Speeds up r.TSR.History.GrandReprojection with r.TSR.History.R11G11B10

diff --git a/Engine/Shaders/Private/TemporalSuperResolution/TSRUpdateHistory.usf b/Engine/Shaders/Private/TemporalSuperResolution/TSRUpdateHistory.usf
--- a/Engine/Shaders/Private/TemporalSuperResolution/TSRUpdateHistory.usf
+++ b/Engine/Shaders/Private/TemporalSuperResolution/TSRUpdateHistory.usf
@@ -1476,6 +1476,11 @@ void MainCS(
 			tsr_halfCx2 PremultipliedHighFrequencyColor = FinalHighFrequencyColor;
 		#endif
 
+		// Encode the overblur with the high frequency color to avoid negative numbers not encodable in a R11G11B10.
+		#if CONFIG_GRAND_REPROJECTION
+			FinalOverblur = FinalOverblur + PremultipliedHighFrequencyColor;
+		#endif
+
 		// Stocastically round up or down using the hardware RWTexture2D truncation unit to take into precision
 		// loss due to pixel format encoding.
 		#if CONFIG_ENABLE_STOCASTIC_QUANTIZATION
@@ -1484,6 +1489,12 @@ void MainCS(
 			tsr_half E = tsr_half(Hammersley16(0, 1, Random).x);
 			
 			PremultipliedHighFrequencyColor = QuantizeForFloatRenderTarget(PremultipliedHighFrequencyColor, E, HistoryQuantizationError);
+			
+			#if CONFIG_GRAND_REPROJECTION
+				FinalPrevHighFrequencyColor = QuantizeForFloatRenderTarget(FinalPrevHighFrequencyColor, E, HistoryQuantizationError);
+				FinalPrevHighFrequencyColorResultant = QuantizeForFloatRenderTarget(FinalPrevHighFrequencyColorResultant, E, HistoryQuantizationError);
+				FinalOverblur = QuantizeForFloatRenderTarget(FinalOverblur, E, HistoryQuantizationError);
+			#endif
 		}
 		#endif
 
@@ -1519,7 +1530,7 @@ void MainCS(
 				#if CONFIG_GRAND_REPROJECTION
 					HistoryOutput_ColorArray[tsr_short3(dpv_lo(LocalHistoryPixelPos), HistoryOutput_ArrayIndices_PrevHighFrequency)] = dpv_lo(FinalPrevHighFrequencyColor);
 					HistoryOutput_ColorArray[tsr_short3(dpv_lo(LocalHistoryPixelPos), HistoryOutput_ArrayIndices_PrevHighFrequencyResultant)] = dpv_lo(FinalPrevHighFrequencyColorResultant);
-					HistoryOutput_ColorArray[tsr_short3(dpv_lo(LocalHistoryPixelPos), HistoryOutput_ArrayIndices_HighFrequencyOverblur)] = dpv_lo(FinalOverblur) + dpv_lo(PremultipliedHighFrequencyColor);
+					HistoryOutput_ColorArray[tsr_short3(dpv_lo(LocalHistoryPixelPos), HistoryOutput_ArrayIndices_HighFrequencyOverblur)] = dpv_lo(FinalOverblur);
 				#endif
 			}
 		
@@ -1543,7 +1554,7 @@ void MainCS(
 				#if CONFIG_GRAND_REPROJECTION
 					HistoryOutput_ColorArray[tsr_short3(dpv_hi(LocalHistoryPixelPos), HistoryOutput_ArrayIndices_PrevHighFrequency)] = dpv_hi(FinalPrevHighFrequencyColor);
 					HistoryOutput_ColorArray[tsr_short3(dpv_hi(LocalHistoryPixelPos), HistoryOutput_ArrayIndices_PrevHighFrequencyResultant)] = dpv_hi(FinalPrevHighFrequencyColorResultant);
-					HistoryOutput_ColorArray[tsr_short3(dpv_hi(LocalHistoryPixelPos), HistoryOutput_ArrayIndices_HighFrequencyOverblur)] = dpv_hi(FinalOverblur) + dpv_hi(PremultipliedHighFrequencyColor);
+					HistoryOutput_ColorArray[tsr_short3(dpv_hi(LocalHistoryPixelPos), HistoryOutput_ArrayIndices_HighFrequencyOverblur)] = dpv_hi(FinalOverblur);
 				#endif
 			}
 			#endif // !CONFIG_ENABLE_DUAL_PIXEL_VECTORIZATION
diff --git a/Engine/Source/Runtime/Renderer/Private/PostProcess/TemporalSuperResolution.cpp b/Engine/Source/Runtime/Renderer/Private/PostProcess/TemporalSuperResolution.cpp
--- a/Engine/Source/Runtime/Renderer/Private/PostProcess/TemporalSuperResolution.cpp
+++ b/Engine/Source/Runtime/Renderer/Private/PostProcess/TemporalSuperResolution.cpp
@@ -1140,7 +1140,7 @@ ITemporalUpscaler::FOutputs AddTemporalSuperResolutionPasses(
 
 		FRDGTextureDesc ArrayDesc = FRDGTextureDesc::Create2DArray(
 			HistoryExtent,
-			(CVarTSRR11G11B10History.GetValueOnRenderThread() != 0 && !bSupportsAlpha && !bGrandReprojection) ? PF_FloatR11G11B10 : PF_FloatRGBA,
+			(CVarTSRR11G11B10History.GetValueOnRenderThread() != 0 && !bSupportsAlpha) ? PF_FloatR11G11B10 : PF_FloatRGBA,
 			FClearValueBinding::None,
 			TexCreate_ShaderResource | TexCreate_UAV,
 			History.ArrayIndices.Size);
