Halfs r.TSR.History.GrandReprojection's memory bandwidth

diff --git a/Engine/Shaders/Private/TemporalSuperResolution/TSRUpdateHistory.usf b/Engine/Shaders/Private/TemporalSuperResolution/TSRUpdateHistory.usf
--- a/Engine/Shaders/Private/TemporalSuperResolution/TSRUpdateHistory.usf
+++ b/Engine/Shaders/Private/TemporalSuperResolution/TSRUpdateHistory.usf
@@ -57,7 +57,11 @@
 	#define CONFIG_TRANSLUCENCY_REJECTION 1
 #endif
 
-#define CONFIG_GRAND_REPROJECTION (DIM_GRAND_REPROJECTION)
+#if DIM_GRAND_REPROJECTION
+	#define CONFIG_GRAND_REPROJECTION 1
+#else
+	#define CONFIG_GRAND_REPROJECTION 0
+#endif
 
 #define CONFIG_RESPONSIVE_STENCIL 1
 #define CONFIG_TAA_RESPONSIVE_BEHAVIOR 1
@@ -86,11 +90,33 @@
 	#define PLATFORM_SPECIFIC_ISOLATE ISOLATE
 #endif
 
+#ifndef CONFIG_INTERPOLATE_VELOCITY
+	#define CONFIG_INTERPOLATE_VELOCITY 0
+#endif
+
 
 //------------------------------------------------------- CONSTANTS
 
 #define TILE_SIZE 8
 
+#if CONFIG_GRAND_REPROJECTION == 1
+	#define CONFIG_METADATA_CHANNELS 2
+#elif CONFIG_PREMULTIPLY_HISTORY
+	#define CONFIG_METADATA_CHANNELS 2
+#else
+	#define CONFIG_METADATA_CHANNELS 1
+#endif
+
+#if CONFIG_METADATA_CHANNELS == 2
+	#define tsr_halfM tsr_half2
+	#define tsr_halfMx2 tsr_half2x2
+#elif CONFIG_METADATA_CHANNELS == 1
+	#define tsr_halfM tsr_half
+	#define tsr_halfMx2 tsr_half2
+#else
+	#error Unknown CONFIG_METADATA_CHANNELS
+#endif
+
 
 //------------------------------------------------------- PARAMETERS
 
@@ -109,7 +135,7 @@ Texture2D<tsr_half> NoiseFilteringTexture;
 Texture2D<tsr_half2> HoleFilledVelocityMaskTexture;
 
 Texture2D<float2> GrandDilatedVelocityTexture;
-Texture2DArray<tsr_halfC> GrandPrevColorArray;
+Texture2D<tsr_halfC> GrandPrevColorTexture;
 
 float2 TranslucencyInfo_Extent;
 float2 TranslucencyInfo_ExtentInverse;
@@ -143,11 +169,7 @@ uint bGenerateOutputMip1;
 uint bHasSeparateTranslucency;
 
 Texture2D<tsr_halfC> PrevHistory_HighFrequency;
-#if CONFIG_PREMULTIPLY_HISTORY
-	Texture2D<tsr_half2> PrevHistory_Metadata;
-#else
-	Texture2D<tsr_half> PrevHistory_Metadata;
-#endif
+Texture2D<tsr_halfM> PrevHistory_Metadata;
 Texture2D<tsr_half3> PrevHistory_Translucency;
 Texture2D<tsr_half>  PrevHistory_TranslucencyAlpha;
 Texture2D<tsr_subpixel_details> PrevHistory_SubpixelDetails;
@@ -166,11 +188,7 @@ uint HistoryOutput_ArrayIndices_PrevHighFrequencyResultant;
 uint HistoryOutput_ArrayIndices_HighFrequencyOverblur;
 
 RWTexture2DArray<tsr_halfC> HistoryOutput_ColorArray;
-#if CONFIG_PREMULTIPLY_HISTORY
-	RWTexture2D<tsr_half2> HistoryOutput_Metadata;
-#else
-	RWTexture2D<tsr_half>  HistoryOutput_Metadata;
-#endif
+RWTexture2D<tsr_halfM> HistoryOutput_Metadata;
 RWTexture2D<tsr_half>  HistoryOutput_TranslucencyAlpha;
 RWTexture2D<tsr_subpixel_details> HistoryOutput_SubpixelDetails;
 
@@ -396,7 +414,11 @@ void MainCS(
 		#endif
 		
 		#if CONFIG_GRAND_REPROJECTION
-			float2 RawEncodedGrandVelocity[DPV_PIXEL_PER_LANE];
+			#if CONFIG_INTERPOLATE_VELOCITY
+				float2 RawEncodedGrandVelocityNeighborhood[4][DPV_PIXEL_PER_LANE];
+			#else
+				float2 RawEncodedGrandVelocity[DPV_PIXEL_PER_LANE];
+			#endif
 		#endif
 
 		tsr_half RawParallaxRejectionMask[DPV_PIXEL_PER_LANE];
@@ -424,10 +446,9 @@ void MainCS(
 				
 				#if !CONFIG_INTERPOLATE_VELOCITY
 					RawEncodedVelocity[PixelId]   = DilatedVelocityTexture[LocalInputPixelPos];
-				#endif
-				
-				#if CONFIG_GRAND_REPROJECTION
-					RawEncodedGrandVelocity[PixelId] = GrandDilatedVelocityTexture[LocalInputPixelPos];
+					#if CONFIG_GRAND_REPROJECTION
+						RawEncodedGrandVelocity[PixelId] = GrandDilatedVelocityTexture[LocalInputPixelPos];
+					#endif
 				#endif
 
 				RawParallaxRejectionMask[PixelId] = ParallaxRejectionMaskTexture[LocalInputPixelPos];
@@ -479,6 +500,13 @@ void MainCS(
 				RawEncodedVelocityNeighborhood[1][PixelId2] = DilatedVelocityTexture[tsr_short2(LocalInputPixelPos1.x, LocalInputPixelPos0.y)];
 				RawEncodedVelocityNeighborhood[2][PixelId2] = DilatedVelocityTexture[tsr_short2(LocalInputPixelPos0.x, LocalInputPixelPos1.y)];
 				RawEncodedVelocityNeighborhood[3][PixelId2] = DilatedVelocityTexture[tsr_short2(LocalInputPixelPos1.x, LocalInputPixelPos1.y)];
+				
+				#if CONFIG_GRAND_REPROJECTION
+					RawEncodedGrandVelocityNeighborhood[0][PixelId2] = GrandDilatedVelocityTexture[tsr_short2(LocalInputPixelPos0.x, LocalInputPixelPos0.y)];
+					RawEncodedGrandVelocityNeighborhood[1][PixelId2] = GrandDilatedVelocityTexture[tsr_short2(LocalInputPixelPos1.x, LocalInputPixelPos0.y)];
+					RawEncodedGrandVelocityNeighborhood[2][PixelId2] = GrandDilatedVelocityTexture[tsr_short2(LocalInputPixelPos0.x, LocalInputPixelPos1.y)];
+					RawEncodedGrandVelocityNeighborhood[3][PixelId2] = GrandDilatedVelocityTexture[tsr_short2(LocalInputPixelPos1.x, LocalInputPixelPos1.y)];
+				#endif
 			}
 			#endif
 		}
@@ -522,26 +550,56 @@ void MainCS(
 					ScreenVelocity[0][PixelId] = LocalScreenVelocity.x;
 					ScreenVelocity[1][PixelId] = LocalScreenVelocity.y;
 				}
+				#if CONFIG_GRAND_REPROJECTION
+				for (uint PixelId = 0; PixelId < DPV_PIXEL_PER_LANE; PixelId++)
+				{
+					float2 V00 = DecodeVelocityFromTexture(float4(RawEncodedGrandVelocityNeighborhood[0][PixelId], 0.0, 0.0)).xy;
+					float2 V10 = DecodeVelocityFromTexture(float4(RawEncodedGrandVelocityNeighborhood[1][PixelId], 0.0, 0.0)).xy;
+					float2 V01 = DecodeVelocityFromTexture(float4(RawEncodedGrandVelocityNeighborhood[2][PixelId], 0.0, 0.0)).xy;
+					float2 V11 = DecodeVelocityFromTexture(float4(RawEncodedGrandVelocityNeighborhood[3][PixelId], 0.0, 0.0)).xy;
+
+					tsr_half2 BilinearInterp = abs(dpv_access_pixel(dInputKO, PixelId));
+					
+					float Bilinear00 = saturate(1.0 - BilinearInterp.x) * saturate(1.0 - BilinearInterp.y);
+					float Bilinear10 = saturate(0.0 + BilinearInterp.x) * saturate(1.0 - BilinearInterp.y);
+					float Bilinear01 = saturate(1.0 - BilinearInterp.x) * saturate(0.0 + BilinearInterp.y);
+					float Bilinear11 = saturate(0.0 + BilinearInterp.x) * saturate(0.0 + BilinearInterp.y);
+					
+					float Bilateral00 = Bilinear00;
+					float Bilateral10 = Bilinear10 * ComputeScreenVelocityBilateralWeight(V10 - V00);
+					float Bilateral01 = Bilinear01 * ComputeScreenVelocityBilateralWeight(V01 - V00);
+					float Bilateral11 = Bilinear11 * ComputeScreenVelocityBilateralWeight(V11 - V00);
+
+					float NormalizeBilateral = rcp(Bilateral00 + Bilateral10 + Bilateral01 + Bilateral11);
+
+					float2 LocalScreenVelocity = (V00 * Bilateral00 + V10 * Bilateral10 + V01 * Bilateral01 + V11 * Bilateral11) * NormalizeBilateral;
+					
+					float2 LocalPrevScreenPos = dpv_access_pixel(ScreenPos, PixelId) - LocalScreenVelocity;
+
+					GrandPrevScreenPos[0][PixelId] = LocalPrevScreenPos.x;
+					GrandPrevScreenPos[1][PixelId] = LocalPrevScreenPos.y;
+				}
+				#endif
 				#else
 				{
 					ScreenVelocity = dpv_interleave_registers(
 						DecodeVelocityFromTexture(float4(RawEncodedVelocity[0], 0.0, 0.0)).xy,
 						DecodeVelocityFromTexture(float4(RawEncodedVelocity[1], 0.0, 0.0)).xy);
+						
+					#if CONFIG_GRAND_REPROJECTION
+					{
+						float2x2 GrandScreenVelocity = dpv_interleave_registers(
+							DecodeVelocityFromTexture(float4(RawEncodedGrandVelocity[0], 0.0, 0.0)).xy,
+							DecodeVelocityFromTexture(float4(RawEncodedGrandVelocity[1], 0.0, 0.0)).xy);
+					
+						GrandPrevScreenPos = ScreenPos - GrandScreenVelocity;
+					}
+					#endif
 				}
 				#endif
 
 				PrevScreenPos = ScreenPos - ScreenVelocity;
 				OutputPixelVelocity = tsr_half2(dpv_length(dpv_mul(ScreenVelocity, HistoryInfo_ViewportSize)));
-				
-				#if CONFIG_GRAND_REPROJECTION
-				{
-					float2x2 GrandScreenVelocity = dpv_interleave_registers(
-						DecodeVelocityFromTexture(float4(RawEncodedGrandVelocity[0], 0.0, 0.0)).xy,
-						DecodeVelocityFromTexture(float4(RawEncodedGrandVelocity[1], 0.0, 0.0)).xy);
-					
-					GrandPrevScreenPos = ScreenPos - GrandScreenVelocity;
-				}
-				#endif
 			}
 			
 			// Detect whether the history reprojection is successful
@@ -621,8 +679,12 @@ void MainCS(
 	
 	#if CONFIG_GRAND_REPROJECTION
 		tsr_halfCx2 GrandPrevToPrev;
-		tsr_halfCx2 GrandPrevToPrevResultant;
-		tsr_halfCx2 PrevHighFrequencyOverblur;
+		#if CONFIG_GRAND_REPROJECTION == 1
+			tsr_half2 GrandPrevToPrevResultant;
+		#elif CONFIG_GRAND_REPROJECTION == 2
+			tsr_halfCx2 GrandPrevToPrevResultant;
+			tsr_halfCx2 PrevHighFrequencyOverblur;
+		#endif
 		tsr_halfCx2 GrandReprojection;
 	#endif
 
@@ -630,11 +692,7 @@ void MainCS(
 	{
 		// Reproject history
 		tsr_halfCx2 AccumulateHighFrequency = tsr_half(0);
-		#if CONFIG_PREMULTIPLY_HISTORY
-			tsr_half2x2 AccumulateMetadata = tsr_half(0);
-		#else
-			tsr_half2 AccumulateMetadata = tsr_half(0);
-		#endif
+		tsr_halfMx2 AccumulateMetadata = tsr_half(0);
 		tsr_half4x2 AccumulateTranslucency = tsr_half(0);
 
 		tsr_halfCx2 HighFrequencyMin = INFINITE_FLOAT;
@@ -642,8 +700,10 @@ void MainCS(
 		
 		#if CONFIG_GRAND_REPROJECTION
 			tsr_halfCx2 AccumulateGrandPrevToPrev = tsr_half(0);
-			tsr_halfCx2 AccumulateGrandPrevToPrevResultant = tsr_half(0);
-			tsr_halfCx2 AccumulatePrevHighFrequencyOverblur = tsr_half(0);
+			#if CONFIG_GRAND_REPROJECTION == 2
+				tsr_halfCx2 AccumulateGrandPrevToPrevResultant = tsr_half(0);
+				tsr_halfCx2 AccumulatePrevHighFrequencyOverblur = tsr_half(0);
+			#endif
 			tsr_halfCx2 AccumulateGrandReprojection = tsr_half(0);
 		#endif
 
@@ -670,11 +730,7 @@ void MainCS(
 			#endif
 
 			tsr_halfC RawHighFrequency[BICUBIC_CATMULL_ROM_SAMPLES][DPV_PIXEL_PER_LANE];
-			#if CONFIG_PREMULTIPLY_HISTORY
-				tsr_half2 RawMetadata[BICUBIC_CATMULL_ROM_SAMPLES][DPV_PIXEL_PER_LANE];
-			#else
-				tsr_half RawMetadata[BICUBIC_CATMULL_ROM_SAMPLES][DPV_PIXEL_PER_LANE];
-			#endif
+			tsr_halfM RawMetadata[BICUBIC_CATMULL_ROM_SAMPLES][DPV_PIXEL_PER_LANE];
 			tsr_half3 RawTranslucencyColor[BICUBIC_CATMULL_ROM_SAMPLES][DPV_PIXEL_PER_LANE];
 			tsr_half  RawTranslucencyAlpha[BICUBIC_CATMULL_ROM_SAMPLES][DPV_PIXEL_PER_LANE];
 
@@ -682,8 +738,10 @@ void MainCS(
 			
 			#if CONFIG_GRAND_REPROJECTION
 				tsr_halfC RawGrandPrevToPrev[BICUBIC_CATMULL_ROM_SAMPLES][DPV_PIXEL_PER_LANE];
-				tsr_halfC RawGrandPrevToPrevResultant[BICUBIC_CATMULL_ROM_SAMPLES][DPV_PIXEL_PER_LANE];
-				tsr_halfC RawPrevHighFrequencyOverblur[BICUBIC_CATMULL_ROM_SAMPLES][DPV_PIXEL_PER_LANE];
+				#if CONFIG_GRAND_REPROJECTION == 2
+					tsr_halfC RawGrandPrevToPrevResultant[BICUBIC_CATMULL_ROM_SAMPLES][DPV_PIXEL_PER_LANE];
+					tsr_halfC RawPrevHighFrequencyOverblur[BICUBIC_CATMULL_ROM_SAMPLES][DPV_PIXEL_PER_LANE];
+				#endif
 				tsr_halfC RawGrandReprojection[BICUBIC_CATMULL_ROM_SAMPLES][DPV_PIXEL_PER_LANE];
 
 				tsr_half2 RawGrandKernelWeight[BICUBIC_CATMULL_ROM_SAMPLES];
@@ -721,9 +779,11 @@ void MainCS(
 						#endif
 						
 						#if CONFIG_GRAND_REPROJECTION
-							RawGrandPrevToPrev[i][PixelId] = PrevHistory_PrevHighFrequency.SampleLevel(GlobalBilinearClampedSampler, SampleUV, 0);
-							RawGrandPrevToPrevResultant[i][PixelId] = PrevHistory_PrevHighFrequencyResultant.SampleLevel(GlobalBilinearClampedSampler, SampleUV, 0);
-							RawPrevHighFrequencyOverblur[i][PixelId] = PrevHistory_HighFrequencyOverblur.SampleLevel(GlobalBilinearClampedSampler, SampleUV, 0);
+							RawGrandPrevToPrev[i][PixelId] = BilinearSampleColorHistory(PrevHistory_PrevHighFrequency, SampleUV);
+							#if CONFIG_GRAND_REPROJECTION == 2
+								RawGrandPrevToPrevResultant[i][PixelId] = BilinearSampleColorHistory(PrevHistory_PrevHighFrequencyResultant, SampleUV);
+								RawPrevHighFrequencyOverblur[i][PixelId] = BilinearSampleColorHistory(PrevHistory_HighFrequencyOverblur, SampleUV);
+							#endif
 
 							{
 								float2 GrandSampleUV0 = clamp(GrandSamples0.UV[i], PrevHistoryInfo_UVViewportBilinearMin, PrevHistoryInfo_UVViewportBilinearMax);
@@ -731,7 +791,7 @@ void MainCS(
 					
 								float2 GrandSampleUV = PixelId == 0 ? GrandSampleUV0 : GrandSampleUV1;
 						
-								RawGrandReprojection[i][PixelId] = GrandPrevColorArray.SampleLevel(GlobalBilinearClampedSampler, float3(GrandSampleUV, 0.0), 0);
+								RawGrandReprojection[i][PixelId] = BilinearSampleColorHistory(GrandPrevColorTexture, GrandSampleUV);
 							
 							}
 						#endif
@@ -747,19 +807,17 @@ void MainCS(
 				for (uint i = 0; i < BICUBIC_CATMULL_ROM_SAMPLES; i++)
 				{
 					tsr_halfCx2 SampleHighFrequency = dpv_force_interleave_registers_array(RawHighFrequency[i]);
-					#if CONFIG_PREMULTIPLY_HISTORY
-						tsr_half2x2 SampleMetadata = dpv_force_interleave_registers_array(RawMetadata[i]);
-					#else
-						tsr_half2 SampleMetadata = dpv_force_interleave_registers_array(RawMetadata[i]);
-					#endif
+					tsr_halfMx2 SampleMetadata = dpv_force_interleave_registers_array(RawMetadata[i]);
 					tsr_half4x2 SampleTranslucency = dpv_force_interleave_registers(
 						tsr_half4(RawTranslucencyColor[i][0], RawTranslucencyAlpha[i][0]),
 						tsr_half4(RawTranslucencyColor[i][1], RawTranslucencyAlpha[i][1]));
 					
 					#if CONFIG_GRAND_REPROJECTION
 						tsr_halfCx2 SampleGrandPrevious = dpv_interleave_registers_array(RawGrandPrevToPrev[i]);
-						tsr_halfCx2 SampleGrandPreviousResultant = dpv_interleave_registers_array(RawGrandPrevToPrevResultant[i]);
-						tsr_halfCx2 SamplePrevHighFrequencyOverblur = dpv_interleave_registers_array(RawPrevHighFrequencyOverblur[i]) - SampleHighFrequency;
+						#if CONFIG_GRAND_REPROJECTION == 2
+							tsr_halfCx2 SampleGrandPreviousResultant = dpv_interleave_registers_array(RawGrandPrevToPrevResultant[i]);
+							tsr_halfCx2 SamplePrevHighFrequencyOverblur = dpv_interleave_registers_array(RawPrevHighFrequencyOverblur[i]) - SampleHighFrequency;
+						#endif
 						tsr_halfCx2 SampleGrandReprojection = dpv_interleave_registers_array(RawGrandReprojection[i]);
 					#endif
 
@@ -778,11 +836,7 @@ void MainCS(
 					#endif
 
 					AccumulateHighFrequency += dpv_scale(SampleHighFrequency, KernelWeight);
-					#if CONFIG_PREMULTIPLY_HISTORY
-						AccumulateMetadata  += dpv_scale(SampleMetadata, KernelWeight);
-					#else
-						AccumulateMetadata  += SampleMetadata * KernelWeight;
-					#endif
+					AccumulateMetadata  += dpv_scale(SampleMetadata, KernelWeight);
 					
 					#if CONFIG_ACCUMULATE_TRANSLUCENCY_SEPARATLY
 						AccumulateTranslucency  += dpv_scale(SampleTranslucency, KernelWeight);
@@ -790,8 +844,10 @@ void MainCS(
 
 					#if CONFIG_GRAND_REPROJECTION
 						AccumulateGrandPrevToPrev  += dpv_scale(SampleGrandPrevious, KernelWeight);
-						AccumulateGrandPrevToPrevResultant  += dpv_scale(SampleGrandPreviousResultant, KernelWeight);
-						AccumulatePrevHighFrequencyOverblur  += dpv_scale(SamplePrevHighFrequencyOverblur, KernelWeight);
+						#if CONFIG_GRAND_REPROJECTION == 2
+							AccumulateGrandPrevToPrevResultant  += dpv_scale(SampleGrandPreviousResultant, KernelWeight);
+							AccumulatePrevHighFrequencyOverblur  += dpv_scale(SamplePrevHighFrequencyOverblur, KernelWeight);
+						#endif
 						AccumulateGrandReprojection  += dpv_scale(SampleGrandReprojection, RawGrandKernelWeight[i]);
 					#endif
 				}
@@ -799,19 +855,17 @@ void MainCS(
 				tsr_half2 FinalMultiplier = dpv_force_interleave_registers(tsr_half(Samples0.FinalMultiplier), tsr_half(Samples1.FinalMultiplier));
 
 				AccumulateHighFrequency = dpv_scale(AccumulateHighFrequency, FinalMultiplier);
-				#if CONFIG_PREMULTIPLY_HISTORY
-					AccumulateMetadata  = dpv_scale(AccumulateMetadata, FinalMultiplier);
-				#else
-					AccumulateMetadata  = AccumulateMetadata * FinalMultiplier;
-				#endif
+				AccumulateMetadata  = dpv_scale(AccumulateMetadata, FinalMultiplier);
 				AccumulateTranslucency  = dpv_scale(AccumulateTranslucency, FinalMultiplier);
 
 				#if CONFIG_GRAND_REPROJECTION
 					tsr_half2 GrandFinalMultiplier = dpv_interleave_registers(tsr_half(GrandSamples0.FinalMultiplier), tsr_half(GrandSamples1.FinalMultiplier));
 
 					AccumulateGrandPrevToPrev          = dpv_scale(AccumulateGrandPrevToPrev, FinalMultiplier);
-					AccumulateGrandPrevToPrevResultant = dpv_scale(AccumulateGrandPrevToPrevResultant, FinalMultiplier);
-					AccumulatePrevHighFrequencyOverblur = dpv_scale(AccumulatePrevHighFrequencyOverblur, FinalMultiplier);
+					#if CONFIG_GRAND_REPROJECTION == 2
+						AccumulateGrandPrevToPrevResultant = dpv_scale(AccumulateGrandPrevToPrevResultant, FinalMultiplier);
+						AccumulatePrevHighFrequencyOverblur = dpv_scale(AccumulatePrevHighFrequencyOverblur, FinalMultiplier);
+					#endif
 					AccumulateGrandReprojection        = dpv_scale(AccumulateGrandReprojection, GrandFinalMultiplier);
 				#endif
 			}
@@ -821,11 +875,7 @@ void MainCS(
 			#else
 				AccumulateHighFrequency = -dpv_min(-AccumulateHighFrequency, tsr_half(0.0).xxx);
 			#endif
-			#if CONFIG_PREMULTIPLY_HISTORY
-				AccumulateMetadata      = -dpv_min(-AccumulateMetadata, tsr_half(0.0).xx);
-			#else
-				AccumulateMetadata      = -min(-AccumulateMetadata, tsr_half(0.0).xx);
-			#endif
+			AccumulateMetadata      = -dpv_min(-AccumulateMetadata, tsr_half(0.0).xx);
 			AccumulateTranslucency  = -dpv_min(-AccumulateTranslucency, tsr_half(0.0).xxxx);
 		}
 		
@@ -844,8 +894,12 @@ void MainCS(
 
 			#if CONFIG_GRAND_REPROJECTION
 				GrandPrevToPrev          = AccumulateGrandPrevToPrev;
-				GrandPrevToPrevResultant = AccumulateGrandPrevToPrevResultant;
-				PrevHighFrequencyOverblur = AccumulatePrevHighFrequencyOverblur;
+				#if CONFIG_GRAND_REPROJECTION == 1
+					GrandPrevToPrevResultant = AccumulateMetadata[1];
+				#elif CONFIG_GRAND_REPROJECTION == 2
+					GrandPrevToPrevResultant = AccumulateGrandPrevToPrevResultant;
+					PrevHighFrequencyOverblur = AccumulatePrevHighFrequencyOverblur;
+				#endif
 				GrandReprojection        = AccumulateGrandReprojection;
 			#endif
 		}
@@ -866,12 +920,14 @@ void MainCS(
 				CorrectExposure(GrandPrevToPrev[0]);
 				CorrectExposure(GrandPrevToPrev[1]);
 				CorrectExposure(GrandPrevToPrev[2]);
-				CorrectExposure(GrandPrevToPrevResultant[0]);
-				CorrectExposure(GrandPrevToPrevResultant[1]);
-				CorrectExposure(GrandPrevToPrevResultant[2]);
-				CorrectExposure(PrevHighFrequencyOverblur[0]);
-				CorrectExposure(PrevHighFrequencyOverblur[1]);
-				CorrectExposure(PrevHighFrequencyOverblur[2]);
+				#if CONFIG_GRAND_REPROJECTION == 2
+					CorrectExposure(GrandPrevToPrevResultant[0]);
+					CorrectExposure(GrandPrevToPrevResultant[1]);
+					CorrectExposure(GrandPrevToPrevResultant[2]);
+					CorrectExposure(PrevHighFrequencyOverblur[0]);
+					CorrectExposure(PrevHighFrequencyOverblur[1]);
+					CorrectExposure(PrevHighFrequencyOverblur[2]);
+				#endif
 				
 				GrandReprojection[0] *= tsr_half(GrandPrevPreExposureCorrection);
 				GrandReprojection[1] *= tsr_half(GrandPrevPreExposureCorrection);
@@ -1168,8 +1224,12 @@ void MainCS(
 	
 	#if CONFIG_GRAND_REPROJECTION
 		tsr_halfCx2 FinalPrevHighFrequencyColor;
-		tsr_halfCx2 FinalPrevHighFrequencyColorResultant;
-		tsr_halfCx2 FinalOverblur = tsr_half(0.0);
+		#if CONFIG_GRAND_REPROJECTION == 1
+			tsr_half2 FinalPrevHighFrequencyResultant;
+		#elif CONFIG_GRAND_REPROJECTION == 2
+			tsr_halfCx2 FinalPrevHighFrequencyColorResultant;
+			tsr_halfCx2 FinalOverblur = tsr_half(0.0);
+		#endif
 		tsr_halfCx2 FinalHighFrequencyColorOverblurCorrection = tsr_half(0.0);
 	#endif
 
@@ -1265,25 +1325,11 @@ void MainCS(
 		#if CONFIG_GRAND_REPROJECTION
 		{
 			#if CONFIG_PREMULTIPLY_HISTORY
-				tsr_half2 FinalHistoryWeight = max(FinalHistoryValidity, MIN_IMPLICIT_HISTORY_WEIGHT);
-			#else
-				const tsr_half2 PrevHistoryWeight = tsr_half(1.0);
-				const tsr_half2 InvPrevHistoryWeight = tsr_half(1.0);
-				const tsr_half2 FinalHistoryWeight = tsr_half(1.0);
+				#error Unsupported
 			#endif
 
-			tsr_halfCx2 BoxSize = dpv_scale(PrevHistoryMommentMax - PrevHistoryMommentMin, tsr_half(0.5) * PrevHistoryWeight);
+			tsr_halfCx2 BoxSize = dpv_scale(PrevHistoryMommentMax - PrevHistoryMommentMin, tsr_half(0.5));
 
-			// Compute MeasuredResultant of the previous frame
-			//tsr_halfCx2 PrevResultant = -dpv_min(-GrandPrevToPrevResultant / GrandPrevToPrev, tsr_half(0.0).xxx);
-			tsr_halfCx2 PrevResultant;
-			PrevResultant[0] = clamp(GrandPrevToPrevResultant[0] * SafeRcp(GrandPrevToPrev[0]), tsr_half(0.0), tsr_half(1.1));
-			PrevResultant[1] = clamp(GrandPrevToPrevResultant[1] * SafeRcp(GrandPrevToPrev[1]), tsr_half(0.0), tsr_half(1.1));
-			PrevResultant[2] = clamp(GrandPrevToPrevResultant[2] * SafeRcp(GrandPrevToPrev[2]), tsr_half(0.0), tsr_half(1.1));
-			#if CONFIG_SCENE_COLOR_APLHA
-				PrevResultant[3] = clamp(GrandPrevToPrevResultant[3] * SafeRcp(GrandPrevToPrev[3]), tsr_half(0.0), tsr_half(1.1));
-			#endif
-			
 			// Compute how much frame N-2 got overblur by reprojecting into N-1 and N versus directly to N.
 			tsr_halfCx2 PrevReprojectionOverblur = GrandReprojection - GrandPrevToPrev;
 
@@ -1295,20 +1341,44 @@ void MainCS(
 
 			tsr_half2 PrevReprojectionOverblurScale = select(bIsOverblurLegal, tsr_half(1.0), tsr_half(0.0));
 
-			tsr_half2 ClampRejection = MeasureRejectionFactor(
+			tsr_half2 HistoryClampResultant = MeasureRejectionFactor(
 				UnpremultipliedPrevHighFrequencyColor, BlendedPrevHighFrequencyColor,
 				FilteredInputColor, InputMinColor, InputMaxColor,
 				MeasureBackbufferLDRQuantizationError());
 
-			tsr_half2 HistoryClampResultant = lerp(ClampRejection, tsr_half(1.0), BlendClamp);
+			tsr_half2 MeasuredResultant = saturate(CommonWeight * BlendPrevHistory * HistoryClampResultant * TotalRejection);
+
+			FinalPrevHighFrequencyColor = PrevHighFrequencyColor;
 
-			FinalOverblur = dpv_scale(PrevResultant, HistoryClampResultant * TotalRejection) * (PrevHighFrequencyOverblur + dpv_scale(PrevReprojectionOverblur, PrevReprojectionOverblurScale));
+			#if CONFIG_GRAND_REPROJECTION == 1
+				FinalPrevHighFrequencyResultant = MeasuredResultant;
 
-			tsr_half2 MeasuredResultant = saturate(CommonWeight * BlendPrevHistory * HistoryClampResultant);
+				const tsr_half CorrectionError = tsr_half(1.0 / sqrt(PI * 0.5));
+				MeasuredResultant = min(MeasuredResultant, saturate(tsr_half(1.0) - Histeresis * rcp(FinalHistoryValidity)));
 
-			FinalPrevHighFrequencyColor = PrevHighFrequencyColor;
-			FinalPrevHighFrequencyColorResultant = dpv_scale(PrevHighFrequencyColor, MeasuredResultant * InvPrevHistoryWeight * FinalHistoryWeight);
-			FinalHighFrequencyColorOverblurCorrection = dpv_scale(FinalOverblur, MeasuredResultant * InvPrevHistoryWeight);
+				tsr_half2 CorrectionResultant = GrandPrevToPrevResultant * MeasuredResultant * PrevReprojectionOverblurScale * CorrectionError;
+
+				tsr_halfCx2 Correction = dpv_scale(PrevReprojectionOverblur, CorrectionResultant);
+
+				tsr_halfCx2 CorrectedFinalHighFrequencyColor = FinalHighFrequencyColor + Correction;
+
+				FinalHighFrequencyColor = clamp(CorrectedFinalHighFrequencyColor, min(FinalHighFrequencyColor, PrevHistoryMommentMin), max(FinalHighFrequencyColor, PrevHistoryMommentMax));
+
+			#elif CONFIG_GRAND_REPROJECTION == 2
+				// Compute MeasuredResultant of the previous frame
+				tsr_halfCx2 PrevResultant;
+				PrevResultant[0] = clamp(GrandPrevToPrevResultant[0] * SafeRcp(GrandPrevToPrev[0]), tsr_half(0.0), tsr_half(1.0));
+				PrevResultant[1] = clamp(GrandPrevToPrevResultant[1] * SafeRcp(GrandPrevToPrev[1]), tsr_half(0.0), tsr_half(1.0));
+				PrevResultant[2] = clamp(GrandPrevToPrevResultant[2] * SafeRcp(GrandPrevToPrev[2]), tsr_half(0.0), tsr_half(1.0));
+				#if CONFIG_SCENE_COLOR_APLHA
+					PrevResultant[3] = clamp(GrandPrevToPrevResultant[3] * SafeRcp(GrandPrevToPrev[3]), tsr_half(0.0), tsr_half(1.0));
+				#endif
+
+				FinalPrevHighFrequencyColorResultant = dpv_scale(PrevHighFrequencyColor, MeasuredResultant);
+
+				FinalOverblur = dpv_scale(PrevResultant, HistoryClampResultant * TotalRejection) * (PrevHighFrequencyOverblur + dpv_scale(PrevReprojectionOverblur, PrevReprojectionOverblurScale));
+				FinalHighFrequencyColorOverblurCorrection = dpv_scale(FinalOverblur, MeasuredResultant);
+			#endif
 		}
 		#endif
 	}
@@ -1406,7 +1476,7 @@ void MainCS(
 	// Compute final output
 	tsr_halfCx2 FinalOutputColor;
 	{
-		#if CONFIG_GRAND_REPROJECTION
+		#if CONFIG_GRAND_REPROJECTION == 2
 			FinalOutputColor = FinalHighFrequencyColor + FinalHighFrequencyColorOverblurCorrection;
 		#else
 			FinalOutputColor = FinalHighFrequencyColor;
@@ -1467,17 +1537,27 @@ void MainCS(
 		tsr_short2x2 HalfResOutputPixelPos;
 		HalfResOutputPixelPos[0] = (LocalHistoryPixelPos[0] >> tsr_short(1)) | (((LocalHistoryPixelPos[0] & tsr_short(0x1))) * tsr_short(~0));
 		HalfResOutputPixelPos[1] = (LocalHistoryPixelPos[1] >> tsr_short(1)) | (((LocalHistoryPixelPos[1] & tsr_short(0x1))) * tsr_short(~0));
+		
+		#if CONFIG_METADATA_CHANNELS == 1
+			tsr_halfMx2 FinalMetadata = FinalHistoryValidity;
+		#else
+			tsr_halfMx2 FinalMetadata;
+			FinalMetadata[0] = FinalHistoryValidity;
+		#endif
 
 		// Premultiply FinalHighFrequencyColor to have validity weighted average in next frame history reprojection.
 		#if CONFIG_PREMULTIPLY_HISTORY
 			tsr_half2 FinalHistoryWeight = max(FinalHistoryValidity, MIN_IMPLICIT_HISTORY_WEIGHT);
 			tsr_halfCx2 PremultipliedHighFrequencyColor = dpv_scale(FinalHighFrequencyColor, FinalHistoryWeight);
+			FinalMetadata[1] = FinalHistoryWeight;
 		#else
 			tsr_halfCx2 PremultipliedHighFrequencyColor = FinalHighFrequencyColor;
 		#endif
-
-		// Encode the overblur with the high frequency color to avoid negative numbers not encodable in a R11G11B10.
-		#if CONFIG_GRAND_REPROJECTION
+		
+		#if CONFIG_GRAND_REPROJECTION == 1
+			FinalMetadata[1] = FinalPrevHighFrequencyResultant;
+		#elif CONFIG_GRAND_REPROJECTION == 2
+			// Encode the overblur with the high frequency color to avoid negative numbers not encodable in a R11G11B10.
 			FinalOverblur = FinalOverblur + PremultipliedHighFrequencyColor;
 		#endif
 
@@ -1492,8 +1572,10 @@ void MainCS(
 			
 			#if CONFIG_GRAND_REPROJECTION
 				FinalPrevHighFrequencyColor = QuantizeForFloatRenderTarget(FinalPrevHighFrequencyColor, E, HistoryQuantizationError);
-				FinalPrevHighFrequencyColorResultant = QuantizeForFloatRenderTarget(FinalPrevHighFrequencyColorResultant, E, HistoryQuantizationError);
-				FinalOverblur = QuantizeForFloatRenderTarget(FinalOverblur, E, HistoryQuantizationError);
+				#if CONFIG_GRAND_REPROJECTION == 2
+					FinalPrevHighFrequencyColorResultant = QuantizeForFloatRenderTarget(FinalPrevHighFrequencyColorResultant, E, HistoryQuantizationError);
+					FinalOverblur = QuantizeForFloatRenderTarget(FinalOverblur, E, HistoryQuantizationError);
+				#endif
 			#endif
 		}
 		#endif
@@ -1514,11 +1596,7 @@ void MainCS(
 			// Output final history lo pixel.
 			{
 				HistoryOutput_ColorArray[tsr_short3(dpv_lo(LocalHistoryPixelPos), HistoryOutput_ArrayIndices_HighFrequency)] = dpv_lo(PremultipliedHighFrequencyColor);
-				#if CONFIG_PREMULTIPLY_HISTORY
-					HistoryOutput_Metadata[dpv_lo(LocalHistoryPixelPos)]      = tsr_half2(dpv_lo(FinalHistoryValidity), dpv_lo(FinalHistoryWeight));
-				#else
-					HistoryOutput_Metadata[dpv_lo(LocalHistoryPixelPos)]      = dpv_lo(FinalHistoryValidity);
-				#endif
+				HistoryOutput_Metadata[dpv_lo(LocalHistoryPixelPos)]      = dpv_lo(FinalMetadata);
 				#if CONFIG_ACCUMULATE_TRANSLUCENCY_SEPARATLY
 					#if CONFIG_SCENE_COLOR_APLHA
 						HistoryOutput_ColorArray[tsr_short3(dpv_lo(LocalHistoryPixelPos), HistoryOutput_ArrayIndices_Translucency)] = dpv_lo(FinalTranslucencyColor);
@@ -1529,8 +1607,10 @@ void MainCS(
 				#endif
 				#if CONFIG_GRAND_REPROJECTION
 					HistoryOutput_ColorArray[tsr_short3(dpv_lo(LocalHistoryPixelPos), HistoryOutput_ArrayIndices_PrevHighFrequency)] = dpv_lo(FinalPrevHighFrequencyColor);
-					HistoryOutput_ColorArray[tsr_short3(dpv_lo(LocalHistoryPixelPos), HistoryOutput_ArrayIndices_PrevHighFrequencyResultant)] = dpv_lo(FinalPrevHighFrequencyColorResultant);
-					HistoryOutput_ColorArray[tsr_short3(dpv_lo(LocalHistoryPixelPos), HistoryOutput_ArrayIndices_HighFrequencyOverblur)] = dpv_lo(FinalOverblur);
+					#if CONFIG_GRAND_REPROJECTION == 2
+						HistoryOutput_ColorArray[tsr_short3(dpv_lo(LocalHistoryPixelPos), HistoryOutput_ArrayIndices_PrevHighFrequencyResultant)] = dpv_lo(FinalPrevHighFrequencyColorResultant);
+						HistoryOutput_ColorArray[tsr_short3(dpv_lo(LocalHistoryPixelPos), HistoryOutput_ArrayIndices_HighFrequencyOverblur)] = dpv_lo(FinalOverblur);
+					#endif
 				#endif
 			}
 		
@@ -1538,11 +1618,7 @@ void MainCS(
 			#if CONFIG_ENABLE_DUAL_PIXEL_VECTORIZATION
 			{
 				HistoryOutput_ColorArray[tsr_short3(dpv_hi(LocalHistoryPixelPos), HistoryOutput_ArrayIndices_HighFrequency)] = dpv_hi(PremultipliedHighFrequencyColor);
-				#if CONFIG_PREMULTIPLY_HISTORY
-					HistoryOutput_Metadata[dpv_hi(LocalHistoryPixelPos)]      = tsr_half2(dpv_hi(FinalHistoryValidity), dpv_hi(FinalHistoryWeight));
-				#else
-					HistoryOutput_Metadata[dpv_hi(LocalHistoryPixelPos)]      = dpv_hi(FinalHistoryValidity);
-				#endif
+				HistoryOutput_Metadata[dpv_hi(LocalHistoryPixelPos)]      = dpv_hi(FinalMetadata);
 				#if CONFIG_ACCUMULATE_TRANSLUCENCY_SEPARATLY
 					#if CONFIG_SCENE_COLOR_APLHA
 						HistoryOutput_ColorArray[tsr_short3(dpv_hi(LocalHistoryPixelPos), HistoryOutput_ArrayIndices_Translucency)] = dpv_hi(FinalTranslucencyColor);
@@ -1553,8 +1629,10 @@ void MainCS(
 				#endif
 				#if CONFIG_GRAND_REPROJECTION
 					HistoryOutput_ColorArray[tsr_short3(dpv_hi(LocalHistoryPixelPos), HistoryOutput_ArrayIndices_PrevHighFrequency)] = dpv_hi(FinalPrevHighFrequencyColor);
-					HistoryOutput_ColorArray[tsr_short3(dpv_hi(LocalHistoryPixelPos), HistoryOutput_ArrayIndices_PrevHighFrequencyResultant)] = dpv_hi(FinalPrevHighFrequencyColorResultant);
-					HistoryOutput_ColorArray[tsr_short3(dpv_hi(LocalHistoryPixelPos), HistoryOutput_ArrayIndices_HighFrequencyOverblur)] = dpv_hi(FinalOverblur);
+					#if CONFIG_GRAND_REPROJECTION == 2
+						HistoryOutput_ColorArray[tsr_short3(dpv_hi(LocalHistoryPixelPos), HistoryOutput_ArrayIndices_PrevHighFrequencyResultant)] = dpv_hi(FinalPrevHighFrequencyColorResultant);
+						HistoryOutput_ColorArray[tsr_short3(dpv_hi(LocalHistoryPixelPos), HistoryOutput_ArrayIndices_HighFrequencyOverblur)] = dpv_hi(FinalOverblur);
+					#endif
 				#endif
 			}
 			#endif // !CONFIG_ENABLE_DUAL_PIXEL_VECTORIZATION
diff --git a/Engine/Shaders/Public/DualPixelVectorization.ush b/Engine/Shaders/Public/DualPixelVectorization.ush
--- a/Engine/Shaders/Public/DualPixelVectorization.ush
+++ b/Engine/Shaders/Public/DualPixelVectorization.ush
@@ -302,6 +302,11 @@ DEFINE_DPV_FUNCTION_FOR_ALL_VECTORS(DEFINE_DPV_FUNCTION_VECTOR)
  *		half3x2 BlueishColor = dpv_mul(Color, half3(1.00, 0.25, 0.25));
  *		half3x2 PremultipliedColor = dpv_scale(Color, Coverage);
  */
+#define DEFINE_DPV_FUNCTION_SCALAR(TYPE) \
+	CALL_SITE_DEBUGLOC __DPV_MULTI_SCALAR(TYPE) dpv_scale(__DPV_MULTI_SCALAR(TYPE) a, __DPV_MULTI_SCALAR(TYPE) b) { \
+		return a * b; \
+	} \
+
 #define DEFINE_DPV_FUNCTION_VECTOR(TYPE, COMPONENTS) \
 	CALL_SITE_DEBUGLOC __DPV_MULTI_VECTOR(TYPE, COMPONENTS) dpv_mul(__DPV_MULTI_VECTOR(TYPE, COMPONENTS) a, __DPV_MULTI_VECTOR(TYPE, COMPONENTS) b) { \
 		return __DPV_FOREACH_COMPONENTS_CONSTRUCTOR(TYPE, COMPONENTS, __DPV_CODE_OPERATOR_VECTOR_VECTOR, a, *, b); \
@@ -324,8 +329,10 @@ DEFINE_DPV_FUNCTION_FOR_ALL_VECTORS(DEFINE_DPV_FUNCTION_VECTOR)
 	CALL_SITE_DEBUGLOC __DPV_MULTI_VECTOR(TYPE, COMPONENTS) dpv_scale(TYPE a, __DPV_MULTI_VECTOR(TYPE, COMPONENTS) b) { \
 		return __DPV_FOREACH_COMPONENTS_CONSTRUCTOR(TYPE, COMPONENTS, __DPV_CODE_OPERATOR_VECTOR_SCALAR, b, *, a); \
 	} \
-
+	
+DEFINE_DPV_FUNCTION_FOR_ALL_SCALARS(DEFINE_DPV_FUNCTION_SCALAR)
 DEFINE_DPV_FUNCTION_FOR_ALL_VECTORS(DEFINE_DPV_FUNCTION_VECTOR)
+#undef DEFINE_DPV_FUNCTION_SCALAR
 #undef DEFINE_DPV_FUNCTION_VECTOR
 
 
@@ -368,6 +375,14 @@ DEFINE_DPV_FUNCTION_FOR_ALL_VECTORS(DEFINE_DPV_FUNCTION_VECTOR)
  * 
  *		half2 Luma = dpv_dot(Color, half3(0.25, 0.5, 0.25));
  */
+#define DEFINE_DPV_FUNCTION_SCALAR(TYPE) \
+	CALL_SITE_DEBUGLOC __DPV_MULTI_SCALAR(TYPE) dpv_min(__DPV_MULTI_SCALAR(TYPE) a, __DPV_MULTI_SCALAR(TYPE) b) { \
+		return min(a, b); \
+	} \
+	CALL_SITE_DEBUGLOC __DPV_MULTI_SCALAR(TYPE) dpv_max(__DPV_MULTI_SCALAR(TYPE) a, __DPV_MULTI_SCALAR(TYPE) b) { \
+		return max(a, b); \
+	} \
+
 #define DEFINE_DPV_FUNCTION_VECTOR(TYPE, COMPONENTS) \
 	CALL_SITE_DEBUGLOC __DPV_MULTI_SCALAR(TYPE) dpv_dot(__DPV_MULTI_VECTOR(TYPE, COMPONENTS) a, __DPV_MULTI_VECTOR(TYPE, COMPONENTS) b) { \
 		return __DPV_FOREACH_COMPONENTS(COMPONENTS, __DPV_CODE_OPERATOR_VECTOR_VECTOR, +, a, *, b); \
@@ -399,8 +414,10 @@ DEFINE_DPV_FUNCTION_FOR_ALL_VECTORS(DEFINE_DPV_FUNCTION_VECTOR)
 	CALL_SITE_DEBUGLOC __DPV_MULTI_VECTOR(TYPE, COMPONENTS) dpv_max(__DPV_MULTI_VECTOR(TYPE, COMPONENTS) a, __DPV_VECTOR(TYPE, COMPONENTS) b) { \
 		return __DPV_FOREACH_COMPONENTS_CONSTRUCTOR(TYPE, COMPONENTS, __DPV_CODE_FUNCTION_2_ARGS, max, a, b); \
 	} \
-
+	
+DEFINE_DPV_FUNCTION_FOR_ALL_SCALARS(DEFINE_DPV_FUNCTION_SCALAR)
 DEFINE_DPV_FUNCTION_FOR_ALL_VECTORS(DEFINE_DPV_FUNCTION_VECTOR)
+#undef DEFINE_DPV_FUNCTION_SCALAR
 #undef DEFINE_DPV_FUNCTION_VECTOR
 
 
diff --git a/Engine/Source/Runtime/Renderer/Private/PostProcess/TemporalSuperResolution.cpp b/Engine/Source/Runtime/Renderer/Private/PostProcess/TemporalSuperResolution.cpp
--- a/Engine/Source/Runtime/Renderer/Private/PostProcess/TemporalSuperResolution.cpp
+++ b/Engine/Source/Runtime/Renderer/Private/PostProcess/TemporalSuperResolution.cpp
@@ -207,6 +207,9 @@ FTSRHistoryArrayIndices TranslateHistoryFormatBitsToArrayIndices(ETSRHistoryForm
 	ArrayIndices.Size = 1;
 	ArrayIndices.HighFrequency = 0;
 	ArrayIndices.Translucency = -1;
+	ArrayIndices.PrevHighFrequency = -1;
+	ArrayIndices.PrevHighFrequencyResultant = -1;
+	ArrayIndices.HighFrequencyOverblur = -1;
 
 	if (EnumHasAnyFlags(HistoryFormatBits, ETSRHistoryFormatBits::Translucency))
 	{
@@ -215,16 +218,9 @@ FTSRHistoryArrayIndices TranslateHistoryFormatBitsToArrayIndices(ETSRHistoryForm
 
 	if (EnumHasAnyFlags(HistoryFormatBits, ETSRHistoryFormatBits::GrandReprojection))
 	{
-		ArrayIndices.PrevHighFrequency = ArrayIndices.Size + 0;
-		ArrayIndices.PrevHighFrequencyResultant = ArrayIndices.Size + 1;
-		ArrayIndices.HighFrequencyOverblur = ArrayIndices.Size + 2;
-		ArrayIndices.Size += 3;
-	}
-	else
-	{
-		ArrayIndices.PrevHighFrequency = -1;
-		ArrayIndices.PrevHighFrequencyResultant = -1;
-		ArrayIndices.HighFrequencyOverblur = -1;
+		ArrayIndices.PrevHighFrequency = ArrayIndices.Size++;
+		//ArrayIndices.PrevHighFrequencyResultant = ArrayIndices.Size++;
+		//ArrayIndices.HighFrequencyOverblur = ArrayIndices.Size++;
 	}
 
 	return ArrayIndices;
@@ -241,11 +237,11 @@ FTSRHistorySRVs CreateSRVs(FRDGBuilder& GraphBuilder, const FTSRHistoryTextures&
 		SRVs.Translucency = GraphBuilder.CreateSRV(FRDGTextureSRVDesc::CreateForSlice(Textures.ColorArray, Textures.ArrayIndices.Translucency));
 		SRVs.TranslucencyAlpha = GraphBuilder.CreateSRV(FRDGTextureSRVDesc::Create(Textures.TranslucencyAlpha));
 	}
-	if (Textures.ArrayIndices.HighFrequencyOverblur >= 0)
+	if (Textures.ArrayIndices.PrevHighFrequency >= 0)
 	{
 		SRVs.PrevHighFrequency = GraphBuilder.CreateSRV(FRDGTextureSRVDesc::CreateForSlice(Textures.ColorArray, Textures.ArrayIndices.PrevHighFrequency));
-		SRVs.PrevHighFrequencyResultant = GraphBuilder.CreateSRV(FRDGTextureSRVDesc::CreateForSlice(Textures.ColorArray, Textures.ArrayIndices.PrevHighFrequencyResultant));
-		SRVs.HighFrequencyOverblur = GraphBuilder.CreateSRV(FRDGTextureSRVDesc::CreateForSlice(Textures.ColorArray, Textures.ArrayIndices.HighFrequencyOverblur));
+		//SRVs.PrevHighFrequencyResultant = GraphBuilder.CreateSRV(FRDGTextureSRVDesc::CreateForSlice(Textures.ColorArray, Textures.ArrayIndices.PrevHighFrequencyResultant));
+		//SRVs.HighFrequencyOverblur = GraphBuilder.CreateSRV(FRDGTextureSRVDesc::CreateForSlice(Textures.ColorArray, Textures.ArrayIndices.HighFrequencyOverblur));
 	}
 
 	SRVs.Guide = Textures.Guide;
@@ -664,7 +660,7 @@ class FTSRUpdateHistoryCS : public FTSRShader
 
 		SHADER_PARAMETER_STRUCT_INCLUDE(FTSRPrevHistoryParameters, PrevHistoryParameters)
 		SHADER_PARAMETER_STRUCT(FTSRHistorySRVs, PrevHistory)
-		SHADER_PARAMETER_RDG_TEXTURE(Texture2DArray, GrandPrevColorArray)
+		SHADER_PARAMETER_RDG_TEXTURE_SRV(Texture2D, GrandPrevColorTexture)
 
 		SHADER_PARAMETER_RDG_TEXTURE_UAV(RWTexture2D, SceneColorOutputMip0)
 		SHADER_PARAMETER_RDG_TEXTURE_UAV(RWTexture2D, SceneColorOutputMip1)
@@ -1154,7 +1150,14 @@ ITemporalUpscaler::FOutputs AddTemporalSuperResolutionPasses(
 			FClearValueBinding::None,
 			TexCreate_ShaderResource | TexCreate_UAV);
 
-		Desc.Format = PF_R8;
+		if (bGrandReprojection && History.ArrayIndices.PrevHighFrequencyResultant == -1)
+		{
+			Desc.Format = PF_R8G8;
+		}
+		else
+		{
+			Desc.Format = PF_R8;
+		}
 		History.Metadata = GraphBuilder.CreateTexture(Desc, TEXT("TSR.History.Metadata"));
 
 		Desc.Format = PF_R8;
@@ -1201,7 +1204,7 @@ ITemporalUpscaler::FOutputs AddTemporalSuperResolutionPasses(
 
 	// Setup the previous frame history
 	FRDGTextureRef PrevVelocity = BlackDummy;
-	FRDGTextureRef GrandPrevColorArray = BlackDummy;
+	FRDGTextureSRVRef GrandPrevColorTexture = GraphBuilder.CreateSRV(BlackDummy);
 	FTSRHistoryTextures PrevHistory;
 	FTSRHistorySRVs PrevHistorySRVs;
 	if (!bCameraCut)
@@ -1221,7 +1224,7 @@ ITemporalUpscaler::FOutputs AddTemporalSuperResolutionPasses(
 		PrevHistorySRVs = CreateSRVs(GraphBuilder, PrevHistory);
 
 		PrevVelocity = InputHistory.Velocity.IsValid() ? GraphBuilder.RegisterExternalTexture(InputHistory.Velocity) : PrevVelocity;
-		GrandPrevColorArray = InputHistory.PrevColorArray.IsValid() ? GraphBuilder.RegisterExternalTexture(InputHistory.PrevColorArray, TEXT("TSR.GrandPrevHistory.ColorArray")) : GrandPrevColorArray;
+		GrandPrevColorTexture = InputHistory.PrevColorArray.IsValid() ?  GraphBuilder.CreateSRV(FRDGTextureSRVDesc::CreateForSlice(GraphBuilder.RegisterExternalTexture(InputHistory.PrevColorArray, TEXT("TSR.GrandPrevHistory.ColorArray")), PrevHistory.ArrayIndices.HighFrequency)) : GrandPrevColorTexture;
 	}
 	else
 	{
@@ -1234,7 +1237,7 @@ ITemporalUpscaler::FOutputs AddTemporalSuperResolutionPasses(
 		// Setup prev history parameters.
 		FScreenPassTextureViewport PrevHistoryViewport(PrevHistorySRVs.HighFrequency->Desc.Texture->Desc.Extent, InputHistory.OutputViewportRect);
 		FScreenPassTextureViewport PrevSubpixelDetailsViewport(PrevHistorySRVs.SubpixelDetails->Desc.Extent, InputHistory.InputViewportRect);
-		FScreenPassTextureViewport GrandPrevSubpixelDetailsViewport(GrandPrevColorArray->Desc.Extent, InputHistory.PrevOutputViewportRect);
+		FScreenPassTextureViewport GrandPrevSubpixelDetailsViewport(GrandPrevColorTexture->Desc.Texture->Desc.Extent, InputHistory.PrevOutputViewportRect);
 
 		if (bCameraCut)
 		{
@@ -1691,7 +1694,7 @@ ITemporalUpscaler::FOutputs AddTemporalSuperResolutionPasses(
 
 		PassParameters->PrevHistoryParameters = PrevHistoryParameters;
 		PassParameters->PrevHistory = PrevHistorySRVs;
-		PassParameters->GrandPrevColorArray = GrandPrevColorArray;
+		PassParameters->GrandPrevColorTexture = GrandPrevColorTexture;
 
 		PassParameters->HistoryOutput = CreateUAVs(GraphBuilder, History);
 		if (HistorySize != OutputRect.Size())
