Fix motion vector bugs for non-nanite objects when the primitive moves but the instances do not.

#rnx

diff --git a/Engine/Shaders/Private/SceneData.ush b/Engine/Shaders/Private/SceneData.ush
--- a/Engine/Shaders/Private/SceneData.ush
+++ b/Engine/Shaders/Private/SceneData.ush
@@ -1020,7 +1020,15 @@ FInstanceSceneData GetInstanceSceneData(uint InstanceId, uint SOAStride, bool bC
 		}
 		else
 		{
-			InstanceData.PrevLocalToWorld		= MakeLWCMatrix(TilePosition, LocalToRelativeWorld);
+			// Since the instance is not dynamic, it cannot move relative to the primitive. However, the primitive itself can
+			// move. So we have to solve for previous Instance->World by solving primitive World->PrevWorld and using it to
+			// transform Instance->World into previous world space.
+			//
+			// TODO: Cut the number of matrix multiplies down by adding World->PrevWorld to FPrimitiveSceneData?
+			// Also, in the case of there being a single instance with identity Instance->Primitive (common), this can be
+			// simplified to merely copying the Primitive PrevLocalToWorld
+			float4x4 PrimRelWorldToPrevRelWorld = mul(PrimitiveData.WorldToLocal.M, PrimitiveData.PreviousLocalToWorld.M);
+			InstanceData.PrevLocalToWorld = MakeLWCMatrix(TilePosition, mul(LocalToRelativeWorld, PrimRelWorldToPrevRelWorld));
 		}
 
 	#if 1 //NEEDS_LIGHTMAP_COORDINATE
