Works arround a shader compiler bug

diff --git a/Engine/Shaders/Private/TemporalSuperResolution/TSRUpdateHistory.usf b/Engine/Shaders/Private/TemporalSuperResolution/TSRUpdateHistory.usf
--- a/Engine/Shaders/Private/TemporalSuperResolution/TSRUpdateHistory.usf
+++ b/Engine/Shaders/Private/TemporalSuperResolution/TSRUpdateHistory.usf
@@ -812,9 +812,10 @@ void MainCS(
 				UNROLL_N(BICUBIC_CATMULL_ROM_SAMPLES)
 				for (uint i = 0; i < BICUBIC_CATMULL_ROM_SAMPLES; i++)
 				{
-					tsr_halfCx2 SampleHighFrequency = dpv_force_interleave_registers_array(RawHighFrequency[i]);
-					tsr_halfMx2 SampleMetadata = dpv_force_interleave_registers_array(RawMetadata[i]);
-					tsr_half4x2 SampleTranslucency = dpv_force_interleave_registers(
+					// TODO: Should use dpv_force_interleave_registers_array() but there is a shader compiler bug
+					tsr_halfCx2 SampleHighFrequency = dpv_interleave_registers_array(RawHighFrequency[i]);
+					tsr_halfMx2 SampleMetadata = dpv_interleave_registers_array(RawMetadata[i]);
+					tsr_half4x2 SampleTranslucency = dpv_interleave_registers(
 						tsr_half4(RawTranslucencyColor[i][0], RawTranslucencyAlpha[i][0]),
 						tsr_half4(RawTranslucencyColor[i][1], RawTranslucencyAlpha[i][1]));
 					
