Fix nan propagation in TSR RejectShading pass when alpha channel is enabled

diff --git a/Engine/Shaders/Private/TemporalSuperResolution/TSRDecimateHistory.usf b/Engine/Shaders/Private/TemporalSuperResolution/TSRDecimateHistory.usf
--- a/Engine/Shaders/Private/TemporalSuperResolution/TSRDecimateHistory.usf
+++ b/Engine/Shaders/Private/TemporalSuperResolution/TSRDecimateHistory.usf
@@ -346,6 +346,9 @@ void MainCS(
 				tsr_halfC LinearReprojectedHistoryGuideColor = ReprojectedHistoryGuideColor * rcp(tsr_half(1.0) - ReprojectedHistoryGuideColor) * tsr_half(HistoryPreExposureCorrection);
 				
 				ReprojectedHistoryGuideColor = LinearReprojectedHistoryGuideColor * rcp(LinearReprojectedHistoryGuideColor + tsr_half(1.0));
+
+				// Removes NaN
+				ReprojectedHistoryGuideColor = -min(-ReprojectedHistoryGuideColor, tsr_half(0.0));
 			}
 
 			#if CONFIG_MOIRE_REPROJECTION
diff --git a/Engine/Shaders/Private/TemporalSuperResolution/TSRRejectShading.usf b/Engine/Shaders/Private/TemporalSuperResolution/TSRRejectShading.usf
--- a/Engine/Shaders/Private/TemporalSuperResolution/TSRRejectShading.usf
+++ b/Engine/Shaders/Private/TemporalSuperResolution/TSRRejectShading.usf
@@ -548,6 +548,12 @@ void MainCS(
 	}
 	#endif
 	
+	#if POST_PROCESS_ALPHA
+		OriginalInput = min(OriginalInput, tsr_simd_halfC::Const(LargestSceneColorRGBA));
+	#else
+		OriginalInput = min(OriginalInput, tsr_simd_halfC::Const(LargestSceneColorRGB));
+	#endif
+
 	// Convert input to LDR
 	OriginalOpaqueInput = OriginalOpaqueInput * rcp(OriginalOpaqueInput + tsr_simd_halfC::Const(1.0));
 	OriginalInput = OriginalInput * rcp(OriginalInput + tsr_simd_halfC::Const(1.0));
