Fixed broken fog rendering on mobile

diff --git a/Engine/Shaders/Private/SceneTexturesCommon.ush b/Engine/Shaders/Private/SceneTexturesCommon.ush
--- a/Engine/Shaders/Private/SceneTexturesCommon.ush
+++ b/Engine/Shaders/Private/SceneTexturesCommon.ush
@@ -200,7 +200,7 @@ float LookupDeviceZ( float2 ScreenUV )
 	return VulkanSubpassDepthFetch();
 #elif MOBILE_DEPTHFECTH && (METAL_PROFILE && !MAC)
 	return DepthbufferFetchES2();
-#elif !MOBILE_DEFERRED_SHADING
+#elif (POST_PROCESS_MATERIAL && !MOBILE_DEFERRED_SHADING)
 	// SceneDepth texture is not accessible during post-processing as we discard it at the end of mobile BasePass (with forward shading) 
 	// instead fetch DeviceZ from SceneDepthAuxTexture
 	return Texture2DSampleLevel(MobileSceneTextures.SceneDepthAuxTexture, MobileSceneTextures.SceneDepthAuxTextureSampler, ScreenUV, 0).r;
