Add strata overview debug info mode.

This summarize most of Strata settings into a single place for quite debugging.

This is mapped onto r.Strata.ViewMode 7

diff --git a/Engine/Shaders/Private/Strata/StrataVisualize.usf b/Engine/Shaders/Private/Strata/StrataVisualize.usf
--- a/Engine/Shaders/Private/Strata/StrataVisualize.usf
+++ b/Engine/Shaders/Private/Strata/StrataVisualize.usf
@@ -20,6 +20,18 @@
 // This is for the fxc compiler to compile faster.
 #define STRATA_VIZ_MAX_BSDF 2
 
+void PrintBool(inout FShaderPrintContext Ctx, bool In)
+{
+	if (In)
+	{
+		Print(Ctx, TEXT("Yes  "), FontGreen);
+	}
+	else
+	{
+		Print(Ctx, TEXT("No   "), FontRed);
+	}
+}
+
 void DrawReferentialTWS(float3 P /*In Translated World Space*/, float3 X, float3 Y, float3 N, float3 InColor)
 {
 	const float Size = 10.f;
@@ -91,22 +103,6 @@ void DrawPixelFootprint(float3 P, float3 dPdx, float3 dPdy, uint2 PixelCoord)
 uint BSDFIndex;
 RWBuffer<uint> RWPositionOffsetBuffer;
 
-// View mode:
-// 1: Plot a single pixel
-// 2: Color code the number of BSDFs (green -> red)
-// 3: Add layer visualization
-void PrintBool(inout FShaderPrintContext Ctx, bool In, FFontColor InColor)
-{
-	if (In)
-	{
-		Print(Ctx, TEXT("Yes  "), FontGreen);
-	}
-	else
-	{
-		Print(Ctx, TEXT("No   "), FontRed);
-	}
-}
-
 void PrintPixelType(inout FShaderPrintContext Ctx, in FStrataPixelHeader Header, FFontColor InColor)
 {
 	if (IsSimpleMaterial(Header))			Print(Ctx, TEXT("Simple  "), InColor);
@@ -183,12 +179,12 @@ void PrintBSDF(inout FShaderPrintContext Context, inout FStrataAddressing Strata
 			Newline(Context, RectMax);
 
 			Print(Context, TEXT(" "), FontSilver);
-			Print(Context, BSDF_GETSHAREDLOCALBASISID(BSDF), FontBSDF, 2, 1);			Print(Context, TEXT("       "), FontSilver);
-			PrintBool(Context, BSDF_GETHASANISOTROPY(BSDF), FontBSDF);					Print(Context, TEXT(" "), FontSilver);
-			PrintBool(Context, BSDF_GETISTOPLAYER(BSDF), FontBSDF);						Print(Context, TEXT("    "), FontSilver);
-			PrintBool(Context, BSDF_GETSSSTYPE(BSDF) != SSS_TYPE_INVALID, FontBSDF);	Print(Context, TEXT("      "), FontSilver);
-			PrintBool(Context, BSDF_GETISSIMPLEVOLUME(BSDF), FontBSDF);					Print(Context, TEXT("        "), FontSilver);
-			PrintBool(Context, BSDF_GETISTHIN(BSDF), FontBSDF);							Print(Context, TEXT("  "), FontSilver);
+			Print(Context, BSDF_GETSHAREDLOCALBASISID(BSDF), FontBSDF, 2, 1);	Print(Context, TEXT("       "), FontSilver);
+			PrintBool(Context, BSDF_GETHASANISOTROPY(BSDF));					Print(Context, TEXT(" "));
+			PrintBool(Context, BSDF_GETISTOPLAYER(BSDF));						Print(Context, TEXT("    "));
+			PrintBool(Context, BSDF_GETSSSTYPE(BSDF) != SSS_TYPE_INVALID);		Print(Context, TEXT("      "));
+			PrintBool(Context, BSDF_GETISSIMPLEVOLUME(BSDF));					Print(Context, TEXT("        "));
+			PrintBool(Context, BSDF_GETISTHIN(BSDF));							Print(Context, TEXT("  "));
 
 			// View Weight
 			{
@@ -306,8 +302,8 @@ void PrintBSDF(inout FShaderPrintContext Context, inout FStrataAddressing Strata
 
 			Print(Context, TEXT(" "), FontSilver);
 			Print(Context, BSDF_GETSHAREDLOCALBASISID(BSDF), FontBSDF);
-			PrintBool(Context, BSDF_GETHASANISOTROPY(BSDF), FontBSDF);
-			PrintBool(Context, BSDF_GETISTOPLAYER(BSDF), FontBSDF);
+			PrintBool(Context, BSDF_GETHASANISOTROPY(BSDF));
+			PrintBool(Context, BSDF_GETISTOPLAYER(BSDF));
 			if (BSDF_GETHASGREYWEIGHT_V(BSDF))
 			{
 				Print(Context, BSDF.LuminanceWeightV.x, FontSilver);
@@ -338,7 +334,7 @@ void PrintBSDF(inout FShaderPrintContext Context, inout FStrataAddressing Strata
 
 			Print(Context, TEXT(" "), FontSilver);
 			Print(Context, BSDF_GETSHAREDLOCALBASISID(BSDF), FontBSDF);
-			PrintBool(Context, BSDF_GETISTOPLAYER(BSDF), FontBSDF);
+			PrintBool(Context, BSDF_GETISTOPLAYER(BSDF));
 			if (BSDF_GETHASGREYWEIGHT_V(BSDF))
 			{
 				Print(Context, BSDF.LuminanceWeightV.x, FontSilver);
@@ -477,11 +473,11 @@ void PrintPixel(uint2 InCoord, float3 InWorldPosition, float3 V)
 		Print(Context, IrradianceAndOcclusion.MaterialAO, FontHeaderPropValue, 5, 3);					Print(Context, TEXT(" "));
 		Print(Context, IrradianceAndOcclusion.IndirectIrradiance, FontHeaderPropValue, 6, 3);			Print(Context, TEXT(" "));
 		Print(Context, TopLayerData.Roughness, FontHeaderPropValue, 5, 3);								Print(Context, TEXT("        "));
-		PrintBool(Context, HasPrecShadowMask(Header), FontHeaderPropValue);								Print(Context, TEXT("     "));
-		PrintBool(Context, HasZeroPrecShadowMask(Header), FontHeaderPropValue);							Print(Context, TEXT("      "));
-		PrintBool(Context, DoesCastContactShadow(Header), FontHeaderPropValue);							Print(Context, TEXT("        "));
-		PrintBool(Context, HasDynamicIndirectShadowCasterRepresentation(Header), FontHeaderPropValue);	Print(Context, TEXT("        "));
-		PrintBool(Context, HasSubsurface(Header), FontHeaderPropValue);
+		PrintBool(Context, HasPrecShadowMask(Header));													Print(Context, TEXT("     "));
+		PrintBool(Context, HasZeroPrecShadowMask(Header));												Print(Context, TEXT("      "));
+		PrintBool(Context, DoesCastContactShadow(Header));												Print(Context, TEXT("        "));
+		PrintBool(Context, HasDynamicIndirectShadowCasterRepresentation(Header));						Print(Context, TEXT("        "));
+		PrintBool(Context, HasSubsurface(Header));
 		Newline(Context, RectMax);
 		Newline(Context);
 		
@@ -587,7 +583,7 @@ void MaterialPrintInfoCS(uint3 DispatchThreadId : SV_DispatchThreadID)
 ///////////////////////////////////////////////////////////////////////////////////////////////////
 // Material Visualize
 
-#if SHADER_MATERIALVISUALIZE
+#if SHADER_MATERIALCOUNT
 
 uint ViewMode;
 
@@ -607,29 +603,33 @@ void VisualizeMaterialPS(
 
 	FStrataAddressing StrataAddressing = GetStrataPixelDataByteOffset(PixelPosDynRes, uint2(View.BufferSizeAndInvSize.xy), Strata.MaxBytesPerPixel);
 	FStrataPixelHeader Header = UnpackStrataHeaderIn(Strata.MaterialTextureArray, StrataAddressing, Strata.TopLayerTexture);
-	bool bStrataMaterial = Header.BSDFCount > 0;
+	const bool bStrataMaterial = Header.BSDFCount > 0;
 
 	OutColor = 0;
+	OutColor.w = 0.1f;
 
 	// BSDF count
-	if (ViewMode == 2)
+	const float MaxBSDFCount = 4.f;
+	if (Header.BSDFCount > 0)
 	{
-		OutColor.w = 0.1f;
-		if (Header.BSDFCount > 0)
-		{
-			const float BSDFCost = saturate(float(Header.BSDFCount) / 4.f);
-			OutColor.xyz = ColorMapTurbo(BSDFCost);
-		}
+		const float BSDFCost = saturate(float(Header.BSDFCount) / MaxBSDFCount);
+		OutColor.xyz = ColorMapTurbo(BSDFCost);
 	}
-	else
+
+	// Print Legend
+	if (all(uint2(SVPos.xy) == uint2(0,0)))
 	{
-		OutColor = float4(0.0f, 0.0f, 0.0f, 1.0f);
+		FShaderPrintContext Context = InitShaderPrintContext(true, uint2(50, 50));
+		Print(Context, TEXT("1 BSDF"), InitFontColor(ColorMapTurbo(1 / MaxBSDFCount))); Newline(Context);
+		Print(Context, TEXT("2 BSDF"), InitFontColor(ColorMapTurbo(2 / MaxBSDFCount))); Newline(Context);
+		Print(Context, TEXT("3 BSDF"), InitFontColor(ColorMapTurbo(3 / MaxBSDFCount))); Newline(Context);
+		Print(Context, TEXT("4 BSDF"), InitFontColor(ColorMapTurbo(4 / MaxBSDFCount))); Newline(Context);
 	}
 }
-#endif // SHADER_MATERIALVISUALIZE
+#endif // SHADER_MATERIALCOUNT
 
 ///////////////////////////////////////////////////////////////////////////////////////////////////
-// Material Visualize
+// Advanced Material Visualize
 
 #define BSDF_ELEMENT_SIZE 210
 #define BSDF_ELEMENT_STARTX 100.0f + 450.0f
@@ -1407,3 +1407,176 @@ void MaterialDebugStrataTreePS(
 }
 
 #endif // SHADER_DEBUGSTRATATREE_PS
+
+
+///////////////////////////////////////////////////////////////////////////////////////////////////
+// Material Print
+
+#if SHADER_SYSTEMINFO
+uint bAdvancedDebugEnabled;
+uint bEnergyConservation;
+uint bEnergyPreservation;
+uint bDbufferPass;
+uint ClassificationCMask;
+uint ClassificationAsync;
+uint Classification8bits;
+uint bRoughRefraction;
+float LumenTileOverflow;
+Buffer<uint> ClassificationTileDrawIndirectBuffer;
+
+uint GetTileCount(uint InType)
+{
+	return ClassificationTileDrawIndirectBuffer[InType * 4];
+}
+
+[numthreads(1, 1, 1)]
+void MainCS(uint3 DispatchThreadId : SV_DispatchThreadID)
+{
+	// Print Legend
+	if (all(DispatchThreadId == 0))
+	{
+		FFontColor TitleColor  = FontWhite;
+		FFontColor LegendColor = FontSilver;
+		FFontColor ValueColor  = FontOrange;
+		FShaderPrintContext Context = InitShaderPrintContext(true, uint2(50, 50));
+
+		// General
+		Print(Context, TEXT("General"), TitleColor);
+		Newline(Context);
+
+		Print(Context, TEXT(" Max byte per pixel : "), LegendColor); 
+		Print(Context, Strata.MaxBytesPerPixel, ValueColor);
+		Newline(Context);
+
+		Print(Context, TEXT(" Shading quality    : "), LegendColor);
+		Print(Context, uint(STRATA_SHADING_QUALITY), ValueColor);
+		Newline(Context);
+
+		Print(Context, TEXT(" Rough diffuse      : "), LegendColor);
+		PrintBool(Context, Strata.bRoughDiffuse);
+		Newline(Context);
+
+		Print(Context, TEXT(" Energy conservation: "), LegendColor);
+		PrintBool(Context, bEnergyConservation);
+		Newline(Context);
+
+		Print(Context, TEXT(" Energy preservation: "), LegendColor);
+		PrintBool(Context, bEnergyPreservation);
+		Newline(Context);
+
+		Print(Context, TEXT(" Lumen overflow     : "), LegendColor);
+		Print(Context, LumenTileOverflow, ValueColor, 3, 1);
+		Newline(Context);
+
+		Print(Context, TEXT(" DBuffer pass       : "), LegendColor);
+		PrintBool(Context, bDbufferPass);
+		Newline(Context);
+
+		Print(Context, TEXT(" Rough refraction   : "), LegendColor);
+		PrintBool(Context, bRoughRefraction);
+		Newline(Context);
+
+		Newline(Context);
+
+		// Classification
+		Print(Context, TEXT("Classification"), TitleColor);
+		Newline(Context);
+
+		Print(Context, TEXT(" Tile Size          : "), LegendColor);
+		Print(Context, Strata.TileSize, ValueColor, 1, 0);
+		Print(Context, TEXT("x"), LegendColor);
+		Print(Context, Strata.TileSize, ValueColor, 1, 0);
+		Newline(Context);
+
+		Print(Context, TEXT(" Tile Count         : "), LegendColor);
+		Print(Context, Strata.TileCount.x, ValueColor, 3, 0);
+		Print(Context, TEXT("x"), LegendColor);
+		Print(Context, Strata.TileCount.y, ValueColor, 3, 0);
+		Newline(Context);
+
+		Print(Context, TEXT(" CMask              : "), LegendColor);
+		PrintBool(Context, ClassificationCMask);
+		Newline(Context);
+
+		Print(Context, TEXT(" Async              : "), LegendColor);
+		PrintBool(Context, ClassificationAsync);
+		Newline(Context);
+
+		Print(Context, TEXT(" 8-bits coord       : "), LegendColor);
+		PrintBool(Context, Classification8bits);
+		Newline(Context);
+
+		Newline(Context);
+
+		// Material tiles
+		Print(Context, TEXT("Material Tiles"), TitleColor);
+		Newline(Context);
+
+		Print(Context, TEXT(" # Simple tiles     : "), LegendColor);
+		Print(Context, GetTileCount(STRATA_TILE_TYPE_SIMPLE), ValueColor);
+		Newline(Context);
+
+		Print(Context, TEXT(" # Single tiles     : "), LegendColor);
+		Print(Context, GetTileCount(STRATA_TILE_TYPE_SINGLE), ValueColor);
+		Newline(Context);
+
+		Print(Context, TEXT(" # Complex tiles    : "), LegendColor);
+		Print(Context, GetTileCount(STRATA_TILE_TYPE_COMPLEX), ValueColor);
+		Newline(Context);
+
+		Newline(Context);
+
+		// Rough tiles
+		if (bRoughRefraction)
+		{
+			Print(Context, TEXT("Rough Refraction Tiles"), TitleColor);
+			Newline(Context);
+
+			Print(Context, TEXT(" # Rough refract.   : "), LegendColor);
+			Print(Context, GetTileCount(STRATA_TILE_TYPE_ROUGH_REFRACT), ValueColor);
+			Newline(Context);
+
+			Print(Context, TEXT(" # Rough r. w/o SSS : "), LegendColor);
+			Print(Context, GetTileCount(STRATA_TILE_TYPE_ROUGH_REFRACT_SSS_WITHOUT), ValueColor);
+			Newline(Context);
+
+			Newline(Context);
+		}
+
+		// DBuffer Tiles
+		if (bDbufferPass)
+		{
+			Print(Context, TEXT("Decal Tiles"), TitleColor);
+			Newline(Context);
+
+			Print(Context, TEXT(" # Simple tiles     : "), LegendColor);
+			Print(Context, GetTileCount(STRATA_TILE_TYPE_DECAL_SIMPLE), ValueColor);
+			Newline(Context);
+
+			Print(Context, TEXT(" # Single tiles     : "), LegendColor);
+			Print(Context, GetTileCount(STRATA_TILE_TYPE_DECAL_SINGLE), ValueColor);
+			Newline(Context);
+
+			Print(Context, TEXT(" # Complex tiles    : "), LegendColor);
+			Print(Context, GetTileCount(STRATA_TILE_TYPE_DECAL_COMPLEX), ValueColor);
+			Newline(Context);
+
+			Newline(Context);
+		}
+
+		// Debug
+		Print(Context, TEXT(" Debug"), TitleColor);
+		Newline(Context);
+
+		Print(Context, TEXT(" Adv. debug         : "), LegendColor);
+		PrintBool(Context, bAdvancedDebugEnabled);
+		Newline(Context);
+
+		Print(Context, TEXT(" Layer peel index   : "), LegendColor);
+		Print(Context, Strata.PeelLayersAboveDepth, ValueColor);
+		Newline(Context);
+
+		Newline(Context);
+	}
+}
+#endif // SHADER_SYSTEMINFO
\ No newline at end of file
diff --git a/Engine/Source/Runtime/Engine/Private/StrataVisualizationData.cpp b/Engine/Source/Runtime/Engine/Private/StrataVisualizationData.cpp
--- a/Engine/Source/Runtime/Engine/Private/StrataVisualizationData.cpp
+++ b/Engine/Source/Runtime/Engine/Private/StrataVisualizationData.cpp
@@ -133,6 +133,14 @@ void FStrataVisualizationData::Initialize()
 				true);
 		}
 
+		AddVisualizationMode(
+			ModeMap,
+			TEXT("StrataInfo"),
+			LOCTEXT("StrataInfo", "Strata Info"),
+			LOCTEXT("StrataInfoDesc", "Visualizes Strata info"),
+			FViewMode::StrataInfo,
+			true);
+
 		ConsoleDocumentationVisualizationMode = ConfigureConsoleCommand(ModeMap);
 	}
 	bIsInitialized = true;
diff --git a/Engine/Source/Runtime/Engine/Public/StrataVisualizationData.h b/Engine/Source/Runtime/Engine/Public/StrataVisualizationData.h
--- a/Engine/Source/Runtime/Engine/Public/StrataVisualizationData.h
+++ b/Engine/Source/Runtime/Engine/Public/StrataVisualizationData.h
@@ -16,6 +16,7 @@ public:
 		MaterialClassification,
 		DecalClassification,
 		RoughRefractionClassification,
+		StrataInfo,
 	};
 
 	/** Describes a single available visualization mode. */
diff --git a/Engine/Source/Runtime/Renderer/Private/Strata/Strata.cpp b/Engine/Source/Runtime/Renderer/Private/Strata/Strata.cpp
--- a/Engine/Source/Runtime/Renderer/Private/Strata/Strata.cpp
+++ b/Engine/Source/Runtime/Renderer/Private/Strata/Strata.cpp
@@ -130,10 +130,6 @@ const TCHAR* ToString(EStrataTileType Type)
 	return TEXT("Unknown");
 }
 
-FORCEINLINE bool SupportsCMask(const FStaticShaderPlatform InPlatform)
-{
-	return CVarUseCmaskClear.GetValueOnRenderThread() > 0 && FDataDrivenShaderPlatformInfo::GetSupportsRenderTargetWriteMask(InPlatform);
-}
 
 namespace Strata
 {
@@ -154,6 +150,11 @@ enum EStrataTileSpace
 	StrataTileSpace_Overflow = 2u
 };
 
+float GetStrataTileOverflowRatio()
+{
+	return FMath::Clamp(CVarStrataTileOverflow.GetValueOnRenderThread(), 0.f, 4.0f);
+}
+
 static FIntPoint GetStrataTextureTileResolution(const FIntPoint& InResolution, uint32 InSpace)
 {
 	FIntPoint Out = InResolution;
@@ -165,7 +166,7 @@ static FIntPoint GetStrataTextureTileResolution(const FIntPoint& InResolution, u
 	}
 	if (InSpace & EStrataTileSpace::StrataTileSpace_Overflow)
 	{
-		const float OverflowRatio = FMath::Clamp(CVarStrataTileOverflow.GetValueOnRenderThread(), 0.f, 4.0f);
+		const float OverflowRatio = GetStrataTileOverflowRatio();
 		Out.Y += FMath::DivideAndRoundUp(FMath::CeilToInt(InResolution.Y * OverflowRatio), STRATA_TILE_SIZE);
 	}
 	return Out;
@@ -184,10 +185,25 @@ FIntPoint GetStrataTextureResolution(const FIntPoint& InResolution)
 
 static void BindStrataGlobalUniformParameters(FRDGBuilder& GraphBuilder, FStrataViewData* StrataViewData, FStrataGlobalUniformParameters& OutStrataUniformParameters);
 
+bool SupportsCMask(const FStaticShaderPlatform InPlatform)
+{
+	return CVarUseCmaskClear.GetValueOnRenderThread() > 0 && FDataDrivenShaderPlatformInfo::GetSupportsRenderTargetWriteMask(InPlatform);
+}
+
+bool IsClassificationCoord8bits()
+{
+	return CVarStrataTileCoord8Bits.GetValueOnRenderThread() == 1;
+}
+
+bool IsClassificationAsync()
+{
+	return CVarStrataAsyncClassification.GetValueOnRenderThread() > 0;
+}
+
 static EPixelFormat GetClassificationTileFormat(const FIntPoint& InResolution)
 {
 	// For platform which whose resolution is never above 1080p, use 8bit tile format for performance.
-	const bool bRequest8bit = CVarStrataTileCoord8Bits.GetValueOnRenderThread() == 1;
+	const bool bRequest8bit = IsClassificationCoord8bits();
 	if (bRequest8bit)
 	{
 		check(InResolution.X <= 2048 && InResolution.Y <= 2048);
@@ -1120,7 +1136,7 @@ void AddStrataMaterialClassificationPass(FRDGBuilder& GraphBuilder, const FMinim
 	}
 
 	// Optionally run tile classification in async compute
-	const ERDGPassFlags PassFlags = CVarStrataAsyncClassification.GetValueOnRenderThread() > 0 ? ERDGPassFlags::AsyncCompute : ERDGPassFlags::Compute;
+	const ERDGPassFlags PassFlags = IsClassificationAsync() ? ERDGPassFlags::AsyncCompute : ERDGPassFlags::Compute;
 
 	for (int32 i = 0; i < Views.Num(); ++i)
 	{
diff --git a/Engine/Source/Runtime/Renderer/Private/Strata/StrataVisualize.cpp b/Engine/Source/Runtime/Renderer/Private/Strata/StrataVisualize.cpp
--- a/Engine/Source/Runtime/Renderer/Private/Strata/StrataVisualize.cpp
+++ b/Engine/Source/Runtime/Renderer/Private/Strata/StrataVisualize.cpp
@@ -76,10 +76,10 @@ class FMaterialPrintInfoCS : public FGlobalShader
 };
 IMPLEMENT_GLOBAL_SHADER(FMaterialPrintInfoCS, "/Engine/Private/Strata/StrataVisualize.usf", "MaterialPrintInfoCS", SF_Compute);
 
-class FVisualizeMaterialPS : public FGlobalShader
+class FVisualizeMaterialCountPS : public FGlobalShader
 {
-	DECLARE_GLOBAL_SHADER(FVisualizeMaterialPS);
-	SHADER_USE_PARAMETER_STRUCT(FVisualizeMaterialPS, FGlobalShader);
+	DECLARE_GLOBAL_SHADER(FVisualizeMaterialCountPS);
+	SHADER_USE_PARAMETER_STRUCT(FVisualizeMaterialCountPS, FGlobalShader);
 
 	using FPermutationDomain = TShaderPermutationDomain<>;
 
@@ -108,10 +108,56 @@ class FVisualizeMaterialPS : public FGlobalShader
 
 		// Stay debug and skip optimizations to reduce compilation time on this long shader.
 		OutEnvironment.CompilerFlags.Add(CFLAG_Debug);
-		OutEnvironment.SetDefine(TEXT("SHADER_MATERIALVISUALIZE"), 1);
+		OutEnvironment.SetDefine(TEXT("SHADER_MATERIALCOUNT"), 1);
 	}
 };
-IMPLEMENT_GLOBAL_SHADER(FVisualizeMaterialPS, "/Engine/Private/Strata/StrataVisualize.usf", "VisualizeMaterialPS", SF_Pixel);
+IMPLEMENT_GLOBAL_SHADER(FVisualizeMaterialCountPS, "/Engine/Private/Strata/StrataVisualize.usf", "VisualizeMaterialPS", SF_Pixel);
+
+
+class FStrataSystemInfoCS : public FGlobalShader
+{
+	DECLARE_GLOBAL_SHADER(FStrataSystemInfoCS);
+	SHADER_USE_PARAMETER_STRUCT(FStrataSystemInfoCS, FGlobalShader);
+
+	using FPermutationDomain = TShaderPermutationDomain<>;
+
+	BEGIN_SHADER_PARAMETER_STRUCT(FParameters, )
+		SHADER_PARAMETER(uint32, bAdvancedDebugEnabled)
+		SHADER_PARAMETER(uint32, bEnergyConservation)
+		SHADER_PARAMETER(uint32, bEnergyPreservation)
+		SHADER_PARAMETER(uint32, bDbufferPass)
+		SHADER_PARAMETER(uint32, ClassificationCMask)
+		SHADER_PARAMETER(uint32, ClassificationAsync)
+		SHADER_PARAMETER(uint32, Classification8bits)
+		SHADER_PARAMETER(uint32, bRoughRefraction)
+		SHADER_PARAMETER(float, LumenTileOverflow)
+		SHADER_PARAMETER_RDG_BUFFER_SRV(Buffer<uint>, ClassificationTileDrawIndirectBuffer)
+		SHADER_PARAMETER_STRUCT_REF(FViewUniformShaderParameters, ViewUniformBuffer)
+		SHADER_PARAMETER_RDG_UNIFORM_BUFFER(FStrataGlobalUniformParameters, Strata)
+		SHADER_PARAMETER_STRUCT_INCLUDE(FSceneTextureParameters, SceneTextures)
+		SHADER_PARAMETER_STRUCT_INCLUDE(ShaderPrint::FShaderParameters, ShaderPrintParameters)
+	END_SHADER_PARAMETER_STRUCT()
+
+		static FPermutationDomain RemapPermutation(FPermutationDomain PermutationVector)
+	{
+		return PermutationVector;
+	}
+
+	static bool ShouldCompilePermutation(const FGlobalShaderPermutationParameters& Parameters)
+	{
+		return Strata::IsStrataEnabled() && StrataDebugVisualizationCanRunOnPlatform(Parameters.Platform) && EnumHasAllFlags(Parameters.Flags, EShaderPermutationFlags::HasEditorOnlyData);
+	}
+
+	static void ModifyCompilationEnvironment(const FGlobalShaderPermutationParameters& Parameters, FShaderCompilerEnvironment& OutEnvironment)
+	{
+		FGlobalShader::ModifyCompilationEnvironment(Parameters, OutEnvironment);
+
+		// Stay debug and skip optimizations to reduce compilation time on this long shader.
+		OutEnvironment.CompilerFlags.Add(CFLAG_Debug);
+		OutEnvironment.SetDefine(TEXT("SHADER_SYSTEMINFO"), 1);
+	}
+};
+IMPLEMENT_GLOBAL_SHADER(FStrataSystemInfoCS, "/Engine/Private/Strata/StrataVisualize.usf", "MainCS", SF_Compute);
 
 class FMaterialDebugStrataTreeCS : public FGlobalShader
 {
@@ -222,7 +268,7 @@ static void AddVisualizeMaterialCountPasses(FRDGBuilder & GraphBuilder, const FV
 	FRDGTextureRef SceneColorTexture = ScreenPassSceneColor.Texture;
 	FRHIBlendState* PreMultipliedColorTransmittanceBlend = TStaticBlendState<CW_RGB, BO_Add, BF_One, BF_SourceAlpha, BO_Add, BF_Zero, BF_One>::GetRHI();
 
-	FVisualizeMaterialPS::FParameters* PassParameters = GraphBuilder.AllocParameters<FVisualizeMaterialPS::FParameters>();
+	FVisualizeMaterialCountPS::FParameters* PassParameters = GraphBuilder.AllocParameters<FVisualizeMaterialCountPS::FParameters>();
 	PassParameters->ViewUniformBuffer = View.ViewUniformBuffer;
 	PassParameters->ViewMode = 2;
 	PassParameters->Strata = Strata::BindStrataGlobalUniformParameters(View);
@@ -230,10 +276,42 @@ static void AddVisualizeMaterialCountPasses(FRDGBuilder & GraphBuilder, const FV
 	PassParameters->RenderTargets[0] = FRenderTargetBinding(SceneColorTexture, ERenderTargetLoadAction::ELoad);
 	ShaderPrint::SetParameters(GraphBuilder, View.ShaderPrintData, PassParameters->ShaderPrintParameters);
 
-	FVisualizeMaterialPS::FPermutationDomain PermutationVector;
-	TShaderMapRef<FVisualizeMaterialPS> PixelShader(View.ShaderMap, PermutationVector);
+	FVisualizeMaterialCountPS::FPermutationDomain PermutationVector;
+	TShaderMapRef<FVisualizeMaterialCountPS> PixelShader(View.ShaderMap, PermutationVector);
 
-	FPixelShaderUtils::AddFullscreenPass<FVisualizeMaterialPS>(GraphBuilder, View.ShaderMap, RDG_EVENT_NAME("Strata::VisualizeMaterial(Draw)"), PixelShader, PassParameters, ScreenPassSceneColor.ViewRect, PreMultipliedColorTransmittanceBlend);
+	FPixelShaderUtils::AddFullscreenPass<FVisualizeMaterialCountPS>(GraphBuilder, View.ShaderMap, RDG_EVENT_NAME("Strata::VisualizeMaterial(Draw)"), PixelShader, PassParameters, ScreenPassSceneColor.ViewRect, PreMultipliedColorTransmittanceBlend);
+}
+
+float GetStrataTileOverflowRatio();
+bool IsClassificationCoord8bits();
+bool IsClassificationAsync();
+bool SupportsCMask(const FStaticShaderPlatform InPlatform);
+
+static void AddVisualizeSystemInfoPasses(FRDGBuilder& GraphBuilder, const FViewInfo& View, FScreenPassTexture& ScreenPassSceneColor, EShaderPlatform Platform)
+{
+	// Force ShaderPrint on.
+	ShaderPrint::SetEnabled(true);
+	ShaderPrint::RequestSpaceForLines(1024);
+	ShaderPrint::RequestSpaceForCharacters(1024);
+
+	FStrataSystemInfoCS::FParameters* PassParameters = GraphBuilder.AllocParameters<FStrataSystemInfoCS::FParameters>();
+	PassParameters->bAdvancedDebugEnabled = IsStrataAdvancedVisualizationShadersEnabled() ? 1u : 0u;
+	PassParameters->bEnergyConservation = View.ViewState ? View.ViewState->ShadingEnergyConservationData.bEnergyConservation : false;;
+	PassParameters->bEnergyPreservation = View.ViewState ? View.ViewState->ShadingEnergyConservationData.bEnergyPreservation : false;;
+	PassParameters->bDbufferPass = IsStrataDbufferPassEnabled(View.GetShaderPlatform()) ? 1 : 0;
+	PassParameters->ClassificationCMask = SupportsCMask(View.GetShaderPlatform()) ? 1 : 0;
+	PassParameters->ClassificationAsync = IsClassificationAsync() ? 1 : 0;
+	PassParameters->Classification8bits = IsClassificationCoord8bits() ? 1 : 0;
+	PassParameters->LumenTileOverflow = GetStrataTileOverflowRatio();
+	PassParameters->bRoughRefraction = IsStrataOpaqueMaterialRoughRefractionEnabled() ? 1 : 0;
+	PassParameters->ClassificationTileDrawIndirectBuffer = GraphBuilder.CreateSRV(View.StrataViewData.ClassificationTileDrawIndirectBuffer, PF_R32_UINT);
+	PassParameters->ViewUniformBuffer = View.ViewUniformBuffer;
+	PassParameters->Strata = Strata::BindStrataGlobalUniformParameters(View);
+	PassParameters->SceneTextures = GetSceneTextureParameters(GraphBuilder, View);
+	ShaderPrint::SetParameters(GraphBuilder, View.ShaderPrintData, PassParameters->ShaderPrintParameters);
+
+	TShaderMapRef<FStrataSystemInfoCS> ComputeShader(View.ShaderMap);
+	FComputeShaderUtils::AddPass(GraphBuilder, RDG_EVENT_NAME("Strata::VisualizeSystemInfo"), ComputeShader, PassParameters, FIntVector(1, 1, 1));
 }
 
 // Draw each material layer independently
@@ -337,16 +415,26 @@ FScreenPassTexture AddStrataDebugPasses(FRDGBuilder& GraphBuilder, const FViewIn
 		{
 			AddVisualizeAdvancedMaterialPasses(GraphBuilder, View, ScreenPassSceneColor, View.GetShaderPlatform());
 		}
+		else if (DebugMode == FStrataVisualizationData::FViewMode::StrataInfo)
+		{
+			AddVisualizeSystemInfoPasses(GraphBuilder, View, ScreenPassSceneColor, View.GetShaderPlatform());
+		}
 		else if (DebugMode == FStrataVisualizationData::FViewMode::DecalClassification)
 		{
-			AddStrataInternalClassificationTilePass(GraphBuilder, View, nullptr, &ScreenPassSceneColor.Texture, EStrataTileType::EDecalSimple, bDebugPass);
-			AddStrataInternalClassificationTilePass(GraphBuilder, View, nullptr, &ScreenPassSceneColor.Texture, EStrataTileType::EDecalSingle, bDebugPass);
-			AddStrataInternalClassificationTilePass(GraphBuilder, View, nullptr, &ScreenPassSceneColor.Texture, EStrataTileType::EDecalComplex, bDebugPass);
+			if (IsStrataDbufferPassEnabled(View.GetShaderPlatform()))
+			{
+				AddStrataInternalClassificationTilePass(GraphBuilder, View, nullptr, &ScreenPassSceneColor.Texture, EStrataTileType::EDecalSimple, bDebugPass);
+				AddStrataInternalClassificationTilePass(GraphBuilder, View, nullptr, &ScreenPassSceneColor.Texture, EStrataTileType::EDecalSingle, bDebugPass);
+				AddStrataInternalClassificationTilePass(GraphBuilder, View, nullptr, &ScreenPassSceneColor.Texture, EStrataTileType::EDecalComplex, bDebugPass);
+			}
 		}
 		else if (DebugMode == FStrataVisualizationData::FViewMode::RoughRefractionClassification)
 		{
-			AddStrataInternalClassificationTilePass(GraphBuilder, View, nullptr, &ScreenPassSceneColor.Texture, EStrataTileType::EOpaqueRoughRefraction, bDebugPass);
-			AddStrataInternalClassificationTilePass(GraphBuilder, View, nullptr, &ScreenPassSceneColor.Texture, EStrataTileType::EOpaqueRoughRefractionSSSWithout, bDebugPass);
+			if (IsStrataOpaqueMaterialRoughRefractionEnabled())
+			{
+				AddStrataInternalClassificationTilePass(GraphBuilder, View, nullptr, &ScreenPassSceneColor.Texture, EStrataTileType::EOpaqueRoughRefraction, bDebugPass);
+				AddStrataInternalClassificationTilePass(GraphBuilder, View, nullptr, &ScreenPassSceneColor.Texture, EStrataTileType::EOpaqueRoughRefractionSSSWithout, bDebugPass);
+			}
 		}
 		else if (DebugMode == FStrataVisualizationData::FViewMode::MaterialClassification)
 		{
