Improves TSR quality arround on shading rejections around geometric edges

#okforpublic github

diff --git a/Engine/Config/BaseScalability.ini b/Engine/Config/BaseScalability.ini
--- a/Engine/Config/BaseScalability.ini
+++ b/Engine/Config/BaseScalability.ini
@@ -42,7 +42,6 @@ r.TemporalAA.Quality=0
 r.TSR.History.R11G11B10=1
 r.TSR.History.ScreenPercentage=100
 r.TSR.History.UpdateQuality=0
-r.TSR.ShadingRejection.SpatialFilter=0
 r.TSR.ShadingRejection.Flickering=0
 r.TSR.Velocity.Extrapolation=0
 r.TSR.Velocity.HoleFill=0
@@ -54,7 +53,6 @@ r.TemporalAA.Quality=1
 r.TSR.History.R11G11B10=1
 r.TSR.History.ScreenPercentage=100
 r.TSR.History.UpdateQuality=1
-r.TSR.ShadingRejection.SpatialFilter=1
 r.TSR.ShadingRejection.Flickering=0
 r.TSR.Velocity.Extrapolation=0
 r.TSR.Velocity.HoleFill=1
@@ -66,7 +64,6 @@ r.TemporalAA.Quality=1
 r.TSR.History.R11G11B10=1
 r.TSR.History.ScreenPercentage=100
 r.TSR.History.UpdateQuality=2
-r.TSR.ShadingRejection.SpatialFilter=1
 r.TSR.ShadingRejection.Flickering=0
 r.TSR.Velocity.Extrapolation=0
 r.TSR.Velocity.HoleFill=1
@@ -78,7 +75,6 @@ r.TemporalAA.Quality=2
 r.TSR.History.R11G11B10=1
 r.TSR.History.ScreenPercentage=100
 r.TSR.History.UpdateQuality=3
-r.TSR.ShadingRejection.SpatialFilter=2
 r.TSR.ShadingRejection.Flickering=1
 r.TSR.Velocity.Extrapolation=1
 r.TSR.Velocity.HoleFill=1
@@ -90,7 +86,6 @@ r.TemporalAA.Quality=2
 r.TSR.History.R11G11B10=0
 r.TSR.History.ScreenPercentage=200
 r.TSR.History.UpdateQuality=3
-r.TSR.ShadingRejection.SpatialFilter=2
 r.TSR.ShadingRejection.Flickering=1
 r.TSR.Velocity.Extrapolation=1
 r.TSR.Velocity.HoleFill=1
diff --git a/Engine/Shaders/Private/TemporalSuperResolution/TSRDilateRejection.usf b/Engine/Shaders/Private/TemporalSuperResolution/TSRDilateRejection.usf
deleted file mode 100644
diff --git a/Engine/Shaders/Private/TemporalSuperResolution/TSRPostfilterRejection.usf b/Engine/Shaders/Private/TemporalSuperResolution/TSRPostfilterRejection.usf
deleted file mode 100644
diff --git a/Engine/Shaders/Private/TemporalSuperResolution/TSRRejectShading.usf b/Engine/Shaders/Private/TemporalSuperResolution/TSRRejectShading.usf
--- a/Engine/Shaders/Private/TemporalSuperResolution/TSRRejectShading.usf
+++ b/Engine/Shaders/Private/TemporalSuperResolution/TSRRejectShading.usf
@@ -41,6 +41,9 @@
 	#define CONFIG_OPTIMIZED_LDS_ACCESS 0
 #endif
 
+#define CONFIG_MANUAL_LDS_SPILL 0
+
+
 //------------------------------------------------------- DERIVED
 
 /** Pixel coordinate stride between each lane in the GROUP_TILE_SIZE. */
@@ -101,7 +104,7 @@ Texture2D<tsr_ushort> IsMovingMaskTexture;
 
 RWTexture2D<tsr_halfC> HistoryGuideOutput;
 RWTexture2D<tsr_half4> HistoryMoireOutput;
-RWTexture2D<tsr_half> HistoryRejectionOutput;
+RWTexture2D<tsr_half2> HistoryRejectionOutput;
 RWTexture2D<tsr_half> InputSceneColorLdrLumaOutput;
 
 
@@ -435,20 +438,108 @@ FSampleType Convolve3x3HV(FSampleType Color, tsr_simd_half Weight, uint GroupThr
 
 //------------------------------------------------------- GENERAL OPERATIONS
 
-tsr_simd_halfC ClampFireFliersWithGuide(tsr_simd_halfC ToClamp, tsr_simd_halfC Guide, tsr_simd_half Weight, uint GroupThreadIndex)
+CALL_SITE_DEBUGLOC
+template<typename FSampleType>
+FSampleType Min3x3(FSampleType Center, uint GroupThreadIndex)
+{
+	return Convolve3x3HV<TAccumulatorMin<FSampleType>, FSampleType>(Center, tsr_simd_half::Const(1.0), GroupThreadIndex);
+}
+
+CALL_SITE_DEBUGLOC
+template<typename FSampleType>
+FSampleType Max3x3(FSampleType Center, uint GroupThreadIndex)
+{
+	return Convolve3x3HV<TAccumulatorMax<FSampleType>, FSampleType>(Center, tsr_simd_half::Const(1.0), GroupThreadIndex);
+}
+
+CALL_SITE_DEBUGLOC
+tsr_simd_halfC ClampFireFliersWithGuide(tsr_simd_halfC ToClamp, tsr_simd_halfC Guide, uint GroupThreadIndex)
 {
-	tsr_simd_halfC ToClampMin = Convolve3x3HV<TAccumulatorMin<tsr_simd_halfC>, tsr_simd_halfC>(ToClamp, Weight, GroupThreadIndex);
-	tsr_simd_halfC ToClampMax = Convolve3x3HV<TAccumulatorMax<tsr_simd_halfC>, tsr_simd_halfC>(ToClamp, Weight, GroupThreadIndex);
+	tsr_simd_halfC ToClampMin = Min3x3(ToClamp, GroupThreadIndex);
+	tsr_simd_halfC ToClampMax = Max3x3(ToClamp, GroupThreadIndex);
 
 	tsr_simd_halfC ClampedGuide = clamp(Guide, ToClampMin, ToClampMax);
 	
-	tsr_simd_halfC ClampedGuideMin = Convolve3x3HV<TAccumulatorMin<tsr_simd_halfC>, tsr_simd_halfC>(ClampedGuide, Weight, GroupThreadIndex);
-	tsr_simd_halfC ClampedGuideMax = Convolve3x3HV<TAccumulatorMax<tsr_simd_halfC>, tsr_simd_halfC>(ClampedGuide, Weight, GroupThreadIndex);
+	tsr_simd_halfC ClampedGuideMin = Min3x3(ClampedGuide, GroupThreadIndex);
+	tsr_simd_halfC ClampedGuideMax = Max3x3(ClampedGuide, GroupThreadIndex);
 	
 	return clamp(ToClamp, ClampedGuideMin, ClampedGuideMax);
 }
 
 
+//------------------------------------------------------- GENERAL OPERATIONS
+
+// Operator that output, lowest, median and highest values from 3 input values.
+CALL_SITE_DEBUGLOC
+template<typename FSampleType>
+void LMHOperator(FSampleType A, FSampleType B, FSampleType C, out FSampleType L, out FSampleType M, out FSampleType H)
+{
+	// TODO: min3, max3, med3.
+	FSampleType X = min(B, C);
+	FSampleType Y = max(B, C);
+	
+	L = min(A, X);
+	FSampleType Z = max(A, X);
+
+	M = min(Z, Y);
+	H = max(Z, Y);
+}
+
+// 3 samples median.
+CALL_SITE_DEBUGLOC
+template<typename FSampleType>
+FSampleType Median(FSampleType A, FSampleType B, FSampleType C)
+{
+	// TODO: med3.
+	FSampleType L, M, H;
+	LMHOperator(A, B, C, L, M, H);
+	return M;
+}
+
+// 9 samples median using Smith1996.
+CALL_SITE_DEBUGLOC
+template<typename FSampleType>
+FSampleType Median9(FSampleType In[9])
+{
+	// First layer.
+	FSampleType L0[3];
+	FSampleType M0[3];
+	FSampleType H0[3];
+	for (uint j = 0; j < 3; j++)
+	{
+		LMHOperator(In[j * 3 + 0], In[j * 3 + 1], In[j * 3 + 2], L0[j], M0[j], H0[j]);
+	}
+	
+	// Second layer.
+	// TODO: min3, max3.
+	FSampleType M1[3];
+	M1[0] = max(max(L0[0], L0[1]), L0[2]);
+	M1[1] = Median(M0[0], M0[1], M0[2]);
+	M1[2] = min(min(H0[0], H0[1]), H0[2]);
+	
+	// Third layer.
+	return Median(M1[0], M1[1], M1[2]);
+}
+
+ISOLATE
+CALL_SITE_DEBUGLOC
+template<typename FSampleType>
+FSampleType Median3x3(FSampleType Center, uint GroupThreadIndex)
+{
+	WriteToLDSForAccessNeighborTexel(Center, GroupThreadIndex);
+	
+	FSampleType InSamples[9];
+	UNROLL
+	for (uint i = 0; i < 9; i++)
+	{
+		FSampleType Sample = AccessNeighborTexel(Center, tsr_short2(kOffsets3x3[i]), GroupThreadIndex);
+		InSamples[i] = Sample;
+	}
+
+	return Median9(InSamples);
+}
+
+
 //------------------------------------------------------- ENTRY POINT
 
 #if COMPILER_SUPPORTS_WAVE_SIZE && DIM_WAVE_SIZE > 0
@@ -600,8 +691,8 @@ void MainCS(
 
 
 		// Compute clamping box of the limance.
-		tsr_simd_half InputLumaMin = Convolve3x3HV<TAccumulatorMin<tsr_simd_half>, tsr_simd_half>(InputLuma, bIsValid, GroupThreadIndex);
-		tsr_simd_half InputLumaMax = Convolve3x3HV<TAccumulatorMax<tsr_simd_half>, tsr_simd_half>(InputLuma, bIsValid, GroupThreadIndex);
+		tsr_simd_half InputLumaMin = Min3x3(InputLuma, GroupThreadIndex);
+		tsr_simd_half InputLumaMax = Max3x3(InputLuma, GroupThreadIndex);
 
 		// Clamp and contribue the history.		
 		tsr_simd_half ClampedPrevFlickeringHistory = clamp(PrevFlickeringHistory, InputLumaMin, InputLumaMax);
@@ -641,8 +732,8 @@ void MainCS(
 
 		// Dilates the total variation and contribution to keep stability despite inacuracies in the history reprojection
 		// of these need a nearest filter which isn't precise at all.
-		tsr_simd_half TotalVariationContribution = Convolve3x3HV<TAccumulatorMax<tsr_simd_half>, tsr_simd_half>(GradientVariation, tsr_simd_half::Const(1.0), GroupThreadIndex);
-		tsr_simd_half TotalVariationCountContribution = Convolve3x3HV<TAccumulatorMax<tsr_simd_half>, tsr_simd_half>(IsFlicker, tsr_simd_half::Const(1.0), GroupThreadIndex);
+		tsr_simd_half TotalVariationContribution = Max3x3(GradientVariation, GroupThreadIndex);
+		tsr_simd_half TotalVariationCountContribution = Max3x3(IsFlicker, GroupThreadIndex);
 
 		UNROLL_N(SIMD_SIZE)
 		for (uint SimdIndex = 0; SimdIndex < SIMD_SIZE; SimdIndex++)
@@ -709,8 +800,8 @@ void MainCS(
 	tsr_simd_halfC ClampedHistory = OriginalHistory;
 	#if CONFIG_FIREFLY_REMOVAL
 	{
-		ClampedInput   = ClampFireFliersWithGuide(OriginalInput  , /* Guide = */ OriginalHistory, bIsValid, GroupThreadIndex);
-		ClampedHistory = ClampFireFliersWithGuide(OriginalHistory, /* Guide = */ OriginalInput  , bIsValid, GroupThreadIndex);
+		ClampedInput   = ClampFireFliersWithGuide(OriginalInput  , /* Guide = */ OriginalHistory, GroupThreadIndex);
+		ClampedHistory = ClampFireFliersWithGuide(OriginalHistory, /* Guide = */ OriginalInput  , GroupThreadIndex);
 	}
 	#endif
 
@@ -719,16 +810,16 @@ void MainCS(
 	tsr_simd_halfC FilteredHistory = Convolve3x3HV<TAccumulatorAvg<tsr_simd_halfC>, tsr_simd_halfC>(ClampedHistory, bIsValid, GroupThreadIndex);
 	
 	// Measure rejection
-	tsr_simd_half Rejection;
-	Rejection.SetAllElements(tsr_half(1.0));
+	tsr_simd_half RejectionBlendFinal = tsr_simd_half::Const(1.0);
+	tsr_simd_half RejectionClampBlend = tsr_simd_half::Const(1.0);
 	#if 1
 	{
 		const tsr_half FilteringWeight = rcp(tsr_half(1.0 + 4 * 0.5 + 4 * 0.25));
 
 		tsr_simd_halfC BackbufferError = tsr_simd_halfC::Const(MeasureBackbufferLDRQuantizationError());
 
-		tsr_simd_halfC FilteredBoxMin = Convolve3x3HV<TAccumulatorMin<tsr_simd_halfC>, tsr_simd_halfC>(FilteredInput, bIsValid, GroupThreadIndex) - BackbufferError;
-		tsr_simd_halfC FilteredBoxMax = Convolve3x3HV<TAccumulatorMax<tsr_simd_halfC>, tsr_simd_halfC>(FilteredInput, bIsValid, GroupThreadIndex) + BackbufferError;
+		tsr_simd_halfC FilteredBoxMin = Min3x3(FilteredInput, GroupThreadIndex) - BackbufferError;
+		tsr_simd_halfC FilteredBoxMax = Max3x3(FilteredInput, GroupThreadIndex) + BackbufferError;
 		
 		tsr_simd_halfC FlickeringClampedFilteredHistory = clamp(FilteredHistory, FilteredBoxMin, FilteredBoxMax);
 
@@ -762,31 +853,49 @@ void MainCS(
 			ClampedFilteredHistory = FlickeringClampedFilteredHistory;
 		#endif
 		
-		tsr_simd_halfC OriginalBoxMin = Convolve3x3HV<TAccumulatorMin<tsr_simd_halfC>, tsr_simd_halfC>(OriginalInput, bIsValid, GroupThreadIndex);
-		tsr_simd_halfC OriginalBoxMax = Convolve3x3HV<TAccumulatorMax<tsr_simd_halfC>, tsr_simd_halfC>(OriginalInput, bIsValid, GroupThreadIndex);
+		tsr_simd_halfC OriginalBoxMin = Min3x3(OriginalInput, GroupThreadIndex);
+		tsr_simd_halfC OriginalBoxMax = Max3x3(OriginalInput, GroupThreadIndex);
 		tsr_simd_halfC BoxSize = (OriginalBoxMax - OriginalBoxMin + (BackbufferError + BackbufferError)) * FilteringWeight;
 		
+		const tsr_simd_halfC BackbufferQuantizationErrorVector = tsr_simd_halfC::Const(MeasureBackbufferLDRQuantizationError());
+		tsr_simd_halfC RawClampedEnergy = abs(ClampedFilteredHistory - FilteredHistory);
+
+		tsr_simd_halfC FilteredClampedEnergy;
+		{
+			FilteredClampedEnergy = Median3x3(RawClampedEnergy, GroupThreadIndex);
+			FilteredClampedEnergy = Max3x3(FilteredClampedEnergy, GroupThreadIndex);
+		}
+
+		#if CONFIG_MOIRE_DETECTION
+		tsr_simd_halfC MoireErrorSize;
 		UNROLL_N(SIMD_SIZE)
 		for (uint SimdIndex = 0; SimdIndex < SIMD_SIZE; SimdIndex++)
 		{
-			#if CONFIG_MOIRE_DETECTION
-				tsr_halfC ClampedEnergy = abs(ClampedFilteredHistory.GetElement(SimdIndex) - FilteredHistory.GetElement(SimdIndex));
-				tsr_halfC ClampedEnergyAmount = ClampedEnergy * SafeRcp(ClampedEnergy[0] + ClampedEnergy[1] + ClampedEnergy[2]);
-
-				// Apply the moire error on each individual channel based on how much they need to clamp.
-				tsr_halfC MoireErrorSize = ClampedEnergyAmount * ((FilteringWeight * tsr_half(3.0)) * MoireError.GetElement(SimdIndex));
-				//tsr_halfC MoireErrorSize = FilteringWeight * 3.0 * MoireError.GetElement(SimdIndex);
-			#else
-				const tsr_halfC MoireErrorSize = tsr_half(0.0);
-			#endif
+			tsr_halfC ClampedEnergy = abs(ClampedFilteredHistory.GetElement(SimdIndex) - FilteredHistory.GetElement(SimdIndex));
+			tsr_halfC ClampedEnergyAmount = ClampedEnergy * SafeRcp(ClampedEnergy[0] + ClampedEnergy[1] + ClampedEnergy[2]);
 
-			Rejection.SetElement(SimdIndex, MeasureRejectionFactor(
-				FilteredHistory.GetElement(SimdIndex), ClampedFilteredHistory.GetElement(SimdIndex),
-				FilteredInput.GetElement(SimdIndex), tsr_half(0.0), max(BoxSize.GetElement(SimdIndex), MoireErrorSize),
-				MeasureBackbufferLDRQuantizationError()));
+			// Apply the moire error on each individual channel based on how much they need to clamp.
+			MoireErrorSize.SetElement(SimdIndex, ClampedEnergyAmount * ((FilteringWeight * tsr_half(3.0)) * MoireError.GetElement(SimdIndex)));
 		}
+		BoxSize = max(BoxSize, MoireErrorSize);
+		#endif
+
+		tsr_simd_halfC Delta = abs(FilteredInput - FilteredHistory);
+		tsr_simd_halfC ClampError = max(
+			BoxSize * tsr_half(0.25),
+			BackbufferQuantizationErrorVector);
+			
+		tsr_simd_halfC RawFactor      = saturate(tsr_simd_halfC::Const(1.0) - max(RawClampedEnergy - ClampError, tsr_simd_halfC::Const(0.0)) * rcp(max(Delta, BoxSize)));
+		tsr_simd_halfC FilteredFactor = saturate(tsr_simd_halfC::Const(1.0) - max(FilteredClampedEnergy - ClampError, tsr_simd_halfC::Const(0.0)) * rcp(max(Delta, BoxSize)));
+		
+		tsr_simd_half RawRejection      = min(RawFactor[0], min(RawFactor[1], RawFactor[2]));
+		tsr_simd_half FilteredRejection = min(FilteredFactor[0], min(FilteredFactor[1], FilteredFactor[2]));
+		
+		RejectionBlendFinal = saturate(FilteredRejection + tsr_simd_half::Const(1.0) - bIsValid);
 		
-		Rejection = saturate(Rejection + tsr_simd_half::Const(1.0) - bIsValid);
+		RejectionClampBlend = Median3x3(RawRejection, GroupThreadIndex);
+		RejectionClampBlend = saturate(RejectionClampBlend + tsr_simd_half::Const(1.0) - bIsValid);
+		RejectionClampBlend = Min3x3(RejectionClampBlend, GroupThreadIndex);
 	}
 	#endif
 	
@@ -804,10 +913,10 @@ void MainCS(
 		{
 			bool bIsOffScreen = false;
 			bool bIsDisoccluded = bIsValid.GetElement(SimdIndex) == tsr_half(0.0);
-			tsr_half LocalRejection = Rejection.GetElement(SimdIndex);
+			tsr_half LocalRejectionBlendFinal = RejectionBlendFinal.GetElement(SimdIndex);
 
-			tsr_half TotalBlend = select(or(bIsOffScreen, bIsDisoccluded), tsr_half(1.0), saturate(tsr_half(1.0) - LocalRejection * tsr_half(4.0)));
-			tsr_half RejectedBlendFactor = max(TheoricBlendFactor, tsr_half(1.0) - LocalRejection);
+			tsr_half TotalBlend = select(or(bIsOffScreen, bIsDisoccluded), tsr_half(1.0), saturate(tsr_half(1.0) - LocalRejectionBlendFinal * tsr_half(4.0)));
+			tsr_half RejectedBlendFactor = max(TheoricBlendFactor, tsr_half(1.0) - LocalRejectionBlendFinal);
 
 
 			BlendFinal.SetElement(SimdIndex, max(RejectedBlendFactor, TotalBlend));
@@ -840,10 +949,12 @@ void MainCS(
 			tsr_short2 OutputPixelPos = InvalidateOutputPixelPos(GroupPixelOffset + LaneSimdPixelOffset, InputInfo_ViewportMax);
 			OutputPixelPos.x = select(bIsNonOverscanPixel, OutputPixelPos.x, ~tsr_short(0));
 		
-			tsr_half LocalRejection = Rejection.GetElement(SimdIndex);
+			tsr_half LocalRejectionBlendFinal = RejectionBlendFinal.GetElement(SimdIndex);
+			tsr_half LocalRejectionClampBlend = RejectionClampBlend.GetElement(SimdIndex);
 			tsr_halfC FinalGuideColor = FinalGuide.GetElement(SimdIndex);
 			
-			LocalRejection = LocalRejection - 0.5 * rcp(tsr_half(256.0));
+			LocalRejectionBlendFinal = LocalRejectionBlendFinal - tsr_half(0.5 * rcp(256.0));
+			LocalRejectionClampBlend = LocalRejectionClampBlend - tsr_half(0.5 * rcp(256.0));
 			#if CONFIG_ENABLE_STOCASTIC_QUANTIZATION
 			{
 				uint2 Random = Rand3DPCG16(int3(LaneSimdPixelOffset, View.StateFrameIndexMod8)).xy;
@@ -857,7 +968,7 @@ void MainCS(
 			#if CONFIG_MOIRE_DETECTION
 				HistoryMoireOutput[OutputPixelPos] = FinalMoire.GetElement(SimdIndex);
 			#endif
-			HistoryRejectionOutput[OutputPixelPos] = LocalRejection;
+			HistoryRejectionOutput[OutputPixelPos] = tsr_half2(LocalRejectionBlendFinal, LocalRejectionClampBlend);
 			InputSceneColorLdrLumaOutput[OutputPixelPos] = CenterLumaLDR.GetElement(SimdIndex);
 			
 			#if DEBUG_OUTPUT
diff --git a/Engine/Shaders/Private/TemporalSuperResolution/TSRUpdateHistory.usf b/Engine/Shaders/Private/TemporalSuperResolution/TSRUpdateHistory.usf
--- a/Engine/Shaders/Private/TemporalSuperResolution/TSRUpdateHistory.usf
+++ b/Engine/Shaders/Private/TemporalSuperResolution/TSRUpdateHistory.usf
@@ -130,7 +130,7 @@ Texture2D<tsr_halfC> InputSceneColorTexture;
 Texture2D<uint2> InputSceneStencilTexture;
 Texture2D<tsr_half4> InputSceneTranslucencyTexture;
 
-Texture2D<tsr_half> HistoryRejectionTexture;
+Texture2D<tsr_half2> HistoryRejectionTexture;
 Texture2D<tsr_half> TranslucencyRejectionTexture;
 
 Texture2D<float2> DilatedVelocityTexture;
@@ -399,6 +399,7 @@ void MainCS(
 		float2x2 GrandPrevScreenPos = ScreenPos;
 	#endif
 	tsr_half2 LowFrequencyRejection = tsr_half(1.0).xx;
+	tsr_half2 LowFrequencyClamp = tsr_half(1.0).xx;
 	tsr_half2 TranslucencyRejection = tsr_half(1.0).xx;
 	tsr_half2 OutputPixelVelocity = tsr_half(0.0).xx;
 	tsr_half2 NoiseFiltering = tsr_half(0.0).xx;
@@ -429,7 +430,7 @@ void MainCS(
 
 		tsr_half RawParallaxRejectionMask[DPV_PIXEL_PER_LANE];
 		uint RawParallaxFactorBits[DPV_PIXEL_PER_LANE];
-		tsr_half RawLowFrequencyRejection[DPV_PIXEL_PER_LANE];
+		tsr_half2 RawLowFrequencyRejection[DPV_PIXEL_PER_LANE];
 		
 		uint RawSceneStencilRef[DPV_PIXEL_PER_LANE];
 		uint RawSceneStencilRefNeighborhood[CONFIG_SAMPLES - 1][DPV_PIXEL_PER_LANE];
@@ -443,8 +444,6 @@ void MainCS(
 		// Issue overlapped texture fetches
 		PLATFORM_SPECIFIC_ISOLATE
 		{
-			tsr_ushort2x2 RejectionPixelPos = tsr_ushort2x2(InputPixelPos) / 2;
-
 			UNROLL_N(DPV_PIXEL_PER_LANE)
 			for (uint PixelId = 0; PixelId < DPV_PIXEL_PER_LANE; PixelId++)
 			{
@@ -461,7 +460,7 @@ void MainCS(
 				RawParallaxFactorBits[PixelId]    = ParallaxFactorTexture[LocalInputPixelPos];
 				RawSceneStencilRef[PixelId]       = InputSceneStencilTexture.Load(int3(LocalInputPixelPos, 0)) STENCIL_COMPONENT_SWIZZLE;
 				
-				RawLowFrequencyRejection[PixelId] = HistoryRejectionTexture[dpv_access_pixel(RejectionPixelPos, PixelId)];
+				RawLowFrequencyRejection[PixelId] = HistoryRejectionTexture[LocalInputPixelPos];
 				
 				#if CONFIG_REJECTION_ANTI_ALIASING
 				{
@@ -617,7 +616,8 @@ void MainCS(
 				/* out */ bIsDisoccluded);
 
 			// Process input texel, forcing to pack register.
-			LowFrequencyRejection = dpv_force_interleave_registers_array(RawLowFrequencyRejection);
+			LowFrequencyRejection = dpv_force_interleave_registers_array(RawLowFrequencyRejection)[0];
+			LowFrequencyClamp     = dpv_force_interleave_registers_array(RawLowFrequencyRejection)[1];
 			SceneStencilRef       = dpv_force_interleave_registers_array(RawSceneStencilRef);
 			
 			// Process neighborhood
@@ -1248,7 +1248,7 @@ void MainCS(
 
 		tsr_half2 TotalRejection = select(or(bIsOffScreen, bIsDisoccluded), tsr_half(0.0).xx, saturate(LowFrequencyRejection * tsr_half(4.0)));
 
-		tsr_half2 BlendClamp = TotalRejection * HoleFilledVelocityMask[1];
+		tsr_half2 BlendClamp = min(TotalRejection * HoleFilledVelocityMask[1], LowFrequencyClamp);
 		
 		#if CONFIG_PREMULTIPLY_HISTORY
 			tsr_halfCx2 UnpremultipliedPrevHighFrequencyColor = dpv_scale(PrevHighFrequencyColor, InvPrevHistoryWeight);
diff --git a/Engine/Source/Runtime/Renderer/Private/PostProcess/TemporalSuperResolution.cpp b/Engine/Source/Runtime/Renderer/Private/PostProcess/TemporalSuperResolution.cpp
--- a/Engine/Source/Runtime/Renderer/Private/PostProcess/TemporalSuperResolution.cpp
+++ b/Engine/Source/Runtime/Renderer/Private/PostProcess/TemporalSuperResolution.cpp
@@ -61,14 +61,6 @@ TAutoConsoleVariable<float> CVarTSRFlickeringMaxParralaxVelocity(
 	TEXT("Maximum parralax velocity in 1080p 60hz pixels allowed before diminishing flickering.\n"),
 	ECVF_RenderThreadSafe);
 
-TAutoConsoleVariable<int32> CVarTSRFilterShadingRejection(
-	TEXT("r.TSR.ShadingRejection.SpatialFilter"), 1,
-	TEXT("Whether the shading rejection should have spatial statistical filtering pass to reduce flickering (default = 1).\n")
-	TEXT(" 0: Disabled;\n")
-	TEXT(" 1: Spatial filter pass is run at lower resolution than CompareHistory pass (default);\n")
-	TEXT(" 2: Spatial filter pass is run CompareHistory pass resolution to improve stability."),
-	ECVF_Scalability | ECVF_RenderThreadSafe);
-
 TAutoConsoleVariable<int32> CVarTSRRejectionAntiAliasingQuality(
 	TEXT("r.TSR.RejectionAntiAliasingQuality"), 3,
 	TEXT("Controls the quality of spatial anti-aliasing on history rejection (default=1)."),
@@ -543,36 +535,6 @@ class FTSRRejectShadingCS : public FTSRShader
 	}
 }; // class FTSRRejectShadingCS
 
-class FTSRPostfilterRejectionCS : public FTSRShader
-{
-	DECLARE_GLOBAL_SHADER(FTSRPostfilterRejectionCS);
-	SHADER_USE_PARAMETER_STRUCT(FTSRPostfilterRejectionCS, FTSRShader);
-
-	class FOutputHalfRes : SHADER_PERMUTATION_BOOL("DIM_OUTPUT_HALF_RES");
-	using FPermutationDomain = TShaderPermutationDomain<FOutputHalfRes>;
-
-	BEGIN_SHADER_PARAMETER_STRUCT(FParameters, )
-		SHADER_PARAMETER(FIntRect, HistoryRejectionViewport)
-		SHADER_PARAMETER_RDG_TEXTURE(Texture2D, HistoryRejectionTexture)
-		SHADER_PARAMETER_RDG_TEXTURE(Texture2D, ParallaxRejectionMaskTexture)
-		SHADER_PARAMETER_RDG_TEXTURE_UAV(RWTexture2D, HistoryRejectionOutput)
-		SHADER_PARAMETER_RDG_TEXTURE_UAV(RWTexture2DArray, DebugOutput)
-	END_SHADER_PARAMETER_STRUCT()
-}; // class FTSRPostfilterRejectionCS
-
-class FTSRDilateRejectionCS : public FTSRShader
-{
-	DECLARE_GLOBAL_SHADER(FTSRDilateRejectionCS);
-	SHADER_USE_PARAMETER_STRUCT(FTSRDilateRejectionCS, FTSRShader);
-
-	BEGIN_SHADER_PARAMETER_STRUCT(FParameters, )
-		SHADER_PARAMETER(FIntRect, HistoryRejectionViewport)
-		SHADER_PARAMETER_RDG_TEXTURE(Texture2D, HistoryRejectionTexture)
-		SHADER_PARAMETER_RDG_TEXTURE_UAV(RWTexture2D, DilatedHistoryRejectionOutput)
-		SHADER_PARAMETER_RDG_TEXTURE_UAV(RWTexture2DArray, DebugOutput)
-	END_SHADER_PARAMETER_STRUCT()
-}; // class FTSRDilateRejectionCS
-
 class FTSRSpatialAntiAliasingCS : public FTSRShader
 {
 	DECLARE_GLOBAL_SHADER(FTSRSpatialAntiAliasingCS);
@@ -756,8 +718,6 @@ IMPLEMENT_GLOBAL_SHADER(FTSRDilateVelocityCS,        "/Engine/Private/TemporalSu
 IMPLEMENT_GLOBAL_SHADER(FTSRDecimateHistoryCS,       "/Engine/Private/TemporalSuperResolution/TSRDecimateHistory.usf",       "MainCS", SF_Compute);
 IMPLEMENT_GLOBAL_SHADER(FTSRCompareTranslucencyCS,   "/Engine/Private/TemporalSuperResolution/TSRCompareTranslucency.usf",   "MainCS", SF_Compute);
 IMPLEMENT_GLOBAL_SHADER(FTSRRejectShadingCS,         "/Engine/Private/TemporalSuperResolution/TSRRejectShading.usf",         "MainCS", SF_Compute);
-IMPLEMENT_GLOBAL_SHADER(FTSRPostfilterRejectionCS,   "/Engine/Private/TemporalSuperResolution/TSRPostfilterRejection.usf",   "MainCS", SF_Compute);
-IMPLEMENT_GLOBAL_SHADER(FTSRDilateRejectionCS,       "/Engine/Private/TemporalSuperResolution/TSRDilateRejection.usf",       "MainCS", SF_Compute);
 IMPLEMENT_GLOBAL_SHADER(FTSRSpatialAntiAliasingCS,   "/Engine/Private/TemporalSuperResolution/TSRSpatialAntiAliasing.usf",   "MainCS", SF_Compute);
 IMPLEMENT_GLOBAL_SHADER(FTSRFilterAntiAliasingCS,    "/Engine/Private/TemporalSuperResolution/TSRFilterAntiAliasing.usf",    "MainCS", SF_Compute);
 IMPLEMENT_GLOBAL_SHADER(FTSRUpdateHistoryCS,         "/Engine/Private/TemporalSuperResolution/TSRUpdateHistory.usf",         "MainCS", SF_Compute);
@@ -858,21 +818,9 @@ ITemporalUpscaler::FOutputs AddTemporalSuperResolutionPasses(
 		RejectionAntiAliasingQuality = 0; 
 	}
 
-	enum class ERejectionPostFilter : uint8
-	{
-		Disabled,
-		PostRejectionDownsample,
-		PreRejectionDownsample,
-	};
-
-	ERejectionPostFilter PostFilter = ERejectionPostFilter(FMath::Clamp(CVarTSRFilterShadingRejection.GetValueOnRenderThread(), 2, 2)); // TODO(TSR)
-
 	FIntPoint InputExtent = PassInputs.SceneColorTexture->Desc.Extent;
 	FIntRect InputRect = View.ViewRect;
 
-	FIntPoint RejectionExtent = InputExtent / 2;
-	FIntRect RejectionRect = FIntRect(InputRect.Min / 2, InputRect.Min / 2 + FIntPoint::DivideAndRoundUp(InputRect.Size(), 2));
-
 	FIntPoint OutputExtent;
 	FIntRect OutputRect;
 	if (View.PrimaryScreenPercentageMethod == EPrimaryScreenPercentageMethod::TemporalUpscale)
@@ -1498,9 +1446,6 @@ ITemporalUpscaler::FOutputs AddTemporalSuperResolutionPasses(
 	FRDGTextureRef HistoryRejectionTexture;
 	FRDGTextureRef InputSceneColorLdrLumaTexture = nullptr;
 	{
-		bool bOutputHalfRes = PostFilter != ERejectionPostFilter::PreRejectionDownsample;
-		check(!bOutputHalfRes); // TODO(TSR)
-
 		{
 			FRDGTextureDesc Desc = FRDGTextureDesc::Create2D(
 				InputExtent,
@@ -1513,8 +1458,8 @@ ITemporalUpscaler::FOutputs AddTemporalSuperResolutionPasses(
 
 		{
 			FRDGTextureDesc Desc = FRDGTextureDesc::Create2D(
-				bOutputHalfRes ? RejectionExtent : InputExtent,
-				PF_R8,
+				InputExtent,
+				PF_R8G8,
 				FClearValueBinding::None,
 				/* InFlags = */ TexCreate_ShaderResource | TexCreate_UAV);
 
@@ -1634,66 +1579,6 @@ ITemporalUpscaler::FOutputs AddTemporalSuperResolutionPasses(
 		}
 	}
 
-	// Post filter the rejection.
-	if (PostFilter != ERejectionPostFilter::Disabled)
-	{
-		bool bOutputHalfRes = PostFilter == ERejectionPostFilter::PreRejectionDownsample;
-		FIntRect Rect = bOutputHalfRes ? InputRect : RejectionRect;
-
-		FRDGTextureRef FilteredHistoryRejectionTexture;
-		{
-			FRDGTextureDesc Desc = FRDGTextureDesc::Create2D(
-				RejectionExtent,
-				PF_R8,
-				FClearValueBinding::None,
-				/* InFlags = */ TexCreate_ShaderResource | TexCreate_UAV);
-
-			FilteredHistoryRejectionTexture = GraphBuilder.CreateTexture(Desc, TEXT("TSR.HistoryRejection"));
-		}
-
-		FTSRPostfilterRejectionCS::FParameters* PassParameters = GraphBuilder.AllocParameters<FTSRPostfilterRejectionCS::FParameters>();
-		PassParameters->HistoryRejectionViewport = Rect;
-		PassParameters->HistoryRejectionTexture = HistoryRejectionTexture;
-		PassParameters->ParallaxRejectionMaskTexture = ParallaxRejectionMaskTexture;
-		PassParameters->HistoryRejectionOutput = GraphBuilder.CreateUAV(FilteredHistoryRejectionTexture);
-		PassParameters->DebugOutput = CreateDebugUAV(bOutputHalfRes ? InputExtent : RejectionExtent, TEXT("Debug.TSR.PostfilterRejection"));
-
-		FTSRPostfilterRejectionCS::FPermutationDomain PermutationVector;
-		PermutationVector.Set<FTSRPostfilterRejectionCS::FOutputHalfRes>(PostFilter == ERejectionPostFilter::PreRejectionDownsample);
-
-		TShaderMapRef<FTSRPostfilterRejectionCS> ComputeShader(View.ShaderMap, PermutationVector);
-		FComputeShaderUtils::AddPass(
-			GraphBuilder,
-			RDG_EVENT_NAME("TSR PostfilterRejection %dx%d", Rect.Width(), Rect.Height()),
-			AsyncComputePasses >= 3 ? ERDGPassFlags::AsyncCompute : ERDGPassFlags::Compute,
-			ComputeShader,
-			PassParameters,
-			FComputeShaderUtils::GetGroupCount(Rect.Size(), 8));
-
-		HistoryRejectionTexture = FilteredHistoryRejectionTexture;
-	}
-
-	// Dilate the rejection.
-	FRDGTextureRef DilatedHistoryRejectionTexture;
-	{
-		DilatedHistoryRejectionTexture = GraphBuilder.CreateTexture(HistoryRejectionTexture->Desc, TEXT("TSR.DilatedHistoryRejection"));
-
-		FTSRDilateRejectionCS::FParameters* PassParameters = GraphBuilder.AllocParameters<FTSRDilateRejectionCS::FParameters>();
-		PassParameters->HistoryRejectionViewport = RejectionRect;
-		PassParameters->HistoryRejectionTexture = HistoryRejectionTexture;
-		PassParameters->DilatedHistoryRejectionOutput = GraphBuilder.CreateUAV(DilatedHistoryRejectionTexture);
-		PassParameters->DebugOutput = CreateDebugUAV(RejectionExtent, TEXT("Debug.TSR.DilateRejection"));
-
-		TShaderMapRef<FTSRDilateRejectionCS> ComputeShader(View.ShaderMap);
-		FComputeShaderUtils::AddPass(
-			GraphBuilder,
-			RDG_EVENT_NAME("TSR DilateRejection %dx%d", RejectionRect.Width(), RejectionRect.Height()),
-			AsyncComputePasses >= 3 ? ERDGPassFlags::AsyncCompute : ERDGPassFlags::Compute,
-			ComputeShader,
-			PassParameters,
-			FComputeShaderUtils::GetGroupCount(RejectionRect.Size(), 8));
-	}
-
 	// Update temporal history.
 	{
 		static const TCHAR* const kUpdateQualityNames[] = {
@@ -1710,7 +1595,7 @@ ITemporalUpscaler::FOutputs AddTemporalSuperResolutionPasses(
 		PassParameters->InputSceneStencilTexture = GraphBuilder.CreateSRV(
 			FRDGTextureSRVDesc::CreateWithPixelFormat(PassInputs.SceneDepthTexture, PF_X24_G8));
 		PassParameters->InputSceneTranslucencyTexture = SeparateTranslucencyTexture;
-		PassParameters->HistoryRejectionTexture = DilatedHistoryRejectionTexture;
+		PassParameters->HistoryRejectionTexture = HistoryRejectionTexture;
 		PassParameters->TranslucencyRejectionTexture = TranslucencyRejectionTexture ? TranslucencyRejectionTexture : BlackDummy;
 
 		PassParameters->DilatedVelocityTexture = DilatedVelocityTexture;
