Fix typos in variable names, comments, and #defines. Put code that's only used when CONFIG_ACCUMULATE_TRANSLUCENCY_SEPARATELY is 1 into #if block for it.

diff --git a/Engine/Shaders/Private/TemporalSuperResolution/TSRCommon.ush b/Engine/Shaders/Private/TemporalSuperResolution/TSRCommon.ush
--- a/Engine/Shaders/Private/TemporalSuperResolution/TSRCommon.ush
+++ b/Engine/Shaders/Private/TemporalSuperResolution/TSRCommon.ush
@@ -39,12 +39,12 @@
 #define CONFIG_SCENE_COLOR_OVERFLOW 1
 
 #if POST_PROCESS_ALPHA
-	#define CONFIG_SCENE_COLOR_APLHA 1
+	#define CONFIG_SCENE_COLOR_ALPHA 1
 #else
-	#define CONFIG_SCENE_COLOR_APLHA 0
+	#define CONFIG_SCENE_COLOR_ALPHA 0
 #endif
 
-#define CONFIG_ENABLE_STOCASTIC_QUANTIZATION (!CONFIG_SCENE_COLOR_APLHA)
+#define CONFIG_ENABLE_STOCASTIC_QUANTIZATION (!CONFIG_SCENE_COLOR_ALPHA)
 
 //------------------------------------------------------- CONSTANTS
 
@@ -205,7 +205,7 @@ static const uint kPlusIndexes3x3[5] = { 4, 1, 3, 7, 5 };
 
 #endif
 
-#if CONFIG_SCENE_COLOR_APLHA
+#if CONFIG_SCENE_COLOR_ALPHA
 	#define tsr_halfC tsr_half4
 	#define tsr_halfCx2 tsr_half4x2
 	#define CONFIG_CHANNEL_COUNT 4
@@ -884,7 +884,7 @@ tsr_half2 MeasureRejectionFactor(
 {
 	const tsr_halfC Null = tsr_half(0.0);
 	
-	#if CONFIG_SCENE_COLOR_APLHA
+	#if CONFIG_SCENE_COLOR_ALPHA
 		tsr_halfCx2 BackbufferQuantizationErrorVector = tsr_half4x2(BackbufferQuantizationError, BackbufferQuantizationError, BackbufferQuantizationError, BackbufferQuantizationError);
 	#else
 		tsr_halfCx2 BackbufferQuantizationErrorVector = tsr_half3x2(BackbufferQuantizationError, BackbufferQuantizationError, BackbufferQuantizationError);
@@ -903,7 +903,7 @@ tsr_half2 MeasureRejectionFactor(
 
 	tsr_half2 Rejection = min3(Factor[0], Factor[1], Factor[2]);
 	
-	#if CONFIG_SCENE_COLOR_APLHA
+	#if CONFIG_SCENE_COLOR_ALPHA
 		Rejection = min(Rejection, Factor[3]);
 	#endif
 
diff --git a/Engine/Shaders/Private/TemporalSuperResolution/TSRCompareTranslucency.usf b/Engine/Shaders/Private/TemporalSuperResolution/TSRCompareTranslucency.usf
--- a/Engine/Shaders/Private/TemporalSuperResolution/TSRCompareTranslucency.usf
+++ b/Engine/Shaders/Private/TemporalSuperResolution/TSRCompareTranslucency.usf
@@ -79,7 +79,7 @@ tsr_half MeasureTranslucencyRejectionFactor(
 	
 	tsr_half BackbufferQuantizationError = 0.0; //MeasureBackbufferLDRQuantizationErrorFromHDRLuma(max(PrevTranslucency[0], Translucency[0]));
 	
-	#if CONFIG_SCENE_COLOR_APLHA
+	#if CONFIG_SCENE_COLOR_ALPHA
 		tsr_half Rejection = MeasureRejectionFactor(
 			PrevTranslucency, ClampedPrevTranslucency,
 			Translucency, TranslucencyMin, TranslucencyMax,
diff --git a/Engine/Shaders/Private/TemporalSuperResolution/TSRRejectShading.usf b/Engine/Shaders/Private/TemporalSuperResolution/TSRRejectShading.usf
--- a/Engine/Shaders/Private/TemporalSuperResolution/TSRRejectShading.usf
+++ b/Engine/Shaders/Private/TemporalSuperResolution/TSRRejectShading.usf
@@ -630,7 +630,7 @@ void MainCS(
 			tsr_half4 Translucency = FilteredTranslucencyInput.GetElement(SimdIndex);
 
 			Color.rgb = Color.rgb * Translucency.a + Translucency.rgb;
-			#if CONFIG_SCENE_COLOR_APLHA
+			#if CONFIG_SCENE_COLOR_ALPHA
 				Color.a *= Translucency.a;
 			#endif 
 
diff --git a/Engine/Shaders/Private/TemporalSuperResolution/TSRResolveHistory.usf b/Engine/Shaders/Private/TemporalSuperResolution/TSRResolveHistory.usf
--- a/Engine/Shaders/Private/TemporalSuperResolution/TSRResolveHistory.usf
+++ b/Engine/Shaders/Private/TemporalSuperResolution/TSRResolveHistory.usf
@@ -157,7 +157,7 @@ void MainCS(
 					SampleLowFrequencyColor.rgb  = SampleLowFrequencyColor.rgb  * SampleTranslucencyAlpha + SampleTranslucencyColor.rgb;
 					SampleHighFrequencyColor.rgb = SampleHighFrequencyColor.rgb * SampleTranslucencyAlpha + SampleTranslucencyColor.rgb * SampleHistoryValidity;
 				
-					#if CONFIG_SCENE_COLOR_APLHA
+					#if CONFIG_SCENE_COLOR_ALPHA
 						SampleLowFrequencyColor.a = saturate(tsr_half(1.0) - saturate(tsr_half(1.0) - SampleLowFrequencyColor.a) * SampleTranslucencyAlpha);
 						SampleHighFrequencyColor.a = saturate(tsr_half(1.0) - saturate(tsr_half(1.0) - SampleHighFrequencyColor.a * InvSampleHistoryValidity) * SampleTranslucencyAlpha) * SampleHistoryValidity;
 					#endif
diff --git a/Engine/Shaders/Private/TemporalSuperResolution/TSRUpdateHistory.usf b/Engine/Shaders/Private/TemporalSuperResolution/TSRUpdateHistory.usf
--- a/Engine/Shaders/Private/TemporalSuperResolution/TSRUpdateHistory.usf
+++ b/Engine/Shaders/Private/TemporalSuperResolution/TSRUpdateHistory.usf
@@ -39,13 +39,13 @@
 #endif
 
 #if DIM_SEPARATE_TRANSLUCENCY
-	#define CONFIG_ACCUMULATE_TRANSLUCENCY_SEPARATLY 1
+	#define CONFIG_ACCUMULATE_TRANSLUCENCY_SEPARATELY 1
 	#define CONFIG_COMPOSE_TRANSLUCENCY 0
-	#define CONFIG_REJECT_TRAMSLUCENCY_ON_RESPONSIVE_AA_ONLY 0
+	#define CONFIG_REJECT_TRANSLUCENCY_ON_RESPONSIVE_AA_ONLY 0
 #else
-	#define CONFIG_ACCUMULATE_TRANSLUCENCY_SEPARATLY 0
+	#define CONFIG_ACCUMULATE_TRANSLUCENCY_SEPARATELY 0
 	#define CONFIG_COMPOSE_TRANSLUCENCY 1
-	#define CONFIG_REJECT_TRAMSLUCENCY_ON_RESPONSIVE_AA_ONLY 0
+	#define CONFIG_REJECT_TRANSLUCENCY_ON_RESPONSIVE_AA_ONLY 0
 #endif
 
 
@@ -69,7 +69,7 @@
 
 #define CONFIG_PREMULTIPLY_HISTORY 0
 
-#if CONFIG_COMPILE_FP16 && !CONFIG_SCENE_COLOR_APLHA
+#if CONFIG_COMPILE_FP16 && !CONFIG_SCENE_COLOR_ALPHA
 	#define CONFIG_MANUAL_LDS_SPILL 0 // TODO
 #else
 	#define CONFIG_MANUAL_LDS_SPILL 0
@@ -393,7 +393,7 @@ void MainCS(
 		tsr_short2x2(InputPPCo),
 		InputPixelPosMin, InputPixelPosMax);
 
-	// Fetch reprojections related informations.
+	// Fetch reprojection-related information.
 	float2x2 PrevScreenPos = ScreenPos;
 	#if CONFIG_GRAND_REPROJECTION
 		float2x2 GrandPrevScreenPos = ScreenPos;
@@ -645,12 +645,12 @@ void MainCS(
 		{
 			TranslucencyRejection = max(TranslucencyRejection, tsr_half(MinTranslucencyRejection));
 		
-			// Fetch wehther the pixel is responsive AA or not.
+			// Fetch whether the pixel is responsive AA or not.
 			#if CONFIG_RESPONSIVE_STENCIL
 			{
 				bIsResponsiveAAPixel = (SceneStencilRef & ResponsiveStencilMask) != 0;
 
-				#if CONFIG_REJECT_TRAMSLUCENCY_ON_RESPONSIVE_AA_ONLY
+				#if CONFIG_REJECT_TRANSLUCENCY_ON_RESPONSIVE_AA_ONLY
 					TranslucencyRejection = max(TranslucencyRejection, select(bIsResponsiveAAPixel, tsr_half(0.0).xx, tsr_half(1.0).xx));
 				#endif
 			}
@@ -666,8 +666,8 @@ void MainCS(
 		}
 	}
 
-	// Reject the translucency to respond quicker on firefliers.
-	#if CONFIG_TRANSLUCENCY_REJECTION && CONFIG_ACCUMULATE_TRANSLUCENCY_SEPARATLY == 0
+	// Reject the translucency to respond quicker on fireflies.
+	#if CONFIG_TRANSLUCENCY_REJECTION && CONFIG_ACCUMULATE_TRANSLUCENCY_SEPARATELY == 0
 		LowFrequencyRejection = min(LowFrequencyRejection, TranslucencyRejection);
 	#endif
 	
@@ -678,8 +678,8 @@ void MainCS(
 	#endif
 	tsr_half2 PrevHistoryValidity;
 	
-	tsr_halfCx2 PrevHistoryMommentMin;
-	tsr_halfCx2 PrevHistoryMommentMax;
+	tsr_halfCx2 PrevHistoryMomentMin;
+	tsr_halfCx2 PrevHistoryMomentMax;
 
 	tsr_half4x2 PrevHistoryTranslucency;
 	
@@ -779,7 +779,7 @@ void MainCS(
 						RawHighFrequency[i][PixelId]     = BilinearSampleColorHistory(PrevHistory_HighFrequency, SampleUV);
 						RawMetadata[i][PixelId]	         = PrevHistory_Metadata.SampleLevel(GlobalBilinearClampedSampler, SampleUV, 0);
 						
-						#if CONFIG_ACCUMULATE_TRANSLUCENCY_SEPARATLY
+						#if CONFIG_ACCUMULATE_TRANSLUCENCY_SEPARATELY
 							RawTranslucencyColor[i][PixelId] = BilinearSampleColorHistory(PrevHistory_Translucency, SampleUV);
 							RawTranslucencyAlpha[i][PixelId] = PrevHistory_TranslucencyAlpha.SampleLevel(GlobalBilinearClampedSampler, SampleUV, 0);
 						#endif
@@ -845,7 +845,7 @@ void MainCS(
 					AccumulateHighFrequency += dpv_scale(SampleHighFrequency, KernelWeight);
 					AccumulateMetadata  += dpv_scale(SampleMetadata, KernelWeight);
 					
-					#if CONFIG_ACCUMULATE_TRANSLUCENCY_SEPARATLY
+					#if CONFIG_ACCUMULATE_TRANSLUCENCY_SEPARATELY
 						AccumulateTranslucency  += dpv_scale(SampleTranslucency, KernelWeight);
 					#endif
 
@@ -877,7 +877,7 @@ void MainCS(
 				#endif
 			}
 			
-			#if CONFIG_SCENE_COLOR_APLHA
+			#if CONFIG_SCENE_COLOR_ALPHA
 				AccumulateHighFrequency = -dpv_min(-AccumulateHighFrequency, tsr_half(0.0).xxxx);
 			#else
 				AccumulateHighFrequency = -dpv_min(-AccumulateHighFrequency, tsr_half(0.0).xxx);
@@ -888,8 +888,8 @@ void MainCS(
 		
 		// Unpack history.
 		{
-			PrevHistoryMommentMin = HighFrequencyMin;
-			PrevHistoryMommentMax = HighFrequencyMax;
+			PrevHistoryMomentMin = HighFrequencyMin;
+			PrevHistoryMomentMax = HighFrequencyMax;
 
 			PrevHighFrequencyColor = AccumulateHighFrequency;
 			#if CONFIG_PREMULTIPLY_HISTORY
@@ -913,12 +913,12 @@ void MainCS(
 
 		// Correct history
 		{
-			CorrectExposure(PrevHistoryMommentMin[0]);
-			CorrectExposure(PrevHistoryMommentMin[1]);
-			CorrectExposure(PrevHistoryMommentMin[2]);
-			CorrectExposure(PrevHistoryMommentMax[0]);
-			CorrectExposure(PrevHistoryMommentMax[1]);
-			CorrectExposure(PrevHistoryMommentMax[2]);
+			CorrectExposure(PrevHistoryMomentMin[0]);
+			CorrectExposure(PrevHistoryMomentMin[1]);
+			CorrectExposure(PrevHistoryMomentMin[2]);
+			CorrectExposure(PrevHistoryMomentMax[0]);
+			CorrectExposure(PrevHistoryMomentMax[1]);
+			CorrectExposure(PrevHistoryMomentMax[2]);
 			CorrectExposure(PrevHighFrequencyColor[0]);
 			CorrectExposure(PrevHighFrequencyColor[1]);
 			CorrectExposure(PrevHighFrequencyColor[2]);
@@ -953,9 +953,11 @@ void MainCS(
 	tsr_halfCx2 InputMinColor;
 	tsr_halfCx2 InputMaxColor;
 	
+#if CONFIG_ACCUMULATE_TRANSLUCENCY_SEPARATELY
 	tsr_half4x2 FilteredTranslucencyColor;
 	tsr_half4x2 TranslucencyMinColor;
 	tsr_half4x2 TranslucencyMaxColor;
+#endif
 
 	tsr_half2 InputPixelAlignement;
 	tsr_half2 ClosestInputLuma4;
@@ -1004,7 +1006,7 @@ void MainCS(
 			tsr_half2x2 dInputKO = tsr_half2x2(InputPPCo - InputPPCk);
 
 
-			// Soften the spatial kernel when there is high noise in the input to make parrallax disocclusion of detailed geometry less distracting.
+			// Soften the spatial kernel when there is high noise in the input to make parallax disocclusion of detailed geometry less distracting.
 			tsr_half2 OffScreenInputToHistoryFactor = tsr_half(1.0) - tsr_half(0.5) * NoiseFiltering;
 
 			bool2 bIsRefining = PrevHistoryValidity != 0.0;
@@ -1051,7 +1053,7 @@ void MainCS(
 				for (uint PixelId = 0; PixelId < DPV_PIXEL_PER_LANE; PixelId++)
 				{
 					RawInputColorArray[SampleId][PixelId] = InputSceneColorTexture[dpv_access_pixel(InputSamplePixelPos, PixelId)];
-					#if CONFIG_ACCUMULATE_TRANSLUCENCY_SEPARATLY || CONFIG_COMPOSE_TRANSLUCENCY
+					#if CONFIG_ACCUMULATE_TRANSLUCENCY_SEPARATELY || CONFIG_COMPOSE_TRANSLUCENCY
 						RawTranslucencyColorArray[SampleId][PixelId] = InputSceneTranslucencyTexture[dpv_access_pixel(TranslucencySamplePixelPos, PixelId)];
 					#endif
 				}
@@ -1061,15 +1063,18 @@ void MainCS(
 		tsr_half2 FilteredInputColorWeight = tsr_half(0.0);
 		
 		FilteredInputColor = tsr_half(0.0);
-		#if CONFIG_SCENE_COLOR_APLHA
+		#if CONFIG_SCENE_COLOR_ALPHA
 			InputMinColor = dpv_interleave_mono_registers(+LargestSceneColorYCoCgA);
 			InputMaxColor = dpv_interleave_mono_registers(-LargestSceneColorYCoCgA);
 		#else
 			InputMinColor = dpv_interleave_mono_registers(+LargestSceneColorYCoCg);
 			InputMaxColor = dpv_interleave_mono_registers(-LargestSceneColorYCoCg);
 		#endif
+
+		#if CONFIG_ACCUMULATE_TRANSLUCENCY_SEPARATELY
 		TranslucencyMinColor = dpv_interleave_mono_registers(+LargestSceneColorRGBA);
 		TranslucencyMaxColor = dpv_interleave_mono_registers(-LargestSceneColorRGBA);
+		#endif
 
 		UNROLL_N(CONFIG_SAMPLES)
 		for (uint SampleId = 0; SampleId < CONFIG_SAMPLES; SampleId++)
@@ -1083,14 +1088,14 @@ void MainCS(
 				InputColor[0] = InputColor[0] * TranslucencyColor[3] + TranslucencyColor[0];
 				InputColor[1] = InputColor[1] * TranslucencyColor[3] + TranslucencyColor[1];
 				InputColor[2] = InputColor[2] * TranslucencyColor[3] + TranslucencyColor[2];
-				#if CONFIG_SCENE_COLOR_APLHA
+				#if CONFIG_SCENE_COLOR_ALPHA
 					InputColor[3] = InputColor[3] * TranslucencyColor[3];
 				#endif
 			}
 			#endif
 
 			#if CONFIG_SCENE_COLOR_OVERFLOW
-				#if CONFIG_SCENE_COLOR_APLHA
+				#if CONFIG_SCENE_COLOR_ALPHA
 					InputColor = min(InputColor, dpv_interleave_mono_registers(LargestSceneColorRGBA));
 				#else
 					InputColor = min(InputColor, dpv_interleave_mono_registers(LargestSceneColorRGB));
@@ -1102,15 +1107,17 @@ void MainCS(
 			FilteredInputColor       += dpv_scale((SampleSpatialWeight * ToneWeight), InputColor);
 			FilteredInputColorWeight += (SampleSpatialWeight * ToneWeight);
 
+			#if CONFIG_ACCUMULATE_TRANSLUCENCY_SEPARATELY
 			if (SampleId == 0)
 			{
 				FilteredTranslucencyColor = TranslucencyColor;
 			}
+			#endif
 
 			InputMinColor = min(InputMinColor, InputColor);
 			InputMaxColor = max(InputMaxColor, InputColor);
 
-			#if CONFIG_ACCUMULATE_TRANSLUCENCY_SEPARATLY
+			#if CONFIG_ACCUMULATE_TRANSLUCENCY_SEPARATELY
 			{
 				TranslucencyMinColor = min(TranslucencyMinColor, TranslucencyColor);
 				TranslucencyMaxColor = max(TranslucencyMaxColor, TranslucencyColor);
@@ -1127,13 +1134,15 @@ void MainCS(
 
 		#if CONFIG_SCENE_COLOR_OVERFLOW
 		{
-			#if CONFIG_SCENE_COLOR_APLHA
+			#if CONFIG_SCENE_COLOR_ALPHA
 				FilteredInputColor = min(FilteredInputColor, dpv_interleave_mono_registers(LargestSceneColorRGBA));
 			#else
 				FilteredInputColor = min(FilteredInputColor, dpv_interleave_mono_registers(LargestSceneColorRGB));
 			#endif
 			ClosestInputLuma4 = min(ClosestInputLuma4, LargestSceneColorYCoCg.x);
+			#if CONFIG_ACCUMULATE_TRANSLUCENCY_SEPARATELY
 			FilteredTranslucencyColor = min(FilteredTranslucencyColor, dpv_interleave_mono_registers(LargestSceneColorRGBA));
+			#endif
 		}
 		#endif
 	}
@@ -1155,7 +1164,7 @@ void MainCS(
 			SharedArray4[LocalGroupThreadIndex] = tsr_half4(InputMaxColor[0], InputMaxColor[1]);
 			SharedArray5[LocalGroupThreadIndex] = tsr_half4(InputMaxColor[2], OutputPixelVelocity);
 	
-			#if CONFIG_ACCUMULATE_TRANSLUCENCY_SEPARATLY
+			#if CONFIG_ACCUMULATE_TRANSLUCENCY_SEPARATELY
 				SharedArray6[LocalGroupThreadIndex] = tsr_half4(FilteredTranslucencyColor[0], FilteredTranslucencyColor[1]);
 				SharedArray7[LocalGroupThreadIndex] = tsr_half4(FilteredTranslucencyColor[2], FilteredTranslucencyColor[3]);
 				SharedArray8[LocalGroupThreadIndex] = tsr_half4(TranslucencyMinColor[0], TranslucencyMinColor[1]);
@@ -1184,7 +1193,7 @@ void MainCS(
 			tsr_half4 RawLDS4 = SharedArray4[LocalGroupThreadIndex];
 			tsr_half4 RawLDS5 = SharedArray5[LocalGroupThreadIndex];
 			
-			#if CONFIG_ACCUMULATE_TRANSLUCENCY_SEPARATLY
+			#if CONFIG_ACCUMULATE_TRANSLUCENCY_SEPARATELY
 				tsr_half4 RawLDS6 = SharedArray6[LocalGroupThreadIndex];
 				tsr_half4 RawLDS7 = SharedArray7[LocalGroupThreadIndex];
 				tsr_half4 RawLDS8 = SharedArray8[LocalGroupThreadIndex];
@@ -1201,7 +1210,7 @@ void MainCS(
 			InputPixelAlignement = RawLDS3.ba;
 			OutputPixelVelocity = RawLDS5.ba;
 			
-			#if CONFIG_ACCUMULATE_TRANSLUCENCY_SEPARATLY
+			#if CONFIG_ACCUMULATE_TRANSLUCENCY_SEPARATELY
 				FilteredTranslucencyColor = tsr_half4x2(RawLDS6.rg, RawLDS6.ba, RawLDS7.rg, RawLDS7.ba);
 				TranslucencyMinColor = tsr_half4x2(RawLDS8.rg, RawLDS8.ba, RawLDS9.rg, RawLDS9.ba);
 				TranslucencyMaxColor = tsr_half4x2(RawLDS10.rg, RawLDS10.ba, RawLDS11.rg, RawLDS11.ba);
@@ -1221,7 +1230,7 @@ void MainCS(
 		#endif
 	}
 	#endif
-	
+
 	// Contribute current frame input into the predictor for next frame.
 	const tsr_half Histeresis = rcp(tsr_half(MAX_SAMPLE_COUNT));
 	const tsr_half PredictionOnlyValidity = Histeresis * tsr_half(2.0);
@@ -1280,7 +1289,7 @@ void MainCS(
 		{
 			tsr_half2 MaxValidity = tsr_half(1.0) - (tsr_half(1.0) - PredictionOnlyValidity) * saturate(OutputPixelVelocity * tsr_half(InvWeightClampingPixelSpeed));
 
-			// Clamp up the max validity to favorise stability under motion on high contrast edges.
+			// Clamp up the max validity to favor stability under motion on high contrast edges.
 			#if 1
 			{
 				tsr_half2 PrevHistoryLuma = Luma4(BlendedPrevHighFrequencyColor);
@@ -1335,7 +1344,7 @@ void MainCS(
 				#error Unsupported
 			#endif
 
-			tsr_halfCx2 BoxSize = dpv_scale(PrevHistoryMommentMax - PrevHistoryMommentMin, tsr_half(0.5));
+			tsr_halfCx2 BoxSize = dpv_scale(PrevHistoryMomentMax - PrevHistoryMomentMin, tsr_half(0.5));
 
 			// Compute how much frame N-2 got overblur by reprojecting into N-1 and N versus directly to N.
 			tsr_halfCx2 PrevReprojectionOverblur = GrandReprojection - GrandPrevToPrev;
@@ -1369,7 +1378,7 @@ void MainCS(
 
 				tsr_halfCx2 CorrectedFinalHighFrequencyColor = FinalHighFrequencyColor + Correction;
 
-				FinalHighFrequencyColor = clamp(CorrectedFinalHighFrequencyColor, min(FinalHighFrequencyColor, PrevHistoryMommentMin), max(FinalHighFrequencyColor, PrevHistoryMommentMax));
+				FinalHighFrequencyColor = clamp(CorrectedFinalHighFrequencyColor, min(FinalHighFrequencyColor, PrevHistoryMomentMin), max(FinalHighFrequencyColor, PrevHistoryMomentMax));
 
 			#elif CONFIG_GRAND_REPROJECTION == 2
 				// Compute MeasuredResultant of the previous frame
@@ -1390,9 +1399,9 @@ void MainCS(
 		#endif
 	}
 	
+	#if CONFIG_ACCUMULATE_TRANSLUCENCY_SEPARATELY
 	// Compute the low frequency history
 	tsr_half4x2 FinalTranslucencyColor;
-	#if CONFIG_ACCUMULATE_TRANSLUCENCY_SEPARATLY
 	{
 		tsr_half2 TranslucencyLumaMin = Luma4(GetColorChannel(TranslucencyMinColor));
 		tsr_half2 TranslucencyLumaMax = Luma4(GetColorChannel(TranslucencyMaxColor));
@@ -1402,7 +1411,7 @@ void MainCS(
 
 		tsr_half4x2 ClampedPrevHistoryTranslucency = clamp(PrevHistoryTranslucency, TranslucencyMinColor, TranslucencyMaxColor);
 		
-		#if CONFIG_SCENE_COLOR_APLHA
+		#if CONFIG_SCENE_COLOR_ALPHA
 			tsr_half2 HighFrequencyTranslucencyRejection = MeasureRejectionFactor(
 				PrevHistoryTranslucency, ClampedPrevHistoryTranslucency,
 				FilteredTranslucencyColor, TranslucencyMinColor, TranslucencyMaxColor,
@@ -1489,18 +1498,18 @@ void MainCS(
 			FinalOutputColor = FinalHighFrequencyColor;
 		#endif
 
-		#if CONFIG_ACCUMULATE_TRANSLUCENCY_SEPARATLY
+		#if CONFIG_ACCUMULATE_TRANSLUCENCY_SEPARATELY
 			FinalOutputColor[0] = FinalOutputColor[0] * FinalTranslucencyColor[3] + FinalTranslucencyColor[0];
 			FinalOutputColor[1] = FinalOutputColor[1] * FinalTranslucencyColor[3] + FinalTranslucencyColor[1];
 			FinalOutputColor[2] = FinalOutputColor[2] * FinalTranslucencyColor[3] + FinalTranslucencyColor[2];
-			#if CONFIG_SCENE_COLOR_APLHA
+			#if CONFIG_SCENE_COLOR_ALPHA
 				FinalOutputColor[3] = FinalOutputColor[3] * FinalTranslucencyColor[3];
 			#endif
 		#endif
 		
 		// Ensure that alpha values that are expected to be opaque (but are only close to opaque) are forced to be opaque.
-		// (0.995 chosen to accomodate handling of 254/255)
-		#if CONFIG_SCENE_COLOR_APLHA
+		// (0.995 chosen to accommodate handling of 254/255)
+		#if CONFIG_SCENE_COLOR_ALPHA
 		{
 			FinalOutputColor[3] = select(FinalOutputColor[3] > 0.995, tsr_half(1.0), FinalOutputColor[3]);
 			FinalOutputColor[3] = select(FinalOutputColor[3] < 0.005, tsr_half(0.0), FinalOutputColor[3]);
@@ -1521,7 +1530,7 @@ void MainCS(
 
 	PLATFORM_SPECIFIC_ISOLATE
 	{
-		#if CONFIG_SCENE_COLOR_APLHA
+		#if CONFIG_SCENE_COLOR_ALPHA
 			const tsr_halfC ColorNull = tsr_half(0.0).xxxx;
 			const tsr_halfC ColorMax10BitsFloat = tsr_half(Max10BitsFloat).xxxx;
 		#else
@@ -1568,7 +1577,7 @@ void MainCS(
 			FinalOverblur = FinalOverblur + PremultipliedHighFrequencyColor;
 		#endif
 
-		// Stocastically round up or down using the hardware RWTexture2D truncation unit to take into precision
+		// Stochastically round up or down using the hardware RWTexture2D truncation unit to take into precision
 		// loss due to pixel format encoding.
 		#if CONFIG_ENABLE_STOCASTIC_QUANTIZATION
 		{
@@ -1588,12 +1597,12 @@ void MainCS(
 		}
 		#endif
 
-		// Protext from NaN and +Inf when writing out the history.
+		// Protect from NaN and +Inf when writing out the history.
 		{
 			PremultipliedHighFrequencyColor = -dpv_min(-PremultipliedHighFrequencyColor, ColorNull);
 			PremultipliedHighFrequencyColor = dpv_min(PremultipliedHighFrequencyColor, ColorMax10BitsFloat);	
 
-			#if CONFIG_ACCUMULATE_TRANSLUCENCY_SEPARATLY
+			#if CONFIG_ACCUMULATE_TRANSLUCENCY_SEPARATELY
 				FinalTranslucencyColor = -dpv_min(-FinalTranslucencyColor, tsr_half(0.0).xxxx);
 				FinalTranslucencyColor = dpv_min(FinalTranslucencyColor, tsr_half(Max10BitsFloat).xxxx);
 			#endif
@@ -1608,7 +1617,7 @@ void MainCS(
 				#endif
 				HistoryOutput_Metadata[dpv_lo(LocalHistoryPixelPos)]  = dpv_lo(FinalMetadata);
 				#if CONFIG_ACCUMULATE_TRANSLUCENCY_SEPARATLY
-					#if CONFIG_SCENE_COLOR_APLHA
+					#if CONFIG_SCENE_COLOR_ALPHA
 						HistoryOutput_ColorArray[tsr_short3(dpv_lo(LocalHistoryPixelPos), HistoryOutput_ArrayIndices_Translucency)] = dpv_lo(FinalTranslucencyColor);
 					#else
 						HistoryOutput_ColorArray[tsr_short3(dpv_lo(LocalHistoryPixelPos), HistoryOutput_ArrayIndices_Translucency)] = dpv_lo(FinalTranslucencyColor).rgb;
@@ -1631,8 +1640,8 @@ void MainCS(
 					HistoryOutput_ColorArray[tsr_short3(dpv_hi(LocalHistoryPixelPos), HistoryOutput_ArrayIndices_HighFrequency)] = dpv_hi(PremultipliedHighFrequencyColor);
 				#endif
 				HistoryOutput_Metadata[dpv_hi(LocalHistoryPixelPos)] = dpv_hi(FinalMetadata);
-				#if CONFIG_ACCUMULATE_TRANSLUCENCY_SEPARATLY
-					#if CONFIG_SCENE_COLOR_APLHA
+				#if CONFIG_ACCUMULATE_TRANSLUCENCY_SEPARATELY
+					#if CONFIG_SCENE_COLOR_ALPHA
 						HistoryOutput_ColorArray[tsr_short3(dpv_hi(LocalHistoryPixelPos), HistoryOutput_ArrayIndices_Translucency)] = dpv_hi(FinalTranslucencyColor);
 					#else
 						HistoryOutput_ColorArray[tsr_short3(dpv_hi(LocalHistoryPixelPos), HistoryOutput_ArrayIndices_Translucency)] = dpv_hi(FinalTranslucencyColor).rgb;
@@ -1683,7 +1692,7 @@ void MainCS(
 			BRANCH
 			if (bGenerateOutputMip1)
 			{
-				#if CONFIG_SCENE_COLOR_APLHA
+				#if CONFIG_SCENE_COLOR_ALPHA
 					SharedArray0[LocalGroupThreadIndex] = dpv_lo(HalfResOutput);
 					#if CONFIG_ENABLE_DUAL_PIXEL_VECTORIZATION
 						SharedArray1[LocalGroupThreadIndex] = dpv_hi(HalfResOutput);
@@ -1705,7 +1714,7 @@ void MainCS(
 					SharedArray1[LocalGroupThreadIndex] += SharedArray1[LocalGroupThreadIndex ^ 0x2];
 				#endif
 					
-				#if CONFIG_SCENE_COLOR_APLHA
+				#if CONFIG_SCENE_COLOR_ALPHA
 					#if CONFIG_ENABLE_DUAL_PIXEL_VECTORIZATION
 						HalfResOutput = dpv_interleave_registers(SharedArray0[LocalGroupThreadIndex], SharedArray1[LocalGroupThreadIndex]);
 					#else
