Separated Sky mie / rayleigh atmosphere lut for better control on cloud look. (start distance, fade region)

This special Mie/Ray only luts are only generated if nedeed (start distance are not 0 on the cloud ocmponent).

Measured on console with some cloud material
Off
- CameraVolumeLUT 0.08 ms
- Some Cloud Tracing   4.39 ms
On
- CameraVolumeLUT 0.09 ms
- Some Cloud Tracing   4.40 ms

=> Almost no impact on CameraVolumeLUT compute
=> cost of luts sampling when tracing material seems hiden in the tracing/marching cost noise.
=> More memory used in game: double buffered ray/mie rgba fp16 volumes are 2*2*(32*32*16)*4*2 = 512KB)

diff --git a/Engine/Shaders/Private/MaterialTemplate.ush b/Engine/Shaders/Private/MaterialTemplate.ush
--- a/Engine/Shaders/Private/MaterialTemplate.ush
+++ b/Engine/Shaders/Private/MaterialTemplate.ush
@@ -1672,6 +1672,7 @@ float4 MaterialExpressionSkyAtmosphereAerialPerspective(FMaterialPixelParameters
 	// distorted uvs and bad visuals with artefact. Only the distance computation can actually be benefit from the surface position specified here.
 	float4 NDCPosition = LWCMultiply(MakeLWCVector4(Parameters.AbsoluteWorldPosition, 1.0f), ResolvedView.WorldToClip);
 
+	const float NearFadeOutRangeInvDepthKm = 1.0 / 0.00001f; // 1 centimeter fade region
 	float4 AerialPerspective = GetAerialPerspectiveLuminanceTransmittance(
 		ResolvedView.RealTimeReflectionCapture, ResolvedView.SkyAtmosphereCameraAerialPerspectiveVolumeSizeAndInvSize,
 		NDCPosition, TranslatedWorldPosition - SkyCameraTranslatedWorldOrigin,
@@ -1681,7 +1682,8 @@ float4 MaterialExpressionSkyAtmosphereAerialPerspective(FMaterialPixelParameters
 		ResolvedView.SkyAtmosphereAerialPerspectiveStartDepthKm,
 		ResolvedView.SkyAtmosphereCameraAerialPerspectiveVolumeDepthSliceLengthKm,
 		ResolvedView.SkyAtmosphereCameraAerialPerspectiveVolumeDepthSliceLengthKmInv,
-		ResolvedView.OneOverPreExposure);
+		ResolvedView.OneOverPreExposure,
+		NearFadeOutRangeInvDepthKm);
 	return AerialPerspective;
 #else
 	return float4(0.0f, 0.0f, 0.0f, 1.0f); // RGB= null scattering, A= null transmittance
diff --git a/Engine/Shaders/Private/SkyAtmosphere.usf b/Engine/Shaders/Private/SkyAtmosphere.usf
--- a/Engine/Shaders/Private/SkyAtmosphere.usf
+++ b/Engine/Shaders/Private/SkyAtmosphere.usf
@@ -67,6 +67,10 @@
 #define VIRTUAL_SHADOW_MAP 0
 #endif
 
+#ifndef SEPARATE_MIE_RAYLEIGH_SCATTERING
+#define SEPARATE_MIE_RAYLEIGH_SCATTERING 0
+#endif
+
 #if SAMPLE_OPAQUE_SHADOW
 
 #define DYNAMICALLY_SHADOWED 1
@@ -374,8 +378,12 @@ float SkyAtmosphereNoise(float2 UV)
 struct SingleScatteringResult
 {
 	float3 L;						// Scattered light (luminance)
+	float3 LMieOnly;				// L but Mie scattering only
+	float3 LRayOnly;				// L but Rayleigh scattering only
 	float3 OpticalDepth;			// Optical depth (1/m)
 	float3 Transmittance;			// Transmittance in [0,1] (unitless)
+	float3 TransmittanceMieOnly;	// Transmittance in [0,1] (unitless) but Mie scattering only
+	float3 TransmittanceRayOnly;	// Transmittance in [0,1] (unitless) but Rayleigh scattering only
 	float3 MultiScatAs1;
 };
 
@@ -398,8 +406,12 @@ SingleScatteringResult IntegrateSingleScatteredLuminance(
 {
 	SingleScatteringResult Result;
 	Result.L = 0;
+	Result.LMieOnly = 0;
+	Result.LRayOnly = 0;
 	Result.OpticalDepth = 0;
 	Result.Transmittance = 1.0f;
+	Result.TransmittanceMieOnly = 1.0f;
+	Result.TransmittanceRayOnly = 1.0f;
 	Result.MultiScatAs1 = 0;
 
 	if (dot(WorldPos, WorldPos) <= Atmosphere.BottomRadiusKm * Atmosphere.BottomRadiusKm)
@@ -500,7 +512,11 @@ SingleScatteringResult IntegrateSingleScatteredLuminance(
 
 	// Ray march the atmosphere to integrate optical depth
 	float3 L = 0.0f;
+	float3 LMieOnly = 0.0f;
+	float3 LRayOnly = 0.0f;
 	float3 Throughput = 1.0f;
+	float3 ThroughputMieOnly = 1.0f;
+	float3 ThroughputRayOnly = 1.0f;
 	float3 OpticalDepth = 0.0f;
 	float t = 0.0f;
 	float tPrev = 0.0f;
@@ -556,31 +572,49 @@ SingleScatteringResult IntegrateSingleScatteredLuminance(
 		const float3 SampleTransmittance = exp(-SampleOpticalDepth);
 		OpticalDepth += SampleOpticalDepth;
 
+		// Transmittance Ray only and Mie only set half of ozone in rayleigh and half of ozone in mie parts.
+		// This is not great but I do not have any better solution. 
+		// Also most of the time Ozone is high in the atmosphere so it should be fine this way.
+		ThroughputMieOnly *= exp(-(Medium.ExtinctionMie + Medium.ExtinctionOzo) * dt * AerialPespectiveViewDistanceScale);
+		ThroughputRayOnly *= exp(-(Medium.ExtinctionRay + Medium.ExtinctionOzo) * dt * AerialPespectiveViewDistanceScale);
+
 		// Phase and transmittance for light 0
 		const float3 UpVector = P / PHeight;
 		float Light0ZenithCosAngle = dot(Light0Dir, UpVector);
 		float3 TransmittanceToLight0 = GetTransmittance(Light0ZenithCosAngle, PHeight);
 		float3 PhaseTimesScattering0;
+		float3 PhaseTimesScattering0MieOnly;
+		float3 PhaseTimesScattering0RayOnly;
 		if (MieRayPhase)
 		{
-			PhaseTimesScattering0 = Medium.ScatteringMie * MiePhaseValueLight0 + Medium.ScatteringRay * RayleighPhaseValueLight0;
+			PhaseTimesScattering0MieOnly= Medium.ScatteringMie * MiePhaseValueLight0;
+			PhaseTimesScattering0RayOnly= Medium.ScatteringRay * RayleighPhaseValueLight0;
+			PhaseTimesScattering0		= PhaseTimesScattering0MieOnly + PhaseTimesScattering0RayOnly;
 		}
 		else
 		{
-			PhaseTimesScattering0 = Medium.Scattering * uniformPhase;
+			PhaseTimesScattering0MieOnly= Medium.ScatteringMie * uniformPhase;
+			PhaseTimesScattering0RayOnly= Medium.ScatteringRay * uniformPhase;
+			PhaseTimesScattering0		= Medium.Scattering * uniformPhase;
 		}
 #if SECOND_ATMOSPHERE_LIGHT_ENABLED
 		// Phase and transmittance for light 1
 		float Light1ZenithCosAngle = dot(Light1Dir, UpVector);
 		float3 TransmittanceToLight1 = GetTransmittance(Light1ZenithCosAngle, PHeight);
 		float3 PhaseTimesScattering1;
+		float3 PhaseTimesScattering1MieOnly;
+		float3 PhaseTimesScattering1RayOnly;
 		if (MieRayPhase)
 		{
-			PhaseTimesScattering1 = Medium.ScatteringMie * MiePhaseValueLight1 + Medium.ScatteringRay * RayleighPhaseValueLight1;
+			PhaseTimesScattering1MieOnly= Medium.ScatteringMie * MiePhaseValueLight1;
+			PhaseTimesScattering1RayOnly= Medium.ScatteringRay * RayleighPhaseValueLight1;
+			PhaseTimesScattering1		= PhaseTimesScattering1MieOnly + PhaseTimesScattering1RayOnly;
 		}
 		else
 		{
-			PhaseTimesScattering1 = Medium.Scattering * uniformPhase;
+			PhaseTimesScattering1MieOnly= Medium.ScatteringMie * uniformPhase;
+			PhaseTimesScattering1RayOnly= Medium.ScatteringRay * uniformPhase;
+			PhaseTimesScattering1		= Medium.Scattering * uniformPhase;
 		}
 #endif
 
@@ -624,7 +658,9 @@ SingleScatteringResult IntegrateSingleScatteredLuminance(
 #endif
 		// MultiScatteredLuminance is already pre-exposed, atmospheric light contribution needs to be pre exposed
 		// Multi-scattering is also not affected by PlanetShadow or TransmittanceToLight because it contains diffuse light after single scattering.
-		float3 S = ExposedLight0Illuminance * (PlanetShadow0 * TransmittanceToLight0 * PhaseTimesScattering0 + MultiScatteredLuminance0 * Medium.Scattering);
+		float3 S		= ExposedLight0Illuminance * (PlanetShadow0 * TransmittanceToLight0 * PhaseTimesScattering0			+ MultiScatteredLuminance0 * Medium.Scattering);
+		float3 SMieOnly = ExposedLight0Illuminance * (PlanetShadow0 * TransmittanceToLight0 * PhaseTimesScattering0MieOnly	+ MultiScatteredLuminance0 * Medium.ScatteringMie);
+		float3 SRayOnly = ExposedLight0Illuminance * (PlanetShadow0 * TransmittanceToLight0 * PhaseTimesScattering0RayOnly	+ MultiScatteredLuminance0 * Medium.ScatteringRay);
 
 #if SECOND_ATMOSPHERE_LIGHT_ENABLED
 		float tPlanet1 = RaySphereIntersectNearest(P, Light1Dir, PlanetO + PLANET_RADIUS_OFFSET * UpVector, Atmosphere.BottomRadiusKm);
@@ -652,7 +688,9 @@ SingleScatteringResult IntegrateSingleScatteredLuminance(
 			VolumetricCloudCommonParameters.CloudShadowmapFarDepthKm[1].x, VolumetricCloudShadowMapTexture1, VolumetricCloudShadowMapTexture1Sampler, OutOpticalDepth3), VolumetricCloudShadowStrength1));
 #endif
 		//  Multi-scattering can work for the second light but it is disabled for the sake of performance.
-		S += ExposedLight1Illuminance * PlanetShadow1 * TransmittanceToLight1 * PhaseTimesScattering1;// +MultiScatteredLuminance * Medium.Scattering);
+		S		+= ExposedLight1Illuminance * PlanetShadow1 * TransmittanceToLight1 * PhaseTimesScattering1			;// +MultiScatteredLuminance * Medium.Scattering);
+		SMieOnly+= ExposedLight1Illuminance * PlanetShadow1 * TransmittanceToLight1 * PhaseTimesScattering1MieOnly	;// +MultiScatteredLuminance * Medium.Scattering);
+		SRayOnly+= ExposedLight1Illuminance * PlanetShadow1 * TransmittanceToLight1 * PhaseTimesScattering1RayOnly	;// +MultiScatteredLuminance * Medium.Scattering);
 #endif
 
 		// When using the power serie to accumulate all sattering order, serie r must be <1 for a serie to converge. 
@@ -670,13 +708,19 @@ SingleScatteringResult IntegrateSingleScatteredLuminance(
 #endif 
 
 #if 0
-		L += Throughput * S * dt;
-		Throughput *= SampleTransmittance;
+		L			+= Throughput * S * dt;
+		LMieOnly	+= Throughput * SMieOnly * dt;
+		LRayOnly	+= Throughput * SRayOnly * dt;
+		Throughput	*= SampleTransmittance;
 #else
 		// See slide 28 at http://www.frostbite.com/2015/08/physically-based-unified-volumetric-rendering-in-frostbite/ 
-		float3 Sint = (S - S * SampleTransmittance) / Medium.Extinction;	// integrate along the current step segment 
-		L += Throughput * Sint;														// accumulate and also take into account the transmittance from previous steps
-		Throughput *= SampleTransmittance;
+		float3 Sint			= (S        - S        * SampleTransmittance) / Medium.Extinction;	// integrate along the current step segment 
+		float3 SintMieOnly	= (SMieOnly - SMieOnly * SampleTransmittance) / Medium.Extinction;
+		float3 SintRayOnly	= (SRayOnly - SRayOnly * SampleTransmittance) / Medium.Extinction;
+		L			+= Throughput * Sint;														// accumulate and also take into account the transmittance from previous steps
+		LMieOnly	+= Throughput * SintMieOnly;
+		LRayOnly	+= Throughput * SintRayOnly;
+		Throughput	*= SampleTransmittance;
 #endif
 
 		tPrev = t;
@@ -704,9 +748,13 @@ SingleScatteringResult IntegrateSingleScatteredLuminance(
 #endif
 	}
 
-	Result.L = L;
-	Result.OpticalDepth = OpticalDepth;
-	Result.Transmittance = Throughput * PlanetOnOpaque;
+	Result.L					= L;
+	Result.LMieOnly				= LMieOnly;
+	Result.LRayOnly				= LRayOnly;
+	Result.OpticalDepth			= OpticalDepth;
+	Result.Transmittance		= Throughput * PlanetOnOpaque;
+	Result.TransmittanceMieOnly	= ThroughputMieOnly * PlanetOnOpaque;
+	Result.TransmittanceRayOnly	= ThroughputRayOnly * PlanetOnOpaque;
 
 	return Result;
 }
@@ -869,6 +917,7 @@ void RenderSkyAtmosphereRayMarchingPS(
 	float3 DepthBufferTranslatedWorldPos = GetScreenTranslatedWorldPos(SVPos, DeviceZ).xyz;
 	float4 NDCPosition = mul(float4(DepthBufferTranslatedWorldPos.xyz, 1), View.TranslatedWorldToClip);
 
+	const float NearFadeOutRangeInvDepthKm = 1.0 / 0.00001f; // 1 centimeter fade region
 	float4 AP = GetAerialPerspectiveLuminanceTransmittance(
 		ResolvedView.RealTimeReflectionCapture, ResolvedView.SkyAtmosphereCameraAerialPerspectiveVolumeSizeAndInvSize,
 		NDCPosition, (DepthBufferTranslatedWorldPos - GetCameraTranslatedWorldPos()) * CM_TO_SKY_UNIT,
@@ -878,7 +927,8 @@ void RenderSkyAtmosphereRayMarchingPS(
 		AerialPerspectiveStartDepthKm,
 		SkyAtmosphere.CameraAerialPerspectiveVolumeDepthSliceLengthKm,
 		SkyAtmosphere.CameraAerialPerspectiveVolumeDepthSliceLengthKmInv,
-		ViewOneOverPreExposure * OutputPreExposure);
+		ViewOneOverPreExposure * OutputPreExposure,
+		NearFadeOutRangeInvDepthKm);
 
 	PreExposedL += AP.rgb * LuminanceScale;
 	float Transmittance = AP.a;
@@ -1288,6 +1338,10 @@ void RenderDistantSkyLightLutCS(uint3 ThreadId : SV_DispatchThreadID)
 ////////////////////////////////////////////////////////////
 
 RWTexture3D<float4> CameraAerialPerspectiveVolumeUAV;
+#if SEPARATE_MIE_RAYLEIGH_SCATTERING
+RWTexture3D<float4> CameraAerialPerspectiveVolumeMieOnlyUAV;
+RWTexture3D<float4> CameraAerialPerspectiveVolumeRayOnlyUAV;
+#endif
 float RealTimeReflection360Mode;
 
 [numthreads(THREADGROUP_SIZE, THREADGROUP_SIZE, THREADGROUP_SIZE)]
@@ -1296,6 +1350,10 @@ void RenderCameraAerialPerspectiveVolumeCS(uint3 ThreadId : SV_DispatchThreadID)
 	if (SkyAtmosphere.FogShowFlagFactor <= 0.0f)
 	{ 
 		CameraAerialPerspectiveVolumeUAV[ThreadId] = float4(0.0f, 0.0f, 0.0f, 1.0f);
+	#if SEPARATE_MIE_RAYLEIGH_SCATTERING
+		CameraAerialPerspectiveVolumeMieOnlyUAV[ThreadId] = float4(0.0f, 0.0f, 0.0f, 1.0f);
+		CameraAerialPerspectiveVolumeRayOnlyUAV[ThreadId] = float4(0.0f, 0.0f, 0.0f, 1.0f);
+	#endif
 		return;
 	}
 
@@ -1364,6 +1422,10 @@ void RenderCameraAerialPerspectiveVolumeCS(uint3 ThreadId : SV_DispatchThreadID)
 			VoxelWorldPos = CamPos + 2.0f * (MiddlePointOnGround - CamPos);
 
 			//CameraAerialPerspectiveVolumeUAV[ThreadId] = float4(1, 0, 0, 0);
+			//#if SEPARATE_MIE_RAYLEIGH_SCATTERING
+			//	CameraAerialPerspectiveVolumeMieOnlyUAV[ThreadId] = float4(1, 0, 0, 0);
+			//	CameraAerialPerspectiveVolumeRayOnlyUAV[ThreadId] = float4(1, 0, 0, 0);
+			//#endif
 			//return; // debug
 		}
 		else if (UnderGround)
@@ -1373,6 +1435,10 @@ void RenderCameraAerialPerspectiveVolumeCS(uint3 ThreadId : SV_DispatchThreadID)
 			//VoxelWorldPos = CamPos + CameraToVoxelDir * PlanetNearT;		// better match but gives visual artefact as visible voxels on a simple plane at altitude 0
 
 			//CameraAerialPerspectiveVolumeUAV[ThreadId] = float4(0, 1, 0, 0);
+			//#if SEPARATE_MIE_RAYLEIGH_SCATTERING
+			//	CameraAerialPerspectiveVolumeMieOnlyUAV[ThreadId] = float4(0, 1, 0, 0);
+			//	CameraAerialPerspectiveVolumeRayOnlyUAV[ThreadId] = float4(0, 1, 0, 0);
+			//#endif
 			//return; // debug
 		}
 		 
@@ -1391,6 +1457,10 @@ void RenderCameraAerialPerspectiveVolumeCS(uint3 ThreadId : SV_DispatchThreadID)
 		{
 			// Ray is not intersecting the atmosphere
 			CameraAerialPerspectiveVolumeUAV[ThreadId] = float4(0.0f, 0.0f, 0.0f, 1.0f);
+		#if SEPARATE_MIE_RAYLEIGH_SCATTERING
+			CameraAerialPerspectiveVolumeMieOnlyUAV[ThreadId] = float4(0.0f, 0.0f, 0.0f, 1.0f);
+			CameraAerialPerspectiveVolumeRayOnlyUAV[ThreadId] = float4(0.0f, 0.0f, 0.0f, 1.0f);
+		#endif
 			return;
 		}
 		float LengthToAtmosphere = length(prevWorlPos - RayStartWorldPos);
@@ -1398,6 +1468,10 @@ void RenderCameraAerialPerspectiveVolumeCS(uint3 ThreadId : SV_DispatchThreadID)
 		{
 			// tMaxMax for this voxel is not within the planet atmosphere
 			CameraAerialPerspectiveVolumeUAV[ThreadId] = float4(0.0f, 0.0f, 0.0f, 1.0f);
+		#if SEPARATE_MIE_RAYLEIGH_SCATTERING
+			CameraAerialPerspectiveVolumeMieOnlyUAV[ThreadId] = float4(0.0f, 0.0f, 0.0f, 1.0f);
+			CameraAerialPerspectiveVolumeRayOnlyUAV[ThreadId] = float4(0.0f, 0.0f, 0.0f, 1.0f);
+		#endif
 			return;
 		}
 		// Now world position has been moved to the atmosphere boundary: we need to reduce tMaxMax accordingly. 
@@ -1424,6 +1498,12 @@ void RenderCameraAerialPerspectiveVolumeCS(uint3 ThreadId : SV_DispatchThreadID)
 
 	const float Transmittance = dot(ss.Transmittance, float3(1.0f / 3.0f, 1.0f / 3.0f, 1.0f / 3.0f));
 	CameraAerialPerspectiveVolumeUAV[ThreadId] = float4(ss.L, Transmittance);
+#if SEPARATE_MIE_RAYLEIGH_SCATTERING
+		const float TransmittanceMieOnly = dot(ss.TransmittanceMieOnly, float3(1.0f / 3.0f, 1.0f / 3.0f, 1.0f / 3.0f));
+		const float TransmittanceRayOnly = dot(ss.TransmittanceRayOnly, float3(1.0f / 3.0f, 1.0f / 3.0f, 1.0f / 3.0f));
+		CameraAerialPerspectiveVolumeMieOnlyUAV[ThreadId]	= float4(ss.LMieOnly, TransmittanceMieOnly);
+		CameraAerialPerspectiveVolumeRayOnlyUAV[ThreadId]	= float4(ss.LRayOnly, TransmittanceRayOnly);
+#endif
 }
 
 
diff --git a/Engine/Shaders/Private/SkyAtmosphereCommon.ush b/Engine/Shaders/Private/SkyAtmosphereCommon.ush
--- a/Engine/Shaders/Private/SkyAtmosphereCommon.ush
+++ b/Engine/Shaders/Private/SkyAtmosphereCommon.ush
@@ -40,7 +40,8 @@ float4 GetAerialPerspectiveLuminanceTransmittance(
 	float AerialPerspectiveVolumeStartDepth,
 	float AerialPerspectiveVolumeDepthSliceLengthKm,
 	float AerialPerspectiveVolumeDepthSliceLengthKmInv,
-	float OneOverExposure)
+	float OneOverExposure,
+	float NearFadeOutRangeInvDepthKm)
 {
 	float2 ScreenUv = (NDC.xy / NDC.ww) * float2(0.5f, -0.5f) + 0.5f;
 
@@ -58,6 +59,7 @@ float4 GetAerialPerspectiveLuminanceTransmittance(
 		// We multiply by weight to fade to 0 at depth 0. It works for luminance and opacity.
 		Weight = saturate(NonLinSlice*NonLinSlice * 2.0f); // Square to have a linear falloff from the change of distribution above
 	}
+	Weight *= saturate(tDepth * NearFadeOutRangeInvDepthKm);
 
 #if AP_SUPPORT_REALTIME_REFLECTION_CAPTURE
 	if (ViewIsRealTimeReflectionCapture)
@@ -110,6 +112,7 @@ float4 GetAerialPerspectiveLuminanceTransmittanceWithFogOver(
 	float AerialPerspectiveVolumeDepthSliceLengthKmInv,
 	float OneOverExposure, float4 FogToApplyOver)
 {
+	const float NearFadeOutRangeInvDepthKm = 1.0 / 0.00001f; // 1 centimeter fade region
 	float4 AP = GetAerialPerspectiveLuminanceTransmittance(
 		ViewIsRealTimeReflectionCapture, CameraAerialPerspectiveVolumeSizeAndInvSize,
 		NDC, WorldPositionRelativeToCamera,
@@ -119,7 +122,8 @@ float4 GetAerialPerspectiveLuminanceTransmittanceWithFogOver(
 		AerialPerspectiveVolumeStartDepth,
 		AerialPerspectiveVolumeDepthSliceLengthKm,
 		AerialPerspectiveVolumeDepthSliceLengthKmInv,
-		OneOverExposure);
+		OneOverExposure,
+		NearFadeOutRangeInvDepthKm);
 
 	float4 FinalFog;
 	// Apply any other fog OVER aerial perspective because AP is usually optically thiner.
diff --git a/Engine/Shaders/Private/VolumetricCloud.usf b/Engine/Shaders/Private/VolumetricCloud.usf
--- a/Engine/Shaders/Private/VolumetricCloud.usf
+++ b/Engine/Shaders/Private/VolumetricCloud.usf
@@ -417,6 +417,17 @@ Texture2D<float> StartTracingDistanceTexture;
 
 #endif // SHADER_RENDERVIEW_CS_WITH_EMPTY_SPACE_SKIPPING
 
+#define GETAERIALPERSPECTIVE(t, ts, x, o, f) GetAerialPerspectiveLuminanceTransmittance(\
+		ResolvedView.RealTimeReflectionCapture, ResolvedView.SkyAtmosphereCameraAerialPerspectiveVolumeSizeAndInvSize,\
+		ClipPos, (x - ResolvedView.TranslatedWorldCameraOrigin) * CM_TO_SKY_UNIT,\
+		t, ts,\
+		ResolvedView.SkyAtmosphereCameraAerialPerspectiveVolumeDepthResolutionInv,\
+		ResolvedView.SkyAtmosphereCameraAerialPerspectiveVolumeDepthResolution,\
+		o,\
+		ResolvedView.SkyAtmosphereCameraAerialPerspectiveVolumeDepthSliceLengthKm,\
+		ResolvedView.SkyAtmosphereCameraAerialPerspectiveVolumeDepthSliceLengthKmInv,\
+		ResolvedView.OneOverPreExposure,\
+		f);\
 
 void MainCommon(in FMaterialPixelParameters MaterialParameters, in float4 SvPosition, inout float4 OutColor0, inout float2 OutColor1)
 {
@@ -1301,18 +1312,25 @@ void MainCommon(in FMaterialPixelParameters MaterialParameters, in float4 SvPosi
 
 	if (RenderVolumetricCloudParameters.EnableAerialPerspectiveSampling && tAPWeightsSum > 0.0f)
 	{
-		// Apply AP only once according to the mean position within the participating media weighted by transmittance/visibility.
-		// This allows to apply AP only once per pixel instead of per sample.
-		float4 AerialPerspective = GetAerialPerspectiveLuminanceTransmittance(
-			ResolvedView.RealTimeReflectionCapture, ResolvedView.SkyAtmosphereCameraAerialPerspectiveVolumeSizeAndInvSize,
-			ClipPos, (FogSampleTranslatedWorldPosition - ResolvedView.TranslatedWorldCameraOrigin) * CM_TO_SKY_UNIT,
-			View.CameraAerialPerspectiveVolume, View.CameraAerialPerspectiveVolumeSampler,
-			ResolvedView.SkyAtmosphereCameraAerialPerspectiveVolumeDepthResolutionInv,
-			ResolvedView.SkyAtmosphereCameraAerialPerspectiveVolumeDepthResolution,
-			ResolvedView.SkyAtmosphereAerialPerspectiveStartDepthKm,
-			ResolvedView.SkyAtmosphereCameraAerialPerspectiveVolumeDepthSliceLengthKm,
-			ResolvedView.SkyAtmosphereCameraAerialPerspectiveVolumeDepthSliceLengthKmInv,
-			ResolvedView.OneOverPreExposure);
+		float4 AerialPerspective;
+
+		if (!ResolvedView.RealTimeReflectionCapture &&	// We do not generate a 360 CameraAerialPerspectiveVolume with separated Mie and Ray contribution as of today.
+			(RenderVolumetricCloudParameters.AerialPerspectiveRayOnlyStartDistanceKm + RenderVolumetricCloudParameters.AerialPerspectiveMieOnlyStartDistanceKm) > 0.0f)
+		{
+			float4 MieAP = GETAERIALPERSPECTIVE(View.CameraAerialPerspectiveVolumeMieOnly, View.CameraAerialPerspectiveVolumeMieOnlySampler, FogSampleTranslatedWorldPosition, 
+				max(ResolvedView.SkyAtmosphereAerialPerspectiveStartDepthKm, RenderVolumetricCloudParameters.AerialPerspectiveMieOnlyStartDistanceKm), RenderVolumetricCloudParameters.AerialPerspectiveMieOnlyFadeDistanceKmInv);
+			float4 RayAP = GETAERIALPERSPECTIVE(View.CameraAerialPerspectiveVolumeRayOnly, View.CameraAerialPerspectiveVolumeRayOnlySampler, FogSampleTranslatedWorldPosition, 
+				max(ResolvedView.SkyAtmosphereAerialPerspectiveStartDepthKm, RenderVolumetricCloudParameters.AerialPerspectiveRayOnlyStartDistanceKm), RenderVolumetricCloudParameters.AerialPerspectiveRayOnlyFadeDistanceKmInv);
+
+			AerialPerspective = float4(MieAP.rgb + RayAP.rgb, MieAP.a * RayAP.a);
+		}
+		else
+		{
+			const float NearFadeOutRangeInvDepthKm = 1.0 / 0.00001f; // 1 centimeter fade region
+			// Apply AP only once according to the mean position within the participating media weighted by transmittance/visibility.
+			// This allows to apply AP only once per pixel instead of per sample.
+			AerialPerspective = GETAERIALPERSPECTIVE(View.CameraAerialPerspectiveVolume, View.CameraAerialPerspectiveVolumeSampler, FogSampleTranslatedWorldPosition, ResolvedView.SkyAtmosphereAerialPerspectiveStartDepthKm, NearFadeOutRangeInvDepthKm);
+		}
 
 		// Apply aerial perspective OVER the cloud, assiming coverage is from the full transmittance.
 		// skip if MeanTransmittance is 1
diff --git a/Engine/Source/Runtime/Engine/Classes/Components/VolumetricCloudComponent.h b/Engine/Source/Runtime/Engine/Classes/Components/VolumetricCloudComponent.h
--- a/Engine/Source/Runtime/Engine/Classes/Components/VolumetricCloudComponent.h
+++ b/Engine/Source/Runtime/Engine/Classes/Components/VolumetricCloudComponent.h
@@ -121,6 +121,20 @@ class UVolumetricCloudComponent : public USceneComponent
 	 */
 	UPROPERTY(EditAnywhere, BlueprintReadOnly, Category = "Cloud Tracing", meta = (UIMin = "0.0", UIMax = "1.0", ClampMin = "0.0", ClampMax = "1.0", SliderExponent = 5.0))
 	float StopTracingTransmittanceThreshold;
+	
+	
+	/** Specify the aerial perspective start distance on cloud for Rayleigh scattering only. (kilometers) */
+	UPROPERTY(EditAnywhere, BlueprintReadOnly, interp, Category = "Art Direction", AdvancedDisplay, meta = (UIMin = 0.0, UIMax = 100.0, ClampMin = 0.0, SliderExponent = 2.0))
+	float AerialPespectiveRayleighScatteringStartDistance;
+	/** Specify the distance over which the Rayleigh scattering will linearly ramp up to ful leffect. (kilometers) */
+	UPROPERTY(EditAnywhere, BlueprintReadOnly, interp, Category = "Art Direction", AdvancedDisplay, meta = (UIMin = 0.0, UIMax = 100.0, ClampMin = 0.0, SliderExponent = 2.0))
+	float AerialPespectiveRayleighScatteringFadeDistance;
+	/** Specify the aerial perspective start distance on cloud for Mie scattering only. (kilometers) */
+	UPROPERTY(EditAnywhere, BlueprintReadOnly, interp, Category = "Art Direction", AdvancedDisplay, meta = (UIMin = 0.0, UIMax = 100.0, ClampMin = 0.0, SliderExponent = 2.0))
+	float AerialPespectiveMieScatteringStartDistance;
+	/** Specify the distance over which the Rayleigh scattering will linearly ramp up to ful leffect. (kilometers) */
+	UPROPERTY(EditAnywhere, BlueprintReadOnly, interp, Category = "Art Direction", AdvancedDisplay, meta = (UIMin = 0.0, UIMax = 100.0, ClampMin = 0.0, SliderExponent = 2.0))
+	float AerialPespectiveMieScatteringFadeDistance;
 
 
 	UFUNCTION(BlueprintCallable, Category = "Rendering")
diff --git a/Engine/Source/Runtime/Engine/Private/Components/VolumetricCloudComponent.cpp b/Engine/Source/Runtime/Engine/Private/Components/VolumetricCloudComponent.cpp
--- a/Engine/Source/Runtime/Engine/Private/Components/VolumetricCloudComponent.cpp
+++ b/Engine/Source/Runtime/Engine/Private/Components/VolumetricCloudComponent.cpp
@@ -54,6 +54,10 @@ UVolumetricCloudComponent::UVolumetricCloudComponent(const FObjectInitializer& O
 	, ShadowReflectionSampleCountScale_DEPRECATED(1.0f)
 	, ShadowTracingDistance(15.0f)
 	, StopTracingTransmittanceThreshold(0.005f)
+	, AerialPespectiveRayleighScatteringStartDistance(0.0f)
+	, AerialPespectiveRayleighScatteringFadeDistance(0.0f)
+	, AerialPespectiveMieScatteringStartDistance(0.0f)
+	, AerialPespectiveMieScatteringFadeDistance(0.0f)
 	, VolumetricCloudSceneProxy(nullptr)
 {
 	static ConstructorHelpers::FObjectFinder<UMaterialInterface> VolumetricCloudDefaultMaterialRef(TEXT("/Engine/EngineSky/VolumetricClouds/m_SimpleVolumetricCloud_Inst.m_SimpleVolumetricCloud_Inst"));
diff --git a/Engine/Source/Runtime/Engine/Private/SceneManagement.cpp b/Engine/Source/Runtime/Engine/Private/SceneManagement.cpp
--- a/Engine/Source/Runtime/Engine/Private/SceneManagement.cpp
+++ b/Engine/Source/Runtime/Engine/Private/SceneManagement.cpp
@@ -823,6 +823,10 @@ FViewUniformShaderParameters::FViewUniformShaderParameters()
 
 	CameraAerialPerspectiveVolume = GBlackAlpha1VolumeTexture->TextureRHI;
 	CameraAerialPerspectiveVolumeSampler = TStaticSamplerState<SF_Bilinear>::GetRHI();
+	CameraAerialPerspectiveVolumeMieOnly = GBlackAlpha1VolumeTexture->TextureRHI;
+	CameraAerialPerspectiveVolumeMieOnlySampler = TStaticSamplerState<SF_Bilinear>::GetRHI();
+	CameraAerialPerspectiveVolumeRayOnly = GBlackAlpha1VolumeTexture->TextureRHI;
+	CameraAerialPerspectiveVolumeRayOnlySampler = TStaticSamplerState<SF_Bilinear>::GetRHI();
 
 	PrimitiveSceneData = GIdentityPrimitiveBuffer.PrimitiveSceneDataBufferSRV;
 	InstanceSceneData = GIdentityPrimitiveBuffer.InstanceSceneDataBufferSRV;
diff --git a/Engine/Source/Runtime/Engine/Private/VolumetricCloudProxy.cpp b/Engine/Source/Runtime/Engine/Private/VolumetricCloudProxy.cpp
--- a/Engine/Source/Runtime/Engine/Private/VolumetricCloudProxy.cpp
+++ b/Engine/Source/Runtime/Engine/Private/VolumetricCloudProxy.cpp
@@ -25,6 +25,10 @@ FVolumetricCloudSceneProxy::FVolumetricCloudSceneProxy(const UVolumetricCloudCom
 	, ShadowReflectionViewSampleCountScale(InComponent->ShadowReflectionViewSampleCountScaleValue)
 	, ShadowTracingDistance(InComponent->ShadowTracingDistance)
 	, StopTracingTransmittanceThreshold(InComponent->StopTracingTransmittanceThreshold)
+	, AerialPespectiveMieScatteringStartDistance(InComponent->AerialPespectiveMieScatteringStartDistance) 
+	, AerialPespectiveMieScatteringFadeDistance(InComponent->AerialPespectiveMieScatteringFadeDistance) 
+	, AerialPespectiveRayleighScatteringStartDistance(InComponent->AerialPespectiveRayleighScatteringStartDistance)
+	, AerialPespectiveRayleighScatteringFadeDistance(InComponent->AerialPespectiveRayleighScatteringFadeDistance)
 	, CloudVolumeMaterial(InComponent->Material)
 {
 }
diff --git a/Engine/Source/Runtime/Engine/Public/SceneView.h b/Engine/Source/Runtime/Engine/Public/SceneView.h
--- a/Engine/Source/Runtime/Engine/Public/SceneView.h
+++ b/Engine/Source/Runtime/Engine/Public/SceneView.h
@@ -905,6 +905,10 @@ BEGIN_GLOBAL_SHADER_PARAMETER_STRUCT_WITH_CONSTRUCTOR(FViewUniformShaderParamete
 	SHADER_PARAMETER_SAMPLER(SamplerState, DistantSkyLightLutTextureSampler)
 	SHADER_PARAMETER_TEXTURE(Texture3D, CameraAerialPerspectiveVolume)
 	SHADER_PARAMETER_SAMPLER(SamplerState, CameraAerialPerspectiveVolumeSampler)
+	SHADER_PARAMETER_TEXTURE(Texture3D, CameraAerialPerspectiveVolumeMieOnly)
+	SHADER_PARAMETER_SAMPLER(SamplerState, CameraAerialPerspectiveVolumeMieOnlySampler)
+	SHADER_PARAMETER_TEXTURE(Texture3D, CameraAerialPerspectiveVolumeRayOnly)
+	SHADER_PARAMETER_SAMPLER(SamplerState, CameraAerialPerspectiveVolumeRayOnlySampler)
 	// Hair
 	SHADER_PARAMETER_TEXTURE(Texture3D, HairScatteringLUTTexture)
 	SHADER_PARAMETER_SAMPLER(SamplerState, HairScatteringLUTSampler)
diff --git a/Engine/Source/Runtime/Engine/Public/VolumetricCloudProxy.h b/Engine/Source/Runtime/Engine/Public/VolumetricCloudProxy.h
--- a/Engine/Source/Runtime/Engine/Public/VolumetricCloudProxy.h
+++ b/Engine/Source/Runtime/Engine/Public/VolumetricCloudProxy.h
@@ -47,6 +47,11 @@ public:
 	float ShadowReflectionViewSampleCountScale;
 	float ShadowTracingDistance;
 	float StopTracingTransmittanceThreshold;
+
+	float AerialPespectiveMieScatteringStartDistance;
+	float AerialPespectiveMieScatteringFadeDistance;
+	float AerialPespectiveRayleighScatteringStartDistance;
+	float AerialPespectiveRayleighScatteringFadeDistance;
 private:
 
 	UMaterialInterface* CloudVolumeMaterial = nullptr;
diff --git a/Engine/Source/Runtime/Renderer/Private/ReflectionEnvironmentRealTimeCapture.cpp b/Engine/Source/Runtime/Renderer/Private/ReflectionEnvironmentRealTimeCapture.cpp
--- a/Engine/Source/Runtime/Renderer/Private/ReflectionEnvironmentRealTimeCapture.cpp
+++ b/Engine/Source/Runtime/Renderer/Private/ReflectionEnvironmentRealTimeCapture.cpp
@@ -576,6 +576,8 @@ void FScene::AllocateAndCaptureFrameSkyEnvMap(
 
 				SkyRC.SkyAtmosphereViewLutTexture = BlackDummy2dTex;
 				SkyRC.SkyAtmosphereCameraAerialPerspectiveVolume = BlackDummy3dTex;
+				SkyRC.SkyAtmosphereCameraAerialPerspectiveVolumeMieOnly = BlackDummy3dTex;
+				SkyRC.SkyAtmosphereCameraAerialPerspectiveVolumeRayOnly = BlackDummy3dTex;
 
 				SkyRC.Viewport = FIntRect(FIntPoint(0, 0), FIntPoint(CubeWidth, CubeWidth));
 				SkyRC.bLightDiskEnabled = false;
diff --git a/Engine/Source/Runtime/Renderer/Private/RendererScene.cpp b/Engine/Source/Runtime/Renderer/Private/RendererScene.cpp
--- a/Engine/Source/Runtime/Renderer/Private/RendererScene.cpp
+++ b/Engine/Source/Runtime/Renderer/Private/RendererScene.cpp
@@ -288,17 +288,32 @@ FPersistentSkyAtmosphereData::FPersistentSkyAtmosphereData()
 	, CurrentDepthResolution(0)
 	, CurrentTextureAerialLUTFormat(PF_Unknown)
 	, CameraAerialPerspectiveVolumeIndex(0)
+	, bSeparatedAtmosphereMieRayLeigh(false)
 {
 }
-void FPersistentSkyAtmosphereData::InitialiseOrNextFrame(ERHIFeatureLevel::Type FeatureLevel, FPooledRenderTargetDesc& AerialPerspectiveDesc, FRHICommandListImmediate& RHICmdList)
+void FPersistentSkyAtmosphereData::InitialiseOrNextFrame(ERHIFeatureLevel::Type FeatureLevel, FPooledRenderTargetDesc& AerialPerspectiveDesc, FRHICommandListImmediate& RHICmdList, bool bSeparatedAtmosphereMieRayLeighIn)
 {
-	if (!bInitialised || (bInitialised && ((AerialPerspectiveDesc.Extent.X != CurrentScreenResolution) || (AerialPerspectiveDesc.Depth != CurrentDepthResolution) || (AerialPerspectiveDesc.Format != CurrentTextureAerialLUTFormat))))
+	if (!bInitialised || (bInitialised && ((AerialPerspectiveDesc.Extent.X != CurrentScreenResolution) || (AerialPerspectiveDesc.Depth != CurrentDepthResolution) 
+		|| (AerialPerspectiveDesc.Format != CurrentTextureAerialLUTFormat) || (bSeparatedAtmosphereMieRayLeigh != bSeparatedAtmosphereMieRayLeighIn))))
 	{
+		bSeparatedAtmosphereMieRayLeigh = bSeparatedAtmosphereMieRayLeighIn;
 		CameraAerialPerspectiveVolumeCount = FeatureLevel == ERHIFeatureLevel::ES3_1 ? 2 : 1;
 		for (int i = 0; i < CameraAerialPerspectiveVolumeCount; ++i)
 		{
 			GRenderTargetPool.FindFreeElement(RHICmdList, AerialPerspectiveDesc, CameraAerialPerspectiveVolumes[i], 
 				i==0 ? TEXT("SkyAtmosphere.CameraAPVolume0") : TEXT("SkyAtmosphere.CameraAPVolume1"));
+			if (bSeparatedAtmosphereMieRayLeigh)
+			{
+				GRenderTargetPool.FindFreeElement(RHICmdList, AerialPerspectiveDesc, CameraAerialPerspectiveVolumesMieOnly[i],
+					i == 0 ? TEXT("SkyAtmosphere.CameraAPVolumeMieOnly0") : TEXT("SkyAtmosphere.CameraAPVolumeMieOnly1"));
+				GRenderTargetPool.FindFreeElement(RHICmdList, AerialPerspectiveDesc, CameraAerialPerspectiveVolumesRayOnly[i],
+					i == 0 ? TEXT("SkyAtmosphere.CameraAPVolumeRayOnly0") : TEXT("SkyAtmosphere.CameraAPVolumeRayOnly1"));
+			}
+			else
+			{
+				CameraAerialPerspectiveVolumesMieOnly[i] = nullptr;
+				CameraAerialPerspectiveVolumesRayOnly[i] = nullptr;
+			}
 		}
 		bInitialised = true;
 		CurrentScreenResolution = AerialPerspectiveDesc.Extent.X;
@@ -313,6 +328,16 @@ TRefCountPtr<IPooledRenderTarget> FPersistentSkyAtmosphereData::GetCurrentCamera
 	check(CameraAerialPerspectiveVolumes[CameraAerialPerspectiveVolumeIndex].IsValid());
 	return CameraAerialPerspectiveVolumes[CameraAerialPerspectiveVolumeIndex];
 }
+TRefCountPtr<IPooledRenderTarget> FPersistentSkyAtmosphereData::GetCurrentCameraAerialPerspectiveVolumeMieOnly()
+{
+	check(CameraAerialPerspectiveVolumesMieOnly[CameraAerialPerspectiveVolumeIndex].IsValid());
+	return CameraAerialPerspectiveVolumesMieOnly[CameraAerialPerspectiveVolumeIndex];
+}
+TRefCountPtr<IPooledRenderTarget> FPersistentSkyAtmosphereData::GetCurrentCameraAerialPerspectiveVolumeRayOnly()
+{
+	check(CameraAerialPerspectiveVolumesRayOnly[CameraAerialPerspectiveVolumeIndex].IsValid());
+	return CameraAerialPerspectiveVolumesRayOnly[CameraAerialPerspectiveVolumeIndex];
+}
 
 /** Default constructor. */
 FSceneViewState::FSceneViewState(ERHIFeatureLevel::Type FeatureLevel, FSceneViewState* ShareOriginTarget)
@@ -836,6 +861,8 @@ uint64 FPersistentSkyAtmosphereData::GetGPUSizeBytes(bool bLogSizes) const
 	for (int32 VolumeIndex = 0; VolumeIndex < UE_ARRAY_COUNT(CameraAerialPerspectiveVolumes); VolumeIndex++)
 	{
 		TotalSize += GetRenderTargetGPUSizeBytes(CameraAerialPerspectiveVolumes[VolumeIndex], bLogSizes);
+		TotalSize += GetRenderTargetGPUSizeBytes(CameraAerialPerspectiveVolumesMieOnly[VolumeIndex], bLogSizes);
+		TotalSize += GetRenderTargetGPUSizeBytes(CameraAerialPerspectiveVolumesRayOnly[VolumeIndex], bLogSizes);
 	}
 	return TotalSize;
 }
diff --git a/Engine/Source/Runtime/Renderer/Private/ScenePrivate.h b/Engine/Source/Runtime/Renderer/Private/ScenePrivate.h
--- a/Engine/Source/Runtime/Renderer/Private/ScenePrivate.h
+++ b/Engine/Source/Runtime/Renderer/Private/ScenePrivate.h
@@ -597,9 +597,11 @@ struct FPersistentSkyAtmosphereData
 {
 	FPersistentSkyAtmosphereData();
 
-	void InitialiseOrNextFrame(ERHIFeatureLevel::Type FeatureLevel, FPooledRenderTargetDesc& AerialPerspectiveDesc, FRHICommandListImmediate& RHICmdList);
+	void InitialiseOrNextFrame(ERHIFeatureLevel::Type FeatureLevel, FPooledRenderTargetDesc& AerialPerspectiveDesc, FRHICommandListImmediate& RHICmdList, bool bSeparatedAtmosphereMieRayLeigh);
 
 	TRefCountPtr<IPooledRenderTarget> GetCurrentCameraAerialPerspectiveVolume();
+	TRefCountPtr<IPooledRenderTarget> GetCurrentCameraAerialPerspectiveVolumeMieOnly();
+	TRefCountPtr<IPooledRenderTarget> GetCurrentCameraAerialPerspectiveVolumeRayOnly();
 
 	uint64 GetGPUSizeBytes(bool bLogSizes) const;
 
@@ -609,8 +611,11 @@ private:
 	int32 CurrentDepthResolution;
 	EPixelFormat CurrentTextureAerialLUTFormat;
 	TRefCountPtr<IPooledRenderTarget> CameraAerialPerspectiveVolumes[2];
+	TRefCountPtr<IPooledRenderTarget> CameraAerialPerspectiveVolumesMieOnly[2];
+	TRefCountPtr<IPooledRenderTarget> CameraAerialPerspectiveVolumesRayOnly[2];
 	uint8 CameraAerialPerspectiveVolumeCount;
 	uint8 CameraAerialPerspectiveVolumeIndex;
+	bool bSeparatedAtmosphereMieRayLeigh;
 };
 
 struct FPersistentGlobalDistanceFieldData : public FThreadSafeRefCountedObject
diff --git a/Engine/Source/Runtime/Renderer/Private/SceneRendering.cpp b/Engine/Source/Runtime/Renderer/Private/SceneRendering.cpp
--- a/Engine/Source/Runtime/Renderer/Private/SceneRendering.cpp
+++ b/Engine/Source/Runtime/Renderer/Private/SceneRendering.cpp
@@ -876,6 +876,8 @@ void FViewInfo::Init()
 	VolumetricFogNearFadeInDistanceInv = 100000000.0f;
 
 	SkyAtmosphereCameraAerialPerspectiveVolume = nullptr;
+	SkyAtmosphereCameraAerialPerspectiveVolumeMieOnly = nullptr;
+	SkyAtmosphereCameraAerialPerspectiveVolumeRayOnly = nullptr;
 	SkyAtmosphereUniformShaderParameters = nullptr;
 
 	VolumetricCloudSkyAO = nullptr;
@@ -1276,6 +1278,8 @@ void FViewInfo::SetupUniformBufferParameters(
 	FRHITexture* TransmittanceLutTextureFound = nullptr;
 	FRHITexture* SkyViewLutTextureFound = nullptr;
 	FRHITexture* CameraAerialPerspectiveVolumeFound = nullptr;
+	FRHITexture* CameraAerialPerspectiveVolumeMieOnlyFound = nullptr;
+	FRHITexture* CameraAerialPerspectiveVolumeRayOnlyFound = nullptr;
 	FRHITexture* DistantSkyLightLutTextureFound = nullptr;
 	if (ShouldRenderSkyAtmosphere(Scene, Family->EngineShowFlags))
 	{
@@ -1302,6 +1306,14 @@ void FViewInfo::SetupUniformBufferParameters(
 		{
 			CameraAerialPerspectiveVolumeFound = this->SkyAtmosphereCameraAerialPerspectiveVolume->GetRHI();
 		}
+		if (this->SkyAtmosphereCameraAerialPerspectiveVolumeMieOnly.IsValid())
+		{
+			CameraAerialPerspectiveVolumeMieOnlyFound = this->SkyAtmosphereCameraAerialPerspectiveVolumeMieOnly->GetRHI();
+		}
+		if (this->SkyAtmosphereCameraAerialPerspectiveVolumeRayOnly.IsValid())
+		{
+			CameraAerialPerspectiveVolumeRayOnlyFound = this->SkyAtmosphereCameraAerialPerspectiveVolumeRayOnly->GetRHI();
+		}
 
 		float SkyViewLutWidth = 1.0f;
 		float SkyViewLutHeight = 1.0f;
@@ -1422,7 +1434,11 @@ void FViewInfo::SetupUniformBufferParameters(
 	ViewUniformShaderParameters.SkyViewLutTexture = OrBlack2DIfNull(SkyViewLutTextureFound);
 	ViewUniformShaderParameters.SkyViewLutTextureSampler = TStaticSamplerState<SF_Bilinear>::GetRHI();
 	ViewUniformShaderParameters.CameraAerialPerspectiveVolume = OrBlack3DAlpha1IfNull(CameraAerialPerspectiveVolumeFound);
+	ViewUniformShaderParameters.CameraAerialPerspectiveVolumeMieOnly = OrBlack3DAlpha1IfNull(CameraAerialPerspectiveVolumeMieOnlyFound);
+	ViewUniformShaderParameters.CameraAerialPerspectiveVolumeRayOnly = OrBlack3DAlpha1IfNull(CameraAerialPerspectiveVolumeRayOnlyFound);
 	ViewUniformShaderParameters.CameraAerialPerspectiveVolumeSampler = TStaticSamplerState<SF_Bilinear>::GetRHI();
+	ViewUniformShaderParameters.CameraAerialPerspectiveVolumeMieOnlySampler = TStaticSamplerState<SF_Bilinear>::GetRHI();
+	ViewUniformShaderParameters.CameraAerialPerspectiveVolumeRayOnlySampler = TStaticSamplerState<SF_Bilinear>::GetRHI();
 
 	ViewUniformShaderParameters.AtmosphereTransmittanceTexture = OrBlack2DIfNull(AtmosphereTransmittanceTexture);
 	ViewUniformShaderParameters.AtmosphereIrradianceTexture = OrBlack2DIfNull(AtmosphereIrradianceTexture);
diff --git a/Engine/Source/Runtime/Renderer/Private/SceneRendering.h b/Engine/Source/Runtime/Renderer/Private/SceneRendering.h
--- a/Engine/Source/Runtime/Renderer/Private/SceneRendering.h
+++ b/Engine/Source/Runtime/Renderer/Private/SceneRendering.h
@@ -1443,6 +1443,8 @@ public:
 
 	// Sky / Atmosphere textures (transient owned by this view info) and pointer to constants owned by SkyAtmosphere proxy.
 	TRefCountPtr<IPooledRenderTarget> SkyAtmosphereCameraAerialPerspectiveVolume;
+	TRefCountPtr<IPooledRenderTarget> SkyAtmosphereCameraAerialPerspectiveVolumeMieOnly;
+	TRefCountPtr<IPooledRenderTarget> SkyAtmosphereCameraAerialPerspectiveVolumeRayOnly;
 	TRefCountPtr<IPooledRenderTarget> SkyAtmosphereViewLutTexture;
 	const FAtmosphereUniformShaderParameters* SkyAtmosphereUniformShaderParameters;
 
diff --git a/Engine/Source/Runtime/Renderer/Private/SkyAtmosphereRendering.cpp b/Engine/Source/Runtime/Renderer/Private/SkyAtmosphereRendering.cpp
--- a/Engine/Source/Runtime/Renderer/Private/SkyAtmosphereRendering.cpp
+++ b/Engine/Source/Runtime/Renderer/Private/SkyAtmosphereRendering.cpp
@@ -595,6 +595,7 @@ class FSampleCloudShadow : SHADER_PERMUTATION_BOOL("SAMPLE_CLOUD_SHADOW");
 class FSampleCloudSkyAO : SHADER_PERMUTATION_BOOL("SAMPLE_CLOUD_SKYAO");
 class FAtmosphereOnClouds : SHADER_PERMUTATION_BOOL("SAMPLE_ATMOSPHERE_ON_CLOUDS");
 class FMSAASampleCount : SHADER_PERMUTATION_SPARSE_INT("MSAA_SAMPLE_COUNT", 1, 2, 4, 8);
+class FSeparateMieAndRayleighScattering : SHADER_PERMUTATION_BOOL("SEPARATE_MIE_RAYLEIGH_SCATTERING");
 
 }
 
@@ -905,7 +906,7 @@ class FRenderCameraAerialPerspectiveVolumeCS : public FGlobalShader
 	DECLARE_GLOBAL_SHADER(FRenderCameraAerialPerspectiveVolumeCS);
 	SHADER_USE_PARAMETER_STRUCT(FRenderCameraAerialPerspectiveVolumeCS, FGlobalShader);
 
-	using FPermutationDomain = TShaderPermutationDomain<FSampleCloudSkyAO, FSecondAtmosphereLight, FSampleOpaqueShadow, FSampleCloudShadow>;
+	using FPermutationDomain = TShaderPermutationDomain<FSampleCloudSkyAO, FSecondAtmosphereLight, FSampleOpaqueShadow, FSampleCloudShadow, FSeparateMieAndRayleighScattering>;
 
 public:
 	const static uint32 GroupSize = 4;
@@ -915,6 +916,8 @@ public:
 		SHADER_PARAMETER_STRUCT_REF(FSkyAtmosphereInternalCommonParameters, SkyAtmosphere)
 		SHADER_PARAMETER_STRUCT_REF(FViewUniformShaderParameters, ViewUniformBuffer)
 		SHADER_PARAMETER_RDG_TEXTURE_UAV(RWTexture3D<float4>, CameraAerialPerspectiveVolumeUAV)
+		SHADER_PARAMETER_RDG_TEXTURE_UAV(RWTexture3D<float4>, CameraAerialPerspectiveVolumeMieOnlyUAV)
+		SHADER_PARAMETER_RDG_TEXTURE_UAV(RWTexture3D<float4>, CameraAerialPerspectiveVolumeRayOnlyUAV)
 		SHADER_PARAMETER_RDG_TEXTURE(Texture2D<float3>, TransmittanceLutTexture)
 		SHADER_PARAMETER_RDG_TEXTURE(Texture2D<float3>, MultiScatteredLuminanceLutTexture)
 		SHADER_PARAMETER_RDG_TEXTURE(Texture2D<float2>, VolumetricCloudShadowMapTexture0)
@@ -1204,16 +1207,27 @@ void InitSkyAtmosphereForView(FRHICommandListImmediate& RHICmdList, const FScene
 		View.SkyAtmosphereUniformShaderParameters = SkyInfo.GetAtmosphereShaderParameters();
 		GRenderTargetPool.FindFreeElement(RHICmdList, SkyAtmosphereViewLutTextureDesc, View.SkyAtmosphereViewLutTexture, TEXT("SkyAtmosphere.SkyViewLut"));
 
+		const bool bSeparatedAtmosphereMieRayLeigh = VolumetricCloudWantsSeparatedAtmosphereMieRayLeigh(Scene);
 		if (View.ViewState)
 		{
 			// Per view non transient double buffered resources when needed.
-			View.ViewState->PersistentSkyAtmosphereData.InitialiseOrNextFrame(View.FeatureLevel, SkyAtmosphereCameraAerialPerspectiveVolumeDesc, RHICmdList);
+			View.ViewState->PersistentSkyAtmosphereData.InitialiseOrNextFrame(View.FeatureLevel, SkyAtmosphereCameraAerialPerspectiveVolumeDesc, RHICmdList, bSeparatedAtmosphereMieRayLeigh);
 			View.SkyAtmosphereCameraAerialPerspectiveVolume = View.ViewState->PersistentSkyAtmosphereData.GetCurrentCameraAerialPerspectiveVolume();
+			if (bSeparatedAtmosphereMieRayLeigh)
+			{
+				View.SkyAtmosphereCameraAerialPerspectiveVolumeMieOnly = View.ViewState->PersistentSkyAtmosphereData.GetCurrentCameraAerialPerspectiveVolumeMieOnly();
+				View.SkyAtmosphereCameraAerialPerspectiveVolumeRayOnly = View.ViewState->PersistentSkyAtmosphereData.GetCurrentCameraAerialPerspectiveVolumeRayOnly();
+			}
 		}
 		else
 		{
-			// Per frame transient resource.
+			// Per frame transient resource for reflection views
 			GRenderTargetPool.FindFreeElement(RHICmdList, SkyAtmosphereCameraAerialPerspectiveVolumeDesc, View.SkyAtmosphereCameraAerialPerspectiveVolume, TEXT("SkyAtmosphere.CameraAPVolume"));
+			if (bSeparatedAtmosphereMieRayLeigh)
+			{
+				GRenderTargetPool.FindFreeElement(RHICmdList, SkyAtmosphereCameraAerialPerspectiveVolumeDesc, View.SkyAtmosphereCameraAerialPerspectiveVolumeMieOnly, TEXT("SkyAtmosphere.CameraAPVolumeMieOnly"));
+				GRenderTargetPool.FindFreeElement(RHICmdList, SkyAtmosphereCameraAerialPerspectiveVolumeDesc, View.SkyAtmosphereCameraAerialPerspectiveVolumeRayOnly, TEXT("SkyAtmosphere.CameraAPVolumeRayOnly"));
+			}
 		}
 	}
 }
@@ -1296,6 +1310,7 @@ void FSceneRenderer::RenderSkyAtmosphereLookUpTables(FRDGBuilder& GraphBuilder)
 
 	const bool bHighQualityMultiScattering = CVarSkyAtmosphereMultiScatteringLUTHighQuality.GetValueOnRenderThread() > 0;
 	const bool bSecondAtmosphereLightEnabled = Scene->IsSecondAtmosphereLightEnabled();
+	const bool bSeparatedAtmosphereMieRayLeigh = VolumetricCloudWantsSeparatedAtmosphereMieRayLeigh(Scene);
 
 	FRHISamplerState* SamplerLinearClamp = TStaticSamplerState<SF_Trilinear>::GetRHI();
 
@@ -1519,6 +1534,7 @@ void FSceneRenderer::RenderSkyAtmosphereLookUpTables(FRDGBuilder& GraphBuilder)
 			PermutationVector.Set<FSecondAtmosphereLight>(bSecondAtmosphereLightEnabled);
 			PermutationVector.Set<FSampleOpaqueShadow>(false); // bShouldSampleOpaqueShadow); Off for now to not have to generate Light0Shadow and Light1Shadow
 			PermutationVector.Set<FSampleCloudShadow>(bShouldSampleCloudShadow);
+			PermutationVector.Set<FSeparateMieAndRayleighScattering>(false);	// Not for reflections for now
 			TShaderMapRef<FRenderCameraAerialPerspectiveVolumeCS> ComputeShader(View.ShaderMap, PermutationVector);
 
 			FRenderCameraAerialPerspectiveVolumeCS::FParameters* PassParameters = GraphBuilder.AllocParameters<FRenderCameraAerialPerspectiveVolumeCS::FParameters>();
@@ -1575,6 +1591,19 @@ void FSceneRenderer::RenderSkyAtmosphereLookUpTables(FRDGBuilder& GraphBuilder)
 		FRDGTextureRef SkyAtmosphereCameraAerialPerspectiveVolume = GraphBuilder.RegisterExternalTexture(View.SkyAtmosphereCameraAerialPerspectiveVolume);
 		FRDGTextureUAVRef SkyAtmosphereCameraAerialPerspectiveVolumeUAV = GraphBuilder.CreateUAV(FRDGTextureUAVDesc(SkyAtmosphereCameraAerialPerspectiveVolume, 0));
 
+		FRDGTextureRef SkyAtmosphereCameraAerialPerspectiveVolumeMieOnly = nullptr;
+		FRDGTextureUAVRef SkyAtmosphereCameraAerialPerspectiveVolumeMieOnlyUAV = nullptr;
+		FRDGTextureRef SkyAtmosphereCameraAerialPerspectiveVolumeRayOnly = nullptr;
+		FRDGTextureUAVRef SkyAtmosphereCameraAerialPerspectiveVolumeRayOnlyUAV = nullptr;
+		if (bSeparatedAtmosphereMieRayLeigh)
+		{
+			SkyAtmosphereCameraAerialPerspectiveVolumeMieOnly = GraphBuilder.RegisterExternalTexture(View.SkyAtmosphereCameraAerialPerspectiveVolumeMieOnly);
+			SkyAtmosphereCameraAerialPerspectiveVolumeMieOnlyUAV = GraphBuilder.CreateUAV(FRDGTextureUAVDesc(SkyAtmosphereCameraAerialPerspectiveVolumeMieOnly, 0));
+
+			SkyAtmosphereCameraAerialPerspectiveVolumeRayOnly = GraphBuilder.RegisterExternalTexture(View.SkyAtmosphereCameraAerialPerspectiveVolumeRayOnly);
+			SkyAtmosphereCameraAerialPerspectiveVolumeRayOnlyUAV = GraphBuilder.CreateUAV(FRDGTextureUAVDesc(SkyAtmosphereCameraAerialPerspectiveVolumeRayOnly, 0));
+		}
+
 		TRDGUniformBufferRef<FVolumeShadowingShaderParametersGlobal0> LightShadowShaderParams0UniformBuffer{};
 		TRDGUniformBufferRef<FVolumeShadowingShaderParametersGlobal1> LightShadowShaderParams1UniformBuffer{};
 		GetSkyAtmosphereLightsUniformBuffers(GraphBuilder, LightShadowShaderParams0UniformBuffer, LightShadowShaderParams1UniformBuffer,
@@ -1635,6 +1664,7 @@ void FSceneRenderer::RenderSkyAtmosphereLookUpTables(FRDGBuilder& GraphBuilder)
 			PermutationVector.Set<FSecondAtmosphereLight>(bSecondAtmosphereLightEnabled);
 			PermutationVector.Set<FSampleOpaqueShadow>(bShouldSampleOpaqueShadow);
 			PermutationVector.Set<FSampleCloudShadow>(bShouldSampleCloudShadow);
+			PermutationVector.Set<FSeparateMieAndRayleighScattering>(bSeparatedAtmosphereMieRayLeigh);
 			TShaderMapRef<FRenderCameraAerialPerspectiveVolumeCS> ComputeShader(View.ShaderMap, PermutationVector);
 
 			FRenderCameraAerialPerspectiveVolumeCS::FParameters * PassParameters = GraphBuilder.AllocParameters<FRenderCameraAerialPerspectiveVolumeCS::FParameters>();
@@ -1654,6 +1684,9 @@ void FSceneRenderer::RenderSkyAtmosphereLookUpTables(FRDGBuilder& GraphBuilder)
 			PassParameters->VolumetricCloudShadowStrength0 = CloudShadowOnAtmosphereStrength0;
 			PassParameters->VolumetricCloudShadowStrength1 = CloudShadowOnAtmosphereStrength1;
 			PassParameters->CameraAerialPerspectiveVolumeUAV = SkyAtmosphereCameraAerialPerspectiveVolumeUAV;
+			PassParameters->CameraAerialPerspectiveVolumeMieOnlyUAV = SkyAtmosphereCameraAerialPerspectiveVolumeMieOnlyUAV;
+			PassParameters->CameraAerialPerspectiveVolumeRayOnlyUAV = SkyAtmosphereCameraAerialPerspectiveVolumeRayOnlyUAV;
+
 			PassParameters->AerialPerspectiveStartDepthKm = AerialPerspectiveStartDepthInCm * CM_TO_KM;
 			PassParameters->RealTimeReflection360Mode = 0.0f;
 			PassParameters->Light0Shadow = LightShadowShaderParams0UniformBuffer;
diff --git a/Engine/Source/Runtime/Renderer/Private/SkyAtmosphereRendering.h b/Engine/Source/Runtime/Renderer/Private/SkyAtmosphereRendering.h
--- a/Engine/Source/Runtime/Renderer/Private/SkyAtmosphereRendering.h
+++ b/Engine/Source/Runtime/Renderer/Private/SkyAtmosphereRendering.h
@@ -85,6 +85,8 @@ struct FSkyAtmosphereRenderContext
 	FRDGTextureRef MultiScatteredLuminanceLut;
 	FRDGTextureRef SkyAtmosphereViewLutTexture;
 	FRDGTextureRef SkyAtmosphereCameraAerialPerspectiveVolume;
+	FRDGTextureRef SkyAtmosphereCameraAerialPerspectiveVolumeMieOnly;
+	FRDGTextureRef SkyAtmosphereCameraAerialPerspectiveVolumeRayOnly;
 
 	///////////////////////////////////
 	// Per view parameters
diff --git a/Engine/Source/Runtime/Renderer/Private/VolumetricCloudRendering.cpp b/Engine/Source/Runtime/Renderer/Private/VolumetricCloudRendering.cpp
--- a/Engine/Source/Runtime/Renderer/Private/VolumetricCloudRendering.cpp
+++ b/Engine/Source/Runtime/Renderer/Private/VolumetricCloudRendering.cpp
@@ -280,6 +280,17 @@ static bool ShouldRenderCloudShadowmap(const FLightSceneProxy* AtmosphericLight)
 	return CVarVolumetricCloudShadowMap.GetValueOnRenderThread() > 0 && AtmosphericLight && AtmosphericLight->GetCastCloudShadows();
 }
 
+bool VolumetricCloudWantsSeparatedAtmosphereMieRayLeigh(const FScene* Scene)
+{
+	const FVolumetricCloudRenderSceneInfo* VCloud = Scene->GetVolumetricCloudSceneInfo();
+	if (VCloud && CVarVolumetricCloud.GetValueOnRenderThread() > 0)
+	{
+		const FVolumetricCloudSceneProxy& VCloudProxy = VCloud->GetVolumetricCloudSceneProxy();
+		return VCloudProxy.AerialPespectiveMieScatteringStartDistance > 0.0f || VCloudProxy.AerialPespectiveRayleighScatteringStartDistance > 0.0f;
+	}
+	return false;
+}
+
 //////////////////////////////////////////////////////////////////////////
 
 static float GetVolumetricCloudShadowmapStrength(const FLightSceneProxy* AtmosphericLight)
@@ -560,6 +571,10 @@ BEGIN_GLOBAL_SHADER_PARAMETER_STRUCT(FRenderVolumetricCloudGlobalParameters, )
 	SHADER_PARAMETER(int32, EnableAtmosphericLightsSampling)
 	SHADER_PARAMETER(int32, EnableHeightFog)
 	SHADER_PARAMETER(float, EmptySpaceSkippingSliceDepth)
+	SHADER_PARAMETER(float, AerialPerspectiveRayOnlyStartDistanceKm)
+	SHADER_PARAMETER(float, AerialPerspectiveRayOnlyFadeDistanceKmInv)
+	SHADER_PARAMETER(float, AerialPerspectiveMieOnlyStartDistanceKm)
+	SHADER_PARAMETER(float, AerialPerspectiveMieOnlyFadeDistanceKmInv)
 	SHADER_PARAMETER_STRUCT_INCLUDE(FFogUniformParameters, FogStruct)
 	SHADER_PARAMETER_STRUCT(FForwardLightData, ForwardLightData)
 END_GLOBAL_SHADER_PARAMETER_STRUCT()
@@ -2120,6 +2135,11 @@ static TRDGUniformBufferRef<FRenderVolumetricCloudGlobalParameters> CreateCloudP
 	VolumetricCloudParams.EnableDistantSkyLightSampling = CVarVolumetricCloudEnableDistantSkyLightSampling.GetValueOnAnyThread() > 0 ? 1 : 0;
 	VolumetricCloudParams.EnableAtmosphericLightsSampling = CVarVolumetricCloudEnableAtmosphericLightsSampling.GetValueOnAnyThread() > 0 ? 1 : 0;
 
+	VolumetricCloudParams.AerialPerspectiveRayOnlyStartDistanceKm = CloudInfo.GetVolumetricCloudSceneProxy().AerialPespectiveRayleighScatteringStartDistance;
+	VolumetricCloudParams.AerialPerspectiveRayOnlyFadeDistanceKmInv = 1.0f / FMath::Max(CloudInfo.GetVolumetricCloudSceneProxy().AerialPespectiveRayleighScatteringFadeDistance, 0.00001f);// One centimeter minimum
+	VolumetricCloudParams.AerialPerspectiveMieOnlyStartDistanceKm = CloudInfo.GetVolumetricCloudSceneProxy().AerialPespectiveMieScatteringStartDistance;
+	VolumetricCloudParams.AerialPerspectiveMieOnlyFadeDistanceKmInv = 1.0f / FMath::Max(CloudInfo.GetVolumetricCloudSceneProxy().AerialPespectiveMieScatteringFadeDistance, 0.00001f);// One centimeter minimum
+
 	const FRDGTextureDesc& Desc = CloudRC.RenderTargets.Output[0].GetTexture()->Desc;
 	VolumetricCloudParams.OutputSizeInvSize = FVector4f(Desc.Extent.X, Desc.Extent.Y, 1.0f / Desc.Extent.X, 1.0f / Desc.Extent.Y);
 
@@ -2709,11 +2729,15 @@ bool FSceneRenderer::RenderVolumetricCloud(
 					{
 						SkyRC.SkyAtmosphereViewLutTexture = GraphBuilder.RegisterExternalTexture(ViewInfo.SkyAtmosphereViewLutTexture);
 						SkyRC.SkyAtmosphereCameraAerialPerspectiveVolume = GraphBuilder.RegisterExternalTexture(ViewInfo.SkyAtmosphereCameraAerialPerspectiveVolume);
+						SkyRC.SkyAtmosphereCameraAerialPerspectiveVolumeMieOnly = GraphBuilder.RegisterExternalTexture(ViewInfo.SkyAtmosphereCameraAerialPerspectiveVolumeMieOnly);
+						SkyRC.SkyAtmosphereCameraAerialPerspectiveVolumeRayOnly = GraphBuilder.RegisterExternalTexture(ViewInfo.SkyAtmosphereCameraAerialPerspectiveVolumeRayOnly);
 					}
 					else
 					{
 						SkyRC.SkyAtmosphereViewLutTexture = GSystemTextures.GetBlackDummy(GraphBuilder);
 						SkyRC.SkyAtmosphereCameraAerialPerspectiveVolume = GSystemTextures.GetVolumetricBlackDummy(GraphBuilder);
+						SkyRC.SkyAtmosphereCameraAerialPerspectiveVolumeMieOnly = GSystemTextures.GetVolumetricBlackDummy(GraphBuilder);
+						SkyRC.SkyAtmosphereCameraAerialPerspectiveVolumeRayOnly = GSystemTextures.GetVolumetricBlackDummy(GraphBuilder);
 					}
 
 					GetSkyAtmosphereLightsUniformBuffers(GraphBuilder, SkyRC.LightShadowShaderParams0UniformBuffer, SkyRC.LightShadowShaderParams1UniformBuffer,
diff --git a/Engine/Source/Runtime/Renderer/Private/VolumetricCloudRendering.h b/Engine/Source/Runtime/Renderer/Private/VolumetricCloudRendering.h
--- a/Engine/Source/Runtime/Renderer/Private/VolumetricCloudRendering.h
+++ b/Engine/Source/Runtime/Renderer/Private/VolumetricCloudRendering.h
@@ -108,6 +108,7 @@ bool ShouldRenderVolumetricCloud(const FScene* Scene, const FEngineShowFlags& En
 bool ShouldViewVisualizeVolumetricCloudConservativeDensity(const FViewInfo& ViewInfo, const FEngineShowFlags& EngineShowFlags);
 bool VolumetricCloudWantsToSampleLocalLights(const FScene* Scene, const FEngineShowFlags& EngineShowFlags);
 uint32 GetVolumetricCloudDebugViewMode(const FEngineShowFlags& ShowFlags);
+bool VolumetricCloudWantsSeparatedAtmosphereMieRayLeigh(const FScene* Scene);
 
 
 // Structure with data necessary to specify a cloud render.
