Fix shader compiler errors

diff --git a/Engine/Shaders/Private/HairStrands/HairStrandsVoxelPageAllocation.usf b/Engine/Shaders/Private/HairStrands/HairStrandsVoxelPageAllocation.usf
--- a/Engine/Shaders/Private/HairStrands/HairStrandsVoxelPageAllocation.usf
+++ b/Engine/Shaders/Private/HairStrands/HairStrandsVoxelPageAllocation.usf
@@ -143,7 +143,7 @@ void AllocatePageIndex(uint2 DispatchThreadId : SV_DispatchThreadID)
 	// Prepare indirect buffer for doing the actual page index allocation and filling the page index
 	{
 		const uint AllocatedPageIndexCount = PageIndexResolution.x * PageIndexResolution.y * PageIndexResolution.z;
-		WriteDispatchIndirectArgs(OutPageIndexAllocationIndirectBufferArgs, MacroGroupId, ceil(AllocatedPageIndexCount / float(IndirectDispatchGroupSize), 1, 1);
+		WriteDispatchIndirectArgs(OutPageIndexAllocationIndirectBufferArgs, MacroGroupId, ceil(AllocatedPageIndexCount / float(IndirectDispatchGroupSize)), 1, 1);
 	}
 
 	// Matrix generation
@@ -709,7 +709,7 @@ void AddIndirectBufferCS(uint GroupIndex : SV_GroupIndex, uint3 DispatchThreadId
 	const uint MaxGroupCountPerDim = 65535;
 	const uint DispatchX = AllocatedPageCount % MaxGroupCountPerDim;
 	const uint DispatchY = ceil(AllocatedPageCount / float(MaxGroupCountPerDim));
-	WriteDispatchIndirectArgs(OutIndirectArgsBuffer, MacroGroupId, DispatchX, DispatchY, ceil(VoxelCountPerPage / float(IndirectGroupSize));
+	WriteDispatchIndirectArgs(OutIndirectArgsBuffer, MacroGroupId, DispatchX, DispatchY, ceil(VoxelCountPerPage / float(IndirectGroupSize)));
 
 	// Copy current total to previous total
 	OutPageIndexGlobalCounter[1] = TotalAllocatedPageCount;
