Volumetric cloud - added permutation for sampling the corners of the froxel.
This is really helpfull to spread occupancy and increase robustness. Enabled by default.

diff --git a/Engine/Shaders/Private/VolumetricCloud.usf b/Engine/Shaders/Private/VolumetricCloud.usf
--- a/Engine/Shaders/Private/VolumetricCloud.usf
+++ b/Engine/Shaders/Private/VolumetricCloud.usf
@@ -1484,15 +1484,14 @@ FThreadData Reduce(FThreadData A, FThreadData B)
 	return Out;
 }
 
-[numthreads(THREADGROUP_SIZEX, THREADGROUP_SIZEY, THREADGROUP_SIZEZ)]
-void MainCS(uint3 DispatchThreadId : SV_DispatchThreadID, uint3 GroupThreadId : SV_GroupThreadID)
+float GetSliceDepth()
 {
-#if MATERIAL_VOLUMETRIC_CLOUD_EMPTY_SPACE_SKIPPING_OUTPUT
-
-	const float SliceDepth = RenderVolumetricCloudParameters.EmptySpaceSkippingSliceDepth;
-	ResolvedView = ResolveView();
+	return RenderVolumetricCloudParameters.EmptySpaceSkippingSliceDepth;
+}
 
-	float2 PixelUV = float2(DispatchThreadId.x + 0.5, DispatchThreadId.y + 0.5) / StartTracingDistanceTextureResolution;
+const float GetSvPositionEmptySpaceSkippingValue(in ViewState ResolvedView, in uint3 DispatchThreadId, in uint3 GroupThreadId, in float3 SampleFroxelOffset)
+{
+	float2 PixelUV = float2(DispatchThreadId.x + SampleFroxelOffset.x, DispatchThreadId.y + SampleFroxelOffset.y) / StartTracingDistanceTextureResolution;
 	float4 SvPosition = float4(PixelUV * ResolvedView.ViewSizeAndInvSize.xy, 0.1, 1);
 
 	FMaterialPixelParameters MaterialParameters = MakeInitializedMaterialPixelParameters();
@@ -1503,13 +1502,50 @@ void MainCS(uint3 DispatchThreadId : SV_DispatchThreadID, uint3 GroupThreadId :
 		RenderVolumetricCloudParameters.CloudLayerCenterKm, RenderVolumetricCloudParameters.PlanetRadiusKm,
 		RenderVolumetricCloudParameters.BottomRadiusKm, RenderVolumetricCloudParameters.TopRadiusKm);
 
-	const float SampleDistanceFromCamera = (float(GroupThreadId.z) + 0.5f) * SliceDepth;
-	const float3 Raydir =-MaterialParameters.CameraVector;
+	const float SampleDistanceFromCamera = (float(GroupThreadId.z) + SampleFroxelOffset.z) * GetSliceDepth();
+	const float3 Raydir = -MaterialParameters.CameraVector;
 	UpdateMaterialCloudParam(MaterialParameters, LWCAdd(ResolvedView.WorldCameraOrigin, LWCPromote(Raydir * SampleDistanceFromCamera)),
-		ResolvedView, CloudLayerParams, TRACING_SHADOW_DISTANCE_OFF, RenderVolumetricCloudParameters.EmptySpaceSkippingSliceDepth);
+		ResolvedView, CloudLayerParams, TRACING_SHADOW_DISTANCE_OFF, GetSliceDepth());
+
+	const float CloudEmptySpaceSkippingOutput0 = GetVolumetricCloudEmptySpaceSkippingOutput0(MaterialParameters);
+	return CloudEmptySpaceSkippingOutput0;
+}
+
+[numthreads(THREADGROUP_SIZEX, THREADGROUP_SIZEY, THREADGROUP_SIZEZ)]
+void MainCS(uint3 DispatchThreadId : SV_DispatchThreadID, uint3 GroupThreadId : SV_GroupThreadID)
+{
+#if MATERIAL_VOLUMETRIC_CLOUD_EMPTY_SPACE_SKIPPING_OUTPUT
+
+	ResolvedView = ResolveView();
+
+	float CloudEmptySpaceSkippingOutput0 = 0.0f;
+
+	// Center samples
+	{
+		const float3 SampleFroxelOffset = float3(0.5f, 0.5f, 0.5f);
+		CloudEmptySpaceSkippingOutput0 = max(CloudEmptySpaceSkippingOutput0, GetSvPositionEmptySpaceSkippingValue(ResolvedView, DispatchThreadId, GroupThreadId, SampleFroxelOffset));
+	}
+
+#if EMPTY_SPACE_SKIPPING_SAMPLE_CORNERS
+
+	// Corner samples
+	UNROLL
+	for(float x=0; x<=1; ++x)
+	{
+		UNROLL
+		for(float y=0; y<=1; ++y)
+		{
+			UNROLL
+			for(float z=0; z<=1; ++z)
+			{
+				CloudEmptySpaceSkippingOutput0 = max(CloudEmptySpaceSkippingOutput0, GetSvPositionEmptySpaceSkippingValue(ResolvedView, DispatchThreadId, GroupThreadId, float3(x, y, z)));
+			}
+		}
+	}
+
+#endif
 
-	float CloudEmptySpaceSkippingOutput0 = GetVolumetricCloudEmptySpaceSkippingOutput0(MaterialParameters);
-	GroupData[GroupThreadId.x][GroupThreadId.y][GroupThreadId.z].FrontDepthKm = max(0.0f, CloudEmptySpaceSkippingOutput0 > 0.0f ? (float(GroupThreadId.z) + StartTracingSliceBias) * SliceDepth * CENTIMETER_TO_KILOMETER : GetDefaultDepth());
+	GroupData[GroupThreadId.x][GroupThreadId.y][GroupThreadId.z].FrontDepthKm = max(0.0f, CloudEmptySpaceSkippingOutput0 > 0.0f ? (float(GroupThreadId.z) + StartTracingSliceBias) * GetSliceDepth() * CENTIMETER_TO_KILOMETER : GetDefaultDepth());
 
 	GroupMemoryBarrierWithGroupSync();
 
diff --git a/Engine/Source/Runtime/Renderer/Private/VolumetricCloudRendering.cpp b/Engine/Source/Runtime/Renderer/Private/VolumetricCloudRendering.cpp
--- a/Engine/Source/Runtime/Renderer/Private/VolumetricCloudRendering.cpp
+++ b/Engine/Source/Runtime/Renderer/Private/VolumetricCloudRendering.cpp
@@ -224,6 +224,11 @@ static TAutoConsoleVariable<float> CVarVolumetricCloudEmptySpaceSkippingStartTra
 	TEXT("The number of slices to bias the start depth with. A valuie of -1 means a bias of one slice towards the view point."),
 	ECVF_RenderThreadSafe | ECVF_Scalability);
 
+static TAutoConsoleVariable<int32> CVarVolumetricCloudEmptySpaceSkippingSampleCorners(
+	TEXT("r.VolumetricCloud.EmptySpaceSkipping.SampleCorners"), 1,
+	TEXT("0 means center samples only, >0 means corner are also sampled."),
+	ECVF_RenderThreadSafe | ECVF_Scalability);
+
 ////////////////////////////////////////////////////////////////////////// 
 
 
@@ -900,8 +905,9 @@ class FRenderVolumetricCloudEmptySpaceSkippingCS : public FMeshMaterialShader
 	SHADER_USE_PARAMETER_STRUCT_WITH_LEGACY_BASE(FRenderVolumetricCloudEmptySpaceSkippingCS, FMeshMaterialShader)
 
 	class FCloudEmptySpaceSkippingDebug : SHADER_PERMUTATION_BOOL("EMPTY_SPACE_SKIPPING_DEBUG");
+	class FCloudEmptySpaceSkippingSampleCorners : SHADER_PERMUTATION_BOOL("EMPTY_SPACE_SKIPPING_SAMPLE_CORNERS");
 
-	using FPermutationDomain = TShaderPermutationDomain<FCloudEmptySpaceSkippingDebug>;
+	using FPermutationDomain = TShaderPermutationDomain<FCloudEmptySpaceSkippingDebug, FCloudEmptySpaceSkippingSampleCorners>;
 
 	BEGIN_SHADER_PARAMETER_STRUCT(FParameters, )
 		SHADER_PARAMETER(FVector2f, StartTracingDistanceTextureResolution)
@@ -2289,8 +2295,11 @@ void FSceneRenderer::RenderVolumetricCloudsInternal(FRDGBuilder& GraphBuilder, F
 			const FMaterial* MaterialResource = &CloudVolumeMaterialProxy->GetMaterialWithFallback(Scene->GetFeatureLevel(), MaterialRenderProxy);
 			MaterialRenderProxy = MaterialRenderProxy ? MaterialRenderProxy : CloudVolumeMaterialProxy;
 
+			const bool bSampleCorners = CVarVolumetricCloudEmptySpaceSkippingSampleCorners.GetValueOnRenderThread() > 0;
+
 			typename FRenderVolumetricCloudEmptySpaceSkippingCS::FPermutationDomain PermutationVector;
 			PermutationVector.Set<typename FRenderVolumetricCloudEmptySpaceSkippingCS::FCloudEmptySpaceSkippingDebug>(bCloudEmptySpaceSkippingDebug);
+			PermutationVector.Set<typename FRenderVolumetricCloudEmptySpaceSkippingCS::FCloudEmptySpaceSkippingSampleCorners>(bSampleCorners);
 
 			TShaderRef<FRenderVolumetricCloudEmptySpaceSkippingCS> ComputeShader = MaterialResource->GetShader<FRenderVolumetricCloudEmptySpaceSkippingCS>(&FLocalVertexFactory::StaticType, PermutationVector, false);
 
