Volumetric cloud -  Improved mode 0 edge fixing
New min max depth buffer
New cloud tracing permutation to output cloud lumiance/transmittance for min and max depth for mode 0.
Better half resolution reconstruction: less holes and not more trailes behind animation and object on disocluded areas.

#swarm https://p4-swarm.epicgames.net/reviews/22826906

diff --git a/Engine/Shaders/Private/DownsampleDepthPixelShader.usf b/Engine/Shaders/Private/DownsampleDepthPixelShader.usf
--- a/Engine/Shaders/Private/DownsampleDepthPixelShader.usf
+++ b/Engine/Shaders/Private/DownsampleDepthPixelShader.usf
@@ -19,17 +19,47 @@ float2 SourceMaxUV;
 float2 DestinationResolution;
 uint DownsampleDepthFilter;
 
+// Only valid when using OUTPUT_MIN_AND_MAX_DEPTH
+int4 DstPixelCoordMinAndMax;
+int4 SrcPixelCoordMinAndMax;
+
 float GetDeviceZ(float2 UV)
 {
 	return DepthTexture.Sample(GlobalPointClampedSampler, min(UV, SourceMaxUV));
 }
 
 void Main(
-	noperspective float4 UV : TEXCOORD0,
-	out float OutDepth : SV_DEPTH,
-	out float4 OutColor : SV_Target0)
+	noperspective float4 OutUVAndScreenPos : TEXCOORD0,
+	in float4 SVPos : SV_POSITION,
+#if OUTPUT_MIN_AND_MAX_DEPTH
+	out float4 OutColor : SV_Target0
+#else
+	out float OutDepth : SV_DEPTH
+#endif
+)
 {
-	OutColor = 0;
+#if OUTPUT_MIN_AND_MAX_DEPTH
+
+	// We do things differently when downsampling min and max depth because it is executed from half to quarter resolution.
+	// And in this case OutUVAndScreenPos is not valid because it uses view rect data.
+
+	// Pixel coordinate for the output 
+	const int2 DstPixelPos = SVPos;
+	// Lower left corner from the input 2x2 quad
+	const int2 PixelPos = DstPixelPos * 2;
+
+	float DeviceZ0 = DepthTexture.Load(uint3(clamp(PixelPos + int2(0,0), SrcPixelCoordMinAndMax.xy, SrcPixelCoordMinAndMax.zw), 0));
+	float DeviceZ1 = DepthTexture.Load(uint3(clamp(PixelPos + int2(0,1), SrcPixelCoordMinAndMax.xy, SrcPixelCoordMinAndMax.zw), 0));
+	float DeviceZ2 = DepthTexture.Load(uint3(clamp(PixelPos + int2(1,1), SrcPixelCoordMinAndMax.xy, SrcPixelCoordMinAndMax.zw), 0));
+	float DeviceZ3 = DepthTexture.Load(uint3(clamp(PixelPos + int2(1,0), SrcPixelCoordMinAndMax.xy, SrcPixelCoordMinAndMax.zw), 0));
+
+	const float MaxDeviceZ = max(max(DeviceZ0, DeviceZ1), max(DeviceZ2, DeviceZ3));
+	const float MinDeviceZ = min(min(DeviceZ0, DeviceZ1), min(DeviceZ2, DeviceZ3));
+	OutColor = float4(MinDeviceZ, MaxDeviceZ, 0.0f, 0.0f);
+
+#else // OUTPUT_MIN_AND_MAX_DEPTH
+
+	const float2 UV = OutUVAndScreenPos.xy;
 	float Depth = 0;
 
 	// Lookup the four DeviceZ's of the full resolution pixels corresponding to this low resolution pixel
@@ -87,6 +117,8 @@ void Main(
 	}
 
 	OutDepth = Depth;
+
+#endif // OUTPUT_MIN_AND_MAX_DEPTH
 }
 
 #ifndef STENCIL_LIGHTING_CHANNELS_SHIFT 
diff --git a/Engine/Shaders/Private/VolumetricCloud.usf b/Engine/Shaders/Private/VolumetricCloud.usf
--- a/Engine/Shaders/Private/VolumetricCloud.usf
+++ b/Engine/Shaders/Private/VolumetricCloud.usf
@@ -59,6 +59,10 @@
 #define CLOUD_DEBUG_VIEW_MODE 0
 #endif
 
+#ifndef CLOUD_MIN_AND_MAX_DEPTH
+#define CLOUD_MIN_AND_MAX_DEPTH 0
+#endif
+
 // Draw debug sample positions
 #if CLOUD_DEBUG_SAMPLES || EMPTY_SPACE_SKIPPING_DEBUG
 #include "ShaderPrint.ush"
@@ -429,7 +433,12 @@ Texture2D<float> StartTracingDistanceTexture;
 		ResolvedView.OneOverPreExposure,\
 		f);\
 
-void MainCommon(in FMaterialPixelParameters MaterialParameters, in float4 SvPosition, inout float4 OutColor0, inout float2 OutColor1)
+void MainCommon(in FMaterialPixelParameters MaterialParameters, in float4 SvPosition, 
+	inout float4 OutColor0
+#if CLOUD_MIN_AND_MAX_DEPTH
+	, inout float4 OutColor1 
+#endif
+	, inout float4 OutDepth)
 {
 	ResolvedView = ResolveView();
 
@@ -448,7 +457,7 @@ void MainCommon(in FMaterialPixelParameters MaterialParameters, in float4 SvPosi
 	//
 	
 	OutColor0					= float4(0.0f, 0.0f, 0.0f, 1.0f);
-	OutColor1					= MaxHalfFloat;
+	OutDepth					= MaxHalfFloat;
 
 	FPixelMaterialInputs PixelMaterialInputs	= (FPixelMaterialInputs)0;
 	CalcMaterialParameters(MaterialParameters, PixelMaterialInputs, SvPosition, true);
@@ -464,6 +473,24 @@ void MainCommon(in FMaterialPixelParameters MaterialParameters, in float4 SvPosi
 	float tAPWeightedSum		= 0.0f;
 	float tAPWeightsSum			= 0.0f;
 
+#if CLOUD_MIN_AND_MAX_DEPTH
+	OutColor1 = float4(0.0f, 0.0f, 0.0f, 1.0f);
+	float3 ClosestLuminance = 0.0f;
+	MATVEC ClosestTransmittanceToView = 1.0f;
+	float ClostestTAPWeightedSum = 0.0f;
+	float ClostestTAPWeightsSum = 0.0f;
+#endif // CLOUD_MIN_AND_MAX_DEPTH
+
+	uint4 TracingCoordToZbufferCoordScaleBias = RenderVolumetricCloudParameters.TracingCoordToZbufferCoordScaleBias;
+	uint2 SceneDepthTextureCoord = uint2(SvPosition.xy - 0.5) * TracingCoordToZbufferCoordScaleBias.xy + TracingCoordToZbufferCoordScaleBias.zw;
+	SceneDepthTextureCoord = clamp(SceneDepthTextureCoord, RenderVolumetricCloudParameters.SceneDepthTextureMinMaxCoord.xy, RenderVolumetricCloudParameters.SceneDepthTextureMinMaxCoord.zw);
+	{
+		float DeviceZ = RenderVolumetricCloudParameters.SceneDepthTexture.Load(uint3(SceneDepthTextureCoord, 0)).r;
+		float3 DepthBufferTranslatedWorldPos = SvPositionToTranslatedWorld(float4(SvPosition.xy, DeviceZ, 1.0));
+		const float TDepthBufferKm = min(length(DepthBufferTranslatedWorldPos * CENTIMETER_TO_KILOMETER - RayTranslatedWorldOriginKm), MaxHalfFloat);
+		OutDepth.xyzw = float4(TDepthBufferKm, TDepthBufferKm, DeviceZ, DeviceZ);
+	}
+
 	////////////////////////////////////////////////////////////////////////////////////////////////////
 	// Begin the cloud ray marching and setup the variable above
 	////////////////////////////////////////////////////////////////////////////////////////////////////
@@ -523,16 +550,34 @@ void MainCommon(in FMaterialPixelParameters MaterialParameters, in float4 SvPosi
 		// Sample the depth buffer and update tracing distance
 		//
 
-		uint4 TracingCoordToZbufferCoordScaleBias = RenderVolumetricCloudParameters.TracingCoordToZbufferCoordScaleBias;
-		uint2 SceneDepthTextureCoord = uint2(SvPosition.xy - 0.5) * TracingCoordToZbufferCoordScaleBias.xy + TracingCoordToZbufferCoordScaleBias.zw;
-		SceneDepthTextureCoord = clamp(SceneDepthTextureCoord, RenderVolumetricCloudParameters.SceneDepthTextureMinMaxCoord.xy, RenderVolumetricCloudParameters.SceneDepthTextureMinMaxCoord.zw);
+#if CLOUD_MIN_AND_MAX_DEPTH
+
+		float2 MinAndMaxDepth = RenderVolumetricCloudParameters.SceneDepthMinAndMaxTexture.Load(uint3(SceneDepthTextureCoord / 2, 0)).rg;
+	#if HAS_INVERTED_Z_BUFFER
+		MinAndMaxDepth = max(0.000000000001, MinAndMaxDepth.yx); // x is clostest depth, y is furthest depth
+	#endif
+
+		// Trace up to furthest depth.
+		float DeviceZ = MinAndMaxDepth.y;
+
+		float3 ClostestDepthBufferTranslatedWorldPos = SvPositionToTranslatedWorld(float4(SvPosition.xy, MinAndMaxDepth.x, 1.0));
+		const float TClosestDepthBufferCm = min(length(ClostestDepthBufferTranslatedWorldPos * CENTIMETER_TO_KILOMETER - RayTranslatedWorldOriginKm), MaxHalfFloat) * KILOMETER_TO_CENTIMETER;
+
+		OutDepth.zw = MinAndMaxDepth;
+
+#else // CLOUD_MIN_AND_MAX_DEPTH
+
 		float DeviceZ = RenderVolumetricCloudParameters.SceneDepthTexture.Load(uint3(SceneDepthTextureCoord, 0)).r;
-#if HAS_INVERTED_Z_BUFFER
+	#if HAS_INVERTED_Z_BUFFER
 		DeviceZ = max(0.000000000001, DeviceZ);
-#endif
+	#endif
+
+#endif // CLOUD_MIN_AND_MAX_DEPTH
+
 		float3 DepthBufferTranslatedWorldPos = SvPositionToTranslatedWorld(float4(SvPosition.xy, DeviceZ, 1.0));
 		const float TDepthBufferKm = min(length(DepthBufferTranslatedWorldPos * CENTIMETER_TO_KILOMETER - RayTranslatedWorldOriginKm), MaxHalfFloat);
-		OutColor1.y = TDepthBufferKm;
+
+		OutDepth.y = TDepthBufferKm;
 
 		//
 		// Skip tracing if the depth buffer is in front of the cloud layer front interface. If inside the layer, clamp the tracing to the depth buffer
@@ -726,6 +771,10 @@ void MainCommon(in FMaterialPixelParameters MaterialParameters, in float4 SvPosi
 			int ms;
 			float3	  SampleTranslatedWorldPosition = RayTranslatedWorldOriginKm + t * Raydir;
 			FLWCVector3 SampleWorldPosition = LWCSubtract(LWCPromote(SampleTranslatedWorldPosition), PrimaryView.PreViewTranslation);
+#if CLOUD_MIN_AND_MAX_DEPTH
+			const bool SampleContributeToClosest = t <= TClosestDepthBufferCm;
+#endif // CLOUD_MIN_AND_MAX
+			 
 
 			//////////////////////////////
 			// Update FMaterialPixelParameters according to current sample.
@@ -1082,6 +1131,11 @@ void MainCommon(in FMaterialPixelParameters MaterialParameters, in float4 SvPosi
 				float tAPWeight = min(TransmittanceToView.MATSUF.r, min(TransmittanceToView.MATSUF.g, TransmittanceToView.MATSUF.b));
 				tAPWeightedSum += t * tAPWeight;
 				tAPWeightsSum += tAPWeight;
+
+			#if CLOUD_MIN_AND_MAX_DEPTH
+				ClostestTAPWeightedSum += SampleContributeToClosest ? t * tAPWeight : 0.0f;
+				ClostestTAPWeightsSum  += SampleContributeToClosest ?     tAPWeight : 0.0f;
+			#endif // CLOUD_MIN_AND_MAX_DEPTH
 			}
 
 
@@ -1232,21 +1286,31 @@ void MainCommon(in FMaterialPixelParameters MaterialParameters, in float4 SvPosi
 				const float3 ScatteredLuminance = SunSkyLuminance * ScatteringCoefficients + EmissiveLuminance;
 
 #if 0
-	// Default iterative integration
+				// Default iterative integration
 				const MATVEC SafePathSegmentTransmittance = exp(-ExtinctionCoefficients * dtMeters);
-				Luminance += TransmittanceToView * ScatteredLuminance * dtMeters;
+				const float3 LuminanceContribution = TransmittanceToView * ScatteredLuminance * dtMeters;
+				Luminance += LuminanceContribution;
+				
+			#if CLOUD_MIN_AND_MAX_DEPTH
+				ClosestLuminance += SampleContributeToClosest ? LuminanceContribution : 0.0;
+			#endif // CLOUD_MIN_AND_MAX_DEPTH
 
 #elif 1
-	// Improved scattering integration. See slide 28 of "Physically Based and Unified Volumetric Rendering in Frostbite"
-	// Automatically works with emission. Emissive color is considered as a constant luminance emitted in all direction uniformly.
+				// Improved scattering integration. See slide 28 of "Physically Based and Unified Volumetric Rendering in Frostbite"
+				// Automatically works with emission. Emissive color is considered as a constant luminance emitted in all direction uniformly.
 				MATVEC SafeExtinctionThreshold = 0.000001f;
 				const MATVEC SafeExtinctionCoefficients = max(SafeExtinctionThreshold, ExtinctionCoefficients);
 				const MATVEC SafePathSegmentTransmittance = exp(-SafeExtinctionCoefficients * dtMeters);
 				float3 LuminanceIntegral = (ScatteredLuminance - ScatteredLuminance * SafePathSegmentTransmittance) / SafeExtinctionCoefficients;
-				Luminance += TransmittanceToView * LuminanceIntegral;
+				const float3 LuminanceContribution = TransmittanceToView * LuminanceIntegral;
+				Luminance += LuminanceContribution;
+				
+			#if CLOUD_MIN_AND_MAX_DEPTH
+				ClosestLuminance += SampleContributeToClosest ? LuminanceContribution : 0.0;
+			#endif // CLOUD_MIN_AND_MAX_DEPTH
 
 #else
-	// Another integration that is not working nicely yet: explodes or need to use high extinction threshold 0.000001 to make it safe. And that has an impact on the final image.
+				// Another integration that is not working nicely yet: explodes or need to use high extinction threshold 0.000001 to make it safe. And that has an impact on the final image.
 				MATVEC S1 = TransmittanceToLight0 * ScatteringCoefficients;
 				MATVEC SafeExtinctionThreshold = 0.000001f;
 				const MATVEC SafeExtinctionCoefficients = max(SafeExtinctionThreshold, ExtinctionCoefficients);
@@ -1258,14 +1322,22 @@ void MainCommon(in FMaterialPixelParameters MaterialParameters, in float4 SvPosi
 				/// Not handled yet: second light and emissive 
 				float3 SkyScatteredLuminance = SunSkyLuminance * ScatteringCoefficients;
 				float3 SkyLuminanceIntegral = (ScatteredLuminance - ScatteredLuminance * SafePathSegmentTransmittance) / SafeExtinctionCoefficients;
-				Luminance += TransmittanceToView * ((AnalyticalShadowedScattered * Light0IlluminanceFinal * PMPC.Phase0[ms]) + (ms == 0 ? SkyLuminanceIntegral : float3(0.0f, 0.0f, 0.0f)));
+				const float3 LuminanceContribution = TransmittanceToView * ((AnalyticalShadowedScattered * Light0IlluminanceFinal * PMPC.Phase0[ms]) + (ms == 0 ? SkyLuminanceIntegral : float3(0.0f, 0.0f, 0.0f)));
+				Luminance += LuminanceContribution;
 				if (ms == 0) S0 = S1;
-
+				
+			#if CLOUD_MIN_AND_MAX_DEPTH
+				ClosestLuminance += SampleContributeToClosest ? LuminanceContribution : 0.0;
+			#endif // CLOUD_MIN_AND_MAX_DEPTH
 #endif
 
 				if (ms == 0)
 				{
 					TransmittanceToView *= SafePathSegmentTransmittance;
+
+				#if CLOUD_MIN_AND_MAX_DEPTH
+					ClosestTransmittanceToView *= TClosestDepthBufferCm > t ? SafePathSegmentTransmittance : 1.0;
+				#endif // CLOUD_MIN_AND_MAX_DEPTH
 				}
 			}
 
@@ -1306,6 +1378,10 @@ void MainCommon(in FMaterialPixelParameters MaterialParameters, in float4 SvPosi
 	const float NoCloudDepth = TMax;
 	const float tAP = tAPWeightsSum==0.0f ? NoCloudDepth : tAPWeightedSum / max(0.0000000001f, tAPWeightsSum);
 	float MeanTransmittance = dot(TransmittanceToView.MATSUF, 1.0f / 3.0f);
+#if CLOUD_MIN_AND_MAX_DEPTH
+	const float ClosestTAP = ClostestTAPWeightsSum == 0.0f ? NoCloudDepth : ClostestTAPWeightedSum / max(0.0000000001f, ClostestTAPWeightsSum);
+	float ClosestMeanTransmittance = dot(ClosestTransmittanceToView.MATSUF, 1.0f / 3.0f);
+#endif // CLOUD_MIN_AND_MAX
 
 	float3 FogSampleTranslatedWorldPosition = RayTranslatedWorldOriginKm + tAP * Raydir;
 	float4 ClipPos = mul(float4(FogSampleTranslatedWorldPosition, 1.0f), PrimaryView.TranslatedWorldToClip);
@@ -1360,6 +1436,64 @@ void MainCommon(in FMaterialPixelParameters MaterialParameters, in float4 SvPosi
 	}
 
 
+	// Apply atmospehre and and fog on closest
+#if CLOUD_MIN_AND_MAX_DEPTH
+	{
+		float3 FogSampleTranslatedWorldPosition = RayTranslatedWorldOriginKm + ClosestTAP * Raydir;
+		float4 ClipPos = mul(float4(FogSampleTranslatedWorldPosition, 1.0f), PrimaryView.TranslatedWorldToClip);
+
+		if (RenderVolumetricCloudParameters.EnableAerialPerspectiveSampling && ClostestTAPWeightsSum > 0.0f)
+		{
+			float4 AerialPerspective;
+
+			if (!ResolvedView.RealTimeReflectionCapture &&	// We do not generate a 360 CameraAerialPerspectiveVolume with separated Mie and Ray contribution as of today.
+				(RenderVolumetricCloudParameters.AerialPerspectiveRayOnlyStartDistanceKm + RenderVolumetricCloudParameters.AerialPerspectiveMieOnlyStartDistanceKm) > 0.0f)
+			{
+				float4 MieAP = GETAERIALPERSPECTIVE(View.CameraAerialPerspectiveVolumeMieOnly, View.CameraAerialPerspectiveVolumeMieOnlySampler, FogSampleTranslatedWorldPosition,
+					max(ResolvedView.SkyAtmosphereAerialPerspectiveStartDepthKm, RenderVolumetricCloudParameters.AerialPerspectiveMieOnlyStartDistanceKm), RenderVolumetricCloudParameters.AerialPerspectiveMieOnlyFadeDistanceKmInv);
+				float4 RayAP = GETAERIALPERSPECTIVE(View.CameraAerialPerspectiveVolumeRayOnly, View.CameraAerialPerspectiveVolumeRayOnlySampler, FogSampleTranslatedWorldPosition,
+					max(ResolvedView.SkyAtmosphereAerialPerspectiveStartDepthKm, RenderVolumetricCloudParameters.AerialPerspectiveRayOnlyStartDistanceKm), RenderVolumetricCloudParameters.AerialPerspectiveRayOnlyFadeDistanceKmInv);
+
+				AerialPerspective = float4(MieAP.rgb + RayAP.rgb, MieAP.a * RayAP.a);
+			}
+			else
+			{
+				const float NearFadeOutRangeInvDepthKm = 1.0 / 0.00001f; // 1 centimeter fade region
+				// Apply AP only once according to the mean position within the participating media weighted by transmittance/visibility.
+				// This allows to apply AP only once per pixel instead of per sample.
+				AerialPerspective = GETAERIALPERSPECTIVE(View.CameraAerialPerspectiveVolume, View.CameraAerialPerspectiveVolumeSampler, FogSampleTranslatedWorldPosition, ResolvedView.SkyAtmosphereAerialPerspectiveStartDepthKm, NearFadeOutRangeInvDepthKm);
+			}
+
+			// Apply aerial perspective OVER the cloud, assiming coverage is from the full transmittance.
+			// skip if ClosestMeanTransmittance is 1
+			float MeanCoverage = 1.0 - ClosestMeanTransmittance;
+			ClosestLuminance = AerialPerspective.rgb * MeanCoverage + AerialPerspective.a * ClosestLuminance;
+		}
+
+
+		//
+		// Apply height fog if needed
+		//
+		if (RenderVolumetricCloudParameters.EnableHeightFog && ClostestTAPWeightsSum > 0.0f)
+		{
+			float4 HeightFogInscatteringAndTransmittance = CalculateHeightFog(FogSampleTranslatedWorldPosition - ResolvedView.TranslatedWorldCameraOrigin);
+
+			if (FogStruct.ApplyVolumetricFog > 0)
+			{
+				float3 VolumeUV = ComputeVolumeUVFromNDC(ClipPos);
+				const uint EyeIndex = 0;
+				HeightFogInscatteringAndTransmittance = CombineVolumetricFog(HeightFogInscatteringAndTransmittance, VolumeUV, EyeIndex, ConvertFromDeviceZ(ClipPos.z));
+			}
+
+			// Apply aerial perspective OVER the cloud, assiming coverage is from the full transmittance.
+			// skip if MeanTransmittance is 1
+			float MeanCoverage = 1.0 - ClosestMeanTransmittance;
+			ClosestLuminance = HeightFogInscatteringAndTransmittance.rgb * MeanCoverage + HeightFogInscatteringAndTransmittance.a * ClosestLuminance;
+		}
+	}
+#endif // CLOUD_MIN_AND_MAX_DEPTH
+
+
 	//
 	// Output result
 	//
@@ -1371,10 +1505,15 @@ void MainCommon(in FMaterialPixelParameters MaterialParameters, in float4 SvPosi
 	OutColor0 = float4(Luminance * OutputPreExposure, GrayScaleTransmittance);
 #endif
 
-	OutColor1.x = MaxHalfFloat; // Default to far away depth to be flat and not intersect with any geometry.
+#if CLOUD_MIN_AND_MAX_DEPTH
+	float ClosestGrayScaleTransmittance = ClosestMeanTransmittance < RenderVolumetricCloudParameters.StopTracingTransmittanceThreshold ? 0.0f : ClosestMeanTransmittance;
+	OutColor1 = float4(ClosestLuminance * OutputPreExposure, ClosestGrayScaleTransmittance);
+#endif // CLOUD_MIN_AND_MAX
+
+	OutDepth.x = MaxHalfFloat; // Default to far away depth to be flat and not intersect with any geometry.
 	if (RenderVolumetricCloudParameters.OpaqueIntersectionMode >= 1)
 	{
-		OutColor1.x = ((GrayScaleTransmittance > 0.99) ? NoCloudDepth : tAP) * CENTIMETER_TO_KILOMETER; // using a small threshold on transmittance
+		OutDepth.x = ((GrayScaleTransmittance > 0.99) ? NoCloudDepth : tAP) * CENTIMETER_TO_KILOMETER; // using a small threshold on transmittance
 	}
 }
 
@@ -1385,11 +1524,13 @@ void MainCommon(in FMaterialPixelParameters MaterialParameters, in float4 SvPosi
 void MainPS(
 	in float4 SvPosition : SV_Position
 	, out float4 OutColor0 : SV_Target0
-	, out float2 OutColor1 : SV_Target1)
+	, out float2 OutDepth : SV_Target1)
 {
 	FMaterialPixelParameters MaterialParameters = MakeInitializedMaterialPixelParameters();
-
-	MainCommon(MaterialParameters, SvPosition, OutColor0, OutColor1);
+	
+	float4 OutDepth4 = 0;
+	MainCommon(MaterialParameters, SvPosition, OutColor0, OutDepth4);
+	OutDepth = OutDepth4.xy;
 }
 
 #endif // SHADER_RENDERVIEW_PS
@@ -1400,8 +1541,11 @@ float4 OutputViewRect;
 int bBlendCloudColor;
 int TargetCubeFace;
 
-RWTexture2D<float4> OutCloudColor;
-RWTexture2D<float2> OutCloudDepth;
+RWTexture2D<float4> OutCloudColor0;
+#if CLOUD_MIN_AND_MAX_DEPTH
+RWTexture2D<float4> OutCloudColor1;
+#endif
+RWTexture2D<float4> OutCloudDepth;
 RWTextureCube<float4> OutCloudColorCube;
 
 [numthreads(THREADGROUP_SIZEX, THREADGROUP_SIZEY, 1)]
@@ -1412,11 +1556,18 @@ void MainCS(uint2 DispatchThreadId : SV_DispatchThreadID)
 	{
 		float4 SvPosition = float4(DispatchThreadId.x + 0.5, DispatchThreadId.y + 0.5, 0.5, 1);
 		float4 OutColor0;
-		float2 OutColor1;
+#if CLOUD_MIN_AND_MAX_DEPTH
+		float4 OutColor1;
+#endif
+		float4 OutDepth;
 
 		FMaterialPixelParameters MaterialParameters = MakeInitializedMaterialPixelParameters();
 
-		MainCommon(MaterialParameters, SvPosition, OutColor0, OutColor1);
+		MainCommon(MaterialParameters, SvPosition, OutColor0
+#if CLOUD_MIN_AND_MAX_DEPTH
+			, OutColor1
+#endif
+			, OutDepth);
 
 		if (TargetCubeFace >= 0)
 		{
@@ -1435,15 +1586,18 @@ void MainCS(uint2 DispatchThreadId : SV_DispatchThreadID)
 		{
 			if (bBlendCloudColor)
 			{
-				float4 CurColor0 = OutCloudColor[DispatchThreadId];
-				OutCloudColor[DispatchThreadId] = float4(CurColor0.rgb * OutColor0.a + OutColor0.rgb, CurColor0.a * OutColor0.a);
+				float4 CurColor0 = OutCloudColor0[DispatchThreadId];
+				OutCloudColor0[DispatchThreadId] = float4(CurColor0.rgb * OutColor0.a + OutColor0.rgb, CurColor0.a * OutColor0.a);
 			}
 			else
 			{
-				OutCloudColor[DispatchThreadId] = OutColor0;
+				OutCloudColor0[DispatchThreadId] = OutColor0;
 			}
 		}
-		OutCloudDepth[DispatchThreadId] = OutColor1;
+#if CLOUD_MIN_AND_MAX_DEPTH
+		OutCloudColor1[DispatchThreadId] = OutColor1;
+#endif
+		OutCloudDepth[DispatchThreadId] = OutDepth;
 	}
 }
 
diff --git a/Engine/Shaders/Private/VolumetricRenderTarget.usf b/Engine/Shaders/Private/VolumetricRenderTarget.usf
--- a/Engine/Shaders/Private/VolumetricRenderTarget.usf
+++ b/Engine/Shaders/Private/VolumetricRenderTarget.usf
@@ -9,6 +9,9 @@
 
 #include "SceneTextureParameters.ush"
 
+#ifndef CLOUD_MIN_AND_MAX_DEPTH
+#define CLOUD_MIN_AND_MAX_DEPTH 0
+#endif
 
 #ifdef SHADER_RECONSTRUCT_VOLUMETRICRT
 
@@ -17,10 +20,26 @@ SamplerState LinearTextureSampler;
 Texture2D HalfResDepthTexture;
 
 Texture2D<float4> TracingVolumetricTexture;
-Texture2D<float2> TracingVolumetricDepthTexture;
+Texture2D<float4> SecondaryTracingVolumetricTexture;
+Texture2D<float4> TracingVolumetricDepthTexture;
 uint4 TracingVolumetricTextureValidCoordRect;
 float4 TracingVolumetricTextureValidUvRect;
 
+struct FDepthData
+{
+	float  CloudFrontDepthFromViewKm;
+	float  SceneDepthFromViewKm;
+	float2 MinMaxDeviceZ;
+};
+FDepthData GetDepthDataFromVector(float4 DepthVector)
+{
+	FDepthData DepthData;
+	DepthData.CloudFrontDepthFromViewKm = DepthVector.x;
+	DepthData.SceneDepthFromViewKm		= DepthVector.y;
+	DepthData.MinMaxDeviceZ				= DepthVector.zw;
+	return DepthData;
+}
+
 float4 SafeLoadTracingVolumetricTexture(uint2 Coord)
 {
 	return TracingVolumetricTexture.Load(uint3(clamp(Coord, TracingVolumetricTextureValidCoordRect.xy, TracingVolumetricTextureValidCoordRect.zw), 0));
@@ -29,13 +48,21 @@ float4 SafeSampleTracingVolumetricTexture(float2 UV)
 {
 	return TracingVolumetricTexture.SampleLevel(LinearTextureSampler, clamp(UV, TracingVolumetricTextureValidUvRect.xy, TracingVolumetricTextureValidUvRect.zw), 0);
 }
-float2 SafeLoadTracingVolumetricDepthTexture(uint2 Coord)
+float4 SafeLoadSecondaryTracingVolumetricTexture(uint2 Coord)
+{
+	return SecondaryTracingVolumetricTexture.Load(uint3(clamp(Coord, TracingVolumetricTextureValidCoordRect.xy, TracingVolumetricTextureValidCoordRect.zw), 0));
+}
+float4 SafeSampleSecondaryTracingVolumetricTexture(float2 UV)
 {
-	return TracingVolumetricDepthTexture.Load(uint3(clamp(Coord, TracingVolumetricTextureValidCoordRect.xy, TracingVolumetricTextureValidCoordRect.zw), 0)).rg;
+	return SecondaryTracingVolumetricTexture.SampleLevel(LinearTextureSampler, clamp(UV, TracingVolumetricTextureValidUvRect.xy, TracingVolumetricTextureValidUvRect.zw), 0);
 }
-float2 SafeSampleTracingVolumetricDepthTexture(float2 UV)
+FDepthData SafeLoadTracingVolumetricDepthTexture(uint2 Coord)
 {
-	return TracingVolumetricDepthTexture.SampleLevel(LinearTextureSampler, clamp(UV, TracingVolumetricTextureValidUvRect.xy, TracingVolumetricTextureValidUvRect.zw), 0).rg;
+	return GetDepthDataFromVector(TracingVolumetricDepthTexture.Load(uint3(clamp(Coord, TracingVolumetricTextureValidCoordRect.xy, TracingVolumetricTextureValidCoordRect.zw), 0)));
+}
+FDepthData SafeSampleTracingVolumetricDepthTexture(float2 UV)
+{
+	return GetDepthDataFromVector(TracingVolumetricDepthTexture.SampleLevel(LinearTextureSampler, clamp(UV, TracingVolumetricTextureValidUvRect.xy, TracingVolumetricTextureValidUvRect.zw), 0));
 }
 
 #if PERMUTATION_HISTORY_AVAILABLE
@@ -54,13 +81,13 @@ float4 SafeSamplePreviousFrameVolumetricTexture(float2 UV)
 {
 	return PreviousFrameVolumetricTexture.SampleLevel(LinearTextureSampler, clamp(UV, PreviousFrameVolumetricTextureValidUvRect.xy, PreviousFrameVolumetricTextureValidUvRect.zw), 0);
 }
-float2 SafeLoadPreviousFrameVolumetricDepthTexture(uint2 Coord)
+FDepthData SafeLoadPreviousFrameVolumetricDepthTexture(uint2 Coord)
 {
-	return PreviousFrameVolumetricDepthTexture.Load(uint3(clamp(Coord, PreviousFrameVolumetricTextureValidCoordRect.xy, PreviousFrameVolumetricTextureValidCoordRect.zw), 0)).rg;
+	return GetDepthDataFromVector(float4(PreviousFrameVolumetricDepthTexture.Load(uint3(clamp(Coord, PreviousFrameVolumetricTextureValidCoordRect.xy, PreviousFrameVolumetricTextureValidCoordRect.zw), 0)).rg, 0.0, 0.0));
 }
-float2 SafeSamplePreviousFrameVolumetricDepthTexture(float2 UV)
+FDepthData SafeSamplePreviousFrameVolumetricDepthTexture(float2 UV)
 {
-	return PreviousFrameVolumetricDepthTexture.SampleLevel(LinearTextureSampler, clamp(UV, PreviousFrameVolumetricTextureValidUvRect.xy, PreviousFrameVolumetricTextureValidUvRect.zw), 0).rg;
+	return GetDepthDataFromVector(float4(PreviousFrameVolumetricDepthTexture.SampleLevel(LinearTextureSampler, clamp(UV, PreviousFrameVolumetricTextureValidUvRect.xy, PreviousFrameVolumetricTextureValidUvRect.zw), 0).rg, 0.0, 0.0));
 }
 #endif // PERMUTATION_HISTORY_AVAILABLE
 
@@ -120,22 +147,24 @@ void ReconstructVolumetricRenderTargetPS(
 
 	const bool bUseNewSample = (XSub == CurrentTracingPixelOffset.x) && (YSub == CurrentTracingPixelOffset.y);
 	float4 RGBA = 0.0f;
-	float2 Depths = 0.0f;
+	FDepthData DepthData = (FDepthData)0;
 
 #if 0
 	// Debug: always use new sample else history result. Should be 
 	if (bUseNewSample)
 	{
 		RGBA  = SafeLoadTracingVolumetricTexture(int2(SVPos.xy) / DownSampleFactor);
-		Depths = SafeLoadTracingVolumetricDepthTexture(int2(SVPos.xy) / DownSampleFactor);
+		//RGBA  = SafeLoadSecondaryTracingVolumetricTexture(int2(SVPos.xy) / DownSampleFactor);
+		DepthData = SafeLoadTracingVolumetricDepthTexture(int2(SVPos.xy) / DownSampleFactor);
 	}
 	else
 	{
-		RGBA = SafeSamplePreviousFrameVolumetricTexture(ScreenUV);
-		Depths = SafeSamplePreviousFrameVolumetricDepthTexture(ScreenUV);
+		RGBA = SafeSampleTracingVolumetricTexture(ScreenUV);
+		//RGBA = SafeSampleSecondaryTracingVolumetricTexture(ScreenUV);
+		DepthData = SafeSamplePreviousFrameVolumetricDepthTexture(ScreenUV);
 	}
 	OutputRt0 = RGBA;
-	OutputRt1 = Depths;
+	OutputRt1 = float2(DepthData.CloudFrontDepthFromViewKm, DepthData.SceneDepthFromViewKm);
 	return;
 #endif
   
@@ -143,7 +172,7 @@ void ReconstructVolumetricRenderTargetPS(
 		float2 ScreenPosition = ViewportUVToScreenPos(ScreenUV); // NDC in [-1,1] not using View.ScreenPositionScaleBias here
 
 		// Sample participating media "front depth" for a better reprojection
-		float TracingVolumetricSampleDepthKm = SafeLoadTracingVolumetricDepthTexture(int2(SVPos.xy) / DownSampleFactor).x;
+		float TracingVolumetricSampleDepthKm = SafeLoadTracingVolumetricDepthTexture(int2(SVPos.xy) / DownSampleFactor).CloudFrontDepthFromViewKm;
 		float TracingVolumetricSampleDepth = TracingVolumetricSampleDepthKm * KILOMETER_TO_CENTIMETER;
 		float DeviceZ = ConvertToDeviceZ(TracingVolumetricSampleDepth); // Approximation. Should try real DeviceZ
 
@@ -166,68 +195,159 @@ void ReconstructVolumetricRenderTargetPS(
 			{
 				// Load the new sample for this pixel we have just traced
 				RGBA = SafeLoadTracingVolumetricTexture(int2(SVPos.xy) / DownSampleFactor);// +float4(0.1, 0.0, 0, 0);
-				Depths = SafeLoadTracingVolumetricDepthTexture(int2(SVPos.xy) / DownSampleFactor);
+				DepthData = SafeLoadTracingVolumetricDepthTexture(int2(SVPos.xy) / DownSampleFactor);
 			}
 			else if(bValidPreviousUVs)
 			{
 				// Sample valid on screen history
 				RGBA = SafeSamplePreviousFrameVolumetricTexture(PrevScreenUVs);// +float4(0, 0.1, 0, 0);
-				Depths = SafeSamplePreviousFrameVolumetricDepthTexture(PrevScreenUVs);
+				DepthData = SafeSamplePreviousFrameVolumetricDepthTexture(PrevScreenUVs);
 			}
 			else
 			{
 				// Bilinear sample new low resoltuion tracing
 				RGBA = SafeSampleTracingVolumetricTexture(ScreenUV);
-				Depths = SafeSampleTracingVolumetricDepthTexture(ScreenUV);
+				DepthData = SafeSampleTracingVolumetricDepthTexture(ScreenUV);
 			}
 			OutputRt0 = RGBA;
-			OutputRt1 = Depths;
+			OutputRt1 = float2(DepthData.CloudFrontDepthFromViewKm, DepthData.SceneDepthFromViewKm);
 			return;
 		}
-		else if (bValidPreviousUVs)
+		else
 		{
+			#define DEBUGCOLORS 0
+
 			int2 CenterSample = int2(SVPos.xy) / DownSampleFactor;
 			float4 NewRGBA = SafeLoadTracingVolumetricTexture(CenterSample);
-			float2 NewDepths = SafeLoadTracingVolumetricDepthTexture(CenterSample);
+			FDepthData NewDepthData = SafeLoadTracingVolumetricDepthTexture(CenterSample);
 
 			if (bUseNewSample)
 			{
 				// Load the new sample for this pixel we have just traced
 				RGBA = NewRGBA;
-				Depths = NewDepths;
-				//RGBA += float4(0.1, 0, 0, 0);
+				DepthData = NewDepthData;
 			}
-			else // bValidPreviousUVs==true always here
+			else if (bValidPreviousUVs)
 			{
+
+
+
+		#if CLOUD_MIN_AND_MAX_DEPTH
+
+
+
+				//
+				// NOTE: The best temporal upsampling path reserver for mode 0 when compute is available.
+				//
+
+				const float SceneHalfResDeviceZ = HalfResDepthTexture.Load(int3(SVPos.xy + ViewViewRectMin, 0)).r;
+				const float SceneHalfResDepth = ConvertFromDeviceZ(max(0.000000000001, SceneHalfResDeviceZ));
+
+				const float ThresholdToUpsamplingCm = 100000.0; // super arbitrary :/
+
+				const float NearCloudTracingDepth = ConvertFromDeviceZ(max(0.000000000001, NewDepthData.MinMaxDeviceZ.x));
+				const float FarCloudTracingDepth  = ConvertFromDeviceZ(max(0.000000000001, NewDepthData.MinMaxDeviceZ.y));
+				if (abs(NearCloudTracingDepth - FarCloudTracingDepth) > ThresholdToUpsamplingCm)
+				{
+					// if the min/max depth has a large difference, then always upsample with the min and max traced data
+					float4 NewRGBANear = SafeLoadSecondaryTracingVolumetricTexture(CenterSample);
+					
+					const float LerpFactor = saturate((SceneHalfResDepth - NearCloudTracingDepth) / max(0.00001, FarCloudTracingDepth - NearCloudTracingDepth));
+					RGBA = lerp(NewRGBANear, NewRGBA, LerpFactor);
+
+				//	DepthData = NewDepthData;
+					DepthData.SceneDepthFromViewKm = SceneHalfResDepth * CENTIMETER_TO_KILOMETER;	// TODO trasform from scene depth to view depth.
+					DepthData.CloudFrontDepthFromViewKm = DepthData.SceneDepthFromViewKm;			// Keer far cloud tracing front depth. That should be good.
+						
+			#if DEBUGCOLORS
+					RGBA = float4(1, 1, 0, 0.5);
+			#endif
+				}
+				else
+				{
+					// Otherwise we use history 
+					FDepthData HistoryDepthData = SafeSamplePreviousFrameVolumetricDepthTexture(PrevScreenUVs);
+
+					const float HistorySceneDepthFromViewCm = HistoryDepthData.SceneDepthFromViewKm * KILOMETER_TO_CENTIMETER;
+
+					//if (HistorySceneDepthFromViewCm < (NearCloudTracingDepth - ThresholdToUpsamplingCm))
+					if (HistorySceneDepthFromViewCm < (NewDepthData.SceneDepthFromViewKm * KILOMETER_TO_CENTIMETER - ThresholdToUpsamplingCm))
+					{
+						// History is closer than the near cloud tracing this frame. This means an object had move and an new disocluded area is discovered.
+						// So we simply use the new data from this frame according to the new depth
+						
+						RGBA = NewRGBA;
+						DepthData = NewDepthData;
+					//	DepthData.SceneDepthFromViewKm = SceneHalfResDepth * CENTIMETER_TO_KILOMETER;	// TODO trasform from scene depth to view depth.
+					//	DepthData.CloudFrontDepthFromViewKm = DepthData.SceneDepthFromViewKm;			// Keer far cloud tracing front depth. That should be good.
+							
+			#if DEBUGCOLORS
+						RGBA = float4(1, 0, 0, 0.5);
+			#endif
+					}
+					else if (HistorySceneDepthFromViewCm > (NewDepthData.SceneDepthFromViewKm * KILOMETER_TO_CENTIMETER + ThresholdToUpsamplingCm))
+					{
+						// An area that just go covered (history is invalid because occluded)
+						RGBA = NewRGBA;
+						DepthData = NewDepthData;
+					//	DepthData.SceneDepthFromViewKm = SceneHalfResDepth * CENTIMETER_TO_KILOMETER;	// TODO trasform from scene depth to view depth.
+					//	DepthData.CloudFrontDepthFromViewKm = DepthData.SceneDepthFromViewKm;			// Keer far cloud tracing front depth. That should be good.
+
+			#if DEBUGCOLORS
+						RGBA = float4(0, 1, 0, 0.5);
+			#endif
+					}
+					else
+					{
+						// Good reprojection
+						float4 HistoryRGBA = SafeSamplePreviousFrameVolumetricTexture(PrevScreenUVs);
+						RGBA = HistoryRGBA;
+						DepthData = HistoryDepthData;
+							
+			#if DEBUGCOLORS
+						RGBA = float4(0, 0, 1, 0.5);
+			#endif
+					}
+				}
+
+
+
+		#else // CLOUD_MIN_AND_MAX_DEPTH
+
+
+
+				//
+				// NOTE: This path is use when mode0 is not used or when compute is not available (min max depth permutation are only generated for compute path)
+				//
+				
 				// Sample valid on screen history
-				float4 HistoryRGBA = SafeSamplePreviousFrameVolumetricTexture(PrevScreenUVs * 1.0);
-				float2 HistoryDepths = SafeSamplePreviousFrameVolumetricDepthTexture(PrevScreenUVs * 1.0);
+				float4 HistoryRGBA = SafeSamplePreviousFrameVolumetricTexture(PrevScreenUVs);
+				FDepthData HistoryDepthData = SafeSamplePreviousFrameVolumetricDepthTexture(PrevScreenUVs);
 
 				// Get information about neightboors
 				int2 NeightboorsOffset[8] = { int2(1,0), int2(1,1), int2(0,1), int2(-1,1), int2(-1,0), int2(-1,-1), int2(0,-1), int2(1,-1)};
 
-
 				const float ReconstructDepthZ = HalfResDepthTexture.Load(int3(SVPos.xy + ViewViewRectMin, 0)).r;
 				const float3 WorldPosition = LWCHackToFloat(SvPositionToWorld(float4(CenterSample, ReconstructDepthZ, 1.0)));
 				const float PixelDistanceFromViewKm = length(WorldPosition - LWCHackToFloat(PrimaryView.WorldCameraOrigin)) * CENTIMETER_TO_KILOMETER;
 
 				RGBA = HistoryRGBA;
-				Depths = HistoryDepths;
+				DepthData = HistoryDepthData;
 
-				if (/*ReconstructDepthZ > 0.0001f &&*/ abs(PixelDistanceFromViewKm - HistoryDepths.y) > PixelDistanceFromViewKm * 0.1f)
+				if (/*ReconstructDepthZ > 0.0001f &&*/ abs(PixelDistanceFromViewKm - DepthData.SceneDepthFromViewKm) > PixelDistanceFromViewKm * 0.1f)
 				{
 					// History has a too large depth difference at depth discontinuities, use the data with closest depth within the neightborhood
 					float ClosestDepth = 99999999.0f;
 					for (int i = 0; i < 8; ++i)
 					{
-						float2 NeighboorsDepths = SafeLoadTracingVolumetricDepthTexture(CenterSample + NeightboorsOffset[i]);
-						const float NeighboorsClosestDepth = abs(PixelDistanceFromViewKm - NeighboorsDepths.y);
+						FDepthData NeighboorsDepthData = SafeLoadTracingVolumetricDepthTexture(CenterSample + NeightboorsOffset[i]);
+						const float NeighboorsClosestDepth = abs(PixelDistanceFromViewKm - NeighboorsDepthData.SceneDepthFromViewKm);
 						if (NeighboorsClosestDepth < ClosestDepth)
 						{
 							ClosestDepth = NeighboorsClosestDepth;
 							float4 NeighboorsRGBA = SafeLoadTracingVolumetricTexture(CenterSample + NeightboorsOffset[i]);
 							RGBA = NeighboorsRGBA;// +float4(0, 1, 0, 0);
-							Depths = NeighboorsDepths;
+							DepthData = NeighboorsDepthData;
 						} 
 					}
 					// After more testing, the code below looked unecessary
@@ -241,7 +361,7 @@ void ReconstructVolumetricRenderTargetPS(
 				else // Because of the test on bUseNewSample above, we know here that we are only dealing with reprojected data  //if(ReconstructDepthZ < 0.000001f)
 				{
 					// TODO: To use this, we need to make sure we prioritise pixe lwith under represented depth.
-#if PERMUTATION_REPROJECTION_BOX_CONSTRAINT
+				#if PERMUTATION_REPROJECTION_BOX_CONSTRAINT
 					// Make sure that history stay in the neightborhood color/transmittance/depth box after reprojection
 					float4 ColorAABBMin = 999999999.0f;
 					float4 ColorAABBMax = 0.0f;
@@ -251,50 +371,64 @@ void ReconstructVolumetricRenderTargetPS(
 					for (int i = 0; i < 8; ++i)
 					{
 						float4 ColorData = SafeLoadTracingVolumetricTexture(CenterSample + NeightboorsOffset[i]);
-						float2 DepthData = SafeLoadTracingVolumetricDepthTexture(CenterSample + NeightboorsOffset[i]);
+						FDepthData NeighboorsDepthData = SafeLoadTracingVolumetricDepthTexture(CenterSample + NeightboorsOffset[i]);
+						float2 NeighboorsDepthData2 = float2(NeighboorsDepthData.CloudFrontDepthFromViewKm, NeighboorsDepthData.SceneDepthFromViewKm);
 						ColorAABBMin = min(ColorAABBMin, ColorData);
 						ColorAABBMax = max(ColorAABBMax, ColorData);
-						DepthsAABBMin = min(DepthsAABBMin, DepthData);
-						DepthsAABBMax = max(DepthsAABBMax, DepthData);
-						bApply = bApply && DepthData.y > 1000.0f;
+						DepthsAABBMin = min(DepthsAABBMin, NeighboorsDepthData2);
+						DepthsAABBMax = max(DepthsAABBMax, NeighboorsDepthData2);
+						bApply = bApply && NeighboorsDepthData2.y > 1000.0f;
 					}
 					ColorAABBMin = min(ColorAABBMin, NewRGBA);
 					ColorAABBMax = max(ColorAABBMax, NewRGBA);
-					DepthsAABBMin = min(DepthsAABBMin, NewDepths);
-					DepthsAABBMax = max(DepthsAABBMax, NewDepths);
-					bApply = bApply && NewDepths.y > 1000.0f;
+					DepthsAABBMin = min(DepthsAABBMin, float2(NewDepthData.CloudFrontDepthFromViewKm, NewDepthData.SceneDepthFromViewKm));
+					DepthsAABBMax = max(DepthsAABBMax, float2(NewDepthData.CloudFrontDepthFromViewKm, NewDepthData.SceneDepthFromViewKm));
+					bApply = bApply && NewDepthData.SceneDepthFromViewKm > 1000.0f;
 					//if (bApply)
 					{
 						RGBA = clamp(RGBA, ColorAABBMin, ColorAABBMax);
-						Depths = clamp(Depths, DepthsAABBMin, DepthsAABBMax);
+
+						float2 Depths = clamp(float2(DepthData.CloudFrontDepthFromViewKm, DepthData.SceneDepthFromViewKm), DepthsAABBMin, DepthsAABBMax);
+
+						DepthData.CloudFrontDepthFromViewKm = Depths.x;
+						DepthData.SceneDepthFromViewKm = Depths.y;
 					}
 					//RGBA += float4(0, 0.8, 0.8, 0);
-#endif
+				#endif
 				}
 
-				if (!(all(IsFinite(RGBA)) && all(IsFinite(Depths))))
+
+
+		#endif // CLOUD_MIN_AND_MAX_DEPTH
+
+
+
+				if (!(all(IsFinite(RGBA)) && all(IsFinite(float2(DepthData.CloudFrontDepthFromViewKm, DepthData.SceneDepthFromViewKm)))))
 				{
 					RGBA = float4(0.0f, 0.0f, 0.0f, 1.0f);
-					Depths = float2(1000.0f, 1000.0f);
+					DepthData.CloudFrontDepthFromViewKm = 1000.0f;
+					DepthData.SceneDepthFromViewKm = 1000.0f;
 				}
 			}
-		}
-		else
-		{
-			// History is invalid so simply use this frame low resolution render with bilinear sampling.
-			RGBA = SafeSampleTracingVolumetricTexture(ScreenUV);
-			Depths = SafeSampleTracingVolumetricDepthTexture(ScreenUV);
+			else // !bValidPreviousUVs
+			{
+				// History is invalid so simply use this frame low resolution render with bilinear sampling.
+				// Single sampel of the far data seem sto always be enough
+				RGBA = SafeSampleTracingVolumetricTexture(ScreenUV);
+				DepthData = SafeSampleTracingVolumetricDepthTexture(ScreenUV);
+			}
 		}
 	}
 
 	OutputRt0 = RGBA;
-	OutputRt1 = Depths;
+	OutputRt1 = float2(DepthData.CloudFrontDepthFromViewKm, DepthData.SceneDepthFromViewKm);
 
 #else // PERMUTATION_HISTORY_AVAILABLE
 
 	// Simple bilinear upsample
 	OutputRt0 = SafeSampleTracingVolumetricTexture(ScreenUV);
-	OutputRt1 = SafeSampleTracingVolumetricDepthTexture(ScreenUV);
+	FDepthData DepthData = SafeSampleTracingVolumetricDepthTexture(ScreenUV);
+	OutputRt1 = float2(DepthData.CloudFrontDepthFromViewKm, DepthData.SceneDepthFromViewKm);
 
 #endif // PERMUTATION_HISTORY_AVAILABLE
 }
@@ -387,7 +521,7 @@ void ComposeVolumetricRTOverScenePS(
 #endif
 )
 {
-	float2 CurResPixelCoord = SVPos.xy;
+	float2 CurResPixelCoord = float2(SVPos.xy);
 	float2 ScreenUV = CurResPixelCoord * View.BufferSizeAndInvSize.zw;
 	float2 ScreenUVNoOffset = (CurResPixelCoord - View.ViewRectMin.xy) * View.BufferSizeAndInvSize.zw;
 	float2 VolumeUV = FullResolutionToVolumetricBufferResolutionScale.x * (ScreenUVNoOffset * View.BufferSizeAndInvSize.xy * VolumetricTextureSizeAndInvSize.zw);
@@ -648,14 +782,14 @@ void ComposeVolumetricRTOverScenePS(
 	PixelDistanceFromViewKm = min(PixelDistanceFromViewKm, max(max(VolumeFrontDepth0.y, VolumeFrontDepth1.y),max(VolumeFrontDepth2.y, VolumeFrontDepth3.y)));
 
 	float4 Depth4 = float4(VolumeFrontDepth0.y, VolumeFrontDepth1.y, VolumeFrontDepth2.y, VolumeFrontDepth3.y);
-	float4 depth4Diff = PixelDistanceFromViewKm - Depth4;
-	depth4Diff = abs(depth4Diff);
-	float MaxDepth4Diff = max(max(depth4Diff.x, depth4Diff.y), max(depth4Diff.z, depth4Diff.w));
+	float4 Depth4Diff = PixelDistanceFromViewKm - Depth4;
+	Depth4Diff = abs(Depth4Diff);
+	float MaxDepth4Diff = max(max(Depth4Diff.x, Depth4Diff.y), max(Depth4Diff.z, Depth4Diff.w));
 
 	float ValidSampleCount = 0;
 	float4 DataAcc = 0;
 
-	const float weightMultiplier = 10000.0f;
+	const float WeightMultiplier = 1000.0f;
 	const float ThresholdToBilinear = PixelDistanceFromViewKm * 0.1;
 
 	if (MaxDepth4Diff > ThresholdToBilinear)
@@ -707,7 +841,7 @@ void ComposeVolumetricRTOverScenePS(
 		if(!bIsWaterPixel || (bIsWaterPixel && ValidSampleCount == 0 && bAnyCloudSamplesInFrontOfWater))
 		{
 			// Depth discontinuities edges
-			float4 weights = 1.0f / (depth4Diff * weightMultiplier + 1.0f);
+			float4 weights = 1.0f / (Depth4Diff * WeightMultiplier + 1.0f);
 			const float weightsSum = dot(weights, float4(1.0f, 1.0f, 1.0f, 1.0f));
 			weights /= weightsSum;
 
@@ -724,8 +858,15 @@ void ComposeVolumetricRTOverScenePS(
 		//DataAcc = float4(0, 1.0, 0, 0);
 	}
 
-
-	if (UvOffsetScale > 0.0 && !bIsWaterPixel) // we are going to fetch far sample so let's skip water pixel because depth cannot be trusted and used in this case.
+		// we are going to fetch far sample so let's skip water pixel because depth cannot be trusted and used in this case.
+	if (UvOffsetScale > 0.0 && !bIsWaterPixel &&	
+		// Also only blur background to not leak cloud over opaque (this can happen when cloud data happear over meshes due to temporal jittering)
+#if HAS_INVERTED_Z_BUFFER
+		PixelDeviceZ < 0.0001				
+#else
+		PixelDeviceZ > 0.9999
+#endif
+		)
 	{
 		// Add some jittered samples to hide upsampling and noise through TAA filtering
 		// We use only Offset0 to save ALU and go faster. Variation across pixels is enough noise for taa to soften it out.
diff --git a/Engine/Source/Runtime/Renderer/Private/DeferredShadingRenderer.cpp b/Engine/Source/Runtime/Renderer/Private/DeferredShadingRenderer.cpp
--- a/Engine/Source/Runtime/Renderer/Private/DeferredShadingRenderer.cpp
+++ b/Engine/Source/Runtime/Renderer/Private/DeferredShadingRenderer.cpp
@@ -2985,7 +2985,6 @@ void FDeferredShadingSceneRenderer::Render(FRDGBuilder& GraphBuilder)
 	const bool bShouldRenderVolumetricCloud = bShouldRenderVolumetricCloudBase && (!ViewFamily.EngineShowFlags.VisualizeVolumetricCloudConservativeDensity && !ViewFamily.EngineShowFlags.VisualizeVolumetricCloudEmptySpaceSkipping);
 	const bool bShouldVisualizeVolumetricCloud = bShouldRenderVolumetricCloudBase && (!!ViewFamily.EngineShowFlags.VisualizeVolumetricCloudConservativeDensity || !!ViewFamily.EngineShowFlags.VisualizeVolumetricCloudEmptySpaceSkipping);
 	bool bAsyncComputeVolumetricCloud = IsVolumetricRenderTargetEnabled() && IsVolumetricRenderTargetAsyncCompute();
-	bool bHasHalfResCheckerboardMinMaxDepth = false;
 	bool bVolumetricRenderTargetRequired = bShouldRenderVolumetricCloud && !bHasRayTracedOverlay;
 
 	if (SkyAtmospherePassLocation == ESkyAtmospherePassLocation::BeforeOcclusion && bShouldRenderSkyAtmosphere)
@@ -3056,16 +3055,30 @@ void FDeferredShadingSceneRenderer::Render(FRDGBuilder& GraphBuilder)
 	UpdateLumenScene(GraphBuilder, LumenFrameTemporaries);
 
 	FRDGTextureRef HalfResolutionDepthCheckerboardMinMaxTexture = nullptr;
+	FRDGTextureRef QuarterResolutionDepthMinMaxTexture = nullptr;
+	bool bQuarterResMinMaxDepthRequired = bShouldRenderVolumetricCloud && ShouldVolumetricCloudTraceWithMinMaxDepth(Views);
 
-	// Kick off async compute cloud eraly if all depth has been written in the prepass
+	auto GenerateQuarterResDepthMinMaxTexture = [&](auto& GraphBuilder, auto& Views, auto& InputTexture)
+	{
+		if (bQuarterResMinMaxDepthRequired)
+		{
+			check(InputTexture != nullptr);	// Must receive a valid texture
+			check(QuarterResolutionDepthMinMaxTexture == nullptr);	// Only generate it once
+			return CreateQuarterResolutionDepthMinAndMax(GraphBuilder, Views, InputTexture);
+		}
+		return FRDGTextureRef(nullptr);
+	};
+
+	// Kick off async compute cloud early if all depth has been written in the prepass
 	if (bShouldRenderVolumetricCloud && bAsyncComputeVolumetricCloud && DepthPass.EarlyZPassMode == DDM_AllOpaque && !bHasRayTracedOverlay)
 	{
 		HalfResolutionDepthCheckerboardMinMaxTexture = CreateHalfResolutionDepthCheckerboardMinMax(GraphBuilder, Views, SceneTextures.Depth.Resolve);
-		bHasHalfResCheckerboardMinMaxDepth = true;
+		QuarterResolutionDepthMinMaxTexture = GenerateQuarterResDepthMinMaxTexture(GraphBuilder, Views, HalfResolutionDepthCheckerboardMinMaxTexture);
 
 		bool bSkipVolumetricRenderTarget = false;
 		bool bSkipPerPixelTracing = true;
-		bAsyncComputeVolumetricCloud = RenderVolumetricCloud(GraphBuilder, SceneTextures, bSkipVolumetricRenderTarget, bSkipPerPixelTracing, HalfResolutionDepthCheckerboardMinMaxTexture, true, InstanceCullingManager);
+		bAsyncComputeVolumetricCloud = RenderVolumetricCloud(GraphBuilder, SceneTextures, bSkipVolumetricRenderTarget, bSkipPerPixelTracing, 
+			HalfResolutionDepthCheckerboardMinMaxTexture, QuarterResolutionDepthMinMaxTexture, true, InstanceCullingManager);
 	}
 	
 	FRDGTextureRef ForwardScreenSpaceShadowMaskTexture = nullptr;
@@ -3371,11 +3384,12 @@ void FDeferredShadingSceneRenderer::Render(FRDGBuilder& GraphBuilder)
 	if (bShouldRenderVolumetricCloud && bAsyncComputeVolumetricCloud && DepthPass.EarlyZPassMode != DDM_AllOpaque && !bHasRayTracedOverlay)
 	{
 		HalfResolutionDepthCheckerboardMinMaxTexture = CreateHalfResolutionDepthCheckerboardMinMax(GraphBuilder, Views, SceneTextures.Depth.Resolve);
-		bHasHalfResCheckerboardMinMaxDepth = true;
+		QuarterResolutionDepthMinMaxTexture = GenerateQuarterResDepthMinMaxTexture(GraphBuilder, Views, HalfResolutionDepthCheckerboardMinMaxTexture);
 
 		bool bSkipVolumetricRenderTarget = false;
 		bool bSkipPerPixelTracing = true;
-		bAsyncComputeVolumetricCloud = RenderVolumetricCloud(GraphBuilder, SceneTextures, bSkipVolumetricRenderTarget, bSkipPerPixelTracing, HalfResolutionDepthCheckerboardMinMaxTexture, true, InstanceCullingManager);
+		bAsyncComputeVolumetricCloud = RenderVolumetricCloud(GraphBuilder, SceneTextures, bSkipVolumetricRenderTarget, bSkipPerPixelTracing, 
+			HalfResolutionDepthCheckerboardMinMaxTexture, QuarterResolutionDepthMinMaxTexture, true, InstanceCullingManager);
 	}
 
 	if (CustomDepthPassLocation == ECustomDepthPassLocation::AfterBasePass)
@@ -3540,9 +3554,10 @@ void FDeferredShadingSceneRenderer::Render(FRDGBuilder& GraphBuilder)
 
 	GraphBuilder.FlushSetupQueue();
 
-	if (bShouldRenderVolumetricCloud && IsVolumetricRenderTargetEnabled() && !bHasHalfResCheckerboardMinMaxDepth && !bHasRayTracedOverlay)
+	if (bShouldRenderVolumetricCloud && IsVolumetricRenderTargetEnabled() && HalfResolutionDepthCheckerboardMinMaxTexture==nullptr && !bHasRayTracedOverlay)
 	{
 		HalfResolutionDepthCheckerboardMinMaxTexture = CreateHalfResolutionDepthCheckerboardMinMax(GraphBuilder, Views, SceneTextures.Depth.Resolve);
+		QuarterResolutionDepthMinMaxTexture = GenerateQuarterResDepthMinMaxTexture(GraphBuilder, Views, HalfResolutionDepthCheckerboardMinMaxTexture);
 	}
 
 	if (bShouldRenderVolumetricCloud && !bHasRayTracedOverlay)
@@ -3552,7 +3567,8 @@ void FDeferredShadingSceneRenderer::Render(FRDGBuilder& GraphBuilder)
 			// Generate the volumetric cloud render target
 			bool bSkipVolumetricRenderTarget = false;
 			bool bSkipPerPixelTracing = true;
-			RenderVolumetricCloud(GraphBuilder, SceneTextures, bSkipVolumetricRenderTarget, bSkipPerPixelTracing, HalfResolutionDepthCheckerboardMinMaxTexture, false, InstanceCullingManager);
+			RenderVolumetricCloud(GraphBuilder, SceneTextures, bSkipVolumetricRenderTarget, bSkipPerPixelTracing, 
+				HalfResolutionDepthCheckerboardMinMaxTexture, QuarterResolutionDepthMinMaxTexture, false, InstanceCullingManager);
 		}
 		// Reconstruct the volumetric cloud render target to be ready to compose it over the scene
 		ReconstructVolumetricRenderTarget(GraphBuilder, Views, SceneTextures.Depth.Resolve, HalfResolutionDepthCheckerboardMinMaxTexture, bAsyncComputeVolumetricCloud);
@@ -3623,7 +3639,8 @@ void FDeferredShadingSceneRenderer::Render(FRDGBuilder& GraphBuilder)
 	{
 		bool bSkipVolumetricRenderTarget = true;
 		bool bSkipPerPixelTracing = false;
-		RenderVolumetricCloud(GraphBuilder, SceneTextures, bSkipVolumetricRenderTarget, bSkipPerPixelTracing, HalfResolutionDepthCheckerboardMinMaxTexture, false, InstanceCullingManager);
+		RenderVolumetricCloud(GraphBuilder, SceneTextures, bSkipVolumetricRenderTarget, bSkipPerPixelTracing, 
+			HalfResolutionDepthCheckerboardMinMaxTexture, QuarterResolutionDepthMinMaxTexture, false, InstanceCullingManager);
 	}
 
 	// or composite the off screen buffer over the scene.
@@ -3852,10 +3869,10 @@ void FDeferredShadingSceneRenderer::Render(FRDGBuilder& GraphBuilder)
 
 	if (bShouldVisualizeVolumetricCloud && !bHasRayTracedOverlay)
 	{
-		RenderVolumetricCloud(GraphBuilder, SceneTextures, false, true, HalfResolutionDepthCheckerboardMinMaxTexture, false, InstanceCullingManager);
+		RenderVolumetricCloud(GraphBuilder, SceneTextures, false, true, HalfResolutionDepthCheckerboardMinMaxTexture, QuarterResolutionDepthMinMaxTexture, false, InstanceCullingManager);
 		ReconstructVolumetricRenderTarget(GraphBuilder, Views, SceneTextures.Depth.Resolve, HalfResolutionDepthCheckerboardMinMaxTexture, false);
 		ComposeVolumetricRenderTargetOverSceneForVisualization(GraphBuilder, Views, SceneTextures.Color.Target, SceneTextures);
-		RenderVolumetricCloud(GraphBuilder, SceneTextures, true, false, HalfResolutionDepthCheckerboardMinMaxTexture, false, InstanceCullingManager);
+		RenderVolumetricCloud(GraphBuilder, SceneTextures, true, false, HalfResolutionDepthCheckerboardMinMaxTexture, QuarterResolutionDepthMinMaxTexture, false, InstanceCullingManager);
 	}
 
 	if (!bHasRayTracedOverlay)
diff --git a/Engine/Source/Runtime/Renderer/Private/SceneRendering.cpp b/Engine/Source/Runtime/Renderer/Private/SceneRendering.cpp
--- a/Engine/Source/Runtime/Renderer/Private/SceneRendering.cpp
+++ b/Engine/Source/Runtime/Renderer/Private/SceneRendering.cpp
@@ -5546,3 +5546,21 @@ FRDGTextureRef CreateHalfResolutionDepthCheckerboardMinMax(FRDGBuilder& GraphBui
 
 	return SmallDepthTexture;
 }
+
+FRDGTextureRef CreateQuarterResolutionDepthMinAndMax(FRDGBuilder& GraphBuilder, TArrayView<const FViewInfo> Views, FRDGTextureRef InputDepthTexture)
+{
+	const FIntPoint SmallDepthExtent = GetDownscaledExtent(InputDepthTexture->Desc.Extent, 2);
+	const FRDGTextureDesc SmallTextureDesc = FRDGTextureDesc::Create2D(SmallDepthExtent, PF_G16R16F, FClearValueBinding::None, TexCreate_RenderTargetable | TexCreate_ShaderResource);
+	FRDGTextureRef SmallTexture = GraphBuilder.CreateTexture(SmallTextureDesc, TEXT("HalfResolutionDepthMinAndMax"));
+
+	for (const FViewInfo& View : Views)
+	{
+		RDG_GPU_MASK_SCOPE(GraphBuilder, View.GPUMask);
+
+		const FScreenPassTexture InputDepth(InputDepthTexture, GetDownscaledRect(View.ViewRect, 2));
+		const FScreenPassRenderTarget SmallTextureRT(SmallTexture, GetDownscaledRect(View.ViewRect, 4), View.DecayLoadAction(ERenderTargetLoadAction::ENoAction));
+		AddDownsampleDepthPass(GraphBuilder, View, InputDepth, SmallTextureRT, EDownsampleDepthFilter::MinAndMaxDepth);
+	}
+
+	return SmallTexture;
+}
diff --git a/Engine/Source/Runtime/Renderer/Private/SceneRendering.h b/Engine/Source/Runtime/Renderer/Private/SceneRendering.h
--- a/Engine/Source/Runtime/Renderer/Private/SceneRendering.h
+++ b/Engine/Source/Runtime/Renderer/Private/SceneRendering.h
@@ -2349,6 +2349,7 @@ protected:
 		bool bSkipVolumetricRenderTarget,
 		bool bSkipPerPixelTracing,
 		FRDGTextureRef HalfResolutionDepthCheckerboardMinMaxTexture,
+		FRDGTextureRef QuarterResolutionDepthMinMaxTexture,
 		bool bAsyncCompute,
 		FInstanceCullingManager& InstanceCullingManager);
 
@@ -2738,6 +2739,9 @@ void VirtualTextureFeedbackEnd(FRDGBuilder& GraphBuilder);
 /** Creates a half resolution checkerboard min / max depth buffer from the input full resolution depth buffer. */
 FRDGTextureRef CreateHalfResolutionDepthCheckerboardMinMax(FRDGBuilder& GraphBuilder, TArrayView<const FViewInfo> Views, FRDGTextureRef SceneDepth);
 
+/** Creates a half resolution depth buffer storing the min and max depth for each 2x2 pixel quad of a half resolution buffer. */
+FRDGTextureRef CreateQuarterResolutionDepthMinAndMax(FRDGBuilder& GraphBuilder, TArrayView<const FViewInfo> Views, FRDGTextureRef DepthTexture);
+
 inline const FSceneTexturesConfig& FViewInfo::GetSceneTexturesConfig() const
 {
 	// TODO:  We are refactoring away use of the FSceneTexturesConfig::Get() global singleton, but need this workaround for now to avoid crashes
diff --git a/Engine/Source/Runtime/Renderer/Private/ScreenPass.cpp b/Engine/Source/Runtime/Renderer/Private/ScreenPass.cpp
--- a/Engine/Source/Runtime/Renderer/Private/ScreenPass.cpp
+++ b/Engine/Source/Runtime/Renderer/Private/ScreenPass.cpp
@@ -177,6 +177,10 @@ public:
 	DECLARE_GLOBAL_SHADER(FDownsampleDepthPS);
 	SHADER_USE_PARAMETER_STRUCT(FDownsampleDepthPS, FGlobalShader);
 
+	class FOutputMinAndMaxDepth : SHADER_PERMUTATION_BOOL("OUTPUT_MIN_AND_MAX_DEPTH");
+
+	using FPermutationDomain = TShaderPermutationDomain<FOutputMinAndMaxDepth>;
+
 	BEGIN_SHADER_PARAMETER_STRUCT(FParameters, )
 		SHADER_PARAMETER_STRUCT_REF(FViewUniformShaderParameters, View)
 		SHADER_PARAMETER_RDG_TEXTURE(Texture2D, DepthTexture)
@@ -184,6 +188,8 @@ public:
 		SHADER_PARAMETER(FVector2f, SourceMaxUV)
 		SHADER_PARAMETER(FVector2f, DestinationResolution)
 		SHADER_PARAMETER(uint32, DownsampleDepthFilter)
+		SHADER_PARAMETER(FIntVector4, DstPixelCoordMinAndMax)
+		SHADER_PARAMETER(FIntVector4, SrcPixelCoordMinAndMax)
 		RENDER_TARGET_BINDING_SLOTS()
 	END_SHADER_PARAMETER_STRUCT()
 
@@ -206,7 +212,11 @@ void AddDownsampleDepthPass(
 	const FScreenPassTextureViewport OutputViewport(Output);
 
 	TShaderMapRef<FScreenPassVS> VertexShader(View.ShaderMap);
-	TShaderMapRef<FDownsampleDepthPS> PixelShader(View.ShaderMap);
+
+	const bool bIsMinAndMaxDepthFilter = DownsampleDepthFilter == EDownsampleDepthFilter::MinAndMaxDepth;
+	FDownsampleDepthPS::FPermutationDomain Permutation;
+	Permutation.Set<FDownsampleDepthPS::FOutputMinAndMaxDepth>(bIsMinAndMaxDepthFilter ? 1 : 0);
+	TShaderMapRef<FDownsampleDepthPS> PixelShader(View.ShaderMap, Permutation);
 
 	FDownsampleDepthPS::FParameters* PassParameters = GraphBuilder.AllocParameters<FDownsampleDepthPS::FParameters>();
 	PassParameters->View = View.ViewUniformBuffer;
@@ -219,14 +229,26 @@ void AddDownsampleDepthPass(
 	const int32 DownsampledSizeY = OutputViewport.Rect.Height();
 	PassParameters->DestinationResolution = FVector2f(DownsampledSizeX, DownsampledSizeY);
 
-	PassParameters->RenderTargets.DepthStencil = FDepthStencilBinding(Output.Texture, Output.LoadAction, Output.LoadAction, FExclusiveDepthStencil::DepthWrite_StencilWrite);
+	PassParameters->DstPixelCoordMinAndMax = FIntVector4(OutputViewport.Rect.Min.X, OutputViewport.Rect.Min.Y, OutputViewport.Rect.Max.X, OutputViewport.Rect.Max.Y);
+	PassParameters->SrcPixelCoordMinAndMax = FIntVector4( InputViewport.Rect.Min.X,  InputViewport.Rect.Min.Y,  InputViewport.Rect.Max.X,  InputViewport.Rect.Max.Y);
 
 	FRHIDepthStencilState* DepthStencilState = TStaticDepthStencilState<true, CF_Always>::GetRHI();
 
+	if (bIsMinAndMaxDepthFilter)
+	{
+		DepthStencilState = TStaticDepthStencilState<false>::GetRHI();
+		PassParameters->RenderTargets[0] = FRenderTargetBinding(Output.Texture, Output.LoadAction);
+	}
+	else
+	{
+		PassParameters->RenderTargets.DepthStencil = FDepthStencilBinding(Output.Texture, Output.LoadAction, Output.LoadAction, FExclusiveDepthStencil::DepthWrite_StencilWrite);
+	}
+
 	static const TCHAR* kFilterNames[] = {
 		TEXT("Point"),
 		TEXT("Max"),
 		TEXT("CheckerMinMax"),
+		TEXT("MinAndMaxDepth"),
 	};
 
 	AddDrawScreenPass(
diff --git a/Engine/Source/Runtime/Renderer/Private/ScreenPass.h b/Engine/Source/Runtime/Renderer/Private/ScreenPass.h
--- a/Engine/Source/Runtime/Renderer/Private/ScreenPass.h
+++ b/Engine/Source/Runtime/Renderer/Private/ScreenPass.h
@@ -660,7 +660,10 @@ enum class EDownsampleDepthFilter
 	Max,
 
 	// Produces a checkerboarded selection of min and max depth values
-	Checkerboard
+	Checkerboard,
+
+	// Procuce a color texture where R=min and G=max DeviceZ
+	MinAndMaxDepth
 };
 
 void AddDownsampleDepthPass(
diff --git a/Engine/Source/Runtime/Renderer/Private/VolumetricCloudRendering.cpp b/Engine/Source/Runtime/Renderer/Private/VolumetricCloudRendering.cpp
--- a/Engine/Source/Runtime/Renderer/Private/VolumetricCloudRendering.cpp
+++ b/Engine/Source/Runtime/Renderer/Private/VolumetricCloudRendering.cpp
@@ -290,6 +290,28 @@ static bool ShouldRenderCloudShadowmap(const FLightSceneProxy* AtmosphericLight)
 	return CVarVolumetricCloudShadowMap.GetValueOnRenderThread() > 0 && AtmosphericLight && AtmosphericLight->GetCastCloudShadows();
 }
 
+bool ShouldVolumetricCloudTraceWithMinMaxDepth(const FViewInfo& ViewInfo)
+{
+	if (ViewInfo.ViewState)
+	{
+		FVolumetricRenderTargetViewStateData& VRT = ViewInfo.ViewState->VolumetricCloudRenderTarget;
+		return ShouldViewRenderVolumetricCloudRenderTarget(ViewInfo) && ShouldUseComputeForCloudTracing(ViewInfo.FeatureLevel) && VRT.GetMode() == 0;
+	}
+	return false;
+}
+
+bool ShouldVolumetricCloudTraceWithMinMaxDepth(const TArray<FViewInfo>& Views)
+{
+	for (int32 ViewIndex = 0; ViewIndex < Views.Num(); ViewIndex++)
+	{
+		if (ShouldVolumetricCloudTraceWithMinMaxDepth(Views[ViewIndex]))
+		{
+			return true;
+		}
+	}
+	return false;
+}
+
 bool VolumetricCloudWantsSeparatedAtmosphereMieRayLeigh(const FScene* Scene)
 {
 	const FVolumetricCloudRenderSceneInfo* VCloud = Scene->GetVolumetricCloudSceneInfo();
@@ -553,6 +575,7 @@ BEGIN_GLOBAL_SHADER_PARAMETER_STRUCT(FRenderVolumetricCloudGlobalParameters, )
 	SHADER_PARAMETER_STRUCT_INCLUDE(FVolumetricCloudCommonShaderParameters, VolumetricCloud)
 	SHADER_PARAMETER_STRUCT(FSceneTextureUniformParameters, SceneTextures)
 	SHADER_PARAMETER_RDG_TEXTURE(Texture2D, SceneDepthTexture)
+	SHADER_PARAMETER_RDG_TEXTURE(Texture2D<float2>, SceneDepthMinAndMaxTexture)
 	SHADER_PARAMETER_RDG_TEXTURE(Texture2D<float3>, CloudShadowTexture0)
 	SHADER_PARAMETER_RDG_TEXTURE(Texture2D<float3>, CloudShadowTexture1)
 	SHADER_PARAMETER_SAMPLER(SamplerState, CloudBilinearTextureSampler)
@@ -810,6 +833,7 @@ class FRenderVolumetricCloudRenderViewCS : public FMeshMaterialShader
 	class FCloudSampleAtmosphericLightShadowmap : SHADER_PERMUTATION_BOOL("CLOUD_SAMPLE_ATMOSPHERIC_LIGHT_SHADOWMAP");
 	class FCloudSampleSecondLight : SHADER_PERMUTATION_BOOL("CLOUD_SAMPLE_SECOND_LIGHT");
 	class FCloudSampleLocalLights : SHADER_PERMUTATION_BOOL("CLOUD_SAMPLE_LOCAL_LIGHTS");
+	class FCloudMinAndMaxDepth : SHADER_PERMUTATION_BOOL("CLOUD_MIN_AND_MAX_DEPTH");
 	class FCloudDebugViewMode : SHADER_PERMUTATION_BOOL("CLOUD_DEBUG_VIEW_MODE");	// Only the CS feature this because CS is typically used in editor.
 
 	using FPermutationDomain = TShaderPermutationDomain<
@@ -817,6 +841,7 @@ class FRenderVolumetricCloudRenderViewCS : public FMeshMaterialShader
 		FCloudSampleAtmosphericLightShadowmap,
 		FCloudSampleSecondLight,
 		FCloudSampleLocalLights,
+		FCloudMinAndMaxDepth,
 		FCloudDebugViewMode>;
 
 	BEGIN_SHADER_PARAMETER_STRUCT(FParameters, )
@@ -828,7 +853,8 @@ class FRenderVolumetricCloudRenderViewCS : public FMeshMaterialShader
 		SHADER_PARAMETER_STRUCT_REF(FViewUniformShaderParameters, View)
 		SHADER_PARAMETER_STRUCT_INCLUDE(FVirtualShadowMapSamplingParameters, VirtualShadowMap)
 		SHADER_PARAMETER(int32, VirtualShadowMapId0)
-		SHADER_PARAMETER_RDG_TEXTURE_UAV(RWTexture2D, OutCloudColor)
+		SHADER_PARAMETER_RDG_TEXTURE_UAV(RWTexture2D, OutCloudColor0)
+		SHADER_PARAMETER_RDG_TEXTURE_UAV(RWTexture2D, OutCloudColor1)
 		SHADER_PARAMETER_RDG_TEXTURE_UAV(RWTexture2D, OutCloudDepth)
 		SHADER_PARAMETER_RDG_TEXTURE_UAV(RWTexture2D, OutCloudColorCube)
 		SHADER_PARAMETER_RDG_TEXTURE(Texture2D, StartTracingDistanceTexture)
@@ -2067,7 +2093,11 @@ void FCloudRenderContext::CreateDefaultTexturesIfNeeded(FRDGBuilder& GraphBuilde
 			FRDGTextureDesc::CreateCube(1, PF_FloatRGBA, FClearValueBinding::Black, TexCreate_ShaderResource | TexCreate_RenderTargetable | TexCreate_UAV),
 			TEXT("Cloud.ColorCubeDummy"));
 
-		DefaultCloudColor2DTexture = GraphBuilder.CreateTexture(
+		DefaultCloudColor02DTexture = GraphBuilder.CreateTexture(
+			FRDGTextureDesc::Create2D(FIntPoint(1, 1), PF_FloatRGBA, FClearValueBinding::None, TexCreate_ShaderResource | TexCreate_RenderTargetable | TexCreate_UAV),
+			TEXT("Cloud.ColorDummy"));
+
+		DefaultCloudColor12DTexture = GraphBuilder.CreateTexture(
 			FRDGTextureDesc::Create2D(FIntPoint(1, 1), PF_FloatRGBA, FClearValueBinding::None, TexCreate_ShaderResource | TexCreate_RenderTargetable | TexCreate_UAV),
 			TEXT("Cloud.ColorDummy"));
 
@@ -2077,7 +2107,8 @@ void FCloudRenderContext::CreateDefaultTexturesIfNeeded(FRDGBuilder& GraphBuilde
 
 
 		DefaultCloudColorCubeTextureUAV = GraphBuilder.CreateUAV(DefaultCloudColorCubeTexture, ERDGUnorderedAccessViewFlags::SkipBarrier);
-		DefaultCloudColor2DTextureUAV = GraphBuilder.CreateUAV(DefaultCloudColor2DTexture, ERDGUnorderedAccessViewFlags::SkipBarrier);
+		DefaultCloudColor02DTextureUAV = GraphBuilder.CreateUAV(DefaultCloudColor12DTexture, ERDGUnorderedAccessViewFlags::SkipBarrier);
+		DefaultCloudColor12DTextureUAV = GraphBuilder.CreateUAV(DefaultCloudColor02DTexture, ERDGUnorderedAccessViewFlags::SkipBarrier);
 		DefaultCloudDepthTextureUAV = GraphBuilder.CreateUAV(DefaultCloudDepthTexture, ERDGUnorderedAccessViewFlags::SkipBarrier);
 	}
 }
@@ -2091,6 +2122,7 @@ static TRDGUniformBufferRef<FRenderVolumetricCloudGlobalParameters> CreateCloudP
 	SetupDefaultRenderVolumetricCloudGlobalParameters(GraphBuilder, VolumetricCloudParams, CloudInfo, MainView);
 
 	VolumetricCloudParams.SceneDepthTexture = CloudRC.SceneDepthZ;
+	VolumetricCloudParams.SceneDepthMinAndMaxTexture = CloudRC.SceneDepthMinAndMax ? CloudRC.SceneDepthMinAndMax : ((bool)ERHIZBuffer::IsInverted ? GSystemTextures.GetBlackDummy(GraphBuilder) : GSystemTextures.GetWhiteDummy(GraphBuilder));
 	if (CloudRC.bShouldViewRenderVolumetricRenderTarget && MainView.ViewState)
 	{
 		FVolumetricRenderTargetViewStateData& VRT = MainView.ViewState->VolumetricCloudRenderTarget;
@@ -2169,7 +2201,12 @@ static TRDGUniformBufferRef<FRenderVolumetricCloudGlobalParameters> CreateCloudP
 	return GraphBuilder.CreateUniformBuffer(&VolumetricCloudParams);
 }
 
-static void GetOutputTexturesWithFallback(FRDGBuilder& GraphBuilder, FCloudRenderContext& CloudRC, FRDGTextureUAVRef& CloudColorCubeTextureUAV, FRDGTextureUAVRef& CloudColorTextureUAV, FRDGTextureUAVRef& CloudDepthTextureUAV)
+static void GetOutputTexturesWithFallback(
+	FRDGBuilder& GraphBuilder, FCloudRenderContext& CloudRC, 
+	FRDGTextureUAVRef& CloudColorCubeTextureUAV, 
+	FRDGTextureUAVRef& CloudColor0TextureUAV,
+	FRDGTextureUAVRef& CloudColor1TextureUAV,
+	FRDGTextureUAVRef& CloudDepthTextureUAV)
 {
 	CloudRC.CreateDefaultTexturesIfNeeded(GraphBuilder);
 
@@ -2191,14 +2228,24 @@ static void GetOutputTexturesWithFallback(FRDGBuilder& GraphBuilder, FCloudRende
 		}
 	}
 
-	FRDGTextureRef CloudColorTexture = CloudRC.RenderTargets[0].GetTexture();
-	if (!CloudColorTexture || CloudColorTexture->Desc.IsTextureCube())
+	FRDGTextureRef CloudColor0Texture = CloudRC.RenderTargets[0].GetTexture();
+	if (!CloudColor0Texture || CloudColor0Texture->Desc.IsTextureCube())
 	{
-		CloudColorTextureUAV = CloudRC.DefaultCloudColor2DTextureUAV;
+		CloudColor0TextureUAV = CloudRC.DefaultCloudColor02DTextureUAV;
 	}
 	else
 	{
-		CloudColorTextureUAV = GraphBuilder.CreateUAV(CloudColorTexture, ERDGUnorderedAccessViewFlags::SkipBarrier);
+		CloudColor0TextureUAV = GraphBuilder.CreateUAV(CloudColor0Texture, ERDGUnorderedAccessViewFlags::SkipBarrier);
+	}
+
+	FRDGTextureRef CloudColor1Texture = CloudRC.SecondaryCloudTracingDataTexture;
+	if (!CloudColor1Texture || CloudColor1Texture->Desc.IsTextureCube())
+	{
+		CloudColor1TextureUAV = CloudRC.DefaultCloudColor12DTextureUAV;
+	}
+	else
+	{
+		CloudColor1TextureUAV = GraphBuilder.CreateUAV(CloudColor1Texture, ERDGUnorderedAccessViewFlags::SkipBarrier);
 	}
 
 	FRDGTextureRef CloudDepthTexture = CloudRC.RenderTargets[1].GetTexture();
@@ -2353,9 +2400,10 @@ void FSceneRenderer::RenderVolumetricCloudsInternal(FRDGBuilder& GraphBuilder, F
 		//
 
 		FRDGTextureUAVRef CloudColorCubeTextureUAV;
-		FRDGTextureUAVRef CloudColorTextureUAV;
+		FRDGTextureUAVRef CloudColor0TextureUAV;
+		FRDGTextureUAVRef CloudColor1TextureUAV;
 		FRDGTextureUAVRef CloudDepthTextureUAV;
-		GetOutputTexturesWithFallback(GraphBuilder, CloudRC, CloudColorCubeTextureUAV, CloudColorTextureUAV, CloudDepthTextureUAV);
+		GetOutputTexturesWithFallback(GraphBuilder, CloudRC, CloudColorCubeTextureUAV, CloudColor0TextureUAV, CloudColor1TextureUAV, CloudDepthTextureUAV);
 
 		bool bSampleVirtualShadowMap = (!bCloudDebugViewModeEnabled && !bSkipAtmosphericLightShadowmap && CVarVolumetricCloudShadowSampleAtmosphericLightShadowmap.GetValueOnRenderThread() > 0);
 
@@ -2369,7 +2417,8 @@ void FSceneRenderer::RenderVolumetricCloudsInternal(FRDGBuilder& GraphBuilder, F
 		{
 			PassParameters->VirtualShadowMap = VirtualShadowMapArray.GetSamplingParameters(GraphBuilder);
 		}
-		PassParameters->OutCloudColor = CloudColorTextureUAV;
+		PassParameters->OutCloudColor0 = CloudColor0TextureUAV;
+		PassParameters->OutCloudColor1 = CloudColor1TextureUAV;
 		PassParameters->OutCloudDepth = CloudDepthTextureUAV;
 		PassParameters->OutCloudColorCube = CloudColorCubeTextureUAV;
 
@@ -2386,11 +2435,14 @@ void FSceneRenderer::RenderVolumetricCloudsInternal(FRDGBuilder& GraphBuilder, F
 		const FMaterial* MaterialResource = &CloudVolumeMaterialProxy->GetMaterialWithFallback(Scene->GetFeatureLevel(), MaterialRenderProxy);
 		MaterialRenderProxy = MaterialRenderProxy ? MaterialRenderProxy : CloudVolumeMaterialProxy;
 
+		FVolumetricRenderTargetViewStateData& VRT = MainView.ViewState->VolumetricCloudRenderTarget;
+
 		typename FRenderVolumetricCloudRenderViewCS::FPermutationDomain PermutationVector;
 		PermutationVector.Set<typename FRenderVolumetricCloudRenderViewCS::FCloudPerSampleAtmosphereTransmittance>(!bCloudDebugViewModeEnabled && ShouldUsePerSampleAtmosphereTransmittance(Scene, &MainView));
 		PermutationVector.Set<typename FRenderVolumetricCloudRenderViewCS::FCloudSampleAtmosphericLightShadowmap>(bSampleVirtualShadowMap);
 		PermutationVector.Set<typename FRenderVolumetricCloudRenderViewCS::FCloudSampleSecondLight>(!bCloudDebugViewModeEnabled && bSecondAtmosphereLightEnabled);
 		PermutationVector.Set<typename FRenderVolumetricCloudRenderViewCS::FCloudSampleLocalLights>(!bCloudDebugViewModeEnabled && bCloudEnableLocalLightSampling && !CloudRC.bIsSkyRealTimeReflectionRendering);
+		PermutationVector.Set<typename FRenderVolumetricCloudRenderViewCS::FCloudMinAndMaxDepth>(ShouldVolumetricCloudTraceWithMinMaxDepth(MainView) && CloudRC.SecondaryCloudTracingDataTexture!=nullptr && !CloudRC.bIsReflectionRendering);// Only for mode 0 in non reflection
 		PermutationVector.Set<typename FRenderVolumetricCloudRenderViewCS::FCloudDebugViewMode>(bCloudDebugViewModeEnabled);
 
 		TShaderRef<FRenderVolumetricCloudRenderViewCS> ComputeShader = MaterialResource->GetShader<FRenderVolumetricCloudRenderViewCS>(&FLocalVertexFactory::StaticType, PermutationVector, false);
@@ -2484,6 +2536,7 @@ bool FSceneRenderer::RenderVolumetricCloud(
 	bool bSkipVolumetricRenderTarget,
 	bool bSkipPerPixelTracing,
 	FRDGTextureRef HalfResolutionDepthCheckerboardMinMaxTexture,
+	FRDGTextureRef QuarterResolutionDepthMinMaxTexture,
 	bool bAsyncCompute,
 	FInstanceCullingManager& InstanceCullingManager)
 {
@@ -2620,10 +2673,17 @@ bool FSceneRenderer::RenderVolumetricCloud(
 					CloudRC.TracingCoordToZbufferCoordScaleBias.W += ViewInfo.CachedViewUniformShaderParameters->ViewRectMin.Y / ((VRT.GetMode() == 0 || VRT.GetMode() == 3) ? 2 : 1);
 					CloudRC.TracingCoordToFullResPixelCoordScaleBias = VRT.GetTracingCoordToFullResPixelCoordScaleBias();
 
+					const bool bShouldVolumetricCloudTraceWithMinMaxDepth = ShouldVolumetricCloudTraceWithMinMaxDepth(ViewInfo);
+					if (bShouldVolumetricCloudTraceWithMinMaxDepth)
+					{
+						CloudRC.SecondaryCloudTracingDataTexture = VRT.GetOrCreateVolumetricSecondaryTracingRT(GraphBuilder); 
+					}
+
 					CloudRC.NoiseFrameIndexModPattern = VRT.GetNoiseFrameIndexModPattern();
 
 					check(VRT.GetMode() != 0 || CloudRC.bCloudDebugViewModeEnabled || HalfResolutionDepthCheckerboardMinMaxTexture);
 					CloudRC.SceneDepthZ = (VRT.GetMode() == 0 || VRT.GetMode() == 3) ? HalfResolutionDepthCheckerboardMinMaxTexture : SceneDepthZ;
+					CloudRC.SceneDepthMinAndMax = bShouldVolumetricCloudTraceWithMinMaxDepth ? QuarterResolutionDepthMinMaxTexture : SceneDepthZ;
 				}
 				else
 				{
@@ -2658,11 +2718,13 @@ bool FSceneRenderer::RenderVolumetricCloud(
 					CloudRC.RenderTargets[1] = FRenderTargetBinding(DestinationRTDepth, bShouldUseHighQualityAerialPerspective ? ERenderTargetLoadAction::EClear : ERenderTargetLoadAction::ENoAction);
 
 					CloudRC.SceneDepthZ = SceneDepthZ;
+					CloudRC.SceneDepthMinAndMax = SceneDepthZ;
 				}
 
 				if (CloudRC.bCloudDebugViewModeEnabled)
 				{
 					CloudRC.SceneDepthZ = (bool)ERHIZBuffer::IsInverted ? GSystemTextures.GetBlackDummy(GraphBuilder) : GSystemTextures.GetWhiteDummy(GraphBuilder);
+					CloudRC.SceneDepthMinAndMax = CloudRC.SceneDepthZ;
 				}
 
 				const FProjectedShadowInfo* ProjectedShadowInfo0 = nullptr;
diff --git a/Engine/Source/Runtime/Renderer/Private/VolumetricCloudRendering.h b/Engine/Source/Runtime/Renderer/Private/VolumetricCloudRendering.h
--- a/Engine/Source/Runtime/Renderer/Private/VolumetricCloudRendering.h
+++ b/Engine/Source/Runtime/Renderer/Private/VolumetricCloudRendering.h
@@ -108,9 +108,10 @@ bool ShouldRenderVolumetricCloud(const FScene* Scene, const FEngineShowFlags& En
 bool ShouldViewVisualizeVolumetricCloudConservativeDensity(const FViewInfo& ViewInfo, const FEngineShowFlags& EngineShowFlags);
 bool VolumetricCloudWantsToSampleLocalLights(const FScene* Scene, const FEngineShowFlags& EngineShowFlags);
 uint32 GetVolumetricCloudDebugViewMode(const FEngineShowFlags& ShowFlags);
+bool ShouldVolumetricCloudTraceWithMinMaxDepth(const FViewInfo& ViewInfo);
+bool ShouldVolumetricCloudTraceWithMinMaxDepth(const TArray<FViewInfo>& Views);
 bool VolumetricCloudWantsSeparatedAtmosphereMieRayLeigh(const FScene* Scene);
 
-
 // Structure with data necessary to specify a cloud render.
 struct FCloudRenderContext
 {
@@ -121,6 +122,7 @@ struct FCloudRenderContext
 	FMaterialRenderProxy* CloudVolumeMaterialProxy;
 
 	FRDGTextureRef SceneDepthZ = nullptr;
+	FRDGTextureRef SceneDepthMinAndMax = nullptr;
 
 	///////////////////////////////////
 	// Per view parameters
@@ -129,6 +131,8 @@ struct FCloudRenderContext
 	TUniformBufferRef<FViewUniformShaderParameters> ViewUniformBuffer;
 	FRenderTargetBindingSlots RenderTargets;
 
+	FRDGTextureRef SecondaryCloudTracingDataTexture = nullptr;
+
 	bool bShouldViewRenderVolumetricRenderTarget;
 	bool bSkipAerialPerspective;
 	bool bIsReflectionRendering;				// Reflection capture and real time sky capture
@@ -149,10 +153,12 @@ struct FCloudRenderContext
 	int VirtualShadowMapId0 = INDEX_NONE;
 
 	FRDGTextureRef DefaultCloudColorCubeTexture = nullptr;
-	FRDGTextureRef DefaultCloudColor2DTexture = nullptr;
+	FRDGTextureRef DefaultCloudColor02DTexture = nullptr;
+	FRDGTextureRef DefaultCloudColor12DTexture = nullptr;
 	FRDGTextureRef DefaultCloudDepthTexture = nullptr;
 	FRDGTextureUAVRef DefaultCloudColorCubeTextureUAV = nullptr;
-	FRDGTextureUAVRef DefaultCloudColor2DTextureUAV = nullptr;
+	FRDGTextureUAVRef DefaultCloudColor02DTextureUAV = nullptr;
+	FRDGTextureUAVRef DefaultCloudColor12DTextureUAV = nullptr;
 	FRDGTextureUAVRef DefaultCloudDepthTextureUAV = nullptr;
 
 	FRDGTextureUAVRef ComputeOverlapCloudColorCubeTextureUAVWithoutBarrier = nullptr;
diff --git a/Engine/Source/Runtime/Renderer/Private/VolumetricRenderTarget.cpp b/Engine/Source/Runtime/Renderer/Private/VolumetricRenderTarget.cpp
--- a/Engine/Source/Runtime/Renderer/Private/VolumetricRenderTarget.cpp
+++ b/Engine/Source/Runtime/Renderer/Private/VolumetricRenderTarget.cpp
@@ -208,6 +208,7 @@ void FVolumetricRenderTargetViewStateData::Initialise(
 
 			// Need a new size so release the low resolution trace buffer
 			VolumetricTracingRT.SafeRelease();
+			VolumetricSecondaryTracingRT.SafeRelease();
 			VolumetricTracingRTDepth.SafeRelease();
 		}
 
@@ -268,14 +269,31 @@ FRDGTextureRef FVolumetricRenderTargetViewStateData::GetOrCreateVolumetricTracin
 	return GraphBuilder.RegisterExternalTexture(VolumetricTracingRT);
 }
 
+FRDGTextureRef FVolumetricRenderTargetViewStateData::GetOrCreateVolumetricSecondaryTracingRT(FRDGBuilder& GraphBuilder)
+{
+	check(FullResolution != FIntPoint::ZeroValue); // check that initialization has been done at least once
+
+	if (!VolumetricSecondaryTracingRT.IsValid())
+	{
+		FPooledRenderTargetDesc Desc = FPooledRenderTargetDesc::Create2DDesc(
+			VolumetricTracingRTResolution, PF_FloatRGBA, FClearValueBinding(FLinearColor(0.0f, 0.0f, 0.0f, 1.0f)),
+			TexCreate_None, TexCreate_ShaderResource | TexCreate_RenderTargetable | TexCreate_UAV, false);
+		GRenderTargetPool.FindFreeElement(GraphBuilder.RHICmdList, Desc, VolumetricSecondaryTracingRT, TEXT("VolumetricRenderTarget.SecondaryTracing"));
+	}
+
+	return GraphBuilder.RegisterExternalTexture(VolumetricSecondaryTracingRT);
+}
+
 FRDGTextureRef FVolumetricRenderTargetViewStateData::GetOrCreateVolumetricTracingRTDepth(FRDGBuilder& GraphBuilder)
 {
 	check(FullResolution != FIntPoint::ZeroValue); // check that initialization has been done at least once
 
 	if (!VolumetricTracingRTDepth.IsValid())
 	{
+		EPixelFormat DepthDataFormat = Mode == 0 ? PF_FloatRGBA : PF_G16R16F;	// TODO This should accoutn for the new mode 0 variation cvar. Or mode 5?
+
 		FPooledRenderTargetDesc Desc = FPooledRenderTargetDesc::Create2DDesc(
-			VolumetricTracingRTResolution, PF_G16R16F, FClearValueBinding(FLinearColor(63000.0f, 63000.0f, 63000.0f, 63000.0f)),
+			VolumetricTracingRTResolution, DepthDataFormat, FClearValueBinding(FLinearColor(63000.0f, 63000.0f, 63000.0f, 63000.0f)),
 			TexCreate_None, TexCreate_ShaderResource | TexCreate_RenderTargetable | TexCreate_UAV, false);
 		GRenderTargetPool.FindFreeElement(GraphBuilder.RHICmdList, Desc, VolumetricTracingRTDepth, TEXT("VolumetricRenderTarget.TracingDepth"));
 	}
@@ -435,11 +453,13 @@ class FReconstructVolumetricRenderTargetPS : public FGlobalShader
 
 	class FHistoryAvailable : SHADER_PERMUTATION_BOOL("PERMUTATION_HISTORY_AVAILABLE");
 	class FReprojectionBoxConstraint : SHADER_PERMUTATION_BOOL("PERMUTATION_REPROJECTION_BOX_CONSTRAINT");
-	using FPermutationDomain = TShaderPermutationDomain<FHistoryAvailable, FReprojectionBoxConstraint>;
+	class FCloudMinAndMaxDepth : SHADER_PERMUTATION_BOOL("CLOUD_MIN_AND_MAX_DEPTH");
+	using FPermutationDomain = TShaderPermutationDomain<FHistoryAvailable, FReprojectionBoxConstraint, FCloudMinAndMaxDepth>;
 
 	BEGIN_SHADER_PARAMETER_STRUCT(FParameters, )
 		SHADER_PARAMETER_STRUCT_REF(FViewUniformShaderParameters, ViewUniformBuffer)
 		SHADER_PARAMETER_RDG_TEXTURE(Texture2D, TracingVolumetricTexture)
+		SHADER_PARAMETER_RDG_TEXTURE(Texture2D, SecondaryTracingVolumetricTexture)
 		SHADER_PARAMETER_RDG_TEXTURE(Texture2D, TracingVolumetricDepthTexture)
 		SHADER_PARAMETER_RDG_TEXTURE(Texture2D, PreviousFrameVolumetricTexture)
 		SHADER_PARAMETER_RDG_TEXTURE(Texture2D, PreviousFrameVolumetricDepthTexture)
@@ -507,6 +527,7 @@ void ReconstructVolumetricRenderTarget(
 		FRDGTextureRef DstVolumetric = VolumetricCloudRT.GetOrCreateDstVolumetricReconstructRT(GraphBuilder);
 		FRDGTextureRef DstVolumetricDepth = VolumetricCloudRT.GetOrCreateDstVolumetricReconstructRTDepth(GraphBuilder);
 		FRDGTextureRef SrcTracingVolumetric = VolumetricCloudRT.GetOrCreateVolumetricTracingRT(GraphBuilder);
+		FRDGTextureRef SrcSecondaryTracingVolumetric = VolumetricCloudRT.GetOrCreateVolumetricSecondaryTracingRT(GraphBuilder);
 		FRDGTextureRef SrcTracingVolumetricDepth = VolumetricCloudRT.GetOrCreateVolumetricTracingRTDepth(GraphBuilder);
 		FRDGTextureRef PreviousFrameVolumetricTexture = VolumetricCloudRT.GetHistoryValid() ? VolumetricCloudRT.GetOrCreateSrcVolumetricReconstructRT(GraphBuilder) : SystemTextures.Black;
 		FRDGTextureRef PreviousFrameVolumetricDepthTexture = VolumetricCloudRT.GetHistoryValid() ? VolumetricCloudRT.GetOrCreateSrcVolumetricReconstructRTDepth(GraphBuilder) : SystemTextures.Black;
@@ -516,6 +537,7 @@ void ReconstructVolumetricRenderTarget(
 		FReconstructVolumetricRenderTargetPS::FPermutationDomain PermutationVector;
 		PermutationVector.Set<FReconstructVolumetricRenderTargetPS::FHistoryAvailable>(VolumetricCloudRT.GetHistoryValid());
 		PermutationVector.Set<FReconstructVolumetricRenderTargetPS::FReprojectionBoxConstraint>(CVarVolumetricRenderTargetReprojectionBoxConstraint.GetValueOnAnyThread() > 0);
+		PermutationVector.Set<FReconstructVolumetricRenderTargetPS::FCloudMinAndMaxDepth>(ShouldVolumetricCloudTraceWithMinMaxDepth(ViewInfo));
 		TShaderMapRef<FReconstructVolumetricRenderTargetPS> PixelShader(ViewInfo.ShaderMap, PermutationVector);
 
 		FReconstructVolumetricRenderTargetPS::FParameters* PassParameters = GraphBuilder.AllocParameters<FReconstructVolumetricRenderTargetPS::FParameters>();
@@ -523,6 +545,7 @@ void ReconstructVolumetricRenderTarget(
 		PassParameters->RenderTargets[0] = FRenderTargetBinding(DstVolumetric, ERenderTargetLoadAction::ENoAction);
 		PassParameters->RenderTargets[1] = FRenderTargetBinding(DstVolumetricDepth, ERenderTargetLoadAction::ENoAction);
 		PassParameters->TracingVolumetricTexture = SrcTracingVolumetric;
+		PassParameters->SecondaryTracingVolumetricTexture = SrcSecondaryTracingVolumetric;
 		PassParameters->TracingVolumetricDepthTexture = SrcTracingVolumetricDepth;
 		PassParameters->PreviousFrameVolumetricTexture = PreviousFrameVolumetricTexture;
 		PassParameters->PreviousFrameVolumetricDepthTexture = PreviousFrameVolumetricDepthTexture;
diff --git a/Engine/Source/Runtime/Renderer/Private/VolumetricRenderTargetViewStateData.h b/Engine/Source/Runtime/Renderer/Private/VolumetricRenderTargetViewStateData.h
--- a/Engine/Source/Runtime/Renderer/Private/VolumetricRenderTargetViewStateData.h
+++ b/Engine/Source/Runtime/Renderer/Private/VolumetricRenderTargetViewStateData.h
@@ -28,6 +28,7 @@ public:
 		int32 UpsamplingMode);
 
 	FRDGTextureRef GetOrCreateVolumetricTracingRT(FRDGBuilder& GraphBuilder);
+	FRDGTextureRef GetOrCreateVolumetricSecondaryTracingRT(FRDGBuilder& GraphBuilder);
 	FRDGTextureRef GetOrCreateVolumetricTracingRTDepth(FRDGBuilder& GraphBuilder);
 
 	FRDGTextureRef GetDstVolumetricReconstructRT(FRDGBuilder& GraphBuilder);
@@ -81,6 +82,7 @@ private:
 	TRefCountPtr<IPooledRenderTarget> VolumetricReconstructRTDepth[kRenderTargetCount];
 
 	TRefCountPtr<IPooledRenderTarget> VolumetricTracingRT;
+	TRefCountPtr<IPooledRenderTarget> VolumetricSecondaryTracingRT;
 	TRefCountPtr<IPooledRenderTarget> VolumetricTracingRTDepth;
 
 	float UvNoiseScale;
