Disabled color clamping in mobile pixel shaders when its not necessary

diff --git a/Engine/Shaders/Private/MobileBasePassPixelShader.usf b/Engine/Shaders/Private/MobileBasePassPixelShader.usf
--- a/Engine/Shaders/Private/MobileBasePassPixelShader.usf
+++ b/Engine/Shaders/Private/MobileBasePassPixelShader.usf
@@ -683,11 +683,14 @@ void Main(
 	OutSceneDepthAux = SvPosition.z;
 #endif
 
-	#if !MATERIALBLENDING_MODULATE
-		OutColor.rgb *= ResolvedView.PreExposure;
-	#endif
+#if !MATERIALBLENDING_MODULATE
+	OutColor.rgb *= ResolvedView.PreExposure;
+#endif
 
+#if !MATERIAL_SHADINGMODEL_UNLIT && !DEFERRED_SHADING_PATH
+	// clamp lit color to avoid overflows
 	OutColor.rgb = SafeGetOutColor(OutColor.rgb);
+#endif
 
 #if USE_EDITOR_COMPOSITING && (MOBILE_EMULATION)
 	// Editor primitive depth testing
diff --git a/Engine/Shaders/Private/MobileLightingCommon.ush b/Engine/Shaders/Private/MobileLightingCommon.ush
--- a/Engine/Shaders/Private/MobileLightingCommon.ush
+++ b/Engine/Shaders/Private/MobileLightingCommon.ush
@@ -125,7 +125,7 @@ return ShadowMap;
 #endif
 
 #ifndef MAX_MOBILE_SHADOWCASCADES
-#define MAX_MOBILE_SHADOWCASCADES 4
+#define MAX_MOBILE_SHADOWCASCADES 4u
 #endif 
 
 half MobileDirectionalLightCSM(float2 ScreenPosition, float SceneDepth, inout float ShadowPositionZ)
@@ -144,7 +144,7 @@ half MobileDirectionalLightCSM(float2 ScreenPosition, float SceneDepth, inout fl
 	Settings.ProjectionDepthBiasParameters = 0;
 
 	float4 ShadowPosition = float4(0, 0, 0, 0);
-	for (int i = 0; i < MAX_MOBILE_SHADOWCASCADES; i++)
+	for (uint i = 0; i < MAX_MOBILE_SHADOWCASCADES; i++)
 	{
 		if (SceneDepth < MobileDirectionalLight.DirectionalLightShadowDistances[i])
 		{
