Include sky materials for color clamping on mobile

diff --git a/Engine/Shaders/Private/MobileBasePassPixelShader.usf b/Engine/Shaders/Private/MobileBasePassPixelShader.usf
--- a/Engine/Shaders/Private/MobileBasePassPixelShader.usf
+++ b/Engine/Shaders/Private/MobileBasePassPixelShader.usf
@@ -687,7 +687,7 @@ void Main(
 	OutColor.rgb *= ResolvedView.PreExposure;
 #endif
 
-#if !MATERIAL_SHADINGMODEL_UNLIT && !DEFERRED_SHADING_PATH
+#if MATERIAL_IS_SKY || (!MATERIAL_SHADINGMODEL_UNLIT && !DEFERRED_SHADING_PATH)
 	// clamp lit color to avoid overflows
 	OutColor.rgb = SafeGetOutColor(OutColor.rgb);
 #endif
