Fixes overblur sources one moving objects on Epic quality

diff --git a/Engine/Shaders/Private/TemporalSuperResolution/TSRUpdateHistory.usf b/Engine/Shaders/Private/TemporalSuperResolution/TSRUpdateHistory.usf
--- a/Engine/Shaders/Private/TemporalSuperResolution/TSRUpdateHistory.usf
+++ b/Engine/Shaders/Private/TemporalSuperResolution/TSRUpdateHistory.usf
@@ -24,7 +24,7 @@
 
 #elif DIM_UPDATE_QUALITY == 3 // Epic
 	#define CONFIG_SAMPLES 6
-	#define CONFIG_INTERPOLATE_VELOCITY 0
+	#define CONFIG_INTERPOLATE_VELOCITY 1
 	#define CONFIG_REJECTION_ANTI_ALIASING 1
 
 	// To support sg.AntiAliasingQuality=4 with same temporal response as sg.AntiAliasingQuality=3
@@ -584,8 +584,8 @@ void MainCS(
 					GrandPrevScreenPos[0][PixelId] = LocalPrevScreenPos.x;
 					GrandPrevScreenPos[1][PixelId] = LocalPrevScreenPos.y;
 				}
-				#endif
-				#else
+				#endif // CONFIG_GRAND_REPROJECTION
+				#else // !CONFIG_INTERPOLATE_VELOCITY
 				{
 					ScreenVelocity = dpv_interleave_registers(
 						DecodeVelocityFromTexture(float4(RawEncodedVelocity[0], 0.0, 0.0)).xy,
@@ -1309,6 +1309,8 @@ void MainCS(
 		}
 		#endif
 
+		ClampedPrevHistoryValidity = min(ClampedPrevHistoryValidity, lerp(tsr_half(0.2), tsr_half(1.0), BlendClamp));
+
 		// Compute the newly rejected validity
 		tsr_half2 RejectedValidity = ClampedPrevHistoryValidity * PrevHistoryRejectionWeight;
 
