Fix IES shader compilation error on CIS.

diff --git a/Engine/Shaders/Private/IESAtlas.usf b/Engine/Shaders/Private/IESAtlas.usf
--- a/Engine/Shaders/Private/IESAtlas.usf
+++ b/Engine/Shaders/Private/IESAtlas.usf
@@ -19,7 +19,14 @@ Texture2D<float4> InTexture_7;
 
 uint4 InSliceIndex[8];
 
-SamplerState InSampler;
+SamplerState InSampler_0;
+SamplerState InSampler_1;
+SamplerState InSampler_2;
+SamplerState InSampler_3;
+SamplerState InSampler_4;
+SamplerState InSampler_5;
+SamplerState InSampler_6;
+SamplerState InSampler_7;
 
 int2 AtlasResolution;
 uint AtlasSliceCount;
@@ -43,14 +50,14 @@ void MainCS(uint3 DispatchThreadId : SV_DispatchThreadID)
 			float Color = 0;
 			switch (SrcSlice)
 			{
-			case 0: Color = InTexture_0.SampleLevel(InSampler, SrcUV, 0).x; DstSlice = InSliceIndex[0].x; break;
-			case 1: Color = InTexture_1.SampleLevel(InSampler, SrcUV, 0).x; DstSlice = InSliceIndex[1].x; break;
-			case 2: Color = InTexture_2.SampleLevel(InSampler, SrcUV, 0).x; DstSlice = InSliceIndex[2].x; break;
-			case 3: Color = InTexture_3.SampleLevel(InSampler, SrcUV, 0).x; DstSlice = InSliceIndex[3].x; break;
-			case 4: Color = InTexture_4.SampleLevel(InSampler, SrcUV, 0).x; DstSlice = InSliceIndex[4].x; break;
-			case 5: Color = InTexture_5.SampleLevel(InSampler, SrcUV, 0).x; DstSlice = InSliceIndex[5].x; break;
-			case 6: Color = InTexture_6.SampleLevel(InSampler, SrcUV, 0).x; DstSlice = InSliceIndex[6].x; break;
-			case 7: Color = InTexture_7.SampleLevel(InSampler, SrcUV, 0).x; DstSlice = InSliceIndex[7].x; break;
+			case 0: Color = InTexture_0.SampleLevel(InSampler_0, SrcUV, 0).x; DstSlice = InSliceIndex[0].x; break;
+			case 1: Color = InTexture_1.SampleLevel(InSampler_1, SrcUV, 0).x; DstSlice = InSliceIndex[1].x; break;
+			case 2: Color = InTexture_2.SampleLevel(InSampler_2, SrcUV, 0).x; DstSlice = InSliceIndex[2].x; break;
+			case 3: Color = InTexture_3.SampleLevel(InSampler_3, SrcUV, 0).x; DstSlice = InSliceIndex[3].x; break;
+			case 4: Color = InTexture_4.SampleLevel(InSampler_4, SrcUV, 0).x; DstSlice = InSliceIndex[4].x; break;
+			case 5: Color = InTexture_5.SampleLevel(InSampler_5, SrcUV, 0).x; DstSlice = InSliceIndex[5].x; break;
+			case 6: Color = InTexture_6.SampleLevel(InSampler_6, SrcUV, 0).x; DstSlice = InSliceIndex[6].x; break;
+			case 7: Color = InTexture_7.SampleLevel(InSampler_7, SrcUV, 0).x; DstSlice = InSliceIndex[7].x; break;
 			}
 
 			// Ensure there is no NaN value
diff --git a/Engine/Source/Runtime/Renderer/Private/IESTextureManager.cpp b/Engine/Source/Runtime/Renderer/Private/IESTextureManager.cpp
--- a/Engine/Source/Runtime/Renderer/Private/IESTextureManager.cpp
+++ b/Engine/Source/Runtime/Renderer/Private/IESTextureManager.cpp
@@ -190,7 +190,7 @@ class FIESAtlasAddTextureCS : public FGlobalShader
 	BEGIN_SHADER_PARAMETER_STRUCT(FParameters, )
 		SHADER_PARAMETER_TEXTURE_ARRAY(Texture2D, InTexture, [8])
 		SHADER_PARAMETER_ARRAY(FUintVector4, InSliceIndex, [8])
-		SHADER_PARAMETER_SAMPLER(SamplerState, InSampler)
+		SHADER_PARAMETER_SAMPLER_ARRAY(SamplerState, InSampler, [8])
 		SHADER_PARAMETER(FIntPoint, AtlasResolution)
 		SHADER_PARAMETER(uint32, AtlasSliceCount)
 		SHADER_PARAMETER(uint32, ValidCount)
@@ -231,11 +231,11 @@ static void AddSlotsPass(
 		Parameters->AtlasResolution = OutAtlas->Desc.Extent;
 		Parameters->AtlasSliceCount = OutAtlas->Desc.ArraySize;
 		Parameters->ValidCount = SlotCount;
-		Parameters->InSampler = TStaticSamplerState<SF_Bilinear, AM_Clamp, AM_Clamp, AM_Clamp>::GetRHI();
 		for (uint32 SlotIt = 0; SlotIt < SlotCountPerPass; ++SlotIt)
 		{
 			Parameters->InTexture[SlotIt] = GSystemTextures.BlackDummy->GetRHI();
 			Parameters->InSliceIndex[SlotIt].X = InvalidSlotIndex;
+			Parameters->InSampler[SlotIt] = TStaticSamplerState<SF_Bilinear, AM_Clamp, AM_Clamp, AM_Clamp>::GetRHI();
 		}
 		for (uint32 SlotIt = 0; SlotIt<SlotCount;++SlotIt)
 		{
