Fix shader compilation issue.

diff --git a/Engine/Shaders/Private/VolumetricCloud.usf b/Engine/Shaders/Private/VolumetricCloud.usf
--- a/Engine/Shaders/Private/VolumetricCloud.usf
+++ b/Engine/Shaders/Private/VolumetricCloud.usf
@@ -1645,6 +1645,7 @@ float GetSliceDepth()
 
 const float GetSvPositionEmptySpaceSkippingValue(in ViewState ResolvedView, in uint3 DispatchThreadId, in uint3 GroupThreadId, in float3 SampleFroxelOffset)
 {
+#if MATERIAL_VOLUMETRIC_CLOUD_EMPTY_SPACE_SKIPPING_OUTPUT
 	float2 PixelUV = float2(DispatchThreadId.x + SampleFroxelOffset.x, DispatchThreadId.y + SampleFroxelOffset.y) / StartTracingDistanceTextureResolution;
 	float4 SvPosition = float4(PixelUV * ResolvedView.ViewSizeAndInvSize.xy, 0.1, 1);
 
@@ -1663,6 +1664,9 @@ const float GetSvPositionEmptySpaceSkippingValue(in ViewState ResolvedView, in u
 
 	const float CloudEmptySpaceSkippingOutput0 = GetVolumetricCloudEmptySpaceSkippingOutput0(MaterialParameters);
 	return CloudEmptySpaceSkippingOutput0;
+#else
+	return 0;
+#endif
 }
 
 [numthreads(THREADGROUP_SIZEX, THREADGROUP_SIZEY, THREADGROUP_SIZEZ)]
