Fix incorrect CalculateDistanceFieldApproxAO(...) definition if !PLATFORM_SUPPORTS_DISTANCE_FIELDS

diff --git a/Engine/Shaders/Private/DistanceField/GlobalDistanceFieldShared.ush b/Engine/Shaders/Private/DistanceField/GlobalDistanceFieldShared.ush
--- a/Engine/Shaders/Private/DistanceField/GlobalDistanceFieldShared.ush
+++ b/Engine/Shaders/Private/DistanceField/GlobalDistanceFieldShared.ush
@@ -395,7 +395,7 @@ float3 GetDistanceFieldGradientGlobal(float3 WorldPosition)
 	return float3(0, 0, .001f);
 }
 
-float CalculateDistanceFieldApproxAO(float3 WorldPosition, float3 WorldNormal, uint NumSteps, float BaseStepDistance, float StepScale, float MaxDistance)
+float CalculateDistanceFieldApproxAO(float3 WorldPosition, float3 WorldNormal, uint NumSteps, float StepDistance, float StepScale, float DistanceBias, float MaxDistance)
 {
 	return 1.0f;
 }
