Speeds up TSR's Reject Shading pass

diff --git a/Engine/Shaders/Private/TemporalSuperResolution/TSRRejectShading.usf b/Engine/Shaders/Private/TemporalSuperResolution/TSRRejectShading.usf
--- a/Engine/Shaders/Private/TemporalSuperResolution/TSRRejectShading.usf
+++ b/Engine/Shaders/Private/TemporalSuperResolution/TSRRejectShading.usf
@@ -86,6 +86,10 @@
 #define tsr_simd_halfC TLaneVector<tsr_half, CONFIG_CHANNEL_COUNT, SIMD_SIZE>
 #define tsr_simd_bool  TLaneVector<bool, 1, SIMD_SIZE>
 
+#ifndef PLATFORM_SPECIFIC_ISOLATE
+	#define PLATFORM_SPECIFIC_ISOLATE ISOLATE
+#endif
+
 
 //------------------------------------------------------- PARAMETERS
 
@@ -521,10 +525,45 @@ FSampleType Median9(FSampleType In[9])
 	return Median(M1[0], M1[1], M1[2]);
 }
 
-ISOLATE
 CALL_SITE_DEBUGLOC
 template<typename FSampleType>
 FSampleType Median3x3(FSampleType Center, uint GroupThreadIndex)
+#if 1
+{
+	WriteToLDSForAccessNeighborTexel(Center, GroupThreadIndex);
+	
+	FSampleType C0 = Center;
+	FSampleType C1 = AccessNeighborTexel(Center, tsr_short2( 1, 0), GroupThreadIndex);
+	FSampleType C2 = AccessNeighborTexel(Center, tsr_short2(-1, 0), GroupThreadIndex);
+	
+	FSampleType M0L0;
+	FSampleType M0M0;
+	FSampleType M0H0;
+	LMHOperator(C0, C1, C2, M0L0, M0M0, M0H0);
+	
+	WriteToLDSForAccessNeighborTexel(M0L0, GroupThreadIndex);
+	FSampleType M0L1 = AccessNeighborTexel(M0L0, tsr_short2(0,  1), GroupThreadIndex);
+	FSampleType M0L2 = AccessNeighborTexel(M0L0, tsr_short2(0, -1), GroupThreadIndex);
+	
+	WriteToLDSForAccessNeighborTexel(M0M0, GroupThreadIndex);
+	FSampleType M0M1 = AccessNeighborTexel(M0M0, tsr_short2(0,  1), GroupThreadIndex);
+	FSampleType M0M2 = AccessNeighborTexel(M0M0, tsr_short2(0, -1), GroupThreadIndex);
+	
+	WriteToLDSForAccessNeighborTexel(M0H0, GroupThreadIndex);
+	FSampleType M0H1 = AccessNeighborTexel(M0H0, tsr_short2(0,  1), GroupThreadIndex);
+	FSampleType M0H2 = AccessNeighborTexel(M0H0, tsr_short2(0, -1), GroupThreadIndex);
+	
+	// Second layer.
+	// TODO: min3, max3.
+	FSampleType M1[3];
+	M1[0] = max(max(M0L0, M0L1), M0L2);
+	M1[1] = Median(M0M0, M0M1, M0M2);
+	M1[2] = min(min(M0H0, M0H1), M0H2);
+	
+	// Third layer.
+	return Median(M1[0], M1[1], M1[2]);
+}
+#else
 {
 	WriteToLDSForAccessNeighborTexel(Center, GroupThreadIndex);
 	
@@ -538,6 +577,7 @@ FSampleType Median3x3(FSampleType Center, uint GroupThreadIndex)
 
 	return Median9(InSamples);
 }
+#endif
 
 
 //------------------------------------------------------- ENTRY POINT
@@ -574,7 +614,7 @@ void MainCS(
 		tsr_simd_half4 PrevMoire;
 		tsr_simd_half bIsStatic;
 	#endif
-	ISOLATE
+	PLATFORM_SPECIFIC_ISOLATE
 	{
 		UNROLL_N(SIMD_SIZE)
 		for (uint SimdIndex = 0; SimdIndex < SIMD_SIZE; SimdIndex++)
@@ -938,7 +978,7 @@ void MainCS(
 	#endif
 
 	// Output data.
-	ISOLATE
+	PLATFORM_SPECIFIC_ISOLATE
 	{
 		UNROLL_N(SIMD_SIZE)
 		for (uint SimdIndex = 0; SimdIndex < SIMD_SIZE; SimdIndex++)
