Fixed border being wrong when downsampling depth buffer due to the fact that the lower right pixel coord is excluded from raster viewport.
Unified code in depthdownsample.
Tested with split screen.

diff --git a/Engine/Shaders/Private/DownsampleDepthPixelShader.usf b/Engine/Shaders/Private/DownsampleDepthPixelShader.usf
--- a/Engine/Shaders/Private/DownsampleDepthPixelShader.usf
+++ b/Engine/Shaders/Private/DownsampleDepthPixelShader.usf
@@ -38,20 +38,17 @@ void Main(
 #endif
 )
 {
-#if OUTPUT_MIN_AND_MAX_DEPTH
-
-	// We do things differently when downsampling min and max depth because it is executed from half to quarter resolution.
-	// And in this case OutUVAndScreenPos is not valid because it uses view rect data.
-
 	// Pixel coordinate for the output 
-	const int2 DstPixelPos = SVPos;
+	const int2 DstPixelPos = SVPos.xy;
 	// Lower left corner from the input 2x2 quad
 	const int2 PixelPos = DstPixelPos * 2;
 
-	float DeviceZ0 = DepthTexture.Load(uint3(clamp(PixelPos + int2(0,0), SrcPixelCoordMinAndMax.xy, SrcPixelCoordMinAndMax.zw), 0));
-	float DeviceZ1 = DepthTexture.Load(uint3(clamp(PixelPos + int2(0,1), SrcPixelCoordMinAndMax.xy, SrcPixelCoordMinAndMax.zw), 0));
-	float DeviceZ2 = DepthTexture.Load(uint3(clamp(PixelPos + int2(1,1), SrcPixelCoordMinAndMax.xy, SrcPixelCoordMinAndMax.zw), 0));
-	float DeviceZ3 = DepthTexture.Load(uint3(clamp(PixelPos + int2(1,0), SrcPixelCoordMinAndMax.xy, SrcPixelCoordMinAndMax.zw), 0));
+	float DeviceZ0 = DepthTexture.Load(uint3(clamp(PixelPos + int2(0, 0), SrcPixelCoordMinAndMax.xy, SrcPixelCoordMinAndMax.zw), 0));
+	float DeviceZ1 = DepthTexture.Load(uint3(clamp(PixelPos + int2(0, 1), SrcPixelCoordMinAndMax.xy, SrcPixelCoordMinAndMax.zw), 0));
+	float DeviceZ2 = DepthTexture.Load(uint3(clamp(PixelPos + int2(1, 1), SrcPixelCoordMinAndMax.xy, SrcPixelCoordMinAndMax.zw), 0));
+	float DeviceZ3 = DepthTexture.Load(uint3(clamp(PixelPos + int2(1, 0), SrcPixelCoordMinAndMax.xy, SrcPixelCoordMinAndMax.zw), 0));
+
+#if OUTPUT_MIN_AND_MAX_DEPTH
 
 	const float MaxDeviceZ = max(max(DeviceZ0, DeviceZ1), max(DeviceZ2, DeviceZ3));
 	const float MinDeviceZ = min(min(DeviceZ0, DeviceZ1), min(DeviceZ2, DeviceZ3));
@@ -59,36 +56,8 @@ void Main(
 
 #else // OUTPUT_MIN_AND_MAX_DEPTH
 
-	const float2 UV = OutUVAndScreenPos.xy;
 	float Depth = 0;
 
-	// Lookup the four DeviceZ's of the full resolution pixels corresponding to this low resolution pixel
-	//
-	// Source texture texels:
-	//	|------||------|------(UV + DestinationTexelSize)
-	//	|      ||      |      |
-	//	|      ||      |      |
-	//	|      ||      |      |
-	//	|------||------|------|
-	//	|      ||      |      |
-	//	|  z0  ||  z2  |      |
-	//	|      ||      |      |
-	//	|------UV------|------(UV + float2(DestinationTexelSize.x, 0))
-	//	|      ||      |      |
-	//	|  z1  ||  z3  |      |
-	//	|      ||      |      |
-	//	|------||------|------|
-	// OffsetScale is 0.25 because:
-	//  - UV.xy is centered on a half resolution texel of the destination render target (in [0,1] uv space). 
-	//  - This texel includes 2x2 texels from the source texture (z0, z1, z2 and z3).
-	//  - Scaling the DestinationTexelSize by 0.5 will get us to the edge of the half resolution texel. 
-	//  - And to sample the center of each texel from the high resolution source texture we need to divide again the offset by half, leading to a final scale of 0.25.
-	const float OffsetScale = 0.25f;
-	float DeviceZ0 = GetDeviceZ(UV.xy + DestinationTexelSize * float2(-OffsetScale,-OffsetScale));
-	float DeviceZ1 = GetDeviceZ(UV.xy + DestinationTexelSize * float2(-OffsetScale, OffsetScale));
-	float DeviceZ2 = GetDeviceZ(UV.xy + DestinationTexelSize * float2( OffsetScale,-OffsetScale));
-	float DeviceZ3 = GetDeviceZ(UV.xy + DestinationTexelSize * float2( OffsetScale, OffsetScale));
-
 	if (DownsampleDepthFilter < DOWNSAMPLE_DEPTH_FILTER_CBMINMAX)
 	{
 #if HAS_INVERTED_Z_BUFFER
@@ -106,12 +75,9 @@ void Main(
 		const float MaxDeviceZ = max(max(DeviceZ0, DeviceZ1), max(DeviceZ2, DeviceZ3));
 		const float MinDeviceZ = min(min(DeviceZ0, DeviceZ1), min(DeviceZ2, DeviceZ3));
 
-		const float2 Resolution = uint2(DestinationResolution * View.BufferSizeAndInvSize.xy * View.ViewSizeAndInvSize.zw); // need to clamp to respect ratios
-		const uint2 PixelPos = uint2(UV.xy * Resolution);
-
-		const uint2 PixelPosStep = (PixelPos >> 1) * 2;
-		uint CheckerBoard = (PixelPos.x - PixelPosStep.x);									// horizontal alternance of black and white
-		CheckerBoard = (PixelPos.y - PixelPosStep.y) == 0 ? CheckerBoard : 1 - CheckerBoard;// vertical toggle of horizontal checker on odd lines
+		const uint2 PixelPosStep = (DstPixelPos >> 1) * 2;
+		uint CheckerBoard = (DstPixelPos.x - PixelPosStep.x);									// horizontal alternance of black and white
+		CheckerBoard = (DstPixelPos.y - PixelPosStep.y) == 0 ? CheckerBoard : 1 - CheckerBoard;// vertical toggle of horizontal checker on odd lines
 
 		Depth = CheckerBoard > 0 ? MaxDeviceZ : MinDeviceZ;
 	}
diff --git a/Engine/Source/Runtime/Renderer/Private/SceneRendering.cpp b/Engine/Source/Runtime/Renderer/Private/SceneRendering.cpp
--- a/Engine/Source/Runtime/Renderer/Private/SceneRendering.cpp
+++ b/Engine/Source/Runtime/Renderer/Private/SceneRendering.cpp
@@ -5547,7 +5547,7 @@ FRDGTextureRef CreateQuarterResolutionDepthMinAndMax(FRDGBuilder& GraphBuilder,
 		RDG_GPU_MASK_SCOPE(GraphBuilder, View.GPUMask);
 
 		const FScreenPassTexture InputDepth(InputDepthTexture, GetDownscaledRect(View.ViewRect, 2));
-		const FScreenPassRenderTarget SmallTextureRT(SmallTexture, GetDownscaledRect(View.ViewRect, 4), View.DecayLoadAction(ERenderTargetLoadAction::ENoAction));
+		const FScreenPassRenderTarget SmallTextureRT(SmallTexture, GetDownscaledRect(InputDepth.ViewRect, 2), View.DecayLoadAction(ERenderTargetLoadAction::ENoAction));
 		AddDownsampleDepthPass(GraphBuilder, View, InputDepth, SmallTextureRT, EDownsampleDepthFilter::MinAndMaxDepth);
 	}
 
diff --git a/Engine/Source/Runtime/Renderer/Private/ScreenPass.cpp b/Engine/Source/Runtime/Renderer/Private/ScreenPass.cpp
--- a/Engine/Source/Runtime/Renderer/Private/ScreenPass.cpp
+++ b/Engine/Source/Runtime/Renderer/Private/ScreenPass.cpp
@@ -218,19 +218,22 @@ void AddDownsampleDepthPass(
 	Permutation.Set<FDownsampleDepthPS::FOutputMinAndMaxDepth>(bIsMinAndMaxDepthFilter ? 1 : 0);
 	TShaderMapRef<FDownsampleDepthPS> PixelShader(View.ShaderMap, Permutation);
 
+	// The lower right corner pixel whose coordinate is max considered excluded https://learn.microsoft.com/en-us/windows/win32/direct3d11/d3d11-rect
+	// That is why we subtract -1 from the maximum value of the source viewport.
+
 	FDownsampleDepthPS::FParameters* PassParameters = GraphBuilder.AllocParameters<FDownsampleDepthPS::FParameters>();
 	PassParameters->View = View.ViewUniformBuffer;
 	PassParameters->DepthTexture = Input.Texture;
 	PassParameters->DestinationTexelSize = FVector2f(1.0f / OutputViewport.Extent.X, 1.0f / OutputViewport.Extent.X);
-	PassParameters->SourceMaxUV = FVector2f((View.ViewRect.Max.X - 0.5f) / InputViewport.Extent.X, (View.ViewRect.Max.Y - 0.5f) / InputViewport.Extent.Y);
+	PassParameters->SourceMaxUV = FVector2f((float(View.ViewRect.Max.X) -1.0f - 0.51f) / InputViewport.Extent.X, (float(View.ViewRect.Max.Y) - 1.0f - 0.51f) / InputViewport.Extent.Y);
 	PassParameters->DownsampleDepthFilter = (uint32)DownsampleDepthFilter;
 
 	const int32 DownsampledSizeX = OutputViewport.Rect.Width();
 	const int32 DownsampledSizeY = OutputViewport.Rect.Height();
 	PassParameters->DestinationResolution = FVector2f(DownsampledSizeX, DownsampledSizeY);
 
-	PassParameters->DstPixelCoordMinAndMax = FIntVector4(OutputViewport.Rect.Min.X, OutputViewport.Rect.Min.Y, OutputViewport.Rect.Max.X, OutputViewport.Rect.Max.Y);
-	PassParameters->SrcPixelCoordMinAndMax = FIntVector4( InputViewport.Rect.Min.X,  InputViewport.Rect.Min.Y,  InputViewport.Rect.Max.X,  InputViewport.Rect.Max.Y);
+	PassParameters->DstPixelCoordMinAndMax = FIntVector4(OutputViewport.Rect.Min.X, OutputViewport.Rect.Min.Y, OutputViewport.Rect.Max.X-1, OutputViewport.Rect.Max.Y-1);
+	PassParameters->SrcPixelCoordMinAndMax = FIntVector4( InputViewport.Rect.Min.X,  InputViewport.Rect.Min.Y,  InputViewport.Rect.Max.X-1,  InputViewport.Rect.Max.Y-1);
 
 	FRHIDepthStencilState* DepthStencilState = TStaticDepthStencilState<true, CF_Always>::GetRHI();
 
