Fix shader compilation issue with Strata & Lumen

diff --git a/Engine/Shaders/Private/DiffuseIndirectComposite.usf b/Engine/Shaders/Private/DiffuseIndirectComposite.usf
--- a/Engine/Shaders/Private/DiffuseIndirectComposite.usf
+++ b/Engine/Shaders/Private/DiffuseIndirectComposite.usf
@@ -261,7 +261,7 @@ FStrataDeferredLighting StrataIndirectLighting(
 		const float  Roughness = StrataGetBSDFRoughness(BSDF);
 		const FShadingOcclusion Occlusion = GetShadingOcclusion(PixelPos, V, StrataBSDFContext.N, Roughness, DiffuseColor, InAO);
 
-		const float3 InGlossyLighting = MixSpecularAndRoughReflections(Roughness, InSpecularLighting, InRoughSpecularLighting);
+		const float3 InGlossyLighting = MixSpecularAndRoughReflections(Roughness, false/*bIsFoliage*/, InSpecularLighting, InRoughSpecularLighting);
 
 		FDirectLighting OutBSDF = (FDirectLighting)0;
 		OutBSDF.Diffuse  += BSDFThroughput * InDiffuseLighting * StrataEnvLight.DiffuseWeight * Occlusion.DiffuseOcclusion;
@@ -270,7 +270,7 @@ FStrataDeferredLighting StrataIndirectLighting(
 		#if STRATA_FASTPATH==0
 		if (any(StrataEnvLight.SpecularHazeWeight > 0.0f))
 		{
-			const float3 InHazeGlossyLighting = MixSpecularAndRoughReflections(StrataEnvLight.SpecularHazeSafeRoughness, InSpecularLighting, InRoughSpecularLighting);
+			const float3 InHazeGlossyLighting = MixSpecularAndRoughReflections(StrataEnvLight.SpecularHazeSafeRoughness, false/*bIsFoliage*/, InSpecularLighting, InRoughSpecularLighting);
 			OutBSDF.Specular += BSDFThroughput * StrataEnvLight.SpecularHazeWeight * InHazeGlossyLighting * Occlusion.SpecularOcclusion;
 		}
 		#endif 
