Remove bad qualifier on a finction output.

diff --git a/Engine/Shaders/Private/VolumetricCloud.usf b/Engine/Shaders/Private/VolumetricCloud.usf
--- a/Engine/Shaders/Private/VolumetricCloud.usf
+++ b/Engine/Shaders/Private/VolumetricCloud.usf
@@ -1651,7 +1651,7 @@ float GetSliceDepth()
 	return RenderVolumetricCloudParameters.EmptySpaceSkippingSliceDepth;
 }
 
-const float GetSvPositionEmptySpaceSkippingValue(in ViewState ResolvedView, in uint3 DispatchThreadId, in uint3 GroupThreadId, in float3 SampleFroxelOffset)
+float GetSvPositionEmptySpaceSkippingValue(in ViewState ResolvedView, in uint3 DispatchThreadId, in uint3 GroupThreadId, in float3 SampleFroxelOffset)
 {
 #if MATERIAL_VOLUMETRIC_CLOUD_EMPTY_SPACE_SKIPPING_OUTPUT
 	float2 PixelUV = float2(DispatchThreadId.x + SampleFroxelOffset.x, DispatchThreadId.y + SampleFroxelOffset.y) / StartTracingDistanceTextureResolution;
