Volumetric cloud added debug visualisation of cloud tiles.

diff --git a/Engine/Shaders/Private/VolumetricCloud.usf b/Engine/Shaders/Private/VolumetricCloud.usf
--- a/Engine/Shaders/Private/VolumetricCloud.usf
+++ b/Engine/Shaders/Private/VolumetricCloud.usf
@@ -420,6 +420,7 @@ ParticipatingMediaPhaseContext SetupParticipatingMediaPhaseContext(float BasePha
 #if SHADER_RENDERVIEW_CS_WITH_EMPTY_SPACE_SKIPPING
 
 float2	StartTracingDistanceTextureResolution;
+float StartTracingSampleVolumeDepth;
 Texture2D<float> StartTracingDistanceTexture;
 
 #endif // SHADER_RENDERVIEW_CS_WITH_EMPTY_SPACE_SKIPPING
@@ -761,9 +762,10 @@ void MainCommon(in FMaterialPixelParameters MaterialParameters, in float4 SvPosi
 		// Read from empty space skipping
 		//
 		uint StartStepCount = 0;
-	#if SHADER_RENDERVIEW_CS_WITH_EMPTY_SPACE_SKIPPING && !CLOUD_DEBUG_VIEW_MODE
+	#if SHADER_RENDERVIEW_CS_WITH_EMPTY_SPACE_SKIPPING
 		const float StartDepthKm = StartTracingDistanceTexture[SvPositionToViewportUV(SvPosition) * StartTracingDistanceTextureResolution];
 		const float StartDepthCm = StartDepthKm * KILOMETER_TO_CENTIMETER;
+	#if !CLOUD_DEBUG_VIEW_MODE
 		if (t < StartDepthCm)
 		{
 			StartStepCount = uint(StartDepthCm / StepT);
@@ -771,6 +773,7 @@ void MainCommon(in FMaterialPixelParameters MaterialParameters, in float4 SvPosi
 		}
 		//OutColor0 = float4(StartDepthKm / 50.0f, 0.0f, 0.0f, 0.5f);
 		//return;
+	#endif
 	#endif
 
 		float3 S0 = 0.0f;
@@ -1362,6 +1365,12 @@ void MainCommon(in FMaterialPixelParameters MaterialParameters, in float4 SvPosi
 		if (RenderVolumetricCloudParameters.CloudDebugViewMode >= 1)
 		{
 			OutColor0.r = 1.0 - saturate(DebugZeroConservativeDensitySampleCount / float(IStepCount));
+		#if SHADER_RENDERVIEW_CS_WITH_EMPTY_SPACE_SKIPPING
+			if (RenderVolumetricCloudParameters.CloudDebugViewMode == 2)
+			{
+				OutColor0.g = saturate(StartDepthCm / StartTracingSampleVolumeDepth);
+			}
+		#endif
 		}
 #endif
 	}
@@ -1508,7 +1517,14 @@ void MainCommon(in FMaterialPixelParameters MaterialParameters, in float4 SvPosi
 
 	float GrayScaleTransmittance = MeanTransmittance < RenderVolumetricCloudParameters.StopTracingTransmittanceThreshold ? 0.0f : MeanTransmittance;
 #if CLOUD_DEBUG_VIEW_MODE
-	OutColor0 = float4(OutColor0.r * OutputPreExposure, Luminance.gb * OutputPreExposure, GrayScaleTransmittance);
+	if (RenderVolumetricCloudParameters.CloudDebugViewMode == 1)
+	{
+		OutColor0 = float4(OutColor0.r * OutputPreExposure, Luminance.gb * OutputPreExposure, GrayScaleTransmittance);
+	}
+	else if (RenderVolumetricCloudParameters.CloudDebugViewMode == 2)
+	{
+		OutColor0 = float4(OutColor0.rg * OutputPreExposure, Luminance.b * OutputPreExposure, GrayScaleTransmittance);
+	}
 #else
 	OutColor0 = float4(Luminance * OutputPreExposure, GrayScaleTransmittance);
 #endif
diff --git a/Engine/Source/Runtime/Renderer/Private/VolumetricCloudRendering.cpp b/Engine/Source/Runtime/Renderer/Private/VolumetricCloudRendering.cpp
--- a/Engine/Source/Runtime/Renderer/Private/VolumetricCloudRendering.cpp
+++ b/Engine/Source/Runtime/Renderer/Private/VolumetricCloudRendering.cpp
@@ -856,6 +856,7 @@ class FRenderVolumetricCloudRenderViewCS : public FMeshMaterialShader
 	BEGIN_SHADER_PARAMETER_STRUCT(FParameters, )
 		SHADER_PARAMETER(FVector4f, OutputViewRect)
 		SHADER_PARAMETER(FVector2f, StartTracingDistanceTextureResolution)
+		SHADER_PARAMETER(float, StartTracingSampleVolumeDepth)
 		SHADER_PARAMETER(int32, bBlendCloudColor)
 		SHADER_PARAMETER(int32, TargetCubeFace)
 		SHADER_PARAMETER_RDG_UNIFORM_BUFFER(FRenderVolumetricCloudGlobalParameters, VolumetricCloudRenderViewParamsUB)
@@ -2433,6 +2434,7 @@ void FSceneRenderer::RenderVolumetricCloudsInternal(FRDGBuilder& GraphBuilder, F
 
 		PassParameters->StartTracingDistanceTextureResolution = FVector2f(EmptySpaceSkippingTexture->Desc.Extent.X, EmptySpaceSkippingTexture->Desc.Extent.Y);
 		PassParameters->StartTracingDistanceTexture = EmptySpaceSkippingTexture;
+		PassParameters->StartTracingSampleVolumeDepth = GetEmptySpaceSkippingVolumeDepth();
 
 #if CLOUD_DEBUG_SAMPLES
 		ShaderPrint::SetEnabled(true);
