Subtract emissive instead of division for exposure illuminance approximation.

diff --git a/Engine/Shaders/Private/PostProcessEyeAdaptation.usf b/Engine/Shaders/Private/PostProcessEyeAdaptation.usf
--- a/Engine/Shaders/Private/PostProcessEyeAdaptation.usf
+++ b/Engine/Shaders/Private/PostProcessEyeAdaptation.usf
@@ -283,36 +283,17 @@ void SetupExposureIlluminance(uint2 ThreadId : SV_DispatchThreadID)
 {
 	const uint2 TexelPos = ThreadId + View.ViewRectMinAndSize.xy;
 
-	float3 Denominator = 0.0f;
+	float3 Emissive = 0.0f;
 
 	const float2 GBufferUV = (TexelPos + 0.5f) * View.BufferSizeAndInvSize.zw;
 	const FGBufferData GBufferData = GetGBufferDataFromSceneTextures(GBufferUV);
 
 	if (GBufferData.ShadingModelID != SHADINGMODELID_UNLIT)
 	{
-		Denominator += GBufferData.DiffuseColor;
-
-		// Emissive
-		Denominator += ColorTexture.Load(int3(TexelPos, 0)).rgb * View.OneOverPreExposure;
-
-		if (GBufferData.ShadingModelID == SHADINGMODELID_SUBSURFACE
-			|| GBufferData.ShadingModelID == SHADINGMODELID_PREINTEGRATED_SKIN
-			|| GBufferData.ShadingModelID == SHADINGMODELID_TWOSIDED_FOLIAGE
-			|| GBufferData.ShadingModelID == SHADINGMODELID_CLOTH
-			|| GBufferData.ShadingModelID == SHADINGMODELID_EYE)
-		{
-			float3 SubsurfaceColor = ExtractSubsurfaceColor(GBufferData);
-
-			if (GBufferData.ShadingModelID == SHADINGMODELID_TWOSIDED_FOLIAGE)
-			{
-				SubsurfaceColor *= 1.0f / PI;
-			}
-
-			Denominator += SubsurfaceColor;
-		}
+		Emissive = ColorTexture.Load(int3(TexelPos, 0)).rgb;
 	}
 
-	RWIlluminanceTexture[TexelPos] = CalculateEyeAdaptationLuminanceWithoutMin(Denominator) * View.PreExposure;
+	RWIlluminanceTexture[TexelPos] = CalculateEyeAdaptationLuminanceWithoutMin(Emissive);
 }
 
 [numthreads(8, 8, 1)]
@@ -329,7 +310,27 @@ void CalculateExposureIlluminance(uint2 ThreadId : SV_DispatchThreadID)
 
 	if (GBufferData.ShadingModelID != SHADINGMODELID_UNLIT)
 	{
-		float Denominator = RWIlluminanceTexture[TexelPos] * View.OneOverPreExposure;
+		const float Emissive = RWIlluminanceTexture[TexelPos] * View.OneOverPreExposure;
+
+		float3 Denominator = 0.0f;
+
+		Denominator += GBufferData.DiffuseColor;
+
+		if (GBufferData.ShadingModelID == SHADINGMODELID_SUBSURFACE
+			|| GBufferData.ShadingModelID == SHADINGMODELID_PREINTEGRATED_SKIN
+			|| GBufferData.ShadingModelID == SHADINGMODELID_TWOSIDED_FOLIAGE
+			|| GBufferData.ShadingModelID == SHADINGMODELID_CLOTH
+			|| GBufferData.ShadingModelID == SHADINGMODELID_EYE)
+		{
+			float3 SubsurfaceColor = ExtractSubsurfaceColor(GBufferData);
+
+			if (GBufferData.ShadingModelID == SHADINGMODELID_TWOSIDED_FOLIAGE)
+			{
+				SubsurfaceColor *= 1.0f / PI;
+			}
+
+			Denominator += SubsurfaceColor;
+		}
 
 		{
 			const float3 TranslatedWorldPosition = ReconstructTranslatedWorldPositionFromDepth(GBufferUV, GBufferData.Depth);
@@ -339,10 +340,13 @@ void CalculateExposureIlluminance(uint2 ThreadId : SV_DispatchThreadID)
 			const float3 V = CameraVector;
 			const float3 EnvBrdf = EnvBRDF(GBufferData.SpecularColor, GBufferData.Roughness, max(0.0, dot(N, V)));
 
-			Denominator += CalculateEyeAdaptationLuminanceWithoutMin(EnvBrdf);
+			Denominator += EnvBrdf;
 		}
 
-		Illuminance /= max(Denominator, EyeAdaptation_IgnoreMaterialsMinBaseColorLuminance);
+		float DenominatorLuminance = CalculateEyeAdaptationLuminanceWithoutMin(Denominator);
+
+		Illuminance = max(0.0f, Illuminance - Emissive);
+		Illuminance /= max(DenominatorLuminance, EyeAdaptation_IgnoreMaterialsMinBaseColorLuminance);
 		Illuminance *= lerp(EyeAdaptation_IgnoreMaterialsLuminanceScale, 1.0f, EyeAdaptation_IgnoreMaterialsMinBaseColorLuminance);
 	}
 	
