Path Tracer: Fix rendering of transparent surfaces viewed over a visible skylight

It was possible for a NaN to be introduced in the throughput which prevented the stochastic hit selection from working properly. This could happen when SigmaT = 0 and HitT=inf which was possible for skylight hits.

Also move the order of operation somewhat to be more correct.

diff --git a/Engine/Shaders/Private/PathTracing/PathTracingCore.ush b/Engine/Shaders/Private/PathTracing/PathTracingCore.ush
--- a/Engine/Shaders/Private/PathTracing/PathTracingCore.ush
+++ b/Engine/Shaders/Private/PathTracing/PathTracingCore.ush
@@ -386,6 +386,9 @@ FPathTracingPayload TraceTransparentRay(inout FPathState PathState, bool bIsCame
 			}
 		}
 #endif
+		// account for Beer's law through the currently active medium
+		// TODO: merge this with the volume raymarching?
+		PathState.PathThroughput *= select(PathState.SigmaT > 0.0, exp(-PathState.SigmaT * (HitPayload.HitT - PathState.Ray.TMin)), 1.0);
 
 		if (HitPayload.BlendingMode == RAY_TRACING_BLEND_MODE_ADDITIVE ||
 			HitPayload.BlendingMode == RAY_TRACING_BLEND_MODE_TRANSLUCENT)
@@ -405,11 +408,11 @@ FPathTracingPayload TraceTransparentRay(inout FPathState PathState, bool bIsCame
 			AccumulateRadiance(PathState.Radiance, PathState.PathThroughput * HitPayload.Radiance, bIsCameraRay);
 		}
 
-		float3 Transparency = GetMaterialTransparency(HitPayload, PathState.Ray.Direction);
-
-		// account for Beer's law through the currently active medium
-		// TODO: merge this with the volume raymarching?
-		PathState.PathThroughput *= exp(-PathState.SigmaT * (HitPayload.HitT - PathState.Ray.TMin));
+		if (HitPayload.HitT == POSITIVE_INFINITY)
+		{
+			// if our hit was against an infinite light, exit now
+			break;
+		}
 
 		if (!bLastBounce)
 		{
@@ -427,14 +430,8 @@ FPathTracingPayload TraceTransparentRay(inout FPathState PathState, bool bIsCame
 			}
 		}
 
-		if (HitPayload.HitT == POSITIVE_INFINITY)
-		{
-			// if our hit was against an infinite light, exit now
-			break;
-		}
-
 		// account for local transparency change
-		PathState.PathThroughput *= Transparency;
+		PathState.PathThroughput *= GetMaterialTransparency(HitPayload, PathState.Ray.Direction);;
 
 		// prepare next step around the loop
 		// retrace the exact same ray with TMin one ulp past the hit we just found
