When `Rough Diffuse Material` is enabled in the project setting, chan's diffuse model will be used by subsurface profile shading model. This will correct the shading at the edge of human skin, better matching the ground truth than Burley's diffuse model.

diff --git a/Engine/Shaders/Private/ShadingModels.ush b/Engine/Shaders/Private/ShadingModels.ush
--- a/Engine/Shaders/Private/ShadingModels.ush
+++ b/Engine/Shaders/Private/ShadingModels.ush
@@ -694,7 +694,15 @@ FDirectLighting SubsurfaceProfileBxDF( FGBufferData GBuffer, half3 N, half3 V, h
 	// Needs to apply shadow at here, since the preintegrated value has already counted it, and we need to skip applying shadow outside.
 	Lighting.Diffuse = lerp(AreaLight.FalloffColor * (Falloff * NoL) * Shadow.SurfaceShadow, BurleyDiffuse.rgb, Tint) * Diffuse_Lambert(GBuffer.DiffuseColor);
 #else
-	Lighting.Diffuse  = AreaLight.FalloffColor * (Falloff * NoL) * Diffuse_Burley( GBuffer.DiffuseColor, GBuffer.Roughness, Context.NoV, NoL, Context.VoH );
+
+#if MATERIAL_ROUGHDIFFUSE
+	// Use Chan's diffuse model. It reduces flatness look and has better match, e.g., at the edge of human face skin when compared to GT.
+	const float3 DiffuseReflection = Diffuse_Chan(GBuffer.DiffuseColor, Pow4(GBuffer.Roughness), Context.NoV, NoL, Context.VoH, Context.NoH, GetAreaLightDiffuseMicroReflWeight(AreaLight));
+#else
+	const float3 DiffuseReflection = Diffuse_Burley(GBuffer.DiffuseColor, GBuffer.Roughness, Context.NoV, NoL, Context.VoH);
+#endif
+
+	Lighting.Diffuse  = AreaLight.FalloffColor * (Falloff * NoL) * DiffuseReflection;
 #endif
 
 	if (IsRectLight(AreaLight))
