Implemented support for decoding single sample depth Z planes

diff --git a/Engine/Shaders/Private/HTileEncoding.ush b/Engine/Shaders/Private/HTileEncoding.ush
--- a/Engine/Shaders/Private/HTileEncoding.ush
+++ b/Engine/Shaders/Private/HTileEncoding.ush
@@ -195,3 +195,274 @@ uint2 ExtractTileHiStencil(uint HTileValue, bool HiStencil)
 
 	return uint2(HiStencil0, HiStencil1);
 }
+
+uint GetZPlaneCount(float ZMask)
+{
+	uint PlaneCount = ZMask;
+
+	if (PlaneCount >= 11u)
+	{
+		PlaneCount += 2u;
+	}
+	else if (PlaneCount >= 9u)
+	{
+		++PlaneCount;
+	}
+
+	return PlaneCount;
+}
+
+// Note: We don't use MSAA depth, this assumes single sample count
+uint PackZLayout(uint ZMask)
+{
+	const uint ZPlaneCount			= GetZPlaneCount(ZMask);
+	const uint ZPlaneSize			= ZPlaneCount * 12u;
+	const uint ZPlanePaddedSize		= Padding(ZPlaneSize, 32u);
+	const uint PMaskBitCount		= CeilLog2(ZPlaneCount);
+	const uint PMaskSize			= 64u * PMaskBitCount / 8u;
+	const uint PMaskPaddedSize		= Padding(PMaskSize, 32u);
+	const uint DepthPaddedSize		= ZPlanePaddedSize + PMaskPaddedSize;
+	const uint DepthPackedSize		= Padding(ZPlaneSize + PMaskSize, 32u);
+	const uint PMaskOffset			= Padding(select(DepthPackedSize < DepthPaddedSize, ZPlaneSize, ZPlanePaddedSize), select(DepthPackedSize > 256u, 256u, 8u));
+	const uint ZPlaneWordCount		= (PMaskOffset + PMaskSize) >> 2u;
+
+	uint ZLayout = ZPlaneCount;
+	ZLayout |= PMaskBitCount;
+	ZLayout |= select(DepthPackedSize <= 32, 1u, 0u) << 16;
+	ZLayout |= (PMaskOffset >> 2u) << 17u;
+	ZLayout |= ZPlaneWordCount << 24u;
+	return ZLayout;
+}
+
+struct ZLayout
+{
+	uint ZPlaneCount;
+	uint ZPlaneWordCount;
+	uint PMaskOffset;
+	uint IndexBitCount;
+	uint Is32BytePadded;
+};
+
+ZLayout UnpackZLayout(uint ZMask)
+{
+	const uint ZLayoutTable[16u] =
+	{
+		PackZLayout(0x0),
+		PackZLayout(0x1),
+		PackZLayout(0x2),
+		PackZLayout(0x3),
+		PackZLayout(0x4),
+		PackZLayout(0x5),
+		PackZLayout(0x6),
+		PackZLayout(0x7),
+		PackZLayout(0x8),
+		PackZLayout(0x9),
+		PackZLayout(0xA),
+		PackZLayout(0xB),
+		PackZLayout(0xC),
+		PackZLayout(0xD),
+		PackZLayout(0xE),
+		PackZLayout(0xF),
+	};
+
+	const uint PackedZLayout = ZLayoutTable[ZMask];
+
+	ZLayout Layout;
+	Layout.ZPlaneCount		= BitFieldExtractU32(PackedZLayout, 5u, 0u);
+	Layout.IndexBitCount	= BitFieldExtractU32(PackedZLayout, 3u, 5u);
+	Layout.Is32BytePadded	= BitFieldExtractU32(PackedZLayout, 4u, 13u);
+	Layout.PMaskOffset		= BitFieldExtractU32(PackedZLayout, 7u, 17u);
+	Layout.ZPlaneWordCount	= BitFieldExtractU32(PackedZLayout, 8u, 24u);
+	return Layout;
+}
+
+// 112 dword max ever read, but can be offset by 8 words
+groupshared uint CompressedDepthData[128];
+
+float DecompressDepthValue_Internal(ZLayout Layout, uint ZOffset, uint ZMask, uint2 PixelPos, uint ThreadIndex, float ClearValue, uint PlatformConfig)
+{
+	GroupMemoryBarrierWithGroupSync();
+
+	uint ZPlaneIndex = 0;
+
+	BRANCH
+	if (Layout.ZPlaneCount > 1) // scalarized
+	{
+		const uint BitOffset = MulU24(ThreadIndex, Layout.IndexBitCount);
+		const uint IndexOffset = ZOffset + Layout.PMaskOffset + (BitOffset >> 5);
+		ZPlaneIndex = BitAlignU32(CompressedDepthData[IndexOffset + 1], CompressedDepthData[IndexOffset], BitOffset);
+		ZPlaneIndex = BitFieldExtractU32(ZPlaneIndex, Layout.IndexBitCount, 0);
+	}
+
+	uint  ZPlaneOffset = (ZPlaneIndex * 3) + ZOffset;
+	uint3 ZPlane = uint3(
+		CompressedDepthData[ZPlaneOffset + 0u],
+		CompressedDepthData[ZPlaneOffset + 1u],
+		CompressedDepthData[ZPlaneOffset + 2u]
+	);
+
+	int Exponent			= BitFieldExtractU32(ZPlane.y, 8u, 24u);
+	int ZCenter				= BitFieldExtractI32(ZPlane.z, 31u, 0u) << 1;
+
+	int2 AdjustedPixelPos	= PixelPos & 7;
+	AdjustedPixelPos		= AdjustedPixelPos * 2 - 7;
+
+	int DepthValue = 0;
+	DepthValue = MadI24(ZPlane.x >> 4,  AdjustedPixelPos.x, DepthValue     );
+	DepthValue = MadI24(ZPlane.y,       AdjustedPixelPos.y, DepthValue     );
+	DepthValue = MadI24(ZPlane.x & 0xF, AdjustedPixelPos.x, DepthValue << 4);
+	DepthValue = MadI24(ZPlane.x >> 28, AdjustedPixelPos.y, DepthValue     );
+
+	if (DepthValue <= -ZCenter)
+	{
+		return 0.0f;
+	}
+
+	DepthValue += ZCenter;
+
+	int Shift = firstbithigh(DepthValue) - 23;
+	Exponent += Shift;
+
+	const bool bHandleDenorms = true;
+	const bool bFlushDenorms = true;
+
+	if (bHandleDenorms && Exponent < 1)
+	{
+		return select(bFlushDenorms, 0.0f, asfloat((DepthValue << (8 - Shift)) >> (9 - Exponent)));
+	}
+	else
+	{
+		return min(1.0f, asfloat(Exponent << 23 | (DepthValue << (9 - Shift)) >> 9));
+	}
+}
+
+float DecompressDepthValue(Texture2D<float> DepthBuffer, uint ZMask, uint2 PixelPos, uint ThreadIndex, float ClearValue, uint PlatformConfig)
+{
+	if (IsZTileClear(ZMask))
+	{
+		return ClearValue;
+	}
+
+	if (ZMask == HTILE_ZMASK_EXPANDED)
+	{
+		// Uncompressed, fetch raw depth
+		return DepthBuffer[PixelPos];
+	}
+
+	const ZLayout Layout = UnpackZLayout(ZMask);
+	const uint ZOffset = select(Is32BytePaddedTileConfig(PlatformConfig), Layout.Is32BytePadded & PixelPos.x, 0u);
+	const uint MaxThreadIndex = ZOffset + Layout.ZPlaneWordCount;
+
+	if (ThreadIndex < MaxThreadIndex)
+	{
+		CompressedDepthData[ThreadIndex] = asuint(DepthBuffer[PixelPos]);
+	}
+
+	return DecompressDepthValue_Internal(Layout, ZOffset, ZMask, PixelPos, ThreadIndex, ClearValue, PlatformConfig);
+}
+
+float DecompressDepthValue(RWTexture2D<float> DepthBuffer, uint ZMask, uint2 PixelPos, uint ThreadIndex, float ClearValue, uint PlatformConfig)
+{
+	if (IsZTileClear(ZMask))
+	{
+		return ClearValue;
+	}
+
+	if (ZMask == HTILE_ZMASK_EXPANDED)
+	{
+		// Uncompressed, fetch raw depth
+		return DepthBuffer[PixelPos];
+	}
+
+	const ZLayout Layout = UnpackZLayout(ZMask);
+	const uint ZOffset = select(Is32BytePaddedTileConfig(PlatformConfig), Layout.Is32BytePadded & PixelPos.x, 0u);
+	const uint MaxThreadIndex = ZOffset + Layout.ZPlaneWordCount;
+
+	if (ThreadIndex < MaxThreadIndex)
+	{
+		CompressedDepthData[ThreadIndex] = asuint(DepthBuffer[PixelPos]);
+	}
+
+	return DecompressDepthValue_Internal(Layout, ZOffset, ZMask, PixelPos, ThreadIndex, ClearValue, PlatformConfig);
+}
+
+uint2 PixelFromThreadIndex(uint Element)
+{
+	uint2 PixelPos;
+	PixelPos.x = ((Element >> 0u) & 1u) | ((Element >> 1u) & 2u) | ((Element >> 2u) & 4u);
+	PixelPos.y = ((Element >> 1u) & 1u) | ((Element >> 2u) & 2u) | ((Element >> 3u) & 4u);
+	return PixelPos;
+}
+
+uint2 SwizzleThreadIndex(uint ThreadIndex)
+{
+	const uint2 LookupTable[64] =
+	{
+		PixelFromThreadIndex(8 * 0 + 0),
+		PixelFromThreadIndex(8 * 0 + 1),
+		PixelFromThreadIndex(8 * 0 + 2),
+		PixelFromThreadIndex(8 * 0 + 3),
+		PixelFromThreadIndex(8 * 0 + 4),
+		PixelFromThreadIndex(8 * 0 + 5),
+		PixelFromThreadIndex(8 * 0 + 6),
+		PixelFromThreadIndex(8 * 0 + 7),
+		PixelFromThreadIndex(8 * 1 + 0),
+		PixelFromThreadIndex(8 * 1 + 1),
+		PixelFromThreadIndex(8 * 1 + 2),
+		PixelFromThreadIndex(8 * 1 + 3),
+		PixelFromThreadIndex(8 * 1 + 4),
+		PixelFromThreadIndex(8 * 1 + 5),
+		PixelFromThreadIndex(8 * 1 + 6),
+		PixelFromThreadIndex(8 * 1 + 7),
+		PixelFromThreadIndex(8 * 2 + 0),
+		PixelFromThreadIndex(8 * 2 + 1),
+		PixelFromThreadIndex(8 * 2 + 2),
+		PixelFromThreadIndex(8 * 2 + 3),
+		PixelFromThreadIndex(8 * 2 + 4),
+		PixelFromThreadIndex(8 * 2 + 5),
+		PixelFromThreadIndex(8 * 2 + 6),
+		PixelFromThreadIndex(8 * 2 + 7),
+		PixelFromThreadIndex(8 * 3 + 0),
+		PixelFromThreadIndex(8 * 3 + 1),
+		PixelFromThreadIndex(8 * 3 + 2),
+		PixelFromThreadIndex(8 * 3 + 3),
+		PixelFromThreadIndex(8 * 3 + 4),
+		PixelFromThreadIndex(8 * 3 + 5),
+		PixelFromThreadIndex(8 * 3 + 6),
+		PixelFromThreadIndex(8 * 3 + 7),
+		PixelFromThreadIndex(8 * 4 + 0),
+		PixelFromThreadIndex(8 * 4 + 1),
+		PixelFromThreadIndex(8 * 4 + 2),
+		PixelFromThreadIndex(8 * 4 + 3),
+		PixelFromThreadIndex(8 * 4 + 4),
+		PixelFromThreadIndex(8 * 4 + 5),
+		PixelFromThreadIndex(8 * 4 + 6),
+		PixelFromThreadIndex(8 * 4 + 7),
+		PixelFromThreadIndex(8 * 5 + 0),
+		PixelFromThreadIndex(8 * 5 + 1),
+		PixelFromThreadIndex(8 * 5 + 2),
+		PixelFromThreadIndex(8 * 5 + 3),
+		PixelFromThreadIndex(8 * 5 + 4),
+		PixelFromThreadIndex(8 * 5 + 5),
+		PixelFromThreadIndex(8 * 5 + 6),
+		PixelFromThreadIndex(8 * 5 + 7),
+		PixelFromThreadIndex(8 * 6 + 0),
+		PixelFromThreadIndex(8 * 6 + 1),
+		PixelFromThreadIndex(8 * 6 + 2),
+		PixelFromThreadIndex(8 * 6 + 3),
+		PixelFromThreadIndex(8 * 6 + 4),
+		PixelFromThreadIndex(8 * 6 + 5),
+		PixelFromThreadIndex(8 * 6 + 6),
+		PixelFromThreadIndex(8 * 6 + 7),
+		PixelFromThreadIndex(8 * 7 + 0),
+		PixelFromThreadIndex(8 * 7 + 1),
+		PixelFromThreadIndex(8 * 7 + 2),
+		PixelFromThreadIndex(8 * 7 + 3),
+		PixelFromThreadIndex(8 * 7 + 4),
+		PixelFromThreadIndex(8 * 7 + 5),
+		PixelFromThreadIndex(8 * 7 + 6),
+		PixelFromThreadIndex(8 * 7 + 7),
+	};
+
+	return LookupTable[ThreadIndex];
+}
\ No newline at end of file
