Wrap logic in HTileEncoding.ush with PLATFORM_SUPPORTS_HTILE_LOOKUP

diff --git a/Engine/Shaders/Private/HTileEncoding.ush b/Engine/Shaders/Private/HTileEncoding.ush
--- a/Engine/Shaders/Private/HTileEncoding.ush
+++ b/Engine/Shaders/Private/HTileEncoding.ush
@@ -4,6 +4,12 @@
 
 #include "Common.ush"
 
+#ifndef PLATFORM_SUPPORTS_HTILE_LOOKUP
+#define PLATFORM_SUPPORTS_HTILE_LOOKUP 0
+#endif
+
+#if PLATFORM_SUPPORTS_HTILE_LOOKUP
+
 #ifdef OVERRIDE_HTILELOOKUP_USH
 
 	#include "/Platform/Private/HTileLookup.ush"
@@ -17,10 +23,6 @@
 
 #endif
 
-#ifndef PLATFORM_SUPPORTS_HTILE_LOOKUP
-#define PLATFORM_SUPPORTS_HTILE_LOOKUP 0
-#endif
-
 // Optimized version for no hi-stencil, just min and max depth.
 uint EncodeTileMinMaxDepth(float MinDepth, float MaxDepth)
 {
@@ -465,4 +467,6 @@ uint2 SwizzleThreadIndex(uint ThreadIndex)
 	};
 
 	return LookupTable[ThreadIndex];
-}
\ No newline at end of file
+}
+
+#endif // PLATFORM_SUPPORTS_HTILE_LOOKUP
\ No newline at end of file
