Add PrintBool to ShaderPrint.

diff --git a/Engine/Shaders/Private/HairStrands/HairStrandsDebug.usf b/Engine/Shaders/Private/HairStrands/HairStrandsDebug.usf
--- a/Engine/Shaders/Private/HairStrands/HairStrandsDebug.usf
+++ b/Engine/Shaders/Private/HairStrands/HairStrandsDebug.usf
@@ -670,18 +670,6 @@ void VisualizeDomPS(
 #include "HairStrandsDeepShadowCommonStruct.ush"
 #include "HairStrandsAABBCommon.ush"
 
-void PrintBool(inout FShaderPrintContext Ctx, bool In)
-{
-	if (In)
-	{
-		Print(Ctx, TEXT("Yes  "), FontGreen);
-	}
-	else
-	{
-		Print(Ctx, TEXT("No   "), FontRed);
-	}
-}
-
 uint AllocatedSlotCount;
 uint MacroGroupCount;
 uint bViewRectOptimizeEnabled;
diff --git a/Engine/Shaders/Private/HairStrands/HairStrandsVisibilityRasterCompute.usf b/Engine/Shaders/Private/HairStrands/HairStrandsVisibilityRasterCompute.usf
--- a/Engine/Shaders/Private/HairStrands/HairStrandsVisibilityRasterCompute.usf
+++ b/Engine/Shaders/Private/HairStrands/HairStrandsVisibilityRasterCompute.usf
@@ -1478,17 +1478,6 @@ void RasterMultiSampleCS(uint DispatchThreadID : SV_DispatchThreadID, uint Group
 #if SHADER_RASTERCOMPUTE_DEBUG
 
 #include "../ShaderPrint.ush"
-void PrintBool(inout FShaderPrintContext Ctx, bool In)
-{
-	if (In)
-	{
-		Print(Ctx, TEXT("Yes"), FontGreen);
-	}
-	else
-	{
-		Print(Ctx, TEXT("No"), FontRed);
-	}
-}
 
 Texture2D<uint>						VisTileDepthGrid;
 Texture2DArray<uint> 				VisTileBinningGrid;
diff --git a/Engine/Shaders/Private/ShaderPrintCommon.ush b/Engine/Shaders/Private/ShaderPrintCommon.ush
--- a/Engine/Shaders/Private/ShaderPrintCommon.ush
+++ b/Engine/Shaders/Private/ShaderPrintCommon.ush
@@ -502,6 +502,33 @@ SHADER_PRINT_OVERLOAD_NN(3, 3, float3x3)
 SHADER_PRINT_OVERLOAD_NN(4, 3, float4x3)
 SHADER_PRINT_OVERLOAD_NN(4, 4, float4x4)
 
+///////////////////////////////////////////////////////////////////////////////////////////////////
+// Bool printing
+
+void PrintBool(inout FShaderPrintContext Ctx, bool In)
+{
+	if (In)
+	{
+		Print(Ctx, TEXT("Yes "), FontGreen);
+	}
+	else
+	{
+		Print(Ctx, TEXT("No  "), FontRed);
+	}
+}
+
+void PrintBool(inout FShaderPrintContext Ctx, bool In, FFontColor InColor)
+{
+	if (In)
+	{
+		Print(Ctx, TEXT("Yes "), InColor);
+	}
+	else
+	{
+		Print(Ctx, TEXT("No  "), InColor);
+	}
+}
+
 ///////////////////////////////////////////////////////////////////////////////////////////////////
 // Widgets (experimental)
 
diff --git a/Engine/Shaders/Private/Strata/StrataVisualize.usf b/Engine/Shaders/Private/Strata/StrataVisualize.usf
--- a/Engine/Shaders/Private/Strata/StrataVisualize.usf
+++ b/Engine/Shaders/Private/Strata/StrataVisualize.usf
@@ -20,18 +20,6 @@
 // This is for the fxc compiler to compile faster.
 #define STRATA_VIZ_MAX_BSDF 2
 
-void PrintBool(inout FShaderPrintContext Ctx, bool In)
-{
-	if (In)
-	{
-		Print(Ctx, TEXT("Yes  "), FontGreen);
-	}
-	else
-	{
-		Print(Ctx, TEXT("No   "), FontRed);
-	}
-}
-
 void DrawReferentialTWS(float3 P /*In Translated World Space*/, float3 X, float3 Y, float3 N, float3 InColor)
 {
 	const float Size = 10.f;
