Remove duplicate of DivideAndRoundUp and move it to Common.ush

diff --git a/Engine/Plugins/FX/Niagara/Shaders/Private/Ribbons/NiagaraRibbonCommon.ush b/Engine/Plugins/FX/Niagara/Shaders/Private/Ribbons/NiagaraRibbonCommon.ush
--- a/Engine/Plugins/FX/Niagara/Shaders/Private/Ribbons/NiagaraRibbonCommon.ush
+++ b/Engine/Plugins/FX/Niagara/Shaders/Private/Ribbons/NiagaraRibbonCommon.ush
@@ -239,12 +239,6 @@ static const FNiagaraID InvalidID = { INT_MIN, INT_MIN };
 
 #endif
 
-
-uint DivideAndRoundUp(uint Dividend, uint Divisor)
-{
-	return (Dividend + Divisor - 1) / Divisor;		
-}
-
 // max absolute error 9.0x10^-3
 // Eberly's polynomial degree 1 - respect bounds
 // input [-1, 1] and output [0, PI]
diff --git a/Engine/Shaders/Private/Common.ush b/Engine/Shaders/Private/Common.ush
--- a/Engine/Shaders/Private/Common.ush
+++ b/Engine/Shaders/Private/Common.ush
@@ -1987,6 +1987,16 @@ uint ZOrder3DEncode(uint3 Coord, const uint SizeLog2)
     return Index;
 }
 
+uint DivideAndRoundUp(uint Dividend, uint Divisor)
+{
+	return (Dividend + Divisor - 1) / Divisor;
+}
+
+// Optimised for power of two because it relies on a division done using bit shift
+uint DivideAndRoundUp(uint Dividend, uint Divisor, uint DivisorAsBitShift)
+{
+	return (Dividend + Divisor - 1) >> DivisorAsBitShift;
+}
 
 // see PixelShaderOutputCommon
 struct FPixelShaderIn
diff --git a/Engine/Shaders/Private/HairStrands/HairStrandsClusterCulling.usf b/Engine/Shaders/Private/HairStrands/HairStrandsClusterCulling.usf
--- a/Engine/Shaders/Private/HairStrands/HairStrandsClusterCulling.usf
+++ b/Engine/Shaders/Private/HairStrands/HairStrandsClusterCulling.usf
@@ -27,11 +27,6 @@ uint ClusterCount;
 /////////////////////////////////////////////////////////////////////////////////////////////////////////////////////////////
 // Utility functions
 
-uint DivideAndRoundUp(uint Dividend, uint Divisor)
-{
-	return (Dividend + Divisor - 1) / Divisor;
-}
-
 uint3 ComputeDispatchGroupCount2D(uint GroupCount)
 {
 	const uint DispatchCountX = int(floor(sqrt(float(GroupCount))));
diff --git a/Engine/Shaders/Private/Strata/Strata.ush b/Engine/Shaders/Private/Strata/Strata.ush
--- a/Engine/Shaders/Private/Strata/Strata.ush
+++ b/Engine/Shaders/Private/Strata/Strata.ush
@@ -1637,12 +1637,6 @@ struct FStrataRaytracingPayload
 ///////////////////////////////////////////////////////////////////////////////
 // Material buffer Addressing
 
-// Optimised for power of two because it relies on a division done using bit shift
-uint DivideAndRoundUp(uint Dividend, uint Divisor, uint DivisorAsBitShift)
-{
-	return (Dividend + Divisor - 1) >> DivisorAsBitShift;
-}
-
 struct FStrataAddressing
 {
 	uint  CurrentIndex;
diff --git a/Engine/Shaders/Private/Strata/StrataMaterialClassification.usf b/Engine/Shaders/Private/Strata/StrataMaterialClassification.usf
--- a/Engine/Shaders/Private/Strata/StrataMaterialClassification.usf
+++ b/Engine/Shaders/Private/Strata/StrataMaterialClassification.usf
@@ -371,8 +371,6 @@ Buffer<uint>   TileDrawIndirectDataBuffer;
 RWBuffer<uint> TileDispatchIndirectDataBuffer;
 RWBuffer<uint> TileDispatchPerThreadIndirectDataBuffer;
 
-uint DivideAndRoundUp(uint X, uint Y) { return (X + Y - 1) / Y; }
-
 [numthreads(1, 1, 1)]
 void ArgsMainCS(uint2 DispatchThreadId : SV_DispatchThreadID)
 {
