Support translucency lighting type on mobile

diff --git a/Engine/Shaders/Private/BasePassCommon.ush b/Engine/Shaders/Private/BasePassCommon.ush
--- a/Engine/Shaders/Private/BasePassCommon.ush
+++ b/Engine/Shaders/Private/BasePassCommon.ush
@@ -293,4 +293,44 @@ FShadingOcclusion ApplyBentNormal(
 	Out.SpecOcclusion = saturate(Out.SpecOcclusion );
 #endif
 	return Out;
+}
+
+#include "SHCommon.ush"
+
+void GetVolumeLightingNonDirectional(float4 AmbientLightingVector, float3 DiffuseColor, inout float3 InterpolatedLighting, out float4 VolumeLighting)
+{
+	// Normal is not taken into account with non directional lighting, and only the ambient term of the SH coefficients are needed
+	FOneBandSHVectorRGB TranslucentLighting;
+	TranslucentLighting.R.V.x = AmbientLightingVector.r;
+	TranslucentLighting.G.V.x = AmbientLightingVector.g;
+	TranslucentLighting.B.V.x = AmbientLightingVector.b;
+
+	FOneBandSHVector DiffuseTransferSH = CalcDiffuseTransferSH1(1);
+	VolumeLighting = float4(DotSH1(TranslucentLighting, DiffuseTransferSH), AmbientLightingVector.a);
+	InterpolatedLighting = DiffuseColor * VolumeLighting.rgb;
+}
+
+void GetVolumeLightingDirectional(float4 AmbientLightingVector, float3 DirectionalLightingVector, float3 WorldNormal, float3 DiffuseColor, inout float3 InterpolatedLighting, out float4 VolumeLighting)
+{
+	float DirectionalLightingIntensity = GetMaterialTranslucencyDirectionalLightingIntensity();
+
+	AmbientLightingVector.rgb /= DirectionalLightingIntensity;
+	DirectionalLightingVector.rgb *= DirectionalLightingIntensity;
+
+	// Reconstruct the SH coefficients based on what was encoded
+	FTwoBandSHVectorRGB TranslucentLighting;
+	TranslucentLighting.R.V.x = AmbientLightingVector.r;
+	TranslucentLighting.G.V.x = AmbientLightingVector.g;
+	TranslucentLighting.B.V.x = AmbientLightingVector.b;
+	float3 NormalizedAmbientColor = AmbientLightingVector.rgb / (Luminance(AmbientLightingVector.rgb) + 0.00001f);
+
+	// Scale the monocrome directional coefficients with the normalzed ambient color as an approximation to the uncompressed values
+	TranslucentLighting.R.V.yzw = DirectionalLightingVector.rgb * NormalizedAmbientColor.r;
+	TranslucentLighting.G.V.yzw = DirectionalLightingVector.rgb * NormalizedAmbientColor.g;
+	TranslucentLighting.B.V.yzw = DirectionalLightingVector.rgb * NormalizedAmbientColor.b;
+
+	// Compute diffuse lighting which takes the normal into account
+	FTwoBandSHVector DiffuseTransferSH = CalcDiffuseTransferSH(WorldNormal, 1);
+	VolumeLighting = float4(max(half3(0, 0, 0), DotSH(TranslucentLighting, DiffuseTransferSH)), AmbientLightingVector.a);
+	InterpolatedLighting += DiffuseColor * VolumeLighting.rgb;
 }
\ No newline at end of file
diff --git a/Engine/Shaders/Private/BasePassPixelShader.usf b/Engine/Shaders/Private/BasePassPixelShader.usf
--- a/Engine/Shaders/Private/BasePassPixelShader.usf
+++ b/Engine/Shaders/Private/BasePassPixelShader.usf
@@ -190,44 +190,6 @@ bool ManualDepthTestEqual(float4 SvPosition, float PSDeviceZWithOffset)
 	return DepthTest_Equal;
 }
 
-void GetVolumeLightingNonDirectional(float4 AmbientLightingVector, float3 DiffuseColor, inout float3 InterpolatedLighting, out float4 VolumeLighting)
-{
-	// Normal is not taken into account with non directional lighting, and only the ambient term of the SH coefficients are needed
-	FOneBandSHVectorRGB TranslucentLighting;
-	TranslucentLighting.R.V.x = AmbientLightingVector.r;
-	TranslucentLighting.G.V.x = AmbientLightingVector.g;
-	TranslucentLighting.B.V.x = AmbientLightingVector.b;
-
-	FOneBandSHVector DiffuseTransferSH = CalcDiffuseTransferSH1(1);
-	VolumeLighting = float4(DotSH1(TranslucentLighting, DiffuseTransferSH), AmbientLightingVector.a);
-	InterpolatedLighting = DiffuseColor * VolumeLighting.rgb;
-}
-
-void GetVolumeLightingDirectional(float4 AmbientLightingVector, float3 DirectionalLightingVector, float3 WorldNormal, float3 DiffuseColor, inout float3 InterpolatedLighting, out float4 VolumeLighting)
-{
-	float DirectionalLightingIntensity = GetMaterialTranslucencyDirectionalLightingIntensity();
-
-	AmbientLightingVector.rgb /= DirectionalLightingIntensity;
-	DirectionalLightingVector.rgb *= DirectionalLightingIntensity;
-
-	// Reconstruct the SH coefficients based on what was encoded
-	FTwoBandSHVectorRGB TranslucentLighting;
-	TranslucentLighting.R.V.x = AmbientLightingVector.r;
-	TranslucentLighting.G.V.x = AmbientLightingVector.g;
-	TranslucentLighting.B.V.x = AmbientLightingVector.b;
-	float3 NormalizedAmbientColor = AmbientLightingVector.rgb / ( Luminance( AmbientLightingVector.rgb ) + 0.00001f );
-
-	// Scale the monocrome directional coefficients with the normalzed ambient color as an approximation to the uncompressed values
-	TranslucentLighting.R.V.yzw = DirectionalLightingVector.rgb * NormalizedAmbientColor.r;
-	TranslucentLighting.G.V.yzw = DirectionalLightingVector.rgb * NormalizedAmbientColor.g;
-	TranslucentLighting.B.V.yzw = DirectionalLightingVector.rgb * NormalizedAmbientColor.b;
-			
-	// Compute diffuse lighting which takes the normal into account
-	FTwoBandSHVector DiffuseTransferSH = CalcDiffuseTransferSH(WorldNormal, 1);
-	VolumeLighting = float4(max(half3(0,0,0), DotSH(TranslucentLighting, DiffuseTransferSH)), AmbientLightingVector.a);
-	InterpolatedLighting += DiffuseColor * VolumeLighting.rgb;
-}
-
 /** Calculates lighting for translucency. */
 float3 GetTranslucencyVolumeLighting(
 	FMaterialPixelParameters MaterialParameters, 
diff --git a/Engine/Shaders/Private/MobileBasePassPixelShader.usf b/Engine/Shaders/Private/MobileBasePassPixelShader.usf
--- a/Engine/Shaders/Private/MobileBasePassPixelShader.usf
+++ b/Engine/Shaders/Private/MobileBasePassPixelShader.usf
@@ -80,6 +80,7 @@
 #endif
 
 #define TRANSLUCENCY_NON_DIRECTIONAL (MATERIALBLENDING_ANY_TRANSLUCENT && (TRANSLUCENCY_LIGHTING_VOLUMETRIC_NONDIRECTIONAL || TRANSLUCENCY_LIGHTING_VOLUMETRIC_PERVERTEX_NONDIRECTIONAL) && !MATERIAL_SHADINGMODEL_SINGLELAYERWATER)
+#define TRANSLUCENCY_SH_LIGHTING (MATERIAL_SHADINGMODEL_DEFAULT_LIT || MATERIAL_SHADINGMODEL_SUBSURFACE) && (MATERIALBLENDING_TRANSLUCENT || MATERIALBLENDING_ADDITIVE) && !TRANSLUCENCY_LIGHTING_SURFACE_FORWARDSHADING && !MATERIAL_SHADINGMODEL_SINGLELAYERWATER
 
 #include "SHCommon.ush"
 #include "/Engine/Generated/Material.ush"
diff --git a/Engine/Shaders/Private/MobileLightingCommon.ush b/Engine/Shaders/Private/MobileLightingCommon.ush
--- a/Engine/Shaders/Private/MobileLightingCommon.ush
+++ b/Engine/Shaders/Private/MobileLightingCommon.ush
@@ -182,6 +182,50 @@ half MobileDirectionalLightCSM(float2 ScreenPosition, float SceneDepth, inout fl
 	return ShadowMap;
 }
 
+/*------------------------------------------------------------------------------
+	Mobile Translucency Light.
+------------------------------------------------------------------------------*/
+
+FLightAccumulator GetTranslucencySHLighting(FGBufferData GBuffer, FDeferredLightData LightData, float3 TranslatedWorldPosition)
+{
+	FLightAccumulator LightAccumulator = (FLightAccumulator)0;
+	float4 VolumeLighting;
+	float3 InterpolatedLighting = 0;
+
+	float Attenuation = 1.0f;
+	float3 L = LightData.Direction;	// Already normalized
+	float3 NormalizedLightVector = L;
+	if (LightData.bRadialLight)
+	{
+		Attenuation = GetLocalLightAttenuation(TranslatedWorldPosition, LightData, L, NormalizedLightVector);
+	}
+
+	float3 Lighting = LightData.Color / PI * Attenuation;
+	FTwoBandSHVectorRGB SHLighting = MulSH(SHBasisFunction(NormalizedLightVector), Lighting);
+
+	float4 AmbientLightingVector = float4(SHLighting.R.V.x, SHLighting.G.V.x, SHLighting.B.V.x, 1.0f);
+
+	float3 LuminanceWeights = float3(.3, .59, .11);
+	float3 Coefficient0 = float3(SHLighting.R.V.y, SHLighting.G.V.y, SHLighting.B.V.y);
+	float3 Coefficient1 = float3(SHLighting.R.V.z, SHLighting.G.V.z, SHLighting.B.V.z);
+	float3 Coefficient2 = float3(SHLighting.R.V.w, SHLighting.G.V.w, SHLighting.B.V.w);
+	float3 DirectionalLightingVector = float3(dot(Coefficient0, LuminanceWeights), dot(Coefficient1, LuminanceWeights), dot(Coefficient2, LuminanceWeights));
+
+#if TRANSLUCENCY_LIGHTING_VOLUMETRIC_PERVERTEX_DIRECTIONAL || TRANSLUCENCY_LIGHTING_VOLUMETRIC_DIRECTIONAL || TRANSLUCENCY_LIGHTING_SURFACE_LIGHTINGVOLUME
+
+	GetVolumeLightingDirectional(AmbientLightingVector, DirectionalLightingVector, GBuffer.WorldNormal, GBuffer.DiffuseColor, InterpolatedLighting, VolumeLighting);
+
+#elif TRANSLUCENCY_LIGHTING_VOLUMETRIC_PERVERTEX_NONDIRECTIONAL || TRANSLUCENCY_LIGHTING_VOLUMETRIC_NONDIRECTIONAL
+
+	GetVolumeLightingNonDirectional(AmbientLightingVector, GBuffer.DiffuseColor, InterpolatedLighting, VolumeLighting);
+
+#endif
+
+	LightAccumulator_AddSplit(LightAccumulator, InterpolatedLighting, 0.0f, InterpolatedLighting, 1.0f, false);
+
+	return LightAccumulator;
+}
+
 /*------------------------------------------------------------------------------
 	Mobile Directional Light.
 ------------------------------------------------------------------------------*/
@@ -225,7 +269,11 @@ void AccumulateDirectionalLighting(FGBufferData GBuffer, float3 TranslatedWorldP
 
 	half4 LightAttenuation = half4(1, 1, DynamicShadowing, DynamicShadowing);
 
+#if TRANSLUCENCY_SH_LIGHTING
+	FLightAccumulator NewLighting = GetTranslucencySHLighting(GBuffer, LightData, TranslatedWorldPosition);
+#else
 	FLightAccumulator NewLighting = AccumulateDynamicLighting(TranslatedWorldPosition, CameraVector, GBuffer, 1, GBuffer.ShadingModelID, LightData, LightAttenuation, 0, uint2(0, 0), OutDirectionalLightShadow);
+#endif
 	DirectLighting = LightAccumulator_Add(DirectLighting, NewLighting);
 }
 
@@ -465,7 +513,11 @@ void AccumulateLightGridLocalLighting(const FCulledLightsGridData InLightGridDat
 		half4 LightAttenuation = half4(1, 1, DynamicShadowing, DynamicShadowing);
 
 		float SurfaceShadow = 0;
+#if TRANSLUCENCY_SH_LIGHTING
+		FLightAccumulator NewLighting = GetTranslucencySHLighting(GBuffer, LightData, TranslatedWorldPosition);
+#else
 		FLightAccumulator NewLighting = AccumulateDynamicLighting(TranslatedWorldPosition, CameraVector, GBuffer, 1, GBuffer.ShadingModelID, LightData, LightAttenuation, 0, uint2(0, 0), SurfaceShadow);
+#endif
 		DirectLighting = LightAccumulator_Add(DirectLighting, NewLighting);
 	}
 }
