Fixed: too strong specular lighting on mobile, since LightColor is no longer pre-multiplied by 1/PI

diff --git a/Engine/Shaders/Private/MobileGGX.ush b/Engine/Shaders/Private/MobileGGX.ush
--- a/Engine/Shaders/Private/MobileGGX.ush
+++ b/Engine/Shaders/Private/MobileGGX.ush
@@ -29,7 +29,7 @@ half GGX_Mobile(half Roughness, float NoH)
 	half a = Roughness * Roughness;
 	half n = NoH * a;
 	half p = a / (OneMinusNoHSqr + n * n);
-	half d = p * p;
+	half d = (1.0 / PI) * p * p;
 	// clamp to avoid overlfow in a bright env
 	return min(d, 2048.0);
 }
\ No newline at end of file
