Use sampling that matches vertex interpolation when sampling per pixel height in landscape.
This fixes the mismatch between rendering landscape to RVT with a single quad and with a high LOD mesh.

diff --git a/Engine/Shaders/Private/LandscapeVertexFactory.ush b/Engine/Shaders/Private/LandscapeVertexFactory.ush
--- a/Engine/Shaders/Private/LandscapeVertexFactory.ush
+++ b/Engine/Shaders/Private/LandscapeVertexFactory.ush
@@ -378,11 +378,54 @@ bool HasVertexFactoryPerPixelHeight()
 	return LandscapeParameters.VirtualTexturePerPixelHeight != 0;
 }
 
+bool UseBarycentricPerPixelHeight(float2 UV, float2 TextureSize)
+{
+	// It only makes sense to use barycentric sampling when sampling resolution is sufficiently above heightmap resolution.
+	float2 dUVdx = ddx(UV * TextureSize);
+	float2 dUVdy = ddy(UV * TextureSize);
+	float T = max(dot(dUVdx, dUVdx), dot(dUVdy, dUVdy));
+	return T < 0.125f;
+}
+
 float GetVertexFactoryPerPixelHeight(FVertexFactoryInterpolantsVSToPS Interpolants)
 {
-	// Use NormalmapTextureSampler for bilinear (HeightmapTextureSampler is usually point).
-	float4 SampleValue = LandscapeParameters.HeightmapTexture.Sample(LandscapeParameters.NormalmapTextureSampler, Interpolants.WeightHeightMapTexCoord.zw);
-	return DecodePackedHeight(SampleValue.xy);
+	float2 UV = Interpolants.WeightHeightMapTexCoord.zw;
+	float2 TextureSize = LandscapeParameters.HeightmapTextureSize.xy;
+
+	if (!UseBarycentricPerPixelHeight(UV, TextureSize))
+	{
+		// Use simple bilinear interpolation of heightmap.
+		// Ideally we want to reuse this sample from VertexFactoryGetPerPixelTangentBasis().
+		float4 SampleValue = LandscapeParameters.NormalmapTexture.Sample(LandscapeParameters.NormalmapTextureSampler, UV);
+		return DecodePackedHeight(SampleValue.xy);
+	}
+
+	// Fake barycentric interpolation to try and match same result that we would get from vertex interpolation with a vertex located at each texel.
+	float2 PixelCoord = UV * TextureSize - 0.5f;
+	float2 FloorPixelCoord = floor(PixelCoord);
+	float2 SubPixelPos = PixelCoord - FloorPixelCoord;
+
+	float B1 = min(SubPixelPos.x, SubPixelPos.y);
+	float B2 = abs(SubPixelPos.y - SubPixelPos.x);
+	float B0 = 1.0f - B1 - B2;
+
+	// Which is the third sample/triangle vertex depends on if we are upper or lower triangle in the quad.
+	bool bLowerTriangle = (SubPixelPos.x > SubPixelPos.y);
+
+	float2 InvTextureSize = LandscapeParameters.HeightmapTextureSize.zw;
+	float2 BaseUV = (FloorPixelCoord + 0.5f) * InvTextureSize;
+	float2 UV1 = InvTextureSize;
+	float2 UV2 = bLowerTriangle ? float2(InvTextureSize.x, 0) : float2(0, InvTextureSize.y);
+
+	float2 S0 = LandscapeParameters.HeightmapTexture.SampleLevel(LandscapeParameters.HeightmapTextureSampler, BaseUV, 0).xy;
+	float2 S1 = LandscapeParameters.HeightmapTexture.SampleLevel(LandscapeParameters.HeightmapTextureSampler, BaseUV + UV1, 0).xy;
+	float2 S2 = LandscapeParameters.HeightmapTexture.SampleLevel(LandscapeParameters.HeightmapTextureSampler, BaseUV + UV2, 0).xy;
+
+	float H0 = DecodePackedHeight(S0);
+	float H1 = DecodePackedHeight(S1);
+	float H2 = DecodePackedHeight(S2);
+
+	return B0 * H0 + B1 * H1 + B2 * H2;
 }
 
 /** Converts from vertex factory specific interpolants to a FMaterialPixelParameters, which is used by material inputs. */
diff --git a/Engine/Source/Runtime/Landscape/Private/LandscapeRender.cpp b/Engine/Source/Runtime/Landscape/Private/LandscapeRender.cpp
--- a/Engine/Source/Runtime/Landscape/Private/LandscapeRender.cpp
+++ b/Engine/Source/Runtime/Landscape/Private/LandscapeRender.cpp
@@ -1577,6 +1577,9 @@ void FLandscapeComponentSceneProxy::OnTransformChanged()
 
 	if (HeightmapTexture)
 	{
+		const float SizeX = FMath::Max(HeightmapTexture->GetSizeX(), 1);
+		const float SizeY = FMath::Max(HeightmapTexture->GetSizeY(), 1);
+		LandscapeParams.HeightmapTextureSize = FVector4f(SizeX, SizeY, 1.f / SizeX, 1.f / SizeY);
 		LandscapeParams.HeightmapTexture = HeightmapTexture->TextureReference.TextureReferenceRHI;
 		LandscapeParams.HeightmapTextureSampler = TStaticSamplerState<SF_Point>::GetRHI();
 		LandscapeParams.NormalmapTexture = HeightmapTexture->TextureReference.TextureReferenceRHI;
@@ -1584,6 +1587,7 @@ void FLandscapeComponentSceneProxy::OnTransformChanged()
 	}
 	else
 	{
+		LandscapeParams.HeightmapTextureSize = FVector4f(1, 1, 1, 1);
 		LandscapeParams.HeightmapTexture = GBlackTexture->TextureRHI;
 		LandscapeParams.HeightmapTextureSampler = GBlackTexture->SamplerStateRHI;
 		LandscapeParams.NormalmapTexture = GBlackTexture->TextureRHI;
diff --git a/Engine/Source/Runtime/Landscape/Public/LandscapeRender.h b/Engine/Source/Runtime/Landscape/Public/LandscapeRender.h
--- a/Engine/Source/Runtime/Landscape/Public/LandscapeRender.h
+++ b/Engine/Source/Runtime/Landscape/Public/LandscapeRender.h
@@ -97,7 +97,8 @@ BEGIN_GLOBAL_SHADER_PARAMETER_STRUCT(FLandscapeUniformShaderParameters, LANDSCAP
 	SHADER_PARAMETER(int32, NumSubsections)
 	SHADER_PARAMETER(int32, LastLOD)
 	SHADER_PARAMETER(uint32, VirtualTexturePerPixelHeight)
-    SHADER_PARAMETER(FVector4f, HeightmapUVScaleBias)
+	SHADER_PARAMETER(FVector4f, HeightmapTextureSize)
+	SHADER_PARAMETER(FVector4f, HeightmapUVScaleBias)
     SHADER_PARAMETER(FVector4f, WeightmapUVScaleBias)
     SHADER_PARAMETER(FVector4f, LandscapeLightmapScaleBias)
     SHADER_PARAMETER(FVector4f, SubsectionSizeVertsLayerUVPan)
