Nanite materials using VT feedback should only ever use [earlydepthstencil], not [ReZ] (Nanite will never have side effects in the material itself like clip() or SV_Depth). Redefined clip to always noop for Nanite (removing CLIP in the process). Also optimized Nanite base pass to always force wave32 regardless of programmable or not.

diff --git a/Engine/Shaders/Private/BasePassPixelShader.usf b/Engine/Shaders/Private/BasePassPixelShader.usf
--- a/Engine/Shaders/Private/BasePassPixelShader.usf
+++ b/Engine/Shaders/Private/BasePassPixelShader.usf
@@ -2094,7 +2094,10 @@ void FPixelShaderInOut_MainPS(
 //		- Pixel depth offset prevents the usage of early depth test because the depth is generated in shader. So we can only relay on LateZ in this case.
 //		- DEPTHSTENCIL_EARLYTEST_LATEWRITE does not help in this case due to UAV writes we need to skip.
 //		- In order to avoid UAV writes, we will run in shader manual depth tests.
-#if !OUTPUT_PIXEL_DEPTH_OFFSET
+#if IS_NANITE_PASS
+	// Even with masking and/or pixel depth offset, Nanite never uses actual clip() or SV_Depth
+	#define PIXELSHADER_EARLYDEPTHSTENCIL EARLYDEPTHSTENCIL
+#elif !OUTPUT_PIXEL_DEPTH_OFFSET
 	#define PIXELSHADER_EARLYDEPTHSTENCIL EARLYDEPTHSTENCIL
 #elif COMPILER_SUPPORTS_DEPTHSTENCIL_EARLYTEST_LATEWRITE
 	// If we support early depth test with late write behaviour then use it since we may be using discard, or modifying depth
@@ -2105,7 +2108,10 @@ void FPixelShaderInOut_MainPS(
 
 // If virtual texture is enabled then use early depth test so that UAV feedback buffer writes respect the depth test
 #if NUM_VIRTUALTEXTURE_SAMPLES || LIGHTMAP_VT_ENABLED || OIT_ENABLED
-	#if COMPILER_SUPPORTS_DEPTHSTENCIL_EARLYTEST_LATEWRITE
+	#if IS_NANITE_PASS
+		// Even with masking and/or pixel depth offset, Nanite never uses actual clip() or SV_Depth
+		#define PIXELSHADER_EARLYDEPTHSTENCIL EARLYDEPTHSTENCIL
+	#elif COMPILER_SUPPORTS_DEPTHSTENCIL_EARLYTEST_LATEWRITE
 		// If we support early depth test with late write behaviour then use it since we may be using discard, or modifying depth
 		#define PIXELSHADER_EARLYDEPTHSTENCIL DEPTHSTENCIL_EARLYTEST_LATEWRITE
 	#elif !OUTPUT_PIXEL_DEPTH_OFFSET
diff --git a/Engine/Shaders/Private/MaterialTemplate.ush b/Engine/Shaders/Private/MaterialTemplate.ush
--- a/Engine/Shaders/Private/MaterialTemplate.ush
+++ b/Engine/Shaders/Private/MaterialTemplate.ush
@@ -25,7 +25,7 @@
 #include "/Engine/Private/DeferredShadingCommon.ush"
 
 // Reduce latency of Nanite material base pass pixel shaders caused by register pressure due to the calculation of WPO
-#if USES_WORLD_POSITION_OFFSET && PIXELSHADER && IS_NANITE_SHADING_PASS && IS_BASE_PASS
+#if PIXELSHADER && IS_NANITE_SHADING_PASS && IS_BASE_PASS
 	COMPILER_FORCE_WAVE32_MODE
 #endif
 
@@ -129,16 +129,14 @@
 
 #if IS_NANITE_PASS
 	#define ALLOW_CLIP 0
-	#define CLIP(x)
+	#define clip(x)
 	#if COMPUTESHADER
-		#define clip(x)
 		#define ddx(x) 0
 		#define ddy(x) 0
 		#define fwidth(x) 0
 	#endif
 #else
 	#define ALLOW_CLIP 1
-	#define CLIP(x) clip(x)
 #endif
 
 #define TEMPLATE_USES_STRATA			(STRATA_ENABLED && MATERIAL_IS_STRATA)
@@ -3167,7 +3165,7 @@ void ClipLODTransition(float2 SvPosition, float DitherFactor)
 		half RetVal = (DitherFactor < 0.0) ?
 			(DitherFactor + 1.0 > RandomVal) :
 			(DitherFactor < RandomVal);
-		CLIP(RetVal - .001);
+		clip(RetVal - .001);
 	}
 }
 
@@ -3212,11 +3210,11 @@ void GetMaterialClippingShadowDepth(FMaterialPixelParameters Parameters, FPixelM
 {
 	ClipLODTransition(Parameters);
 	#if MATERIALBLENDING_MASKED
-		CLIP(GetMaterialMask(PixelMaterialInputs));
+		clip(GetMaterialMask(PixelMaterialInputs));
 	#elif TRANSLUCENT_SHADOW_WITH_MASKED_OPACITY
-		CLIP(GetMaterialMaskedOpacity(PixelMaterialInputs));
+		clip(GetMaterialMaskedOpacity(PixelMaterialInputs));
 	#elif MATERIALBLENDING_TRANSLUCENT
-		CLIP(GetMaterialOpacity(PixelMaterialInputs) - 1.0f / 255.0f);
+		clip(GetMaterialOpacity(PixelMaterialInputs) - 1.0f / 255.0f);
 	#endif
 }
 
@@ -3258,9 +3256,9 @@ void GetMaterialCoverageAndClipping(FMaterialPixelParameters Parameters, FPixelM
 		float Noise = frac( dot( float2( 171.0, 231.0 ) / 71, Pos.xy ) );
 		float Dither = ( Dither5 * 5 + Noise ) * (1.0 / 6.0);
 
-		CLIP( GetMaterialMask(PixelMaterialInputs) + Dither - 0.5 );
+		clip( GetMaterialMask(PixelMaterialInputs) + Dither - 0.5 );
 	#else
-		CLIP(GetMaterialMask(PixelMaterialInputs));
+		clip(GetMaterialMask(PixelMaterialInputs));
 	#endif
 #endif
 }
@@ -3271,7 +3269,7 @@ void GetMaterialClippingVelocity(FMaterialPixelParameters Parameters, FPixelMate
 {
 #if MATERIALBLENDING_TRANSLUCENT || MATERIALBLENDING_ADDITIVE || MATERIALBLENDING_MODULATE
 	ClipLODTransition(Parameters);
-	CLIP(bIsThinTranslucent ? 1.0f : GetMaterialOpacity(PixelMaterialInputs) - 1.0 / 255.0 - GetMaterialOpacityMaskClipValue());
+	clip(bIsThinTranslucent ? 1.0f : GetMaterialOpacity(PixelMaterialInputs) - 1.0 / 255.0 - GetMaterialOpacityMaskClipValue());
 #else
 	GetMaterialCoverageAndClipping(Parameters, PixelMaterialInputs);
 #endif
@@ -3302,10 +3300,10 @@ uint DiscardMaterialWithPixelCoverage(FMaterialPixelParameters MaterialParameter
 		float Mask = (OriginalMask - MaskClip) / (1.0 - MaskClip);
 		uint CurrentPixelCoverage = GetDerivativeCoverageFromMask(Mask);
 		// Discard pixel shader if all sample are masked to avoid computing other material inputs.
-		CLIP(float(CurrentPixelCoverage) - 0.5);
+		clip(float(CurrentPixelCoverage) - 0.5);
 		return CurrentPixelCoverage;
 	}
-	CLIP(OriginalMask - MaskClip);
+	clip(OriginalMask - MaskClip);
 	return 0xF;
 }
 
diff --git a/Engine/Shaders/Public/ShaderVersion.ush b/Engine/Shaders/Public/ShaderVersion.ush
--- a/Engine/Shaders/Public/ShaderVersion.ush
+++ b/Engine/Shaders/Public/ShaderVersion.ush
@@ -2,4 +2,4 @@
 // Each time the console command is executed it generates a new GUID. As this file is included
 // in Platform.ush (which should be included in any shader) it allows to invalidate the shader DDC.
 // 
-// GUID = 22434BFA-4B31-4043-9494-8D69F6E789DD
\ No newline at end of file
+// GUID = 28A1A473-9B5D-40E9-B951-B09A08EB4D7F
\ No newline at end of file
