Fix shader compiler issue on Android.

diff --git a/Engine/Shaders/Private/PostProcessEyeAdaptation.usf b/Engine/Shaders/Private/PostProcessEyeAdaptation.usf
--- a/Engine/Shaders/Private/PostProcessEyeAdaptation.usf
+++ b/Engine/Shaders/Private/PostProcessEyeAdaptation.usf
@@ -8,6 +8,8 @@
 #include "PostProcessCommon.ush"
 #include "PostProcessHistogramCommon.ush"
 
+#if COMPUTESHADER
+
 #define FrontLayerTranslucencyReflectionsStruct LumenGIVolumeStruct
 #define RadianceCacheInterpolation LumenGIVolumeStruct
 #include "Lumen/LumenTranslucencyVolumeShared.ush"
@@ -18,6 +20,8 @@
 #include "SHCommon.ush"
 #include "/Engine/Private/ScreenPass.ush"
 
+#endif
+
 Texture2D EyeAdaptationTexture;
 
 Texture2D HistogramTexture;
