Add and use blue noise for VSM projection to improve performance.
Add ability to adjust the adaptive ray count heuristic to be less aggressive (larger numbers).

diff --git a/Engine/Shaders/Private/VirtualShadowMaps/VirtualShadowMapProjection.usf b/Engine/Shaders/Private/VirtualShadowMaps/VirtualShadowMapProjection.usf
--- a/Engine/Shaders/Private/VirtualShadowMaps/VirtualShadowMapProjection.usf
+++ b/Engine/Shaders/Private/VirtualShadowMaps/VirtualShadowMapProjection.usf
@@ -173,7 +173,7 @@ float SMRTCotMaxRayAngleFromLight;	// Spot/point lights
 float SMRTTexelDitherScale;			// Currently only directional lights
 float SMRTExtrapolateSlope;
 float SMRTMaxSlopeBias;				// Currently only local lights
-uint bSMRTUseAdaptiveRayCount;
+uint SMRTAdaptiveRayCount;
 
 // One pass projection
 RWTexture2D< uint4 > OutShadowMaskBits;
@@ -256,8 +256,6 @@ FVirtualShadowMapSampleResult ProjectLight(
 
 		if (SMRTRayCount > 0)
 		{
-			const bool bUseAdaptiveRayCount = bSMRTUseAdaptiveRayCount > 0;
-
 		#if DIRECTIONAL_LIGHT
 			Result = TraceDirectional(
 				VirtualShadowMapId,
@@ -274,7 +272,7 @@ FVirtualShadowMapSampleResult ProjectLight(
 				WorldNormal,
 				SMRTTexelDitherScale,
 				SMRTExtrapolateSlope,
-				bUseAdaptiveRayCount);
+				SMRTAdaptiveRayCount);
 		#else
 			Result = TraceLocalLight(
 				VirtualShadowMapId,
@@ -292,7 +290,7 @@ FVirtualShadowMapSampleResult ProjectLight(
 				SMRTTexelDitherScale,
 				SMRTExtrapolateSlope,
 				SMRTMaxSlopeBias,
-				bUseAdaptiveRayCount);
+				SMRTAdaptiveRayCount);
 		#endif // DIRECTIONAL_LIGHT
 		}
 		else
@@ -646,7 +644,7 @@ void VirtualShadowMapProjection(
 			ScreenRayLengthWorld,
 			TranslatedWorldPosition,
 			Noise);
-		
+
 		float TransmissionShadow = Result.ShadowFactor;
 		if ( SubsurfaceOpacity < 1.0f )
 		{
diff --git a/Engine/Shaders/Private/VirtualShadowMaps/VirtualShadowMapProjectionDirectional.ush b/Engine/Shaders/Private/VirtualShadowMaps/VirtualShadowMapProjectionDirectional.ush
--- a/Engine/Shaders/Private/VirtualShadowMaps/VirtualShadowMapProjectionDirectional.ush
+++ b/Engine/Shaders/Private/VirtualShadowMaps/VirtualShadowMapProjectionDirectional.ush
@@ -13,6 +13,9 @@ VirtualShadowMapProjectionDirectional.ush:
 #include "VirtualShadowMapProjectionCommon.ush"
 #include "VirtualShadowMapSMRTCommon.ush"
 #include "../Visualization.ush"
+#include "../BlueNoise.ush"
+
+#define USE_BLUE_NOISE 1
 
 float2 ComputeDepthSlopeDirectionalUV(
 	FVirtualShadowMapProjectionShaderData ProjectionData,
@@ -167,6 +170,37 @@ bool IsBackfaceToDirectionalLight(float3 Normal, float3 LightDirection, float Li
 	return dot(Normal, LightDirection) < -SinAlpha;
 }
 
+// See http://extremelearning.com.au/unreasonable-effectiveness-of-quasirandom-sequences/
+float2 R2Sequence(int n)
+{
+	return frac(float(n) * float2(0.754877669f, 0.569840296f));
+}
+
+float4 GetRandomSample(
+	uint2 PixelPos,
+	uint TimeIndex,
+	uint SampleIndex,
+	uint MaxSampleCount)
+{
+#if USE_BLUE_NOISE
+	{
+		// Blue noise is cheaper to compute
+		int2 Offset1 = int2(R2Sequence(SampleIndex) * BlueNoise.Dimensions.xy);
+		int2 Offset2 = int2(R2Sequence(SampleIndex + MaxSampleCount) * BlueNoise.Dimensions.xy);
+		float4 RandomSample;
+		RandomSample.xy = BlueNoiseVec2(PixelPos + Offset1, TimeIndex);
+		RandomSample.zw = BlueNoiseVec2(PixelPos + Offset2, TimeIndex);
+		return RandomSample;
+	}
+#else
+	{
+		RandomSequence RandSequence;
+		RandomSequence_Initialize(RandSequence, PixelPos, SampleIndex, TimeIndex, MaxSampleCount);
+		return RandomSequence_GenerateSample4D(RandSequence);
+	}
+#endif
+}
+
 FVirtualShadowMapSampleResult TraceDirectional(
 	int VirtualShadowMapId,
 	FLightShaderParameters Light,
@@ -182,7 +216,7 @@ FVirtualShadowMapSampleResult TraceDirectional(
 	half3 WorldNormal,
 	float TexelDitherScale,
 	float ExtrapolateSlope,
-	bool bUseAdaptiveRayCount)
+	uint SMRTAdaptiveRayCount)
 {
 	FLWCVector3 WorldPosition = LWCSubtract(LWCPromote(TranslatedWorldPosition), PrimaryView.PreViewTranslation);
 	float3 ViewPosition = mul(float4(TranslatedWorldPosition, 1.0f), View.TranslatedWorldToView).xyz;
@@ -234,13 +268,8 @@ FVirtualShadowMapSampleResult TraceDirectional(
 		for ( ; i < MaxRayCount; i++)
 		{
 			// One sample for ray, one for texel dither
-			RandomSequence RandSequence;
-			RandomSequence_Initialize(RandSequence, PixelPos, i, View.StateFrameIndex, MaxRayCount);
-			float4 RandSample = RandomSequence_GenerateSample4D(RandSequence);
-			float2 E = RandSample.xy;
-
-			float3 RayDir = GetRandomDirectionalLightRayDir(Light, E);
-
+			float4 RandSample = GetRandomSample(PixelPos, View.StateFrameIndex, i, MaxRayCount);
+			float3 RayDir = GetRandomDirectionalLightRayDir(Light, RandSample.xy);
 			float2 TexelOffset = (RandSample.zw - 0.5f) * DitherScale;
 
 			FSMRTClipmapRayState RayState = SMRTClipmapRayInitialize(
@@ -275,8 +304,9 @@ FVirtualShadowMapSampleResult TraceDirectional(
 				++RayMissCount;
 			}
 
-			if (bUseAdaptiveRayCount)
+			if (SMRTAdaptiveRayCount > 0)
 			{
+				// TODO: Adapt this heuristic based on SMRTAdaptiveRayCount as well?
 				if( i == 0 )
 				{
 					bool bHit = SMRTResult.bValidHit;
@@ -288,7 +318,7 @@ FVirtualShadowMapSampleResult TraceDirectional(
 						break;
 					}
 				}
-				else if( i == 1 )
+				else if( i >= SMRTAdaptiveRayCount )
 				{
 					// After 2 iterations and all have hit, assume umbra
 					bool bAllLanesHit = WaveActiveAllTrue( RayMissCount == 0 );
@@ -337,9 +367,7 @@ bool GenerateRayDirectional(
 	// Too low values will greatly restrict how large penumbras can be in screen space.
 	float RayLength = RayLengthScale * DistanceFromViewOrigin;
 
-	RandomSequence RandSequence;
-	RandomSequence_Initialize(RandSequence, PixelPos, RayIndex, View.StateFrameIndexMod8, RayCount);
-	float2 E = RandomSequence_GenerateSample2D(RandSequence);
+	float2 E = GetRandomSample(PixelPos, View.StateFrameIndex, RayIndex, RayCount).xy;
 
 	float3 RayDir = GetRandomDirectionalLightRayDir(Light, E);
 	OutRayStart = TranslatedWorldPosition;
diff --git a/Engine/Shaders/Private/VirtualShadowMaps/VirtualShadowMapProjectionSpot.ush b/Engine/Shaders/Private/VirtualShadowMaps/VirtualShadowMapProjectionSpot.ush
--- a/Engine/Shaders/Private/VirtualShadowMaps/VirtualShadowMapProjectionSpot.ush
+++ b/Engine/Shaders/Private/VirtualShadowMaps/VirtualShadowMapProjectionSpot.ush
@@ -334,7 +334,7 @@ FVirtualShadowMapSampleResult TraceLocalLight(
 	float TexelDitherScale,
 	float ExtrapolateSlope,
 	float MaxSlopeBias,
-	bool bUseAdaptiveRayCount)
+	uint SMRTAdaptiveRayCount)
 {
 	const uint LightType = Light.SpotAngles.x > -2.0f ? LIGHT_TYPE_SPOT : LIGHT_TYPE_POINT;
 	const float3 ToLight = Light.TranslatedWorldPosition - TranslatedWorldPosition;
@@ -478,7 +478,7 @@ FVirtualShadowMapSampleResult TraceLocalLight(
 			Result.VirtualTexelAddress = RayState.VirtualTexelAddress;
 			Result.PhysicalTexelAddress = RayState.PhysicalTexelAddress;
 
-			if (bUseAdaptiveRayCount)
+			if (SMRTAdaptiveRayCount > 0)
 			{
 				if( i == 0 )
 				{
@@ -505,7 +505,7 @@ FVirtualShadowMapSampleResult TraceLocalLight(
 						break;
 					}
 				}
-				else if( i == 1 )
+				else if( i >= SMRTAdaptiveRayCount )
 				{
 					// After 2 iterations and all have hit, must be in umbra.
 					// TODO: Modify this heuristic based on ray counts to avoid blocky artifacts in near-umbra
diff --git a/Engine/Source/Runtime/Renderer/Private/VirtualShadowMaps/VirtualShadowMapProjection.cpp b/Engine/Source/Runtime/Renderer/Private/VirtualShadowMaps/VirtualShadowMapProjection.cpp
--- a/Engine/Source/Runtime/Renderer/Private/VirtualShadowMaps/VirtualShadowMapProjection.cpp
+++ b/Engine/Source/Runtime/Renderer/Private/VirtualShadowMaps/VirtualShadowMapProjection.cpp
@@ -25,6 +25,7 @@
 #include "SceneRendering.h"
 #include "VirtualShadowMapClipmap.h"
 #include "HairStrands/HairStrandsData.h"
+#include "BlueNoise.h"
 
 #define MAX_TEST_PERMUTATION 0
 
@@ -225,6 +226,7 @@ class FVirtualShadowMapProjectionCS : public FGlobalShader
 		SHADER_PARAMETER_RDG_UNIFORM_BUFFER(FVirtualVoxelParameters, HairStrandsVoxel)
 		SHADER_PARAMETER_RDG_UNIFORM_BUFFER(FStrataGlobalUniformParameters, Strata)
 		SHADER_PARAMETER_STRUCT_REF(FViewUniformShaderParameters, View)
+		SHADER_PARAMETER_STRUCT_REF(FBlueNoise, BlueNoise)
 		SHADER_PARAMETER(FIntVector4, ProjectionRect)
 		SHADER_PARAMETER(float, ScreenRayLength)
 		SHADER_PARAMETER(float, NormalBias)
@@ -237,7 +239,7 @@ class FVirtualShadowMapProjectionCS : public FGlobalShader
 		SHADER_PARAMETER(float, SMRTTexelDitherScale)
 		SHADER_PARAMETER(float, SMRTExtrapolateSlope)
 		SHADER_PARAMETER(float, SMRTMaxSlopeBias)
-		SHADER_PARAMETER(uint32, bSMRTUseAdaptiveRayCount)
+		SHADER_PARAMETER(uint32, SMRTAdaptiveRayCount)
 		SHADER_PARAMETER(uint32, InputType)
 		SHADER_PARAMETER(uint32, bCullBackfacingPixels)
 		// One pass projection parameters
@@ -329,7 +331,7 @@ static void RenderVirtualShadowMapProjectionCommon(
 	PassParameters->SubsurfaceMinSourceRadius = FMath::Sin(0.5f * FMath::DegreesToRadians(CVarSubsurfaceShadowMinSourceAngle.GetValueOnRenderThread()));
 	PassParameters->InputType = uint32(InputType);
 	PassParameters->bCullBackfacingPixels = VirtualShadowMapArray.ShouldCullBackfacingPixels() ? 1 : 0;
-	PassParameters->bSMRTUseAdaptiveRayCount = CVarSMRTAdaptiveRayCount.GetValueOnRenderThread() != 0 ? 1 : 0;
+	PassParameters->SMRTAdaptiveRayCount = CVarSMRTAdaptiveRayCount.GetValueOnRenderThread();
 	PassParameters->Strata = Strata::BindStrataGlobalUniformParameters(View);
 	if (bUseTileList)
 	{
@@ -342,6 +344,9 @@ static void RenderVirtualShadowMapProjectionCommon(
 		PassParameters->HairStrandsVoxel = HairStrands::BindHairStrandsVoxelUniformParameters(View);
 	}
 
+	FBlueNoise BlueNoise = GetBlueNoiseParameters();
+	PassParameters->BlueNoise = CreateUniformBufferImmediate(BlueNoise, EUniformBufferUsage::UniformBuffer_SingleDraw);
+
 	bool bDirectionalLight = false;
 	bool bOnePassProjection = LightProxy == nullptr;
 	if (bOnePassProjection)
@@ -424,7 +429,7 @@ static void RenderVirtualShadowMapProjectionCommon(
 		GraphBuilder.AddPass(
 			RDG_EVENT_NAME("VirtualShadowMapProjection(RayCount:%u(%s),SamplesPerRay:%u,Input:%s%s,TileList)",
 				PassParameters->SMRTRayCount,
-				PassParameters->bSMRTUseAdaptiveRayCount ? TEXT("Adaptive") : TEXT("Static"),
+				PassParameters->SMRTUseAdaptiveRayCount > 0 ? TEXT("Adaptive") : TEXT("Static"),
 				PassParameters->SMRTSamplesPerRay,
 				ToString(InputType),
 				bDebugOutput ? TEXT(",Debug") : TEXT("")),
@@ -442,7 +447,7 @@ static void RenderVirtualShadowMapProjectionCommon(
 			GraphBuilder,
 			RDG_EVENT_NAME("VirtualShadowMapProjection(RayCount:%u(%s),SamplesPerRay:%u,Input:%s%s)",
 				PassParameters->SMRTRayCount,
-				PassParameters->bSMRTUseAdaptiveRayCount ? TEXT("Adaptive") : TEXT("Static"),
+				PassParameters->SMRTUseAdaptiveRayCount > 0 ? TEXT("Adaptive") : TEXT("Static"),
 				PassParameters->SMRTSamplesPerRay,
 				ToString(InputType),
 				bDebugOutput ? TEXT(",Debug") : TEXT("")),
