Add landscape actor setting for choosing to use high quality per pixel heightmap sampling used with RVT.
Add a per platform CVar to completely disable high quality per pixel heightmap sampling if we want minimal cost.

diff --git a/Engine/Shaders/Private/LandscapeVertexFactory.ush b/Engine/Shaders/Private/LandscapeVertexFactory.ush
--- a/Engine/Shaders/Private/LandscapeVertexFactory.ush
+++ b/Engine/Shaders/Private/LandscapeVertexFactory.ush
@@ -370,62 +370,66 @@ float3x3 VertexFactoryGetPerPixelTangentBasis(FVertexFactoryInterpolantsVSToPS I
 	return Result;
 }
 
-/** Lookup per-pixel height from heightmap texture */
+/** Lookup per-pixel height from heightmap texture. */
 #define VF_PER_PIXEL_HEIGHTMAP 1
+/** Use high quality barycentric version of interpolation. */
+#ifndef PER_PIXEL_HEIGHTMAP_HQ
+#define PER_PIXEL_HEIGHTMAP_HQ 0
+#endif
 
 bool HasVertexFactoryPerPixelHeight()
 {
 	return LandscapeParameters.VirtualTexturePerPixelHeight != 0;
 }
 
-bool UseBarycentricPerPixelHeight(float2 UV, float2 TextureSize)
-{
-	// It only makes sense to use barycentric sampling when sampling resolution is sufficiently above heightmap resolution.
-	float2 dUVdx = ddx(UV * TextureSize);
-	float2 dUVdy = ddy(UV * TextureSize);
-	float T = max(dot(dUVdx, dUVdx), dot(dUVdy, dUVdy));
-	return T < 0.125f;
-}
-
 float GetVertexFactoryPerPixelHeight(FVertexFactoryInterpolantsVSToPS Interpolants)
 {
 	float2 UV = Interpolants.WeightHeightMapTexCoord.zw;
-	float2 TextureSize = LandscapeParameters.HeightmapTextureSize.xy;
 
-	if (!UseBarycentricPerPixelHeight(UV, TextureSize))
+#if PER_PIXEL_HEIGHTMAP_HQ
+	if (LandscapeParameters.VirtualTexturePerPixelHeight > 1)
 	{
-		// Use simple bilinear interpolation of heightmap.
-		// Ideally we want to reuse this sample from VertexFactoryGetPerPixelTangentBasis().
-		float4 SampleValue = LandscapeParameters.NormalmapTexture.Sample(LandscapeParameters.NormalmapTextureSampler, UV);
-		return DecodePackedHeight(SampleValue.xy);
+		// It only makes sense to use barycentric sampling when sampling resolution is sufficiently above heightmap resolution.
+		float2 TextureSize = LandscapeParameters.HeightmapTextureSize.xy;
+		float2 dUVdx = ddx(UV * TextureSize);
+		float2 dUVdy = ddy(UV * TextureSize);
+		float T = max(dot(dUVdx, dUVdx), dot(dUVdy, dUVdy));
+		if (T < 0.125f)
+		{
+			// Fake barycentric interpolation to try and match same result that we would get from vertex interpolation with a vertex located at each texel.
+			float2 PixelCoord = UV * TextureSize - 0.5f;
+			float2 FloorPixelCoord = floor(PixelCoord);
+			float2 SubPixelPos = PixelCoord - FloorPixelCoord;
+
+			float B1 = min(SubPixelPos.x, SubPixelPos.y);
+			float B2 = abs(SubPixelPos.y - SubPixelPos.x);
+			float B0 = 1.0f - B1 - B2;
+
+			// Which is the third sample/triangle vertex depends on if we are upper or lower triangle in the quad.
+			bool bLowerTriangle = (SubPixelPos.x > SubPixelPos.y);
+
+			float2 InvTextureSize = LandscapeParameters.HeightmapTextureSize.zw;
+			float2 BaseUV = (FloorPixelCoord + 0.5f) * InvTextureSize;
+			float2 UV1 = InvTextureSize;
+			float2 UV2 = bLowerTriangle ? float2(InvTextureSize.x, 0) : float2(0, InvTextureSize.y);
+
+			float2 S0 = LandscapeParameters.HeightmapTexture.SampleLevel(LandscapeParameters.HeightmapTextureSampler, BaseUV, 0).xy;
+			float2 S1 = LandscapeParameters.HeightmapTexture.SampleLevel(LandscapeParameters.HeightmapTextureSampler, BaseUV + UV1, 0).xy;
+			float2 S2 = LandscapeParameters.HeightmapTexture.SampleLevel(LandscapeParameters.HeightmapTextureSampler, BaseUV + UV2, 0).xy;
+
+			float H0 = DecodePackedHeight(S0);
+			float H1 = DecodePackedHeight(S1);
+			float H2 = DecodePackedHeight(S2);
+
+			return B0 * H0 + B1 * H1 + B2 * H2;
+		}
 	}
+#endif // PER_PIXEL_HEIGHTMAP_HQ
 
-	// Fake barycentric interpolation to try and match same result that we would get from vertex interpolation with a vertex located at each texel.
-	float2 PixelCoord = UV * TextureSize - 0.5f;
-	float2 FloorPixelCoord = floor(PixelCoord);
-	float2 SubPixelPos = PixelCoord - FloorPixelCoord;
-
-	float B1 = min(SubPixelPos.x, SubPixelPos.y);
-	float B2 = abs(SubPixelPos.y - SubPixelPos.x);
-	float B0 = 1.0f - B1 - B2;
-
-	// Which is the third sample/triangle vertex depends on if we are upper or lower triangle in the quad.
-	bool bLowerTriangle = (SubPixelPos.x > SubPixelPos.y);
-
-	float2 InvTextureSize = LandscapeParameters.HeightmapTextureSize.zw;
-	float2 BaseUV = (FloorPixelCoord + 0.5f) * InvTextureSize;
-	float2 UV1 = InvTextureSize;
-	float2 UV2 = bLowerTriangle ? float2(InvTextureSize.x, 0) : float2(0, InvTextureSize.y);
-
-	float2 S0 = LandscapeParameters.HeightmapTexture.SampleLevel(LandscapeParameters.HeightmapTextureSampler, BaseUV, 0).xy;
-	float2 S1 = LandscapeParameters.HeightmapTexture.SampleLevel(LandscapeParameters.HeightmapTextureSampler, BaseUV + UV1, 0).xy;
-	float2 S2 = LandscapeParameters.HeightmapTexture.SampleLevel(LandscapeParameters.HeightmapTextureSampler, BaseUV + UV2, 0).xy;
-
-	float H0 = DecodePackedHeight(S0);
-	float H1 = DecodePackedHeight(S1);
-	float H2 = DecodePackedHeight(S2);
-
-	return B0 * H0 + B1 * H1 + B2 * H2;
+	// Use simple bilinear interpolation of heightmap.
+	// Ideally we want to reuse this sample from VertexFactoryGetPerPixelTangentBasis().
+	float4 SampleValue = LandscapeParameters.NormalmapTexture.Sample(LandscapeParameters.NormalmapTextureSampler, UV);
+	return DecodePackedHeight(SampleValue.xy);
 }
 
 /** Converts from vertex factory specific interpolants to a FMaterialPixelParameters, which is used by material inputs. */
diff --git a/Engine/Source/Runtime/Landscape/Classes/LandscapeProxy.h b/Engine/Source/Runtime/Landscape/Classes/LandscapeProxy.h
--- a/Engine/Source/Runtime/Landscape/Classes/LandscapeProxy.h
+++ b/Engine/Source/Runtime/Landscape/Classes/LandscapeProxy.h
@@ -511,6 +511,13 @@ public:
 	UPROPERTY(EditAnywhere, AdvancedDisplay, BlueprintReadOnly, Category = VirtualTexture)
 	bool bVirtualTextureRenderWithQuad = false;
 
+	/** 
+	 * Use highest quality heightmap interpolation when using a single quad to render this landscape to runtime virtual texture pages.
+	 * This also requires the project setting: r.VT.RVT.HighQualityPerPixelHeight.
+	 */
+	UPROPERTY(EditAnywhere, AdvancedDisplay, BlueprintReadOnly, Category = VirtualTexture, meta = (DisplayName = "High Quality PerPixel Height", EditCondition = "bVirtualTextureRenderWithQuad"))
+	bool bVirtualTextureRenderWithQuadHQ = true;
+
 	/** 
 	 * Number of mesh levels to use when rendering landscape into runtime virtual texture.
 	 * Lower values reduce vertex count when rendering to the runtime virtual texture but decrease accuracy when using values that require vertex interpolation.
diff --git a/Engine/Source/Runtime/Landscape/Private/Landscape.cpp b/Engine/Source/Runtime/Landscape/Private/Landscape.cpp
--- a/Engine/Source/Runtime/Landscape/Private/Landscape.cpp
+++ b/Engine/Source/Runtime/Landscape/Private/Landscape.cpp
@@ -3502,6 +3502,7 @@ void ALandscapeProxy::GetSharedProperties(ALandscapeProxy* Landscape)
 		RuntimeVirtualTextures = Landscape->RuntimeVirtualTextures;
 		VirtualTextureLodBias = Landscape->VirtualTextureLodBias;
 		bVirtualTextureRenderWithQuad = Landscape->bVirtualTextureRenderWithQuad;
+		bVirtualTextureRenderWithQuadHQ = Landscape->bVirtualTextureRenderWithQuadHQ;
 		VirtualTextureNumLods = Landscape->VirtualTextureNumLods;
 		VirtualTextureRenderPassType = Landscape->VirtualTextureRenderPassType;
 
diff --git a/Engine/Source/Runtime/Landscape/Private/LandscapeEdit.cpp b/Engine/Source/Runtime/Landscape/Private/LandscapeEdit.cpp
--- a/Engine/Source/Runtime/Landscape/Private/LandscapeEdit.cpp
+++ b/Engine/Source/Runtime/Landscape/Private/LandscapeEdit.cpp
@@ -5459,6 +5459,7 @@ void ALandscapeProxy::PostEditChangeProperty(FPropertyChangedEvent& PropertyChan
 	else if (PropertyName == GET_MEMBER_NAME_CHECKED(ALandscapeProxy, RuntimeVirtualTextures)
 		|| PropertyName == GET_MEMBER_NAME_CHECKED(ALandscapeProxy, VirtualTextureRenderPassType)
 		|| PropertyName == GET_MEMBER_NAME_CHECKED(ALandscapeProxy, bVirtualTextureRenderWithQuad)
+		|| PropertyName == GET_MEMBER_NAME_CHECKED(ALandscapeProxy, bVirtualTextureRenderWithQuadHQ)
 		|| PropertyName == GET_MEMBER_NAME_CHECKED(ALandscapeProxy, VirtualTextureNumLods)
 		|| PropertyName == GET_MEMBER_NAME_CHECKED(ALandscapeProxy, VirtualTextureLodBias))
 	{
@@ -5929,6 +5930,7 @@ void ALandscape::PostEditChangeProperty(FPropertyChangedEvent& PropertyChangedEv
 		PropertyName == GET_MEMBER_NAME_CHECKED(ALandscapeProxy, RuntimeVirtualTextures) ||
 		PropertyName == GET_MEMBER_NAME_CHECKED(ALandscapeProxy, VirtualTextureRenderPassType) ||
 		PropertyName == GET_MEMBER_NAME_CHECKED(ALandscapeProxy, bVirtualTextureRenderWithQuad) ||
+		PropertyName == GET_MEMBER_NAME_CHECKED(ALandscapeProxy, bVirtualTextureRenderWithQuadHQ) ||
 		PropertyName == GET_MEMBER_NAME_CHECKED(ALandscapeProxy, VirtualTextureNumLods) ||
 		PropertyName == GET_MEMBER_NAME_CHECKED(ALandscapeProxy, VirtualTextureLodBias))
 	{
diff --git a/Engine/Source/Runtime/Landscape/Private/LandscapeRender.cpp b/Engine/Source/Runtime/Landscape/Private/LandscapeRender.cpp
--- a/Engine/Source/Runtime/Landscape/Private/LandscapeRender.cpp
+++ b/Engine/Source/Runtime/Landscape/Private/LandscapeRender.cpp
@@ -890,7 +890,10 @@ FLandscapeComponentSceneProxy::FLandscapeComponentSceneProxy(ULandscapeComponent
 	LODSettings.LastLODScreenSizeSquared = LODScreenRatioSquared[LastLOD];
 	LODSettings.ForcedLOD = ForcedLOD;
 
-	bVirtualTextureRenderWithQuad = InComponent->GetLandscapeProxy()->bVirtualTextureRenderWithQuad;
+	const bool bVirtualTextureRenderWithQuad = InComponent->GetLandscapeProxy()->bVirtualTextureRenderWithQuad;
+	const bool bVirtualTextureRenderWithQuadHQ = InComponent->GetLandscapeProxy()->bVirtualTextureRenderWithQuadHQ;
+	VirtualTexturePerPixelHeight = bVirtualTextureRenderWithQuad ? bVirtualTextureRenderWithQuadHQ ? 2 : 1 : 0;
+
 	LastVirtualTextureLOD = MaxLOD;
 	FirstVirtualTextureLOD = bVirtualTextureRenderWithQuad ? MaxLOD : FMath::Max(MaxLOD - InComponent->GetLandscapeProxy()->VirtualTextureNumLods, 0);
 	VirtualTextureLodBias = bVirtualTextureRenderWithQuad ? 0 : InComponent->GetLandscapeProxy()->VirtualTextureLodBias;
@@ -1546,7 +1549,7 @@ void FLandscapeComponentSceneProxy::OnTransformChanged()
 	LandscapeParams.SubsectionSizeVerts = SubsectionSizeVerts;
 	LandscapeParams.NumSubsections = NumSubsections;
 	LandscapeParams.LastLOD = LastLOD;
-	LandscapeParams.VirtualTexturePerPixelHeight = bVirtualTextureRenderWithQuad ? 1 : 0;
+	LandscapeParams.VirtualTexturePerPixelHeight = VirtualTexturePerPixelHeight;
 	LandscapeParams.HeightmapUVScaleBias = HeightmapScaleBias;
 	LandscapeParams.WeightmapUVScaleBias = WeightmapScaleBias;
 	LandscapeParams.LocalToWorldNoScaling = FMatrix44f(LocalToWorldNoScaling);			// LWC_TODO: Precision loss
diff --git a/Engine/Source/Runtime/Landscape/Public/LandscapeRender.h b/Engine/Source/Runtime/Landscape/Public/LandscapeRender.h
--- a/Engine/Source/Runtime/Landscape/Public/LandscapeRender.h
+++ b/Engine/Source/Runtime/Landscape/Public/LandscapeRender.h
@@ -666,7 +666,7 @@ protected:
 	int8						MaxLOD;		// Maximum LOD level, user override possible
 	int8						NumWeightmapLayerAllocations;
 	uint8						StaticLightingLOD;
-	uint8						bVirtualTextureRenderWithQuad : 1;
+	uint8						VirtualTexturePerPixelHeight;
 	float						WeightmapSubsectionOffset;
 	TArray<float>				LODScreenRatioSquared;		// Table of valid screen size -> LOD index
 	int32						FirstLOD;	// First LOD we have batch elements for
diff --git a/Engine/Source/Runtime/Renderer/Private/VT/RuntimeVirtualTextureRender.cpp b/Engine/Source/Runtime/Renderer/Private/VT/RuntimeVirtualTextureRender.cpp
--- a/Engine/Source/Runtime/Renderer/Private/VT/RuntimeVirtualTextureRender.cpp
+++ b/Engine/Source/Runtime/Renderer/Private/VT/RuntimeVirtualTextureRender.cpp
@@ -20,12 +20,17 @@
 
 namespace RuntimeVirtualTexture
 {
+	static TAutoConsoleVariable<int32> CVarVTHighQualityPerPixelHeight(
+		TEXT("r.VT.RVT.HighQualityPerPixelHeight"),
+		1,
+		TEXT("Use higher quality sampling of per pixel heightmaps when rendering to Runtime Virtual Texture.\n"),
+		ECVF_ReadOnly);
+
 	static TAutoConsoleVariable<int32> CVarVTDirectCompress(
 		TEXT("r.VT.RVT.DirectCompress"),
 		1,
 		TEXT("Compress texture data direct to the physical texture on platforms that support it."),
-		ECVF_RenderThreadSafe
-	);
+		ECVF_RenderThreadSafe);
 
     int32 RenderCaptureNextRVTPagesDraws = 0;
     static FAutoConsoleVariableRef CVarRenderCaptureNextRVTPagesDraws(
@@ -129,6 +134,10 @@ namespace RuntimeVirtualTexture
 			FMeshMaterialShader::ModifyCompilationEnvironment(Parameters, OutEnvironment);
 			OutEnvironment.SetDefine(TEXT("VIRTUAL_TEXTURE_PAGE_RENDER"), 1);
 			OutEnvironment.SetDefine(TEXT("IS_VIRTUAL_TEXTURE_MATERIAL"), 1);
+
+			static FShaderPlatformCachedIniValue<bool> HighQualityPerPixelHeightValue(TEXT("r.VT.RVT.HighQualityPerPixelHeight"));
+			const bool bHighQualityPerPixelHeight = (HighQualityPerPixelHeightValue.Get((EShaderPlatform)Parameters.Platform) != 0);
+			OutEnvironment.SetDefine(TEXT("PER_PIXEL_HEIGHTMAP_HQ"), bHighQualityPerPixelHeight ? 1 : 0);
 		}
 
 		FShader_VirtualTextureMaterialDraw()
