Implement support for PixelDepth shading node for Ray Tracing and Path Tracing shaders

This has a side benefit of not requiring the pixel coordinate to be passed in via the payload. This can be cleaned up in a follow up refactor.

diff --git a/Engine/Shaders/Private/PathTracing/PathTracingMaterialHitShader.usf b/Engine/Shaders/Private/PathTracing/PathTracingMaterialHitShader.usf
--- a/Engine/Shaders/Private/PathTracing/PathTracingMaterialHitShader.usf
+++ b/Engine/Shaders/Private/PathTracing/PathTracingMaterialHitShader.usf
@@ -38,11 +38,10 @@ RAY_TRACING_ENTRY_CLOSEST_HIT(PathTracingMaterialCHS,
 	PackedPayload.HitT = RayTCurrent();
 
 #if USE_MATERIAL_CLOSEST_HIT_SHADER
-	const uint2 PixelCoord = PackedPayload.GetPixelCoord();
-
 	ResolvedView = ResolveView();
 
-	float4 SvPosition = float4(PixelCoord.xy, 0.0, 1.0);
+	const float3 TranslatedWorldPosition = TranslatedWorldRayOrigin() + RayTCurrent() * WorldRayDirection();
+	float4 SvPosition = mul(float4(TranslatedWorldPosition, 1.0f), ResolvedView.TranslatedWorldToClip);
 
 #if VF_SUPPORTS_RAYTRACING_PREPARE_MATERIAL_PIXEL_PARAMETERS
 	// this is a newer codepath that is both more flexible and allows for more direct calculation compared to the other codepath
@@ -429,11 +428,10 @@ RAY_TRACING_ENTRY_ANY_HIT(PathTracingMaterialAHS,
 
 	// If we got here -- we have a reason to evaluate the material -- do the setup 
 
-	const uint2 PixelCoord = PackedPayload.GetPixelCoord();
-
 	ResolvedView = ResolveView();
 
-	float4 SvPosition = float4(PixelCoord.xy, 0.0, 1.0);
+	const float3 TranslatedWorldPosition = TranslatedWorldRayOrigin() + RayTCurrent() * WorldRayDirection();
+	float4 SvPosition = mul(float4(TranslatedWorldPosition, 1.0f), ResolvedView.TranslatedWorldToClip);
 
 #if VF_SUPPORTS_RAYTRACING_PREPARE_MATERIAL_PIXEL_PARAMETERS
 	// this is a newer codepath that is both more flexible and allows for more direct calculation compared to the other codepath
diff --git a/Engine/Shaders/Private/RayTracing/RayTracingDecalMaterialShader.usf b/Engine/Shaders/Private/RayTracing/RayTracingDecalMaterialShader.usf
--- a/Engine/Shaders/Private/RayTracing/RayTracingDecalMaterialShader.usf
+++ b/Engine/Shaders/Private/RayTracing/RayTracingDecalMaterialShader.usf
@@ -24,7 +24,6 @@ RAY_TRACING_ENTRY_CALLABLE(RayTracingDecalMaterialShader, FDecalShaderParams, Pa
 	ResolvedView = ResolveView();
 
 	const float3 TranslatedWorldPosition = Params.GetInputTranslatedWorldPosition();
-	const uint2 PixelCoord = Params.GetPixelCoord();
 
 	float4 DecalPos = mul(float4(TranslatedWorldPosition.xyz, 1), WorldToDecal);
 
@@ -55,7 +54,7 @@ RAY_TRACING_ENTRY_CALLABLE(RayTracingDecalMaterialShader, FDecalShaderParams, Pa
 	MaterialParameters.TwoSidedSign = 1;
 	MaterialParameters.VertexColor = 1;
 	MaterialParameters.CameraVector = CameraVector;
-	MaterialParameters.SvPosition = float4(PixelCoord.xy, 0.0, 1.0);
+	MaterialParameters.SvPosition = mul(float4(TranslatedWorldPosition, 1.0f), ResolvedView.TranslatedWorldToClip);
 	MaterialParameters.ScreenPosition = SvPositionToResolvedScreenPosition(MaterialParameters.SvPosition);
 	MaterialParameters.LightVector = SwizzlePos;
 
diff --git a/Engine/Shaders/Private/RayTracing/RayTracingMaterialHitShaders.usf b/Engine/Shaders/Private/RayTracing/RayTracingMaterialHitShaders.usf
--- a/Engine/Shaders/Private/RayTracing/RayTracingMaterialHitShaders.usf
+++ b/Engine/Shaders/Private/RayTracing/RayTracingMaterialHitShaders.usf
@@ -492,12 +492,12 @@ RAY_TRACING_ENTRY_CLOSEST_HIT(MaterialCHS,
 
 	GlobalTextureMipBias = PackedPayload.GetMipBias();
 
-	const uint2 PixelCoord = PackedPayload.GetPixelCoord();
-
 	ResolvedView = ResolveView();
 
 	FVertexFactoryInterpolantsVSToPS Interpolants;
-	float4 SvPosition = float4(PixelCoord.xy, 0.0, 1.0);
+	const float3 TranslatedWorldPosition = TranslatedWorldRayOrigin() + RayTCurrent() * WorldRayDirection();
+	float4 SvPosition = mul(float4(TranslatedWorldPosition, 1.0f), ResolvedView.TranslatedWorldToClip);
+
 	
 	FRayCone PropagatedCone = PropagateRayCone(PackedPayload.GetRayCone(), 0 /* surface curvature */, RayTCurrent());
 	CalcInterpolants(PropagatedCone, Attributes, Interpolants);
@@ -612,12 +612,11 @@ RAY_TRACING_ENTRY_ANY_HIT(MaterialAHS,
 
 	GlobalTextureMipBias = PackedPayload.GetMipBias();
 
-	const uint2 PixelCoord = PackedPayload.GetPixelCoord();
-
 	ResolvedView = ResolveView();
 
 	FVertexFactoryInterpolantsVSToPS Interpolants;
-	float4 SvPosition = float4(PixelCoord.xy, 0.0, 1.0);
+	const float3 TranslatedWorldPosition = TranslatedWorldRayOrigin() + RayTCurrent() * WorldRayDirection();
+	float4 SvPosition = mul(float4(TranslatedWorldPosition, 1.0f), ResolvedView.TranslatedWorldToClip);
 
 	CalcInterpolants(PackedPayload.GetRayCone(), Attributes, Interpolants);
 
