Path Tracer: Remove PixelCoord from Payload and PixelIndex from PathState

This field is no longer used, which frees up 4 bytes in the payload and in the path state.

diff --git a/Engine/Plugins/Experimental/GPULightmass/Shaders/Private/LightmapCommon.ush b/Engine/Plugins/Experimental/GPULightmass/Shaders/Private/LightmapCommon.ush
--- a/Engine/Plugins/Experimental/GPULightmass/Shaders/Private/LightmapCommon.ush
+++ b/Engine/Plugins/Experimental/GPULightmass/Shaders/Private/LightmapCommon.ush
@@ -10,7 +10,7 @@ FMinimalPayload TraceLightmapVisibilityRay(
 	in FRayDesc Ray)
 {
 	// GPULightmass only binds path tracing shaders, so we need to use it to trace rays, even though this function only cares about the hit distance
-	FPackedPathTracingPayload PackedPayload = InitPathTracingVisibilityPayload(uint2(0, 0), 1.0f);
+	FPackedPathTracingPayload PackedPayload = InitPathTracingVisibilityPayload(1.0f);
 
 	TraceRay(
 		TLAS,
diff --git a/Engine/Plugins/Experimental/GPULightmass/Shaders/Private/LightmapPathTracing.usf b/Engine/Plugins/Experimental/GPULightmass/Shaders/Private/LightmapPathTracing.usf
--- a/Engine/Plugins/Experimental/GPULightmass/Shaders/Private/LightmapPathTracing.usf
+++ b/Engine/Plugins/Experimental/GPULightmass/Shaders/Private/LightmapPathTracing.usf
@@ -186,7 +186,7 @@ FPathTracingPayload TraceTransparentRay(FRayDesc Ray, bool IsCameraRay, bool Las
 		// also avoid noise from hot spots (they can be properly
 		// importance sampled if a capturing skylight is added)
 		const bool bIgnoreSky = !IsCameraRay;
-		FPackedPathTracingPayload PackedPayload = InitPathTracingPayload(LaunchIndex, bIgnoreSky);
+		FPackedPathTracingPayload PackedPayload = InitPathTracingPayload(bIgnoreSky);
 		// Trace the ray
 		TraceRay(
 			TLAS,
@@ -289,7 +289,7 @@ FPathTracingPayload TraceTransparentRay(FRayDesc Ray, bool IsCameraRay, bool Las
 	return Payload;
 }
 
-float3 TraceTransparentVisibilityRay(FRayDesc Ray, uint2 PixelCoord, float PathRoughness)
+float3 TraceTransparentVisibilityRay(FRayDesc Ray, float PathRoughness)
 {
 	const uint RayFlags = RAY_FLAG_ACCEPT_FIRST_HIT_AND_END_SEARCH | RAY_FLAG_SKIP_CLOSEST_HIT_SHADER;
 	const uint InstanceInclusionMask = RAY_TRACING_MASK_SHADOW;
@@ -297,7 +297,7 @@ float3 TraceTransparentVisibilityRay(FRayDesc Ray, uint2 PixelCoord, float PathR
 	const uint MultiplierForGeometryContributionToShaderIndex = RAY_TRACING_NUM_SHADER_SLOTS;
 	const uint MissShaderIndex = 0;
 
-	FPackedPathTracingPayload PackedPayload = InitPathTracingVisibilityPayload(PixelCoord, PathRoughness);
+	FPackedPathTracingPayload PackedPayload = InitPathTracingVisibilityPayload(PathRoughness);
 
 	// Trace the ray
 	TraceRay(
@@ -608,7 +608,7 @@ void PathTracingKernel(
 					
 						float AvgRoughness = ApproximateCaustics ? GetAverageRoughness(Payload) : 0.0;
 
-						Visibility = TraceTransparentVisibilityRay(LightRay, LaunchIndex, AvgRoughness);
+						Visibility = TraceTransparentVisibilityRay(LightRay, AvgRoughness);
 						LightSample.RadianceOverPdf *= Visibility;
 					}
 
@@ -783,7 +783,7 @@ void PathTracingKernel(
 					{
 						FRayDesc LightRay = Ray;
 						LightRay.TMax = LightResult.HitT;
-						float3 Visibility = TraceTransparentVisibilityRay(LightRay, LaunchIndex, PathRoughness);
+						float3 Visibility = TraceTransparentVisibilityRay(LightRay, PathRoughness);
 						LightContrib *= Visibility;
 						
 						// Separate direct lighting (bounce 0) for GPU Lightmass
diff --git a/Engine/Shaders/Private/PathTracing/PathTracing.usf b/Engine/Shaders/Private/PathTracing/PathTracing.usf
--- a/Engine/Shaders/Private/PathTracing/PathTracing.usf
+++ b/Engine/Shaders/Private/PathTracing/PathTracing.usf
@@ -44,7 +44,6 @@ FPathState LoadPathStateData(uint Index)
 {
 	FPathTracingPackedPathState Data = PathStateData[Index];
 	FPathState Output = (FPathState)0;
-	Output.PixelIndex = uint2(Data.PixelIndex & 0xFFFFu, Data.PixelIndex >> 16);
 	Output.RandSequence.SampleIndex = Data.RandSeqSampleIndex;
 	Output.RandSequence.SampleSeed  = Data.RandSeqSampleSeed;
 	Output.Radiance = Data.Radiance;
@@ -71,7 +70,6 @@ FPathState LoadPathStateData(uint Index)
 void StorePathStateData(FPathState PathState, uint Index)
 {
 	FPathTracingPackedPathState Data;
-	Data.PixelIndex = PathState.PixelIndex.x + PathState.PixelIndex.y * 65536;;
 	Data.RandSeqSampleIndex = PathState.RandSequence.SampleIndex;
 	Data.RandSeqSampleSeed  = PathState.RandSequence.SampleSeed;
 	Data.Radiance = PathState.Radiance;
diff --git a/Engine/Shaders/Private/PathTracing/PathTracingCommon.ush b/Engine/Shaders/Private/PathTracing/PathTracingCommon.ush
--- a/Engine/Shaders/Private/PathTracing/PathTracingCommon.ush
+++ b/Engine/Shaders/Private/PathTracing/PathTracingCommon.ush
@@ -46,13 +46,9 @@ float3 PayloadDecodeUnitVector(uint Encoded)
 struct FPackedPathTracingPayload : FMinimalPayload
 {
 	// float FMinimalPayload::HitT                       // 4  bytes
-	uint PackedPixelCoord;                               // 4  bytes
-	uint PackedData[14];								 // 56 bytes (encoded data, depends on shading model and ray type)
+	uint PackedData[15];								 // 60 bytes (encoded data, depends on shading model and ray type)
 														 // 64 bytes total
 
-	void SetPixelCoord(uint2 PixelCoord) { PackedPixelCoord = (PixelCoord.x & 0xFFFF) | (PixelCoord.y << 16); }
-	uint2 GetPixelCoord() { return uint2(PackedPixelCoord & 0xFFFF, PackedPixelCoord >> 16); }
-
 	uint GetFlags() { return (PackedData[7] >> 23) & 0x3F; }
 	void SetFlag(uint Flag) { PackedData[7] |= Flag << 23; }
 	uint GetPrimitiveLightingChannelMask() { return (PackedData[7] >> 29) & 0x7; }
@@ -190,10 +186,9 @@ struct FPathTracingPayload : FMinimalPayload
 	}
 };
 
-FPackedPathTracingPayload InitPathTracingPayload(uint2 PixelCoord, bool bIgnoreSkyMaterials)
+FPackedPathTracingPayload InitPathTracingPayload(bool bIgnoreSkyMaterials)
 {
 	FPackedPathTracingPayload Output = (FPackedPathTracingPayload)0;
-	Output.SetPixelCoord(PixelCoord);
 	if (bIgnoreSkyMaterials)
 	{
 		Output.SetIgnoreSkyMaterials();
@@ -201,10 +196,9 @@ FPackedPathTracingPayload InitPathTracingPayload(uint2 PixelCoord, bool bIgnoreS
 	return Output;
 }
 
-FPackedPathTracingPayload InitPathTracingVisibilityPayload(uint2 PixelCoord, float PathRoughness)
+FPackedPathTracingPayload InitPathTracingVisibilityPayload(float PathRoughness)
 {
 	FPackedPathTracingPayload Output = (FPackedPathTracingPayload)0;
-	Output.SetPixelCoord(PixelCoord);
 	// Signal to the AHS we want to evaluate the opacity
 	// The payload is used to carry the path throughput (for transparent shadows)
 	// and current path roughness (for approximate caustics)
diff --git a/Engine/Shaders/Private/PathTracing/PathTracingCore.ush b/Engine/Shaders/Private/PathTracing/PathTracingCore.ush
--- a/Engine/Shaders/Private/PathTracing/PathTracingCore.ush
+++ b/Engine/Shaders/Private/PathTracing/PathTracingCore.ush
@@ -71,7 +71,6 @@ Buffer<float> StartingExtinctionCoefficient;
 
 struct FPathState
 {                                   // packed size
-	uint2 PixelIndex;               //  4 bytes
 	RandomSequence RandSequence;    //  8 bytes
 	float3 Radiance;                // 12 bytes
 	float  BackgroundVisibility;    //  4 bytes
@@ -81,7 +80,7 @@ struct FPathState
 	float3 PathThroughput;          // 12 bytes
 	float PathRoughness;            //  2 bytes // half precision
 	float3 SigmaT; // Extinction    //  6 bytes // half precision 
-									// 84 bytes total (see FPathTracingPackedPathState)
+									// 80 bytes total (see FPathTracingPackedPathState)
 
 	void WritePixel(uint2 TextureIndex)
 	{
@@ -180,7 +179,7 @@ FPathTracingPayload TraceTransparentRay(inout FPathState PathState, bool bIsCame
 		// also avoid noise from hot spots (they can be properly
 		// importance sampled if a capturing skylight is added)
 		const bool bIgnoreSky = !bIsCameraRay;
-		FPackedPathTracingPayload PackedPayload = InitPathTracingPayload(PathState.PixelIndex, bIgnoreSky);
+		FPackedPathTracingPayload PackedPayload = InitPathTracingPayload(bIgnoreSky);
 		// Trace the ray
 		TraceRay(
 			TLAS,
@@ -378,7 +377,6 @@ FPathTracingPayload TraceTransparentRay(inout FPathState PathState, bool bIsCame
 
 				FDecalShaderParams Params = (FDecalShaderParams)0;
 				Params.SetInputTranslatedWorldPosition(TranslatedWorldPos);
-				Params.SetPixelCoord(PathState.PixelIndex);
 
 				CallShader(DecalId, Params);
 
@@ -459,7 +457,7 @@ FPathTracingPayload TraceTransparentRay(inout FPathState PathState, bool bIsCame
 	return Payload;
 }
 
-float3 TraceTransparentVisibilityRay(FRayDesc Ray, uint2 PixelCoord, float PathRoughness, uint MissShaderIndex, bool bCastShadows, inout RandomSequence RandSequence)
+float3 TraceTransparentVisibilityRay(FRayDesc Ray, float PathRoughness, uint MissShaderIndex, bool bCastShadows, inout RandomSequence RandSequence)
 {
 	if (!bCastShadows && MissShaderIndex == 0)
 	{
@@ -472,7 +470,7 @@ float3 TraceTransparentVisibilityRay(FRayDesc Ray, uint2 PixelCoord, float PathR
 	const uint RayContributionToHitGroupIndex = RAY_TRACING_SHADER_SLOT_SHADOW;
 	const uint MultiplierForGeometryContributionToShaderIndex = RAY_TRACING_NUM_SHADER_SLOTS;
 
-	FPackedPathTracingPayload PackedPayload = InitPathTracingVisibilityPayload(PixelCoord, PathRoughness);
+	FPackedPathTracingPayload PackedPayload = InitPathTracingVisibilityPayload(PathRoughness);
 
 	if (!bCastShadows)
 	{
@@ -528,10 +526,10 @@ struct FProbeResult
 	bool IsMiss() { return HitT <= 0; }
 };
 
-FProbeResult TraceProbeRay(FRayDesc Ray, uint2 PixelCoord)
+FProbeResult TraceProbeRay(FRayDesc Ray)
 {
 	// Trace a short ray to see if we escaped the surface
-	FPackedPathTracingPayload PackedPayload = InitPathTracingPayload(PixelCoord, true);
+	FPackedPathTracingPayload PackedPayload = InitPathTracingPayload(true);
 	// TODO: use smaller payload + flag to skip shading work? (only want smooth+geo normals)
 	//       however supporting blockers would require some shading ...
 	const uint RayFlags = 0;
@@ -578,7 +576,7 @@ void ApplyRayBias(inout FRayDesc Ray, float HitT, float3 Normal, bool ShadowTerm
 #endif
 }
 
-bool ProcessSubsurfaceRandomWalk(inout FPathTracingPayload Payload, inout float3 PathThroughput, inout RandomSequence RandSequence, float3 RayDirection, uint2 PixelCoord, bool SimplifySSS)
+bool ProcessSubsurfaceRandomWalk(inout FPathTracingPayload Payload, inout float3 PathThroughput, inout RandomSequence RandSequence, float3 RayDirection, bool SimplifySSS)
 {
 	if (!Payload.IsSubsurfaceMaterial())
 	{
@@ -687,7 +685,7 @@ bool ProcessSubsurfaceRandomWalk(inout FPathTracingPayload Payload, inout float3
 		}
 		for (;;)
 		{
-			FProbeResult ProbeResult = TraceProbeRay(Ray, PixelCoord);
+			FProbeResult ProbeResult = TraceProbeRay(Ray);
 			if (ProbeResult.IsMiss())
 			{
 				// we didn't hit anything, so scatter randomly in the volume and keep tracing
@@ -763,8 +761,7 @@ FPathState CreatePathState(int2 PixelIndex)
 {
 	FPathState Output = (FPathState)0;
 
-	Output.PixelIndex = PixelIndex;
-	uint2 LaunchIndex = Output.PixelIndex + View.ViewRectMin.xy;
+	uint2 LaunchIndex = PixelIndex + View.ViewRectMin.xy;
 
 	// Initialize random sequence 
 	if (SamplerType == PATHTRACER_SAMPLER_ERROR_DIFFUSION)
@@ -939,7 +936,7 @@ bool PathTracingKernel(inout FPathState PathState, int Bounce)
 						LightRay.TMin = 0;
 						LightRay.Direction = LightSample.Direction;
 						LightRay.TMax = LightSample.Distance;
-						LightSample.RadianceOverPdf *= TraceTransparentVisibilityRay(LightRay, PathState.PixelIndex, 1.0, MissShaderIndex, bCastShadows, PathState.RandSequence);
+						LightSample.RadianceOverPdf *= TraceTransparentVisibilityRay(LightRay, 1.0, MissShaderIndex, bCastShadows, PathState.RandSequence);
 						// #dxr_todo: Is it cheaper to fire the ray first? Or eval the material first?
 						if (any(LightSample.RadianceOverPdf > 0))
 						{
@@ -976,7 +973,7 @@ bool PathTracingKernel(inout FPathState PathState, int Bounce)
 								LightContrib *= MISWeightRobust(MaterialSample.Pdf, LightResult.Pdf);
 							}
 							MaterialRay.TMax = LightResult.HitT;
-							LightContrib *= TraceTransparentVisibilityRay(MaterialRay, PathState.PixelIndex, MaterialSample.Roughness, MissShaderIndex, bCastShadows, PathState.RandSequence);
+							LightContrib *= TraceTransparentVisibilityRay(MaterialRay, MaterialSample.Roughness, MissShaderIndex, bCastShadows, PathState.RandSequence);
 							AccumulateRadiance(PathState.Radiance, LightContrib, bIsCameraRay);
 						}
 					}
@@ -1014,7 +1011,7 @@ bool PathTracingKernel(inout FPathState PathState, int Bounce)
 	}
 
 	// Processing the random walk will (stochastically) move the shading point elsewhere on the surface of the object
-	if (!ProcessSubsurfaceRandomWalk(Payload, PathState.PathThroughput, PathState.RandSequence, PathState.Ray.Direction, PathState.PixelIndex, bSimplifySSS))
+	if (!ProcessSubsurfaceRandomWalk(Payload, PathState.PathThroughput, PathState.RandSequence, PathState.Ray.Direction, bSimplifySSS))
 	{
 		// random walk did not terminate at a valid point
 		return false;
@@ -1111,7 +1108,7 @@ bool PathTracingKernel(inout FPathState PathState, int Bounce)
 
 				const bool bCastShadows = CastsShadow(LightId);
 				const uint MissShaderIndex = GetLightMissShaderIndex(LightId);
-				LightSample.RadianceOverPdf *= TraceTransparentVisibilityRay(LightRay, PathState.PixelIndex, AvgRoughness, MissShaderIndex, bCastShadows, PathState.RandSequence);
+				LightSample.RadianceOverPdf *= TraceTransparentVisibilityRay(LightRay, AvgRoughness, MissShaderIndex, bCastShadows, PathState.RandSequence);
 
 				if (bIsVolumeSample)
 				{
@@ -1278,7 +1275,7 @@ bool PathTracingKernel(inout FPathState PathState, int Bounce)
 				LightRay.TMax = LightResult.HitT;
 				const bool bCastShadows = CastsShadow(LightId);
 				const uint MissShaderIndex = GetLightMissShaderIndex(LightId);
-				LightContrib *= TraceTransparentVisibilityRay(LightRay, PathState.PixelIndex, ShadowRayRoughness, MissShaderIndex, bCastShadows, PathState.RandSequence);
+				LightContrib *= TraceTransparentVisibilityRay(LightRay, ShadowRayRoughness, MissShaderIndex, bCastShadows, PathState.RandSequence);
 				// the light made some contribution, and there was nothing along the shadow ray
 				AccumulateRadiance(PathState.Radiance, LightContrib, bIsCameraRay);
 			}
diff --git a/Engine/Shaders/Private/PathTracing/PathTracingInitExtinctionCoefficient.usf b/Engine/Shaders/Private/PathTracing/PathTracingInitExtinctionCoefficient.usf
--- a/Engine/Shaders/Private/PathTracing/PathTracingInitExtinctionCoefficient.usf
+++ b/Engine/Shaders/Private/PathTracing/PathTracingInitExtinctionCoefficient.usf
@@ -9,7 +9,7 @@ RWBuffer<float> RWStartingExtinctionCoefficient;
 // TODO: figure out the right start value for camera rays inside glass
 RAY_TRACING_ENTRY_RAYGEN(PathTracingInitExtinctionCoefficientRG)
 {
-	FPackedPathTracingPayload PackedPayload = InitPathTracingPayload(0, true);
+	FPackedPathTracingPayload PackedPayload = InitPathTracingPayload(true);
 
 	const uint RayFlags = RAY_FLAG_FORCE_OPAQUE;
 	const uint MissShaderIndex = 0;
diff --git a/Engine/Shaders/Shared/RayTracingTypes.h b/Engine/Shaders/Shared/RayTracingTypes.h
--- a/Engine/Shaders/Shared/RayTracingTypes.h
+++ b/Engine/Shaders/Shared/RayTracingTypes.h
@@ -37,7 +37,6 @@ struct FPathTracingLight {
 static_assert(sizeof(FPathTracingLight) == 132, "Path tracing light structure should be kept as small as possible");
 
 struct FPathTracingPackedPathState {
-	uint32    PixelIndex;
 	uint32    RandSeqSampleIndex;
 	uint32    RandSeqSampleSeed;
 	FVector3f Radiance;
@@ -51,7 +50,7 @@ struct FPathTracingPackedPathState {
 	uint16    SigmaT[3];
 };
 
-static_assert(sizeof(FPathTracingPackedPathState) == 84, "Packed Path State size should be minimized");
+static_assert(sizeof(FPathTracingPackedPathState) == 80, "Packed Path State size should be minimized");
 
 // C++ representation of a decal for ray tracing
 struct FRayTracingDecal
@@ -92,7 +91,6 @@ struct FPathTracingLight {
 };
 
 struct FPathTracingPackedPathState {
-	uint      PixelIndex;
 	uint      RandSeqSampleIndex;
 	uint      RandSeqSampleSeed;
 	float3    Radiance;
