Add implementation of Fresnel F82 curve and fix F_Fresnel logic to clamp SpecularColor correctly

diff --git a/Engine/Shaders/Private/BRDF.ush b/Engine/Shaders/Private/BRDF.ush
--- a/Engine/Shaders/Private/BRDF.ush
+++ b/Engine/Shaders/Private/BRDF.ush
@@ -481,9 +481,20 @@ float3 F_Schlick(float3 F0, float3 F90, float VoH)
 	return F90 * Fc + (1 - Fc) * F0;
 }
 
+float3 F_AdobeF82(float3 F0, float3 F82, float VoH)
+{
+	// [Kutz et al. 2021, "Novel aspects of the Adobe Standard Material" ]
+	// See Section 2.3 (note the formulas in the paper do not match the code, the code is the correct version)
+	// The constants below are derived by just constant folding the terms dependent on CosThetaMax=1/7
+	const float Fc = Pow5(1 - VoH);
+	const float K = 49.0 / 46656.0;
+	float3 b = (K - K * F82) * (7776.0 + 9031.0 * F0);
+	return saturate(F0 + Fc * ((1 - F0) - b * (VoH - VoH * VoH)));
+}
+
 float3 F_Fresnel( float3 SpecularColor, float VoH )
 {
-	float3 SpecularColorSqrt = sqrt( clamp( float3(0, 0, 0), float3(0.99, 0.99, 0.99), SpecularColor ) );
+	float3 SpecularColorSqrt = sqrt( clamp(SpecularColor, float3(0, 0, 0), float3(0.99, 0.99, 0.99) ) );
 	float3 n = ( 1 + SpecularColorSqrt ) / ( 1 - SpecularColorSqrt );
 	float3 g = sqrt( n*n + VoH*VoH - 1 );
 	return 0.5 * Square( (g - VoH) / (g + VoH) ) * ( 1 + Square( ((g+VoH)*VoH - 1) / ((g-VoH)*VoH + 1) ) );
