Fixed: Use last frame eye adaptation constant instead of sampling EyeAdaptation buffer directly in materials to avoid issue in a MALI GL drivers (r9p0)

diff --git a/Engine/Shaders/Private/EyeAdaptationCommon.ush b/Engine/Shaders/Private/EyeAdaptationCommon.ush
--- a/Engine/Shaders/Private/EyeAdaptationCommon.ush
+++ b/Engine/Shaders/Private/EyeAdaptationCommon.ush
@@ -46,7 +46,10 @@ float EyeAdaptationLookup()
 	#if SHADING_PATH_DEFERRED
 		return EyeAdaptationLookupTexture2D(EyeAdaptationStruct.EyeAdaptationTexture).x;
 	#elif SHADING_PATH_MOBILE
-		return EyeAdaptationLookupBuffer(EyeAdaptationStruct.EyeAdaptationBuffer).x;
+		// Unofortunatelly on older MALI GL drivers (r9p0), passing this buffer into 'log' or other functions, fails with this error: Function call discards 'readonly' access qualifier
+		// So we have to use last frame eye adaptation value as a workaround on mobile
+		// return EyeAdaptationLookupBuffer(EyeAdaptationBuffer).x;
+		return ResolvedView.PreExposure;
 	#else
 		#error Unknown shading path.
 	#endif
diff --git a/Engine/Shaders/Private/MobileBasePassPixelShader.usf b/Engine/Shaders/Private/MobileBasePassPixelShader.usf
--- a/Engine/Shaders/Private/MobileBasePassPixelShader.usf
+++ b/Engine/Shaders/Private/MobileBasePassPixelShader.usf
@@ -561,7 +561,9 @@ void Main(
 	float DirectionalLightShadow = 1.0f;
 
 	// Directional light
+#if !MATERIAL_SHADINGMODEL_UNLIT
 	AccumulateDirectionalLighting(GBuffer, MaterialParameters.WorldPosition_CamRelative, CameraVector, MaterialParameters.ScreenPosition, SvPosition, DynamicShadowFactors, DirectionalLightShadow, DirectLighting);
+#endif
 
 	float2 LocalPosition = SvPosition.xy - ResolvedView.ViewRectMin.xy;
 	uint GridIndex = ComputeLightGridCellIndex(uint2(LocalPosition.x, LocalPosition.y), SvPosition.w);
