Prevent issues in FRayTracingInstanceDescriptor due to out-of-range values by masking bits.
- Otherwise they can end up overwriting InstanceMask/Flags.
- Also assert when setting up instance descriptors on CPU.

diff --git a/Engine/Shaders/Private/RayTracing/RayTracingCommon.ush b/Engine/Shaders/Private/RayTracing/RayTracingCommon.ush
--- a/Engine/Shaders/Private/RayTracing/RayTracingCommon.ush
+++ b/Engine/Shaders/Private/RayTracing/RayTracingCommon.ush
@@ -213,8 +213,8 @@ uint GetRayTracingInstanceDescriptorFlags(FRayTracingInstanceDescriptor Desc)
 FRayTracingInstanceDescriptor BuildPlatformRayTracingInstanceDesc(uint InstanceMask, uint InstanceId, uint InstanceFlags, uint InstanceContributionToHitGroupIndex, float3x4 LocalToWorld, uint2 AccelerationStructureAddress)
 {
 	FRayTracingInstanceDescriptor Desc;
-	Desc.InstanceIdAndMask = (InstanceMask << 24) | InstanceId;
-	Desc.InstanceContributionToHitGroupIndexAndFlags = (InstanceFlags << 24) | InstanceContributionToHitGroupIndex;
+	Desc.InstanceIdAndMask = (InstanceMask << 24) | (InstanceId & 0x00FFFFFF);
+	Desc.InstanceContributionToHitGroupIndexAndFlags = (InstanceFlags << 24) | (InstanceContributionToHitGroupIndex & 0x00FFFFFF);
 	Desc.AccelerationStructure = AccelerationStructureAddress;
 	SetRayTracingInstanceTransform(Desc, LocalToWorld);
 
diff --git a/Engine/Source/Runtime/Renderer/Private/RayTracing/RayTracingInstanceBufferUtil.cpp b/Engine/Source/Runtime/Renderer/Private/RayTracing/RayTracingInstanceBufferUtil.cpp
--- a/Engine/Source/Runtime/Renderer/Private/RayTracing/RayTracingInstanceBufferUtil.cpp
+++ b/Engine/Source/Runtime/Renderer/Private/RayTracing/RayTracingInstanceBufferUtil.cpp
@@ -234,6 +234,9 @@ void FillRayTracingInstanceUploadBuffer(
 				InstanceDesc.InstanceContributionToHitGroupIndex = SceneInitializer.SegmentPrefixSum[SceneInstanceIndex] * SceneInitializer.ShaderSlotsPerGeometrySegment;
 				InstanceDesc.bApplyLocalBoundsTransform = SceneInstance.bApplyLocalBoundsTransform;
 
+				checkf(InstanceDesc.InstanceId <= 0xFFFFFF, TEXT("InstanceId must fit in 24 bits."));
+				checkf(InstanceDesc.InstanceContributionToHitGroupIndex <= 0xFFFFFF, TEXT("InstanceContributionToHitGroupIndex must fit in 24 bits."));
+
 				if (!SceneInstance.ActivationMask.IsEmpty() && (SceneInstance.ActivationMask[TransformIndex / 32] & (1 << (TransformIndex % 32))) == 0)
 				{
 					// Set flag for deactivated instances
