Fix potential artifacts when SMRT crosses a clipmap boundary and the new sample depth is outside the Z range of the original clipmap.
Most common repro is when sampling against the clear value with a shadow being cast on a non-shadow-casting receiver.

#localization none
#tests Editor PIE, EngineTests, preflight cooks

diff --git a/Engine/Shaders/Private/VirtualShadowMaps/VirtualShadowMapSMRTTemplate.ush b/Engine/Shaders/Private/VirtualShadowMaps/VirtualShadowMapSMRTTemplate.ush
--- a/Engine/Shaders/Private/VirtualShadowMaps/VirtualShadowMapSMRTTemplate.ush
+++ b/Engine/Shaders/Private/VirtualShadowMaps/VirtualShadowMapSMRTTemplate.ush
@@ -28,8 +28,14 @@ FSMRTResult SMRTRayCast(
 	int NumSteps,
 	float StepOffset)
 {
-	float DepthHistory = -1;
-	float DepthHistoryTime = -1;
+	// NOTE: Need a magic initializer that we can detect. Would be natural to use DepthHistoryTime for this,
+	// but we want that register to disappear (DCE) when SMRT_EXTRAPOLATE_SLOPE is disabled.
+	// DepthHistory *can* be negative though when crossing clipmap boundaries, as the depth range expands.
+	// In practice it will only ever be down to -2 or so.
+	const float DepthHistoryNotSet = -10000.0f;
+	float DepthHistory = DepthHistoryNotSet;
+
+	float DepthHistoryTime = -1.0f;
 	float DepthSlope = 0;
 
 	const float TimeScale = -1.0f / NumSteps;
@@ -59,7 +65,7 @@ FSMRTResult SMRTRayCast(
 		if (Sample.bValid)
 		{
 			const float SampleDepth = Sample.SampleDepth;
-			if (DepthHistory < 0)
+			if (DepthHistory == DepthHistoryNotSet)
 			{
 				// First valid sample we've seen. Do a regular depth compare.
 				DepthHistory = SampleDepth;
