VCloud - remove pixel position bias not needed when sampling blue noise since the texture will repeat naturally anyway.

diff --git a/Engine/Shaders/Private/VolumetricCloud.usf b/Engine/Shaders/Private/VolumetricCloud.usf
--- a/Engine/Shaders/Private/VolumetricCloud.usf
+++ b/Engine/Shaders/Private/VolumetricCloud.usf
@@ -742,11 +742,10 @@ void MainCommon(in FMaterialPixelParameters MaterialParameters, in float4 SvPosi
 		float t = TMin + 0.5 * StepT;
 		if (RenderVolumetricCloudParameters.IsReflectionRendering == 0)
 		{
-		#if 0
-			// Using simple ALU based noise as it seems to work the best now. Mod8 is also used to stabilise the result when the amount of sample cannot converge to the correct integration.
+		#if 0 // Simple noise
 			t = TMin + (float(Rand3DPCG16(int3(SvPosition.xy, View.StateFrameIndexMod8)).x) * rcp(65536.0)) * StepT;
-		#else
-			uint2 FullResPixelCoord = SvPosition.xy * RenderVolumetricCloudParameters.TracingCoordToFullResPixelCoordScaleBias.xy + RenderVolumetricCloudParameters.TracingCoordToFullResPixelCoordScaleBias.xy;
+		#else // Blue noise works best
+			uint2 FullResPixelCoord = SvPosition.xy * RenderVolumetricCloudParameters.TracingCoordToFullResPixelCoordScaleBias.xy; // bias is not needed, noise will still repeat as needed
 			t = TMin + BlueNoiseScalar(FullResPixelCoord, View.StateFrameIndexMod8) * StepT;
 		#endif
 		}
