Revert PostprocessSubsurface (screen space) flag disabling the visual effect of `subsurface` shading model.

diff --git a/Engine/Shaders/Private/BasePassPixelShader.usf b/Engine/Shaders/Private/BasePassPixelShader.usf
--- a/Engine/Shaders/Private/BasePassPixelShader.usf
+++ b/Engine/Shaders/Private/BasePassPixelShader.usf
@@ -879,10 +879,6 @@ void FPixelShaderInOut_MainPS(
 	if (ShadingModel == SHADINGMODELID_SUBSURFACE || ShadingModel == SHADINGMODELID_PREINTEGRATED_SKIN || ShadingModel == SHADINGMODELID_SUBSURFACE_PROFILE || ShadingModel == SHADINGMODELID_TWOSIDED_FOLIAGE || ShadingModel == SHADINGMODELID_CLOTH || ShadingModel == SHADINGMODELID_EYE)
 	{
 		float4 SubsurfaceData = GetMaterialSubsurfaceData(PixelMaterialInputs);
-		
-		// If the subsurface feature is disabled, the subsurface color should be black
-		// This will also ensure that the subsurface profile does not contribute when subsurface scattering is disabled
-		SubsurfaceData *= View.bSubsurfacePostprocessEnabled;
 
 		if (false) // Dummy if to make the ifdef logic play nicely
 		{
diff --git a/Engine/Shaders/Private/MobileBasePassPixelShader.usf b/Engine/Shaders/Private/MobileBasePassPixelShader.usf
--- a/Engine/Shaders/Private/MobileBasePassPixelShader.usf
+++ b/Engine/Shaders/Private/MobileBasePassPixelShader.usf
@@ -331,10 +331,6 @@ void Main(
 	{
 		half4 SubsurfaceData = GetMaterialSubsurfaceData(PixelMaterialInputs);
 
-		// If the subsurface feature is disabled, the subsurface color should be black
-		// This will also ensure that the subsurface profile does not contribute when subsurface scattering is disabled
-		SubsurfaceData *= View.bSubsurfacePostprocessEnabled;
-
 		if (ShadingModelID == SHADINGMODELID_SUBSURFACE || ShadingModelID == SHADINGMODELID_PREINTEGRATED_SKIN || ShadingModelID == SHADINGMODELID_TWOSIDED_FOLIAGE)
 		{
 			SubsurfaceColor = SubsurfaceData.rgb * ResolvedView.DiffuseOverrideParameter.w + ResolvedView.DiffuseOverrideParameter.xyz;
