Strata - fixed some wanring related to HQ normals.

diff --git a/Engine/Shaders/Private/Strata/Strata.ush b/Engine/Shaders/Private/Strata/Strata.ush
--- a/Engine/Shaders/Private/Strata/Strata.ush
+++ b/Engine/Shaders/Private/Strata/Strata.ush
@@ -1822,10 +1822,11 @@ uint StrataLoadUint1AtIndexOffset(in FStrataMaterialContainer StrataBuffer, in F
 
 STRATA_TOP_LAYER_TYPE StrataLoadTopLayerData(FStrataTopLayerDataContainer InTopLayerContainer, inout FStrataAddressing StrataAddressing)
 {
+	// Top layer container is always uint2 so we cast to XXX to avoid warnings.
 	#if STRATA_MATERIALCONTAINER_IS_WRITABLE
-	return InTopLayerContainer[StrataAddressing.PixelCoords];
+	return (STRATA_TOP_LAYER_TYPE)InTopLayerContainer[StrataAddressing.PixelCoords];
 	#else
-	return InTopLayerContainer.Load(uint3(StrataAddressing.PixelCoords, 0));
+	return (STRATA_TOP_LAYER_TYPE)InTopLayerContainer.Load(uint3(StrataAddressing.PixelCoords, 0));
 	#endif
 }
 
