Strata - Fixed debug probes to handle strata MRT count.

diff --git a/Engine/Shaders/Private/DebugProbes.usf b/Engine/Shaders/Private/DebugProbes.usf
--- a/Engine/Shaders/Private/DebugProbes.usf
+++ b/Engine/Shaders/Private/DebugProbes.usf
@@ -81,9 +81,10 @@ void WriteOutMaterial(
 	in float3 Normal
 #if STRATA_ENABLED
 	, inout float4 OutGBufferA
-	, inout uint OutGBufferB
-	, inout uint OutGBufferC
-	, inout STRATA_TOP_LAYER_TYPE OutGBufferD
+	, inout uint StrataMat0
+	, inout uint StrataMat1
+	, inout uint StrataMat2
+	, inout STRATA_TOP_LAYER_TYPE StrataTopLayerData
 	, inout uint OutGBufferE
 #else
 	, inout float4 OutGBufferA
@@ -97,6 +98,10 @@ void WriteOutMaterial(
 {
 #if STRATA_ENABLED
 
+#if STRATA_BASE_PASS_MRT_OUTPUT_COUNT!=3
+#error Please update debug probe MRT code.
+#endif
+
 	FStrataPixelHeader StrataPixelHeader = InitialiseStrataPixelHeader();
 	StrataPixelHeader.StrataTree = GetInitialisedStrataTree();
 	StrataPixelHeader.IrradianceAO = InitIrradianceAndOcclusion();
@@ -162,11 +167,12 @@ void WriteOutMaterial(
 		SSSData, TopLayerData, OpaqueRoughRefractionData);
 
 	// Write out MRTs. Note: this only works for deferred with basepassvelocity off
-	OutGBufferA = 0;
-	OutGBufferB = RWStrataMaterialContainer.MaterialRenderTargets[0];
-	OutGBufferC = RWStrataMaterialContainer.MaterialRenderTargets[1];
-	OutGBufferD = StrataPackTopLayerData(TopLayerData);
-	OutGBufferE = 0;
+	OutGBufferA			= 0;
+	StrataMat0			= RWStrataMaterialContainer.MaterialRenderTargets[0];
+	StrataMat1			= RWStrataMaterialContainer.MaterialRenderTargets[1];
+	StrataMat2			= RWStrataMaterialContainer.MaterialRenderTargets[2];
+	StrataTopLayerData	= StrataPackTopLayerData(TopLayerData);
+	OutGBufferE			= 0;
 
 #else // STRATA_ENABLED
 
@@ -197,11 +203,12 @@ void MainPS(
 	, out float4 OutColor		: SV_Target0
 #if PERMUTATION_PASS == RENDER_BASEPASS
 	#if STRATA_ENABLED
-	, out float4 OutGBufferA				: SV_Target1
-	, out uint OutGBufferB					: SV_Target2
-	, out uint OutGBufferC					: SV_Target3
-	, out STRATA_TOP_LAYER_TYPE OutGBufferD : SV_Target4
-	, out uint OutGBufferE					: SV_Target5
+	, out float4 OutGBufferA						: SV_Target2
+	, out uint StrataMat0							: SV_Target3
+	, out uint StrataMat1							: SV_Target4
+	, out uint StrataMat2							: SV_Target5
+	, out STRATA_TOP_LAYER_TYPE StrataTopLayerData	: SV_Target6
+	, out uint OutGBufferE							: SV_Target7
 	#else
 	, out float4 OutGBufferA	: SV_Target1
 	, out float4 OutGBufferB	: SV_Target2
@@ -217,11 +224,20 @@ void MainPS(
 {
 	OutColor = 0.f;
 #if PERMUTATION_PASS == RENDER_BASEPASS
+#if STRATA_ENABLED
+	OutGBufferA = 0;
+	StrataMat0 = 0;
+	StrataMat1 = 0;
+	StrataMat2 = 0;
+	StrataTopLayerData = 0;
+	OutGBufferE = 0;
+#else
 	OutGBufferA = 0;
 	OutGBufferB = 0;
 	OutGBufferC = 0;
 	OutGBufferD = 0;
 	OutGBufferE = 0;
+#endif
 #endif
 
 	ResolvedView = ResolveView();
@@ -289,11 +305,21 @@ void MainPS(
 				Roughness,
 				IndirectIrradianceAndAO,
 				WorldNormal,
+			#if STRATA_ENABLED
+				OutGBufferA,
+				StrataMat0,
+				StrataMat1,
+				StrataMat2,
+				StrataTopLayerData,
+				OutGBufferE
+			#else
 				OutGBufferA,
 				OutGBufferB,
 				OutGBufferC,
 				OutGBufferD,
-				OutGBufferE);
+				OutGBufferE
+			#endif
+			);
 		}
 	#endif
 
@@ -491,12 +517,28 @@ void MainPS(
 			Roughness,
 			IndirectIrradianceAndAO,
 			WorldNormal,
+		#if STRATA_ENABLED
+			OutGBufferA,
+			StrataMat0,
+			StrataMat1,
+			StrataMat2,
+			StrataTopLayerData,
+			OutGBufferE
+		#else
 			OutGBufferA,
 			OutGBufferB,
 			OutGBufferC,
 			OutGBufferD,
-			OutGBufferE);
+			OutGBufferE
+		#endif
+		);
+	}
+#if STRATA_ENABLED
+	else
+	{
+		discard;
 	}
+#endif
 
 	OutColor = 0.0f;
 
