Fix bad directional light when it is parallel to (1,0,0) with a non-zero source radius.

diff --git a/Engine/Shaders/Private/RayTracing/RayTracingDirectionalLight.ush b/Engine/Shaders/Private/RayTracing/RayTracingDirectionalLight.ush
--- a/Engine/Shaders/Private/RayTracing/RayTracingDirectionalLight.ush
+++ b/Engine/Shaders/Private/RayTracing/RayTracingDirectionalLight.ush
@@ -50,7 +50,7 @@ void GenerateDirectionalLightOcclusionRay(
 	float3 LightDirection = LightParameters.Direction;
 	float3 N = LightDirection;
 	float3 dPdu = float3(1, 0, 0);
-	if (dot(N, dPdu) != 0)
+	if (abs(dot(N, dPdu)) != 1)
 	{
 		dPdu = cross(N, dPdu);
 	}
