Add initial mobile support for Strata.

This CL adds
* Forward shading support (deferred support will be added later on)
* Slab/Eye/Hair/SLW support
* Directional/Local/Env/Planar lighting

There are a couple of followup to the initial CL:
* Fix incorrect env. specular contribution which is too low to what is supposed to be
* Add SH lighting support
* Merge StrataForwardLighting and StrataMobileFowardLighting if possible as there are a lot of overlap

diff --git a/Engine/Shaders/Private/DeferredLightingCommon.ush b/Engine/Shaders/Private/DeferredLightingCommon.ush
--- a/Engine/Shaders/Private/DeferredLightingCommon.ush
+++ b/Engine/Shaders/Private/DeferredLightingCommon.ush
@@ -252,9 +252,11 @@ float GetLocalLightAttenuation(
 	if (LightData.bInverseSquared)
 	{
 		LightMask = Square( saturate( 1 - Square( DistanceSqr * Square(LightData.InvRadius) ) ) );
-#if SHADING_PATH_MOBILE
+		// This extra attenuation has been added for supporting existing 'legacy' shading model on mobile. 
+		// This is not needed for Strata which unifies all lighting paths
+		#if SHADING_PATH_MOBILE && !STRATA_ENABLED
 		LightMask *= 1.0f / (DistanceSqr + 1.0f);
-#endif
+		#endif
 	}
 	else
 	{
diff --git a/Engine/Shaders/Private/LightData.ush b/Engine/Shaders/Private/LightData.ush
--- a/Engine/Shaders/Private/LightData.ush
+++ b/Engine/Shaders/Private/LightData.ush
@@ -192,6 +192,12 @@ uint UnpackLightingChannelMask(uint In)
 	return (In >> 8) & LIGHTING_CHANNEL_MASK;
 }
 
+uint UnpackLightingChannelMask(FLocalLightData In)
+{
+	const uint LightTypeAndPackedShadowMapChannelMask = asuint(In.LightDirectionAndShadowMask.w);
+	return UnpackLightingChannelMask(LightTypeAndPackedShadowMapChannelMask);
+}
+
 uint UnpackLightType(uint InShadowMapChannelMaskPacked)
 {
 	return (InShadowMapChannelMaskPacked >> 16) & 0x3;
diff --git a/Engine/Shaders/Private/MobileBasePassPixelShader.usf b/Engine/Shaders/Private/MobileBasePassPixelShader.usf
--- a/Engine/Shaders/Private/MobileBasePassPixelShader.usf
+++ b/Engine/Shaders/Private/MobileBasePassPixelShader.usf
@@ -12,11 +12,15 @@
 
 #define ForwardLightData MobileBasePass.Forward
 #define ReflectionStruct MobileBasePass.ReflectionsParameters
+#define StrataStruct MobileBasePass.Strata
 
 #ifndef MOBILE_QL_FORCE_DISABLE_PREINTEGRATEDGF
 #define MOBILE_QL_FORCE_DISABLE_PREINTEGRATEDGF 0
 #endif
 
+#define UseBasePassSkylightDiffuse  (1)
+#define UseBasePassSkylightSpecular (MobileBasePass.ReflectionsParameters.SkyLightParameters.y)
+
 #define MOBILE_USE_PREINTEGRATED_GF (MATERIAL_USE_PREINTEGRATED_GF && !MOBILE_QL_FORCE_DISABLE_PREINTEGRATEDGF)
 
 //use preintegrated GF lut for simple IBL
@@ -48,6 +52,12 @@
 	#define SINGLE_LAYER_WATER_SHADING_QUALITY 2
 #endif
 
+// Enable Strata. This define & include need to be defined before certains includes (i.e., DBufferDecalShared which uses them internally)
+#if !MATERIAL_IS_STRATA && STRATA_ENABLED
+#undef STRATA_ENABLED
+#define STRATA_ENABLED 0
+#endif
+
 #ifndef PROJECT_MOBILE_ENABLE_MOVABLE_SPOTLIGHT_SHADOWS
 #define PROJECT_MOBILE_ENABLE_MOVABLE_SPOTLIGHT_SHADOWS 0
 #endif
@@ -68,6 +78,14 @@
 #define SUPPORT_SPOTLIGHTS_SHADOW (MATERIALBLENDING_SOLID || MATERIALBLENDING_MASKED) && PROJECT_MOBILE_ENABLE_MOVABLE_SPOTLIGHT_SHADOWS
 #define DEFERRED_SHADING_PATH (MOBILE_DEFERRED_SHADING && ((MATERIALBLENDING_SOLID || MATERIALBLENDING_MASKED) && !MATERIAL_SHADINGMODEL_SINGLELAYERWATER))
 
+// Force FORWARD_SHADING on mobile when MOBILE_DEFERRED_SHADING is disabled
+#ifndef FORWARD_SHADING
+#define FORWARD_SHADING (!MOBILE_DEFERRED_SHADING)
+#elif defined(FORWARD_SHADING)
+#undef  FORWARD_SHADING
+#define FORWARD_SHADING (!MOBILE_DEFERRED_SHADING)
+#endif
+
 // SceneDepthAux always enabled for IOS, on other platforms only enabled when in forward and MobileHDR=true
 #define HAS_SCENE_DEPTH_AUX_OUTPUT (USE_SCENE_DEPTH_AUX && (MATERIALBLENDING_SOLID || MATERIALBLENDING_MASKED))
 
@@ -94,6 +112,19 @@
 	#include "SingleLayerWaterShading.ush"
 #endif
 
+// STRATA_TODO: Shared with BasePassPixelShader.usf
+#define MATERIAL_STRATA_OPAQUE_PRECOMPUTED_LIGHTING 0
+
+#if STRATA_TRANSLUCENT_FORWARD || STRATA_FORWARD_SHADING || MATERIAL_STRATA_OPAQUE_PRECOMPUTED_LIGHTING || STRATA_MATERIAL_EXPORT_EXECUTED
+#include "/Engine/Private/Strata/StrataEvaluation.ush"
+#endif
+#if STRATA_TRANSLUCENT_FORWARD || STRATA_FORWARD_SHADING
+#include "/Engine/Private/Strata/StrataMobileForwardLighting.ush"
+#endif
+#if MATERIAL_STRATA_OPAQUE_PRECOMPUTED_LIGHTING || STRATA_MATERIAL_EXPORT_EXECUTED
+#include "/Engine/Private/Strata/StrataExport.ush"
+#endif
+
 half3 FrameBufferBlendOp(half4 Source)
 {
 	half4 Dest = half4 (0,0,0,0);
@@ -304,6 +335,17 @@ void Main(
 	GetMaterialCoverageAndClipping(MaterialParameters, PixelMaterialInputs);
 #endif
 
+	half MaterialAO = GetMaterialAmbientOcclusion(PixelMaterialInputs);
+#if APPLY_AO
+	{
+		half4 GatheredAmbientOcclusion = Texture2DSample(AmbientOcclusionTexture, AmbientOcclusionSampler, SvPositionToBufferUV(SvPosition));
+		MaterialAO *= lerp(1.0, GatheredAmbientOcclusion.r, AmbientOcclusionStaticFraction);
+	}
+#endif
+
+	half Opacity = GetMaterialOpacity(PixelMaterialInputs);
+	
+#if !STRATA_ENABLED
 	// Store the results in local variables and reuse instead of calling the functions multiple times.
 	half3 BaseColor = GetMaterialBaseColor(PixelMaterialInputs);
 	half Metallic = GetMaterialMetallic(PixelMaterialInputs);
@@ -318,15 +360,13 @@ void Main(
 	half Roughness = max(0.015625f, GetMaterialRoughness(PixelMaterialInputs));
 	half Anisotropy = GetMaterialAnisotropy(PixelMaterialInputs);
 	uint ShadingModelID = GetMaterialShadingModel(PixelMaterialInputs);
-	half Opacity = GetMaterialOpacity(PixelMaterialInputs);
-	half MaterialAO = GetMaterialAmbientOcclusion(PixelMaterialInputs);
 	
 	// If we don't use this shading model the color should be black (don't generate shader code for unused data, don't do indirectlighting cache lighting with this color).
 	float3 SubsurfaceColor = 0;
 	// 0..1, SubsurfaceProfileId = int(x * 255)
 	float SubsurfaceProfile = 0;
 
-#if (MATERIAL_SHADINGMODEL_SUBSURFACE || MATERIAL_SHADINGMODEL_PREINTEGRATED_SKIN || MATERIAL_SHADINGMODEL_SUBSURFACE_PROFILE || MATERIAL_SHADINGMODEL_TWOSIDED_FOLIAGE || MATERIAL_SHADINGMODEL_CLOTH || MATERIAL_SHADINGMODEL_EYE) 
+	#if (MATERIAL_SHADINGMODEL_SUBSURFACE || MATERIAL_SHADINGMODEL_PREINTEGRATED_SKIN || MATERIAL_SHADINGMODEL_SUBSURFACE_PROFILE || MATERIAL_SHADINGMODEL_TWOSIDED_FOLIAGE || MATERIAL_SHADINGMODEL_CLOTH || MATERIAL_SHADINGMODEL_EYE) 
 	if (ShadingModelID == SHADINGMODELID_SUBSURFACE || ShadingModelID == SHADINGMODELID_PREINTEGRATED_SKIN || ShadingModelID == SHADINGMODELID_SUBSURFACE_PROFILE || ShadingModelID == SHADINGMODELID_TWOSIDED_FOLIAGE || ShadingModelID == SHADINGMODELID_CLOTH || ShadingModelID == SHADINGMODELID_EYE)
 	{
 		half4 SubsurfaceData = GetMaterialSubsurfaceData(PixelMaterialInputs);
@@ -342,14 +382,45 @@ void Main(
 
 		SubsurfaceProfile = SubsurfaceData.a;
 	}
+	#endif
+#endif // !STRATA_ENABLED
+
+
+#if STRATA_ENABLED && !STRATA_OPTIMIZED_UNLIT
+	FStrataData StrataData = PixelMaterialInputs.FrontMaterial;
+
+	// Initialise a Strata header with normal in registers
+	FStrataPixelHeader StrataPixelHeader = InitialiseStrataPixelHeader();
+	StrataPixelHeader.StrataTree = MaterialParameters.StrataTree;
+	StrataPixelHeader.BSDFCount = MaterialParameters.StrataTree.BSDFCount;
+	StrataPixelHeader.SharedLocalBases = MaterialParameters.SharedLocalBases;
+	StrataPixelHeader.IrradianceAO = InitIrradianceAndOcclusion();
+	StrataPixelHeader.IrradianceAO.MaterialAO = MaterialAO;
+	SetCastContactShadow(StrataPixelHeader, GetPrimitiveData(MaterialParameters).Flags & PRIMITIVE_SCENE_DATA_FLAG_HAS_CAST_CONTACT_SHADOW);
+	SetDynamicIndirectShadowCasterRepresentation(StrataPixelHeader, GetPrimitiveData(MaterialParameters).Flags & PRIMITIVE_SCENE_DATA_FLAG_HAS_CAPSULE_REPRESENTATION);
+#endif
+
+#if STRATA_ENABLED && STRATA_INLINE_SINGLELAYERWATER
+	SetMaterialMode(StrataPixelHeader, HEADER_MATERIALMODE_SLWATER);
+
+	// Override GBuffer data with Strata SLW water BSDF to run forward shadfing code.
+	// STRATA_TODO: run the shading through a strata path? (by adding a special BSDF?)
+	SanitizeStrataSingleLayerWater(StrataPixelHeader.StrataTree.BSDFs[0]);
+	FStrataBSDF SLWBSDF = StrataPixelHeader.StrataTree.BSDFs[0];
+
+	half3 BaseColor = SLW_BASECOLOR(SLWBSDF);
+	half Metallic = SLW_METALLIC(SLWBSDF);
+	half Specular = SLW_SPECULAR(SLWBSDF);
+	half Roughness = SLW_ROUGHNESS(SLWBSDF);
+	Opacity = SLW_TOPMATERIALOPACITY(SLWBSDF);
+	MaterialParameters.WorldNormal = normalize(StrataPixelHeader.SharedLocalBases.Normals[BSDF_GETSHAREDLOCALBASISID(SLWBSDF)]);
+	half Anisotropy = 0.0f;
+	uint ShadingModelID = MATERIAL_SHADINGMODEL_SINGLELAYERWATER;
+	float3 SubsurfaceColor = 0;
+	float SubsurfaceProfile = 0;
 #endif
 
-#if APPLY_AO
-	half4 GatheredAmbientOcclusion = Texture2DSample(AmbientOcclusionTexture, AmbientOcclusionSampler, SvPositionToBufferUV(SvPosition));
-
-	MaterialAO *= lerp(1.0, GatheredAmbientOcclusion.r, AmbientOcclusionStaticFraction);
-#endif
-
+#if !STRATA_ENABLED || STRATA_INLINE_SINGLELAYERWATER
 	const half BaseMaterialCoverageOverWater = Opacity;
 	const half WaterVisibility = 1.0 - BaseMaterialCoverageOverWater;
 	// Fade out diffuse as this will be handled by the single scattering lighting. when over the water surface.
@@ -444,6 +515,9 @@ void Main(
 	GBuffer.GBufferAO = AOMultiBounce(Luminance(GBuffer.SpecularColor), ShadingOcclusion.SpecOcclusion).g;
 #endif
 
+#endif // !STRATA_ENABLED || STRATA_INLINE_SINGLELAYERWATER
+
+#if !STRATA_ENABLED
 	
 #if MATERIAL_SHADINGMODEL_THIN_TRANSLUCENT
 	float TopMaterialCoverage = Opacity;
@@ -471,6 +545,8 @@ void Main(
 	GBuffer.BaseColor *= TopMaterialCoverage;
 #endif
 
+#endif // !STRATA_ENABLED
+
 	VTPageTableResult LightmapVTPageTableResult = (VTPageTableResult)0.0f;
 #if LIGHTMAP_VT_ENABLED
 	{
@@ -487,6 +563,180 @@ void Main(
 	StoreVirtualTextureFeedback(MaterialParameters.VirtualTextureFeedback, 0, LightmapVTPageTableResult.PackedRequest);
 #endif
 
+#if STRATA_ENABLED && !STRATA_INLINE_SINGLELAYERWATER
+	
+	float3 Color = 0;
+	float3 DualBlendSurfaceLuminancePostCoverage = 0.0f;
+	float3 DualBlendSurfaceTransmittancePreCoverage = 1.0f;
+	float DualBlendSurfaceCoverage = 1.0f;
+
+  #if !STRATA_OPTIMIZED_UNLIT
+
+	// Static shadow mask
+	#if GBUFFER_HAS_PRECSHADOWFACTOR
+	{
+		// Encode shadow mask only if the shadow mask is entirely non-zero and non-one
+	#if WRITES_PRECSHADOWFACTOR_ZERO
+		SetHasPrecShadowMask(StrataPixelHeader, false);
+		SetZeroPrecShadowMask(StrataPixelHeader, true);
+	#else
+	  #if ALLOW_STATIC_LIGHTING
+		const bool bAllZero = all(GBuffer.PrecomputedShadowFactors == 0);
+		const bool bAllOne  = all(GBuffer.PrecomputedShadowFactors == 1);
+		if (!bAllZero && !bAllOne)
+		{
+			SetHasPrecShadowMask(StrataPixelHeader, true);
+		}
+		else if (bAllZero)
+		{
+			SetHasPrecShadowMask(StrataPixelHeader, false);
+			SetZeroPrecShadowMask(StrataPixelHeader, true);
+		}
+		else if (bAllOne)
+	  #endif
+		{
+			SetHasPrecShadowMask(StrataPixelHeader, false);
+			SetZeroPrecShadowMask(StrataPixelHeader, false);
+		}
+	#endif
+	}
+	#endif
+
+
+
+	#if STRATA_INLINE_SHADING
+	// We must normalize each normal and tangent to avoid non normalised vectors due to per vertex interpolation or texture filtering,
+	// for the deferred (our packing relies on normalized normal) and forward (normals are going to be used as-is from registers) paths.
+	#if STRATA_FORCE_SINGLE_BSDF
+	uint i = 0;
+	if (StrataPixelHeader.SharedLocalBases.Count > 0)
+	#else
+	UNROLL
+	for (uint i = 0; i < StrataPixelHeader.SharedLocalBases.Count; ++i)
+	#endif // STRATA_FORCE_SINGLE_BSDF
+	{
+		StrataPixelHeader.SharedLocalBases.Normals[i] = normalize(StrataPixelHeader.SharedLocalBases.Normals[i]);
+		if (StrataGetSharedLocalBasisType(StrataPixelHeader.SharedLocalBases.Types, i) == STRATA_BASIS_TYPE_TANGENT)
+		{
+			StrataPixelHeader.SharedLocalBases.Tangents[i] = normalize(StrataPixelHeader.SharedLocalBases.Tangents[i]);
+		}
+	}
+	#endif // STRATA_INLINE_SHADING
+
+	FLightAccumulator DirectLighting = (FLightAccumulator)0;
+	half3 CameraVector = -MaterialParameters.CameraVector;
+	float2 ScreenUV = ScreenPositionToBufferUV(MaterialParameters.ScreenPosition);
+
+	#define STRATA_DEBUG_ENABLE_LIGHTING 1
+	#if STRATA_DEBUG_ENABLE_LIGHTING
+	if (StrataPixelHeader.BSDFCount > 0)
+	{
+		#if FULLY_ROUGH
+		const bool bForceFullyRough = true;
+		#else
+		const bool bForceFullyRough = View.RenderingReflectionCaptureMask > 0; // Can mobile capture reflection probe? Is this needed for mobile?
+		#endif
+		FStrataIntegrationSettings Settings = InitStrataIntegrationSettings(bForceFullyRough, StrataStruct.bRoughDiffuse, StrataStruct.PeelLayersAboveDepth);
+
+
+		// TODO: needed? Remove?
+		float3 VolumetricLightmapBrickTextureUVs = 0;
+
+		const float SceneDepth = MaterialParameters.ScreenPosition.w;;
+		const uint EyeIndex = 0;
+		Color = StrataMobileForwardLighting(
+			EyeIndex,
+			SvPosition,
+			Settings,
+			BasePassInterpolants,
+			Interpolants,
+			LightmapVTPageTableResult,
+			VolumetricLightmapBrickTextureUVs,
+			MaterialParameters,
+			SceneDepth,
+			ScreenUV,
+			StrataPixelHeader,
+			StrataData,
+			DualBlendSurfaceTransmittancePreCoverage,
+			DualBlendSurfaceCoverage);
+
+	#if STRATA_TRANSLUCENT_MATERIAL
+		DualBlendSurfaceLuminancePostCoverage = Color;
+		Color = 0.0f;
+		Opacity = 1.0f; // nullify following operation
+	#elif STRATA_FORWARD_SHADING
+		// Color unchanged for opaque materials
+		Opacity = 1.0f;
+	#else
+		#error Unhandled Strata forward shading mode
+	#endif
+	}
+	#endif // STRATA_DEBUG_ENABLE_LIGHTING
+
+  #else // STRATA_OPTIMIZED_UNLIT
+
+	// Unlit forces a single BSDF
+	FStrataBSDF UnlitBSDF = PixelMaterialInputs.FrontMaterial.InlinedBSDF;
+	DualBlendSurfaceLuminancePostCoverage = BSDF_GETEMISSIVE(UnlitBSDF);
+	DualBlendSurfaceTransmittancePreCoverage = UNLIT_TRANSMITTANCE(UnlitBSDF);
+	DualBlendSurfaceCoverage = UnlitBSDF.Coverage;
+	Color = 0.0f;
+	Opacity = 1.0f;
+
+  #endif // STRATA_OPTIMIZED_UNLIT
+
+	half4 Fogging = half4(0, 0, 0, 1);
+	#if USE_VERTEX_FOG
+	{
+	  #if PACK_INTERPOLANTS
+		Fogging = PackedInterpolants[0];
+	  #else
+		Fogging = BasePassInterpolants.VertexFog;
+	  #endif
+	}
+	#endif
+
+	// Compute final value
+	{
+	#if STRATA_OPTIMIZED_UNLIT
+		OutColor = float4(DualBlendSurfaceLuminancePostCoverage, 1);
+	#elif STRATA_OPAQUE_MATERIAL
+		OutColor = float4(Color, 1);
+	#elif STRATA_TRANSLUCENT_MATERIAL
+			
+	#if STRATA_LEGACY_PREMULT_ALPHA_OVERRIDE && !STRATA_OPTIMIZED_UNLIT
+		// This data patching is working because DualBlendSurfaceCoverage is always 1 when converting material from legacy to Strata with premultiplied alpha blending.
+		DualBlendSurfaceTransmittancePreCoverage= StrataData.PreMultipliedAlphaOverrideCoverage >= 0.0f ? 0.0f											: DualBlendSurfaceTransmittancePreCoverage;
+		DualBlendSurfaceCoverage				= StrataData.PreMultipliedAlphaOverrideCoverage >= 0.0f ? StrataData.PreMultipliedAlphaOverrideCoverage	: DualBlendSurfaceCoverage;
+	#endif
+
+		// Add fog luminance according to surfacecoverage and reduce surface luminance according to fog coverage.
+		float3 AdjustedDualBlendAdd = DualBlendSurfaceCoverage * Fogging.rgb + Fogging.a * DualBlendSurfaceLuminancePostCoverage;
+		// Fade the surface color transmittance out to 1 according to the surface coverage, and take into account the fog coverage to the surface.
+		float3 AdjustedDualBlendMul = lerp(1.0f, Fogging.a * DualBlendSurfaceTransmittancePreCoverage, DualBlendSurfaceCoverage);
+
+	#if STRATA_BLENDING_COLOREDTRANSMITTANCEONLY
+		// RETURN_COLOR not needed with modulative blending
+		half3 FoggedColor = lerp(float3(1, 1, 1), DualBlendSurfaceTransmittancePreCoverage, Fogging.aaa* DualBlendSurfaceCoverage);
+		OutColor = half4(FoggedColor, 1.0f);
+	#elif STRATA_BLENDING_ALPHAHOLDOUT
+		float AdjustedAlpha = saturate(1 - dot(AdjustedDualBlendMul, float3(1.0f, 1.0f, 1.0f) / 3.0f));
+		OutColor = half4(0.0f.xxx, AdjustedAlpha);
+	#else
+		// Pre multipled alpha blending
+		float AdjustedAlpha = saturate(1 - dot(AdjustedDualBlendMul, float3(1.0f, 1.0f, 1.0f) / 3.0f));
+		#if STRATA_USES_CONVERSION_FROM_LEGACY && MATERIALBLENDING_ADDITIVE
+			AdjustedAlpha = 0.0f;
+		#endif
+		// Pre multipled alpha blending
+		OutColor = half4(AdjustedDualBlendAdd, AdjustedAlpha);
+	#endif
+
+	#endif // STRATA_TRANSLUCENT_MATERIAL
+	}
+
+#else // STRATA_ENABLED
+
 	half IndirectIrradiance = 0;
 	half3 DiffuseIndirectLighting = 0;
 	half3 SubsurfaceIndirectLighting = 0;
@@ -621,6 +871,11 @@ void Main(
 
 #endif// DEFERRED_SHADING_PATH
 
+#endif // STRATA_ENABLED
+
+
+#if !STRATA_ENABLED || STRATA_INLINE_SINGLELAYERWATER
+
 #if !MATERIAL_SHADINGMODEL_UNLIT
 	half3 Color = DirectLighting.TotalLight;
 #else
@@ -673,6 +928,7 @@ void Main(
 			OutColor.a = uint((Mask + 0.25f) * 4.0f) / 4.0f;
 		#endif
 	#endif
+#endif // STRATA_ENABLED
 
 #if HAS_SCENE_DEPTH_AUX_OUTPUT
 	OutSceneDepthAux = SvPosition.z;
diff --git a/Engine/Shaders/Private/MobileLightingCommon.ush b/Engine/Shaders/Private/MobileLightingCommon.ush
--- a/Engine/Shaders/Private/MobileLightingCommon.ush
+++ b/Engine/Shaders/Private/MobileLightingCommon.ush
@@ -172,11 +172,11 @@ half MobileDirectionalLightCSM(float2 ScreenPosition, float SceneDepth, inout fl
 
 		ShadowMap = MobileShadowPCF(ShadowPosition.xy, Settings);
 
-		#if FADE_CSM
-			float Fade = saturate(SceneDepth * MobileDirectionalLight.DirectionalLightDistanceFadeMADAndSpecularScale.x + MobileDirectionalLight.DirectionalLightDistanceFadeMADAndSpecularScale.y);
-			// lerp out shadow based on fade params.
-			ShadowMap = lerp(ShadowMap, 1.0, Fade * Fade);
-		#endif
+	#if FADE_CSM
+		float Fade = saturate(SceneDepth * MobileDirectionalLight.DirectionalLightDistanceFadeMADAndSpecularScale.x + MobileDirectionalLight.DirectionalLightDistanceFadeMADAndSpecularScale.y);
+		// lerp out shadow based on fade params.
+		ShadowMap = lerp(ShadowMap, 1.0, Fade * Fade);
+	#endif
 	}
 #endif //ENABLE_MOBILE_CSM
 	return ShadowMap;
@@ -230,31 +230,31 @@ FLightAccumulator GetTranslucencySHLighting(FGBufferData GBuffer, FDeferredLight
 	Mobile Directional Light.
 ------------------------------------------------------------------------------*/
 
-void AccumulateDirectionalLighting(FGBufferData GBuffer, float3 TranslatedWorldPosition, half3 CameraVector, float4 ScreenPosition, float4 SvPosition, inout half4 DynamicShadowFactors, inout float OutDirectionalLightShadow, inout FLightAccumulator DirectLighting)
+FDeferredLightData GetDirectionalLightData(float4 ScreenPosition, float4 SvPosition, inout half4 OutDynamicShadowFactors, inout half OutDynamicShadowing)
 {
 #if USE_SHADOWMASKTEXTURE
 	half4 PreviewShadowMapChannelMask = half4(1.0f, 0.0f, 0.0f, 0.0f);
-	DynamicShadowFactors = Texture2DSample(MobileBasePass.ScreenSpaceShadowMaskTexture, MobileBasePass.ScreenSpaceShadowMaskSampler, SvPositionToBufferUV(SvPosition));
-	DynamicShadowFactors = DecodeLightAttenuation(DynamicShadowFactors);
+	OutDynamicShadowFactors = Texture2DSample(MobileBasePass.ScreenSpaceShadowMaskTexture, MobileBasePass.ScreenSpaceShadowMaskSampler, SvPositionToBufferUV(SvPosition));
+	OutDynamicShadowFactors = DecodeLightAttenuation(OutDynamicShadowFactors);
 		
 	PreviewShadowMapChannelMask = UnpackShadowMapChannelMask(MobileDirectionalLight.DirectionalLightShadowMapChannelMask >> 4);
 
-	half DynamicShadowing = dot(PreviewShadowMapChannelMask, DynamicShadowFactors);
+	OutDynamicShadowing = dot(PreviewShadowMapChannelMask, OutDynamicShadowFactors);
 #else
 
 	float  ShadowPositionZ = 0;
 	half ShadowMap = MobileDirectionalLightCSM(ScreenPosition.xy, ScreenPosition.w, ShadowPositionZ);
-	half DynamicShadowing = 1.0f;
+	OutDynamicShadowing = 1.0f;
 
-	#if !MOBILE_DEFERRED_SHADING && DIRECTIONAL_LIGHT_CSM && !MATERIAL_SHADINGMODEL_SINGLELAYERWATER
-		// Cascaded Shadow Map
-		if (IsCSMEnabled())
-		{		
-			DynamicShadowing = ShadowMap;
-		}
-	#elif MOBILE_DEFERRED_SHADING
-		DynamicShadowing = ShadowMap;
-	#endif
+   #if !MOBILE_DEFERRED_SHADING && DIRECTIONAL_LIGHT_CSM && !MATERIAL_SHADINGMODEL_SINGLELAYERWATER
+	// Cascaded Shadow Map
+	if (IsCSMEnabled())
+	{		
+		OutDynamicShadowing = ShadowMap;
+	}
+  #elif MOBILE_DEFERRED_SHADING
+	OutDynamicShadowing = ShadowMap;
+  #endif
 #endif
 
 	FDeferredLightData LightData = (FDeferredLightData)0;
@@ -267,6 +267,13 @@ void AccumulateDirectionalLighting(FGBufferData GBuffer, float3 TranslatedWorldP
 	LightData.HairTransmittance = InitHairTransmittanceData();
 	LightData.ShadowMapChannelMask = UnpackShadowMapChannelMask(MobileDirectionalLight.DirectionalLightShadowMapChannelMask);
 
+	return LightData;
+}
+
+void AccumulateDirectionalLighting(FGBufferData GBuffer, float3 TranslatedWorldPosition, half3 CameraVector, float4 ScreenPosition, float4 SvPosition, inout half4 DynamicShadowFactors, inout float OutDirectionalLightShadow, inout FLightAccumulator DirectLighting)
+{
+	half DynamicShadowing = 1.0f;
+	FDeferredLightData LightData = GetDirectionalLightData(ScreenPosition, SvPosition, DynamicShadowFactors, DynamicShadowing);
 	half4 LightAttenuation = half4(1, 1, DynamicShadowing, DynamicShadowing);
 
 #if TRANSLUCENCY_SH_LIGHTING
@@ -291,13 +298,14 @@ half3 GetMobileSkyLightReflection(half3 ReflectionVector, half Roughness, half C
 }
 #endif
 
-half3 GetImageBasedReflectionLighting(half3 ReflectionVector
-	, float3 TranslatedWorldPosition
-	, half Roughness
-	, half IndirectIrradiance
-	, half CompositeAlpha
-	, uint GridIndex
-)
+// Common helper function for evaluating reflection probes for mobile
+half3 GetImageBasedReflectionLighting_Mobile(
+	half3 ReflectionVector, 
+	float3 TranslatedWorldPosition, 
+	half Roughness, 
+	half IndirectIrradiance, 
+	half CompositeAlpha, 
+	uint GridIndex)
 {
 	half3 SpecularIBL = (half3)0.0f;
 
@@ -356,13 +364,58 @@ half3 GetImageBasedReflectionLighting(half3 ReflectionVector
 	return SpecularIBL;
 }
 
-void AccumulateReflection(FGBufferData GBuffer
-	, half3 CameraVector
-	, float3 TranslatedWorldPosition
-	, half3 ReflectionVector
-	, half IndirectIrradiance
-	, uint GridIndex
-	, inout FLightAccumulator DirectLighting)
+// Common helper function for evaluating planar reflection for mobile
+half3 GetPlanarReflectionbasedReflectionLighting_Mobile(float3 TranslatedWorldPosition, half3 WorldNormal, half Roughness, half3 InSpecularIBL)
+{
+	half3 Out = InSpecularIBL;
+#if ENABLE_PLANAR_REFLECTION || MATERIAL_PLANAR_FORWARD_REFLECTIONS
+	BRANCH
+	if (abs(dot(PlanarReflectionStruct.ReflectionPlane.xyz, 1)) > .0001f)
+	{
+		half4 PlanarReflection = GetPlanarReflection(TranslatedWorldPosition, WorldNormal, Roughness);
+		// Planar reflections win over reflection environment
+		Out = lerp(InSpecularIBL, PlanarReflection.rgb, PlanarReflection.a);
+	}
+#endif
+	return Out;
+}
+
+// Version used for Strata
+float3 GetImageBasedReflectionSpecular(
+	  half3 SpecularDirection, 
+	  half Roughness, 
+	  float3 TranslatedWorldPosition, 
+	  half3 WorldNormal, 
+	  FShadingOcclusion Occlusion, 
+	  half IndirectIrradiance, 
+	  uint GridIndex)
+{
+	half3 Out = 0;
+
+	// IBL reflection
+	Out = GetImageBasedReflectionLighting_Mobile(
+		SpecularDirection,
+		TranslatedWorldPosition,
+		Roughness,
+		IndirectIrradiance,
+		Occlusion.SpecOcclusion,
+		GridIndex);
+
+	// Planar reflection
+	Out = GetPlanarReflectionbasedReflectionLighting_Mobile(TranslatedWorldPosition, WorldNormal, Roughness, Out);
+
+	return Out;
+}
+
+// Version used for legacy shading
+void AccumulateReflection(
+	FGBufferData GBuffer, 
+	half3 CameraVector, 
+	float3 TranslatedWorldPosition, 
+	half3 ReflectionVector, 
+	half IndirectIrradiance, 
+	uint GridIndex, 
+	inout FLightAccumulator DirectLighting)
 {
 	half3 SpecularIBLLighting = (half3)0.0f;
 
@@ -381,7 +434,8 @@ void AccumulateReflection(FGBufferData GBuffer
 		SpecularOcclusion = GetSpecularOcclusion(NoV, RoughnessSq, SpecularOcclusion);
 	}
 
-	half3 SpecularIBL = GetImageBasedReflectionLighting(ReflectionVector
+	// IBL reflection (primary)
+	half3 SpecularIBL = GetImageBasedReflectionLighting_Mobile(ReflectionVector
 		, TranslatedWorldPosition
 		, GBuffer.Roughness
 		, IndirectIrradiance
@@ -389,15 +443,8 @@ void AccumulateReflection(FGBufferData GBuffer
 		, GridIndex
 	);
 
-#if ENABLE_PLANAR_REFLECTION || MATERIAL_PLANAR_FORWARD_REFLECTIONS
-	BRANCH
-	if (abs(dot(PlanarReflectionStruct.ReflectionPlane.xyz, 1)) > .0001f)
-	{
-		half4 PlanarReflection = GetPlanarReflection(TranslatedWorldPosition, GBuffer.WorldNormal, GBuffer.Roughness);
-		// Planar reflections win over reflection environment
-		SpecularIBL = lerp(SpecularIBL, PlanarReflection.rgb, PlanarReflection.a);
-	}
-#endif
+	// Planar reflection
+	SpecularIBL = GetPlanarReflectionbasedReflectionLighting_Mobile(TranslatedWorldPosition, GBuffer.WorldNormal, GBuffer.Roughness, SpecularIBL);
 
 	half3 DiffuseColor = GBuffer.DiffuseColor;
 	half3 SpecularColor = GBuffer.SpecularColor;
@@ -421,7 +468,8 @@ void AccumulateReflection(FGBufferData GBuffer
 	#endif
 		SpecularIBLLighting += SpecularIBL * LayerAttenuation * (SpecularColor * AB.x + AB.y * saturate(50 * SpecularColor.g) * (1 - ClearCoat));
 
-		SpecularIBL = GetImageBasedReflectionLighting(TopLayerR
+		// IBL reflection (secondary)
+		SpecularIBL = GetImageBasedReflectionLighting_Mobile(TopLayerR
 			, TranslatedWorldPosition
 			, ClearCoatRoughness
 			, IndirectIrradiance
@@ -441,6 +489,13 @@ void AccumulateReflection(FGBufferData GBuffer
 	LightAccumulator_AddSplit(DirectLighting, 0.0f, SpecularIBLLighting, 0.0f, 1.0f, false);
 }
 
+// Return true if the point is influenced by the light
+bool IsLightVisible(FLocalLightData InLight, float3 TranslatedWorldPosition)
+{
+	float InvLightRadiusSq = InLight.LightPositionAndInvRadius.w * InLight.LightPositionAndInvRadius.w;
+	return length2(TranslatedWorldPosition - InLight.LightPositionAndInvRadius.xyz) * InvLightRadiusSq <= 1.0f;
+}
+
 /*------------------------------------------------------------------------------
 	Mobile Local lights.
 ------------------------------------------------------------------------------*/
@@ -480,8 +535,7 @@ void AccumulateLightGridLocalLighting(const FCulledLightsGridData InLightGridDat
 
 		// extra-early out since we know light grid is sloppy and all lights in list are radial (have a range)
 		// appears useless
-		float invLightRadiusSq = LocalLight.LightPositionAndInvRadius.w * LocalLight.LightPositionAndInvRadius.w;
-		if (length2(TranslatedWorldPosition - LocalLight.LightPositionAndInvRadius.xyz) * invLightRadiusSq > 1.0f)
+		if (!IsLightVisible(LocalLight, TranslatedWorldPosition))
 		{
 			continue;
 		}
@@ -490,7 +544,7 @@ void AccumulateLightGridLocalLighting(const FCulledLightsGridData InLightGridDat
 
 #if !MOBILE_DEFERRED_LIGHTING
 		// Check lighting channels for mobile forward
-		if (((LightTypeAndPackedShadowMapChannelMask >> 8) & LightingChannelMask) == 0)
+		if ((UnpackLightingChannelMask(LocalLight) & LightingChannelMask) == 0)
 		{
 			continue;
 		}
diff --git a/Engine/Shaders/Private/Strata/StrataEvaluation.ush b/Engine/Shaders/Private/Strata/StrataEvaluation.ush
--- a/Engine/Shaders/Private/Strata/StrataEvaluation.ush
+++ b/Engine/Shaders/Private/Strata/StrataEvaluation.ush
@@ -72,8 +72,14 @@
 // Point lobe in off-specular peak direction
 float3 StrataGetOffSpecularPeakReflectionDir(float3 Normal, float3 ReflectionVector, float Roughness)
 {
+	// This function is used for both regular & mobile path. For performance reason, mobile 
+	// version uses (optionally) a simplifed version.
+#if !SHADING_PATH_MOBILE || MOBILE_HIGH_QUALITY_BRDF
 	float a = Square(Roughness);
 	return lerp(Normal, ReflectionVector, (1 - a) * (sqrt(1 - a) + a));
+#else
+	return ReflectionVector;
+#endif
 }
 
 ///////////////////////////////////////////////////////////////////////////////
diff --git a/Engine/Shaders/Private/Strata/StrataForwardLighting.ush b/Engine/Shaders/Private/Strata/StrataForwardLighting.ush
--- a/Engine/Shaders/Private/Strata/StrataForwardLighting.ush
+++ b/Engine/Shaders/Private/Strata/StrataForwardLighting.ush
@@ -218,7 +218,7 @@ float3 StrataForwardLighting(
 				if (bEnableSpecular)
 				{
 					// Evaluate specular parameters without bent normal
-					FStrataEnvLightResult StrataEnvLight = StrataEvaluateForEnvLight(StrataBSDFContext, bEnableSpecular, Settings);
+					FStrataEnvLightResult StrataEnvLight = StrataEvaluateForEnvLight(StrataBSDFContext, true/*bEnableSpecular*/, Settings);
 
 					// Evaluate specular lighting for the main lobe
 					if (any(StrataEnvLight.SpecularWeight > 0.0f))
@@ -344,6 +344,8 @@ float3 StrataForwardLighting(
 				////
 				//// Evaluate the single directional light selected for forward lighting.
 				////
+				BRANCH
+				if (DirectionalLightData.HasDirectionalLight)
 				{
 					float3 DirLightL = DirLightData.Direction;	// Already normalized
 					float3 ToDirLight = DirLightL;
@@ -353,36 +355,30 @@ float3 StrataForwardLighting(
 						DirLightMask = GetLocalLightAttenuation(TranslatedWorldPosition, DirLightData, ToDirLight, DirLightL);
 					}
 
-					bool bEvaluateHasBeenExecuted = false;
-					BRANCH
-					if (DirectionalLightData.HasDirectionalLight)
-					{
-						float DirectionalLightCloudShadow = 1.0f;
-					#if NEEDS_BASEPASS_CLOUD_SHADOW_INTERPOLATOR
-						DirectionalLightCloudShadow = BasePassInterpolants.VertexCloudShadow;
-					#endif
+					float DirectionalLightCloudShadow = 1.0f;
+				#if NEEDS_BASEPASS_CLOUD_SHADOW_INTERPOLATOR
+					DirectionalLightCloudShadow = BasePassInterpolants.VertexCloudShadow;
+				#endif
 
-						StrataUpdateBSDFContext(StrataBSDFContext, ToDirLight);
-
-						FStrataEvaluateResult BSDFEvaluate = (FStrataEvaluateResult)0;
-						Color += StrataForwardLightingCommon(
-							Dither,
-							Settings,
-							DirLightData,
-							ToDirLight,
-							DirLightMask,
-							DirLightAttenuation,
-							InitRectTexture(),
-							DirLightingChannelMask,
-							PrimitiveLightingChannelMask,
-							TranslatedWorldPosition,
-							SceneDepth,
-							LuminanceWeight(StrataBSDFContext, CurrentBSDF),
-							StrataPixelHeader,
-							StrataBSDFContext,
-							bEvaluateHasBeenExecuted,
-							BSDFEvaluate) * DirectionalLightCloudShadow;
-					}
+					StrataUpdateBSDFContext(StrataBSDFContext, ToDirLight);
+
+					FStrataEvaluateResult BSDFEvaluate = (FStrataEvaluateResult)0;
+					Color += StrataForwardLightingCommon(
+						Dither,
+						Settings,
+						DirLightData,
+						ToDirLight,
+						DirLightMask,
+						DirLightAttenuation,
+						InitRectTexture(),
+						DirLightingChannelMask,
+						PrimitiveLightingChannelMask,
+						TranslatedWorldPosition,
+						SceneDepth,
+						LuminanceWeight(StrataBSDFContext, CurrentBSDF),
+						StrataPixelHeader,
+						StrataBSDFContext,
+						BSDFEvaluate) * DirectionalLightCloudShadow;
 				}
 
 				////
@@ -434,7 +430,6 @@ float3 StrataForwardLighting(
 					// Update the strata BSDF context accordgin to the new L
 					StrataUpdateBSDFContext(StrataBSDFContext, ToLight);
 
-					bool bEvaluateHasBeenExecuted = false;
 					FStrataEvaluateResult BSDFEvaluate = (FStrataEvaluateResult)0;
 					Color += StrataForwardLightingCommon(
 						Dither,
@@ -451,7 +446,6 @@ float3 StrataForwardLighting(
 						LuminanceWeight(StrataBSDFContext, CurrentBSDF),
 						StrataPixelHeader,
 						StrataBSDFContext,
-						bEvaluateHasBeenExecuted,
 						BSDFEvaluate);
 				}
 
diff --git a/Engine/Shaders/Private/Strata/StrataLightingCommon.ush b/Engine/Shaders/Private/Strata/StrataLightingCommon.ush
--- a/Engine/Shaders/Private/Strata/StrataLightingCommon.ush
+++ b/Engine/Shaders/Private/Strata/StrataLightingCommon.ush
@@ -20,10 +20,8 @@ float3 StrataForwardLightingCommon(
 	float3 BSDFColoredVisibility,
 	FStrataPixelHeader StrataPixelHeader,
 	FStrataBSDFContext StrataBSDFContext,
-	inout bool bEvaluateHasBeenExecuted,
 	inout FStrataEvaluateResult BSDFEvaluate)
 {
-	bEvaluateHasBeenExecuted = false;
 	float3 Color = 0.0;
 
 	if (LightMask > 0.0)
@@ -31,8 +29,9 @@ float3 StrataForwardLightingCommon(
 		// Evaluate the ShadowTerm that can then be used when integrating the lighting
 		FShadowTerms ShadowTerms = { StrataGetAO(StrataPixelHeader), 1.0, 1.0, InitHairTransmittanceData() };
 
+		// PrecomputedShadowFactors are not supported on mobile at the moment
 		float4 PrecomputedShadowFactors = HasZeroPrecShadowMask(StrataPixelHeader) ? 0.0f : 1.0f;
-	#if TRANSLUCENCY_LIGHTING_SURFACE_FORWARDSHADING
+	#if TRANSLUCENCY_LIGHTING_SURFACE_FORWARDSHADING && !SHADING_PATH_MOBILE
 		PrecomputedShadowFactors.x = ComputeDirectionalLightStaticShadowing(TranslatedWorldPosition).x;
 	#else
 		PrecomputedShadowFactors.x = 1;
@@ -65,7 +64,6 @@ float3 StrataForwardLightingCommon(
 				// We must have the evaluate inside the if due to the rectlight texture: it must be now be ambiguous which texture is going ot be used.
 				// After te compilation, a local resource must map to a unique global resource (the default or the actual rect light texture).
 				BSDFEvaluate = StrataIntegrateBSDF(StrataBSDFContext, ShadowTerms, AreaLightContext, Settings, INTEGRATION_AREA_LIGHT_RECT);
-				bEvaluateHasBeenExecuted = true;
 			}
 			else
 			{
@@ -82,7 +80,6 @@ float3 StrataForwardLightingCommon(
 				{
 					BSDFEvaluate = StrataIntegrateBSDF(StrataBSDFContext, ShadowTerms, AreaLightContext, Settings, INTEGRATION_PUNCTUAL_LIGHT);
 				}
-				bEvaluateHasBeenExecuted = true;
 			}
 
 			FLATTEN
@@ -105,7 +102,6 @@ float3 StrataForwardLightingCommon(
 				// We must have the evaluate inside the if due to the rectlight texture: it must be now be ambiguous which texture is going ot be used.
 				// After te compilation, a local resource must map to a unique global resource (the default or the actual rect light texture).
 				BSDFEvaluate = StrataIntegrateBSDF(StrataBSDFContext, ShadowTerms, AreaLightContext, Settings, INTEGRATION_AREA_LIGHT_RECT);
-				bEvaluateHasBeenExecuted = true;
 			}
 			else
 			{
@@ -121,7 +117,6 @@ float3 StrataForwardLightingCommon(
 				{
 					BSDFEvaluate = StrataIntegrateBSDF(StrataBSDFContext, ShadowTerms, AreaLightContext, Settings, INTEGRATION_PUNCTUAL_LIGHT);
 				}
-				bEvaluateHasBeenExecuted = true;
 			}
 
 			FLATTEN
diff --git a/Engine/Shaders/Private/Strata/StrataMobileForwardLighting.ush b/Engine/Shaders/Private/Strata/StrataMobileForwardLighting.ush
new file mode 100644
--- /dev/null
+++ b/Engine/Shaders/Private/Strata/StrataMobileForwardLighting.ush
@@ -0,0 +1,342 @@
+// Copyright Epic Games, Inc. All Rights Reserved.
+
+#pragma once
+
+#define USE_STRATA_FORWARD_LIGHTING_COMMON 1
+#include "StrataLightingCommon.ush"
+#include "../LightData.ush"
+#include "../MobileLightingCommon.ush"
+
+#define FORWARD_PER_PIXEL_SHADING (FORWARD_SHADING || TRANSLUCENCY_LIGHTING_SURFACE_FORWARDSHADING || MATERIAL_SHADINGMODEL_SINGLELAYERWATER)
+
+// For debug or iteration purpose only
+#ifndef STRATA_FORCE_SINGLE_BSDF
+#define STRATA_FORCE_SINGLE_BSDF 1
+#endif
+
+// For debug or iteration purpose only
+#ifndef STRATA_FORCE_SINGLE_LOCAL_LIGHT
+#define STRATA_FORCE_SINGLE_LOCAL_LIGHT 0
+#endif
+
+// Forward declarations
+void GetVolumeLightingNonDirectional(float4 AmbientLightingVector, float3 DiffuseColor, inout float3 InterpolatedLighting, out float4 VolumeLighting);
+void GetVolumeLightingDirectional(float4 AmbientLightingVector, float3 DirectionalLightingVector, float3 WorldNormal, float3 DiffuseColor, inout float3 InterpolatedLighting, out float4 VolumeLighting);
+
+float3 GetWorldBentNormalZero(in FMaterialPixelParameters MaterialParameters);
+
+void GetPrecomputedIndirectLightingAndSkyLight(
+	VTPageTableResult LightmapVTPageTableResult, 
+	FVertexFactoryInterpolantsVSToPS Interpolants, 
+	bool bEvaluateBackface, 
+	half3 DiffuseDir, 
+	half3 SkyLightColor, 
+	out half3 OutDiffuseIndirectLighting, 
+	out half3 OutSubsurfaceIndirectLighting, 
+	out half OutIndirectIrradiance, 
+	out half3 OutWaterDiffuseIndirectLuminance);
+
+// This code should map to GetForwardDirectLightingSplit
+float3 StrataMobileForwardLighting(
+	uint EyeIndex,
+	float4 SvPosition,
+	FStrataIntegrationSettings Settings,
+	FMobileBasePassInterpolantsVSToPS BasePassInterpolants,
+	FVertexFactoryInterpolantsVSToPS Interpolants,
+	VTPageTableResult LightmapVTPageTableResult,
+	float3 VolumetricLightmapBrickTextureUVs,
+	FMaterialPixelParameters MaterialParameters,
+	float SceneDepth,
+	float2 ScreenUV,
+	FStrataPixelHeader StrataPixelHeader,
+	FStrataData StrataData,
+	inout float3 OutTransmittancePreCoverage,
+	inout float OutCoverage
+	)
+{
+	const uint PrimitiveId = MaterialParameters.PrimitiveId;
+	const uint2 PixelPos = uint2(SvPosition.xy);
+	const float Dither = InterleavedGradientNoise(PixelPos, View.StateFrameIndexMod8);
+	const float3 AbsoluteWorldPosition = LWCToFloat(MaterialParameters.AbsoluteWorldPosition);
+	const float3 TranslatedWorldPosition = MaterialParameters.WorldPosition_CamRelative;
+	const float3 LightingPositionOffset = MaterialParameters.LightingPositionOffset;
+	float3 V = MaterialParameters.CameraVector;
+
+	FStrataAddressing NullStrataAddressing = (FStrataAddressing)0;	// Fake unused in StrataCreateBSDFContext when using Forward inline shading
+
+	float3 Color = 0;
+	OutCoverage = 0.0f;
+	OutTransmittancePreCoverage = 1.0f;
+
+	float SpecularScale = 1;
+#if TRANSLUCENCY_ANY_VOLUMETRIC
+	// No specular on volumetric translucency lighting modes
+	SpecularScale = 0;
+#endif
+
+	const uint PrimitiveLightingChannelMask = GetPrimitive_LightingChannelMask(PrimitiveId);
+	const float2 LocalPosition = SvPosition.xy - ResolvedView.ViewRectMin.xy;
+	const uint GridIndex = ComputeLightGridCellIndex(uint2(LocalPosition.x, LocalPosition.y), SvPosition.w);
+
+	FStrataTree StrataTree = StrataPixelHeader.StrataTree;
+	if (StrataTree.BSDFCount > 0)
+	{
+		// Write Irradiance/AO data into header in order to be retrieve correctly during lighting
+		#if STRATA_INLINE_SHADING
+		HEADER_SETIRRADIANCE_AO(StrataPixelHeader.State, StrataPackIrradianceAndOcclusion(StrataPixelHeader.IrradianceAO, 0));
+		#endif
+
+		// Update tree (coverage/transmittance/luminace weights)
+		StrataUpdateTree(NullStrataAddressing, StrataPixelHeader, StrataTree, StrataData, V, Settings, OutCoverage, OutTransmittancePreCoverage);
+
+		//
+		// Forward lighting
+		//
+		// 
+		#if STRATA_FORCE_SINGLE_BSDF
+		int BSDFIdx = 0;
+		#else
+		STRATA_UNROLL_N(STRATA_CLAMPED_BSDF_COUNT)
+		for (int BSDFIdx = 0; BSDFIdx < StrataTree.BSDFCount; ++BSDFIdx)
+		#endif
+		{
+		#define CurrentBSDF StrataTree.BSDFs[BSDFIdx]
+
+			if (StrataIsBSDFVisible(CurrentBSDF))
+			{
+				// Starting with environment lighting, the context does not specify a light
+				FStrataBSDFContext StrataBSDFContext = StrataCreateBSDFContext(StrataPixelHeader, CurrentBSDF, NullStrataAddressing, V);
+					
+				////
+				//// Evaluate environment lighting
+				////
+
+				// Compute the bent normal from the BSDF normal
+				const float Roughness = StrataGetBSDFRoughness(CurrentBSDF);
+				const FShadingOcclusion ShadingOcclusion = ApplyBentNormal(MaterialParameters.CameraVector, StrataBSDFContext.N, GetWorldBentNormalZero(MaterialParameters), Roughness, StrataGetAO(StrataPixelHeader));
+				float IndirectOcclusion = 1.0f;
+				half  IndirectIrradiance = 0.0f; 
+				half3 WaterDiffuseIndirectLuminance = 0.f;
+
+				// 1. Diffuse Sky evaluation using bent normal
+				if (UseBasePassSkylightDiffuse > 0)
+				{
+					// Update a context specific for environment lighting
+					FStrataBSDFContext EnvBSDFContext = StrataBSDFContext;
+					EnvBSDFContext.N = ShadingOcclusion.BentNormal;
+					StrataUpdateBSDFContext(EnvBSDFContext, EnvBSDFContext.L);
+
+					// And evaluate diffuse parameters
+					FStrataEnvLightResult StrataEnvLight = StrataEvaluateForEnvLight(EnvBSDFContext, false/*bEnableSpecular*/, Settings);
+					float3 DiffuseNormal = StrataEnvLight.DiffuseNormal;
+					float3 DiffuseColorForIndirect = StrataEnvLight.DiffuseWeight;
+					float3 SubsurfaceColorForIndirect = StrataEnvLight.DiffuseBackFaceWeight;
+
+					// Evaluate diffuse lighting
+					if (any((DiffuseColorForIndirect + SubsurfaceColorForIndirect) > 0.0f))
+					{
+						const bool bEvaluateBackface = any(SubsurfaceColorForIndirect > 0.0);
+
+						const half NoL = max(0, dot(StrataBSDFContext.N, MobileDirectionalLight.DirectionalLightDirectionAndShadowTransition.xyz));
+						const bool bApplySkyLighting = !(MATERIAL_TWOSIDED && LQ_TEXTURE_LIGHTMAP) || NoL == 0;
+
+						half3 DiffuseIndirectLighting = 0.0f;
+						half3 SubsurfaceIndirectLighting = 0.0f;
+						GetPrecomputedIndirectLightingAndSkyLight(
+							LightmapVTPageTableResult, 
+							Interpolants, 
+							bEvaluateBackface, 
+							DiffuseNormal, 
+							ResolvedView.SkyLightColor.rgb, 
+							DiffuseIndirectLighting, 
+							SubsurfaceIndirectLighting, 
+							IndirectIrradiance, 
+							WaterDiffuseIndirectLuminance);
+
+						// STRATA_TODO: GetIndirectOcclusion require occlusion data, which are not available on mobile. The base pass has a 
+						// AmbientOcclusionTexture which could be used, but it does not seems to be hooked for now
+					#if 0 && FORWARD_SHADING && (MATERIALBLENDING_SOLID || MATERIALBLENDING_MASKED)
+						float2 NearestResolvedDepthScreenUV = CalculateNearestResolvedDepthScreenUV(ScreenUV, MaterialParameters.ScreenPosition.w);
+
+						IndirectOcclusion = GetIndirectOcclusion(NearestResolvedDepthScreenUV, HasDynamicIndirectShadowCasterRepresentation(StrataPixelHeader));
+						DiffuseIndirectLighting *= IndirectOcclusion;
+						SubsurfaceIndirectLighting *= IndirectOcclusion;
+						IndirectIrradiance *= IndirectOcclusion;
+					#endif
+
+						// For diffuse, we specify a perpendicular to the surface light direction for the transmittance to light to not be view dependent.
+						const float DiffuseEnvLightingNoL = 1.0f;
+
+						Color += (DiffuseIndirectLighting * DiffuseColorForIndirect + SubsurfaceIndirectLighting * SubsurfaceColorForIndirect) * 
+								 LuminanceWeight(DiffuseEnvLightingNoL, CurrentBSDF) * AOMultiBounce(StrataEnvLight.DiffuseColor, ShadingOcclusion.DiffOcclusion);
+					}
+				}
+
+				// 2. Specular evaluation (using regular normal)
+				//#if MATERIALBLENDING_MASKED || MATERIALBLENDING_SOLID || TRANSLUCENCY_LIGHTING_SURFACE_FORWARDSHADING || MATERIAL_SHADINGMODEL_SINGLELAYERWATER
+				#if FORWARD_PER_PIXEL_SHADING || TRANSLUCENCY_LIGHTING_SURFACE_LIGHTINGVOLUME
+				if (UseBasePassSkylightSpecular > 0)
+				{
+					// Evaluate specular parameters without bent normal
+					FStrataEnvLightResult StrataEnvLight = StrataEvaluateForEnvLight(StrataBSDFContext, true /*bEnableSpecular*/, Settings);
+
+					// Evaluate specular lighting for the main lobe
+					if (any(StrataEnvLight.SpecularWeight > 0.0f))
+					{
+						Color += GetImageBasedReflectionSpecular(StrataEnvLight.SpecularDirection, StrataEnvLight.SpecularSafeRoughness, MaterialParameters.WorldPosition_CamRelative, StrataBSDFContext.N, ShadingOcclusion, IndirectIrradiance, GridIndex) *
+								 StrataEnvLight.SpecularWeight * IndirectOcclusion * LuminanceWeight(StrataBSDFContext, CurrentBSDF) * AOMultiBounce(StrataEnvLight.SpecularColor, ShadingOcclusion.SpecOcclusion);
+					}
+					// And for the second lobe if needed
+					if (any(StrataEnvLight.SpecularHazeWeight > 0.0f))
+					{
+						Color += GetImageBasedReflectionSpecular(StrataEnvLight.SpecularHazeDirection, StrataEnvLight.SpecularHazeSafeRoughness, MaterialParameters.WorldPosition_CamRelative, StrataBSDFContext.N, ShadingOcclusion, IndirectIrradiance, GridIndex) *
+								 StrataEnvLight.SpecularHazeWeight * IndirectOcclusion * LuminanceWeight(StrataBSDFContext, CurrentBSDF) * AOMultiBounce(StrataEnvLight.SpecularColor, ShadingOcclusion.SpecOcclusion);
+					}
+				}
+				#endif
+
+
+				////
+				//// Evaluate translucent lighting volume / vertex lighting.
+				////
+				// This is currently not supported by mobile renderer. See StrataForwardLighting for more details
+
+
+				////
+				//// Evaluate emissive
+				////
+				Color += BSDF_GETEMISSIVE(CurrentBSDF) * CurrentBSDF.LuminanceWeightV;
+
+				////
+				//// Evaluate the single directional light selected for forward lighting.
+				////
+				const bool bHasDirectionalLight = any(MobileDirectionalLight.DirectionalLightColor > 0);
+				BRANCH
+				if (bHasDirectionalLight)
+				{
+					half DynamicShadowing = 1.0f;
+					half4 DynamicShadowFactors = 1.0f;
+					FDeferredLightData DirLightData = GetDirectionalLightData(MaterialParameters.ScreenPosition, SvPosition, DynamicShadowFactors, DynamicShadowing);
+					half4 DirLightAttenuation = half4(1, 1, DynamicShadowing, DynamicShadowing);
+					uint DirLightingChannelMask = LIGHTING_CHANNEL_MASK;
+					float3 ToDirLight = DirLightData.Direction;	// Already normalized
+					float DirLightMask = 1;
+	
+					float DirectionalLightCloudShadow = 1.0f;
+				#if NEEDS_BASEPASS_CLOUD_SHADOW_INTERPOLATOR
+					DirectionalLightCloudShadow = BasePassInterpolants.VertexCloudShadow;
+				#endif
+
+					StrataUpdateBSDFContext(StrataBSDFContext, ToDirLight);
+
+					FStrataEvaluateResult BSDFEvaluate = (FStrataEvaluateResult)0;
+					Color += StrataForwardLightingCommon(
+						Dither,
+						Settings,
+						DirLightData,
+						ToDirLight,
+						DirLightMask,
+						DirLightAttenuation,
+						InitRectTexture(),
+						DirLightingChannelMask,
+						PrimitiveLightingChannelMask,
+						TranslatedWorldPosition,
+						SceneDepth,
+						LuminanceWeight(StrataBSDFContext, CurrentBSDF),
+						StrataPixelHeader,
+						StrataBSDFContext,
+						BSDFEvaluate) * DirectionalLightCloudShadow;
+				}
+
+				////
+				//// Evaluate local lights
+				////
+				#if !DISABLE_FORWARD_LOCAL_LIGHTS && (FORWARD_SHADING || TRANSLUCENCY_LIGHTING_SURFACE_FORWARDSHADING) // Forbidden for Single Layer Water
+				{
+					// Basic implementation for FORWARD_LOCAL_LIGHTS where local lights are reloaded per BSDF layer. It is fast by default when layer==1.
+					const FCulledLightsGridData CulledLightsGrid = GetCulledLightsGrid(GridIndex, EyeIndex);
+
+					// Limit max to ForwardLightData.NumLocalLights.
+					// This prevents GPU hangs when the PS tries to read from uninitialized NumCulledLightsGrid buffer
+					const uint NumLocalLights = min(CulledLightsGrid.NumLocalLights, GetNumLocalLights(EyeIndex));
+
+				#if STRATA_FORCE_SINGLE_LOCAL_LIGHT
+					#define LOCAL_LIGHT_EARLY_OUT return Color;
+					uint LocalLightListIndex = 0;
+					if (NumLocalLights > 0)
+				#else
+					#define LOCAL_LIGHT_EARLY_OUT continue;
+					LOOP
+					for (uint LocalLightListIndex = 0; LocalLightListIndex < NumLocalLights; LocalLightListIndex++)
+				#endif
+					{
+						const FLocalLightData LocalLight = GetLocalLightData(CulledLightsGrid.DataStartIndex + LocalLightListIndex, EyeIndex);
+
+						half4 PreviewShadowMapChannelMask = 1;
+						half4 DynamicShadowFactors = 1;
+						uint  LocalLightChannelMask = LIGHTING_CHANNEL_MASK;
+						FDeferredLightData LightData = ConvertToDeferredLight_Mobile(LocalLight);
+					
+						// Rect light optionally supported in forward.
+						// * This is not done explicitely in the legacy path (resulting in no lighting) but we do it for strata in order to avoid visual artefact.
+						// * This can be removed when LocalLight properly supports rect lights.
+						if (LightData.bRectLight && !SUPPORT_RECTLIGHT_ON_FORWARD_LIT_TRANSLUCENT)
+						{
+							LOCAL_LIGHT_EARLY_OUT;
+						}
+
+						// When disabled, force the compiler to remove the texture atlas sampling to avoid taking an extra texture sampler
+						#if !SUPPORT_RECTLIGHT_ON_FORWARD_LIT_TRANSLUCENT
+						LightData.bRectLight = false;
+						#endif
+						const FRectTexture LocalRectTexture = ConvertToRectTexture(LightData);
+
+						float DynamicShadowing = dot(PreviewShadowMapChannelMask, DynamicShadowFactors);
+						float4 LightAttenuation = float4(1, 1, DynamicShadowing.x, DynamicShadowing.x);
+
+						float3 L = LightData.Direction;	// Already normalized
+						float3 ToLight = L;
+						float LightMask = 1;
+						if (LightData.bRadialLight)
+						{
+							LightMask = GetLocalLightAttenuation(TranslatedWorldPosition, LightData, ToLight, L);
+						}
+
+						if (!IsLightVisible(LocalLight, TranslatedWorldPosition))
+						{
+							LOCAL_LIGHT_EARLY_OUT;
+						}
+
+						// Update the strata BSDF context accordgin to the new L
+						StrataUpdateBSDFContext(StrataBSDFContext, ToLight);
+
+						FStrataEvaluateResult BSDFEvaluate = (FStrataEvaluateResult)0;
+						Color += StrataForwardLightingCommon(
+							Dither,
+							Settings,
+							LightData,
+							ToLight,
+							LightMask,
+							LightAttenuation,
+							LocalRectTexture,
+							LocalLightChannelMask,
+							PrimitiveLightingChannelMask,
+							TranslatedWorldPosition,
+							SceneDepth,
+							LuminanceWeight(StrataBSDFContext, CurrentBSDF),
+							StrataPixelHeader,
+							StrataBSDFContext,
+							BSDFEvaluate);
+					}
+				}
+				#endif // !DISABLE_FORWARD_LOCAL_LIGHTS
+
+			}// StrataIsBSDFVisible
+		#undef CurrentBSDF
+		}
+	}
+
+	return Color;
+}
+
+
diff --git a/Engine/Shaders/Private/Strata/StrataVisualize.usf b/Engine/Shaders/Private/Strata/StrataVisualize.usf
--- a/Engine/Shaders/Private/Strata/StrataVisualize.usf
+++ b/Engine/Shaders/Private/Strata/StrataVisualize.usf
@@ -1087,7 +1087,6 @@ void EvaluateLighting(
 
 		FStrataBSDFContext StrataBSDFContext = StrataCreateBSDFContext(TangentBasis, BSDF, V, ToLight);
 
-		bool bEvaluateHasBeenExecuted = false;
 		FStrataEvaluateResult BSDFEvaluate = (FStrataEvaluateResult)0;
 		float3 DirLightLuminance = StrataForwardLightingCommon(
 			Dither,
@@ -1104,7 +1103,6 @@ void EvaluateLighting(
 			BSDFColoredVisibility,
 			StrataPixelHeader,
 			StrataBSDFContext,
-			bEvaluateHasBeenExecuted,
 			BSDFEvaluate);
 
 		OutColor.rgb += DirLightLuminance * View.PreExposure;
diff --git a/Engine/Source/Runtime/Renderer/Private/MobileBasePassRendering.cpp b/Engine/Source/Runtime/Renderer/Private/MobileBasePassRendering.cpp
--- a/Engine/Source/Runtime/Renderer/Private/MobileBasePassRendering.cpp
+++ b/Engine/Source/Runtime/Renderer/Private/MobileBasePassRendering.cpp
@@ -227,6 +227,9 @@ void SetupMobileBasePassUniformParameters(
 		BasePassParameters.ScreenSpaceShadowMaskSampler = TStaticSamplerState<SF_Point, AM_Clamp, AM_Clamp, AM_Clamp>::GetRHI();
 	}
 
+	// Strata
+	Strata::BindStrataMobileForwardPasslUniformParameters(GraphBuilder, View, BasePassParameters.Strata);
+
 	SetupReflectionUniformParameters(View, BasePassParameters.ReflectionsParameters);
 }
 
diff --git a/Engine/Source/Runtime/Renderer/Private/MobileBasePassRendering.h b/Engine/Source/Runtime/Renderer/Private/MobileBasePassRendering.h
--- a/Engine/Source/Runtime/Renderer/Private/MobileBasePassRendering.h
+++ b/Engine/Source/Runtime/Renderer/Private/MobileBasePassRendering.h
@@ -37,6 +37,7 @@ BEGIN_GLOBAL_SHADER_PARAMETER_STRUCT(FMobileBasePassUniformParameters, )
 	SHADER_PARAMETER_STRUCT(FForwardLightData, Forward)
 	SHADER_PARAMETER_STRUCT(FPlanarReflectionUniformParameters, PlanarReflection) // Single global planar reflection for the forward pass.
 	SHADER_PARAMETER_STRUCT(FMobileSceneTextureUniformParameters, SceneTextures)
+	SHADER_PARAMETER_STRUCT(FStrataMobileForwardPassUniformParameters, Strata)
 	SHADER_PARAMETER_STRUCT(FDebugViewModeUniformParameters, DebugViewMode)
 	SHADER_PARAMETER_STRUCT(FReflectionUniformParameters, ReflectionsParameters)
 	SHADER_PARAMETER_TEXTURE(Texture2D, PreIntegratedGFTexture)
diff --git a/Engine/Source/Runtime/Renderer/Private/MobileShadingRenderer.cpp b/Engine/Source/Runtime/Renderer/Private/MobileShadingRenderer.cpp
--- a/Engine/Source/Runtime/Renderer/Private/MobileShadingRenderer.cpp
+++ b/Engine/Source/Runtime/Renderer/Private/MobileShadingRenderer.cpp
@@ -850,7 +850,14 @@ void FMobileSceneRenderer::Render(FRDGBuilder& GraphBuilder)
 
 	FRDGSystemTextures::Create(GraphBuilder);
 
-	// Strata initialisation is always run even when not enabled.
+	// Strata initialization is always run even when not enabled.
+	if (Strata::IsStrataEnabled())
+	{
+		for (FViewInfo& View : Views)
+		{
+			ShadingEnergyConservation::Init(GraphBuilder, View);
+		}
+	}
 	Strata::InitialiseStrataFrameSceneData(GraphBuilder, *this);
 
 	// Force the subsurface profile texture to be updated.
diff --git a/Engine/Source/Runtime/Renderer/Private/Strata/Strata.cpp b/Engine/Source/Runtime/Renderer/Private/Strata/Strata.cpp
--- a/Engine/Source/Runtime/Renderer/Private/Strata/Strata.cpp
+++ b/Engine/Source/Runtime/Renderer/Private/Strata/Strata.cpp
@@ -607,6 +607,23 @@ void BindStrataForwardPasslUniformParameters(FRDGBuilder& GraphBuilder, const FV
 	}
 }
 
+void BindStrataMobileForwardPasslUniformParameters(FRDGBuilder& GraphBuilder, const FViewInfo& View, FStrataMobileForwardPassUniformParameters& OutStrataUniformParameters)
+{
+	FStrataSceneData* StrataSceneData = View.StrataViewData.SceneData;
+	if (IsStrataEnabled() && StrataSceneData)
+	{
+		OutStrataUniformParameters.MaxBytesPerPixel = StrataSceneData->MaxBytesPerPixel;
+		OutStrataUniformParameters.bRoughDiffuse = StrataSceneData->bRoughDiffuse ? 1u : 0u;
+		OutStrataUniformParameters.PeelLayersAboveDepth = StrataSceneData->PeelLayersAboveDepth;
+	}
+	else
+	{
+		OutStrataUniformParameters.MaxBytesPerPixel = 0;
+		OutStrataUniformParameters.bRoughDiffuse = 0;
+		OutStrataUniformParameters.PeelLayersAboveDepth = 0;
+	}
+}
+
 TRDGUniformBufferRef<FStrataGlobalUniformParameters> BindStrataGlobalUniformParameters(const FViewInfo& View)
 {
 	check(View.StrataViewData.StrataGlobalUniformParameters != nullptr || !IsStrataEnabled());
diff --git a/Engine/Source/Runtime/Renderer/Private/Strata/Strata.h b/Engine/Source/Runtime/Renderer/Private/Strata/Strata.h
--- a/Engine/Source/Runtime/Renderer/Private/Strata/Strata.h
+++ b/Engine/Source/Runtime/Renderer/Private/Strata/Strata.h
@@ -37,6 +37,12 @@ BEGIN_SHADER_PARAMETER_STRUCT(FStrataForwardPassUniformParameters, )
 	SHADER_PARAMETER_RDG_TEXTURE(Texture2D<uint2>, TopLayerTexture)
 END_SHADER_PARAMETER_STRUCT()
 
+BEGIN_SHADER_PARAMETER_STRUCT(FStrataMobileForwardPassUniformParameters, )
+	SHADER_PARAMETER(uint32, MaxBytesPerPixel)
+	SHADER_PARAMETER(uint32, bRoughDiffuse)
+	SHADER_PARAMETER(uint32, PeelLayersAboveDepth)
+END_SHADER_PARAMETER_STRUCT()
+
 BEGIN_SHADER_PARAMETER_STRUCT(FStrataTileParameter, )
 	SHADER_PARAMETER_RDG_BUFFER_SRV(Buffer<uint>, TileListBuffer)
 	RDG_BUFFER_ACCESS(TileIndirectBuffer, ERHIAccess::IndirectArgs)
@@ -160,6 +166,7 @@ void InitialiseStrataFrameSceneData(FRDGBuilder& GraphBuilder, FSceneRenderer& S
 
 void BindStrataBasePassUniformParameters(FRDGBuilder& GraphBuilder, const FViewInfo& View, FStrataBasePassUniformParameters& OutStrataUniformParameters);
 void BindStrataForwardPasslUniformParameters(FRDGBuilder& GraphBuilder, const FViewInfo& View, FStrataForwardPassUniformParameters& OutStrataUniformParameters);
+void BindStrataMobileForwardPasslUniformParameters(FRDGBuilder& GraphBuilder, const FViewInfo& View, FStrataMobileForwardPassUniformParameters& OutStrataUniformParameters);
 TRDGUniformBufferRef<FStrataGlobalUniformParameters> BindStrataGlobalUniformParameters(const FViewInfo& View);
 
 void AppendStrataMRTs(const FSceneRenderer& SceneRenderer, uint32& BasePassTextureCount, TStaticArray<FTextureRenderTargetBinding, MaxSimultaneousRenderTargets>& BasePassTextures);
