Move IsLightVisible function into LightData.ush to remove duplicate/share code.

diff --git a/Engine/Shaders/Private/ClusteredDeferredShadingPixelShader.usf b/Engine/Shaders/Private/ClusteredDeferredShadingPixelShader.usf
--- a/Engine/Shaders/Private/ClusteredDeferredShadingPixelShader.usf
+++ b/Engine/Shaders/Private/ClusteredDeferredShadingPixelShader.usf
@@ -108,18 +108,6 @@ bool GetClusterDebugColor(float2 LocalPosition, float SceneDepth, FCulledLightsG
 #endif
 }
 
-/**
- * Return true is the light is visible, false otherwise
- */
-bool IsLightVisible(const FLocalLightData In, float3 TranslatedWorldPosition)
-{
-	// extra-early out since we know light grid is sloppy and all lights in list are radial (have a range)
-	// appears useless
-	const float InvLightRadiusSq = In.LightPositionAndInvRadius.w * In.LightPositionAndInvRadius.w;
-	return length2(TranslatedWorldPosition - In.LightPositionAndInvRadius.xyz) * InvLightRadiusSq <= 1.0f;
-}
-
-
 /**
  * Adds local lighting of simple lights and returns the index at which non-simple light should starts
  */
diff --git a/Engine/Shaders/Private/LightData.ush b/Engine/Shaders/Private/LightData.ush
--- a/Engine/Shaders/Private/LightData.ush
+++ b/Engine/Shaders/Private/LightData.ush
@@ -429,4 +429,14 @@ FSimpleDeferredLightData ConvertToSimpleLight(FLocalLightData In)
 	Out.FalloffExponent			= In.LightColorAndFalloffExponent.w;
 	Out.bInverseSquared			= Out.FalloffExponent == 0;
 	return Out;
+}
+
+///////////////////////////////////////////////////////////////////////////////////////////////////
+// Helper functions
+
+// Return true if the point is influenced by the light
+bool IsLightVisible(FLocalLightData InLight, float3 TranslatedWorldPosition)
+{
+	float InvLightRadiusSq = InLight.LightPositionAndInvRadius.w * InLight.LightPositionAndInvRadius.w;
+	return length2(TranslatedWorldPosition - InLight.LightPositionAndInvRadius.xyz) * InvLightRadiusSq <= 1.0f;
 }
\ No newline at end of file
diff --git a/Engine/Shaders/Private/MobileLightingCommon.ush b/Engine/Shaders/Private/MobileLightingCommon.ush
--- a/Engine/Shaders/Private/MobileLightingCommon.ush
+++ b/Engine/Shaders/Private/MobileLightingCommon.ush
@@ -489,13 +489,6 @@ void AccumulateReflection(
 	LightAccumulator_AddSplit(DirectLighting, 0.0f, SpecularIBLLighting, 0.0f, 1.0f, false);
 }
 
-// Return true if the point is influenced by the light
-bool IsLightVisible(FLocalLightData InLight, float3 TranslatedWorldPosition)
-{
-	float InvLightRadiusSq = InLight.LightPositionAndInvRadius.w * InLight.LightPositionAndInvRadius.w;
-	return length2(TranslatedWorldPosition - InLight.LightPositionAndInvRadius.xyz) * InvLightRadiusSq <= 1.0f;
-}
-
 /*------------------------------------------------------------------------------
 	Mobile Local lights.
 ------------------------------------------------------------------------------*/
diff --git a/Engine/Shaders/Private/RayTracing/RayTracingLightingCommon.ush b/Engine/Shaders/Private/RayTracing/RayTracingLightingCommon.ush
--- a/Engine/Shaders/Private/RayTracing/RayTracingLightingCommon.ush
+++ b/Engine/Shaders/Private/RayTracing/RayTracingLightingCommon.ush
@@ -442,7 +442,6 @@ float3 ComputeDirectLighting(
 			continue;
 		}
 
-		bool IsLightVisible = false;
 		float ShadowRayLength = 0.0;
 		if (ReflectedShadowsType == 2)
 		{
