* Add opacity support for Strata SceneTexture evaluation
* Shared SceneTexture evaluation between Strata & legacy paths.

diff --git a/Engine/Shaders/Private/MaterialTemplate.ush b/Engine/Shaders/Private/MaterialTemplate.ush
--- a/Engine/Shaders/Private/MaterialTemplate.ush
+++ b/Engine/Shaders/Private/MaterialTemplate.ush
@@ -2361,8 +2361,8 @@ float4 SceneTextureLookup(float2 UV, int SceneTextureIndex, bool bFiltered)
 		return 0.0f;
 	}
 
-#else // !SCENE_TEXTURES_DISABLED
-#if MATERIAL_SHADINGMODEL_SINGLELAYERWATER
+#elif MATERIAL_SHADINGMODEL_SINGLELAYERWATER
+
 	if (SceneTextureIndex == PPI_CustomDepth)
 	{
 		MaterialFloat Depth = ConvertFromDeviceZ(Texture2DSample(SingleLayerWater.CustomDepthTexture, SingleLayerWater.CustomDepthSampler, UV).r);
@@ -2374,154 +2374,106 @@ float4 SceneTextureLookup(float2 UV, int SceneTextureIndex, bool bFiltered)
 		return MaterialFloat4(Stencil, Stencil, Stencil, 0.f);
 	}
 
-	return float4(0, 0, 0, 0);
-#else // MATERIAL_SHADINGMODEL_SINGLELAYERWATER
-#if STRATA_ENABLED
+#else // SCENE_TEXTURES_DISABLED
 
 	const uint2 PixelPos = uint2(UV * View.BufferSizeAndInvSize.xy);
 
 	// 1. Post-process inputs which are independent on material data
 	switch (SceneTextureIndex)
 	{
-	case PPI_SceneColor:			return float4(CalcSceneColor(UV), 0);
-	case PPI_SceneDepth:			return CalcSceneDepth(UV);
-	case PPI_SeparateTranslucency:	return float4(1, 1, 1, 1);	// todo
-	case PPI_Opacity:				return 1.f; // Opacity has different meaning depending on the legacy shading model.
-	case PPI_CustomDepth:			return CalcSceneCustomDepth(UV);
-#if POST_PROCESS_MATERIAL
-	case PPI_PostProcessInput0:		return Texture2DSample(PostProcessInput_0_Texture, bFiltered ? PostProcessInput_BilinearSampler : PostProcessInput_0_SharedSampler, UV);
-	case PPI_PostProcessInput1:		return Texture2DSample(PostProcessInput_1_Texture, bFiltered ? PostProcessInput_BilinearSampler : PostProcessInput_1_SharedSampler, UV);
-	case PPI_PostProcessInput2:		return Texture2DSample(PostProcessInput_2_Texture, bFiltered ? PostProcessInput_BilinearSampler : PostProcessInput_2_SharedSampler, UV);
-	case PPI_PostProcessInput3:		return Texture2DSample(PostProcessInput_3_Texture, bFiltered ? PostProcessInput_BilinearSampler : PostProcessInput_3_SharedSampler, UV);
-	case PPI_PostProcessInput4:		return Texture2DSample(PostProcessInput_4_Texture, bFiltered ? PostProcessInput_BilinearSampler : PostProcessInput_4_SharedSampler, UV);
-#endif // __POST_PROCESS_COMMON__
-	case PPI_DecalMask:				return 0;  // material compiler will return an error
-	case PPI_AmbientOcclusion:		return CalcSceneAO(UV);
-	case PPI_CustomStencil:			return CalcSceneCustomStencil(PixelPos);
-#if POST_PROCESS_MATERIAL
-	case PPI_Velocity:				return float4(PostProcessVelocityLookup(ConvertToDeviceZ(CalcSceneDepth(UV)), UV), 0, 0);
-#endif
+		case PPI_SceneColor:			return float4(CalcSceneColor(UV), 0);
+		case PPI_SceneDepth:			return CalcSceneDepth(UV);
+		case PPI_SeparateTranslucency:	return float4(1, 1, 1, 1);	// todo
+		case PPI_CustomDepth:			return CalcSceneCustomDepth(UV);
+	#if POST_PROCESS_MATERIAL
+		case PPI_PostProcessInput0:		return Texture2DSample(PostProcessInput_0_Texture, bFiltered ? PostProcessInput_BilinearSampler : PostProcessInput_0_SharedSampler, UV);
+		case PPI_PostProcessInput1:		return Texture2DSample(PostProcessInput_1_Texture, bFiltered ? PostProcessInput_BilinearSampler : PostProcessInput_1_SharedSampler, UV);
+		case PPI_PostProcessInput2:		return Texture2DSample(PostProcessInput_2_Texture, bFiltered ? PostProcessInput_BilinearSampler : PostProcessInput_2_SharedSampler, UV);
+		case PPI_PostProcessInput3:		return Texture2DSample(PostProcessInput_3_Texture, bFiltered ? PostProcessInput_BilinearSampler : PostProcessInput_3_SharedSampler, UV);
+		case PPI_PostProcessInput4:		return Texture2DSample(PostProcessInput_4_Texture, bFiltered ? PostProcessInput_BilinearSampler : PostProcessInput_4_SharedSampler, UV);
+	#endif // __POST_PROCESS_COMMON__
+		case PPI_DecalMask:				return 0;  // material compiler will return an error
+		case PPI_AmbientOcclusion:		return CalcSceneAO(UV);
+		case PPI_CustomStencil:			return CalcSceneCustomStencil(PixelPos);
+	#if POST_PROCESS_MATERIAL
+		case PPI_Velocity:				return float4(PostProcessVelocityLookup(ConvertToDeviceZ(CalcSceneDepth(UV)), UV), 0, 0);
+	#endif
 	}
 
 	// 2. Post-process inputs which are dependent on material data
 	//    Note: For PostProcess material we always return the first layer dataCalcSceneCustomStencil
-#ifdef StrataStruct
-	FStrataAddressing StrataAddressing   = GetStrataPixelDataByteOffset(PixelPos, uint2(View.BufferSizeAndInvSize.xy), StrataStruct.MaxBytesPerPixel);
-	FStrataPixelHeader StrataPixelHeader = UnpackStrataHeaderIn(StrataStruct.MaterialTextureArray, StrataAddressing, StrataStruct.TopLayerTexture);
-	BRANCH
-	if (StrataPixelHeader.BSDFCount > 0)
+	#if STRATA_ENABLED
 	{
-		const FStrataBSDF BSDF = UnpackStrataBSDFIn(StrataStruct.MaterialTextureArray, StrataAddressing, StrataPixelHeader);
-#else
-	FStrataAddressing StrataAddressing = GetStrataPixelDataByteOffset(PixelPos, uint2(View.BufferSizeAndInvSize.xy), Strata.MaxBytesPerPixel);
-	FStrataPixelHeader StrataPixelHeader = UnpackStrataHeaderIn(Strata.MaterialTextureArray, StrataAddressing, Strata.TopLayerTexture);
-	BRANCH
-	if (StrataPixelHeader.BSDFCount > 0)
+		#ifdef StrataStruct
+			#define StrataSceneTexture StrataStructure
+		#else
+			#define StrataSceneTexture Strata
+		#endif
+		FStrataAddressing StrataAddressing   = GetStrataPixelDataByteOffset(PixelPos, uint2(View.BufferSizeAndInvSize.xy), StrataSceneTexture.MaxBytesPerPixel);
+		FStrataPixelHeader StrataPixelHeader = UnpackStrataHeaderIn(StrataSceneTexture.MaterialTextureArray, StrataAddressing, StrataSceneTexture.TopLayerTexture);
+		BRANCH
+		if (StrataPixelHeader.BSDFCount > 0)
+		{
+			const FStrataSubsurfaceHeader SSSHeader = StrataLoadSubsurfaceHeader(StrataSceneTexture.MaterialTextureArray, StrataSceneTexture.FirstSliceStoringStrataSSSData, StrataAddressing.PixelCoords);
+			const uint SSSType = StrataSubSurfaceHeaderGetSSSType(SSSHeader);
+
+			const FStrataBSDF BSDF = UnpackStrataBSDFIn(StrataSceneTexture.MaterialTextureArray, StrataAddressing, StrataPixelHeader);
+			switch (SceneTextureIndex)
+			{
+				// order needs to match to ESceneTextureId
+				case PPI_DiffuseColor:			return float4(StrataGetBSDFDiffuseColor(BSDF), 0);
+				case PPI_SpecularColor:			return float4(StrataGetBSDFSpecularColor(BSDF), 0);
+				case PPI_SubsurfaceColor:		return float4(StrataGetBSDFSubSurfaceColor(BSDF), 0);
+				case PPI_BaseColor:				return float4(StrataGetBSDFBaseColor(BSDF), 0);
+				case PPI_Specular:				return StrataGetBSDFSpecular(BSDF);
+				case PPI_Metallic:				return StrataGetBSDFMetallic(BSDF);
+				case PPI_WorldNormal:			return float4(StrataGetWorldNormal(StrataPixelHeader, BSDF, StrataAddressing), 0);
+				case PPI_Opacity:				return SSSType != SSS_TYPE_INVALID ? StrataSubSurfaceHeaderGetOpacity(SSSHeader) : 1.f;
+				case PPI_Roughness:				return StrataGetBSDFRoughness(BSDF);
+				case PPI_MaterialAO:			return StrataGetAO(StrataPixelHeader);
+				case PPI_ShadingModelColor:		return float4(GetShadingModelColor(StrataGetLegacyShadingModels(BSDF)), 1);
+				case PPI_ShadingModelID:		return float4(StrataGetLegacyShadingModels(BSDF), 0, 0, 0);
+				case PPI_StoredBaseColor:		return float4(StrataGetBSDFBaseColor(BSDF), 0);
+				case PPI_StoredSpecular:		return float4(StrataGetBSDFSpecular(BSDF).rrr, 0);
+				case PPI_WorldTangent:			return float4(StrataGetWorldTangent(StrataPixelHeader, BSDF, StrataAddressing), 0);
+				case PPI_Anisotropy:			return StrataGetBSDFAnisotropy(BSDF);
+			}
+		}
+		#undef StrataSceneTexture
+	}
+	#else // STRATA_ENABLED
 	{
-		const FStrataBSDF BSDF = UnpackStrataBSDFIn(Strata.MaterialTextureArray, StrataAddressing, StrataPixelHeader);
-#endif
+		FScreenSpaceData ScreenSpaceData = GetScreenSpaceData(UV, false);
 		switch (SceneTextureIndex)
 		{
 			// order needs to match to ESceneTextureId
-			case PPI_DiffuseColor:			return float4(StrataGetBSDFDiffuseColor(BSDF), 0);
-			case PPI_SpecularColor:			return float4(StrataGetBSDFSpecularColor(BSDF), 0);
-			case PPI_SubsurfaceColor:		return float4(StrataGetBSDFSubSurfaceColor(BSDF), 0);
-			case PPI_BaseColor:				return float4(StrataGetBSDFBaseColor(BSDF), 0);
-			case PPI_Specular:				return StrataGetBSDFSpecular(BSDF);
-			case PPI_Metallic:				return StrataGetBSDFMetallic(BSDF);
-			case PPI_WorldNormal:			return float4(StrataGetWorldNormal(StrataPixelHeader, BSDF, StrataAddressing), 0);
-			case PPI_Roughness:				return StrataGetBSDFRoughness(BSDF);
-			case PPI_MaterialAO:			return StrataGetAO(StrataPixelHeader);
-			case PPI_ShadingModelColor:		return float4(GetShadingModelColor(StrataGetLegacyShadingModels(BSDF)), 1);
-			case PPI_ShadingModelID:		return float4(StrataGetLegacyShadingModels(BSDF), 0, 0, 0);
-			case PPI_StoredBaseColor:		return float4(StrataGetBSDFBaseColor(BSDF), 0);
-			case PPI_StoredSpecular:		return float4(StrataGetBSDFSpecular(BSDF).rrr, 0);
-			case PPI_WorldTangent:			return float4(StrataGetWorldTangent(StrataPixelHeader, BSDF, StrataAddressing), 0);
-			case PPI_Anisotropy:			return StrataGetBSDFAnisotropy(BSDF);
+			case PPI_DiffuseColor:			return float4(ScreenSpaceData.GBuffer.DiffuseColor, 0);
+			case PPI_SpecularColor:			return float4(ScreenSpaceData.GBuffer.SpecularColor, 0);
+			case PPI_SubsurfaceColor:		return IsSubsurfaceModel(ScreenSpaceData.GBuffer.ShadingModelID) ? float4(ExtractSubsurfaceColor(ScreenSpaceData.GBuffer), ScreenSpaceData.GBuffer.CustomData.a) : ScreenSpaceData.GBuffer.CustomData;
+			case PPI_BaseColor:				return float4(ScreenSpaceData.GBuffer.BaseColor, 0);
+			case PPI_Specular:				return ScreenSpaceData.GBuffer.Specular;
+			case PPI_Metallic:				return ScreenSpaceData.GBuffer.Metallic;
+			#if IS_DBUFFER_DECAL
+			case PPI_WorldNormal:			return GetDBufferReprojectedWorldNormal(UV); // DBuffer uses depth buffer gradients or reprojection of last frame normal.
+			#else
+			case PPI_WorldNormal:			return float4(ScreenSpaceData.GBuffer.WorldNormal, 0);
+			#endif
+			case PPI_Opacity:				return ScreenSpaceData.GBuffer.CustomData.a;
+			case PPI_Roughness:				return ScreenSpaceData.GBuffer.Roughness;
+			case PPI_MaterialAO:			return ScreenSpaceData.GBuffer.GBufferAO;
+			case PPI_ShadingModelColor:		return float4(GetShadingModelColor(ScreenSpaceData.GBuffer.ShadingModelID), 1);
+			case PPI_ShadingModelID:		return float4(ScreenSpaceData.GBuffer.ShadingModelID, 0, 0, 0);
+			case PPI_StoredBaseColor:		return float4(ScreenSpaceData.GBuffer.StoredBaseColor, 0);
+			case PPI_StoredSpecular:		return float4(ScreenSpaceData.GBuffer.StoredSpecular.rrr, 0);
+			case PPI_WorldTangent:			return float4(ScreenSpaceData.GBuffer.WorldTangent, 0);
+			case PPI_Anisotropy:			return ScreenSpaceData.GBuffer.Anisotropy;
 		}
 	}
+	#endif // STRATA_ENABLED
 
-	return float4(0, 0, 0, 0);
-
-#else // STRATA_ENABLED
-	FScreenSpaceData ScreenSpaceData = GetScreenSpaceData(UV, false);
-	switch(SceneTextureIndex)
-	{
-		// order needs to match to ESceneTextureId
-
-		case PPI_SceneColor:
-			return float4(CalcSceneColor(UV), 0);
-		case PPI_SceneDepth:
-			return ScreenSpaceData.GBuffer.Depth;
-		case PPI_DiffuseColor:
-			return float4(ScreenSpaceData.GBuffer.DiffuseColor, 0);
-		case PPI_SpecularColor:
-			return float4(ScreenSpaceData.GBuffer.SpecularColor, 0);
-		case PPI_SubsurfaceColor:
-			return IsSubsurfaceModel(ScreenSpaceData.GBuffer.ShadingModelID) ? float4( ExtractSubsurfaceColor(ScreenSpaceData.GBuffer), ScreenSpaceData.GBuffer.CustomData.a ) : ScreenSpaceData.GBuffer.CustomData;
-		case PPI_BaseColor:
-			return float4(ScreenSpaceData.GBuffer.BaseColor, 0);
-		case PPI_Specular:
-			return ScreenSpaceData.GBuffer.Specular;
-		case PPI_Metallic:
-			return ScreenSpaceData.GBuffer.Metallic;
-		case PPI_WorldNormal:
-#if IS_DBUFFER_DECAL
-			// DBuffer uses depth buffer gradients or reprojection of last frame normal.
-			return GetDBufferReprojectedWorldNormal(UV);
-#else
-			return float4(ScreenSpaceData.GBuffer.WorldNormal, 0);
-#endif
-		case PPI_SeparateTranslucency:
-			return float4(1, 1, 1, 1);	// todo
-		case PPI_Opacity:
-			return ScreenSpaceData.GBuffer.CustomData.a;
-		case PPI_Roughness:
-			return ScreenSpaceData.GBuffer.Roughness;
-		case PPI_MaterialAO:
-			return ScreenSpaceData.GBuffer.GBufferAO;
-		case PPI_CustomDepth:
-			return ScreenSpaceData.GBuffer.CustomDepth;
-#if POST_PROCESS_MATERIAL
-		case PPI_PostProcessInput0:
-			return Texture2DSample(PostProcessInput_0_Texture, bFiltered ? PostProcessInput_BilinearSampler : PostProcessInput_0_SharedSampler, UV);
-		case PPI_PostProcessInput1:
-			return Texture2DSample(PostProcessInput_1_Texture, bFiltered ? PostProcessInput_BilinearSampler : PostProcessInput_1_SharedSampler, UV);
-		case PPI_PostProcessInput2:
-			return Texture2DSample(PostProcessInput_2_Texture, bFiltered ? PostProcessInput_BilinearSampler : PostProcessInput_2_SharedSampler, UV);
-		case PPI_PostProcessInput3:
-			return Texture2DSample(PostProcessInput_3_Texture, bFiltered ? PostProcessInput_BilinearSampler : PostProcessInput_3_SharedSampler, UV);
-		case PPI_PostProcessInput4:
-			return Texture2DSample(PostProcessInput_4_Texture, bFiltered ? PostProcessInput_BilinearSampler : PostProcessInput_4_SharedSampler, UV);
-#endif // __POST_PROCESS_COMMON__
-		case PPI_DecalMask:
-			return 0;  // material compiler will return an error
-		case PPI_ShadingModelColor:
-			return float4(GetShadingModelColor(ScreenSpaceData.GBuffer.ShadingModelID), 1);
-		case PPI_ShadingModelID:
-			return float4(ScreenSpaceData.GBuffer.ShadingModelID, 0, 0, 0);
-		case PPI_AmbientOcclusion:
-			return ScreenSpaceData.AmbientOcclusion;
-		case PPI_CustomStencil:
-			return ScreenSpaceData.GBuffer.CustomStencil;
-		case PPI_StoredBaseColor:
-			return float4(ScreenSpaceData.GBuffer.StoredBaseColor, 0);
-		case PPI_StoredSpecular:
-			return float4(ScreenSpaceData.GBuffer.StoredSpecular.rrr, 0);
-#if POST_PROCESS_MATERIAL
-		case PPI_Velocity:
-			return float4(PostProcessVelocityLookup(ConvertToDeviceZ(ScreenSpaceData.GBuffer.Depth), UV), 0, 0);
-#endif
-		case PPI_WorldTangent:
-			return float4(ScreenSpaceData.GBuffer.WorldTangent, 0);
-		case PPI_Anisotropy:
-			return ScreenSpaceData.GBuffer.Anisotropy;
-		default:
-			return float4(0, 0, 0, 0);
-	}
-#endif // STRATA_ENABLED
-#endif // MATERIAL_SHADINGMODEL_SINGLELAYERWATER
 #endif // SCENE_TEXTURES_DISABLED
+
+	return float4(0, 0, 0, 0);
 }
 
 #endif // SHADING_PATH_DEFERRED
