Added helper function to volumetric cloud to help with color correction.

diff --git a/Engine/Shaders/Private/VolumetricCloud.usf b/Engine/Shaders/Private/VolumetricCloud.usf
--- a/Engine/Shaders/Private/VolumetricCloud.usf
+++ b/Engine/Shaders/Private/VolumetricCloud.usf
@@ -135,6 +135,42 @@ float SampleAmbientOcclusion(in FPixelMaterialInputs PixelMaterialInputs)
 	return AO;
 }
 
+#ifdef WORKING_COLOR_SPACE_IS_SRGB
+#undef WORKING_COLOR_SPACE_IS_SRGB
+#endif
+#define WORKING_COLOR_SPACE_IS_SRGB 0
+void ConvertCloudPixelMaterialInputsToWorkingColorSpace(inout FPixelMaterialInputs PixelMaterialInputs)
+{
+#if WORKING_COLOR_SPACE_IS_SRGB
+#if STRATA_ENABLED
+	// TODO
+#else
+	// Get raw data
+	float3 ExtinctionSRGB = SampleExtinctionCoefficients(PixelMaterialInputs);
+	float3 AlbedoSRGB = SampleAlbedo(PixelMaterialInputs);
+	float3 EmissiveSRGB = SampleEmissive(PixelMaterialInputs);
+
+	// Compute scattering&absorption coefficients for SRGB
+	float3 ScatteringSRGB = AlbedoSRGB * ExtinctionSRGB;
+	float3 AbsorptionSRGB = max(0.0f, ExtinctionSRGB - ScatteringSRGB);
+
+	// Convert all values to working color space
+	float3 ScatteringColorSpace = ConvertCoefficientsToWorkingColorSpace(ScatteringSRGB);
+	float3 AbsorptionColorSpace = ConvertCoefficientsToWorkingColorSpace(AbsorptionSRGB);
+	float3 EmissiveColorSpace	= ConvertLuminanceToWorkingColorSpace(EmissiveSRGB);
+
+	// Convert back to extinction/albedo for working color space
+	float3 ExtinctionColorSpace = ScatteringColorSpace + AbsorptionColorSpace;
+	float3 AlbedoColorSpace		= ScatteringColorSpace / max(0.000000000001f, ExtinctionColorSpace);
+
+	// Set raw values back
+	PixelMaterialInputs.EmissiveColor	= EmissiveColorSpace;
+	PixelMaterialInputs.Subsurface		= ExtinctionColorSpace;
+	PixelMaterialInputs.BaseColor		= AlbedoColorSpace;
+#endif // STRATA_ENABLED
+#endif // WORKING_COLOR_SPACE_IS_SRGB
+}
+
 #endif
 
 #endif
@@ -828,6 +864,7 @@ void MainCommon(in FMaterialPixelParameters MaterialParameters, in float4 SvPosi
 			// Sample the participating media material at the sample position.
 			//////////////////////////////
 			CalcPixelMaterialInputs(MaterialParameters, PixelMaterialInputs);
+			ConvertCloudPixelMaterialInputsToWorkingColorSpace(PixelMaterialInputs);
 			const MATVEC ExtinctionCoefficients = SampleExtinctionCoefficients(PixelMaterialInputs);
 			const MATVEC EmissiveLuminance = USES_EMISSIVE_COLOR ? SampleEmissive(PixelMaterialInputs) : 0.0f;
 			const MATVEC Albedo = SampleAlbedo(PixelMaterialInputs);
@@ -923,6 +960,7 @@ void MainCommon(in FMaterialPixelParameters MaterialParameters, in float4 SvPosi
 						}
 #endif
 						CalcPixelMaterialInputs(MaterialParameters, PixelMaterialInputs);
+						ConvertCloudPixelMaterialInputsToWorkingColorSpace(PixelMaterialInputs);
 						OpticalDepth += SampleExtinctionCoefficients(PixelMaterialInputs) * ShadowLengthTest * CENTIMETER_TO_METER * DetlaNormT;
 					}
 #else
@@ -942,6 +980,7 @@ void MainCommon(in FMaterialPixelParameters MaterialParameters, in float4 SvPosi
 						}
 #endif
 						CalcPixelMaterialInputs(MaterialParameters, PixelMaterialInputs);
+						ConvertCloudPixelMaterialInputsToWorkingColorSpace(PixelMaterialInputs);
 						OpticalDepth += SampleExtinctionCoefficients(PixelMaterialInputs) * ShadowDtMeter;
 					}
 #endif
@@ -1044,6 +1083,7 @@ void MainCommon(in FMaterialPixelParameters MaterialParameters, in float4 SvPosi
 					}
 
 					CalcPixelMaterialInputs(MaterialParameters, PixelMaterialInputs);
+					ConvertCloudPixelMaterialInputsToWorkingColorSpace(PixelMaterialInputs);
 					MATVEC ShadowExtinctionCoefficients = SampleExtinctionCoefficients(PixelMaterialInputs);
 
 					ParticipatingMediaContext ShadowPMC = SetupParticipatingMediaContext(0.0f, ShadowExtinctionCoefficients, MsScattFactor, MsExtinFactor, 0.0f, 0.0f);
@@ -1108,6 +1148,7 @@ void MainCommon(in FMaterialPixelParameters MaterialParameters, in float4 SvPosi
 						}
 
 						CalcPixelMaterialInputs(MaterialParameters, PixelMaterialInputs);
+						ConvertCloudPixelMaterialInputsToWorkingColorSpace(PixelMaterialInputs);
 						MATVEC ShadowExtinctionCoefficients = SampleExtinctionCoefficients(PixelMaterialInputs);
 
 						ParticipatingMediaContext ShadowPMC = SetupParticipatingMediaContext(0.0f, ShadowExtinctionCoefficients, MsScattFactor, MsExtinFactor, 0.0f, 0.0f);
@@ -1231,6 +1272,7 @@ void MainCommon(in FMaterialPixelParameters MaterialParameters, in float4 SvPosi
 								const float ExtinctionFactor = 1.0f;
 
 								CalcPixelMaterialInputs(MaterialParameters, PixelMaterialInputs);
+								ConvertCloudPixelMaterialInputsToWorkingColorSpace(PixelMaterialInputs);
 								MATVEC ShadowExtinctionCoefficients = SampleExtinctionCoefficients(PixelMaterialInputs);
 
 								if (MaterialParameters.CloudSampleNormAltitudeInLayer < 0.0f || MaterialParameters.CloudSampleNormAltitudeInLayer > 1.0f)
@@ -1958,7 +2000,7 @@ void MainPS(
 		#endif
 
 			CalcPixelMaterialInputs(MaterialParameters, PixelMaterialInputs);
-
+			ConvertCloudPixelMaterialInputsToWorkingColorSpace(PixelMaterialInputs);
 			float3 ShadowExtinctionCoefficients = SampleExtinctionCoefficients(PixelMaterialInputs);
 
 			bool MediumPresent = any(ShadowExtinctionCoefficients > 0.0f);
