Cleanup histogram shader

diff --git a/Engine/Shaders/Private/PostProcessHistogram.usf b/Engine/Shaders/Private/PostProcessHistogram.usf
--- a/Engine/Shaders/Private/PostProcessHistogram.usf
+++ b/Engine/Shaders/Private/PostProcessHistogram.usf
@@ -132,7 +132,7 @@ void MainCS(
 					float LogLuminanceHist = ComputeHistogramPositionFromLogLuminance(LogLuminance);
 
 					// Map the normalized histogram position into texels.
-					float fBucket = saturate(LogLuminanceHist) * (HISTOGRAM_SIZE - 1);// * 0.9999f;
+					float fBucket = saturate(LogLuminanceHist) * (HISTOGRAM_SIZE - 1);
 
 					// Find two discrete buckets that straddle the continuous histogram position.
 					uint Bucket0 = (uint)(fBucket);
@@ -145,7 +145,15 @@ void MainCS(
 					float Weight1 = frac(fBucket);
 					float Weight0 = 1.0f - Weight1;
 
-#if !BILATERAL_GRID
+					// Accumulate the weight to the nearby history buckets.
+
+#if BILATERAL_GRID
+					uint LogLuminanceHist0 = PackTwoFloatToUINT32(Weight0, LogLuminanceHist * Weight0);
+					uint LogLuminanceHist1 = PackTwoFloatToUINT32(Weight1, LogLuminanceHist * Weight1);
+
+					InterlockedAdd(SharedHistogram[Bucket0][GroupThreadId.x][GroupThreadId.y], LogLuminanceHist0);
+					InterlockedAdd(SharedHistogram[Bucket1][GroupThreadId.x][GroupThreadId.y], LogLuminanceHist1);
+#else
 					// When EyeAdaptation_BlackHistogramBucketInfluence=.0, we will ignore the last bucket. The main use
 					// case is so the black background pixels in the editor have no effect. But if we have cases where
 					// pixel values can actually be black, we want to set EyeAdaptation_LastHistogramBucketInfluence=1.0.
@@ -157,17 +165,7 @@ void MainCS(
 
 					Weight0 *= ScreenWeight;
 					Weight1 *= ScreenWeight;
-#endif
 
-					// Accumulate the weight to the nearby history buckets.
-
-#if BILATERAL_GRID
-					uint LogLuminanceHist0 = PackTwoFloatToUINT32(Weight0, LogLuminanceHist * Weight0);
-					uint LogLuminanceHist1 = PackTwoFloatToUINT32(Weight1, LogLuminanceHist * Weight1);
-
-					InterlockedAdd(SharedHistogram[Bucket0][GroupThreadId.x][GroupThreadId.y], LogLuminanceHist0);
-					InterlockedAdd(SharedHistogram[Bucket1][GroupThreadId.x][GroupThreadId.y], LogLuminanceHist1);
-#else
 					uint Weight0Int = PackFloatToUINT32(Weight0);
 					uint Weight1Int = PackFloatToUINT32(Weight1);
 
diff --git a/Engine/Shaders/Private/PostProcessMobile.usf b/Engine/Shaders/Private/PostProcessMobile.usf
--- a/Engine/Shaders/Private/PostProcessMobile.usf
+++ b/Engine/Shaders/Private/PostProcessMobile.usf
@@ -1350,7 +1350,7 @@ void Histogram_MainCS(
 				ScreenWeight = AdaptationWeightTexture(BufferUV);
 
 				// Map the normalized histogram position into texels.
-				fBucket = LogLuminance * HistogramSizeMinusOne;// * 0.9999f;
+				fBucket = LogLuminance * HistogramSizeMinusOne;
 
 				// Find two discrete buckets that straddle the continuous histogram position.
 				Bucket0 = (uint)(fBucket);
