Fix Strata mobile env. specular.

diff --git a/Engine/Shaders/Private/Strata/StrataMobileForwardLighting.ush b/Engine/Shaders/Private/Strata/StrataMobileForwardLighting.ush
--- a/Engine/Shaders/Private/Strata/StrataMobileForwardLighting.ush
+++ b/Engine/Shaders/Private/Strata/StrataMobileForwardLighting.ush
@@ -128,13 +128,14 @@ float3 StrataMobileForwardLighting(
 
 					// And evaluate diffuse parameters
 					FStrataEnvLightResult StrataEnvLight = StrataEvaluateForEnvLight(EnvBSDFContext, false/*bEnableSpecular*/, Settings);
-					float3 DiffuseNormal = StrataEnvLight.DiffuseNormal;
-					float3 DiffuseColorForIndirect = StrataEnvLight.DiffuseWeight;
-					float3 SubsurfaceColorForIndirect = StrataEnvLight.DiffuseBackFaceWeight;
 
 					// Evaluate diffuse lighting
-					if (any((DiffuseColorForIndirect + SubsurfaceColorForIndirect) > 0.0f))
+					// Note: The evaluation is alway done for computing IndirectIrradiance, which is used for ImageBasedReflection evaluation
 					{
+						const float3 DiffuseNormal = StrataEnvLight.DiffuseNormal;
+						const float3 DiffuseColorForIndirect = StrataEnvLight.DiffuseWeight;
+						const float3 SubsurfaceColorForIndirect = StrataEnvLight.DiffuseBackFaceWeight;
+
 						const bool bEvaluateBackface = any(SubsurfaceColorForIndirect > 0.0);
 
 						const half NoL = max(0, dot(StrataBSDFContext.N, MobileDirectionalLight.DirectionalLightDirectionAndShadowTransition.xyz));
@@ -164,11 +165,14 @@ float3 StrataMobileForwardLighting(
 						IndirectIrradiance *= IndirectOcclusion;
 					#endif
 
-						// For diffuse, we specify a perpendicular to the surface light direction for the transmittance to light to not be view dependent.
-						const float DiffuseEnvLightingNoL = 1.0f;
+						if (any((DiffuseColorForIndirect + SubsurfaceColorForIndirect) > 0.0f))
+						{
+							// For diffuse, we specify a perpendicular to the surface light direction for the transmittance to light to not be view dependent.
+							const float DiffuseEnvLightingNoL = 1.0f;
 
-						Color += (DiffuseIndirectLighting * DiffuseColorForIndirect + SubsurfaceIndirectLighting * SubsurfaceColorForIndirect) * 
-								 LuminanceWeight(DiffuseEnvLightingNoL, CurrentBSDF) * AOMultiBounce(StrataEnvLight.DiffuseColor, ShadingOcclusion.DiffOcclusion);
+							Color += (DiffuseIndirectLighting * DiffuseColorForIndirect + SubsurfaceIndirectLighting * SubsurfaceColorForIndirect) * 
+									 LuminanceWeight(DiffuseEnvLightingNoL, CurrentBSDF) * AOMultiBounce(StrataEnvLight.DiffuseColor, ShadingOcclusion.DiffOcclusion);
+						}
 					}
 				}
 
