Add postprocess material support for Strata mobile.

diff --git a/Engine/Shaders/Private/PostProcessMaterialShaders.usf b/Engine/Shaders/Private/PostProcessMaterialShaders.usf
--- a/Engine/Shaders/Private/PostProcessMaterialShaders.usf
+++ b/Engine/Shaders/Private/PostProcessMaterialShaders.usf
@@ -352,14 +352,36 @@ void MainPS(
 	// fill out other related material parameters
 	CalcMaterialParametersPost(Parameters, PixelMaterialInputs, SvPosition, true);
 
+#if STRATA_ENABLED
+
+	// Unlit forces a single BSDF
+	const FStrataBSDF UnlitBSDF = PixelMaterialInputs.FrontMaterial.InlinedBSDF;
+
+	// Grab emissive colour as output
+  #if MATERIAL_OUTPUT_OPACITY_AS_ALPHA
+  // Transmittance = 1 - opacity. See UMaterialExpressionStrataPostProcess.
+	const half UnlitBSDFTransmittance = UNLIT_TRANSMITTANCE(UnlitBSDF).x;
+	const half UnlitBSDFOpacity = 1.0f - UnlitBSDFTransmittance;
+	half Alpha = UnlitBSDFOpacity;
+  #else
+	half Alpha = 0.0f;
+  #endif
+
+	const half3 EmissiveColor = BSDF_GETEMISSIVE(UnlitBSDF);
+
+#else // STRATA_ENABLED
+
 	// Grab emissive colour as output
 	half3 EmissiveColor = GetMaterialEmissive(PixelMaterialInputs);
 	
-#if MATERIAL_OUTPUT_OPACITY_AS_ALPHA
+  #if MATERIAL_OUTPUT_OPACITY_AS_ALPHA
 	half Alpha = GetMaterialOpacity(PixelMaterialInputs);
-#else
+  #else
 	half Alpha = 0.0f;
-#endif
+  #endif
+
+#endif // STRATA_ENABLED
+
 	half4 FullSceneColor = half4(EmissiveColor, Alpha);
 	   
 #if POST_PROCESS_MATERIAL_BEFORE_TONEMAP
