Add precomputed shadow factor for Strata mobile.

diff --git a/Engine/Shaders/Private/Strata/StrataForwardLighting.ush b/Engine/Shaders/Private/Strata/StrataForwardLighting.ush
--- a/Engine/Shaders/Private/Strata/StrataForwardLighting.ush
+++ b/Engine/Shaders/Private/Strata/StrataForwardLighting.ush
@@ -373,6 +373,7 @@ float3 StrataForwardLighting(
 						InitRectTexture(),
 						DirLightingChannelMask,
 						PrimitiveLightingChannelMask,
+						GetPrecomputedShadowFactors(StrataPixelHeader, TranslatedWorldPosition),
 						TranslatedWorldPosition,
 						SceneDepth,
 						LuminanceWeight(StrataBSDFContext, CurrentBSDF),
@@ -441,6 +442,7 @@ float3 StrataForwardLighting(
 						LocalRectTexture,
 						LocalLightChannelMask,
 						PrimitiveLightingChannelMask,
+						GetPrecomputedShadowFactors(StrataPixelHeader, TranslatedWorldPosition),
 						TranslatedWorldPosition,
 						SceneDepth,
 						LuminanceWeight(StrataBSDFContext, CurrentBSDF),
diff --git a/Engine/Shaders/Private/Strata/StrataLightingCommon.ush b/Engine/Shaders/Private/Strata/StrataLightingCommon.ush
--- a/Engine/Shaders/Private/Strata/StrataLightingCommon.ush
+++ b/Engine/Shaders/Private/Strata/StrataLightingCommon.ush
@@ -4,6 +4,19 @@
 
 #ifdef USE_STRATA_FORWARD_LIGHTING_COMMON
 
+// Used by non-mobile path
+half4 GetPrecomputedShadowFactors(FStrataPixelHeader StrataPixelHeader, float3 TranslatedWorldPosition)
+{
+	// PrecomputedShadowFactors are not supported on mobile at the moment
+	float4 OutPrecomputedShadowFactors = HasZeroPrecShadowMask(StrataPixelHeader) ? 0.0f : 1.0f;
+#if TRANSLUCENCY_LIGHTING_SURFACE_FORWARDSHADING && !SHADING_PATH_MOBILE
+	OutPrecomputedShadowFactors.x = ComputeDirectionalLightStaticShadowing(TranslatedWorldPosition).x; // Is that actually correct for forward & deferred?
+#else
+	OutPrecomputedShadowFactors.x = 1;
+#endif
+	return half4(OutPrecomputedShadowFactors);
+}
+
 // Common function for forward lighting per pixel using light data
 float3 StrataForwardLightingCommon(
 	float Dither,
@@ -15,6 +28,7 @@ float3 StrataForwardLightingCommon(
 	FRectTexture RectTexture,
 	uint LightChannelMask,
 	uint PrimitiveLightingChannelMask,
+	half4 PrecomputedShadowFactors,
 	float3 TranslatedWorldPosition,
 	float SceneDepth,
 	float3 BSDFColoredVisibility,
@@ -29,14 +43,6 @@ float3 StrataForwardLightingCommon(
 		// Evaluate the ShadowTerm that can then be used when integrating the lighting
 		FShadowTerms ShadowTerms = { StrataGetAO(StrataPixelHeader), 1.0, 1.0, InitHairTransmittanceData() };
 
-		// PrecomputedShadowFactors are not supported on mobile at the moment
-		float4 PrecomputedShadowFactors = HasZeroPrecShadowMask(StrataPixelHeader) ? 0.0f : 1.0f;
-	#if TRANSLUCENCY_LIGHTING_SURFACE_FORWARDSHADING && !SHADING_PATH_MOBILE
-		PrecomputedShadowFactors.x = ComputeDirectionalLightStaticShadowing(TranslatedWorldPosition).x;
-	#else
-		PrecomputedShadowFactors.x = 1;
-	#endif
-
 		const uint FakeShadingModelID = 0;
 		const float FakeContactShadowOpacity = 1.0f;
 		GetShadowTerms(SceneDepth, PrecomputedShadowFactors, FakeShadingModelID, FakeContactShadowOpacity,
diff --git a/Engine/Shaders/Private/Strata/StrataMobileForwardLighting.ush b/Engine/Shaders/Private/Strata/StrataMobileForwardLighting.ush
--- a/Engine/Shaders/Private/Strata/StrataMobileForwardLighting.ush
+++ b/Engine/Shaders/Private/Strata/StrataMobileForwardLighting.ush
@@ -16,7 +16,7 @@
 
 // For debug or iteration purpose only
 #ifndef STRATA_FORCE_SINGLE_LOCAL_LIGHT
-#define STRATA_FORCE_SINGLE_LOCAL_LIGHT 0
+#define STRATA_FORCE_SINGLE_LOCAL_LIGHT 1
 #endif
 
 // Forward declarations
@@ -60,6 +60,7 @@ float3 StrataMobileForwardLighting(
 	const float3 AbsoluteWorldPosition = LWCToFloat(MaterialParameters.AbsoluteWorldPosition);
 	const float3 TranslatedWorldPosition = MaterialParameters.WorldPosition_CamRelative;
 	const float3 LightingPositionOffset = MaterialParameters.LightingPositionOffset;
+	const half4 PrecomputedShadowFactors = GetPrimaryPrecomputedShadowMask(LightmapVTPageTableResult, Interpolants, MaterialParameters);
 	float3 V = MaterialParameters.CameraVector;
 
 	FStrataAddressing NullStrataAddressing = (FStrataAddressing)0;	// Fake unused in StrataCreateBSDFContext when using Forward inline shading
@@ -244,6 +245,7 @@ float3 StrataMobileForwardLighting(
 						InitRectTexture(),
 						DirLightingChannelMask,
 						PrimitiveLightingChannelMask,
+						PrecomputedShadowFactors,
 						TranslatedWorldPosition,
 						SceneDepth,
 						LuminanceWeight(StrataBSDFContext, CurrentBSDF),
@@ -325,6 +327,7 @@ float3 StrataMobileForwardLighting(
 							LocalRectTexture,
 							LocalLightChannelMask,
 							PrimitiveLightingChannelMask,
+							PrecomputedShadowFactors,
 							TranslatedWorldPosition,
 							SceneDepth,
 							LuminanceWeight(StrataBSDFContext, CurrentBSDF),
diff --git a/Engine/Shaders/Private/Strata/StrataVisualize.usf b/Engine/Shaders/Private/Strata/StrataVisualize.usf
--- a/Engine/Shaders/Private/Strata/StrataVisualize.usf
+++ b/Engine/Shaders/Private/Strata/StrataVisualize.usf
@@ -1082,6 +1082,7 @@ void EvaluateLighting(
 		FRectTexture RectTexture = InitRectTexture();
 		uint LightChannelMask = 0xFFFFFFFF;
 		uint PrimitiveLightingChannelMask = LightChannelMask;
+		half4 PrecomputedShadowFactors = 1;
 
 		float3 BSDFColoredVisibility = 1.0f;
 
@@ -1098,6 +1099,7 @@ void EvaluateLighting(
 			RectTexture,
 			LightChannelMask,
 			PrimitiveLightingChannelMask,
+			PrecomputedShadowFactors,
 			TranslatedWorldPosition,
 			SceneDepth,
 			BSDFColoredVisibility,
