Add SH volume lighting support for Strata mobile.

diff --git a/Engine/Shaders/Private/MobileLightingCommon.ush b/Engine/Shaders/Private/MobileLightingCommon.ush
--- a/Engine/Shaders/Private/MobileLightingCommon.ush
+++ b/Engine/Shaders/Private/MobileLightingCommon.ush
@@ -186,11 +186,15 @@ half MobileDirectionalLightCSM(float2 ScreenPosition, float SceneDepth, inout fl
 	Mobile Translucency Light.
 ------------------------------------------------------------------------------*/
 
-FLightAccumulator GetTranslucencySHLighting(FGBufferData GBuffer, FDeferredLightData LightData, float3 TranslatedWorldPosition)
+struct FTranslucencyLightingVector
 {
-	FLightAccumulator LightAccumulator = (FLightAccumulator)0;
-	float4 VolumeLighting;
-	float3 InterpolatedLighting = 0;
+	float4 AmbientLightingVector;
+	float3 DirectionalLightingVector;
+};
+
+FTranslucencyLightingVector GetTranslucencyLightingVector(FDeferredLightData LightData, float3 TranslatedWorldPosition)
+{
+	FTranslucencyLightingVector Out;
 
 	float Attenuation = 1.0f;
 	float3 L = LightData.Direction;	// Already normalized
@@ -203,26 +207,36 @@ FLightAccumulator GetTranslucencySHLighting(FGBufferData GBuffer, FDeferredLight
 	float3 Lighting = LightData.Color / PI * Attenuation;
 	FTwoBandSHVectorRGB SHLighting = MulSH(SHBasisFunction(NormalizedLightVector), Lighting);
 
-	float4 AmbientLightingVector = float4(SHLighting.R.V.x, SHLighting.G.V.x, SHLighting.B.V.x, 1.0f);
+	Out.AmbientLightingVector = float4(SHLighting.R.V.x, SHLighting.G.V.x, SHLighting.B.V.x, 1.0f);
 
 	float3 LuminanceWeights = float3(.3, .59, .11);
 	float3 Coefficient0 = float3(SHLighting.R.V.y, SHLighting.G.V.y, SHLighting.B.V.y);
 	float3 Coefficient1 = float3(SHLighting.R.V.z, SHLighting.G.V.z, SHLighting.B.V.z);
 	float3 Coefficient2 = float3(SHLighting.R.V.w, SHLighting.G.V.w, SHLighting.B.V.w);
-	float3 DirectionalLightingVector = float3(dot(Coefficient0, LuminanceWeights), dot(Coefficient1, LuminanceWeights), dot(Coefficient2, LuminanceWeights));
+	Out.DirectionalLightingVector = float3(dot(Coefficient0, LuminanceWeights), dot(Coefficient1, LuminanceWeights), dot(Coefficient2, LuminanceWeights));
+
+	return Out;
+}
+
+FLightAccumulator GetTranslucencySHLighting(FGBufferData GBuffer, FDeferredLightData LightData, float3 TranslatedWorldPosition)
+{
+	float4 VolumeLighting;
+	float3 InterpolatedLighting = 0;
+
+	FTranslucencyLightingVector LightingVector = GetTranslucencyLightingVector(LightData, TranslatedWorldPosition);
 
 #if TRANSLUCENCY_LIGHTING_VOLUMETRIC_PERVERTEX_DIRECTIONAL || TRANSLUCENCY_LIGHTING_VOLUMETRIC_DIRECTIONAL || TRANSLUCENCY_LIGHTING_SURFACE_LIGHTINGVOLUME
 
-	GetVolumeLightingDirectional(AmbientLightingVector, DirectionalLightingVector, GBuffer.WorldNormal, GBuffer.DiffuseColor, InterpolatedLighting, VolumeLighting);
+	GetVolumeLightingDirectional(LightingVector.AmbientLightingVector, LightingVector.DirectionalLightingVector, GBuffer.WorldNormal, GBuffer.DiffuseColor, InterpolatedLighting, VolumeLighting);
 
 #elif TRANSLUCENCY_LIGHTING_VOLUMETRIC_PERVERTEX_NONDIRECTIONAL || TRANSLUCENCY_LIGHTING_VOLUMETRIC_NONDIRECTIONAL
 
-	GetVolumeLightingNonDirectional(AmbientLightingVector, GBuffer.DiffuseColor, InterpolatedLighting, VolumeLighting);
+	GetVolumeLightingNonDirectional(LightingVector.AmbientLightingVector, GBuffer.DiffuseColor, InterpolatedLighting, VolumeLighting);
 
 #endif
 
+	FLightAccumulator LightAccumulator = (FLightAccumulator)0;
 	LightAccumulator_AddSplit(LightAccumulator, InterpolatedLighting, 0.0f, InterpolatedLighting, 1.0f, false);
-
 	return LightAccumulator;
 }
 
diff --git a/Engine/Shaders/Private/Strata/StrataMobileForwardLighting.ush b/Engine/Shaders/Private/Strata/StrataMobileForwardLighting.ush
--- a/Engine/Shaders/Private/Strata/StrataMobileForwardLighting.ush
+++ b/Engine/Shaders/Private/Strata/StrataMobileForwardLighting.ush
@@ -204,8 +204,14 @@ float3 StrataMobileForwardLighting(
 				////
 				//// Evaluate translucent lighting volume / vertex lighting.
 				////
-				// This is currently not supported by mobile renderer. See StrataForwardLighting for more details
+				#define STRATA_TRANSLUCENCY_LIGHTING_VOLUMETRIC_DIRECTIONAL    (TRANSLUCENCY_LIGHTING_VOLUMETRIC_PERVERTEX_DIRECTIONAL || TRANSLUCENCY_LIGHTING_VOLUMETRIC_DIRECTIONAL || TRANSLUCENCY_LIGHTING_SURFACE_LIGHTINGVOLUME)
+				#define STRATA_TRANSLUCENCY_LIGHTING_VOLUMETRIC_NONDIRECTIONAL (TRANSLUCENCY_LIGHTING_VOLUMETRIC_PERVERTEX_NONDIRECTIONAL || TRANSLUCENCY_LIGHTING_VOLUMETRIC_NONDIRECTIONAL)
 
+				#if STRATA_TRANSLUCENCY_LIGHTING_VOLUMETRIC_NONDIRECTIONAL || STRATA_TRANSLUCENCY_LIGHTING_VOLUMETRIC_DIRECTIONAL
+				float4 VolumeLighting = 0;
+				float3 InterpolatedLighting = 0;
+				FStrataEvaluateResult BSDFEvaluate = StrataEvaluateBSDF(StrataBSDFContext, Settings);
+				#endif
 
 				////
 				//// Evaluate emissive
@@ -232,6 +238,12 @@ float3 StrataMobileForwardLighting(
 					DirectionalLightCloudShadow = BasePassInterpolants.VertexCloudShadow;
 				#endif
 
+					FTranslucencyLightingVector LightingVector = GetTranslucencyLightingVector(DirLightData, TranslatedWorldPosition);
+				#if STRATA_TRANSLUCENCY_LIGHTING_VOLUMETRIC_DIRECTIONAL
+					GetVolumeLightingDirectional(LightingVector.AmbientLightingVector, LightingVector.DirectionalLightingVector, StrataBSDFContext.N, BSDFEvaluate.DiffuseColor, InterpolatedLighting, VolumeLighting);
+				#elif STRATA_TRANSLUCENCY_LIGHTING_VOLUMETRIC_NONDIRECTIONAL
+					GetVolumeLightingNonDirectional(LightingVector.AmbientLightingVector, BSDFEvaluate.DiffuseColor, InterpolatedLighting, VolumeLighting);
+				#else
 					StrataUpdateBSDFContext(StrataBSDFContext, ToDirLight);
 
 					FStrataEvaluateResult BSDFEvaluate = (FStrataEvaluateResult)0;
@@ -252,6 +264,7 @@ float3 StrataMobileForwardLighting(
 						StrataPixelHeader,
 						StrataBSDFContext,
 						BSDFEvaluate) * DirectionalLightCloudShadow;
+				#endif
 				}
 
 				////
@@ -313,6 +326,12 @@ float3 StrataMobileForwardLighting(
 							LOCAL_LIGHT_EARLY_OUT;
 						}
 
+						FTranslucencyLightingVector LightingVector = GetTranslucencyLightingVector(LightData, TranslatedWorldPosition);
+					#if STRATA_TRANSLUCENCY_LIGHTING_VOLUMETRIC_DIRECTIONAL
+						GetVolumeLightingDirectional(LightingVector.AmbientLightingVector, LightingVector.DirectionalLightingVector, StrataBSDFContext.N, BSDFEvaluate.DiffuseColor, InterpolatedLighting, VolumeLighting);
+					#elif STRATA_TRANSLUCENCY_LIGHTING_VOLUMETRIC_NONDIRECTIONAL
+						GetVolumeLightingNonDirectional(LightingVector.AmbientLightingVector, BSDFEvaluate.DiffuseColor, InterpolatedLighting, VolumeLighting);
+					#else
 						// Update the strata BSDF context accordgin to the new L
 						StrataUpdateBSDFContext(StrataBSDFContext, ToLight);
 
@@ -334,6 +353,7 @@ float3 StrataMobileForwardLighting(
 							StrataPixelHeader,
 							StrataBSDFContext,
 							BSDFEvaluate);
+					#endif
 					}
 				}
 				#endif // !DISABLE_FORWARD_LOCAL_LIGHTS
