Fix hair shading for Strata mobile.

diff --git a/Engine/Shaders/Private/HairBsdf.ush b/Engine/Shaders/Private/HairBsdf.ush
--- a/Engine/Shaders/Private/HairBsdf.ush
+++ b/Engine/Shaders/Private/HairBsdf.ush
@@ -34,11 +34,12 @@ float3 KajiyaKayDiffuseAttenuation(FGBufferData GBuffer, float3 L, float3 V, hal
 	N = FakeNormal;
 
 	// Hack approximation for multiple scattering.
+	float MinValue = 0.0001f;
 	float Wrap = 1;
 	float NoL = saturate((dot(N, L) + Wrap) / Square(1 + Wrap));
 	float DiffuseScatter = (1 / PI) * lerp(NoL, KajiyaDiffuse, 0.33) * GBuffer.Metallic;
 	float Luma = Luminance(GBuffer.BaseColor);
-	float3 ScatterTint = pow(abs(GBuffer.BaseColor / Luma), 1 - Shadow);
+	float3 ScatterTint = pow(abs(GBuffer.BaseColor / max(Luma, MinValue)), clamp(1 - Shadow, MinValue, 1.f));
 	return sqrt(abs(GBuffer.BaseColor)) * DiffuseScatter * ScatterTint;
 }
 
