Add simpler GGX D/Vis evaluation for Strata mobile.

diff --git a/Engine/Shaders/Private/MobileGGX.ush b/Engine/Shaders/Private/MobileGGX.ush
--- a/Engine/Shaders/Private/MobileGGX.ush
+++ b/Engine/Shaders/Private/MobileGGX.ush
@@ -22,7 +22,7 @@
 // Taken from https://gist.github.com/romainguy/a2e9208f14cae37c579448be99f78f25
 // Modified by Epic Games, Inc. To account for premultiplied light color and code style rules.
 
-half GGX_Mobile(half Roughness, float NoH)
+half D_GGX_Mobile(half Roughness, float NoH)
 {
     // Walter et al. 2007, "Microfacet Models for Refraction through Rough Surfaces"
 	float OneMinusNoHSqr = 1.0 - NoH * NoH; 
diff --git a/Engine/Shaders/Private/ShadingModels.ush b/Engine/Shaders/Private/ShadingModels.ush
--- a/Engine/Shaders/Private/ShadingModels.ush
+++ b/Engine/Shaders/Private/ShadingModels.ush
@@ -226,7 +226,7 @@ float3 SpecularGGX( float Roughness, float3 SpecularColor, BxDFContext Context,
 	float Energy = EnergyNormalization( a2, Context.VoH, AreaLight );
 	
 #if SHADING_PATH_MOBILE
-	half D = GGX_Mobile(Roughness, Context.NoH) * Energy;
+	half D = D_GGX_Mobile(Roughness, Context.NoH) * Energy;
 	return MobileSpecularGGXInner(D, SpecularColor, Roughness, Context.NoV, NoL, Context.VoH, MOBILE_HIGH_QUALITY_BRDF);
 #else
 	// Generalized microfacet specular
@@ -248,7 +248,7 @@ half3 DualSpecularGGX(half AverageRoughness, half Lobe0Roughness, half Lobe1Roug
 	float Lobe1Energy = EnergyNormalization(Lobe1Alpha2, Context.VoH, AreaLight);
 
 #if SHADING_PATH_MOBILE
-	half D = lerp(GGX_Mobile(Lobe0Roughness, Context.NoH) * Lobe0Energy, GGX_Mobile(Lobe1Roughness, Context.NoH) * Lobe1Energy, LobeMix);
+	half D = lerp(D_GGX_Mobile(Lobe0Roughness, Context.NoH) * Lobe0Energy, D_GGX_Mobile(Lobe1Roughness, Context.NoH) * Lobe1Energy, LobeMix);
 
 	return MobileSpecularGGXInner(D, SpecularColor, AverageRoughness, Context.NoV, NoL, Context.VoH, MOBILE_HIGH_QUALITY_BRDF);
 #else
@@ -469,7 +469,7 @@ FDirectLighting ClearCoatBxDF( FGBufferData GBuffer, half3 N, half3 V, half3 L,
 		float ClearCoatEnergy = EnergyNormalization(a2, Context.VoH, AreaLight);
 		half Vis = Vis_SmithJointApprox(a2, Context.NoV, NoL);
 #if SHADING_PATH_MOBILE
-		half D = GGX_Mobile(ClearCoatRoughness, Context.NoH) * ClearCoatEnergy;
+		half D = D_GGX_Mobile(ClearCoatRoughness, Context.NoH) * ClearCoatEnergy;
 		Vis = saturate(Vis);
 #else
 		float D = D_GGX(a2, Context.NoH) * ClearCoatEnergy;
@@ -593,7 +593,7 @@ FDirectLighting ClearCoatBxDF( FGBufferData GBuffer, half3 N, half3 V, half3 L,
 			// NoL is chosen to provide better parity with DefaultLit when ClearCoat=0
 			Vis2 = Vis_SmithJointApprox(a2, BottomContext.NoV, NoL);
 #if SHADING_PATH_MOBILE
-			D2 = GGX_Mobile(GBuffer.Roughness, BottomContext.NoH);
+			D2 = D_GGX_Mobile(GBuffer.Roughness, BottomContext.NoH);
 			Vis2 = saturate(Vis2);
 #else
 			D2 = D_GGX(a2, BottomContext.NoH);
@@ -836,7 +836,7 @@ FDirectLighting TwoSidedBxDF( FGBufferData GBuffer, half3 N, half3 V, half3 L, f
 	// Scatter distribution
 	half VoL = dot(V, L);
 #if SHADING_PATH_MOBILE
-	half Scatter = GGX_Mobile(0.6, saturate(-VoL));
+	half Scatter = D_GGX_Mobile(0.6, saturate(-VoL));
 #else
 	float Scatter = D_GGX( 0.6*0.6, saturate( -VoL ) );
 #endif
@@ -899,7 +899,7 @@ FDirectLighting EyeBxDF( FGBufferData GBuffer, half3 N, half3 V, half3 L, float
 		// Generalized microfacet specular
 		float Vis = Vis_SmithJointApprox(a2, Context.NoV, NoL);
 #if SHADING_PATH_MOBILE
-		half D = GGX_Mobile(GBuffer.Roughness, Context.NoH) * Energy;
+		half D = D_GGX_Mobile(GBuffer.Roughness, Context.NoH) * Energy;
 		Vis = saturate(Vis);
 #else
 		float D = D_GGX(a2, Context.NoH) * Energy;
diff --git a/Engine/Shaders/Private/Strata/StrataEvaluation.ush b/Engine/Shaders/Private/Strata/StrataEvaluation.ush
--- a/Engine/Shaders/Private/Strata/StrataEvaluation.ush
+++ b/Engine/Shaders/Private/Strata/StrataEvaluation.ush
@@ -203,21 +203,47 @@ float3 LuminanceWeight(in FStrataBSDFContext InContext, in FStrataBSDF InBSDF)
 	return LuminanceWeight(InContext.SatNoL, InBSDF);
 }
 
-float StrataVisSmithJoint(float a2, float NoV, float NoL)
+half Strata_D_GGX(half Roughness, half a2, half NoH)
 {
-#if STRATA_SHADING_QUALITY > 1
-	return Vis_SmithJointApprox(a2, NoV, NoL);
+#if SHADING_PATH_MOBILE
+	return D_GGX_Mobile(Roughness, NoH);
 #else
+	return D_GGX(a2, NoH);
+#endif
+}
+
+half Strata_D_GGX_Aniso(float ax, float ay, float NoH, float XoH, float YoH)
+{
+	return D_GGXaniso(ax, ay, NoH, XoH, YoH);
+}
+
+half Strata_Vis_GGX(half Roughness, half a2, half NoV, half NoL)
+{
+#if SHADING_PATH_MOBILE && !MOBILE_HIGH_QUALITY_BRDF
+	return (Roughness * 0.25 + 0.25);
+#elif STRATA_SHADING_QUALITY > 1
 	return Vis_SmithJoint(a2, NoV, NoL);
+#else
+	return Vis_SmithJointApprox(a2, NoV, NoL);
 #endif
 }
 
-float StrataVisSmithJointAniso(float ax, float ay, float NoV, float NoL, float XoV, float XoL, float YoV, float YoL)
+half Strata_Vis_GGX_Aniso(float ax, float ay, float NoV, float NoL, float XoV, float XoL, float YoV, float YoL)
 {
 	// STRATA_TODO: build approximation for the aniso version
 	return Vis_SmithJointAniso(ax, ay, NoV, NoL, XoV, XoL, YoV, YoL);
 }
 
+half3 Strata_F_GGX(half3 F0, half3 F90, half VoH)
+{
+// STRATA_TODO: Legacy mobile path use EnvBRDFApprox in place of Fresnel evaluation. However its cost seems higher. Need reevaluate if this is needed.
+//#if SHADING_PATH_MOBILE && !MOBILE_HIGH_QUALITY_BRDF
+//	return EnvBRDFApprox(F0, F90, Roughness, NoV);
+//#else
+	return F_Schlick(F0, F90, VoH);
+//#endif
+}
+
 ///////////////////////////////////////////////////////////////////////////////////////////////////
 // BSDF evaluation (punctual/area)
 
@@ -370,8 +396,8 @@ FStrataEvaluateResult StrataEvaluateBSDFCommon(FStrataBSDFContext BSDFContext, F
 						float2 AlphaXY = 0;
 						GetAnisotropicRoughness(Alpha, SLAB_ANISOTROPY(BSDFContext.BSDF), AlphaXY.x, AlphaXY.y);
 
-						D   = D_GGXaniso(AlphaXY.x, AlphaXY.y, BSDFContext.Context.NoH, BSDFContext.Context.XoH, BSDFContext.Context.YoH);
-						Vis = StrataVisSmithJointAniso(AlphaXY.x, AlphaXY.y, BSDFContext.Context.NoV, BSDFContext.SatNoL, BSDFContext.Context.XoV, BSDFContext.Context.XoL, BSDFContext.Context.YoV, BSDFContext.Context.YoL);
+						D   = Strata_D_GGX_Aniso(AlphaXY.x, AlphaXY.y, BSDFContext.Context.NoH, BSDFContext.Context.XoH, BSDFContext.Context.YoH);
+						Vis = Strata_Vis_GGX_Aniso(AlphaXY.x, AlphaXY.y, BSDFContext.Context.NoV, BSDFContext.SatNoL, BSDFContext.Context.XoV, BSDFContext.Context.XoL, BSDFContext.Context.YoV, BSDFContext.Context.YoL);
 						const float H_PDF = VisibleGGXPDF_aniso(BSDFContext.TangentV, BSDFContext.TangentH, AlphaXY);
 						PDF = RayPDFToReflectionRayPDF(BSDFContext.Context.VoH, H_PDF);
 					}
@@ -399,14 +425,14 @@ FStrataEvaluateResult StrataEvaluateBSDFCommon(FStrataBSDFContext BSDFContext, F
 							float Alpha2Spec = Alpha2;
 							if(IntegrationType == INTEGRATION_PUNCTUAL_LIGHT)
 							{
-								D = D_GGX(Alpha2Spec, BSDFContext.Context.NoH);
+								D = Strata_D_GGX(SafeRoughness, Alpha2Spec, BSDFContext.Context.NoH);
 							}
 							else
 							{
 								const float Energy = EnergyNormalization(Alpha2Spec, BSDFContext.Context.VoH, AreaLightContext.AreaLight);
-								D = D_GGX(Alpha2Spec, BSDFContext.Context.NoH) * Energy;
+								D = Strata_D_GGX(SafeRoughness, Alpha2Spec, BSDFContext.Context.NoH) * Energy;
 							}
-							Vis = StrataVisSmithJoint(Alpha2Spec, BSDFContext.Context.NoV, AreaLightContext.NoL);
+							Vis = Strata_Vis_GGX(SafeRoughness, Alpha2Spec, BSDFContext.Context.NoV, AreaLightContext.NoL);
 							const float H_PDF = VisibleGGXPDF(BSDFContext.TangentV, BSDFContext.TangentH, Alpha2Spec);
 							PDF = RayPDFToReflectionRayPDF(BSDFContext.Context.VoH, H_PDF);
 						}
@@ -423,7 +449,7 @@ FStrataEvaluateResult StrataEvaluateBSDFCommon(FStrataBSDFContext BSDFContext, F
 				}
 				else
 				{
-					const float3 FresnelTerm = F_Schlick(F0, F90, BSDFContext.Context.VoH);
+					const float3 FresnelTerm = Strata_F_GGX(F0, F90, BSDFContext.Context.VoH);
 					Sample.SpecularPathValue = ((SpecularMicroOcclusion * D * Vis) * MSScale) * FresnelTerm;
 				}
 
@@ -467,8 +493,8 @@ FStrataEvaluateResult StrataEvaluateBSDFCommon(FStrataBSDFContext BSDFContext, F
 						float2 HazeAlpha = 0;
 						GetAnisotropicRoughness(HazeSafeRoughness, SLAB_ANISOTROPY(BSDFContext.BSDF), HazeAlpha.x, HazeAlpha.y);
 
-						HazeD   = D_GGXaniso(HazeAlpha.x, HazeAlpha.y, BSDFContext.Context.NoH, BSDFContext.Context.XoH, BSDFContext.Context.YoH);
-						HazeVis = StrataVisSmithJointAniso(HazeAlpha.x, HazeAlpha.y, BSDFContext.Context.NoV, BSDFContext.Context.NoL, BSDFContext.Context.XoV, BSDFContext.Context.XoL, BSDFContext.Context.YoV, BSDFContext.Context.YoL);
+						HazeD   = Strata_D_GGX_Aniso(HazeAlpha.x, HazeAlpha.y, BSDFContext.Context.NoH, BSDFContext.Context.XoH, BSDFContext.Context.YoH);
+						HazeVis = Strata_Vis_GGX_Aniso(HazeAlpha.x, HazeAlpha.y, BSDFContext.Context.NoV, BSDFContext.Context.NoL, BSDFContext.Context.XoV, BSDFContext.Context.XoL, BSDFContext.Context.YoV, BSDFContext.Context.YoL);
 						const float H_PDF = VisibleGGXPDF_aniso(BSDFContext.TangentV, BSDFContext.TangentH, HazeAlpha);
 						HazePDF = RayPDFToReflectionRayPDF(BSDFContext.Context.VoH, H_PDF);
 					}
@@ -492,14 +518,14 @@ FStrataEvaluateResult StrataEvaluateBSDFCommon(FStrataBSDFContext BSDFContext, F
 							float Alpha2SpecHaze = Pow4(HazeSafeRoughness);
 							if (IntegrationType == INTEGRATION_PUNCTUAL_LIGHT)
 							{
-								HazeD = D_GGX(Alpha2SpecHaze, BSDFContext.Context.NoH);
+								HazeD = Strata_D_GGX(HazeSafeRoughness, Alpha2SpecHaze, BSDFContext.Context.NoH);
 							}
 							else
 							{
 								const float Energy = EnergyNormalization(Alpha2SpecHaze, BSDFContext.Context.VoH, AreaLightContext.AreaLight);
-								HazeD = D_GGX(Alpha2SpecHaze, BSDFContext.Context.NoH) * Energy;
+								HazeD = Strata_D_GGX(HazeSafeRoughness, Alpha2SpecHaze, BSDFContext.Context.NoH) * Energy;
 							}
-							HazeVis = StrataVisSmithJoint(Alpha2SpecHaze, BSDFContext.Context.NoV, AreaLightContext.NoL);
+							HazeVis = Strata_Vis_GGX(HazeSafeRoughness, Alpha2SpecHaze, BSDFContext.Context.NoV, AreaLightContext.NoL);
 							const float H_PDF = VisibleGGXPDF(BSDFContext.TangentV, BSDFContext.TangentH, Alpha2SpecHaze);
 							HazePDF = RayPDFToReflectionRayPDF(BSDFContext.Context.VoH, H_PDF);
 						}
@@ -508,7 +534,7 @@ FStrataEvaluateResult StrataEvaluateBSDFCommon(FStrataBSDFContext BSDFContext, F
 
 				
 				const float HazeSpecularMicroOcclusion = F0RGBToMicroOcclusion(HazeF0);
-				const float3 HazeFresnelTerm = F_Schlick(HazeF0, HazeF90, BSDFContext.Context.VoH);
+				const float3 HazeFresnelTerm = Strata_F_GGX(HazeF0, HazeF90, BSDFContext.Context.VoH);
 				if (bIsRectLight)
 				{					
 					Sample.SpecularHazePathValue = RectLightSpecHaze * (HazeMSScale * HazeSpecularMicroOcclusion);
@@ -659,7 +685,7 @@ FStrataEvaluateResult StrataEvaluateBSDFCommon(FStrataBSDFContext BSDFContext, F
 					const float WrapNoL = saturate((-dot(BSDFContext.N, BSDFContext.L) + Wrap) / Square(1 + Wrap));
 					// Scatter distribution
 					const float VoL = dot(BSDFContext.V, BSDFContext.L);
-					const float Scatter = D_GGX(0.6 * 0.6, saturate(-VoL));
+					const float Scatter = Strata_D_GGX(0.6, 0.6 * 0.6, saturate(-VoL));
 
 					TransmissionThroughput = (WrapNoL * Scatter) * SubSurfaceColor;
 				}
@@ -906,20 +932,20 @@ FStrataEvaluateResult StrataEvaluateBSDFCommon(FStrataBSDFContext BSDFContext, F
 				float Alpha2Spec = Alpha2;
 				if(IntegrationType == INTEGRATION_PUNCTUAL_LIGHT)
 				{
-					D = D_GGX(Alpha2Spec, BSDFContext.Context.NoH);
+					D = Strata_D_GGX(SafeRoughness, Alpha2Spec, BSDFContext.Context.NoH);
 				}
 				else
 				{
 					const float Energy = EnergyNormalization(Alpha2Spec, BSDFContext.Context.VoH, AreaLightContext.AreaLight);
-					D = D_GGX(Alpha2Spec, BSDFContext.Context.NoH) * Energy;
+					D = Strata_D_GGX(SafeRoughness, Alpha2Spec, BSDFContext.Context.NoH) * Energy;
 				}
-				Vis = StrataVisSmithJoint(Alpha2Spec, BSDFContext.Context.NoV, AreaLightContext.NoL);
+				Vis = Strata_Vis_GGX(SafeRoughness, Alpha2Spec, BSDFContext.Context.NoV, AreaLightContext.NoL);
 				const float H_PDF = VisibleGGXPDF(BSDFContext.TangentV, BSDFContext.TangentH, Alpha2Spec);
 				PDF = RayPDFToReflectionRayPDF(BSDFContext.Context.VoH, H_PDF);
 			}
 
 			const float SpecularMicroOcclusion = F0RGBToMicroOcclusion(F0);
-			const float3 FresnelTerm = SpecularMicroOcclusion * F_Schlick(F0, F90, BSDFContext.Context.VoH);
+			const float3 FresnelTerm = SpecularMicroOcclusion * Strata_F_GGX(F0, F90, BSDFContext.Context.VoH);
 
 			Sample.SpecularPathProbability	= 1.0f;
 			Sample.SpecularPDF				= PDF * Sample.SpecularPathProbability;
