Fix invalid memory access when no virtual shadow map is present for subsurface material in clustered shading.

#rnx

diff --git a/Engine/Shaders/Private/ClusteredDeferredShadingPixelShader.usf b/Engine/Shaders/Private/ClusteredDeferredShadingPixelShader.usf
--- a/Engine/Shaders/Private/ClusteredDeferredShadingPixelShader.usf
+++ b/Engine/Shaders/Private/ClusteredDeferredShadingPixelShader.usf
@@ -273,7 +273,7 @@ float4 GetLightGridLocalLighting(
 			LightAttenuation = VSMResult.ShadowFactor.xxxx;
 		}
 
-		if ( SubsurfaceOpacity < 1.0f )
+		if ( SubsurfaceOpacity < 1.0f && LocalLight.VirtualShadowMapId != INDEX_NONE )
 		{
 			// Replace SSS portion for subsurface materials
 			half SSSTransmission = ComputeSimpleSubsurfaceTransmissionFromVirtualShadowMap(LocalLight.VirtualShadowMapId, TranslatedWorldPosition, SubsurfaceOpacity);
