Fix incorrect initialization of 'bIsActive' on ShaderPrint context.

diff --git a/Engine/Shaders/Private/ShaderPrint.ush b/Engine/Shaders/Private/ShaderPrint.ush
--- a/Engine/Shaders/Private/ShaderPrint.ush
+++ b/Engine/Shaders/Private/ShaderPrint.ush
@@ -68,7 +68,7 @@ FShaderPrintContext InitShaderPrintContext(bool bActive, float2 InStartPos)
 
 FShaderPrintContext InitShaderPrintContext(bool bActive, uint2 InStartCoord) 
 {
-	return InitShaderPrintContext(true, float2(InStartCoord) / float2(ShaderPrintData.Resolution));
+	return InitShaderPrintContext(bActive, float2(InStartCoord) / float2(ShaderPrintData.Resolution));
 }
 
 FShaderPrintContext InitShaderPrintContext()
diff --git a/Engine/Shaders/Private/ShaderPrintCommon.ush b/Engine/Shaders/Private/ShaderPrintCommon.ush
--- a/Engine/Shaders/Private/ShaderPrintCommon.ush
+++ b/Engine/Shaders/Private/ShaderPrintCommon.ush
@@ -331,7 +331,7 @@ struct FShaderPrintContext
 void ShaderPrint_Internal(FShaderPrintContext Ctx, in FShaderPrintItem Item)
 {
 	// If MaxValueCount is 0 then we don't reset the buffer counter so early out here
-	if (Ctx.Config.MaxValueCount == 0)
+	if (!Ctx.bIsActive || Ctx.Config.MaxValueCount == 0)
 	{
 		return;
 	}
@@ -376,8 +376,11 @@ void PrintSymbol(inout FShaderPrintContext Ctx, in int Symbol) { PrintSymbol(Ctx
 
 void Newline(inout FShaderPrintContext Ctx)
 {
-	Ctx.Pos.x = Ctx.StartPos.x;
-	Ctx.Pos.y += Ctx.Config.FontSpacing.y;
+	if (Ctx.bIsActive)
+	{
+		Ctx.Pos.x = Ctx.StartPos.x;
+		Ctx.Pos.y += Ctx.Config.FontSpacing.y;
+	}
 }
 
 int2 GetCursorPos(FShaderPrintContext Ctx)
