Fix Strata material visualziation alignment issue.

diff --git a/Engine/Shaders/Private/Strata/StrataVisualize.usf b/Engine/Shaders/Private/Strata/StrataVisualize.usf
--- a/Engine/Shaders/Private/Strata/StrataVisualize.usf
+++ b/Engine/Shaders/Private/Strata/StrataVisualize.usf
@@ -168,11 +168,11 @@ void PrintBSDF(inout FShaderPrintContext Context, inout FStrataAddressing Strata
 
 			Print(Context, TEXT(" "), FontSilver);
 			Print(Context, BSDF_GETSHAREDLOCALBASISID(BSDF), FontBSDF, 2, 1);	Print(Context, TEXT("       "), FontSilver);
-			PrintBool(Context, BSDF_GETHASANISOTROPY(BSDF));					Print(Context, TEXT(" "));
-			PrintBool(Context, BSDF_GETISTOPLAYER(BSDF));						Print(Context, TEXT("    "));
-			PrintBool(Context, BSDF_GETSSSTYPE(BSDF) != SSS_TYPE_INVALID);		Print(Context, TEXT("      "));
-			PrintBool(Context, BSDF_GETISSIMPLEVOLUME(BSDF));					Print(Context, TEXT("        "));
-			PrintBool(Context, BSDF_GETISTHIN(BSDF));							Print(Context, TEXT("  "));
+			PrintBool(Context, BSDF_GETHASANISOTROPY(BSDF));					Print(Context, TEXT("  "));
+			PrintBool(Context, BSDF_GETISTOPLAYER(BSDF));						Print(Context, TEXT("     "));
+			PrintBool(Context, BSDF_GETSSSTYPE(BSDF) != SSS_TYPE_INVALID);		Print(Context, TEXT("       "));
+			PrintBool(Context, BSDF_GETISSIMPLEVOLUME(BSDF));					Print(Context, TEXT("         "));
+			PrintBool(Context, BSDF_GETISTHIN(BSDF));							Print(Context, TEXT("   "));
 
 			// View Weight
 			{
@@ -186,6 +186,7 @@ void PrintBSDF(inout FShaderPrintContext Context, inout FStrataAddressing Strata
 					Print(Context, BSDF.LuminanceWeightV.g, FontLightGreen, 6, 4);
 					Print(Context, BSDF.LuminanceWeightV.b, FontLightBlue,  6, 4);
 				}
+				Print(Context, TEXT(" "), FontWhite);
 			}
 
 			// Lighy weight
@@ -461,10 +462,10 @@ void PrintPixel(uint2 InCoord, float3 InWorldPosition, float3 V)
 		Print(Context, IrradianceAndOcclusion.MaterialAO, FontHeaderPropValue, 5, 3);					Print(Context, TEXT(" "));
 		Print(Context, IrradianceAndOcclusion.IndirectIrradiance, FontHeaderPropValue, 6, 3);			Print(Context, TEXT(" "));
 		Print(Context, TopLayerData.Roughness, FontHeaderPropValue, 5, 3);								Print(Context, TEXT("        "));
-		PrintBool(Context, HasPrecShadowMask(Header));													Print(Context, TEXT("     "));
-		PrintBool(Context, HasZeroPrecShadowMask(Header));												Print(Context, TEXT("      "));
-		PrintBool(Context, DoesCastContactShadow(Header));												Print(Context, TEXT("        "));
-		PrintBool(Context, HasDynamicIndirectShadowCasterRepresentation(Header));						Print(Context, TEXT("        "));
+		PrintBool(Context, HasPrecShadowMask(Header));													Print(Context, TEXT("      "));
+		PrintBool(Context, HasZeroPrecShadowMask(Header));												Print(Context, TEXT("       "));
+		PrintBool(Context, DoesCastContactShadow(Header));												Print(Context, TEXT("         "));
+		PrintBool(Context, HasDynamicIndirectShadowCasterRepresentation(Header));						Print(Context, TEXT("         "));
 		PrintBool(Context, HasSubsurface(Header));
 		Newline(Context, RectMax);
 		Newline(Context);
