Fix shader compilation issue on mobile deferred.

diff --git a/Engine/Shaders/Private/MobileLightingCommon.ush b/Engine/Shaders/Private/MobileLightingCommon.ush
--- a/Engine/Shaders/Private/MobileLightingCommon.ush
+++ b/Engine/Shaders/Private/MobileLightingCommon.ush
@@ -400,7 +400,7 @@ float3 GetImageBasedReflectionSpecular(
 	  half Roughness, 
 	  float3 TranslatedWorldPosition, 
 	  half3 WorldNormal, 
-	  FShadingOcclusion Occlusion, 
+	  half SpecularOcclusion,
 	  half IndirectIrradiance, 
 	  uint GridIndex)
 {
@@ -412,7 +412,7 @@ float3 GetImageBasedReflectionSpecular(
 		TranslatedWorldPosition,
 		Roughness,
 		IndirectIrradiance,
-		Occlusion.SpecOcclusion,
+		SpecularOcclusion,
 		GridIndex);
 
 	// Planar reflection
diff --git a/Engine/Shaders/Private/Strata/StrataMobileForwardLighting.ush b/Engine/Shaders/Private/Strata/StrataMobileForwardLighting.ush
--- a/Engine/Shaders/Private/Strata/StrataMobileForwardLighting.ush
+++ b/Engine/Shaders/Private/Strata/StrataMobileForwardLighting.ush
@@ -188,13 +188,13 @@ float3 StrataMobileForwardLighting(
 					// Evaluate specular lighting for the main lobe
 					if (any(StrataEnvLight.SpecularWeight > 0.0f))
 					{
-						Color += GetImageBasedReflectionSpecular(StrataEnvLight.SpecularDirection, StrataEnvLight.SpecularSafeRoughness, MaterialParameters.WorldPosition_CamRelative, StrataBSDFContext.N, ShadingOcclusion, IndirectIrradiance, GridIndex) *
+						Color += GetImageBasedReflectionSpecular(StrataEnvLight.SpecularDirection, StrataEnvLight.SpecularSafeRoughness, MaterialParameters.WorldPosition_CamRelative, StrataBSDFContext.N, ShadingOcclusion.SpecularOcclusion, IndirectIrradiance, GridIndex) *
 								 StrataEnvLight.SpecularWeight * IndirectOcclusion * LuminanceWeight(StrataBSDFContext, CurrentBSDF) * AOMultiBounce(StrataEnvLight.SpecularColor, ShadingOcclusion.SpecOcclusion);
 					}
 					// And for the second lobe if needed
 					if (any(StrataEnvLight.SpecularHazeWeight > 0.0f))
 					{
-						Color += GetImageBasedReflectionSpecular(StrataEnvLight.SpecularHazeDirection, StrataEnvLight.SpecularHazeSafeRoughness, MaterialParameters.WorldPosition_CamRelative, StrataBSDFContext.N, ShadingOcclusion, IndirectIrradiance, GridIndex) *
+						Color += GetImageBasedReflectionSpecular(StrataEnvLight.SpecularHazeDirection, StrataEnvLight.SpecularHazeSafeRoughness, MaterialParameters.WorldPosition_CamRelative, StrataBSDFContext.N, ShadingOcclusion.SpecularOcclusion, IndirectIrradiance, GridIndex) *
 								 StrataEnvLight.SpecularHazeWeight * IndirectOcclusion * LuminanceWeight(StrataBSDFContext, CurrentBSDF) * AOMultiBounce(StrataEnvLight.SpecularColor, ShadingOcclusion.SpecOcclusion);
 					}
 				}
