Remove StrataIntegrateBSDF to avoid an intermediate step.

diff --git a/Engine/Shaders/Private/Strata/StrataDeferredLighting.ush b/Engine/Shaders/Private/Strata/StrataDeferredLighting.ush
--- a/Engine/Shaders/Private/Strata/StrataDeferredLighting.ush
+++ b/Engine/Shaders/Private/Strata/StrataDeferredLighting.ush
@@ -70,7 +70,7 @@ FStrataDeferredLighting StrataDeferredLighting(
 
 			// We must have the StrataIntegrateBSDF inside the if due to the rectlight texture: it must be non ambiguous which texture is going to be used.
 			// After the compilation, a local resource must map to a unique global resource (the default or the actual rect light texture).
-			BSDFEvaluate = StrataIntegrateBSDF(BSDFContext, ShadowTerms, AreaLightContext, Settings, INTEGRATION_AREA_LIGHT_RECT);
+			BSDFEvaluate = StrataEvaluateBSDFCommon(BSDFContext, ShadowTerms, AreaLightContext, Settings, INTEGRATION_AREA_LIGHT_RECT);
 		}
 		else
 		{
@@ -80,11 +80,11 @@ FStrataDeferredLighting StrataDeferredLighting(
 			BRANCH
 			if(IsAreaLight(AreaLightContext.AreaLight))
 			{
-				BSDFEvaluate = StrataIntegrateBSDF(BSDFContext, ShadowTerms, AreaLightContext, Settings, INTEGRATION_AREA_LIGHT_CAPSULE);
+				BSDFEvaluate = StrataEvaluateBSDFCommon(BSDFContext, ShadowTerms, AreaLightContext, Settings, INTEGRATION_AREA_LIGHT_CAPSULE);
 			}
 			else
 			{
-				BSDFEvaluate = StrataIntegrateBSDF(BSDFContext, ShadowTerms, AreaLightContext, Settings, INTEGRATION_PUNCTUAL_LIGHT);
+				BSDFEvaluate = StrataEvaluateBSDFCommon(BSDFContext, ShadowTerms, AreaLightContext, Settings, INTEGRATION_PUNCTUAL_LIGHT);
 			}
 		}
 
diff --git a/Engine/Shaders/Private/Strata/StrataEvaluation.ush b/Engine/Shaders/Private/Strata/StrataEvaluation.ush
--- a/Engine/Shaders/Private/Strata/StrataEvaluation.ush
+++ b/Engine/Shaders/Private/Strata/StrataEvaluation.ush
@@ -997,22 +997,6 @@ FStrataEvaluateResult StrataEvaluateBSDF(FStrataBSDFContext BSDFContext, FStrata
 	return StrataEvaluateBSDFCommon(BSDFContext, IdentityShadow, DummyAreaLightContext, Settings, INTEGRATION_PUNCTUAL_LIGHT);
 }
 
-// Area light evaluation
-FStrataEvaluateResult StrataIntegrateBSDF(FStrataBSDFContext BSDFContext, FShadowTerms Shadow, FAreaLightIntegrateContext AreaLightContext, FStrataIntegrationSettings Settings, uint IntegrationMode)
-{
-	BRANCH
-	if (IntegrationMode == INTEGRATION_AREA_LIGHT_RECT)
-	{
-		return StrataEvaluateBSDFCommon(BSDFContext, Shadow, AreaLightContext, Settings, INTEGRATION_AREA_LIGHT_RECT);
-	}
-	else if (IntegrationMode == INTEGRATION_AREA_LIGHT_CAPSULE)
-	{
-		return StrataEvaluateBSDFCommon(BSDFContext, Shadow, AreaLightContext, Settings, INTEGRATION_AREA_LIGHT_CAPSULE);
-	}
-
-	return StrataEvaluateBSDFCommon(BSDFContext, Shadow, AreaLightContext, Settings, INTEGRATION_PUNCTUAL_LIGHT);
-}
-
 ///////////////////////////////////////////////////////////////////////////////////////////////////
 // BSDF Importance sampling
  
diff --git a/Engine/Shaders/Private/Strata/StrataLightingCommon.ush b/Engine/Shaders/Private/Strata/StrataLightingCommon.ush
--- a/Engine/Shaders/Private/Strata/StrataLightingCommon.ush
+++ b/Engine/Shaders/Private/Strata/StrataLightingCommon.ush
@@ -69,7 +69,7 @@ float3 StrataForwardLightingCommon(
 
 				// We must have the evaluate inside the if due to the rectlight texture: it must be now be ambiguous which texture is going ot be used.
 				// After te compilation, a local resource must map to a unique global resource (the default or the actual rect light texture).
-				BSDFEvaluate = StrataIntegrateBSDF(StrataBSDFContext, ShadowTerms, AreaLightContext, Settings, INTEGRATION_AREA_LIGHT_RECT);
+				BSDFEvaluate = StrataEvaluateBSDFCommon(StrataBSDFContext, ShadowTerms, AreaLightContext, Settings, INTEGRATION_AREA_LIGHT_RECT);
 			}
 			else
 			{
@@ -80,11 +80,11 @@ float3 StrataForwardLightingCommon(
 				BRANCH
 				if(IsAreaLight(AreaLightContext.AreaLight))
 				{
-					BSDFEvaluate = StrataIntegrateBSDF(StrataBSDFContext, ShadowTerms, AreaLightContext, Settings, INTEGRATION_AREA_LIGHT_CAPSULE);
+					BSDFEvaluate = StrataEvaluateBSDFCommon(StrataBSDFContext, ShadowTerms, AreaLightContext, Settings, INTEGRATION_AREA_LIGHT_CAPSULE);
 				}
 				else
 				{
-					BSDFEvaluate = StrataIntegrateBSDF(StrataBSDFContext, ShadowTerms, AreaLightContext, Settings, INTEGRATION_PUNCTUAL_LIGHT);
+					BSDFEvaluate = StrataEvaluateBSDFCommon(StrataBSDFContext, ShadowTerms, AreaLightContext, Settings, INTEGRATION_PUNCTUAL_LIGHT);
 				}
 			}
 
@@ -107,7 +107,7 @@ float3 StrataForwardLightingCommon(
 
 				// We must have the evaluate inside the if due to the rectlight texture: it must be now be ambiguous which texture is going ot be used.
 				// After te compilation, a local resource must map to a unique global resource (the default or the actual rect light texture).
-				BSDFEvaluate = StrataIntegrateBSDF(StrataBSDFContext, ShadowTerms, AreaLightContext, Settings, INTEGRATION_AREA_LIGHT_RECT);
+				BSDFEvaluate = StrataEvaluateBSDFCommon(StrataBSDFContext, ShadowTerms, AreaLightContext, Settings, INTEGRATION_AREA_LIGHT_RECT);
 			}
 			else
 			{
@@ -117,11 +117,11 @@ float3 StrataForwardLightingCommon(
 				BRANCH
 				if(IsAreaLight(AreaLightContext.AreaLight))
 				{
-					BSDFEvaluate = StrataIntegrateBSDF(StrataBSDFContext, ShadowTerms, AreaLightContext, Settings, INTEGRATION_AREA_LIGHT_CAPSULE);
+					BSDFEvaluate = StrataEvaluateBSDFCommon(StrataBSDFContext, ShadowTerms, AreaLightContext, Settings, INTEGRATION_AREA_LIGHT_CAPSULE);
 				}
 				else
 				{
-					BSDFEvaluate = StrataIntegrateBSDF(StrataBSDFContext, ShadowTerms, AreaLightContext, Settings, INTEGRATION_PUNCTUAL_LIGHT);
+					BSDFEvaluate = StrataEvaluateBSDFCommon(StrataBSDFContext, ShadowTerms, AreaLightContext, Settings, INTEGRATION_PUNCTUAL_LIGHT);
 				}
 			}
 
