Strata - fixed forward light and image based specular define for mobile.

diff --git a/Engine/Shaders/Private/Strata/StrataMobileForwardLighting.ush b/Engine/Shaders/Private/Strata/StrataMobileForwardLighting.ush
--- a/Engine/Shaders/Private/Strata/StrataMobileForwardLighting.ush
+++ b/Engine/Shaders/Private/Strata/StrataMobileForwardLighting.ush
@@ -179,7 +179,7 @@ float3 StrataMobileForwardLighting(
 
 				// 2. Specular evaluation (using regular normal)
 				//#if MATERIALBLENDING_MASKED || MATERIALBLENDING_SOLID || TRANSLUCENCY_LIGHTING_SURFACE_FORWARDSHADING || MATERIAL_SHADINGMODEL_SINGLELAYERWATER
-				#if FORWARD_PER_PIXEL_SHADING || TRANSLUCENCY_LIGHTING_SURFACE_LIGHTINGVOLUME
+				#if ENABLE_CLUSTERED_REFLECTION && (FORWARD_PER_PIXEL_SHADING || TRANSLUCENCY_LIGHTING_SURFACE_LIGHTINGVOLUME)
 				if (UseBasePassSkylightSpecular > 0)
 				{
 					// Evaluate specular parameters without bent normal
@@ -270,7 +270,7 @@ float3 StrataMobileForwardLighting(
 				////
 				//// Evaluate local lights
 				////
-				#if !DISABLE_FORWARD_LOCAL_LIGHTS && (FORWARD_SHADING || TRANSLUCENCY_LIGHTING_SURFACE_FORWARDSHADING) // Forbidden for Single Layer Water
+				#if ENABLE_CLUSTERED_LIGHTS && !DISABLE_FORWARD_LOCAL_LIGHTS && (FORWARD_SHADING || TRANSLUCENCY_LIGHTING_SURFACE_FORWARDSHADING) // Forbidden for Single Layer Water
 				{
 					// Basic implementation for FORWARD_LOCAL_LIGHTS where local lights are reloaded per BSDF layer. It is fast by default when layer==1.
 					const FCulledLightsGridData CulledLightsGrid = GetCulledLightsGrid(GridIndex, EyeIndex);
