Honor bReceivesDecal flag when doing manual DBuffer samples in material.

diff --git a/Engine/Shaders/Private/MaterialTemplate.ush b/Engine/Shaders/Private/MaterialTemplate.ush
--- a/Engine/Shaders/Private/MaterialTemplate.ush
+++ b/Engine/Shaders/Private/MaterialTemplate.ush
@@ -2501,29 +2501,31 @@ float3 DecodeSceneColorForMaterialNode(float2 ScreenUV)
 
 #endif // SHADING_PATH_DEFERRED
 
-float4 MaterialExpressionDBufferTextureLookup(float2 BufferUV, int DBufferTextureIndex)
+float4 MaterialExpressionDBufferTextureLookup(FMaterialPixelParameters Parameters, float2 BufferUV, int DBufferTextureIndex)
 {
-	uint2 PixelPos = uint2(BufferUV * View.BufferSizeAndInvSize.xy);
-	uint ValidDBufferTargetMask = GetDBufferTargetMask(PixelPos) & (1u << DBufferTextureIndex);
-
-#if STRATA_ENABLED
-	const FStrataDBuffer DBufferData = GetDBufferData(BufferUV, DBufferTextureIndex);
-	switch (DBufferTextureIndex)
+	if ((GetPrimitiveData(Parameters).Flags & PRIMITIVE_SCENE_DATA_FLAG_DECAL_RECEIVER) != 0 && View.ShowDecalsMask > 0)
 	{
-	// All data (BaseColor, WorldNormal, Roughness) are pre-multiplied by their coverage
-	case 0:	return float4(DBufferData.BaseColor,	DBufferData.OneMinusCoverage_BaseColor);
-	case 1:	return float4(DBufferData.WorldNormal,	DBufferData.OneMinusCoverage_WorldNormal);
-	case 2:	return float4(DBufferData.Roughness,	DBufferData.Metallic, DBufferData.Specular, DBufferData.OneMinusCoverage_Roughness);
-	}
+#if STRATA_ENABLED
+		const FStrataDBuffer DBufferData = GetDBufferData(BufferUV, DBufferTextureIndex);
+		switch (DBufferTextureIndex)
+		{
+		// All data (BaseColor, WorldNormal, Roughness) are pre-multiplied by their coverage
+		case 0:	return float4(DBufferData.BaseColor,	DBufferData.OneMinusCoverage_BaseColor);
+		case 1:	return float4(DBufferData.WorldNormal,	DBufferData.OneMinusCoverage_WorldNormal);
+		case 2:	return float4(DBufferData.Roughness,	DBufferData.Metallic, DBufferData.Specular, DBufferData.OneMinusCoverage_Roughness);
+		}
 #else
-	FDBufferData DBufferData = GetDBufferData(BufferUV, ValidDBufferTargetMask);
-	switch (DBufferTextureIndex)
-	{
-	case 0:	return float4(DBufferData.PreMulColor, DBufferData.ColorOpacity);
-	case 1:	return float4(DBufferData.PreMulWorldNormal, DBufferData.NormalOpacity);
-	case 2:	return float4(DBufferData.PreMulRoughness, DBufferData.PreMulMetallic, DBufferData.PreMulSpecular, DBufferData.RoughnessOpacity);
-	}
+		uint2 PixelPos = uint2(BufferUV * View.BufferSizeAndInvSize.xy);
+		uint ValidDBufferTargetMask = GetDBufferTargetMask(PixelPos) & (1u << DBufferTextureIndex);
+		FDBufferData DBufferData = GetDBufferData(BufferUV, ValidDBufferTargetMask);
+		switch (DBufferTextureIndex)
+		{
+		case 0:	return float4(DBufferData.PreMulColor, DBufferData.ColorOpacity);
+		case 1:	return float4(DBufferData.PreMulWorldNormal, DBufferData.NormalOpacity);
+		case 2:	return float4(DBufferData.PreMulRoughness, DBufferData.PreMulMetallic, DBufferData.PreMulSpecular, DBufferData.RoughnessOpacity);
+		}
 #endif
+	}
 
 	return float4(0, 0, 0, 1);
 }
diff --git a/Engine/Source/Runtime/Engine/Private/Materials/HLSLMaterialTranslator.cpp b/Engine/Source/Runtime/Engine/Private/Materials/HLSLMaterialTranslator.cpp
--- a/Engine/Source/Runtime/Engine/Private/Materials/HLSLMaterialTranslator.cpp
+++ b/Engine/Source/Runtime/Engine/Private/Materials/HLSLMaterialTranslator.cpp
@@ -6819,7 +6819,7 @@ int32 FHLSLMaterialTranslator::DBufferTextureLookup(int32 ViewportUV, uint32 DBu
 	MaterialCompilationOutput.SetIsDBufferTextureUsed(DBufferTextureIndex);
 	AddEstimatedTextureSample();
 
-	return AddCodeChunk(MCT_Float4, TEXT("MaterialExpressionDBufferTextureLookup(%s, %d)"), *CoerceParameter(BufferUV, MCT_Float2), (int)DBufferTextureIndex);
+	return AddCodeChunk(MCT_Float4, TEXT("MaterialExpressionDBufferTextureLookup(Parameters, %s, %d)"), *CoerceParameter(BufferUV, MCT_Float2), (int)DBufferTextureIndex);
 }
 
 int32 FHLSLMaterialTranslator::Texture(UTexture* InTexture, int32& TextureReferenceIndex, EMaterialSamplerType SamplerType, ESamplerSourceMode SamplerSource, ETextureMipValueMode MipValueMode)
