PR #9802: [Mobile] Mobile Deferred Renderer Fix: support particle lighting & light function fix (Contributed by [REDACTED_PERSON])

diff --git a/Engine/Plugins/FX/Niagara/Source/Niagara/Private/NiagaraRendererLights.cpp b/Engine/Plugins/FX/Niagara/Source/Niagara/Private/NiagaraRendererLights.cpp
--- a/Engine/Plugins/FX/Niagara/Source/Niagara/Private/NiagaraRendererLights.cpp
+++ b/Engine/Plugins/FX/Niagara/Source/Niagara/Private/NiagaraRendererLights.cpp
@@ -56,7 +56,8 @@ FPrimitiveViewRelevance FNiagaraRendererLights::GetViewRelevance(const FSceneVie
 FNiagaraDynamicDataBase* FNiagaraRendererLights::GenerateDynamicData(const FNiagaraSceneProxy* Proxy, const UNiagaraRendererProperties* InProperties, const FNiagaraEmitterInstance* Emitter) const
 {
 	// particle (simple) lights are only supported with deferred shading
-	if (!bHasLights || Proxy->GetScene().GetShadingPath() != EShadingPath::Deferred)
+	
+	if (!bHasLights || (Proxy->GetScene().GetShadingPath() != EShadingPath::Deferred && !IsMobileDeferredShadingEnabled(Proxy->GetScene().GetShaderPlatform())))
 	{
 		return nullptr;
 	}
diff --git a/Engine/Shaders/Private/MobileDeferredShading.usf b/Engine/Shaders/Private/MobileDeferredShading.usf
--- a/Engine/Shaders/Private/MobileDeferredShading.usf
+++ b/Engine/Shaders/Private/MobileDeferredShading.usf
@@ -178,13 +178,15 @@ void MobileDirectionalLightPS(
 	half4 DynamicShadowFactors = 1.0f;
 	float DirectionalLightShadow = 1.0f;
 	AccumulateDirectionalLighting(GBuffer, TranslatedWorldPosition, CameraVector, ScreenPosition, SvPosition, DynamicShadowFactors, DirectionalLightShadow, DirectLighting);
+	// LightFunction should only affect direct lighting result
+	DirectLighting.TotalLight *= ComputeLightFunctionMultiplier(TranslatedWorldPosition);
 
 #if ENABLE_CLUSTERED_REFLECTION || ENABLE_SKY_LIGHT || ENABLE_PLANAR_REFLECTION
 	// If we have a single directional light, apply relfection and sky contrubution here
 	ReflectionEnvironmentSkyLighting(GBuffer, CameraVector, TranslatedWorldPosition, ReflectionVector, GridIndex, DirectLighting);
 #endif
 
-	half3 Color = DirectLighting.TotalLight * ComputeLightFunctionMultiplier(TranslatedWorldPosition);
+	half3 Color = DirectLighting.TotalLight;
 	// MobileHDR applies PreExposure in tonemapper
 	Color *= View.PreExposure;
 
