Fixed: missing WPO on mobile with GPUScene enabled

diff --git a/Engine/Shaders/Private/InstanceCulling/BuildInstanceDrawCommands.usf b/Engine/Shaders/Private/InstanceCulling/BuildInstanceDrawCommands.usf
--- a/Engine/Shaders/Private/InstanceCulling/BuildInstanceDrawCommands.usf
+++ b/Engine/Shaders/Private/InstanceCulling/BuildInstanceDrawCommands.usf
@@ -128,7 +128,7 @@ void WritePackedInstanceData(uint Offset, uint InstanceId, FInstanceSceneData In
 	InstanceAuxData1 = 0; // TODO: pack something useful here 
 #endif // ALLOW_STATIC_LIGHTING
 
-	const uint PackedPrimitiveFlags = ((PrimitiveData.Flags & 0xffff) | (InstanceData.Flags << 16u));
+	const uint PackedPrimitiveFlags = ((PrimitiveData.Flags & 0x01FFFFFF) | (InstanceData.Flags << 25u));
 
 	// This will only write out the relative instance transform, but that's fine; the tile coordinate of the instance transform comes from the primitive data
 	float4 InstanceOriginAndId = InstanceData.LocalToWorld.M[3];
diff --git a/Engine/Shaders/Private/SceneData.ush b/Engine/Shaders/Private/SceneData.ush
--- a/Engine/Shaders/Private/SceneData.ush
+++ b/Engine/Shaders/Private/SceneData.ush
@@ -1133,9 +1133,10 @@ FSceneDataIntermediates GetSceneDataIntermediates(uint DrawInstanceId, float4 In
 	// TODO: add support for LWC, we should pack it as integer tile coordinates 
 	float3 TilePosition = float3(0,0,0);
 
-	// TODO: pack important primitive and instance flags here
-	const uint PrimitiveFlags = asuint(InstanceTransform1.w);
-	const uint InstanceFlags = (PrimitiveFlags >> 16);
+	// TODO: pack important primitive and instance flags here 
+	// (25 + 7)
+	const uint PrimitiveFlags = asuint(InstanceTransform1.w) & 0x01FFFFFF;
+	const uint InstanceFlags = (PrimitiveFlags >> 25u);
 
 	// Reconstruct InstanceData from a packed data
 	FInstanceSceneData InstanceData = (FInstanceSceneData)0;
