Add Legacy foliage support for Strata.

diff --git a/Engine/Shaders/Private/Strata/Strata.ush b/Engine/Shaders/Private/Strata/Strata.ush
--- a/Engine/Shaders/Private/Strata/Strata.ush
+++ b/Engine/Shaders/Private/Strata/Strata.ush
@@ -348,6 +348,7 @@ struct FStrataBSDF
 #define SINGLE_OPTLEGACYMODE_CLOTH				2
 #define SINGLE_OPTLEGACYMODE_SSSWRAP			3
 #define SINGLE_OPTLEGACYMODE_SSSPROFILE			4
+#define SINGLE_OPTLEGACYMODE_SSSWRAP_THIN		5
 // 5 more remains: Also how to deal with Aniso?
 
 // Layout 0 (Fast)
@@ -3829,7 +3830,7 @@ FStrataBSDF UnpackStrataBSDFIn(FStrataMaterialContainer StrataBuffer, inout FStr
 				BSDF_SETHASFUZZ(OutBSDF, 1);
 				BSDF_SETHASTRANSABOVE(OutBSDF, 0);
 			}
-			else if (OptimisedLegacyMode == SINGLE_OPTLEGACYMODE_SSSWRAP)
+			else if (OptimisedLegacyMode == SINGLE_OPTLEGACYMODE_SSSWRAP || OptimisedLegacyMode == SINGLE_OPTLEGACYMODE_SSSWRAP_THIN)
 			{
 			#if STRATA_DEFERRED_SHADING
 				uint Data0 = StrataHeader.SharedLocalBasesTypes_PackedHeader;
@@ -3864,7 +3865,7 @@ FStrataBSDF UnpackStrataBSDFIn(FStrataMaterialContainer StrataBuffer, inout FStr
 				BSDF_SETHASMFP(OutBSDF, 1);
 				BSDF_SETHASHAZINESS(OutBSDF, 0);
 				BSDF_SETISSIMPLEVOLUME(OutBSDF, 0);
-				BSDF_SETISTHIN(OutBSDF, 0);
+				BSDF_SETISTHIN(OutBSDF, (OptimisedLegacyMode == SINGLE_OPTLEGACYMODE_SSSWRAP_THIN ? 1 : 0));
 				BSDF_SETHASFUZZ(OutBSDF, 0);
 				BSDF_SETHASTRANSABOVE(OutBSDF, 0);
 			}
diff --git a/Engine/Shaders/Private/Strata/StrataExport.ush b/Engine/Shaders/Private/Strata/StrataExport.ush
--- a/Engine/Shaders/Private/Strata/StrataExport.ush
+++ b/Engine/Shaders/Private/Strata/StrataExport.ush
@@ -355,6 +355,12 @@ void PackStrataOut(
 				{
 					OptimisedLegacyMode = SINGLE_OPTLEGACYMODE_CLOTH;
 				}
+				
+				else if ((OptimisedModeMask & SingleIsThin) == SingleIsThin)
+				{
+					OptimisedLegacyMode = SINGLE_OPTLEGACYMODE_SSSWRAP_THIN;
+					SSSData.Header.Bytes = 0; // Force to skip SSSData writing
+				}
 				else if ((OptimisedModeMask & SingleIsWrap) == SingleIsWrap)
 				{
 					OptimisedLegacyMode = SINGLE_OPTLEGACYMODE_SSSWRAP;
@@ -574,7 +580,7 @@ void PackStrataOut(
 							// Ensure the rest of the BSDF is not stored with the regular path
 							BSDF.State = 0;
 						}
-						else if (OptimisedLegacyMode == SINGLE_OPTLEGACYMODE_SSSWRAP)
+						else if (OptimisedLegacyMode == SINGLE_OPTLEGACYMODE_SSSWRAP || OptimisedLegacyMode == SINGLE_OPTLEGACYMODE_SSSWRAP_THIN)
 						{
 							const float SSSWOpacity = StrataSubSurfaceGetWrapOpacityFromAnisotropy(SLAB_SSSPHASEANISOTROPY(BSDF));
 
diff --git a/Engine/Shaders/Private/Strata/StrataVisualize.usf b/Engine/Shaders/Private/Strata/StrataVisualize.usf
--- a/Engine/Shaders/Private/Strata/StrataVisualize.usf
+++ b/Engine/Shaders/Private/Strata/StrataVisualize.usf
@@ -447,6 +447,10 @@ void PrintPixel(uint2 InCoord, float3 InWorldPosition, float3 V)
 			{
 				Print(Context, TEXT(" - Legecy SSS-Wrap"), FontLightRed);
 			}
+			else if (OptimisedLegacyMode == SINGLE_OPTLEGACYMODE_SSSWRAP_THIN)
+			{
+				Print(Context, TEXT(" - Legecy Foliage"), FontLightRed);
+			}
 			else if (OptimisedLegacyMode == SINGLE_OPTLEGACYMODE_SSSPROFILE)
 			{
 				Print(Context, TEXT(" - Legecy SSS-Profile"), FontLightRed);
