Strata - fixed typos

diff --git a/Engine/Shaders/Private/Strata/StrataVisualize.usf b/Engine/Shaders/Private/Strata/StrataVisualize.usf
--- a/Engine/Shaders/Private/Strata/StrataVisualize.usf
+++ b/Engine/Shaders/Private/Strata/StrataVisualize.usf
@@ -437,23 +437,23 @@ void PrintPixel(uint2 InCoord, float3 InWorldPosition, float3 V)
 			uint OptimisedLegacyMode = (Header.SharedLocalBasesTypes_PackedHeader >> (HEADER_SINGLEENCODING_BIT_COUNT)) & HEADER_SINGLE_OPTLEGACYMODE_BIT_MASK;
 			if (OptimisedLegacyMode ==SINGLE_OPTLEGACYMODE_CLEARCOAT)
 			{
-				Print(Context, TEXT(" - Legecy Clear-Coat"), FontLightRed);
+				Print(Context, TEXT(" - Legacy Clear-Coat"), FontLightRed);
 			}
 			else if (OptimisedLegacyMode ==SINGLE_OPTLEGACYMODE_CLOTH)
 			{
-				Print(Context, TEXT(" - Legecy Cloth"), FontLightRed);
+				Print(Context, TEXT(" - Legacy Cloth"), FontLightRed);
 			}
 			else if (OptimisedLegacyMode ==SINGLE_OPTLEGACYMODE_SSSWRAP)
 			{
-				Print(Context, TEXT(" - Legecy SSS-Wrap"), FontLightRed);
+				Print(Context, TEXT(" - Legacy SSS-Wrap"), FontLightRed);
 			}
 			else if (OptimisedLegacyMode == SINGLE_OPTLEGACYMODE_SSSWRAP_THIN)
 			{
-				Print(Context, TEXT(" - Legecy Foliage"), FontLightRed);
+				Print(Context, TEXT(" - Legacy Foliage"), FontLightRed);
 			}
 			else if (OptimisedLegacyMode == SINGLE_OPTLEGACYMODE_SSSPROFILE)
 			{
-				Print(Context, TEXT(" - Legecy SSS-Profile"), FontLightRed);
+				Print(Context, TEXT(" - Legacy SSS-Profile"), FontLightRed);
 			}
 		#endif
 		}
