Fix a bug that GPUScene breaks the velocity on mobile.

diff --git a/Engine/Shaders/Private/SceneData.ush b/Engine/Shaders/Private/SceneData.ush
--- a/Engine/Shaders/Private/SceneData.ush
+++ b/Engine/Shaders/Private/SceneData.ush
@@ -1141,6 +1141,17 @@ FSceneDataIntermediates GetSceneDataIntermediates(uint DrawInstanceId, float4 In
 	// Reconstruct InstanceData from a packed data
 	FInstanceSceneData InstanceData = (FInstanceSceneData)0;
 
+	// Populate instance data from primitive data
+	InstanceData.LocalToWorld = Primitive.LocalToWorld;
+	InstanceData.PrevLocalToWorld = Primitive.PreviousLocalToWorld;
+	InstanceData.WorldToLocal = Primitive.WorldToLocal;
+	InstanceData.NonUniformScale = Primitive.NonUniformScale;
+	InstanceData.InvNonUniformScale = Primitive.InvNonUniformScale;
+	InstanceData.DeterminantSign = GetPrimitive_DeterminantSign_FromFlags(Primitive.Flags);
+	InstanceData.LocalBoundsCenter = (Primitive.LocalObjectBoundsMax + Primitive.LocalObjectBoundsMin) * 0.5f;
+	InstanceData.LocalBoundsExtent = (Primitive.LocalObjectBoundsMax - Primitive.LocalObjectBoundsMin) * 0.5f;
+	InstanceData.ValidInstance = true;
+
 	float4x4 LocalToRelativeWorld = float4x4(
 		float4(InstanceTransform1.xyz, 0.0f), 
 		float4(InstanceTransform2.xyz, 0.0f), 
@@ -1176,18 +1187,18 @@ FSceneDataIntermediates GetSceneDataIntermediates(uint DrawInstanceId, float4 In
 
 #if ALLOW_STATIC_LIGHTING 
 	Intermediates.Primitive.LightmapDataIndex = asuint(InstanceTransform2.w);
-	InstanceData.LightMapAndShadowMapUVBias = float4(
+	Intermediates.InstanceData.LightMapAndShadowMapUVBias = float4(
 		UnpackSnorm2x16(asuint(InstanceAuxData.y)), 
 		UnpackSnorm2x16(asuint(InstanceAuxData.w)));
 #else
-	InstanceData.LocalBoundsCenter = UnpackLocalBoundsCenter(InstanceAuxData.xy);
-	InstanceData.LocalBoundsExtent = UnpackLocalBoundsExtent(InstanceAuxData.zw);
-	Intermediates.Primitive.LocalObjectBoundsMin = InstanceData.LocalBoundsCenter - InstanceData.LocalBoundsExtent; 
-	Intermediates.Primitive.LocalObjectBoundsMax = InstanceData.LocalBoundsCenter + InstanceData.LocalBoundsExtent; 
+	Intermediates.InstanceData.LocalBoundsCenter = UnpackLocalBoundsCenter(InstanceAuxData.xy);
+	Intermediates.InstanceData.LocalBoundsExtent = UnpackLocalBoundsExtent(InstanceAuxData.zw);
+	Intermediates.Primitive.LocalObjectBoundsMin = Intermediates.InstanceData.LocalBoundsCenter - Intermediates.InstanceData.LocalBoundsExtent;
+	Intermediates.Primitive.LocalObjectBoundsMax = Intermediates.InstanceData.LocalBoundsCenter + Intermediates.InstanceData.LocalBoundsExtent;
 
-	float3 ObjectRelativeWorldPosition = mul(float4(InstanceData.LocalBoundsCenter.xyz, 0.0f), LocalToRelativeWorld).xyz;
+	float3 ObjectRelativeWorldPosition = mul(float4(Intermediates.InstanceData.LocalBoundsCenter.xyz, 0.0f), LocalToRelativeWorld).xyz;
 	Intermediates.Primitive.ObjectWorldPosition = MakeLWCVector3(TilePosition, ObjectRelativeWorldPosition);
-	Intermediates.Primitive.ObjectRadius = length(InstanceData.LocalBoundsExtent * InstanceData.NonUniformScale.xyz);
+	Intermediates.Primitive.ObjectRadius = length(Intermediates.InstanceData.LocalBoundsExtent * Intermediates.InstanceData.NonUniformScale.xyz);
 #endif
 	
 	return Intermediates; 
