Rename Thickness to ThicknessCm in Strata.ush for more clarity about units.

diff --git a/Engine/Shaders/Private/Strata/Strata.ush b/Engine/Shaders/Private/Strata/Strata.ush
--- a/Engine/Shaders/Private/Strata/Strata.ush
+++ b/Engine/Shaders/Private/Strata/Strata.ush
@@ -1345,7 +1345,7 @@ FStrataData GetStrataSlabBSDF(
 	float3 Emissive, 
 	float SecondRoughness, float SecondRoughnessWeight, float SecondRoughnessAsSimpleClearCoat,
 	float FuzzAmount, float3 FuzzColor,
-	float Thickness,
+	float ThicknessCm,
 	bool bIsThinSurface,
 	uint SharedLocalBasisIndex, inout uint SharedLocalBasisTypes)
 {
@@ -1353,13 +1353,13 @@ FStrataData GetStrataSlabBSDF(
 
 	if (bIsThinSurface)
 	{
-		Thickness = STRATA_LAYER_ISTHIN_THICKNESS_THRESHOLD_CM - 1e-5f;
+		ThicknessCm = STRATA_LAYER_ISTHIN_THICKNESS_THRESHOLD_CM - 1e-5f;
 	}
 
 	SSSMFPScale = saturate(SSSMFPScale);
 	const bool bSSSProfilePlugged = SSSProfileID > 0.0f;
 	const bool bMFPPlugged = any(SSSMFP > 0);
-	bool bIsThin = bMFPPlugged && Thickness < STRATA_LAYER_ISTHIN_THICKNESS_THRESHOLD_CM && !UseSSSDiffusion; // Thin is used for flagging 'legacy foliage' case
+	bool bIsThin = bMFPPlugged && ThicknessCm < STRATA_LAYER_ISTHIN_THICKNESS_THRESHOLD_CM && !UseSSSDiffusion; // Thin is used for flagging 'legacy foliage' case
 	bool bHasSSSWrap = bMFPPlugged && !UseSSSDiffusion;
 	bool bHasSSSProfile = !bIsThin && bSSSProfilePlugged && SSSMFPScale > 0.f;	// When no profile is provide, we use Burley. STRATA_TODO we should clarify as burley is also an option of the profile.
 	bool bHasSSS = bMFPPlugged || bHasSSSProfile || bHasSSSWrap;
@@ -1400,8 +1400,8 @@ FStrataData GetStrataSlabBSDF(
 	// Convert MFP be expressed relative to simple volume thickness
 	if (!STRATA_FASTPATH && bHasSSSWrap)
 	{
-		SSSMFP = RescaleMFPToComputationSpace(SSSMFP, Thickness, STRATA_SIMPLEVOLUME_THICKNESS_CM);
-		Thickness = STRATA_SIMPLEVOLUME_THICKNESS_CM;
+		SSSMFP = RescaleMFPToComputationSpace(SSSMFP, ThicknessCm, STRATA_SIMPLEVOLUME_THICKNESS_CM);
+		ThicknessCm = STRATA_SIMPLEVOLUME_THICKNESS_CM;
 	}
 
 	const uint bHasHaziness = ((Haziness.Roughness != Roughness  && Haziness.Weight > 0.0f) || Haziness.bSimpleClearCoat) ? 1 : 0;
@@ -1432,7 +1432,7 @@ FStrataData GetStrataSlabBSDF(
 	BSDF_SETHASMFP			(StrataData.InlinedBSDF,		bMFPPlugged ? 1 : 0);
 	BSDF_SETHASHAZINESS		(StrataData.InlinedBSDF,		bHasHaziness);
 	BSDF_SETHASFUZZ			(StrataData.InlinedBSDF,		bHasFuzz);
-	BSDF_SETTHICKNESSCM		(StrataData.InlinedBSDF,		Thickness);
+	BSDF_SETTHICKNESSCM		(StrataData.InlinedBSDF,		ThicknessCm);
 #endif
 
 #if !STRATA_FASTPATH
@@ -2468,14 +2468,14 @@ FStrataData StrataWeightParameterBlending(FStrataData A, float Weight)
 	return StrataData;
 }
 
-void StrataGetThinFilmF0F90(float NoV, float ThinFilmThickness, float ThinFilmIOR, inout float3 OutF0, inout float3 OutF90)
+void StrataGetThinFilmF0F90(float NoV, float ThinFilmNormalizedThickness, float ThinFilmIOR, inout float3 OutF0, inout float3 OutF90)
 {
 	const float3 F0 = OutF0;
 	const float3 F90 = OutF90;
 
 	const float NoL = NoV;
 	const float VoH = NoV;
-	float3 RThinFilm = F_ThinFilm(NoV, NoL, VoH, F0, F90, ThinFilmIOR, ThinFilmThickness);
+	float3 RThinFilm = F_ThinFilm(NoV, NoL, VoH, F0, F90, ThinFilmIOR, ThinFilmNormalizedThickness);
 
 	// Compute a F0 which match evaluation Thin-Film reflectance for the current view angle using Schlick's Fresnel
 	const float FReference = min(0.99f, F_Schlick(0, 1.0f, NoV).x);
@@ -2499,12 +2499,12 @@ struct FStrataThinFilmOutput
 	float3 F0;
 	float3 F90;
 };
-FStrataThinFilmOutput StrataGetThinFilmF0F90(float NoV, float3 F0, float3 F90, float ThinFilmThickness, float ThinFilmIOR)
+FStrataThinFilmOutput StrataGetThinFilmF0F90(float NoV, float3 F0, float3 F90, float ThinFilmNormalizedThickness, float ThinFilmIOR)
 {
 	FStrataThinFilmOutput Output;
 	Output.F0 = F0;
 	Output.F90 = F90;
-	StrataGetThinFilmF0F90(NoV, ThinFilmThickness, ThinFilmIOR, Output.F0, Output.F90);
+	StrataGetThinFilmF0F90(NoV, ThinFilmNormalizedThickness, ThinFilmIOR, Output.F0, Output.F90);
 	return Output;
 }
 
@@ -3414,7 +3414,7 @@ void StrataConvertFromDBuffer(
 
 	// Decal only support a single Slab BSDF with no-anisotropy
 	const FStrataPixelFootprint PixelFootprint = (FStrataPixelFootprint)0;
-	const float DefaultThickness = STRATA_LAYER_DEFAULT_THICKNESS_CM;
+	const float DefaultThicknessCm = STRATA_LAYER_DEFAULT_THICKNESS_CM;
 	OutData = GetStrataSlabBSDF(
 		PixelFootprint,		// PixelFootprint
 		ComputeDiffuseAlbedo(In.BaseColor, In.Metallic),
@@ -3433,7 +3433,7 @@ void StrataConvertFromDBuffer(
 		0.f,				// SecondRoughnessAsSimpleClearCoat
 		0.f,				// FuzzAmount
 		0.f,				// FuzzColor
-		DefaultThickness,
+		DefaultThicknessCm,
 		false,				// bIsThinSurface,
 		SharedLocalBasisIndex,
 		SharedLocalBasisTypes);
