Path Tracer: Fix interpretation of clearcoat parameter

Setting clearcoat to 0 was not properly fading out the coating. After careful re-reading of the raster evaluation, the logic was subtly different between the two. I also refactored the code to directly compute the three lobes without falling back to the logic in default lit. This prepares the ability to sample diffuse and specular independently.

Implement additional version of lobe selection for 3 and 4 lobed materials. Like the previous method, these are robust to any combination of lobes having zero weights which helps avoid special cases and numerical issues.

Removed unused variables in ClearCoatBxDF

diff --git a/Engine/Shaders/Private/PathTracing/Material/PathTracingClearCoat.ush b/Engine/Shaders/Private/PathTracing/Material/PathTracingClearCoat.ush
--- a/Engine/Shaders/Private/PathTracing/Material/PathTracingClearCoat.ush
+++ b/Engine/Shaders/Private/PathTracing/Material/PathTracingClearCoat.ush
@@ -14,26 +14,14 @@
 #define CLEAR_COAT_F0	0.04 // From hard-coded value in ShadingModels.ush
 
 #include "PathTracingMaterialCommon.ush"
+#include "PathTracingGlossy.ush"
 #include "PathTracingFresnel.ush"
 #include "PathTracingEnergyConservation.ush"
 
-float CalcClearCoatSelectionPdf(float CoatEnergy)
-{
-#if EVAL_TOP_LAYER == 0
-	// never select coating if top is disabled
-	return 0.0; 
-#elif EVAL_BOTTOM_LAYER == 0
-	// always select coating if bottom is disabled
-	return 1.0; 
-#endif
-	// probability of choosing the coat (proportional to its directional albedo)
-	return CoatEnergy;
-}
 
 float ClearCoatFresnel(float VoH) {
 	// from hardcoded logic in ClearCoatBxDF
-	float Fc = Pow5(1 - VoH);
-	return Fc + (1 - Fc) * CLEAR_COAT_F0;
+	return CLEAR_COAT_F0 + (1.0 - CLEAR_COAT_F0) * Pow5(1 - VoH);
 }
 
 FMaterialEval ClearCoat_EvalMaterial(
@@ -44,60 +32,83 @@ FMaterialEval ClearCoat_EvalMaterial(
 {
 	FMaterialEval Result = NullMaterialEval();
 	
-	float3 V_World = -IncomingDirection;
-	float3 L_World = OutgoingDirection;
-	float3 N_World = Payload.WorldNormal;
+	const float3 V_World = -IncomingDirection;
+	const float3 L_World = OutgoingDirection;
 
-	float2 Alpha = 0;
-	float3x3 Basis = GetGGXBasis(Payload.GetClearCoatRoughness(), Payload.WorldNormal, Alpha);
+	float2 CoatAlpha = 0;
+	float2 SpecAlpha = 0;
+	const float3x3 CoatBasis = GetGGXBasis(Payload.GetClearCoatRoughness(), Payload.WorldNormal, CoatAlpha);
+	const float3x3 BaseBasis = GetGGXBasis(Payload.Roughness, Payload.Anisotropy, Payload.GetClearCoatBottomNormal(), Payload.WorldTangent, SpecAlpha);
 
 	// move vectors into right shading frame
-	float3 V = mul(Basis, V_World);
-	float3 L = mul(Basis, L_World);
-	float3 H = normalize(V + L);
+	const float3 CoatV = mul(CoatBasis, V_World);
+	const float3 CoatL = mul(CoatBasis, L_World);
+	const float3 CoatH = normalize(CoatV + CoatL);
 
-	float NoL = saturate(L.z);
-	float NoV = saturate(V.z);
-	float NoH = saturate(H.z);
-	float VoH = saturate(dot(V, H));
+	const float3 BaseV = mul(BaseBasis, V_World);
+	const float3 BaseL = mul(BaseBasis, L_World);
+	const float3 BaseH = normalize(BaseV + BaseL);
 
-	const FBxDFEnergyTerms Coat = ComputeGGXSpecEnergyTerms(Payload.GetClearCoatRoughness(), NoV, CLEAR_COAT_F0);
 
-	// Specular lobe
-	const float2 GGXResult = GGXEvalReflection(L, V, H, Alpha);
-	const float F = ClearCoatFresnel(VoH);
-	const float3 SpecWeight = F * GGXResult.x * Coat.W;
-	const float SpecPdf = GGXResult.y;
+	const float CoatNoL = saturate(CoatL.z);
+	const float CoatNoV = saturate(CoatV.z);
+	const float CoatNoH = saturate(CoatH.z);
+	const float CoatVoH = saturate(dot(CoatV, CoatH));
 
-	const float CoatSelectionPdf = CalcClearCoatSelectionPdf(Coat.E.x);
+	const float BaseNoL = saturate(BaseL.z);
+	const float BaseNoV = saturate(BaseV.z);
+	const float BaseNoH = saturate(BaseH.z);
+	const float BaseVoH = saturate(dot(BaseV, BaseH));
 
-	AddLobeWithMIS(Result.Weight, Result.Pdf, SpecWeight, SpecPdf, CoatSelectionPdf);
-	
+#if EVAL_BOTTOM_LAYER
+#if GREY_INNER_LAYER
+	const float3 SpecColor = 0.0;
+	const float3 DiffColor = 0.18;
+#else
+	const float3 SpecColor = Payload.SpecularColor;
+	const float3 DiffColor = Payload.DiffuseColor;
+#endif
+#else
+	const float3 SpecColor = 0.0;
+	const float3 DiffColor = 0.0;
+#endif
 
-	// bottom-layer clear coat evaluation
-	FMaterialEval BottomEval = NullMaterialEval();
-	Payload.WorldNormal = Payload.GetClearCoatBottomNormal();
-#if GREY_INNER_LAYER == 0
-	// NOTE: Should we try to account for the IOR ratio between layers? We could try to darken Specular Color by the coating IOR
-	// but this would apply to both the dielctric and metal portions equally ...
-	// It seems better to punt until we have a Strata going
-	BottomEval = DefaultLit_EvalMaterial(IncomingDirection, OutgoingDirection, Payload);
+#if EVAL_TOP_LAYER
+	const float ClearCoat = Payload.GetClearCoat();
 #else
-	float3 OldColor = Payload.BaseColor;
-	Payload.BaseColor = 0.18;
-	BottomEval = Lambert_EvalMaterial(OutgoingDirection, Payload);
-	Payload.BaseColor = OldColor;
+	const float ClearCoat = 0.0;
 #endif
-	Payload.WorldNormal = N_World;
 
-	// Calculate transmission through substrate
-	float3 Transmission = SimpleClearCoatTransmittance(NoL, NoV, Payload.Metallic, Payload.BaseColor);
+	const FBxDFEnergyTerms Coat = ComputeGGXSpecEnergyTerms(Payload.GetClearCoatRoughness(), CoatNoV, CLEAR_COAT_F0);
+	const FBxDFEnergyTerms Spec = ComputeGGXSpecEnergyTerms(Payload.Roughness              , BaseNoV, SpecColor);
 
-	BottomEval.Weight *= lerp(1.0, Transmission, Payload.GetClearCoat());
-	BottomEval.Weight *= 1.0 - Coat.E; // remove energy reflected by clearcoat to maintain energy preservation
+	const float3 CoatE = ClearCoat * Coat.E;
+	const float2 LobeCdf = LobeSelectionCdf(CoatE, (1.0 - CoatE) * Spec.E, (1.0 - CoatE) * (1.0 - Spec.E) * DiffColor);
+	const float3 LobePdf = LobeSelectionPdf(LobeCdf);
 
-	AddLobeWithMIS(Result.Weight, Result.Pdf, BottomEval.Weight, BottomEval.Pdf, 1.0 - CoatSelectionPdf);
+	// Coat Specular lobe
+	const float2 CoatGGXResult = GGXEvalReflection(CoatL, CoatV, CoatH, CoatAlpha);
+	const float3 CoatWeight = ClearCoat * ClearCoatFresnel(CoatVoH) * CoatGGXResult.x * Coat.W;
+	const float CoatPdf = CoatGGXResult.y;
 
+	// Calculate transmission through substrate
+	float3 Transmission = (1.0 - CoatE) * lerp(1.0, SimpleClearCoatTransmittance(CoatNoL, CoatNoV, Payload.Metallic, Payload.BaseColor), ClearCoat);
+
+	// Base Specular lobe
+	const float2 SpecGGXResult = GGXEvalReflection(BaseL, BaseV, BaseH, SpecAlpha);
+	const float3 SpecWeight = Transmission * F_Schlick(SpecColor, BaseVoH) * CoatGGXResult.x * Spec.W;
+	const float SpecPdf = SpecGGXResult.y;
+
+	// Base Diffuse lobe
+	const float3 Diffuse = GetPathTracingDiffuseModel(DiffColor, Payload.Roughness, BaseNoV, BaseNoL, BaseVoH, BaseNoH);
+	const float3 DiffWeight = Transmission * (1 - Spec.E) * Diffuse * ShadowTerminatorTerm(L_World, Payload.WorldNormal, Payload.WorldSmoothNormal);
+	const float DiffPdf = BaseNoL / PI;
+
+	// Combine all lobes together
+	AddLobeWithMIS(Result.Weight, Result.Pdf, CoatWeight, CoatPdf, LobePdf.x);
+	AddLobeWithMIS(Result.Weight, Result.Pdf, SpecWeight, SpecPdf, LobePdf.y);
+	AddLobeWithMIS(Result.Weight, Result.Pdf, DiffWeight, DiffPdf, LobePdf.z);
+	
 	return Result;
 }
 
@@ -107,29 +118,59 @@ FMaterialSample ClearCoat_SampleMaterial(
 	float4 RandSample
 )
 {
-	float3 N_World = Payload.WorldNormal;
 	float3 V_World = -RayDirection;
 
-	const float NoV = saturate(dot(N_World, V_World));
-	const FBxDFEnergyTerms Coat = ComputeGGXSpecEnergyTerms(Payload.Roughness, NoV, CLEAR_COAT_F0);
 
-	const float CoatSelectionPdf = CalcClearCoatSelectionPdf(Coat.E.x);
+	float2 CoatAlpha = 0;
+	float2 SpecAlpha = 0;
+	const float3x3 CoatBasis = GetGGXBasis(Payload.GetClearCoatRoughness(), Payload.WorldNormal, CoatAlpha);
+	const float3x3 BaseBasis = GetGGXBasis(Payload.Roughness, Payload.Anisotropy, Payload.GetClearCoatBottomNormal(), Payload.WorldTangent, SpecAlpha);
+
+	// move view vector into right shading frame
+	const float3 CoatV = mul(CoatBasis, V_World);
+	const float3 BaseV = mul(BaseBasis, V_World);
+	const float CoatNoV = saturate(CoatV.z);
+	const float BaseNoV = saturate(BaseV.z);
+
+
+#if EVAL_BOTTOM_LAYER
+#if GREY_INNER_LAYER
+	const float3 SpecColor = 0.0;
+	const float3 DiffColor = 0.18;
+#else
+	const float3 SpecColor = Payload.SpecularColor;
+	const float3 DiffColor = Payload.DiffuseColor;
+#endif
+#else
+	const float3 SpecColor = 0.0;
+	const float3 DiffColor = 0.0;
+#endif
+
+#if EVAL_TOP_LAYER
+	const float ClearCoat = Payload.GetClearCoat();
+#else
+	const float ClearCoat = 0.0;
+#endif
+
+	const FBxDFEnergyTerms Coat = ComputeGGXSpecEnergyTerms(Payload.GetClearCoatRoughness(), CoatNoV, CLEAR_COAT_F0);
+	const FBxDFEnergyTerms Spec = ComputeGGXSpecEnergyTerms(Payload.Roughness              , BaseNoV, Payload.SpecularColor);
+
+	const float3 CoatE = ClearCoat * Coat.E;
+	const float2 LobeCdf = LobeSelectionCdf(CoatE, (1.0 - CoatE) * Spec.E, (1.0 - CoatE) * (1.0 - Spec.E) * DiffColor);
+	const float3 LobePdf = LobeSelectionPdf(LobeCdf);
 
-	// choose between coat and bottom layer
+	// choose between layers
 	float3 OutDirection = 0;
 	float OutRoughness = 0;
-	if (RandSample.x < CoatSelectionPdf)
+	if (RandSample.x < LobeCdf.x)
 	{
-		RandSample.x /= CoatSelectionPdf;
-
-		// No anisotropy, just pick an arbitrary basis
-		float2 Alpha = 0;
-		float3x3 Basis = GetGGXBasis(Payload.GetClearCoatRoughness(), Payload.WorldNormal, Alpha);
+		// Sampled coating
+		RandSample.x /= LobeCdf.x;
 
-		float3 V = -mul(Basis, RayDirection);
+		float3 V = -mul(CoatBasis, RayDirection);
 
 		// #dxr_todo: Switch to UniformSampleDiskConcentric after fixing 0/0 bug
-		float3 H = ImportanceSampleVisibleGGX_aniso(UniformSampleDisk(RandSample.xy), Alpha, V).xyz;
+		float3 H = ImportanceSampleVisibleGGX_aniso(UniformSampleDisk(RandSample.xy), CoatAlpha, V).xyz;
 
 		float3 L = reflect(-V, H);
 		if (L.z <= 0)
@@ -138,32 +179,80 @@ FMaterialSample ClearCoat_SampleMaterial(
 			return NullMaterialSample();
 		}
 		// reflect and transform
-		OutDirection = mul(L, Basis);
+		OutDirection = mul(L, CoatBasis);
 		OutRoughness = Payload.GetClearCoatRoughness();
 	}
+	else if (RandSample.x < LobeCdf.y)
+	{
+		// Sampled base spec
+		RandSample.x -= LobeCdf.x;
+		RandSample.x /= LobeCdf.y - LobeCdf.x;
+
+		float3 V = -mul(BaseBasis, RayDirection);
+		float3 H = ImportanceSampleVisibleGGX_aniso(UniformSampleDisk(RandSample.xy), SpecAlpha, V).xyz;
+		float3 L = reflect(-V, H);
+		if (L.z <= 0)
+		{
+			// invalid output direction, exit early
+			return NullMaterialSample();
+		}
+
+		OutDirection = mul(L, BaseBasis);
+		OutRoughness = Payload.Roughness;
+	}
 	else
 	{
-		RandSample.x -= CoatSelectionPdf;
-		RandSample.x /= (1.0 - CoatSelectionPdf);
-
-#if GREY_INNER_LAYER == 0
-		float3 Weight = 0;
-		float Pdf = 0;
-		Payload.WorldNormal = Payload.GetClearCoatBottomNormal();
-		FMaterialSample InnerSample = DefaultLit_SampleMaterial(RayDirection, Payload, RandSample);
-		Payload.WorldNormal = N_World;
-
-		OutDirection = InnerSample.Direction;
-		OutRoughness = InnerSample.Roughness;
-#else
-		float3 Nb = Payload.GetClearCoatBottomNormal();
-		OutDirection = TangentToWorld(CosineSampleHemisphere(RandSample.xy).xyz, Nb);
+		// Sampled base diffuse
+		RandSample.x -= LobeCdf.y;
+		RandSample.x /= 1.0 - LobeCdf.y;
+
+		float3 L = CosineSampleHemisphere(RandSample.xy).xyz;
+		OutDirection = mul(L, BaseBasis);
 		OutRoughness = 1.0;
-#endif
 	}
 
-	// defer to eval for now to get proper MIS mix of both events
-	// TODO: could use temporary results above to inline Eval logic here
-	FMaterialEval Result = ClearCoat_EvalMaterial(RayDirection, OutDirection, Payload);
-	return CreateMaterialSample(OutDirection, Result.Weight, Result.Pdf, 1.0, OutRoughness);
+
+	// move light vector into right shading frame
+	const float3 L_World = OutDirection;
+
+	const float3 CoatL = mul(CoatBasis, L_World);
+	const float3 CoatH = normalize(CoatV + CoatL);
+
+	const float3 BaseL = mul(BaseBasis, L_World);
+	const float3 BaseH = normalize(BaseV + BaseL);
+
+	const float CoatNoL = saturate(CoatL.z);
+	const float CoatNoH = saturate(CoatH.z);
+	const float CoatVoH = saturate(dot(CoatV, CoatH));
+
+	const float BaseNoL = saturate(BaseL.z);
+	const float BaseNoH = saturate(BaseH.z);
+	const float BaseVoH = saturate(dot(BaseV, BaseH));
+
+	// Coat Specular lobe
+	const float2 CoatGGXResult = GGXEvalReflection(CoatL, CoatV, CoatH, CoatAlpha);
+	const float3 CoatWeight = ClearCoat * ClearCoatFresnel(CoatVoH) * CoatGGXResult.x * Coat.W;
+	const float CoatPdf = CoatGGXResult.y;
+
+	// Calculate transmission through substrate
+	float3 Transmission = (1.0 - CoatE) * lerp(1.0, SimpleClearCoatTransmittance(CoatNoL, CoatNoV, Payload.Metallic, Payload.BaseColor), ClearCoat);
+
+	// Base Specular lobe
+	const float2 SpecGGXResult = GGXEvalReflection(BaseL, BaseV, BaseH, SpecAlpha);
+	const float3 SpecWeight = Transmission * F_Schlick(SpecColor, BaseVoH) * CoatGGXResult.x * Spec.W;
+	const float SpecPdf = SpecGGXResult.y;
+
+	// Base Diffuse lobe
+	const float3 Diffuse = GetPathTracingDiffuseModel(DiffColor, Payload.Roughness, BaseNoV, BaseNoL, BaseVoH, BaseNoH);
+	const float3 DiffWeight = Transmission * (1 - Spec.E) * Diffuse * ShadowTerminatorTerm(L_World, Payload.WorldNormal, Payload.WorldSmoothNormal);
+	const float DiffPdf = BaseNoL / PI;
+
+	// Combine all lobes together
+	float3 OutWeight = 0.0;
+	float OutPdf = 0.0;
+	AddLobeWithMIS(OutWeight, OutPdf, CoatWeight, CoatPdf, LobePdf.x);
+	AddLobeWithMIS(OutWeight, OutPdf, SpecWeight, SpecPdf, LobePdf.y);
+	AddLobeWithMIS(OutWeight, OutPdf, DiffWeight, DiffPdf, LobePdf.z);
+
+	return CreateMaterialSample(OutDirection, OutWeight, OutPdf, 1.0, OutRoughness);
 }
diff --git a/Engine/Shaders/Private/PathTracing/Material/PathTracingMaterialCommon.ush b/Engine/Shaders/Private/PathTracing/Material/PathTracingMaterialCommon.ush
--- a/Engine/Shaders/Private/PathTracing/Material/PathTracingMaterialCommon.ush
+++ b/Engine/Shaders/Private/PathTracing/Material/PathTracingMaterialCommon.ush
@@ -19,6 +19,53 @@ float LobeSelectionProb(float3 A, float3 B)
 	return MISWeightBalanced(SumA, SumB);
 }
 
+// Given three lobes that will roughly contribute colors A, B and C to the total (estimated for example using directional albedo)
+// return a CDF for choosing lobe A, B or C
+// The individual probabilities of each lobe can be estimated using LobeSelectionPdf.
+float2 LobeSelectionCdf(float3 A, float3 B, float3 C)
+{
+	const float SumA = A.x + A.y + A.z;
+	const float SumB = B.x + B.y + B.z;
+	const float SumC = C.x + C.y + C.z;
+
+	// use MIS routine to robustly compute the ratios even when both are 0, or overflow, etc ...
+	const float ProbAB = MISWeightBalanced(SumA, SumB); // A or B
+	const float ProbLR = MISWeightBalanced(SumA + SumB, SumC); // AB or C
+
+	// Now build CDF
+	return float2(ProbAB * ProbLR, ProbLR);
+}
+
+// Same as above for choosing among 4 colored lobes
+float3 LobeSelectionCdf(float3 A, float3 B, float3 C, float3 D)
+{
+	const float SumA = A.x + A.y + A.z;
+	const float SumB = B.x + B.y + B.z;
+	const float SumC = C.x + C.y + C.z;
+	const float SumD = D.x + D.y + D.z;
+
+	// use MIS routine to robustly compute the ratios even when both are 0, or overflow, etc ...
+	const float ProbAB = MISWeightBalanced(SumA, SumB); // A or B
+	const float ProbCD = MISWeightBalanced(SumC, SumD); // C or D
+	const float ProbLR = MISWeightBalanced(SumA + SumB, SumC + SumD); // AB or CD
+
+	// Now build CDF
+	return float3(ProbAB * ProbLR, ProbLR, (ProbLR + ProbCD) - ProbLR * ProbCD);
+}
+
+// Given the CDF for a choice between three lobes (as returned by LobeSelectionCdf above), compute the Pdf for each individual lobe
+float3 LobeSelectionPdf(float2 LobeCdf)
+{
+	return float3(LobeCdf.x, LobeCdf.y - LobeCdf.x, 1.0 - LobeCdf.y);
+}
+
+// Given the CDF for a choice between four lobes (as returned by LobeSelectionCdf above), compute the Pdf for each individual lobe
+float4 LobeSelectionPdf(float3 LobeCdf)
+{
+	return float4(LobeCdf.x, LobeCdf.y - LobeCdf.x, LobeCdf.z - LobeCdf.y, 1.0 - LobeCdf.z);
+}
+
+
 // Takes as input the sample weight and pdf for a certain lobe of a mixed model, together with the probability of picking that lobe
 // This function then updates a running total Weight and Pdf value that represents the overall contribution of the BxDF
 // This function should be called when a BxDF is made up of multiple lobes combined with a sum to correctly account for the probability
diff --git a/Engine/Shaders/Private/ShadingModels.ush b/Engine/Shaders/Private/ShadingModels.ush
--- a/Engine/Shaders/Private/ShadingModels.ush
+++ b/Engine/Shaders/Private/ShadingModels.ush
@@ -405,8 +405,6 @@ FDirectLighting ClearCoatBxDF( FGBufferData GBuffer, half3 N, half3 V, half3 L,
 {
 	const half ClearCoat			= GBuffer.CustomData.x;
 	const half ClearCoatRoughness	= max(GBuffer.CustomData.y, 0.02f);
-	const half Film = 1 * ClearCoat;
-	const half MetalSpec = 0.9;
 
 	FDirectLighting Lighting = {
 		float3(0.0, 0.0, 0.0),
