Back out part of recent change that implementd GetPrevInstanceToWorld() for instanced geom.
Go back to the old behavior which was to return current frame InstanceToWorld.
Keep the part of the change that applied for non instanced geom.
Need to do this until we solve some issues with the calculation of PrevFrameLocalToWorld.
For static instances we hit a precision issue in matrix concatenation in GetInstanceSceneData().
The precision can blow up in WPO calculations to give bad velocity vectors.

diff --git a/Engine/Shaders/Private/MaterialTemplate.ush b/Engine/Shaders/Private/MaterialTemplate.ush
--- a/Engine/Shaders/Private/MaterialTemplate.ush
+++ b/Engine/Shaders/Private/MaterialTemplate.ush
@@ -682,7 +682,8 @@ FLWCMatrix GetInstanceToWorld(FMaterialVertexParameters Parameters)
 FLWCMatrix GetPrevInstanceToWorld(FMaterialVertexParameters Parameters)
 {
 #if USE_INSTANCING || USE_INSTANCE_CULLING || IS_MESHPARTICLE_FACTORY || IS_NANITE_PASS
-	return Parameters.PrevFrameLocalToWorld;
+	// todo: return Parameters.PrevFrameLocalToWorld when we have fixed precision issue of PrevLocalToWorld in GetInstanceSceneData()
+	return Parameters.InstanceLocalToWorld;
 #else
 	return GetPrimitiveData(Parameters).PreviousLocalToWorld;
 #endif
