Lumen GPU crash fix when using hair procedural geometry:
* GetLumenGeometryNormal() was calling LoadVertexFactoryInputForHGS() which was intended to work only for triangle geometry
* Return geometry normal calculated from material properties for materials that use VF_SUPPORTS_RAYTRACING_PREPARE_MATERIAL_PIXEL_PARAMETERS (which includes procedural hair)
* Prevent accessing triangle data in LoadVertexFactoryInputForHGS() when using procedural geometry.

diff --git a/Engine/Shaders/Private/HairStrands/HairStrandsVertexFactory.ush b/Engine/Shaders/Private/HairStrands/HairStrandsVertexFactory.ush
--- a/Engine/Shaders/Private/HairStrands/HairStrandsVertexFactory.ush
+++ b/Engine/Shaders/Private/HairStrands/HairStrandsVertexFactory.ush
@@ -155,13 +155,17 @@ FVertexInfo GetVertexInfo(FVertexFactoryInput Input)
 #if RAYHITGROUPSHADER
 FVertexFactoryInput LoadVertexFactoryInputForHGS(uint TriangleIndex, int VertexIndex)
 {
+	FVertexFactoryInput Input = (FVertexFactoryInput)0;
+
+#if !ENABLE_PROCEDURAL_INTERSECTOR
 	FTriangleBaseAttributes TriangleBase = LoadTriangleBaseAttributes(TriangleIndex);
 
-	FVertexFactoryInput Input;
 	Input.Position = float4(TriangleBase.LocalPositions[VertexIndex], 1.0f);
 	Input.VertexId = TriangleBase.Indices[VertexIndex];
 	Input.TriangleId = TriangleIndex;
 	Input.TriangleVertexId = VertexIndex;
+#endif // !ENABLE_PROCEDURAL_INTERSECTOR
+
 	// Note: GetInstanceUserData() stores the GPU-Scene primitive ID
 	VF_GPUSCENE_SET_INPUT_FOR_RT(Input, GetInstanceUserData(), 0U);
 
diff --git a/Engine/Shaders/Private/RayTracing/RayTracingMaterialHitShaders.usf b/Engine/Shaders/Private/RayTracing/RayTracingMaterialHitShaders.usf
--- a/Engine/Shaders/Private/RayTracing/RayTracingMaterialHitShaders.usf
+++ b/Engine/Shaders/Private/RayTracing/RayTracingMaterialHitShaders.usf
@@ -468,11 +468,29 @@ RAY_TRACING_ENTRY_CLOSEST_HIT(MaterialCHS,
 	PackedPayload.HitT = RayTCurrent();
 
 #if USE_MATERIAL_CLOSEST_HIT_SHADER
+	GlobalTextureMipBias = PackedPayload.GetMipBias();
+
+	ResolvedView = ResolveView();
+
+	FVertexFactoryInterpolantsVSToPS Interpolants;
+	const float3 TranslatedWorldPosition = TranslatedWorldRayOrigin() + RayTCurrent() * WorldRayDirection();
+	float4 SvPosition = mul(float4(TranslatedWorldPosition, 1.0f), ResolvedView.TranslatedWorldToClip);
+	
+	FRayCone PropagatedCone = PropagateRayCone(PackedPayload.GetRayCone(), 0 /* surface curvature */, RayTCurrent());
+	CalcInterpolants(PropagatedCone, Attributes, Interpolants);
+
+#if VF_SUPPORTS_RAYTRACING_PREPARE_MATERIAL_PIXEL_PARAMETERS
+	float3 GeoNormal = 0;
+	FMaterialPixelParameters MaterialParameters = GetMaterialPixelParameters(TranslatedWorldRayOrigin(), WorldRayDirection(), RayTCurrent(), PrimitiveIndex(), Attributes, HitKind(), GeoNormal);
+#else
+	float3 GeoNormal = GetLumenGeometryNormal();
+	FMaterialPixelParameters MaterialParameters = GetMaterialPixelParameters(Interpolants, SvPosition);
+#endif
 
 	bool bIsLumenPayload = PackedPayload.IsLumenPayload();
 	if (bIsLumenPayload)
 	{
-		float3 SurfaceNormal = GetLumenGeometryNormal();
+		const float3 SurfaceNormal = GeoNormal;
 		const int PrimitiveIndex = GetInstanceUserData();
 		const int PrimitiveInstanceIndex = InstanceIndex() - GetBaseInstanceIndex();
 
@@ -490,24 +508,6 @@ RAY_TRACING_ENTRY_CLOSEST_HIT(MaterialCHS,
 		return;
 	}
 
-	GlobalTextureMipBias = PackedPayload.GetMipBias();
-
-	ResolvedView = ResolveView();
-
-	FVertexFactoryInterpolantsVSToPS Interpolants;
-	const float3 TranslatedWorldPosition = TranslatedWorldRayOrigin() + RayTCurrent() * WorldRayDirection();
-	float4 SvPosition = mul(float4(TranslatedWorldPosition, 1.0f), ResolvedView.TranslatedWorldToClip);
-
-	
-	FRayCone PropagatedCone = PropagateRayCone(PackedPayload.GetRayCone(), 0 /* surface curvature */, RayTCurrent());
-	CalcInterpolants(PropagatedCone, Attributes, Interpolants);
-
-#if VF_SUPPORTS_RAYTRACING_PREPARE_MATERIAL_PIXEL_PARAMETERS
-	float3 GeoNormal = 0;
-	FMaterialPixelParameters MaterialParameters = GetMaterialPixelParameters(TranslatedWorldRayOrigin(), WorldRayDirection(), RayTCurrent(), PrimitiveIndex(), Attributes, HitKind(), GeoNormal);
-#else
-	FMaterialPixelParameters MaterialParameters = GetMaterialPixelParameters(Interpolants, SvPosition);
-#endif
 	FPixelMaterialInputs PixelMaterialInputs;
 
 	{
