Changes to EyeAdaptationCommon.ush missing from previous CL.

diff --git a/Engine/Shaders/Private/EyeAdaptationCommon.ush b/Engine/Shaders/Private/EyeAdaptationCommon.ush
--- a/Engine/Shaders/Private/EyeAdaptationCommon.ush
+++ b/Engine/Shaders/Private/EyeAdaptationCommon.ush
@@ -5,24 +5,21 @@
 #include "Common.ush"
 
 #if !EYE_ADAPTATION_DISABLED
-
-Texture2D EyeAdaptationTexture;
 StructuredBuffer<float4> EyeAdaptationBuffer;
 
+float4 EyeAdaptationLookupBuffer(StructuredBuffer<float4> InEyeAdaptation)
+{
+	return InEyeAdaptation[0];
+}
+
 #if SHADING_PATH_DEFERRED
+Texture2D EyeAdaptationTexture;
+
 float4 EyeAdaptationLookupTexture2D(Texture2D InEyeAdaptation)
 {
 	return InEyeAdaptation.Load(int3(0, 0, 0));
 }
-#elif SHADING_PATH_MOBILE
-float4 EyeAdaptationLookupBuffer(StructuredBuffer<float4> InEyeAdaptation)
-{
-	return InEyeAdaptation[0];
-}
-#else
-	#error Unknown shading path.
 #endif
-
 #endif // !EYE_ADAPTATION_DISABLED
 
 #if defined (EyeAdaptationStruct)
@@ -46,7 +43,7 @@ float EyeAdaptationLookup()
 	#if SHADING_PATH_DEFERRED
 		return EyeAdaptationLookupTexture2D(EyeAdaptationStruct.EyeAdaptationTexture).x;
 	#elif SHADING_PATH_MOBILE
-		// Unofortunatelly on older MALI GL drivers (r9p0), passing this buffer into 'log' or other functions, fails with this error: Function call discards 'readonly' access qualifier
+		// Unfortunately on older MALI GL drivers (r9p0), passing this buffer into 'log' or other functions, fails with this error: Function call discards 'readonly' access qualifier
 		// So we have to use last frame eye adaptation value as a workaround on mobile
 		// return EyeAdaptationLookupBuffer(EyeAdaptationBuffer).x;
 		return ResolvedView.PreExposure;
