IES Atlas: Ensure lookups use centered texels to prevent a subtle shift in the projection

Some of the path tracing EngineTests showed a difference to the previous implementation. This restores the previous behavior.

diff --git a/Engine/Shaders/Private/IESAtlas.usf b/Engine/Shaders/Private/IESAtlas.usf
--- a/Engine/Shaders/Private/IESAtlas.usf
+++ b/Engine/Shaders/Private/IESAtlas.usf
@@ -42,7 +42,7 @@ void MainCS(uint3 DispatchThreadId : SV_DispatchThreadID)
 		const uint2 DstPixelPos = DispatchThreadId.xy;
 		uint DstSlice = 0;
 
-		const float2 SrcUV = DstPixelPos / float2(AtlasResolution);
+		const float2 SrcUV = (DstPixelPos + 0.5) / float2(AtlasResolution);
 		const uint SrcSlice = DispatchThreadId.z;
 
 		if (SrcSlice < ValidCount)
