Path Tracer: Refactor payload encoding to improve precision

Improve precision of unit vectors to 16:16 octahedral for better result on very smooth surfaces

Re-organize the flags to come first, followed by the smooth and geometric normal (which are independent of shading data). This is in preparation for a more Strata friendly layout.

diff --git a/Engine/Shaders/Private/PathTracing/PathTracingCommon.ush b/Engine/Shaders/Private/PathTracing/PathTracingCommon.ush
--- a/Engine/Shaders/Private/PathTracing/PathTracingCommon.ush
+++ b/Engine/Shaders/Private/PathTracing/PathTracingCommon.ush
@@ -16,7 +16,7 @@
 #define PATH_TRACING_PAYLOAD_INPUT_FLAG_IGNORE_SKY_MATERIALS          (1 << 0) // Indicates that the AHS/CHS should skip processing if the material is tagged as sky
 #define PATH_TRACING_PAYLOAD_INPUT_FLAG_PATHTRACER_VISIBILITY         (1 << 1) // Indicates that this is a path tracer visibility ray (which supports material evaluation for transparent shadows and fake caustics)
 
-#define PATH_TRACING_PAYLOAD_NORMAL_BITS  24
+#define PATH_TRACING_PAYLOAD_NORMAL_BITS  32
 
 // shading model extensions beyond the default count
 #define SHADINGMODELID_MEDIUM  (SHADINGMODELID_NUM + 1)
@@ -49,9 +49,9 @@ struct FPackedPathTracingPayload : FMinimalPayload
 	uint PackedData[15];								 // 60 bytes (encoded data, depends on shading model and ray type)
 														 // 64 bytes total
 
-	uint GetFlags() { return (PackedData[7] >> 23) & 0x3F; }
-	void SetFlag(uint Flag) { PackedData[7] |= Flag << 23; }
-	uint GetPrimitiveLightingChannelMask() { return (PackedData[7] >> 29) & 0x7; }
+	uint GetFlags() { return (PackedData[0] >> 23) & 0x3F; }
+	void SetFlag(uint Flag) { PackedData[0] |= Flag << 23; }
+	uint GetPrimitiveLightingChannelMask() { return (PackedData[0] >> 29) & 0x7; }
 
 	// Flag methods
 	void SetIgnoreSkyMaterials() { SetFlag(PATH_TRACING_PAYLOAD_INPUT_FLAG_IGNORE_SKY_MATERIALS); }
@@ -67,44 +67,44 @@ struct FPackedPathTracingPayload : FMinimalPayload
 	float3 GetRayThroughput()
 	{
 		return float3(
-			asfloat(PackedData[0]),
 			asfloat(PackedData[1]),
-			asfloat(PackedData[2])
+			asfloat(PackedData[2]),
+			asfloat(PackedData[3])
 		);
 	}
 
 	void SetRayThroughput(float3 RayThroughput)
 	{
-		PackedData[0] = asuint(RayThroughput.x);
-		PackedData[1] = asuint(RayThroughput.y);
-		PackedData[2] = asuint(RayThroughput.z);
+		PackedData[1] = asuint(RayThroughput.x);
+		PackedData[2] = asuint(RayThroughput.y);
+		PackedData[3] = asuint(RayThroughput.z);
 	}
 
 	float GetPathRoughness()
 	{
-		return asfloat(PackedData[3]);
+		return asfloat(PackedData[4]);
 	}
 
 	void SetPathRoughness(float PathRoughness)
 	{
-		PackedData[3] = asuint(PathRoughness);
+		PackedData[4] = asuint(PathRoughness);
 	}
 
 	// tau: optical depth along shadow ray, this is the integral of the extinction coefficient
 	float3 GetTau()
 	{
 		return float3(
-			asfloat(PackedData[4]),
 			asfloat(PackedData[5]),
-			asfloat(PackedData[6])
+			asfloat(PackedData[6]),
+			asfloat(PackedData[7])
 		);
 	}
 
 	void SetTau(float3 Tau)
 	{
-		PackedData[4] = asuint(Tau.x);
-		PackedData[5] = asuint(Tau.y);
-		PackedData[6] = asuint(Tau.z);
+		PackedData[5] = asuint(Tau.x);
+		PackedData[6] = asuint(Tau.y);
+		PackedData[7] = asuint(Tau.z);
 	}
 };
 #if IS_PAYLOAD_ENABLED(RT_PAYLOAD_PATH_TRACING_MATERIAL)
@@ -217,48 +217,41 @@ FPackedPathTracingPayload PackPathTracingPayload(FPathTracingPayload Input)
 	FPackedPathTracingPayload Output = (FPackedPathTracingPayload)0;
 	Output.HitT = Input.HitT;
 
-	Output.PackedData[0] = f32tof16(Input.Radiance.x);
-	Output.PackedData[0] |= f32tof16(Input.Radiance.y) << 16;
-	Output.PackedData[1] = f32tof16(Input.Radiance.z);
-
-	uint ShdNormal = PayloadEncodeUnitVector(Input.WorldNormal);
-	uint GeoNormal = PayloadEncodeUnitVector(Input.WorldGeoNormal);
-	Output.PackedData[1] |= ShdNormal << 16; // pack low 16 bits
-	Output.PackedData[2] = ShdNormal >> 16;  // pack upper 8 bits
-	Output.PackedData[2] |= GeoNormal << 8;  // pack all 24 bits
-
-	Output.PackedData[3] = f32tof16(Input.BaseColor.x);
-	Output.PackedData[3] |= f32tof16(Input.BaseColor.y) << 16;
-	Output.PackedData[4] = f32tof16(Input.BaseColor.z);
-
-	Output.PackedData[4] |= f32tof16(Input.Opacity) << 16;
-
-	Output.PackedData[5] = f32tof16(Input.Metallic);
-	Output.PackedData[5] |= f32tof16(Input.Specular) << 16;
-	Output.PackedData[6] = f32tof16(Input.Roughness);
-	Output.PackedData[6] |= f32tof16(Input.Anisotropy) << 16;
-
-	Output.PackedData[7] = f32tof16(Input.Ior);                              // 16 bits
-	Output.PackedData[7] |= (Input.ShadingModelID & 0xF) << 16;               // 4 bits
-	Output.PackedData[7] |= (Input.BlendingMode & 0x7) << 20;                 // 3 bits
-	Output.PackedData[7] |= (Input.Flags & 0x3F) << 23;                       // 6 bits
-	Output.PackedData[7] |= (Input.PrimitiveLightingChannelMask & 0x7) << 29; // 3 bits
-                                                                       // total 32 bits
-
-	uint SmoothNormal = PayloadEncodeUnitVector(Input.WorldSmoothNormal);
-	uint TangentVec   = PayloadEncodeUnitVector(Input.WorldTangent);
-	Output.PackedData[8] = SmoothNormal;                     // pack all 24 bits
-	Output.PackedData[8] |= (TangentVec << 8) & 0xFF000000u; // pack upper 8 bits
-	Output.PackedData[9] = TangentVec & 0xFFFFu;             // pack lower 16 bits
-
-	Output.PackedData[9] |= f32tof16(Input.CustomVector.x) << 16;
-	Output.PackedData[10] = f32tof16(Input.CustomVector.y);
-	Output.PackedData[10] |= f32tof16(Input.CustomVector.z) << 16;
-
-	Output.PackedData[11] = f32tof16(Input.CustomData.x);
-	Output.PackedData[11] |= f32tof16(Input.CustomData.y) << 16;
-	Output.PackedData[12] = f32tof16(Input.CustomData.z);
-	Output.PackedData[12] |= f32tof16(Input.CustomData.w) << 16;
+	Output.PackedData[0] = f32tof16(Input.Opacity);
+	Output.PackedData[0] |= (Input.ShadingModelID & 0xF) << 16;               // 4 bits
+	Output.PackedData[0] |= (Input.BlendingMode & 0x7) << 20;                 // 3 bits
+	Output.PackedData[0] |= (Input.Flags & 0x3F) << 23;                       // 6 bits
+	Output.PackedData[0] |= (Input.PrimitiveLightingChannelMask & 0x7) << 29; // 3 bits
+	                                                                   // total 32 bits
+	Output.PackedData[1] = PayloadEncodeUnitVector(Input.WorldGeoNormal);
+	Output.PackedData[2] = PayloadEncodeUnitVector(Input.WorldSmoothNormal);
+	Output.PackedData[3] = PayloadEncodeUnitVector(Input.WorldNormal);
+
+	Output.PackedData[4] = f32tof16(Input.Radiance.x);
+	Output.PackedData[4] |= f32tof16(Input.Radiance.y) << 16;
+	Output.PackedData[5] = f32tof16(Input.Radiance.z);
+
+	Output.PackedData[5] |= f32tof16(Input.BaseColor.x) << 16;
+	Output.PackedData[6] = f32tof16(Input.BaseColor.y);
+	Output.PackedData[6] |= f32tof16(Input.BaseColor.z) << 16;
+
+	Output.PackedData[7] = f32tof16(Input.Metallic);
+	Output.PackedData[7] |= f32tof16(Input.Specular) << 16;
+	Output.PackedData[8] = f32tof16(Input.Roughness);
+	Output.PackedData[8] |= f32tof16(Input.Anisotropy) << 16;
+
+	Output.PackedData[9] = PayloadEncodeUnitVector(Input.WorldTangent);
+
+	Output.PackedData[10] = f32tof16(Input.CustomVector.x);
+	Output.PackedData[10] |= f32tof16(Input.CustomVector.y) << 16;
+	Output.PackedData[11] = f32tof16(Input.CustomVector.z);
+
+	Output.PackedData[11] |= f32tof16(Input.CustomData.x) << 16;
+	Output.PackedData[12] = f32tof16(Input.CustomData.y);
+	Output.PackedData[12] |= f32tof16(Input.CustomData.z) << 16;
+	Output.PackedData[13] = f32tof16(Input.CustomData.w);
+	Output.PackedData[13] |= f32tof16(Input.Ior) << 16;
+	Output.PackedData[14] = 0; // unused
 
 	return Output;
 }
@@ -268,36 +261,39 @@ FPathTracingPayload UnpackPathTracingPayload(FPackedPathTracingPayload Input, FR
 	FPathTracingPayload Output = (FPathTracingPayload)0;
 
 	Output.HitT = Input.HitT;
-	Output.Radiance.x = f16tof32(Input.PackedData[0]);
-	Output.Radiance.y = f16tof32(Input.PackedData[0] >> 16);
-	Output.Radiance.z = f16tof32(Input.PackedData[1]);
-	Output.WorldNormal = PayloadDecodeUnitVector((Input.PackedData[1] >> 16) | ((Input.PackedData[2] & 0xFF) << 16));
-	Output.WorldGeoNormal = PayloadDecodeUnitVector(Input.PackedData[2] >> 8);
-	Output.BaseColor.x = f16tof32(Input.PackedData[3]);
-	Output.BaseColor.y = f16tof32(Input.PackedData[3] >> 16);
-	Output.BaseColor.z = f16tof32(Input.PackedData[4]);
-	Output.Opacity = f16tof32(Input.PackedData[4] >> 16);
-	Output.Metallic = f16tof32(Input.PackedData[5]);
-	Output.Specular = f16tof32(Input.PackedData[5] >> 16);
-	Output.Roughness = f16tof32(Input.PackedData[6]);
-	Output.Anisotropy = f16tof32(Input.PackedData[6] >> 16);
-	Output.Ior = f16tof32(Input.PackedData[7]);
-	Output.ShadingModelID = (Input.PackedData[7] >> 16) & 0xF;
-	Output.BlendingMode = (Input.PackedData[7] >> 20) & 0x7;
+	Output.Opacity = f16tof32(Input.PackedData[0]);
+	Output.ShadingModelID = (Input.PackedData[0] >> 16) & 0xF;
+	Output.BlendingMode = (Input.PackedData[0] >> 20) & 0x7;
 	Output.Flags = Input.GetFlags();
 	Output.PrimitiveLightingChannelMask = Input.GetPrimitiveLightingChannelMask();
 
-	Output.WorldSmoothNormal = PayloadDecodeUnitVector(Input.PackedData[8] & 0x00FFFFFFu);
-	Output.WorldTangent = PayloadDecodeUnitVector((Input.PackedData[8] & 0xFF000000u) >> 8 | (Input.PackedData[9] & 0xFFFFu));
+	Output.WorldGeoNormal = PayloadDecodeUnitVector(Input.PackedData[1]);
+	Output.WorldSmoothNormal = PayloadDecodeUnitVector(Input.PackedData[2]);
+	Output.WorldNormal = PayloadDecodeUnitVector(Input.PackedData[3]);
 
-	Output.CustomVector.x = f16tof32(Input.PackedData[9] >> 16);
-	Output.CustomVector.y = f16tof32(Input.PackedData[10]);
-	Output.CustomVector.z = f16tof32(Input.PackedData[10] >> 16);
-
-	Output.CustomData.x = f16tof32(Input.PackedData[11]);
-	Output.CustomData.y = f16tof32(Input.PackedData[11] >> 16);
-	Output.CustomData.z = f16tof32(Input.PackedData[12]);
-	Output.CustomData.w = f16tof32(Input.PackedData[12] >> 16);
+	Output.Radiance.x = f16tof32(Input.PackedData[4]);
+	Output.Radiance.y = f16tof32(Input.PackedData[4] >> 16);
+	Output.Radiance.z = f16tof32(Input.PackedData[5]);
+	Output.BaseColor.x = f16tof32(Input.PackedData[5] >> 16);
+	Output.BaseColor.y = f16tof32(Input.PackedData[6]);
+	Output.BaseColor.z = f16tof32(Input.PackedData[6] >> 16);
+	
+	Output.Metallic = f16tof32(Input.PackedData[7]);
+	Output.Specular = f16tof32(Input.PackedData[7] >> 16);
+	Output.Roughness = f16tof32(Input.PackedData[8]);
+	Output.Anisotropy = f16tof32(Input.PackedData[8] >> 16);
+
+	Output.WorldTangent = PayloadDecodeUnitVector(Input.PackedData[9]);
+
+	Output.CustomVector.x = f16tof32(Input.PackedData[10]);
+	Output.CustomVector.y = f16tof32(Input.PackedData[10] >> 16);
+	Output.CustomVector.z = f16tof32(Input.PackedData[11]);
+
+	Output.CustomData.x = f16tof32(Input.PackedData[11] >> 16);
+	Output.CustomData.y = f16tof32(Input.PackedData[12]);
+	Output.CustomData.z = f16tof32(Input.PackedData[12] >> 16);
+	Output.CustomData.w = f16tof32(Input.PackedData[13]);
+	Output.Ior = f16tof32(Input.PackedData[13] >> 16);
 
 	Output.TranslatedWorldPos = Ray.Origin + Output.HitT * Ray.Direction;
 
