Add Fuzz roughness to Strata slab.

If FuzzFoughness is not connected, the Roughness input is used instead. This allows optional independent roughness control.

diff --git a/Engine/Shaders/Private/DebugProbes.usf b/Engine/Shaders/Private/DebugProbes.usf
--- a/Engine/Shaders/Private/DebugProbes.usf
+++ b/Engine/Shaders/Private/DebugProbes.usf
@@ -132,6 +132,7 @@ void WriteOutMaterial(
 		/*flaot SecondRoughnessAsSimpleClearCoat*/0.0f,
 		/*float FuzzAmount						*/0.0f,
 		/*float3 FuzzColor						*/0.0f,
+		/*float FuzzRoughness					*/0.0f,
 		/*float Thickness						*/0.0f,
 		/*bool bIsThinSurface					*/false,
 		/*uint SharedLocalBasisIndex			*/SharedLocalBasisIndex,
diff --git a/Engine/Shaders/Private/Strata/Strata.ush b/Engine/Shaders/Private/Strata/Strata.ush
--- a/Engine/Shaders/Private/Strata/Strata.ush
+++ b/Engine/Shaders/Private/Strata/Strata.ush
@@ -471,6 +471,7 @@ struct FStrataBSDF
 // Fuzz to simulate cloth
 #define SLAB_FUZZ_COLOR(X)				X.VGPRs[4].xyz
 #define SLAB_FUZZ_AMOUNT(X)				X.VGPRs[4].w
+#define SLAB_FUZZ_ROUGHNESS(X)			X.VGPRs[2].w
 
 ///////////////////////////////////////////////////////////////////////////////
 // Volumetric BSDF members
@@ -532,10 +533,16 @@ uint   PackRGBA8(float4 rgba);
 float4 UnpackRGBA8(uint rgba);
 uint   Pack10F(float Value);
 float  Unpack10F(uint Value);
+uint   PackR6(float Value, float Dither);
+float  UnpackR6(uint Value);
+uint   PackR7(float Value, float Dither);
+float  UnpackR7(uint Value);
 uint   PackR11G11B10F(float3 rgb);
 float3 UnpackR11G11B10F(uint rgb);
 uint   PackR10G10B10F(float3 rgb);
 float3 UnpackR10G10B10F(uint rgb);
+uint   PackR7G7B6Gamma2(float3 rgb, float Dither);
+float3 UnpackR7G7B6Gamma2(uint rgb);
 uint   PackColorLinearToGamma2AlphaLinear(float4 rgba);
 float4 UnpackColorGamma2ToLinearAlphaLinear(uint rgba);
 uint   StrataPackNormal(in float3 Normal);
@@ -956,6 +963,21 @@ uint StrataSubsurfaceProfileIdTo8bits(float In)
 	return uint(In * 255.0f + 0.5f);
 }
 
+///////////////////////////////////////////////////////////////////////////////
+// Fuzz
+
+uint PackFuzz(half3 InColor, half InAmount, half InRoughness, half InDither)
+{
+	return PackR7G7B6Gamma2(InColor, InDither) | (PackR6(InAmount, InDither) << 20) | (PackR6(InRoughness, InDither) << 26);
+}
+
+void UnpackFuzz(uint In, inout half3 OutColor, inout half OutAmount, inout half OutRoughness)
+{
+	OutColor     = UnpackR7G7B6Gamma2(In & 0xFFFFF);
+	OutAmount    = UnpackR6((In >> 20) & 0x3F);
+	OutRoughness = UnpackR6((In >> 26) & 0x3F);
+}
+
 ///////////////////////////////////////////////////////////////////////////////
 // Misc
 
@@ -1344,7 +1366,7 @@ FStrataData GetStrataSlabBSDF(
 	float SSSProfileID, float3 SSSMFP, float SSSMFPScale, float SSSPhaseAnisotropy, float UseSSSDiffusion,
 	float3 Emissive, 
 	float SecondRoughness, float SecondRoughnessWeight, float SecondRoughnessAsSimpleClearCoat,
-	float FuzzAmount, float3 FuzzColor,
+	float FuzzAmount, float3 FuzzColor, float FuzzRoughness,
 	float ThicknessCm,
 	bool bIsThinSurface,
 	uint SharedLocalBasisIndex, inout uint SharedLocalBasisTypes)
@@ -1460,6 +1482,7 @@ FStrataData GetStrataSlabBSDF(
 
 	SLAB_FUZZ_AMOUNT		(StrataData.InlinedBSDF)		= FuzzAmount;
 	SLAB_FUZZ_COLOR			(StrataData.InlinedBSDF)		= FuzzColor;
+	SLAB_FUZZ_ROUGHNESS		(StrataData.InlinedBSDF)		= FuzzRoughness;
 
 #endif // !STRATA_FASTPATH
 
@@ -2186,6 +2209,7 @@ FStrataData StrataHorizontalMixingParameterBlending(FStrataData Background, FStr
 	SLAB_ROUGHNESS(ResultBSDF) =		lerp(SLAB_ROUGHNESS(BackgroundBSDF),			SLAB_ROUGHNESS(ForegroundBSDF),				ForegroundOtherMixFactor);
 	SLAB_ANISOTROPY(ResultBSDF) =		lerp(SLAB_ANISOTROPY(BackgroundBSDF),			SLAB_ANISOTROPY(ForegroundBSDF),			ForegroundOtherMixFactor);
 	SLAB_FUZZ_AMOUNT(ResultBSDF) =		lerp(SLAB_FUZZ_AMOUNT(BackgroundBSDF),			SLAB_FUZZ_AMOUNT(ForegroundBSDF),			ForegroundOtherMixFactor);
+	SLAB_FUZZ_ROUGHNESS(ResultBSDF)	=	lerp(SLAB_FUZZ_ROUGHNESS(BackgroundBSDF),		SLAB_FUZZ_ROUGHNESS(ForegroundBSDF),		ForegroundOtherMixFactor);
 	SLAB_HAZINESS(ResultBSDF) =			LerpHazinessParameterBlending(BackgroundBSDF,	ForegroundBSDF,								ForegroundOtherMixFactor);
 	
 
@@ -2286,6 +2310,7 @@ FStrataData StrataAddParameterBlending(FStrataData A, FStrataData B, float AMixF
 	SLAB_ROUGHNESS(ResultBSDF) =		lerp(SLAB_ROUGHNESS(BBSDF),						SLAB_ROUGHNESS(ABSDF),						AMixFactor);
 	SLAB_ANISOTROPY(ResultBSDF) =		lerp(SLAB_ANISOTROPY(BBSDF),					SLAB_ANISOTROPY(ABSDF),						AMixFactor);
 	SLAB_FUZZ_AMOUNT(ResultBSDF) =		lerp(SLAB_FUZZ_AMOUNT(BBSDF),					SLAB_FUZZ_AMOUNT(ABSDF),					AMixFactor);
+	SLAB_FUZZ_ROUGHNESS(ResultBSDF) =	lerp(SLAB_FUZZ_ROUGHNESS(BBSDF),				SLAB_FUZZ_ROUGHNESS(ABSDF),					AMixFactor);
 	SLAB_HAZINESS(ResultBSDF) =			LerpHazinessParameterBlending(BBSDF,			ABSDF,										AMixFactor);
 
 	if (BSDF_GETSSSTYPE(ResultBSDF) == SSS_TYPE_DIFFUSION_PROFILE)
@@ -2398,6 +2423,7 @@ FStrataBSDF StrataVerticalLayeringParameterBlending(in FStrataBSDF TopBSDF, in F
 		
 	SLAB_ROUGHNESS(ResultBSDF) =		lerp(SLAB_ROUGHNESS(TopBSDF),					SLAB_ROUGHNESS(BaseBSDF),	TopToBottomLerp);
 	SLAB_FUZZ_AMOUNT(ResultBSDF) =		lerp(SLAB_FUZZ_AMOUNT(TopBSDF),					SLAB_FUZZ_AMOUNT(BaseBSDF), TopToBottomLerp);
+	SLAB_FUZZ_ROUGHNESS(ResultBSDF) =	lerp(SLAB_FUZZ_ROUGHNESS(TopBSDF),				SLAB_FUZZ_ROUGHNESS(BaseBSDF),TopToBottomLerp);
 	SLAB_HAZINESS(ResultBSDF) =			LerpHazinessParameterBlending(TopBSDF,			BaseBSDF,					TopToBottomLerp);
 
 	// Anisotropy is disabled with vertical blending. 
@@ -2576,6 +2602,7 @@ void SanitizeStrataSlab(inout FStrataBSDF BSDF)
 
 	SLAB_FUZZ_AMOUNT(BSDF) = saturate(SLAB_FUZZ_AMOUNT(BSDF));
 	SLAB_FUZZ_COLOR(BSDF) = saturate(SLAB_FUZZ_COLOR(BSDF));
+	SLAB_FUZZ_ROUGHNESS(BSDF) = saturate(SLAB_FUZZ_ROUGHNESS(BSDF));
 }
 
 void SanitizeStrataHair(inout FStrataBSDF BSDF)
@@ -3133,6 +3160,19 @@ float UnpackR7(uint Value)
 	return float(Value & 0x7F) * (1.0f / 127.0f);
 }
 
+uint PackR6(float Value, float Dither)
+{
+	const float CenteredDither = Dither - 0.5f;
+	const float Dither7bits = CenteredDither * 1.0f / 63.0f;
+	const float DitheredValue = DitherIfNonBlack(saturate(Value), Dither7bits);
+	return uint(DitheredValue * 63.0f);
+}
+
+float UnpackR6(uint Value)
+{
+	return float(Value & 0x3F) * (1.0f / 63.0f);
+}
+
 // Simple Octahedral from "A Survey of Efficient Representations for Independent Unit Vectors" http://jcgt.org/published/0003/02/01/
 float SignNotZero(in float k)
 {
@@ -3428,6 +3468,7 @@ void StrataConvertFromDBuffer(
 		0.f,				// SecondRoughnessAsSimpleClearCoat
 		0.f,				// FuzzAmount
 		0.f,				// FuzzColor
+		In.Roughness,		// FuzzRoughness
 		DefaultThicknessCm,
 		false,				// bIsThinSurface,
 		SharedLocalBasisIndex,
@@ -3817,6 +3858,7 @@ FStrataBSDF UnpackStrataBSDFIn(FStrataMaterialContainer StrataBuffer, inout FStr
 				SLAB_FUZZ_COLOR(OutBSDF)	= UnpackR7G7B6Gamma2(PackedFuzzColor20bits);
 				SLAB_FUZZ_AMOUNT(OutBSDF)	= UnpackR8(PackedFuzzAmount8bits);
 				SLAB_ROUGHNESS(OutBSDF)		= StrataUnpackTopLayerData(StrataHeader.PackedTopLayerData).Roughness;
+				SLAB_FUZZ_ROUGHNESS(OutBSDF)= SLAB_ROUGHNESS(OutBSDF);
 			#endif
 
 				// Force a BSDF state to be a single layer on top.
@@ -4028,7 +4070,7 @@ FStrataBSDF UnpackStrataBSDFIn(FStrataMaterialContainer StrataBuffer, inout FStr
 
 		if (BSDF_GETHASFUZZ(OutBSDF))
 		{
-			StrataLoad_ColorGamma2ToLinearAlphaLinear(StrataBuffer, StrataAddressing, SLAB_FUZZ_COLOR(OutBSDF), SLAB_FUZZ_AMOUNT(OutBSDF));
+			UnpackFuzz((StrataLoadUint1(StrataBuffer, StrataAddressing)), SLAB_FUZZ_COLOR(OutBSDF), SLAB_FUZZ_AMOUNT(OutBSDF), SLAB_FUZZ_ROUGHNESS(OutBSDF));
 		}
 
 		if (BSDF_GETHASTRANSABOVE(OutBSDF))
diff --git a/Engine/Shaders/Private/Strata/StrataDbuffer.usf b/Engine/Shaders/Private/Strata/StrataDbuffer.usf
--- a/Engine/Shaders/Private/Strata/StrataDbuffer.usf
+++ b/Engine/Shaders/Private/Strata/StrataDbuffer.usf
@@ -248,7 +248,7 @@ void MainCS(
 
 				if (BSDF_GETHASFUZZ(BSDF))
 				{
-					MaterialTextureArrayUAV[uint3(PixelCoord, RWOffset)] = PackColorLinearToGamma2AlphaLinear(float4(FuzzColor, SLAB_FUZZ_AMOUNT(BSDF)));
+					MaterialTextureArrayUAV[uint3(PixelCoord, RWOffset)] = PackFuzz(FuzzColor, SLAB_FUZZ_AMOUNT(BSDF), SLAB_FUZZ_ROUGHNESS(BSDF), Dither);
 					RWOffset++;
 				}
 			}
diff --git a/Engine/Shaders/Private/Strata/StrataEvaluation.ush b/Engine/Shaders/Private/Strata/StrataEvaluation.ush
--- a/Engine/Shaders/Private/Strata/StrataEvaluation.ush
+++ b/Engine/Shaders/Private/Strata/StrataEvaluation.ush
@@ -611,17 +611,18 @@ FStrataEvaluateResult StrataEvaluateBSDFCommon(FStrataBSDFContext BSDFContext, F
 			BRANCH
 			if (BSDF_GETHASFUZZ(BSDFContext.BSDF))
 			{
-				const float FuzzAmount = SLAB_FUZZ_AMOUNT(BSDFContext.BSDF);
+				const float  FuzzAmount = SLAB_FUZZ_AMOUNT(BSDFContext.BSDF);
 				const float3 FuzzF0  = SLAB_FUZZ_COLOR(BSDFContext.BSDF);
+				const float  FuzzF90 = 1.0f;
+				const float  FuzzRoughness = MakeRoughnessSafe(SLAB_FUZZ_ROUGHNESS(BSDFContext.BSDF), STRATA_MIN_FUZZ_ROUGHNESS);
 
-				const float ClothF90 = 1.0f;
-				FBxDFEnergyTerms EnergyTerms = ComputeClothEnergyTerms(SafeRoughness, BSDFContext.Context.NoV, FuzzF0, ClothF90);
+				FBxDFEnergyTerms EnergyTerms = ComputeClothEnergyTerms(FuzzRoughness, BSDFContext.Context.NoV, FuzzF0, FuzzF90);
 
 				// The specular and diffuse components below the fuzz are only attenuated linearly according to the amount of fuzz.
 				const float3 Cloth_DirectionalAlbedo_SpecularTransmission = lerp(1.0, ComputeEnergyPreservation(EnergyTerms), FuzzAmount);
 
 				// Charlie cloth modle for D and Ashikhmin for visibility (more efficient than Charlie).
-				float ClothD = D_Charlie(SafeRoughness, BSDFContext.Context.NoH);
+				float ClothD   = D_Charlie(FuzzRoughness, BSDFContext.Context.NoH);
 				float ClothVis = Vis_Ashikhmin(BSDFContext.Context.NoV, AreaLightContext.NoL);
 				float3 ClothF  = F_Schlick(FuzzF0, BSDFContext.Context.VoH);
 				float3 ClothSpecularPathValue = (ClothD * ClothVis) * ClothF * ComputeEnergyConservation(EnergyTerms);
@@ -1459,11 +1460,12 @@ FStrataEnvLightResult StrataEvaluateForEnvLight(FStrataBSDFContext BSDFContext,
 			BRANCH
 			if (BSDF_GETHASFUZZ(BSDFContext.BSDF))
 			{
-				const float FuzzAmount = SLAB_FUZZ_AMOUNT(BSDFContext.BSDF);
+				const float  FuzzAmount = SLAB_FUZZ_AMOUNT(BSDFContext.BSDF);
 				const float3 FuzzF0 = SLAB_FUZZ_COLOR(BSDFContext.BSDF);
+				const float  FuzzF90 = 1.0f;
+				const float  FuzzRoughness = MakeRoughnessSafe(SLAB_FUZZ_ROUGHNESS(BSDFContext.BSDF), STRATA_MIN_FUZZ_ROUGHNESS);
 
-				const float ClothF90 = 1.0f;
-				const FBxDFEnergyTermsRGB ClothEnergyTerms = ComputeClothEnergyTermsRGB(SafeRoughness, BSDFContext.Context.NoV, FuzzF0, ClothF90);
+				const FBxDFEnergyTermsRGB ClothEnergyTerms = ComputeClothEnergyTermsRGB(FuzzRoughness, BSDFContext.Context.NoV, FuzzF0, FuzzF90);
 
 				// The specular and diffuse components below the fuzz are only attenuated linearly according to the amount of fuzz.
 				const float Cloth_DirectionalAlbedo_SpecularTransmission = lerp(1.0, ComputeEnergyPreservation(ClothEnergyTerms), FuzzAmount);
@@ -1587,4 +1589,4 @@ FThreeBandSHVector StrataBSDFToSH(FStrataBSDFContext BSDFContext)
 	return SHVector;
 }
 
-#endif // STRATA_ENABLED
\ No newline at end of file
+#endif // STRATA_ENABLED
diff --git a/Engine/Shaders/Private/Strata/StrataExport.ush b/Engine/Shaders/Private/Strata/StrataExport.ush
--- a/Engine/Shaders/Private/Strata/StrataExport.ush
+++ b/Engine/Shaders/Private/Strata/StrataExport.ush
@@ -326,7 +326,7 @@ void PackStrataOut(
 						OptimisedModeMask |= SingleHasClearCoat;
 					}
 				}
-				if (BSDF_GETHASFUZZ(BSDF))
+				if (BSDF_GETHASFUZZ(BSDF) && SLAB_ROUGHNESS(BSDF) == SLAB_FUZZ_ROUGHNESS(BSDF))
 				{
 					OptimisedModeMask |= SingleHasCloth;
 				}
@@ -564,7 +564,7 @@ void PackStrataOut(
 							const uint PackedF020Bits		= PackR7G7B6Gamma2(	SLAB_F0(BSDF)				, Dither);
 							const uint PackedFuzzColor20bits= PackR7G7B6Gamma2(	SLAB_FUZZ_COLOR(BSDF)		, Dither);
 							const uint PackedFuzzAmount8bits= PackR8(			SLAB_FUZZ_AMOUNT(BSDF));
-							// Roughness is stored in the TopLayer texture.
+							// Roughness is stored in the TopLayer texture. Fuzz roughness is identical to slab roughness in this path.
 
 							// MRT0 
 							#if (32 - (HEADER_SINGLEENCODING_BIT_COUNT + HEADER_SINGLE_OPTLEGACYMODE_BIT_COUNT)) < 8
@@ -775,7 +775,7 @@ void PackStrataOut(
 
 							if (BSDF_GETHASFUZZ(BSDF))
 							{
-								STRATA_STORE_UINT1(PackColorLinearToGamma2AlphaLinear(float4(SLAB_FUZZ_COLOR(BSDF), SLAB_FUZZ_AMOUNT(BSDF))));
+								STRATA_STORE_UINT1(PackFuzz(SLAB_FUZZ_COLOR(BSDF), SLAB_FUZZ_AMOUNT(BSDF), SLAB_FUZZ_ROUGHNESS(BSDF), Dither));
 							}
 
 							if (BSDF_GETHASTRANSABOVE(BSDF))
diff --git a/Engine/Shaders/Private/Strata/StrataLegacyConversion.ush b/Engine/Shaders/Private/Strata/StrataLegacyConversion.ush
--- a/Engine/Shaders/Private/Strata/StrataLegacyConversion.ush
+++ b/Engine/Shaders/Private/Strata/StrataLegacyConversion.ush
@@ -119,6 +119,7 @@ FStrataData CreateLegacySlab(
 		In.SecondRoughnessAsSimpleClearCoat ? 1.0f : 0.0f,	// SecondRoughnessAsSimpleClearCoat
 		In.FuzzAmount,										// FuzzAmount
 		In.FuzzColor,										// FuzzColor
+		In.Roughness,										// FuzzRoughness
 		In.Thickness,
 		In.bIsThinSurface,
 		In.SharedLocalBasisIndex, SharedLocalBasisTypes);
diff --git a/Engine/Shaders/Private/Strata/StrataVisualize.usf b/Engine/Shaders/Private/Strata/StrataVisualize.usf
--- a/Engine/Shaders/Private/Strata/StrataVisualize.usf
+++ b/Engine/Shaders/Private/Strata/StrataVisualize.usf
@@ -279,6 +279,7 @@ void PrintBSDF(inout FShaderPrintContext Context, inout FStrataAddressing Strata
 			{
 				Print(Context, TEXT(" Fuzz Amount "), FontBSDFPropName);	Print(Context, SLAB_FUZZ_AMOUNT(BSDF), FontBSDFPropValu);				Newline(Context, RectMax);
 				Print(Context, TEXT(" Fuzz Color  "), FontBSDFPropName);	Print(Context, SLAB_FUZZ_COLOR(BSDF), FontBSDFPropValu);				Newline(Context, RectMax);
+				Print(Context, TEXT(" Fuzz Rough. "), FontBSDFPropName);	Print(Context, SLAB_FUZZ_ROUGHNESS(BSDF), FontBSDFPropValu);			Newline(Context, RectMax);
 			}
 		}
 		break;
diff --git a/Engine/Shaders/Shared/StrataDefinitions.h b/Engine/Shaders/Shared/StrataDefinitions.h
--- a/Engine/Shaders/Shared/StrataDefinitions.h
+++ b/Engine/Shaders/Shared/StrataDefinitions.h
@@ -36,6 +36,9 @@
 // The default thickness of a layer is considered to be 0.01 centimeter = 0.1 millimeter
 #define STRATA_LAYER_DEFAULT_THICKNESS_CM	0.01f
 
+// Min Fuzz Roughness to avoid numerical issue
+#define STRATA_MIN_FUZZ_ROUGHNESS 0.02f
+
 #define STRATA_BASE_PASS_MRT_OUTPUT_COUNT	3
 
 #define STRATA_SSS_DATA_UINT_COUNT		2
diff --git a/Engine/Source/Runtime/Engine/Classes/Materials/MaterialExpressionStrata.h b/Engine/Source/Runtime/Engine/Classes/Materials/MaterialExpressionStrata.h
--- a/Engine/Source/Runtime/Engine/Classes/Materials/MaterialExpressionStrata.h
+++ b/Engine/Source/Runtime/Engine/Classes/Materials/MaterialExpressionStrata.h
@@ -278,10 +278,10 @@ class UMaterialExpressionStrataSlabBSDF : public UMaterialExpressionStrataBSDF
 	FExpressionInput SecondRoughnessWeight;
 
 	/**
-	 * The slab thickness. (type = float, centimeters, default = 0.01 centimeter = 0.1 millimeter)
+	 * Controls how rough the Fuzz layer is. Roughness of 0 is smooth and 1 is rough. If FuzzRoughness Is not connected, the Roughness input will be used instead. (type = float, unit = unitless, defaults to 0.5)
 	 */
 	UPROPERTY()
-	FExpressionInput Thickness; // TODO: Remove
+	FExpressionInput FuzzRoughness;
 
 	/**
 	 * The amount of fuzz on top of the surface used to simulate cloth-like appearance.
@@ -318,6 +318,7 @@ class UMaterialExpressionStrataSlabBSDF : public UMaterialExpressionStrataBSDF
 
 	bool HasEdgeColor() const;
 	bool HasFuzz() const;
+	bool HasFuzzRoughness() const;
 	bool HasSecondRoughness() const;
 	bool HasSSS() const;
 	bool HasSSSProfile() const;
diff --git a/Engine/Source/Runtime/Engine/Private/Materials/HLSLMaterialTranslator.cpp b/Engine/Source/Runtime/Engine/Private/Materials/HLSLMaterialTranslator.cpp
--- a/Engine/Source/Runtime/Engine/Private/Materials/HLSLMaterialTranslator.cpp
+++ b/Engine/Source/Runtime/Engine/Private/Materials/HLSLMaterialTranslator.cpp
@@ -10706,7 +10706,7 @@ int32 FHLSLMaterialTranslator::StrataSlabBSDF(
 	int32 SSSProfileId, int32 SSSMFP, int32 SSSMFPScale, int32 SSSPhaseAniso, int32 UseSSSDiffusion,
 	int32 EmissiveColor,
 	int32 SecondRoughness, int32 SecondRoughnessWeight, int32 SecondRoughnessAsSimpleClearCoat,
-	int32 FuzzAmount, int32 FuzzColor,
+	int32 FuzzAmount, int32 FuzzColor, int32 FuzzRoughness,
 	int32 Thickness, 
 	int32 Normal, int32 Tangent, const FString& SharedLocalBasisIndexMacro,
 	FStrataOperator* PromoteToOperator)
@@ -10724,7 +10724,7 @@ int32 FHLSLMaterialTranslator::StrataSlabBSDF(
 			return INDEX_NONE;
 		}
 		return AddCodeChunk(
-			MCT_Strata, TEXT("PromoteParameterBlendedBSDFToOperator(GetStrataSlabBSDF(Parameters.StrataPixelFootprint, %s, %s, %s, %s, %s, %s, %s, %s, %s, %s, %s, %s, %s, %s, %s, %s, %s, %s, %s, Parameters.SharedLocalBases.Types) /* Normal = %s ; Tangent = %s ; Thickness = %s */, Parameters.StrataTree, %u, %u, %u, %u)"),
+			MCT_Strata, TEXT("PromoteParameterBlendedBSDFToOperator(GetStrataSlabBSDF(Parameters.StrataPixelFootprint, %s, %s, %s, %s, %s, %s, %s, %s, %s, %s, %s, %s, %s, %s, %s, %s, %s, %s, %s, %s, Parameters.SharedLocalBases.Types) /* Normal = %s ; Tangent = %s ; Thickness = %s */, Parameters.StrataTree, %u, %u, %u, %u)"),
 			*StrataGetCastParameterCode(DiffuseAlbedo,			MCT_Float3),
 			*StrataGetCastParameterCode(F0,						MCT_Float3),
 			*StrataGetCastParameterCode(F90,					MCT_Float3),
@@ -10741,6 +10741,7 @@ int32 FHLSLMaterialTranslator::StrataSlabBSDF(
 			*StrataGetCastParameterCode(SecondRoughnessAsSimpleClearCoat, MCT_Float),
 			*StrataGetCastParameterCode(FuzzAmount,				MCT_Float),
 			*StrataGetCastParameterCode(FuzzColor,				MCT_Float3),
+			*StrataGetCastParameterCode(FuzzRoughness,			MCT_Float1),
 			*StrataGetCastParameterCode(Thickness,				MCT_Float),
 			bIsThinSurface ? TEXT("true") : TEXT("false"),
 			*SharedLocalBasisIndexMacro,
@@ -10755,7 +10756,7 @@ int32 FHLSLMaterialTranslator::StrataSlabBSDF(
 	}
 	
 	return AddCodeChunk(
-		MCT_Strata, TEXT("GetStrataSlabBSDF(Parameters.StrataPixelFootprint, %s, %s, %s, %s, %s, %s, %s, %s, %s, %s, %s, %s, %s, %s, %s, %s, %s, %s, %s, Parameters.SharedLocalBases.Types) /* Normal = %s ; Tangent = %s ; Thickness = %s */"),
+		MCT_Strata, TEXT("GetStrataSlabBSDF(Parameters.StrataPixelFootprint, %s, %s, %s, %s, %s, %s, %s, %s, %s, %s, %s, %s, %s, %s, %s, %s, %s, %s, %s, %s, Parameters.SharedLocalBases.Types) /* Normal = %s ; Tangent = %s ; Thickness = %s */"),
 		*StrataGetCastParameterCode(DiffuseAlbedo,			MCT_Float3),
 		*StrataGetCastParameterCode(F0,						MCT_Float3),
 		*StrataGetCastParameterCode(F90,					MCT_Float3),
@@ -10772,6 +10773,7 @@ int32 FHLSLMaterialTranslator::StrataSlabBSDF(
 		*StrataGetCastParameterCode(SecondRoughnessAsSimpleClearCoat, MCT_Float),
 		*StrataGetCastParameterCode(FuzzAmount,				MCT_Float),
 		*StrataGetCastParameterCode(FuzzColor,				MCT_Float3),
+		*StrataGetCastParameterCode(FuzzRoughness,			MCT_Float1),
 		*StrataGetCastParameterCode(Thickness,				MCT_Float),
 		bIsThinSurface ? TEXT("true") : TEXT("false"),
 		*SharedLocalBasisIndexMacro,
diff --git a/Engine/Source/Runtime/Engine/Private/Materials/HLSLMaterialTranslator.h b/Engine/Source/Runtime/Engine/Private/Materials/HLSLMaterialTranslator.h
--- a/Engine/Source/Runtime/Engine/Private/Materials/HLSLMaterialTranslator.h
+++ b/Engine/Source/Runtime/Engine/Private/Materials/HLSLMaterialTranslator.h
@@ -1074,7 +1074,7 @@ protected:
 		int32 SSSProfileId, int32 SSSMFP, int32 SSSMFPScale, int32 SSSPhaseAniso, int32 bUseSSSDiffusion,
 		int32 EmissiveColor, 
 		int32 SecondRoughness, int32 SecondRoughnessWeight, int32 SecondRoughnessAsSimpleClearCoat,
-		int32 FuzzAmount, int32 FuzzColor, 
+		int32 FuzzAmount, int32 FuzzColor, int32 FuzzRoughness,
 		int32 Thickness,
 		int32 Normal, int32 Tangent, const FString& SharedLocalBasisIndexMacro,
 		FStrataOperator* PromoteToOperator) override;
diff --git a/Engine/Source/Runtime/Engine/Private/Materials/MaterialExpressions.cpp b/Engine/Source/Runtime/Engine/Private/Materials/MaterialExpressions.cpp
--- a/Engine/Source/Runtime/Engine/Private/Materials/MaterialExpressions.cpp
+++ b/Engine/Source/Runtime/Engine/Private/Materials/MaterialExpressions.cpp
@@ -22511,10 +22511,11 @@ int32 UMaterialExpressionStrataSlabBSDF::Compile(class FMaterialCompiler* Compil
 
 	const bool bHasEdgeColor = HasEdgeColor();
 	const bool bHasFuzz = HasFuzz();
+	const bool bHasFuzzRoughness = HasFuzzRoughness();
 	const bool bHasSecondRoughness = HasSecondRoughness();
 	const bool bHasMFPPluggedIn = HasMFPPluggedIn();
-	int32 SSSProfileCodeChunk = INDEX_NONE;
 	const bool bHasSSS = HasSSS();
+	int32 SSSProfileCodeChunk = INDEX_NONE;
 	if (bHasSSS)
 	{
 		SSSProfileCodeChunk = Compiler->ForceCast(Compiler->ScalarParameter(GetSubsurfaceProfileParameterName(), 1.0f), MCT_Float1);
@@ -22559,6 +22560,7 @@ int32 UMaterialExpressionStrataSlabBSDF::Compile(class FMaterialCompiler* Compil
 		Compiler->Constant(0.0f),										// SecondRoughnessAsSimpleClearCoat
 		CompileWithDefaultFloat1(Compiler, FuzzAmount, 0.0f),
 		CompileWithDefaultFloat3(Compiler, FuzzColor, 0.0f, 0.0f, 0.0f),
+		bHasFuzzRoughness ? CompileWithDefaultFloat1(Compiler, FuzzRoughness, 0.5f) : RoughnessCodeChunk,
 		ThicknesCodeChunk,
 		NormalCodeChunk,
 		TangentCodeChunk,
@@ -22585,7 +22587,7 @@ const TArray<FExpressionInput*> UMaterialExpressionStrataSlabBSDF::GetInputs()
 	Result.Add(&EmissiveColor);
 	Result.Add(&SecondRoughness);
 	Result.Add(&SecondRoughnessWeight);
-	Result.Add(&Thickness);
+	Result.Add(&FuzzRoughness);
 	Result.Add(&FuzzAmount);
 	Result.Add(&FuzzColor);
 	return Result;
@@ -22657,7 +22659,7 @@ uint32 UMaterialExpressionStrataSlabBSDF::GetInputType(int32 InputIndex)
 	}
 	else if (InputIndex == 13)
 	{
-		return MCT_Float1; // Thickness
+		return MCT_Float1; // FuzzRoughness
 	}
 	else if (InputIndex == 14)
 	{
@@ -22728,7 +22730,7 @@ FName UMaterialExpressionStrataSlabBSDF::GetInputName(int32 InputIndex) const
 	}
 	else if (InputIndex == 13)
 	{
-		return TEXT("Thickness (DEPRECATED)");
+		return TEXT("Fuzz Roughness");
 	}
 	else if (InputIndex == 14)
 	{
@@ -22824,6 +22826,11 @@ bool UMaterialExpressionStrataSlabBSDF::HasFuzz() const
 	return FuzzAmount.IsConnected();
 }
 
+bool UMaterialExpressionStrataSlabBSDF::HasFuzzRoughness() const
+{
+	return FuzzRoughness.IsConnected();
+}
+
 bool UMaterialExpressionStrataSlabBSDF::HasSecondRoughness() const
 {
 	return SecondRoughnessWeight.IsConnected();
@@ -22867,12 +22874,13 @@ int32 UMaterialExpressionStrataSimpleClearCoatBSDF::Compile(class FMaterialCompi
 	int32 ThicknessCodeChunk = Compiler->StrataThicknessStackGetThicknessCode(StrataOperator.ThicknessIndex);
 	check(ThicknessCodeChunk != INDEX_NONE);
 
+	int32 RoughnessCodeChunk = CompileWithDefaultFloat1(Compiler, Roughness, 0.5f);
+
 	int32 OutputCodeChunk = Compiler->StrataSlabBSDF(
 		CompileWithDefaultFloat3(Compiler, DiffuseAlbedo, 0.18f, 0.18f, 0.18f),		// DiffuseAlbedo
 		CompileWithDefaultFloat3(Compiler, F0, DefaultF0, DefaultF0, DefaultF0),	// F0
-		Compiler->Constant3(1.0f, 1.0f, 1.0f),					// F90
-		
-		CompileWithDefaultFloat1(Compiler, Roughness, 0.5f),
+		Compiler->Constant3(1.0f, 1.0f, 1.0f),					// F90		
+		RoughnessCodeChunk,										// Roughness
 		Compiler->Constant(0.0f),								// Anisotropy
 		Compiler->Constant(0.0f),								// SSSProfile
 		Compiler->Constant3(0.0f, 0.0f, 0.0f),					// SSSMFP
@@ -22885,6 +22893,7 @@ int32 UMaterialExpressionStrataSimpleClearCoatBSDF::Compile(class FMaterialCompi
 		Compiler->Constant(1.0f),								// SecondRoughnessAsSimpleClearCoat == true for UMaterialExpressionStrataSimpleClearCoatBSDF
 		Compiler->Constant(0.0f),								// FuzzAmount
 		Compiler->Constant3(0.0f, 0.0f, 0.0f),					// FuzzColor
+		RoughnessCodeChunk,										// FuzzRoughness
 		ThicknessCodeChunk,										// Thickness
 		NormalCodeChunk,
 		NullTangentCodeChunk,
diff --git a/Engine/Source/Runtime/Engine/Public/MaterialCompiler.h b/Engine/Source/Runtime/Engine/Public/MaterialCompiler.h
--- a/Engine/Source/Runtime/Engine/Public/MaterialCompiler.h
+++ b/Engine/Source/Runtime/Engine/Public/MaterialCompiler.h
@@ -640,7 +640,7 @@ public:
 		int32 SSSProfileId, int32 SSSMFP, int32 SSSMFPScale, int32 SSSPhaseAniso, int32 UseSSSDiffusion,
 		int32 EmissiveColor, 
 		int32 SecondRoughness, int32 SecondRoughnessWeight, int32 SecondRoughnessAsSimpleClearCoat,
-		int32 FuzzAmount, int32 FuzzColor,
+		int32 FuzzAmount, int32 FuzzColor, int32 FuzzRoughness,
 		int32 Thickness,
 		int32 Normal, int32 Tangent, const FString& SharedLocalBasisIndexMacro,
 		FStrataOperator* PromoteToOperator) = 0;
@@ -1236,7 +1236,7 @@ public:
 		int32 SSSProfileId, int32 SSSMFP, int32 SSSMFPScale, int32 SSSPhaseAniso, int32 UseSSSDiffusion,
 		int32 EmissiveColor, 
 		int32 SecondRoughness, int32 SecondRoughnessWeight, int32 SecondRoughnessAsSimpleClearCoat,
-		int32 FuzzAmount, int32 FuzzColor,
+		int32 FuzzAmount, int32 FuzzColor, int32 FuzzRoughness,
 		int32 Thickness,
 		int32 Normal, int32 Tangent, const FString& SharedLocalBasisIndexMacro,
 		FStrataOperator* PromoteToOperator) override
@@ -1247,7 +1247,7 @@ public:
 			SSSProfileId, SSSMFP, SSSMFPScale, SSSPhaseAniso, UseSSSDiffusion,
 			EmissiveColor, 
 			SecondRoughness, SecondRoughnessWeight, SecondRoughnessAsSimpleClearCoat,
-			FuzzAmount, FuzzColor,
+			FuzzAmount, FuzzColor, FuzzRoughness,
 			Thickness,
 			Normal, Tangent, SharedLocalBasisIndexMacro,
 			PromoteToOperator);
