Change Fuzz energy conservation/preservation for Strata.

diff --git a/Engine/Shaders/Private/ShadingEnergyConservationTemplate.ush b/Engine/Shaders/Private/ShadingEnergyConservationTemplate.ush
--- a/Engine/Shaders/Private/ShadingEnergyConservationTemplate.ush
+++ b/Engine/Shaders/Private/ShadingEnergyConservationTemplate.ush
@@ -74,12 +74,13 @@ BXDF_ENERGY_SUFFIX(FBxDFEnergyTerms) BXDF_ENERGY_SUFFIX(ComputeGGXSpecEnergyTerm
 	return BXDF_ENERGY_SUFFIX(ComputeGGXSpecEnergyTerms)(Roughness, NoV, F0, F90);
 }
 
-BXDF_ENERGY_SUFFIX(FBxDFEnergyTerms) BXDF_ENERGY_SUFFIX(ComputeClothEnergyTerms)(float Roughness, float NoV, float3 F0, float3 F90)
+BXDF_ENERGY_SUFFIX(FBxDFEnergyTerms) BXDF_ENERGY_SUFFIX(ComputeClothEnergyTerms)(float Roughness, float NoV, float3 F0)
 {
 	BXDF_ENERGY_SUFFIX(FBxDFEnergyTerms) Out;
 #if USE_ENERGY_CONSERVATION > 0
 	{
-		Out = BXDF_ENERGY_SUFFIX(ComputeFresnelEnergyTerms)(ClothEnergyLookup(Roughness, NoV), F0, F90);
+		Out.W = 1.f;
+		Out.E = BXDF_ENERGY_SUFFIX(GetF0F90)(F0) * ClothEnergyLookup(Roughness, NoV).x;
 	}
 #else
 	{
@@ -90,12 +91,6 @@ BXDF_ENERGY_SUFFIX(FBxDFEnergyTerms) BXDF_ENERGY_SUFFIX(ComputeClothEnergyTerms)
 	return Out;
 }
 
-BXDF_ENERGY_SUFFIX(FBxDFEnergyTerms) BXDF_ENERGY_SUFFIX(ComputeClothEnergyTerms)(float Roughness, float NoV, float3 F0)
-{
-	const float F90 = saturate(50.0 * F0.g); // See F_Schlick implementation
-	return BXDF_ENERGY_SUFFIX(ComputeClothEnergyTerms)(Roughness, NoV, F0, F90);
-}
-
 BXDF_ENERGY_SUFFIX(FBxDFEnergyTerms) BXDF_ENERGY_SUFFIX(ComputeDiffuseEnergyTerms)(float Roughness, float NoV)
 {
 	BXDF_ENERGY_SUFFIX(FBxDFEnergyTerms) Out;
diff --git a/Engine/Shaders/Private/Strata/StrataEvaluation.ush b/Engine/Shaders/Private/Strata/StrataEvaluation.ush
--- a/Engine/Shaders/Private/Strata/StrataEvaluation.ush
+++ b/Engine/Shaders/Private/Strata/StrataEvaluation.ush
@@ -616,7 +616,8 @@ FStrataEvaluateResult StrataEvaluateBSDFCommon(FStrataBSDFContext BSDFContext, F
 				const float  FuzzF90 = 1.0f;
 				const float  FuzzRoughness = MakeRoughnessSafe(SLAB_FUZZ_ROUGHNESS(BSDFContext.BSDF), STRATA_MIN_FUZZ_ROUGHNESS);
 
-				FBxDFEnergyTerms EnergyTerms = ComputeClothEnergyTerms(FuzzRoughness, BSDFContext.Context.NoV, FuzzF0, FuzzF90);
+				// F0 is ignored, as energy preservation is handle solely based on the white directional albedo and FuzzAmount
+				FBxDFEnergyTerms EnergyTerms = ComputeClothEnergyTerms(FuzzRoughness, BSDFContext.Context.NoV, 1.0f /*FuzzF0*/);
 
 				// The specular and diffuse components below the fuzz are only attenuated linearly according to the amount of fuzz.
 				const float3 Cloth_DirectionalAlbedo_SpecularTransmission = lerp(1.0, ComputeEnergyPreservation(EnergyTerms), FuzzAmount);
@@ -1465,7 +1466,8 @@ FStrataEnvLightResult StrataEvaluateForEnvLight(FStrataBSDFContext BSDFContext,
 				const float  FuzzF90 = 1.0f;
 				const float  FuzzRoughness = MakeRoughnessSafe(SLAB_FUZZ_ROUGHNESS(BSDFContext.BSDF), STRATA_MIN_FUZZ_ROUGHNESS);
 
-				const FBxDFEnergyTermsRGB ClothEnergyTerms = ComputeClothEnergyTermsRGB(FuzzRoughness, BSDFContext.Context.NoV, FuzzF0, FuzzF90);
+				// F0 is ignored, as energy preservation is handle solely based on the white directional albedo and FuzzAmount
+				const FBxDFEnergyTermsRGB ClothEnergyTerms = ComputeClothEnergyTermsRGB(FuzzRoughness, BSDFContext.Context.NoV, 1.0f /*FuzzF0*/);
 
 				// The specular and diffuse components below the fuzz are only attenuated linearly according to the amount of fuzz.
 				const float Cloth_DirectionalAlbedo_SpecularTransmission = lerp(1.0, ComputeEnergyPreservation(ClothEnergyTerms), FuzzAmount);
@@ -1478,8 +1480,8 @@ FStrataEnvLightResult StrataEvaluateForEnvLight(FStrataBSDFContext BSDFContext,
 				
 				// We currently send the cloth environment lighting to be taken into account by the diffuse term.
 				// STRATA_TODO: true a mix of diffuse and specular contribution according to roughness.
-				StrataEnvLightResult.DiffuseColor		+= ClothEnergyTerms.E * FuzzAmount;
-				StrataEnvLightResult.DiffuseWeight		+= ClothEnergyTerms.E * FuzzAmount;
+				StrataEnvLightResult.DiffuseColor		+= FuzzF0 * ClothEnergyTerms.E * FuzzAmount;
+				StrataEnvLightResult.DiffuseWeight		+= FuzzF0 * ClothEnergyTerms.E * FuzzAmount;
 				break;
 			}
 			#endif
