Fix for water looking pink-ish like a milkshake on mobile under certain conditions. Fixed by multiplying the indirect diffuse term used on water by the sky light color (similar to what the non-mobile version does).

#rnx

diff --git a/Engine/Shaders/Private/MobileBasePassPixelShader.usf b/Engine/Shaders/Private/MobileBasePassPixelShader.usf
--- a/Engine/Shaders/Private/MobileBasePassPixelShader.usf
+++ b/Engine/Shaders/Private/MobileBasePassPixelShader.usf
@@ -175,7 +175,7 @@ void GetSkyLighting(half3 SkyLightColor, half3 WorldNormal, bool bEvaluateBackfa
 		IndirectIrradiance += Luminance(OutSkyDiffuseLighting);
 
 	#if MATERIAL_SHADINGMODEL_SINGLELAYERWATER
-		OutWaterDiffuseIndirectLuminance += SkyDiffuse;
+		OutWaterDiffuseIndirectLuminance += OutSkyDiffuseLighting;
 	#endif
 
 	#if SHADINGMODEL_REQUIRES_BACKFACE_LIGHTING
