Fix additive material rendering when per-pixel OIT is enabled.

diff --git a/Engine/Shaders/Private/BasePassPixelShader.usf b/Engine/Shaders/Private/BasePassPixelShader.usf
--- a/Engine/Shaders/Private/BasePassPixelShader.usf
+++ b/Engine/Shaders/Private/BasePassPixelShader.usf
@@ -2033,11 +2033,10 @@ void FPixelShaderInOut_MainPS(
 		#else
 		{
 			float3 WriteColor = Out.MRT[0].rgb;
-			float  Opacity = Out.MRT[0].a;
 		#if MATERIALBLENDING_TRANSLUCENT
 			WriteColor *= Opacity; // WriteColor is premultiplied alpha color
 		#elif MATERIALBLENDING_ADDITIVE
-			Opacity = 1.0f;
+			// NOP
 		#elif MATERIALBLENDING_ALPHACOMPOSITE || MATERIALBLENDING_ALPHAHOLDOUT
 			// Not clear what should be done in this case.
 		#else
