Fix RT PSO creation error when Strata is enabled with high-quality normal.

diff --git a/Engine/Shaders/Private/Strata/Strata.ush b/Engine/Shaders/Private/Strata/Strata.ush
--- a/Engine/Shaders/Private/Strata/Strata.ush
+++ b/Engine/Shaders/Private/Strata/Strata.ush
@@ -1639,7 +1639,13 @@ FStrataData GetStrataSingleLayerWaterBSDF(
 // Raytracing material buffer
 
 #ifndef STRATA_RT_PAYLOAD_NUM_UINTS
-#define STRATA_RT_PAYLOAD_NUM_UINTS  9
+  // When STRATA_NORMAL_QUALITY == 1, PackedTopLayerData is a uint2. To keep FPackedMaterialClosestHitPayload 
+  // at 64 bytes (see declaration in BuiltInRayTracingShaders.cpp), we make data payload 4 bytes shorter)
+  #if STRATA_NORMAL_QUALITY == 0
+	#define STRATA_RT_PAYLOAD_NUM_UINTS  9
+  #else
+	#define STRATA_RT_PAYLOAD_NUM_UINTS  8
+  #endif
 #endif
 
 struct FStrataRaytracingPayload
@@ -1647,7 +1653,7 @@ struct FStrataRaytracingPayload
 	// Total size: 40 bytes. (assuming the default value of STRATA_RT_PAYLOAD_NUM_UINTS)
 	// If this changes, update FPackedMaterialClosestHitPayload in RayTracingCommon.ush 
 	// and ensure the payload size is smaller than the maximum size
-	STRATA_TOP_LAYER_TYPE PackedTopLayerData;
+	STRATA_TOP_LAYER_TYPE PackedTopLayerData; //
 	uint Data[STRATA_RT_PAYLOAD_NUM_UINTS];
 };
 
