Fix issue in Global SDF when multiple objects with Mesh SDF that extend outside the bounding box overlap.

Context:
- When composing objects into GlobalSDF we shorten the max distance based on the distance to the closest surface found so far
- the condition `if (DistanceToBox < MaxDistance)` is incorrect if mesh SDFs stick out of the bounding box
- this can cause instability in Global SDF since result will depend on the order objects are composited in

- we have logic to prevent this but only when the query position is outside of the bounding box
- this is a problem since Global SDF voxel centers can be inside the box

Fix:
- also prevent negative SDF distance to stick out of the bounding box when query position is inside the box.
- Modified how DistanceToBox is calculated to be valid when point is inside box.

diff --git a/Engine/Shaders/Private/MeshDistanceFieldCommon.ush b/Engine/Shaders/Private/MeshDistanceFieldCommon.ush
--- a/Engine/Shaders/Private/MeshDistanceFieldCommon.ush
+++ b/Engine/Shaders/Private/MeshDistanceFieldCommon.ush
@@ -6,23 +6,24 @@
 
 float DistanceToNearestSurfaceForObject(FDFObjectData DFObjectData, FLWCVector3 WorldPosition, float MaxDistance)
 {
-	float3 VolumePosition = LWCMultiply(WorldPosition, DFObjectData.WorldToVolume);
+	const float3 VolumePosition = LWCMultiply(WorldPosition, DFObjectData.WorldToVolume);
 
-	float3 ClampedVolumePosition = clamp(VolumePosition, -DFObjectData.VolumePositionExtent, DFObjectData.VolumePositionExtent);
-	float DistanceToBox = length((VolumePosition - ClampedVolumePosition) * DFObjectData.VolumeToWorldScale);
+	const float3 ClampedVolumePosition = clamp(VolumePosition, -DFObjectData.VolumePositionExtent, DFObjectData.VolumePositionExtent);
+
+	const float3 ToBox = (abs(VolumePosition) - DFObjectData.VolumePositionExtent) * DFObjectData.VolumeToWorldScale;
+	const float DistanceToBoxOutside = length(max(0.0f, ToBox));
+	const float DistanceToBoxInside = min(0.0f, max3(ToBox.x, ToBox.y, ToBox.z));
+	const float DistanceToBox = DistanceToBoxOutside + DistanceToBoxInside;
 
 	float DistanceToOccluder = MaxDistance;
 
 	BRANCH
 	if (DistanceToBox < MaxDistance)
 	{
-		DistanceToOccluder = DistanceToMeshSurfaceStandalone(ClampedVolumePosition, DFObjectData) * DFObjectData.VolumeScale + DistanceToBox;
+		DistanceToOccluder = DistanceToMeshSurfaceStandalone(ClampedVolumePosition, DFObjectData) * DFObjectData.VolumeScale + max(0.0f, DistanceToBox);
 
 		// Don't allow negative SDF distance to stick out of the bounding box
-		if (DistanceToBox > 0.0f)
-		{
-			DistanceToOccluder = max(DistanceToOccluder, DistanceToBox);
-		}
+		DistanceToOccluder = max(DistanceToOccluder, DistanceToBox);
 	}
 
 	return DistanceToOccluder;
