Fix groom cluster debug display.

diff --git a/Engine/Shaders/Private/HairStrands/HairStrandsClusterCulling.usf b/Engine/Shaders/Private/HairStrands/HairStrandsClusterCulling.usf
--- a/Engine/Shaders/Private/HairStrands/HairStrandsClusterCulling.usf
+++ b/Engine/Shaders/Private/HairStrands/HairStrandsClusterCulling.usf
@@ -100,41 +100,43 @@ void MainDrawDebugAABBCS(uint3 DispatchThreadId : SV_DispatchThreadID)
 	if (DispatchThreadId.x == 0)
 	{
 		const bool bDrawClusterInfo = ClusterDebugMode == MODE_CLUSTER || ClusterDebugMode == MODE_CULLED || ClusterDebugMode == MODE_LOD;
-		const uint HairGroupHeight = bDrawClusterInfo ? 120u : 60u;
+		const uint HairGroupHeight = 15; //bDrawClusterInfo ? 120u : 60u;
 
-		FShaderPrintContext Context = InitShaderPrintContext(true, uint2(50, 50 + HairGroupId * HairGroupHeight));
+		FShaderPrintContext Context = InitShaderPrintContext(true, uint2(50, 50 + (HairGroupId == 0 ? 0 : HairGroupHeight) + HairGroupId * HairGroupHeight));
+
+		FFontColor FontValue = FontWhite;
+		FFontColor FontPercent = FontSilver;
 
 		if (bDrawClusterInfo)
 		{
+			if (HairGroupId == 0)
+			{
+				Print(Context, TEXT("HairGroup   Clusters      Triangles"), FontYellow);
+			#if PERMUTATION_DEBUGAABBBUFFER
+				Print(Context, TEXT("         ScreenSize LOD"), FontYellow);
+			#endif
+				Newline(Context);
+			}
+
 			const uint VisibleClusterCount  = CulledDispatchIndirectParametersClusterCountBuffer[0];
 			const uint VisibleTriangleCount = CulledDrawIndirectParameters[0] / 3; // CulledDrawIndirectParameters contains the total number of triangle indices (3 per triangle)
 
-			Print(Context, TEXT("HairGroup "), FontYellow);
-			Print(Context, HairGroupId);
-			Print(Context, TEXT("       Count        Percent"), FontYellow);
-			Newline(Context);
-
-			Print(Context, TEXT("Visible Clusters  "), FontYellow);
-			Print(Context, VisibleClusterCount);
-			Print(Context, TEXT(" "), FontYellow);
-			PrintSymbol(Context, _SPC_);
-			Print(Context, 100.0f * float(VisibleClusterCount) / float(ClusterCount));
-			Newline(Context);
-
-			Print(Context, TEXT("Visible Triangles "), FontYellow);
-			Print(Context, VisibleTriangleCount);
-			PrintSymbol(Context, _SPC_);
-			Print(Context, 100.0f * float(VisibleTriangleCount) / float(TriangleCount));
-			Newline(Context);
+			Print(Context, HairGroupId, FontOrange);
+
+			Print(Context, VisibleClusterCount, FontValue, 5, 0);
+			Print(Context, TEXT(" ["), FontPercent);
+			Print(Context, 100.0f * float(VisibleClusterCount) / float(ClusterCount), FontPercent, 4, 1);
+			Print(Context, TEXT("%] "), FontPercent);
+
+			Print(Context, VisibleTriangleCount, FontValue, 9, 0);
+			Print(Context, TEXT(" ["), FontPercent);
+			Print(Context, 100.0f * float(VisibleTriangleCount) / float(TriangleCount), FontPercent, 4, 1);
+			Print(Context, TEXT("%] "), FontPercent);
+
 		}
 
 		#if PERMUTATION_DEBUGAABBBUFFER
 		{
-			// Extra new line to seperate from previous block
-			Newline(Context);
-
-			Print(Context, TEXT("Screen Size       "), FontYellow);
-
 			float ScreenSize = 1;
 			{
 				const FHairAABB Bound = ReadHairAABB(0, GroupAABBBuffer);
@@ -143,18 +145,11 @@ void MainDrawDebugAABBCS(uint3 DispatchThreadId : SV_DispatchThreadID)
 				const float  SphereRadius = length(Bound.Max - SphereOrigin);
 				ScreenSize = ComputeBoundsScreenSize(SphereOrigin, SphereRadius, LWCHackToFloat(PrimaryView.WorldCameraOrigin), View.ViewToClip);
 			}
-			Print(Context, ScreenSize);
-			Newline(Context);
+			Print(Context, ScreenSize, FontValue, 8, 2);
+			Print(Context, TEXT("   "));
 
 			// Show actual LOD distribution for all clusters. Otherwise show the average LOD index
 			const bool bShowLODHistogram = false;
-
-			if (bShowLODHistogram)
-			{
-				Print(Context, TEXT("LOD0        LOD1        LOD2        LOD3        LOD4        LOD5        LOD6        LOD7"), FontYellow);
-				Newline(Context);
-			}
-		
 			// Aggregate all the cluster results
 			uint LODHistogram[MAX_HAIR_LOD] = { 0,0,0,0,0,0,0,0 };
 			uint TotalVertexCount = 0;
@@ -187,24 +182,15 @@ void MainDrawDebugAABBCS(uint3 DispatchThreadId : SV_DispatchThreadID)
 			}
 			else
 			{
-				Print(Context, TEXT("LOD               "), FontYellow);
-
 				if (TotalVisibleCluster > 0)
 				{
-					Print(Context, AvgLOD);
+					Print(Context, AvgLOD, FontValue, 4, 2);
 				}
 				else
 				{
-					Print(Context, TEXT("Culled"), FontYellow);
+					Print(Context, TEXT("Culled"), FontLightRed);
 				}
 			}
-
-			Newline(Context);
-
-			// Extra stats for validation purpose
-			// Print(Context, TotalVertexCount);
-			// Print(Context, TotalVisibleCluster);
-			// Print(Context, ScreenSize);
 		}
 		#endif
 
