Removed unused function.

diff --git a/Engine/Shaders/Private/ShadingCommon.ush b/Engine/Shaders/Private/ShadingCommon.ush
--- a/Engine/Shaders/Private/ShadingCommon.ush
+++ b/Engine/Shaders/Private/ShadingCommon.ush
@@ -158,11 +158,6 @@ float MakeRoughnessSafe(float Roughness, float MinRoughness=0.001f)
 	return clamp(Roughness, MinRoughness, 1.0f);
 }
 
-float ComputeHazyLobeRoughness(float Roughness, float Haziness)
-{
-	return lerp(Roughness, 1.0f, Haziness);
-}
-
 float F0ToMetallic(float F0)
 {
 	// Approximate the metallic input from F0 with a small lerp region
