Fixed YCoCg variants for RVT ccompression on ETC2 platforms

diff --git a/Engine/Shaders/Private/ETCCompressionCommon.ush b/Engine/Shaders/Private/ETCCompressionCommon.ush
--- a/Engine/Shaders/Private/ETCCompressionCommon.ush
+++ b/Engine/Shaders/Private/ETCCompressionCommon.ush
@@ -293,8 +293,17 @@ uint4 CompressBlock_ETC2_SRGBA(in float3 BlockRGB[16], in float BlockA[16])
 
 uint4 CompressBlock_ETC2_YCoCg(in float3 BlockRGB[16])
 {
+	float BlockY[16];
+	for (int i = 0; i < 16; ++i)
+	{
+		float3 YCoCg = RGB2YCoCg(BlockRGB[i]);
+		BlockRGB[i] = float3(YCoCg.yz, 0.0);
+		BlockY[i] = YCoCg.x;
+	}
+		
 	uint2 CompressedRGB = CompressBlock_ETC2_SRGB(BlockRGB);
-	return uint4(uint2(0xffffffff, 0x0), CompressedRGB);
+	uint2 CompressedAlpha = CompressBlock_ETC2_Alpha(BlockY);
+	return uint4(CompressedAlpha, CompressedRGB);
 }
 
 uint4 CompressBlock_ETC2_RG(in float BlockU[16], in float BlockV[16])
