Strata - rough refraction now account for transmittance.

diff --git a/Engine/Shaders/Private/DistortApplyScreenPS.usf b/Engine/Shaders/Private/DistortApplyScreenPS.usf
--- a/Engine/Shaders/Private/DistortApplyScreenPS.usf
+++ b/Engine/Shaders/Private/DistortApplyScreenPS.usf
@@ -164,6 +164,20 @@ void Main(
 	OutColor = Texture2DSampleLevel(SceneColorTexture, SceneColorTextureSampler, NewBufferUV, RefractionMipLevel);
 #endif
 
+	// Reveal the original scene color unrefracted/sharp as a function of coverage.
+	const float Coverage = RefractionVarianceCoverage.y;
+	if (Coverage < 1.0f)
+	{
+		float4 OriginalSceneColor = 0.0f;
+	#if USE_MSAA
+		OriginalSceneColor = SceneColorMSAATexture.Load(int2(NoDistortionUVs * View.BufferSizeAndInvSize.xy), SampleIndex);
+	#else
+		OriginalSceneColor = Texture2DSample(SceneColorTexture, SceneColorTextureSampler, NoDistortionUVs);
+	#endif
+
+		OutColor = lerp(OriginalSceneColor, OutColor, Coverage);
+	}
+
 #else // USE_ROUGH_REFRACTION
 
 	float2 NewBufferUV = TexCoord.xy + DistBufferUVOffset;
