Fix Eye Adaptation debug output when ViewRect is not on top-left of texture.

diff --git a/Engine/Shaders/Private/Histogram.usf b/Engine/Shaders/Private/Histogram.usf
--- a/Engine/Shaders/Private/Histogram.usf
+++ b/Engine/Shaders/Private/Histogram.usf
@@ -302,14 +302,12 @@ void MainAtomicCS(
 	
 	GroupMemoryBarrierWithGroupSync();
 
-
 	//////////////////////////////////////////////////////
 	// * ACCUMULATE HISTOGRAM FOR ENTIRE LINE IN LDS *  (SCATTER)
-	uint2 Offset = Input_ViewportMin;
 	uint y = GroupId.x;
 	for (uint x = 0 ; x < Input_ViewportSize.x ; x+=THREADGROUP_SIZEX)
 	{
-		uint2 TexelPos = uint2(x + GroupThreadId.x + Offset.x, y + Offset.y);
+		uint2 TexelPos = uint2(x + GroupThreadId.x, y) + Input_ViewportMin;
 
 		if (TexelPos.x < Input_ViewportMax.x)
 		{
@@ -364,7 +362,7 @@ void MainAtomicCS(
 #endif
 
 #if USE_APPROX_ILLUMINANCE_DEBUG
-			DebugOutput[TexelPos] = LuminanceVal;
+			DebugOutput[TexelPos - Input_ViewportMin] = LuminanceVal;
 #endif
 
 			float Weight0,Weight1;
