Fix incorrect sampling of LocalExposure bilateral grid in BloomSetupCommon and PostProcessVisualizeLocalExposure.
- Bilateral grid must be sampled using viewport UV.

diff --git a/Engine/Shaders/Private/PostProcessBloom.usf b/Engine/Shaders/Private/PostProcessBloom.usf
--- a/Engine/Shaders/Private/PostProcessBloom.usf
+++ b/Engine/Shaders/Private/PostProcessBloom.usf
@@ -34,7 +34,7 @@ float2 GetExposureScaleMiddleGreyLumValue()
 	return ExposureScaleMiddleGreyLumValue;
 }
 
-float4 BloomSetupCommon(float2 UV, float2 ExposureScaleMiddleGreyLumValue)
+float4 BloomSetupCommon(float2 UV, float2 ViewportUV, float2 ExposureScaleMiddleGreyLumValue)
 {
 	float4 SceneColor = Texture2DSample(InputTexture, InputSampler, UV) * View.OneOverPreExposure;
 
@@ -46,7 +46,7 @@ float4 BloomSetupCommon(float2 UV, float2 ExposureScaleMiddleGreyLumValue)
 		float LogLuminance = log2(LuminanceVal);
 		float MiddleGreyLumValue = ExposureScaleMiddleGreyLumValue.y;
 
-		float BaseLogLum = CalculateBaseLogLuminance(LogLuminance, EyeAdaptation_LocalExposureBlurredLuminanceBlend, ExposureScale, UV, LumBilateralGrid, BlurredLogLum, LumBilateralGridSampler, BlurredLogLumSampler);
+		float BaseLogLum = CalculateBaseLogLuminance(LogLuminance, EyeAdaptation_LocalExposureBlurredLuminanceBlend, ExposureScale, ViewportUV, LumBilateralGrid, BlurredLogLum, LumBilateralGridSampler, BlurredLogLumSampler);
 		float LocalExposure = CalculateLocalExposure(LogLuminance + log2(ExposureScale), BaseLogLum, MiddleGreyLumValue, EyeAdaptation_LocalExposureHighlightContrastScale, EyeAdaptation_LocalExposureShadowContrastScale, EyeAdaptation_LocalExposureDetailStrength);
 		
 		SceneColor.rgb *= LocalExposure;
@@ -74,10 +74,11 @@ void BloomSetupPS(
 	out float4 OutColor : SV_Target0)
 {
 	float2 UV = ApplyScreenTransform(SvPosition.xy, SvPositionToInputTextureUV);
+	float2 ViewportUV = (int2(SvPosition.xy) - Input_ViewportMin) * Input_ViewportSizeInverse;
 	
 	float2 ExposureScaleMiddleGreyLumValue = GetExposureScaleMiddleGreyLumValue();
 
-	OutColor = BloomSetupCommon(UV, ExposureScaleMiddleGreyLumValue);
+	OutColor = BloomSetupCommon(UV, ViewportUV, ExposureScaleMiddleGreyLumValue);
 }
 
 #elif COMPUTESHADER
@@ -94,9 +95,11 @@ void BloomSetupCS(uint2 DispatchThreadId : SV_DispatchThreadID)
 		return;
 	}
 
+	float2 ViewportUV = DispatchThreadId * Input_ViewportSizeInverse;
+
 	float2 ExposureScaleMiddleGreyLumValue = GetExposureScaleMiddleGreyLumValue();
 
-	float4 OutColor = BloomSetupCommon(UV, ExposureScaleMiddleGreyLumValue);
+	float4 OutColor = BloomSetupCommon(UV, ViewportUV, ExposureScaleMiddleGreyLumValue);
 
 	uint2 PixelPos = DispatchThreadId + Input_ViewportMin;
 
diff --git a/Engine/Shaders/Private/PostProcessVisualizeLocalExposure.usf b/Engine/Shaders/Private/PostProcessVisualizeLocalExposure.usf
--- a/Engine/Shaders/Private/PostProcessVisualizeLocalExposure.usf
+++ b/Engine/Shaders/Private/PostProcessVisualizeLocalExposure.usf
@@ -24,7 +24,7 @@ SCREEN_PASS_TEXTURE_VIEWPORT(Output)
 float4 MainPS(noperspective float4 UVAndScreenPos : TEXCOORD0, float4 SvPosition : SV_POSITION) : SV_Target0
 {
 	float2 UV = UVAndScreenPos.xy;
-	int2 PixelPos = (int2)SvPosition.xy;
+	float2 ViewportUV = (int2(SvPosition.xy) - Output_ViewportMin) * Output_ViewportSizeInverse;
 
 	float2 ExposureScaleMiddleGrey = EyeAdaptationBuffer[0].xw;
 
@@ -34,7 +34,7 @@ float4 MainPS(noperspective float4 UVAndScreenPos : TEXCOORD0, float4 SvPosition
 	float LogLuminance = log2(LuminanceVal);
 	float MiddleGreyLumValue = log2(0.18 * ExposureScaleMiddleGrey.y * EyeAdaptation_LocalExposureMiddleGreyExposureCompensation);
 
-	float BaseLogLum = CalculateBaseLogLuminance(LogLuminance, EyeAdaptation_LocalExposureBlurredLuminanceBlend, ExposureScaleMiddleGrey.x, UV, LumBilateralGrid, BlurredLogLum, TextureSampler, TextureSampler);
+	float BaseLogLum = CalculateBaseLogLuminance(LogLuminance, EyeAdaptation_LocalExposureBlurredLuminanceBlend, ExposureScaleMiddleGrey.x, ViewportUV, LumBilateralGrid, BlurredLogLum, TextureSampler, TextureSampler);
 	float LocalExposure = CalculateLocalExposure(LogLuminance + log2(ExposureScaleMiddleGrey.x), BaseLogLum, MiddleGreyLumValue, EyeAdaptation_LocalExposureHighlightContrastScale, EyeAdaptation_LocalExposureShadowContrastScale, EyeAdaptation_LocalExposureDetailStrength);
 	float DetailLogLum = (LogLuminance + log2(ExposureScaleMiddleGrey.x)) - BaseLogLum;
 
