Supports thin translucent with dual src blending if it's supported on mobile platform with DXC.
It's disabled on Android because some android devices don't support dual src blending, and it's disabled on IOS because it doesn't work with MRT.

diff --git a/Engine/Config/Mac/DataDrivenPlatformInfo.ini b/Engine/Config/Mac/DataDrivenPlatformInfo.ini
--- a/Engine/Config/Mac/DataDrivenPlatformInfo.ini
+++ b/Engine/Config/Mac/DataDrivenPlatformInfo.ini
@@ -65,6 +65,7 @@ bSupportsIndexBufferUAVs=true
 bSupports4ComponentUAVReadWrite=true
 bSupportsManualVertexFetch=true
 bSupportsVolumeTextureAtomics=false
+bSupportsDualSourceBlending=true
 FriendlyName=LOCTEXT("FriendlyShaderPlatformName_MacSM5_MRT", "Mac Shader Model 5 MRT")
 
 [ShaderPlatform METAL_MACES3_1]
@@ -76,4 +77,4 @@ bIsMobile=true
 bSupports4ComponentUAVReadWrite=true
 bSupportsManualVertexFetch=false
 bSupportsVolumeTextureAtomics=false
-FriendlyName=LOCTEXT("FriendlyShaderPlatformName_MacES31", "Mac Shader Model 3.1")
\ No newline at end of file
+FriendlyName=LOCTEXT("FriendlyShaderPlatformName_MacES31", "Mac Shader Model 3.1")
diff --git a/Engine/Config/TVOS/DataDrivenPlatformInfo.ini b/Engine/Config/TVOS/DataDrivenPlatformInfo.ini
--- a/Engine/Config/TVOS/DataDrivenPlatformInfo.ini
+++ b/Engine/Config/TVOS/DataDrivenPlatformInfo.ini
@@ -25,6 +25,7 @@ bTargetsTiledGPU=true
 bSupportsIndexBufferUAVs=true
 bSupports4ComponentUAVReadWrite=true
 bSupportsManualVertexFetch=false
+bSupportsDualSourceBlending=true
 FriendlyName=LOCTEXT("FriendlyShaderPlatformName_TVOS_SM5", "TVOS Shader Model 5")
 
 [ShaderPlatform METAL_TVOS]
diff --git a/Engine/Config/VulkanPC/DataDrivenPlatformInfo.ini b/Engine/Config/VulkanPC/DataDrivenPlatformInfo.ini
--- a/Engine/Config/VulkanPC/DataDrivenPlatformInfo.ini
+++ b/Engine/Config/VulkanPC/DataDrivenPlatformInfo.ini
@@ -96,7 +96,7 @@ bNeedsOfflineCompiler = false
 bSupportsComputeFramework = false
 
 ; //#todo-vulkanddpi
-bSupportsDualSourceBlending = false
+bSupportsDualSourceBlending = true
 
 ; //#todo-vulkanddpi
 bRequiresGeneratePrevTransformBuffer = false
@@ -147,4 +147,4 @@ bSupportsMobileMultiView=true
 bSupportsIndexBufferUAVs=true
 bSupportsManualVertexFetch=false
 bSupportsSceneDataCompressedTransforms=true
-FriendlyName=LOCTEXT("FriendlyShaderPlatformName_VulkanPC_ES31", "Vulkan PC Shader Model 3_1")
\ No newline at end of file
+FriendlyName=LOCTEXT("FriendlyShaderPlatformName_VulkanPC_ES31", "Vulkan PC Shader Model 3_1")
diff --git a/Engine/Config/Windows/DataDrivenPlatformInfo.ini b/Engine/Config/Windows/DataDrivenPlatformInfo.ini
--- a/Engine/Config/Windows/DataDrivenPlatformInfo.ini
+++ b/Engine/Config/Windows/DataDrivenPlatformInfo.ini
@@ -76,6 +76,7 @@ bSupportsOIT=true
 BindlessSupport=RayTracingOnly
 SupportsMultiViewport=RuntimeDependent
 FriendlyName=LOCTEXT("FriendlyShaderPlatformName_D3D_SM5", "D3D Shader Model 5")
+bSupportsDualSourceBlending=true
 
 [PreviewPlatform PCD3D_SM5]
 PlatformName=PCD3D_SM5
@@ -133,6 +134,7 @@ bSupportsOIT=true
 bSupportsRealTypes=RuntimeDependent
 EnablesHLSL2021ByDefault=1
 SupportsMultiViewport=RuntimeGuaranteed
+bSupportsDualSourceBlending=true
 
 [PreviewPlatform PCD3D_SM6]
 PlatformName=PCD3D_SM6
@@ -171,4 +173,4 @@ MaxFeatureLevel=ES3_1
 ShaderFormat=GLSL_150_ES31
 bIsMobile=true
 bSupportsManualVertexFetch=false
-FriendlyName=LOCTEXT("FriendlyShaderPlatformName_PCGL_ES31", "Opengl PC Shader Model 3.1")
\ No newline at end of file
+FriendlyName=LOCTEXT("FriendlyShaderPlatformName_PCGL_ES31", "Opengl PC Shader Model 3.1")
diff --git a/Engine/Shaders/Private/MobileBasePassPixelShader.usf b/Engine/Shaders/Private/MobileBasePassPixelShader.usf
--- a/Engine/Shaders/Private/MobileBasePassPixelShader.usf
+++ b/Engine/Shaders/Private/MobileBasePassPixelShader.usf
@@ -107,6 +107,7 @@
 #include "LightmapCommon.ush"  
 #include "MobileLightingCommon.ush"
 #include "ShadingModelsMaterial.ush"
+#include "ThinTranslucentCommon.ush"
 
 #if MATERIAL_SHADINGMODEL_SINGLELAYERWATER
 	#include "SingleLayerWaterShading.ush"
@@ -274,7 +275,12 @@ void Main(
 	, out HALF4_TYPE OutGBufferD : SV_Target4
 	#endif
 #else
+	#if DUAL_SOURCE_COLOR_BLENDING_ENABLED && !THIN_TRANSLUCENT_FALL_BACK
+	, out HALF4_TYPE OutColor DUAL_SOURCE_BLENDING_SLOT(0) : SV_Target0
+	, out HALF4_TYPE OutColor1 DUAL_SOURCE_BLENDING_SLOT(1) : SV_Target1
+	#else
 	, out HALF4_TYPE OutColor : SV_Target0
+	#endif
 #endif
 #if HAS_SCENE_DEPTH_AUX_OUTPUT
 	, out float OutSceneDepthAux : SV_TargetDepthAux
@@ -361,6 +367,12 @@ void Main(
 	half Anisotropy = GetMaterialAnisotropy(PixelMaterialInputs);
 	uint ShadingModelID = GetMaterialShadingModel(PixelMaterialInputs);
 	
+	// Opacity for this model is the coverage of the top layer over the transmission surface. Since the
+	// underlying material isn't allowed to be metallic, we are scaling the metallic value by the coverage.
+#if MATERIAL_SHADINGMODEL_THIN_TRANSLUCENT
+	Metallic *= Opacity;
+#endif
+
 	// If we don't use this shading model the color should be black (don't generate shader code for unused data, don't do indirectlighting cache lighting with this color).
 	float3 SubsurfaceColor = 0;
 	// 0..1, SubsurfaceProfileId = int(x * 255)
@@ -515,38 +527,6 @@ void Main(
 	GBuffer.GBufferAO = AOMultiBounce(Luminance(GBuffer.SpecularColor), ShadingOcclusion.SpecOcclusion).g;
 #endif
 
-#endif // !STRATA_ENABLED || STRATA_INLINE_SINGLELAYERWATER
-
-#if !STRATA_ENABLED
-	
-#if MATERIAL_SHADINGMODEL_THIN_TRANSLUCENT
-	float TopMaterialCoverage = Opacity;
-
-	// Adjust Opacity
-	{
-		float3 TransmittanceColor = GetThinTranslucentMaterialOutput0(MaterialParameters);
-		float Transmittance = dot(TransmittanceColor, float3(1.0f,1.0f,1.0f)/3.0f);
-
-		// We can't use the NoV from the shading models because it uses saturate(), whereas we are using abs().
-		// The length through the surface is the same for both front and back faces.
-		float NoV = abs(dot(MaterialParameters.WorldNormal, MaterialParameters.CameraVector));
-		float PathLength = rcp(max(NoV, 1e-5f));
-		float NegativeAbsorptionCoefficient = log(Transmittance);
-		
-		// Material Modulation is how much of the background light goes through the surface
-		float MaterialModulation = exp(NegativeAbsorptionCoefficient * PathLength);
-
-		// The alpha of the material in translucent mode is one minus the amount that it is modulating the background by.
-		float MaterialOpacity = 1.0f - MaterialModulation;
-		
-		Opacity = 1.0f - (1.0f - TopMaterialCoverage) * (1.0f - MaterialOpacity);
-	}
-
-	GBuffer.BaseColor *= TopMaterialCoverage;
-#endif
-
-#endif // !STRATA_ENABLED
-
 	VTPageTableResult LightmapVTPageTableResult = (VTPageTableResult)0.0f;
 #if LIGHTMAP_VT_ENABLED
 	{
@@ -894,17 +874,52 @@ void Main(
 	// NEEDS_BASEPASS_PIXEL_FOGGING is not allowed on mobile for the sake of performance.
 				 
 	half3 Emissive = GetMaterialEmissive(PixelMaterialInputs);
-#if MATERIAL_SHADINGMODEL_THIN_TRANSLUCENT
-	Emissive *= TopMaterialCoverage;
-#endif
+
 	Color += Emissive;
 
 #if !MATERIAL_SHADINGMODEL_UNLIT && MOBILE_EMULATION
 	Color = lerp(Color, DiffuseColor, ResolvedView.UnlitViewmodeMask);
 #endif
 
+#if MATERIAL_SHADINGMODEL_THIN_TRANSLUCENT && !STRATA_ENABLED
+	half3 DualBlendSurfaceLuminancePostCoverage = 0.0f;
+	half3 DualBlendSurfaceTransmittancePreCoverage = 1.0f;
+	half DualBlendSurfaceCoverage = 1.0f;
+	{
+		AccumulateThinTranslucentModel(
+			DualBlendSurfaceLuminancePostCoverage,
+			DualBlendSurfaceTransmittancePreCoverage,
+			DualBlendSurfaceCoverage,
+			MaterialParameters,
+			GBuffer,
+			DirectLighting.TotalLightDiffuse,
+			DirectLighting.TotalLightSpecular,
+			Emissive,
+			Opacity);
+
+		Color = 0;
+		Opacity = 1.0f;
+	}
+#endif // MATERIAL_SHADINGMODEL_THIN_TRANSLUCENT
+	
+	// MATERIAL_SHADINGMODEL_THIN_TRANSLUCENT must come first because it also has MATERIALBLENDING_TRANSLUCENT defined
+	#if DUAL_SOURCE_COLOR_BLENDING_ENABLED
+
+		// Add fog luminance according to surfacecoverage and reduce surface luminance according to fog coverage.
+		half3 AdjustedDualBlendAdd = DualBlendSurfaceCoverage * VertexFog.rgb + VertexFog.a * DualBlendSurfaceLuminancePostCoverage;
+		// Fade the surface color transmittance out to 1 according to the surface coverage, and take into account the fog coverage to the surface.
+		half3 AdjustedDualBlendMul = lerp(1.0f, VertexFog.a * DualBlendSurfaceTransmittancePreCoverage, DualBlendSurfaceCoverage);
+
+		#if THIN_TRANSLUCENT_FALL_BACK
+			// In the fallback case, use standard alpha blending(grey scale transmittance)
+			half AdjustedAlpha = saturate(1 - dot(AdjustedDualBlendMul, half3(1.0f, 1.0f, 1.0f) / 3.0f));
+			OutColor = half4(AdjustedDualBlendAdd, AdjustedAlpha);
+		#else
+			OutColor = half4(AdjustedDualBlendAdd,0.0);
+			OutColor1 = half4(AdjustedDualBlendMul,1.0);
+		#endif
 	// On mobile, water (an opaque material) is rendered as trnaslucent with forced premultiplied alpha blending (see MobileBasePass::SetTranslucentRenderState)
-	#if MATERIALBLENDING_ALPHACOMPOSITE || MATERIAL_SHADINGMODEL_SINGLELAYERWATER
+	#elif MATERIALBLENDING_ALPHACOMPOSITE || MATERIAL_SHADINGMODEL_SINGLELAYERWATER
 		OutColor = half4(Color * VertexFog.a + VertexFog.rgb * Opacity, Opacity);
 	#elif MATERIALBLENDING_ALPHAHOLDOUT
 		// not implemented for holdout
diff --git a/Engine/Shaders/Private/ThinTranslucentCommon.ush b/Engine/Shaders/Private/ThinTranslucentCommon.ush
--- a/Engine/Shaders/Private/ThinTranslucentCommon.ush
+++ b/Engine/Shaders/Private/ThinTranslucentCommon.ush
@@ -6,7 +6,7 @@
 
 #define THIN_TRANSLUCENT_USE_DOTNV_THICKNESS 1
 
-#define THIN_TRANSLUCENT_ACCURATE_FRESNEL 1
+#define THIN_TRANSLUCENT_ACCURATE_FRESNEL (SHADING_PATH_DEFERRED || MOBILE_HIGH_QUALITY_BRDF)
 
 #if MATERIAL_SHADINGMODEL_THIN_TRANSLUCENT
 void AccumulateThinTranslucentModel(inout float3 DualBlendSurfaceLuminancePostCoverage,
diff --git a/Engine/Source/Runtime/RHI/Public/DataDrivenShaderPlatformInfo.h b/Engine/Source/Runtime/RHI/Public/DataDrivenShaderPlatformInfo.h
--- a/Engine/Source/Runtime/RHI/Public/DataDrivenShaderPlatformInfo.h
+++ b/Engine/Source/Runtime/RHI/Public/DataDrivenShaderPlatformInfo.h
@@ -942,9 +942,8 @@ inline bool RHISupportsShaderPipelines(const FStaticShaderPlatform Platform)
 
 inline bool RHISupportsDualSourceBlending(const FStaticShaderPlatform Platform)
 {
-	// For now only enable support for SM5
-	// Metal RHI doesn't support dual source blending properly at the moment.
-	return IsFeatureLevelSupported(Platform, ERHIFeatureLevel::SM5) && (IsD3DPlatform(Platform) || FDataDrivenShaderPlatformInfo::GetSupportsDualSourceBlending(Platform) || IsVulkanPlatform(Platform));
+	// Check if the platform supports dual src blending from DDPI
+	return FDataDrivenShaderPlatformInfo::GetSupportsDualSourceBlending(Platform) && !FDataDrivenShaderPlatformInfo::GetIsHlslcc(Platform);
 }
 
 // helper to check that the shader platform supports creating a UAV off an index buffer.
diff --git a/Engine/Source/Runtime/Renderer/Private/MobileBasePass.cpp b/Engine/Source/Runtime/Renderer/Private/MobileBasePass.cpp
--- a/Engine/Source/Runtime/Renderer/Private/MobileBasePass.cpp
+++ b/Engine/Source/Runtime/Renderer/Private/MobileBasePass.cpp
@@ -390,19 +390,37 @@ void MobileBasePass::SetOpaqueRenderState(FMeshPassProcessorRenderState& DrawRen
 
 void MobileBasePass::SetTranslucentRenderState(FMeshPassProcessorRenderState& DrawRenderState, const FMaterial& Material, FMaterialShadingModelField ShadingModels)
 {
-	const bool bIsUsingMobilePixelProjectedReflection = Material.IsUsingPlanarForwardReflections() && IsUsingMobilePixelProjectedReflection(GetFeatureLevelShaderPlatform(Material.GetFeatureLevel()));
+	EShaderPlatform ShaderPlatform = GetFeatureLevelShaderPlatform(Material.GetFeatureLevel());
+
+	const bool bIsUsingMobilePixelProjectedReflection = Material.IsUsingPlanarForwardReflections() && IsUsingMobilePixelProjectedReflection(ShaderPlatform);
 
 	if (ShadingModels.HasShadingModel(MSM_ThinTranslucent))
 	{
-		// the mobile thin translucent fallback uses a similar mode as BLEND_Translucent, but multiplies color by 1 insead of SrcAlpha.
-		DrawRenderState.SetBlendState(TStaticBlendState<CW_RGBA, BO_Add, BF_One, BF_InverseSourceAlpha, BO_Add, BF_Zero, BF_InverseSourceAlpha,
-														CW_NONE, BO_Add, BF_One, BF_Zero, BO_Add, BF_One, BF_Zero,
-														CW_NONE, BO_Add, BF_One, BF_Zero, BO_Add, BF_One, BF_Zero,
-														CW_NONE, BO_Add, BF_One, BF_Zero, BO_Add, BF_One, BF_Zero,
-														CW_NONE, BO_Add, BF_One, BF_Zero, BO_Add, BF_One, BF_Zero,
-														CW_NONE, BO_Add, BF_One, BF_Zero, BO_Add, BF_One, BF_Zero,
-														CW_NONE, BO_Add, BF_One, BF_Zero, BO_Add, BF_One, BF_Zero,
-														CW_NONE, BO_Add, BF_One, BF_Zero, BO_Add, BF_One, BF_Zero>::GetRHI());
+		const bool bIsDualSourceBlending = RHISupportsDualSourceBlending(ShaderPlatform);
+		if (bIsDualSourceBlending)
+		{
+			// Blend by putting add in target 0 and multiply by background in target 1.
+			DrawRenderState.SetBlendState(TStaticBlendState<CW_RGBA, BO_Add, BF_One, BF_Source1Color, BO_Add, BF_One, BF_Source1Alpha,
+															CW_NONE, BO_Add, BF_One, BF_Zero, BO_Add, BF_One, BF_Zero,
+															CW_NONE, BO_Add, BF_One, BF_Zero, BO_Add, BF_One, BF_Zero,
+															CW_NONE, BO_Add, BF_One, BF_Zero, BO_Add, BF_One, BF_Zero,
+															CW_NONE, BO_Add, BF_One, BF_Zero, BO_Add, BF_One, BF_Zero,
+															CW_NONE, BO_Add, BF_One, BF_Zero, BO_Add, BF_One, BF_Zero,
+															CW_NONE, BO_Add, BF_One, BF_Zero, BO_Add, BF_One, BF_Zero,
+															CW_NONE, BO_Add, BF_One, BF_Zero, BO_Add, BF_One, BF_Zero>::GetRHI());
+		}
+		else
+		{
+			// the mobile thin translucent fallback uses a similar mode as BLEND_Translucent, but multiplies color by 1 insead of SrcAlpha.
+			DrawRenderState.SetBlendState(TStaticBlendState<CW_RGBA, BO_Add, BF_One, BF_InverseSourceAlpha, BO_Add, BF_Zero, BF_InverseSourceAlpha,
+															CW_NONE, BO_Add, BF_One, BF_Zero, BO_Add, BF_One, BF_Zero,
+															CW_NONE, BO_Add, BF_One, BF_Zero, BO_Add, BF_One, BF_Zero,
+															CW_NONE, BO_Add, BF_One, BF_Zero, BO_Add, BF_One, BF_Zero,
+															CW_NONE, BO_Add, BF_One, BF_Zero, BO_Add, BF_One, BF_Zero,
+															CW_NONE, BO_Add, BF_One, BF_Zero, BO_Add, BF_One, BF_Zero,
+															CW_NONE, BO_Add, BF_One, BF_Zero, BO_Add, BF_One, BF_Zero,
+															CW_NONE, BO_Add, BF_One, BF_Zero, BO_Add, BF_One, BF_Zero>::GetRHI());
+		}
 	}
 	else
 	{
diff --git a/Engine/Source/Runtime/Renderer/Private/MobileBasePassRendering.h b/Engine/Source/Runtime/Renderer/Private/MobileBasePassRendering.h
--- a/Engine/Source/Runtime/Renderer/Private/MobileBasePassRendering.h
+++ b/Engine/Source/Runtime/Renderer/Private/MobileBasePassRendering.h
@@ -394,6 +394,11 @@ public:
 		OutEnvironment.SetDefine(TEXT("ENABLE_CLUSTERED_LIGHTS"), bEnableLocalLights ? 1u : 0u);
 		OutEnvironment.SetDefine(TEXT("ENABLE_CLUSTERED_REFLECTION"), bEnableClusteredReflections ? 1u : 0u);
 		OutEnvironment.SetDefine(TEXT("USE_SHADOWMASKTEXTURE"), bMobileUsesShadowMaskTexture && !bTranslucentMaterial ? 1u : 0u);
+
+		if (!RHISupportsDualSourceBlending(Parameters.Platform))
+		{
+			OutEnvironment.SetDefine(TEXT("THIN_TRANSLUCENT_FALL_BACK"), TEXT("1"));
+		}
 	}
 	
 	/** Initialization constructor. */
