MobileBasePassPixelShader - fix missing endif from bad automerge.
#rnx

diff --git a/Engine/Shaders/Private/MobileBasePassPixelShader.usf b/Engine/Shaders/Private/MobileBasePassPixelShader.usf
--- a/Engine/Shaders/Private/MobileBasePassPixelShader.usf
+++ b/Engine/Shaders/Private/MobileBasePassPixelShader.usf
@@ -527,6 +527,8 @@ void Main(
 	GBuffer.GBufferAO = AOMultiBounce(Luminance(GBuffer.SpecularColor), ShadingOcclusion.SpecOcclusion).g;
 #endif
 
+#endif // !STRATA_ENABLED || STRATA_INLINE_SINGLELAYERWATER
+
 	VTPageTableResult LightmapVTPageTableResult = (VTPageTableResult)0.0f;
 #if LIGHTMAP_VT_ENABLED
 	{
