Sky light illuminance meter now account for sky color.
The sun light direction is now affected by the transmittance of the sun light to the skylight component position.
Fixed view issues.

diff --git a/Engine/Shaders/Private/ReflectionEnvironmentShaders.usf b/Engine/Shaders/Private/ReflectionEnvironmentShaders.usf
--- a/Engine/Shaders/Private/ReflectionEnvironmentShaders.usf
+++ b/Engine/Shaders/Private/ReflectionEnvironmentShaders.usf
@@ -218,6 +218,8 @@ void CopyCubemapToIlluminanceMeterCubemapPS(
 	out float4 OutColor : SV_Target0
 	)
 {
+	ResolvedView = ResolveView();
+
 	float2 UV = SvPosition.xy * SvPositionToUVScale;
 	float2 ScaledUVs = UV * 2 - 1;
 
@@ -230,7 +232,7 @@ void CopyCubemapToIlluminanceMeterCubemapPS(
 	float Tmp= 1.0f + ScaledUVs.x * ScaledUVs.x + ScaledUVs.y * ScaledUVs.y; 
 	float dA = 4.0f / (sqrt(Tmp) * Tmp); // Not really dA, but using same variable name
 
-	const float3 L = TextureCubeSampleLevel(SourceCubemapTexture, SourceCubemapSampler, CubeCoordinates, 0).rgb;
+	const float3 L = TextureCubeSampleLevel(SourceCubemapTexture, SourceCubemapSampler, CubeCoordinates, 0).rgb * ResolvedView.SkyLightColor.rgb;
 	OutColor = float4(L * NoL* dA, dA);
 }
 
@@ -885,8 +887,10 @@ void RenderRealTimeReflectionHeightFogPS(
 #ifdef VISUALIZE_ILLUMINANCE_METER
 
 #include "/Engine/Private/ShaderPrint.ush"
+#include "/Engine/Private/SkyAtmosphereCommon.ush"
 
 int CubeMipIndexToSampleIlluminance;
+float3 SkyLightCaptureWorlPos;
 
 TextureCube SkyLightCubeTexture;
 SamplerState SkyLightCubeSampler;
@@ -912,7 +916,7 @@ float GetSimpleAutoExposure()
 	float4 IlluminanceAcc = GetIlluminance();
 	float MaxIlluminance = max3(IlluminanceAcc.r, IlluminanceAcc.g, IlluminanceAcc.b);
 	float SimpleAutoExposure = 1.0f / (MaxIlluminance <= 0.0 ? 0.01f : MaxIlluminance);
-	return SimpleAutoExposure;
+	return 1.0;// SimpleAutoExposure; Disabled for now, using the view exposure
 }
 
 float4 SampleCubeAsLatLon(TextureCube CubeTexture, SamplerState CubeSampler, float2 UV)
@@ -943,7 +947,7 @@ float4 SampleCubeAsSphere(TextureCube CubeTexture, SamplerState CubeSampler, flo
 			N *= float3(1.0, -1.0, 1.0);
 		}
 
-		DebugVis = TextureCubeSampleLevel(CubeTexture, CubeSampler, mul(N, (float3x3)View.ViewToTranslatedWorld), 0);
+		DebugVis = TextureCubeSampleLevel(CubeTexture, CubeSampler, mul(N, (float3x3)ResolvedView.ViewToTranslatedWorld), 0);
 		DebugVis.rgb *= GetSimpleAutoExposure();
 	}
 
@@ -954,7 +958,9 @@ void VisualizeIlluminanceMeterPS(
 	float4 SVPos : SV_POSITION,
 	out float4 OutColor : SV_Target0)
 {
-	const uint2 PixelPosDynRes = uint2(float2(SVPos.xy) * View.ViewResolutionFraction);
+	ResolvedView = ResolveView();
+
+	const uint2 PixelPosDynRes = uint2(float2(SVPos.xy) * ResolvedView.ViewResolutionFraction);
 	float2 BufferUV = SvPositionToBufferUV(float4(PixelPosDynRes, SVPos.zw));
 
 	OutColor = 0.0f;
@@ -983,7 +989,17 @@ void VisualizeIlluminanceMeterPS(
 
 		Newline(Context);
 
-		const float3 SunLightIlluminance = View.AtmosphereLightIlluminanceOnGroundPostTransmittance[0];
+	#if PROJECT_SUPPORT_SKY_ATMOSPHERE
+		const float3 PlanetCenterToTranslatedWorldPos = (LWCToFloat(LWCAdd(SkyLightCaptureWorlPos, ResolvedView.PreViewTranslation)) - ResolvedView.SkyPlanetTranslatedWorldCenterAndViewHeight.xyz) * CM_TO_SKY_UNIT;
+		// GetAtmosphereTransmittance does a shadow test against the virtual planet.
+		const float3 TransmittanceToLight = GetAtmosphereTransmittance(
+			PlanetCenterToTranslatedWorldPos, ResolvedView.AtmosphereLightDirection[0].xyz, ResolvedView.SkyAtmosphereBottomRadiusKm, ResolvedView.SkyAtmosphereTopRadiusKm,
+			View.TransmittanceLutTexture, View.TransmittanceLutTextureSampler);
+
+		const float3 SunLightIlluminance = ResolvedView.AtmosphereLightIlluminanceOuterSpace[0].rgb * TransmittanceToLight;
+	#else
+		const float3 SunLightIlluminance = ResolvedView.AtmosphereLightIlluminanceOuterSpace[0];
+	#endif
 
 		Print(Context, TEXT("Sun light illuminance = "), InitFontColor(0.8, 0.8, 0.8));
 		Print(Context, SunLightIlluminance.r, InitFontColor(1.0, 0.8, 0.8));
diff --git a/Engine/Source/Runtime/Engine/Classes/Components/SkyLightComponent.h b/Engine/Source/Runtime/Engine/Classes/Components/SkyLightComponent.h
--- a/Engine/Source/Runtime/Engine/Classes/Components/SkyLightComponent.h
+++ b/Engine/Source/Runtime/Engine/Classes/Components/SkyLightComponent.h
@@ -204,7 +204,8 @@ class ENGINE_API USkyLightComponent : public ULightComponentBase
 	TEnumAsByte<enum EOcclusionCombineMode> OcclusionCombineMode;
 	
 	/**
-	 * When enabled, draw on screen the bull cube map as well as the integrate illuminance half hemisphere, together with the sky illuminance contribution. It also show the atmospheric sun illuminance post SkyAtmosphere transmittance on ground if one is enabled.
+	 * When enabled, visualize on screen the sky cube map as well as the integrate half hemisphere of luminance, together with the sky illuminance contribution at the position of the skylight component.
+	 * It also show the atmospheric directional light illuminance post SkyAtmosphere transmittance at the position on the skylight component.
 	 */
 	UPROPERTY(EditAnywhere, BlueprintReadOnly, Category = Visualization)
 	uint32 bShowIlluminanceMeter : 1;
diff --git a/Engine/Source/Runtime/Renderer/Private/IlluminanceMeter.cpp b/Engine/Source/Runtime/Renderer/Private/IlluminanceMeter.cpp
--- a/Engine/Source/Runtime/Renderer/Private/IlluminanceMeter.cpp
+++ b/Engine/Source/Runtime/Renderer/Private/IlluminanceMeter.cpp
@@ -21,6 +21,7 @@ public:
 	SHADER_USE_PARAMETER_STRUCT(FCopyCubemapToIlluminanceMeterCubemapPS, FGlobalShader);
 
 	BEGIN_SHADER_PARAMETER_STRUCT(FParameters, )
+		SHADER_PARAMETER_STRUCT_REF(FViewUniformShaderParameters, ViewUniformBuffer)
 		SHADER_PARAMETER_TEXTURE(TextureCube, SourceCubemapTexture)
 		SHADER_PARAMETER_SAMPLER(SamplerState, SourceCubemapSampler)
 		SHADER_PARAMETER(FVector2f, SvPositionToUVScale)
@@ -67,6 +68,7 @@ class FPrintIlluminanceMeterPS : public FGlobalShader
 		SHADER_PARAMETER_TEXTURE(TextureCube, SkyLightCubeTexture)
 		SHADER_PARAMETER_SAMPLER(SamplerState, SkyLightCubeSampler)
 		SHADER_PARAMETER(int32, CubeMipIndexToSampleIlluminance)
+		SHADER_PARAMETER(FVector3f, SkyLightCaptureWorlPos)
 		SHADER_PARAMETER_STRUCT_INCLUDE(ShaderPrint::FShaderParameters, ShaderPrintParameters)
 		RENDER_TARGET_BINDING_SLOTS()
 	END_SHADER_PARAMETER_STRUCT()
@@ -135,6 +137,7 @@ void FScene::ProcessAndRenderIlluminanceMeter(FRDGBuilder& GraphBuilder, TArrayV
 		for (uint32 CubeFace = 0; CubeFace < CubeFace_MAX; CubeFace++)
 		{
 			auto* PassParameters = GraphBuilder.AllocParameters<FCopyCubemapToIlluminanceMeterCubemapPS::FParameters>();
+			PassParameters->ViewUniformBuffer = MainView.ViewUniformBuffer;
 			PassParameters->SourceCubemapSampler = TStaticSamplerState<SF_Point, AM_Clamp, AM_Clamp, AM_Clamp>::GetRHI();
 			PassParameters->SourceCubemapTexture = SkyLightTextureRHI;
 			PassParameters->CubeFace = CubeFace;
@@ -202,6 +205,7 @@ void FScene::ProcessAndRenderIlluminanceMeter(FRDGBuilder& GraphBuilder, TArrayV
 		PassParameters->SkyLightCubeTexture = SkyLightTextureRHI;
 		PassParameters->SkyLightCubeSampler = TStaticSamplerState<SF_Trilinear, AM_Clamp, AM_Clamp, AM_Clamp>::GetRHI();
 		PassParameters->CubeMipIndexToSampleIlluminance = IlluminanceMeterCubeTexture->Desc.NumMips - 1;
+		PassParameters->SkyLightCaptureWorlPos = FVector3f(SkyLight->CapturePosition);
 		ShaderPrint::SetParameters(GraphBuilder, MainView.ShaderPrintData, PassParameters->ShaderPrintParameters);
 		PassParameters->RenderTargets[0] = FRenderTargetBinding(SceneColorTexture, ERenderTargetLoadAction::ELoad);
 
