Path Tracer: Decouple refraction IOR from specularity for solid and thin glass models

In the thin case, Ior now affects the transmitted roughness, so that rough refraction is progressively disabled as Ior is reduced toward 1.0. When not using refraction, ThinTranslucent materials are rendered using transparency, while with refraction they are rendered with an actual bounce (to allow for roughness, and for consistency with the solid case). SpecularColor controls the fresnel effect.

In the solid case, Ior now only affects ray bending, while SpecularColor controls the fresnel effect. Refraction must be enabled in the material to get a refractive event. Disabling refraction means the surface is simply transparent.

diff --git a/Engine/Shaders/Private/PathTracing/Material/PathTracingFresnel.ush b/Engine/Shaders/Private/PathTracing/Material/PathTracingFresnel.ush
--- a/Engine/Shaders/Private/PathTracing/Material/PathTracingFresnel.ush
+++ b/Engine/Shaders/Private/PathTracing/Material/PathTracingFresnel.ush
@@ -5,6 +5,10 @@ PathTracingFresnel.usf: Fresnel utilities for material sampling functions
 ===============================================================================================*/
 #pragma once
 
+// The functions in this file obey the following convention:
+// Eta = bFrontFacing ? Ior : rcp(Ior);
+// In other words, when we "enter" the material, we use the IOR (a value greater than 1.0), and when we exit, we use the inverse (a value less than 1.0)
+
 float FresnelReflectance(float CosI, float Eta)
 {
 	// This is a float only version of F_Fresnel that takes the Ior ratio as input directly (and handles TIR)
@@ -21,7 +25,27 @@ float FresnelReflectance(float CosI, float Eta)
 		float b2 = Pow2((c * (g + c) - 1) / (c * (g - c) + 1));
 		return 0.5 * a2 * (1.0 + b2);
 	}
+	// TIR -- must become fully reflective because refraction direction is not defined
+	return 1.0;
+}
+
+// TIR handled by Eta parameter, reflectivity by F0
+float FresnelReflectance(float CosI, float Eta, float F0)
+{
+	// Eta = bFrontFacing ? Ior : rcp(Ior);
+	float F90 = saturate(F0 * 50.0);
 
+	if (Eta >= 1.0)
+	{
+		// no TIR possible when "entering" (lower ior -> higher ior)
+		return lerp(F0, F90, Pow5(1 - CosI));
+	}
+	float g2 = Eta * Eta - 1 + CosI * CosI;
+	if (g2 >= 0.0)
+	{
+		return lerp(F0, F90, Pow5(1.0 - sqrt(g2) * rcp(Eta)));
+	}
+	// TIR -- must become fully reflective because refraction direction is not defined
 	return 1.0;
 }
 
@@ -60,3 +84,36 @@ bool SampleRefraction(float3 RayDirection, float3 Normal, float Eta, float RandS
 	OutDirection = RayDirection + (2 * CosI) * Normal;
 	return false;
 }
+
+// Sample refraction where bending is driven by Eta and reflectivity by F0
+bool SampleRefraction(float3 RayDirection, float3 Normal, float Eta, float F0, float RandSample, out float3 OutDirection, out float OutFresnel)
+{
+	float CosI = abs(dot(RayDirection, Normal));
+	float g2 = Eta * Eta - 1.0 + CosI * CosI;
+	if (g2 >= 0.0)
+	{
+		// refraction is possible, compute fresnel
+		float g = sqrt(g2);
+		float c = Eta >= 1.0 ? CosI : g * rcp(Eta);
+		float F90 = saturate(F0 * 50.0);
+		float Fr = lerp(F0, F90, Pow5(1 - c));
+		// Stochastically decide to transmit or not based on Fr
+		// This assumes RandSample is in [0,1)
+		if (RandSample >= Fr)
+		{
+			OutFresnel = 1 - Fr;
+			OutDirection = normalize(RayDirection + Normal * (CosI - g));
+			return true;
+		}
+		// sampled reflection
+		OutFresnel = Fr;
+	}
+	else
+	{
+		// TIR (no refraction possible, so everything must be reflected
+		OutFresnel = 1.0;
+	}
+	// either TIR, or sampled reflection
+	OutDirection = RayDirection + (2 * CosI) * Normal;
+	return false;
+}
diff --git a/Engine/Shaders/Private/PathTracing/Material/PathTracingGlossy.ush b/Engine/Shaders/Private/PathTracing/Material/PathTracingGlossy.ush
--- a/Engine/Shaders/Private/PathTracing/Material/PathTracingGlossy.ush
+++ b/Engine/Shaders/Private/PathTracing/Material/PathTracingGlossy.ush
@@ -13,6 +13,27 @@ Glossy.usf: Microfacet GGX BRDF sampling functions
 #define GGX_MIN_ROUGHNESS 0.01
 #define GGX_MIN_ANISOTROPY 0.001 // ignore anisotropy below this amount
 
+float2 GetGGXAlpha(float Roughness)
+{
+	Roughness = max(Roughness, GGX_MIN_ROUGHNESS);
+	return Roughness * Roughness;
+}
+
+float2 GetGGXAlpha(float Roughness, float Anisotropy)
+{
+	Roughness = max(Roughness, GGX_MIN_ROUGHNESS);
+	float2 Alpha = 0.0;
+	if (abs(Anisotropy) >= GGX_MIN_ANISOTROPY)
+	{
+		GetAnisotropicRoughness(Roughness * Roughness, Anisotropy, Alpha.x, Alpha.y);
+	}
+	else
+	{
+		Alpha = Roughness * Roughness;
+	}
+	return Alpha;
+}
+
 float3x3 GetGGXBasis(float Roughness, float Anisotropy, float3 WorldNormal, float3 WorldTangent, out float2 Alpha)
 {
 	Roughness = max(Roughness, GGX_MIN_ROUGHNESS);
diff --git a/Engine/Shaders/Private/PathTracing/Material/PathTracingMaterialSampling.ush b/Engine/Shaders/Private/PathTracing/Material/PathTracingMaterialSampling.ush
--- a/Engine/Shaders/Private/PathTracing/Material/PathTracingMaterialSampling.ush
+++ b/Engine/Shaders/Private/PathTracing/Material/PathTracingMaterialSampling.ush
@@ -12,7 +12,6 @@ PathTracingMaterialSampling.usf: Material sampling functions
 #include "PathTracingRadianceProbe.ush"
 #include "PathTracingGlossy.ush"
 #include "PathTracingSpecularReflection.ush"
-#include "PathTracingSpecularTransmission.ush"
 #include "PathTracingSolidGlass.ush"
 #include "PathTracingThinGlass.ush"
 #include "PathTracingDefaultLit.ush"
@@ -50,31 +49,32 @@ bool IsNeeValidMaterial(FPathTracingPayload Payload)
 
 float3 GetMaterialTransparency(FPathTracingPayload Payload, float3 RayDirection)
 {
-	switch (Payload.ShadingModelID)
+	// ThinTranslucent is special, all other shading models should go through blend mode logic
+	if (Payload.ShadingModelID == SHADINGMODELID_THIN_TRANSLUCENT)
 	{
-		case SHADINGMODELID_THIN_TRANSLUCENT:
+#if SIMPLIFIED_MATERIAL_SHADER
+		const float VoN = dot(RayDirection, Payload.WorldNormal);
+		return (1 - Payload.Opacity) * pow(Payload.CustomData.xyz, 1.0 / abs(VoN));
+#else
+		if (!Payload.HasRefraction() && Payload.Opacity < 1.0)
 		{
-#if !SIMPLIFIED_MATERIAL_SHADER
-			if (Payload.Ior != 1.0)
-			{
-				return 0;
-			}
-			else
-#endif
-			{
-				float VoN = dot(RayDirection, Payload.WorldNormal);
-				return (1 - Payload.Opacity) * pow(Payload.CustomData.xyz, 1.0 / abs(VoN));
-			}
+			// ThinTranslucent has an optimized mode where it ignores refractivity and just handles the transmission via transparency
+			const float VoN = abs(dot(RayDirection, Payload.WorldNormal));
+			const float F0 = F0RGBToF0(Payload.SpecularColor);
+			return (1 - Payload.Opacity) * ComputeThinSlabWeights(Payload.CustomData.xyz, VoN, 0.0, F0).Transmitted;
 		}
-		default:
-			// other shading models support blend modes
-			break;
+		else
+		{
+			// opaque, or handled via the transmission bsdf lobe
+			return 0;
+		}
+#endif
 	}
 	switch (Payload.BlendingMode)
 	{
 		case RAY_TRACING_BLEND_MODE_TRANSLUCENT:
 #if !SIMPLIFIED_MATERIAL_SHADER
-			if (Payload.Ior != 1.0)
+			if (Payload.HasRefraction())
 			{
 				// not using transparency - material will be opaque
 				return 0;
@@ -105,10 +105,15 @@ float3 EstimateMaterialAlbedo(FPathTracingPayload Payload)
 	{
 		case RAY_TRACING_BLEND_MODE_TRANSLUCENT:
 #if !SIMPLIFIED_MATERIAL_SHADER
-			if (Payload.Ior != 1.0)
+			if (Payload.HasRefraction())
 			{
 				HasRefraction = true;
 			}
+			else if (Payload.ShadingModelID == SHADINGMODELID_THIN_TRANSLUCENT)
+			{
+				// special case -- thin translucent can use transparency without refraction but should still contain a specular highlight at Opacity=0
+				HasRefraction = true;
+			}
 			else
 #endif
 			{
@@ -150,6 +155,7 @@ float3 EstimateMaterialAlbedo(FPathTracingPayload Payload)
 		case SHADINGMODELID_SINGLELAYERWATER:
 		case SHADINGMODELID_DEFAULT_LIT:
 		case SHADINGMODELID_CLEAR_COAT:
+		case SHADINGMODELID_THIN_TRANSLUCENT:
 			// TODO: include clearcoat albedo as well
 			Albedo = BrdfWeight * (Payload.DiffuseColor + Payload.SpecularColor);
 			break;
@@ -166,10 +172,16 @@ float3 EstimateMaterialAlbedo(FPathTracingPayload Payload)
 #if !SIMPLIFIED_MATERIAL_SHADER
 	if (HasRefraction)
 	{
-		// glass is currently perfectly clear - so has unit albedo
-		// TODO: estimate some kind of albedo for thin-translucent case since that one uses a color?
-		float3 White = 1;
-		Albedo = lerp(White, Albedo, Payload.Opacity);
+		// glass has unit albedo, except in the thin case where the transmission portion can be tinted
+		float3 GlassAlbedo = 1.0;
+		if (Payload.ShadingModelID == SHADINGMODELID_THIN_TRANSLUCENT)
+		{
+			// If transparency is taking care of the tint, just account for the specular portion
+			// Otherwise, we are using (rough) refraction here, so account for both specular + transmission
+			const float F0 = F0RGBToF0(Payload.SpecularColor);
+			GlassAlbedo = Payload.HasRefraction() ? lerp(Payload.CustomData.xyz, 1.0, F0) : F0;
+		}
+		Albedo = lerp(GlassAlbedo, Albedo, Payload.Opacity);
 	}
 #endif
 	return Albedo;
@@ -195,10 +207,16 @@ void AccumulateAlbedoNormal(inout FPathTracingPayload Payload, float3 PathThroug
 	{
 		case RAY_TRACING_BLEND_MODE_TRANSLUCENT:
 #if !SIMPLIFIED_MATERIAL_SHADER
-			if (Payload.Ior != 1.0)
+			if (Payload.HasRefraction())
 			{
 				HasRefraction = true;
 			}
+			else if (Payload.ShadingModelID == SHADINGMODELID_THIN_TRANSLUCENT)
+			{
+				// special case -- thin translucent can use transparency without refraction but should still contain a specular highlight at Opacity=0
+				HasRefraction = true;
+			}
+			else
 #endif
 			{
 				BrdfWeight = Payload.Opacity;
@@ -232,31 +250,38 @@ void AccumulateAlbedoNormal(inout FPathTracingPayload Payload, float3 PathThroug
 	case SHADINGMODELID_SUBSURFACE_PROFILE:
 	case SHADINGMODELID_EYE:
 		Albedo += BrdfWeight * PathThroughput * (Payload.DiffuseColor + RoughnessWeight * Payload.SpecularColor + Payload.GetSubsurfaceColor());
-		Normal += Luminance(BrdfWeight * PathThroughput) * Payload.WorldNormal;
+		Normal += BrdfWeight * Luminance(PathThroughput) * Payload.WorldNormal;
 		break;
 	case SHADINGMODELID_SINGLELAYERWATER:
 	case SHADINGMODELID_DEFAULT_LIT:
 	case SHADINGMODELID_CLEAR_COAT:
 	case SHADINGMODELID_THIN_TRANSLUCENT:
 		Albedo += BrdfWeight * PathThroughput * (Payload.DiffuseColor + RoughnessWeight * Payload.SpecularColor);
-		Normal += Luminance(BrdfWeight * PathThroughput) * Payload.WorldNormal;
+		Normal += BrdfWeight * Luminance(PathThroughput) * Payload.WorldNormal;
 		break;
 	case SHADINGMODELID_CLOTH:
 		Albedo += BrdfWeight * PathThroughput * (Payload.DiffuseColor + lerp(RoughnessWeight * Payload.SpecularColor, Payload.CustomData.xyz, Payload.CustomData.w));
-		Normal += Luminance(BrdfWeight * PathThroughput) * Payload.WorldNormal;
+		Normal += BrdfWeight * Luminance(PathThroughput) * Payload.WorldNormal;
 		break;
 #endif
 	default:
 		Albedo += BrdfWeight * PathThroughput * Payload.BaseColor;
-		Normal += Luminance(BrdfWeight * PathThroughput) * Payload.WorldNormal;
+		Normal += BrdfWeight * Luminance(PathThroughput) * Payload.WorldNormal;
 		break;
 	}
 #if !SIMPLIFIED_MATERIAL_SHADER
 	if (HasRefraction)
 	{
-		// glass is currently perfectly clear - so has unit albedo
-		// TODO: estimate some kind of albedo for thin-translucent case since that one uses a color?
-		Albedo += (1 - BrdfWeight) * RoughnessWeight * PathThroughput;
+		// glass has unit albedo, except in the thin case where the transmission portion can be tinted
+		float3 GlassAlbedo = 1.0;
+		if (Payload.ShadingModelID == SHADINGMODELID_THIN_TRANSLUCENT)
+		{
+			// If transparency is taking care of the tint, just account for the specular portion
+			// Otherwise, we are using (rough) refraction here, so account for both specular + transmission
+			const float F0 = F0RGBToF0(Payload.SpecularColor);
+			GlassAlbedo = Payload.HasRefraction() ? lerp(Payload.CustomData.xyz, 1.0, F0) : F0;
+		}
+		Albedo += (1 - BrdfWeight) * RoughnessWeight * PathThroughput * GlassAlbedo;
 		Normal += (1 - BrdfWeight) * RoughnessWeight * Luminance(PathThroughput) * Payload.WorldNormal;
 	}
 #endif
@@ -327,7 +352,7 @@ float GetAverageRoughness(FPathTracingPayload Payload)
 #if SIMPLIFIED_MATERIAL_SHADER
 	return 1.0;
 #else
-	const float GlassPercent = (Payload.BlendingMode == RAY_TRACING_BLEND_MODE_TRANSLUCENT && Payload.Ior != 1.0) ? 1 - Payload.Opacity : 0.0;
+	const float GlassPercent = (Payload.BlendingMode == RAY_TRACING_BLEND_MODE_TRANSLUCENT && Payload.HasRefraction()) ? 1 - Payload.Opacity : 0.0;
 	float Roughness = Payload.Roughness;
 	switch (Payload.ShadingModelID)
 	{
@@ -502,7 +527,7 @@ FMaterialSample SampleMaterial(
 	switch (Payload.BlendingMode)
 	{
 		case RAY_TRACING_BLEND_MODE_TRANSLUCENT:
-			if (Payload.Ior == 1.0)
+			if (!Payload.HasRefraction() && Payload.ShadingModelID != SHADINGMODELID_THIN_TRANSLUCENT)
 			{
 				// scale down BRDF contribution, but don't pick any indirect ray because its already accounted for by transparency handling
 				BrdfWeight = Payload.Opacity;
@@ -510,8 +535,8 @@ FMaterialSample SampleMaterial(
 			else
 			{
 				// randomly choose between brdf and glass
-				BrdfWeight = 1;
-				TransWeight = 1;
+				BrdfWeight = Payload.Opacity;
+				TransWeight = 1 - Payload.Opacity;
 				BrdfProb = Payload.Opacity;
 			}
 			break;
@@ -539,10 +564,6 @@ FMaterialSample SampleMaterial(
 		{
 			Result = RoughThinGlass_SampleMaterial(RayDirection, Payload, RandSample);
 		}
-		else if (Payload.Roughness == 0)
-		{
-			Result = SpecularRefraction_SampleMaterial(RayDirection, Payload, RandSample);
-		}
 		else
 		{
 			Result = RoughGlass_SampleMaterial(RayDirection, Payload, RandSample);
@@ -632,7 +653,7 @@ FMaterialEval EvalMaterial(
 		case RAY_TRACING_BLEND_MODE_TRANSLUCENT:
 			// brdf:   Opacity
 			// glass:  1.0 - Opacity
-			if (Payload.Ior == 1.0)
+			if (!Payload.HasRefraction() && Payload.ShadingModelID != SHADINGMODELID_THIN_TRANSLUCENT)
 			{
 				// handled by native transpaarency, just weight down the brdf contrib
 				BrdfWeight = Payload.Opacity;
@@ -640,8 +661,8 @@ FMaterialEval EvalMaterial(
 			else
 			{
 				// brdf or glass
-				BrdfWeight = 1;
-				TransWeight = 1;
+				BrdfWeight = Payload.Opacity;
+				TransWeight = 1 - Payload.Opacity;
 				BrdfProb = Payload.Opacity;
 			}
 			break;
@@ -657,67 +678,64 @@ FMaterialEval EvalMaterial(
 	// NOTE: the translucent portions are all singular -- so we can ignore them in the eval method for now as they have Pdf=+inf for a single direction only
 
 	FMaterialEval Result = NullMaterialEval();
-	switch (Payload.ShadingModelID)
+	if (BrdfWeight > 0.0)
 	{
-		case SHADINGMODELID_UNLIT:
-			break;
-		case SHADINGMODELID_DEFAULT_LIT:
-		case SHADINGMODELID_SUBSURFACE:
-		case SHADINGMODELID_PREINTEGRATED_SKIN:
-		case SHADINGMODELID_SINGLELAYERWATER:
-		case SHADINGMODELID_THIN_TRANSLUCENT:
-			Result = DefaultLit_EvalMaterial(IncomingDirection, OutgoingDirection, Payload);
-			break;
-		case SHADINGMODELID_SUBSURFACE_PROFILE:
-			Result = SubsurfaceProfile_EvalMaterial(IncomingDirection, OutgoingDirection, Payload);
-			break;
-		case SHADINGMODELID_CLOTH:
-			Result = Cloth_EvalMaterial(IncomingDirection, OutgoingDirection, Payload);
-			break;
-		case SHADINGMODELID_CLEAR_COAT:
-			Result = ClearCoat_EvalMaterial(IncomingDirection, OutgoingDirection, Payload);
-			break;
-		case SHADINGMODELID_TWOSIDED_FOLIAGE:
-			Result = TwoSidedFoliage_EvalMaterial(IncomingDirection, OutgoingDirection, Payload);
-			break;
-		case SHADINGMODELID_EYE:
-			Result = Eye_EvalMaterial(IncomingDirection, OutgoingDirection, Payload);
-			break;
-		case SHADINGMODELID_HAIR:
-			Result = Hair_EvalMaterial(IncomingDirection, OutgoingDirection, Payload);
-			break;
-		case SHADINGMODELID_MEDIUM:
-			Result = Medium_EvalMaterial(IncomingDirection, OutgoingDirection, Payload);
-			break;
-		default:
-			Result = Lambert_EvalMaterial(OutgoingDirection, Payload);
-			break;
+		switch (Payload.ShadingModelID)
+		{
+			case SHADINGMODELID_UNLIT:
+				break;
+			case SHADINGMODELID_DEFAULT_LIT:
+			case SHADINGMODELID_SUBSURFACE:
+			case SHADINGMODELID_PREINTEGRATED_SKIN:
+			case SHADINGMODELID_SINGLELAYERWATER:
+			case SHADINGMODELID_THIN_TRANSLUCENT:
+				Result = DefaultLit_EvalMaterial(IncomingDirection, OutgoingDirection, Payload);
+				break;
+			case SHADINGMODELID_SUBSURFACE_PROFILE:
+				Result = SubsurfaceProfile_EvalMaterial(IncomingDirection, OutgoingDirection, Payload);
+				break;
+			case SHADINGMODELID_CLOTH:
+				Result = Cloth_EvalMaterial(IncomingDirection, OutgoingDirection, Payload);
+				break;
+			case SHADINGMODELID_CLEAR_COAT:
+				Result = ClearCoat_EvalMaterial(IncomingDirection, OutgoingDirection, Payload);
+				break;
+			case SHADINGMODELID_TWOSIDED_FOLIAGE:
+				Result = TwoSidedFoliage_EvalMaterial(IncomingDirection, OutgoingDirection, Payload);
+				break;
+			case SHADINGMODELID_EYE:
+				Result = Eye_EvalMaterial(IncomingDirection, OutgoingDirection, Payload);
+				break;
+			case SHADINGMODELID_HAIR:
+				Result = Hair_EvalMaterial(IncomingDirection, OutgoingDirection, Payload);
+				break;
+			case SHADINGMODELID_MEDIUM:
+				Result = Medium_EvalMaterial(IncomingDirection, OutgoingDirection, Payload);
+				break;
+			default:
+				Result = Lambert_EvalMaterial(OutgoingDirection, Payload);
+				break;
+		}
+		// final adjustment to weight / prob for translucency
+		Result.Weight *= BrdfWeight;
+		Result.Pdf *= BrdfProb;
 	}
 
-	// final adjustment to weight / prob for translucency
-	Result.Weight *= BrdfWeight;
-	Result.Pdf *= BrdfProb;
-
 	if (BrdfProb < 1.0)
 	{
-		// refraction model (only evaluate for rough case as singular can't be evaluated)
-		if (Payload.Roughness != 0)
+		// refraction model
+		FMaterialEval GlassEval;
+		if (Payload.ShadingModelID == SHADINGMODELID_THIN_TRANSLUCENT)
 		{
-			FMaterialEval GlassEval;
-			if (Payload.ShadingModelID == SHADINGMODELID_THIN_TRANSLUCENT)
-			{
-				// just a thin slab (reflection-like lobe on both sides)
-				GlassEval = RoughThinGlass_EvalMaterial(IncomingDirection, OutgoingDirection, Payload);
-			}
-			else
-			{
-				// real solid glass, reflection and refraction
-				GlassEval = RoughGlass_EvalMaterial(IncomingDirection, OutgoingDirection, Payload);
-			}
-			
-
-			AddLobeWithMIS(Result.Weight, Result.Pdf, GlassEval.Weight, GlassEval.Pdf, 1.0 - BrdfProb);
+			// just a thin slab (reflection-like lobe on both sides)
+			GlassEval = RoughThinGlass_EvalMaterial(IncomingDirection, OutgoingDirection, Payload);
+		}
+		else
+		{
+			// real solid glass, reflection and refraction
+			GlassEval = RoughGlass_EvalMaterial(IncomingDirection, OutgoingDirection, Payload);
 		}
+		AddLobeWithMIS(Result.Weight, Result.Pdf, GlassEval.Weight * TransWeight, GlassEval.Pdf, 1.0 - BrdfProb);
 	}
 	return Result;
 #endif
diff --git a/Engine/Shaders/Private/PathTracing/Material/PathTracingSolidGlass.ush b/Engine/Shaders/Private/PathTracing/Material/PathTracingSolidGlass.ush
--- a/Engine/Shaders/Private/PathTracing/Material/PathTracingSolidGlass.ush
+++ b/Engine/Shaders/Private/PathTracing/Material/PathTracingSolidGlass.ush
@@ -11,6 +11,7 @@ PathTracingSolidGlass.ush: Microfacet Refractive BSDF for solid glass
 #include "PathTracingGlossy.ush"
 #include "PathTracingEnergyConservation.ush"
 
+
 FMaterialSample RoughGlass_SampleMaterial(
 	float3 RayDirection,
 	FPathTracingPayload Payload,
@@ -22,8 +23,8 @@ FMaterialSample RoughGlass_SampleMaterial(
 	float3 V_World = -RayDirection;
 
 	// #dxr_todo: Maintain a refraction stack on the path tracing payload
-	float Ior = Payload.Ior;
-	float Eta = bIsEntering ? Ior : rcp(Ior);
+	const float Eta = bIsEntering ? Payload.Ior : rcp(Payload.Ior);
+	const float F0 = F0RGBToF0(Payload.SpecularColor);
 
 	float2 Alpha;
 	float3x3 Basis = GetGGXBasis(Payload.Roughness, Payload.Anisotropy, Payload.WorldNormal, Payload.WorldTangent, Alpha);
@@ -41,13 +42,10 @@ FMaterialSample RoughGlass_SampleMaterial(
 	const float E = GGXEnergyLookup(Payload.Roughness, NoV, Eta);
 
 	float3 H = ImportanceSampleVisibleGGX_aniso(UniformSampleDisk(RandSample.xy), Alpha, V).xyz;
-	float3 N = float3(0, 0, 1);
-	float NoH = saturate(H.z);
-	float VoH = saturate(dot(V, H));
 
 	float3 L = 0;
 	float F = 0;
-	if (SampleRefraction(-V, H, Eta, RandSample.z, L, F))
+	if (SampleRefraction(-V, H, Eta, F0, RandSample.z, L, F))
 	{
 		const float2 GGXResult = GGXEvalRefraction(L, V, H, Alpha, Eta);
 
@@ -76,8 +74,8 @@ FMaterialEval RoughGlass_EvalMaterial(
 	float3 N_World = Payload.WorldNormal;
 
 	// #dxr_todo: Maintain a refraction stack on the path tracing payload
-	float Ior = Payload.Ior;
-	float Eta = bIsEntering ? Ior : rcp(Ior);
+	const float Eta = bIsEntering ? Payload.Ior : rcp(Payload.Ior);
+	const float F0 = F0RGBToF0(Payload.SpecularColor);
 
 	float2 Alpha;
 	float3x3 Basis = GetGGXBasis(Payload.Roughness, Payload.Anisotropy, Payload.WorldNormal, Payload.WorldTangent, Alpha);
@@ -102,7 +100,7 @@ FMaterialEval RoughGlass_EvalMaterial(
 		// reflection side
 		const float3 H = normalize(L + V);
 		const float VoH = saturate(dot(V, H));
-		const float F = FresnelReflectance(VoH, Eta);
+		const float F = FresnelReflectance(VoH, Eta, F0);
 
 		const float2 GGXResult = GGXEvalReflection(L, V, H, Alpha);
 		return CreateMaterialEval(GGXResult.x / E, F * GGXResult.y);
@@ -115,7 +113,7 @@ FMaterialEval RoughGlass_EvalMaterial(
 		Ht = normalize((Eta < 1.0f) ? -Ht : Ht);
 
 		float VoH = dot(V, Ht);
-		float F = 1.0f - FresnelReflectance(VoH, Eta);
+		float F = 1.0f - FresnelReflectance(VoH, Eta, F0);
 		if (F > 0)
 		{
 			const float2 GGXResult = GGXEvalRefraction(L, V, Ht, Alpha, Eta);
diff --git a/Engine/Shaders/Private/PathTracing/Material/PathTracingSpecularTransmission.ush b/Engine/Shaders/Private/PathTracing/Material/PathTracingSpecularTransmission.ush
--- a/Engine/Shaders/Private/PathTracing/Material/PathTracingSpecularTransmission.ush
+++ b/Engine/Shaders/Private/PathTracing/Material/PathTracingSpecularTransmission.ush
@@ -9,19 +9,19 @@ PathTracingSpecularTransmission.ush: Refractive BSDFs (pure singular)
 #include "PathTracingMaterialCommon.ush"
 #include "PathTracingGlossy.ush"
 
+// NOTE: this is no longer used, keeping around as reference
+
 FMaterialSample SpecularRefraction_SampleMaterial(
 	float3 RayDirection,
 	FPathTracingPayload Payload,
 	float4 RandSample)
 {
 	// NOTE: IsFrontFace() determines polygon orientation, because the normal is always flipped towards in the incoming ray
-	bool bIsEntering = Payload.IsFrontFace();
-	float3 WorldNormal = Payload.WorldNormal;
+	const bool bIsEntering = Payload.IsFrontFace();
+	const float3 WorldNormal = Payload.WorldNormal;
 
 	// #dxr_todo: Maintain a refraction stack on the path tracing payload
-	float Ior = Payload.Ior;
-	float Eta = bIsEntering ? Ior : 1.0 / Ior;
-
+	const float Eta = bIsEntering ? Payload.Ior : rcp(Payload.Ior);
 
 	float F = 0, BiasSign = 1.0;
 	float3 OutDirection = 0;
diff --git a/Engine/Shaders/Private/PathTracing/Material/PathTracingThinGlass.ush b/Engine/Shaders/Private/PathTracing/Material/PathTracingThinGlass.ush
--- a/Engine/Shaders/Private/PathTracing/Material/PathTracingThinGlass.ush
+++ b/Engine/Shaders/Private/PathTracing/Material/PathTracingThinGlass.ush
@@ -11,7 +11,7 @@ PathTracingThinGlass.ush: Microfacet BSDF for thin glass
 #include "PathTracingGlossy.ush"
 #include "PathTracingEnergyConservation.ush"
 
-struct SlabWeights
+struct FThinSlabWeights
 {
 	float3 Reflected;
 	float3 Transmitted;
@@ -19,30 +19,36 @@ struct SlabWeights
 
 // return the amount of reflected and transmitted energy for an infinite slab viewed from a particular direction
 // taking into account the (repeated) Fresnel terms and the transmittance through the slab
-SlabWeights ComputeThinSlabWeights(float3 SlabColor, float CosV, float Eta)
+FThinSlabWeights ComputeThinSlabWeights(float3 SlabColor, float CosV, float Eta, float F0)
 {
 	// Compute Fresnel reflection off the top of the slab
 	// The Fresnel refletion at the bottom layer is exactly equal so we can assume energy is split into R and T at each bounce
-	const float R = FresnelReflectance(CosV, Eta);
-	const float T = 1.0 - R;
-
+	// NOTE: we don't need to pass Eta through here because TIR (which could only occur in contrived scenarios) isn't really supported
+	const float R = FresnelReflectance(CosV, 1.0, F0);
 
 	// The slab color is normalized to unit thickness, so we only need to account for the refracted viewing angle
-	const float3 A = pow(SlabColor, Eta / sqrt(Eta * Eta - 1.0 + CosV * CosV));
+	// Note that we only allow Eta > 1.0 here so that we don't need to worry about TIR
+	const float3 A = pow(SlabColor, Eta > 1.0 ? Eta * rsqrt(Eta * Eta - 1.0 + CosV * CosV) : rcp(abs(CosV)));
 
 	// figure out how much is reflected/transmitted from a single interaction
 	float3 SlabR = R;
-	float3 SlabT = T * A * T;
+	float3 SlabT = A - R * A; // (1 - R) * A
 
-	// if there was some transmission -- we must also account for the infinite series of bounces within the slab
-	// which can be computed exactly with geometric series
+#if 1
+	// Enable bounces within the slab
 	if (R < 1.0)
 	{
-		SlabR += R * Pow2(T * A) / (1.0 - Pow2(R * A));
-		SlabT /= 1.0 - Pow2(R * A);
+		// If there was some transmission we must also account for the infinite series of bounces within the slab
+		// which can be computed exactly with a geometric series. Note that this causes a bit of energy loss because we don't
+		// account for this non-linear fresnel term in the energy conservation calculation. However, this is fairly minimal in
+		// practice, and the added visual richness seems worth it.
+		float3 InvDenom = rcp(1.0 - Pow2(R * A));
+		SlabR += R * Pow2(SlabT) * InvDenom;
+		SlabT *= InvDenom - R * InvDenom; // (1 - R) * InvDenom
 	}
+#endif
 
-	SlabWeights Result = (SlabWeights) 0;
+	FThinSlabWeights Result = (FThinSlabWeights) 0;
 	Result.Reflected = SlabR;
 	Result.Transmitted = SlabT;
 	return Result;
@@ -54,13 +60,22 @@ FMaterialEval RoughThinGlass_EvalMaterial(
 	FPathTracingPayload Payload
 )
 {
-	float3 V_World = -IncomingDirection;
-	float3 L_World = OutgoingDirection;
-	float3 N_World = Payload.WorldNormal;
+	const float3 V_World = -IncomingDirection;
+	const float3 L_World = OutgoingDirection;
+	const float3 N_World = Payload.WorldNormal;
+
+	const float Eta = Payload.Ior;
+	const float F0 = F0RGBToF0(Payload.SpecularColor);
+	const float F90 = saturate(50.0 * F0);
+	const float3 SlabColor = Payload.CustomData.xyz;
 
+	// transmitted lobe should have slightly less roughness due to the squeeze implied by the IOR (eventually converging to 0.0 as Ior approaches 1.0)
+	const float RoughnessR = Payload.Roughness;
+	const float RoughnessT = Payload.Roughness * sqrt(1.0 - min(Eta, rcp(Eta)));
 
-	float2 Alpha;
-	float3x3 Basis = GetGGXBasis(Payload.Roughness, Payload.Anisotropy, Payload.WorldNormal, Payload.WorldTangent, Alpha);
+	float2 AlphaR;
+	const float3x3 Basis = GetGGXBasis(RoughnessR, Payload.Anisotropy, Payload.WorldNormal, Payload.WorldTangent, AlphaR);
+	const float2 AlphaT = GetGGXAlpha(RoughnessT, Payload.Anisotropy);
 
 	// move vectors into right shading frame
 	float3 V = mul(Basis, V_World);
@@ -72,33 +87,47 @@ FMaterialEval RoughThinGlass_EvalMaterial(
 		return NullMaterialEval();
 	}
 
-	const bool IsTransmission = L.z < 0;
+	const bool bIsReflection = L.z >= 0;
 	L.z = abs(L.z); // push L to the same side as V
 
 	const float NoV = saturate(V.z);
 	const float3 H = normalize(L + V);
 	const float VoH = saturate(dot(V, H));
 
-	const float Eta = Payload.Ior;
-
-	SlabWeights Result = ComputeThinSlabWeights(Payload.CustomData.xyz, VoH, Payload.Ior);
-
-	const float ProbR = LobeSelectionProb(Result.Reflected, Result.Transmitted);
-
-	const float2 GGXResult = GGXEvalReflection(L, V, H, Alpha);
-
-	// GGX reflected energy is split between reflection and refraction equally
-	// so just use the overall GGX energy loss for energy compensation (as in the solid glass case)
-	const float OutWeight = GGXResult.x / GGXEnergyLookup(Payload.Roughness, NoV).x;
 
-	if (IsTransmission)
+	// because the roughnesses may not match, we need to measure the response for each individually, taking into account the reflected and transmitted fresnel respectively
+#if 1
+	// NOTE: 1.0 - F_Schlick(NoV, F0, F90) == F_Schlick(NoV, 1.0 - F0, 1.0 - F90)
+	const FBxDFEnergyTerms SpecR = ComputeGGXSpecEnergyTerms(RoughnessR, NoV, F0, F90);
+	const FBxDFEnergyTerms SpecT = ComputeGGXSpecEnergyTerms(RoughnessT, NoV, 1.0 - F0, 1.0 - F90);
+	// Approximate the transmission as just a plain tint by slab color because we can't depend on the half-vector for lobe selection
+	const float NudgeE = rcp(SpecR.E.x + SpecT.E.x);
+	const float ProbR = LobeSelectionProb(SpecR.E, Payload.HasRefraction() ? SpecT.E * SlabColor : 0.0);
+#else
+	const float Er = GGXEnergyLookup(RoughnessR, NoV).x;
+	const float Et = GGXEnergyLookup(RoughnessT, NoV).x;
+	const float ProbR = LobeSelectionProb(F0 * Er, Payload.HasRefraction() ? (1 - F0) * Et * SlabColor : 0.0);
+#endif
+
+	const FThinSlabWeights Result = ComputeThinSlabWeights(SlabColor, VoH, Payload.Ior, F0);
+	const float2 GGXResult = GGXEvalReflection(L, V, H, bIsReflection ? AlphaR : AlphaT);
+	if (bIsReflection)
 	{
-		return CreateMaterialEval(OutWeight / (1.0 - ProbR) * Result.Transmitted, GGXResult.y * (1.0 - ProbR));
+		if (ProbR > 0.0)
+		{
+			const float OutWeight = NudgeE * GGXResult.x * SpecR.W.x;
+			return CreateMaterialEval(OutWeight / ProbR * Result.Reflected, GGXResult.y * ProbR);
+		}
 	}
 	else
 	{
-		return CreateMaterialEval(OutWeight / ProbR * Result.Reflected, GGXResult.y * ProbR);
+		if (ProbR < 1.0)
+		{
+			const float OutWeight = NudgeE * GGXResult.x * SpecT.W.x;
+			return CreateMaterialEval(OutWeight / (1.0 - ProbR) * Result.Transmitted, GGXResult.y * (1.0 - ProbR));
+		}
 	}
+	return NullMaterialEval();
 }
 
 
@@ -107,16 +136,24 @@ FMaterialSample RoughThinGlass_SampleMaterial(
 	FPathTracingPayload Payload,
 	float4 RandSample)
 {
-	float3 N_World = Payload.WorldNormal;
-	float3 V_World = -RayDirection;
+	const float3 N_World = Payload.WorldNormal;
+	const float3 V_World = -RayDirection;
 
-	float Eta = Payload.Ior;
+	const float Eta = Payload.Ior;
+	const float F0 = F0RGBToF0(Payload.SpecularColor);
+	const float F90 = saturate(50.0 * F0);
+	const float3 SlabColor = Payload.CustomData.xyz;
 
-	float2 Alpha;
-	float3x3 Basis = GetGGXBasis(Payload.Roughness, Payload.Anisotropy, Payload.WorldNormal, Payload.WorldTangent, Alpha);
+	// transmitted lobe should have slightly less roughness due to the squeeze implied by the IOR (eventually converging to 0.0 as Ior approaches 1.0)
+	const float RoughnessR = Payload.Roughness;
+	const float RoughnessT = Payload.Roughness * sqrt(1.0 - min(Eta, rcp(Eta)));
 
-	float3 V = mul(Basis, V_World);
-	float NoV = saturate(V.z);
+	float2 AlphaR;
+	const float3x3 Basis = GetGGXBasis(RoughnessR, Payload.Anisotropy, Payload.WorldNormal, Payload.WorldTangent, AlphaR);
+	const float2 AlphaT = GetGGXAlpha(RoughnessT, Payload.Anisotropy);
+
+	const float3 V = mul(Basis, V_World);
+	const float NoV = saturate(V.z);
 
 	if (NoV == 0)
 	{
@@ -124,36 +161,45 @@ FMaterialSample RoughThinGlass_SampleMaterial(
 		return NullMaterialSample();
 	}
 
-	float3 H = ImportanceSampleVisibleGGX_aniso(UniformSampleDisk(RandSample.xy), Alpha, V).xyz;
-	float3 N = float3(0, 0, 1);
-	float NoH = saturate(H.z);
-	float VoH = saturate(dot(V, H));
+	// because the roughnesses may not match, we need to measure the response for each individually, taking into account the reflected and transmitted fresnel respectively
+#if 1
+	// NOTE: 1.0 - F_Schlick(NoV, F0, F90) == F_Schlick(NoV, 1.0 - F0, 1.0 - F90)
+	const FBxDFEnergyTerms SpecR = ComputeGGXSpecEnergyTerms(RoughnessR, NoV, F0, F90);
+	const FBxDFEnergyTerms SpecT = ComputeGGXSpecEnergyTerms(RoughnessT, NoV, 1.0 - F0, 1.0 - F90);
+	// Approximate the transmission as just a plain tint by slab color because we can't depend on the half-vector for lobe selection
+	const float NudgeE = rcp(SpecR.E.x + SpecT.E.x);
+	const float ProbR = LobeSelectionProb(SpecR.E, Payload.HasRefraction() ? SpecT.E * SlabColor : 0.0);
+#else
+	const float Er = GGXEnergyLookup(RoughnessR, NoV).x;
+	const float Et = GGXEnergyLookup(RoughnessT, NoV).x;
+	const float ProbR = LobeSelectionProb(F0 * Er, Payload.HasRefraction() ? (1 - F0) * Et * SlabColor : 0.0);
+#endif
+
+	const bool bIsReflection = RandSample.z < ProbR;
+
+	const float3 H = ImportanceSampleVisibleGGX_aniso(UniformSampleDisk(RandSample.xy), bIsReflection ? AlphaR : AlphaT, V).xyz;
+	const float NoH = saturate(H.z);
+	const float VoH = saturate(dot(V, H));
 
-	float3 L = reflect(-V, H);
+	const float3 L = reflect(-V, H);
 	if (L.z <= 0)
 	{
 		// invalid output direction, exit early
 		return NullMaterialSample();
 	}
 
-	float NoL = saturate(L.z);
-
-	SlabWeights Result = ComputeThinSlabWeights(Payload.CustomData.xyz, VoH, Payload.Ior);
+	const float NoL = saturate(L.z);
 
-	const float ProbR = LobeSelectionProb(Result.Reflected, Result.Transmitted);
+	FThinSlabWeights Result = ComputeThinSlabWeights(SlabColor, VoH, Payload.Ior, F0);
 
-	const float2 GGXResult = GGXEvalReflection(L, V, H, Alpha);
+	const float2 GGXResult = GGXEvalReflection(L, V, H, bIsReflection ? AlphaR : AlphaT);
 
-	// GGX reflected energy is split between reflection and refraction equally
-	// so just use the overall GGX energy loss for energy compensation (as in the solid glass case)
-	float3 OutWeight = GGXResult.x / GGXEnergyLookup(Payload.Roughness, NoV).x;
+	float3 OutWeight = GGXResult.x;
 	float OutPdf = GGXResult.y;
-
-	// decide between reflection/transmission
-	if (RandSample.z < ProbR)
+	if (bIsReflection)
 	{
 		// reflected ray
-		OutWeight *= Result.Reflected / ProbR;
+		OutWeight *= NudgeE * SpecR.W.x * Result.Reflected / ProbR;
 		OutPdf *= ProbR;
 		float3 OutDirection = normalize(mul(L, Basis));
 		return CreateMaterialSample(OutDirection, OutWeight, OutPdf, 1.0, Payload.Roughness);
@@ -161,7 +207,7 @@ FMaterialSample RoughThinGlass_SampleMaterial(
 	else
 	{
 		// transmitted ray
-		OutWeight *= Result.Transmitted / (1.0 - ProbR);
+		OutWeight *= NudgeE * SpecT.W.x * Result.Transmitted / (1.0 - ProbR);
 		OutPdf *= 1.0 - ProbR;
 		float3 OutDirection = normalize(mul(float3(L.xy, -L.z), Basis)); // flip reflection to other side
 		return CreateMaterialSample(OutDirection, OutWeight, OutPdf, -1.0, Payload.Roughness);
diff --git a/Engine/Shaders/Private/PathTracing/PathTracingCommon.ush b/Engine/Shaders/Private/PathTracing/PathTracingCommon.ush
--- a/Engine/Shaders/Private/PathTracing/PathTracingCommon.ush
+++ b/Engine/Shaders/Private/PathTracing/PathTracingCommon.ush
@@ -165,19 +165,29 @@ struct FPathTracingPayload : FMinimalPayload
 		Roughness = lerp(RoughnessData.x, RoughnessData.y, RoughnessData.z);
 	}
 
+	bool HasRefraction()
+	{
+		return Ior > 0.0;
+	}
+
 	bool IsMaterialSolidGlass()
 	{
-		return ShadingModelID == SHADINGMODELID_DEFAULT_LIT && BlendingMode == RAY_TRACING_BLEND_MODE_TRANSLUCENT && Ior != 1.0 && Opacity < 1;
+		return ShadingModelID == SHADINGMODELID_DEFAULT_LIT && BlendingMode == RAY_TRACING_BLEND_MODE_TRANSLUCENT && HasRefraction() && Opacity < 1;
+	}
+
+	bool IsMaterialThinGlass()
+	{
+		return ShadingModelID == SHADINGMODELID_THIN_TRANSLUCENT && BlendingMode == RAY_TRACING_BLEND_MODE_TRANSLUCENT && HasRefraction() && Opacity < 1;
 	}
 
 	bool IsMaterialTransmissive()
 	{
 		return ShadingModelID == SHADINGMODELID_TWOSIDED_FOLIAGE ||
-			   ShadingModelID == SHADINGMODELID_THIN_TRANSLUCENT || 
 			   ShadingModelID == SHADINGMODELID_HAIR ||
 			   ShadingModelID == SHADINGMODELID_MEDIUM ||
 			   ShadingModelID == SHADINGMODELID_SINGLELAYERWATER ||	
-			   IsMaterialSolidGlass();
+			   IsMaterialSolidGlass() ||
+			   IsMaterialThinGlass();
 	}
 
 	bool IsSubsurfaceMaterial()
diff --git a/Engine/Shaders/Private/PathTracing/PathTracingMaterialHitShader.usf b/Engine/Shaders/Private/PathTracing/PathTracingMaterialHitShader.usf
--- a/Engine/Shaders/Private/PathTracing/PathTracingMaterialHitShader.usf
+++ b/Engine/Shaders/Private/PathTracing/PathTracingMaterialHitShader.usf
@@ -41,7 +41,6 @@ float GetRefractionIor(FPixelMaterialInputs PixelMaterialInputs)
 	float3 F0 = ComputeF0(Specular, BaseColor, Metallic);
 	float Ior = DielectricF0RGBToIor(F0);
 #endif
-	Ior = Ior == 0.0 ? 1.0 : Ior;
 	return Ior;
 }
 
@@ -137,7 +136,7 @@ RAY_TRACING_ENTRY_CLOSEST_HIT(PathTracingMaterialCHS,
 #if MATERIALBLENDING_TRANSLUCENT && REFRACTION_USE_INDEX_OF_REFRACTION
 	Payload.Ior = GetRefractionIor(PixelMaterialInputs);
 #else
-	Payload.Ior = 1;
+	Payload.Ior = 0;
 #endif
 
 #if !SIMPLIFIED_MATERIAL_SHADER
@@ -383,7 +382,7 @@ RAY_TRACING_ENTRY_CLOSEST_HIT(PathTracingMaterialCHS,
 #else
 	Payload.Opacity = GetMaterialOpacity(PixelMaterialInputs);
 #endif
-	Payload.Ior = 1;
+	Payload.Ior = 0;
 
 #if !MATERIAL_TWOSIDED
 	if (!Payload.IsFrontFace() && !Payload.IsMaterialTransmissive())
@@ -503,15 +502,18 @@ RAY_TRACING_ENTRY_ANY_HIT(PathTracingMaterialAHS,
 #if REFRACTION_USE_INDEX_OF_REFRACTION && !SIMPLIFIED_MATERIAL_SHADER
 	float Ior = GetRefractionIor(PixelMaterialInputs);
 #else
-	float Ior = 1.0;
+	float Ior = 0.0;
 #endif
 
 	float3 Transparency = 1 - Opacity;
 	float3 V = WorldRayDirection();
 	float3 N = normalize(MaterialParameters.WorldNormal);
 	float VoN = abs(dot(V, N));
-#if !SIMPLIFIED_MATERIAL_SHADER
-	if (Opacity < 1.0 && Ior != 1.0)
+#if SIMPLIFIED_MATERIAL_SHADER
+	// simplified logic with no bending at the interface and no fresnel
+	Transparency *= pow(Transmission, 1.0 / VoN);
+#else
+	if (Opacity < 1.0)
 	{
 		float PathRoughness = PackedPayload.GetPathRoughness();
 		float Roughness = GetMaterialRoughness(PixelMaterialInputs);
@@ -521,18 +523,17 @@ RAY_TRACING_ENTRY_ANY_HIT(PathTracingMaterialAHS,
 			return;
 		}
 		// compute transmission through the slab (fresnel + absorption)
-		Transparency *= ComputeThinSlabWeights(Transmission, VoN, Ior).Transmitted;
+		float3 BaseColor = GetMaterialBaseColor(PixelMaterialInputs);
+		float  Metallic = GetMaterialMetallic(PixelMaterialInputs);
+		float  Specular = GetMaterialSpecular(PixelMaterialInputs);
+		float  F0 = F0RGBToF0(ComputeF0(Specular, BaseColor, Metallic));
+		Transparency *= ComputeThinSlabWeights(Transmission, VoN, Ior, F0).Transmitted;
 
 		// fake caustic approximation (see comments below)
 		Transparency *= (1 - Roughness * Roughness) * saturate(PathRoughness - Roughness);
 	}
-	else
 #endif
-	{
-		// no bending at the interface so no fresnel - equation above simplifies to just A
-		float3 A = pow(Transmission, 1.0 / VoN);
-		Transparency *= A;
-	}
+
 #elif MATERIALBLENDING_MASKED && MATERIAL_DITHER_OPACITY_MASK
 	// See MATERIAL_DITHER_OPACITY_MASK comment above
 	float Opacity = saturate(GetMaterialMaskInputRaw(PixelMaterialInputs));
@@ -550,7 +551,7 @@ RAY_TRACING_ENTRY_ANY_HIT(PathTracingMaterialAHS,
 #if !SIMPLIFIED_MATERIAL_SHADER && MATERIAL_SHADINGMODEL_DEFAULT_LIT && REFRACTION_USE_INDEX_OF_REFRACTION
 	float Ior = GetRefractionIor(PixelMaterialInputs);
 	int ShadingModelID = GetMaterialShadingModel(PixelMaterialInputs);
-	if (Transparency > 0 && ShadingModelID == SHADINGMODELID_DEFAULT_LIT && Ior != 1.0)
+	if (Transparency > 0 && ShadingModelID == SHADINGMODELID_DEFAULT_LIT && Ior > 0.0)
 	{
 		// current hit has some refraction
 		float PathRoughness = PackedPayload.GetPathRoughness();
@@ -570,7 +571,11 @@ RAY_TRACING_ENTRY_ANY_HIT(PathTracingMaterialAHS,
 		float3 N = normalize(MaterialParameters.WorldNormal);
 		float Eta = Ior;
 		float VoN = abs(dot(WorldRayDirection(), N));
-		float Fr = FresnelReflectance(VoN, Eta);
+		float3 BaseColor = GetMaterialBaseColor(PixelMaterialInputs);
+		float  Metallic = GetMaterialMetallic(PixelMaterialInputs);
+		float  Specular = GetMaterialSpecular(PixelMaterialInputs);
+		float  F0 = F0RGBToF0(ComputeF0(Specular, BaseColor, Metallic));
+		float Fr = FresnelReflectance(VoN, Eta, F0);
 		Transparency *= pow(1 - Fr, 2) * (1 - Roughness * Roughness) * saturate(PathRoughness - Roughness);
 
 #if HAVE_GetAbsorptionMediumMaterialOutput0
diff --git a/Engine/Shaders/Private/ShadingCommon.ush b/Engine/Shaders/Private/ShadingCommon.ush
--- a/Engine/Shaders/Private/ShadingCommon.ush
+++ b/Engine/Shaders/Private/ShadingCommon.ush
@@ -100,9 +100,14 @@ float F0ToDielectricSpecular(float F0)
 	return saturate(F0 / 0.08f);
 }
 
+float F0RGBToF0(float3 F0)
+{
+	return dot(0.3333333.xxx, F0);
+}
+
 float F0RGBToDielectricSpecular(float3 F0)
 {
-	return F0ToDielectricSpecular(max3(F0.r, F0.g, F0.b));
+	return F0ToDielectricSpecular(F0RGBToF0(F0));
 }
 
 half DielectricSpecularToF0(half Specular)
@@ -113,12 +118,12 @@ half DielectricSpecularToF0(half Specular)
 // [Burley, "Extending the Disney BRDF to a BSDF with Integrated Subsurface Scattering"]
 float DielectricF0ToIor(float F0)
 {
-	return 2.0f / (1.0f - sqrt(F0)) - 1.0f;
+	return 2.0f / (1.0f - sqrt(min(F0,0.99))) - 1.0f;
 }
 
 float DielectricF0RGBToIor(float3 F0)
 {
-	return DielectricF0ToIor(max3(F0.r, F0.g, F0.b));
+	return DielectricF0ToIor(F0RGBToF0(F0));
 }
 
 float DielectricIorToF0(float Ior)
diff --git a/Engine/Shaders/Private/ShadingEnergyConservationTable.usf b/Engine/Shaders/Private/ShadingEnergyConservationTable.usf
--- a/Engine/Shaders/Private/ShadingEnergyConservationTable.usf
+++ b/Engine/Shaders/Private/ShadingEnergyConservationTable.usf
@@ -2,6 +2,7 @@
 
 #include "Common.ush"
 #include "BRDF.ush"
+#include "ShadingCommon.ush"
 #include "/Engine/Shared/PathTracingDefinitions.h"
 
 // Move these utility functions into a common/non-path-tracer specific place
@@ -27,14 +28,19 @@ RWTexture3D<float2> OutputTexture3D;
 [numthreads(THREADGROUPSIZE_X, THREADGROUPSIZE_Y, 1)]
 void BuildEnergyTableCS(uint3 DispatchThreadId : SV_DispatchThreadID)
 {
-	float3 Index = (float3(DispatchThreadId) + 0.5) / float(ENERGY_TABLE_RESOLUTION);
+	const float3 Index = (float3(DispatchThreadId) + 0.5) / float(ENERGY_TABLE_RESOLUTION);
 
 	// Map from thread ID to the exact roughness/viewing angle we want to compute
-	float NoV = max(Index.x, 1e-8);
-	float Roughness = Index.y;
+	const float NoV = max(Index.x, 1e-8);
+	const float Roughness = Index.y;
 #if BUILD_ENERGY_TABLE == 1
-	float EtaIn = 1.0 + 2.0 * Index.z; // only handle Ior in [1,3] range
-	float EtaOut = 1.0 / EtaIn;
+	const float EtaIn = 1.0 + 2.0 * Index.z; // only handle Ior in [1,3] range
+	const float EtaOut = rcp(EtaIn);
+	// NOTE: Baking a table for all possible F0 values would greatly complicate things
+	// Instead, we aim for optimal results when the Ior values is consistent with F0.
+	// In practice, it appears that this gives very good results, and the penalty when
+	// Ior and F0 are mismatched does not appear to be that large.
+	const float F0 = DielectricIorToF0(EtaIn); // NOTE: this is always equal to DielectricIorToF0(EtaOut)
 #endif
 
 	float2 Alpha; GetGGXBasis(Roughness, float3(0, 0, 1), Alpha);
@@ -66,7 +72,7 @@ void BuildEnergyTableCS(uint3 DispatchThreadId : SV_DispatchThreadID)
 			// compute weight for outside -> in case
 			float3 L = 0;
 			float F = 0;
-			if (SampleRefraction(-V, H, EtaIn, RandSample.z, L, F))
+			if (SampleRefraction(-V, H, EtaIn, F0, RandSample.z, L, F))
 			{
 				Weight.x = GGXEvalRefraction(L, V, H, Alpha, EtaIn).x;
 			}
@@ -79,7 +85,7 @@ void BuildEnergyTableCS(uint3 DispatchThreadId : SV_DispatchThreadID)
 			// compute weight for inside -> out case
 			float3 L = 0;
 			float F = 0;
-			if (SampleRefraction(-V, H, EtaOut, RandSample.z, L, F))
+			if (SampleRefraction(-V, H, EtaOut, F0, RandSample.z, L, F))
 			{
 				Weight.y = GGXEvalRefraction(L, V, H, Alpha, EtaOut).x;
 			}
