Path Tracer: Fix emissive materials not respecting two-sided flag

The path tracer only uses backfacing culing on camera rays to avoid light leaks and to properly render glass and subsurface materials. However this means that emissive properties get picked up on both sides of thin models.

This change zeroes out the emissive contribution on the backside if the material is single-sided to avoid this issue.

diff --git a/Engine/Shaders/Private/PathTracing/PathTracingMaterialHitShader.usf b/Engine/Shaders/Private/PathTracing/PathTracingMaterialHitShader.usf
--- a/Engine/Shaders/Private/PathTracing/PathTracingMaterialHitShader.usf
+++ b/Engine/Shaders/Private/PathTracing/PathTracingMaterialHitShader.usf
@@ -152,6 +152,13 @@ RAY_TRACING_ENTRY_CLOSEST_HIT(PathTracingMaterialCHS,
 	Payload.WorldNormal = MaterialParameters.WorldNormal;
 	Payload.WorldSmoothNormal = MaterialParameters.TwoSidedSign * TransformTangentNormalToWorld(MaterialParameters.TangentToWorld, float3(0, 0, 1));
 	Payload.Radiance = GetMaterialEmissive(PixelMaterialInputs);
+#if TWO_SIDED_MATERIAL == 0
+	if (MaterialParameters.TwoSidedSign < 0)
+	{
+		// when viewing the surface from "inside", don't include emission
+		Payload.Radiance = 0;
+	}
+#endif
 	Payload.BaseColor = BaseColor;
 	Payload.Specular = Specular;
 	Payload.Roughness = Roughness;
@@ -372,6 +379,13 @@ RAY_TRACING_ENTRY_CLOSEST_HIT(PathTracingMaterialCHS,
 	float GeoNormalSign = MaterialParameters.TwoSidedSign * GetPrimitive_DeterminantSign(MaterialParameters.PrimitiveId);
 	Payload.WorldGeoNormal = GeoNormalSign * GeoNormal;
 	Payload.Radiance = GetMaterialEmissive(PixelMaterialInputs);
+#if TWO_SIDED_MATERIAL == 0
+	if (MaterialParameters.TwoSidedSign < 0)
+	{
+		// when viewing the surface from "inside", don't include emission
+		Payload.Radiance = 0;
+	}
+#endif
 	Payload.BaseColor = BaseColor;
 	Payload.Specular = 0;
 	Payload.Roughness = 1;
diff --git a/Engine/Source/Runtime/Renderer/Private/PathTracing.cpp b/Engine/Source/Runtime/Renderer/Private/PathTracing.cpp
--- a/Engine/Source/Runtime/Renderer/Private/PathTracing.cpp
+++ b/Engine/Source/Runtime/Renderer/Private/PathTracing.cpp
@@ -1161,6 +1161,7 @@ public:
 		OutEnvironment.SetDefine(TEXT("USE_RAYTRACED_TEXTURE_RAYCONE_LOD"), 0);
 		OutEnvironment.SetDefine(TEXT("SCENE_TEXTURES_DISABLED"), 1);
 		OutEnvironment.SetDefine(TEXT("SIMPLIFIED_MATERIAL_SHADER"), UseSimplifiedShader);
+		OutEnvironment.SetDefine(TEXT("TWO_SIDED_MATERIAL"), Parameters.MaterialParameters.bIsTwoSided ? 1 : 0);
 		FMeshMaterialShader::ModifyCompilationEnvironment(Parameters, OutEnvironment);
 	}
 
