Initial forward raster version for hair strands.

This is a early prototype version, which has a lot of todos. This first version is meant for easing future iteration. Most of the parts are still in a naive version as this stage, and does not render correctly (not shadow support, ordering issue, ...).

diff --git a/Engine/Shaders/Private/HairStrands/HairStrandsForwardCulling.usf b/Engine/Shaders/Private/HairStrands/HairStrandsForwardCulling.usf
new file mode 100644
--- /dev/null
+++ b/Engine/Shaders/Private/HairStrands/HairStrandsForwardCulling.usf
@@ -0,0 +1,339 @@
+// Copyright Epic Games, Inc. All Rights Reserved.
+
+#define HAIR_STRANDS_PARAMETERS 1
+
+#define CULL_ENABLE 0
+
+#define CULLING_CURVE 1
+#define CULLING_POINT 2
+#define CULLING_TYPE CULLING_CURVE
+
+#if SHADER_RASTERCOMPUTE_CULL
+#include "../Common.ush"
+#include "HairStrandsVisibilityCommon.ush"
+#include "HairStrandsVertexFactoryCommon.ush"
+
+// TODO 
+// * HZB
+// * Position quantization 3 x 20 bits?
+// * Enabling culling (viewport/depth), we need to track of hiddent parts, so that CP type is correct (1:begin/0:middle/2:end)
+
+// Input
+Texture2D<float>	SceneDepthTexture;
+float2				OutputResolution;
+uint				HairMaterialId;
+uint				ControlPointCount;
+uint				CurveCount;
+
+// Output
+RWBuffer<float4>					RWPoints;
+RWStructuredBuffer<FPackedHairVis>	RWHairVis;
+RWBuffer<uint2>						RWCoord;
+RWTexture2D<uint>					RWCounter;
+
+///////////////////////////////////////////////////////////////////////////////////////////////////
+// Utils
+
+float3 NDCToPixelCoord(float4 InDC)
+{
+	const float3 NDC = InDC.xyz / InDC.w;
+	float2 UV = NDC.xy * ResolvedView.ScreenPositionScaleBias.xy + ResolvedView.ScreenPositionScaleBias.wz;
+	return float3(UV * OutputResolution, NDC.z);
+}
+
+// Line clipping based on "CLIPPING USING HOMOGENEOUS COORDINATES" by Blinn et al.
+bool IsViewportVisible(inout float4 P0, inout float4 P1)
+{
+	float2 T = float2(0.0f, 1.0f);
+	bool bIsRemoved = P0.w < 0.0f && P1.w < 0.0f; // Both points behind near plane
+
+	bool bSign = false;
+
+	UNROLL
+	for (uint PlaneIdx = 0; PlaneIdx < 6; ++PlaneIdx)
+	{
+		// Compute boundary coordinates of both points (w+x, w-x, w+y, w-y, z, w-z)
+		bSign = !bSign;
+		const uint CompIdx = PlaneIdx / 2;
+		const float Sign = bSign ? 1.0f : -1.0f;
+		const float WFactor = PlaneIdx != 4 ? 1.0f : 0.0f;
+		const float2 BC = WFactor * float2(P0.w, P1.w) + Sign * float2(P0[CompIdx], P1[CompIdx]);
+
+		float Num = BC.x;
+		float Denom = BC.x - BC.y;
+		bIsRemoved = bIsRemoved || (BC.x < 0.0f && BC.y < 0.0f); // Both points outside the plane
+		float Alpha = Num / Denom;
+
+		// If the denominator is negative, P0 has a smaller boundary coordinate than P1, so we can assume
+		// that P1 is inside the plane (or bIsRemoved is true), so we need to update the alpha for P0.
+		// The reverse is true if the denominator is positive.
+		if (Denom < 0.0f)
+		{
+			T.x = max(T.x, Alpha);
+		}
+		else
+		{
+			T.y = min(T.y, Alpha);
+		}
+	}
+
+	return !bIsRemoved;
+}
+
+///////////////////////////////////////////////////////////////////////////////////////////////////
+// Curve version
+#if CULLING_TYPE == CULLING_CURVE
+groupshared uint s_ValidSegmentCount;
+groupshared uint s_WriteOffset;
+
+[numthreads(1024, 1, 1)]
+void CSMain(uint DispatchThreadID : SV_DispatchThreadID, uint GroupThreadID : SV_GroupThreadID, uint GroupID : SV_GroupID)
+{
+	ResolvedView = ResolveView();
+
+	if (GroupThreadID == 0)
+	{
+		s_ValidSegmentCount= 0;
+		s_WriteOffset = 0;
+	}	
+	GroupMemoryBarrierWithGroupSync();
+
+	const uint CurveID = DispatchThreadID;
+	const bool bValidCurve = CurveID < CurveCount;
+	const float3 PositionOffset = HairStrandsVF_PositionOffsetBuffer[0].xyz;
+
+	bool bIsCurveVisibile = false;
+	uint PrimWriteOffset = 0;
+	FHairCurve Curve = (FHairCurve)0;
+	if (bValidCurve)
+	{
+		Curve = UnpackHairCurve(HairStrandsVF_CurveBuffer[CurveID]);
+		
+		uint VisiblePointCount = 0;
+		for (uint PointIt=0; PointIt < Curve.PointCount-1; ++PointIt)
+		{
+			// Sanity check
+			const uint PrimID = PointIt + Curve.PointIndex;
+			const bool bValidPoint = PrimID + 1 < ControlPointCount;
+			if (bValidPoint)
+			{
+				// Fetch both control points
+				const FHairControlPoint CP0 = UnpackHairControlPoint(
+					HairStrandsVF_PositionBuffer[PrimID],
+					PositionOffset,
+					HairStrandsVF_Radius,
+					HairStrandsVF_RootScale,
+					HairStrandsVF_TipScale);
+
+				const FHairControlPoint CP1 = UnpackHairControlPoint(
+					HairStrandsVF_PositionBuffer[PrimID + 1],
+					PositionOffset,
+					HairStrandsVF_Radius,
+					HairStrandsVF_RootScale,
+					HairStrandsVF_TipScale);
+
+				// Transform to world space (TODO change this to translated world space?)
+				const float3 CP0WP = mul(float4(CP0.Position, 1.0f), HairStrandsVF_LocalToWorldPrimitiveTransform).xyz;
+				const float3 CP1WP = mul(float4(CP1.Position, 1.0f), HairStrandsVF_LocalToWorldPrimitiveTransform).xyz;
+				float4 SP0 = mul(float4(CP0WP, 1.0f), LWCHackToFloat(PrimaryView.WorldToClip));
+				float4 SP1 = mul(float4(CP1WP, 1.0f), LWCHackToFloat(PrimaryView.WorldToClip));
+
+				#if CULL_ENABLE
+				if (IsViewportVisible(SP0, SP1))
+				#endif
+				{
+					++VisiblePointCount;
+					if (CP1.Type == HAIR_CONTROLPOINT_END)
+					{
+						++VisiblePointCount;
+					}
+					bIsCurveVisibile = true;
+				}
+			}
+		}
+		InterlockedAdd(s_ValidSegmentCount, VisiblePointCount, PrimWriteOffset);
+	}
+	GroupMemoryBarrierWithGroupSync();
+
+	// Alloc
+	if (GroupThreadID == 0)
+	{
+		InterlockedAdd(RWCounter[uint2(0,0)], s_ValidSegmentCount, s_WriteOffset);
+	}
+	GroupMemoryBarrierWithGroupSync();
+
+	// Write valid segments
+	if (bValidCurve && bIsCurveVisibile)
+	{
+		uint WriteIndex = s_WriteOffset + PrimWriteOffset;
+		for (uint PointIt=0; PointIt < Curve.PointCount; ++PointIt)
+		{
+			// Sanity check
+			const uint PrimID = PointIt + Curve.PointIndex;
+			const bool bValidPoint = PrimID < ControlPointCount;
+			if (bValidPoint)
+			{
+				// Fetch both control points
+				const FHairControlPoint CP0 = UnpackHairControlPoint(
+					HairStrandsVF_PositionBuffer[PrimID],
+					PositionOffset,
+					HairStrandsVF_Radius,
+					HairStrandsVF_RootScale,
+					HairStrandsVF_TipScale);
+
+				const FHairControlPoint CP1 = UnpackHairControlPoint(
+					HairStrandsVF_PositionBuffer[PrimID + 1],
+					PositionOffset,
+					HairStrandsVF_Radius,
+					HairStrandsVF_RootScale,
+					HairStrandsVF_TipScale);
+
+				// Transform to world space (TODO change this to translated world space?)
+				const float3 CP0WP = mul(float4(CP0.Position, 1.0f), HairStrandsVF_LocalToWorldPrimitiveTransform).xyz;
+				const float3 CP1WP = mul(float4(CP1.Position, 1.0f), HairStrandsVF_LocalToWorldPrimitiveTransform).xyz;
+				float4 SP0 = mul(float4(CP0WP, 1.0f), LWCHackToFloat(PrimaryView.WorldToClip));
+				float4 SP1 = mul(float4(CP1WP, 1.0f), LWCHackToFloat(PrimaryView.WorldToClip));
+
+				#if CULL_ENABLE
+				if (IsViewportVisible(SP0, SP1))
+				#endif
+				{
+					SP0 = float4(NDCToPixelCoord(SP0), 1.0f / SP0.w);
+					SP1 = float4(NDCToPixelCoord(SP1), 1.0f / SP1.w);
+
+					FHairControlPoint Point0 = CP0;
+					Point0.Position = CP0WP;
+
+					const float DeviceZ = SP0.z;
+					const float2 Coord = SP0.xy;
+
+					FHairVis Out;
+					Out.Depth			= DeviceZ;
+					Out.Coverage		= 1;
+					Out.Coverage8bit	= 255;
+					Out.PrimitiveId		= PrimID;
+					Out.MaterialId		= HairMaterialId;
+
+					uint PackedW = f32tof16(Point0.WorldRadius);
+					PackedW = PackedW | (uint(saturate(Point0.UCoord) * 0xFF)<<16);
+					PackedW = PackedW | (Point0.Type<<24);
+
+					RWHairVis[WriteIndex] = PackHairVis(Out);
+					RWCoord[WriteIndex]   = Coord;
+					RWPoints[WriteIndex] = float4(Point0.Position, asfloat(PackedW));
+					++WriteIndex;
+				}
+			}
+		}
+	}
+}
+#endif // CULLING_TYPE == CULLING_CURVE
+
+///////////////////////////////////////////////////////////////////////////////////////////////////
+// Point version
+#if CULLING_TYPE == CULLING_POINT
+groupshared uint s_ValidSegmentCount;
+groupshared uint s_WriteOffset;
+
+[numthreads(1024, 1, 1)]
+void CSMain(uint DispatchThreadID : SV_DispatchThreadID, uint GroupThreadID : SV_GroupThreadID, uint GroupID : SV_GroupID)
+{
+	ResolvedView = ResolveView();
+
+	if (GroupThreadID == 0)
+	{
+		s_ValidSegmentCount= 0;
+		s_WriteOffset = 0;
+	}
+	
+	GroupMemoryBarrierWithGroupSync();
+
+	uint   PrimID = DispatchThreadID;
+	bool   bIsPrimVisibile = false;
+	uint   PrimWriteOffset = 0;
+	float  DeviceZ = 0;
+	float2 Coord = 0;
+
+	FHairControlPoint Point0 = (FHairControlPoint)0;
+
+	// Count valid segments
+	bool bIsPrimValid = PrimID < ControlPointCount;
+	if (bIsPrimValid)
+	{
+		const float3 PositionOffset = HairStrandsVF_PositionOffsetBuffer[0].xyz;
+
+		// Fetch both control points
+		FHairControlPoint CP0 = UnpackHairControlPoint(
+			HairStrandsVF_PositionBuffer[PrimID],
+			PositionOffset,
+			HairStrandsVF_Radius,
+			HairStrandsVF_RootScale,
+			HairStrandsVF_TipScale);
+
+		FHairControlPoint CP1 = CP0;
+		if (CP0.Type != HAIR_CONTROLPOINT_END)
+		{
+			CP1 = UnpackHairControlPoint(
+				HairStrandsVF_PositionBuffer[PrimID + 1],
+				PositionOffset,
+				HairStrandsVF_Radius,
+				HairStrandsVF_RootScale,
+				HairStrandsVF_TipScale);
+		}
+
+		// Transform to world space
+		const float3 CP0WP = mul(float4(CP0.Position, 1.0f), HairStrandsVF_LocalToWorldPrimitiveTransform).xyz;
+		const float3 CP1WP = mul(float4(CP1.Position, 1.0f), HairStrandsVF_LocalToWorldPrimitiveTransform).xyz;
+		float4 SP0 = mul(float4(CP0WP, 1.0f), LWCHackToFloat(PrimaryView.WorldToClip));
+		float4 SP1 = mul(float4(CP1WP, 1.0f), LWCHackToFloat(PrimaryView.WorldToClip));
+
+		#if CULL_ENABLE
+		if (IsViewportVisible(SP0, SP1))
+		#endif
+		{
+			SP0 = float4(NDCToPixelCoord(SP0), 1.0f / SP0.w);
+			SP1 = float4(NDCToPixelCoord(SP1), 1.0f / SP1.w);
+
+			Point0 = CP0;
+			Point0.Position = CP0WP;
+
+			DeviceZ = SP0.z;
+			Coord = SP0.xy;
+			bIsPrimVisibile = true;
+			InterlockedAdd(s_ValidSegmentCount, 1, PrimWriteOffset);
+		}
+	}	
+	GroupMemoryBarrierWithGroupSync();
+
+	// Alloc
+	if (GroupThreadID == 0)
+	{
+		InterlockedAdd(RWCounter[uint2(0,0)], s_ValidSegmentCount, s_WriteOffset);
+	}
+	GroupMemoryBarrierWithGroupSync();
+
+	// Write valid segments
+	if (bIsPrimVisibile)
+	{
+		FHairVis Out;
+		Out.Depth			= DeviceZ;
+		Out.Coverage		= 1;
+		Out.Coverage8bit	= 255;
+		Out.PrimitiveId		= PrimID;
+		Out.MaterialId		= HairMaterialId;
+
+		uint PackedW = f32tof16(Point0.WorldRadius);
+		PackedW = PackedW | (uint(saturate(Point0.UCoord) * 0xFF)<<16);
+		PackedW = PackedW | (Point0.Type<<24);
+
+		// TODO: for now, no culling, at all, write all segments
+		const uint WriteIndex = s_WriteOffset + PrimWriteOffset;
+		//const uint WriteIndex = PrimID;
+		RWHairVis[WriteIndex] = PackHairVis(Out);
+		RWCoord[WriteIndex]   = Coord;
+		RWPoints[WriteIndex] = float4(Point0.Position, asfloat(PackedW));
+	}
+}
+#endif // CULLING_TYPE == CULLING_CURVE
+
+#endif // SHADER_RASTERCOMPUTE_CULL
\ No newline at end of file
diff --git a/Engine/Shaders/Private/HairStrands/HairStrandsForwardRaster.usf b/Engine/Shaders/Private/HairStrands/HairStrandsForwardRaster.usf
new file mode 100644
--- /dev/null
+++ b/Engine/Shaders/Private/HairStrands/HairStrandsForwardRaster.usf
@@ -0,0 +1,1810 @@
+// Copyright Epic Games, Inc. All Rights Reserved.
+
+#define HAIR_STRANDS_PARAMETERS 0
+
+#include "../Common.ush"
+#include "../WaveOpUtil.ush"
+#include "HairStrandsClusterCommon.ush"
+#include "HairStrandsVertexFactoryCommon.ush"
+#include "HairStrandsVisibilityCommon.ush"
+#include "../ColorMap.ush"
+
+#if PERMUTATION_DEBUG
+#include "../ShaderPrint.ush"
+#endif
+
+////////////////////////////////////////////////////////////
+// Pack/unpack helpers
+uint PackTileCoord(uint2 In)
+{
+	return (In.x & 0xffff) | ((In.y & 0xffff) << 16);
+}
+
+uint2 UnpackTileCoord(uint In)
+{
+	return uint2(In & 0xffff, (In >> 16) & 0xffff);
+}
+
+uint PackR11G11B10F(float3 rgb)
+{
+	uint r = (f32tof16(rgb.r) << 17) & 0xFFE00000;
+	uint g = (f32tof16(rgb.g) << 6) & 0x001FFC00;
+	uint b = (f32tof16(rgb.b) >> 5) & 0x000003FF;
+	return r | g | b;
+}
+
+float3 UnpackR11G11B10F(uint rgb)
+{
+	float r = f16tof32((rgb >> 17) & 0x7FF0);
+	float g = f16tof32((rgb >> 6) & 0x7FF0);
+	float b = f16tof32((rgb << 5) & 0x7FE0);
+	return float3(r, g, b);
+}
+
+uint PackDepth(float In)
+{
+	return asuint(In);
+}
+
+float UnpackDepth(uint In)
+{
+	return asfloat(In);
+}
+
+uint PackDepthRange(float2 In)
+{
+	return f32tof16(In.x) | (f32tof16(In.y) << 16);
+}
+
+float2 UnpackDepthRange(uint In)
+{
+	return float2(f16tof32(In & 0xffff), f16tof32(In >> 16));
+}
+
+///////////////////////////////////////////////////////////////////////////
+// Tile Helpers
+
+// Max number of iterator that a rasterizer can do. This is for preventing any kind of infinite loop.
+#define MAX_WORK_COUNT 256
+
+#define BIN_TILE_SIZE 32
+#define BIN_RCP_TILE_SIZE (1.f / BIN_TILE_SIZE)
+#define BIN_TILE_SIZE_AS_SHIFT 5
+#define BIN_THREAD_COUNT 1024
+
+#define RASTER_TILE_SIZE 8
+#define RASTER_RCP_TILE_SIZE (1.f / RASTER_TILE_SIZE)
+#define RASTER_TILE_SIZE_AS_SHIFT 3
+#define RASTER_THREAD_COUNT 64
+
+uint2 LinearTo2D_Common(uint In, uint InTileSize, float InRcpTileSize)
+{
+	uint2 Out;
+#if 0
+	Out.y = (In + 0.5f) * InRcpTileSize;
+	Out.x = In - (Out.y * InTileSize);
+#else
+	Out.x = In%InTileSize;
+	Out.y = In/InTileSize;
+#endif
+	return Out;
+}
+
+uint2 LinearTo2D_Bin(uint In)	 { return LinearTo2D_Common(In, BIN_TILE_SIZE, BIN_RCP_TILE_SIZE); }
+uint2 LinearTo2D_Raster(uint In) { return LinearTo2D_Common(In, RASTER_TILE_SIZE, RASTER_RCP_TILE_SIZE); }
+
+///////////////////////////////////////////////////////////////////////////
+// DDA helper
+ 
+// TODO: Without setting a limit, we are sometimes caught in an infinite loop even though this should not happen.
+#define DDA_MAX_ITERATIONS 256
+
+struct FDDAContext
+{
+	float2 Coord;
+	float2 DeltaDist;
+	float2 Step;
+	float2 SideDist;
+};
+
+FDDAContext DDACreateContext(float2 RayStart, float2 RayDir)
+{
+	const float2 RayDirRcp = 1.0f / RayDir;
+
+	FDDAContext Context;
+	Context.Coord = floor(RayStart);
+	Context.DeltaDist = abs(RayDirRcp);
+	Context.Step = sign(RayDir);
+	Context.SideDist = (Context.Coord - RayStart + max(Context.Step, 0.0f)) * RayDirRcp;
+
+	return Context;
+}
+
+void DDAAdvance(inout FDDAContext Context)
+{
+	if (Context.SideDist.x < Context.SideDist.y)
+	{
+		Context.SideDist.x += Context.DeltaDist.x;
+		Context.Coord.x += Context.Step.x;
+	}
+	else
+	{
+		Context.SideDist.y += Context.DeltaDist.y;
+		Context.Coord.y += Context.Step.y;
+	}
+}
+
+///////////////////////////////////////////////////////////////////////////
+// Visibility Tile data
+/*
+// use untyped buffer for segment tiles to reduce VGPR usage - 16 bytes
+struct FVisTile
+{
+	uint PrimOffset;
+	uint PrimCount;
+	uint TileCoord;
+	uint MinDepth;
+};
+*/
+#define VT_PrimOffset 0
+#define VT_PrimCount 1
+#define VT_Coord 2
+#define VT_MinWriteIndex 3
+#define VT_MinMaxDepth 4
+#define VT_SIZE 5
+
+// Visibility tile data are stored as:
+//   _______________________________________________________________________________________________________________________________________________________________________
+// ||                      Tile 0                          ||                      Tile 1                          ||                      Tile 2                          ||
+// ||______________________________________________________||______________________________________________________||______________________________________________________||
+// ||            |           |             |               ||            |           |             |               ||            |           |             |               ||
+// || PrimOffset | PrimCount | (Tile)Coord | MinWriteIndex || PrimOffset | PrimCount | (Tile)Coord | MinWriteIndex || PrimOffset | PrimCount | (Tile)Coord | MinWriteIndex ||
+
+uint PackVisTileCoord(uint2 In)
+{
+	return (In.x & 0xff) | ((In.y & 0xff) << 8);
+}
+
+uint2 UnpackVisTileCoord(uint In)
+{
+	return uint2(In & 0xff, (In >> 8) & 0xff);
+}
+
+uint LoadOutVisTileData(RWByteAddressBuffer OutBuffer, uint Index, uint VTEntry)
+{
+	// Each entry is 4 bytes
+	return OutBuffer.Load(Index * VT_SIZE * 4 + VTEntry * 4);
+}
+
+void StoreOutVisTileData(RWByteAddressBuffer OutBuffer, uint Index, uint VTEntry, uint Value)
+{
+	// Each entry is 4 bytes
+	OutBuffer.Store(Index * VT_SIZE * 4 + VTEntry * 4, Value);
+}
+
+uint LoadVisTileData(ByteAddressBuffer InBuffer, uint Index, uint VTEntry)
+{
+	return InBuffer.Load(Index * VT_SIZE * 4 + VTEntry * 4);
+}
+
+///////////////////////////////////////////////////////////////////////////
+// Misc.
+float4 Transparent(float4 Color) { return float4(Color.xyz, 0.5f); }
+
+///////////////////////////////////////////////////////////////////////////
+// Common parameters
+
+int2	BinTileRes;
+int2	RasterTileRes;
+
+uint	NumBinners;
+float	RcpNumBinners;
+
+uint	NumRasterizers;
+float	RcpNumRasterizers;
+
+int2	OutputResolution;
+float2	OutputResolutionf;
+
+///////////////////////////////////////////////////////////////////////////
+// Control points
+Buffer<float4> ControlPoints;
+Texture2D<uint> ControlPointCount;
+uint MaxControlPointCount;
+
+uint GetControlPointCount()
+{
+	return ControlPointCount.Load(uint3(0,0,0));
+}
+
+FHairControlPoint UnpackHairControlPoint(uint InPrimId)
+{
+	const float4 Packed = ControlPoints[InPrimId];
+	const uint W = asuint(Packed.w);
+	const float R = f16tof32(W & 0xFFFF);
+	const float U = ((W >> 16) & 0xFF) * (1.f / 255.f);
+	const uint  T = (W >> 24) & 0x3;
+
+	FHairControlPoint Out;
+	Out.Position = Packed.xyz;
+	Out.WorldRadius = R;
+	Out.UCoord = U;
+	Out.Type = T;
+	return Out;
+}
+
+#if SHADER_RASTERCOMPUTE_BINNING || SHADER_RASTERCOMPUTE_COMPACTION || SHADER_RASTERCOMPUTE_RASTER || SHADER_RASTERCOMPUTE_DEPTH_GRID || SHADER_RASTERCOMPUTE_DEBUG
+
+///////////////////////////////////////////////////////////////////////////
+
+float3 NDCToPixelCoord(float4 InDC)
+{
+	const float3 NDC = InDC.xyz / InDC.w;
+	float2 UV = NDC.xy * ResolvedView.ScreenPositionScaleBias.xy + ResolvedView.ScreenPositionScaleBias.wz;
+	return float3(UV * OutputResolution, NDC.z);
+}
+
+void CalcHomogenousPos(in uint PrimId, out float4 HP, out uint Type)
+{
+	const FHairControlPoint CP = UnpackHairControlPoint(PrimId);
+
+	const float3 WP = CP.Position; // This is actually WorldPosition
+	HP = mul(float4(WP, 1.0f), LWCHackToFloat(PrimaryView.WorldToClip)); // TODO move this at least into translated world space
+	Type = CP.Type;
+}
+
+void CalcHomogenousPosAndRad(in uint PrimId, out float4 HP, out float Rad, out uint Type)
+{
+	const FHairControlPoint CP = UnpackHairControlPoint(PrimId);
+
+	const float3 WP = CP.Position; // This is actually WorldPosition
+	HP = mul(float4(WP, 1.0f), LWCHackToFloat(PrimaryView.WorldToClip));
+	Rad = CP.WorldRadius * 2000.0; // OutputResolutionf.x; //TODO: figure this out correctly?
+	Type = CP.Type;
+}
+
+float ComputeLerpAlpha(int2 Coord, float2 P0, float2 P1, float SegmentLenSqRcp)
+{
+	// Project P onto line segment and compute the lerp alpha between P0 and P1
+	// Simplification of:
+	// A = P - P0
+	// B = P1 - P0
+	// Alpha = dot(A, B) / dot(B, B)
+	const float2 P = Coord + 0.5f;
+	const float Alpha = saturate(dot(P - P0, P1 - P0) * SegmentLenSqRcp);
+	return Alpha;
+}
+
+float ComputePerspectiveCorrectRadius(float Rad0, float Rad1, float Alpha, float RcpW0, float RcpW1)
+{
+	// Alpha value for perspective correct interpolation. We store the reciprocal of w in the w component of P0 and P1,
+	// so this is a simplification of:
+	// (Alpha / w1) / ((1 - Alpha) / w0 + Alpha / w1)
+	const float LerpedRcpW = lerp(RcpW0, RcpW1, Alpha);
+	const float PerspectiveAlpha = (Alpha * RcpW1) / LerpedRcpW;
+	// Divide by W to make thickness dependent on screen space depth? This division was kept from the previous line rasterization algorithm.
+	const float Rad = lerp(Rad0, Rad1, PerspectiveAlpha) * LerpedRcpW;
+	return Rad;
+}
+
+// Line clipping based on "CLIPPING USING HOMOGENEOUS COORDINATES" by Blinn et al.
+bool BlinnLineClipping(inout float4 P0, inout float4 P1)
+{
+	float2 T = float2(0.0f, 1.0f);
+	bool bIsRemoved = P0.w < 0.0f && P1.w < 0.0f; // Both points behind near plane
+
+	bool bSign = false;
+
+	UNROLL
+	for (uint PlaneIdx = 0; PlaneIdx < 6; ++PlaneIdx)
+	{
+		// Compute boundary coordinates of both points (w+x, w-x, w+y, w-y, z, w-z)
+		bSign = !bSign;
+		const uint CompIdx = PlaneIdx / 2;
+		const float Sign = bSign ? 1.0f : -1.0f;
+		const float WFactor = PlaneIdx != 4 ? 1.0f : 0.0f;
+		const float2 BC = WFactor * float2(P0.w, P1.w) + Sign * float2(P0[CompIdx], P1[CompIdx]);
+
+		float Num = BC.x;
+		float Denom = BC.x - BC.y;
+		bIsRemoved = bIsRemoved || (BC.x < 0.0f && BC.y < 0.0f); // Both points outside the plane
+		float Alpha = Num / Denom;
+		
+		// If the denominator is negative, P0 has a smaller boundary coordinate than P1, so we can assume
+		// that P1 is inside the plane (or bIsRemoved is true), so we need to update the alpha for P0.
+		// The reverse is true if the denominator is positive.
+		if (Denom < 0.0f)
+		{
+			T.x = max(T.x, Alpha);
+		}
+		else
+		{
+			T.y = min(T.y, Alpha);
+		}
+	}
+
+	if (!bIsRemoved)
+	{
+		const float4 P0Clipped = lerp(P0, P1, T.x);
+		const float4 P1Clipped = lerp(P0, P1, T.y);
+		P0 = P0Clipped;
+		P1 = P1Clipped;
+	}
+
+	return !bIsRemoved;
+}
+
+bool ClipRaySegment(float2 AABBMin, float2 AABBMax, float4 P0, float4 P1, out float2 T, out bool2 bClipped)
+{
+	bClipped = false;
+	T = float2(0.0f, 1.0f);
+	const bool bP0Outside = any(P0.xy < AABBMin) || any(P0.xy > AABBMax);
+	const bool bP1Outside = any(P1.xy < AABBMin) || any(P1.xy > AABBMax);
+	if (!bP0Outside && !bP1Outside)
+	{
+		return true;
+	}
+
+	const float2 Origin = P0.xy;
+	const float2 Dir = P1.xy - P0.xy;
+	const float2 RcpDir = 1.0f / Dir;
+
+	const float2 T0 = (AABBMin - Origin) * RcpDir;
+	const float2 T1 = (AABBMax - Origin) * RcpDir;
+
+	T.x = max(min(T0.x, T1.x), min(T0.y, T1.y));
+	T.y = min(max(T0.x, T1.x), max(T0.y, T1.y));
+
+	// Ray intersects the AABB but the segment is completely outside or no intersection at all.
+	if (T.y < 0.0f || T.x > T.y || T.x > 1.f)
+	{
+		bClipped = true;
+		return false;
+	}
+
+	if (bP0Outside && T.x > 0.0f && T.x < 1.0f)
+	{
+		bClipped.x = true;
+	}
+	if (bP1Outside && T.y > 0.0f && T.y < 1.0f)
+	{
+		bClipped.y = true;
+	}
+
+	return true;
+}
+
+bool ClipRaySegment(float2 AABBMin, float2 AABBMax, inout float4 P0, inout float4 P1, inout float Rad0, inout float Rad1, inout float Alpha0, inout float Alpha1, out bool2 bClipped, out float2 T)
+{
+	//float2 T;
+	bool bIsValid = ClipRaySegment(AABBMin, AABBMax, P0, P1, T, bClipped);
+	
+	if (bIsValid)
+	{
+		const bool bP0Outside = any(P0.xy < AABBMin) || any(P0.xy > AABBMax);
+		const bool bP1Outside = any(P1.xy < AABBMin) || any(P1.xy > AABBMax);
+
+		float4 P0New = P0;
+		float4 P1New = P1;
+		float Rad0New = Rad0;
+		float Rad1New = Rad1;
+		if (bP0Outside && T.x > 0.0f && T.x < 1.0f)
+		{
+			Alpha0 = T.x;
+			P0New = lerp(P0, P1, T.x);
+			Rad0New = lerp(Rad0, Rad1, T.x);
+			bClipped.x = true;
+		}
+		if (bP1Outside && T.y > 0.0f && T.y < 1.0f)
+		{
+			Alpha1 = T.y;
+			P1New = lerp(P0, P1, T.y);
+			Rad1New = lerp(Rad0, Rad1, T.y);
+			bClipped.y = true;
+		}
+		P0 = P0New;
+		P1 = P1New;
+		Rad0 = Rad0New;
+		Rad1 = Rad1New;
+	}
+
+	return bIsValid;
+}
+
+#endif // Common rasetrizer helper function & parameters
+
+///////////////////////////////////////////////////////////////////////////
+
+#if SHADER_RASTERCOMPUTE_DEPTH_GRID
+
+Texture2D<float> SceneDepthTexture;
+RWTexture2D<uint> OutVisTileDepthGrid;
+groupshared uint group_FurthestDepth; // (4 bytes)
+
+[numthreads(BIN_THREAD_COUNT, 1, 1)]
+void PrepareDepthGridCS(uint DispatchThreadID : SV_DispatchThreadID, uint GroupThreadID : SV_GroupThreadID, uint2 GroupID : SV_GroupID)
+{
+	if (GroupThreadID == 0)
+	{
+		group_FurthestDepth = 0xFFFFFFFF; // Inverse-Z
+	}
+
+	GroupMemoryBarrierWithGroupSync();
+
+	// Assumes a workgroup size of 1024 threads to cover a maximum tile size of 32x32.
+	if (GroupThreadID < BIN_THREAD_COUNT)
+	{
+		const uint2 PixelCoord = LinearTo2D_Bin(GroupThreadID) + GroupID * BIN_TILE_SIZE;
+
+		if (all(PixelCoord < (uint2)OutputResolution))
+		{
+			const float Depth = SceneDepthTexture.Load(uint3(PixelCoord, 0));
+
+			// Compute furthest depth inside this tile
+			WaveInterlockedMin(group_FurthestDepth, PackDepth(Depth)); // Inverse-Z
+		}
+	}
+
+	GroupMemoryBarrierWithGroupSync();
+
+	if (GroupThreadID == 0)
+	{
+		OutVisTileDepthGrid[GroupID] = group_FurthestDepth;
+	}
+}
+
+#endif //SHADER_RASTERCOMPUTE_DEPTH_GRID
+
+///////////////////////////////////////////////////////////////////////////
+
+#define BIN_MINMAX 0
+
+#if SHADER_RASTERCOMPUTE_BINNING
+
+RWTexture2DArray<uint> 				OutVisTileBinningGrid;
+#if BIN_MINMAX
+RWTexture2DArray<uint> 				OutVisTileBinningGridMin;
+RWTexture2DArray<uint> 				OutVisTileBinningGridMax;
+#endif
+RWBuffer<uint>						OutVisTilePrims;
+RWBuffer<uint>						OutVisTileArgs;
+RWByteAddressBuffer					OutVisTileData;
+Texture2D<uint>						VisTileDepthGrid;
+ByteAddressBuffer					IndirectPrimIDCount;
+
+groupshared uint group_LoopNum;
+groupshared uint group_VerticesNum;
+groupshared uint group_BatchNum;
+
+#define TILES_TO_ALLOCATE_MAX 1024
+
+groupshared uint group_TilesToAllocate[TILES_TO_ALLOCATE_MAX];
+groupshared uint group_TilesToAllocateCount;
+
+// The total number of line segments (ControlPointCount) is divided up equally between N binners - each binner = a workgroup which loops through the designated set segments in batches of 1024
+// NB there is still potential to use LDS to prevent/reduce contention when adding segments to the binning grid and this may improve perf
+
+[numthreads(1024, 1, 1)]
+void BinningCS(uint GroupThreadID : SV_GroupThreadID, uint GroupID : SV_GroupID)
+{
+	ResolvedView = ResolveView();
+	if (GroupThreadID == 0)
+	{
+		group_VerticesNum = GetControlPointCount();
+		group_BatchNum = DivideAndRoundUp(group_VerticesNum, 1024);
+		group_LoopNum  = DivideAndRoundUp(group_BatchNum, NumBinners);
+	}
+	GroupMemoryBarrierWithGroupSync();
+
+	#if PERMUTATION_DEBUG
+	const bool bDebugEnabled = false; //GroupID == 0 && GroupThreadID <= 64;
+	FShaderPrintContext Ctx = InitShaderPrintContext(bDebugEnabled, uint2(250, 50));
+	#endif
+
+	LOOP for (uint LoopIndex = 0; LoopIndex < group_LoopNum; LoopIndex++)
+	{
+		const uint BatchIndex = LoopIndex + (GroupID * group_LoopNum);
+		bool bSegValid = (BatchIndex < group_BatchNum);
+
+		const uint PrimID = BatchIndex * 1024 + GroupThreadID;
+		bSegValid = bSegValid && (PrimID < group_VerticesNum);
+
+		const uint SegmentCountLayerIdx = GroupID; // Stores number of segments per tile per workgroup.
+		const uint TmpSegmentCountLayerIdx = SegmentCountLayerIdx + NumBinners; // Also stores number of segments per tile per workgroup. Used as second counter for this two pass algorithm.
+		const uint TileAllocInfoLayerIdx = SegmentCountLayerIdx + NumBinners * 2; // Stores per tile per workgroup allocation info.
+
+		uint NearestDepth  = 0;
+		uint FurthestDepth = 1;
+		float2 TileCoord0F = 0.0f;
+		float2 TileCoord1F = 0.0f;
+
+		#if PERMUTATION_DEBUG
+		FHairControlPoint CP0;
+		FHairControlPoint CP1;
+		#endif
+
+		// 1. Project segment end points and clip them to the screen
+		if (bSegValid)
+		{
+			float4 H0 = 0.0f;
+			float4 H1 = 0.0f;
+			uint Type = -1;
+			CalcHomogenousPos(PrimID, H0, Type);
+
+			bool bIsEndCV = (Type == HAIR_CONTROLPOINT_END);
+			bSegValid = !bIsEndCV;
+
+			if (bSegValid)
+			{
+				CalcHomogenousPos(PrimID + 1, H1, Type);
+				
+				// Do clipping in homogenous coordinates
+				bSegValid = BlinnLineClipping(H0, H1);
+
+				if (bSegValid)
+				{
+					float3 SP0 = NDCToPixelCoord(H0);
+					float3 SP1 = NDCToPixelCoord(H1);
+					SP0.xy *= BIN_RCP_TILE_SIZE;
+					SP1.xy *= BIN_RCP_TILE_SIZE;
+
+					// For peace of mind, make sure these are actually clamped to a valid range.
+					SP0 = clamp(SP0, 0.0f, float3(BinTileRes, 1.0f));
+					SP1 = clamp(SP1, 0.0f, float3(BinTileRes, 1.0f));
+
+					NearestDepth  = PackDepth(max(SP0.z, SP1.z)); // Inverse-Z
+					FurthestDepth = PackDepth(min(SP0.z, SP1.z)); // Inverse-Z
+
+					TileCoord0F = SP0.xy;
+					TileCoord1F = SP1.xy;
+
+					#if PERMUTATION_DEBUG
+					if (bDebugEnabled && 0)
+					{
+						CP0 = UnpackHairControlPoint(PrimID);
+						CP1 = UnpackHairControlPoint(PrimID +1);
+
+						AddLineWS(Ctx, CP0.Position, CP1.Position, ColorRed);
+					}
+					#endif
+				}
+			}
+		}
+
+		// 2. Reset allocation counter
+		if (GroupThreadID == 0)
+		{
+			group_TilesToAllocateCount = 0;
+		}
+		GroupMemoryBarrierWithGroupSync();
+		
+		// 3. Increment per workgroup per tile counters and add tiles to be allocated
+		if (bSegValid)
+		{
+			FDDAContext DDAContext = DDACreateContext(TileCoord0F, normalize(TileCoord1F - TileCoord0F));
+			const int2 EndCoord = (int2)floor(TileCoord1F);
+
+			for (int DDAIt = 0; DDAIt < DDA_MAX_ITERATIONS; ++DDAIt)
+			{
+				const int2 TileCoord = (int2)floor(DDAContext.Coord);
+				
+				uint DebugInsertMode = 0;
+				BRANCH
+				if (NearestDepth > VisTileDepthGrid[TileCoord])
+				{
+					uint OldTileSegmentCount;
+					InterlockedAdd(OutVisTileBinningGrid[uint3(TileCoord, SegmentCountLayerIdx)], 1, OldTileSegmentCount);
+					DebugInsertMode = 1;
+
+					// Min/Max
+					#if BIN_MINMAX
+					InterlockedMin(OutVisTileBinningGridMin[uint3(TileCoord, SegmentCountLayerIdx)], FurthestDepth);
+					InterlockedMax(OutVisTileBinningGridMax[uint3(TileCoord, SegmentCountLayerIdx)], NearestDepth);
+					#endif
+
+					BRANCH
+					if ((OldTileSegmentCount % 1024) == 0)
+					{
+						uint WritePos;
+						InterlockedAdd(group_TilesToAllocateCount, 1, WritePos);
+						if (WritePos < TILES_TO_ALLOCATE_MAX)
+						{
+							group_TilesToAllocate[WritePos] = PackVisTileCoord(TileCoord);
+							DebugInsertMode = 2;
+						}
+					}
+				}
+
+				#if PERMUTATION_DEBUG
+				if (bDebugEnabled)
+				{
+					//CP0 = UnpackHairControlPoint(PrimID);
+					//CP1 = UnpackHairControlPoint(PrimID +1);
+
+					//AddLineWS(Ctx, CP0.Position, CP1.Position, ColorRed);
+
+					float4 DebugColor = ColorRed;
+					if (DebugInsertMode == 1) DebugColor = ColorGreen;
+					if (DebugInsertMode == 2) DebugColor = ColorYellow;
+					AddQuadSS(Ctx, TileCoord * BIN_TILE_SIZE, TileCoord * BIN_TILE_SIZE + BIN_TILE_SIZE, DebugColor);
+				}
+				#endif
+
+				if (all(TileCoord == EndCoord))
+				{
+					break;
+				}
+
+				DDAAdvance(DDAContext);
+			}
+		}
+
+		GroupMemoryBarrierWithGroupSync();
+
+		// 4. Allocate tiles
+		const uint TilesToAllocateCount = min(TILES_TO_ALLOCATE_MAX, group_TilesToAllocateCount);
+		for (uint TileIdx = GroupThreadID; TileIdx < TilesToAllocateCount; TileIdx += 1024)
+		{
+			const uint PackedTileCoord = group_TilesToAllocate[TileIdx];
+			const uint2 TileCoord = UnpackVisTileCoord(PackedTileCoord);
+
+			const uint TotalNewWriteCount = OutVisTileBinningGrid[uint3(TileCoord, SegmentCountLayerIdx)];
+			const uint TotalOldWriteCount = OutVisTileBinningGrid[uint3(TileCoord, TmpSegmentCountLayerIdx)];
+
+			#if BIN_MINMAX
+			const float MinDepth = UnpackDepth(OutVisTileBinningGridMin[uint3(TileCoord, SegmentCountLayerIdx)]);
+			const float MaxDepth = UnpackDepth(OutVisTileBinningGridMax[uint3(TileCoord, SegmentCountLayerIdx)]);
+			#endif
+
+			uint NewTile;
+			WaveInterlockedAddScalar_(OutVisTileArgs[0], 1, NewTile);
+
+			StoreOutVisTileData(OutVisTileData, NewTile, VT_Coord, PackedTileCoord);
+			// Round down the count to the start of the tile and later compare against this to decide which tile to write to.
+			StoreOutVisTileData(OutVisTileData, NewTile, VT_MinWriteIndex, TotalNewWriteCount & ~1023u);
+			// Min/Max depth
+			#if BIN_MINMAX
+			StoreOutVisTileData(OutVisTileData, NewTile, VT_MinMaxDepth, PackDepthRange(float2(NearestDepth, FurthestDepth));
+			#endif
+
+			const uint PrevTile = (OutVisTileBinningGrid[uint3(TileCoord, TileAllocInfoLayerIdx)] & 0xffff);
+
+			if (TotalOldWriteCount > 0)
+			{
+				StoreOutVisTileData(OutVisTileData, PrevTile, VT_PrimCount, 1024);
+			}
+
+			OutVisTileBinningGrid[uint3(TileCoord, TileAllocInfoLayerIdx)] = (PrevTile << 16) | (NewTile & 0xffff);
+		}
+
+		GroupMemoryBarrierWithGroupSync();
+
+		// 5. Write PrimID to tiles
+		if (bSegValid)
+		{
+			FDDAContext DDAContext = DDACreateContext(TileCoord0F, normalize(TileCoord1F - TileCoord0F));
+			const int2 EndCoord = (int2)floor(TileCoord1F);
+
+			for (int DDAIt = 0; DDAIt < DDA_MAX_ITERATIONS; ++DDAIt)
+			{
+				const int2 TileCoord = (int2)floor(DDAContext.Coord);
+
+				BRANCH
+				if (NearestDepth > VisTileDepthGrid[TileCoord])
+				{
+					const uint PackedTiles = OutVisTileBinningGrid[uint3(TileCoord, TileAllocInfoLayerIdx)];
+					const uint CurTile = (PackedTiles & 0xffff);
+					const uint PrevTile = ((PackedTiles >> 16) & 0xffff);
+
+					// Currently we need this to get our write position, but maybe there is a cheaper way to keep track of that?
+					uint OldTileSegmentCount;
+					InterlockedAdd(OutVisTileBinningGrid[uint3(TileCoord, TmpSegmentCountLayerIdx)], 1, OldTileSegmentCount);
+
+					const bool bWriteToCurTile = OldTileSegmentCount >= LoadOutVisTileData(OutVisTileData, CurTile, VT_MinWriteIndex);
+					const uint LocalWritePos = OldTileSegmentCount % 1024;
+					const uint WritePos = (bWriteToCurTile ? CurTile : PrevTile) * 1024 + LocalWritePos;
+
+					OutVisTilePrims[WritePos] = PrimID;
+
+					BRANCH
+					if (bWriteToCurTile)
+					{
+						if ((OldTileSegmentCount + 1) == OutVisTileBinningGrid[uint3(TileCoord, SegmentCountLayerIdx)])
+						{
+							StoreOutVisTileData(OutVisTileData, CurTile, VT_PrimCount, (OldTileSegmentCount == 1023) ? 1024 : ((OldTileSegmentCount + 1) % 1024));
+						}
+					}
+				}
+
+				if (all(TileCoord == EndCoord))
+				{
+					break;
+				}
+
+				DDAAdvance(DDAContext);
+			}
+		}
+	}
+}
+#endif //SHADER_RASTERCOMPUTE_BINNING
+
+///////////////////////////////////////////////////////////////////////////
+
+#if SHADER_RASTERCOMPUTE_COMPACTION
+
+ByteAddressBuffer					InData;
+Buffer<uint> 						InPrims;
+Buffer<uint> 						InArgs;
+RWByteAddressBuffer					OutData;
+RWBuffer<uint> 						OutPrims;
+RWBuffer<uint> 						OutArgs;
+
+groupshared uint group_TotalPrimCount;
+groupshared uint group_PrimWriteOffset;
+groupshared uint group_NumTiles;
+groupshared uint group_TilesToCompact[1024];
+groupshared uint group_MaxLDSTileIdx;
+
+[numthreads(1024, 1, 1)]
+void CompactionCS(uint DispatchThreadID : SV_DispatchThreadID, uint GroupThreadID : SV_GroupThreadID, uint2 GroupID : SV_GroupID)
+{
+	if (GroupThreadID == 0)
+	{
+		group_TotalPrimCount = 0;
+		group_NumTiles = 0;
+		group_MaxLDSTileIdx = 0;
+	}
+	
+	GroupMemoryBarrierWithGroupSync();
+
+	const uint NumTiles = InArgs[0];
+	const uint PackedCoord = PackVisTileCoord(GroupID);
+
+	// 1. Compute total number of primitives at this tile coordinate
+	uint LocalPrimCount = 0;
+	{
+		for (uint TileIdx = GroupThreadID; TileIdx < NumTiles; TileIdx += 1024)
+		{
+			const uint TilePackedCoord = LoadVisTileData(InData, TileIdx, VT_Coord);
+			if (PackedCoord == TilePackedCoord)
+			{
+				LocalPrimCount += LoadVisTileData(InData, TileIdx, VT_PrimCount);
+				
+				uint WritePos;
+				WaveInterlockedAddScalar_(group_NumTiles, 1, WritePos);
+				if (WritePos < 1024)
+				{
+					group_TilesToCompact[WritePos] = TileIdx;
+					WaveInterlockedMax(group_MaxLDSTileIdx, TileIdx);
+				}
+			}
+		}
+	}
+	GroupMemoryBarrierWithGroupSync();
+
+	if (LocalPrimCount > 0)
+	{
+		WaveInterlockedAdd(group_TotalPrimCount, LocalPrimCount);
+	}
+	GroupMemoryBarrierWithGroupSync();
+
+	const uint TotalPrimCount = group_TotalPrimCount;
+
+	if (TotalPrimCount == 0)
+	{
+		return;
+	}
+
+	// 2. Allocate space
+	if (GroupThreadID == 0)
+	{
+		const uint NumTilesToAllocate = DivideAndRoundUp(TotalPrimCount, 1024);
+
+		uint FirstCompactedTile;
+		InterlockedAdd(OutArgs[0], NumTilesToAllocate, FirstCompactedTile);
+
+		group_PrimWriteOffset = FirstCompactedTile * 1024;
+
+		// Initialize new tiles
+		for (uint TileIdx = 0; TileIdx < NumTilesToAllocate; ++TileIdx)
+		{
+			const uint CompactedTile = FirstCompactedTile + TileIdx;
+
+			const uint PrimCount = min(TotalPrimCount - TileIdx * 1024, 1024);
+			StoreOutVisTileData(OutData, CompactedTile, VT_PrimCount, PrimCount);
+			StoreOutVisTileData(OutData, CompactedTile, VT_Coord, PackedCoord);
+		}
+	}
+	GroupMemoryBarrierWithGroupSync();
+
+	// 3. Copy PrimIDs to compacted memory
+	{
+		uint CurrentWriteOffset = group_PrimWriteOffset;
+
+		// 3.1 First process the LDS list of tiles
+		const uint NumInputTiles = min(group_NumTiles, 1024);
+		for (uint LDSIdx = 0; LDSIdx < NumInputTiles; ++LDSIdx)
+		{
+			const uint TileIdx = group_TilesToCompact[LDSIdx];
+
+			const uint TilePrimOffset = TileIdx * 1024;
+			const uint TilePrimCount = LoadVisTileData(InData, TileIdx, VT_PrimCount);
+
+			if (GroupThreadID < TilePrimCount)
+			{
+				OutPrims[CurrentWriteOffset + GroupThreadID] = InPrims[TilePrimOffset + GroupThreadID];
+			}
+
+			CurrentWriteOffset += TilePrimCount;
+		}
+
+		// 3.2 Check any remaning tiles
+		if (group_NumTiles > 1024)
+		{
+			for (uint TileIdx = group_MaxLDSTileIdx + 1; TileIdx < NumTiles; ++TileIdx)
+			{
+				const uint TilePackedCoord = LoadVisTileData(InData, TileIdx, VT_Coord);
+				if (PackedCoord == TilePackedCoord)
+				{
+					const uint TilePrimOffset = TileIdx * 1024;
+					const uint TilePrimCount = LoadVisTileData(InData, TileIdx, VT_PrimCount);
+
+					if (GroupThreadID < TilePrimCount)
+					{
+						OutPrims[CurrentWriteOffset + GroupThreadID] = InPrims[TilePrimOffset + GroupThreadID];
+					}
+
+					CurrentWriteOffset += TilePrimCount;
+				}
+			}
+		}
+	}
+}
+
+#endif // SHADER_RASTERCOMPUTE_COMPACTION
+
+///////////////////////////////////////////////////////////////////////////
+
+#define RASTER_DEPTH_BUCKET 32
+#define RASTER_SEGMENT_COUNT 32
+#define SEGMENT_COUNT_PER_GROUP 1024
+
+// For editing convenience
+#if !SHADER_RASTERCOMPUTE_DEBUG && !SHADER_RASTERCOMPUTE_COMPACTION && !SHADER_RASTERCOMPUTE_BINNING && !SHADER_RASTERCOMPUTE_RASTER_MULTI_SAMPLE && !SHADER_RASTERCOMPUTE_DEPTH_GRID
+#define SHADER_RASTERCOMPUTE_RASTER 1
+#endif
+
+#if SHADER_RASTERCOMPUTE_RASTER
+
+Buffer<uint> 				VisTilePrims;
+Buffer<uint> 				VisTileArgs;
+ByteAddressBuffer			VisTileData;
+RWTexture2D<float4>			OutSceneColorTexture;
+RWStructuredBuffer<uint>	RWWorkCounter;
+
+#if PERMUTATION_DEBUG
+RWTexture2D<uint>			OutHairCountTexture_ForDebug;
+RWTexture2D<uint>			OutHairPixelCountPerTile_ForDebug;
+#endif
+
+int2						SampleLightingViewportResolution;
+Texture2D<float4>			SampleLightingTexture;
+
+// Depth|ID
+groupshared uint2  s_Segments[1024];
+groupshared uint2  s_Segments_Sorted[1024];
+groupshared uint   s_Segments_Min;
+groupshared uint   s_Segments_Max;
+groupshared uint   s_SegmentsCount[RASTER_DEPTH_BUCKET];
+groupshared uint   s_SegmentsAlloc[RASTER_DEPTH_BUCKET];
+
+//groupshared uint   s_Mask[RASTER_THREAD_COUNT];
+groupshared uint   s_Mask[8][8];
+groupshared uint   s_OpaqueMask;
+groupshared uint   s_Data[RASTER_SEGMENT_COUNT][8];
+groupshared uint   s_Color[RASTER_SEGMENT_COUNT][8];
+groupshared uint   s_bDataOrder;
+groupshared uint   s_WorkID;
+
+#if PERMUTATION_DEBUG
+groupshared float s_Coverage[RASTER_THREAD_COUNT];
+#endif
+
+void ClearMask(uint2 In)
+{
+	s_Mask[In.x][In.y] = 0;
+	//s_Mask[In.x + In.y * RASTER_TILE_SIZE] = 0;
+}
+
+uint ReadMask(uint2 In)
+{
+	return s_Mask[In.x][In.y];
+	//return s_Mask[In.x + In.y * RASTER_TILE_SIZE];
+}
+
+void WriteMask(uint2 In, uint InValue)
+{
+	InterlockedOr(s_Mask[In.x][In.y], InValue);
+	//InterlockedOr(s_Mask[In.x + In.y * RASTER_TILE_SIZE], InValue);
+}
+
+#if PERMUTATION_DEBUG
+void PrintCoverage(inout FShaderPrintContext Ctx)
+{
+	Print(Ctx, TEXT("Coverage"), FontWhite);
+	Newline(Ctx);
+	for (uint y = 0; y < RASTER_TILE_SIZE; ++y)
+	{
+		for (uint x = 0; x < RASTER_TILE_SIZE; ++x)
+		{
+			const uint ValidSegments = ReadMask(uint2(x,y)); //s_Mask[x][y];
+			if (ValidSegments > 0)
+			{
+				Print(Ctx, TEXT("x "), FontGreen);
+			}
+			else
+			{
+				Print(Ctx, TEXT(". "), FontRed);
+			}
+		}
+		Newline(Ctx);
+	}
+}
+
+void PrintPixelMask(inout FShaderPrintContext Ctx, uint InMask)
+{
+	for (uint s = 0; s < RASTER_SEGMENT_COUNT; ++s)
+	{
+		const bool bValid = ((1u<<s) & InMask) != 0u;
+		if (bValid)
+		{
+			Print(Ctx, TEXT("x"), FontGreen);
+		}
+		else
+		{
+
+
+			Print(Ctx, TEXT("."), FontRed);
+		}		
+	}
+}
+#endif
+
+struct FSampleData
+{
+	uint2 Coord; // Pixel coord within the tile 8x8 (4bits x 4bits)
+	float Depth;
+	float Coverage;
+};
+
+uint PackSampleData(FSampleData In)
+{
+	return f32tof16(In.Depth) | (uint(saturate(In.Coverage) * 0xFF)<<16u) | ((In.Coord.x & 0xF)<<24u) | ((In.Coord.y & 0xF) << 28u);
+}
+
+FSampleData UnpackSampleData(uint In)
+{
+	FSampleData Out;
+	Out.Depth = f16tof32(In & 0xFFFF);
+	Out.Coverage = ((In>>16) & 0xFF) * (1.f/255.f);
+	Out.Coord =  uint2((In>>24) & 0xF, (In>>28) & 0xF);
+	return Out;
+}
+
+uint PackColorData(float3 In)
+{
+	return PackR11G11B10F(In);
+}
+
+float3 UnpackColorData(uint In)
+{
+	return UnpackR11G11B10F(In);
+}
+
+float3 LoadSampleColor(uint InPrimId)
+{
+	uint2 SampleCoord = 0;
+	SampleCoord.x = InPrimId % SampleLightingViewportResolution.x;
+	SampleCoord.y = InPrimId / SampleLightingViewportResolution.x;
+	return SampleLightingTexture.Load(uint3(SampleCoord, 0)).xyz;
+}
+
+uint GetDepthBinIndex(float InDepth, float InvDepthExtent)
+{
+	// Inverse-Z
+	const float MinDepth = UnpackDepth(s_Segments_Min);
+	const float MaxDepth = UnpackDepth(s_Segments_Max);
+	const uint DepthIt = min(saturate((InDepth - MinDepth) * InvDepthExtent) * RASTER_DEPTH_BUCKET, RASTER_DEPTH_BUCKET - 1);
+	return (RASTER_DEPTH_BUCKET-1) - DepthIt;
+}
+
+[numthreads(RASTER_TILE_SIZE, RASTER_TILE_SIZE, 1)]
+void RasterCS(
+	uint  GroupThread1D    : SV_GroupIndex,		/* 64  */
+	uint2 GroupThread2D    : SV_GroupThreadID,  /* 8x8 */
+	uint  GroupID          : SV_GroupID)		/* Rasterizer ID */
+{
+	ResolvedView = ResolveView();
+	const uint BinTileNum = VisTileArgs[0];
+
+	LOOP
+	for (uint WorkIndex = 0; WorkIndex < MAX_WORK_COUNT; WorkIndex++)
+	{
+		// 0. Fetch work item
+		if (GroupThread1D == 0)
+		{
+			InterlockedAdd(RWWorkCounter[0], 1, s_WorkID);
+		}
+
+		// 0. Reset all LDS variables
+		{
+			if (GroupThread1D == 0)
+			{
+				s_Segments_Min = PackDepth(1e8);
+				s_Segments_Max = 0;
+				s_bDataOrder = 0;
+				s_OpaqueMask = 0;
+			}
+			if (GroupThread1D < RASTER_DEPTH_BUCKET)
+			{
+				s_SegmentsCount[GroupThread1D] = 0;
+				s_SegmentsAlloc[GroupThread1D] = 0;
+			}
+			ClearMask(GroupThread2D);
+			//s_Mask[Thread_Coord.x][Thread_Coord.y] = 0;
+
+			if (GroupThread1D < 32)
+			{
+				s_Data[GroupThread1D][0] = 0;
+				s_Data[GroupThread1D][1] = 0;
+				s_Data[GroupThread1D][2] = 0;
+				s_Data[GroupThread1D][3] = 0;
+			}
+			else
+			{
+				s_Data[GroupThread1D-32][4] = 0;
+				s_Data[GroupThread1D-32][5] = 0;
+				s_Data[GroupThread1D-32][6] = 0;
+				s_Data[GroupThread1D-32][7] = 0;
+			}
+
+			#if PERMUTATION_DEBUG
+			s_Coverage[GroupThread1D] = 0;
+			#endif
+		}
+		GroupMemoryBarrierWithGroupSync();
+
+		// 0. Load all tile data
+		const uint BinTileIndex = s_WorkID / 16; // 1 x Bin32x32 -> 16 x Raster8x8
+		const bool bTileValid = (BinTileIndex < BinTileNum);
+		if (!bTileValid)
+		{
+			return;
+		}
+
+		const uint PrimOffset	= BinTileIndex * SEGMENT_COUNT_PER_GROUP;
+		const uint PrimCount	= LoadVisTileData(VisTileData, BinTileIndex, VT_PrimCount);
+		const uint2 BinTileCoord= UnpackVisTileCoord(LoadVisTileData(VisTileData, BinTileIndex, VT_Coord));
+		const uint2 BinTileMin	= BinTileCoord * BIN_TILE_SIZE;
+		const uint2 BinTileMax	= BinTileMin + BIN_TILE_SIZE;
+
+		const uint QuadrantIndex	= s_WorkID % 16; // 1 x Bin32x32 -> 16 x Raster8x8
+		const uint2 QuadrantCoord	= LinearTo2D_Common(QuadrantIndex, 4, 1.f / 4.f);
+		const uint2 RasterTileCoord = BinTileCoord * 4 + QuadrantCoord;
+		const uint2 RasterTileMin	= RasterTileCoord * RASTER_TILE_SIZE;
+		const uint2 RasterTileMax	= RasterTileMin   + RASTER_TILE_SIZE;
+
+		const uint LoopCount64  = DivideAndRoundUp(PrimCount, RASTER_THREAD_COUNT);
+		const uint LoopCount32  = DivideAndRoundUp(PrimCount, RASTER_SEGMENT_COUNT);
+
+		// For debug only
+		#if PERMUTATION_DEBUG
+		const bool bDebugEnableAll	= all(RasterTileCoord == uint2(ShaderPrintData.CursorCoord / float(RASTER_TILE_SIZE)));
+		const bool bDebugEnable		= all(RasterTileCoord == uint2(ShaderPrintData.CursorCoord / float(RASTER_TILE_SIZE))) && GroupThread1D == 0;
+		FShaderPrintContext Ctx		= InitShaderPrintContext(bDebugEnable, uint2(350, 50));
+		const FFontColor FontLegend = FontWhite;
+		const FFontColor FontValue  = FontOrange;
+		#endif
+
+		#if PERMUTATION_DEBUG
+		if (bDebugEnable)
+		{
+			AddFilledQuadSS(BinTileMin, BinTileMax, Transparent(ColorLightGreen));
+			//AddFilledQuadSS(RasterTileMin, RasterTileMax, Transparent(ColorLightGreen));
+			AddQuadSS(RasterTileMin, RasterTileMax, Transparent(ColorYellow));
+		}
+		#endif
+
+		#if PERMUTATION_DEBUG
+		if (0 && GroupThread1D == 0)
+		{
+			const float4 TileColor = ColorLightGreen;
+			//const float4 TileColor = float4(ColorMapTurbo(s_WorkID / 2048.f),1);
+			//AddFilledQuadSS(RasterTileMin, RasterTileMax, Transparent(TileColor));
+			AddQuadSS(RasterTileMin, RasterTileMax, ColorYellow);
+		}		
+		#endif
+
+		// 1. Load all the segments and compute min/max bound (move this during the compaction)
+		{
+			LOOP
+			for (uint LoopIndex = 0; LoopIndex < LoopCount64; ++LoopIndex)
+			{
+				const uint Prim = GroupThread1D + LoopIndex * RASTER_THREAD_COUNT;
+				if (Prim < PrimCount)
+				{
+					const uint PrimID = VisTilePrims[PrimOffset + Prim];
+
+					uint TypeDummy = 0;
+					float4 SP0 = 0;
+					float4 SP1 = 0;
+					float Rad0 = 0;
+					float Rad1 = 0;
+					CalcHomogenousPosAndRad(PrimID,   SP0, Rad0, TypeDummy);
+					CalcHomogenousPosAndRad(PrimID+1, SP1, Rad1, TypeDummy);
+				
+					float Alpha0 = 0;
+					float Alpha1 = 1;
+					// Clipping
+					{
+						SP0 = float4(NDCToPixelCoord(SP0), 1.0f / SP0.w);
+						SP1 = float4(NDCToPixelCoord(SP1), 1.0f / SP1.w);
+
+						// Clip against tile
+						bool2 bClipped = false;
+						float2 T = 0;
+						ClipRaySegment(BinTileMin - 0.5f, BinTileMax + 0.5f, SP0, SP1, Rad0, Rad1, Alpha0, Alpha1, bClipped, T);
+					}
+
+					const uint uDepth = PackDepth(max(SP0.z, SP1.z)); // Inverse Z
+					s_Segments[Prim] = uint2(uDepth, PrimID);
+
+					InterlockedMin(s_Segments_Min, uDepth);
+					InterlockedMax(s_Segments_Max, uDepth);
+
+					#if 0 && PERMUTATION_DEBUG
+					if (bDebugEnableAll)
+					{
+						const float3 P0 = UnpackHairControlPoint(PrimID).Position;
+						const float3 P1 = UnpackHairControlPoint(PrimID+1).Position;
+
+						const float4 Color0 = float4(LoadSampleColor(PrimID), 1);
+						const float4 Color1 = float4(LoadSampleColor(PrimID+1), 1);
+
+						const float4 LineColor = float4(ColorMapTurbo(GroupThread1D / float(RASTER_THREAD_COUNT)), 1);
+						//AddLineWS(P0, P1, LineColor);
+						AddLineWS(P0, P1, Color0, Color1);
+					}
+					#endif
+				}
+			}
+		}
+		GroupMemoryBarrierWithGroupSync();
+		const float InvDepthExtent = 1.f / max(0.0001f, UnpackDepth(s_Segments_Max)- UnpackDepth(s_Segments_Min));
+
+		// DEBUG
+		#if PERMUTATION_DEBUG
+		if (bDebugEnable)//  && GroupThreadID == 63)
+		{
+			const uint2 OutCoord = GroupThread2D + RasterTileMin;
+
+			//Print(Ctx, TEXT("Out Coord  :"), FontLegend); Print(Ctx, OutCoord, FontValue); Newline(Ctx);
+			//Print(Ctx, TEXT("Min Raster :"), FontLegend); Print(Ctx, RasterTileMin, FontValue); Newline(Ctx);
+			//Print(Ctx, TEXT("Thread     :"), FontLegend); Print(Ctx, Thread_Coord, FontValue); Newline(Ctx);
+			//Print(Ctx, TEXT("Max Raster :"), FontLegend); Print(Ctx, RasterTileMax, FontValue); Newline(Ctx);
+			//Newline(Ctx);
+
+			//Print(Ctx, TEXT("Cursor     :"), FontLegend); Print(Ctx, uint2(ShaderPrintData.CursorCoord), FontValue); Newline(Ctx);
+			//Newline(Ctx);
+
+			Print(Ctx, TEXT("Work ID      :"), FontLegend); Print(Ctx, s_WorkID, FontValue); Newline(Ctx);
+			Print(Ctx, TEXT("Prim Count   :"), FontLegend); Print(Ctx, PrimCount, FontValue); Newline(Ctx);
+			Print(Ctx, TEXT("Loop 32 Count:"), FontLegend); Print(Ctx, LoopCount32, FontValue); Newline(Ctx);
+			Print(Ctx, TEXT("Loop 64 Count:"), FontLegend); Print(Ctx, LoopCount64, FontValue); Newline(Ctx);
+			Print(Ctx, TEXT("Min Depth    :"), FontLegend); Print(Ctx, UnpackDepth(s_Segments_Min), FontValue); Newline(Ctx);
+			Print(Ctx, TEXT("Max Depth    :"), FontLegend); Print(Ctx, UnpackDepth(s_Segments_Max), FontValue); Newline(Ctx);
+			Newline(Ctx);
+
+			//Print(Ctx, TEXT("Bin ID     :"), FontLegend); Print(Ctx, BinTileIndex, FontValue); Newline(Ctx);
+			//Print(Ctx, TEXT("Raster ID  :"), FontLegend); Print(Ctx, QuadrantIndex, FontValue); Newline(Ctx);
+			//Newline(Ctx);
+
+			//Print(Ctx, TEXT("Raster Tile:"), FontLegend); Print(Ctx, RasterTileCoord, FontValue); Newline(Ctx);
+			//Print(Ctx, TEXT("Bin Tile   :"), FontLegend); Print(Ctx, BinTileCoord, FontValue); Newline(Ctx);
+			//Print(Ctx, TEXT("Out Coord  :"), FontLegend); Print(Ctx, OutCoord, FontValue); Newline(Ctx);
+			//Newline(Ctx);
+
+		}
+		#endif
+
+		// 2. Compute the count of depth bucket
+		{
+			LOOP
+			for (uint LoopIndex = 0; LoopIndex < LoopCount64; LoopIndex++)
+			{
+				const uint Prim = GroupThread1D + LoopIndex * RASTER_THREAD_COUNT;
+				if (Prim < PrimCount)
+				{
+					const float Depth	= UnpackDepth(s_Segments[Prim].x);
+					const uint DepthIt  = GetDepthBinIndex(Depth, InvDepthExtent);
+
+					InterlockedAdd(s_SegmentsCount[DepthIt], 1);
+				}
+			}
+		}
+		GroupMemoryBarrierWithGroupSync();
+
+		// Replace this with parallel version
+		if (GroupThread1D == 0)
+		{
+			uint Acc = 0;
+			for (uint It = 0; It < RASTER_DEPTH_BUCKET; It++)
+			{
+				const uint Next = s_SegmentsCount[It];
+				s_SegmentsCount[It] = Acc;
+				Acc += Next;
+
+				#if 0 && PERMUTATION_DEBUG
+				if (bDebugEnable)
+				{
+					if (Next > 0)	{ Print(Ctx, TEXT("x"), FontValue); }
+					else			{ Print(Ctx, TEXT("."), FontValue); }
+					if (It == RASTER_DEPTH_BUCKET-1) { Newline(Ctx); }
+				}
+				#endif
+			}
+		}
+		GroupMemoryBarrierWithGroupSync();
+
+		// 3. Insert the segment into the right bucket
+		{
+			LOOP
+			for (uint LoopIndex = 0; LoopIndex < LoopCount64; LoopIndex++)
+			{
+				const uint Prim = GroupThread1D + LoopIndex * RASTER_THREAD_COUNT;
+				if (Prim < PrimCount)
+				{
+					const uint2 Segment = s_Segments[Prim];
+					const float Depth = UnpackDepth(Segment.x);
+					const uint DepthIt = GetDepthBinIndex(Depth, InvDepthExtent);
+
+					uint AllocOffset = 0;
+					InterlockedAdd(s_SegmentsAlloc[DepthIt], 1, AllocOffset);
+					const uint NewIndex = AllocOffset + s_SegmentsCount[DepthIt];
+					s_Segments_Sorted[NewIndex] = Segment;
+				}
+			}
+		}
+		GroupMemoryBarrierWithGroupSync();
+
+		// 4. Raster segments
+		{
+			// These are global Color/Coverage for the final pixel handled by this thread
+			float3 Thread_Color    = 0;
+			float  Thread_Coverage = 0;
+			uint   Thread_Complete = 0;
+			uint   Thread_LoopCountToFullCoverage=0;
+
+			// 4.1 Loop over all segment within the tile, and rastize 32 of them at each loop
+			LOOP
+			for (uint LoopIndex = 0; LoopIndex < LoopCount32; LoopIndex++)
+			{
+				// 4.0 Reset
+				if (GroupThread1D < RASTER_SEGMENT_COUNT)
+				{
+					for (int J = 0; J < RASTER_TILE_SIZE; ++J)
+					{
+						s_Data[GroupThread1D][J] = 0;
+						s_Color[GroupThread1D][J] = 0;
+					}
+				}
+				s_Mask[GroupThread2D.x][GroupThread2D.y] = 0;
+
+				GroupMemoryBarrierWithGroupSync();
+
+				if (s_OpaqueMask == 64)
+				{
+					break;
+				}
+
+				// 4.1 Raster segment (1 thread = 1 segment)
+				// Half of the thread are doing nothing, we could raster the two half of the segments (one per each thread)
+				const uint Prim = GroupThread1D + LoopIndex * RASTER_SEGMENT_COUNT;
+				if (GroupThread1D < RASTER_SEGMENT_COUNT && Prim < PrimCount)
+				{
+					const uint PrimID = s_Segments_Sorted[Prim].y;
+					uint TypeDummy = 0;
+					float4 SP0 = 0;
+					float4 SP1 = 0;
+					float Rad0 = 0;
+					float Rad1 = 0;
+					CalcHomogenousPosAndRad(PrimID,     SP0, Rad0, TypeDummy);
+					CalcHomogenousPosAndRad(PrimID + 1, SP1, Rad1, TypeDummy);
+
+					float Alpha0=0;
+					float Alpha1=1;
+
+					bool2 bClipped = false;
+					bool bIsSegmentValid = false;
+					float2 T = 0;
+					// 4.1.1 Clipping
+					{
+						SP0 = float4(NDCToPixelCoord(SP0), 1.0f / SP0.w);
+						SP1 = float4(NDCToPixelCoord(SP1), 1.0f / SP1.w);
+				
+						// Clip against tile
+						bIsSegmentValid = ClipRaySegment(RasterTileMin - 0.5f, RasterTileMax + 0.5f, SP0, SP1, Rad0, Rad1, Alpha0, Alpha1, bClipped, T);
+					}
+
+					// 4.1.2 Rasterize segment (1 thread = 1 segments)
+					if (bIsSegmentValid)
+					{
+						const float SegmentLenSqRcp = 1.0f / dot(SP1.xy - SP0.xy, SP1.xy - SP0.xy);
+						const bool bIsSteep = abs(SP1.x - SP0.x) < abs(SP1.y - SP0.y);
+						const float X0 = bIsSteep ? min(SP0.y, SP1.y) : min(SP0.x, SP1.x);
+						const float X1 = bIsSteep ? max(SP0.y, SP1.y) : max(SP0.x, SP1.x);
+						const float RcpNumSteps = 1.0f / (X1 - X0);
+						const int NumSteps = (int)(ceil(X1) - floor(X0));
+
+						if (bIsSteep)
+						{
+							InterlockedOr(s_bDataOrder, 1u << GroupThread1D);
+						}
+
+						// DEBUG (Write coord)
+						#if 0 && PERMUTATION_DEBUG
+						if (GroupThread1D == 0 && bDebugEnable)
+						{
+							Print(Ctx, TEXT("NumSteps:"), FontYellow);
+							Print(Ctx, NumSteps, FontWhite); 
+							Newline(Ctx);
+
+							Print(Ctx, TEXT("X0      :"), FontYellow);
+							Print(Ctx, X0, FontWhite);
+							Newline(Ctx);
+
+							Print(Ctx, TEXT("X1      :"), FontYellow);
+							Print(Ctx, X1, FontWhite);
+							Newline(Ctx);
+
+							Print(Ctx, TEXT("T       :"), FontYellow);
+							Print(Ctx, T, FontWhite);
+							Newline(Ctx);
+
+							Print(Ctx, TEXT("bValid  :"), FontYellow);
+							PrintBool(Ctx, bIsSegmentValid);
+							Newline(Ctx);
+							
+							
+							
+							//const float3 P0 = UnpackHairControlPoint(PrimID).Position;
+							//const float3 P1 = UnpackHairControlPoint(PrimID + 1).Position;
+
+							//const float4 Color0 = float4(LoadSampleColor(PrimID), 1);
+							//const float4 Color1 = float4(LoadSampleColor(PrimID + 1), 1);
+
+							//const float4 LineColor = float4(ColorMapTurbo(GroupThread1D / float(RASTER_THREAD_COUNT)), 1);
+							//AddLineWS(P0, P1, LineColor);
+							AddLineSS(SP0, SP1, ColorRed);
+						}
+						#endif
+
+						// DEBUG (Legend)
+						#if 0 && PERMUTATION_DEBUG
+						if (GroupThread1D == 0 && bDebugEnable)
+						{
+							//Print(Ctx, s_Mask[IntraTileCoord.x][IntraTileCoord.y], FontWhite);
+							//PrintPixelMask(Ctx, ReadMask(IntraTileCoord)/*s_Mask[IntraTileCoord.x][IntraTileCoord.y]*/);
+							//Print(Ctx, TEXT("Color0: "), FontWhite); Print(Ctx, Color0); Newline(Ctx);
+							//Print(Ctx, TEXT("Color1: "), FontWhite); Print(Ctx, Color1); Newline(Ctx);
+							//Print(Ctx, TEXT("Color :  "), FontWhite); Print(Ctx, Color);
+							const float3 Color0 = LoadSampleColor(PrimID);
+							const float3 Color1 = LoadSampleColor(PrimID + 1);
+							Print(Ctx, TEXT("Color0 :  "), FontWhite); Print(Ctx, Color0); Newline(Ctx);
+							Print(Ctx, TEXT("Color1 :  "), FontWhite); Print(Ctx, Color1); Newline(Ctx);
+							Print(Ctx, TEXT("  Alpha       A0          A1          AColor     Color"), FontWhite);
+							Newline(Ctx);
+						}
+						#endif
+						LOOP
+						for (int J = 0; J < NumSteps; ++J)
+						//for (int J = 0; J < RASTER_TILE_SIZE; ++J)
+						{
+							const float Alpha = saturate(J * RcpNumSteps);
+							const float4 SP = lerp(SP0, SP1, Alpha);
+
+							int2 Coords = SP.xy;
+							const int2 IntraTileCoord = Coords - RasterTileMin;
+
+							// TO it needs to store data in sweeping order 
+							FSampleData Sample = (FSampleData)0;
+							if (all(IntraTileCoord >= 0) && all(IntraTileCoord < RASTER_TILE_SIZE))
+							{
+								const float Alpha = ComputeLerpAlpha(Coords, SP0.xy, SP1.xy, SegmentLenSqRcp);
+								const float Depth = lerp(SP0.z, SP1.z, Alpha);
+
+								const float AlphaColor = lerp(Alpha0, Alpha1, Alpha);
+								const float3 Color0 = LoadSampleColor(PrimID);
+								const float3 Color1 = LoadSampleColor(PrimID+1);
+								const float3 Color = lerp(Color0, Color1, AlphaColor);
+
+								// Fill in sample data
+								Sample.Depth = Depth;
+								Sample.Coord = IntraTileCoord;
+								Sample.Coverage = 0.5f; //TODO;
+
+								// Write data
+								WriteMask(IntraTileCoord, 1u << GroupThread1D);
+								//InterlockedOr(s_Mask[IntraTileCoord.x][IntraTileCoord.y], 1u<<GroupThread1D);
+
+								const uint JCoord = bIsSteep ? IntraTileCoord.y : IntraTileCoord.x;
+								s_Color[GroupThread1D][JCoord] = PackColorData(Color);
+								s_Data[GroupThread1D][JCoord]  = PackSampleData(Sample);
+
+
+								// DEBUG (Write coord)
+								#if 0 && PERMUTATION_DEBUG
+								if (GroupThread1D == 0 && bDebugEnable)
+								{
+									Print(Ctx, J, FontYellow, 2, 0);
+									//Print(Ctx, IntraTileCoord, FontWhite,2,0);
+									//Print(Ctx, s_Mask[IntraTileCoord.x][IntraTileCoord.y], FontWhite);
+									//PrintPixelMask(Ctx, ReadMask(IntraTileCoord)/*s_Mask[IntraTileCoord.x][IntraTileCoord.y]*/);
+									//Print(Ctx, TEXT("Color0: "), FontWhite); Print(Ctx, Color0); Newline(Ctx);
+									//Print(Ctx, TEXT("Color1: "), FontWhite); Print(Ctx, Color1); Newline(Ctx);
+									//Print(Ctx, TEXT("Color :  "), FontWhite); Print(Ctx, Color);
+									Print(Ctx, Alpha, FontWhite);
+									Print(Ctx, Alpha0, FontRed);
+									Print(Ctx, Alpha1, FontGreen);
+									Print(Ctx, AlphaColor, FontBlue);
+									Print(Ctx, Color, FontOrange);
+									Newline(Ctx);
+								}
+								#endif
+							}
+						}
+					} // if (bIsSegmentValid)
+				}
+				GroupMemoryBarrierWithGroupSync();
+
+				// DEBUG (Coverage mask)
+				#if 0 && PERMUTATION_DEBUG
+				if (GroupThread1D == 0 && bDebugEnable)
+				{
+					PrintCoverage(Ctx);
+				}
+				#endif
+
+				// 4.2 Combine all samples within the same pixel (1 thread = 1 pixel)
+				{
+					const uint ValidSegments = ReadMask(GroupThread2D);//s_Mask[Thread_Coord.x][Thread_Coord.y];
+
+					
+					// Compute transmittance function along depth
+					#define USE_TRANSMITTANCE 0
+					#if USE_TRANSMITTANCE
+					float Coverages[RASTER_SEGMENT_COUNT];
+					float Pixel_InvDepthExtent = 0.010f;
+					{
+						// 1. Reset coverage & compute min/max depth of pixel
+						float Pixel_MinDepth = 1e8;
+						float Pixel_MaxDepth = 0;
+						{
+							for (uint SegmentIt = 0; SegmentIt < RASTER_SEGMENT_COUNT; ++SegmentIt)
+							{
+								Coverages[SegmentIt] = 0;
+
+								const uint SegmentBit = 1u << SegmentIt;
+								const bool bIsValid = (ValidSegments & SegmentBit) > 0;
+								const bool bIsSteep = (s_bDataOrder & SegmentBit) > 0;
+								if (bIsValid)
+								{
+									const uint J = bIsSteep ? Thread_Coord.y : Thread_Coord.x;
+									const FSampleData Sample = UnpackSampleData(s_Data[SegmentIt][J]);
+									Pixel_MinDepth = min(Pixel_MinDepth, Sample.Depth);
+									Pixel_MaxDepth = max(Pixel_MaxDepth, Sample.Depth);
+								}
+							}
+						}
+						Pixel_InvDepthExtent = max(0.010f, Pixel_MaxDepth - Pixel_MinDepth);
+
+						// 2. Compute coverage per bin
+						for (uint SegmentIt=0; SegmentIt < RASTER_SEGMENT_COUNT; ++SegmentIt)
+						{
+							const uint SegmentBit = 1u << SegmentIt;
+							const bool bIsValid = (ValidSegments & SegmentBit) > 0;
+							const bool bIsSteep = (s_bDataOrder & SegmentBit) > 0;
+							if (bIsValid)
+							{
+								const uint J = bIsSteep ? Thread_Coord.y : Thread_Coord.x;
+
+								const FSampleData Sample = UnpackSampleData(s_Data[SegmentIt][J]);
+								const uint DepthIndex = GetDepthBinIndex(Sample.Depth, Pixel_InvDepthExtent);
+
+								Coverages[DepthIndex] = min(1, Coverages[DepthIndex] + Sample.Coverage);
+							}
+						}
+
+						// 3. Compute coverage function (accumulated)
+						for (uint TransIt = 1; TransIt < RASTER_SEGMENT_COUNT; ++TransIt)
+						{
+							const float PrevCoverage = Coverages[TransIt - 1];
+							const float CurrCoverage = Coverages[TransIt];
+							Coverages[TransIt] = saturate(PrevCoverage + (1 - PrevCoverage) + CurrCoverage);
+						}
+					}
+					#endif
+
+					// Change this loop into to bit logic?
+					{
+						for (uint SegmentIt=0; SegmentIt < RASTER_SEGMENT_COUNT; ++SegmentIt)
+						{
+							const uint SegmentBit = 1u << SegmentIt;
+							const bool bIsValid = (ValidSegments & SegmentBit) > 0;
+							const bool bIsSteep = (s_bDataOrder & SegmentBit) > 0;
+							if (bIsValid)
+							{
+								const uint J = bIsSteep ? GroupThread2D.y : GroupThread2D.x;
+
+								const FSampleData Sample = UnpackSampleData(s_Data[SegmentIt][J]);
+								const float3 Color = UnpackColorData(s_Color[SegmentIt][J]);
+
+								// Accumulate sample contribution 
+								// TODO: needs to compute transmittance funciton, as seg ordered from front to back, but the raster pixels might be slightly different
+								#if USE_TRANSMITTANCE
+								const uint DepthIndex = GetDepthBinIndex(Sample.Depth, Pixel_InvDepthExtent);
+								const float Coverage  = Coverages[DepthIndex];
+								//const float Transmittance = Transmittances[DepthIndex];
+								//Thread_Coverage += Transmittance * Sample.Coverage;
+								//Thread_Color    += Transmittance * Sample.Coverage * Color;
+								Thread_Coverage = Coverage;
+								Thread_Color += Coverage * Color;
+								#elif 1
+								const float AccTransmittance = (1.f-Thread_Coverage);
+								Thread_Coverage += AccTransmittance * Sample.Coverage;
+								Thread_Color    += AccTransmittance * Sample.Coverage * Color;
+
+								//Thread_Coverage = 0.5;
+								//Thread_Color    = 0;
+
+								#else
+								Thread_Coverage += 1;
+								Thread_Color = Color;
+								#endif
+							}
+						}
+					}
+
+					const float CoverageThreshold = 0.95f;
+					if (Thread_Complete == 0 && Thread_Coverage > CoverageThreshold)
+					{
+						InterlockedAdd(s_OpaqueMask, 1);
+						Thread_Complete = 1;
+					}
+
+					#if PERMUTATION_DEBUG
+					if (Thread_LoopCountToFullCoverage == 0 && Thread_Coverage > CoverageThreshold)
+					{
+						Thread_LoopCountToFullCoverage = LoopIndex;
+					}
+					#endif
+				}
+				//GroupMemoryBarrierWithGroupSync();
+
+			} // for (uint LoopIndex =...) ...
+
+			#if PERMUTATION_DEBUG
+			s_Coverage[GroupThread1D] = Thread_Coverage;
+			GroupMemoryBarrierWithGroupSync(); // For sync s_Coverage
+			#endif
+
+			// 4.2. Write final color
+			if (Thread_Coverage>0)
+			{
+				const uint2 PixelCoord = GroupThread2D + RasterTileMin;
+				const bool bOutput = true;//View.GeneralPurposeTweak > 1;
+				if (bOutput)
+				{
+					const float3 SourceColor = OutSceneColorTexture[PixelCoord].xyz;
+					OutSceneColorTexture[PixelCoord] = float4(lerp(SourceColor, Thread_Color, Thread_Coverage), 1);
+					//OutSceneColorTexture[PixelCoord] = float4(Thread_Color, 1);
+					//OutSceneColorTexture[PixelCoord] = float4(1, 1, 0, 1);
+					//OutSceneColorTexture[PixelCoord] = float4(PixelCoord / OutputResolutionf,0,1);
+					//OutSceneColorTexture[PixelCoord] = float4(GroupThread2D / float2(RASTER_TILE_SIZE, RASTER_TILE_SIZE), 0, 1);
+				}
+
+				// For debug purpose only
+				#if PERMUTATION_DEBUG
+				{
+					InterlockedAdd(OutHairCountTexture_ForDebug[PixelCoord], 1);
+					InterlockedAdd(OutHairPixelCountPerTile_ForDebug[RasterTileCoord], 1);
+				}
+				#endif
+			}
+
+			#if PERMUTATION_DEBUG
+			if (GroupThread1D == 0)
+			{
+				// Full/Partial coverage tiles
+				if (0)
+				{
+					const bool bFullCover = s_OpaqueMask == 64;
+					if (bFullCover)	AddFilledQuadSS(RasterTileMin, RasterTileMax, Transparent(ColorLightRed));
+					else			AddFilledQuadSS(RasterTileMin, RasterTileMax, Transparent(ColorLightGreen));
+				}
+
+				if (bDebugEnable)
+				{
+					Newline(Ctx);
+					Print(Ctx, TEXT("Full Coverage: "), FontWhite);
+					Print(Ctx, s_OpaqueMask, FontYellow, 2, 0);
+					Newline(Ctx);
+
+					Newline(Ctx);
+					for (uint y = 0; y < RASTER_TILE_SIZE; ++y)
+					{
+						for (uint x = 0; x < RASTER_TILE_SIZE; ++x)
+						{
+							const uint Index = x + y * RASTER_TILE_SIZE;
+							FFontColor Color = InitFontColor(ColorMapTurbo(s_Coverage[Index]));
+							//Print(Ctx, s_Coverage[Index], Color, 4, 2);
+							
+							Print(Ctx, Thread_LoopCountToFullCoverage, Color, 2, 0);
+						}
+						Newline(Ctx);
+						Newline(Ctx);
+					}
+					Newline(Ctx);
+				}
+
+				//InterlockedAdd(OutHairPixelCountPerTile_ForDebug[RasterTileCoord], PrimCount);
+				OutHairPixelCountPerTile_ForDebug[RasterTileCoord] = PrimCount;
+			}
+			#endif
+
+		} // 4. Raster
+
+	}// Work item
+}
+
+#endif //SHADER_RASTERCOMPUTE_RASTER
+
+///////////////////////////////////////////////////////////////////////////
+
+#if SHADER_RASTERCOMPUTE_DEBUG
+
+#include "../ShaderPrint.ush"
+
+Buffer<uint>						VisTileArgs;
+ByteAddressBuffer					VisTileData;
+
+Texture2D<uint>						HairCountTexture_ForDebug;
+Texture2D<uint>						HairPixelCountPerTile_ForDebug;
+uint								InstanceCount;
+
+uint GetTileTotalSegment(uint2 TileCoord, bool bPrintDetails, uint InTileSize)
+{
+	const float TileDisplayScale = 1.5f;
+	const uint DisplayTileSize = InTileSize * TileDisplayScale;
+	uint2 InlinedTileCoord = uint2(0, 0);
+
+	const uint TileCount = VisTileArgs[0];
+
+	uint TotalSegments = 0;
+	for (uint TileIndex=0; TileIndex<TileCount; ++TileIndex)
+	{
+		const uint PackedTileCoord = LoadVisTileData(VisTileData, TileIndex, VT_Coord);
+		const uint2 VisTileCoord = UnpackVisTileCoord(PackedTileCoord);
+		if (all(VisTileCoord == TileCoord))
+		{
+			const uint TileSegments = LoadVisTileData(VisTileData, TileIndex, VT_PrimCount);
+			TotalSegments += TileSegments;
+
+			if (bPrintDetails)
+			{
+				AddFilledQuadSS(InlinedTileCoord * DisplayTileSize, (InlinedTileCoord + 1) * DisplayTileSize, TileSegments > 0 ? Transparent(ColorLightGreen) : Transparent(ColorLightRed));
+				AddQuadSS(InlinedTileCoord * DisplayTileSize, (InlinedTileCoord + 1) * DisplayTileSize, ColorYellow);
+
+				const uint2 TilePrintOffset = InTileSize >> 1;
+
+				FShaderPrintContext Context = InitShaderPrintContext(true, InlinedTileCoord * DisplayTileSize + TilePrintOffset);
+				Print(Context, TileSegments, FontWhite);
+				++InlinedTileCoord.x;
+			}
+		}
+	}
+	return TotalSegments;
+}
+
+void PrintTile(uint2 TileCoord, uint TotalSegments, bool bPrintText, uint InTileSize)
+{
+	AddFilledQuadSS(TileCoord * InTileSize, (TileCoord + 1) * InTileSize, TotalSegments > 0 ? Transparent(ColorLightGreen) : Transparent(ColorLightRed));
+	if (bPrintText)
+	{
+		FShaderPrintContext Context = InitShaderPrintContext(true, TileCoord * InTileSize + uint2(0, InTileSize * 1.5f));
+		Print(Context, TotalSegments, FontWhite);
+
+		AddQuadSS(TileCoord * InTileSize, (TileCoord + 1) * InTileSize, ColorYellow);
+	}
+}
+
+[numthreads(8, 8, 1)]
+void MainCS(uint3 ThreadId : SV_DispatchThreadID)
+{
+	const bool bIsCursorPixel = (all(ThreadId == 0) && all(ShaderPrintData.CursorCoord >= 0));
+
+	// Info/Stats
+	if (all(ThreadId == 0))
+	{
+		FFontColor FontValue = FontOrange;
+		FFontColor FontTitle = FontYellow;
+		FFontColor FontLegend = FontWhite;
+		FShaderPrintContext Context = InitShaderPrintContext(true, uint2(50, 110));
+		Print(Context, TEXT("Raster compute         "), FontTitle); Newline(Context);
+		Print(Context, TEXT("Instance Count       : "), FontLegend); Print(Context, InstanceCount, FontValue, 3, 0); Newline(Context);
+		Print(Context, TEXT("Total segments Count : "), FontLegend); Print(Context, GetControlPointCount(), FontValue); Newline(Context);
+		Print(Context, TEXT("Max.  segments Count : "), FontLegend); Print(Context, MaxControlPointCount, FontValue);  Newline(Context);
+		Newline(Context);
+
+		Print(Context, TEXT("Configuration          "), FontTitle); Newline(Context);
+		Print(Context, TEXT("Output Resolution    : "), FontLegend); Print(Context, OutputResolution.x, FontValue, 4, 0); Print(Context, TEXT("x"), FontSilver); Print(Context, OutputResolution.y, FontValue, 4, 0); Newline(Context);
+		Newline(Context);
+
+		Print(Context, TEXT("Bin Tile Size        : "), FontLegend); Print(Context, uint(BIN_TILE_SIZE), FontValue, 2, 0); Newline(Context);
+		Print(Context, TEXT("Bin Tile Res         : "), FontLegend); Print(Context, BinTileRes.x, FontValue, 3, 0); Print(Context, TEXT("x"), FontSilver); Print(Context, BinTileRes.y, FontValue, 3, 0); Newline(Context);
+		Print(Context, TEXT("Num Binners          : "), FontLegend); Print(Context, NumBinners, FontValue); Newline(Context);
+		Newline(Context);
+
+		Print(Context, TEXT("Raster Tile Size     : "), FontLegend); Print(Context, uint(RASTER_TILE_SIZE), FontValue, 2, 0); Newline(Context);
+		Print(Context, TEXT("Raster Tile Res      : "), FontLegend); Print(Context, RasterTileRes.x, FontValue, 3, 0); Print(Context, TEXT("x"), FontSilver); Print(Context, RasterTileRes.y, FontValue, 3, 0); Newline(Context);
+		Print(Context, TEXT("Num Rasterizers      : "), FontLegend); Print(Context, NumRasterizers, FontValue); Newline(Context);
+		Newline(Context);
+
+		Print(Context, TEXT("Alloc. Tile Count    : "), FontLegend); Print(Context, VisTileArgs[0], FontValue); Newline(Context);
+		Print(Context, TEXT("Rasterizer Max Work  : "), FontLegend); Print(Context, NumRasterizers * MAX_WORK_COUNT, FontValue); Newline(Context);
+		Newline(Context);
+		
+		Newline(Context);
+
+		if (bIsCursorPixel)
+		{
+		const uint2 PixelCoord = ShaderPrintData.CursorCoord;
+		const uint2 RasterTileCoord = uint2(ShaderPrintData.CursorCoord) >> RASTER_TILE_SIZE_AS_SHIFT;
+
+		const uint HairCount = HairCountTexture_ForDebug.Load(uint3(PixelCoord, 0));
+		const uint RasterizedPixels = HairPixelCountPerTile_ForDebug.Load(uint3(RasterTileCoord, 0));
+
+		Print(Context, TEXT("Hair Count           : "), FontLegend); Print(Context, HairCount, FontValue); Newline(Context);
+		Print(Context, TEXT("Hair #Pixel in Tile  : "), FontLegend); Print(Context, RasterizedPixels, FontValue); Newline(Context);
+		Newline(Context);
+		}
+	}
+	#if 0
+	// Cursor info
+	if (bIsCursorPixel)
+	{
+		const uint2 PixelCoord = ShaderPrintData.CursorCoord;
+		const uint2 BinTileCoord = uint2(ShaderPrintData.CursorCoord) >> BIN_TILE_SIZE_AS_SHIFT;
+		if (all(BinTileCoord < BinTileRes))
+		{
+			const uint TotalSegments = GetTileTotalSegment(BinTileCoord, true, BIN_TILE_SIZE);
+			PrintTile(BinTileCoord, TotalSegments, true, BIN_TILE_SIZE);
+		}
+	}
+
+	// All tile
+	{
+		const uint2 BinTileCoord = ThreadId.xy;
+		if (all(BinTileCoord < BinTileRes))
+		{
+			const uint TotalSegments = GetTileTotalSegment(BinTileCoord, false, BIN_TILE_SIZE);
+			if (TotalSegments)
+			{
+				PrintTile(BinTileCoord, TotalSegments, false, BIN_TILE_SIZE);
+			}
+		}
+	}
+	#endif
+}
+#endif //SHADER_RASTERCOMPUTE_DEBUG
+
diff --git a/Engine/Source/Runtime/Renderer/Private/HairStrands/HairStrandsComposition.cpp b/Engine/Source/Runtime/Renderer/Private/HairStrands/HairStrandsComposition.cpp
--- a/Engine/Source/Runtime/Renderer/Private/HairStrands/HairStrandsComposition.cpp
+++ b/Engine/Source/Runtime/Renderer/Private/HairStrands/HairStrandsComposition.cpp
@@ -12,6 +12,7 @@
 #include "HairStrandsTile.h"
 #include "FogRendering.h"
 #include "PostProcess/TemporalAA.h"
+#include "HairStrandsForwardRaster.h"
 
 /////////////////////////////////////////////////////////////////////////////////////////
 
@@ -467,7 +468,6 @@ static void AddHairVisibilityGBufferWritePass(
 		Parameters);
 }
 
-
 static void InternalRenderHairComposition(
 	FRDGBuilder& GraphBuilder,
 	const FViewInfo& View,
@@ -480,131 +480,138 @@ static void InternalRenderHairComposition(
 	RDG_GPU_STAT_SCOPE(GraphBuilder, HairStrandsComposition);
 	SCOPED_NAMED_EVENT(HairStrandsComposition, FColor::Emerald);
 
+	if (View.HairStrandsViewData.VisibilityData.RasterizedInstanceCount > 0)
+	{
+		AddHairStrandsForwardRasterPass(
+			GraphBuilder,
+			View,
+			SceneDepthTexture->Desc.Extent,
+			View.HairStrandsViewData.VisibilityData,
+			SceneDepthTexture,
+			SceneColorTexture);
+	}
+	else
 	{
+		const FHairStrandsVisibilityData& VisibilityData = View.HairStrandsViewData.VisibilityData;
+
+		if (!VisibilityData.CoverageTexture)
+		{
+			return; // Automatically skip for any view not rendering hair
+		}
+
+		FRDGTextureRef DOFDepth = nullptr;
+		const bool bHairDOF = GHairStrandsComposeDOFDepth > 0 ? 1 : 0;
+		if (bHairDOF)
 		{
+			DOFDepth = AddHairDOFDepthPass(
+				GraphBuilder,
+				View,
+				VisibilityData,
+				VisibilityData.CoverageTexture,
+				SceneColorTexture,
+				SceneDepthTexture);
+		}
+
+		FRDGTextureRef AccumulatedColor = nullptr;
+
+		const bool bTemporalLayeringEnabled = IsHairVisibilityComputeRasterTemporalLayeringEnabled() && View.ViewState;
+
+		if (bTemporalLayeringEnabled)
+		{
+			FIntPoint OutputResolution = SceneColorTexture->Desc.Extent;
+
+			FRDGTextureDesc HairColorInputDesc = FRDGTextureDesc::Create2D(OutputResolution, PF_FloatRGBA, FClearValueBinding(FLinearColor(0, 0, 0, 0)), TexCreate_RenderTargetable | TexCreate_ShaderResource | TexCreate_UAV, 1);
+			FRDGTextureRef HairColorInputTexture = GraphBuilder.CreateTexture(HairColorInputDesc, TEXT("Hair.ColorInputTexture"));
+
+			AddHairVisibilityComposeSamplePass(
+				GraphBuilder,
+				View,
+				VisibilityData,
+				VisibilityData.CoverageTexture,
+				DOFDepth,
+				nullptr,
+				HairColorInputTexture,
+				SceneDepthTexture,
+				bTemporalLayeringEnabled);
+
+
+			check(View.ViewState);
+			FTAAPassParameters TAASettings(View);
+			TAASettings.SceneDepthTexture = SceneDepthTexture;
+			TAASettings.SceneVelocityTexture = SceneVelocityTexture;
+			TAASettings.Pass = ETAAPassConfig::Hair;
+			TAASettings.SceneColorInput = HairColorInputTexture;
+			TAASettings.bOutputRenderTargetable = true;
+
+			FTAAOutputs TAAOutputs = AddTemporalAAPass(
+				GraphBuilder,
+				View,
+				TAASettings,
+				View.PrevViewInfo.HairHistory,
+				&View.ViewState->PrevFrameViewInfo.HairHistory);
+
+			AccumulatedColor = TAAOutputs.SceneColor;			
+		}
+
+		AddHairVisibilityComposeSamplePass(
+			GraphBuilder,
+			View,
+			VisibilityData,
+			VisibilityData.CoverageTexture,
+			DOFDepth,
+			AccumulatedColor,
+			SceneColorTexture,
+			SceneDepthTexture,
+			bTemporalLayeringEnabled);
+
+		if (VisibilityData.ResolveMaskTexture)
+		{
+			AddHairVisibilityFastResolveMaskPass(
+				GraphBuilder,
+				View,
+				VisibilityData.ResolveMaskTexture,
+				VisibilityData.TileData,
+				SceneDepthTexture);
+		}
+
+		const bool bWriteDummyData		= View.Family->ViewMode != VMI_VisualizeBuffer && GHairWriteGBufferData == 1;
+		const bool bWritePartialGBuffer = View.Family->ViewMode != VMI_VisualizeBuffer && (GHairWriteGBufferData == 1 || GHairWriteGBufferData == 2);
+		const bool bWriteFullGBuffer	= View.Family->ViewMode == VMI_VisualizeBuffer || (GHairWriteGBufferData == 3);
+		if (bWriteFullGBuffer || bWritePartialGBuffer)
+		{
+			const FSceneTextures& SceneTextures = View.GetSceneTextures();
+			const FRDGTextureRef GBufferATexture = SceneTextures.GBufferA;
+			const FRDGTextureRef GBufferBTexture = SceneTextures.GBufferB;
+			const FRDGTextureRef GBufferCTexture = SceneTextures.GBufferC;
+			const FRDGTextureRef GBufferDTexture = SceneTextures.GBufferD;
+			const FRDGTextureRef GBufferETexture = SceneTextures.GBufferE;
+			if (bWritePartialGBuffer && GBufferATexture && GBufferBTexture)
+			{
+				AddHairVisibilityGBufferWritePass(
+					GraphBuilder,
+					View,
+					bWriteDummyData,
+					VisibilityData.TileData,
+					GBufferATexture,
+					GBufferBTexture,
+					nullptr,
+					nullptr,
+					nullptr,
+					nullptr);
+			}
+			else if (bWriteFullGBuffer && GBufferATexture && GBufferBTexture && GBufferCTexture && SceneDepthTexture)
 			{
-				const FHairStrandsVisibilityData& VisibilityData = View.HairStrandsViewData.VisibilityData;
-
-				if (!VisibilityData.CoverageTexture)
-				{
-					return; // Automatically skip for any view not rendering hair
-				}
-
-				FRDGTextureRef DOFDepth = nullptr;
-				const bool bHairDOF = GHairStrandsComposeDOFDepth > 0 ? 1 : 0;
-				if (bHairDOF)
-				{
-					DOFDepth = AddHairDOFDepthPass(
-						GraphBuilder,
-						View,
-						VisibilityData,
-						VisibilityData.CoverageTexture,
-						SceneColorTexture,
-						SceneDepthTexture);
-				}
-
-				FRDGTextureRef AccumulatedColor = nullptr;
-
-				const bool bTemporalLayeringEnabled = IsHairVisibilityComputeRasterTemporalLayeringEnabled() && View.ViewState;
-
-				if (bTemporalLayeringEnabled)
-				{
-					FIntPoint OutputResolution = SceneColorTexture->Desc.Extent;
-
-					FRDGTextureDesc HairColorInputDesc = FRDGTextureDesc::Create2D(OutputResolution, PF_FloatRGBA, FClearValueBinding(FLinearColor(0, 0, 0, 0)), TexCreate_RenderTargetable | TexCreate_ShaderResource | TexCreate_UAV, 1);
-					FRDGTextureRef HairColorInputTexture = GraphBuilder.CreateTexture(HairColorInputDesc, TEXT("Hair.ColorInputTexture"));
-
-					AddHairVisibilityComposeSamplePass(
-						GraphBuilder,
-						View,
-						VisibilityData,
-						VisibilityData.CoverageTexture,
-						DOFDepth,
-						nullptr,
-						HairColorInputTexture,
-						SceneDepthTexture,
-						bTemporalLayeringEnabled);
-
-
-					check(View.ViewState);
-					FTAAPassParameters TAASettings(View);
-					TAASettings.SceneDepthTexture = SceneDepthTexture;
-					TAASettings.SceneVelocityTexture = SceneVelocityTexture;
-					TAASettings.Pass = ETAAPassConfig::Hair;
-					TAASettings.SceneColorInput = HairColorInputTexture;
-					TAASettings.bOutputRenderTargetable = true;
-
-					FTAAOutputs TAAOutputs = AddTemporalAAPass(
-						GraphBuilder,
-						View,
-						TAASettings,
-						View.PrevViewInfo.HairHistory,
-						&View.ViewState->PrevFrameViewInfo.HairHistory);
-
-					AccumulatedColor = TAAOutputs.SceneColor;			
-				}
-
-				AddHairVisibilityComposeSamplePass(
+				AddHairVisibilityGBufferWritePass(
 					GraphBuilder,
 					View,
-					VisibilityData,
-					VisibilityData.CoverageTexture,
-					DOFDepth,
-					AccumulatedColor,
-					SceneColorTexture,
-					SceneDepthTexture,
-					bTemporalLayeringEnabled);
-
-				if (VisibilityData.ResolveMaskTexture)
-				{
-					AddHairVisibilityFastResolveMaskPass(
-						GraphBuilder,
-						View,
-						VisibilityData.ResolveMaskTexture,
-						VisibilityData.TileData,
-						SceneDepthTexture);
-				}
-
-				const bool bWriteDummyData		= View.Family->ViewMode != VMI_VisualizeBuffer && GHairWriteGBufferData == 1;
-				const bool bWritePartialGBuffer = View.Family->ViewMode != VMI_VisualizeBuffer && (GHairWriteGBufferData == 1 || GHairWriteGBufferData == 2);
-				const bool bWriteFullGBuffer	= View.Family->ViewMode == VMI_VisualizeBuffer || (GHairWriteGBufferData == 3);
-				if (bWriteFullGBuffer || bWritePartialGBuffer)
-				{
-					const FSceneTextures& SceneTextures = View.GetSceneTextures();
-					const FRDGTextureRef GBufferATexture = SceneTextures.GBufferA;
-					const FRDGTextureRef GBufferBTexture = SceneTextures.GBufferB;
-					const FRDGTextureRef GBufferCTexture = SceneTextures.GBufferC;
-					const FRDGTextureRef GBufferDTexture = SceneTextures.GBufferD;
-					const FRDGTextureRef GBufferETexture = SceneTextures.GBufferE;
-					if (bWritePartialGBuffer && GBufferATexture && GBufferBTexture)
-					{
-						AddHairVisibilityGBufferWritePass(
-							GraphBuilder,
-							View,
-							bWriteDummyData,
-							VisibilityData.TileData,
-							GBufferATexture,
-							GBufferBTexture,
-							nullptr,
-							nullptr,
-							nullptr,
-							nullptr);
-					}
-					else if (bWriteFullGBuffer && GBufferATexture && GBufferBTexture && GBufferCTexture && SceneDepthTexture)
-					{
-						AddHairVisibilityGBufferWritePass(
-							GraphBuilder,
-							View,
-							bWriteDummyData,
-							VisibilityData.TileData,
-							GBufferATexture,
-							GBufferBTexture,
-							GBufferCTexture,
-							GBufferDTexture,
-							GBufferETexture,
-							SceneDepthTexture);
-					}
-				}
+					bWriteDummyData,
+					VisibilityData.TileData,
+					GBufferATexture,
+					GBufferBTexture,
+					GBufferCTexture,
+					GBufferDTexture,
+					GBufferETexture,
+					SceneDepthTexture);
 			}
 		}
 	}
diff --git a/Engine/Source/Runtime/Renderer/Private/HairStrands/HairStrandsData.h b/Engine/Source/Runtime/Renderer/Private/HairStrands/HairStrandsData.h
--- a/Engine/Source/Runtime/Renderer/Private/HairStrands/HairStrandsData.h
+++ b/Engine/Source/Runtime/Renderer/Private/HairStrands/HairStrandsData.h
@@ -141,6 +141,8 @@ struct FHairStrandsVisibilityData
 	uint32			NodeGroupSize = 0;
 	FVector4f		HairOnlyDepthHZBParameters = FVector4f::Zero();
 
+	uint32				RasterizedInstanceCount = 0;
+	uint32				MaxControlPointCount = 0;
 	FRDGBufferSRVRef	ControlPointsSRV = nullptr;
 	FRDGTextureRef		ControlPointCount = nullptr;
 
diff --git a/Engine/Source/Runtime/Renderer/Private/HairStrands/HairStrandsForwardRaster.cpp b/Engine/Source/Runtime/Renderer/Private/HairStrands/HairStrandsForwardRaster.cpp
new file mode 100644
--- /dev/null
+++ b/Engine/Source/Runtime/Renderer/Private/HairStrands/HairStrandsForwardRaster.cpp
@@ -0,0 +1,632 @@
+// Copyright Epic Games, Inc. All Rights Reserved.
+
+#include "HairStrandsForwardRaster.h"
+
+#include "Shader.h"
+#include "GlobalShader.h"
+#include "ShaderParameters.h"
+#include "ShaderParameterStruct.h"
+#include "SceneTextureParameters.h"
+#include "RenderGraphUtils.h"
+#include "HairStrandsInterface.h"
+#include "SceneRendering.h"
+
+extern int32 GHairVisibilityComputeRaster_Culling;
+extern int32 GHairVisibilityComputeRaster_MaxTiles;
+extern int32 GHairVisibilityComputeRaster_TileSize;
+extern int32 GHairVisibility_NumClassifiers;
+extern int32 GHairVisibilityComputeRaster_NumBinners;
+extern int32 GHairVisibilityComputeRaster_NumRasterizers;
+extern int32 GHairVisibilityComputeRaster_NumRasterizersNaive;
+extern int32 GHairVisibilityComputeRaster_Debug;
+
+void SetUpViewHairRenderInfo(const FViewInfo& ViewInfo, FVector4f& OutHairRenderInfo, uint32& OutHairRenderInfoBits, uint32& OutHairComponents);
+void SetUpViewHairRenderInfo(const FViewInfo& ViewInfo, bool bEnableMSAA, FVector4f& OutHairRenderInfo, uint32& OutHairRenderInfoBits, uint32& OutHairComponents);
+
+
+/////////////////////////////////////////////////////////////////////////////////////////
+// Culling pass
+
+class FHairCullSegmentCS : public FGlobalShader
+{
+	DECLARE_GLOBAL_SHADER(FHairCullSegmentCS);
+	SHADER_USE_PARAMETER_STRUCT(FHairCullSegmentCS, FGlobalShader);
+
+	BEGIN_SHADER_PARAMETER_STRUCT(FParameters, )
+		SHADER_PARAMETER(FVector2f, OutputResolution)
+		SHADER_PARAMETER(uint32, HairMaterialId)
+		SHADER_PARAMETER(uint32, ControlPointCount)
+		SHADER_PARAMETER(uint32, CurveCount)
+		SHADER_PARAMETER_STRUCT_INCLUDE(FHairStrandsInstanceParameters, HairInstance)
+		SHADER_PARAMETER_STRUCT_REF(FViewUniformShaderParameters, ViewUniformBuffer)
+		SHADER_PARAMETER_RDG_TEXTURE(Texture2D, SceneDepthTexture)
+		SHADER_PARAMETER_RDG_BUFFER_UAV(RWStructuredBuffer<FPackedHairVis>, RWHairVis)
+		SHADER_PARAMETER_RDG_BUFFER_UAV(RWBuffer<uint2>, RWCoord)
+		SHADER_PARAMETER_RDG_TEXTURE_UAV(RWTexture2D<uint>, RWCounter)
+		SHADER_PARAMETER_RDG_BUFFER_UAV(RWBuffer<float4>, RWPoints)
+
+		END_SHADER_PARAMETER_STRUCT()
+
+public:
+	static bool ShouldCompilePermutation(const FGlobalShaderPermutationParameters& Parameters) { return IsHairStrandsSupported(EHairStrandsShaderType::Strands, Parameters.Platform); }
+	static void ModifyCompilationEnvironment(const FGlobalShaderPermutationParameters& Parameters, FShaderCompilerEnvironment& OutEnvironment)
+	{
+		FGlobalShader::ModifyCompilationEnvironment(Parameters, OutEnvironment);
+		OutEnvironment.SetDefine(TEXT("SHADER_RASTERCOMPUTE_CULL"), 1);
+	}
+};
+
+IMPLEMENT_GLOBAL_SHADER(FHairCullSegmentCS, "/Engine/Private/HairStrands/HairStrandsForwardCulling.usf", "CSMain", SF_Compute);
+
+///////////////////////////////////////////////////////////////////////////////////////////////////
+// Raster compute common parameters
+
+BEGIN_SHADER_PARAMETER_STRUCT(FRasterComputeForwardCommonParameters, )
+	SHADER_PARAMETER(FIntPoint, OutputResolution)
+	SHADER_PARAMETER(FVector2f, OutputResolutionf)
+
+	SHADER_PARAMETER(FIntPoint, BinTileRes)
+	SHADER_PARAMETER(FIntPoint, RasterTileRes)
+
+	SHADER_PARAMETER(uint32,    NumRasterizers)
+	SHADER_PARAMETER(float,     RcpNumRasterizers)
+
+	SHADER_PARAMETER(uint32,    NumBinners)
+	SHADER_PARAMETER(float,     RcpNumBinners)
+
+	SHADER_PARAMETER(uint32, MaxControlPointCount)
+	SHADER_PARAMETER_RDG_BUFFER_SRV(Buffer, ControlPoints)
+	SHADER_PARAMETER_RDG_TEXTURE(Texture2D<uint>, ControlPointCount)
+END_SHADER_PARAMETER_STRUCT()
+
+///////////////////////////////////////////////////////////////////////////////////////////////////
+// Compute depth tile data based on scene data
+
+class FVisiblityRasterComputeForwardPrepareDepthGridCS : public FGlobalShader
+{
+	DECLARE_GLOBAL_SHADER(FVisiblityRasterComputeForwardPrepareDepthGridCS);
+	SHADER_USE_PARAMETER_STRUCT(FVisiblityRasterComputeForwardPrepareDepthGridCS, FGlobalShader);
+
+	BEGIN_SHADER_PARAMETER_STRUCT(FParameters, )
+		SHADER_PARAMETER_STRUCT_INCLUDE(FRasterComputeForwardCommonParameters, Common)
+		SHADER_PARAMETER_RDG_TEXTURE(Texture2D, SceneDepthTexture)
+		SHADER_PARAMETER_RDG_TEXTURE_UAV(RWTexture2D, OutVisTileDepthGrid)
+	END_SHADER_PARAMETER_STRUCT()
+
+public:
+	static bool ShouldCompilePermutation(const FGlobalShaderPermutationParameters& Parameters) { return IsHairStrandsSupported(EHairStrandsShaderType::Strands, Parameters.Platform); }
+	static void ModifyCompilationEnvironment(const FGlobalShaderPermutationParameters& Parameters, FShaderCompilerEnvironment& OutEnvironment)
+	{
+		FGlobalShader::ModifyCompilationEnvironment(Parameters, OutEnvironment);
+		OutEnvironment.SetDefine(TEXT("SHADER_RASTERCOMPUTE_DEPTH_GRID"), 1);
+	}
+};
+
+IMPLEMENT_GLOBAL_SHADER(FVisiblityRasterComputeForwardPrepareDepthGridCS, "/Engine/Private/HairStrands/HairStrandsForwardRaster.usf", "PrepareDepthGridCS", SF_Compute);
+
+/////////////////////////////////////////////////////////////////////////////////////////
+// Bin strands segments
+
+class FVisiblityRasterComputeForwardBinningCS : public FGlobalShader
+{
+	DECLARE_GLOBAL_SHADER(FVisiblityRasterComputeForwardBinningCS);
+	SHADER_USE_PARAMETER_STRUCT(FVisiblityRasterComputeForwardBinningCS, FGlobalShader);
+
+	class FDebug : SHADER_PERMUTATION_BOOL("PERMUTATION_DEBUG");
+	using FPermutationDomain = TShaderPermutationDomain<FDebug>;
+
+	BEGIN_SHADER_PARAMETER_STRUCT(FParameters, )
+		SHADER_PARAMETER_STRUCT_INCLUDE(FRasterComputeForwardCommonParameters, Common)
+		SHADER_PARAMETER_RDG_TEXTURE(Texture2DArray, VisTileBinningGrid)
+		SHADER_PARAMETER_RDG_BUFFER_UAV(RWBuffer, OutVisTilePrims)
+		SHADER_PARAMETER_RDG_TEXTURE_UAV(RWTexture2DArray, OutVisTileBinningGrid)
+		SHADER_PARAMETER_RDG_TEXTURE(Texture2D, VisTileDepthGrid)
+		SHADER_PARAMETER_RDG_BUFFER_UAV(RWBuffer, OutVisTileArgs)
+		SHADER_PARAMETER_RDG_BUFFER_UAV(RWByteAddressBuffer, OutVisTileData)
+		SHADER_PARAMETER_STRUCT_REF(FViewUniformShaderParameters, ViewUniformBuffer)
+		RDG_TEXTURE_ACCESS(VisTileBinningGridTex, ERHIAccess::UAVCompute)
+		SHADER_PARAMETER_STRUCT_INCLUDE(ShaderPrint::FShaderParameters, ShaderPrintParameters)
+	END_SHADER_PARAMETER_STRUCT()
+
+public:
+	static bool ShouldCompilePermutation(const FGlobalShaderPermutationParameters& Parameters) { return IsHairStrandsSupported(EHairStrandsShaderType::Strands, Parameters.Platform); }
+	static void ModifyCompilationEnvironment(const FGlobalShaderPermutationParameters& Parameters, FShaderCompilerEnvironment& OutEnvironment)
+	{
+		FGlobalShader::ModifyCompilationEnvironment(Parameters, OutEnvironment);
+		OutEnvironment.SetDefine(TEXT("SHADER_RASTERCOMPUTE_BINNING"), 1);
+	}
+};
+
+IMPLEMENT_GLOBAL_SHADER(FVisiblityRasterComputeForwardBinningCS, "/Engine/Private/HairStrands/HairStrandsForwardRaster.usf", "BinningCS", SF_Compute);
+
+/////////////////////////////////////////////////////////////////////////////////////////
+// Compacted binned segments into contiguous list 
+
+class FVisiblityRasterComputeForwardCompactionCS : public FGlobalShader
+{
+	DECLARE_GLOBAL_SHADER(FVisiblityRasterComputeForwardCompactionCS);
+	SHADER_USE_PARAMETER_STRUCT(FVisiblityRasterComputeForwardCompactionCS, FGlobalShader);
+
+	class FDebug : SHADER_PERMUTATION_BOOL("PERMUTATION_DEBUG");
+	using FPermutationDomain = TShaderPermutationDomain<FDebug>;
+
+	BEGIN_SHADER_PARAMETER_STRUCT(FParameters, )
+		SHADER_PARAMETER_STRUCT_INCLUDE(FRasterComputeForwardCommonParameters, Common)
+		SHADER_PARAMETER_RDG_BUFFER_SRV(ByteAddressBuffer, InData)
+		SHADER_PARAMETER_RDG_BUFFER_SRV(Buffer, InPrims)
+		SHADER_PARAMETER_RDG_BUFFER_SRV(Buffer, InArgs)
+		SHADER_PARAMETER_RDG_BUFFER_UAV(RWByteAddressBuffer, OutData)
+		SHADER_PARAMETER_RDG_BUFFER_UAV(RWBuffer, OutPrims)
+		SHADER_PARAMETER_RDG_BUFFER_UAV(RWBuffer, OutArgs)
+		SHADER_PARAMETER_STRUCT_REF(FViewUniformShaderParameters, ViewUniformBuffer)
+		SHADER_PARAMETER_STRUCT_INCLUDE(ShaderPrint::FShaderParameters, ShaderPrintParameters)
+	END_SHADER_PARAMETER_STRUCT()
+
+public:
+	static bool ShouldCompilePermutation(const FGlobalShaderPermutationParameters &Parameters) { return IsHairStrandsSupported(EHairStrandsShaderType::Strands, Parameters.Platform); }
+	static void ModifyCompilationEnvironment(const FGlobalShaderPermutationParameters &Parameters, FShaderCompilerEnvironment &OutEnvironment)
+	{
+		FGlobalShader::ModifyCompilationEnvironment(Parameters, OutEnvironment);
+		OutEnvironment.SetDefine(TEXT("SHADER_RASTERCOMPUTE_COMPACTION"), 1);
+	}
+};
+
+IMPLEMENT_GLOBAL_SHADER(FVisiblityRasterComputeForwardCompactionCS, "/Engine/Private/HairStrands/HairStrandsForwardRaster.usf", "CompactionCS", SF_Compute);
+
+/////////////////////////////////////////////////////////////////////////////////////////
+// Rasterized binned & shaded segements
+
+class FVisiblityRasterComputeForwardRasterizeCS : public FGlobalShader
+{
+	DECLARE_GLOBAL_SHADER(FVisiblityRasterComputeForwardRasterizeCS);
+	SHADER_USE_PARAMETER_STRUCT(FVisiblityRasterComputeForwardRasterizeCS, FGlobalShader);
+
+	class FDebug : SHADER_PERMUTATION_BOOL("PERMUTATION_DEBUG");
+	using FPermutationDomain = TShaderPermutationDomain<FDebug>;
+
+	BEGIN_SHADER_PARAMETER_STRUCT(FParameters, )
+		SHADER_PARAMETER_STRUCT_INCLUDE(FRasterComputeForwardCommonParameters, Common)
+		SHADER_PARAMETER(FIntPoint, SampleLightingViewportResolution)
+		SHADER_PARAMETER_RDG_TEXTURE(Texture2D, SampleLightingTexture)
+
+		SHADER_PARAMETER_RDG_TEXTURE(Texture2D, SceneDepthTexture)
+		SHADER_PARAMETER_RDG_BUFFER_SRV(Buffer, VisTilePrims)
+		SHADER_PARAMETER_RDG_BUFFER_SRV(Buffer, VisTileArgs)
+		SHADER_PARAMETER_RDG_BUFFER_SRV(ByteAddressBuffer, VisTileData)
+		SHADER_PARAMETER_RDG_TEXTURE_UAV(RWTexture2D, OutHairCountTexture_ForDebug)
+		SHADER_PARAMETER_RDG_TEXTURE_UAV(RWTexture2D, OutHairPixelCountPerTile_ForDebug)
+		SHADER_PARAMETER_RDG_TEXTURE_UAV(RWTexture2D, OutSceneColorTexture)
+		SHADER_PARAMETER_RDG_BUFFER_UAV(RWStructuredBuffer, RWWorkCounter)
+		SHADER_PARAMETER_STRUCT_INCLUDE(ShaderPrint::FShaderParameters, ShaderPrintParameters)
+		SHADER_PARAMETER_STRUCT_REF(FViewUniformShaderParameters, ViewUniformBuffer)
+	END_SHADER_PARAMETER_STRUCT()
+
+public:
+	static bool ShouldCompilePermutation(const FGlobalShaderPermutationParameters& Parameters) { return IsHairStrandsSupported(EHairStrandsShaderType::Strands, Parameters.Platform); }
+	static void ModifyCompilationEnvironment(const FGlobalShaderPermutationParameters& Parameters, FShaderCompilerEnvironment& OutEnvironment)
+	{
+		FGlobalShader::ModifyCompilationEnvironment(Parameters, OutEnvironment);
+		OutEnvironment.SetDefine(TEXT("SHADER_RASTERCOMPUTE_RASTER"), 1);
+
+		OutEnvironment.CompilerFlags.Add(CFLAG_AllowTypedUAVLoads);
+	}
+};
+
+IMPLEMENT_GLOBAL_SHADER(FVisiblityRasterComputeForwardRasterizeCS, "/Engine/Private/HairStrands/HairStrandsForwardRaster.usf", "RasterCS", SF_Compute);
+
+/////////////////////////////////////////////////////////////////////////////////////////
+// Debug pass
+
+class FVisiblityRasterComputeForwardDebugCS : public FGlobalShader
+{
+	DECLARE_GLOBAL_SHADER(FVisiblityRasterComputeForwardDebugCS);
+	SHADER_USE_PARAMETER_STRUCT(FVisiblityRasterComputeForwardDebugCS, FGlobalShader);
+
+	using FPermutationDomain = TShaderPermutationDomain<>;
+
+	BEGIN_SHADER_PARAMETER_STRUCT(FParameters, )
+		SHADER_PARAMETER_STRUCT_INCLUDE(FRasterComputeForwardCommonParameters, Common)
+		SHADER_PARAMETER(uint32, InstanceCount)
+		SHADER_PARAMETER_RDG_TEXTURE(Texture2D, HairCountTexture_ForDebug)
+		SHADER_PARAMETER_RDG_TEXTURE(Texture2D, HairPixelCountPerTile_ForDebug)
+		SHADER_PARAMETER_RDG_BUFFER_SRV(Buffer, VisTileArgs)
+		SHADER_PARAMETER_RDG_BUFFER_SRV(ByteAddressBuffer, VisTileData)
+
+		SHADER_PARAMETER_STRUCT_REF(FViewUniformShaderParameters, ViewUniformBuffer)
+		SHADER_PARAMETER_STRUCT_INCLUDE(ShaderPrint::FShaderParameters, ShaderPrintParameters)
+	END_SHADER_PARAMETER_STRUCT()
+
+public:
+	static bool ShouldCompilePermutation(const FGlobalShaderPermutationParameters& Parameters) { return IsHairStrandsSupported(EHairStrandsShaderType::Strands, Parameters.Platform); }
+	static void ModifyCompilationEnvironment(const FGlobalShaderPermutationParameters& Parameters, FShaderCompilerEnvironment& OutEnvironment)
+	{
+		FGlobalShader::ModifyCompilationEnvironment(Parameters, OutEnvironment);
+		OutEnvironment.SetDefine(TEXT("SHADER_RASTERCOMPUTE_DEBUG"), 1);
+	}
+};
+
+IMPLEMENT_GLOBAL_SHADER(FVisiblityRasterComputeForwardDebugCS, "/Engine/Private/HairStrands/HairStrandsForwardRaster.usf", "MainCS", SF_Compute);
+
+
+/////////////////////////////////////////////////////////////////////////////////////////
+// Culling pass
+
+FRasterForwardCullingOutput AddHairStrandsForwardCullingPass(
+	FRDGBuilder& GraphBuilder,
+	const FViewInfo& ViewInfo,
+	const FHairStrandsMacroGroupDatas& MacroGroupDatas,
+	const FIntPoint& InResolution,
+	const FRDGTextureRef SceneDepthTexture,
+	bool bSupportCulling,
+	bool bForceRegister)
+{	
+	// Compute maximum number of PrimIDs
+	uint32 MaxNumPrimIDs = 0;
+	for (const FHairStrandsMacroGroupData &MacroGroup : MacroGroupDatas)
+	{
+		for (const FHairStrandsMacroGroupData::PrimitiveInfo &PrimitiveInfo : MacroGroup.PrimitivesInfos)
+		{
+			// If a groom is not visible in primary view, but visible in shadow view, its PrimitiveInfo.Mesh will be null.
+			if (PrimitiveInfo.Mesh == nullptr || PrimitiveInfo.Mesh->Elements.Num() == 0)
+			{
+				continue;
+			}
+
+			const FHairGroupPublicData* HairGroupPublicData = reinterpret_cast<const FHairGroupPublicData*>(PrimitiveInfo.Mesh->Elements[0].VertexFactoryUserData);
+			check(HairGroupPublicData);
+
+			const FHairGroupPublicData::FVertexFactoryInput& VFInput = HairGroupPublicData->VFInput;
+
+			uint32 VertexCount = VFInput.Strands.VertexCount;
+			const bool bCullingPossible = bSupportCulling && GHairVisibilityComputeRaster_Culling;
+			if (!bCullingPossible)
+			{
+				// calculate current view screen size - which can result in fewer strands rasterized in current view
+				const FSphere BoundsSphere = HairGroupPublicData->ContinuousLODBounds.GetSphere();
+				const float ScreenSize = ComputeBoundsScreenSize(FVector4(BoundsSphere.Center, 1), BoundsSphere.W, ViewInfo);
+
+				VertexCount = HairGroupPublicData->GetActiveStrandsVertexCount(VFInput.Strands.VertexCount, ScreenSize);
+			}
+
+			MaxNumPrimIDs += VertexCount;
+		}
+	}
+
+	if (MaxNumPrimIDs == 0)
+	{
+		return FRasterForwardCullingOutput();
+	}
+	const uint32 NodeVisInBytes = 8; // See HairStrandsVisibilityInternal::NodeVis in HairStrandsVisibility.cpp
+	const EPixelFormat ControlPointFormat = PF_A32B32G32R32F;
+
+	FRasterForwardCullingOutput Out;
+	Out.Resolution		= InResolution;
+	Out.NodeIndex		= GraphBuilder.CreateTexture(FRDGTextureDesc::Create2D(Out.Resolution, PF_R32_UINT, FClearValueBinding::None, TexCreate_UAV | TexCreate_ShaderResource), TEXT("Hair.VisibilityCompactNodeIndex"));
+	Out.NodeVis			= GraphBuilder.CreateBuffer(FRDGBufferDesc::CreateStructuredDesc(NodeVisInBytes, MaxNumPrimIDs), TEXT("Hair.VisibilityNodeVis"));
+	Out.NodeCoord		= GraphBuilder.CreateBuffer(FRDGBufferDesc::CreateBufferDesc(sizeof(uint32), MaxNumPrimIDs), TEXT("Hair.VisibilityNodeCoord"));
+	Out.Points			= GraphBuilder.CreateBuffer(FRDGBufferDesc::CreateBufferDesc(sizeof(uint32)*4u, MaxNumPrimIDs), TEXT("Hair.ControlPoints"));
+	Out.PointsSRV		= GraphBuilder.CreateSRV(Out.Points, ControlPointFormat);
+	Out.PointCount		= GraphBuilder.CreateTexture(FRDGTextureDesc::Create2D(FIntPoint(1,1), PF_R32_UINT, FClearValueBinding::None, TexCreate_UAV | TexCreate_ShaderResource), TEXT("Hair.ControlPointCount"));
+	Out.RasterizedInstanceCount = 0;
+
+	FRDGBufferUAVRef PointsUAV = GraphBuilder.CreateUAV(Out.Points, ControlPointFormat);
+
+	const uint32 ClearValues[4] = { 0u,0u,0u,0u };
+	AddClearUAVPass(GraphBuilder, GraphBuilder.CreateUAV(Out.PointCount), ClearValues);
+	AddClearUAVPass(GraphBuilder, GraphBuilder.CreateUAV(Out.NodeIndex), ClearValues);
+
+	// Create and set the uniform buffer
+	const bool bEnableMSAA = false;
+	TUniformBufferRef<FViewUniformShaderParameters> ViewUniformShaderParameters;
+	SetUpViewHairRenderInfo(ViewInfo, bEnableMSAA, ViewInfo.CachedViewUniformShaderParameters->HairRenderInfo, ViewInfo.CachedViewUniformShaderParameters->HairRenderInfoBits, ViewInfo.CachedViewUniformShaderParameters->HairComponents);
+	ViewUniformShaderParameters = TUniformBufferRef<FViewUniformShaderParameters>::CreateUniformBufferImmediate(*ViewInfo.CachedViewUniformShaderParameters, UniformBuffer_SingleFrame);
+
+	TShaderMapRef<FHairCullSegmentCS> CullShader(ViewInfo.ShaderMap);
+	for (const FHairStrandsMacroGroupData& MacroGroup : MacroGroupDatas)
+	{
+		const FHairStrandsMacroGroupData::TPrimitiveInfos& PrimitiveSceneInfos = MacroGroup.PrimitivesInfos;
+
+		for (const FHairStrandsMacroGroupData::PrimitiveInfo& PrimitiveInfo : PrimitiveSceneInfos)
+		{
+			// If a groom is not visible in primary view, but visible in shadow view, its PrimitiveInfo.Mesh will be null.
+			if (PrimitiveInfo.Mesh == nullptr || PrimitiveInfo.Mesh->Elements.Num() == 0)
+			{
+				continue;
+			}
+
+			const FHairGroupPublicData* HairGroupPublicData = reinterpret_cast<const FHairGroupPublicData*>(PrimitiveInfo.Mesh->Elements[0].VertexFactoryUserData);
+			check(HairGroupPublicData);
+
+			const FHairGroupPublicData::FVertexFactoryInput& VFInput = HairGroupPublicData->VFInput;
+			const FMatrix44d LocalToWorldPrimitiveTransform = VFInput.LocalToWorldTransform.ToMatrixWithScale(); 
+
+			const bool bCullingEnable = false; //bSupportCulling && GHairVisibilityComputeRaster_Culling ? HairGroupPublicData->GetCullingResultAvailable() : false;
+			
+			// calculate current view screen size - which can result in fewer strands rasterized in current view
+			const FSphere BoundsSphere = HairGroupPublicData->ContinuousLODBounds.GetSphere();
+			const float ScreenSize = ComputeBoundsScreenSize(FVector4(BoundsSphere.Center, 1), BoundsSphere.W, ViewInfo);
+			const uint32 CurveCount = VFInput.Strands.CurveCount; //HairGroupPublicData->GetActiveStrandsVertexCount(VFInput.Strands.VertexCount, ScreenSize);
+			const uint32 VertexCount = VFInput.Strands.VertexCount; //HairGroupPublicData->GetActiveStrandsVertexCount(VFInput.Strands.VertexCount, ScreenSize);
+
+			// Curve version
+			#if 1
+			FHairCullSegmentCS::FParameters* Parameters = GraphBuilder.AllocParameters<FHairCullSegmentCS::FParameters>();
+			Parameters->OutputResolution		= Out.Resolution;
+			Parameters->HairMaterialId			= PrimitiveInfo.MaterialId;
+			Parameters->ControlPointCount		= VertexCount;
+			Parameters->CurveCount				= CurveCount;
+			Parameters->RWHairVis				= GraphBuilder.CreateUAV(Out.NodeVis);
+			Parameters->RWCoord					= GraphBuilder.CreateUAV(Out.NodeCoord, PF_R16G16_UINT);
+			Parameters->RWCounter				= GraphBuilder.CreateUAV(Out.PointCount);
+			Parameters->RWPoints				= PointsUAV;
+			Parameters->HairInstance			= GetHairStrandsInstanceParameters(GraphBuilder, ViewInfo, HairGroupPublicData, bCullingEnable, bForceRegister);
+			Parameters->ViewUniformBuffer		= ViewUniformShaderParameters;
+			Parameters->SceneDepthTexture		= SceneDepthTexture;
+
+			FIntVector DispatchCount = FComputeShaderUtils::GetGroupCount(CurveCount, 1024);
+			FComputeShaderUtils::AddPass(
+				GraphBuilder, 
+				RDG_EVENT_NAME("HairStrands::CullSegment(Curves)"), 
+				CullShader,
+				Parameters, 
+				DispatchCount);
+			#endif
+
+			// Vertex version
+			#if 0
+			FHairCullSegmentCS::FParameters* Parameters = GraphBuilder.AllocParameters<FHairCullSegmentCS::FParameters>();
+			Parameters->OutputResolution		= Out.Resolution;
+			Parameters->HairMaterialId			= PrimitiveInfo.MaterialId;
+			Parameters->ControlPointCount		= VertexCount;
+			Parameters->CurveCount				= CurveCount;
+			Parameters->RWHairVis				= GraphBuilder.CreateUAV(Out.NodeVis);
+			Parameters->RWCoord					= GraphBuilder.CreateUAV(Out.NodeCoord, PF_R16G16_UINT);
+			Parameters->RWCounter				= GraphBuilder.CreateUAV(Out.PointCount);
+			Parameters->RWPoints				= GraphBuilder.CreateUAV(Out.Points);
+			Parameters->HairInstance			= GetHairStrandsInstanceParameters(GraphBuilder, ViewInfo, HairGroupPublicData, bCullingEnable, bForceRegister);
+			Parameters->ViewUniformBuffer		= ViewUniformShaderParameters;
+			Parameters->SceneDepthTexture		= SceneDepthTexture;
+
+			FIntVector DispatchCount = FComputeShaderUtils::GetGroupCount(VertexCount, 1024);
+			FComputeShaderUtils::AddPass(
+				GraphBuilder, 
+				RDG_EVENT_NAME("HairStrands::CullSegment(Points)"), 
+				CullShader,
+				Parameters, 
+				DispatchCount);
+			#endif
+
+			++Out.RasterizedInstanceCount;
+		}
+	}
+
+	return Out;
+}
+
+/////////////////////////////////////////////////////////////////////////////////////////
+// Raster pass 
+
+void AddHairStrandsForwardRasterPass(
+	FRDGBuilder& GraphBuilder,
+	const FViewInfo& ViewInfo,
+	const FIntPoint& InResolution,
+	const FHairStrandsVisibilityData& InData,
+	const FRDGTextureRef SceneDepthTexture,
+	const FRDGTextureRef SceneColorTexture)
+{	
+	const FIntPoint OutResolution = InResolution;
+
+	// See the comment on the GHairVisibilityComputeRaster_MaxTiles declaration for an explanation for the large upper bound.
+	const uint32 NumBinners = FMath::Min(FMath::Max(GHairVisibilityComputeRaster_NumBinners, 1), 256);
+	const uint32 NumRasterizers = FMath::Min(FMath::Max(GHairVisibilityComputeRaster_NumRasterizers, 1), 1024);
+
+	const uint32 BinTileSize = 32;
+	const uint32 RasterTileSize = 8;
+
+	const FIntPoint BinTileRes = FIntPoint(FMath::DivideAndRoundUp(uint32(InResolution.X), BinTileSize), FMath::DivideAndRoundUp(uint32(InResolution.Y), BinTileSize));
+	const FIntPoint RasterTileRes = FIntPoint(FMath::DivideAndRoundUp(uint32(InResolution.X), RasterTileSize), FMath::DivideAndRoundUp(uint32(InResolution.Y), RasterTileSize));
+
+	struct FBinningData
+	{
+		// Binned segments
+		FRDGBufferRef VisTilePrims;
+		FRDGBufferSRVRef VisTilePrimsSRV;
+		
+		FRDGBufferRef VisTileData;
+		FRDGBufferSRVRef VisTileDataSRV;
+
+		FRDGBufferRef VisTileArgs;
+		FRDGBufferUAVRef VisTileArgsUAV;
+		FRDGBufferSRVRef VisTileArgsSRV;
+
+		FRDGTextureRef VisTileBinningGrid;
+		FRDGTextureUAVRef VisTileBinningGridUAV;
+
+		// Binned & compacted segments
+		FRDGBufferRef CompactedVisTilePrims;
+		FRDGBufferRef CompactedVisTileData;
+		FRDGBufferRef CompactedVisTileArgs;
+	};
+
+	FBinningData BinData;
+	{
+		const uint32 BinMaxTiles = FMath::Min(FMath::Max(GHairVisibilityComputeRaster_MaxTiles, 1024), 262144);
+
+		BinData.VisTilePrims = GraphBuilder.CreateBuffer(FRDGBufferDesc::CreateBufferDesc(sizeof(uint32), BinMaxTiles * 1024), TEXT("Hair.VisTilePrims"));
+		BinData.VisTilePrimsSRV = GraphBuilder.CreateSRV(BinData.VisTilePrims, PF_R32_UINT);
+
+		const uint32 EntryCount = 5; // See VT_XXX in .usf: PrimOffset / PrimCount / Coord / MaxIndex / MinMaxDepth
+		FRDGBufferDesc DescTileData = FRDGBufferDesc::CreateStructuredDesc(sizeof(uint32), BinMaxTiles * EntryCount * 2);
+		DescTileData.Usage = EBufferUsageFlags(DescTileData.Usage | BUF_ByteAddressBuffer);
+		BinData.VisTileData = GraphBuilder.CreateBuffer(DescTileData, TEXT("Hair.VisTileData"));
+		BinData.VisTileDataSRV = GraphBuilder.CreateSRV(BinData.VisTileData);
+
+		FRDGBufferDesc VisTileArgsDesc = FRDGBufferDesc::CreateBufferDesc(sizeof(uint32), 4 + 16 + (((BinTileRes.X * BinTileRes.Y) + 31) / 32));
+		BinData.VisTileArgs = GraphBuilder.CreateBuffer(VisTileArgsDesc, TEXT("Hair.VisTileArgs"));
+		BinData.VisTileArgsUAV = GraphBuilder.CreateUAV(BinData.VisTileArgs, PF_R32_UINT);
+		BinData.VisTileArgsSRV = GraphBuilder.CreateSRV(BinData.VisTileArgs, PF_R32_UINT);
+
+		BinData.VisTileBinningGrid = GraphBuilder.CreateTexture(FRDGTextureDesc::Create2DArray(BinTileRes, PF_R32_UINT, FClearValueBinding::None, TexCreate_UAV | TexCreate_ShaderResource, NumBinners * 3), TEXT("Hair.VisTileBinningGrid"));
+		BinData.VisTileBinningGridUAV = GraphBuilder.CreateUAV(BinData.VisTileBinningGrid);
+
+		BinData.CompactedVisTilePrims = GraphBuilder.CreateBuffer(FRDGBufferDesc::CreateBufferDesc(sizeof(uint32), BinMaxTiles * 1024), TEXT("Hair.CompactedVisTilePrims"));
+		BinData.CompactedVisTileData = GraphBuilder.CreateBuffer(DescTileData, TEXT("Hair.CompactedVisTileData"));
+		BinData.CompactedVisTileArgs = GraphBuilder.CreateBuffer(VisTileArgsDesc, TEXT("Hair.CompactedVisTileArgs"));
+	}
+
+	// Create and set the uniform buffer
+	TUniformBufferRef<FViewUniformShaderParameters> ViewUniformShaderParameters;
+	{
+		const bool bEnableMSAA = false;
+		SetUpViewHairRenderInfo(ViewInfo, bEnableMSAA, ViewInfo.CachedViewUniformShaderParameters->HairRenderInfo, ViewInfo.CachedViewUniformShaderParameters->HairRenderInfoBits, ViewInfo.CachedViewUniformShaderParameters->HairComponents);
+		ViewUniformShaderParameters = TUniformBufferRef<FViewUniformShaderParameters>::CreateUniformBufferImmediate(*ViewInfo.CachedViewUniformShaderParameters, UniformBuffer_SingleFrame);
+	}
+
+	// Common parameters
+	FRasterComputeForwardCommonParameters Common;
+	{		
+		Common.OutputResolution = OutResolution;
+		Common.OutputResolutionf = FVector2f(OutResolution.X, OutResolution.Y);
+
+		Common.RasterTileRes = RasterTileRes;
+		Common.NumRasterizers = NumRasterizers;
+		Common.RcpNumRasterizers = 1.0 / NumRasterizers;
+
+		Common.BinTileRes = BinTileRes;
+		Common.NumBinners = NumBinners;
+		Common.RcpNumBinners = 1.0 / NumBinners;
+
+		Common.MaxControlPointCount = InData.MaxControlPointCount;
+		Common.ControlPoints = InData.ControlPointsSRV;
+		Common.ControlPointCount = InData.ControlPointCount;
+	}
+
+	// Fill in tile depth
+	FRDGTextureRef VisTileDepthGrid = nullptr;
+	{
+		VisTileDepthGrid = GraphBuilder.CreateTexture(FRDGTextureDesc::Create2D(BinTileRes, PF_R32_UINT, FClearValueBinding::None, TexCreate_UAV | TexCreate_ShaderResource), TEXT("Hair.VisTileDepthGrid"));
+
+		TShaderMapRef<FVisiblityRasterComputeForwardPrepareDepthGridCS> ComputeShaderPrepareDepthGrid(ViewInfo.ShaderMap);
+		FVisiblityRasterComputeForwardPrepareDepthGridCS::FParameters* Parameters = GraphBuilder.AllocParameters<FVisiblityRasterComputeForwardPrepareDepthGridCS::FParameters>();		
+		Parameters->Common = Common;
+		Parameters->SceneDepthTexture = SceneDepthTexture;
+		Parameters->OutVisTileDepthGrid = GraphBuilder.CreateUAV(VisTileDepthGrid);
+		FComputeShaderUtils::AddPass(GraphBuilder, RDG_EVENT_NAME("HairStrands::VisibilityComputeRasterPrepDepthGrid"), ComputeShaderPrepareDepthGrid, Parameters, FIntVector(BinTileRes.X, BinTileRes.Y, 1));
+	}
+
+	const bool bDebugEnabled = GHairVisibilityComputeRaster_Debug > 0;
+	uint32 TotalPrimitiveInfoCount = 1; // TODO = total number of groups
+
+	if (bDebugEnabled)
+	{
+		ShaderPrint::SetEnabled(true);
+		ShaderPrint::RequestSpaceForTriangles(2 * NumBinners * Common.RasterTileRes.X * Common.RasterTileRes.Y + 2 * NumBinners * 10);
+		ShaderPrint::RequestSpaceForLines(4 * 2 * NumBinners * 128);
+	}
+
+	// Reset buffers
+	{
+		uint32 IndexGridClearValues[4] = { 0x0,0x0,0x0,0x0 };
+		AddClearUAVPass(GraphBuilder, BinData.VisTileBinningGridUAV, IndexGridClearValues);
+		AddClearUAVPass(GraphBuilder, BinData.VisTileArgsUAV, 0u);
+		AddClearUAVPass(GraphBuilder, GraphBuilder.CreateUAV(BinData.CompactedVisTileArgs, PF_R32_UINT), 0u);
+	}
+
+	// Binning pass
+	{
+		FVisiblityRasterComputeForwardBinningCS::FParameters* Parameters = GraphBuilder.AllocParameters<FVisiblityRasterComputeForwardBinningCS::FParameters>();
+		Parameters->Common = Common;
+		Parameters->ViewUniformBuffer = ViewUniformShaderParameters;
+		Parameters->OutVisTilePrims = GraphBuilder.CreateUAV(BinData.VisTilePrims, PF_R32_UINT);
+		Parameters->OutVisTileBinningGrid = BinData.VisTileBinningGridUAV;
+		Parameters->VisTileDepthGrid = VisTileDepthGrid;
+		Parameters->OutVisTileArgs = BinData.VisTileArgsUAV;
+		Parameters->OutVisTileData = GraphBuilder.CreateUAV(BinData.VisTileData);
+		Parameters->VisTileBinningGridTex = BinData.VisTileBinningGrid;
+		ShaderPrint::SetParameters(GraphBuilder, ViewInfo.ShaderPrintData, Parameters->ShaderPrintParameters);
+
+		FVisiblityRasterComputeForwardBinningCS::FPermutationDomain PermutationVector;
+		PermutationVector.Set<FVisiblityRasterComputeForwardBinningCS::FDebug>(bDebugEnabled);
+		TShaderMapRef<FVisiblityRasterComputeForwardBinningCS> ComputeShaderBinning(ViewInfo.ShaderMap, PermutationVector);
+		FComputeShaderUtils::AddPass(
+			GraphBuilder, 
+			RDG_EVENT_NAME("HairStrands::VisibilityComputeRasterBinning"),
+			ComputeShaderBinning,
+			Parameters, 
+			FIntVector(NumBinners, 1, 1));
+	}
+
+	// Compaction
+	{
+		FVisiblityRasterComputeForwardCompactionCS::FParameters* Parameters = GraphBuilder.AllocParameters<FVisiblityRasterComputeForwardCompactionCS::FParameters>();
+		Parameters->Common = Common;
+		Parameters->ViewUniformBuffer = ViewUniformShaderParameters;
+		Parameters->InData  = BinData.VisTileDataSRV;
+		Parameters->InPrims = BinData.VisTilePrimsSRV;
+		Parameters->InArgs  = BinData.VisTileArgsSRV;
+		Parameters->OutData = GraphBuilder.CreateUAV(BinData.CompactedVisTileData);
+		Parameters->OutPrims= GraphBuilder.CreateUAV(BinData.CompactedVisTilePrims, PF_R32_UINT);
+		Parameters->OutArgs = GraphBuilder.CreateUAV(BinData.CompactedVisTileArgs, PF_R32_UINT);
+		ShaderPrint::SetParameters(GraphBuilder, ViewInfo.ShaderPrintData, Parameters->ShaderPrintParameters);
+
+		FVisiblityRasterComputeForwardCompactionCS::FPermutationDomain PermutationVector;
+		PermutationVector.Set<FVisiblityRasterComputeForwardCompactionCS::FDebug>(bDebugEnabled);
+		TShaderMapRef<FVisiblityRasterComputeForwardCompactionCS> ComputeShaderCompaction(ViewInfo.ShaderMap, PermutationVector);
+		FComputeShaderUtils::AddPass(GraphBuilder, RDG_EVENT_NAME("HairStrands::VisibilityComputeRasterCompaction"), ComputeShaderCompaction, Parameters, FIntVector(BinTileRes.X, BinTileRes.Y, 1));
+	}
+
+	// Debug data
+	FRDGTextureRef HairPixelCountPerTile_ForDebug = nullptr;
+	FRDGTextureRef HairCountTexture_ForDebug = nullptr;
+	if (bDebugEnabled)
+	{
+		HairPixelCountPerTile_ForDebug = GraphBuilder.CreateTexture(FRDGTextureDesc::Create2D(RasterTileRes, PF_R32_UINT, FClearValueBinding::None, TexCreate_UAV | TexCreate_ShaderResource | TexCreate_RenderTargetable), TEXT("Hair.HairPixelTileTexture_ForDebug"));
+		HairCountTexture_ForDebug = GraphBuilder.CreateTexture(FRDGTextureDesc::Create2D(OutResolution, PF_R32_UINT, FClearValueBinding::None, TexCreate_UAV | TexCreate_ShaderResource | TexCreate_RenderTargetable), TEXT("Hair.HairCountTexture_ForDebug"));
+
+		uint32 ClearValues[4] = { 0,0,0,0 };
+		AddClearUAVPass(GraphBuilder, GraphBuilder.CreateUAV(HairCountTexture_ForDebug), ClearValues);
+		AddClearUAVPass(GraphBuilder, GraphBuilder.CreateUAV(HairPixelCountPerTile_ForDebug), ClearValues);
+	}
+
+	// Raster pass
+	{
+		FRDGBufferRef WorkCounter = GraphBuilder.CreateBuffer(FRDGBufferDesc::CreateStructuredDesc(sizeof(uint32), 1), TEXT("Hair.RasterWorkCounter"));
+		FRDGBufferUAVRef WorkCounterUAV = GraphBuilder.CreateUAV(WorkCounter, PF_R32_UINT);
+		AddClearUAVPass(GraphBuilder, WorkCounterUAV, 0);
+
+		FVisiblityRasterComputeForwardRasterizeCS::FParameters* Parameters = GraphBuilder.AllocParameters<FVisiblityRasterComputeForwardRasterizeCS::FParameters>();
+		Parameters->Common = Common;
+		Parameters->ViewUniformBuffer = ViewUniformShaderParameters;
+		Parameters->SceneDepthTexture = SceneDepthTexture;
+		Parameters->OutHairCountTexture_ForDebug = HairCountTexture_ForDebug ? GraphBuilder.CreateUAV(HairCountTexture_ForDebug) : nullptr;
+		Parameters->OutHairPixelCountPerTile_ForDebug = HairPixelCountPerTile_ForDebug ? GraphBuilder.CreateUAV(HairPixelCountPerTile_ForDebug) : nullptr;
+		Parameters->SampleLightingViewportResolution = InData.SampleLightingViewportResolution;
+		Parameters->SampleLightingTexture = InData.SampleLightingTexture;
+		Parameters->VisTilePrims = GraphBuilder.CreateSRV(BinData.CompactedVisTilePrims, PF_R32_UINT);
+		Parameters->VisTileArgs = GraphBuilder.CreateSRV(BinData.CompactedVisTileArgs, PF_R32_UINT);
+		Parameters->VisTileData = GraphBuilder.CreateSRV(BinData.CompactedVisTileData);
+		Parameters->OutSceneColorTexture = GraphBuilder.CreateUAV(SceneColorTexture);
+		Parameters->RWWorkCounter = WorkCounterUAV;
+		ShaderPrint::SetParameters(GraphBuilder, ViewInfo.ShaderPrintData, Parameters->ShaderPrintParameters);
+
+		FVisiblityRasterComputeForwardRasterizeCS::FPermutationDomain PermutationVector;
+		PermutationVector.Set<FVisiblityRasterComputeForwardRasterizeCS::FDebug>(bDebugEnabled);
+		TShaderMapRef<FVisiblityRasterComputeForwardRasterizeCS> ComputeShaderRaster(ViewInfo.ShaderMap, PermutationVector);
+		FComputeShaderUtils::AddPass(GraphBuilder, RDG_EVENT_NAME("HairStrands::VisibilityComputeRasterRaster(tiled)"), ComputeShaderRaster, Parameters, FIntVector(NumRasterizers, 1, 1));
+	}
+
+	if (bDebugEnabled)
+	{
+		FVisiblityRasterComputeForwardDebugCS::FParameters* Parameters = GraphBuilder.AllocParameters<FVisiblityRasterComputeForwardDebugCS::FParameters>();
+		Parameters->Common = Common;
+		Parameters->InstanceCount = InData.RasterizedInstanceCount;
+		Parameters->VisTileData = GraphBuilder.CreateSRV(BinData.CompactedVisTileData);
+		Parameters->VisTileArgs = BinData.VisTileArgsSRV;
+		Parameters->HairCountTexture_ForDebug = HairCountTexture_ForDebug;
+		Parameters->HairPixelCountPerTile_ForDebug = HairPixelCountPerTile_ForDebug;
+		ShaderPrint::SetParameters(GraphBuilder, ViewInfo.ShaderPrintData, Parameters->ShaderPrintParameters);
+
+		TShaderMapRef<FVisiblityRasterComputeForwardDebugCS> DebugComputeShader(ViewInfo.ShaderMap);
+		FComputeShaderUtils::AddPass(GraphBuilder, RDG_EVENT_NAME("HairStrands::VisibilityComputeRaster(Debug)"), DebugComputeShader, Parameters, FIntVector(FMath::DivideAndRoundUp(Common.BinTileRes.X, 8), FMath::DivideAndRoundUp(Common.BinTileRes.Y, 8), 1));
+	}
+}
diff --git a/Engine/Source/Runtime/Renderer/Private/HairStrands/HairStrandsForwardRaster.h b/Engine/Source/Runtime/Renderer/Private/HairStrands/HairStrandsForwardRaster.h
new file mode 100644
--- /dev/null
+++ b/Engine/Source/Runtime/Renderer/Private/HairStrands/HairStrandsForwardRaster.h
@@ -0,0 +1,44 @@
+// Copyright Epic Games, Inc. All Rights Reserved.
+
+/*=============================================================================
+	HairVisibilityRendering.h: Hair strands visibility buffer implementation.
+=============================================================================*/
+
+#pragma once
+
+#include "CoreMinimal.h"
+#include "RendererInterface.h"
+#include "RenderGraphResources.h"
+#include "HairStrandsData.h"
+
+class FViewInfo;
+class FInstanceCullingManager;
+
+struct FRasterForwardCullingOutput
+{
+	FIntPoint			Resolution;
+	uint32				RasterizedInstanceCount = 0;
+	FRDGTextureRef		NodeIndex = nullptr;
+	FRDGBufferRef		NodeCoord = nullptr;
+	FRDGBufferRef		NodeVis = nullptr;
+	FRDGTextureRef		PointCount = nullptr;
+	FRDGBufferRef		Points = nullptr;
+	FRDGBufferSRVRef	PointsSRV = nullptr;
+};
+
+FRasterForwardCullingOutput AddHairStrandsForwardCullingPass(
+	FRDGBuilder& GraphBuilder,
+	const FViewInfo& ViewInfo,
+	const FHairStrandsMacroGroupDatas& MacroGroupDatas,
+	const FIntPoint& InResolution,
+	const FRDGTextureRef SceneDepthTexture,
+	bool bSupportCulling,
+	bool bForceRegister);
+
+void AddHairStrandsForwardRasterPass(
+	FRDGBuilder& GraphBuilder,
+	const FViewInfo& ViewInfo,
+	const FIntPoint& InResolution,
+	const FHairStrandsVisibilityData& InData,
+	const FRDGTextureRef SceneDepthTexture,
+	const FRDGTextureRef SceneColorTexture);
\ No newline at end of file
diff --git a/Engine/Source/Runtime/Renderer/Private/HairStrands/HairStrandsInterface.cpp b/Engine/Source/Runtime/Renderer/Private/HairStrands/HairStrandsInterface.cpp
--- a/Engine/Source/Runtime/Renderer/Private/HairStrands/HairStrandsInterface.cpp
+++ b/Engine/Source/Runtime/Renderer/Private/HairStrands/HairStrandsInterface.cpp
@@ -492,7 +492,12 @@ bool IsHairStrandContinuousDecimationReorderingEnabled()
 
 bool IsHairVisibilityComputeRasterEnabled()
 {
-	return CVarHairStrandsVisibilityComputeRaster.GetValueOnAnyThread() > 0;
+	return CVarHairStrandsVisibilityComputeRaster.GetValueOnAnyThread() == 1;
+}
+
+bool IsHairVisibilityComputeRasterForwardEnabled()
+{
+	return CVarHairStrandsVisibilityComputeRaster.GetValueOnAnyThread() == 2;
 }
 
 bool IsHairVisibilityComputeRasterContinuousLODEnabled()
diff --git a/Engine/Source/Runtime/Renderer/Private/HairStrands/HairStrandsVisibility.cpp b/Engine/Source/Runtime/Renderer/Private/HairStrands/HairStrandsVisibility.cpp
--- a/Engine/Source/Runtime/Renderer/Private/HairStrands/HairStrandsVisibility.cpp
+++ b/Engine/Source/Runtime/Renderer/Private/HairStrands/HairStrandsVisibility.cpp
@@ -5,6 +5,7 @@
 #include "HairStrandsUtils.h"
 #include "HairStrandsInterface.h"
 #include "HairStrandsTile.h"
+#include "HairStrandsForwardRaster.h"
 
 #include "Shader.h"
 #include "GlobalShader.h"
@@ -148,16 +149,12 @@ enum EHairVisibilityRenderMode
 	HairVisibilityRenderMode_MSAA_Visibility,
 	HairVisibilityRenderMode_TransmittanceAndHairCount,
 	HairVisibilityRenderMode_ComputeRaster,
+	HairVisibilityRenderMode_ComputeRasterForward,
 	HairVisibilityRenderModeCount
 };
 
-inline bool DoesSupportRasterCompute()
-{
-	//for now always return true 
-	return true; //GRHISupportsAtomicUInt64; //TODO: re-enable support for 64 bit visibility buffers
-}
-
 bool IsHairVisibilityComputeRasterEnabled();
+bool IsHairVisibilityComputeRasterForwardEnabled();
 
 inline EHairVisibilityRenderMode GetHairVisibilityRenderMode()
 {
@@ -165,10 +162,14 @@ inline EHairVisibilityRenderMode GetHairVisibilityRenderMode()
 	{
 		return HairVisibilityRenderMode_PPLL;
 	}
-	else if (IsHairVisibilityComputeRasterEnabled() && DoesSupportRasterCompute())
+	else if (IsHairVisibilityComputeRasterEnabled())
 	{
 		return HairVisibilityRenderMode_ComputeRaster;
 	}
+	else if (IsHairVisibilityComputeRasterForwardEnabled())
+	{
+		return HairVisibilityRenderMode_ComputeRasterForward;
+	}
 	else
 	{
 		return HairVisibilityRenderMode_MSAA_Visibility;
@@ -253,6 +254,7 @@ inline uint32 GetMeanSamplePerPixel()
 	const uint32 SamplePerPixel = GetMaxSamplePerPixel();
 	switch (GetHairVisibilityRenderMode())
 	{
+	case HairVisibilityRenderMode_ComputeRasterForward:
 	case HairVisibilityRenderMode_ComputeRaster:
 		return FMath::Max(1, FMath::FloorToInt(SamplePerPixel * FMath::Clamp(GHairVisibilityCompute_MeanSamplePerPixel, 0.f, 1.f)));
 	case HairVisibilityRenderMode_MSAA_Visibility:
@@ -287,7 +289,7 @@ static uint32 GetTotalSampleCountForAllocation(FIntPoint Resolution)
 	return Resolution.X * Resolution.Y * GetMeanSamplePerPixel();
 }
 
-static void SetUpViewHairRenderInfo(const FViewInfo& ViewInfo, bool bEnableMSAA, FVector4f& OutHairRenderInfo, uint32& OutHairRenderInfoBits, uint32& OutHairComponents)
+void SetUpViewHairRenderInfo(const FViewInfo& ViewInfo, bool bEnableMSAA, FVector4f& OutHairRenderInfo, uint32& OutHairRenderInfoBits, uint32& OutHairComponents)
 {
 	FVector2f PixelVelocity(1.f / (ViewInfo.ViewRect.Width() * 2), 1.f / (ViewInfo.ViewRect.Height() * 2));
 	const float VelocityMagnitudeScale = FMath::Clamp(CVarHairVelocityMagnitudeScale.GetValueOnAnyThread(), 0, 512) * FMath::Min(PixelVelocity.X, PixelVelocity.Y);
@@ -2728,7 +2730,8 @@ void AddMeshDrawTransitionPass(
 			ExternalAccessQueue.Add(CulledVertexRadiusScaleBuffer);
 			ExternalAccessQueue.Add(DrawIndirectBuffer, ERHIAccess::IndirectArgs);
 
-			if (GetHairVisibilityRenderMode() != HairVisibilityRenderMode_ComputeRaster)
+			const EHairVisibilityRenderMode RasterMode = GetHairVisibilityRenderMode();
+			if (RasterMode != HairVisibilityRenderMode_ComputeRaster && RasterMode != HairVisibilityRenderMode_ComputeRasterForward)
 			{
 				VFInput.Strands.PositionBuffer				= FRDGImportedBuffer();
 				VFInput.Strands.PrevPositionBuffer			= FRDGImportedBuffer();
@@ -3190,8 +3193,6 @@ static FRasterComputeOutput AddVisibilityComputeRasterPass(
 	bool bSupportCulling,
 	bool bForceRegister)
 {	
-	check(DoesSupportRasterCompute());
-
 	FRasterComputeOutput Out;
 	Out.Resolution = InResolution;
 
@@ -3964,7 +3965,7 @@ void DrawHitProxies(
 	// Proxy rendering is only supported/compatible with MSAA-visibility rendering. 
 	// PPLL is not supported, but it is supposed to be used only for final render.
 	const EHairVisibilityRenderMode RenderMode = GetHairVisibilityRenderMode();
-	if (RenderMode != HairVisibilityRenderMode_MSAA_Visibility && RenderMode != HairVisibilityRenderMode_ComputeRaster)
+	if (RenderMode != HairVisibilityRenderMode_MSAA_Visibility && RenderMode != HairVisibilityRenderMode_ComputeRaster && RenderMode != HairVisibilityRenderMode_ComputeRasterForward)
 	{
 		return;
 	}
@@ -4201,11 +4202,15 @@ void RenderHairStrandsVisibilityBuffer(
 			const FIntPoint Resolution = SceneDepthTexture->Desc.Extent;
 
 			const EHairVisibilityRenderMode RenderMode = GetHairVisibilityRenderMode();
-			check(RenderMode == HairVisibilityRenderMode_MSAA_Visibility || RenderMode == HairVisibilityRenderMode_PPLL || RenderMode == HairVisibilityRenderMode_ComputeRaster);
+			check(
+				RenderMode == HairVisibilityRenderMode_MSAA_Visibility || 
+				RenderMode == HairVisibilityRenderMode_PPLL || 
+				RenderMode == HairVisibilityRenderMode_ComputeRaster ||
+				RenderMode == HairVisibilityRenderMode_ComputeRasterForward);
 
 			const bool bRunColorAndDepthPatching = 
 				(SceneMaterial0 && SceneMaterial1) &&
-				SceneColorTexture && RenderMode != HairVisibilityRenderMode_ComputeRaster;
+				SceneColorTexture && (RenderMode != HairVisibilityRenderMode_ComputeRaster || RenderMode != HairVisibilityRenderMode_ComputeRasterForward);
 
 			FRDGTextureRef HairOnlyDepthTexture = GraphBuilder.CreateTexture(SceneDepthTexture->Desc, TEXT("Hair.HairOnlyDepthTexture"));
 			FRDGTextureRef CoverageTexture = nullptr;
@@ -4214,124 +4219,223 @@ void RenderHairStrandsVisibilityBuffer(
 			FRDGBufferRef  CompactNodeVis = nullptr;
 			FRDGTextureRef NodeCounter = nullptr;
 
-			if (RenderMode == HairVisibilityRenderMode_ComputeRaster)
+			if (RenderMode == HairVisibilityRenderMode_ComputeRasterForward)
 			{
-				FRasterComputeOutput RasterOutput = AddVisibilityComputeRasterPass(
+				FRasterForwardCullingOutput Out = AddHairStrandsForwardCullingPass(
 					GraphBuilder,
 					View,
 					MacroGroupDatas,
 					Resolution,
-					VisibilityData.MaxSampleCount,
 					SceneDepthTexture,
 					true,
 					false);
 
-				// Generate Tile data
-				check(RasterOutput.PrimMatTexture);
-				VisibilityData.TileData = AddHairStrandsGenerateTilesPass(GraphBuilder, View, RasterOutput.HairCountTexture);
-
+				if (Out.NodeCoord == nullptr)
 				{
-					{
-						FRDGBufferRef CompactNodeCoord;
-						FRDGBufferRef IndirectArgsBuffer;
-						FRDGTextureRef ResolveMaskTexture = nullptr;
-						AddHairVisibilityCompactionComputeRasterPass(
-							GraphBuilder,
-							View,
-							VisibilityData.NodeGroupSize,
-							VisibilityData.MaxSampleCount,
-							RasterOutput,
-							VisibilityData.TileData,
-							NodeCounter,
-							CompactNodeIndex,
-							CompactNodeVis,
-							CompactNodeCoord,
-							CoverageTexture,
-							IndirectArgsBuffer,
-							VisibilityData.MaxNodeCount);
+					return;
+				}
+				VisibilityData.MaxNodeCount = Out.NodeCoord->Desc.NumElements;
 
+				VisibilityData.MaxControlPointCount = Out.NodeCoord->Desc.NumElements;
+				VisibilityData.ControlPointsSRV = Out.PointsSRV;
+				VisibilityData.ControlPointCount = Out.PointCount;
+				VisibilityData.RasterizedInstanceCount = Out.RasterizedInstanceCount;
 
-						// Evaluate material based on the visiblity pass result
-						// Output both complete sample data + per-sample velocity
-						FMaterialPassOutput PassOutput = AddHairMaterialPass(
-							GraphBuilder,
-							Scene,
-							&View,
-							false,
-							MacroGroupDatas,
-							InstanceCullingManager,
-							VisibilityData.NodeGroupSize,
-							CompactNodeIndex,
-							CompactNodeVis,
-							CompactNodeCoord,
-							NodeCounter,
-							IndirectArgsBuffer);
+				CompactNodeIndex = Out.NodeIndex;
+				NodeCounter = Out.PointCount;
+				CompactNodeVis = Out.NodeVis;
+				FRDGBufferRef CompactNodeCoord = Out.NodeCoord;
 
-						// Merge per-sample velocity into the scene velocity buffer
-						AddHairVelocityPass(
-							GraphBuilder,
-							View,
-							MacroGroupDatas,
-							VisibilityData.TileData,
-							CoverageTexture,
-							CompactNodeIndex,
-							CompactNodeVis,
-							PassOutput.NodeVelocity,
-							SceneVelocityTexture,
-							ResolveMaskTexture);
+				FRDGBufferRef IndirectArgsBuffer = AddCopyIndirectArgPass(GraphBuilder, &View, VisibilityData.NodeGroupSize, 1, Out.PointCount);
 
-						CompactNodeData = PassOutput.NodeData;
+				// Generate Tile data
+				VisibilityData.TileData = AddHairStrandsGenerateTilesPass(GraphBuilder, View,Out.Resolution); // Tile for all pixels == no tile accelaration
 
-						VisibilityData.SampleLightingViewportResolution = PassOutput.SampleLightingTexture->Desc.Extent;
-						VisibilityData.SampleLightingTexture = PassOutput.SampleLightingTexture;
-						VisibilityData.NodeIndex = CompactNodeIndex;
-						VisibilityData.CoverageTexture = CoverageTexture;
-						VisibilityData.HairOnlyDepthTexture = HairOnlyDepthTexture;
-						VisibilityData.NodeData = CompactNodeData;
-						VisibilityData.NodeVisData = CompactNodeVis;
-						VisibilityData.NodeCoord = CompactNodeCoord;
-						VisibilityData.NodeIndirectArg = IndirectArgsBuffer;
-						VisibilityData.NodeCount = NodeCounter;
-						VisibilityData.ResolveMaskTexture = ResolveMaskTexture;	
-					}
+				// Evaluate material based on the visibility pass result
+				// Output both complete sample data + per-sample velocity
+				FMaterialPassOutput PassOutput = AddHairMaterialPass(
+					GraphBuilder,
+					Scene,
+					&View,
+					false,
+					MacroGroupDatas,
+					InstanceCullingManager,
+					VisibilityData.NodeGroupSize,
+					CompactNodeIndex,
+					CompactNodeVis,
+					CompactNodeCoord,
+					NodeCounter,
+					IndirectArgsBuffer);
 
-					// For fully covered pixels, write: 
-					// * black color into the scene color
-					// * closest depth
-					// * unlit shading model ID 
-					if (bRunColorAndDepthPatching)
-					{
-						AddHairMaterialDataPatchPass(
-							GraphBuilder,
-							View,
-							VisibilityData.TileData,
-							CoverageTexture,
-							CompactNodeIndex,
-							CompactNodeData,
-							SceneMaterial0,
-							SceneMaterial1,
-							SceneColorTexture,
-							SceneDepthTexture,
-							VisibilityData.LightChannelMaskTexture);
-					}
+				// Dummy coverage texture
+				CoverageTexture = GraphBuilder.CreateTexture(FRDGTextureDesc::Create2D(Resolution, FHairStrandsVisibilityData::CoverageFormat, FClearValueBinding::None, TexCreate_UAV | TexCreate_ShaderResource), TEXT("Hair.CoverageTexture"));
+				AddClearUAVPass(GraphBuilder, GraphBuilder.CreateUAV(CoverageTexture), 0.f);
 
-					AddHairOnlyDepthPass(
+				CompactNodeData = PassOutput.NodeData;
+
+				VisibilityData.SampleLightingViewportResolution = PassOutput.SampleLightingTexture->Desc.Extent;
+				VisibilityData.SampleLightingTexture = PassOutput.SampleLightingTexture;
+				VisibilityData.NodeIndex = CompactNodeIndex;
+				VisibilityData.CoverageTexture = CoverageTexture;
+				VisibilityData.HairOnlyDepthTexture = HairOnlyDepthTexture;
+				VisibilityData.NodeData = CompactNodeData;
+				VisibilityData.NodeVisData = CompactNodeVis;
+				VisibilityData.NodeCoord = CompactNodeCoord;
+				VisibilityData.NodeIndirectArg = IndirectArgsBuffer;
+				VisibilityData.NodeCount = NodeCounter;
+				VisibilityData.ResolveMaskTexture = nullptr;
+
+				// For fully covered pixels, write: 
+				// * black color into the scene color
+				// * closest depth
+				// * unlit shading model ID 
+				if (bRunColorAndDepthPatching)
+				{
+					AddHairMaterialDataPatchPass(
 						GraphBuilder,
 						View,
 						VisibilityData.TileData,
 						CoverageTexture,
 						CompactNodeIndex,
 						CompactNodeData,
-						HairOnlyDepthTexture);
+						SceneMaterial0,
+						SceneMaterial1,
+						SceneColorTexture,
+						SceneDepthTexture,
+						VisibilityData.LightChannelMaskTexture);
+				}
 
-					AddHairOnlyHZBPass(
+				AddHairOnlyDepthPass(
+					GraphBuilder,
+					View,
+					VisibilityData.TileData,
+					CoverageTexture,
+					CompactNodeIndex,
+					CompactNodeData,
+					HairOnlyDepthTexture);
+
+				AddHairOnlyHZBPass(
+					GraphBuilder,
+					View,
+					HairOnlyDepthTexture,
+					VisibilityData.HairOnlyDepthHZBParameters,
+					VisibilityData.HairOnlyDepthClosestHZBTexture,
+					VisibilityData.HairOnlyDepthFurthestHZBTexture);
+			}
+			else if (RenderMode == HairVisibilityRenderMode_ComputeRaster)
+			{
+				FRasterComputeOutput RasterOutput = AddVisibilityComputeRasterPass(
+					GraphBuilder,
+					View,
+					MacroGroupDatas,
+					Resolution,
+					VisibilityData.MaxSampleCount,
+					SceneDepthTexture,
+					true,
+					false);
+
+				// Generate Tile data
+				check(RasterOutput.PrimMatTexture);
+				VisibilityData.TileData = AddHairStrandsGenerateTilesPass(GraphBuilder, View, RasterOutput.HairCountTexture);
+
+				FRDGBufferRef CompactNodeCoord;
+				FRDGBufferRef IndirectArgsBuffer;
+				FRDGTextureRef ResolveMaskTexture = nullptr;
+				AddHairVisibilityCompactionComputeRasterPass(
+					GraphBuilder,
+					View,
+					VisibilityData.NodeGroupSize,
+					VisibilityData.MaxSampleCount,
+					RasterOutput,
+					VisibilityData.TileData,
+					NodeCounter,
+					CompactNodeIndex,
+					CompactNodeVis,
+					CompactNodeCoord,
+					CoverageTexture,
+					IndirectArgsBuffer,
+					VisibilityData.MaxNodeCount);
+
+
+				// Evaluate material based on the visiblity pass result
+				// Output both complete sample data + per-sample velocity
+				FMaterialPassOutput PassOutput = AddHairMaterialPass(
+					GraphBuilder,
+					Scene,
+					&View,
+					false,
+					MacroGroupDatas,
+					InstanceCullingManager,
+					VisibilityData.NodeGroupSize,
+					CompactNodeIndex,
+					CompactNodeVis,
+					CompactNodeCoord,
+					NodeCounter,
+					IndirectArgsBuffer);
+
+				// Merge per-sample velocity into the scene velocity buffer
+				AddHairVelocityPass(
+					GraphBuilder,
+					View,
+					MacroGroupDatas,
+					VisibilityData.TileData,
+					CoverageTexture,
+					CompactNodeIndex,
+					CompactNodeVis,
+					PassOutput.NodeVelocity,
+					SceneVelocityTexture,
+					ResolveMaskTexture);
+
+				CompactNodeData = PassOutput.NodeData;
+
+				VisibilityData.SampleLightingViewportResolution = PassOutput.SampleLightingTexture->Desc.Extent;
+				VisibilityData.SampleLightingTexture = PassOutput.SampleLightingTexture;
+				VisibilityData.NodeIndex = CompactNodeIndex;
+				VisibilityData.CoverageTexture = CoverageTexture;
+				VisibilityData.HairOnlyDepthTexture = HairOnlyDepthTexture;
+				VisibilityData.NodeData = CompactNodeData;
+				VisibilityData.NodeVisData = CompactNodeVis;
+				VisibilityData.NodeCoord = CompactNodeCoord;
+				VisibilityData.NodeIndirectArg = IndirectArgsBuffer;
+				VisibilityData.NodeCount = NodeCounter;
+				VisibilityData.ResolveMaskTexture = ResolveMaskTexture;	
+
+				// For fully covered pixels, write: 
+				// * black color into the scene color
+				// * closest depth
+				// * unlit shading model ID 
+				if (bRunColorAndDepthPatching)
+				{
+					AddHairMaterialDataPatchPass(
 						GraphBuilder,
 						View,
-						HairOnlyDepthTexture,
-						VisibilityData.HairOnlyDepthHZBParameters,
-						VisibilityData.HairOnlyDepthClosestHZBTexture,
-						VisibilityData.HairOnlyDepthFurthestHZBTexture);
+						VisibilityData.TileData,
+						CoverageTexture,
+						CompactNodeIndex,
+						CompactNodeData,
+						SceneMaterial0,
+						SceneMaterial1,
+						SceneColorTexture,
+						SceneDepthTexture,
+						VisibilityData.LightChannelMaskTexture);
 				}
+
+				AddHairOnlyDepthPass(
+					GraphBuilder,
+					View,
+					VisibilityData.TileData,
+					CoverageTexture,
+					CompactNodeIndex,
+					CompactNodeData,
+					HairOnlyDepthTexture);
+
+				AddHairOnlyHZBPass(
+					GraphBuilder,
+					View,
+					HairOnlyDepthTexture,
+					VisibilityData.HairOnlyDepthHZBParameters,
+					VisibilityData.HairOnlyDepthClosestHZBTexture,
+					VisibilityData.HairOnlyDepthFurthestHZBTexture);
 			}
 			else if (RenderMode == HairVisibilityRenderMode_MSAA_Visibility)
 			{
