Illuminance meter fixes and clarifications.
SkyAtmosphere tooltip update.

diff --git a/Engine/Shaders/Private/ReflectionEnvironmentShaders.usf b/Engine/Shaders/Private/ReflectionEnvironmentShaders.usf
--- a/Engine/Shaders/Private/ReflectionEnvironmentShaders.usf
+++ b/Engine/Shaders/Private/ReflectionEnvironmentShaders.usf
@@ -895,7 +895,7 @@ float3 SkyLightCaptureWorlPos;
 TextureCube SkyLightCubeTexture;
 SamplerState SkyLightCubeSampler;
 
-float4 GetIlluminance()
+float3 GetIlluminance()
 {
 	float4 IlluminanceAcc = 0;
 	IlluminanceAcc += TextureCubeSampleLevel(SourceCubemapTexture, SourceCubemapSampler, float3(-1, 0, 0), CubeMipIndexToSampleIlluminance);
@@ -913,7 +913,7 @@ float4 GetIlluminance()
 
 float GetSimpleAutoExposure()
 {
-	float4 IlluminanceAcc = GetIlluminance();
+	float3 IlluminanceAcc = GetIlluminance();
 	float MaxIlluminance = max3(IlluminanceAcc.r, IlluminanceAcc.g, IlluminanceAcc.b);
 	float SimpleAutoExposure = 1.0f / (MaxIlluminance <= 0.0 ? 0.01f : MaxIlluminance);
 	return 1.0;// SimpleAutoExposure; Disabled for now, using the view exposure
@@ -921,7 +921,10 @@ float GetSimpleAutoExposure()
 
 float4 SampleCubeAsLatLon(TextureCube CubeTexture, SamplerState CubeSampler, float2 UV)
 {
-	float3 DirFromLatLong = float3(cos(UV.x * 2.0 * PI), sin(UV.x * 2.0 * PI), -sin(UV.y * PI - PI / 2));
+	float HAngle = UV.x * 2.0 * PI;
+	float VAngle = UV.y * PI - PI / 2;
+	float CosVAngle = cos(VAngle);
+	float3 DirFromLatLong = float3(cos(HAngle) * CosVAngle, sin(HAngle) * CosVAngle, -sin(VAngle));
 
 	float4 DebugVis = TextureCubeSampleLevel(CubeTexture, CubeSampler, DirFromLatLong, 0);
 	DebugVis.rgb *= GetSimpleAutoExposure();
@@ -965,31 +968,23 @@ void VisualizeIlluminanceMeterPS(
 
 	OutColor = 0.0f;
 
-	const uint2 CubeMapPreviewOrigin	= uint2(100, 110);
-	const uint2 CubeMapPreviewSize		= uint2(500, 250);
+	const uint2 IlluMapPreviewOrigin	= uint2(100, 110);
+	const uint2 IlluMapPreviewSize		= uint2(500, 250);
 
 	const uint2 SkyMapPreviewOrigin		= uint2(100, 370);
 	const uint2 SkyMapPreviewSize		= uint2(500, 250);	
 
-	const uint2 SphereCubeMapPreviewOrigin	= CubeMapPreviewOrigin + uint2(CubeMapPreviewSize.x + 10, 0);
-	const uint2 SphereCubeMapPreviewSize	= uint2(250, 250);
+	const uint2 SphereIlluMapPreviewOrigin	= IlluMapPreviewOrigin + uint2(IlluMapPreviewSize.x + 10, 0);
+	const uint2 SphereIlluMapPreviewSize	= uint2(250, 250);
 
 	const uint2 SphereSkyMapPreviewOrigin	= SkyMapPreviewOrigin + uint2(SkyMapPreviewSize.x + 10, 0);
 	const uint2 SphereSkyMapPreviewSize		= uint2(250, 250);
 
 	if (all(uint2(SVPos.xy) == uint2(1, 1)))
 	{
-		float4 IlluminanceAcc = GetIlluminance();
+		float3 SkyLightIlluminance = GetIlluminance();
 
-		FShaderPrintContext Context = InitShaderPrintContext(true, uint2(110, 70));
-		Print(Context, TEXT("Sky light illuminance = "), InitFontColor(0.8, 0.8, 0.8));
-		Print(Context, IlluminanceAcc.r, InitFontColor(1.0, 0.8, 0.8));
-		Print(Context, IlluminanceAcc.g, InitFontColor(0.8, 1.0, 0.8));
-		Print(Context, IlluminanceAcc.b, InitFontColor(0.8, 0.8, 1.0));
-
-		Newline(Context);
-
-	#if PROJECT_SUPPORT_SKY_ATMOSPHERE
+#if PROJECT_SUPPORT_SKY_ATMOSPHERE
 		const float3 PlanetCenterToTranslatedWorldPos = (LWCToFloat(LWCAdd(SkyLightCaptureWorlPos, ResolvedView.PreViewTranslation)) - ResolvedView.SkyPlanetTranslatedWorldCenterAndViewHeight.xyz) * CM_TO_SKY_UNIT;
 		// GetAtmosphereTransmittance does a shadow test against the virtual planet.
 		const float3 TransmittanceToLight = GetAtmosphereTransmittance(
@@ -997,19 +992,52 @@ void VisualizeIlluminanceMeterPS(
 			View.TransmittanceLutTexture, View.TransmittanceLutTextureSampler);
 
 		const float3 SunLightIlluminance = ResolvedView.AtmosphereLightIlluminanceOuterSpace[0].rgb * TransmittanceToLight;
-	#else
+#else
 		const float3 SunLightIlluminance = ResolvedView.AtmosphereLightIlluminanceOuterSpace[0];
-	#endif
+#endif
+
+		float SkyIlluminance = max3(SkyLightIlluminance.r, SkyLightIlluminance.g, SkyLightIlluminance.b);
+		float SunIlluminance = max3(SunLightIlluminance.r, SunLightIlluminance.g, SunLightIlluminance.b);
+		//float SkyIlluminance = dot(SkyLightIlluminance.rgb, (1.0/3.0).xxx);
+		//float SunIlluminance = dot(SunLightIlluminance.rgb, (1.0/3.0).xxx);
+		float TotalIlluminance = SkyIlluminance + SunIlluminance;
+
+		FShaderPrintContext Context = InitShaderPrintContext(true, uint2(110, 70));
+		Print(Context, TEXT("Sky light illuminance = "), InitFontColor(0.8, 0.8, 0.8));
+		Print(Context, SkyLightIlluminance.r, InitFontColor(1.0, 0.8, 0.8));
+		Print(Context, SkyLightIlluminance.g, InitFontColor(0.8, 1.0, 0.8));
+		Print(Context, SkyLightIlluminance.b, InitFontColor(0.8, 0.8, 1.0));
+		Print(Context, 100.0 * SkyIlluminance / TotalIlluminance, InitFontColor(0.7, 0.7, 0.7));
+		Print(Context, TEXT("%"), InitFontColor(0.7, 0.7, 0.7));
+
+		Newline(Context);
 
 		Print(Context, TEXT("Sun light illuminance = "), InitFontColor(0.8, 0.8, 0.8));
 		Print(Context, SunLightIlluminance.r, InitFontColor(1.0, 0.8, 0.8));
 		Print(Context, SunLightIlluminance.g, InitFontColor(0.8, 1.0, 0.8));
 		Print(Context, SunLightIlluminance.b, InitFontColor(0.8, 0.8, 1.0));
-	//	Print(Context, TEXT("(On ground, post SkyAtmosphere transmittance)"), InitFontColor(0.8, 0.8, 0.8));
+		Print(Context, 100.0 * SunIlluminance / TotalIlluminance, InitFontColor(0.7, 0.7, 0.7));
+		Print(Context, TEXT("%"), InitFontColor(0.7, 0.7, 0.7));
+
+		Newline(Context);
+
+		Print(Context, TEXT("Total illuminance     = "), InitFontColor(0.8, 0.8, 0.8));
+		Print(Context, SkyLightIlluminance.r + SunLightIlluminance.r, InitFontColor(1.0, 0.8, 0.8));
+		Print(Context, SkyLightIlluminance.g + SunLightIlluminance.g, InitFontColor(0.8, 1.0, 0.8));
+		Print(Context, SkyLightIlluminance.b + SunLightIlluminance.b, InitFontColor(0.8, 0.8, 1.0));
+		Print(Context, 100.0 * (SkyIlluminance + SunIlluminance) / TotalIlluminance, InitFontColor(0.7, 0.7, 0.7));
+		Print(Context, TEXT("%"), InitFontColor(0.7, 0.7, 0.7));
+
+		Context = InitShaderPrintContext(true, SphereIlluMapPreviewOrigin + uint2(SphereIlluMapPreviewSize.x + 10, SphereIlluMapPreviewSize.y / 2));
+		Print(Context, TEXT("Illuminance meter view"), InitFontColor(0.9, 0.9, 0.9));
+
+		Context = InitShaderPrintContext(true, SphereSkyMapPreviewOrigin + uint2(SphereSkyMapPreviewSize.x + 10, SphereSkyMapPreviewSize.y / 2));
+		Print(Context, TEXT("Sky light cube map view"), InitFontColor(0.9, 0.9, 0.9));
+
 	}
-	else if (all(uint2(SVPos.xy) >= CubeMapPreviewOrigin) && all(uint2(SVPos.xy) <= (CubeMapPreviewOrigin + CubeMapPreviewSize)))
+	else if (all(uint2(SVPos.xy) >= IlluMapPreviewOrigin) && all(uint2(SVPos.xy) <= (IlluMapPreviewOrigin + IlluMapPreviewSize)))
 	{
-		float2 UV = (SVPos.xy - float2(CubeMapPreviewOrigin)) / float2(CubeMapPreviewSize);
+		float2 UV = (SVPos.xy - float2(IlluMapPreviewOrigin)) / float2(IlluMapPreviewSize);
 		float4 DebugVis = SampleCubeAsLatLon(SourceCubemapTexture, SourceCubemapSampler, UV);
 		DebugVis.rgb *= 1.0 / DebugVis.a;	// This is a way to remove the texel dArea factor to show what luminance is integrated as a unction of NoL.
 		OutColor = float4(DebugVis.rgb, 0.0);
@@ -1020,9 +1048,9 @@ void VisualizeIlluminanceMeterPS(
 		float4 DebugVis = SampleCubeAsLatLon(SkyLightCubeTexture, SkyLightCubeSampler, UV);
 		OutColor = float4(DebugVis.rgb, 0.0);
 	}
-	else if (all(uint2(SVPos.xy) >= SphereCubeMapPreviewOrigin) && all(uint2(SVPos.xy) <= (SphereCubeMapPreviewOrigin + SphereCubeMapPreviewSize)))
+	else if (all(uint2(SVPos.xy) >= SphereIlluMapPreviewOrigin) && all(uint2(SVPos.xy) <= (SphereIlluMapPreviewOrigin + SphereIlluMapPreviewSize)))
 	{
-		float2 UV = (SVPos.xy - float2(SphereCubeMapPreviewOrigin)) / float2(SphereCubeMapPreviewSize);
+		float2 UV = (SVPos.xy - float2(SphereIlluMapPreviewOrigin)) / float2(SphereIlluMapPreviewSize);
 		float4 DebugVis = SampleCubeAsSphere(SourceCubemapTexture, SourceCubemapSampler, UV);
 		DebugVis.rgb *= 1.0 / DebugVis.a;	// This is a way to remove the texel dArea factor to show what luminance is integrated as a unction of NoL.
 		OutColor = float4(DebugVis.rgb, 0.0);
diff --git a/Engine/Source/Runtime/Engine/Classes/Components/SkyAtmosphereComponent.h b/Engine/Source/Runtime/Engine/Classes/Components/SkyAtmosphereComponent.h
--- a/Engine/Source/Runtime/Engine/Classes/Components/SkyAtmosphereComponent.h
+++ b/Engine/Source/Runtime/Engine/Classes/Components/SkyAtmosphereComponent.h
@@ -68,7 +68,10 @@ class USkyAtmosphereComponent : public USceneComponent
 	UPROPERTY(EditAnywhere, BlueprintReadOnly, interp, Category = "Atmosphere", meta = (UIMin = 1.0, UIMax = 200.0, ClampMin = 0.1, SliderExponent = 2.0))
 	float AtmosphereHeight;
 
-	/** Render multi scattering as if sun light would bounce around in the atmosphere. This is achieved using a dual scattering approach. */
+	/** Factor applied to multiple scattering only (after the sun light has bounced around in the atmosphere at least once). 
+	 * Multiple scattering is evaluated using a dual scattering approach. 
+	 * A value of 2 is recommended to better represent default atmosphere when r.SkyAtmosphere.MultiScatteringLUT.HighQuality=0. 
+	 */
 	UPROPERTY(EditAnywhere, BlueprintReadOnly, interp, Category = "Atmosphere", meta = (DisplayName = "MultiScattering", UIMin = 0.0, UIMax = 1.0, ClampMin = 0.0, ClampMax = 2.0))
 	float MultiScatteringFactor;
 
