Fixed shader compilation error.

diff --git a/Engine/Shaders/Private/DistortFiltering.usf b/Engine/Shaders/Private/DistortFiltering.usf
--- a/Engine/Shaders/Private/DistortFiltering.usf
+++ b/Engine/Shaders/Private/DistortFiltering.usf
@@ -26,7 +26,7 @@ void CopySceneColorDepthTexturePS(FScreenVertexOutput Input, out float4 OutColor
 	float4 LinearColor = Texture2DSample(SceneColorTexture, SceneColorSampler, TextureUV);
 
 	// Clamping to a certain max luminance help to reduce bright specular highlight flickering, even with temporal AA stabilization.
-	LinearColor.rgb = min(LinearColor, MaxExposedLuminance.xxx);
+	LinearColor.rgb = min(LinearColor.rgb, MaxExposedLuminance.xxx);
 
 	float DeviceZ = Texture2DSample(SceneDepthTexture, SceneDepthSampler, TextureUV).r;
 #if HAS_INVERTED_Z_BUFFER
