Fix implicit truncation in PostProcessMaterialShader.usf.

#rnx

diff --git a/Engine/Shaders/Private/PostProcessMaterialShaders.usf b/Engine/Shaders/Private/PostProcessMaterialShaders.usf
--- a/Engine/Shaders/Private/PostProcessMaterialShaders.usf
+++ b/Engine/Shaders/Private/PostProcessMaterialShaders.usf
@@ -323,7 +323,7 @@ void MainPS(
 
 	// can be optimized
 	float4 SvPosition = Input.Position;
-	float2 ViewportUV = (SvPosition - PostProcessOutput_ViewportMin) * PostProcessOutput_ViewportSizeInverse;
+	float2 ViewportUV = (SvPosition.xy - PostProcessOutput_ViewportMin.xy) * PostProcessOutput_ViewportSizeInverse.xy;
 	float2 BufferUV = ViewportUVToBufferUV(ViewportUV);
 	
 #if NUM_TEX_COORD_INTERPOLATORS
