Temporary CIS fix.

diff --git a/Engine/Shaders/Private/HairStrands/HairStrandsForwardRaster.usf b/Engine/Shaders/Private/HairStrands/HairStrandsForwardRaster.usf
--- a/Engine/Shaders/Private/HairStrands/HairStrandsForwardRaster.usf
+++ b/Engine/Shaders/Private/HairStrands/HairStrandsForwardRaster.usf
@@ -1024,6 +1024,7 @@ void RasterCS(
 	uint2 GroupThread2D    : SV_GroupThreadID,  /* 8x8 */
 	uint  GroupID          : SV_GroupID)		/* Rasterizer ID */
 {
+	#if 0
 	ResolvedView = ResolveView();
 	const uint BinTileNum = VisTileArgs[0];
 
@@ -1665,6 +1666,7 @@ void RasterCS(
 		} // 4. Raster
 
 	}// Work item
+	#endif
 }
 
 #endif //SHADER_RASTERCOMPUTE_RASTER
