Path Tracer: Fix potential NaN when rendering solid glass

diff --git a/Engine/Shaders/Private/PathTracing/Material/PathTracingGlossy.ush b/Engine/Shaders/Private/PathTracing/Material/PathTracingGlossy.ush
--- a/Engine/Shaders/Private/PathTracing/Material/PathTracingGlossy.ush
+++ b/Engine/Shaders/Private/PathTracing/Material/PathTracingGlossy.ush
@@ -104,7 +104,8 @@ float2 GGXEvalRefraction(float3 L, float3 V, float3 H, float2 Alpha, float Eta)
 		const float Weight = NoL * (LenV + NoV) / (NoV * LenL + NoL * LenV);
 		const float Pdf = (LoH * VoH * Eta * Eta * 2.0 * D) / (Pow2(VoH - Eta * LoH) * (LenV + NoV));
 
-		return float2(Weight, Pdf);
+		// NOTE: filter out potential NaNs from Pdf
+		return float2(Weight, max(Pdf, 0.0));
 	}
 	return 0;
 }
