Fix hair strands RT geometry not visible in RT debug view mode.

diff --git a/Engine/Shaders/Private/RayTracing/RayTracingDebug.usf b/Engine/Shaders/Private/RayTracing/RayTracingDebug.usf
--- a/Engine/Shaders/Private/RayTracing/RayTracingDebug.usf
+++ b/Engine/Shaders/Private/RayTracing/RayTracingDebug.usf
@@ -55,7 +55,7 @@ float4 DebugRayTracingMaterial(FRayDesc Ray, bool bIsFarField)
 	RayCone.SpreadAngle = View.EyeToPixelSpreadAngle;
 
 	const uint RayFlags = RAY_FLAG_CULL_BACK_FACING_TRIANGLES;
-	const uint InstanceInclusionMask = (bIsFarField) ? RAY_TRACING_MASK_FAR_FIELD : (OpaqueOnly ? RAY_TRACING_MASK_OPAQUE : RAY_TRACING_MASK_ALL);
+	const uint InstanceInclusionMask = (bIsFarField) ? RAY_TRACING_MASK_FAR_FIELD : (OpaqueOnly ? (RAY_TRACING_MASK_OPAQUE | RAY_TRACING_MASK_HAIR_STRANDS) : RAY_TRACING_MASK_ALL);
 	const bool bEnableSkyLightContribution = true;
 	const bool bIgnoreTranslucentMaterials = false;
 
