Restoring side by side textures and samplers in particle simulation shaders since those are still needed in certain OpenGL shader compilation paths.

diff --git a/Engine/Shaders/Private/ParticleSimulationShader.usf b/Engine/Shaders/Private/ParticleSimulationShader.usf
--- a/Engine/Shaders/Private/ParticleSimulationShader.usf
+++ b/Engine/Shaders/Private/ParticleSimulationShader.usf
@@ -131,14 +131,13 @@ Texture3D VectorFieldTextures_1;
 Texture3D VectorFieldTextures_2;
 Texture3D VectorFieldTextures_3;
 #endif
-//Texture3D VectorFieldTextures_4;
-//Texture3D VectorFieldTextures_5;
-//Texture3D VectorFieldTextures_6;
-//Texture3D VectorFieldTextures_7;
-//Texture3D VectorFieldTextures_8;
-//Texture3D VectorFieldTextures_9;
 
-SamplerState VectorFieldTexturesSampler;
+SamplerState VectorFieldTexturesSampler_0;
+#if MAX_VECTOR_FIELDS > 1
+SamplerState VectorFieldTexturesSampler_1;
+SamplerState VectorFieldTexturesSampler_2;
+SamplerState VectorFieldTexturesSampler_3;
+#endif
 
 /**
  * Computes the orbit velocity to apply to the particle based on time.
@@ -171,22 +170,13 @@ float3 ComputeOrbitVelocity(float Time, float RandomOrbit)
  float3 SampleVectorFieldTexture(int Index, float3 UV)
  {
 #if MAX_VECTOR_FIELDS == 1
-	 return Texture3DSample(VectorFieldTextures_0, VectorFieldTexturesSampler, UV).xyz;
+	 return Texture3DSample(VectorFieldTextures_0, VectorFieldTexturesSampler_0, UV).xyz;
 #else
-	if (Index == 0) return Texture3DSample(VectorFieldTextures_0, VectorFieldTexturesSampler, UV).xyz;
-	if (Index == 1) return Texture3DSample(VectorFieldTextures_1, VectorFieldTexturesSampler, UV).xyz;
-	if (Index == 2) return Texture3DSample(VectorFieldTextures_2, VectorFieldTexturesSampler, UV).xyz;
-	return Texture3DSample(VectorFieldTextures_3, VectorFieldTexturesSampler, UV).xyz;
+	if (Index == 0) return Texture3DSample(VectorFieldTextures_0, VectorFieldTexturesSampler_0, UV).xyz;
+	if (Index == 1) return Texture3DSample(VectorFieldTextures_1, VectorFieldTexturesSampler_1, UV).xyz;
+	if (Index == 2) return Texture3DSample(VectorFieldTextures_2, VectorFieldTexturesSampler_2, UV).xyz;
+	return Texture3DSample(VectorFieldTextures_3, VectorFieldTexturesSampler_3, UV).xyz;
 #endif
-
-//	if (Index == 3) return Texture3DSample(VectorFieldTextures_3, VectorFieldTexturesSampler, UV).xyz;
-//	if (Index == 4) return Texture3DSample(VectorFieldTextures_4, VectorFieldTexturesSampler, UV).xyz;
-//	if (Index == 5) return Texture3DSample(VectorFieldTextures_5, VectorFieldTexturesSampler, UV).xyz;
-//	if (Index == 6) return Texture3DSample(VectorFieldTextures_6, VectorFieldTexturesSampler, UV).xyz;
-//	if (Index == 7) return Texture3DSample(VectorFieldTextures_7, VectorFieldTexturesSampler, UV).xyz;
-//	if (Index == 8) return Texture3DSample(VectorFieldTextures_8, VectorFieldTexturesSampler, UV).xyz;
-//	
-//	return Texture3DSample(VectorFieldTextures_9, VectorFieldTexturesSampler, UV).xyz;
  }
 
 /**
diff --git a/Engine/Source/Runtime/Engine/Private/Particles/ParticleGpuSimulation.cpp b/Engine/Source/Runtime/Engine/Private/Particles/ParticleGpuSimulation.cpp
--- a/Engine/Source/Runtime/Engine/Private/Particles/ParticleGpuSimulation.cpp
+++ b/Engine/Source/Runtime/Engine/Private/Particles/ParticleGpuSimulation.cpp
@@ -940,7 +940,7 @@ BEGIN_SHADER_PARAMETER_STRUCT(FParticleSimulationPsParameters, )
 	SHADER_PARAMETER_SAMPLER(SamplerState, CurveTextureSampler)
 
 	SHADER_PARAMETER_TEXTURE_ARRAY(Texture3D, VectorFieldTextures, [MAX_VECTOR_FIELDS])
-	SHADER_PARAMETER_SAMPLER(SamplerState, VectorFieldTexturesSampler)
+	SHADER_PARAMETER_SAMPLER_ARRAY(SamplerState, VectorFieldTexturesSampler, [MAX_VECTOR_FIELDS])
 
 	SHADER_PARAMETER_STRUCT_INCLUDE(FParticlePerFrameSimulationShaderParameters, PerFrameParameters)
 	SHADER_PARAMETER_STRUCT_INCLUDE(FGlobalDistanceFieldParameters2, GlobalDistanceFieldParameters)
@@ -1274,7 +1274,10 @@ void ExecuteSimulationCommands(
 		PsParameters.CurveTexture = GParticleCurveTexture.GetCurveTexture();
 		PsParameters.CurveTextureSampler = SamplerStateLinear;
 
-		PsParameters.VectorFieldTexturesSampler = SamplerStateLinear;
+		for (int32 i = 0; i < MAX_VECTOR_FIELDS; ++i)
+		{
+			PsParameters.VectorFieldTexturesSampler[i] = SamplerStateLinear;
+		}
 
 		if (CollisionMode != PCM_None)
 		{
