Fix Strata blending mode with per-pixel OIT.

diff --git a/Engine/Shaders/Private/BasePassPixelShader.usf b/Engine/Shaders/Private/BasePassPixelShader.usf
--- a/Engine/Shaders/Private/BasePassPixelShader.usf
+++ b/Engine/Shaders/Private/BasePassPixelShader.usf
@@ -2023,13 +2023,28 @@ void FPixelShaderInOut_MainPS(
 
 		#if MATERIAL_WORKS_WITH_DUAL_SOURCE_COLOR_BLENDING
 		{
+			#if !STRATA_ENABLED
 			float3 AdjustedDualBlendAdd = Out.MRT[0];
 			float3 AdjustedDualBlendMul = Out.MRT[1];
+			#endif
+
 			// Add early out threshold?
 			AddOITSample(uint2(In.SvPosition.xy), AdjustedDualBlendAdd, AdjustedDualBlendMul, MaterialParameters.ScreenPosition.w);
 			Out.MRT[0] = half4(0, 0, 0, 0);
 			Out.MRT[1] = half4(1, 1, 1, 1);
 		}
+		#elif STRATA_TRANSLUCENT_MATERIAL
+		{
+			#if STRATA_BLENDING_COLOREDTRANSMITTANCEONLY
+			AddOITSample(uint2(In.SvPosition.xy), float3(0, 0, 0), FoggedColor, MaterialParameters.ScreenPosition.w);
+			Out.MRT[0] = half4(1, 1, 1, 1);
+			Out.MRT[1] = half4(0, 0, 0, 0);
+			#else
+			AddOITSample(uint2(In.SvPosition.xy), AdjustedDualBlendAdd, AdjustedDualBlendMul, MaterialParameters.ScreenPosition.w);
+			Out.MRT[0] = half4(0, 0, 0, 0);
+			Out.MRT[1] = half4(1, 1, 1, 1);
+			#endif
+		}
 		#else
 		{
 			float3 WriteColor = Out.MRT[0].rgb;
