Clean up old and duplicate code from some ray/path tracing shaders
- Removed unused SetPixelCoord/GetPixelCoord from FDecalShaderParams.
- Removed duplicate calculations of TranslatedWorldPosition from PathTracingMaterialHitShader.

diff --git a/Engine/Shaders/Private/PathTracing/PathTracingMaterialHitShader.usf b/Engine/Shaders/Private/PathTracing/PathTracingMaterialHitShader.usf
--- a/Engine/Shaders/Private/PathTracing/PathTracingMaterialHitShader.usf
+++ b/Engine/Shaders/Private/PathTracing/PathTracingMaterialHitShader.usf
@@ -54,7 +54,7 @@ RAY_TRACING_ENTRY_CLOSEST_HIT(PathTracingMaterialCHS,
 	ResolvedView = ResolveView();
 
 	const float3 TranslatedWorldPosition = TranslatedWorldRayOrigin() + RayTCurrent() * WorldRayDirection();
-	float4 SvPosition = mul(float4(TranslatedWorldPosition, 1.0f), ResolvedView.TranslatedWorldToClip);
+	const float4 SvPosition = mul(float4(TranslatedWorldPosition, 1.0f), ResolvedView.TranslatedWorldToClip);
 
 #if VF_SUPPORTS_RAYTRACING_PREPARE_MATERIAL_PIXEL_PARAMETERS
 	// this is a newer codepath that is both more flexible and allows for more direct calculation compared to the other codepath
@@ -72,10 +72,8 @@ RAY_TRACING_ENTRY_CLOSEST_HIT(PathTracingMaterialCHS,
 	FPixelMaterialInputs PixelMaterialInputs;
 
 	{
-		float4 ScreenPosition = SvPositionToResolvedScreenPosition(SvPosition);
-		float3 TranslatedWorldPosition = TranslatedWorldRayOrigin() + WorldRayDirection() * RayTCurrent();
-
-		bool bIsFrontFace = HitKind() == HIT_KIND_TRIANGLE_FRONT_FACE;
+		const float4 ScreenPosition = SvPositionToResolvedScreenPosition(SvPosition);
+		const bool bIsFrontFace = HitKind() == HIT_KIND_TRIANGLE_FRONT_FACE;
 
 		// #dxr_todo: UE-72130 support world position offset 
 		// #if USE_WORLD_POSITION_EXCLUDING_SHADER_OFFSETS
@@ -457,7 +455,7 @@ RAY_TRACING_ENTRY_ANY_HIT(PathTracingMaterialAHS,
 	ResolvedView = ResolveView();
 
 	const float3 TranslatedWorldPosition = TranslatedWorldRayOrigin() + RayTCurrent() * WorldRayDirection();
-	float4 SvPosition = mul(float4(TranslatedWorldPosition, 1.0f), ResolvedView.TranslatedWorldToClip);
+	const float4 SvPosition = mul(float4(TranslatedWorldPosition, 1.0f), ResolvedView.TranslatedWorldToClip);
 
 #if VF_SUPPORTS_RAYTRACING_PREPARE_MATERIAL_PIXEL_PARAMETERS
 	// this is a newer codepath that is both more flexible and allows for more direct calculation compared to the other codepath
@@ -477,8 +475,7 @@ RAY_TRACING_ENTRY_ANY_HIT(PathTracingMaterialAHS,
 	const bool bIsFrontFace = HitKind() == HIT_KIND_TRIANGLE_FRONT_FACE;
 
 	{
-		float4 ScreenPosition = SvPositionToResolvedScreenPosition(SvPosition);
-		float3 TranslatedWorldPosition = TranslatedWorldRayOrigin() + WorldRayDirection() * RayTCurrent();
+		const float4 ScreenPosition = SvPositionToResolvedScreenPosition(SvPosition);
 
 		MaterialParameters.CameraVector = -WorldRayDirection();
 
diff --git a/Engine/Shaders/Private/RayTracing/RayTracingCommon.ush b/Engine/Shaders/Private/RayTracing/RayTracingCommon.ush
--- a/Engine/Shaders/Private/RayTracing/RayTracingCommon.ush
+++ b/Engine/Shaders/Private/RayTracing/RayTracingCommon.ush
@@ -931,13 +931,6 @@ struct FDecalShaderParams
 	}
 	float3 GetInputTranslatedWorldPosition() { return float3(Data[0], Data[1], Data[2]); }
 
-	void SetPixelCoord(uint2 PixelCoord)
-	{
-		Data[3] = asfloat(PixelCoord.x);
-		Data[4] = asfloat(PixelCoord.y);
-	}
-	uint2 GetPixelCoord() { return asuint(float2(Data[3], Data[4])); }
-
 	void SetMetallic(float In) { Data[0] = In; }
 	float GetMetallic() { return Data[0]; }
 
