Prevent NaNs from going into the convolved realtime sky capture cube map

diff --git a/Engine/Shaders/Private/ReflectionEnvironmentShaders.usf b/Engine/Shaders/Private/ReflectionEnvironmentShaders.usf
--- a/Engine/Shaders/Private/ReflectionEnvironmentShaders.usf
+++ b/Engine/Shaders/Private/ReflectionEnvironmentShaders.usf
@@ -513,6 +513,10 @@ void FilterPS(
 	if( Roughness < 0.01 )
 	{
 		OutColor = SourceCubemapTexture.SampleLevel(SourceCubemapSampler, CubeCoordinates, 0 );
+		
+		// Prevent NaNs from going into the cube map
+		OutColor = -min(-OutColor, 0);
+		
 	#ifdef USE_COMPUTE
 		OutTextureMipColor[uint3(FaceCoord, SelectedCubeFace)] = OutColor;
 	#endif
@@ -593,6 +597,9 @@ void FilterPS(
 		OutColor = FilteredColor / Weight;
 	}
 
+	// Prevent NaNs from going into the cube map
+	OutColor = -min(-OutColor, 0);
+
 #ifdef USE_COMPUTE
 	OutTextureMipColor[uint3(FaceCoord, SelectedCubeFace)] = OutColor;
 #endif
