Thin-translucent: adjust logic for thin-translucent to be a bit more consistent with the path tracer.

The fresnel calculations were needlessly expensive and can be simplified as entering and exiting events must have the same fresnel factor. Added the mathematica "proof" of this in the comments.

The specular color was only using Spec and not Metallic. This limited the effective IOR range of thin translucent materials to 1.8 at most. Now metallic is no longer scaled by opacity, which allows for creating surfaces made of higher IORs like diamond for example.

This is motivated by the fact that Strata will be using Diffuse/Specular workflow internally and therefore does not treat metallic specially in this case.

Adjust RT Translucency to match the path tracer behavior, and implement the missing blend modes.

diff --git a/Engine/Shaders/Private/BasePassPixelShader.usf b/Engine/Shaders/Private/BasePassPixelShader.usf
--- a/Engine/Shaders/Private/BasePassPixelShader.usf
+++ b/Engine/Shaders/Private/BasePassPixelShader.usf
@@ -864,12 +864,6 @@ void FPixelShaderInOut_MainPS(
 
 	float MaterialAO = GetMaterialAmbientOcclusion(PixelMaterialInputs);
 
-// Opacity for this model is the coverage of the top layer over the transmission surface. Since the
-// underlying material isn't allowed to be metallic, we are scaling the metallic value by the coverage.
-#if MATERIAL_SHADINGMODEL_THIN_TRANSLUCENT
-	Metallic *= Opacity;
-#endif
-
 	// If we don't use this shading model the color should be black (don't generate shader code for unused data, don't do indirectlighting cache lighting with this color).
 	float3 SubsurfaceColor = 0;
 	// 0..1, SubsurfaceProfileId = int(x * 255)
diff --git a/Engine/Shaders/Private/MobileBasePassPixelShader.usf b/Engine/Shaders/Private/MobileBasePassPixelShader.usf
--- a/Engine/Shaders/Private/MobileBasePassPixelShader.usf
+++ b/Engine/Shaders/Private/MobileBasePassPixelShader.usf
@@ -367,12 +367,6 @@ void Main(
 	half Anisotropy = GetMaterialAnisotropy(PixelMaterialInputs);
 	uint ShadingModelID = GetMaterialShadingModel(PixelMaterialInputs);
 	
-	// Opacity for this model is the coverage of the top layer over the transmission surface. Since the
-	// underlying material isn't allowed to be metallic, we are scaling the metallic value by the coverage.
-#if MATERIAL_SHADINGMODEL_THIN_TRANSLUCENT
-	Metallic *= Opacity;
-#endif
-
 	// If we don't use this shading model the color should be black (don't generate shader code for unused data, don't do indirectlighting cache lighting with this color).
 	float3 SubsurfaceColor = 0;
 	// 0..1, SubsurfaceProfileId = int(x * 255)
diff --git a/Engine/Shaders/Private/RayTracing/RayTracingMaterialHitShaders.usf b/Engine/Shaders/Private/RayTracing/RayTracingMaterialHitShaders.usf
--- a/Engine/Shaders/Private/RayTracing/RayTracingMaterialHitShaders.usf
+++ b/Engine/Shaders/Private/RayTracing/RayTracingMaterialHitShaders.usf
@@ -325,7 +325,7 @@ void GetMaterialPayload(
 	half IOR = DielectricF0RGBToIor(F0);
 	#endif
 #else
-	half IOR = 1;
+	half IOR = 0;
 #endif
 
 	/**
@@ -338,7 +338,11 @@ void GetMaterialPayload(
 	Payload.Specular = Specular;
 	Payload.Roughness = Roughness;
 	Payload.Metallic = Metallic;
+#if MATERIALBLENDING_MODULATE
+	Payload.Opacity = 0.0;
+#else
 	Payload.Opacity = GetMaterialOpacity(PixelMaterialInputs);
+#endif
 
 #if MATERIALBLENDING_MASKED
 	// regular masked mode - binary decision
@@ -348,7 +352,7 @@ void GetMaterialPayload(
 	}
 #endif // MATERIALBLENDING_MASKED
 
-	Payload.Ior = IOR == 0 ? 1.0 : IOR;
+	Payload.Ior = IOR;
 	Payload.CustomData = float4(GetMaterialCustomData0(MaterialParameters), GetMaterialCustomData1(MaterialParameters), 0, 0);
 	
 #if MATERIAL_SHADINGMODEL_CLEAR_COAT
diff --git a/Engine/Shaders/Private/RayTracing/RayTracingPrimaryRays.usf b/Engine/Shaders/Private/RayTracing/RayTracingPrimaryRays.usf
--- a/Engine/Shaders/Private/RayTracing/RayTracingPrimaryRays.usf
+++ b/Engine/Shaders/Private/RayTracing/RayTracingPrimaryRays.usf
@@ -17,6 +17,7 @@
 #include "../ReflectionEnvironmentShared.ush"
 #include "../Montecarlo.ush"
 #include "../PathTracing/Utilities/PathTracingRandomSequence.ush" 
+#include "../PathTracing/Material/PathTracingFresnel.ush"
 #include "../HeightFogCommon.ush"
 #include "../SobolRandom.ush"
 #include "../SceneTextureParameters.ush"
@@ -129,13 +130,13 @@ RAY_TRACING_ENTRY_RAYGEN(RayTracingPrimaryRaysRGS)
 
 	// Integrated data by path tracing
 	float3 PathRadiance = 0.0;
-	float PathThroughput = 1.0;	// A float for now because UE does not support colored translucency as of today.
+	float3 PathThroughput = 1.0;
 	float LastRoughness = 0.0;
 	float HitDistance = 0.0f;
 
 	for (uint RefractionRayIndex = 0; RefractionRayIndex < MaxRefractionRays; ++RefractionRayIndex)
 	{
-		const uint RefractionRayFlags = RefractionRayIndex == 0? RAY_FLAG_CULL_BACK_FACING_TRIANGLES : 0;
+		const uint RefractionRayFlags = RefractionRayIndex == 0 ? RAY_FLAG_CULL_BACK_FACING_TRIANGLES : 0;
 		const uint RefractionInstanceInclusionMask = (!bPrimaryView && RefractionRayIndex == 0) || !TranslucencyRefraction ? RAY_TRACING_MASK_TRANSLUCENT : RAY_TRACING_MASK_ALL;
 		const bool bRefractionRayTraceSkyLightContribution = false;
 		const bool bRefractionDecoupleSampleGeneration = true;
@@ -201,7 +202,7 @@ RAY_TRACING_ENTRY_RAYGEN(RayTracingPrimaryRaysRGS)
 			// For that sample the fog for the hit position and view assuming no refraction. This seems acceptable.
 			ViewFogInScatteringAndTransmittance = CalculateHeightFog(WorldPositionRelativeToCameraPrimaryRay);
 
-			if (FogStruct.ApplyVolumetricFog > 0 
+			if (FogStruct.ApplyVolumetricFog > 0
 				&& !bHasScattered) // We can only sample the volumetric fog contribution (always camera aligned) if the ray is not scattered around
 			{
 				float4 ClipPos = mul(float4(HitPoint, 1.0f), PrimaryView.TranslatedWorldToClip); // HitPoint is translated world space
@@ -282,29 +283,63 @@ RAY_TRACING_ENTRY_RAYGEN(RayTracingPrimaryRaysRGS)
 			PathRadiance += PathThroughput * vertexRadianceWeight * ViewFogInScatteringAndTransmittance.a * ReflectionThroughput * ReflectionRadiance;
 		}
 
+		// Update the refraction/transparency path transmittance and check stop condition
+		switch (Payload.BlendingMode)
+		{
+			case RAY_TRACING_BLEND_MODE_ADDITIVE:
+			{
+				break;
+			}
+			case RAY_TRACING_BLEND_MODE_MODULATE:
+			{
+				PathThroughput *= Payload.Radiance;
+				break;
+			}
+			case RAY_TRACING_BLEND_MODE_ALPHA_COMPOSITE:
+			case RAY_TRACING_BLEND_MODE_TRANSLUCENT:
+			{
+				PathThroughput *= 1.0 - Payload.Opacity;
+				break;
+			}
+			default:
+			{
+				PathThroughput *= 0.0;
+				break;
+			}
+		}
+		
+		if (all(PathThroughput <= 0.0))
+		{
+			break;
+		}
 
 		//
 		// Handle refraction through the surface.
 		//
 
-		// Update the refraction path transmittance and check stop condition
-		float PathVertexTransmittance = Payload.BlendingMode == RAY_TRACING_BLEND_MODE_ADDITIVE ? 1.0 : 1.0 - Payload.Opacity;
-		PathThroughput *= PathVertexTransmittance;
-		if (PathThroughput <= 0.0)
-		{
-			break;
-		}
-
 		// Set refraction ray for next iteration
 		float3 RefractedDirection = Ray.Direction;
-		if (TranslucencyRefraction)
+		if (Payload.ShadingModelID == SHADINGMODELID_THIN_TRANSLUCENT)
+		{
+			const float3 N = Payload.WorldNormal;
+			const float3 V = -Ray.Direction;
+			const float NoV = dot(N, V);
+
+			const float F0 = F0RGBToF0(Payload.SpecularColor);
+			const float F = FresnelReflectance(NoV, 1.0, F0);
+			const float3 Transmittance = Payload.CustomData.xyz;
+			const float PathLength = rcp(saturate(abs(NoV) + 1e-5));
+			PathThroughput *= pow(Transmittance, PathLength) * Pow2(1 - F);
+		}
+		else if (TranslucencyRefraction && Payload.BlendingMode == RAY_TRACING_BLEND_MODE_TRANSLUCENT && Payload.Ior > 0.0)
 		{
 			float Ior = Payload.Ior;
 			bHasScattered |= Ior > 1.0 ? true : false;
 
-			float ReflectionRefractionEventThroughput;
 			bool bIsEntering = Payload.IsFrontFace();
 
+			float Eta = bIsEntering ? Ior : rcp(Ior);
+
 			float3 N = Payload.WorldNormal;
 			float3 V = -Ray.Direction;
 			float NoV = dot(N, V);
@@ -317,23 +352,12 @@ RAY_TRACING_ENTRY_RAYGEN(RayTracingPrimaryRaysRGS)
 				bIsEntering = true;
 			}
 
-			float N1 = bIsEntering ? 1.0 : Ior;
-			float N2 = bIsEntering ? Ior : 1.0;
-			float Eta = N1 / N2;
-			float NoT = CalcNoT(NoV, N1, N2);
-			float Fr = FresnelDielectric(Eta, NoV, NoT);
+			float F0 = F0RGBToF0(Payload.SpecularColor);
 
-			float3 T = refract(Ray.Direction, N, Eta);
-			if (any(T) > 0.0)
-			{
-				RefractedDirection = T;
-				PathThroughput *= 1.0 - Fr;
-			}
-			// Handle total internal reflection
-			else
-			{
-				RefractedDirection = reflect(Ray.Direction, N);
-			}
+
+			float F = 0.0; // fresnel
+			SampleRefraction(-V, N, Eta, 1.0 /* always sample refraction or TIR */, RefractedDirection, F);
+			PathThroughput *= F;
 
 			// ray has bent, so it may need to go arbitrarily far
 			NextMaxRayDistance = LocalMaxRayDistance;
diff --git a/Engine/Shaders/Private/ThinTranslucentCommon.ush b/Engine/Shaders/Private/ThinTranslucentCommon.ush
--- a/Engine/Shaders/Private/ThinTranslucentCommon.ush
+++ b/Engine/Shaders/Private/ThinTranslucentCommon.ush
@@ -6,8 +6,6 @@
 
 #define THIN_TRANSLUCENT_USE_DOTNV_THICKNESS 1
 
-#define THIN_TRANSLUCENT_ACCURATE_FRESNEL (SHADING_PATH_DEFERRED || MOBILE_HIGH_QUALITY_BRDF)
-
 #if MATERIAL_SHADINGMODEL_THIN_TRANSLUCENT
 void AccumulateThinTranslucentModel(inout float3 DualBlendSurfaceLuminancePostCoverage,
 								inout float3 DualBlendSurfaceTransmittancePreCoverage,
@@ -34,7 +32,9 @@ void AccumulateThinTranslucentModel(inout float3 DualBlendSurfaceLuminancePostCo
 	
 	// color is for normalized thickness
 #if THIN_TRANSLUCENT_USE_DOTNV_THICKNESS
-	float PathLength = 1.0f/NoV;
+    // TODO: in theory this should account for the IOR as well, see the path tracing code for a reference
+	// However in practice the difference is subtle, so keep using this simple logic for now.
+	float PathLength = rcp(NoV);
 #else
 	float PathLength = 1.0f;
 #endif	
@@ -50,19 +50,15 @@ void AccumulateThinTranslucentModel(inout float3 DualBlendSurfaceLuminancePostCo
 	// Light gets lost from absorption through the surface
 	Transmittance = Transmittance * ColorTopMaterialCoverage;
 
-	// Exiting the surface, FresnelRatio light is reflected back into the surface, and the rest refracts out of the surface.
-#if THIN_TRANSLUCENT_ACCURATE_FRESNEL
-	// Accurate calculation
-	const float F0 = dot(float3(1.0f,1.0f,1.0f)/3.0f,GBuffer.SpecularColor.rgb); // air to glass F0
-	const float NIor = DielectricF0ToIor(F0);
-	const float ReverseNIor = 1.0/NIor; // glass to air IOR
-	const float ReverseF0 = DielectricIorToF0(ReverseNIor); // glass to air F0
-	const float3 ReverseFresnelRatio = F_Schlick(ReverseF0, NoV);
-	Transmittance = Transmittance * (1-ReverseFresnelRatio);
-#else
-	// Approximation, assumes same amount refracts from air->material as material->air.
+	// Exiting the surface, the Fresnel ratio is the same! This is technically only true for F_Fresnel, but since we are using
+	// F_Schlick as a stand-in, applying the same logic is a valid approximation
+	// In mathematica, using fr[cos,eta] we can verify that:
+	//   fr[c_, eta_] := 1 / 2 * A ^ 2 * (1 + B ^ 2) /.A -> (g - c) / (g + c) /.B -> (c * (g + c) - 1) / (c * (g - c) + 1) /.g->Sqrt[eta ^ 2 - 1 + c ^ 2];
+	//   FullSimplify[fr[c, 1/ior] == fr[Sqrt[1 - (1 - c ^ 2) * ior ^ 2], ior], Assumptions->ior > 1 && c > 0 && c < 1]
+	// evaluates to
+	//   True
 	Transmittance = Transmittance * (1-FresnelRatio);
-#endif
+
 	// We are treating the BaseColor and Emissive color as a layer on top of the absorbing media, but below specular layer.
 	float3 DefaultLitColor = DiffuseColor + EmissiveColor;
 
