Path Tracer: Apply roughness clamping and shading normal adjustments in the material instead of the integrator

This will be required for Strata as we won't be able to modify the encoded BSDF tree after it has been returned by the material.

Refactor how optional payload fields are managed as a consequence of this re-organization. The SSS profile is now partially decoded inside the material. This reduces the number of shading model specific operations that need to happen inside the integrator. More work will shift to the material in future refactors.

There are now 8 custom data fields in the payload, used by the various shading models to encode additional information.

diff --git a/Engine/Plugins/Experimental/GPULightmass/Shaders/Private/LightmapPathTracing.usf b/Engine/Plugins/Experimental/GPULightmass/Shaders/Private/LightmapPathTracing.usf
--- a/Engine/Plugins/Experimental/GPULightmass/Shaders/Private/LightmapPathTracing.usf
+++ b/Engine/Plugins/Experimental/GPULightmass/Shaders/Private/LightmapPathTracing.usf
@@ -187,7 +187,7 @@ FPathTracingPayload TraceTransparentRay(FRayDesc Ray, bool IsCameraRay, bool Las
 		// also avoid noise from hot spots (they can be properly
 		// importance sampled if a capturing skylight is added)
 		const bool bIgnoreSky = !IsCameraRay;
-		FPackedPathTracingPayload PackedPayload = InitPathTracingPayload(bIgnoreSky);
+		FPackedPathTracingPayload PackedPayload = InitPathTracingPayload(bIgnoreSky, 1.0);
 		// Trace the ray
 		TraceRay(
 			TLAS,
@@ -397,7 +397,6 @@ void PathTracingKernel(
 			Payload.Roughness = 1;
 			Payload.Metallic = 0;
 			Payload.Ior = 1;
-			Payload.CustomData = float4(0, 0, 0, 0);
 			Payload.Opacity = 1;
 			Payload.ShadingModelID = SHADINGMODELID_DEFAULT_LIT;
 			Payload.BlendingMode = RAY_TRACING_BLEND_MODE_OPAQUE;
@@ -542,18 +541,6 @@ void PathTracingKernel(
 		}
 	#endif
 		
-		// As soon as the path is blurry enough, we can get away with diffuse sampling only
-		const bool bSimplifySSS = PathRoughness >= 0.15;
-
-		if (ApproximateCaustics)
-		{
-			// modify the payload roughness to minimize difficult caustics
-			// This is inspired by a trick used in the Arnold renderer:
-			//   https://cgg.mff.cuni.cz/~jaroslav/gicourse2010/giai2010-02-marcos_fajardo-slides.pdf (slide 39)
-			//   https://www.arnoldrenderer.com/research/Arnold_TOG2018.pdf (section 4.2)
-			AdjustMaterialRoughness(Payload, PathRoughness);
-		}
-		
 		FLightLoopCount LightLoopCount = LightGridLookup(Payload.TranslatedWorldPos);
 
 		// Choose a random number for both Light sampling and BxDF sampling
diff --git a/Engine/Shaders/Private/PathTracing/Material/PathTracingCloth.ush b/Engine/Shaders/Private/PathTracing/Material/PathTracingCloth.ush
--- a/Engine/Shaders/Private/PathTracing/Material/PathTracingCloth.ush
+++ b/Engine/Shaders/Private/PathTracing/Material/PathTracingCloth.ush
@@ -34,8 +34,8 @@ FMaterialEval Cloth_EvalMaterial(
 	const float VoH = saturate(dot(V, H));
 	const float NoH = saturate(H.z);
 
-	const float3 FuzzColor = Payload.CustomData.xyz;
-	const float Cloth = Payload.CustomData.w;
+	const float3 FuzzColor = Payload.GetClothColor();
+	const float Cloth = Payload.GetClothAmount();
 
 	const FBxDFEnergyTerms Spec = ComputeGGXSpecEnergyTerms(Payload.Roughness, NoV, Payload.SpecularColor);
 	const FBxDFEnergyTerms Fuzz = ComputeClothEnergyTerms(Payload.Roughness, NoV, FuzzColor);
@@ -81,8 +81,8 @@ FMaterialSample Cloth_SampleMaterial(
 
 	const float NoV = saturate(V.z);
 
-	const float3 FuzzColor = Payload.CustomData.xyz;
-	const float Cloth = Payload.CustomData.w;
+	const float3 FuzzColor = Payload.GetClothColor();
+	const float Cloth = Payload.GetClothAmount();
 	const FBxDFEnergyTerms Spec = ComputeGGXSpecEnergyTerms(Payload.Roughness, NoV, Payload.SpecularColor);
 	const FBxDFEnergyTerms Fuzz = ComputeClothEnergyTerms(Payload.Roughness, NoV, FuzzColor);
 
diff --git a/Engine/Shaders/Private/PathTracing/Material/PathTracingEye.ush b/Engine/Shaders/Private/PathTracing/Material/PathTracingEye.ush
--- a/Engine/Shaders/Private/PathTracing/Material/PathTracingEye.ush
+++ b/Engine/Shaders/Private/PathTracing/Material/PathTracingEye.ush
@@ -40,9 +40,9 @@ FMaterialEval Eye_EvalMaterial(
 	const float LobeSelectionPdf = LobeSelectionProb(Payload.DiffuseColor * (1 - Spec.E), Spec.E);
 
 
-	const float IrisMask = Payload.Anisotropy;
-	const float3 IrisNormal = Payload.WorldTangent;
-	const float3 CausticNormal = Payload.CustomVector;
+	const float IrisMask = Payload.GetEyeIrisMask();
+	const float3 IrisNormal = Payload.GetEyeIrisNormal();
+	const float3 CausticNormal = Payload.GetEyeCausticNormal();
 	const float IrisNoL = saturate(dot(IrisNormal, L_World));
 	const float Power = lerp(12, 1, IrisNoL);
 	const float Caustic = 0.8 + 0.2 * (Power + 1) * pow(saturate(dot(CausticNormal, L_World)), Power);
@@ -128,9 +128,9 @@ FMaterialSample Eye_SampleMaterial(
 	float3 OutWeight = 0;
 	float OutPdf = 0;
 
-	const float IrisMask = Payload.Anisotropy;
-	const float3 IrisNormal = Payload.WorldTangent;
-	const float3 CausticNormal = Payload.CustomVector;
+	const float IrisMask = Payload.GetEyeIrisMask();
+	const float3 IrisNormal = Payload.GetEyeIrisNormal();
+	const float3 CausticNormal = Payload.GetEyeCausticNormal();
 	const float IrisNoL = saturate(dot(IrisNormal, L_World));
 	const float Power = lerp(12, 1, IrisNoL);
 	const float Caustic = 0.8 + 0.2 * (Power + 1) * pow(saturate(dot(CausticNormal, L_World)), Power);
diff --git a/Engine/Shaders/Private/PathTracing/Material/PathTracingHair.ush b/Engine/Shaders/Private/PathTracing/Material/PathTracingHair.ush
--- a/Engine/Shaders/Private/PathTracing/Material/PathTracingHair.ush
+++ b/Engine/Shaders/Private/PathTracing/Material/PathTracingHair.ush
@@ -301,13 +301,13 @@ FMaterialSample Hair_SampleMaterial(
 	FPathTracingPayload Payload,
 	float4 RandSample)
 {
-	float3 V_World = -RayDirection;
-	float3x3 HairBasis = GetHairBasis(Payload.WorldTangent, V_World);
+	const float3 V_World = -RayDirection;
+	const float3x3 HairBasis = GetHairBasis(Payload.WorldTangent, V_World);
 
-	float3 V = mul(HairBasis, V_World); // TODO: simplify this since HairBasis is already sort of aligned with V ?
-	float h = 2 * Payload.CustomData.y - 1; // remap back to [-1,1]
-	float AzimuthalRoughness = lerp(1.0, 0.1, Payload.Metallic);
-	FHairData HairData = PrepareHairData(V, Payload.BaseColor, Payload.Roughness, AzimuthalRoughness, h, Payload.Specular);
+	const float3 V = mul(HairBasis, V_World); // TODO: simplify this since HairBasis is already sort of aligned with V ?
+	const float h = 2 * Payload.GetHairPrimitiveUV().y - 1; // remap back to [-1,1]
+	const float AzimuthalRoughness = lerp(1.0, 0.1, Payload.Metallic);
+	const FHairData HairData = PrepareHairData(V, Payload.BaseColor, Payload.Roughness, AzimuthalRoughness, h, Payload.Specular);
 
 	// 1) Decide which lobe to sample
 	float u = RandSample.z * HairData.LobeCdf.w;
@@ -347,23 +347,23 @@ FMaterialSample Hair_SampleMaterial(
 	u = saturate(u);
 
 	// 2) Sample longitudinal angle
-	float CosTheta = 1 + v * log(RandSample.x + (1 - RandSample.x) * exp(-2 / v));
-	float CosPhi = cos(2 * PI * RandSample.y);
-	float SinThetaI = -CosTheta * SinThetaOp + sqrt((1 - CosTheta * CosTheta) * (1 - SinThetaOp * SinThetaOp)) * CosPhi;
-	float CosThetaI = TrigInverse(SinThetaI);
+	const float CosTheta = 1 + v * log(RandSample.x + (1 - RandSample.x) * exp(-2 / v));
+	const float CosPhi = cos(2 * PI * RandSample.y);
+	const float SinThetaI = -CosTheta * SinThetaOp + sqrt((1 - CosTheta * CosTheta) * (1 - SinThetaOp * SinThetaOp)) * CosPhi;
+	const float CosThetaI = TrigInverse(SinThetaI);
 
 	// 3) Sample Azimuthal angle
-	float PhiI = (p < 3) ? Phi(p, HairData.GammaO, HairData.GammaT) + SampleTrimmedLogistic(u, HairData.s) : 2 * PI * u;
+	const float PhiI = (p < 3) ? Phi(p, HairData.GammaO, HairData.GammaT) + SampleTrimmedLogistic(u, HairData.s) : 2 * PI * u;
 
-	float3 L = float3(CosThetaI * sin(PhiI), CosThetaI * cos(PhiI), SinThetaI);
+	const float3 L = float3(CosThetaI * sin(PhiI), CosThetaI * cos(PhiI), SinThetaI);
 
 	float3 Weight = 0;
 	float Pdf = 0;
-	float InvS = rcp(HairData.s);
-	float Pdf0 = Mp(SinThetaI, HairData.SinThetaO.x, HairData.v0) * Np(PhiI, 0, InvS, HairData.GammaO, HairData.GammaT);
-	float Pdf1 = Mp(SinThetaI, HairData.SinThetaO.y, HairData.v1) * Np(PhiI, 1, InvS, HairData.GammaO, HairData.GammaT);
-	float Pdf2 = Mp(SinThetaI, HairData.SinThetaO.z, HairData.v2) * Np(PhiI, 2, InvS, HairData.GammaO, HairData.GammaT);
-	float Pdf3 = Mp(SinThetaI, HairData.SinThetaO.w, HairData.v3) * (1 / (2 * PI)); // last lobe is isotropic
+	const float InvS = rcp(HairData.s);
+	const float Pdf0 = Mp(SinThetaI, HairData.SinThetaO.x, HairData.v0) * Np(PhiI, 0, InvS, HairData.GammaO, HairData.GammaT);
+	const float Pdf1 = Mp(SinThetaI, HairData.SinThetaO.y, HairData.v1) * Np(PhiI, 1, InvS, HairData.GammaO, HairData.GammaT);
+	const float Pdf2 = Mp(SinThetaI, HairData.SinThetaO.z, HairData.v2) * Np(PhiI, 2, InvS, HairData.GammaO, HairData.GammaT);
+	const float Pdf3 = Mp(SinThetaI, HairData.SinThetaO.w, HairData.v3) * (1 / (2 * PI)); // last lobe is isotropic
 
 	AddLobeWithMIS(Weight, Pdf, HairData.A0, Pdf0, HairData.LobePdf.x);
 	AddLobeWithMIS(Weight, Pdf, HairData.A1, Pdf1, HairData.LobePdf.y);
@@ -371,7 +371,7 @@ FMaterialSample Hair_SampleMaterial(
 	AddLobeWithMIS(Weight, Pdf, HairData.A3, Pdf3, HairData.LobePdf.w);
 
 	// finally transform into world space
-	float3 L_World = mul(L, HairBasis);
+	const float3 L_World = mul(L, HairBasis);
 
 	return CreateMaterialSample(L_World, Weight, Pdf, sign(dot(Payload.WorldGeoNormal, L_World)), Payload.Roughness);
 }
@@ -382,29 +382,29 @@ FMaterialEval Hair_EvalMaterial(
 	FPathTracingPayload Payload
 )
 {
-	float3 V_World = -IncomingDirection;
-	float3x3 HairBasis = GetHairBasis(Payload.WorldTangent, V_World);
+	const float3 V_World = -IncomingDirection;
+	const float3x3 HairBasis = GetHairBasis(Payload.WorldTangent, V_World);
 
-	float3 L_World = OutgoingDirection;
-	float3 V = mul(HairBasis, V_World);
-	float3 L = mul(HairBasis, L_World);
+	const float3 L_World = OutgoingDirection;
+	const float3 V = mul(HairBasis, V_World);
+	const float3 L = mul(HairBasis, L_World);
 
-	float h = 2 * Payload.CustomData.y - 1; // remap back to [-1,1]
-	float AzimuthalRoughness = lerp(1.0, 0.1, Payload.Metallic);
-	FHairData HairData = PrepareHairData(V, Payload.BaseColor, Payload.Roughness, AzimuthalRoughness, h, Payload.Specular);
+	const float h = 2 * Payload.GetHairPrimitiveUV().y - 1; // remap back to [-1,1]
+	const float AzimuthalRoughness = lerp(1.0, 0.1, Payload.Metallic);
+	const FHairData HairData = PrepareHairData(V, Payload.BaseColor, Payload.Roughness, AzimuthalRoughness, h, Payload.Specular);
 
 	// precompute data based of lighting direction
-	float SinThetaI = L.z;
-	float CosThetaI = TrigInverse(SinThetaI);
-	float PhiI = atan2(L.x, L.y); // NOTE: in our coordinate system, PhiO == 0, so we only need to account for L
+	const float SinThetaI = L.z;
+	const float CosThetaI = TrigInverse(SinThetaI);
+	const float PhiI = atan2(L.x, L.y); // NOTE: in our coordinate system, PhiO == 0, so we only need to account for L
 
 	FMaterialEval Result = NullMaterialEval();
 
-	float InvS = rcp(HairData.s);
-	float Pdf0 = Mp(SinThetaI, HairData.SinThetaO.x, HairData.v0) * Np(PhiI, 0, InvS, HairData.GammaO, HairData.GammaT);
-	float Pdf1 = Mp(SinThetaI, HairData.SinThetaO.y, HairData.v1) * Np(PhiI, 1, InvS, HairData.GammaO, HairData.GammaT);
-	float Pdf2 = Mp(SinThetaI, HairData.SinThetaO.z, HairData.v2) * Np(PhiI, 2, InvS, HairData.GammaO, HairData.GammaT);
-	float Pdf3 = Mp(SinThetaI, HairData.SinThetaO.w, HairData.v3) * (1 / (2 * PI)); // last lobe is isotropic
+	const float InvS = rcp(HairData.s);
+	const float Pdf0 = Mp(SinThetaI, HairData.SinThetaO.x, HairData.v0) * Np(PhiI, 0, InvS, HairData.GammaO, HairData.GammaT);
+	const float Pdf1 = Mp(SinThetaI, HairData.SinThetaO.y, HairData.v1) * Np(PhiI, 1, InvS, HairData.GammaO, HairData.GammaT);
+	const float Pdf2 = Mp(SinThetaI, HairData.SinThetaO.z, HairData.v2) * Np(PhiI, 2, InvS, HairData.GammaO, HairData.GammaT);
+	const float Pdf3 = Mp(SinThetaI, HairData.SinThetaO.w, HairData.v3) * (1 / (2 * PI)); // last lobe is isotropic
 
 	AddLobeWithMIS(Result.Weight, Result.Pdf, HairData.A0, Pdf0, HairData.LobePdf.x);
 	AddLobeWithMIS(Result.Weight, Result.Pdf, HairData.A1, Pdf1, HairData.LobePdf.y);
diff --git a/Engine/Shaders/Private/PathTracing/Material/PathTracingMaterialSampling.ush b/Engine/Shaders/Private/PathTracing/Material/PathTracingMaterialSampling.ush
--- a/Engine/Shaders/Private/PathTracing/Material/PathTracingMaterialSampling.ush
+++ b/Engine/Shaders/Private/PathTracing/Material/PathTracingMaterialSampling.ush
@@ -54,14 +54,14 @@ float3 GetMaterialTransparency(FPathTracingPayload Payload, float3 RayDirection)
 	{
 #if SIMPLIFIED_MATERIAL_SHADER
 		const float VoN = dot(RayDirection, Payload.WorldNormal);
-		return (1 - Payload.Opacity) * pow(Payload.CustomData.xyz, 1.0 / abs(VoN));
+		return (1 - Payload.Opacity) * pow(Payload.GetTransmittanceColor(), 1.0 / abs(VoN));
 #else
 		if (!Payload.HasRefraction() && Payload.Opacity < 1.0)
 		{
 			// ThinTranslucent has an optimized mode where it ignores refractivity and just handles the transmission via transparency
 			const float VoN = abs(dot(RayDirection, Payload.WorldNormal));
 			const float F0 = F0RGBToF0(Payload.SpecularColor);
-			return (1 - Payload.Opacity) * ComputeThinSlabWeights(Payload.CustomData.xyz, VoN, 0.0, F0).Transmitted;
+			return (1 - Payload.Opacity) * ComputeThinSlabWeights(Payload.GetTransmittanceColor(), VoN, 0.0, F0).Transmitted;
 		}
 		else
 		{
@@ -160,7 +160,7 @@ float3 EstimateMaterialAlbedo(FPathTracingPayload Payload)
 			Albedo = BrdfWeight * (Payload.DiffuseColor + Payload.SpecularColor);
 			break;
 		case SHADINGMODELID_CLOTH:
-			Albedo = BrdfWeight * (Payload.DiffuseColor + lerp(Payload.SpecularColor, Payload.CustomData.xyz, Payload.CustomData.w));
+			Albedo = BrdfWeight * (Payload.DiffuseColor + lerp(Payload.SpecularColor, Payload.GetClothColor(), Payload.GetClothAmount()));
 			break;
 		case SHADINGMODELID_HAIR:
 			Albedo = BrdfWeight * Payload.BaseColor;
@@ -179,7 +179,7 @@ float3 EstimateMaterialAlbedo(FPathTracingPayload Payload)
 			// If transparency is taking care of the tint, just account for the specular portion
 			// Otherwise, we are using (rough) refraction here, so account for both specular + transmission
 			const float F0 = F0RGBToF0(Payload.SpecularColor);
-			GlassAlbedo = Payload.HasRefraction() ? lerp(Payload.CustomData.xyz, 1.0, F0) : F0;
+			GlassAlbedo = Payload.HasRefraction() ? lerp(Payload.GetTransmittanceColor(), 1.0, F0) : F0;
 		}
 		Albedo = lerp(GlassAlbedo, Albedo, Payload.Opacity);
 	}
@@ -260,7 +260,7 @@ void AccumulateAlbedoNormal(inout FPathTracingPayload Payload, float3 PathThroug
 		Normal += BrdfWeight * Luminance(PathThroughput) * Payload.WorldNormal;
 		break;
 	case SHADINGMODELID_CLOTH:
-		Albedo += BrdfWeight * PathThroughput * (Payload.DiffuseColor + lerp(RoughnessWeight * Payload.SpecularColor, Payload.CustomData.xyz, Payload.CustomData.w));
+		Albedo += BrdfWeight * PathThroughput * (Payload.DiffuseColor + lerp(RoughnessWeight * Payload.SpecularColor, Payload.GetClothColor(), Payload.GetClothAmount()));
 		Normal += BrdfWeight * Luminance(PathThroughput) * Payload.WorldNormal;
 		break;
 #endif
@@ -279,7 +279,7 @@ void AccumulateAlbedoNormal(inout FPathTracingPayload Payload, float3 PathThroug
 			// If transparency is taking care of the tint, just account for the specular portion
 			// Otherwise, we are using (rough) refraction here, so account for both specular + transmission
 			const float F0 = F0RGBToF0(Payload.SpecularColor);
-			GlassAlbedo = Payload.HasRefraction() ? lerp(Payload.CustomData.xyz, 1.0, F0) : F0;
+			GlassAlbedo = Payload.HasRefraction() ? lerp(Payload.GetTransmittanceColor(), 1.0, F0) : F0;
 		}
 		Albedo += (1 - BrdfWeight) * RoughnessWeight * PathThroughput * GlassAlbedo;
 		Normal += (1 - BrdfWeight) * RoughnessWeight * Luminance(PathThroughput) * Payload.WorldNormal;
@@ -306,46 +306,6 @@ void AccumulateAlbedo(float3 SigmaS, float3 PathThroughput, float PathRoughness,
 	Albedo += RoughnessWeight * SigmaS * PathThroughput;
 }
 
-void AdjustMaterialRoughness(inout FPathTracingPayload Payload, float PathRoughness)
-{
-#if !SIMPLIFIED_MATERIAL_SHADER
-	// boost material roughness according to what we've seen so far
-	switch (Payload.ShadingModelID)
-	{
-		case SHADINGMODELID_UNLIT:
-			// this object does not scatter light
-			return;
-		case SHADINGMODELID_CLEAR_COAT:
-			// clamp coat roughness
-			Payload.SetClearCoatRoughness(max(Payload.GetClearCoatRoughness(), PathRoughness));
-			// fallthrough
-		case SHADINGMODELID_SINGLELAYERWATER:
-		case SHADINGMODELID_TWOSIDED_FOLIAGE:
-		case SHADINGMODELID_DEFAULT_LIT:
-		case SHADINGMODELID_THIN_TRANSLUCENT:
-		case SHADINGMODELID_CLOTH:
-		case SHADINGMODELID_SUBSURFACE:
-		case SHADINGMODELID_PREINTEGRATED_SKIN:
-		case SHADINGMODELID_EYE:
-		case SHADINGMODELID_HAIR:
-			Payload.Roughness = max(Payload.Roughness, PathRoughness);
-			break;
-		case SHADINGMODELID_SUBSURFACE_PROFILE:
-		{
-			// Clamp both specular lobes
-			float3 DualRoughnessData = Payload.GetDualRoughnessSpecular();
-			DualRoughnessData = max(DualRoughnessData, float3(PathRoughness, PathRoughness, 0.0));
-			Payload.SetDualRoughnessSpecular(DualRoughnessData);
-			break;
-		}
-		default:
-			// fallback is diffuse
-			Payload.Roughness = 1.0;
-			break;
-	}
-#endif
-}
-
 // returns an average roughness for the whole material (for the path roughness of shadow rays from this material)
 float GetAverageRoughness(FPathTracingPayload Payload)
 {
@@ -372,10 +332,9 @@ float GetAverageRoughness(FPathTracingPayload Payload)
 #endif
 }
 
+
 float3 DecodeSSSProfileRadius(uint ProfileId, float3 DiffuseColor, float Opacity)
 {
-	// NOTE: would be better to just precompute all this on the CPU and just have a single value to lookup
-
 	// Burley parameterization
 	float3 SurfaceAlbedo = View.SSProfilesTexture.Load(int3(BSSS_SURFACEALBEDO_OFFSET, ProfileId, 0)).xyz;
 	float3 DiffuseMeanFreePath = DecodeDiffuseMeanFreePath(View.SSProfilesTexture.Load(int3(BSSS_DMFP_OFFSET, ProfileId, 0))).xyz;
@@ -387,16 +346,11 @@ float3 DecodeSSSProfileRadius(uint ProfileId, float3 DiffuseColor, float Opacity
 	float3 SSSRadius = GetMFPFromDMFPApprox(SurfaceAlbedo, DiffuseColor, Opacity * WorldUnitScale * DiffuseMeanFreePath);
 
 	return SSSRadius * BURLEY_MM_2_CM;
-
 }
 
 
 FPathTracingPayload UnpackPathTracingPayloadMaterial(FPackedPathTracingPayload PackedInput, FRayDesc Ray)
 {
-	// the path tracer has a few more requirements compared to the ray-tracing path.
-	// Long term, we should reoranize the payload for the path tracer, but for now, just
-	// do a bit of post-processing and re-use fields to encode extracted information
-
 	// This function can also do a bit of per-material post-processing to simplify the rest of the code
 	FPathTracingPayload HitPayload = UnpackPathTracingPayload(PackedInput, Ray);
 
@@ -434,7 +388,7 @@ FPathTracingPayload UnpackPathTracingPayloadMaterial(FPackedPathTracingPayload P
 #if SIMPLIFIED_MATERIAL_SHADER
 		default:
 		{
-			HitPayload.CustomData = 0;
+			HitPayload.SetSubsurfaceColor(0.0);
 			break;
 		}
 #else
@@ -455,14 +409,6 @@ FPathTracingPayload UnpackPathTracingPayloadMaterial(FPackedPathTracingPayload P
 				HitPayload.SubsurfaceRadius = 0.0;
 				HitPayload.SetSubsurfaceColor(0);
 			}
-			// Decode dual roughness info
-			float Lobe0Roughness = 0;
-			float Lobe1Roughness = 0;
-			float LobeMix = 0;
-
-			GetSubsurfaceProfileDualSpecular(ProfileId, HitPayload.Roughness, HitPayload.Opacity, Lobe0Roughness, Lobe1Roughness, LobeMix);
-
-			HitPayload.SetDualRoughnessSpecular(float3(Lobe0Roughness, Lobe1Roughness, LobeMix));
 			// Reset opacity to avoid special cases elsewhere
 			HitPayload.Opacity = 1;
 			break;
@@ -475,7 +421,7 @@ FPathTracingPayload UnpackPathTracingPayloadMaterial(FPackedPathTracingPayload P
 				HitPayload.SubsurfaceRadius = DecodeSSSProfileRadius(ProfileId, HitPayload.DiffuseColor, 1.0);
 
 				// eye still needs a diffuse component for the iris center
-				float IrisMask = HitPayload.Anisotropy;
+				float IrisMask = HitPayload.GetEyeIrisMask();
 				HitPayload.SetSubsurfaceColor((1 - IrisMask) * HitPayload.DiffuseColor);
 				HitPayload.DiffuseColor *= IrisMask;
 			}
diff --git a/Engine/Shaders/Private/PathTracing/Material/PathTracingMedium.ush b/Engine/Shaders/Private/PathTracing/Material/PathTracingMedium.ush
--- a/Engine/Shaders/Private/PathTracing/Material/PathTracingMedium.ush
+++ b/Engine/Shaders/Private/PathTracing/Material/PathTracingMedium.ush
@@ -17,8 +17,8 @@ FMaterialSample Medium_SampleMaterial(
 	float4 RandSample)
 {
 	const float3 RayleighWeight = Payload.BaseColor;
-	const float3 HGWeight = Payload.CustomData.xyz;
-	const float G = Payload.CustomData.w;
+	const float3 HGWeight = Payload.GetHGWeight();
+	const float G = Payload.GetHGPhaseG();
 
 #if USE_UNIFORM_PHASE_FUNCTION
 	const float4 Result = UniformSampleSphere(RandSample.xy);
@@ -49,8 +49,9 @@ FMaterialEval Medium_EvalMaterial(
 )
 {
 	const float3 RayleighWeight = Payload.BaseColor;
-	const float3 HGWeight = Payload.CustomData.xyz;
-	const float G = Payload.CustomData.w;
+	const float3 HGWeight = Payload.GetHGWeight();
+	const float G = Payload.GetHGPhaseG();
+
 #if USE_UNIFORM_PHASE_FUNCTION
 	// simple omni-directional evaluation
 	return CreateMaterialEval(RayleighWeight + HGWeight, 1.0 / (4 * PI));
diff --git a/Engine/Shaders/Private/PathTracing/Material/PathTracingThinGlass.ush b/Engine/Shaders/Private/PathTracing/Material/PathTracingThinGlass.ush
--- a/Engine/Shaders/Private/PathTracing/Material/PathTracingThinGlass.ush
+++ b/Engine/Shaders/Private/PathTracing/Material/PathTracingThinGlass.ush
@@ -67,7 +67,7 @@ FMaterialEval RoughThinGlass_EvalMaterial(
 	const float Eta = Payload.Ior;
 	const float F0 = F0RGBToF0(Payload.SpecularColor);
 	const float F90 = saturate(50.0 * F0);
-	const float3 SlabColor = Payload.CustomData.xyz;
+	const float3 SlabColor = Payload.GetTransmittanceColor();
 
 	// transmitted lobe should have slightly less roughness due to the squeeze implied by the IOR (eventually converging to 0.0 as Ior approaches 1.0)
 	const float RoughnessR = Payload.Roughness;
@@ -142,7 +142,7 @@ FMaterialSample RoughThinGlass_SampleMaterial(
 	const float Eta = Payload.Ior;
 	const float F0 = F0RGBToF0(Payload.SpecularColor);
 	const float F90 = saturate(50.0 * F0);
-	const float3 SlabColor = Payload.CustomData.xyz;
+	const float3 SlabColor = Payload.GetTransmittanceColor();
 
 	// transmitted lobe should have slightly less roughness due to the squeeze implied by the IOR (eventually converging to 0.0 as Ior approaches 1.0)
 	const float RoughnessR = Payload.Roughness;
diff --git a/Engine/Shaders/Private/PathTracing/PathTracingCommon.ush b/Engine/Shaders/Private/PathTracing/PathTracingCommon.ush
--- a/Engine/Shaders/Private/PathTracing/PathTracingCommon.ush
+++ b/Engine/Shaders/Private/PathTracing/PathTracingCommon.ush
@@ -43,6 +43,24 @@ float3 PayloadDecodeUnitVector(uint Encoded)
 	return OctahedronToUnitVector(2 * OctF - 1);
 }
 
+float3 AdjustShadingNormal(float3 ShadingNormal, float3 GeoNormal, float3 RayDirection)
+{
+	// Clip shading normal in a view dependent way such that the reflection stays above the geometric normal
+	// This introduces a bit of view-dependency to the shading normal but fixes dark artifacts around grazing angles
+
+	float3 D = RayDirection;
+	float3 R = reflect(D, ShadingNormal);
+
+	// https://iquilezles.org/www/articles/dontflip/dontflip.htm
+	float k = dot(R, GeoNormal);
+	if (k < 0.0)
+	{
+		return normalize(normalize(R - k * GeoNormal) - D);
+	}
+	return ShadingNormal;
+}
+
+
 // This payload structure is what we transport between RGS/CHS/AHS programs
 struct FPackedPathTracingPayload : FMinimalPayload
 {
@@ -126,16 +144,16 @@ struct FPathTracingPayload : FMinimalPayload
 	float  Metallic;
 	float  Specular;
 	float  Roughness;
+	float  Anisotropy;
 	float  Ior;
 	uint   ShadingModelID;
 	uint   BlendingMode;
 	uint   Flags;
 	uint   PrimitiveLightingChannelMask;
-	float4 CustomData;
-	float3 CustomVector;
+	float4 CustomData0;
+	float4 CustomData1;
 	float3 WorldTangent;
 	float3 SubsurfaceRadius;
-	float  Anisotropy;
 	float3 TranslatedWorldPos;
 
 	void SetFrontFace() { Flags |= PATH_TRACING_PAYLOAD_OUTPUT_FLAG_FRONT_FACE; }
@@ -146,25 +164,61 @@ struct FPathTracingPayload : FMinimalPayload
 
 
 	// Various ways to interpret CustomData (depending on ShadingModelID)
-	float GetClearCoat() { return CustomData.x; }
-	float GetClearCoatRoughness() { return CustomData.y; }
-	void SetClearCoatRoughness(float ClearCoatRoughness) { CustomData.y = ClearCoatRoughness; }
-	float3 GetClearCoatBottomNormal() { return CustomVector; }
+	// NOTE: This is not always following the same conventions as GetMaterialCustomData0,1()
+
+	// SHADINGMODELID_CLOTH
+	float3 GetClothColor() { return CustomData0.xyz; }
+	void SetClothColor(float3 ClothColor) { CustomData0.xyz = ClothColor; }
+	float GetClothAmount() { return CustomData0.w; }
+	void SetClothAmount(float ClothAmount) { CustomData0.w = ClothAmount; }
+
+	// SHADINGMODELID_CLEAR_COAT
+	float GetClearCoat() { return CustomData0.x; }
+	float GetClearCoatRoughness() { return CustomData0.y; }
+	void SetClearCoat(float ClearCoat) { CustomData0.x = ClearCoat; }
+	void SetClearCoatRoughness(float ClearCoatRoughness) { CustomData0.y = ClearCoatRoughness; }
+	float3 GetClearCoatBottomNormal() { return CustomData1.xyz; }
+	void SetClearCoatBottomNormal(float3 BottomNormal) { CustomData1.xyz = BottomNormal; }
+
+	// SHADINGMODELID_SUBSURFACE_PROFILE and SHADINGMODELID_EYE
+	uint GetSubsurfaceProfileId() { return uint(CustomData0.a * 255.0 + 0.5); }
+
+	// SHADINGMODELID_SUBSURFACE, SHADINGMODELID_TWO_SIDED_FOLIAGE (and SHADINGMODELID_SUBSURFACE_PROFILE, SHADINGMODELID_EYE after decode)
+	float3 GetSubsurfaceColor() { return CustomData0.xyz; }
+	void SetSubsurfaceColor(float3 SubsurfaceColor) { CustomData0.xyz = SubsurfaceColor; }
+	
+	// SHADINGMODELID_SUBSURFACE_PROFILE
+	float3 GetDualRoughnessSpecular() { return CustomData1.xyz; }
+	void SetDualRoughnessSpecular(float Roughness0, float Roughness1, float LobeMix) { CustomData1.xyz = float3(Roughness0, Roughness1, LobeMix); }
 
-	uint GetSubsurfaceProfileId() { return uint(CustomData.a * 255.0 + 0.5); }
 
-	float3 GetSubsurfaceColor() { return CustomData.xyz; }
-	void SetSubsurfaceColor(float3 SubsurfaceColor) { CustomData.xyz = SubsurfaceColor; }
-	
-	// Steal an unused float3 (the only material with dual specular does not support anisotropy)
-	// Data is encoded as (x: roughness 0, y: roughness 1, z: roughness mix)
-	float3 GetDualRoughnessSpecular() { return WorldTangent; }
-	void SetDualRoughnessSpecular(float3 RoughnessData)
-	{
-		// store data in unused field and compute average roughness
-		WorldTangent = RoughnessData;
-		Roughness = lerp(RoughnessData.x, RoughnessData.y, RoughnessData.z);
-	}
+	// SHADINGMODELID_EYE
+	float3 GetEyeCausticNormal() { return CustomData1.xyz; }
+	void SetEyeCausticNormal(float3 CausticNormal) { CustomData1.xyz = CausticNormal; }
+	float GetEyeIrisMask() { return Anisotropy; }
+	float3 GetEyeIrisNormal() { return WorldTangent; }
+	void SetEyeIrisMask(float IrisMask) { Anisotropy = IrisMask; }
+	void SetEyeIrisNormal(float3 IrisNormal) { WorldTangent = IrisNormal; }
+
+	// SHADINGMODELID_HAIR
+	float2 GetHairPrimitiveUV() { return CustomData0.xy; }
+	void SetHairPrimitiveUV(float2 UV) { CustomData0.xy = UV; }
+
+	// SHADINGMODELID_DEFAULT_LIT, SHADINGMODELID_THIN_TRANSLUCENT
+	float3 GetTransmittanceColor() { return CustomData0.xyz; }
+	void SetTransmittanceColor(float3 TransmittanceColor) { CustomData0.xyz = TransmittanceColor; }
+
+	// SHADINGMODELID_SINGLELAYERWATER
+	float3 GetWaterExtinction() { return CustomData0.xyz; }
+	void SetWaterExtinction(float3 SigmaT) { CustomData0.xyz = SigmaT; }
+	float GetWaterPhaseG() { return CustomData0.w; }
+	void SetWaterPhaseG(float G) { CustomData0.w = G; }
+
+	// SHADINGMODELID_MEDIUM
+	float3 GetHGWeight() { return CustomData0.xyz; }
+	float GetHGPhaseG() { return CustomData0.w; }
+	void SetHG(float3 Weight, float G) { CustomData0 = float4(Weight, G); }
+
 
 	bool HasRefraction()
 	{
@@ -200,9 +254,10 @@ struct FPathTracingPayload : FMinimalPayload
 	}
 };
 
-FPackedPathTracingPayload InitPathTracingPayload(bool bIgnoreSkyMaterials)
+FPackedPathTracingPayload InitPathTracingPayload(bool bIgnoreSkyMaterials, float PathRoughness)
 {
 	FPackedPathTracingPayload Output = (FPackedPathTracingPayload)0;
+	Output.SetPathRoughness(PathRoughness);
 	if (bIgnoreSkyMaterials)
 	{
 		Output.SetIgnoreSkyMaterials();
@@ -253,16 +308,18 @@ FPackedPathTracingPayload PackPathTracingPayload(FPathTracingPayload Input)
 
 	Output.PackedData[9] = PayloadEncodeUnitVector(Input.WorldTangent);
 
-	Output.PackedData[10] = f32tof16(Input.CustomVector.x);
-	Output.PackedData[10] |= f32tof16(Input.CustomVector.y) << 16;
-	Output.PackedData[11] = f32tof16(Input.CustomVector.z);
+	Output.PackedData[10] = f32tof16(Input.CustomData0.x);
+	Output.PackedData[10] |= f32tof16(Input.CustomData0.y) << 16;
+	Output.PackedData[11] = f32tof16(Input.CustomData0.z);
+	Output.PackedData[11] |= f32tof16(Input.CustomData0.w) << 16;
 
-	Output.PackedData[11] |= f32tof16(Input.CustomData.x) << 16;
-	Output.PackedData[12] = f32tof16(Input.CustomData.y);
-	Output.PackedData[12] |= f32tof16(Input.CustomData.z) << 16;
-	Output.PackedData[13] = f32tof16(Input.CustomData.w);
-	Output.PackedData[13] |= f32tof16(Input.Ior) << 16;
-	Output.PackedData[14] = 0; // unused
+	Output.PackedData[12] = f32tof16(Input.CustomData1.x);
+	Output.PackedData[12] |= f32tof16(Input.CustomData1.y) << 16;
+	Output.PackedData[13] = f32tof16(Input.CustomData1.z);
+	Output.PackedData[13] |= f32tof16(Input.CustomData1.w) << 16;
+
+	Output.PackedData[14] = f32tof16(Input.Ior);
+	// NOTE: 16 bits left in PackedData[14]
 
 	return Output;
 }
@@ -296,15 +353,17 @@ FPathTracingPayload UnpackPathTracingPayload(FPackedPathTracingPayload Input, FR
 
 	Output.WorldTangent = PayloadDecodeUnitVector(Input.PackedData[9]);
 
-	Output.CustomVector.x = f16tof32(Input.PackedData[10]);
-	Output.CustomVector.y = f16tof32(Input.PackedData[10] >> 16);
-	Output.CustomVector.z = f16tof32(Input.PackedData[11]);
+	Output.CustomData0.x = f16tof32(Input.PackedData[10]);
+	Output.CustomData0.y = f16tof32(Input.PackedData[10] >> 16);
+	Output.CustomData0.z = f16tof32(Input.PackedData[11]);
+	Output.CustomData0.w = f16tof32(Input.PackedData[11] >> 16);
+
+	Output.CustomData1.x = f16tof32(Input.PackedData[12]);
+	Output.CustomData1.y = f16tof32(Input.PackedData[12] >> 16);
+	Output.CustomData1.z = f16tof32(Input.PackedData[13]);
+	Output.CustomData1.w = f16tof32(Input.PackedData[13] >> 16);
 
-	Output.CustomData.x = f16tof32(Input.PackedData[11] >> 16);
-	Output.CustomData.y = f16tof32(Input.PackedData[12]);
-	Output.CustomData.z = f16tof32(Input.PackedData[12] >> 16);
-	Output.CustomData.w = f16tof32(Input.PackedData[13]);
-	Output.Ior = f16tof32(Input.PackedData[13] >> 16);
+	Output.Ior = f16tof32(Input.PackedData[14]);
 
 	Output.TranslatedWorldPos = Ray.Origin + Output.HitT * Ray.Direction;
 
diff --git a/Engine/Shaders/Private/PathTracing/PathTracingCore.ush b/Engine/Shaders/Private/PathTracing/PathTracingCore.ush
--- a/Engine/Shaders/Private/PathTracing/PathTracingCore.ush
+++ b/Engine/Shaders/Private/PathTracing/PathTracingCore.ush
@@ -152,14 +152,10 @@ void ApplyDecal(inout FPathTracingPayload HitPayload, FDecalShaderParams Params)
 FPathTracingPayload TraceTransparentRay(inout FPathState PathState, bool bIsCameraRay, bool bLastBounce, bool bIncludeEmission, uint NumVisibleLights, inout FVolumeSegment VolumeSegment)
 {
 	const uint RayFlags = bIsCameraRay && EnableCameraBackfaceCulling ? RAY_FLAG_CULL_BACK_FACING_TRIANGLES : 0;
-	uint InstanceInclusionMask = 0; 
-	if (bIsCameraRay)
+	uint InstanceInclusionMask = bIsCameraRay ? PATHTRACER_MASK_CAMERA : PATHTRACER_MASK_INDIRECT;
+	if (length2(PathState.Ray.Origin) < Pow2(PATH_TRACING_SKIP_HAIR_DISTANCE))
 	{
-		InstanceInclusionMask = PATHTRACER_MASK_CAMERA | (length2(PathState.Ray.Origin) < Pow2(PATH_TRACING_SKIP_HAIR_DISTANCE) ? PATHTRACER_MASK_HAIR_CAMERA : 0);
-	}
-	else
-	{
-		InstanceInclusionMask = PATHTRACER_MASK_INDIRECT | (length2(PathState.Ray.Origin) < Pow2(PATH_TRACING_SKIP_HAIR_DISTANCE) ? PATHTRACER_MASK_HAIR_INDIRECT : 0);
+		InstanceInclusionMask |= bIsCameraRay ? PATHTRACER_MASK_HAIR_CAMERA : PATHTRACER_MASK_HAIR_INDIRECT;
 	}
 	
 	const uint MissShaderIndex = 0;
@@ -188,7 +184,7 @@ FPathTracingPayload TraceTransparentRay(inout FPathState PathState, bool bIsCame
 		// also avoid noise from hot spots (they can be properly
 		// importance sampled if a capturing skylight is added)
 		const bool bIgnoreSky = !bIsCameraRay;
-		FPackedPathTracingPayload PackedPayload = InitPathTracingPayload(bIgnoreSky);
+		FPackedPathTracingPayload PackedPayload = InitPathTracingPayload(bIgnoreSky, ApproximateCaustics ? PathState.PathRoughness : 0.0);
 		// Trace the ray
 		TraceRay(
 			TLAS,
@@ -390,6 +386,12 @@ FPathTracingPayload TraceTransparentRay(inout FPathState PathState, bool bIsCame
 				CallShader(DecalId, Params);
 
 				ApplyDecal(HitPayload, Params);
+
+				if (ApproximateCaustics)
+				{
+					Payload.Roughness = max(Payload.Roughness, PathState.PathRoughness);
+				}
+				HitPayload.WorldNormal = AdjustShadingNormal(HitPayload.WorldNormal, HitPayload.WorldGeoNormal, PathState.Ray.Direction);
 			}
 		}
 #endif
@@ -538,7 +540,7 @@ struct FProbeResult
 FProbeResult TraceProbeRay(FRayDesc Ray)
 {
 	// Trace a short ray to see if we escaped the surface
-	FPackedPathTracingPayload PackedPayload = InitPathTracingPayload(true);
+	FPackedPathTracingPayload PackedPayload = InitPathTracingPayload(true, 1.0);
 	// TODO: use smaller payload + flag to skip shading work? (only want smooth+geo normals)
 	//       however supporting blockers would require some shading ...
 	const uint RayFlags = 0;
@@ -831,23 +833,6 @@ FPathState CreatePathState(int2 PixelIndex)
 	return Output;
 }
 
-float3 AdjustShadingNormal(float3 ShadingNormal, float3 GeoNormal, float3 RayDirection)
-{
-	// Clip shading normal in a view dependent way such that the reflection stays above the geometric normal
-	// This introduces a bit of view-dependency to the shading normal but fixes dark artifacts around grazing angles
-
-	float3 D = RayDirection;
-	float3 R = reflect(D, ShadingNormal);
-
-	// https://iquilezles.org/www/articles/dontflip/dontflip.htm
-	float k = dot(R, GeoNormal);
-	if (k < 0.0)
-	{
-		return normalize(normalize(R - k * GeoNormal) - D);
-	}
-	return ShadingNormal;
-}
-
 bool PathTracingKernel(inout FPathState PathState, int Bounce)
 {
 	// This array will hold a CDF for light picking
@@ -1000,11 +985,6 @@ bool PathTracingKernel(inout FPathState PathState, int Bounce)
 		return false;
 	}
 
-	// fix dark edges in reflections from shading normals
-	Payload.WorldNormal       = AdjustShadingNormal(Payload.WorldNormal      , Payload.WorldGeoNormal, PathState.Ray.Direction);
-	Payload.WorldSmoothNormal = AdjustShadingNormal(Payload.WorldSmoothNormal, Payload.WorldGeoNormal, PathState.Ray.Direction);
-
-
 #if 0
 	PathState.Radiance = 0.5 * Payload.WorldGeoNormal + 0.5;
 	return false;
@@ -1013,15 +993,6 @@ bool PathTracingKernel(inout FPathState PathState, int Bounce)
 	// As soon as the path is blurry enough, we can get away with diffuse sampling only
 	const bool bSimplifySSS = PathState.PathRoughness >= 0.15;
 
-	if (ApproximateCaustics)
-	{
-		// modify the payload roughness to minimize difficult caustics
-		// This is inspired by a trick used in the Arnold renderer:
-		//   https://cgg.mff.cuni.cz/~jaroslav/gicourse2010/giai2010-02-marcos_fajardo-slides.pdf (slide 39)
-		//   https://www.arnoldrenderer.com/research/Arnold_TOG2018.pdf (section 4.2)
-		AdjustMaterialRoughness(Payload, PathState.PathRoughness);
-	}
-
 	// Processing the random walk will (stochastically) move the shading point elsewhere on the surface of the object
 	if (!ProcessSubsurfaceRandomWalk(Payload, PathState.PathThroughput, PathState.RandSequence, PathState.Ray.Direction, bSimplifySSS))
 	{
@@ -1204,28 +1175,15 @@ bool PathTracingKernel(inout FPathState PathState, int Bounce)
 	// enlarge roughness based on the chosen lobe roughness
 	PathState.PathRoughness = max(PathState.PathRoughness, MaterialSample.Roughness);
 
-	// update the current extinction if we are crossing a boundary on glass
+	// update the current extinction if we are crossing a boundary on glass or water
 	// summing the local extinction gives a rudimentary way of dealing with overlapping regions
 	// long term we will probably want a stack with priorities
-	if (MaterialSample.PositionBiasSign < 0 && Payload.IsMaterialSolidGlass())
-	{
-		float3 LocalSigmaT = -log(max(Payload.CustomData.xyz, 1e-8f)) * AbsorptionScale;
-		if (Payload.IsFrontFace())
-		{
-			// entering
-			PathState.SigmaT += LocalSigmaT;
-		}
-		else
-		{
-			// exiting
-			PathState.SigmaT -= LocalSigmaT;
-			PathState.SigmaT = max(PathState.SigmaT, 0);
-		}
-	}
-
-	if (MaterialSample.PositionBiasSign < 0 && Payload.ShadingModelID == SHADINGMODELID_SINGLELAYERWATER)
+	if (MaterialSample.PositionBiasSign < 0 && (Payload.IsMaterialSolidGlass() || Payload.ShadingModelID == SHADINGMODELID_SINGLELAYERWATER))
 	{
-		float3 LocalSigmaT = max(Payload.CustomData.xyz, 1e-8f);
+		const float3 LocalSigmaT =
+			Payload.ShadingModelID == SHADINGMODELID_SINGLELAYERWATER ?
+			Payload.GetWaterExtinction() :
+			-log(max(Payload.GetTransmittanceColor(), 1e-8f)) * AbsorptionScale;
 		if (Payload.IsFrontFace())
 		{
 			// entering
diff --git a/Engine/Shaders/Private/PathTracing/PathTracingInitExtinctionCoefficient.usf b/Engine/Shaders/Private/PathTracing/PathTracingInitExtinctionCoefficient.usf
--- a/Engine/Shaders/Private/PathTracing/PathTracingInitExtinctionCoefficient.usf
+++ b/Engine/Shaders/Private/PathTracing/PathTracingInitExtinctionCoefficient.usf
@@ -9,8 +9,6 @@ RWBuffer<float> RWStartingExtinctionCoefficient;
 // TODO: figure out the right start value for camera rays inside glass
 RAY_TRACING_ENTRY_RAYGEN(PathTracingInitExtinctionCoefficientRG)
 {
-	FPackedPathTracingPayload PackedPayload = InitPathTracingPayload(true);
-
 	const uint RayFlags = RAY_FLAG_FORCE_OPAQUE;
 	const uint MissShaderIndex = 0;
 	const uint InstanceInclusionMask = PATHTRACER_MASK_ALL 
@@ -23,27 +21,10 @@ RAY_TRACING_ENTRY_RAYGEN(PathTracingInitExtinctionCoefficientRG)
 	UpRay.TMin = 0.0;
 	UpRay.TMax = POSITIVE_INFINITY;
 		
-	TraceRay(
-		TLAS,
-		RayFlags,
-		InstanceInclusionMask,
-		RAY_TRACING_SHADER_SLOT_MATERIAL,
-		RAY_TRACING_NUM_SHADER_SLOTS,
-		MissShaderIndex,
-		UpRay.GetNativeDesc(),
-		PackedPayload);
-
-	float3 StartingSigma = 0.0;
-	while (PackedPayload.IsHit())
+	float3 StartingSigmaT = 0.0;
+	for (;;)
 	{
-		FPathTracingPayload HitPayload = UnpackPathTracingPayload(PackedPayload, UpRay);
-		if (HitPayload.ShadingModelID == SHADINGMODELID_SINGLELAYERWATER)
-		{			
-			StartingSigma = max(HitPayload.CustomData.xyz, 1e-8f);
-			break;
-		}
-
-		UpRay.TMin = asfloat(asuint(HitPayload.HitT) + 1);
+		FPackedPathTracingPayload PackedPayload = InitPathTracingPayload(true, 0.0);
 		TraceRay(
 			TLAS,
 			RayFlags,
@@ -53,9 +34,23 @@ RAY_TRACING_ENTRY_RAYGEN(PathTracingInitExtinctionCoefficientRG)
 			MissShaderIndex,
 			UpRay.GetNativeDesc(),
 			PackedPayload);
+		if (PackedPayload.IsMiss())
+		{
+			// we didn't hit anything
+			break;
+		}
+		FPathTracingPayload HitPayload = UnpackPathTracingPayload(PackedPayload, UpRay);
+		if (HitPayload.ShadingModelID == SHADINGMODELID_SINGLELAYERWATER)
+		{
+			// Found water! we are done
+			StartingSigmaT = HitPayload.GetWaterExtinction();
+			break;
+		}
+		// otherwise, keep tracing
+		UpRay.TMin = asfloat(asuint(HitPayload.HitT) + 1);
 	}
 
-	RWStartingExtinctionCoefficient[0] = StartingSigma.x;
-	RWStartingExtinctionCoefficient[1] = StartingSigma.y;
-	RWStartingExtinctionCoefficient[2] = StartingSigma.z;
+	RWStartingExtinctionCoefficient[0] = StartingSigmaT.x;
+	RWStartingExtinctionCoefficient[1] = StartingSigmaT.y;
+	RWStartingExtinctionCoefficient[2] = StartingSigmaT.z;
 }
diff --git a/Engine/Shaders/Private/PathTracing/PathTracingMaterialHitShader.usf b/Engine/Shaders/Private/PathTracing/PathTracingMaterialHitShader.usf
--- a/Engine/Shaders/Private/PathTracing/PathTracingMaterialHitShader.usf
+++ b/Engine/Shaders/Private/PathTracing/PathTracingMaterialHitShader.usf
@@ -30,6 +30,16 @@
 #include "/Engine/Private/PathTracing/Material/PathTracingFresnel.ush"
 #include "/Engine/Private/PathTracing/Material/PathTracingThinGlass.ush"
 
+float AdjustMaterialRoughness(float Roughness, float PathRoughness)
+{
+	// Modify the payload roughness to minimize difficult caustics
+	// This is inspired by a trick used in the Arnold renderer:
+	//   https://cgg.mff.cuni.cz/~jaroslav/gicourse2010/giai2010-02-marcos_fajardo-slides.pdf (slide 39)
+	//   https://www.arnoldrenderer.com/research/Arnold_TOG2018.pdf (section 4.2)
+	// NOTE: If approximate caustics are disabled, the path roughness passed in here will be 0.0 which effectively turns off this optimization
+	return max(Roughness, PathRoughness);
+}
+
 float GetRefractionIor(FPixelMaterialInputs PixelMaterialInputs)
 {
 #if REFRACTION_ROOT_NODE_OVERRIDES_DEFAULT || STRATA_OPTIMIZED_UNLIT
@@ -149,6 +159,11 @@ RAY_TRACING_ENTRY_CLOSEST_HIT(PathTracingMaterialCHS,
 	Payload.WorldGeoNormal = GeoNormalSign * GeoNormal;
 	Payload.WorldNormal = MaterialParameters.WorldNormal;
 	Payload.WorldSmoothNormal = MaterialParameters.TwoSidedSign * TransformTangentNormalToWorld(MaterialParameters.TangentToWorld, float3(0, 0, 1));
+
+	// fix dark edges in reflections from shading normals
+	Payload.WorldNormal       = AdjustShadingNormal(Payload.WorldNormal      , Payload.WorldGeoNormal, WorldRayDirection());
+	Payload.WorldSmoothNormal = AdjustShadingNormal(Payload.WorldSmoothNormal, Payload.WorldGeoNormal, WorldRayDirection());
+
 	Payload.Radiance = GetMaterialEmissive(PixelMaterialInputs);
 #if TWO_SIDED_MATERIAL == 0
 	if (MaterialParameters.TwoSidedSign < 0)
@@ -159,8 +174,8 @@ RAY_TRACING_ENTRY_CLOSEST_HIT(PathTracingMaterialCHS,
 #endif
 	Payload.BaseColor = BaseColor;
 	Payload.Specular = Specular;
-	Payload.Roughness = Roughness;
 	Payload.Metallic = Metallic;
+	Payload.Roughness = Roughness;
 	Payload.Opacity = GetMaterialOpacity(PixelMaterialInputs);
 
 #if MATERIAL_USES_ANISOTROPY
@@ -193,22 +208,28 @@ RAY_TRACING_ENTRY_CLOSEST_HIT(PathTracingMaterialCHS,
 #if MATERIAL_SHADINGMODEL_CLEAR_COAT
 	if (Payload.ShadingModelID == SHADINGMODELID_CLEAR_COAT)
 	{
-    	Payload.CustomData.x = GetMaterialCustomData0(MaterialParameters);
-		Payload.CustomData.y = GetMaterialCustomData1(MaterialParameters);
+		const float ClearCoat = GetMaterialCustomData0(MaterialParameters);
+		const float ClearCoatRoughness = GetMaterialCustomData1(MaterialParameters);
+		Payload.SetClearCoat(ClearCoat);
+		Payload.SetClearCoatRoughness(AdjustMaterialRoughness(ClearCoatRoughness, PackedPayload.GetPathRoughness()));
 #if CLEAR_COAT_BOTTOM_NORMAL
 #if NUM_MATERIAL_OUTPUTS_CLEARCOATBOTTOMNORMAL > 0
 #if MATERIAL_TANGENTSPACENORMAL
-		float3 tempnormal = normalize(TransformTangentVectorToWorld(MaterialParameters.TangentToWorld, ClearCoatBottomNormal0(MaterialParameters)));
+		float3 BottomNormal = normalize(TransformTangentVectorToWorld(MaterialParameters.TangentToWorld, ClearCoatBottomNormal0(MaterialParameters)));
 #else
-		float3 tempnormal = ClearCoatBottomNormal0(MaterialParameters);
+		float3 BottomNormal = ClearCoatBottomNormal0(MaterialParameters);
 #endif
-		Payload.CustomVector = tempnormal;
+		
+		// if we got a custom normal for the clearcoat, adjust it now
+		BottomNormal = AdjustShadingNormal(BottomNormal, Payload.WorldGeoNormal, WorldRayDirection());
+		Payload.SetClearCoatBottomNormal(BottomNormal);
 #else
-		Payload.CustomVector = Payload.WorldNormal;
+		Payload.SetClearCoatBottomNormal(Payload.WorldNormal);
 #endif
 #else
-		Payload.CustomVector = Payload.WorldNormal;
+		Payload.SetClearCoatBottomNormal(Payload.WorldNormal);
 #endif
+		
 	}
 	else
 #endif
@@ -216,8 +237,10 @@ RAY_TRACING_ENTRY_CLOSEST_HIT(PathTracingMaterialCHS,
 #if MATERIAL_SHADINGMODEL_CLOTH
 	if (Payload.ShadingModelID == SHADINGMODELID_CLOTH)
 	{
-		Payload.CustomData = GetMaterialSubsurfaceData(PixelMaterialInputs);
-		Payload.CustomData.a = saturate(GetMaterialCustomData0(MaterialParameters)); // Fuzz
+		const float3 ClothColor = GetMaterialSubsurfaceData(PixelMaterialInputs).rgb;
+		const float Fuzz = saturate(GetMaterialCustomData0(MaterialParameters));
+		Payload.SetClothColor(ClothColor);
+		Payload.SetClothAmount(Fuzz);
 	}
 	else
 #endif
@@ -225,14 +248,32 @@ RAY_TRACING_ENTRY_CLOSEST_HIT(PathTracingMaterialCHS,
 #if MATERIAL_SHADINGMODEL_TWOSIDED_FOLIAGE || MATERIAL_SHADINGMODEL_SUBSURFACE || MATERIAL_SHADINGMODEL_PREINTEGRATED_SKIN || MATERIAL_SHADINGMODEL_SUBSURFACE_PROFILE
 	if (Payload.ShadingModelID == SHADINGMODELID_TWOSIDED_FOLIAGE ||
 		Payload.ShadingModelID == SHADINGMODELID_SUBSURFACE ||
-		Payload.ShadingModelID == SHADINGMODELID_PREINTEGRATED_SKIN ||
-		Payload.ShadingModelID == SHADINGMODELID_SUBSURFACE_PROFILE)
+		Payload.ShadingModelID == SHADINGMODELID_PREINTEGRATED_SKIN)
 	{
-		Payload.CustomData = GetMaterialSubsurfaceData(PixelMaterialInputs);
+		const float3 SubsurfaceColor = GetMaterialSubsurfaceData(PixelMaterialInputs).rgb;
+		Payload.SetSubsurfaceColor(SubsurfaceColor);
 	}
 	else
 #endif
+#if MATERIAL_SHADINGMODEL_SUBSURFACE_PROFILE
+	if (Payload.ShadingModelID == SHADINGMODELID_SUBSURFACE_PROFILE)
+	{
+		Payload.CustomData0 = GetMaterialSubsurfaceData(PixelMaterialInputs);
+		uint ProfileId = Payload.GetProfileId();
+
+		// Decode dual roughness info
+		float Lobe0Roughness = 0;
+		float Lobe1Roughness = 0;
+		float LobeMix = 0;
+
+		GetSubsurfaceProfileDualSpecular(ProfileId, Payload.Roughness, Payload.Opacity, Lobe0Roughness, Lobe1Roughness, LobeMix);
 
+		Lobe0Roughness = AdjustMaterialRoughness(Lobe0Roughness, PackedPayload.GetPathRoughness();
+		Lobe1Roughness = AdjustMaterialRoughness(Lobe1Roughness, PackedPayload.GetPathRoughness();
+		Payload.SetDualRoughnessSpecular(Lobe0Roughness, Lobe1Roughness, LobeMix);
+	}
+	else
+#endif
 #if MATERIAL_SHADINGMODEL_SINGLELAYERWATER
 	if (Payload.ShadingModelID == SHADINGMODELID_SINGLELAYERWATER)
 	{
@@ -243,14 +284,14 @@ RAY_TRACING_ENTRY_CLOSEST_HIT(PathTracingMaterialCHS,
 
 		const float3 ScatteringCoeff = max(0.0f, LWCToFloat(GetSingleLayerWaterMaterialOutput0(WaterMaterialParameters)));
 		const float3 AbsorptionCoeff = max(0.0f, LWCToFloat(GetSingleLayerWaterMaterialOutput1(WaterMaterialParameters)));
-		const float PhaseG = clamp(LWCToFloat(GetSingleLayerWaterMaterialOutput2(WaterMaterialParameters)), -1.0f, 1.0f);
+		const float PhaseG = clamp(LWCToFloat(GetSingleLayerWaterMaterialOutput2(WaterMaterialParameters)), -0.99f, 0.99f);
 
 		const float3 WaterExtinction = ScatteringCoeff + AbsorptionCoeff;
 		const float3 WaterAlbedo = ScatteringCoeff / WaterExtinction;
 
-		Payload.BaseColor.xyz = WaterAlbedo;
-		Payload.CustomData.xyz = WaterExtinction;
-		Payload.CustomData.w = PhaseG;
+		Payload.BaseColor = WaterAlbedo;
+		Payload.SetWaterExtinction(WaterExtinction);
+		Payload.SetWaterPhaseG(PhaseG);
 
 		Payload.Ior = DielectricF0ToIor(DielectricSpecularToF0(Specular));
 		Payload.BlendingMode = RAY_TRACING_BLEND_MODE_TRANSLUCENT;
@@ -261,7 +302,7 @@ RAY_TRACING_ENTRY_CLOSEST_HIT(PathTracingMaterialCHS,
 #if MATERIAL_SHADINGMODEL_THIN_TRANSLUCENT
 		if (Payload.ShadingModelID == SHADINGMODELID_THIN_TRANSLUCENT)
 		{
-			Payload.CustomData.xyz = GetThinTranslucentMaterialOutput0(MaterialParameters);
+			Payload.SetTransmittanceColor(GetThinTranslucentMaterialOutput0(MaterialParameters));
 		}
 		else
 #endif
@@ -271,9 +312,9 @@ RAY_TRACING_ENTRY_CLOSEST_HIT(PathTracingMaterialCHS,
 	{
 		// TODO: encode hair specific info
 #if HAIR_STRAND_MESH_FACTORY
-		Payload.CustomData.xy = MaterialParameters.HairPrimitiveUV;
+		Payload.SetHairPrimitiveUV(MaterialParameters.HairPrimitiveUV);
 #else
-		Payload.CustomData.xy = 0.5;
+		Payload.SetHairPrimitiveUV(0.5);
 #endif
 	}
 	else
@@ -285,8 +326,8 @@ RAY_TRACING_ENTRY_CLOSEST_HIT(PathTracingMaterialCHS,
 		// This is all based on logic from ShadingModelsMaterial.ush
 		const float IrisMask     = saturate(GetMaterialCustomData0(MaterialParameters));
 		const float IrisDistance = saturate(GetMaterialCustomData1(MaterialParameters));
-		Payload.Anisotropy = IrisMask;
-		Payload.CustomData.a = GetMaterialSubsurfaceData(PixelMaterialInputs).a; // SubsurfaceProfile (NOTE: this is not used, EyeBxDF leaves Diffuse=0, so does not actually apply SSS?)
+		Payload.SetEyeIrisMask(IrisMask);
+		Payload.CustomData0.a = GetMaterialSubsurfaceData(PixelMaterialInputs).a; // SubsurfaceProfile
 		Payload.Metallic = 0.0;
 #if NUM_MATERIAL_OUTPUTS_GETTANGENTOUTPUT > 0
 		// Blend in the negative intersection normal to create some concavity
@@ -295,10 +336,11 @@ RAY_TRACING_ENTRY_CLOSEST_HIT(PathTracingMaterialCHS,
 		// introduce some concavity, this does the job.
 		const float3 PlaneNormal = normalize(GetTangentOutput0(MaterialParameters));
 		const float3 CausticNormal = normalize(lerp(PlaneNormal, -Payload.WorldNormal, IrisMask * IrisDistance));
-		Payload.CustomVector = CausticNormal;
+		Payload.SetEyeCausticNormal(CausticNormal);
+		// NOTE: calling AdjustShadingNormal on the custom normal does not seem necessary since this is only used for a diffuse calculation
 #else
 		const float3 PlaneNormal = Payload.WorldNormal;
-		Payload.CustomVector = Payload.WorldNormal;
+		Payload.SetEyeCausticNormal(Payload.WorldNormal);
 #endif
 
 #if IRIS_NORMAL // on
@@ -312,10 +354,10 @@ RAY_TRACING_ENTRY_CLOSEST_HIT(PathTracingMaterialCHS,
 	#else
 		float3 IrisNormal = PlaneNormal;
 	#endif
-		Payload.WorldTangent = IrisNormal;
+		Payload.SetEyeIrisNormal(IrisNormal);
 #else
 		// IRIS_NORMAL off
-		Payload.WorldTangent = PlaneNormal;
+		Payload.SetEyeIrisNormal(PlaneNormal);
 #endif
 	}
 	else
@@ -326,9 +368,10 @@ RAY_TRACING_ENTRY_CLOSEST_HIT(PathTracingMaterialCHS,
 	{
 		// only allow refraction for default lit materials since we need space for the absorption amount
 #if HAVE_GetAbsorptionMediumMaterialOutput0
-		Payload.CustomData.xyz = GetAbsorptionMediumMaterialOutput0(MaterialParameters);
+		Payload.SetTransmittanceColor(GetAbsorptionMediumMaterialOutput0(MaterialParameters));
 #else
-		Payload.CustomData.xyz = 1.0;
+		// Make the glass totally clear if no custom medium is set
+		Payload.SetTransmittanceColor(1.0);
 #endif
 	}
 	else
@@ -337,6 +380,8 @@ RAY_TRACING_ENTRY_CLOSEST_HIT(PathTracingMaterialCHS,
 		// terminal case for the conditionals above
 	}
 
+	// adjust after everything else (because SSS profile case needs to decode the dual spec info from the unmodified roughness)
+	Payload.Roughness = AdjustMaterialRoughness(Payload.Roughness, PackedPayload.GetPathRoughness());
 
 #else // SIMPLIFIED_MATERIAL_SHADER
 
@@ -347,14 +392,14 @@ RAY_TRACING_ENTRY_CLOSEST_HIT(PathTracingMaterialCHS,
 #if MATERIAL_SHADINGMODEL_TWOSIDED_FOLIAGE
 	if (Payload.ShadingModelID == SHADINGMODELID_TWOSIDED_FOLIAGE)
 	{
-		Payload.CustomData = GetMaterialSubsurfaceData(PixelMaterialInputs);
+		Payload.SetSubsurfaceColor(GetMaterialSubsurfaceData(PixelMaterialInputs).rgb);
 	}
 	else
 #endif
 #if MATERIAL_SHADINGMODEL_THIN_TRANSLUCENT
 	if (Payload.ShadingModelID == SHADINGMODELID_THIN_TRANSLUCENT)
 	{
-		Payload.CustomData.xyz = GetThinTranslucentMaterialOutput0(MaterialParameters);
+		Payload.SetTransmittanceColor(GetThinTranslucentMaterialOutput0(MaterialParameters));
 	}
 	else
 #endif
diff --git a/Engine/Shaders/Private/PathTracing/Volume/PathTracingVolumeCommon.ush b/Engine/Shaders/Private/PathTracing/Volume/PathTracingVolumeCommon.ush
--- a/Engine/Shaders/Private/PathTracing/Volume/PathTracingVolumeCommon.ush
+++ b/Engine/Shaders/Private/PathTracing/Volume/PathTracingVolumeCommon.ush
@@ -205,10 +205,11 @@ FPathTracingPayload CreateMediumHitPayload(float HitT, float3 TranslatedWorldPos
 	Result.Opacity = 1.0;
 	Result.PrimitiveLightingChannelMask = 7;
 	Result.SetFrontFace();
-	float3 SigmaS = ShadedResult.SigmaSRayleigh + ShadedResult.SigmaSHG;
-	Result.BaseColor = select(SigmaS > 0.0, min(ShadedResult.SigmaSRayleigh, SigmaS) / SigmaS, 0.0);
-	Result.CustomData.xyz = select(SigmaS > 0, min(ShadedResult.SigmaSHG, SigmaS) / SigmaS, 0.0);
-	Result.CustomData.w = ShadedResult.PhaseG;
+	const float3 SigmaS = ShadedResult.SigmaSRayleigh + ShadedResult.SigmaSHG;
+	const float3 RayleighWeight = select(SigmaS > 0.0, min(ShadedResult.SigmaSRayleigh, SigmaS) / SigmaS, 0.0);
+	const float3 HGWeight = select(SigmaS > 0, min(ShadedResult.SigmaSHG, SigmaS) / SigmaS, 0.0);
+	Result.BaseColor = RayleighWeight;
+	Result.SetHG(HGWeight, ShadedResult.PhaseG);
 	return Result;
 }
 
