Path Tracer: Fix typos in refactored code

diff --git a/Engine/Shaders/Private/PathTracing/PathTracingMaterialHitShader.usf b/Engine/Shaders/Private/PathTracing/PathTracingMaterialHitShader.usf
--- a/Engine/Shaders/Private/PathTracing/PathTracingMaterialHitShader.usf
+++ b/Engine/Shaders/Private/PathTracing/PathTracingMaterialHitShader.usf
@@ -27,6 +27,7 @@
 #include "/Engine/Private/RayTracing/RayTracingCalcInterpolants.ush"
 #include "/Engine/Private/ShadingCommon.ush"
 #include "/Engine/Private/DeferredShadingCommon.ush"
+#include "/Engine/Private/SubsurfaceProfileCommon.ush"
 #include "/Engine/Private/PathTracing/Material/PathTracingFresnel.ush"
 #include "/Engine/Private/PathTracing/Material/PathTracingThinGlass.ush"
 
@@ -259,7 +260,7 @@ RAY_TRACING_ENTRY_CLOSEST_HIT(PathTracingMaterialCHS,
 	if (Payload.ShadingModelID == SHADINGMODELID_SUBSURFACE_PROFILE)
 	{
 		Payload.CustomData0 = GetMaterialSubsurfaceData(PixelMaterialInputs);
-		uint ProfileId = Payload.GetProfileId();
+		uint ProfileId = Payload.GetSubsurfaceProfileId();
 
 		// Decode dual roughness info
 		float Lobe0Roughness = 0;
@@ -268,8 +269,8 @@ RAY_TRACING_ENTRY_CLOSEST_HIT(PathTracingMaterialCHS,
 
 		GetSubsurfaceProfileDualSpecular(ProfileId, Payload.Roughness, Payload.Opacity, Lobe0Roughness, Lobe1Roughness, LobeMix);
 
-		Lobe0Roughness = AdjustMaterialRoughness(Lobe0Roughness, PackedPayload.GetPathRoughness();
-		Lobe1Roughness = AdjustMaterialRoughness(Lobe1Roughness, PackedPayload.GetPathRoughness();
+		Lobe0Roughness = AdjustMaterialRoughness(Lobe0Roughness, PackedPayload.GetPathRoughness());
+		Lobe1Roughness = AdjustMaterialRoughness(Lobe1Roughness, PackedPayload.GetPathRoughness());
 		Payload.SetDualRoughnessSpecular(Lobe0Roughness, Lobe1Roughness, LobeMix);
 	}
 	else
