* Add better filtering to rect atlas texture when uploaded
* Add width/height texture ratio clamp to get better filtering

diff --git a/Engine/Shaders/Private/RectLightAtlas.usf b/Engine/Shaders/Private/RectLightAtlas.usf
--- a/Engine/Shaders/Private/RectLightAtlas.usf
+++ b/Engine/Shaders/Private/RectLightAtlas.usf
@@ -67,14 +67,14 @@ void MainPS(
 	float4 Color = 0;
 	switch (InId)
 	{
-	case 0: Color = InTexture0.Load(uint3(InCoord, InTextureMIPBias)); break;
-	case 1: Color = InTexture1.Load(uint3(InCoord, InTextureMIPBias)); break;
-	case 2: Color = InTexture2.Load(uint3(InCoord, InTextureMIPBias)); break;
-	case 3: Color = InTexture3.Load(uint3(InCoord, InTextureMIPBias)); break;
-	case 4: Color = InTexture4.Load(uint3(InCoord, InTextureMIPBias)); break;
-	case 5: Color = InTexture5.Load(uint3(InCoord, InTextureMIPBias)); break;
-	case 6: Color = InTexture6.Load(uint3(InCoord, InTextureMIPBias)); break;
-	case 7: Color = InTexture7.Load(uint3(InCoord, InTextureMIPBias)); break;
+	case 0: Color = InTexture0.SampleLevel(InSampler, InUV, InTextureMIPBias); break;
+	case 1: Color = InTexture1.SampleLevel(InSampler, InUV, InTextureMIPBias); break;
+	case 2: Color = InTexture2.SampleLevel(InSampler, InUV, InTextureMIPBias); break;
+	case 3: Color = InTexture3.SampleLevel(InSampler, InUV, InTextureMIPBias); break;
+	case 4: Color = InTexture4.SampleLevel(InSampler, InUV, InTextureMIPBias); break;
+	case 5: Color = InTexture5.SampleLevel(InSampler, InUV, InTextureMIPBias); break;
+	case 6: Color = InTexture6.SampleLevel(InSampler, InUV, InTextureMIPBias); break;
+	case 7: Color = InTexture7.SampleLevel(InSampler, InUV, InTextureMIPBias); break;
 	}
 
 	// Ensure there is no NaN value
@@ -115,6 +115,7 @@ void MainPS(
 #if SHADER_FILTER_TEXTURE
 
 float KernelSize;
+uint FilterQuality;
 int2 SrcAtlasResolution;
 Buffer<uint4> SrcSlotBuffer;
 Texture2D<float4> SourceAtlasTexture;
@@ -127,55 +128,59 @@ void MainPS(
 	in uint InId : RECT_ID,
 	out float4 OutTexture : SV_Target0)
 {
-#if 1
-	// 2x2 box filtering
-	const float2 SrcCoord = uint2(InCoord) * 2.f;
-	const uint4 SrcRect = SrcSlotBuffer[InId];
-	const uint2 SrcRectOffset = SrcRect.xy;
-	const uint2 SrcRectResolution = SrcRect.zw;
-
-	const float4 Color0 = SourceAtlasTexture.Load(uint3(SrcRectOffset + SrcCoord + float2(0,0), 0));
-	const float4 Color1 = SourceAtlasTexture.Load(uint3(SrcRectOffset + SrcCoord + float2(1,0), 0));
-	const float4 Color2 = SourceAtlasTexture.Load(uint3(SrcRectOffset + SrcCoord + float2(0,1), 0));
-	const float4 Color3 = SourceAtlasTexture.Load(uint3(SrcRectOffset + SrcCoord + float2(1,1), 0));
-	OutTexture = (Color0 + Color1 + Color2 + Color3) * 0.25f;
-#else
-	// 5x5 Gaussian filtering (naive, non-separable, with bilinear filtering)
-	// 
-	// Possible improvement, but probably not worth it since the filtering is done only once, when the texture is uploaded: 
-	// * Separable filter (but requires compute with border, or RT ping-pong)
-	// * Recursive sparse filter with MIP prefiltering (Deriche & co)
-	const float2 SrcCoord = InCoord * 2.f; // -> InCoord is xxx.5f, so SrcCoord is at the center of the 2x2 block
-	const uint4 SrcRect = SrcSlotBuffer[InId];
-	const uint2 SrcRectOffset = SrcRect.xy;
-	const uint2 SrcRectResolution = SrcRect.zw;
-
-	const float2 MinCoord = SrcRectOffset;
-	const float2 MaxCoord = SrcRectOffset + SrcRectResolution;
-	const float2 InvAtlasResolution = 1.0f / float2(SrcAtlasResolution);
-
-	float3 AccColor = 0;
-	float AccW = 0;
-	const float Variance = 2.0f; // std ~1.4. weight, at radius=3 ~0.01;
-	const int KernelExtent = 2;
-	for (int y = -KernelExtent; y <= KernelExtent; ++y)
-	for (int x = -KernelExtent; x <= KernelExtent; ++x
+
+	if (FilterQuality)
+	{
+		// 2x2 box filtering
+		const float2 SrcCoord = uint2(InCoord) * 2.f;
+		const uint4 SrcRect = SrcSlotBuffer[InId];
+		const uint2 SrcRectOffset = SrcRect.xy;
+		const uint2 SrcRectResolution = SrcRect.zw;
+
+		const float4 Color0 = SourceAtlasTexture.Load(uint3(SrcRectOffset + SrcCoord + float2(0,0), 0));
+		const float4 Color1 = SourceAtlasTexture.Load(uint3(SrcRectOffset + SrcCoord + float2(1,0), 0));
+		const float4 Color2 = SourceAtlasTexture.Load(uint3(SrcRectOffset + SrcCoord + float2(0,1), 0));
+		const float4 Color3 = SourceAtlasTexture.Load(uint3(SrcRectOffset + SrcCoord + float2(1,1), 0));
+		OutTexture = (Color0 + Color1 + Color2 + Color3) * 0.25f;
+	}
+	else
 	{
-		const float2 Offset = float2(x, y);
-		const float2 Coord = SrcRectOffset + SrcCoord + Offset;// +float2(0.5f, 0.5f);
-		if (all(Coord >= MinCoord) && all(Coord < MaxCoord))
+		// 5x5 Gaussian filtering (naive, non-separable, with bilinear filtering)
+		// 
+		// Possible improvement, but probably not worth it since the filtering is done only once, when the texture is uploaded: 
+		// * Separable filter (but requires compute with border, or RT ping-pong)
+		// * Recursive sparse filter with MIP prefiltering (Deriche & co)
+		const float2 SrcCoord = InCoord * 2.f; // -> InCoord is xxx.5f, so SrcCoord is at the center of the 2x2 block
+		const uint4 SrcRect = SrcSlotBuffer[InId];
+		const uint2 SrcRectOffset = SrcRect.xy;
+		const uint2 SrcRectResolution = SrcRect.zw;
+
+		const float2 MinCoord = SrcRectOffset;
+		const float2 MaxCoord = SrcRectOffset + SrcRectResolution;
+		const float2 InvAtlasResolution = 1.0f / float2(SrcAtlasResolution);
+
+		float3 AccColor = 0;
+		float AccW = 0;
+		const float Variance = 3.0f; // std ~1.4. weight, at radius=3 ~0.01;
+		const int KernelExtent = 2;
+		for (int y = -KernelExtent; y <= KernelExtent; ++y)
+		for (int x = -KernelExtent; x <= KernelExtent; ++x)
 		{
-			const float2 UV = Coord * InvAtlasResolution;
-			const float  D2 = dot(Offset, Offset);
-			const float  W = exp(-D2 / Variance);
-			const float3 Color = SourceAtlasTexture.SampleLevel(SourceAtlasSampler, UV, 0).xyz; // MIP=0 as it is an SRV narrowed around SrcMIP
+			const float2 Offset = float2(x, y);
+			const float2 Coord = SrcRectOffset + SrcCoord + Offset;// +float2(0.5f, 0.5f);
+			if (all(Coord >= MinCoord) && all(Coord < MaxCoord))
+			{
+				const float2 UV = Coord * InvAtlasResolution;
+				const float  D2 = dot(Offset, Offset);
+				const float  W = exp(-D2 / Variance);
+				const float3 Color = SourceAtlasTexture.SampleLevel(SourceAtlasSampler, UV, 0).xyz; // MIP=0 as it is an SRV narrowed around SrcMIP
 
-			AccColor += W * Color;
-			AccW += W;
+				AccColor += W * Color;
+				AccW += W;
+			}
 		}
+		OutTexture = float4(AccColor / AccW, 1.0f);
 	}
-	OutTexture = float4(AccColor / AccW, 1.0f);
-#endif 
 }
 
 #endif // SHADER_FILTER_TEXTURE
diff --git a/Engine/Source/Runtime/Renderer/Private/RectLightTextureManager.cpp b/Engine/Source/Runtime/Renderer/Private/RectLightTextureManager.cpp
--- a/Engine/Source/Runtime/Renderer/Private/RectLightTextureManager.cpp
+++ b/Engine/Source/Runtime/Renderer/Private/RectLightTextureManager.cpp
@@ -55,6 +55,18 @@ static TAutoConsoleVariable<int32> CVarRectLighTranslucent(
 	TEXT("Enable rect light support for translucent surfaces. When enabled, it will add an extrat sampler to the pixel shader (limited to 16 on dx11 based system)"),
 	ECVF_ReadOnly | ECVF_RenderThreadSafe);
 
+static TAutoConsoleVariable<int32> CVarRectLighFilterQuality(
+	TEXT("r.RectLightAtlas.FilterQuality"),
+	1,
+	TEXT("Define the filtering quality used for filtering texture (0:Box filter, 1:Gaussian filter)."),
+	ECVF_RenderThreadSafe);
+
+static TAutoConsoleVariable<int32> CVarRectLighMaxTextureRatio(
+	TEXT("r.RectLightAtlas.MaxTextureRatio"),
+	2,
+	TEXT("Define the max Width/Height or Height/Width ratio that a texture can have."),
+	ECVF_RenderThreadSafe);
+
 namespace RectLightAtlas
 {
 
@@ -79,6 +91,30 @@ namespace RectLightAtlas
 static const uint32 InvalidSlotIndex = ~0u;
 static const FIntPoint InvalidOrigin = FIntPoint(-1, -1);
 
+FIntPoint GetSourceResolution(const FRHITexture* Tex)
+{
+	FIntPoint Out(1, 1);
+	if (Tex)
+	{
+		const FIntVector SourceOriginalResolution = Tex->GetSizeXYZ();
+		const float RatioYX = float(SourceOriginalResolution.Y) / float(SourceOriginalResolution.X);
+		const float RatioXY = float(SourceOriginalResolution.X) / float(SourceOriginalResolution.Y);
+
+		Out = FIntPoint(SourceOriginalResolution.X, SourceOriginalResolution.Y);
+		// Max Ratio of 2
+		const float RatioThreshold = FMath::Clamp(CVarRectLighMaxTextureRatio.GetValueOnAnyThread(), 1, 16);
+		if (RatioYX > RatioThreshold)
+		{
+			Out.X *= RatioYX / RatioThreshold;
+		}
+		else if (RatioXY > RatioThreshold)
+		{
+			Out.Y *= RatioXY / RatioThreshold;
+		}
+	}
+	return Out;
+}
+
 struct FAtlasRect
 {
 	FIntPoint Origin = InvalidOrigin;
@@ -130,6 +166,11 @@ struct FAtlasSlot
 	bool bForceRefresh = false;
 	bool IsValid() const { return SourceTexture != nullptr; }
 	FRHITexture* GetTextureRHI() const { return SourceTexture->GetReferencedTexture(); }
+
+	FIntPoint GetSourceResolution() const 
+	{ 
+		return RectLightAtlas::GetSourceResolution(SourceTexture ? SourceTexture->GetReferencedTexture() : nullptr);
+	}
 };
 
 // Store info for copying one atlas slot to another one, when a new layout is created
@@ -601,7 +642,7 @@ class FRectAtlasFilterTexturePS : public FGlobalShader
 
 	BEGIN_SHADER_PARAMETER_STRUCT(FParameters, )
 		SHADER_PARAMETER(FIntPoint, SrcAtlasResolution)
-		SHADER_PARAMETER(float, KernelSize)
+		SHADER_PARAMETER(uint32, FilterQuality)
 		SHADER_PARAMETER_RDG_BUFFER_SRV(Buffer, DstSlotBuffer)
 		SHADER_PARAMETER_RDG_BUFFER_SRV(Buffer, SrcSlotBuffer)
 		SHADER_PARAMETER_SAMPLER(SamplerState, SourceAtlasSampler)
@@ -665,7 +706,7 @@ static void FilterSlotsPass(
 			FRDGBufferRef DstSlotBuffer = CreateSlotBuffer(GraphBuilder, DstMIPSlots, TEXT("RectLight.DstMIPSlotBuffer"));
 
 			FRectAtlasFilterTexturePS::FParameters* Parameters = GraphBuilder.AllocParameters<FRectAtlasFilterTexturePS::FParameters>();
-			Parameters->KernelSize			= 5;
+			Parameters->FilterQuality		= FMath::Clamp(CVarRectLighFilterQuality.GetValueOnRenderThread(), 0, 1);
 			Parameters->SrcAtlasResolution	= ToMIP(InAtlas->Desc.Extent, SrcMip);
 			Parameters->DstSlotBuffer		= GraphBuilder.CreateSRV(DstSlotBuffer, PF_R16G16B16A16_UINT);
 			Parameters->SrcSlotBuffer		= GraphBuilder.CreateSRV(SrcSlotBuffer, PF_R16G16B16A16_UINT);
@@ -1228,8 +1269,8 @@ static void PackAtlas(
 
 			// Check if the texture resolution has changed (due to streaming)
 			// If the texture resolution has change, reset the slot to be handled as a new slot
-			const FIntVector TextureResolution = ValidSlot.GetTextureRHI()->GetSizeXYZ();
-			const bool bHasTextureResolutionChanged = ValidSlot.Rect.Resolution != FIntPoint(TextureResolution.X, TextureResolution.Y);
+			const FIntPoint TextureResolution = ValidSlot.GetSourceResolution();
+			const bool bHasTextureResolutionChanged = ValidSlot.Rect.Resolution != TextureResolution;
 			if (bHasTextureResolutionChanged || ValidSlot.bForceRefresh)
 			{
 				ValidSlot.Rect.Origin = InvalidOrigin;
@@ -1335,9 +1376,9 @@ static void PackAtlas(
 					CurrentSourceTextureMIPBias++;
 					for (FAtlasSlot& Slot : ValidSlots)
 					{
-						const FRHITexture* Tex = Slot.GetTextureRHI();
+						const FIntPoint SourceResolution = Slot.GetSourceResolution();
 						Slot.Rect.Origin = InvalidOrigin;
-						Slot.Rect.Resolution = ToMIP(FIntPoint(Tex->GetSizeXYZ().X, Tex->GetSizeXYZ().Y), CurrentSourceTextureMIPBias);
+						Slot.Rect.Resolution = ToMIP(SourceResolution, CurrentSourceTextureMIPBias);
 					}
 				}
 
@@ -1395,7 +1436,7 @@ uint32 AddTexture(UTexture* In)
 			Slot->SourceTexture = In->TextureReference.TextureReferenceRHI;
 			Slot->Id = SlotIndex;
 			Slot->Rect.Origin = InvalidOrigin;
-			Slot->Rect.Resolution = FIntPoint(Tex->GetSizeXYZ().X, Tex->GetSizeXYZ().Y);
+			Slot->Rect.Resolution = GetSourceResolution(Tex);
 			Slot->RefCount = 1;
 
 			GRectLightTextureManager.bHasPendingAdds = true;
@@ -1496,11 +1537,7 @@ void UpdateAtlasTexture(FRDGBuilder& GraphBuilder, const ERHIFeatureLevel::Type
 		{
 			if (Slot.IsValid() && Slot.GetTextureRHI())
 			{
-				FRHITexture* TextureRHI = Slot.GetTextureRHI();
-				check(TextureRHI != nullptr);
-
-				const FIntVector SourceResolution = TextureRHI->GetSizeXYZ();
-				const FIntPoint  SourceResolutionMIPed = ToMIP(FIntPoint(SourceResolution.X, SourceResolution.Y), GRectLightTextureManager.AtlasLayout.SourceTextureMIPBias);
+				const FIntPoint SourceResolutionMIPed = ToMIP(Slot.GetSourceResolution(), GRectLightTextureManager.AtlasLayout.SourceTextureMIPBias);
 				if (SourceResolutionMIPed.X > Slot.Rect.Resolution.X || SourceResolutionMIPed.Y > Slot.Rect.Resolution.Y)
 				{
 					// Invalid the slot
