* Fix raster coverage
* Improve hair sorting

diff --git a/Engine/Shaders/Private/HairStrands/HairStrandsForwardRaster.usf b/Engine/Shaders/Private/HairStrands/HairStrandsForwardRaster.usf
--- a/Engine/Shaders/Private/HairStrands/HairStrandsForwardRaster.usf
+++ b/Engine/Shaders/Private/HairStrands/HairStrandsForwardRaster.usf
@@ -205,6 +205,7 @@ float	RcpNumRasterizers;
 
 int2	OutputResolution;
 float2	OutputResolutionf;
+float	RadiusAtDepth1;
 
 ///////////////////////////////////////////////////////////////////////////
 // Control points
@@ -259,7 +260,7 @@ void CalcHomogenousPosAndRad(in uint PrimId, out float4 HP, out float Rad, out u
 
 	const float3 WP = CP.Position; // This is actually WorldPosition
 	HP = mul(float4(WP, 1.0f), LWCHackToFloat(PrimaryView.WorldToClip));
-	Rad = CP.WorldRadius * 2000.0; // OutputResolutionf.x; //TODO: figure this out correctly?
+	Rad = CP.WorldRadius;
 	Type = CP.Type;
 }
 
@@ -868,6 +869,7 @@ void CompactionCS(uint DispatchThreadID : SV_DispatchThreadID, uint GroupThreadI
 #define RASTER_DEPTH_BUCKET 32
 #define RASTER_SEGMENT_COUNT 32
 #define SEGMENT_COUNT_PER_GROUP 1024
+#define INVALID_PRIM_ID 0xFFFFFFFF
 
 // For editing convenience
 #if !SHADER_RASTERCOMPUTE_DEBUG && !SHADER_RASTERCOMPUTE_COMPACTION && !SHADER_RASTERCOMPUTE_BINNING && !SHADER_RASTERCOMPUTE_RASTER_MULTI_SAMPLE && !SHADER_RASTERCOMPUTE_DEPTH_GRID
@@ -895,6 +897,7 @@ groupshared uint2  s_Segments[1024];
 groupshared uint2  s_Segments_Sorted[1024];
 groupshared uint   s_Segments_Min;
 groupshared uint   s_Segments_Max;
+groupshared uint   s_Segments_ValidCount;
 groupshared uint   s_SegmentsCount[RASTER_DEPTH_BUCKET];
 groupshared uint   s_SegmentsAlloc[RASTER_DEPTH_BUCKET];
 
@@ -1014,10 +1017,26 @@ uint GetDepthBinIndex(float InDepth, float InvDepthExtent)
 	// Inverse-Z
 	const float MinDepth = UnpackDepth(s_Segments_Min);
 	const float MaxDepth = UnpackDepth(s_Segments_Max);
-	const uint DepthIt = min(saturate((InDepth - MinDepth) * InvDepthExtent) * RASTER_DEPTH_BUCKET, RASTER_DEPTH_BUCKET - 1);
+	const uint DepthIt = clamp(saturate((InDepth - MinDepth) * InvDepthExtent) * RASTER_DEPTH_BUCKET, 0, RASTER_DEPTH_BUCKET - 1);
 	return (RASTER_DEPTH_BUCKET-1) - DepthIt;
 }
 
+#if PERMUTATION_DEBUG
+void ShiftX(inout FShaderPrintContext Out, uint InPixelCountX)
+{	
+	const float fShift = float(InPixelCountX) / float(ShaderPrintData.Resolution.x);
+	Out.StartPos.x += fShift;
+	Out.Pos.x += fShift;
+}
+
+void ShiftY(inout FShaderPrintContext Out, uint InPixelCountY)
+{
+	const float fShift = float(InPixelCountY) / float(ShaderPrintData.Resolution.y);
+	Out.StartPos.y += fShift;
+	Out.Pos.y += fShift;
+}
+#endif
+
 [numthreads(RASTER_TILE_SIZE, RASTER_TILE_SIZE, 1)]
 void RasterCS(
 	uint  GroupThread1D    : SV_GroupIndex,		/* 64  */
@@ -1042,6 +1061,7 @@ void RasterCS(
 			{
 				s_Segments_Min = PackDepth(1e8);
 				s_Segments_Max = 0;
+				s_Segments_ValidCount = 0;
 				s_bDataOrder = 0;
 				s_OpaqueMask = 0;
 			}
@@ -1095,15 +1115,17 @@ void RasterCS(
 		const uint2 RasterTileMax	= RasterTileMin   + RASTER_TILE_SIZE;
 
 		const uint LoopCount64  = DivideAndRoundUp(PrimCount, RASTER_THREAD_COUNT);
-		const uint LoopCount32  = DivideAndRoundUp(PrimCount, RASTER_SEGMENT_COUNT);
 
 		// For debug only
 		#if PERMUTATION_DEBUG
 		const bool bDebugEnableAll	= all(RasterTileCoord == uint2(ShaderPrintData.CursorCoord / float(RASTER_TILE_SIZE)));
-		const bool bDebugEnable		= all(RasterTileCoord == uint2(ShaderPrintData.CursorCoord / float(RASTER_TILE_SIZE))) && GroupThread1D == 0;
+		const bool bDebugEnable		= all(RasterTileCoord == uint2(ShaderPrintData.CursorCoord / float(RASTER_TILE_SIZE))) && GroupThread1D == View.GeneralPurposeTweak;
 		FShaderPrintContext Ctx		= InitShaderPrintContext(bDebugEnable, uint2(350, 50));
 		const FFontColor FontLegend = FontWhite;
 		const FFontColor FontValue  = FontOrange;
+
+		//const bool bDebugEnable2 = all(RasterTileCoord == uint2(ShaderPrintData.CursorCoord / float(RASTER_TILE_SIZE))) && GroupThread1D == 1;
+		FShaderPrintContext CtxAll = InitShaderPrintContext(bDebugEnableAll, uint2(750, 50));
 		#endif
 
 		#if PERMUTATION_DEBUG
@@ -1133,7 +1155,7 @@ void RasterCS(
 				const uint Prim = GroupThread1D + LoopIndex * RASTER_THREAD_COUNT;
 				if (Prim < PrimCount)
 				{
-					const uint PrimID = VisTilePrims[PrimOffset + Prim];
+					uint PrimID = VisTilePrims[PrimOffset + Prim];
 
 					uint TypeDummy = 0;
 					float4 SP0 = 0;
@@ -1152,18 +1174,23 @@ void RasterCS(
 
 						// Clip against tile
 						bool2 bClipped = false;
+						bool bIsValidSegment = false;
 						float2 T = 0;
-						ClipRaySegment(BinTileMin - 0.5f, BinTileMax + 0.5f, SP0, SP1, Rad0, Rad1, Alpha0, Alpha1, bClipped, T);
+						bIsValidSegment = ClipRaySegment(RasterTileMin - 0.5f, RasterTileMax + 0.5f, SP0, SP1, Rad0, Rad1, Alpha0, Alpha1, bClipped, T);
+
+						PrimID = bIsValidSegment ? PrimID : INVALID_PRIM_ID;
 					}
 
 					const uint uDepth = PackDepth(max(SP0.z, SP1.z)); // Inverse Z
 					s_Segments[Prim] = uint2(uDepth, PrimID);
+					if (PrimID != INVALID_PRIM_ID)
+					{
+						InterlockedMin(s_Segments_Min, uDepth);
+						InterlockedMax(s_Segments_Max, uDepth);
+					}
 
-					InterlockedMin(s_Segments_Min, uDepth);
-					InterlockedMax(s_Segments_Max, uDepth);
-
-					#if 0 && PERMUTATION_DEBUG
-					if (bDebugEnableAll)
+					#if 1 && PERMUTATION_DEBUG
+					if (bDebugEnableAll && PrimID != INVALID_PRIM_ID)
 					{
 						const float3 P0 = UnpackHairControlPoint(PrimID).Position;
 						const float3 P1 = UnpackHairControlPoint(PrimID+1).Position;
@@ -1172,8 +1199,8 @@ void RasterCS(
 						const float4 Color1 = float4(LoadSampleColor(PrimID+1), 1);
 
 						const float4 LineColor = float4(ColorMapTurbo(GroupThread1D / float(RASTER_THREAD_COUNT)), 1);
-						//AddLineWS(P0, P1, LineColor);
-						AddLineWS(P0, P1, Color0, Color1);
+						AddLineWS(P0, P1, LineColor);
+						//AddLineWS(P0, P1, Color0, Color1);
 					}
 					#endif
 				}
@@ -1182,41 +1209,6 @@ void RasterCS(
 		GroupMemoryBarrierWithGroupSync();
 		const float InvDepthExtent = 1.f / max(0.0001f, UnpackDepth(s_Segments_Max)- UnpackDepth(s_Segments_Min));
 
-		// DEBUG
-		#if PERMUTATION_DEBUG
-		if (bDebugEnable)//  && GroupThreadID == 63)
-		{
-			const uint2 OutCoord = GroupThread2D + RasterTileMin;
-
-			//Print(Ctx, TEXT("Out Coord  :"), FontLegend); Print(Ctx, OutCoord, FontValue); Newline(Ctx);
-			//Print(Ctx, TEXT("Min Raster :"), FontLegend); Print(Ctx, RasterTileMin, FontValue); Newline(Ctx);
-			//Print(Ctx, TEXT("Thread     :"), FontLegend); Print(Ctx, Thread_Coord, FontValue); Newline(Ctx);
-			//Print(Ctx, TEXT("Max Raster :"), FontLegend); Print(Ctx, RasterTileMax, FontValue); Newline(Ctx);
-			//Newline(Ctx);
-
-			//Print(Ctx, TEXT("Cursor     :"), FontLegend); Print(Ctx, uint2(ShaderPrintData.CursorCoord), FontValue); Newline(Ctx);
-			//Newline(Ctx);
-
-			Print(Ctx, TEXT("Work ID      :"), FontLegend); Print(Ctx, s_WorkID, FontValue); Newline(Ctx);
-			Print(Ctx, TEXT("Prim Count   :"), FontLegend); Print(Ctx, PrimCount, FontValue); Newline(Ctx);
-			Print(Ctx, TEXT("Loop 32 Count:"), FontLegend); Print(Ctx, LoopCount32, FontValue); Newline(Ctx);
-			Print(Ctx, TEXT("Loop 64 Count:"), FontLegend); Print(Ctx, LoopCount64, FontValue); Newline(Ctx);
-			Print(Ctx, TEXT("Min Depth    :"), FontLegend); Print(Ctx, UnpackDepth(s_Segments_Min), FontValue); Newline(Ctx);
-			Print(Ctx, TEXT("Max Depth    :"), FontLegend); Print(Ctx, UnpackDepth(s_Segments_Max), FontValue); Newline(Ctx);
-			Newline(Ctx);
-
-			//Print(Ctx, TEXT("Bin ID     :"), FontLegend); Print(Ctx, BinTileIndex, FontValue); Newline(Ctx);
-			//Print(Ctx, TEXT("Raster ID  :"), FontLegend); Print(Ctx, QuadrantIndex, FontValue); Newline(Ctx);
-			//Newline(Ctx);
-
-			//Print(Ctx, TEXT("Raster Tile:"), FontLegend); Print(Ctx, RasterTileCoord, FontValue); Newline(Ctx);
-			//Print(Ctx, TEXT("Bin Tile   :"), FontLegend); Print(Ctx, BinTileCoord, FontValue); Newline(Ctx);
-			//Print(Ctx, TEXT("Out Coord  :"), FontLegend); Print(Ctx, OutCoord, FontValue); Newline(Ctx);
-			//Newline(Ctx);
-
-		}
-		#endif
-
 		// 2. Compute the count of depth bucket
 		{
 			LOOP
@@ -1225,16 +1217,20 @@ void RasterCS(
 				const uint Prim = GroupThread1D + LoopIndex * RASTER_THREAD_COUNT;
 				if (Prim < PrimCount)
 				{
-					const float Depth	= UnpackDepth(s_Segments[Prim].x);
-					const uint DepthIt  = GetDepthBinIndex(Depth, InvDepthExtent);
-
-					InterlockedAdd(s_SegmentsCount[DepthIt], 1);
+					const bool bIsValid = s_Segments[Prim].y != INVALID_PRIM_ID;
+					if (bIsValid)
+					{
+						const float Depth	= UnpackDepth(s_Segments[Prim].x);
+						const uint DepthIt  = GetDepthBinIndex(Depth, InvDepthExtent);
+						InterlockedAdd(s_SegmentsCount[DepthIt], 1);
+					}
 				}
 			}
 		}
 		GroupMemoryBarrierWithGroupSync();
 
 		// Replace this with parallel version
+		s_Segments_ValidCount = 0;
 		if (GroupThread1D == 0)
 		{
 			uint Acc = 0;
@@ -1244,7 +1240,7 @@ void RasterCS(
 				s_SegmentsCount[It] = Acc;
 				Acc += Next;
 
-				#if 0 && PERMUTATION_DEBUG
+				#if 1 && PERMUTATION_DEBUG
 				if (bDebugEnable)
 				{
 					if (Next > 0)	{ Print(Ctx, TEXT("x"), FontValue); }
@@ -1253,6 +1249,7 @@ void RasterCS(
 				}
 				#endif
 			}
+			s_Segments_ValidCount = Acc;
 		}
 		GroupMemoryBarrierWithGroupSync();
 
@@ -1265,18 +1262,60 @@ void RasterCS(
 				if (Prim < PrimCount)
 				{
 					const uint2 Segment = s_Segments[Prim];
-					const float Depth = UnpackDepth(Segment.x);
-					const uint DepthIt = GetDepthBinIndex(Depth, InvDepthExtent);
+					if (Segment.y != INVALID_PRIM_ID)
+					{
+						const float Depth = UnpackDepth(Segment.x);
+						const uint DepthIt = GetDepthBinIndex(Depth, InvDepthExtent);
 
-					uint AllocOffset = 0;
-					InterlockedAdd(s_SegmentsAlloc[DepthIt], 1, AllocOffset);
-					const uint NewIndex = AllocOffset + s_SegmentsCount[DepthIt];
-					s_Segments_Sorted[NewIndex] = Segment;
+						uint AllocOffset = 0;
+						InterlockedAdd(s_SegmentsAlloc[DepthIt], 1, AllocOffset);
+						const uint NewIndex = AllocOffset + s_SegmentsCount[DepthIt];
+						s_Segments_Sorted[NewIndex] = Segment;
+					}
 				}
 			}
 		}
 		GroupMemoryBarrierWithGroupSync();
 
+		const uint LoopCount32 = DivideAndRoundUp(s_Segments_ValidCount, RASTER_SEGMENT_COUNT);
+		const uint LoopCount32_All = DivideAndRoundUp(PrimCount, RASTER_SEGMENT_COUNT);
+
+		// DEBUG
+		#if PERMUTATION_DEBUG
+		if (bDebugEnable)
+		{
+			const uint2 OutCoord = GroupThread2D + RasterTileMin;
+
+			//Print(Ctx, TEXT("Out Coord  :"), FontLegend); Print(Ctx, OutCoord, FontValue); Newline(Ctx);
+			//Print(Ctx, TEXT("Min Raster :"), FontLegend); Print(Ctx, RasterTileMin, FontValue); Newline(Ctx);
+			//Print(Ctx, TEXT("Thread     :"), FontLegend); Print(Ctx, Thread_Coord, FontValue); Newline(Ctx);
+			//Print(Ctx, TEXT("Max Raster :"), FontLegend); Print(Ctx, RasterTileMax, FontValue); Newline(Ctx);
+			//Newline(Ctx);
+
+			//Print(Ctx, TEXT("Cursor     :"), FontLegend); Print(Ctx, uint2(ShaderPrintData.CursorCoord), FontValue); Newline(Ctx);
+			//Newline(Ctx);
+
+			Print(Ctx, TEXT("Work ID      :"), FontLegend); Print(Ctx, s_WorkID, FontValue); Newline(Ctx);
+			Print(Ctx, TEXT("Prim Count   :"), FontLegend); Print(Ctx, s_Segments_ValidCount, FontValue);	Print(Ctx, TEXT(" / "), FontLegend); Print(Ctx, PrimCount, FontValue); Newline(Ctx);
+			Print(Ctx, TEXT("Loop 32 Count:"), FontLegend); Print(Ctx, LoopCount32, FontValue);				Print(Ctx, TEXT(" / "), FontLegend); Print(Ctx, LoopCount32_All, FontValue); Newline(Ctx);
+			Print(Ctx, TEXT("Loop 64 Count:"), FontLegend); Print(Ctx, LoopCount64, FontValue); Newline(Ctx);
+			Print(Ctx, TEXT("Min Depth    :"), FontLegend); Print(Ctx, UnpackDepth(s_Segments_Min), FontValue); Newline(Ctx);
+			Print(Ctx, TEXT("Max Depth    :"), FontLegend); Print(Ctx, UnpackDepth(s_Segments_Max), FontValue); Newline(Ctx);
+			Print(Ctx, TEXT("Rad.at Depth1:"), FontLegend); Print(Ctx, RadiusAtDepth1, FontValue); Newline(Ctx);
+			Newline(Ctx);
+
+			//Print(Ctx, TEXT("Bin ID     :"), FontLegend); Print(Ctx, BinTileIndex, FontValue); Newline(Ctx);
+			//Print(Ctx, TEXT("Raster ID  :"), FontLegend); Print(Ctx, QuadrantIndex, FontValue); Newline(Ctx);
+			//Newline(Ctx);
+
+			//Print(Ctx, TEXT("Raster Tile:"), FontLegend); Print(Ctx, RasterTileCoord, FontValue); Newline(Ctx);
+			//Print(Ctx, TEXT("Bin Tile   :"), FontLegend); Print(Ctx, BinTileCoord, FontValue); Newline(Ctx);
+			//Print(Ctx, TEXT("Out Coord  :"), FontLegend); Print(Ctx, OutCoord, FontValue); Newline(Ctx);
+			//Newline(Ctx);
+
+		}
+		#endif
+
 		// 4. Raster segments
 		{
 			// These are global Color/Coverage for the final pixel handled by this thread
@@ -1336,6 +1375,53 @@ void RasterCS(
 						bIsSegmentValid = ClipRaySegment(RasterTileMin - 0.5f, RasterTileMax + 0.5f, SP0, SP1, Rad0, Rad1, Alpha0, Alpha1, bClipped, T);
 					}
 
+					// DEBUG (Write coord)
+					#if 0 && PERMUTATION_DEBUG
+					if (bDebugEnableAll)
+					{
+						ShiftY(CtxAll, 500 + GroupThread1D * 150);
+						
+						Print(CtxAll, TEXT("Thrd 1D:"), FontYellow);
+						Print(CtxAll, GroupThread1D, FontWhite);
+						Newline(CtxAll);
+
+						Print(CtxAll, TEXT("ID     :"), FontYellow);
+						Print(CtxAll, PrimID, FontWhite);
+						Newline(CtxAll);
+
+						Print(CtxAll, TEXT("Valid  :"), FontYellow);
+						PrintBool(CtxAll, bIsSegmentValid);
+						Newline(CtxAll);
+
+						Print(CtxAll, TEXT("Clip X :"), FontYellow);
+						PrintBool(CtxAll, bClipped.x);
+						Newline(CtxAll);
+
+						Print(CtxAll, TEXT("Clip Y :"), FontYellow);
+						PrintBool(CtxAll, bClipped.y);
+						Newline(CtxAll);
+
+						Print(CtxAll, TEXT("X0      :"), FontYellow);
+						Print(CtxAll, SP0, FontWhite);
+						Newline(CtxAll);
+
+						Print(CtxAll, TEXT("X1      :"), FontYellow);
+						Print(CtxAll, SP1, FontWhite);
+						Newline(CtxAll);
+
+						//const float3 P0 = UnpackHairControlPoint(PrimID).Position;
+						//const float3 P1 = UnpackHairControlPoint(PrimID + 1).Position;
+
+						//const float4 Color0 = float4(LoadSampleColor(PrimID), 1);
+						//const float4 Color1 = float4(LoadSampleColor(PrimID + 1), 1);
+
+						//const float4 LineColor = float4(ColorMapTurbo(GroupThread1D / float(RASTER_THREAD_COUNT)), 1);
+						//AddLineWS(P0, P1, LineColor);
+						if (GroupThread1D == View.GeneralPurposeTweak)
+							AddLineSS(SP0, SP1, ColorPurple);
+					}
+					#endif
+
 					// 4.1.2 Rasterize segment (1 thread = 1 segments)
 					if (bIsSegmentValid)
 					{
@@ -1353,28 +1439,39 @@ void RasterCS(
 
 						// DEBUG (Write coord)
 						#if 0 && PERMUTATION_DEBUG
-						if (GroupThread1D == 0 && bDebugEnable)
-						{
-							Print(Ctx, TEXT("NumSteps:"), FontYellow);
-							Print(Ctx, NumSteps, FontWhite); 
-							Newline(Ctx);
-
-							Print(Ctx, TEXT("X0      :"), FontYellow);
-							Print(Ctx, X0, FontWhite);
-							Newline(Ctx);
-
-							Print(Ctx, TEXT("X1      :"), FontYellow);
-							Print(Ctx, X1, FontWhite);
-							Newline(Ctx);
-
-							Print(Ctx, TEXT("T       :"), FontYellow);
-							Print(Ctx, T, FontWhite);
-							Newline(Ctx);
-
-							Print(Ctx, TEXT("bValid  :"), FontYellow);
-							PrintBool(Ctx, bIsSegmentValid);
-							Newline(Ctx);
-							
+						//if (GroupThread1D == 0 && bDebugEnable)
+						if (bDebugEnableAll)
+						{							
+							ShiftY(CtxAll, GroupThread1D * 150);
+							//ShiftY(CtxAll, GroupThread1D * )
+
+							Print(CtxAll, TEXT("NumSteps:"), FontYellow);
+							Print(CtxAll, NumSteps, FontWhite);
+							Newline(CtxAll);
+
+							Print(CtxAll, TEXT("X0      :"), FontYellow);
+							Print(CtxAll, X0, FontWhite);
+							Newline(CtxAll);
+
+							Print(CtxAll, TEXT("X1      :"), FontYellow);
+							Print(CtxAll, X1, FontWhite);
+							Newline(CtxAll);
+
+							Print(CtxAll, TEXT("T       :"), FontYellow);
+							Print(CtxAll, T, FontWhite);
+							Newline(CtxAll);
+
+							Print(CtxAll, TEXT("bValid  :"), FontYellow);
+							PrintBool(CtxAll, bIsSegmentValid);
+							Newline(CtxAll);
+
+							Print(CtxAll, TEXT("bClip X :"), FontYellow);
+							PrintBool(CtxAll, bClipped.x);
+							Newline(CtxAll);
+
+							Print(CtxAll, TEXT("bClip Y :"), FontYellow);
+							PrintBool(CtxAll, bClipped.y);
+							Newline(CtxAll);
 							
 							
 							//const float3 P0 = UnpackHairControlPoint(PrimID).Position;
@@ -1391,7 +1488,7 @@ void RasterCS(
 
 						// DEBUG (Legend)
 						#if 0 && PERMUTATION_DEBUG
-						if (GroupThread1D == 0 && bDebugEnable)
+						if (bDebugEnable)
 						{
 							//Print(Ctx, s_Mask[IntraTileCoord.x][IntraTileCoord.y], FontWhite);
 							//PrintPixelMask(Ctx, ReadMask(IntraTileCoord)/*s_Mask[IntraTileCoord.x][IntraTileCoord.y]*/);
@@ -1402,7 +1499,7 @@ void RasterCS(
 							const float3 Color1 = LoadSampleColor(PrimID + 1);
 							Print(Ctx, TEXT("Color0 :  "), FontWhite); Print(Ctx, Color0); Newline(Ctx);
 							Print(Ctx, TEXT("Color1 :  "), FontWhite); Print(Ctx, Color1); Newline(Ctx);
-							Print(Ctx, TEXT("  Alpha       A0          A1          AColor     Color"), FontWhite);
+							Print(Ctx, TEXT("  Alpha       A0          A1          AColor     Color    Coord  Coverage"), FontWhite);
 							Newline(Ctx);
 						}
 						#endif
@@ -1428,11 +1525,24 @@ void RasterCS(
 								const float3 Color1 = LoadSampleColor(PrimID+1);
 								const float3 Color = lerp(Color0, Color1, AlphaColor);
 
+								const float HairRadius = lerp(Rad0, Rad1, AlphaColor);
+
 								// Fill in sample data
 								Sample.Depth = Depth;
 								Sample.Coord = IntraTileCoord;
 								Sample.Coverage = 0.5f; //TODO;
 
+								// Compute actual coverage
+								//{
+									// Minimal radius to snap the strand to a sample/pixel center (to avoid aliasing)
+									const bool bIsOrthoView = false;
+									const float SceneDistance = ConvertFromDeviceZ(Depth);
+									const float MinHairRadius = bIsOrthoView ? RadiusAtDepth1 : (SceneDistance * RadiusAtDepth1);
+									Sample.Coverage = saturate(HairRadius / MinHairRadius);
+								//}
+
+
+
 								// Write data
 								WriteMask(IntraTileCoord, 1u << GroupThread1D);
 								//InterlockedOr(s_Mask[IntraTileCoord.x][IntraTileCoord.y], 1u<<GroupThread1D);
@@ -1444,7 +1554,7 @@ void RasterCS(
 
 								// DEBUG (Write coord)
 								#if 0 && PERMUTATION_DEBUG
-								if (GroupThread1D == 0 && bDebugEnable)
+								if (bDebugEnable)
 								{
 									Print(Ctx, J, FontYellow, 2, 0);
 									//Print(Ctx, IntraTileCoord, FontWhite,2,0);
@@ -1454,10 +1564,14 @@ void RasterCS(
 									//Print(Ctx, TEXT("Color1: "), FontWhite); Print(Ctx, Color1); Newline(Ctx);
 									//Print(Ctx, TEXT("Color :  "), FontWhite); Print(Ctx, Color);
 									Print(Ctx, Alpha, FontWhite);
-									Print(Ctx, Alpha0, FontRed);
-									Print(Ctx, Alpha1, FontGreen);
+									//Print(Ctx, Alpha0, FontRed);
+									//Print(Ctx, Alpha1, FontGreen);
 									Print(Ctx, AlphaColor, FontBlue);
-									Print(Ctx, Color, FontOrange);
+									//Print(Ctx, Color, FontOrange);
+									//Print(Ctx, IntraTileCoord, FontOrange);
+									Print(Ctx, HairRadius, FontRed);
+									Print(Ctx, MinHairRadius, FontGreen);
+									Print(Ctx, Sample.Coverage, FontBlue);
 									Newline(Ctx);
 								}
 								#endif
@@ -1469,7 +1583,7 @@ void RasterCS(
 
 				// DEBUG (Coverage mask)
 				#if 0 && PERMUTATION_DEBUG
-				if (GroupThread1D == 0 && bDebugEnable)
+				if (bDebugEnable)
 				{
 					PrintCoverage(Ctx);
 				}
@@ -1606,7 +1720,8 @@ void RasterCS(
 				if (bOutput)
 				{
 					const float3 SourceColor = OutSceneColorTexture[PixelCoord].xyz;
-					OutSceneColorTexture[PixelCoord] = float4(lerp(SourceColor, Thread_Color, Thread_Coverage), 1);
+					OutSceneColorTexture[PixelCoord] = float4(SourceColor * (1-Thread_Coverage) + Thread_Color, 1);
+					//OutSceneColorTexture[PixelCoord] = float4(lerp(SourceColor, Thread_Color, Thread_Coverage), 1);
 					//OutSceneColorTexture[PixelCoord] = float4(Thread_Color, 1);
 					//OutSceneColorTexture[PixelCoord] = float4(1, 1, 0, 1);
 					//OutSceneColorTexture[PixelCoord] = float4(PixelCoord / OutputResolutionf,0,1);
@@ -1623,38 +1738,42 @@ void RasterCS(
 			}
 
 			#if PERMUTATION_DEBUG
-			if (GroupThread1D == 0)
+			if (bDebugEnable)
 			{
-				// Full/Partial coverage tiles
-				if (0)
-				{
-					const bool bFullCover = s_OpaqueMask == 64;
-					if (bFullCover)	AddFilledQuadSS(RasterTileMin, RasterTileMax, Transparent(ColorLightRed));
-					else			AddFilledQuadSS(RasterTileMin, RasterTileMax, Transparent(ColorLightGreen));
-				}
+				Newline(Ctx);
+				Print(Ctx, TEXT("Coverage: "), FontWhite);
+				Print(Ctx, Thread_Coverage, FontYellow);
+				Newline(Ctx);
 
-				if (bDebugEnable)
-				{
-					Newline(Ctx);
-					Print(Ctx, TEXT("Full Coverage: "), FontWhite);
-					Print(Ctx, s_OpaqueMask, FontYellow, 2, 0);
-					Newline(Ctx);
+				Print(Ctx, TEXT("Full Coverage: "), FontWhite);
+				Print(Ctx, s_OpaqueMask, FontYellow, 2, 0);
+				Newline(Ctx);
 
-					Newline(Ctx);
-					for (uint y = 0; y < RASTER_TILE_SIZE; ++y)
+				Newline(Ctx);
+				for (uint y = 0; y < RASTER_TILE_SIZE; ++y)
+				{
+					for (uint x = 0; x < RASTER_TILE_SIZE; ++x)
 					{
-						for (uint x = 0; x < RASTER_TILE_SIZE; ++x)
-						{
-							const uint Index = x + y * RASTER_TILE_SIZE;
-							FFontColor Color = InitFontColor(ColorMapTurbo(s_Coverage[Index]));
-							//Print(Ctx, s_Coverage[Index], Color, 4, 2);
+						const uint Index = x + y * RASTER_TILE_SIZE;
+						FFontColor Color = InitFontColor(ColorMapTurbo(s_Coverage[Index]));
+						//Print(Ctx, s_Coverage[Index], Color, 4, 2);
 							
-							Print(Ctx, Thread_LoopCountToFullCoverage, Color, 2, 0);
-						}
-						Newline(Ctx);
-						Newline(Ctx);
+						Print(Ctx, Thread_LoopCountToFullCoverage, Color, 2, 0);
 					}
 					Newline(Ctx);
+					Newline(Ctx);
+				}
+				Newline(Ctx);
+			}
+
+			if (GroupThread1D == 0)
+			{
+				// Full/Partial coverage tiles
+				if (1)
+				{
+					const bool bFullCover = s_OpaqueMask == 64;
+					if (bFullCover)	AddFilledQuadSS(RasterTileMin, RasterTileMax, Transparent(ColorLightRed));
+					//else			AddFilledQuadSS(RasterTileMin, RasterTileMax, Transparent(ColorLightGreen));
 				}
 
 				//InterlockedAdd(OutHairPixelCountPerTile_ForDebug[RasterTileCoord], PrimCount);
diff --git a/Engine/Source/Runtime/Renderer/Private/HairStrands/HairStrandsForwardRaster.cpp b/Engine/Source/Runtime/Renderer/Private/HairStrands/HairStrandsForwardRaster.cpp
--- a/Engine/Source/Runtime/Renderer/Private/HairStrands/HairStrandsForwardRaster.cpp
+++ b/Engine/Source/Runtime/Renderer/Private/HairStrands/HairStrandsForwardRaster.cpp
@@ -23,6 +23,7 @@ extern int32 GHairVisibilityComputeRaster_Debug;
 
 void SetUpViewHairRenderInfo(const FViewInfo& ViewInfo, FVector4f& OutHairRenderInfo, uint32& OutHairRenderInfoBits, uint32& OutHairComponents);
 void SetUpViewHairRenderInfo(const FViewInfo& ViewInfo, bool bEnableMSAA, FVector4f& OutHairRenderInfo, uint32& OutHairRenderInfoBits, uint32& OutHairComponents);
+FMinHairRadiusAtDepth1 ComputeMinStrandRadiusAtDepth1(const FIntPoint& Resolution, const float FOV, const uint32 SampleCount, const float OverrideStrandHairRasterizationScale);
 
 inline bool IsHairStrandsForwardRasterSupported(EShaderPlatform In) 
 { 
@@ -79,7 +80,9 @@ BEGIN_SHADER_PARAMETER_STRUCT(FRasterComputeForwardCommonParameters, )
 	SHADER_PARAMETER(uint32,    NumBinners)
 	SHADER_PARAMETER(float,     RcpNumBinners)
 
-	SHADER_PARAMETER(uint32, MaxControlPointCount)
+	SHADER_PARAMETER(float,     RadiusAtDepth1)
+	SHADER_PARAMETER(uint32,    MaxControlPointCount)
+
 	SHADER_PARAMETER_RDG_BUFFER_SRV(Buffer, ControlPoints)
 	SHADER_PARAMETER_RDG_TEXTURE(Texture2D<uint>, ControlPointCount)
 END_SHADER_PARAMETER_STRUCT()
@@ -514,6 +517,8 @@ void AddHairStrandsForwardRasterPass(
 		Common.MaxControlPointCount = InData.MaxControlPointCount;
 		Common.ControlPoints = InData.ControlPointsSRV;
 		Common.ControlPointCount = InData.ControlPointCount;
+
+		Common.RadiusAtDepth1 = ComputeMinStrandRadiusAtDepth1(FIntPoint(ViewInfo.UnconstrainedViewRect.Width(), ViewInfo.UnconstrainedViewRect.Height()), ViewInfo.FOV, 1 /*SampleCount*/, 0 /*ScaleOverride*/).Primary;
 	}
 
 	// Fill in tile depth
