Add hair motion vector output support.

diff --git a/Engine/Shaders/Private/HairStrands/HairStrandsForwardRaster.usf b/Engine/Shaders/Private/HairStrands/HairStrandsForwardRaster.usf
--- a/Engine/Shaders/Private/HairStrands/HairStrandsForwardRaster.usf
+++ b/Engine/Shaders/Private/HairStrands/HairStrandsForwardRaster.usf
@@ -41,6 +41,20 @@ float3 UnpackR11G11B10F(uint rgb)
 	return float3(r, g, b);
 }
 
+uint PackRG16(float2 In)
+{
+#if COMPILER_SUPPORTS_PACK_INTRINSICS
+	return PackUInt2ToUInt(f32tof16(In.r), f32tof16(In.g));
+#else
+	return f32tof16(In.r) | (f32tof16(In.g) << 16);
+#endif
+}
+
+float2 UnpackRG16(uint In)
+{
+	return float2(f16tof32(In), f16tof32(In >> 16));
+}
+
 uint PackDepth(float In)
 {
 	return asuint(In);
@@ -870,6 +884,7 @@ void CompactionCS(uint DispatchThreadID : SV_DispatchThreadID, uint GroupThreadI
 #define RASTER_SEGMENT_COUNT 32
 #define SEGMENT_COUNT_PER_GROUP 1024
 #define INVALID_PRIM_ID 0xFFFFFFFF
+#define INVALID_VELOCITY -1e8f
 
 // For editing convenience
 #if !SHADER_RASTERCOMPUTE_DEBUG && !SHADER_RASTERCOMPUTE_COMPACTION && !SHADER_RASTERCOMPUTE_BINNING && !SHADER_RASTERCOMPUTE_RASTER_MULTI_SAMPLE && !SHADER_RASTERCOMPUTE_DEPTH_GRID
@@ -882,6 +897,7 @@ Buffer<uint> 				VisTilePrims;
 Buffer<uint> 				VisTileArgs;
 ByteAddressBuffer			VisTileData;
 RWTexture2D<float4>			OutSceneColorTexture;
+RWTexture2D<float4>			OutSceneVelocityTexture;
 RWStructuredBuffer<uint>	RWWorkCounter;
 
 #if PERMUTATION_DEBUG
@@ -891,6 +907,7 @@ RWTexture2D<uint>			OutHairPixelCountPerTile_ForDebug;
 
 int2						SampleLightingViewportResolution;
 Texture2D<float4>			SampleLightingTexture;
+Buffer<float4>				SampleVelocityBuffer;
 
 // Depth|ID
 groupshared uint2  s_Segments[1024];
@@ -906,6 +923,7 @@ groupshared uint   s_Mask[8][8];
 groupshared uint   s_OpaqueMask;
 groupshared uint   s_Data[RASTER_SEGMENT_COUNT][8];
 groupshared uint   s_Color[RASTER_SEGMENT_COUNT][8];
+groupshared uint2  s_Velocity[RASTER_SEGMENT_COUNT][8];
 groupshared uint   s_bDataOrder;
 groupshared uint   s_WorkID;
 
@@ -1004,6 +1022,16 @@ float3 UnpackColorData(uint In)
 	return UnpackR11G11B10F(In);
 }
 
+uint2 PackVelocityData(float4 In)
+{
+	return uint2(PackRG16(In.xy), PackRG16(In.zw));
+}
+
+float4 UnpackVelocityData(uint2 In)
+{
+	return float4(UnpackRG16(In.x), UnpackRG16(In.y));
+}
+
 float3 LoadSampleColor(uint InPrimId)
 {
 	uint2 SampleCoord = 0;
@@ -1012,6 +1040,13 @@ float3 LoadSampleColor(uint InPrimId)
 	return SampleLightingTexture.Load(uint3(SampleCoord, 0)).xyz;
 }
 
+float4 LoadSampleVelocity(uint InPrimId)
+{
+	// This return the encoded velocity
+	// For decoding the actual velocity, use DecodeVelocityFromTexture(...)
+	return SampleVelocityBuffer[InPrimId];
+}
+
 uint GetDepthBinIndex(float InDepth, float InvDepthExtent)
 {
 	// Inverse-Z
@@ -1322,6 +1357,7 @@ void RasterCS(
 			float3 Thread_Color    = 0;
 			float  Thread_Coverage = 0;
 			uint   Thread_Complete = 0;
+			float4 Thread_Velocity = INVALID_VELOCITY;
 			uint   Thread_LoopCountToFullCoverage=0;
 
 			// 4.1 Loop over all segment within the tile, and rastize 32 of them at each loop
@@ -1335,6 +1371,7 @@ void RasterCS(
 					{
 						s_Data[GroupThread1D][J] = 0;
 						s_Color[GroupThread1D][J] = 0;
+						s_Velocity[GroupThread1D][J] = 0;
 					}
 				}
 				s_Mask[GroupThread2D.x][GroupThread2D.y] = 0;
@@ -1525,6 +1562,10 @@ void RasterCS(
 								const float3 Color1 = LoadSampleColor(PrimID+1);
 								const float3 Color = lerp(Color0, Color1, AlphaColor);
 
+								const float4 Velocity0 = LoadSampleVelocity(PrimID);
+								const float4 Velocity1 = LoadSampleVelocity(PrimID + 1);
+								const float4 Velocity = lerp(Velocity0, Velocity1, AlphaColor);
+
 								const float HairRadius = lerp(Rad0, Rad1, AlphaColor);
 
 								// Fill in sample data
@@ -1550,6 +1591,7 @@ void RasterCS(
 								const uint JCoord = bIsSteep ? IntraTileCoord.y : IntraTileCoord.x;
 								s_Color[GroupThread1D][JCoord] = PackColorData(Color);
 								s_Data[GroupThread1D][JCoord]  = PackSampleData(Sample);
+								s_Velocity[GroupThread1D][JCoord] = PackVelocityData(Velocity);
 
 
 								// DEBUG (Write coord)
@@ -1662,6 +1704,7 @@ void RasterCS(
 
 								const FSampleData Sample = UnpackSampleData(s_Data[SegmentIt][J]);
 								const float3 Color = UnpackColorData(s_Color[SegmentIt][J]);
+								const float4 Velocity = UnpackVelocityData(s_Velocity[SegmentIt][J]);
 
 								// Accumulate sample contribution 
 								// TODO: needs to compute transmittance funciton, as seg ordered from front to back, but the raster pixels might be slightly different
@@ -1678,6 +1721,12 @@ void RasterCS(
 								Thread_Coverage += AccTransmittance * Sample.Coverage;
 								Thread_Color    += AccTransmittance * Sample.Coverage * Color;
 
+								// Use the closest valid segment as output velocity
+								if (Thread_Velocity.x <= INVALID_VELOCITY)
+								{
+									Thread_Velocity = Velocity;
+								}
+
 								//Thread_Coverage = 0.5;
 								//Thread_Color    = 0;
 
@@ -1721,6 +1770,10 @@ void RasterCS(
 				{
 					const float3 SourceColor = OutSceneColorTexture[PixelCoord].xyz;
 					OutSceneColorTexture[PixelCoord] = float4(SourceColor * (1-Thread_Coverage) + Thread_Color, 1);
+					if (Thread_Velocity.x > INVALID_VELOCITY)
+					{
+						OutSceneVelocityTexture[PixelCoord] = Thread_Velocity;
+					}
 					//OutSceneColorTexture[PixelCoord] = float4(lerp(SourceColor, Thread_Color, Thread_Coverage), 1);
 					//OutSceneColorTexture[PixelCoord] = float4(Thread_Color, 1);
 					//OutSceneColorTexture[PixelCoord] = float4(1, 1, 0, 1);
diff --git a/Engine/Source/Runtime/Renderer/Private/HairStrands/HairStrandsComposition.cpp b/Engine/Source/Runtime/Renderer/Private/HairStrands/HairStrandsComposition.cpp
--- a/Engine/Source/Runtime/Renderer/Private/HairStrands/HairStrandsComposition.cpp
+++ b/Engine/Source/Runtime/Renderer/Private/HairStrands/HairStrandsComposition.cpp
@@ -488,7 +488,8 @@ static void InternalRenderHairComposition(
 			SceneDepthTexture->Desc.Extent,
 			View.HairStrandsViewData.VisibilityData,
 			SceneDepthTexture,
-			SceneColorTexture);
+			SceneColorTexture,
+			SceneVelocityTexture);
 	}
 	else
 	{
diff --git a/Engine/Source/Runtime/Renderer/Private/HairStrands/HairStrandsData.h b/Engine/Source/Runtime/Renderer/Private/HairStrands/HairStrandsData.h
--- a/Engine/Source/Runtime/Renderer/Private/HairStrands/HairStrandsData.h
+++ b/Engine/Source/Runtime/Renderer/Private/HairStrands/HairStrandsData.h
@@ -145,6 +145,7 @@ struct FHairStrandsVisibilityData
 	uint32				MaxControlPointCount = 0;
 	FRDGBufferSRVRef	ControlPointsSRV = nullptr;
 	FRDGTextureRef		ControlPointCount = nullptr;
+	FRDGBufferSRVRef	ControlPointVelocitySRV = nullptr;
 
 	FHairStrandsTiles TileData;
 
diff --git a/Engine/Source/Runtime/Renderer/Private/HairStrands/HairStrandsForwardRaster.cpp b/Engine/Source/Runtime/Renderer/Private/HairStrands/HairStrandsForwardRaster.cpp
--- a/Engine/Source/Runtime/Renderer/Private/HairStrands/HairStrandsForwardRaster.cpp
+++ b/Engine/Source/Runtime/Renderer/Private/HairStrands/HairStrandsForwardRaster.cpp
@@ -196,6 +196,7 @@ class FHairStrandsForwardRasterRasterizeCS : public FGlobalShader
 		SHADER_PARAMETER_STRUCT_INCLUDE(FRasterComputeForwardCommonParameters, Common)
 		SHADER_PARAMETER(FIntPoint, SampleLightingViewportResolution)
 		SHADER_PARAMETER_RDG_TEXTURE(Texture2D, SampleLightingTexture)
+		SHADER_PARAMETER_RDG_BUFFER_SRV(Buffer<float4>, SampleVelocityBuffer)
 
 		SHADER_PARAMETER_RDG_TEXTURE(Texture2D, SceneDepthTexture)
 		SHADER_PARAMETER_RDG_BUFFER_SRV(Buffer, VisTilePrims)
@@ -204,6 +205,7 @@ class FHairStrandsForwardRasterRasterizeCS : public FGlobalShader
 		SHADER_PARAMETER_RDG_TEXTURE_UAV(RWTexture2D, OutHairCountTexture_ForDebug)
 		SHADER_PARAMETER_RDG_TEXTURE_UAV(RWTexture2D, OutHairPixelCountPerTile_ForDebug)
 		SHADER_PARAMETER_RDG_TEXTURE_UAV(RWTexture2D, OutSceneColorTexture)
+		SHADER_PARAMETER_RDG_TEXTURE_UAV(RWTexture2D, OutSceneVelocityTexture)
 		SHADER_PARAMETER_RDG_BUFFER_UAV(RWStructuredBuffer, RWWorkCounter)
 		SHADER_PARAMETER_STRUCT_INCLUDE(ShaderPrint::FShaderParameters, ShaderPrintParameters)
 		SHADER_PARAMETER_STRUCT_REF(FViewUniformShaderParameters, ViewUniformBuffer)
@@ -425,7 +427,8 @@ void AddHairStrandsForwardRasterPass(
 	const FIntPoint& InResolution,
 	const FHairStrandsVisibilityData& InData,
 	const FRDGTextureRef SceneDepthTexture,
-	const FRDGTextureRef SceneColorTexture)
+	const FRDGTextureRef SceneColorTexture,
+	const FRDGTextureRef SceneVelocityTexture)
 {	
 	if (!IsHairStrandsForwardRasterSupported(ViewInfo.GetShaderPlatform()))
 	{
@@ -622,10 +625,12 @@ void AddHairStrandsForwardRasterPass(
 		Parameters->OutHairPixelCountPerTile_ForDebug = HairPixelCountPerTile_ForDebug ? GraphBuilder.CreateUAV(HairPixelCountPerTile_ForDebug) : nullptr;
 		Parameters->SampleLightingViewportResolution = InData.SampleLightingViewportResolution;
 		Parameters->SampleLightingTexture = InData.SampleLightingTexture;
+		Parameters->SampleVelocityBuffer = InData.ControlPointVelocitySRV;
 		Parameters->VisTilePrims = GraphBuilder.CreateSRV(BinData.CompactedVisTilePrims, PF_R32_UINT);
 		Parameters->VisTileArgs = GraphBuilder.CreateSRV(BinData.CompactedVisTileArgs, PF_R32_UINT);
 		Parameters->VisTileData = GraphBuilder.CreateSRV(BinData.CompactedVisTileData);
 		Parameters->OutSceneColorTexture = GraphBuilder.CreateUAV(SceneColorTexture);
+		Parameters->OutSceneVelocityTexture = GraphBuilder.CreateUAV(SceneVelocityTexture);
 		Parameters->RWWorkCounter = WorkCounterUAV;
 		ShaderPrint::SetParameters(GraphBuilder, ViewInfo.ShaderPrintData, Parameters->ShaderPrintParameters);
 
diff --git a/Engine/Source/Runtime/Renderer/Private/HairStrands/HairStrandsForwardRaster.h b/Engine/Source/Runtime/Renderer/Private/HairStrands/HairStrandsForwardRaster.h
--- a/Engine/Source/Runtime/Renderer/Private/HairStrands/HairStrandsForwardRaster.h
+++ b/Engine/Source/Runtime/Renderer/Private/HairStrands/HairStrandsForwardRaster.h
@@ -41,4 +41,5 @@ void AddHairStrandsForwardRasterPass(
 	const FIntPoint& InResolution,
 	const FHairStrandsVisibilityData& InData,
 	const FRDGTextureRef SceneDepthTexture,
-	const FRDGTextureRef SceneColorTexture);
\ No newline at end of file
+	const FRDGTextureRef SceneColorTexture,
+	const FRDGTextureRef SceneVelocityColorTexture);
\ No newline at end of file
diff --git a/Engine/Source/Runtime/Renderer/Private/HairStrands/HairStrandsVisibility.cpp b/Engine/Source/Runtime/Renderer/Private/HairStrands/HairStrandsVisibility.cpp
--- a/Engine/Source/Runtime/Renderer/Private/HairStrands/HairStrandsVisibility.cpp
+++ b/Engine/Source/Runtime/Renderer/Private/HairStrands/HairStrandsVisibility.cpp
@@ -812,6 +812,7 @@ struct FMaterialPassOutput
 {
 	FRDGBufferRef NodeData = nullptr;
 	FRDGBufferRef NodeVelocity = nullptr;
+	FRDGBufferSRVRef NodeVelocitySRV = nullptr;
 	FRDGTextureRef SampleLightingTexture = nullptr;
 };
 
@@ -842,6 +843,7 @@ static FMaterialPassOutput AddHairMaterialPass(
 	Output.NodeData				 = GraphBuilder.CreateBuffer(FRDGBufferDesc::CreateStructuredDesc(sizeof(HairStrandsVisibilityInternal::NodeData), MaxNodeCount), TEXT("Hair.CompactNodeData"));
 	Output.NodeVelocity			 = GraphBuilder.CreateBuffer(FRDGBufferDesc::CreateBufferDesc(VelocityFormat == PF_G16R16 ? 4 : 8, CompactNodeVis->Desc.NumElements), TEXT("Hair.CompactNodeVelocity"));
 	Output.SampleLightingTexture = AddClearLightSamplePass(GraphBuilder, ViewInfo, MaxNodeCount, CompactNodeCounter);
+	Output.NodeVelocitySRV		 = GraphBuilder.CreateSRV(Output.NodeVelocity, VelocityFormat);
 
 	const uint32 ResolutionDim = FMath::CeilToInt(FMath::Sqrt(static_cast<float>(MaxNodeCount)));
 	const FIntPoint Resolution(ResolutionDim, ResolutionDim);
@@ -4284,6 +4286,7 @@ void RenderHairStrandsVisibilityBuffer(
 				VisibilityData.NodeIndirectArg = IndirectArgsBuffer;
 				VisibilityData.NodeCount = NodeCounter;
 				VisibilityData.ResolveMaskTexture = nullptr;
+				VisibilityData.ControlPointVelocitySRV = PassOutput.NodeVelocitySRV;
 
 				// For fully covered pixels, write: 
 				// * black color into the scene color
