Fixing android shader compiles caused by using a single SamplerState with multiple Textures.

diff --git a/Engine/Shaders/Private/RectLightAtlas.usf b/Engine/Shaders/Private/RectLightAtlas.usf
--- a/Engine/Shaders/Private/RectLightAtlas.usf
+++ b/Engine/Shaders/Private/RectLightAtlas.usf
@@ -55,7 +55,14 @@ Texture2D<float4> InTexture5;
 Texture2D<float4> InTexture6;
 Texture2D<float4> InTexture7;
 
-SamplerState InSampler;
+SamplerState InSampler0;
+SamplerState InSampler1;
+SamplerState InSampler2;
+SamplerState InSampler3;
+SamplerState InSampler4;
+SamplerState InSampler5;
+SamplerState InSampler6;
+SamplerState InSampler7;
 
 void MainPS(
 	in float4 InPosition : SV_Position,
@@ -67,14 +74,14 @@ void MainPS(
 	float4 Color = 0;
 	switch (InId)
 	{
-	case 0: Color = InTexture0.SampleLevel(InSampler, InUV, InTextureMIPBias); break;
-	case 1: Color = InTexture1.SampleLevel(InSampler, InUV, InTextureMIPBias); break;
-	case 2: Color = InTexture2.SampleLevel(InSampler, InUV, InTextureMIPBias); break;
-	case 3: Color = InTexture3.SampleLevel(InSampler, InUV, InTextureMIPBias); break;
-	case 4: Color = InTexture4.SampleLevel(InSampler, InUV, InTextureMIPBias); break;
-	case 5: Color = InTexture5.SampleLevel(InSampler, InUV, InTextureMIPBias); break;
-	case 6: Color = InTexture6.SampleLevel(InSampler, InUV, InTextureMIPBias); break;
-	case 7: Color = InTexture7.SampleLevel(InSampler, InUV, InTextureMIPBias); break;
+	case 0: Color = InTexture0.SampleLevel(InSampler0, InUV, InTextureMIPBias); break;
+	case 1: Color = InTexture1.SampleLevel(InSampler1, InUV, InTextureMIPBias); break;
+	case 2: Color = InTexture2.SampleLevel(InSampler2, InUV, InTextureMIPBias); break;
+	case 3: Color = InTexture3.SampleLevel(InSampler3, InUV, InTextureMIPBias); break;
+	case 4: Color = InTexture4.SampleLevel(InSampler4, InUV, InTextureMIPBias); break;
+	case 5: Color = InTexture5.SampleLevel(InSampler5, InUV, InTextureMIPBias); break;
+	case 6: Color = InTexture6.SampleLevel(InSampler6, InUV, InTextureMIPBias); break;
+	case 7: Color = InTexture7.SampleLevel(InSampler7, InUV, InTextureMIPBias); break;
 	}
 
 	// Ensure there is no NaN value
diff --git a/Engine/Source/Runtime/Renderer/Private/RectLightTextureManager.cpp b/Engine/Source/Runtime/Renderer/Private/RectLightTextureManager.cpp
--- a/Engine/Source/Runtime/Renderer/Private/RectLightTextureManager.cpp
+++ b/Engine/Source/Runtime/Renderer/Private/RectLightTextureManager.cpp
@@ -414,7 +414,14 @@ class FRectAtlasAddTexturePS : public FGlobalShader
 		SHADER_PARAMETER_TEXTURE(Texture2D, InTexture5)
 		SHADER_PARAMETER_TEXTURE(Texture2D, InTexture6)
 		SHADER_PARAMETER_TEXTURE(Texture2D, InTexture7)
-		SHADER_PARAMETER_SAMPLER(SamplerState, InSampler)
+		SHADER_PARAMETER_SAMPLER(SamplerState, InSampler0)
+		SHADER_PARAMETER_SAMPLER(SamplerState, InSampler1)
+		SHADER_PARAMETER_SAMPLER(SamplerState, InSampler2)
+		SHADER_PARAMETER_SAMPLER(SamplerState, InSampler3)
+		SHADER_PARAMETER_SAMPLER(SamplerState, InSampler4)
+		SHADER_PARAMETER_SAMPLER(SamplerState, InSampler5)
+		SHADER_PARAMETER_SAMPLER(SamplerState, InSampler6)
+		SHADER_PARAMETER_SAMPLER(SamplerState, InSampler7)
 		SHADER_PARAMETER_STRUCT_INCLUDE(FRectLightAtlasVS::FParameters, VS)
 		RENDER_TARGET_BINDING_SLOTS()
 	END_SHADER_PARAMETER_STRUCT()
@@ -481,8 +488,17 @@ static void AddSlotsPass(
 			}
 		}
 
+		FRHISamplerState* SamplerState = TStaticSamplerState<SF_Bilinear, AM_Clamp, AM_Clamp, AM_Clamp>::GetRHI();
+		Parameters->InSampler0 = SamplerState;
+		Parameters->InSampler1 = SamplerState;
+		Parameters->InSampler2 = SamplerState;
+		Parameters->InSampler3 = SamplerState;
+		Parameters->InSampler4 = SamplerState;
+		Parameters->InSampler5 = SamplerState;
+		Parameters->InSampler6 = SamplerState;
+		Parameters->InSampler7 = SamplerState;
+
 		Parameters->InTextureMIPBias = TextureMIPBias;
-		Parameters->InSampler = TStaticSamplerState<SF_Bilinear, AM_Clamp, AM_Clamp, AM_Clamp>::GetRHI();
 		Parameters->VS.AtlasResolution = Resolution;
 		Parameters->VS.SlotBufferOffset = SlotOffset;
 		Parameters->VS.SlotBuffer = GraphBuilder.CreateSRV(SlotBuffer, PF_R16G16B16A16_UINT);
