* Add Min/MaxZ tracking per binning tile
* Add segment sorting during compaction

diff --git a/Engine/Shaders/Private/HairStrands/HairStrandsForwardRaster.usf b/Engine/Shaders/Private/HairStrands/HairStrandsForwardRaster.usf
--- a/Engine/Shaders/Private/HairStrands/HairStrandsForwardRaster.usf
+++ b/Engine/Shaders/Private/HairStrands/HairStrandsForwardRaster.usf
@@ -67,19 +67,19 @@ float UnpackDepth(uint In)
 
 uint PackDepthRange(float2 In)
 {
-	return f32tof16(In.x) | (f32tof16(In.y) << 16);
+	return PackRG16(In);
 }
 
 float2 UnpackDepthRange(uint In)
 {
-	return float2(f16tof32(In & 0xffff), f16tof32(In >> 16));
+	return UnpackRG16(In);
 }
 
 ///////////////////////////////////////////////////////////////////////////
 // Tile Helpers
 
 // Max number of iterator that a rasterizer can do. This is for preventing any kind of infinite loop.
-#define MAX_WORK_COUNT 256
+#define MAX_WORK_COUNT 4096
 
 #define BIN_TILE_SIZE 32
 #define BIN_RCP_TILE_SIZE (1.f / BIN_TILE_SIZE)
@@ -168,11 +168,11 @@ struct FVisTile
 #define VT_SIZE 5
 
 // Visibility tile data are stored as:
-//   _______________________________________________________________________________________________________________________________________________________________________
-// ||                      Tile 0                          ||                      Tile 1                          ||                      Tile 2                          ||
-// ||______________________________________________________||______________________________________________________||______________________________________________________||
-// ||            |           |             |               ||            |           |             |               ||            |           |             |               ||
-// || PrimOffset | PrimCount | (Tile)Coord | MinWriteIndex || PrimOffset | PrimCount | (Tile)Coord | MinWriteIndex || PrimOffset | PrimCount | (Tile)Coord | MinWriteIndex ||
+//   ________________________________________________________________________________________________________________________________________________________________________________________________________________
+// ||                      Tile 0                                        ||                      Tile 1                                        ||                      Tile 2                                        ||
+// ||____________________________________________________________________||____________________________________________________________________||____________________________________________________________________||
+// ||            |           |             |               |             ||            |           |             |               |             ||            |           |             |               |             ||
+// || PrimOffset | PrimCount | (Tile)Coord | MinWriteIndex | MinMaxDepth || PrimOffset | PrimCount | (Tile)Coord | MinWriteIndex | MinMaxDepth || PrimOffset | PrimCount | (Tile)Coord | MinWriteIndex | MinMaxDepth ||
 
 uint PackVisTileCoord(uint2 In)
 {
@@ -472,16 +472,17 @@ void PrepareDepthGridCS(uint DispatchThreadID : SV_DispatchThreadID, uint GroupT
 
 ///////////////////////////////////////////////////////////////////////////
 
-#define BIN_MINMAX 0
+#define BIN_MINMAX 1
 
 #if SHADER_RASTERCOMPUTE_BINNING
 
 RWTexture2DArray<uint> 				OutVisTileBinningGrid;
 #if BIN_MINMAX
-RWTexture2DArray<uint> 				OutVisTileBinningGridMin;
-RWTexture2DArray<uint> 				OutVisTileBinningGridMax;
+RWTexture2DArray<uint> 				OutVisTileBinningGridMinZ;
+RWTexture2DArray<uint> 				OutVisTileBinningGridMaxZ;
 #endif
 RWBuffer<uint>						OutVisTilePrims;
+RWBuffer<uint>						OutVisTilePrimDepths;
 RWBuffer<uint>						OutVisTileArgs;
 RWByteAddressBuffer					OutVisTileData;
 Texture2D<uint>						VisTileDepthGrid;
@@ -512,7 +513,7 @@ void BinningCS(uint GroupThreadID : SV_GroupThreadID, uint GroupID : SV_GroupID)
 	GroupMemoryBarrierWithGroupSync();
 
 	#if PERMUTATION_DEBUG
-	const bool bDebugEnabled = false; //GroupID == 0 && GroupThreadID <= 64;
+	const bool bDebugEnabled = GroupID == 0 && GroupThreadID <= 64;
 	FShaderPrintContext Ctx = InitShaderPrintContext(bDebugEnabled, uint2(250, 50));
 	#endif
 
@@ -528,8 +529,8 @@ void BinningCS(uint GroupThreadID : SV_GroupThreadID, uint GroupID : SV_GroupID)
 		const uint TmpSegmentCountLayerIdx = SegmentCountLayerIdx + NumBinners; // Also stores number of segments per tile per workgroup. Used as second counter for this two pass algorithm.
 		const uint TileAllocInfoLayerIdx = SegmentCountLayerIdx + NumBinners * 2; // Stores per tile per workgroup allocation info.
 
-		uint NearestDepth  = 0;
-		uint FurthestDepth = 1;
+		uint MaxZ = 0;
+		uint MinZ = 0xFFFFFFFF;
 		float2 TileCoord0F = 0.0f;
 		float2 TileCoord1F = 0.0f;
 
@@ -567,8 +568,8 @@ void BinningCS(uint GroupThreadID : SV_GroupThreadID, uint GroupID : SV_GroupID)
 					SP0 = clamp(SP0, 0.0f, float3(BinTileRes, 1.0f));
 					SP1 = clamp(SP1, 0.0f, float3(BinTileRes, 1.0f));
 
-					NearestDepth  = PackDepth(max(SP0.z, SP1.z)); // Inverse-Z
-					FurthestDepth = PackDepth(min(SP0.z, SP1.z)); // Inverse-Z
+					MaxZ = PackDepth(max(SP0.z, SP1.z));
+					MinZ = PackDepth(min(SP0.z, SP1.z));
 
 					TileCoord0F = SP0.xy;
 					TileCoord1F = SP1.xy;
@@ -605,7 +606,7 @@ void BinningCS(uint GroupThreadID : SV_GroupThreadID, uint GroupID : SV_GroupID)
 				
 				uint DebugInsertMode = 0;
 				BRANCH
-				if (NearestDepth > VisTileDepthGrid[TileCoord])
+				if (MaxZ > VisTileDepthGrid[TileCoord])  // Inverse-Z
 				{
 					uint OldTileSegmentCount;
 					InterlockedAdd(OutVisTileBinningGrid[uint3(TileCoord, SegmentCountLayerIdx)], 1, OldTileSegmentCount);
@@ -613,8 +614,8 @@ void BinningCS(uint GroupThreadID : SV_GroupThreadID, uint GroupID : SV_GroupID)
 
 					// Min/Max
 					#if BIN_MINMAX
-					InterlockedMin(OutVisTileBinningGridMin[uint3(TileCoord, SegmentCountLayerIdx)], FurthestDepth);
-					InterlockedMax(OutVisTileBinningGridMax[uint3(TileCoord, SegmentCountLayerIdx)], NearestDepth);
+					InterlockedMin(OutVisTileBinningGridMinZ[uint3(TileCoord, SegmentCountLayerIdx)], MinZ);
+					InterlockedMax(OutVisTileBinningGridMaxZ[uint3(TileCoord, SegmentCountLayerIdx)], MaxZ);
 					#endif
 
 					BRANCH
@@ -635,7 +636,6 @@ void BinningCS(uint GroupThreadID : SV_GroupThreadID, uint GroupID : SV_GroupID)
 				{
 					//CP0 = UnpackHairControlPoint(PrimID);
 					//CP1 = UnpackHairControlPoint(PrimID +1);
-
 					//AddLineWS(Ctx, CP0.Position, CP1.Position, ColorRed);
 
 					float4 DebugColor = ColorRed;
@@ -667,8 +667,8 @@ void BinningCS(uint GroupThreadID : SV_GroupThreadID, uint GroupID : SV_GroupID)
 			const uint TotalOldWriteCount = OutVisTileBinningGrid[uint3(TileCoord, TmpSegmentCountLayerIdx)];
 
 			#if BIN_MINMAX
-			const float MinDepth = UnpackDepth(OutVisTileBinningGridMin[uint3(TileCoord, SegmentCountLayerIdx)]);
-			const float MaxDepth = UnpackDepth(OutVisTileBinningGridMax[uint3(TileCoord, SegmentCountLayerIdx)]);
+			const float TileMinZ = UnpackDepth(OutVisTileBinningGridMinZ[uint3(TileCoord, SegmentCountLayerIdx)]);
+			const float TileMaxZ = UnpackDepth(OutVisTileBinningGridMaxZ[uint3(TileCoord, SegmentCountLayerIdx)]);
 			#endif
 
 			uint NewTile;
@@ -679,7 +679,7 @@ void BinningCS(uint GroupThreadID : SV_GroupThreadID, uint GroupID : SV_GroupID)
 			StoreOutVisTileData(OutVisTileData, NewTile, VT_MinWriteIndex, TotalNewWriteCount & ~1023u);
 			// Min/Max depth
 			#if BIN_MINMAX
-			StoreOutVisTileData(OutVisTileData, NewTile, VT_MinMaxDepth, PackDepthRange(float2(NearestDepth, FurthestDepth));
+			StoreOutVisTileData(OutVisTileData, NewTile, VT_MinMaxDepth, PackDepthRange(float2(TileMinZ, TileMaxZ)));
 			#endif
 
 			const uint PrevTile = (OutVisTileBinningGrid[uint3(TileCoord, TileAllocInfoLayerIdx)] & 0xffff);
@@ -705,7 +705,7 @@ void BinningCS(uint GroupThreadID : SV_GroupThreadID, uint GroupID : SV_GroupID)
 				const int2 TileCoord = (int2)floor(DDAContext.Coord);
 
 				BRANCH
-				if (NearestDepth > VisTileDepthGrid[TileCoord])
+				if (MaxZ > VisTileDepthGrid[TileCoord])  // Inverse-Z
 				{
 					const uint PackedTiles = OutVisTileBinningGrid[uint3(TileCoord, TileAllocInfoLayerIdx)];
 					const uint CurTile = (PackedTiles & 0xffff);
@@ -720,7 +720,7 @@ void BinningCS(uint GroupThreadID : SV_GroupThreadID, uint GroupID : SV_GroupID)
 					const uint WritePos = (bWriteToCurTile ? CurTile : PrevTile) * 1024 + LocalWritePos;
 
 					OutVisTilePrims[WritePos] = PrimID;
-
+					OutVisTilePrimDepths[WritePos] = MaxZ; // Inverse-Z
 					BRANCH
 					if (bWriteToCurTile)
 					{
@@ -749,6 +749,7 @@ void BinningCS(uint GroupThreadID : SV_GroupThreadID, uint GroupID : SV_GroupID)
 
 ByteAddressBuffer					InData;
 Buffer<uint> 						InPrims;
+Buffer<uint> 						InDepths;
 Buffer<uint> 						InArgs;
 RWByteAddressBuffer					OutData;
 RWBuffer<uint> 						OutPrims;
@@ -759,7 +760,26 @@ groupshared uint group_PrimWriteOffset;
 groupshared uint group_NumTiles;
 groupshared uint group_TilesToCompact[1024];
 groupshared uint group_MaxLDSTileIdx;
+groupshared uint group_MinZ;
+groupshared uint group_MaxZ;
+
+#define COMPACTION_DEPTH_BUCKET 64
+groupshared uint s_BinOffset[COMPACTION_DEPTH_BUCKET];
+groupshared uint s_BinCount[COMPACTION_DEPTH_BUCKET];
+
+uint GetDepthBinIndex(float InDepth)
+{
+	// Inverse-Z
+	const float MinDepth = UnpackDepth(group_MinZ);
+	const float MaxDepth = UnpackDepth(group_MaxZ);
+	const float InvDepthExtent = 1.f / max(MaxDepth - MinDepth, 1e-5f);
+	const uint DepthIt = clamp(saturate((InDepth - MinDepth) * InvDepthExtent) * COMPACTION_DEPTH_BUCKET, 0, COMPACTION_DEPTH_BUCKET - 1);
+	return (COMPACTION_DEPTH_BUCKET - 1) - DepthIt;
+}
 
+// Launch based on CPU BinTileResX x BinTileResY
+// 1 group per screen-tile, 1 threads per bin-tile matching the screen-tile coord
+// There can be/are several bins for the same screen area
 [numthreads(1024, 1, 1)]
 void CompactionCS(uint DispatchThreadID : SV_DispatchThreadID, uint GroupThreadID : SV_GroupThreadID, uint2 GroupID : SV_GroupID)
 {
@@ -768,12 +788,24 @@ void CompactionCS(uint DispatchThreadID : SV_DispatchThreadID, uint GroupThreadI
 		group_TotalPrimCount = 0;
 		group_NumTiles = 0;
 		group_MaxLDSTileIdx = 0;
+		group_MinZ = 0xFFFFFFFF;
+		group_MaxZ = 0;
 	}
-	
+	if (GroupThreadID < COMPACTION_DEPTH_BUCKET)
+	{
+		s_BinCount[GroupThreadID] = 0;
+	}
+
 	GroupMemoryBarrierWithGroupSync();
 
 	const uint NumTiles = InArgs[0];
-	const uint PackedCoord = PackVisTileCoord(GroupID);
+	const uint PackedCoord = PackVisTileCoord(GroupID); // All thread will process the same tile
+
+	#if PERMUTATION_DEBUG
+	const uint2 TileCoord = UnpackVisTileCoord(PackedCoord);
+	const bool bDebugEnabled = all(TileCoord == uint2(ShaderPrintData.CursorCoord / float(BIN_TILE_SIZE)));
+	FShaderPrintContext Ctx = InitShaderPrintContext(bDebugEnabled, uint2(750, 50));
+	#endif
 
 	// 1. Compute total number of primitives at this tile coordinate
 	uint LocalPrimCount = 0;
@@ -784,6 +816,10 @@ void CompactionCS(uint DispatchThreadID : SV_DispatchThreadID, uint GroupThreadI
 			if (PackedCoord == TilePackedCoord)
 			{
 				LocalPrimCount += LoadVisTileData(InData, TileIdx, VT_PrimCount);
+
+				const float2 LocalDepthRange = UnpackDepthRange(LoadVisTileData(InData, TileIdx, VT_MinMaxDepth));
+				InterlockedMin(group_MinZ, PackDepth(LocalDepthRange.x));
+				InterlockedMax(group_MaxZ, PackDepth(LocalDepthRange.y));
 				
 				uint WritePos;
 				WaveInterlockedAddScalar_(group_NumTiles, 1, WritePos);
@@ -813,7 +849,7 @@ void CompactionCS(uint DispatchThreadID : SV_DispatchThreadID, uint GroupThreadI
 	// 2. Allocate space
 	if (GroupThreadID == 0)
 	{
-		const uint NumTilesToAllocate = DivideAndRoundUp(TotalPrimCount, 1024);
+		uint NumTilesToAllocate = DivideAndRoundUp(TotalPrimCount, 1024);
 
 		uint FirstCompactedTile;
 		InterlockedAdd(OutArgs[0], NumTilesToAllocate, FirstCompactedTile);
@@ -832,7 +868,25 @@ void CompactionCS(uint DispatchThreadID : SV_DispatchThreadID, uint GroupThreadI
 	}
 	GroupMemoryBarrierWithGroupSync();
 
+
+	#if PERMUTATION_DEBUG
+	if (bDebugEnabled)
+	{
+		float4 DebugColor = ColorRed;
+		if (GroupThreadID == 0)
+		{
+			AddQuadSS(Ctx, TileCoord * BIN_TILE_SIZE, TileCoord * BIN_TILE_SIZE + BIN_TILE_SIZE, DebugColor);
+			Print(Ctx, TEXT("TileCoord      :"), FontWhite); Print(Ctx, TileCoord, FontWhite); Newline(Ctx);
+			Print(Ctx, TEXT("TotalPrimCount :"), FontWhite); Print(Ctx, TotalPrimCount, FontWhite); Newline(Ctx);
+			Print(Ctx, TEXT("group_NumTiles :"), FontWhite); Print(Ctx, group_NumTiles, FontWhite); Newline(Ctx);
+			Print(Ctx, TEXT("group_MinZ     :"), FontWhite); Print(Ctx, UnpackDepth(group_MinZ), FontWhite); Newline(Ctx);
+			Print(Ctx, TEXT("group_MaxZ     :"), FontWhite); Print(Ctx, UnpackDepth(group_MaxZ), FontWhite); Newline(Ctx);
+		}
+	}
+	#endif
+
 	// 3. Copy PrimIDs to compacted memory
+	#if 0
 	{
 		uint CurrentWriteOffset = group_PrimWriteOffset;
 
@@ -874,6 +928,125 @@ void CompactionCS(uint DispatchThreadID : SV_DispatchThreadID, uint GroupThreadI
 			}
 		}
 	}
+	#elif 1
+	{
+		const uint NumInputTiles = min(group_NumTiles, 1024);
+
+		{
+			// 3.1 First process the LDS list of tiles
+			{
+				for (uint LDSIdx = 0; LDSIdx < NumInputTiles; ++LDSIdx)
+				{
+					const uint TileIdx = group_TilesToCompact[LDSIdx];
+
+					const uint TilePrimOffset = TileIdx * 1024;
+					const uint TilePrimCount = LoadVisTileData(InData, TileIdx, VT_PrimCount);
+
+					if (GroupThreadID < TilePrimCount)
+					{
+						const float Depth = UnpackDepth(InDepths[TilePrimOffset + GroupThreadID]);
+						const uint BinIndex = GetDepthBinIndex(Depth);
+						InterlockedAdd(s_BinCount[BinIndex], 1);
+
+						#if PERMUTATION_DEBUG
+						if (0 && GroupThreadID == 0)
+						{
+							Print(Ctx, TEXT("Depth0 :"), FontWhite); Print(Ctx, Depth, FontWhite);
+							Print(Ctx, TEXT(" - BinIndex : - "), FontWhite); Print(Ctx, BinIndex, FontWhite);
+							Print(Ctx, TEXT(" - BinMinZ : - "), FontWhite); Print(Ctx, UnpackDepth(group_MinZ), FontWhite);
+							Print(Ctx, TEXT(" - BinMaxZ : - "), FontWhite); Print(Ctx, UnpackDepth(group_MaxZ), FontWhite);
+							Newline(Ctx);
+						}
+						#endif
+					}
+				}
+			}
+
+			GroupMemoryBarrierWithGroupSync();
+
+			// 3.2 Prefix sum of bin count
+			if (GroupThreadID == 0)
+			{
+				uint GlobalOffset = 0;
+				for (uint It=0; It < COMPACTION_DEPTH_BUCKET;++It)
+				{
+					s_BinOffset[It] = GlobalOffset;
+					GlobalOffset += s_BinCount[It];
+				}
+
+				#if PERMUTATION_DEBUG
+				for (uint It2 = 0; It2 < COMPACTION_DEPTH_BUCKET; ++It2)
+				{					
+					if (s_BinCount[It2] > 0)
+						Print(Ctx, TEXT("x"), FontWhite);
+					else
+						Print(Ctx, TEXT("."), FontWhite);
+
+					if (It2 % 32 == 0)
+						Newline(Ctx);
+				}
+				Newline(Ctx);
+				#endif
+			}
+
+			GroupMemoryBarrierWithGroupSync();
+
+			// 3.3 Clear insertion counter
+			if (GroupThreadID < COMPACTION_DEPTH_BUCKET)
+			{
+				s_BinCount[GroupThreadID] = 0;
+			}
+
+			GroupMemoryBarrierWithGroupSync();
+
+			// 3.4 Insert primitive into bins
+			{
+				uint CurrentWriteOffset = group_PrimWriteOffset;
+				for (uint LDSIdx = 0; LDSIdx < NumInputTiles; ++LDSIdx)
+				{
+					const uint TileIdx = group_TilesToCompact[LDSIdx];
+
+					const uint TilePrimOffset = TileIdx * 1024;
+					const uint TilePrimCount = LoadVisTileData(InData, TileIdx, VT_PrimCount);
+
+					if (GroupThreadID < TilePrimCount)
+					{
+						const float Depth = UnpackDepth(InDepths[TilePrimOffset + GroupThreadID]);
+						const uint BinIndex = GetDepthBinIndex(Depth);
+						const uint GlobalOffset = s_BinOffset[BinIndex];
+						uint LocalOffset = 0;
+						InterlockedAdd(s_BinCount[BinIndex], 1, LocalOffset);
+						const uint WriteIndex = group_PrimWriteOffset + GlobalOffset + LocalOffset;
+						OutPrims[WriteIndex] = InPrims[TilePrimOffset + GroupThreadID];
+					}
+
+					CurrentWriteOffset += TilePrimCount;
+				}
+			}
+
+			// 3.5 Check any remaning tiles (Unlikely?)
+			//if (group_NumTiles > 1024)
+			//{
+			//	for (uint TileIdx = group_MaxLDSTileIdx + 1; TileIdx < NumTiles; ++TileIdx)
+			//	{
+			//		const uint TilePackedCoord = LoadVisTileData(InData, TileIdx, VT_Coord);
+			//		if (PackedCoord == TilePackedCoord)
+			//		{
+			//			const uint TilePrimOffset = TileIdx * 1024;
+			//			const uint TilePrimCount = LoadVisTileData(InData, TileIdx, VT_PrimCount);
+			//
+			//			if (GroupThreadID < TilePrimCount)
+			//			{
+			//				OutPrims[CurrentWriteOffset + GroupThreadID] = InPrims[TilePrimOffset + GroupThreadID];
+			//			}
+			//
+			//			CurrentWriteOffset += TilePrimCount;
+			//		}
+			//	}
+			//}
+		}
+	}
+	#endif
 }
 
 #endif // SHADER_RASTERCOMPUTE_COMPACTION
@@ -1172,16 +1345,6 @@ void RasterCS(
 		}
 		#endif
 
-		#if PERMUTATION_DEBUG
-		if (0 && GroupThread1D == 0)
-		{
-			const float4 TileColor = ColorLightGreen;
-			//const float4 TileColor = float4(ColorMapTurbo(s_WorkID / 2048.f),1);
-			//AddFilledQuadSS(RasterTileMin, RasterTileMax, Transparent(TileColor));
-			AddQuadSS(RasterTileMin, RasterTileMax, ColorYellow);
-		}		
-		#endif
-
 		// 1. Load all the segments and compute min/max bound (move this during the compaction)
 		{
 			LOOP
@@ -1454,7 +1617,7 @@ void RasterCS(
 
 						//const float4 LineColor = float4(ColorMapTurbo(GroupThread1D / float(RASTER_THREAD_COUNT)), 1);
 						//AddLineWS(P0, P1, LineColor);
-						if (GroupThread1D == View.GeneralPurposeTweak)
+						if (GroupThread1D == 0)
 							AddLineSS(SP0, SP1, ColorPurple);
 					}
 					#endif
@@ -1635,62 +1798,6 @@ void RasterCS(
 				{
 					const uint ValidSegments = ReadMask(GroupThread2D);//s_Mask[Thread_Coord.x][Thread_Coord.y];
 
-					
-					// Compute transmittance function along depth
-					#define USE_TRANSMITTANCE 0
-					#if USE_TRANSMITTANCE
-					float Coverages[RASTER_SEGMENT_COUNT];
-					float Pixel_InvDepthExtent = 0.010f;
-					{
-						// 1. Reset coverage & compute min/max depth of pixel
-						float Pixel_MinDepth = 1e8;
-						float Pixel_MaxDepth = 0;
-						{
-							for (uint SegmentIt = 0; SegmentIt < RASTER_SEGMENT_COUNT; ++SegmentIt)
-							{
-								Coverages[SegmentIt] = 0;
-
-								const uint SegmentBit = 1u << SegmentIt;
-								const bool bIsValid = (ValidSegments & SegmentBit) > 0;
-								const bool bIsSteep = (s_bDataOrder & SegmentBit) > 0;
-								if (bIsValid)
-								{
-									const uint J = bIsSteep ? Thread_Coord.y : Thread_Coord.x;
-									const FSampleData Sample = UnpackSampleData(s_Data[SegmentIt][J]);
-									Pixel_MinDepth = min(Pixel_MinDepth, Sample.Depth);
-									Pixel_MaxDepth = max(Pixel_MaxDepth, Sample.Depth);
-								}
-							}
-						}
-						Pixel_InvDepthExtent = max(0.010f, Pixel_MaxDepth - Pixel_MinDepth);
-
-						// 2. Compute coverage per bin
-						for (uint SegmentIt=0; SegmentIt < RASTER_SEGMENT_COUNT; ++SegmentIt)
-						{
-							const uint SegmentBit = 1u << SegmentIt;
-							const bool bIsValid = (ValidSegments & SegmentBit) > 0;
-							const bool bIsSteep = (s_bDataOrder & SegmentBit) > 0;
-							if (bIsValid)
-							{
-								const uint J = bIsSteep ? Thread_Coord.y : Thread_Coord.x;
-
-								const FSampleData Sample = UnpackSampleData(s_Data[SegmentIt][J]);
-								const uint DepthIndex = GetDepthBinIndex(Sample.Depth, Pixel_InvDepthExtent);
-
-								Coverages[DepthIndex] = min(1, Coverages[DepthIndex] + Sample.Coverage);
-							}
-						}
-
-						// 3. Compute coverage function (accumulated)
-						for (uint TransIt = 1; TransIt < RASTER_SEGMENT_COUNT; ++TransIt)
-						{
-							const float PrevCoverage = Coverages[TransIt - 1];
-							const float CurrCoverage = Coverages[TransIt];
-							Coverages[TransIt] = saturate(PrevCoverage + (1 - PrevCoverage) + CurrCoverage);
-						}
-					}
-					#endif
-
 					// Change this loop into to bit logic?
 					{
 						for (uint SegmentIt=0; SegmentIt < RASTER_SEGMENT_COUNT; ++SegmentIt)
@@ -1706,17 +1813,6 @@ void RasterCS(
 								const float3 Color = UnpackColorData(s_Color[SegmentIt][J]);
 								const float4 Velocity = UnpackVelocityData(s_Velocity[SegmentIt][J]);
 
-								// Accumulate sample contribution 
-								// TODO: needs to compute transmittance funciton, as seg ordered from front to back, but the raster pixels might be slightly different
-								#if USE_TRANSMITTANCE
-								const uint DepthIndex = GetDepthBinIndex(Sample.Depth, Pixel_InvDepthExtent);
-								const float Coverage  = Coverages[DepthIndex];
-								//const float Transmittance = Transmittances[DepthIndex];
-								//Thread_Coverage += Transmittance * Sample.Coverage;
-								//Thread_Color    += Transmittance * Sample.Coverage * Color;
-								Thread_Coverage = Coverage;
-								Thread_Color += Coverage * Color;
-								#elif 1
 								const float AccTransmittance = (1.f-Thread_Coverage);
 								Thread_Coverage += AccTransmittance * Sample.Coverage;
 								Thread_Color    += AccTransmittance * Sample.Coverage * Color;
@@ -1726,14 +1822,6 @@ void RasterCS(
 								{
 									Thread_Velocity = Velocity;
 								}
-
-								//Thread_Coverage = 0.5;
-								//Thread_Color    = 0;
-
-								#else
-								Thread_Coverage += 1;
-								Thread_Color = Color;
-								#endif
 							}
 						}
 					}
@@ -1822,7 +1910,7 @@ void RasterCS(
 			if (GroupThread1D == 0)
 			{
 				// Full/Partial coverage tiles
-				if (1)
+				if (0)
 				{
 					const bool bFullCover = s_OpaqueMask == 64;
 					if (bFullCover)	AddFilledQuadSS(RasterTileMin, RasterTileMax, Transparent(ColorLightRed));
diff --git a/Engine/Source/Runtime/Renderer/Private/HairStrands/HairStrandsForwardRaster.cpp b/Engine/Source/Runtime/Renderer/Private/HairStrands/HairStrandsForwardRaster.cpp
--- a/Engine/Source/Runtime/Renderer/Private/HairStrands/HairStrandsForwardRaster.cpp
+++ b/Engine/Source/Runtime/Renderer/Private/HairStrands/HairStrandsForwardRaster.cpp
@@ -127,7 +127,10 @@ class FHairStrandsForwardRasterBinningCS : public FGlobalShader
 		SHADER_PARAMETER_STRUCT_INCLUDE(FRasterComputeForwardCommonParameters, Common)
 		SHADER_PARAMETER_RDG_TEXTURE(Texture2DArray, VisTileBinningGrid)
 		SHADER_PARAMETER_RDG_BUFFER_UAV(RWBuffer, OutVisTilePrims)
+		SHADER_PARAMETER_RDG_BUFFER_UAV(RWBuffer, OutVisTilePrimDepths)
 		SHADER_PARAMETER_RDG_TEXTURE_UAV(RWTexture2DArray, OutVisTileBinningGrid)
+		SHADER_PARAMETER_RDG_TEXTURE_UAV(RWTexture2DArray, OutVisTileBinningGridMinZ)
+		SHADER_PARAMETER_RDG_TEXTURE_UAV(RWTexture2DArray, OutVisTileBinningGridMaxZ)
 		SHADER_PARAMETER_RDG_TEXTURE(Texture2D, VisTileDepthGrid)
 		SHADER_PARAMETER_RDG_BUFFER_UAV(RWBuffer, OutVisTileArgs)
 		SHADER_PARAMETER_RDG_BUFFER_UAV(RWByteAddressBuffer, OutVisTileData)
@@ -162,6 +165,7 @@ class FHairStrandsForwardRasterCompactionCS : public FGlobalShader
 		SHADER_PARAMETER_STRUCT_INCLUDE(FRasterComputeForwardCommonParameters, Common)
 		SHADER_PARAMETER_RDG_BUFFER_SRV(ByteAddressBuffer, InData)
 		SHADER_PARAMETER_RDG_BUFFER_SRV(Buffer, InPrims)
+		SHADER_PARAMETER_RDG_BUFFER_SRV(Buffer, InDepths)
 		SHADER_PARAMETER_RDG_BUFFER_SRV(Buffer, InArgs)
 		SHADER_PARAMETER_RDG_BUFFER_UAV(RWByteAddressBuffer, OutData)
 		SHADER_PARAMETER_RDG_BUFFER_UAV(RWBuffer, OutPrims)
@@ -452,7 +456,10 @@ void AddHairStrandsForwardRasterPass(
 		// Binned segments
 		FRDGBufferRef VisTilePrims;
 		FRDGBufferSRVRef VisTilePrimsSRV;
-		
+
+		FRDGBufferRef VisTilePrimDepths;
+		FRDGBufferSRVRef VisTilePrimDepthsSRV;
+
 		FRDGBufferRef VisTileData;
 		FRDGBufferSRVRef VisTileDataSRV;
 
@@ -463,6 +470,12 @@ void AddHairStrandsForwardRasterPass(
 		FRDGTextureRef VisTileBinningGrid;
 		FRDGTextureUAVRef VisTileBinningGridUAV;
 
+		// BinZ Min/Max
+		FRDGTextureRef VisTileBinningGridMinZ;
+		FRDGTextureRef VisTileBinningGridMaxZ;
+		FRDGTextureUAVRef VisTileBinningGridMinZUAV;
+		FRDGTextureUAVRef VisTileBinningGridMaxZUAV;
+
 		// Binned & compacted segments
 		FRDGBufferRef CompactedVisTilePrims;
 		FRDGBufferRef CompactedVisTileData;
@@ -476,6 +489,9 @@ void AddHairStrandsForwardRasterPass(
 		BinData.VisTilePrims = GraphBuilder.CreateBuffer(FRDGBufferDesc::CreateBufferDesc(sizeof(uint32), BinMaxTiles * 1024), TEXT("Hair.VisTilePrims"));
 		BinData.VisTilePrimsSRV = GraphBuilder.CreateSRV(BinData.VisTilePrims, PF_R32_UINT);
 
+		BinData.VisTilePrimDepths = GraphBuilder.CreateBuffer(FRDGBufferDesc::CreateBufferDesc(sizeof(uint32), BinMaxTiles * 1024), TEXT("Hair.VisTilePrimDepths"));
+		BinData.VisTilePrimDepthsSRV = GraphBuilder.CreateSRV(BinData.VisTilePrimDepths, PF_R32_UINT);
+
 		const uint32 EntryCount = 5; // See VT_XXX in .usf: PrimOffset / PrimCount / Coord / MaxIndex / MinMaxDepth
 		FRDGBufferDesc DescTileData = FRDGBufferDesc::CreateStructuredDesc(sizeof(uint32), BinMaxTiles * EntryCount * 2);
 		DescTileData.Usage = EBufferUsageFlags(DescTileData.Usage | BUF_ByteAddressBuffer);
@@ -490,6 +506,12 @@ void AddHairStrandsForwardRasterPass(
 		BinData.VisTileBinningGrid = GraphBuilder.CreateTexture(FRDGTextureDesc::Create2DArray(BinTileRes, PF_R32_UINT, FClearValueBinding::None, TexCreate_UAV | TexCreate_ShaderResource, NumBinners * 3), TEXT("Hair.VisTileBinningGrid"));
 		BinData.VisTileBinningGridUAV = GraphBuilder.CreateUAV(BinData.VisTileBinningGrid);
 
+		// BinZ Min/Max
+		BinData.VisTileBinningGridMinZ = GraphBuilder.CreateTexture(FRDGTextureDesc::Create2DArray(BinTileRes, PF_R32_UINT, FClearValueBinding::None, TexCreate_UAV | TexCreate_ShaderResource, NumBinners * 3), TEXT("Hair.VisTileBinningGridMinZ"));
+		BinData.VisTileBinningGridMaxZ = GraphBuilder.CreateTexture(FRDGTextureDesc::Create2DArray(BinTileRes, PF_R32_UINT, FClearValueBinding::None, TexCreate_UAV | TexCreate_ShaderResource, NumBinners * 3), TEXT("Hair.VisTileBinningGridMaxZ"));
+		BinData.VisTileBinningGridMinZUAV = GraphBuilder.CreateUAV(BinData.VisTileBinningGridMinZ);
+		BinData.VisTileBinningGridMaxZUAV = GraphBuilder.CreateUAV(BinData.VisTileBinningGridMaxZ);
+
 		BinData.CompactedVisTilePrims = GraphBuilder.CreateBuffer(FRDGBufferDesc::CreateBufferDesc(sizeof(uint32), BinMaxTiles * 1024), TEXT("Hair.CompactedVisTilePrims"));
 		BinData.CompactedVisTileData = GraphBuilder.CreateBuffer(DescTileData, TEXT("Hair.CompactedVisTileData"));
 		BinData.CompactedVisTileArgs = GraphBuilder.CreateBuffer(VisTileArgsDesc, TEXT("Hair.CompactedVisTileArgs"));
@@ -549,6 +571,12 @@ void AddHairStrandsForwardRasterPass(
 
 	// Reset buffers
 	{
+		// BinZ Min/Max
+		uint32 IndexGridClearMinZ[4] = { 0xFFFFFFFF,0xFFFFFFFF,0xFFFFFFFF,0xFFFFFFFF };
+		uint32 IndexGridClearMaxZ[4] = { 0x0,0x0,0x0,0x0 };
+		AddClearUAVPass(GraphBuilder, BinData.VisTileBinningGridMinZUAV, IndexGridClearMinZ);
+		AddClearUAVPass(GraphBuilder, BinData.VisTileBinningGridMaxZUAV, IndexGridClearMaxZ);
+
 		uint32 IndexGridClearValues[4] = { 0x0,0x0,0x0,0x0 };
 		AddClearUAVPass(GraphBuilder, BinData.VisTileBinningGridUAV, IndexGridClearValues);
 		AddClearUAVPass(GraphBuilder, BinData.VisTileArgsUAV, 0u);
@@ -561,7 +589,10 @@ void AddHairStrandsForwardRasterPass(
 		Parameters->Common = Common;
 		Parameters->ViewUniformBuffer = ViewUniformShaderParameters;
 		Parameters->OutVisTilePrims = GraphBuilder.CreateUAV(BinData.VisTilePrims, PF_R32_UINT);
+		Parameters->OutVisTilePrimDepths = GraphBuilder.CreateUAV(BinData.VisTilePrimDepths, PF_R32_UINT);
 		Parameters->OutVisTileBinningGrid = BinData.VisTileBinningGridUAV;
+		Parameters->OutVisTileBinningGridMinZ = BinData.VisTileBinningGridMinZUAV;
+		Parameters->OutVisTileBinningGridMaxZ = BinData.VisTileBinningGridMaxZUAV;
 		Parameters->VisTileDepthGrid = VisTileDepthGrid;
 		Parameters->OutVisTileArgs = BinData.VisTileArgsUAV;
 		Parameters->OutVisTileData = GraphBuilder.CreateUAV(BinData.VisTileData);
@@ -586,6 +617,7 @@ void AddHairStrandsForwardRasterPass(
 		Parameters->ViewUniformBuffer = ViewUniformShaderParameters;
 		Parameters->InData  = BinData.VisTileDataSRV;
 		Parameters->InPrims = BinData.VisTilePrimsSRV;
+		Parameters->InDepths= BinData.VisTilePrimDepthsSRV;
 		Parameters->InArgs  = BinData.VisTileArgsSRV;
 		Parameters->OutData = GraphBuilder.CreateUAV(BinData.CompactedVisTileData);
 		Parameters->OutPrims= GraphBuilder.CreateUAV(BinData.CompactedVisTilePrims, PF_R32_UINT);
