Process tile from front to back to ensure correct composition.

Tile data projected into the same screen tile, are now processed by the same raster processor starting from front to back.

diff --git a/Engine/Shaders/Private/HairStrands/HairStrandsForwardRaster.usf b/Engine/Shaders/Private/HairStrands/HairStrandsForwardRaster.usf
--- a/Engine/Shaders/Private/HairStrands/HairStrandsForwardRaster.usf
+++ b/Engine/Shaders/Private/HairStrands/HairStrandsForwardRaster.usf
@@ -75,6 +75,16 @@ float2 UnpackDepthRange(uint In)
 	return UnpackRG16(In);
 }
 
+uint PackWork(uint InTileIndex, uint InTileCount)
+{
+	return InTileIndex | (InTileCount << 16);
+}
+
+uint2 UnpackWork(uint In)
+{
+	return uint2(In & 0xFFFF, In >> 16);
+}
+
 ///////////////////////////////////////////////////////////////////////////
 // Tile Helpers
 
@@ -754,6 +764,9 @@ Buffer<uint> 						InArgs;
 RWByteAddressBuffer					OutData;
 RWBuffer<uint> 						OutPrims;
 RWBuffer<uint> 						OutArgs;
+RWStructuredBuffer<uint> 			OutWork; // Offset & Count
+RWStructuredBuffer<uint> 			OutDataCount;
+RWStructuredBuffer<uint> 			OutWorkCount;
 
 groupshared uint group_TotalPrimCount;
 groupshared uint group_PrimWriteOffset;
@@ -803,7 +816,7 @@ void CompactionCS(uint DispatchThreadID : SV_DispatchThreadID, uint GroupThreadI
 
 	#if PERMUTATION_DEBUG
 	const uint2 TileCoord = UnpackVisTileCoord(PackedCoord);
-	const bool bDebugEnabled = all(TileCoord == uint2(ShaderPrintData.CursorCoord / float(BIN_TILE_SIZE)));
+	const bool bDebugEnabled = false; //all(TileCoord == uint2(ShaderPrintData.CursorCoord / float(BIN_TILE_SIZE)));
 	FShaderPrintContext Ctx = InitShaderPrintContext(bDebugEnabled, uint2(750, 50));
 	#endif
 
@@ -854,6 +867,10 @@ void CompactionCS(uint DispatchThreadID : SV_DispatchThreadID, uint GroupThreadI
 		uint FirstCompactedTile;
 		InterlockedAdd(OutArgs[0], NumTilesToAllocate, FirstCompactedTile);
 
+		uint WorkIndex;
+		InterlockedAdd(OutWorkCount[0], 1, WorkIndex);
+		OutWork[WorkIndex] = PackWork(FirstCompactedTile, NumTilesToAllocate);
+
 		group_PrimWriteOffset = FirstCompactedTile * 1024;
 
 		// Initialize new tiles
@@ -886,49 +903,6 @@ void CompactionCS(uint DispatchThreadID : SV_DispatchThreadID, uint GroupThreadI
 	#endif
 
 	// 3. Copy PrimIDs to compacted memory
-	#if 0
-	{
-		uint CurrentWriteOffset = group_PrimWriteOffset;
-
-		// 3.1 First process the LDS list of tiles
-		const uint NumInputTiles = min(group_NumTiles, 1024);
-		for (uint LDSIdx = 0; LDSIdx < NumInputTiles; ++LDSIdx)
-		{
-			const uint TileIdx = group_TilesToCompact[LDSIdx];
-
-			const uint TilePrimOffset = TileIdx * 1024;
-			const uint TilePrimCount = LoadVisTileData(InData, TileIdx, VT_PrimCount);
-
-			if (GroupThreadID < TilePrimCount)
-			{
-				OutPrims[CurrentWriteOffset + GroupThreadID] = InPrims[TilePrimOffset + GroupThreadID];
-			}
-
-			CurrentWriteOffset += TilePrimCount;
-		}
-
-		// 3.2 Check any remaning tiles
-		if (group_NumTiles > 1024)
-		{
-			for (uint TileIdx = group_MaxLDSTileIdx + 1; TileIdx < NumTiles; ++TileIdx)
-			{
-				const uint TilePackedCoord = LoadVisTileData(InData, TileIdx, VT_Coord);
-				if (PackedCoord == TilePackedCoord)
-				{
-					const uint TilePrimOffset = TileIdx * 1024;
-					const uint TilePrimCount = LoadVisTileData(InData, TileIdx, VT_PrimCount);
-
-					if (GroupThreadID < TilePrimCount)
-					{
-						OutPrims[CurrentWriteOffset + GroupThreadID] = InPrims[TilePrimOffset + GroupThreadID];
-					}
-
-					CurrentWriteOffset += TilePrimCount;
-				}
-			}
-		}
-	}
-	#elif 1
 	{
 		const uint NumInputTiles = min(group_NumTiles, 1024);
 
@@ -982,7 +956,7 @@ void CompactionCS(uint DispatchThreadID : SV_DispatchThreadID, uint GroupThreadI
 					else
 						Print(Ctx, TEXT("."), FontWhite);
 
-					if (It2 % 32 == 0)
+					if (It2 != 0 && (It2 % 32) == 0)
 						Newline(Ctx);
 				}
 				Newline(Ctx);
@@ -1046,7 +1020,6 @@ void CompactionCS(uint DispatchThreadID : SV_DispatchThreadID, uint GroupThreadI
 			//}
 		}
 	}
-	#endif
 }
 
 #endif // SHADER_RASTERCOMPUTE_COMPACTION
@@ -1067,6 +1040,8 @@ void CompactionCS(uint DispatchThreadID : SV_DispatchThreadID, uint GroupThreadI
 #if SHADER_RASTERCOMPUTE_RASTER
 
 Buffer<uint> 				VisTilePrims;
+StructuredBuffer<uint>		VisTileWork;
+StructuredBuffer<uint>		VisTileWorkCount;
 Buffer<uint> 				VisTileArgs;
 ByteAddressBuffer			VisTileData;
 RWTexture2D<float4>			OutSceneColorTexture;
@@ -1100,6 +1075,9 @@ groupshared uint2  s_Velocity[RASTER_SEGMENT_COUNT][8];
 groupshared uint   s_bDataOrder;
 groupshared uint   s_WorkID;
 
+groupshared uint   s_BinTileOffset;
+groupshared uint   s_BinTileCount;
+
 #if PERMUTATION_DEBUG
 groupshared float s_Coverage[RASTER_THREAD_COUNT];
 #endif
@@ -1252,18 +1230,67 @@ void RasterCS(
 	uint  GroupID          : SV_GroupID)		/* Rasterizer ID */
 {
 	ResolvedView = ResolveView();
+	const uint FetchWorkCount = VisTileWorkCount[0];
 	const uint BinTileNum = VisTileArgs[0];
 
+	s_BinTileOffset = 0;
+	s_BinTileCount  = 0;
+
+	// These are global Color/Coverage for the final pixel handled by this thread
+	float3 Thread_Color = 0;
+	float  Thread_Coverage = 0;
+	uint   Thread_Complete = 0;
+	float4 Thread_Velocity = INVALID_VELOCITY;
+	uint   Thread_LoopCountToFullCoverage = 0;
+
 	LOOP
 	for (uint WorkIndex = 0; WorkIndex < MAX_WORK_COUNT; WorkIndex++)
 	{
-		// 0. Fetch work item
+		// 0.1 Fetch work item
 		if (GroupThread1D == 0)
 		{
 			InterlockedAdd(RWWorkCounter[0], 1, s_WorkID);
 		}
 
-		// 0. Reset all LDS variables
+		if (GroupThread1D == 0)
+		{
+			const uint FetchWorkIndex = s_WorkID / 16; // 1 x Bin32x32 -> 16 x Raster8x8
+
+			uint2 Work = 0;
+			if (FetchWorkIndex < FetchWorkCount)
+			{
+				Work = UnpackWork(VisTileWork[FetchWorkIndex]);
+			}
+
+			s_BinTileOffset = Work.x;
+			s_BinTileCount  = Work.y;
+		}
+
+		GroupMemoryBarrierWithGroupSync();
+
+		// 0.3 If we start a new screen tile, clear out final output
+		{
+			Thread_Color = 0;
+			Thread_Coverage = 0;
+			Thread_Complete = 0;
+			Thread_Velocity = INVALID_VELOCITY;
+			Thread_LoopCountToFullCoverage = 0;
+		}
+
+		// Early out if we are done
+		{
+			const uint FetchWorkIndex = s_WorkID / 16; // 1 x Bin32x32 -> 16 x Raster8x8
+			if (FetchWorkIndex >= FetchWorkCount)
+			{
+				return;
+			}
+		}
+
+		const uint FrontToBackCount = min(s_BinTileCount, 64);
+		for (uint FrontToBackIndex = 0; FrontToBackIndex < FrontToBackCount; ++FrontToBackIndex)
+		{
+
+		// 0.2 Reset all LDS variables
 		{
 			if (GroupThread1D == 0)
 			{
@@ -1302,8 +1329,8 @@ void RasterCS(
 		}
 		GroupMemoryBarrierWithGroupSync();
 
-		// 0. Load all tile data
-		const uint BinTileIndex = s_WorkID / 16; // 1 x Bin32x32 -> 16 x Raster8x8
+		// 0.4 Early out when running out of valid tiles
+		const uint BinTileIndex = s_BinTileOffset + FrontToBackIndex;
 		const bool bTileValid = (BinTileIndex < BinTileNum);
 		if (!bTileValid)
 		{
@@ -1327,7 +1354,7 @@ void RasterCS(
 		// For debug only
 		#if PERMUTATION_DEBUG
 		const bool bDebugEnableAll	= all(RasterTileCoord == uint2(ShaderPrintData.CursorCoord / float(RASTER_TILE_SIZE)));
-		const bool bDebugEnable		= all(RasterTileCoord == uint2(ShaderPrintData.CursorCoord / float(RASTER_TILE_SIZE))) && GroupThread1D == View.GeneralPurposeTweak;
+		const bool bDebugEnable		= all(RasterTileCoord == uint2(ShaderPrintData.CursorCoord / float(RASTER_TILE_SIZE))) && GroupThread1D == 0;
 		FShaderPrintContext Ctx		= InitShaderPrintContext(bDebugEnable, uint2(350, 50));
 		const FFontColor FontLegend = FontWhite;
 		const FFontColor FontValue  = FontOrange;
@@ -1339,6 +1366,10 @@ void RasterCS(
 		#if PERMUTATION_DEBUG
 		if (bDebugEnable)
 		{
+			Print(Ctx, TEXT("Work - Index :"), FontLegend); Print(Ctx, FrontToBackIndex, FontValue); Newline(Ctx);
+			Print(Ctx, TEXT("Work - Offset:"), FontLegend); Print(Ctx, s_BinTileOffset, FontValue); Newline(Ctx);
+			Print(Ctx, TEXT("Work - Count :"), FontLegend); Print(Ctx, s_BinTileCount, FontValue); Newline(Ctx);
+
 			AddFilledQuadSS(BinTileMin, BinTileMax, Transparent(ColorLightGreen));
 			//AddFilledQuadSS(RasterTileMin, RasterTileMax, Transparent(ColorLightGreen));
 			AddQuadSS(RasterTileMin, RasterTileMax, Transparent(ColorYellow));
@@ -1502,6 +1533,11 @@ void RasterCS(
 			Print(Ctx, TEXT("Rad.at Depth1:"), FontLegend); Print(Ctx, RadiusAtDepth1, FontValue); Newline(Ctx);
 			Newline(Ctx);
 
+			Print(Ctx, TEXT("Work - Index :"), FontLegend); Print(Ctx, WorkIndex, FontValue, 2, 0);
+			Print(Ctx, TEXT(" - Offset:"), FontLegend); Print(Ctx, s_BinTileOffset, FontValue, 5, 0);
+			Print(Ctx, TEXT(" - Count :"), FontLegend); Print(Ctx, s_BinTileCount, FontValue, 3, 0);
+			Newline(Ctx);
+
 			//Print(Ctx, TEXT("Bin ID     :"), FontLegend); Print(Ctx, BinTileIndex, FontValue); Newline(Ctx);
 			//Print(Ctx, TEXT("Raster ID  :"), FontLegend); Print(Ctx, QuadrantIndex, FontValue); Newline(Ctx);
 			//Newline(Ctx);
@@ -1516,13 +1552,6 @@ void RasterCS(
 
 		// 4. Raster segments
 		{
-			// These are global Color/Coverage for the final pixel handled by this thread
-			float3 Thread_Color    = 0;
-			float  Thread_Coverage = 0;
-			uint   Thread_Complete = 0;
-			float4 Thread_Velocity = INVALID_VELOCITY;
-			uint   Thread_LoopCountToFullCoverage=0;
-
 			// 4.1 Loop over all segment within the tile, and rastize 32 of them at each loop
 			LOOP
 			for (uint LoopIndex = 0; LoopIndex < LoopCount32; LoopIndex++)
@@ -1845,12 +1874,15 @@ void RasterCS(
 			} // for (uint LoopIndex =...) ...
 
 			#if PERMUTATION_DEBUG
-			s_Coverage[GroupThread1D] = Thread_Coverage;
-			GroupMemoryBarrierWithGroupSync(); // For sync s_Coverage
+			s_Coverage[GroupThread1D] = Thread_Coverage; // For sync s_Coverage
 			#endif
 
+			GroupMemoryBarrierWithGroupSync();
+
 			// 4.2. Write final color
-			if (Thread_Coverage>0)
+			const bool bLastBinTile = FrontToBackIndex == (s_BinTileCount - 1);
+			const bool bWriteOut = Thread_Coverage > 0 && (bLastBinTile /* || raster tile is fully covered */);
+			if (bWriteOut)
 			{
 				const uint2 PixelCoord = GroupThread2D + RasterTileMin;
 				const bool bOutput = true;//View.GeneralPurposeTweak > 1;
@@ -1924,6 +1956,8 @@ void RasterCS(
 
 		} // 4. Raster
 
+		} // for (... FrontToBackIndex ...)
+
 	}// Work item
 }
 
diff --git a/Engine/Source/Runtime/Renderer/Private/HairStrands/HairStrandsForwardRaster.cpp b/Engine/Source/Runtime/Renderer/Private/HairStrands/HairStrandsForwardRaster.cpp
--- a/Engine/Source/Runtime/Renderer/Private/HairStrands/HairStrandsForwardRaster.cpp
+++ b/Engine/Source/Runtime/Renderer/Private/HairStrands/HairStrandsForwardRaster.cpp
@@ -170,6 +170,9 @@ class FHairStrandsForwardRasterCompactionCS : public FGlobalShader
 		SHADER_PARAMETER_RDG_BUFFER_UAV(RWByteAddressBuffer, OutData)
 		SHADER_PARAMETER_RDG_BUFFER_UAV(RWBuffer, OutPrims)
 		SHADER_PARAMETER_RDG_BUFFER_UAV(RWBuffer, OutArgs)
+		SHADER_PARAMETER_RDG_BUFFER_UAV(RWBuffer, OutWork)
+		SHADER_PARAMETER_RDG_BUFFER_UAV(RWBuffer, OutDataCount)
+		SHADER_PARAMETER_RDG_BUFFER_UAV(RWBuffer, OutWorkCount)
 		SHADER_PARAMETER_STRUCT_REF(FViewUniformShaderParameters, ViewUniformBuffer)
 		SHADER_PARAMETER_STRUCT_INCLUDE(ShaderPrint::FShaderParameters, ShaderPrintParameters)
 	END_SHADER_PARAMETER_STRUCT()
@@ -205,6 +208,8 @@ class FHairStrandsForwardRasterRasterizeCS : public FGlobalShader
 		SHADER_PARAMETER_RDG_TEXTURE(Texture2D, SceneDepthTexture)
 		SHADER_PARAMETER_RDG_BUFFER_SRV(Buffer, VisTilePrims)
 		SHADER_PARAMETER_RDG_BUFFER_SRV(Buffer, VisTileArgs)
+		SHADER_PARAMETER_RDG_BUFFER_SRV(StructuredBuffer, VisTileWork)
+		SHADER_PARAMETER_RDG_BUFFER_SRV(StructuredBuffer, VisTileWorkCount)
 		SHADER_PARAMETER_RDG_BUFFER_SRV(ByteAddressBuffer, VisTileData)
 		SHADER_PARAMETER_RDG_TEXTURE_UAV(RWTexture2D, OutHairCountTexture_ForDebug)
 		SHADER_PARAMETER_RDG_TEXTURE_UAV(RWTexture2D, OutHairPixelCountPerTile_ForDebug)
@@ -480,11 +485,14 @@ void AddHairStrandsForwardRasterPass(
 		FRDGBufferRef CompactedVisTilePrims;
 		FRDGBufferRef CompactedVisTileData;
 		FRDGBufferRef CompactedVisTileArgs;
+		FRDGBufferRef CompactedVisTileWork;
+		FRDGBufferRef CompactedVisTileDataCount;
+		FRDGBufferRef CompactedVisTileWorkCount;
 	};
 
 	FBinningData BinData;
 	{
-		const uint32 BinMaxTiles = FMath::Min(FMath::Max(GHairVisibilityComputeRaster_MaxTiles, 1024), 262144);
+		const uint32 BinMaxTiles = FMath::Min(FMath::Max(GHairVisibilityComputeRaster_MaxTiles, 1024), 262144); //?
 
 		BinData.VisTilePrims = GraphBuilder.CreateBuffer(FRDGBufferDesc::CreateBufferDesc(sizeof(uint32), BinMaxTiles * 1024), TEXT("Hair.VisTilePrims"));
 		BinData.VisTilePrimsSRV = GraphBuilder.CreateSRV(BinData.VisTilePrims, PF_R32_UINT);
@@ -515,6 +523,9 @@ void AddHairStrandsForwardRasterPass(
 		BinData.CompactedVisTilePrims = GraphBuilder.CreateBuffer(FRDGBufferDesc::CreateBufferDesc(sizeof(uint32), BinMaxTiles * 1024), TEXT("Hair.CompactedVisTilePrims"));
 		BinData.CompactedVisTileData = GraphBuilder.CreateBuffer(DescTileData, TEXT("Hair.CompactedVisTileData"));
 		BinData.CompactedVisTileArgs = GraphBuilder.CreateBuffer(VisTileArgsDesc, TEXT("Hair.CompactedVisTileArgs"));
+		BinData.CompactedVisTileWork		= GraphBuilder.CreateBuffer(FRDGBufferDesc::CreateStructuredDesc(sizeof(uint32), BinMaxTiles * 1024), TEXT("Hair.CompactedVisTileWork"));
+		BinData.CompactedVisTileDataCount	= GraphBuilder.CreateBuffer(FRDGBufferDesc::CreateStructuredDesc(sizeof(uint32), 1), TEXT("Hair.CompactedVisTileDataCount"));
+		BinData.CompactedVisTileWorkCount	= GraphBuilder.CreateBuffer(FRDGBufferDesc::CreateStructuredDesc(sizeof(uint32), 1), TEXT("Hair.CompactedVisTileWorkCount"));
 	}
 
 	// Create and set the uniform buffer
@@ -581,6 +592,8 @@ void AddHairStrandsForwardRasterPass(
 		AddClearUAVPass(GraphBuilder, BinData.VisTileBinningGridUAV, IndexGridClearValues);
 		AddClearUAVPass(GraphBuilder, BinData.VisTileArgsUAV, 0u);
 		AddClearUAVPass(GraphBuilder, GraphBuilder.CreateUAV(BinData.CompactedVisTileArgs, PF_R32_UINT), 0u);
+		AddClearUAVPass(GraphBuilder, GraphBuilder.CreateUAV(BinData.CompactedVisTileDataCount, PF_R32_UINT), 0u);
+		AddClearUAVPass(GraphBuilder, GraphBuilder.CreateUAV(BinData.CompactedVisTileWorkCount, PF_R32_UINT), 0u);
 	}
 
 	// Binning pass
@@ -613,15 +626,19 @@ void AddHairStrandsForwardRasterPass(
 	// Compaction
 	{
 		FHairStrandsForwardRasterCompactionCS::FParameters* Parameters = GraphBuilder.AllocParameters<FHairStrandsForwardRasterCompactionCS::FParameters>();
-		Parameters->Common = Common;
-		Parameters->ViewUniformBuffer = ViewUniformShaderParameters;
-		Parameters->InData  = BinData.VisTileDataSRV;
-		Parameters->InPrims = BinData.VisTilePrimsSRV;
-		Parameters->InDepths= BinData.VisTilePrimDepthsSRV;
-		Parameters->InArgs  = BinData.VisTileArgsSRV;
-		Parameters->OutData = GraphBuilder.CreateUAV(BinData.CompactedVisTileData);
-		Parameters->OutPrims= GraphBuilder.CreateUAV(BinData.CompactedVisTilePrims, PF_R32_UINT);
-		Parameters->OutArgs = GraphBuilder.CreateUAV(BinData.CompactedVisTileArgs, PF_R32_UINT);
+		Parameters->Common				= Common;
+		Parameters->ViewUniformBuffer	= ViewUniformShaderParameters;
+		Parameters->InData				= BinData.VisTileDataSRV;
+		Parameters->InPrims				= BinData.VisTilePrimsSRV;
+		Parameters->InDepths			= BinData.VisTilePrimDepthsSRV;
+		Parameters->InArgs				= BinData.VisTileArgsSRV;
+		Parameters->OutData				= GraphBuilder.CreateUAV(BinData.CompactedVisTileData);
+		Parameters->OutPrims			= GraphBuilder.CreateUAV(BinData.CompactedVisTilePrims, PF_R32_UINT);
+		Parameters->OutArgs 			= GraphBuilder.CreateUAV(BinData.CompactedVisTileArgs, PF_R32_UINT);
+		Parameters->OutWork				= GraphBuilder.CreateUAV(BinData.CompactedVisTileWork, PF_R32_UINT);
+		Parameters->OutDataCount		= GraphBuilder.CreateUAV(BinData.CompactedVisTileDataCount, PF_R32_UINT);
+		Parameters->OutWorkCount		= GraphBuilder.CreateUAV(BinData.CompactedVisTileWorkCount, PF_R32_UINT);
+
 		ShaderPrint::SetParameters(GraphBuilder, ViewInfo.ShaderPrintData, Parameters->ShaderPrintParameters);
 
 		FHairStrandsForwardRasterCompactionCS::FPermutationDomain PermutationVector;
@@ -660,6 +677,8 @@ void AddHairStrandsForwardRasterPass(
 		Parameters->SampleVelocityBuffer = InData.ControlPointVelocitySRV;
 		Parameters->VisTilePrims = GraphBuilder.CreateSRV(BinData.CompactedVisTilePrims, PF_R32_UINT);
 		Parameters->VisTileArgs = GraphBuilder.CreateSRV(BinData.CompactedVisTileArgs, PF_R32_UINT);
+		Parameters->VisTileWork = GraphBuilder.CreateSRV(BinData.CompactedVisTileWork, PF_R32_UINT);
+		Parameters->VisTileWorkCount = GraphBuilder.CreateSRV(BinData.CompactedVisTileWorkCount, PF_R32_UINT);
 		Parameters->VisTileData = GraphBuilder.CreateSRV(BinData.CompactedVisTileData);
 		Parameters->OutSceneColorTexture = GraphBuilder.CreateUAV(SceneColorTexture);
 		Parameters->OutSceneVelocityTexture = GraphBuilder.CreateUAV(SceneVelocityTexture);
