Fix incorrect opaque depth testing against hair segments.

diff --git a/Engine/Shaders/Private/HairStrands/HairStrandsForwardRaster.usf b/Engine/Shaders/Private/HairStrands/HairStrandsForwardRaster.usf
--- a/Engine/Shaders/Private/HairStrands/HairStrandsForwardRaster.usf
+++ b/Engine/Shaders/Private/HairStrands/HairStrandsForwardRaster.usf
@@ -1057,6 +1057,8 @@ int2						SampleLightingViewportResolution;
 Texture2D<float4>			SampleLightingTexture;
 Buffer<float4>				SampleVelocityBuffer;
 
+Texture2D<float>			SceneDepthTexture;
+
 // Depth|ID
 groupshared uint2  s_Segments[1024];
 groupshared uint2  s_Segments_Sorted[1024];
@@ -1748,67 +1750,54 @@ void RasterCS(
 							{
 								const float Alpha = ComputeLerpAlpha(Coords, SP0.xy, SP1.xy, SegmentLenSqRcp);
 								const float Depth = lerp(SP0.z, SP1.z, Alpha);
+								const float OpaqueDepth = SceneDepthTexture.Load(uint3(Coords, 0));
+								if (Depth > OpaqueDepth)
+								{
+									const float AlphaColor = lerp(Alpha0, Alpha1, Alpha);
+									const float3 Color0 = LoadSampleColor(PrimID);
+									const float3 Color1 = LoadSampleColor(PrimID+1);
+									const float3 Color = lerp(Color0, Color1, AlphaColor);
 
-								const float AlphaColor = lerp(Alpha0, Alpha1, Alpha);
-								const float3 Color0 = LoadSampleColor(PrimID);
-								const float3 Color1 = LoadSampleColor(PrimID+1);
-								const float3 Color = lerp(Color0, Color1, AlphaColor);
-
-								const float4 Velocity0 = LoadSampleVelocity(PrimID);
-								const float4 Velocity1 = LoadSampleVelocity(PrimID + 1);
-								const float4 Velocity = lerp(Velocity0, Velocity1, AlphaColor);
+									const float4 Velocity0 = LoadSampleVelocity(PrimID);
+									const float4 Velocity1 = LoadSampleVelocity(PrimID+1);
+									const float4 Velocity = lerp(Velocity0, Velocity1, AlphaColor);
 
-								const float HairRadius = lerp(Rad0, Rad1, AlphaColor);
 
-								// Fill in sample data
-								Sample.Depth = Depth;
-								Sample.Coord = IntraTileCoord;
-								Sample.Coverage = 0.5f; //TODO;
+									// Fill in sample data
+									Sample.Depth = Depth;
+									Sample.Coord = IntraTileCoord;
 
-								// Compute actual coverage
-								//{
+									// Compute coverage
 									// Minimal radius to snap the strand to a sample/pixel center (to avoid aliasing)
 									const bool bIsOrthoView = false;
 									const float SceneDistance = ConvertFromDeviceZ(Depth);
 									const float MinHairRadius = bIsOrthoView ? RadiusAtDepth1 : (SceneDistance * RadiusAtDepth1);
+									const float HairRadius = lerp(Rad0, Rad1, AlphaColor);
 									Sample.Coverage = saturate(HairRadius / MinHairRadius);
-								//}
-
-
 
-								// Write data
-								WriteMask(IntraTileCoord, 1u << GroupThread1D);
-								//InterlockedOr(s_Mask[IntraTileCoord.x][IntraTileCoord.y], 1u<<GroupThread1D);
-
-								const uint JCoord = bIsSteep ? IntraTileCoord.y : IntraTileCoord.x;
-								s_Color[GroupThread1D][JCoord] = PackColorData(Color);
-								s_Data[GroupThread1D][JCoord]  = PackSampleData(Sample);
-								s_Velocity[GroupThread1D][JCoord] = PackVelocityData(Velocity);
-
-
-								// DEBUG (Write coord)
-								#if 0 && PERMUTATION_DEBUG
-								if (bDebugEnable)
-								{
-									Print(Ctx, J, FontYellow, 2, 0);
-									//Print(Ctx, IntraTileCoord, FontWhite,2,0);
-									//Print(Ctx, s_Mask[IntraTileCoord.x][IntraTileCoord.y], FontWhite);
-									//PrintPixelMask(Ctx, ReadMask(IntraTileCoord)/*s_Mask[IntraTileCoord.x][IntraTileCoord.y]*/);
-									//Print(Ctx, TEXT("Color0: "), FontWhite); Print(Ctx, Color0); Newline(Ctx);
-									//Print(Ctx, TEXT("Color1: "), FontWhite); Print(Ctx, Color1); Newline(Ctx);
-									//Print(Ctx, TEXT("Color :  "), FontWhite); Print(Ctx, Color);
-									Print(Ctx, Alpha, FontWhite);
-									//Print(Ctx, Alpha0, FontRed);
-									//Print(Ctx, Alpha1, FontGreen);
-									Print(Ctx, AlphaColor, FontBlue);
-									//Print(Ctx, Color, FontOrange);
-									//Print(Ctx, IntraTileCoord, FontOrange);
-									Print(Ctx, HairRadius, FontRed);
-									Print(Ctx, MinHairRadius, FontGreen);
-									Print(Ctx, Sample.Coverage, FontBlue);
-									Newline(Ctx);
+									// Write data
+									WriteMask(IntraTileCoord, 1u << GroupThread1D);
+
+									const uint JCoord = bIsSteep ? IntraTileCoord.y : IntraTileCoord.x;
+									s_Color[GroupThread1D][JCoord] = PackColorData(Color);
+									s_Data[GroupThread1D][JCoord]  = PackSampleData(Sample);
+									s_Velocity[GroupThread1D][JCoord] = PackVelocityData(Velocity);
+
+
+									// DEBUG (Write coord)
+									#if 0 && PERMUTATION_DEBUG
+									if (bDebugEnable)
+									{
+										Print(Ctx, J, FontYellow, 2, 0);
+										Print(Ctx, Alpha, FontWhite);
+										Print(Ctx, AlphaColor, FontBlue);
+										Print(Ctx, HairRadius, FontRed);
+										Print(Ctx, MinHairRadius, FontGreen);
+										Print(Ctx, Sample.Coverage, FontBlue);
+										Newline(Ctx);
+									}
+									#endif
 								}
-								#endif
 							}
 						}
 					} // if (bIsSegmentValid)
