Move final pixel output (color/velocity/...) outside of the tile loop.

diff --git a/Engine/Shaders/Private/HairStrands/HairStrandsForwardRaster.usf b/Engine/Shaders/Private/HairStrands/HairStrandsForwardRaster.usf
--- a/Engine/Shaders/Private/HairStrands/HairStrandsForwardRaster.usf
+++ b/Engine/Shaders/Private/HairStrands/HairStrandsForwardRaster.usf
@@ -1244,6 +1244,7 @@ void RasterCS(
 	uint   Thread_Complete = 0;
 	float4 Thread_Velocity = INVALID_VELOCITY;
 	uint   Thread_LoopCountToFullCoverage = 0;
+	uint2  Thread_PixelCoord = 0;
 
 	LOOP
 	for (uint WorkIndex = 0; WorkIndex < MAX_WORK_COUNT; WorkIndex++)
@@ -1277,6 +1278,7 @@ void RasterCS(
 			Thread_Complete = 0;
 			Thread_Velocity = INVALID_VELOCITY;
 			Thread_LoopCountToFullCoverage = 0;
+			Thread_PixelCoord = 0;
 		}
 
 		// Early out if we are done
@@ -1288,423 +1290,361 @@ void RasterCS(
 			}
 		}
 
+		// Iterate over all bins for the current raster/screen tile
 		const uint FrontToBackCount = min(s_BinTileCount, 64);
 		for (uint FrontToBackIndex = 0; FrontToBackIndex < FrontToBackCount; ++FrontToBackIndex)
 		{
 
-		// 0.2 Reset all LDS variables
-		{
-			if (GroupThread1D == 0)
-			{
-				s_Segments_Min = PackDepth(1e8);
-				s_Segments_Max = 0;
-				s_Segments_ValidCount = 0;
-				s_bDataOrder = 0;
-				s_OpaqueMask = 0;
-			}
-			if (GroupThread1D < RASTER_DEPTH_BUCKET)
+			// 0.2 Reset all LDS variables
 			{
-				s_SegmentsCount[GroupThread1D] = 0;
-				s_SegmentsAlloc[GroupThread1D] = 0;
-			}
-			ClearMask(GroupThread2D);
-			//s_Mask[Thread_Coord.x][Thread_Coord.y] = 0;
+				if (GroupThread1D == 0)
+				{
+					s_Segments_Min = PackDepth(1e8);
+					s_Segments_Max = 0;
+					s_Segments_ValidCount = 0;
+					s_bDataOrder = 0;
+					s_OpaqueMask = 0;
+				}
+				if (GroupThread1D < RASTER_DEPTH_BUCKET)
+				{
+					s_SegmentsCount[GroupThread1D] = 0;
+					s_SegmentsAlloc[GroupThread1D] = 0;
+				}
+				ClearMask(GroupThread2D);
+				//s_Mask[Thread_Coord.x][Thread_Coord.y] = 0;
 
-			if (GroupThread1D < 32)
-			{
-				s_Data[GroupThread1D][0] = 0;
-				s_Data[GroupThread1D][1] = 0;
-				s_Data[GroupThread1D][2] = 0;
-				s_Data[GroupThread1D][3] = 0;
+				if (GroupThread1D < 32)
+				{
+					s_Data[GroupThread1D][0] = 0;
+					s_Data[GroupThread1D][1] = 0;
+					s_Data[GroupThread1D][2] = 0;
+					s_Data[GroupThread1D][3] = 0;
+				}
+				else
+				{
+					s_Data[GroupThread1D-32][4] = 0;
+					s_Data[GroupThread1D-32][5] = 0;
+					s_Data[GroupThread1D-32][6] = 0;
+					s_Data[GroupThread1D-32][7] = 0;
+				}
+
+				#if PERMUTATION_DEBUG
+				s_Coverage[GroupThread1D] = 0;
+				#endif
 			}
-			else
+			GroupMemoryBarrierWithGroupSync();
+
+			// 0.4 Early out when running out of valid tiles
+			const uint BinTileIndex = s_BinTileOffset + FrontToBackIndex;
+			const bool bTileValid = (BinTileIndex < BinTileNum);
+			if (!bTileValid)
 			{
-				s_Data[GroupThread1D-32][4] = 0;
-				s_Data[GroupThread1D-32][5] = 0;
-				s_Data[GroupThread1D-32][6] = 0;
-				s_Data[GroupThread1D-32][7] = 0;
+				return;
 			}
 
-			#if PERMUTATION_DEBUG
-			s_Coverage[GroupThread1D] = 0;
-			#endif
-		}
-		GroupMemoryBarrierWithGroupSync();
+			const uint PrimOffset	= BinTileIndex * SEGMENT_COUNT_PER_GROUP;
+			const uint PrimCount	= LoadVisTileData(VisTileData, BinTileIndex, VT_PrimCount);
+			const uint2 BinTileCoord= UnpackVisTileCoord(LoadVisTileData(VisTileData, BinTileIndex, VT_Coord));
+			const uint2 BinTileMin	= BinTileCoord * BIN_TILE_SIZE;
+			const uint2 BinTileMax	= BinTileMin + BIN_TILE_SIZE;
 
-		// 0.4 Early out when running out of valid tiles
-		const uint BinTileIndex = s_BinTileOffset + FrontToBackIndex;
-		const bool bTileValid = (BinTileIndex < BinTileNum);
-		if (!bTileValid)
-		{
-			return;
-		}
+			const uint QuadrantIndex	= s_WorkID % 16; // 1 x Bin32x32 -> 16 x Raster8x8
+			const uint2 QuadrantCoord	= LinearTo2D_Common(QuadrantIndex, 4, 1.f / 4.f);
+			const uint2 RasterTileCoord = BinTileCoord * 4 + QuadrantCoord;
+			const uint2 RasterTileMin	= RasterTileCoord * RASTER_TILE_SIZE;
+			const uint2 RasterTileMax	= RasterTileMin   + RASTER_TILE_SIZE;
 
-		const uint PrimOffset	= BinTileIndex * SEGMENT_COUNT_PER_GROUP;
-		const uint PrimCount	= LoadVisTileData(VisTileData, BinTileIndex, VT_PrimCount);
-		const uint2 BinTileCoord= UnpackVisTileCoord(LoadVisTileData(VisTileData, BinTileIndex, VT_Coord));
-		const uint2 BinTileMin	= BinTileCoord * BIN_TILE_SIZE;
-		const uint2 BinTileMax	= BinTileMin + BIN_TILE_SIZE;
+			const uint LoopCount64  = DivideAndRoundUp(PrimCount, RASTER_THREAD_COUNT);
 
-		const uint QuadrantIndex	= s_WorkID % 16; // 1 x Bin32x32 -> 16 x Raster8x8
-		const uint2 QuadrantCoord	= LinearTo2D_Common(QuadrantIndex, 4, 1.f / 4.f);
-		const uint2 RasterTileCoord = BinTileCoord * 4 + QuadrantCoord;
-		const uint2 RasterTileMin	= RasterTileCoord * RASTER_TILE_SIZE;
-		const uint2 RasterTileMax	= RasterTileMin   + RASTER_TILE_SIZE;
+			Thread_PixelCoord = RasterTileMin + GroupThread2D;
 
-		const uint LoopCount64  = DivideAndRoundUp(PrimCount, RASTER_THREAD_COUNT);
-
-		// For debug only
-		#if PERMUTATION_DEBUG
-		const bool bDebugEnableAll	= all(RasterTileCoord == uint2(ShaderPrintData.CursorCoord / float(RASTER_TILE_SIZE)));
-		const bool bDebugEnable		= all(RasterTileCoord == uint2(ShaderPrintData.CursorCoord / float(RASTER_TILE_SIZE))) && GroupThread1D == 0;
-		FShaderPrintContext Ctx		= InitShaderPrintContext(bDebugEnable, uint2(350, 50));
-		const FFontColor FontLegend = FontWhite;
-		const FFontColor FontValue  = FontOrange;
-
-		//const bool bDebugEnable2 = all(RasterTileCoord == uint2(ShaderPrintData.CursorCoord / float(RASTER_TILE_SIZE))) && GroupThread1D == 1;
-		FShaderPrintContext CtxAll = InitShaderPrintContext(bDebugEnableAll, uint2(750, 50));
-		#endif
+			// For debug only
+			#if PERMUTATION_DEBUG
+			const bool bDebugEnableAll	= all(RasterTileCoord == uint2(ShaderPrintData.CursorCoord / float(RASTER_TILE_SIZE)));
+			const bool bDebugEnable		= all(RasterTileCoord == uint2(ShaderPrintData.CursorCoord / float(RASTER_TILE_SIZE))) && GroupThread1D == 0;
+			FShaderPrintContext Ctx		= InitShaderPrintContext(bDebugEnable, uint2(350, 50));
+			const FFontColor FontLegend = FontWhite;
+			const FFontColor FontValue  = FontOrange;
+
+			//const bool bDebugEnable2 = all(RasterTileCoord == uint2(ShaderPrintData.CursorCoord / float(RASTER_TILE_SIZE))) && GroupThread1D == 1;
+			FShaderPrintContext CtxAll = InitShaderPrintContext(bDebugEnableAll, uint2(750, 50));
+			#endif
 
-		#if PERMUTATION_DEBUG
-		if (bDebugEnable)
-		{
-			Print(Ctx, TEXT("Work - Index :"), FontLegend); Print(Ctx, FrontToBackIndex, FontValue); Newline(Ctx);
-			Print(Ctx, TEXT("Work - Offset:"), FontLegend); Print(Ctx, s_BinTileOffset, FontValue); Newline(Ctx);
-			Print(Ctx, TEXT("Work - Count :"), FontLegend); Print(Ctx, s_BinTileCount, FontValue); Newline(Ctx);
+			#if PERMUTATION_DEBUG
+			if (bDebugEnable)
+			{
+				Print(Ctx, TEXT("Work - Index :"), FontLegend); Print(Ctx, FrontToBackIndex, FontValue); Newline(Ctx);
+				Print(Ctx, TEXT("Work - Offset:"), FontLegend); Print(Ctx, s_BinTileOffset, FontValue); Newline(Ctx);
+				Print(Ctx, TEXT("Work - Count :"), FontLegend); Print(Ctx, s_BinTileCount, FontValue); Newline(Ctx);
 
-			AddFilledQuadSS(BinTileMin, BinTileMax, Transparent(ColorLightGreen));
-			//AddFilledQuadSS(RasterTileMin, RasterTileMax, Transparent(ColorLightGreen));
-			AddQuadSS(RasterTileMin, RasterTileMax, Transparent(ColorYellow));
-		}
-		#endif
+				AddFilledQuadSS(BinTileMin, BinTileMax, Transparent(ColorLightGreen));
+				//AddFilledQuadSS(RasterTileMin, RasterTileMax, Transparent(ColorLightGreen));
+				AddQuadSS(RasterTileMin, RasterTileMax, Transparent(ColorYellow));
+			}
+			#endif
 
-		// 1. Load all the segments and compute min/max bound (move this during the compaction)
-		{
-			LOOP
-			for (uint LoopIndex = 0; LoopIndex < LoopCount64; ++LoopIndex)
+			// 1. Load all the segments and compute min/max bound (move this during the compaction)
 			{
-				const uint Prim = GroupThread1D + LoopIndex * RASTER_THREAD_COUNT;
-				if (Prim < PrimCount)
+				LOOP
+				for (uint LoopIndex = 0; LoopIndex < LoopCount64; ++LoopIndex)
 				{
-					uint PrimID = VisTilePrims[PrimOffset + Prim];
-
-					uint TypeDummy = 0;
-					float4 SP0 = 0;
-					float4 SP1 = 0;
-					float Rad0 = 0;
-					float Rad1 = 0;
-					CalcHomogenousPosAndRad(PrimID,   SP0, Rad0, TypeDummy);
-					CalcHomogenousPosAndRad(PrimID+1, SP1, Rad1, TypeDummy);
-				
-					float Alpha0 = 0;
-					float Alpha1 = 1;
-					// Clipping
+					const uint Prim = GroupThread1D + LoopIndex * RASTER_THREAD_COUNT;
+					if (Prim < PrimCount)
 					{
-						SP0 = float4(NDCToPixelCoord(SP0), 1.0f / SP0.w);
-						SP1 = float4(NDCToPixelCoord(SP1), 1.0f / SP1.w);
+						uint PrimID = VisTilePrims[PrimOffset + Prim];
+
+						uint TypeDummy = 0;
+						float4 SP0 = 0;
+						float4 SP1 = 0;
+						float Rad0 = 0;
+						float Rad1 = 0;
+						CalcHomogenousPosAndRad(PrimID,   SP0, Rad0, TypeDummy);
+						CalcHomogenousPosAndRad(PrimID+1, SP1, Rad1, TypeDummy);
+				
+						float Alpha0 = 0;
+						float Alpha1 = 1;
+						// Clipping
+						{
+							SP0 = float4(NDCToPixelCoord(SP0), 1.0f / SP0.w);
+							SP1 = float4(NDCToPixelCoord(SP1), 1.0f / SP1.w);
 
-						// Clip against tile
-						bool2 bClipped = false;
-						bool bIsValidSegment = false;
-						float2 T = 0;
-						bIsValidSegment = ClipRaySegment(RasterTileMin - 0.5f, RasterTileMax + 0.5f, SP0, SP1, Rad0, Rad1, Alpha0, Alpha1, bClipped, T);
+							// Clip against tile
+							bool2 bClipped = false;
+							bool bIsValidSegment = false;
+							float2 T = 0;
+							bIsValidSegment = ClipRaySegment(RasterTileMin - 0.5f, RasterTileMax + 0.5f, SP0, SP1, Rad0, Rad1, Alpha0, Alpha1, bClipped, T);
 
-						PrimID = bIsValidSegment ? PrimID : INVALID_PRIM_ID;
-					}
+							PrimID = bIsValidSegment ? PrimID : INVALID_PRIM_ID;
+						}
 
-					const uint uDepth = PackDepth(max(SP0.z, SP1.z)); // Inverse Z
-					s_Segments[Prim] = uint2(uDepth, PrimID);
-					if (PrimID != INVALID_PRIM_ID)
-					{
-						InterlockedMin(s_Segments_Min, uDepth);
-						InterlockedMax(s_Segments_Max, uDepth);
-					}
+						const uint uDepth = PackDepth(max(SP0.z, SP1.z)); // Inverse Z
+						s_Segments[Prim] = uint2(uDepth, PrimID);
+						if (PrimID != INVALID_PRIM_ID)
+						{
+							InterlockedMin(s_Segments_Min, uDepth);
+							InterlockedMax(s_Segments_Max, uDepth);
+						}
 
-					#if 1 && PERMUTATION_DEBUG
-					if (bDebugEnableAll && PrimID != INVALID_PRIM_ID)
-					{
-						const float3 P0 = UnpackHairControlPoint(PrimID).Position;
-						const float3 P1 = UnpackHairControlPoint(PrimID+1).Position;
+						#if 1 && PERMUTATION_DEBUG
+						if (bDebugEnableAll && PrimID != INVALID_PRIM_ID)
+						{
+							const float3 P0 = UnpackHairControlPoint(PrimID).Position;
+							const float3 P1 = UnpackHairControlPoint(PrimID+1).Position;
 
-						const float4 Color0 = float4(LoadSampleColor(PrimID), 1);
-						const float4 Color1 = float4(LoadSampleColor(PrimID+1), 1);
+							const float4 Color0 = float4(LoadSampleColor(PrimID), 1);
+							const float4 Color1 = float4(LoadSampleColor(PrimID+1), 1);
 
-						const float4 LineColor = float4(ColorMapTurbo(GroupThread1D / float(RASTER_THREAD_COUNT)), 1);
-						AddLineWS(P0, P1, LineColor);
-						//AddLineWS(P0, P1, Color0, Color1);
+							const float4 LineColor = float4(ColorMapTurbo(GroupThread1D / float(RASTER_THREAD_COUNT)), 1);
+							AddLineWS(P0, P1, LineColor);
+							//AddLineWS(P0, P1, Color0, Color1);
+						}
+						#endif
 					}
-					#endif
 				}
 			}
-		}
-		GroupMemoryBarrierWithGroupSync();
-		const float InvDepthExtent = 1.f / max(0.0001f, UnpackDepth(s_Segments_Max)- UnpackDepth(s_Segments_Min));
+			GroupMemoryBarrierWithGroupSync();
+			const float InvDepthExtent = 1.f / max(0.0001f, UnpackDepth(s_Segments_Max)- UnpackDepth(s_Segments_Min));
 
-		// 2. Compute the count of depth bucket
-		{
-			LOOP
-			for (uint LoopIndex = 0; LoopIndex < LoopCount64; LoopIndex++)
+			// 2. Compute the count of depth bucket
 			{
-				const uint Prim = GroupThread1D + LoopIndex * RASTER_THREAD_COUNT;
-				if (Prim < PrimCount)
+				LOOP
+				for (uint LoopIndex = 0; LoopIndex < LoopCount64; LoopIndex++)
 				{
-					const bool bIsValid = s_Segments[Prim].y != INVALID_PRIM_ID;
-					if (bIsValid)
+					const uint Prim = GroupThread1D + LoopIndex * RASTER_THREAD_COUNT;
+					if (Prim < PrimCount)
 					{
-						const float Depth	= UnpackDepth(s_Segments[Prim].x);
-						const uint DepthIt  = GetDepthBinIndex(Depth, InvDepthExtent);
-						InterlockedAdd(s_SegmentsCount[DepthIt], 1);
+						const bool bIsValid = s_Segments[Prim].y != INVALID_PRIM_ID;
+						if (bIsValid)
+						{
+							const float Depth	= UnpackDepth(s_Segments[Prim].x);
+							const uint DepthIt  = GetDepthBinIndex(Depth, InvDepthExtent);
+							InterlockedAdd(s_SegmentsCount[DepthIt], 1);
+						}
 					}
 				}
 			}
-		}
-		GroupMemoryBarrierWithGroupSync();
+			GroupMemoryBarrierWithGroupSync();
 
-		// Replace this with parallel version
-		s_Segments_ValidCount = 0;
-		if (GroupThread1D == 0)
-		{
-			uint Acc = 0;
-			for (uint It = 0; It < RASTER_DEPTH_BUCKET; It++)
+			// Replace this with parallel version
+			s_Segments_ValidCount = 0;
+			if (GroupThread1D == 0)
 			{
-				const uint Next = s_SegmentsCount[It];
-				s_SegmentsCount[It] = Acc;
-				Acc += Next;
-
-				#if 1 && PERMUTATION_DEBUG
-				if (bDebugEnable)
+				uint Acc = 0;
+				for (uint It = 0; It < RASTER_DEPTH_BUCKET; It++)
 				{
-					if (Next > 0)	{ Print(Ctx, TEXT("x"), FontValue); }
-					else			{ Print(Ctx, TEXT("."), FontValue); }
-					if (It == RASTER_DEPTH_BUCKET-1) { Newline(Ctx); }
+					const uint Next = s_SegmentsCount[It];
+					s_SegmentsCount[It] = Acc;
+					Acc += Next;
+
+					#if 1 && PERMUTATION_DEBUG
+					if (bDebugEnable)
+					{
+						if (Next > 0)	{ Print(Ctx, TEXT("x"), FontValue); }
+						else			{ Print(Ctx, TEXT("."), FontValue); }
+						if (It == RASTER_DEPTH_BUCKET-1) { Newline(Ctx); }
+					}
+					#endif
 				}
-				#endif
+				s_Segments_ValidCount = Acc;
 			}
-			s_Segments_ValidCount = Acc;
-		}
-		GroupMemoryBarrierWithGroupSync();
+			GroupMemoryBarrierWithGroupSync();
 
-		// 3. Insert the segment into the right bucket
-		{
-			LOOP
-			for (uint LoopIndex = 0; LoopIndex < LoopCount64; LoopIndex++)
+			// 3. Insert the segment into the right bucket
 			{
-				const uint Prim = GroupThread1D + LoopIndex * RASTER_THREAD_COUNT;
-				if (Prim < PrimCount)
+				LOOP
+				for (uint LoopIndex = 0; LoopIndex < LoopCount64; LoopIndex++)
 				{
-					const uint2 Segment = s_Segments[Prim];
-					if (Segment.y != INVALID_PRIM_ID)
+					const uint Prim = GroupThread1D + LoopIndex * RASTER_THREAD_COUNT;
+					if (Prim < PrimCount)
 					{
-						const float Depth = UnpackDepth(Segment.x);
-						const uint DepthIt = GetDepthBinIndex(Depth, InvDepthExtent);
+						const uint2 Segment = s_Segments[Prim];
+						if (Segment.y != INVALID_PRIM_ID)
+						{
+							const float Depth = UnpackDepth(Segment.x);
+							const uint DepthIt = GetDepthBinIndex(Depth, InvDepthExtent);
 
-						uint AllocOffset = 0;
-						InterlockedAdd(s_SegmentsAlloc[DepthIt], 1, AllocOffset);
-						const uint NewIndex = AllocOffset + s_SegmentsCount[DepthIt];
-						s_Segments_Sorted[NewIndex] = Segment;
+							uint AllocOffset = 0;
+							InterlockedAdd(s_SegmentsAlloc[DepthIt], 1, AllocOffset);
+							const uint NewIndex = AllocOffset + s_SegmentsCount[DepthIt];
+							s_Segments_Sorted[NewIndex] = Segment;
+						}
 					}
 				}
 			}
-		}
-		GroupMemoryBarrierWithGroupSync();
+			GroupMemoryBarrierWithGroupSync();
 
-		const uint LoopCount32 = DivideAndRoundUp(s_Segments_ValidCount, RASTER_SEGMENT_COUNT);
-		const uint LoopCount32_All = DivideAndRoundUp(PrimCount, RASTER_SEGMENT_COUNT);
+			const uint LoopCount32 = DivideAndRoundUp(s_Segments_ValidCount, RASTER_SEGMENT_COUNT);
+			const uint LoopCount32_All = DivideAndRoundUp(PrimCount, RASTER_SEGMENT_COUNT);
 
-		// DEBUG
-		#if PERMUTATION_DEBUG
-		if (bDebugEnable)
-		{
-			const uint2 OutCoord = GroupThread2D + RasterTileMin;
-
-			//Print(Ctx, TEXT("Out Coord  :"), FontLegend); Print(Ctx, OutCoord, FontValue); Newline(Ctx);
-			//Print(Ctx, TEXT("Min Raster :"), FontLegend); Print(Ctx, RasterTileMin, FontValue); Newline(Ctx);
-			//Print(Ctx, TEXT("Thread     :"), FontLegend); Print(Ctx, Thread_Coord, FontValue); Newline(Ctx);
-			//Print(Ctx, TEXT("Max Raster :"), FontLegend); Print(Ctx, RasterTileMax, FontValue); Newline(Ctx);
-			//Newline(Ctx);
-
-			//Print(Ctx, TEXT("Cursor     :"), FontLegend); Print(Ctx, uint2(ShaderPrintData.CursorCoord), FontValue); Newline(Ctx);
-			//Newline(Ctx);
-
-			Print(Ctx, TEXT("Work ID      :"), FontLegend); Print(Ctx, s_WorkID, FontValue); Newline(Ctx);
-			Print(Ctx, TEXT("Prim Count   :"), FontLegend); Print(Ctx, s_Segments_ValidCount, FontValue);	Print(Ctx, TEXT(" / "), FontLegend); Print(Ctx, PrimCount, FontValue); Newline(Ctx);
-			Print(Ctx, TEXT("Loop 32 Count:"), FontLegend); Print(Ctx, LoopCount32, FontValue);				Print(Ctx, TEXT(" / "), FontLegend); Print(Ctx, LoopCount32_All, FontValue); Newline(Ctx);
-			Print(Ctx, TEXT("Loop 64 Count:"), FontLegend); Print(Ctx, LoopCount64, FontValue); Newline(Ctx);
-			Print(Ctx, TEXT("Min Depth    :"), FontLegend); Print(Ctx, UnpackDepth(s_Segments_Min), FontValue); Newline(Ctx);
-			Print(Ctx, TEXT("Max Depth    :"), FontLegend); Print(Ctx, UnpackDepth(s_Segments_Max), FontValue); Newline(Ctx);
-			Print(Ctx, TEXT("Rad.at Depth1:"), FontLegend); Print(Ctx, RadiusAtDepth1, FontValue); Newline(Ctx);
-			Newline(Ctx);
-
-			Print(Ctx, TEXT("Work - Index :"), FontLegend); Print(Ctx, WorkIndex, FontValue, 2, 0);
-			Print(Ctx, TEXT(" - Offset:"), FontLegend); Print(Ctx, s_BinTileOffset, FontValue, 5, 0);
-			Print(Ctx, TEXT(" - Count :"), FontLegend); Print(Ctx, s_BinTileCount, FontValue, 3, 0);
-			Newline(Ctx);
-
-			//Print(Ctx, TEXT("Bin ID     :"), FontLegend); Print(Ctx, BinTileIndex, FontValue); Newline(Ctx);
-			//Print(Ctx, TEXT("Raster ID  :"), FontLegend); Print(Ctx, QuadrantIndex, FontValue); Newline(Ctx);
-			//Newline(Ctx);
-
-			//Print(Ctx, TEXT("Raster Tile:"), FontLegend); Print(Ctx, RasterTileCoord, FontValue); Newline(Ctx);
-			//Print(Ctx, TEXT("Bin Tile   :"), FontLegend); Print(Ctx, BinTileCoord, FontValue); Newline(Ctx);
-			//Print(Ctx, TEXT("Out Coord  :"), FontLegend); Print(Ctx, OutCoord, FontValue); Newline(Ctx);
-			//Newline(Ctx);
+			// DEBUG
+			#if PERMUTATION_DEBUG
+			if (bDebugEnable)
+			{
+				const uint2 OutCoord = GroupThread2D + RasterTileMin;
+
+				//Print(Ctx, TEXT("Out Coord  :"), FontLegend); Print(Ctx, OutCoord, FontValue); Newline(Ctx);
+				//Print(Ctx, TEXT("Min Raster :"), FontLegend); Print(Ctx, RasterTileMin, FontValue); Newline(Ctx);
+				//Print(Ctx, TEXT("Thread     :"), FontLegend); Print(Ctx, Thread_Coord, FontValue); Newline(Ctx);
+				//Print(Ctx, TEXT("Max Raster :"), FontLegend); Print(Ctx, RasterTileMax, FontValue); Newline(Ctx);
+				//Newline(Ctx);
+
+				//Print(Ctx, TEXT("Cursor     :"), FontLegend); Print(Ctx, uint2(ShaderPrintData.CursorCoord), FontValue); Newline(Ctx);
+				//Newline(Ctx);
+
+				Print(Ctx, TEXT("Work ID      :"), FontLegend); Print(Ctx, s_WorkID, FontValue); Newline(Ctx);
+				Print(Ctx, TEXT("Prim Count   :"), FontLegend); Print(Ctx, s_Segments_ValidCount, FontValue);	Print(Ctx, TEXT(" / "), FontLegend); Print(Ctx, PrimCount, FontValue); Newline(Ctx);
+				Print(Ctx, TEXT("Loop 32 Count:"), FontLegend); Print(Ctx, LoopCount32, FontValue);				Print(Ctx, TEXT(" / "), FontLegend); Print(Ctx, LoopCount32_All, FontValue); Newline(Ctx);
+				Print(Ctx, TEXT("Loop 64 Count:"), FontLegend); Print(Ctx, LoopCount64, FontValue); Newline(Ctx);
+				Print(Ctx, TEXT("Min Depth    :"), FontLegend); Print(Ctx, UnpackDepth(s_Segments_Min), FontValue); Newline(Ctx);
+				Print(Ctx, TEXT("Max Depth    :"), FontLegend); Print(Ctx, UnpackDepth(s_Segments_Max), FontValue); Newline(Ctx);
+				Print(Ctx, TEXT("Rad.at Depth1:"), FontLegend); Print(Ctx, RadiusAtDepth1, FontValue); Newline(Ctx);
+				Newline(Ctx);
 
-		}
-		#endif
+				Print(Ctx, TEXT("Work - Index :"), FontLegend); Print(Ctx, WorkIndex, FontValue, 2, 0);
+				Print(Ctx, TEXT(" - Offset:"), FontLegend); Print(Ctx, s_BinTileOffset, FontValue, 5, 0);
+				Print(Ctx, TEXT(" - Count :"), FontLegend); Print(Ctx, s_BinTileCount, FontValue, 3, 0);
+				Newline(Ctx);
 
-		// 4. Raster segments
-		{
-			// 4.1 Loop over all segment within the tile, and rastize 32 of them at each loop
-			LOOP
-			for (uint LoopIndex = 0; LoopIndex < LoopCount32; LoopIndex++)
-			{
-				// 4.0 Reset
-				if (GroupThread1D < RASTER_SEGMENT_COUNT)
-				{
-					for (int J = 0; J < RASTER_TILE_SIZE; ++J)
-					{
-						s_Data[GroupThread1D][J] = 0;
-						s_Color[GroupThread1D][J] = 0;
-						s_Velocity[GroupThread1D][J] = 0;
-					}
-				}
-				s_Mask[GroupThread2D.x][GroupThread2D.y] = 0;
+				//Print(Ctx, TEXT("Bin ID     :"), FontLegend); Print(Ctx, BinTileIndex, FontValue); Newline(Ctx);
+				//Print(Ctx, TEXT("Raster ID  :"), FontLegend); Print(Ctx, QuadrantIndex, FontValue); Newline(Ctx);
+				//Newline(Ctx);
 
-				GroupMemoryBarrierWithGroupSync();
+				//Print(Ctx, TEXT("Raster Tile:"), FontLegend); Print(Ctx, RasterTileCoord, FontValue); Newline(Ctx);
+				//Print(Ctx, TEXT("Bin Tile   :"), FontLegend); Print(Ctx, BinTileCoord, FontValue); Newline(Ctx);
+				//Print(Ctx, TEXT("Out Coord  :"), FontLegend); Print(Ctx, OutCoord, FontValue); Newline(Ctx);
+				//Newline(Ctx);
 
-				if (s_OpaqueMask == 64)
-				{
-					break;
-				}
+			}
+			#endif
 
-				// 4.1 Raster segment (1 thread = 1 segment)
-				// Half of the thread are doing nothing, we could raster the two half of the segments (one per each thread)
-				const uint Prim = GroupThread1D + LoopIndex * RASTER_SEGMENT_COUNT;
-				if (GroupThread1D < RASTER_SEGMENT_COUNT && Prim < PrimCount)
+			// 4. Raster segments
+			{
+				// 4.1 Loop over all segment within the tile, and rastize 32 of them at each loop
+				LOOP
+				for (uint LoopIndex = 0; LoopIndex < LoopCount32; LoopIndex++)
 				{
-					const uint PrimID = s_Segments_Sorted[Prim].y;
-					uint TypeDummy = 0;
-					float4 SP0 = 0;
-					float4 SP1 = 0;
-					float Rad0 = 0;
-					float Rad1 = 0;
-					CalcHomogenousPosAndRad(PrimID,     SP0, Rad0, TypeDummy);
-					CalcHomogenousPosAndRad(PrimID + 1, SP1, Rad1, TypeDummy);
-
-					float Alpha0=0;
-					float Alpha1=1;
-
-					bool2 bClipped = false;
-					bool bIsSegmentValid = false;
-					float2 T = 0;
-					// 4.1.1 Clipping
+					// 4.0 Reset
+					if (GroupThread1D < RASTER_SEGMENT_COUNT)
 					{
-						SP0 = float4(NDCToPixelCoord(SP0), 1.0f / SP0.w);
-						SP1 = float4(NDCToPixelCoord(SP1), 1.0f / SP1.w);
-				
-						// Clip against tile
-						bIsSegmentValid = ClipRaySegment(RasterTileMin - 0.5f, RasterTileMax + 0.5f, SP0, SP1, Rad0, Rad1, Alpha0, Alpha1, bClipped, T);
+						for (int J = 0; J < RASTER_TILE_SIZE; ++J)
+						{
+							s_Data[GroupThread1D][J] = 0;
+							s_Color[GroupThread1D][J] = 0;
+							s_Velocity[GroupThread1D][J] = 0;
+						}
 					}
+					s_Mask[GroupThread2D.x][GroupThread2D.y] = 0;
 
-					// DEBUG (Write coord)
-					#if 0 && PERMUTATION_DEBUG
-					if (bDebugEnableAll)
-					{
-						ShiftY(CtxAll, 500 + GroupThread1D * 150);
-						
-						Print(CtxAll, TEXT("Thrd 1D:"), FontYellow);
-						Print(CtxAll, GroupThread1D, FontWhite);
-						Newline(CtxAll);
-
-						Print(CtxAll, TEXT("ID     :"), FontYellow);
-						Print(CtxAll, PrimID, FontWhite);
-						Newline(CtxAll);
-
-						Print(CtxAll, TEXT("Valid  :"), FontYellow);
-						PrintBool(CtxAll, bIsSegmentValid);
-						Newline(CtxAll);
-
-						Print(CtxAll, TEXT("Clip X :"), FontYellow);
-						PrintBool(CtxAll, bClipped.x);
-						Newline(CtxAll);
+					GroupMemoryBarrierWithGroupSync();
 
-						Print(CtxAll, TEXT("Clip Y :"), FontYellow);
-						PrintBool(CtxAll, bClipped.y);
-						Newline(CtxAll);
-
-						Print(CtxAll, TEXT("X0      :"), FontYellow);
-						Print(CtxAll, SP0, FontWhite);
-						Newline(CtxAll);
-
-						Print(CtxAll, TEXT("X1      :"), FontYellow);
-						Print(CtxAll, SP1, FontWhite);
-						Newline(CtxAll);
-
-						//const float3 P0 = UnpackHairControlPoint(PrimID).Position;
-						//const float3 P1 = UnpackHairControlPoint(PrimID + 1).Position;
-
-						//const float4 Color0 = float4(LoadSampleColor(PrimID), 1);
-						//const float4 Color1 = float4(LoadSampleColor(PrimID + 1), 1);
-
-						//const float4 LineColor = float4(ColorMapTurbo(GroupThread1D / float(RASTER_THREAD_COUNT)), 1);
-						//AddLineWS(P0, P1, LineColor);
-						if (GroupThread1D == 0)
-							AddLineSS(SP0, SP1, ColorPurple);
+					if (s_OpaqueMask == 64)
+					{
+						break;
 					}
-					#endif
 
-					// 4.1.2 Rasterize segment (1 thread = 1 segments)
-					if (bIsSegmentValid)
+					// 4.1 Raster segment (1 thread = 1 segment)
+					// Half of the thread are doing nothing, we could raster the two half of the segments (one per each thread)
+					const uint Prim = GroupThread1D + LoopIndex * RASTER_SEGMENT_COUNT;
+					if (GroupThread1D < RASTER_SEGMENT_COUNT && Prim < PrimCount)
 					{
-						const float SegmentLenSqRcp = 1.0f / dot(SP1.xy - SP0.xy, SP1.xy - SP0.xy);
-						const bool bIsSteep = abs(SP1.x - SP0.x) < abs(SP1.y - SP0.y);
-						const float X0 = bIsSteep ? min(SP0.y, SP1.y) : min(SP0.x, SP1.x);
-						const float X1 = bIsSteep ? max(SP0.y, SP1.y) : max(SP0.x, SP1.x);
-						const float RcpNumSteps = 1.0f / (X1 - X0);
-						const int NumSteps = (int)(ceil(X1) - floor(X0));
-
-						if (bIsSteep)
+						const uint PrimID = s_Segments_Sorted[Prim].y;
+						uint TypeDummy = 0;
+						float4 SP0 = 0;
+						float4 SP1 = 0;
+						float Rad0 = 0;
+						float Rad1 = 0;
+						CalcHomogenousPosAndRad(PrimID,     SP0, Rad0, TypeDummy);
+						CalcHomogenousPosAndRad(PrimID + 1, SP1, Rad1, TypeDummy);
+
+						float Alpha0=0;
+						float Alpha1=1;
+
+						bool2 bClipped = false;
+						bool bIsSegmentValid = false;
+						float2 T = 0;
+						// 4.1.1 Clipping
 						{
-							InterlockedOr(s_bDataOrder, 1u << GroupThread1D);
+							SP0 = float4(NDCToPixelCoord(SP0), 1.0f / SP0.w);
+							SP1 = float4(NDCToPixelCoord(SP1), 1.0f / SP1.w);
+				
+							// Clip against tile
+							bIsSegmentValid = ClipRaySegment(RasterTileMin - 0.5f, RasterTileMax + 0.5f, SP0, SP1, Rad0, Rad1, Alpha0, Alpha1, bClipped, T);
 						}
 
 						// DEBUG (Write coord)
 						#if 0 && PERMUTATION_DEBUG
-						//if (GroupThread1D == 0 && bDebugEnable)
 						if (bDebugEnableAll)
-						{							
-							ShiftY(CtxAll, GroupThread1D * 150);
-							//ShiftY(CtxAll, GroupThread1D * )
-
-							Print(CtxAll, TEXT("NumSteps:"), FontYellow);
-							Print(CtxAll, NumSteps, FontWhite);
+						{
+							ShiftY(CtxAll, 500 + GroupThread1D * 150);
+						
+							Print(CtxAll, TEXT("Thrd 1D:"), FontYellow);
+							Print(CtxAll, GroupThread1D, FontWhite);
 							Newline(CtxAll);
 
-							Print(CtxAll, TEXT("X0      :"), FontYellow);
-							Print(CtxAll, X0, FontWhite);
+							Print(CtxAll, TEXT("ID     :"), FontYellow);
+							Print(CtxAll, PrimID, FontWhite);
 							Newline(CtxAll);
 
-							Print(CtxAll, TEXT("X1      :"), FontYellow);
-							Print(CtxAll, X1, FontWhite);
+							Print(CtxAll, TEXT("Valid  :"), FontYellow);
+							PrintBool(CtxAll, bIsSegmentValid);
 							Newline(CtxAll);
 
-							Print(CtxAll, TEXT("T       :"), FontYellow);
-							Print(CtxAll, T, FontWhite);
+							Print(CtxAll, TEXT("Clip X :"), FontYellow);
+							PrintBool(CtxAll, bClipped.x);
 							Newline(CtxAll);
 
-							Print(CtxAll, TEXT("bValid  :"), FontYellow);
-							PrintBool(CtxAll, bIsSegmentValid);
+							Print(CtxAll, TEXT("Clip Y :"), FontYellow);
+							PrintBool(CtxAll, bClipped.y);
 							Newline(CtxAll);
 
-							Print(CtxAll, TEXT("bClip X :"), FontYellow);
-							PrintBool(CtxAll, bClipped.x);
+							Print(CtxAll, TEXT("X0      :"), FontYellow);
+							Print(CtxAll, SP0, FontWhite);
 							Newline(CtxAll);
 
-							Print(CtxAll, TEXT("bClip Y :"), FontYellow);
-							PrintBool(CtxAll, bClipped.y);
+							Print(CtxAll, TEXT("X1      :"), FontYellow);
+							Print(CtxAll, SP1, FontWhite);
 							Newline(CtxAll);
-							
-							
+
 							//const float3 P0 = UnpackHairControlPoint(PrimID).Position;
 							//const float3 P1 = UnpackHairControlPoint(PrimID + 1).Position;
 
@@ -1713,241 +1653,253 @@ void RasterCS(
 
 							//const float4 LineColor = float4(ColorMapTurbo(GroupThread1D / float(RASTER_THREAD_COUNT)), 1);
 							//AddLineWS(P0, P1, LineColor);
-							AddLineSS(SP0, SP1, ColorRed);
+							if (GroupThread1D == 0)
+								AddLineSS(SP0, SP1, ColorPurple);
 						}
 						#endif
 
-						// DEBUG (Legend)
-						#if 0 && PERMUTATION_DEBUG
-						if (bDebugEnable)
+						// 4.1.2 Rasterize segment (1 thread = 1 segments)
+						if (bIsSegmentValid)
 						{
-							//Print(Ctx, s_Mask[IntraTileCoord.x][IntraTileCoord.y], FontWhite);
-							//PrintPixelMask(Ctx, ReadMask(IntraTileCoord)/*s_Mask[IntraTileCoord.x][IntraTileCoord.y]*/);
-							//Print(Ctx, TEXT("Color0: "), FontWhite); Print(Ctx, Color0); Newline(Ctx);
-							//Print(Ctx, TEXT("Color1: "), FontWhite); Print(Ctx, Color1); Newline(Ctx);
-							//Print(Ctx, TEXT("Color :  "), FontWhite); Print(Ctx, Color);
-							const float3 Color0 = LoadSampleColor(PrimID);
-							const float3 Color1 = LoadSampleColor(PrimID + 1);
-							Print(Ctx, TEXT("Color0 :  "), FontWhite); Print(Ctx, Color0); Newline(Ctx);
-							Print(Ctx, TEXT("Color1 :  "), FontWhite); Print(Ctx, Color1); Newline(Ctx);
-							Print(Ctx, TEXT("  Alpha       A0          A1          AColor     Color    Coord  Coverage"), FontWhite);
-							Newline(Ctx);
-						}
-						#endif
-						LOOP
-						for (int J = 0; J < NumSteps; ++J)
-						//for (int J = 0; J < RASTER_TILE_SIZE; ++J)
-						{
-							const float Alpha = saturate(J * RcpNumSteps);
-							const float4 SP = lerp(SP0, SP1, Alpha);
-
-							int2 Coords = SP.xy;
-							const int2 IntraTileCoord = Coords - RasterTileMin;
-
-							// TO it needs to store data in sweeping order 
-							FSampleData Sample = (FSampleData)0;
-							if (all(IntraTileCoord >= 0) && all(IntraTileCoord < RASTER_TILE_SIZE))
+							const float SegmentLenSqRcp = 1.0f / dot(SP1.xy - SP0.xy, SP1.xy - SP0.xy);
+							const bool bIsSteep = abs(SP1.x - SP0.x) < abs(SP1.y - SP0.y);
+							const float X0 = bIsSteep ? min(SP0.y, SP1.y) : min(SP0.x, SP1.x);
+							const float X1 = bIsSteep ? max(SP0.y, SP1.y) : max(SP0.x, SP1.x);
+							const float RcpNumSteps = 1.0f / (X1 - X0);
+							const int NumSteps = (int)(ceil(X1) - floor(X0));
+
+							if (bIsSteep)
 							{
-								const float Alpha = ComputeLerpAlpha(Coords, SP0.xy, SP1.xy, SegmentLenSqRcp);
-								const float Depth = lerp(SP0.z, SP1.z, Alpha);
-								const float OpaqueDepth = SceneDepthTexture.Load(uint3(Coords, 0));
-								if (Depth > OpaqueDepth)
-								{
-									const float AlphaColor = lerp(Alpha0, Alpha1, Alpha);
-									const float3 Color0 = LoadSampleColor(PrimID);
-									const float3 Color1 = LoadSampleColor(PrimID+1);
-									const float3 Color = lerp(Color0, Color1, AlphaColor);
-
-									const float4 Velocity0 = LoadSampleVelocity(PrimID);
-									const float4 Velocity1 = LoadSampleVelocity(PrimID+1);
-									const float4 Velocity = lerp(Velocity0, Velocity1, AlphaColor);
-
+								InterlockedOr(s_bDataOrder, 1u << GroupThread1D);
+							}
 
-									// Fill in sample data
-									Sample.Depth = Depth;
-									Sample.Coord = IntraTileCoord;
+							// DEBUG (Write coord)
+							#if 0 && PERMUTATION_DEBUG
+							//if (GroupThread1D == 0 && bDebugEnable)
+							if (bDebugEnableAll)
+							{							
+								ShiftY(CtxAll, GroupThread1D * 150);
+								//ShiftY(CtxAll, GroupThread1D * )
+
+								Print(CtxAll, TEXT("NumSteps:"), FontYellow);
+								Print(CtxAll, NumSteps, FontWhite);
+								Newline(CtxAll);
+
+								Print(CtxAll, TEXT("X0      :"), FontYellow);
+								Print(CtxAll, X0, FontWhite);
+								Newline(CtxAll);
+
+								Print(CtxAll, TEXT("X1      :"), FontYellow);
+								Print(CtxAll, X1, FontWhite);
+								Newline(CtxAll);
+
+								Print(CtxAll, TEXT("T       :"), FontYellow);
+								Print(CtxAll, T, FontWhite);
+								Newline(CtxAll);
+
+								Print(CtxAll, TEXT("bValid  :"), FontYellow);
+								PrintBool(CtxAll, bIsSegmentValid);
+								Newline(CtxAll);
+
+								Print(CtxAll, TEXT("bClip X :"), FontYellow);
+								PrintBool(CtxAll, bClipped.x);
+								Newline(CtxAll);
+
+								Print(CtxAll, TEXT("bClip Y :"), FontYellow);
+								PrintBool(CtxAll, bClipped.y);
+								Newline(CtxAll);
+							
+							
+								//const float3 P0 = UnpackHairControlPoint(PrimID).Position;
+								//const float3 P1 = UnpackHairControlPoint(PrimID + 1).Position;
 
-									// Compute coverage
-									// Minimal radius to snap the strand to a sample/pixel center (to avoid aliasing)
-									const bool bIsOrthoView = false;
-									const float SceneDistance = ConvertFromDeviceZ(Depth);
-									const float MinHairRadius = bIsOrthoView ? RadiusAtDepth1 : (SceneDistance * RadiusAtDepth1);
-									const float HairRadius = lerp(Rad0, Rad1, AlphaColor);
-									Sample.Coverage = saturate(HairRadius / MinHairRadius);
+								//const float4 Color0 = float4(LoadSampleColor(PrimID), 1);
+								//const float4 Color1 = float4(LoadSampleColor(PrimID + 1), 1);
 
-									// Write data
-									WriteMask(IntraTileCoord, 1u << GroupThread1D);
+								//const float4 LineColor = float4(ColorMapTurbo(GroupThread1D / float(RASTER_THREAD_COUNT)), 1);
+								//AddLineWS(P0, P1, LineColor);
+								AddLineSS(SP0, SP1, ColorRed);
+							}
+							#endif
 
-									const uint JCoord = bIsSteep ? IntraTileCoord.y : IntraTileCoord.x;
-									s_Color[GroupThread1D][JCoord] = PackColorData(Color);
-									s_Data[GroupThread1D][JCoord]  = PackSampleData(Sample);
-									s_Velocity[GroupThread1D][JCoord] = PackVelocityData(Velocity);
+							// DEBUG (Legend)
+							#if 0 && PERMUTATION_DEBUG
+							if (bDebugEnable)
+							{
+								//Print(Ctx, s_Mask[IntraTileCoord.x][IntraTileCoord.y], FontWhite);
+								//PrintPixelMask(Ctx, ReadMask(IntraTileCoord)/*s_Mask[IntraTileCoord.x][IntraTileCoord.y]*/);
+								//Print(Ctx, TEXT("Color0: "), FontWhite); Print(Ctx, Color0); Newline(Ctx);
+								//Print(Ctx, TEXT("Color1: "), FontWhite); Print(Ctx, Color1); Newline(Ctx);
+								//Print(Ctx, TEXT("Color :  "), FontWhite); Print(Ctx, Color);
+								const float3 Color0 = LoadSampleColor(PrimID);
+								const float3 Color1 = LoadSampleColor(PrimID + 1);
+								Print(Ctx, TEXT("Color0 :  "), FontWhite); Print(Ctx, Color0); Newline(Ctx);
+								Print(Ctx, TEXT("Color1 :  "), FontWhite); Print(Ctx, Color1); Newline(Ctx);
+								Print(Ctx, TEXT("  Alpha       A0          A1          AColor     Color    Coord  Coverage"), FontWhite);
+								Newline(Ctx);
+							}
+							#endif
+							LOOP
+							for (int J = 0; J < NumSteps; ++J)
+							//for (int J = 0; J < RASTER_TILE_SIZE; ++J)
+							{
+								const float Alpha = saturate(J * RcpNumSteps);
+								const float4 SP = lerp(SP0, SP1, Alpha);
 
+								int2 Coords = SP.xy;
+								const int2 IntraTileCoord = Coords - RasterTileMin;
 
-									// DEBUG (Write coord)
-									#if 0 && PERMUTATION_DEBUG
-									if (bDebugEnable)
+								// TO it needs to store data in sweeping order 
+								FSampleData Sample = (FSampleData)0;
+								if (all(IntraTileCoord >= 0) && all(IntraTileCoord < RASTER_TILE_SIZE))
+								{
+									const float Alpha = ComputeLerpAlpha(Coords, SP0.xy, SP1.xy, SegmentLenSqRcp);
+									const float Depth = lerp(SP0.z, SP1.z, Alpha);
+									const float OpaqueDepth = SceneDepthTexture.Load(uint3(Coords, 0));
+									if (Depth > OpaqueDepth)
 									{
-										Print(Ctx, J, FontYellow, 2, 0);
-										Print(Ctx, Alpha, FontWhite);
-										Print(Ctx, AlphaColor, FontBlue);
-										Print(Ctx, HairRadius, FontRed);
-										Print(Ctx, MinHairRadius, FontGreen);
-										Print(Ctx, Sample.Coverage, FontBlue);
-										Newline(Ctx);
+										const float AlphaColor = lerp(Alpha0, Alpha1, Alpha);
+										const float3 Color0 = LoadSampleColor(PrimID);
+										const float3 Color1 = LoadSampleColor(PrimID+1);
+										const float3 Color = lerp(Color0, Color1, AlphaColor);
+
+										const float4 Velocity0 = LoadSampleVelocity(PrimID);
+										const float4 Velocity1 = LoadSampleVelocity(PrimID+1);
+										const float4 Velocity = lerp(Velocity0, Velocity1, AlphaColor);
+
+
+										// Fill in sample data
+										Sample.Depth = Depth;
+										Sample.Coord = IntraTileCoord;
+
+										// Compute coverage
+										// Minimal radius to snap the strand to a sample/pixel center (to avoid aliasing)
+										const bool bIsOrthoView = false;
+										const float SceneDistance = ConvertFromDeviceZ(Depth);
+										const float MinHairRadius = bIsOrthoView ? RadiusAtDepth1 : (SceneDistance * RadiusAtDepth1);
+										const float HairRadius = lerp(Rad0, Rad1, AlphaColor);
+										Sample.Coverage = saturate(HairRadius / MinHairRadius);
+
+										// Write data
+										WriteMask(IntraTileCoord, 1u << GroupThread1D);
+
+										const uint JCoord = bIsSteep ? IntraTileCoord.y : IntraTileCoord.x;
+										s_Color[GroupThread1D][JCoord] = PackColorData(Color);
+										s_Data[GroupThread1D][JCoord]  = PackSampleData(Sample);
+										s_Velocity[GroupThread1D][JCoord] = PackVelocityData(Velocity);
+
+
+										// DEBUG (Write coord)
+										#if 0 && PERMUTATION_DEBUG
+										if (bDebugEnable)
+										{
+											Print(Ctx, J, FontYellow, 2, 0);
+											Print(Ctx, Alpha, FontWhite);
+											Print(Ctx, AlphaColor, FontBlue);
+											Print(Ctx, HairRadius, FontRed);
+											Print(Ctx, MinHairRadius, FontGreen);
+											Print(Ctx, Sample.Coverage, FontBlue);
+											Newline(Ctx);
+										}
+										#endif
 									}
-									#endif
 								}
 							}
-						}
-					} // if (bIsSegmentValid)
-				}
-				GroupMemoryBarrierWithGroupSync();
-
-				// DEBUG (Coverage mask)
-				#if 0 && PERMUTATION_DEBUG
-				if (bDebugEnable)
-				{
-					PrintCoverage(Ctx);
-				}
-				#endif
+						} // if (bIsSegmentValid)
+					}
+					GroupMemoryBarrierWithGroupSync();
 
-				// 4.2 Combine all samples within the same pixel (1 thread = 1 pixel)
-				{
-					const uint ValidSegments = ReadMask(GroupThread2D);//s_Mask[Thread_Coord.x][Thread_Coord.y];
+					// DEBUG (Coverage mask)
+					#if 0 && PERMUTATION_DEBUG
+					if (bDebugEnable)
+					{
+						PrintCoverage(Ctx);
+					}
+					#endif
 
-					// Change this loop into to bit logic?
+					// 4.2 Combine all samples within the same pixel (1 thread = 1 pixel)
 					{
-						for (uint SegmentIt=0; SegmentIt < RASTER_SEGMENT_COUNT; ++SegmentIt)
+						const uint ValidSegments = ReadMask(GroupThread2D);//s_Mask[Thread_Coord.x][Thread_Coord.y];
+
+						// Change this loop into to bit logic?
 						{
-							const uint SegmentBit = 1u << SegmentIt;
-							const bool bIsValid = (ValidSegments & SegmentBit) > 0;
-							const bool bIsSteep = (s_bDataOrder & SegmentBit) > 0;
-							if (bIsValid)
+							for (uint SegmentIt=0; SegmentIt < RASTER_SEGMENT_COUNT; ++SegmentIt)
 							{
-								const uint J = bIsSteep ? GroupThread2D.y : GroupThread2D.x;
+								const uint SegmentBit = 1u << SegmentIt;
+								const bool bIsValid = (ValidSegments & SegmentBit) > 0;
+								const bool bIsSteep = (s_bDataOrder & SegmentBit) > 0;
+								if (bIsValid)
+								{
+									const uint J = bIsSteep ? GroupThread2D.y : GroupThread2D.x;
 
-								const FSampleData Sample = UnpackSampleData(s_Data[SegmentIt][J]);
-								const float3 Color = UnpackColorData(s_Color[SegmentIt][J]);
-								const float4 Velocity = UnpackVelocityData(s_Velocity[SegmentIt][J]);
+									const FSampleData Sample = UnpackSampleData(s_Data[SegmentIt][J]);
+									const float3 Color = UnpackColorData(s_Color[SegmentIt][J]);
+									const float4 Velocity = UnpackVelocityData(s_Velocity[SegmentIt][J]);
 
-								const float AccTransmittance = (1.f-Thread_Coverage);
-								Thread_Coverage += AccTransmittance * Sample.Coverage;
-								Thread_Color    += AccTransmittance * Sample.Coverage * Color;
+									const float AccTransmittance = (1.f-Thread_Coverage);
+									Thread_Coverage += AccTransmittance * Sample.Coverage;
+									Thread_Color    += AccTransmittance * Sample.Coverage * Color;
 
-								// Use the closest valid segment as output velocity
-								if (Thread_Velocity.x <= INVALID_VELOCITY)
-								{
-									Thread_Velocity = Velocity;
+									// Use the closest valid segment as output velocity
+									if (Thread_Velocity.x <= INVALID_VELOCITY)
+									{
+										Thread_Velocity = Velocity;
+									}
 								}
 							}
 						}
-					}
 
-					const float CoverageThreshold = 0.95f;
-					if (Thread_Complete == 0 && Thread_Coverage > CoverageThreshold)
-					{
-						InterlockedAdd(s_OpaqueMask, 1);
-						Thread_Complete = 1;
-					}
+						const float CoverageThreshold = 0.95f;
+						if (Thread_Complete == 0 && Thread_Coverage > CoverageThreshold)
+						{
+							InterlockedAdd(s_OpaqueMask, 1);
+							Thread_Complete = 1;
+						}
 
-					#if PERMUTATION_DEBUG
-					if (Thread_LoopCountToFullCoverage == 0 && Thread_Coverage > CoverageThreshold)
-					{
-						Thread_LoopCountToFullCoverage = LoopIndex;
+						#if PERMUTATION_DEBUG
+						if (Thread_LoopCountToFullCoverage == 0 && Thread_Coverage > CoverageThreshold)
+						{
+							Thread_LoopCountToFullCoverage = LoopIndex;
+						}
+						#endif
 					}
-					#endif
-				}
-				//GroupMemoryBarrierWithGroupSync();
-
-			} // for (uint LoopIndex =...) ...
-
-			#if PERMUTATION_DEBUG
-			s_Coverage[GroupThread1D] = Thread_Coverage; // For sync s_Coverage
-			#endif
+					//GroupMemoryBarrierWithGroupSync();
 
-			GroupMemoryBarrierWithGroupSync();
+				} // for (uint LoopIndex =...) ...
 
-			// 4.2. Write final color
-			const bool bLastBinTile = FrontToBackIndex == (s_BinTileCount - 1);
-			const bool bWriteOut = Thread_Coverage > 0 && (bLastBinTile /* || raster tile is fully covered */);
-			if (bWriteOut)
-			{
-				const uint2 PixelCoord = GroupThread2D + RasterTileMin;
-				const bool bOutput = true;//View.GeneralPurposeTweak > 1;
-				if (bOutput)
-				{
-					const float3 SourceColor = OutSceneColorTexture[PixelCoord].xyz;
-					OutSceneColorTexture[PixelCoord] = float4(SourceColor * (1-Thread_Coverage) + Thread_Color, 1);
-					if (Thread_Velocity.x > INVALID_VELOCITY)
-					{
-						OutSceneVelocityTexture[PixelCoord] = Thread_Velocity;
-					}
-					//OutSceneColorTexture[PixelCoord] = float4(lerp(SourceColor, Thread_Color, Thread_Coverage), 1);
-					//OutSceneColorTexture[PixelCoord] = float4(Thread_Color, 1);
-					//OutSceneColorTexture[PixelCoord] = float4(1, 1, 0, 1);
-					//OutSceneColorTexture[PixelCoord] = float4(PixelCoord / OutputResolutionf,0,1);
-					//OutSceneColorTexture[PixelCoord] = float4(GroupThread2D / float2(RASTER_TILE_SIZE, RASTER_TILE_SIZE), 0, 1);
-				}
-
-				// For debug purpose only
 				#if PERMUTATION_DEBUG
-				{
-					InterlockedAdd(OutHairCountTexture_ForDebug[PixelCoord], 1);
-					InterlockedAdd(OutHairPixelCountPerTile_ForDebug[RasterTileCoord], 1);
-				}
+				s_Coverage[GroupThread1D] = Thread_Coverage; // For sync s_Coverage
 				#endif
-			}
 
-			#if PERMUTATION_DEBUG
-			if (bDebugEnable)
-			{
-				Newline(Ctx);
-				Print(Ctx, TEXT("Coverage: "), FontWhite);
-				Print(Ctx, Thread_Coverage, FontYellow);
-				Newline(Ctx);
+				GroupMemoryBarrierWithGroupSync();
 
-				Print(Ctx, TEXT("Full Coverage: "), FontWhite);
-				Print(Ctx, s_OpaqueMask, FontYellow, 2, 0);
-				Newline(Ctx);
+			} // 4. Raster
 
-				Newline(Ctx);
-				for (uint y = 0; y < RASTER_TILE_SIZE; ++y)
-				{
-					for (uint x = 0; x < RASTER_TILE_SIZE; ++x)
-					{
-						const uint Index = x + y * RASTER_TILE_SIZE;
-						FFontColor Color = InitFontColor(ColorMapTurbo(s_Coverage[Index]));
-						//Print(Ctx, s_Coverage[Index], Color, 4, 2);
-							
-						Print(Ctx, Thread_LoopCountToFullCoverage, Color, 2, 0);
-					}
-					Newline(Ctx);
-					Newline(Ctx);
-				}
-				Newline(Ctx);
-			}
+		} // for (... FrontToBackIndex ...)
 
-			if (GroupThread1D == 0)
+		// 5. Write final color
+		const bool bWriteOut = Thread_Coverage > 0;
+		if (bWriteOut)
+		{
+			const uint2 PixelCoord = Thread_PixelCoord;
+			const float3 SourceColor = OutSceneColorTexture[PixelCoord].xyz;
+			OutSceneColorTexture[PixelCoord] = float4(SourceColor * (1-Thread_Coverage) + Thread_Color, 1);
+			if (Thread_Velocity.x > INVALID_VELOCITY)
 			{
-				// Full/Partial coverage tiles
-				if (0)
-				{
-					const bool bFullCover = s_OpaqueMask == 64;
-					if (bFullCover)	AddFilledQuadSS(RasterTileMin, RasterTileMax, Transparent(ColorLightRed));
-					//else			AddFilledQuadSS(RasterTileMin, RasterTileMax, Transparent(ColorLightGreen));
-				}
-
-				//InterlockedAdd(OutHairPixelCountPerTile_ForDebug[RasterTileCoord], PrimCount);
-				OutHairPixelCountPerTile_ForDebug[RasterTileCoord] = PrimCount;
+				OutSceneVelocityTexture[PixelCoord] = Thread_Velocity;
+			}
+			
+			// For debug purpose only
+			#if PERMUTATION_DEBUG
+			{
+				const uint2 RasterTileCoord = PixelCoord / RASTER_TILE_SIZE;
+				InterlockedAdd(OutHairCountTexture_ForDebug[PixelCoord], 1);
+				InterlockedAdd(OutHairPixelCountPerTile_ForDebug[RasterTileCoord], 1);
 			}
 			#endif
+		}
 
-		} // 4. Raster
-
-		} // for (... FrontToBackIndex ...)
-
-	}// Work item
+	} // for ( ... Work item ... )
 }
 
 #endif //SHADER_RASTERCOMPUTE_RASTER
