* Change opaque pixel tracking using bit logic rather than pixel counting
* Exit front to tile raster when reaching full coverage

diff --git a/Engine/Shaders/Private/HairStrands/HairStrandsForwardRaster.usf b/Engine/Shaders/Private/HairStrands/HairStrandsForwardRaster.usf
--- a/Engine/Shaders/Private/HairStrands/HairStrandsForwardRaster.usf
+++ b/Engine/Shaders/Private/HairStrands/HairStrandsForwardRaster.usf
@@ -1070,7 +1070,7 @@ groupshared uint   s_SegmentsAlloc[RASTER_DEPTH_BUCKET];
 
 //groupshared uint   s_Mask[RASTER_THREAD_COUNT];
 groupshared uint   s_Mask[8][8];
-groupshared uint   s_OpaqueMask;
+groupshared uint2  s_OpaqueMask;
 groupshared uint   s_Data[RASTER_SEGMENT_COUNT][8];
 groupshared uint   s_Color[RASTER_SEGMENT_COUNT][8];
 groupshared uint2  s_Velocity[RASTER_SEGMENT_COUNT][8];
@@ -1136,10 +1136,25 @@ void PrintPixelMask(inout FShaderPrintContext Ctx, uint InMask)
 		}
 		else
 		{
+			Print(Ctx, TEXT("."), FontRed);
+		}
+	}
+}
 
+void PrintOpaqueMask(inout FShaderPrintContext Ctx, uint2 InOpaqueMask)
+{
+	for (uint y = 0; y < RASTER_TILE_SIZE; ++y)
+	for (uint x = 0; x < RASTER_TILE_SIZE; ++x)
+	{
+		const uint s = x + y * RASTER_TILE_SIZE;
+		bool bValid = false;
+		if (s < 32) { bValid = ((1u << s) & InOpaqueMask.x) != 0u; }
+		else { bValid = ((1u << (s - 32)) & InOpaqueMask.y) != 0u; }
 
-			Print(Ctx, TEXT("."), FontRed);
-		}		
+		if (bValid) { Print(Ctx, TEXT("x "), FontGreen); }
+		else { Print(Ctx, TEXT(". "), FontRed); }
+
+		if (x == RASTER_TILE_SIZE-1) Newline(Ctx);
 	}
 }
 #endif
@@ -1249,6 +1264,10 @@ void RasterCS(
 	LOOP
 	for (uint WorkIndex = 0; WorkIndex < MAX_WORK_COUNT; WorkIndex++)
 	{
+		#if PERMUTATION_DEBUG
+		FShaderPrintContext GlobalCtx = InitShaderPrintContext(false, uint2(0, 0));
+		#endif
+
 		// 0.1 Fetch work item
 		if (GroupThread1D == 0)
 		{
@@ -1267,6 +1286,7 @@ void RasterCS(
 
 			s_BinTileOffset = Work.x;
 			s_BinTileCount  = Work.y;
+			s_OpaqueMask    = 0;
 		}
 
 		GroupMemoryBarrierWithGroupSync();
@@ -1292,9 +1312,9 @@ void RasterCS(
 
 		// Iterate over all bins for the current raster/screen tile
 		const uint FrontToBackCount = min(s_BinTileCount, 64);
+		uint ExitFrontToBackIndex = FrontToBackCount;
 		for (uint FrontToBackIndex = 0; FrontToBackIndex < FrontToBackCount; ++FrontToBackIndex)
 		{
-
 			// 0.2 Reset all LDS variables
 			{
 				if (GroupThread1D == 0)
@@ -1303,7 +1323,6 @@ void RasterCS(
 					s_Segments_Max = 0;
 					s_Segments_ValidCount = 0;
 					s_bDataOrder = 0;
-					s_OpaqueMask = 0;
 				}
 				if (GroupThread1D < RASTER_DEPTH_BUCKET)
 				{
@@ -1358,6 +1377,12 @@ void RasterCS(
 
 			Thread_PixelCoord = RasterTileMin + GroupThread2D;
 
+			if (all(s_OpaqueMask == 0xFFFFFFFF))
+			{
+				ExitFrontToBackIndex = FrontToBackIndex;
+				break;
+			}
+
 			// For debug only
 			#if PERMUTATION_DEBUG
 			const bool bDebugEnableAll	= all(RasterTileCoord == uint2(ShaderPrintData.CursorCoord / float(RASTER_TILE_SIZE)));
@@ -1577,7 +1602,8 @@ void RasterCS(
 
 					GroupMemoryBarrierWithGroupSync();
 
-					if (s_OpaqueMask == 64)
+					// If raster tile is fully covered, exit
+					if (all(s_OpaqueMask == 0xFFFFFFFF))
 					{
 						break;
 					}
@@ -1852,7 +1878,15 @@ void RasterCS(
 						const float CoverageThreshold = 0.95f;
 						if (Thread_Complete == 0 && Thread_Coverage > CoverageThreshold)
 						{
-							InterlockedAdd(s_OpaqueMask, 1);
+							// Mark pixel has fully covered
+							if (GroupThread1D >= 32)
+							{
+								InterlockedOr(s_OpaqueMask.y, 1u << (GroupThread1D - 32));
+							}
+							else
+							{
+								InterlockedOr(s_OpaqueMask.x, 1u << GroupThread1D);
+							}
 							Thread_Complete = 1;
 						}
 
@@ -1869,6 +1903,7 @@ void RasterCS(
 
 				#if PERMUTATION_DEBUG
 				s_Coverage[GroupThread1D] = Thread_Coverage; // For sync s_Coverage
+				GlobalCtx = Ctx;
 				#endif
 
 				GroupMemoryBarrierWithGroupSync();
@@ -1895,6 +1930,16 @@ void RasterCS(
 				const uint2 RasterTileCoord = PixelCoord / RASTER_TILE_SIZE;
 				InterlockedAdd(OutHairCountTexture_ForDebug[PixelCoord], 1);
 				InterlockedAdd(OutHairPixelCountPerTile_ForDebug[RasterTileCoord], 1);
+
+				const FFontColor FontLegend = FontWhite;
+				const FFontColor FontValue = FontOrange;
+				Print(GlobalCtx, TEXT("Exit :"), FontLegend);
+				Print(GlobalCtx, ExitFrontToBackIndex, FontValue, 3, 0);
+				Print(GlobalCtx, TEXT(" / "), FontLegend);
+				Print(GlobalCtx, FrontToBackCount, FontValue, 3, 0);
+				Newline(GlobalCtx);
+
+				PrintOpaqueMask(GlobalCtx, s_OpaqueMask);
 			}
 			#endif
 		}
