Fix missing steep bit reset in raster loop causing random wrong pixel composition.

diff --git a/Engine/Shaders/Private/HairStrands/HairStrandsForwardRaster.usf b/Engine/Shaders/Private/HairStrands/HairStrandsForwardRaster.usf
--- a/Engine/Shaders/Private/HairStrands/HairStrandsForwardRaster.usf
+++ b/Engine/Shaders/Private/HairStrands/HairStrandsForwardRaster.usf
@@ -523,7 +523,7 @@ void BinningCS(uint GroupThreadID : SV_GroupThreadID, uint GroupID : SV_GroupID)
 	GroupMemoryBarrierWithGroupSync();
 
 	#if PERMUTATION_DEBUG
-	const bool bDebugEnabled = GroupID == 0 && GroupThreadID <= 64;
+	const bool bDebugEnabled = false && GroupID == 0 && GroupThreadID <= 64;
 	FShaderPrintContext Ctx = InitShaderPrintContext(bDebugEnabled, uint2(250, 50));
 	#endif
 
@@ -816,7 +816,7 @@ void CompactionCS(uint DispatchThreadID : SV_DispatchThreadID, uint GroupThreadI
 
 	#if PERMUTATION_DEBUG
 	const uint2 TileCoord = UnpackVisTileCoord(PackedCoord);
-	const bool bDebugEnabled = false; //all(TileCoord == uint2(ShaderPrintData.CursorCoord / float(BIN_TILE_SIZE)));
+	const bool bDebugEnabled = false && all(TileCoord == uint2(ShaderPrintData.CursorCoord / float(BIN_TILE_SIZE)));
 	FShaderPrintContext Ctx = InitShaderPrintContext(bDebugEnabled, uint2(750, 50));
 	#endif
 
@@ -1112,7 +1112,7 @@ void PrintCoverage(inout FShaderPrintContext Ctx)
 		for (uint x = 0; x < RASTER_TILE_SIZE; ++x)
 		{
 			const uint ValidSegments = ReadMask(uint2(x,y)); //s_Mask[x][y];
-			if (ValidSegments > 0)
+			if (ValidSegments != 0)
 			{
 				Print(Ctx, TEXT("x "), FontGreen);
 			}
@@ -1402,7 +1402,7 @@ void RasterCS(
 				Print(Ctx, TEXT("Work - Offset:"), FontLegend); Print(Ctx, s_BinTileOffset, FontValue); Newline(Ctx);
 				Print(Ctx, TEXT("Work - Count :"), FontLegend); Print(Ctx, s_BinTileCount, FontValue); Newline(Ctx);
 
-				AddFilledQuadSS(BinTileMin, BinTileMax, Transparent(ColorLightGreen));
+				//AddFilledQuadSS(BinTileMin, BinTileMax, Transparent(ColorLightGreen));
 				//AddFilledQuadSS(RasterTileMin, RasterTileMax, Transparent(ColorLightGreen));
 				AddQuadSS(RasterTileMin, RasterTileMax, Transparent(ColorYellow));
 			}
@@ -1450,7 +1450,7 @@ void RasterCS(
 							InterlockedMax(s_Segments_Max, uDepth);
 						}
 
-						#if 1 && PERMUTATION_DEBUG
+						#if 0 && PERMUTATION_DEBUG
 						if (bDebugEnableAll && PrimID != INVALID_PRIM_ID)
 						{
 							const float3 P0 = UnpackHairControlPoint(PrimID).Position;
@@ -1599,6 +1599,7 @@ void RasterCS(
 						}
 					}
 					s_Mask[GroupThread2D.x][GroupThread2D.y] = 0;
+					s_bDataOrder = 0;
 
 					GroupMemoryBarrierWithGroupSync();
 
@@ -1852,8 +1853,8 @@ void RasterCS(
 							for (uint SegmentIt=0; SegmentIt < RASTER_SEGMENT_COUNT; ++SegmentIt)
 							{
 								const uint SegmentBit = 1u << SegmentIt;
-								const bool bIsValid = (ValidSegments & SegmentBit) > 0;
-								const bool bIsSteep = (s_bDataOrder & SegmentBit) > 0;
+								const bool bIsValid = (ValidSegments & SegmentBit) != 0;
+								const bool bIsSteep = (s_bDataOrder & SegmentBit) != 0;
 								if (bIsValid)
 								{
 									const uint J = bIsSteep ? GroupThread2D.y : GroupThread2D.x;
@@ -1862,7 +1863,7 @@ void RasterCS(
 									const float3 Color = UnpackColorData(s_Color[SegmentIt][J]);
 									const float4 Velocity = UnpackVelocityData(s_Velocity[SegmentIt][J]);
 
-									const float AccTransmittance = (1.f-Thread_Coverage);
+									const float AccTransmittance = saturate(1.f-Thread_Coverage);
 									Thread_Coverage += AccTransmittance * Sample.Coverage;
 									Thread_Color    += AccTransmittance * Sample.Coverage * Color;
 
@@ -2048,8 +2049,8 @@ void MainCS(uint3 ThreadId : SV_DispatchThreadID)
 		const FFontColor AllocColor			 = InitFontColor(ColorMapTurbo(VisTileArgs[0] / float(CPUAllocatedTileCount)));
 		const FFontColor AllocCompactedColor = InitFontColor(ColorMapTurbo(CompactedVisTileArgs[0] / float(CPUAllocatedCompactedTileCount)));
 
-		Print(Context, TEXT("Alloc. Tile          : "), FontLegend); Print(Context, VisTileArgs[0], AllocColor); Print(Context, TEXT(" / "), FontLegend); Print(Context, CPUAllocatedTileCount, FontValue); Newline(Context);
-		Print(Context, TEXT("Alloc. Compacted Tile: "), FontLegend); Print(Context, CompactedVisTileArgs[0], AllocColor); Print(Context, TEXT(" / "), FontLegend); Print(Context, CPUAllocatedCompactedTileCount, FontValue); Newline(Context);
+		Print(Context, TEXT("Alloc. Tile          : "), FontLegend); Print(Context, VisTileArgs[0], AllocColor, 6, 0); Print(Context, TEXT(" / "), FontLegend); Print(Context, CPUAllocatedTileCount, FontValue,6,0); Newline(Context);
+		Print(Context, TEXT("Alloc. Compacted Tile: "), FontLegend); Print(Context, CompactedVisTileArgs[0], AllocColor, 6, 0); Print(Context, TEXT(" / "), FontLegend); Print(Context, CPUAllocatedCompactedTileCount, FontValue, 6, 0); Newline(Context);
 		Print(Context, TEXT("Rasterizer Max Work  : "), FontLegend); Print(Context, NumRasterizers * MAX_WORK_COUNT, FontValue); Newline(Context);
 		Newline(Context);
 		
