* Add depth range struct for clarifying binning/raster code.
* Store bin tile depth range into the compacted bin tile data.

diff --git a/Engine/Shaders/Private/HairStrands/HairStrandsForwardRaster.usf b/Engine/Shaders/Private/HairStrands/HairStrandsForwardRaster.usf
--- a/Engine/Shaders/Private/HairStrands/HairStrandsForwardRaster.usf
+++ b/Engine/Shaders/Private/HairStrands/HairStrandsForwardRaster.usf
@@ -65,14 +65,24 @@ float UnpackDepth(uint In)
 	return asfloat(In);
 }
 
-uint PackDepthRange(float2 In)
+struct FDepthRange
 {
-	return PackRG16(In);
+	float MinZ;
+	float MaxZ;
+};
+
+uint PackDepthRange(FDepthRange In)
+{
+	return PackRG16(float2(In.MinZ, In.MaxZ));
 }
 
-float2 UnpackDepthRange(uint In)
+FDepthRange UnpackDepthRange(uint In)
 {
-	return UnpackRG16(In);
+	FDepthRange Out;
+	const float2 D = UnpackRG16(In);
+	Out.MinZ = D.x;
+	Out.MaxZ = D.y;
+	return Out;
 }
 
 uint PackWork(uint InTileIndex, uint InTileCount)
@@ -677,8 +687,9 @@ void BinningCS(uint GroupThreadID : SV_GroupThreadID, uint GroupID : SV_GroupID)
 			const uint TotalOldWriteCount = OutVisTileBinningGrid[uint3(TileCoord, TmpSegmentCountLayerIdx)];
 
 			#if BIN_MINMAX
-			const float TileMinZ = UnpackDepth(OutVisTileBinningGridMinZ[uint3(TileCoord, SegmentCountLayerIdx)]);
-			const float TileMaxZ = UnpackDepth(OutVisTileBinningGridMaxZ[uint3(TileCoord, SegmentCountLayerIdx)]);
+			FDepthRange TileDepthRange;
+			TileDepthRange.MinZ = UnpackDepth(OutVisTileBinningGridMinZ[uint3(TileCoord, SegmentCountLayerIdx)]);
+			TileDepthRange.MaxZ = UnpackDepth(OutVisTileBinningGridMaxZ[uint3(TileCoord, SegmentCountLayerIdx)]);
 			#endif
 
 			uint NewTile;
@@ -689,7 +700,7 @@ void BinningCS(uint GroupThreadID : SV_GroupThreadID, uint GroupID : SV_GroupID)
 			StoreOutVisTileData(OutVisTileData, NewTile, VT_MinWriteIndex, TotalNewWriteCount & ~1023u);
 			// Min/Max depth
 			#if BIN_MINMAX
-			StoreOutVisTileData(OutVisTileData, NewTile, VT_MinMaxDepth, PackDepthRange(float2(TileMinZ, TileMaxZ)));
+			StoreOutVisTileData(OutVisTileData, NewTile, VT_MinMaxDepth, PackDepthRange(TileDepthRange));
 			#endif
 
 			const uint PrevTile = (OutVisTileBinningGrid[uint3(TileCoord, TileAllocInfoLayerIdx)] & 0xffff);
@@ -830,9 +841,9 @@ void CompactionCS(uint DispatchThreadID : SV_DispatchThreadID, uint GroupThreadI
 			{
 				LocalPrimCount += LoadVisTileData(InData, TileIdx, VT_PrimCount);
 
-				const float2 LocalDepthRange = UnpackDepthRange(LoadVisTileData(InData, TileIdx, VT_MinMaxDepth));
-				InterlockedMin(group_MinZ, PackDepth(LocalDepthRange.x));
-				InterlockedMax(group_MaxZ, PackDepth(LocalDepthRange.y));
+				const FDepthRange LocalDepthRange = UnpackDepthRange(LoadVisTileData(InData, TileIdx, VT_MinMaxDepth));
+				InterlockedMin(group_MinZ, PackDepth(LocalDepthRange.MinZ));
+				InterlockedMax(group_MaxZ, PackDepth(LocalDepthRange.MaxZ));
 				
 				uint WritePos;
 				WaveInterlockedAddScalar_(group_NumTiles, 1, WritePos);
@@ -881,6 +892,11 @@ void CompactionCS(uint DispatchThreadID : SV_DispatchThreadID, uint GroupThreadI
 			const uint PrimCount = min(TotalPrimCount - TileIdx * 1024, 1024);
 			StoreOutVisTileData(OutData, CompactedTile, VT_PrimCount, PrimCount);
 			StoreOutVisTileData(OutData, CompactedTile, VT_Coord, PackedCoord);
+
+			FDepthRange DepthRange;
+			DepthRange.MinZ = group_MinZ;
+			DepthRange.MinZ = group_MaxZ;
+			StoreOutVisTileData(OutData, CompactedTile, VT_MinMaxDepth, PackDepthRange(DepthRange));
 		}
 	}
 	GroupMemoryBarrierWithGroupSync();
