Change sample data encoding to improve coverage precision (8bits -> 10bits).

diff --git a/Engine/Shaders/Private/HairStrands/HairStrandsForwardRaster.usf b/Engine/Shaders/Private/HairStrands/HairStrandsForwardRaster.usf
--- a/Engine/Shaders/Private/HairStrands/HairStrandsForwardRaster.usf
+++ b/Engine/Shaders/Private/HairStrands/HairStrandsForwardRaster.usf
@@ -1177,22 +1177,22 @@ void PrintOpaqueMask(inout FShaderPrintContext Ctx, uint2 InOpaqueMask)
 
 struct FSampleData
 {
-	uint2 Coord; // Pixel coord within the tile 8x8 (4bits x 4bits)
-	float Depth;
-	float Coverage;
+	uint2 Coord;	// 2x3bits - Pixel coord within the tile 8x8
+	float Depth;	// 16 bits
+	float Coverage;	// 10 bits
 };
 
 uint PackSampleData(FSampleData In)
 {
-	return f32tof16(In.Depth) | (uint(saturate(In.Coverage) * 0xFF)<<16u) | ((In.Coord.x & 0xF)<<24u) | ((In.Coord.y & 0xF) << 28u);
+	return f32tof16(In.Depth) | (uint(saturate(In.Coverage) * 0x3FF) << 16u) | ((In.Coord.x & 0x7) << 26u) | ((In.Coord.y & 0x7) << 29u);
 }
 
 FSampleData UnpackSampleData(uint In)
 {
 	FSampleData Out;
 	Out.Depth = f16tof32(In & 0xFFFF);
-	Out.Coverage = ((In>>16) & 0xFF) * (1.f/255.f);
-	Out.Coord =  uint2((In>>24) & 0xF, (In>>28) & 0xF);
+	Out.Coverage = ((In >> 16) & 0x3FF) * (1.f / 1023.f);
+	Out.Coord = uint2((In >> 26) & 0x7, (In >> 29) & 0x7);
 	return Out;
 }
 
@@ -2066,7 +2066,7 @@ void MainCS(uint3 ThreadId : SV_DispatchThreadID)
 		const FFontColor AllocCompactedColor = InitFontColor(ColorMapTurbo(CompactedVisTileArgs[0] / float(CPUAllocatedCompactedTileCount)));
 
 		Print(Context, TEXT("Alloc. Tile          : "), FontLegend); Print(Context, VisTileArgs[0], AllocColor, 6, 0); Print(Context, TEXT(" / "), FontLegend); Print(Context, CPUAllocatedTileCount, FontValue,6,0); Newline(Context);
-		Print(Context, TEXT("Alloc. Compacted Tile: "), FontLegend); Print(Context, CompactedVisTileArgs[0], AllocColor, 6, 0); Print(Context, TEXT(" / "), FontLegend); Print(Context, CPUAllocatedCompactedTileCount, FontValue, 6, 0); Newline(Context);
+		Print(Context, TEXT("Alloc. Compacted Tile: "), FontLegend); Print(Context, CompactedVisTileArgs[0], AllocCompactedColor, 6, 0); Print(Context, TEXT(" / "), FontLegend); Print(Context, CPUAllocatedCompactedTileCount, FontValue, 6, 0); Newline(Context);
 		Print(Context, TEXT("Rasterizer Max Work  : "), FontLegend); Print(Context, NumRasterizers * MAX_WORK_COUNT, FontValue); Newline(Context);
 		Newline(Context);
 		
