* Fix invalid binning of segments at screen border due to clamping error.
* Fix rasterization of invalid segments due to missing reset, and incorrect segments count check.

diff --git a/Engine/Shaders/Private/HairStrands/HairStrandsForwardRaster.usf b/Engine/Shaders/Private/HairStrands/HairStrandsForwardRaster.usf
--- a/Engine/Shaders/Private/HairStrands/HairStrandsForwardRaster.usf
+++ b/Engine/Shaders/Private/HairStrands/HairStrandsForwardRaster.usf
@@ -585,8 +585,8 @@ void BinningCS(uint GroupThreadID : SV_GroupThreadID, uint GroupID : SV_GroupID)
 					SP1.xy *= BIN_RCP_TILE_SIZE;
 
 					// For peace of mind, make sure these are actually clamped to a valid range.
-					SP0 = clamp(SP0, 0.0f, float3(BinTileRes, 1.0f));
-					SP1 = clamp(SP1, 0.0f, float3(BinTileRes, 1.0f));
+					SP0 = clamp(SP0, 0.0f, float3(BinTileRes-0.01f, 1.0f));
+					SP1 = clamp(SP1, 0.0f, float3(BinTileRes-0.01f, 1.0f));
 
 					MaxZ = PackDepth(max(SP0.z, SP1.z));
 					MinZ = PackDepth(min(SP0.z, SP1.z));
@@ -1420,6 +1420,7 @@ void RasterCS(
 
 				//AddFilledQuadSS(BinTileMin, BinTileMax, Transparent(ColorLightGreen));
 				//AddFilledQuadSS(RasterTileMin, RasterTileMax, Transparent(ColorLightGreen));
+				AddQuadSS(BinTileMin, BinTileMax, Transparent(ColorLightGreen));
 				AddQuadSS(RasterTileMin, RasterTileMax, Transparent(ColorYellow));
 			}
 			#endif
@@ -1430,6 +1431,11 @@ void RasterCS(
 				for (uint LoopIndex = 0; LoopIndex < LoopCount64; ++LoopIndex)
 				{
 					const uint Prim = GroupThread1D + LoopIndex * RASTER_THREAD_COUNT;
+
+					// Need to reset s_Segments and s_Segments_Sorted, as they are used for tracking valid segments
+					s_Segments[Prim] = uint2(0, INVALID_PRIM_ID);
+					s_Segments_Sorted[Prim] = uint2(0, INVALID_PRIM_ID);
+
 					if (Prim < PrimCount)
 					{
 						uint PrimID = VisTilePrims[PrimOffset + Prim];
@@ -1598,6 +1604,11 @@ void RasterCS(
 			}
 			#endif
 
+			if (LoopCount32 == 0)
+			{
+				continue;
+			}
+
 			// 4. Raster segments
 			{
 				// 4.1 Loop over all segment within the tile, and rastize 32 of them at each loop
@@ -1628,7 +1639,7 @@ void RasterCS(
 					// 4.1 Raster segment (1 thread = 1 segment)
 					// Half of the thread are doing nothing, we could raster the two half of the segments (one per each thread)
 					const uint Prim = GroupThread1D + LoopIndex * RASTER_SEGMENT_COUNT;
-					if (GroupThread1D < RASTER_SEGMENT_COUNT && Prim < PrimCount)
+					if (GroupThread1D < RASTER_SEGMENT_COUNT && Prim < s_Segments_ValidCount)
 					{
 						const uint PrimID = s_Segments_Sorted[Prim].y;
 						uint TypeDummy = 0;
