Increase depth bins sorting for compaction and rasterization, to remove flickering happening on overlapping distant grooms.

diff --git a/Engine/Shaders/Private/HairStrands/HairStrandsForwardRaster.usf b/Engine/Shaders/Private/HairStrands/HairStrandsForwardRaster.usf
--- a/Engine/Shaders/Private/HairStrands/HairStrandsForwardRaster.usf
+++ b/Engine/Shaders/Private/HairStrands/HairStrandsForwardRaster.usf
@@ -787,7 +787,7 @@ groupshared uint group_MaxLDSTileIdx;
 groupshared uint group_MinZ;
 groupshared uint group_MaxZ;
 
-#define COMPACTION_DEPTH_BUCKET 64
+#define COMPACTION_DEPTH_BUCKET 1024
 groupshared uint s_BinOffset[COMPACTION_DEPTH_BUCKET];
 groupshared uint s_BinCount[COMPACTION_DEPTH_BUCKET];
 
@@ -1042,7 +1042,7 @@ void CompactionCS(uint DispatchThreadID : SV_DispatchThreadID, uint GroupThreadI
 
 ///////////////////////////////////////////////////////////////////////////
 
-#define RASTER_DEPTH_BUCKET 32
+#define RASTER_DEPTH_BUCKET 64
 #define RASTER_SEGMENT_COUNT 32
 #define SEGMENT_COUNT_PER_GROUP 1024
 #define INVALID_PRIM_ID 0xFFFFFFFF
@@ -1797,8 +1797,8 @@ void RasterCS(
 							for (int J = 0; J < NumSteps; ++J)
 							//for (int J = 0; J < RASTER_TILE_SIZE; ++J)
 							{
-								const float Alpha = saturate(J * RcpNumSteps);
-								const float4 SP = lerp(SP0, SP1, Alpha);
+								const float AlphaSP = saturate(J * RcpNumSteps);
+								const float4 SP = lerp(SP0, SP1, AlphaSP);
 
 								int2 Coords = SP.xy;
 								const int2 IntraTileCoord = Coords - RasterTileMin;
