Path Tracer: Fix incorrect geometric normal on nanite instances with negative scale

diff --git a/Engine/Shaders/Private/PathTracing/PathTracingMaterialHitShader.usf b/Engine/Shaders/Private/PathTracing/PathTracingMaterialHitShader.usf
--- a/Engine/Shaders/Private/PathTracing/PathTracingMaterialHitShader.usf
+++ b/Engine/Shaders/Private/PathTracing/PathTracingMaterialHitShader.usf
@@ -154,9 +154,12 @@ RAY_TRACING_ENTRY_CLOSEST_HIT(PathTracingMaterialCHS,
 	 * Set material attributes for full materials
 	 **/
 	Payload.TranslatedWorldPos = LWCToFloat(LWCAdd(MaterialParameters.AbsoluteWorldPosition, ResolvedView.PreViewTranslation));
+	float GeoNormalSign = MaterialParameters.TwoSidedSign;
+#if !VF_SUPPORTS_RAYTRACING_PREPARE_MATERIAL_PIXEL_PARAMETERS
 	// Because the geometric normal is computed directly in world space
 	// it doesn't reflect the sign flip from the object transform, so apply it here
-	float GeoNormalSign = MaterialParameters.TwoSidedSign * GetPrimitive_DeterminantSign(MaterialParameters.PrimitiveId);
+	GeoNormalSign *= GetPrimitive_DeterminantSign(MaterialParameters.PrimitiveId);
+#endif
 	Payload.WorldGeoNormal = GeoNormalSign * GeoNormal;
 	Payload.WorldNormal = MaterialParameters.WorldNormal;
 	Payload.WorldSmoothNormal = MaterialParameters.TwoSidedSign * TransformTangentNormalToWorld(MaterialParameters.TangentToWorld, float3(0, 0, 1));
