Path Tracer: Lower the shadow terminator ray bias value

This results in more accurate contact shadows, and can prevent artifacts in digital humans, specifically around the lips. Low polygon meshes may show slightly more faceting artifacts around the shadow terminator as a result.

diff --git a/Engine/Shaders/Private/PathTracing/PathTracingCore.ush b/Engine/Shaders/Private/PathTracing/PathTracingCore.ush
--- a/Engine/Shaders/Private/PathTracing/PathTracingCore.ush
+++ b/Engine/Shaders/Private/PathTracing/PathTracingCore.ush
@@ -620,18 +620,15 @@ FProbeResult TraceProbeRay(FRayDesc Ray)
 	return Result;
 }
 
-void ApplyRayBias(inout FRayDesc Ray, float HitT, float3 Normal, bool ShadowTerminatorFix = false)
+void ApplyRayBias(inout FRayDesc Ray, float HitT, float3 Normal)
 {
 #if 1
 	// Take maximum of position or hit distance to determine "radius" of hit avoidance
 	const float RefValue = max(max(abs(Ray.Origin.x), abs(Ray.Origin.y)), max(abs(Ray.Origin.z), HitT));
 	const uint UlpRadius = 16; // number of floating point ulps to skip around an intersection
 	const float RelativeOffset = asfloat(asuint(RefValue) + UlpRadius) - RefValue;
-	const float ShadowOffset = 0.1; // 1mm
 	const float BaseOffset = 0.001; // 0.01mm (avoid precision issues very close to the origin where ulps become tiny)
-	// for shadow rays, push out the point a bit more so we can reduce shadow terminator issues from poor tesselations
-	// other rays are less likely to arrive at grazing angles, so don't need the extra offset
-	Ray.Origin += max(ShadowTerminatorFix ? ShadowOffset : BaseOffset, RelativeOffset) * Normal;
+	Ray.Origin += max(BaseOffset, RelativeOffset) * Normal;
 #else
 	ApplyPositionBias(Ray.Origin, Ray.Direction, Normal, MaxNormalBias);
 #endif
@@ -1126,13 +1123,12 @@ bool PathTracingKernel(inout FPathState PathState, int Bounce)
 			{
 				// for transmissive materials, bias the position to the other side of the surface if the light is coming from behind
 				const float SignedPositionBias = bIsTransmissiveMaterial ? sign(dot(Payload.WorldGeoNormal, LightSample.Direction)) : 1.0;
-				const bool bEnableShadowTermFix = Payload.ShadingModelID != SHADINGMODELID_HAIR && Payload.ShadingModelID != SHADINGMODELID_MEDIUM;
 				FRayDesc LightRay;
 				LightRay.Origin = TranslatedWorldPos;
 				LightRay.TMin = 0;
 				LightRay.Direction = LightSample.Direction;
 				LightRay.TMax = LightSample.Distance;
-				ApplyRayBias(LightRay, Payload.HitT, SignedPositionBias * Payload.WorldGeoNormal, bEnableShadowTermFix);
+				ApplyRayBias(LightRay, Payload.HitT, SignedPositionBias * Payload.WorldGeoNormal);
 
 				float AvgRoughness = ApproximateCaustics ? GetAverageRoughness(Payload) : 0.0;
 
