Path Tracer: Implement holdout flag on primitives so that users can mark certain objects as being holdouts without having to touch the material

When a holdout primitive is shaded by a camera ray, the regular shading is ignored and the object becomes unlit/holdout instead. All other ray types see the object as normal. The logic takes into account the blend mode, so partialy transparent objects will become partial holdouts. And since the original shader is used, any masking will be maintained.

diff --git a/Engine/Plugins/Experimental/GPULightmass/Shaders/Private/LightmapPathTracing.usf b/Engine/Plugins/Experimental/GPULightmass/Shaders/Private/LightmapPathTracing.usf
--- a/Engine/Plugins/Experimental/GPULightmass/Shaders/Private/LightmapPathTracing.usf
+++ b/Engine/Plugins/Experimental/GPULightmass/Shaders/Private/LightmapPathTracing.usf
@@ -183,11 +183,7 @@ FPathTracingPayload TraceTransparentRay(FRayDesc Ray, bool IsCameraRay, bool Las
 	}
 	for (;;)
 	{
-		// avoid double counting what was captured by the skylight
-		// also avoid noise from hot spots (they can be properly
-		// importance sampled if a capturing skylight is added)
-		const bool bIgnoreSky = !IsCameraRay;
-		FPackedPathTracingPayload PackedPayload = InitPathTracingPayload(bIgnoreSky, 1.0);
+		FPackedPathTracingPayload PackedPayload = InitPathTracingPayload(IsCameraRay, 1.0);
 		// Trace the ray
 		TraceRay(
 			TLAS,
diff --git a/Engine/Plugins/FX/Niagara/Source/Niagara/Private/NiagaraComponent.cpp b/Engine/Plugins/FX/Niagara/Source/Niagara/Private/NiagaraComponent.cpp
--- a/Engine/Plugins/FX/Niagara/Source/Niagara/Private/NiagaraComponent.cpp
+++ b/Engine/Plugins/FX/Niagara/Source/Niagara/Private/NiagaraComponent.cpp
@@ -345,6 +345,7 @@ TUniformBuffer<FPrimitiveUniformShaderParameters>* FNiagaraSceneProxy::GetCustom
 				.OutputVelocity(bOutputVelocity)
 				.CastContactShadow(CastsContactShadow())
 				.CastShadow(CastsDynamicShadow())
+				.Holdout(Holdout())
 				.HasCapsuleRepresentation(HasDynamicIndirectShadowCasterRepresentation())
 				.UseVolumetricLightmap(bHasPrecomputedVolumetricLightmap)
 				.UseSingleSampleShadowFromStationaryLights(UseSingleSampleShadowFromStationaryLights());
diff --git a/Engine/Shaders/Private/PathTracing/PathTracingCommon.ush b/Engine/Shaders/Private/PathTracing/PathTracingCommon.ush
--- a/Engine/Shaders/Private/PathTracing/PathTracingCommon.ush
+++ b/Engine/Shaders/Private/PathTracing/PathTracingCommon.ush
@@ -14,7 +14,7 @@
 #define PATH_TRACING_PAYLOAD_OUTPUT_FLAG_FRONT_FACE                   (1 << 0) // Indicates that ray has hit the front face of a primitive. Set by closest hit shader.
 #define PATH_TRACING_PAYLOAD_OUTPUT_FLAG_DECAL_RECEIVER               (1 << 1) // Indicates that ray has hit a primitive that receives decals. Set by closest hit shader.
 
-#define PATH_TRACING_PAYLOAD_INPUT_FLAG_IGNORE_SKY_MATERIALS          (1 << 0) // Indicates that the AHS/CHS should skip processing if the material is tagged as sky
+#define PATH_TRACING_PAYLOAD_INPUT_FLAG_CAMERA_RAY                    (1 << 0) // Indicates that this is a camera ray which can trigger a number of special behaviors in the closest hit shader
 #define PATH_TRACING_PAYLOAD_INPUT_FLAG_PATHTRACER_VISIBILITY         (1 << 1) // Indicates that this is a path tracer visibility ray (which supports material evaluation for transparent shadows and fake caustics)
 
 #define PATH_TRACING_PAYLOAD_NORMAL_BITS  32
@@ -73,8 +73,8 @@ struct FPackedPathTracingPayload : FMinimalPayload
 	uint GetPrimitiveLightingChannelMask() { return (PackedData[0] >> 29) & 0x7; }
 
 	// Flag methods
-	void SetIgnoreSkyMaterials() { SetFlag(PATH_TRACING_PAYLOAD_INPUT_FLAG_IGNORE_SKY_MATERIALS); }
-	bool IsIgnoreSkyMaterials() { return (GetFlags() & PATH_TRACING_PAYLOAD_INPUT_FLAG_IGNORE_SKY_MATERIALS) != 0; }
+	void SetCameraRay() { SetFlag(PATH_TRACING_PAYLOAD_INPUT_FLAG_CAMERA_RAY); }
+	bool IsCameraRay() { return (GetFlags() & PATH_TRACING_PAYLOAD_INPUT_FLAG_CAMERA_RAY) != 0; }
 
 	void SetVisibilityRay() { SetFlag(PATH_TRACING_PAYLOAD_INPUT_FLAG_PATHTRACER_VISIBILITY); }
 	bool IsVisibilityRay() { return (GetFlags() & PATH_TRACING_PAYLOAD_INPUT_FLAG_PATHTRACER_VISIBILITY) != 0; }
@@ -254,13 +254,13 @@ struct FPathTracingPayload : FMinimalPayload
 	}
 };
 
-FPackedPathTracingPayload InitPathTracingPayload(bool bIgnoreSkyMaterials, float PathRoughness)
+FPackedPathTracingPayload InitPathTracingPayload(bool bIsCameraRay, float PathRoughness)
 {
 	FPackedPathTracingPayload Output = (FPackedPathTracingPayload)0;
 	Output.SetPathRoughness(PathRoughness);
-	if (bIgnoreSkyMaterials)
+	if (bIsCameraRay)
 	{
-		Output.SetIgnoreSkyMaterials();
+		Output.SetCameraRay();
 	}
 	return Output;
 }
diff --git a/Engine/Shaders/Private/PathTracing/PathTracingCore.ush b/Engine/Shaders/Private/PathTracing/PathTracingCore.ush
--- a/Engine/Shaders/Private/PathTracing/PathTracingCore.ush
+++ b/Engine/Shaders/Private/PathTracing/PathTracingCore.ush
@@ -184,11 +184,7 @@ FPathTracingPayload TraceTransparentRay(inout FPathState PathState, bool bIsCame
 	}
 	for (;;)
 	{
-		// avoid double counting what was captured by the skylight
-		// also avoid noise from hot spots (they can be properly
-		// importance sampled if a capturing skylight is added)
-		const bool bIgnoreSky = !bIsCameraRay;
-		FPackedPathTracingPayload PackedPayload = InitPathTracingPayload(bIgnoreSky, ApproximateCaustics ? PathState.PathRoughness : 0.0);
+		FPackedPathTracingPayload PackedPayload = InitPathTracingPayload(bIsCameraRay, ApproximateCaustics ? PathState.PathRoughness : 0.0);
 		// Trace the ray
 		TraceRay(
 			TLAS,
diff --git a/Engine/Shaders/Private/PathTracing/PathTracingMaterialHitShader.usf b/Engine/Shaders/Private/PathTracing/PathTracingMaterialHitShader.usf
--- a/Engine/Shaders/Private/PathTracing/PathTracingMaterialHitShader.usf
+++ b/Engine/Shaders/Private/PathTracing/PathTracingMaterialHitShader.usf
@@ -116,9 +116,14 @@ RAY_TRACING_ENTRY_CLOSEST_HIT(PathTracingMaterialCHS,
 #endif
 
 	Payload.BlendingMode = MaterialBlendingMode;
-	Payload.PrimitiveLightingChannelMask = GetPrimitive_LightingChannelMask(MaterialParameters.PrimitiveId);
 
-	if ((GetPrimitiveData(MaterialParameters.PrimitiveId).Flags & PRIMITIVE_SCENE_DATA_FLAG_DECAL_RECEIVER) != 0)
+	// fetch primitive flags only once
+	// TODO: would be nice to keep this inside MaterialParameters as it is also needed there as well
+	const uint PrimitiveFlags = GetPrimitiveData(MaterialParameters.PrimitiveId).Flags;
+
+	Payload.PrimitiveLightingChannelMask = GetPrimitive_LightingChannelMask_FromFlags(PrimitiveFlags);
+
+	if ((PrimitiveFlags & PRIMITIVE_SCENE_DATA_FLAG_DECAL_RECEIVER) != 0)
 	{
 		Payload.SetDecalReceiver();
 	}
@@ -130,8 +135,11 @@ RAY_TRACING_ENTRY_CLOSEST_HIT(PathTracingMaterialCHS,
 	}
 
 #if MATERIAL_IS_SKY
-	if (PackedPayload.IsIgnoreSkyMaterials())
+	if (!PackedPayload.IsCameraRay())
 	{
+		// avoid double counting what was captured by the skylight
+		// also avoid noise from hot spots (they can be properly
+		// importance sampled if a capturing skylight is added)
 		PackedPayload = PackPathTracingPayload(Payload);
 		return;
 	}
@@ -384,6 +392,27 @@ RAY_TRACING_ENTRY_CLOSEST_HIT(PathTracingMaterialCHS,
 		// terminal case for the conditionals above
 	}
 
+	if (PackedPayload.IsCameraRay())
+	{
+		if ((PrimitiveFlags & PRIMITIVE_SCENE_DATA_FLAG_HOLDOUT) != 0)
+		{
+			if (Payload.BlendingMode != RAY_TRACING_BLEND_MODE_TRANSLUCENT || Payload.Ior != 0.0)
+			{
+				// If we are either not translucent, or if the surface is marked as refractive
+				// then treat opacity as 1.0
+				Payload.Opacity = 1.0;
+			}
+			else
+			{
+				// keep the opacity from the material, it might indicate partial transparency
+				// so we should get a partial holdout
+			}
+			Payload.BlendingMode = MaterialBlendingMode = RAY_TRACING_BLEND_MODE_ALPHA_HOLDOUT;
+			Payload.ShadingModelID = SHADINGMODELID_UNLIT;
+			Payload.Radiance = 0;
+		}
+	}
+
 	// adjust after everything else (because SSS profile case needs to decode the dual spec info from the unmodified roughness)
 	Payload.Roughness = AdjustMaterialRoughness(Payload.Roughness, PackedPayload.GetPathRoughness());
 
diff --git a/Engine/Shaders/Private/SceneData.ush b/Engine/Shaders/Private/SceneData.ush
--- a/Engine/Shaders/Private/SceneData.ush
+++ b/Engine/Shaders/Private/SceneData.ush
@@ -88,6 +88,7 @@
 #define PRIMITIVE_SCENE_DATA_FLAG_INSTANCE_DRAW_DISTANCE_CULL			0x2000000
 #define PRIMITIVE_SCENE_DATA_FLAG_WPO_DISABLE_DISTANCE					0x4000000
 #define PRIMITIVE_SCENE_DATA_FLAG_WRITES_CUSTOM_DEPTH_STENCIL			0x8000000
+#define PRIMITIVE_SCENE_DATA_FLAG_HOLDOUT                               0x10000000
 
 uint GetDefaultPrimitiveSceneDataFlags()
 {
diff --git a/Engine/Source/Runtime/Engine/Classes/Components/PrimitiveComponent.h b/Engine/Source/Runtime/Engine/Classes/Components/PrimitiveComponent.h
--- a/Engine/Source/Runtime/Engine/Classes/Components/PrimitiveComponent.h
+++ b/Engine/Source/Runtime/Engine/Classes/Components/PrimitiveComponent.h
@@ -430,6 +430,10 @@ public:
 	UPROPERTY(EditAnywhere, AdvancedDisplay, BlueprintReadOnly, Category=Rendering)
 	uint8 bReceivesDecals:1;
 
+	/** If this is True, this primitive will render black with an alpha of 0, but all secondary effects (shadows, reflections, indirect lighting) remain. This feature is currently only implemented in the Path Tracer. */
+	UPROPERTY(EditAnywhere, BlueprintReadWrite, Category = PathTracing)
+	uint8 bHoldout : 1;
+
 	/** If this is True, this component won't be visible when the view actor is the component's owner, directly or indirectly. */
 	UPROPERTY(EditAnywhere, AdvancedDisplay, BlueprintReadOnly, Category = Rendering)
 	uint8 bOwnerNoSee:1;
diff --git a/Engine/Source/Runtime/Engine/Private/PrimitiveSceneProxy.cpp b/Engine/Source/Runtime/Engine/Private/PrimitiveSceneProxy.cpp
--- a/Engine/Source/Runtime/Engine/Private/PrimitiveSceneProxy.cpp
+++ b/Engine/Source/Runtime/Engine/Private/PrimitiveSceneProxy.cpp
@@ -275,6 +275,7 @@ FPrimitiveSceneProxy::FPrimitiveSceneProxy(const UPrimitiveComponent* InComponen
 ,	bHasPerInstanceEditorData(false)
 #endif
 ,	bAllowApproximateOcclusion(InComponent->Mobility != EComponentMobility::Movable)
+,   bHoldout(InComponent->bHoldout)
 ,	bUseAsOccluder(InComponent->bUseAsOccluder)
 ,	bSelectable(InComponent->bSelectable)
 ,	bHasPerInstanceHitProxies(InComponent->bHasPerInstanceHitProxies)
@@ -543,6 +544,7 @@ void FPrimitiveSceneProxy::UpdateUniformBuffer()
 				.CastContactShadow(CastsContactShadow())
 				.CastHiddenShadow(CastsHiddenShadow())
 				.CastShadow(CastsDynamicShadow())
+				.Holdout(Holdout())
 				.InstanceSceneDataOffset(PrimitiveSceneInfo ? PrimitiveSceneInfo->GetInstanceSceneDataOffset() : INDEX_NONE)
 				.NumInstanceSceneDataEntries(PrimitiveSceneInfo ? PrimitiveSceneInfo->GetNumInstanceSceneDataEntries() : 0)
 				.InstancePayloadDataOffset(PrimitiveSceneInfo ? PrimitiveSceneInfo->GetInstancePayloadDataOffset() : INDEX_NONE)
diff --git a/Engine/Source/Runtime/Engine/Private/PrimitiveUniformShaderParameters.cpp b/Engine/Source/Runtime/Engine/Private/PrimitiveUniformShaderParameters.cpp
--- a/Engine/Source/Runtime/Engine/Private/PrimitiveUniformShaderParameters.cpp
+++ b/Engine/Source/Runtime/Engine/Private/PrimitiveUniformShaderParameters.cpp
@@ -191,6 +191,7 @@ FPrimitiveSceneShaderData::FPrimitiveSceneShaderData(const FPrimitiveSceneProxy*
 		.VisibleInSceneCaptureOnly(Proxy->IsVisibleInSceneCaptureOnly())
 		.HiddenInSceneCapture(Proxy->IsHiddenInSceneCapture())
 		.ForceHidden(Proxy->IsForceHidden())
+		.Holdout(Proxy->Holdout())
 		.UseVolumetricLightmap(bHasPrecomputedVolumetricLightmap)
 		.NaniteResourceID(NaniteResourceID)
 		.NaniteHierarchyOffset(NaniteHierarchyOffset)
diff --git a/Engine/Source/Runtime/Engine/Public/PrimitiveSceneProxy.h b/Engine/Source/Runtime/Engine/Public/PrimitiveSceneProxy.h
--- a/Engine/Source/Runtime/Engine/Public/PrimitiveSceneProxy.h
+++ b/Engine/Source/Runtime/Engine/Public/PrimitiveSceneProxy.h
@@ -649,6 +649,7 @@ public:
 	inline bool DoesVFRequirePrimitiveUniformBuffer() const { return bVFRequiresPrimitiveUniformBuffer; }
 	inline bool ShouldUseAsOccluder() const { return bUseAsOccluder; }
 	inline bool AllowApproximateOcclusion() const { return bAllowApproximateOcclusion; }
+	inline bool Holdout() const { return bHoldout; }
 
 	inline FRHIUniformBuffer* GetUniformBuffer() const
 	{
@@ -1305,6 +1306,12 @@ protected:
 	/** If this is True, this primitive doesn't need exact occlusion info. */
 	uint8 bAllowApproximateOcclusion : 1;
 
+	/**
+	 * If this is True, this primitive should render black with an alpha of 0, but all secondary effects (shadows, refletions, indirect lighting)
+	 * should behave as usual. This feature is currently only implemented in the Path Tracer.
+	 */
+	uint8 bHoldout : 1;
+
 private:
 
 	/** If this is True, this primitive will be used to occlusion cull other primitives. */
diff --git a/Engine/Source/Runtime/Engine/Public/PrimitiveUniformShaderParameters.h b/Engine/Source/Runtime/Engine/Public/PrimitiveUniformShaderParameters.h
--- a/Engine/Source/Runtime/Engine/Public/PrimitiveUniformShaderParameters.h
+++ b/Engine/Source/Runtime/Engine/Public/PrimitiveUniformShaderParameters.h
@@ -95,6 +95,7 @@ END_GLOBAL_SHADER_PARAMETER_STRUCT()
 #define PRIMITIVE_SCENE_DATA_FLAG_INSTANCE_DRAW_DISTANCE_CULL			0x2000000
 #define PRIMITIVE_SCENE_DATA_FLAG_WPO_DISABLE_DISTANCE					0x4000000
 #define PRIMITIVE_SCENE_DATA_FLAG_WRITES_CUSTOM_DEPTH_STENCIL			0x8000000
+#define PRIMITIVE_SCENE_DATA_FLAG_HOLDOUT                               0x10000000
 
 struct FPrimitiveUniformShaderParametersBuilder
 {
@@ -131,6 +132,7 @@ public:
 		bHasInstanceDrawDistanceCull				= false;
 		bHasWPODisableDistance						= false;
 		bWritesCustomDepthStencil					= false;
+		bHoldout                                    = false;
 
 		Parameters.BoundsScale						= 1.0f;
 
@@ -187,6 +189,7 @@ public:
 	PRIMITIVE_UNIFORM_BUILDER_FLAG_METHOD(bool,			VisibleInSceneCaptureOnly);
 	PRIMITIVE_UNIFORM_BUILDER_FLAG_METHOD(bool,			HiddenInSceneCapture);
 	PRIMITIVE_UNIFORM_BUILDER_FLAG_METHOD(bool,			ForceHidden);
+	PRIMITIVE_UNIFORM_BUILDER_FLAG_METHOD(bool,         Holdout);
 
 	PRIMITIVE_UNIFORM_BUILDER_METHOD(uint32,			InstanceSceneDataOffset);
 	PRIMITIVE_UNIFORM_BUILDER_METHOD(uint32,			NumInstanceSceneDataEntries);
@@ -447,6 +450,7 @@ public:
 		Parameters.Flags |= bHasInstanceDrawDistanceCull ? PRIMITIVE_SCENE_DATA_FLAG_INSTANCE_DRAW_DISTANCE_CULL : 0u;
 		Parameters.Flags |= bHasWPODisableDistance ? PRIMITIVE_SCENE_DATA_FLAG_WPO_DISABLE_DISTANCE : 0u;
 		Parameters.Flags |= bWritesCustomDepthStencil ? PRIMITIVE_SCENE_DATA_FLAG_WRITES_CUSTOM_DEPTH_STENCIL : 0u;
+		Parameters.Flags |= bHoldout ? PRIMITIVE_SCENE_DATA_FLAG_HOLDOUT : 0u;
 		return Parameters;
 	}
 
@@ -487,6 +491,7 @@ private:
 	uint32 bHasInstanceDrawDistanceCull : 1;
 	uint32 bHasWPODisableDistance : 1;
 	uint32 bWritesCustomDepthStencil : 1;
+	uint32 bHoldout : 1;
 };
 
 inline TUniformBufferRef<FPrimitiveUniformShaderParameters> CreatePrimitiveUniformBufferImmediate(
