Fix arial perspective transmitance on mobile.

GetAerialPerspectiveLuminanceTransmittanceWithFogOver function takes the WorldPosition in camera space since the last engine merge, so we need to have the translated view translation.

diff --git a/Engine/Shaders/Private/MobileFog.usf b/Engine/Shaders/Private/MobileFog.usf
--- a/Engine/Shaders/Private/MobileFog.usf
+++ b/Engine/Shaders/Private/MobileFog.usf
@@ -58,7 +58,7 @@ void MobileFogPS(
 			ResolvedView.RealTimeReflectionCapture, 
 			ResolvedView.SkyAtmosphereCameraAerialPerspectiveVolumeSizeAndInvSize,
 			NDCPosition, 
-			(WorldPosition.xyz - ResolvedView.TranslatedWorldCameraOrigin) * CM_TO_SKY_UNIT,
+			(WorldPositionRelativeToCamera.xyz - ResolvedView.TranslatedWorldCameraOrigin) * CM_TO_SKY_UNIT,
 			View.CameraAerialPerspectiveVolume, 
 			View.CameraAerialPerspectiveVolumeSampler,
 			ResolvedView.SkyAtmosphereCameraAerialPerspectiveVolumeDepthResolutionInv,
