Path Tracer: Refactor handling of SSS Radius in the path tracer

Decode SSS profile entirely in the hit material

Reserve a slot in the payload for phase function anisotropy

Simplify logic by tracking SSS Radius/Color explicitly after payload decode

diff --git a/Engine/Shaders/Private/PathTracing/Material/PathTracingLambert.ush b/Engine/Shaders/Private/PathTracing/Material/PathTracingLambert.ush
--- a/Engine/Shaders/Private/PathTracing/Material/PathTracingLambert.ush
+++ b/Engine/Shaders/Private/PathTracing/Material/PathTracingLambert.ush
@@ -38,7 +38,7 @@ FMaterialSample Simplified_SampleMaterial(
 	float3 N_World = Payload.WorldNormal;
 
 	float3 FrontColor = Payload.BaseColor;
-	float3 BackColor = Payload.ShadingModelID == SHADINGMODELID_TWOSIDED_FOLIAGE ? Payload.GetSubsurfaceColor() : 0.0;
+	float3 BackColor = Payload.SubsurfaceColor;
 
 	float4 SampledValue = CosineSampleHemisphere(RandSample.xy);
 
@@ -59,7 +59,7 @@ FMaterialEval Simplified_EvalMaterial(
 	float NoL = dot(N_World, OutgoingDirection);
 
 	float3 FrontColor = Payload.BaseColor;
-	float3 BackColor = Payload.ShadingModelID == SHADINGMODELID_TWOSIDED_FOLIAGE ? Payload.GetSubsurfaceColor() : 0.0;
+	float3 BackColor = Payload.SubsurfaceColor;
 	float TransmissionProb = LobeSelectionProb(BackColor, FrontColor);
 
 	float3 Weight = NoL < 0.0 ? BackColor : FrontColor;
diff --git a/Engine/Shaders/Private/PathTracing/Material/PathTracingMaterialSampling.ush b/Engine/Shaders/Private/PathTracing/Material/PathTracingMaterialSampling.ush
--- a/Engine/Shaders/Private/PathTracing/Material/PathTracingMaterialSampling.ush
+++ b/Engine/Shaders/Private/PathTracing/Material/PathTracingMaterialSampling.ush
@@ -22,7 +22,6 @@ PathTracingMaterialSampling.usf: Material sampling functions
 #include "PathTracingEye.ush"
 #include "PathTracingHair.ush"
 #include "PathTracingMedium.ush"
-#include "../../BurleyNormalizedSSSCommon.ush"
 
 // Simplified material shaders for GPU Lightmass, supporting Lambert, Two Sided Foliage and Thin Translucent
 // Make sure material evaluation code below matches with hit shader which fills the payload
@@ -142,7 +141,7 @@ float3 EstimateMaterialAlbedo(FPathTracingPayload Payload)
 			return 0;
 #if SIMPLIFIED_MATERIAL_SHADER
 		case SHADINGMODELID_TWOSIDED_FOLIAGE:
-			Albedo = BrdfWeight * (Payload.BaseColor + Payload.GetSubsurfaceColor());
+			Albedo = BrdfWeight * (Payload.BaseColor + Payload.SubsurfaceColor);
 			break;
 #else
 		case SHADINGMODELID_TWOSIDED_FOLIAGE:
@@ -150,7 +149,7 @@ float3 EstimateMaterialAlbedo(FPathTracingPayload Payload)
 		case SHADINGMODELID_PREINTEGRATED_SKIN:
 		case SHADINGMODELID_SUBSURFACE_PROFILE:
 		case SHADINGMODELID_EYE:
-			Albedo = BrdfWeight * (Payload.DiffuseColor + Payload.SpecularColor + Payload.GetSubsurfaceColor());
+			Albedo = BrdfWeight * (Payload.DiffuseColor + Payload.SpecularColor + Payload.SubsurfaceColor);
 			break;
 		case SHADINGMODELID_SINGLELAYERWATER:
 		case SHADINGMODELID_DEFAULT_LIT:
@@ -240,7 +239,7 @@ void AccumulateAlbedoNormal(inout FPathTracingPayload Payload, float3 PathThroug
 		return;
 #if SIMPLIFIED_MATERIAL_SHADER
 	case SHADINGMODELID_TWOSIDED_FOLIAGE:
-		Albedo += BrdfWeight * PathThroughput * (Payload.BaseColor + Payload.GetSubsurfaceColor());
+		Albedo += BrdfWeight * PathThroughput * (Payload.BaseColor + Payload.SubsurfaceColor);
 		Normal += Luminance(BrdfWeight * PathThroughput) * Payload.WorldNormal;
 		break;
 #else
@@ -249,7 +248,7 @@ void AccumulateAlbedoNormal(inout FPathTracingPayload Payload, float3 PathThroug
 	case SHADINGMODELID_PREINTEGRATED_SKIN:
 	case SHADINGMODELID_SUBSURFACE_PROFILE:
 	case SHADINGMODELID_EYE:
-		Albedo += BrdfWeight * PathThroughput * (Payload.DiffuseColor + RoughnessWeight * Payload.SpecularColor + Payload.GetSubsurfaceColor());
+		Albedo += BrdfWeight * PathThroughput * (Payload.DiffuseColor + RoughnessWeight * Payload.SpecularColor + Payload.SubsurfaceColor);
 		Normal += BrdfWeight * Luminance(PathThroughput) * Payload.WorldNormal;
 		break;
 	case SHADINGMODELID_SINGLELAYERWATER:
@@ -332,23 +331,6 @@ float GetAverageRoughness(FPathTracingPayload Payload)
 #endif
 }
 
-
-float3 DecodeSSSProfileRadius(uint ProfileId, float3 DiffuseColor, float Opacity)
-{
-	// Burley parameterization
-	float3 SurfaceAlbedo = View.SSProfilesTexture.Load(int3(BSSS_SURFACEALBEDO_OFFSET, ProfileId, 0)).xyz;
-	float3 DiffuseMeanFreePath = DecodeDiffuseMeanFreePath(View.SSProfilesTexture.Load(int3(BSSS_DMFP_OFFSET, ProfileId, 0))).xyz;
-	float WorldUnitScale = DecodeWorldUnitScale(View.SSProfilesTexture.Load(int3(SSSS_TINT_SCALE_OFFSET, ProfileId, 0)).a) * BURLEY_CM_2_MM;
-
-	// Opacity acts as a per-pixel radius multiplier
-	// NOTE: this seems backwards? Opacity=0 acts like default-lit while Opacity=1 acts like SSS?
-	// NOTE2: Confirm if this interpretation of opacity is correct ...
-	float3 SSSRadius = GetMFPFromDMFPApprox(SurfaceAlbedo, DiffuseColor, Opacity * WorldUnitScale * DiffuseMeanFreePath);
-
-	return SSSRadius * BURLEY_MM_2_CM;
-}
-
-
 FPathTracingPayload UnpackPathTracingPayloadMaterial(FPackedPathTracingPayload PackedInput, FRayDesc Ray)
 {
 	// This function can also do a bit of per-material post-processing to simplify the rest of the code
@@ -360,20 +342,15 @@ FPathTracingPayload UnpackPathTracingPayloadMaterial(FPackedPathTracingPayload P
 		case SHADINGMODELID_SUBSURFACE:
 		case SHADINGMODELID_PREINTEGRATED_SKIN:
 		{
-			// TODO: is this accurate enough? hard to measure exactly since the raster path uses a very different approach
-			float SSSRadius = (1 - HitPayload.Opacity) * 10.0; // simple formula, up to 10cm radius
-			HitPayload.SubsurfaceRadius = SSSRadius;
-
-			// Reset opacity to avoid special cases elsewhere
-			HitPayload.Opacity = 1;
-			
+			HitPayload.SubsurfaceRadius = HitPayload.GetSubsurfaceRadius();
 			// fallthrough to case below to avoid over-bright SSS
 		}
 #endif
 		case SHADINGMODELID_TWOSIDED_FOLIAGE:
 		{
+			HitPayload.SubsurfaceColor = HitPayload.GetSubsurfaceColor();
 			float3 Front = HitPayload.DiffuseColor;
-			float3 Back = HitPayload.GetSubsurfaceColor();
+			float3 Back = HitPayload.SubsurfaceColor;
 			float3 Sum = Front + Back;
 			float S = max(Sum.x, max(Sum.y, Sum.z));
 			if (S > 1)
@@ -381,56 +358,32 @@ FPathTracingPayload UnpackPathTracingPayloadMaterial(FPackedPathTracingPayload P
 				// nudge the material back to something physically plausible
 				// NOTE: we ignore spec here since it should be accounted for by the brdf model itself
 				HitPayload.DiffuseColor = Front / S;
-				HitPayload.SetSubsurfaceColor(Back / S);
+				HitPayload.SubsurfaceColor = Back / S;
 			}
 			break;
 		}
 #if SIMPLIFIED_MATERIAL_SHADER
 		default:
 		{
-			HitPayload.SetSubsurfaceColor(0.0);
+			HitPayload.SubsurfaceColor = 0.0;
 			break;
 		}
 #else
 		case SHADINGMODELID_SUBSURFACE_PROFILE:
 		{
-			uint ProfileId = HitPayload.GetSubsurfaceProfileId();
-			if (GetSubsurfaceProfileUseBurley(ProfileId))
-			{
-				HitPayload.SubsurfaceRadius = DecodeSSSProfileRadius(ProfileId, HitPayload.DiffuseColor, HitPayload.Opacity);
-
-				// burley mode doesn't have any diffuse component
-				HitPayload.SetSubsurfaceColor(HitPayload.DiffuseColor);
-				HitPayload.DiffuseColor = 0;
-			}
-			else
-			{
-				// NOTE: Non-burley scattering is deprecated going forward, so not worth supporting this mode
-				HitPayload.SubsurfaceRadius = 0.0;
-				HitPayload.SetSubsurfaceColor(0);
-			}
-			// Reset opacity to avoid special cases elsewhere
-			HitPayload.Opacity = 1;
+			HitPayload.SubsurfaceRadius = HitPayload.GetSubsurfaceRadius();
+			// NOTE: if radius is 0 the sss color will get flipped back into diffuse
+			HitPayload.SubsurfaceColor = HitPayload.DiffuseColor;
+			HitPayload.DiffuseColor = 0;
 			break;
 		}
 		case SHADINGMODELID_EYE:
 		{
-			uint ProfileId = HitPayload.GetSubsurfaceProfileId();
-			if (GetSubsurfaceProfileUseBurley(ProfileId))
-			{
-				HitPayload.SubsurfaceRadius = DecodeSSSProfileRadius(ProfileId, HitPayload.DiffuseColor, 1.0);
-
-				// eye still needs a diffuse component for the iris center
-				float IrisMask = HitPayload.GetEyeIrisMask();
-				HitPayload.SetSubsurfaceColor((1 - IrisMask) * HitPayload.DiffuseColor);
-				HitPayload.DiffuseColor *= IrisMask;
-			}
-			else
-			{
-				// NOTE: Non-burley scattering is deprecated going forward, so not worth supporting this mode
-				HitPayload.SubsurfaceRadius = 0.0;
-				HitPayload.SetSubsurfaceColor(0);
-			}
+			HitPayload.SubsurfaceRadius = HitPayload.GetSubsurfaceRadius();
+			// eye still needs a diffuse component for the iris center
+			float IrisMask = HitPayload.GetEyeIrisMask();
+			HitPayload.SubsurfaceColor = HitPayload.DiffuseColor - IrisMask * HitPayload.DiffuseColor;
+			HitPayload.DiffuseColor *= IrisMask;
 			break;
 		}
 #endif
diff --git a/Engine/Shaders/Private/PathTracing/Material/PathTracingTwoSidedFoliage.ush b/Engine/Shaders/Private/PathTracing/Material/PathTracingTwoSidedFoliage.ush
--- a/Engine/Shaders/Private/PathTracing/Material/PathTracingTwoSidedFoliage.ush
+++ b/Engine/Shaders/Private/PathTracing/Material/PathTracingTwoSidedFoliage.ush
@@ -20,7 +20,7 @@ FMaterialSample TwoSidedFoliage_SampleMaterial(
 
 	// pick back lobe proportionally to its contribution relative to front side
 	const float3 Diff = (1 - Spec.E) * Payload.DiffuseColor;
-	const float3 Back = (1 - Spec.E) * Payload.GetSubsurfaceColor();
+	const float3 Back = (1 - Spec.E) * Payload.SubsurfaceColor;
 	float TransmissionSelectionPdf = LobeSelectionProb(Back, Diff + Spec.E);
 
 	if (RandSample.x < TransmissionSelectionPdf)
@@ -56,7 +56,7 @@ FMaterialEval TwoSidedFoliage_EvalMaterial(
 
 	// pick back lobe proportionally to its contribution relative to front side
 	const float3 Diff = (1 - Spec.E) * Payload.DiffuseColor;
-	const float3 Back = (1 - Spec.E) * Payload.GetSubsurfaceColor();
+	const float3 Back = (1 - Spec.E) * Payload.SubsurfaceColor;
 	float TransmissionSelectionPdf = LobeSelectionProb(Back, Diff + Spec.E);
 
 	float3 N = Payload.WorldNormal;
diff --git a/Engine/Shaders/Private/PathTracing/PathTracingCommon.ush b/Engine/Shaders/Private/PathTracing/PathTracingCommon.ush
--- a/Engine/Shaders/Private/PathTracing/PathTracingCommon.ush
+++ b/Engine/Shaders/Private/PathTracing/PathTracingCommon.ush
@@ -153,6 +153,7 @@ struct FPathTracingPayload : FMinimalPayload
 	float4 CustomData0;
 	float4 CustomData1;
 	float3 WorldTangent;
+	float3 SubsurfaceColor;
 	float3 SubsurfaceRadius;
 	float3 TranslatedWorldPos;
 
@@ -180,12 +181,15 @@ struct FPathTracingPayload : FMinimalPayload
 	float3 GetClearCoatBottomNormal() { return CustomData1.xyz; }
 	void SetClearCoatBottomNormal(float3 BottomNormal) { CustomData1.xyz = BottomNormal; }
 
-	// SHADINGMODELID_SUBSURFACE_PROFILE and SHADINGMODELID_EYE
-	uint GetSubsurfaceProfileId() { return uint(CustomData0.a * 255.0 + 0.5); }
+	// SHADINGMODELID_SUBSURFACE, SHADINGMODELID_SUBSURFACE_PROFILE and SHADINGMODELID_EYE
+	float3 GetSubsurfaceRadius() { return CustomData0.rgb; }
+	void SetSubsurfaceRadius(float3 Radius) { CustomData0.rgb = Radius; }
+	float GetSubsurfacePhaseFunction() { return CustomData0.a; }
+	void SetSubsurfacePhaseFunction(float G) { CustomData0.a = G; }
 
-	// SHADINGMODELID_SUBSURFACE, SHADINGMODELID_TWO_SIDED_FOLIAGE (and SHADINGMODELID_SUBSURFACE_PROFILE, SHADINGMODELID_EYE after decode)
-	float3 GetSubsurfaceColor() { return CustomData0.xyz; }
-	void SetSubsurfaceColor(float3 SubsurfaceColor) { CustomData0.xyz = SubsurfaceColor; }
+	// SHADINGMODELID_SUBSURFACE, SHADINGMODELID_TWO_SIDED_FOLIAGE (note that SHADINGMODELID_SUBSURFACE_PROFILE and SHADINGMODELID_EYE don't use this as they use DiffuseColor directly)
+	float3 GetSubsurfaceColor() { return CustomData1.xyz; }
+	void SetSubsurfaceColor(float3 SubsurfaceColor) { CustomData1.xyz = SubsurfaceColor; }
 	
 	// SHADINGMODELID_SUBSURFACE_PROFILE
 	float3 GetDualRoughnessSpecular() { return CustomData1.xyz; }
diff --git a/Engine/Shaders/Private/PathTracing/PathTracingCore.ush b/Engine/Shaders/Private/PathTracing/PathTracingCore.ush
--- a/Engine/Shaders/Private/PathTracing/PathTracingCore.ush
+++ b/Engine/Shaders/Private/PathTracing/PathTracingCore.ush
@@ -639,7 +639,7 @@ bool ProcessSubsurfaceRandomWalk(inout FPathTracingPayload Payload, inout float3
 		return true;
 	}
 
-	float3 SSSColor = Payload.GetSubsurfaceColor();
+	float3 SSSColor = Payload.SubsurfaceColor;
 	float3 SSSRadius = Payload.SubsurfaceRadius;
 
 	if (SimplifySSS || all(SSSRadius == 0) || all(SSSColor == 0) || MaxSSSBounces == 0)
@@ -649,7 +649,7 @@ bool ProcessSubsurfaceRandomWalk(inout FPathTracingPayload Payload, inout float3
 		// or if the user decided not to perform any SSS bounces
 		Payload.DiffuseColor += SSSColor;
 		Payload.DiffuseColor = saturate(Payload.DiffuseColor);
-		Payload.SetSubsurfaceColor(0.0);
+		Payload.SubsurfaceColor = 0.0;
 		return true;
 	}
 	
diff --git a/Engine/Shaders/Private/PathTracing/PathTracingMaterialHitShader.usf b/Engine/Shaders/Private/PathTracing/PathTracingMaterialHitShader.usf
--- a/Engine/Shaders/Private/PathTracing/PathTracingMaterialHitShader.usf
+++ b/Engine/Shaders/Private/PathTracing/PathTracingMaterialHitShader.usf
@@ -28,9 +28,11 @@
 #include "/Engine/Private/ShadingCommon.ush"
 #include "/Engine/Private/DeferredShadingCommon.ush"
 #include "/Engine/Private/SubsurfaceProfileCommon.ush"
+#include "/Engine/Private/BurleyNormalizedSSSCommon.ush"
 #include "/Engine/Private/PathTracing/Material/PathTracingFresnel.ush"
 #include "/Engine/Private/PathTracing/Material/PathTracingThinGlass.ush"
 
+
 float AdjustMaterialRoughness(float Roughness, float PathRoughness)
 {
 	// Modify the payload roughness to minimize difficult caustics
@@ -55,6 +57,22 @@ float GetRefractionIor(FPixelMaterialInputs PixelMaterialInputs)
 	return Ior;
 }
 
+float3 DecodeSSSProfileRadius(uint ProfileId, float3 DiffuseColor, float Opacity)
+{
+	// Burley parameterization
+	float3 SurfaceAlbedo = View.SSProfilesTexture.Load(int3(BSSS_SURFACEALBEDO_OFFSET, ProfileId, 0)).xyz;
+	float3 DiffuseMeanFreePath = DecodeDiffuseMeanFreePath(View.SSProfilesTexture.Load(int3(BSSS_DMFP_OFFSET, ProfileId, 0))).xyz;
+	float WorldUnitScale = DecodeWorldUnitScale(View.SSProfilesTexture.Load(int3(SSSS_TINT_SCALE_OFFSET, ProfileId, 0)).a) * BURLEY_CM_2_MM;
+
+	// Opacity acts as a per-pixel radius multiplier
+	// NOTE: this seems backwards? Opacity=0 acts like default-lit while Opacity=1 acts like SSS?
+	// NOTE2: Confirm if this interpretation of opacity is correct ...
+	float3 SSSRadius = GetMFPFromDMFPApprox(SurfaceAlbedo, DiffuseColor, Opacity * WorldUnitScale * DiffuseMeanFreePath);
+
+	return SSSRadius * BURLEY_MM_2_CM;
+}
+
+
 RAY_TRACING_ENTRY_CLOSEST_HIT(PathTracingMaterialCHS,
 	FPackedPathTracingPayload, PackedPayload,
 	FRayTracingIntersectionAttributes, Attributes)
@@ -257,21 +275,36 @@ RAY_TRACING_ENTRY_CLOSEST_HIT(PathTracingMaterialCHS,
 	else
 #endif
 
-#if MATERIAL_SHADINGMODEL_TWOSIDED_FOLIAGE || MATERIAL_SHADINGMODEL_SUBSURFACE || MATERIAL_SHADINGMODEL_PREINTEGRATED_SKIN || MATERIAL_SHADINGMODEL_SUBSURFACE_PROFILE
-	if (Payload.ShadingModelID == SHADINGMODELID_TWOSIDED_FOLIAGE ||
-		Payload.ShadingModelID == SHADINGMODELID_SUBSURFACE ||
+#if MATERIAL_SHADINGMODEL_TWOSIDED_FOLIAGE
+	if (Payload.ShadingModelID == SHADINGMODELID_TWOSIDED_FOLIAGE)
+	{
+		const float3 SubsurfaceColor = GetMaterialSubsurfaceData(PixelMaterialInputs).rgb;
+		Payload.SetSubsurfaceColor(SubsurfaceColor);
+	}
+	else
+#endif
+
+#if MATERIAL_SHADINGMODEL_SUBSURFACE || MATERIAL_SHADINGMODEL_PREINTEGRATED_SKIN
+	if (Payload.ShadingModelID == SHADINGMODELID_SUBSURFACE ||
 		Payload.ShadingModelID == SHADINGMODELID_PREINTEGRATED_SKIN)
 	{
 		const float3 SubsurfaceColor = GetMaterialSubsurfaceData(PixelMaterialInputs).rgb;
 		Payload.SetSubsurfaceColor(SubsurfaceColor);
+
+		// TODO: is this accurate enough? hard to measure exactly since the raster path uses a very different approach
+		// TODO: It would be nice to have a dedicated node to control the SSS Radius explicitly instead of this
+		const float SSSRadius = (1 - Payload.Opacity) * 10.0; // simple formula, up to 10cm radius
+		Payload.SetSubsurfaceRadius(SSSRadius);
+		Payload.SetSubsurfacePhaseFunction(0.0); // TODO: expose via a dedicated node?
+		// Reset opacity to avoid special cases elsewhere
+		Payload.Opacity = 1;
 	}
 	else
 #endif
 #if MATERIAL_SHADINGMODEL_SUBSURFACE_PROFILE
 	if (Payload.ShadingModelID == SHADINGMODELID_SUBSURFACE_PROFILE)
 	{
-		Payload.CustomData0 = GetMaterialSubsurfaceData(PixelMaterialInputs);
-		uint ProfileId = Payload.GetSubsurfaceProfileId();
+		const uint ProfileId = ExtractSubsurfaceProfileInt(GetMaterialSubsurfaceDataRaw(PixelMaterialInputs).a);
 
 		// Decode dual roughness info
 		float Lobe0Roughness = 0;
@@ -283,6 +316,19 @@ RAY_TRACING_ENTRY_CLOSEST_HIT(PathTracingMaterialCHS,
 		Lobe0Roughness = AdjustMaterialRoughness(Lobe0Roughness, PackedPayload.GetPathRoughness());
 		Lobe1Roughness = AdjustMaterialRoughness(Lobe1Roughness, PackedPayload.GetPathRoughness());
 		Payload.SetDualRoughnessSpecular(Lobe0Roughness, Lobe1Roughness, LobeMix);
+
+		if (GetSubsurfaceProfileUseBurley(ProfileId))
+		{
+			// Decode SSS radius
+			float3 DiffuseColor = Payload.BaseColor * (1.0 - Payload.Metallic);
+			Payload.SetSubsurfaceRadius(DecodeSSSProfileRadius(ProfileId, DiffuseColor, Payload.Opacity));
+		}
+		else
+		{
+			Payload.SetSubsurfaceRadius(0.0);
+		}
+		Payload.SetSubsurfacePhaseFunction(0.0); // TODO: extract from profile data
+		Payload.Opacity = 1.0; // Reset opacity to avoid special cases elsewhere
 	}
 	else
 #endif
@@ -339,7 +385,6 @@ RAY_TRACING_ENTRY_CLOSEST_HIT(PathTracingMaterialCHS,
 		const float IrisMask     = saturate(GetMaterialCustomData0(MaterialParameters));
 		const float IrisDistance = saturate(GetMaterialCustomData1(MaterialParameters));
 		Payload.SetEyeIrisMask(IrisMask);
-		Payload.CustomData0.a = GetMaterialSubsurfaceData(PixelMaterialInputs).a; // SubsurfaceProfile
 		Payload.Metallic = 0.0;
 #if NUM_MATERIAL_OUTPUTS_GETTANGENTOUTPUT > 0
 		// Blend in the negative intersection normal to create some concavity
@@ -371,6 +416,21 @@ RAY_TRACING_ENTRY_CLOSEST_HIT(PathTracingMaterialCHS,
 		// IRIS_NORMAL off
 		Payload.SetEyeIrisNormal(PlaneNormal);
 #endif
+		
+		const uint ProfileId = ExtractSubsurfaceProfileInt(GetMaterialSubsurfaceDataRaw(PixelMaterialInputs).a);
+		if (GetSubsurfaceProfileUseBurley(ProfileId))
+		{
+			// Decode SSS radius
+			float3 DiffuseColor = Payload.BaseColor * (1.0 - Payload.Metallic);
+			Payload.SetSubsurfaceRadius(DecodeSSSProfileRadius(ProfileId, DiffuseColor, Payload.Opacity));
+		}
+		else
+		{
+			Payload.SetSubsurfaceRadius(0.0);
+		}
+		Payload.Opacity = 1.0; // revert to default now that decode is done
+		Payload.SetSubsurfacePhaseFunction(0.0); // TODO: extract from profile data
+
 	}
 	else
 #endif
