Fix SceneTextures(CustomStencil) node not taking into account UV offset

diff --git a/Engine/Shaders/Private/MaterialTemplate.ush b/Engine/Shaders/Private/MaterialTemplate.ush
--- a/Engine/Shaders/Private/MaterialTemplate.ush
+++ b/Engine/Shaders/Private/MaterialTemplate.ush
@@ -2262,7 +2262,8 @@ MaterialFloat4 MobileSceneTextureLookup(inout FMaterialPixelParameters Parameter
 		return 0.0f;
 	}
 #else
-	FGBufferData GBuffer = MobileFetchAndDecodeGBuffer(UV, Parameters.SvPosition.xy);
+	const float2 PixelPos = UV * View.BufferSizeAndInvSize.xy;
+	FGBufferData GBuffer = MobileFetchAndDecodeGBuffer(UV, PixelPos);
 
 	switch(SceneTextureId)
 	{
