Enable one pass projection by default
Hard code "SubsurfaceShadowMode = 1" path. Restrict to directional lights as that's the only place it is implemented currently anyways.
Misc cleanup

diff --git a/Engine/Shaders/Private/ClusteredDeferredShadingPixelShader.usf b/Engine/Shaders/Private/ClusteredDeferredShadingPixelShader.usf
--- a/Engine/Shaders/Private/ClusteredDeferredShadingPixelShader.usf
+++ b/Engine/Shaders/Private/ClusteredDeferredShadingPixelShader.usf
@@ -233,8 +233,7 @@ float4 GetLightGridLocalLighting(
 
 	int2 PixelPos = int2( ScreenUV.xy * View.BufferSizeAndInvSize.xy );
 	uint4 ShadowMask = ~ShadowMaskBits[ PixelPos ];
-	half SubsurfaceOpacity = GetSubsurfaceOpacityFromGbuffer( PixelPos );
-
+	
 	LOOP
 	for (uint LocalLightListIndex = FirstNonSimpleLightIndex; LocalLightListIndex < NumLocalLights; LocalLightListIndex++)
 	{
@@ -273,13 +272,6 @@ float4 GetLightGridLocalLighting(
 			LightAttenuation = VSMResult.ShadowFactor.xxxx;
 		}
 
-		if ( SubsurfaceOpacity < 1.0f && LocalLight.VirtualShadowMapId != INDEX_NONE )
-		{
-			// Replace SSS portion for subsurface materials
-			half SSSTransmission = ComputeSimpleSubsurfaceTransmissionFromVirtualShadowMap(LocalLight.VirtualShadowMapId, TranslatedWorldPosition, SubsurfaceOpacity);
-			LightAttenuation.yw = SSSTransmission;
-		}
-
 		#if USE_HAIR_LIGHTING
 		LightAttenuation = 1.0f; // TODO: use shadow mask bits (e.g., subpixel shadow mask)
 		if (LightAttenuation.x > 0)
diff --git a/Engine/Shaders/Private/VirtualShadowMaps/VirtualShadowMapCacheManagement.usf b/Engine/Shaders/Private/VirtualShadowMaps/VirtualShadowMapCacheManagement.usf
--- a/Engine/Shaders/Private/VirtualShadowMaps/VirtualShadowMapCacheManagement.usf
+++ b/Engine/Shaders/Private/VirtualShadowMaps/VirtualShadowMapCacheManagement.usf
@@ -107,8 +107,6 @@ void VirtualSmInvalidateInstancePagesCS(
 		// TODO: test the flag on the instance instead once it is updated correctly InstanceSceneData.CastShadows
 		if (bCastShadows)
 		{
-			bool bAnyStaticPagesInvalidated = false;
-
 			for (uint VirtualShadowMapId = FirstVirtualShadowMapId; VirtualShadowMapId < EndVirtualShadowMapId; ++VirtualShadowMapId)
 			{
 				// 1. Load cached projection data
@@ -136,7 +134,6 @@ void VirtualSmInvalidateInstancePagesCS(
 
 				bool bInvalidateStaticPage = bForceInvalidateStaticPage || ShouldCacheInstanceAsStatic(InstanceSceneData, ProjectionData.bUnCached);
 				uint InvalidationFlags = (bInvalidateStaticPage ? VSM_STATIC_UNCACHED_FLAG : VSM_DYNAMIC_UNCACHED_FLAG) << VSM_PHYSICAL_PAGE_INVALIDATION_FLAGS_SHIFT;
-				bAnyStaticPagesInvalidated = bAnyStaticPagesInvalidated || bInvalidateStaticPage;
 
 				// Go back to clip space
 				float4x4 UVToClip;
diff --git a/Engine/Shaders/Private/VirtualShadowMaps/VirtualShadowMapProjection.usf b/Engine/Shaders/Private/VirtualShadowMaps/VirtualShadowMapProjection.usf
--- a/Engine/Shaders/Private/VirtualShadowMaps/VirtualShadowMapProjection.usf
+++ b/Engine/Shaders/Private/VirtualShadowMaps/VirtualShadowMapProjection.usf
@@ -165,7 +165,6 @@ int4 ProjectionRect;
 float ScreenRayLength;
 float NormalBias;
 float SubsurfaceMinSourceRadius;
-int SubsurfaceShadowMode;			// Temp for testing and A/B comparison
 int SMRTRayCount;					// 0 = off
 int SMRTSamplesPerRay;
 float SMRTRayLengthScale;			// Directional lights
@@ -628,12 +627,15 @@ void VirtualShadowMapProjection(
 		FLightShaderParameters Light = GetRootLightShaderParameters();
 
 		// Subsurface approximation
+#if DIRECTIONAL_LIGHT
+		const float SubsurfaceOpacity = bIsHairInput ? 1.0f : GetSubsurfaceOpacityFromGbuffer(PixelPos);
 		// Increases light source radius for subsurface material as the opacity reduces to emulates light diffusion.
-		const float SubsurfaceOpacity = bIsHairInput ? 1.0f : GetSubsurfaceOpacityFromGbuffer( PixelPos );
+		// TODO: Does something similar make sense for local lights?
 		if (SubsurfaceOpacity < 1.0f)
 		{
-			Light.SourceRadius = max(Light.SourceRadius, (1-SubsurfaceOpacity)*SubsurfaceMinSourceRadius);
+			Light.SourceRadius = max(Light.SourceRadius, (1.0f - SubsurfaceOpacity) * SubsurfaceMinSourceRadius);
 		}
+#endif
 
 		FVirtualShadowMapSampleResult Result = ProjectLight(
 			VirtualShadowMapId,
@@ -645,24 +647,20 @@ void VirtualShadowMapProjection(
 			TranslatedWorldPosition,
 			Noise);
 
-		float TransmissionShadow = Result.ShadowFactor;
-		if ( SubsurfaceOpacity < 1.0f )
+#if DIRECTIONAL_LIGHT
+		// TODO: Implement occluder distance and similar falloff for local lights, and unify logic into one pass projection path
+		if (SubsurfaceOpacity < 1.0f)
 		{
-			if (SubsurfaceShadowMode > 0)
-			{
-				TransmissionShadow = GetShadowFactorSubsurface(Result.ShadowFactor, Result.OccluderDistance, SubsurfaceOpacity);
-				Result.ShadowFactor = TransmissionShadow;
-			}
-			else
-			{
-				TransmissionShadow = ComputeSimpleSubsurfaceTransmissionFromVirtualShadowMap(
-					VirtualShadowMapId,
-					TranslatedWorldPosition,
-					SubsurfaceOpacity);
-			}
+			// We use a single term for both the shadow and "transmission shadow" in the VSM path to provide a
+			// higher quality filtered result on back faces and avoid artifacts at the point where normals flip.
+			// This is especially important for foliage as normals are often "warped" to provide softer shading
+			// and do not represent the underlying geometry accurately, so having a continuous shadow function
+			// is important.
+			Result.ShadowFactor = GetShadowFactorSubsurface(Result.ShadowFactor, Result.OccluderDistance, SubsurfaceOpacity);
 		}
+#endif
 
-		OutShadowFactor[ PixelPos ] = float2( Result.ShadowFactor, TransmissionShadow );
+		OutShadowFactor[ PixelPos ] = Result.ShadowFactor.xx;
 
 		if (VisualizeVirtualShadowMapId == VirtualShadowMapId)
 		{
diff --git a/Engine/Shaders/Private/VirtualShadowMaps/VirtualShadowMapTransmissionCommon.ush b/Engine/Shaders/Private/VirtualShadowMaps/VirtualShadowMapTransmissionCommon.ush
--- a/Engine/Shaders/Private/VirtualShadowMaps/VirtualShadowMapTransmissionCommon.ush
+++ b/Engine/Shaders/Private/VirtualShadowMaps/VirtualShadowMapTransmissionCommon.ush
@@ -49,36 +49,4 @@ float GetShadowFactorSubsurface(float ShadowFactor, float OccluderDistance, floa
 		return Square(lerp(SSSTransmission, 1, ShadowFactor));
 	}
 	return ShadowFactor;
-
-}
-
-/**
-* Sample VSM to determine approximate thickness and compute simple subsurface transmission matching the
-* PCF path.
-*/
-float ComputeSimpleSubsurfaceTransmissionFromVirtualShadowMap(
-	int VirtualShadowMapId, 
-	float3 TranslatedWorldPosition,
-	float SubsurfaceOpacity)
-{
-	// If the receiver is a subsurface material, we do a very simple transmission approximation
-	// This matches the logic in the PCF kernel (ShadowFilteringCommon.ush) and the contact
-	// shadows trace (DeferredLightingCommon.ush). Ideally these should all be unified in a single location.
-	// Could optimize this in the case where we are doing single-sample VSM lookup anyways, but we would need
-	// exclude contact shadows overwriting the occluder distance. Single sample direct shadows are not currently
-	// a common case.
-
-	// Do a separate unfiltered lookup to get approximate occluder distance
-	FVirtualShadowMapSampleResult SSSResult = SampleVirtualShadowMapTranslatedWorld(
-		VirtualShadowMapId,
-		TranslatedWorldPosition,
-		0.0f // RayStartDistance
-	);
-
-	if (SSSResult.bValid)
-	{
-		return GetShadowFactorSubsurface(SSSResult.ShadowFactor, SSSResult.OccluderDistance, SubsurfaceOpacity);
-	}
-
-	return 1.0f;
 }
diff --git a/Engine/Source/Runtime/Renderer/Private/LightRendering.cpp b/Engine/Source/Runtime/Renderer/Private/LightRendering.cpp
--- a/Engine/Source/Runtime/Renderer/Private/LightRendering.cpp
+++ b/Engine/Source/Runtime/Renderer/Private/LightRendering.cpp
@@ -49,8 +49,6 @@ IMPLEMENT_GLOBAL_SHADER_PARAMETER_STRUCT(FDeferredLightUniformStruct, "DeferredL
 
 extern int32 GUseTranslucentLightingVolumes;
 
-extern TAutoConsoleVariable<int32> CVarVirtualShadowOnePassProjection;
-
 static int32 GAllowDepthBoundsTest = 1;
 static FAutoConsoleVariableRef CVarAllowDepthBoundsTest(
 	TEXT("r.AllowDepthBoundsTest"),
@@ -1441,7 +1439,6 @@ void FDeferredShadingSceneRenderer::RenderLights(
 					// In certain cases we can skip creating the screen shadow mask texture
 					// In particular right now this is true if we are doing one pass projection with only a virtual shadow map
 					// with no light functions, as in that case we can directly sample the shadow mask bits in the lighting shader.
-					// TODO: Add cvar
 					bElideScreenShadowMask = 
 						CVarOnePassProjectionSkipScreenShadowMask.GetValueOnRenderThread() != 0 &&
 						ShadowSceneRenderer->UsePackedShadowMaskBits() &&
diff --git a/Engine/Source/Runtime/Renderer/Private/Shadows/ShadowSceneRenderer.cpp b/Engine/Source/Runtime/Renderer/Private/Shadows/ShadowSceneRenderer.cpp
--- a/Engine/Source/Runtime/Renderer/Private/Shadows/ShadowSceneRenderer.cpp
+++ b/Engine/Source/Runtime/Renderer/Private/Shadows/ShadowSceneRenderer.cpp
@@ -31,7 +31,6 @@ TAutoConsoleVariable<int32> CVarMaxDistantLightsPerFrame(
 	ECVF_Scalability | ECVF_RenderThreadSafe
 );
 
-
 static TAutoConsoleVariable<int32> CVarDistantLightMode(
 	TEXT("r.Shadow.Virtual.DistantLightMode"),
 	0,
@@ -45,6 +44,13 @@ static TAutoConsoleVariable<float> CVarNaniteShadowsLODBias(
 	TEXT("LOD bias for nanite geometry in shadows. 0 = full detail. >0 = reduced detail."),
 	ECVF_RenderThreadSafe);
 
+TAutoConsoleVariable<int32> CVarVirtualShadowOnePassProjection(
+	TEXT("r.Shadow.Virtual.OnePassProjection"),
+	1,
+	TEXT("Projects all local light virtual shadow maps in a single pass for better performance."),
+	ECVF_Scalability | ECVF_RenderThreadSafe
+);
+
 
 DECLARE_DWORD_COUNTER_STAT(TEXT("VSM Total Raster Bins"), STAT_VSMNaniteBasePassTotalRasterBins, STATGROUP_ShadowRendering);
 DECLARE_DWORD_COUNTER_STAT(TEXT("VSM Total Shading Draws"), STAT_VSMNaniteBasePassTotalShadingDraws, STATGROUP_ShadowRendering);
@@ -296,8 +302,6 @@ void FShadowSceneRenderer::PostSetupDebugRender()
 #endif
 }
 
-extern TAutoConsoleVariable<int32> CVarVirtualShadowOnePassProjection;
-
 void FShadowSceneRenderer::RenderVirtualShadowMapProjectionMaskBits(
 	FRDGBuilder& GraphBuilder,
 	FMinimalSceneTextures& SceneTextures)
diff --git a/Engine/Source/Runtime/Renderer/Private/VirtualShadowMaps/VirtualShadowMapArray.cpp b/Engine/Source/Runtime/Renderer/Private/VirtualShadowMaps/VirtualShadowMapArray.cpp
--- a/Engine/Source/Runtime/Renderer/Private/VirtualShadowMaps/VirtualShadowMapArray.cpp
+++ b/Engine/Source/Runtime/Renderer/Private/VirtualShadowMaps/VirtualShadowMapArray.cpp
@@ -1926,6 +1926,9 @@ void FVirtualShadowMapArray::LogStats(FRDGBuilder& GraphBuilder, const FViewInfo
 
 	if (StatsBufferRDG)
 	{
+		// Convenience, enable shader print automatically
+		ShaderPrint::SetEnabled(true);
+
 		int ShowStatsValue = CVarShowStats.GetValueOnRenderThread();
 
 		FVirtualSmLogStatsCS::FParameters* PassParameters = GraphBuilder.AllocParameters<FVirtualSmLogStatsCS::FParameters>();
diff --git a/Engine/Source/Runtime/Renderer/Private/VirtualShadowMaps/VirtualShadowMapProjection.cpp b/Engine/Source/Runtime/Renderer/Private/VirtualShadowMaps/VirtualShadowMapProjection.cpp
--- a/Engine/Source/Runtime/Renderer/Private/VirtualShadowMaps/VirtualShadowMapProjection.cpp
+++ b/Engine/Source/Runtime/Renderer/Private/VirtualShadowMaps/VirtualShadowMapProjection.cpp
@@ -45,13 +45,6 @@ static TAutoConsoleVariable<float> CVarNormalBias(
 	ECVF_Scalability | ECVF_RenderThreadSafe
 );
 
-TAutoConsoleVariable<int32> CVarVirtualShadowOnePassProjection(
-	TEXT("r.Shadow.Virtual.OnePassProjection"),
-	0,
-	TEXT("Single pass projects all local VSMs culled with the light grid. Used in conjunction with clustered deferred shading."),
-	ECVF_Scalability | ECVF_RenderThreadSafe
-);
-
 static TAutoConsoleVariable<int32> CVarSMRTRayCountLocal(
 	TEXT( "r.Shadow.Virtual.SMRT.RayCountLocal" ),
 	7,
@@ -153,13 +146,6 @@ static TAutoConsoleVariable<int32> CVarForcePerLightShadowMaskClear(
 	ECVF_RenderThreadSafe
 );
 
-static TAutoConsoleVariable<int32> CVarSubsurfaceShadowMode(
-	TEXT( "r.Shadow.Virtual.SubsurfaceShadowMode" ),
-	0,
-	TEXT( "" ),
-	ECVF_RenderThreadSafe
-);
-
 static TAutoConsoleVariable<int32> CVarSubsurfaceShadowMinSourceAngle(
 	TEXT("r.Shadow.Virtual.SubsurfaceShadowMinSourceAngle"),
 	5,
@@ -231,7 +217,6 @@ class FVirtualShadowMapProjectionCS : public FGlobalShader
 		SHADER_PARAMETER(FIntVector4, ProjectionRect)
 		SHADER_PARAMETER(float, ScreenRayLength)
 		SHADER_PARAMETER(float, NormalBias)
-		SHADER_PARAMETER(int32, SubsurfaceShadowMode)
 		SHADER_PARAMETER(float, SubsurfaceMinSourceRadius)
 		SHADER_PARAMETER(int32, SMRTRayCount)
 		SHADER_PARAMETER(int32, SMRTSamplesPerRay)
@@ -328,7 +313,6 @@ static void RenderVirtualShadowMapProjectionCommon(
 	PassParameters->ProjectionRect = FIntVector4(ProjectionRect.Min.X, ProjectionRect.Min.Y, ProjectionRect.Max.X, ProjectionRect.Max.Y);
 	PassParameters->ScreenRayLength = CVarScreenRayLength.GetValueOnRenderThread();
 	PassParameters->NormalBias = GetNormalBiasForShader();
-	PassParameters->SubsurfaceShadowMode = CVarSubsurfaceShadowMode.GetValueOnRenderThread();
 	PassParameters->SubsurfaceMinSourceRadius = FMath::Sin(0.5f * FMath::DegreesToRadians(CVarSubsurfaceShadowMinSourceAngle.GetValueOnRenderThread()));
 	PassParameters->InputType = uint32(InputType);
 	PassParameters->bCullBackfacingPixels = VirtualShadowMapArray.ShouldCullBackfacingPixels() ? 1 : 0;
