Removing accidentially submitted debug code from PQ to llinear conversion code
#rnx

diff --git a/Engine/Shaders/Private/GammaCorrectionCommon.ush b/Engine/Shaders/Private/GammaCorrectionCommon.ush
--- a/Engine/Shaders/Private/GammaCorrectionCommon.ush
+++ b/Engine/Shaders/Private/GammaCorrectionCommon.ush
@@ -151,7 +151,7 @@ ST2084ToLinear(float3 pq)
 	float3 Np = pow( pq, 1./m2 );
 	float3 L = Np - c1;
 	L = max(0., L);
-	L = L / max(0.000000001, c2 - c3 * Np);		// DETERMINE SENSICAL CUT-OFF
+	L = L / (c2 - c3 * Np);
 	L = pow( L, 1./m1 );
 	float3 P = L * C;
 
