Path Tracer: Implement diffuse/specular override and support the various EngineShowFlags that use these

This allows debugging diffuse/specular light contributions as well as quickly rendering grey lighting passes.

Implement Subsurface show flag to quickly turn random walks on and off

diff --git a/Engine/Shaders/Private/PathTracing/Material/PathTracingMaterialSampling.ush b/Engine/Shaders/Private/PathTracing/Material/PathTracingMaterialSampling.ush
--- a/Engine/Shaders/Private/PathTracing/Material/PathTracingMaterialSampling.ush
+++ b/Engine/Shaders/Private/PathTracing/Material/PathTracingMaterialSampling.ush
@@ -331,11 +331,17 @@ float GetAverageRoughness(FPathTracingPayload Payload)
 #endif
 }
 
-FPathTracingPayload UnpackPathTracingPayloadMaterial(FPackedPathTracingPayload PackedInput, FRayDesc Ray)
+FPathTracingPayload UnpackPathTracingPayloadMaterial(FPackedPathTracingPayload PackedInput, FRayDesc Ray, bool bApplyOverrides)
 {
 	// This function can also do a bit of per-material post-processing to simplify the rest of the code
 	FPathTracingPayload HitPayload = UnpackPathTracingPayload(PackedInput, Ray);
 
+	if (bApplyOverrides)
+	{
+		HitPayload.DiffuseColor  = HitPayload.DiffuseColor  * View.DiffuseOverrideParameter.w  + View.DiffuseOverrideParameter.xyz;
+		HitPayload.SpecularColor = HitPayload.SpecularColor * View.SpecularOverrideParameter.w + View.SpecularOverrideParameter.xyz;
+	}
+
 	switch (HitPayload.ShadingModelID)
 	{
 #if !SIMPLIFIED_MATERIAL_SHADER
@@ -360,6 +366,10 @@ FPathTracingPayload UnpackPathTracingPayloadMaterial(FPackedPathTracingPayload P
 				HitPayload.DiffuseColor = Front / S;
 				HitPayload.SubsurfaceColor = Back / S;
 			}
+			if (bApplyOverrides)
+			{
+				HitPayload.SubsurfaceColor = HitPayload.SubsurfaceColor * View.DiffuseOverrideParameter.w + View.DiffuseOverrideParameter.xyz;
+			}
 			break;
 		}
 #if SIMPLIFIED_MATERIAL_SHADER
@@ -386,6 +396,24 @@ FPathTracingPayload UnpackPathTracingPayloadMaterial(FPackedPathTracingPayload P
 			HitPayload.DiffuseColor *= IrisMask;
 			break;
 		}
+		case SHADINGMODELID_HAIR:
+		{
+			if (bApplyOverrides)
+			{
+				// This is not ideal since hair is technically specular transmission
+				HitPayload.BaseColor = HitPayload.BaseColor  * View.DiffuseOverrideParameter.w  + View.DiffuseOverrideParameter.xyz;
+				HitPayload.Specular = HitPayload.Specular * View.SpecularOverrideParameter.w + View.SpecularOverrideParameter.x;
+			}
+			break;
+		}
+		case SHADINGMODELID_CLEAR_COAT:
+		{
+			if (bApplyOverrides)
+			{
+				HitPayload.SetClearCoat(HitPayload.GetClearCoat() * View.SpecularOverrideParameter.w + View.SpecularOverrideParameter.x);
+			}
+			break;
+		}
 #endif
 	}
 
diff --git a/Engine/Shaders/Private/PathTracing/PathTracingCore.ush b/Engine/Shaders/Private/PathTracing/PathTracingCore.ush
--- a/Engine/Shaders/Private/PathTracing/PathTracingCore.ush
+++ b/Engine/Shaders/Private/PathTracing/PathTracingCore.ush
@@ -6,7 +6,6 @@
 
 #define USE_PATH_TRACING_LIGHT_GRID	1
 #define USE_RAY_TRACING_DECAL_GRID	1
-#define USE_ENERGY_CONSERVATION 1
 
 // Ignore hair strands if the ray origin is beyond this radius. This is meant to workaround a performance cliff that can occur on some hardware.
 // In ordinary scenes, the most likely occurance of this problem is when combining hair and volumetric atmosphere as rays from several Km can travel back towards the hair and cause long delays or even TDRs
@@ -147,9 +146,6 @@ void ApplyDecal(inout FPathTracingPayload HitPayload, FDecalShaderParams Params)
 		HitPayload.Radiance += Params.GetEmissive() * DecalAlpha;
 	}
 
-	HitPayload.DiffuseColor = HitPayload.BaseColor - HitPayload.BaseColor * HitPayload.Metallic;
-	HitPayload.SpecularColor = ComputeF0(HitPayload.Specular, HitPayload.BaseColor, HitPayload.Metallic);
-
 	// TODO: Payload.Opacity?
 }
 
@@ -367,8 +363,9 @@ FPathTracingPayload TraceTransparentRay(inout FPathState PathState, bool bIsCame
 		}
 
 		// Unpack the payload
-		FPathTracingPayload HitPayload = UnpackPathTracingPayloadMaterial(PackedPayload, PathState.Ray);
+		FPathTracingPayload HitPayload = UnpackPathTracingPayloadMaterial(PackedPayload, PathState.Ray, bIsCameraRay);
 
+		float SumDecalOpacity = 0;
 		if (PathState.PathRoughness < MeshDecalRoughnessCutoff && HitPayload.IsDecalReceiver())
 		{
 			FRayDesc DecalRay = (FRayDesc)0;
@@ -404,11 +401,7 @@ FPathTracingPayload TraceTransparentRay(inout FPathState PathState, bool bIsCame
 			{
 				ApplyDecal(HitPayload, DecalPayload.Params);
 
-				if (ApproximateCaustics)
-				{
-					Payload.Roughness = max(Payload.Roughness, PathState.PathRoughness);
-				}
-				HitPayload.WorldNormal = AdjustShadingNormal(HitPayload.WorldNormal, HitPayload.WorldGeoNormal, PathState.Ray.Direction);
+				SumDecalOpacity += DecalPayload.Params.GetOpacity();
 			}
 		}
 
@@ -418,7 +411,6 @@ FPathTracingPayload TraceTransparentRay(inout FPathState PathState, bool bIsCame
 			const float3 TranslatedWorldPos = HitPayload.TranslatedWorldPos;
 
 			FDecalLoopCount DecalLoopCount = DecalGridLookup(TranslatedWorldPos);
-
 			for (uint Index = 0, Num = DecalLoopCount.NumDecals; Index < Num; ++Index)
 			{
 				uint DecalId = GetDecalId(Index, DecalLoopCount);
@@ -430,14 +422,27 @@ FPathTracingPayload TraceTransparentRay(inout FPathState PathState, bool bIsCame
 
 				ApplyDecal(HitPayload, Params);
 
-				if (ApproximateCaustics)
-				{
-					Payload.Roughness = max(Payload.Roughness, PathState.PathRoughness);
-				}
-				HitPayload.WorldNormal = AdjustShadingNormal(HitPayload.WorldNormal, HitPayload.WorldGeoNormal, PathState.Ray.Direction);
+				SumDecalOpacity += Params.GetOpacity();
 			}
 		}
 #endif
+		if (SumDecalOpacity > 0)
+		{
+			// if we ran any decal logic, update the payload
+			if (ApproximateCaustics)
+			{
+				Payload.Roughness = max(Payload.Roughness, PathState.PathRoughness);
+			}
+			HitPayload.WorldNormal = AdjustShadingNormal(HitPayload.WorldNormal, HitPayload.WorldGeoNormal, PathState.Ray.Direction);
+			HitPayload.DiffuseColor = HitPayload.BaseColor - HitPayload.BaseColor * HitPayload.Metallic;
+			HitPayload.SpecularColor = ComputeF0(HitPayload.Specular, HitPayload.BaseColor, HitPayload.Metallic);
+			if (bIsCameraRay)
+			{
+				HitPayload.DiffuseColor = HitPayload.DiffuseColor * View.DiffuseOverrideParameter.w + View.DiffuseOverrideParameter.xyz;
+				HitPayload.SpecularColor = HitPayload.SpecularColor * View.SpecularOverrideParameter.w + View.SpecularOverrideParameter.xyz;
+			}
+		}
+
 		// account for Beer's law through the currently active medium
 		// TODO: merge this with the volume raymarching?
 		PathState.PathThroughput *= select(PathState.SigmaT > 0.0, exp(-PathState.SigmaT * (HitPayload.HitT - PathState.Ray.TMin)), 1.0);
diff --git a/Engine/Source/Runtime/Renderer/Private/PathTracing.cpp b/Engine/Source/Runtime/Renderer/Private/PathTracing.cpp
--- a/Engine/Source/Runtime/Renderer/Private/PathTracing.cpp
+++ b/Engine/Source/Runtime/Renderer/Private/PathTracing.cpp
@@ -514,7 +514,7 @@ static void PreparePathTracingData(const FScene* Scene, const FViewInfo& View, F
 	}
 
 	PathTracingData.MaxBounces = MaxBounces;
-	PathTracingData.MaxSSSBounces = CVarPathTracingMaxSSSBounces.GetValueOnRenderThread();
+	PathTracingData.MaxSSSBounces = View.Family->EngineShowFlags.SubsurfaceScattering ? CVarPathTracingMaxSSSBounces.GetValueOnRenderThread() : 0;
 	PathTracingData.MaxNormalBias = GetRaytracingMaxNormalBias();
 	PathTracingData.MISMode = CVarPathTracingMISMode.GetValueOnRenderThread();
 	PathTracingData.VolumeMISMode = CVarPathTracingVolumeMISMode.GetValueOnRenderThread();
@@ -2080,6 +2080,13 @@ void FDeferredShadingSceneRenderer::RenderPathTracing(
 	Config.LightShowFlags |= View.Family->EngineShowFlags.TexturedLightProfiles ? 1 << 5 : 0;
 	Config.LightShowFlags |= View.Family->EngineShowFlags.LightFunctions        ? 1 << 6 : 0;
 	Config.LightShowFlags |= CVarPathTracingLightFunctionColor.GetValueOnRenderThread() ? 1 << 7 : 0;
+	// the following flags all mess with diffuse/spec overrides and therefore change the image
+	Config.LightShowFlags |= View.Family->EngineShowFlags.Diffuse                    ? 1 << 8 : 0;
+	Config.LightShowFlags |= View.Family->EngineShowFlags.Specular                   ? 1 << 9 : 0;
+	Config.LightShowFlags |= View.Family->EngineShowFlags.OverrideDiffuseAndSpecular ? 1 << 10 : 0;
+	Config.LightShowFlags |= View.Family->EngineShowFlags.LightingOnlyOverride       ? 1 << 11 : 0;
+	Config.LightShowFlags |= View.Family->EngineShowFlags.ReflectionOverride         ? 1 << 12 : 0;
+	Config.LightShowFlags |= View.Family->EngineShowFlags.SubsurfaceScattering       ? 1 << 13 : 0;
 
 	PreparePathTracingData(Scene, View, Config.PathTracingData);
 
