Fix invalid hair sample fetching.

diff --git a/Engine/Shaders/Private/HairStrands/HairStrandsForwardRaster.usf b/Engine/Shaders/Private/HairStrands/HairStrandsForwardRaster.usf
--- a/Engine/Shaders/Private/HairStrands/HairStrandsForwardRaster.usf
+++ b/Engine/Shaders/Private/HairStrands/HairStrandsForwardRaster.usf
@@ -1216,11 +1216,11 @@ float4 UnpackVelocityData(uint2 In)
 	return float4(UnpackRG16(In.x), UnpackRG16(In.y));
 }
 
-float3 LoadSampleColor(uint InPrimId)
+float3 LoadSampleColor(uint InPrimId, uint2 InSampleResolution)
 {
 	uint2 SampleCoord = 0;
-	SampleCoord.x = InPrimId % SampleLightingViewportResolution.x;
-	SampleCoord.y = InPrimId / SampleLightingViewportResolution.x;
+	SampleCoord.x = InPrimId % InSampleResolution.x;
+	SampleCoord.y = InPrimId / InSampleResolution.x;
 	return SampleLightingTexture.Load(uint3(SampleCoord, 0)).xyz;
 }
 
@@ -1277,6 +1277,13 @@ void RasterCS(
 	uint   Thread_LoopCountToFullCoverage = 0;
 	uint2  Thread_PixelCoord = 0;
 
+	uint2 SampleLightingEffectiveResolution = 0;
+	{
+		const uint EffectivePointCount = ControlPointCount.Load(uint3(0, 0, 0));
+		SampleLightingEffectiveResolution.x = ceil(sqrt(EffectivePointCount));
+		SampleLightingEffectiveResolution.y = SampleLightingEffectiveResolution.x;
+	}
+
 	LOOP
 	for (uint WorkIndex = 0; WorkIndex < MAX_WORK_COUNT; WorkIndex++)
 	{
@@ -1478,8 +1485,8 @@ void RasterCS(
 							const float3 P0 = UnpackHairControlPoint(PrimID).Position;
 							const float3 P1 = UnpackHairControlPoint(PrimID+1).Position;
 
-							const float4 Color0 = float4(LoadSampleColor(PrimID), 1);
-							const float4 Color1 = float4(LoadSampleColor(PrimID+1), 1);
+							const float4 Color0 = float4(LoadSampleColor(PrimID, SampleLightingEffectiveResolution), 1);
+							const float4 Color1 = float4(LoadSampleColor(PrimID+1, SampleLightingEffectiveResolution), 1);
 
 							const float4 LineColor = float4(ColorMapTurbo(GroupThread1D / float(RASTER_THREAD_COUNT)), 1);
 							AddLineWS(P0, P1, LineColor);
@@ -1767,8 +1774,8 @@ void RasterCS(
 								//const float3 P0 = UnpackHairControlPoint(PrimID).Position;
 								//const float3 P1 = UnpackHairControlPoint(PrimID + 1).Position;
 
-								//const float4 Color0 = float4(LoadSampleColor(PrimID), 1);
-								//const float4 Color1 = float4(LoadSampleColor(PrimID + 1), 1);
+								//const float4 Color0 = float4(LoadSampleColor(PrimID, SampleLightingEffectiveResolution), 1);
+								//const float4 Color1 = float4(LoadSampleColor(PrimID + 1, SampleLightingEffectiveResolution), 1);
 
 								//const float4 LineColor = float4(ColorMapTurbo(GroupThread1D / float(RASTER_THREAD_COUNT)), 1);
 								//AddLineWS(P0, P1, LineColor);
@@ -1785,8 +1792,8 @@ void RasterCS(
 								//Print(Ctx, TEXT("Color0: "), FontWhite); Print(Ctx, Color0); Newline(Ctx);
 								//Print(Ctx, TEXT("Color1: "), FontWhite); Print(Ctx, Color1); Newline(Ctx);
 								//Print(Ctx, TEXT("Color :  "), FontWhite); Print(Ctx, Color);
-								const float3 Color0 = LoadSampleColor(PrimID);
-								const float3 Color1 = LoadSampleColor(PrimID + 1);
+								const float3 Color0 = LoadSampleColor(PrimID, SampleLightingEffectiveResolution);
+								const float3 Color1 = LoadSampleColor(PrimID + 1, SampleLightingEffectiveResolution);
 								Print(Ctx, TEXT("Color0 :  "), FontWhite); Print(Ctx, Color0); Newline(Ctx);
 								Print(Ctx, TEXT("Color1 :  "), FontWhite); Print(Ctx, Color1); Newline(Ctx);
 								Print(Ctx, TEXT("  Alpha       A0          A1          AColor     Color    Coord  Coverage"), FontWhite);
@@ -1813,8 +1820,8 @@ void RasterCS(
 									if (Depth > OpaqueDepth)
 									{
 										const float AlphaColor = lerp(Alpha0, Alpha1, Alpha);
-										const float3 Color0 = LoadSampleColor(PrimID);
-										const float3 Color1 = LoadSampleColor(PrimID+1);
+										const float3 Color0 = LoadSampleColor(PrimID, SampleLightingEffectiveResolution);
+										const float3 Color1 = LoadSampleColor(PrimID+1, SampleLightingEffectiveResolution);
 										const float3 Color = lerp(Color0, Color1, AlphaColor);
 
 										const float4 Velocity0 = LoadSampleVelocity(PrimID);
