Cloud - fixed ground contribution to work all around the globe.

diff --git a/Engine/Shaders/Private/VolumetricCloud.usf b/Engine/Shaders/Private/VolumetricCloud.usf
--- a/Engine/Shaders/Private/VolumetricCloud.usf
+++ b/Engine/Shaders/Private/VolumetricCloud.usf
@@ -761,11 +761,6 @@ void MainCommon(in FMaterialPixelParameters MaterialParameters, in float4 SvPosi
 			RenderVolumetricCloudParameters.CloudLayerCenterKm, RenderVolumetricCloudParameters.PlanetRadiusKm,
 			RenderVolumetricCloudParameters.BottomRadiusKm, RenderVolumetricCloudParameters.TopRadiusKm);
 
-		const float3 PlanetSurfaceNormal = float3(0.0f, 0.0f, 1.0f);	// Ambient contribution from the clouds is only done on a plane above the planet, e.g. do not support space view yet
-		const float3 GroundBrdfNdotL = saturate(dot(Light0Direction, PlanetSurfaceNormal)) * (RenderVolumetricCloudParameters.GroundAlbedo.rgb / PI); // Assuming pure Lambert diffuse surface.
-		const float3 GroundHemisphereLuminanceIsotropic = (2.0f * PI) * IsotropicPhase(); // Assumes the ground is uniform luminance to the cloud and solid angle is bottom hemisphere 2PI
-		const float3 GroundToCloudTransfertIsoScatter = GroundBrdfNdotL * GroundHemisphereLuminanceIsotropic;
-
 		//
 		// TRACING LOOP
 		//
@@ -927,10 +922,6 @@ void MainCommon(in FMaterialPixelParameters MaterialParameters, in float4 SvPosi
 				//
 				// A- we compute lighting bouncing off the ground: affected by light direction, transmittance in the atmosphere and albedo.
 				//
-				const float3 AtmosphereTransmittanceToGround0 = AtmosphereTransmittanceToLight0; // big approximation when CLOUD_PER_SAMPLE_ATMOSPHERE_TRANSMITTANCE is true
-				const float3 AtmosphereTransmittanceToGround1 = AtmosphereTransmittanceToLight1; // idem
-				const float3 ScatteredLightLuminance = (AtmosphereTransmittanceToGround0 * Light0Illuminance + AtmosphereTransmittanceToGround1 * Light1Illuminance) * GroundToCloudTransfertIsoScatter;
-				const float3 TransmittedScatteredLightLuminance = ScatteredLightLuminance * exp(-PMC.ExtinctionCoefficients[0] * MaterialParameters.CloudSampleAltitudeInLayer * CENTIMETER_TO_METER);
 #if MATERIAL_VOLUMETRIC_ADVANCED_GROUND_CONTRIBUTION
 				if (MaxTransmittanceToView > 0.01f)
 				{
@@ -942,6 +933,10 @@ void MainCommon(in FMaterialPixelParameters MaterialParameters, in float4 SvPosi
 					const float ShadowStepCount = 5.0f;
 					const float InvShadowStepCount = 1.0f / ShadowStepCount;
 
+					// Evaluate the direction toward the ground only once
+					const float3 GroundNormal    = normalize(PlanetCenterToWorldPos);	// Ambient contribution from the clouds is only done on a plane above the planet, e.g. do not support space view yet
+					const float3 GroundDirection = -GroundNormal;
+
 #if 1
 					// Non-linear shadow sample distribution, hardcoded to x^2
 					float PreviousNormT = 0.0f;
@@ -950,7 +945,7 @@ void MainCommon(in FMaterialPixelParameters MaterialParameters, in float4 SvPosi
 						float CurrentNormT = ShadowT * ShadowT; // CurrentNormT is the end of the considered segment to integrate, PreviousNormT is its beginning.
 						const float DetlaNormT = CurrentNormT - PreviousNormT;
 						const float ShadowSampleDistance = ShadowLengthTest * (CurrentNormT - 0.5 * DetlaNormT);
-						UpdateMaterialCloudParam(MaterialParameters, LWCAdd(SampleWorldPosition, LWCPromote(float3(0.0f, 0.0f, -1.0f) * ShadowSampleDistance)),
+						UpdateMaterialCloudParam(MaterialParameters, LWCAdd(SampleWorldPosition, LWCPromote(GroundDirection * ShadowSampleDistance)),
 							ResolvedView, CloudLayerParams, ShadowSampleDistance, SPACE_SKIPPING_SLICE_DEPTH_OFF);
 						PreviousNormT = CurrentNormT;
 #if MATERIAL_VOLUMETRIC_ADVANCED_CONSERVATIVE_DENSITY
@@ -985,6 +980,15 @@ void MainCommon(in FMaterialPixelParameters MaterialParameters, in float4 SvPosi
 					}
 #endif
 
+					const float3 GroundBrdfNdotL = saturate(dot(Light0Direction, GroundNormal)) * (RenderVolumetricCloudParameters.GroundAlbedo.rgb / PI); // Assuming pure Lambert diffuse surface.
+					const float3 GroundHemisphereLuminanceIsotropic = (2.0f * PI) * IsotropicPhase(); // Assumes the ground is uniform luminance to the cloud and solid angle is bottom hemisphere 2PI
+					const float3 GroundToCloudTransfertIsoScatter = GroundBrdfNdotL * GroundHemisphereLuminanceIsotropic;
+
+					const float3 AtmosphereTransmittanceToGround0 = AtmosphereTransmittanceToLight0; // big approximation when CLOUD_PER_SAMPLE_ATMOSPHERE_TRANSMITTANCE is true
+					const float3 AtmosphereTransmittanceToGround1 = AtmosphereTransmittanceToLight1; // idem
+					const float3 ScatteredLightLuminance = (AtmosphereTransmittanceToGround0 * Light0Illuminance + AtmosphereTransmittanceToGround1 * Light1Illuminance) * GroundToCloudTransfertIsoScatter;
+					const float3 TransmittedScatteredLightLuminance = ScatteredLightLuminance * exp(-PMC.ExtinctionCoefficients[0] * MaterialParameters.CloudSampleAltitudeInLayer * CENTIMETER_TO_METER);
+
 					DistantLightLuminance += ScatteredLightLuminance * exp(-OpticalDepth);
 				}
 #endif // MATERIAL_VOLUMETRIC_ADVANCED_GROUND_CONTRIBUTION
