Fix EngineTest shader compile error.

#rnx

diff --git a/Engine/Shaders/Private/SparseVolumeTexture/VisualizeSparseVolumeTexture.usf b/Engine/Shaders/Private/SparseVolumeTexture/VisualizeSparseVolumeTexture.usf
--- a/Engine/Shaders/Private/SparseVolumeTexture/VisualizeSparseVolumeTexture.usf
+++ b/Engine/Shaders/Private/SparseVolumeTexture/VisualizeSparseVolumeTexture.usf
@@ -189,7 +189,7 @@ void VisualizeSparseVolumeTexturePS(
 				else
 				{
 					float4 VoxelData = SparseVolumeTextureB.Load(int4(VoxelCoord, 0));
-					Density = VoxelData[min(ComponentToVisualize - 4, 3)];
+					Density = VoxelData[min(ComponentToVisualize - 4, 3u)];
 				}
 
 				if (Density > 0.0f)
