Fix an issue with one pass projection attempting to clear a null screen shadow mask
Temporary work around mismatch in spot angle due to quantization in light grid

diff --git a/Engine/Shaders/Private/VirtualShadowMaps/VirtualShadowMapProjection.usf b/Engine/Shaders/Private/VirtualShadowMaps/VirtualShadowMapProjection.usf
--- a/Engine/Shaders/Private/VirtualShadowMaps/VirtualShadowMapProjection.usf
+++ b/Engine/Shaders/Private/VirtualShadowMaps/VirtualShadowMapProjection.usf
@@ -220,7 +220,15 @@ FVirtualShadowMapSampleResult ProjectLight(
 	float3 L = ToLight * InvDist;
 	
 	bool bInLightRadius = InvDist >= Light.InvRadius;
-	bool bInLightCone = dot( L, Light.Direction ) >= Light.SpotAngles.x;
+	bool bInLightCone = SpotAttenuationMask(L, -Light.Direction, Light.SpotAngles) > 0.0f;
+
+	#if ONE_PASS_PROJECTION
+		// Hacky workaround for spot angles being quantized to 16-bit when they are loaded from the light grid!
+		// This avoids potential artifacts at cone edges, and in this path we have the tile culling to lessen the
+		// overhead. Still needs to be fixed though as it is not trivial overhead!
+		bInLightCone = true;
+	#endif
+
 	bool bInLightRegion = bInLightRadius && bInLightCone;
 #endif // DIRECTIONAL_LIGHT
 
@@ -302,6 +310,7 @@ FVirtualShadowMapSampleResult ProjectLight(
 		}
 
 		//Result.GeneralDebug = GreenToRedTurbo(1.0f - (SMRTRayOffset / ScreenRayLengthWorld));
+		//Result.GeneralDebug = SpotAttenuation(-L, Light.Direction, Light.SpotAngles).xxx;
 
 		// Hair strands voxel traversal to apply hair shadow on top of opaque geometry
 		#if HAS_HAIR_STRANDS
diff --git a/Engine/Source/Runtime/Renderer/Private/LightRendering.cpp b/Engine/Source/Runtime/Renderer/Private/LightRendering.cpp
--- a/Engine/Source/Runtime/Renderer/Private/LightRendering.cpp
+++ b/Engine/Source/Runtime/Renderer/Private/LightRendering.cpp
@@ -1889,7 +1889,11 @@ void FDeferredShadingSceneRenderer::RenderLights(
 							}
 						};
 
-						ClearShadowMask(ScreenShadowMaskTexture);
+						// Skip clearing if it doesn't exist (due to being elided by VSM projection)
+						if (ScreenShadowMaskTexture)
+						{
+							ClearShadowMask(ScreenShadowMaskTexture);
+						}
 						if (ScreenShadowMaskSubPixelTexture)
 						{
 							ClearShadowMask(ScreenShadowMaskSubPixelTexture);
