Pack data in FDecalShaderParams.
- avoids having to overwrite translated world position.
- also required to support blending of mesh decals.

diff --git a/Engine/Shaders/Private/PathTracing/PathTracingCore.ush b/Engine/Shaders/Private/PathTracing/PathTracingCore.ush
--- a/Engine/Shaders/Private/PathTracing/PathTracingCore.ush
+++ b/Engine/Shaders/Private/PathTracing/PathTracingCore.ush
@@ -131,9 +131,11 @@ void ApplyDecal(inout FPathTracingPayload HitPayload, FDecalShaderParams Params)
 
 	if (DecalFlags & DECAL_WRITE_ROUGHNESS_SPECULAR_METALLIC_FLAG)
 	{
-		HitPayload.Metallic = lerp(HitPayload.Metallic, Params.GetMetallic(), DecalAlpha);
-		HitPayload.Specular = lerp(HitPayload.Specular, Params.GetSpecular(), DecalAlpha);
-		HitPayload.Roughness = lerp(HitPayload.Roughness, Params.GetRoughness(), DecalAlpha);
+		const float3 MetallicSpecularRoughness = Params.GetMetallicSpecularRoughness();
+		
+		HitPayload.Metallic = lerp(HitPayload.Metallic, MetallicSpecularRoughness.x, DecalAlpha);
+		HitPayload.Specular = lerp(HitPayload.Specular, MetallicSpecularRoughness.y, DecalAlpha);
+		HitPayload.Roughness = lerp(HitPayload.Roughness, MetallicSpecularRoughness.z, DecalAlpha);
 	}
 
 	if (DecalFlags & DECAL_WRITE_AMBIENT_OCCLUSION_FLAG)
diff --git a/Engine/Shaders/Private/RayTracing/RayTracingCommon.ush b/Engine/Shaders/Private/RayTracing/RayTracingCommon.ush
--- a/Engine/Shaders/Private/RayTracing/RayTracingCommon.ush
+++ b/Engine/Shaders/Private/RayTracing/RayTracingCommon.ush
@@ -921,54 +921,62 @@ CHECK_RT_PAYLOAD_SIZE(FPackedMaterialClosestHitPayload)
 #if IS_PAYLOAD_ENABLED(RT_PAYLOAD_TYPE_DECALS | RT_PAYLOAD_TYPE_MESH_DECALS)
 struct FDecalShaderParams
 {
-	float Data[15];
+	uint Data[15];
 
 	void SetInputTranslatedWorldPosition(float3 In)
 	{
-		Data[0] = In.x;
-		Data[1] = In.y;
-		Data[2] = In.z;
+		Data[0] = asuint(In.x);
+		Data[1] = asuint(In.y);
+		Data[2] = asuint(In.z);
 	}
-	float3 GetInputTranslatedWorldPosition() { return float3(Data[0], Data[1], Data[2]); }
+	float3 GetInputTranslatedWorldPosition() { return float3(asfloat(Data[0]), asfloat(Data[1]), asfloat(Data[2])); }
 
-	void SetMetallic(float In) { Data[0] = In; }
-	float GetMetallic() { return Data[0]; }
-
-	void SetSpecular(float In) { Data[1] = In; }
-	float GetSpecular() { return Data[1]; }
-
-	void SetRoughness(float In) { Data[2] = In; }
-	float GetRoughness() { return Data[2]; }
-
-	void SetAO(float In) { Data[3] = In; }
-	float GetAO() { return Data[3]; }
+	void SetMetallicSpecularRoughness(float3 In)
+	{
+		Data[3] = f32tof16(In.r);
+		Data[3] |= f32tof16(In.g) << 16;
+		Data[4] = f32tof16(In.b);
+	}
+	float3 GetMetallicSpecularRoughness() { return float3(f16tof32(Data[3]), f16tof32(Data[3] >> 16), f16tof32(Data[4])); }
 
-	void SetBaseColor(float3 In) {
-		Data[4] = In.r;
-		Data[5] = In.g;
-		Data[6] = In.b;
+	void SetBaseColor(float3 In)
+	{
+		Data[5] = f32tof16(In.r);
+		Data[5] |= f32tof16(In.g) << 16;
+		Data[6] = f32tof16(In.b);
 	}
-	float3 GetBaseColor() { return float3(Data[4], Data[5], Data[6]); }
+	float3 GetBaseColor() { return float3(f16tof32(Data[5]), f16tof32(Data[5] >> 16), f16tof32(Data[6])); }
 
-	void SetOpacity(float In) { Data[7] = In; }
-	float GetOpacity() { return Data[7]; }
+	void SetWorldNormal(float3 In)
+	{
+		Data[7] = asuint(In.r);
+		Data[8] = asuint(In.g);
+		Data[9] = asuint(In.b);
+	}
+	float3 GetWorldNormal() { return float3(asfloat(Data[7]), asfloat(Data[8]), asfloat(Data[9])); }
 
-	void SetWorldNormal(float3 In) {
-		Data[8] = In.r;
-		Data[9] = In.g;
-		Data[10] = In.b;
+	void SetEmissive(float3 In)
+	{
+		Data[10] = f32tof16(In.r);
+		Data[10] |= f32tof16(In.g) << 16;
+		Data[11] = f32tof16(In.b);
 	}
-	float3 GetWorldNormal() { return float3(Data[8], Data[9], Data[10]); }
+	float3 GetEmissive() { return float3(f16tof32(Data[10]), f16tof32(Data[10] >> 16), f16tof32(Data[11])); }
 
-	void SetFlags(uint In) { Data[11] = asfloat(In); }
-	uint GetFlags() { return asuint(Data[11]); }
+	void SetAO(float In)
+	{
+		Data[12] = asuint(In);
+	}
+	float GetAO() { return asfloat(Data[12]); }
 
-	void SetEmissive(float3 In) {
-		Data[12] = In.r;
-		Data[13] = In.g;
-		Data[14] = In.b;
+	void SetOpacity(float In)
+	{
+		Data[13] = asuint(In);
 	}
-	float3 GetEmissive() { return float3(Data[12], Data[13], Data[14]); }
+	float GetOpacity() { return asfloat(Data[13]); }
+
+	void SetFlags(uint In) { Data[14] = In; }
+	uint GetFlags() { return Data[14]; }
 };
 #endif
 
diff --git a/Engine/Shaders/Private/RayTracing/RayTracingDecalMaterialShader.usf b/Engine/Shaders/Private/RayTracing/RayTracingDecalMaterialShader.usf
--- a/Engine/Shaders/Private/RayTracing/RayTracingDecalMaterialShader.usf
+++ b/Engine/Shaders/Private/RayTracing/RayTracingDecalMaterialShader.usf
@@ -78,9 +78,7 @@ RAY_TRACING_ENTRY_CALLABLE(RayTracingDecalMaterialShader, FDecalShaderParams, Pa
 
 	if (DecalWriteFlags & DECAL_WRITE_ROUGHNESS_SPECULAR_METALLIC_FLAG)
 	{
-		Params.SetMetallic(GetMaterialMetallic(PixelMaterialInputs));
-		Params.SetSpecular(GetMaterialSpecular(PixelMaterialInputs));
-		Params.SetRoughness(GetMaterialRoughness(PixelMaterialInputs));
+		Params.SetMetallicSpecularRoughness(float3(GetMaterialMetallic(PixelMaterialInputs), GetMaterialSpecular(PixelMaterialInputs), GetMaterialRoughness(PixelMaterialInputs)));
 	}
 
 	if (DecalWriteFlags & DECAL_WRITE_AMBIENT_OCCLUSION_FLAG)
@@ -161,9 +159,7 @@ RAY_TRACING_ENTRY_CLOSEST_HIT(RayTracingDecalMaterialCHS,
 
 	if (DECAL_PAYLOAD_FLAGS & DECAL_WRITE_ROUGHNESS_SPECULAR_METALLIC_FLAG)
 	{
-		Payload.Params.SetMetallic(GetMaterialMetallic(PixelMaterialInputs));
-		Payload.Params.SetSpecular(GetMaterialSpecular(PixelMaterialInputs));
-		Payload.Params.SetRoughness(GetMaterialRoughness(PixelMaterialInputs));
+		Payload.Params.SetMetallicSpecularRoughness(float3(GetMaterialMetallic(PixelMaterialInputs), GetMaterialSpecular(PixelMaterialInputs), GetMaterialRoughness(PixelMaterialInputs)));
 	}
 
 	if (DECAL_PAYLOAD_FLAGS & DECAL_WRITE_AMBIENT_OCCLUSION_FLAG)
