Fix RT RectLight Shadow sampling issue. Sample from spherical rect instead of the rect light plane so that the sampling matches the pdf to remove visible large banding bias.

diff --git a/Engine/Shaders/Private/RayTracing/RayTracingRectLight.ush b/Engine/Shaders/Private/RayTracing/RayTracingRectLight.ush
--- a/Engine/Shaders/Private/RayTracing/RayTracingRectLight.ush
+++ b/Engine/Shaders/Private/RayTracing/RayTracingRectLight.ush
@@ -24,34 +24,23 @@ bool GenerateRectLightOcclusionRay(
 	RayTMax = 0.0;
 	RayPdf = 0.0;
 
-	float3 Tangent = float3(0,0,0);
-	float3 BiTangent = float3(0,0,0);
-	GetRectLightTangents(LightParameters, /* out */ Tangent, /* out */ BiTangent);
-	float2 LightDimensions = GetRectLightDimensions(LightParameters);
-	
-	// Draw random variable
-	RandSample -= 0.5;
+	// Define rectangle
+	FRect Rect = GetRect(LightParameters, TranslatedWorldPosition);
+
+	// 1. Derive the pdf when uniformly sampling from the spherical rectangle
+	FSphericalRect SphericalRect = BuildSphericalRect(Rect);
+	RayPdf = 1.0 / SphericalRect.SolidAngle;
 
-	// Map sample point to quad
-	float3 LightSamplePosition = LightParameters.TranslatedWorldPosition + Tangent * LightDimensions.x * RandSample.x + BiTangent * LightDimensions.y * RandSample.y;
-	float3 LightDirection = normalize(LightSamplePosition - TranslatedWorldPosition);
+	// 2. Uniformly sample from the spherical rectangle
+	float3 LightDirection = UniformSampleSphericalRect(RandSample, SphericalRect);
 
 	// Light-normal culling
 	bool bIsValidRay = dot(-LightDirection, -LightParameters.Direction) >= 0.0;
-
-	// Apply normal perturbation when defining ray
-	RayDirection = LightDirection;
-	RayTMax = length(LightSamplePosition - TranslatedWorldPosition);
-
-	// Define rectangle
-	FRect Rect;
-	Rect.Origin = LightParameters.TranslatedWorldPosition - TranslatedWorldPosition;
-	Rect.Axis[0] = Tangent;
-	Rect.Axis[1] = BiTangent;
-	Rect.Axis[2] = cross(Tangent, BiTangent);
-	Rect.Extent = LightDimensions * 0.5f;
-	FSphericalRect SphericalRect = BuildSphericalRect(Rect);
-	RayPdf = 1.0 / SphericalRect.SolidAngle;
+	
+	// Apply normal perturbation when defining ray. normalize() resolves several banding artifacts
+	// at the cost of creating some bias under idea lighting condition but is neglectable.
+	RayDirection = normalize(LightDirection);
+	RayTMax = length(LightDirection);
 
 	return bIsValidRay;
 }
\ No newline at end of file
diff --git a/Engine/Shaders/Private/RectLight.ush b/Engine/Shaders/Private/RectLight.ush
--- a/Engine/Shaders/Private/RectLight.ush
+++ b/Engine/Shaders/Private/RectLight.ush
@@ -660,6 +660,19 @@ FRect GetRect(
 	return Rect;
 }
 
+// Get the Rect of the light with respect to the translated world position.
+FRect GetRect(FLightShaderParameters In, float3 TranslatedWorldPosition)
+{
+	return GetRect(In.TranslatedWorldPosition - TranslatedWorldPosition,
+		In.Direction,
+		In.Tangent,
+		In.SourceRadius,
+		In.SourceLength,
+		In.RectLightBarnCosAngle,
+		In.RectLightBarnLength,
+		true);
+}
+
 bool IsRectVisible(FRect Rect)
 {
 	// No-visible rect light due to barn door occlusion
