Use premultiplied alpha for applying decals in path tracing.
- will simplify blending in AHS later.

diff --git a/Engine/Shaders/Private/PathTracing/PathTracingCore.ush b/Engine/Shaders/Private/PathTracing/PathTracingCore.ush
--- a/Engine/Shaders/Private/PathTracing/PathTracingCore.ush
+++ b/Engine/Shaders/Private/PathTracing/PathTracingCore.ush
@@ -116,26 +116,26 @@ void AccumulateRadiance(inout float3 TotalRadiance, float3 PathRadiance, bool bI
 
 void ApplyDecal(inout FPathTracingPayload HitPayload, FDecalShaderParams Params)
 {
-	const float DecalAlpha = Params.GetOpacity();
+	const float DecalTransparency = Params.GetTransparency();
 	const uint DecalFlags = Params.GetFlags();
 
 	if (DecalFlags & DECAL_WRITE_BASE_COLOR_FLAG)
 	{
-		HitPayload.BaseColor = lerp(HitPayload.BaseColor, Params.GetBaseColor(), DecalAlpha);
+		HitPayload.BaseColor = HitPayload.BaseColor * DecalTransparency + Params.GetBaseColor();
 	}
 
 	if (DecalFlags & DECAL_WRITE_NORMAL_FLAG)
 	{
-		HitPayload.WorldNormal = lerp(HitPayload.WorldNormal, Params.GetWorldNormal(), DecalAlpha);
+		HitPayload.WorldNormal = HitPayload.WorldNormal * DecalTransparency + Params.GetWorldNormal();
 	}
 
 	if (DecalFlags & DECAL_WRITE_ROUGHNESS_SPECULAR_METALLIC_FLAG)
 	{
 		const float3 MetallicSpecularRoughness = Params.GetMetallicSpecularRoughness();
-		
-		HitPayload.Metallic = lerp(HitPayload.Metallic, MetallicSpecularRoughness.x, DecalAlpha);
-		HitPayload.Specular = lerp(HitPayload.Specular, MetallicSpecularRoughness.y, DecalAlpha);
-		HitPayload.Roughness = lerp(HitPayload.Roughness, MetallicSpecularRoughness.z, DecalAlpha);
+
+		HitPayload.Metallic = HitPayload.Metallic * DecalTransparency + MetallicSpecularRoughness.x;
+		HitPayload.Specular = HitPayload.Specular * DecalTransparency + MetallicSpecularRoughness.y;
+		HitPayload.Roughness = HitPayload.Roughness * DecalTransparency + MetallicSpecularRoughness.z;
 	}
 
 	if (DecalFlags & DECAL_WRITE_AMBIENT_OCCLUSION_FLAG)
@@ -145,7 +145,7 @@ void ApplyDecal(inout FPathTracingPayload HitPayload, FDecalShaderParams Params)
 
 	if (DecalFlags & DECAL_WRITE_EMISSIVE_FLAG)
 	{
-		HitPayload.Radiance += Params.GetEmissive() * DecalAlpha;
+		HitPayload.Radiance += Params.GetEmissive();
 	}
 
 	// TODO: Payload.Opacity?
@@ -367,7 +367,7 @@ FPathTracingPayload TraceTransparentRay(inout FPathState PathState, bool bIsCame
 		// Unpack the payload
 		FPathTracingPayload HitPayload = UnpackPathTracingPayloadMaterial(PackedPayload, PathState.Ray, bIsCameraRay);
 
-		float SumDecalOpacity = 0;
+		float DecalThroughput = 1.0f;
 		if (PathState.PathRoughness < MeshDecalRoughnessCutoff && HitPayload.IsDecalReceiver())
 		{
 			FRayDesc DecalRay = (FRayDesc)0;
@@ -403,7 +403,7 @@ FPathTracingPayload TraceTransparentRay(inout FPathState PathState, bool bIsCame
 			{
 				ApplyDecal(HitPayload, DecalPayload.Params);
 
-				SumDecalOpacity += DecalPayload.Params.GetOpacity();
+				DecalThroughput *= DecalPayload.Params.GetTransparency();
 			}
 		}
 
@@ -424,11 +424,11 @@ FPathTracingPayload TraceTransparentRay(inout FPathState PathState, bool bIsCame
 
 				ApplyDecal(HitPayload, Params);
 
-				SumDecalOpacity += Params.GetOpacity();
+				DecalThroughput *= Params.GetTransparency();
 			}
 		}
 #endif
-		if (SumDecalOpacity > 0)
+		if (DecalThroughput < 1.0f)
 		{
 			// if we ran any decal logic, update the payload
 			if (ApproximateCaustics)
diff --git a/Engine/Shaders/Private/RayTracing/RayTracingCommon.ush b/Engine/Shaders/Private/RayTracing/RayTracingCommon.ush
--- a/Engine/Shaders/Private/RayTracing/RayTracingCommon.ush
+++ b/Engine/Shaders/Private/RayTracing/RayTracingCommon.ush
@@ -969,11 +969,11 @@ struct FDecalShaderParams
 	}
 	float GetAO() { return asfloat(Data[12]); }
 
-	void SetOpacity(float In)
+	void SetTransparency(float In)
 	{
 		Data[13] = asuint(In);
 	}
-	float GetOpacity() { return asfloat(Data[13]); }
+	float GetTransparency() { return asfloat(Data[13]); }
 
 	void SetFlags(uint In) { Data[14] = In; }
 	uint GetFlags() { return Data[14]; }
diff --git a/Engine/Shaders/Private/RayTracing/RayTracingDecalMaterialShader.usf b/Engine/Shaders/Private/RayTracing/RayTracingDecalMaterialShader.usf
--- a/Engine/Shaders/Private/RayTracing/RayTracingDecalMaterialShader.usf
+++ b/Engine/Shaders/Private/RayTracing/RayTracingDecalMaterialShader.usf
@@ -29,6 +29,7 @@ RAY_TRACING_ENTRY_CALLABLE(RayTracingDecalMaterialShader, FDecalShaderParams, Pa
 	if (any(abs(DecalPos.xyz) > 1))
 	{
 		// clip content outside the decal
+		Params.SetTransparency(1.0f);
 		return;
 	}
 
@@ -64,31 +65,32 @@ RAY_TRACING_ENTRY_CALLABLE(RayTracingDecalMaterialShader, FDecalShaderParams, Pa
 	CalcPixelMaterialInputs(MaterialParameters, PixelMaterialInputs);
 
 	const float DecalFading = saturate(4 - 4 * abs(SwizzlePos.z * 2 - 1)) * DecalParams.x;
-	Params.SetOpacity(GetMaterialOpacity(PixelMaterialInputs) * DecalFading);
+	const float Opacity = GetMaterialOpacity(PixelMaterialInputs) * DecalFading;
+	Params.SetTransparency(1.0f - Opacity);
 
 	if (DecalWriteFlags & DECAL_WRITE_BASE_COLOR_FLAG)
 	{
-		Params.SetBaseColor(GetMaterialBaseColor(PixelMaterialInputs));
+		Params.SetBaseColor(GetMaterialBaseColor(PixelMaterialInputs) * Opacity);
 	}
 
 	if (DecalWriteFlags & DECAL_WRITE_NORMAL_FLAG)
 	{
-		Params.SetWorldNormal(MaterialParameters.WorldNormal);
+		Params.SetWorldNormal(MaterialParameters.WorldNormal * Opacity);
 	}
 
 	if (DecalWriteFlags & DECAL_WRITE_ROUGHNESS_SPECULAR_METALLIC_FLAG)
 	{
-		Params.SetMetallicSpecularRoughness(float3(GetMaterialMetallic(PixelMaterialInputs), GetMaterialSpecular(PixelMaterialInputs), GetMaterialRoughness(PixelMaterialInputs)));
+		Params.SetMetallicSpecularRoughness(float3(GetMaterialMetallic(PixelMaterialInputs), GetMaterialSpecular(PixelMaterialInputs), GetMaterialRoughness(PixelMaterialInputs)) * Opacity);
 	}
 
 	if (DecalWriteFlags & DECAL_WRITE_AMBIENT_OCCLUSION_FLAG)
 	{
-		Params.SetAO(GetMaterialAmbientOcclusion(PixelMaterialInputs));
+		Params.SetAO(GetMaterialAmbientOcclusion(PixelMaterialInputs) * Opacity);
 	}
 
 	if (DecalWriteFlags & DECAL_WRITE_EMISSIVE_FLAG)
 	{
-		Params.SetEmissive(GetMaterialEmissive(PixelMaterialInputs));
+		Params.SetEmissive(GetMaterialEmissive(PixelMaterialInputs) * Opacity);
 	}
 
 	Params.SetFlags(DecalWriteFlags);
@@ -144,32 +146,33 @@ RAY_TRACING_ENTRY_CLOSEST_HIT(RayTracingDecalMaterialCHS,
 		CalcMaterialParametersEx(MaterialParameters, PixelMaterialInputs, SvPosition, ScreenPosition, bIsFrontFace, TranslatedWorldPosition, TranslatedWorldPosition);
 		// #endif
 	}
-	
-	Payload.Params.SetOpacity(GetMaterialOpacity(PixelMaterialInputs));
+
+	const float Opacity = GetMaterialOpacity(PixelMaterialInputs);
+	Payload.Params.SetTransparency(1.0f - Opacity);
 
 	if (DECAL_PAYLOAD_FLAGS & DECAL_WRITE_BASE_COLOR_FLAG)
 	{
-		Payload.Params.SetBaseColor(GetMaterialBaseColor(PixelMaterialInputs));
+		Payload.Params.SetBaseColor(GetMaterialBaseColor(PixelMaterialInputs) * Opacity);
 	}
 
 	if (DECAL_PAYLOAD_FLAGS & DECAL_WRITE_NORMAL_FLAG)
 	{
-		Payload.Params.SetWorldNormal(MaterialParameters.WorldNormal);
+		Payload.Params.SetWorldNormal(MaterialParameters.WorldNormal * Opacity);
 	}
 
 	if (DECAL_PAYLOAD_FLAGS & DECAL_WRITE_ROUGHNESS_SPECULAR_METALLIC_FLAG)
 	{
-		Payload.Params.SetMetallicSpecularRoughness(float3(GetMaterialMetallic(PixelMaterialInputs), GetMaterialSpecular(PixelMaterialInputs), GetMaterialRoughness(PixelMaterialInputs)));
+		Payload.Params.SetMetallicSpecularRoughness(float3(GetMaterialMetallic(PixelMaterialInputs), GetMaterialSpecular(PixelMaterialInputs), GetMaterialRoughness(PixelMaterialInputs)) * Opacity);
 	}
 
 	if (DECAL_PAYLOAD_FLAGS & DECAL_WRITE_AMBIENT_OCCLUSION_FLAG)
 	{
-		Payload.Params.SetAO(GetMaterialAmbientOcclusion(PixelMaterialInputs));
+		Payload.Params.SetAO(GetMaterialAmbientOcclusion(PixelMaterialInputs) * Opacity);
 	}
 
 	if (DECAL_PAYLOAD_FLAGS & DECAL_WRITE_EMISSIVE_FLAG)
 	{
-		Payload.Params.SetEmissive(GetMaterialEmissive(PixelMaterialInputs));
+		Payload.Params.SetEmissive(GetMaterialEmissive(PixelMaterialInputs) * Opacity);
 	}
 
 	Payload.Params.SetFlags(DECAL_PAYLOAD_FLAGS);
diff --git a/Engine/Shaders/Private/RayTracing/RayTracingDefaultDecalHitShader.usf b/Engine/Shaders/Private/RayTracing/RayTracingDefaultDecalHitShader.usf
--- a/Engine/Shaders/Private/RayTracing/RayTracingDefaultDecalHitShader.usf
+++ b/Engine/Shaders/Private/RayTracing/RayTracingDefaultDecalHitShader.usf
@@ -6,7 +6,7 @@ RAY_TRACING_ENTRY_CLOSEST_HIT(RayTracingDefaultOpaqueDecalMaterialCHS,
 {
 	Payload.HitT = RayTCurrent();
 
-	Payload.Params.SetOpacity(1.0f);
+	Payload.Params.SetTransparency(0.0f);
 	Payload.Params.SetBaseColor(float3(0, 0, 0));
 	Payload.Params.SetFlags(DECAL_WRITE_BASE_COLOR_FLAG);
 }
