RayTracing: Cleanup unused code related to path tracing materials inside RTGI. This has been ifdef'd out for a while now and can no longer be supported

diff --git a/Engine/Shaders/Private/RayTracing/RayTracingCreateGatherPointsRGS.usf b/Engine/Shaders/Private/RayTracing/RayTracingCreateGatherPointsRGS.usf
--- a/Engine/Shaders/Private/RayTracing/RayTracingCreateGatherPointsRGS.usf
+++ b/Engine/Shaders/Private/RayTracing/RayTracingCreateGatherPointsRGS.usf
@@ -21,7 +21,6 @@
 #include "RayTracingGatherPoints.ush"
 
 #include "../PathTracing/Light/PathTracingLightSampling.ush"
-#include "../PathTracing/Material/PathTracingMaterialSampling.ush"
 #include "../PathTracing/Utilities/PathTracingRandomSequence.ush"
 
 RaytracingAccelerationStructure TLAS;
diff --git a/Engine/Shaders/Private/RayTracing/RayTracingFinalGatherRGS.usf b/Engine/Shaders/Private/RayTracing/RayTracingFinalGatherRGS.usf
--- a/Engine/Shaders/Private/RayTracing/RayTracingFinalGatherRGS.usf
+++ b/Engine/Shaders/Private/RayTracing/RayTracingFinalGatherRGS.usf
@@ -12,9 +12,7 @@
 #include "RayTracingGatherPoints.ush"
 
 #include "../PathTracing/Light/PathTracingLightSampling.ush"
-#include "../PathTracing/Material/PathTracingMaterialSampling.ush"
 
-#define USE_PATHTRACING_MATERIALS 0
 #define USE_REJECTION_CRITERIA 0
 #define USE_REPROJECTION 1
 #define USE_NEIGHBOR_SAMPLES 1
diff --git a/Engine/Shaders/Private/RayTracing/RayTracingGlobalIlluminationRGS.usf b/Engine/Shaders/Private/RayTracing/RayTracingGlobalIlluminationRGS.usf
--- a/Engine/Shaders/Private/RayTracing/RayTracingGlobalIlluminationRGS.usf
+++ b/Engine/Shaders/Private/RayTracing/RayTracingGlobalIlluminationRGS.usf
@@ -11,12 +11,9 @@
 #include "RayTracingHitGroupCommon.ush"
 
 #include "../PathTracing/Light/PathTracingLightSampling.ush"
-#include "../PathTracing/Material/PathTracingMaterialSampling.ush"
 #include "../PathTracing/Utilities/PathTracingRandomSequence.ush"
 
 
-#define USE_PATHTRACING_MATERIALS 0
-
 RaytracingAccelerationStructure TLAS;
 RWTexture2D<float4> RWGlobalIlluminationUAV;
 RWTexture2D<float2> RWGlobalIlluminationRayDistanceUAV;
@@ -260,11 +257,7 @@ RAY_TRACING_ENTRY_RAYGEN(GlobalIlluminationRGS)
 						// Evaluate material
 						float3 MaterialThroughput;
 						float MaterialEvalPdf = 0.0;
-#if USE_PATHTRACING_MATERIALS
-						float3 MaterialWeight = 0;
-						EvalMaterial(Ray.Direction, LightRay.Direction, Payload, MaterialWeight, MaterialEvalPdf);
-						MaterialThroughput = MaterialWeight * MaterialEvalPdf;
-#else
+
 						half3 N = Payload.WorldNormal;
 						half3 V = -Ray.Direction;
 						half3 L = LightRay.Direction;
@@ -277,7 +270,6 @@ RAY_TRACING_ENTRY_RAYGEN(GlobalIlluminationRGS)
 							MaterialThroughput += LightingSample.Transmission;
 						}
 						MaterialEvalPdf = 1.0;
-#endif
 
 						// Record the contribution
 						float3 ExitantRadianceSample = LightRadianceOverPdf * MaterialThroughput * RayThroughput * SplitFactor / RayPdf;
@@ -298,14 +290,6 @@ RAY_TRACING_ENTRY_RAYGEN(GlobalIlluminationRGS)
 			{
 				float3 Direction;
 				float3 Throughput = 1.0;
-#if USE_PATHTRACING_MATERIALS
-				RandSample4 = RandomSequence_GenerateSample4D(RandSequence);
-
-				float PositionBiasSign;
-				float3 Weight = 0;
-				SampleMaterial(Ray.Direction, Payload, RandSample, Direction, Weight, MaterialPdf, PositionBiasSign);
-				Throughput = Weight * MaterialPdf;
-#else
 
 				float3 RandSample3 = RandomSequence_GenerateSample3D(RandSequence);
 
@@ -319,7 +303,7 @@ RAY_TRACING_ENTRY_RAYGEN(GlobalIlluminationRGS)
 				FShadowTerms ShadowTerms = { 0.0, 0.0, 0.0, InitHairTransmittanceData() };
 				FDirectLighting LightingSample = EvaluateBxDF(GBufferData, N, V, L, NoL, ShadowTerms);
 				Throughput = LightingSample.Diffuse;
-#endif
+
 				if (MaterialPdf <= 0.0)
 				{
 					break;
