Path Tracer: Restore energy conservation behavior for cloth lobe

The behavior was changed for Strata, but this accidentally modified the path tracer behavior. Since Strata always supplies F0=1.0, it is safe to restore the logic to what it was.

diff --git a/Engine/Shaders/Private/ShadingEnergyConservationTemplate.ush b/Engine/Shaders/Private/ShadingEnergyConservationTemplate.ush
--- a/Engine/Shaders/Private/ShadingEnergyConservationTemplate.ush
+++ b/Engine/Shaders/Private/ShadingEnergyConservationTemplate.ush
@@ -74,13 +74,12 @@ BXDF_ENERGY_SUFFIX(FBxDFEnergyTerms) BXDF_ENERGY_SUFFIX(ComputeGGXSpecEnergyTerm
 	return BXDF_ENERGY_SUFFIX(ComputeGGXSpecEnergyTerms)(Roughness, NoV, F0, F90);
 }
 
-BXDF_ENERGY_SUFFIX(FBxDFEnergyTerms) BXDF_ENERGY_SUFFIX(ComputeClothEnergyTerms)(float Roughness, float NoV, float3 F0)
+BXDF_ENERGY_SUFFIX(FBxDFEnergyTerms) BXDF_ENERGY_SUFFIX(ComputeClothEnergyTerms)(float Roughness, float NoV, float3 F0, float3 F90)
 {
 	BXDF_ENERGY_SUFFIX(FBxDFEnergyTerms) Out;
 #if USE_ENERGY_CONSERVATION > 0
 	{
-		Out.W = 1.f;
-		Out.E = BXDF_ENERGY_SUFFIX(GetF0F90)(F0) * ClothEnergyLookup(Roughness, NoV).x;
+		Out = BXDF_ENERGY_SUFFIX(ComputeFresnelEnergyTerms)(ClothEnergyLookup(Roughness, NoV), F0, F90);
 	}
 #else
 	{
@@ -91,6 +90,12 @@ BXDF_ENERGY_SUFFIX(FBxDFEnergyTerms) BXDF_ENERGY_SUFFIX(ComputeClothEnergyTerms)
 	return Out;
 }
 
+BXDF_ENERGY_SUFFIX(FBxDFEnergyTerms) BXDF_ENERGY_SUFFIX(ComputeClothEnergyTerms)(float Roughness, float NoV, float3 F0)
+{
+	const float F90 = saturate(50.0 * F0.g); // See F_Schlick implementation
+	return BXDF_ENERGY_SUFFIX(ComputeClothEnergyTerms)(Roughness, NoV, F0, F90);
+}
+
 BXDF_ENERGY_SUFFIX(FBxDFEnergyTerms) BXDF_ENERGY_SUFFIX(ComputeDiffuseEnergyTerms)(float Roughness, float NoV)
 {
 	BXDF_ENERGY_SUFFIX(FBxDFEnergyTerms) Out;
