Fix for height fog on volumetric cloud not working when start distance is >=1.

diff --git a/Engine/Shaders/Private/VolumetricCloud.usf b/Engine/Shaders/Private/VolumetricCloud.usf
--- a/Engine/Shaders/Private/VolumetricCloud.usf
+++ b/Engine/Shaders/Private/VolumetricCloud.usf
@@ -1445,8 +1445,8 @@ void MainCommon(in FMaterialPixelParameters MaterialParameters, in float4 SvPosi
 	float ClosestMeanTransmittance = dot(ClosestTransmittanceToView.MATSUF, 1.0f / 3.0f);
 #endif // CLOUD_MIN_AND_MAX
 
-	float3 FogSampleTranslatedWorldPosition = RayTranslatedWorldOriginKm + tAP * Raydir;
-	float4 ClipPos = mul(float4(FogSampleTranslatedWorldPosition, 1.0f), PrimaryView.TranslatedWorldToClip);
+	float3 FogSampleWorldPositionRelativeToCameraCm = tAP * Raydir;
+	float4 ClipPos = mul(float4(RayTranslatedWorldOriginKm * KILOMETER_TO_CENTIMETER + FogSampleWorldPositionRelativeToCameraCm, 1.0f), PrimaryView.TranslatedWorldToClip);
 
 	if (RenderVolumetricCloudParameters.EnableAerialPerspectiveSampling && tAPWeightsSum > 0.0f)
 	{
@@ -1455,9 +1455,9 @@ void MainCommon(in FMaterialPixelParameters MaterialParameters, in float4 SvPosi
 		if (!ResolvedView.RealTimeReflectionCapture &&	// We do not generate a 360 CameraAerialPerspectiveVolume with separated Mie and Ray contribution as of today.
 			(RenderVolumetricCloudParameters.AerialPerspectiveRayOnlyStartDistanceKm + RenderVolumetricCloudParameters.AerialPerspectiveMieOnlyStartDistanceKm) > 0.0f)
 		{
-			float4 MieAP = GETAERIALPERSPECTIVE(View.CameraAerialPerspectiveVolumeMieOnly, View.CameraAerialPerspectiveVolumeMieOnlySampler, FogSampleTranslatedWorldPosition, 
+			float4 MieAP = GETAERIALPERSPECTIVE(View.CameraAerialPerspectiveVolumeMieOnly, View.CameraAerialPerspectiveVolumeMieOnlySampler, FogSampleWorldPositionRelativeToCameraCm,
 				max(ResolvedView.SkyAtmosphereAerialPerspectiveStartDepthKm, RenderVolumetricCloudParameters.AerialPerspectiveMieOnlyStartDistanceKm), RenderVolumetricCloudParameters.AerialPerspectiveMieOnlyFadeDistanceKmInv);
-			float4 RayAP = GETAERIALPERSPECTIVE(View.CameraAerialPerspectiveVolumeRayOnly, View.CameraAerialPerspectiveVolumeRayOnlySampler, FogSampleTranslatedWorldPosition, 
+			float4 RayAP = GETAERIALPERSPECTIVE(View.CameraAerialPerspectiveVolumeRayOnly, View.CameraAerialPerspectiveVolumeRayOnlySampler, FogSampleWorldPositionRelativeToCameraCm,
 				max(ResolvedView.SkyAtmosphereAerialPerspectiveStartDepthKm, RenderVolumetricCloudParameters.AerialPerspectiveRayOnlyStartDistanceKm), RenderVolumetricCloudParameters.AerialPerspectiveRayOnlyFadeDistanceKmInv);
 
 			AerialPerspective = float4(MieAP.rgb + RayAP.rgb, MieAP.a * RayAP.a);
@@ -1467,7 +1467,7 @@ void MainCommon(in FMaterialPixelParameters MaterialParameters, in float4 SvPosi
 			const float NearFadeOutRangeInvDepthKm = 1.0 / 0.00001f; // 1 centimeter fade region
 			// Apply AP only once according to the mean position within the participating media weighted by transmittance/visibility.
 			// This allows to apply AP only once per pixel instead of per sample.
-			AerialPerspective = GETAERIALPERSPECTIVE(View.CameraAerialPerspectiveVolume, View.CameraAerialPerspectiveVolumeSampler, FogSampleTranslatedWorldPosition, ResolvedView.SkyAtmosphereAerialPerspectiveStartDepthKm, NearFadeOutRangeInvDepthKm);
+			AerialPerspective = GETAERIALPERSPECTIVE(View.CameraAerialPerspectiveVolume, View.CameraAerialPerspectiveVolumeSampler, FogSampleWorldPositionRelativeToCameraCm, ResolvedView.SkyAtmosphereAerialPerspectiveStartDepthKm, NearFadeOutRangeInvDepthKm);
 		}
 
 		// Apply aerial perspective OVER the cloud, assiming coverage is from the full transmittance.
@@ -1482,13 +1482,13 @@ void MainCommon(in FMaterialPixelParameters MaterialParameters, in float4 SvPosi
 	//
 	if (RenderVolumetricCloudParameters.EnableHeightFog && tAPWeightsSum > 0.0f)
 	{
-		float4 HeightFogInscatteringAndTransmittance = CalculateHeightFog(FogSampleTranslatedWorldPosition - ResolvedView.TranslatedWorldCameraOrigin);
+		float4 HeightFogInscatteringAndTransmittance = CalculateHeightFog(FogSampleWorldPositionRelativeToCameraCm);
 		
 		if (FogStruct.ApplyVolumetricFog > 0)
 		{
 			float3 VolumeUV = ComputeVolumeUVFromNDC(ClipPos);
 			const uint EyeIndex = 0;
-			HeightFogInscatteringAndTransmittance = CombineVolumetricFog(HeightFogInscatteringAndTransmittance, VolumeUV, EyeIndex, ConvertFromDeviceZ(ClipPos.z));
+			HeightFogInscatteringAndTransmittance = CombineVolumetricFog(HeightFogInscatteringAndTransmittance, VolumeUV, EyeIndex, length(FogSampleWorldPositionRelativeToCameraCm));
 		}
 
 		// Apply aerial perspective OVER the cloud, assiming coverage is from the full transmittance.
@@ -1501,8 +1501,8 @@ void MainCommon(in FMaterialPixelParameters MaterialParameters, in float4 SvPosi
 	// Apply atmospehre and and fog on closest
 #if CLOUD_MIN_AND_MAX_DEPTH
 	{
-		float3 FogSampleTranslatedWorldPosition = RayTranslatedWorldOriginKm + ClosestTAP * Raydir;
-		float4 ClipPos = mul(float4(FogSampleTranslatedWorldPosition, 1.0f), PrimaryView.TranslatedWorldToClip);
+		float3 FogSampleWorldPositionRelativeToCameraCm = ClosestTAP * Raydir;
+		float4 ClipPos = mul(float4(RayTranslatedWorldOriginKm * KILOMETER_TO_CENTIMETER + FogSampleWorldPositionRelativeToCameraCm, 1.0f), PrimaryView.TranslatedWorldToClip);
 
 		if (RenderVolumetricCloudParameters.EnableAerialPerspectiveSampling && ClostestTAPWeightsSum > 0.0f)
 		{
@@ -1511,9 +1511,9 @@ void MainCommon(in FMaterialPixelParameters MaterialParameters, in float4 SvPosi
 			if (!ResolvedView.RealTimeReflectionCapture &&	// We do not generate a 360 CameraAerialPerspectiveVolume with separated Mie and Ray contribution as of today.
 				(RenderVolumetricCloudParameters.AerialPerspectiveRayOnlyStartDistanceKm + RenderVolumetricCloudParameters.AerialPerspectiveMieOnlyStartDistanceKm) > 0.0f)
 			{
-				float4 MieAP = GETAERIALPERSPECTIVE(View.CameraAerialPerspectiveVolumeMieOnly, View.CameraAerialPerspectiveVolumeMieOnlySampler, FogSampleTranslatedWorldPosition,
+				float4 MieAP = GETAERIALPERSPECTIVE(View.CameraAerialPerspectiveVolumeMieOnly, View.CameraAerialPerspectiveVolumeMieOnlySampler, FogSampleWorldPositionRelativeToCameraCm,
 					max(ResolvedView.SkyAtmosphereAerialPerspectiveStartDepthKm, RenderVolumetricCloudParameters.AerialPerspectiveMieOnlyStartDistanceKm), RenderVolumetricCloudParameters.AerialPerspectiveMieOnlyFadeDistanceKmInv);
-				float4 RayAP = GETAERIALPERSPECTIVE(View.CameraAerialPerspectiveVolumeRayOnly, View.CameraAerialPerspectiveVolumeRayOnlySampler, FogSampleTranslatedWorldPosition,
+				float4 RayAP = GETAERIALPERSPECTIVE(View.CameraAerialPerspectiveVolumeRayOnly, View.CameraAerialPerspectiveVolumeRayOnlySampler, FogSampleWorldPositionRelativeToCameraCm,
 					max(ResolvedView.SkyAtmosphereAerialPerspectiveStartDepthKm, RenderVolumetricCloudParameters.AerialPerspectiveRayOnlyStartDistanceKm), RenderVolumetricCloudParameters.AerialPerspectiveRayOnlyFadeDistanceKmInv);
 
 				AerialPerspective = float4(MieAP.rgb + RayAP.rgb, MieAP.a * RayAP.a);
@@ -1523,7 +1523,7 @@ void MainCommon(in FMaterialPixelParameters MaterialParameters, in float4 SvPosi
 				const float NearFadeOutRangeInvDepthKm = 1.0 / 0.00001f; // 1 centimeter fade region
 				// Apply AP only once according to the mean position within the participating media weighted by transmittance/visibility.
 				// This allows to apply AP only once per pixel instead of per sample.
-				AerialPerspective = GETAERIALPERSPECTIVE(View.CameraAerialPerspectiveVolume, View.CameraAerialPerspectiveVolumeSampler, FogSampleTranslatedWorldPosition, ResolvedView.SkyAtmosphereAerialPerspectiveStartDepthKm, NearFadeOutRangeInvDepthKm);
+				AerialPerspective = GETAERIALPERSPECTIVE(View.CameraAerialPerspectiveVolume, View.CameraAerialPerspectiveVolumeSampler, FogSampleWorldPositionRelativeToCameraCm, ResolvedView.SkyAtmosphereAerialPerspectiveStartDepthKm, NearFadeOutRangeInvDepthKm);
 			}
 
 			// Apply aerial perspective OVER the cloud, assiming coverage is from the full transmittance.
@@ -1538,13 +1538,13 @@ void MainCommon(in FMaterialPixelParameters MaterialParameters, in float4 SvPosi
 		//
 		if (RenderVolumetricCloudParameters.EnableHeightFog && ClostestTAPWeightsSum > 0.0f)
 		{
-			float4 HeightFogInscatteringAndTransmittance = CalculateHeightFog(FogSampleTranslatedWorldPosition - ResolvedView.TranslatedWorldCameraOrigin);
+			float4 HeightFogInscatteringAndTransmittance = CalculateHeightFog(FogSampleWorldPositionRelativeToCameraCm);
 
 			if (FogStruct.ApplyVolumetricFog > 0)
 			{
 				float3 VolumeUV = ComputeVolumeUVFromNDC(ClipPos);
 				const uint EyeIndex = 0;
-				HeightFogInscatteringAndTransmittance = CombineVolumetricFog(HeightFogInscatteringAndTransmittance, VolumeUV, EyeIndex, ConvertFromDeviceZ(ClipPos.z));
+				HeightFogInscatteringAndTransmittance = CombineVolumetricFog(HeightFogInscatteringAndTransmittance, VolumeUV, EyeIndex, length(FogSampleWorldPositionRelativeToCameraCm));
 			}
 
 			// Apply aerial perspective OVER the cloud, assiming coverage is from the full transmittance.
