Add option to evaluate hair material without interpolation (i.e., at control point).

This is necessary for now with hair forward raster.

diff --git a/Engine/Shaders/Private/HairStrands/HairStrandsMaterialCommon.ush b/Engine/Shaders/Private/HairStrands/HairStrandsMaterialCommon.ush
--- a/Engine/Shaders/Private/HairStrands/HairStrandsMaterialCommon.ush
+++ b/Engine/Shaders/Private/HairStrands/HairStrandsMaterialCommon.ush
@@ -86,7 +86,8 @@ FEvaluateOutput Evaluate(
 	in uint SamplePrimitiveId,
 	in uint SampleCoverage8bit,
 	in uint PrimitiveId,
-	in bool bUpdateSampleCoverage)
+	in bool bUpdateSampleCoverage,
+	in bool bInterpolationEnabled)
 {
 	// World_P / World_P0 / World_P1 are in translated world space
 	const float SceneDepth = ConvertFromDeviceZ(SampleDepth);
@@ -96,48 +97,62 @@ FEvaluateOutput Evaluate(
 
 	float3 Velocity_World_P;
 	float3 Velocity_World_P_Prev;
-	float U = 0;
 	float WorldStrandRadius = 0;
 	FVertexFactoryInterpolantsVSToPS Interpolants;
 	{
 		const uint HairPrimitiveId0 = SamplePrimitiveId;
-		const uint HairPrimitiveId1 = HairPrimitiveId0 + 1; //TODO: min(HairPrimitiveId0 + 1, MaxVertexCount-1);
 
 		FHairVertexData Vertex0;
-		FHairVertexData Vertex1;
-
 		FVertexFactoryInterpolantsVSToPS Interpolants0 = GetInterpolants(PrimitiveId, HairPrimitiveId0, Vertex0);
-		FVertexFactoryInterpolantsVSToPS Interpolants1 = GetInterpolants(PrimitiveId, HairPrimitiveId1, Vertex1);
-
 		Interpolants = Interpolants0;
 
-		// Compute U
-		// Use the projection of the current sample point onto the hair segment (which is in the middle of the hair strands)
-		const float3 A = (World_P  - Vertex0.World_P);
-		const float3 B = (Vertex1.World_P - Vertex0.World_P);
-		const float CosA = dot(A, B);
-		const float LengthB2 = dot(B, B);
-		const float SegmentU = LengthB2 > 0 ? saturate(CosA / LengthB2) : 0;
-
-		// Compute V
+		float SegmentU = 0.0f;
 		float SegmentV = 0.5f;
+		if (bInterpolationEnabled)
 		{
-			const float3 W = cross(A, B);
-			const bool bIsRight = dot(W, View.ViewForward) < 0;
-			const float WorldV	= length(A - B * SegmentU);
-			WorldStrandRadius	= lerp(Vertex0.Radius, Vertex1.Radius, SegmentU);
-			SegmentV = WorldV / WorldStrandRadius;
-			SegmentV = bIsRight ? SegmentV * 0.5f + 0.5f : (1 - SegmentV) * 0.5f;
+			const uint HairPrimitiveId1 = HairPrimitiveId0 + 1; //TODO: min(HairPrimitiveId0 + 1, MaxVertexCount-1);
+
+			FHairVertexData Vertex1;
+			FVertexFactoryInterpolantsVSToPS Interpolants1 = GetInterpolants(PrimitiveId, HairPrimitiveId1, Vertex1);
+
+			// Compute U
+			// Use the projection of the current sample point onto the hair segment (which is in the middle of the hair strands)
+			const float3 A = (World_P  - Vertex0.World_P);
+			const float3 B = (Vertex1.World_P - Vertex0.World_P);
+			{
+				const float CosA = dot(A, B);
+				const float LengthB2 = dot(B, B);
+				SegmentU = LengthB2 > 0 ? saturate(CosA / LengthB2) : 0;
+			}
+
+			// Compute V
+			{
+				const float3 W = cross(A, B);
+				const bool bIsRight = dot(W, View.ViewForward) < 0;
+				const float WorldV	= length(A - B * SegmentU);
+				WorldStrandRadius	= lerp(Vertex0.Radius, Vertex1.Radius, SegmentU);
+				SegmentV = WorldV / WorldStrandRadius;
+				SegmentV = bIsRight ? SegmentV * 0.5f + 0.5f : (1 - SegmentV) * 0.5f;
+			}
+
+			Velocity_World_P	  = lerp(Vertex0.World_P, Vertex1.World_P, SegmentU);
+			Velocity_World_P_Prev = lerp(Vertex0.World_P_Prev, Vertex1.World_P_Prev, SegmentU);
+
+			Interpolants.TangentToWorld0 = lerp(Interpolants0.TangentToWorld0, Interpolants1.TangentToWorld0, SegmentU);
+			Interpolants.TangentToWorld2 = lerp(Interpolants0.TangentToWorld2, Interpolants1.TangentToWorld2, SegmentU);
+		}
+		else
+		{
+			WorldStrandRadius				= Vertex0.Radius;
+			Interpolants.TangentToWorld0	= Interpolants0.TangentToWorld0;
+			Interpolants.TangentToWorld2	= Interpolants0.TangentToWorld2;
+			Velocity_World_P				= Vertex0.World_P;
+			Velocity_World_P_Prev			= Vertex0.World_P_Prev;
 		}
-
-		Velocity_World_P	  = lerp(Vertex0.World_P, Vertex1.World_P, SegmentU);
-		Velocity_World_P_Prev = lerp(Vertex0.World_P_Prev, Vertex1.World_P_Prev, SegmentU);
-
-		Interpolants.TangentToWorld0 = lerp(Interpolants0.TangentToWorld0, Interpolants1.TangentToWorld0, SegmentU);
-		Interpolants.TangentToWorld2 = lerp(Interpolants0.TangentToWorld2, Interpolants1.TangentToWorld2, SegmentU);
 
 		Interpolants.TangentToWorld0 = normalize(Interpolants.TangentToWorld0);
 		Interpolants.TangentToWorld2 = normalize(Interpolants.TangentToWorld2);
+
 #if VF_USE_PRIMITIVE_SCENE_DATA
 		Interpolants.PrimitiveId	 = PrimitiveId;
 #endif // VF_USE_PRIMITIVE_SCENE_DATA
@@ -145,8 +160,6 @@ FEvaluateOutput Evaluate(
 		Interpolants.HairPrimitiveUV = float2(SegmentU, SegmentV);
 	}
 	
-
-	#if 1
 	float4 EncodedVelocity = 0;
 	{
 		const float4 ScreenPos		= mul(float4(Velocity_World_P.xyz, 1), ResolvedView.TranslatedWorldToClip);
@@ -155,7 +168,6 @@ FEvaluateOutput Evaluate(
 		const float3 Velocity	= Calculate3DVelocity(ScreenPos, PrevScreenPos);
 		EncodedVelocity			= EncodeVelocityToTexture(Velocity);
 	}
-	#endif
 
 	// Sample Position < consider to be the center??
 	// VS
@@ -171,8 +183,7 @@ FEvaluateOutput Evaluate(
 		FHairRenderInfo HairRenderInfo = GetHairRenderInfo(ResolvedView.HairRenderInfo, ResolvedView.HairRenderInfoBits, bUseStableRasterization);
 		const float LocalSceneDepth = ConvertFromDeviceZ(SvPosition.z); // Linear depth in world unit
 		const float PixelRadius = HairRenderInfo.bIsOrthoView ? HairRenderInfo.RadiusAtDepth1Primary : LocalSceneDepth * HairRenderInfo.RadiusAtDepth1Primary; // Not correct but the coverage is not used (we count instead the number of sub-sample covered)
-		const float StrandRealRadius = WorldStrandRadius;
-		Coverage = StrandRealRadius / max(StrandRealRadius, PixelRadius);
+		Coverage = WorldStrandRadius / max(WorldStrandRadius, PixelRadius);
 
 		SampleCoverage8bit = min(uint(Coverage * 0xFF), 0xFFu);
 	}
diff --git a/Engine/Shaders/Private/HairStrands/HairStrandsMaterialPS.usf b/Engine/Shaders/Private/HairStrands/HairStrandsMaterialPS.usf
--- a/Engine/Shaders/Private/HairStrands/HairStrandsMaterialPS.usf
+++ b/Engine/Shaders/Private/HairStrands/HairStrandsMaterialPS.usf
@@ -33,7 +33,7 @@ void Main(
 			FEvaluateOutput Out = (FEvaluateOutput)0;
 			if (SampleDepth > 0)
 			{
-				Out = Evaluate(SamplePixelCoord, SampleDepth, SamplePrimitiveId, SampleCoverage8bit, InPrimitiveId, MaterialPassParameters.bUpdateSampleCoverage > 0);
+				Out = Evaluate(SamplePixelCoord, SampleDepth, SamplePrimitiveId, SampleCoverage8bit, InPrimitiveId, MaterialPassParameters.bUpdateSampleCoverage > 0, MaterialPassParameters.bInterpolationEnabled);
 			}
 
 			// Note: Use the syntax _ over the syntax . as on certain platform the UAV buffers are 
diff --git a/Engine/Source/Runtime/Renderer/Private/HairStrands/HairStrandsVisibility.cpp b/Engine/Source/Runtime/Renderer/Private/HairStrands/HairStrandsVisibility.cpp
--- a/Engine/Source/Runtime/Renderer/Private/HairStrands/HairStrandsVisibility.cpp
+++ b/Engine/Source/Runtime/Renderer/Private/HairStrands/HairStrandsVisibility.cpp
@@ -742,6 +742,7 @@ BEGIN_GLOBAL_SHADER_PARAMETER_STRUCT(FVisibilityMaterialPassUniformParameters, )
 	SHADER_PARAMETER(uint32, MaxSampleCount)
 	SHADER_PARAMETER(uint32, NodeGroupSize)
 	SHADER_PARAMETER(uint32, bUpdateSampleCoverage)
+	SHADER_PARAMETER(uint32, bInterpolationEnabled)
 	SHADER_PARAMETER_RDG_TEXTURE(Texture2D<uint>, NodeIndex)
 	SHADER_PARAMETER_RDG_TEXTURE(Texture2D<uint>, TotalNodeCounter)
 	SHADER_PARAMETER_RDG_BUFFER_SRV(Buffer<uint2>, NodeCoord)
@@ -824,6 +825,7 @@ static FMaterialPassOutput AddHairMaterialPass(
 	const FScene* Scene,
 	const FViewInfo* ViewInfo,
 	const bool bUpdateSampleCoverage,
+	const bool bInteroplationEnabled,
 	const FHairStrandsMacroGroupDatas& MacroGroupDatas,
 	FInstanceCullingManager& InstanceCullingManager,
 	const uint32 NodeGroupSize,
@@ -887,6 +889,7 @@ static FMaterialPassOutput AddHairMaterialPass(
 			FVisibilityMaterialPassUniformParameters* UniformParameters = GraphBuilder.AllocParameters<FVisibilityMaterialPassUniformParameters>();
 
 			UniformParameters->bUpdateSampleCoverage = bUpdateSampleCoverage ? 1 : 0;
+			UniformParameters->bInterpolationEnabled = bInteroplationEnabled ? 1 : 0;
 			UniformParameters->MaxResolution = Resolution;
 			UniformParameters->NodeGroupSize = NodeGroupSize;
 			UniformParameters->MaxSampleCount = MaxNodeCount;
@@ -4263,6 +4266,7 @@ void RenderHairStrandsVisibilityBuffer(
 					Scene,
 					&View,
 					false,
+					false,
 					MacroGroupDatas,
 					InstanceCullingManager,
 					VisibilityData.NodeGroupSize,
@@ -4370,6 +4374,7 @@ void RenderHairStrandsVisibilityBuffer(
 					Scene,
 					&View,
 					false,
+					true,
 					MacroGroupDatas,
 					InstanceCullingManager,
 					VisibilityData.NodeGroupSize,
@@ -4542,6 +4547,7 @@ void RenderHairStrandsVisibilityBuffer(
 							Scene,
 							&View,
 							bUpdateSampleCoverage,
+							true,
 							MacroGroupDatas,
 							InstanceCullingManager,
 							VisibilityData.NodeGroupSize,
@@ -4696,6 +4702,7 @@ void RenderHairStrandsVisibilityBuffer(
 							Scene,
 							&View,
 							bUpdateSampleCoverage,
+							true,
 							MacroGroupDatas,
 							InstanceCullingManager,
 							VisibilityData.NodeGroupSize,
