Clean up visualize HDR from lux derived from luminance hwen there is an illuminance meter now. This confuses our users.

diff --git a/Engine/Shaders/Private/PostProcessVisualizeHDR.usf b/Engine/Shaders/Private/PostProcessVisualizeHDR.usf
--- a/Engine/Shaders/Private/PostProcessVisualizeHDR.usf
+++ b/Engine/Shaders/Private/PostProcessVisualizeHDR.usf
@@ -245,7 +245,6 @@ void MainPS(noperspective float4 UVAndScreenPos : TEXCOORD0, float4 SvPosition :
 			}
 
 			float MaxLuminance = max(LuminanceAvg.r, max(LuminanceAvg.g, LuminanceAvg.b));
-			float MaxIlluminance = MaxLuminance * PI;
 
 			// Luminance
 			int2 TopLeft = OutputViewportCenter + int2(11, 11);
@@ -267,29 +266,9 @@ void MainPS(noperspective float4 UVAndScreenPos : TEXCOORD0, float4 SvPosition :
 			PrintCharacter(PixelPos, OutColor.xyz, NumberColor, TopLeft, _N_);
 			PrintCharacter(PixelPos, OutColor.xyz, NumberColor, TopLeft, _I_);
 			PrintCharacter(PixelPos, OutColor.xyz, NumberColor, TopLeft, _T_);
-			// Illuminance
-			TopLeft = OutputViewportCenter + int2(11, 22);
-			if (MaxIlluminance < 1000.0f)
-			{
-				PrintFloat(PixelPos, OutColor.xyz, NumberColor, TopLeft, MaxIlluminance);
-				TopLeft.x += 60;
-			}
-			else if(MaxIlluminance < 999999.0f)
-			{
-				PrintFloatNoFractionLarge(PixelPos, OutColor.xyz, NumberColor, TopLeft, MaxIlluminance);
-				TopLeft.x += 60;
-			}
-			else
-			{
-				PrintFloatNoFraction(PixelPos, OutColor.xyz, NumberColor, TopLeft, MaxIlluminance, 10);
-				TopLeft.x += 100;
-			}
-			PrintCharacter(PixelPos, OutColor.xyz, NumberColor, TopLeft, _L_);
-			PrintCharacter(PixelPos, OutColor.xyz, NumberColor, TopLeft, _U_);
-			PrintCharacter(PixelPos, OutColor.xyz, NumberColor, TopLeft, _X_);
 
 			// EV100
-			TopLeft = OutputViewportCenter + int2(11, 33);
+			TopLeft = OutputViewportCenter + int2(11, 22);
 
 			const float EyeAdaptationEV100 = log2((ComputeRgbLuminance(LuminanceAvg)/.18)/EyeAdaptation_LuminanceMax);
 			PrintMediumFloatWithSign(PixelPos, OutColor.xyz, NumberColor, TopLeft, EyeAdaptationEV100);
