Fix planar reflection preflitering sampling nans on some platforms

#rnx

diff --git a/Engine/Shaders/Private/PlanarReflectionShaders.usf b/Engine/Shaders/Private/PlanarReflectionShaders.usf
--- a/Engine/Shaders/Private/PlanarReflectionShaders.usf
+++ b/Engine/Shaders/Private/PlanarReflectionShaders.usf
@@ -56,7 +56,8 @@ void VerticalSample(float2 InUV, float y, float KernelRadiusY, float CenterDista
 	// Half texel offset for bilinear
 	float4 SampleColor = SampleSceneColorInput(SampleUV + float2(0, ColorHalfTexelOffset));
 
-	AccumulatedColor += SampleColor * SampleWeight;
+	// SampleColor might contain NaNs when sampling out of the view rect, and multiplying by 0 won't help
+	AccumulatedColor += -min(-SampleColor, 0.0) * SampleWeight;
 	AccumulatedWeight += SampleWeight;
 }
 
