Fix general RT shadow issue because of an unmatched sqaure for the encoding of the shadow. It affects several mismatch between RTshadow and pathtracing.

diff --git a/Engine/Shaders/Private/RayTracing/RayTracingOcclusionRGS.usf b/Engine/Shaders/Private/RayTracing/RayTracingOcclusionRGS.usf
--- a/Engine/Shaders/Private/RayTracing/RayTracingOcclusionRGS.usf
+++ b/Engine/Shaders/Private/RayTracing/RayTracingOcclusionRGS.usf
@@ -716,7 +716,7 @@ RAY_TRACING_ENTRY_RAYGEN(OcclusionRGS)
 	{
 		// the channel assignment is documented in ShadowRendering.cpp (look for Light Attenuation channel assignment)
 		const float HairShadow = OcclusionToShadow(HairOcclusion, RequestedSamplePerPixel);
-		const float FadedShadow = Square(HairShadow);
+		const float FadedShadow = HairShadow;
 		const float4 OutColor = EncodeLightAttenuation(FadedShadow.xxxx);
 		RWSubPixelOcclusionMaskUAV[PixelCoord] = OutColor;
 	}
@@ -756,8 +756,8 @@ RAY_TRACING_ENTRY_RAYGEN(OcclusionRGS)
 
 		// 0 is shadowed, 1 is unshadowed
 		// RETURN_COLOR not needed unless writing to SceneColor;
-		float FadedShadow = lerp(1.0f, Square(Shadow), ShadowFadeFraction);
-		float FadedSSSShadow = lerp(1.0f, Square(SSSTransmission), ShadowFadeFraction);
+		float FadedShadow = lerp(1.0f, Shadow, ShadowFadeFraction);
+		float FadedSSSShadow = lerp(1.0f, SSSTransmission, ShadowFadeFraction);
 
 		// the channel assignment is documented in ShadowRendering.cpp (look for Light Attenuation channel assignment)
 		float4 OutColor;
diff --git a/Engine/Shaders/Private/ScreenSpaceDenoise/SSDSpatialAccumulation.usf b/Engine/Shaders/Private/ScreenSpaceDenoise/SSDSpatialAccumulation.usf
--- a/Engine/Shaders/Private/ScreenSpaceDenoise/SSDSpatialAccumulation.usf
+++ b/Engine/Shaders/Private/ScreenSpaceDenoise/SSDSpatialAccumulation.usf
@@ -1282,8 +1282,8 @@ void MainCS(
 
 				// 0 is shadowed, 1 is unshadowed
 				// RETURN_COLOR not needed unless writing to SceneColor;
-				float FadedShadow = lerp(1.0f, Square(Shadow), ShadowFadeFraction);
-				float FadedSSSShadow = lerp(1.0f, Square(SSSTransmission), ShadowFadeFraction);
+				float FadedShadow = lerp(1.0f, Shadow, ShadowFadeFraction);
+				float FadedSSSShadow = lerp(1.0f, SSSTransmission, ShadowFadeFraction);
 
 				// the channel assignment is documented in ShadowRendering.cpp (look for Light Attenuation channel assignment)
 				float4 OutColor;
diff --git a/Engine/Shaders/Private/ShadowProjectionPixelShader.usf b/Engine/Shaders/Private/ShadowProjectionPixelShader.usf
--- a/Engine/Shaders/Private/ShadowProjectionPixelShader.usf
+++ b/Engine/Shaders/Private/ShadowProjectionPixelShader.usf
@@ -338,7 +338,7 @@ void Main(
 
 	// 0 is shadowed, 1 is unshadowed
 	// RETURN_COLOR not needed unless writing to SceneColor;
-	float FadedShadow = lerp(1.0f, Square(Shadow), ShadowFadeFraction * PerObjectDistanceFadeFraction);
+	float FadedShadow = lerp(1.0f, Shadow, ShadowFadeFraction * PerObjectDistanceFadeFraction);
 
 #if FORWARD_SHADING || SHADING_PATH_MOBILE 
 	float LightInfluenceMask = GetLightInfluenceMask(TranslateWorldPosition);
@@ -347,7 +347,7 @@ void Main(
 	// Write into all channels, the write mask will constrain to the correct one
 	OutColor = EncodeLightAttenuation(FadedShadow);
 #else
-	float FadedSSSShadow = lerp(1.0f, Square(SSSTransmission), ShadowFadeFraction * PerObjectDistanceFadeFraction);
+	float FadedSSSShadow = lerp(1.0f, SSSTransmission, ShadowFadeFraction * PerObjectDistanceFadeFraction);
 
 	// the channel assignment is documented in ShadowRendering.cpp (look for Light Attenuation channel assignment)
 	OutColor = EncodeLightAttenuation(half4(FadedShadow, FadedSSSShadow, FadedShadow, FadedSSSShadow));
@@ -436,7 +436,7 @@ void MainOnePassPointLightPS(
 	Shadow = saturate( (Shadow - 0.5) * ShadowSharpen + 0.5 );
 #endif
 
-	float FadedShadow = lerp(1.0f, Square(Shadow), ShadowFadeFraction);
+	float FadedShadow = lerp(1.0f, Shadow, ShadowFadeFraction);
 
 #if FORWARD_SHADING
 	float LightInfluenceMask = GetLightInfluenceMask(TranslateWorldPosition);
