Path Tracer: Refactor material eval/sample API to give independent control over diffuse vs specular scattering

This allows us to implement the Affect Specular control in the lights and also be able to decompose the image by the type of scattering (volume,diffuse,specular,etc..).

The PathState now stores the first scattering type in the path state so that we can decompose the image by scatter event type. Refraction is given slightly special treatment. A path is allowed to refract multiple times before we decide on the diffuse vs. specular vs volume status. This is done so that objects behind glass can be treated the same as objects seen directly. This can be important in the case of (for example) a character wearing glasses, or characters seen behind the windows of a vehicle.

Add controls to post process volume to select which lighting components we want to have enabled. This allows rendering of diffuse/specular/volume elements separately from the beauty image. These top level controls also impact early exits in the path tracer when we realize that a given path will no longer make any contribution to the image. This allows specular only or volume only renders to be faster than the combined render.

There were many refactors along the way:
 * Refactor blend mode handling to reduce the amount of logic needed. Take advantage of the fact that some modes such as modulate or holdout are only encountered with unlit mode, which the integrator will never pick for lighting calculations.
 * Move RadianceProbe eval/sample to GPULightmass only
 * Refactor all BSDFs to use V_World and L_World instead of incoming/outgoing directions as all BSDFs did this renaming immediately anyway.
 * Only pass 3 random numbers to MaterialSample as this is the most that any of the shading models uses (and the 4th component is reserved for path termination anyway).
 * Refactor all BSDFs to allow sharing of the common setup code relating to the probability of lobe selection
 * SolidGlass and ThinGlass now both implement the entire bsdf, which allows the logic for refraction to take into account all lobes at once.

diff --git a/Engine/Plugins/Experimental/GPULightmass/Shaders/Private/LightmapPathTracing.usf b/Engine/Plugins/Experimental/GPULightmass/Shaders/Private/LightmapPathTracing.usf
--- a/Engine/Plugins/Experimental/GPULightmass/Shaders/Private/LightmapPathTracing.usf
+++ b/Engine/Plugins/Experimental/GPULightmass/Shaders/Private/LightmapPathTracing.usf
@@ -35,6 +35,8 @@ Texture2D<float> RayGuidingCDFY;
 #include "/Engine/Private/PathTracing/Utilities/PathTracingRandomSequence.ush"
 #include "/Engine/Private/PathTracing/Light/PathTracingLightSampling.ush"
 #include "/Engine/Private/PathTracing/Material/PathTracingMaterialSampling.ush"
+#include "/Engine/Private/PathTracing/Material/PathTracingRadianceProbe.ush"
+
 
 #include "LightmapEncoding.ush"
 #include "LightmapCommon.ush"
@@ -542,15 +544,12 @@ void PathTracingKernel(
 		// Choose a random number for both Light sampling and BxDF sampling
 		float4 RandSample = RandomSequence_GenerateSample4D(RandSequence);
 
-		// Does this material require NEE? (will be false if MaterialPdf is always +inf)
-		const bool bIsNeeValid = IsNeeValidMaterial(Payload);
-
 		const bool bDoLightLoop = EnableDirectLighting != 0 || Bounce > 0;
 
 		float LightPickingCdfSum = 0;
 
 		// If we are using Light sampling and the material can use it ...
-		if (MISMode != 0 && bIsNeeValid && SceneLightCount > 0 && bDoLightLoop)
+		if (MISMode != 0 && SceneLightCount > 0 && bDoLightLoop)
 		{
 			// Choose a light and sample it
 			float3 TranslatedWorldPos = Payload.TranslatedWorldPos;
@@ -608,7 +607,16 @@ void PathTracingKernel(
 					if (any(LightSample.RadianceOverPdf > 0))
 					{
 						// Evaluate material
-						FMaterialEval MaterialEval = EvalMaterial(Ray.Direction, LightSample.Direction, Payload, Bounce == 0);
+						FMaterialEval MaterialEval;
+						if (Bounce == 0)
+						{
+							MaterialEval = RadianceProbe_EvalMaterial(LightSample.Direction, Payload);
+						}
+						else
+						{
+							MaterialEval = EvalMaterial(-Ray.Direction, LightSample.Direction, Payload, float2(1.0, 0.0));
+						}
+							
 
 						// Record the contribution
 						float3 LightContrib = PathThroughput * LightSample.RadianceOverPdf * MaterialEval.Weight * MaterialEval.Pdf;						
@@ -658,7 +666,15 @@ void PathTracingKernel(
 		}
 	
 		// Sample material
-		FMaterialSample MaterialSample = SampleMaterial(Ray.Direction, Payload, RandSample, Bounce == 0);
+		FMaterialSample MaterialSample;
+		if (Bounce == 0)
+		{
+			MaterialSample = RadianceProbe_SampleMaterial(Payload, RandSample.xyz);
+		}
+		else
+		{
+			MaterialSample = SampleMaterial(-Ray.Direction, Payload, RandSample.xyz);
+		}
 
 		if (MaterialSample.Pdf < 0 || asuint(MaterialSample.Pdf) > 0x7F800000)
 		{
@@ -692,7 +708,7 @@ void PathTracingKernel(
 		if (ContinuationProb < 1)
 		{
 			// If there is some chance we should terminate the ray, draw an extra random value
-			float RussianRouletteRand = RandSample.w; // SampleMaterial does not use this value at the moment
+			float RussianRouletteRand = RandSample.w;
 			//RussianRouletteRand = RandomSequence_GenerateSample1D(RandSequence);
 			if (RussianRouletteRand >= ContinuationProb)
 			{
diff --git a/Engine/Plugins/Experimental/GPULightmass/Source/GPULightmass/Private/PathTracingLightParameters.inl b/Engine/Plugins/Experimental/GPULightmass/Source/GPULightmass/Private/PathTracingLightParameters.inl
--- a/Engine/Plugins/Experimental/GPULightmass/Source/GPULightmass/Private/PathTracingLightParameters.inl
+++ b/Engine/Plugins/Experimental/GPULightmass/Source/GPULightmass/Private/PathTracingLightParameters.inl
@@ -113,7 +113,7 @@ void SetupPathTracingLightParameters(
 		DestLight.dPdv = (FVector3f)Light.Tangent;
 		DestLight.Color = FVector3f(Light.Color);
 		DestLight.Dimensions = FVector2f(Light.SourceRadius, Light.SourceLength);
-		DestLight.Shaping = FVector2f(Light.SpotAngles);
+		DestLight.Shaping = EncodeToF16x2(FVector2f(Light.SpotAngles));
 		DestLight.Attenuation = 1.0f / Light.AttenuationRadius;
 		DestLight.FalloffExponent = Light.FalloffExponent;
 		DestLight.IESAtlasIndex = Light.IESAtlasIndex;
@@ -162,7 +162,7 @@ void SetupPathTracingLightParameters(
 
 		DestLight.Dimensions = FVector2f(Light.SourceWidth, Light.SourceHeight);
 		DestLight.Attenuation = 1.0f / Light.AttenuationRadius;
-		DestLight.Shaping = FVector2f(FMath::Cos(FMath::DegreesToRadians(Light.BarnDoorAngle)), Light.BarnDoorLength);
+		DestLight.Shaping = EncodeToF16x2(FVector2f(FMath::Cos(FMath::DegreesToRadians(Light.BarnDoorAngle)), Light.BarnDoorLength));
 		DestLight.IESAtlasIndex = Light.IESAtlasIndex;
 
 		DestLight.Flags = PATHTRACER_FLAG_TRANSMISSION_MASK;
diff --git a/Engine/Shaders/Private/PathTracing/Light/PathTracingLightCommon.ush b/Engine/Shaders/Private/PathTracing/Light/PathTracingLightCommon.ush
--- a/Engine/Shaders/Private/PathTracing/Light/PathTracingLightCommon.ush
+++ b/Engine/Shaders/Private/PathTracing/Light/PathTracingLightCommon.ush
@@ -222,9 +222,17 @@ float2 GetRectSize(int LightId)
 	return SceneLights[LightId].Dimensions.xy;
 }
 
+float2 DecodeLightF16x2(uint Encoded)
+{
+	return float2(
+		f16tof32(Encoded),
+		f16tof32(Encoded >> 16)
+	);
+}
+
 float2 GetCosConeAngles(int LightId)
 {
-	return SceneLights[LightId].Shaping;
+	return DecodeLightF16x2(SceneLights[LightId].Shaping);
 }
 
 float3 GetColor(int LightId)
@@ -249,30 +257,22 @@ float GetAttenuation(int LightId)
 
 float GetRectLightBarnCosAngle(int LightId)
 {
-	return SceneLights[LightId].Shaping.x;
+	return DecodeLightF16x2(SceneLights[LightId].Shaping).x;
 }
 
 float GetRectLightBarnLength(int LightId)
 {
-	return SceneLights[LightId].Shaping.y;
+	return DecodeLightF16x2(SceneLights[LightId].Shaping).y;
 }
 
 float2 GetRectLightAtlasUVScale(int LightId)
 {
-	const uint Encoded = SceneLights[LightId].RectLightAtlasUVScale;
-	float2 UVScale;
-	UVScale.x = f16tof32(0xFFFF & Encoded);
-	UVScale.y = f16tof32(0xFFFF & (Encoded>>16));
-	return UVScale;
+	return DecodeLightF16x2(SceneLights[LightId].RectLightAtlasUVScale);
 }
 
 float2 GetRectLightAtlasUVOffset(int LightId)
 {
-	const uint Encoded = SceneLights[LightId].RectLightAtlasUVOffset;
-	float2 UVOffset;
-	UVOffset.x = f16tof32(0xFFFF & Encoded);
-	UVOffset.y = f16tof32(0xFFFF & (Encoded >> 16));
-	return UVOffset;
+	return DecodeLightF16x2(SceneLights[LightId].RectLightAtlasUVOffset);
 }
 
 #ifndef USE_ATTENUATION_TERM
@@ -365,6 +365,11 @@ float GetVolumetricScatteringIntensity(int LightId)
 	return SceneLights[LightId].VolumetricScatteringIntensity;
 }
 
+float GetLightSpecularScale(int LightId)
+{
+	return SceneLights[LightId].SpecularScale;
+}
+
 uint GetLightMissShaderIndex(int LightId)
 {
 	return SceneLights[LightId].MissShaderIndex;
diff --git a/Engine/Shaders/Private/PathTracing/Material/PathTracingClearCoat.ush b/Engine/Shaders/Private/PathTracing/Material/PathTracingClearCoat.ush
--- a/Engine/Shaders/Private/PathTracing/Material/PathTracingClearCoat.ush
+++ b/Engine/Shaders/Private/PathTracing/Material/PathTracingClearCoat.ush
@@ -24,235 +24,238 @@ float ClearCoatFresnel(float VoH) {
 	return CLEAR_COAT_F0 + (1.0 - CLEAR_COAT_F0) * Pow5(1 - VoH);
 }
 
-FMaterialEval ClearCoat_EvalMaterial(
-	float3 IncomingDirection,
-	float3 OutgoingDirection,
-	inout FPathTracingPayload Payload
-)
+
+struct FClearCoatData
 {
-	FMaterialEval Result = NullMaterialEval();
-	
-	const float3 V_World = -IncomingDirection;
-	const float3 L_World = OutgoingDirection;
+	float3x3 BaseBasis;
+	float3x3 CoatBasis;
+	float2 BaseAlpha;
+	float2 CoatAlpha;
 
-	float2 CoatAlpha = 0;
-	float2 SpecAlpha = 0;
-	const float3x3 CoatBasis = GetGGXBasis(Payload.GetClearCoatRoughness(), Payload.WorldNormal, CoatAlpha);
-	const float3x3 BaseBasis = GetGGXBasis(Payload.Roughness, Payload.Anisotropy, Payload.GetClearCoatBottomNormal(), Payload.WorldTangent, SpecAlpha);
+	float3 BaseV;
+	float3 CoatV;
 
-	// move vectors into right shading frame
-	const float3 CoatV = mul(CoatBasis, V_World);
-	const float3 CoatL = mul(CoatBasis, L_World);
-	const float3 CoatH = normalize(CoatV + CoatL);
+	float3 DiffColor;
+	float3 SpecColor;
+	float ClearCoat;
 
-	const float3 BaseV = mul(BaseBasis, V_World);
-	const float3 BaseL = mul(BaseBasis, L_World);
-	const float3 BaseH = normalize(BaseV + BaseL);
+	FBxDFEnergyTermsRGB Spec;
+	FBxDFEnergyTermsA   Coat;
+	float CoatE;
 
+	float2 LobeCdf;
+	float3 LobePdf;
+};
 
-	const float CoatNoL = saturate(CoatL.z);
-	const float CoatNoV = saturate(CoatV.z);
-	const float CoatNoH = saturate(CoatH.z);
-	const float CoatVoH = saturate(dot(CoatV, CoatH));
+FClearCoatData CreateClearCoatData(FPathTracingPayload Payload, float3 V_World)
+{
+	FClearCoatData Data = (FClearCoatData)0;
+
+	const float3 CoatN = Payload.WorldNormal;
+	const float3 BaseN = Payload.GetClearCoatBottomNormal();
+
+	Data.BaseBasis = GetGGXBasis(Payload.Roughness, Payload.Anisotropy, BaseN, Payload.WorldTangent, Data.BaseAlpha);
+	Data.CoatBasis = GetGGXBasis(Payload.GetClearCoatRoughness(), CoatN, Data.CoatAlpha);
+	
+
+	Data.CoatV = mul(Data.CoatBasis, V_World);
+	Data.BaseV = mul(Data.BaseBasis, V_World);
+
+	const float CoatNoV = saturate(Data.CoatV.z);
+	const float BaseNoV = saturate(Data.BaseV.z);
 
-	const float BaseNoL = saturate(BaseL.z);
-	const float BaseNoV = saturate(BaseV.z);
-	const float BaseNoH = saturate(BaseH.z);
-	const float BaseVoH = saturate(dot(BaseV, BaseH));
 
 #if EVAL_BOTTOM_LAYER
 #if GREY_INNER_LAYER
-	const float3 SpecColor = 0.0;
-	const float3 DiffColor = 0.18;
+	Data.SpecColor = 0.0;
+	Data.DiffColor = 0.18;
 #else
-	const float3 SpecColor = Payload.SpecularColor;
-	const float3 DiffColor = Payload.DiffuseColor;
+	Data.SpecColor = Payload.SpecularColor;
+	Data.DiffColor = Payload.DiffuseColor;
 #endif
 #else
-	const float3 SpecColor = 0.0;
-	const float3 DiffColor = 0.0;
+	Data.SpecColor = 0.0;
+	Data.DiffColor = 0.0;
 #endif
 
 #if EVAL_TOP_LAYER
-	const float ClearCoat = Payload.GetClearCoat();
+	Data.ClearCoat = Payload.GetClearCoat();
 #else
-	const float ClearCoat = 0.0;
+	Data.ClearCoat = 0.0;
 #endif
 
-	const FBxDFEnergyTerms Coat = ComputeGGXSpecEnergyTerms(Payload.GetClearCoatRoughness(), CoatNoV, CLEAR_COAT_F0);
-	const FBxDFEnergyTerms Spec = ComputeGGXSpecEnergyTerms(Payload.Roughness              , BaseNoV, SpecColor);
+	Data.Coat = ComputeGGXSpecEnergyTermsA(Payload.GetClearCoatRoughness(), CoatNoV, CLEAR_COAT_F0);
+	Data.Spec = ComputeGGXSpecEnergyTermsRGB(Payload.Roughness, BaseNoV, Data.SpecColor);
+
+	Data.CoatE = Data.ClearCoat * Data.Coat.E;
+	Data.LobeCdf = LobeSelectionCdf(Data.CoatE, (1.0 - Data.CoatE) * Data.Spec.E, (1.0 - Data.CoatE) * (1.0 - Data.Spec.E) * Data.DiffColor);
+	Data.LobePdf = LobeSelectionPdf(Data.LobeCdf);
+
+	return Data;
+}
+
+FMaterialEval ClearCoat_EvalMaterial(
+	float3 V_World,
+	float3 L_World,
+	FPathTracingPayload Payload,
+	float2 DiffuseSpecularScale
+)
+{
+	const FClearCoatData Data = CreateClearCoatData(Payload, V_World);
+	
+	// move vectors into right shading frame
+	const float3 CoatV = Data.CoatV;
+	const float3 CoatL = mul(Data.CoatBasis, L_World);
+	const float3 CoatH = normalize(CoatV + CoatL);
+
+	const float3 BaseV = Data.BaseV;
+	const float3 BaseL = mul(Data.BaseBasis, L_World);
+	const float3 BaseH = normalize(BaseV + BaseL);
+
+
+	const float CoatNoL = saturate(CoatL.z);
+	const float CoatNoV = saturate(CoatV.z);
+	const float CoatNoH = saturate(CoatH.z);
+	const float CoatVoH = saturate(dot(CoatV, CoatH));
+
+	const float BaseNoL = saturate(BaseL.z);
+	const float BaseNoV = saturate(BaseV.z);
+	const float BaseNoH = saturate(BaseH.z);
+	const float BaseVoH = saturate(dot(BaseV, BaseH));
 
-	const float3 CoatE = ClearCoat * Coat.E;
-	const float2 LobeCdf = LobeSelectionCdf(CoatE, (1.0 - CoatE) * Spec.E, (1.0 - CoatE) * (1.0 - Spec.E) * DiffColor);
-	const float3 LobePdf = LobeSelectionPdf(LobeCdf);
 
 	// Coat Specular lobe
-	const float2 CoatGGXResult = GGXEvalReflection(CoatL, CoatV, CoatH, CoatAlpha);
-	const float3 CoatWeight = ClearCoat * ClearCoatFresnel(CoatVoH) * CoatGGXResult.x * Coat.W;
+	const float2 CoatGGXResult = GGXEvalReflection(CoatL, CoatV, CoatH, Data.CoatAlpha);
+	const float3 CoatWeight = Data.ClearCoat * ClearCoatFresnel(CoatVoH) * CoatGGXResult.x * Data.Coat.W;
 	const float CoatPdf = CoatGGXResult.y;
 
 	// Calculate transmission through substrate
-	float3 Transmission = (1.0 - CoatE) * lerp(1.0, SimpleClearCoatTransmittance(CoatNoL, CoatNoV, Payload.Metallic, Payload.BaseColor), ClearCoat);
+	const float3 Transmission = (1.0 - Data.CoatE) * lerp(1.0, SimpleClearCoatTransmittance(CoatNoL, CoatNoV, Payload.Metallic, Payload.BaseColor), Data.ClearCoat);
 
 	// Base Specular lobe
-	const float2 SpecGGXResult = GGXEvalReflection(BaseL, BaseV, BaseH, SpecAlpha);
-	const float3 SpecWeight = Transmission * F_Schlick(SpecColor, BaseVoH) * CoatGGXResult.x * Spec.W;
+	const float2 SpecGGXResult = GGXEvalReflection(BaseL, BaseV, BaseH, Data.BaseAlpha);
+	const float3 SpecWeight = Transmission * F_Schlick(Data.SpecColor, BaseVoH) * SpecGGXResult.x * Data.Spec.W;
 	const float SpecPdf = SpecGGXResult.y;
 
 	// Base Diffuse lobe
-	const float3 Diffuse = GetPathTracingDiffuseModel(DiffColor, Payload.Roughness, BaseNoV, BaseNoL, BaseVoH, BaseNoH);
-	const float3 DiffWeight = Transmission * (1 - Spec.E) * Diffuse * ShadowTerminatorTerm(L_World, Payload.WorldNormal, Payload.WorldSmoothNormal);
+	const float3 Diffuse = GetPathTracingDiffuseModel(Data.DiffColor, Payload.Roughness, BaseNoV, BaseNoL, BaseVoH, BaseNoH);
+	const float3 DiffWeight = Transmission * (1 - Data.Spec.E) * Diffuse * ShadowTerminatorTerm(L_World, Payload.WorldNormal, Payload.WorldSmoothNormal);
 	const float DiffPdf = BaseNoL / PI;
 
 	// Combine all lobes together
-	AddLobeWithMIS(Result.Weight, Result.Pdf, CoatWeight, CoatPdf, LobePdf.x);
-	AddLobeWithMIS(Result.Weight, Result.Pdf, SpecWeight, SpecPdf, LobePdf.y);
-	AddLobeWithMIS(Result.Weight, Result.Pdf, DiffWeight, DiffPdf, LobePdf.z);
+	FMaterialEval Result = NullMaterialEval();
+	Result.AddLobeWithMIS(CoatWeight * DiffuseSpecularScale.y, CoatPdf, Data.LobePdf.x);
+	Result.AddLobeWithMIS(SpecWeight * DiffuseSpecularScale.y, SpecPdf, Data.LobePdf.y);
+	Result.AddLobeWithMIS(DiffWeight * DiffuseSpecularScale.x, DiffPdf, Data.LobePdf.z);
+
+	Result.Weight *= Payload.GetBSDFOpacity();
 	
 	return Result;
 }
 
 FMaterialSample ClearCoat_SampleMaterial(
-	float3 RayDirection,
+	float3 V_World,
 	FPathTracingPayload Payload,
-	float4 RandSample
+	float3 RandSample
 )
 {
-	float3 V_World = -RayDirection;
-
+	const FClearCoatData Data = CreateClearCoatData(Payload, V_World);
 
-	float2 CoatAlpha = 0;
-	float2 SpecAlpha = 0;
-	const float3x3 CoatBasis = GetGGXBasis(Payload.GetClearCoatRoughness(), Payload.WorldNormal, CoatAlpha);
-	const float3x3 BaseBasis = GetGGXBasis(Payload.Roughness, Payload.Anisotropy, Payload.GetClearCoatBottomNormal(), Payload.WorldTangent, SpecAlpha);
-
-	// move view vector into right shading frame
-	const float3 CoatV = mul(CoatBasis, V_World);
-	const float3 BaseV = mul(BaseBasis, V_World);
-	const float CoatNoV = saturate(CoatV.z);
-	const float BaseNoV = saturate(BaseV.z);
-
-
-#if EVAL_BOTTOM_LAYER
-#if GREY_INNER_LAYER
-	const float3 SpecColor = 0.0;
-	const float3 DiffColor = 0.18;
-#else
-	const float3 SpecColor = Payload.SpecularColor;
-	const float3 DiffColor = Payload.DiffuseColor;
-#endif
-#else
-	const float3 SpecColor = 0.0;
-	const float3 DiffColor = 0.0;
-#endif
-
-#if EVAL_TOP_LAYER
-	const float ClearCoat = Payload.GetClearCoat();
-#else
-	const float ClearCoat = 0.0;
-#endif
-
-	const FBxDFEnergyTerms Coat = ComputeGGXSpecEnergyTerms(Payload.GetClearCoatRoughness(), CoatNoV, CLEAR_COAT_F0);
-	const FBxDFEnergyTerms Spec = ComputeGGXSpecEnergyTerms(Payload.Roughness              , BaseNoV, Payload.SpecularColor);
-
-	const float3 CoatE = ClearCoat * Coat.E;
-	const float2 LobeCdf = LobeSelectionCdf(CoatE, (1.0 - CoatE) * Spec.E, (1.0 - CoatE) * (1.0 - Spec.E) * DiffColor);
-	const float3 LobePdf = LobeSelectionPdf(LobeCdf);
+	const float CoatNoV = saturate(Data.CoatV.z);
+	const float BaseNoV = saturate(Data.BaseV.z);
 
 	// choose between layers
-	float3 OutDirection = 0;
-	float OutRoughness = 0;
-	if (RandSample.x < LobeCdf.x)
+	FMaterialSample Result = NullMaterialSample();
+	const bool bSampledSpecular = RandSample.x < Data.LobeCdf.y;
+	if (bSampledSpecular)
 	{
-		// Sampled coating
-		RandSample.x /= LobeCdf.x;
+		// Sampled specular (coat or base)
+		const bool bSampledCoat = RandSample.x < Data.LobeCdf.x;
 
-		float3 V = -mul(CoatBasis, RayDirection);
-
-		// #dxr_todo: Switch to UniformSampleDiskConcentric after fixing 0/0 bug
-		float3 H = ImportanceSampleVisibleGGX_aniso(UniformSampleDisk(RandSample.xy), CoatAlpha, V).xyz;
-
-		float3 L = reflect(-V, H);
-		if (L.z <= 0)
+		if (bSampledCoat)
 		{
-			// invalid output direction, skip some work
-			return NullMaterialSample();
+			RandSample.x /= Data.LobeCdf.x;
 		}
-		// reflect and transform
-		OutDirection = mul(L, CoatBasis);
-		OutRoughness = Payload.GetClearCoatRoughness();
-	}
-	else if (RandSample.x < LobeCdf.y)
-	{
-		// Sampled base spec
-		RandSample.x -= LobeCdf.x;
-		RandSample.x /= LobeCdf.y - LobeCdf.x;
-
-		float3 V = -mul(BaseBasis, RayDirection);
-		float3 H = ImportanceSampleVisibleGGX_aniso(UniformSampleDisk(RandSample.xy), SpecAlpha, V).xyz;
-		float3 L = reflect(-V, H);
+		else
+		{
+			RandSample.x -= Data.LobeCdf.x;
+			RandSample.x /= Data.LobeCdf.y - Data.LobeCdf.x;
+		}
+		const float3 V = bSampledCoat ? Data.CoatV : Data.BaseV;
+		const float3 H = ImportanceSampleVisibleGGX_aniso(UniformSampleDisk(RandSample.xy), bSampledCoat ? Data.CoatAlpha : Data.BaseAlpha, V).xyz;
+		const float3 L = reflect(-V, H);
 		if (L.z <= 0)
 		{
-			// invalid output direction, exit early
+			// invalid output direction, skip some work
 			return NullMaterialSample();
 		}
 
-		OutDirection = mul(L, BaseBasis);
-		OutRoughness = Payload.Roughness;
+		// reflect and transform
+		Result.Direction = bSampledCoat ? mul(L, Data.CoatBasis) : mul(L, Data.BaseBasis);
+		Result.Roughness = bSampledCoat ? Payload.GetClearCoatRoughness() : Payload.Roughness;
+		Result.ScatterType = PATHTRACER_SCATTER_SPECULAR;
 	}
 	else
 	{
 		// Sampled base diffuse
-		RandSample.x -= LobeCdf.y;
-		RandSample.x /= 1.0 - LobeCdf.y;
+		RandSample.x -= Data.LobeCdf.y;
+		RandSample.x /= 1.0 - Data.LobeCdf.y;
 
-		float3 L = CosineSampleHemisphere(RandSample.xy).xyz;
-		OutDirection = mul(L, BaseBasis);
-		OutRoughness = 1.0;
+		const float3 L = CosineSampleHemisphere(RandSample.xy).xyz;
+		Result.Direction = mul(L, Data.BaseBasis);
+		Result.Roughness = 1.0;
+		Result.ScatterType = PATHTRACER_SCATTER_DIFFUSE;
 	}
 
-
 	// move light vector into right shading frame
-	const float3 L_World = OutDirection;
+	const float3 L_World = Result.Direction;
 
-	const float3 CoatL = mul(CoatBasis, L_World);
-	const float3 CoatH = normalize(CoatV + CoatL);
+	const float3 CoatL = mul(Data.CoatBasis, L_World);
+	const float3 CoatH = normalize(Data.CoatV + CoatL);
 
-	const float3 BaseL = mul(BaseBasis, L_World);
-	const float3 BaseH = normalize(BaseV + BaseL);
+	const float3 BaseL = mul(Data.BaseBasis, L_World);
+	const float3 BaseH = normalize(Data.BaseV + BaseL);
 
 	const float CoatNoL = saturate(CoatL.z);
 	const float CoatNoH = saturate(CoatH.z);
-	const float CoatVoH = saturate(dot(CoatV, CoatH));
+	const float CoatVoH = saturate(dot(Data.CoatV, CoatH));
 
 	const float BaseNoL = saturate(BaseL.z);
 	const float BaseNoH = saturate(BaseH.z);
-	const float BaseVoH = saturate(dot(BaseV, BaseH));
-
-	// Coat Specular lobe
-	const float2 CoatGGXResult = GGXEvalReflection(CoatL, CoatV, CoatH, CoatAlpha);
-	const float3 CoatWeight = ClearCoat * ClearCoatFresnel(CoatVoH) * CoatGGXResult.x * Coat.W;
-	const float CoatPdf = CoatGGXResult.y;
+	const float BaseVoH = saturate(dot(Data.BaseV, BaseH));
 
 	// Calculate transmission through substrate
-	float3 Transmission = (1.0 - CoatE) * lerp(1.0, SimpleClearCoatTransmittance(CoatNoL, CoatNoV, Payload.Metallic, Payload.BaseColor), ClearCoat);
+	const float3 Transmission = (1.0 - Data.CoatE) * lerp(1.0, SimpleClearCoatTransmittance(CoatNoL, CoatNoV, Payload.Metallic, Payload.BaseColor), Data.ClearCoat);
 
-	// Base Specular lobe
-	const float2 SpecGGXResult = GGXEvalReflection(BaseL, BaseV, BaseH, SpecAlpha);
-	const float3 SpecWeight = Transmission * F_Schlick(SpecColor, BaseVoH) * CoatGGXResult.x * Spec.W;
-	const float SpecPdf = SpecGGXResult.y;
+	// evaluate the lobe we sampled
+	if (bSampledSpecular)
+	{
+		// Coat Specular lobe
+		const float2 CoatGGXResult = GGXEvalReflection(CoatL, Data.CoatV, CoatH, Data.CoatAlpha);
+		const float3 CoatWeight = Data.ClearCoat * ClearCoatFresnel(CoatVoH) * CoatGGXResult.x * Data.Coat.W;
+		const float CoatPdf = CoatGGXResult.y;
+
+		// Base Specular lobe
+		const float2 SpecGGXResult = GGXEvalReflection(BaseL, Data.BaseV, BaseH, Data.BaseAlpha);
+		const float3 SpecWeight = Transmission * F_Schlick(Data.SpecColor, BaseVoH) * SpecGGXResult.x * Data.Spec.W;
+		const float SpecPdf = SpecGGXResult.y;
+
+		Result.AddLobeWithMIS(CoatWeight, CoatPdf, Data.LobePdf.x);
+		Result.AddLobeWithMIS(SpecWeight, SpecPdf, Data.LobePdf.y);
+	}
+	else
+	{
+		// Base Diffuse lobe
+		const float3 Diffuse = GetPathTracingDiffuseModel(Data.DiffColor, Payload.Roughness, BaseNoV, BaseNoL, BaseVoH, BaseNoH);
+		const float3 DiffWeight = Transmission * (1 - Data.Spec.E) * Diffuse * ShadowTerminatorTerm(L_World, Payload.WorldNormal, Payload.WorldSmoothNormal);
+		const float DiffPdf = BaseNoL / PI;
 
-	// Base Diffuse lobe
-	const float3 Diffuse = GetPathTracingDiffuseModel(DiffColor, Payload.Roughness, BaseNoV, BaseNoL, BaseVoH, BaseNoH);
-	const float3 DiffWeight = Transmission * (1 - Spec.E) * Diffuse * ShadowTerminatorTerm(L_World, Payload.WorldNormal, Payload.WorldSmoothNormal);
-	const float DiffPdf = BaseNoL / PI;
+		Result.AddLobeWithMIS(DiffWeight, DiffPdf, Data.LobePdf.z);
+	}
 
-	// Combine all lobes together
-	float3 OutWeight = 0.0;
-	float OutPdf = 0.0;
-	AddLobeWithMIS(OutWeight, OutPdf, CoatWeight, CoatPdf, LobePdf.x);
-	AddLobeWithMIS(OutWeight, OutPdf, SpecWeight, SpecPdf, LobePdf.y);
-	AddLobeWithMIS(OutWeight, OutPdf, DiffWeight, DiffPdf, LobePdf.z);
 
-	return CreateMaterialSample(OutDirection, OutWeight, OutPdf, 1.0, OutRoughness);
+	Result.Weight *= Payload.GetBSDFOpacity();
+	Result.PositionBiasSign = 1.0;
+
+	return Result;
 }
diff --git a/Engine/Shaders/Private/PathTracing/Material/PathTracingCloth.ush b/Engine/Shaders/Private/PathTracing/Material/PathTracingCloth.ush
--- a/Engine/Shaders/Private/PathTracing/Material/PathTracingCloth.ush
+++ b/Engine/Shaders/Private/PathTracing/Material/PathTracingCloth.ush
@@ -9,24 +9,55 @@
 #include "PathTracingFresnel.ush"
 #include "PathTracingEnergyConservation.ush"
 
-FMaterialEval Cloth_EvalMaterial(
-	float3 IncomingDirection,
-	float3 OutgoingDirection,
-	FPathTracingPayload Payload
-)
+struct FClothData
 {
-	FMaterialEval Result = NullMaterialEval();
+	float3x3 Basis;
+	float2 Alpha;
+	float3 V;
+
+	FBxDFEnergyTermsRGB Spec;
+	FBxDFEnergyTermsRGB Fuzz;
+
+	float3 DiffWeight;
+
+	float2 LobeCdf;
+	float3 LobePdf;
+};
 
-	const float3 V_World = -IncomingDirection;
-	const float3 L_World = OutgoingDirection;
-	const float3 N_World = Payload.WorldNormal;
+FClothData PrepareClothData(FPathTracingPayload Payload, float3 V_World)
+{
+	FClothData Data = (FClothData)0;
+
+	Data.Basis = GetGGXBasis(Payload.Roughness, Payload.Anisotropy, Payload.WorldNormal, Payload.WorldTangent, Data.Alpha);
+
+	Data.V = mul(Data.Basis, V_World);
+
+	const float ClothBlend = Payload.GetClothAmount();
 
-	float2 Alpha = 0;
-	const float3x3 Basis = GetGGXBasis(Payload.Roughness, Payload.Anisotropy, Payload.WorldNormal, Payload.WorldTangent, Alpha);
+	const float NoV = saturate(Data.V.z);
+	Data.Spec = ComputeGGXSpecEnergyTermsRGB(Payload.Roughness, NoV, Payload.SpecularColor);
+	Data.Fuzz = ComputeClothEnergyTermsRGB(Payload.Roughness, NoV, Payload.GetClothColor());
+
+	Data.DiffWeight = (1.0 - lerp(Data.Spec.E, Data.Fuzz.E, ClothBlend)) * Payload.DiffuseColor;
+
+	Data.LobeCdf = LobeSelectionCdf(Data.DiffWeight, ClothBlend * Data.Fuzz.E, (1.0 - ClothBlend) * Data.Spec.E);
+	Data.LobePdf = LobeSelectionPdf(Data.LobeCdf);
+
+	return Data;
+}
+
+FMaterialEval Cloth_EvalMaterial(
+	float3 V_World,
+	float3 L_World,
+	FPathTracingPayload Payload,
+	float2 DiffuseSpecularScale
+)
+{
+	const FClothData Data = PrepareClothData(Payload, V_World);
 
 	// move vectors into right shading frame
-	const float3 V = mul(Basis, V_World);
-	const float3 L = mul(Basis, L_World);
+	const float3 V = Data.V;
+	const float3 L = mul(Data.Basis, L_World);
 	const float3 H = normalize(V + L);
 
 	const float NoL = saturate(L.z);
@@ -34,129 +65,120 @@ FMaterialEval Cloth_EvalMaterial(
 	const float VoH = saturate(dot(V, H));
 	const float NoH = saturate(H.z);
 
-	const float3 FuzzColor = Payload.GetClothColor();
-	const float Cloth = Payload.GetClothAmount();
 
-	const FBxDFEnergyTerms Spec = ComputeGGXSpecEnergyTerms(Payload.Roughness, NoV, Payload.SpecularColor);
-	const FBxDFEnergyTerms Fuzz = ComputeClothEnergyTerms(Payload.Roughness, NoV, FuzzColor);
+	FMaterialEval Result = NullMaterialEval();
 
-	// Probability of picking diffuse lobe vs. specular lobe
-	const float3 DiffWeight = (1.0 - lerp(Spec.E, Fuzz.E, Cloth)) * Payload.DiffuseColor;
-	const float3 DiffLobeAlbedo = DiffWeight + Cloth * Fuzz.E;
-	const float3 SpecLobeAlbedo = (1.0 - Cloth) * Spec.E;
+	// Diffuse Lobe
+	const float ShadowTerminator = ShadowTerminatorTerm(L_World, Payload.WorldNormal, Payload.WorldSmoothNormal);
+	Result.AddLobeWithMIS(Data.DiffWeight * ShadowTerminator * DiffuseSpecularScale.x, NoL / PI, Data.LobePdf.x);
 
-	const float LobeSelectionPdf = LobeSelectionProb(DiffLobeAlbedo, SpecLobeAlbedo);
+	const float ClothBlend = Payload.GetClothAmount();
 
-	// Diffuse and Cloth Lobe (see ClothBxDF)
+	// Cloth Lobe (see ClothBxDF - it is counted as specular)
 	const float DCloth = D_InvGGX(Pow4(Payload.Roughness), NoH);
-	const float3 FCloth = F_Schlick(FuzzColor, VoH);
-	const float3 ClothWeight = NoL > 0 && NoV > 0 ? (PI * Cloth * DCloth * Vis_Cloth(NoV, NoL)) * FCloth * Fuzz.W : 0.0;
-
-	AddLobeWithMIS(Result.Weight, Result.Pdf, (DiffWeight + ClothWeight) * ShadowTerminatorTerm(L_World, Payload.WorldNormal, Payload.WorldSmoothNormal), NoL / PI, LobeSelectionPdf);
+	const float3 FCloth = F_Schlick(Payload.GetClothColor(), VoH);
+	const float3 ClothWeight = NoL > 0 && NoV > 0 ? (PI * ClothBlend * DCloth * Vis_Cloth(NoV, NoL)) * FCloth * Data.Fuzz.W : 0.0;
+	Result.AddLobeWithMIS(ClothWeight * ShadowTerminator * DiffuseSpecularScale.y, NoL / PI, Data.LobePdf.y);
 
 	// Specular lobe
-	const float2 GGXResult = GGXEvalReflection(L, V, H, Alpha);
+	const float2 GGXResult = GGXEvalReflection(L, V, H, Data.Alpha);
 	const float3 F = F_Schlick(Payload.SpecularColor, VoH);
-	const float3 SpecWeight = (1 - Cloth) * F * GGXResult.x * Spec.W;
+	const float3 SpecWeight = (1 - ClothBlend) * F * GGXResult.x * Data.Spec.W;
 	const float SpecPdf = GGXResult.y;
+	Result.AddLobeWithMIS(SpecWeight * DiffuseSpecularScale.y, SpecPdf, Data.LobePdf.z);
 
-	AddLobeWithMIS(Result.Weight, Result.Pdf, SpecWeight, SpecPdf, 1 - LobeSelectionPdf);
+	Result.Weight *= Payload.GetBSDFOpacity();
 
 	return Result;
 }
 
 FMaterialSample Cloth_SampleMaterial(
-	float3 RayDirection,
+	float3 V_World,
 	FPathTracingPayload Payload,
-	float4 RandSample
+	float3 RandSample
 )
 {
-	float3 N = Payload.WorldNormal;
-	float3 V = -RayDirection;
-
-	float2 Alpha;
-	float3x3 Basis = GetGGXBasis(Payload.Roughness, Payload.Anisotropy, Payload.WorldNormal, Payload.WorldTangent, Alpha);
-
-	V = mul(Basis, V);
-
-	const float NoV = saturate(V.z);
-
-	const float3 FuzzColor = Payload.GetClothColor();
-	const float Cloth = Payload.GetClothAmount();
-	const FBxDFEnergyTerms Spec = ComputeGGXSpecEnergyTerms(Payload.Roughness, NoV, Payload.SpecularColor);
-	const FBxDFEnergyTerms Fuzz = ComputeClothEnergyTerms(Payload.Roughness, NoV, FuzzColor);
-
-
-	// Probability of picking diffuse lobe vs. specular lobe
-	const float3 DiffWeight = (1.0 - lerp(Spec.E, Fuzz.E, Cloth)) * Payload.DiffuseColor;
-	const float3 DiffLobeAlbedo = DiffWeight + Cloth * Fuzz.E;
-	const float3 SpecLobeAlbedo = (1.0 - Cloth) * Spec.E;
-
-	const float LobeSelectionPdf = LobeSelectionProb(DiffLobeAlbedo, SpecLobeAlbedo);
+	const FClothData Data = PrepareClothData(Payload, V_World);
 
+	const float3 V = Data.V;
+	const float NoV = saturate(Data.V.z);
 
 	// Randomly choose to sample diffuse+cloth or specular lobe
 	// The cloth lobe is low enough frequency that sampling with cosine weighting is sufficient
-	float3 L = 0, H = 0;
-	float OutRoughness = 0;
-	if (RandSample.x < LobeSelectionPdf)
+
+	float3 L = 0.0, H = 0.0;
+	const bool bSampledDiffuse = RandSample.x < Data.LobeCdf.x;
+	bool bSampledCloth = false;
+	if (RandSample.x < Data.LobeCdf.y)
 	{
-		RandSample.x /= LobeSelectionPdf;
+		if (RandSample.x < Data.LobeCdf.x)
+		{
+			RandSample.x /= Data.LobeCdf.x;
+		}
+		else
+		{
+			RandSample.x -= Data.LobeCdf.x;
+			RandSample.x /= Data.LobeCdf.y - Data.LobeCdf.x;
+			bSampledCloth = true;
+		}
 
 		// Lambert
-#if 1
-		float4 Result = CosineSampleHemisphere(RandSample.xy);
-#else
-		float4 Result = CosineSampleHemisphereConcentric(RandSample.xy);
-#endif
-
-		L = Result.xyz;
-		H = normalize(L + V);
-		OutRoughness = 1.0;
+		L = CosineSampleHemisphere(RandSample.xy).xyz;
+		H = normalize(L + Data.V);
 	}
 	else
 	{
-		RandSample.x -= LobeSelectionPdf;
-		RandSample.x /= 1 - LobeSelectionPdf;
-
+		RandSample.x -= Data.LobeCdf.y;
+		RandSample.x /= 1.0 - Data.LobeCdf.y;
 
 		// #dxr_todo: Evaluate UniformSampleDiskConcentric
-		H = ImportanceSampleVisibleGGX_aniso(UniformSampleDisk(RandSample.xy), Alpha, V).xyz;
-
-		L = reflect(-V, H);
+		H = ImportanceSampleVisibleGGX_aniso(UniformSampleDisk(RandSample.xy), Data.Alpha, Data.V).xyz;
+		L = reflect(-Data.V, H);
 		if (L.z <= 0)
 		{
 			// invalid output direction, exit early
 			return NullMaterialSample();
 		}
-		OutRoughness = Payload.Roughness;
 	}
 
 	// With a valid direction in hand -- now evaluate the BxDF (taking advantage of already computed terms)
 
-	const float3 L_World = normalize(mul(L, Basis));
+	// transform to world space
+	const float3 L_World = normalize(mul(L, Data.Basis));
 
 	const float NoL = saturate(L.z);
 	const float VoH = saturate(dot(V, H));
 	const float NoH = saturate(H.z);
 
-	// Diffuse and Cloth Lobe (see ClothBxDF)
-	const float DCloth = D_InvGGX(Pow4(Payload.Roughness), NoH);
-	const float3 FCloth = F_Schlick(FuzzColor, VoH);
-	const float3 ClothWeight = NoL > 0 && NoV > 0 ? (PI * Cloth * DCloth * Vis_Cloth(NoV, NoL)) * FCloth * Fuzz.W : 0.0;
-
-	float3 OutWeight = 0;
-	float OutPdf = 0;
-	AddLobeWithMIS(OutWeight, OutPdf, (DiffWeight + ClothWeight) * ShadowTerminatorTerm(L_World, Payload.WorldNormal, Payload.WorldSmoothNormal), NoL / PI, LobeSelectionPdf);
+	const float ClothBlend = Payload.GetClothAmount();
 
-	// Specular lobe
+	FMaterialSample Result = CreateMaterialSample(L_World, 0.0, 0.0, 1.0, 1.0, PATHTRACER_SCATTER_DIFFUSE);
 
-	const float2 GGXResult = GGXEvalReflection(L, V, H, Alpha);
-	const float3 F = F_Schlick(Payload.SpecularColor, VoH);
-	const float3 SpecWeight = (1 - Cloth) * F * GGXResult.x * Spec.W;
-	const float SpecPdf = GGXResult.y;
-	AddLobeWithMIS(OutWeight, OutPdf, SpecWeight, SpecPdf, 1 - LobeSelectionPdf);
+	const float ShadowTerminator = ShadowTerminatorTerm(L_World, Payload.WorldNormal, Payload.WorldSmoothNormal);
 
-	// transform to world space
-	return CreateMaterialSample(L_World, OutWeight, OutPdf, 1.0, OutRoughness);
+	if (bSampledDiffuse)
+	{
+		Result.AddLobeWithMIS(Data.DiffWeight * ShadowTerminator, NoL / PI, Data.LobePdf.x);
+	}
+	else if (bSampledCloth)
+	{
+		const float DCloth = D_InvGGX(Pow4(Payload.Roughness), NoH);
+		const float3 FCloth = F_Schlick(Payload.GetClothColor(), VoH);
+		const float3 ClothWeight = NoL > 0 && NoV > 0 ? (PI * ClothBlend * DCloth * Vis_Cloth(NoV, NoL)) * FCloth * Data.Fuzz.W : 0.0;
+		Result.AddLobeWithMIS(ClothWeight * ShadowTerminator, NoL / PI, Data.LobePdf.y);
+		Result.ScatterType = PATHTRACER_SCATTER_SPECULAR;
+	}
+	else
+	{
+		// Specular lobe
+		const float2 GGXResult = GGXEvalReflection(L, V, H, Data.Alpha);
+		const float3 F = F_Schlick(Payload.SpecularColor, VoH);
+		const float3 SpecWeight = (1 - ClothBlend) * F * GGXResult.x * Data.Spec.W;
+		const float SpecPdf = GGXResult.y;
+		Result.AddLobeWithMIS(SpecWeight, SpecPdf, Data.LobePdf.z);
+		Result.Roughness = Payload.Roughness;
+		Result.ScatterType = PATHTRACER_SCATTER_SPECULAR;
+	}
+	Result.Weight *= Payload.GetBSDFOpacity();
+	return Result;
 }
diff --git a/Engine/Shaders/Private/PathTracing/Material/PathTracingDefaultLit.ush b/Engine/Shaders/Private/PathTracing/Material/PathTracingDefaultLit.ush
--- a/Engine/Shaders/Private/PathTracing/Material/PathTracingDefaultLit.ush
+++ b/Engine/Shaders/Private/PathTracing/Material/PathTracingDefaultLit.ush
@@ -10,24 +10,44 @@
 #include "PathTracingFresnel.ush"
 #include "PathTracingEnergyConservation.ush"
 
-FMaterialEval DefaultLit_EvalMaterial(
-	float3 IncomingDirection,
-	float3 OutgoingDirection,
-	FPathTracingPayload Payload
-)
+struct FDefaultLitData
 {
-	FMaterialEval Result = NullMaterialEval();
+	float3x3 Basis;
+	float2 Alpha;
+	float3 V;
+	FBxDFEnergyTermsRGB Spec;
+	float3 DiffuseWeight;
+
+	float LobeProb;
+};
 
-	const float3 V_World = -IncomingDirection;
-	const float3 L_World = OutgoingDirection;
-	const float3 N_World = Payload.WorldNormal;
+FDefaultLitData CreateDefaultLitData(FPathTracingPayload Payload, float3 V_World)
+{
+	FDefaultLitData Data = (FDefaultLitData)0;
 
-	float2 Alpha = 0;
-	const float3x3 Basis = GetGGXBasis(Payload.Roughness, Payload.Anisotropy, Payload.WorldNormal, Payload.WorldTangent, Alpha);
+	Data.Basis = GetGGXBasis(Payload.Roughness, Payload.Anisotropy, Payload.WorldNormal, Payload.WorldTangent, Data.Alpha);
+	Data.V = mul(Data.Basis, V_World);
+	const float NoV = saturate(Data.V.z);
+	Data.Spec = ComputeGGXSpecEnergyTermsRGB(Payload.Roughness, NoV, Payload.SpecularColor);
+	Data.DiffuseWeight = Payload.DiffuseColor * (1 - Data.Spec.E);
+	// Probability of picking diffuse lobe vs. specular lobe
+	Data.LobeProb = LobeSelectionProb(Data.DiffuseWeight, Data.Spec.E);
+
+	return Data;
+}
+
+FMaterialEval DefaultLit_EvalMaterial(
+	float3 V_World,
+	float3 L_World,
+	FPathTracingPayload Payload,
+	float2 DiffuseSpecularScale
+)
+{
+	const FDefaultLitData Data = CreateDefaultLitData(Payload, V_World);
 
 	// move vectors into right shading frame
-	const float3 V = mul(Basis, V_World);
-	const float3 L = mul(Basis, L_World);
+	const float3 V = Data.V;
+	const float3 L = mul(Data.Basis, L_World);
 	const float3 H = normalize(V + L);
 
 	const float NoV = saturate(V.z);
@@ -35,100 +55,87 @@ FMaterialEval DefaultLit_EvalMaterial(
 	const float VoH = saturate(dot(V, H));
 	const float NoH = saturate(H.z);
 
-	const FBxDFEnergyTerms Spec = ComputeGGXSpecEnergyTerms(Payload.Roughness, NoV, Payload.SpecularColor);
-
-	// Probability of picking diffuse lobe vs. specular lobe
-	const float LobeSelectionPdf = LobeSelectionProb(Payload.DiffuseColor * (1 - Spec.E), Spec.E);
+	FMaterialEval Result = NullMaterialEval();
 
 	// Diffuse Lobe
-	const float3 Diffuse = GetPathTracingDiffuseModel(Payload.DiffuseColor, Payload.Roughness, NoV, NoL, VoH, NoH);
-	AddLobeWithMIS(Result.Weight, Result.Pdf, (1 - Spec.E) * Diffuse * ShadowTerminatorTerm(L_World, Payload.WorldNormal, Payload.WorldSmoothNormal), NoL / PI, LobeSelectionPdf);
+	const float3 Diffuse = GetPathTracingDiffuseModel(Data.DiffuseWeight, Payload.Roughness, NoV, NoL, VoH, NoH);
+	Result.AddLobeWithMIS(DiffuseSpecularScale.x * Diffuse * ShadowTerminatorTerm(L_World, Payload.WorldNormal, Payload.WorldSmoothNormal), NoL / PI, Data.LobeProb);
 
 	// Specular lobe
-	const float2 GGXResult = GGXEvalReflection(L, V, H, Alpha);
+	const float2 GGXResult = GGXEvalReflection(L, V, H, Data.Alpha);
 	const float3 F = F_Schlick(Payload.SpecularColor, VoH);
-	const float3 SpecWeight = F * GGXResult.x * Spec.W;
+	const float3 SpecWeight = F * GGXResult.x * Data.Spec.W;
 	const float SpecPdf = GGXResult.y;
 
-	AddLobeWithMIS(Result.Weight, Result.Pdf, SpecWeight, SpecPdf, 1.0 - LobeSelectionPdf);
+	Result.AddLobeWithMIS(DiffuseSpecularScale.y * SpecWeight, SpecPdf, 1.0 - Data.LobeProb);
+
+	Result.Weight *= Payload.GetBSDFOpacity();
 
 	return Result;
 }
 
 FMaterialSample DefaultLit_SampleMaterial(
-	float3 RayDirection,
+	float3 V_World,
 	FPathTracingPayload Payload,
-	float4 RandSample
+	float3 RandSample
 )
 {
-	float3 N = Payload.WorldNormal;
-	float3 V = -RayDirection;
-	
-	float2 Alpha;
-	float3x3 Basis = GetGGXBasis(Payload.Roughness, Payload.Anisotropy, Payload.WorldNormal, Payload.WorldTangent, Alpha);
-
-	V = mul(Basis, V);
-
-	const float NoV = saturate(V.z);
-	const FBxDFEnergyTerms Spec = ComputeGGXSpecEnergyTerms(Payload.Roughness, NoV, Payload.SpecularColor);
+	const FDefaultLitData Data = CreateDefaultLitData(Payload, V_World);
 
-	// Probability of picking diffuse lobe vs. specular lobe
-	const float LobeSelectionPdf = LobeSelectionProb(Payload.DiffuseColor * (1 - Spec.E), Spec.E);
+	const float3 V = Data.V;
+	float3 L = 0, H = 0;
 
 	// Randomly choose to sample diffuse or specular
-	float3 L = 0, H = 0;
-	float OutRoughness = 0;
-	if (RandSample.x < LobeSelectionPdf)
+	const bool bSampledDiffuse = RandSample.x < Data.LobeProb;
+	if (bSampledDiffuse)
 	{
-		RandSample.x /= LobeSelectionPdf;
+		RandSample.x /= Data.LobeProb;
 
 		// Lambert
 		// TODO: evaluate CosineSampleHemisphereConcentric
-		float4 Result = CosineSampleHemisphere(RandSample.xy);
-
-		L = Result.xyz;
+		L = CosineSampleHemisphere(RandSample.xy).xyz;
 		H = normalize(L + V);
-		OutRoughness = 1.0;
 	}
 	else
 	{
-		RandSample.x -= LobeSelectionPdf;
-		RandSample.x /= (1.0 - LobeSelectionPdf);
+		RandSample.x -= Data.LobeProb;
+		RandSample.x /= 1.0 - Data.LobeProb;
 
 		// #dxr_todo: Evaluate UniformSampleDiskConcentric
-		H = ImportanceSampleVisibleGGX_aniso(UniformSampleDisk(RandSample.xy), Alpha, V).xyz;
-
+		H = ImportanceSampleVisibleGGX_aniso(UniformSampleDisk(RandSample.xy), Data.Alpha, Data.V).xyz;
 		L = reflect(-V, H);
 		if (L.z <= 0)
 		{
 			// invalid output direction, exit early
 			return NullMaterialSample();
 		}
-		OutRoughness = Payload.Roughness;
 	}
 
-	// With a valid direction in hand -- now evaluate the BxDF (taking advantage of already computed terms)
-
-	const float3 L_World = normalize(mul(L, Basis));
+	// transform to world space
+	const float3 L_World = normalize(mul(L, Data.Basis));
 
+	const float NoV = saturate(V.z);
 	const float NoL = saturate(L.z);
 	const float VoH = saturate(dot(V, H));
 	const float NoH = saturate(H.z);
 
-	float3 OutWeight = 0;
-	float OutPdf = 0;
-
-	// Diffuse Lobe
-	const float3 Diffuse = GetPathTracingDiffuseModel(Payload.DiffuseColor, Payload.Roughness, NoV, NoL, VoH, NoH);
-	AddLobeWithMIS(OutWeight, OutPdf, (1 - Spec.E) * Diffuse * ShadowTerminatorTerm(L_World, Payload.WorldNormal, Payload.WorldSmoothNormal), NoL / PI, LobeSelectionPdf);
-
-	// Specular lobe
-	const float2 GGXResult = GGXEvalReflection(L, V, H, Alpha);
-	const float3 F = F_Schlick(Payload.SpecularColor, VoH);
-	const float3 SpecWeight = F * GGXResult.x * Spec.W;
-	const float SpecPdf = GGXResult.y;
-	AddLobeWithMIS(OutWeight, OutPdf, SpecWeight, SpecPdf, 1.0 - LobeSelectionPdf);
-
-	// transform to world space
-	return CreateMaterialSample(L_World, OutWeight, OutPdf, 1.0, OutRoughness);
+	FMaterialSample Result = CreateMaterialSample(L_World, 0.0, 0.0, 1.0, 1.0, PATHTRACER_SCATTER_DIFFUSE);
+	if (bSampledDiffuse)
+	{
+		const float3 Diffuse = GetPathTracingDiffuseModel(Data.DiffuseWeight, Payload.Roughness, NoV, NoL, VoH, NoH);
+		Result.AddLobeWithMIS(Diffuse * ShadowTerminatorTerm(L_World, Payload.WorldNormal, Payload.WorldSmoothNormal), NoL / PI, Data.LobeProb);
+	}
+	else
+	{
+		const float2 GGXResult = GGXEvalReflection(L, V, H, Data.Alpha);
+		const float3 F = F_Schlick(Payload.SpecularColor, VoH);
+		const float3 SpecWeight = F * GGXResult.x * Data.Spec.W;
+		const float SpecPdf = GGXResult.y;
+
+		Result.AddLobeWithMIS(SpecWeight, SpecPdf, 1.0 - Data.LobeProb);
+		Result.Roughness = Payload.Roughness;
+		Result.ScatterType = PATHTRACER_SCATTER_SPECULAR;
+	}
+	Result.Weight *= Payload.GetBSDFOpacity();
+	return Result;
 }
diff --git a/Engine/Shaders/Private/PathTracing/Material/PathTracingEye.ush b/Engine/Shaders/Private/PathTracing/Material/PathTracingEye.ush
--- a/Engine/Shaders/Private/PathTracing/Material/PathTracingEye.ush
+++ b/Engine/Shaders/Private/PathTracing/Material/PathTracingEye.ush
@@ -10,35 +10,56 @@
 #include "PathTracingFresnel.ush"
 #include "PathTracingEnergyConservation.ush"
 
-FMaterialEval Eye_EvalMaterial(
-	float3 IncomingDirection,
-	float3 OutgoingDirection,
-	FPathTracingPayload Payload
-)
+struct FEyeData
 {
-	FMaterialEval Result = NullMaterialEval();
+	float3x3 Basis;
+	float2 Alpha;
+	float3 V;
+
+	FBxDFEnergyTermsRGB Spec;
+	float3 DiffWeight;
+	float LobeProb;
+};
+
+FEyeData PrepareEyeData(FPathTracingPayload Payload, float3 V_World)
+{
+	FEyeData Data = (FEyeData)0;
 
-	const float3 V_World = -IncomingDirection;
-	const float3 L_World = OutgoingDirection;
-	const float3 N_World = Payload.WorldNormal;
+	Data.Basis = GetTangentBasis(Payload.WorldNormal);
+	Data.Alpha = GetGGXAlpha(Payload.Roughness);
+	Data.V = mul(Data.Basis, V_World);
 
-	float2 Alpha = 0;
-	const float3x3 Basis = GetGGXBasis(Payload.Roughness, Payload.WorldNormal, Alpha);
+	const float NoV = saturate(Data.V.z);
+	Data.Spec = ComputeGGXSpecEnergyTermsRGB(Payload.Roughness, NoV, Payload.SpecularColor);
+
+	Data.DiffWeight = (1 - Data.Spec.E) * Payload.DiffuseColor;
+
+	// Probability of picking diffuse lobe vs. specular lobe
+	Data.LobeProb = LobeSelectionProb(Data.DiffWeight, Data.Spec.E);
+
+	return Data;
+}
+
+FMaterialEval Eye_EvalMaterial(
+	float3 V_World,
+	float3 L_World,
+	FPathTracingPayload Payload,
+	float2 DiffuseSpecularScale
+)
+{
+	const FEyeData Data = PrepareEyeData(Payload, V_World);
 
 	// move vectors into right shading frame
-	const float3 V = mul(Basis, V_World);
-	const float3 L = mul(Basis, L_World);
+	const float3 V = Data.V;
+	const float3 L = mul(Data.Basis, L_World);
 	const float3 H = normalize(V + L);
 
 	const float NoV = saturate(V.z);
 	const float NoL = saturate(L.z);
 	const float VoH = saturate(dot(V, H));
 
-	const FBxDFEnergyTerms Spec = ComputeGGXSpecEnergyTerms(Payload.Roughness, NoV, Payload.SpecularColor);
-
-	// Probability of picking diffuse lobe vs. specular lobe
-	const float LobeSelectionPdf = LobeSelectionProb(Payload.DiffuseColor * (1 - Spec.E), Spec.E);
 
+	FMaterialEval Result = NullMaterialEval();
 
 	const float IrisMask = Payload.GetEyeIrisMask();
 	const float3 IrisNormal = Payload.GetEyeIrisNormal();
@@ -48,107 +69,95 @@ FMaterialEval Eye_EvalMaterial(
 	const float Caustic = 0.8 + 0.2 * (Power + 1) * pow(saturate(dot(CausticNormal, L_World)), Power);
 	const float Iris = IrisNoL * Caustic;
 	const float Sclera = NoL;
-
 	const float EyeDiffuseTweak = 2 * lerp(Sclera, Iris, IrisMask);
 
-
 	// Diffuse Lobe
-	AddLobeWithMIS(Result.Weight, Result.Pdf, (1 - Spec.E) * Payload.DiffuseColor * EyeDiffuseTweak, 1 / (2 * PI), LobeSelectionPdf);
+	Result.AddLobeWithMIS(Data.DiffWeight * EyeDiffuseTweak * DiffuseSpecularScale.x, 1 / (2 * PI), Data.LobeProb);
 
 	// Specular lobe
-	const float2 GGXResult = GGXEvalReflection(L, V, H, Alpha);
+	const float2 GGXResult = GGXEvalReflection(L, V, H, Data.Alpha);
 	const float3 F = F_Schlick(Payload.SpecularColor, VoH);
-	const float3 SpecWeight = F * GGXResult.x * Spec.W;
+	const float3 SpecWeight = F * GGXResult.x * Data.Spec.W;
 	const float SpecPdf = GGXResult.y;
 
-	AddLobeWithMIS(Result.Weight, Result.Pdf, SpecWeight, SpecPdf, 1.0 - LobeSelectionPdf);
+	Result.AddLobeWithMIS(SpecWeight * DiffuseSpecularScale.y, SpecPdf, 1.0 - Data.LobeProb);
+
+	Result.Weight *= Payload.GetBSDFOpacity();
 
 	return Result;
 }
 
 FMaterialSample Eye_SampleMaterial(
-	float3 RayDirection,
+	float3 V_World,
 	FPathTracingPayload Payload,
-	float4 RandSample
+	float3 RandSample
 )
 {
-	float3 N = Payload.WorldNormal;
-	float3 V = -RayDirection;
-
-	float2 Alpha;
-	float3x3 Basis = GetGGXBasis(Payload.Roughness, Payload.Anisotropy, Payload.WorldNormal, Payload.WorldTangent, Alpha);
-
-	V = mul(Basis, V);
+	const FEyeData Data = PrepareEyeData(Payload, V_World);
 
+	const float3 V = Data.V;
 	const float NoV = saturate(V.z);
-	const FBxDFEnergyTerms Spec = ComputeGGXSpecEnergyTerms(Payload.Roughness, NoV, Payload.SpecularColor);
-
-	// Probability of picking diffuse lobe vs. specular lobe
-	const float LobeSelectionPdf = LobeSelectionProb(Payload.DiffuseColor * (1 - Spec.E), Spec.E);
 
 	// Randomly choose to sample diffuse or specular
 	float3 L = 0, H = 0;
-	float OutRoughness = 0;
-	if (RandSample.x < LobeSelectionPdf)
+	const bool bSampledDiffuse = RandSample.x < Data.LobeProb;
+	if (bSampledDiffuse)
 	{
-		RandSample.x /= LobeSelectionPdf;
+		RandSample.x /= Data.LobeProb;
 
-		// Lambert
-		// TODO: evaluate CosineSampleHemisphereConcentric
-		float4 Result = UniformSampleHemisphere(RandSample.xy);
-
-		L = Result.xyz;
+		L = UniformSampleHemisphere(RandSample.xy).xyz;
 		H = normalize(L + V);
-		OutRoughness = 1.0;
 	}
 	else
 	{
-		RandSample.x -= LobeSelectionPdf;
-		RandSample.x /= (1.0 - LobeSelectionPdf);
-
-		// #dxr_todo: Evaluate UniformSampleDiskConcentric
-		H = ImportanceSampleVisibleGGX_aniso(UniformSampleDisk(RandSample.xy), Alpha, V).xyz;
+		RandSample.x -= Data.LobeProb;
+		RandSample.x /= 1.0 - Data.LobeProb;
 
+		H = ImportanceSampleVisibleGGX_aniso(UniformSampleDisk(RandSample.xy), Data.Alpha, V).xyz;
 		L = reflect(-V, H);
 		if (L.z <= 0)
 		{
 			// invalid output direction, exit early
 			return NullMaterialSample();
 		}
-		OutRoughness = Payload.Roughness;
 	}
 
 	// With a valid direction in hand -- now evaluate the BxDF (taking advantage of already computed terms)
 
-	const float3 L_World = normalize(mul(L, Basis));
+	// transform to world space
+	const float3 L_World = normalize(mul(L, Data.Basis));
 
 	const float NoL = saturate(L.z);
 	const float VoH = saturate(dot(V, H));
 
-	float3 OutWeight = 0;
-	float OutPdf = 0;
-
-	const float IrisMask = Payload.GetEyeIrisMask();
-	const float3 IrisNormal = Payload.GetEyeIrisNormal();
-	const float3 CausticNormal = Payload.GetEyeCausticNormal();
-	const float IrisNoL = saturate(dot(IrisNormal, L_World));
-	const float Power = lerp(12, 1, IrisNoL);
-	const float Caustic = 0.8 + 0.2 * (Power + 1) * pow(saturate(dot(CausticNormal, L_World)), Power);
-	const float Iris = IrisNoL * Caustic;
-	const float Sclera = NoL;
-
-	const float EyeDiffuseTweak = 2 * lerp(Sclera, Iris, IrisMask);
-
-	// Diffuse Lobe
-	AddLobeWithMIS(OutWeight, OutPdf, (1 - Spec.E) * Payload.DiffuseColor * EyeDiffuseTweak, 1 / (2 * PI), LobeSelectionPdf);
-
-	// Specular lobe
-	const float2 GGXResult = GGXEvalReflection(L, V, H, Alpha);
-	const float3 F = F_Schlick(Payload.SpecularColor, VoH);
-	const float3 SpecWeight = F * GGXResult.x * Spec.W;
-	const float SpecPdf = GGXResult.y;
-	AddLobeWithMIS(OutWeight, OutPdf, SpecWeight, SpecPdf, 1.0 - LobeSelectionPdf);
-
-	// transform to world space
-	return CreateMaterialSample(L_World, OutWeight, OutPdf, 1.0, OutRoughness);
+	FMaterialSample Result = CreateMaterialSample(L_World, 0.0, 0.0, 1.0, 1.0, PATHTRACER_SCATTER_DIFFUSE);
+	if (bSampledDiffuse)
+	{
+		const float IrisMask = Payload.GetEyeIrisMask();
+		const float3 IrisNormal = Payload.GetEyeIrisNormal();
+		const float3 CausticNormal = Payload.GetEyeCausticNormal();
+		const float IrisNoL = saturate(dot(IrisNormal, L_World));
+		const float Power = lerp(12, 1, IrisNoL);
+		const float Caustic = 0.8 + 0.2 * (Power + 1) * pow(saturate(dot(CausticNormal, L_World)), Power);
+		const float Iris = IrisNoL * Caustic;
+		const float Sclera = NoL;
+
+		const float EyeDiffuseTweak = 2 * lerp(Sclera, Iris, IrisMask);
+
+		// Diffuse Lobe
+		Result.AddLobeWithMIS(Data.DiffWeight * EyeDiffuseTweak, 1 / (2 * PI), Data.LobeProb);
+	}
+	else
+	{
+		// Specular lobe
+		const float2 GGXResult = GGXEvalReflection(L, V, H, Data.Alpha);
+		const float3 F = F_Schlick(Payload.SpecularColor, VoH);
+		const float3 SpecWeight = F * GGXResult.x * Data.Spec.W;
+		const float SpecPdf = GGXResult.y;
+		Result.AddLobeWithMIS(SpecWeight, SpecPdf, 1.0 - Data.LobeProb);
+		Result.Roughness = Payload.Roughness;
+		Result.ScatterType = PATHTRACER_SCATTER_SPECULAR;
+	}
+	Result.Weight *= Payload.GetBSDFOpacity();
+	return Result;
 }
diff --git a/Engine/Shaders/Private/PathTracing/Material/PathTracingHair.ush b/Engine/Shaders/Private/PathTracing/Material/PathTracingHair.ush
--- a/Engine/Shaders/Private/PathTracing/Material/PathTracingHair.ush
+++ b/Engine/Shaders/Private/PathTracing/Material/PathTracingHair.ush
@@ -45,11 +45,11 @@ float ZinkeHairDiffuse(float3 L)
 }
 
 FMaterialSample Hair_SampleMaterial(
-	float3 RayDirection,
+	float3 V_World,
 	FPathTracingPayload Payload,
-	float4 RandSample)
+	float3 RandSample)
 {
-	float3x3 HairBasis = GetHairBasis(Payload.WorldTangent, -RayDirection);
+	float3x3 HairBasis = GetHairBasis(Payload.WorldTangent, V_World);
 
 	// sample a random position across the width of the fiber
 	float h = 2 * RandSample.x - 1;
@@ -65,22 +65,22 @@ FMaterialSample Hair_SampleMaterial(
 	float3 L_World = mul(L, HairBasis);
 
 	float Pdf = ZinkeHairDiffuse(L);
-	return CreateMaterialSample(L_World, Payload.BaseColor, Pdf, sign(dot(Payload.WorldGeoNormal, L_World)), 1.0);
+	return CreateMaterialSample(L_World, Payload.BaseColor, Pdf, sign(dot(Payload.WorldGeoNormal, L_World)), 1.0, PATHTRACER_SCATTER_DIFFUSE);
 }
 
 FMaterialEval Hair_EvalMaterial(
-	float3 IncomingDirection,
-	float3 OutgoingDirection,
-	FPathTracingPayload Payload
+	float3 V_World,
+	float3 L_World,
+	FPathTracingPayload Payload,
+	float2 DiffuseSpecularScale
 )
 {
-	float3x3 HairBasis = GetHairBasis(Payload.WorldTangent, -IncomingDirection);
+	float3x3 HairBasis = GetHairBasis(Payload.WorldTangent, V_World);
 
-	float3 L_World = OutgoingDirection;
 	float3 L = mul(HairBasis, L_World);
 
 	float Pdf = ZinkeHairDiffuse(L);
-	return CreateMaterialEval(Payload.BaseColor, Pdf);
+	return CreateMaterialEval(Payload.BaseColor * DiffuseSpecularScale.x, Pdf);
 }
 
 
@@ -179,7 +179,8 @@ float TrigInverse(float x)
 
 // precomputed terms for evaluating/sampling the hair bxdf
 // TODO: can we somehow do this stuff only once instead of repeating it in eval and sample?
-struct FHairData {
+struct FHairData
+{
 	// trig that only depends on viewing direction and offset h
 	float4 SinThetaO; // Viewing angles (adjusted by scale tilt per lobe)
 	float GammaO;
@@ -297,11 +298,10 @@ FHairData PrepareHairData(float3 V, float3 Color, float LongitudinalRoughness, f
 }
 
 FMaterialSample Hair_SampleMaterial(
-	float3 RayDirection,
+	float3 V_World,
 	FPathTracingPayload Payload,
-	float4 RandSample)
+	float3 RandSample)
 {
-	const float3 V_World = -RayDirection;
 	const float3x3 HairBasis = GetHairBasis(Payload.WorldTangent, V_World);
 
 	const float3 V = mul(HairBasis, V_World); // TODO: simplify this since HairBasis is already sort of aligned with V ?
@@ -357,35 +357,34 @@ FMaterialSample Hair_SampleMaterial(
 
 	const float3 L = float3(CosThetaI * sin(PhiI), CosThetaI * cos(PhiI), SinThetaI);
 
-	float3 Weight = 0;
-	float Pdf = 0;
+	// transform direction into world space
+	const float3 L_World = mul(L, HairBasis);
+
+	FMaterialSample Result = CreateMaterialSample(L_World, 0.0, 0.0, sign(dot(Payload.WorldGeoNormal, L_World)), Payload.Roughness, PATHTRACER_SCATTER_SPECULAR);
+
 	const float InvS = rcp(HairData.s);
 	const float Pdf0 = Mp(SinThetaI, HairData.SinThetaO.x, HairData.v0) * Np(PhiI, 0, InvS, HairData.GammaO, HairData.GammaT);
 	const float Pdf1 = Mp(SinThetaI, HairData.SinThetaO.y, HairData.v1) * Np(PhiI, 1, InvS, HairData.GammaO, HairData.GammaT);
 	const float Pdf2 = Mp(SinThetaI, HairData.SinThetaO.z, HairData.v2) * Np(PhiI, 2, InvS, HairData.GammaO, HairData.GammaT);
 	const float Pdf3 = Mp(SinThetaI, HairData.SinThetaO.w, HairData.v3) * (1 / (2 * PI)); // last lobe is isotropic
 
-	AddLobeWithMIS(Weight, Pdf, HairData.A0, Pdf0, HairData.LobePdf.x);
-	AddLobeWithMIS(Weight, Pdf, HairData.A1, Pdf1, HairData.LobePdf.y);
-	AddLobeWithMIS(Weight, Pdf, HairData.A2, Pdf2, HairData.LobePdf.z);
-	AddLobeWithMIS(Weight, Pdf, HairData.A3, Pdf3, HairData.LobePdf.w);
+	Result.AddLobeWithMIS(HairData.A0, Pdf0, HairData.LobePdf.x);
+	Result.AddLobeWithMIS(HairData.A1, Pdf1, HairData.LobePdf.y);
+	Result.AddLobeWithMIS(HairData.A2, Pdf2, HairData.LobePdf.z);
+	Result.AddLobeWithMIS(HairData.A3, Pdf3, HairData.LobePdf.w);
 
-	// finally transform into world space
-	const float3 L_World = mul(L, HairBasis);
-
-	return CreateMaterialSample(L_World, Weight, Pdf, sign(dot(Payload.WorldGeoNormal, L_World)), Payload.Roughness);
+	return Result;
 }
 
 FMaterialEval Hair_EvalMaterial(
-	float3 IncomingDirection,
-	float3 OutgoingDirection,
-	FPathTracingPayload Payload
+	float3 V_World,
+	float3 L_World,
+	FPathTracingPayload Payload,
+	float2 DiffuseSpecularScale
 )
 {
-	const float3 V_World = -IncomingDirection;
 	const float3x3 HairBasis = GetHairBasis(Payload.WorldTangent, V_World);
 
-	const float3 L_World = OutgoingDirection;
 	const float3 V = mul(HairBasis, V_World);
 	const float3 L = mul(HairBasis, L_World);
 
@@ -406,11 +405,11 @@ FMaterialEval Hair_EvalMaterial(
 	const float Pdf2 = Mp(SinThetaI, HairData.SinThetaO.z, HairData.v2) * Np(PhiI, 2, InvS, HairData.GammaO, HairData.GammaT);
 	const float Pdf3 = Mp(SinThetaI, HairData.SinThetaO.w, HairData.v3) * (1 / (2 * PI)); // last lobe is isotropic
 
-	AddLobeWithMIS(Result.Weight, Result.Pdf, HairData.A0, Pdf0, HairData.LobePdf.x);
-	AddLobeWithMIS(Result.Weight, Result.Pdf, HairData.A1, Pdf1, HairData.LobePdf.y);
-	AddLobeWithMIS(Result.Weight, Result.Pdf, HairData.A2, Pdf2, HairData.LobePdf.z);
-	AddLobeWithMIS(Result.Weight, Result.Pdf, HairData.A3, Pdf3, HairData.LobePdf.w);
-
+	Result.AddLobeWithMIS(HairData.A0, Pdf0, HairData.LobePdf.x);
+	Result.AddLobeWithMIS(HairData.A1, Pdf1, HairData.LobePdf.y);
+	Result.AddLobeWithMIS(HairData.A2, Pdf2, HairData.LobePdf.z);
+	Result.AddLobeWithMIS(HairData.A3, Pdf3, HairData.LobePdf.w);
+	Result.Weight *= DiffuseSpecularScale.y;
 	return Result;
 }
 
diff --git a/Engine/Shaders/Private/PathTracing/Material/PathTracingLambert.ush b/Engine/Shaders/Private/PathTracing/Material/PathTracingLambert.ush
--- a/Engine/Shaders/Private/PathTracing/Material/PathTracingLambert.ush
+++ b/Engine/Shaders/Private/PathTracing/Material/PathTracingLambert.ush
@@ -7,7 +7,7 @@ PathTracingLambert.usf: Lambertian BRDF sampling functions
 
 FMaterialSample Lambert_SampleMaterial(
 	FPathTracingPayload Payload,
-	float4 RandSample)
+	float3 RandSample)
 {
 	float3 N_World = Payload.WorldNormal;
 
@@ -15,25 +15,25 @@ FMaterialSample Lambert_SampleMaterial(
 
 	float3 L_World = TangentToWorld(SampledValue.xyz, N_World);
 
-	return CreateMaterialSample(L_World, Payload.BaseColor * ShadowTerminatorTerm(L_World, Payload.WorldNormal, Payload.WorldSmoothNormal), SampledValue.w, 1.0, 1.0);
+	return CreateMaterialSample(L_World, Payload.GetBSDFOpacity() * Payload.BaseColor * ShadowTerminatorTerm(L_World, Payload.WorldNormal, Payload.WorldSmoothNormal), SampledValue.w, 1.0, 1.0, PATHTRACER_SCATTER_DIFFUSE);
 }
 
 FMaterialEval Lambert_EvalMaterial(
-	float3 OutgoingDirection,
-	FPathTracingPayload Payload
+	float3 L_World,
+	FPathTracingPayload Payload,
+	float2 DiffuseSpecularScale
 )
 {
 	float3 N_World = Payload.WorldNormal;
-	float3 L_World = OutgoingDirection;
-	float NoL = saturate(dot(N_World, OutgoingDirection));
+	float NoL = saturate(dot(N_World, L_World));
 
-	return CreateMaterialEval(Payload.BaseColor * ShadowTerminatorTerm(L_World, Payload.WorldNormal, Payload.WorldSmoothNormal), NoL / PI);
+	return CreateMaterialEval(Payload.GetBSDFOpacity() * Payload.BaseColor * ShadowTerminatorTerm(L_World, Payload.WorldNormal, Payload.WorldSmoothNormal) * DiffuseSpecularScale.x, NoL / PI);
 }
 
 
 FMaterialSample Simplified_SampleMaterial(
 	FPathTracingPayload Payload,
-	float4 RandSample)
+	float3 RandSample)
 {
 	float3 N_World = Payload.WorldNormal;
 
@@ -47,16 +47,18 @@ FMaterialSample Simplified_SampleMaterial(
 	float3 Weight = bDoTransmission ? BackColor / TransmissionProb : FrontColor / (1 - TransmissionProb);
 	float Pdf = SampledValue.w * (bDoTransmission ? TransmissionProb : 1.0 - TransmissionProb);
 
-	return CreateMaterialSample(TangentToWorld(SampledValue.xyz, N_World * (bDoTransmission ? -1.0 : 1.0)), Weight, Pdf, 1.0, 1.0);
+	Weight *= Payload.GetBSDFOpacity();
+	float Sign = bDoTransmission ? -1.0 : 1.0;
+	return CreateMaterialSample(TangentToWorld(SampledValue.xyz, N_World * Sign), Weight, Pdf, Sign, 1.0, PATHTRACER_SCATTER_DIFFUSE);
 }
 
 FMaterialEval Simplified_EvalMaterial(
-	float3 OutgoingDirection,
+	float3 L_World,
 	FPathTracingPayload Payload
 )
 {
 	float3 N_World = Payload.WorldNormal;
-	float NoL = dot(N_World, OutgoingDirection);
+	float NoL = dot(N_World, L_World);
 
 	float3 FrontColor = Payload.BaseColor;
 	float3 BackColor = Payload.SubsurfaceColor;
@@ -68,5 +70,6 @@ FMaterialEval Simplified_EvalMaterial(
 	{
 		Weight /= Pdf;
 	}
+	Weight *= Payload.GetBSDFOpacity();
 	return CreateMaterialEval(Weight, Pdf * saturate(abs(NoL)) / PI);
 }
diff --git a/Engine/Shaders/Private/PathTracing/Material/PathTracingMaterialCommon.ush b/Engine/Shaders/Private/PathTracing/Material/PathTracingMaterialCommon.ush
--- a/Engine/Shaders/Private/PathTracing/Material/PathTracingMaterialCommon.ush
+++ b/Engine/Shaders/Private/PathTracing/Material/PathTracingMaterialCommon.ush
@@ -9,14 +9,25 @@ PathTracingMaterialCommon.usf: Brdf utility functions
 #include "../PathTracingCommon.ush"
 #include "/Engine/Private/BRDF.ush"
 
+float LobeColorToWeight(float3 C)
+{
+	// TODO: What is the best heuristic to use here?
+	// average seems to perform slightly better, but needs more careful testing
+#if 1
+	return C.x + C.y + C.z;
+#else
+	return max3(C.x, C.y, C.z);
+#endif
+}
+
 // Given two lobes that will roughly contribute colors A and B to the total (estimated for example using directional albedo)
 // return the probability of choosing lobe A
 float LobeSelectionProb(float3 A, float3 B)
 {
-	const float SumA = A.x + A.y + A.z;
-	const float SumB = B.x + B.y + B.z;
-	// use MIS routine to robustly compute SumA / (SumA + SumB) even when both are 0, or overflow, etc ...
-	return MISWeightBalanced(SumA, SumB);
+	const float Aw = LobeColorToWeight(A);
+	const float Bw = LobeColorToWeight(B);
+	// use MIS routine to robustly compute Aw / (Aw + Bw) even when both are 0, or overflow, etc ...
+	return MISWeightBalanced(Aw, Bw);
 }
 
 // Given three lobes that will roughly contribute colors A, B and C to the total (estimated for example using directional albedo)
@@ -24,13 +35,13 @@ float LobeSelectionProb(float3 A, float3 B)
 // The individual probabilities of each lobe can be estimated using LobeSelectionPdf.
 float2 LobeSelectionCdf(float3 A, float3 B, float3 C)
 {
-	const float SumA = A.x + A.y + A.z;
-	const float SumB = B.x + B.y + B.z;
-	const float SumC = C.x + C.y + C.z;
+	const float Aw = LobeColorToWeight(A);
+	const float Bw = LobeColorToWeight(B);
+	const float Cw = LobeColorToWeight(C);
 
 	// use MIS routine to robustly compute the ratios even when both are 0, or overflow, etc ...
-	const float ProbAB = MISWeightBalanced(SumA, SumB); // A or B
-	const float ProbLR = MISWeightBalanced(SumA + SumB, SumC); // AB or C
+	const float ProbAB = MISWeightBalanced(Aw, Bw); // A or B
+	const float ProbLR = MISWeightBalanced(Aw + Bw, Cw); // AB or C
 
 	// Now build CDF
 	return float2(ProbAB * ProbLR, ProbLR);
@@ -39,15 +50,15 @@ float2 LobeSelectionCdf(float3 A, float3 B, float3 C)
 // Same as above for choosing among 4 colored lobes
 float3 LobeSelectionCdf(float3 A, float3 B, float3 C, float3 D)
 {
-	const float SumA = A.x + A.y + A.z;
-	const float SumB = B.x + B.y + B.z;
-	const float SumC = C.x + C.y + C.z;
-	const float SumD = D.x + D.y + D.z;
+	const float Aw = LobeColorToWeight(A);
+	const float Bw = LobeColorToWeight(B);
+	const float Cw = LobeColorToWeight(C);
+	const float Dw = LobeColorToWeight(D);
 
 	// use MIS routine to robustly compute the ratios even when both are 0, or overflow, etc ...
-	const float ProbAB = MISWeightBalanced(SumA, SumB); // A or B
-	const float ProbCD = MISWeightBalanced(SumC, SumD); // C or D
-	const float ProbLR = MISWeightBalanced(SumA + SumB, SumC + SumD); // AB or CD
+	const float ProbAB = MISWeightBalanced(Aw, Bw); // A or B
+	const float ProbCD = MISWeightBalanced(Cw, Dw); // C or D
+	const float ProbLR = MISWeightBalanced(Aw + Bw, Cw + Dw); // AB or CD
 
 	// Now build CDF
 	return float3(ProbAB * ProbLR, ProbLR, (ProbLR + ProbCD) - ProbLR * ProbCD);
@@ -88,6 +99,12 @@ void AddLobeWithMIS(inout float3 Weight, inout float Pdf, float3 LobeWeight, flo
 
 // The following structs are used as return types for the Eval/Sample methods
 
+#define PATHTRACER_SCATTER_CAMERA   0 // this is used to mark camera rays until we have actually scattered
+#define PATHTRACER_SCATTER_DIFFUSE  1
+#define PATHTRACER_SCATTER_SPECULAR 2
+#define PATHTRACER_SCATTER_REFRACT  3
+#define PATHTRACER_SCATTER_VOLUME   4
+
 struct FMaterialSample
 {
 	float3 Direction;
@@ -95,6 +112,12 @@ struct FMaterialSample
 	float Pdf;
 	float PositionBiasSign;
 	float Roughness;
+	int   ScatterType;
+
+	void AddLobeWithMIS(float3 LobeWeight, float LobePdf, float LobeProb)
+	{
+		::AddLobeWithMIS(Weight, Pdf, LobeWeight, LobePdf, LobeProb);
+	}
 };
 
 FMaterialSample NullMaterialSample()
@@ -103,7 +126,7 @@ FMaterialSample NullMaterialSample()
 	return (FMaterialSample)0;
 }
 
-FMaterialSample CreateMaterialSample(float3 Direction, float3 Weight, float Pdf, float PositionBiasSign, float Roughness)
+FMaterialSample CreateMaterialSample(float3 Direction, float3 Weight, float Pdf, float PositionBiasSign, float Roughness, int ScatterType)
 {
 	FMaterialSample Result;
 	Result.Direction = Direction;
@@ -111,6 +134,7 @@ FMaterialSample CreateMaterialSample(float3 Direction, float3 Weight, float Pdf,
 	Result.Pdf = Pdf;
 	Result.PositionBiasSign = PositionBiasSign;
 	Result.Roughness = Roughness;
+	Result.ScatterType = ScatterType;
 	return Result;
 }
 
@@ -118,6 +142,11 @@ struct FMaterialEval
 {
 	float3 Weight;
 	float Pdf;
+
+	void AddLobeWithMIS(float3 LobeWeight, float LobePdf, float LobeProb)
+	{
+		::AddLobeWithMIS(Weight, Pdf, LobeWeight, LobePdf, LobeProb);
+	}
 };
 
 FMaterialEval NullMaterialEval()
diff --git a/Engine/Shaders/Private/PathTracing/Material/PathTracingMaterialSampling.ush b/Engine/Shaders/Private/PathTracing/Material/PathTracingMaterialSampling.ush
--- a/Engine/Shaders/Private/PathTracing/Material/PathTracingMaterialSampling.ush
+++ b/Engine/Shaders/Private/PathTracing/Material/PathTracingMaterialSampling.ush
@@ -9,9 +9,7 @@ PathTracingMaterialSampling.usf: Material sampling functions
 #include "PathTracingMaterialCommon.ush"
 #include "PathTracingFresnel.ush"
 #include "PathTracingLambert.ush"
-#include "PathTracingRadianceProbe.ush"
 #include "PathTracingGlossy.ush"
-#include "PathTracingSpecularReflection.ush"
 #include "PathTracingSolidGlass.ush"
 #include "PathTracingThinGlass.ush"
 #include "PathTracingDefaultLit.ush"
@@ -29,23 +27,6 @@ PathTracingMaterialSampling.usf: Material sampling functions
 	#define SIMPLIFIED_MATERIAL_SHADER 0
 #endif
 
-bool IsNeeValidMaterial(FPathTracingPayload Payload)
-{
-	// Checks done in the same order as other sampling and evaluation functions to ensure consistency
-	if (Payload.BlendingMode == RAY_TRACING_BLEND_MODE_MODULATE ||
-		Payload.BlendingMode == RAY_TRACING_BLEND_MODE_ALPHA_HOLDOUT)
-	{
-		return false;
-	}
-
-	if (Payload.ShadingModelID == SHADINGMODELID_UNLIT)
-	{
-		return false;
-	}
-
-	return true;
-}
-
 float3 GetMaterialTransparency(FPathTracingPayload Payload, float3 RayDirection)
 {
 	// ThinTranslucent is special, all other shading models should go through blend mode logic
@@ -98,40 +79,6 @@ float3 GetMaterialTransparency(FPathTracingPayload Payload, float3 RayDirection)
 
 float3 EstimateMaterialAlbedo(FPathTracingPayload Payload)
 {
-	float BrdfWeight = 1;
-	bool HasRefraction = false;
-	switch (Payload.BlendingMode)
-	{
-		case RAY_TRACING_BLEND_MODE_TRANSLUCENT:
-#if !SIMPLIFIED_MATERIAL_SHADER
-			if (Payload.HasRefraction())
-			{
-				HasRefraction = true;
-			}
-			else if (Payload.ShadingModelID == SHADINGMODELID_THIN_TRANSLUCENT)
-			{
-				// special case -- thin translucent can use transparency without refraction but should still contain a specular highlight at Opacity=0
-				HasRefraction = true;
-			}
-			else
-#endif
-			{
-				// material is not configured with refraction, can use transparency
-				// NOTE: the albedo when considering refraction is handled below
-				BrdfWeight = Payload.Opacity;
-			}
-			break;
-		case RAY_TRACING_BLEND_MODE_ADDITIVE:
-			BrdfWeight = Payload.Opacity;
-			break;
-		case RAY_TRACING_BLEND_MODE_MODULATE:
-		case RAY_TRACING_BLEND_MODE_ALPHA_HOLDOUT:
-			// In these modes, the light does not interact with the material at all
-			return 0;
-		default:
-			break;
-	}
-
 	// TODO: make this more accurate by estimating the directional albedo
 	float3 Albedo = 0;
 	switch (Payload.ShadingModelID)
@@ -141,7 +88,7 @@ float3 EstimateMaterialAlbedo(FPathTracingPayload Payload)
 			return 0;
 #if SIMPLIFIED_MATERIAL_SHADER
 		case SHADINGMODELID_TWOSIDED_FOLIAGE:
-			Albedo = BrdfWeight * (Payload.BaseColor + Payload.SubsurfaceColor);
+			Albedo = Payload.BaseColor + Payload.SubsurfaceColor;
 			break;
 #else
 		case SHADINGMODELID_TWOSIDED_FOLIAGE:
@@ -149,38 +96,38 @@ float3 EstimateMaterialAlbedo(FPathTracingPayload Payload)
 		case SHADINGMODELID_PREINTEGRATED_SKIN:
 		case SHADINGMODELID_SUBSURFACE_PROFILE:
 		case SHADINGMODELID_EYE:
-			Albedo = BrdfWeight * (Payload.DiffuseColor + Payload.SpecularColor + Payload.SubsurfaceColor);
+			Albedo = Payload.DiffuseColor + Payload.SpecularColor + Payload.SubsurfaceColor;
 			break;
 		case SHADINGMODELID_SINGLELAYERWATER:
 		case SHADINGMODELID_DEFAULT_LIT:
 		case SHADINGMODELID_CLEAR_COAT:
-		case SHADINGMODELID_THIN_TRANSLUCENT:
 			// TODO: include clearcoat albedo as well
-			Albedo = BrdfWeight * (Payload.DiffuseColor + Payload.SpecularColor);
+			Albedo = Payload.DiffuseColor + Payload.SpecularColor;
 			break;
+		case SHADINGMODELID_THIN_TRANSLUCENT:
+		{
+			// glass portion is a mixture of pure specular and tinted transmission (when refracting)
+			// or glass portion is just the specular highlight with transmission accounted for by transparency
+			const float F0 = F0RGBToF0(Payload.SpecularColor);
+			const float3 GlassAlbedo = Payload.HasRefraction() ? lerp(Payload.GetTransmittanceColor(), 1.0, F0) : F0;
+			// surface is same as default-lit
+			const float3 SurfAlbedo = Payload.DiffuseColor + Payload.SpecularColor;
+			return lerp(GlassAlbedo, SurfAlbedo, Payload.Opacity);
+		}
 		case SHADINGMODELID_CLOTH:
-			Albedo = BrdfWeight * (Payload.DiffuseColor + lerp(Payload.SpecularColor, Payload.GetClothColor(), Payload.GetClothAmount()));
+			Albedo = (Payload.DiffuseColor + lerp(Payload.SpecularColor, Payload.GetClothColor(), Payload.GetClothAmount()));
 			break;
 		case SHADINGMODELID_HAIR:
-			Albedo = BrdfWeight * Payload.BaseColor;
 #endif
 		default:
-			Albedo = BrdfWeight * Payload.BaseColor;
+			Albedo = Payload.BaseColor;
 			break;
 	}
+	Albedo *= Payload.GetBSDFOpacity();
 #if !SIMPLIFIED_MATERIAL_SHADER
-	if (HasRefraction)
+	if (Payload.HasRefraction())
 	{
-		// glass has unit albedo, except in the thin case where the transmission portion can be tinted
-		float3 GlassAlbedo = 1.0;
-		if (Payload.ShadingModelID == SHADINGMODELID_THIN_TRANSLUCENT)
-		{
-			// If transparency is taking care of the tint, just account for the specular portion
-			// Otherwise, we are using (rough) refraction here, so account for both specular + transmission
-			const float F0 = F0RGBToF0(Payload.SpecularColor);
-			GlassAlbedo = Payload.HasRefraction() ? lerp(Payload.GetTransmittanceColor(), 1.0, F0) : F0;
-		}
-		Albedo = lerp(GlassAlbedo, Albedo, Payload.Opacity);
+		Albedo += 1 - Payload.Opacity; // glass albedo is 1.0
 	}
 #endif
 	return Albedo;
@@ -200,38 +147,7 @@ void AccumulateAlbedoNormal(inout FPathTracingPayload Payload, float3 PathThroug
 	// figure out how much the specular lobes should contribute to the total (relative to the current PathRoughness)
 	// if we are above the MinRoughness, use the PathRoughness as are lower bound instead so that multiple bounces don't accumulate too much albedo
 	const float RoughnessWeight = saturate((Payload.Roughness - max(MinRoughness, PathRoughness)) / (MaxRoughness - max(MinRoughness, PathRoughness)));
-	float BrdfWeight = 1;
-	bool HasRefraction = false;
-	switch (Payload.BlendingMode)
-	{
-		case RAY_TRACING_BLEND_MODE_TRANSLUCENT:
-#if !SIMPLIFIED_MATERIAL_SHADER
-			if (Payload.HasRefraction())
-			{
-				HasRefraction = true;
-			}
-			else if (Payload.ShadingModelID == SHADINGMODELID_THIN_TRANSLUCENT)
-			{
-				// special case -- thin translucent can use transparency without refraction but should still contain a specular highlight at Opacity=0
-				HasRefraction = true;
-			}
-			else
-#endif
-			{
-				BrdfWeight = Payload.Opacity;
-			}
-			break;
-		case RAY_TRACING_BLEND_MODE_ADDITIVE:
-			BrdfWeight = Payload.Opacity;
-			break;
-		case RAY_TRACING_BLEND_MODE_MODULATE:
-		case RAY_TRACING_BLEND_MODE_ALPHA_HOLDOUT:
-			// In these modes, the light does not interact with the material at all
-			return;
-		default:
-			break;
-	}
-
+	const float BrdfWeight = Payload.GetBSDFOpacity();
 	switch (Payload.ShadingModelID)
 	{
 	case SHADINGMODELID_UNLIT:
@@ -253,8 +169,16 @@ void AccumulateAlbedoNormal(inout FPathTracingPayload Payload, float3 PathThroug
 		break;
 	case SHADINGMODELID_SINGLELAYERWATER:
 	case SHADINGMODELID_DEFAULT_LIT:
+	{
+		if (Payload.HasRefraction())
+		{
+			const float GlassAlbedo = 1.0; // solid glass has unit albedo
+			Albedo += (1 - BrdfWeight) * RoughnessWeight * PathThroughput * GlassAlbedo;
+			Normal += (1 - BrdfWeight) * RoughnessWeight * Luminance(PathThroughput) * Payload.WorldNormal;
+		}
+		// fallthrough
+	}
 	case SHADINGMODELID_CLEAR_COAT:
-	case SHADINGMODELID_THIN_TRANSLUCENT:
 		Albedo += BrdfWeight * PathThroughput * (Payload.DiffuseColor + RoughnessWeight * Payload.SpecularColor);
 		Normal += BrdfWeight * Luminance(PathThroughput) * Payload.WorldNormal;
 		break;
@@ -262,28 +186,22 @@ void AccumulateAlbedoNormal(inout FPathTracingPayload Payload, float3 PathThroug
 		Albedo += BrdfWeight * PathThroughput * (Payload.DiffuseColor + lerp(RoughnessWeight * Payload.SpecularColor, Payload.GetClothColor(), Payload.GetClothAmount()));
 		Normal += BrdfWeight * Luminance(PathThroughput) * Payload.WorldNormal;
 		break;
+	case SHADINGMODELID_THIN_TRANSLUCENT:
+	{
+		const float F0 = F0RGBToF0(Payload.SpecularColor);
+		const float3 GlassAlbedo = Payload.HasRefraction() ? lerp(Payload.GetTransmittanceColor(), 1.0, F0) : F0;
+		Albedo += (1 - BrdfWeight) * RoughnessWeight * PathThroughput * GlassAlbedo;
+		Normal += (1 - BrdfWeight) * RoughnessWeight * Luminance(PathThroughput) * Payload.WorldNormal;
+		Albedo += BrdfWeight * PathThroughput * (Payload.DiffuseColor + RoughnessWeight * Payload.SpecularColor);
+		Normal += BrdfWeight * Luminance(PathThroughput) * Payload.WorldNormal;
+		break;
+	}
 #endif
 	default:
 		Albedo += BrdfWeight * PathThroughput * Payload.BaseColor;
 		Normal += BrdfWeight * Luminance(PathThroughput) * Payload.WorldNormal;
 		break;
 	}
-#if !SIMPLIFIED_MATERIAL_SHADER
-	if (HasRefraction)
-	{
-		// glass has unit albedo, except in the thin case where the transmission portion can be tinted
-		float3 GlassAlbedo = 1.0;
-		if (Payload.ShadingModelID == SHADINGMODELID_THIN_TRANSLUCENT)
-		{
-			// If transparency is taking care of the tint, just account for the specular portion
-			// Otherwise, we are using (rough) refraction here, so account for both specular + transmission
-			const float F0 = F0RGBToF0(Payload.SpecularColor);
-			GlassAlbedo = Payload.HasRefraction() ? lerp(Payload.GetTransmittanceColor(), 1.0, F0) : F0;
-		}
-		Albedo += (1 - BrdfWeight) * RoughnessWeight * PathThroughput * GlassAlbedo;
-		Normal += (1 - BrdfWeight) * RoughnessWeight * Luminance(PathThroughput) * Payload.WorldNormal;
-	}
-#endif
 }
 
 void AccumulateAlbedo(float3 SigmaS, float3 PathThroughput, float PathRoughness, inout float3 Albedo)
@@ -311,7 +229,7 @@ float GetAverageRoughness(FPathTracingPayload Payload)
 #if SIMPLIFIED_MATERIAL_SHADER
 	return 1.0;
 #else
-	const float GlassPercent = (Payload.BlendingMode == RAY_TRACING_BLEND_MODE_TRANSLUCENT && Payload.HasRefraction()) ? 1 - Payload.Opacity : 0.0;
+	const float GlassPercent = Payload.HasRefraction() ? 1 - Payload.Opacity : 0.0;
 	float Roughness = Payload.Roughness;
 	switch (Payload.ShadingModelID)
 	{
@@ -400,8 +318,7 @@ FPathTracingPayload UnpackPathTracingPayloadMaterial(FPackedPathTracingPayload P
 		{
 			if (bApplyOverrides)
 			{
-				// This is not ideal since hair is technically specular transmission
-				HitPayload.BaseColor = HitPayload.BaseColor  * View.DiffuseOverrideParameter.w  + View.DiffuseOverrideParameter.xyz;
+				HitPayload.BaseColor = HitPayload.BaseColor  * View.SpecularOverrideParameter.w  + View.SpecularOverrideParameter.xyz;
 				HitPayload.Specular = HitPayload.Specular * View.SpecularOverrideParameter.w + View.SpecularOverrideParameter.x;
 			}
 			break;
@@ -414,6 +331,14 @@ FPathTracingPayload UnpackPathTracingPayloadMaterial(FPackedPathTracingPayload P
 			}
 			break;
 		}
+		case SHADINGMODELID_CLOTH:
+		{
+			if (bApplyOverrides)
+			{
+				HitPayload.SetClothColor(HitPayload.GetClothColor() * View.SpecularOverrideParameter.w + View.SpecularOverrideParameter.xyz);
+			}
+			break;
+		}
 #endif
 	}
 
@@ -421,249 +346,101 @@ FPathTracingPayload UnpackPathTracingPayloadMaterial(FPackedPathTracingPayload P
 }
 
 FMaterialSample SampleMaterial(
-	float3 RayDirection,
+	float3 V_World,
 	FPathTracingPayload Payload,
-	float4 RandSample,
-	bool bForceRadianceProbe
+	float3 RandSample
 )
 {
-	if (bForceRadianceProbe)
-	{
-		return RadianceProbe_SampleMaterial(Payload, RandSample);
-	}
-
 #if SIMPLIFIED_MATERIAL_SHADER
-	FMaterialSample Result = Simplified_SampleMaterial(Payload, RandSample);
-	Result.Weight *= Payload.Opacity;
-	return Result;
+	return Simplified_SampleMaterial(Payload, RandSample);
 	//return Lambert_SampleMaterial(Payload, RandSample);
 #else
-
-	if (Payload.BlendingMode == RAY_TRACING_BLEND_MODE_MODULATE ||
-		Payload.BlendingMode == RAY_TRACING_BLEND_MODE_ALPHA_HOLDOUT)
-	{
-		// transparency only, nothing to sample
-		return NullMaterialSample();
-	}
-
-	// Decide the overall weight of the BRDF portion vs. the translucent portion
-	// Note that the translucent portion could be refraction, transparency or SSS
-	float BrdfWeight = 1;
-	float TransWeight = 0;
-	float BrdfProb = 1;
-	switch (Payload.BlendingMode)
-	{
-		case RAY_TRACING_BLEND_MODE_TRANSLUCENT:
-			if (!Payload.HasRefraction() && Payload.ShadingModelID != SHADINGMODELID_THIN_TRANSLUCENT)
-			{
-				// scale down BRDF contribution, but don't pick any indirect ray because its already accounted for by transparency handling
-				BrdfWeight = Payload.Opacity;
-			}
-			else
-			{
-				// randomly choose between brdf and glass
-				BrdfWeight = Payload.Opacity;
-				TransWeight = 1 - Payload.Opacity;
-				BrdfProb = Payload.Opacity;
-			}
-			break;
-		case RAY_TRACING_BLEND_MODE_ADDITIVE:
-			BrdfWeight = Payload.Opacity;
-			break;
-		default:
-			// opaque or one of the modes handled above
-			break;
-	}
-
-	FMaterialSample Result = NullMaterialSample();
-	if (RandSample.x < BrdfProb)
-	{
-		// we randomly decided to sample the brdf portion of the material
-		RandSample.x /= BrdfProb;
-	}
-	else
-	{
-		// we randomly decided to sample the translucent portion
-		RandSample.x = (RandSample.x - BrdfProb) / (1 - BrdfProb);
-
-		// only case that should reach this point is the refraction model (which can be thin or solid)
-		if (Payload.ShadingModelID == SHADINGMODELID_THIN_TRANSLUCENT)
-		{
-			Result = RoughThinGlass_SampleMaterial(RayDirection, Payload, RandSample);
-		}
-		else
-		{
-			Result = RoughGlass_SampleMaterial(RayDirection, Payload, RandSample);
-		}
-		Result.Weight *= TransWeight;
-		Result.Pdf *= 1 - BrdfProb;
-		return Result;
-	}
-
 	switch (Payload.ShadingModelID)
 	{
 		case SHADINGMODELID_UNLIT:
 			// Nothing to do, material does not scatter light
-			return Result;
+			return NullMaterialSample();
 		case SHADINGMODELID_DEFAULT_LIT:
+		case SHADINGMODELID_SINGLELAYERWATER:
+		{
+			if (Payload.HasRefraction() && Payload.Opacity < 1)
+			{
+				return RoughGlass_SampleMaterial(V_World, Payload, RandSample);
+			}
+			// fallthrough
+		}
 		case SHADINGMODELID_SUBSURFACE:
 		case SHADINGMODELID_PREINTEGRATED_SKIN:
-		case SHADINGMODELID_SINGLELAYERWATER:
+			return DefaultLit_SampleMaterial(V_World, Payload, RandSample);
 		case SHADINGMODELID_THIN_TRANSLUCENT:
-			Result = DefaultLit_SampleMaterial(RayDirection, Payload, RandSample);
-			break;
+			return RoughThinGlass_SampleMaterial(V_World, Payload, RandSample);
 		case SHADINGMODELID_SUBSURFACE_PROFILE:
-			Result = SubsurfaceProfile_SampleMaterial(RayDirection, Payload, RandSample);
-			break;
+			return SubsurfaceProfile_SampleMaterial(V_World, Payload, RandSample);
 		case SHADINGMODELID_CLOTH:
-			Result = Cloth_SampleMaterial(RayDirection, Payload, RandSample);
-			break;
+			return Cloth_SampleMaterial(V_World, Payload, RandSample);
 		case SHADINGMODELID_CLEAR_COAT:
-			Result = ClearCoat_SampleMaterial(RayDirection, Payload, RandSample);
-			break;
+			return ClearCoat_SampleMaterial(V_World, Payload, RandSample);
 		case SHADINGMODELID_TWOSIDED_FOLIAGE:
-			Result = TwoSidedFoliage_SampleMaterial(RayDirection, Payload, RandSample);
-			break;
+			return TwoSidedFoliage_SampleMaterial(V_World, Payload, RandSample);
 		case SHADINGMODELID_EYE:
-			Result = Eye_SampleMaterial(RayDirection, Payload, RandSample);
-			break;
+			return Eye_SampleMaterial(V_World, Payload, RandSample);
 		case SHADINGMODELID_HAIR:
-			Result = Hair_SampleMaterial(RayDirection, Payload, RandSample);
-			break;
+			return Hair_SampleMaterial(V_World, Payload, RandSample);
 		case SHADINGMODELID_MEDIUM:
-			Result = Medium_SampleMaterial(RayDirection, Payload, RandSample);
-			break;
+			return Medium_SampleMaterial(V_World, Payload, RandSample);
 		default:
-			Result = Lambert_SampleMaterial(Payload, RandSample);
-			break;
+			return Lambert_SampleMaterial(Payload, RandSample);
 	}
-
-	// final adjustment to weight / prob based on the stochastic decision for translucency earlier
-	Result.Weight *= BrdfWeight;
-	Result.Pdf *= BrdfProb;
-	return Result;
-#endif
+	return NullMaterialSample();
+#endif // SIMPLIFIED_MATERIAL_SHADER
 }
 
 FMaterialEval EvalMaterial(
-	float3 IncomingDirection,
-	float3 OutgoingDirection,
+	float3 V_World,
+	float3 L_World,
 	FPathTracingPayload Payload,
-	bool bForceRadianceProbe
+	float2 DiffuseSpecularScale
 )
 {
-	if (bForceRadianceProbe)
-	{
-		return RadianceProbe_EvalMaterial(OutgoingDirection, Payload);
-	}
 #if SIMPLIFIED_MATERIAL_SHADER
-	FMaterialEval Result = Simplified_EvalMaterial(OutgoingDirection, Payload);
-	Result.Weight *= Payload.Opacity;
-	return Result;
-
-	//return Lambert_EvalMaterial(OutgoingDirection, Payload);
+	return Simplified_EvalMaterial(L_World, Payload);
+	//return Lambert_EvalMaterial(L_World, Payload);
 #else
-
-	if (Payload.BlendingMode == RAY_TRACING_BLEND_MODE_MODULATE ||
-		Payload.BlendingMode == RAY_TRACING_BLEND_MODE_ALPHA_HOLDOUT)
-	{
-		return NullMaterialEval();
-	}
-
-	// Decide the overall weight of the BRDF portion vs. the translucent portion
-	// Note that the translucent portion could be refraction, transparency or SSS
-	float BrdfWeight = 1;
-	float TransWeight = 0;
-	float BrdfProb = 1;
-	switch (Payload.BlendingMode)
+	switch (Payload.ShadingModelID)
 	{
-		case RAY_TRACING_BLEND_MODE_TRANSLUCENT:
-			// brdf:   Opacity
-			// glass:  1.0 - Opacity
-			if (!Payload.HasRefraction() && Payload.ShadingModelID != SHADINGMODELID_THIN_TRANSLUCENT)
-			{
-				// handled by native transpaarency, just weight down the brdf contrib
-				BrdfWeight = Payload.Opacity;
-			}
-			else
-			{
-				// brdf or glass
-				BrdfWeight = Payload.Opacity;
-				TransWeight = 1 - Payload.Opacity;
-				BrdfProb = Payload.Opacity;
-			}
-			break;
-		case RAY_TRACING_BLEND_MODE_ADDITIVE:
-			// brdf:   Opacity
-			BrdfWeight = Payload.Opacity;
-			break;
-		default:
-			// opaque or one of the modes handled above
+		case SHADINGMODELID_UNLIT:
 			break;
-	}
-
-	// NOTE: the translucent portions are all singular -- so we can ignore them in the eval method for now as they have Pdf=+inf for a single direction only
-
-	FMaterialEval Result = NullMaterialEval();
-	if (BrdfWeight > 0.0)
-	{
-		switch (Payload.ShadingModelID)
-		{
-			case SHADINGMODELID_UNLIT:
-				break;
-			case SHADINGMODELID_DEFAULT_LIT:
-			case SHADINGMODELID_SUBSURFACE:
-			case SHADINGMODELID_PREINTEGRATED_SKIN:
-			case SHADINGMODELID_SINGLELAYERWATER:
-			case SHADINGMODELID_THIN_TRANSLUCENT:
-				Result = DefaultLit_EvalMaterial(IncomingDirection, OutgoingDirection, Payload);
-				break;
-			case SHADINGMODELID_SUBSURFACE_PROFILE:
-				Result = SubsurfaceProfile_EvalMaterial(IncomingDirection, OutgoingDirection, Payload);
-				break;
-			case SHADINGMODELID_CLOTH:
-				Result = Cloth_EvalMaterial(IncomingDirection, OutgoingDirection, Payload);
-				break;
-			case SHADINGMODELID_CLEAR_COAT:
-				Result = ClearCoat_EvalMaterial(IncomingDirection, OutgoingDirection, Payload);
-				break;
-			case SHADINGMODELID_TWOSIDED_FOLIAGE:
-				Result = TwoSidedFoliage_EvalMaterial(IncomingDirection, OutgoingDirection, Payload);
-				break;
-			case SHADINGMODELID_EYE:
-				Result = Eye_EvalMaterial(IncomingDirection, OutgoingDirection, Payload);
-				break;
-			case SHADINGMODELID_HAIR:
-				Result = Hair_EvalMaterial(IncomingDirection, OutgoingDirection, Payload);
-				break;
-			case SHADINGMODELID_MEDIUM:
-				Result = Medium_EvalMaterial(IncomingDirection, OutgoingDirection, Payload);
-				break;
-			default:
-				Result = Lambert_EvalMaterial(OutgoingDirection, Payload);
-				break;
-		}
-		// final adjustment to weight / prob for translucency
-		Result.Weight *= BrdfWeight;
-		Result.Pdf *= BrdfProb;
-	}
-
-	if (BrdfProb < 1.0)
-	{
-		// refraction model
-		FMaterialEval GlassEval;
-		if (Payload.ShadingModelID == SHADINGMODELID_THIN_TRANSLUCENT)
-		{
-			// just a thin slab (reflection-like lobe on both sides)
-			GlassEval = RoughThinGlass_EvalMaterial(IncomingDirection, OutgoingDirection, Payload);
-		}
-		else
+		case SHADINGMODELID_DEFAULT_LIT:
+		case SHADINGMODELID_SINGLELAYERWATER:
 		{
-			// real solid glass, reflection and refraction
-			GlassEval = RoughGlass_EvalMaterial(IncomingDirection, OutgoingDirection, Payload);
+			if (Payload.HasRefraction() && Payload.Opacity < 1)
+			{
+				return RoughGlass_EvalMaterial(V_World, L_World, Payload, DiffuseSpecularScale);
+			}
+			// fallthrough
 		}
-		AddLobeWithMIS(Result.Weight, Result.Pdf, GlassEval.Weight * TransWeight, GlassEval.Pdf, 1.0 - BrdfProb);
+		case SHADINGMODELID_SUBSURFACE:
+		case SHADINGMODELID_PREINTEGRATED_SKIN:
+			return DefaultLit_EvalMaterial(V_World, L_World, Payload, DiffuseSpecularScale);
+		case SHADINGMODELID_THIN_TRANSLUCENT:
+			return RoughThinGlass_EvalMaterial(V_World, L_World, Payload, DiffuseSpecularScale);
+		case SHADINGMODELID_SUBSURFACE_PROFILE:
+			return SubsurfaceProfile_EvalMaterial(V_World, L_World, Payload, DiffuseSpecularScale);
+		case SHADINGMODELID_CLOTH:
+			return Cloth_EvalMaterial(V_World, L_World, Payload, DiffuseSpecularScale);
+		case SHADINGMODELID_CLEAR_COAT:
+			return ClearCoat_EvalMaterial(V_World, L_World, Payload, DiffuseSpecularScale);
+		case SHADINGMODELID_TWOSIDED_FOLIAGE:
+			return TwoSidedFoliage_EvalMaterial(V_World, L_World, Payload, DiffuseSpecularScale);
+		case SHADINGMODELID_EYE:
+			return Eye_EvalMaterial(V_World, L_World, Payload, DiffuseSpecularScale);
+		case SHADINGMODELID_HAIR:
+			return Hair_EvalMaterial(V_World, L_World, Payload, DiffuseSpecularScale);
+		case SHADINGMODELID_MEDIUM:
+			return Medium_EvalMaterial(V_World, L_World, Payload, DiffuseSpecularScale);
+		default:
+			return Lambert_EvalMaterial(L_World, Payload, DiffuseSpecularScale);
 	}
-	return Result;
-#endif
+	return NullMaterialEval();
+#endif // SIMPLIFIED_MATERIAL_SHADER
 }
diff --git a/Engine/Shaders/Private/PathTracing/Material/PathTracingMedium.ush b/Engine/Shaders/Private/PathTracing/Material/PathTracingMedium.ush
--- a/Engine/Shaders/Private/PathTracing/Material/PathTracingMedium.ush
+++ b/Engine/Shaders/Private/PathTracing/Material/PathTracingMedium.ush
@@ -12,9 +12,9 @@ PathTracingMedium.usf: Volumetric phase function
 #define USE_UNIFORM_PHASE_FUNCTION 0
 
 FMaterialSample Medium_SampleMaterial(
-	float3 IncomingDirection,
+	float3 V_World,
 	FPathTracingPayload Payload,
-	float4 RandSample)
+	float3 RandSample)
 {
 	const float3 RayleighWeight = Payload.BaseColor;
 	const float3 HGWeight = Payload.GetHGWeight();
@@ -23,7 +23,7 @@ FMaterialSample Medium_SampleMaterial(
 #if USE_UNIFORM_PHASE_FUNCTION
 	const float4 Result = UniformSampleSphere(RandSample.xy);
 	// uniform scattering
-	return CreateMaterialSample(Result.xyz, RayleighWeight + HGWeight, Result.w, 1.0, 1.0);
+	return CreateMaterialSample(Result.xyz, RayleighWeight + HGWeight, Result.w, 1.0, 1.0, PATHTRACER_SCATTER_VOLUME);
 #else
 	const float LobeProb = LobeSelectionProb(RayleighWeight, HGWeight);
 
@@ -33,19 +33,20 @@ FMaterialSample Medium_SampleMaterial(
 	const float SinTheta = sqrt(saturate(1.0 - CosTheta * CosTheta));
 	const float Phi = RandSample.y * (2.0 * PI);
 
-	FMaterialSample Result = NullMaterialSample();
-	Result.Direction = TangentToWorld(float3(cos(Phi) * SinTheta, sin(Phi) * SinTheta, CosTheta), -IncomingDirection);
-	Result.Roughness = 1.0;
-	AddLobeWithMIS(Result.Weight, Result.Pdf, RayleighWeight, RayleighPhase(CosTheta), LobeProb);
-	AddLobeWithMIS(Result.Weight, Result.Pdf, HGWeight      , HenyeyGreensteinPhase(G, CosTheta )  , 1.0 - LobeProb);
+	const float3 L_World = TangentToWorld(float3(cos(Phi) * SinTheta, sin(Phi) * SinTheta, CosTheta), V_World);
+
+	FMaterialSample Result = CreateMaterialSample(L_World, 0.0, 0.0, 1.0, 1.0, PATHTRACER_SCATTER_VOLUME);
+	Result.AddLobeWithMIS(RayleighWeight, RayleighPhase(CosTheta), LobeProb);
+	Result.AddLobeWithMIS(HGWeight      , HenyeyGreensteinPhase(G, CosTheta )  , 1.0 - LobeProb);
 	return Result;
 #endif
 }
 
 FMaterialEval Medium_EvalMaterial(
-	float3 IncomingDirection,
-	float3 OutgoingDirection,
-	FPathTracingPayload Payload
+	float3 V_World,
+	float3 L_World,
+	FPathTracingPayload Payload,
+	float2 DiffuseSpecularScale
 )
 {
 	const float3 RayleighWeight = Payload.BaseColor;
@@ -58,10 +59,11 @@ FMaterialEval Medium_EvalMaterial(
 #else
 	const float LobeProb = LobeSelectionProb(RayleighWeight, HGWeight);
 
-	const float CosTheta = -dot(IncomingDirection, OutgoingDirection);
+	const float CosTheta = dot(V_World, L_World);
 	FMaterialEval Result = NullMaterialEval();
-	AddLobeWithMIS(Result.Weight, Result.Pdf, RayleighWeight, RayleighPhase(CosTheta), LobeProb);
-	AddLobeWithMIS(Result.Weight, Result.Pdf, HGWeight      , HenyeyGreensteinPhase(G, CosTheta), 1.0 - LobeProb);
+	Result.AddLobeWithMIS(RayleighWeight, RayleighPhase(CosTheta), LobeProb);
+	Result.AddLobeWithMIS(HGWeight      , HenyeyGreensteinPhase(G, CosTheta), 1.0 - LobeProb);
+	Result.Weight *= DiffuseSpecularScale.x;
 	return Result;
 #endif
 }
diff --git a/Engine/Shaders/Private/PathTracing/Material/PathTracingRadianceProbe.ush b/Engine/Shaders/Private/PathTracing/Material/PathTracingRadianceProbe.ush
--- a/Engine/Shaders/Private/PathTracing/Material/PathTracingRadianceProbe.ush
+++ b/Engine/Shaders/Private/PathTracing/Material/PathTracingRadianceProbe.ush
@@ -50,7 +50,7 @@ float2 InverseConcentricMapping(float2 p)
 
 FMaterialSample RadianceProbe_SampleMaterial(
 	FPathTracingPayload Payload,
-	float4 RandSample
+	float3 RandSample
 )
 {
 	float PdfAdjust = 1.0f; 
@@ -136,11 +136,11 @@ FMaterialSample RadianceProbe_SampleMaterial(
 	// As the name suggests, radiance probe has a combined throughput (Pdf * Weight) == 1
 	float OutWeight = (OutPdf > 0.0f) ? (1.0f / OutPdf) : 0;
 
-	return CreateMaterialSample(OutDirection, OutWeight, OutPdf, 1.0, 1.0);
+	return CreateMaterialSample(OutDirection, OutWeight, OutPdf, 1.0, 1.0, PATHTRACER_SCATTER_DIFFUSE);
 }
 
 FMaterialEval RadianceProbe_EvalMaterial(
-	float3 OutgoingDirection,
+	float3 L_World,
 	FPathTracingPayload Payload
 )
 {
@@ -150,7 +150,7 @@ FMaterialEval RadianceProbe_EvalMaterial(
 	
 #ifdef LIGHTMAP_PATH_TRACING_MAIN_RG 
 #if USE_FIRST_BOUNCE_RAY_GUIDING
-	float3 TangentDirection = WorldToTangent(OutgoingDirection, N_World);
+	float3 TangentDirection = WorldToTangent(L_World, N_World);
 	float2 PrimarySample = InverseConcentricMapping(TangentDirection.xy);
 	
 	uint2 BatchedLaunchIndex = DispatchRaysIndex().xy;
@@ -209,7 +209,7 @@ FMaterialEval RadianceProbe_EvalMaterial(
 	}
 #endif	
 #endif
-	float NoL = max(dot(N_World, OutgoingDirection), 0);
+	float NoL = max(dot(N_World, L_World), 0);
 
 #ifdef LIGHTMAP_PATH_TRACING_MAIN_RG
 	float OutPdf = NoL / PI * PdfAdjust;
@@ -226,29 +226,29 @@ FMaterialEval RadianceProbe_EvalMaterial(
 
 FMaterialSample RadianceProbe_SampleMaterial(
 	FPathTracingPayload Payload,
-	float4 RandSample)
+	float3 RandSample)
 {
 	float3 N_World = Payload.WorldNormal;
 
 	float4 SampledValue = CosineSampleHemisphere(RandSample.xy);
 
-	return CreateMaterialSample(TangentToWorld(SampledValue.xyz, N_World), 1.0, SampledValue.w, 1.0, 1.0);
+	return CreateMaterialSample(TangentToWorld(SampledValue.xyz, N_World), 1.0, SampledValue.w, 1.0, 1.0, PATHTRACER_SCATTER_DIFFUSE);
 }
 
-float RadianceProbe_Pdf(float3 OutgoingDirection, FPathTracingPayload Payload)
+float RadianceProbe_Pdf(float3 L_World, FPathTracingPayload Payload)
 {
 	float3 N_World = Payload.WorldNormal;
-	float NoL = saturate(dot(N_World, OutgoingDirection));
+	float NoL = saturate(dot(N_World, L_World));
 	return NoL / PI;
 }
 
 FMaterialEval RadianceProbe_EvalMaterial(
-	float3 OutgoingDirection,
+	float3 L_World,
 	FPathTracingPayload Payload
 )
 {
 	float3 N_World = Payload.WorldNormal;
-	float NoL = saturate(dot(N_World, OutgoingDirection));
+	float NoL = saturate(dot(N_World, L_World));
 
 	return CreateMaterialEval(Payload.DiffuseColor, NoL / PI);
 }
diff --git a/Engine/Shaders/Private/PathTracing/Material/PathTracingSolidGlass.ush b/Engine/Shaders/Private/PathTracing/Material/PathTracingSolidGlass.ush
--- a/Engine/Shaders/Private/PathTracing/Material/PathTracingSolidGlass.ush
+++ b/Engine/Shaders/Private/PathTracing/Material/PathTracingSolidGlass.ush
@@ -11,115 +11,217 @@ PathTracingSolidGlass.ush: Microfacet Refractive BSDF for solid glass
 #include "PathTracingGlossy.ush"
 #include "PathTracingEnergyConservation.ush"
 
+struct FRoughGlassData
+{
+	float3x3 Basis;
+	float2 Alpha;
+	float3 V;
 
-FMaterialSample RoughGlass_SampleMaterial(
-	float3 RayDirection,
-	FPathTracingPayload Payload,
-	float4 RandSample)
+	FBxDFEnergyTermsRGB Spec;
+	float3 DiffWeight;
+
+	float2 LobeCdf;
+	float3 LobePdf;
+
+	// glass lobe data
+	float Eta;
+	float F0;
+	float EGlass;
+};
+
+FRoughGlassData PrepareRoughGlassData(FPathTracingPayload Payload, float3 V_World)
 {
-	// NOTE: IsFrontFace() determines polygon orientation, because the normal is always flipped towards in the incoming ray
-	bool bIsEntering = Payload.IsFrontFace();
-	float3 N_World = Payload.WorldNormal;
-	float3 V_World = -RayDirection;
+	FRoughGlassData Data = (FRoughGlassData)0;
 
+	Data.Basis = GetGGXBasis(Payload.Roughness, Payload.Anisotropy, Payload.WorldNormal, Payload.WorldTangent, Data.Alpha);
+	Data.V = mul(Data.Basis, V_World);
+	const float NoV = saturate(Data.V.z);
+
+
+	Data.Spec = ComputeGGXSpecEnergyTermsRGB(Payload.Roughness, NoV, Payload.SpecularColor);
+	Data.DiffWeight = Payload.DiffuseColor * (1 - Data.Spec.E);
+
+	// Probability of picking diffuse vs spec vs glass
+	Data.LobeCdf = LobeSelectionCdf(Payload.Opacity * Data.DiffWeight, Payload.Opacity * Data.Spec.E, 1.0 - Payload.Opacity);
+	Data.LobePdf = LobeSelectionPdf(Data.LobeCdf);
+
+
+	// NOTE: IsFrontFace() determines polygon orientation, because the normal is always flipped towards in the incoming ray
 	// #dxr_todo: Maintain a refraction stack on the path tracing payload
-	const float Eta = bIsEntering ? Payload.Ior : rcp(Payload.Ior);
-	const float F0 = F0RGBToF0(Payload.SpecularColor);
+	Data.Eta = Payload.IsFrontFace() ? Payload.Ior : rcp(Payload.Ior);
+	Data.F0 = F0RGBToF0(Payload.SpecularColor);
 
-	float2 Alpha;
-	float3x3 Basis = GetGGXBasis(Payload.Roughness, Payload.Anisotropy, Payload.WorldNormal, Payload.WorldTangent, Alpha);
+	// correct for energy loss by scaling the whole BSDF
+	Data.EGlass = GGXEnergyLookup(Payload.Roughness, NoV, Data.Eta);
+
+	return Data;
+}
 
-	float3 V = mul(Basis, V_World);
-	float NoV = saturate(V.z);
 
+FMaterialSample RoughGlass_SampleMaterial(
+	float3 V_World,
+	FPathTracingPayload Payload,
+	float3 RandSample)
+{
+	const FRoughGlassData Data = PrepareRoughGlassData(Payload, V_World);
+	const float3 V = Data.V;
+	const float NoV = saturate(V.z);
 	if (NoV == 0)
 	{
 		// invalid grazing angle
 		return NullMaterialSample();
 	}
 
-	// correct for energy loss by scaling the whole BSDF
-	const float E = GGXEnergyLookup(Payload.Roughness, NoV, Eta);
+	// Randomly choose to sample diffuse or specular
+	if (RandSample.x < Data.LobeCdf.x)
+	{
+		RandSample.x /= Data.LobeCdf.x;
 
-	float3 H = ImportanceSampleVisibleGGX_aniso(UniformSampleDisk(RandSample.xy), Alpha, V).xyz;
+		// Lambert
+		// TODO: evaluate CosineSampleHemisphereConcentric
+		const float4 Result = CosineSampleHemisphere(RandSample.xy);
 
-	float3 L = 0;
-	float F = 0;
-	if (SampleRefraction(-V, H, Eta, F0, RandSample.z, L, F))
-	{
-		const float2 GGXResult = GGXEvalRefraction(L, V, H, Alpha, Eta);
+		const float3 L = Result.xyz;
+		const float3 H = normalize(L + V);
 
 		// transform to world space
-		return CreateMaterialSample(normalize(mul(L, Basis)), GGXResult.x / E, F * GGXResult.y, -1.0, Payload.Roughness);
+		const float3 L_World = normalize(mul(L, Data.Basis));
+
+		const float NoL = saturate(L.z);
+		const float VoH = saturate(dot(V, H));
+		const float NoH = saturate(H.z);
+
+		const float3 Diffuse = GetPathTracingDiffuseModel(Data.DiffWeight, Payload.Roughness, NoV, NoL, VoH, NoH);
+		const float3 DiffWeight = (Payload.Opacity / Data.LobePdf.x) * Diffuse * ShadowTerminatorTerm(L_World, Payload.WorldNormal, Payload.WorldSmoothNormal);
+		const float DiffPdf = Data.LobePdf.x * NoL / PI;
+
+		return CreateMaterialSample(L_World, DiffWeight, DiffPdf, 1.0, 1.0, PATHTRACER_SCATTER_DIFFUSE);
 	}
 	else
 	{
-		const float2 GGXResult = GGXEvalReflection(L, V, H, Alpha);
+		const bool bDoSpec = RandSample.x < Data.LobeCdf.y;
 
-		// transform to world space
-		return CreateMaterialSample(normalize(mul(L, Basis)), GGXResult.x / E, F * GGXResult.y, 1.0, Payload.Roughness);
+		if (bDoSpec)
+		{
+			RandSample.x -= Data.LobeCdf.x;
+			RandSample.x /= Data.LobeCdf.y - Data.LobeCdf.x;
+		}
+		else
+		{
+			RandSample.x -= Data.LobeCdf.y;
+			RandSample.x /= 1.0 - Data.LobeCdf.y;
+		}
+
+		// #dxr_todo: Evaluate UniformSampleDiskConcentric
+		const float3 H = ImportanceSampleVisibleGGX_aniso(UniformSampleDisk(RandSample.xy), Data.Alpha, V).xyz;
+
+		if (bDoSpec)
+		{
+			const float3 L = reflect(-V, H);
+			if (L.z <= 0)
+			{
+				// invalid output direction, exit early
+				return NullMaterialSample();
+			}
+
+			// transform to world space
+			const float3 L_World = normalize(mul(L, Data.Basis));
+
+			const float NoL = saturate(L.z);
+			const float VoH = saturate(dot(V, H));
+			const float NoH = saturate(H.z);
+
+			// Specular lobe
+			const float2 GGXResult = GGXEvalReflection(L, V, H, Data.Alpha);
+			const float3 F = F_Schlick(Payload.SpecularColor, VoH);
+			const float3 SpecWeight = (Payload.Opacity / Data.LobePdf.y) * F * GGXResult.x * Data.Spec.W;
+			const float SpecPdf = GGXResult.y * Data.LobePdf.y;
+
+			return CreateMaterialSample(L_World, SpecWeight, SpecPdf, 1.0, Payload.Roughness, PATHTRACER_SCATTER_SPECULAR);
+		}
+		else
+		{
+			// Glass lobe (reflection and refraction)
+			float3 L = 0;
+			float F = 0;
+
+			const bool bRefract = SampleRefraction(-V, H, Data.Eta, Data.F0, RandSample.z, L, F);
+			// transform to world space
+			const float3 L_World = normalize(mul(L, Data.Basis));
+
+			const float2 GGXResult = bRefract ? GGXEvalRefraction(L, V, H, Data.Alpha, Data.Eta) : GGXEvalReflection(L, V, H, Data.Alpha);
+
+			return CreateMaterialSample(L_World, (1.0 - Payload.Opacity) * GGXResult.x / Data.EGlass, F * GGXResult.y, bRefract ? -1.0 : 1.0, Payload.Roughness, bRefract ? PATHTRACER_SCATTER_REFRACT : PATHTRACER_SCATTER_SPECULAR);
+		}
 	}
 }
 
 FMaterialEval RoughGlass_EvalMaterial(
-	float3 IncomingDirection,
-	float3 OutgoingDirection,
-	FPathTracingPayload Payload
+	float3 V_World,
+	float3 L_World,
+	FPathTracingPayload Payload,
+	float2 DiffuseSpecularScale
 )
 {
-	// NOTE: IsFrontFace() determines polygon orientation, because the normal is always flipped towards in the incoming ray
-	bool bIsEntering = Payload.IsFrontFace();
-	float3 V_World = -IncomingDirection;
-	float3 L_World = OutgoingDirection;
-	float3 N_World = Payload.WorldNormal;
-
-	// #dxr_todo: Maintain a refraction stack on the path tracing payload
-	const float Eta = bIsEntering ? Payload.Ior : rcp(Payload.Ior);
-	const float F0 = F0RGBToF0(Payload.SpecularColor);
-
-	float2 Alpha;
-	float3x3 Basis = GetGGXBasis(Payload.Roughness, Payload.Anisotropy, Payload.WorldNormal, Payload.WorldTangent, Alpha);
-
-	// move vectors into right shading frame
-	float3 V = mul(Basis, V_World);
-	float3 L = mul(Basis, L_World);
-
-	if (V.z <= 0)
+	const FRoughGlassData Data = PrepareRoughGlassData(Payload, V_World);
+	const float3 V = Data.V;
+	const float NoV = saturate(V.z);
+	if (NoV == 0)
 	{
-		// invalid input
+		// invalid grazing angle
 		return NullMaterialEval();
 	}
 
-	float NoV = saturate(V.z);
 
-	// correct for energy loss by scaling the whole BSDF
-	const float E = GGXEnergyLookup(Payload.Roughness, NoV, Eta);
+	// move vectors into right shading frame
+	const float3 L = mul(Data.Basis, L_World);
+
+	FMaterialEval Result = NullMaterialEval();
 
 	if (L.z > 0)
 	{
 		// reflection side
-		const float3 H = normalize(L + V);
-		const float VoH = saturate(dot(V, H));
-		const float F = FresnelReflectance(VoH, Eta, F0);
+		const float3 H = normalize(V + L);
 
-		const float2 GGXResult = GGXEvalReflection(L, V, H, Alpha);
-		return CreateMaterialEval(GGXResult.x / E, F * GGXResult.y);
+		const float NoL = saturate(L.z);
+		const float VoH = saturate(dot(V, H));
+		const float NoH = saturate(H.z);
+
+		// Diffuse Lobe
+		const float3 Diffuse = GetPathTracingDiffuseModel(Data.DiffWeight, Payload.Roughness, NoV, NoL, VoH, NoH);
+		Result.AddLobeWithMIS(Payload.Opacity * Diffuse * ShadowTerminatorTerm(L_World, Payload.WorldNormal, Payload.WorldSmoothNormal) * DiffuseSpecularScale.x, NoL / PI, Data.LobePdf.x);
+
+		// Specular lobe
+		const float2 GGXResult = GGXEvalReflection(L, V, H, Data.Alpha);
+		const float3 Fs = F_Schlick(Payload.SpecularColor, VoH);
+		const float3 SpecWeight = Payload.Opacity * GGXResult.x * Data.Spec.W * Fs;
+		const float SpecPdf = GGXResult.y;
+		Result.AddLobeWithMIS(SpecWeight * DiffuseSpecularScale.y, SpecPdf, Data.LobePdf.y);
+
+		// Glass lobe
+		const float Fg = FresnelReflectance(VoH, Data.Eta, Data.F0);
+		const float3 GlassWeight = GGXResult.x * (1 - Payload.Opacity) / Data.EGlass;
+		const float GlassPdf = GGXResult.y * Fg;
+		Result.AddLobeWithMIS(GlassWeight * DiffuseSpecularScale.y, GlassPdf, Data.LobePdf.z);
 	}
-	else if (L.z < 0)
+	else if (L.z < 0 && Data.LobePdf.z > 0)
 	{
 		// refraction side
-		float NoL = saturate(-L.z);
-		float3 Ht = -(Eta * L + V);
-		Ht = normalize((Eta < 1.0f) ? -Ht : Ht);
+		const float NoL = saturate(-L.z);
+		float3 Ht = -(Data.Eta * L + V);
+		Ht = normalize((Data.Eta < 1.0f) ? -Ht : Ht);
 
-		float VoH = dot(V, Ht);
-		float F = 1.0f - FresnelReflectance(VoH, Eta, F0);
-		if (F > 0)
+		const float VoH = dot(V, Ht);
+		const float Fg = 1.0f - FresnelReflectance(VoH, Data.Eta, Data.F0);
+		if (Fg > 0)
 		{
-			const float2 GGXResult = GGXEvalRefraction(L, V, Ht, Alpha, Eta);
+			const float2 GGXResult = GGXEvalRefraction(L, V, Ht, Data.Alpha, Data.Eta);
+
+			const float3 GlassWeight = GGXResult.x * (1 - Payload.Opacity) / Data.EGlass;
+			const float GlassPdf = GGXResult.y * Fg;
 
-			return CreateMaterialEval(GGXResult.x / E, F * GGXResult.y);
+			Result.AddLobeWithMIS(GlassWeight * DiffuseSpecularScale.y, GlassPdf, Data.LobePdf.z);
 		}
 	}
-	return NullMaterialEval();
+	return Result;
 }
diff --git a/Engine/Shaders/Private/PathTracing/Material/PathTracingSpecularReflection.ush b/Engine/Shaders/Private/PathTracing/Material/PathTracingSpecularReflection.ush
deleted file mode 100644
diff --git a/Engine/Shaders/Private/PathTracing/Material/PathTracingSpecularTransmission.ush b/Engine/Shaders/Private/PathTracing/Material/PathTracingSpecularTransmission.ush
deleted file mode 100644
diff --git a/Engine/Shaders/Private/PathTracing/Material/PathTracingSubsurfaceProfile.ush b/Engine/Shaders/Private/PathTracing/Material/PathTracingSubsurfaceProfile.ush
--- a/Engine/Shaders/Private/PathTracing/Material/PathTracingSubsurfaceProfile.ush
+++ b/Engine/Shaders/Private/PathTracing/Material/PathTracingSubsurfaceProfile.ush
@@ -11,157 +11,145 @@
 #pragma once
 
 #include "PathTracingMaterialCommon.ush"
-#include "PathTracingFresnel.ush"
 
-float2 SubsurfaceProfile_LobeCdf(float3 DiffuseColor, float3 SpecE, float SpecMix)
+struct FSubsurfaceProfileData {
+	float3x3 Basis;
+	float2 Alpha0;
+	float2 Alpha1;
+
+	float3 V;
+
+	FBxDFEnergyTermsRGB Spec0;
+	FBxDFEnergyTermsRGB Spec1;
+	float3 DiffuseWeight;
+
+	float2 LobeCdf;
+	float3 LobePdf;
+};
+
+FSubsurfaceProfileData CreateSubsurfaceProfileData(FPathTracingPayload Payload, float3 V_World)
 {
-	float3 DiffE = (1 - SpecE) * DiffuseColor;
+	FSubsurfaceProfileData Data = (FSubsurfaceProfileData)0;
+
+	const float3 DualRoughnessData = Payload.GetDualRoughnessSpecular();
+	Data.Basis = GetTangentBasis(Payload.WorldNormal);
+	Data.Alpha0 = GetGGXAlpha(DualRoughnessData.x);
+	Data.Alpha1 = GetGGXAlpha(DualRoughnessData.y);
+
+	Data.V = mul(Data.Basis, V_World);
+	const float NoV = saturate(Data.V.z);
+	Data.Spec0 = ComputeGGXSpecEnergyTermsRGB(DualRoughnessData.x, NoV, Payload.SpecularColor);
+	Data.Spec1 = ComputeGGXSpecEnergyTermsRGB(DualRoughnessData.y, NoV, Payload.SpecularColor);
+
+	const float3 SpecE = lerp(Data.Spec0.E, Data.Spec1.E, DualRoughnessData.z);
+	Data.DiffuseWeight = Payload.DiffuseColor * (1 - SpecE);
+	Data.LobeCdf = LobeSelectionCdf(Data.DiffuseWeight, Data.Spec0.E * (1 - DualRoughnessData.z), Data.Spec1.E * DualRoughnessData.z);
+	Data.LobePdf = LobeSelectionPdf(Data.LobeCdf);
 
-	float DiffuseProb = LobeSelectionProb(DiffE, SpecE);
-	return float2(DiffuseProb, DiffuseProb + (1 - DiffuseProb) * (1 - SpecMix));
+	return Data;
 }
 
 FMaterialEval SubsurfaceProfile_EvalMaterial(
-	float3 IncomingDirection,
-	float3 OutgoingDirection,
-	FPathTracingPayload Payload
+	float3 V_World,
+	float3 L_World,
+	FPathTracingPayload Payload,
+	float2 DiffuseSpecularScale
 )
 {
-	FMaterialEval Result = NullMaterialEval();
-
-	const float3 V_World = -IncomingDirection;
-	const float3 L_World = OutgoingDirection;
-	const float3 N_World = Payload.WorldNormal;
-
-	float3 DualRoughnessData = Payload.GetDualRoughnessSpecular();
-	float2 Alpha0 = 0;
-	float2 Alpha1 = 0;
-	const float3x3 Basis = GetGGXBasis(DualRoughnessData.x, Payload.WorldNormal, Alpha0);
-	GetGGXBasis(DualRoughnessData.y, Payload.WorldNormal, Alpha1);
+	const FSubsurfaceProfileData Data = CreateSubsurfaceProfileData(Payload, V_World);
+	const float3 DualRoughnessData = Payload.GetDualRoughnessSpecular();
 
 	// move vectors into right shading frame
-	const float3 V = mul(Basis, V_World);
-	const float3 L = mul(Basis, L_World);
+	const float3 V = Data.V;
+	const float3 L = mul(Data.Basis, L_World);
 	const float3 H = normalize(V + L);
-
-	const float NoV = saturate(V.z);
-	const FBxDFEnergyTerms Spec0 = ComputeGGXSpecEnergyTerms(DualRoughnessData.x, NoV, Payload.SpecularColor);
-	const FBxDFEnergyTerms Spec1 = ComputeGGXSpecEnergyTerms(DualRoughnessData.y, NoV, Payload.SpecularColor);
-	const float3 SpecE = lerp(Spec0.E, Spec1.E, DualRoughnessData.z);
-	const float2 LobeCdf = SubsurfaceProfile_LobeCdf(Payload.DiffuseColor, SpecE, DualRoughnessData.z);
-	const float3 LobeProb = float3(LobeCdf.x, LobeCdf.y - LobeCdf.x, 1 - LobeCdf.y);
-
-
 	const float NoL = saturate(L.z);
 	const float VoH = saturate(dot(V, H));
 
+	FMaterialEval Result = NullMaterialEval();
+
 	// Diffuse Lobe
-	AddLobeWithMIS(Result.Weight, Result.Pdf, Payload.DiffuseColor * (1 - SpecE) * ShadowTerminatorTerm(L_World, Payload.WorldNormal, Payload.WorldSmoothNormal), NoL / PI, LobeProb.x);
+	Result.AddLobeWithMIS(Data.DiffuseWeight * ShadowTerminatorTerm(L_World, Payload.WorldNormal, Payload.WorldSmoothNormal) * DiffuseSpecularScale.x, NoL / PI, Data.LobePdf.x);
 
 	// Specular lobes
-	const float2 GGXResult0 = GGXEvalReflection(L, V, H, Alpha0);
-	const float2 GGXResult1 = GGXEvalReflection(L, V, H, Alpha1);
-	const float3 F = F_Schlick(Payload.SpecularColor, VoH);
-	AddLobeWithMIS(Result.Weight, Result.Pdf, (1 - DualRoughnessData.z) * F * GGXResult0.x * Spec0.W, GGXResult0.y, LobeProb.y);
-	AddLobeWithMIS(Result.Weight, Result.Pdf, (    DualRoughnessData.z) * F * GGXResult1.x * Spec1.W, GGXResult1.y, LobeProb.z);
+	const float2 GGXResult0 = GGXEvalReflection(L, V, H, Data.Alpha0);
+	const float2 GGXResult1 = GGXEvalReflection(L, V, H, Data.Alpha1);
+	const float3 F = F_Schlick(Payload.SpecularColor, VoH) * DiffuseSpecularScale.y;
+	Result.AddLobeWithMIS((1 - DualRoughnessData.z) * F * GGXResult0.x * Data.Spec0.W, GGXResult0.y, Data.LobePdf.y);
+	Result.AddLobeWithMIS((    DualRoughnessData.z) * F * GGXResult1.x * Data.Spec1.W, GGXResult1.y, Data.LobePdf.z);
 
 	return Result;
 }
 
 FMaterialSample SubsurfaceProfile_SampleMaterial(
-	float3 RayDirection,
+	float3 V_World,
 	FPathTracingPayload Payload,
-	float4 RandSample
+	float3 RandSample
 )
 {
-	float3 N = Payload.WorldNormal;
-	float3 V = -RayDirection;
-
-	float3 DualRoughnessData = Payload.GetDualRoughnessSpecular();
-	float2 Alpha0 = 0;
-	float2 Alpha1 = 0;
-	const float3x3 Basis = GetGGXBasis(DualRoughnessData.x, Payload.WorldNormal, Alpha0);
-	GetGGXBasis(DualRoughnessData.y, Payload.WorldNormal, Alpha1);
-
-	V = mul(Basis, V);
-
-	const float NoV = saturate(V.z);
-	const FBxDFEnergyTerms Spec0 = ComputeGGXSpecEnergyTerms(DualRoughnessData.x, NoV, Payload.SpecularColor);
-	const FBxDFEnergyTerms Spec1 = ComputeGGXSpecEnergyTerms(DualRoughnessData.y, NoV, Payload.SpecularColor);
-	const float3 SpecE = lerp(Spec0.E, Spec1.E, DualRoughnessData.z);
-	const float2 LobeCdf = SubsurfaceProfile_LobeCdf(Payload.DiffuseColor, SpecE, DualRoughnessData.z);
-	const float3 LobeProb = float3(LobeCdf.x, LobeCdf.y - LobeCdf.x, 1 - LobeCdf.y);
+	const FSubsurfaceProfileData Data = CreateSubsurfaceProfileData(Payload, V_World);
+	const float3 V = Data.V;
+	const float3 DualRoughnessData = Payload.GetDualRoughnessSpecular();
 
 	// Randomly choose to sample diffuse or specular
+	const bool bSampledDiffuse = RandSample.x < Data.LobeCdf.x;
 	float3 L = 0, H = 0;
 	float OutRoughness = 0;
-	if (RandSample.x < LobeCdf.x)
+	if (bSampledDiffuse)
 	{
-		RandSample.x /= LobeCdf.x;
+		RandSample.x /= Data.LobeCdf.x;
 
-		// Lambert
-		// TODO: evaluate CosineSampleHemisphereConcentric
-		float4 Result = CosineSampleHemisphere(RandSample.xy);
-
-		L = Result.xyz;
+		L = CosineSampleHemisphere(RandSample.xy).xyz;
 		H = normalize(L + V);
 		OutRoughness = 1.0;
 	}
-	else if (RandSample.x < LobeCdf.y)
+	else
 	{
-		RandSample.x -= LobeCdf.x;
-		RandSample.x /= LobeCdf.y - LobeCdf.x;
-
-
-		// #dxr_todo: Evaluate UniformSampleDiskConcentric
-		H = ImportanceSampleVisibleGGX_aniso(UniformSampleDisk(RandSample.xy), Alpha0, V).xyz;
-
-		L = reflect(-V, H);
-		if (L.z <= 0)
+		const bool bUseSpec0 = RandSample.x < Data.LobeCdf.y;
+		if (bUseSpec0)
 		{
-			// invalid output direction, exit early
-			return NullMaterialSample();
+			RandSample.x -= Data.LobeCdf.x;
+			RandSample.x /= Data.LobeCdf.y - Data.LobeCdf.x;
+			OutRoughness = DualRoughnessData.x;
+		}
+		else
+		{
+			RandSample.x -= Data.LobeCdf.y;
+			RandSample.x /= 1.0 - Data.LobeCdf.y;
+			OutRoughness = DualRoughnessData.y;
 		}
-		OutRoughness = DualRoughnessData.x;
-	}
-	else
-	{
-		RandSample.x -= LobeCdf.y;
-		RandSample.x /= 1 - LobeCdf.y;
 
 
 		// #dxr_todo: Evaluate UniformSampleDiskConcentric
-		H = ImportanceSampleVisibleGGX_aniso(UniformSampleDisk(RandSample.xy), Alpha1, V).xyz;
-
+		H = ImportanceSampleVisibleGGX_aniso(UniformSampleDisk(RandSample.xy), bUseSpec0 ? Data.Alpha0 : Data.Alpha1, V).xyz;
 		L = reflect(-V, H);
 		if (L.z <= 0)
 		{
 			// invalid output direction, exit early
 			return NullMaterialSample();
 		}
-		OutRoughness = DualRoughnessData.y;
 	}
 
-	// With a valid direction in hand -- now evaluate the BxDF (taking advantage of already computed terms)
-
-	const float3 L_World = normalize(mul(L, Basis));
-
-	const float NoL = saturate(L.z);
-	const float VoH = saturate(dot(V, H));
-
-	float3 OutWeight = 0;
-	float OutPdf = 0;
-
-	// Diffuse Lobe
-	AddLobeWithMIS(OutWeight, OutPdf, Payload.DiffuseColor * (1 - SpecE) * ShadowTerminatorTerm(L_World, Payload.WorldNormal, Payload.WorldSmoothNormal), NoL / PI, LobeProb.x);
-
-	// Specular lobe
-	const float2 GGXResult0 = GGXEvalReflection(L, V, H, Alpha0);
-	const float2 GGXResult1 = GGXEvalReflection(L, V, H, Alpha1);
-	const float3 F = F_Schlick(Payload.SpecularColor, VoH);
-	AddLobeWithMIS(OutWeight, OutPdf, (1 - DualRoughnessData.z) * F * GGXResult0.x * Spec0.W, GGXResult0.y, LobeProb.y);
-	AddLobeWithMIS(OutWeight, OutPdf, (    DualRoughnessData.z) * F * GGXResult1.x * Spec1.W, GGXResult1.y, LobeProb.z);
-
 	// transform to world space
-	return CreateMaterialSample(L_World, OutWeight, OutPdf, 1.0, OutRoughness);
+	const float3 L_World = normalize(mul(L, Data.Basis));
+
+	FMaterialSample Result = CreateMaterialSample(L_World, 0.0, 0.0, 1.0, OutRoughness, PATHTRACER_SCATTER_DIFFUSE);
+	if (bSampledDiffuse)
+	{
+		// Diffuse Lobe
+		const float NoL = L.z;
+		Result.AddLobeWithMIS(Data.DiffuseWeight * ShadowTerminatorTerm(L_World, Payload.WorldNormal, Payload.WorldSmoothNormal), NoL / PI, Data.LobePdf.x);
+	}
+	else 
+	{
+		const float VoH = saturate(dot(V, H));
+		// Specular lobes
+		const float2 GGXResult0 = GGXEvalReflection(L, V, H, Data.Alpha0);
+		const float2 GGXResult1 = GGXEvalReflection(L, V, H, Data.Alpha1);
+		const float3 F = F_Schlick(Payload.SpecularColor, VoH);
+		Result.AddLobeWithMIS((1 - DualRoughnessData.z) * F * GGXResult0.x * Data.Spec0.W, GGXResult0.y, Data.LobePdf.y);
+		Result.AddLobeWithMIS((    DualRoughnessData.z) * F * GGXResult1.x * Data.Spec1.W, GGXResult1.y, Data.LobePdf.z);
+		Result.ScatterType = PATHTRACER_SCATTER_SPECULAR;
+	}
+	return Result;
 }
diff --git a/Engine/Shaders/Private/PathTracing/Material/PathTracingThinGlass.ush b/Engine/Shaders/Private/PathTracing/Material/PathTracingThinGlass.ush
--- a/Engine/Shaders/Private/PathTracing/Material/PathTracingThinGlass.ush
+++ b/Engine/Shaders/Private/PathTracing/Material/PathTracingThinGlass.ush
@@ -54,32 +54,74 @@ FThinSlabWeights ComputeThinSlabWeights(float3 SlabColor, float CosV, float Eta,
 	return Result;
 }
 
-FMaterialEval RoughThinGlass_EvalMaterial(
-	float3 IncomingDirection,
-	float3 OutgoingDirection,
-	FPathTracingPayload Payload
-)
+
+struct FRoughThinGlassData
 {
-	const float3 V_World = -IncomingDirection;
-	const float3 L_World = OutgoingDirection;
-	const float3 N_World = Payload.WorldNormal;
+	float F0, F90;
+	float3x3 Basis;
+	float2 AlphaR, AlphaT;
+	float3 V;
+
+	FBxDFEnergyTermsRGB Spec;
+	FBxDFEnergyTermsA SpecR;
+	FBxDFEnergyTermsA SpecT;
+	float NudgeE;
+	float3 DiffWeight;
+
+	float3 LobeCdf;
+	float4 LobePdf;
+};
 
-	const float Eta = Payload.Ior;
-	const float F0 = F0RGBToF0(Payload.SpecularColor);
-	const float F90 = saturate(50.0 * F0);
-	const float3 SlabColor = Payload.GetTransmittanceColor();
+FRoughThinGlassData PrepareRoughThinGlassData(FPathTracingPayload Payload, float3 V_World)
+{
+	FRoughThinGlassData Data = (FRoughThinGlassData)0;
+
+	Data.F0 = F0RGBToF0(Payload.SpecularColor);
+	Data.F90 = saturate(50.0 * Data.F0);
 
 	// transmitted lobe should have slightly less roughness due to the squeeze implied by the IOR (eventually converging to 0.0 as Ior approaches 1.0)
 	const float RoughnessR = Payload.Roughness;
-	const float RoughnessT = Payload.Roughness * sqrt(1.0 - min(Eta, rcp(Eta)));
+	const float RoughnessT = Payload.Roughness * sqrt(1.0 - min(Payload.Ior, rcp(Payload.Ior)));
+
+	Data.Basis = GetGGXBasis(RoughnessR, Payload.Anisotropy, Payload.WorldNormal, Payload.WorldTangent, Data.AlphaR);
+	Data.AlphaT = GetGGXAlpha(RoughnessT, Payload.Anisotropy);
+
+	Data.V = mul(Data.Basis, V_World);
+	const float NoV = saturate(Data.V.z);
+
+	Data.Spec = ComputeGGXSpecEnergyTermsRGB(Payload.Roughness, NoV, Payload.SpecularColor);
+	// because the roughnesses may not match, we need to measure the response for each individually, taking into account the reflected and transmitted fresnel respectively
+	// NOTE: 1.0 - F_Schlick(NoV, F0, F90) == F_Schlick(NoV, 1.0 - F0, 1.0 - F90)
+	Data.SpecR = ComputeGGXSpecEnergyTermsA(RoughnessR, NoV, Data.F0, Data.F90);
+	Data.SpecT = ComputeGGXSpecEnergyTermsA(RoughnessT, NoV, 1.0 - Data.F0, 1.0 - Data.F90);
+	Data.NudgeE = rcp(Data.SpecR.E + Data.SpecT.E);
+	Data.DiffWeight = (1 - Data.Spec.E) * Payload.DiffuseColor;
+
+	// Figure out probability of the 4 lobes: diffuse,spec,glass reflection, glass transmission (if enabled)
+	// Approximate the transmission as just a plain tint by slab color because we can't depend on the half-vector for lobe selection
+	Data.LobeCdf = LobeSelectionCdf(
+		Payload.Opacity * Data.DiffWeight,
+		Payload.Opacity * Data.Spec.E,
+		(1 - Payload.Opacity) * Data.NudgeE * Data.SpecR.E,
+		(1 - Payload.Opacity) * Data.NudgeE * (Payload.HasRefraction() ? Data.SpecT.E * Payload.GetTransmittanceColor() : 0.0)
+	);
+	Data.LobePdf = LobeSelectionPdf(Data.LobeCdf);
+
+	return Data;
+}
 
-	float2 AlphaR;
-	const float3x3 Basis = GetGGXBasis(RoughnessR, Payload.Anisotropy, Payload.WorldNormal, Payload.WorldTangent, AlphaR);
-	const float2 AlphaT = GetGGXAlpha(RoughnessT, Payload.Anisotropy);
+FMaterialEval RoughThinGlass_EvalMaterial(
+	float3 V_World,
+	float3 L_World,
+	FPathTracingPayload Payload,
+	float2 DiffuseSpecularScale
+)
+{
+	const FRoughThinGlassData Data = PrepareRoughThinGlassData(Payload, V_World);
 
 	// move vectors into right shading frame
-	float3 V = mul(Basis, V_World);
-	float3 L = mul(Basis, L_World);
+	float3 V = Data.V;
+	float3 L = mul(Data.Basis, L_World);
 
 	if (V.z <= 0)
 	{
@@ -90,69 +132,55 @@ FMaterialEval RoughThinGlass_EvalMaterial(
 	const bool bIsReflection = L.z >= 0;
 	L.z = abs(L.z); // push L to the same side as V
 
+	const float NoL = saturate(L.z);
 	const float NoV = saturate(V.z);
 	const float3 H = normalize(L + V);
 	const float VoH = saturate(dot(V, H));
+	const float NoH = saturate(H.z);
 
 
-	// because the roughnesses may not match, we need to measure the response for each individually, taking into account the reflected and transmitted fresnel respectively
-#if 1
-	// NOTE: 1.0 - F_Schlick(NoV, F0, F90) == F_Schlick(NoV, 1.0 - F0, 1.0 - F90)
-	const FBxDFEnergyTerms SpecR = ComputeGGXSpecEnergyTerms(RoughnessR, NoV, F0, F90);
-	const FBxDFEnergyTerms SpecT = ComputeGGXSpecEnergyTerms(RoughnessT, NoV, 1.0 - F0, 1.0 - F90);
-	// Approximate the transmission as just a plain tint by slab color because we can't depend on the half-vector for lobe selection
-	const float NudgeE = rcp(SpecR.E.x + SpecT.E.x);
-	const float ProbR = LobeSelectionProb(SpecR.E, Payload.HasRefraction() ? SpecT.E * SlabColor : 0.0);
-#else
-	const float Er = GGXEnergyLookup(RoughnessR, NoV).x;
-	const float Et = GGXEnergyLookup(RoughnessT, NoV).x;
-	const float ProbR = LobeSelectionProb(F0 * Er, Payload.HasRefraction() ? (1 - F0) * Et * SlabColor : 0.0);
-#endif
-
-	const FThinSlabWeights Result = ComputeThinSlabWeights(SlabColor, VoH, Payload.Ior, F0);
-	const float2 GGXResult = GGXEvalReflection(L, V, H, bIsReflection ? AlphaR : AlphaT);
+	const FThinSlabWeights SlabResult = ComputeThinSlabWeights(Payload.GetTransmittanceColor(), VoH, Payload.Ior, Data.F0);
+	const float2 GGXResult = GGXEvalReflection(L, V, H, bIsReflection ? Data.AlphaR : Data.AlphaT);
+	FMaterialEval Result = NullMaterialEval();
 	if (bIsReflection)
 	{
-		if (ProbR > 0.0)
-		{
-			const float OutWeight = NudgeE * GGXResult.x * SpecR.W.x;
-			return CreateMaterialEval(OutWeight / ProbR * Result.Reflected, GGXResult.y * ProbR);
-		}
+		// Diffuse Lobe
+		const float3 Diffuse = GetPathTracingDiffuseModel(Data.DiffWeight, Payload.Roughness, NoV, NoL, VoH, NoH);
+		const float3 DiffWeight = Payload.Opacity * Diffuse * ShadowTerminatorTerm(L_World, Payload.WorldNormal, Payload.WorldSmoothNormal);
+		const float DiffPdf = NoL / PI;
+		Result.AddLobeWithMIS(DiffWeight * DiffuseSpecularScale.x, DiffPdf, Data.LobePdf.x);
+
+		// Specular lobe
+		const float2 GGXResult = GGXEvalReflection(L, V, H, Data.AlphaR);
+		const float3 F = F_Schlick(Payload.SpecularColor, VoH);
+		const float3 SpecWeight = Payload.Opacity * F * GGXResult.x * Data.Spec.W;
+		const float SpecPdf = GGXResult.y;
+		Result.AddLobeWithMIS(SpecWeight * DiffuseSpecularScale.y, SpecPdf, Data.LobePdf.y);
+
+
+		// Transmission reflection
+		const float3 GlassWeight = ((1 - Payload.Opacity) * Data.NudgeE * GGXResult.x * Data.SpecR.W) * SlabResult.Reflected;
+		const float GlassPdf = GGXResult.y;
+		Result.AddLobeWithMIS(GlassWeight * DiffuseSpecularScale.y, GlassPdf, Data.LobePdf.z);
 	}
 	else
 	{
-		if (ProbR < 1.0)
-		{
-			const float OutWeight = NudgeE * GGXResult.x * SpecT.W.x;
-			return CreateMaterialEval(OutWeight / (1.0 - ProbR) * Result.Transmitted, GGXResult.y * (1.0 - ProbR));
-		}
+		const float3 GlassWeight = ((1 - Payload.Opacity) * Data.NudgeE * GGXResult.x * Data.SpecT.W) * SlabResult.Transmitted;
+		const float GlassPdf = GGXResult.y;
+		Result.AddLobeWithMIS(GlassWeight * DiffuseSpecularScale.y, GlassPdf, Data.LobePdf.w);
 	}
-	return NullMaterialEval();
+	return Result;
 }
 
 
 FMaterialSample RoughThinGlass_SampleMaterial(
-	float3 RayDirection,
+	float3 V_World,
 	FPathTracingPayload Payload,
-	float4 RandSample)
+	float3 RandSample)
 {
-	const float3 N_World = Payload.WorldNormal;
-	const float3 V_World = -RayDirection;
-
-	const float Eta = Payload.Ior;
-	const float F0 = F0RGBToF0(Payload.SpecularColor);
-	const float F90 = saturate(50.0 * F0);
-	const float3 SlabColor = Payload.GetTransmittanceColor();
-
-	// transmitted lobe should have slightly less roughness due to the squeeze implied by the IOR (eventually converging to 0.0 as Ior approaches 1.0)
-	const float RoughnessR = Payload.Roughness;
-	const float RoughnessT = Payload.Roughness * sqrt(1.0 - min(Eta, rcp(Eta)));
-
-	float2 AlphaR;
-	const float3x3 Basis = GetGGXBasis(RoughnessR, Payload.Anisotropy, Payload.WorldNormal, Payload.WorldTangent, AlphaR);
-	const float2 AlphaT = GetGGXAlpha(RoughnessT, Payload.Anisotropy);
+	const FRoughThinGlassData Data = PrepareRoughThinGlassData(Payload, V_World);
 
-	const float3 V = mul(Basis, V_World);
+	const float3 V = Data.V;
 	const float NoV = saturate(V.z);
 
 	if (NoV == 0)
@@ -161,55 +189,96 @@ FMaterialSample RoughThinGlass_SampleMaterial(
 		return NullMaterialSample();
 	}
 
-	// because the roughnesses may not match, we need to measure the response for each individually, taking into account the reflected and transmitted fresnel respectively
-#if 1
-	// NOTE: 1.0 - F_Schlick(NoV, F0, F90) == F_Schlick(NoV, 1.0 - F0, 1.0 - F90)
-	const FBxDFEnergyTerms SpecR = ComputeGGXSpecEnergyTerms(RoughnessR, NoV, F0, F90);
-	const FBxDFEnergyTerms SpecT = ComputeGGXSpecEnergyTerms(RoughnessT, NoV, 1.0 - F0, 1.0 - F90);
-	// Approximate the transmission as just a plain tint by slab color because we can't depend on the half-vector for lobe selection
-	const float NudgeE = rcp(SpecR.E.x + SpecT.E.x);
-	const float ProbR = LobeSelectionProb(SpecR.E, Payload.HasRefraction() ? SpecT.E * SlabColor : 0.0);
-#else
-	const float Er = GGXEnergyLookup(RoughnessR, NoV).x;
-	const float Et = GGXEnergyLookup(RoughnessT, NoV).x;
-	const float ProbR = LobeSelectionProb(F0 * Er, Payload.HasRefraction() ? (1 - F0) * Et * SlabColor : 0.0);
-#endif
-
-	const bool bIsReflection = RandSample.z < ProbR;
-
-	const float3 H = ImportanceSampleVisibleGGX_aniso(UniformSampleDisk(RandSample.xy), bIsReflection ? AlphaR : AlphaT, V).xyz;
-	const float NoH = saturate(H.z);
-	const float VoH = saturate(dot(V, H));
-
-	const float3 L = reflect(-V, H);
-	if (L.z <= 0)
+	if (RandSample.x < Data.LobeCdf.x)
 	{
-		// invalid output direction, exit early
-		return NullMaterialSample();
-	}
+		// Diffuse lobe
+		const float3 L = CosineSampleHemisphere(RandSample.xy).xyz;
+		const float3 H = normalize(L + V);
 
-	const float NoL = saturate(L.z);
+		// transform to world space
+		const float3 L_World = normalize(mul(L, Data.Basis));
 
-	FThinSlabWeights Result = ComputeThinSlabWeights(SlabColor, VoH, Payload.Ior, F0);
+		const float NoL = saturate(L.z);
+		const float VoH = saturate(dot(V, H));
+		const float NoH = saturate(H.z);
 
-	const float2 GGXResult = GGXEvalReflection(L, V, H, bIsReflection ? AlphaR : AlphaT);
+		// Diffuse Lobe
+		const float3 Diffuse = GetPathTracingDiffuseModel(Data.DiffWeight, Payload.Roughness, NoV, NoL, VoH, NoH);
+		const float3 DiffWeight = Payload.Opacity * Diffuse * ShadowTerminatorTerm(L_World, Payload.WorldNormal, Payload.WorldSmoothNormal);
+		const float DiffPdf = NoL / PI;
 
-	float3 OutWeight = GGXResult.x;
-	float OutPdf = GGXResult.y;
-	if (bIsReflection)
-	{
-		// reflected ray
-		OutWeight *= NudgeE * SpecR.W.x * Result.Reflected / ProbR;
-		OutPdf *= ProbR;
-		float3 OutDirection = normalize(mul(L, Basis));
-		return CreateMaterialSample(OutDirection, OutWeight, OutPdf, 1.0, Payload.Roughness);
+		FMaterialSample Result = CreateMaterialSample(L_World, 0.0, 0.0, 1.0, 1.0, PATHTRACER_SCATTER_DIFFUSE);
+		Result.AddLobeWithMIS(DiffWeight, DiffPdf, Data.LobePdf.x);
+		return Result;
 	}
 	else
 	{
-		// transmitted ray
-		OutWeight *= NudgeE * SpecT.W.x * Result.Transmitted / (1.0 - ProbR);
-		OutPdf *= 1.0 - ProbR;
-		float3 OutDirection = normalize(mul(float3(L.xy, -L.z), Basis)); // flip reflection to other side
-		return CreateMaterialSample(OutDirection, OutWeight, OutPdf, -1.0, Payload.Roughness);
+		// Specular/Glass lobes
+		const bool bIsReflection = RandSample.x < Data.LobeCdf.z;
+		if (RandSample.x < Data.LobeCdf.y)
+		{
+			RandSample.x -= Data.LobeCdf.x;
+			RandSample.x /= Data.LobeCdf.y - Data.LobeCdf.x;
+		}
+		else if (RandSample.x < Data.LobeCdf.z)
+		{
+			RandSample.x -= Data.LobeCdf.y;
+			RandSample.x /= Data.LobeCdf.z - Data.LobeCdf.y;
+		}
+		else
+		{
+			RandSample.x -= Data.LobeCdf.z;
+			RandSample.x /= 1.0 - Data.LobeCdf.z;
+		}
+		const float3 H = ImportanceSampleVisibleGGX_aniso(UniformSampleDisk(RandSample.xy), bIsReflection ? Data.AlphaR : Data.AlphaT, V).xyz;
+		const float NoH = saturate(H.z);
+		const float VoH = saturate(dot(V, H));
+
+		const float3 L = reflect(-V, H);
+		if (L.z <= 0)
+		{
+			// invalid output direction, exit early
+			return NullMaterialSample();
+		}
+
+		const float NoL = saturate(L.z);
+
+		FThinSlabWeights SlabResult = ComputeThinSlabWeights(Payload.GetTransmittanceColor(), VoH, Payload.Ior, Data.F0);
+
+		const float2 GGXResult = GGXEvalReflection(L, V, H, bIsReflection ? Data.AlphaR : Data.AlphaT);
+
+		FMaterialSample Result = NullMaterialSample();
+		Result.Roughness = Payload.Roughness; // use common roughness for all lobes, even though transmission is squeezed
+		if (bIsReflection)
+		{
+			// Specular lobe
+			const float3 F = F_Schlick(Payload.SpecularColor, VoH);
+			const float3 SpecWeight = Payload.Opacity * F * GGXResult.x * Data.Spec.W;
+			const float SpecPdf = GGXResult.y;
+			Result.AddLobeWithMIS(SpecWeight, SpecPdf, Data.LobePdf.y);
+
+
+			// Transmission reflection
+			const float3 GlassWeight = ((1 - Payload.Opacity) * Data.NudgeE * GGXResult.x * Data.SpecR.W) * SlabResult.Reflected;
+			const float GlassPdf = GGXResult.y;
+			Result.AddLobeWithMIS(GlassWeight, GlassPdf, Data.LobePdf.z);
+
+			Result.Direction = mul(L, Data.Basis);
+			Result.PositionBiasSign = 1.0;
+			Result.ScatterType = PATHTRACER_SCATTER_SPECULAR;
+		}
+		else
+		{
+			// transmitted ray
+			const float3 GlassWeight = ((1 - Payload.Opacity) * Data.NudgeE * GGXResult.x * Data.SpecT.W) * SlabResult.Transmitted;
+			const float GlassPdf = GGXResult.y;
+			Result.AddLobeWithMIS(GlassWeight, GlassPdf, Data.LobePdf.w);
+
+			Result.Direction = mul(float3(L.xy, -L.z), Data.Basis); // flip reflection to other side
+			Result.PositionBiasSign = -1.0;
+			Result.ScatterType = PATHTRACER_SCATTER_REFRACT;
+		}
+		Result.Direction = normalize(Result.Direction);
+		return Result;
 	}
 }
diff --git a/Engine/Shaders/Private/PathTracing/Material/PathTracingTwoSidedFoliage.ush b/Engine/Shaders/Private/PathTracing/Material/PathTracingTwoSidedFoliage.ush
--- a/Engine/Shaders/Private/PathTracing/Material/PathTracingTwoSidedFoliage.ush
+++ b/Engine/Shaders/Private/PathTracing/Material/PathTracingTwoSidedFoliage.ush
@@ -9,74 +9,84 @@
 #include "PathTracingDefaultLit.ush"
 #include "PathTracingEnergyConservation.ush"
 
-FMaterialSample TwoSidedFoliage_SampleMaterial(
-	float3 RayDirection,
-	FPathTracingPayload Payload,
-	float4 RandSample
-)
+struct FTwoSidedFoliageData
+{
+	float NoV;
+	FBxDFEnergyTermsRGB Spec;
+	float3 Back;
+
+	float ProbT;
+};
+
+FTwoSidedFoliageData PrepareTwoSidedFoliageData(FPathTracingPayload Payload, float3 V_World)
 {
-	const float NoV = saturate(-dot(RayDirection, Payload.WorldNormal));
-	const FBxDFEnergyTerms Spec = ComputeGGXSpecEnergyTerms(Payload.Roughness, NoV, Payload.SpecularColor);
+	FTwoSidedFoliageData Data = (FTwoSidedFoliageData)0;
+	Data.NoV = saturate(dot(V_World, Payload.WorldNormal));
+	Data.Spec = ComputeGGXSpecEnergyTermsRGB(Payload.Roughness, Data.NoV, Payload.SpecularColor);
 
 	// pick back lobe proportionally to its contribution relative to front side
-	const float3 Diff = (1 - Spec.E) * Payload.DiffuseColor;
-	const float3 Back = (1 - Spec.E) * Payload.SubsurfaceColor;
-	float TransmissionSelectionPdf = LobeSelectionProb(Back, Diff + Spec.E);
+	const float3 Diff = (1 - Data.Spec.E) * Payload.DiffuseColor;
+	Data.Back         = (1 - Data.Spec.E) * Payload.SubsurfaceColor;
+	Data.ProbT = LobeSelectionProb(Data.Back, Diff + Data.Spec.E);
 
-	if (RandSample.x < TransmissionSelectionPdf)
-	{
-		RandSample.x /= TransmissionSelectionPdf;
+	return Data;
+}
 
-		float3 N = Payload.WorldNormal;
+FMaterialSample TwoSidedFoliage_SampleMaterial(
+	float3 V_World,
+	FPathTracingPayload Payload,
+	float3 RandSample
+)
+{
+	const FTwoSidedFoliageData Data = PrepareTwoSidedFoliageData(Payload, V_World);
+	if (RandSample.x < Data.ProbT)
+	{
+		RandSample.x /= Data.ProbT;
 
-		float4 SampledValue = CosineSampleHemisphere(RandSample.yz);
+		const float3 N = Payload.WorldNormal;
+		const float4 SampledValue = CosineSampleHemisphere(RandSample.yz);
 
-		return CreateMaterialSample(TangentToWorld(SampledValue.xyz, -N), Back / TransmissionSelectionPdf, TransmissionSelectionPdf * SampledValue.w, -1.0, 1.0);
+		return CreateMaterialSample(TangentToWorld(SampledValue.xyz, -N), Payload.GetBSDFOpacity() * Data.Back / Data.ProbT, Data.ProbT * SampledValue.w, -1.0, 1.0, PATHTRACER_SCATTER_DIFFUSE);
 	}
 	else
 	{
-		RandSample.x -= TransmissionSelectionPdf;
-		RandSample.x /= 1.0 - TransmissionSelectionPdf;
+		RandSample.x -= Data.ProbT;
+		RandSample.x /= 1.0 - Data.ProbT;
 
-		FMaterialSample Result = DefaultLit_SampleMaterial(RayDirection, Payload, RandSample);
-		Result.Weight /= 1.0 - TransmissionSelectionPdf;
-		Result.Pdf *= 1.0 - TransmissionSelectionPdf;
+		FMaterialSample Result = DefaultLit_SampleMaterial(V_World, Payload, RandSample);
+		Result.Weight /= 1.0 - Data.ProbT;
+		Result.Pdf *= 1.0 - Data.ProbT;
 		return Result;
 	}
 }
 
 FMaterialEval TwoSidedFoliage_EvalMaterial(
-	float3 IncomingDirection,
-	float3 OutgoingDirection,
-	FPathTracingPayload Payload
+	float3 V_World,
+	float3 L_World,
+	FPathTracingPayload Payload,
+	float2 DiffuseSpecularScale
 )
 {
-	const float NoV = saturate(-dot(IncomingDirection, Payload.WorldNormal));
-	const FBxDFEnergyTerms Spec = ComputeGGXSpecEnergyTerms(Payload.Roughness, NoV, Payload.SpecularColor);
-
-	// pick back lobe proportionally to its contribution relative to front side
-	const float3 Diff = (1 - Spec.E) * Payload.DiffuseColor;
-	const float3 Back = (1 - Spec.E) * Payload.SubsurfaceColor;
-	float TransmissionSelectionPdf = LobeSelectionProb(Back, Diff + Spec.E);
+	const FTwoSidedFoliageData Data = PrepareTwoSidedFoliageData(Payload, V_World);
 
-	float3 N = Payload.WorldNormal;
-	float NoL = dot(N, OutgoingDirection);
+	const float3 N = Payload.WorldNormal;
+	const float NoL = dot(N, L_World);
 	if (NoL < 0.0)
 	{
 		// Diffuse transmission (the implementation in TwoSidedBxDF does not appear to be a real BxDF)
 
-		if (TransmissionSelectionPdf > 0)
+		if (Data.ProbT > 0)
 		{
 			float LambertPdf = -NoL / PI;
 
-			return CreateMaterialEval(Back / TransmissionSelectionPdf, TransmissionSelectionPdf * LambertPdf);
+			return CreateMaterialEval(Payload.GetBSDFOpacity() * Data.Back / Data.ProbT * DiffuseSpecularScale.x, Data.ProbT * LambertPdf);
 		}
 	}
-	else if (TransmissionSelectionPdf < 1)
+	else if (Data.ProbT < 1)
 	{
-		FMaterialEval Result = DefaultLit_EvalMaterial(IncomingDirection, OutgoingDirection, Payload);
-		Result.Weight /= 1.0 - TransmissionSelectionPdf;
-		Result.Pdf *= 1.0 - TransmissionSelectionPdf;
+		FMaterialEval Result = DefaultLit_EvalMaterial(V_World, L_World, Payload, DiffuseSpecularScale);
+		Result.Weight /= 1.0 - Data.ProbT;
+		Result.Pdf *= 1.0 - Data.ProbT;
 		return Result;
 	}
 	return NullMaterialEval();
diff --git a/Engine/Shaders/Private/PathTracing/PathTracing.usf b/Engine/Shaders/Private/PathTracing/PathTracing.usf
--- a/Engine/Shaders/Private/PathTracing/PathTracing.usf
+++ b/Engine/Shaders/Private/PathTracing/PathTracing.usf
@@ -63,7 +63,13 @@ FPathState LoadPathStateData(uint Index)
 	Output.SigmaT.x = f16tof32(Data.PackedRoughnessSigma.x >> 16);
 	Output.SigmaT.y = f16tof32(Data.PackedRoughnessSigma.y);
 	Output.SigmaT.z = f16tof32(Data.PackedRoughnessSigma.y >> 16);
-
+	// extract scatter type from sign bits of SigmaT
+	Output.FirstScatterType = 0;
+	Output.FirstScatterType |= asuint(Output.SigmaT.x) >> 31;
+	Output.FirstScatterType |= asuint(Output.SigmaT.y) >> 30;
+	Output.FirstScatterType |= asuint(Output.SigmaT.z) >> 29;
+	// remove sign bits from SigmaT (must be a positive value)
+	Output.SigmaT = abs(Output.SigmaT);
 	return Output;
 }
 
@@ -80,8 +86,13 @@ void StorePathStateData(FPathState PathState, uint Index)
 	Data.RayOrigin = PathState.Ray.Origin;
 	Data.RayDirection = PathState.Ray.Direction;
 	Data.PathThroughput = PathState.PathThroughput;
-	Data.PackedRoughnessSigma.x = f32tof16(PathState.PathRoughness) | (f32tof16(PathState.SigmaT.x) << 16);
-	Data.PackedRoughnessSigma.y = f32tof16(PathState.SigmaT.y) | (f32tof16(PathState.SigmaT.z) << 16);
+
+	float3 ScatterAndSigmaT = abs(PathState.SigmaT);
+	ScatterAndSigmaT.x = asfloat(((PathState.FirstScatterType & 1) << 31) | asuint(ScatterAndSigmaT.x));
+	ScatterAndSigmaT.y = asfloat(((PathState.FirstScatterType & 2) << 30) | asuint(ScatterAndSigmaT.y));
+	ScatterAndSigmaT.z = asfloat(((PathState.FirstScatterType & 4) << 29) | asuint(ScatterAndSigmaT.z));
+	Data.PackedRoughnessSigma.x = f32tof16(PathState.PathRoughness) | (f32tof16(ScatterAndSigmaT.x) << 16);
+	Data.PackedRoughnessSigma.y = f32tof16(ScatterAndSigmaT.y) | (f32tof16(ScatterAndSigmaT.z) << 16);
 	PathStateData[Index] = Data;
 }
 
diff --git a/Engine/Shaders/Private/PathTracing/PathTracingCommon.ush b/Engine/Shaders/Private/PathTracing/PathTracingCommon.ush
--- a/Engine/Shaders/Private/PathTracing/PathTracingCommon.ush
+++ b/Engine/Shaders/Private/PathTracing/PathTracingCommon.ush
@@ -224,6 +224,14 @@ struct FPathTracingPayload : FMinimalPayload
 	void SetHG(float3 Weight, float G) { CustomData0 = float4(Weight, G); }
 
 
+	// Methods used by the integrator
+
+	float GetBSDFOpacity()
+	{
+		// How much should the BSDF components by weighted?
+		return BlendingMode == RAY_TRACING_BLEND_MODE_ALPHA_COMPOSITE ? 1.0 : Opacity;
+	}
+
 	bool HasRefraction()
 	{
 		return Ior > 0.0;
diff --git a/Engine/Shaders/Private/PathTracing/PathTracingCore.ush b/Engine/Shaders/Private/PathTracing/PathTracingCore.ush
--- a/Engine/Shaders/Private/PathTracing/PathTracingCore.ush
+++ b/Engine/Shaders/Private/PathTracing/PathTracingCore.ush
@@ -52,6 +52,8 @@ uint EnableEmissive;
 uint SamplerType;
 uint VisualizeLightGrid;
 uint VisualizeDecalGrid;
+uint ApplyDiffuseSpecularOverrides;
+uint EnabledLightingContributions;
 float DecalRoughnessCutoff;
 float MeshDecalRoughnessCutoff;
 float MeshDecalBias;
@@ -83,8 +85,84 @@ struct FPathState
 	float3 PathThroughput;          // 12 bytes
 	float PathRoughness;            //  2 bytes // half precision
 	float3 SigmaT; // Extinction    //  6 bytes // half precision 
+	uint FirstScatterType;          //  3 bits - packed in sign bits of SigmaT
 									// 80 bytes total (see FPathTracingPackedPathState)
 
+	bool HasMadeContributionScatter()
+	{
+		// The first scattering type that matters for the PATHTRACER_CONTRIBUTION_* flags is the first non-refractive one
+		// We choose to track camera and pure refracted paths together so that objects behind glass can be treated the same as directly visible objects.
+		return !(FirstScatterType == PATHTRACER_SCATTER_CAMERA ||
+			     FirstScatterType == PATHTRACER_SCATTER_REFRACT);
+	}
+
+	bool UpdateScatterType(uint ScatterType)
+	{
+		// Keep track of the first non-refractive scatter type, and leave it "locked" beyond this
+		// This ensures a particular pixel can only contribute to one lighting component.
+		// The reason to not split refractive events is that you typically want to treat refracted paths the same as camera paths.
+		// For example, when extracting just the diffuse component of a character wearing glasses, you would want the directly visible diffuse
+		// and refracted diffuse to be treated the same.
+		if (!HasMadeContributionScatter())
+		{
+			FirstScatterType = ScatterType;
+		}
+		switch (FirstScatterType)
+		{
+			// In these cases, only keep tracing if the corresponding lighting component is enabled
+			case PATHTRACER_SCATTER_DIFFUSE:  return (EnabledLightingContributions & PATHTRACER_CONTRIBUTION_DIFFUSE ) != 0;
+			case PATHTRACER_SCATTER_SPECULAR: return (EnabledLightingContributions & PATHTRACER_CONTRIBUTION_SPECULAR) != 0;
+			case PATHTRACER_SCATTER_VOLUME:   return (EnabledLightingContributions & PATHTRACER_CONTRIBUTION_VOLUME  ) != 0;
+		}
+		// In all other cases, keep tracing, we haven't made a final decision on what this path is yet
+		return true;
+	}
+
+	// The following functions allow us to scale down the contributions for lighting components before we have scattered
+	bool ShouldAccumulateEmissive()
+	{
+		return (HasMadeContributionScatter() || (EnabledLightingContributions & PATHTRACER_CONTRIBUTION_EMISSIVE) != 0);
+	}
+
+	bool ShouldAccumulateDiffuse()
+	{
+		return (HasMadeContributionScatter() || (EnabledLightingContributions & PATHTRACER_CONTRIBUTION_DIFFUSE) != 0);
+	}
+
+	bool ShouldAccumulateSpecular()
+	{
+		return (HasMadeContributionScatter() || (EnabledLightingContributions & PATHTRACER_CONTRIBUTION_SPECULAR) != 0);
+	}
+
+	bool ShouldAccumulateVolume()
+	{
+		return (HasMadeContributionScatter() || (EnabledLightingContributions & PATHTRACER_CONTRIBUTION_VOLUME) != 0);
+	}
+
+	float2 GetDiffuseSpecularScale(bool bIsVolumeSample)
+	{
+		float2 Result = 0.0;
+		if (bIsVolumeSample)
+		{
+			if (ShouldAccumulateVolume())
+			{
+				Result.x = 1.0;
+			}
+		}
+		else
+		{
+			if (ShouldAccumulateDiffuse())
+			{
+				Result.x = 1.0;
+			}
+			if (ShouldAccumulateSpecular())
+			{
+				Result.y = 1.0;
+			}
+		}
+		return Result;
+	}
+
 	void WritePixel(uint2 TextureIndex)
 	{
 		// Avoid reading the old pixel on the first sample on the off-chance there is a NaN/Inf pixel ...
@@ -365,7 +443,7 @@ FPathTracingPayload TraceTransparentRay(inout FPathState PathState, bool bIsCame
 		}
 
 		// Unpack the payload
-		FPathTracingPayload HitPayload = UnpackPathTracingPayloadMaterial(PackedPayload, PathState.Ray, bIsCameraRay);
+		FPathTracingPayload HitPayload = UnpackPathTracingPayloadMaterial(PackedPayload, PathState.Ray, bIsCameraRay && ApplyDiffuseSpecularOverrides);
 
 		float DecalThroughput = 1.0f;
 		if (PathState.PathRoughness < MeshDecalRoughnessCutoff && HitPayload.IsDecalReceiver())
@@ -672,14 +750,14 @@ bool ProcessSubsurfaceRandomWalk(inout FPathTracingPayload Payload, inout float3
 	{
 		// dual spec
 		const float3 DualRoughnessData = Payload.GetDualRoughnessSpecular();
-		const FBxDFEnergyTerms Spec0 = ComputeGGXSpecEnergyTerms(DualRoughnessData.x, NoV, Payload.SpecularColor);
-		const FBxDFEnergyTerms Spec1 = ComputeGGXSpecEnergyTerms(DualRoughnessData.y, NoV, Payload.SpecularColor);
-		SpecE = lerp(Spec0.E, Spec1.E, DualRoughnessData.z);
+		const float3 SpecE0 = ComputeGGXSpecEnergyTermsRGB(DualRoughnessData.x, NoV, Payload.SpecularColor).E;
+		const float3 SpecE1 = ComputeGGXSpecEnergyTermsRGB(DualRoughnessData.y, NoV, Payload.SpecularColor).E;
+		SpecE = lerp(SpecE0, SpecE1, DualRoughnessData.z);
 	}
 	else
 	{
 		// single spec
-		SpecE = ComputeGGXSpecEnergyTerms(Payload.Roughness, NoV, Payload.SpecularColor).E;
+		SpecE = ComputeGGXSpecEnergyTermsRGB(Payload.Roughness, NoV, Payload.SpecularColor).E;
 	}
 
 	const float3 SSSLobeAlbedo = (1 - SpecE) * SSSColor;
@@ -876,6 +954,7 @@ FPathState CreatePathState(int2 PixelIndex)
 	Output.PathThroughput = 1.0;
 	Output.PathRoughness = 0;
 	Output.SigmaT = float3(StartingExtinctionCoefficient[0], StartingExtinctionCoefficient[1], StartingExtinctionCoefficient[2]);
+	Output.FirstScatterType = PATHTRACER_SCATTER_CAMERA;
 
 	return Output;
 }
@@ -892,8 +971,10 @@ bool PathTracingKernel(inout FPathState PathState, int Bounce)
 #endif
 	const bool bIsCameraRay = Bounce == 0;
 	const bool bIsLastBounce = Bounce == MaxBounces;
-	const bool bIncludeEmissive = (EnableDirectLighting != 0 || Bounce > 1) &&
-		(EnableEmissive != 0 || bIsCameraRay);
+	const bool bIncludeEmissive =
+		(EnableDirectLighting != 0 || Bounce > 1) &&
+		(EnableEmissive != 0 || bIsCameraRay) &&
+		PathState.ShouldAccumulateEmissive();
 
 	// number of directly visible lights for the first bounce
 	// after the first bounce, we don't need to include lights in the trace call
@@ -910,7 +991,7 @@ bool PathTracingKernel(inout FPathState PathState, int Bounce)
 		VolumeSegment);
 
 	// process the returned volume segment if we got one
-	if (VolumeSegment.IsValid() && VolumeMISMode != 0)
+	if (VolumeSegment.IsValid() && VolumeMISMode != 0 && PathState.ShouldAccumulateVolume())
 	{
 		const float3 Ro = PathState.Ray.Origin;
 		const float3 Rd = PathState.Ray.Direction;
@@ -965,7 +1046,7 @@ bool PathTracingKernel(inout FPathState PathState, int Bounce)
 
 				// prepare a minimal payload that describes the hit we need to shade
 				FPathTracingPayload VolPayload = CreateMediumHitPayload(EquiAngularT, TranslatedWorldPos, Result);
-				float4 LightRandValue = RandomSequence_GenerateSample4D(PathState.RandSequence);
+				float3 LightRandValue = RandomSequence_GenerateSample3D(PathState.RandSequence);
 
 				const bool bCastShadows = CastsVolumeShadow(LightSampleLightId);
 				const uint MissShaderIndex = GetLightMissShaderIndex(LightSampleLightId);
@@ -985,7 +1066,7 @@ bool PathTracingKernel(inout FPathState PathState, int Bounce)
 						if (any(LightSample.RadianceOverPdf > 0))
 						{
 							// Evaluate material
-							FMaterialEval MaterialEval = Medium_EvalMaterial(Rd, LightSample.Direction, VolPayload);
+							FMaterialEval MaterialEval = Medium_EvalMaterial(-Rd, LightSample.Direction, VolPayload, float2(1.0, 0.0));
 
 							// Record the contribution
 							float3 LightContrib = Contrib * LightSample.RadianceOverPdf * GetVolumetricScatteringIntensity(LightSampleLightId) * MaterialEval.Weight * MaterialEval.Pdf;
@@ -1000,7 +1081,7 @@ bool PathTracingKernel(inout FPathState PathState, int Bounce)
 				// now compute again with phase function MIS
 				if (MISMode != 1)
 				{
-					FMaterialSample MaterialSample = Medium_SampleMaterial(Rd, VolPayload, LightRandValue);
+					FMaterialSample MaterialSample = Medium_SampleMaterial(-Rd, VolPayload, LightRandValue);
 					if (MaterialSample.Pdf > 0)
 					{
 						FRayDesc MaterialRay;
@@ -1077,17 +1158,16 @@ bool PathTracingKernel(inout FPathState PathState, int Bounce)
 	// Choose a random number for both Light sampling and BxDF sampling
 	float4 RandSample = RandomSequence_GenerateSample4D(PathState.RandSequence);
 
-	// Does this material require NEE? (will be false if MaterialPdf is always +inf)
-	const bool bIsNeeValid = IsNeeValidMaterial(Payload);
+	const bool bIsVolumeSample = Payload.ShadingModelID == SHADINGMODELID_MEDIUM;
 
-	const bool bDoLightLoop = EnableDirectLighting != 0 || Bounce > 0;
+	float2 DiffuseSpecularScale = PathState.GetDiffuseSpecularScale(bIsVolumeSample);
 
-	const bool bIsVolumeSample = Payload.ShadingModelID == SHADINGMODELID_MEDIUM;
+	const bool bDoLightLoop = (EnableDirectLighting != 0 || Bounce > 0) && any(DiffuseSpecularScale > 0);
 
 	float LightPickingCdfSum = 0;
 
 	// If we are using Light sampling and the material can use it ...
-	if (MISMode != 0 && bIsNeeValid && SceneLightCount > 0 && bDoLightLoop)
+	if (MISMode != 0 && SceneLightCount > 0 && bDoLightLoop)
 	{
 		// Choose a light and sample it
 		float3 TranslatedWorldPos = Payload.TranslatedWorldPos;
@@ -1148,7 +1228,9 @@ bool PathTracingKernel(inout FPathState PathState, int Bounce)
 				if (any(LightSample.RadianceOverPdf > 0))
 				{
 					// Evaluate material
-					FMaterialEval MaterialEval = EvalMaterial(PathState.Ray.Direction, LightSample.Direction, Payload, false);
+					float2 LightDiffuseSpecularScale = DiffuseSpecularScale * float2(1.0, GetLightSpecularScale(LightId));
+
+					FMaterialEval MaterialEval = EvalMaterial(-PathState.Ray.Direction, LightSample.Direction, Payload, LightDiffuseSpecularScale);
 
 					// Record the contribution
 					float3 LightContrib = PathState.PathThroughput * LightSample.RadianceOverPdf * MaterialEval.Weight * MaterialEval.Pdf;
@@ -1163,7 +1245,7 @@ bool PathTracingKernel(inout FPathState PathState, int Bounce)
 	}
 
 	// Sample material
-	FMaterialSample MaterialSample = SampleMaterial(PathState.Ray.Direction, Payload, RandSample, false);
+	FMaterialSample MaterialSample = SampleMaterial(-PathState.Ray.Direction, Payload, RandSample.xyz);
 
 	if (MaterialSample.Pdf < 0 || asuint(MaterialSample.Pdf) > 0x7F800000)
 	{
@@ -1185,6 +1267,12 @@ bool PathTracingKernel(inout FPathState PathState, int Bounce)
 		return false;
 	}
 
+	if (!PathState.UpdateScatterType(MaterialSample.ScatterType))
+	{
+		// path scattered for a non-enabled lighting component, we can exit immediately
+		return false;
+	}
+
 	// Russian roulette:
 	//   The probability of keeping the path should be roughly proportional to the weight at the current shade point,
 	//  but just using MaterialWeight would miss out on cases where the path throughput changes color (like in a cornell
@@ -1197,8 +1285,7 @@ bool PathTracingKernel(inout FPathState PathState, int Bounce)
 	if (ContinuationProb < 1)
 	{
 		// If there is some chance we should terminate the ray, draw an extra random value
-		float RussianRouletteRand = RandSample.w; // SampleMaterial does not use this value at the moment
-		//RussianRouletteRand = RandomSequence_GenerateSample1D(RandSequence);
+		float RussianRouletteRand = RandSample.w;
 		if (RussianRouletteRand >= ContinuationProb)
 		{
 			// stochastically terminate the path
@@ -1268,9 +1355,24 @@ bool PathTracingKernel(inout FPathState PathState, int Bounce)
 			}
 
 			float3 LightContrib = PathState.PathThroughput * LightResult.Radiance;
-			if (bIsVolumeSample)
+			switch (MaterialSample.ScatterType)
 			{
-				LightContrib *= GetVolumetricScatteringIntensity(LightId);
+				case PATHTRACER_SCATTER_DIFFUSE:
+				{
+					LightContrib *= DiffuseSpecularScale.x;
+					break;
+				}
+				case PATHTRACER_SCATTER_SPECULAR:
+				case PATHTRACER_SCATTER_REFRACT:
+				{
+					LightContrib *= GetLightSpecularScale(LightId) * DiffuseSpecularScale.y;
+					break;
+				}
+				case PATHTRACER_SCATTER_VOLUME:
+				{
+					LightContrib *= GetVolumetricScatteringIntensity(LightId) * DiffuseSpecularScale.x;
+					break;
+				}
 			}
 
 			if (bUseMIS)
diff --git a/Engine/Shaders/Private/PathTracing/PathTracingMaterialHitShader.usf b/Engine/Shaders/Private/PathTracing/PathTracingMaterialHitShader.usf
--- a/Engine/Shaders/Private/PathTracing/PathTracingMaterialHitShader.usf
+++ b/Engine/Shaders/Private/PathTracing/PathTracingMaterialHitShader.usf
@@ -1,13 +1,5 @@
 // Copyright Epic Games, Inc. All Rights Reserved.
 
-#ifndef USE_MATERIAL_CLOSEST_HIT_SHADER
-	#define USE_MATERIAL_CLOSEST_HIT_SHADER 1
-#endif
-
-#ifndef USE_MATERIAL_ANY_HIT_SHADER
-	#define USE_MATERIAL_ANY_HIT_SHADER 1
-#endif
-
 #define PATH_TRACING 1
 #define ENABLE_SKY_LIGHT 1      
 #define NEEDS_LIGHTMAP_COORDINATE 0
@@ -79,7 +71,6 @@ RAY_TRACING_ENTRY_CLOSEST_HIT(PathTracingMaterialCHS,
 {
 	PackedPayload.HitT = RayTCurrent();
 
-#if USE_MATERIAL_CLOSEST_HIT_SHADER
 	ResolvedView = ResolveView();
 
 	const float3 TranslatedWorldPosition = TranslatedWorldRayOrigin() + RayTCurrent() * WorldRayDirection();
@@ -104,12 +95,7 @@ RAY_TRACING_ENTRY_CLOSEST_HIT(PathTracingMaterialCHS,
 		const float4 ScreenPosition = SvPositionToResolvedScreenPosition(SvPosition);
 		const bool bIsFrontFace = HitKind() == HIT_KIND_TRIANGLE_FRONT_FACE;
 
-		// #dxr_todo: UE-72130 support world position offset 
-		// #if USE_WORLD_POSITION_EXCLUDING_SHADER_OFFSETS
-		// CalcMaterialParametersEx(MaterialParameters, PixelMaterialInputs, SvPosition, ScreenPosition, bIsFrontFace, TranslatedWorldPosition, BasePassInterpolants.PixelPositionExcludingWPO);
-		// #else
 		CalcMaterialParametersEx(MaterialParameters, PixelMaterialInputs, SvPosition, ScreenPosition, bIsFrontFace, TranslatedWorldPosition, TranslatedWorldPosition);
-		// #endif
 	}
 
 	FPathTracingPayload Payload = (FPathTracingPayload)0;
@@ -123,12 +109,16 @@ RAY_TRACING_ENTRY_CLOSEST_HIT(PathTracingMaterialCHS,
 	uint MaterialBlendingMode = RAY_TRACING_BLEND_MODE_ALPHA_COMPOSITE;
 #elif MATERIALBLENDING_ALPHAHOLDOUT
 	uint MaterialBlendingMode = RAY_TRACING_BLEND_MODE_ALPHA_HOLDOUT;
+	HLSL_STATIC_ASSERT(MATERIAL_SHADINGMODEL_UNLIT == 1, "Alpha holdout blend mode requires unlit shading model");
+	Payload.ShadingModelID = SHADINGMODELID_UNLIT;
 #elif MATERIALBLENDING_TRANSLUCENT
 	uint MaterialBlendingMode = RAY_TRACING_BLEND_MODE_TRANSLUCENT;
 #elif MATERIALBLENDING_ADDITIVE
 	uint MaterialBlendingMode = RAY_TRACING_BLEND_MODE_ADDITIVE;
 #elif MATERIALBLENDING_MODULATE
 	uint MaterialBlendingMode = RAY_TRACING_BLEND_MODE_MODULATE;
+	HLSL_STATIC_ASSERT(MATERIAL_SHADINGMODEL_UNLIT == 1, "Modulate blend mode requires unlit shading model");
+	Payload.ShadingModelID = SHADINGMODELID_UNLIT;
 #else
 	uint MaterialBlendingMode = RAY_TRACING_BLEND_MODE_OPAQUE;
 #endif
@@ -168,10 +158,11 @@ RAY_TRACING_ENTRY_CLOSEST_HIT(PathTracingMaterialCHS,
 	half  Metallic = GetMaterialMetallic(PixelMaterialInputs);
 	half  Specular = GetMaterialSpecular(PixelMaterialInputs);
 	half Roughness = GetMaterialRoughness(PixelMaterialInputs);
-#if MATERIALBLENDING_TRANSLUCENT && REFRACTION_USE_INDEX_OF_REFRACTION
-	Payload.Ior = GetRefractionIor(PixelMaterialInputs);
-#else
-	Payload.Ior = 0;
+	float Ior = 0.0;
+#if MATERIALBLENDING_TRANSLUCENT && REFRACTION_USE_INDEX_OF_REFRACTION && (MATERIAL_SHADINGMODEL_DEFAULT_LIT || MATERIAL_SHADINGMODEL_THIN_TRANSLUCENT)
+	// NOTE: only default-lit and thin-translucent use ior, the other material models only support plain transparency
+	// This is an attempt to limit the complexity of the material eval/sample API which must take into account an extra lobe if supporting refraction
+	Ior = GetRefractionIor(PixelMaterialInputs);
 #endif
 
 #if !SIMPLIFIED_MATERIAL_SHADER
@@ -358,11 +349,12 @@ RAY_TRACING_ENTRY_CLOSEST_HIT(PathTracingMaterialCHS,
 #endif
 
 #if MATERIAL_SHADINGMODEL_THIN_TRANSLUCENT
-		if (Payload.ShadingModelID == SHADINGMODELID_THIN_TRANSLUCENT)
-		{
-			Payload.SetTransmittanceColor(GetThinTranslucentMaterialOutput0(MaterialParameters));
-		}
-		else
+	if (Payload.ShadingModelID == SHADINGMODELID_THIN_TRANSLUCENT)
+	{
+		Payload.SetTransmittanceColor(GetThinTranslucentMaterialOutput0(MaterialParameters));
+		Payload.Ior = Ior;
+	}
+	else
 #endif
 
 #if MATERIAL_SHADINGMODEL_HAIR
@@ -445,6 +437,7 @@ RAY_TRACING_ENTRY_CLOSEST_HIT(PathTracingMaterialCHS,
 		// Make the glass totally clear if no custom medium is set
 		Payload.SetTransmittanceColor(1.0);
 #endif
+		Payload.Ior = Ior;
 	}
 	else
 #endif
@@ -545,14 +538,14 @@ RAY_TRACING_ENTRY_CLOSEST_HIT(PathTracingMaterialCHS,
 #endif
 
 	PackedPayload = PackPathTracingPayload(Payload);
-#endif // USE_MATERIAL_CLOSEST_HIT_SHADER
 }
 
+#if USE_MATERIAL_ANY_HIT_SHADER
+
 RAY_TRACING_ENTRY_ANY_HIT(PathTracingMaterialAHS,
 	FPackedPathTracingPayload, PackedPayload,
 	FRayTracingIntersectionAttributes, Attributes)
 {
-#if USE_MATERIAL_ANY_HIT_SHADER
 
 #if MATERIALBLENDING_SOLID || MATERIALBLENDING_ALPHAHOLDOUT
 	// Continue traversal by simply returning from this shader
@@ -577,7 +570,6 @@ RAY_TRACING_ENTRY_ANY_HIT(PathTracingMaterialAHS,
 	// For any other blending models, we have a few options:
 	if (PackedPayload.IsVisibilityRay())
 	{
-		// this case is used by the path tracer only currently
 		// We want to evaluate the material opacity and accumulate it
 		// Keep going and evaluate the material
 	}
@@ -617,12 +609,7 @@ RAY_TRACING_ENTRY_ANY_HIT(PathTracingMaterialAHS,
 
 		MaterialParameters.CameraVector = -WorldRayDirection();
 
-		// #dxr_todo: UE-72130 support world position offset 
-		// #if USE_WORLD_POSITION_EXCLUDING_SHADER_OFFSETS
-		// CalcMaterialParametersEx(MaterialParameters, PixelMaterialInputs, SvPosition, ScreenPosition, bIsFrontFace, TranslatedWorldPosition, BasePassInterpolants.PixelPositionExcludingWPO);
-		// #else
 		CalcMaterialParametersEx(MaterialParameters, PixelMaterialInputs, SvPosition, ScreenPosition, bIsFrontFace, TranslatedWorldPosition, TranslatedWorldPosition);
-		// #endif
 	}
 
 	// Now extract the relevant info from PixelMaterialInputs according to the blending mode
@@ -635,117 +622,132 @@ RAY_TRACING_ENTRY_ANY_HIT(PathTracingMaterialAHS,
 	}
 #else
 	// If we got this far -- we have a blending mode that can produce partial opacity _and_ we want to
-	// know how much it contributes to the throughput along this ray. These cases are currently only
-	// exercised by the path tracer.
+	// know how much it contributes to the throughput along this ray (ie: this is a shadow ray).
 #if MATERIALBLENDING_MODULATE
-	// this material is the only one to lead to colored opacity
+	// Compute colored opacity
 	float3 Transparency = GetMaterialEmissive(PixelMaterialInputs);
 #elif MATERIALBLENDING_ALPHACOMPOSITE
 	float Opacity = GetMaterialOpacity(PixelMaterialInputs);
 	float Transparency = 1 - Opacity;
-#elif MATERIAL_SHADINGMODEL_THIN_TRANSLUCENT
-	// Assume Payload.ShadingModelID == SHADINGMODELID_THIN_TRANSLUCENT
-	float3 Transmission = GetThinTranslucentMaterialOutput0(MaterialParameters);
+#elif MATERIALBLENDING_MASKED && MATERIAL_DITHER_OPACITY_MASK
+	// See MATERIAL_DITHER_OPACITY_MASK comment above
+	float Opacity = saturate(GetMaterialMaskInputRaw(PixelMaterialInputs));
+	float Transparency = 1 - Opacity;
+#elif MATERIALBLENDING_TRANSLUCENT
 	float Opacity = GetMaterialOpacity(PixelMaterialInputs);
+#if MATERIAL_SHADINGMODEL_THIN_TRANSLUCENT || SIMPLIFIED_MATERIAL_SHADER
+	// We can get colored shadows either for thin translucent, or for GPULM which has backwards compatibility with some cpu lightmass tricks
+	float3 Transparency = 1 - Opacity;
+#else
+	float Transparency = 1 - Opacity;
+#endif
 
+	uint ShadingModelID = GetMaterialShadingModel(PixelMaterialInputs);
+	
+#if MATERIAL_SHADINGMODEL_THIN_TRANSLUCENT
+	if (ShadingModelID == SHADINGMODELID_THIN_TRANSLUCENT)
+	{
+		float3 Transmission = GetThinTranslucentMaterialOutput0(MaterialParameters);
 #if REFRACTION_USE_INDEX_OF_REFRACTION && !SIMPLIFIED_MATERIAL_SHADER
-	float Ior = GetRefractionIor(PixelMaterialInputs);
+		float Ior = GetRefractionIor(PixelMaterialInputs);
 #else
-	float Ior = 0.0;
+		float Ior = 0.0;
 #endif
-
-	float3 Transparency = 1 - Opacity;
-	float3 V = WorldRayDirection();
-	float3 N = normalize(MaterialParameters.WorldNormal);
-	float VoN = abs(dot(V, N));
+		float3 V = WorldRayDirection();
+		float3 N = normalize(MaterialParameters.WorldNormal);
+		float VoN = abs(dot(V, N));
 #if SIMPLIFIED_MATERIAL_SHADER
-	// simplified logic with no bending at the interface and no fresnel
-	Transparency *= pow(Transmission, 1.0 / VoN);
+		// simplified logic with no bending at the interface and no fresnel
+		Transparency *= pow(Transmission, 1.0 / VoN);
 #else
-	if (Opacity < 1.0)
-	{
-		float PathRoughness = PackedPayload.GetPathRoughness();
-		float Roughness = GetMaterialRoughness(PixelMaterialInputs);
-		if (PathRoughness <= Roughness)
+		if (Opacity < 1.0)
 		{
-			// not using fast caustic approximation - treat as opaque
-			return;
+			float PathRoughness = PackedPayload.GetPathRoughness();
+			float Roughness = GetMaterialRoughness(PixelMaterialInputs);
+			if (PathRoughness <= Roughness)
+			{
+				// not using fast caustic approximation - treat as opaque
+				return;
+			}
+			// compute transmission through the slab (fresnel + absorption)
+			float3 BaseColor = GetMaterialBaseColor(PixelMaterialInputs);
+			float  Metallic = GetMaterialMetallic(PixelMaterialInputs);
+			float  Specular = GetMaterialSpecular(PixelMaterialInputs);
+			float  F0 = F0RGBToF0(ComputeF0(Specular, BaseColor, Metallic));
+			Transparency *= ComputeThinSlabWeights(Transmission, VoN, Ior, F0).Transmitted;
+
+			// fake caustic approximation (see comments below)
+			Transparency *= (1 - Roughness * Roughness) * saturate(PathRoughness - Roughness);
 		}
-		// compute transmission through the slab (fresnel + absorption)
-		float3 BaseColor = GetMaterialBaseColor(PixelMaterialInputs);
-		float  Metallic = GetMaterialMetallic(PixelMaterialInputs);
-		float  Specular = GetMaterialSpecular(PixelMaterialInputs);
-		float  F0 = F0RGBToF0(ComputeF0(Specular, BaseColor, Metallic));
-		Transparency *= ComputeThinSlabWeights(Transmission, VoN, Ior, F0).Transmitted;
-
-		// fake caustic approximation (see comments below)
-		Transparency *= (1 - Roughness * Roughness) * saturate(PathRoughness - Roughness);
-	}
 #endif
-
-#elif MATERIALBLENDING_MASKED && MATERIAL_DITHER_OPACITY_MASK
-	// See MATERIAL_DITHER_OPACITY_MASK comment above
-	float Opacity = saturate(GetMaterialMaskInputRaw(PixelMaterialInputs));
-	float Transparency = 1 - Opacity;
-#elif MATERIALBLENDING_TRANSLUCENT
-	#if !SIMPLIFIED_MATERIAL_SHADER
-		float Opacity = GetMaterialOpacity(PixelMaterialInputs);
-		float Transparency = 1 - Opacity;
-	#else
+	}
+	else
+#endif // MATERIAL_SHADINGMODEL_THIN_TRANSLUCENT
+#if MATERIAL_SHADINGMODEL_DEFAULT_LIT
+	if (ShadingModelID == SHADINGMODELID_DEFAULT_LIT)
+	{
+#if SIMPLIFIED_MATERIAL_SHADER
 		// CPU Lightmass translucency behavior emulation
-		float3 Transparency = lerp(float3(1, 1, 1), GetMaterialBaseColor(PixelMaterialInputs), GetMaterialOpacity(PixelMaterialInputs));
-	#endif
-	
-	// Is refraction enabled?
-#if !SIMPLIFIED_MATERIAL_SHADER && MATERIAL_SHADINGMODEL_DEFAULT_LIT && REFRACTION_USE_INDEX_OF_REFRACTION
-	float Ior = GetRefractionIor(PixelMaterialInputs);
-	int ShadingModelID = GetMaterialShadingModel(PixelMaterialInputs);
-	if (Transparency > 0 && ShadingModelID == SHADINGMODELID_DEFAULT_LIT && Ior > 0.0)
-	{
-		// current hit has some refraction
-		float PathRoughness = PackedPayload.GetPathRoughness();
-		float Roughness = GetMaterialRoughness(PixelMaterialInputs);
-		if (PathRoughness <= Roughness)
+		Transparency = lerp(float3(1, 1, 1), GetMaterialBaseColor(PixelMaterialInputs), Opacity);
+#elif REFRACTION_USE_INDEX_OF_REFRACTION
+		// Is refraction enabled?
+		float Ior = GetRefractionIor(PixelMaterialInputs);
+		if (Transparency > 0 && Ior > 0.0)
 		{
-			// not using fast caustic approximation - treat as opaque
-			return;
-		}
-		// The heuristic used here is inspired by the following presentations:
-		//   - Kulla & Conty: Revisiting Physically Based Shading at Imageworks
-		//       https://blog.selfshadow.com/publications/s2017-shading-course/imageworks/s2017_pbs_imageworks_slides_v2.pdf
-		//   - Colin Barré-Brisebois - Pica Pica & Nvidia Turing
-		//       https://www.ea.com/seed/news/siggraph-2018-picapica-nv-turing
-
-		// fake caustics
-		float3 N = normalize(MaterialParameters.WorldNormal);
-		float Eta = Ior;
-		float VoN = abs(dot(WorldRayDirection(), N));
-		float3 BaseColor = GetMaterialBaseColor(PixelMaterialInputs);
-		float  Metallic = GetMaterialMetallic(PixelMaterialInputs);
-		float  Specular = GetMaterialSpecular(PixelMaterialInputs);
-		float  F0 = F0RGBToF0(ComputeF0(Specular, BaseColor, Metallic));
-		float Fr = FresnelReflectance(VoN, Eta, F0);
-		Transparency *= pow(1 - Fr, 2) * (1 - Roughness * Roughness) * saturate(PathRoughness - Roughness);
+			// current hit has some refraction
+			float PathRoughness = PackedPayload.GetPathRoughness();
+			float Roughness = GetMaterialRoughness(PixelMaterialInputs);
+			if (PathRoughness <= Roughness)
+			{
+				// not using fast caustic approximation - treat as opaque
+				return;
+			}
+			// The heuristic used here is inspired by the following presentations:
+			//   - Kulla & Conty: Revisiting Physically Based Shading at Imageworks
+			//       https://blog.selfshadow.com/publications/s2017-shading-course/imageworks/s2017_pbs_imageworks_slides_v2.pdf
+			//   - Colin Barré-Brisebois - Pica Pica & Nvidia Turing
+			//       https://www.ea.com/seed/news/siggraph-2018-picapica-nv-turing
+
+			// fake caustics
+			float3 N = normalize(MaterialParameters.WorldNormal);
+			float Eta = Ior;
+			float VoN = abs(dot(WorldRayDirection(), N));
+			float3 BaseColor = GetMaterialBaseColor(PixelMaterialInputs);
+			float  Metallic = GetMaterialMetallic(PixelMaterialInputs);
+			float  Specular = GetMaterialSpecular(PixelMaterialInputs);
+			float  F0 = F0RGBToF0(ComputeF0(Specular, BaseColor, Metallic));
+			float Fr = FresnelReflectance(VoN, Eta, F0);
+			Transparency *= pow(1 - Fr, 2) * (1 - Roughness * Roughness) * saturate(PathRoughness - Roughness);
 
 #if HAVE_GetAbsorptionMediumMaterialOutput0
-		// Now track the optical thickness so that we can account for Beer's law along the shadow ray
-		// TODO: to support lights inside glass we would need to track an extra offset, but this is hopefully an uncommon scenario
-		float Distance = RayTCurrent();
-		float3 TransmittanceColor = GetAbsorptionMediumMaterialOutput0(MaterialParameters);
-		float3 LocalSigmaT = -log(max(TransmittanceColor, 1e-8f));
-		if (bIsFrontFace)
-		{
-			// remove contribution from the ray origin to the current hit (assuming the backside will be hit)
-			PackedPayload.SetTau(PackedPayload.GetTau() - LocalSigmaT * Distance);
-		}
-		else
-		{
-			// add contribution from ray origin to current hit (backside) - excess portion not covered by the ray will be removed by the front face hit
-			PackedPayload.SetTau(PackedPayload.GetTau() + LocalSigmaT * Distance);
-		}
+			// Does the material have any kind of volumetric absorption to apply?
+			// Now track the optical thickness so that we can account for Beer's law along the shadow ray
+			// TODO: to support lights inside glass we would need to track an extra offset, but this is hopefully an uncommon scenario
+			// TODO2: SingleLayerWater needs something similar, but it generally is setup with opaque or masked blend mode and water geometry does not cast shadows ...
+			float Distance = RayTCurrent();
+			float3 TransmittanceColor = GetAbsorptionMediumMaterialOutput0(MaterialParameters);
+			float3 LocalSigmaT = -log(max(TransmittanceColor, 1e-8f));
+			if (bIsFrontFace)
+			{
+				// remove contribution from the ray origin to the current hit (assuming the backside will be hit)
+				PackedPayload.SetTau(PackedPayload.GetTau() - LocalSigmaT * Distance);
+			}
+			else
+			{
+				// add contribution from ray origin to current hit (backside) - excess portion not covered by the ray will be removed by the front face hit
+				PackedPayload.SetTau(PackedPayload.GetTau() + LocalSigmaT * Distance);
+			}
 #endif // HAVE_GetAbsorptionMediumMaterialOutput0
+
+		}
+#endif // REFRACTION_USE_INDEX_OF_REFRACTION
+	}
+#else // MATERIAL_SHADINGMODEL_DEFAULT_LIT
+	{
+		// base case for shadingmodel if/else
 	}
-#endif // refraction enabled?
+#endif // MATERIAL_SHADINGMODEL_DEFAULT_LIT
+
 #else
 #error "Control flow does not account for all blending modes!"
 #endif
@@ -767,5 +769,6 @@ RAY_TRACING_ENTRY_ANY_HIT(PathTracingMaterialAHS,
 
 #endif // !MATERIALBLENDING_SOLID && !MATERIALBLENDING_ALPHAHOLDOUT && !MATERIALBLENDING_ADDITIVE
 
-#endif // USE_MATERIAL_ANY_HIT_SHADER
 }
+
+#endif // USE_MATERIAL_ANY_HIT_SHADER
diff --git a/Engine/Shaders/Shared/PathTracingDefinitions.h b/Engine/Shaders/Shared/PathTracingDefinitions.h
--- a/Engine/Shaders/Shared/PathTracingDefinitions.h
+++ b/Engine/Shaders/Shared/PathTracingDefinitions.h
@@ -34,6 +34,14 @@
 #define PATHTRACER_MASK_UNUSED	(PATHTRACER_MASK_EMPTY_SLOT1|PATHTRACER_MASK_EMPTY_SLOT2)	
 #define PATHTRACER_MASK_ALL									0xFF
 
+// Constants for light contribution types (AOV decomposition of the image)
+// Leaving all constants enabled creates the beauty image, but turning off some bits allows
+// the path tracer to create an image with only certain components enabled
+#define PATHTRACER_CONTRIBUTION_EMISSIVE                     1
+#define PATHTRACER_CONTRIBUTION_DIFFUSE                      2
+#define PATHTRACER_CONTRIBUTION_SPECULAR                     4
+#define PATHTRACER_CONTRIBUTION_VOLUME                       8
+
 // Constants for the path tracer light grid
 #define PATHTRACER_LIGHT_GRID_SINGULAR_MASK					0x80000000u
 #define PATHTRACER_LIGHT_GRID_LIGHT_COUNT_MASK				0x7FFFFFFFu
diff --git a/Engine/Shaders/Shared/RayTracingTypes.h b/Engine/Shaders/Shared/RayTracingTypes.h
--- a/Engine/Shaders/Shared/RayTracingTypes.h
+++ b/Engine/Shaders/Shared/RayTracingTypes.h
@@ -24,7 +24,8 @@ struct FPathTracingLight {
 	float3  dPdv;
 	float3  Color;
 	float2  Dimensions; // Radius,Length or RectWidth,RectHeight or Sin(Angle/2),0 depending on light type
-	float2  Shaping;    // Barndoor controls for RectLights, Cone angles for spots lights
+	uint    Shaping;    // Barndoor controls for RectLights, Cone angles for spots lights, encoded as f16x2
+	float   SpecularScale;
 	float   Attenuation;
 	float   FalloffExponent; // for non-inverse square decay lights only
 	float   VolumetricScatteringIntensity;  // scale for volume contributions
diff --git a/Engine/Source/Runtime/Engine/Classes/Engine/Scene.h b/Engine/Source/Runtime/Engine/Classes/Engine/Scene.h
--- a/Engine/Source/Runtime/Engine/Classes/Engine/Scene.h
+++ b/Engine/Source/Runtime/Engine/Classes/Engine/Scene.h
@@ -1287,6 +1287,19 @@ struct FPostProcessSettings
 	UPROPERTY(EditAnywhere, BlueprintReadWrite, Category = Overrides, meta = (PinHiddenByDefault, InlineEditConditionToggle))
 	uint32 bOverride_PathTracingEnableDenoiser : 1;
 
+	UPROPERTY(EditAnywhere, BlueprintReadWrite, Category = Overrides, meta = (PinHiddenByDefault, InlineEditConditionToggle))
+	uint32 bOverride_PathTracingIncludeEmissive : 1;
+
+	UPROPERTY(EditAnywhere, BlueprintReadWrite, Category = Overrides, meta = (PinHiddenByDefault, InlineEditConditionToggle))
+	uint32 bOverride_PathTracingIncludeDiffuse : 1;
+
+	UPROPERTY(EditAnywhere, BlueprintReadWrite, Category = Overrides, meta = (PinHiddenByDefault, InlineEditConditionToggle))
+	uint32 bOverride_PathTracingIncludeSpecular : 1;
+
+	UPROPERTY(EditAnywhere, BlueprintReadWrite, Category = Overrides, meta = (PinHiddenByDefault, InlineEditConditionToggle))
+	uint32 bOverride_PathTracingIncludeVolume : 1;
+
+
 	// -----------------------------------------------------------------------
 
 	/** Enable HQ Gaussian on high end mobile platforms. (ES3_1) */
@@ -2134,38 +2147,53 @@ struct FPostProcessSettings
 
 	// Path Tracing
 	/** Sets the path tracing maximum bounces */
-	UPROPERTY(interp, EditAnywhere, BlueprintReadWrite, Category = "Rendering Features|PathTracing", meta = (ClampMin = "0", ClampMax = "100", editcondition = "bOverride_PathTracingMaxBounces", DisplayName = "Max. Bounces"))
+	UPROPERTY(interp, EditAnywhere, BlueprintReadWrite, Category = "Path Tracing", meta = (ClampMin = "0", ClampMax = "100", editcondition = "bOverride_PathTracingMaxBounces", DisplayName = "Max. Bounces"))
 	int32 PathTracingMaxBounces;
 
 	/** Sets the samples per pixel for the path tracer. */
-	UPROPERTY(interp, EditAnywhere, BlueprintReadWrite, Category = "Rendering Features|PathTracing", meta = (ClampMin = "1", UIMax = "65536", editcondition = "bOverride_PathTracingSamplesPerPixel", DisplayName = "Samples Per Pixel"))
+	UPROPERTY(interp, EditAnywhere, BlueprintReadWrite, Category = "Path Tracing", meta = (ClampMin = "1", UIMax = "65536", editcondition = "bOverride_PathTracingSamplesPerPixel", DisplayName = "Samples Per Pixel"))
 	int32 PathTracingSamplesPerPixel;
 
 	/** Sets anti-aliasing filter width for the path tracer. Lower values are sharper (and more aliased), larger values are softer (and blurrier). */
-	UPROPERTY(interp, EditAnywhere, BlueprintReadWrite, Category = "Rendering Features|PathTracing", meta = (ClampMin = "1.0", ClampMax = "6.0", editcondition = "bOverride_PathTracingFilterWidth", DisplayName = "Filter Width"))
+	UPROPERTY(interp, EditAnywhere, BlueprintReadWrite, Category = "Path Tracing", meta = (ClampMin = "1.0", ClampMax = "6.0", editcondition = "bOverride_PathTracingFilterWidth", DisplayName = "Filter Width"))
 	float PathTracingFilterWidth;
 
 	/** Enables emissive materials for the path tracer. This can prevent double-counting of illumination from surfaces that are also represented by light sources, and noise from small emitters. */
-	UPROPERTY(EditAnywhere, BlueprintReadWrite, Category = "Rendering Features|PathTracing", meta = (editcondition = "bOverride_PathTracingEnableEmissive", DisplayName = "Emissive Materials"))
+	UPROPERTY(EditAnywhere, BlueprintReadWrite, Category = "Path Tracing", meta = (editcondition = "bOverride_PathTracingEnableEmissive", DisplayName = "Emissive Materials"))
 	uint32 PathTracingEnableEmissive : 1;
 
 	/** Sets the maximum exposure allowed in the path tracer to reduce fireflies. This should be set a few stops higher than the scene exposure. */
-	UPROPERTY(interp, EditAnywhere, BlueprintReadWrite, Category = "Rendering Features|PathTracing", meta = (ClampMin = "-10.0", ClampMax = "30.0", editcondition = "bOverride_PathTracingMaxPathExposure", DisplayName = "Max Path Exposure"))
+	UPROPERTY(interp, EditAnywhere, BlueprintReadWrite, Category = "Path Tracing", meta = (ClampMin = "-10.0", ClampMax = "30.0", editcondition = "bOverride_PathTracingMaxPathExposure", DisplayName = "Max Path Exposure"))
 	float PathTracingMaxPathExposure;
 
 	/** Enables a reference quality depth-of-field which replaces the post-process effect. */
-	UPROPERTY(EditAnywhere, BlueprintReadWrite, Category = "Rendering Features|PathTracing", meta = (editcondition = "bOverride_PathTracingEnableReferenceDOF", DisplayName = "Reference Depth Of Field"))
+	UPROPERTY(EditAnywhere, BlueprintReadWrite, Category = "Path Tracing", meta = (editcondition = "bOverride_PathTracingEnableReferenceDOF", DisplayName = "Reference Depth Of Field"))
 	uint32 PathTracingEnableReferenceDOF : 1;
 
 	/** Enables path tracing in the atmosphere instead of baking the sky atmosphere contribution into a skylight. Any skylight present in the scene will be automatically ignored when this is enabled. */
-	UPROPERTY(EditAnywhere, BlueprintReadWrite, Category = "Rendering Features|PathTracing", meta = (editcondition = "bOverride_PathTracingEnableReferenceAtmosphere", DisplayName = "Reference Atmosphere"))
+	UPROPERTY(EditAnywhere, BlueprintReadWrite, Category = "Path Tracing", meta = (editcondition = "bOverride_PathTracingEnableReferenceAtmosphere", DisplayName = "Reference Atmosphere"))
 	uint32 PathTracingEnableReferenceAtmosphere : 1;
 
-
 	/** Run the currently loaded denoiser plugin on the last sample to remove noise from the output. Has no effect if a plug-in is not loaded. */
-	UPROPERTY(EditAnywhere, BlueprintReadWrite, Category = "Rendering Features|PathTracing", meta = (editcondition = "bOverride_PathTracingEnableDenoiser", DisplayName = "Denoiser"))
+	UPROPERTY(EditAnywhere, BlueprintReadWrite, Category = "Path Tracing", meta = (editcondition = "bOverride_PathTracingEnableDenoiser", DisplayName = "Denoiser"))
 	uint32 PathTracingEnableDenoiser : 1;
 
+	/** Should the render include directly visible emissive elements? */
+	UPROPERTY(EditAnywhere, BlueprintReadWrite, Category = "Path Tracing|Lighting Components", meta = (editcondition = "bOverride_PathTracingIncludeEmissive", DisplayName = "Emissive"))
+	uint32 PathTracingIncludeEmissive : 1;
+
+	/** Should the render include diffuse lighting contributions? */
+	UPROPERTY(EditAnywhere, BlueprintReadWrite, Category = "Path Tracing|Lighting Components", meta = (editcondition = "bOverride_PathTracingIncludeDiffuse", DisplayName = "Diffuse"))
+	uint32 PathTracingIncludeDiffuse : 1;
+
+	/** Should the render include specular lighting contributions? */
+	UPROPERTY(EditAnywhere, BlueprintReadWrite, Category = "Path Tracing|Lighting Components", meta = (editcondition = "bOverride_PathTracingIncludeSpecular", DisplayName = "Specular"))
+	uint32 PathTracingIncludeSpecular : 1;
+
+	/** Should the render include volume lighting contributions? */
+	UPROPERTY(EditAnywhere, BlueprintReadWrite, Category = "Path Tracing|Lighting Components", meta = (editcondition = "bOverride_PathTracingIncludeVolume", DisplayName = "Volume"))
+	uint32 PathTracingIncludeVolume : 1;
+
 
 	UPROPERTY()
 	float LPVFadeRange_DEPRECATED;
diff --git a/Engine/Source/Runtime/Engine/Private/Scene.cpp b/Engine/Source/Runtime/Engine/Private/Scene.cpp
--- a/Engine/Source/Runtime/Engine/Private/Scene.cpp
+++ b/Engine/Source/Runtime/Engine/Private/Scene.cpp
@@ -636,6 +636,11 @@ FPostProcessSettings::FPostProcessSettings()
 	PathTracingEnableReferenceDOF = 0;
 	PathTracingEnableReferenceAtmosphere = 0;
 	PathTracingEnableDenoiser = 1;
+
+	PathTracingIncludeEmissive = 1;
+	PathTracingIncludeDiffuse = 1;
+	PathTracingIncludeSpecular = 1;
+	PathTracingIncludeVolume = 1;
 	
 	bMobileHQGaussian = false;
 
@@ -844,6 +849,10 @@ FPostProcessSettings::FPostProcessSettings(const FPostProcessSettings& Settings)
 	, bOverride_PathTracingEnableReferenceDOF(Settings.bOverride_PathTracingEnableReferenceDOF)
 	, bOverride_PathTracingEnableReferenceAtmosphere(Settings.bOverride_PathTracingEnableReferenceAtmosphere)
 	, bOverride_PathTracingEnableDenoiser(Settings.bOverride_PathTracingEnableDenoiser)
+	, bOverride_PathTracingIncludeEmissive(Settings.bOverride_PathTracingIncludeEmissive)
+	, bOverride_PathTracingIncludeDiffuse(Settings.bOverride_PathTracingIncludeDiffuse)
+	, bOverride_PathTracingIncludeSpecular(Settings.bOverride_PathTracingIncludeSpecular)
+	, bOverride_PathTracingIncludeVolume(Settings.bOverride_PathTracingIncludeVolume)
 
 	, bMobileHQGaussian(Settings.bMobileHQGaussian)
 	, BloomMethod(Settings.BloomMethod)
@@ -1039,6 +1048,10 @@ FPostProcessSettings::FPostProcessSettings(const FPostProcessSettings& Settings)
 	, PathTracingEnableReferenceDOF(Settings.PathTracingEnableReferenceDOF)
 	, PathTracingEnableReferenceAtmosphere(Settings.PathTracingEnableReferenceAtmosphere)
 	, PathTracingEnableDenoiser(Settings.PathTracingEnableDenoiser)
+	, PathTracingIncludeEmissive(Settings.PathTracingIncludeEmissive)
+	, PathTracingIncludeDiffuse(Settings.PathTracingIncludeDiffuse)
+	, PathTracingIncludeSpecular(Settings.PathTracingIncludeSpecular)
+	, PathTracingIncludeVolume(Settings.PathTracingIncludeVolume)
 
 	, ScreenPercentage_DEPRECATED(Settings.ScreenPercentage_DEPRECATED)
 
diff --git a/Engine/Source/Runtime/Engine/Private/SceneView.cpp b/Engine/Source/Runtime/Engine/Private/SceneView.cpp
--- a/Engine/Source/Runtime/Engine/Private/SceneView.cpp
+++ b/Engine/Source/Runtime/Engine/Private/SceneView.cpp
@@ -1745,6 +1745,26 @@ void FSceneView::OverridePostProcessSettings(const FPostProcessSettings& Src, fl
 			Dest.PathTracingEnableDenoiser = Src.PathTracingEnableDenoiser;
 		}
 
+		if (Src.bOverride_PathTracingIncludeEmissive)
+		{
+			Dest.PathTracingIncludeEmissive = Src.PathTracingIncludeEmissive;
+		}
+
+		if (Src.bOverride_PathTracingIncludeDiffuse)
+		{
+			Dest.PathTracingIncludeDiffuse = Src.PathTracingIncludeDiffuse;
+		}
+
+		if (Src.bOverride_PathTracingIncludeSpecular)
+		{
+			Dest.PathTracingIncludeSpecular = Src.PathTracingIncludeSpecular;
+		}
+
+		if (Src.bOverride_PathTracingIncludeVolume)
+		{
+			Dest.PathTracingIncludeVolume = Src.PathTracingIncludeVolume;
+		}
+
 		if (Src.bOverride_DepthOfFieldBladeCount)
 		{
 			Dest.DepthOfFieldBladeCount = Src.DepthOfFieldBladeCount;
diff --git a/Engine/Source/Runtime/Renderer/Private/PathTracing.cpp b/Engine/Source/Runtime/Renderer/Private/PathTracing.cpp
--- a/Engine/Source/Runtime/Renderer/Private/PathTracing.cpp
+++ b/Engine/Source/Runtime/Renderer/Private/PathTracing.cpp
@@ -337,6 +337,8 @@ BEGIN_SHADER_PARAMETER_STRUCT(FPathTracingData, )
 	SHADER_PARAMETER(uint32, VisualizeDecalGrid)
 	SHADER_PARAMETER(uint32, EnableAtmosphere)
 	SHADER_PARAMETER(uint32, EnableFog)
+	SHADER_PARAMETER(uint32, EnabledLightingContributions) // PATHTRACER_CONTRIBUTION_*
+	SHADER_PARAMETER(uint32, ApplyDiffuseSpecularOverrides)
 	SHADER_PARAMETER(int32, MaxRaymarchSteps)
 	SHADER_PARAMETER(int32, MeshDecalsUseViewRay)
 	SHADER_PARAMETER(float, MaxPathIntensity)
@@ -386,6 +388,8 @@ struct FPathTracingConfig
 			PathTracingData.AbsorptionScale != Other.PathTracingData.AbsorptionScale ||
 			PathTracingData.EnableAtmosphere != Other.PathTracingData.EnableAtmosphere ||
 			PathTracingData.EnableFog != Other.PathTracingData.EnableFog ||
+			PathTracingData.ApplyDiffuseSpecularOverrides != Other.PathTracingData.ApplyDiffuseSpecularOverrides ||
+			PathTracingData.EnabledLightingContributions != Other.PathTracingData.EnabledLightingContributions ||
 			PathTracingData.DecalRoughnessCutoff != Other.PathTracingData.DecalRoughnessCutoff ||
 			PathTracingData.MeshDecalRoughnessCutoff != Other.PathTracingData.MeshDecalRoughnessCutoff ||
 			PathTracingData.MeshDecalBias != Other.PathTracingData.MeshDecalBias ||
@@ -570,6 +574,18 @@ static void PreparePathTracingData(const FScene* Scene, const FViewInfo& View, F
 	PathTracingData.MeshDecalsUseViewRay = CVarPathTracingMeshDecalsUseViewRay.GetValueOnRenderThread();
 
 	PathTracingData.MaxRaymarchSteps = CVarPathTracingMaxRaymarchSteps.GetValueOnRenderThread();
+
+	// NOTE: Diffuse and Specular show flags also modify the override colors, but we prefer to tie those to the lighting contribution mechanism below which is more principled
+	PathTracingData.ApplyDiffuseSpecularOverrides =
+		View.Family->EngineShowFlags.LightingOnlyOverride       != 0 ||
+		View.Family->EngineShowFlags.OverrideDiffuseAndSpecular != 0 ||
+		View.Family->EngineShowFlags.ReflectionOverride         != 0;
+
+	PathTracingData.EnabledLightingContributions = 0;
+	PathTracingData.EnabledLightingContributions |= (View.FinalPostProcessSettings.PathTracingIncludeEmissive != 0                                              ) ? PATHTRACER_CONTRIBUTION_EMISSIVE : 0;
+	PathTracingData.EnabledLightingContributions |= (View.FinalPostProcessSettings.PathTracingIncludeDiffuse  != 0 && View.Family->EngineShowFlags.Diffuse  != 0) ? PATHTRACER_CONTRIBUTION_DIFFUSE  : 0;
+	PathTracingData.EnabledLightingContributions |= (View.FinalPostProcessSettings.PathTracingIncludeSpecular != 0 && View.Family->EngineShowFlags.Specular != 0) ? PATHTRACER_CONTRIBUTION_SPECULAR : 0;
+	PathTracingData.EnabledLightingContributions |= (View.FinalPostProcessSettings.PathTracingIncludeVolume   != 0                                              ) ? PATHTRACER_CONTRIBUTION_VOLUME   : 0;
 }
 
 static bool ShouldCompilePathTracingShadersForProject(EShaderPlatform ShaderPlatform)
@@ -1699,6 +1715,7 @@ void SetLightParameters(FRDGBuilder& GraphBuilder, FPathTracingRG::FParameters*
 		DestLight.Flags |= PATHTRACING_LIGHT_SKY;
 		DestLight.Flags |= Scene->SkyLight->bCastShadows ? PATHTRACER_FLAG_CAST_SHADOW_MASK : 0;
 		DestLight.Flags |= Scene->SkyLight->bCastVolumetricShadow ? PATHTRACER_FLAG_CAST_VOL_SHADOW_MASK : 0;
+		DestLight.SpecularScale = 1.0f;
 		DestLight.VolumetricScatteringIntensity = Scene->SkyLight->VolumetricScatteringIntensity;
 		DestLight.IESAtlasIndex = INDEX_NONE;
 		DestLight.MissShaderIndex = 0;
@@ -1750,6 +1767,8 @@ void SetLightParameters(FRDGBuilder& GraphBuilder, FPathTracingRG::FParameters*
 			DestLight.Normal = -LightParameters.Direction;
 			DestLight.dPdu = FVector3f::CrossProduct(LightParameters.Tangent, LightParameters.Direction);
 			DestLight.dPdv = LightParameters.Tangent;
+			DestLight.Shaping = 0;
+			DestLight.SpecularScale = LightParameters.SpecularScale;
 			DestLight.Attenuation = LightParameters.InvRadius;
 			DestLight.FalloffExponent = 0;
 			DestLight.VolumetricScatteringIntensity = Light.LightSceneInfo->Proxy->GetVolumetricScatteringIntensity();
@@ -1816,6 +1835,8 @@ void SetLightParameters(FRDGBuilder& GraphBuilder, FPathTracingRG::FParameters*
 		DestLight.Normal = -LightParameters.Direction;
 		DestLight.dPdu = FVector3f::CrossProduct(LightParameters.Tangent, LightParameters.Direction);
 		DestLight.dPdv = LightParameters.Tangent;
+		DestLight.Shaping = 0;
+		DestLight.SpecularScale = LightParameters.SpecularScale;
 		DestLight.Attenuation = LightParameters.InvRadius;
 		DestLight.FalloffExponent = 0;
 		DestLight.VolumetricScatteringIntensity = Light.LightSceneInfo->Proxy->GetVolumetricScatteringIntensity();
@@ -1836,7 +1857,7 @@ void SetLightParameters(FRDGBuilder& GraphBuilder, FPathTracingRG::FParameters*
 			case LightType_Rect:
 			{
 				DestLight.Dimensions = FVector2f(2.0f * LightParameters.SourceRadius, 2.0f * LightParameters.SourceLength);
-				DestLight.Shaping = FVector2f(LightParameters.RectLightBarnCosAngle, LightParameters.RectLightBarnLength);
+				DestLight.Shaping = EncodeToF16x2(FVector2f(LightParameters.RectLightBarnCosAngle, LightParameters.RectLightBarnLength));
 				DestLight.FalloffExponent = LightParameters.FalloffExponent;
 				DestLight.Flags |= Light.LightSceneInfo->Proxy->IsInverseSquared() ? 0 : PATHTRACER_FLAG_NON_INVERSE_SQUARE_FALLOFF_MASK;
 				DestLight.Flags |= PATHTRACING_LIGHT_RECT;
@@ -1867,7 +1888,7 @@ void SetLightParameters(FRDGBuilder& GraphBuilder, FPathTracingRG::FParameters*
 			case LightType_Spot:
 			{
 				DestLight.Dimensions = FVector2f(LightParameters.SourceRadius, LightParameters.SourceLength);
-				DestLight.Shaping = LightParameters.SpotAngles;
+				DestLight.Shaping = EncodeToF16x2(LightParameters.SpotAngles);
 				DestLight.FalloffExponent = LightParameters.FalloffExponent;
 				DestLight.Flags |= Light.LightSceneInfo->Proxy->IsInverseSquared() ? 0 : PATHTRACER_FLAG_NON_INVERSE_SQUARE_FALLOFF_MASK;
 				DestLight.Flags |= PATHTRACING_LIGHT_SPOT;
diff --git a/Engine/Source/Runtime/Renderer/Private/RayTracing/RayTracingGlobalIllumination.cpp b/Engine/Source/Runtime/Renderer/Private/RayTracing/RayTracingGlobalIllumination.cpp
--- a/Engine/Source/Runtime/Renderer/Private/RayTracing/RayTracingGlobalIllumination.cpp
+++ b/Engine/Source/Runtime/Renderer/Private/RayTracing/RayTracingGlobalIllumination.cpp
@@ -246,6 +246,12 @@ DECLARE_GPU_STAT_NAMED(RayTracingGIBruteForce, TEXT("Ray Tracing GI: Brute Force
 DECLARE_GPU_STAT_NAMED(RayTracingGIFinalGather, TEXT("Ray Tracing GI: Final Gather"));
 DECLARE_GPU_STAT_NAMED(RayTracingGICreateGatherPoints, TEXT("Ray Tracing GI: Create Gather Points"));
 
+static uint32 EncodeToF16x2(const FVector2f& In)
+{
+	return FFloat16(In.X).Encoded | (FFloat16(In.Y).Encoded << 16);
+}
+
+// TODO: Share this logic with the path tracer code
 static void SetupLightParameters(
 	FScene* Scene,
 	const FViewInfo& View, FRDGBuilder& GraphBuilder,
@@ -298,6 +304,7 @@ static void SetupLightParameters(
 		DestLight.Flags |= Light.LightSceneInfo->Proxy->IsInverseSquared() ? 0 : PATHTRACER_FLAG_NON_INVERSE_SQUARE_FALLOFF_MASK;
 
 		DestLight.FalloffExponent = LightShaderParameters.FalloffExponent;
+		DestLight.SpecularScale = LightShaderParameters.SpecularScale;
 		DestLight.Attenuation = LightShaderParameters.InvRadius;
 		DestLight.IESAtlasIndex = LightShaderParameters.IESAtlasIndex;
 
@@ -323,7 +330,7 @@ static void SetupLightParameters(
 			DestLight.dPdv = LightShaderParameters.Tangent;
 			DestLight.Color = FVector3f(LightShaderParameters.Color);
 			DestLight.Dimensions = FVector2f(2.0f * LightShaderParameters.SourceRadius, 2.0f * LightShaderParameters.SourceLength);
-			DestLight.Shaping = FVector2f(LightShaderParameters.RectLightBarnCosAngle, LightShaderParameters.RectLightBarnLength);
+			DestLight.Shaping = EncodeToF16x2(FVector2f(LightShaderParameters.RectLightBarnCosAngle, LightShaderParameters.RectLightBarnLength));
 			DestLight.Flags |= PATHTRACING_LIGHT_RECT;
 			break;
 		}
@@ -350,7 +357,7 @@ static void SetupLightParameters(
 			DestLight.Color = FVector3f(LightShaderParameters.Color);
 			float SourceRadius = 0.0; // LightShaderParameters.SourceRadius causes too much noise for little pay off at this time
 			DestLight.Dimensions = FVector2f(SourceRadius, 0.0);
-			DestLight.Shaping = LightShaderParameters.SpotAngles;
+			DestLight.Shaping = EncodeToF16x2(LightShaderParameters.SpotAngles);
 			DestLight.Flags |= PATHTRACING_LIGHT_SPOT;
 			break;
 		}
