Eye Adapation: Implement an extra mode that just uses the previous frame exposure to allow basic support from contexts such as ray tracing or path tracing where it isn't convenient to bind an extra buffer.

diff --git a/Engine/Shaders/Private/EyeAdaptationCommon.ush b/Engine/Shaders/Private/EyeAdaptationCommon.ush
--- a/Engine/Shaders/Private/EyeAdaptationCommon.ush
+++ b/Engine/Shaders/Private/EyeAdaptationCommon.ush
@@ -4,14 +4,14 @@
 
 #include "Common.ush"
 
-#if !EYE_ADAPTATION_DISABLED
+#if !EYE_ADAPTATION_DISABLED && !EYE_ADAPTATION_PREV_FRAME_EXPOSURE
 StructuredBuffer<float4> EyeAdaptationBuffer;
 
 float4 EyeAdaptationLookupBuffer(StructuredBuffer<float4> InEyeAdaptation)
 {
 	return InEyeAdaptation[0];
 }
-#endif // !EYE_ADAPTATION_DISABLED
+#endif // !EYE_ADAPTATION_DISABLED && !EYE_ADAPTATION_PREV_FRAME_EXPOSURE
 
 #if defined (EyeAdaptationStruct)
 	#if EYE_ADAPTATION_LOOSE_PARAMETERS
@@ -43,6 +43,9 @@ float EyeAdaptationLookup()
 	#endif
 #elif EYE_ADAPTATION_LOOSE_PARAMETERS
 	return EyeAdaptationLookupBuffer(EyeAdaptationBuffer).x;
+#elif EYE_ADAPTATION_PREV_FRAME_EXPOSURE
+	// fallback implementation if the EyeAdaptationBuffer is not available
+	return ResolvedView.PreExposure;
 #else
 	return 0.0f;
 #endif
diff --git a/Engine/Shaders/Private/PathTracing/PathTracingMaterialHitShader.usf b/Engine/Shaders/Private/PathTracing/PathTracingMaterialHitShader.usf
--- a/Engine/Shaders/Private/PathTracing/PathTracingMaterialHitShader.usf
+++ b/Engine/Shaders/Private/PathTracing/PathTracingMaterialHitShader.usf
@@ -9,6 +9,9 @@
 // Needed for VertexFactoryInterpolate to interpolate attributes from vertices to hit point
 #define NEEDS_VERTEX_FACTORY_INTERPOLATION 1
 
+// This should be good enough for path tracing and avoids having to bind an extra buffer
+#define EYE_ADAPTATION_PREV_FRAME_EXPOSURE 1
+
 #include "/Engine/Private/Common.ush"
 #include "/Engine/Private/RayTracing/RayTracingCommon.ush"
 #include "/Engine/Private/RayTracing/RayTracingHitGroupCommon.ush"
diff --git a/Engine/Shaders/Private/RayTracing/RayTracingMaterialHitShaders.usf b/Engine/Shaders/Private/RayTracing/RayTracingMaterialHitShaders.usf
--- a/Engine/Shaders/Private/RayTracing/RayTracingMaterialHitShaders.usf
+++ b/Engine/Shaders/Private/RayTracing/RayTracingMaterialHitShaders.usf
@@ -10,6 +10,10 @@
 // Needed for VertexFactoryInterpolate to interpolate attributes from vertices to hit point
 #define NEEDS_VERTEX_FACTORY_INTERPOLATION 1
 
+// This should be good enough for ray tracing and avoids having to bind an extra buffer
+#define EYE_ADAPTATION_PREV_FRAME_EXPOSURE 1
+
+
 #if STRATA_ENABLED
 #define STRATA_INLINE_SHADING 1
 #endif
