Fixed bug with VT compression where bit fields where booleans were bit shifted

diff --git a/Engine/Shaders/Private/ETCCompressionCommon.ush b/Engine/Shaders/Private/ETCCompressionCommon.ush
--- a/Engine/Shaders/Private/ETCCompressionCommon.ush
+++ b/Engine/Shaders/Private/ETCCompressionCommon.ush
@@ -122,12 +122,13 @@ uint2 CompressBlock_ETC2_RGB(in float3 Block[16])
 	uint ColorBits;
 	float3 BaseColor1_Quant, BaseColor2_Quant;
 
-	uint BlockMode;
+	// Changed block mode to store bits because of DXC compiler optimization error
+	uint BlockModeBits;
 	int3 MinDiff = { -4, -4, -4 }, MaxDiff = { 3, 3, 3 };
 	if (all(Diff >= MinDiff) && all(Diff <= MaxDiff))
 	{
 		// We can use differential mode.
-		BlockMode = BLOCK_MODE_DIFFERENTIAL;
+		BlockModeBits = BLOCK_MODE_DIFFERENTIAL << 25;
 		ColorBits = ((Diff.b & 7) << 16) | (BaseColor1.b << 19) | ((Diff.g & 7) << 8) | (BaseColor1.g << 11) | (Diff.r & 7) | (BaseColor1.r << 3);
 		BaseColor1_Quant = ExpandColor555(BaseColor1);
 		BaseColor2_Quant = ExpandColor555(BaseColor2);
@@ -135,7 +136,7 @@ uint2 CompressBlock_ETC2_RGB(in float3 Block[16])
 	else
 	{
 		// We must use the lower precision individual mode.
-		BlockMode = BLOCK_MODE_INDIVIDUAL;
+		BlockModeBits = BLOCK_MODE_INDIVIDUAL << 25;
 		BaseColor1 >>= 1;
 		BaseColor2 >>= 1;
 		ColorBits = (BaseColor1.b << 20) | (BaseColor2.b << 16) | (BaseColor1.g << 12) | (BaseColor2.g << 8) | (BaseColor1.r << 4) | BaseColor2.r;
@@ -161,7 +162,7 @@ uint2 CompressBlock_ETC2_RGB(in float3 Block[16])
 	
 	// Both these values need to be big-endian. We can build ModeBits directly in big-endian layout, but for IndexBits
 	// it's too hard, so we'll just swap here.
-	uint ModeBits = (SubBlock1TableIdx << 29) | (SubBlock2TableIdx << 26) | (BlockMode << 25) | (FlipBit << 24) | ColorBits;
+	uint ModeBits = (SubBlock1TableIdx << 29) | (SubBlock2TableIdx << 26) | (BlockModeBits) | (FlipBit << 24) | ColorBits;
 	uint IndexBits = SwapEndian32(SubBlock1Weights | SubBlock2Weights);
 
 	return uint2(ModeBits, IndexBits);
