Revert sqaure() removal for PCF shadow pass to keep the original PCF look, which is less blurry visually. More comments are added so it will not be applied accidentally elsewhere.

diff --git a/Engine/Shaders/Private/ShadowProjectionPixelShader.usf b/Engine/Shaders/Private/ShadowProjectionPixelShader.usf
--- a/Engine/Shaders/Private/ShadowProjectionPixelShader.usf
+++ b/Engine/Shaders/Private/ShadowProjectionPixelShader.usf
@@ -80,6 +80,13 @@ float3 ModulatedShadowBlendOp(float3 Source)
 	return Source.rgb*Dest.rgb;
 }
 
+// PCF falloff is overly blurry, apply to get
+// a more reasonable look artistically. Should not be applied to
+// other passes that target pbr (e.g., raytracing and denoising).
+float ApplyPCFOverBlurCorrection(float Shadow)
+{
+	return Square(Shadow);
+}
 
 #if SUBPIXEL_SHADOW
 #include "HairStrands/HairStrandsVisibilityCommon.ush"
@@ -326,6 +333,8 @@ void Main(
 				// ideally we use a larger filter kernel for SSSbut as Gather4 makes that harder
 				SSSTransmission = ManualPCF(ShadowPosition.xy, Settings);
 			}
+
+			SSSTransmission = ApplyPCFOverBlurCorrection(SSSTransmission);
 		}
 
 	#endif
@@ -338,6 +347,13 @@ void Main(
 
 	// 0 is shadowed, 1 is unshadowed
 	// RETURN_COLOR not needed unless writing to SceneColor;
+
+	// Apply PCF overblur correction for shadow when PCSS & PCF.
+	#define CORRECT_SHADOW_PCF_BLUR  !UNFILTERED_SHADOW_PROJECTION && !APPLY_TRANSLUCENCY_SHADOWS
+
+	#if CORRECT_SHADOW_PCF_BLUR
+		Shadow = ApplyPCFOverBlurCorrection(Shadow);
+	#endif
 	float FadedShadow = lerp(1.0f, Shadow, ShadowFadeFraction * PerObjectDistanceFadeFraction);
 
 #if FORWARD_SHADING || SHADING_PATH_MOBILE 
@@ -434,6 +450,8 @@ void MainOnePassPointLightPS(
 	float Shadow = CubemapHardwarePCF(TranslateWorldPosition, LightPositionAndInvRadius.xyz, LightPositionAndInvRadius.w, PointLightDepthBias.x, SlopeBias, PointLightDepthBias.z);
 
 	Shadow = saturate( (Shadow - 0.5) * ShadowSharpen + 0.5 );
+
+	Shadow = ApplyPCFOverBlurCorrection(Shadow);
 #endif
 
 	float FadedShadow = lerp(1.0f, Shadow, ShadowFadeFraction);
