Fix motion vector errors for instances that do not have dynamic data but whose primitive can move.

This required calculating a World->Previous World transform in higher precision on CPU and uploading it to the GPU to prevent precision issues that were observed to happen with WPO that used TransformPosition (Instance Local to World) material nodes when the actor is very far from the origin (about 10,000 units).

#rnx

diff --git a/Engine/Shaders/Private/MaterialTemplate.ush b/Engine/Shaders/Private/MaterialTemplate.ush
--- a/Engine/Shaders/Private/MaterialTemplate.ush
+++ b/Engine/Shaders/Private/MaterialTemplate.ush
@@ -682,8 +682,7 @@ FLWCMatrix GetInstanceToWorld(FMaterialVertexParameters Parameters)
 FLWCMatrix GetPrevInstanceToWorld(FMaterialVertexParameters Parameters)
 {
 #if USE_INSTANCING || USE_INSTANCE_CULLING || IS_MESHPARTICLE_FACTORY || IS_NANITE_PASS
-	// todo: return Parameters.PrevFrameLocalToWorld when we have fixed precision issue of PrevLocalToWorld in GetInstanceSceneData()
-	return Parameters.InstanceLocalToWorld;
+	return Parameters.PrevFrameLocalToWorld;
 #else
 	return GetPrimitiveData(Parameters).PreviousLocalToWorld;
 #endif
diff --git a/Engine/Shaders/Private/SceneData.ush b/Engine/Shaders/Private/SceneData.ush
--- a/Engine/Shaders/Private/SceneData.ush
+++ b/Engine/Shaders/Private/SceneData.ush
@@ -153,6 +153,7 @@ struct FPrimitiveSceneData
 	FLWCInverseMatrix WorldToLocal;
 	FLWCMatrix	PreviousLocalToWorld;
 	FLWCInverseMatrix PreviousWorldToLocal;
+	float4x4	WorldToPreviousWorld; // NOTE: Intentionally not LWC
 	float3		InvNonUniformScale;
 	float		ObjectBoundsX;
 	FLWCVector3	ObjectWorldPosition;
@@ -201,6 +202,7 @@ FPrimitiveSceneData GetPrimitiveDataFromUniformBuffer()
 	PrimitiveData.WorldToLocal						= MakeLWCInverseMatrix4x3(Primitive.TilePosition, Primitive.RelativeWorldToLocal);
 	PrimitiveData.PreviousLocalToWorld				= MakeLWCMatrix4x3(Primitive.TilePosition, Primitive.PreviousLocalToRelativeWorld);
 	PrimitiveData.PreviousWorldToLocal				= MakeLWCInverseMatrix4x3(Primitive.TilePosition, Primitive.PreviousRelativeWorldToLocal);
+	PrimitiveData.WorldToPreviousWorld				= Primitive.WorldToPreviousWorld;
 	PrimitiveData.InvNonUniformScale				= Primitive.InvNonUniformScale;
 	PrimitiveData.ObjectBoundsX						= Primitive.ObjectBoundsX;
 	PrimitiveData.ObjectWorldPosition				= MakeLWCVector3(Primitive.TilePosition, Primitive.ObjectRelativeWorldPositionAndRadius.xyz);
@@ -250,7 +252,7 @@ RWStructuredBuffer<float4> GPUScenePrimitiveSceneDataRW;
 #endif
 
 // Stride of a single primitive's data in float4's, must match C++ (FPrimitiveSceneShaderData::DataStrideInFloat4s and FScatterUploadBuffer::PrimitiveDataStrideInFloat4s)
-#define PRIMITIVE_SCENE_DATA_STRIDE 42
+#define PRIMITIVE_SCENE_DATA_STRIDE 41
 
 float4 LoadPrimitivePrimitiveSceneDataElement(uint PrimitiveIndex, uint ItemIndex)
 {
@@ -283,31 +285,42 @@ FPrimitiveSceneData GetPrimitiveData(uint PrimitiveId)
 	
 	float3 TilePosition = LoadPrimitivePrimitiveSceneDataElement(PrimitiveIndex, 1).xyz;
 
-	float4x4 LocalToWorld;
-	LocalToWorld[0] = float4(LoadPrimitivePrimitiveSceneDataElement(PrimitiveIndex, 2).xyz, 0.0f);
-	LocalToWorld[1] = float4(LoadPrimitivePrimitiveSceneDataElement(PrimitiveIndex, 3).xyz, 0.0f);
-	LocalToWorld[2] = float4(LoadPrimitivePrimitiveSceneDataElement(PrimitiveIndex, 4).xyz, 0.0f);
-	LocalToWorld[3] = float4(LoadPrimitivePrimitiveSceneDataElement(PrimitiveIndex, 5).xyz, 1.0f);
-
-	float4x4 PreviousLocalToWorld;
-	PreviousLocalToWorld[0]	= float4(LoadPrimitivePrimitiveSceneDataElement(PrimitiveIndex, 10).xyz, 0.0f);
-	PreviousLocalToWorld[1]	= float4(LoadPrimitivePrimitiveSceneDataElement(PrimitiveIndex, 11).xyz, 0.0f);
-	PreviousLocalToWorld[2]	= float4(LoadPrimitivePrimitiveSceneDataElement(PrimitiveIndex, 12).xyz, 0.0f);
-	PreviousLocalToWorld[3]	= float4(LoadPrimitivePrimitiveSceneDataElement(PrimitiveIndex, 13).xyz, 1.0f);
-
-	float4x4 WorldToLocal;
-	WorldToLocal[0]	= float4(LoadPrimitivePrimitiveSceneDataElement(PrimitiveIndex, 6).xyz, 0.0f);
-	WorldToLocal[1]	= float4(LoadPrimitivePrimitiveSceneDataElement(PrimitiveIndex, 7).xyz, 0.0f);
-	WorldToLocal[2]	= float4(LoadPrimitivePrimitiveSceneDataElement(PrimitiveIndex, 8).xyz, 0.0f);
-	WorldToLocal[3]	= float4(LoadPrimitivePrimitiveSceneDataElement(PrimitiveIndex, 9).xyz, 1.0f);
-
-	float4x4 PreviousWorldToLocal;
-	PreviousWorldToLocal[0]	= float4(LoadPrimitivePrimitiveSceneDataElement(PrimitiveIndex, 14).xyz, 0.0f);
-	PreviousWorldToLocal[1]	= float4(LoadPrimitivePrimitiveSceneDataElement(PrimitiveIndex, 15).xyz, 0.0f);
-	PreviousWorldToLocal[2]	= float4(LoadPrimitivePrimitiveSceneDataElement(PrimitiveIndex, 16).xyz, 0.0f);
-	PreviousWorldToLocal[3]	= float4(LoadPrimitivePrimitiveSceneDataElement(PrimitiveIndex, 17).xyz, 1.0f);
-
-	float4 ObjectWorldPositionAndRadius	= LoadPrimitivePrimitiveSceneDataElement(PrimitiveIndex, 19);
+	const float4x4 LocalToWorld = transpose(float4x4(
+		LoadPrimitivePrimitiveSceneDataElement(PrimitiveIndex, 2),
+		LoadPrimitivePrimitiveSceneDataElement(PrimitiveIndex, 3),
+		LoadPrimitivePrimitiveSceneDataElement(PrimitiveIndex, 4),
+		float4(0, 0, 0, 1)
+	));
+
+	const float4x4 PreviousLocalToWorld = transpose(float4x4(
+		LoadPrimitivePrimitiveSceneDataElement(PrimitiveIndex, 8),
+		LoadPrimitivePrimitiveSceneDataElement(PrimitiveIndex, 9),
+		LoadPrimitivePrimitiveSceneDataElement(PrimitiveIndex, 10),
+		float4(0, 0, 0, 1)
+	));
+
+	const float4x4 WorldToLocal = transpose(float4x4(
+		LoadPrimitivePrimitiveSceneDataElement(PrimitiveIndex, 5),
+		LoadPrimitivePrimitiveSceneDataElement(PrimitiveIndex, 6),
+		LoadPrimitivePrimitiveSceneDataElement(PrimitiveIndex, 7),
+		float4(0, 0, 0, 1)
+	));
+
+	const float4x4 PreviousWorldToLocal = transpose(float4x4(
+		LoadPrimitivePrimitiveSceneDataElement(PrimitiveIndex, 11),
+		LoadPrimitivePrimitiveSceneDataElement(PrimitiveIndex, 12),
+		LoadPrimitivePrimitiveSceneDataElement(PrimitiveIndex, 13),
+		float4(0, 0, 0, 1)
+	));
+
+	const float4x4 WorldToPreviousWorld = transpose(float4x4(
+		LoadPrimitivePrimitiveSceneDataElement(PrimitiveIndex, 14),
+		LoadPrimitivePrimitiveSceneDataElement(PrimitiveIndex, 15),
+		LoadPrimitivePrimitiveSceneDataElement(PrimitiveIndex, 16),
+		float4(0, 0, 0, 1)
+	));
+
+	float4 ObjectWorldPositionAndRadius	= LoadPrimitivePrimitiveSceneDataElement(PrimitiveIndex, 18);
 
 	PrimitiveData.Flags									= asuint(LoadPrimitivePrimitiveSceneDataElement(PrimitiveIndex, 0).x);
 	PrimitiveData.InstanceSceneDataOffset				= asuint(LoadPrimitivePrimitiveSceneDataElement(PrimitiveIndex, 0).y);
@@ -319,36 +332,37 @@ FPrimitiveSceneData GetPrimitiveData(uint PrimitiveId)
 	PrimitiveData.WorldToLocal							= MakeLWCInverseMatrix4x3(TilePosition, WorldToLocal);
 	PrimitiveData.PreviousLocalToWorld					= MakeLWCMatrix4x3(TilePosition, PreviousLocalToWorld);
 	PrimitiveData.PreviousWorldToLocal					= MakeLWCInverseMatrix4x3(TilePosition, PreviousWorldToLocal);
-	PrimitiveData.InvNonUniformScale					= LoadPrimitivePrimitiveSceneDataElement(PrimitiveIndex, 18).xyz;
-	PrimitiveData.ObjectBoundsX							= LoadPrimitivePrimitiveSceneDataElement(PrimitiveIndex, 18).w;
+	PrimitiveData.WorldToPreviousWorld					= WorldToPreviousWorld;
+	PrimitiveData.InvNonUniformScale					= LoadPrimitivePrimitiveSceneDataElement(PrimitiveIndex, 17).xyz;
+	PrimitiveData.ObjectBoundsX							= LoadPrimitivePrimitiveSceneDataElement(PrimitiveIndex, 17).w;
 	PrimitiveData.ObjectWorldPosition					= MakeLWCVector3(TilePosition, ObjectWorldPositionAndRadius.xyz);
 	PrimitiveData.ObjectRadius							= ObjectWorldPositionAndRadius.w;
-	PrimitiveData.ActorWorldPosition					= MakeLWCVector3(TilePosition, LoadPrimitivePrimitiveSceneDataElement(PrimitiveIndex, 20).xyz);
-	PrimitiveData.LightmapUVIndex						= asuint(LoadPrimitivePrimitiveSceneDataElement(PrimitiveIndex, 20).w);
-	PrimitiveData.ObjectOrientation						= LoadPrimitivePrimitiveSceneDataElement(PrimitiveIndex, 21).xyz;
-	PrimitiveData.LightmapDataIndex						= asuint(LoadPrimitivePrimitiveSceneDataElement(PrimitiveIndex, 21).w);
-	PrimitiveData.NonUniformScale						= LoadPrimitivePrimitiveSceneDataElement(PrimitiveIndex, 22);
-	PrimitiveData.PreSkinnedLocalBoundsMin				= LoadPrimitivePrimitiveSceneDataElement(PrimitiveIndex, 23).xyz;
-	PrimitiveData.NaniteResourceID						= asuint(LoadPrimitivePrimitiveSceneDataElement(PrimitiveIndex, 23).w);
-	PrimitiveData.PreSkinnedLocalBoundsMax				= LoadPrimitivePrimitiveSceneDataElement(PrimitiveIndex, 24).xyz;
-	PrimitiveData.NaniteHierarchyOffset					= asuint(LoadPrimitivePrimitiveSceneDataElement(PrimitiveIndex, 24).w);
-	PrimitiveData.LocalObjectBoundsMin					= LoadPrimitivePrimitiveSceneDataElement(PrimitiveIndex, 25).xyz;
-	PrimitiveData.ObjectBoundsY							= LoadPrimitivePrimitiveSceneDataElement(PrimitiveIndex, 25).w;
-	PrimitiveData.LocalObjectBoundsMax					= LoadPrimitivePrimitiveSceneDataElement(PrimitiveIndex, 26).xyz;
-	PrimitiveData.ObjectBoundsZ							= LoadPrimitivePrimitiveSceneDataElement(PrimitiveIndex, 26).w;
-	PrimitiveData.InstanceLocalBoundsCenter				= LoadPrimitivePrimitiveSceneDataElement(PrimitiveIndex, 27).xyz;
-	PrimitiveData.InstancePayloadDataOffset				= asuint(LoadPrimitivePrimitiveSceneDataElement(PrimitiveIndex, 27).w);
-	PrimitiveData.InstanceLocalBoundsExtent				= LoadPrimitivePrimitiveSceneDataElement(PrimitiveIndex, 28).xyz;
-	PrimitiveData.InstancePayloadDataStride				= asuint(LoadPrimitivePrimitiveSceneDataElement(PrimitiveIndex, 28).w);
-	PrimitiveData.WireframeColor						= LoadPrimitivePrimitiveSceneDataElement(PrimitiveIndex, 29).xyz;
-	PrimitiveData.PackedNaniteFlags						= asuint(LoadPrimitivePrimitiveSceneDataElement(PrimitiveIndex, 29).w);
-	PrimitiveData.LevelColor							= LoadPrimitivePrimitiveSceneDataElement(PrimitiveIndex, 30).xyz;
-	PrimitiveData.PersistentPrimitiveIndex				= asint(LoadPrimitivePrimitiveSceneDataElement(PrimitiveIndex, 30).w);
-	PrimitiveData.InstanceDrawDistanceMinMaxSquared		= LoadPrimitivePrimitiveSceneDataElement(PrimitiveIndex, 31).xy;
-	PrimitiveData.InstanceWPODisableDistanceSquared		= LoadPrimitivePrimitiveSceneDataElement(PrimitiveIndex, 31).z;
-	PrimitiveData.NaniteRayTracingDataOffset			= asuint(LoadPrimitivePrimitiveSceneDataElement(PrimitiveIndex, 31).w);
-	PrimitiveData.BoundsScale							= LoadPrimitivePrimitiveSceneDataElement(PrimitiveIndex, 32).x;
-	PrimitiveData.CustomStencilValueAndMask				= asuint(LoadPrimitivePrimitiveSceneDataElement(PrimitiveIndex, 32).y);
+	PrimitiveData.ActorWorldPosition					= MakeLWCVector3(TilePosition, LoadPrimitivePrimitiveSceneDataElement(PrimitiveIndex, 19).xyz);
+	PrimitiveData.LightmapUVIndex						= asuint(LoadPrimitivePrimitiveSceneDataElement(PrimitiveIndex, 19).w);
+	PrimitiveData.ObjectOrientation						= LoadPrimitivePrimitiveSceneDataElement(PrimitiveIndex, 20).xyz;
+	PrimitiveData.LightmapDataIndex						= asuint(LoadPrimitivePrimitiveSceneDataElement(PrimitiveIndex, 20).w);
+	PrimitiveData.NonUniformScale						= LoadPrimitivePrimitiveSceneDataElement(PrimitiveIndex, 21);
+	PrimitiveData.PreSkinnedLocalBoundsMin				= LoadPrimitivePrimitiveSceneDataElement(PrimitiveIndex, 22).xyz;
+	PrimitiveData.NaniteResourceID						= asuint(LoadPrimitivePrimitiveSceneDataElement(PrimitiveIndex, 22).w);
+	PrimitiveData.PreSkinnedLocalBoundsMax				= LoadPrimitivePrimitiveSceneDataElement(PrimitiveIndex, 23).xyz;
+	PrimitiveData.NaniteHierarchyOffset					= asuint(LoadPrimitivePrimitiveSceneDataElement(PrimitiveIndex, 23).w);
+	PrimitiveData.LocalObjectBoundsMin					= LoadPrimitivePrimitiveSceneDataElement(PrimitiveIndex, 24).xyz;
+	PrimitiveData.ObjectBoundsY							= LoadPrimitivePrimitiveSceneDataElement(PrimitiveIndex, 24).w;
+	PrimitiveData.LocalObjectBoundsMax					= LoadPrimitivePrimitiveSceneDataElement(PrimitiveIndex, 25).xyz;
+	PrimitiveData.ObjectBoundsZ							= LoadPrimitivePrimitiveSceneDataElement(PrimitiveIndex, 25).w;
+	PrimitiveData.InstanceLocalBoundsCenter				= LoadPrimitivePrimitiveSceneDataElement(PrimitiveIndex, 26).xyz;
+	PrimitiveData.InstancePayloadDataOffset				= asuint(LoadPrimitivePrimitiveSceneDataElement(PrimitiveIndex, 26).w);
+	PrimitiveData.InstanceLocalBoundsExtent				= LoadPrimitivePrimitiveSceneDataElement(PrimitiveIndex, 27).xyz;
+	PrimitiveData.InstancePayloadDataStride				= asuint(LoadPrimitivePrimitiveSceneDataElement(PrimitiveIndex, 27).w);
+	PrimitiveData.WireframeColor						= LoadPrimitivePrimitiveSceneDataElement(PrimitiveIndex, 28).xyz;
+	PrimitiveData.PackedNaniteFlags						= asuint(LoadPrimitivePrimitiveSceneDataElement(PrimitiveIndex, 28).w);
+	PrimitiveData.LevelColor							= LoadPrimitivePrimitiveSceneDataElement(PrimitiveIndex, 29).xyz;
+	PrimitiveData.PersistentPrimitiveIndex				= asint(LoadPrimitivePrimitiveSceneDataElement(PrimitiveIndex, 29).w);
+	PrimitiveData.InstanceDrawDistanceMinMaxSquared		= LoadPrimitivePrimitiveSceneDataElement(PrimitiveIndex, 30).xy;
+	PrimitiveData.InstanceWPODisableDistanceSquared		= LoadPrimitivePrimitiveSceneDataElement(PrimitiveIndex, 30).z;
+	PrimitiveData.NaniteRayTracingDataOffset			= asuint(LoadPrimitivePrimitiveSceneDataElement(PrimitiveIndex, 30).w);
+	PrimitiveData.BoundsScale							= LoadPrimitivePrimitiveSceneDataElement(PrimitiveIndex, 31).x;
+	PrimitiveData.CustomStencilValueAndMask				= asuint(LoadPrimitivePrimitiveSceneDataElement(PrimitiveIndex, 31).y);
 	// .zw Unused
 
 	// NOTE: Please make sure GetPrimitiveDataFromUniformBuffer() gets updated as well when adding new members here
@@ -357,7 +371,7 @@ FPrimitiveSceneData GetPrimitiveData(uint PrimitiveId)
 	UNROLL
 	for (int DataIndex = 0; DataIndex < NUM_CUSTOM_PRIMITIVE_DATA; ++DataIndex)
 	{
-		PrimitiveData.CustomPrimitiveData[DataIndex] = LoadPrimitivePrimitiveSceneDataElement(PrimitiveIndex,  33 + DataIndex);
+		PrimitiveData.CustomPrimitiveData[DataIndex] = LoadPrimitivePrimitiveSceneDataElement(PrimitiveIndex,  32 + DataIndex);
 	}
 	return PrimitiveData;
 #endif
@@ -1028,14 +1042,9 @@ FInstanceSceneData GetInstanceSceneData(uint InstanceId, uint SOAStride, bool bC
 		else
 		{
 			// Since the instance is not dynamic, it cannot move relative to the primitive. However, the primitive itself can
-			// move. So we have to solve for previous Instance->World by solving primitive World->PrevWorld and using it to
+			// move. So we have to solve for previous Instance->World by taking primitive World->PrevWorld and using it to
 			// transform Instance->World into previous world space.
-			//
-			// TODO: Cut the number of matrix multiplies down by adding World->PrevWorld to FPrimitiveSceneData?
-			// Also, in the case of there being a single instance with identity Instance->Primitive (common), this can be
-			// simplified to merely copying the Primitive PrevLocalToWorld
-			float4x4 PrimRelWorldToPrevRelWorld = mul(PrimitiveData.WorldToLocal.M, PrimitiveData.PreviousLocalToWorld.M);
-			InstanceData.PrevLocalToWorld = MakeLWCMatrix(TilePosition, mul(LocalToRelativeWorld, PrimRelWorldToPrevRelWorld));
+			InstanceData.PrevLocalToWorld = MakeLWCMatrix(TilePosition, mul(LocalToRelativeWorld, PrimitiveData.WorldToPreviousWorld));
 		}
 
 	#if 1 //NEEDS_LIGHTMAP_COORDINATE
diff --git a/Engine/Source/Runtime/Engine/Private/PrimitiveUniformShaderParameters.cpp b/Engine/Source/Runtime/Engine/Private/PrimitiveUniformShaderParameters.cpp
--- a/Engine/Source/Runtime/Engine/Private/PrimitiveUniformShaderParameters.cpp
+++ b/Engine/Source/Runtime/Engine/Private/PrimitiveUniformShaderParameters.cpp
@@ -233,7 +233,6 @@ FPrimitiveSceneShaderData::FPrimitiveSceneShaderData(const FPrimitiveSceneProxy*
 
 void FPrimitiveSceneShaderData::Setup(const FPrimitiveUniformShaderParameters& PrimitiveUniformShaderParameters)
 {
-	static_assert(sizeof(FPrimitiveUniformShaderParameters) == sizeof(FPrimitiveSceneShaderData), "The FPrimitiveSceneShaderData manual layout below and in usf must match FPrimitiveUniformShaderParameters.  Update this assert when adding a new member.");
 	static_assert(NUM_LIGHTING_CHANNELS == 3, "The FPrimitiveSceneShaderData packing currently assumes a maximum of 3 lighting channels.");
 
 	// Note: layout must match GetPrimitiveData in usf
@@ -250,77 +249,84 @@ void FPrimitiveSceneShaderData::Setup(const FPrimitiveUniformShaderParameters& P
 	Data[1].Z	= PrimitiveUniformShaderParameters.TilePosition.Z;
 	Data[1].W	= FMath::AsFloat(PrimitiveUniformShaderParameters.PrimitiveComponentId);
 
-	Data[2]		= *(const FVector4f*)&PrimitiveUniformShaderParameters.LocalToRelativeWorld.M[0][0];
-	Data[3]		= *(const FVector4f*)&PrimitiveUniformShaderParameters.LocalToRelativeWorld.M[1][0];
-	Data[4]		= *(const FVector4f*)&PrimitiveUniformShaderParameters.LocalToRelativeWorld.M[2][0];
-	Data[5]		= *(const FVector4f*)&PrimitiveUniformShaderParameters.LocalToRelativeWorld.M[3][0];
+	// Pack these matrices into the buffer as float3x4 transposed
 
-	Data[6]		= *(const FVector4f*)&PrimitiveUniformShaderParameters.RelativeWorldToLocal.M[0][0];
-	Data[7]		= *(const FVector4f*)&PrimitiveUniformShaderParameters.RelativeWorldToLocal.M[1][0];
-	Data[8]		= *(const FVector4f*)&PrimitiveUniformShaderParameters.RelativeWorldToLocal.M[2][0];
-	Data[9]		= *(const FVector4f*)&PrimitiveUniformShaderParameters.RelativeWorldToLocal.M[3][0];
+	FMatrix44f LocalToRelativeWorldTranspose = PrimitiveUniformShaderParameters.LocalToRelativeWorld.GetTransposed();
+	Data[2]		= *(const FVector4f*)&LocalToRelativeWorldTranspose.M[0][0];
+	Data[3]		= *(const FVector4f*)&LocalToRelativeWorldTranspose.M[1][0];
+	Data[4]		= *(const FVector4f*)&LocalToRelativeWorldTranspose.M[2][0];
 
-	Data[10]	= *(const FVector4f*)&PrimitiveUniformShaderParameters.PreviousLocalToRelativeWorld.M[0][0];
-	Data[11]	= *(const FVector4f*)&PrimitiveUniformShaderParameters.PreviousLocalToRelativeWorld.M[1][0];
-	Data[12]	= *(const FVector4f*)&PrimitiveUniformShaderParameters.PreviousLocalToRelativeWorld.M[2][0];
-	Data[13]	= *(const FVector4f*)&PrimitiveUniformShaderParameters.PreviousLocalToRelativeWorld.M[3][0];
+	FMatrix44f RelativeWorldToLocalTranspose = PrimitiveUniformShaderParameters.RelativeWorldToLocal.GetTransposed();
+	Data[5]		= *(const FVector4f*)&RelativeWorldToLocalTranspose.M[0][0];
+	Data[6]		= *(const FVector4f*)&RelativeWorldToLocalTranspose.M[1][0];
+	Data[7]		= *(const FVector4f*)&RelativeWorldToLocalTranspose.M[2][0];
 
-	Data[14]	= *(const FVector4f*)&PrimitiveUniformShaderParameters.PreviousRelativeWorldToLocal.M[0][0];
-	Data[15]	= *(const FVector4f*)&PrimitiveUniformShaderParameters.PreviousRelativeWorldToLocal.M[1][0];
-	Data[16]	= *(const FVector4f*)&PrimitiveUniformShaderParameters.PreviousRelativeWorldToLocal.M[2][0];
-	Data[17]	= *(const FVector4f*)&PrimitiveUniformShaderParameters.PreviousRelativeWorldToLocal.M[3][0];
+	FMatrix44f PreviousLocalToRelativeWorldTranspose = PrimitiveUniformShaderParameters.PreviousLocalToRelativeWorld.GetTransposed();
+	Data[8]		= *(const FVector4f*)&PreviousLocalToRelativeWorldTranspose.M[0][0];
+	Data[9]		= *(const FVector4f*)&PreviousLocalToRelativeWorldTranspose.M[1][0];
+	Data[10]	= *(const FVector4f*)&PreviousLocalToRelativeWorldTranspose.M[2][0];
 
-	Data[18]	= FVector4f(PrimitiveUniformShaderParameters.InvNonUniformScale, PrimitiveUniformShaderParameters.ObjectBoundsX);
-	Data[19]	= PrimitiveUniformShaderParameters.ObjectRelativeWorldPositionAndRadius;
+	FMatrix44f PreviousRelativeWorldToLocalTranspose = PrimitiveUniformShaderParameters.PreviousRelativeWorldToLocal.GetTransposed();
+	Data[11]	= *(const FVector4f*)&PreviousRelativeWorldToLocalTranspose.M[0][0];
+	Data[12]	= *(const FVector4f*)&PreviousRelativeWorldToLocalTranspose.M[1][0];
+	Data[13]	= *(const FVector4f*)&PreviousRelativeWorldToLocalTranspose.M[2][0];
 
-	Data[20]	= FVector4f(PrimitiveUniformShaderParameters.ActorRelativeWorldPosition, 0.0f);
-	Data[20].W	= FMath::AsFloat(PrimitiveUniformShaderParameters.LightmapUVIndex);
+	FMatrix44f WorldToPreviousWorldTranspose = PrimitiveUniformShaderParameters.WorldToPreviousWorld.GetTransposed();
+	Data[14]	= *(const FVector4f*)&WorldToPreviousWorldTranspose.M[0][0];
+	Data[15]	= *(const FVector4f*)&WorldToPreviousWorldTranspose.M[1][0];
+	Data[16]	= *(const FVector4f*)&WorldToPreviousWorldTranspose.M[2][0];
 
-	Data[21]	= FVector4f(PrimitiveUniformShaderParameters.ObjectOrientation, 0.0f);
-	Data[21].W	= FMath::AsFloat(PrimitiveUniformShaderParameters.LightmapDataIndex);
+	Data[17]	= FVector4f(PrimitiveUniformShaderParameters.InvNonUniformScale, PrimitiveUniformShaderParameters.ObjectBoundsX);
+	Data[18]	= PrimitiveUniformShaderParameters.ObjectRelativeWorldPositionAndRadius;
 
-	Data[22]	= PrimitiveUniformShaderParameters.NonUniformScale;
+	Data[19]	= FVector4f(PrimitiveUniformShaderParameters.ActorRelativeWorldPosition, 0.0f);
+	Data[19].W	= FMath::AsFloat(PrimitiveUniformShaderParameters.LightmapUVIndex);
 
-	Data[23]	= FVector4f(PrimitiveUniformShaderParameters.PreSkinnedLocalBoundsMin, 0.0f);
-	Data[23].W	= FMath::AsFloat(PrimitiveUniformShaderParameters.NaniteResourceID);
+	Data[20]	= FVector4f(PrimitiveUniformShaderParameters.ObjectOrientation, 0.0f);
+	Data[20].W	= FMath::AsFloat(PrimitiveUniformShaderParameters.LightmapDataIndex);
 
-	Data[24]	= FVector4f(PrimitiveUniformShaderParameters.PreSkinnedLocalBoundsMax, 0.0f);
-	Data[24].W	= FMath::AsFloat(PrimitiveUniformShaderParameters.NaniteHierarchyOffset);
+	Data[21]	= PrimitiveUniformShaderParameters.NonUniformScale;
 
-	Data[25]	= FVector4f(PrimitiveUniformShaderParameters.LocalObjectBoundsMin, PrimitiveUniformShaderParameters.ObjectBoundsY);
-	Data[26]	= FVector4f(PrimitiveUniformShaderParameters.LocalObjectBoundsMax, PrimitiveUniformShaderParameters.ObjectBoundsZ);
+	Data[22]	= FVector4f(PrimitiveUniformShaderParameters.PreSkinnedLocalBoundsMin, 0.0f);
+	Data[22].W	= FMath::AsFloat(PrimitiveUniformShaderParameters.NaniteResourceID);
 
-	Data[27].X = PrimitiveUniformShaderParameters.InstanceLocalBoundsCenter.X;
-	Data[27].Y = PrimitiveUniformShaderParameters.InstanceLocalBoundsCenter.Y;
-	Data[27].Z = PrimitiveUniformShaderParameters.InstanceLocalBoundsCenter.Z;
-	Data[27].W = FMath::AsFloat(PrimitiveUniformShaderParameters.InstancePayloadDataOffset);
+	Data[23]	= FVector4f(PrimitiveUniformShaderParameters.PreSkinnedLocalBoundsMax, 0.0f);
+	Data[23].W	= FMath::AsFloat(PrimitiveUniformShaderParameters.NaniteHierarchyOffset);
 
-	Data[28].X = PrimitiveUniformShaderParameters.InstanceLocalBoundsExtent.X;
-	Data[28].Y = PrimitiveUniformShaderParameters.InstanceLocalBoundsExtent.Y;
-	Data[28].Z = PrimitiveUniformShaderParameters.InstanceLocalBoundsExtent.Z;
-	Data[28].W = FMath::AsFloat(PrimitiveUniformShaderParameters.InstancePayloadDataStride);
+	Data[24]	= FVector4f(PrimitiveUniformShaderParameters.LocalObjectBoundsMin, PrimitiveUniformShaderParameters.ObjectBoundsY);
+	Data[25]	= FVector4f(PrimitiveUniformShaderParameters.LocalObjectBoundsMax, PrimitiveUniformShaderParameters.ObjectBoundsZ);
+
+	Data[26].X = PrimitiveUniformShaderParameters.InstanceLocalBoundsCenter.X;
+	Data[26].Y = PrimitiveUniformShaderParameters.InstanceLocalBoundsCenter.Y;
+	Data[26].Z = PrimitiveUniformShaderParameters.InstanceLocalBoundsCenter.Z;
+	Data[26].W = FMath::AsFloat(PrimitiveUniformShaderParameters.InstancePayloadDataOffset);
+
+	Data[27].X = PrimitiveUniformShaderParameters.InstanceLocalBoundsExtent.X;
+	Data[27].Y = PrimitiveUniformShaderParameters.InstanceLocalBoundsExtent.Y;
+	Data[27].Z = PrimitiveUniformShaderParameters.InstanceLocalBoundsExtent.Z;
+	Data[27].W = FMath::AsFloat(PrimitiveUniformShaderParameters.InstancePayloadDataStride);
 	
-	Data[29].X = PrimitiveUniformShaderParameters.WireframeColor.X;
-	Data[29].Y = PrimitiveUniformShaderParameters.WireframeColor.Y;
-	Data[29].Z = PrimitiveUniformShaderParameters.WireframeColor.Z;
-	Data[29].W = FMath::AsFloat(PrimitiveUniformShaderParameters.PackedNaniteFlags);
-
-	Data[30].X = PrimitiveUniformShaderParameters.LevelColor.X;
-	Data[30].Y = PrimitiveUniformShaderParameters.LevelColor.Y;
-	Data[30].Z = PrimitiveUniformShaderParameters.LevelColor.Z;
-	Data[30].W = FMath::AsFloat(uint32(PrimitiveUniformShaderParameters.PersistentPrimitiveIndex));
-
-	Data[31].X = PrimitiveUniformShaderParameters.InstanceDrawDistanceMinMaxSquared.X;
-	Data[31].Y = PrimitiveUniformShaderParameters.InstanceDrawDistanceMinMaxSquared.Y;
-	Data[31].Z = PrimitiveUniformShaderParameters.InstanceWPODisableDistanceSquared;
-	Data[31].W = FMath::AsFloat(PrimitiveUniformShaderParameters.NaniteRayTracingDataOffset);
-
-	Data[32].X = PrimitiveUniformShaderParameters.BoundsScale;
-	Data[32].Y = FMath::AsFloat(PrimitiveUniformShaderParameters.CustomStencilValueAndMask);
+	Data[28].X = PrimitiveUniformShaderParameters.WireframeColor.X;
+	Data[28].Y = PrimitiveUniformShaderParameters.WireframeColor.Y;
+	Data[28].Z = PrimitiveUniformShaderParameters.WireframeColor.Z;
+	Data[28].W = FMath::AsFloat(PrimitiveUniformShaderParameters.PackedNaniteFlags);
+
+	Data[29].X = PrimitiveUniformShaderParameters.LevelColor.X;
+	Data[29].Y = PrimitiveUniformShaderParameters.LevelColor.Y;
+	Data[29].Z = PrimitiveUniformShaderParameters.LevelColor.Z;
+	Data[29].W = FMath::AsFloat(uint32(PrimitiveUniformShaderParameters.PersistentPrimitiveIndex));
+
+	Data[30].X = PrimitiveUniformShaderParameters.InstanceDrawDistanceMinMaxSquared.X;
+	Data[30].Y = PrimitiveUniformShaderParameters.InstanceDrawDistanceMinMaxSquared.Y;
+	Data[30].Z = PrimitiveUniformShaderParameters.InstanceWPODisableDistanceSquared;
+	Data[30].W = FMath::AsFloat(PrimitiveUniformShaderParameters.NaniteRayTracingDataOffset);
+
+	Data[31].X = PrimitiveUniformShaderParameters.BoundsScale;
+	Data[31].Y = FMath::AsFloat(PrimitiveUniformShaderParameters.CustomStencilValueAndMask);
 	// .ZW Unused
 
 	// Set all the custom primitive data float4. This matches the loop in SceneData.ush
-	const int32 CustomPrimitiveDataStartIndex = 33;
+	const int32 CustomPrimitiveDataStartIndex = 32;
 	for (int32 DataIndex = 0; DataIndex < FCustomPrimitiveData::NumCustomPrimitiveDataFloat4s; ++DataIndex)
 	{
 		Data[CustomPrimitiveDataStartIndex + DataIndex] = PrimitiveUniformShaderParameters.CustomPrimitiveData[DataIndex];
diff --git a/Engine/Source/Runtime/Engine/Public/PrimitiveUniformShaderParameters.h b/Engine/Source/Runtime/Engine/Public/PrimitiveUniformShaderParameters.h
--- a/Engine/Source/Runtime/Engine/Public/PrimitiveUniformShaderParameters.h
+++ b/Engine/Source/Runtime/Engine/Public/PrimitiveUniformShaderParameters.h
@@ -34,6 +34,7 @@ BEGIN_GLOBAL_SHADER_PARAMETER_STRUCT(FPrimitiveUniformShaderParameters,ENGINE_AP
 	SHADER_PARAMETER(FMatrix44f,	RelativeWorldToLocal)									// Rarely needed
 	SHADER_PARAMETER(FMatrix44f,	PreviousLocalToRelativeWorld)							// Used to calculate velocity
 	SHADER_PARAMETER(FMatrix44f,	PreviousRelativeWorldToLocal)							// Rarely used when calculating velocity, if material uses vertex offset along with world->local transform
+	SHADER_PARAMETER(FMatrix44f,	WorldToPreviousWorld)									// Used when calculating instance prev local->world for static instances that do not store it (calculated via doubles to resolve precision issues)
 	SHADER_PARAMETER_EX(FVector3f,	InvNonUniformScale,  EShaderPrecisionModifier::Half)	// Often needed
 	SHADER_PARAMETER(float,			ObjectBoundsX)											// Only needed for editor/development
 	SHADER_PARAMETER(FVector4f,		ObjectRelativeWorldPositionAndRadius)					// Needed by some materials
@@ -370,28 +371,31 @@ public:
 
 		{
 			// Inverse on FMatrix44f can generate NaNs if the source matrix contains large scaling, so do it in double precision.
+			// Also use double precision to calculate WorldToPreviousWorld to prevent precision issues at far distances
 			FMatrix LocalToRelativeWorld = FLargeWorldRenderScalar::MakeToRelativeWorldMatrixDouble(TilePositionOffset, AbsoluteLocalToWorld);
+			FMatrix PrevLocalToRelativeWorld = FLargeWorldRenderScalar::MakeClampedToRelativeWorldMatrixDouble(TilePositionOffset, AbsolutePreviousLocalToWorld);
+			FMatrix RelativeWorldToLocal = LocalToRelativeWorld.Inverse();
+
 			Parameters.LocalToRelativeWorld = FMatrix44f(LocalToRelativeWorld);
-			Parameters.RelativeWorldToLocal = FMatrix44f(LocalToRelativeWorld.Inverse());
+			Parameters.RelativeWorldToLocal = FMatrix44f(RelativeWorldToLocal);
+			Parameters.WorldToPreviousWorld = FMatrix44f(RelativeWorldToLocal * PrevLocalToRelativeWorld);
+
+			if (bHasPreviousLocalToWorld)
+			{
+				Parameters.PreviousLocalToRelativeWorld = FMatrix44f(PrevLocalToRelativeWorld);
+				Parameters.PreviousRelativeWorldToLocal = FMatrix44f(PrevLocalToRelativeWorld.Inverse());
+			}
+			else
+			{
+				Parameters.PreviousLocalToRelativeWorld = Parameters.LocalToRelativeWorld;
+				Parameters.PreviousRelativeWorldToLocal = Parameters.RelativeWorldToLocal;
+			}
 		}
 
 		Parameters.ActorRelativeWorldPosition = FVector3f(AbsoluteActorWorldPosition - TilePositionOffset);	//LWC_TODO: Precision loss
 		const FVector3f ObjectRelativeWorldPositionAsFloat = FVector3f(AbsoluteObjectWorldPosition - TilePositionOffset);
 		Parameters.ObjectRelativeWorldPositionAndRadius = FVector4f(ObjectRelativeWorldPositionAsFloat, ObjectRadius);
 
-		if (bHasPreviousLocalToWorld)
-		{
-			// Inverse on FMatrix44f can generate NaNs if the source matrix contains large scaling, so do it in double precision.
-			FMatrix PrevLocalToRelativeWorld = FLargeWorldRenderScalar::MakeClampedToRelativeWorldMatrixDouble(TilePositionOffset, AbsolutePreviousLocalToWorld);
-			Parameters.PreviousLocalToRelativeWorld = FMatrix44f(PrevLocalToRelativeWorld);
-			Parameters.PreviousRelativeWorldToLocal = FMatrix44f(PrevLocalToRelativeWorld.Inverse());
-		}
-		else
-		{
-			Parameters.PreviousLocalToRelativeWorld = Parameters.LocalToRelativeWorld;
-			Parameters.PreviousRelativeWorldToLocal = Parameters.RelativeWorldToLocal;
-		}
-
 		if (!bHasInstanceLocalBounds)
 		{
 			FRenderBounds InstanceLocalBounds(Parameters.LocalObjectBoundsMin, Parameters.LocalObjectBoundsMax);
@@ -562,7 +566,7 @@ extern ENGINE_API TGlobalResource<FIdentityPrimitiveUniformBuffer> GIdentityPrim
 struct FPrimitiveSceneShaderData
 {
 	// Must match PRIMITIVE_SCENE_DATA_STRIDE in SceneData.ush
-	enum { DataStrideInFloat4s = 42 };
+	enum { DataStrideInFloat4s = 41 };
 
 	TStaticArray<FVector4f, DataStrideInFloat4s> Data;
 
diff --git a/Engine/Source/Runtime/RenderCore/Public/UnifiedBuffer.h b/Engine/Source/Runtime/RenderCore/Public/UnifiedBuffer.h
--- a/Engine/Source/Runtime/RenderCore/Public/UnifiedBuffer.h
+++ b/Engine/Source/Runtime/RenderCore/Public/UnifiedBuffer.h
@@ -68,7 +68,7 @@ RENDERCORE_API bool ResizeResourceIfNeeded(FRDGBuilder& GraphBuilder, FRWBuffer&
 class FScatterUploadBuffer
 {
 public:
-	enum { PrimitiveDataStrideInFloat4s = 42 };
+	enum { PrimitiveDataStrideInFloat4s = 41 };
 
 	FByteAddressBuffer ScatterBuffer;
 	FByteAddressBuffer UploadBuffer;
