Perform sRGB conversion for base color when using compute shader (need to do this manually with UAVs).

diff --git a/Engine/Shaders/Private/ComputeShaderOutputCommon.ush b/Engine/Shaders/Private/ComputeShaderOutputCommon.ush
--- a/Engine/Shaders/Private/ComputeShaderOutputCommon.ush
+++ b/Engine/Shaders/Private/ComputeShaderOutputCommon.ush
@@ -7,6 +7,7 @@ through a single struct, allowing for a more readable code
 =============================================================================*/
 
 #include "ShaderOutputCommon.ush"
+#include "GammaCorrectionCommon.ush"
 
 // TEMP: For initial debugging
 #include "Visualization.ush"
@@ -116,6 +117,9 @@ uint Pack_B8G8R8A8_UNorm(precise float4 Input)
 	return Output;
 }
 
+// TODO: Is this ever used? (see bHighPrecisionGBuffers)
+#define HIGH_PRECISION_GBUFFERS 0
+
 [numthreads(8, 8, 1)]
 void MainCS(uint ThreadIndex : SV_GroupThreadID, uint3 GroupID : SV_GroupID, uint3 DTID : SV_DispatchThreadID)
 {
@@ -168,12 +172,17 @@ void MainCS(uint ThreadIndex : SV_GroupThreadID, uint3 GroupID : SV_GroupID, uin
 
 #if PIXELSHADEROUTPUT_MRT2
 	//OutTarget2UAV[PixelPos] = float4(IntToColor(MaterialDepthId), PixelShaderOut.MRT[2].w);
-	OutTarget2UAV[PixelPos] = PixelShaderOut.MRT[2].w;
+	OutTarget2UAV[PixelPos] = PixelShaderOut.MRT[2];
 #endif
 
 #if PIXELSHADEROUTPUT_MRT3
-	//OutTarget3UAV[PixelPos] = Pack_B8G8R8A8_UNorm(PixelShaderOut.MRT[3]);
-	OutTarget3UAV[PixelPos] = PixelShaderOut.MRT[3];
+	#if HIGH_PRECISION_GBUFFERS
+		OutTarget3UAV[PixelPos] = PixelShaderOut.MRT[3];
+	#else
+		// const bool bLegacyAlbedoSrgb = true;
+		//Info.Targets[3].Init(DiffuseAndSpecularGBufferFormat,  TEXT("GBufferC"), bLegacyAlbedoSrgb && !bHighPrecisionGBuffers,  true,  true,  true);
+		OutTarget3UAV[PixelPos] = float4(LinearToSrgb(PixelShaderOut.MRT[3].rgb), PixelShaderOut.MRT[3].a); // BaseColor is sRGB
+	#endif
 #endif
 
 #if STRATA_OPAQUE_DEFERRED
