PathTracing: Fix incorrect sign in medium sampling function

Discovered the sampling call was flipped the wrong way around which made it inconsistent with light sampling and could produce subtly wrong results on volumetric multiple scattering when the henyey-greenstein lobe was dominant.

diff --git a/Engine/Shaders/Private/PathTracing/Material/PathTracingMedium.ush b/Engine/Shaders/Private/PathTracing/Material/PathTracingMedium.ush
--- a/Engine/Shaders/Private/PathTracing/Material/PathTracingMedium.ush
+++ b/Engine/Shaders/Private/PathTracing/Material/PathTracingMedium.ush
@@ -33,7 +33,7 @@ FMaterialSample Medium_SampleMaterial(
 	const float SinTheta = sqrt(saturate(1.0 - CosTheta * CosTheta));
 	const float Phi = RandSample.y * (2.0 * PI);
 
-	const float3 L_World = TangentToWorld(float3(cos(Phi) * SinTheta, sin(Phi) * SinTheta, CosTheta), V_World);
+	const float3 L_World = TangentToWorld(float3(cos(Phi) * SinTheta, sin(Phi) * SinTheta, CosTheta), -V_World);
 
 	FMaterialSample Result = CreateMaterialSample(L_World, 0.0, 0.0, 1.0, 1.0, PATHTRACER_SCATTER_VOLUME);
 	Result.AddLobeWithMIS(RayleighWeight, RayleighPhase(CosTheta), LobeProb);
@@ -55,7 +55,7 @@ FMaterialEval Medium_EvalMaterial(
 
 #if USE_UNIFORM_PHASE_FUNCTION
 	// simple omni-directional evaluation
-	return CreateMaterialEval(RayleighWeight + HGWeight, 1.0 / (4 * PI));
+	return CreateMaterialEval((RayleighWeight + HGWeight) * DiffuseSpecularScale.x, 1.0 / (4 * PI));
 #else
 	const float LobeProb = LobeSelectionProb(RayleighWeight, HGWeight);
 
