Remapped compute base pass group threads into pixel-like 2x2 quads to leverage the 4 lane cross bar and other optimizations.

diff --git a/Engine/Shaders/Private/ComputeShaderOutputCommon.ush b/Engine/Shaders/Private/ComputeShaderOutputCommon.ush
--- a/Engine/Shaders/Private/ComputeShaderOutputCommon.ush
+++ b/Engine/Shaders/Private/ComputeShaderOutputCommon.ush
@@ -8,6 +8,7 @@ through a single struct, allowing for a more readable code
 
 #include "ShaderOutputCommon.ush"
 #include "GammaCorrectionCommon.ush"
+#include "MortonCode.ush"
 
 // TEMP: For initial debugging
 #include "Visualization.ush"
@@ -82,50 +83,38 @@ through a single struct, allowing for a more readable code
 #endif // STRATA_OPAQUE_DEFERRED
 
 // .xy = min, .zw = max
-uint4 ViewRect;
-
-uint FloatToUInt(float Value, float Scale)
-{
-	return (uint)floor(Value * Scale + 0.5f);
-}
-
-uint Pack_R10G10B10A2_UNorm(precise float4 Input)
-{
-	uint Output;
-	Output =
-	(
-		(FloatToUInt(saturate(Input.x), 1023))			|
-		(FloatToUInt(saturate(Input.y), 1023) << 10)	|
-		(FloatToUInt(saturate(Input.z), 1023) << 20)	|
-		(FloatToUInt(saturate(Input.w),    3) << 30)
-	);
-
-	return Output;
-}
-
-uint Pack_B8G8R8A8_UNorm(precise float4 Input)
-{
-	uint Output;
-	Output =
-	(
-		(FloatToUInt(saturate(Input.z), 255))			|
-		(FloatToUInt(saturate(Input.y), 255) << 8)		|
-		(FloatToUInt(saturate(Input.x), 255) << 16)		|
-		(FloatToUInt(saturate(Input.w), 255) << 24)
-	);
-
-	return Output;
-}
+uint4 ViewRect; // TODO: Not bound yet
 
 // TODO: Is this ever used? (see bHighPrecisionGBuffers)
 #define HIGH_PRECISION_GBUFFERS 0
 
+// Just like pixel shaders, compute 2x2 CS quads from thread index so we can leverage the 4 lane cross bar.
+// [0 1][2 3] -> [0 1][4 5]
+// [4 5][6 7] -> [2 3][7 6]
+#define REMAP_SHADING_QUADS 1
+
 [numthreads(8, 8, 1)]
-void MainCS(uint ThreadIndex : SV_GroupThreadID, uint3 GroupID : SV_GroupID, uint3 DTID : SV_DispatchThreadID)
+void MainCS(uint ThreadIndex : SV_GroupIndex, uint3 GroupID : SV_GroupID, uint3 DTID : SV_DispatchThreadID)
 {
+#if REMAP_SHADING_QUADS
+	// Just like pixel shaders, compute 2x2 CS quads from thread index so we can leverage the 4 lane cross bar.
+	// [0 1][2 3] -> [0 1][4 5]
+	// [4 5][6 7] -> [2 3][7 6]
+#if 1
+	const uint X = BitFieldInsertU32(BitFieldMaskU32(1u, 0u), ThreadIndex, BitFieldExtractU32(ThreadIndex, 3u, 2u));
+	const uint Y = BitFieldInsertU32(BitFieldMaskU32(2u, 0u), BitFieldExtractU32(ThreadIndex, 2u, 1u), BitFieldExtractU32(ThreadIndex, 3u, 3u));
+#else
+	const uint2 Decoded = MortonDecode(ThreadIndex);
+	const uint X = Decoded.x;
+	const uint Y = Decoded.y;
+#endif
+	const uint2 PixelCoord = uint2(X, Y) + GroupID.xy * 8u;
+#else
 	const uint2 PixelCoord = DTID.xy;
-	const float2 HalfPixelOffset = float2(0.5, 0.5);
-	const float2 PixelPos = float2(PixelCoord) + HalfPixelOffset;
+#endif
+
+	const float2 HalfPixelOffset = float2(0.5f, 0.5f);
+	const float2 PixelPos   = float2(PixelCoord /* TODO: + ViewRect.xy */) + HalfPixelOffset;
 
 #if IS_NANITE_PASS
 	FResolvedMaterial ResolvedMaterial = UnpackMaterialResolve(Nanite.MaterialResolve[PixelPos]);
@@ -171,7 +160,6 @@ void MainCS(uint ThreadIndex : SV_GroupThreadID, uint3 GroupID : SV_GroupID, uin
 #endif
 
 #if PIXELSHADEROUTPUT_MRT2
-	//OutTarget2UAV[PixelPos] = float4(IntToColor(MaterialDepthId), PixelShaderOut.MRT[2].w);
 	OutTarget2UAV[PixelPos] = PixelShaderOut.MRT[2];
 #endif
 
@@ -182,6 +170,10 @@ void MainCS(uint ThreadIndex : SV_GroupThreadID, uint3 GroupID : SV_GroupID, uin
 		// const bool bLegacyAlbedoSrgb = true;
 		//Info.Targets[3].Init(DiffuseAndSpecularGBufferFormat,  TEXT("GBufferC"), bLegacyAlbedoSrgb && !bHighPrecisionGBuffers,  true,  true,  true);
 		OutTarget3UAV[PixelPos] = float4(LinearToSrgb(PixelShaderOut.MRT[3].rgb), PixelShaderOut.MRT[3].a); // BaseColor is sRGB
+
+		//float3 DebugVal = IntToColor(MaterialDepthId);
+		//float3 DebugVal = IntToColor(ThreadIndex >> 2u);
+		//OutTarget3UAV[PixelPos] = float4(LinearToSrgb(DebugVal), PixelShaderOut.MRT[3].w);
 	#endif
 #endif
 
