Add shader print assert.

diff --git a/Engine/Shaders/Private/ShaderPrint.ush b/Engine/Shaders/Private/ShaderPrint.ush
--- a/Engine/Shaders/Private/ShaderPrint.ush
+++ b/Engine/Shaders/Private/ShaderPrint.ush
@@ -201,3 +201,33 @@ void AddOBB(float3 Min, float3 Max, float4 Color, float4x4 Transform)
 {
 	AddOBB(InitShaderPrintContext(), Min, Max, Color, Transform);
 }
+
+// Assert (experimental)
+// Can be used in any shader binding ShaderPrint uniform buffer. It will print the 
+// assert for all pixels/threads not complying to the condition. All print will be 
+// stacked correctly.
+// 
+// Example:
+//  Assert(DispatchThreadId < MaxThread, TEXT("Error: ThreadId higher than MaxThread");
+// 
+// This function also returns a FShaderPrintContext so that one can add more details if needed.
+//
+// Example:
+//  FShaderPrintContext Ctx = Assert(DispatchThreadId < MaxThread, TEXT("Error: thread higher MaxThread");
+//  Print(Ctx, DispatchThreadId, FontWhite);
+//
+FShaderPrintContext Assert(bool Condition, FShaderPrintText Text, uint2 InWritePos = uint2(50, 150))
+{
+	FShaderPrintContext OutCtx = InitShaderPrintContext(false, InWritePos);
+	if (!Condition)
+	{
+		uint LineOffset = 0;
+		InterlockedAdd(ShaderPrint_RWEntryBuffer[3] /* Free counter */, 1, LineOffset);
+		InWritePos.y += LineOffset * 12;
+
+		OutCtx = InitShaderPrintContext(true, InWritePos);
+		Print(OutCtx, Text, FontRed);
+	}
+
+	return OutCtx;
+}
