Support blending of mesh decals in path tracer
- Mesh decals are applied back to front to simplify logic. Front to back would require accumulating decal properties in temporary variables and tracking transparency per property due to decals with different write flags.
- Use iterative closest hit rays to guarantee correct order and avoid issues with duplicate AHS execution.

diff --git a/Engine/Shaders/Private/PathTracing/PathTracingCore.ush b/Engine/Shaders/Private/PathTracing/PathTracingCore.ush
--- a/Engine/Shaders/Private/PathTracing/PathTracingCore.ush
+++ b/Engine/Shaders/Private/PathTracing/PathTracingCore.ush
@@ -229,6 +229,46 @@ void ApplyDecal(inout FPathTracingPayload HitPayload, FDecalShaderParams Params)
 	// TODO: Payload.Opacity?
 }
 
+void TraceDecalRay(FRayDesc DecalRay, bool bReverseCulling, uint InstanceInclusionMask, inout FPathTracingPayload HitPayload, inout float DecalThroughput)
+{
+	const uint MissShaderIndex = 0; // TODO
+
+	for (;;)
+	{
+		FRayTracingDecalPayload DecalPayload = (FRayTracingDecalPayload)0;
+		DecalPayload.Params.SetInputTranslatedWorldPosition(HitPayload.TranslatedWorldPos);
+
+		TraceRay(
+			DecalTLAS,
+			bReverseCulling ? RAY_FLAG_CULL_FRONT_FACING_TRIANGLES : RAY_FLAG_CULL_BACK_FACING_TRIANGLES,
+			InstanceInclusionMask,
+			RAY_TRACING_SHADER_SLOT_MATERIAL,
+			RAY_TRACING_NUM_SHADER_SLOTS,
+			MissShaderIndex,
+			DecalRay.GetNativeDesc(),
+			DecalPayload);
+
+		if (DecalPayload.IsMiss())
+		{
+			break;
+		}
+
+		if (bReverseCulling)
+		{
+			// also reverse normal
+			DecalPayload.Params.SetWorldNormal(-DecalPayload.Params.GetWorldNormal());
+		}
+		
+		ApplyDecal(HitPayload, DecalPayload.Params);
+		
+		DecalThroughput *= DecalPayload.Params.GetTransparency();
+
+		// prepare next step around the loop
+		// retrace the exact same ray with TMin one ulp past the hit we just found
+		DecalRay.TMin = asfloat(asuint(DecalPayload.HitT) + 1);
+	}
+}
+
 FPathTracingPayload TraceTransparentRay(inout FPathState PathState, bool bIsCameraRay, bool bLastBounce, bool bIncludeEmission, uint NumVisibleLights, inout FVolumeSegment VolumeSegment)
 {
 	const uint RayFlags = bIsCameraRay && EnableCameraBackfaceCulling ? RAY_FLAG_CULL_BACK_FACING_TRIANGLES : 0;
@@ -451,38 +491,22 @@ FPathTracingPayload TraceTransparentRay(inout FPathState PathState, bool bIsCame
 			FRayDesc DecalRay = (FRayDesc)0;
 			if (MeshDecalsUseViewRay)
 			{
-				DecalRay = PathState.Ray;
 				// go a bit further than the current hit, to handle decals partially inside the model
-				// NOTE: TMin will be inherited from current ray
+				DecalRay.Direction = -PathState.Ray.Direction;
+				DecalRay.Origin = HitPayload.TranslatedWorldPos - DecalRay.Direction * MeshDecalBias;
+				DecalRay.TMin = 0.0f;
 				DecalRay.TMax = HitPayload.HitT + MeshDecalBias;
 			}
 			else
 			{
-				DecalRay.Direction = -HitPayload.WorldGeoNormal;
+				DecalRay.Direction = HitPayload.WorldGeoNormal;
 				DecalRay.Origin = HitPayload.TranslatedWorldPos - DecalRay.Direction * MeshDecalBias;
 				DecalRay.TMin = 0.0f;
 				DecalRay.TMax = MeshDecalBias * 2;
 			}
 
-			FRayTracingDecalPayload DecalPayload = (FRayTracingDecalPayload)0;
-			DecalPayload.Params.SetInputTranslatedWorldPosition(HitPayload.TranslatedWorldPos);
-
-			TraceRay(
-				DecalTLAS,
-				RAY_FLAG_CULL_BACK_FACING_TRIANGLES,
-				InstanceInclusionMask,
-				RAY_TRACING_SHADER_SLOT_MATERIAL,
-				RAY_TRACING_NUM_SHADER_SLOTS,
-				MissShaderIndex,
-				DecalRay.GetNativeDesc(),
-				DecalPayload);
-
-			if (!DecalPayload.IsMiss())
-			{
-				ApplyDecal(HitPayload, DecalPayload.Params);
-
-				DecalThroughput *= DecalPayload.Params.GetTransparency();
-			}
+			const bool bReverseCulling = true;
+			TraceDecalRay(DecalRay, bReverseCulling, InstanceInclusionMask, HitPayload, DecalThroughput);
 		}
 
 #if PLATFORM_SUPPORTS_CALLABLE_SHADERS
diff --git a/Engine/Shaders/Private/RayTracing/RayTracingDecalMaterialShader.usf b/Engine/Shaders/Private/RayTracing/RayTracingDecalMaterialShader.usf
--- a/Engine/Shaders/Private/RayTracing/RayTracingDecalMaterialShader.usf
+++ b/Engine/Shaders/Private/RayTracing/RayTracingDecalMaterialShader.usf
@@ -211,10 +211,8 @@ RAY_TRACING_ENTRY_ANY_HIT(RayTracingDecalMaterialAHS,
 		// #endif
 	}
 
-	// We don't support blending of mesh decals in ray tracing at the moment.
-	// Instead we only use alpha masking to provide some limited overlap support, especially useful when large meshes are used to render sparse decals.
 	float Opacity = GetMaterialOpacity(PixelMaterialInputs);
-	if (Opacity <= 0.01f)
+	if (Opacity <= 0.00001f)
 	{
 		IgnoreHit();
 		return;
