Strata: Fix excess energy with cloth shading after changes in 23734405

Also modify the path tracer and ClothBxDF to match the new energy conservation behavior.

diff --git a/Engine/Shaders/Private/PathTracing/Material/PathTracingCloth.ush b/Engine/Shaders/Private/PathTracing/Material/PathTracingCloth.ush
--- a/Engine/Shaders/Private/PathTracing/Material/PathTracingCloth.ush
+++ b/Engine/Shaders/Private/PathTracing/Material/PathTracingCloth.ush
@@ -16,7 +16,7 @@ struct FClothData
 	float3 V;
 
 	FBxDFEnergyTermsRGB Spec;
-	FBxDFEnergyTermsRGB Fuzz;
+	FBxDFEnergyTermsA Fuzz;
 
 	float3 DiffWeight;
 
@@ -36,9 +36,11 @@ FClothData PrepareClothData(FPathTracingPayload Payload, float3 V_World)
 
 	const float NoV = saturate(Data.V.z);
 	Data.Spec = ComputeGGXSpecEnergyTermsRGB(Payload.Roughness, NoV, Payload.SpecularColor);
-	Data.Fuzz = ComputeClothEnergyTermsRGB(Payload.Roughness, NoV, Payload.GetClothColor());
+	Data.Fuzz = ComputeClothEnergyTermsA(Payload.Roughness, NoV);
 
-	Data.DiffWeight = (1.0 - lerp(Data.Spec.E, Data.Fuzz.E, ClothBlend)) * Payload.DiffuseColor;
+	const float SpecEp = ComputeEnergyPreservation(Data.Spec);
+	const float FuzzEp = ComputeEnergyPreservation(Data.Fuzz);
+	Data.DiffWeight = lerp(SpecEp, FuzzEp, ClothBlend) * Payload.DiffuseColor;
 
 	Data.LobeCdf = LobeSelectionCdf(Data.DiffWeight, ClothBlend * Data.Fuzz.E, (1.0 - ClothBlend) * Data.Spec.E);
 	Data.LobePdf = LobeSelectionPdf(Data.LobeCdf);
@@ -77,7 +79,7 @@ FMaterialEval Cloth_EvalMaterial(
 	// Cloth Lobe (see ClothBxDF - it is counted as specular)
 	const float DCloth = D_InvGGX(Pow4(Payload.Roughness), NoH);
 	const float3 FCloth = F_Schlick(Payload.GetClothColor(), VoH);
-	const float3 ClothWeight = NoL > 0 && NoV > 0 ? (PI * ClothBlend * DCloth * Vis_Cloth(NoV, NoL)) * FCloth * Data.Fuzz.W : 0.0;
+	const float3 ClothWeight = NoL > 0 && NoV > 0 ? (PI * ClothBlend * DCloth * Vis_Cloth(NoV, NoL)) * FCloth * ComputeEnergyConservation(Data.Fuzz) : 0.0;
 	Result.AddLobeWithMIS(ClothWeight * ShadowTerminator * DiffuseSpecularScale.y, NoL / PI, Data.LobePdf.y);
 
 	// Specular lobe
@@ -164,7 +166,7 @@ FMaterialSample Cloth_SampleMaterial(
 	{
 		const float DCloth = D_InvGGX(Pow4(Payload.Roughness), NoH);
 		const float3 FCloth = F_Schlick(Payload.GetClothColor(), VoH);
-		const float3 ClothWeight = NoL > 0 && NoV > 0 ? (PI * ClothBlend * DCloth * Vis_Cloth(NoV, NoL)) * FCloth * Data.Fuzz.W : 0.0;
+		const float3 ClothWeight = NoL > 0 && NoV > 0 ? (PI * ClothBlend * DCloth * Vis_Cloth(NoV, NoL)) * FCloth * ComputeEnergyConservation(Data.Fuzz) : 0.0;
 		Result.AddLobeWithMIS(ClothWeight * ShadowTerminator, NoL / PI, Data.LobePdf.y);
 		Result.ScatterType = PATHTRACER_SCATTER_SPECULAR;
 	}
diff --git a/Engine/Shaders/Private/ShadingEnergyConservationTemplate.ush b/Engine/Shaders/Private/ShadingEnergyConservationTemplate.ush
--- a/Engine/Shaders/Private/ShadingEnergyConservationTemplate.ush
+++ b/Engine/Shaders/Private/ShadingEnergyConservationTemplate.ush
@@ -74,28 +74,23 @@ BXDF_ENERGY_SUFFIX(FBxDFEnergyTerms) BXDF_ENERGY_SUFFIX(ComputeGGXSpecEnergyTerm
 	return BXDF_ENERGY_SUFFIX(ComputeGGXSpecEnergyTerms)(Roughness, NoV, F0, F90);
 }
 
-BXDF_ENERGY_SUFFIX(FBxDFEnergyTerms) BXDF_ENERGY_SUFFIX(ComputeClothEnergyTerms)(float Roughness, float NoV, float3 F0, float3 F90)
+BXDF_ENERGY_SUFFIX(FBxDFEnergyTerms) BXDF_ENERGY_SUFFIX(ComputeClothEnergyTerms)(float Roughness, float NoV)
 {
 	BXDF_ENERGY_SUFFIX(FBxDFEnergyTerms) Out;
 #if USE_ENERGY_CONSERVATION > 0
 	{
-		Out = BXDF_ENERGY_SUFFIX(ComputeFresnelEnergyTerms)(ClothEnergyLookup(Roughness, NoV), F0, F90);
+		Out.W = 1.f;
+		Out.E = ClothEnergyLookup(Roughness, NoV).x;
 	}
 #else
 	{
 		Out.W = 1.0f;
-		Out.E = BXDF_ENERGY_SUFFIX(GetF0F90)(F0);
+		Out.E = 1.0f;
 	}
 #endif
 	return Out;
 }
 
-BXDF_ENERGY_SUFFIX(FBxDFEnergyTerms) BXDF_ENERGY_SUFFIX(ComputeClothEnergyTerms)(float Roughness, float NoV, float3 F0)
-{
-	const float F90 = saturate(50.0 * F0.g); // See F_Schlick implementation
-	return BXDF_ENERGY_SUFFIX(ComputeClothEnergyTerms)(Roughness, NoV, F0, F90);
-}
-
 BXDF_ENERGY_SUFFIX(FBxDFEnergyTerms) BXDF_ENERGY_SUFFIX(ComputeDiffuseEnergyTerms)(float Roughness, float NoV)
 {
 	BXDF_ENERGY_SUFFIX(FBxDFEnergyTerms) Out;
diff --git a/Engine/Shaders/Private/ShadingModels.ush b/Engine/Shaders/Private/ShadingModels.ush
--- a/Engine/Shaders/Private/ShadingModels.ush
+++ b/Engine/Shaders/Private/ShadingModels.ush
@@ -772,7 +772,7 @@ FDirectLighting ClothBxDF( FGBufferData GBuffer, half3 N, half3 V, half3 L, floa
 	float3 F2 = F_Schlick( FuzzColor, Context.VoH );
 	float3 Spec2 = AreaLight.FalloffColor * (Falloff * NoL) * (D2 * Vis2) * F2;
 	
-	const FBxDFEnergyTerms EnergyTerms2 = ComputeClothEnergyTerms(GBuffer.Roughness, Context.NoV, FuzzColor);
+	const FBxDFEnergyTermsA EnergyTerms2 = ComputeClothEnergyTermsA(GBuffer.Roughness, Context.NoV);
 	Spec2 *= ComputeEnergyConservation(EnergyTerms2);
 
 	FDirectLighting Lighting;
diff --git a/Engine/Shaders/Private/Strata/StrataEvaluation.ush b/Engine/Shaders/Private/Strata/StrataEvaluation.ush
--- a/Engine/Shaders/Private/Strata/StrataEvaluation.ush
+++ b/Engine/Shaders/Private/Strata/StrataEvaluation.ush
@@ -613,11 +613,10 @@ FStrataEvaluateResult StrataEvaluateBSDFCommon(FStrataBSDFContext BSDFContext, F
 			{
 				const float  FuzzAmount = SLAB_FUZZ_AMOUNT(BSDFContext.BSDF);
 				const float3 FuzzF0  = SLAB_FUZZ_COLOR(BSDFContext.BSDF);
-				const float  FuzzF90 = 1.0f;
 				const float  FuzzRoughness = MakeRoughnessSafe(SLAB_FUZZ_ROUGHNESS(BSDFContext.BSDF), STRATA_MIN_FUZZ_ROUGHNESS);
 
 				// F0 is ignored, as energy preservation is handle solely based on the white directional albedo and FuzzAmount
-				FBxDFEnergyTerms EnergyTerms = ComputeClothEnergyTerms(FuzzRoughness, BSDFContext.Context.NoV, 1.0f /*FuzzF0*/);
+				FBxDFEnergyTermsA EnergyTerms = ComputeClothEnergyTermsA(FuzzRoughness, BSDFContext.Context.NoV);
 
 				// The specular and diffuse components below the fuzz are only attenuated linearly according to the amount of fuzz.
 				const float3 Cloth_DirectionalAlbedo_SpecularTransmission = lerp(1.0, ComputeEnergyPreservation(EnergyTerms), FuzzAmount);
@@ -626,7 +625,9 @@ FStrataEvaluateResult StrataEvaluateBSDFCommon(FStrataBSDFContext BSDFContext, F
 				float ClothD   = D_Charlie(FuzzRoughness, BSDFContext.Context.NoH);
 				float ClothVis = Vis_Ashikhmin(BSDFContext.Context.NoV, AreaLightContext.NoL);
 				float3 ClothF  = F_Schlick(FuzzF0, BSDFContext.Context.VoH);
-				float3 ClothSpecularPathValue = (ClothD * ClothVis) * ClothF * ComputeEnergyConservation(EnergyTerms);
+				// ComputeEnergyConservation(EnergyTerms) is not taken into account because we do not white white fuzz to reflect all incoming energy
+				// The fuzz should have some transmission to the lower lobes. In order to pass a furnace test, a combination of white diffuse+white fuzz is required
+				float3 ClothSpecularPathValue = (ClothD * ClothVis) * ClothF;
 
 				// Area light are not supported by the cloth BRDF
 				float3 ClothIntegratedSpecularValue = (ShadowTerms.SurfaceShadow * AreaLightContext.NoL * AreaLightContext.Falloff * FuzzAmount) * AreaLightContext.AreaLight.FalloffColor * ClothSpecularPathValue;
@@ -1463,11 +1464,10 @@ FStrataEnvLightResult StrataEvaluateForEnvLight(FStrataBSDFContext BSDFContext,
 			{
 				const float  FuzzAmount = SLAB_FUZZ_AMOUNT(BSDFContext.BSDF);
 				const float3 FuzzF0 = SLAB_FUZZ_COLOR(BSDFContext.BSDF);
-				const float  FuzzF90 = 1.0f;
 				const float  FuzzRoughness = MakeRoughnessSafe(SLAB_FUZZ_ROUGHNESS(BSDFContext.BSDF), STRATA_MIN_FUZZ_ROUGHNESS);
 
 				// F0 is ignored, as energy preservation is handle solely based on the white directional albedo and FuzzAmount
-				const FBxDFEnergyTermsRGB ClothEnergyTerms = ComputeClothEnergyTermsRGB(FuzzRoughness, BSDFContext.Context.NoV, 1.0f /*FuzzF0*/);
+				const FBxDFEnergyTermsA ClothEnergyTerms = ComputeClothEnergyTermsA(FuzzRoughness, BSDFContext.Context.NoV);
 
 				// The specular and diffuse components below the fuzz are only attenuated linearly according to the amount of fuzz.
 				const float Cloth_DirectionalAlbedo_SpecularTransmission = lerp(1.0, ComputeEnergyPreservation(ClothEnergyTerms), FuzzAmount);
