- Fix for bad motion vectors with Cascade mesh renderers

#rnx

diff --git a/Engine/Shaders/Private/MeshParticleVertexFactory.ush b/Engine/Shaders/Private/MeshParticleVertexFactory.ush
--- a/Engine/Shaders/Private/MeshParticleVertexFactory.ush
+++ b/Engine/Shaders/Private/MeshParticleVertexFactory.ush
@@ -240,7 +240,8 @@ FLWCMatrix GetParticlePrevTransform(FVertexFactoryInput Input)
 	FLWCMatrix FinalMat;
 	if ( MeshParticleVF.bUseLocalSpace )
 	{
-		FinalMat = LWCMultiply(PrevTransform4x4, GetPrimitiveDataFromUniformBuffer().PreviousLocalToWorld);
+		float3 LWCTile = LWCGetTile(LWCGetOrigin(GetPrimitiveDataFromUniformBuffer().PreviousLocalToWorld));
+		FinalMat = MakeLWCMatrix(LWCTile, PrevTransform4x4);
 	}
 	else
 	{
