Capitalizes on TSR perf optimisations on consoles to upgrade TSR Epic quality with r.TSR.History.ScreenPercentage=200

diff --git a/Engine/Config/BaseScalability.ini b/Engine/Config/BaseScalability.ini
--- a/Engine/Config/BaseScalability.ini
+++ b/Engine/Config/BaseScalability.ini
@@ -42,10 +42,8 @@ r.TemporalAA.Quality=0
 r.TSR.History.R11G11B10=1
 r.TSR.History.ScreenPercentage=100
 r.TSR.History.UpdateQuality=0
-r.TSR.History.GrandReprojection=0
 r.TSR.ShadingRejection.Flickering=0
 r.TSR.Velocity.Extrapolation=0
-r.TSR.Velocity.HoleFill=0
 r.TSR.RejectionAntiAliasingQuality=0
 
 [AntiAliasingQuality@1]
@@ -54,10 +52,8 @@ r.TemporalAA.Quality=1
 r.TSR.History.R11G11B10=1
 r.TSR.History.ScreenPercentage=100
 r.TSR.History.UpdateQuality=1
-r.TSR.History.GrandReprojection=0
 r.TSR.ShadingRejection.Flickering=0
 r.TSR.Velocity.Extrapolation=0
-r.TSR.Velocity.HoleFill=1
 r.TSR.RejectionAntiAliasingQuality=1
 
 [AntiAliasingQuality@2]
@@ -66,22 +62,18 @@ r.TemporalAA.Quality=1
 r.TSR.History.R11G11B10=1
 r.TSR.History.ScreenPercentage=100
 r.TSR.History.UpdateQuality=2
-r.TSR.History.GrandReprojection=0
 r.TSR.ShadingRejection.Flickering=1
 r.TSR.Velocity.Extrapolation=0
-r.TSR.Velocity.HoleFill=1
 r.TSR.RejectionAntiAliasingQuality=1
 
 [AntiAliasingQuality@3]
 r.FXAA.Quality=4
 r.TemporalAA.Quality=2
 r.TSR.History.R11G11B10=1
-r.TSR.History.ScreenPercentage=100
+r.TSR.History.ScreenPercentage=200
 r.TSR.History.UpdateQuality=3
-r.TSR.History.GrandReprojection=1
 r.TSR.ShadingRejection.Flickering=1
 r.TSR.Velocity.Extrapolation=1
-r.TSR.Velocity.HoleFill=1
 r.TSR.RejectionAntiAliasingQuality=2
 
 [AntiAliasingQuality@Cine]
@@ -90,10 +82,8 @@ r.TemporalAA.Quality=2
 r.TSR.History.R11G11B10=0
 r.TSR.History.ScreenPercentage=200
 r.TSR.History.UpdateQuality=3
-r.TSR.History.GrandReprojection=1
 r.TSR.ShadingRejection.Flickering=1
 r.TSR.Velocity.Extrapolation=1
-r.TSR.Velocity.HoleFill=1
 r.TSR.RejectionAntiAliasingQuality=2
 
 
diff --git a/Engine/Shaders/Private/TemporalSuperResolution/TSRCommon.ush b/Engine/Shaders/Private/TemporalSuperResolution/TSRCommon.ush
--- a/Engine/Shaders/Private/TemporalSuperResolution/TSRCommon.ush
+++ b/Engine/Shaders/Private/TemporalSuperResolution/TSRCommon.ush
@@ -1194,4 +1194,12 @@ float ComputeScreenVelocityBilateralWeight(float2 ScreenVelocityDelta)
 	return ComputePixelVelocityBilateralWeight(ScreenVelocityDelta * ScreenVelocityToInputPixelVelocity);
 }
 
+tsr_half2 ComputePrevWeightMultiplier(tsr_half2 BlendFactor)
+{
+	return (1.0 - BlendFactor) / BlendFactor;
+}
 
+tsr_half2 RejectionFactorToBlendFactor(tsr_half2 RejectionFactor)
+{
+	return tsr_half(1.0) - RejectionFactor;
+}
diff --git a/Engine/Shaders/Private/TemporalSuperResolution/TSRRejectShading.usf b/Engine/Shaders/Private/TemporalSuperResolution/TSRRejectShading.usf
--- a/Engine/Shaders/Private/TemporalSuperResolution/TSRRejectShading.usf
+++ b/Engine/Shaders/Private/TemporalSuperResolution/TSRRejectShading.usf
@@ -98,6 +98,8 @@ float2 TranslucencyTextureUVMin;
 float2 TranslucencyTextureUVMax;
 float3 HistoryGuideQuantizationError;
 float FlickeringFramePeriod;
+float TheoricBlendFactor;
+
 
 Texture2D<tsr_halfC> InputTexture;
 Texture2D<tsr_half>  InputMoireLumaTexture;
@@ -951,10 +953,6 @@ void MainCS(
 	// Update the guide for next frame.
 	tsr_simd_halfC FinalGuide;
 	{
-		const tsr_half CurrentFrameSampleCount = tsr_half(1.0);
-		const tsr_half PrevLowFrequencySampleCount = tsr_half(8.0);
-		const tsr_half TheoricBlendFactor = CurrentFrameSampleCount / (CurrentFrameSampleCount + PrevLowFrequencySampleCount);
-		
 		tsr_simd_half BlendFinal;
 		// Completely discard the history
 		UNROLL_N(SIMD_SIZE)
@@ -964,9 +962,9 @@ void MainCS(
 			bool bIsDisoccluded = bIsValid.GetElement(SimdIndex) == tsr_half(0.0);
 			tsr_half LocalRejectionBlendFinal = RejectionBlendFinal.GetElement(SimdIndex);
 
+			//tsr_half TotalBlend = select(or(bIsOffScreen, bIsDisoccluded), tsr_half(1.0), tsr_half(0.0));
 			tsr_half TotalBlend = select(or(bIsOffScreen, bIsDisoccluded), tsr_half(1.0), saturate(tsr_half(1.0) - LocalRejectionBlendFinal * tsr_half(4.0)));
-			tsr_half RejectedBlendFactor = max(TheoricBlendFactor, tsr_half(1.0) - LocalRejectionBlendFinal);
-
+			tsr_half RejectedBlendFactor = max(tsr_half(TheoricBlendFactor), tsr_half(1.0) - LocalRejectionBlendFinal);
 
 			BlendFinal.SetElement(SimdIndex, max(RejectedBlendFactor, TotalBlend));
 		}
diff --git a/Engine/Shaders/Private/TemporalSuperResolution/TSRUpdateHistory.usf b/Engine/Shaders/Private/TemporalSuperResolution/TSRUpdateHistory.usf
--- a/Engine/Shaders/Private/TemporalSuperResolution/TSRUpdateHistory.usf
+++ b/Engine/Shaders/Private/TemporalSuperResolution/TSRUpdateHistory.usf
@@ -27,9 +27,6 @@
 	#define CONFIG_INTERPOLATE_VELOCITY 1
 	#define CONFIG_REJECTION_ANTI_ALIASING 1
 
-	// To support sg.AntiAliasingQuality=4 with same temporal response as sg.AntiAliasingQuality=3
-	#define CONFIG_APPLY_CONTRIBUTION_MULTIPLIER 1
-
 #else
 	#error Unknown history update quality
 #endif
@@ -82,10 +79,6 @@
 	#define CONFIG_ENABLE_DUAL_PIXEL_VECTORIZATION 0
 #endif
 
-#ifndef CONFIG_APPLY_CONTRIBUTION_MULTIPLIER
-	#define CONFIG_APPLY_CONTRIBUTION_MULTIPLIER 0
-#endif
-
 #ifndef PLATFORM_SPECIFIC_ISOLATE
 	#define PLATFORM_SPECIFIC_ISOLATE ISOLATE
 #endif
@@ -166,9 +159,12 @@ FScreenTransform HistoryPixelPosToTranslucencyPPCo;
 FScreenTransform ScreenPosToGrandPrevHistoryBufferUV;
 float3 HistoryQuantizationError;
 float MinTranslucencyRejection;
+float HistorySampleCount;
+float HistoryHisteresis;
+float WeightClampingRejection;
+float WeightClampingPixelSpeedAmplitude;
 float InvWeightClampingPixelSpeed;
 float InputToHistoryFactor;
-float InputContributionMultiplier;
 float GrandPrevPreExposureCorrection;
 uint ResponsiveStencilMask;
 uint bGenerateOutputMip1;
@@ -772,6 +768,7 @@ void MainCS(
 			{
 				bIsDisoccluded = false;
 				LowFrequencyRejection = tsr_half(1.0).xx;
+				LowFrequencyClamp = tsr_half(1.0).xx;
 				TranslucencyRejection = tsr_half(1.0).xx;
 			}
 			#endif
@@ -785,9 +782,6 @@ void MainCS(
 	
 	// Reproject history
 	tsr_halfCx2 PrevHighFrequencyColor;
-	#if CONFIG_PREMULTIPLY_HISTORY
-		tsr_half2 PrevHistoryWeight;
-	#endif
 	tsr_half2 PrevHistoryValidity;
 	
 	tsr_halfCx2 PrevHistoryMomentMin;
@@ -942,17 +936,8 @@ void MainCS(
 
 					tsr_half2 KernelWeight = RawKernelWeight[i];
 
-					#if CONFIG_PREMULTIPLY_HISTORY
-						tsr_half2 InvWeight = SafeRcp(SampleMetadata[1]);
-
-						HighFrequencyMin = min(HighFrequencyMin, dpv_scale(SampleHighFrequency, InvWeight));
-						HighFrequencyMax = max(HighFrequencyMax, dpv_scale(SampleHighFrequency, InvWeight));
-
-					#else
-						HighFrequencyMin = min(HighFrequencyMin, SampleHighFrequency);
-						HighFrequencyMax = max(HighFrequencyMax, SampleHighFrequency);
-
-					#endif
+					HighFrequencyMin = min(HighFrequencyMin, SampleHighFrequency);
+					HighFrequencyMax = max(HighFrequencyMax, SampleHighFrequency);
 
 					AccumulateHighFrequency += dpv_scale(SampleHighFrequency, KernelWeight);
 					AccumulateMetadata  += dpv_scale(SampleMetadata, KernelWeight);
@@ -1004,9 +989,6 @@ void MainCS(
 			PrevHistoryMomentMax = HighFrequencyMax;
 
 			PrevHighFrequencyColor = AccumulateHighFrequency;
-			#if CONFIG_PREMULTIPLY_HISTORY
-				PrevHistoryWeight = AccumulateMetadata[1];
-			#endif
 			PrevHistoryValidity = AccumulateMetadata[0];
 
 			PrevHistoryTranslucency = AccumulateTranslucency;
@@ -1074,6 +1056,9 @@ void MainCS(
 	tsr_half2 InputPixelAlignement;
 	tsr_half2 ClosestInputLuma4;
 	
+	tsr_half2 BlendClamp;
+	tsr_half2 PrevWeight;
+	tsr_half2 CurrentWeight;
 	
 	PLATFORM_SPECIFIC_ISOLATE
 	{
@@ -1116,22 +1101,72 @@ void MainCS(
 			
 			// Vector in pixel between pixel K -> O.
 			tsr_half2x2 dInputKO = tsr_half2x2(InputPPCo - InputPPCk);
+			
+			// Compute upscaling kernel size based of the rejections and number of samples already in history.
+			tsr_half2 KernelInputToHistoryFactor;
+			{
+				const tsr_half ActualHistorySampleCount = tsr_half(HistorySampleCount);
+				const tsr_half ActualHistoryHisteresis = tsr_half(HistoryHisteresis);
 
+				tsr_half2 MinRejectionBlendFactor = RejectionFactorToBlendFactor(LowFrequencyRejection);
+				MinRejectionBlendFactor = select(bIsOffScreen, tsr_half(1.0).xx, MinRejectionBlendFactor);
 
-			// Soften the spatial kernel when there is high noise in the input to make parallax disocclusion of detailed geometry less distracting.
-			tsr_half2 OffScreenInputToHistoryFactor = tsr_half(1.0) - tsr_half(0.5) * NoiseFiltering;
+				tsr_half2 CoarseInputPixelAlignement = ComputeSampleWeigth(tsr_half(1.0), dInputKO, /* MinimalContribution = */ float(0.0));
+				tsr_half2 IdealInputPixelAlignement = ComputeSampleWeigth(tsr_half(InputToHistoryFactor), dInputKO, /* MinimalContribution = */ float(0.0));
+			
+				tsr_half2 CoarseCurrentContribution = CoarseInputPixelAlignement * tsr_half(ActualHistoryHisteresis);
+				tsr_half2 IdealCurrentContribution = IdealInputPixelAlignement * tsr_half(ActualHistoryHisteresis);
 
-			bool2 bIsRefining = PrevHistoryValidity != 0.0;
+				// Select whether history should be clamped.
+				BlendClamp = select(bIsDisoccluded, tsr_half(0.0), min(HoleFilledVelocityMask[1], LowFrequencyClamp));
+		
+				// Reduces the number of pixels when rejecting history
+				tsr_half2 ClampedPrevHistoryValidity = PrevHistoryValidity;
+				#if 1
+				{
+					tsr_half2 MaxValidity = lerp(tsr_half(WeightClampingRejection), tsr_half(1.0), select(bIsDisoccluded, tsr_half(0.0), LowFrequencyClamp));
+					//tsr_half2 MaxValidity = lerp(tsr_half(WeightClampingRejection), tsr_half(1.0), BlendClamp);
+
+					ClampedPrevHistoryValidity = min(ClampedPrevHistoryValidity, MaxValidity);
+				}
+				#endif
 
-			tsr_half2 KernelInputToHistoryLerp = select(or(or(bIsOffScreen, bIsDisoccluded), !bIsRefining), tsr_half(0.0).xx, saturate(LowFrequencyRejection * tsr_half(16.0) - tsr_half(13.0)));
-			tsr_half2 KernelInputToHistoryFactor = lerp(OffScreenInputToHistoryFactor, tsr_half(InputToHistoryFactor), KernelInputToHistoryLerp);
-			tsr_half2 KernelInputToHistoryAlignmentFactor = lerp(tsr_half(1.0).xx, tsr_half(InputToHistoryFactor), KernelInputToHistoryLerp);
+				const tsr_half MaxWeight = tsr_half(1.0);
+	
+				// Compute the histeresis if was refining the history
+				tsr_half2 RefiningHisteresis = select(
+					ClampedPrevHistoryValidity + IdealCurrentContribution > tsr_half(0.0),
+					IdealCurrentContribution * rcp(ClampedPrevHistoryValidity + IdealCurrentContribution),
+					tsr_half(1.0));
+				
+				// Compute the prev weight with coarse pixel alignment when rejecting the history.
+				tsr_half2 CoarseRejectedPrevWeight = min(
+					CoarseCurrentContribution * ComputePrevWeightMultiplier(MinRejectionBlendFactor),
+					MaxWeight);
+					
+				// Compute the prev weight with coarse pixel alignment when refining.
+				tsr_half2 CoarseRefiningPrevWeight = min(
+					CoarseCurrentContribution * ComputePrevWeightMultiplier(RefiningHisteresis),
+					MaxWeight);
+				
+				// Do not refine when rejected previous weight is lower than refining with coarse pixel alignment.
+				// bIsNotRefining = CoarseRejectedPrevWeight < CoarseRefiningPrevWeight
+				tsr_half2 bIsRefining = select(CoarseRejectedPrevWeight < CoarseRefiningPrevWeight, tsr_half(0.0), tsr_half(1.0));
+				bIsRefining = min(bIsRefining, saturate(ClampedPrevHistoryValidity * ActualHistorySampleCount));
+	
+				// Soften the spatial kernel when there is high noise in the input to make parallax disocclusion of detailed geometry less distracting.
+				tsr_half2 OffScreenInputToHistoryFactor = tsr_half(1.0) - tsr_half(0.5) * NoiseFiltering;
+
+				tsr_half2 KernelInputToHistoryLerp = select(or(bIsOffScreen, bIsDisoccluded), tsr_half(0.0).xx, saturate(LowFrequencyRejection * tsr_half(16.0) - tsr_half(13.0)) * bIsRefining);
+				KernelInputToHistoryFactor = lerp(OffScreenInputToHistoryFactor, tsr_half(InputToHistoryFactor), KernelInputToHistoryLerp);
+				tsr_half2 KernelInputToHistoryAlignmentFactor = lerp(tsr_half(1.0).xx, tsr_half(InputToHistoryFactor), KernelInputToHistoryLerp);
 			
-			InputPixelAlignement = ComputeSampleWeigth(KernelInputToHistoryAlignmentFactor, dInputKO, /* MinimalContribution = */ float(0.0));
-		
-			#if CONFIG_APPLY_CONTRIBUTION_MULTIPLIER
-				InputPixelAlignement *= tsr_half(InputContributionMultiplier);
-			#endif
+				InputPixelAlignement = ComputeSampleWeigth(KernelInputToHistoryAlignmentFactor, dInputKO, /* MinimalContribution = */ float(0.0));
+			
+				CurrentWeight = InputPixelAlignement * tsr_half(ActualHistoryHisteresis);
+				PrevWeight = min(select(CurrentWeight > tsr_half(0.0), CurrentWeight, CoarseCurrentContribution) * ComputePrevWeightMultiplier(MinRejectionBlendFactor), ClampedPrevHistoryValidity);
+				PrevWeight = min(PrevWeight, MaxWeight - CurrentWeight);	
+			}
 
 			UNROLL_N(CONFIG_SAMPLES)
 			for (uint SampleId = 0; SampleId < CONFIG_SAMPLES; SampleId++)
@@ -1344,9 +1379,6 @@ void MainCS(
 	#endif
 
 	// Contribute current frame input into the predictor for next frame.
-	const tsr_half Histeresis = rcp(tsr_half(MAX_SAMPLE_COUNT));
-	const tsr_half PredictionOnlyValidity = Histeresis * tsr_half(2.0);
-	
 	tsr_halfCx2 FinalHighFrequencyColor;
 	tsr_half2 FinalHistoryValidity;
 	
@@ -1362,44 +1394,16 @@ void MainCS(
 	#endif
 
 	{
-		#if CONFIG_PREMULTIPLY_HISTORY
-			tsr_half2 InvPrevHistoryWeight = SafeRcp(PrevHistoryWeight);
-		#endif
-		tsr_half2 InvPrevHistoryValidity = SafeRcp(PrevHistoryValidity);
-
-		tsr_half2 TotalRejection = select(or(bIsOffScreen, bIsDisoccluded), tsr_half(0.0).xx, saturate(LowFrequencyRejection * tsr_half(4.0)));
-
-		tsr_half2 BlendClamp = min(TotalRejection * HoleFilledVelocityMask[1], LowFrequencyClamp);
-		
-		#if CONFIG_PREMULTIPLY_HISTORY
-			tsr_halfCx2 UnpremultipliedPrevHighFrequencyColor = dpv_scale(PrevHighFrequencyColor, InvPrevHistoryWeight);
-		#else
-			tsr_halfCx2 UnpremultipliedPrevHighFrequencyColor = PrevHighFrequencyColor;
-		#endif
-		tsr_halfCx2 ClampedPrevHighFrequencyColor = clamp(UnpremultipliedPrevHighFrequencyColor, InputMinColor, InputMaxColor);
-		tsr_halfCx2 BlendedPrevHighFrequencyColor = UnpremultipliedPrevHighFrequencyColor;
-		BlendedPrevHighFrequencyColor[0] = lerp(ClampedPrevHighFrequencyColor[0], UnpremultipliedPrevHighFrequencyColor[0], BlendClamp);
-		BlendedPrevHighFrequencyColor[1] = lerp(ClampedPrevHighFrequencyColor[1], UnpremultipliedPrevHighFrequencyColor[1], BlendClamp);
-		BlendedPrevHighFrequencyColor[2] = lerp(ClampedPrevHighFrequencyColor[2], UnpremultipliedPrevHighFrequencyColor[2], BlendClamp);
-
-		// Compute how much the history needs to be rejected, based on the completeness.
-		tsr_half2 PrevHistoryRejectionWeight = select(bIsOffScreen, tsr_half(0.0), LowFrequencyRejection);
-	
-		tsr_half2 DesiredCurrentContribution = max(Histeresis * InputPixelAlignement, tsr_half(0.0));
-
-		// Determine whether the the prediction based rejection was confident enough.
-		tsr_half2 RejectionConfidentEnough = tsr_half(1); // saturate(RejectionValidity * MAX_SAMPLE_COUNT - 3.0);
+		tsr_halfCx2 ClampedPrevHighFrequencyColor = clamp(PrevHighFrequencyColor, InputMinColor, InputMaxColor);
+		tsr_halfCx2 BlendedPrevHighFrequencyColor = PrevHighFrequencyColor;
+		BlendedPrevHighFrequencyColor[0] = lerp(ClampedPrevHighFrequencyColor[0], PrevHighFrequencyColor[0], BlendClamp);
+		BlendedPrevHighFrequencyColor[1] = lerp(ClampedPrevHighFrequencyColor[1], PrevHighFrequencyColor[1], BlendClamp);
+		BlendedPrevHighFrequencyColor[2] = lerp(ClampedPrevHighFrequencyColor[2], PrevHighFrequencyColor[2], BlendClamp);
 		
 		// Clamp the validity due to motion to maintain better sharpness in history reprojection under motion.
-		tsr_half2 ClampedPrevHistoryValidity;
-		#if CONFIG_GRAND_REPROJECTION
+		#if !CONFIG_GRAND_REPROJECTION
 		{
-			tsr_half2 MaxValidity = lerp(PredictionOnlyValidity, tsr_half(1.0), TotalRejection);
-			ClampedPrevHistoryValidity = min(PrevHistoryValidity, MaxValidity);
-		}
-		#elif 1
-		{
-			tsr_half2 MaxValidity = tsr_half(1.0) - (tsr_half(1.0) - PredictionOnlyValidity) * saturate(OutputPixelVelocity * tsr_half(InvWeightClampingPixelSpeed));
+			tsr_half2 MaxValidity = tsr_half(1.0) - tsr_half(WeightClampingPixelSpeedAmplitude) * saturate(OutputPixelVelocity * tsr_half(InvWeightClampingPixelSpeed));
 
 			// Clamp up the max validity to favor stability under motion on high contrast edges.
 			#if 1
@@ -1408,40 +1412,22 @@ void MainCS(
 				tsr_half2 FilteredLuma = Luma4(FilteredInputColor);
 
 				tsr_half2 MinValidityForStability = abs(FilteredLuma - PrevHistoryLuma) / max(FilteredLuma, PrevHistoryLuma);
-				
+			
 				MaxValidity = max(MaxValidity, MinValidityForStability);
 			}
 			#endif
 			
-			ClampedPrevHistoryValidity = min(PrevHistoryValidity, MaxValidity);
-		}
-		#else
-		{
-			ClampedPrevHistoryValidity = PrevHistoryValidity;
+			PrevWeight = min(PrevWeight, MaxValidity);
 		}
 		#endif
 
-		ClampedPrevHistoryValidity = min(ClampedPrevHistoryValidity, lerp(tsr_half(0.2), tsr_half(1.0), BlendClamp));
-
-		// Compute the newly rejected validity
-		tsr_half2 RejectedValidity = ClampedPrevHistoryValidity * PrevHistoryRejectionWeight;
-
-		// Compute the maximum output validity.
-		tsr_half2 OutputValidity = (
-			clamp(RejectedValidity + DesiredCurrentContribution, tsr_half(0.0), PredictionOnlyValidity) +
-			clamp(RejectedValidity + DesiredCurrentContribution * PrevHistoryRejectionWeight * RejectionConfidentEnough - PredictionOnlyValidity, tsr_half(0.0), tsr_half(1.0) - PredictionOnlyValidity));
-
-		tsr_half2 OutputValidity2 = select(bIsOffScreen, max(MIN_IMPLICIT_HISTORY_WEIGHT, OutputValidity), OutputValidity);
-		DesiredCurrentContribution = select(bIsOffScreen, max(MIN_IMPLICIT_HISTORY_WEIGHT, DesiredCurrentContribution), DesiredCurrentContribution);
-
-		tsr_half2 PrevMomentWeight = max(OutputValidity2 - DesiredCurrentContribution, tsr_half(0.0));
-		tsr_half2 CurrentMomentWeight = min(DesiredCurrentContribution, OutputValidity2);
+		tsr_half2 OutputValidity = CurrentWeight + PrevWeight;
 
 		tsr_half2 PrevHistoryToneWeight = HdrWeightY(Luma4(BlendedPrevHighFrequencyColor));
 		tsr_half2 FilteredInputToneWeight = HdrWeight4(FilteredInputColor);
 			
-		tsr_half2 BlendPrevHistory = PrevMomentWeight * PrevHistoryToneWeight;
-		tsr_half2 BlendFilteredInput = CurrentMomentWeight * FilteredInputToneWeight;
+		tsr_half2 BlendPrevHistory = PrevWeight * PrevHistoryToneWeight;
+		tsr_half2 BlendFilteredInput = CurrentWeight * FilteredInputToneWeight;
 
 		tsr_half2 CommonWeight = SafeRcp(BlendPrevHistory + BlendFilteredInput);
 
@@ -1454,10 +1440,6 @@ void MainCS(
 
 		#if CONFIG_GRAND_REPROJECTION
 		{
-			#if CONFIG_PREMULTIPLY_HISTORY
-				#error Unsupported
-			#endif
-
 			tsr_halfCx2 BoxSize = dpv_scale(PrevHistoryMomentMax - PrevHistoryMomentMin, tsr_half(0.5));
 
 			// Compute how much frame N-2 got overblur by reprojecting into N-1 and N versus directly to N.
@@ -1472,11 +1454,11 @@ void MainCS(
 			tsr_half2 PrevReprojectionOverblurScale = select(bIsOverblurLegal, tsr_half(1.0), tsr_half(0.0));
 
 			tsr_half2 HistoryClampResultant = MeasureRejectionFactor(
-				UnpremultipliedPrevHighFrequencyColor, BlendedPrevHighFrequencyColor,
+				PrevHighFrequencyColor, BlendedPrevHighFrequencyColor,
 				FilteredInputColor, InputMinColor, InputMaxColor,
 				MeasureBackbufferLDRQuantizationError());
 
-			tsr_half2 MeasuredResultant = saturate(CommonWeight * BlendPrevHistory * HistoryClampResultant * TotalRejection);
+			tsr_half2 MeasuredResultant = saturate(CommonWeight * BlendPrevHistory * HistoryClampResultant);
 
 			FinalPrevHighFrequencyColor = PrevHighFrequencyColor;
 
@@ -1484,7 +1466,7 @@ void MainCS(
 				FinalPrevHighFrequencyResultant = MeasuredResultant;
 
 				const tsr_half CorrectionError = tsr_half(1.0 / sqrt(PI * 0.5));
-				MeasuredResultant = min(MeasuredResultant, saturate(tsr_half(1.0) - Histeresis * rcp(FinalHistoryValidity)));
+				MeasuredResultant = min(MeasuredResultant, saturate(tsr_half(1.0) - tsr_half(HistoryHisteresis) * rcp(FinalHistoryValidity)));
 
 				tsr_half2 CorrectionResultant = GrandPrevToPrevResultant * MeasuredResultant * PrevReprojectionOverblurScale * CorrectionError;
 
@@ -1671,20 +1653,11 @@ void MainCS(
 			FinalMetadata[0] = FinalHistoryValidity;
 		#endif
 
-		// Premultiply FinalHighFrequencyColor to have validity weighted average in next frame history reprojection.
-		#if CONFIG_PREMULTIPLY_HISTORY
-			tsr_half2 FinalHistoryWeight = max(FinalHistoryValidity, MIN_IMPLICIT_HISTORY_WEIGHT);
-			tsr_halfCx2 PremultipliedHighFrequencyColor = dpv_scale(FinalHighFrequencyColor, FinalHistoryWeight);
-			FinalMetadata[1] = FinalHistoryWeight;
-		#else
-			tsr_halfCx2 PremultipliedHighFrequencyColor = FinalHighFrequencyColor;
-		#endif
-		
 		#if CONFIG_GRAND_REPROJECTION == 1
 			FinalMetadata[1] = FinalPrevHighFrequencyResultant;
 		#elif CONFIG_GRAND_REPROJECTION == 2
 			// Encode the overblur with the high frequency color to avoid negative numbers not encodable in a R11G11B10.
-			FinalOverblur = FinalOverblur + PremultipliedHighFrequencyColor;
+			FinalOverblur = FinalOverblur + FinalHighFrequencyColor;
 		#endif
 
 		// Stochastically round up or down using the hardware RWTexture2D truncation unit to take into precision
@@ -1695,7 +1668,7 @@ void MainCS(
 			tsr_half E = tsr_half(Hammersley16(0, 1, Random).x);
 			
 			FinalOutputColor = QuantizeForFloatRenderTarget(FinalOutputColor, E, HistoryQuantizationError);
-			PremultipliedHighFrequencyColor = QuantizeForFloatRenderTarget(PremultipliedHighFrequencyColor, E, HistoryQuantizationError);
+			FinalHighFrequencyColor = QuantizeForFloatRenderTarget(FinalHighFrequencyColor, E, HistoryQuantizationError);
 
 			#if CONFIG_GRAND_REPROJECTION
 				FinalPrevHighFrequencyColor = QuantizeForFloatRenderTarget(FinalPrevHighFrequencyColor, E, HistoryQuantizationError);
@@ -1709,8 +1682,8 @@ void MainCS(
 
 		// Protect from NaN and +Inf when writing out the history.
 		{
-			PremultipliedHighFrequencyColor = -dpv_min(-PremultipliedHighFrequencyColor, ColorNull);
-			PremultipliedHighFrequencyColor = dpv_min(PremultipliedHighFrequencyColor, ColorMax10BitsFloat);	
+			FinalHighFrequencyColor = -dpv_min(-FinalHighFrequencyColor, ColorNull);
+			FinalHighFrequencyColor = dpv_min(FinalHighFrequencyColor, ColorMax10BitsFloat);	
 
 			#if CONFIG_ACCUMULATE_TRANSLUCENCY_SEPARATELY
 				FinalTranslucencyColor = -dpv_min(-FinalTranslucencyColor, tsr_half(0.0).xxxx);
@@ -1723,7 +1696,7 @@ void MainCS(
 			// Output final history lo pixel.
 			{
 				#if CONFIG_OUTPUT_HIGH_FREQUENCY
-					HistoryOutput_ColorArray[tsr_short3(dpv_lo(LocalHistoryPixelPos), HistoryOutput_ArrayIndices_HighFrequency)] = dpv_lo(PremultipliedHighFrequencyColor);
+					HistoryOutput_ColorArray[tsr_short3(dpv_lo(LocalHistoryPixelPos), HistoryOutput_ArrayIndices_HighFrequency)] = dpv_lo(FinalHighFrequencyColor);
 				#endif
 				HistoryOutput_Metadata[dpv_lo(LocalHistoryPixelPos)]  = dpv_lo(FinalMetadata);
 				#if CONFIG_ACCUMULATE_TRANSLUCENCY_SEPARATLY
@@ -1747,7 +1720,7 @@ void MainCS(
 			#if CONFIG_ENABLE_DUAL_PIXEL_VECTORIZATION
 			{
 				#if CONFIG_OUTPUT_HIGH_FREQUENCY
-					HistoryOutput_ColorArray[tsr_short3(dpv_hi(LocalHistoryPixelPos), HistoryOutput_ArrayIndices_HighFrequency)] = dpv_hi(PremultipliedHighFrequencyColor);
+					HistoryOutput_ColorArray[tsr_short3(dpv_hi(LocalHistoryPixelPos), HistoryOutput_ArrayIndices_HighFrequency)] = dpv_hi(FinalHighFrequencyColor);
 				#endif
 				HistoryOutput_Metadata[dpv_hi(LocalHistoryPixelPos)] = dpv_hi(FinalMetadata);
 				#if CONFIG_ACCUMULATE_TRANSLUCENCY_SEPARATELY
diff --git a/Engine/Source/Runtime/Renderer/Private/PostProcess/TemporalSuperResolution.cpp b/Engine/Source/Runtime/Renderer/Private/PostProcess/TemporalSuperResolution.cpp
--- a/Engine/Source/Runtime/Renderer/Private/PostProcess/TemporalSuperResolution.cpp
+++ b/Engine/Source/Runtime/Renderer/Private/PostProcess/TemporalSuperResolution.cpp
@@ -16,9 +16,13 @@
 namespace
 {
 
+TAutoConsoleVariable<float> CVarTSRHistorySampleCount(
+	TEXT("r.TSR.History.SampleCount"), 16.0f,
+	TEXT("Maximum number sample for each output pixel in the history."),
+	ECVF_RenderThreadSafe);
+
 TAutoConsoleVariable<float> CVarTSRHistorySP(
-	TEXT("r.TSR.History.ScreenPercentage"),
-	100.0f,
+	TEXT("r.TSR.History.ScreenPercentage"), 100.0f,
 	TEXT("Size of TSR's history."),
 	ECVF_Scalability | ECVF_RenderThreadSafe);
 
@@ -40,13 +44,18 @@ TAutoConsoleVariable<int32> CVarTSRHistorySeparateTranslucency(
 TAutoConsoleVariable<int32> CVarTSRHistoryGrandReprojection(
 	TEXT("r.TSR.History.GrandReprojection"), 0,
 	TEXT("Experimental functionality to keep higher sharpness in TSR's history in motion."),
-	ECVF_Scalability | ECVF_RenderThreadSafe);
+	ECVF_RenderThreadSafe);
 
 TAutoConsoleVariable<int32> CVarTSRWaveOps(
 	TEXT("r.TSR.WaveOps"), 1,
 	TEXT("Whether to use wave ops in the shading rejection heuristics"),
 	ECVF_RenderThreadSafe);
 
+TAutoConsoleVariable<float> CVarTSRHistoryRejectionSampleCount(
+	TEXT("r.TSR.ShadingRejection.SampleCount"), 2.0f,
+	TEXT("Maximum number of sample in each output pixel of the history after total shading rejection."),
+	ECVF_Scalability | ECVF_RenderThreadSafe);
+
 TAutoConsoleVariable<int32> CVarTSRFlickeringEnable(
 	TEXT("r.TSR.ShadingRejection.Flickering"), 1,
 	TEXT("Whether to enable the flickering detection heuristic.\n"),
@@ -101,9 +110,14 @@ TAutoConsoleVariable<int32> CVarTSREnableResponiveAA(
 	TEXT("Whether the responsive AA should keep history fully clamped."),
 	ECVF_RenderThreadSafe);
 
+TAutoConsoleVariable<float> CVarTSRWeightClampingSampleCount(
+	TEXT("r.TSR.Velocity.WeightClampingSampleCount"), 4.0f,
+	TEXT("Sample count in history pixel to clamp history to when output pixel velocity reach r.TSR.Velocity.WeightClampingPixelSpeed (Default = 4.0f)."),
+	ECVF_RenderThreadSafe);
+
 TAutoConsoleVariable<float> CVarTSRWeightClampingPixelSpeed(
 	TEXT("r.TSR.Velocity.WeightClampingPixelSpeed"), 1.0f,
-	TEXT("Defines the pixel velocity at which the the high frequencies of the history get's their contributing weight clamped. ")
+	TEXT("Defines the output pixel velocity at which the the high frequencies of the history get's their contributing weight clamped. ")
 	TEXT("Smallest reduce blur in movement (Default = 1.0f)."),
 	ECVF_RenderThreadSafe);
 
@@ -495,6 +509,7 @@ class FTSRRejectShadingCS : public FTSRShader
 		SHADER_PARAMETER(FVector2f, TranslucencyTextureUVMax)
 		SHADER_PARAMETER(FVector3f, HistoryGuideQuantizationError)
 		SHADER_PARAMETER(float, FlickeringFramePeriod)
+		SHADER_PARAMETER(float, TheoricBlendFactor)
 
 		SHADER_PARAMETER_RDG_TEXTURE(Texture2D, InputTexture)
 		SHADER_PARAMETER_RDG_TEXTURE(Texture2D, InputMoireLumaTexture)
@@ -658,6 +673,10 @@ class FTSRUpdateHistoryCS : public FTSRShader
 
 		SHADER_PARAMETER(FVector3f, HistoryQuantizationError)
 		SHADER_PARAMETER(float, MinTranslucencyRejection)
+		SHADER_PARAMETER(float, HistorySampleCount)
+		SHADER_PARAMETER(float, HistoryHisteresis)
+		SHADER_PARAMETER(float, WeightClampingRejection)
+		SHADER_PARAMETER(float, WeightClampingPixelSpeedAmplitude)
 		SHADER_PARAMETER(float, InvWeightClampingPixelSpeed)
 		SHADER_PARAMETER(float, InputToHistoryFactor)
 		SHADER_PARAMETER(float, InputContributionMultiplier)
@@ -837,6 +856,9 @@ ITemporalUpscaler::FOutputs AddTemporalSuperResolutionPasses(
 	const float RefreshRateToFrameRateCap = (View.Family->Time.GetDeltaRealTimeSeconds() > 0.0f && CVarTSRFlickeringAdjustToFrameRate.GetValueOnRenderThread())
 		? View.Family->Time.GetDeltaRealTimeSeconds() * CVarTSRFlickeringFrameRateCap.GetValueOnRenderThread() : 1.0f;
 
+	// Maximum number sample for each output pixel in the history
+	const float MaxHistorySampleCount = FMath::Clamp(CVarTSRHistorySampleCount.GetValueOnRenderThread(), 8.0f, 32.0f);
+
 	// whether TSR passes can run on async compute.
 	int32 AsyncComputePasses = GSupportsEfficientAsyncCompute ? CVarTSRAsyncCompute.GetValueOnRenderThread() : 0;
 
@@ -926,9 +948,14 @@ ITemporalUpscaler::FOutputs AddTemporalSuperResolutionPasses(
 			FMath::Max(InputExtent.X, QuantizedHistoryViewSize.X),
 			FMath::Max(InputExtent.Y, QuantizedHistoryViewSize.Y));
 	}
+	float OutputToHistoryResolutionFraction = float(HistorySize.X) / float(OutputRect.Width());
+	float OutputToHistoryResolutionFractionSquare = OutputToHistoryResolutionFraction * OutputToHistoryResolutionFraction;
 
-	float ScreenPercentage = float(InputRect.Width()) / float(OutputRect.Width());
-	float InvScreenPercentage = float(OutputRect.Width()) / float(InputRect.Width());
+	float InputToHistoryResolutionFraction = float(HistorySize.X) / float(InputRect.Width());
+	float InputToHistoryResolutionFractionSquare = InputToHistoryResolutionFraction * InputToHistoryResolutionFraction;
+
+	float OutputToInputResolutionFraction = float(InputRect.Width()) / float(OutputRect.Width());
+	float OutputToInputResolutionFractionSquare = OutputToInputResolutionFraction * OutputToInputResolutionFraction;
 
 	RDG_EVENT_SCOPE(GraphBuilder, "TemporalSuperResolution(%s) %dx%d -> %dx%d",
 		bSupportsAlpha ? TEXT("Alpha") : TEXT(""),
@@ -1551,6 +1578,7 @@ ITemporalUpscaler::FOutputs AddTemporalSuperResolutionPasses(
 		PassParameters->TranslucencyTextureUVMax = GetScreenPassTextureViewportParameters(TranslucencyViewport).UVViewportBilinearMax;
 		PassParameters->HistoryGuideQuantizationError = ComputePixelFormatQuantizationError(History.Guide->Desc.Format);
 		PassParameters->FlickeringFramePeriod = FlickeringFramePeriod;
+		PassParameters->TheoricBlendFactor = 1.0f / (1.0f + MaxHistorySampleCount / OutputToInputResolutionFractionSquare);
 
 		PassParameters->InputTexture = PassInputs.SceneColorTexture;
 		if (PassInputs.MoireInputTexture.IsValid())
@@ -1699,9 +1727,16 @@ ITemporalUpscaler::FOutputs AddTemporalSuperResolutionPasses(
 		PassParameters->HistoryPixelPosToTranslucencyPPCo = HistoryPixelPosToViewportUV * PassParameters->TranslucencyInfo.ViewportSize + CommonParameters.InputJitter * PassParameters->TranslucencyInfo.ViewportSize / CommonParameters.InputInfo.ViewportSize + SeparateTranslucencyRect.Min;
 		PassParameters->HistoryQuantizationError = ComputePixelFormatQuantizationError((History.ColorArray ? History.ColorArray : History.Output)->Desc.Format);
 		PassParameters->MinTranslucencyRejection = TranslucencyRejectionTexture == nullptr ? 1.0 : 0.0;
-		PassParameters->InvWeightClampingPixelSpeed = 1.0f / CVarTSRWeightClampingPixelSpeed.GetValueOnRenderThread();
+
+		// All parameters to control the sample count in history.
+		PassParameters->HistorySampleCount = MaxHistorySampleCount / OutputToHistoryResolutionFractionSquare;
+		PassParameters->HistoryHisteresis = 1.0f / PassParameters->HistorySampleCount;
+		PassParameters->WeightClampingRejection = (CVarTSRHistoryRejectionSampleCount.GetValueOnRenderThread() / OutputToHistoryResolutionFractionSquare) * PassParameters->HistoryHisteresis;
+		PassParameters->WeightClampingPixelSpeedAmplitude = FMath::Clamp(1.0f - CVarTSRWeightClampingSampleCount.GetValueOnRenderThread() * PassParameters->HistoryHisteresis, 0.0f, 1.0f);
+		PassParameters->InvWeightClampingPixelSpeed = 1.0f / (CVarTSRWeightClampingPixelSpeed.GetValueOnRenderThread() * OutputToHistoryResolutionFraction);
+		
 		PassParameters->InputToHistoryFactor = float(HistorySize.X) / float(InputRect.Width());
-		PassParameters->InputContributionMultiplier = FMath::Pow(float(HistorySize.X) / float(OutputRect.Width()), 2.0f); 
+		PassParameters->InputContributionMultiplier = OutputToHistoryResolutionFractionSquare; 
 		PassParameters->GrandPrevPreExposureCorrection = bCameraCut ? 1.0f : View.PreExposure / InputHistory.PrevSceneColorPreExposure;
 		PassParameters->ResponsiveStencilMask = CVarTSREnableResponiveAA.GetValueOnRenderThread() ? (STENCIL_TEMPORAL_RESPONSIVE_AA_MASK) : 0;
 		PassParameters->bGenerateOutputMip1 = false;
