Strata - Fix for flickering trees du to strata wrap_thin not being handled correctly in a few spots.

diff --git a/Engine/Shaders/Private/Strata/StrataDbuffer.usf b/Engine/Shaders/Private/Strata/StrataDbuffer.usf
--- a/Engine/Shaders/Private/Strata/StrataDbuffer.usf
+++ b/Engine/Shaders/Private/Strata/StrataDbuffer.usf
@@ -276,7 +276,7 @@ void MainCS(
 				MaterialTextureArrayUAV[uint3(PixelCoord, 1)] = (PackedDiffuse20Bits | ((PackedFuzzColor20bits & 0xFFC00) << 10));
 				MaterialTextureArrayUAV[uint3(PixelCoord, 2)] = (PackedF020Bits | ((PackedFuzzColor20bits & 0x3FF) << 20));
 			}
-			else if (OptimisedLegacyMode == SINGLE_OPTLEGACYMODE_SSSWRAP)
+			else if (OptimisedLegacyMode == SINGLE_OPTLEGACYMODE_SSSWRAP || OptimisedLegacyMode == SINGLE_OPTLEGACYMODE_SSSWRAP_THIN)
 			{
 				// Roughness is only stored in the TopLayer texture.
 				const float3 MFP = lerp(0.f, SLAB_SSSMFP(BSDF), DBufferData.OneMinusCoverage_BaseColor);
diff --git a/Engine/Shaders/Private/Strata/StrataExport.ush b/Engine/Shaders/Private/Strata/StrataExport.ush
--- a/Engine/Shaders/Private/Strata/StrataExport.ush
+++ b/Engine/Shaders/Private/Strata/StrataExport.ush
@@ -621,13 +621,6 @@ void PackStrataOut(
 							// Ensure the rest of the BSDF is not stored with the regular path
 							BSDF.State = 0;
 						}
-						// else if (OptimisedLegacyMode == SINGLE_OPTLEGACYMODE_xxx)
-						//
-						// STRATA_TODO
-						// - Fast cloth
-						// - Fast subsurface
-						// - Fast diffusion with profile
-						// - Anisotropy as an option for Disney and clear coat as for current UE5?
 
 						#if (HEADER_SIMPLEENCODING_BIT_COUNT) > 32
 						#error Strata fast path header is > 32bits
diff --git a/Engine/Shaders/Private/Strata/StrataMaterialClassification.usf b/Engine/Shaders/Private/Strata/StrataMaterialClassification.usf
--- a/Engine/Shaders/Private/Strata/StrataMaterialClassification.usf
+++ b/Engine/Shaders/Private/Strata/StrataMaterialClassification.usf
@@ -147,12 +147,12 @@ void TileMainCS(uint2 DispatchThreadId : SV_DispatchThreadID, uint LinearIndex :
 
 		// Output/Patch SSS data for legacy encoding (this allows to save ALU & bandwidth during the base pass0
 		uint OptimisedLegacyMode = ((PackedHeader >> (HEADER_SINGLEENCODING_BIT_COUNT)) & HEADER_SINGLE_OPTLEGACYMODE_BIT_MASK);
-		const bool bIsLegacyWrap = OptimisedLegacyMode == SINGLE_OPTLEGACYMODE_SSSWRAP;
+		const bool bIsLegacyWrapOrWrapThin = OptimisedLegacyMode == SINGLE_OPTLEGACYMODE_SSSWRAP || OptimisedLegacyMode == SINGLE_OPTLEGACYMODE_SSSWRAP_THIN; // Wrap and Wrap thin have same packing
 		const bool bIsLegacySSSProfile = OptimisedLegacyMode == SINGLE_OPTLEGACYMODE_SSSPROFILE;
-		if (bIsSingle && (bIsLegacyWrap || bIsLegacySSSProfile))
+		if (bIsSingle && (bIsLegacyWrapOrWrapThin || bIsLegacySSSProfile))
 		{
 			bContainsScreenSpaceSubsurfaceScattering = true;
-			if (bIsLegacyWrap)
+			if (bIsLegacyWrapOrWrapThin)
 			{
 				const uint PackedSSSWOpacity7bits = (PackedHeader >> (HEADER_SINGLEENCODING_BIT_COUNT + HEADER_SINGLE_OPTLEGACYMODE_BIT_COUNT)) & 0x7F;
 				FStrataSubsurfaceHeader SSSHeader = (FStrataSubsurfaceHeader)0;
