Strata visualise shader update.

diff --git a/Engine/Shaders/Private/Strata/StrataVisualize.usf b/Engine/Shaders/Private/Strata/StrataVisualize.usf
--- a/Engine/Shaders/Private/Strata/StrataVisualize.usf
+++ b/Engine/Shaders/Private/Strata/StrataVisualize.usf
@@ -190,8 +190,8 @@ void PrintBSDF(inout FShaderPrintContext Context, inout FStrataAddressing Strata
 				Print(Context, TEXT(" "), FontWhite);
 			}
 
-			// Lighy weight
-			if (bGreyTopTrans)
+			// Light Weight
+			if (bGreyTopTrans && bWeightL)
 			{
 				Print(Context, TEXT(" "), FontSilver);
 				if (bGreyTopTrans)
