Eye Adaptation: Use the new previous frame mode for raytraced decals as well

diff --git a/Engine/Shaders/Private/RayTracing/RayTracingDecalMaterialShader.usf b/Engine/Shaders/Private/RayTracing/RayTracingDecalMaterialShader.usf
--- a/Engine/Shaders/Private/RayTracing/RayTracingDecalMaterialShader.usf
+++ b/Engine/Shaders/Private/RayTracing/RayTracingDecalMaterialShader.usf
@@ -3,6 +3,9 @@
 #define PATH_TRACING 1
 #define SCENE_TEXTURES_DISABLED 1 // Ray tracing shaders cannot access scene textures
 
+// This should be good enough for ray tracing and avoids having to bind an extra buffer
+#define EYE_ADAPTATION_PREV_FRAME_EXPOSURE 1
+
 #include "/Engine/Private/Common.ush"
 #include "/Engine/Private/RayTracing/RayTracingCommon.ush"
 
