Fixes highlight instability at r.TSR.History.ScreenPercentage=200

diff --git a/Engine/Shaders/Private/TemporalSuperResolution/TSRResolveHistory.usf b/Engine/Shaders/Private/TemporalSuperResolution/TSRResolveHistory.usf
--- a/Engine/Shaders/Private/TemporalSuperResolution/TSRResolveHistory.usf
+++ b/Engine/Shaders/Private/TemporalSuperResolution/TSRResolveHistory.usf
@@ -9,8 +9,6 @@
 
 #define KERNEL_SIZE 4
 
-#define CONFIG_PREMULTIPLY_HISTORY 0
-
 
 //------------------------------------------------------- CONSTANTS
 
@@ -85,7 +83,6 @@ uint2 OutputViewRectMax;
 uint bGenerateOutputMip1;
 float HistoryValidityMultiply;
 
-Texture2D<tsr_halfC> History_LowFrequency;
 Texture2D<tsr_halfC> History_HighFrequency;
 Texture2D<tsr_half2> History_Metadata;
 Texture2D<tsr_half3> History_Translucency;
@@ -108,9 +105,7 @@ void MainCS(
 		ZOrder2D(GroupThreadIndex, uint(log2(float(TILE_SIZE)))) +
 		GroupId * uint2(TILE_SIZE, TILE_SIZE));
 
-	tsr_halfC LowFrequencyColor = 0.0;
-	tsr_halfC HighFrequencyColor = 0.0;
-	tsr_half  HistoryValidity = 0.0;
+	tsr_halfC Color = tsr_half(0.0);
 	
 	// Downsample the history
 	{
@@ -121,9 +116,9 @@ void MainCS(
 		float2 OutputPixelPos = ApplyScreenTransform(float2(DispatchThreadId), DispatchThreadToHistoryPixelPos);
 		float2 TopLeftKernel = floor(OutputPixelPos + (0.5 - 0.5 * KERNEL_SIZE)) + 0.5;
 		
-		tsr_half LowFrequencyWeight = 0.0;
-		tsr_half HighFrequencyWeight = 0.0;
-		tsr_half HistoryValidityWeight = 0.0;
+		tsr_halfC MinColor = tsr_half(POSITIVE_INFINITY);
+		tsr_halfC MaxColor = tsr_half(0.0);
+		tsr_half ColorWeight = tsr_half(0.0);
 
 		UNROLL_N(KERNEL_SIZE)
 		for (uint x = 0; x < KERNEL_SIZE; x++)
@@ -137,29 +132,16 @@ void MainCS(
 
 				SampleUV = clamp(SampleUV, HistoryInfo_UVViewportBilinearMin, HistoryInfo_UVViewportBilinearMax);
 
-				tsr_halfC SampleLowFrequencyColor  = History_LowFrequency.SampleLevel(GlobalBilinearClampedSampler, SampleUV, 0);
-				tsr_halfC SampleHighFrequencyColor = History_HighFrequency.SampleLevel(GlobalBilinearClampedSampler, SampleUV, 0);
-				tsr_half2 SampleMetadata   	       = History_Metadata.SampleLevel(GlobalBilinearClampedSampler, SampleUV, 0);
-				
+				tsr_halfC SampleColor    = History_HighFrequency.SampleLevel(GlobalBilinearClampedSampler, SampleUV, 0);
 				tsr_half3 SampleTranslucencyColor = History_Translucency.SampleLevel(GlobalBilinearClampedSampler, SampleUV, 0);
 				tsr_half  SampleTranslucencyAlpha = History_TranslucencyAlpha.SampleLevel(GlobalBilinearClampedSampler, SampleUV, 0);
 				
-				#if CONFIG_PREMULTIPLY_HISTORY
-					tsr_half SampleHistoryValidity = SampleMetadata.y;
-					tsr_half InvSampleHistoryValidity = SafeRcp(SampleHistoryValidity);
-				#else
-					const tsr_half SampleHistoryValidity = tsr_half(1.0);
-					const tsr_half InvSampleHistoryValidity = tsr_half(1.0);
-				#endif
-
 				#if DIM_SEPARATE_TRANSLUCENCY
 				{
-					SampleLowFrequencyColor.rgb  = SampleLowFrequencyColor.rgb  * SampleTranslucencyAlpha + SampleTranslucencyColor.rgb;
-					SampleHighFrequencyColor.rgb = SampleHighFrequencyColor.rgb * SampleTranslucencyAlpha + SampleTranslucencyColor.rgb * SampleHistoryValidity;
+					SampleColor.rgb = SampleColor.rgb * SampleTranslucencyAlpha + SampleTranslucencyColor.rgb;
 				
 					#if CONFIG_SCENE_COLOR_ALPHA
-						SampleLowFrequencyColor.a = saturate(tsr_half(1.0) - saturate(tsr_half(1.0) - SampleLowFrequencyColor.a) * SampleTranslucencyAlpha);
-						SampleHighFrequencyColor.a = saturate(tsr_half(1.0) - saturate(tsr_half(1.0) - SampleHighFrequencyColor.a * InvSampleHistoryValidity) * SampleTranslucencyAlpha) * SampleHistoryValidity;
+						SampleColor.a = saturate(tsr_half(1.0) - saturate(tsr_half(1.0) - SampleColor.a) * SampleTranslucencyAlpha);
 					#endif
 				}
 				#endif
@@ -167,28 +149,21 @@ void MainCS(
 				float2 PixelOffset = abs(SamplePixelPos - OutputPixelPos);
 
 				tsr_half kernelWeight = tsr_half(MitchellNetravali(PixelOffset.x, B, C) * MitchellNetravali(PixelOffset.y, B, C));
-				
-				tsr_half LowFrequencyHdrWeight  = HdrWeight4(SampleLowFrequencyColor);
-				tsr_half HighFrequencyHdrWeight = HdrWeight4(SampleHighFrequencyColor * InvSampleHistoryValidity);
-
-				LowFrequencyColor   += (kernelWeight * LowFrequencyHdrWeight)  * SampleLowFrequencyColor;
-				LowFrequencyWeight  += (kernelWeight * LowFrequencyHdrWeight);
-
-				HighFrequencyColor  += (kernelWeight * HighFrequencyHdrWeight) * SampleHighFrequencyColor;
-				HighFrequencyWeight += (kernelWeight * HighFrequencyHdrWeight) * SampleHistoryValidity;
+				tsr_half SampleHdrWeight = HdrWeight4(SampleColor);
 
-				HistoryValidity       += kernelWeight * SampleHistoryValidity;
-				HistoryValidityWeight += kernelWeight;
+				MinColor = min(MinColor, SampleColor);
+				MaxColor = max(MaxColor, SampleColor);
+				Color  += (kernelWeight * SampleHdrWeight) * SampleColor;
+				ColorWeight += (kernelWeight * SampleHdrWeight);
 			}
 		}
 
-		LowFrequencyColor  *= rcp(LowFrequencyWeight);
-		HighFrequencyColor *= SafeRcp(HighFrequencyWeight);
-		HistoryValidity *= tsr_half(HistoryValidityMultiply) / HistoryValidityWeight;
+		Color *= SafeRcp(ColorWeight);
+		Color = clamp(Color, MinColor, MaxColor);
 	}
 	
 	// Compute final output
-	tsr_halfC FinalOutputColor = HighFrequencyColor;
+	tsr_halfC FinalOutputColor = Color;
 
 	tsr_short2 OutputPixelPosition = InvalidateOutputPixelPos(tsr_short2(OutputViewRectMin + DispatchThreadId), OutputViewRectMax);
 	
