Implement occluder distance and subsurface for local lights and one pass projection.
Generalize noise dither to work with both.

diff --git a/Engine/Shaders/Private/VirtualShadowMaps/VirtualShadowMapProjection.usf b/Engine/Shaders/Private/VirtualShadowMaps/VirtualShadowMapProjection.usf
--- a/Engine/Shaders/Private/VirtualShadowMaps/VirtualShadowMapProjection.usf
+++ b/Engine/Shaders/Private/VirtualShadowMaps/VirtualShadowMapProjection.usf
@@ -202,7 +202,8 @@ FVirtualShadowMapSampleResult ProjectLight(
 	float SceneDepth,
 	float ScreenRayLengthWorld,
 	float3 TranslatedWorldPosition,
-	const float Noise)
+	const float Noise,
+	float SubsurfaceOpacity)
 {
 	const bool bIsHairInput = InputType == INPUT_TYPE_HAIRSTRANDS;
 
@@ -309,8 +310,19 @@ FVirtualShadowMapSampleResult ProjectLight(
 				WorldNormal);
 		}
 
+		if (SubsurfaceOpacity < 1.0f)
+		{
+			// We use a single term for both the shadow and "transmission shadow" in the VSM path to provide a
+			// higher quality filtered result on back faces and avoid artifacts at the point where normals flip.
+			// This is especially important for foliage as normals are often "warped" to provide softer shading
+			// and do not represent the underlying geometry accurately, so having a continuous shadow function
+			// is important.
+			Result.ShadowFactor = GetShadowFactorSubsurface(Result.ShadowFactor, Result.OccluderDistance, SubsurfaceOpacity);
+		}
+
 		//Result.GeneralDebug = GreenToRedTurbo(1.0f - (SMRTRayOffset / ScreenRayLengthWorld));
 		//Result.GeneralDebug = SpotAttenuation(-L, Light.Direction, Light.SpotAngles).xxx;
+		//Result.GeneralDebug = (Result.OccluderDistance / 20.0f).xxx;
 
 		// Hair strands voxel traversal to apply hair shadow on top of opaque geometry
 		#if HAS_HAIR_STRANDS
@@ -402,6 +414,20 @@ FVirtualShadowMapSampleResult ProjectLight(
 		#endif
 	}
 
+	// Add some simple dither to break up banding
+	// This is not perfect for stuff like subsurface stuff in motion, as the banding on those smooth
+	// gradients can be significant and in motion TSR can't get enough samples to average it all out,
+	// but it's inexpensive amd usually pretty reasonable.
+	// Current noise scale aligns with one pass projection packing, but also works well for typical SMRT ray counts
+	// The tolerance on the lower end is somewhat arbitrary, but needed to trade off brightening umbra regions in subsurface
+	// materials (due to adding noise then clamping) against banding in penumbra regions on these surfaces.
+	const float NoiseScale = (1.0f / 15.0f);
+	if (Result.ShadowFactor > (NoiseScale/4.0f) && Result.ShadowFactor < 1.0f)
+	{
+		Result.ShadowFactor += (BlueNoiseScalar(PixelPos, View.StateFrameIndex) - 0.5f) * NoiseScale;
+		Result.ShadowFactor = saturate(Result.ShadowFactor);
+	}
+
 	return Result;
 }
 
@@ -584,6 +610,7 @@ void VirtualShadowMapProjection(
 	const float Noise = InterleavedGradientNoise( SvPosition.xy, View.StateFrameIndexMod8 );
 
 	const FProjectionShadingInfo ShadingInfo = GetProjectionShadingInfo(PixelPos);
+	const float SubsurfaceOpacity = bIsHairInput ? 1.0f : GetSubsurfaceOpacityFromGbuffer(PixelPos);
 
 	FVirtualShadowMapSampleResult VisualizeResult = InitVirtualShadowMapSampleResult();
 
@@ -615,9 +642,8 @@ void VirtualShadowMapProjection(
 				SceneDepth,
 				ScreenRayLengthWorld,
 				TranslatedWorldPosition,
-				Noise);
-
-			// NOTE: Subsurface opacity is handled in the clustered shading pass in the one pass projection path
+				Noise,
+				SubsurfaceOpacity);
 
 			PackShadowMask(ShadowMask, Result.ShadowFactor, Index);
 
@@ -636,8 +662,8 @@ void VirtualShadowMapProjection(
 		FLightShaderParameters Light = GetRootLightShaderParameters();
 
 		// Subsurface approximation
-#if DIRECTIONAL_LIGHT
 		const float SubsurfaceOpacity = bIsHairInput ? 1.0f : GetSubsurfaceOpacityFromGbuffer(PixelPos);
+#if DIRECTIONAL_LIGHT
 		// Increases light source radius for subsurface material as the opacity reduces to emulates light diffusion.
 		// TODO: Does something similar make sense for local lights?
 		if (SubsurfaceOpacity < 1.0f)
@@ -654,20 +680,8 @@ void VirtualShadowMapProjection(
 			SceneDepth,
 			ScreenRayLengthWorld,
 			TranslatedWorldPosition,
-			Noise);
-
-#if DIRECTIONAL_LIGHT
-		// TODO: Implement occluder distance and similar falloff for local lights, and unify logic into one pass projection path
-		if (SubsurfaceOpacity < 1.0f)
-		{
-			// We use a single term for both the shadow and "transmission shadow" in the VSM path to provide a
-			// higher quality filtered result on back faces and avoid artifacts at the point where normals flip.
-			// This is especially important for foliage as normals are often "warped" to provide softer shading
-			// and do not represent the underlying geometry accurately, so having a continuous shadow function
-			// is important.
-			Result.ShadowFactor = GetShadowFactorSubsurface(Result.ShadowFactor, Result.OccluderDistance, SubsurfaceOpacity);
-		}
-#endif
+			Noise,
+			SubsurfaceOpacity);
 
 		OutShadowFactor[ PixelPos ] = Result.ShadowFactor.xx;
 
diff --git a/Engine/Shaders/Private/VirtualShadowMaps/VirtualShadowMapProjectionDirectional.ush b/Engine/Shaders/Private/VirtualShadowMaps/VirtualShadowMapProjectionDirectional.ush
--- a/Engine/Shaders/Private/VirtualShadowMaps/VirtualShadowMapProjectionDirectional.ush
+++ b/Engine/Shaders/Private/VirtualShadowMaps/VirtualShadowMapProjectionDirectional.ush
@@ -338,7 +338,7 @@ FVirtualShadowMapSampleResult TraceDirectional(
 		*/
 		float OccluderDistance = (OccluderDistanceSum / float(max(1, RayCount - RayMissCount)));
 
-		Result.ShadowFactor = SMRTComputeRayMissFactor(RayMissCount, RayCount, MaxRayCount, Noise);
+		Result.ShadowFactor = float(RayMissCount) / float(RayCount);
 		Result.OccluderDistance = OccluderDistance;
 		Result.RayCount = RayCount;
 	}
diff --git a/Engine/Shaders/Private/VirtualShadowMaps/VirtualShadowMapProjectionSpot.ush b/Engine/Shaders/Private/VirtualShadowMaps/VirtualShadowMapProjectionSpot.ush
--- a/Engine/Shaders/Private/VirtualShadowMaps/VirtualShadowMapProjectionSpot.ush
+++ b/Engine/Shaders/Private/VirtualShadowMaps/VirtualShadowMapProjectionSpot.ush
@@ -393,6 +393,8 @@ FVirtualShadowMapSampleResult TraceLocalLight(
 
 		uint i = 0;
 		uint RayMissCount = 0;
+		float OccluderDistanceSum = 0.0f;
+		float MaxOccluderDistance = -1.0f;
 		for( ; i < MaxRayCount; i++ )
 		{
 			RandomSequence RandSequence;
@@ -471,7 +473,24 @@ FVirtualShadowMapSampleResult TraceLocalLight(
 					RayState);
 			}
 
-			RayMissCount += SMRTResult.bValidHit ? 0 : 1;
+			if (SMRTResult.bValidHit)
+			{
+				// TODO: Do we want to mess with this at all due to our offset?
+				float ReceiverDistance = length(RayStart);
+				float OccluderDistance = ComputeOccluderDistancePerspective(
+					// Re-fetch this here to avoid reg pressure
+					GetVirtualShadowMapProjectionData(RayState.VirtualShadowMapId).ShadowViewToClipMatrix,
+					SMRTResult.HitDepth,
+					saturate(RayState.RayStartUVz.z),
+					ReceiverDistance);
+
+				OccluderDistanceSum += OccluderDistance;
+				MaxOccluderDistance = max(MaxOccluderDistance, OccluderDistance);
+			}
+			else
+			{
+				++RayMissCount;
+			}
 			
 			// Debug output (DCE'd if not used)
 			Result.ClipmapIndexOrMipLevel = RayState.MipLevel;
@@ -480,25 +499,12 @@ FVirtualShadowMapSampleResult TraceLocalLight(
 
 			if (SMRTAdaptiveRayCount > 0)
 			{
+				// TODO: Adapt this heuristic based on SMRTAdaptiveRayCount as well?
 				if( i == 0 )
 				{
 					bool bHit = SMRTResult.bValidHit;
-					#if 0			// Disabled heuristic penumbra size path for now as it creates some speckle artifacts near the edges of sharp shadows
-						if( bHit )
-						{
-						#if 0
-							//float WorldBluringRadius = ComputeLightSampleWorldBluringRadius( TranslatedWorldPosition, LIGHT_TYPE_SPOT, Light, SMRTResult.OccluderDistance );
-							//float PixelBluringRadius = ProjectWorldDistanceToPixelDistance( SceneDepth, WorldBluringRadius );
-							float WorldBluringRadius = Light.SourceRadius * SMRTResult.OccluderDistance / max( DistToLight - SMRTResult.OccluderDistance, 0.001 );
-							float PixelBluringRadius = WorldBluringRadius * View.ViewSizeAndInvSize.x * 0.5 / ( SceneDepth * View.ClipToView[0][0] );
-							bHit = PixelBluringRadius > 2;
-						#else
-							bHit = SMRTResult.OccluderDistance * Light.SourceRadius * ( 0.25 * View.ViewSizeAndInvSize.y / ( SceneDepth * View.ClipToView[1][1] ) ) > max( DistToLight - Result.OccluderDistance, 0.001 );
-						#endif
-						}
-					#endif
-
-					// All lanes miss or hit something close enough to not be noisy.
+
+					// All lanes missed
 					bool bAllLanesMiss = WaveActiveAllTrue( !bHit );
 					if( bAllLanesMiss )
 					{
@@ -507,8 +513,7 @@ FVirtualShadowMapSampleResult TraceLocalLight(
 				}
 				else if( i >= SMRTAdaptiveRayCount )
 				{
-					// After 2 iterations and all have hit, must be in umbra.
-					// TODO: Modify this heuristic based on ray counts to avoid blocky artifacts in near-umbra
+					// After N iterations and all have hit, assume umbra
 					bool bAllLanesHit = WaveActiveAllTrue( RayMissCount == 0 );
 					if( bAllLanesHit )
 					{
@@ -522,9 +527,12 @@ FVirtualShadowMapSampleResult TraceLocalLight(
 		}
 		uint RayCount = min(i + 1, MaxRayCount);		// break vs regular for loop exit
 
-		Result.ShadowFactor = SMRTComputeRayMissFactor(RayMissCount, RayCount, MaxRayCount, Noise);
+		float OccluderDistance = (OccluderDistanceSum / float(max(1, RayCount - RayMissCount)));
+		//OccluderDistance = MaxOccluderDistance;
+
+		Result.ShadowFactor = float(RayMissCount) / float(RayCount);
 		Result.RayCount = RayCount;
-		Result.OccluderDistance = -1.0f;	// TODO?
+		Result.OccluderDistance = OccluderDistance;
 	}
 
 	return Result;
diff --git a/Engine/Shaders/Private/VirtualShadowMaps/VirtualShadowMapSMRTCommon.ush b/Engine/Shaders/Private/VirtualShadowMaps/VirtualShadowMapSMRTCommon.ush
--- a/Engine/Shaders/Private/VirtualShadowMaps/VirtualShadowMapSMRTCommon.ush
+++ b/Engine/Shaders/Private/VirtualShadowMaps/VirtualShadowMapSMRTCommon.ush
@@ -31,16 +31,3 @@ struct FSMRTResult
 	bool bValidHit;
 	float HitDepth;	// Valid if bValidHit is true
 };
-
-float SMRTComputeRayMissFactor(uint RayMissCount, uint RayCount, uint MaxRayCount, half Noise)
-{
-	float RayMissFactor = float(RayMissCount) / float(RayCount);
-
-	// Add some simple dither to break up banding
-	if (RayMissFactor > 0.0f && RayMissFactor < 1.0f)
-	{
-		RayMissFactor += (Noise - 0.5f) * float(1.0f / MaxRayCount);
-	}
-
-	return RayMissFactor;
-}
diff --git a/Engine/Source/Runtime/Renderer/Private/VirtualShadowMaps/VirtualShadowMapProjection.cpp b/Engine/Source/Runtime/Renderer/Private/VirtualShadowMaps/VirtualShadowMapProjection.cpp
--- a/Engine/Source/Runtime/Renderer/Private/VirtualShadowMaps/VirtualShadowMapProjection.cpp
+++ b/Engine/Source/Runtime/Renderer/Private/VirtualShadowMaps/VirtualShadowMapProjection.cpp
@@ -149,7 +149,7 @@ static TAutoConsoleVariable<int32> CVarForcePerLightShadowMaskClear(
 static TAutoConsoleVariable<int32> CVarSubsurfaceShadowMinSourceAngle(
 	TEXT("r.Shadow.Virtual.SubsurfaceShadowMinSourceAngle"),
 	5,
-	TEXT("Minimum source angle (in degrees) used for shadow & transmittance of sub-surface materials"),
+	TEXT("Minimum source angle (in degrees) used for shadow & transmittance of sub-surface materials with directional lights"),
 	ECVF_RenderThreadSafe
 );
 
