Add anisotropic subsurface scattering for path tracing.
1. Add a new node Subsurface Medium with two input MeanFreePath and ScatteringDistribution. If MeanFreePath is not connected, fallback to the old behavior (e.g., use the derived MFP for subsurface profile shading model), If ScatteringDistribution is not connected, fallback to zero.
2. It can be used by Subsurface, eye and subsurface profile/Preintegrated shading model.

diff --git a/Engine/Shaders/Private/PathTracing/PathTracingCore.ush b/Engine/Shaders/Private/PathTracing/PathTracingCore.ush
--- a/Engine/Shaders/Private/PathTracing/PathTracingCore.ush
+++ b/Engine/Shaders/Private/PathTracing/PathTracingCore.ush
@@ -824,14 +824,22 @@ bool ProcessSubsurfaceRandomWalk(inout FPathTracingPayload Payload, inout float3
 
 
 #if SSS_USE_TIR
-	const float3 Albedo = 1 - exp(SSSColor * (-11.43 + SSSColor * (15.38 - 13.91 * SSSColor)));
+	float3 Albedo = 1 - exp(SSSColor * (-11.43 + SSSColor * (15.38 - 13.91 * SSSColor)));
 #else
 	// Van de-Hulst inverse mapping
 	// https://blog.selfshadow.com/publications/s2017-shading-course/imageworks/s2017_pbs_imageworks_slides_v2.pdf (Slide 44)
 	// [REDACTED_PROFILE_URL] (Section 7.5.3 of v0.1.3)
-	const float3 Albedo = 1 - Pow2(4.09712 + 4.20863 * SSSColor - sqrt(9.59217 + SSSColor * (41.6808 + 17.7126 * SSSColor)));
+	float3 Albedo = 1 - Pow2(4.09712 + 4.20863 * SSSColor - sqrt(9.59217 + SSSColor * (41.6808 + 17.7126 * SSSColor)));
 	SSSRadius *= 2.0; // roughly match parameterization above
 #endif
+
+	// Apply albedo remapping for anisotropic scattering 
+	float G = Payload.GetSubsurfacePhaseFunction();
+
+	// Revisiting Physically Based Shading at Imageworks.
+	// https://blog.selfshadow.com/publications/s2017-shading-course/imageworks/s2017_pbs_imageworks_slides_v2.pdf
+	Albedo = Albedo / (1 - G * (1 - Albedo));
+
 	const int MAX_SSS_BOUNCES = MaxSSSBounces;
 	const float3 SigmaT = 1.0 / SSSRadius;
 	const float3 SigmaS = Albedo * SigmaT;
@@ -851,10 +859,12 @@ bool ProcessSubsurfaceRandomWalk(inout FPathTracingPayload Payload, inout float3
 			FProbeResult ProbeResult = TraceProbeRay(Ray);
 			if (ProbeResult.IsMiss())
 			{
-				// we didn't hit anything, so scatter randomly in the volume and keep tracing
+				// we didn't hit anything, so scatter according to the scattering distribution in the volume and keep tracing
 				RandSample = RandomSequence_GenerateSample3D(RandSequence);
 				Ray.Origin += Ray.TMax * Ray.Direction;
-				Ray.Direction = UniformSampleSphere(RandSample.xy).xyz;
+				
+				float4 DirectionAndPhase = ImportanceSampleHenyeyGreensteinPhase(RandSample.xy, G);
+				Ray.Direction = TangentToWorld(DirectionAndPhase.xyz, Ray.Direction);
 				Ray.TMin = 0;
 				// account for transmittance and scattering up to this point
 				PathThroughput *= SigmaS * EvaluateSpectralTransmittanceHit(Ray.TMax, SigmaT, ColorChannelPdf).xyz;
diff --git a/Engine/Shaders/Private/PathTracing/PathTracingMaterialHitShader.usf b/Engine/Shaders/Private/PathTracing/PathTracingMaterialHitShader.usf
--- a/Engine/Shaders/Private/PathTracing/PathTracingMaterialHitShader.usf
+++ b/Engine/Shaders/Private/PathTracing/PathTracingMaterialHitShader.usf
@@ -67,6 +67,11 @@ float3 DecodeSSSProfileRadius(uint ProfileId, float3 DiffuseColor, float Opacity
 	return SSSRadius * BURLEY_MM_2_CM;
 }
 
+float DecodeSSSProfileScatteringDistribution(uint ProfileId)
+{
+	float EncodedScatteringDistribution = View.SSProfilesTexture.Load(int3(SSSS_TRANSMISSION_OFFSET, ProfileId, 0)).z;
+	return DecodeScatteringDistribution(EncodedScatteringDistribution);
+}
 
 RAY_TRACING_ENTRY_CLOSEST_HIT(PathTracingMaterialCHS,
 	FPackedPathTracingPayload, PackedPayload,
@@ -285,11 +290,20 @@ RAY_TRACING_ENTRY_CLOSEST_HIT(PathTracingMaterialCHS,
 		const float3 SubsurfaceColor = GetMaterialSubsurfaceData(PixelMaterialInputs).rgb;
 		Payload.SetSubsurfaceColor(SubsurfaceColor);
 
+#if HAVE_GetSubsurfaceMediumMaterialOutput0
+		Payload.SetSubsurfaceRadius(GetSubsurfaceMediumMaterialOutput0(MaterialParameters));
+#else
 		// TODO: is this accurate enough? hard to measure exactly since the raster path uses a very different approach
-		// TODO: It would be nice to have a dedicated node to control the SSS Radius explicitly instead of this
 		const float SSSRadius = (1 - Payload.Opacity) * 10.0; // simple formula, up to 10cm radius
 		Payload.SetSubsurfaceRadius(SSSRadius);
-		Payload.SetSubsurfacePhaseFunction(0.0); // TODO: expose via a dedicated node?
+#endif
+
+#if HAVE_GetSubsurfaceMediumMaterialOutput1
+		float ScatteringDistribution = clamp(GetSubsurfaceMediumMaterialOutput1(MaterialParameters), -0.99f, 0.99f);
+		Payload.SetSubsurfacePhaseFunction(ScatteringDistribution);
+#else
+		Payload.SetSubsurfacePhaseFunction(0.0);
+#endif
 		// Reset opacity to avoid special cases elsewhere
 		Payload.Opacity = 1;
 	}
@@ -311,6 +325,9 @@ RAY_TRACING_ENTRY_CLOSEST_HIT(PathTracingMaterialCHS,
 		Lobe1Roughness = AdjustMaterialRoughness(Lobe1Roughness, PackedPayload.GetPathRoughness());
 		Payload.SetDualRoughnessSpecular(Lobe0Roughness, Lobe1Roughness, LobeMix);
 
+#if HAVE_GetSubsurfaceMediumMaterialOutput0
+		Payload.SetSubsurfaceRadius(GetSubsurfaceMediumMaterialOutput0(MaterialParameters));
+#else
 		if (GetSubsurfaceProfileUseBurley(ProfileId))
 		{
 			// Decode SSS radius
@@ -321,7 +338,15 @@ RAY_TRACING_ENTRY_CLOSEST_HIT(PathTracingMaterialCHS,
 		{
 			Payload.SetSubsurfaceRadius(0.0);
 		}
-		Payload.SetSubsurfacePhaseFunction(0.0); // TODO: extract from profile data
+#endif
+
+#if HAVE_GetSubsurfaceMediumMaterialOutput1
+		float ScatteringDistribution = clamp(GetSubsurfaceMediumMaterialOutput1(MaterialParameters), -0.99f, 0.99f);
+		Payload.SetSubsurfacePhaseFunction(ScatteringDistribution);
+#else
+		Payload.SetSubsurfacePhaseFunction(0.0);
+#endif
+
 		Payload.Opacity = 1.0; // Reset opacity to avoid special cases elsewhere
 	}
 	else
@@ -411,7 +436,10 @@ RAY_TRACING_ENTRY_CLOSEST_HIT(PathTracingMaterialCHS,
 		// IRIS_NORMAL off
 		Payload.SetEyeIrisNormal(PlaneNormal);
 #endif
-		
+
+#if HAVE_GetSubsurfaceMediumMaterialOutput0
+		Payload.SetSubsurfaceRadius(GetSubsurfaceMediumMaterialOutput0(MaterialParameters));
+#else		
 		const uint ProfileId = ExtractSubsurfaceProfileInt(GetMaterialSubsurfaceDataRaw(PixelMaterialInputs).a);
 		if (GetSubsurfaceProfileUseBurley(ProfileId))
 		{
@@ -423,8 +451,15 @@ RAY_TRACING_ENTRY_CLOSEST_HIT(PathTracingMaterialCHS,
 		{
 			Payload.SetSubsurfaceRadius(0.0);
 		}
+#endif
 		Payload.Opacity = 1.0; // revert to default now that decode is done
-		Payload.SetSubsurfacePhaseFunction(0.0); // TODO: extract from profile data
+
+#if HAVE_GetSubsurfaceMediumMaterialOutput1
+		float ScatteringDistribution = clamp(GetSubsurfaceMediumMaterialOutput1(MaterialParameters), -0.99f, 0.99f);
+		Payload.SetSubsurfacePhaseFunction(ScatteringDistribution);
+#else
+		Payload.SetSubsurfacePhaseFunction(0.0);
+#endif
 
 	}
 	else
diff --git a/Engine/Shaders/Private/SubsurfaceProfileCommon.ush b/Engine/Shaders/Private/SubsurfaceProfileCommon.ush
--- a/Engine/Shaders/Private/SubsurfaceProfileCommon.ush
+++ b/Engine/Shaders/Private/SubsurfaceProfileCommon.ush
@@ -117,6 +117,16 @@ float4 DecodeDiffuseMeanFreePath(float4 EncodedDiffuseMeanFreePath)
 	return EncodedDiffuseMeanFreePath * DEC_UNIT_TO_DIFFUSEMEANFREEPATH_IN_MM;
 }
 
+float EncodeScatteringDistribution(float ScatteringDistribution)
+{
+	return (ScatteringDistribution + 1.0f) * 0.5f;
+}
+
+float DecodeScatteringDistribution(float ScatteringDistribution)
+{
+	return ScatteringDistribution * 2.0f - 1.0f;
+}
+
 bool GetSubsurfaceTransmittanceProfileUseBurley(uint SubsurfaceProfileInt)
 {
 	half Type = GetSubsurfaceProfileTexture(SSSS_BOUNDARY_COLOR_BLEED_OFFSET, SubsurfaceProfileInt).a;
diff --git a/Engine/Shaders/Private/TransmissionCommon.ush b/Engine/Shaders/Private/TransmissionCommon.ush
--- a/Engine/Shaders/Private/TransmissionCommon.ush
+++ b/Engine/Shaders/Private/TransmissionCommon.ush
@@ -36,7 +36,7 @@ FTransmissionProfileParams GetTransmissionProfileParams(uint SubsurfaceProfileIn
 
 	Result.ExtinctionScale = Data.x;
 	Result.NormalScale = Data.y;
-	Result.ScatteringDistribution = Data.z;
+	Result.ScatteringDistribution = DecodeScatteringDistribution(Data.z);
 	Result.OneOverIOR = Data.w;
 
 	return Result;
diff --git a/Engine/Source/Runtime/Engine/Classes/Materials/MaterialExpressionSubsurfaceMediumMaterialOutput.h b/Engine/Source/Runtime/Engine/Classes/Materials/MaterialExpressionSubsurfaceMediumMaterialOutput.h
new file mode 100644
--- /dev/null
+++ b/Engine/Source/Runtime/Engine/Classes/Materials/MaterialExpressionSubsurfaceMediumMaterialOutput.h
@@ -0,0 +1,38 @@
+// Copyright Epic Games, Inc. All Rights Reserved.
+
+#pragma once
+
+#include "CoreMinimal.h"
+#include "Materials/MaterialExpressionCustomOutput.h"
+#include "UObject/ObjectMacros.h"
+#include "MaterialExpressionSubsurfaceMediumMaterialOutput.generated.h"
+
+/** Material output expression for setting the mean free path and scattering distribution properties of subsurface profile (for the Path Tracer Only). */
+UCLASS(MinimalAPI, collapsecategories, hidecategories = Object)
+class UMaterialExpressionSubsurfaceMediumMaterialOutput : public UMaterialExpressionCustomOutput
+{
+	GENERATED_UCLASS_BODY()
+
+	/** Input for mean free path (cm). Fallback to the default behavior if not set (e.g., Subsurfaceprofile shading: Use the derived MFP) */
+	UPROPERTY()
+	FExpressionInput MeanFreePath;
+
+	/** Input for scattering distribution. Valid range is (-1, 1). Fallback to zero if not set*/
+	UPROPERTY()
+	FExpressionInput ScatteringDistribution;
+
+public:
+#if WITH_EDITOR
+	//~ Begin UMaterialExpression Interface
+	virtual int32 Compile(class FMaterialCompiler* Compiler, int32 OutputIndex) override;
+	virtual void GetCaption(TArray<FString>& OutCaptions) const override;
+	virtual uint32 GetInputType(int32 InputIndex) override;
+	//~ End UMaterialExpression Interface
+#endif
+
+	//~ Begin UMaterialExpressionCustomOutput Interface
+	virtual int32 GetNumOutputs() const override;
+	virtual FString GetFunctionName() const override;
+	virtual FString GetDisplayName() const override;
+	//~ End UMaterialExpressionCustomOutput Interface
+};
diff --git a/Engine/Source/Runtime/Engine/Private/Materials/MaterialExpressions.cpp b/Engine/Source/Runtime/Engine/Private/Materials/MaterialExpressions.cpp
--- a/Engine/Source/Runtime/Engine/Private/Materials/MaterialExpressions.cpp
+++ b/Engine/Source/Runtime/Engine/Private/Materials/MaterialExpressions.cpp
@@ -152,6 +152,7 @@
 #include "Materials/MaterialExpressionStaticBoolParameter.h"
 #include "Materials/MaterialExpressionStaticSwitchParameter.h"
 #include "Materials/MaterialExpressionStaticComponentMaskParameter.h"
+#include "Materials/MaterialExpressionSubsurfaceMediumMaterialOutput.h"
 #include "Materials/MaterialExpressionVectorParameter.h"
 #include "Materials/MaterialExpressionDoubleVectorParameter.h"
 #include "Materials/MaterialExpressionParticleColor.h"
@@ -25817,4 +25818,90 @@ bool UMaterialExpressionSparseVolumeTextureSampleParameter::MatchesSearchQuery(c
 }
 #endif
 
+///////////////////////////////////////////////////////////////////////////////
+// UMaterialExpressionSubsurfaceMediumMaterialOutput
+///////////////////////////////////////////////////////////////////////////////
+
+UMaterialExpressionSubsurfaceMediumMaterialOutput::UMaterialExpressionSubsurfaceMediumMaterialOutput(const FObjectInitializer& ObjectInitializer)
+	: Super(ObjectInitializer)
+{
+	// Structure to hold one-time initialization
+	struct FConstructorStatics
+	{
+		FText NAME_SubsurfaceMedium;
+		FConstructorStatics()
+			: NAME_SubsurfaceMedium(LOCTEXT("Shading", "Shading"))
+		{
+		}
+	};
+	static FConstructorStatics ConstructorStatics;
+
+#if WITH_EDITORONLY_DATA
+	MenuCategories.Add(ConstructorStatics.NAME_SubsurfaceMedium);
+#endif
+
+#if WITH_EDITOR
+	Outputs.Reset();
+#endif
+}
+
+#if WITH_EDITOR
+
+int32 UMaterialExpressionSubsurfaceMediumMaterialOutput::Compile(class FMaterialCompiler* Compiler, int32 OutputIndex)
+{
+
+	int32 CodeInput = INDEX_NONE;
+	
+	if (OutputIndex == 0)
+	{
+		CodeInput = MeanFreePath.IsConnected() ? MeanFreePath.Compile(Compiler) : INDEX_NONE;
+	}
+	else if (OutputIndex == 1)
+	{
+		CodeInput = ScatteringDistribution.IsConnected() ? ScatteringDistribution.Compile(Compiler) : INDEX_NONE;
+	}
+
+	return Compiler->CustomOutput(this, OutputIndex, CodeInput);
+}
+
+void UMaterialExpressionSubsurfaceMediumMaterialOutput::GetCaption(TArray<FString>& OutCaptions) const
+{
+	OutCaptions.Add(FString(TEXT("Subsurface Medium (Path Tracer Only)")));
+}
+
+uint32 UMaterialExpressionSubsurfaceMediumMaterialOutput::GetInputType(int32 InputIndex)
+{
+	switch (InputIndex)
+	{
+	case 0:
+		return MCT_Float3;
+		break;
+	case 1:
+		return MCT_Float1;
+		break;
+	default:
+		break;
+	}
+
+	check(false);
+	return MCT_Float1;
+}
+
+#endif // WITH_EDITOR
+
+int32 UMaterialExpressionSubsurfaceMediumMaterialOutput::GetNumOutputs() const
+{
+	return 2;
+}
+
+FString UMaterialExpressionSubsurfaceMediumMaterialOutput::GetFunctionName() const
+{
+	return TEXT("GetSubsurfaceMediumMaterialOutput");
+}
+
+FString UMaterialExpressionSubsurfaceMediumMaterialOutput::GetDisplayName() const
+{
+	return TEXT("Subsurface Medium");
+}
+
 #undef LOCTEXT_NAMESPACE
diff --git a/Engine/Source/Runtime/Engine/Private/Rendering/SubsurfaceProfile.cpp b/Engine/Source/Runtime/Engine/Private/Rendering/SubsurfaceProfile.cpp
--- a/Engine/Source/Runtime/Engine/Private/Rendering/SubsurfaceProfile.cpp
+++ b/Engine/Source/Runtime/Engine/Private/Rendering/SubsurfaceProfile.cpp
@@ -380,6 +380,17 @@ FLinearColor DecodeDiffuseMeanFreePath(FLinearColor EncodedDiffuseMeanFreePath)
 	return EncodedDiffuseMeanFreePath * DEC_UNIT_TO_DIFFUSEMEANFREEPATH_IN_MM;
 }
 
+//in [0,1]
+float EncodeScatteringDistribution(float ScatteringDistribution)
+{
+	return (ScatteringDistribution + 1.0f) * 0.5f;
+}
+
+float DecodeScatteringDistribution(float ScatteringDistribution)
+{
+	return ScatteringDistribution * 2.0f - 1.0f;
+}
+
 void SetupSurfaceAlbedoAndDiffuseMeanFreePath(FLinearColor& SurfaceAlbedo, FLinearColor& Dmfp)
 {
 	int32 SamplingSelectionMethod = FMath::Clamp(CVarSSProfilesSamplingChannelSelection.GetValueOnAnyThread(), 0, 1);
@@ -497,7 +508,7 @@ void FSubsurfaceProfileTexture::CreateTexture(FRHICommandListImmediate& RHICmdLi
 		//X:ExtinctionScale, Y:Normal Scale, Z:ScatteringDistribution, W:OneOverIOR
 		TextureRow[SSSS_TRANSMISSION_OFFSET].R = Data.ExtinctionScale;
 		TextureRow[SSSS_TRANSMISSION_OFFSET].G = Data.NormalScale;
-		TextureRow[SSSS_TRANSMISSION_OFFSET].B = Data.ScatteringDistribution;
+		TextureRow[SSSS_TRANSMISSION_OFFSET].B = EncodeScatteringDistribution(Data.ScatteringDistribution);
 		TextureRow[SSSS_TRANSMISSION_OFFSET].A = 1.0f / Data.IOR;
 
 		if (Data.bEnableBurley)
