Fix 'ambiguous type for bit shift' shader compiler spam from "OcclusionMask |= bIsBentNormalOccluded ? (1 << TracingRayIndex) : 0;"

diff --git a/Engine/Shaders/Private/BasePassPixelShader.usf b/Engine/Shaders/Private/BasePassPixelShader.usf
--- a/Engine/Shaders/Private/BasePassPixelShader.usf
+++ b/Engine/Shaders/Private/BasePassPixelShader.usf
@@ -679,7 +679,7 @@ uint GetDiffuseIndirectSampleOcclusion(FSharedLocalBases SharedLocalBases, float
 				float LVisibility = saturate(Evaluate(VisibleSG, BxDFSample.L) / Evaluate(HemisphereSG, BxDFSample.L));
 				bool bIsBentNormalOccluded = LVisibility < VisibilityThreshold;
 
-				OcclusionMask |= bIsBentNormalOccluded ? (1 << TracingRayIndex) : 0;
+				OcclusionMask |= bIsBentNormalOccluded ? (1u << TracingRayIndex) : 0u;
 			}
 		}
 
@@ -713,7 +713,7 @@ uint GetDiffuseIndirectSampleOcclusion(FGBufferData GBuffer, float3 V, float3 Wo
 			float LVisibility = saturate(Evaluate(VisibleSG, BxDFSample.L) / Evaluate(HemisphereSG, BxDFSample.L));
 			bool bIsBentNormalOccluded = LVisibility < VisibilityThreshold;
 
-			OcclusionMask |= bIsBentNormalOccluded ? (1 << TracingRayIndex) : 0;
+			OcclusionMask |= bIsBentNormalOccluded ? (1u << TracingRayIndex) : 0u;
 		}
 
 		DiffuseIndirectSampleOcclusion = OcclusionMask;
