Use ETC1S for RVT RGB compression in ETC path
Added CoCg scaling in ETC path

diff --git a/Engine/Shaders/Private/ETCCompressionCommon.ush b/Engine/Shaders/Private/ETCCompressionCommon.ush
--- a/Engine/Shaders/Private/ETCCompressionCommon.ush
+++ b/Engine/Shaders/Private/ETCCompressionCommon.ush
@@ -14,169 +14,179 @@
 #define NUM_RGB_TABLES 8
 #define NUM_ALPHA_TABLES 16
 
-int3 FloatColorToUint555(in float3 FloatColor)
+float EtcLuminance(in float3 c)
 {
-	float3 Scale       = float3(31.f, 31.f, 31.f);
-	float3 ColorScaled = round(saturate(FloatColor) * Scale);
-	return int3((int)ColorScaled.r, (int)ColorScaled.g, (int)ColorScaled.b);
+	return dot(c, float3(0.3f, 0.5f, 0.2f));
 }
 
-float3 ExpandColor444(in int3 Color444)
+uint EtcPackColorBits(in float3 FloatColor)
 {
-	int3 Color888 = (Color444 << 4) + Color444;
-	return Color888 / 255.f;
+	uint3 ColorScaled = uint3(round(saturate(FloatColor) * 255.f));
+	uint ColorPacked = (ColorScaled.r & 0x000000f8) | ((ColorScaled.g & 0x000000f8) << 8u) | ((ColorScaled.b & 0x000000f8) << 16u);
+	return ColorPacked;
 }
 
-float3 ExpandColor555(in int3 Color555)
+uint EtcSwapEndian32(in uint x)
 {
-	int3 Color888 = (Color555 << 3) + (Color555 >> 2);
-	return Color888 / 255.f;
+	return
+		((x & 0x0000ff) << 24u) |
+		((x & 0x00ff00) <<  8u) |
+		((x & 0xff0000) >>  8u) |
+		(x >> 24u);
 }
 
-uint SwapEndian32(in uint x)
+// Get a single scale factor to use for a YCoCg color block
+// This increases precision at the expense of potential blending artifacts across blocks
+float EtcGetYCoCgScale(float2 MinCoCg, float2 MaxCoCg)
 {
-	return
-		((x & 0x0000ff) << 24) |
-		((x & 0x00ff00) <<  8) |
-		((x & 0xff0000) >>  8) |
-		(x >> 24);
+	MinCoCg = abs(MinCoCg - 128.f / 255.f);
+	MaxCoCg = abs(MaxCoCg - 128.f / 255.f);
+
+	float MaxComponent = max(max(MinCoCg.x, MinCoCg.y), max(MaxCoCg.x, MaxCoCg.y));
+
+	return (MaxComponent < 32.f / 255.f) ? 4.f : (MaxComponent < 64.f / 255.f) ? 2.f : 1.f;
 }
 
-void FindPixelWeights(in float3 Block[16], in float3 BaseColor, in uint TableIdx, in int StartX, in int EndX, in int StartY, in int EndY, out uint SubBlockWeights)
+void EtcApplyYCoCgScale(inout float2 CoCg, float Scale)
 {
-	//int PIXEL_INDEX_ENCODE_TABLE[4] = { 3, 2, 0, 1 };
-	SubBlockWeights = 0;
-	
-	float TableRangeMax = EtcParameters.RGB_DISTANCE_TABLES[TableIdx].w / 255.f;
-	float BaseLum = Luminance(BaseColor);
+	CoCg = (CoCg - 128.f / 255.f) * Scale + 128.f / 255.f;
+}
+
+uint EtcGetBlockWeights(in float3 Block[16], in float3 BaseColor, in uint TableIdx)
+{
+	// weights
+	// 3 - big negative distance
+	// 2 - small negative distance
+	// 1 - big positive distance
+	// 0 - small positive distance
+	uint BlockWeights = 0;
 
-	for (int Y = StartY; Y < EndY; ++Y)
+	float Distance0 = EtcParameters.RGB_DISTANCE_TABLES[TableIdx].z / 255.f;
+	float Distance1 = EtcParameters.RGB_DISTANCE_TABLES[TableIdx].w / 255.f;
+	
+	float BaseLum = EtcLuminance(BaseColor);
+	float Lum1 = min(1.f, BaseLum + Distance1);
+	float Lum0 = BaseLum + Distance0;
+	float Lum3 = max(0.f, BaseLum - Distance1);
+	float Lum2 = BaseLum - Distance0;
+	
+	for (uint Y = 0; Y < 4; ++Y)
 	{
-		for (int X = StartX; X < EndX; ++X)
+		for (uint X = 0; X < 4; ++X)
 		{
-			float3 OrigColor = Block[4 * Y + X];
-			float Diff = Luminance(OrigColor) - BaseLum;
-			int EncIndex = 0;
-			if (Diff < 0.f)
+			float3 Color = Block[4u * Y + X];
+			float ColorLum = EtcLuminance(Color);
+			uint EncIndex = 0u;
+			if (ColorLum < BaseLum || ColorLum == 0.f) // make sure fully Black block is encoded correctly
 			{
-				EncIndex = (-Diff * 1.58) > TableRangeMax ? 3 : 2;  
+				EncIndex = (ColorLum - Lum3) < (Lum2 - ColorLum) ? 3u : 2u;
 			}
 			else
 			{
-				EncIndex = (Diff * 1.58) > TableRangeMax ? 1 : 0;
+				EncIndex = (Lum1 - ColorLum) < (ColorLum - Lum0) ? 1u : 0u;
 			}
-			//int EncIndex = PIXEL_INDEX_ENCODE_TABLE[SelectedIndex];
-			int IndexInBlock = X * 4 + Y;
-			SubBlockWeights |= ((EncIndex & 1) << IndexInBlock) | ((EncIndex >> 1) << (16 + IndexInBlock));
+			
+			uint IndexInBlock = X * 4u + Y;
+			BlockWeights |= ((EncIndex & 1u) << IndexInBlock) | ((EncIndex >> 1u) << (16u + IndexInBlock));
 		}
 	}
+	return BlockWeights;
 }
 
-uint SelectRGBTableIndex(float LuminanceR)
+uint EtcSelectRGBTableIndex(float LuminanceR)
 {
 	// guess a table using sub-block luminance range
-	float Range = (LuminanceR + LuminanceR * 0.1)*255.0;
+	float Range = LuminanceR*255.f + 1.f;
 	if (Range < 8.0)
 	{
-		return 0;
+		return 0u;
 	}
 	else if (Range < 17.0)
 	{
-		return 1;
+		return 1u;
 	}
 	else if (Range < 29.0)
 	{
-		return 2;
+		return 2u;
 	}
 	else if (Range < 42.0)
 	{
-		return 3;
+		return 3u;
 	}
 	else if (Range < 60.0)
 	{
-		return 4;
+		return 4u;
 	}
 	else if (Range < 80.0)
 	{
-		return 5;
+		return 5u;
 	}
 	else if (Range < 106.0)
 	{
-		return 6;
+		return 6u;
 	}
-	return 7;
+	return 7u;
 }
 
-uint2 CompressBlock_ETC2_RGB(in float3 Block[16])
+uint2 CompressBlock_ETC1S(in float3 Block[16], bool bSRGB, bool bCoCg)
 {
-	// Always use side-by-side mode (flip bit set to 0).
-	uint FlipBit = 0;
-	
-	float3 BaseColor1_Float = (Block[0] + Block[1] + Block[4] + Block[5] +  Block[8] +  Block[9] + Block[12] + Block[13]) * 0.125;
-	float3 BaseColor2_Float = (Block[2] + Block[3] + Block[6] + Block[7] + Block[10] + Block[11] + Block[14] + Block[15]) * 0.125;
-		
-	int3 BaseColor1 = FloatColorToUint555(BaseColor1_Float);
-	int3 BaseColor2 = FloatColorToUint555(BaseColor2_Float);
-	int3 Diff = BaseColor2 - BaseColor1;
-
-	uint ColorBits;
-	float3 BaseColor1_Quant, BaseColor2_Quant;
+	// ETC1S: differential with color delta=0
+	uint FlipBit = 0u;
+	uint BlockModeBits = (BLOCK_MODE_DIFFERENTIAL << 25u);
 
-	// Changed block mode to store bits because of DXC compiler optimization error
-	uint BlockModeBits;
-	int3 MinDiff = { -4, -4, -4 }, MaxDiff = { 3, 3, 3 };
-	if (all(Diff >= MinDiff) && all(Diff <= MaxDiff))
+	float3 BaseColorMin = 1;
+	float3 BaseColorMax = 0;
+	float3 BaseColorAvg = 0;
+	for (int i = 0; i < 16; ++i)
 	{
-		// We can use differential mode.
-		BlockModeBits = BLOCK_MODE_DIFFERENTIAL << 25;
-		ColorBits = ((Diff.b & 7) << 16) | (BaseColor1.b << 19) | ((Diff.g & 7) << 8) | (BaseColor1.g << 11) | (Diff.r & 7) | (BaseColor1.r << 3);
-		BaseColor1_Quant = ExpandColor555(BaseColor1);
-		BaseColor2_Quant = ExpandColor555(BaseColor2);
+		BaseColorMin = min(BaseColorMin, Block[i]);
+		BaseColorMax = max(BaseColorMax, Block[i]);
+		BaseColorAvg = BaseColorAvg + Block[i];
+	}
+	BaseColorAvg = BaseColorAvg / 16.f;
+		
+	float3 BaseColor = (bSRGB ? LinearToSrgb(BaseColorAvg) : BaseColorAvg);
+	
+	uint ColorBits = 0u;
+	if (bCoCg)
+	{
+		float CoCgScale = EtcGetYCoCgScale(BaseColorMin.xy, BaseColorMax.xy);
+		EtcApplyYCoCgScale(BaseColor.xy, CoCgScale);
+		ColorBits = EtcPackColorBits(BaseColor);
+		ColorBits |= (((uint)CoCgScale - 1) << 19u);
 	}
 	else
 	{
-		// We must use the lower precision individual mode.
-		BlockModeBits = BLOCK_MODE_INDIVIDUAL << 25;
-		BaseColor1 >>= 1;
-		BaseColor2 >>= 1;
-		ColorBits = (BaseColor1.b << 20) | (BaseColor2.b << 16) | (BaseColor1.g << 12) | (BaseColor2.g << 8) | (BaseColor1.r << 4) | BaseColor2.r;
-		BaseColor1_Quant = ExpandColor444(BaseColor1);
-		BaseColor2_Quant = ExpandColor444(BaseColor2);
+		ColorBits = EtcPackColorBits(BaseColor);
 	}
 
-	float l00 = Luminance(Block[0]);
-	float l08 = Luminance(Block[8]);
-	float l13 = Luminance(Block[13]);
-	float LuminanceR1 = (max3(l00, l08, l13) - min3(l00, l08, l13))*0.5;
-	uint SubBlock1TableIdx = SelectRGBTableIndex(LuminanceR1);
-	uint SubBlock1Weights = 0;
-	FindPixelWeights(Block, BaseColor1_Quant, SubBlock1TableIdx, 0, 2, 0, 4, SubBlock1Weights);
-
-	float l02 = Luminance(Block[2]);
-	float l10 = Luminance(Block[10]);
-	float l15 = Luminance(Block[15]);
-	float LuminanceR2 =  (max3(l02, l10, l15) - min3(l02, l10, l15))*0.5;
-	uint SubBlock2TableIdx = SelectRGBTableIndex(LuminanceR2);
-	uint SubBlock2Weights = 0;
-	FindPixelWeights(Block, BaseColor2_Quant, SubBlock2TableIdx, 2, 4, 0, 4, SubBlock2Weights);
-	
+	float LumMin = EtcLuminance(BaseColorMin);
+	float LumMax = EtcLuminance(BaseColorMax);
+	float LumRange = (LumMax - LumMin) * 0.5f;
+
+	uint BlockTableIdx = EtcSelectRGBTableIndex(LumRange);
+	uint BlockWeights = EtcGetBlockWeights(Block, BaseColorAvg, BlockTableIdx);
+
 	// Both these values need to be big-endian. We can build ModeBits directly in big-endian layout, but for IndexBits
 	// it's too hard, so we'll just swap here.
-	uint ModeBits = (SubBlock1TableIdx << 29) | (SubBlock2TableIdx << 26) | (BlockModeBits) | (FlipBit << 24) | ColorBits;
-	uint IndexBits = SwapEndian32(SubBlock1Weights | SubBlock2Weights);
+	uint ModeBits = (BlockTableIdx << 29u) | (BlockTableIdx << 26u) | (BlockModeBits) | (FlipBit << 24u) | ColorBits;
+	uint IndexBits = EtcSwapEndian32(BlockWeights);
 
 	return uint2(ModeBits, IndexBits);
 }
 
 uint2 CompressBlock_ETC2_SRGB(in float3 Block[16])
 {
-	// sRGB not supported on mobile. Get rid of this function completely, to avoid confusion, and call the non-sRGB version in VirtualTextureCompress?
-	//for(int i = 0; i < 16; ++i)
-	//{
-	//	Block[i] = LinearToSrgb(Block[i]);
-	//}
+	const bool bSRGB = false;// true;
+	const bool bCoCg = false;
+	return CompressBlock_ETC1S(Block, bSRGB, bCoCg);
+}
 
-	return CompressBlock_ETC2_RGB(Block);
+uint2 CompressBlock_ETC2_RGB(in float3 Block[16])
+{
+	const bool bSRGB = false;
+	const bool bCoCg = false;
+	return CompressBlock_ETC1S(Block, bSRGB, bCoCg);
 }
 
 void SelectAlphaMod(in float SourceAlpha, in float EncodedAlpha, int IndexInTable, inout int SelectedIndex, inout float MinDif)
@@ -272,8 +282,8 @@ uint2 CompressBlock_ETC2_Alpha(in float BlockA[16])
 	int BaseValueInt = BaseValue;
 	
 	uint2 AlphaBits;
-	AlphaBits.x = SwapEndian32(BlockWeights.y | (TableIdx << 16) | (MultiplierInt << 20) | (BaseValueInt << 24));
-	AlphaBits.y = SwapEndian32(BlockWeights.x);
+	AlphaBits.x = EtcSwapEndian32(BlockWeights.y | (TableIdx << 16) | (MultiplierInt << 20) | (BaseValueInt << 24));
+	AlphaBits.y = EtcSwapEndian32(BlockWeights.x);
 
 	return AlphaBits;
 }
@@ -301,8 +311,10 @@ uint4 CompressBlock_ETC2_YCoCg(in float3 BlockRGB[16])
 		BlockRGB[i] = float3(YCoCg.yz, 0.0);
 		BlockY[i] = YCoCg.x;
 	}
-		
-	uint2 CompressedRGB = CompressBlock_ETC2_SRGB(BlockRGB);
+
+	const bool bSRGB = false;
+	const bool bCoCg = true;
+	uint2 CompressedRGB = CompressBlock_ETC1S(BlockRGB, bSRGB, bCoCg);
 	uint2 CompressedAlpha = CompressBlock_ETC2_Alpha(BlockY);
 	return uint4(CompressedAlpha, CompressedRGB);
 }
