Fixed unpacking of InstanceFlags in a mobile path of GPUScene

diff --git a/Engine/Shaders/Private/SceneData.ush b/Engine/Shaders/Private/SceneData.ush
--- a/Engine/Shaders/Private/SceneData.ush
+++ b/Engine/Shaders/Private/SceneData.ush
@@ -1151,8 +1151,9 @@ FSceneDataIntermediates GetSceneDataIntermediates(uint DrawInstanceId, float4 In
 
 	// TODO: pack important primitive and instance flags here 
 	// (25 + 7)
-	const uint PrimitiveFlags = asuint(InstanceTransform1.w) & 0x01FFFFFF;
-	const uint InstanceFlags = (PrimitiveFlags >> 25u);
+	const uint PackedPrimitiveFlags = asuint(InstanceTransform1.w);
+	const uint PrimitiveFlags = (PackedPrimitiveFlags & 0x01FFFFFF);
+	const uint InstanceFlags = (PackedPrimitiveFlags >> 25u);
 
 	// Reconstruct InstanceData from a packed data
 	FInstanceSceneData InstanceData = (FInstanceSceneData)0;
