Strata - deferred shading now handle light SpecularScale correctly.

diff --git a/Engine/Shaders/Private/Strata/StrataDeferredLighting.ush b/Engine/Shaders/Private/Strata/StrataDeferredLighting.ush
--- a/Engine/Shaders/Private/Strata/StrataDeferredLighting.ush
+++ b/Engine/Shaders/Private/Strata/StrataDeferredLighting.ush
@@ -89,7 +89,7 @@ FStrataDeferredLighting StrataDeferredLighting(
 		}
 
 		float3 DiffuseLuminance = BSDFEvaluate.IntegratedDiffuseValue;
-		float3 SpecularLuminance = BSDFEvaluate.IntegratedSpecularValue;
+		float3 SpecularLuminance = BSDFEvaluate.IntegratedSpecularValue * LightData.SpecularScale;
 
 		const float3 CommonMultiplier = LightData.Color * LightMask * LuminanceWeight(BSDFContext, BSDF);
 
