PathTracing: Add ray-type switch node

This allows creative control over indirect lighting response in materials. Note that this is not intended for optimization purposes as the current implementation evaluates all branches of the switch before picking one.

diff --git a/Engine/Plugins/Experimental/GPULightmass/Shaders/Private/LightmapPathTracing.usf b/Engine/Plugins/Experimental/GPULightmass/Shaders/Private/LightmapPathTracing.usf
--- a/Engine/Plugins/Experimental/GPULightmass/Shaders/Private/LightmapPathTracing.usf
+++ b/Engine/Plugins/Experimental/GPULightmass/Shaders/Private/LightmapPathTracing.usf
@@ -185,7 +185,7 @@ FPathTracingPayload TraceTransparentRay(FRayDesc Ray, bool IsCameraRay, bool Las
 	}
 	for (;;)
 	{
-		FPackedPathTracingPayload PackedPayload = InitPathTracingPayload(IsCameraRay, 1.0);
+		FPackedPathTracingPayload PackedPayload = InitPathTracingPayload(IsCameraRay ? PATHTRACER_SCATTER_CAMERA : PATHTRACER_SCATTER_DIFFUSE, 1.0);
 		// Trace the ray
 		TraceRay(
 			TLAS,
diff --git a/Engine/Shaders/Private/Common.ush b/Engine/Shaders/Private/Common.ush
--- a/Engine/Shaders/Private/Common.ush
+++ b/Engine/Shaders/Private/Common.ush
@@ -1779,14 +1779,15 @@ bool GetRayTracingQualitySwitch()
 #endif
 }
 
-bool GetPathTracingQualitySwitch()
-{
-#if (RAYHITGROUPSHADER || RAYMISSSHADER || RAYCALLABLESHADER) && PATH_TRACING && !SIMPLIFIED_MATERIAL_SHADER
-	return true;
-#else
-	return false;
+#if !PATH_TRACING
+// Default implementations for non-path traced shaders
+// The actual implementations used for the path tracer exist in PathTracingShaderUtils.h to minimize cruft in Common.ush
+bool GetPathTracingQualitySwitch() { return false;  }
+bool GetPathTracingIsShadow() { return false; }
+bool GetPathTracingIsIndirectDiffuse() { return false; }
+bool GetPathTracingIsIndirectSpecular() { return false; }
+bool GetPathTracingIsIndirectVolume() { return false; }
 #endif
-}
 
 bool GetLightmassReplaceState()
 {
diff --git a/Engine/Shaders/Private/PathTracing/Material/PathTracingMaterialCommon.ush b/Engine/Shaders/Private/PathTracing/Material/PathTracingMaterialCommon.ush
--- a/Engine/Shaders/Private/PathTracing/Material/PathTracingMaterialCommon.ush
+++ b/Engine/Shaders/Private/PathTracing/Material/PathTracingMaterialCommon.ush
@@ -6,7 +6,7 @@ PathTracingMaterialCommon.usf: Brdf utility functions
 
 #pragma once
 
-#include "../PathTracingCommon.ush"
+#include "/Engine/Private/MonteCarlo.ush"
 #include "/Engine/Private/BRDF.ush"
 
 float LobeColorToWeight(float3 C)
diff --git a/Engine/Shaders/Private/PathTracing/PathTracingCommon.ush b/Engine/Shaders/Private/PathTracing/PathTracingCommon.ush
--- a/Engine/Shaders/Private/PathTracing/PathTracingCommon.ush
+++ b/Engine/Shaders/Private/PathTracing/PathTracingCommon.ush
@@ -6,6 +6,7 @@
 
 #pragma once
 
+#include "Material/PathTracingMaterialCommon.ush"
 #include "../RayTracing/RayTracingCommon.ush"
 #include "../OctahedralCommon.ush"
 #include "/Engine/Shared/PathTracingDefinitions.h"
@@ -14,8 +15,12 @@
 #define PATH_TRACING_PAYLOAD_OUTPUT_FLAG_FRONT_FACE                   (1 << 0) // Indicates that ray has hit the front face of a primitive. Set by closest hit shader.
 #define PATH_TRACING_PAYLOAD_OUTPUT_FLAG_DECAL_RECEIVER               (1 << 1) // Indicates that ray has hit a primitive that receives decals. Set by closest hit shader.
 
-#define PATH_TRACING_PAYLOAD_INPUT_FLAG_CAMERA_RAY                    (1 << 0) // Indicates that this is a camera ray which can trigger a number of special behaviors in the closest hit shader
-#define PATH_TRACING_PAYLOAD_INPUT_FLAG_PATHTRACER_VISIBILITY         (1 << 1) // Indicates that this is a path tracer visibility ray (which supports material evaluation for transparent shadows and fake caustics)
+#define PATH_TRACING_PAYLOAD_INPUT_FLAG_RAY_TYPE_MASK                 (7 << 0) // Reserve 3 bits for ray-mask
+#define PATH_TRACING_PAYLOAD_INPUT_FLAG_RAY_TYPE_CAMERA                0       // Indicates that this is a camera ray which can trigger a number of special behaviors in the closest hit shader
+#define PATH_TRACING_PAYLOAD_INPUT_FLAG_RAY_TYPE_SHADOW                1       // Indicates that this is a path tracer visibility ray (which supports material evaluation for transparent shadows and fake caustics)
+#define PATH_TRACING_PAYLOAD_INPUT_FLAG_RAY_TYPE_INDIRECT_DIFFUSE      2
+#define PATH_TRACING_PAYLOAD_INPUT_FLAG_RAY_TYPE_INDIRECT_SPECULAR     3
+#define PATH_TRACING_PAYLOAD_INPUT_FLAG_RAY_TYPE_INDIRECT_VOLUME       4
 
 #define PATH_TRACING_PAYLOAD_NORMAL_BITS  32
 
@@ -73,11 +78,11 @@ struct FPackedPathTracingPayload : FMinimalPayload
 	uint GetPrimitiveLightingChannelMask() { return (PackedData[0] >> 29) & 0x7; }
 
 	// Flag methods
-	void SetCameraRay() { SetFlag(PATH_TRACING_PAYLOAD_INPUT_FLAG_CAMERA_RAY); }
-	bool IsCameraRay() { return (GetFlags() & PATH_TRACING_PAYLOAD_INPUT_FLAG_CAMERA_RAY) != 0; }
+	void SetCameraRay() { SetFlag(PATH_TRACING_PAYLOAD_INPUT_FLAG_RAY_TYPE_CAMERA); }
+	bool IsCameraRay() { return (GetFlags() & PATH_TRACING_PAYLOAD_INPUT_FLAG_RAY_TYPE_MASK) == PATH_TRACING_PAYLOAD_INPUT_FLAG_RAY_TYPE_CAMERA; }
 
-	void SetVisibilityRay() { SetFlag(PATH_TRACING_PAYLOAD_INPUT_FLAG_PATHTRACER_VISIBILITY); }
-	bool IsVisibilityRay() { return (GetFlags() & PATH_TRACING_PAYLOAD_INPUT_FLAG_PATHTRACER_VISIBILITY) != 0; }
+	void SetVisibilityRay() { SetFlag(PATH_TRACING_PAYLOAD_INPUT_FLAG_RAY_TYPE_SHADOW); }
+	bool IsVisibilityRay() { return (GetFlags() & PATH_TRACING_PAYLOAD_INPUT_FLAG_RAY_TYPE_MASK) == PATH_TRACING_PAYLOAD_INPUT_FLAG_RAY_TYPE_SHADOW; }
 
 	bool IsFrontFace() { return (GetFlags() & PATH_TRACING_PAYLOAD_OUTPUT_FLAG_FRONT_FACE) != 0; }
 
@@ -266,13 +271,18 @@ struct FPathTracingPayload : FMinimalPayload
 	}
 };
 
-FPackedPathTracingPayload InitPathTracingPayload(bool bIsCameraRay, float PathRoughness)
+FPackedPathTracingPayload InitPathTracingPayload(uint ScatterType, float PathRoughness)
 {
 	FPackedPathTracingPayload Output = (FPackedPathTracingPayload)0;
 	Output.SetPathRoughness(PathRoughness);
-	if (bIsCameraRay)
+	switch (ScatterType)
 	{
-		Output.SetCameraRay();
+		case PATHTRACER_SCATTER_DIFFUSE:   Output.SetFlag(PATH_TRACING_PAYLOAD_INPUT_FLAG_RAY_TYPE_INDIRECT_DIFFUSE);  break;
+		case PATHTRACER_SCATTER_SPECULAR:  
+		case PATHTRACER_SCATTER_REFRACT:   Output.SetFlag(PATH_TRACING_PAYLOAD_INPUT_FLAG_RAY_TYPE_INDIRECT_SPECULAR); break;
+		case PATHTRACER_SCATTER_VOLUME:    Output.SetFlag(PATH_TRACING_PAYLOAD_INPUT_FLAG_RAY_TYPE_INDIRECT_VOLUME);   break;
+		case PATHTRACER_SCATTER_CAMERA: 
+		default:                           Output.SetFlag(PATH_TRACING_PAYLOAD_INPUT_FLAG_RAY_TYPE_CAMERA); break;
 	}
 	return Output;
 }
diff --git a/Engine/Shaders/Private/PathTracing/PathTracingCore.ush b/Engine/Shaders/Private/PathTracing/PathTracingCore.ush
--- a/Engine/Shaders/Private/PathTracing/PathTracingCore.ush
+++ b/Engine/Shaders/Private/PathTracing/PathTracingCore.ush
@@ -300,7 +300,7 @@ FPathTracingPayload TraceTransparentRay(inout FPathState PathState, bool bIsCame
 	}
 	for (;;)
 	{
-		FPackedPathTracingPayload PackedPayload = InitPathTracingPayload(bIsCameraRay, ApproximateCaustics ? PathState.PathRoughness : 0.0);
+		FPackedPathTracingPayload PackedPayload = InitPathTracingPayload(PathState.FirstScatterType, ApproximateCaustics ? PathState.PathRoughness : 0.0);
 		// Trace the ray
 		TraceRay(
 			TLAS,
@@ -692,7 +692,9 @@ struct FProbeResult
 FProbeResult TraceProbeRay(FRayDesc Ray)
 {
 	// Trace a short ray to see if we escaped the surface
-	FPackedPathTracingPayload PackedPayload = InitPathTracingPayload(true, 1.0);
+	// NOTE: we leave these rays tagged as "camera" because we are interested in getting the same shading normal the camera saw (on the off chance the material does some ray-simplfiication)
+	//      Since the SSS code mostly runs on camera rays, this is a reasonable approximation
+	FPackedPathTracingPayload PackedPayload = InitPathTracingPayload(PATHTRACER_SCATTER_CAMERA, 1.0);
 	// TODO: use smaller payload + flag to skip shading work? (only want smooth+geo normals)
 	//       however supporting blockers would require some shading ...
 	const uint RayFlags = 0;
diff --git a/Engine/Shaders/Private/PathTracing/PathTracingInitExtinctionCoefficient.usf b/Engine/Shaders/Private/PathTracing/PathTracingInitExtinctionCoefficient.usf
--- a/Engine/Shaders/Private/PathTracing/PathTracingInitExtinctionCoefficient.usf
+++ b/Engine/Shaders/Private/PathTracing/PathTracingInitExtinctionCoefficient.usf
@@ -24,7 +24,7 @@ RAY_TRACING_ENTRY_RAYGEN(PathTracingInitExtinctionCoefficientRG)
 	float3 StartingSigmaT = 0.0;
 	for (;;)
 	{
-		FPackedPathTracingPayload PackedPayload = InitPathTracingPayload(true, 0.0);
+		FPackedPathTracingPayload PackedPayload = InitPathTracingPayload(PATHTRACER_SCATTER_CAMERA, 0.0);
 		TraceRay(
 			TLAS,
 			RayFlags,
diff --git a/Engine/Shaders/Private/PathTracing/PathTracingLightingMissShader.usf b/Engine/Shaders/Private/PathTracing/PathTracingLightingMissShader.usf
--- a/Engine/Shaders/Private/PathTracing/PathTracingLightingMissShader.usf
+++ b/Engine/Shaders/Private/PathTracing/PathTracingLightingMissShader.usf
@@ -1,6 +1,9 @@
 // Copyright Epic Games, Inc. All Rights Reserved.
 
+#define PATH_TRACING 1
+
 #include "PathTracingCommon.ush"
+#include "PathTracingShaderUtils.ush"
 
 #include "/Engine/Generated/Material.ush"
 
@@ -14,6 +17,7 @@
 RAY_TRACING_ENTRY_MISS(PathTracingLightingMS, FPackedPathTracingPayload, Payload)
 {
 	Payload.SetMiss();
+	CurrentPayloadInputFlags = Payload.GetFlags();
 	float3 TranslatedWorldPosition = WorldRayOrigin(); // Shadow ray is traced from the shaded point
 	float3 Throughput = Payload.GetRayThroughput();
 	float3 LightFunctionMult = GetRayTracingLightFunction(TranslatedWorldPosition);
diff --git a/Engine/Shaders/Private/PathTracing/PathTracingMaterialHitShader.usf b/Engine/Shaders/Private/PathTracing/PathTracingMaterialHitShader.usf
--- a/Engine/Shaders/Private/PathTracing/PathTracingMaterialHitShader.usf
+++ b/Engine/Shaders/Private/PathTracing/PathTracingMaterialHitShader.usf
@@ -15,7 +15,7 @@
 #include "/Engine/Private/Common.ush"
 #include "/Engine/Private/RayTracing/RayTracingCommon.ush"
 #include "/Engine/Private/RayTracing/RayTracingHitGroupCommon.ush"
-
+#include "/Engine/Private/PathTracing/PathTracingShaderUtils.ush"
 #include "/Engine/Generated/Material.ush"   
 #include "/Engine/Generated/VertexFactory.ush"
 
@@ -84,6 +84,8 @@ RAY_TRACING_ENTRY_CLOSEST_HIT(PathTracingMaterialCHS,
 	const float3 TranslatedWorldPosition = TranslatedWorldRayOrigin() + RayTCurrent() * WorldRayDirection();
 	const float4 SvPosition = mul(float4(TranslatedWorldPosition, 1.0f), ResolvedView.TranslatedWorldToClip);
 
+	CurrentPayloadInputFlags = PackedPayload.GetFlags();
+
 #if VF_SUPPORTS_RAYTRACING_PREPARE_MATERIAL_PIXEL_PARAMETERS
 	// this is a newer codepath that is both more flexible and allows for more direct calculation compared to the other codepath
 	// TODO: implement such a method for all vertex factories
@@ -625,6 +627,8 @@ RAY_TRACING_ENTRY_ANY_HIT(PathTracingMaterialAHS,
 	const float3 TranslatedWorldPosition = TranslatedWorldRayOrigin() + RayTCurrent() * WorldRayDirection();
 	const float4 SvPosition = mul(float4(TranslatedWorldPosition, 1.0f), ResolvedView.TranslatedWorldToClip);
 
+	CurrentPayloadInputFlags = PackedPayload.GetFlags();
+
 #if VF_SUPPORTS_RAYTRACING_PREPARE_MATERIAL_PIXEL_PARAMETERS
 	// this is a newer codepath that is both more flexible and allows for more direct calculation compared to the other codepath
 	// TODO: implement such a method for all vertex factories
diff --git a/Engine/Shaders/Private/PathTracing/PathTracingShaderUtils.ush b/Engine/Shaders/Private/PathTracing/PathTracingShaderUtils.ush
new file mode 100644
--- /dev/null
+++ b/Engine/Shaders/Private/PathTracing/PathTracingShaderUtils.ush
@@ -0,0 +1,43 @@
+// Copyright Epic Games, Inc. All Rights Reserved.
+
+#pragma once
+
+// This file contains utility methods that can be used by path tracer specific material shaders (such as the CHS/AHS pairs, light miss shaders, etc ...)
+
+#ifndef PATH_TRACING
+#error This header is only for use within path tracing shaders
+#endif
+
+#include "PathTracingCommon.ush"
+
+bool GetPathTracingQualitySwitch()
+{
+#if SIMPLIFIED_MATERIAL_SHADER
+	// GPU Lightmass does not take the pathtracing branch as it has a different node dedicated to it
+	return false;
+#else
+	return true;
+#endif
+}
+
+// Use a global variable for convenience. This must be set prior to running the generated material code.
+static int CurrentPayloadInputFlags = 0;
+bool GetPathTracingIsShadow()
+{
+	return (CurrentPayloadInputFlags & PATH_TRACING_PAYLOAD_INPUT_FLAG_RAY_TYPE_MASK) == PATH_TRACING_PAYLOAD_INPUT_FLAG_RAY_TYPE_SHADOW;
+}
+
+bool GetPathTracingIsIndirectDiffuse()
+{
+	return (CurrentPayloadInputFlags & PATH_TRACING_PAYLOAD_INPUT_FLAG_RAY_TYPE_MASK) == PATH_TRACING_PAYLOAD_INPUT_FLAG_RAY_TYPE_INDIRECT_DIFFUSE;
+}
+
+bool GetPathTracingIsIndirectSpecular()
+{
+	return (CurrentPayloadInputFlags & PATH_TRACING_PAYLOAD_INPUT_FLAG_RAY_TYPE_MASK) == PATH_TRACING_PAYLOAD_INPUT_FLAG_RAY_TYPE_INDIRECT_SPECULAR;
+}
+
+bool GetPathTracingIsIndirectVolume()
+{
+	return (CurrentPayloadInputFlags & PATH_TRACING_PAYLOAD_INPUT_FLAG_RAY_TYPE_MASK) == PATH_TRACING_PAYLOAD_INPUT_FLAG_RAY_TYPE_INDIRECT_VOLUME;
+}
diff --git a/Engine/Source/Runtime/Engine/Classes/Materials/MaterialExpressionPathTracingRayTypeSwitch.h b/Engine/Source/Runtime/Engine/Classes/Materials/MaterialExpressionPathTracingRayTypeSwitch.h
new file mode 100644
--- /dev/null
+++ b/Engine/Source/Runtime/Engine/Classes/Materials/MaterialExpressionPathTracingRayTypeSwitch.h
@@ -0,0 +1,46 @@
+// Copyright Epic Games, Inc. All Rights Reserved.
+
+#pragma once
+
+#include "CoreMinimal.h"
+#include "UObject/ObjectMacros.h"
+#include "MaterialExpressionIO.h"
+#include "MaterialValueType.h"
+#include "Materials/MaterialExpression.h"
+#include "MaterialExpressionPathTracingRayTypeSwitch.generated.h"
+
+UCLASS()
+class UMaterialExpressionPathTracingRayTypeSwitch : public UMaterialExpression
+{
+	GENERATED_UCLASS_BODY()
+
+	/** Used for camera rays (or for non-path traced shading) */
+	UPROPERTY()
+	FExpressionInput Main;
+
+	/** Used by the path tracer on shadow rays (this only applies for non-opaque blend modes) */
+	UPROPERTY()
+	FExpressionInput Shadow;
+
+	/** Used by the path tracer on indirect diffuse rays */
+	UPROPERTY()
+	FExpressionInput IndirectDiffuse;
+
+	/** Used by the path tracer on indirect specular rays */
+	UPROPERTY()
+	FExpressionInput IndirectSpecular;
+
+	/** Used by the path tracer on indirect volume rays */
+	UPROPERTY()
+	FExpressionInput IndirectVolume;
+
+	//~ Begin UMaterialExpression Interface
+#if WITH_EDITOR
+	virtual int32 Compile(class FMaterialCompiler* Compiler, int32 OutputIndex) override;
+	virtual void GetCaption(TArray<FString>& OutCaptions) const override;
+	virtual bool IsResultMaterialAttributes(int32 OutputIndex) override;
+	virtual uint32 GetInputType(int32 InputIndex) override { return MCT_Unknown; }
+	virtual uint32 GetOutputType(int32 OutputIndex) override { return MCT_Unknown; }
+#endif
+	//~ End UMaterialExpression Interface
+};
diff --git a/Engine/Source/Runtime/Engine/Private/Materials/HLSLMaterialTranslator.cpp b/Engine/Source/Runtime/Engine/Private/Materials/HLSLMaterialTranslator.cpp
--- a/Engine/Source/Runtime/Engine/Private/Materials/HLSLMaterialTranslator.cpp
+++ b/Engine/Source/Runtime/Engine/Private/Materials/HLSLMaterialTranslator.cpp
@@ -9268,6 +9268,17 @@ int32 FHLSLMaterialTranslator::PathTracingQualitySwitchReplace(int32 Normal, int
 	return GenericSwitch(TEXT("GetPathTracingQualitySwitch()"), PathTraced, Normal);
 }
 
+int32 FHLSLMaterialTranslator::PathTracingRayTypeSwitch(int32 Main, int32 Shadow, int32 IndirectDiffuse, int32 IndirectSpecular, int32 IndirectVolume)
+{
+	// Generate a sequential series of switches so we can easily account for type promotion across the ports
+	// Any input port that does not have anything connected to it (INDEX_NONE) defaults back to Main
+	int32 TmpA = GenericSwitch(TEXT("GetPathTracingIsShadow()")          , Shadow           == INDEX_NONE ? Main : Shadow          , Main);
+	int32 TmpB = GenericSwitch(TEXT("GetPathTracingIsIndirectDiffuse()") , IndirectDiffuse  == INDEX_NONE ? Main : IndirectDiffuse , TmpA);
+	int32 TmpC = GenericSwitch(TEXT("GetPathTracingIsIndirectSpecular()"), IndirectSpecular == INDEX_NONE ? Main : IndirectSpecular, TmpB);
+	int32 TmpD = GenericSwitch(TEXT("GetPathTracingIsIndirectVolume()")  , IndirectVolume   == INDEX_NONE ? Main : IndirectVolume  , TmpC);
+	return TmpD;
+}
+
 int32 FHLSLMaterialTranslator::LightmassReplace(int32 Realtime, int32 Lightmass)
 {
 	return GenericSwitch(TEXT("GetLightmassReplaceState()"), Lightmass, Realtime);
diff --git a/Engine/Source/Runtime/Engine/Private/Materials/HLSLMaterialTranslator.h b/Engine/Source/Runtime/Engine/Private/Materials/HLSLMaterialTranslator.h
--- a/Engine/Source/Runtime/Engine/Private/Materials/HLSLMaterialTranslator.h
+++ b/Engine/Source/Runtime/Engine/Private/Materials/HLSLMaterialTranslator.h
@@ -986,6 +986,7 @@ protected:
 
 	virtual int32 RayTracingQualitySwitchReplace(int32 Normal, int32 RayTraced) override;
 	virtual int32 PathTracingQualitySwitchReplace(int32 Normal, int32 PathTraced) override;
+	virtual int32 PathTracingRayTypeSwitch(int32 Main, int32 Shadow, int32 IndirectDiffuse, int32 IndirectSpecular, int32 IndirectVolume) override;
 	virtual int32 LightmassReplace(int32 Realtime, int32 Lightmass) override;
 
 	virtual int32 VirtualTextureOutputReplace(int32 Default, int32 VirtualTexture) override;
diff --git a/Engine/Source/Runtime/Engine/Private/Materials/MaterialExpressions.cpp b/Engine/Source/Runtime/Engine/Private/Materials/MaterialExpressions.cpp
--- a/Engine/Source/Runtime/Engine/Private/Materials/MaterialExpressions.cpp
+++ b/Engine/Source/Runtime/Engine/Private/Materials/MaterialExpressions.cpp
@@ -110,6 +110,7 @@
 #include "Materials/MaterialExpressionGIReplace.h"
 #include "Materials/MaterialExpressionRayTracingQualitySwitch.h"
 #include "Materials/MaterialExpressionPathTracingQualitySwitch.h"
+#include "Materials/MaterialExpressionPathTracingRayTypeSwitch.h"
 #include "Materials/MaterialExpressionGetMaterialAttributes.h"
 #include "Materials/MaterialExpressionHairAttributes.h"
 #include "Materials/MaterialExpressionHairColor.h"
@@ -17922,6 +17923,67 @@ FStrataOperator* UMaterialExpressionPathTracingQualitySwitch::StrataGenerateMate
 #endif // WITH_EDITOR
 
 
+//
+// UMaterialExpressionPathTracingRayTypeSwitch
+//
+UMaterialExpressionPathTracingRayTypeSwitch::UMaterialExpressionPathTracingRayTypeSwitch(const FObjectInitializer& ObjectInitializer)
+	: Super(ObjectInitializer)
+{
+	// Structure to hold one-time initialization
+	struct FConstructorStatics
+	{
+		FText NAME_Utility;
+		FConstructorStatics()
+			: NAME_Utility(LOCTEXT("Utility", "Utility"))
+		{
+		}
+	};
+	static FConstructorStatics ConstructorStatics;
+
+#if WITH_EDITORONLY_DATA
+	MenuCategories.Add(ConstructorStatics.NAME_Utility);
+#endif
+}
+
+#if WITH_EDITOR
+int32 UMaterialExpressionPathTracingRayTypeSwitch::Compile(class FMaterialCompiler* Compiler, int32 OutputIndex)
+{
+	if (!Main.GetTracedInput().Expression)
+	{
+		return Compiler->Errorf(TEXT("Missing PathTracingRayTypeSwitch input 'Main'"));
+	}
+	else
+	{
+		// compile all arguments (its ok if some of these are not connected, the will default to using Main)
+		int32 ArgMain = Main.Compile(Compiler);
+		int32 ArgShadow = Shadow.Compile(Compiler);
+		int32 ArgDiffuse = IndirectDiffuse.Compile(Compiler);
+		int32 ArgSpecular = IndirectSpecular.Compile(Compiler);
+		int32 ArgVolume = IndirectVolume.Compile(Compiler);
+
+		return Compiler->PathTracingRayTypeSwitch(ArgMain, ArgShadow, ArgDiffuse, ArgSpecular, ArgVolume);
+	}
+}
+
+bool UMaterialExpressionPathTracingRayTypeSwitch::IsResultMaterialAttributes(int32 OutputIndex)
+{
+	// Only check the Main expression since it must be connected. If the other plugs have something else connected that is not compatible,
+	// we would get an error during translation therefore it is sufficient to check the main plug only.
+	if (Main.Expression)
+	{
+		return Main.Expression->IsResultMaterialAttributes(OutputIndex);
+	}
+	return false;
+}
+
+void UMaterialExpressionPathTracingRayTypeSwitch::GetCaption(TArray<FString>& OutCaptions) const
+{
+	OutCaptions.Add(TEXT("PathTracingRayTypeSwitch"));
+}
+
+#endif // WITH_EDITOR
+
+
 UMaterialExpressionObjectOrientation::UMaterialExpressionObjectOrientation(const FObjectInitializer& ObjectInitializer)
 	: Super(ObjectInitializer)
 {
diff --git a/Engine/Source/Runtime/Engine/Public/MaterialCompiler.h b/Engine/Source/Runtime/Engine/Public/MaterialCompiler.h
--- a/Engine/Source/Runtime/Engine/Public/MaterialCompiler.h
+++ b/Engine/Source/Runtime/Engine/Public/MaterialCompiler.h
@@ -550,6 +550,7 @@ public:
 	virtual int32 NaniteReplace(int32 Default, int32 Nanite) = 0;
 	virtual int32 RayTracingQualitySwitchReplace(int32 Normal, int32 RayTraced) = 0;
 	virtual int32 PathTracingQualitySwitchReplace(int32 Normal, int32 PathTraced) = 0;
+	virtual int32 PathTracingRayTypeSwitch(int32 Main, int32 Shadow, int32 IndirectDiffuse, int32 IndirectSpecular, int32 IndirectVolume) = 0;
 	virtual int32 LightmassReplace(int32 Realtime, int32 Lightmass) = 0;
 	virtual int32 VirtualTextureOutputReplace(int32 Default, int32 VirtualTexture) = 0;
 	virtual int32 ReflectionCapturePassSwitch(int32 Default, int32 Reflection) = 0;
@@ -993,6 +994,7 @@ public:
 	virtual int32 NaniteReplace(int32 Default, int32 Nanite) override { return Compiler->NaniteReplace(Default, Nanite); }
 	virtual int32 RayTracingQualitySwitchReplace(int32 Normal, int32 RayTraced) override { return Compiler->RayTracingQualitySwitchReplace(Normal, RayTraced); }
 	virtual int32 PathTracingQualitySwitchReplace(int32 Normal, int32 PathTraced) override { return Compiler->PathTracingQualitySwitchReplace(Normal, PathTraced); }
+	virtual int32 PathTracingRayTypeSwitch(int32 Main, int32 Shadow, int32 IndirectDiffuse, int32 IndirectSpecular, int32 IndirectVolume) override { return Compiler->PathTracingRayTypeSwitch(Main, Shadow, IndirectDiffuse, IndirectSpecular, IndirectVolume); }
 	virtual int32 LightmassReplace(int32 Realtime, int32 Lightmass) override { return Compiler->LightmassReplace(Realtime, Lightmass); }
 	virtual int32 VirtualTextureOutputReplace(int32 Default, int32 VirtualTexture) override { return Compiler->VirtualTextureOutputReplace(Default, VirtualTexture); }
 	virtual int32 ReflectionCapturePassSwitch(int32 Default, int32 Reflection) override { return Compiler->ReflectionCapturePassSwitch(Default, Reflection); }
