Fix subsurface profile raytraced transmission where extinction scale does not have an effect. The max transmission needs to be divided by extinction scale.We divide by sqrt(extinction scale) to match it closer to rasterizer parameterization.

diff --git a/Engine/Shaders/Private/RayTracing/RayTracingOcclusionRGS.usf b/Engine/Shaders/Private/RayTracing/RayTracingOcclusionRGS.usf
--- a/Engine/Shaders/Private/RayTracing/RayTracingOcclusionRGS.usf
+++ b/Engine/Shaders/Private/RayTracing/RayTracingOcclusionRGS.usf
@@ -47,6 +47,7 @@ float LODTransitionEnd;
 float AvoidSelfIntersectionTraceDistance;
 uint bTransmissionSamplingDistanceCulling;
 uint TransmissionSamplingTechnique;
+uint TransmissionMeanFreePathType;
 uint RejectionSamplingTrials;
 uint bAcceptFirstHit;
 uint bTwoSidedGeometry;
@@ -54,6 +55,9 @@ uint bTwoSidedGeometry;
 #define TransmissionSampling_ConstantTrackingInfiniteDomain 0
 #define TransmissionSampling_ConstantTrackingFiniteDomain 1
 
+#define TransmissionMeanFreePathType_ExtinctionScale	0
+#define TransmissionMeanFreePathType_MaxMeanFreePath	1
+
 #if USE_HAIR_LIGHTING
 #include "../HairStrands/HairStrandsRaytracing.ush"
 uint				bUseHairVoxel;
@@ -155,12 +159,12 @@ bool GenerateOcclusionRay(
 
 float ComputeDiffuseTransmission(
 	float3 V, float3 P, float3 N,
-	float MeanFreePath, float Eta,
+	float MeanFreePath, float MaxTransmissionDistance, float Eta,
 	FLightShaderParameters LightParameters, float2 LightSample,
 	inout RandomSequence RandSequence, 
 	uint RejectionRetryMax, uint2 PixelCoord)
 {
-	float TransmissionDistance = SSSS_MAX_TRANSMISSION_PROFILE_DISTANCE;
+	float TransmissionDistance = MaxTransmissionDistance;
 	float RcpMeanFreePath = rcp(MeanFreePath);
 
 	// Refract inward and force a scattering event
@@ -189,10 +193,15 @@ float ComputeDiffuseTransmission(
 			InstanceInclusionMask,
 			RejectionTestRay);
 		TMax = (RejectionTestPayload.IsHit() ? RejectionTestPayload.HitT : RejectionTestRay.TMax);
-		if (RejectionTestPayload.IsMiss())
+		
+		// Assume the back facing is parallel to the front face.
+		// If the first ray marching into the volume does not hit anything, we should still
+		// resume for the scattering such that it has a chance to fly out of the volume
+		// based on the incident angle and the angle to light 
+		/*if (RejectionTestPayload.IsMiss())
 		{
 			return TransmissionDistance;
-		}
+		}*/
 	}
 
 	float RandSample = RandomSequence_GenerateSample1D(RandSequence);
@@ -204,7 +213,7 @@ float ComputeDiffuseTransmission(
 		// Use constant tracking through homogeneous media, but renormalize to force a scatter event within the interval: [0, TMax]
 		float NormalizationConstant = 1.0 - exp(-RcpMeanFreePath * TMax);
 		ScatterDistance = -log(1.0 - RandSample * NormalizationConstant) * MeanFreePath;
-		ScatterDistancePdf = RcpMeanFreePath * exp(-RcpMeanFreePath * ScatterDistance) * rcp(NormalizationConstant);
+		ScatterDistancePdf = (TMax * RcpMeanFreePath) * exp(-RcpMeanFreePath * ScatterDistance) * rcp(NormalizationConstant);
 	}
 	else // TransmissionSamplingTechnique = TransmissionSampling_ConstantTrackingInfiniteDomain
 	{
@@ -223,8 +232,14 @@ float ComputeDiffuseTransmission(
 		}
 		if (ScatterDistance > TMax)	return 0.0;
 	}
+
+	// After sampling by pdf, the first interaction could go out of the range
+	// resulting in no transmission. Need to reconsider adding it in by commenting this out?
 	ScatterDistance /= ScatterDistancePdf;
 
+	// Avoid self intersection.
+	ScatterDistance = max(0.01, ScatterDistance);
+
 	// Build sub-surface scattering ray
 	FRayDesc SSSRay;
 	SSSRay.Origin = P + Tr * ScatterDistance;
@@ -232,7 +247,7 @@ float ComputeDiffuseTransmission(
 
 	// Clip the ray length by the maximum transmission distance
 	float LightTMax = SSSRay.TMax;
-	SSSRay.TMax = SSSS_MAX_TRANSMISSION_PROFILE_DISTANCE - ScatterDistance;
+	SSSRay.TMax = TransmissionDistance - ScatterDistance;
 
 	FMinimalPayload SSSPayload = TraceVisibilityRay(
 		TLAS,
@@ -319,6 +334,7 @@ FOcclusionResult ComputeOcclusion(
 	const FOcclusionShadingParameters ShadingDesc,
 	const FLightShaderParameters LightParameters,
 	const FTransmissionProfileParams TransmissionProfileParams,
+	const float SubsurfaceSamplingMeanFreePath,
 	const uint LocalSamplesPerPixel)
 {
 	FOcclusionResult Out = InitOcclusionResult();
@@ -552,19 +568,28 @@ FOcclusionResult ComputeOcclusion(
 		const float3 TranslatedWorldViewOrigin = LWCToFloat(LWCAdd(PrimaryView.WorldViewOrigin, PrimaryView.PreViewTranslation));
 		const float3 V = normalize(TranslatedWorldPosition - TranslatedWorldViewOrigin);
 
-		const float MeanFreePath = TransmissionProfileParams.ExtinctionScale;
+		// Apply sqrt to match the behavior of the rasterizer closer visually.
+		const float ExtinctionScale = sqrt(TransmissionProfileParams.ExtinctionScale);
+		
+		// Decrease ExtinctionScale should increase the max transmission distance. Keep rasterizer and raytracing shadow behave similar for transmission.
+		const float MaxTransmissionDistance = SSSS_MAX_TRANSMISSION_PROFILE_DISTANCE / ExtinctionScale; 
 		const float Eta = TransmissionProfileParams.OneOverIOR;
 		const uint RejectionRetryMax = LocalSamplesPerPixel - 1;
-		Out.TransmissionDistance = ComputeDiffuseTransmission(V, TranslatedWorldPosition, WorldNormal, MeanFreePath, Eta, LightParameters, LightSample, RandSequence, RejectionRetryMax, PixelCoord);
+		Out.TransmissionDistance = ComputeDiffuseTransmission(V, TranslatedWorldPosition, WorldNormal, SubsurfaceSamplingMeanFreePath,
+			MaxTransmissionDistance, Eta, LightParameters, LightSample, RandSequence, RejectionRetryMax, PixelCoord);
 
 		// Overload closest-hit distance to express occlusion distance when facing the light or transmission distance when back-facing
-		if (Out.TransmissionDistance < SSSS_MAX_TRANSMISSION_PROFILE_DISTANCE)
+		if (Out.TransmissionDistance < MaxTransmissionDistance)
 		{
 			Out.ClosestRayDistance = max(Out.ClosestRayDistance, Out.TransmissionDistance);
 		}
 
-		// Conversion to raster model: 1.0 - (distance / SSSS_MAX_TRANSMISSION_PROFILE_DISTANCE)
-		Out.TransmissionDistance = saturate(EncodeThickness(Out.TransmissionDistance / (SSSS_MAX_TRANSMISSION_PROFILE_DISTANCE)));
+		// Respond to the normal scale
+		const float NormalScale = TransmissionProfileParams.NormalScale * 0.5;
+		Out.TransmissionDistance += NormalScale;
+
+		// Conversion to raster model: 1.0 - (distance / SSSS_MAX_TRANSMISSION_PROFILE_DISTANCE).
+		Out.TransmissionDistance = saturate(EncodeThickness(Out.TransmissionDistance / MaxTransmissionDistance));
 	}
 	else if (ShadingDesc.bIsSubsurface || ShadingDesc.bIsHair)
 	{
@@ -584,6 +609,23 @@ FOcclusionResult ComputeOcclusion(
 	return Out;
 }
 
+#if !STRATA_ENABLED
+float DecodeSamplingSSSProfileRadius(FGBufferData GBufferData)
+{
+	uint SubsurfaceProfileInt = ExtractSubsurfaceProfileInt(GBufferData);
+
+	// Burley Sampling parameterization
+	float3 SurfaceAlbedo = View.SSProfilesTexture.Load(int3(BSSS_SURFACEALBEDO_OFFSET, SubsurfaceProfileInt, 0)).www;
+	float3 DiffuseMeanFreePath = DecodeDiffuseMeanFreePath(View.SSProfilesTexture.Load(int3(BSSS_DMFP_OFFSET, SubsurfaceProfileInt, 0))).www;
+	float WorldUnitScale = DecodeWorldUnitScale(View.SSProfilesTexture.Load(int3(SSSS_TINT_SCALE_OFFSET, SubsurfaceProfileInt, 0)).a) * BURLEY_CM_2_MM;
+	float Opacity = UseSubsurfaceProfile(GBufferData.ShadingModelID) ? GBufferData.CustomData.a : 0.0f;
+
+	float SSSRadius = GetMFPFromDMFPApprox(SurfaceAlbedo, SurfaceAlbedo, Opacity * WorldUnitScale * DiffuseMeanFreePath).z;
+
+	return SSSRadius * BURLEY_MM_2_CM;
+}
+#endif
+
 RAY_TRACING_ENTRY_RAYGEN(OcclusionRGS)
 {
 	const uint2 PixelCoord = DispatchRaysIndex().xy + View.ViewRectMin.xy + PixelOffset;
@@ -606,6 +648,7 @@ RAY_TRACING_ENTRY_RAYGEN(OcclusionRGS)
 		float3 WorldNormal = 0;
 		FOcclusionShadingParameters ShadingParameters = (FOcclusionShadingParameters)0;
 		FTransmissionProfileParams TransmissionParameters = (FTransmissionProfileParams)0;
+		float SubsurfaceSamplingMeanFreePath = 0;
 		#if STRATA_ENABLED
 		{
 			FStrataAddressing StrataAddressing = GetStrataPixelDataByteOffset(PixelCoord, uint2(View.BufferSizeAndInvSize.xy), Strata.MaxBytesPerPixel);
@@ -627,6 +670,8 @@ RAY_TRACING_ENTRY_RAYGEN(OcclusionRGS)
 			if (ENABLE_TRANSMISSION && ShadingParameters.bHasSubsurfaceProfile)
 			{
 				TransmissionParameters = GetTransmissionProfileParams(StrataSubSurfaceHeaderGetProfileId(SSSHeader));
+				//@TODO: Read the max MFP from STRATA?
+				SubsurfaceSamplingMeanFreePath = TransmissionParameters.ExtinctionScale;
 			}
 
 			if (ShadingParameters.bIsSubsurface)
@@ -651,6 +696,15 @@ RAY_TRACING_ENTRY_RAYGEN(OcclusionRGS)
 			if (ENABLE_TRANSMISSION)
 			{
 				TransmissionParameters = GetTransmissionProfileParams(ExtractSubsurfaceProfileInt(GBufferData));
+				
+				// If MFP is available, use the actual mean free path, otherwise, fallback to the old behavior using extinction scale in the subsurface profile.
+				SubsurfaceSamplingMeanFreePath = TransmissionParameters.ExtinctionScale;
+
+				BRANCH
+				if (TransmissionMeanFreePathType == TransmissionMeanFreePathType_MaxMeanFreePath)
+				{
+					SubsurfaceSamplingMeanFreePath = DecodeSamplingSSSProfileRadius(GBufferData);
+				}
 			}
 
 			if (ShadingParameters.bIsSubsurface)
@@ -685,6 +739,7 @@ RAY_TRACING_ENTRY_RAYGEN(OcclusionRGS)
 			ShadingParameters,
 			LightParameters, 
 			TransmissionParameters,
+			SubsurfaceSamplingMeanFreePath,
 			LocalSamplesPerPixel);
 
 		#if USE_HAIR_LIGHTING
@@ -699,6 +754,7 @@ RAY_TRACING_ENTRY_RAYGEN(OcclusionRGS)
 				ShadingParameters,
 				LightParameters,
 				TransmissionParameters,
+				SubsurfaceSamplingMeanFreePath,
 				RequestedSamplePerPixel);
 		}
 		#endif
diff --git a/Engine/Source/Runtime/Renderer/Private/RayTracing/RayTracingShadows.cpp b/Engine/Source/Runtime/Renderer/Private/RayTracing/RayTracingShadows.cpp
--- a/Engine/Source/Runtime/Renderer/Private/RayTracing/RayTracingShadows.cpp
+++ b/Engine/Source/Runtime/Renderer/Private/RayTracing/RayTracingShadows.cpp
@@ -65,6 +65,14 @@ static TAutoConsoleVariable<int32> CVarRayTracingTransmissionSamplingTechnique(
 	ECVF_RenderThreadSafe
 );
 
+static TAutoConsoleVariable CVarRayTracingTransmissionMeanFreePathType(
+	TEXT("r.RayTracing.Transmission.MeanFreePathType"),
+	0,
+	TEXT("0: Use the extinction scale from subsurface profile as MFP.")
+	TEXT("1: Use the max MFP from Subsurface profile to generate samples for transmission (Strata is not supported)."),
+	ECVF_RenderThreadSafe
+);
+
 static TAutoConsoleVariable<int32> CVarRayTracingTransmissionRejectionSamplingTrials(
 	TEXT("r.RayTracing.Transmission.RejectionSamplingTrials"),
 	0,
@@ -106,6 +114,18 @@ bool EnableRayTracingShadowTwoSidedGeometry()
 	return CVarRayTracingShadowsEnableTwoSidedGeometry.GetValueOnRenderThread() != 0;
 }
 
+uint32 GetRayTracingTransmissionMeanFreePathType()
+{	
+	uint32 MeanFreePathType = FMath::Clamp(CVarRayTracingTransmissionMeanFreePathType.GetValueOnRenderThread(), 0, 1);
+
+	if (Strata::IsStrataEnabled())
+	{
+		MeanFreePathType = 0u;
+	}
+
+	return MeanFreePathType;
+}
+
 class FOcclusionRGS : public FGlobalShader
 {
 	DECLARE_GLOBAL_SHADER(FOcclusionRGS)
@@ -164,6 +184,7 @@ class FOcclusionRGS : public FGlobalShader
 		SHADER_PARAMETER(float, AvoidSelfIntersectionTraceDistance)
 		SHADER_PARAMETER(uint32, bTransmissionSamplingDistanceCulling)
 		SHADER_PARAMETER(uint32, TransmissionSamplingTechnique)
+		SHADER_PARAMETER(uint32, TransmissionMeanFreePathType)
 		SHADER_PARAMETER(uint32, RejectionSamplingTrials)
 		SHADER_PARAMETER(uint32, bAcceptFirstHit)
 		SHADER_PARAMETER(uint32, bTwoSidedGeometry)
@@ -315,6 +336,7 @@ void FDeferredShadingSceneRenderer::RenderRayTracingShadows(
 		PassParameters->SSProfilesTexture = View.RayTracingSubSurfaceProfileTexture;
 		PassParameters->bTransmissionSamplingDistanceCulling = CVarRayTracingTransmissionSamplingDistanceCulling.GetValueOnRenderThread();
 		PassParameters->TransmissionSamplingTechnique = CVarRayTracingTransmissionSamplingTechnique.GetValueOnRenderThread();
+		PassParameters->TransmissionMeanFreePathType = GetRayTracingTransmissionMeanFreePathType();
 		PassParameters->RejectionSamplingTrials = CVarRayTracingTransmissionRejectionSamplingTrials.GetValueOnRenderThread();
 		PassParameters->Strata = Strata::BindStrataGlobalUniformParameters(View);
 		if (bUseHairLighting)
