Strata - fixed path tracer compilation for SLW hit lighting.

diff --git a/Engine/Shaders/Private/PathTracing/PathTracingMaterialHitShader.usf b/Engine/Shaders/Private/PathTracing/PathTracingMaterialHitShader.usf
--- a/Engine/Shaders/Private/PathTracing/PathTracingMaterialHitShader.usf
+++ b/Engine/Shaders/Private/PathTracing/PathTracingMaterialHitShader.usf
@@ -356,6 +356,22 @@ RAY_TRACING_ENTRY_CLOSEST_HIT(PathTracingMaterialCHS,
 #if MATERIAL_SHADINGMODEL_SINGLELAYERWATER
 	if (Payload.ShadingModelID == SHADINGMODELID_SINGLELAYERWATER)
 	{
+#if STRATA_INLINE_SINGLELAYERWATER
+		SanitizeStrataSingleLayerWater(MaterialParameters.StrataTree.BSDFs[0]);
+		FStrataBSDF SLWBSDF = MaterialParameters.StrataTree.BSDFs[0];
+		const float3 Albedo = SLW_WATERALBEDO(SLWBSDF);
+		const float3 ExtinctionCoeff = SLW_WATEREXTINCTION(SLWBSDF);
+		const float3 ScatteringCoeff = Albedo * ExtinctionCoeff;
+		const float3 AbsorptionCoeff = max(0.0, ExtinctionCoeff - ScatteringCoeff);
+		const float PhaseG = SLW_WATERPHASEG(SLWBSDF);
+
+		const float3 WaterExtinction = ScatteringCoeff + AbsorptionCoeff;
+		const float3 WaterAlbedo = ScatteringCoeff / WaterExtinction;
+
+		Payload.BaseColor = WaterAlbedo;
+		Payload.SetWaterExtinction(WaterExtinction);
+		Payload.SetWaterPhaseG(PhaseG);
+#else
 		// For single layer water we always want to treat it as front facing when evaluating parameters 
 		// as hitting back face is treated as being underwater and we get completely different parameters.
 		FMaterialPixelParameters WaterMaterialParameters = MaterialParameters;
@@ -367,6 +383,7 @@ RAY_TRACING_ENTRY_CLOSEST_HIT(PathTracingMaterialCHS,
 
 		const float3 WaterExtinction = ScatteringCoeff + AbsorptionCoeff;
 		const float3 WaterAlbedo = ScatteringCoeff / WaterExtinction;
+#endif
 
 		Payload.BaseColor = WaterAlbedo;
 		Payload.SetWaterExtinction(WaterExtinction);
