Fix for SLW to not blow up with invalid normal.

diff --git a/Engine/Shaders/Private/BasePassPixelShader.usf b/Engine/Shaders/Private/BasePassPixelShader.usf
--- a/Engine/Shaders/Private/BasePassPixelShader.usf
+++ b/Engine/Shaders/Private/BasePassPixelShader.usf
@@ -923,7 +923,20 @@ void FPixelShaderInOut_MainPS(
 	Specular = SLW_SPECULAR(SLWBSDF);
 	Roughness = SLW_ROUGHNESS(SLWBSDF);
 	Opacity = SLW_TOPMATERIALOPACITY(SLWBSDF);
+
 	MaterialParameters.WorldNormal = normalize(StrataPixelHeader.SharedLocalBases.Normals[BSDF_GETSHAREDLOCALBASISID(SLWBSDF)]);
+	const float SLWWorldNormalSquaredLen = dot(MaterialParameters.WorldNormal, MaterialParameters.WorldNormal);
+	BRANCH
+	if (!IsFinite(SLWWorldNormalSquaredLen)) // A way to detect that the normalisation process was done on 0-length normal or with NaNs.
+	{
+		// In this case, reset to the world normal to the vertex normal in order to avoid abvious bad bloom on screen.
+		MaterialParameters.WorldNormal = normalize(MaterialParameters.TangentToWorld[2]);
+		#if MATERIAL_TANGENTSPACENORMAL
+		MaterialParameters.WorldNormal *= MaterialParameters.TwoSidedSign;
+		#endif
+		StrataPixelHeader.SharedLocalBases.Normals[BSDF_GETSHAREDLOCALBASISID(SLWBSDF)] = MaterialParameters.WorldNormal;
+	}
+
 	Anisotropy = 0.0f;
 	ShadingModel = MATERIAL_SHADINGMODEL_SINGLELAYERWATER;
 #endif
