Fix shader typo in MobileDeferredShading.usf

diff --git a/Engine/Shaders/Private/MobileDeferredShading.usf b/Engine/Shaders/Private/MobileDeferredShading.usf
--- a/Engine/Shaders/Private/MobileDeferredShading.usf
+++ b/Engine/Shaders/Private/MobileDeferredShading.usf
@@ -78,7 +78,7 @@ half3 SkyLightDiffuseMobile(FGBufferData GBuffer, half3 DiffuseColor, half3 V, i
 	if (GBuffer.ShadingModelID == SHADINGMODELID_HAIR)
 	{
 		const half3 N = GBuffer.WorldNormal;
-		DiffuseColorForIndirect = EvaluateEnvHair(GBuffer, V, N, SkyDiffuseLookUpNormal);
+		DiffuseColor += EvaluateEnvHair(GBuffer, V, N, SkyDiffuseLookUpNormal);
 	}
 
 	if (GBuffer.ShadingModelID == SHADINGMODELID_CLOTH)
