Fixed invalid ddx/ddy in Nanite material CS shaders under some conditions (like the CitySample clear color vehicle shading - WorldNormal would be DCE'd out due to literal 0s flattening math through a cross product). Compute shading threads have been swizzled into 2x2 quad layout, so we can rely on ddx/ddy remapping in the compiler to QuadReadAcrossX/Y/Diagonal and producing correct results.

diff --git a/Engine/Shaders/Private/Common.ush b/Engine/Shaders/Private/Common.ush
--- a/Engine/Shaders/Private/Common.ush
+++ b/Engine/Shaders/Private/Common.ush
@@ -28,6 +28,10 @@
 	#define MaterialFloat4x3 float4x3 
 #endif
 
+#ifndef COMPUTE_SHADED
+#define COMPUTE_SHADED 0
+#endif
+
 struct FloatDeriv
 {
 	float Value;
@@ -152,6 +156,9 @@ const static float3 Max111110BitsFloat3 = float3(Max11BitsFloat, Max11BitsFloat,
 // * USE_FORCE_TEXTURE_MIP : enable/disable manual MIP level selection
 // * FORCED_TEXTURE_MIP    : force a specific MIP level
 //
+#if COMPUTE_SHADED && !defined(USE_FORCE_TEXTURE_MIP)
+	#define USE_FORCE_TEXTURE_MIP 0
+#endif
 #if !PIXELSHADER && !defined(USE_FORCE_TEXTURE_MIP)
 	#define USE_FORCE_TEXTURE_MIP 1
 #endif
diff --git a/Engine/Shaders/Private/MaterialTemplate.ush b/Engine/Shaders/Private/MaterialTemplate.ush
--- a/Engine/Shaders/Private/MaterialTemplate.ush
+++ b/Engine/Shaders/Private/MaterialTemplate.ush
@@ -130,7 +130,10 @@
 #if IS_NANITE_PASS
 	#define ALLOW_CLIP 0
 	#define clip(x)
-	#if COMPUTESHADER
+	#if COMPUTE_SHADED
+		// Threads are remapped into 2x2 quad layout, so quad read x/y/d work for ddx/ddy
+		#define fwidth(x) 0
+	#elif COMPUTESHADER
 		#define ddx(x) 0
 		#define ddy(x) 0
 		#define fwidth(x) 0
