Added manual VT filtering for mobile platforms that do not use TAA (contributed by [REDACTED_PERSON] at ARM)

diff --git a/Engine/Shaders/Private/VirtualTextureCommon.ush b/Engine/Shaders/Private/VirtualTextureCommon.ush
--- a/Engine/Shaders/Private/VirtualTextureCommon.ush
+++ b/Engine/Shaders/Private/VirtualTextureCommon.ush
@@ -198,8 +198,9 @@ struct VTPageTableResult
 	float2 UV;
 	float2 dUVdx;
 	float2 dUVdy;
-	uint4 PageTableValue[2];
+	uint4 PageTableValue[4];
 	uint PackedRequest;
+	float MipLevelFrac;
 };
 
 uint GetRequestedLevel(VTPageTableResult PageTableResult)
@@ -260,9 +261,18 @@ int TextureComputeVirtualMipLevel(
 	const float ComputedLevel = MipLevelAniso2D(OutResult.dUVdx, OutResult.dUVdy, PageTableUniform.MaxAnisoLog2);
 
 	const float GlobalMipBias = GetGlobalVirtualTextureMipBias();
+#if VIRTUAL_TEXTURE_MANUAL_TRILINEAR_FILTERING
+	const float Noise = 0.f;
+#else
 	const float Noise = GetStochasticMipNoise(SvPositionXY);
+#endif
+
+	const float MipLevel = ComputedLevel + MipBias + GlobalMipBias + Noise;
+	const float MipLevelFloor = floor(MipLevel);
+	
+	OutResult.MipLevelFrac = MipLevel - MipLevelFloor;
 
-	return (int)floor(ComputedLevel + MipBias + GlobalMipBias + Noise) + int(PageTableUniform.vPageTableMipBias);
+	return (int)MipLevelFloor + int(PageTableUniform.vPageTableMipBias);
 }
 
 /** Samples page table indirection and any apply changes to UV, vLevel and PageTable information for adaptive page tables. */
@@ -309,12 +319,30 @@ void TextureLoadVirtualPageTableInternal(
 	OutResult.UV = UV * PageTableUniform.SizeInPages;
 
 	const uint vLevelClamped = clamp(vLevel, 0, int(PageTableUniform.MaxLevel));
-	const uint vPageX = (uint(OutResult.UV.x) + PageTableUniform.XOffsetInPages) >> vLevelClamped;
-	const uint vPageY = (uint(OutResult.UV.y) + PageTableUniform.YOffsetInPages) >> vLevelClamped;
+	uint vPageX = (uint(OutResult.UV.x) + PageTableUniform.XOffsetInPages) >> vLevelClamped;
+	uint vPageY = (uint(OutResult.UV.y) + PageTableUniform.YOffsetInPages) >> vLevelClamped;
 
 	OutResult.PageTableValue[0] = PageTable0.Load(int3(vPageX, vPageY, vLevelClamped));
 	OutResult.PageTableValue[1] = uint4(0u, 0u, 0u, 0u);
 
+#if VIRTUAL_TEXTURE_MANUAL_TRILINEAR_FILTERING
+	// Second page table for trilinear.
+	const uint vLevelClamped2 = clamp(vLevel + 1, 0, int(PageTableUniform.MaxLevel));
+	const uint vPageX2 = (uint(OutResult.UV.x) + PageTableUniform.XOffsetInPages) >> vLevelClamped2;
+	const uint vPageY2 = (uint(OutResult.UV.y) + PageTableUniform.YOffsetInPages) >> vLevelClamped2;
+	
+	OutResult.PageTableValue[2] = PageTable0.Load(int3(vPageX2, vPageY2, vLevelClamped2));
+	OutResult.PageTableValue[3] = uint4(0u, 0u, 0u, 0u);
+	
+	// Alternate requests to both mip levels
+	if ((View.FrameNumber & 1u) == 0u)
+	{
+		vLevel += 1;
+		vPageX = vPageX2;
+		vPageY = vPageY2;
+	}
+#endif
+
 	// PageTableID packed in upper 4 bits of 'PackedPageTableUniform', which is the bit position we want it in for PackedRequest as well, just need to mask off extra bits
 	OutResult.PackedRequest = PageTableUniform.ShiftedPageTableID;
 	OutResult.PackedRequest |= vPageX;
@@ -336,12 +364,30 @@ void TextureLoadVirtualPageTableInternal(
 	OutResult.UV = UV * PageTableUniform.SizeInPages;
 
 	const uint vLevelClamped = clamp(vLevel, 0, int(PageTableUniform.MaxLevel));
-	const uint vPageX = (uint(OutResult.UV.x) + PageTableUniform.XOffsetInPages) >> vLevelClamped;
-	const uint vPageY = (uint(OutResult.UV.y) + PageTableUniform.YOffsetInPages) >> vLevelClamped;
+	uint vPageX = (uint(OutResult.UV.x) + PageTableUniform.XOffsetInPages) >> vLevelClamped;
+	uint vPageY = (uint(OutResult.UV.y) + PageTableUniform.YOffsetInPages) >> vLevelClamped;
 
 	OutResult.PageTableValue[0] = PageTable0.Load(int3(vPageX, vPageY, vLevelClamped));
 	OutResult.PageTableValue[1] = PageTable1.Load(int3(vPageX, vPageY, vLevelClamped));
 
+#if VIRTUAL_TEXTURE_MANUAL_TRILINEAR_FILTERING
+	// Second page table for trilinear.
+	const uint vLevelClamped2 = clamp(vLevel + 1, 0, int(PageTableUniform.MaxLevel));
+	const uint vPageX2 = (uint(OutResult.UV.x) + PageTableUniform.XOffsetInPages) >> vLevelClamped2;
+	const uint vPageY2 = (uint(OutResult.UV.y) + PageTableUniform.YOffsetInPages) >> vLevelClamped2;
+
+	OutResult.PageTableValue[2] = PageTable0.Load(int3(vPageX2, vPageY2, vLevelClamped2));
+	OutResult.PageTableValue[3] = PageTable1.Load(int3(vPageX2, vPageY2, vLevelClamped2));
+
+	// Alternate requests to both mip levels
+	if ((View.FrameNumber & 1u) == 0u)
+	{
+		vLevel += 1;
+		vPageX = vPageX2;
+		vPageY = vPageY2;
+	}
+#endif
+
 	// PageTableID packed in upper 4 bits of 'PackedPageTableUniform', which is the bit position we want it in for PackedRequest as well, just need to mask off extra bits
 	OutResult.PackedRequest = PageTableUniform.ShiftedPageTableID;
 	OutResult.PackedRequest |= vPageX;
@@ -695,20 +741,34 @@ float2 VTComputePhysicalUVs(in out VTPageTableResult PageTableResult, uint Layer
 	return pUV;
 }
 
+
+MaterialFloat4 TextureVirtualSampleAniso(Texture2D Physical, SamplerState PhysicalSampler, float2 pUV, 	float2 dUVdx, float2 dUVdy)
+{
+#if	VIRTUAL_TEXTURE_ANISOTROPIC_FILTERING
+	return Physical.SampleGrad(PhysicalSampler, pUV, dUVdx, dUVdy);
+#else
+	// no need for dUVdx/dUVdy unless we have anistropic filtering enabled
+	return Physical.SampleLevel(PhysicalSampler, pUV, 0.0f);
+#endif
+}
+
 // SamplePhysicalTexture: Aniso On
 MaterialFloat4 TextureVirtualSample(
 	Texture2D Physical, SamplerState PhysicalSampler,
 	VTPageTableResult PageTableResult, uint LayerIndex,
 	VTUniform Uniform)
 {
-	const float2 pUV = VTComputePhysicalUVs(PageTableResult, LayerIndex, Uniform);
-	const bool bValid = IsValid(PageTableResult, LayerIndex);
+	float2 pUV = VTComputePhysicalUVs(PageTableResult, LayerIndex, Uniform);
+	bool bValid = IsValid(PageTableResult, LayerIndex);
+	const MaterialFloat4 Sample1 = bValid ? TextureVirtualSampleAniso(Physical, PhysicalSampler, pUV, PageTableResult.dUVdx, PageTableResult.dUVdy) : Uniform.FallbackValue;
 
-#if VIRTUAL_TEXTURE_ANISOTROPIC_FILTERING
-	return bValid ? Physical.SampleGrad(PhysicalSampler, pUV, PageTableResult.dUVdx, PageTableResult.dUVdy) : Uniform.FallbackValue;
+#if VIRTUAL_TEXTURE_MANUAL_TRILINEAR_FILTERING
+	pUV = VTComputePhysicalUVs(PageTableResult, LayerIndex + 8u, Uniform);
+	bValid = IsValid(PageTableResult, LayerIndex + 8u);
+	const MaterialFloat4 Sample2 = bValid ? TextureVirtualSampleAniso(Physical, PhysicalSampler, pUV, PageTableResult.dUVdx, PageTableResult.dUVdy) : Sample1;
+	return lerp(Sample1, Sample2, PageTableResult.MipLevelFrac);
 #else
-	// no need for dUVdx/dUVdy unless we have anistropic filtering enabled
-	return bValid ? Physical.SampleLevel(PhysicalSampler, pUV, 0.0f) : Uniform.FallbackValue;
+	return Sample1;
 #endif
 }
 
diff --git a/Engine/Source/Runtime/Engine/Private/ShaderCompiler/ShaderCompiler.cpp b/Engine/Source/Runtime/Engine/Private/ShaderCompiler/ShaderCompiler.cpp
--- a/Engine/Source/Runtime/Engine/Private/ShaderCompiler/ShaderCompiler.cpp
+++ b/Engine/Source/Runtime/Engine/Private/ShaderCompiler/ShaderCompiler.cpp
@@ -6444,9 +6444,15 @@ void GlobalBeginCompileShader(
 	{
 		static const auto CVar = IConsoleManager::Get().FindConsoleVariable(TEXT("r.VT.AnisotropicFiltering"));
 		Input.Environment.SetDefine(TEXT("VIRTUAL_TEXTURE_ANISOTROPIC_FILTERING"), CVar ? (CVar->GetInt() != 0) : 0);
+		
+		if (bIsMobilePlatform)
+		{
+			static FShaderPlatformCachedIniValue<bool> CVarVTMobileManualTrilinearFiltering(TEXT("r.VT.Mobile.ManualTrilinearFiltering"));
+			Input.Environment.SetDefine(TEXT("VIRTUAL_TEXTURE_MANUAL_TRILINEAR_FILTERING"), (CVarVTMobileManualTrilinearFiltering.Get(Target.GetPlatform()) ? 1 : 0));
+		}
 	}
 
-	if (IsMobilePlatform((EShaderPlatform)Target.Platform))
+	if (bIsMobilePlatform)
 	{
 		const bool bMobileMovableSpotlightShadowsEnabled = IsMobileMovableSpotlightShadowsEnabled(Target.GetPlatform());
 		Input.Environment.SetDefine(TEXT("PROJECT_MOBILE_ENABLE_MOVABLE_SPOTLIGHT_SHADOWS"), bMobileMovableSpotlightShadowsEnabled ? 1 : 0);
diff --git a/Engine/Source/Runtime/Engine/Private/VT/VirtualTextureScalability.cpp b/Engine/Source/Runtime/Engine/Private/VT/VirtualTextureScalability.cpp
--- a/Engine/Source/Runtime/Engine/Private/VT/VirtualTextureScalability.cpp
+++ b/Engine/Source/Runtime/Engine/Private/VT/VirtualTextureScalability.cpp
@@ -143,6 +143,13 @@ namespace VirtualTextureScalability
 		TEXT("Is anisotropic filtering for VTs enabled?"),
 		ECVF_RenderThreadSafe | ECVF_ReadOnly);
 
+	static TAutoConsoleVariable<int32> CVarVTMobileEnableManualTrilinearFiltering(
+		TEXT("r.VT.Mobile.ManualTrilinearFiltering"),
+		1,
+		TEXT("Whether to use a manual trilinear filtering for VTs on mobile platforms.\n")
+		TEXT("This more expensive filtering is used on mobile platforms that do not support Temporal Anti-Aliasing.\n"),
+		ECVF_RenderThreadSafe | ECVF_ReadOnly);
+
 	static TAutoConsoleVariable<int32> CVarVTPageFreeThreshold(
 		TEXT("r.VT.PageFreeThreshold"),
 		60,
diff --git a/Engine/Source/Runtime/RenderCore/Private/Shader.cpp b/Engine/Source/Runtime/RenderCore/Private/Shader.cpp
--- a/Engine/Source/Runtime/RenderCore/Private/Shader.cpp
+++ b/Engine/Source/Runtime/RenderCore/Private/Shader.cpp
@@ -2042,18 +2042,21 @@ void ShaderMapAppendKeyString(EShaderPlatform Platform, FString& KeyString)
 		static const auto CVarVirtualTexture = IConsoleManager::Get().FindTConsoleVariableDataInt(TEXT("r.VirtualTextures"));
 		bool VTTextures = CVarVirtualTexture && CVarVirtualTexture->GetValueOnAnyThread() != 0;
 
+		static const auto CVarVTAnisotropic = IConsoleManager::Get().FindTConsoleVariableDataInt(TEXT("r.VT.AnisotropicFiltering"));
+		int32 VTFiltering = CVarVTAnisotropic && CVarVTAnisotropic->GetValueOnAnyThread() != 0 ? 1 : 0;
+
 		static const auto CVarMobileVirtualTexture = IConsoleManager::Get().FindTConsoleVariableDataInt(TEXT("r.Mobile.VirtualTextures"));
 		if (IsMobilePlatform(Platform) && VTTextures)
 		{
 			VTTextures = (CVarMobileVirtualTexture->GetValueOnAnyThread() != 0);
+
+			static FShaderPlatformCachedIniValue<bool> CVarVTMobileManualTrilinearFiltering(TEXT("r.VT.Mobile.ManualTrilinearFiltering"));
+			VTFiltering += (CVarVTMobileManualTrilinearFiltering.Get(Platform) ? 2 : 0);
 		}
 
 		const bool VTSupported = TargetPlatform != nullptr && TargetPlatform->SupportsFeature(ETargetPlatformFeatures::VirtualTextureStreaming);
 
-		static const auto CVarVTAnisotropic = IConsoleManager::Get().FindTConsoleVariableDataInt(TEXT("r.VT.AnisotropicFiltering"));
-		const bool VTAnisotropic = CVarVTAnisotropic && CVarVTAnisotropic->GetValueOnAnyThread() != 0;
-
-		auto tt = FString::Printf(TEXT("_VT-%d-%d-%d-%d"), VTLightmaps, VTTextures, VTSupported, VTAnisotropic);
+		auto tt = FString::Printf(TEXT("_VT-%d-%d-%d-%d"), VTLightmaps, VTTextures, VTSupported, VTFiltering);
  		KeyString += tt;
 	}
 
