Path Tracing: Minor cleanups for decals

 * Remove projection type for decals (camera ray is now always used for compatbility with raster)
 * Disable decals on non-opaque geometry (for compatibility with raster and for improved performance when tracing through many layers of transparency)
 * Disable decals on holdouts
 * Remove AO field as it isn't used
 * Merge the decal and mesh decal paylods into one

diff --git a/Engine/Shaders/Private/PathTracing/PathTracingCore.ush b/Engine/Shaders/Private/PathTracing/PathTracingCore.ush
--- a/Engine/Shaders/Private/PathTracing/PathTracingCore.ush
+++ b/Engine/Shaders/Private/PathTracing/PathTracingCore.ush
@@ -57,7 +57,6 @@ uint EnabledLightingContributions;
 float DecalRoughnessCutoff;
 float MeshDecalRoughnessCutoff;
 float MeshDecalBias;
-int MeshDecalsUseViewRay;
 float MaxPathIntensity;
 float MaxNormalBias;
 float FilterWidth;
@@ -192,38 +191,33 @@ void AccumulateRadiance(inout float3 TotalRadiance, float3 PathRadiance, bool bI
 	TotalRadiance += PathRadiance;
 }
 
-void ApplyDecal(inout FPathTracingPayload HitPayload, FDecalShaderParams Params)
+void ApplyDecal(inout FPathTracingPayload HitPayload, FDecalShaderPayload DecalPayload)
 {
-	const float DecalTransparency = Params.GetTransparency();
-	const uint DecalFlags = Params.GetFlags();
+	const float DecalTransparency = DecalPayload.GetTransparency();
+	const uint DecalFlags = DecalPayload.GetFlags();
 
 	if (DecalFlags & DECAL_WRITE_BASE_COLOR_FLAG)
 	{
-		HitPayload.BaseColor = HitPayload.BaseColor * DecalTransparency + Params.GetBaseColor();
+		HitPayload.BaseColor = HitPayload.BaseColor * DecalTransparency + DecalPayload.GetBaseColor();
 	}
 
 	if (DecalFlags & DECAL_WRITE_NORMAL_FLAG)
 	{
-		HitPayload.WorldNormal = HitPayload.WorldNormal * DecalTransparency + Params.GetWorldNormal();
+		HitPayload.WorldNormal = HitPayload.WorldNormal * DecalTransparency + DecalPayload.GetWorldNormal();
 	}
 
 	if (DecalFlags & DECAL_WRITE_ROUGHNESS_SPECULAR_METALLIC_FLAG)
 	{
-		const float3 MetallicSpecularRoughness = Params.GetMetallicSpecularRoughness();
+		const float3 MetallicSpecularRoughness = DecalPayload.GetMetallicSpecularRoughness();
 
 		HitPayload.Metallic = HitPayload.Metallic * DecalTransparency + MetallicSpecularRoughness.x;
 		HitPayload.Specular = HitPayload.Specular * DecalTransparency + MetallicSpecularRoughness.y;
 		HitPayload.Roughness = HitPayload.Roughness * DecalTransparency + MetallicSpecularRoughness.z;
 	}
 
-	if (DecalFlags & DECAL_WRITE_AMBIENT_OCCLUSION_FLAG)
-	{
-		// TODO: DecalAO?
-	}
-
 	if (DecalFlags & DECAL_WRITE_EMISSIVE_FLAG)
 	{
-		HitPayload.Radiance += Params.GetEmissive();
+		HitPayload.Radiance += DecalPayload.GetEmissive();
 	}
 
 	// TODO: Payload.Opacity?
@@ -235,8 +229,8 @@ void TraceDecalRay(FRayDesc DecalRay, bool bReverseCulling, uint InstanceInclusi
 
 	for (;;)
 	{
-		FRayTracingDecalPayload DecalPayload = (FRayTracingDecalPayload)0;
-		DecalPayload.Params.SetInputTranslatedWorldPosition(HitPayload.TranslatedWorldPos);
+		FDecalShaderPayload DecalPayload = (FDecalShaderPayload)0;
+		DecalPayload.SetInputTranslatedWorldPosition(HitPayload.TranslatedWorldPos);
 
 		TraceRay(
 			DecalTLAS,
@@ -256,12 +250,12 @@ void TraceDecalRay(FRayDesc DecalRay, bool bReverseCulling, uint InstanceInclusi
 		if (bReverseCulling)
 		{
 			// also reverse normal
-			DecalPayload.Params.SetWorldNormal(-DecalPayload.Params.GetWorldNormal());
+			DecalPayload.SetWorldNormal(-DecalPayload.GetWorldNormal());
 		}
 		
-		ApplyDecal(HitPayload, DecalPayload.Params);
+		ApplyDecal(HitPayload, DecalPayload);
 		
-		DecalThroughput *= DecalPayload.Params.GetTransparency();
+		DecalThroughput *= DecalPayload.GetTransparency();
 
 		// prepare next step around the loop
 		// retrace the exact same ray with TMin one ulp past the hit we just found
@@ -489,21 +483,11 @@ FPathTracingPayload TraceTransparentRay(inout FPathState PathState, bool bIsCame
 		if (PathState.PathRoughness < MeshDecalRoughnessCutoff && HitPayload.IsDecalReceiver())
 		{
 			FRayDesc DecalRay = (FRayDesc)0;
-			if (MeshDecalsUseViewRay)
-			{
-				// go a bit further than the current hit, to handle decals partially inside the model
-				DecalRay.Direction = -PathState.Ray.Direction;
-				DecalRay.Origin = HitPayload.TranslatedWorldPos - DecalRay.Direction * MeshDecalBias;
-				DecalRay.TMin = 0.0f;
-				DecalRay.TMax = HitPayload.HitT + MeshDecalBias;
-			}
-			else
-			{
-				DecalRay.Direction = HitPayload.WorldGeoNormal;
-				DecalRay.Origin = HitPayload.TranslatedWorldPos - DecalRay.Direction * MeshDecalBias;
-				DecalRay.TMin = 0.0f;
-				DecalRay.TMax = MeshDecalBias * 2;
-			}
+			// go a bit further than the current hit, to handle decals partially inside the model
+			DecalRay.Direction = -PathState.Ray.Direction;
+			DecalRay.Origin = HitPayload.TranslatedWorldPos - DecalRay.Direction * MeshDecalBias;
+			DecalRay.TMin = 0.0f;
+			DecalRay.TMax = HitPayload.HitT + MeshDecalBias;
 
 			const bool bReverseCulling = true;
 			TraceDecalRay(DecalRay, bReverseCulling, InstanceInclusionMask, HitPayload, DecalThroughput);
@@ -519,14 +503,14 @@ FPathTracingPayload TraceTransparentRay(inout FPathState PathState, bool bIsCame
 			{
 				uint DecalId = GetDecalId(Index, DecalLoopCount);
 
-				FDecalShaderParams Params = (FDecalShaderParams)0;
-				Params.SetInputTranslatedWorldPosition(TranslatedWorldPos);
+				FDecalShaderPayload DecalPayload = (FDecalShaderPayload)0;
+				DecalPayload.SetInputTranslatedWorldPosition(TranslatedWorldPos);
 
-				CallShader(DecalId, Params);
+				CallShader(DecalId, DecalPayload);
 
-				ApplyDecal(HitPayload, Params);
+				ApplyDecal(HitPayload, DecalPayload);
 
-				DecalThroughput *= Params.GetTransparency();
+				DecalThroughput *= DecalPayload.GetTransparency();
 			}
 		}
 #endif
diff --git a/Engine/Shaders/Private/PathTracing/PathTracingMaterialHitShader.usf b/Engine/Shaders/Private/PathTracing/PathTracingMaterialHitShader.usf
--- a/Engine/Shaders/Private/PathTracing/PathTracingMaterialHitShader.usf
+++ b/Engine/Shaders/Private/PathTracing/PathTracingMaterialHitShader.usf
@@ -116,36 +116,48 @@ RAY_TRACING_ENTRY_CLOSEST_HIT(PathTracingMaterialCHS,
 	Payload.ShadingModelID = GetMaterialShadingModel(PixelMaterialInputs);
 
 #if MATERIALBLENDING_ALPHACOMPOSITE
-	uint MaterialBlendingMode = RAY_TRACING_BLEND_MODE_ALPHA_COMPOSITE;
+	Payload.Opacity = GetMaterialOpacity(PixelMaterialInputs);
+	Payload.BlendingMode = RAY_TRACING_BLEND_MODE_ALPHA_COMPOSITE;
 #elif MATERIALBLENDING_ALPHAHOLDOUT
-	uint MaterialBlendingMode = RAY_TRACING_BLEND_MODE_ALPHA_HOLDOUT;
+	Payload.Opacity = GetMaterialOpacity(PixelMaterialInputs);
+	Payload.BlendingMode = RAY_TRACING_BLEND_MODE_ALPHA_HOLDOUT;
 	HLSL_STATIC_ASSERT(MATERIAL_SHADINGMODEL_UNLIT == 1, "Alpha holdout blend mode requires unlit shading model");
 	Payload.ShadingModelID = SHADINGMODELID_UNLIT;
 #elif MATERIALBLENDING_TRANSLUCENT
-	uint MaterialBlendingMode = RAY_TRACING_BLEND_MODE_TRANSLUCENT;
+	Payload.Opacity = GetMaterialOpacity(PixelMaterialInputs);
+	Payload.BlendingMode = RAY_TRACING_BLEND_MODE_TRANSLUCENT;
 #elif MATERIALBLENDING_ADDITIVE
-	uint MaterialBlendingMode = RAY_TRACING_BLEND_MODE_ADDITIVE;
+	Payload.Opacity = GetMaterialOpacity(PixelMaterialInputs);
+	Payload.BlendingMode = RAY_TRACING_BLEND_MODE_ADDITIVE;
 #elif MATERIALBLENDING_MODULATE
-	uint MaterialBlendingMode = RAY_TRACING_BLEND_MODE_MODULATE;
+	Payload.Opacity = GetMaterialOpacity(PixelMaterialInputs);
+	Payload.BlendingMode = RAY_TRACING_BLEND_MODE_MODULATE;
 	HLSL_STATIC_ASSERT(MATERIAL_SHADINGMODEL_UNLIT == 1, "Modulate blend mode requires unlit shading model");
 	Payload.ShadingModelID = SHADINGMODELID_UNLIT;
+#elif MATERIALBLENDING_MASKED && MATERIAL_DITHER_OPACITY_MASK
+	// dithering emulates real transparency, so switch to translucent
+	// NOTE: the raster path technically takes into account the opacity mask clip value, so the effective transparency should be:
+	//        saturate(MaskRaw - ClipValue + 0.5)
+	// (See derivation in DitheredOpacityMaskToOpacity)
+	// However this behavior is surprising to most users and does not exactly match the rasterizer anyway due to how the realtime AA
+	// code performs blending.
+	// Since the goal of dithered opacity is to emulate ordinary transparency, just use the mask input as opacity directly and
+	// ignore the configured clip value.
+	Payload.Opacity = saturate(GetMaterialMaskInputRaw(PixelMaterialInputs));
+	Payload.BlendingMode = RAY_TRACING_BLEND_MODE_TRANSLUCENT;
+#elif MATERIALBLENDING_SOLID || MATERIALBLENDING_MASKED
+	Payload.Opacity = 1.0;
+	Payload.BlendingMode = RAY_TRACING_BLEND_MODE_OPAQUE;
 #else
-	uint MaterialBlendingMode = RAY_TRACING_BLEND_MODE_OPAQUE;
+#error Unknown material blending mode
 #endif
 
-	Payload.BlendingMode = MaterialBlendingMode;
-
 	// fetch primitive flags only once
 	// TODO: would be nice to keep this inside MaterialParameters as it is also needed there as well
 	const uint PrimitiveFlags = GetPrimitiveData(MaterialParameters.PrimitiveId).Flags;
 
 	Payload.PrimitiveLightingChannelMask = GetPrimitive_LightingChannelMask_FromFlags(PrimitiveFlags);
 
-	if ((PrimitiveFlags & PRIMITIVE_SCENE_DATA_FLAG_DECAL_RECEIVER) != 0)
-	{
-		Payload.SetDecalReceiver();
-	}
-
 	Payload.HitT = RayTCurrent();
 	if (HitKind() == HIT_KIND_TRIANGLE_FRONT_FACE)
 	{
@@ -207,30 +219,12 @@ RAY_TRACING_ENTRY_CLOSEST_HIT(PathTracingMaterialCHS,
 	Payload.Specular = Specular;
 	Payload.Metallic = Metallic;
 	Payload.Roughness = Roughness;
-	Payload.Opacity = GetMaterialOpacity(PixelMaterialInputs);
 
 #if MATERIAL_USES_ANISOTROPY
 	Payload.WorldTangent = CalculateAnisotropyTangent(MaterialParameters, PixelMaterialInputs);
 	Payload.Anisotropy = GetMaterialAnisotropy(PixelMaterialInputs);
 #endif	
 
-#if MATERIALBLENDING_MASKED
-#if MATERIAL_DITHER_OPACITY_MASK
-	// dithering emulates real transparency, so switch to translucent
-	// NOTE: the raster path technically takes into account the opacity mask clip value, so the effective transparency should be:
-	//        saturate(MaskRaw - ClipValue + 0.5)
-	// (See derivation in DitheredOpacityMaskToOpacity)
-	// However this behavior is surprising to most users and does not exactly match the rasterizer anyway due to how the realtime AA
-	// code performs blending.
-	// Since the goal of dithered opacity is to emulate ordinary transparency, just use the mask input as opacity directly and
-	// ignore the configured clip value.
-	Payload.Opacity = saturate(GetMaterialMaskInputRaw(PixelMaterialInputs));
-	Payload.BlendingMode = RAY_TRACING_BLEND_MODE_TRANSLUCENT;
-#else
-	// regular masked mode - accounted for in AHS
-#endif
-#endif // MATERIALBLENDING_MASKED
-
 #if HAIR_STRAND_MESH_FACTORY
 	Payload.WorldSmoothNormal = Payload.WorldNormal = Payload.WorldGeoNormal;
 	Payload.WorldTangent = MaterialParameters.TangentToWorld[2];
@@ -517,12 +511,18 @@ RAY_TRACING_ENTRY_CLOSEST_HIT(PathTracingMaterialCHS,
 				// keep the opacity from the material, it might indicate partial transparency
 				// so we should get a partial holdout
 			}
-			Payload.BlendingMode = MaterialBlendingMode = RAY_TRACING_BLEND_MODE_ALPHA_HOLDOUT;
+			Payload.BlendingMode = RAY_TRACING_BLEND_MODE_ALPHA_HOLDOUT;
 			Payload.ShadingModelID = SHADINGMODELID_UNLIT;
 			Payload.Radiance = 0;
 		}
 	}
 
+	// Only opaque surfaces can recieve decals. This is both for compatibility with raster and improving performance when we have stacks of transparent surfaces
+	if ((Payload.BlendingMode == RAY_TRACING_BLEND_MODE_OPAQUE) && ((PrimitiveFlags & PRIMITIVE_SCENE_DATA_FLAG_DECAL_RECEIVER) != 0))
+	{
+		Payload.SetDecalReceiver();
+	}
+
 	// adjust after everything else (because SSS profile case needs to decode the dual spec info from the unmodified roughness)
 	Payload.Roughness = AdjustMaterialRoughness(Payload.Roughness, PackedPayload.GetPathRoughness());
 
diff --git a/Engine/Shaders/Private/RayTracing/RayTracingCommon.ush b/Engine/Shaders/Private/RayTracing/RayTracingCommon.ush
--- a/Engine/Shaders/Private/RayTracing/RayTracingCommon.ush
+++ b/Engine/Shaders/Private/RayTracing/RayTracingCommon.ush
@@ -913,12 +913,11 @@ CHECK_RT_PAYLOAD_SIZE(FPackedMaterialClosestHitPayload)
 #define DECAL_WRITE_NORMAL_FLAG							(1 << 1)
 #define DECAL_WRITE_ROUGHNESS_SPECULAR_METALLIC_FLAG	(1 << 2)
 #define DECAL_WRITE_EMISSIVE_FLAG						(1 << 3)
-#define DECAL_WRITE_AMBIENT_OCCLUSION_FLAG				(1 << 4)
 
-#if IS_PAYLOAD_ENABLED(RT_PAYLOAD_TYPE_DECALS | RT_PAYLOAD_TYPE_MESH_DECALS)
-struct FDecalShaderParams
+#if IS_PAYLOAD_ENABLED(RT_PAYLOAD_TYPE_DECALS)
+struct FDecalShaderPayload : FMinimalPayload
 {
-	uint Data[15];
+	uint Data[14];
 
 	void SetInputTranslatedWorldPosition(float3 In)
 	{
@@ -960,34 +959,18 @@ struct FDecalShaderParams
 	}
 	float3 GetEmissive() { return float3(f16tof32(Data[10]), f16tof32(Data[10] >> 16), f16tof32(Data[11])); }
 
-	void SetAO(float In)
-	{
-		Data[12] = asuint(In);
-	}
-	float GetAO() { return asfloat(Data[12]); }
-
 	void SetTransparency(float In)
 	{
-		Data[13] = asuint(In);
+		Data[12] = asuint(In);
 	}
-	float GetTransparency() { return asfloat(Data[13]); }
+	float GetTransparency() { return asfloat(Data[12]); }
 
-	void SetFlags(uint In) { Data[14] = In; }
-	uint GetFlags() { return Data[14]; }
+	void SetFlags(uint In) { Data[13] = In; }
+	uint GetFlags() { return Data[13]; }
 };
+CHECK_RT_PAYLOAD_SIZE(FDecalShaderPayload)
 #endif
 
-#if IS_PAYLOAD_ENABLED(RT_PAYLOAD_TYPE_DECALS)
-CHECK_RT_PAYLOAD_SIZE(FDecalShaderParams)
-#endif
-
-#if IS_PAYLOAD_ENABLED(RT_PAYLOAD_TYPE_MESH_DECALS)
-struct FRayTracingDecalPayload : FMinimalPayload
-{
-	FDecalShaderParams Params;
-};
-CHECK_RT_PAYLOAD_SIZE(FRayTracingDecalPayload)
-#endif
 
 ///////////////////////////////////////////////////////////////////////////////////////////////////
 // Tracing util functions
diff --git a/Engine/Shaders/Private/RayTracing/RayTracingDecalMaterialShader.usf b/Engine/Shaders/Private/RayTracing/RayTracingDecalMaterialShader.usf
--- a/Engine/Shaders/Private/RayTracing/RayTracingDecalMaterialShader.usf
+++ b/Engine/Shaders/Private/RayTracing/RayTracingDecalMaterialShader.usf
@@ -21,18 +21,19 @@
 
 #include "/Engine/Generated/Material.ush"
 
-RAY_TRACING_ENTRY_CALLABLE(RayTracingDecalMaterialShader, FDecalShaderParams, Params)
+#if RAYCALLABLESHADER
+RAY_TRACING_ENTRY_CALLABLE(RayTracingDecalMaterialShader, FDecalShaderPayload, Payload)
 {
 	ResolvedView = ResolveView();
 
-	const float3 TranslatedWorldPosition = Params.GetInputTranslatedWorldPosition();
+	const float3 TranslatedWorldPosition = Payload.GetInputTranslatedWorldPosition();
 
 	float4 DecalPos = mul(float4(TranslatedWorldPosition.xyz, 1), WorldToDecal);
 
 	if (any(abs(DecalPos.xyz) > 1))
 	{
 		// clip content outside the decal
-		Params.SetTransparency(1.0f);
+		Payload.SetTransparency(1.0f);
 		return;
 	}
 
@@ -69,35 +70,31 @@ RAY_TRACING_ENTRY_CALLABLE(RayTracingDecalMaterialShader, FDecalShaderParams, Pa
 
 	const float DecalFading = saturate(4 - 4 * abs(SwizzlePos.z * 2 - 1)) * DecalParams.x;
 	const float Opacity = GetMaterialOpacity(PixelMaterialInputs) * DecalFading;
-	Params.SetTransparency(1.0f - Opacity);
+	Payload.SetTransparency(1.0f - Opacity);
 
 	if (DecalWriteFlags & DECAL_WRITE_BASE_COLOR_FLAG)
 	{
-		Params.SetBaseColor(GetMaterialBaseColor(PixelMaterialInputs) * Opacity);
+		Payload.SetBaseColor(GetMaterialBaseColor(PixelMaterialInputs) * Opacity);
 	}
 
 	if (DecalWriteFlags & DECAL_WRITE_NORMAL_FLAG)
 	{
-		Params.SetWorldNormal(MaterialParameters.WorldNormal * Opacity);
+		Payload.SetWorldNormal(MaterialParameters.WorldNormal * Opacity);
 	}
 
 	if (DecalWriteFlags & DECAL_WRITE_ROUGHNESS_SPECULAR_METALLIC_FLAG)
 	{
-		Params.SetMetallicSpecularRoughness(float3(GetMaterialMetallic(PixelMaterialInputs), GetMaterialSpecular(PixelMaterialInputs), GetMaterialRoughness(PixelMaterialInputs)) * Opacity);
-	}
-
-	if (DecalWriteFlags & DECAL_WRITE_AMBIENT_OCCLUSION_FLAG)
-	{
-		Params.SetAO(GetMaterialAmbientOcclusion(PixelMaterialInputs) * Opacity);
+		Payload.SetMetallicSpecularRoughness(float3(GetMaterialMetallic(PixelMaterialInputs), GetMaterialSpecular(PixelMaterialInputs), GetMaterialRoughness(PixelMaterialInputs)) * Opacity);
 	}
 
 	if (DecalWriteFlags & DECAL_WRITE_EMISSIVE_FLAG)
 	{
-		Params.SetEmissive(GetMaterialEmissive(PixelMaterialInputs) * Opacity);
+		Payload.SetEmissive(GetMaterialEmissive(PixelMaterialInputs) * Opacity);
 	}
 
-	Params.SetFlags(DecalWriteFlags);
+	Payload.SetFlags(DecalWriteFlags);
 }
+#endif // RAYCALLABLESHADER
 
 #if RAYHITGROUPSHADER
 
@@ -114,7 +111,7 @@ RAY_TRACING_ENTRY_CALLABLE(RayTracingDecalMaterialShader, FDecalShaderParams, Pa
 #include "/Engine/Private/RayTracing/RayTracingCalcInterpolants.ush"
 
 RAY_TRACING_ENTRY_CLOSEST_HIT(RayTracingDecalMaterialCHS,
-	FRayTracingDecalPayload, Payload,
+	FDecalShaderPayload, Payload,
 	FRayTracingIntersectionAttributes, Attributes)
 {
 	Payload.HitT = RayTCurrent();
@@ -151,40 +148,35 @@ RAY_TRACING_ENTRY_CLOSEST_HIT(RayTracingDecalMaterialCHS,
 	}
 
 	const float Opacity = GetMaterialOpacity(PixelMaterialInputs);
-	Payload.Params.SetTransparency(1.0f - Opacity);
+	Payload.SetTransparency(1.0f - Opacity);
 
 	if (DECAL_PAYLOAD_FLAGS & DECAL_WRITE_BASE_COLOR_FLAG)
 	{
-		Payload.Params.SetBaseColor(GetMaterialBaseColor(PixelMaterialInputs) * Opacity);
+		Payload.SetBaseColor(GetMaterialBaseColor(PixelMaterialInputs) * Opacity);
 	}
 
 	if (DECAL_PAYLOAD_FLAGS & DECAL_WRITE_NORMAL_FLAG)
 	{
-		Payload.Params.SetWorldNormal(MaterialParameters.WorldNormal * Opacity);
+		Payload.SetWorldNormal(MaterialParameters.WorldNormal * Opacity);
 	}
 
 	if (DECAL_PAYLOAD_FLAGS & DECAL_WRITE_ROUGHNESS_SPECULAR_METALLIC_FLAG)
 	{
-		Payload.Params.SetMetallicSpecularRoughness(float3(GetMaterialMetallic(PixelMaterialInputs), GetMaterialSpecular(PixelMaterialInputs), GetMaterialRoughness(PixelMaterialInputs)) * Opacity);
-	}
-
-	if (DECAL_PAYLOAD_FLAGS & DECAL_WRITE_AMBIENT_OCCLUSION_FLAG)
-	{
-		Payload.Params.SetAO(GetMaterialAmbientOcclusion(PixelMaterialInputs) * Opacity);
+		Payload.SetMetallicSpecularRoughness(float3(GetMaterialMetallic(PixelMaterialInputs), GetMaterialSpecular(PixelMaterialInputs), GetMaterialRoughness(PixelMaterialInputs)) * Opacity);
 	}
 
 	if (DECAL_PAYLOAD_FLAGS & DECAL_WRITE_EMISSIVE_FLAG)
 	{
-		Payload.Params.SetEmissive(GetMaterialEmissive(PixelMaterialInputs) * Opacity);
+		Payload.SetEmissive(GetMaterialEmissive(PixelMaterialInputs) * Opacity);
 	}
 
-	Payload.Params.SetFlags(DECAL_PAYLOAD_FLAGS);
+	Payload.SetFlags(DECAL_PAYLOAD_FLAGS);
 
 	return;
 }
 
 RAY_TRACING_ENTRY_ANY_HIT(RayTracingDecalMaterialAHS,
-	FRayTracingDecalPayload, Payload,
+	FDecalShaderPayload, Payload,
 	FRayTracingIntersectionAttributes, Attributes)
 {
 	ResolvedView = ResolveView();
@@ -206,12 +198,7 @@ RAY_TRACING_ENTRY_ANY_HIT(RayTracingDecalMaterialAHS,
 
 		MaterialParameters.CameraVector = -WorldRayDirection();
 
-		// #dxr_todo: UE-72130 support world position offset 
-		// #if USE_WORLD_POSITION_EXCLUDING_SHADER_OFFSETS
-		// CalcMaterialParametersEx(MaterialParameters, PixelMaterialInputs, SvPosition, ScreenPosition, bIsFrontFace, TranslatedWorldPosition, BasePassInterpolants.PixelPositionExcludingWPO);
-		// #else
 		CalcMaterialParametersEx(MaterialParameters, PixelMaterialInputs, SvPosition, ScreenPosition, bIsFrontFace, TranslatedWorldPosition, TranslatedWorldPosition);
-		// #endif
 	}
 
 	float Opacity = GetMaterialOpacity(PixelMaterialInputs);
diff --git a/Engine/Shaders/Private/RayTracing/RayTracingDefaultDecalHitShader.usf b/Engine/Shaders/Private/RayTracing/RayTracingDefaultDecalHitShader.usf
--- a/Engine/Shaders/Private/RayTracing/RayTracingDefaultDecalHitShader.usf
+++ b/Engine/Shaders/Private/RayTracing/RayTracingDefaultDecalHitShader.usf
@@ -1,24 +1,24 @@
 #include "/Engine/Private/RayTracing/RayTracingCommon.ush"
 
 RAY_TRACING_ENTRY_CLOSEST_HIT(RayTracingDefaultOpaqueDecalMaterialCHS,
-	FRayTracingDecalPayload, Payload,
+	FDecalShaderPayload, Payload,
 	FRayTracingIntersectionAttributes, Attributes)
 {
 	Payload.HitT = RayTCurrent();
 
-	Payload.Params.SetTransparency(0.0f);
-	Payload.Params.SetBaseColor(float3(0, 0, 0));
-	Payload.Params.SetFlags(DECAL_WRITE_BASE_COLOR_FLAG);
+	Payload.SetTransparency(0.0f);
+	Payload.SetBaseColor(float3(0, 0, 0));
+	Payload.SetFlags(DECAL_WRITE_BASE_COLOR_FLAG);
 }
 
 RAY_TRACING_ENTRY_CLOSEST_HIT(RayTracingDefaultHiddenDecalMaterialCHS,
-	FRayTracingDecalPayload, Payload,
+	FDecalShaderPayload, Payload,
 	FRayTracingIntersectionAttributes, Attributes)
 {
 }
 
 RAY_TRACING_ENTRY_ANY_HIT(RayTracingDefaultHiddenDecalMaterialAHS,
-	FRayTracingDecalPayload, Payload,
+	FDecalShaderPayload, Payload,
 	FRayTracingIntersectionAttributes, Attributes)
 {
 	IgnoreHit();
diff --git a/Engine/Shaders/Private/RayTracing/RayTracingMaterialDefaultHitShaders.usf b/Engine/Shaders/Private/RayTracing/RayTracingMaterialDefaultHitShaders.usf
--- a/Engine/Shaders/Private/RayTracing/RayTracingMaterialDefaultHitShaders.usf
+++ b/Engine/Shaders/Private/RayTracing/RayTracingMaterialDefaultHitShaders.usf
@@ -3,6 +3,7 @@
 #include "../Common.ush"
 #include "RayTracingCommon.ush"
 
+#if IS_PAYLOAD_ENABLED(RT_PAYLOAD_TYPE_RAYTRACING_MATERIAL)
 RAY_TRACING_ENTRY_CLOSEST_HIT(OpaqueShadowCHS,
 	FPackedMaterialClosestHitPayload, PackedPayload,
 	FRayTracingIntersectionAttributes, Attributes)
@@ -22,10 +23,11 @@ RAY_TRACING_ENTRY_ANY_HIT(HiddenMaterialAHS,
 {
 	IgnoreHit();
 }
+#endif
 
 #if IS_PAYLOAD_ENABLED(RT_PAYLOAD_TYPE_DECALS)
 RAY_TRACING_ENTRY_CALLABLE(DefaultCallableShader,
-	FDecalShaderParams, Params)
+	FDecalShaderPayload, Params)
 {
 	// This shader serves only has a placeholder in RTPSO
 	// If RTPSO doesn't include a required shader, a NULL SBT entry should be used.
diff --git a/Engine/Shaders/Shared/RayTracingPayloadType.h b/Engine/Shaders/Shared/RayTracingPayloadType.h
--- a/Engine/Shaders/Shared/RayTracingPayloadType.h
+++ b/Engine/Shaders/Shared/RayTracingPayloadType.h
@@ -12,9 +12,8 @@
 #define RT_PAYLOAD_TYPE_PATHTRACING_MATERIAL	(1 << 5)	// FPackedPathTracingPayload
 #define RT_PAYLOAD_TYPE_LUMEN_MINIMAL			(1 << 6)	// FLumenMinimalPayload
 #define RT_PAYLOAD_TYPE_VFX   					(1 << 7)	// FVFXTracePayload
-#define RT_PAYLOAD_TYPE_DECALS					(1 << 8)	// FDecalShaderParams 
-#define RT_PAYLOAD_TYPE_MESH_DECALS				(1 << 9)	// FRayTracingDecalPayload 
-#define RT_PAYLOAD_TYPE_SPARSE_VOXEL			(1 << 10)	// FSparseVoxelPayload
+#define RT_PAYLOAD_TYPE_DECALS					(1 << 8)	// FDecalShaderPayload 
+#define RT_PAYLOAD_TYPE_SPARSE_VOXEL			(1 << 9)	// FSparseVoxelPayload
 
 #ifdef __cplusplus
 /**
@@ -39,7 +38,6 @@ enum ERayTracingPayloadType : uint // HLSL
 	LumenMinimal = RT_PAYLOAD_TYPE_LUMEN_MINIMAL,
 	VFX = RT_PAYLOAD_TYPE_VFX,
 	Decals = RT_PAYLOAD_TYPE_DECALS,
-	MeshDecals = RT_PAYLOAD_TYPE_MESH_DECALS,
 	SparseVoxel = RT_PAYLOAD_TYPE_SPARSE_VOXEL,
 };
 
diff --git a/Engine/Source/Runtime/Renderer/Private/PathTracing.cpp b/Engine/Source/Runtime/Renderer/Private/PathTracing.cpp
--- a/Engine/Source/Runtime/Renderer/Private/PathTracing.cpp
+++ b/Engine/Source/Runtime/Renderer/Private/PathTracing.cpp
@@ -303,13 +303,6 @@ TAutoConsoleVariable<float> CVarPathTracingMeshDecalBias(
 	ECVF_RenderThreadSafe
 );
 
-TAutoConsoleVariable<int32> CVarPathTracingMeshDecalsUseViewRay(
-	TEXT("r.PathTracing.MeshDecalUsesViewRay"),
-	1,
-	TEXT("Controls if mesh decals should be probed by the viewing ray (default) or by a short ray perpendicular to the surface (experimental)"),
-	ECVF_RenderThreadSafe
-);
-
 TAutoConsoleVariable<int32> CVarPathTracingLightFunctionColor(
 	TEXT("r.PathTracing.LightFunctionColor"),
 	0,
@@ -341,7 +334,6 @@ BEGIN_SHADER_PARAMETER_STRUCT(FPathTracingData, )
 	SHADER_PARAMETER(uint32, EnabledLightingContributions) // PATHTRACER_CONTRIBUTION_*
 	SHADER_PARAMETER(uint32, ApplyDiffuseSpecularOverrides)
 	SHADER_PARAMETER(int32, MaxRaymarchSteps)
-	SHADER_PARAMETER(int32, MeshDecalsUseViewRay)
 	SHADER_PARAMETER(float, MaxPathIntensity)
 	SHADER_PARAMETER(float, MaxNormalBias)
 	SHADER_PARAMETER(float, FilterWidth)
@@ -394,7 +386,6 @@ struct FPathTracingConfig
 			PathTracingData.DecalRoughnessCutoff != Other.PathTracingData.DecalRoughnessCutoff ||
 			PathTracingData.MeshDecalRoughnessCutoff != Other.PathTracingData.MeshDecalRoughnessCutoff ||
 			PathTracingData.MeshDecalBias != Other.PathTracingData.MeshDecalBias ||
-			PathTracingData.MeshDecalsUseViewRay != Other.PathTracingData.MeshDecalsUseViewRay ||
 			PathTracingData.MaxRaymarchSteps != Other.PathTracingData.MaxRaymarchSteps ||
 			ViewRect != Other.ViewRect ||
 			LightShowFlags != Other.LightShowFlags ||
@@ -572,7 +563,6 @@ static void PreparePathTracingData(const FScene* Scene, const FViewInfo& View, F
 
 	PathTracingData.MeshDecalRoughnessCutoff = PathTracing::UsesDecals(*View.Family) && Scene->RayTracingScene.GetRHIRayTracingScene()->GetInitializer().NumNativeInstancesPerLayer[(uint32)ERayTracingSceneLayer::Decals] > 0 ? CVarPathTracingMeshDecalRoughnessCutoff.GetValueOnRenderThread() : -1.0f;
 	PathTracingData.MeshDecalBias = CVarPathTracingMeshDecalBias.GetValueOnRenderThread();
-	PathTracingData.MeshDecalsUseViewRay = CVarPathTracingMeshDecalsUseViewRay.GetValueOnRenderThread();
 
 	PathTracingData.MaxRaymarchSteps = CVarPathTracingMaxRaymarchSteps.GetValueOnRenderThread();
 
@@ -773,7 +763,7 @@ class FPathTracingRG : public FGlobalShader
 
 	static ERayTracingPayloadType GetRayTracingPayloadType(const int32 PermutationId)
 	{
-		return ERayTracingPayloadType::PathTracingMaterial | ERayTracingPayloadType::Decals | ERayTracingPayloadType::MeshDecals;
+		return ERayTracingPayloadType::PathTracingMaterial | ERayTracingPayloadType::Decals;
 	}
 
 	BEGIN_SHADER_PARAMETER_STRUCT(FParameters, )
diff --git a/Engine/Source/Runtime/Renderer/Private/RayTracing/RayTracingDecals.cpp b/Engine/Source/Runtime/Renderer/Private/RayTracing/RayTracingDecals.cpp
--- a/Engine/Source/Runtime/Renderer/Private/RayTracing/RayTracingDecals.cpp
+++ b/Engine/Source/Runtime/Renderer/Private/RayTracing/RayTracingDecals.cpp
@@ -296,7 +296,7 @@ public:
 
 	static ERayTracingPayloadType GetRayTracingPayloadType(const int32 PermutationId)
 	{
-		return ERayTracingPayloadType::MeshDecals;
+		return ERayTracingPayloadType::Decals;
 	}
 };
 
@@ -306,8 +306,6 @@ using FRayTracingDecalMaterialCHS_AHS = TRayTracingDecalMaterial<true>;
 IMPLEMENT_MATERIAL_SHADER_TYPE(template<>, FRayTracingDecalMaterialCHS, TEXT("/Engine/Private/RayTracing/RayTracingDecalMaterialShader.usf"), TEXT("closesthit=RayTracingDecalMaterialCHS"), SF_RayHitGroup);
 IMPLEMENT_MATERIAL_SHADER_TYPE(template<>, FRayTracingDecalMaterialCHS_AHS, TEXT("/Engine/Private/RayTracing/RayTracingDecalMaterialShader.usf"), TEXT("closesthit=RayTracingDecalMaterialCHS anyhit=RayTracingDecalMaterialAHS"), SF_RayHitGroup);
 
-IMPLEMENT_RT_PAYLOAD_TYPE(ERayTracingPayloadType::MeshDecals, 64);
-
 FShaderType* GetRayTracingDecalMaterialShaderType(bool bUseAnyHitShader)
 {
 	if (bUseAnyHitShader)
@@ -332,7 +330,7 @@ class FDefaultOpaqueMeshDecalHitGroup : public FGlobalShader
 
 	static ERayTracingPayloadType GetRayTracingPayloadType(const int32 PermutationId)
 	{
-		return ERayTracingPayloadType::MeshDecals;
+		return ERayTracingPayloadType::Decals;
 	}
 
 	using FParameters = FEmptyShaderParameters;
@@ -352,7 +350,7 @@ class FDefaultHiddenMeshDecalHitGroup : public FGlobalShader
 
 	static ERayTracingPayloadType GetRayTracingPayloadType(const int32 PermutationId)
 	{
-		return ERayTracingPayloadType::MeshDecals;
+		return ERayTracingPayloadType::Decals;
 	}
 
 	using FParameters = FEmptyShaderParameters;
diff --git a/Engine/Source/Runtime/Renderer/Private/RayTracing/RayTracingMaterialHitShaders.cpp b/Engine/Source/Runtime/Renderer/Private/RayTracing/RayTracingMaterialHitShaders.cpp
--- a/Engine/Source/Runtime/Renderer/Private/RayTracing/RayTracingMaterialHitShaders.cpp
+++ b/Engine/Source/Runtime/Renderer/Private/RayTracing/RayTracingMaterialHitShaders.cpp
@@ -648,11 +648,6 @@ FRayTracingPipelineState* FDeferredShadingSceneRenderer::CreateRayTracingMateria
 		? (ERayTracingPayloadType::PathTracingMaterial | ERayTracingPayloadType::Decals)
 		: ERayTracingPayloadType::RayTracingMaterial;
 
-	if (bSupportMeshDecals)
-	{
-		PayloadType |= ERayTracingPayloadType::MeshDecals;
-	}
-
 	FRayTracingPipelineStateInitializer Initializer;
 
 	Initializer.MaxPayloadSizeInBytes = GetRayTracingPayloadTypeMaxSize(PayloadType);
@@ -801,7 +796,7 @@ FRayTracingPipelineState* FDeferredShadingSceneRenderer::CreateRayTracingMateria
 						const FVisibleRayTracingMeshCommand VisibleMeshCommand = MeshCommands[CommandIndex];
 						const FRayTracingMeshCommand& MeshCommand = *VisibleMeshCommand.RayTracingMeshCommand;
 
-						const bool bIsMeshDecalShader = MeshCommand.MaterialShader->RayTracingPayloadType == (uint32)ERayTracingPayloadType::MeshDecals;
+						const bool bIsMeshDecalShader = MeshCommand.MaterialShader->RayTracingPayloadType == (uint32)ERayTracingPayloadType::Decals;
 
 						// TODO: Following check is disabled since FRayTracingMeshProcessor non-path-tracing code paths still don't assign the appropriate shader to decal mesh commands.
 						// We could also potentially use regular materials to approximate decals in ray tracing in some situations.
