Fix SSS causes an orange rim using Ray Traced lighting at certain angles. The issue is caused by not enough penalty to the depth in the bilateral filtering.

diff --git a/Engine/Shaders/Private/ScreenSpaceDenoise/SSDSignalFramework.ush b/Engine/Shaders/Private/ScreenSpaceDenoise/SSDSignalFramework.ush
--- a/Engine/Shaders/Private/ScreenSpaceDenoise/SSDSignalFramework.ush
+++ b/Engine/Shaders/Private/ScreenSpaceDenoise/SSDSignalFramework.ush
@@ -416,6 +416,17 @@ float3 ComputeMajorLightRayDirection(FSSDSampleSceneInfos SceneMetadata, FSSDSig
 }
 #endif
 
+float ComputeLightSourceRadius(FSSDSignalDomainKnowledge DomainKnowledge)
+#if CONFIG_SIGNAL_PROCESSING == SIGNAL_PROCESSING_SHADOW_VISIBILITY_MASK
+{
+	return  DomainKnowledge.Light.SourceRadius;
+}
+#else
+{
+	return 0.0f;
+}
+#endif
+
 /** COmpute the weight of the bilateral filter to use between a reference sample and a neighbor. */
 float ComputeBilateralWeight(
 	const uint BilateralSettings,
@@ -495,14 +506,21 @@ float ComputeBilateralWeight(
 		
 		float PlaneMisAlignementBilateral;
 		{
-			const float NormalEncodingError = 0.25;
+			// Use source radius as encoding error to add more penalty to depth to avoid SSS backface transmission rim for spotlight with source radius > 0.
+			// It makes the scaling of the distance by (DistanceFromLight - HitT)/HitT instead of (DistanceFromLight - HitT)/HitT*SourceRadius
+			// When light is behind close to neighbor in view direction, the light tangent bilaterial will be close to 1. Little penalty
+			// is applied in depth. Adding the effect of SourceRadius on NeighborToRefError could add more penalty to depth to mitigate the
+			// bleeding error.
+			const bool bApplyDepthPenaltyWeightForSSS = RefSceneMetadata.ShadingModelID == SHADINGMODELID_SUBSURFACE_PROFILE && ComputeLightSourceRadius(DomainKnowledge) > 0.0f;
+			const float	NormalEncodingError = bApplyDepthPenaltyWeightForSSS ? DomainKnowledge.Light.SourceRadius : -0.25;
+
 			
 			float DistFromRefPlane = dot(RefNormal, -NeighborToRefVector);
 
 			float NeighborToRefPlaneDistance = length(-NeighborToRefVector - DistFromRefPlane * RefNormal);
 			float NeighborToRefError = NeighborToRefPlaneDistance * NormalEncodingError;
 
-			PlaneMisAlignementBilateral = saturate(1 - 6 * (abs(DistFromRefPlane) - NeighborToRefError) * InvMaxWorldBluringDistance);
+			PlaneMisAlignementBilateral = saturate(1 - 6 * (abs(DistFromRefPlane) + NeighborToRefError) * InvMaxWorldBluringDistance);
 		}
 
 		PositionBilateralWeight = LightTangentBilateral * PlaneMisAlignementBilateral;
