Strata - fixed ConvertCloudPixelMaterialInputsToWorkingColorSpace.

diff --git a/Engine/Shaders/Private/VolumetricCloud.usf b/Engine/Shaders/Private/VolumetricCloud.usf
--- a/Engine/Shaders/Private/VolumetricCloud.usf
+++ b/Engine/Shaders/Private/VolumetricCloud.usf
@@ -142,9 +142,6 @@ float SampleAmbientOcclusion(in FPixelMaterialInputs PixelMaterialInputs)
 void ConvertCloudPixelMaterialInputsToWorkingColorSpace(inout FPixelMaterialInputs PixelMaterialInputs)
 {
 #if WORKING_COLOR_SPACE_IS_SRGB
-#if STRATA_ENABLED
-	// TODO
-#else
 	// Get raw data
 	float3 ExtinctionSRGB = SampleExtinctionCoefficients(PixelMaterialInputs);
 	float3 AlbedoSRGB = SampleAlbedo(PixelMaterialInputs);
@@ -163,6 +160,11 @@ void ConvertCloudPixelMaterialInputsToWorkingColorSpace(inout FPixelMaterialInpu
 	float3 ExtinctionColorSpace = ScatteringColorSpace + AbsorptionColorSpace;
 	float3 AlbedoColorSpace		= ScatteringColorSpace / max(0.000000000001f, ExtinctionColorSpace);
 
+#if STRATA_ENABLED
+	BSDF_GETEMISSIVE(PixelMaterialInputs.FrontMaterial.InlinedBSDF)					= EmissiveColorSpace;
+	VOLUMETRICFOGCLOUD_EXTINCTION(PixelMaterialInputs.FrontMaterial.InlinedBSDF)	= ExtinctionColorSpace;
+	VOLUMETRICFOGCLOUD_ALBEDO(PixelMaterialInputs.FrontMaterial.InlinedBSDF)		= AlbedoColorSpace;
+#else
 	// Set raw values back
 	PixelMaterialInputs.EmissiveColor	= EmissiveColorSpace;
 	PixelMaterialInputs.Subsurface		= ExtinctionColorSpace;
