Strata - fixed volumetric cloud extinction being grayscale with Strata only.

diff --git a/Engine/Shaders/Private/VolumetricCloud.usf b/Engine/Shaders/Private/VolumetricCloud.usf
--- a/Engine/Shaders/Private/VolumetricCloud.usf
+++ b/Engine/Shaders/Private/VolumetricCloud.usf
@@ -88,7 +88,7 @@ float3 SampleExtinctionCoefficients(in FPixelMaterialInputs PixelMaterialInputs)
 	float3 Extinction = 0.0f;
 #if STRATA_ENABLED
 	FStrataBSDF BSDF = PixelMaterialInputs.FrontMaterial.InlinedBSDF;
-	Extinction = VOLUMETRICFOGCLOUD_EXTINCTION(BSDF).r;
+	Extinction = VOLUMETRICFOGCLOUD_EXTINCTION(BSDF).rgb;
 #else
 	#if !MATERIAL_SHADINGMODEL_UNLIT
 	Extinction = GetMaterialSubsurfaceDataRaw(PixelMaterialInputs).rgb;
@@ -102,9 +102,9 @@ float3 SampleEmissive(in FPixelMaterialInputs PixelMaterialInputs)
 	float3 EmissiveColor = 0.0f;
 #if STRATA_ENABLED
 	FStrataBSDF BSDF = PixelMaterialInputs.FrontMaterial.InlinedBSDF;
-	EmissiveColor = BSDF_GETEMISSIVE(BSDF);
+	EmissiveColor = BSDF_GETEMISSIVE(BSDF).rgb;
 #else
-	EmissiveColor = GetMaterialEmissiveRaw(PixelMaterialInputs);
+	EmissiveColor = GetMaterialEmissiveRaw(PixelMaterialInputs).rgb;
 #endif
 	return clamp(EmissiveColor, 0.0f, 65000.0f);
 }
@@ -114,10 +114,10 @@ float3 SampleAlbedo(in FPixelMaterialInputs PixelMaterialInputs)
 	float3 Albedo = 0.0f;
 #if STRATA_ENABLED
 	FStrataBSDF BSDF = PixelMaterialInputs.FrontMaterial.InlinedBSDF;
-	Albedo = VOLUMETRICFOGCLOUD_ALBEDO(BSDF);
+	Albedo = VOLUMETRICFOGCLOUD_ALBEDO(BSDF).rgb;
 #else
 	#if !MATERIAL_SHADINGMODEL_UNLIT
-	Albedo = GetMaterialBaseColor(PixelMaterialInputs) * View.DiffuseOverrideParameter.w + View.DiffuseOverrideParameter.xyz;
+	Albedo = GetMaterialBaseColor(PixelMaterialInputs).rgb * View.DiffuseOverrideParameter.w + View.DiffuseOverrideParameter.xyz;
 	#endif
 #endif
 	return saturate(Albedo);
