Fix RT shadow with Strata

The 'is valid material' bit on the TopLayerData is incorrectly set/read for some reason. This information is now retrieve from the header. This is a temporary fix which needs further investigation.

diff --git a/Engine/Shaders/Private/RayTracing/RayTracingOcclusionRGS.usf b/Engine/Shaders/Private/RayTracing/RayTracingOcclusionRGS.usf
--- a/Engine/Shaders/Private/RayTracing/RayTracingOcclusionRGS.usf
+++ b/Engine/Shaders/Private/RayTracing/RayTracingOcclusionRGS.usf
@@ -659,7 +659,7 @@ RAY_TRACING_ENTRY_RAYGEN(OcclusionRGS)
 			const bool bSSSIsValid = StrataSubSurfaceHeaderGetUseDiffusion(SSSHeader);
 			const bool bSSSIsProfile = StrataSubSurfaceHeaderGetIsProfile(SSSHeader);
 
-			ShadingParameters.bIsValid = IsStrataMaterial(TopLayerData);
+			ShadingParameters.bIsValid = IsStrataMaterial(StrataPixelHeader);
 			ShadingParameters.bIsHair = StrataGetBSDFType(StrataPixelHeader) == STRATA_BSDF_TYPE_HAIR;
 			ShadingParameters.bIsSubsurface = bSSSIsValid && !bSSSIsProfile;
 			ShadingParameters.bHasSubsurfaceProfile = bSSSIsValid && bSSSIsProfile;
