Strata - fixed single layer water with pixel depth offset.
Project Opal.

diff --git a/Engine/Shaders/Private/BasePassPixelShader.usf b/Engine/Shaders/Private/BasePassPixelShader.usf
--- a/Engine/Shaders/Private/BasePassPixelShader.usf
+++ b/Engine/Shaders/Private/BasePassPixelShader.usf
@@ -1597,8 +1597,9 @@ void FPixelShaderInOut_MainPS(
 	GBuffer.PrecomputedShadowFactors = PrecomputedShadowFactors;
 	GBuffer.Velocity = EncodedVelocity;
 
-#if OUTPUT_PIXEL_DEPTH_OFFSET
+#if OUTPUT_PIXEL_DEPTH_OFFSET && STRATA_INLINE_SINGLELAYERWATER==0
 	// When in deferred, opaque materials with pixel depth offset must execute a custom depth test in order to avoid conflicting UAV writes.
+	// This is however not done for single layer water as the depth will never match. Note: that can lead overlapping SLW polys to cause flicker on screen.
 	if (ManualDepthTestEqual(In.SvPosition, Out.Depth))
 #endif
 	{
