Clean up u32 -> u16x2 pack/unpack HLSL intrinsics

diff --git a/Engine/Shaders/Private/HairStrands/HairStrandsForwardRaster.usf b/Engine/Shaders/Private/HairStrands/HairStrandsForwardRaster.usf
--- a/Engine/Shaders/Private/HairStrands/HairStrandsForwardRaster.usf
+++ b/Engine/Shaders/Private/HairStrands/HairStrandsForwardRaster.usf
@@ -41,20 +41,6 @@ float3 UnpackR11G11B10F(uint rgb)
 	return float3(r, g, b);
 }
 
-uint PackRG16(float2 In)
-{
-#if COMPILER_SUPPORTS_PACK_INTRINSICS
-	return PackUInt2ToUInt(f32tof16(In.r), f32tof16(In.g));
-#else
-	return f32tof16(In.r) | (f32tof16(In.g) << 16);
-#endif
-}
-
-float2 UnpackRG16(uint In)
-{
-	return float2(f16tof32(In), f16tof32(In >> 16));
-}
-
 uint PackDepth(float In)
 {
 	return asuint(In);
@@ -73,13 +59,13 @@ struct FDepthRange
 
 uint PackDepthRange(FDepthRange In)
 {
-	return PackRG16(float2(In.MinZ, In.MaxZ));
+	return PackFloat2ToUInt(In.MinZ, In.MaxZ);
 }
 
 FDepthRange UnpackDepthRange(uint In)
 {
 	FDepthRange Out;
-	const float2 D = UnpackRG16(In);
+	const float2 D = UnpackFloat2FromUInt(In);
 	Out.MinZ = D.x;
 	Out.MaxZ = D.y;
 	return Out;
@@ -1208,12 +1194,12 @@ float3 UnpackColorData(uint In)
 
 uint2 PackVelocityData(float4 In)
 {
-	return uint2(PackRG16(In.xy), PackRG16(In.zw));
+	return uint2(PackFloat2ToUInt(In.xy), PackFloat2ToUInt(In.zw));
 }
 
 float4 UnpackVelocityData(uint2 In)
 {
-	return float4(UnpackRG16(In.x), UnpackRG16(In.y));
+	return float4(UnpackFloat2FromUInt(In.x), UnpackFloat2FromUInt(In.y));
 }
 
 float3 LoadSampleColor(uint InPrimId, uint2 InSampleResolution)
diff --git a/Engine/Shaders/Private/Strata/Strata.ush b/Engine/Shaders/Private/Strata/Strata.ush
--- a/Engine/Shaders/Private/Strata/Strata.ush
+++ b/Engine/Shaders/Private/Strata/Strata.ush
@@ -3016,28 +3016,14 @@ float3 UnpackR10G10B10F(uint rgb)
 	return float3(r, g, b);
 }
 
-uint PackRG16(float2 In)
-{
-#if COMPILER_SUPPORTS_PACK_INTRINSICS
-	return PackUInt2ToUInt(f32tof16(In.r), f32tof16(In.g));
-#else
-	return f32tof16(In.r) | (f32tof16(In.g) << 16);
-#endif
-}
-
-float2 UnpackRG16(uint In)
-{
-	return float2(f16tof32(In), f16tof32(In>>16));
-}
-
 uint2 PackR16G16B16A16F(float4 In)
 {
-	return uint2(PackRG16(In.xy), PackRG16(In.zw));
+	return uint2(PackFloat2ToUInt(In.xy), PackFloat2ToUInt(In.zw));
 }
 
 float4 UnpackR16G16B16A16F(uint2 In)
 {
-	return float4(UnpackRG16(In.x), UnpackRG16(In.y));
+	return float4(UnpackFloat2FromUInt(In.x), UnpackFloat2FromUInt(In.y));
 }
 
 uint PackColorLinearToGamma2(float3 In)
@@ -3199,12 +3185,12 @@ float2 SignNotZero(in float2 v)
 
 uint StrataPackNormal(in float3 Normal)
 {
-	return PackRG16(UnitVectorToOctahedron(Normal));
+	return PackFloat2ToUInt(UnitVectorToOctahedron(Normal));
 }
 
 float3 StrataUnpackNormal(uint PackedNormal)
 {
-	return OctahedronToUnitVector(UnpackRG16(PackedNormal));
+	return OctahedronToUnitVector(UnpackFloat2FromUInt(PackedNormal));
 }
 
 uint StrataPackNormal24(in float3 Normal)
diff --git a/Engine/Shaders/Public/Platform.ush b/Engine/Shaders/Public/Platform.ush
--- a/Engine/Shaders/Public/Platform.ush
+++ b/Engine/Shaders/Public/Platform.ush
@@ -188,8 +188,39 @@
 
 #ifndef COMPILER_SUPPORTS_PACK_INTRINSICS
 #define COMPILER_SUPPORTS_PACK_INTRINSICS 0
+
+uint PackUInt2ToUInt(uint X, uint Y)
+{
+	return X | (Y << 16u);
+}
+
 #endif
 
+uint PackUInt2ToUInt(uint2 XY)
+{
+	return PackUInt2ToUInt(XY.x, XY.y);
+}
+
+uint2 UnpackUIntFromUInt2(uint Packed)
+{
+	return uint2(Packed & 0xffff, Packed >> 16);
+}
+
+uint PackFloat2ToUInt(float X, float Y)
+{
+	return PackUInt2ToUInt(f32tof16(X), f32tof16(Y));
+}
+
+uint PackFloat2ToUInt(float2 XY)
+{
+	return PackFloat2ToUInt(XY.x, XY.y);
+}
+
+float2 UnpackFloat2FromUInt(uint In)
+{
+	return float2(f16tof32(In), f16tof32(In >> 16));
+}
+
 #ifndef COMPILER_SUPPORTS_WAVE_32_64_MODE
 #define COMPILER_SUPPORTS_WAVE_32_64_MODE 0
 #endif
