Fix for static like output during TSRResolveHistory on macOS due to what apears to be a GPU thread group race condition.  Could use simd_shuffle() here to fix on Metal to avoid the memory barrier and use of thread group memory but this has to be SM5 level - use a thread group memory barrier instead.

diff --git a/Engine/Shaders/Private/TemporalSuperResolution/TSRResolveHistory.usf b/Engine/Shaders/Private/TemporalSuperResolution/TSRResolveHistory.usf
--- a/Engine/Shaders/Private/TemporalSuperResolution/TSRResolveHistory.usf
+++ b/Engine/Shaders/Private/TemporalSuperResolution/TSRResolveHistory.usf
@@ -196,6 +196,11 @@ void MainCS(
 		if (bGenerateOutputMip1)
 		{
 			SharedArray0[LocalGroupThreadIndex] = HalfResOutput;
+
+			#if COMPILER_METAL
+				GroupMemoryBarrierWithGroupSync();
+			#endif
+
 			SharedArray0[LocalGroupThreadIndex] += SharedArray0[LocalGroupThreadIndex ^ 0x1];
 			SharedArray0[LocalGroupThreadIndex] += SharedArray0[LocalGroupThreadIndex ^ 0x2];
 			
@@ -212,4 +217,4 @@ void MainCS(
 	}
 	#endif
 }
- 
\ No newline at end of file
+
