Fix ClosestHZBTexture to be conservative by rounding up to nearest fp16, which fixes Lumen HZB Screen Traces missing intersections and travelling through geometry

diff --git a/Engine/Shaders/Private/HZB.usf b/Engine/Shaders/Private/HZB.usf
--- a/Engine/Shaders/Private/HZB.usf
+++ b/Engine/Shaders/Private/HZB.usf
@@ -62,6 +62,12 @@ RWTexture2D<float> ClosestHZBOutput_3;
 groupshared float SharedMinDeviceZ[GROUP_TILE_SIZE * GROUP_TILE_SIZE];
 groupshared float SharedMaxDeviceZ[GROUP_TILE_SIZE * GROUP_TILE_SIZE];
 
+float RoundUpF16(float DeviceZ)
+{
+	// ClosestDeviceZ needs to be rounded up to nearest fp16 to be conservative
+	return f16tof32(f32tof16(DeviceZ) + 1);
+}
+
 void OutputMipLevel(uint MipLevel, uint2 OutputPixelPos, float FurthestDeviceZ, float ClosestDeviceZ)
 {
 	#if DIM_MIP_LEVEL_COUNT >= 2
@@ -71,7 +77,7 @@ void OutputMipLevel(uint MipLevel, uint2 OutputPixelPos, float FurthestDeviceZ,
 			FurthestHZBOutput_1[OutputPixelPos] = FurthestDeviceZ;
 		#endif
 		#if DIM_CLOSEST
-			ClosestHZBOutput_1[OutputPixelPos] = ClosestDeviceZ;
+			ClosestHZBOutput_1[OutputPixelPos] = RoundUpF16(ClosestDeviceZ);
 		#endif
 	}
 	#endif
@@ -82,7 +88,7 @@ void OutputMipLevel(uint MipLevel, uint2 OutputPixelPos, float FurthestDeviceZ,
 			FurthestHZBOutput_2[OutputPixelPos] = FurthestDeviceZ;
 		#endif
 		#if DIM_CLOSEST
-			ClosestHZBOutput_2[OutputPixelPos] = ClosestDeviceZ;
+			ClosestHZBOutput_2[OutputPixelPos] = RoundUpF16(ClosestDeviceZ);
 		#endif
 	}
 	#endif
@@ -93,7 +99,7 @@ void OutputMipLevel(uint MipLevel, uint2 OutputPixelPos, float FurthestDeviceZ,
 			FurthestHZBOutput_3[OutputPixelPos] = FurthestDeviceZ;
 		#endif
 		#if DIM_CLOSEST
-			ClosestHZBOutput_3[OutputPixelPos] = ClosestDeviceZ;
+			ClosestHZBOutput_3[OutputPixelPos] = RoundUpF16(ClosestDeviceZ);
 		#endif
 	}		
 	#endif
@@ -189,11 +195,6 @@ void HZBBuildCS(
 
 	float MinDeviceZ = min(min3(DeviceZ.x, DeviceZ.y, DeviceZ.z), DeviceZ.w);
 	float MaxDeviceZ = max(max3(DeviceZ.x, DeviceZ.y, DeviceZ.z), DeviceZ.w);
-
-	// NOTE: For 16-bit HZB buffers we should be careful about rounding in the right direction for Min/Max.
-	// fp32 -> fp16 rounds towards zero in DirectX so the "min" HZB will be conservative for positive values,
-	// which is the common case for culling. Special care would be needed here to ensure that this quantization
-	// is fully conservative for the "max" HZB and/or negative numbers, but this is not currently necessary.
 	
 	uint2 OutputPixelPos = DispatchThreadId;
 	
@@ -202,7 +203,7 @@ void HZBBuildCS(
 	#endif
 	
 	#if DIM_CLOSEST
-		ClosestHZBOutput_0[OutputPixelPos] = MaxDeviceZ;
+		ClosestHZBOutput_0[OutputPixelPos] = RoundUpF16(MaxDeviceZ);
 	#endif
 
 	#if DIM_MIP_LEVEL_COUNT == 1
