Fix flickering in Lightmap Density viewmode:
* Make output position invariant so it matches the depth written to in depth prepass.

diff --git a/Engine/Shaders/Private/LightMapDensityShader.usf b/Engine/Shaders/Private/LightMapDensityShader.usf
--- a/Engine/Shaders/Private/LightMapDensityShader.usf
+++ b/Engine/Shaders/Private/LightMapDensityShader.usf
@@ -47,7 +47,8 @@ void MainVertexShader(
 	Output.WorldPosition = WorldPosition;
 	{
 		float4 RasterizedWorldPosition = VertexFactoryGetRasterizedWorldPosition(Input, VFIntermediates, Output.WorldPosition);
-		Output.Position = mul(RasterizedWorldPosition, ResolvedView.TranslatedWorldToClip);
+		float4 ClipSpacePosition = mul(RasterizedWorldPosition, ResolvedView.TranslatedWorldToClip);
+		Output.Position = INVARIANT(ClipSpacePosition);
 	}
 
 #if USE_GLOBAL_CLIP_PLANE
