Strata - fix for no blur to happen to refraction when roughness=0.

diff --git a/Engine/Shaders/Private/Strata/Strata.ush b/Engine/Shaders/Private/Strata/Strata.ush
--- a/Engine/Shaders/Private/Strata/Strata.ush
+++ b/Engine/Shaders/Private/Strata/Strata.ush
@@ -758,7 +758,8 @@ float GetVarianceFromRoughness(float Roughness)
 	float Roughness2 = Roughness * Roughness;
 	float Roughness3 = Roughness2 * Roughness;
 	float Roughness4 = Roughness2 * Roughness2;
-	return saturate(0.0000226575 - 0.000332564 * Roughness - 0.0168059 * Roughness2 + 0.115569 * Roughness3 - 0.0526336 * Roughness4);
+//	return saturate(0.0000226575 - 0.000332564 * Roughness - 0.0168059 * Roughness2 + 0.115569 * Roughness3 - 0.0526336 * Roughness4);	// Original
+	return saturate(             - 0.000332564 * Roughness - 0.0168059 * Roughness2 + 0.115569 * Roughness3 - 0.0526336 * Roughness4);	// Removing the 0.0000226575 to get GetVarianceFromRoughness(0)=0, to not get any blurred refraction when roughness==0.
 }
 
 ///////////////////////////////////////////////////////////////////////////////
