Fix LOD count on hair instance print.

diff --git a/Engine/Shaders/Private/HairStrands/HairStrandsDebug.usf b/Engine/Shaders/Private/HairStrands/HairStrandsDebug.usf
--- a/Engine/Shaders/Private/HairStrands/HairStrandsDebug.usf
+++ b/Engine/Shaders/Private/HairStrands/HairStrandsDebug.usf
@@ -1081,7 +1081,7 @@ struct FInstanceInfo
 {
 	uint  GroupIndex;
 	uint  GroupCount;
-	float LODCount;
+	uint  LODCount;
 	uint  GeometryType;
 	uint  BindingType;
 	bool  bHasSim;
@@ -1279,7 +1279,7 @@ void MainCS(uint GroupIndex : SV_GroupIndex, uint3 DispatchThreadId : SV_Dispatc
 		// LOD
 		Print(Context,Info.LOD, FontYellow, 4, 2);
 		Print(Context, TEXT("/"));
-		Print(Context,Info.LODCount, FontYellow, 1);
+		Print(Context,Info.LODCount-1, FontYellow, 1);
 		Print(Context, TEXT("  "));
 
 		// Geometry
