Merge all hair voxel code into HairStrandsVoxelRasterCompute.

diff --git a/Engine/Shaders/Private/HairStrands/HairStrandsDeepVoxelCommon.ush b/Engine/Shaders/Private/HairStrands/HairStrandsDeepVoxelCommon.ush
deleted file mode 100644
diff --git a/Engine/Shaders/Private/HairStrands/HairStrandsVoxelDebug.usf b/Engine/Shaders/Private/HairStrands/HairStrandsVoxelDebug.usf
new file mode 100644
--- /dev/null
+++ b/Engine/Shaders/Private/HairStrands/HairStrandsVoxelDebug.usf
@@ -0,0 +1,386 @@
+// Copyright Epic Games, Inc. All Rights Reserved.
+
+#include "../Common.ush"
+#include "../SceneTextureParameters.ush"
+#include "../DeferredShadingCommon.ush"
+#include "../ShaderPrint.ush"
+#include "../PositionReconstructionCommon.ush"  
+#include "HairStrandsVisibilityCommon.ush"
+#include "HairStrandsVoxelPageCommon.ush"
+#include "../ColorMap.ush"  
+
+#if SHADER_VOXEL_DEBUG
+
+#define VOXEL_TRAVERSAL_DEBUG_PAGEOCCUPANCY 0
+#define VOXEL_TRAVERSAL_DEBUG 0
+#define VOXEL_TRAVERSAL_FORCE_MIP_ENABLE 1
+#if PERMUTATION_TRAVERSAL == 1
+	#define VOXEL_TRAVERSAL_TYPE VOXEL_TRAVERSAL_LINEAR_MIPMAP
+#else
+	#define VOXEL_TRAVERSAL_TYPE VOXEL_TRAVERSAL_LINEAR
+#endif
+#include "HairStrandsVoxelPageTraversal.ush"
+
+float2 OutputResolution; 
+uint bDrawPage;
+int ForcedMipLevel;
+uint MacroGroupId;
+uint MacroGroupCount;
+uint MaxTotalPageIndexCount;
+
+RWTexture2D<float4> OutputTexture;
+Buffer<uint> TotalValidPageCounter;
+
+#define VOXEL_BLOCK 0 
+#define VOXEL_SMOOTH 1
+#define VOXEL_FILTERING VOXEL_BLOCK
+
+uint GetVoxelPageIndex(uint3 InPageIndexCoord, uint3 InPageIndexResolution, uint InLinearPageIndexOffset)
+{
+	const uint LinearPageIndexCoord = CoordToIndex(InPageIndexCoord, InPageIndexResolution, InLinearPageIndexOffset);
+	return VirtualVoxel.PageIndexBuffer.Load(LinearPageIndexCoord);
+}
+
+uint2 GetVoxelPageIndexOccupancy(uint3 InPageIndexCoord, uint3 InPageIndexResolution, uint InLinearPageIndexOffset)
+{
+	const uint LinearPageIndexCoord = CoordToIndex(InPageIndexCoord, InPageIndexResolution, InLinearPageIndexOffset);
+	return VirtualVoxel.PageIndexOccupancyBuffer.Load(LinearPageIndexCoord);
+}
+
+FFontColor GetOccupancyColor(float In)
+{
+	float3 Color = lerp(float3(0, 1, 0), float3(1, 0, 0), saturate(In));
+	return InitFontColor(Color);
+}
+
+[numthreads(8, 8, 1)]
+void MainCS(uint3 DispatchThreadId : SV_DispatchThreadID)
+{
+	const float2 PixelCoord = DispatchThreadId.xy + 0.5f;
+	const float2 UV = PixelCoord / float2(OutputResolution);
+
+	const FVirtualVoxelNodeDesc NodeDesc = UnpackVoxelNode(VirtualVoxel.NodeDescBuffer[MacroGroupId], VirtualVoxel.PageResolution);
+	const bool bIsNodeValid = all(NodeDesc.PageIndexResolution != 0);
+
+
+	const bool bPrint = all(uint2(PixelCoord) == 0);
+
+	// Draw macro group AABB & debug infos
+	if (bPrint)
+	{
+		const uint TotalPageIndexCount = MaxTotalPageIndexCount;
+		const uint TotalPageCount = VirtualVoxel.PageCountResolution.x * VirtualVoxel.PageCountResolution.y * VirtualVoxel.PageCountResolution.z;
+		const uint TotalAllocatedValidPageCount = TotalValidPageCounter.Load(0);
+		uint TotalAllocatedPageCount = 0;
+		for (uint MacroGroupIt = 0; MacroGroupIt < MacroGroupCount; ++MacroGroupIt)
+		{
+			const FPackedVirtualVoxelNodeDesc LocalPackedNode = VirtualVoxel.NodeDescBuffer[MacroGroupIt];
+			const FVirtualVoxelNodeDesc LocalNodeDesc = UnpackVoxelNode(LocalPackedNode, VirtualVoxel.PageResolution);
+			TotalAllocatedPageCount += LocalNodeDesc.PageIndexResolution.x * LocalNodeDesc.PageIndexResolution.y * LocalNodeDesc.PageIndexResolution.z;
+		}
+
+		AddAABBTWS(NodeDesc.TranslatedWorldMinAABB, NodeDesc.TranslatedWorldMaxAABB, ColorOrange);
+
+		// General information
+		// General
+		//  Voxel World size        : 
+		//  Page Resolution         :
+		//  Page Count Resolution   :
+		//  Page Texture Resolution :
+		//
+		//  Total Allocatled Page Index   :
+		//  Total Page Index              :
+		//  Total Page Index Occupancy    :
+		//
+		//  Total Allocated Page          :
+		//  Total Page                    :
+		//  Total Page Occupancy          :
+		if (MacroGroupId == 0)
+		{
+			FShaderPrintContext Context = InitShaderPrintContext(true, uint2(50, 100));
+
+			// Adaptive voxels
+			Print(Context, TEXT("Alloc. Feedback  "), FontOrange);
+			if (VirtualVoxel.AllocationFeedbackEnable)
+			{
+				Print(Context, TEXT("True"), FontGreen);
+			}
+			else
+			{
+				Print(Context, TEXT("False"), FontRed);
+			}
+			Newline(Context);
+
+			// Actual voxel/pages size & resolution
+			Print(Context, TEXT("GPU Curr. Voxel Size        "), FontOrange);
+			Print(Context, VirtualVoxel.CurrGPUMinVoxelSize[0], FontEmerald);
+			Newline(Context);
+
+			Print(Context, TEXT("GPU Next. Voxel Size        "));
+			Print(Context, VirtualVoxel.NextGPUMinVoxelSize[0]);
+			Newline(Context);
+
+			Print(Context, TEXT("CPU Min. Voxel Size         "));
+			Print(Context, VirtualVoxel.CPUMinVoxelWorldSize);
+			Newline(Context);
+			
+			Print(Context, TEXT("Page Res                    "));
+			Print(Context, VirtualVoxel.PageResolution);
+			Newline(Context);
+
+			Print(Context, TEXT("Page Count Res              "));
+			Print(Context, VirtualVoxel.PageCountResolution.x);
+			Print(Context, VirtualVoxel.PageCountResolution.y);
+			Print(Context, VirtualVoxel.PageCountResolution.z);
+			Newline(Context);
+
+			Print(Context, TEXT("Page Texture Res            "), FontEmerald);
+			Print(Context, VirtualVoxel.PageTextureResolution.x, FontEmerald);
+			Print(Context, VirtualVoxel.PageTextureResolution.y, FontEmerald);
+			Print(Context, VirtualVoxel.PageTextureResolution.z, FontEmerald);
+			Newline(Context);
+
+			// Total page indices stats
+			Newline(Context);
+
+			Print(Context, TEXT("Total Allocated Page Index  "));
+			Print(Context, TotalAllocatedPageCount);
+			Newline(Context);
+
+			Print(Context, TEXT("Total Page Index            "));
+			Print(Context, TotalPageIndexCount);
+			Newline(Context);
+
+			const float PageIndexOccupancy = TotalAllocatedPageCount / float(TotalPageIndexCount);
+			Print(Context, TEXT("Page Index Occupancy        "), FontOrange);
+			Print(Context, PageIndexOccupancy * 100.f, GetOccupancyColor(PageIndexOccupancy));
+			Newline(Context);
+
+			// Total pages stats
+			Newline(Context);
+
+			Print(Context, TEXT("Total Allocated Valid Page  "));
+			Print(Context, TotalAllocatedValidPageCount);
+			Newline(Context);
+
+			Print(Context, TEXT("Total Page                  "));
+			Print(Context, TotalPageCount);
+			Newline(Context);
+
+			const float PageOccupancy = TotalAllocatedValidPageCount / float(TotalPageCount);
+			Print(Context, TEXT("Page Occupancy              "), FontOrange);
+			Print(Context, PageOccupancy * 100.f, GetOccupancyColor(PageOccupancy));
+			Newline(Context);
+
+			Print(Context, TEXT("Translated world offset     "), FontOrange);
+			Print(Context, VirtualVoxel.TranslatedWorldOffset);
+			Newline(Context);
+		}
+
+		// MacroGroupId XX
+		//  Page Index Resolution: X Y Z (Linear/Total - %) 
+		//  Virtual Resolution   : X Y Z
+		//  Page Index Offset    : O
+		if (MacroGroupId == 0)
+		{
+			FShaderPrintContext Context = InitShaderPrintContext(true, uint2(50, 420));
+
+			Print(Context, TEXT("Macro Group        "), FontOrange);
+			Newline(Context);
+
+			Print(Context, TEXT(" Virtual Res X     "), FontEmerald);
+			Newline(Context);
+			Print(Context, TEXT(" Virtual Res Y     "), FontEmerald);
+			Newline(Context);
+			Print(Context, TEXT(" Virtual Res Z     "), FontEmerald);
+			Newline(Context);
+
+			Newline(Context);
+
+			Print(Context, TEXT(" Page Index Res X  "));
+			Newline(Context);
+			Print(Context, TEXT(" Page Index Res Y  "));
+			Newline(Context);
+			Print(Context, TEXT(" Page Index Res Z  "));
+			Newline(Context);
+
+			Newline(Context);
+
+			Print(Context, TEXT(" Page Index Offset "));
+			Newline(Context);
+
+			Newline(Context);
+
+			Print(Context, TEXT(" Voxel World Size  "));
+			Newline(Context);
+		}
+
+		{
+			FShaderPrintContext Context = InitShaderPrintContext(true, uint2(250 + MacroGroupId * 35, 420));
+
+			Print(Context, MacroGroupId, FontOrange);
+			Newline(Context);
+
+			Print(Context, NodeDesc.VirtualResolution.x, FontEmerald);
+			Newline(Context);
+			Print(Context, NodeDesc.VirtualResolution.y, FontEmerald);
+			Newline(Context);
+			Print(Context, NodeDesc.VirtualResolution.z, FontEmerald);
+			Newline(Context);
+
+			Newline(Context);
+
+			Print(Context, NodeDesc.PageIndexResolution.x, FontYellow);
+			Newline(Context);
+			Print(Context, NodeDesc.PageIndexResolution.y, FontYellow);
+			Newline(Context);
+			Print(Context, NodeDesc.PageIndexResolution.z, FontYellow);
+			Newline(Context);
+
+			Newline(Context);
+
+			// Offset of 1 line for every other voxel size, to avoid super-position with large float number
+			const bool bAddLine = (MacroGroupId & 1) == 1;
+
+			if (bAddLine) { Newline(Context); }
+
+			Print(Context, NodeDesc.PageIndexOffset, FontSilver);
+			Newline(Context);
+
+			Newline(Context);
+
+			Print(Context, NodeDesc.VoxelWorldSize, FontSilver);
+			Newline(Context);
+		}
+	}
+
+	bool bIsValid = false;
+	float3 HitColor = 0;
+	if (bDrawPage > 0)
+	{
+		// Reverse-Z
+		const float3 O = PrimaryView.TranslatedWorldCameraOrigin;
+		const float3 E = ReconstructTranslatedWorldPositionFromDepth(UV, ConvertFromDeviceZ(0.0001f));
+		const float3 D = (E - O);
+
+		float PixelRadius = -1;
+		const float2 HitT = LineBoxIntersect(O, E, NodeDesc.TranslatedWorldMinAABB, NodeDesc.TranslatedWorldMaxAABB);
+		float AccDensity = 0;
+		if (HitT.x < HitT.y)
+		{
+			// Use different offseting to see either: 
+			// * Individual voxel (for debugging mainly)
+			// * Smoothed voxel (as used for transmittance computation)
+		#if VOXEL_FILTERING == VOXEL_BLOCK
+			const float3 Offset = 0;
+			const float Dither = InterleavedGradientNoise(PixelCoord.xy, 1);
+		#endif
+		#if VOXEL_FILTERING == VOXEL_SMOOTH
+			const float3 Offset = float3(InterleavedGradientNoise(PixelCoord.xy, 1), InterleavedGradientNoise(PixelCoord.xy, 2), InterleavedGradientNoise(PixelCoord.xy, 3));
+			const float Dither = 0;
+		#endif
+
+
+			// Intersect the indirection table to see if a brick is empty or not
+			const uint MaxStep = 256;
+			const float Delta = (HitT.y-HitT.x)/float(MaxStep);
+			for (uint StepIt=0; StepIt<MaxStep; ++StepIt)
+			{			
+				const float3 HitP = Offset + O + ((Dither+StepIt)*Delta + HitT.x) * D;
+				const int3 PageIndexCoord = PositionToCoord(HitP, NodeDesc.TranslatedWorldMinAABB, NodeDesc.TranslatedWorldMaxAABB, NodeDesc.PageIndexResolution);
+				const uint PageIndex = GetVoxelPageIndex(PageIndexCoord, NodeDesc.PageIndexResolution, NodeDesc.PageIndexOffset);
+
+				bool bHasValidData = true;
+				if (bDrawPage == 2)
+				{
+					const uint2 PageIndexOccuancy = GetVoxelPageIndexOccupancy(PageIndexCoord, NodeDesc.PageIndexResolution, NodeDesc.PageIndexOffset);
+					const int3 VoxelCoord = PositionToCoord(HitP, NodeDesc.TranslatedWorldMinAABB, NodeDesc.TranslatedWorldMaxAABB, NodeDesc.VirtualResolution);
+					const int3 VoxelBasePageCoord = PageIndexCoord * VirtualVoxel.PageResolution;
+					const int3 VoxelRelativeCoord = VoxelCoord - VoxelBasePageCoord;
+					bHasValidData = DoesVoxelContainData(PageIndexOccuancy, VoxelRelativeCoord);
+				}
+
+				if (PageIndex != INVALID_VOXEL_PAGE_INDEX && bHasValidData)
+				{
+					const uint3 PageCoord = IndexToCoord(PageIndex, VirtualVoxel.PageCountResolution);
+
+					AccDensity = 1;
+					bIsValid = true;
+					HitColor = PageCoord / 12.f;
+					break;
+				}
+			}
+		}
+	}
+	else 
+	{
+		FHairTraversalResult TraversalResult = (FHairTraversalResult)0;
+
+		const float DistanceThreshold = 10000;
+		const bool bDebugEnabled = false;
+		const float3 SampleRandom = GetHairVoxelJitter(PixelCoord, View.StateFrameIndexMod8, VirtualVoxel.JitterMode);
+
+		const float3 TranslatedWP0 = PrimaryView.TranslatedWorldCameraOrigin;
+		const float3 E = ReconstructTranslatedWorldPositionFromDepth(UV, ConvertFromDeviceZ(0.0001f));
+		const float3 TranslatedWP1 = normalize(E - TranslatedWP0) * DistanceThreshold + TranslatedWP0;
+
+		FVirtualVoxelCommonDesc CommonDesc;
+		CommonDesc.PageCountResolution	= VirtualVoxel.PageCountResolution;
+		CommonDesc.PageTextureResolution= VirtualVoxel.PageTextureResolution;
+		CommonDesc.PageResolution		= VirtualVoxel.PageResolution;
+		CommonDesc.PageResolutionLog2	= VirtualVoxel.PageResolutionLog2;
+
+		if (bIsNodeValid)
+		{
+			FHairTraversalSettings TraversalSettings = InitHairTraversalSettings();
+			TraversalSettings.DensityScale		= VirtualVoxel.DensityScale;
+			TraversalSettings.CountThreshold	= GetOpaqueVoxelValue();
+			TraversalSettings.DistanceThreshold = DistanceThreshold;
+			TraversalSettings.bDebugEnabled		= bDebugEnabled;
+			TraversalSettings.SteppingScale		= VirtualVoxel.SteppingScale_Shadow;
+			TraversalSettings.Random			= SampleRandom;
+			TraversalSettings.TanConeAngle		= 0;
+			TraversalSettings.bIsPrimaryRay		= true;
+			TraversalSettings.PixelRadius		= -1;
+			TraversalSettings.ForcedMip			= ForcedMipLevel;
+
+			TraversalResult = ComputeHairCountVirtualVoxel(
+				TranslatedWP0,
+				TranslatedWP1,
+				CommonDesc,
+				NodeDesc,
+				VirtualVoxel.PageIndexBuffer,
+				VirtualVoxel.PageIndexOccupancyBuffer,
+				VirtualVoxel.PageTexture,
+				TraversalSettings);
+
+			HitColor = GetHSVDebugColor(saturate(TraversalResult.HairCount / 16));
+			bIsValid = TraversalResult.HairCount > 0;
+		}
+
+		const bool bIsPixelEnable = all(int2(PixelCoord) == GetCursorPos());
+		if (bIsPixelEnable && bIsValid)
+		{
+			FShaderPrintContext Context = InitShaderPrintContext(true, uint2(50, 50));
+
+			Print(Context, TEXT("X, Y "));
+			Print(Context, GetCursorPos());
+			Newline(Context);
+
+			Print(Context, TEXT("T.Count "));
+			Print(Context, TraversalResult.HairCount);
+			Newline(Context);
+			
+			Print(Context, TEXT("T.Visibility "));
+			Print(Context, TraversalResult.Visibility);
+			Newline(Context);
+		}
+	}
+
+	if (bIsValid)
+	{
+		OutputTexture[DispatchThreadId.xy] = float4(HitColor, 1);
+	}
+}
+
+#endif // SHADER_VOXEL_DEBUG
\ No newline at end of file
diff --git a/Engine/Shaders/Private/HairStrands/HairStrandsVoxelDepthInjection.usf b/Engine/Shaders/Private/HairStrands/HairStrandsVoxelDepthInjection.usf
deleted file mode 100644
diff --git a/Engine/Shaders/Private/HairStrands/HairStrandsVoxelMip.usf b/Engine/Shaders/Private/HairStrands/HairStrandsVoxelMip.usf
deleted file mode 100644
diff --git a/Engine/Shaders/Private/HairStrands/HairStrandsVoxelOpaque.usf b/Engine/Shaders/Private/HairStrands/HairStrandsVoxelOpaque.usf
deleted file mode 100644
diff --git a/Engine/Shaders/Private/HairStrands/HairStrandsVoxelPageAllocation.usf b/Engine/Shaders/Private/HairStrands/HairStrandsVoxelPageAllocation.usf
deleted file mode 100644
diff --git a/Engine/Shaders/Private/HairStrands/HairStrandsVoxelPageRayMarching.usf b/Engine/Shaders/Private/HairStrands/HairStrandsVoxelPageRayMarching.usf
deleted file mode 100644
diff --git a/Engine/Shaders/Private/HairStrands/HairStrandsVoxelRasterCompute.usf b/Engine/Shaders/Private/HairStrands/HairStrandsVoxelRasterCompute.usf
--- a/Engine/Shaders/Private/HairStrands/HairStrandsVoxelRasterCompute.usf
+++ b/Engine/Shaders/Private/HairStrands/HairStrandsVoxelRasterCompute.usf
@@ -3,11 +3,924 @@
 #define HAIR_STRANDS_PARAMETERS 1
 
 #include "../Common.ush"
+#include "../Common.ush"
+#include "../Matrices.ush"
+#include "../SceneTextureParameters.ush"
+#include "../PositionReconstructionCommon.ush" 
+#include "../DeferredShadingCommon.ush"
+#include "../ShaderPrint.ush"
+#include "HairStrandsAABBCommon.ush"
+#include "HairStrandsVisibilityCommon.ush"
 #include "HairStrandsClusterCommon.ush"
 #include "HairStrandsVertexFactoryCommon.ush"
 #include "HairStrandsVoxelPageCommon.ush" 
+#include "HairStrandsDeepShadowCommonStruct.ush"
+
+///////////////////////////////////////////////////////////////////////////////////////////////////
+// Page allocation
+
+#if SHADER_ALLOCATEPAGEINDEX 
+
+float				CPUPageWorldSize;
+float				CPUVoxelWorldSize;
+uint				bUseCPUVoxelWorldSize;	// When adaptive voxel size is disabled, we use CPU voxel size value
+uint				TotalPageIndexCount; 	// This is the max page index count;
+uint				PageResolution; 		// Resolution of a page
+uint				MacroGroupCount;
+uint				IndirectDispatchGroupSize;
+uint				bDoesMacroGroupSupportVoxelization;
+
+// For testing parity with CPU version
+float3				CPU_TranslatedWorldMinAABB;
+float3				CPU_TranslatedWorldMaxAABB;
+int3				CPU_PageIndexResolution;
+uint				CPU_bUseCPUData;
+
+Buffer<float>								GPUVoxelWorldSize;
+Buffer<float>								MacroGroupVoxelSizeBuffer;
+Buffer<int>									MacroGroupAABBBuffer;
+RWBuffer<int>								MacroGroupVoxelAlignedAABBBuffer;
+RWBuffer<uint4>								OutPageIndexResolutionAndOffsetBuffer;
+RWBuffer<uint>								OutPageIndexAllocationIndirectBufferArgs;
+RWStructuredBuffer<FVoxelizationViewInfo>	OutVoxelizationViewInfoBuffer;
+
+#ifndef MAX_HAIR_MACROGROUP_COUNT
+#error MAX_HAIR_MACROGROUP_COUNT needs to be defined
+#endif
+
+#define INVALID_OFFSET 0xFFFFFFFF
+
+groupshared uint PageIndexOffsets[MAX_HAIR_MACROGROUP_COUNT];
+
+// This code assume we have less than 32 macro group (which fit into a single CU/SM)
+[numthreads(MAX_HAIR_MACROGROUP_COUNT, 1, 1)]
+void AllocatePageIndex(uint2 DispatchThreadId : SV_DispatchThreadID)
+{
+	const uint MacroGroupId = DispatchThreadId.x;
+
+	FHairAABB Bound = InitHairAABB();
+	float PageWorldSize = CPUPageWorldSize;
+	bool bIsValid = MacroGroupId < MacroGroupCount;
+	if (bIsValid)
+	{
+		const bool bSupportVoxelization = (bDoesMacroGroupSupportVoxelization >> MacroGroupId) & 0x1;
+		Bound = ReadHairAABB(MacroGroupId, MacroGroupAABBBuffer);
+		const float VoxelWorldSize = bUseCPUVoxelWorldSize ? CPUVoxelWorldSize : max(GPUVoxelWorldSize[0], MacroGroupVoxelSizeBuffer[MacroGroupId]);
+		PageWorldSize = VoxelWorldSize * PageResolution;
+
+		if (any(Bound.Min > Bound.Max) || !bSupportVoxelization)
+		{
+			Bound.Min = 0;
+			Bound.Max = 0;
+			bIsValid = false;
+		}
+	}
+
+	// Page index allocation
+	int3 PageIndexResolution = 0;
+	{
+		// Snap the max AABB to the voxel size.
+
+		// The contents of MacroGroupAABBBuffer (tight fitting AABBs) and MacroGroupVoxelAlignedAABBBuffer diverge here
+		// because the macro group AABBs for voxelization need to be snapped to the voxel page boundary.
+
+		// Allocate enough pages to cover the AABB, where page (0,0,0) origin sit on MinAABB.
+		float3 MacroGroupSize = 0;
+		if (bIsValid)
+		{
+			MacroGroupSize = Bound.Max - Bound.Min;
+			PageIndexResolution = ceil(MacroGroupSize / PageWorldSize);
+			MacroGroupSize = PageIndexResolution * PageWorldSize;
+			Bound.Max = MacroGroupSize + Bound.Min;
+		}
+
+		const uint TotalPageIndex = PageIndexResolution.x * PageIndexResolution.y * PageIndexResolution.z;
+		PageIndexOffsets[MacroGroupId] = TotalPageIndex;
+		GroupMemoryBarrierWithGroupSync();
+
+		// Postfix sum to instance group are always ordered by index
+		if (DispatchThreadId.x == 0)
+		{
+			bool bValidAllocation = true;
+			uint PageIndexOffset = 0;
+			for (uint LocalMacroGroupId = 0; LocalMacroGroupId < MacroGroupCount; ++LocalMacroGroupId)
+			{
+				const uint PageCount = PageIndexOffsets[LocalMacroGroupId];
+				bValidAllocation = bValidAllocation && (PageIndexOffset + PageCount <= TotalPageIndexCount);
+				PageIndexOffsets[LocalMacroGroupId] = bValidAllocation ? PageIndexOffset : INVALID_OFFSET;
+				PageIndexOffset += PageCount;
+			}
+		}
+		GroupMemoryBarrierWithGroupSync();
+
+		const uint PageIndexOffset = PageIndexOffsets[MacroGroupId];
+		bIsValid = bIsValid && (PageIndexOffset != INVALID_OFFSET);
+		if (bIsValid)
+		{
+			OutPageIndexResolutionAndOffsetBuffer[MacroGroupId] = uint4(PageIndexResolution, PageIndexOffset);
+			WriteHairAABB(MacroGroupId, Bound, MacroGroupVoxelAlignedAABBBuffer);
+		}
+		else
+		{
+			// Clear all output if the allocation is not valid
+			OutPageIndexResolutionAndOffsetBuffer[MacroGroupId] = uint4(0, 0, 0, 0);
+			WriteDispatchIndirectArgs(OutPageIndexAllocationIndirectBufferArgs, MacroGroupId, 0, 1, 1);
+
+			FVoxelizationViewInfo ViewInfo;
+			ViewInfo.TranslatedWorldToClip = 0;
+			ViewInfo.ViewForward = 0;
+			ViewInfo.RasterResolution = 0;
+			ViewInfo.Pad0 = 0;
+			ViewInfo.Pad1 = 0;
+			OutVoxelizationViewInfoBuffer[MacroGroupId] = ViewInfo;
+		}
+	}
+
+	if (!bIsValid)
+	{
+		return;
+	}
+
+	// Use CPU data, which allow to compare the GPU driven part with CPU data
+	// This code path is only for debug purpose and work only when there is a single instance group
+	if (CPU_bUseCPUData > 0)
+	{
+		Bound.Min = CPU_TranslatedWorldMinAABB;
+		Bound.Max = CPU_TranslatedWorldMaxAABB;
+		PageIndexResolution = CPU_PageIndexResolution;
+		WriteHairAABB(MacroGroupId, Bound, MacroGroupVoxelAlignedAABBBuffer);
+	}
+
+	// Prepare indirect buffer for doing the actual page index allocation and filling the page index
+	{
+		const uint AllocatedPageIndexCount = PageIndexResolution.x * PageIndexResolution.y * PageIndexResolution.z;
+		WriteDispatchIndirectArgs(OutPageIndexAllocationIndirectBufferArgs, MacroGroupId, ceil(AllocatedPageIndexCount / float(IndirectDispatchGroupSize)), 1, 1);
+	}
+
+	// Matrix generation
+	if (bIsValid)
+	{
+		// Find the largest resolution and its dominant axis
+		uint2 RasterResolution = uint2(0, 0);
+		float3 RasterProjectionSize = float3(0, 0, 0);
+		float3 RasterDirection = float3(0, 0, 0);
+		float3 RasterUp = float3(0, 0, 0);
+		const int3 TotalVoxelResolution = PageIndexResolution * PageResolution;
+		{
+			const uint ResolutionXY = TotalVoxelResolution.x * TotalVoxelResolution.y;
+			const uint ResolutionXZ = TotalVoxelResolution.x * TotalVoxelResolution.z;
+			const uint ResolutionYZ = TotalVoxelResolution.y * TotalVoxelResolution.y;
+			if (ResolutionXY >= ResolutionXZ && ResolutionXY >= ResolutionYZ)
+			{
+				RasterResolution = int2(TotalVoxelResolution.x, TotalVoxelResolution.y);
+				RasterDirection = float3(0, 0, 1);
+				RasterUp = float3(0, 1, 0);
+
+				FHairAABB ProjRasterAABB;
+				ProjRasterAABB.Min = float3(Bound.Min.x, Bound.Min.y, Bound.Min.z);
+				ProjRasterAABB.Max = float3(Bound.Max.x, Bound.Max.y, Bound.Max.z);
+				RasterProjectionSize = ProjRasterAABB.Max - ProjRasterAABB.Min;
+			}
+			else if (ResolutionXZ >= ResolutionXY && ResolutionXZ >= ResolutionYZ)
+			{
+				RasterResolution = int2(TotalVoxelResolution.x, TotalVoxelResolution.z);
+				RasterDirection = float3(0, -1, 0);
+				RasterUp = float3(0, 0, 1);
+
+				FHairAABB ProjRasterAABB;
+				ProjRasterAABB.Min = float3(Bound.Min.x, Bound.Min.z, Bound.Min.y);
+				ProjRasterAABB.Max = float3(Bound.Max.x, Bound.Max.z, Bound.Max.y);
+				RasterProjectionSize = ProjRasterAABB.Max - ProjRasterAABB.Min;
+			}
+			else
+			{
+				RasterResolution = int2(TotalVoxelResolution.y, TotalVoxelResolution.z);
+				RasterDirection = float3(1, 0, 0);
+				RasterUp = float3(0, 0, 1);
+
+				FHairAABB ProjRasterAABB;
+				ProjRasterAABB.Min = float3(Bound.Min.y, Bound.Min.z, Bound.Min.x);
+				ProjRasterAABB.Max = float3(Bound.Max.y, Bound.Max.z, Bound.Max.x);
+				RasterProjectionSize = ProjRasterAABB.Max - ProjRasterAABB.Min;
+			}
+		}
+
+		const float3 RasterAABBSize = Bound.Max - Bound.Min;
+		const float3 RasterAABBCenter = (Bound.Max + Bound.Min) * 0.5f;
+
+		const float4x4 ProjMatrix = ReversedZOrthoMatrix(0.5f * RasterProjectionSize.x, 0.5f * RasterProjectionSize.y, 1.0f / RasterProjectionSize.z, 0);
+		const float4x4 TranslatedViewMatrix = LookAtMatrix(RasterAABBCenter - RasterDirection * RasterProjectionSize.z * 0.5f, RasterAABBCenter, RasterUp);
+
+		FVoxelizationViewInfo ViewInfo;
+		ViewInfo.TranslatedWorldToClip = mul(TranslatedViewMatrix, ProjMatrix);
+		ViewInfo.ViewForward = RasterDirection;
+		ViewInfo.RasterResolution = RasterResolution;
+		ViewInfo.Pad0 = 0;
+		ViewInfo.Pad1 = 0;
+		OutVoxelizationViewInfoBuffer[MacroGroupId] = ViewInfo;
+	}
+}
+#endif
+
+///////////////////////////////////////////////////////////////////////////////////////////////////
+// Common struct/functions
+
+#if SHADER_MARKVALID_PREPARE || SHADER_MARKVALID_INDIRECTARG || SHADER_MARKVALID_SCATTER
+struct FScatterCoord
+{
+	uint3 MinCoord;
+	uint  PartIndex;
+	uint3 Resolution;
+	uint  MacroGroupId;
+};
+
+uint2 PackScatterCoord(FScatterCoord In)
+{
+	uint2 Out;
+	Out.x =
+		((In.MinCoord.x & 0xFF)) |
+		((In.MinCoord.y & 0xFF) << 8) |
+		((In.MinCoord.z & 0xFF) << 16) |
+		((In.PartIndex & 0xFF) << 24);
+
+	Out.y =
+		((In.Resolution.x & 0xFF)) |
+		((In.Resolution.y & 0xFF) << 8) |
+		((In.Resolution.z & 0xFF) << 16) |
+		((In.MacroGroupId & 0xFF) << 24);
+
+	return Out;
+}
+
+FScatterCoord UnpackScatterCoord(uint2 In)
+{
+	FScatterCoord Out;
+	Out.MinCoord = uint3(
+		In.x & 0xFF,
+		(In.x >> 8) & 0xFF,
+		(In.x >> 16) & 0xFF);
+	Out.PartIndex = (In.x >> 24) & 0xFF;
+
+	Out.Resolution = uint3(
+		In.y & 0xFF,
+		(In.y >> 8) & 0xFF,
+		(In.y >> 16) & 0xFF);
+	Out.MacroGroupId = (In.y >> 24) & 0xFF;
+
+	return Out;
+}
+
+void UpdatePackScatterCoord(inout uint2 Out, uint PartIndex)
+{
+	Out.x = (Out.x & 0x00FFFFFF) | ((PartIndex & 0xFF) << 24);
+}
+
+#define SCATTER_WORK_ITEM_COUNT 8
+#define SCATTER_GROUP_SIZE 32
+#define SCATTER_GROUP_LINE 4
+
+#endif // SHADER_MARKVALID_PREPARE || SHADER_MARKVALID_INDIRECTARG || SHADER_MARKVALID_SCATTER
+
+///////////////////////////////////////////////////////////////////////////////////////////////////
+
+#if SHADER_MARKVALID_PREPARE
+uint			MaxClusterCount;
+uint			MacroGroupId;
+uint			MaxScatterAllocationCount;
+uint			bForceDirectPageAllocation;
+
+Buffer<int>		GroupAABBsBuffer;
+Buffer<int>		ClusterAABBsBuffer;
+Buffer<int>		MacroGroupVoxelAlignedAABBBuffer;
+Buffer<uint4>	PageIndexResolutionAndOffsetBuffer;
+
+RWBuffer<uint>	OutValidPageIndexBuffer;
+RWBuffer<uint>  OutDeferredScatterCounter;
+RWBuffer<uint2> OutDeferredScatterBuffer;
+
+
+// PageIndexBuffer      is sampled with linear coordinate computed from the 3d page coordinate. VALID NODE ARE NOT COMPACTED. It contains the LINEAR PAGE INDEX (to map to the 3d volume).
+// PageIndexCoordBuffer is sampled with linear coordinate for allocated nodes. VALID NODE ARE COMPACTED. It contains the 3d page coordinate and ClustedId. Only used for opaque voxel injection.
+
+#if PERMUTATION_USE_CLUSTER
+[numthreads(SCATTER_GROUP_SIZE, 1, 1)]
+void MarkValid_PrepareCS(uint2 DispatchThreadId : SV_DispatchThreadID)
+{
+	const uint ClusterIndex = DispatchThreadId.x;
+	if (ClusterIndex >= MaxClusterCount)
+		return;
+
+	const uint BaseClusterIndex = 6 * ClusterIndex;
+
+	FHairAABB ClusterBound;
+	ClusterBound.Min.x = float(ClusterAABBsBuffer[BaseClusterIndex + 0]);
+	ClusterBound.Min.y = float(ClusterAABBsBuffer[BaseClusterIndex + 1]);
+	ClusterBound.Min.z = float(ClusterAABBsBuffer[BaseClusterIndex + 2]);
+
+	ClusterBound.Max.x = float(ClusterAABBsBuffer[BaseClusterIndex + 3]);
+	ClusterBound.Max.y = float(ClusterAABBsBuffer[BaseClusterIndex + 4]);
+	ClusterBound.Max.z = float(ClusterAABBsBuffer[BaseClusterIndex + 5]);
+
+	if (any(ClusterBound.Min >= ClusterBound.Max))
+		return;
+
+	if (any(!IsFinite(ClusterBound.Min)) || any(!IsFinite(ClusterBound.Max)))
+		return;
+
+	const uint4 PageIndexResolutionAndOffset = PageIndexResolutionAndOffsetBuffer.Load(MacroGroupId);
+	FHairAABB MacroGroupBound = ReadHairAABB(MacroGroupId, MacroGroupVoxelAlignedAABBBuffer);
+	const int3 PageIndexResolution = PageIndexResolutionAndOffset.xyz;
+	const uint PageIndexOffset = PageIndexResolutionAndOffset.w;
+
+
+	if (any(MacroGroupBound.Min >= MacroGroupBound.Max))
+		return;
+
+	if (any(!IsFinite(MacroGroupBound.Min)) || any(!IsFinite(MacroGroupBound.Max)))
+		return;
+
+	uint3 MinCoord = PositionToCoord(ClusterBound.Min, MacroGroupBound.Min, MacroGroupBound.Max, PageIndexResolution);
+	uint3 MaxCoord = PositionToCoord(ClusterBound.Max, MacroGroupBound.Min, MacroGroupBound.Max, PageIndexResolution);
+
+	uint3 PageIndexResolutionMinusOne = uint3(PageIndexResolution - 1);
+	MinCoord = clamp(MinCoord, uint3(0, 0, 0), PageIndexResolutionMinusOne);
+	MaxCoord = clamp(MaxCoord, uint3(0, 0, 0), PageIndexResolutionMinusOne);
+
+	FScatterCoord Coord;
+	Coord.MacroGroupId = MacroGroupId;
+	Coord.PartIndex = 0;
+	Coord.MinCoord = MinCoord;
+	Coord.Resolution = (MaxCoord - MinCoord) + 1;
+
+	const uint ScatterCount = Coord.Resolution.x * Coord.Resolution.y * Coord.Resolution.z;
+
+	// Arbitrary large number (e.g., 100x10x10 pages covered)
+	// This acts as guards against degenerated case, where the sim would deformed strands large position making the cluster arbitratry large.
+	if (ScatterCount > 10000)
+		return;
+
+	if (any(!IsFinite(float3(MinCoord))))
+		return;
+
+	if (any(!IsFinite(float3(MaxCoord))))
+		return;
+
+	// Find a good sweet spot
+	if (ScatterCount <= 4 || bForceDirectPageAllocation > 0)
+	{
+		for (uint z = MinCoord.z; z <= MaxCoord.z; ++z)
+			for (uint y = MinCoord.y; y <= MaxCoord.y; ++y)
+				for (uint x = MinCoord.x; x <= MaxCoord.x; ++x)
+				{
+					const uint3 PageIndexCoord = uint3(x, y, z);
+					const uint LinearPageIndexCoord = CoordToIndex(PageIndexCoord, PageIndexResolution, PageIndexOffset);
+					InterlockedOr(OutValidPageIndexBuffer[LinearPageIndexCoord], 1u);
+				}
+	}
+	else
+	{
+		const uint WorkItemCount = (ScatterCount + SCATTER_WORK_ITEM_COUNT - 1) / SCATTER_WORK_ITEM_COUNT;
+
+		uint AllocationOffset = 0;
+		InterlockedAdd(OutDeferredScatterCounter[0], WorkItemCount, AllocationOffset);
+		if (WorkItemCount + AllocationOffset < MaxScatterAllocationCount)
+		{
+			uint2 PackedCoord = PackScatterCoord(Coord);
+			for (uint Item = 0; Item < WorkItemCount; ++Item)
+			{
+				UpdatePackScatterCoord(PackedCoord, Item);
+				OutDeferredScatterBuffer[AllocationOffset + Item] = PackedCoord;
+			}
+		}
+	}
+}
+#else // PERMUTATION_USE_CLUSTER
+[numthreads(8, 8, 8)]
+void MarkValid_PrepareCS(uint3 DispatchThreadId : SV_DispatchThreadID)
+{
+	const uint4 PageIndexResolutionAndOffset = PageIndexResolutionAndOffsetBuffer.Load(MacroGroupId);
+	const FHairAABB MacroGroupBound = ReadHairAABB(MacroGroupId, MacroGroupVoxelAlignedAABBBuffer);
+	const int3 PageIndexResolution = PageIndexResolutionAndOffset.xyz;
+	const uint PageIndexOffset = PageIndexResolutionAndOffset.w;
+
+	const uint3 Coord = DispatchThreadId;
+	if (any(Coord >= PageIndexResolution))
+		return;
+
+	const FHairAABB GroupBound = ReadHairAABB(0, GroupAABBsBuffer);
+	const uint3 MinCoord = PositionToCoord(GroupBound.Min, MacroGroupBound.Min, MacroGroupBound.Max, PageIndexResolution);
+	const uint3 MaxCoord = PositionToCoord(GroupBound.Max, MacroGroupBound.Min, MacroGroupBound.Max, PageIndexResolution);
+
+	if (any(!IsFinite(float3(MinCoord))))
+		return;
+
+	if (any(!IsFinite(float3(MaxCoord))))
+		return;
+
+	if (all(Coord >= MinCoord) && all(Coord <= MaxCoord))
+	{
+		const uint3 PageIndexCoord = uint3(Coord.x, Coord.y, Coord.z);
+		const uint LinearPageIndexCoord = CoordToIndex(PageIndexCoord, PageIndexResolution, PageIndexOffset);
+		InterlockedOr(OutValidPageIndexBuffer[LinearPageIndexCoord], 1u);
+	}
+
+	// Fake writting to insure RDG allocate buffers even there is only Group-based allocated elements (vs. cluster-based allocated elements)
+	if (Coord.x > 65000)
+	{
+		OutDeferredScatterCounter[0] = 0;
+		OutDeferredScatterBuffer[0] = uint2(0, 0);
+	}
+}
+#endif // PERMUTATION_USE_CLUSTER
+#endif // SHADER_MARKVALID_PREPARE
+
+///////////////////////////////////////////////////////////////////////////////////////////////////
+
+#if SHADER_MARKVALID_INDIRECTARG
+Buffer<uint> DeferredScatterCounter;
+RWBuffer<uint> OutIndirectArgsBuffer;
+
+[numthreads(1, 1, 1)]
+void MarkValid_BuildIndirectArgCS(uint2 DispatchThreadId : SV_DispatchThreadID)
+{
+	const uint ScatterItemCount = DeferredScatterCounter[0] * SCATTER_WORK_ITEM_COUNT;
+	const uint DispatchCount = (ScatterItemCount + SCATTER_GROUP_SIZE - 1) / SCATTER_GROUP_SIZE;
+	const uint DispatchX = min(DispatchCount, uint(SCATTER_GROUP_LINE));
+	const uint DispatchY = (DispatchCount + SCATTER_GROUP_LINE - 1) / SCATTER_GROUP_LINE;
+
+	WriteDispatchIndirectArgs(OutIndirectArgsBuffer, 0, DispatchX, DispatchY, 1);
+}
+#endif
+
+///////////////////////////////////////////////////////////////////////////////////////////////////
+
+#if SHADER_MARKVALID_SCATTER
+Buffer<int>		IndirectBufferArgs;
+Buffer<uint4>	PageIndexResolutionAndOffsetBuffer;
+Buffer<uint>	DeferredScatterCounter;
+Buffer<uint2>	DeferredScatterBuffer;
+RWBuffer<uint>	OutValidPageIndexBuffer;
+
+[numthreads(SCATTER_GROUP_SIZE, 1, 1)]
+void MarkValid_ScatterCS(uint2 DispatchThreadId : SV_DispatchThreadID)
+{
+	const uint LinearThreadId = DispatchThreadId.x + DispatchThreadId.y * SCATTER_GROUP_SIZE * SCATTER_GROUP_LINE;
+	const uint ScatterItem = LinearThreadId;
+	const uint WorkItemIndex = LinearThreadId / SCATTER_WORK_ITEM_COUNT;
+	const uint ItemIndex = LinearThreadId % SCATTER_WORK_ITEM_COUNT;
+	const uint MaxScatterItemCount = DeferredScatterCounter[0] * SCATTER_WORK_ITEM_COUNT;
+	if (ScatterItem < MaxScatterItemCount)
+	{
+		const FScatterCoord Coord = UnpackScatterCoord(DeferredScatterBuffer[WorkItemIndex]);
+
+		const uint4 PageIndexResolutionAndOffset = PageIndexResolutionAndOffsetBuffer.Load(Coord.MacroGroupId);
+		const int3  PageIndexResolution = PageIndexResolutionAndOffset.xyz;
+		const uint  PageIndexOffset = PageIndexResolutionAndOffset.w;
+
+		const uint  ClusterPageCount = Coord.Resolution.x * Coord.Resolution.y * Coord.Resolution.z;
+		const uint  LinearVoxelIt = (Coord.PartIndex * SCATTER_WORK_ITEM_COUNT) + ItemIndex;
+		if (LinearVoxelIt < ClusterPageCount)
+		{
+			const uint3 PageIndexCoord = IndexToCoord(LinearVoxelIt, Coord.Resolution) + Coord.MinCoord;
+			const uint LinearPageIndexCoord = CoordToIndex(PageIndexCoord, PageIndexResolution, PageIndexOffset);
+			InterlockedOr(OutValidPageIndexBuffer[LinearPageIndexCoord], 1u);
+		}
+	}
+}
+#endif
+
+///////////////////////////////////////////////////////////////////////////////////////////////////
+
+#if SHADER_MARKVALID
+uint			MaxClusterCount;
+int3			CPU_PageIndexResolution;
+float3			CPU_TranslatedWorldMinAABB;
+uint			CPU_PageIndexOffset;
+float3			CPU_TranslatedWorldMaxAABB;
+uint			MacroGroupId;
+
+Buffer<int>		ClusterAABBsBuffer;
+Buffer<int>		MacroGroupVoxelAlignedAABBBuffer;
+Buffer<uint4>	PageIndexResolutionAndOffsetBuffer;
+RWBuffer<uint>	OutValidPageIndexBuffer;
+
+#define GROUP_SIZE 32
+
+// PageIndexBuffer      is sampled with linear coordinate computed from the 3d page coordinate. VALID NODE ARE NOT COMPACTED. It contains the LINEAR PAGE INDEX (to map to the 3d volume).
+// PageIndexCoordBuffer is sampled with linear coordinate for allocated nodes. VALID NODE ARE COMPACTED. It contains the 3d page coordinate and ClustedId. Only used for opaque voxel injection.
+
+[numthreads(GROUP_SIZE, 1, 1)]
+void MarkValidCS(uint2 DispatchThreadId : SV_DispatchThreadID)
+{
+	const uint ClusterIndex = DispatchThreadId.x;
+	if (ClusterIndex >= MaxClusterCount)
+		return;
+
+	const uint BaseClusterIndex = 6 * ClusterIndex;
+
+	FHairAABB ClusterBound;
+	ClusterBound.Min.x = float(ClusterAABBsBuffer[BaseClusterIndex + 0]);
+	ClusterBound.Min.y = float(ClusterAABBsBuffer[BaseClusterIndex + 1]);
+	ClusterBound.Min.z = float(ClusterAABBsBuffer[BaseClusterIndex + 2]);
+
+	ClusterBound.Max.x = float(ClusterAABBsBuffer[BaseClusterIndex + 3]);
+	ClusterBound.Max.y = float(ClusterAABBsBuffer[BaseClusterIndex + 4]);
+	ClusterBound.Max.z = float(ClusterAABBsBuffer[BaseClusterIndex + 5]);
+
+	FHairAABB MacroGroupBound;
+	int3 PageIndexResolution;
+	uint PageIndexOffset;
+
+#if PERMUTATION_GPU_DRIVEN == 1
+	const uint4 PageIndexResolutionAndOffset = PageIndexResolutionAndOffsetBuffer.Load(MacroGroupId);
+
+	MacroGroupBound = ReadHairAABB(MacroGroupId, MacroGroupVoxelAlignedAABBBuffer);
+	PageIndexResolution = PageIndexResolutionAndOffset.xyz;
+	PageIndexOffset = PageIndexResolutionAndOffset.w;
+#else
+	MacroGroupBound.Min = CPU_TranslatedWorldMinAABB;
+	MacroGroupBound.Max = CPU_TranslatedWorldMaxAABB;
+	PageIndexResolution = CPU_PageIndexResolution;
+	PageIndexOffset = CPU_PageIndexOffset;
+#endif
+
+	uint3 MinCoord = PositionToCoord(ClusterBound.Min, MacroGroupBound.Min, MacroGroupBound.Max, PageIndexResolution);
+	uint3 MaxCoord = PositionToCoord(ClusterBound.Max, MacroGroupBound.Min, MacroGroupBound.Max, PageIndexResolution);
+
+	uint3 PageIndexResolutionMinusOne = PageIndexResolution - 1;
+	MinCoord = clamp(MinCoord, uint3(0, 0, 0), PageIndexResolutionMinusOne);
+	MaxCoord = clamp(MaxCoord, uint3(0, 0, 0), PageIndexResolutionMinusOne);
+
+	// Splat cluster contribution
+	// Change this to be a bit smarter (use graphics pipe?)
+	for (uint z = MinCoord.z; z <= MaxCoord.z; ++z)
+		for (uint y = MinCoord.y; y <= MaxCoord.y; ++y)
+			for (uint x = MinCoord.x; x <= MaxCoord.x; ++x)
+			{
+				const uint3 PageIndexCoord = uint3(x, y, z);
+				const uint LinearPageIndexCoord = CoordToIndex(PageIndexCoord, PageIndexResolution, PageIndexOffset);
+				InterlockedOr(OutValidPageIndexBuffer[LinearPageIndexCoord], 1u);
+			}
+}
+#endif
+
+///////////////////////////////////////////////////////////////////////////////////////////////////
+
+#if SHADER_ALLOCATE
+uint MacroGroupId;
+uint PageCount;
+uint CPU_PageIndexCount;
+uint CPU_PageIndexOffset;
+uint3 CPU_PageIndexResolution;
+
+Buffer<uint4> PageIndexResolutionAndOffsetBuffer;
+Buffer<int> IndirectBufferArgs;
+
+RWBuffer<uint> PageIndexGlobalCounter;
+RWBuffer<uint> PageIndexBuffer;
+RWBuffer<uint> PageToPageIndexBuffer;
+RWBuffer<uint4> PageIndexCoordBuffer;
+
+groupshared uint LocalCounter;
+groupshared uint GroupBase;
+
+#define GROUP_SIZE 32
+
+[numthreads(GROUP_SIZE, 1, 1)]
+void AllocateCS(uint GroupIndex : SV_GroupIndex, uint3 DispatchThreadId : SV_DispatchThreadID)
+{
+	if (GroupIndex == 0)
+	{
+		GroupBase = 0;
+		LocalCounter = 0;
+	}
+	GroupMemoryBarrierWithGroupSync();
+
+#if PERMUTATION_GPU_DRIVEN == 1
+	const uint4 PageIndexResolutionAndOffset = PageIndexResolutionAndOffsetBuffer.Load(MacroGroupId);
+	const uint3 PageIndexResolution = PageIndexResolutionAndOffset.xyz;
+	const uint  PageIndexOffset = PageIndexResolutionAndOffset.w;
+	const uint  PageIndexCount = PageIndexResolution.x * PageIndexResolution.y * PageIndexResolution.z;
+#else
+	const uint3 PageIndexResolution = CPU_PageIndexResolution;
+	const uint  PageIndexOffset = CPU_PageIndexOffset;
+	const uint  PageIndexCount = CPU_PageIndexCount;
+#endif
+
+	const uint GridIndex = DispatchThreadId.x + PageIndexOffset;
+	bool bIsValid = false;
+	if (DispatchThreadId.x < PageIndexCount)
+	{
+		bIsValid = PageIndexBuffer[GridIndex] > 0;
+	}
+
+	uint Offset = 0;
+	if (bIsValid)
+	{
+		InterlockedAdd(LocalCounter, 1u, Offset);
+	}
+	GroupMemoryBarrierWithGroupSync();
+
+	if (GroupIndex == 0)
+	{
+		InterlockedAdd(PageIndexGlobalCounter[0], LocalCounter, GroupBase);
+	}
+	GroupMemoryBarrierWithGroupSync();
+
+	if (bIsValid)
+	{
+		const uint PageIndex = GroupBase + Offset;
+		const bool bIsAllocationValid = PageIndex < PageCount;
+
+		PageIndexBuffer[GridIndex] = bIsAllocationValid ? PageIndex : INVALID_VOXEL_PAGE_INDEX;
+		if (bIsAllocationValid)
+		{
+			PageToPageIndexBuffer[PageIndex] = GridIndex;
+		}
+
+		// Output the coordinates of the allocated page for indirect dispatch usage
+		// If the allocated failed (run out of page), then we mark the IndexCoord with a invalid GroupID
+		const uint LinearIndex = DispatchThreadId.x;
+		const uint3 PageIndexCoord = IndexToCoord(LinearIndex, PageIndexResolution);
+		const uint IndexRelativeToMacroGroupIndexOffset = PageIndex - PageIndexGlobalCounter[1];
+		PageIndexCoordBuffer[PageIndexOffset + IndexRelativeToMacroGroupIndexOffset] = uint4(PageIndexCoord, bIsAllocationValid ? MacroGroupId : INVALID_MACRO_GROUP_ID);
+	}
+	// Mark page index as invalid
+	// Insure that even if write more (due to larger dispatch count that needed), we do not stomp other instance group page index
+	else if (DispatchThreadId.x < PageIndexCount)
+	{
+		PageIndexBuffer[GridIndex] = INVALID_VOXEL_PAGE_INDEX;
+	}
+}
+#endif
+
+///////////////////////////////////////////////////////////////////////////////////////////////////
+
+#if SHADER_ADDDESC
+
+float3 CPU_TranslatedWorldMinAABB;
+uint MacroGroupId;
+float3 CPU_TranslatedWorldMaxAABB;
+uint CPU_PageIndexOffset;
+int3 CPU_PageIndexResolution;
+float CPU_VoxelWorldSize;
+uint bUseCPUVoxelWorldSize; // When adaptive voxel size is disabled, we use CPU voxel size value
+
+Buffer<float> GPU_VoxelWorldSize;
+Buffer<int> MacroGroupVoxelAlignedAABBBuffer;
+Buffer<float> MacroGroupVoxelSizeBuffer;
+Buffer<uint4> PageIndexResolutionAndOffsetBuffer;
+RWStructuredBuffer<FPackedVirtualVoxelNodeDesc> OutNodeDescBuffer;
+
+[numthreads(1, 1, 1)]
+void AddDescCS(uint GroupIndex : SV_GroupIndex, uint3 DispatchThreadId : SV_DispatchThreadID)
+{
+	FVirtualVoxelNodeDesc Node;
+
+#if PERMUTATION_GPU_DRIVEN == 1
+	const uint4 PageIndexResolutionAndOffset = PageIndexResolutionAndOffsetBuffer.Load(MacroGroupId);
+	const FHairAABB TranslatedWorldBound = ReadHairAABB(MacroGroupId, MacroGroupVoxelAlignedAABBBuffer);
+	const float VoxelWorldSize = MacroGroupVoxelSizeBuffer[MacroGroupId];
+
+	Node.TranslatedWorldMinAABB = TranslatedWorldBound.Min;
+	Node.TranslatedWorldMaxAABB = TranslatedWorldBound.Max;
+	Node.PageIndexResolution = PageIndexResolutionAndOffset.xyz;
+	Node.PageIndexOffset = PageIndexResolutionAndOffset.w;
+	Node.VoxelWorldSize = bUseCPUVoxelWorldSize ? CPU_VoxelWorldSize : max(GPU_VoxelWorldSize[0], VoxelWorldSize);
+#else
+	Node.TranslatedWorldMinAABB = CPU_TranslatedWorldMinAABB;
+	Node.TranslatedWorldMaxAABB = CPU_TranslatedWorldMaxAABB;
+	Node.PageIndexResolution = CPU_PageIndexResolution;
+	Node.PageIndexOffset = CPU_PageIndexOffset;
+	Node.VoxelWorldSize = CPU_VoxelWorldSize;
+#endif
+
+	FPackedVirtualVoxelNodeDesc PackedNode = PackVoxelNode(Node);
+	OutNodeDescBuffer[MacroGroupId] = PackedNode;
+}
+#endif
+
+///////////////////////////////////////////////////////////////////////////////////////////////////
+
+#if SHADER_ADDINDIRECTBUFFER
+
+uint PageResolution;
+uint MacroGroupId;
+uint IndirectGroupSize;
+
+RWBuffer<uint> OutPageIndexGlobalCounter;
+RWBuffer<uint> OutIndirectArgsBuffer;
+RWBuffer<uint> OutTotalRequestedPageAllocationBuffer;
+
+[numthreads(1, 1, 1)]
+void AddIndirectBufferCS(uint GroupIndex : SV_GroupIndex, uint3 DispatchThreadId : SV_DispatchThreadID)
+{
+	const uint TotalAllocatedPageCount = OutPageIndexGlobalCounter[0];
+	const uint PrevTotalAllocatedPageCount = OutPageIndexGlobalCounter[1];
+	const uint AllocatedPageCount = TotalAllocatedPageCount - PrevTotalAllocatedPageCount;
+	const uint VoxelCountPerPage = PageResolution * PageResolution * PageResolution;
+
+	const uint MaxGroupCountPerDim = 65535;
+	const uint DispatchX = AllocatedPageCount % MaxGroupCountPerDim;
+	const uint DispatchY = ceil(AllocatedPageCount / float(MaxGroupCountPerDim));
+	WriteDispatchIndirectArgs(OutIndirectArgsBuffer, MacroGroupId, DispatchX, DispatchY, ceil(VoxelCountPerPage / float(IndirectGroupSize)));
+
+	// Copy current total to previous total
+	OutPageIndexGlobalCounter[1] = TotalAllocatedPageCount;
+
+	// The value will be written several times, but only the value written during the last macro group will give the final page count allocated/requested
+	OutTotalRequestedPageAllocationBuffer[0] = TotalAllocatedPageCount;
+}
+
+#endif
+
+///////////////////////////////////////////////////////////////////////////////////////////////////
+
+#define IND_CLEAR_GROUP_SIZE_Z 64
+
+#if SHADER_INDPAGECLEARBUFFERGEN
+
+uint PageResolution;
+Buffer<uint> PageIndexGlobalCounter;
+RWBuffer<uint> OutIndirectArgsBuffer;
+
+[numthreads(1, 1, 1)]
+void VoxelIndPageClearBufferGenCS(uint GroupIndex : SV_GroupIndex, uint3 DispatchThreadId : SV_DispatchThreadID)
+{
+	const uint TotalAllocatedPageCount = PageIndexGlobalCounter[0];
+	const uint VoxelCountPerPage = PageResolution * PageResolution * PageResolution;
+
+	const uint MaxGroupCountPerDim = 65535;
+	const uint DispatchX = TotalAllocatedPageCount % MaxGroupCountPerDim;
+	const uint DispatchY = ceil(TotalAllocatedPageCount / float(MaxGroupCountPerDim));
+
+	WriteDispatchIndirectArgs(OutIndirectArgsBuffer, 0, DispatchX, DispatchY, ceil(VoxelCountPerPage / float(IND_CLEAR_GROUP_SIZE_Z)));
+}
+
+#endif
+
+///////////////////////////////////////////////////////////////////////////////////////////////////
+
+#if SHADER_INDPAGECLEAR
+
+uint				VirtualVoxelParams_PageResolution;
+int3				VirtualVoxelParams_PageCountResolution;
+Buffer<uint4>		VirtualVoxelParams_PageIndexCoordBuffer;
+
+RWTexture3D<uint>	OutPageTexture;
+
+[numthreads(1, 1, IND_CLEAR_GROUP_SIZE_Z)]
+void VoxelIndPageClearCS(uint3 DispatchThreadId : SV_DispatchThreadID, uint3 GroupId : SV_GroupID, uint3 GroupThreadId : SV_GroupThreadID)
+{
+	const uint  LinearVoxelCoord = DispatchThreadId.z;
+
+	const uint VoxelCountPerPage = VirtualVoxelParams_PageResolution * VirtualVoxelParams_PageResolution * VirtualVoxelParams_PageResolution;
+	//	if (AllocatedPageIndex < TotalAllocatedPageCount)		// Not needed today because the XY size exactly matches the range (page count) we need to clear.
+	if (LinearVoxelCoord < VoxelCountPerPage)
+	{
+		const uint  MaxDispatchCountPerDim = 65535u;
+		const uint  AllocatedPageIndex = DispatchThreadId.x + DispatchThreadId.y * MaxDispatchCountPerDim;
+		const uint3 VoxelCoordOffset = IndexToCoord(LinearVoxelCoord, VirtualVoxelParams_PageResolution.xxx);
+
+		const uint PageIndex = AllocatedPageIndex; // PageIndexBuffer is not needed, we already know those tiles are allocated linearly in 3D within OutPageTexture.
+
+		const uint3 PageCoord = IndexToCoord(PageIndex, VirtualVoxelParams_PageCountResolution);
+		const int3 VoxelPageBase = PageCoord * VirtualVoxelParams_PageResolution;
+		const int3 VoxelCoord = VoxelPageBase + VoxelCoordOffset;
+
+		OutPageTexture[VoxelCoord] = 0;
+	}
+}
+
+#endif
+
+///////////////////////////////////////////////////////////////////////////////////////////////////
+
+#if SHADER_ADAPTIVE_FEEDBACK
+
+#define DEBUG_ENABLE 0
+
+#if DEBUG_ENABLE
+#include "../ShaderPrint.ush"
+#endif
+
+uint  CPUAllocatedPageCount;
+float CPUMinVoxelWorldSize;
+float AdaptiveCorrectionThreshold;
+float AdaptiveCorrectionSpeed;
+
+Buffer<uint> TotalRequestedPageAllocationBuffer;
+Buffer<float> CurrGPUMinVoxelWorldSize;
+RWBuffer<float> NextGPUMinVoxelWorldSize;
+
+float RoundHairVoxelSize(float In)
+{
+	// Round voxel size to 0.01f to avoid oscillation issue
+	return floor(In * 1000.f + 0.5f) * 0.001f;
+}
+
+[numthreads(1, 1, 1)]
+void FeedbackCS(uint3 DispatchThreadId : SV_DispatchThreadID)
+{
+	const float CurrVoxelWorldSize = RoundHairVoxelSize(CurrGPUMinVoxelWorldSize[0]);
+
+	// Voxel pages are represent a volume. To derive a better estimate of the ratio by which voxel size needs to be scale, 
+	// compute the cubic root of this ratio.
+	//
+	// AllocatedPage   AllocatedRes^3
+	// ------------- = --------------  = VolumeRatio = LinearRatio^3
+	//    MaxPage          MaxRes^3
+
+	// Ratio used for predicting voxel size increase
+	const uint GPUAllocatedPageCount = TotalRequestedPageAllocationBuffer[0];
+	const float VolumeRatio = float(GPUAllocatedPageCount) / float(CPUAllocatedPageCount);
+	const float LinearRatio = pow(VolumeRatio, 1.f / 3.f);
+
+	// Ratio used for predicting voxel size decrease (i.e. when requested allocation fit, 
+	// but the voxel size does not match the (more precise) target).
+	// In this case, we add a threshold/margin to to the target, so that there is no oscillation.
+	const float VolumeRatio_Thres = float(GPUAllocatedPageCount) / float(CPUAllocatedPageCount * AdaptiveCorrectionThreshold);
+	const float LinearRatio_Thres = pow(VolumeRatio_Thres, 1.f / 3.f);
+
+	// If the page pool is not large enough increase voxel size
+	float NextVoxelWorldSize = CPUMinVoxelWorldSize;
+	if (GPUAllocatedPageCount > CPUAllocatedPageCount)
+	{
+		//NextVoxelWorldSize = CurrVoxelWorldSize * LinearRatio;
+		NextVoxelWorldSize = CurrVoxelWorldSize * LinearRatio_Thres;
+	}
+	// If the page pool is large enough but the voxel are larger than the requested size decrease voxel size
+	else if (GPUAllocatedPageCount < CPUAllocatedPageCount && CurrVoxelWorldSize > CPUMinVoxelWorldSize)
+	{
+		const float TargetVoxelWorldSize = CurrVoxelWorldSize * LinearRatio_Thres;
+		NextVoxelWorldSize = max(CPUMinVoxelWorldSize, lerp(CurrVoxelWorldSize, TargetVoxelWorldSize, AdaptiveCorrectionSpeed));
+	}
+	//else if (GPUAllocatedPageCount > CPUAllocatedPageCount * AdaptiveCorrectionThreshold)
+	//{
+	//	const float TargetVoxelWorldSize = CurrVoxelWorldSize * LinearRatio_Thres;
+	//	NextVoxelWorldSize = max(CPUMinVoxelWorldSize, lerp(CurrVoxelWorldSize, TargetVoxelWorldSize, AdaptiveCorrectionSpeed));
+	//}
+	else
+	{
+		NextVoxelWorldSize = CPUMinVoxelWorldSize;
+	}
+
+	// Clamp voxel size into a reasonable amount (e.g. 0.1mm - 100mm)
+	const float ClampMinVoxelWorldSize = 0.01f;
+	const float ClampMaxVoxelWorldSize = 10.0f;
+	NextVoxelWorldSize = clamp(RoundHairVoxelSize(NextVoxelWorldSize), ClampMinVoxelWorldSize, ClampMaxVoxelWorldSize);
+
+	// Debug
+#if DEBUG_ENABLE
+	FFontColor CPUColor = FontEmerald;
+	FFontColor GPUColor = FontOrange;
+	FFontColor CstColor = FontSilver;
+
+	FShaderPrintContext Context = InitShaderPrintContext(true, uint2(700, 50));
+
+	Print(Context, TEXT(" ------------------------------- "), FontSilver); Newline(Context);
+	Print(Context, TEXT("|          Allocations          |"), FontSilver); Newline(Context);
+	Print(Context, TEXT(" ------------------------------- "), FontSilver); Newline(Context);
+
+	Print(Context, TEXT("GPU Allocated      "), GPUColor);
+	Print(Context, GPUAllocatedPageCount, GPUColor);
+	Newline(Context);
+
+	Print(Context, TEXT("CPU Allocated      "), CPUColor);
+	Print(Context, CPUAllocatedPageCount, CPUColor);
+	Newline(Context);
+
+	Print(Context, TEXT("GPU Curr Min. Size "), GPUColor);
+	Print(Context, CurrVoxelWorldSize, GPUColor);
+	Newline(Context);
+
+	Print(Context, TEXT("GPU Next Min. Size "), GPUColor);
+	Print(Context, NextVoxelWorldSize, GPUColor);
+	Newline(Context);
+
+	Print(Context, TEXT("CPU Min. Size      "), CPUColor);
+	Print(Context, CPUMinVoxelWorldSize, CPUColor);
+	Newline(Context);
+
+	Print(Context, TEXT("Correction Thres.  "), CstColor);
+	Print(Context, AdaptiveCorrectionThreshold, CstColor);
+	Newline(Context);
+
+	Print(Context, TEXT("Correction Speed   "), CstColor);
+	Print(Context, AdaptiveCorrectionSpeed, CstColor);
+	Newline(Context);
+#endif
+
+	// Update state data
+	NextGPUMinVoxelWorldSize[0] = RoundHairVoxelSize(NextVoxelWorldSize);
+}
+
+#endif // SHADER_ADAPTIVE_FEEDBACK
 
 ///////////////////////////////////////////////////////////////////////////
+// Voxel Raster Compute
+
+#if SHADER_RASTERCOMPUTE
 
 uint			DispatchCountX;
 uint			MaxRasterCount;
@@ -27,9 +940,6 @@ RWTexture3D<uint> OutPageTexture;
 
 float SampleWeight;
 
-///////////////////////////////////////////////////////////////////////////
-// Define
-
 #define JITTER_ENABLE 0
 
 #if PERMUTATION_GROUP_SIZE == 64
@@ -40,8 +950,6 @@ float SampleWeight;
 	#error Unknown group size	
 #endif
 
-///////////////////////////////////////////////////////////////////////////
-// Utils
 float3 GetHairVoxelJitter(uint2 PixelCoord, uint Seed)
 {
 	return float3(
@@ -50,8 +958,6 @@ float3 GetHairVoxelJitter(uint2 PixelCoord, uint Seed)
 		InterleavedGradientNoise(PixelCoord.xy, Seed * 7901));
 }
 
-///////////////////////////////////////////////////////////////////////////
-
 [numthreads(GROUP_SIZE, 1, 1)]
 void MainCS(uint2 DispatchThreadID : SV_DispatchThreadID)
 {
@@ -208,4 +1114,476 @@ void MainCS(uint2 DispatchThreadID : SV_DispatchThreadID)
 			PreviousCoord = VoxelPageCoord;
 		}
 	}
-}
\ No newline at end of file
+}
+#endif // SHADER_RASTERCOMPUTE
+
+///////////////////////////////////////////////////////////////////////////
+// Inject opaque surface into voxels
+
+#if SHADER_INJECTOPAQUE_VIRTUALVOXEL
+uint MacroGroupId;
+float2 SceneDepthResolution;
+uint VoxelBiasCount;
+uint VoxelMarkCount;
+RWTexture3D<uint> OutPageTexture;
+
+uint		 VirtualVoxelParams_PageIndexCount;
+uint		 VirtualVoxelParams_PageResolution;
+uint3		 VirtualVoxelParams_PageCountResolution;
+Buffer<uint> VirtualVoxelParams_PageIndexBuffer;
+Buffer<uint4>VirtualVoxelParams_PageIndexCoordBuffer;
+StructuredBuffer<FPackedVirtualVoxelNodeDesc> VirtualVoxelParams_NodeDescBuffer;
+
+#define GROUP_SIZE 64
+[numthreads(1, 1, GROUP_SIZE)]
+void MainCS(uint3 DispatchThreadId : SV_DispatchThreadID)
+{
+	const uint  MaxDispatchCountPerDim = 65535u;
+	const uint  AllocatedPageIndex = DispatchThreadId.x + DispatchThreadId.y * MaxDispatchCountPerDim;
+	const uint  LinearVoxelCoord = DispatchThreadId.z;
+	const uint3 VoxelCoordOffset = IndexToCoord(LinearVoxelCoord, VirtualVoxelParams_PageResolution.xxx);
+
+	const FVirtualVoxelNodeDesc NodeDesc = UnpackVoxelNode(VirtualVoxelParams_NodeDescBuffer[MacroGroupId], VirtualVoxelParams_PageResolution);
+
+	const uint4 PageIndexCoord = VirtualVoxelParams_PageIndexCoordBuffer[NodeDesc.PageIndexOffset + AllocatedPageIndex];
+	const uint  LinearPageIndexCoord = CoordToIndex(PageIndexCoord.xyz, NodeDesc.PageIndexResolution, NodeDesc.PageIndexOffset);
+
+	// PageIndexCoord have invalid .w component if we run out of available page during the allocation
+	const bool bIsValid = LinearPageIndexCoord < VirtualVoxelParams_PageIndexCount&& PageIndexCoord.w != INVALID_MACRO_GROUP_ID;
+	if (bIsValid)
+	{
+		const uint PageIndex = VirtualVoxelParams_PageIndexBuffer.Load(LinearPageIndexCoord);
+		if (PageIndex != INVALID_VOXEL_PAGE_INDEX)
+		{
+			const uint3 VoxelCoordBase = PageIndexCoord.xyz * VirtualVoxelParams_PageResolution;
+			const uint3 VoxelCoord = VoxelCoordBase + VoxelCoordOffset;
+			const float3 TranslatedWorldPosition = VoxelCoord / float3(NodeDesc.VirtualResolution) * (NodeDesc.TranslatedWorldMaxAABB - NodeDesc.TranslatedWorldMinAABB) + NodeDesc.TranslatedWorldMinAABB;
+
+			float4 ClipPos = mul(float4(TranslatedWorldPosition, 1), PrimaryView.TranslatedWorldToClip);
+			ClipPos /= ClipPos.w;
+			const float DepthBias = VoxelBiasCount * NodeDesc.VoxelWorldSize;
+			const float VoxelDepth = ConvertFromDeviceZ(ClipPos.z) - DepthBias;
+			float2 SceneUV = float2(0.5f * (ClipPos.x + 1), 1 - 0.5f * (ClipPos.y + 1));
+			SceneUV = ViewportUVToBufferUV(SceneUV);
+			const float2 ScenePixelCoord = SceneUV * SceneDepthResolution;
+
+			const bool bIsOnScreen = SceneUV.x >= 0 && SceneUV.x < 1 && SceneUV.y >= 0 && SceneUV.y < 1;
+			if (!bIsOnScreen)
+				return;
+
+			const float ClosestDepth = ConvertFromDeviceZ(SceneDepthTexture.Load(uint3(ScenePixelCoord, 0)).x);
+			const float3 SceneTranslatedWorldPos = ReconstructTranslatedWorldPositionFromDepth(SceneUV, ClosestDepth);
+
+			const bool bIsInVolume =
+				SceneTranslatedWorldPos.x >= NodeDesc.TranslatedWorldMinAABB.x && SceneTranslatedWorldPos.x < NodeDesc.TranslatedWorldMaxAABB.x&&
+				SceneTranslatedWorldPos.y >= NodeDesc.TranslatedWorldMinAABB.y && SceneTranslatedWorldPos.y < NodeDesc.TranslatedWorldMaxAABB.y&&
+				SceneTranslatedWorldPos.z >= NodeDesc.TranslatedWorldMinAABB.z && SceneTranslatedWorldPos.z < NodeDesc.TranslatedWorldMaxAABB.z;
+
+			if (!bIsInVolume)
+				return;
+
+			// Inject opaque depth on a thin layer (Dist < DistThreshold) for avoiding weird projection
+			if (ClosestDepth < VoxelDepth && abs(ClosestDepth - VoxelDepth) < VoxelMarkCount * NodeDesc.VoxelWorldSize)
+			{
+				const uint3 VoxelPageIndexCoord = VoxelCoord / VirtualVoxelParams_PageResolution;
+				const uint3 VoxelIndexCoordBase = VoxelPageIndexCoord * VirtualVoxelParams_PageResolution;
+				const uint3 VoxelPageOffset = VoxelCoord - VoxelIndexCoordBase;
+
+				const uint3 PageCoord = IndexToCoord(PageIndex, VirtualVoxelParams_PageCountResolution);
+				const int3 VoxelPageBase = PageCoord * VirtualVoxelParams_PageResolution;
+				const int3 VoxelPageCoord = VoxelPageOffset + VoxelPageBase;
+
+				InterlockedOr(OutPageTexture[VoxelPageCoord], VOXEL_OPAQUE_ADD);
+			}
+		}
+	}
+}
+#endif // SHADER_INJECTOPAQUE_VIRTUALVOXEL
+
+///////////////////////////////////////////////////////////////////////////
+
+#if SHADER_DEPTH_INJECTION
+float4x4 CPU_TranslatedWorldToClip;
+
+float2 OutputResolution;
+uint MacroGroupId;
+uint AtlasSlotIndex;
+
+float3 LightDirection;
+uint bIsDirectional;
+float3 TranslatedLightPosition;
+uint bIsGPUDriven;
+
+StructuredBuffer<FDeepShadowViewInfo> DeepShadowViewInfoBuffer;
+
+void MainVS(
+	uint VertexId : SV_VertexID,
+	out float4 OutPosition : SV_POSITION,
+	out float3 OutTranslatedWorldPosition : WORLD_POSITION)
+{
+	const FPackedVirtualVoxelNodeDesc PackedNode = VirtualVoxel.NodeDescBuffer[MacroGroupId];
+	const FVirtualVoxelNodeDesc NodeDesc = UnpackVoxelNode(PackedNode, VirtualVoxel.PageResolution);
+
+	// Move this to an actual vertex/index buffer
+	const float3 Min = NodeDesc.TranslatedWorldMinAABB;
+	const float3 Max = NodeDesc.TranslatedWorldMaxAABB;
+
+	const float3 Center = (Min + Max) * 0.5f;
+	const float3 Extent = (Max - Min) * 0.5f;
+
+	const float3 Position0 = Center + float3(-Extent.x, -Extent.y, -Extent.z);
+	const float3 Position1 = Center + float3(+Extent.x, -Extent.y, -Extent.z);
+	const float3 Position2 = Center + float3(+Extent.x, +Extent.y, -Extent.z);
+	const float3 Position3 = Center + float3(-Extent.x, +Extent.y, -Extent.z);
+	const float3 Position4 = Center + float3(-Extent.x, -Extent.y, +Extent.z);
+	const float3 Position5 = Center + float3(+Extent.x, -Extent.y, +Extent.z);
+	const float3 Position6 = Center + float3(+Extent.x, +Extent.y, +Extent.z);
+	const float3 Position7 = Center + float3(-Extent.x, +Extent.y, +Extent.z);
+
+	float3 TranslatedWorldPosition = 0;
+	switch (VertexId)
+	{
+	case  0: TranslatedWorldPosition = Position0; break;
+	case  1: TranslatedWorldPosition = Position1; break;
+	case  2: TranslatedWorldPosition = Position2; break;
+	case  3: TranslatedWorldPosition = Position0; break;
+	case  4: TranslatedWorldPosition = Position2; break;
+	case  5: TranslatedWorldPosition = Position3; break;
+
+	case  6: TranslatedWorldPosition = Position4; break;
+	case  7: TranslatedWorldPosition = Position5; break;
+	case  8: TranslatedWorldPosition = Position6; break;
+	case  9: TranslatedWorldPosition = Position4; break;
+	case 10: TranslatedWorldPosition = Position6; break;
+	case 11: TranslatedWorldPosition = Position7; break;
+
+	case 12: TranslatedWorldPosition = Position0; break;
+	case 13: TranslatedWorldPosition = Position1; break;
+	case 14: TranslatedWorldPosition = Position5; break;
+	case 15: TranslatedWorldPosition = Position0; break;
+	case 16: TranslatedWorldPosition = Position5; break;
+	case 17: TranslatedWorldPosition = Position4; break;
+
+	case 18: TranslatedWorldPosition = Position2; break;
+	case 19: TranslatedWorldPosition = Position3; break;
+	case 20: TranslatedWorldPosition = Position7; break;
+	case 21: TranslatedWorldPosition = Position2; break;
+	case 22: TranslatedWorldPosition = Position7; break;
+	case 23: TranslatedWorldPosition = Position6; break;
+
+	case 24: TranslatedWorldPosition = Position1; break;
+	case 25: TranslatedWorldPosition = Position2; break;
+	case 26: TranslatedWorldPosition = Position6; break;
+	case 27: TranslatedWorldPosition = Position1; break;
+	case 28: TranslatedWorldPosition = Position6; break;
+	case 29: TranslatedWorldPosition = Position5; break;
+
+	case 30: TranslatedWorldPosition = Position3; break;
+	case 31: TranslatedWorldPosition = Position0; break;
+	case 32: TranslatedWorldPosition = Position4; break;
+	case 33: TranslatedWorldPosition = Position3; break;
+	case 34: TranslatedWorldPosition = Position4; break;
+	case 35: TranslatedWorldPosition = Position7; break;
+	}
+
+	float4x4 TranslatedWorldToClipMatrix = CPU_TranslatedWorldToClip;
+	if (bIsGPUDriven)
+	{
+		FDeepShadowViewInfo DeepShadowViewInfo = DeepShadowViewInfoBuffer[AtlasSlotIndex];
+		TranslatedWorldToClipMatrix = DeepShadowViewInfo.TranslatedWorldToClipScaledBiased;
+	}
+
+	OutTranslatedWorldPosition = TranslatedWorldPosition;
+	OutPosition = mul(float4(TranslatedWorldPosition, 1), TranslatedWorldToClipMatrix);
+}
+
+#define VOXEL_TRAVERSAL_DEBUG_PAGEOCCUPANCY 0
+//#define VOXEL_TRAVERSAL_TYPE VOXEL_TRAVERSAL_LINEAR_MIPMAP
+#define VOXEL_TRAVERSAL_TYPE VOXEL_TRAVERSAL_LINEAR
+#include "HairStrandsVoxelPageTraversal.ush"
+
+void MainPS(
+	in float4 InPosition : SV_POSITION,
+	in float3 InTranslatedWorldPosition : WORLD_POSITION,
+	out float OutDepth : SV_DEPTH)
+{
+	const float2 PixelCoord = InPosition.xy;
+	const float2 UV = PixelCoord / float2(OutputResolution); // todo view rect offset
+
+	const float DistanceThreshold = 1000;
+	const bool bDebugEnabled = false;
+	const float3 SampleRandom = GetHairVoxelJitter(PixelCoord, View.StateFrameIndexMod8, VirtualVoxel.JitterMode);
+
+	const float3 TracingDirection = bIsDirectional ? LightDirection : normalize(InTranslatedWorldPosition - TranslatedLightPosition);
+	const float3 TranslatedWP0 = InTranslatedWorldPosition;
+	const float3 TranslatedWP1 = InTranslatedWorldPosition + TracingDirection * DistanceThreshold;
+
+	FVirtualVoxelCommonDesc CommonDesc;
+	CommonDesc.PageCountResolution = VirtualVoxel.PageCountResolution;
+	CommonDesc.PageTextureResolution = VirtualVoxel.PageTextureResolution;
+	CommonDesc.PageResolution = VirtualVoxel.PageResolution;
+	CommonDesc.PageResolutionLog2 = VirtualVoxel.PageResolutionLog2;
+
+	const FPackedVirtualVoxelNodeDesc PackedNode = VirtualVoxel.NodeDescBuffer[MacroGroupId];
+	const FVirtualVoxelNodeDesc NodeDesc = UnpackVoxelNode(PackedNode, VirtualVoxel.PageResolution);
+
+	FHairTraversalSettings TraversalSettings = InitHairTraversalSettings();
+	TraversalSettings.DensityScale = VirtualVoxel.DensityScale;
+	TraversalSettings.CountThreshold = 0.9f; // GetOpaqueVoxelValue();
+	TraversalSettings.DistanceThreshold = DistanceThreshold;
+	TraversalSettings.bDebugEnabled = bDebugEnabled;
+	TraversalSettings.SteppingScale = VirtualVoxel.SteppingScale_Shadow;
+	TraversalSettings.Random = SampleRandom;
+	TraversalSettings.TanConeAngle = 0;
+	TraversalSettings.bIsPrimaryRay = true;
+	TraversalSettings.bUseOpaqueVisibility = true;
+	TraversalSettings.PixelRadius = -1;
+	TraversalSettings.ForcedMip = -1;
+
+	const FHairTraversalResult TraversalResult = ComputeHairCountVirtualVoxel(
+		TranslatedWP0,
+		TranslatedWP1,
+		CommonDesc,
+		NodeDesc,
+		VirtualVoxel.PageIndexBuffer,
+		VirtualVoxel.PageIndexOccupancyBuffer,
+		VirtualVoxel.PageTexture,
+		TraversalSettings);
+	bool bIsValid = TraversalResult.HairCount > 0;
+
+	if (bIsValid)
+	{
+		float4x4 TranslatedWorldToClipMatrix = CPU_TranslatedWorldToClip;
+		if (bIsGPUDriven)
+		{
+			FDeepShadowViewInfo DeepShadowViewInfo = DeepShadowViewInfoBuffer[AtlasSlotIndex];
+			TranslatedWorldToClipMatrix = DeepShadowViewInfo.TranslatedWorldToClipScaledBiased;
+		}
+		const float3 HitP = TranslatedWP0 + normalize(TranslatedWP1 - TranslatedWP0) * TraversalResult.HitT;
+		float4 ClipP = mul(float4(HitP, 1), TranslatedWorldToClipMatrix);
+		OutDepth = ClipP.z /= ClipP.w;
+	}
+	else
+	{
+		discard;
+	}
+}
+#endif //SHADER_DEPTH_INJECTION
+
+///////////////////////////////////////////////////////////////////////////
+// Common function for mipmapping voxels
+
+#if SHADER_MIP_VIRTUALVOXEL || SHADER_MIP_INDIRECTARGS || SHADER_UPDATE_PAGEINDEX
+uint ComputeMipDensity(
+	const uint RawDensity0,
+	const uint RawDensity1,
+	const uint RawDensity2,
+	const uint RawDensity3,
+	const uint RawDensity4,
+	const uint RawDensity5,
+	const uint RawDensity6,
+	const uint RawDensity7)
+{
+	const float TotalOpaque =
+		((RawDensity0 & VOXEL_OPAQUE_MASK) >> VOXEL_OPAQUE_SHIFT) +
+		((RawDensity1 & VOXEL_OPAQUE_MASK) >> VOXEL_OPAQUE_SHIFT) +
+		((RawDensity2 & VOXEL_OPAQUE_MASK) >> VOXEL_OPAQUE_SHIFT) +
+		((RawDensity3 & VOXEL_OPAQUE_MASK) >> VOXEL_OPAQUE_SHIFT) +
+		((RawDensity4 & VOXEL_OPAQUE_MASK) >> VOXEL_OPAQUE_SHIFT) +
+		((RawDensity5 & VOXEL_OPAQUE_MASK) >> VOXEL_OPAQUE_SHIFT) +
+		((RawDensity6 & VOXEL_OPAQUE_MASK) >> VOXEL_OPAQUE_SHIFT) +
+		((RawDensity7 & VOXEL_OPAQUE_MASK) >> VOXEL_OPAQUE_SHIFT);
+	const uint OutTotalOpaque = uint(clamp(TotalOpaque / 8.f, TotalOpaque > 0 ? 1 : 0, 0x7F)) << VOXEL_OPAQUE_SHIFT;
+
+	// Propagate no shadow casting flag, only if other valid/non-empty voxels are also no-shadow caster
+	const bool bHasShadowCaster =
+		((RawDensity0 & VOXEL_HAIR_MASK) > 0 && (RawDensity0 & VOXEL_CAST_NO_SHADOW_MASK) == 0) ||
+		((RawDensity1 & VOXEL_HAIR_MASK) > 0 && (RawDensity1 & VOXEL_CAST_NO_SHADOW_MASK) == 0) ||
+		((RawDensity2 & VOXEL_HAIR_MASK) > 0 && (RawDensity2 & VOXEL_CAST_NO_SHADOW_MASK) == 0) ||
+		((RawDensity3 & VOXEL_HAIR_MASK) > 0 && (RawDensity3 & VOXEL_CAST_NO_SHADOW_MASK) == 0) ||
+		((RawDensity4 & VOXEL_HAIR_MASK) > 0 && (RawDensity4 & VOXEL_CAST_NO_SHADOW_MASK) == 0) ||
+		((RawDensity5 & VOXEL_HAIR_MASK) > 0 && (RawDensity5 & VOXEL_CAST_NO_SHADOW_MASK) == 0) ||
+		((RawDensity6 & VOXEL_HAIR_MASK) > 0 && (RawDensity6 & VOXEL_CAST_NO_SHADOW_MASK) == 0) ||
+		((RawDensity7 & VOXEL_HAIR_MASK) > 0 && (RawDensity7 & VOXEL_CAST_NO_SHADOW_MASK) == 0);
+
+	uint TotalHair =
+		(RawDensity0 & VOXEL_HAIR_MASK) +
+		(RawDensity1 & VOXEL_HAIR_MASK) +
+		(RawDensity2 & VOXEL_HAIR_MASK) +
+		(RawDensity3 & VOXEL_HAIR_MASK) +
+		(RawDensity4 & VOXEL_HAIR_MASK) +
+		(RawDensity5 & VOXEL_HAIR_MASK) +
+		(RawDensity6 & VOXEL_HAIR_MASK) +
+		(RawDensity7 & VOXEL_HAIR_MASK);
+	const bool bHasData = TotalHair > 0;
+	TotalHair /= 8;
+
+	// Insure that if a voxel contains some hair data, its total hair remains > 0 after averaging. 
+	// This is important for pruning invalid/empty page later on, to not remove non-empty page (which 
+	// could arise due to numerical precision).
+	return min(uint(VOXEL_HAIR_MASK), bHasData ? max(TotalHair, 1u) : 0u) | OutTotalOpaque | (bHasShadowCaster ? 0u : VOXEL_CAST_NO_SHADOW_MASK);
+}
+#endif // SHADER_MIP_VIRTUALVOXEL || SHADER_MIP_INDIRECTARGS || SHADER_UPDATE_PAGEINDEX
+
+///////////////////////////////////////////////////////////////////////////
+
+#if SHADER_MIP_VIRTUALVOXEL
+int3 PageCountResolution;
+uint PageResolution;
+uint SourceMip;
+uint TargetMip;
+
+Texture3D<uint>   InDensityTexture;
+RWTexture3D<uint> OutDensityTexture;
+
+[numthreads(1, 1, 64)]
+void MainCS(uint3 DispatchThreadId : SV_DispatchThreadID)
+{
+	const uint  MaxDispatchCountPerDim = 65535;
+	const uint  PageIndex = DispatchThreadId.x + DispatchThreadId.y * MaxDispatchCountPerDim;
+	const uint3 PageCoord = IndexToCoord(PageIndex, PageCountResolution);
+
+	const uint  InPageResolution = PageResolution >> SourceMip;
+	const uint OutPageResolution = PageResolution >> TargetMip;
+	const uint TotalVoxelPerOutPageResolution = OutPageResolution * OutPageResolution * OutPageResolution;
+	const uint VoxelIndex = DispatchThreadId.z;
+	if (VoxelIndex >= TotalVoxelPerOutPageResolution)
+		return;
+
+	const uint3 OutVoxelCoordOffset = IndexToCoord(VoxelIndex, OutPageResolution);
+	const uint3  InVoxelCoordOffset = OutVoxelCoordOffset << 1;
+
+	const uint3 OutVoxelCoord = PageCoord * OutPageResolution + OutVoxelCoordOffset;
+	const uint3  InVoxelCoord = PageCoord * InPageResolution + InVoxelCoordOffset;
+
+	const uint3 InVoxelCoord0 = InVoxelCoord;
+	const uint3 InVoxelCoord1 = InVoxelCoord0 + uint3(1, 0, 0);
+	const uint3 InVoxelCoord2 = InVoxelCoord0 + uint3(0, 1, 0);
+	const uint3 InVoxelCoord3 = InVoxelCoord0 + uint3(1, 1, 0);
+	const uint3 InVoxelCoord4 = InVoxelCoord0 + uint3(0, 0, 1);
+	const uint3 InVoxelCoord5 = InVoxelCoord0 + uint3(1, 0, 1);
+	const uint3 InVoxelCoord6 = InVoxelCoord0 + uint3(0, 1, 1);
+	const uint3 InVoxelCoord7 = InVoxelCoord0 + uint3(1, 1, 1);
+
+	const uint RawDensity0 = InDensityTexture[InVoxelCoord0];
+	const uint RawDensity1 = InDensityTexture[InVoxelCoord1];
+	const uint RawDensity2 = InDensityTexture[InVoxelCoord2];
+	const uint RawDensity3 = InDensityTexture[InVoxelCoord3];
+	const uint RawDensity4 = InDensityTexture[InVoxelCoord4];
+	const uint RawDensity5 = InDensityTexture[InVoxelCoord5];
+	const uint RawDensity6 = InDensityTexture[InVoxelCoord6];
+	const uint RawDensity7 = InDensityTexture[InVoxelCoord7];
+
+	const uint Total = ComputeMipDensity(
+		RawDensity0,
+		RawDensity1,
+		RawDensity2,
+		RawDensity3,
+		RawDensity4,
+		RawDensity5,
+		RawDensity6,
+		RawDensity7);
+
+	OutDensityTexture[OutVoxelCoord] = Total;
+}
+#endif
+
+///////////////////////////////////////////////////////////////////////////
+
+#if SHADER_MIP_INDIRECTARGS
+#include "HairStrandsVoxelPageCommon.ush"
+
+uint PageResolution;
+uint TargetMipIndex;
+uint DispatchGroupSize;
+
+Buffer<uint>   InIndirectArgs;
+RWBuffer<uint> OutIndirectArgs;
+
+[numthreads(1, 1, 1)]
+void MainCS(uint3 DispatchThreadId : SV_DispatchThreadID)
+{
+	const uint MacroGroupId = DispatchThreadId.x;
+
+	const uint TargetPageResolution = PageResolution >> TargetMipIndex;
+	const uint TotalVoxelCount = TargetPageResolution * TargetPageResolution * TargetPageResolution;
+	const uint DispatchZ = ceil(TotalVoxelCount / float(DispatchGroupSize));
+
+	WriteDispatchIndirectArgs(OutIndirectArgs, 0, InIndirectArgs[0], InIndirectArgs[1], DispatchZ);
+}
+
+#endif
+
+///////////////////////////////////////////////////////////////////////////
+
+#if SHADER_UPDATE_PAGEINDEX
+int3 PageCountResolution;
+uint PageResolution;
+uint MipIt;
+
+Buffer<uint>	  PageIndexGlobalCounter; // Allocated page counter
+Texture3D<uint>	  DensityTexture;
+Buffer<uint>	  PageToPageIndexBuffer;
+RWBuffer<uint>	  OutPageIndexBuffer;
+RWBuffer<uint2>	  OutPageIndexOccupancyBuffer;
+
+[numthreads(1, 1, 64)]
+void MainCS(uint3 DispatchThreadId : SV_DispatchThreadID)
+{
+	const uint  MaxDispatchCountPerDim = 65535;
+	const uint  PageIndex = DispatchThreadId.x + DispatchThreadId.y * MaxDispatchCountPerDim;
+	const uint3 PageCoord = IndexToCoord(PageIndex, PageCountResolution);
+
+	const uint AllocatedPageCount = PageIndexGlobalCounter[0];
+	if (PageIndex >= AllocatedPageCount)
+		return;
+
+	const uint OutPageResolution = PageResolution >> MipIt;
+	const uint TotalVoxelPerOutPageResolution = OutPageResolution * OutPageResolution * OutPageResolution;
+	const uint VoxelIndex = DispatchThreadId.z;
+	if (VoxelIndex >= TotalVoxelPerOutPageResolution)
+		return;
+
+	{
+		const uint3 OutVoxelCoordOffset = IndexToCoord(VoxelIndex, OutPageResolution);
+		const uint3 OutVoxelCoord = PageCoord * OutPageResolution + OutVoxelCoordOffset;
+
+		const uint Density = DensityTexture.Load(uint4(OutVoxelCoord, MipIt));
+		// Update the page index with invalid page index if the voxel does not contain any data.
+		// This allow to save tracing cost when evaluating the transmittance.
+		if (Density == 0)
+		{
+			const uint PageIndexOffset = PageToPageIndexBuffer[PageIndex];
+			OutPageIndexBuffer[PageIndexOffset] = INVALID_VOXEL_PAGE_INDEX;
+		}
+	}
+
+	// Fill in the occupancy table
+	{
+		const uint OccMipIt = 3; // A voxel at this mip has a size of 8x8x8
+		const uint OccPageResolution = PageResolution >> OccMipIt;
+		uint2 Occupancy = 0;
+		uint BitIndex = 0;
+		for (uint z = 0; z < 4; ++z)
+			for (uint y = 0; y < 4; ++y)
+				for (uint x = 0; x < 4; ++x)
+				{
+					const uint3 OutVoxelCoord = PageCoord * OccPageResolution + uint3(x, y, z);
+					const uint Density = DensityTexture.Load(uint4(OutVoxelCoord, OccMipIt));
+					if (Density > 0)
+					{
+						if (BitIndex < 32)
+							Occupancy.x = Occupancy.x | (1u << BitIndex);
+						else
+							Occupancy.y = Occupancy.y | (1u << (BitIndex - 32));
+					}
+					++BitIndex;
+				}
+		const uint PageIndexOffset = PageToPageIndexBuffer[PageIndex];
+		OutPageIndexOccupancyBuffer[PageIndexOffset] = Occupancy;
+	}
+}
+#endif
+
+///////////////////////////////////////////////////////////////////////////
+
diff --git a/Engine/Source/Runtime/Renderer/Private/HairStrands/HairStrandsDebug.cpp b/Engine/Source/Runtime/Renderer/Private/HairStrands/HairStrandsDebug.cpp
--- a/Engine/Source/Runtime/Renderer/Private/HairStrands/HairStrandsDebug.cpp
+++ b/Engine/Source/Runtime/Renderer/Private/HairStrands/HairStrandsDebug.cpp
@@ -742,10 +742,11 @@ public:
 		// Skip optimization for avoiding long compilation time due to large UAV writes
 		FGlobalShader::ModifyCompilationEnvironment(Parameters, OutEnvironment);
 		OutEnvironment.CompilerFlags.Add(CFLAG_Debug);
+		OutEnvironment.SetDefine(TEXT("SHADER_VOXEL_DEBUG"), 1);
 	}
 };
 
-IMPLEMENT_GLOBAL_SHADER(FVoxelVirtualRaymarchingCS, "/Engine/Private/HairStrands/HairStrandsVoxelPageRayMarching.usf", "MainCS", SF_Compute);
+IMPLEMENT_GLOBAL_SHADER(FVoxelVirtualRaymarchingCS, "/Engine/Private/HairStrands/HairStrandsVoxelDebug.usf", "MainCS", SF_Compute);
 
 static void AddVoxelPageRaymarchingPass(
 	FRDGBuilder& GraphBuilder,
diff --git a/Engine/Source/Runtime/Renderer/Private/HairStrands/HairStrandsDeepShadow.cpp b/Engine/Source/Runtime/Renderer/Private/HairStrands/HairStrandsDeepShadow.cpp
--- a/Engine/Source/Runtime/Renderer/Private/HairStrands/HairStrandsDeepShadow.cpp
+++ b/Engine/Source/Runtime/Renderer/Private/HairStrands/HairStrandsDeepShadow.cpp
@@ -81,8 +81,8 @@ class FHairStrandsShadowDepthInjectionPS : public FHairStrandsShadowDepthInjecti
 	END_SHADER_PARAMETER_STRUCT()
 };
 
-IMPLEMENT_GLOBAL_SHADER(FHairStrandsShadowDepthInjectionPS, "/Engine/Private/HairStrands/HairStrandsVoxelDepthInjection.usf", "MainPS", SF_Pixel);
-IMPLEMENT_GLOBAL_SHADER(FHairStrandsShadowDepthInjectionVS, "/Engine/Private/HairStrands/HairStrandsVoxelDepthInjection.usf", "MainVS", SF_Vertex);
+IMPLEMENT_GLOBAL_SHADER(FHairStrandsShadowDepthInjectionPS, "/Engine/Private/HairStrands/HairStrandsVoxelRasterCompute.usf", "MainPS", SF_Pixel);
+IMPLEMENT_GLOBAL_SHADER(FHairStrandsShadowDepthInjectionVS, "/Engine/Private/HairStrands/HairStrandsVoxelRasterCompute.usf", "MainVS", SF_Vertex);
 
 void AddInjectHairVoxelShadowCaster(
 	FRDGBuilder& GraphBuilder,
diff --git a/Engine/Source/Runtime/Renderer/Private/HairStrands/HairStrandsVoxelization.cpp b/Engine/Source/Runtime/Renderer/Private/HairStrands/HairStrandsVoxelization.cpp
--- a/Engine/Source/Runtime/Renderer/Private/HairStrands/HairStrandsVoxelization.cpp
+++ b/Engine/Source/Runtime/Renderer/Private/HairStrands/HairStrandsVoxelization.cpp
@@ -173,7 +173,7 @@ public:
 	}
 };
 
-IMPLEMENT_GLOBAL_SHADER(FVirtualVoxelInjectOpaqueCS, "/Engine/Private/HairStrands/HairStrandsVoxelOpaque.usf", "MainCS", SF_Compute);
+IMPLEMENT_GLOBAL_SHADER(FVirtualVoxelInjectOpaqueCS, "/Engine/Private/HairStrands/HairStrandsVoxelRasterCompute.usf", "MainCS", SF_Compute);
 
 static void AddVirtualVoxelInjectOpaquePass(
 	FRDGBuilder& GraphBuilder,
@@ -510,18 +510,18 @@ public:
 	}
 };
 
-IMPLEMENT_GLOBAL_SHADER(FVoxelMarkValidPageIndex_PrepareCS, "/Engine/Private/HairStrands/HairStrandsVoxelPageAllocation.usf", "MarkValid_PrepareCS", SF_Compute);
-IMPLEMENT_GLOBAL_SHADER(FVoxelMarkValidPageIndex_IndirectArgsCS, "/Engine/Private/HairStrands/HairStrandsVoxelPageAllocation.usf", "MarkValid_BuildIndirectArgCS", SF_Compute);
-IMPLEMENT_GLOBAL_SHADER(FVoxelMarkValidPageIndex_ScatterCS, "/Engine/Private/HairStrands/HairStrandsVoxelPageAllocation.usf", "MarkValid_ScatterCS", SF_Compute);
+IMPLEMENT_GLOBAL_SHADER(FVoxelMarkValidPageIndex_PrepareCS, "/Engine/Private/HairStrands/HairStrandsVoxelRasterCompute.usf", "MarkValid_PrepareCS", SF_Compute);
+IMPLEMENT_GLOBAL_SHADER(FVoxelMarkValidPageIndex_IndirectArgsCS, "/Engine/Private/HairStrands/HairStrandsVoxelRasterCompute.usf", "MarkValid_BuildIndirectArgCS", SF_Compute);
+IMPLEMENT_GLOBAL_SHADER(FVoxelMarkValidPageIndex_ScatterCS, "/Engine/Private/HairStrands/HairStrandsVoxelRasterCompute.usf", "MarkValid_ScatterCS", SF_Compute);
 
-IMPLEMENT_GLOBAL_SHADER(FVoxelAllocatePageIndexCS, "/Engine/Private/HairStrands/HairStrandsVoxelPageAllocation.usf", "AllocatePageIndex", SF_Compute);
-IMPLEMENT_GLOBAL_SHADER(FVoxelMarkValidPageIndexCS, "/Engine/Private/HairStrands/HairStrandsVoxelPageAllocation.usf", "MarkValidCS", SF_Compute);
-IMPLEMENT_GLOBAL_SHADER(FVoxelAllocateVoxelPageCS, "/Engine/Private/HairStrands/HairStrandsVoxelPageAllocation.usf", "AllocateCS", SF_Compute);
-IMPLEMENT_GLOBAL_SHADER(FVoxelAddNodeDescCS, "/Engine/Private/HairStrands/HairStrandsVoxelPageAllocation.usf", "AddDescCS", SF_Compute);
-IMPLEMENT_GLOBAL_SHADER(FVoxelAddIndirectBufferCS, "/Engine/Private/HairStrands/HairStrandsVoxelPageAllocation.usf", "AddIndirectBufferCS", SF_Compute);
-IMPLEMENT_GLOBAL_SHADER(FVoxelIndPageClearBufferGenCS, "/Engine/Private/HairStrands/HairStrandsVoxelPageAllocation.usf", "VoxelIndPageClearBufferGenCS", SF_Compute);
-IMPLEMENT_GLOBAL_SHADER(FVoxelIndPageClearCS, "/Engine/Private/HairStrands/HairStrandsVoxelPageAllocation.usf", "VoxelIndPageClearCS", SF_Compute);
-IMPLEMENT_GLOBAL_SHADER(FVoxelAdaptiveFeedbackCS, "/Engine/Private/HairStrands/HairStrandsVoxelPageAllocation.usf", "FeedbackCS", SF_Compute);
+IMPLEMENT_GLOBAL_SHADER(FVoxelAllocatePageIndexCS, "/Engine/Private/HairStrands/HairStrandsVoxelRasterCompute.usf", "AllocatePageIndex", SF_Compute);
+IMPLEMENT_GLOBAL_SHADER(FVoxelMarkValidPageIndexCS, "/Engine/Private/HairStrands/HairStrandsVoxelRasterCompute.usf", "MarkValidCS", SF_Compute);
+IMPLEMENT_GLOBAL_SHADER(FVoxelAllocateVoxelPageCS, "/Engine/Private/HairStrands/HairStrandsVoxelRasterCompute.usf", "AllocateCS", SF_Compute);
+IMPLEMENT_GLOBAL_SHADER(FVoxelAddNodeDescCS, "/Engine/Private/HairStrands/HairStrandsVoxelRasterCompute.usf", "AddDescCS", SF_Compute);
+IMPLEMENT_GLOBAL_SHADER(FVoxelAddIndirectBufferCS, "/Engine/Private/HairStrands/HairStrandsVoxelRasterCompute.usf", "AddIndirectBufferCS", SF_Compute);
+IMPLEMENT_GLOBAL_SHADER(FVoxelIndPageClearBufferGenCS, "/Engine/Private/HairStrands/HairStrandsVoxelRasterCompute.usf", "VoxelIndPageClearBufferGenCS", SF_Compute);
+IMPLEMENT_GLOBAL_SHADER(FVoxelIndPageClearCS, "/Engine/Private/HairStrands/HairStrandsVoxelRasterCompute.usf", "VoxelIndPageClearCS", SF_Compute);
+IMPLEMENT_GLOBAL_SHADER(FVoxelAdaptiveFeedbackCS, "/Engine/Private/HairStrands/HairStrandsVoxelRasterCompute.usf", "FeedbackCS", SF_Compute);
 
 inline FIntVector CeilToInt(const FVector& V)
 {
@@ -1432,9 +1432,9 @@ public:
 };
 
 
-IMPLEMENT_GLOBAL_SHADER(FVirtualVoxelGenerateMipCS, "/Engine/Private/HairStrands/HairStrandsVoxelMip.usf", "MainCS", SF_Compute);
-IMPLEMENT_GLOBAL_SHADER(FVirtualVoxelIndirectArgMipCS, "/Engine/Private/HairStrands/HairStrandsVoxelMip.usf", "MainCS", SF_Compute);
-IMPLEMENT_GLOBAL_SHADER(FVirtualVoxelPatchPageIndexWithMipDataCS, "/Engine/Private/HairStrands/HairStrandsVoxelMip.usf", "MainCS", SF_Compute);
+IMPLEMENT_GLOBAL_SHADER(FVirtualVoxelGenerateMipCS, "/Engine/Private/HairStrands/HairStrandsVoxelRasterCompute.usf", "MainCS", SF_Compute);
+IMPLEMENT_GLOBAL_SHADER(FVirtualVoxelIndirectArgMipCS, "/Engine/Private/HairStrands/HairStrandsVoxelRasterCompute.usf", "MainCS", SF_Compute);
+IMPLEMENT_GLOBAL_SHADER(FVirtualVoxelPatchPageIndexWithMipDataCS, "/Engine/Private/HairStrands/HairStrandsVoxelRasterCompute.usf", "MainCS", SF_Compute);
 
 
 static void AddVirtualVoxelGenerateMipPass(
