Strata - fixed forward shading ambient lighting to work with lumen GI enabled.

diff --git a/Engine/Shaders/Private/Strata/StrataForwardLighting.ush b/Engine/Shaders/Private/Strata/StrataForwardLighting.ush
--- a/Engine/Shaders/Private/Strata/StrataForwardLighting.ush
+++ b/Engine/Shaders/Private/Strata/StrataForwardLighting.ush
@@ -161,27 +161,60 @@ float3 StrataForwardLighting(
 				float IndirectIrradiance = 0.0f; 
 
 				// 1. Diffuse Sky evaluation using bent normal
-				if (UseBasePassSkylight > 0)
 				{
+					// Update a context specific for environment lighting
+					FStrataBSDFContext EnvBSDFContext = StrataBSDFContext;
+					EnvBSDFContext.N = ShadingOcclusion.BentNormal;
+					StrataUpdateBSDFContext(EnvBSDFContext, EnvBSDFContext.L);
+
+					// And evaluate diffuse parameters
+					FStrataEnvLightResult StrataEnvLight = StrataEvaluateForEnvLight(EnvBSDFContext, false/*bEnableSpecular*/, Settings);
+					float3 DiffuseNormal = StrataEnvLight.DiffuseNormal;
+					float3 DiffuseColorForIndirect = StrataEnvLight.DiffuseWeight;
+					float3 SubsurfaceColorForIndirect = StrataEnvLight.DiffuseBackFaceWeight;
+
+					// Evaluate diffuse lighting
+					if (any((DiffuseColorForIndirect + SubsurfaceColorForIndirect) > 0.0f))
 					{
-						// Update a context specific for environment lighting
-						FStrataBSDFContext EnvBSDFContext = StrataBSDFContext;
-						EnvBSDFContext.N = ShadingOcclusion.BentNormal;
-						StrataUpdateBSDFContext(EnvBSDFContext, EnvBSDFContext.L);
-
-						// And evaluate diffuse parameters
-						FStrataEnvLightResult StrataEnvLight = StrataEvaluateForEnvLight(EnvBSDFContext, false/*bEnableSpecular*/, Settings);
-						float3 DiffuseNormal = StrataEnvLight.DiffuseNormal;
-						float3 DiffuseColorForIndirect = StrataEnvLight.DiffuseWeight;
-						float3 SubsurfaceColorForIndirect = StrataEnvLight.DiffuseBackFaceWeight;
-
-						// Evaluate diffuse lighting
-						if (any((DiffuseColorForIndirect + SubsurfaceColorForIndirect) > 0.0f))
-						{
-							const bool bEvaluateBackface = any(SubsurfaceColorForIndirect > 0.0);
+						const bool bEvaluateBackface = any(SubsurfaceColorForIndirect > 0.0);
+
+						float3 DiffuseIndirectLighting = 0.0f;
+						float3 SubsurfaceIndirectLighting = 0.0f;
 
-							float3 DiffuseIndirectLighting = 0.0f;
-							float3 SubsurfaceIndirectLighting = 0.0f;
+					#if (MATERIALBLENDING_TRANSLUCENT || MATERIALBLENDING_ADDITIVE || MATERIALBLENDING_MODULATE) && PROJECT_SUPPORTS_LUMEN && !FORWARD_SHADING
+						if (IsLumenTranslucencyGIEnabled())
+						{
+							// Lumen Dynamic GI + shadowed Skylight
+							FTwoBandSHVectorRGB TranslucencyGISH = GetTranslucencyGIVolumeLighting(MaterialParameters.AbsoluteWorldPosition, ResolvedView.WorldToClip, true);
+
+							#if TRANSLUCENCY_LIGHTING_VOLUMETRIC_NONDIRECTIONAL || TRANSLUCENCY_LIGHTING_VOLUMETRIC_PERVERTEX_NONDIRECTIONAL
+
+								FOneBandSHVectorRGB TranslucencyGISH1;
+								TranslucencyGISH1.R.V = TranslucencyGISH.R.V.x;
+								TranslucencyGISH1.G.V = TranslucencyGISH.G.V.x;
+								TranslucencyGISH1.B.V = TranslucencyGISH.B.V.x;
+
+								FOneBandSHVector DiffuseTransferSH = CalcDiffuseTransferSH1(1);
+								DiffuseIndirectLighting += max(float3(0,0,0), DotSH1(TranslucencyGISH1, DiffuseTransferSH)) / PI;
+
+							#else
+								// Diffuse convolution
+								FTwoBandSHVector DiffuseTransferSH = CalcDiffuseTransferSH(DiffuseNormal, 1);
+								DiffuseIndirectLighting += max(half3(0,0,0), DotSH(TranslucencyGISH, DiffuseTransferSH)) / PI;
+
+								#if SHADINGMODEL_REQUIRES_BACKFACE_LIGHTING
+								if (bEvaluateBackface)
+								{
+									FTwoBandSHVector SubsurfaceTransferSH = CalcDiffuseTransferSH(-DiffuseNormal, 1);
+									SubsurfaceIndirectLighting += max(half3(0,0,0), DotSH(TranslucencyGISH, SubsurfaceTransferSH)) / PI;
+								}
+								#endif
+							#endif
+						}
+						else
+					#endif
+						if (UseBasePassSkylight > 0)
+						{
 							GetPrecomputedIndirectLightingAndSkyLight(
 								MaterialParameters,
 								Interpolants,
@@ -193,22 +226,22 @@ float3 StrataForwardLighting(
 								DiffuseIndirectLighting,
 								SubsurfaceIndirectLighting,
 								IndirectIrradiance);
+						}
 
-						#if FORWARD_SHADING && (MATERIALBLENDING_SOLID || MATERIALBLENDING_MASKED)
-							float2 NearestResolvedDepthScreenUV = CalculateNearestResolvedDepthScreenUV(ScreenUV, MaterialParameters.ScreenPosition.w);
+					#if FORWARD_SHADING && (MATERIALBLENDING_SOLID || MATERIALBLENDING_MASKED)
+						float2 NearestResolvedDepthScreenUV = CalculateNearestResolvedDepthScreenUV(ScreenUV, MaterialParameters.ScreenPosition.w);
 
-							IndirectOcclusion = GetIndirectOcclusion(NearestResolvedDepthScreenUV, HasDynamicIndirectShadowCasterRepresentation(StrataPixelHeader));
-							DiffuseIndirectLighting *= IndirectOcclusion;
-							SubsurfaceIndirectLighting *= IndirectOcclusion;
-							IndirectIrradiance *= IndirectOcclusion;
-						#endif
+						IndirectOcclusion = GetIndirectOcclusion(NearestResolvedDepthScreenUV, HasDynamicIndirectShadowCasterRepresentation(StrataPixelHeader));
+						DiffuseIndirectLighting *= IndirectOcclusion;
+						SubsurfaceIndirectLighting *= IndirectOcclusion;
+						IndirectIrradiance *= IndirectOcclusion;
+					#endif
 
-							// For diffuse, we specify a perpendicular to the surface light direction for the transmittance to light to not be view dependent.
-							const float DiffuseEnvLightingNoL = 1.0f;
+						// For diffuse, we specify a perpendicular to the surface light direction for the transmittance to light to not be view dependent.
+						const float DiffuseEnvLightingNoL = 1.0f;
 
-							Color += (DiffuseIndirectLighting * DiffuseColorForIndirect + SubsurfaceIndirectLighting * SubsurfaceColorForIndirect) * 
-								LuminanceWeight(DiffuseEnvLightingNoL, CurrentBSDF) * AOMultiBounce(StrataEnvLight.DiffuseColor, ShadingOcclusion.DiffOcclusion);
-						}
+						Color += (DiffuseIndirectLighting * DiffuseColorForIndirect + SubsurfaceIndirectLighting * SubsurfaceColorForIndirect) * 
+							LuminanceWeight(DiffuseEnvLightingNoL, CurrentBSDF) * AOMultiBounce(StrataEnvLight.DiffuseColor, ShadingOcclusion.DiffOcclusion);
 					}
 				}
 
