ColorManagement: Working color space support in volumetric cloud rendering.

diff --git a/Engine/Shaders/Private/ColorSpace.ush b/Engine/Shaders/Private/ColorSpace.ush
--- a/Engine/Shaders/Private/ColorSpace.ush
+++ b/Engine/Shaders/Private/ColorSpace.ush
@@ -266,4 +266,28 @@ float3 LAB_2_LCH(float3 LAB)
 	}
 	LCH.z = HInDegree0To360;
 	return LCH;
-}
\ No newline at end of file
+}
+
+
+//------------------------------------------------------------------
+// sRGB / Working Color Space
+//------------------------------------------------------------------
+float3 SRGBCoefficientsToWorkingColorSpace(in float3 CoefficientsSRGB)
+{
+#if WORKING_COLOR_SPACE_IS_SRGB
+	return CoefficientsSRGB;
+#else
+	float3 TransmittanceSRGB = exp(-CoefficientsSRGB);
+	float3 TransmittanceWCS = mul(SRGB_TO_WORKING_COLOR_SPACE_MAT, TransmittanceSRGB);
+	return -log(TransmittanceWCS);
+#endif
+}
+
+float3 SRGBColorToWorkingColorSpace(in float3 ColorSRGB)
+{
+#if WORKING_COLOR_SPACE_IS_SRGB
+	return ColorSRGB;
+#else
+	return mul(SRGB_TO_WORKING_COLOR_SPACE_MAT, ColorSRGB);
+#endif
+}
diff --git a/Engine/Shaders/Private/VolumetricCloud.usf b/Engine/Shaders/Private/VolumetricCloud.usf
--- a/Engine/Shaders/Private/VolumetricCloud.usf
+++ b/Engine/Shaders/Private/VolumetricCloud.usf
@@ -8,6 +8,7 @@
 
 #include "Common.ush"
 #include "ColorMap.ush"
+#include "ColorSpace.ush"
 
 #define SceneTexturesStruct RenderVolumetricCloudParameters.SceneTextures
 
@@ -135,40 +136,32 @@ float SampleAmbientOcclusion(in FPixelMaterialInputs PixelMaterialInputs)
 	return AO;
 }
 
-#ifdef WORKING_COLOR_SPACE_IS_SRGB
-#undef WORKING_COLOR_SPACE_IS_SRGB
-#endif
-#define WORKING_COLOR_SPACE_IS_SRGB 0
 void ConvertCloudPixelMaterialInputsToWorkingColorSpace(inout FPixelMaterialInputs PixelMaterialInputs)
 {
-#if WORKING_COLOR_SPACE_IS_SRGB
+#if !WORKING_COLOR_SPACE_IS_SRGB
 	// Get raw data
 	float3 ExtinctionSRGB = SampleExtinctionCoefficients(PixelMaterialInputs);
 	float3 AlbedoSRGB = SampleAlbedo(PixelMaterialInputs);
 	float3 EmissiveSRGB = SampleEmissive(PixelMaterialInputs);
 
-	// Compute scattering&absorption coefficients for SRGB
-	float3 ScatteringSRGB = AlbedoSRGB * ExtinctionSRGB;
-	float3 AbsorptionSRGB = max(0.0f, ExtinctionSRGB - ScatteringSRGB);
-
-	// Convert all values to working color space
-	float3 ScatteringColorSpace = ConvertCoefficientsToWorkingColorSpace(ScatteringSRGB);
-	float3 AbsorptionColorSpace = ConvertCoefficientsToWorkingColorSpace(AbsorptionSRGB);
-	float3 EmissiveColorSpace	= ConvertLuminanceToWorkingColorSpace(EmissiveSRGB);
-
-	// Convert back to extinction/albedo for working color space
-	float3 ExtinctionColorSpace = ScatteringColorSpace + AbsorptionColorSpace;
-	float3 AlbedoColorSpace		= ScatteringColorSpace / max(0.000000000001f, ExtinctionColorSpace);
-
-#if STRATA_ENABLED
-	BSDF_GETEMISSIVE(PixelMaterialInputs.FrontMaterial.InlinedBSDF)					= EmissiveColorSpace;
-	VOLUMETRICFOGCLOUD_EXTINCTION(PixelMaterialInputs.FrontMaterial.InlinedBSDF)	= ExtinctionColorSpace;
-	VOLUMETRICFOGCLOUD_ALBEDO(PixelMaterialInputs.FrontMaterial.InlinedBSDF)		= AlbedoColorSpace;
+	/** To be revisited:
+	 * Converting albedo and extinction directly appears to yield better results than calculating scattering
+	 * and absorption, converting those to the working color space and extracting the results back.
+	 * The (ScatteringSRGB = AlbedoSRGB * ExtinctionSRGB) multiplication is sensitive to its color space context.
+	 */
+	float3 ExtinctionWCS	= SRGBCoefficientsToWorkingColorSpace(ExtinctionSRGB);
+	float3 AlbedoWCS		= SRGBColorToWorkingColorSpace(AlbedoSRGB);
+	float3 EmissiveWCS		= SRGBColorToWorkingColorSpace(EmissiveSRGB);
+
+#if STRATA_ENABLED	
+	VOLUMETRICFOGCLOUD_EXTINCTION(PixelMaterialInputs.FrontMaterial.InlinedBSDF)	= ExtinctionWCS;
+	VOLUMETRICFOGCLOUD_ALBEDO(PixelMaterialInputs.FrontMaterial.InlinedBSDF)		= AlbedoWCS;
+	BSDF_GETEMISSIVE(PixelMaterialInputs.FrontMaterial.InlinedBSDF)					= EmissiveWCS;
 #else
 	// Set raw values back
-	PixelMaterialInputs.EmissiveColor	= EmissiveColorSpace;
-	PixelMaterialInputs.Subsurface		= ExtinctionColorSpace;
-	PixelMaterialInputs.BaseColor		= AlbedoColorSpace;
+	PixelMaterialInputs.Subsurface.rgb	= ExtinctionWCS;
+	PixelMaterialInputs.BaseColor		= AlbedoWCS;
+	PixelMaterialInputs.EmissiveColor	= EmissiveWCS;
 #endif // STRATA_ENABLED
 #endif // WORKING_COLOR_SPACE_IS_SRGB
 }
diff --git a/Engine/Source/Runtime/Engine/Private/ShaderCompiler/ShaderCompiler.cpp b/Engine/Source/Runtime/Engine/Private/ShaderCompiler/ShaderCompiler.cpp
--- a/Engine/Source/Runtime/Engine/Private/ShaderCompiler/ShaderCompiler.cpp
+++ b/Engine/Source/Runtime/Engine/Private/ShaderCompiler/ShaderCompiler.cpp
@@ -50,6 +50,7 @@
 #include "UObject/UObjectIterator.h"
 #include "Math/UnitConversion.h"
 #include "UnrealEngine.h"
+#include "ColorSpace.h"
 
 #if WITH_EDITOR
 #include "UObject/ArchiveCookContext.h"
@@ -6465,6 +6466,27 @@ void GlobalBeginCompileShader(
 		Input.Environment.SetDefine(TEXT("PROJECT_MOBILE_ENABLE_MOVABLE_SPOTLIGHT_SHADOWS"), bMobileMovableSpotlightShadowsEnabled ? 1 : 0);
 	}
 
+	{
+		using namespace UE::Color;
+		const bool bWorkingColorSpaceIsSRGB = FColorSpace::GetWorking().IsSRGB();
+		Input.Environment.SetDefine(TEXT("WORKING_COLOR_SPACE_IS_SRGB"), bWorkingColorSpaceIsSRGB ? 1 : 0);
+		
+		// We limit matrix definition below when WORKING_COLOR_SPACE_IS_SRGB == 0.
+		if (!bWorkingColorSpaceIsSRGB)
+		{
+			TCHAR MatrixFormat[] = TEXT("float3x3(%0.10f, %0.10f, %0.10f, %0.10f, %0.10f, %0.10f, %0.10f, %0.10f, %0.10f)");
+
+			const FColorSpaceTransform FromSRGB(FColorSpace(EColorSpace::sRGB), FColorSpace::GetWorking());
+
+			Input.Environment.SetDefine(
+				TEXT("SRGB_TO_WORKING_COLOR_SPACE_MAT"),
+				FString::Printf(MatrixFormat,
+					FromSRGB.M[0][0], FromSRGB.M[1][0], FromSRGB.M[2][0],
+					FromSRGB.M[0][1], FromSRGB.M[1][1], FromSRGB.M[2][1],
+					FromSRGB.M[0][2], FromSRGB.M[1][2], FromSRGB.M[2][2]));
+		}
+	}
+
 	const double TileSize = FLargeWorldRenderScalar::GetTileSize();
 	Input.Environment.SetDefine(TEXT("UE_LWC_RENDER_TILE_SIZE"), (float)TileSize);
 	Input.Environment.SetDefine(TEXT("UE_LWC_RENDER_TILE_SIZE_SQRT"), (float)FMath::Sqrt(TileSize));
diff --git a/Engine/Source/Runtime/RenderCore/Private/Shader.cpp b/Engine/Source/Runtime/RenderCore/Private/Shader.cpp
--- a/Engine/Source/Runtime/RenderCore/Private/Shader.cpp
+++ b/Engine/Source/Runtime/RenderCore/Private/Shader.cpp
@@ -24,6 +24,7 @@
 #include "Misc/LargeWorldRenderPosition.h"
 #include "DataDrivenShaderPlatformInfo.h"
 #include "ShaderPlatformCachedIniValue.h"
+#include "ColorSpace.h"
 
 #if WITH_EDITORONLY_DATA
 #include "Interfaces/IShaderFormat.h"
@@ -2032,6 +2033,20 @@ void ShaderMapAppendKeyString(EShaderPlatform Platform, FString& KeyString)
  		KeyString += FString::Printf(TEXT("_VT-%d-%d-%d-%d"), VTLightmaps, VTTextures, VTSupported, VTFiltering);
 	}
 
+	{
+		const UE::Color::FColorSpace& WCS = UE::Color::FColorSpace::GetWorking();
+		if (!WCS.IsSRGB())
+		{
+			// The working color space is uniquely defined by its chromaticities (as loaded from renderer settings).
+			uint32 WCSHash = 0;
+			WCSHash ^= GetTypeHash(WCS.GetRedChromaticity());
+			WCSHash ^= GetTypeHash(WCS.GetGreenChromaticity());
+			WCSHash ^= GetTypeHash(WCS.GetBlueChromaticity());
+			WCSHash ^= GetTypeHash(WCS.GetWhiteChromaticity());
+			KeyString += FString::Printf(TEXT("_WCS-%u"), WCSHash);
+		}
+	}
+
 	{
 		static const auto CVar = IConsoleManager::Get().FindTConsoleVariableDataInt(TEXT("r.Shaders.RemoveDeadCode"));
 		if (CVar && CVar->GetValueOnAnyThread() != 0)
diff --git a/Engine/Source/Runtime/RenderCore/RenderCore.Build.cs b/Engine/Source/Runtime/RenderCore/RenderCore.Build.cs
--- a/Engine/Source/Runtime/RenderCore/RenderCore.Build.cs
+++ b/Engine/Source/Runtime/RenderCore/RenderCore.Build.cs
@@ -39,7 +39,7 @@ public class RenderCore : ModuleRules
 			RuntimeDependencies.Add(Path.Combine(Unreal.EngineDirectory.ToString(), "Extras/GPUDumpViewer/..."), StagedFileType.DebugNonUFS);
 		}
 
-		PrivateDependencyModuleNames.AddRange(new string[] { "Core", "Projects", "RHI", "ApplicationCore", "TraceLog", "CookOnTheFly" });
+		PrivateDependencyModuleNames.AddRange(new string[] { "Core", "Projects", "RHI", "ApplicationCore", "TraceLog", "CookOnTheFly", "ColorManagement" });
 
         PrivateIncludePathModuleNames.AddRange(new string[] { "DerivedDataCache" });
 		
