Fix incorrect Legacy unlit material conversion.

diff --git a/Engine/Shaders/Private/Strata/StrataLegacyConversion.ush b/Engine/Shaders/Private/Strata/StrataLegacyConversion.ush
--- a/Engine/Shaders/Private/Strata/StrataLegacyConversion.ush
+++ b/Engine/Shaders/Private/Strata/StrataLegacyConversion.ush
@@ -477,7 +477,7 @@ FStrataData StrataConvertLegacyMaterialStatic(
 	if (ShadingModel == SHADINGMODELID_UNLIT)
 	{
 		const float GreyTransmittance = saturate(1.0f - Opacity);
-	#if MATERIALBLENDING_OPAQUE || MATERIALBLENDING_MASKED
+	#if MATERIALBLENDING_SOLID || MATERIALBLENDING_MASKED
 		// Opaque materials only write emissive.
 		Out = GetStrataUnlitBSDF(Emissive, 0.0f);
 	#elif MATERIALBLENDING_ALPHACOMPOSITE
