Add LocalExposure 'IgnoreMaterial' support for Strata.

diff --git a/Engine/Shaders/Private/PostProcessEyeAdaptation.usf b/Engine/Shaders/Private/PostProcessEyeAdaptation.usf
--- a/Engine/Shaders/Private/PostProcessEyeAdaptation.usf
+++ b/Engine/Shaders/Private/PostProcessEyeAdaptation.usf
@@ -7,6 +7,8 @@
 #include "Common.ush"
 #include "PostProcessCommon.ush"
 #include "PostProcessHistogramCommon.ush"
+#include "Strata/Strata.ush"
+#include "Strata/StrataEvaluation.ush"
 
 #if COMPUTESHADER
 
@@ -366,70 +368,143 @@ float4 GetVolumeLightingDirectional(float4 AmbientLightingVector, float3 Directi
 	return float4(max(half3(0, 0, 0), DotSH(TranslucentLighting, DiffuseTransferSH)), AmbientLightingVector.a);
 }
 
-[numthreads(8, 8, 1)]
-void CalculateExposureIlluminance(uint2 ThreadId : SV_DispatchThreadID)
+struct FExposureMaterial
 {
-	const uint2 OutputTexelPos = ThreadId + Illuminance_ViewportMin;
-	const uint2 TexelPos = ThreadId * IllumiananceDownscaleFactor + View.ViewRectMinAndSize.xy;
+	float3 TranslatedWorldPosition;
+	float3 WorldNormal;
+	float3 DirectionalAlbedo;
 
-	const float3 SceneColor = ColorTexture.Load(int3(TexelPos, 0)).rgb * View.OneOverPreExposure;
+	bool bIsValid;
+	bool bHasBackfaceLighting;
+	bool bHasValidDirectionalAlbedo;
+};
 
-	float Luminance = CalculateEyeAdaptationLuminance(SceneColor);
+FExposureMaterial GetExposureMaterial(uint2 InPixelPos)
+{
+	FExposureMaterial Out = (FExposureMaterial)0;
+	#if STRATA_ENABLED
+	{	
+		FStrataAddressing StrataAddressing = GetStrataPixelDataByteOffset(InPixelPos, uint2(View.BufferSizeAndInvSize.xy), Strata.MaxBytesPerPixel);
+		const FStrataPixelHeader StrataPixelHeader = UnpackStrataHeaderIn(Strata.MaterialTextureArray, StrataAddressing, Strata.TopLayerTexture);
+		const FStrataTopLayerData TopLayerData = StrataUnpackTopLayerData(Strata.TopLayerTexture.Load(uint3(InPixelPos, 0)));
+		const FStrataSubsurfaceHeader SSSHeader = StrataLoadSubsurfaceHeader(Strata.MaterialTextureArray, Strata.FirstSliceStoringStrataSSSData, InPixelPos);
+
+		const float DeviceZ = ConvertFromDeviceZ(SceneDepthTexture.Load(int3(InPixelPos, 0)).r);
+		const float3 TranslatedWorldPosition = ReconstructTranslatedWorldPositionFromDeviceZ(InPixelPos, DeviceZ);
+		const float3 V = normalize(View.TranslatedWorldCameraOrigin - TranslatedWorldPosition);
+
+		float3 DiffuseDirectionalAlbedo = 0;
+		float3 DirectionalAlbedo = 0;
+		bool bHasMaterialBackfaceDiffuse = false;
+		{
+			FStrataDeferredLighting Out = GetInitialisedStrataDeferredLighting();
+			const FStrataIntegrationSettings Settings = InitStrataIntegrationSettings(false /*bForceFullyRough*/, Strata.bRoughDiffuse, Strata.PeelLayersAboveDepth);
+			Strata_for(uint BSDFIndex = 0, BSDFIndex < StrataPixelHeader.BSDFCount, ++BSDFIndex)
+			{
+				FStrataBSDF BSDF = UnpackStrataBSDF(Strata.MaterialTextureArray, StrataAddressing, StrataPixelHeader);
+				FStrataBSDFContext StrataBSDFContext = StrataCreateBSDFContext(StrataPixelHeader, BSDF, StrataAddressing, V);
+				const float3 BSDFThroughput = LuminanceWeight(StrataBSDFContext, BSDF);
 
-	const float2 GBufferUV = (TexelPos + 0.5f) * View.BufferSizeAndInvSize.zw;
-	const FGBufferData GBufferData = GetGBufferDataFromSceneTextures(GBufferUV);
+				// Evaluate environment lighting
+				FStrataEnvLightResult StrataEnvLight = StrataEvaluateForEnvLight(StrataBSDFContext, true /*bEnableSpecular*/, Settings);
 
-	if (GBufferData.ShadingModelID != SHADINGMODELID_UNLIT)
+				DiffuseDirectionalAlbedo += BSDFThroughput * StrataEnvLight.DiffuseWeight;
+				DirectionalAlbedo        += BSDFThroughput * StrataEnvLight.DiffuseWeight;
+				DirectionalAlbedo        += BSDFThroughput * StrataEnvLight.SpecularWeight;
+
+				const bool bHasBackfaceDiffuse = StrataGetBSDFType(BSDF) == STRATA_BSDF_TYPE_SLAB && BSDF_GETSSSTYPE(BSDF) != SSS_TYPE_INVALID;
+				if (bHasBackfaceDiffuse)
+				{
+					DirectionalAlbedo += BSDFThroughput * StrataEnvLight.DiffuseBackFaceWeight; // StrataEnvLight.DiffuseBackFaceWeight is already divided PI
+				}
+				if (any(StrataEnvLight.SpecularHazeWeight > 0.0f))
+				{
+					DirectionalAlbedo += BSDFThroughput * StrataEnvLight.SpecularHazeWeight;
+				}
+			}
+		}
+
+		Out.TranslatedWorldPosition = ReconstructTranslatedWorldPositionFromDeviceZ(InPixelPos, DeviceZ);
+		Out.WorldNormal = TopLayerData.WorldNormal;
+		Out.DirectionalAlbedo = DirectionalAlbedo;
+		Out.bIsValid = IsStrataMaterial(StrataPixelHeader);
+		Out.bHasBackfaceLighting = bHasMaterialBackfaceDiffuse;
+		Out.bHasValidDirectionalAlbedo = CalculateEyeAdaptationLuminanceWithoutMin(DiffuseDirectionalAlbedo) > EyeAdaptation_IgnoreMaterialsMinBaseColorLuminance;
+	}
+	#else // STRATA_ENABLED
 	{
+		const float2 GBufferUV = (InPixelPos + 0.5f) * View.BufferSizeAndInvSize.zw;
+		const FGBufferData GBufferData = GetGBufferDataFromSceneTextures(GBufferUV);
 		const float3 TranslatedWorldPosition = ReconstructTranslatedWorldPositionFromDepth(GBufferUV, GBufferData.Depth);
 
-		bool bUseTranslucencyVolume = !EyeAdaptation_IgnoreMaterialsReconstructFromSceneColor;
-
-		// when DiffuseColor is black there's not enough information in SceneColor to reconstruct illuminance
-		if (CalculateEyeAdaptationLuminanceWithoutMin(GBufferData.DiffuseColor) <= EyeAdaptation_IgnoreMaterialsMinBaseColorLuminance)
+		// Calculate approximate illuminance from SceneColor and GBufferData
+		// Illuminance ~= (SceneColorLuminance - Emissive) / (DiffuseColor + SubsurfaceColor + EnvBRDF)
+		float3 Denominator = 0.f;
+		Denominator += GBufferData.DiffuseColor;
+		if (GBufferData.ShadingModelID == SHADINGMODELID_SUBSURFACE
+			|| GBufferData.ShadingModelID == SHADINGMODELID_PREINTEGRATED_SKIN
+			|| GBufferData.ShadingModelID == SHADINGMODELID_TWOSIDED_FOLIAGE
+			|| GBufferData.ShadingModelID == SHADINGMODELID_CLOTH
+			|| GBufferData.ShadingModelID == SHADINGMODELID_EYE)
 		{
-			bUseTranslucencyVolume = true;
-		}
+			float3 SubsurfaceColor = ExtractSubsurfaceColor(GBufferData);
 
-		float Illuminance;
+			if (GBufferData.ShadingModelID == SHADINGMODELID_TWOSIDED_FOLIAGE)
+			{
+				SubsurfaceColor *= 1.0f / PI;
+			}
+
+			Denominator += SubsurfaceColor;
+		}
 
-		if (!bUseTranslucencyVolume)
 		{
-			// Calculate approximate illuminance from SceneColor and GBufferData
-			// Illuminance ~= (SceneColorLuminance - Emissive) / (DiffuseColor + SubsurfaceColor + EnvBRDF)
+			const float3 CameraVector = normalize(View.TranslatedWorldCameraOrigin - TranslatedWorldPosition);
 
-			float3 Denominator = 0.0f;
+			const float3 N = GBufferData.WorldNormal;
+			const float3 V = CameraVector;
+			const float3 EnvBrdf = EnvBRDF(GBufferData.SpecularColor, GBufferData.Roughness, max(0.0, dot(N, V)));
 
-			Denominator += GBufferData.DiffuseColor;
+			Denominator += EnvBrdf;
+		}
 
-			if (GBufferData.ShadingModelID == SHADINGMODELID_SUBSURFACE
-				|| GBufferData.ShadingModelID == SHADINGMODELID_PREINTEGRATED_SKIN
-				|| GBufferData.ShadingModelID == SHADINGMODELID_TWOSIDED_FOLIAGE
-				|| GBufferData.ShadingModelID == SHADINGMODELID_CLOTH
-				|| GBufferData.ShadingModelID == SHADINGMODELID_EYE)
-			{
-				float3 SubsurfaceColor = ExtractSubsurfaceColor(GBufferData);
+		Out.TranslatedWorldPosition = TranslatedWorldPosition;
+		Out.WorldNormal = GBufferData.WorldNormal;
+		Out.DirectionalAlbedo = Denominator;
+		Out.bIsValid = GBufferData.ShadingModelID != SHADINGMODELID_UNLIT;
+		Out.bHasBackfaceLighting = GetShadingModelRequiresBackfaceLighting(GBufferData.ShadingModelID);
+		Out.bHasValidDirectionalAlbedo = CalculateEyeAdaptationLuminanceWithoutMin(GBufferData.DiffuseColor) > EyeAdaptation_IgnoreMaterialsMinBaseColorLuminance;
+	}
+	#endif // STRATA_ENABLED
+	return Out;
+}
 
-				if (GBufferData.ShadingModelID == SHADINGMODELID_TWOSIDED_FOLIAGE)
-				{
-					SubsurfaceColor *= 1.0f / PI;
-				}
+[numthreads(8, 8, 1)]
+void CalculateExposureIlluminance(uint2 ThreadId : SV_DispatchThreadID)
+{
+	const uint2 OutputTexelPos = ThreadId + Illuminance_ViewportMin;
+	const uint2 TexelPos = ThreadId * IllumiananceDownscaleFactor + View.ViewRectMinAndSize.xy;
 
-				Denominator += SubsurfaceColor;
-			}
+	const float3 SceneColor = ColorTexture.Load(int3(TexelPos, 0)).rgb * View.OneOverPreExposure;
 
-			{
-				const float3 CameraVector = normalize(View.TranslatedWorldCameraOrigin - TranslatedWorldPosition);
+	float Luminance = CalculateEyeAdaptationLuminance(SceneColor);
 
-				const float3 N = GBufferData.WorldNormal;
-				const float3 V = CameraVector;
-				const float3 EnvBrdf = EnvBRDF(GBufferData.SpecularColor, GBufferData.Roughness, max(0.0, dot(N, V)));
+	const FExposureMaterial Material = GetExposureMaterial(TexelPos);
 
-				Denominator += EnvBrdf;
-			}
+	if (Material.bIsValid)
+	{
+		bool bUseTranslucencyVolume = !EyeAdaptation_IgnoreMaterialsReconstructFromSceneColor;
 
-			const float DenominatorLuminance = CalculateEyeAdaptationLuminanceWithoutMin(Denominator);
+		// when DirectionalAlbedo is black there's not enough information in SceneColor to reconstruct illuminance
+		if (!Material.bHasValidDirectionalAlbedo)
+		{
+			bUseTranslucencyVolume = true;
+		}
+
+		float Illuminance;
 
+		if (!bUseTranslucencyVolume)
+		{
+			const float DenominatorLuminance = CalculateEyeAdaptationLuminanceWithoutMin(Material.DirectionalAlbedo);
 			// Emissive is copied to IlluminanceTexture after base pass and before lighting pass
 			const float Emissive = RWIlluminanceTexture[OutputTexelPos] * View.OneOverPreExposure;
 
@@ -437,8 +512,7 @@ void CalculateExposureIlluminance(uint2 ThreadId : SV_DispatchThreadID)
 		}
 		else
 		{
-			const float3 TranslucencyEvaluationTranslatedPosition = TranslatedWorldPosition + GBufferData.WorldNormal * EyeAdaptation_IgnoreMaterialsEvaluationPositionBias;
-			const bool bEvaluateBackface = GetShadingModelRequiresBackfaceLighting(GBufferData.ShadingModelID);
+			const float3 TranslucencyEvaluationTranslatedPosition = Material.TranslatedWorldPosition + Material.WorldNormal * EyeAdaptation_IgnoreMaterialsEvaluationPositionBias;
 
 			const float3 LightingPositionOffset = 0;
 
@@ -450,14 +524,14 @@ void CalculateExposureIlluminance(uint2 ThreadId : SV_DispatchThreadID)
 			float4 AmbientLightingVector = GetAmbientLightingVectorFromTranslucentLightingVolume(InnerVolumeUVs, OuterVolumeUVs, FinalLerpFactor);
 			float3 DirectionalLightingVector = GetDirectionalLightingVectorFromTranslucentLightingVolume(InnerVolumeUVs, OuterVolumeUVs, FinalLerpFactor);
 			
-			float3 VolumeLighting = GetVolumeLightingDirectional(AmbientLightingVector, DirectionalLightingVector, GBufferData.WorldNormal).rgb;
+			float3 VolumeLighting = GetVolumeLightingDirectional(AmbientLightingVector, DirectionalLightingVector, Material.WorldNormal).rgb;
 			
-			if (bEvaluateBackface)
+			if (Material.bHasBackfaceLighting)
 			{
-				VolumeLighting += GetVolumeLightingDirectional(AmbientLightingVector, DirectionalLightingVector, -GBufferData.WorldNormal).rgb;
+				VolumeLighting += GetVolumeLightingDirectional(AmbientLightingVector, DirectionalLightingVector, -Material.WorldNormal).rgb;
 			}
 
-#if PROJECT_SUPPORTS_LUMEN && !FORWARD_SHADING
+		#if PROJECT_SUPPORTS_LUMEN && !FORWARD_SHADING
 			if (IsLumenTranslucencyGIEnabled())
 			{
 				// Lumen Dynamic GI + shadowed Skylight
@@ -466,16 +540,16 @@ void CalculateExposureIlluminance(uint2 ThreadId : SV_DispatchThreadID)
 				FTwoBandSHVectorRGB TranslucencyGISH = GetTranslucencyGIVolumeLighting(TranslucencyEvaluationPosition, PrimaryView.WorldToClip, true);
 
 				// Diffuse convolution
-				FTwoBandSHVector DiffuseTransferSH = CalcDiffuseTransferSH(GBufferData.WorldNormal, 1);
+				FTwoBandSHVector DiffuseTransferSH = CalcDiffuseTransferSH(Material.WorldNormal, 1);
 				VolumeLighting.rgb += max(half3(0, 0, 0), DotSH(TranslucencyGISH, DiffuseTransferSH)) / PI;
 
-				if (bEvaluateBackface)
+				if (Material.bHasBackfaceLighting)
 				{
-					FTwoBandSHVector SubsurfaceTransferSH = CalcDiffuseTransferSH(-GBufferData.WorldNormal, 1);
+					FTwoBandSHVector SubsurfaceTransferSH = CalcDiffuseTransferSH(-Material.WorldNormal, 1);
 					VolumeLighting.rgb += max(half3(0, 0, 0), DotSH(TranslucencyGISH, SubsurfaceTransferSH)) / PI;
 				}
 			}
-#endif
+		#endif
 			else
 			{
 				// TODO: Use SkyLight
diff --git a/Engine/Source/Runtime/Renderer/Private/PostProcess/PostProcessEyeAdaptation.cpp b/Engine/Source/Runtime/Renderer/Private/PostProcess/PostProcessEyeAdaptation.cpp
--- a/Engine/Source/Runtime/Renderer/Private/PostProcess/PostProcessEyeAdaptation.cpp
+++ b/Engine/Source/Runtime/Renderer/Private/PostProcess/PostProcessEyeAdaptation.cpp
@@ -4,6 +4,7 @@
 
 #include "Engine/Texture.h"
 #include "TranslucentLighting.h"
+#include "Strata/Strata.h"
 
 #include "SceneTextureParameters.h"
 #include "SystemTextures.h"
@@ -710,6 +711,7 @@ class FCalculateExposureIlluminanceCS : public FGlobalShader
 
 		SHADER_PARAMETER_STRUCT_INCLUDE(FTranslucencyLightingVolumeParameters, TranslucencyLightingVolume)
 		SHADER_PARAMETER_RDG_UNIFORM_BUFFER(FLumenTranslucencyLightingUniforms, LumenGIVolumeStruct)
+		SHADER_PARAMETER_RDG_UNIFORM_BUFFER(FStrataGlobalUniformParameters, Strata)
 	END_SHADER_PARAMETER_STRUCT()
 
 	static bool ShouldCompilePermutation(const FGlobalShaderPermutationParameters& Parameters)
@@ -762,6 +764,7 @@ FRDGTextureRef AddCalculateExposureIlluminancePass(
 			auto* LumenUniforms = GraphBuilder.AllocParameters<FLumenTranslucencyLightingUniforms>();
 			LumenUniforms->Parameters = GetLumenTranslucencyLightingParameters(GraphBuilder, View.LumenTranslucencyGIVolume, View.LumenFrontLayerTranslucency);
 			PassParameters->LumenGIVolumeStruct = GraphBuilder.CreateUniformBuffer(LumenUniforms);
+			PassParameters->Strata = Strata::BindStrataGlobalUniformParameters(View);
 
 			auto ComputeShader = View.ShaderMap->GetShader<FCalculateExposureIlluminanceCS>();
 
