Move strip and control point define into hair definitions.

diff --git a/Engine/Plugins/Runtime/HairStrands/Source/HairStrandsCore/Private/GroomBuilder.cpp b/Engine/Plugins/Runtime/HairStrands/Source/HairStrandsCore/Private/GroomBuilder.cpp
--- a/Engine/Plugins/Runtime/HairStrands/Source/HairStrandsCore/Private/GroomBuilder.cpp
+++ b/Engine/Plugins/Runtime/HairStrands/Source/HairStrandsCore/Private/GroomBuilder.cpp
@@ -349,7 +349,7 @@ namespace HairStrandsBuilder
 					FHairStrandsPositionFormat::Type& PackedPosition = OutPackedPositions[PointIndex + IndexOffset];
 					CopyVectorToPosition(PointPosition - HairBoxCenter, PackedPosition);
 					FPackedRadiusAndType PackedRadiusAndType;
-					PackedRadiusAndType.Data.ControlPointType = (PointIndex == 0) ? 1u : (PointIndex == (PointCount - 1) ? 2u : 0u);
+					PackedRadiusAndType.Data.ControlPointType = (PointIndex == 0) ? HAIR_CONTROLPOINT_START : (PointIndex == (PointCount - 1) ? HAIR_CONTROLPOINT_END : HAIR_CONTROLPOINT_INSIDE);
 					PackedRadiusAndType.Data.NormalizedRadius = uint8(FMath::Clamp(NormalizedRadius * 63.f, 0.f, 63.f));
 					PackedPosition.PackedRadiusAndType = PackedRadiusAndType.Packed;
 					PackedPosition.UCoord = uint8(FMath::Clamp(CoordU * 255.f, 0.f, 255.f));
diff --git a/Engine/Plugins/Runtime/HairStrands/Source/HairStrandsCore/Private/GroomComponent.cpp b/Engine/Plugins/Runtime/HairStrands/Source/HairStrandsCore/Private/GroomComponent.cpp
--- a/Engine/Plugins/Runtime/HairStrands/Source/HairStrandsCore/Private/GroomComponent.cpp
+++ b/Engine/Plugins/Runtime/HairStrands/Source/HairStrandsCore/Private/GroomComponent.cpp
@@ -70,8 +70,6 @@ static FAutoConsoleVariableRef CVarHairStrands_UseAttachedSimulationComponents(T
 
 #define LOCTEXT_NAMESPACE "GroomComponent"
 
-#define USE_HAIR_TRIANGLE_STRIP 0
-
 ////////////////////////////////////////////////////////////////////////////////////////////////////////////////////////
 
 const FLinearColor GetHairGroupDebugColor(int32 GroupIt)
@@ -803,13 +801,9 @@ public:
 		{
 			VertexFactory = (FVertexFactory*)Instance->Strands.VertexFactory;
 			HairVertexCount = Instance->Strands.RestResource->GetVertexCount();
-			#if USE_HAIR_TRIANGLE_STRIP
-			MaxVertexIndex = HairVertexCount * 2;
-			#else
-			MaxVertexIndex = HairVertexCount * 6;
-			#endif
+			MaxVertexIndex = HairVertexCount * HAIR_POINT_TO_VERTEX;
 			bUseCulling = Instance->Strands.bIsCullingEnabled;
-			NumPrimitive = bUseCulling ? 0 : HairVertexCount * 2;
+			NumPrimitive = bUseCulling ? 0 : HairVertexCount * HAIR_POINT_TO_TRIANGLE;
 			if (MaterialRenderProxy == nullptr)
 			{
 				MaterialRenderProxy = HairGroupMaterialProxies[GroupIndex].Strands;
diff --git a/Engine/Shaders/Private/HairStrands/HairStrandsClusterCulling.usf b/Engine/Shaders/Private/HairStrands/HairStrandsClusterCulling.usf
--- a/Engine/Shaders/Private/HairStrands/HairStrandsClusterCulling.usf
+++ b/Engine/Shaders/Private/HairStrands/HairStrandsClusterCulling.usf
@@ -392,11 +392,7 @@ void MainClusterCullingCS(uint3 DispatchThreadId : SV_DispatchThreadID)
 					InterlockedAdd(DispatchIndirectParametersClusterCount[0], 1, VisibleClusterId);
 
 					//uint PreviousVertexCount = 0;
-				#if USE_HAIR_TRIANGLE_STRIP
-					InterlockedAdd(DrawIndirectParameters[0], ClusterLOD.VertexCount * 2/*, PreviousVertexCount*/);
-				#else
-					InterlockedAdd(DrawIndirectParameters[0], ClusterLOD.VertexCount * 6/*, PreviousVertexCount*/);
-				#endif
+					InterlockedAdd(DrawIndirectParameters[0], ClusterLOD.VertexCount * HAIR_POINT_TO_VERTEX/*, PreviousVertexCount*/);
 
 					checkBufferAccessSlow(GlobalClusterIdBuffer, VisibleClusterId);
 					GlobalClusterIdBuffer[VisibleClusterId]	  = ClusterId;
@@ -443,7 +439,7 @@ void MainClusterCullingPrepareIndirectDispatchCS(
 	{
 		// DrawIndirectBuffer contains the total number of vertex count needed for drawing quad triangle strips, so 6 vertices per quad / control point
 		const uint VertexCount    = DrawIndirectBuffer[0];
-		const uint ControlPoint   = VertexCount / (USE_HAIR_TRIANGLE_STRIP ? 2u : 6u);
+		const uint ControlPoint   = VertexCount / (HAIR_POINT_TO_VERTEX);
 		const uint DispatchCount  = DivideAndRoundUp(ControlPoint, HAIR_VERTEXCOUNT_GROUP_SIZE);
 
 		DispatchIndirectBuffer[0] = DispatchCount;
diff --git a/Engine/Shaders/Private/HairStrands/HairStrandsMaterialCommon.ush b/Engine/Shaders/Private/HairStrands/HairStrandsMaterialCommon.ush
--- a/Engine/Shaders/Private/HairStrands/HairStrandsMaterialCommon.ush
+++ b/Engine/Shaders/Private/HairStrands/HairStrandsMaterialCommon.ush
@@ -62,11 +62,10 @@ FVertexFactoryInterpolantsVSToPS GetInterpolants(uint PrimitiveId, uint HairPrim
 	FVertexFactoryInput Input;
 	// Assumes no actual instancing used for these draws. Will load the first instance data.
 	VF_GPUSCENE_SET_INPUT_FOR_RT(Input, PrimitiveId, 0U);
-#if USE_HAIR_TRIANGLE_STRIP
-	Input.VertexId = HairPrimitiveId * 2; // 6: the vertex factory expect triangles strip index. A quad is made of two triangles with indices: 0,1.
-#else
-	Input.VertexId = HairPrimitiveId * 6; // 6: the vertex factory expect triangles list index. A quad is made of two triangles with indices: 0,1,2,3,4,5.
-#endif
+	// * USE_HAIR_TRIANGLE_STRIP=0: the vertex factory expect triangles strip index. A quad is made of two triangles with indices: 0,1.
+	// * USE_HAIR_TRIANGLE_STRIP=1: the vertex factory expect triangles list index. A quad is made of two triangles with indices: 0,1,2,3,4,5.
+	Input.VertexId = HairPrimitiveId * HAIR_POINT_TO_VERTEX; 
+
 	FVertexFactoryIntermediates	Intermediates = GetVertexFactoryIntermediates(Input);
 	FVertexFactoryInterpolantsVSToPS Interpolants = VertexFactoryGetInterpolantsVSToPS(Input, Intermediates, Empty);
 
diff --git a/Engine/Shaders/Private/HairStrands/HairStrandsVertexFactoryCommon.ush b/Engine/Shaders/Private/HairStrands/HairStrandsVertexFactoryCommon.ush
--- a/Engine/Shaders/Private/HairStrands/HairStrandsVertexFactoryCommon.ush
+++ b/Engine/Shaders/Private/HairStrands/HairStrandsVertexFactoryCommon.ush
@@ -4,12 +4,6 @@
 
 #include "/Engine/Shared/HairStrandsDefinitions.h"
 
-#define HAIR_CONTROLPOINT_INSIDE 0
-#define HAIR_CONTROLPOINT_START	1
-#define HAIR_CONTROLPOINT_END	2
-
-#define USE_HAIR_TRIANGLE_STRIP 0
-
 ////////////////////////////////////////////////////////////////////////////////////////////////////////////////////////////////
 // Hair control points
 
diff --git a/Engine/Shaders/Shared/HairStrandsDefinitions.h b/Engine/Shaders/Shared/HairStrandsDefinitions.h
--- a/Engine/Shaders/Shared/HairStrandsDefinitions.h
+++ b/Engine/Shaders/Shared/HairStrandsDefinitions.h
@@ -23,6 +23,24 @@
 // Max number of discrete LOD that a hair group can have
 #define MAX_HAIR_LOD 8
 
+// Use triangle strip for HW raster path
+#define USE_HAIR_TRIANGLE_STRIP 0
+
+// Number of vertex per control-point
+#if USE_HAIR_TRIANGLE_STRIP
+#define HAIR_POINT_TO_VERTEX 2u
+#else
+#define HAIR_POINT_TO_VERTEX 6u
+#endif
+
+// Number of triangle per control-point
+#define HAIR_POINT_TO_TRIANGLE 2u
+
+// Type of control points
+#define HAIR_CONTROLPOINT_INSIDE 0u
+#define HAIR_CONTROLPOINT_START	 1u
+#define HAIR_CONTROLPOINT_END	 2u
+
 // Group size for dispatching based on a groom vertex/curve/cluster count. 
 // This defines ensures the group size is consistent across the hair pipeline, 
 // and ensures dispatch count are smaller than 65k limits
