Fix Shader compilation error due to missing renaming.

diff --git a/Engine/Shaders/Private/HairStrands/HairStrandsAttributeCommon.ush b/Engine/Shaders/Private/HairStrands/HairStrandsAttributeCommon.ush
--- a/Engine/Shaders/Private/HairStrands/HairStrandsAttributeCommon.ush
+++ b/Engine/Shaders/Private/HairStrands/HairStrandsAttributeCommon.ush
@@ -37,9 +37,9 @@ uint InternalGetHairStrandsAttributes(uint Index, uint AttributeOffsetInBytes, u
 	return Packed >> (8u * ElementIndex);
 }
 
-uint InternalGetHairStrandsAttributesPerCurve(uint HairPrimitiveId, uint AttributeOffsetInBytes, uint AttributeStrideInBytes, ByteAddressBuffer InAttributeBuffer, Buffer<uint> InVertexToCurveBuffer)
+uint InternalGetHairStrandsAttributesPerCurve(uint HairPrimitiveId, uint AttributeOffsetInBytes, uint AttributeStrideInBytes, ByteAddressBuffer InAttributeBuffer, Buffer<uint> InPointToCurveBuffer)
 {
-	const uint CurveIndex = InVertexToCurveBuffer[HairPrimitiveId];
+	const uint CurveIndex = InPointToCurveBuffer[HairPrimitiveId];
 	return InternalGetHairStrandsAttributes(CurveIndex, AttributeOffsetInBytes, AttributeStrideInBytes, InAttributeBuffer);
 }
 
@@ -101,7 +101,7 @@ float2 GetHairStrandsUV(uint HairPrimitiveId, float2 HairPrimitiveUV)
 
 float2 GetHairStrandsDimensions(uint HairPrimitiveId, float2 HairPrimitiveUV)
 {
-	const float CurveLength = UnpackHairLength(InternalGetHairStrandsAttributesPerCurve(HairPrimitiveId, HAIR_ATTRIBUTE_OFFSET_LENGTH(HairStrandsVF.AttributeOffsets), HAIR_ATTRIBUTE_STRIDE_LENGTH, HairStrandsVF.AttributeBuffer, HairStrandsVF.VertexToCurveBuffer));
+	const float CurveLength = UnpackHairLength(InternalGetHairStrandsAttributesPerCurve(HairPrimitiveId, HAIR_ATTRIBUTE_OFFSET_LENGTH(HairStrandsVF.AttributeOffsets), HAIR_ATTRIBUTE_STRIDE_LENGTH, HairStrandsVF.AttributeBuffer, HairStrandsVF.PointToCurveBuffer));
 	const float CoordU = InternalGetHairStrandsCoordU(HairPrimitiveId, HairPrimitiveUV);
 	const float Radius = InternalGetHairStrandsWorldRadius(HairPrimitiveId, HairPrimitiveUV);
 	return float2(CurveLength * CoordU, Radius);
@@ -110,12 +110,12 @@ float2 GetHairStrandsDimensions(uint HairPrimitiveId, float2 HairPrimitiveUV)
 float2 GetHairStrandsRootUV(uint HairPrimitiveId, float2 HairPrimitiveUV)
 {
 	// Encoding need to match GroomBuilder.cpp - 
-	return UnpackHairRootUV(InternalGetHairStrandsAttributesPerCurve(HairPrimitiveId, HAIR_ATTRIBUTE_OFFSET_ROOTUV(HairStrandsVF.AttributeOffsets), HAIR_ATTRIBUTE_STRIDE_ROOTUV, HairStrandsVF.AttributeBuffer, HairStrandsVF.VertexToCurveBuffer));
+	return UnpackHairRootUV(InternalGetHairStrandsAttributesPerCurve(HairPrimitiveId, HAIR_ATTRIBUTE_OFFSET_ROOTUV(HairStrandsVF.AttributeOffsets), HAIR_ATTRIBUTE_STRIDE_ROOTUV, HairStrandsVF.AttributeBuffer, HairStrandsVF.PointToCurveBuffer));
 }
 
 float GetHairStrandsSeed(uint HairPrimitiveId, float2 HairPrimitiveUV)
 {
-	return UnpackHairSeed(InternalGetHairStrandsAttributesPerCurve(HairPrimitiveId, HAIR_ATTRIBUTE_OFFSET_SEED(HairStrandsVF.AttributeOffsets), HAIR_ATTRIBUTE_STRIDE_SEED, HairStrandsVF.AttributeBuffer, HairStrandsVF.VertexToCurveBuffer));
+	return UnpackHairSeed(InternalGetHairStrandsAttributesPerCurve(HairPrimitiveId, HAIR_ATTRIBUTE_OFFSET_SEED(HairStrandsVF.AttributeOffsets), HAIR_ATTRIBUTE_STRIDE_SEED, HairStrandsVF.AttributeBuffer, HairStrandsVF.PointToCurveBuffer));
 }
 
 uint GetHairStrandsClumpID(uint HairPrimitiveId, float2 HairPrimitiveUV)
@@ -123,7 +123,7 @@ uint GetHairStrandsClumpID(uint HairPrimitiveId, float2 HairPrimitiveUV)
 	const uint ClumpIdOffset = HAIR_ATTRIBUTE_OFFSET_CLUMPID(HairStrandsVF.AttributeOffsets);
 	if (IsAttributeValid(ClumpIdOffset))
 	{
-		return UnpackHairClumpID(InternalGetHairStrandsAttributesPerCurve(HairPrimitiveId, ClumpIdOffset, HAIR_ATTRIBUTE_STRIDE_CLUMPID, HairStrandsVF.AttributeBuffer, HairStrandsVF.VertexToCurveBuffer));
+		return UnpackHairClumpID(InternalGetHairStrandsAttributesPerCurve(HairPrimitiveId, ClumpIdOffset, HAIR_ATTRIBUTE_STRIDE_CLUMPID, HairStrandsVF.AttributeBuffer, HairStrandsVF.PointToCurveBuffer));
 	}
 	else
 	{
