Path Tracer: Fix potential artifacts in transparent shadows

Instead of using the AHS alone to implement shadows, use an iterative loop with CHS. This avoids duplicate AHS invocations without forcing us to disable them at the acceleration structure level. The AHS is only used for masked materials now.

Refactor the decal AHS logic to be distinct from the path tracer material's AHS logic.

diff --git a/Engine/Plugins/Experimental/GPULightmass/Shaders/Private/LightmapPathTracing.usf b/Engine/Plugins/Experimental/GPULightmass/Shaders/Private/LightmapPathTracing.usf
--- a/Engine/Plugins/Experimental/GPULightmass/Shaders/Private/LightmapPathTracing.usf
+++ b/Engine/Plugins/Experimental/GPULightmass/Shaders/Private/LightmapPathTracing.usf
@@ -302,7 +302,7 @@ FPathTracingPayload TraceTransparentRay(FRayDesc Ray, bool IsCameraRay, bool Las
 
 float3 TraceTransparentVisibilityRay(FRayDesc Ray, float PathRoughness)
 {
-	const uint RayFlags = RAY_FLAG_ACCEPT_FIRST_HIT_AND_END_SEARCH | RAY_FLAG_SKIP_CLOSEST_HIT_SHADER;
+	const uint RayFlags = 0;
 	const uint InstanceInclusionMask = PATHTRACER_MASK_SHADOW;
 	const uint RayContributionToHitGroupIndex = RAY_TRACING_SHADER_SLOT_SHADOW;
 	const uint MultiplierForGeometryContributionToShaderIndex = RAY_TRACING_NUM_SHADER_SLOTS;
@@ -310,25 +310,32 @@ float3 TraceTransparentVisibilityRay(FRayDesc Ray, float PathRoughness)
 
 	FPackedPathTracingPayload PackedPayload = InitPathTracingVisibilityPayload(PathRoughness);
 
-	// Trace the ray
-	TraceRay(
-		TLAS,
-		RayFlags,
-		InstanceInclusionMask,
-		RayContributionToHitGroupIndex,
-		MultiplierForGeometryContributionToShaderIndex,
-		MissShaderIndex,
-		Ray.GetNativeDesc(),
-		PackedPayload);
-
-	if (PackedPayload.IsMiss())
+	for (;;)
 	{
-		// we didn't hit any geo - just return whatever throughput is left
-		return PackedPayload.GetRayThroughput();
+		// Trace the ray
+		TraceRay(
+			TLAS,
+			RayFlags,
+			InstanceInclusionMask,
+			RayContributionToHitGroupIndex,
+			MultiplierForGeometryContributionToShaderIndex,
+			MissShaderIndex,
+			Ray.GetNativeDesc(),
+			PackedPayload);
+		if (PackedPayload.IsMiss())
+		{
+			break;
+		}
+		if (!any(PackedPayload.GetRayThroughput()))
+		{
+			// accumulated to fully opacity already
+			break;
+		}
+		// one past the current hit to get the next transparent hit
+		Ray.TMin = asfloat(asuint(PackedPayload.HitT) + 1);
 	}
 
-	// we registered a hit -- assumed to be fully opaque
-	return 0;
+	return PackedPayload.GetRayThroughput();
 }
 
 void PathTracingKernel(
diff --git a/Engine/Shaders/Private/PathTracing/PathTracingCore.ush b/Engine/Shaders/Private/PathTracing/PathTracingCore.ush
--- a/Engine/Shaders/Private/PathTracing/PathTracingCore.ush
+++ b/Engine/Shaders/Private/PathTracing/PathTracingCore.ush
@@ -612,7 +612,7 @@ float3 TraceTransparentVisibilityRay(FRayDesc Ray, float PathRoughness, uint Mis
 		return 1.0;
 	}
 
-	const uint RayFlags = RAY_FLAG_ACCEPT_FIRST_HIT_AND_END_SEARCH | RAY_FLAG_SKIP_CLOSEST_HIT_SHADER;
+	const uint RayFlags = 0;
 	const uint InstanceInclusionMask = PATHTRACER_MASK_SHADOW | (length2(Ray.Origin) < Pow2(PATH_TRACING_SKIP_HAIR_DISTANCE) ? PATHTRACER_MASK_HAIR_SHADOW : 0);
 	const uint RayContributionToHitGroupIndex = RAY_TRACING_SHADER_SLOT_SHADOW;
 	const uint MultiplierForGeometryContributionToShaderIndex = RAY_TRACING_NUM_SHADER_SLOTS;
@@ -626,40 +626,51 @@ float3 TraceTransparentVisibilityRay(FRayDesc Ray, float PathRoughness, uint Mis
 		Ray.TMax = POSITIVE_INFINITY;
 	}
 
-	// Trace the ray
-	TraceRay(
-		TLAS,
-		RayFlags,
-		InstanceInclusionMask,
-		RayContributionToHitGroupIndex,
-		MultiplierForGeometryContributionToShaderIndex,
-		MissShaderIndex,
-		Ray.GetNativeDesc(),
-		PackedPayload);
-
-	if (PackedPayload.IsMiss())
+	// Trace the ray - progressing in steps through transparent geometry
+	for (;;)
 	{
-		float3 Throughput = PackedPayload.GetRayThroughput();
-		if (bCastShadows)
+		TraceRay(
+			TLAS,
+			RayFlags,
+			InstanceInclusionMask,
+			RayContributionToHitGroupIndex,
+			MultiplierForGeometryContributionToShaderIndex,
+			MissShaderIndex,
+			Ray.GetNativeDesc(),
+			PackedPayload);
+
+		if (PackedPayload.IsMiss())
 		{
-			FVolumeIntersectionList VolumeIsectList = VolumeIntersect(Ray.Origin, Ray.Direction, Ray.TMin, Ray.TMax);
-			if (VolumeIsectList.HitBlocker())
-			{
-				// blockers are opaque -- done!
-				return 0.0;
-			}
-			if (VolumeIsectList.HitVolume())
-			{
-				Throughput = VolumeGetTransmittance(Throughput, Ray.Origin, Ray.Direction, VolumeIsectList, RandSequence);
-			}
+			// we didn't hit anything, we are done
+			break;
+		}
+		if (!any(PackedPayload.GetRayThroughput() > 0))
+		{
+			// we've accumulated to full opacity, we are done
+			break;
 		}
-		// we didn't hit any geo - just return whatever throughput is left
-		float3 Tau = PackedPayload.GetTau() * AbsorptionScale;
-		return Throughput * exp(-max(Tau, 0.0));
+		// we hit some geometry, move one ulp past the hit and trace the same ray again
+		Ray.TMin = asfloat(asuint(PackedPayload.HitT) + 1);
 	}
 
-	// we registered a hit -- assumed to be fully opaque
-	return 0.0;
+	float3 Throughput = PackedPayload.GetRayThroughput();
+	if (bCastShadows && any(Throughput > 0))
+	{
+		FVolumeIntersectionList VolumeIsectList = VolumeIntersect(Ray.Origin, Ray.Direction, Ray.TMin, Ray.TMax);
+		if (VolumeIsectList.HitBlocker())
+		{
+			// blockers are opaque -- done!
+			return 0.0;
+		}
+		if (VolumeIsectList.HitVolume())
+		{
+			Throughput = VolumeGetTransmittance(Throughput, Ray.Origin, Ray.Direction, VolumeIsectList, RandSequence);
+		}
+		// add in the medium extinction
+		const float3 Tau = PackedPayload.GetTau() * AbsorptionScale;
+		Throughput *= exp(-max(Tau, 0.0));
+	}
+	return Throughput;
 }
 
 struct FProbeResult
diff --git a/Engine/Shaders/Private/PathTracing/PathTracingDefaultHitShader.usf b/Engine/Shaders/Private/PathTracing/PathTracingDefaultHitShader.usf
--- a/Engine/Shaders/Private/PathTracing/PathTracingDefaultHitShader.usf
+++ b/Engine/Shaders/Private/PathTracing/PathTracingDefaultHitShader.usf
@@ -8,6 +8,20 @@ RAY_TRACING_ENTRY_CLOSEST_HIT(PathTracingDefaultOpaqueCHS,
 	FRayTracingIntersectionAttributes, Attributes)
 {
 	PackedPayload.HitT = RayTCurrent();
+	if (PackedPayload.IsVisibilityRay())
+	{
+		// mark shadow fully opaque
+		PackedPayload.SetRayThroughput(0.0);
+	}
+	else
+	{
+		// mark hit as solid black
+		FPathTracingPayload Payload = (FPathTracingPayload)0;
+		Payload.ShadingModelID = SHADINGMODELID_UNLIT;
+		Payload.BlendingMode   = RAY_TRACING_BLEND_MODE_OPAQUE;
+		Payload.Opacity = 1.0;
+		PackedPayload = PackPathTracingPayload(Payload);
+	}
 }
 
 RAY_TRACING_ENTRY_CLOSEST_HIT(PathTracingDefaultHiddenCHS,
diff --git a/Engine/Shaders/Private/PathTracing/PathTracingMaterialHitShader.usf b/Engine/Shaders/Private/PathTracing/PathTracingMaterialHitShader.usf
--- a/Engine/Shaders/Private/PathTracing/PathTracingMaterialHitShader.usf
+++ b/Engine/Shaders/Private/PathTracing/PathTracingMaterialHitShader.usf
@@ -78,6 +78,15 @@ RAY_TRACING_ENTRY_CLOSEST_HIT(PathTracingMaterialCHS,
 	FRayTracingIntersectionAttributes, Attributes)
 {
 	PackedPayload.HitT = RayTCurrent();
+#if MATERIALBLENDING_SOLID || MATERIALBLENDING_ALPHAHOLDOUT || (MATERIALBLENDING_MASKED && !MATERIAL_DITHER_OPACITY_MASK)
+	if (PackedPayload.IsVisibilityRay())
+	{
+		// Shadow ray case -- all these blend modes get treated as fully opaque, so exit early.
+		// In the case of Masked materials without dithered opacity, the masking logic was already handled by AHS.
+		PackedPayload.SetRayThroughput(0.0);
+		return;
+	}
+#endif
 
 	ResolvedView = ResolveView();
 
@@ -101,13 +110,147 @@ RAY_TRACING_ENTRY_CLOSEST_HIT(PathTracingMaterialCHS,
 
 	FPixelMaterialInputs PixelMaterialInputs;
 
+	const bool bIsFrontFace = HitKind() == HIT_KIND_TRIANGLE_FRONT_FACE;
 	{
 		const float4 ScreenPosition = SvPositionToResolvedScreenPosition(SvPosition);
-		const bool bIsFrontFace = HitKind() == HIT_KIND_TRIANGLE_FRONT_FACE;
-
 		CalcMaterialParametersEx(MaterialParameters, PixelMaterialInputs, SvPosition, ScreenPosition, bIsFrontFace, TranslatedWorldPosition, TranslatedWorldPosition);
 	}
 
+	// the following blend modes need to process shadows after having executed the material
+#if MATERIALBLENDING_MODULATE || (MATERIALBLENDING_MASKED && MATERIAL_DITHER_OPACITY_MASK) || MATERIALBLENDING_ALPHACOMPOSITE || MATERIALBLENDING_TRANSLUCENT
+	if (PackedPayload.IsVisibilityRay())
+	{
+#if MATERIALBLENDING_MODULATE
+		const float3 Transparency = GetMaterialEmissive(PixelMaterialInputs);
+#elif MATERIALBLENDING_ALPHACOMPOSITE
+		const float Opacity = GetMaterialOpacity(PixelMaterialInputs);
+		const float Transparency = 1 - Opacity;
+#elif MATERIALBLENDING_MASKED && MATERIAL_DITHER_OPACITY_MASK
+		// See MATERIAL_DITHER_OPACITY_MASK comment below
+		const float Opacity = saturate(GetMaterialMaskInputRaw(PixelMaterialInputs));
+		const float Transparency = 1 - Opacity;
+#elif MATERIALBLENDING_TRANSLUCENT
+		const float Opacity = GetMaterialOpacity(PixelMaterialInputs);
+#if MATERIAL_SHADINGMODEL_THIN_TRANSLUCENT || SIMPLIFIED_MATERIAL_SHADER
+		// We can get colored shadows either for thin translucent, or for GPULM which has backwards compatibility with some cpu lightmass tricks
+		float3 Transparency = 1 - Opacity;
+#else
+		float Transparency = 1 - Opacity;
+#endif
+		uint ShadingModelID = GetMaterialShadingModel(PixelMaterialInputs);
+
+#if MATERIAL_SHADINGMODEL_THIN_TRANSLUCENT
+		if (ShadingModelID == SHADINGMODELID_THIN_TRANSLUCENT)
+		{
+			float3 Transmission = GetThinTranslucentMaterialOutput0(MaterialParameters);
+#if REFRACTION_USE_INDEX_OF_REFRACTION && !SIMPLIFIED_MATERIAL_SHADER
+			float Ior = GetRefractionIor(PixelMaterialInputs);
+#else
+			float Ior = 0.0;
+#endif
+			float3 V = WorldRayDirection();
+			float3 N = normalize(MaterialParameters.WorldNormal);
+			float VoN = abs(dot(V, N));
+#if SIMPLIFIED_MATERIAL_SHADER
+			// simplified logic with no bending at the interface and no fresnel
+			Transparency *= pow(Transmission, 1.0 / VoN);
+#else
+			if (Opacity < 1.0)
+			{
+				float PathRoughness = PackedPayload.GetPathRoughness();
+				float Roughness = GetMaterialRoughness(PixelMaterialInputs);
+				if (PathRoughness <= Roughness)
+				{
+					// not using fast caustic approximation - treat as opaque
+					PackedPayload.SetRayThroughput(0.0);
+					return;
+				}
+				// compute transmission through the slab (fresnel + absorption)
+				float3 BaseColor = GetMaterialBaseColor(PixelMaterialInputs);
+				float  Metallic = GetMaterialMetallic(PixelMaterialInputs);
+				float  Specular = GetMaterialSpecular(PixelMaterialInputs);
+				float  F0 = F0RGBToF0(ComputeF0(Specular, BaseColor, Metallic));
+				Transparency *= ComputeThinSlabWeights(Transmission, VoN, Ior, F0).Transmitted;
+
+				// fake caustic approximation (see comments below)
+				Transparency *= (1 - Roughness * Roughness) * saturate(PathRoughness - Roughness);
+			}
+#endif
+		}
+		else
+#endif // MATERIAL_SHADINGMODEL_THIN_TRANSLUCENT
+#if MATERIAL_SHADINGMODEL_DEFAULT_LIT
+		if (ShadingModelID == SHADINGMODELID_DEFAULT_LIT)
+		{
+#if SIMPLIFIED_MATERIAL_SHADER
+			// CPU Lightmass translucency behavior emulation
+			Transparency = lerp(float3(1, 1, 1), GetMaterialBaseColor(PixelMaterialInputs), Opacity);
+#elif REFRACTION_USE_INDEX_OF_REFRACTION
+			// Is refraction enabled?
+			float Ior = GetRefractionIor(PixelMaterialInputs);
+			if (Transparency > 0 && Ior > 0.0)
+			{
+				// current hit has some refraction
+				float PathRoughness = PackedPayload.GetPathRoughness();
+				float Roughness = GetMaterialRoughness(PixelMaterialInputs);
+				if (PathRoughness <= Roughness)
+				{
+					// not using fast caustic approximation - treat as opaque
+					PackedPayload.SetRayThroughput(0.0);
+					return;
+				}
+				// The heuristic used here is inspired by the following presentations:
+				//   - Kulla & Conty: Revisiting Physically Based Shading at Imageworks
+				//       https://blog.selfshadow.com/publications/s2017-shading-course/imageworks/s2017_pbs_imageworks_slides_v2.pdf
+				//   - Colin Barré-Brisebois - Pica Pica & Nvidia Turing
+				//       https://www.ea.com/seed/news/siggraph-2018-picapica-nv-turing
+
+				// fake caustics
+				float3 N = normalize(MaterialParameters.WorldNormal);
+				float Eta = Ior;
+				float VoN = abs(dot(WorldRayDirection(), N));
+				float3 BaseColor = GetMaterialBaseColor(PixelMaterialInputs);
+				float  Metallic = GetMaterialMetallic(PixelMaterialInputs);
+				float  Specular = GetMaterialSpecular(PixelMaterialInputs);
+				float  F0 = F0RGBToF0(ComputeF0(Specular, BaseColor, Metallic));
+				float Fr = FresnelReflectance(VoN, Eta, F0);
+				Transparency *= pow(1 - Fr, 2) * (1 - Roughness * Roughness) * saturate(PathRoughness - Roughness);
+
+#if HAVE_GetAbsorptionMediumMaterialOutput0
+				// Does the material have any kind of volumetric absorption to apply?
+				// Now track the optical thickness so that we can account for Beer's law along the shadow ray
+				// TODO: to support lights inside glass we would need to track an extra offset, but this is hopefully an uncommon scenario
+				// TODO2: SingleLayerWater needs something similar, but it generally is setup with opaque or masked blend mode and water geometry does not cast shadows ...
+				float Distance = RayTCurrent();
+				float3 TransmittanceColor = GetAbsorptionMediumMaterialOutput0(MaterialParameters);
+				float3 LocalSigmaT = -log(max(TransmittanceColor, 1e-8f));
+				
+				// Front Face: remove contribution from the ray origin to the current hit (assuming the backside will be hit)
+				// Back Face:  add contribution from ray origin to current hit (backside) - excess portion not covered by the ray will be removed by the front face hit
+				Distance *= bIsFrontFace ? -1.0 : +1.0;
+				PackedPayload.SetTau(PackedPayload.GetTau() + LocalSigmaT * Distance);
+#endif // HAVE_GetAbsorptionMediumMaterialOutput0
+			}
+#endif // REFRACTION_USE_INDEX_OF_REFRACTION
+		}
+#else // MATERIAL_SHADINGMODEL_DEFAULT_LIT
+		{
+			// base case for shadingmodel if/else
+		}
+#endif // MATERIAL_SHADINGMODEL_DEFAULT_LIT
+
+#else // MATERIALBLEINDG_*
+#error Unhandled blending mode!
+#endif
+
+		// Update the ray throughput (it is packed simply into the payload since we don't need to carry any other information across hits)
+		float3 RayThroughput = PackedPayload.GetRayThroughput();
+		RayThroughput *= Transparency;
+		PackedPayload.SetRayThroughput(RayThroughput);
+		return;
+	}
+#endif
+
 	FPathTracingPayload Payload = (FPathTracingPayload)0;
 
 	/**
@@ -290,7 +433,8 @@ RAY_TRACING_ENTRY_CLOSEST_HIT(PathTracingMaterialCHS,
 		Payload.SetSubsurfaceRadius(GetSubsurfaceMediumMaterialOutput0(MaterialParameters));
 #else
 		// TODO: is this accurate enough? hard to measure exactly since the raster path uses a very different approach
-		const float SSSRadius = (1 - Payload.Opacity) * 10.0; // simple formula, up to 10cm radius
+		const float Opacity = GetMaterialOpacity(PixelMaterialInputs);
+		const float SSSRadius = (1 - Opacity) * 10.0; // simple formula, up to 10cm radius
 		Payload.SetSubsurfaceRadius(SSSRadius);
 #endif
 
@@ -327,8 +471,9 @@ RAY_TRACING_ENTRY_CLOSEST_HIT(PathTracingMaterialCHS,
 		if (GetSubsurfaceProfileUseBurley(ProfileId))
 		{
 			// Decode SSS radius
-			float3 DiffuseColor = Payload.BaseColor * (1.0 - Payload.Metallic);
-			Payload.SetSubsurfaceRadius(DecodeSSSProfileRadius(ProfileId, DiffuseColor, Payload.Opacity));
+			const float Opacity = GetMaterialOpacity(PixelMaterialInputs);
+			const float3 DiffuseColor = Payload.BaseColor * (1.0 - Payload.Metallic);
+			Payload.SetSubsurfaceRadius(DecodeSSSProfileRadius(ProfileId, DiffuseColor, Opacity));
 		}
 		else
 		{
@@ -457,8 +602,9 @@ RAY_TRACING_ENTRY_CLOSEST_HIT(PathTracingMaterialCHS,
 		if (GetSubsurfaceProfileUseBurley(ProfileId))
 		{
 			// Decode SSS radius
-			float3 DiffuseColor = Payload.BaseColor * (1.0 - Payload.Metallic);
-			Payload.SetSubsurfaceRadius(DecodeSSSProfileRadius(ProfileId, DiffuseColor, Payload.Opacity));
+			const float Opacity = GetMaterialOpacity(PixelMaterialInputs);
+			const float3 DiffuseColor = Payload.BaseColor * (1.0 - Payload.Metallic);
+			Payload.SetSubsurfaceRadius(DecodeSSSProfileRadius(ProfileId, DiffuseColor, Opacity));
 		}
 		else
 		{
@@ -598,42 +744,8 @@ RAY_TRACING_ENTRY_ANY_HIT(PathTracingMaterialAHS,
 	FPackedPathTracingPayload, PackedPayload,
 	FRayTracingIntersectionAttributes, Attributes)
 {
-
-#if MATERIALBLENDING_SOLID || MATERIALBLENDING_ALPHAHOLDOUT
-	// Continue traversal by simply returning from this shader
-	// Usually this path is inactive because opaque materials don't have AHS bound
-	return;
-#elif MATERIALBLENDING_ADDITIVE
-	if (PackedPayload.IsVisibilityRay())
-	{
-		// path tracer wants to treat these cases as fully transparent
-		// no material evaluation is required
-		IgnoreHit();
-	}
-	// otherwise accept hit
-#else
-	// All other blending modes _might_ need material evaluation - so we need
-	// to generate some code
-
 #if MATERIALBLENDING_MASKED && !MATERIAL_DITHER_OPACITY_MASK
-	// NOTE: Masked mode execute always - regardless of ray flags
-	// This ensures that both material rays and shadow rays "see" the same thing
-#else
-	// For any other blending models, we have a few options:
-	if (PackedPayload.IsVisibilityRay())
-	{
-		// We want to evaluate the material opacity and accumulate it
-		// Keep going and evaluate the material
-	}
-	else
-	{
-		// non-opaque blending mode but no special behavior was requested
-		return;
-	}
-#endif  // !MATERIALBLENDING_MASKED
-
-	// If we got here -- we have a reason to evaluate the material -- do the setup 
-
+	// This is the only case which actually needs to run the full material
 	ResolvedView = ResolveView();
 
 	const float3 TranslatedWorldPosition = TranslatedWorldRayOrigin() + RayTCurrent() * WorldRayDirection();
@@ -666,163 +778,16 @@ RAY_TRACING_ENTRY_ANY_HIT(PathTracingMaterialAHS,
 		CalcMaterialParametersEx(MaterialParameters, PixelMaterialInputs, SvPosition, ScreenPosition, bIsFrontFace, TranslatedWorldPosition, TranslatedWorldPosition);
 	}
 
-	// Now extract the relevant info from PixelMaterialInputs according to the blending mode
-
-#if MATERIALBLENDING_MASKED && !MATERIAL_DITHER_OPACITY_MASK
 	// Regardless of payload flags -- we always apply this 
 	if (GetMaterialMask(PixelMaterialInputs) < 0)
 	{
 		IgnoreHit();
 	}
 #else
-	// If we got this far -- we have a blending mode that can produce partial opacity _and_ we want to
-	// know how much it contributes to the throughput along this ray (ie: this is a shadow ray).
-#if MATERIALBLENDING_MODULATE
-	// Compute colored opacity
-	float3 Transparency = GetMaterialEmissive(PixelMaterialInputs);
-#elif MATERIALBLENDING_ALPHACOMPOSITE
-	float Opacity = GetMaterialOpacity(PixelMaterialInputs);
-	float Transparency = 1 - Opacity;
-#elif MATERIALBLENDING_MASKED && MATERIAL_DITHER_OPACITY_MASK
-	// See MATERIAL_DITHER_OPACITY_MASK comment above
-	float Opacity = saturate(GetMaterialMaskInputRaw(PixelMaterialInputs));
-	float Transparency = 1 - Opacity;
-#elif MATERIALBLENDING_TRANSLUCENT
-	float Opacity = GetMaterialOpacity(PixelMaterialInputs);
-#if MATERIAL_SHADINGMODEL_THIN_TRANSLUCENT || SIMPLIFIED_MATERIAL_SHADER
-	// We can get colored shadows either for thin translucent, or for GPULM which has backwards compatibility with some cpu lightmass tricks
-	float3 Transparency = 1 - Opacity;
-#else
-	float Transparency = 1 - Opacity;
-#endif
-
-	uint ShadingModelID = GetMaterialShadingModel(PixelMaterialInputs);
-	
-#if MATERIAL_SHADINGMODEL_THIN_TRANSLUCENT
-	if (ShadingModelID == SHADINGMODELID_THIN_TRANSLUCENT)
-	{
-		float3 Transmission = GetThinTranslucentMaterialOutput0(MaterialParameters);
-#if REFRACTION_USE_INDEX_OF_REFRACTION && !SIMPLIFIED_MATERIAL_SHADER
-		float Ior = GetRefractionIor(PixelMaterialInputs);
-#else
-		float Ior = 0.0;
-#endif
-		float3 V = WorldRayDirection();
-		float3 N = normalize(MaterialParameters.WorldNormal);
-		float VoN = abs(dot(V, N));
-#if SIMPLIFIED_MATERIAL_SHADER
-		// simplified logic with no bending at the interface and no fresnel
-		Transparency *= pow(Transmission, 1.0 / VoN);
-#else
-		if (Opacity < 1.0)
-		{
-			float PathRoughness = PackedPayload.GetPathRoughness();
-			float Roughness = GetMaterialRoughness(PixelMaterialInputs);
-			if (PathRoughness <= Roughness)
-			{
-				// not using fast caustic approximation - treat as opaque
-				return;
-			}
-			// compute transmission through the slab (fresnel + absorption)
-			float3 BaseColor = GetMaterialBaseColor(PixelMaterialInputs);
-			float  Metallic = GetMaterialMetallic(PixelMaterialInputs);
-			float  Specular = GetMaterialSpecular(PixelMaterialInputs);
-			float  F0 = F0RGBToF0(ComputeF0(Specular, BaseColor, Metallic));
-			Transparency *= ComputeThinSlabWeights(Transmission, VoN, Ior, F0).Transmitted;
-
-			// fake caustic approximation (see comments below)
-			Transparency *= (1 - Roughness * Roughness) * saturate(PathRoughness - Roughness);
-		}
+	// All other blending modes need to run CHS to decide what to do, or are handled via the default opaque/translucent CHS/AHS materials.
+	// in these cases, the AHS should be disabled, so emit an error if we get here somehow
+	#error Unexpected blend mode encountered in AHS!
 #endif
-	}
-	else
-#endif // MATERIAL_SHADINGMODEL_THIN_TRANSLUCENT
-#if MATERIAL_SHADINGMODEL_DEFAULT_LIT
-	if (ShadingModelID == SHADINGMODELID_DEFAULT_LIT)
-	{
-#if SIMPLIFIED_MATERIAL_SHADER
-		// CPU Lightmass translucency behavior emulation
-		Transparency = lerp(float3(1, 1, 1), GetMaterialBaseColor(PixelMaterialInputs), Opacity);
-#elif REFRACTION_USE_INDEX_OF_REFRACTION
-		// Is refraction enabled?
-		float Ior = GetRefractionIor(PixelMaterialInputs);
-		if (Transparency > 0 && Ior > 0.0)
-		{
-			// current hit has some refraction
-			float PathRoughness = PackedPayload.GetPathRoughness();
-			float Roughness = GetMaterialRoughness(PixelMaterialInputs);
-			if (PathRoughness <= Roughness)
-			{
-				// not using fast caustic approximation - treat as opaque
-				return;
-			}
-			// The heuristic used here is inspired by the following presentations:
-			//   - Kulla & Conty: Revisiting Physically Based Shading at Imageworks
-			//       https://blog.selfshadow.com/publications/s2017-shading-course/imageworks/s2017_pbs_imageworks_slides_v2.pdf
-			//   - Colin Barré-Brisebois - Pica Pica & Nvidia Turing
-			//       https://www.ea.com/seed/news/siggraph-2018-picapica-nv-turing
-
-			// fake caustics
-			float3 N = normalize(MaterialParameters.WorldNormal);
-			float Eta = Ior;
-			float VoN = abs(dot(WorldRayDirection(), N));
-			float3 BaseColor = GetMaterialBaseColor(PixelMaterialInputs);
-			float  Metallic = GetMaterialMetallic(PixelMaterialInputs);
-			float  Specular = GetMaterialSpecular(PixelMaterialInputs);
-			float  F0 = F0RGBToF0(ComputeF0(Specular, BaseColor, Metallic));
-			float Fr = FresnelReflectance(VoN, Eta, F0);
-			Transparency *= pow(1 - Fr, 2) * (1 - Roughness * Roughness) * saturate(PathRoughness - Roughness);
-
-#if HAVE_GetAbsorptionMediumMaterialOutput0
-			// Does the material have any kind of volumetric absorption to apply?
-			// Now track the optical thickness so that we can account for Beer's law along the shadow ray
-			// TODO: to support lights inside glass we would need to track an extra offset, but this is hopefully an uncommon scenario
-			// TODO2: SingleLayerWater needs something similar, but it generally is setup with opaque or masked blend mode and water geometry does not cast shadows ...
-			float Distance = RayTCurrent();
-			float3 TransmittanceColor = GetAbsorptionMediumMaterialOutput0(MaterialParameters);
-			float3 LocalSigmaT = -log(max(TransmittanceColor, 1e-8f));
-			if (bIsFrontFace)
-			{
-				// remove contribution from the ray origin to the current hit (assuming the backside will be hit)
-				PackedPayload.SetTau(PackedPayload.GetTau() - LocalSigmaT * Distance);
-			}
-			else
-			{
-				// add contribution from ray origin to current hit (backside) - excess portion not covered by the ray will be removed by the front face hit
-				PackedPayload.SetTau(PackedPayload.GetTau() + LocalSigmaT * Distance);
-			}
-#endif // HAVE_GetAbsorptionMediumMaterialOutput0
-
-		}
-#endif // REFRACTION_USE_INDEX_OF_REFRACTION
-	}
-#else // MATERIAL_SHADINGMODEL_DEFAULT_LIT
-	{
-		// base case for shadingmodel if/else
-	}
-#endif // MATERIAL_SHADINGMODEL_DEFAULT_LIT
-
-#else
-#error "Control flow does not account for all blending modes!"
-#endif
-
-	// Update the ray throughput (it is packed simply into the payload since we don't need to carry any other information across hits)
-	float3 RayThroughput = PackedPayload.GetRayThroughput();
-
-	RayThroughput *= Transparency;
-
-	PackedPayload.SetRayThroughput(RayThroughput);
-
-	if (any(RayThroughput > 0))
-	{
-		// not at full opacity yet -- ignore the hit so we can keep going
-		IgnoreHit();
-	}
-	// If RayThroughput reached 0 -- just return which will accept the hit
-#endif // !MATERIALBLENDING_MASKED
-
-#endif // !MATERIALBLENDING_SOLID && !MATERIALBLENDING_ALPHAHOLDOUT && !MATERIALBLENDING_ADDITIVE
-
 }
 
 #endif // USE_MATERIAL_ANY_HIT_SHADER
diff --git a/Engine/Shaders/Private/RayTracing/RayTracingDecalMaterialShader.usf b/Engine/Shaders/Private/RayTracing/RayTracingDecalMaterialShader.usf
--- a/Engine/Shaders/Private/RayTracing/RayTracingDecalMaterialShader.usf
+++ b/Engine/Shaders/Private/RayTracing/RayTracingDecalMaterialShader.usf
@@ -175,6 +175,8 @@ RAY_TRACING_ENTRY_CLOSEST_HIT(RayTracingDecalMaterialCHS,
 	return;
 }
 
+#if USE_MATERIAL_ANY_HIT_SHADER
+
 RAY_TRACING_ENTRY_ANY_HIT(RayTracingDecalMaterialAHS,
 	FDecalShaderPayload, Payload,
 	FRayTracingIntersectionAttributes, Attributes)
@@ -209,4 +211,6 @@ RAY_TRACING_ENTRY_ANY_HIT(RayTracingDecalMaterialAHS,
 	}
 }
 
+#endif // USE_MATERIAL_ANY_HIT_SHADER
+
 #endif // RAYHITGROUPSHADER
diff --git a/Engine/Source/Runtime/Engine/Private/Materials/MaterialShared.cpp b/Engine/Source/Runtime/Engine/Private/Materials/MaterialShared.cpp
--- a/Engine/Source/Runtime/Engine/Private/Materials/MaterialShared.cpp
+++ b/Engine/Source/Runtime/Engine/Private/Materials/MaterialShared.cpp
@@ -4960,6 +4960,7 @@ FMaterialShaderParameters::FMaterialShaderParameters(const FMaterial* InMaterial
 	bIsDefaultMaterial = InMaterial->IsDefaultMaterial();
 	bIsSpecialEngineMaterial = InMaterial->IsSpecialEngineMaterial();
 	bIsMasked = InMaterial->IsMasked();
+	bIsDitherMasked = InMaterial->IsDitherMasked();
 	bIsTwoSided = InMaterial->IsTwoSided();
 	bIsThinSurface = InMaterial->IsThinSurface();
 	bIsDistorted = InMaterial->IsDistorted();
diff --git a/Engine/Source/Runtime/Engine/Public/MaterialShared.h b/Engine/Source/Runtime/Engine/Public/MaterialShared.h
--- a/Engine/Source/Runtime/Engine/Public/MaterialShared.h
+++ b/Engine/Source/Runtime/Engine/Public/MaterialShared.h
@@ -2924,6 +2924,7 @@ struct FMaterialShaderParameters
 			uint64 bIsDefaultMaterial : 1;
 			uint64 bIsSpecialEngineMaterial : 1;
 			uint64 bIsMasked : 1;
+			uint64 bIsDitherMasked : 1;
 			uint64 bIsTwoSided : 1;
 			uint64 bIsThinSurface : 1;
 			uint64 bIsDistorted : 1;
diff --git a/Engine/Source/Runtime/Renderer/Private/PathTracing.cpp b/Engine/Source/Runtime/Renderer/Private/PathTracing.cpp
--- a/Engine/Source/Runtime/Renderer/Private/PathTracing.cpp
+++ b/Engine/Source/Runtime/Renderer/Private/PathTracing.cpp
@@ -1129,25 +1129,16 @@ FRayTracingLightFunctionMap GatherLightFunctionLightsPathTracing(FScene* Scene,
 	return RayTracingLightFunctionMap;
 }
 
-static bool NeedsAnyHitShader(bool bIsMasked, EBlendMode BlendMode)
+static bool NeedsAnyHitShader(bool bIsMasked, bool bIsDitherMasked)
 {
-	if (bIsMasked)
-	{
-		return true;
-	}
-	else if (IsOpaqueBlendMode(BlendMode) || IsAlphaHoldoutBlendMode(BlendMode))
-	{
-		return false;
-	}
-	else
-	{
-		return true;
-	}
+	// Masked materials use AHS to quickly reject masked out portions
+	// However, dither masking gets treated as ordinary translucency
+	return (bIsMasked && !bIsDitherMasked);
 }
 
 static bool NeedsAnyHitShader(const FMaterial& RESTRICT MaterialResource)
 {
-	return NeedsAnyHitShader(MaterialResource.IsMasked(), MaterialResource.GetBlendMode());
+	return NeedsAnyHitShader(MaterialResource.IsMasked(), MaterialResource.IsDitherMasked());
 }
 
 template<bool UseAnyHitShader, bool UseIntersectionShader, bool UseSimplifiedShader>
@@ -1173,9 +1164,13 @@ public:
 			// does the VF support ray tracing at all?
 			return false;
 		}
-		if (NeedsAnyHitShader(Parameters.MaterialParameters.bIsMasked, Parameters.MaterialParameters.BlendMode) != UseAnyHitShader)
+		if (Parameters.MaterialParameters.MaterialDomain == MD_DeferredDecal)
+		{
+			// decals are handled elsewhere
+			return false;
+		}
+		if (NeedsAnyHitShader(Parameters.MaterialParameters.bIsMasked, Parameters.MaterialParameters.bIsDitherMasked) != UseAnyHitShader)
 		{
-			// the anyhit permutation is only required if the material is masked or has a non-opaque blend mode
 			return false;
 		}
 		const bool bUseProceduralPrimitive = Parameters.VertexFactoryType->SupportsRayTracingProceduralPrimitive() && FDataDrivenShaderPlatformInfo::GetSupportsRayTracingProceduralPrimitive(Parameters.Platform);
@@ -1322,52 +1317,50 @@ bool FRayTracingMeshProcessor::ProcessPathTracing(
 {
 	FMaterialShaderTypes ShaderTypes;
 
-	const bool bUseProceduralPrimitive = MeshBatch.VertexFactory->GetType()->SupportsRayTracingProceduralPrimitive() &&
-		FDataDrivenShaderPlatformInfo::GetSupportsRayTracingProceduralPrimitive(GMaxRHIShaderPlatform);
-
-	const bool bIsDecal = MaterialResource.GetMaterialDomain() == MD_DeferredDecal;
-	FShaderType* DecalShaderType_CHS = GetRayTracingDecalMaterialShaderType(false);
-	FShaderType* DecalShaderType_CHS_AHS = GetRayTracingDecalMaterialShaderType(true);
-
-	switch (RayTracingMeshCommandsMode)
+	if (MaterialResource.GetMaterialDomain() == MD_DeferredDecal)
 	{
-		case ERayTracingMeshCommandsMode::PATH_TRACING:
+		ShaderTypes.AddShaderType(GetRayTracingDecalMaterialShaderType(MaterialResource.GetBlendMode()));
+	}
+	else
+	{
+		const bool bUseProceduralPrimitive = MeshBatch.VertexFactory->GetType()->SupportsRayTracingProceduralPrimitive() &&
+			FDataDrivenShaderPlatformInfo::GetSupportsRayTracingProceduralPrimitive(GMaxRHIShaderPlatform);
+		switch (RayTracingMeshCommandsMode)
 		{
-			if (NeedsAnyHitShader(MaterialResource))
+			case ERayTracingMeshCommandsMode::PATH_TRACING:
 			{
-				if (bIsDecal)
-					ShaderTypes.AddShaderType(DecalShaderType_CHS_AHS);
-				else if (bUseProceduralPrimitive)
-					ShaderTypes.AddShaderType<FPathTracingMaterialCHS_AHS_IS>();
+				if (NeedsAnyHitShader(MaterialResource))
+				{
+					if (bUseProceduralPrimitive)
+						ShaderTypes.AddShaderType<FPathTracingMaterialCHS_AHS_IS>();
+					else
+						ShaderTypes.AddShaderType<FPathTracingMaterialCHS_AHS>();
+				}
 				else
-					ShaderTypes.AddShaderType<FPathTracingMaterialCHS_AHS>();
+				{
+					if (bUseProceduralPrimitive)
+						ShaderTypes.AddShaderType<FPathTracingMaterialCHS_IS>();
+					else
+						ShaderTypes.AddShaderType<FPathTracingMaterialCHS>();
+				}
+				break;
 			}
-			else
+			case ERayTracingMeshCommandsMode::LIGHTMAP_TRACING:
 			{
-				if (bIsDecal)
-					ShaderTypes.AddShaderType(DecalShaderType_CHS);
-				else if (bUseProceduralPrimitive)
-					ShaderTypes.AddShaderType<FPathTracingMaterialCHS_IS>();
+				if (NeedsAnyHitShader(MaterialResource))
+				{
+					ShaderTypes.AddShaderType<FGPULightmassCHS_AHS>();
+				}
 				else
-					ShaderTypes.AddShaderType<FPathTracingMaterialCHS>();
-			}
-			break;
-		}
-		case ERayTracingMeshCommandsMode::LIGHTMAP_TRACING:
-		{
-			if (NeedsAnyHitShader(MaterialResource))
-			{
-				ShaderTypes.AddShaderType<FGPULightmassCHS_AHS>();
+				{
+					ShaderTypes.AddShaderType<FGPULightmassCHS>();
+				}
+				break;
 			}
-			else
+			default:
 			{
-				ShaderTypes.AddShaderType<FGPULightmassCHS>();
+				return false;
 			}
-			break;
-		}
-		default:
-		{
-		    return false;
 		}
 	}
 
diff --git a/Engine/Source/Runtime/Renderer/Private/RayTracing/RayTracingDecals.cpp b/Engine/Source/Runtime/Renderer/Private/RayTracing/RayTracingDecals.cpp
--- a/Engine/Source/Runtime/Renderer/Private/RayTracing/RayTracingDecals.cpp
+++ b/Engine/Source/Runtime/Renderer/Private/RayTracing/RayTracingDecals.cpp
@@ -225,6 +225,11 @@ IMPLEMENT_RT_PAYLOAD_TYPE(ERayTracingPayloadType::Decals, 60);
 
 IMPLEMENT_SHADER_TYPE(, FRayTracingDecalMaterialShader, TEXT("/Engine/Private/RayTracing/RayTracingDecalMaterialShader.usf"), TEXT("RayTracingDecalMaterialShader"), SF_RayCallable);
 
+static bool NeedsAnyHitShader(EBlendMode BlendMode)
+{
+	return BlendMode != BLEND_Opaque;
+}
+
 template<bool bUseAnyHitShader>
 class TRayTracingDecalMaterial : public FMeshMaterialShader
 {
@@ -252,7 +257,7 @@ public:
 		{
 			return false;
 		}
-		if ((Parameters.MaterialParameters.bIsMasked || Parameters.MaterialParameters.BlendMode != BLEND_Opaque) != bUseAnyHitShader)
+		if (NeedsAnyHitShader(Parameters.MaterialParameters.BlendMode) != bUseAnyHitShader)
 		{
 			// the anyhit permutation is only required if the material is masked or has a non-opaque blend mode
 			return false;
@@ -270,6 +275,7 @@ public:
 		FDecalBlendDesc DecalBlendDesc = DecalRendering::ComputeDecalBlendDesc(Parameters.Platform, Parameters.MaterialParameters);
 
 		OutEnvironment.SetDefine(TEXT("DECAL_PAYLOAD_FLAGS"), CalculateDecalWriteFlags(DecalBlendDesc));
+		OutEnvironment.SetDefine(TEXT("USE_MATERIAL_ANY_HIT_SHADER"), bUseAnyHitShader ? 1 : 0);
 	}
 
 	static bool ValidateCompiledResult(EShaderPlatform Platform, const FShaderParameterMap& ParameterMap, TArray<FString>& OutError)
@@ -306,11 +312,11 @@ using FRayTracingDecalMaterialCHS_AHS = TRayTracingDecalMaterial<true>;
 IMPLEMENT_MATERIAL_SHADER_TYPE(template<>, FRayTracingDecalMaterialCHS, TEXT("/Engine/Private/RayTracing/RayTracingDecalMaterialShader.usf"), TEXT("closesthit=RayTracingDecalMaterialCHS"), SF_RayHitGroup);
 IMPLEMENT_MATERIAL_SHADER_TYPE(template<>, FRayTracingDecalMaterialCHS_AHS, TEXT("/Engine/Private/RayTracing/RayTracingDecalMaterialShader.usf"), TEXT("closesthit=RayTracingDecalMaterialCHS anyhit=RayTracingDecalMaterialAHS"), SF_RayHitGroup);
 
-FShaderType* GetRayTracingDecalMaterialShaderType(bool bUseAnyHitShader)
+FShaderType* GetRayTracingDecalMaterialShaderType(EBlendMode BlendMode)
 {
-	if (bUseAnyHitShader)
+	if (NeedsAnyHitShader(BlendMode))
 	{
-		return &FRayTracingDecalMaterialCHS_AHS::StaticType;
+		return &FRayTracingDecalMaterialCHS_AHS::StaticType; 
 	}
 	else
 	{
diff --git a/Engine/Source/Runtime/Renderer/Private/RayTracing/RayTracingDecals.h b/Engine/Source/Runtime/Renderer/Private/RayTracing/RayTracingDecals.h
--- a/Engine/Source/Runtime/Renderer/Private/RayTracing/RayTracingDecals.h
+++ b/Engine/Source/Runtime/Renderer/Private/RayTracing/RayTracingDecals.h
@@ -27,7 +27,7 @@ END_GLOBAL_SHADER_PARAMETER_STRUCT()
 TRDGUniformBufferRef<FRayTracingDecals> CreateNullRayTracingDecalsUniformBuffer(FRDGBuilder& GraphBuilder);
 TRDGUniformBufferRef<FRayTracingDecals> CreateRayTracingDecalData(FRDGBuilder& GraphBuilder, FScene& Scene, const FViewInfo& View, uint32 BaseCallableSlotIndex);
 
-FShaderType* GetRayTracingDecalMaterialShaderType(bool bUseAnyHitShader);
+FShaderType* GetRayTracingDecalMaterialShaderType(EBlendMode BlendMode);
 
 FRHIRayTracingShader* GetDefaultOpaqueMeshDecalHitShader(const FGlobalShaderMap* ShaderMap);
 FRHIRayTracingShader* GetDefaultHiddenMeshDecalHitShader(const FGlobalShaderMap* ShaderMap);
diff --git a/Engine/Source/Runtime/Renderer/Private/RayTracing/RayTracingInstanceMask.cpp b/Engine/Source/Runtime/Renderer/Private/RayTracing/RayTracingInstanceMask.cpp
--- a/Engine/Source/Runtime/Renderer/Private/RayTracing/RayTracingInstanceMask.cpp
+++ b/Engine/Source/Runtime/Renderer/Private/RayTracing/RayTracingInstanceMask.cpp
@@ -245,7 +245,7 @@ void SetupRayTracingMeshCommandMaskAndStatus(FRayTracingMeshCommand& MeshCommand
 {
 	const FVertexFactory* VertexFactory = MeshBatch.VertexFactory;
 
-	MeshCommand.bCastRayTracedShadows = MeshBatch.CastRayTracedShadow && MaterialResource.CastsRayTracedShadows();
+	MeshCommand.bCastRayTracedShadows = MeshBatch.CastRayTracedShadow && MaterialResource.CastsRayTracedShadows() && MaterialResource.GetBlendMode() != BLEND_Additive;
 	MeshCommand.bOpaque = MaterialResource.GetBlendMode() == EBlendMode::BLEND_Opaque && !(VertexFactory->GetType()->SupportsRayTracingProceduralPrimitive() && FDataDrivenShaderPlatformInfo::GetSupportsRayTracingProceduralPrimitive(GMaxRHIShaderPlatform));
 	MeshCommand.bDecal = MaterialResource.IsDeferredDecal();
 	MeshCommand.bIsSky = MaterialResource.IsSky();
