Add missing coverage scale material change.

diff --git a/Engine/Shaders/Private/HairStrands/HairStrandsMaterialCommon.ush b/Engine/Shaders/Private/HairStrands/HairStrandsMaterialCommon.ush
--- a/Engine/Shaders/Private/HairStrands/HairStrandsMaterialCommon.ush
+++ b/Engine/Shaders/Private/HairStrands/HairStrandsMaterialCommon.ush
@@ -34,6 +34,7 @@ uint MaterialPass_MacroGroupId;
 uint MaterialPass_MaterialId;  
 uint MaterialPass_PrimitiveId;
 uint MaterialPass_LightChannelMask;
+float MaterialPass_HairCoverageScale;
 
 #ifndef HAIR_MATERIAL_DEBUG_OUTPUT
 #define HAIR_MATERIAL_DEBUG_OUTPUT 0
@@ -182,7 +183,7 @@ FEvaluateOutput Evaluate(
 		FHairRenderInfo HairRenderInfo = GetHairRenderInfo(ResolvedView.HairRenderInfo, ResolvedView.HairRenderInfoBits, bUseStableRasterization);
 		const float LocalSceneDepth = ConvertFromDeviceZ(SvPosition.z); // Linear depth in world unit
 		const float PixelRadius = HairRenderInfo.bIsOrthoView ? HairRenderInfo.RadiusAtDepth1Primary : LocalSceneDepth * HairRenderInfo.RadiusAtDepth1Primary; // Not correct but the coverage is not used (we count instead the number of sub-sample covered)
-		Coverage = WorldStrandRadius / max(WorldStrandRadius, PixelRadius);
+		Coverage = saturate(WorldStrandRadius / max(WorldStrandRadius, PixelRadius) * MaterialPass_HairCoverageScale);
 
 		SampleCoverage8bit = min(uint(Coverage * 0xFF), 0xFFu);
 	}
