Add optional accurate sRGB encoding for Strata pixel payload.

diff --git a/Engine/Shaders/Private/Strata/Strata.ush b/Engine/Shaders/Private/Strata/Strata.ush
--- a/Engine/Shaders/Private/Strata/Strata.ush
+++ b/Engine/Shaders/Private/Strata/Strata.ush
@@ -21,6 +21,7 @@
 #include "../OctahedralCommon.ush"
 #include "../SubsurfaceProfileCommon.ush"
 #include "../BurleyNormalizedSSSCommon.ush"
+#include "../GammaCorrectionCommon.ush"
 #include "/Engine/Shared/StrataDefinitions.h"
 #include "/Engine/Private/Strata/StrataStatisticalOperators.ush"
 
@@ -2975,6 +2976,24 @@ float StrataGetLegacyShadingModels(in FStrataBSDF BSDF)
 
 float3 sqrtFast3(float3 In) { return float3(sqrtFast(In.x), sqrtFast(In.y), sqrtFast(In.z)); }
 
+half3 StrataLinearToSrgb(half3 In)
+{
+#if STRATA_USE_ACCURATE_SRGB
+	return LinearToSrgb(In);
+#else
+	return sqrtFast3(In); // Simple gamma 2.0
+#endif
+}
+
+half3 StrataSrgbToLinear(half3 In)
+{
+#if STRATA_USE_ACCURATE_SRGB
+	return sRGBToLinear(In);
+#else
+	return In * In; // Simple gamma 2.0
+#endif
+}
+
 uint Pack10F(float Value)
 {
 	return (f32tof16(Value) >> 5) & 0x000003FF;
@@ -3029,24 +3048,24 @@ float4 UnpackR16G16B16A16F(uint2 In)
 
 uint PackColorLinearToGamma2(float3 In)
 {
-	return PackRGBA8(float4(sqrtFast3(saturate(In.rgb)), 0)); // Simple gamma 2.0
+	return PackRGBA8(float4(StrataLinearToSrgb(saturate(In.rgb)), 0));
 }
 
 float3 UnpackColorGamma2ToLinear(uint In)
 {
 	float3 Out = UnpackRGBA8(In).xyz;
-	return Out.rgb * Out.rgb; // Simple gamma 2.0	
+	return StrataSrgbToLinear(Out.rgb);
 }
 
 uint PackColorLinearToGamma2AlphaLinear(float4 In)
 {
-	return PackRGBA8(float4(sqrtFast3(saturate(In.rgb)), In.a)); // Simple gamma 2.0
+	return PackRGBA8(float4(StrataLinearToSrgb(saturate(In.rgb)), In.a));
 }
 
 float4 UnpackColorGamma2ToLinearAlphaLinear(uint In)
 {
 	float4 Out = UnpackRGBA8(In);
-	return float4(Out.rgb * Out.rgb, Out.a); // Simple gamma 2.0	
+	return float4(StrataSrgbToLinear(Out.rgb), Out.a);
 }
 
 uint PackR8(float Value)
@@ -3104,7 +3123,7 @@ uint PackR7G7B6Gamma2(float3 rgb, float Dither)
 	const float3 Dither776 = float3(Dither7bits, Dither7bits, Dither6bits);
 
 	// Dither value to hide more aggressive bit reduction.
-	const float3 dithered_rgb = DitherIfNonBlack(sqrtFast3(saturate(rgb)), Dither776);
+	const float3 dithered_rgb = DitherIfNonBlack(StrataLinearToSrgb(saturate(rgb)), Dither776);
 
 	uint r = (uint(dithered_rgb.r * 127.0f) << 0);
 	uint g = (uint(dithered_rgb.g * 127.0f) << 7);
@@ -3118,7 +3137,8 @@ float3 UnpackR7G7B6Gamma2(uint rgb)
 	Out.r = float((rgb >> 0)  & 0x7F) * (1.0f / 127.0f);
 	Out.g = float((rgb >> 7)  & 0x7F) * (1.0f / 127.0f);
 	Out.b = float((rgb >> 14) & 0x3F) * (1.0f /  63.0f);
-	return Out * Out;
+
+	return StrataSrgbToLinear(Out);
 }
 
 uint PackR7G7B6Linear(float3 rgb, float Dither)
diff --git a/Engine/Source/Runtime/Engine/Private/ShaderCompiler/ShaderCompiler.cpp b/Engine/Source/Runtime/Engine/Private/ShaderCompiler/ShaderCompiler.cpp
--- a/Engine/Source/Runtime/Engine/Private/ShaderCompiler/ShaderCompiler.cpp
+++ b/Engine/Source/Runtime/Engine/Private/ShaderCompiler/ShaderCompiler.cpp
@@ -6430,6 +6430,9 @@ void GlobalBeginCompileShader(
 			Input.Environment.SetDefine(TEXT("STRATA_USE_DBUFFER_PASS"), bStrataDBufferPass ? 1 : 0);
 		}
 
+		const bool bStrataUseAccurateSRGB = bStrata && Strata::IsAccurateSRGBEnabled();
+		Input.Environment.SetDefine(TEXT("STRATA_USE_ACCURATE_SRGB"), bStrataUseAccurateSRGB ? 1 : 0);
+
 		const bool bStrataBackCompatibility = bStrata && Strata::IsBackCompatibilityEnabled();
 		Input.Environment.SetDefine(TEXT("PROJECT_STRATA_BACKCOMPATIBILITY"), bStrataBackCompatibility ? 1 : 0);
 
diff --git a/Engine/Source/Runtime/RenderCore/Private/RenderUtils.cpp b/Engine/Source/Runtime/RenderCore/Private/RenderUtils.cpp
--- a/Engine/Source/Runtime/RenderCore/Private/RenderUtils.cpp
+++ b/Engine/Source/Runtime/RenderCore/Private/RenderUtils.cpp
@@ -1568,6 +1568,12 @@ static TAutoConsoleVariable<int32> CVarStrataTileCoord8Bits(
 	TEXT("Format of tile coord. This variable is read-only."),
 	ECVF_ReadOnly | ECVF_RenderThreadSafe);
 
+static TAutoConsoleVariable<int32> CVarStrataAccurateSRGB(
+	TEXT("r.Strata.AccurateSRGB"),
+	0,
+	TEXT("Enable accurate sRGB encoding for pixel data encoding/decoding during rendering passes (base pass, lighting, ...). Otherwise use a simple 'Gamma2' approximation."),
+	ECVF_ReadOnly | ECVF_RenderThreadSafe);
+
 namespace Strata
 {
 	bool IsStrataEnabled()
@@ -1635,6 +1641,11 @@ namespace Strata
 		return CVarStrataTileCoord8Bits.GetValueOnAnyThread() > 0 ? 1 : 0;
 	}
 
+	bool IsAccurateSRGBEnabled()
+	{
+		return CVarStrataAccurateSRGB.GetValueOnAnyThread() > 0 ? 1 : 0;
+	}
+
 	uint32 GetShadingQuality()
 	{
 		return CVarStrataShadingQuality.GetValueOnAnyThread();
diff --git a/Engine/Source/Runtime/RenderCore/Private/Shader.cpp b/Engine/Source/Runtime/RenderCore/Private/Shader.cpp
--- a/Engine/Source/Runtime/RenderCore/Private/Shader.cpp
+++ b/Engine/Source/Runtime/RenderCore/Private/Shader.cpp
@@ -1953,6 +1953,11 @@ void ShaderMapAppendKeyString(EShaderPlatform Platform, FString& KeyString)
 		{
 			KeyString += FString::Printf(TEXT("_ADVDEBUG"));
 		}
+
+		if (Strata::IsAccurateSRGBEnabled())
+		{
+			KeyString += FString::Printf(TEXT("_SRGB"));
+		}
 	}
 
 	{
diff --git a/Engine/Source/Runtime/RenderCore/Public/RenderUtils.h b/Engine/Source/Runtime/RenderCore/Public/RenderUtils.h
--- a/Engine/Source/Runtime/RenderCore/Public/RenderUtils.h
+++ b/Engine/Source/Runtime/RenderCore/Public/RenderUtils.h
@@ -487,6 +487,7 @@ namespace Strata
 	RENDERCORE_API bool IsOpaqueRoughRefractionEnabled();
 	RENDERCORE_API bool IsAdvancedVisualizationEnabled();
 	RENDERCORE_API bool Is8bitTileCoordEnabled();
+	RENDERCORE_API bool IsAccurateSRGBEnabled();
 
 	RENDERCORE_API uint32 GetRayTracingMaterialPayloadSizeInBytes();
 	RENDERCORE_API uint32 GetRayTracingMaterialPayloadSizeInBytes(EShaderPlatform InPlatform);
