Fixed: WPO does not work on Mobile with GPUScene enabled

diff --git a/Engine/Shaders/Private/InstanceCulling/BuildInstanceDrawCommands.usf b/Engine/Shaders/Private/InstanceCulling/BuildInstanceDrawCommands.usf
--- a/Engine/Shaders/Private/InstanceCulling/BuildInstanceDrawCommands.usf
+++ b/Engine/Shaders/Private/InstanceCulling/BuildInstanceDrawCommands.usf
@@ -128,7 +128,8 @@ void WritePackedInstanceData(uint Offset, uint InstanceId, FInstanceSceneData In
 	InstanceAuxData1 = 0; // TODO: pack something useful here 
 #endif // ALLOW_STATIC_LIGHTING
 
-	const uint PackedPrimitiveFlags = ((PrimitiveData.Flags & 0x01FFFFFF) | (InstanceData.Flags << 25u));
+	// 24 bits for Primitive flags + 8 bits for Instance flags
+	const uint PackedPrimitiveFlags = ((PrimitiveData.Flags & 0x00FFFFFF) | (InstanceData.Flags << 24u));
 
 	// This will only write out the relative instance transform, but that's fine; the tile coordinate of the instance transform comes from the primitive data
 	float4 InstanceOriginAndId = InstanceData.LocalToWorld.M[3];
diff --git a/Engine/Shaders/Private/SceneData.ush b/Engine/Shaders/Private/SceneData.ush
--- a/Engine/Shaders/Private/SceneData.ush
+++ b/Engine/Shaders/Private/SceneData.ush
@@ -75,7 +75,7 @@
 #define PRIMITIVE_SCENE_DATA_FLAG_LIGHTING_CHANNEL_1					0x1000
 #define PRIMITIVE_SCENE_DATA_FLAG_LIGHTING_CHANNEL_2					0x2000
 #define PRIMITIVE_SCENE_DATA_FLAG_HAS_INSTANCE_LOCAL_BOUNDS				0x4000
-#define PRIMITIVE_SCENE_DATA_FLAG_HAS_NANITE_IMPOSTER					0x8000
+#define PRIMITIVE_SCENE_DATA_FLAG_EVALUATE_WORLD_POSITION_OFFSET		0x8000
 #define PRIMITIVE_SCENE_DATA_FLAG_VISIBLE_IN_GAME						0x10000
 #define PRIMITIVE_SCENE_DATA_FLAG_VISIBLE_IN_EDITOR						0x20000
 #define PRIMITIVE_SCENE_DATA_FLAG_VISIBLE_IN_REFLECTION_CAPTURES		0x40000
@@ -83,9 +83,9 @@
 #define PRIMITIVE_SCENE_DATA_FLAG_VISIBLE_IN_RAY_TRACING				0x100000
 #define PRIMITIVE_SCENE_DATA_FLAG_VISIBLE_IN_SCENE_CAPTURE_ONLY			0x200000
 #define PRIMITIVE_SCENE_DATA_FLAG_HIDDEN_IN_SCENE_CAPTURE				0x400000
-#define PRIMITIVE_SCENE_DATA_FLAG_FORCE_HIDDEN							0x800000
+#define PRIMITIVE_SCENE_DATA_FLAG_FORCE_HIDDEN							0x800000 // Mobile allocates only 24 bits for primitive flags
 #define PRIMITIVE_SCENE_DATA_FLAG_CAST_HIDDEN_SHADOW					0x1000000
-#define PRIMITIVE_SCENE_DATA_FLAG_EVALUATE_WORLD_POSITION_OFFSET		0x2000000
+#define PRIMITIVE_SCENE_DATA_FLAG_HAS_NANITE_IMPOSTER					0x2000000
 #define PRIMITIVE_SCENE_DATA_FLAG_INSTANCE_DRAW_DISTANCE_CULL			0x4000000
 #define PRIMITIVE_SCENE_DATA_FLAG_WPO_DISABLE_DISTANCE					0x8000000
 #define PRIMITIVE_SCENE_DATA_FLAG_WRITES_CUSTOM_DEPTH_STENCIL			0x10000000
@@ -129,7 +129,7 @@ uint GetDefaultPrimitiveSceneDataFlags()
 #define INSTANCE_SCENE_DATA_FLAG_HAS_LIGHTSHADOW_UV_BIAS	0x10
 #define INSTANCE_SCENE_DATA_FLAG_HAS_HIERARCHY_OFFSET		0x20
 #define INSTANCE_SCENE_DATA_FLAG_HAS_LOCAL_BOUNDS			0x40
-#define INSTANCE_SCENE_DATA_FLAG_HAS_EDITOR_DATA			0x80
+#define INSTANCE_SCENE_DATA_FLAG_HAS_EDITOR_DATA			0x80 // Mobile allocates only 8 bits for instance flags
 #define INSTANCE_SCENE_DATA_FLAG_IS_RAYTRACING_FAR_FIELD	0x100
 
 #if PLATFORM_ALLOW_SCENE_DATA_COMPRESSED_TRANSFORMS
@@ -1149,10 +1149,10 @@ FSceneDataIntermediates GetSceneDataIntermediates(uint DrawInstanceId, float4 In
 	float3 TilePosition = float3(0,0,0);
 
 	// TODO: pack important primitive and instance flags here 
-	// (25 + 7)
+	// (24 + 8)
 	const uint PackedPrimitiveFlags = asuint(InstanceTransform1.w);
-	const uint PrimitiveFlags = (PackedPrimitiveFlags & 0x01FFFFFF);
-	const uint InstanceFlags = (PackedPrimitiveFlags >> 25u);
+	const uint PrimitiveFlags = (PackedPrimitiveFlags & 0x00FFFFFF);
+	const uint InstanceFlags = (PackedPrimitiveFlags >> 24u);
 
 	// Reconstruct InstanceData from a packed data
 	FInstanceSceneData InstanceData = (FInstanceSceneData)0;
diff --git a/Engine/Source/Runtime/Engine/Public/InstanceUniformShaderParameters.h b/Engine/Source/Runtime/Engine/Public/InstanceUniformShaderParameters.h
--- a/Engine/Source/Runtime/Engine/Public/InstanceUniformShaderParameters.h
+++ b/Engine/Source/Runtime/Engine/Public/InstanceUniformShaderParameters.h
@@ -19,7 +19,7 @@
 #define INSTANCE_SCENE_DATA_FLAG_HAS_LIGHTSHADOW_UV_BIAS	0x10
 #define INSTANCE_SCENE_DATA_FLAG_HAS_HIERARCHY_OFFSET		0x20
 #define INSTANCE_SCENE_DATA_FLAG_HAS_LOCAL_BOUNDS			0x40
-#define INSTANCE_SCENE_DATA_FLAG_HAS_EDITOR_DATA			0x80
+#define INSTANCE_SCENE_DATA_FLAG_HAS_EDITOR_DATA			0x80 // Mobile allocates only 8 bits for instance flags
 #define INSTANCE_SCENE_DATA_FLAG_IS_RAYTRACING_FAR_FIELD	0x100
 
 #define INSTANCE_SCENE_DATA_FLAG_PAYLOAD_MASK ( \
diff --git a/Engine/Source/Runtime/Engine/Public/PrimitiveUniformShaderParameters.h b/Engine/Source/Runtime/Engine/Public/PrimitiveUniformShaderParameters.h
--- a/Engine/Source/Runtime/Engine/Public/PrimitiveUniformShaderParameters.h
+++ b/Engine/Source/Runtime/Engine/Public/PrimitiveUniformShaderParameters.h
@@ -83,7 +83,7 @@ END_GLOBAL_SHADER_PARAMETER_STRUCT()
 #define PRIMITIVE_SCENE_DATA_FLAG_LIGHTING_CHANNEL_1					0x1000
 #define PRIMITIVE_SCENE_DATA_FLAG_LIGHTING_CHANNEL_2					0x2000
 #define PRIMITIVE_SCENE_DATA_FLAG_HAS_INSTANCE_LOCAL_BOUNDS				0x4000
-#define PRIMITIVE_SCENE_DATA_FLAG_HAS_NANITE_IMPOSTER					0x8000
+#define PRIMITIVE_SCENE_DATA_FLAG_EVALUATE_WORLD_POSITION_OFFSET		0x8000
 #define PRIMITIVE_SCENE_DATA_FLAG_VISIBLE_IN_GAME						0x10000
 #define PRIMITIVE_SCENE_DATA_FLAG_VISIBLE_IN_EDITOR						0x20000
 #define PRIMITIVE_SCENE_DATA_FLAG_VISIBLE_IN_REFLECTION_CAPTURES		0x40000
@@ -91,9 +91,9 @@ END_GLOBAL_SHADER_PARAMETER_STRUCT()
 #define PRIMITIVE_SCENE_DATA_FLAG_VISIBLE_IN_RAY_TRACING				0x100000
 #define PRIMITIVE_SCENE_DATA_FLAG_VISIBLE_IN_SCENE_CAPTURE_ONLY			0x200000
 #define PRIMITIVE_SCENE_DATA_FLAG_HIDDEN_IN_SCENE_CAPTURE				0x400000
-#define PRIMITIVE_SCENE_DATA_FLAG_FORCE_HIDDEN							0x800000
+#define PRIMITIVE_SCENE_DATA_FLAG_FORCE_HIDDEN							0x800000 // Mobile allocates only 24 bits for primitive flags
 #define PRIMITIVE_SCENE_DATA_FLAG_CAST_HIDDEN_SHADOW					0x1000000
-#define PRIMITIVE_SCENE_DATA_FLAG_EVALUATE_WORLD_POSITION_OFFSET		0x2000000
+#define PRIMITIVE_SCENE_DATA_FLAG_HAS_NANITE_IMPOSTER					0x2000000
 #define PRIMITIVE_SCENE_DATA_FLAG_INSTANCE_DRAW_DISTANCE_CULL			0x4000000
 #define PRIMITIVE_SCENE_DATA_FLAG_WPO_DISABLE_DISTANCE					0x8000000
 #define PRIMITIVE_SCENE_DATA_FLAG_WRITES_CUSTOM_DEPTH_STENCIL			0x10000000
