Improve hair instance stats for better clarity.

diff --git a/Engine/Shaders/Private/HairStrands/HairStrandsDebug.usf b/Engine/Shaders/Private/HairStrands/HairStrandsDebug.usf
--- a/Engine/Shaders/Private/HairStrands/HairStrandsDebug.usf
+++ b/Engine/Shaders/Private/HairStrands/HairStrandsDebug.usf
@@ -2,6 +2,81 @@
 
 #include "../Common.ush"
 #include "../HairShadingCommon.ush"
+#include "../ColorMap.ush"
+
+#if SHADER_PRINT_INSTANCE || SHADER_PRINT
+/////////////////////////////////////////////////////////////////////////////////////////
+// Utils function
+
+#include "../ShaderPrint.ush"
+
+// Print float number in 0..1000 aligned on the right
+void PrintNumberAlignedRight(inout FShaderPrintContext Context, float Value, FFontColor InColor)
+{
+	if (Value >= 100.f)
+	{
+		Print(Context, Value, InColor, 5, 1);
+	}
+	else if (Value >= 10.f)
+	{
+		Print(Context, TEXT(" "), InColor);
+		Print(Context, Value, InColor, 4, 1);
+	}
+	else
+	{
+		Print(Context, TEXT("  "), InColor);
+		Print(Context, Value, InColor, 3, 1);
+	}
+}
+
+// Print 
+void PrintLargeNumber(inout FShaderPrintContext Context, uint Value, FFontColor InColor)
+{
+	if (Value > 1000000000)
+	{
+		const float fValue = Value / 1000000000.f;
+		//Print(Context, fValue, InColor, 5, 1);
+		PrintNumberAlignedRight(Context, fValue, InColor);
+		Print(Context, TEXT("B"), InColor);
+	}
+	else if (Value > 1000000)
+	{
+		const float fValue = Value / 1000000.f;
+		//Print(Context, fValue, InColor, 5, 1);
+		PrintNumberAlignedRight(Context, fValue, InColor);
+		Print(Context, TEXT("M"), InColor);
+	}
+	else if (Value > 1000)
+	{
+		const float fValue = Value / 1000.f;
+		//Print(Context, fValue, InColor, 5, 1);
+		PrintNumberAlignedRight(Context, fValue, InColor);
+		Print(Context, TEXT("k"), InColor);
+	}
+	else
+	{
+		Print(Context, TEXT("   "), InColor);
+		Print(Context, Value, InColor, 3);
+	}
+}
+
+void PrintPercentage(inout FShaderPrintContext Context, float Value, FFontColor InColor)
+{
+	const float fValue = clamp(Value * 100.f, 0.f, 99.9f);
+	if (fValue < 10.f)
+	{
+		Print(Context, TEXT(" "), InColor);
+		Print(Context, fValue, InColor, 3, 1);
+	}
+	else
+	{
+		Print(Context, fValue, InColor, 4, 1);
+	}
+	Print(Context, TEXT("% "), InColor);
+}
+
+#endif // SHADER_PRINT_INSTANCE || SHADER_PRINT
+
 /////////////////////////////////////////////////////////////////////////////////////////
 // Visualize hair strands data as full screen view mode (BaseColor, Tangent, Roughness, ...)
 
@@ -879,20 +954,22 @@ void MainCS(uint InGroupIndex : SV_GroupIndex, uint3 DispatchThreadId : SV_Dispa
 		const uint2 EquivalentResolution = ceil(sqrt(UsedSampleCount));
 	
 		// Pixel coord
-		FShaderPrintContext Context = InitShaderPrintContext(true, uint2(450, 50));
+		FShaderPrintContext Context = InitShaderPrintContext(true, uint2(450, 100));
 		Print(Context, TEXT("Sample Count (Used/Alloc)  "), FontCyan);
-		Print(Context, UsedSampleCount, FontCyan);
-		Print(Context, AllocatedSampleCount, FontCyan);
+		PrintLargeNumber(Context, UsedSampleCount, FontCyan);
+		Print(Context, TEXT(" / "), FontCyan);
+		PrintLargeNumber(Context, AllocatedSampleCount, FontCyan);
+		Print(Context, TEXT(" "), FontCyan);
 
 		const float AllocationPercentage = float(UsedSampleCount) / float(AllocatedSampleCount);
-		const FFontColor OccupancyColor = GetOccupancyColor(AllocationPercentage);
-		Print(Context, TEXT(" %% "), OccupancyColor);
-		Print(Context, AllocationPercentage * 100.f, OccupancyColor);
+		PrintPercentage(Context, AllocationPercentage, GetOccupancyColor(AllocationPercentage));
 
 		Newline(Context);
 
 		Print(Context, TEXT("Sample Lighting Resolution "), FontCyan);
-		Print(Context, EquivalentResolution, FontCyan);
+		Print(Context, EquivalentResolution.x, FontCyan, 4);
+		Print(Context, TEXT("x"), FontSilver);
+		Print(Context, EquivalentResolution.y, FontCyan, 4);
 		Newline(Context);
 
 		Print(Context, TEXT("MacroGroup Count           "), FontYellow);
@@ -1059,7 +1136,6 @@ void MainCS(uint InGroupIndex : SV_GroupIndex, uint3 DispatchThreadId : SV_Dispa
 
 #if SHADER_PRINT_INSTANCE
 
-#include "../ShaderPrint.ush"
 #include "HairStrandsAABBCommon.ush"
 
 int2					MaxResolution;
@@ -1216,7 +1292,7 @@ void MainCS(uint GroupIndex : SV_GroupIndex, uint3 DispatchThreadId : SV_Dispatc
 		return;
 	}
 
-	FShaderPrintContext Context = InitShaderPrintContext(true, uint2(50, 50));
+	FShaderPrintContext Context = InitShaderPrintContext(true, uint2(50, 100));
 	// Draw legend
 
 	// Draw totals
@@ -1238,36 +1314,34 @@ void MainCS(uint GroupIndex : SV_GroupIndex, uint3 DispatchThreadId : SV_Dispatc
 		Newline(Context);
 
 		Print(Context, TEXT("Total Strand Curves : "), FontWhite);
-		Print(Context,TotalStrandsCount, FontOrange, 8);
+		PrintLargeNumber(Context,TotalStrandsCount, FontOrange);
 		Newline(Context);
 
 		Print(Context, TEXT("Total Strand Vertex : "), FontWhite);
-		Print(Context,TotalVertexCount, FontSilver, 8);
+		PrintLargeNumber(Context,TotalVertexCount, FontSilver);
 		Newline(Context);
 	}
 	Newline(Context);
 
 	{
-		Print(Context, TEXT("Total Strand (primary)        : "), FontWhite);
-		Print(Context,InstanceCount_StrandsPrimaryView, FontYellow, 8);
-		Newline(Context);
-
-		Print(Context, TEXT("Total Strand (shadow)         : "), FontWhite);
-		Print(Context,InstanceCount_StrandsShadowView, FontYellow, 8);
-		Newline(Context);
-
-		Print(Context, TEXT("Total Cards / Meshes (primary): "), FontWhite);
-		Print(Context,InstanceCount_CardsOrMeshesPrimaryView, FontYellow, 8);
+		Print(Context, TEXT("Total Strand        : "), FontWhite);
+		Print(Context, TEXT(" Visible:"), FontSilver);
+		Print(Context,InstanceCount_StrandsPrimaryView, FontYellow, 4);
+		Print(Context, TEXT("| Shadow:"), FontSilver);
+		Print(Context,InstanceCount_StrandsShadowView, FontYellow, 4);
 		Newline(Context);
 
-		Print(Context, TEXT("Total Cards / Meshes (shadow) : "), FontWhite);
-		Print(Context, InstanceCount_CardsOrMeshesShadowView, FontYellow, 8);
+		Print(Context, TEXT("Total Cards/Meshes  : "), FontWhite);
+		Print(Context, TEXT(" Visible:"), FontSilver);
+		Print(Context,InstanceCount_CardsOrMeshesPrimaryView, FontYellow, 4);
+		Print(Context, TEXT("| Shadow:"), FontSilver);
+		Print(Context, InstanceCount_CardsOrMeshesShadowView, FontYellow, 4);
 		Newline(Context);
 	}
 	Newline(Context);
 
 	// Draw instance details
-	Print(Context, TEXT("   Group  LOD     Geometry Binding Sim   RBF   Clip Screen Strands          Vertex           Bind.Tri/Sect Primary Shadow Name"));
+	Print(Context, TEXT("   Group  LOD     Geometry Binding Sim RBF Clip Screen Strands             Vertex              Bind.Tri/Sect Vis. Shad. Name"));
 	Newline(Context);
 
 	for (uint InstanceIndex = 0; InstanceIndex < InstanceCount; ++InstanceIndex)
@@ -1302,18 +1376,10 @@ void MainCS(uint GroupIndex : SV_GroupIndex, uint3 DispatchThreadId : SV_Dispatc
 		else if (Info.BindingType == 2)	{ Print(Context, TEXT("None    "), FontTurquoise); }
 
 		// Sim
-		if (Info.bHasSim)
-			Print(Context, TEXT("True "), FontGreen);
-		else
-			Print(Context, TEXT("False"), FontRed);
-		Print(Context, TEXT(" "), FontWhite);
+		PrintBool(Context, Info.bHasSim);
 
 		// RBF
-		if (Info.bHasRBF)
-			Print(Context, TEXT("True "), FontGreen);
-		else
-			Print(Context, TEXT("False"), FontRed);
-		Print(Context, TEXT(" "));
+		PrintBool(Context, Info.bHasRBF);
 
 		// Clip
 		if (Info.ClipValue >= 0)
@@ -1325,15 +1391,21 @@ void MainCS(uint GroupIndex : SV_GroupIndex, uint3 DispatchThreadId : SV_Dispatc
 		Print(Context, Info.ScreenSize, FontOrange, 4, 2);
 		Print(Context, TEXT("   "));
 
-		// Strands/Vertex sCount
-		const uint EffectiveStrandsCount = ceil(Info.StrandsCount * Info.ScreenSize);
-		Print(Context, Info.ActiveCurveCount, FontOrange, 8);
+		// Curve sCount
+		const float CurveRatio = Info.ActiveCurveCount / float(Info.StrandsCount);
+		PrintPercentage(Context, CurveRatio, InitFontColor(ColorMapTurbo(CurveRatio)));
+		PrintLargeNumber(Context, Info.ActiveCurveCount, FontOrange);
 		Print(Context, TEXT("/"));
-		Print(Context,Info.StrandsCount, FontSilver, 8);
+		PrintLargeNumber(Context, Info.StrandsCount, FontSilver);
+		Print(Context, TEXT(" "), FontWhite);
 
-		Print(Context, Info.ActivePointCount, FontOrange, 8);
+		// Point sCount
+		const float VertexRatio = Info.ActivePointCount / float(Info.VertexCount);
+		PrintPercentage(Context, VertexRatio, InitFontColor(ColorMapTurbo(VertexRatio)));
+		PrintLargeNumber(Context, Info.ActivePointCount, FontOrange);
 		Print(Context, TEXT("/"));
-		Print(Context,Info.VertexCount, FontSilver, 8);
+		PrintLargeNumber(Context, Info.VertexCount, FontSilver);
+		Print(Context, TEXT(" "), FontWhite);
 
 		// Binding data
 		Print(Context,Info.UniqueTriangle, FontOrange, 8);
@@ -1342,19 +1414,12 @@ void MainCS(uint GroupIndex : SV_GroupIndex, uint3 DispatchThreadId : SV_Dispatc
 		Print(Context, TEXT("  "));
 
 		// Primary view
-		if (Info.bIsVisibleInPrimaryView)
-			Print(Context, TEXT("True "), FontGreen);
-		else
-			Print(Context, TEXT("False"), FontRed);
-		Print(Context, TEXT(" "), FontWhite);
+		PrintBool(Context, Info.bIsVisibleInPrimaryView);
+		Print(Context, TEXT(" "));
 
 		// Shadow view
+		PrintBool(Context, Info.bIsVisibleInShadowView);
 		Print(Context, TEXT("  "));
-		if (Info.bIsVisibleInShadowView)
-			Print(Context, TEXT("True "), FontGreen);
-		else
-			Print(Context, TEXT("False"), FontRed);
-		Print(Context, TEXT("  "), FontWhite);
 
 		// Name
 		PrintInstanceName(Context, InstanceIndex, FontWhite);
