Apply PreExposure to Heterogeneous Volumes radiance formats.

diff --git a/Engine/Shaders/Private/HeterogeneousVolumes/HeterogeneousVolumesHardwareRayTracing.usf b/Engine/Shaders/Private/HeterogeneousVolumes/HeterogeneousVolumesHardwareRayTracing.usf
--- a/Engine/Shaders/Private/HeterogeneousVolumes/HeterogeneousVolumesHardwareRayTracing.usf
+++ b/Engine/Shaders/Private/HeterogeneousVolumes/HeterogeneousVolumesHardwareRayTracing.usf
@@ -471,7 +471,7 @@ RAY_TRACING_ENTRY_RAYGEN(RenderLightingCacheWithPreshadingRGS)
 
 #elif DIM_LIGHTING_CACHE_MODE == LIGHTING_CACHE_INSCATTERING
 	float3 Inscattering = ComputeInscattering(WorldRayOrigin, LightData, LightType, MaxStepCount, bApplyShadowTransmittance);
-	RWLightingCacheTexture[VoxelIndex] += Inscattering;
+	RWLightingCacheTexture[VoxelIndex] += Inscattering * View.PreExposure;
 #endif // DIM_LIGHTING_CACHE_MODE
 }
 
@@ -590,5 +590,5 @@ RAY_TRACING_ENTRY_RAYGEN(RenderSingleScatteringWithPreshadingRGS)
 
 	// Output..
 	float3 Opacity = bApplyEmissionAndTransmittance ? 1.0 - Transmittance : 0.0;
-	RWLightingTexture[PixelCoord] += float4(Radiance, Luminance(Opacity));
+	RWLightingTexture[PixelCoord] += float4(Radiance * View.PreExposure, Luminance(Opacity));
 }
diff --git a/Engine/Shaders/Private/HeterogeneousVolumes/HeterogeneousVolumesLiveShadingPipeline.usf b/Engine/Shaders/Private/HeterogeneousVolumes/HeterogeneousVolumesLiveShadingPipeline.usf
--- a/Engine/Shaders/Private/HeterogeneousVolumes/HeterogeneousVolumesLiveShadingPipeline.usf
+++ b/Engine/Shaders/Private/HeterogeneousVolumes/HeterogeneousVolumesLiveShadingPipeline.usf
@@ -246,7 +246,7 @@ void RenderLightingCacheWithLiveShadingCS(
 
 #elif DIM_LIGHTING_CACHE_MODE == LIGHTING_CACHE_INSCATTERING
 	float3 Inscattering = ComputeInscattering(WorldRayOrigin, LightData, LightType, MaxStepCount, bApplyShadowTransmittance);
-	RWLightingCacheTexture[VoxelIndex] += Inscattering;
+	RWLightingCacheTexture[VoxelIndex] += Inscattering * View.PreExposure;
 #endif // DIM_LIGHTING_CACHE_MODE
 }
 
@@ -335,7 +335,7 @@ void RenderSingleScatteringWithLiveShadingCS(
 
 		// Output..
 		float3 Opacity = bApplyEmissionAndTransmittance ? 1.0 - Transmittance : 0.0;
-		RWLightingTexture[PixelCoord] += float4(Radiance, Luminance(Opacity));
+		RWLightingTexture[PixelCoord] += float4(Radiance * View.PreExposure, Luminance(Opacity));
 
 		// Debug output
 		//RWDebugOutputBuffer[LinearIndex] = DebugOutput;
diff --git a/Engine/Shaders/Private/HeterogeneousVolumes/HeterogeneousVolumesPreshadingPipeline.usf b/Engine/Shaders/Private/HeterogeneousVolumes/HeterogeneousVolumesPreshadingPipeline.usf
--- a/Engine/Shaders/Private/HeterogeneousVolumes/HeterogeneousVolumesPreshadingPipeline.usf
+++ b/Engine/Shaders/Private/HeterogeneousVolumes/HeterogeneousVolumesPreshadingPipeline.usf
@@ -492,7 +492,7 @@ void RenderLightingCacheWithPreshadingCS(
 
 #elif DIM_LIGHTING_CACHE_MODE == LIGHTING_CACHE_INSCATTERING
 	float3 Inscattering = ComputeInscattering(WorldRayOrigin, LightData, LightType, MaxStepCount, bApplyShadowTransmittance);
-	RWLightingCacheTexture[VoxelIndex] += Inscattering;
+	RWLightingCacheTexture[VoxelIndex] += Inscattering * View.PreExposure;
 #endif // DIM_LIGHTING_CACHE_MODE
 }
 
@@ -622,7 +622,7 @@ void RenderSingleScatteringWithPreshadingCS(
 
 	// Output..
 	float3 Opacity = bApplyEmissionAndTransmittance ? 1.0 - Transmittance : 0.0;
-	RWLightingTexture[PixelCoord] += float4(Radiance, Luminance(Opacity));
+	RWLightingTexture[PixelCoord] += float4(Radiance * View.PreExposure, Luminance(Opacity));
 
 #if DIM_DEBUG
 	//if (!all(DispatchThreadId))
diff --git a/Engine/Shaders/Private/HeterogeneousVolumes/HeterogeneousVolumesRayMarchingUtils.ush b/Engine/Shaders/Private/HeterogeneousVolumes/HeterogeneousVolumesRayMarchingUtils.ush
--- a/Engine/Shaders/Private/HeterogeneousVolumes/HeterogeneousVolumesRayMarchingUtils.ush
+++ b/Engine/Shaders/Private/HeterogeneousVolumes/HeterogeneousVolumesRayMarchingUtils.ush
@@ -274,7 +274,7 @@ float3 ComputeInscattering(
 	float3 TransmittanceUVW = saturate((LocalShadowRayOrigin - LocalBoundsMin) / (2.0 * GetLocalBoundsExtent()));
 	float MipLevel = 0;
 
-	Inscattering = SampleLightingCache(TransmittanceUVW, MipLevel);
+	Inscattering = SampleLightingCache(TransmittanceUVW, MipLevel) * View.OneOverPreExposure;
 
 #else // DIM_USE_INSCATTERING_VOLUME
 	float3 L = LightData.Direction;
