- Fix for TransformLocalPositionToPrevWorld using GetLocalToWorld vs GetPrevLocalToWorld

diff --git a/Engine/Shaders/Private/MaterialTemplate.ush b/Engine/Shaders/Private/MaterialTemplate.ush
--- a/Engine/Shaders/Private/MaterialTemplate.ush
+++ b/Engine/Shaders/Private/MaterialTemplate.ush
@@ -1287,7 +1287,7 @@ FLWCVector3 TransformLocalPositionToWorld(FMaterialVertexParameters Parameters,f
 /** Transforms a position from local space to previous frame absolute world space */
 FLWCVector3 TransformLocalPositionToPrevWorld(FMaterialVertexParameters Parameters,float3 InLocalPosition)
 {
-	return LWCMultiply(InLocalPosition, GetLocalToWorld(Parameters));
+	return LWCMultiply(InLocalPosition, GetPrevLocalToWorld(Parameters));
 }
 
 /** Return the object's position in world space */
