Fixed Nanite material CS quad binning corruption (binning was correct, but actual shade had a quad index overrun that would fetch valid quads from other random materials adjacent in memory).

diff --git a/Engine/Shaders/Private/ComputeShaderOutputCommon.ush b/Engine/Shaders/Private/ComputeShaderOutputCommon.ush
--- a/Engine/Shaders/Private/ComputeShaderOutputCommon.ush
+++ b/Engine/Shaders/Private/ComputeShaderOutputCommon.ush
@@ -374,14 +374,23 @@ void MainCS(uint ThreadIndex : SV_GroupIndex, uint GroupID : SV_GroupID, uint3 G
 	
 	const uint LaneTL = WaveGetLaneIndex() & ~3u;
 
-	uint PackedQuad;
+	uint PackedQuad = 0u;
 	if (bIsQuadTL)
 	{
-		const uint QuadStart = Nanite.ShadingBinMeta[ShadingBin].z;
-		PackedQuad = Nanite.ShadingBinQuads[QuadStart + QuadIndex];
+		const uint QuadCount = Nanite.ShadingBinMeta[ShadingBin].x;
+		if (QuadIndex < QuadCount)
+		{
+			const uint QuadStart = Nanite.ShadingBinMeta[ShadingBin].z;
+			PackedQuad = Nanite.ShadingBinQuads[QuadStart + QuadIndex];
+		}
 	}
 	PackedQuad = WaveReadLaneAt(PackedQuad, LaneTL);
 	FShadingQuad Quad = UnpackShadingQuad(PackedQuad);
+	if (!any(Quad.Coverage))
+	{
+		// None of the 2x2 pixels have coverage (likely an out of bounds wave).
+		return;
+	}
 
 	// Determine final pixel position on the screen for shading
 	const float2 PixelPos = float2(Quad.TopLeft) + float2(CornerIndex & 1u, (CornerIndex >> 1u) & 1u);
@@ -394,7 +403,9 @@ void MainCS(uint ThreadIndex : SV_GroupIndex, uint GroupID : SV_GroupID, uint3 G
 	float R =  (DispatchThreadID.x & 0xFFu) / 255.0f;
 	float G = ((DispatchThreadID.x & 0xFF00u) >> 8u) / 255.0f;
 	float B = ((DispatchThreadID.x & 0xFF0000u) >> 16u) / 255.0f;
-  	ShadedPixel.MRT[3].rgb = float3(R, G, B);
+	ShadedPixel.MRT[3].rgb = float3(R, G, B);
+#elif 0
+	ShadedPixel.MRT[3].rgb = IntToColor(ShadingBin);
 #elif 0
 	ShadedPixel.MRT[3].rgb = IntToColor(QuadIndex);
 #elif 0
