Remove sRGB-only output restriction in tonemapper for mobile platforms

diff --git a/Engine/Shaders/Private/PostProcessCombineLUTs.usf b/Engine/Shaders/Private/PostProcessCombineLUTs.usf
--- a/Engine/Shaders/Private/PostProcessCombineLUTs.usf
+++ b/Engine/Shaders/Private/PostProcessCombineLUTs.usf
@@ -135,12 +135,7 @@ float OutputMaxLuminance;
 
 uint GetOutputDevice()
 {
-#if FEATURE_LEVEL > FEATURE_LEVEL_ES3_1
 	return OutputDevice;
-#else
-	// Only sRGB output for Mobile
-	return 0;
-#endif 
 }
 
 float BlueCorrection;
diff --git a/Engine/Shaders/Private/PostProcessDeviceEncodingOnly.usf b/Engine/Shaders/Private/PostProcessDeviceEncodingOnly.usf
--- a/Engine/Shaders/Private/PostProcessDeviceEncodingOnly.usf
+++ b/Engine/Shaders/Private/PostProcessDeviceEncodingOnly.usf
@@ -43,12 +43,7 @@ float ACESSceneColorMultiplier;
 
 uint GetOutputDevice()
 {
-#if FEATURE_LEVEL > FEATURE_LEVEL_ES3_1
 	return OutputDevice;
-#else
-	// Only sRGB output for Mobile
-	return 0;
-#endif 
 }
 
 float4 SampleSceneColor(float2 SceneUV)
