Fix a bug that the ForwardLightData in VolumetricCloud.usf is expanded to RenderVolumetricCloudParameters.RenderVolumetricCloudParameters.ForwardLightData.

diff --git a/Engine/Shaders/Private/LightGridCommon.ush b/Engine/Shaders/Private/LightGridCommon.ush
--- a/Engine/Shaders/Private/LightGridCommon.ush
+++ b/Engine/Shaders/Private/LightGridCommon.ush
@@ -19,12 +19,24 @@
 // If this is changed, the LightGridUses16BitBuffers function from LightGridInjection.cpp should also be updated.
 #define LIGHT_GRID_USES_16BIT_BUFFERS (PLATFORM_SUPPORTS_BUFFER_LOAD_TYPE_CONVERSION && !SHADING_PATH_MOBILE)
 
-// ForwardLightData.CulledLightDataGrid could be StructuredBuffer or 16bit Buffer, make sure to access it from this MACRO
+// ForwardLightData.CulledLightDataGrid could be 32bit StructuredBuffer or 16bit Buffer, make sure to access it from this function
+uint GetCulledLightDataGrid(uint GridIndex)
+{
+#if LIGHT_GRID_USES_16BIT_BUFFERS
+	return ForwardLightData.CulledLightDataGrid16Bit[GridIndex];
+#else
+	return ForwardLightData.CulledLightDataGrid32Bit[GridIndex];
+#endif
+}
+
+uint GetCulledLightDataGridISR(uint GridIndex)
+{
 #if LIGHT_GRID_USES_16BIT_BUFFERS
-	#define GetCulledLightDataGrid(UB, GridIndex) UB.CulledLightDataGrid16Bit[GridIndex]
+	return ForwardLightDataISR.CulledLightDataGrid16Bit[GridIndex];
 #else
-	#define GetCulledLightDataGrid(UB, GridIndex) UB.CulledLightDataGrid32Bit[GridIndex]
+	return ForwardLightDataISR.CulledLightDataGrid32Bit[GridIndex];
 #endif
+}
 
 struct FLightGridData
 {
@@ -215,7 +227,7 @@ FLocalLightData GetLocalLightData(uint GridIndex, uint EyeIndex)
 	if (EyeIndex == 0)
 	{
 #endif
-		uint LocalLightIndex = GetCulledLightDataGrid(ForwardLightData, GridIndex);
+		uint LocalLightIndex = GetCulledLightDataGrid(GridIndex);
 		uint LocalLightBaseIndex = LocalLightIndex * LOCAL_LIGHT_DATA_STRIDE;
 
 		Out = GetLocalLightData_Internal(
@@ -233,7 +245,7 @@ FLocalLightData GetLocalLightData(uint GridIndex, uint EyeIndex)
 	}
 	else
 	{
-		uint LocalLightIndex = GetCulledLightDataGrid(ForwardLightDataISR, GridIndex);
+		uint LocalLightIndex = GetCulledLightDataGridISR(GridIndex);
 		uint LocalLightBaseIndex = LocalLightIndex * LOCAL_LIGHT_DATA_STRIDE;
 
 		Out = GetLocalLightData_Internal(
diff --git a/Engine/Shaders/Private/ReflectionEnvironmentComposite.ush b/Engine/Shaders/Private/ReflectionEnvironmentComposite.ush
--- a/Engine/Shaders/Private/ReflectionEnvironmentComposite.ush
+++ b/Engine/Shaders/Private/ReflectionEnvironmentComposite.ush
@@ -48,13 +48,13 @@ float3 CompositeReflectionCapturesAndSkylightTWS(
 		{
 	#endif
 	
-			CaptureIndex = GetCulledLightDataGrid(ForwardLightData, CaptureDataStartIndex + TileCaptureIndex);
+			CaptureIndex = GetCulledLightDataGrid(CaptureDataStartIndex + TileCaptureIndex);
 	
 	#if INSTANCED_STEREO
 		}
 		else
 		{
-			CaptureIndex = GetCulledLightDataGrid(ForwardLightDataISR, CaptureDataStartIndex + TileCaptureIndex);
+			CaptureIndex = GetCulledLightDataGridISR(CaptureDataStartIndex + TileCaptureIndex);
 		}
 	#endif
 #endif
diff --git a/Engine/Shaders/Private/VirtualShadowMaps/VirtualShadowMapPageManagement.usf b/Engine/Shaders/Private/VirtualShadowMaps/VirtualShadowMapPageManagement.usf
--- a/Engine/Shaders/Private/VirtualShadowMaps/VirtualShadowMapPageManagement.usf
+++ b/Engine/Shaders/Private/VirtualShadowMaps/VirtualShadowMapPageManagement.usf
@@ -173,7 +173,7 @@ void PruneLightGridCS(uint GridLinearIndex : SV_DispatchThreadID)
 
 				// Copy light grid data index
 				checkStructuredBufferAccessSlow(OutPrunedLightGridData, CulledLightGridData.DataStartIndex + NumRetainedLights);
-				OutPrunedLightGridData[CulledLightGridData.DataStartIndex + NumRetainedLights++] = GetCulledLightDataGrid(ForwardLightData, CulledLightGridData.DataStartIndex + Index);
+				OutPrunedLightGridData[CulledLightGridData.DataStartIndex + NumRetainedLights++] = GetCulledLightDataGrid(CulledLightGridData.DataStartIndex + Index);
 			}
 		}
 	}
