Fix a bug introduced in CL#2411787

diff --git a/Engine/Shaders/Private/LightGridCommon.ush b/Engine/Shaders/Private/LightGridCommon.ush
--- a/Engine/Shaders/Private/LightGridCommon.ush
+++ b/Engine/Shaders/Private/LightGridCommon.ush
@@ -29,6 +29,7 @@ uint GetCulledLightDataGrid(uint GridIndex)
 #endif
 }
 
+#if INSTANCED_STEREO
 uint GetCulledLightDataGridISR(uint GridIndex)
 {
 #if LIGHT_GRID_USES_16BIT_BUFFERS
@@ -37,6 +38,7 @@ uint GetCulledLightDataGridISR(uint GridIndex)
 	return ForwardLightDataISR.CulledLightDataGrid32Bit[GridIndex];
 #endif
 }
+#endif
 
 struct FLightGridData
 {
