Fix incorrect VSM backface culling with Strata for pixel having subsurface.

diff --git a/Engine/Shaders/Private/VirtualShadowMaps/VirtualShadowMapPageManagement.usf b/Engine/Shaders/Private/VirtualShadowMaps/VirtualShadowMapPageManagement.usf
--- a/Engine/Shaders/Private/VirtualShadowMaps/VirtualShadowMapPageManagement.usf
+++ b/Engine/Shaders/Private/VirtualShadowMaps/VirtualShadowMapPageManagement.usf
@@ -239,9 +239,11 @@ void GeneratePageFlagsFromPixels(
 		DeviceZ = LookupDeviceZ(PixelPos);
 #if STRATA_ENABLED
 		const FStrataTopLayerData TopLayerData = StrataUnpackTopLayerData(Strata.TopLayerTexture.Load(uint3(PixelPos, 0)));
+		FStrataAddressing StrataAddressing = GetStrataPixelDataByteOffset(PixelPos, uint2(View.BufferSizeAndInvSize.xy), Strata.MaxBytesPerPixel);
+		FStrataPixelHeader StrataPixelHeader = UnpackStrataHeaderIn(Strata.MaterialTextureArray, StrataAddressing, Strata.TopLayerTexture);
 		WorldNormal = TopLayerData.WorldNormal;
 		bIsValid = IsStrataMaterial(TopLayerData);
-		// bBackfaceCull = bBackfaceCull && !IsSubsurfaceModel(GBuffer.ShadingModelID); STRATA_TODO
+		bBackfaceCull = bBackfaceCull && !HasSubsurface(StrataPixelHeader);
 #else // !STRATA_ENABLED
 		FGBufferData GBuffer = GetGBufferDataUint(PixelPos, true);
 		WorldNormal = GBuffer.WorldNormal;
