Heterogeneous Volumes: Respect Component scale when ray-marching objects.

diff --git a/Engine/Shaders/Private/HeterogeneousVolumes/HeterogeneousVolumesRayMarchingUtils.ush b/Engine/Shaders/Private/HeterogeneousVolumes/HeterogeneousVolumesRayMarchingUtils.ush
--- a/Engine/Shaders/Private/HeterogeneousVolumes/HeterogeneousVolumesRayMarchingUtils.ush
+++ b/Engine/Shaders/Private/HeterogeneousVolumes/HeterogeneousVolumesRayMarchingUtils.ush
@@ -29,6 +29,7 @@ struct FRayMarchingContext
 	// World-space
 	float3 WorldRayOrigin;
 	float3 WorldRayDirection;
+	float LocalToWorldScale;
 
 	// Ray-step attributes
 	float Jitter;
@@ -81,6 +82,12 @@ FRayMarchingContext CreateRayMarchingContext(
 	Context.bApplyDirectLighting = bApplyDirectLighting;
 	Context.bApplyShadowTransmittance = bApplyShadowTransmittance;
 
+	// Transform local-scale properties into world-scale
+	float3 LocalRayHit = LocalRayOrigin + LocalRayDirection * LocalRayTMax;
+	float3 WorldRayHit = mul(float4(LocalRayHit, 1.0), GetLocalToWorld()).xyz;
+	float WorldRayTMax = length(WorldRayHit - WorldRayOrigin);
+	Context.LocalToWorldScale = WorldRayTMax / LocalRayTMax;
+
 	Context.MipLevel = 0.0f;
 
 	return Context;
@@ -158,8 +165,9 @@ void RayMarchTransmittance(
 {
 	for (uint StepIndex = 0; StepIndex < StepCount; ++StepIndex)
 	{
-		float3 LocalPosition = RayMarchingContext.LocalRayOrigin + RayMarchingContext.LocalRayDirection * (RayMarchingContext.LocalRayTMin + RayMarchingContext.StepSize * (RayMarchingContext.Jitter + StepIndex));
-		float3 WorldPosition = RayMarchingContext.WorldRayOrigin + RayMarchingContext.WorldRayDirection * (RayMarchingContext.LocalRayTMin + RayMarchingContext.StepSize * (RayMarchingContext.Jitter + StepIndex));
+		float LocalHitT = RayMarchingContext.LocalRayTMin + RayMarchingContext.StepSize * (RayMarchingContext.Jitter + StepIndex);
+		float3 LocalPosition = RayMarchingContext.LocalRayOrigin + RayMarchingContext.LocalRayDirection * LocalHitT;
+		float3 WorldPosition = RayMarchingContext.WorldRayOrigin + RayMarchingContext.WorldRayDirection * LocalHitT * RayMarchingContext.LocalToWorldScale;
 
 		FVolumeSampleContext SampleContext = CreateVolumeSampleContext(LocalPosition, WorldPosition, RayMarchingContext.MipLevel);
 		float3 Extinction = SampleExtinction(SampleContext);
@@ -319,8 +327,9 @@ void RayMarchEmissionAbsorption(
 {
 	for (uint StepIndex = 0; StepIndex < StepCount; ++StepIndex)
 	{
-		float3 LocalPosition = RayMarchingContext.LocalRayOrigin + RayMarchingContext.LocalRayDirection * (RayMarchingContext.LocalRayTMin + RayMarchingContext.StepSize * (RayMarchingContext.Jitter + StepIndex));
-		float3 WorldPosition = RayMarchingContext.WorldRayOrigin + RayMarchingContext.WorldRayDirection * (RayMarchingContext.LocalRayTMin + RayMarchingContext.StepSize * (RayMarchingContext.Jitter + StepIndex));
+		float LocalHitT = RayMarchingContext.LocalRayTMin + RayMarchingContext.StepSize * (RayMarchingContext.Jitter + StepIndex);
+		float3 LocalPosition = RayMarchingContext.LocalRayOrigin + RayMarchingContext.LocalRayDirection * LocalHitT;
+		float3 WorldPosition = RayMarchingContext.WorldRayOrigin + RayMarchingContext.WorldRayDirection * LocalHitT * RayMarchingContext.LocalToWorldScale;
 
 		FVolumeSampleContext SampleContext = CreateVolumeSampleContext(LocalPosition, WorldPosition, RayMarchingContext.MipLevel);
 		float3 Extinction = SampleExtinction(SampleContext);
@@ -344,8 +353,9 @@ void RayMarchSingleScattering(
 {
 	for (uint StepIndex = 0; StepIndex < StepCount; ++StepIndex)
 	{
-		float3 LocalPosition = RayMarchingContext.LocalRayOrigin + RayMarchingContext.LocalRayDirection * (RayMarchingContext.LocalRayTMin + RayMarchingContext.StepSize * (RayMarchingContext.Jitter + StepIndex));
-		float3 WorldPosition = RayMarchingContext.WorldRayOrigin + RayMarchingContext.WorldRayDirection * (RayMarchingContext.LocalRayTMin + RayMarchingContext.StepSize * (RayMarchingContext.Jitter + StepIndex));
+		float LocalHitT = RayMarchingContext.LocalRayTMin + RayMarchingContext.StepSize * (RayMarchingContext.Jitter + StepIndex);
+		float3 LocalPosition = RayMarchingContext.LocalRayOrigin + RayMarchingContext.LocalRayDirection * LocalHitT;
+		float3 WorldPosition = RayMarchingContext.WorldRayOrigin + RayMarchingContext.WorldRayDirection * LocalHitT * RayMarchingContext.LocalToWorldScale;
 		FVolumeSampleContext SampleContext = CreateVolumeSampleContext(LocalPosition, WorldPosition, RayMarchingContext.MipLevel);
 
 		if (RayMarchingContext.bApplyEmissionAndTransmittance)
