Strata - comment clean up

diff --git a/Engine/Shaders/Private/Strata/Strata.ush b/Engine/Shaders/Private/Strata/Strata.ush
--- a/Engine/Shaders/Private/Strata/Strata.ush
+++ b/Engine/Shaders/Private/Strata/Strata.ush
@@ -414,23 +414,23 @@ struct FStrataBSDF
 #define BSDF_GETWEIGHT10F(X)			READ_BITS(X.State,  BSDF_BIT_COUNT_WEIGHT10F, BSDF_BIT_OFFSET_WEIGHT10F)
 #define BSDF_SETWEIGHT10F(X, Weight)	WRITE_BITS(X.State, BSDF_BIT_COUNT_WEIGHT10F, BSDF_BIT_OFFSET_WEIGHT10F, Weight)
 
-// Indicates if the BSDF is part of the top layer
+// Indicates if the BSDF uses our simple volume shading to represent a layer of participating media
 #define BSDF_GETISSIMPLEVOLUME(X)		READ_BITS(X.State,  BSDF_BIT_COUNT_SIMPLEVOLUME, BSDF_BIT_OFFSET_SIMPLEVOLUME)
 #define BSDF_SETISSIMPLEVOLUME(X, Value) WRITE_BITS(X.State, BSDF_BIT_COUNT_SIMPLEVOLUME, BSDF_BIT_OFFSET_SIMPLEVOLUME, Value)
 
-// Indicates if the BSDF is part of the top layer
+// Indicates if the BSDF has explicit MFP data
 #define BSDF_GETHASMFP(X)				READ_BITS(X.State,  BSDF_BIT_COUNT_MFPPLUGGED, BSDF_BIT_OFFSET_MFPPLUGGED)
 #define BSDF_SETHASMFP(X, Value)		WRITE_BITS(X.State, BSDF_BIT_COUNT_MFPPLUGGED, BSDF_BIT_OFFSET_MFPPLUGGED, Value)
 
-// Indicates if the BSDF is part of the top layer
+// Indicates if the BSDF has fuzz data
 #define BSDF_GETHASFUZZ(X)				READ_BITS(X.State,  BSDF_BIT_COUNT_HASFUZZ, BSDF_BIT_OFFSET_HASFUZZ)
 #define BSDF_SETHASFUZZ(X, Value)		WRITE_BITS(X.State, BSDF_BIT_COUNT_HASFUZZ, BSDF_BIT_OFFSET_HASFUZZ, Value)
 
-// Indicates if the BSDF is part of the top layer
+// Indicates if the BSDF is thin (e.g. not thick as a mesh but thin as a piece of paper, drives the shading model used)
 #define BSDF_GETISTHIN(X)				READ_BITS(X.State,  BSDF_BIT_COUNT_ISTHIN, BSDF_BIT_OFFSET_ISTHIN)
 #define BSDF_SETISTHIN(X, Value)		WRITE_BITS(X.State, BSDF_BIT_COUNT_ISTHIN, BSDF_BIT_OFFSET_ISTHIN, Value)
 
-// Indicates if the BSDF has light weight data
+// Indicates if the BSDF has color transmittance data (due to layers of matter above it)
 #define BSDF_GETHASTRANSABOVE(X)		READ_BITS(X.State,  BSDF_BIT_COUNT_HASTRANSABOVE, BSDF_BIT_OFFSET_HASTRANSABOVE)
 #define BSDF_SETHASTRANSABOVE(X, Value)	WRITE_BITS(X.State, BSDF_BIT_COUNT_HASTRANSABOVE, BSDF_BIT_OFFSET_HASTRANSABOVE, Value)
 
