- Fix GetPreviousLocalToWorld3x3 using LocalToWorld rather than PreviousLocalToWorld

diff --git a/Engine/Shaders/Private/MaterialTemplate.ush b/Engine/Shaders/Private/MaterialTemplate.ush
--- a/Engine/Shaders/Private/MaterialTemplate.ush
+++ b/Engine/Shaders/Private/MaterialTemplate.ush
@@ -637,7 +637,7 @@ MaterialFloat3x3 GetPreviousLocalToWorld3x3(FMaterialVertexParameters Parameters
 #if DECAL_PRIMITIVE
 	return (float3x3)DecalToWorld;
 #else
-	return LWCToFloat3x3(GetPrimitiveData(Parameters).LocalToWorld);
+	return LWCToFloat3x3(GetPrimitiveData(Parameters).PreviousLocalToWorld);
 #endif
 }
 
