Fix missing shadow when high quality normal are enabled with Strata.

diff --git a/Engine/Shaders/Private/VirtualShadowMaps/VirtualShadowMapPageManagement.usf b/Engine/Shaders/Private/VirtualShadowMaps/VirtualShadowMapPageManagement.usf
--- a/Engine/Shaders/Private/VirtualShadowMaps/VirtualShadowMapPageManagement.usf
+++ b/Engine/Shaders/Private/VirtualShadowMaps/VirtualShadowMapPageManagement.usf
@@ -242,7 +242,7 @@ void GeneratePageFlagsFromPixels(
 		FStrataAddressing StrataAddressing = GetStrataPixelDataByteOffset(PixelPos, uint2(View.BufferSizeAndInvSize.xy), Strata.MaxBytesPerPixel);
 		FStrataPixelHeader StrataPixelHeader = UnpackStrataHeaderIn(Strata.MaterialTextureArray, StrataAddressing, Strata.TopLayerTexture);
 		WorldNormal = TopLayerData.WorldNormal;
-		bIsValid = IsStrataMaterial(TopLayerData);
+		bIsValid = IsStrataMaterial(StrataPixelHeader);
 		bBackfaceCull = bBackfaceCull && !HasSubsurface(StrataPixelHeader);
 #else // !STRATA_ENABLED
 		FGBufferData GBuffer = GetGBufferDataUint(PixelPos, true);
diff --git a/Engine/Shaders/Private/VirtualShadowMaps/VirtualShadowMapProjection.usf b/Engine/Shaders/Private/VirtualShadowMaps/VirtualShadowMapProjection.usf
--- a/Engine/Shaders/Private/VirtualShadowMaps/VirtualShadowMapProjection.usf
+++ b/Engine/Shaders/Private/VirtualShadowMaps/VirtualShadowMapProjection.usf
@@ -56,7 +56,7 @@ FProjectionShadingInfo GetProjectionShadingInfo(uint2 PixelPos)
 	const bool bIsProfile = StrataSubSurfaceHeaderGetIsProfile(SSSHeader);
 	Out.bIsSubsurface = bIsValid && !bIsProfile;
 
-	Out.bIsValid = IsStrataMaterial(TopLayerData);
+	Out.bIsValid = IsStrataMaterial(StrataPixelHeader);
 	Out.bIsHair = StrataGetBSDFType(StrataPixelHeader) == STRATA_BSDF_TYPE_HAIR;
 	Out.WorldNormal = TopLayerData.WorldNormal;
 #else
