Fixes a crash with r.TSR.History.SeparateTranslucency=1

diff --git a/Engine/Shaders/Private/TemporalSuperResolution/TSRUpdateHistory.usf b/Engine/Shaders/Private/TemporalSuperResolution/TSRUpdateHistory.usf
--- a/Engine/Shaders/Private/TemporalSuperResolution/TSRUpdateHistory.usf
+++ b/Engine/Shaders/Private/TemporalSuperResolution/TSRUpdateHistory.usf
@@ -87,7 +87,7 @@
 	#define CONFIG_INTERPOLATE_VELOCITY 0
 #endif
 
-#if CONFIG_GRAND_REPROJECTION == 2 || CONFIG_ACCUMULATE_TRANSLUCENCY_SEPARATLY || CONFIG_PREMULTIPLY_HISTORY
+#if CONFIG_GRAND_REPROJECTION == 2 || CONFIG_ACCUMULATE_TRANSLUCENCY_SEPARATELY || CONFIG_PREMULTIPLY_HISTORY
 	#define CONFIG_OUTPUT_HIGH_FREQUENCY 1
 #else
 	#define CONFIG_OUTPUT_HIGH_FREQUENCY 0
@@ -1699,7 +1699,7 @@ void MainCS(
 					HistoryOutput_ColorArray[tsr_short3(dpv_lo(LocalHistoryPixelPos), HistoryOutput_ArrayIndices_HighFrequency)] = dpv_lo(FinalHighFrequencyColor);
 				#endif
 				HistoryOutput_Metadata[dpv_lo(LocalHistoryPixelPos)]  = dpv_lo(FinalMetadata);
-				#if CONFIG_ACCUMULATE_TRANSLUCENCY_SEPARATLY
+				#if CONFIG_ACCUMULATE_TRANSLUCENCY_SEPARATELY
 					#if CONFIG_SCENE_COLOR_ALPHA
 						HistoryOutput_ColorArray[tsr_short3(dpv_lo(LocalHistoryPixelPos), HistoryOutput_ArrayIndices_Translucency)] = dpv_lo(FinalTranslucencyColor);
 					#else
