Ray Tracing: Fix potential preformance issue with invalid shadow rays being traced

For spot lights, we don't trace a shadow ray outside of its cone, so make sure the ray struct does not contain uninitialized values as this could lead to a random ray being fired that could severely impact performance.

diff --git a/Engine/Shaders/Private/RayTracing/RayTracingOcclusionRGS.usf b/Engine/Shaders/Private/RayTracing/RayTracingOcclusionRGS.usf
--- a/Engine/Shaders/Private/RayTracing/RayTracingOcclusionRGS.usf
+++ b/Engine/Shaders/Private/RayTracing/RayTracingOcclusionRGS.usf
@@ -243,7 +243,12 @@ float ComputeDiffuseTransmission(
 	// Build sub-surface scattering ray
 	FRayDesc SSSRay;
 	SSSRay.Origin = P + Tr * ScatterDistance;
-	GenerateOcclusionRay(LightParameters, SSSRay.Origin, N, LightSample, SSSRay.Origin, SSSRay.Direction, SSSRay.TMin, SSSRay.TMax);
+	bool bIsValidRay = GenerateOcclusionRay(LightParameters, SSSRay.Origin, N, LightSample, SSSRay.Origin, SSSRay.Direction, SSSRay.TMin, SSSRay.TMax);
+
+	if (!bIsValidRay)
+	{
+		return TransmissionDistance;
+	}
 
 	// Clip the ray length by the maximum transmission distance
 	float LightTMax = SSSRay.TMax;
@@ -389,7 +394,7 @@ FOcclusionResult ComputeOcclusion(
 
 		float2 RandSample = RandomSequence_GenerateSample2D(RandSequence);
 
-		FRayDesc Ray;
+		FRayDesc Ray = (FRayDesc)0;
 		bool bIsValidRay = GenerateOcclusionRay(
 			LightParameters,
 			TranslatedWorldPosition, WorldNormal,
