Strata - now properly account for Allo Negative Color

diff --git a/Engine/Shaders/Private/Strata/Strata.ush b/Engine/Shaders/Private/Strata/Strata.ush
--- a/Engine/Shaders/Private/Strata/Strata.ush
+++ b/Engine/Shaders/Private/Strata/Strata.ush
@@ -63,6 +63,10 @@
 #define MATERIAL_FULLY_ROUGH 0
 #endif
 
+#ifndef MATERIAL_ALLOW_NEGATIVE_EMISSIVECOLOR
+#define MATERIAL_ALLOW_NEGATIVE_EMISSIVECOLOR 0
+#endif
+
 #ifndef STRATA_CLAMPED_BSDF_COUNT
 #define STRATA_CLAMPED_BSDF_COUNT 0
 #endif
@@ -2644,7 +2648,11 @@ void SanitizeStrataSingleLayerWater(inout FStrataBSDF BSDF)
 
 void SanitizeStrataUnlit(inout FStrataBSDF BSDF)
 {
+#if MATERIAL_ALLOW_NEGATIVE_EMISSIVECOLOR
+	BSDF_SETEMISSIVE(BSDF, min(BSDF_GETEMISSIVE(BSDF), Max10BitsFloat));
+#else
 	BSDF_SETEMISSIVE(BSDF, clamp(BSDF_GETEMISSIVE(BSDF), 0.0, Max10BitsFloat));
+#endif
 	UNLIT_TRANSMITTANCE(BSDF) = saturate(UNLIT_TRANSMITTANCE(BSDF));
 }
 
