Distortion effect is not clipped out if distortion UVs are  smaller than any pixel (ignoring dynamic resolution scaling to be able to upscale correctly). Instead of an arbitrary threshold.

diff --git a/Engine/Shaders/Private/DistortionCommon.ush b/Engine/Shaders/Private/DistortionCommon.ush
--- a/Engine/Shaders/Private/DistortionCommon.ush
+++ b/Engine/Shaders/Private/DistortionCommon.ush
@@ -38,11 +38,14 @@ float2 ComputeBufferUVDistortion(
 
 	float2 BufferUVDistortion = ViewportUVDistortion * ResolvedView.ViewSizeAndInvSize.xy * ResolvedView.BufferSizeAndInvSize.zw;
 
-	//clip if the distortion distance (squared) is too small to be noticed
-	//this will result in a less expensive apply pass since the clipped pixels won't set stencil to pass
-	if (TryToClip)
+	// Clip if the distortion distance (squared) is too small to be noticed.
+	// This will result in a less expensive apply pass since the clipped pixels won't set stencil to pass.
+	// We only clip if the distortion is less than the size of a pixel of the unscaled viewport, i.e. dynamic resolution = 100%, to ensure a good reconstruction.
+	const float2 FullResolutionDistortionPixelSize = ResolvedView.BufferSizeAndInvSize.zw;
+	if (TryToClip && 
+		all(abs(BufferUVDistortion) < FullResolutionDistortionPixelSize))
 	{
-		clip(dot(BufferUVDistortion, BufferUVDistortion) - .00001);
+		clip(-1.0f);
 	}
 
 	// When ISR is enabled we store two FOVs in the DistortionParameters and compute the aspect ratio here
