Strata - fix for fireflies

diff --git a/Engine/Shaders/Private/Strata/StrataEvaluation.ush b/Engine/Shaders/Private/Strata/StrataEvaluation.ush
--- a/Engine/Shaders/Private/Strata/StrataEvaluation.ush
+++ b/Engine/Shaders/Private/Strata/StrataEvaluation.ush
@@ -145,6 +145,13 @@ FStrataBSDFContext StrataCreateBSDFContext(float3x3 TangentBasis, FStrataBSDF BS
 	BSDFContext.SatNoL = saturate(BSDFContext.Context.NoL);
 	BSDFContext.SatNoV = saturate(BSDFContext.Context.NoV);
 
+#if STRATA_NORMAL_QUALITY==1
+	// It looks like when we use the high quality normals, NoH can become 1.
+	// This can cause D_GGX to be infinite du to a division by 0 when roughness is also 0.
+	// So to avoid fireflies, we clamp NoH to a value a little bit below 1.
+	BSDFContext.Context.NoH = clamp(BSDFContext.Context.NoH, 0.0f, 0.9999f);
+#endif
+
 	return BSDFContext;
 }
 
diff --git a/Engine/Source/Runtime/Engine/Private/ShaderCompiler/ShaderCompiler.cpp b/Engine/Source/Runtime/Engine/Private/ShaderCompiler/ShaderCompiler.cpp
--- a/Engine/Source/Runtime/Engine/Private/ShaderCompiler/ShaderCompiler.cpp
+++ b/Engine/Source/Runtime/Engine/Private/ShaderCompiler/ShaderCompiler.cpp
@@ -6506,7 +6506,6 @@ void GlobalBeginCompileShader(
 				Input.Environment.SetDefine(TEXT("USE_ACHROMATIC_BXDF_ENERGY"), 1u);
 			}
 
-			static const auto CVarStrataGBufferFormat = IConsoleManager::Get().FindTConsoleVariableDataInt(TEXT("r.GBufferFormat"));
 			const uint32 StrataNormalQuality = Strata::GetNormalQuality();
 			Input.Environment.SetDefine(TEXT("STRATA_NORMAL_QUALITY"), StrataNormalQuality);
 
diff --git a/Engine/Source/Runtime/Renderer/Private/Strata/Strata.cpp b/Engine/Source/Runtime/Renderer/Private/Strata/Strata.cpp
--- a/Engine/Source/Runtime/Renderer/Private/Strata/Strata.cpp
+++ b/Engine/Source/Runtime/Renderer/Private/Strata/Strata.cpp
@@ -304,7 +304,7 @@ static void RecordStrataAnalytics(EShaderPlatform InPlatform)
 
 static EPixelFormat GetTopLayerTextureFormat()
 {
-	const bool bStrataHighQualityNormal = GetNormalQuality() > 1;
+	const bool bStrataHighQualityNormal = GetNormalQuality() > 0;
 
 	// High quality normal is not supported on platforms that do not support R32G32 UAV load.
 	// This is dues to the way Strata account for decals. See FStrataDBufferPassCS, updating TopLayerTexture this way.
