Path Tracer: Fix potential TDRs caused by infinite loops during raytracing inside the path tracer

There were multiple issues here:
  * The default opaque shader did not correctly set HitT for non-shadow rays
  * The default Hidden shader did not behave correctly if ForceOpaque was set (this should be impossible in theory, but it seems safer to be robust to this case)
  * ForceOpaque could be set for cached meshes that had only opaque sections but had some sections set to not cast shadows (either from the mesh itself or the material)
  * BLEND_Additive was not correctly accounted for in all places that determine is a mesh command casts shadows
  * Add the Additive blend mode case for shadow rays to the path tracing material as a defensive measure (it doesn't cost much and will avoid other corner cases where the cast shadow flag is not correctly set on the mesh command)
  * Tracing into a scene with missing CHS entries could lead to crashes as the missing value was reported as a hit, but did not write to HitT - update all trace calls so that the payload is in a "miss" state before tracing to defend against this (this causes visual issues as can be seen by using the path tracer with mesh section isolation in the StaticMeshEditor, but this seems like the best solution for now)
  * Apply the same logic to the default decal shaders as well
  * Fix out-of-bounds access on scenes with volumetrics and more than 256 lights (render will be incorrect, but at least will not crash)

I have left some of the debug scaffolding for debuggin of shadow rays which may be handy in the future.

diff --git a/Engine/Plugins/Experimental/GPULightmass/Shaders/Private/LightmapPathTracing.usf b/Engine/Plugins/Experimental/GPULightmass/Shaders/Private/LightmapPathTracing.usf
--- a/Engine/Plugins/Experimental/GPULightmass/Shaders/Private/LightmapPathTracing.usf
+++ b/Engine/Plugins/Experimental/GPULightmass/Shaders/Private/LightmapPathTracing.usf
@@ -312,6 +312,7 @@ float3 TraceTransparentVisibilityRay(FRayDesc Ray, float PathRoughness)
 	for (;;)
 	{
 		// Trace the ray
+		PackedPayload.SetMiss();
 		TraceRay(
 			TLAS,
 			RayFlags,
diff --git a/Engine/Shaders/Private/PathTracing/PathTracingCommon.ush b/Engine/Shaders/Private/PathTracing/PathTracingCommon.ush
--- a/Engine/Shaders/Private/PathTracing/PathTracingCommon.ush
+++ b/Engine/Shaders/Private/PathTracing/PathTracingCommon.ush
@@ -278,6 +278,7 @@ FPackedPathTracingPayload InitPathTracingPayload(uint ScatterType, float PathRou
 {
 	FPackedPathTracingPayload Output = (FPackedPathTracingPayload)0;
 	Output.SetPathRoughness(PathRoughness);
+	Output.SetMiss();
 	switch (ScatterType)
 	{
 		case PATHTRACER_SCATTER_DIFFUSE:   Output.SetFlag(PATH_TRACING_PAYLOAD_INPUT_FLAG_RAY_TYPE_INDIRECT_DIFFUSE);  break;
@@ -299,6 +300,7 @@ FPackedPathTracingPayload InitPathTracingVisibilityPayload(float PathRoughness)
 	Output.SetVisibilityRay();
 	Output.SetPathRoughness(PathRoughness);
 	Output.SetRayThroughput(1.0);
+	Output.SetMiss();
 	return Output;
 }
 
diff --git a/Engine/Shaders/Private/PathTracing/PathTracingCore.ush b/Engine/Shaders/Private/PathTracing/PathTracingCore.ush
--- a/Engine/Shaders/Private/PathTracing/PathTracingCore.ush
+++ b/Engine/Shaders/Private/PathTracing/PathTracingCore.ush
@@ -235,7 +235,7 @@ void TraceDecalRay(FRayDesc DecalRay, bool bReverseCulling, uint InstanceInclusi
 	{
 		FDecalShaderPayload DecalPayload = (FDecalShaderPayload)0;
 		DecalPayload.SetInputTranslatedWorldPosition(HitPayload.TranslatedWorldPos);
-
+		DecalPayload.SetMiss();
 		TraceRay(
 			DecalTLAS,
 			bReverseCulling ? RAY_FLAG_CULL_FRONT_FACING_TRIANGLES : RAY_FLAG_CULL_BACK_FACING_TRIANGLES,
@@ -608,6 +608,8 @@ FPathTracingPayload TraceTransparentRay(inout FPathState PathState, bool bIsCame
 	return Payload;
 }
 
+#define PATH_TRACING_DEBUG_TRANSPARENCY_RAY 0
+
 float3 TraceTransparentVisibilityRay(FRayDesc Ray, float PathRoughness, uint MissShaderIndex, bool bCastShadows, inout RandomSequence RandSequence)
 {
 	if (!bCastShadows && MissShaderIndex == 0)
@@ -630,9 +632,20 @@ float3 TraceTransparentVisibilityRay(FRayDesc Ray, float PathRoughness, uint Mis
 		Ray.TMax = POSITIVE_INFINITY;
 	}
 
+	const float OrigTMin = Ray.TMin;
+
 	// Trace the ray - progressing in steps through transparent geometry
+#if PATH_TRACING_DEBUG_TRANSPARENCY_RAY
+	const int MaxSteps = 256; // avoid TDRs during debugging by capping how many times we go around
+	int Steps = 0;
+	for (int Counter = 0; Counter < MaxSteps; Counter++)
+#else
 	for (;;)
+#endif
 	{
+		// Reset payload to a miss for the rare case of a mesh without a valid CHS assigned
+		// Note that in this case we will get incorrect shadowing, but this is preferable to an infinite loop
+		PackedPayload.SetMiss();
 		TraceRay(
 			TLAS,
 			RayFlags,
@@ -655,12 +668,23 @@ float3 TraceTransparentVisibilityRay(FRayDesc Ray, float PathRoughness, uint Mis
 		}
 		// we hit some geometry, move one ulp past the hit and trace the same ray again
 		Ray.TMin = asfloat(asuint(PackedPayload.HitT) + 1);
+#if PATH_TRACING_DEBUG_TRANSPARENCY_RAY
+		Steps++;
+#endif
 	}
 
+#if PATH_TRACING_DEBUG_TRANSPARENCY_RAY
+	if (Steps == 0) return 1.0;
+	if (Steps == 1) return float3(0, 0, 1);
+	if (Steps == 2) return float3(0, 1, 0);
+	if (Steps == MaxSteps) return float3(1, 0, 0);
+	const float t = float(Steps) / float(MaxSteps);
+	return t;
+#else
 	float3 Throughput = PackedPayload.GetRayThroughput();
 	if (bCastShadows && any(Throughput > 0))
 	{
-		FVolumeIntersectionList VolumeIsectList = VolumeIntersect(Ray.Origin, Ray.Direction, Ray.TMin, Ray.TMax);
+		FVolumeIntersectionList VolumeIsectList = VolumeIntersect(Ray.Origin, Ray.Direction, OrigTMin, Ray.TMax);
 		if (VolumeIsectList.HitBlocker())
 		{
 			// blockers are opaque -- done!
@@ -675,6 +699,7 @@ float3 TraceTransparentVisibilityRay(FRayDesc Ray, float PathRoughness, uint Mis
 		Throughput *= exp(-max(Tau, 0.0));
 	}
 	return Throughput;
+#endif
 }
 
 struct FProbeResult
@@ -999,11 +1024,12 @@ bool PathTracingKernel(inout FPathState PathState, int Bounce)
 	// This array will hold a CDF for light picking
 	float LightPickingCdf[RAY_TRACING_LIGHT_COUNT_MAXIMUM];
 
-#if 0
-	// visualize pixel seeds
-	Radiance = HashToColor(StrongIntegerHash(RandPixel.x + StrongIntegerHash(RandPixel.y)));
+#if PATH_TRACING_DEBUG_TRANSPARENCY_RAY
+	// visualize the shadow ray handling directly
+	PathState.Radiance = TraceTransparentVisibilityRay(PathState.Ray, 0.0, 0, true, PathState.RandSequence);
 	return false;
 #endif
+
 	const bool bIsCameraRay = Bounce == 0;
 	const bool bIsLastBounce = Bounce == MaxBounces;
 	const bool bIncludeEmissive = PathState.ShouldAccumulateEmissive();
@@ -1030,8 +1056,9 @@ bool PathTracingKernel(inout FPathState PathState, int Bounce)
 		const float VTMin = VolumeSegment.Interval.VolumeTMin;
 		const float VTMax = VolumeSegment.Interval.VolumeTMax;
 
+		const int NumLights = min(SceneLightCount, RAY_TRACING_LIGHT_COUNT_MAXIMUM); // make sure we don't overflow the cdf (TODO: relax this limit)
 		float LightPickingCdfSum = 0.0;
-		for (int LightId = 0; LightId < SceneLightCount; LightId++)
+		for (int LightId = 0; LightId < NumLights; LightId++)
 		{
 			FVolumeLightSampleSetup LightSetup = PrepareLightVolumeSample(LightId, Ro, Rd, VTMin, VTMax);
 			if (LightSetup.IsValid())
@@ -1048,7 +1075,7 @@ bool PathTracingKernel(inout FPathState PathState, int Bounce)
 			float3 RandSample = RandomSequence_GenerateSample3D(PathState.RandSequence);
 			int LightSampleLightId = 0;
 			float LightPickPdf = 0.0;
-			SelectLight(RandSample.x * LightPickingCdfSum, SceneLightCount, LightPickingCdf, LightSampleLightId, LightPickPdf);
+			SelectLight(RandSample.x * LightPickingCdfSum, NumLights, LightPickingCdf, LightSampleLightId, LightPickPdf);
 			LightPickPdf /= LightPickingCdfSum;
 
 			// picked a light! now use the equi-angular sampler to pick a position along the ray and store it for later
diff --git a/Engine/Shaders/Private/PathTracing/PathTracingDefaultHitShader.usf b/Engine/Shaders/Private/PathTracing/PathTracingDefaultHitShader.usf
--- a/Engine/Shaders/Private/PathTracing/PathTracingDefaultHitShader.usf
+++ b/Engine/Shaders/Private/PathTracing/PathTracingDefaultHitShader.usf
@@ -7,7 +7,6 @@ RAY_TRACING_ENTRY_CLOSEST_HIT(PathTracingDefaultOpaqueCHS,
 	FPackedPathTracingPayload, PackedPayload,
 	FRayTracingIntersectionAttributes, Attributes)
 {
-	PackedPayload.HitT = RayTCurrent();
 	if (PackedPayload.IsVisibilityRay())
 	{
 		// mark shadow fully opaque
@@ -22,12 +21,28 @@ RAY_TRACING_ENTRY_CLOSEST_HIT(PathTracingDefaultOpaqueCHS,
 		Payload.Opacity = 1.0;
 		PackedPayload = PackPathTracingPayload(Payload);
 	}
+	PackedPayload.HitT = RayTCurrent();
 }
 
 RAY_TRACING_ENTRY_CLOSEST_HIT(PathTracingDefaultHiddenCHS,
 	FPackedPathTracingPayload, PackedPayload,
 	FRayTracingIntersectionAttributes, Attributes)
 {
+	// This is not reachable in theory since AHS calls IgnoreHit(), but could be if the mesh gets flagged with ForceOpaque (or the ray ignores AHS)
+	if (PackedPayload.IsVisibilityRay())
+	{
+		// don't change the throughput (object should be invisible)
+	}
+	else
+	{
+		// mark hit as fully transparent
+		FPathTracingPayload Payload = (FPathTracingPayload)0;
+		Payload.ShadingModelID = SHADINGMODELID_UNLIT;
+		Payload.BlendingMode = RAY_TRACING_BLEND_MODE_ADDITIVE;
+		Payload.Opacity = 0.0;
+		PackedPayload = PackPathTracingPayload(Payload);
+	}
+	PackedPayload.HitT = RayTCurrent();
 }
 
 RAY_TRACING_ENTRY_ANY_HIT(PathTracingDefaultHiddenAHS,
diff --git a/Engine/Shaders/Private/PathTracing/PathTracingMaterialHitShader.usf b/Engine/Shaders/Private/PathTracing/PathTracingMaterialHitShader.usf
--- a/Engine/Shaders/Private/PathTracing/PathTracingMaterialHitShader.usf
+++ b/Engine/Shaders/Private/PathTracing/PathTracingMaterialHitShader.usf
@@ -87,7 +87,14 @@ RAY_TRACING_ENTRY_CLOSEST_HIT(PathTracingMaterialCHS,
 		return;
 	}
 #endif
-
+#if MATERIALBLENDING_ADDITIVE
+	if (PackedPayload.IsVisibilityRay())
+	{
+		// should be fully invisible to shadow rays
+		// this case should be caught by making the mesh invisible to shadows, but it doesn't cost too much to have this extra check here just in case that fails
+		return;
+	}
+#endif
 	ResolvedView = ResolveView();
 
 	const float3 TranslatedWorldPosition = TranslatedWorldRayOrigin() + RayTCurrent() * WorldRayDirection();
diff --git a/Engine/Source/Runtime/Renderer/Private/DeferredShadingRenderer.cpp b/Engine/Source/Runtime/Renderer/Private/DeferredShadingRenderer.cpp
--- a/Engine/Source/Runtime/Renderer/Private/DeferredShadingRenderer.cpp
+++ b/Engine/Source/Runtime/Renderer/Private/DeferredShadingRenderer.cpp
@@ -698,6 +698,7 @@ struct FRayTracingRelevantPrimitive
 	uint8 InstanceMask = 0;
 	bool bStatic = false;
 	bool bAllSegmentsOpaque = true;
+	bool bAllSegmentsCastShadow = true;
 	bool bAnySegmentsCastShadow = false;
 	bool bAnySegmentsDecal = false;
 	bool bAllSegmentsDecal = true;
@@ -714,12 +715,13 @@ struct FRayTracingRelevantPrimitive
 		uint64 Key = StateHash;
 		Key ^= uint64(InstanceMask) << 32;
 		Key ^= bAllSegmentsOpaque ? 0x1ull << 40 : 0x0;
-		Key ^= bAnySegmentsCastShadow ? 0x1ull << 41 : 0x0;
-		Key ^= bAnySegmentsDecal ? 0x1ull << 42 : 0x0;
-		Key ^= bAllSegmentsDecal ? 0x1ull << 43 : 0x0;
-		Key ^= bTwoSided ? 0x1ull << 44 : 0x0;
-		Key ^= bIsSky ? 0x1ull << 45 : 0x0;
-		Key ^= bAllSegmentsTranslucent ? 0x1ull << 46 : 0x0;
+		Key ^= bAllSegmentsCastShadow ? 0x1ull << 41 : 0x0;
+		Key ^= bAnySegmentsCastShadow ? 0x1ull << 42 : 0x0;
+		Key ^= bAnySegmentsDecal ? 0x1ull << 43 : 0x0;
+		Key ^= bAllSegmentsDecal ? 0x1ull << 44 : 0x0;
+		Key ^= bTwoSided ? 0x1ull << 45 : 0x0;
+		Key ^= bIsSky ? 0x1ull << 46 : 0x0;
+		Key ^= bAllSegmentsTranslucent ? 0x1ull << 47 : 0x0;
 		return Key ^ reinterpret_cast<uint64>(RayTracingGeometryRHI);
 	}
 
@@ -961,6 +963,7 @@ static void GatherRayTracingRelevantPrimitives(FScene& Scene, const FViewInfo& V
 
 							RelevantPrimitive.InstanceMask |= RayTracingMeshCommand.InstanceMask;
 							RelevantPrimitive.bAllSegmentsOpaque &= RayTracingMeshCommand.bOpaque;
+							RelevantPrimitive.bAllSegmentsCastShadow &= RayTracingMeshCommand.bCastRayTracedShadows;
 							RelevantPrimitive.bAnySegmentsCastShadow |= RayTracingMeshCommand.bCastRayTracedShadows;
 							RelevantPrimitive.bAnySegmentsDecal |= RayTracingMeshCommand.bDecal;
 							RelevantPrimitive.bAllSegmentsDecal &= RayTracingMeshCommand.bDecal;
@@ -1695,7 +1698,7 @@ bool FDeferredShadingSceneRenderer::GatherRayTracingWorldInstancesForView(FRDGBu
 
 						RayTracingInstance.Mask = RelevantPrimitive.InstanceMask; // When no cached command is found, InstanceMask == 0 and the instance is effectively filtered out
 
-						if (RelevantPrimitive.bAllSegmentsOpaque)
+						if (RelevantPrimitive.bAllSegmentsOpaque && RelevantPrimitive.bAllSegmentsCastShadow)
 						{
 							RayTracingInstance.Flags |= ERayTracingInstanceFlags::ForceOpaque;
 						}
diff --git a/Engine/Source/Runtime/Renderer/Private/RayTracing/RayTracingInstanceMask.cpp b/Engine/Source/Runtime/Renderer/Private/RayTracing/RayTracingInstanceMask.cpp
--- a/Engine/Source/Runtime/Renderer/Private/RayTracing/RayTracingInstanceMask.cpp
+++ b/Engine/Source/Runtime/Renderer/Private/RayTracing/RayTracingInstanceMask.cpp
@@ -162,8 +162,9 @@ FRayTracingMaskAndFlags BuildRayTracingInstanceMaskAndFlags(TArrayView<const FMe
 			const EBlendMode BlendMode = Material.GetBlendMode();
 			Result.Mask |= BlendModeToRayTracingInstanceMask(BlendMode, MaskMode);
 			bAllSegmentsOpaque &= BlendMode == BLEND_Opaque;
-			bAnySegmentsCastShadow |= MeshBatch.CastRayTracedShadow && Material.CastsRayTracedShadows();
-			bAllSegmentsCastShadow &= MeshBatch.CastRayTracedShadow && Material.CastsRayTracedShadows();
+			const bool bSegmentCastsShadow = MeshBatch.CastRayTracedShadow && Material.CastsRayTracedShadows() && BlendMode != BLEND_Additive;
+			bAnySegmentsCastShadow |= bSegmentCastsShadow;
+			bAllSegmentsCastShadow &= bSegmentCastsShadow;
 			bAnySegmentsDecal |= Material.IsDeferredDecal();
 			bAllSegmentsDecal &= Material.IsDeferredDecal();
 			bDoubleSided |= MeshBatch.bDisableBackfaceCulling || Material.IsTwoSided();
diff --git a/Engine/Source/Runtime/Renderer/Private/RayTracing/RayTracingMaterialHitShaders.cpp b/Engine/Source/Runtime/Renderer/Private/RayTracing/RayTracingMaterialHitShaders.cpp
--- a/Engine/Source/Runtime/Renderer/Private/RayTracing/RayTracingMaterialHitShaders.cpp
+++ b/Engine/Source/Runtime/Renderer/Private/RayTracing/RayTracingMaterialHitShaders.cpp
@@ -762,8 +762,8 @@ FRayTracingPipelineState* FDeferredShadingSceneRenderer::CreateRayTracingMateria
 	const int32 OpaqueShadowMaterialIndex = FindRayTracingHitGroupIndex(PipelineState, OpaqueShadowShader, true);
 	const int32 HiddenMaterialIndex = FindRayTracingHitGroupIndex(PipelineState, HiddenMaterialShader, true);
 
-	const int32 OpaqueMeshDecalHitGroupIndex = FindRayTracingHitGroupIndex(PipelineState, OpaqueShadowShader, true);
-	const int32 HiddenMeshDecalHitGroupIndex = FindRayTracingHitGroupIndex(PipelineState, HiddenMaterialShader, true);
+	const int32 OpaqueMeshDecalHitGroupIndex = bSupportMeshDecals ? FindRayTracingHitGroupIndex(PipelineState, OpaqueMeshDecalHitShader, true) : INDEX_NONE;
+	const int32 HiddenMeshDecalHitGroupIndex = bSupportMeshDecals ? FindRayTracingHitGroupIndex(PipelineState, HiddenMeshDecalHitShader, true) : INDEX_NONE;
 
 	FViewInfo& ReferenceView = Views[0];
 
@@ -844,8 +844,17 @@ FRayTracingPipelineState* FDeferredShadingSceneRenderer::CreateRayTracingMateria
 						}
 
 						// Bind shadow shader
+						if (bIsMeshDecalShader)
+						{
+							// mesh decals do not use the shadow slot, so do minimal work
+							FRayTracingLocalShaderBindings& Binding = BindingWriter->AddWithExternalParameters();
+							Binding.InstanceIndex = VisibleMeshCommand.InstanceIndex;
+							Binding.SegmentIndex = MeshCommand.GeometrySegmentIndex;
+							Binding.ShaderSlot = RAY_TRACING_SHADER_SLOT_SHADOW;
+							Binding.ShaderIndexInPipeline = OpaqueMeshDecalHitGroupIndex;
 
-						if (MeshCommand.bCastRayTracedShadows && !VisibleMeshCommand.bHidden)
+						}
+						else if (MeshCommand.bCastRayTracedShadows && !VisibleMeshCommand.bHidden)
 						{
 							if (MeshCommand.bOpaque || !bEnableShadowMaterials)
 							{
