Path Tracing: Refactor path tracing code to minimize material/blend type related special cases in the integrator

Revisit payload encoding to handle blend mode logic entirely in the shader. This removes the need for a blendmode field and simplifies the logic in the raygen shader, reducing the amount of switch statements. We also gain back 3 extra bits that can be spent on output flags.

Refactor how random walk selection is performed to avoid shadingmodel specific code in the integrator. Also refactor the code to allow transparency to be combined with SSS.

Unify the handling of SingleLayerWater and solid glass since they were mostly sharing the same code. Both cases now use the DEFAULT_LIT shading model. This required extending the logic for camera medium tracking a bit. It now always runs (with a cvar to turn it off for debugging) instead of being tied to View.IsUnderwater() which could only work for water geometry. This makes it possible to achieve the under-water effect with regular refractive materials as well.

diff --git a/Engine/Plugins/Experimental/GPULightmass/Shaders/Private/LightmapPathTracing.usf b/Engine/Plugins/Experimental/GPULightmass/Shaders/Private/LightmapPathTracing.usf
--- a/Engine/Plugins/Experimental/GPULightmass/Shaders/Private/LightmapPathTracing.usf
+++ b/Engine/Plugins/Experimental/GPULightmass/Shaders/Private/LightmapPathTracing.usf
@@ -220,21 +220,12 @@ FPathTracingPayload TraceTransparentRay(FRayDesc Ray, bool IsCameraRay, bool Las
 		// Unpack the payload
 		FPathTracingPayload HitPayload = UnpackPathTracingPayload(PackedPayload, Ray);
 
-		if (HitPayload.BlendingMode == RAY_TRACING_BLEND_MODE_ADDITIVE ||
-			HitPayload.BlendingMode == RAY_TRACING_BLEND_MODE_TRANSLUCENT)
-		{
-			// emission needs to be pre-multiplied by the opacity
-			HitPayload.Radiance *= HitPayload.Opacity;
-		}
-
-		// add in surface emission (except for modulate blend mode which uses Radiance as the transparency multiplier)
-		if (IncludeEmission && HitPayload.BlendingMode != RAY_TRACING_BLEND_MODE_MODULATE)
+		// add in surface emission
+		if (IncludeEmission)
 		{
 			AccumulateRadiance(Radiance, PathThroughput * HitPayload.Radiance);
 		}
 
-		float3 Transparency = GetMaterialTransparency(HitPayload, Ray.Direction);
-
 		if (!LastBounce)
 		{
 			float3 Contrib = PathThroughput * EstimateMaterialAlbedo(HitPayload);
@@ -277,7 +268,7 @@ FPathTracingPayload TraceTransparentRay(FRayDesc Ray, bool IsCameraRay, bool Las
 
 		// prepare next step around the loop
 		// retrace the exact same ray with TMin one ulp past the hit we just found
-		PathThroughput *= Transparency;
+		PathThroughput *= HitPayload.TransparencyColor;
 		Ray.TMin = asfloat(asuint(HitPayload.HitT) + 1);
 
 		if (all(PathThroughput == 0))
@@ -326,7 +317,7 @@ float3 TraceTransparentVisibilityRay(FRayDesc Ray, float PathRoughness)
 		{
 			break;
 		}
-		if (!any(PackedPayload.GetRayThroughput()))
+		if (!any(PackedPayload.GetRayThroughput() > 0))
 		{
 			// accumulated to fully opacity already
 			break;
@@ -406,17 +397,17 @@ void PathTracingKernel(
 			Payload.WorldNormal = ShadingNormal;
 			Payload.WorldGeoNormal = ShadingNormal;
 			Payload.WorldSmoothNormal = ShadingNormal;
-			Payload.Radiance = float3(0, 0, 0);
-			Payload.BaseColor = float3(1, 1, 1);
-			Payload.DiffuseColor = float3(1, 1, 1);
-			Payload.SpecularColor = float3(0, 0, 0);
+			Payload.Radiance = 0;
+			Payload.BaseColor = 1;
+			Payload.DiffuseColor = 1;
+			Payload.SpecularColor = 0;
 			Payload.Specular = 0;
 			Payload.Roughness = 1;
 			Payload.Metallic = 0;
 			Payload.Ior = 1;
-			Payload.Opacity = 1;
+			Payload.BSDFOpacity = 1;
+			Payload.TransparencyColor = 0;
 			Payload.ShadingModelID = SHADINGMODELID_DEFAULT_LIT;
-			Payload.BlendingMode = RAY_TRACING_BLEND_MODE_OPAQUE;
 			Payload.PrimitiveLightingChannelMask = 0b111;
 			Payload.HitT = 1.0f;
 			Payload.SetFrontFace();
diff --git a/Engine/Shaders/Private/PathTracing/Material/PathTracingClearCoat.ush b/Engine/Shaders/Private/PathTracing/Material/PathTracingClearCoat.ush
--- a/Engine/Shaders/Private/PathTracing/Material/PathTracingClearCoat.ush
+++ b/Engine/Shaders/Private/PathTracing/Material/PathTracingClearCoat.ush
@@ -148,7 +148,7 @@ FMaterialEval ClearCoat_EvalMaterial(
 	Result.AddLobeWithMIS(SpecWeight * DiffuseSpecularScale.y, SpecPdf, Data.LobePdf.y);
 	Result.AddLobeWithMIS(DiffWeight * DiffuseSpecularScale.x, DiffPdf, Data.LobePdf.z);
 
-	Result.Weight *= Payload.GetBSDFOpacity();
+	Result.Weight *= Payload.BSDFOpacity;
 	
 	return Result;
 }
@@ -254,7 +254,7 @@ FMaterialSample ClearCoat_SampleMaterial(
 	}
 
 
-	Result.Weight *= Payload.GetBSDFOpacity();
+	Result.Weight *= Payload.BSDFOpacity;
 	Result.PositionBiasSign = 1.0;
 
 	return Result;
diff --git a/Engine/Shaders/Private/PathTracing/Material/PathTracingCloth.ush b/Engine/Shaders/Private/PathTracing/Material/PathTracingCloth.ush
--- a/Engine/Shaders/Private/PathTracing/Material/PathTracingCloth.ush
+++ b/Engine/Shaders/Private/PathTracing/Material/PathTracingCloth.ush
@@ -89,7 +89,7 @@ FMaterialEval Cloth_EvalMaterial(
 	const float SpecPdf = GGXResult.y;
 	Result.AddLobeWithMIS(SpecWeight * DiffuseSpecularScale.y, SpecPdf, Data.LobePdf.z);
 
-	Result.Weight *= Payload.GetBSDFOpacity();
+	Result.Weight *= Payload.BSDFOpacity;
 
 	return Result;
 }
@@ -181,6 +181,6 @@ FMaterialSample Cloth_SampleMaterial(
 		Result.Roughness = Payload.Roughness;
 		Result.ScatterType = PATHTRACER_SCATTER_SPECULAR;
 	}
-	Result.Weight *= Payload.GetBSDFOpacity();
+	Result.Weight *= Payload.BSDFOpacity;
 	return Result;
 }
diff --git a/Engine/Shaders/Private/PathTracing/Material/PathTracingDefaultLit.ush b/Engine/Shaders/Private/PathTracing/Material/PathTracingDefaultLit.ush
--- a/Engine/Shaders/Private/PathTracing/Material/PathTracingDefaultLit.ush
+++ b/Engine/Shaders/Private/PathTracing/Material/PathTracingDefaultLit.ush
@@ -69,7 +69,7 @@ FMaterialEval DefaultLit_EvalMaterial(
 
 	Result.AddLobeWithMIS(DiffuseSpecularScale.y * SpecWeight, SpecPdf, 1.0 - Data.LobeProb);
 
-	Result.Weight *= Payload.GetBSDFOpacity();
+	Result.Weight *= Payload.BSDFOpacity;
 
 	return Result;
 }
@@ -136,6 +136,6 @@ FMaterialSample DefaultLit_SampleMaterial(
 		Result.Roughness = Payload.Roughness;
 		Result.ScatterType = PATHTRACER_SCATTER_SPECULAR;
 	}
-	Result.Weight *= Payload.GetBSDFOpacity();
+	Result.Weight *= Payload.BSDFOpacity;
 	return Result;
 }
diff --git a/Engine/Shaders/Private/PathTracing/Material/PathTracingEye.ush b/Engine/Shaders/Private/PathTracing/Material/PathTracingEye.ush
--- a/Engine/Shaders/Private/PathTracing/Material/PathTracingEye.ush
+++ b/Engine/Shaders/Private/PathTracing/Material/PathTracingEye.ush
@@ -82,7 +82,7 @@ FMaterialEval Eye_EvalMaterial(
 
 	Result.AddLobeWithMIS(SpecWeight * DiffuseSpecularScale.y, SpecPdf, 1.0 - Data.LobeProb);
 
-	Result.Weight *= Payload.GetBSDFOpacity();
+	Result.Weight *= Payload.BSDFOpacity;
 
 	return Result;
 }
@@ -158,6 +158,6 @@ FMaterialSample Eye_SampleMaterial(
 		Result.Roughness = Payload.Roughness;
 		Result.ScatterType = PATHTRACER_SCATTER_SPECULAR;
 	}
-	Result.Weight *= Payload.GetBSDFOpacity();
+	Result.Weight *= Payload.BSDFOpacity;
 	return Result;
 }
diff --git a/Engine/Shaders/Private/PathTracing/Material/PathTracingLambert.ush b/Engine/Shaders/Private/PathTracing/Material/PathTracingLambert.ush
--- a/Engine/Shaders/Private/PathTracing/Material/PathTracingLambert.ush
+++ b/Engine/Shaders/Private/PathTracing/Material/PathTracingLambert.ush
@@ -15,7 +15,7 @@ FMaterialSample Lambert_SampleMaterial(
 
 	float3 L_World = TangentToWorld(SampledValue.xyz, N_World);
 
-	return CreateMaterialSample(L_World, Payload.GetBSDFOpacity() * Payload.BaseColor * ShadowTerminatorTerm(L_World, Payload.WorldNormal, Payload.WorldSmoothNormal), SampledValue.w, 1.0, 1.0, PATHTRACER_SCATTER_DIFFUSE);
+	return CreateMaterialSample(L_World, Payload.BSDFOpacity * Payload.BaseColor * ShadowTerminatorTerm(L_World, Payload.WorldNormal, Payload.WorldSmoothNormal), SampledValue.w, 1.0, 1.0, PATHTRACER_SCATTER_DIFFUSE);
 }
 
 FMaterialEval Lambert_EvalMaterial(
@@ -27,7 +27,7 @@ FMaterialEval Lambert_EvalMaterial(
 	float3 N_World = Payload.WorldNormal;
 	float NoL = saturate(dot(N_World, L_World));
 
-	return CreateMaterialEval(Payload.GetBSDFOpacity() * Payload.BaseColor * ShadowTerminatorTerm(L_World, Payload.WorldNormal, Payload.WorldSmoothNormal) * DiffuseSpecularScale.x, NoL / PI);
+	return CreateMaterialEval(Payload.BSDFOpacity * Payload.BaseColor * ShadowTerminatorTerm(L_World, Payload.WorldNormal, Payload.WorldSmoothNormal) * DiffuseSpecularScale.x, NoL / PI);
 }
 
 
@@ -47,7 +47,7 @@ FMaterialSample Simplified_SampleMaterial(
 	float3 Weight = bDoTransmission ? BackColor / TransmissionProb : FrontColor / (1 - TransmissionProb);
 	float Pdf = SampledValue.w * (bDoTransmission ? TransmissionProb : 1.0 - TransmissionProb);
 
-	Weight *= Payload.GetBSDFOpacity();
+	Weight *= Payload.BSDFOpacity;
 	float Sign = bDoTransmission ? -1.0 : 1.0;
 	return CreateMaterialSample(TangentToWorld(SampledValue.xyz, N_World * Sign), Weight, Pdf, Sign, 1.0, PATHTRACER_SCATTER_DIFFUSE);
 }
@@ -70,6 +70,6 @@ FMaterialEval Simplified_EvalMaterial(
 	{
 		Weight /= Pdf;
 	}
-	Weight *= Payload.GetBSDFOpacity();
+	Weight *= Payload.BSDFOpacity;
 	return CreateMaterialEval(Weight, Pdf * saturate(abs(NoL)) / PI);
 }
diff --git a/Engine/Shaders/Private/PathTracing/Material/PathTracingMaterialSampling.ush b/Engine/Shaders/Private/PathTracing/Material/PathTracingMaterialSampling.ush
--- a/Engine/Shaders/Private/PathTracing/Material/PathTracingMaterialSampling.ush
+++ b/Engine/Shaders/Private/PathTracing/Material/PathTracingMaterialSampling.ush
@@ -27,56 +27,6 @@ PathTracingMaterialSampling.usf: Material sampling functions
 	#define SIMPLIFIED_MATERIAL_SHADER 0
 #endif
 
-float3 GetMaterialTransparency(FPathTracingPayload Payload, float3 RayDirection)
-{
-	// ThinTranslucent is special, all other shading models should go through blend mode logic
-	if (Payload.ShadingModelID == SHADINGMODELID_THIN_TRANSLUCENT)
-	{
-#if SIMPLIFIED_MATERIAL_SHADER
-		const float VoN = dot(RayDirection, Payload.WorldNormal);
-		return (1 - Payload.Opacity) * pow(Payload.GetTransmittanceColor(), 1.0 / abs(VoN));
-#else
-		if (!Payload.HasRefraction() && Payload.Opacity < 1.0)
-		{
-			// ThinTranslucent has an optimized mode where it ignores refractivity and just handles the transmission via transparency
-			const float VoN = abs(dot(RayDirection, Payload.WorldNormal));
-			const float F0 = F0RGBToF0(Payload.SpecularColor);
-			return (1 - Payload.Opacity) * ComputeThinSlabWeights(Payload.GetTransmittanceColor(), VoN, 0.0, F0).Transmitted;
-		}
-		else
-		{
-			// opaque, or handled via the transmission bsdf lobe
-			return 0;
-		}
-#endif
-	}
-	switch (Payload.BlendingMode)
-	{
-		case RAY_TRACING_BLEND_MODE_TRANSLUCENT:
-#if !SIMPLIFIED_MATERIAL_SHADER
-			if (Payload.HasRefraction())
-			{
-				// not using transparency - material will be opaque
-				return 0;
-			}
-			else
-#endif
-			{
-				// material is not configured with refraction, can use transparency
-				return 1 - Payload.Opacity;
-			}
-		case RAY_TRACING_BLEND_MODE_ADDITIVE:
-			return 1;
-		case RAY_TRACING_BLEND_MODE_MODULATE:
-			return Payload.Radiance;
-		case RAY_TRACING_BLEND_MODE_ALPHA_COMPOSITE:
-		case RAY_TRACING_BLEND_MODE_ALPHA_HOLDOUT:
-			return 1 - Payload.Opacity;
-		default: 
-			return 0;
-	}
-}
-
 float3 EstimateMaterialAlbedo(FPathTracingPayload Payload)
 {
 	// TODO: make this more accurate by estimating the directional albedo
@@ -98,7 +48,6 @@ float3 EstimateMaterialAlbedo(FPathTracingPayload Payload)
 		case SHADINGMODELID_EYE:
 			Albedo = Payload.DiffuseColor + Payload.SpecularColor + Payload.SubsurfaceColor;
 			break;
-		case SHADINGMODELID_SINGLELAYERWATER:
 		case SHADINGMODELID_DEFAULT_LIT:
 		case SHADINGMODELID_CLEAR_COAT:
 			// TODO: include clearcoat albedo as well
@@ -112,7 +61,7 @@ float3 EstimateMaterialAlbedo(FPathTracingPayload Payload)
 			const float3 GlassAlbedo = Payload.HasRefraction() ? lerp(Payload.GetTransmittanceColor(), 1.0, F0) : F0;
 			// surface is same as default-lit
 			const float3 SurfAlbedo = Payload.DiffuseColor + Payload.SpecularColor;
-			return lerp(GlassAlbedo, SurfAlbedo, Payload.Opacity);
+			return lerp(GlassAlbedo, SurfAlbedo, Payload.BSDFOpacity);
 		}
 		case SHADINGMODELID_CLOTH:
 			Albedo = (Payload.DiffuseColor + lerp(Payload.SpecularColor, Payload.GetClothColor(), Payload.GetClothAmount()));
@@ -123,11 +72,11 @@ float3 EstimateMaterialAlbedo(FPathTracingPayload Payload)
 			Albedo = Payload.BaseColor;
 			break;
 	}
-	Albedo *= Payload.GetBSDFOpacity();
+	Albedo *= Payload.BSDFOpacity;
 #if !SIMPLIFIED_MATERIAL_SHADER
 	if (Payload.HasRefraction())
 	{
-		Albedo += 1 - Payload.Opacity; // glass albedo is 1.0
+		Albedo += 1 - Payload.BSDFOpacity; // glass albedo is 1.0
 	}
 #endif
 	return Albedo;
@@ -144,10 +93,11 @@ void AccumulateAlbedoNormal(inout FPathTracingPayload Payload, float3 PathThroug
 	}
 	// figure out how much the current path should contribute to the denoiser aovs
 	PathThroughput *= 1 - saturate((PathRoughness - MinRoughness) / (MaxRoughness - MinRoughness));
+	float NormThroughput = Luminance(PathThroughput);
 	// figure out how much the specular lobes should contribute to the total (relative to the current PathRoughness)
 	// if we are above the MinRoughness, use the PathRoughness as are lower bound instead so that multiple bounces don't accumulate too much albedo
 	const float RoughnessWeight = saturate((Payload.Roughness - max(MinRoughness, PathRoughness)) / (MaxRoughness - max(MinRoughness, PathRoughness)));
-	const float BrdfWeight = Payload.GetBSDFOpacity();
+	const float BrdfWeight = Payload.BSDFOpacity;
 	switch (Payload.ShadingModelID)
 	{
 	case SHADINGMODELID_UNLIT:
@@ -165,41 +115,40 @@ void AccumulateAlbedoNormal(inout FPathTracingPayload Payload, float3 PathThroug
 	case SHADINGMODELID_SUBSURFACE_PROFILE:
 	case SHADINGMODELID_EYE:
 		Albedo += BrdfWeight * PathThroughput * (Payload.DiffuseColor + RoughnessWeight * Payload.SpecularColor + Payload.SubsurfaceColor);
-		Normal += BrdfWeight * Luminance(PathThroughput) * Payload.WorldNormal;
+		Normal += BrdfWeight * NormThroughput * Payload.WorldNormal;
 		break;
-	case SHADINGMODELID_SINGLELAYERWATER:
 	case SHADINGMODELID_DEFAULT_LIT:
 	{
 		if (Payload.HasRefraction())
 		{
 			const float GlassAlbedo = 1.0; // solid glass has unit albedo
 			Albedo += (1 - BrdfWeight) * RoughnessWeight * PathThroughput * GlassAlbedo;
-			Normal += (1 - BrdfWeight) * RoughnessWeight * Luminance(PathThroughput) * Payload.WorldNormal;
+			Normal += (1 - BrdfWeight) * RoughnessWeight * NormThroughput * Payload.WorldNormal;
 		}
 		// fallthrough
 	}
 	case SHADINGMODELID_CLEAR_COAT:
 		Albedo += BrdfWeight * PathThroughput * (Payload.DiffuseColor + RoughnessWeight * Payload.SpecularColor);
-		Normal += BrdfWeight * Luminance(PathThroughput) * Payload.WorldNormal;
+		Normal += BrdfWeight * NormThroughput * Payload.WorldNormal;
 		break;
 	case SHADINGMODELID_CLOTH:
 		Albedo += BrdfWeight * PathThroughput * (Payload.DiffuseColor + lerp(RoughnessWeight * Payload.SpecularColor, Payload.GetClothColor(), Payload.GetClothAmount()));
-		Normal += BrdfWeight * Luminance(PathThroughput) * Payload.WorldNormal;
+		Normal += BrdfWeight * NormThroughput * Payload.WorldNormal;
 		break;
 	case SHADINGMODELID_THIN_TRANSLUCENT:
 	{
 		const float F0 = F0RGBToF0(Payload.SpecularColor);
 		const float3 GlassAlbedo = Payload.HasRefraction() ? lerp(Payload.GetTransmittanceColor(), 1.0, F0) : F0;
 		Albedo += (1 - BrdfWeight) * RoughnessWeight * PathThroughput * GlassAlbedo;
-		Normal += (1 - BrdfWeight) * RoughnessWeight * Luminance(PathThroughput) * Payload.WorldNormal;
+		Normal += (1 - BrdfWeight) * RoughnessWeight * NormThroughput * Payload.WorldNormal;
 		Albedo += BrdfWeight * PathThroughput * (Payload.DiffuseColor + RoughnessWeight * Payload.SpecularColor);
-		Normal += BrdfWeight * Luminance(PathThroughput) * Payload.WorldNormal;
+		Normal += BrdfWeight * NormThroughput * Payload.WorldNormal;
 		break;
 	}
 #endif
 	default:
 		Albedo += BrdfWeight * PathThroughput * Payload.BaseColor;
-		Normal += BrdfWeight * Luminance(PathThroughput) * Payload.WorldNormal;
+		Normal += BrdfWeight * NormThroughput * Payload.WorldNormal;
 		break;
 	}
 }
@@ -229,7 +178,7 @@ float GetAverageRoughness(FPathTracingPayload Payload)
 #if SIMPLIFIED_MATERIAL_SHADER
 	return 1.0;
 #else
-	const float GlassPercent = Payload.HasRefraction() ? 1 - Payload.Opacity : 0.0;
+	const float GlassPercent = Payload.HasRefraction() ? 1 - Payload.BSDFOpacity : 0.0;
 	float Roughness = Payload.Roughness;
 	switch (Payload.ShadingModelID)
 	{
@@ -345,6 +294,60 @@ FPathTracingPayload UnpackPathTracingPayloadMaterial(FPackedPathTracingPayload P
 	return HitPayload;
 }
 
+struct FSSSRandomWalkInfo
+{
+	float3 Color;
+	float3 Radius;
+
+	float3 Weight;  // multiplier for the path throughput to take into account energy conservation from specular above the SSS
+	float  Prob;    // probability of choosing to perform the random walk
+};
+
+
+FSSSRandomWalkInfo GetMaterialSSSInfo(FPathTracingPayload Payload, float3 V_World)
+{
+	FSSSRandomWalkInfo Result = (FSSSRandomWalkInfo)0;
+	Result.Color = Payload.SubsurfaceColor;
+	Result.Radius = Payload.SubsurfaceRadius;
+	Result.Weight = 0;
+	Result.Prob = 0;
+	if (Payload.IsSubsurfaceMaterial())
+	{
+		const float NoV = saturate(dot(Payload.WorldNormal, V_World));
+		float3 SpecE = 0;
+		if (Payload.ShadingModelID == SHADINGMODELID_SUBSURFACE_PROFILE)
+		{
+			// dual spec
+			const float3 DualRoughnessData = Payload.GetDualRoughnessSpecular();
+			const float3 SpecE0 = ComputeGGXSpecEnergyTermsRGB(DualRoughnessData.x, NoV, Payload.SpecularColor).E;
+			const float3 SpecE1 = ComputeGGXSpecEnergyTermsRGB(DualRoughnessData.y, NoV, Payload.SpecularColor).E;
+			SpecE = lerp(SpecE0, SpecE1, DualRoughnessData.z);
+		}
+		else
+		{
+			// single spec
+			SpecE = ComputeGGXSpecEnergyTermsRGB(Payload.Roughness, NoV, Payload.SpecularColor).E;
+		}
+
+		const float3 SSSLobeAlbedo = (1 - SpecE) * Payload.SubsurfaceColor;
+		const float3 DiffLobeAlbedo = (1 - SpecE) * Payload.DiffuseColor;
+		const float3 SpecLobeAlbedo = SpecE;
+		Result.Prob = LobeSelectionProb(SSSLobeAlbedo, DiffLobeAlbedo + SpecLobeAlbedo);
+		Result.Weight = Payload.BSDFOpacity * (1 - SpecE);
+	}
+	return Result;
+}
+
+void RemoveMaterialSSS(inout FPathTracingPayload Payload)
+{
+	// Tweak the payload to remove subsurface scattering without chaning the overall color
+	// This function is free to assume this material contains SSS
+	Payload.DiffuseColor += Payload.SubsurfaceColor;
+	Payload.DiffuseColor = saturate(Payload.DiffuseColor);
+	Payload.SubsurfaceColor = 0.0;
+}
+
+
 FMaterialSample SampleMaterial(
 	float3 V_World,
 	FPathTracingPayload Payload,
@@ -361,9 +364,8 @@ FMaterialSample SampleMaterial(
 			// Nothing to do, material does not scatter light
 			return NullMaterialSample();
 		case SHADINGMODELID_DEFAULT_LIT:
-		case SHADINGMODELID_SINGLELAYERWATER:
 		{
-			if (Payload.HasRefraction() && Payload.Opacity < 1)
+			if (Payload.HasRefraction() && Payload.BSDFOpacity < 1)
 			{
 				return RoughGlass_SampleMaterial(V_World, Payload, RandSample);
 			}
@@ -411,9 +413,8 @@ FMaterialEval EvalMaterial(
 		case SHADINGMODELID_UNLIT:
 			break;
 		case SHADINGMODELID_DEFAULT_LIT:
-		case SHADINGMODELID_SINGLELAYERWATER:
 		{
-			if (Payload.HasRefraction() && Payload.Opacity < 1)
+			if (Payload.HasRefraction() && Payload.BSDFOpacity < 1)
 			{
 				return RoughGlass_EvalMaterial(V_World, L_World, Payload, DiffuseSpecularScale);
 			}
diff --git a/Engine/Shaders/Private/PathTracing/Material/PathTracingSolidGlass.ush b/Engine/Shaders/Private/PathTracing/Material/PathTracingSolidGlass.ush
--- a/Engine/Shaders/Private/PathTracing/Material/PathTracingSolidGlass.ush
+++ b/Engine/Shaders/Private/PathTracing/Material/PathTracingSolidGlass.ush
@@ -42,7 +42,7 @@ FRoughGlassData PrepareRoughGlassData(FPathTracingPayload Payload, float3 V_Worl
 	Data.DiffWeight = Payload.DiffuseColor * (1 - Data.Spec.E);
 
 	// Probability of picking diffuse vs spec vs glass
-	Data.LobeCdf = LobeSelectionCdf(Payload.Opacity * Data.DiffWeight, Payload.Opacity * Data.Spec.E, 1.0 - Payload.Opacity);
+	Data.LobeCdf = LobeSelectionCdf(Payload.BSDFOpacity * Data.DiffWeight, Payload.BSDFOpacity * Data.Spec.E, 1.0 - Payload.BSDFOpacity);
 	Data.LobePdf = LobeSelectionPdf(Data.LobeCdf);
 
 
@@ -92,7 +92,7 @@ FMaterialSample RoughGlass_SampleMaterial(
 		const float NoH = saturate(H.z);
 
 		const float3 Diffuse = GetPathTracingDiffuseModel(Data.DiffWeight, Payload.Roughness, NoV, NoL, VoH, NoH);
-		const float3 DiffWeight = (Payload.Opacity / Data.LobePdf.x) * Diffuse * ShadowTerminatorTerm(L_World, Payload.WorldNormal, Payload.WorldSmoothNormal);
+		const float3 DiffWeight = (Payload.BSDFOpacity / Data.LobePdf.x) * Diffuse * ShadowTerminatorTerm(L_World, Payload.WorldNormal, Payload.WorldSmoothNormal);
 		const float DiffPdf = Data.LobePdf.x * NoL / PI;
 
 		return CreateMaterialSample(L_World, DiffWeight, DiffPdf, 1.0, 1.0, PATHTRACER_SCATTER_DIFFUSE);
@@ -134,7 +134,7 @@ FMaterialSample RoughGlass_SampleMaterial(
 			// Specular lobe
 			const float2 GGXResult = GGXEvalReflection(L, V, H, Data.Alpha);
 			const float3 F = F_Schlick(Payload.SpecularColor, VoH);
-			const float3 SpecWeight = (Payload.Opacity / Data.LobePdf.y) * F * GGXResult.x * Data.Spec.W;
+			const float3 SpecWeight = (Payload.BSDFOpacity / Data.LobePdf.y) * F * GGXResult.x * Data.Spec.W;
 			const float SpecPdf = GGXResult.y * Data.LobePdf.y;
 
 			return CreateMaterialSample(L_World, SpecWeight, SpecPdf, 1.0, Payload.Roughness, PATHTRACER_SCATTER_SPECULAR);
@@ -151,7 +151,7 @@ FMaterialSample RoughGlass_SampleMaterial(
 
 			const float2 GGXResult = bRefract ? GGXEvalRefraction(L, V, H, Data.Alpha, Data.Eta) : GGXEvalReflection(L, V, H, Data.Alpha);
 
-			return CreateMaterialSample(L_World, (1.0 - Payload.Opacity) * GGXResult.x / Data.EGlass, F * GGXResult.y, bRefract ? -1.0 : 1.0, Payload.Roughness, bRefract ? PATHTRACER_SCATTER_REFRACT : PATHTRACER_SCATTER_SPECULAR);
+			return CreateMaterialSample(L_World, (1.0 - Payload.BSDFOpacity) * GGXResult.x / Data.EGlass, F * GGXResult.y, bRefract ? -1.0 : 1.0, Payload.Roughness, bRefract ? PATHTRACER_SCATTER_REFRACT : PATHTRACER_SCATTER_SPECULAR);
 		}
 	}
 }
@@ -189,18 +189,18 @@ FMaterialEval RoughGlass_EvalMaterial(
 
 		// Diffuse Lobe
 		const float3 Diffuse = GetPathTracingDiffuseModel(Data.DiffWeight, Payload.Roughness, NoV, NoL, VoH, NoH);
-		Result.AddLobeWithMIS(Payload.Opacity * Diffuse * ShadowTerminatorTerm(L_World, Payload.WorldNormal, Payload.WorldSmoothNormal) * DiffuseSpecularScale.x, NoL / PI, Data.LobePdf.x);
+		Result.AddLobeWithMIS(Payload.BSDFOpacity * Diffuse * ShadowTerminatorTerm(L_World, Payload.WorldNormal, Payload.WorldSmoothNormal) * DiffuseSpecularScale.x, NoL / PI, Data.LobePdf.x);
 
 		// Specular lobe
 		const float2 GGXResult = GGXEvalReflection(L, V, H, Data.Alpha);
 		const float3 Fs = F_Schlick(Payload.SpecularColor, VoH);
-		const float3 SpecWeight = Payload.Opacity * GGXResult.x * Data.Spec.W * Fs;
+		const float3 SpecWeight = Payload.BSDFOpacity * GGXResult.x * Data.Spec.W * Fs;
 		const float SpecPdf = GGXResult.y;
 		Result.AddLobeWithMIS(SpecWeight * DiffuseSpecularScale.y, SpecPdf, Data.LobePdf.y);
 
 		// Glass lobe
 		const float Fg = FresnelReflectance(VoH, Data.Eta, Data.F0);
-		const float3 GlassWeight = GGXResult.x * (1 - Payload.Opacity) / Data.EGlass;
+		const float3 GlassWeight = GGXResult.x * (1 - Payload.BSDFOpacity) / Data.EGlass;
 		const float GlassPdf = GGXResult.y * Fg;
 		Result.AddLobeWithMIS(GlassWeight * DiffuseSpecularScale.y, GlassPdf, Data.LobePdf.z);
 	}
@@ -217,7 +217,7 @@ FMaterialEval RoughGlass_EvalMaterial(
 		{
 			const float2 GGXResult = GGXEvalRefraction(L, V, Ht, Data.Alpha, Data.Eta);
 
-			const float3 GlassWeight = GGXResult.x * (1 - Payload.Opacity) / Data.EGlass;
+			const float3 GlassWeight = GGXResult.x * (1 - Payload.BSDFOpacity) / Data.EGlass;
 			const float GlassPdf = GGXResult.y * Fg;
 
 			Result.AddLobeWithMIS(GlassWeight * DiffuseSpecularScale.y, GlassPdf, Data.LobePdf.z);
diff --git a/Engine/Shaders/Private/PathTracing/Material/PathTracingThinGlass.ush b/Engine/Shaders/Private/PathTracing/Material/PathTracingThinGlass.ush
--- a/Engine/Shaders/Private/PathTracing/Material/PathTracingThinGlass.ush
+++ b/Engine/Shaders/Private/PathTracing/Material/PathTracingThinGlass.ush
@@ -100,10 +100,10 @@ FRoughThinGlassData PrepareRoughThinGlassData(FPathTracingPayload Payload, float
 	// Figure out probability of the 4 lobes: diffuse,spec,glass reflection, glass transmission (if enabled)
 	// Approximate the transmission as just a plain tint by slab color because we can't depend on the half-vector for lobe selection
 	Data.LobeCdf = LobeSelectionCdf(
-		Payload.Opacity * Data.DiffWeight,
-		Payload.Opacity * Data.Spec.E,
-		(1 - Payload.Opacity) * Data.NudgeE * Data.SpecR.E,
-		(1 - Payload.Opacity) * Data.NudgeE * (Payload.HasRefraction() ? Data.SpecT.E * Payload.GetTransmittanceColor() : 0.0)
+		Payload.BSDFOpacity * Data.DiffWeight,
+		Payload.BSDFOpacity * Data.Spec.E,
+		(1 - Payload.BSDFOpacity) * Data.NudgeE * Data.SpecR.E,
+		(1 - Payload.BSDFOpacity) * Data.NudgeE * (Payload.HasRefraction() ? Data.SpecT.E * Payload.GetTransmittanceColor() : 0.0)
 	);
 	Data.LobePdf = LobeSelectionPdf(Data.LobeCdf);
 
@@ -146,26 +146,26 @@ FMaterialEval RoughThinGlass_EvalMaterial(
 	{
 		// Diffuse Lobe
 		const float3 Diffuse = GetPathTracingDiffuseModel(Data.DiffWeight, Payload.Roughness, NoV, NoL, VoH, NoH);
-		const float3 DiffWeight = Payload.Opacity * Diffuse * ShadowTerminatorTerm(L_World, Payload.WorldNormal, Payload.WorldSmoothNormal);
+		const float3 DiffWeight = Payload.BSDFOpacity * Diffuse * ShadowTerminatorTerm(L_World, Payload.WorldNormal, Payload.WorldSmoothNormal);
 		const float DiffPdf = NoL / PI;
 		Result.AddLobeWithMIS(DiffWeight * DiffuseSpecularScale.x, DiffPdf, Data.LobePdf.x);
 
 		// Specular lobe
 		const float2 GGXResult = GGXEvalReflection(L, V, H, Data.AlphaR);
 		const float3 F = F_Schlick(Payload.SpecularColor, VoH);
-		const float3 SpecWeight = Payload.Opacity * F * GGXResult.x * Data.Spec.W;
+		const float3 SpecWeight = Payload.BSDFOpacity * F * GGXResult.x * Data.Spec.W;
 		const float SpecPdf = GGXResult.y;
 		Result.AddLobeWithMIS(SpecWeight * DiffuseSpecularScale.y, SpecPdf, Data.LobePdf.y);
 
 
 		// Transmission reflection
-		const float3 GlassWeight = ((1 - Payload.Opacity) * Data.NudgeE * GGXResult.x * Data.SpecR.W) * SlabResult.Reflected;
+		const float3 GlassWeight = ((1 - Payload.BSDFOpacity) * Data.NudgeE * GGXResult.x * Data.SpecR.W) * SlabResult.Reflected;
 		const float GlassPdf = GGXResult.y;
 		Result.AddLobeWithMIS(GlassWeight * DiffuseSpecularScale.y, GlassPdf, Data.LobePdf.z);
 	}
 	else
 	{
-		const float3 GlassWeight = ((1 - Payload.Opacity) * Data.NudgeE * GGXResult.x * Data.SpecT.W) * SlabResult.Transmitted;
+		const float3 GlassWeight = ((1 - Payload.BSDFOpacity) * Data.NudgeE * GGXResult.x * Data.SpecT.W) * SlabResult.Transmitted;
 		const float GlassPdf = GGXResult.y;
 		Result.AddLobeWithMIS(GlassWeight * DiffuseSpecularScale.y, GlassPdf, Data.LobePdf.w);
 	}
@@ -204,7 +204,7 @@ FMaterialSample RoughThinGlass_SampleMaterial(
 
 		// Diffuse Lobe
 		const float3 Diffuse = GetPathTracingDiffuseModel(Data.DiffWeight, Payload.Roughness, NoV, NoL, VoH, NoH);
-		const float3 DiffWeight = Payload.Opacity * Diffuse * ShadowTerminatorTerm(L_World, Payload.WorldNormal, Payload.WorldSmoothNormal);
+		const float3 DiffWeight = Payload.BSDFOpacity * Diffuse * ShadowTerminatorTerm(L_World, Payload.WorldNormal, Payload.WorldSmoothNormal);
 		const float DiffPdf = NoL / PI;
 
 		FMaterialSample Result = CreateMaterialSample(L_World, 0.0, 0.0, 1.0, 1.0, PATHTRACER_SCATTER_DIFFUSE);
@@ -253,13 +253,13 @@ FMaterialSample RoughThinGlass_SampleMaterial(
 		{
 			// Specular lobe
 			const float3 F = F_Schlick(Payload.SpecularColor, VoH);
-			const float3 SpecWeight = Payload.Opacity * F * GGXResult.x * Data.Spec.W;
+			const float3 SpecWeight = Payload.BSDFOpacity * F * GGXResult.x * Data.Spec.W;
 			const float SpecPdf = GGXResult.y;
 			Result.AddLobeWithMIS(SpecWeight, SpecPdf, Data.LobePdf.y);
 
 
 			// Transmission reflection
-			const float3 GlassWeight = ((1 - Payload.Opacity) * Data.NudgeE * GGXResult.x * Data.SpecR.W) * SlabResult.Reflected;
+			const float3 GlassWeight = ((1 - Payload.BSDFOpacity) * Data.NudgeE * GGXResult.x * Data.SpecR.W) * SlabResult.Reflected;
 			const float GlassPdf = GGXResult.y;
 			Result.AddLobeWithMIS(GlassWeight, GlassPdf, Data.LobePdf.z);
 
@@ -270,7 +270,7 @@ FMaterialSample RoughThinGlass_SampleMaterial(
 		else
 		{
 			// transmitted ray
-			const float3 GlassWeight = ((1 - Payload.Opacity) * Data.NudgeE * GGXResult.x * Data.SpecT.W) * SlabResult.Transmitted;
+			const float3 GlassWeight = ((1 - Payload.BSDFOpacity) * Data.NudgeE * GGXResult.x * Data.SpecT.W) * SlabResult.Transmitted;
 			const float GlassPdf = GGXResult.y;
 			Result.AddLobeWithMIS(GlassWeight, GlassPdf, Data.LobePdf.w);
 
diff --git a/Engine/Shaders/Private/PathTracing/Material/PathTracingTwoSidedFoliage.ush b/Engine/Shaders/Private/PathTracing/Material/PathTracingTwoSidedFoliage.ush
--- a/Engine/Shaders/Private/PathTracing/Material/PathTracingTwoSidedFoliage.ush
+++ b/Engine/Shaders/Private/PathTracing/Material/PathTracingTwoSidedFoliage.ush
@@ -46,7 +46,7 @@ FMaterialSample TwoSidedFoliage_SampleMaterial(
 		const float3 N = Payload.WorldNormal;
 		const float4 SampledValue = CosineSampleHemisphere(RandSample.yz);
 
-		return CreateMaterialSample(TangentToWorld(SampledValue.xyz, -N), Payload.GetBSDFOpacity() * Data.Back / Data.ProbT, Data.ProbT * SampledValue.w, -1.0, 1.0, PATHTRACER_SCATTER_DIFFUSE);
+		return CreateMaterialSample(TangentToWorld(SampledValue.xyz, -N), Payload.BSDFOpacity * Data.Back / Data.ProbT, Data.ProbT * SampledValue.w, -1.0, 1.0, PATHTRACER_SCATTER_DIFFUSE);
 	}
 	else
 	{
@@ -79,7 +79,7 @@ FMaterialEval TwoSidedFoliage_EvalMaterial(
 		{
 			float LambertPdf = -NoL / PI;
 
-			return CreateMaterialEval(Payload.GetBSDFOpacity() * Data.Back / Data.ProbT * DiffuseSpecularScale.x, Data.ProbT * LambertPdf);
+			return CreateMaterialEval(Payload.BSDFOpacity * Data.Back / Data.ProbT * DiffuseSpecularScale.x, Data.ProbT * LambertPdf);
 		}
 	}
 	else if (Data.ProbT < 1)
diff --git a/Engine/Shaders/Private/PathTracing/PathTracingCommon.ush b/Engine/Shaders/Private/PathTracing/PathTracingCommon.ush
--- a/Engine/Shaders/Private/PathTracing/PathTracingCommon.ush
+++ b/Engine/Shaders/Private/PathTracing/PathTracingCommon.ush
@@ -11,10 +11,11 @@
 #include "../OctahedralCommon.ush"
 #include "/Engine/Shared/PathTracingDefinitions.h"
 
-// These flags are analogous to RAY_TRACING_PAYLOAD_* flags. There are currently 6 bits available
+// These flags are analogous to RAY_TRACING_PAYLOAD_* flags. There are currently 9 bits available
 #define PATH_TRACING_PAYLOAD_OUTPUT_FLAG_FRONT_FACE                   (1 << 0) // Indicates that ray has hit the front face of a primitive. Set by closest hit shader.
 #define PATH_TRACING_PAYLOAD_OUTPUT_FLAG_DECAL_RECEIVER               (1 << 1) // Indicates that ray has hit a primitive that receives decals. Set by closest hit shader.
 #define PATH_TRACING_PAYLOAD_OUTPUT_FLAG_TWO_SIDED_MATERIAL           (1 << 2) // Indicates that ray has hit a primitive surface with a two sided material. Set by closest hit shader. Used by GPU Lightmass to detect invalid surfaces.
+#define PATH_TRACING_PAYLOAD_OUTPUT_FLAG_IS_HOLDOUT                   (1 << 3) // Indicates that ray has hit a primitive which should be held out from alpha
 
 #define PATH_TRACING_PAYLOAD_INPUT_FLAG_RAY_TYPE_MASK                 (7 << 0) // Reserve 3 bits for ray-mask
 #define PATH_TRACING_PAYLOAD_INPUT_FLAG_RAY_TYPE_CAMERA                0       // Indicates that this is a camera ray which can trigger a number of special behaviors in the closest hit shader
@@ -28,6 +29,20 @@
 // shading model extensions beyond the default count
 #define SHADINGMODELID_MEDIUM  (SHADINGMODELID_NUM + 1)
 
+#define PATH_TRACING_ABSORPTION_SCALE  0.01        // inverse distance at which we hit the transmission color (Beer's law)
+
+// This function is meant for the solid glass case so we can store the transmittance as a color in [0,1]
+float3 PathTracingGlassTransmittanceToExtinction(float3 TransmittanceColor)
+{
+	return -log(max(TransmittanceColor, 1e-8f)) * PATH_TRACING_ABSORPTION_SCALE;
+}
+
+float3 PathTracingGlassExtinctionToTransmittance(float3 ExtinctionCoefficient)
+{
+	// inverse of function above
+	return exp(-max(ExtinctionCoefficient, 0.0) / PATH_TRACING_ABSORPTION_SCALE);
+}
+
 
 uint PayloadEncodeUnitVector(float3 Normal)
 {
@@ -49,6 +64,43 @@ float3 PayloadDecodeUnitVector(uint Encoded)
 	return OctahedronToUnitVector(2 * OctF - 1);
 }
 
+uint PayloadEncodeHDRColor(float3 rgb)
+{
+	// clamp to the largest representable value
+	rgb = clamp(rgb, 0.0, Max111110BitsFloat3);
+	// See PackR11G11B10F (duplicated to keep the code independent of Strata for now)
+	uint r = (f32tof16(rgb.r) << 17) & 0xFFE00000;
+	uint g = (f32tof16(rgb.g) <<  6) & 0x001FFC00;
+	uint b = (f32tof16(rgb.b) >>  5) & 0x000003FF;
+	return r | g | b;
+}
+
+float3 PayloadDecodeHDRColor(uint rgb)
+{
+	float r = f16tof32((rgb >> 17) & 0x7FF0);
+	float g = f16tof32((rgb >>  6) & 0x7FF0);
+	float b = f16tof32((rgb <<  5) & 0x7FE0);
+	return float3(r, g, b);
+}
+
+uint PayloadEncodeLDRColor(float3 rgb)
+{
+	// NOTE: Could add a simple gamma2 here to improve precision of dark colors, but this does not seem necessary for now.
+	// unorm 11/11/10 encoding
+	uint r = uint(floor(saturate(rgb.r) * 2047.0 + 0.5));
+	uint g = uint(floor(saturate(rgb.g) * 2047.0 + 0.5));
+	uint b = uint(floor(saturate(rgb.b) * 1023.0 + 0.5));
+	return (r << 21) | (g << 10) | b;
+}
+
+float3 PayloadDecodeLDRColor(uint rgb)
+{
+	float r = float((rgb >> 21)       ) * (1.0 / 2047.0);
+	float g = float((rgb >> 10) & 2047) * (1.0 / 2047.0);
+	float b = float((rgb      ) & 1023) * (1.0 / 1023.0);
+	return float3(r, g, b);
+}
+
 float3 AdjustShadingNormal(float3 ShadingNormal, float3 GeoNormal, float3 RayDirection)
 {
 	// Clip shading normal in a view dependent way such that the reflection stays above the geometric normal
@@ -74,8 +126,8 @@ struct FPackedPathTracingPayload : FMinimalPayload
 	uint PackedData[15];								 // 60 bytes (encoded data, depends on shading model and ray type)
 														 // 64 bytes total
 
-	uint GetFlags() { return (PackedData[0] >> 23) & 0x3F; }
-	void SetFlag(uint Flag) { PackedData[0] |= Flag << 23; }
+	uint GetFlags() { return (PackedData[0] >> 20) & 0x1FF; }
+	void SetFlag(uint Flag) { PackedData[0] |= Flag << 20; }
 	uint GetPrimitiveLightingChannelMask() { return (PackedData[0] >> 29) & 0x7; }
 
 	// Flag methods
@@ -146,14 +198,14 @@ struct FPathTracingPayload : FMinimalPayload
 	float3 BaseColor;
 	float3 DiffuseColor;
 	float3 SpecularColor;
-	float  Opacity;
+	float3 TransparencyColor;        // how much do we see straight through the surface?
+	float  BSDFOpacity;              // how much should we weigh down the bsdf by? (for glass, controls the blend between solid and glass lobes)
 	float  Metallic;
 	float  Specular;
 	float  Roughness;
 	float  Anisotropy;
 	float  Ior;
 	uint   ShadingModelID;
-	uint   BlendingMode;
 	uint   Flags;
 	uint   PrimitiveLightingChannelMask;
 	float4 CustomData0;
@@ -169,6 +221,9 @@ struct FPathTracingPayload : FMinimalPayload
 	void SetDecalReceiver() { Flags |= PATH_TRACING_PAYLOAD_OUTPUT_FLAG_DECAL_RECEIVER; }
 	bool IsDecalReceiver() { return (Flags & PATH_TRACING_PAYLOAD_OUTPUT_FLAG_DECAL_RECEIVER) != 0; }
 	
+	void SetHoldout() { Flags |= PATH_TRACING_PAYLOAD_OUTPUT_FLAG_IS_HOLDOUT; }
+	bool IsHoldout() { return (Flags & PATH_TRACING_PAYLOAD_OUTPUT_FLAG_IS_HOLDOUT) != 0; }
+
 	void SetMaterialTwoSided() { Flags |= PATH_TRACING_PAYLOAD_OUTPUT_FLAG_TWO_SIDED_MATERIAL; }
 	bool IsMaterialTwoSided() { return (Flags & PATH_TRACING_PAYLOAD_OUTPUT_FLAG_TWO_SIDED_MATERIAL) != 0; }
 
@@ -220,12 +275,6 @@ struct FPathTracingPayload : FMinimalPayload
 	float3 GetTransmittanceColor() { return CustomData0.xyz; }
 	void SetTransmittanceColor(float3 TransmittanceColor) { CustomData0.xyz = TransmittanceColor; }
 
-	// SHADINGMODELID_SINGLELAYERWATER
-	float3 GetWaterExtinction() { return CustomData0.xyz; }
-	void SetWaterExtinction(float3 SigmaT) { CustomData0.xyz = SigmaT; }
-	float GetWaterPhaseG() { return CustomData0.w; }
-	void SetWaterPhaseG(float G) { CustomData0.w = G; }
-
 	// SHADINGMODELID_MEDIUM
 	float3 GetHGWeight() { return CustomData0.xyz; }
 	float GetHGPhaseG() { return CustomData0.w; }
@@ -234,12 +283,6 @@ struct FPathTracingPayload : FMinimalPayload
 
 	// Methods used by the integrator
 
-	float GetBSDFOpacity()
-	{
-		// How much should the BSDF components by weighted?
-		return BlendingMode == RAY_TRACING_BLEND_MODE_ALPHA_COMPOSITE ? 1.0 : Opacity;
-	}
-
 	bool HasRefraction()
 	{
 		return Ior > 0.0;
@@ -247,12 +290,12 @@ struct FPathTracingPayload : FMinimalPayload
 
 	bool IsMaterialSolidGlass()
 	{
-		return ShadingModelID == SHADINGMODELID_DEFAULT_LIT && BlendingMode == RAY_TRACING_BLEND_MODE_TRANSLUCENT && HasRefraction() && Opacity < 1;
+		return ShadingModelID == SHADINGMODELID_DEFAULT_LIT && HasRefraction() && BSDFOpacity < 1;
 	}
 
 	bool IsMaterialThinGlass()
 	{
-		return ShadingModelID == SHADINGMODELID_THIN_TRANSLUCENT && BlendingMode == RAY_TRACING_BLEND_MODE_TRANSLUCENT && HasRefraction() && Opacity < 1;
+		return ShadingModelID == SHADINGMODELID_THIN_TRANSLUCENT && HasRefraction() && BSDFOpacity < 1;
 	}
 
 	bool IsMaterialTransmissive()
@@ -260,7 +303,6 @@ struct FPathTracingPayload : FMinimalPayload
 		return ShadingModelID == SHADINGMODELID_TWOSIDED_FOLIAGE ||
 			   ShadingModelID == SHADINGMODELID_HAIR ||
 			   ShadingModelID == SHADINGMODELID_MEDIUM ||
-			   ShadingModelID == SHADINGMODELID_SINGLELAYERWATER ||	
 			   IsMaterialSolidGlass() ||
 			   IsMaterialThinGlass();
 	}
@@ -304,29 +346,23 @@ FPackedPathTracingPayload InitPathTracingVisibilityPayload(float PathRoughness)
 	return Output;
 }
 
-
 FPackedPathTracingPayload PackPathTracingPayload(FPathTracingPayload Input)
 {
 	FPackedPathTracingPayload Output = (FPackedPathTracingPayload)0;
 	Output.HitT = Input.HitT;
 
-	Output.PackedData[0] = f32tof16(Input.Opacity);
+	Output.PackedData[0] = f32tof16(Input.BSDFOpacity);
 	Output.PackedData[0] |= (Input.ShadingModelID & 0xF) << 16;               // 4 bits
-	Output.PackedData[0] |= (Input.BlendingMode & 0x7) << 20;                 // 3 bits
-	Output.PackedData[0] |= (Input.Flags & 0x3F) << 23;                       // 6 bits
+	Output.PackedData[0] |= (Input.Flags & 0x1FF) << 20;                      // 9 bits
 	Output.PackedData[0] |= (Input.PrimitiveLightingChannelMask & 0x7) << 29; // 3 bits
 	                                                                   // total 32 bits
 	Output.PackedData[1] = PayloadEncodeUnitVector(Input.WorldGeoNormal);
 	Output.PackedData[2] = PayloadEncodeUnitVector(Input.WorldSmoothNormal);
 	Output.PackedData[3] = PayloadEncodeUnitVector(Input.WorldNormal);
 
-	Output.PackedData[4] = f32tof16(Input.Radiance.x);
-	Output.PackedData[4] |= f32tof16(Input.Radiance.y) << 16;
-	Output.PackedData[5] = f32tof16(Input.Radiance.z);
-
-	Output.PackedData[5] |= f32tof16(Input.BaseColor.x) << 16;
-	Output.PackedData[6] = f32tof16(Input.BaseColor.y);
-	Output.PackedData[6] |= f32tof16(Input.BaseColor.z) << 16;
+	Output.PackedData[4] = PayloadEncodeHDRColor(Input.Radiance);
+	Output.PackedData[5] = PayloadEncodeLDRColor(Input.BaseColor);
+	Output.PackedData[6] = PayloadEncodeLDRColor(Input.TransparencyColor);
 
 	Output.PackedData[7] = f32tof16(Input.Metallic);
 	Output.PackedData[7] |= f32tof16(Input.Specular) << 16;
@@ -356,9 +392,8 @@ FPathTracingPayload UnpackPathTracingPayload(FPackedPathTracingPayload Input, FR
 	FPathTracingPayload Output = (FPathTracingPayload)0;
 
 	Output.HitT = Input.HitT;
-	Output.Opacity = f16tof32(Input.PackedData[0]);
+	Output.BSDFOpacity = f16tof32(Input.PackedData[0]);
 	Output.ShadingModelID = (Input.PackedData[0] >> 16) & 0xF;
-	Output.BlendingMode = (Input.PackedData[0] >> 20) & 0x7;
 	Output.Flags = Input.GetFlags();
 	Output.PrimitiveLightingChannelMask = Input.GetPrimitiveLightingChannelMask();
 
@@ -366,12 +401,9 @@ FPathTracingPayload UnpackPathTracingPayload(FPackedPathTracingPayload Input, FR
 	Output.WorldSmoothNormal = PayloadDecodeUnitVector(Input.PackedData[2]);
 	Output.WorldNormal = PayloadDecodeUnitVector(Input.PackedData[3]);
 
-	Output.Radiance.x = f16tof32(Input.PackedData[4]);
-	Output.Radiance.y = f16tof32(Input.PackedData[4] >> 16);
-	Output.Radiance.z = f16tof32(Input.PackedData[5]);
-	Output.BaseColor.x = f16tof32(Input.PackedData[5] >> 16);
-	Output.BaseColor.y = f16tof32(Input.PackedData[6]);
-	Output.BaseColor.z = f16tof32(Input.PackedData[6] >> 16);
+	Output.Radiance = PayloadDecodeHDRColor(Input.PackedData[4]);
+	Output.BaseColor = PayloadDecodeLDRColor(Input.PackedData[5]);
+	Output.TransparencyColor = PayloadDecodeLDRColor(Input.PackedData[6]);
 	
 	Output.Metallic = f16tof32(Input.PackedData[7]);
 	Output.Specular = f16tof32(Input.PackedData[7] >> 16);
diff --git a/Engine/Shaders/Private/PathTracing/PathTracingCore.ush b/Engine/Shaders/Private/PathTracing/PathTracingCore.ush
--- a/Engine/Shaders/Private/PathTracing/PathTracingCore.ush
+++ b/Engine/Shaders/Private/PathTracing/PathTracingCore.ush
@@ -59,7 +59,6 @@ float MeshDecalBias;
 float MaxPathIntensity;
 float MaxNormalBias;
 float FilterWidth;
-float AbsorptionScale; // inverse distance at which we hit BaseColor - exposed mainly so we can disable absorption handling and experiment with the default scale factor
 float CameraFocusDistance;
 float2 CameraLensRadius;
 
@@ -223,8 +222,6 @@ void ApplyDecal(inout FPathTracingPayload HitPayload, FDecalShaderPayload DecalP
 	{
 		HitPayload.Radiance += DecalPayload.GetEmissive();
 	}
-
-	// TODO: Payload.Opacity?
 }
 
 void TraceDecalRay(FRayDesc DecalRay, bool bReverseCulling, uint InstanceInclusionMask, inout FPathTracingPayload HitPayload, inout float DecalThroughput)
@@ -334,8 +331,9 @@ FPathTracingPayload TraceTransparentRay(inout FPathState PathState, bool bIsCame
 				LightHitPayload.HitT = LightHit.HitT;
 				LightHitPayload.Radiance = LightHit.Radiance;
 				LightHitPayload.ShadingModelID = SHADINGMODELID_UNLIT;
-				// alpha composite assumes premult, which allows emissive hits that are fully transparent
-				LightHitPayload.BlendingMode = RAY_TRACING_BLEND_MODE_ALPHA_COMPOSITE;
+				// light hits are fully transparent to match how the light loop works (the light surface does not cast a shadow)
+				LightHitPayload.BSDFOpacity = 1.0;
+				LightHitPayload.TransparencyColor = 0.0;
 				LightHitPayload.PrimitiveLightingChannelMask = 7;
 				LightHitPayload.SetFrontFace();
 				PackedPayload = PackPathTracingPayload(LightHitPayload);
@@ -539,20 +537,13 @@ FPathTracingPayload TraceTransparentRay(inout FPathState PathState, bool bIsCame
 		// TODO: merge this with the volume raymarching?
 		PathState.PathThroughput *= select(PathState.SigmaT > 0.0, exp(-PathState.SigmaT * (HitPayload.HitT - PathState.Ray.TMin)), 1.0);
 
-		if (HitPayload.BlendingMode == RAY_TRACING_BLEND_MODE_ADDITIVE ||
-			HitPayload.BlendingMode == RAY_TRACING_BLEND_MODE_TRANSLUCENT)
+		if (HitPayload.IsHoldout() && bIsCameraRay)
 		{
-			// emission needs to be pre-multiplied by the opacity
-			HitPayload.Radiance *= HitPayload.Opacity;
+			PathState.BackgroundVisibility += Luminance(PathState.PathThroughput * HitPayload.BSDFOpacity);
 		}
 
-		if (HitPayload.BlendingMode == RAY_TRACING_BLEND_MODE_ALPHA_HOLDOUT && bIsCameraRay)
-		{
-			PathState.BackgroundVisibility += Luminance(PathState.PathThroughput * HitPayload.Opacity);
-		}
-
-		// add in surface emission (except for modulate blend mode which uses Radiance as the transparency multiplier)
-		if (bIncludeEmission && HitPayload.BlendingMode != RAY_TRACING_BLEND_MODE_MODULATE)
+		// add in surface emission
+		if (bIncludeEmission)
 		{
 			AccumulateRadiance(PathState.Radiance, PathState.PathThroughput * HitPayload.Radiance, bIsCameraRay);
 		}
@@ -580,7 +571,7 @@ FPathTracingPayload TraceTransparentRay(inout FPathState PathState, bool bIsCame
 		}
 
 		// account for local transparency change
-		PathState.PathThroughput *= GetMaterialTransparency(HitPayload, PathState.Ray.Direction);;
+		PathState.PathThroughput *= HitPayload.TransparencyColor;
 
 		// prepare next step around the loop
 		// retrace the exact same ray with TMin one ulp past the hit we just found
@@ -695,8 +686,7 @@ float3 TraceTransparentVisibilityRay(FRayDesc Ray, float PathRoughness, uint Mis
 			Throughput = VolumeGetTransmittance(Throughput, Ray.Origin, Ray.Direction, VolumeIsectList, RandSequence);
 		}
 		// add in the medium extinction
-		const float3 Tau = PackedPayload.GetTau() * AbsorptionScale;
-		Throughput *= exp(-max(Tau, 0.0));
+		Throughput *= exp(-max(PackedPayload.GetTau(), 0.0));
 	}
 	return Throughput;
 #endif
@@ -774,57 +764,30 @@ bool ProcessSubsurfaceRandomWalk(inout FPathTracingPayload Payload, inout float3
 		return true;
 	}
 
-	float3 SSSColor = Payload.SubsurfaceColor;
-	float3 SSSRadius = Payload.SubsurfaceRadius;
+	FSSSRandomWalkInfo SSS = GetMaterialSSSInfo(Payload, -RayDirection);
+
 
-	if (SimplifySSS || all(SSSRadius == 0) || all(SSSColor == 0) || MaxSSSBounces == 0)
+	if (SimplifySSS || all(SSS.Radius == 0) || all(SSS.Color == 0) || MaxSSSBounces == 0)
 	{
 		// if we are running in a context that won't see the SSS directly -- just turn it into diffuse and skip the random walk
 		// we can also skip random walk if the radius or color is 0
 		// or if the user decided not to perform any SSS bounces
-		Payload.DiffuseColor += SSSColor;
-		Payload.DiffuseColor = saturate(Payload.DiffuseColor);
-		Payload.SubsurfaceColor = 0.0;
+		RemoveMaterialSSS(Payload);
 		return true;
 	}
 	
-	// Support SSS for mfp smaller than 1mm.
-	SSSRadius = max(SSSRadius, 0.0009); // 0.009mm minimum to make sure random walk can move forward and to keep SigmaT finite
-
-
 	// decide randomly to evaluate the SSS part of the material, or keep shading the current part
-	// This also computes the energy conservation to take out the portion of lighting reflected away by the specular layer
-	const float NoV = saturate(dot(Payload.WorldNormal, -RayDirection));
-	float3 SpecE = 0;
-	if (Payload.ShadingModelID == SHADINGMODELID_SUBSURFACE_PROFILE)
-	{
-		// dual spec
-		const float3 DualRoughnessData = Payload.GetDualRoughnessSpecular();
-		const float3 SpecE0 = ComputeGGXSpecEnergyTermsRGB(DualRoughnessData.x, NoV, Payload.SpecularColor).E;
-		const float3 SpecE1 = ComputeGGXSpecEnergyTermsRGB(DualRoughnessData.y, NoV, Payload.SpecularColor).E;
-		SpecE = lerp(SpecE0, SpecE1, DualRoughnessData.z);
-	}
-	else
-	{
-		// single spec
-		SpecE = ComputeGGXSpecEnergyTermsRGB(Payload.Roughness, NoV, Payload.SpecularColor).E;
-	}
-
-	const float3 SSSLobeAlbedo = (1 - SpecE) * SSSColor;
-	const float3 DiffLobeAlbedo = (1 - SpecE) * Payload.DiffuseColor;
-	const float3 SpecLobeAlbedo = SpecE;
-	float SSSProb = LobeSelectionProb(SSSLobeAlbedo, DiffLobeAlbedo + SpecLobeAlbedo);
 	float3 RandSample = RandomSequence_GenerateSample3D(RandSequence);
-	if (RandSample.x < SSSProb)
+	if (RandSample.x < SSS.Prob)
 	{
 		// do random walk (and adjust throughput to remove energy reflected away by the specular lobe)
-		PathThroughput *= (1 - SpecE) / SSSProb;
-		RandSample.x /= SSSProb;
+		PathThroughput *= SSS.Weight / SSS.Prob;
+		RandSample.x /= SSS.Prob;
 	}
 	else
 	{
 		// skip random walk and shade current point instead
-		PathThroughput *= 1 / (1 - SSSProb);
+		PathThroughput *= 1 / (1 - SSS.Prob);
 		return true;
 	}
 
@@ -835,6 +798,9 @@ bool ProcessSubsurfaceRandomWalk(inout FPathTracingPayload Payload, inout float3
 	Ray.TMin = 0;
 	ApplyRayBias(Ray, Payload.HitT, -Payload.WorldGeoNormal);
 
+	// Support SSS for mfp smaller than 1mm.
+	SSS.Radius = max(SSS.Radius, 0.0009); // 0.009mm minimum to make sure random walk can move forward and to keep SigmaT finite
+
     // Use the technique mentioned in "A Hero Beneath the Surface" (Section 6.3.2)
 	// https://jo.dreggn.org/home/2021_spectral_imaging.pdf
 	// This allows multiple overlapping meshes to be treated as a single "volume" by discarding internal hits until we have
@@ -850,13 +816,13 @@ bool ProcessSubsurfaceRandomWalk(inout FPathTracingPayload Payload, inout float3
 
 
 #if SSS_USE_TIR
-	float3 Albedo = 1 - exp(SSSColor * (-11.43 + SSSColor * (15.38 - 13.91 * SSSColor)));
+	float3 Albedo = 1 - exp(SSS.Color * (-11.43 + SSS.Color * (15.38 - 13.91 * SSS.Color)));
 #else
 	// Van de-Hulst inverse mapping
 	// https://blog.selfshadow.com/publications/s2017-shading-course/imageworks/s2017_pbs_imageworks_slides_v2.pdf (Slide 44)
 	// [REDACTED_PROFILE_URL] (Section 7.5.3 of v0.1.3)
-	float3 Albedo = 1 - Pow2(4.09712 + 4.20863 * SSSColor - sqrt(9.59217 + SSSColor * (41.6808 + 17.7126 * SSSColor)));
-	SSSRadius *= 2.0; // roughly match parameterization above
+	float3 Albedo = 1 - Pow2(4.09712 + 4.20863 * SSS.Color - sqrt(9.59217 + SSS.Color * (41.6808 + 17.7126 * SSS.Color)));
+	SSS.Radius *= 2.0; // roughly match parameterization above
 #endif
 
 	// Apply albedo remapping for anisotropic scattering 
@@ -867,7 +833,7 @@ bool ProcessSubsurfaceRandomWalk(inout FPathTracingPayload Payload, inout float3
 	Albedo = Albedo / (1 - G * (1 - Albedo));
 
 	const int MAX_SSS_BOUNCES = MaxSSSBounces;
-	const float3 SigmaT = 1.0 / SSSRadius;
+	const float3 SigmaT = rcp(SSS.Radius);
 	const float3 SigmaS = Albedo * SigmaT;
 	for (int i = 0; i < MAX_SSS_BOUNCES; i++)
 	{
@@ -946,8 +912,9 @@ bool ProcessSubsurfaceRandomWalk(inout FPathTracingPayload Payload, inout float3
 					Payload.WorldGeoNormal *= -1.0;
 				}
 				Payload.ShadingModelID = SHADINGMODELID_NUM; // invalid value so that we get diffuse shading
-				Payload.BlendingMode = RAY_TRACING_BLEND_MODE_OPAQUE;
+				Payload.BSDFOpacity = 1;
 				Payload.BaseColor = 1;
+				Payload.TransparencyColor = 0;
 				return true;
 			}
 		}
@@ -1187,8 +1154,10 @@ bool PathTracingKernel(inout FPathState PathState, int Bounce)
 		return false;
 	}
 
+	const bool bIsVolumeSample = Payload.ShadingModelID == SHADINGMODELID_MEDIUM;
+
 	FLightLoopCount LightLoopCount = LightGridLookup(Payload.TranslatedWorldPos);
-	if (Payload.ShadingModelID == SHADINGMODELID_MEDIUM && VolumeMISMode != 0)
+	if (bIsVolumeSample && VolumeMISMode != 0)
 	{
 		// if we are using the volume segment for local lights, exclude them from the light loop below
 		LightLoopCount.NumLights = SceneInfiniteLightCount;
@@ -1217,8 +1186,6 @@ bool PathTracingKernel(inout FPathState PathState, int Bounce)
 	// Choose a random number for both Light sampling and BxDF sampling
 	float4 RandSample = RandomSequence_GenerateSample4D(PathState.RandSequence);
 
-	const bool bIsVolumeSample = Payload.ShadingModelID == SHADINGMODELID_MEDIUM;
-
 	float2 DiffuseSpecularScale = PathState.GetDiffuseSpecularScale(bIsVolumeSample);
 
 	const bool bDoLightLoop = any(DiffuseSpecularScale > 0);
@@ -1370,12 +1337,9 @@ bool PathTracingKernel(inout FPathState PathState, int Bounce)
 	// update the current extinction if we are crossing a boundary on glass or water
 	// summing the local extinction gives a rudimentary way of dealing with overlapping regions
 	// long term we will probably want a stack with priorities
-	if (MaterialSample.PositionBiasSign < 0 && (Payload.IsMaterialSolidGlass() || Payload.ShadingModelID == SHADINGMODELID_SINGLELAYERWATER))
+	if (MaterialSample.PositionBiasSign < 0 && Payload.IsMaterialSolidGlass())
 	{
-		const float3 LocalSigmaT =
-			Payload.ShadingModelID == SHADINGMODELID_SINGLELAYERWATER ?
-			Payload.GetWaterExtinction() :
-			-log(max(Payload.GetTransmittanceColor(), 1e-8f)) * AbsorptionScale;
+		const float3 LocalSigmaT = PathTracingGlassTransmittanceToExtinction(Payload.GetTransmittanceColor());
 		if (Payload.IsFrontFace())
 		{
 			// entering
diff --git a/Engine/Shaders/Private/PathTracing/PathTracingDefaultHitShader.usf b/Engine/Shaders/Private/PathTracing/PathTracingDefaultHitShader.usf
--- a/Engine/Shaders/Private/PathTracing/PathTracingDefaultHitShader.usf
+++ b/Engine/Shaders/Private/PathTracing/PathTracingDefaultHitShader.usf
@@ -17,8 +17,7 @@ RAY_TRACING_ENTRY_CLOSEST_HIT(PathTracingDefaultOpaqueCHS,
 		// mark hit as solid black
 		FPathTracingPayload Payload = (FPathTracingPayload)0;
 		Payload.ShadingModelID = SHADINGMODELID_UNLIT;
-		Payload.BlendingMode   = RAY_TRACING_BLEND_MODE_OPAQUE;
-		Payload.Opacity = 1.0;
+		Payload.BSDFOpacity = 1.0;
 		PackedPayload = PackPathTracingPayload(Payload);
 	}
 	PackedPayload.HitT = RayTCurrent();
@@ -38,8 +37,7 @@ RAY_TRACING_ENTRY_CLOSEST_HIT(PathTracingDefaultHiddenCHS,
 		// mark hit as fully transparent
 		FPathTracingPayload Payload = (FPathTracingPayload)0;
 		Payload.ShadingModelID = SHADINGMODELID_UNLIT;
-		Payload.BlendingMode = RAY_TRACING_BLEND_MODE_ADDITIVE;
-		Payload.Opacity = 0.0;
+		Payload.TransparencyColor = 1.0;
 		PackedPayload = PackPathTracingPayload(Payload);
 	}
 	PackedPayload.HitT = RayTCurrent();
diff --git a/Engine/Shaders/Private/PathTracing/PathTracingInitExtinctionCoefficient.usf b/Engine/Shaders/Private/PathTracing/PathTracingInitExtinctionCoefficient.usf
--- a/Engine/Shaders/Private/PathTracing/PathTracingInitExtinctionCoefficient.usf
+++ b/Engine/Shaders/Private/PathTracing/PathTracingInitExtinctionCoefficient.usf
@@ -6,14 +6,11 @@
 RaytracingAccelerationStructure TLAS;
 RWBuffer<float> RWStartingExtinctionCoefficient;
 
-// TODO: figure out the right start value for camera rays inside glass
 RAY_TRACING_ENTRY_RAYGEN(PathTracingInitExtinctionCoefficientRG)
 {
-	const uint RayFlags = RAY_FLAG_FORCE_OPAQUE;
+	const uint RayFlags = 0;
 	const uint MissShaderIndex = 0;
-	const uint InstanceInclusionMask = PATHTRACER_MASK_ALL 
-										- PATHTRACER_MASK_HAIR_INDIRECT 
-										- PATHTRACER_MASK_HAIR_CAMERA; // Ignore hair strands
+	const uint InstanceInclusionMask = PATHTRACER_MASK_CAMERA; // Ignore hair strands, only trace what camera rays can see
 
 	FRayDesc UpRay;
 	UpRay.Origin = PrimaryView.TranslatedWorldCameraOrigin;
@@ -40,16 +37,22 @@ RAY_TRACING_ENTRY_RAYGEN(PathTracingInitExtinctionCoefficientRG)
 			break;
 		}
 		FPathTracingPayload HitPayload = UnpackPathTracingPayload(PackedPayload, UpRay);
-		if (HitPayload.ShadingModelID == SHADINGMODELID_SINGLELAYERWATER)
+		if (HitPayload.IsMaterialSolidGlass())
 		{
-			// Found water! we are done
-			StartingSigmaT = HitPayload.GetWaterExtinction();
-			break;
+			// Found a solid transmissive medium -- pickup the extinction
+			float3 LocalSigmaT = PathTracingGlassTransmittanceToExtinction(HitPayload.GetTransmittanceColor());
+
+			// increase when we leave a medium (that means we must have been inside before)
+			// decrease when we enter a medium (will be cancelled out when we leave)
+			StartingSigmaT += HitPayload.IsFrontFace() ? -LocalSigmaT : +LocalSigmaT;
 		}
-		// otherwise, keep tracing
-		UpRay.TMin = asfloat(asuint(HitPayload.HitT) + 1);
+		// keep tracing
+		UpRay.TMin = asfloat(asuint(PackedPayload.HitT) + 1);
 	}
 
+	// clamp against 0 in the unlikely case that we only hit front faces of glass on our way up
+	StartingSigmaT = max(StartingSigmaT, 0.0);
+
 	RWStartingExtinctionCoefficient[0] = StartingSigmaT.x;
 	RWStartingExtinctionCoefficient[1] = StartingSigmaT.y;
 	RWStartingExtinctionCoefficient[2] = StartingSigmaT.z;
diff --git a/Engine/Shaders/Private/PathTracing/PathTracingMaterialHitShader.usf b/Engine/Shaders/Private/PathTracing/PathTracingMaterialHitShader.usf
--- a/Engine/Shaders/Private/PathTracing/PathTracingMaterialHitShader.usf
+++ b/Engine/Shaders/Private/PathTracing/PathTracingMaterialHitShader.usf
@@ -166,7 +166,7 @@ RAY_TRACING_ENTRY_CLOSEST_HIT(PathTracingMaterialCHS,
 			{
 				float PathRoughness = PackedPayload.GetPathRoughness();
 				float Roughness = GetMaterialRoughness(PixelMaterialInputs);
-				if (PathRoughness <= Roughness)
+				if (Ior > 0.0 && PathRoughness <= Roughness)
 				{
 					// not using fast caustic approximation - treat as opaque
 					PackedPayload.SetRayThroughput(0.0);
@@ -180,7 +180,10 @@ RAY_TRACING_ENTRY_CLOSEST_HIT(PathTracingMaterialCHS,
 				Transparency *= ComputeThinSlabWeights(Transmission, VoN, Ior, F0).Transmitted;
 
 				// fake caustic approximation (see comments below)
-				Transparency *= (1 - Roughness * Roughness) * saturate(PathRoughness - Roughness);
+				if (Ior > 0.0)
+				{
+					Transparency *= (1 - Roughness * Roughness) * saturate(PathRoughness - Roughness);
+				}
 			}
 #endif
 		}
@@ -230,7 +233,7 @@ RAY_TRACING_ENTRY_CLOSEST_HIT(PathTracingMaterialCHS,
 				// TODO2: SingleLayerWater needs something similar, but it generally is setup with opaque or masked blend mode and water geometry does not cast shadows ...
 				float Distance = RayTCurrent();
 				float3 TransmittanceColor = GetAbsorptionMediumMaterialOutput0(MaterialParameters);
-				float3 LocalSigmaT = -log(max(TransmittanceColor, 1e-8f));
+				float3 LocalSigmaT = PathTracingGlassTransmittanceToExtinction(TransmittanceColor);
 				
 				// Front Face: remove contribution from the ray origin to the current hit (assuming the backside will be hit)
 				// Back Face:  add contribution from ray origin to current hit (backside) - excess portion not covered by the ray will be removed by the front face hit
@@ -266,22 +269,23 @@ RAY_TRACING_ENTRY_CLOSEST_HIT(PathTracingMaterialCHS,
 	Payload.ShadingModelID = GetMaterialShadingModel(PixelMaterialInputs);
 
 #if MATERIALBLENDING_ALPHACOMPOSITE
-	Payload.Opacity = GetMaterialOpacity(PixelMaterialInputs);
-	Payload.BlendingMode = RAY_TRACING_BLEND_MODE_ALPHA_COMPOSITE;
+	Payload.BSDFOpacity = 1.0;
+	Payload.TransparencyColor = 1.0 - GetMaterialOpacity(PixelMaterialInputs);
 #elif MATERIALBLENDING_ALPHAHOLDOUT
-	Payload.Opacity = GetMaterialOpacity(PixelMaterialInputs);
-	Payload.BlendingMode = RAY_TRACING_BLEND_MODE_ALPHA_HOLDOUT;
+	Payload.BSDFOpacity = GetMaterialOpacity(PixelMaterialInputs);
+	Payload.TransparencyColor = 1.0 - GetMaterialOpacity(PixelMaterialInputs);
+	Payload.SetHoldout();
 	HLSL_STATIC_ASSERT(MATERIAL_SHADINGMODEL_UNLIT == 1, "Alpha holdout blend mode requires unlit shading model");
 	Payload.ShadingModelID = SHADINGMODELID_UNLIT;
 #elif MATERIALBLENDING_TRANSLUCENT
-	Payload.Opacity = GetMaterialOpacity(PixelMaterialInputs);
-	Payload.BlendingMode = RAY_TRACING_BLEND_MODE_TRANSLUCENT;
+	Payload.BSDFOpacity = GetMaterialOpacity(PixelMaterialInputs);
+	Payload.TransparencyColor = 1.0 - Payload.BSDFOpacity;
 #elif MATERIALBLENDING_ADDITIVE
-	Payload.Opacity = GetMaterialOpacity(PixelMaterialInputs);
-	Payload.BlendingMode = RAY_TRACING_BLEND_MODE_ADDITIVE;
+	Payload.BSDFOpacity = GetMaterialOpacity(PixelMaterialInputs);
+	Payload.TransparencyColor = 1.0;
 #elif MATERIALBLENDING_MODULATE
-	Payload.Opacity = GetMaterialOpacity(PixelMaterialInputs);
-	Payload.BlendingMode = RAY_TRACING_BLEND_MODE_MODULATE;
+	Payload.BSDFOpacity = 0.0;
+	Payload.TransparencyColor = GetMaterialEmissive(PixelMaterialInputs);
 	HLSL_STATIC_ASSERT(MATERIAL_SHADINGMODEL_UNLIT == 1, "Modulate blend mode requires unlit shading model");
 	Payload.ShadingModelID = SHADINGMODELID_UNLIT;
 #elif MATERIALBLENDING_MASKED && MATERIAL_DITHER_OPACITY_MASK
@@ -293,11 +297,11 @@ RAY_TRACING_ENTRY_CLOSEST_HIT(PathTracingMaterialCHS,
 	// code performs blending.
 	// Since the goal of dithered opacity is to emulate ordinary transparency, just use the mask input as opacity directly and
 	// ignore the configured clip value.
-	Payload.Opacity = saturate(GetMaterialMaskInputRaw(PixelMaterialInputs));
-	Payload.BlendingMode = RAY_TRACING_BLEND_MODE_TRANSLUCENT;
+	Payload.BSDFOpacity = saturate(GetMaterialMaskInputRaw(PixelMaterialInputs));
+	Payload.TransparencyColor = 1.0 - Payload.BSDFOpacity;
 #elif MATERIALBLENDING_SOLID || MATERIALBLENDING_MASKED
-	Payload.Opacity = 1.0;
-	Payload.BlendingMode = RAY_TRACING_BLEND_MODE_OPAQUE;
+	Payload.BSDFOpacity = 1.0;
+	Payload.TransparencyColor = 0.0;
 #else
 #error Unknown material blending mode
 #endif
@@ -337,11 +341,6 @@ RAY_TRACING_ENTRY_CLOSEST_HIT(PathTracingMaterialCHS,
 	Ior = GetRefractionIor(PixelMaterialInputs);
 #endif
 
-#if !SIMPLIFIED_MATERIAL_SHADER
-
-	/**
-	 * Set material attributes for full materials
-	 **/
 	Payload.TranslatedWorldPos = LWCToFloat(LWCAdd(MaterialParameters.AbsoluteWorldPosition, ResolvedView.PreViewTranslation));
 	float GeoNormalSign = MaterialParameters.TwoSidedSign;
 #if !VF_SUPPORTS_RAYTRACING_PREPARE_MATERIAL_PIXEL_PARAMETERS
@@ -353,12 +352,21 @@ RAY_TRACING_ENTRY_CLOSEST_HIT(PathTracingMaterialCHS,
 	Payload.WorldNormal = MaterialParameters.WorldNormal;
 	Payload.WorldSmoothNormal = MaterialParameters.TwoSidedSign * TransformTangentNormalToWorld(MaterialParameters.TangentToWorld, float3(0, 0, 1));
 
+
+#if !SIMPLIFIED_MATERIAL_SHADER
+	/**
+	 * Set material attributes for full materials
+	 **/
+
 	// fix dark edges in reflections from shading normals
 	Payload.WorldNormal       = AdjustShadingNormal(Payload.WorldNormal      , Payload.WorldGeoNormal, WorldRayDirection());
 	Payload.WorldSmoothNormal = AdjustShadingNormal(Payload.WorldSmoothNormal, Payload.WorldGeoNormal, WorldRayDirection());
 
 	Payload.Radiance = GetMaterialEmissive(PixelMaterialInputs);
-#if TWO_SIDED_MATERIAL == 0
+	Payload.Radiance *= Payload.BSDFOpacity; // pre-multiply
+
+
+#if MATERIAL_TWOSIDED == 0
 	if (MaterialParameters.TwoSidedSign < 0)
 	{
 		// when viewing the surface from "inside", don't include emission
@@ -451,8 +459,6 @@ RAY_TRACING_ENTRY_CLOSEST_HIT(PathTracingMaterialCHS,
 #else
 		Payload.SetSubsurfacePhaseFunction(0.0);
 #endif
-		// Reset opacity to avoid special cases elsewhere
-		Payload.Opacity = 1;
 	}
 	else
 #endif
@@ -466,7 +472,8 @@ RAY_TRACING_ENTRY_CLOSEST_HIT(PathTracingMaterialCHS,
 		float Lobe1Roughness = 0;
 		float LobeMix = 0;
 
-		GetSubsurfaceProfileDualSpecular(ProfileId, Payload.Roughness, Payload.Opacity, Lobe0Roughness, Lobe1Roughness, LobeMix);
+		// NOTE: we pass 1.0 for opacity because we don't really want to fade-out the dual roughness as opacity is lowered
+		GetSubsurfaceProfileDualSpecular(ProfileId, Payload.Roughness, 1.0, Lobe0Roughness, Lobe1Roughness, LobeMix);
 
 		Lobe0Roughness = AdjustMaterialRoughness(Lobe0Roughness, PackedPayload.GetPathRoughness());
 		Lobe1Roughness = AdjustMaterialRoughness(Lobe1Roughness, PackedPayload.GetPathRoughness());
@@ -494,8 +501,6 @@ RAY_TRACING_ENTRY_CLOSEST_HIT(PathTracingMaterialCHS,
 #else
 		Payload.SetSubsurfacePhaseFunction(0.0);
 #endif
-
-		Payload.Opacity = 1.0; // Reset opacity to avoid special cases elsewhere
 	}
 	else
 #endif
@@ -513,10 +518,6 @@ RAY_TRACING_ENTRY_CLOSEST_HIT(PathTracingMaterialCHS,
 
 		const float3 WaterExtinction = ScatteringCoeff + AbsorptionCoeff;
 		const float3 WaterAlbedo = ScatteringCoeff / WaterExtinction;
-
-		Payload.BaseColor = WaterAlbedo;
-		Payload.SetWaterExtinction(WaterExtinction);
-		Payload.SetWaterPhaseG(PhaseG);
 #else
 		// For single layer water we always want to treat it as front facing when evaluating parameters 
 		// as hitting back face is treated as being underwater and we get completely different parameters.
@@ -531,12 +532,13 @@ RAY_TRACING_ENTRY_CLOSEST_HIT(PathTracingMaterialCHS,
 		const float3 WaterAlbedo = ScatteringCoeff / WaterExtinction;
 #endif
 
-		Payload.BaseColor = WaterAlbedo;
-		Payload.SetWaterExtinction(WaterExtinction);
-		Payload.SetWaterPhaseG(PhaseG);
-
+		Payload.ShadingModelID = SHADINGMODELID_DEFAULT_LIT; // replace shading model to avoid extra cases in raygen
+		Payload.BaseColor = WaterAlbedo; // TODO: unused at the moment
+		Payload.SetTransmittanceColor(PathTracingGlassExtinctionToTransmittance(WaterExtinction));
 		Payload.Ior = DielectricF0ToIor(DielectricSpecularToF0(Specular));
-		Payload.BlendingMode = RAY_TRACING_BLEND_MODE_TRANSLUCENT;
+		// we are using actual refraction, so don't apply any transparency
+		Payload.TransparencyColor = 0.0;
+		Payload.BSDFOpacity = 0.0; // only shade with the glass portion
 	}
 	else
 #endif
@@ -544,8 +546,26 @@ RAY_TRACING_ENTRY_CLOSEST_HIT(PathTracingMaterialCHS,
 #if MATERIAL_SHADINGMODEL_THIN_TRANSLUCENT
 	if (Payload.ShadingModelID == SHADINGMODELID_THIN_TRANSLUCENT)
 	{
-		Payload.SetTransmittanceColor(GetThinTranslucentMaterialOutput0(MaterialParameters));
+		const float3 TransmittanceColor = GetThinTranslucentMaterialOutput0(MaterialParameters);
+		Payload.SetTransmittanceColor(TransmittanceColor);
 		Payload.Ior = Ior;
+		if (Ior > 0.0)
+		{
+			// In this mode we get rough refraction, so don't apply any transparency
+			Payload.TransparencyColor = 0.0;
+		}
+		else
+		{
+			// In this mode we just get straight transparency
+			const float3 V = WorldRayDirection();
+			const float3 N = normalize(Payload.WorldNormal);
+			const float VoN = abs(dot(V, N));
+			const float3 BaseColor = Payload.BaseColor;
+			const float  Metallic = Payload.Metallic;
+			const float  Specular = Payload.Specular;
+			const float  F0 = F0RGBToF0(ComputeF0(Specular, BaseColor, Metallic));
+			Payload.TransparencyColor = (1 - Payload.BSDFOpacity) * ComputeThinSlabWeights(TransmittanceColor, VoN, 0.0, F0).Transmitted;
+		}
 	}
 	else
 #endif
@@ -618,7 +638,6 @@ RAY_TRACING_ENTRY_CLOSEST_HIT(PathTracingMaterialCHS,
 			Payload.SetSubsurfaceRadius(0.0);
 		}
 #endif
-		Payload.Opacity = 1.0; // revert to default now that decode is done
 
 #if HAVE_GetSubsurfaceMediumMaterialOutput1
 		float ScatteringDistribution = clamp(GetSubsurfaceMediumMaterialOutput1(MaterialParameters), -0.99f, 0.99f);
@@ -642,6 +661,11 @@ RAY_TRACING_ENTRY_CLOSEST_HIT(PathTracingMaterialCHS,
 		Payload.SetTransmittanceColor(1.0);
 #endif
 		Payload.Ior = Ior;
+		if (Ior > 0.0)
+		{
+			// we are using actual refraction, so disable transparency
+			Payload.TransparencyColor = 0.0;
+		}
 	}
 	else
 #endif
@@ -649,32 +673,38 @@ RAY_TRACING_ENTRY_CLOSEST_HIT(PathTracingMaterialCHS,
 		// terminal case for the conditionals above
 	}
 
+#if MATERIALBLENDING_ALPHAHOLDOUT
+	// the material is already a holdout, no need to do anything else
+#else
 	if (PackedPayload.IsCameraRay())
 	{
 		if ((PrimitiveFlags & PRIMITIVE_SCENE_DATA_FLAG_HOLDOUT) != 0)
 		{
-			if (Payload.BlendingMode != RAY_TRACING_BLEND_MODE_TRANSLUCENT || Payload.Ior != 0.0)
+			if (Payload.Ior != 0.0)
 			{
-				// If we are either not translucent, or if the surface is marked as refractive
-				// then treat opacity as 1.0
-				Payload.Opacity = 1.0;
+				// If the surface is marked as refractive then treat it as opaque for holdout
+				Payload.BSDFOpacity = 1.0;
 			}
 			else
 			{
 				// keep the opacity from the material, it might indicate partial transparency
 				// so we should get a partial holdout
 			}
-			Payload.BlendingMode = RAY_TRACING_BLEND_MODE_ALPHA_HOLDOUT;
+
+			Payload.SetHoldout();
 			Payload.ShadingModelID = SHADINGMODELID_UNLIT;
 			Payload.Radiance = 0;
 		}
 	}
+#endif
 
 	// Only opaque surfaces can recieve decals. This is both for compatibility with raster and improving performance when we have stacks of transparent surfaces
-	if ((Payload.BlendingMode == RAY_TRACING_BLEND_MODE_OPAQUE) && ((PrimitiveFlags & PRIMITIVE_SCENE_DATA_FLAG_DECAL_RECEIVER) != 0))
+#if MATERIALBLENDING_SOLID || MATERIALBLENDING_MASKED
+	if ((PrimitiveFlags & PRIMITIVE_SCENE_DATA_FLAG_DECAL_RECEIVER) != 0)
 	{
 		Payload.SetDecalReceiver();
 	}
+#endif
 
 	// adjust after everything else (because SSS profile case needs to decode the dual spec info from the unmodified roughness)
 	Payload.Roughness = AdjustMaterialRoughness(Payload.Roughness, PackedPayload.GetPathRoughness());
@@ -684,6 +714,16 @@ RAY_TRACING_ENTRY_CLOSEST_HIT(PathTracingMaterialCHS,
 	/**
 	 * Set material attributes for simplified materials
 	 **/
+
+#if MATERIALBLENDING_TRANSLUCENT && !MATERIAL_SHADINGMODEL_THIN_TRANSLUCENT
+	 // Force opacity to 0 so it is ignored for path hits and only evaluated in AHS as a modulation to throughput
+	Payload.BSDFOpacity = 0;
+	Payload.TransparencyColor = 1;
+#else
+	Payload.BSDFOpacity = GetMaterialOpacity(PixelMaterialInputs);
+	Payload.TransparencyColor = 1 - Payload.BSDFOpacity;
+#endif
+
 	 // Anything unsupported will be forced to default lit
 #if MATERIAL_SHADINGMODEL_TWOSIDED_FOLIAGE
 	if (Payload.ShadingModelID == SHADINGMODELID_TWOSIDED_FOLIAGE)
@@ -695,7 +735,12 @@ RAY_TRACING_ENTRY_CLOSEST_HIT(PathTracingMaterialCHS,
 #if MATERIAL_SHADINGMODEL_THIN_TRANSLUCENT
 	if (Payload.ShadingModelID == SHADINGMODELID_THIN_TRANSLUCENT)
 	{
-		Payload.SetTransmittanceColor(GetThinTranslucentMaterialOutput0(MaterialParameters));
+		const float3 Transmission = GetThinTranslucentMaterialOutput0(MaterialParameters);
+		const float3 V = WorldRayDirection();
+		const float3 N = normalize(MaterialParameters.WorldNormal);
+		const float VoN = abs(dot(V, N));
+		// simplified logic with no bending at the interface and no fresnel
+		Payload.TransparencyColor *= pow(Transmission, 1.0 / VoN);
 	}
 	else
 #endif
@@ -713,31 +758,21 @@ RAY_TRACING_ENTRY_CLOSEST_HIT(PathTracingMaterialCHS,
 		Payload.ShadingModelID = SHADINGMODELID_DEFAULT_LIT;
 #endif
 	}
-	Payload.TranslatedWorldPos = LWCToFloat(LWCAdd(MaterialParameters.AbsoluteWorldPosition, ResolvedView.PreViewTranslation));
-	Payload.WorldNormal = normalize(MaterialParameters.WorldNormal);
-	float GeoNormalSign = MaterialParameters.TwoSidedSign * GetPrimitive_DeterminantSign(MaterialParameters.PrimitiveId);
-	Payload.WorldGeoNormal = GeoNormalSign * GeoNormal;
 	Payload.Radiance = GetMaterialEmissive(PixelMaterialInputs);
-#if TWO_SIDED_MATERIAL == 0
+	Payload.Radiance *= Payload.BSDFOpacity; // premultiply
+#if MATERIAL_TWOSIDED
+	Payload.SetMaterialTwoSided();
+#else
 	if (MaterialParameters.TwoSidedSign < 0)
 	{
 		// when viewing the surface from "inside", don't include emission
 		Payload.Radiance = 0;
 	}
-#endif
-#if MATERIAL_TWOSIDED
-	Payload.SetMaterialTwoSided();
 #endif
 	Payload.BaseColor = BaseColor;
 	Payload.Specular = 0;
 	Payload.Roughness = 1;
 	Payload.Metallic = 0;
-#if MATERIALBLENDING_TRANSLUCENT && !MATERIAL_SHADINGMODEL_THIN_TRANSLUCENT
-	// Force opacity to 0 so it is ignored for path hits and only evaluated in AHS as a modulation to throughput
-	Payload.Opacity = 0;
-#else
-	Payload.Opacity = GetMaterialOpacity(PixelMaterialInputs);
-#endif
 	Payload.Ior = 0;
 
 #endif
diff --git a/Engine/Shaders/Private/PathTracing/Volume/PathTracingAtmosphere.ush b/Engine/Shaders/Private/PathTracing/Volume/PathTracingAtmosphere.ush
--- a/Engine/Shaders/Private/PathTracing/Volume/PathTracingAtmosphere.ush
+++ b/Engine/Shaders/Private/PathTracing/Volume/PathTracingAtmosphere.ush
@@ -148,7 +148,7 @@ FPackedPathTracingPayload AtmosphereGetBlockerHit(float3 Origin, float3 Directio
 
 	FPathTracingPayload Payload = (FPathTracingPayload)0;
 	Payload.BaseColor = Atmosphere.GroundAlbedo.rgb;
-	Payload.Opacity = 1.0;
+	Payload.BSDFOpacity = 1.0;
 	Payload.WorldNormal = Payload.WorldGeoNormal = Payload.WorldSmoothNormal = WorldN;
 	Payload.ShadingModelID = SHADINGMODELID_NUM; // defaults to diffuse
 	Payload.HitT = HitT;
diff --git a/Engine/Shaders/Private/PathTracing/Volume/PathTracingVolumeCommon.ush b/Engine/Shaders/Private/PathTracing/Volume/PathTracingVolumeCommon.ush
--- a/Engine/Shaders/Private/PathTracing/Volume/PathTracingVolumeCommon.ush
+++ b/Engine/Shaders/Private/PathTracing/Volume/PathTracingVolumeCommon.ush
@@ -201,8 +201,7 @@ FPathTracingPayload CreateMediumHitPayload(float HitT, float3 TranslatedWorldPos
 	Result.HitT = HitT;
 	Result.TranslatedWorldPos = TranslatedWorldPos;
 	Result.ShadingModelID = SHADINGMODELID_MEDIUM;
-	Result.BlendingMode = RAY_TRACING_BLEND_MODE_OPAQUE;
-	Result.Opacity = 1.0;
+	Result.BSDFOpacity = 1.0;
 	Result.PrimitiveLightingChannelMask = 7;
 	Result.SetFrontFace();
 	const float3 SigmaS = ShadedResult.SigmaSRayleigh + ShadedResult.SigmaSHG;
diff --git a/Engine/Source/Runtime/Renderer/Private/PathTracing.cpp b/Engine/Source/Runtime/Renderer/Private/PathTracing.cpp
--- a/Engine/Source/Runtime/Renderer/Private/PathTracing.cpp
+++ b/Engine/Source/Runtime/Renderer/Private/PathTracing.cpp
@@ -86,14 +86,6 @@ TAutoConsoleVariable<float> CVarPathTracingFilterWidth(
 	ECVF_RenderThreadSafe
 );
 
-TAutoConsoleVariable<float> CVarPathTracingAbsorptionScale(
-	TEXT("r.PathTracing.AbsorptionScale"),
-	0.01,
-	TEXT("Sets the inverse distance at which BaseColor is reached for transmittance in refractive glass (default = 1/100 units)\n")
-	TEXT("Setting this value to 0 will disable absorption handling for refractive glass\n"),
-	ECVF_RenderThreadSafe
-);
-
 TAutoConsoleVariable<int32> CVarPathTracingMISMode(
 	TEXT("r.PathTracing.MISMode"),
 	2,
@@ -312,6 +304,14 @@ TAutoConsoleVariable<int32> CVarPathTracingLightFunctionColor(
 	ECVF_RenderThreadSafe
 );
 
+TAutoConsoleVariable<int32> CVarPathTracingCameraMediumTracking(
+	TEXT("r.PathTracing.CameraMediumTracking"),
+	1,
+	TEXT("Enables automatic camera medium tracking to detect when a camera starts inside water or solid glass automatically (default = 1)\n")
+	TEXT("0: off\n")
+	TEXT("1: on (default)\n"),
+	ECVF_RenderThreadSafe
+);
 
 BEGIN_SHADER_PARAMETER_STRUCT(FPathTracingData, )
 	SHADER_PARAMETER(float, BlendFactor)
@@ -336,7 +336,6 @@ BEGIN_SHADER_PARAMETER_STRUCT(FPathTracingData, )
 	SHADER_PARAMETER(float, MaxPathIntensity)
 	SHADER_PARAMETER(float, MaxNormalBias)
 	SHADER_PARAMETER(float, FilterWidth)
-	SHADER_PARAMETER(float, AbsorptionScale)
 	SHADER_PARAMETER(float, DecalRoughnessCutoff)
 	SHADER_PARAMETER(float, MeshDecalRoughnessCutoff)
 	SHADER_PARAMETER(float, MeshDecalBias)
@@ -356,6 +355,7 @@ struct FPathTracingConfig
 	bool VisibleLights;
 	bool UseMISCompensation;
 	bool LockedSamplingPattern;
+	bool UseCameraMediumTracking;
 	bool UseMultiGPU; // NOTE: Requires invalidation because the buffer layout changes
 	int DenoiserMode; // NOTE: does not require path tracing invalidation
 
@@ -375,7 +375,6 @@ struct FPathTracingConfig
 			PathTracingData.VisualizeDecalGrid != Other.PathTracingData.VisualizeDecalGrid ||
 			PathTracingData.MaxPathIntensity != Other.PathTracingData.MaxPathIntensity ||
 			PathTracingData.FilterWidth != Other.PathTracingData.FilterWidth ||
-			PathTracingData.AbsorptionScale != Other.PathTracingData.AbsorptionScale ||
 			PathTracingData.EnableAtmosphere != Other.PathTracingData.EnableAtmosphere ||
 			PathTracingData.EnableFog != Other.PathTracingData.EnableFog ||
 			PathTracingData.ApplyDiffuseSpecularOverrides != Other.PathTracingData.ApplyDiffuseSpecularOverrides ||
@@ -392,6 +391,7 @@ struct FPathTracingConfig
 			VisibleLights != Other.VisibleLights ||
 			UseMISCompensation != Other.UseMISCompensation ||
 			LockedSamplingPattern != Other.LockedSamplingPattern ||
+			UseCameraMediumTracking != Other.UseCameraMediumTracking ||
 			UseMultiGPU != Other.UseMultiGPU;
 	}
 
@@ -460,6 +460,9 @@ struct FPathTracingState {
 	TRefCountPtr<IPooledRenderTarget> AtmosphereOpticalDepthLUT;
 	FAtmosphereConfig LastAtmosphereConfig;
 
+	// Buffer containing the starting medium extinction
+	TRefCountPtr<FRDGPooledBuffer>    StartingExtinctionCoefficient;
+
 	// Current sample index to be rendered by the path tracer - this gets incremented each time the path tracer accumulates a frame of samples
 	uint32 SampleIndex = 0;
 
@@ -510,7 +513,6 @@ static void PreparePathTracingData(const FScene* Scene, const FViewInfo& View, F
 		FilterWidth = PPV.PathTracingFilterWidth;
 	}
 	PathTracingData.FilterWidth = FilterWidth;
-	PathTracingData.AbsorptionScale = CVarPathTracingAbsorptionScale.GetValueOnRenderThread();
 	PathTracingData.CameraFocusDistance = 0;
 	PathTracingData.CameraLensRadius = FVector2f::ZeroVector;
 	if (ShowFlags.DepthOfField &&
@@ -1194,7 +1196,6 @@ public:
 		OutEnvironment.SetDefine(TEXT("USE_RAYTRACED_TEXTURE_RAYCONE_LOD"), 0);
 		OutEnvironment.SetDefine(TEXT("SCENE_TEXTURES_DISABLED"), 1);
 		OutEnvironment.SetDefine(TEXT("SIMPLIFIED_MATERIAL_SHADER"), UseSimplifiedShader);
-		OutEnvironment.SetDefine(TEXT("TWO_SIDED_MATERIAL"), Parameters.MaterialParameters.bIsTwoSided ? 1 : 0);
 		FMeshMaterialShader::ModifyCompilationEnvironment(Parameters, OutEnvironment);
 	}
 
@@ -2063,6 +2064,7 @@ void FDeferredShadingSceneRenderer::RenderPathTracing(
 	uint32 MaxSPP = SamplesPerPixelCVar > -1 ? SamplesPerPixelCVar : View.FinalPostProcessSettings.PathTracingSamplesPerPixel;
 	MaxSPP = FMath::Max(MaxSPP, 1u);
 	Config.LockedSamplingPattern = CVarPathTracingFrameIndependentTemporalSeed.GetValueOnRenderThread() == 0;
+	Config.UseCameraMediumTracking = CVarPathTracingCameraMediumTracking.GetValueOnRenderThread() != 0;
 
 	// compute an integer code of what show flags and booleans related to lights are currently enabled so we can detect changes
 	Config.LightShowFlags = 0;
@@ -2126,18 +2128,29 @@ void FDeferredShadingSceneRenderer::RenderPathTracing(
 		
 	}
 
-	// prepare extinction coefficient for camera rays
-	FRDGBufferRef StartingExtinctionCoefficient = GraphBuilder.CreateBuffer(FRDGBufferDesc::CreateBufferDesc(sizeof(float), 3), TEXT("PathTracer.StartingExtinctionCoefficient"));
-	FRDGBufferUAVRef StartingExtinctionCoefficientUAV = GraphBuilder.CreateUAV(StartingExtinctionCoefficient, PF_R32_FLOAT);
 
-	if (View.IsUnderwater())
+	// If this is the first sample, recompute the initial medium
+	// In this case of an offline render, do this every frame so that motion blur through a boundary is properly accounted for
+	FRDGBufferRef StartingExtinctionCoefficient = nullptr;
+	if (!Config.UseCameraMediumTracking)
+	{
+		PathTracingState->StartingExtinctionCoefficient.SafeRelease();
+		// camera medium tracking is not enabled - just make a temp buffer and set it to 0
+		StartingExtinctionCoefficient = GraphBuilder.CreateBuffer(FRDGBufferDesc::CreateBufferDesc(sizeof(float), 3), TEXT("PathTracer.StartingExtinctionCoefficient"));
+		AddClearUAVPass(GraphBuilder, GraphBuilder.CreateUAV(StartingExtinctionCoefficient, PF_R32_FLOAT), 0);
+		
+	}
+	else if (!PathTracingState->StartingExtinctionCoefficient.IsValid() || PathTracingState->SampleIndex == 0 || View.bIsOfflineRender)
 	{
 		auto RayGenShader = GetGlobalShaderMap(View.FeatureLevel)->GetShader<FPathTracingInitExtinctionCoefficientRG>();
 
+		// prepare extinction coefficient for camera rays
+		StartingExtinctionCoefficient = GraphBuilder.CreateBuffer(FRDGBufferDesc::CreateBufferDesc(sizeof(float), 3), TEXT("PathTracer.StartingExtinctionCoefficient"));
+
 		FPathTracingInitExtinctionCoefficientRG::FParameters* PassParameters = GraphBuilder.AllocParameters<FPathTracingInitExtinctionCoefficientRG::FParameters>();
 		PassParameters->TLAS = Scene->RayTracingScene.GetLayerView(ERayTracingSceneLayer::Base);
 		PassParameters->ViewUniformBuffer = View.ViewUniformBuffer;
-		PassParameters->RWStartingExtinctionCoefficient = StartingExtinctionCoefficientUAV;
+		PassParameters->RWStartingExtinctionCoefficient = GraphBuilder.CreateUAV(StartingExtinctionCoefficient, PF_R32_FLOAT);
 
 		GraphBuilder.AddPass(
 			RDG_EVENT_NAME("Path Tracer Init Sigma"),
@@ -2157,10 +2170,12 @@ void FDeferredShadingSceneRenderer::RenderPathTracing(
 					1, 1
 				);
 			});
+		GraphBuilder.QueueBufferExtraction(StartingExtinctionCoefficient, &PathTracingState->StartingExtinctionCoefficient);
 	}
 	else
 	{
-		AddClearUAVPass(GraphBuilder, StartingExtinctionCoefficientUAV, 0);
+		check(PathTracingState->StartingExtinctionCoefficient.IsValid());
+		StartingExtinctionCoefficient = GraphBuilder.RegisterExternalBuffer(PathTracingState->StartingExtinctionCoefficient, TEXT("PathTracer.StartingExtinctionCoefficient"));
 	}
 
 	// prepare atmosphere optical depth lookup texture (if needed)
