Path Tracer: Improve volume light sampling for local lights

This appears to improve renders involving multiple lights by tuning the relative probability assigned to each light.

diff --git a/Engine/Shaders/Private/PathTracing/Light/PathTracingLightCommon.ush b/Engine/Shaders/Private/PathTracing/Light/PathTracingLightCommon.ush
--- a/Engine/Shaders/Private/PathTracing/Light/PathTracingLightCommon.ush
+++ b/Engine/Shaders/Private/PathTracing/Light/PathTracingLightCommon.ush
@@ -87,21 +87,19 @@ struct FVolumeLightSampleSetup
 
 FVolumeLightSampleSetup CreateEquiAngularSampler(float LightImportance, float3 Center, float3 Ro, float3 Rd, float TMin, float TMax)
 {
-#if 0
-	float T0 = dot(Center - Ro, Rd);
-	float3 P0 = Ro + T0 * Rd;
-	Center = Center - normalize(Center - P0);
-#endif
 	const float3 V = Center - Ro;
 	FVolumeLightSampleSetup Result = (FVolumeLightSampleSetup)0;
 	Result.VolumeTMin = TMin;
 	Result.VolumeTMax = TMax;
 	Result.T0 = dot(V, Rd);
 	Result.Height = length(cross(V, Rd));
-	Result.ThetaMin = Result.Height > 0 ? atan((TMin - Result.T0) / Result.Height) : 0.0;
-	Result.ThetaMax = Result.Height > 0 ? atan((TMax - Result.T0) / Result.Height) : 0.0;
+	// TODO: handle Height == 0 case better, the pdf becomes degenerate in this case
+	// but unless the light is parented to the camera exactly, it is hard to see in practice
+	const float RcpHeight = Result.Height > 0 ? rcp(Result.Height) : 0.0;
+	Result.ThetaMin = atan((TMin - Result.T0) * RcpHeight);
+	Result.ThetaMax = atan((TMax - Result.T0) * RcpHeight);
 	Result.K = (Result.ThetaMax - Result.ThetaMin) * Result.Height;
-	Result.LightImportance = LightImportance * (TMax - TMin);
+	Result.LightImportance = LightImportance * RcpHeight;
 	return Result;
 }
 
