Avoid non-necessary initialization causing massive shader stall.

Reduce base pass cost of complex material by ~50%

diff --git a/Engine/Shaders/Private/Strata/Strata.ush b/Engine/Shaders/Private/Strata/Strata.ush
--- a/Engine/Shaders/Private/Strata/Strata.ush
+++ b/Engine/Shaders/Private/Strata/Strata.ush
@@ -1118,8 +1118,10 @@ struct FStrataTree
 
 FStrataTree GetInitialisedStrataTree()
 {
-	FStrataTree StrataTree = (FStrataTree)0;
-	return StrataTree;
+	FStrataTree Out;
+	Out.BSDFCount = 0;
+	Out.OperatorCount = 0;
+	return Out;
 }
 
 void StrataPackOutStrataTree(
@@ -1965,7 +1967,7 @@ FStrataPixelHeader InitialiseStrataPixelHeader()
 	Out.BSDFCount = 0;
 #if STRATA_INLINE_SHADING
 	Out.SharedLocalBases = (FSharedLocalBases)0;
-	Out.StrataTree = (FStrataTree)0;
+	Out.StrataTree = GetInitialisedStrataTree();
 	Out.IrradianceAO = InitIrradianceAndOcclusion();
 #endif
 #if STRATA_DEFERRED_SHADING
