Fix a trivial bug of mobile deferred with extended gbuffer enabled on IOS.

diff --git a/Engine/Shaders/Private/DeferredShadingCommon.ush b/Engine/Shaders/Private/DeferredShadingCommon.ush
--- a/Engine/Shaders/Private/DeferredShadingCommon.ush
+++ b/Engine/Shaders/Private/DeferredShadingCommon.ush
@@ -771,8 +771,10 @@ void MobileFetchGBuffer(in float2 UV, out half4 GBufferA, out half4 GBufferB, ou
 	GBufferC = SubpassFetchRGBA_3(); 
 #if MOBILE_EXTENDED_GBUFFER
 	GBufferD = SubpassFetchRGBA_4();
-#endif
+	SceneDepth = ConvertFromDeviceZ(SubpassFetchR_5());
+#else
 	SceneDepth = ConvertFromDeviceZ(SubpassFetchR_4());
+#endif
 #elif USE_GLES_FBF_DEFERRED
 	GBufferA = GLSubpassFetch1(); 
 	GBufferB = GLSubpassFetch2(); 
