Add normal curvature-based roughness support for Strata's Slab.

diff --git a/Engine/Shaders/Private/BasePassCommon.ush b/Engine/Shaders/Private/BasePassCommon.ush
--- a/Engine/Shaders/Private/BasePassCommon.ush
+++ b/Engine/Shaders/Private/BasePassCommon.ush
@@ -149,16 +149,6 @@ float3 GetDirectionalLightingVectorFromTranslucentLightingVolume(float3 InnerVol
 
 #include "SphericalGaussian.ush"
 
-float NormalCurvatureToRoughness(float3 WorldNormal)
-{
-	float3 dNdx = ddx(WorldNormal);
-	float3 dNdy = ddy(WorldNormal);
-	float x = dot(dNdx, dNdx);
-	float y = dot(dNdy, dNdy);
-	float CurvatureApprox = pow(max(x, y), View.NormalCurvatureToRoughnessScaleBias.z);
-	return saturate(CurvatureApprox * View.NormalCurvatureToRoughnessScaleBias.x + View.NormalCurvatureToRoughnessScaleBias.y);
-}
-
 struct FShadingOcclusion
 {
 	half  DiffOcclusion;
diff --git a/Engine/Shaders/Private/BasePassPixelShader.usf b/Engine/Shaders/Private/BasePassPixelShader.usf
--- a/Engine/Shaders/Private/BasePassPixelShader.usf
+++ b/Engine/Shaders/Private/BasePassPixelShader.usf
@@ -1072,13 +1072,7 @@ void FPixelShaderInOut_MainPS(
 	GBuffer.SpecularColor = ComputeF0(Specular, BaseColor, Metallic);
 
 #if MATERIAL_NORMAL_CURVATURE_TO_ROUGHNESS
-	// Curvature-to-roughness uses derivatives of the WorldVertexNormal, which is incompatible with centroid interpolation because
-	// the samples are not uniformly distributed. Therefore we use WorldVertexNormal_Center which is guaranteed to be center interpolated.
-#if USE_WORLDVERTEXNORMAL_CENTER_INTERPOLATION
-	float GeometricAARoughness = NormalCurvatureToRoughness(MaterialParameters.WorldVertexNormal_Center);
-#else
-	float GeometricAARoughness = NormalCurvatureToRoughness(MaterialParameters.TangentToWorld[2].xyz);
-#endif
+	const float GeometricAARoughness = GetRoughnessFromNormalCurvature(MaterialParameters);
 	GBuffer.Roughness = max(GBuffer.Roughness, GeometricAARoughness);
 
 #if MATERIAL_SHADINGMODEL_CLEAR_COAT
diff --git a/Engine/Shaders/Private/MaterialTemplate.ush b/Engine/Shaders/Private/MaterialTemplate.ush
--- a/Engine/Shaders/Private/MaterialTemplate.ush
+++ b/Engine/Shaders/Private/MaterialTemplate.ush
@@ -3057,6 +3057,31 @@ float3 CalculateAnisotropyTangent(in out FMaterialPixelParameters Parameters, FP
 	return Tangent;
 }
 
+float NormalCurvatureToRoughness(float3 WorldNormal)
+{
+	float3 dNdx = ddx(WorldNormal);
+	float3 dNdy = ddy(WorldNormal);
+	float x = dot(dNdx, dNdx);
+	float y = dot(dNdy, dNdy);
+	float CurvatureApprox = pow(max(x, y), View.NormalCurvatureToRoughnessScaleBias.z);
+	return saturate(CurvatureApprox * View.NormalCurvatureToRoughnessScaleBias.x + View.NormalCurvatureToRoughnessScaleBias.y);
+}
+
+float GetRoughnessFromNormalCurvature(FMaterialPixelParameters InMaterialParameters)
+{
+#if MATERIAL_NORMAL_CURVATURE_TO_ROUGHNESS
+	// Curvature-to-roughness uses derivatives of the WorldVertexNormal, which is incompatible with centroid interpolation because
+	// the samples are not uniformly distributed. Therefore we use WorldVertexNormal_Center which is guaranteed to be center interpolated.
+	#if USE_WORLDVERTEXNORMAL_CENTER_INTERPOLATION
+	return NormalCurvatureToRoughness(InMaterialParameters.WorldVertexNormal_Center);
+	#else
+	return NormalCurvatureToRoughness(InMaterialParameters.TangentToWorld[2].xyz);
+	#endif
+#else
+	return 0.f;
+#endif
+}
+
 // EvaluateMaterialAttributes
 %s
 
@@ -3071,7 +3096,7 @@ void CalcPixelMaterialInputs(in out FMaterialPixelParameters Parameters, in out
 %s
 
 #if TEMPLATE_USES_STRATA
-	Parameters.StrataPixelFootprint = StrataGetPixelFootprint(Parameters.WorldPosition_CamRelative);
+	Parameters.StrataPixelFootprint = StrataGetPixelFootprint(Parameters.WorldPosition_CamRelative, GetRoughnessFromNormalCurvature(Parameters));
 	Parameters.SharedLocalBases = StrataInitialiseSharedLocalBases();
 	Parameters.StrataTree = GetInitialisedStrataTree();
 #endif
@@ -3129,7 +3154,7 @@ void CalcPixelMaterialInputsAnalyticDerivatives(in out FMaterialPixelParameters
 %s
 
 #if TEMPLATE_USES_STRATA
-	Parameters.StrataPixelFootprint = StrataGetPixelFootprint(Parameters.WorldPosition_CamRelative);
+	Parameters.StrataPixelFootprint = StrataGetPixelFootprint(Parameters.WorldPosition_CamRelative, GetRoughnessFromNormalCurvature(Parameters));
 	Parameters.SharedLocalBases = StrataInitialiseSharedLocalBases();
 	Parameters.StrataTree = GetInitialisedStrataTree();
 #endif
diff --git a/Engine/Shaders/Private/Strata/Strata.ush b/Engine/Shaders/Private/Strata/Strata.ush
--- a/Engine/Shaders/Private/Strata/Strata.ush
+++ b/Engine/Shaders/Private/Strata/Strata.ush
@@ -1339,20 +1339,22 @@ FStrataData GetInitialisedStrataData()
 
 struct FStrataPixelFootprint
 {
-	float PixelRadiusInWorldSpace; // In cm
+	float PixelRadiusInWorldSpace;	// In cm
+	float NormalCurvatureRoughness; // Surface curvature converted into roughness for AA purpose
 };
 
-FStrataPixelFootprint StrataGetPixelFootprint(float3 dPdx, float3 dPdy)
+FStrataPixelFootprint StrataGetPixelFootprint(float3 dPdx, float3 dPdy, float InNormalCurvatureRoughness)
 {
 	const float dX = lengthFast(dPdx);
 	const float dY = lengthFast(dPdy);
 
 	FStrataPixelFootprint Out;
 	Out.PixelRadiusInWorldSpace = min(dX, dY) * 0.5f;
+	Out.NormalCurvatureRoughness = InNormalCurvatureRoughness;
 	return Out;
 }
 
-FStrataPixelFootprint StrataGetPixelFootprint(float3 InWorldPosition)
+FStrataPixelFootprint StrataGetPixelFootprint(float3 InWorldPosition, float InNormalCurvatureRoughness)
 {
 #if CLOUD_LAYER_PIXEL_SHADER==1
 	// Clouds materials feature dynamic loops and thus cannot support derivatives.
@@ -1360,7 +1362,7 @@ FStrataPixelFootprint StrataGetPixelFootprint(float3 InWorldPosition)
 #else
 	const float3 dPdx = ddx(InWorldPosition);
 	const float3 dPdy = ddy(InWorldPosition);
-	return StrataGetPixelFootprint(dPdx, dPdy);
+	return StrataGetPixelFootprint(dPdx, dPdy, InNormalCurvatureRoughness);
 #endif
 }
 
@@ -1379,6 +1381,12 @@ FStrataData GetStrataSlabBSDF(
 	bool bIsThinSurface,
 	uint SharedLocalBasisIndex, inout uint SharedLocalBasisTypes)
 {
+	// Curvature-based roughness for geometric AA
+	Roughness		= max(Roughness, InPixelFootprint.NormalCurvatureRoughness);
+	SecondRoughness = max(SecondRoughness, InPixelFootprint.NormalCurvatureRoughness);
+	FuzzRoughness	= max(FuzzRoughness, InPixelFootprint.NormalCurvatureRoughness);
+	
+
 	FStrataData StrataData = GetInitialisedStrataData();
 
 	SSSMFPScale = saturate(SSSMFPScale);
