Strata - fixed opacity clipping for translucent writting velocity

diff --git a/Engine/Shaders/Private/MaterialTemplate.ush b/Engine/Shaders/Private/MaterialTemplate.ush
--- a/Engine/Shaders/Private/MaterialTemplate.ush
+++ b/Engine/Shaders/Private/MaterialTemplate.ush
@@ -3320,11 +3320,11 @@ void GetMaterialCoverageAndClipping(FMaterialPixelParameters Parameters, FPixelM
 
 // The material blending mode alone is not enough information. We also need to know if the material is ThinTranslucent because in this case we should never clip.
 // Indeed, ThinTranslucent surfaces always have a coverage of 1 and opacity represent the coverage of the opaque lit material sitting on top of the translucent one.
-void GetMaterialClippingVelocity(FMaterialPixelParameters Parameters, FPixelMaterialInputs PixelMaterialInputs, bool bIsThinTranslucent)
+void GetMaterialClippingVelocity(FMaterialPixelParameters Parameters, FPixelMaterialInputs PixelMaterialInputs, bool bIsThinTranslucent, float MaterialOpacity)
 {
-#if MATERIALBLENDING_TRANSLUCENT || MATERIALBLENDING_ADDITIVE || MATERIALBLENDING_MODULATE
+#if MATERIALBLENDING_TRANSLUCENT || MATERIALBLENDING_ADDITIVE || MATERIALBLENDING_MODULATE || STRATA_TRANSLUCENT_MATERIAL
 	ClipLODTransition(Parameters);
-	clip(bIsThinTranslucent ? 1.0f : GetMaterialOpacity(PixelMaterialInputs) - 1.0 / 255.0 - GetMaterialOpacityMaskClipValue());
+	clip(bIsThinTranslucent ? 1.0f : MaterialOpacity - 1.0 / 255.0 - GetMaterialOpacityMaskClipValue());
 #else
 	GetMaterialCoverageAndClipping(Parameters, PixelMaterialInputs);
 #endif
diff --git a/Engine/Shaders/Private/VelocityShader.usf b/Engine/Shaders/Private/VelocityShader.usf
--- a/Engine/Shaders/Private/VelocityShader.usf
+++ b/Engine/Shaders/Private/VelocityShader.usf
@@ -137,7 +137,36 @@ void MainPixelShader(
 #if MATERIALBLENDING_MASKED_USING_COVERAGE
 	OutCoverage = DiscardMaterialWithPixelCoverage(MaterialParameters, PixelMaterialInputs);
 #else
-	GetMaterialClippingVelocity(MaterialParameters, PixelMaterialInputs, VELOCITY_THIN_TRANSLUCENT_MODE > 0);
+
+#if STRATA_ENABLED
+	const bool bThinTranslucentMaterial = false;
+
+	FStrataData StrataData = PixelMaterialInputs.FrontMaterial;
+	FStrataAddressing NullStrataAddressing = (FStrataAddressing)0;	// Fake unused in StrataCreateBSDFContext when using Forward inline shading
+
+	FStrataPixelHeader StrataPixelHeader = InitialiseStrataPixelHeader();
+	StrataPixelHeader.StrataTree = MaterialParameters.StrataTree;
+	StrataPixelHeader.BSDFCount = MaterialParameters.StrataTree.BSDFCount;
+	StrataPixelHeader.SharedLocalBases = MaterialParameters.SharedLocalBases;
+
+	float MaterialOpacity = 1.0f;
+	FStrataTree StrataTree = StrataPixelHeader.StrataTree;
+	if (StrataTree.BSDFCount > 0)
+	{
+		// Update tree (coverage/transmittance/luminace weights)
+		const FStrataIntegrationSettings Settings = InitStrataIntegrationSettings();
+		const float3 V = MaterialParameters.CameraVector;
+		StrataUpdateTree(NullStrataAddressing, StrataPixelHeader, StrataTree, V, Settings);
+
+		MaterialOpacity = StrataTree.Operators[StrataData.OperatorIndex].Coverage;
+	}
+
+#else
+	const bool bThinTranslucentMaterial = VELOCITY_THIN_TRANSLUCENT_MODE > 0;
+	const float MaterialOpacity = GetMaterialOpacity(PixelMaterialInputs);
+#endif
+
+	GetMaterialClippingVelocity(MaterialParameters, PixelMaterialInputs, bThinTranslucentMaterial, MaterialOpacity);
 #endif
 
 	// 3d velocity, includes camera an object motion
