Fixed flickering SkyViewLUT texture when viewed from afar in space.

diff --git a/Engine/Shaders/Private/SkyAtmosphere.usf b/Engine/Shaders/Private/SkyAtmosphere.usf
--- a/Engine/Shaders/Private/SkyAtmosphere.usf
+++ b/Engine/Shaders/Private/SkyAtmosphere.usf
@@ -421,11 +421,14 @@ SingleScatteringResult IntegrateSingleScatteredLuminance(
 
 	float2 PixPos = SVPos.xy;
 
-	// Compute next intersection with atmosphere or ground 
+	// Compute next intersection with atmosphere or ground
 	float3 PlanetO = float3(0.0f, 0.0f, 0.0f);
+	float tMax = 0.0f;
+#if 0
+	// The bottom code causes the skyview lut to flicker when view from space afar.
+	// Remove that code when the else section is proven.
 	float tBottom = RaySphereIntersectNearest(WorldPos, WorldDir, PlanetO, Atmosphere.BottomRadiusKm);
 	float tTop = RaySphereIntersectNearest(WorldPos, WorldDir, PlanetO, Atmosphere.TopRadiusKm);
-	float tMax = 0.0f;
 	if (tBottom < 0.0f)
 	{
 		if (tTop < 0.0f)
@@ -445,6 +448,34 @@ SingleScatteringResult IntegrateSingleScatteredLuminance(
 			tMax = min(tTop, tBottom);
 		}
 	}
+#else
+	float tBottom = 0.0f;
+	float2 SolB = RayIntersectSphere(WorldPos, WorldDir, float4(PlanetO, Atmosphere.BottomRadiusKm));
+	float2 SolT = RayIntersectSphere(WorldPos, WorldDir, float4(PlanetO, Atmosphere.TopRadiusKm));
+
+	const bool bNoBotIntersection = all(SolB < 0.0f);
+	const bool bNoTopIntersection = all(SolT < 0.0f);
+	if (bNoTopIntersection)
+	{
+		// No intersection with planet or its atmosphere.
+		tMax = 0.0f;
+		return Result;
+	}
+	else if (bNoBotIntersection)
+	{
+		// No intersection with planet, so we trace up to the far end of the top atmosphere 
+		// (looking up from ground or edges when see from afar in space).
+		tMax = max(SolT.x, SolT.y);
+	}
+	else
+	{
+		// Interesection with planet and atmospehre: we simply trace up to the planet ground.
+		// We know there is at least one intersection thanks to [REDACTED_PERSON].
+		// If one of the solution is invalid=-1, that means we are inside the planet: we stop tracing by setting tBottom=0.
+		tBottom = max(0.0f, min(SolB.x, SolB.y));
+		tMax = tBottom;
+	}
+#endif
 
 	float PlanetOnOpaque = 1.0f;	// This is used to hide opaque meshes under the planet ground
 #if VIEWDATA_AVAILABLE
