Fix for landscape vertex factory

diff --git a/Engine/Shaders/Private/LandscapeVertexFactory.ush b/Engine/Shaders/Private/LandscapeVertexFactory.ush
--- a/Engine/Shaders/Private/LandscapeVertexFactory.ush
+++ b/Engine/Shaders/Private/LandscapeVertexFactory.ush
@@ -670,21 +670,26 @@ FVertexFactoryInterpolantsVSToPS VertexFactoryGetInterpolantsVSToPS(FVertexFacto
 	Interpolants.LayerTexCoord				= LandscapeTexCoords.LayerTexCoord;
 	Interpolants.WeightHeightMapTexCoord.xy	= LandscapeTexCoords.WeightMapTexCoord;
 	Interpolants.WeightHeightMapTexCoord.zw	= LandscapeTexCoords.HeightMapTexCoord;
-	#if NUM_MATERIAL_TEXCOORDS
-		float2 CustomizedUVs[NUM_MATERIAL_TEXCOORDS];
-		GetMaterialCustomizedUVs(VertexParameters, CustomizedUVs);
 
-		Interpolants.LayerTexCoord.xy = CustomizedUVs[0];
-		Interpolants.TransformedTexCoords = 0;
-
-		#if NUM_MATERIAL_TEXCOORDS > 1
-			Interpolants.TransformedTexCoords.xy = CustomizedUVs[1];
-			#if NUM_MATERIAL_TEXCOORDS > 2
-				Interpolants.TransformedTexCoords.zw = CustomizedUVs[2];
-			#endif // 2
-		#endif // 1
-
-	#endif // NUM_MATERIAL_TEXCOORDS
+	// Landscape material doesn't support custom vertex interpolators for now but GetMaterialCustomizedUVs is declared using NUM_TEX_COORD_INTERPOLATORS
+#if NUM_TEX_COORD_INTERPOLATORS
+	float2 CustomizedUVs[NUM_TEX_COORD_INTERPOLATORS];
+	GetMaterialCustomizedUVs(VertexParameters, CustomizedUVs);
+#endif // NUM_TEX_COORD_INTERPOLATORS
+
+	// Landscape material only supports up to NUM_MATERIAL_TEXCOORDS 5 (0-3 are reserved for landscape-specific stuff: see UMaterialExpressionLandscapeLayerCoords, 4 is for lightmaps and 5 for heightmap coordinates, 
+	//  anything above is not supported : see GetMaterialPixelParameters) : 
+#if NUM_MATERIAL_TEXCOORDS
+	Interpolants.LayerTexCoord.xy = CustomizedUVs[0];
+	Interpolants.TransformedTexCoords = 0;
+
+#if NUM_MATERIAL_TEXCOORDS > 1
+	Interpolants.TransformedTexCoords.xy = CustomizedUVs[1];
+#if NUM_MATERIAL_TEXCOORDS > 2
+	Interpolants.TransformedTexCoords.zw = CustomizedUVs[2];
+#endif // NUM_MATERIAL_TEXCOORDS > 2
+#endif // NUM_MATERIAL_TEXCOORDS > 1
+#endif // NUM_MATERIAL_TEXCOORDS
 
 #if NEEDS_LIGHTMAP_COORDINATE
 	#if VF_USE_PRIMITIVE_SCENE_DATA
