Avoid unecessary RT shadow tracing for hair strands.

Save 25% of shadow cost in test case with many lights.

diff --git a/Engine/Shaders/Private/RayTracing/RayTracingOcclusionRGS.usf b/Engine/Shaders/Private/RayTracing/RayTracingOcclusionRGS.usf
--- a/Engine/Shaders/Private/RayTracing/RayTracingOcclusionRGS.usf
+++ b/Engine/Shaders/Private/RayTracing/RayTracingOcclusionRGS.usf
@@ -64,13 +64,12 @@ uint				bUseHairVoxel;
 float				HairOcclusionThreshold;
 Texture2D<uint>		HairLightChannelMaskTexture;
 
-void NeedTraceHair(
+bool NeedTraceHair(
 	in uint2 PixelCoord,
-	inout bool bTraceHairRay,
 	inout float HairDeviceZ)
 {
 	const float HairSampleDepth = HairStrands.HairOnlyDepthTexture.Load(uint3(PixelCoord, 0)).x;
-	bTraceHairRay = false;
+	bool bTraceHairRay = false;
 	HairDeviceZ = 0;
 	if (HairSampleDepth > 0)
 	{
@@ -78,6 +77,8 @@ void NeedTraceHair(
 		HairDeviceZ = HairSampleDepth;
 		bTraceHairRay = (HairLightChannel & LightingChannelMask) != 0;
 	}
+
+	return bTraceHairRay;
 }
 #endif
 
@@ -327,6 +328,7 @@ struct FOcclusionShadingParameters
 	bool bIsValid;
 	bool bBackfaceCulling;
 	bool bIsHair;
+	bool bIsHairStrands;
 	bool bIsSubsurface;
 	bool bHasSubsurfaceProfile;
 };
@@ -358,7 +360,7 @@ FOcclusionResult ComputeOcclusion(
 	// toward the view point (this create issue for transmitted light)
 	// It seems having proper transmission is more important, thus the light 
 	// version is enabled by default
-	if (ShadingDesc.bIsHair)
+	if (ShadingDesc.bIsHair || ShadingDesc.bIsHairStrands)
 	{
 	#if LIGHT_TYPE == LIGHT_TYPE_DIRECTIONAL
 		const float3 LightDirection = LightParameters.Direction;
@@ -452,7 +454,7 @@ FOcclusionResult ComputeOcclusion(
 		}
 
 	#if USE_HAIR_LIGHTING
-		if (ShadingDesc.bIsHair)
+		if (ShadingDesc.bIsHairStrands)
 		{
 			// Denoiser mode 0 doesn't use depth, so take first hit
 			RayFlags |= RAY_FLAG_SKIP_CLOSEST_HIT_SHADER | RAY_FLAG_ACCEPT_FIRST_HIT_AND_END_SEARCH;
@@ -497,7 +499,7 @@ FOcclusionResult ComputeOcclusion(
 #endif
 
 		#if USE_HAIR_LIGHTING
-		if (!ShadingDesc.bIsHair && bUseHairVoxel)
+		if (!ShadingDesc.bIsHairStrands && bUseHairVoxel)
 		{
 			MinimalPayload.HitT = TraverseHair(PixelCoord, RandSequence, Ray.Origin, Ray.Direction, MinimalPayload.HitT, VirtualVoxel.Raytracing_ShadowOcclusionThreshold);
 		}
@@ -730,12 +732,6 @@ RAY_TRACING_ENTRY_RAYGEN(OcclusionRGS)
 			LocalSamplesPerPixel = 0;
 		}
 
-	#if USE_HAIR_LIGHTING
-		bool bTraceHairRay = false;
-		float HairDeviceZ = 0;
-		NeedTraceHair(PixelCoord, bTraceHairRay, HairDeviceZ);
-	#endif
-
 		Occlusion = ComputeOcclusion(
 			PixelCoord,
 			RAY_TRACING_MASK_SHADOW | RAY_TRACING_MASK_THIN_SHADOW,
@@ -748,18 +744,23 @@ RAY_TRACING_ENTRY_RAYGEN(OcclusionRGS)
 			LocalSamplesPerPixel);
 
 		#if USE_HAIR_LIGHTING
-		if (bTraceHairRay)
+		float HairDeviceZ = 0;
+		if (NeedTraceHair(PixelCoord, HairDeviceZ))
 		{
-			ShadingParameters.bIsHair = true;
+			FTransmissionProfileParams HairTransmissionParameters = (FTransmissionProfileParams)0;
+			FOcclusionShadingParameters HairOcclusionShadingParameters = (FOcclusionShadingParameters)0;
+			HairOcclusionShadingParameters.bIsValid = true;
+			HairOcclusionShadingParameters.bIsHairStrands = true;
+
 			HairOcclusion = ComputeOcclusion(
 				PixelCoord,
 				RAY_TRACING_MASK_SHADOW,
 				HairDeviceZ,
 				float3(1,0,0),
-				ShadingParameters,
+				HairOcclusionShadingParameters,
 				LightParameters,
-				TransmissionParameters,
-				SubsurfaceSamplingMeanFreePath,
+				HairTransmissionParameters,
+				0.f/*SubsurfaceSamplingMeanFreePath/*,
 				RequestedSamplePerPixel);
 		}
 		#endif
