Apply indiret lighting scale on moible.

diff --git a/Engine/Shaders/Private/BasePassPixelShader.usf b/Engine/Shaders/Private/BasePassPixelShader.usf
--- a/Engine/Shaders/Private/BasePassPixelShader.usf
+++ b/Engine/Shaders/Private/BasePassPixelShader.usf
@@ -596,7 +596,7 @@ void GetPrecomputedIndirectLightingAndSkyLight(
 		LightmapUVType LightmapUV0, LightmapUV1;
 		uint LightmapDataIndex;
 		GetLightMapCoordinates(Interpolants, LightmapUV0, LightmapUV1, LightmapDataIndex);
-		GetLightMapColorLQ(LightmapVTPageTableResult, LightmapUV0, LightmapUV1, LightmapDataIndex, DiffuseDir, bEvaluateBackface, OutDiffuseLighting, OutSubsurfaceLighting, OutIndirectIrradiance);
+		GetLightMapColorLQ(LightmapVTPageTableResult, LightmapUV0, LightmapUV1, LightmapDataIndex, DiffuseDir, bEvaluateBackface, OutDiffuseLighting, OutSubsurfaceLighting);
 
 	#endif
 
diff --git a/Engine/Shaders/Private/LightmapCommon.ush b/Engine/Shaders/Private/LightmapCommon.ush
--- a/Engine/Shaders/Private/LightmapCommon.ush
+++ b/Engine/Shaders/Private/LightmapCommon.ush
@@ -75,7 +75,7 @@ half4 SampleLightmapVT(VTPageTableResult LightmapVTPageTableResult, uint LayerIn
 }
 #endif // LIGHTMAP_VT_ENABLED
 
-void GetLightMapColorLQ( VTPageTableResult LightmapVTPageTableResult, LightmapUVType LightmapUV0, LightmapUVType LightmapUV1, uint LightmapDataIndex, half3 WorldNormal, bool bEvaluateBackface, out half3 OutDiffuseLighting, out half3 OutSubsurfaceLighting, out half OutIndirectIrradiance)
+void GetLightMapColorLQ( VTPageTableResult LightmapVTPageTableResult, LightmapUVType LightmapUV0, LightmapUVType LightmapUV1, uint LightmapDataIndex, half3 WorldNormal, bool bEvaluateBackface, out half3 OutDiffuseLighting, out half3 OutSubsurfaceLighting)
 {
 	OutSubsurfaceLighting = 0;
 
@@ -127,7 +127,6 @@ void GetLightMapColorLQ( VTPageTableResult LightmapVTPageTableResult, LightmapUV
 	half3 Color = LogRGB * (Luma / max(0.00001f, LogL));	// 1 smax, 1 rcp, 1 smul, 1 vmul
 
 	OutDiffuseLighting = Color;
-	OutIndirectIrradiance = Luma;
 }
 
 void GetLightMapColorHQ( VTPageTableResult LightmapVTPageTableResult, LightmapUVType LightmapUV0, LightmapUVType LightmapUV1, uint LightmapDataIndex, half3 WorldNormal, float2 SvPositionXY, bool bEvaluateBackface, out half3 OutDiffuseLighting, out half3 OutSubsurfaceLighting )
diff --git a/Engine/Shaders/Private/MobileBasePassPixelShader.usf b/Engine/Shaders/Private/MobileBasePassPixelShader.usf
--- a/Engine/Shaders/Private/MobileBasePassPixelShader.usf
+++ b/Engine/Shaders/Private/MobileBasePassPixelShader.usf
@@ -170,7 +170,7 @@ void ApplyPixelDepthOffsetForMobileBasePass(inout FMaterialPixelParameters Mater
 	#define PIXELSHADER_EARLYDEPTHSTENCIL 	
 #endif
 
-void GetSkyLighting(half3 SkyLightColor, half3 WorldNormal, bool bEvaluateBackface, inout half IndirectIrradiance, out half3 OutSkyDiffuseLighting, out half3 OutSkySubsurfaceLighting, out half3 OutWaterDiffuseIndirectLuminance)
+void GetSkyLighting(half3 SkyLightColor, half3 WorldNormal, bool bEvaluateBackface, out half3 OutSkyDiffuseLighting, out half3 OutSkySubsurfaceLighting, out half3 OutWaterDiffuseIndirectLuminance)
 {
 	OutSkyDiffuseLighting = 0;
 	OutSkySubsurfaceLighting = 0;
@@ -180,7 +180,6 @@ void GetSkyLighting(half3 SkyLightColor, half3 WorldNormal, bool bEvaluateBackfa
 	{
 		half3 SkyDiffuse = GetSkySHDiffuseSimple(WorldNormal);
 		OutSkyDiffuseLighting = SkyDiffuse * SkyLightColor;
-		IndirectIrradiance += Luminance(OutSkyDiffuseLighting);
 
 	#if MATERIAL_SHADINGMODEL_SINGLELAYERWATER
 		OutWaterDiffuseIndirectLuminance += OutSkyDiffuseLighting;
@@ -209,7 +208,7 @@ void GetPrecomputedIndirectLightingAndSkyLight(VTPageTableResult LightmapVTPageT
 	uint LightmapDataIndex;
 	GetLightMapCoordinates(Interpolants, LightmapUV0, LightmapUV1, LightmapDataIndex);
 
-	GetLightMapColorLQ(LightmapVTPageTableResult, LightmapUV0, LightmapUV1, LightmapDataIndex, DiffuseDir, bEvaluateBackface, OutDiffuseIndirectLighting, OutSubsurfaceIndirectLighting, OutIndirectIrradiance);
+	GetLightMapColorLQ(LightmapVTPageTableResult, LightmapUV0, LightmapUV1, LightmapDataIndex, DiffuseDir, bEvaluateBackface, OutDiffuseIndirectLighting, OutSubsurfaceIndirectLighting);
 #elif CACHED_POINT_INDIRECT_LIGHTING
 	#if MATERIALBLENDING_MASKED || MATERIALBLENDING_SOLID
 		// Take the normal into account for opaque
@@ -251,16 +250,22 @@ void GetPrecomputedIndirectLightingAndSkyLight(VTPageTableResult LightmapVTPageT
 		}
 		#endif
 	#endif
-
-	OutIndirectIrradiance = Luminance(OutDiffuseIndirectLighting);
 #endif
 
+	// Apply indirect lighting scale while we have only accumulated lightmaps
+	OutDiffuseIndirectLighting *= View.PrecomputedIndirectLightingColorScale;
+	OutSubsurfaceIndirectLighting *= View.PrecomputedIndirectLightingColorScale;
+
 	half3 SkyDiffuseLighting;
 	half3 SkySubsurfaceLighting;
-	GetSkyLighting(SkyLightColor, DiffuseDir, bEvaluateBackface, OutIndirectIrradiance, SkyDiffuseLighting, SkySubsurfaceLighting, OutWaterDiffuseIndirectLuminance);
+	GetSkyLighting(SkyLightColor, DiffuseDir, bEvaluateBackface, SkyDiffuseLighting, SkySubsurfaceLighting, OutWaterDiffuseIndirectLuminance);
 
 	OutDiffuseIndirectLighting += SkyDiffuseLighting;
 	OutSubsurfaceIndirectLighting += SkySubsurfaceLighting;
+
+#if LQ_TEXTURE_LIGHTMAP || CACHED_POINT_INDIRECT_LIGHTING
+	OutIndirectIrradiance = Luminance(OutDiffuseIndirectLighting);
+#endif
 }
 
 PIXELSHADER_EARLYDEPTHSTENCIL
diff --git a/Engine/Shaders/Private/RayTracing/RayTracingMaterialHitShaders.usf b/Engine/Shaders/Private/RayTracing/RayTracingMaterialHitShaders.usf
--- a/Engine/Shaders/Private/RayTracing/RayTracingMaterialHitShaders.usf
+++ b/Engine/Shaders/Private/RayTracing/RayTracingMaterialHitShaders.usf
@@ -184,7 +184,7 @@ void GetPrecomputedIndirectLightingAndSkyLight(
 	// Low quality texture lightmaps
 	#elif LQ_TEXTURE_LIGHTMAP
 
-		GetLightMapColorLQ(LightmapVTPageTableResult, LightmapUV0, LightmapUV1, LightmapDataIndex, DiffuseDir, bEvaluateBackface, OutDiffuseLighting, OutSubsurfaceLighting, OutIndirectIrradiance);
+		GetLightMapColorLQ(LightmapVTPageTableResult, LightmapUV0, LightmapUV1, LightmapDataIndex, DiffuseDir, bEvaluateBackface, OutDiffuseLighting, OutSubsurfaceLighting);
 
 	#endif
 
