RayTracingOcclusion ShaderCode: Fix malformed /* */ comment.
#rnx

diff --git a/Engine/Shaders/Private/RayTracing/RayTracingOcclusionRGS.usf b/Engine/Shaders/Private/RayTracing/RayTracingOcclusionRGS.usf
--- a/Engine/Shaders/Private/RayTracing/RayTracingOcclusionRGS.usf
+++ b/Engine/Shaders/Private/RayTracing/RayTracingOcclusionRGS.usf
@@ -760,7 +760,7 @@ RAY_TRACING_ENTRY_RAYGEN(OcclusionRGS)
 				HairOcclusionShadingParameters,
 				LightParameters,
 				HairTransmissionParameters,
-				0.f/*SubsurfaceSamplingMeanFreePath/*,
+				0.f/*SubsurfaceSamplingMeanFreePath*/,
 				RequestedSamplePerPixel);
 		}
 		#endif
