Fix flickering on car glass.

diff --git a/Engine/Shaders/Private/DistortAccumulatePS.usf b/Engine/Shaders/Private/DistortAccumulatePS.usf
--- a/Engine/Shaders/Private/DistortAccumulatePS.usf
+++ b/Engine/Shaders/Private/DistortAccumulatePS.usf
@@ -141,7 +141,7 @@ void Main(
 	OutClosestDepthMeter = RefractionDepthMeter;
 
 	// Only allow clipping if roughness is close to 0.
-	const bool TryToClip = OutVarianceCoverage.x < 0.01f; // STRATA_TODO || RefractionCoverage <= 0.0f;	// Have an option to account for coverage or not?, ignored by default
+	const bool TryToClip = false;  //OutVarianceCoverage.x < 0.01f; // STRATA_TODO || RefractionCoverage <= 0.0f;	// Have an option to account for coverage or not?, ignored by default
 
 	// Compute UV distortion
 	float2 BufferUVDistortion = ComputeBufferUVDistortion(
