Substrate - fix or top layer roughness affecting the bottom layer roughness when it should not when coverage is 0, or coverage is 1 but transmittance is 0.

diff --git a/Engine/Shaders/Private/Strata/StrataTree.ush b/Engine/Shaders/Private/Strata/StrataTree.ush
--- a/Engine/Shaders/Private/Strata/StrataTree.ush
+++ b/Engine/Shaders/Private/Strata/StrataTree.ush
@@ -399,10 +399,12 @@ void UpdateAllBSDFWithBottomUpOperatorVisit_Vertical(
 
 		// Evaluate the reflected lobe after refraction through the top layer of the vertical operator.
 		// => This tells how light should be integrated when traveling up thorugh that layer.
-		const float TopLayerSigma = Op.VerticalTop_BottomUpRefractionLobeStat.Sigma;
+		// However, light should only refract and spread from roughness according to the amount of coverage, and only if transmittance is > 0.
+		const float DiracLobeSigma = 0.0f;
+		const float TopLayerSigma = lerp(DiracLobeSigma, Op.VerticalTop_BottomUpRefractionLobeStat.Sigma, Op.VerticalTop_Coverage * dot(Op.VerticalTop_TransmittanceAlongN, (1.0f/3.0f).xxx));
 		const bool bIsTop = Op.LayerDepth == 0;
 		const float BottomUpInterfaceEta = bIsTop ? WATER_IOR / AIR_IOR : 1.0f;
-		FStrataLobeStatistic RefractedLobe = StrataGetRefractedLobe(StrataOppositeLobe(ReflectedLobe), DummyInterfaceDFG, StrataLobeVarianceToRoughness(TopLayerSigma), BottomUpInterfaceEta);/////////////////
+		FStrataLobeStatistic RefractedLobe = StrataGetRefractedLobe(StrataOppositeLobe(ReflectedLobe), DummyInterfaceDFG, StrataLobeVarianceToRoughness(TopLayerSigma), BottomUpInterfaceEta);
 
 		StrataSetBSDFRoughness(CurrentBSDF, StrataLobeVarianceToRoughness(RefractedLobe.Sigma));
 	}
