Support instance custom data in pixel shader with mobile GpuScene enabled.

diff --git a/Engine/Shaders/Private/InstanceCulling/BuildInstanceDrawCommands.usf b/Engine/Shaders/Private/InstanceCulling/BuildInstanceDrawCommands.usf
--- a/Engine/Shaders/Private/InstanceCulling/BuildInstanceDrawCommands.usf
+++ b/Engine/Shaders/Private/InstanceCulling/BuildInstanceDrawCommands.usf
@@ -125,7 +125,7 @@ void WritePackedInstanceData(uint Offset, uint InstanceId, FInstanceSceneData In
 #else
 	InstanceAuxData4.xy = PackLocalBoundsCenter(InstanceData.LocalBoundsCenter);
 	InstanceAuxData4.zw = PackLocalBoundsExtent(InstanceData.LocalBoundsExtent);
-	InstanceAuxData1 = 0; // TODO: pack something useful here 
+	InstanceAuxData1 = InstanceData.RandomID;
 #endif // ALLOW_STATIC_LIGHTING
 
 	// 24 bits for Primitive flags + 8 bits for Instance flags
@@ -141,7 +141,11 @@ void WritePackedInstanceData(uint Offset, uint InstanceId, FInstanceSceneData In
 	
 	InstanceTransform1.w = asfloat(PackedPrimitiveFlags); 
 	InstanceTransform2.w = InstanceAuxData1; 
-	InstanceTransform3.w = InstanceData.RandomID;
+
+	uint CustomDataCount;
+	LoadInstanceRelativeIdAndCustomDataCount(InstanceId, InstanceSceneDataSOAStride, InstanceData.RelativeId, CustomDataCount);
+	FInstancePayloadDataOffsets Offsets = GetInstancePayloadDataOffsets(InstanceData.PrimitiveId, InstanceData.Flags, InstanceData.RelativeId);
+	InstanceTransform3.w = asfloat((CustomDataCount & ((1 << INSTANCE_CUSTOM_DATA_COUNT_NUM_BITS) - 1)) | (Offsets.CustomData << INSTANCE_CUSTOM_DATA_COUNT_NUM_BITS));
 		
 	InstanceIdsBufferOutMobile[Offset*5 + 0] = InstanceOriginAndId;
 	InstanceIdsBufferOutMobile[Offset*5 + 1] = InstanceTransform1;
diff --git a/Engine/Shaders/Private/SceneData.ush b/Engine/Shaders/Private/SceneData.ush
--- a/Engine/Shaders/Private/SceneData.ush
+++ b/Engine/Shaders/Private/SceneData.ush
@@ -1178,10 +1178,6 @@ FSceneDataIntermediates GetSceneDataIntermediates(uint DrawInstanceId, float4 In
 
 	InstanceData.Flags = InstanceFlags;
 	InstanceData.PrimitiveId = asuint(InstanceOrigin.w);
-	
-#if USES_PER_INSTANCE_RANDOM || USE_DITHERED_LOD_TRANSITION
-	InstanceData.RandomID = InstanceTransform3.w;
-#endif
 
 	ComputeInstanceDerivedData(InstanceData, TilePosition, LocalToRelativeWorld);
 
@@ -1215,6 +1211,15 @@ FSceneDataIntermediates GetSceneDataIntermediates(uint DrawInstanceId, float4 In
 	float3 ObjectRelativeWorldPosition = mul(float4(Intermediates.InstanceData.LocalBoundsCenter.xyz, 0.0f), LocalToRelativeWorld).xyz;
 	Intermediates.Primitive.ObjectWorldPosition = MakeLWCVector3(TilePosition, ObjectRelativeWorldPosition);
 	Intermediates.Primitive.ObjectRadius = length(Intermediates.InstanceData.LocalBoundsExtent * Intermediates.InstanceData.NonUniformScale.xyz);
+#if USES_PER_INSTANCE_RANDOM || USE_DITHERED_LOD_TRANSITION
+	Intermediates.InstanceData.RandomID = InstanceTransform2.w;
+#endif
+#endif
+
+#if ENABLE_PER_INSTANCE_CUSTOM_DATA
+	uint PackedCustomDataCountAndOffset = asuint(InstanceTransform3.w);
+	Intermediates.InstanceData.CustomDataCount = uint(PackedCustomDataCountAndOffset & ((1 << INSTANCE_CUSTOM_DATA_COUNT_NUM_BITS) - 1));
+	Intermediates.InstanceData.CustomDataOffset = uint(PackedCustomDataCountAndOffset >> INSTANCE_CUSTOM_DATA_COUNT_NUM_BITS);
 #endif
 	
 	return Intermediates; 
