Fix shader compiler error
- Use LWC codepath in ClipRayToStartOutsideGlobalSDF(...)

diff --git a/Engine/Shaders/Private/DistanceField/GlobalDistanceFieldUtils.ush b/Engine/Shaders/Private/DistanceField/GlobalDistanceFieldUtils.ush
--- a/Engine/Shaders/Private/DistanceField/GlobalDistanceFieldUtils.ush
+++ b/Engine/Shaders/Private/DistanceField/GlobalDistanceFieldUtils.ush
@@ -195,21 +195,11 @@ void ClipRayToStartOutsideGlobalSDF(inout float3 TranslatedRayOrigin, float3 Ray
 {
 	uint ClipmapIndex = NumGlobalSDFClipmaps - 1;
 
-	// For merging with LWC converted Global SDF in main
-#define LWC_GLOBALSDF 0 
-#if LWC_GLOBALSDF
 	float3 TranslatedWorldRayEnd = TranslatedRayOrigin + RayDirection * RayLength;
 	float ClipmapVoxelExtent = GlobalVolumeTranslatedCenterAndExtent[ClipmapIndex].w * GlobalVolumeTexelSize;
 	float GlobalVolumeExtent = GlobalVolumeTranslatedCenterAndExtent[ClipmapIndex].w - ClipmapVoxelExtent;
 	float3 GlobalVolumeTranslatedCenter = GlobalVolumeTranslatedCenterAndExtent[ClipmapIndex].xyz;
 	float2 IntersectionTimes = LineBoxIntersect(TranslatedRayOrigin, TranslatedWorldRayEnd, GlobalVolumeTranslatedCenter - GlobalVolumeExtent.xxx, GlobalVolumeTranslatedCenter + GlobalVolumeExtent.xxx);
-#else
-	float3 TranslatedRayEnd = TranslatedRayOrigin + RayDirection * RayLength;
-	float ClipmapVoxelExtent = GlobalVolumeCenterAndExtent[ClipmapIndex].w * GlobalVolumeTexelSize;
-	float3 TranslatedGlobalVolumeCenter = GlobalVolumeCenterAndExtent[ClipmapIndex].xyz + LWCHackToFloat(PrimaryView.PreViewTranslation);
-	float GlobalVolumeExtent = GlobalVolumeCenterAndExtent[ClipmapIndex].w - ClipmapVoxelExtent;
-	float2 IntersectionTimes = LineBoxIntersect(TranslatedRayOrigin, TranslatedRayEnd, TranslatedGlobalVolumeCenter - GlobalVolumeExtent.xxx, TranslatedGlobalVolumeCenter + GlobalVolumeExtent.xxx);
-#endif
 
 	if (IntersectionTimes.x < IntersectionTimes.y)
 	{
