Substrate - fix for slab to respect additive blending as for legacy.

diff --git a/Engine/Shaders/Private/BasePassPixelShader.usf b/Engine/Shaders/Private/BasePassPixelShader.usf
--- a/Engine/Shaders/Private/BasePassPixelShader.usf
+++ b/Engine/Shaders/Private/BasePassPixelShader.usf
@@ -1846,7 +1846,7 @@ void FPixelShaderInOut_MainPS(
 		#else
 			// Pre multipled alpha blending
 			float AdjustedAlpha = saturate(1 - dot(AdjustedDualBlendMul, float3(1.0f, 1.0f, 1.0f) / 3.0f));
-			#if STRATA_USES_CONVERSION_FROM_LEGACY && MATERIALBLENDING_ADDITIVE
+			#if MATERIALBLENDING_ADDITIVE
 				AdjustedAlpha = 0.0f;
 			#endif
 			// Pre multipled alpha blending
