Fixed emissive exhibiting different behavior when using Substrate.

diff --git a/Engine/Shaders/Private/BasePassPixelShader.usf b/Engine/Shaders/Private/BasePassPixelShader.usf
--- a/Engine/Shaders/Private/BasePassPixelShader.usf
+++ b/Engine/Shaders/Private/BasePassPixelShader.usf
@@ -2078,12 +2078,10 @@ void FPixelShaderInOut_MainPS(
 	}
 #endif // OIT_ENABLED
 
-#if MATERIAL_IS_SKY
-	// Sky materials can result in high luminance values, e.g. the sun disk. 
-	// This is so we make sure to at least stay within the boundaries of fp10 and not cause NaN on some platforms.
+	// Sky dome, or any emissive, materials can result in high luminance values, e.g. the sun disk. 
+	// This Min here is to we make sure pre-exposed luminance remains within the boundaries of fp10 and not cause NaN on some platforms.
 	// We also half that range to also make sure we have room for other additive elements such as bloom, clouds or particle visual effects.
 	Out.MRT[0].xyz = min(Out.MRT[0].xyz, Max10BitsFloat.xxx * 0.5f);
-#endif
 
 #if NUM_VIRTUALTEXTURE_SAMPLES || LIGHTMAP_VT_ENABLED
 #if OUTPUT_PIXEL_DEPTH_OFFSET
diff --git a/Engine/Shaders/Private/Strata/Strata.ush b/Engine/Shaders/Private/Strata/Strata.ush
--- a/Engine/Shaders/Private/Strata/Strata.ush
+++ b/Engine/Shaders/Private/Strata/Strata.ush
@@ -2593,7 +2593,8 @@ float SanitizeF0(float F0)
 
 void SanitizeStrataSlab(inout FStrataBSDF BSDF)
 {
-	BSDF_SETEMISSIVE(BSDF, clamp(BSDF_GETEMISSIVE(BSDF), 0.0, Max10BitsFloat));
+	// We do not clamp emissive luminance upper bound. This is done later in the base pass within pre-exposed luminance space.
+	BSDF_SETEMISSIVE(BSDF, max(BSDF_GETEMISSIVE(BSDF), 0.0));
 	SLAB_DIFFUSEALBEDO(BSDF) = SanitizeDiffuseAlbedo(SLAB_DIFFUSEALBEDO(BSDF));
 	SLAB_F0(BSDF) = SanitizeF0(SLAB_F0(BSDF));
 	SLAB_F90(BSDF) = saturate(SLAB_F90(BSDF));
@@ -2621,7 +2622,8 @@ void SanitizeStrataSlab(inout FStrataBSDF BSDF)
 
 void SanitizeStrataHair(inout FStrataBSDF BSDF)
 {
-	BSDF_SETEMISSIVE(BSDF, clamp(BSDF_GETEMISSIVE(BSDF), 0.0, Max10BitsFloat));
+	// We do not clamp emissive luminance upper bound. This is done later in the base pass within pre-exposed luminance space.
+	BSDF_SETEMISSIVE(BSDF, max(BSDF_GETEMISSIVE(BSDF), 0.0));
 	HAIR_BASECOLOR(BSDF) = SanitizeDiffuseAlbedo(HAIR_BASECOLOR(BSDF));
 	HAIR_SCATTER(BSDF) = saturate(HAIR_SCATTER(BSDF));
 	HAIR_ROUGHNESS(BSDF) = SanitizeRoughness(HAIR_ROUGHNESS(BSDF));
@@ -2631,7 +2633,8 @@ void SanitizeStrataHair(inout FStrataBSDF BSDF)
 
 void SanitizeStrataEye(inout FStrataBSDF BSDF)
 {
-	BSDF_SETEMISSIVE(BSDF, clamp(BSDF_GETEMISSIVE(BSDF), 0.0, Max10BitsFloat));
+	// We do not clamp emissive luminance upper bound. This is done later in the base pass within pre-exposed luminance space.
+	BSDF_SETEMISSIVE(BSDF, max(BSDF_GETEMISSIVE(BSDF), 0.0));
 	EYE_DIFFUSEALBEDO(BSDF) = SanitizeDiffuseAlbedo(EYE_DIFFUSEALBEDO(BSDF));
 	EYE_ROUGHNESS(BSDF) = SanitizeRoughness(EYE_ROUGHNESS(BSDF));
 	EYE_IRISDISTANCE(BSDF) = saturate(EYE_IRISDISTANCE(BSDF));
@@ -2642,7 +2645,8 @@ void SanitizeStrataEye(inout FStrataBSDF BSDF)
 
 void SanitizeStrataSingleLayerWater(inout FStrataBSDF BSDF)
 {
-	BSDF_SETEMISSIVE(BSDF, clamp(BSDF_GETEMISSIVE(BSDF), 0.0, Max10BitsFloat));
+	// We do not clamp emissive luminance upper bound. This is done later in the base pass within pre-exposed luminance space.
+	BSDF_SETEMISSIVE(BSDF, max(BSDF_GETEMISSIVE(BSDF), 0.0));
 	SLW_BASECOLOR(BSDF) = SanitizeDiffuseAlbedo(SLW_BASECOLOR(BSDF));
 	SLW_METALLIC(BSDF) = saturate(SLW_METALLIC(BSDF));
 	SLW_SPECULAR(BSDF) = SanitizeF0(SLW_SPECULAR(BSDF));
@@ -2658,10 +2662,9 @@ void SanitizeStrataSingleLayerWater(inout FStrataBSDF BSDF)
 
 void SanitizeStrataUnlit(inout FStrataBSDF BSDF)
 {
-#if MATERIAL_ALLOW_NEGATIVE_EMISSIVECOLOR
-	BSDF_SETEMISSIVE(BSDF, min(BSDF_GETEMISSIVE(BSDF), Max10BitsFloat));
-#else
-	BSDF_SETEMISSIVE(BSDF, clamp(BSDF_GETEMISSIVE(BSDF), 0.0, Max10BitsFloat));
+#if MATERIAL_ALLOW_NEGATIVE_EMISSIVECOLOR == 0
+	// We do not clamp emissive luminance upper bound. This is done later in the base pass within pre-exposed luminance space.
+	BSDF_SETEMISSIVE(BSDF, max(BSDF_GETEMISSIVE(BSDF), 0.0));
 #endif
 	UNLIT_TRANSMITTANCE(BSDF) = saturate(UNLIT_TRANSMITTANCE(BSDF));
 }
