Work around compiler flow control bugs on a particular platform with per-pixel CS material binning.

diff --git a/Engine/Shaders/Private/ComputeShaderOutputCommon.ush b/Engine/Shaders/Private/ComputeShaderOutputCommon.ush
--- a/Engine/Shaders/Private/ComputeShaderOutputCommon.ush
+++ b/Engine/Shaders/Private/ComputeShaderOutputCommon.ush
@@ -258,100 +258,56 @@ void ExportPixel(const float2 PixelPos, FPixelShaderOut ShadedPixel)
 #endif // STRATA_OPAQUE_DEFERRED
 }
 
-struct FShadingQuadSwizzle
-{
-	bool bActiveLane;
-	uint SwizzleIndex;
-};
+#define VIS_HELPER_LANES 0
 
-FShadingQuadSwizzle GetShadingQuad(uint QuadX, uint QuadY, uint LaneIndex, uint ShadingRate)
+void ProcessPixel(bool bActiveLane, uint ShadingBin, const uint2 PixelPos, uint QuadIndex, bool3 PixelCopy, uint HelperLaneCount)
 {
-	FShadingQuadSwizzle ShadingQuad;
-
-	// https://microsoft.github.io/DirectX-Specs/d3d/HLSL_SM_6_6_Derivatives.html
-	// https://microsoft.github.io/DirectX-Specs/d3d/VariableRateShading.html
-	
-	// D3D12_SHADING_RATE_1X1
-	ShadingQuad.bActiveLane = true;
-	ShadingQuad.SwizzleIndex = LaneIndex;
-
-	const uint QuadIndex = QuadX * QuadY * 2u;
+	// All lanes shade (regardless of export - so ddx/ddy are valid)
+	FPixelShaderOut ShadedPixel = ShadePixel(PixelPos, QuadIndex);
 
-#if 1
-	const bool bTopLeft		= ((QuadIndex & 3u) == 0u);
-	const bool bTopRight	= ((QuadIndex & 3u) == 1u);
-	const bool bBottomLeft	= ((QuadIndex & 3u) == 2u);
-	const bool bBottomRight	= ((QuadIndex & 3u) == 3u);
-#else
-	// To test the QuadRead swizzles
-	const int XSigned = int(X);
-	const int YSigned = int(Y);
-	int LaneXDelta = QuadReadAcrossX(XSigned) - XSigned;
-	int LaneYDelta = QuadReadAcrossY(YSigned) - YSigned;
-	const bool bTopLeft		= (LaneXDelta > 0 && LaneYDelta > 0);
-	const bool bTopRight	= (LaneXDelta < 0 && LaneYDelta > 0);
-	const bool bBottomLeft	= (LaneXDelta > 0 && LaneYDelta < 0);
-	const bool bBottomRight	= (LaneXDelta < 0 && LaneYDelta < 0);
+#if VIS_HELPER_LANES
+	ShadedPixel.MRT[3].rgb = ColorMapTurbo(float(HelperLaneCount) / 4.0f);
+#elif 0
+	ShadedPixel.MRT[3].rgb = ColorMapTurbo(50.0);
+#elif 0
+	// Coherency vis
+	float R =  (DispatchThreadID.x & 0xFFu) / 255.0f;
+	float G = ((DispatchThreadID.x & 0xFF00u) >> 8u) / 255.0f;
+	float B = ((DispatchThreadID.x & 0xFF0000u) >> 16u) / 255.0f;
+	ShadedPixel.MRT[3].rgb = float3(R, G, B);
+#elif 0
+	ShadedPixel.MRT[3].rgb = IntToColor(ShadingBin);
+#elif 0
+	ShadedPixel.MRT[3].rgb = IntToColor(QuadIndex);
+#elif 0
+	ShadedPixel.MRT[3].rgb = VisualizeShadingRate(ShadingRate).rgb;
 #endif
 
-	if (ShadingRate == D3D12_SHADING_RATE_1X2)
-	{
-		// [A]
-		// [B]
-		ShadingQuad.bActiveLane  = (QuadIndex & 1u) == 0u;
-		ShadingQuad.SwizzleIndex = select(ShadingQuad.bActiveLane, LaneIndex, QuadReadAcrossX(LaneIndex));
-	}
-	else if (ShadingRate == D3D12_SHADING_RATE_2X1)
-	{
-		// [A][B]
-		ShadingQuad.bActiveLane  = (QuadIndex & 1u) == 0u;
-		ShadingQuad.SwizzleIndex = select(ShadingQuad.bActiveLane, LaneIndex, QuadReadAcrossY(LaneIndex));
-	}
-	else if (ShadingRate == D3D12_SHADING_RATE_2X2)
-	{
-		// [A][B]
-		// [C][D]
-		ShadingQuad.bActiveLane  = (QuadIndex & 3u) == 0u;
-		ShadingQuad.SwizzleIndex = 
-			select(
-				ShadingQuad.bActiveLane,
-				LaneIndex,
-				select(
-					bBottomRight,
-					QuadReadAcrossDiagonal(LaneIndex),
-					select(
-						bBottomLeft,
-						QuadReadAcrossY(LaneIndex),
-						QuadReadAcrossX(LaneIndex)
-					)
-				)
-			)
-		;
-	}
-	else if (ShadingRate == D3D12_SHADING_RATE_2X4) // not generated yet
-	{
-		// [A][B]
-		// [C][D]
-		// [E][F]
-		// [G][H]
-	}
-	else if (ShadingRate == D3D12_SHADING_RATE_4X2) // not generated yet
-	{
-		// [A][B][C][D]
-		// [E][F][G][H]
-	}
-	else if (ShadingRate == D3D12_SHADING_RATE_4X4) // not generated yet
+	// Disable helper lanes from final export
+	BRANCH
+	if (bActiveLane)
 	{
-		// [A][B][C][D]
-		// [E][F][G][H]
-		// [I][J][K][L]
-		// [M][N][O][P]
-	}
+		ExportPixel(PixelPos, ShadedPixel);
 
-	return ShadingQuad;
-}
+		BRANCH
+		if (PixelCopy.x) // Copy H
+		{
+			ExportPixel(PixelPos + uint2(1, 0), ShadedPixel);
+		}
 
-#define VIS_HELPER_LANES 0
+		BRANCH
+		if (PixelCopy.y) // Copy V
+		{
+			ExportPixel(PixelPos + uint2(0, 1), ShadedPixel);
+		}
+
+		BRANCH
+		if (PixelCopy.z) // Copy D
+		{
+			ExportPixel(PixelPos + uint2(1, 1), ShadedPixel);
+		}
+	}
+}
 
 [numthreads(64, 1, 1)]
 void MainCS(uint ThreadIndex : SV_GroupIndex, uint GroupID : SV_GroupID, uint3 GroupThreadID : SV_GroupThreadID, uint3 DispatchThreadID : SV_DispatchThreadID)
@@ -367,21 +323,19 @@ void MainCS(uint ThreadIndex : SV_GroupIndex, uint GroupID : SV_GroupID, uint3 G
 
 	const bool bQuadBinning = PassData.z != 0u;
 
-	uint2 PixelPos;
-	uint QuadIndex;
-	bool bActiveLane = true;
-
-#if VIS_HELPER_LANES
 	uint HelperLaneCount = 0u;
-#endif
 
-	bool3 PixelCopy = bool3(false, false, false); // H, V, D
+	uint CornerIndex = 0;
+	uint QuadIndex = 0;
+	bool3 PixelCopy = bool3(false, false, false);
+
+	FShadingQuad Quad;
 
 	BRANCH
 	if (bQuadBinning)
 	{
 		// New 2x2 shading quad every 4 samples
-		const uint CornerIndex = ThreadIndex & 3u;
+		CornerIndex = ThreadIndex & 3u;
 		QuadIndex = GroupID * 16u + (ThreadIndex >> 2u);
 
 		const bool bIsQuadTL = CornerIndex == 0u;
@@ -402,87 +356,49 @@ void MainCS(uint ThreadIndex : SV_GroupIndex, uint GroupID : SV_GroupID, uint3 G
 			}
 		}
 		PackedQuad = WaveReadLaneAt(PackedQuad, LaneTL);
-		FShadingQuad Quad = UnpackShadingQuad(PackedQuad);
-		if (!any(Quad.Coverage))
-		{
-			// None of the 2x2 pixels have coverage (likely an out of bounds wave).
-			return;
-		}
-
-		bActiveLane = Quad.Coverage[CornerIndex];
-
-	#if VIS_HELPER_LANES
-		HelperLaneCount = 4u;
-
-		UNROLL
-		for (uint CoverageIndex = 0; CoverageIndex < 4u; ++CoverageIndex)
-		{
-			HelperLaneCount -= Quad.Coverage[CoverageIndex] ? 1u : 0u;
-		}
-	#endif
-
-		// Determine final pixel position on the screen for shading
-		PixelPos = Quad.TopLeft + uint2(CornerIndex & 1u, (CornerIndex >> 1u) & 1u);
+		Quad = UnpackShadingQuad(PackedQuad);
 	}
 	else
 	{
 		const uint PixelIndex = DispatchThreadID.x;
-		QuadIndex = PixelIndex;
 		const uint PixelCount = Nanite.ShadingBinMeta[ShadingBin].y;
+		
 		if (PixelIndex >= PixelCount)
 		{
-			return;
+			Quad.Coverage = bool4(false, false, false, false);
 		}
+		else
+		{
+			Quad.Coverage = bool4(true, true, true, true);
 
-		const uint PixelStart = Nanite.ShadingBinMeta[ShadingBin].z;
-		FShadingPixel Pixel = UnpackShadingPixel(Nanite.ShadingBinData[PixelStart + PixelIndex]);
-		PixelPos = Pixel.PosXY;
-		PixelCopy = Pixel.Copy;
+			const uint PixelStart = Nanite.ShadingBinMeta[ShadingBin].z;
+			FShadingPixel Pixel = UnpackShadingPixel(Nanite.ShadingBinData[PixelStart + PixelIndex]);
+			Quad.TopLeft = Pixel.PosXY;
+		}
 	}
 
-	// All lanes shade (regardless of export - so ddx/ddy are valid)
-	FPixelShaderOut ShadedPixel = ShadePixel(PixelPos, QuadIndex);
+	if (!any(Quad.Coverage))
+	{
+		// None of the 2x2 pixels have coverage (likely an out of bounds wave).
+		return;
+	}
+
+	const bool bActiveLane = Quad.Coverage[CornerIndex];
 
 #if VIS_HELPER_LANES
-	ShadedPixel.MRT[3].rgb = ColorMapTurbo(float(HelperLaneCount) / 4.0f);
-#elif 0
-	ShadedPixel.MRT[3].rgb = ColorMapTurbo(50.0);
-#elif 0
-	// Coherency vis
-	float R =  (DispatchThreadID.x & 0xFFu) / 255.0f;
-	float G = ((DispatchThreadID.x & 0xFF00u) >> 8u) / 255.0f;
-	float B = ((DispatchThreadID.x & 0xFF0000u) >> 16u) / 255.0f;
-	ShadedPixel.MRT[3].rgb = float3(R, G, B);
-#elif 0
-	ShadedPixel.MRT[3].rgb = IntToColor(ShadingBin);
-#elif 0
-	ShadedPixel.MRT[3].rgb = IntToColor(QuadIndex);
-#elif 0
-	ShadedPixel.MRT[3].rgb = VisualizeShadingRate(ShadingRate).rgb;
-#endif
+	HelperLaneCount = 4u;
 
-	// Disable helper lanes from final export
-	if (bActiveLane)
+	UNROLL
+	for (uint CoverageIndex = 0; CoverageIndex < 4u; ++CoverageIndex)
 	{
-		ExportPixel(PixelPos, ShadedPixel);
-
-		BRANCH
-		if (PixelCopy.x) // Copy H
-		{
-			ExportPixel(PixelPos + uint2(1, 0), ShadedPixel);
-		}
+		HelperLaneCount -= Quad.Coverage[CoverageIndex] ? 1u : 0u;
+	}
+#endif
 
-		BRANCH
-		if (PixelCopy.y) // Copy V
-		{
-			ExportPixel(PixelPos + uint2(0, 1), ShadedPixel);
-		}
+	// Determine final pixel position on the screen for shading
+	const uint2 PixelPos = Quad.TopLeft + uint2(CornerIndex & 1u, (CornerIndex >> 1u) & 1u);
 
-		BRANCH
-		if (PixelCopy.z) // Copy D
-		{
-			ExportPixel(PixelPos + uint2(1, 1), ShadedPixel);
-		}
-	}
+	// Shade and export quad
+	ProcessPixel(bActiveLane, ShadingBin, PixelPos, QuadIndex, PixelCopy, HelperLaneCount);
 }
 
