Path Tracer: Fix Direct vs Indirect lighting components not adding back up to beauty

diff --git a/Engine/Shaders/Private/PathTracing/PathTracingCore.ush b/Engine/Shaders/Private/PathTracing/PathTracingCore.ush
--- a/Engine/Shaders/Private/PathTracing/PathTracingCore.ush
+++ b/Engine/Shaders/Private/PathTracing/PathTracingCore.ush
@@ -107,38 +107,34 @@ struct FPathState
 		switch (FirstScatterType)
 		{
 			// In these cases, only keep tracing if the corresponding lighting component is enabled
-			case PATHTRACER_SCATTER_DIFFUSE:  return (GetEnabledLightingContributions() & PATHTRACER_CONTRIBUTION_DIFFUSE ) != 0;
-			case PATHTRACER_SCATTER_SPECULAR: return (GetEnabledLightingContributions() & PATHTRACER_CONTRIBUTION_SPECULAR) != 0;
-			case PATHTRACER_SCATTER_VOLUME:   return (GetEnabledLightingContributions() & PATHTRACER_CONTRIBUTION_VOLUME  ) != 0;
+			case PATHTRACER_SCATTER_DIFFUSE:  return (EnabledIndirectLightingContributions & PATHTRACER_CONTRIBUTION_DIFFUSE ) != 0;
+			case PATHTRACER_SCATTER_SPECULAR: return (EnabledIndirectLightingContributions & PATHTRACER_CONTRIBUTION_SPECULAR) != 0;
+			case PATHTRACER_SCATTER_VOLUME:   return (EnabledIndirectLightingContributions & PATHTRACER_CONTRIBUTION_VOLUME  ) != 0;
 		}
 		// In all other cases, keep tracing, we haven't made a final decision on what this path is yet
 		return true;
 	}
 
-	uint GetEnabledLightingContributions()
-	{
-		return HasMadeContributionScatter() ? EnabledIndirectLightingContributions : EnabledDirectLightingContributions;
-	}
-
-	// The following functions allow us to scale down the contributions for lighting components before we have scattered
 	bool ShouldAccumulateEmissive()
 	{
-		return (GetEnabledLightingContributions() & PATHTRACER_CONTRIBUTION_EMISSIVE) != 0;
+		uint EnabledLightingContributions = HasMadeContributionScatter() ? EnabledIndirectLightingContributions : EnabledDirectLightingContributions;
+		return (EnabledLightingContributions & PATHTRACER_CONTRIBUTION_EMISSIVE) != 0;
 	}
 
+	// The following functions allow us to scale down the contributions for lighting components before we have scattered
 	bool ShouldAccumulateDiffuse()
 	{
-		return (GetEnabledLightingContributions() & PATHTRACER_CONTRIBUTION_DIFFUSE) != 0;
+		return (HasMadeContributionScatter() || (EnabledDirectLightingContributions & PATHTRACER_CONTRIBUTION_DIFFUSE) != 0);
 	}
 
 	bool ShouldAccumulateSpecular()
 	{
-		return (GetEnabledLightingContributions() & PATHTRACER_CONTRIBUTION_SPECULAR) != 0;
+		return (HasMadeContributionScatter() || (EnabledDirectLightingContributions & PATHTRACER_CONTRIBUTION_SPECULAR) != 0);
 	}
 
 	bool ShouldAccumulateVolume()
 	{
-		return (GetEnabledLightingContributions() & PATHTRACER_CONTRIBUTION_VOLUME) != 0;
+		return (HasMadeContributionScatter() || (EnabledDirectLightingContributions & PATHTRACER_CONTRIBUTION_VOLUME) != 0);
 	}
 
 	float2 GetDiffuseSpecularScale(bool bIsVolumeSample)
@@ -1186,7 +1182,7 @@ bool PathTracingKernel(inout FPathState PathState, int Bounce)
 	// Choose a random number for both Light sampling and BxDF sampling
 	float4 RandSample = RandomSequence_GenerateSample4D(PathState.RandSequence);
 
-	float2 DiffuseSpecularScale = PathState.GetDiffuseSpecularScale(bIsVolumeSample);
+	const float2 DiffuseSpecularScale = PathState.GetDiffuseSpecularScale(bIsVolumeSample);
 
 	const bool bDoLightLoop = any(DiffuseSpecularScale > 0);
 
@@ -1293,12 +1289,6 @@ bool PathTracingKernel(inout FPathState PathState, int Bounce)
 		return false;
 	}
 
-	if (!PathState.UpdateScatterType(MaterialSample.ScatterType))
-	{
-		// path scattered for a non-enabled lighting component, we can exit immediately
-		return false;
-	}
-
 	// Russian roulette:
 	//   The probability of keeping the path should be roughly proportional to the weight at the current shade point,
 	//  but just using MaterialWeight would miss out on cases where the path throughput changes color (like in a cornell
@@ -1424,6 +1414,6 @@ bool PathTracingKernel(inout FPathState PathState, int Bounce)
 
 	}
 
-	// this path should continue bouncing
-	return true;
+	// continue the path only if the relevant indirect contribution is enabled
+	return PathState.UpdateScatterType(MaterialSample.ScatterType);
 }
