Substrate - MFP is no longer clamped to fp111110 when in forward to be able to get to a transmittance closer to identity 1.0f. (important when multiple invisible translucent layers are accumulated to not see any transmittance effect).
This also makesure non complex BSDFs are sanitized before they are exported out to the Substrate gbuffer.

diff --git a/Engine/Shaders/Private/Strata/Strata.ush b/Engine/Shaders/Private/Strata/Strata.ush
--- a/Engine/Shaders/Private/Strata/Strata.ush
+++ b/Engine/Shaders/Private/Strata/Strata.ush
@@ -97,6 +97,9 @@
 
 #define STRATA_INLINE_SINGLELAYERWATER (STRATA_ENABLED && MATERIAL_IS_STRATA && STRATA_INLINE_SHADING && MATERIAL_SHADINGMODEL_SINGLELAYERWATER)
 
+#ifndef STRATA_OPAQUE_MATERIAL
+#define STRATA_OPAQUE_MATERIAL 0
+#endif
 
 // Behind the scene, we do all the simple volumetric lighting & transmitance computations considering a slab of 100 centimeters = 1 meter.
 // If the user specifies a different thickness to simulate varying appearance, we rescale the mean free path in order to not have to store thickness as part of the BSDF description.
@@ -2601,7 +2604,9 @@ void SanitizeStrataSlab(inout FStrataBSDF BSDF)
 	SLAB_ROUGHNESS(BSDF) = SanitizeRoughness(SLAB_ROUGHNESS(BSDF));
 	SLAB_ANISOTROPY(BSDF) = clamp(SLAB_ANISOTROPY(BSDF), -1.f, 1.f);
 	// SLAB_HAZINESS is sanitize when packed, see PackHaziness.
-#if STRATA_INLINE_SHADING
+#if STRATA_INLINE_SHADING && STRATA_OPAQUE_MATERIAL
+	// We only clamp MFP when we are going to export it packed into the GBuffer.
+	// Otherwise, for forward shading, it is already safe to use, see RescaleMFPToComputationSpace.
 	BSDF.TmpMFP = clamp(BSDF.TmpMFP, 0.0f, Max111110BitsFloat3);
 #endif
 
@@ -2751,7 +2756,9 @@ float3 RescaleMFPToComputationSpace(float3 InMFPInCm, float InSrcThicknessInCm,
 {
 	// Rescale the BSDF mean free path so to run volumetric material math on a unit slab of medium whose thickness is defined by InDstThicknessInCm
 	// For instance, if thickness is larger => we need to reduce the MFP for the medium to appear visually thicker.
-	return clamp(InMFPInCm * (InDstThicknessInCm / max(STRATA_EPSILON, InSrcThicknessInCm)), 0.0f, Max111110BitsFloat3);
+	// We do not cap the recovered MFP to be able to get closer to and indentity transmittance=1.0. 
+	// MFP is safe already thanks to STRATA_EPSILON and TransmittanceToMeanFreePath. And it will be made safe later with a SanitizeStrataSlab before it is exported to the gbuffer.
+	return max(InMFPInCm * (InDstThicknessInCm / max(STRATA_EPSILON, InSrcThicknessInCm)), 0.0f);
 }
 
 void EnableSlabMFPComputationSpace(inout FStrataBSDF BSDF)
diff --git a/Engine/Shaders/Private/Strata/StrataTree.ush b/Engine/Shaders/Private/Strata/StrataTree.ush
--- a/Engine/Shaders/Private/Strata/StrataTree.ush
+++ b/Engine/Shaders/Private/Strata/StrataTree.ush
@@ -44,9 +44,6 @@ void PostProcessSlabBeforeLighting(inout FStrataBSDF BSDF)
 
 void PostProcessBSDFBeforeLighting(inout FStrataBSDF BSDF)
 {
-	// Sanitize BSDF before it is used for forward shading
-	StrataSanitizeBSDF(BSDF);
-
 	switch (BSDF_GETTYPE(BSDF))
 	{
 	case STRATA_BSDF_TYPE_SLAB:
@@ -89,6 +86,9 @@ void UpdateSingleBSDFOperatorCoverageTransmittance(
 	Roughness = Settings.bRoughnessTracking ? Roughness : 0.0;
 #endif
 
+	// Sanitize BSDF before it is used for forward shading or exported to GBuffer
+	StrataSanitizeBSDF(CurrentBSDF);
+
 	// Distortion and translucency passes accumnulate pass calls for colored transmittance, coverage and TopDownRefractionLobeStat to be setup correctly.
 #if STRATA_COMPLEXPATH || MATERIALBLENDING_ANY_TRANSLUCENT
 
