Use scene UB in hair strands debug view instead of USE_GLOBAL_GPU_SCENE_DATA

diff --git a/Engine/Shaders/Private/HairStrands/HairStrandsDebug.usf b/Engine/Shaders/Private/HairStrands/HairStrandsDebug.usf
--- a/Engine/Shaders/Private/HairStrands/HairStrandsDebug.usf
+++ b/Engine/Shaders/Private/HairStrands/HairStrandsDebug.usf
@@ -867,9 +867,6 @@ uint HairVisibilityNodeGroupSize;
 uint AllocatedSampleCount;
 uint HairInstanceCount;
 
-uint NumInstances;
-uint InstanceDataSOAStride;
-
 Buffer<uint> HairInstanceIDs;
 Texture2D<float> HairCountTexture;
 Texture2D<uint> HairCountUintTexture;
@@ -895,10 +892,10 @@ void DrawInstanceAABB(uint PrimitiveId)
 	// Since we don't have the actual instance ID (they are transient instance vs. persistent instance), 
 	// we find the primitive with the matching PrimitiveID
 	uint InstanceID = ~0;
-	const uint TotalInstanceCount = min(8192u, NumInstances);
+	const uint TotalInstanceCount = min(8192u, Scene.GPUScene.NumInstances);
 	for (uint LocalInstanceID = 0; LocalInstanceID < TotalInstanceCount; ++LocalInstanceID)
 	{
-		FInstanceSceneData InstanceData = GetInstanceSceneData(LocalInstanceID, InstanceDataSOAStride);
+		FInstanceSceneData InstanceData = GetInstanceSceneData(LocalInstanceID, Scene.GPUScene.InstanceDataSOAStride);
 		if (InstanceData.PrimitiveId == PrimitiveId)
 		{
 			InstanceID = LocalInstanceID;
@@ -908,7 +905,7 @@ void DrawInstanceAABB(uint PrimitiveId)
 
 	if (InstanceID != ~0)
 	{
-		FInstanceSceneData InstanceData = GetInstanceSceneData(InstanceID, InstanceDataSOAStride);
+		FInstanceSceneData InstanceData = GetInstanceSceneData(InstanceID, Scene.GPUScene.InstanceDataSOAStride);
 		const float4 InstanceBoundColor = ColorLightGreen;
 		AddOBBWS(InstanceData.LocalBoundsCenter - InstanceData.LocalBoundsExtent, InstanceData.LocalBoundsCenter + InstanceData.LocalBoundsExtent, InstanceBoundColor, LWCHackToFloat(InstanceData.LocalToWorld));
 		AddReferentialWS(LWCHackToFloat(InstanceData.LocalToWorld), 50.f);
diff --git a/Engine/Source/Runtime/Renderer/Private/HairStrands/HairStrandsDebug.cpp b/Engine/Source/Runtime/Renderer/Private/HairStrands/HairStrandsDebug.cpp
--- a/Engine/Source/Runtime/Renderer/Private/HairStrands/HairStrandsDebug.cpp
+++ b/Engine/Source/Runtime/Renderer/Private/HairStrands/HairStrandsDebug.cpp
@@ -196,7 +196,7 @@ class FHairDebugPrintCS : public FGlobalShader
 		SHADER_PARAMETER_RDG_BUFFER_SRV(Buffer, HairMacroGroupVoxelAlignedAABBBuffer)
 		SHADER_PARAMETER_RDG_TEXTURE_SRV(Texture2D, StencilTexture)
 		SHADER_PARAMETER_SAMPLER(SamplerState, LinearSampler)
-		SHADER_PARAMETER_STRUCT_INCLUDE(FGPUSceneResourceParameters, GPUSceneResource)
+		SHADER_PARAMETER_RDG_UNIFORM_BUFFER(FSceneUniformParameters, Scene)
 		SHADER_PARAMETER_STRUCT_REF(FViewUniformShaderParameters, ViewUniformBuffer)
 		SHADER_PARAMETER_STRUCT_INCLUDE(ShaderPrint::FShaderParameters, ShaderPrintUniformBuffer)
 		SHADER_PARAMETER_RDG_UNIFORM_BUFFER(FHairStrandsViewUniformParameters, HairStrands)
@@ -211,7 +211,6 @@ public:
 		OutEnvironment.CompilerFlags.Add(CFLAG_Debug);
 		OutEnvironment.SetDefine(TEXT("SHADER_PRINT"), 1);
 		OutEnvironment.SetDefine(TEXT("VF_SUPPORTS_PRIMITIVE_SCENE_DATA"), 1);
-		OutEnvironment.SetDefine(TEXT("USE_GLOBAL_GPU_SCENE_DATA"), 1);
 	}
 };
 
@@ -268,7 +267,7 @@ static void AddDebugHairPrintPass(
 	const FIntPoint Resolution(Viewport.Width(), Viewport.Height());
 
 	FHairDebugPrintCS::FParameters* Parameters = GraphBuilder.AllocParameters<FHairDebugPrintCS::FParameters>();
-	Parameters->GPUSceneResource = Scene->GPUScene.GetShaderParameters();
+	Parameters->Scene = View->GetSceneUniforms().GetBuffer(GraphBuilder);
 	Parameters->HairInstanceCount = InstanceIDs.Num();
 	Parameters->HairInstanceIDs = GraphBuilder.CreateSRV(InstancesIDBuffer, PF_R32_UINT);
 	Parameters->GroupSize = GetVendorOptimalGroupSize2D();
