* Added possibility for UMaterialExpressionCustomOutput to limit the max number of outputs. Display an error when the material is using more than the max allowed number
* Increased the maximum landscape grass type outputs from 18 to 32
Note: FMaterialCachedHLSLTree was not updated accordingly but it already doesn't support the same features as the legacy HLSL material translator (e.g. AllowMultipleCustomOutputs)

diff --git a/Engine/Shaders/Private/LandscapeGrassWeight.usf b/Engine/Shaders/Private/LandscapeGrassWeight.usf
--- a/Engine/Shaders/Private/LandscapeGrassWeight.usf
+++ b/Engine/Shaders/Private/LandscapeGrassWeight.usf
@@ -81,108 +81,153 @@ void PSMain(
 #endif // VF_USE_PRIMITIVE_SCENE_DATA
 	float2 PackedHeight = float2(floor(Height / 256), fmod(Height, 256)) / 255;
 
-#if NUM_MATERIAL_OUTPUTS_GETGRASSWEIGHT > 2
+	// Height = first 2 outputs:
+	OutColor = float4(PackedHeight, 0.0, 0.0);
+
 	switch(OutputPass)
 	{
 	case 0:
-#endif
-		// Height, outputs 1 & 2
-		OutColor.xy = PackedHeight;
-		#if NUM_MATERIAL_OUTPUTS_GETGRASSWEIGHT > 0
-			OutColor.z = GetGrassWeight0(MaterialParameters);
-			#if NUM_MATERIAL_OUTPUTS_GETGRASSWEIGHT > 1
-				OutColor.w = GetGrassWeight1(MaterialParameters);
-			#else
-				OutColor.w = 0;
-			#endif
-		#else
-			OutColor.z = 0;
-			OutColor.w = 0;
-		#endif
-#if NUM_MATERIAL_OUTPUTS_GETGRASSWEIGHT > 2
+	{
+#if HAVE_GetGrassWeight0
+		OutColor[2] = GetGrassWeight0(MaterialParameters);
+#endif // HAVE_GetGrassWeight0
+#if HAVE_GetGrassWeight1
+		OutColor[3] = GetGrassWeight1(MaterialParameters);
+#endif // HAVE_GetGrassWeight1
 		break;
+	}
 	case 1:
-		OutColor.x = GetGrassWeight2(MaterialParameters);
-		#if NUM_MATERIAL_OUTPUTS_GETGRASSWEIGHT > 3
-			OutColor.y = GetGrassWeight3(MaterialParameters);
-		#else
-			OutColor.y = 0;
-		#endif
-		#if NUM_MATERIAL_OUTPUTS_GETGRASSWEIGHT > 4
-			OutColor.z = GetGrassWeight4(MaterialParameters);
-		#else
-			OutColor.z = 0;
-		#endif
-		#if NUM_MATERIAL_OUTPUTS_GETGRASSWEIGHT > 5
-			OutColor.w = GetGrassWeight5(MaterialParameters);
-		#else
-			OutColor.w = 0;
-		#endif
+	{
+#if HAVE_GetGrassWeight2
+		OutColor[0] = GetGrassWeight2(MaterialParameters);
+#endif // HAVE_GetGrassWeight2
+#if HAVE_GetGrassWeight3
+		OutColor[1] = GetGrassWeight3(MaterialParameters);
+#endif // HAVE_GetGrassWeight3
+#if HAVE_GetGrassWeight4
+		OutColor[2] = GetGrassWeight4(MaterialParameters);
+#endif // HAVE_GetGrassWeight4
+#if HAVE_GetGrassWeight5
+		OutColor[3] = GetGrassWeight5(MaterialParameters);
+#endif // HAVE_GetGrassWeight5
 		break;
-#if NUM_MATERIAL_OUTPUTS_GETGRASSWEIGHT > 6
+	}
 	case 2:
-		OutColor.x = GetGrassWeight6(MaterialParameters);
-		#if NUM_MATERIAL_OUTPUTS_GETGRASSWEIGHT > 7
-			OutColor.y = GetGrassWeight7(MaterialParameters);
-		#else
-			OutColor.y = 0;
-		#endif
-		#if NUM_MATERIAL_OUTPUTS_GETGRASSWEIGHT > 8
-			OutColor.z = GetGrassWeight8(MaterialParameters);
-		#else
-			OutColor.z = 0;
-		#endif
-		#if NUM_MATERIAL_OUTPUTS_GETGRASSWEIGHT > 9
-			OutColor.w = GetGrassWeight9(MaterialParameters);
-		#else
-			OutColor.w = 0;
-		#endif
+	{
+#if HAVE_GetGrassWeight6
+		OutColor[0] = GetGrassWeight6(MaterialParameters);
+#endif // HAVE_GetGrassWeight6
+#if HAVE_GetGrassWeight7
+		OutColor[1] = GetGrassWeight7(MaterialParameters);
+#endif // HAVE_GetGrassWeight7
+#if HAVE_GetGrassWeight8
+		OutColor[2] = GetGrassWeight8(MaterialParameters);
+#endif // HAVE_GetGrassWeight8
+#if HAVE_GetGrassWeight9
+		OutColor[3] = GetGrassWeight9(MaterialParameters);
+#endif // HAVE_GetGrassWeight9
 		break;
-#endif // NUM_MATERIAL_OUTPUTS_GETGRASSWEIGHT > 6
-#if NUM_MATERIAL_OUTPUTS_GETGRASSWEIGHT > 10
+	}
 	case 3:
-		OutColor.x = GetGrassWeight10(MaterialParameters);
-		#if NUM_MATERIAL_OUTPUTS_GETGRASSWEIGHT > 11
-			OutColor.y = GetGrassWeight11(MaterialParameters);
-		#else
-			OutColor.y = 0;
-		#endif
-		#if NUM_MATERIAL_OUTPUTS_GETGRASSWEIGHT > 12
-			OutColor.z = GetGrassWeight12(MaterialParameters);
-		#else
-			OutColor.z = 0;
-		#endif
-		#if NUM_MATERIAL_OUTPUTS_GETGRASSWEIGHT > 13
-			OutColor.w = GetGrassWeight13(MaterialParameters);
-		#else
-			OutColor.w = 0;
-		#endif
+	{
+#if HAVE_GetGrassWeight10
+		OutColor[0] = GetGrassWeight10(MaterialParameters);
+#endif // HAVE_GetGrassWeight10
+#if HAVE_GetGrassWeight11
+		OutColor[1] = GetGrassWeight11(MaterialParameters);
+#endif // HAVE_GetGrassWeight11
+#if HAVE_GetGrassWeight12
+		OutColor[2] = GetGrassWeight12(MaterialParameters);
+#endif // HAVE_GetGrassWeight12
+#if HAVE_GetGrassWeight13
+		OutColor[3] = GetGrassWeight13(MaterialParameters);
+#endif // HAVE_GetGrassWeight13
 		break;
-#endif // NUM_MATERIAL_OUTPUTS_GETGRASSWEIGHT > 10
-#if NUM_MATERIAL_OUTPUTS_GETGRASSWEIGHT > 14
+	}
 	case 4:
-		OutColor.x = GetGrassWeight14(MaterialParameters);
-		#if NUM_MATERIAL_OUTPUTS_GETGRASSWEIGHT > 15
-			OutColor.y = GetGrassWeight15(MaterialParameters);
-		#else
-			OutColor.y = 0;
-		#endif
-		#if NUM_MATERIAL_OUTPUTS_GETGRASSWEIGHT > 16
-			OutColor.z = GetGrassWeight16(MaterialParameters);
-		#else
-			OutColor.z = 0;
-		#endif
-		#if NUM_MATERIAL_OUTPUTS_GETGRASSWEIGHT > 17
-			OutColor.w = GetGrassWeight17(MaterialParameters);
-		#else
-			OutColor.w = 0;
-		#endif
+	{
+#if HAVE_GetGrassWeight14
+		OutColor[0] = GetGrassWeight14(MaterialParameters);
+#endif // HAVE_GetGrassWeight14
+#if HAVE_GetGrassWeight15
+		OutColor[1] = GetGrassWeight15(MaterialParameters);
+#endif // HAVE_GetGrassWeight15
+#if HAVE_GetGrassWeight16
+		OutColor[2] = GetGrassWeight16(MaterialParameters);
+#endif // HAVE_GetGrassWeight16
+#if HAVE_GetGrassWeight17
+		OutColor[3] = GetGrassWeight17(MaterialParameters);
+#endif // HAVE_GetGrassWeight17
+		break;
+	}
+	case 5:
+	{
+#if HAVE_GetGrassWeight18
+		OutColor[0] = GetGrassWeight18(MaterialParameters);
+#endif // HAVE_GetGrassWeight18
+#if HAVE_GetGrassWeight19
+		OutColor[1] = GetGrassWeight19(MaterialParameters);
+#endif // HAVE_GetGrassWeight19
+#if HAVE_GetGrassWeight20
+		OutColor[2] = GetGrassWeight20(MaterialParameters);
+#endif // HAVE_GetGrassWeight20
+#if HAVE_GetGrassWeight21
+		OutColor[3] = GetGrassWeight21(MaterialParameters);
+#endif // HAVE_GetGrassWeight21
+		break;
+	}
+	case 6:
+	{
+#if HAVE_GetGrassWeight22
+		OutColor[0] = GetGrassWeight22(MaterialParameters);
+#endif // HAVE_GetGrassWeight22
+#if HAVE_GetGrassWeight23
+		OutColor[1] = GetGrassWeight23(MaterialParameters);
+#endif // HAVE_GetGrassWeight23
+#if HAVE_GetGrassWeight24
+		OutColor[2] = GetGrassWeight24(MaterialParameters);
+#endif // HAVE_GetGrassWeight24
+#if HAVE_GetGrassWeight25
+		OutColor[3] = GetGrassWeight25(MaterialParameters);
+#endif // HAVE_GetGrassWeight25
+		break;
+	}
+	case 7:
+	{
+#if HAVE_GetGrassWeight26
+		OutColor[0] = GetGrassWeight26(MaterialParameters);
+#endif // HAVE_GetGrassWeight26
+#if HAVE_GetGrassWeight27
+		OutColor[1] = GetGrassWeight27(MaterialParameters);
+#endif // HAVE_GetGrassWeight27
+#if HAVE_GetGrassWeight28
+		OutColor[2] = GetGrassWeight28(MaterialParameters);
+#endif // HAVE_GetGrassWeight28
+#if HAVE_GetGrassWeight29
+		OutColor[3] = GetGrassWeight29(MaterialParameters);
+#endif // HAVE_GetGrassWeight29
+		break;
+	}
+	case 8:
+	{
+#if HAVE_GetGrassWeight30
+		OutColor[0] = GetGrassWeight30(MaterialParameters);
+#endif // HAVE_GetGrassWeight30
+#if HAVE_GetGrassWeight31
+		OutColor[1] = GetGrassWeight31(MaterialParameters);
+#endif // HAVE_GetGrassWeight31
 		break;
-#endif // NUM_MATERIAL_OUTPUTS_GETGRASSWEIGHT > 14
+
+// UMaterialExpressionLandscapeGrassOutput::MaxGrassTypes == 32. Uncomment this if it ever gets increased :
+//#if HAVE_GetGrassWeight32
+//		OutColor[2] = GetGrassWeight32(MaterialParameters);
+//#endif // HAVE_GetGrassWeight32
+//#if HAVE_GetGrassWeight33
+//		OutColor[3] = GetGrassWeight33(MaterialParameters);
+//#endif // HAVE_GetGrassWeight33
+	}
 	default:
-		OutColor = float4(0,0,0,0);
 		break;
 	}
-#endif // NUM_MATERIAL_OUTPUTS_GETGRASSWEIGHT > 2
 }
 #endif
\ No newline at end of file
diff --git a/Engine/Source/Runtime/Engine/Classes/Materials/MaterialExpressionCustomOutput.h b/Engine/Source/Runtime/Engine/Classes/Materials/MaterialExpressionCustomOutput.h
--- a/Engine/Source/Runtime/Engine/Classes/Materials/MaterialExpressionCustomOutput.h
+++ b/Engine/Source/Runtime/Engine/Classes/Materials/MaterialExpressionCustomOutput.h
@@ -17,6 +17,8 @@ class UMaterialExpressionCustomOutput : public UMaterialExpression
 
 	// Override to enable multiple outputs
 	virtual int32 GetNumOutputs() const { return 1; };
+	// Override to limit the maximum number of outputs
+	virtual int32 GetMaxOutputs() const { return -1; };
 	virtual FString GetFunctionName() const PURE_VIRTUAL(UMaterialExpressionCustomOutput::GetFunctionName, return TEXT("GetCustomOutput"););
 	virtual FString GetDisplayName() const { return GetFunctionName(); }
 
diff --git a/Engine/Source/Runtime/Engine/Private/Materials/HLSLMaterialTranslator.cpp b/Engine/Source/Runtime/Engine/Private/Materials/HLSLMaterialTranslator.cpp
--- a/Engine/Source/Runtime/Engine/Private/Materials/HLSLMaterialTranslator.cpp
+++ b/Engine/Source/Runtime/Engine/Private/Materials/HLSLMaterialTranslator.cpp
@@ -562,31 +562,39 @@ void FHLSLMaterialTranslator::CompileCustomOutputs(TArray<UMaterialExpressionCus
 		else
 		{
 			SeenCustomOutputExpressionsClasses.Add(CustomOutput->GetClass());
-			int32 NumOutputs = CustomOutput->GetNumOutputs();
 
-			if (CustomOutput->NeedsCustomOutputDefines())
+			int32 MaxOutputs = CustomOutput->GetMaxOutputs();
+			int32 NumOutputs = CustomOutput->GetNumOutputs();
+			if (NumOutputs > MaxOutputs)
 			{
-				ResourcesString += FString::Printf(TEXT("#define NUM_MATERIAL_OUTPUTS_%s %d\r\n"), *CustomOutput->GetFunctionName().ToUpper(), NumOutputs);
+				Errorf(TEXT("The material can only contain up to %i output %s nodes (current number: %i)"), MaxOutputs, *CustomOutput->GetDescription(), NumOutputs);
 			}
-
-			if (NumOutputs > 0)
+			else
 			{
-				EShaderFrequency CustomOutputShaderFrequency = CustomOutput->GetShaderFrequency();
-				for (int32 Index = 0; Index < NumOutputs; Index++)
+				if (CustomOutput->NeedsCustomOutputDefines())
+				{
+					ResourcesString += FString::Printf(TEXT("#define NUM_MATERIAL_OUTPUTS_%s %d\r\n"), *CustomOutput->GetFunctionName().ToUpper(), NumOutputs);
+				}
+
+				if (NumOutputs > 0)
 				{
+					EShaderFrequency CustomOutputShaderFrequency = CustomOutput->GetShaderFrequency();
+					for (int32 Index = 0; Index < NumOutputs; Index++)
 					{
-						ClearFunctionStack(CustomOutputShaderFrequency);
-						FunctionStacks[CustomOutputShaderFrequency].Add(new FMaterialFunctionCompileState(nullptr));
+						{
+							ClearFunctionStack(CustomOutputShaderFrequency);
+							FunctionStacks[CustomOutputShaderFrequency].Add(new FMaterialFunctionCompileState(nullptr));
+						}
+						MaterialProperty = MP_MAX; // Indicates we're not compiling any material property.
+						ShaderFrequency = CustomOutputShaderFrequency;
+						TArray<FShaderCodeChunk> CustomExpressionChunks;
+						AssignTempScope(CustomExpressionChunks);
+						CustomOutput->Compile(this, Index);
 					}
-					MaterialProperty = MP_MAX; // Indicates we're not compiling any material property.
-					ShaderFrequency = CustomOutputShaderFrequency;
-					TArray<FShaderCodeChunk> CustomExpressionChunks;
-					AssignTempScope(CustomExpressionChunks);
-					CustomOutput->Compile(this, Index);
-				}
 
-				ClearFunctionStack(CustomOutputShaderFrequency);
-				FunctionStacks[CustomOutputShaderFrequency].Add(new FMaterialFunctionCompileState(nullptr));
+					ClearFunctionStack(CustomOutputShaderFrequency);
+					FunctionStacks[CustomOutputShaderFrequency].Add(new FMaterialFunctionCompileState(nullptr));
+				}
 			}
 		}
 	}
diff --git a/Engine/Source/Runtime/Engine/Private/Materials/MaterialCachedHLSLTree.cpp b/Engine/Source/Runtime/Engine/Private/Materials/MaterialCachedHLSLTree.cpp
--- a/Engine/Source/Runtime/Engine/Private/Materials/MaterialCachedHLSLTree.cpp
+++ b/Engine/Source/Runtime/Engine/Private/Materials/MaterialCachedHLSLTree.cpp
@@ -185,6 +185,9 @@ void FMaterialCachedHLSLTree::EmitSharedCode(FStringBuilderBase& OutCode) const
 		const FString OutputName = CustomOutput->GetFunctionName();
 		const EShaderFrequency ShaderFrequency = CustomOutput->GetShaderFrequency();
 
+		// TODO [jonathan.bard] : Add support for UMaterialExpressionCustomOutput::GetMaxOutputs and UMaterialExpressionCustomOutput::AllowMultipleCustomOutputs, these should generate compilation errors. 
+		//  Not sure how errors are reported by FMaterialCachedHLSLTree ATM...
+
 		if (CustomOutput->NeedsCustomOutputDefines())
 		{
 			OutCode.Appendf(TEXT("#define NUM_MATERIAL_OUTPUTS_%s %d\n"), *OutputName.ToUpper(), NumOutputs);
diff --git a/Engine/Source/Runtime/Landscape/Classes/Materials/MaterialExpressionLandscapeGrassOutput.h b/Engine/Source/Runtime/Landscape/Classes/Materials/MaterialExpressionLandscapeGrassOutput.h
--- a/Engine/Source/Runtime/Landscape/Classes/Materials/MaterialExpressionLandscapeGrassOutput.h
+++ b/Engine/Source/Runtime/Landscape/Classes/Materials/MaterialExpressionLandscapeGrassOutput.h
@@ -42,6 +42,9 @@ class LANDSCAPE_API UMaterialExpressionLandscapeGrassOutput : public UMaterialEx
 {
 	GENERATED_UCLASS_BODY()
 
+	// Maximum number of supported grass types on a given landscape material. Whenever adjusting this, make sure to update LandscapeGrassWeight.usf accordingly:
+	static constexpr int32 MaxGrassTypes = 32;
+
 #if WITH_EDITOR
 	virtual int32 Compile(class FMaterialCompiler* Compiler, int32 OutputIndex) override;
 	virtual void GetCaption(TArray<FString>& OutCaptions) const override;
@@ -60,6 +63,7 @@ class LANDSCAPE_API UMaterialExpressionLandscapeGrassOutput : public UMaterialEx
 #endif
 
 	virtual int32 GetNumOutputs() const override { return GrassTypes.Num(); }
+	virtual int32 GetMaxOutputs() const { return MaxGrassTypes; };
 	virtual FString GetFunctionName() const override { return TEXT("GetGrassWeight"); }
 
 	UPROPERTY(EditAnywhere, Category = UMaterialExpressionLandscapeGrassOutput)
