Workaround broken interpolator arrays on some Android devices by flattening array in a Slate elements

diff --git a/Engine/Shaders/Private/SlateElementPixelShader.usf b/Engine/Shaders/Private/SlateElementPixelShader.usf
--- a/Engine/Shaders/Private/SlateElementPixelShader.usf
+++ b/Engine/Shaders/Private/SlateElementPixelShader.usf
@@ -70,6 +70,33 @@ half3 GammaCorrect(half3 InColor)
 	return CorrectedColor;
 }
 
+void GetSlateUVs(VertexToPixelInterpolants Interpolants, inout float2 UVs[8])
+{
+#if NUM_MATERIAL_TEXCOORDS > 0
+	UVs[0] = SLATE_UV0( Interpolants );
+#endif
+#if NUM_MATERIAL_TEXCOORDS > 1
+	UVs[1] = SLATE_UV1( Interpolants );
+#endif
+#if NUM_MATERIAL_TEXCOORDS > 2
+	UVs[2] = SLATE_UV2( Interpolants );
+#endif
+#if NUM_MATERIAL_TEXCOORDS > 3
+	UVs[3] = SLATE_UV3( Interpolants );
+#endif
+#if NUM_MATERIAL_TEXCOORDS > 4
+	UVs[4] = SLATE_UV4( Interpolants );
+#endif
+#if NUM_MATERIAL_TEXCOORDS > 5
+	UVs[5] = SLATE_UV5( Interpolants );
+#endif
+#if NUM_MATERIAL_TEXCOORDS > 6
+	UVs[6] = SLATE_UV6( Interpolants );
+#endif
+#if NUM_MATERIAL_TEXCOORDS > 7
+	UVs[7] = SLATE_UV7( Interpolants );
+#endif
+}
 
 #if USE_MATERIALS
 void FillTexCoords( inout FMaterialPixelParameters Parameters, VertexToPixelInterpolants InVertex, float2 InDefaultUV )
@@ -77,37 +104,40 @@ void FillTexCoords( inout FMaterialPixelParameters Parameters, VertexToPixelInte
 #if SHADER_TYPE == ST_GrayscaleFont || SHADER_TYPE == ST_ColorFont
 	#if NUM_MATERIAL_TEXCOORDS == 1
 		// Note: The first channel is used to look up into the texture atlas. The second channel is the custom 0-1 UV set for the text run
-		Parameters.TexCoords[0] = GetUV( InVertex, 1 );
+		Parameters.TexCoords[0] = SLATE_UV1( InVertex );
 	#elif NUM_MATERIAL_TEXCOORDS == 2
 		// Note: The first channel is used to look up into the texture atlas. The second channel is the custom 0-1 UV set for the text run
-		Parameters.TexCoords[0] = GetUV( InVertex, 1 );
+		Parameters.TexCoords[0] = SLATE_UV1( InVertex );
 		Parameters.TexCoords[1] = InVertex.MaterialTexCoords;
 	#endif
 #elif (SHADER_TYPE == ST_Custom || NUM_CUSTOMIZED_UVS != 0) && NUM_MATERIAL_TEXCOORDS > 0
 	// note: We assume users who customize UVs or are using completely custom elements know what they are doing with UVs and want no assumptions made about how uvs are used
+	float2 UVs[8]; GetSlateUVs(InVertex, UVs);
 	UNROLL
 	for( int CoordinateIndex = 0; CoordinateIndex < NUM_MATERIAL_TEXCOORDS; CoordinateIndex++ )
 	{
-		Parameters.TexCoords[CoordinateIndex] = GetUV( InVertex, CoordinateIndex );
+		Parameters.TexCoords[CoordinateIndex] = UVs[CoordinateIndex];
 	}
 #else
 	#if NUM_MATERIAL_TEXCOORDS == 1
 		Parameters.TexCoords[0] = InDefaultUV;
 	#elif NUM_MATERIAL_TEXCOORDS == 2
 		Parameters.TexCoords[0] = InDefaultUV;
-		Parameters.TexCoords[1] = GetUV( InVertex, 0 );
+		Parameters.TexCoords[1] = SLATE_UV0( InVertex );
 	#elif NUM_MATERIAL_TEXCOORDS == 3
 		Parameters.TexCoords[0] = InDefaultUV;
-		Parameters.TexCoords[1] = GetUV( InVertex, 0 );
-		Parameters.TexCoords[2] = GetUV( InVertex, 1 );
+		Parameters.TexCoords[1] = SLATE_UV0( InVertex );
+		Parameters.TexCoords[2] = SLATE_UV1( InVertex );
 	#elif NUM_MATERIAL_TEXCOORDS >= 4
 		Parameters.TexCoords[0] = InDefaultUV;
-		Parameters.TexCoords[1] = GetUV( InVertex, 0 );
-		Parameters.TexCoords[2] = GetUV( InVertex, 1 );
+		Parameters.TexCoords[1] = SLATE_UV0( InVertex );
+		Parameters.TexCoords[2] = SLATE_UV1( InVertex );
+
+		float2 UVs[8]; GetSlateUVs(InVertex, UVs);
 		UNROLL
 		for( int CoordinateIndex = 3; CoordinateIndex < NUM_MATERIAL_TEXCOORDS; CoordinateIndex++ )
 		{
-			Parameters.TexCoords[CoordinateIndex] = GetUV( InVertex, CoordinateIndex );
+			Parameters.TexCoords[CoordinateIndex] = UVs[CoordinateIndex];
 		}
 	#endif // NUM_MATERIAL_TEXCOORDS == 1
 #endif	// SHADER_TYPE
@@ -227,7 +257,7 @@ half4 GetColor( VertexToPixelInterpolants InVertex, float2 UV )
 #if SHADER_TYPE == ST_GrayscaleFont || SHADER_TYPE == ST_ColorFont
 half4 GetFontElementColor( VertexToPixelInterpolants InVertex )
 {
-	float2 AtlasTextureCoordinates = GetUV( InVertex, 0 );
+	float2 AtlasTextureCoordinates = SLATE_UV0( InVertex );
 #if USE_MATERIALS
 	half4 BaseColor = GetMaterialColor( InVertex, float2(0,0) ) * InVertex.Color;
 #else
@@ -252,14 +282,14 @@ half4 GetFontElementColor( VertexToPixelInterpolants InVertex )
 #if SHADER_TYPE == ST_Default || SHADER_TYPE == ST_Custom
 half4 GetDefaultElementColor( VertexToPixelInterpolants InVertex )
 {
-	return GetColor( InVertex, GetUV( InVertex, 0) * GetUV( InVertex, 1 ) );
+	return GetColor( InVertex, SLATE_UV0( InVertex ) * SLATE_UV1( InVertex ) );
 }
 #endif
 
 #if SHADER_TYPE == ST_Border 
 half4 GetBorderElementColor( VertexToPixelInterpolants InVertex )
 {
-	float4 TextureCoordinates = InVertex.TextureCoordinates[0];
+	float4 TextureCoordinates = InVertex.TextureCoordinates0;
 	float2 NewUV;
 	if( TextureCoordinates.z == 0.0f && TextureCoordinates.w == 0.0f )
 	{
@@ -296,7 +326,7 @@ half4 GetBorderElementColor( VertexToPixelInterpolants InVertex )
 #if SHADER_TYPE == ST_RoundedBox
 half4 GetRoundedBoxElementColor( VertexToPixelInterpolants InVertex )
 {
-	return GetRoundedBoxElementColorInternal(ShaderParams.zw, InVertex.MaterialTexCoords.xy, ShaderParams.y, ShaderParams2, GetColor(InVertex, GetUV(InVertex, 0) * GetUV(InVertex, 1)),  InVertex.SecondaryColor);
+	return GetRoundedBoxElementColorInternal(ShaderParams.zw, InVertex.MaterialTexCoords.xy, ShaderParams.y, ShaderParams2, GetColor(InVertex, SLATE_UV0(InVertex) * SLATE_UV1(InVertex)),  InVertex.SecondaryColor);
 }
 #endif
 
@@ -313,7 +343,7 @@ half4 GetRoundedBoxElementColor( VertexToPixelInterpolants InVertex )
  */
 half4 GetLineSegmentElementColor( VertexToPixelInterpolants InVertex )
 {
-	const float2 Gradient = GetUV(InVertex, 0);
+	const float2 Gradient = SLATE_UV0(InVertex);
 
 	// Get coverage based on distance from segment sides and ends
 	const float2 OutsideFilterUV = float2(1.0f, 1.0f);
diff --git a/Engine/Shaders/Private/SlateShaderCommon.ush b/Engine/Shaders/Private/SlateShaderCommon.ush
--- a/Engine/Shaders/Private/SlateShaderCommon.ush
+++ b/Engine/Shaders/Private/SlateShaderCommon.ush
@@ -36,14 +36,27 @@ struct VertexToPixelInterpolants
 	half4 SecondaryColor : COLOR1;
 	float4 OriginalPosition : ORIGINAL_POSITION;
 	float2 MaterialTexCoords : TEXCOORD0;
-	float4 TextureCoordinates[NUM_SLATE_TEXCOORDS] : TEXCOORD1;
+	
+	float4 TextureCoordinates0 : TEXCOORD1;
+#if NUM_SLATE_TEXCOORDS > 2
+	float4 TextureCoordinates1 : TEXCOORD2;
+#endif
+#if NUM_SLATE_TEXCOORDS > 4
+	float4 TextureCoordinates2 : TEXCOORD3;
+#endif
+#if NUM_SLATE_TEXCOORDS > 6
+	float4 TextureCoordinates3 : TEXCOORD4;
+#endif
 };
 
-float2 GetUV(VertexToPixelInterpolants Interpolants, uint UVIndex)
-{
-	float4 UVVector = Interpolants.TextureCoordinates[UVIndex / 2];
-	return Mod(UVIndex, 2) ? UVVector.zw : UVVector.xy;
-}
+#define SLATE_UV0(Interpolants) Interpolants.TextureCoordinates0.xy
+#define SLATE_UV1(Interpolants) Interpolants.TextureCoordinates0.zw
+#define SLATE_UV2(Interpolants) Interpolants.TextureCoordinates1.xy
+#define SLATE_UV3(Interpolants) Interpolants.TextureCoordinates1.zw
+#define SLATE_UV4(Interpolants) Interpolants.TextureCoordinates2.xy
+#define SLATE_UV5(Interpolants) Interpolants.TextureCoordinates2.zw
+#define SLATE_UV6(Interpolants) Interpolants.TextureCoordinates3.xy
+#define SLATE_UV7(Interpolants) Interpolants.TextureCoordinates3.zw
 
 float cross(float2 a, float2 b)
 {
diff --git a/Engine/Shaders/Private/SlateVertexShader.usf b/Engine/Shaders/Private/SlateVertexShader.usf
--- a/Engine/Shaders/Private/SlateVertexShader.usf
+++ b/Engine/Shaders/Private/SlateVertexShader.usf
@@ -34,7 +34,7 @@ VertexToPixelInterpolants Main(
 	VOut.MaterialTexCoords = InMaterialTextureCoordinates;
 	VOut.Color = InColor FCOLOR_COMPONENT_SWIZZLE;;
 	VOut.SecondaryColor = InSecondaryColor FCOLOR_COMPONENT_SWIZZLE;
-	VOut.TextureCoordinates[0] = InTextureCoordinates;
+	VOut.TextureCoordinates0 = InTextureCoordinates;
 
 #if USE_MATERIALS
 	FMaterialVertexParameters VertexParameters = (FMaterialVertexParameters)0;
@@ -77,22 +77,29 @@ VertexToPixelInterpolants Main(
 
 		float2 CustomizedUVs[NUM_SLATE_TEXCOORDS];
 		GetMaterialCustomizedUVs(VertexParameters, CustomizedUVs);
-
-		{
-			UNROLL
-			for( uint CoordinateIndex = 0; CoordinateIndex < NUM_SLATE_TEXCOORDS; CoordinateIndex++ )
-			{
-				FLATTEN
-				if (Mod(CoordinateIndex, 2))
-				{
-					VOut.TextureCoordinates[CoordinateIndex / 2].zw = CustomizedUVs[CoordinateIndex];
-				}
-				else
-				{
-					VOut.TextureCoordinates[CoordinateIndex / 2].xy = CustomizedUVs[CoordinateIndex];
-				}
-			}
-		}
+		
+		SLATE_UV0(VOut) = CustomizedUVs[0];
+	#if NUM_SLATE_TEXCOORDS > 1
+		SLATE_UV1(VOut) = CustomizedUVs[1];
+	#endif
+	#if NUM_SLATE_TEXCOORDS > 2
+		SLATE_UV2(VOut) = CustomizedUVs[2];
+	#endif
+	#if NUM_SLATE_TEXCOORDS > 3
+		SLATE_UV3(VOut) = CustomizedUVs[3];
+	#endif
+	#if NUM_SLATE_TEXCOORDS > 4
+		SLATE_UV4(VOut) = CustomizedUVs[4];
+	#endif
+	#if NUM_SLATE_TEXCOORDS > 5
+		SLATE_UV5(VOut) = CustomizedUVs[5];
+	#endif
+	#if NUM_SLATE_TEXCOORDS > 6
+		SLATE_UV6(VOut) = CustomizedUVs[6];
+	#endif
+	#if NUM_SLATE_TEXCOORDS > 7
+		SLATE_UV7(VOut) = CustomizedUVs[7];
+	#endif
 		
 	#endif // HAS_MATERIAL_TEXCOORDS
 #endif // NUM_MATERIAL_TEXCOORDS_VERTEX
