Readd independent samplers support for decals.

diff --git a/Engine/Shaders/Private/DBufferDecalShared.ush b/Engine/Shaders/Private/DBufferDecalShared.ush
--- a/Engine/Shaders/Private/DBufferDecalShared.ush
+++ b/Engine/Shaders/Private/DBufferDecalShared.ush
@@ -16,9 +16,16 @@
 #define OpaqueBasePassDBufferCTexture OpaqueBasePass.DBufferCTexture
 #define OpaqueBasePassDBufferRenderMask OpaqueBasePass.DBufferRenderMask
 
+#if SUPPORTS_INDEPENDENT_SAMPLERS
+#define OpaqueBasePassDBufferATextureSampler OpaqueBasePass.DBufferATextureSampler
+#define OpaqueBasePassDBufferBTextureSampler OpaqueBasePass.DBufferATextureSampler
+#define OpaqueBasePassDBufferCTextureSampler OpaqueBasePass.DBufferATextureSampler
+#else
 #define OpaqueBasePassDBufferATextureSampler OpaqueBasePass.DBufferATextureSampler
 #define OpaqueBasePassDBufferBTextureSampler OpaqueBasePass.DBufferBTextureSampler
 #define OpaqueBasePassDBufferCTextureSampler OpaqueBasePass.DBufferCTextureSampler
+#endif
+
 #endif // DBUFFER_BASEPASS_RESOURCE
 
 // Returns a bitmask of DBuffer targets that contain valid data for this pixel.
diff --git a/Engine/Shaders/Private/Strata/StrataDbuffer.usf b/Engine/Shaders/Private/Strata/StrataDbuffer.usf
--- a/Engine/Shaders/Private/Strata/StrataDbuffer.usf
+++ b/Engine/Shaders/Private/Strata/StrataDbuffer.usf
@@ -61,9 +61,15 @@ SamplerState DBufferCTextureSampler;
 #define OpaqueBasePassDBufferCTexture DBufferCTexture
 #define OpaqueBasePassDBufferRenderMask DBufferRenderMask
 
-#define OpaqueBasePassDBufferBTextureSampler DBufferBTextureSampler
+#if SUPPORTS_INDEPENDENT_SAMPLERS
+#define OpaqueBasePassDBufferATextureSampler DBufferATextureSampler
+#define OpaqueBasePassDBufferBTextureSampler DBufferATextureSampler
+#define OpaqueBasePassDBufferCTextureSampler DBufferATextureSampler
+#else
 #define OpaqueBasePassDBufferATextureSampler DBufferATextureSampler
+#define OpaqueBasePassDBufferBTextureSampler DBufferBTextureSampler
 #define OpaqueBasePassDBufferCTextureSampler DBufferCTextureSampler
+#endif
 
 #include "../DBufferDecalShared.ush"
 
