Fix StrataVisualization shader compilation error.

diff --git a/Engine/Shaders/Private/Strata/StrataVisualize.usf b/Engine/Shaders/Private/Strata/StrataVisualize.usf
--- a/Engine/Shaders/Private/Strata/StrataVisualize.usf
+++ b/Engine/Shaders/Private/Strata/StrataVisualize.usf
@@ -76,7 +76,7 @@ void DrawPixelFootprint(float3 P, float3 dPdx, float3 dPdy, uint2 PixelCoord)
 			MFP = StrataSubSurfaceHeaderGetMFP(SSSHeader);
 		}
 
-		FStrataPixelFootprint Footprint = StrataGetPixelFootprint(dPdx, dPdy);
+		FStrataPixelFootprint Footprint = StrataGetPixelFootprint(dPdx, dPdy, 0.f /*InNormalCurvatureRoughness*/);
 		AddDrawPixelFootprint(P, dPdx, dPdy, 0.5f, false, ColorRed);
 		AddDrawPixelFootprint(P, dPdx, dPdy, Footprint.PixelRadiusInWorldSpace, true, ColorOrange);
 		AddDrawPixelFootprint(P, dPdx, dPdy, max3(MFP.x, MFP.y, MFP.z), true, ColorCyan);
