Fix incorrect hair coverage threshold comparison.

diff --git a/Engine/Shaders/Private/HairStrands/HairStrandsHitProxy.usf b/Engine/Shaders/Private/HairStrands/HairStrandsHitProxy.usf
--- a/Engine/Shaders/Private/HairStrands/HairStrandsHitProxy.usf
+++ b/Engine/Shaders/Private/HairStrands/HairStrandsHitProxy.usf
@@ -25,7 +25,7 @@ void EmitPS(
 	// Take the first sample
 	const FNodeDesc NodeDesc = DecodeNodeDesc(VisNodeIndex.Load(uint3(PixelCoord,0)));
 	const float Coverage = CoverageTexture.Load(uint3(PixelCoord, 0));
-	if (NodeDesc.Count > 0 && Coverage > CoverageThreshold)
+	if (NodeDesc.Count > 0 && Coverage >= CoverageThreshold)
 	{
 		const uint SampleIndex = NodeDesc.Offset + 0;
 		const FHairVis InNode  = UnpackHairVis(VisNodeData[SampleIndex]);
@@ -76,7 +76,7 @@ void EmitPS(
 		const uint SampleIndex = NodeDesc.Offset + 0;
 		const FHairVis InNode = UnpackHairVis(VisNodeData[SampleIndex]);
 		const uint MaterialId = clamp(InNode.MaterialId, 0, MaxMaterialCount - 1);
-		if (SelectionMaterialIdBuffer[MaterialId] > 0 && Coverage > CoverageThreshold)
+		if (SelectionMaterialIdBuffer[MaterialId] > 0 && Coverage >= CoverageThreshold)
 		{
 			bIsSelected = true;
 			OutDepth = InNode.Depth;
