Added/fixed local light support for mobile-multi view forward rendering. Although local lights were enabled before, only light data for the primary view (left eye) was used, leading to severe warping and artifacts in the right eye. Secondary view (right eye) data is now passed in via the mobile base pass uniform buffer in a similar manner to Instanced Stereo.

diff --git a/Engine/Shaders/Private/ForwardLightingCommon.ush b/Engine/Shaders/Private/ForwardLightingCommon.ush
--- a/Engine/Shaders/Private/ForwardLightingCommon.ush
+++ b/Engine/Shaders/Private/ForwardLightingCommon.ush
@@ -426,7 +426,7 @@ half3 GetImageBasedReflectionSpecular(FMaterialPixelParameters MaterialParameter
 		uint DataStartIndex = 0;
 
 #if REFLECTION_COMPOSITE_USE_BLENDED_REFLECTION_CAPTURES
-		#if INSTANCED_STEREO
+		#if (INSTANCED_STEREO || MOBILE_MULTI_VIEW)
 			BRANCH
 			if (EyeIndex == 0)
 			{
@@ -436,7 +436,7 @@ half3 GetImageBasedReflectionSpecular(FMaterialPixelParameters MaterialParameter
 				NumLocalReflectionCaptures = min(ForwardLightData.NumCulledLightsGrid[NumCulledEntryIndex + 0], ForwardLightData.NumReflectionCaptures);
 				DataStartIndex = ForwardLightData.NumCulledLightsGrid[NumCulledEntryIndex + 1];
 
-		#if INSTANCED_STEREO
+		#if (INSTANCED_STEREO || MOBILE_MULTI_VIEW)
 			}
 			else
 			{
diff --git a/Engine/Shaders/Private/LightGridCommon.ush b/Engine/Shaders/Private/LightGridCommon.ush
--- a/Engine/Shaders/Private/LightGridCommon.ush
+++ b/Engine/Shaders/Private/LightGridCommon.ush
@@ -8,7 +8,9 @@
 
 #include "LightData.ush"
 
-#if INSTANCED_STEREO
+#if MOBILE_MULTI_VIEW
+	#define ForwardLightDataISR MobileBasePass.ForwardMMV
+#elif INSTANCED_STEREO
 	#if MATERIALBLENDING_ANY_TRANSLUCENT
 		#define ForwardLightDataISR TranslucentBasePass.Shared.ForwardISR
 	#else
@@ -51,7 +53,7 @@ FLightGridData GetLightGridData(uint EyeIndex)
 {
 	FLightGridData Result;
 
-#if INSTANCED_STEREO
+#if (INSTANCED_STEREO || MOBILE_MULTI_VIEW)
 	BRANCH
 	if (EyeIndex == 0)
 	{
@@ -61,7 +63,7 @@ FLightGridData GetLightGridData(uint EyeIndex)
 		Result.LightGridZParams = ForwardLightData.LightGridZParams;
 		Result.CulledGridSize = ForwardLightData.CulledGridSize;
 
-#if INSTANCED_STEREO
+#if (INSTANCED_STEREO || MOBILE_MULTI_VIEW)
 	}
 	else
 	{
@@ -108,7 +110,7 @@ uint ComputeLightGridCellIndex(uint2 PixelPos, float SceneDepth)
 
 uint GetNumLocalLights(uint EyeIndex)
 {
-#if INSTANCED_STEREO
+#if (INSTANCED_STEREO || MOBILE_MULTI_VIEW)
 	return (EyeIndex == 0) ? ForwardLightData.NumLocalLights : ForwardLightDataISR.NumLocalLights;
 #else
 	return ForwardLightData.NumLocalLights;
@@ -125,7 +127,7 @@ FCulledLightsGridData GetCulledLightsGrid(uint GridIndex, uint EyeIndex)
 {
 	FCulledLightsGridData Result;
 
-#if INSTANCED_STEREO
+#if (INSTANCED_STEREO || MOBILE_MULTI_VIEW)
 	BRANCH
 	if (EyeIndex == 0)
 	{
@@ -134,7 +136,7 @@ FCulledLightsGridData GetCulledLightsGrid(uint GridIndex, uint EyeIndex)
 		Result.NumLocalLights = min(ForwardLightData.NumCulledLightsGrid[GridIndex * NUM_CULLED_LIGHTS_GRID_STRIDE + 0], ForwardLightData.NumLocalLights);
 		Result.DataStartIndex = ForwardLightData.NumCulledLightsGrid[GridIndex * NUM_CULLED_LIGHTS_GRID_STRIDE + 1];
 
-#if INSTANCED_STEREO
+#if (INSTANCED_STEREO || MOBILE_MULTI_VIEW)
 	}
 	else
 	{
@@ -178,7 +180,7 @@ FDirectionalLightData GetDirectionalLightData(uint EyeIndex)
 {
 	FDirectionalLightData Result;
 
-#if INSTANCED_STEREO
+#if (INSTANCED_STEREO || MOBILE_MULTI_VIEW)
 	BRANCH
 	if (EyeIndex == 0)
 	{
@@ -191,7 +193,7 @@ FDirectionalLightData GetDirectionalLightData(uint EyeIndex)
 		Result.DirectionalLightDirection			= ForwardLightData.DirectionalLightDirection;
 		Result.DirectionalLightSourceRadius			= ForwardLightData.DirectionalLightSourceRadius;
 
-#if INSTANCED_STEREO
+#if (INSTANCED_STEREO || MOBILE_MULTI_VIEW)
 	}
 	else
 	{
@@ -254,7 +256,7 @@ FLocalLightData GetLocalLightData(uint GridIndex, uint EyeIndex)
 {
 	FLocalLightData Out;
 
-#if INSTANCED_STEREO
+#if (INSTANCED_STEREO || MOBILE_MULTI_VIEW)
 	BRANCH
 	if (EyeIndex == 0)
 	{
@@ -273,7 +275,7 @@ FLocalLightData GetLocalLightData(uint GridIndex, uint EyeIndex)
 			ForwardLightData.ForwardLocalLightBuffer[LocalLightBaseIndex + 3],
 			ForwardLightData.ForwardLocalLightBuffer[LocalLightBaseIndex + 4],
 			ForwardLightData.ForwardLocalLightBuffer[LocalLightBaseIndex + 5]);
-#if INSTANCED_STEREO
+#if (INSTANCED_STEREO || MOBILE_MULTI_VIEW)
 	}
 	else
 	{
diff --git a/Engine/Shaders/Private/MobileBasePassPixelShader.usf b/Engine/Shaders/Private/MobileBasePassPixelShader.usf
--- a/Engine/Shaders/Private/MobileBasePassPixelShader.usf
+++ b/Engine/Shaders/Private/MobileBasePassPixelShader.usf
@@ -817,8 +817,15 @@ void Main(
 	// Local lights
 #if  !MATERIAL_SHADINGMODEL_SINGLELAYERWATER && ENABLE_CLUSTERED_LIGHTS
 	{
+
+#if MOBILE_MULTI_VIEW
+		const uint EyeIndex = BasePassInterpolants.MultiViewId;
+#else
 		const uint EyeIndex = 0;
+#endif
+
 		const FCulledLightsGridData CulledLightGridData = GetCulledLightsGrid(GridIndex, EyeIndex);
+
 		half4 LocalLightDynamicShadowFactors = 1.0f;
 
 		#if USE_SHADOWMASKTEXTURE
@@ -994,4 +1001,5 @@ void Main(
 #if DEFERRED_SHADING_PATH && USE_GLES_FBF_DEFERRED 
 	OutProxy.rgb = OutColor.rgb;
 #endif
+
 }
diff --git a/Engine/Shaders/Private/ReflectionEnvironmentComposite.ush b/Engine/Shaders/Private/ReflectionEnvironmentComposite.ush
--- a/Engine/Shaders/Private/ReflectionEnvironmentComposite.ush
+++ b/Engine/Shaders/Private/ReflectionEnvironmentComposite.ush
@@ -42,7 +42,7 @@ float3 CompositeReflectionCapturesAndSkylightTWS(
 #ifdef REFLECTION_COMPOSITE_NO_CULLING_DATA
 		CaptureIndex = TileCaptureIndex;	// Go from 0 to NumCapturesAffectingTile as absolute index in capture array
 #else
-	#if INSTANCED_STEREO
+	#if (INSTANCED_STEREO || MOBILE_MULTI_VIEW)
 		BRANCH
 		if (EyeIndex == 0)
 		{
@@ -50,7 +50,7 @@ float3 CompositeReflectionCapturesAndSkylightTWS(
 	
 			CaptureIndex = GetCulledLightDataGrid(CaptureDataStartIndex + TileCaptureIndex);
 	
-	#if INSTANCED_STEREO
+	#if (INSTANCED_STEREO || MOBILE_MULTI_VIEW)
 		}
 		else
 		{
diff --git a/Engine/Source/Runtime/Renderer/Private/MobileBasePassRendering.cpp b/Engine/Source/Runtime/Renderer/Private/MobileBasePassRendering.cpp
--- a/Engine/Source/Runtime/Renderer/Private/MobileBasePassRendering.cpp
+++ b/Engine/Source/Runtime/Renderer/Private/MobileBasePassRendering.cpp
@@ -182,6 +182,18 @@ void SetupMobileBasePassUniformParameters(
 		SetupDummyForwardLightUniformParameters(GraphBuilder, BasePassParameters.Forward, View.GetShaderPlatform());
 	}
 
+	// Setup forward light data for mobile multi-view secondary view if enabled and available
+	const FViewInfo* InstancedView = View.GetInstancedView();
+	const FForwardLightData* InstancedForwardLightData = InstancedView ? InstancedView->ForwardLightingResources.ForwardLightData : nullptr;
+	if (View.bIsMobileMultiViewEnabled && InstancedForwardLightData)
+	{
+		BasePassParameters.ForwardMMV = *InstancedForwardLightData;
+	}
+	else
+	{
+		BasePassParameters.ForwardMMV = BasePassParameters.Forward; // Duplicate primary data if not
+	}
+
 	const FScene* Scene = View.Family->Scene ? View.Family->Scene->GetRenderScene() : nullptr;
 	const FPlanarReflectionSceneProxy* ReflectionSceneProxy = Scene ? Scene->GetForwardPassGlobalPlanarReflection() : nullptr;
 	SetupPlanarReflectionUniformParameters(View, ReflectionSceneProxy, BasePassParameters.PlanarReflection);
diff --git a/Engine/Source/Runtime/Renderer/Private/MobileBasePassRendering.h b/Engine/Source/Runtime/Renderer/Private/MobileBasePassRendering.h
--- a/Engine/Source/Runtime/Renderer/Private/MobileBasePassRendering.h
+++ b/Engine/Source/Runtime/Renderer/Private/MobileBasePassRendering.h
@@ -35,6 +35,7 @@ BEGIN_GLOBAL_SHADER_PARAMETER_STRUCT(FMobileBasePassUniformParameters, )
 	SHADER_PARAMETER(float, AmbientOcclusionStaticFraction)
 	SHADER_PARAMETER_STRUCT(FFogUniformParameters, Fog)
 	SHADER_PARAMETER_STRUCT(FForwardLightData, Forward)
+	SHADER_PARAMETER_STRUCT(FForwardLightData, ForwardMMV)
 	SHADER_PARAMETER_STRUCT(FPlanarReflectionUniformParameters, PlanarReflection) // Single global planar reflection for the forward pass.
 	SHADER_PARAMETER_STRUCT(FMobileSceneTextureUniformParameters, SceneTextures)
 	SHADER_PARAMETER_STRUCT(FStrataMobileForwardPassUniformParameters, Strata)
