Fixed VRS checkerboard artifacts on opacity mask blend mode materials when MSAA is enabled.

Coverage duplication will be gated by SV_ShadingRate in 5.3 onwards. Branch was tested to ensure no performance hit.

diff --git a/Engine/Shaders/Private/BasePassPixelShader.usf b/Engine/Shaders/Private/BasePassPixelShader.usf
--- a/Engine/Shaders/Private/BasePassPixelShader.usf
+++ b/Engine/Shaders/Private/BasePassPixelShader.usf
@@ -132,6 +132,7 @@
 #define EDITOR_ALPHA2COVERAGE (USE_EDITOR_COMPOSITING && SUPPORTS_PIXEL_COVERAGE)
 #define POST_PROCESS_SUBSURFACE ((MATERIAL_SHADINGMODEL_SUBSURFACE_PROFILE || MATERIAL_SHADINGMODEL_EYE) && USES_GBUFFER)
 #define MATERIAL_WORKS_WITH_DUAL_SOURCE_COLOR_BLENDING (MATERIAL_SHADINGMODEL_THIN_TRANSLUCENT || STRATA_BLENDING_TRANSLUCENT_COLOREDTRANSMITTANCE )
+#define OUTPUTS_COVERAGE (MATERIALBLENDING_MASKED_USING_COVERAGE && !EARLY_Z_PASS_ONLY_MATERIAL_MASKING)
 
 #define OIT_ENABLED (PROJECT_OIT && PLATFORM_SUPPORTS_ROV && (MATERIALBLENDING_TRANSLUCENT || MATERIALBLENDING_ADDITIVE || MATERIAL_WORKS_WITH_DUAL_SOURCE_COLOR_BLENDING))
 
@@ -2109,6 +2110,18 @@ void FPixelShaderInOut_MainPS(
 	Out.MRT[0] = StrataMaterialExport;
 #endif
 
+// Duplicate coverage bit values in case multiple pixels are written by one invocation (VRS)
+#if OUTPUTS_COVERAGE
+	// If centroids are integers, it means we're at a non-1x1 shading rate
+	if(View.NumSceneColorMSAASamples > 1 && (frac(In.SvPosition.x) == 0 || frac(In.SvPosition.y) == 0)) 
+	{
+		Out.Coverage &= 0xF;
+		Out.Coverage |= (Out.Coverage << 4);
+		Out.Coverage |= (Out.Coverage << 8);
+		Out.Coverage |= (Out.Coverage << 12);
+	}
+#endif
+
 }
 
 #if STRATA_OPAQUE_DEFERRED
@@ -2176,7 +2189,7 @@ void FPixelShaderInOut_MainPS(
 //	#define PIXELSHADEROUTPUT_MRT1 (WRITES_VELOCITY_TO_GBUFFER || (MATERIAL_SHADINGMODEL_THIN_TRANSLUCENT && THIN_TRANSLUCENT_USE_DUAL_BLEND))
 //#endif //USES_GBUFFER
 #define PIXELSHADEROUTPUT_A2C ((EDITOR_ALPHA2COVERAGE) != 0)
-#define PIXELSHADEROUTPUT_COVERAGE (MATERIALBLENDING_MASKED_USING_COVERAGE && !EARLY_Z_PASS_ONLY_MATERIAL_MASKING)
+#define PIXELSHADEROUTPUT_COVERAGE ((OUTPUTS_COVERAGE) != 0)
 
 // all PIXELSHADEROUTPUT_ and "void FPixelShaderInOut_MainPS()" need to be setup before this include
 // this include generates the wrapper code to call MainPS(inout FPixelShaderOutput PixelShaderOutput)
