Refactor how decals are accumulated and applied in path tracer
- Evaluate and accumulate decals into temporary structure FDecalData.
    - done immediately after tracing the main ray so that it can be used to support DBuffer in a follow up CL.
- Apply accumulated decals after main FPackedPathTracingPayload is unpacked.

diff --git a/Engine/Shaders/Private/PathTracing/PathTracingCommon.ush b/Engine/Shaders/Private/PathTracing/PathTracingCommon.ush
--- a/Engine/Shaders/Private/PathTracing/PathTracingCommon.ush
+++ b/Engine/Shaders/Private/PathTracing/PathTracingCommon.ush
@@ -148,6 +148,7 @@ struct FPackedPathTracingPayload : FMinimalPayload
 
 	bool IsFrontFace() { return (GetFlags() & PATH_TRACING_PAYLOAD_OUTPUT_FLAG_FRONT_FACE) != 0; }
 
+	bool IsDecalReceiver() { return (GetFlags() & PATH_TRACING_PAYLOAD_OUTPUT_FLAG_DECAL_RECEIVER) != 0; }
 
 	// These method are meant to be used only when IsVisibilityRay() is true
 	float3 GetRayThroughput()
diff --git a/Engine/Shaders/Private/PathTracing/PathTracingCore.ush b/Engine/Shaders/Private/PathTracing/PathTracingCore.ush
--- a/Engine/Shaders/Private/PathTracing/PathTracingCore.ush
+++ b/Engine/Shaders/Private/PathTracing/PathTracingCore.ush
@@ -190,74 +190,126 @@ void AccumulateRadiance(inout float3 TotalRadiance, float3 PathRadiance, bool bI
 	TotalRadiance += PathRadiance;
 }
 
-void ApplyDecal(inout FPathTracingPayload HitPayload, FDecalShaderPayload DecalPayload)
+struct FDecalData
+{
+	float4 BaseColor;
+	float4 WorldNormal;
+	float4 MetallicSpecularRoughness;
+	float3 Emissive;
+};
+
+float ApplyDecal(inout FPathTracingPayload HitPayload, FDecalData DecalData)
+{
+	HitPayload.BaseColor = HitPayload.BaseColor * DecalData.BaseColor.a + DecalData.BaseColor.rgb;
+
+	// We normalize the normal to get smoother visual result (it helps to avoid having D_GGX explodes toward infinity, and matches ApplyDBufferData(...))
+	HitPayload.WorldNormal = normalize(HitPayload.WorldNormal * DecalData.WorldNormal.a + DecalData.WorldNormal.rgb);
+
+	HitPayload.Metallic = HitPayload.Metallic * DecalData.MetallicSpecularRoughness.a + DecalData.MetallicSpecularRoughness.x;
+	HitPayload.Specular = HitPayload.Specular * DecalData.MetallicSpecularRoughness.a + DecalData.MetallicSpecularRoughness.y;
+	HitPayload.Roughness = HitPayload.Roughness * DecalData.MetallicSpecularRoughness.a + DecalData.MetallicSpecularRoughness.z;
+
+	HitPayload.Radiance += DecalData.Emissive;
+
+	return min3(DecalData.BaseColor.a, DecalData.WorldNormal.a, DecalData.MetallicSpecularRoughness.a);
+}
+
+void CombineDecal(FDecalShaderPayload DecalPayload, inout FDecalData DecalData)
 {
 	const float DecalTransparency = DecalPayload.GetTransparency();
 	const uint DecalFlags = DecalPayload.GetFlags();
 
 	if (DecalFlags & DECAL_WRITE_BASE_COLOR_FLAG)
 	{
-		HitPayload.BaseColor = HitPayload.BaseColor * DecalTransparency + DecalPayload.GetBaseColor();
+		DecalData.BaseColor.rgb = DecalData.BaseColor.rgb * DecalTransparency + DecalPayload.GetBaseColor();
+		DecalData.BaseColor.a *= DecalTransparency;
 	}
 
 	if (DecalFlags & DECAL_WRITE_NORMAL_FLAG)
 	{
-		HitPayload.WorldNormal = HitPayload.WorldNormal * DecalTransparency + DecalPayload.GetWorldNormal();
+		DecalData.WorldNormal.rgb = DecalData.WorldNormal.rgb * DecalTransparency + DecalPayload.GetWorldNormal();
+		DecalData.WorldNormal.a *= DecalTransparency;
 	}
 
 	if (DecalFlags & DECAL_WRITE_ROUGHNESS_SPECULAR_METALLIC_FLAG)
 	{
-		const float3 MetallicSpecularRoughness = DecalPayload.GetMetallicSpecularRoughness();
-
-		HitPayload.Metallic = HitPayload.Metallic * DecalTransparency + MetallicSpecularRoughness.x;
-		HitPayload.Specular = HitPayload.Specular * DecalTransparency + MetallicSpecularRoughness.y;
-		HitPayload.Roughness = HitPayload.Roughness * DecalTransparency + MetallicSpecularRoughness.z;
+		DecalData.MetallicSpecularRoughness.rgb = DecalData.MetallicSpecularRoughness.rgb * DecalTransparency + DecalPayload.GetMetallicSpecularRoughness();
+		DecalData.MetallicSpecularRoughness.a *= DecalTransparency;
 	}
 
 	if (DecalFlags & DECAL_WRITE_EMISSIVE_FLAG)
 	{
-		HitPayload.Radiance += DecalPayload.GetEmissive();
+		DecalData.Emissive += DecalPayload.GetEmissive();
 	}
 }
 
-void TraceDecalRay(FRayDesc DecalRay, bool bReverseCulling, uint InstanceInclusionMask, inout FPathTracingPayload HitPayload, inout float DecalThroughput)
+void EvaluateDecals(FPathState PathState, uint InstanceInclusionMask, FPackedPathTracingPayload PackedPayload, inout FDecalData DecalData)
 {
-	const uint MissShaderIndex = 0; // TODO
+	const float3 HitTranslatedWorldPos = PathState.Ray.Origin + PackedPayload.HitT * PathState.Ray.Direction;
 
-	for (;;)
+	if (PathState.PathRoughness < MeshDecalRoughnessCutoff)
 	{
-		FDecalShaderPayload DecalPayload = (FDecalShaderPayload)0;
-		DecalPayload.SetInputTranslatedWorldPosition(HitPayload.TranslatedWorldPos);
-		DecalPayload.SetMiss();
-		TraceRay(
-			DecalTLAS,
-			bReverseCulling ? RAY_FLAG_CULL_FRONT_FACING_TRIANGLES : RAY_FLAG_CULL_BACK_FACING_TRIANGLES,
-			InstanceInclusionMask,
-			RAY_TRACING_SHADER_SLOT_MATERIAL,
-			RAY_TRACING_NUM_SHADER_SLOTS,
-			MissShaderIndex,
-			DecalRay.GetNativeDesc(),
-			DecalPayload);
+		FRayDesc DecalRay = (FRayDesc)0;
+		// go a bit further than the current hit, to handle decals partially inside the model
+		DecalRay.Direction = -PathState.Ray.Direction;
+		DecalRay.Origin = HitTranslatedWorldPos - DecalRay.Direction * MeshDecalBias;
+		DecalRay.TMin = 0.0f;
+		DecalRay.TMax = PackedPayload.HitT + MeshDecalBias;
+
+		const bool bReverseCulling = true;
+		const uint MissShaderIndex = 0; // TODO
 
-		if (DecalPayload.IsMiss())
+		for (;;)
 		{
-			break;
+			FDecalShaderPayload DecalPayload = (FDecalShaderPayload)0;
+			DecalPayload.SetInputTranslatedWorldPosition(HitTranslatedWorldPos);
+			DecalPayload.SetMiss();
+			TraceRay(
+				DecalTLAS,
+				bReverseCulling ? RAY_FLAG_CULL_FRONT_FACING_TRIANGLES : RAY_FLAG_CULL_BACK_FACING_TRIANGLES,
+				InstanceInclusionMask,
+				RAY_TRACING_SHADER_SLOT_MATERIAL,
+				RAY_TRACING_NUM_SHADER_SLOTS,
+				MissShaderIndex,
+				DecalRay.GetNativeDesc(),
+				DecalPayload);
+
+			if (DecalPayload.IsMiss())
+			{
+				break;
+			}
+
+			if (bReverseCulling)
+			{
+				// also reverse normal
+				DecalPayload.SetWorldNormal(-DecalPayload.GetWorldNormal());
+			}
+
+			CombineDecal(DecalPayload, DecalData);
+
+			// prepare next step around the loop
+			// retrace the exact same ray with TMin one ulp past the hit we just found
+			DecalRay.TMin = asfloat(asuint(DecalPayload.HitT) + 1);
 		}
+	}
 
-		if (bReverseCulling)
+#if PLATFORM_SUPPORTS_CALLABLE_SHADERS
+	if (PathState.PathRoughness < DecalRoughnessCutoff)
+	{
+		FDecalLoopCount DecalLoopCount = DecalGridLookup(HitTranslatedWorldPos);
+		for (uint Index = 0, Num = DecalLoopCount.NumDecals; Index < Num; ++Index)
 		{
-			// also reverse normal
-			DecalPayload.SetWorldNormal(-DecalPayload.GetWorldNormal());
-		}
-		
-		ApplyDecal(HitPayload, DecalPayload);
-		
-		DecalThroughput *= DecalPayload.GetTransparency();
+			uint DecalId = GetDecalId(Index, DecalLoopCount);
 
-		// prepare next step around the loop
-		// retrace the exact same ray with TMin one ulp past the hit we just found
-		DecalRay.TMin = asfloat(asuint(DecalPayload.HitT) + 1);
+			FDecalShaderPayload DecalPayload = (FDecalShaderPayload)0;
+			DecalPayload.SetInputTranslatedWorldPosition(HitTranslatedWorldPos);
+
+			CallShader(DecalId, DecalPayload);
+
+			CombineDecal(DecalPayload, DecalData);
+		}
 	}
+#endif
 }
 
 FPathTracingPayload TraceTransparentRay(inout FPathState PathState, bool bIsCameraRay, bool bLastBounce, bool bIncludeEmission, uint NumVisibleLights, inout FVolumeSegment VolumeSegment)
@@ -303,6 +355,17 @@ FPathTracingPayload TraceTransparentRay(inout FPathState PathState, bool bIsCame
 			PathState.Ray.GetNativeDesc(),
 			PackedPayload);
 
+		FDecalData DecalData;
+		DecalData.BaseColor = float4(0, 0, 0, 1);
+		DecalData.WorldNormal = float4(0, 0, 0, 1);
+		DecalData.MetallicSpecularRoughness = float4(0, 0, 0, 1);
+		DecalData.Emissive = 0;
+
+		if (!PackedPayload.IsMiss() && PackedPayload.IsDecalReceiver())
+		{
+			EvaluateDecals(PathState, InstanceInclusionMask, PackedPayload, DecalData);
+		}
+
 		if (PackedPayload.IsMiss() && VolumeIsectList.HitBlocker())
 		{
 			// we didn't hit any real geometry, but we did hit the volume's blocker geometry
@@ -477,41 +540,8 @@ FPathTracingPayload TraceTransparentRay(inout FPathState PathState, bool bIsCame
 		// Unpack the payload
 		FPathTracingPayload HitPayload = UnpackPathTracingPayloadMaterial(PackedPayload, PathState.Ray, bIsCameraRay && ApplyDiffuseSpecularOverrides);
 
-		float DecalThroughput = 1.0f;
-		if (PathState.PathRoughness < MeshDecalRoughnessCutoff && HitPayload.IsDecalReceiver())
-		{
-			FRayDesc DecalRay = (FRayDesc)0;
-			// go a bit further than the current hit, to handle decals partially inside the model
-			DecalRay.Direction = -PathState.Ray.Direction;
-			DecalRay.Origin = HitPayload.TranslatedWorldPos - DecalRay.Direction * MeshDecalBias;
-			DecalRay.TMin = 0.0f;
-			DecalRay.TMax = HitPayload.HitT + MeshDecalBias;
-
-			const bool bReverseCulling = true;
-			TraceDecalRay(DecalRay, bReverseCulling, InstanceInclusionMask, HitPayload, DecalThroughput);
-		}
-
-#if PLATFORM_SUPPORTS_CALLABLE_SHADERS
-		if (PathState.PathRoughness < DecalRoughnessCutoff && HitPayload.IsDecalReceiver())
-		{
-			const float3 TranslatedWorldPos = HitPayload.TranslatedWorldPos;
-
-			FDecalLoopCount DecalLoopCount = DecalGridLookup(TranslatedWorldPos);
-			for (uint Index = 0, Num = DecalLoopCount.NumDecals; Index < Num; ++Index)
-			{
-				uint DecalId = GetDecalId(Index, DecalLoopCount);
+		const float DecalThroughput = ApplyDecal(HitPayload, DecalData);
 
-				FDecalShaderPayload DecalPayload = (FDecalShaderPayload)0;
-				DecalPayload.SetInputTranslatedWorldPosition(TranslatedWorldPos);
-
-				CallShader(DecalId, DecalPayload);
-
-				ApplyDecal(HitPayload, DecalPayload);
-
-				DecalThroughput *= DecalPayload.GetTransparency();
-			}
-		}
-#endif
 		if (DecalThroughput < 1.0f)
 		{
 			// if we ran any decal logic, update the payload
