Support DBuffer functionality in Path Tracer
- Respect Decal Response Mask set in Material
- Support DBufferTextureLookup material node in path tracer hit shaders.

Implementation:
- Provide DBuffer data in FPackedPathTracingPayload to PathTracingMaterialCHS
- Existing TraceRay uses "null" DBuffer data (since hit point is not known yet)
- If ray hits a material that uses DBufferTextureLookup, calculate DBuffer data at hitpoint and, if necessary, retrace ray to reevaluate material with correct DBuffer data.
- MATERIALDECALRESPONSEMASK is passed from CHS to RGS and decal is applied based on that.

diff --git a/Engine/Shaders/Private/DBufferDecalShared.ush b/Engine/Shaders/Private/DBufferDecalShared.ush
--- a/Engine/Shaders/Private/DBufferDecalShared.ush
+++ b/Engine/Shaders/Private/DBufferDecalShared.ush
@@ -28,21 +28,30 @@
 
 #endif // DBUFFER_BASEPASS_RESOURCE
 
+#if USE_DBUFFER && RAYHITGROUPSHADER
+// This must be set prior to running the generated material code.
+static float4 CurrentPayloadDBufferA = float4(0.0, 0.0, 0.0, 1.0);
+static float4 CurrentPayloadDBufferB = float4(0.0, 0.0, 0.0, 1.0);
+static float4 CurrentPayloadDBufferC = float4(0.0, 0.0, 0.0, 1.0);
+#endif
+
 // Returns a bitmask of DBuffer targets that contain valid data for this pixel.
 // @param BufferUV - UV space in the DBuffer textures
 uint GetDBufferTargetMask(uint2 PixelPos)
 {
 #if USE_DBUFFER && ENABLE_DBUFFER_TEXTURES && !RAYHITGROUPSHADER && !LUMEN_CARD_CAPTURE
-
-#if PLATFORM_SUPPORTS_RENDERTARGET_WRITE_MASK
-	return DecodeRTWriteMask(PixelPos, OpaqueBasePassDBufferRenderMask, 3);
-#elif PLATFORM_SUPPORTS_PER_PIXEL_DBUFFER_MASK
-	uint Mask = OpaqueBasePassDBufferRenderMask.Load(uint3(PixelPos, 0));
-	return Mask > 0 ? 0x07 : 0x00;
-#else
+	{
+#	if PLATFORM_SUPPORTS_RENDERTARGET_WRITE_MASK
+		return DecodeRTWriteMask(PixelPos, OpaqueBasePassDBufferRenderMask, 3);
+#	elif PLATFORM_SUPPORTS_PER_PIXEL_DBUFFER_MASK
+		uint Mask = OpaqueBasePassDBufferRenderMask.Load(uint3(PixelPos, 0));
+		return Mask > 0 ? 0x07 : 0x00;
+#	else
+		return 0x07;
+#	endif
+	}
+#elif USE_DBUFFER && RAYHITGROUPSHADER
 	return 0x07;
-#endif
-
 #else
 	return 0;
 #endif
@@ -89,6 +98,13 @@ FStrataDBuffer GetDBufferData(float2 BufferUV, uint RTMaskBit)
 	if (RTMaskBit & 0x2) { DBufferB = Texture2DSampleLevel(OpaqueBasePassDBufferBTexture, OpaqueBasePassDBufferBTextureSampler, BufferUV, 0);	}
 	BRANCH
 	if (RTMaskBit & 0x4) { DBufferC = Texture2DSampleLevel(OpaqueBasePassDBufferCTexture, OpaqueBasePassDBufferCTextureSampler, BufferUV, 0);	}
+#elif USE_DBUFFER && RAYHITGROUPSHADER
+	BRANCH
+	if (RTMaskBit & 0x1) { DBufferA = CurrentPayloadDBufferA; }
+	BRANCH // StrataUnpackDBuffer expects DBufferB to have a specific scale and bias
+	if (RTMaskBit & 0x2) { DBufferB = float4(CurrentPayloadDBufferB.rgb * 0.5f + 128.0f / 255.0f, CurrentPayloadDBufferB.a); }
+	BRANCH
+	if (RTMaskBit & 0x4) { DBufferC = CurrentPayloadDBufferC; }
 #endif
 	return StrataUnpackDBuffer(DBufferA, DBufferB, DBufferC);
 }
@@ -361,6 +377,25 @@ FDBufferData GetDBufferData(float2 BufferUV, uint RTMaskBit)
 	{
 		DBufferC = Texture2DSampleLevel(OpaqueBasePass.DBufferCTexture, OpaqueBasePassDBufferCTextureSampler, BufferUV, 0);
 	}
+#elif USE_DBUFFER && RAYHITGROUPSHADER
+	BRANCH
+	if (RTMaskBit & 0x1)
+	{
+		DBufferA = CurrentPayloadDBufferA;
+	}
+
+	BRANCH
+	if (RTMaskBit & 0x2)
+	{
+		// DecodeDBufferData expects DBufferB to have a specific scale and bias
+		DBufferB = float4(CurrentPayloadDBufferB.rgb * 0.5f + 128.0f / 255.0f, CurrentPayloadDBufferB.a);
+	}
+
+	BRANCH
+	if (RTMaskBit & 0x4)
+	{
+		DBufferC = CurrentPayloadDBufferC;
+	}
 #endif
 
 	return DecodeDBufferData(DBufferA, DBufferB, DBufferC);
diff --git a/Engine/Shaders/Private/PathTracing/PathTracingCommon.ush b/Engine/Shaders/Private/PathTracing/PathTracingCommon.ush
--- a/Engine/Shaders/Private/PathTracing/PathTracingCommon.ush
+++ b/Engine/Shaders/Private/PathTracing/PathTracingCommon.ush
@@ -11,11 +11,17 @@
 #include "../OctahedralCommon.ush"
 #include "/Engine/Shared/PathTracingDefinitions.h"
 
-// These flags are analogous to RAY_TRACING_PAYLOAD_* flags. There are currently 9 bits available
+// These flags are analogous to RAY_TRACING_PAYLOAD_* flags. There are currently 9 bits available.
 #define PATH_TRACING_PAYLOAD_OUTPUT_FLAG_FRONT_FACE                   (1 << 0) // Indicates that ray has hit the front face of a primitive. Set by closest hit shader.
-#define PATH_TRACING_PAYLOAD_OUTPUT_FLAG_DECAL_RECEIVER               (1 << 1) // Indicates that ray has hit a primitive that receives decals. Set by closest hit shader.
-#define PATH_TRACING_PAYLOAD_OUTPUT_FLAG_TWO_SIDED_MATERIAL           (1 << 2) // Indicates that ray has hit a primitive surface with a two sided material. Set by closest hit shader. Used by GPU Lightmass to detect invalid surfaces.
-#define PATH_TRACING_PAYLOAD_OUTPUT_FLAG_IS_HOLDOUT                   (1 << 3) // Indicates that ray has hit a primitive which should be held out from alpha
+#define PATH_TRACING_PAYLOAD_OUTPUT_FLAG_TWO_SIDED_MATERIAL           (1 << 1) // Indicates that ray has hit a primitive surface with a two sided material. Set by closest hit shader. Used by GPU Lightmass to detect invalid surfaces.
+#define PATH_TRACING_PAYLOAD_OUTPUT_FLAG_IS_HOLDOUT                   (1 << 2) // Indicates that ray has hit a primitive which should be held out from alpha.
+#define PATH_TRACING_PAYLOAD_OUTPUT_FLAG_DECAL_RESPONSE_MASK_SHIFT    (3)
+#define PATH_TRACING_PAYLOAD_OUTPUT_FLAG_DECAL_RESPONSE_COLOR         (1 << 3) // Indicates that ray has hit a primitive that receives decals color. Set by closest hit shader.
+#define PATH_TRACING_PAYLOAD_OUTPUT_FLAG_DECAL_RESPONSE_NORMAL        (1 << 4) // Indicates that ray has hit a primitive that receives decals normal. Set by closest hit shader.
+#define PATH_TRACING_PAYLOAD_OUTPUT_FLAG_DECAL_RESPONSE_ROUGHNESS     (1 << 5) // Indicates that ray has hit a primitive that receives decals metallic/specular/roughness. Set by closest hit shader.
+#define PATH_TRACING_PAYLOAD_OUTPUT_FLAG_USE_DBUFFER_LOOKUP           (1 << 6) // Indicates that ray has hit a primitive that uses DBufferTextureLookup. Set by closest hit shader.
+#define PATH_TRACING_PAYLOAD_OUTPUT_FLAG_UNUSED7                      (1 << 7) // Free for future use
+#define PATH_TRACING_PAYLOAD_OUTPUT_FLAG_UNUSED8                      (1 << 8) // Free for future use
 
 #define PATH_TRACING_PAYLOAD_INPUT_FLAG_RAY_TYPE_MASK                 (7 << 0) // Reserve 3 bits for ray-mask
 #define PATH_TRACING_PAYLOAD_INPUT_FLAG_RAY_TYPE_CAMERA                0       // Indicates that this is a camera ray which can trigger a number of special behaviors in the closest hit shader
@@ -147,8 +153,17 @@ struct FPackedPathTracingPayload : FMinimalPayload
 	bool IsVisibilityRay() { return (GetFlags() & PATH_TRACING_PAYLOAD_INPUT_FLAG_RAY_TYPE_MASK) == PATH_TRACING_PAYLOAD_INPUT_FLAG_RAY_TYPE_SHADOW; }
 
 	bool IsFrontFace() { return (GetFlags() & PATH_TRACING_PAYLOAD_OUTPUT_FLAG_FRONT_FACE) != 0; }
-
-	bool IsDecalReceiver() { return (GetFlags() & PATH_TRACING_PAYLOAD_OUTPUT_FLAG_DECAL_RECEIVER) != 0; }
+	
+	bool IsDecalReceiver() 
+	{
+		const uint Mask = 
+			PATH_TRACING_PAYLOAD_OUTPUT_FLAG_DECAL_RESPONSE_COLOR |
+			PATH_TRACING_PAYLOAD_OUTPUT_FLAG_DECAL_RESPONSE_NORMAL | 
+			PATH_TRACING_PAYLOAD_OUTPUT_FLAG_DECAL_RESPONSE_ROUGHNESS | 
+			PATH_TRACING_PAYLOAD_OUTPUT_FLAG_USE_DBUFFER_LOOKUP;
+		return (GetFlags() & Mask) != 0;
+	}
+	bool UsesDBufferLookup() { return (GetFlags() & PATH_TRACING_PAYLOAD_OUTPUT_FLAG_USE_DBUFFER_LOOKUP) != 0; }
 
 	// These method are meant to be used only when IsVisibilityRay() is true
 	float3 GetRayThroughput()
@@ -196,6 +211,60 @@ struct FPackedPathTracingPayload : FMinimalPayload
 		PackedData[6] = asuint(Tau.y);
 		PackedData[7] = asuint(Tau.z);
 	}
+
+	float4 GetDBufferA()
+	{
+		return float4(
+			f16tof32(PackedData[8]),
+			f16tof32(PackedData[8] >> 16),
+			f16tof32(PackedData[9]),
+			f16tof32(PackedData[9] >> 16)
+		);
+	}
+
+	void SetDBufferA(float4 DBufferA)
+	{
+		PackedData[8] = f32tof16(DBufferA.x);
+		PackedData[8] |= f32tof16(DBufferA.y) << 16;
+		PackedData[9] = f32tof16(DBufferA.z);
+		PackedData[9] |= f32tof16(DBufferA.w) << 16;
+	}
+
+	float4 GetDBufferB()
+	{
+		return float4(
+			f16tof32(PackedData[10]),
+			f16tof32(PackedData[10] >> 16),
+			f16tof32(PackedData[11]),
+			f16tof32(PackedData[11] >> 16)
+		);
+	}
+
+	void SetDBufferB(float4 DBufferB)
+	{
+		PackedData[10] = f32tof16(DBufferB.x);
+		PackedData[10] |= f32tof16(DBufferB.y) << 16;
+		PackedData[11] = f32tof16(DBufferB.z);
+		PackedData[11] |= f32tof16(DBufferB.w) << 16;
+	}
+
+	float4 GetDBufferC()
+	{
+		return float4(
+			f16tof32(PackedData[12]),
+			f16tof32(PackedData[12] >> 16),
+			f16tof32(PackedData[13]),
+			f16tof32(PackedData[13] >> 16)
+		);
+	}
+
+	void SetDBufferC(float4 DBufferC)
+	{
+		PackedData[12] = f32tof16(DBufferC.x);
+		PackedData[12] |= f32tof16(DBufferC.y) << 16;
+		PackedData[13] = f32tof16(DBufferC.z);
+		PackedData[13] |= f32tof16(DBufferC.w) << 16;
+	}
 #endif
 };
 
@@ -246,8 +315,14 @@ struct FPathTracingPayload : FMinimalPayload
 	void SetFrontFace() { Flags |= PATH_TRACING_PAYLOAD_OUTPUT_FLAG_FRONT_FACE; }
 	bool IsFrontFace() { return (Flags & PATH_TRACING_PAYLOAD_OUTPUT_FLAG_FRONT_FACE) != 0; }
 
-	void SetDecalReceiver() { Flags |= PATH_TRACING_PAYLOAD_OUTPUT_FLAG_DECAL_RECEIVER; }
-	bool IsDecalReceiver() { return (Flags & PATH_TRACING_PAYLOAD_OUTPUT_FLAG_DECAL_RECEIVER) != 0; }
+	void SetDecalReceiver(uint DecalResponseMask) { Flags |= (DecalResponseMask & 0x07) << PATH_TRACING_PAYLOAD_OUTPUT_FLAG_DECAL_RESPONSE_MASK_SHIFT; }
+	bool IsDecalReceiver() { return (Flags & (PATH_TRACING_PAYLOAD_OUTPUT_FLAG_DECAL_RESPONSE_COLOR | PATH_TRACING_PAYLOAD_OUTPUT_FLAG_DECAL_RESPONSE_NORMAL | PATH_TRACING_PAYLOAD_OUTPUT_FLAG_DECAL_RESPONSE_ROUGHNESS)) != 0; }
+	bool HasDecalResponseColor() { return (Flags & PATH_TRACING_PAYLOAD_OUTPUT_FLAG_DECAL_RESPONSE_COLOR) != 0; }
+	bool HasDecalResponseNormal() { return (Flags & PATH_TRACING_PAYLOAD_OUTPUT_FLAG_DECAL_RESPONSE_NORMAL) != 0; }
+	bool HasDecalResponseRoughness() { return (Flags & PATH_TRACING_PAYLOAD_OUTPUT_FLAG_DECAL_RESPONSE_ROUGHNESS) != 0; }
+
+	void SetUseDBufferLookup() { Flags |= PATH_TRACING_PAYLOAD_OUTPUT_FLAG_USE_DBUFFER_LOOKUP; }
+	bool UsesDBufferLookup() { return (Flags & PATH_TRACING_PAYLOAD_OUTPUT_FLAG_USE_DBUFFER_LOOKUP) != 0; }
 	
 	void SetHoldout() { Flags |= PATH_TRACING_PAYLOAD_OUTPUT_FLAG_IS_HOLDOUT; }
 	bool IsHoldout() { return (Flags & PATH_TRACING_PAYLOAD_OUTPUT_FLAG_IS_HOLDOUT) != 0; }
diff --git a/Engine/Shaders/Private/PathTracing/PathTracingCore.ush b/Engine/Shaders/Private/PathTracing/PathTracingCore.ush
--- a/Engine/Shaders/Private/PathTracing/PathTracingCore.ush
+++ b/Engine/Shaders/Private/PathTracing/PathTracingCore.ush
@@ -50,6 +50,7 @@ uint EnableCameraBackfaceCulling;
 uint SamplerType;
 uint VisualizeLightGrid;
 uint VisualizeDecalGrid;
+uint EnableDBuffer;
 uint ApplyDiffuseSpecularOverrides;
 uint EnabledDirectLightingContributions;
 uint EnabledIndirectLightingContributions;
@@ -200,14 +201,23 @@ struct FDecalData
 
 float ApplyDecal(inout FPathTracingPayload HitPayload, FDecalData DecalData)
 {
-	HitPayload.BaseColor = HitPayload.BaseColor * DecalData.BaseColor.a + DecalData.BaseColor.rgb;
+	if (!EnableDBuffer || HitPayload.HasDecalResponseColor())
+	{
+		HitPayload.BaseColor = HitPayload.BaseColor * DecalData.BaseColor.a + DecalData.BaseColor.rgb;
+	}
 
-	// We normalize the normal to get smoother visual result (it helps to avoid having D_GGX explodes toward infinity, and matches ApplyDBufferData(...))
-	HitPayload.WorldNormal = normalize(HitPayload.WorldNormal * DecalData.WorldNormal.a + DecalData.WorldNormal.rgb);
+	if (!EnableDBuffer || HitPayload.HasDecalResponseNormal())
+	{
+		// We normalize the normal to get smoother visual result (it helps to avoid having D_GGX explodes toward infinity, and matches ApplyDBufferData(...))
+		HitPayload.WorldNormal = normalize(HitPayload.WorldNormal * DecalData.WorldNormal.a + DecalData.WorldNormal.rgb);
+	}
 
-	HitPayload.Metallic = HitPayload.Metallic * DecalData.MetallicSpecularRoughness.a + DecalData.MetallicSpecularRoughness.x;
-	HitPayload.Specular = HitPayload.Specular * DecalData.MetallicSpecularRoughness.a + DecalData.MetallicSpecularRoughness.y;
-	HitPayload.Roughness = HitPayload.Roughness * DecalData.MetallicSpecularRoughness.a + DecalData.MetallicSpecularRoughness.z;
+	if (!EnableDBuffer || HitPayload.HasDecalResponseRoughness())
+	{
+		HitPayload.Metallic = HitPayload.Metallic * DecalData.MetallicSpecularRoughness.a + DecalData.MetallicSpecularRoughness.x;
+		HitPayload.Specular = HitPayload.Specular * DecalData.MetallicSpecularRoughness.a + DecalData.MetallicSpecularRoughness.y;
+		HitPayload.Roughness = HitPayload.Roughness * DecalData.MetallicSpecularRoughness.a + DecalData.MetallicSpecularRoughness.z;
+	}
 
 	HitPayload.Radiance += DecalData.Emissive;
 
@@ -344,6 +354,9 @@ FPathTracingPayload TraceTransparentRay(inout FPathState PathState, bool bIsCame
 	for (;;)
 	{
 		FPackedPathTracingPayload PackedPayload = InitPathTracingPayload(PathState.FirstScatterType, ApproximateCaustics ? PathState.PathRoughness : 0.0);
+		PackedPayload.SetDBufferA(float4(0, 0, 0, 1));
+		PackedPayload.SetDBufferB(float4(0, 0, 0, 1));
+		PackedPayload.SetDBufferC(float4(0, 0, 0, 1));
 		// Trace the ray
 		TraceRay(
 			TLAS,
@@ -364,6 +377,29 @@ FPathTracingPayload TraceTransparentRay(inout FPathState PathState, bool bIsCame
 		if (!PackedPayload.IsMiss() && PackedPayload.IsDecalReceiver())
 		{
 			EvaluateDecals(PathState, InstanceInclusionMask, PackedPayload, DecalData);
+
+#if USE_DBUFFER
+			const float DecalThroughput = min3(DecalData.BaseColor.a, DecalData.WorldNormal.a, DecalData.MetallicSpecularRoughness.a);
+			if (EnableDBuffer && PackedPayload.UsesDBufferLookup() && DecalThroughput < 1.0f)
+			{
+				// Retrace ray with DBuffer data
+
+				PackedPayload = InitPathTracingPayload(PathState.FirstScatterType, ApproximateCaustics ? PathState.PathRoughness : 0.0);
+
+				PackedPayload.SetDBufferA(DecalData.BaseColor);
+				PackedPayload.SetDBufferB(DecalData.WorldNormal);
+				PackedPayload.SetDBufferC(DecalData.MetallicSpecularRoughness);
+				TraceRay(
+					TLAS,
+					RayFlags,
+					InstanceInclusionMask,
+					RAY_TRACING_SHADER_SLOT_MATERIAL,
+					RAY_TRACING_NUM_SHADER_SLOTS,
+					MissShaderIndex,
+					PathState.Ray.GetNativeDesc(),
+					PackedPayload);
+			}
+#endif
 		}
 
 		if (PackedPayload.IsMiss() && VolumeIsectList.HitBlocker())
diff --git a/Engine/Shaders/Private/PathTracing/PathTracingMaterialHitShader.usf b/Engine/Shaders/Private/PathTracing/PathTracingMaterialHitShader.usf
--- a/Engine/Shaders/Private/PathTracing/PathTracingMaterialHitShader.usf
+++ b/Engine/Shaders/Private/PathTracing/PathTracingMaterialHitShader.usf
@@ -102,6 +102,12 @@ RAY_TRACING_ENTRY_CLOSEST_HIT(PathTracingMaterialCHS,
 
 	CurrentPayloadInputFlags = PackedPayload.GetFlags();
 
+#if USE_DBUFFER && MATERIAL_USES_DECAL_LOOKUP
+	CurrentPayloadDBufferA = PackedPayload.GetDBufferA();
+	CurrentPayloadDBufferB = PackedPayload.GetDBufferB();
+	CurrentPayloadDBufferC = PackedPayload.GetDBufferC();
+#endif
+
 #if VF_SUPPORTS_RAYTRACING_PREPARE_MATERIAL_PIXEL_PARAMETERS
 	// this is a newer codepath that is both more flexible and allows for more direct calculation compared to the other codepath
 	// TODO: implement such a method for all vertex factories
@@ -685,7 +691,14 @@ RAY_TRACING_ENTRY_CLOSEST_HIT(PathTracingMaterialCHS,
 #if MATERIALBLENDING_SOLID || MATERIALBLENDING_MASKED
 	if ((PrimitiveFlags & PRIMITIVE_SCENE_DATA_FLAG_DECAL_RECEIVER) != 0)
 	{
-		Payload.SetDecalReceiver();
+#if USE_DBUFFER
+		Payload.SetDecalReceiver(MATERIALDECALRESPONSEMASK);
+#	if MATERIAL_USES_DECAL_LOOKUP
+		Payload.SetUseDBufferLookup();
+#	endif
+#else
+		Payload.SetDecalReceiver(0x07);
+#endif
 	}
 #endif
 
diff --git a/Engine/Source/Runtime/Engine/Private/Materials/HLSLMaterialTranslator.cpp b/Engine/Source/Runtime/Engine/Private/Materials/HLSLMaterialTranslator.cpp
--- a/Engine/Source/Runtime/Engine/Private/Materials/HLSLMaterialTranslator.cpp
+++ b/Engine/Source/Runtime/Engine/Private/Materials/HLSLMaterialTranslator.cpp
@@ -1807,6 +1807,7 @@ void FHLSLMaterialTranslator::GetMaterialEnvironment(EShaderPlatform InPlatform,
 	OutEnvironment.SetDefine(TEXT("MATERIAL_USES_ANISOTROPY"), bUsesAnisotropy && FDataDrivenShaderPlatformInfo::GetSupportsAnisotropicMaterials(InPlatform));
 
 	OutEnvironment.SetDefine(TEXT("MATERIAL_DECAL_READ_MASK"), MaterialCompilationOutput.UsedDBufferTextures);
+	OutEnvironment.SetDefine(TEXT("MATERIAL_USES_DECAL_LOOKUP"), MaterialCompilationOutput.bUsesDBufferTextureLookup);
 	OutEnvironment.SetDefine(TEXT("MATERIAL_PATH_TRACING_BUFFER_READ"), MaterialCompilationOutput.UsedPathTracingBufferTextures);
 
 	// Count the number of VTStacks (each stack will allocate a feedback slot)
@@ -7053,6 +7054,9 @@ int32 FHLSLMaterialTranslator::DBufferTextureLookup(int32 ViewportUV, uint32 DBu
 	}
 
 	MaterialCompilationOutput.SetIsDBufferTextureUsed(DBufferTextureIndex);
+	// set separate flag to indicate that material uses DBuffer lookup specifically
+	// can't rely on UsedDBufferTextures because those bits are also set depending on the default decal response behavior.
+	MaterialCompilationOutput.SetIsDBufferTextureLookupUsed(true);
 	AddEstimatedTextureSample();
 
 	return AddCodeChunk(MCT_Float4, TEXT("MaterialExpressionDBufferTextureLookup(Parameters, %s, %d)"), *CoerceParameter(BufferUV, MCT_Float2), (int)DBufferTextureIndex);
diff --git a/Engine/Source/Runtime/Engine/Public/MaterialShared.h b/Engine/Source/Runtime/Engine/Public/MaterialShared.h
--- a/Engine/Source/Runtime/Engine/Public/MaterialShared.h
+++ b/Engine/Source/Runtime/Engine/Public/MaterialShared.h
@@ -678,6 +678,7 @@ public:
 		UsedDBufferTextures(0),
 		RuntimeVirtualTextureOutputAttributeMask(0),
 		bNeedsSceneTextures(false),
+		bUsesDBufferTextureLookup(false),
 		bUsesEyeAdaptation(false),
 		bModifiesMeshPosition(false),
 		bUsesWorldPositionOffset(false),
@@ -697,6 +698,7 @@ public:
 	ENGINE_API void SetIsSceneTextureUsed(ESceneTextureId TexId) { UsedSceneTextures |= (1 << TexId); }
 
 	ENGINE_API void SetIsDBufferTextureUsed(int32 TextureIndex) { UsedDBufferTextures |= (1 << TextureIndex); }
+	ENGINE_API void SetIsDBufferTextureLookupUsed(bool bValue) { bUsesDBufferTextureLookup = bValue; }
 
 	ENGINE_API void SetIsPathTracingBufferTextureUsed(int32 TexId) { UsedPathTracingBufferTextures |= (1 << TexId); }
 
@@ -766,6 +768,9 @@ public:
 	/** true if the material needs the scene texture lookups. */
 	LAYOUT_BITFIELD(uint8, bNeedsSceneTextures, 1);
 
+	/** true if the material uses DBuffer texture lookups. */
+	LAYOUT_BITFIELD(uint8, bUsesDBufferTextureLookup, 1);
+
 	/** true if the material uses the EyeAdaptationLookup */
 	LAYOUT_BITFIELD(uint8, bUsesEyeAdaptation, 1);
 
diff --git a/Engine/Source/Runtime/Renderer/Private/PathTracing.cpp b/Engine/Source/Runtime/Renderer/Private/PathTracing.cpp
--- a/Engine/Source/Runtime/Renderer/Private/PathTracing.cpp
+++ b/Engine/Source/Runtime/Renderer/Private/PathTracing.cpp
@@ -274,6 +274,13 @@ TAutoConsoleVariable<int32> CVarPathTracingDecalGridVisualize(
 	ECVF_RenderThreadSafe
 );
 
+TAutoConsoleVariable<int32> CVarPathTracingUseDBuffer(
+	TEXT("r.PathTracing.UseDBuffer"),
+	1,
+	TEXT("Whether to support DBuffer functionality (default=1)"),
+	ECVF_RenderThreadSafe
+);
+
 TAutoConsoleVariable<float> CVarPathTracingDecalRoughnessCutoff(
 	TEXT("r.PathTracing.DecalRoughnessCutoff"),
 	0.15f,
@@ -336,6 +343,7 @@ BEGIN_SHADER_PARAMETER_STRUCT(FPathTracingData, )
 	SHADER_PARAMETER(uint32, SamplerType)
 	SHADER_PARAMETER(uint32, VisualizeLightGrid)
 	SHADER_PARAMETER(uint32, VisualizeDecalGrid)
+	SHADER_PARAMETER(uint32, EnableDBuffer)
 	SHADER_PARAMETER(uint32, EnableAtmosphere)
 	SHADER_PARAMETER(uint32, EnableFog)
 	SHADER_PARAMETER(uint32, EnabledDirectLightingContributions)   // PATHTRACER_CONTRIBUTION_*
@@ -382,6 +390,7 @@ struct FPathTracingConfig
 			PathTracingData.EnableCameraBackfaceCulling != Other.PathTracingData.EnableCameraBackfaceCulling ||
 			PathTracingData.VisualizeLightGrid != Other.PathTracingData.VisualizeLightGrid ||
 			PathTracingData.VisualizeDecalGrid != Other.PathTracingData.VisualizeDecalGrid ||
+			PathTracingData.EnableDBuffer != Other.PathTracingData.EnableDBuffer ||
 			PathTracingData.MaxPathIntensity != Other.PathTracingData.MaxPathIntensity ||
 			PathTracingData.FilterWidth != Other.PathTracingData.FilterWidth ||
 			PathTracingData.EnableAtmosphere != Other.PathTracingData.EnableAtmosphere ||
@@ -547,6 +556,8 @@ static void PreparePathTracingData(const FScene* Scene, const FViewInfo& View, F
 		&& (Scene->ExponentialFogs[0].FogData[0].Density > 0 ||
 			Scene->ExponentialFogs[0].FogData[1].Density > 0);
 
+	PathTracingData.EnableDBuffer = CVarPathTracingUseDBuffer.GetValueOnRenderThread();
+
 	PathTracingData.DecalRoughnessCutoff = PathTracing::UsesDecals(*View.Family) && View.bHasRayTracingDecals ? CVarPathTracingDecalRoughnessCutoff.GetValueOnRenderThread() : -1.0f;
 
 	PathTracingData.MeshDecalRoughnessCutoff = PathTracing::UsesDecals(*View.Family) && Scene->RayTracingScene.GetRHIRayTracingScene()->GetInitializer().NumNativeInstancesPerLayer[(uint32)ERayTracingSceneLayer::Decals] > 0 ? CVarPathTracingMeshDecalRoughnessCutoff.GetValueOnRenderThread() : -1.0f;
