Fix shader compilation issue on switch... but also fix shader being compiled for switch unnecessarily.
Fix minor typos.

diff --git a/Engine/Shaders/Private/ForwardLightingCommon.ush b/Engine/Shaders/Private/ForwardLightingCommon.ush
--- a/Engine/Shaders/Private/ForwardLightingCommon.ush
+++ b/Engine/Shaders/Private/ForwardLightingCommon.ush
@@ -30,9 +30,11 @@
 
 #if VIRTUAL_SHADOW_MAP
 #include "VirtualShadowMaps/VirtualShadowMapProjectionCommon.ush"
+#if SUPPORT_VSM_FOWARD_QUALITY > 0
 #include "VirtualShadowMaps/VirtualShadowMapProjectionFilter.ush"
 #include "VirtualShadowMaps/VirtualShadowMapProjectionDirectional.ush"
-#endif
+#endif // SUPPORT_VSM_FOWARD_QUALITY
+#endif // VIRTUAL_SHADOW_MAP
 
 #ifndef ENABLE_FORWARD_CLOUD_SHADOW
 #define ENABLE_FORWARD_CLOUD_SHADOW 0
@@ -193,7 +195,7 @@ FDeferredLightingSplit GetForwardDirectLightingSplit(
 				#if SUPPORT_VSM_FOWARD_QUALITY == 1
 				{
 					const FLightShaderParameters DireciontalLight = ConvertToLightShaderParameters(LightData);
-					FVirtualShadowMapSMRTSettings SMRTSettings = GetDirectionalLighSMRTSettings();
+					FVirtualShadowMapSMRTSettings SMRTSettings = GetDirectionalLightSMRTSettings();
 
 					const bool bBackfaceCull = !IsSubsurfaceModel(GBufferData.ShadingModelID);
 					const float ScreenRayLengthWorld = SMRTSettings.ScreenRayLength * View.ClipToView[1][1] * GBufferData.Depth;
diff --git a/Engine/Shaders/Private/LightGridCommon.ush b/Engine/Shaders/Private/LightGridCommon.ush
--- a/Engine/Shaders/Private/LightGridCommon.ush
+++ b/Engine/Shaders/Private/LightGridCommon.ush
@@ -161,7 +161,7 @@ struct FVirtualShadowMapSMRTSettings
 	uint  AdaptiveRayCount;
 };
 
-FVirtualShadowMapSMRTSettings GetDirectionalLighSMRTSettings()
+FVirtualShadowMapSMRTSettings GetDirectionalLightSMRTSettings()
 {
 	FVirtualShadowMapSMRTSettings Out;
 	Out.ScreenRayLength			= ForwardLightData.DirectionalLightSMRTSettings.ScreenRayLength;
diff --git a/Engine/Shaders/Private/Strata/StrataForwardLighting.ush b/Engine/Shaders/Private/Strata/StrataForwardLighting.ush
--- a/Engine/Shaders/Private/Strata/StrataForwardLighting.ush
+++ b/Engine/Shaders/Private/Strata/StrataForwardLighting.ush
@@ -9,9 +9,11 @@
 #if VIRTUAL_SHADOW_MAP
 #include "../BlueNoise.ush"
 #include "../VirtualShadowMaps/VirtualShadowMapProjectionCommon.ush"
+if SUPPORT_VSM_FOWARD_QUALITY > 0
 #include "../VirtualShadowMaps/VirtualShadowMapProjectionFilter.ush"
 #include "../VirtualShadowMaps/VirtualShadowMapProjectionDirectional.ush"
-#endif
+#endif // SUPPORT_VSM_FOWARD_QUALITY
+#endif // VIRTUAL_SHADOW_MAP
 
 #define FORWARD_PER_PIXEL_SHADING (FORWARD_SHADING || TRANSLUCENCY_LIGHTING_SURFACE_FORWARDSHADING || MATERIAL_SHADINGMODEL_SINGLELAYERWATER)
 
@@ -119,7 +121,7 @@ float3 StrataForwardLighting(
 		#if SUPPORT_VSM_FOWARD_QUALITY == 1
 		{
 			const FLightShaderParameters DireciontalLight = ConvertToLightShaderParameters(DirLightData);
-			FVirtualShadowMapSMRTSettings SMRTSettings = GetDirectionalLighSMRTSettings();
+			FVirtualShadowMapSMRTSettings SMRTSettings = GetDirectionalLightSMRTSettings();
 
 			const bool bBackfaceCull = true; // STRATA_TODO: this needs to be fixed but require to extract hair/SSS/diffusebackface prior to BSDF looping
 			const float ScreenRayLengthWorld = SMRTSettings.ScreenRayLength * View.ClipToView[1][1] * SceneDepth;
diff --git a/Engine/Shaders/Private/VirtualShadowMaps/VirtualShadowMapProjectionDirectional.ush b/Engine/Shaders/Private/VirtualShadowMaps/VirtualShadowMapProjectionDirectional.ush
--- a/Engine/Shaders/Private/VirtualShadowMaps/VirtualShadowMapProjectionDirectional.ush
+++ b/Engine/Shaders/Private/VirtualShadowMaps/VirtualShadowMapProjectionDirectional.ush
@@ -331,14 +331,14 @@ FVirtualShadowMapSampleResult TraceDirectional(
 			}
 			#endif
 		}
-		uint RayCount = min(i + 1, MaxRayCount);		// break vs regular for loop exit
+		uint RayCount = min(i + 1U, MaxRayCount);		// break vs regular for loop exit
 
 		/* TODO: Experiment with max vs avg distance in various cases
 		float OccluderDistance = (View.GeneralPurposeTweak == 0)
 			? MaxOccluderDistance
 			: (OccluderDistanceSum / float(max(1, RayCount - RayMissCount)));
 		*/
-		float OccluderDistance = (OccluderDistanceSum / float(max(1, RayCount - RayMissCount)));
+		float OccluderDistance = (OccluderDistanceSum / float(max(1U, RayCount - RayMissCount)));
 
 		Result.ShadowFactor = float(RayMissCount) / float(RayCount);
 		Result.OccluderDistance = OccluderDistance;
diff --git a/Engine/Source/Runtime/Renderer/Private/LightRendering.cpp b/Engine/Source/Runtime/Renderer/Private/LightRendering.cpp
--- a/Engine/Source/Runtime/Renderer/Private/LightRendering.cpp
+++ b/Engine/Source/Runtime/Renderer/Private/LightRendering.cpp
@@ -842,7 +842,7 @@ class FDeferredLightPS : public FGlobalShader
 			}
 		}
 
-		if (!DoesPlatformSupportNanite(Parameters.Platform, false) && PermutationVector.Get<FVirtualShadowMapMask>() != 0)
+		if (!DoesPlatformSupportVirtualShadowMaps(Parameters.Platform) && PermutationVector.Get<FVirtualShadowMapMask>() != 0)
 		{
 			return false;
 		}
diff --git a/Engine/Source/Runtime/Renderer/Private/TranslucentLighting.cpp b/Engine/Source/Runtime/Renderer/Private/TranslucentLighting.cpp
--- a/Engine/Source/Runtime/Renderer/Private/TranslucentLighting.cpp
+++ b/Engine/Source/Runtime/Renderer/Private/TranslucentLighting.cpp
@@ -707,6 +707,13 @@ public:
 	  */
 	static bool ShouldCompilePermutation(const FMaterialShaderPermutationParameters& Parameters)
 	{
+		FPermutationDomain PermutationVector(Parameters.PermutationId);
+
+		if (!DoesPlatformSupportVirtualShadowMaps(Parameters.Platform) && PermutationVector.Get<FVirtualShadowMap>() != 0)
+		{
+			return false;
+		}
+
 		return (Parameters.MaterialParameters.MaterialDomain == MD_LightFunction || Parameters.MaterialParameters.bIsSpecialEngineMaterial) &&
 			(IsFeatureLevelSupported(Parameters.Platform, ERHIFeatureLevel::SM5) &&
 			(RHISupportsGeometryShaders(Parameters.Platform) || RHISupportsVertexShaderLayer(Parameters.Platform)));
