Fix typo

diff --git a/Engine/Shaders/Private/Strata/StrataForwardLighting.ush b/Engine/Shaders/Private/Strata/StrataForwardLighting.ush
--- a/Engine/Shaders/Private/Strata/StrataForwardLighting.ush
+++ b/Engine/Shaders/Private/Strata/StrataForwardLighting.ush
@@ -9,7 +9,7 @@
 #if VIRTUAL_SHADOW_MAP
 #include "../BlueNoise.ush"
 #include "../VirtualShadowMaps/VirtualShadowMapProjectionCommon.ush"
-if SUPPORT_VSM_FOWARD_QUALITY > 0
+#if SUPPORT_VSM_FOWARD_QUALITY > 0
 #include "../VirtualShadowMaps/VirtualShadowMapProjectionFilter.ush"
 #include "../VirtualShadowMaps/VirtualShadowMapProjectionDirectional.ush"
 #endif // SUPPORT_VSM_FOWARD_QUALITY
