Fixes TSR DilateVelocity's CONFIG_AA_CROSS_OPTIMISATION that could sample depth outside viewports

diff --git a/Engine/Shaders/Private/TemporalSuperResolution/TSRDilateVelocity.usf b/Engine/Shaders/Private/TemporalSuperResolution/TSRDilateVelocity.usf
--- a/Engine/Shaders/Private/TemporalSuperResolution/TSRDilateVelocity.usf
+++ b/Engine/Shaders/Private/TemporalSuperResolution/TSRDilateVelocity.usf
@@ -121,10 +121,14 @@ void MainCS(
 	#endif
 	ISOLATE
 	{
-		DeviceZ = SceneDepthTexture[InputPixelPos];
-		#if CONFIG_SUBPIXEL_DEPTH
-			PrevCompressedSubpixelDepth = PrevScatteredSubpixelDepthTexture[InputPixelPos];
-		#endif
+		{
+			int2 NeighborPixelPos = clamp(int2(InputPixelPos), InputPixelPosMin, InputPixelPosMax);
+				
+			DeviceZ = SceneDepthTexture[NeighborPixelPos];
+			#if CONFIG_SUBPIXEL_DEPTH
+				PrevCompressedSubpixelDepth = PrevScatteredSubpixelDepthTexture[NeighborPixelPos];
+			#endif
+		}
 		#if CONFIG_AA_CROSS_OPTIMISATION
 			UNROLL_N(3)
 			for (uint i = 0; i < 3; i++)
