Support virtual texture feedback behind a single layer water material.
We need to stochastically drop some writes to the feedback buffer so that we don't overwrite all of the content from the surface beneath.

diff --git a/Engine/Shaders/Private/VirtualTextureCommon.ush b/Engine/Shaders/Private/VirtualTextureCommon.ush
--- a/Engine/Shaders/Private/VirtualTextureCommon.ush
+++ b/Engine/Shaders/Private/VirtualTextureCommon.ush
@@ -65,10 +65,10 @@ void FinalizeVirtualTextureFeedback(in FVirtualTextureFeedbackParams Params, flo
 
 		// When using DBuffer, feedback buffer is blending using decal opacity for opaque objects
 		// This way we avoid overwriting values already written to feedback buffer while rendering dbuffer decals
-#if MATERIALBLENDING_TRANSLUCENT || MATERIALBLENDING_ADDITIVE || MATERIALBLENDING_MODULATE || IS_DECAL || (USE_DBUFFER && MATERIAL_DECAL_READ_MASK)
+#if MATERIALBLENDING_TRANSLUCENT || MATERIALBLENDING_ADDITIVE || MATERIALBLENDING_MODULATE || MATERIAL_SHADINGMODEL_SINGLELAYERWATER || IS_DECAL || (USE_DBUFFER && MATERIAL_DECAL_READ_MASK)
 
 		// Stochastic alpha testing
-#if MATERIALBLENDING_ADDITIVE || MATERIALBLENDING_MODULATE || MATERIAL_SHADINGMODEL_THIN_TRANSLUCENT
+#if MATERIALBLENDING_ADDITIVE || MATERIALBLENDING_MODULATE || MATERIAL_SHADINGMODEL_THIN_TRANSLUCENT || MATERIAL_SHADINGMODEL_SINGLELAYERWATER
 		// Opacity may not make sense for additive and modulation blending
 		// Use a value of 0.25f so we have up to 4 layers of VT transparancy
 		const float Alpha = 0.25f;  
