Fix a bug where reading custom data in a Nanite material to calculate World Position Offset would always result in using the default value.

diff --git a/Engine/Shaders/Private/MaterialTemplate.ush b/Engine/Shaders/Private/MaterialTemplate.ush
--- a/Engine/Shaders/Private/MaterialTemplate.ush
+++ b/Engine/Shaders/Private/MaterialTemplate.ush
@@ -1157,7 +1157,7 @@ float GetPerInstanceCustomData(FMaterialPixelParameters Parameters, int Index, f
 /** Get the per-instance custom data when instancing */
 float GetPerInstanceCustomData(FMaterialVertexParameters Parameters, int Index, float DefaultValue)
 {
-#if USE_INSTANCE_CULLING && USES_PER_INSTANCE_CUSTOM_DATA
+#if VF_USE_PRIMITIVE_SCENE_DATA && USES_PER_INSTANCE_CUSTOM_DATA
 	const uint FloatIndex = uint(Index);
 	BRANCH
 	if (FloatIndex < Parameters.CustomDataCount)
