Path Tracer: Improve SSS random walk sampling

Implement the Dwivedi sampling scheme to favor random walks that go back to the surface. To better handle thin regions, a slab thickness is measured first so that we can guide either forwards or backwards in the medium, depending on the position in the virtual slab.

Refactor the interface counting logic so we can trace SSS probe rays more easily (in particular, re-use the same logic for estimating the thickness of the object).

Since the overall change to the random walk algorithm is small, I am keeping the old code around for comparison via ifdefs. This will come in handy for debugging in the future.

I've expose to SSSGuidingRatio for debugging purposes via a cvar. It is potentially useful for end-users to opt-out of the new random walk in cases where it does not help.

diff --git a/Engine/Shaders/Private/PathTracing/PathTracingCore.ush b/Engine/Shaders/Private/PathTracing/PathTracingCore.ush
--- a/Engine/Shaders/Private/PathTracing/PathTracingCore.ush
+++ b/Engine/Shaders/Private/PathTracing/PathTracingCore.ush
@@ -35,6 +35,7 @@ uint TemporalSeed;
 uint MaxSamples;
 uint MaxBounces;
 uint MaxSSSBounces;
+float SSSGuidingRatio;
 
 // 0: only Material sampling
 // 1: only Light sampling
@@ -765,42 +766,67 @@ struct FProbeResult
 	bool IsMiss() { return HitT <= 0; }
 };
 
-FProbeResult TraceProbeRay(FRayDesc Ray)
+FProbeResult TraceSSSProbeRay(FRayDesc Ray, inout int InterfaceCounter)
 {
-	// Trace a short ray to see if we escaped the surface
-	// NOTE: we leave these rays tagged as "camera" because we are interested in getting the same shading normal the camera saw (on the off chance the material does some ray-simplfiication)
-	//      Since the SSS code mostly runs on camera rays, this is a reasonable approximation
-	FPackedPathTracingPayload PackedPayload = InitPathTracingPayload(PATHTRACER_SCATTER_CAMERA, 1.0);
-	// TODO: use smaller payload + flag to skip shading work? (only want smooth+geo normals)
-	//       however supporting blockers would require some shading ...
-	const uint RayFlags = 0;
-	const uint InstanceInclusionMask = PATHTRACER_MASK_ALL
-										- PATHTRACER_MASK_HAIR_SHADOW 
-										- PATHTRACER_MASK_HAIR_INDIRECT
-										- PATHTRACER_MASK_HAIR_CAMERA; // Ignore hair strands for SSS rays
-	const uint MissShaderIndex = 0;
-	TraceRay(
-		TLAS,
-		RayFlags,
-		InstanceInclusionMask,
-		RAY_TRACING_SHADER_SLOT_MATERIAL,
-		RAY_TRACING_NUM_SHADER_SLOTS,
-		MissShaderIndex,
-		Ray.GetNativeDesc(),
-		PackedPayload);
-
-	if (PackedPayload.IsMiss())
+	// Use the technique mentioned in "A Hero Beneath the Surface" (Section 6.3.2)
+	// https://jo.dreggn.org/home/2021_spectral_imaging.pdf
+	// This allows multiple overlapping meshes to be treated as a single "volume" by discarding internal hits until we have
+	// crossed the appropriate number of interfaces
+#define SSS_USE_INTERFACE_COUNTING 1
+
+#if SSS_USE_INTERFACE_COUNTING
+	for (;;)
+#endif
 	{
-		return (FProbeResult)0;
+		// Trace a short ray to see if we escaped the surface
+		// NOTE: we leave these rays tagged as "camera" because we are interested in getting the same shading normal the camera saw (on the off chance the material does some ray-simplfiication)
+		//      Since the SSS code mostly runs on camera rays, this is a reasonable approximation
+		FPackedPathTracingPayload PackedPayload = InitPathTracingPayload(PATHTRACER_SCATTER_CAMERA, 1.0);
+		// TODO: use smaller payload + flag to skip shading work? (only want smooth+geo normals)
+		//       however supporting blockers would require some shading ...
+		//       
+		const uint RayFlags = 0;
+		const uint InstanceInclusionMask = PATHTRACER_MASK_ALL
+			- PATHTRACER_MASK_HAIR_SHADOW
+			- PATHTRACER_MASK_HAIR_INDIRECT
+			- PATHTRACER_MASK_HAIR_CAMERA; // Ignore hair strands for SSS rays
+		const uint MissShaderIndex = 0;
+		TraceRay(
+			TLAS,
+			RayFlags,
+			InstanceInclusionMask,
+			RAY_TRACING_SHADER_SLOT_MATERIAL,
+			RAY_TRACING_NUM_SHADER_SLOTS,
+			MissShaderIndex,
+			Ray.GetNativeDesc(),
+			PackedPayload);
+
+		if (PackedPayload.IsMiss())
+		{
+			// we did not hit anything
+			return (FProbeResult)0;
+		}
+
+#if SSS_USE_INTERFACE_COUNTING
+		// update counter
+		InterfaceCounter += PackedPayload.IsFrontFace() ? +1 : -1;
+		if (InterfaceCounter != 0)
+		{
+			// we have not yet crossed the right number of interfaces, so ignore the current hit and try the segment again
+			Ray.TMin = asfloat(asuint(PackedPayload.HitT) + 1);
+			continue;
+		}
+#endif
+
+		FPathTracingPayload Payload = UnpackPathTracingPayload(PackedPayload, Ray);
+		FProbeResult Result;
+		Result.HitT = PackedPayload.HitT;
+		Result.WorldNormal = Payload.WorldNormal;
+		Result.WorldSmoothNormal = Payload.WorldSmoothNormal;
+		Result.WorldGeoNormal = Payload.WorldGeoNormal;
+		Result.FrontFace = Payload.IsFrontFace();
+		return Result;
 	}
-	FPathTracingPayload Payload = UnpackPathTracingPayload(PackedPayload, Ray);
-	FProbeResult Result;
-	Result.HitT = PackedPayload.HitT;
-	Result.WorldNormal = Payload.WorldNormal;
-	Result.WorldSmoothNormal = Payload.WorldSmoothNormal;
-	Result.WorldGeoNormal = Payload.WorldGeoNormal;
-	Result.FrontFace = Payload.IsFrontFace();
-	return Result;
 }
 
 void ApplyRayBias(inout FRayDesc Ray, float HitT, float3 Normal)
@@ -817,6 +843,190 @@ void ApplyRayBias(inout FRayDesc Ray, float HitT, float3 Normal)
 #endif
 }
 
+float3 ComputeDwivediScale(float3 Albedo)
+{
+	// "Zero-Variance Theory for Efficient Subsurface Scattering"
+	// [REDACTED_PERSON] and Jaroslav Krivanek (SIGGRAPH 2020)
+	// [REDACTED_PROFILE_URL]
+	// Eq. 67
+	const float3 ClampedAlbedo = clamp(Albedo, 0.001, 0.999); // avoid singularities at the extremeties
+	return rsqrt(1.0 - pow(ClampedAlbedo, 2.44294 - 0.0215813 * ClampedAlbedo + 0.578637 / ClampedAlbedo));
+}
+
+float RescaleRandomNumber(float RandVal, float LowerBound, float UpperBound)
+{
+	const float OneMinusEpsilon = 0.99999994; // 32-bit float just before 1.0
+	return min((RandVal - LowerBound) / (UpperBound - LowerBound), OneMinusEpsilon);
+}
+
+float SampleGuidedSpectralTransmittance(float RandValue, float SlabCosine, float3 DwivediScale, float GuidingFactor, float3 Sigma, float3 ProbT, float3 ColorChannelPdf)
+{
+	float3 ColorChannelCdf = float3(
+		ColorChannelPdf.x,
+		ColorChannelPdf.x + ColorChannelPdf.y,
+		ColorChannelPdf.x + ColorChannelPdf.y + ColorChannelPdf.z);
+	if (ColorChannelCdf.z > 0)
+	{
+		
+		// choose guided vs non-guided
+		if (RandValue < GuidingFactor)
+		{
+			// dwivedi walk, stretch sigma
+			RandValue = RescaleRandomNumber(RandValue, 0.0, GuidingFactor);
+		}
+		else
+		{
+			// classic walk, use unmodified sigma
+			SlabCosine = 0;
+			RandValue = RescaleRandomNumber(RandValue, GuidingFactor, 1.0);
+
+		}
+		const float q = RandValue * ColorChannelCdf.z;
+		if (q < ColorChannelCdf.x)
+		{
+			const float RescaleRand1 = RescaleRandomNumber(q, 0.0, ColorChannelCdf.x);
+			const float RescaleRand2 = RescaleRand1 < ProbT.x ? RescaleRandomNumber(RescaleRand1, 0.0, ProbT.x) : RescaleRandomNumber(RescaleRand1, ProbT.x, 1.0);
+			SlabCosine *= RescaleRand1 < ProbT.x ? -1.0 : 1.0;
+			const float StretchedSigma = Sigma.x * (1 - SlabCosine / DwivediScale.x);
+			return -log(1 - RescaleRand2) / StretchedSigma;
+		}
+		else if (q < ColorChannelCdf.y)
+		{
+			const float RescaleRand1 = RescaleRandomNumber(q, ColorChannelCdf.x, ColorChannelCdf.y);
+			const float RescaleRand2 = RescaleRand1 < ProbT.y ? RescaleRandomNumber(RescaleRand1, 0.0, ProbT.y) : RescaleRandomNumber(RescaleRand1, ProbT.y, 1.0);
+			SlabCosine *= RescaleRand1 < ProbT.y ? -1.0 : 1.0;
+			const float StretchedSigma = Sigma.y * (1 - SlabCosine / DwivediScale.y);
+			return -log(1 - RescaleRand2) / StretchedSigma;
+
+		}
+		else
+		{
+			const float RescaleRand1 = RescaleRandomNumber(q, ColorChannelCdf.y, ColorChannelCdf.z);
+			const float RescaleRand2 = RescaleRand1 < ProbT.z ? RescaleRandomNumber(RescaleRand1, 0.0, ProbT.z) : RescaleRandomNumber(RescaleRand1, ProbT.z, 1.0);
+			SlabCosine *= RescaleRand1 < ProbT.z ? -1.0 : 1.0;
+			const float StretchedSigma = Sigma.z * (1 - SlabCosine / DwivediScale.z);
+			return -log(1 - RescaleRand2) / StretchedSigma;
+		}
+	}
+	// all channels have 0 probability
+	return -1.0;
+}
+
+float4 EvaluateGuidedSpectralTransmittanceHit(float SampledT, float SlabCosine, float3 DwivediScale, float GuidingFactor, float3 Sigma, float3 ProbT, float3 ColorChannelPdf)
+{
+	// normalize the pdf (to match code above)
+	ColorChannelPdf *= rcp(ColorChannelPdf.x + ColorChannelPdf.y + ColorChannelPdf.z);
+	float3 Transmittance = exp(-SampledT * Sigma);
+	// probability of reaching the sampled point with classic sampling
+	float3 TransmittancePdf = Sigma * Transmittance;
+	// probability of reaching the sampled point with guiding
+	float3 GuidedSigmaR = (1 - SlabCosine / DwivediScale) * Sigma;
+	float3 GuidedSigmaT = (1 + SlabCosine / DwivediScale) * Sigma;
+	float3 GuidedPdfR = GuidedSigmaR * exp(-SampledT * GuidedSigmaR);
+	float3 GuidedPdfT = GuidedSigmaT * exp(-SampledT * GuidedSigmaT);
+	float3 GuidedPdf = lerp(GuidedPdfR, GuidedPdfT, ProbT);
+	float MisPdf = dot(ColorChannelPdf, lerp(TransmittancePdf, GuidedPdf, GuidingFactor));
+	return MisPdf > 0 ? float4(Transmittance / MisPdf, MisPdf) : 0.0;
+}
+
+float4 EvaluateGuidedSpectralTransmittanceMiss(float MaxT, float SlabCosine, float3 DwivediScale, float GuidingFactor, float3 Sigma, float3 ProbT, float3 ColorChannelPdf)
+{
+	// normalize the pdf (to match code above)
+	ColorChannelPdf *= rcp(ColorChannelPdf.x + ColorChannelPdf.y + ColorChannelPdf.z);
+	float3 Transmittance = exp(-MaxT * Sigma);
+	float3 TransmittancePdf = Transmittance; // probability of going past MaxT with classic sampling
+
+	// probability of going past MaxT (integral of the pdf from MaxT to infinity) with guiding
+	float3 GuidedSigmaR = (1 - SlabCosine / DwivediScale) * Sigma;
+	float3 GuidedSigmaT = (1 + SlabCosine / DwivediScale) * Sigma;
+	float3 GuidedPdfR = exp(-MaxT * GuidedSigmaR);
+	float3 GuidedPdfT = exp(-MaxT * GuidedSigmaT);
+	float3 GuidedPdf = lerp(GuidedPdfR, GuidedPdfT, ProbT);
+	
+	float MisPdf = dot(ColorChannelPdf, lerp(TransmittancePdf, GuidedPdf, GuidingFactor));
+	return MisPdf > 0 ? float4(Transmittance / MisPdf, MisPdf) : 0.0;
+}
+
+// returns: xyz: world space direction, w: throughput scale
+float4 SampleDwivediPhaseFunction(float3 ColorChannelPdf, float3 DwivediScale, float GuidingFraction, float3 ProbT, float3 DwivediSlabNormal, float3 RayDirection, float G, float2 RandSample)
+{
+	float4 Result = 0;
+	float3 ColorChannelCdf = float3(
+		ColorChannelPdf.x,
+		ColorChannelPdf.x + ColorChannelPdf.y,
+		ColorChannelPdf.x + ColorChannelPdf.y + ColorChannelPdf.z);
+	if (ColorChannelCdf.z > 0)
+	{
+		const float3 PhaseLog = log((DwivediScale + 1.0) / (DwivediScale - 1.0));
+		const float OneMinusEpsilon = 0.99999994; // 32-bit float just before 1.0
+		if (RandSample.x < GuidingFraction)
+		{
+			// sample the dwivedi guiding pdf
+			RandSample.x = RescaleRandomNumber(RandSample.x, 0.0, GuidingFraction);
+			const float q = RandSample.x * ColorChannelCdf.z;
+			float CosineZ = 0;
+			float Sign = 1;
+			if (q < ColorChannelCdf.x)
+			{
+				const float RescaleRand1 = RescaleRandomNumber(q, 0.0, ColorChannelCdf.x);
+				const float RescaleRand2 = RescaleRand1 < ProbT.x ? RescaleRandomNumber(RescaleRand1, 0.0, ProbT.x) : RescaleRandomNumber(RescaleRand1, ProbT.x, 1.0);
+				CosineZ = (DwivediScale.x - (DwivediScale.x + 1) * exp(-RescaleRand2 * PhaseLog.x));
+				Sign = RescaleRand1 < ProbT.x ? -1.0 : +1.0;
+			}
+			else if (q < ColorChannelCdf.y)
+			{
+				const float RescaleRand1 = RescaleRandomNumber(q, ColorChannelCdf.x, ColorChannelCdf.y);
+				const float RescaleRand2 = RescaleRand1 < ProbT.y ? RescaleRandomNumber(RescaleRand1, 0.0, ProbT.y) : RescaleRandomNumber(RescaleRand1, ProbT.y, 1.0);
+				CosineZ = (DwivediScale.y - (DwivediScale.y + 1) * exp(-RescaleRand2 * PhaseLog.y));
+				Sign = RescaleRand1 < ProbT.y ? -1.0 : +1.0;
+			}
+			else
+			{
+				const float RescaleRand1 = RescaleRandomNumber(q, ColorChannelCdf.y, ColorChannelCdf.z);
+				const float RescaleRand2 = RescaleRand1 < ProbT.z ? RescaleRandomNumber(RescaleRand1, 0.0, ProbT.z) : RescaleRandomNumber(RescaleRand1, ProbT.z, 1.0);
+				CosineZ = (DwivediScale.z - (DwivediScale.z + 1) * exp(-RescaleRand2 * PhaseLog.z));
+				Sign = RescaleRand1 < ProbT.z ? -1.0 : +1.0;
+
+			}
+			float3 PhasePdfR = rcp((DwivediScale - CosineZ) * PhaseLog * (2 * PI));
+			float3 PhasePdfT = rcp((DwivediScale + CosineZ) * PhaseLog * (2 * PI));
+			float3 PhasePdf = lerp(PhasePdfR, PhasePdfT, ProbT);
+			ColorChannelPdf *= rcp(ColorChannelCdf.z);
+			float MisPdf = dot(ColorChannelPdf, PhasePdf);
+
+
+			float SineZ = sqrt(saturate(1 - CosineZ * CosineZ));
+			float Phi = (2 * PI) * RandSample.y;
+			// produce output direction in slab tangent frame
+			Result.xyz = normalize(TangentToWorld(float3(SineZ * cos(Phi), SineZ * sin(Phi), Sign * CosineZ), DwivediSlabNormal));
+			// final throughput is the phase function divided by pdf
+			float PhaseCosine = -dot(RayDirection, Result.xyz);
+			float PhaseEval = HenyeyGreensteinPhase(G, PhaseCosine);
+			Result.w = PhaseEval / lerp(PhaseEval, MisPdf, GuidingFraction);
+		}
+		else
+		{
+			// sample the classic HG lobe directly (but include probability of sampling the guided lobe)
+			RandSample.x = RescaleRandomNumber(RandSample.x, GuidingFraction, 1.0);
+			float4 DirectionAndPhase = ImportanceSampleHenyeyGreensteinPhase(RandSample, G);
+			Result.xyz = normalize(TangentToWorld(DirectionAndPhase.xyz, RayDirection));
+
+			float CosineZ = dot(Result.xyz, DwivediSlabNormal);
+
+			float3 GuidedPhasePdfR = rcp((DwivediScale - CosineZ) * PhaseLog * (2 * PI));
+			float3 GuidedPhasePdfT = rcp((DwivediScale + CosineZ) * PhaseLog * (2 * PI));
+			float3 GuidedPhasePdf = lerp(GuidedPhasePdfR, GuidedPhasePdfT, ProbT);
+			ColorChannelPdf *= rcp(ColorChannelCdf.z);
+			float MisPdf = dot(ColorChannelPdf, GuidedPhasePdf);
+
+			// final throughput is the phase function divided by pdf
+			float PhaseEval = DirectionAndPhase.w;
+			Result.w = PhaseEval / lerp(PhaseEval, MisPdf, GuidingFraction);
+		}
+	}
+	return Result;
+}
+
 bool ProcessSubsurfaceRandomWalk(inout FPathTracingPayload Payload, inout float3 PathThroughput, inout RandomSequence RandSequence, float3 RayDirection, bool SimplifySSS)
 {
 	if (!Payload.IsSubsurfaceMaterial())
@@ -863,11 +1073,6 @@ bool ProcessSubsurfaceRandomWalk(inout FPathTracingPayload Payload, inout float3
 	// Support SSS for mfp smaller than 1mm.
 	SSS.Radius = max(SSS.Radius, 0.0009); // 0.009mm minimum to make sure random walk can move forward and to keep SigmaT finite
 
-    // Use the technique mentioned in "A Hero Beneath the Surface" (Section 6.3.2)
-	// https://jo.dreggn.org/home/2021_spectral_imaging.pdf
-	// This allows multiple overlapping meshes to be treated as a single "volume" by discarding internal hits until we have
-	// crossed the appropriate number of interfaces
-#define SSS_USE_INTERFACE_COUNTING 1
 
 	int InterfaceCounter = Payload.IsFrontFace() ? +1 : -1;
 
@@ -887,98 +1092,173 @@ bool ProcessSubsurfaceRandomWalk(inout FPathTracingPayload Payload, inout float3
 	SSS.Radius *= 2.0; // roughly match parameterization above
 #endif
 
-	// Apply albedo remapping for anisotropic scattering 
-	float G = Payload.GetSubsurfacePhaseFunction();
+	// Subsurface guiding is implemented following the Dwivedi random walk technique described here:
+	// https://cgg.mff.cuni.cz/~jaroslav/papers/2014-zerovar/
+	// [REDACTED_PROFILE_URL]
+	// A thin-slab approximation is used to improve the guiding in thin regions as well as described in the video presentation (slides 37-39).
+#define SSS_USE_DWIVEDI 1
+#define SSS_USE_DWIVEDI_USE_THIN_SLABS 1 // Probe the geometry to have an estimate of thickness - and use this to guide toward front or backside, depending on which is closer
 
 	// Revisiting Physically Based Shading at Imageworks.
 	// https://blog.selfshadow.com/publications/s2017-shading-course/imageworks/s2017_pbs_imageworks_slides_v2.pdf
+	float G = Payload.GetSubsurfacePhaseFunction();
 	Albedo = Albedo / (1 - G * (1 - Albedo));
 
+#if SSS_USE_DWIVEDI
+	// Setup an oriented slab to approximate the surface. This is used to guide paths back towards the surface so they can escape
+	const float3 DwivediScale = ComputeDwivediScale(Albedo);
+	float3 DwivediSlabNormal = Payload.WorldSmoothNormal;
+	float3 DwivediSlabOrigin = Payload.TranslatedWorldPos;
+	// Guiding only works for isotropic phase functions - limit the guiding to a narrow range of phase functions around G=0.0
+	// This heuristic was compared to the one from:
+	// https://jo.dreggn.org/home/2016_dwivedi.pdf (see Equation 15)
+	// but the latter does not appear sufficient for G > 0.25 or so. This is potentially explained by the lack of Closest point or incident illumination guiding.
+	const float GuidedRatio = SSSGuidingRatio * (1.0 - pow(saturate(abs(G * 4)), 0.0625));
+
+#if SSS_USE_DWIVEDI_USE_THIN_SLABS
+	bool bDoSlabSearch = GuidedRatio > 0;
+	float SlabThickness = -1.0; // negative means we didn't find a valid thickness
+#endif
+
+#endif // SSS_USE_DWIVEDI
+
+
 	const int MAX_SSS_BOUNCES = MaxSSSBounces;
 	const float3 SigmaT = rcp(SSS.Radius);
 	const float3 SigmaS = Albedo * SigmaT;
 	for (int i = 0; i < MAX_SSS_BOUNCES; i++)
 	{
-		// TODO: add Dwivedi sampling for faster escape in low albedo media ...
-
 		float3 ColorChannelPdf = PathThroughput * Albedo;
+#if SSS_USE_DWIVEDI
+		float SlabCosine = dot(Ray.Direction, DwivediSlabNormal);
+#if SSS_USE_DWIVEDI_USE_THIN_SLABS
+		if (bDoSlabSearch)
+		{
+			FRayDesc ProbeRay;
+			ProbeRay.Origin = Ray.Origin;
+			ProbeRay.Direction = -DwivediSlabNormal;
+			ProbeRay.TMin = 0.0;
+			ProbeRay.TMax = 10 * max3(SSS.Radius.x, SSS.Radius.y, SSS.Radius.z);
+			int ProbeInterfaceCounter = InterfaceCounter;
+			FProbeResult Result = TraceSSSProbeRay(ProbeRay, ProbeInterfaceCounter);
+			if (Result.IsMiss())
+			{
+				// didn't find a hit, register missing slab
+				SlabThickness = -1.0;
+			}
+			else
+			{
+				// got a valid hit -- use it as our thickness
+				SlabThickness = Result.HitT;
+			}
+			bDoSlabSearch = false;
+		}
+
+		// Instead of only guiding towards the slab front (reflection), also guide toward the slack back (tranmission) when the surface is thin
+		// The heuristic to choose between guiding front or back is determined by the following probability, given in [2] on slide 37 of the video
+		// Note that the depths in the video presentation are optical depths, so have to be multiplied by SigmaT
+		float SlabZ = clamp(dot(DwivediSlabOrigin - Ray.Origin, DwivediSlabNormal), 0.0, SlabThickness);
+		float3 ProbT = SlabThickness > 0.0 ? rcp(1 + exp(SigmaT * (SlabThickness - 2 * SlabZ) / DwivediScale)) : 0.0;
+#else
+		float3 ProbT = 0.0;
+#endif
+		Ray.TMax = SampleGuidedSpectralTransmittance(RandSample.z, SlabCosine, DwivediScale, GuidedRatio, SigmaT, ProbT, ColorChannelPdf);
+#else // SSS_USE_DWIVEDI
 		Ray.TMax = SampleSpectralTransmittance(RandSample.z, SigmaT, ColorChannelPdf);
+#endif
 		if (Ray.TMax < 0.0)
 		{
 			// no more energy left in the path
 			break;
 		}
-		for (;;)
+
+		FProbeResult ProbeResult = TraceSSSProbeRay(Ray, InterfaceCounter);
+
+		RandSample = RandomSequence_GenerateSample3D(RandSequence);
+		if (ProbeResult.IsMiss())
 		{
-			FProbeResult ProbeResult = TraceProbeRay(Ray);
-			if (ProbeResult.IsMiss())
-			{
-				// we didn't hit anything, so scatter according to the scattering distribution in the volume and keep tracing
-				RandSample = RandomSequence_GenerateSample3D(RandSequence);
-				Ray.Origin += Ray.TMax * Ray.Direction;
-				
-				float4 DirectionAndPhase = ImportanceSampleHenyeyGreensteinPhase(RandSample.xy, G);
-				Ray.Direction = TangentToWorld(DirectionAndPhase.xyz, Ray.Direction);
-				Ray.TMin = 0;
-				// account for transmittance and scattering up to this point
-				PathThroughput *= SigmaS * EvaluateSpectralTransmittanceHit(Ray.TMax, SigmaT, ColorChannelPdf).xyz;
-				// NOTE: phase function cancels out since it is being perfectly importance sampled
-
-				// exit the Interface counting loop and keep scattering
-				break;
-			}
-#if SSS_USE_INTERFACE_COUNTING
-			// update counter
-			InterfaceCounter += ProbeResult.FrontFace ? +1 : -1;
-			if (InterfaceCounter != 0)
-			{
-				// we have not yet crossed the right number of interfaces, so ignore the current hit and try the segment again
-				Ray.TMin = asfloat(asuint(ProbeResult.HitT) + 1);
-			}
-			else
+			// we didn't hit anything, so scatter according to the scattering distribution in the volume and keep tracing
+			Ray.Origin += Ray.TMax * Ray.Direction;
+
+#if SSS_USE_DWIVEDI
+			// account for transmittance and scattering up to this point
+			PathThroughput *= SigmaS * EvaluateGuidedSpectralTransmittanceHit(Ray.TMax, SlabCosine, DwivediScale, GuidedRatio, SigmaT, ProbT, ColorChannelPdf).xyz;
+
+#if SSS_USE_DWIVEDI_USE_THIN_SLABS
+			// Recompute ProbT to account for new position before choosing a new direction
+			// NOTE: this does not appear to help much, so leave it off for now
+			//SlabZ = clamp(dot(DwivediSlabOrigin - Ray.Origin, DwivediSlabNormal), 0.0, SlabThickness);
+			//ProbT = SlabThickness > 0.0 ? rcp(1 + exp(SigmaT * (SlabThickness - 2 * SlabZ) / DwivediScale)) : 0.0;
 #endif
-			{
-				// account for transmittance to the boundary
-				PathThroughput *= EvaluateSpectralTransmittanceMiss(ProbeResult.HitT, SigmaT, ColorChannelPdf).xyz;
 
-				// our short ray hit the geomtry - decide if we should exit or not
+			float4 Result = SampleDwivediPhaseFunction(ColorChannelPdf, DwivediScale, GuidedRatio, ProbT, DwivediSlabNormal, Ray.Direction, G, RandSample.xy);
+			Ray.Direction = Result.xyz;
+			PathThroughput *= Result.w;
+#else
+			// importance sample the phase function
+			float4 DirectionAndPhase = ImportanceSampleHenyeyGreensteinPhase(RandSample.xy, G);
+			Ray.Direction = TangentToWorld(DirectionAndPhase.xyz, Ray.Direction);
+			// NOTE: phase function cancels out since it is being perfectly importance sampled
+			// account for transmittance and scattering up to this point
+			PathThroughput *= SigmaS * EvaluateSpectralTransmittanceHit(Ray.TMax, SigmaT, ColorChannelPdf).xyz;
+#endif
+			// keep scattering
+			continue;
+		}
+		else
+		{
+#if SSS_USE_DWIVEDI
+            // account for transmittance to the boundary as well as the guiding probability
+			PathThroughput *= EvaluateGuidedSpectralTransmittanceMiss(ProbeResult.HitT, SlabCosine, DwivediScale, GuidedRatio, SigmaT, ProbT, ColorChannelPdf).xyz;
+#else
+			// account for transmittance to the boundary
+			PathThroughput *= EvaluateSpectralTransmittanceMiss(ProbeResult.HitT, SigmaT, ColorChannelPdf).xyz;
+#endif
+			// our short ray hit the geometry - decide if we should exit or not
 #if SSS_USE_TIR
-				RandSample = RandomSequence_GenerateSample2D(RandSequence).xxy;
-				float3 WorldNormal = ProbeResult.WorldNormal;
-				float CosTheta = abs(dot(Ray.Direction, WorldNormal));
-				float Fresnel = FresnelReflectance(CosTheta, 1.0 / 1.4);
-				if (RandSample.x < Fresnel)
-				{
-					// internal reflection occured -- reflect and keep tracing
-					// NOTE: weight and probability cancel out, so no throughput adjustment is needed
-					Ray.Origin += ProbeResult.HitT * Ray.Direction;
-					Ray.Direction = reflect(Ray.Direction, WorldNormal);
-					Ray.TMin = 0;
-					ApplyRayBias(Ray, ProbeResult.HitT, ProbeResult.WorldGeoNormal);
-					// Reset interface counter as if we had just entered the model (taking into account that we are seeing it from the opposite side now)
-					InterfaceCounter = ProbeResult.FrontFace ? -1 : +1;
-					// exit from the interface counting loop and keep scattering
-					break;
-				}
+			float3 WorldNormal = ProbeResult.WorldNormal;
+			float CosTheta = abs(dot(Ray.Direction, WorldNormal));
+			float Fresnel = FresnelReflectance(CosTheta, 1.0 / 1.4);
+			if (RandSample.x < Fresnel)
+			{
+				// internal reflection occured -- reflect and keep tracing
+				// NOTE: weight and probability cancel out, so no throughput adjustment is needed
+				Ray.Origin += ProbeResult.HitT * Ray.Direction;
+				Ray.Direction = reflect(Ray.Direction, WorldNormal);
+
+#if SSS_USE_DWIVEDI && SSS_USE_DWIVEDI_USE_THIN_SLABS
+				// we hit the boundary, so reset the slab origin and look for a new one
+				DwivediSlabOrigin = Ray.Origin;
+				DwivediSlabNormal = ProbeResult.WorldSmoothNormal * ((ProbeResult.FrontFace != Payload.IsFrontFace()) ? -1.0 : 1.0);
+				bDoSlabSearch = GuidedRatio > 0;
 #endif
-				// we hit the boundary! overwrite the current payload and exit the walk with a diffuse scattering event
-				Payload.TranslatedWorldPos = Ray.Origin + ProbeResult.HitT * Ray.Direction;
-				Payload.WorldNormal = ProbeResult.WorldNormal;
-				Payload.WorldSmoothNormal = ProbeResult.WorldSmoothNormal;
-				Payload.WorldGeoNormal = ProbeResult.WorldGeoNormal;
-				// make sure normal is pointed outward from the object so we capture illumination from the exterior
-				if (ProbeResult.FrontFace != Payload.IsFrontFace())
-				{
-					Payload.WorldNormal *= -1.0;
-					Payload.WorldSmoothNormal *= -1.0;
-					Payload.WorldGeoNormal *= -1.0;
-				}
-				Payload.ShadingModelID = SHADINGMODELID_NUM; // invalid value so that we get diffuse shading
-				Payload.BSDFOpacity = 1;
-				Payload.BaseColor = 1;
-				Payload.TransparencyColor = 0;
-				return true;
+
+				ApplyRayBias(Ray, ProbeResult.HitT, ProbeResult.WorldGeoNormal);
+
+				// Reset interface counter as if we had just entered the model (taking into account that we are seeing it from the opposite side now)
+				InterfaceCounter = ProbeResult.FrontFace ? -1 : +1;
+				// keep scattering
+				continue;
+			}
+#endif // SSS_USE_TIR
+
+			// we hit the boundary! overwrite the current payload and exit the walk with a diffuse scattering event
+			Payload.TranslatedWorldPos = Ray.Origin + ProbeResult.HitT * Ray.Direction;
+			Payload.WorldNormal = ProbeResult.WorldNormal;
+			Payload.WorldSmoothNormal = ProbeResult.WorldSmoothNormal;
+			Payload.WorldGeoNormal = ProbeResult.WorldGeoNormal;
+			// make sure normal is pointed outward from the object so we capture illumination from the exterior
+			if (ProbeResult.FrontFace != Payload.IsFrontFace())
+			{
+				Payload.WorldNormal *= -1.0;
+				Payload.WorldSmoothNormal *= -1.0;
+				Payload.WorldGeoNormal *= -1.0;
 			}
+			Payload.ShadingModelID = SHADINGMODELID_NUM; // invalid value so that we get diffuse shading
+			Payload.BSDFOpacity = 1;
+			Payload.BaseColor = 1;
+			Payload.TransparencyColor = 0;
+			return true;
 		}
 	}
 	// we scattered a bunch of times and never hit anything -- give up
diff --git a/Engine/Source/Runtime/Renderer/Private/PathTracing.cpp b/Engine/Source/Runtime/Renderer/Private/PathTracing.cpp
--- a/Engine/Source/Runtime/Renderer/Private/PathTracing.cpp
+++ b/Engine/Source/Runtime/Renderer/Private/PathTracing.cpp
@@ -145,6 +145,14 @@ TAutoConsoleVariable<int32> CVarPathTracingMaxSSSBounces(
 	ECVF_RenderThreadSafe
 );
 
+TAutoConsoleVariable<float> CVarPathTracingSSSGuidingRatio(
+	TEXT("r.PathTracing.SSSGuidingRatio"),
+	0.5f,
+	TEXT("Sets the ratio between classical random walks and walks guided towards the surface. A value of 0.0 corresponds to a purely classical random walk, while a value of 1.0 is fully guided towards the surface (at the expense of fireflies in non-flat regions of the model. (default = 0.5)"),
+	ECVF_RenderThreadSafe
+);
+
+
 TAutoConsoleVariable<float> CVarPathTracingMaxPathIntensity(
 	TEXT("r.PathTracing.MaxPathIntensity"),
 	-1,
@@ -336,6 +344,7 @@ BEGIN_SHADER_PARAMETER_STRUCT(FPathTracingData, )
 	SHADER_PARAMETER(uint32, MaxSamples)
 	SHADER_PARAMETER(uint32, MaxBounces)
 	SHADER_PARAMETER(uint32, MaxSSSBounces)
+	SHADER_PARAMETER(float , SSSGuidingRatio)
 	SHADER_PARAMETER(uint32, MISMode)
 	SHADER_PARAMETER(uint32, VolumeMISMode)
 	SHADER_PARAMETER(uint32, ApproximateCaustics)
@@ -383,6 +392,7 @@ struct FPathTracingConfig
 			PathTracingData.MaxSamples != Other.PathTracingData.MaxSamples ||
 			PathTracingData.MaxBounces != Other.PathTracingData.MaxBounces ||
 			PathTracingData.MaxSSSBounces != Other.PathTracingData.MaxSSSBounces ||
+			PathTracingData.SSSGuidingRatio != Other.PathTracingData.SSSGuidingRatio ||
 			PathTracingData.MISMode != Other.PathTracingData.MISMode ||
 			PathTracingData.VolumeMISMode != Other.PathTracingData.VolumeMISMode ||
 			PathTracingData.SamplerType != Other.PathTracingData.SamplerType ||
@@ -511,6 +521,7 @@ static void PreparePathTracingData(const FScene* Scene, const FViewInfo& View, F
 
 	PathTracingData.MaxBounces = MaxBounces;
 	PathTracingData.MaxSSSBounces = ShowFlags.SubsurfaceScattering ? CVarPathTracingMaxSSSBounces.GetValueOnRenderThread() : 0;
+	PathTracingData.SSSGuidingRatio = FMath::Clamp(CVarPathTracingSSSGuidingRatio.GetValueOnRenderThread(), 0.0f, 1.0f);
 	PathTracingData.MaxNormalBias = GetRaytracingMaxNormalBias();
 	PathTracingData.MISMode = CVarPathTracingMISMode.GetValueOnRenderThread();
 	PathTracingData.VolumeMISMode = CVarPathTracingVolumeMISMode.GetValueOnRenderThread();
@@ -2036,8 +2047,6 @@ void FDeferredShadingSceneRenderer::PreparePathTracing(const FSceneViewFamily& V
 		const int CompactionType = CVarPathTracingCompaction.GetValueOnRenderThread();
 		FPathTracingRG::FPermutationDomain PermutationVector;
 		PermutationVector.Set<FPathTracingRG::FCompactionType>(CompactionType);
-		FGlobalShaderPermutationParameters Parameters(FPathTracingRG::StaticGetTypeLayout().Name, ViewFamily.GetShaderPlatform(), PermutationVector.ToDimensionValueId());
-		if (FPathTracingRG::ShouldCompilePermutation(Parameters))
 		{
 			auto RayGenShader = GetGlobalShaderMap(ViewFamily.GetShaderPlatform())->GetShader<FPathTracingRG>(PermutationVector);
 			OutRayGenShaders.Add(RayGenShader.GetRayTracingShader());
