Fix a shader compilation failure introduced by the interaction of CL 24117687 and the robomerged 5.2 MMV forward light data changes.

diff --git a/Engine/Shaders/Private/LightGridCommon.ush b/Engine/Shaders/Private/LightGridCommon.ush
--- a/Engine/Shaders/Private/LightGridCommon.ush
+++ b/Engine/Shaders/Private/LightGridCommon.ush
@@ -31,7 +31,7 @@ uint GetCulledLightDataGrid(uint GridIndex)
 #endif
 }
 
-#if INSTANCED_STEREO
+#if (INSTANCED_STEREO || MOBILE_MULTI_VIEW)
 uint GetCulledLightDataGridISR(uint GridIndex)
 {
 #if LIGHT_GRID_USES_16BIT_BUFFERS
