PostProcessUpscale: Fixed incorrect truncation.

Squiggly vertical lines were seen in Panini projection renders, and upon examination it seems to be due to PostProcessUpscale.usf shifting the horizontal vertex position due to a cast to uint of the value that fmod returns, which might not always be exact.

diff --git a/Engine/Shaders/Private/PostProcessUpscale.usf b/Engine/Shaders/Private/PostProcessUpscale.usf
--- a/Engine/Shaders/Private/PostProcessUpscale.usf
+++ b/Engine/Shaders/Private/PostProcessUpscale.usf
@@ -97,7 +97,7 @@ void MainVS(
 	// amount of Quads in X and Y, from GTesselatedScreenRectangleIndexBuffer
 	uint2 RectangleTessellation = uint2(TESS_RECT_X, TESS_RECT_Y);
 
-	float2 TexCoord = float2((uint)fmod(VertexId, (RectangleTessellation.x + 1)), VertexId / (RectangleTessellation.x + 1)) / float2(RectangleTessellation);
+	float2 TexCoord = float2(VertexId % (RectangleTessellation.x + 1), VertexId / (RectangleTessellation.x + 1)) / float2(RectangleTessellation);
 
 	// still in 0..1 range
 	float4 Position = float4(TexCoord.x, TexCoord.y, 0, 1);
