Do not use MaxWPODisplacement on mobile platforms with GPUScene enabled as it currently breaks auto-instancing
Remove WPO usage branch in the shader when material does not use WPO

diff --git a/Engine/Shaders/Private/MaterialTemplate.ush b/Engine/Shaders/Private/MaterialTemplate.ush
--- a/Engine/Shaders/Private/MaterialTemplate.ush
+++ b/Engine/Shaders/Private/MaterialTemplate.ush
@@ -2925,6 +2925,10 @@ bool ShouldEnableWorldPositionOffset(FMaterialVertexParameters Parameters)
 	}
 #endif
 
+#if !USES_WORLD_POSITION_OFFSET
+	return false;
+#endif
+
 #if USE_INSTANCING || USE_INSTANCE_CULLING
 	// skip if this instance is hidden
 	if (Parameters.PerInstanceParams.y < 1.f)
@@ -2943,8 +2947,13 @@ bool ShouldEnableWorldPositionOffset(FMaterialVertexParameters Parameters)
 
 float3 ClampWorldPositionOffset(FMaterialVertexParameters Parameters, float3 InOffset)
 {
+#if (FEATURE_LEVEL == FEATURE_LEVEL_ES3_1 && VF_SUPPORTS_PRIMITIVE_SCENE_DATA)
+	// Do not use WPO clamping on mobile platforms as MaxWPODisplacement has to be pulled from PrimitiveUB breaking auto-instancing
+	return InOffset;
+#else
 	const float MaxWPODim = GetPrimitiveData(Parameters).MaxWPODisplacement;
 	return MaxWPODim <= 0.0f ? InOffset : clamp(InOffset, -MaxWPODim.xxx, MaxWPODim.xxx);
+#endif
 }
 
 float3 GetMaterialWorldPositionOffsetRaw(FMaterialVertexParameters Parameters)
