Clearcoat: Fix shading artifacts when one of the components of the base color is exactly 0.0

diff --git a/Engine/Shaders/Private/BRDF.ush b/Engine/Shaders/Private/BRDF.ush
--- a/Engine/Shaders/Private/BRDF.ush
+++ b/Engine/Shaders/Private/BRDF.ush
@@ -724,7 +724,7 @@ float3 SimpleClearCoatTransmittance(float NoL, float NoV, float Metallic, float3
 		float3 TransmittanceColor = Diffuse_Lambert(BaseColor);
 		// ... and after being affected by the substrate transmittance when traveling through it, back and forth.
 		// Because of this, extinction is normalized by traveling through layer thickness twice
-		float3 ExtinctionCoefficient = -log(TransmittanceColor) / (2.0 * LayerThickness);
+		float3 ExtinctionCoefficient = -log(max(TransmittanceColor,0.0001)) / (2.0 * LayerThickness);
 
 		// Optical depth is evaluated for ThinDistance computed above. But we also subtract 2.0 * LayerThickness
 		// in order to respect the fact that BaseColor represents reflected color viewed at 0 incidence angle. => In this case, OpticalDepth will be 0 and thus Transmittance=1.
