Fix light leak on opaque geometry when using layering with simple volume.

diff --git a/Engine/Shaders/Private/Strata/StrataEvaluation.ush b/Engine/Shaders/Private/Strata/StrataEvaluation.ush
--- a/Engine/Shaders/Private/Strata/StrataEvaluation.ush
+++ b/Engine/Shaders/Private/Strata/StrataEvaluation.ush
@@ -810,7 +810,11 @@ FStrataEvaluateResult StrataEvaluateBSDFCommon(FStrataBSDFContext BSDFContext, F
 				Sample.bSubsurface = false;	// It should already be the case because Enforce in this case because BSDF_GETHASSSS should be false when BSDF_GETISSIMPLEVOLUME is true
 
 				// Use Context.NoL which is not saturate, and allows to handle backface lighting, necessary for medium support
+				#if STRATA_TRANSLUCENT_ENABLED
 				const float MediumNoL = abs(BSDFContext.Context.NoL);
+				#else
+				const float MediumNoL = saturate(BSDFContext.Context.NoL);
+				#endif
 
 				Sample.IntegratedDiffuseValue = lerp(
 					Sample.IntegratedDiffuseValue, 
