Force early_depth test for a masked VT materials on mobile. Contributed by: Gucai Zhu (hkzhugc)
#10223

diff --git a/Engine/Shaders/Private/MobileBasePassPixelShader.usf b/Engine/Shaders/Private/MobileBasePassPixelShader.usf
--- a/Engine/Shaders/Private/MobileBasePassPixelShader.usf
+++ b/Engine/Shaders/Private/MobileBasePassPixelShader.usf
@@ -163,8 +163,8 @@ void ApplyPixelDepthOffsetForMobileBasePass(inout FMaterialPixelParameters Mater
     float PixelDepthOffset = ApplyPixelDepthOffsetToMaterialParameters(MaterialParameters, PixelMaterialInputs, OutDepth);
 }
 
-// Force early depth_stencil for non-masked material that use VT feedback
-#if (NUM_VIRTUALTEXTURE_SAMPLES || LIGHTMAP_VT_ENABLED) && !(MATERIALBLENDING_MASKED || USE_DITHERED_LOD_TRANSITION || OUTPUT_PIXEL_DEPTH_OFFSET)
+// Force early depth_stencil for materials that use VT feedback and don't modify the pixel depth
+#if (NUM_VIRTUALTEXTURE_SAMPLES || LIGHTMAP_VT_ENABLED) && !(OUTPUT_PIXEL_DEPTH_OFFSET)
 	#define PIXELSHADER_EARLYDEPTHSTENCIL EARLYDEPTHSTENCIL	
 #else
 	#define PIXELSHADER_EARLYDEPTHSTENCIL 	
