Fix a bug of filter bloom input.

diff --git a/Engine/Shaders/Private/PostProcessMobile.usf b/Engine/Shaders/Private/PostProcessMobile.usf
--- a/Engine/Shaders/Private/PostProcessMobile.usf
+++ b/Engine/Shaders/Private/PostProcessMobile.usf
@@ -233,9 +233,6 @@ void BloomVS_Mobile(
 	OutTexCoords[3] = TransformedUV + BufferSizeAndInvSize.zw * float2( 1,  1);
 }
 
-// Temporal jitter at the pixel scale.
-float2 TemporalJitterPixels;
-
 static const float2 kOffsets[5] =
 {
 	int2( 0, -1),
@@ -308,13 +305,11 @@ void BloomPS_Mobile(
 
 	half NeighborsFinalWeight = 0;
 	float3 NeighborsColor = 0;
-
+	float2 PPCk = floor(SvPosition.xy) * 4 + 0.5f;
 	// Filter input pixels. Referenced from the FilterCurrentFrameInputSamples in TemporalAA.usf
 	for (uint i = 0; i < 5; i++)
 	{
- 		float2 PPCk = SvPosition.xy + TemporalJitterPixels;
- 
-		float2 TexCoords = BufferSizeAndInvSize.zw * (kOffsets[i] + PPCk);
+		float2 TexCoords = BufferSizeAndInvSize.zw * (kOffsets[i] + PPCk) - View.TemporalAAJitter.xy * float2(0.5f, -0.5f);
 		half3 SampleColor = SceneColorTexture.Sample(SceneColorSampler, TexCoords).rgb;
 		half SampleHdrWeight = HdrWeight4(SampleColor);
 		half SampleSpatialWeight = GET_SCALAR_ARRAY_ELEMENT(PlusWeights, i);
diff --git a/Engine/Source/Runtime/Renderer/Private/PostProcess/PostProcessMobile.cpp b/Engine/Source/Runtime/Renderer/Private/PostProcess/PostProcessMobile.cpp
--- a/Engine/Source/Runtime/Renderer/Private/PostProcess/PostProcessMobile.cpp
+++ b/Engine/Source/Runtime/Renderer/Private/PostProcess/PostProcessMobile.cpp
@@ -217,7 +217,6 @@ public:
 
 	BEGIN_SHADER_PARAMETER_STRUCT(FParameters, )
 		SHADER_PARAMETER(float, BloomThreshold)
-		SHADER_PARAMETER(FVector2f, TemporalJitterPixels)
 		SHADER_PARAMETER(FVector4f, BufferSizeAndInvSize)
 		SHADER_PARAMETER_SCALAR_ARRAY(float, PlusWeights, [5])
 		SHADER_PARAMETER_STRUCT_REF(FViewUniformShaderParameters, View)
@@ -376,8 +375,6 @@ FMobileBloomSetupOutputs AddMobileBloomSetupPass(FRDGBuilder& GraphBuilder, cons
 		{
 			GET_SCALAR_ARRAY_ELEMENT(PSShaderParameters->PlusWeights, i) /= TotalWeight;
 		}
-
-		PSShaderParameters->TemporalJitterPixels = TemporalJitterPixels;
 	}
 
 	const FScreenPassTextureViewport InputViewport(Inputs.SceneColor);
