Address some global shader compilation warnings for implicit coversion from float to half

#rnx

diff --git a/Engine/Shaders/Private/GBufferHelpers.ush b/Engine/Shaders/Private/GBufferHelpers.ush
--- a/Engine/Shaders/Private/GBufferHelpers.ush
+++ b/Engine/Shaders/Private/GBufferHelpers.ush
@@ -374,10 +374,10 @@ void AdjustBaseColorAndSpecularColorForSubsurfaceProfileLightingCopyHack(inout f
 
 void GBufferPostDecode(inout FGBufferData Ret, bool bChecker, bool bGetNormalizedNormal)
 {
-	Ret.CustomData = HasCustomGBufferData(Ret.ShadingModelID) ? Ret.CustomData : 0.0f;
+	Ret.CustomData = HasCustomGBufferData(Ret.ShadingModelID) ? Ret.CustomData : 0.0h;
 
-	Ret.PrecomputedShadowFactors = !(Ret.SelectiveOutputMask & 0x2) ? Ret.PrecomputedShadowFactors : ((Ret.SelectiveOutputMask & 0x4) ? 0.0f : 1.0f);
-	Ret.Velocity = !(Ret.SelectiveOutputMask & 0x8) ? Ret.Velocity : 0.0f;
+	Ret.PrecomputedShadowFactors = !(Ret.SelectiveOutputMask & 0x2) ? Ret.PrecomputedShadowFactors : ((Ret.SelectiveOutputMask & 0x4) ? 0.0h : 1.0h);
+	Ret.Velocity = !(Ret.SelectiveOutputMask & 0x8) ? Ret.Velocity : 0.0h;
 	bool bHasAnisotropy = (Ret.SelectiveOutputMask & 0x1);
 
 	Ret.StoredBaseColor = Ret.BaseColor;
@@ -434,8 +434,8 @@ void GBufferPostDecode(inout FGBufferData Ret, bool bChecker, bool bGetNormalize
 	
 	if (bHasAnisotropy)
 	{
-		Ret.WorldTangent = DecodeNormal(Ret.WorldTangent);
-		Ret.Anisotropy = Ret.Anisotropy * 2.0f - 1.0f;
+		Ret.WorldTangent = half3(DecodeNormal(Ret.WorldTangent));
+		Ret.Anisotropy = half(Ret.Anisotropy * 2.0f - 1.0f);
 	
 		if(bGetNormalizedNormal)
 		{
diff --git a/Engine/Shaders/Private/LightData.ush b/Engine/Shaders/Private/LightData.ush
--- a/Engine/Shaders/Private/LightData.ush
+++ b/Engine/Shaders/Private/LightData.ush
@@ -50,7 +50,7 @@ FRectLightData UnpackRectLightData(uint In0, uint In1, uint In2)
 	Out.AtlasData.AtlasUVScale	= half2(f16tof32(In1 & 0xFFFF), f16tof32(In1 >> 16));
 	Out.BarnLength				= f16tof32(In2 & 0xFFFF);			// 16 bits
 	Out.BarnCosAngle			= ((In2 >> 16) & 0x3FF) / 1024.f;	// 10 bits
-	Out.AtlasData.AtlasMaxLevel =  (In2 >> 26);						//  6 bits
+	Out.AtlasData.AtlasMaxLevel = half(In2 >> 26);					//  6 bits
 	return Out;
 }
 
diff --git a/Engine/Shaders/Private/ShadingCommon.ush b/Engine/Shaders/Private/ShadingCommon.ush
--- a/Engine/Shaders/Private/ShadingCommon.ush
+++ b/Engine/Shaders/Private/ShadingCommon.ush
@@ -112,7 +112,7 @@ float F0RGBToDielectricSpecular(float3 F0)
 
 half DielectricSpecularToF0(half Specular)
 {
-	return 0.08f * Specular;
+	return half(0.08f * Specular);
 }
 
 // [Burley, "Extending the Disney BRDF to a BSDF with Integrated Subsurface Scattering"]
diff --git a/Engine/Shaders/Private/VirtualShadowMaps/VirtualShadowMapProjection.usf b/Engine/Shaders/Private/VirtualShadowMaps/VirtualShadowMapProjection.usf
--- a/Engine/Shaders/Private/VirtualShadowMaps/VirtualShadowMapProjection.usf
+++ b/Engine/Shaders/Private/VirtualShadowMaps/VirtualShadowMapProjection.usf
@@ -272,7 +272,7 @@ FVirtualShadowMapSampleResult ProjectLight(
 				SMRTRayLengthScale,
 				Noise,
 				bBackfaceCull,
-				WorldNormal,
+				half3(WorldNormal),
 				SMRTTexelDitherScale,
 				SMRTExtrapolateSlope,
 				SMRTAdaptiveRayCount);
