Substrate - fix for wrong "transmittance above N" usage with the vrtical layering operator.
Deleted unused TransmittanceAboveAlongN.

Verified car in lookdev, furnace and Opal map.

diff --git a/Engine/Shaders/Private/Strata/Strata.ush b/Engine/Shaders/Private/Strata/Strata.ush
--- a/Engine/Shaders/Private/Strata/Strata.ush
+++ b/Engine/Shaders/Private/Strata/Strata.ush
@@ -123,8 +123,8 @@ struct FStrataBSDF
 	int					OperatorIndex;
 	float3				LuminanceWeightV;
 
-	float3				TransmittanceAboveAlongN;	// The transmittance towards the top of the material along the normal and ignoring this BSDF, e.g. only accumulating matter above it.
-	float3				TransmittanceBelowAlongN;	// The transmittance towards the bottom of the material along the normal and ignoring this BSDF, e.g. only accumulating matter below it.
+	float				CoverageAboveAlongN;
+	float3				TransmittanceAboveAlongN;	// The "pre-coverage" transmittance towards the top of the material along the normal and ignoring this BSDF, e.g. only accumulating matter above it.
 
 	int					bIsBottom;
 	int					bIsTop;
@@ -1163,8 +1163,8 @@ void StrataPackOutStrataTree(
 		WRITE_UINT(StrataTree.BSDFs[i].State);
 		WRITE_INT(StrataTree.BSDFs[i].OperatorIndex);
 		WRITE_FLOAT3(StrataTree.BSDFs[i].LuminanceWeightV);
+		WRITE_FLOAT(StrataTree.BSDFs[i].CoverageAboveAlongN);
 		WRITE_FLOAT3(StrataTree.BSDFs[i].TransmittanceAboveAlongN);
-		WRITE_FLOAT3(StrataTree.BSDFs[i].TransmittanceBelowAlongN);
 		WRITE_UINT(StrataTree.BSDFs[i].bIsBottom);
 		WRITE_UINT(StrataTree.BSDFs[i].bIsTop);
 
@@ -1260,8 +1260,8 @@ FStrataBSDF StrataUnpackInStrataTreeBSDF(
 	READ_UINT(BSDF.State);
 	READ_INT(BSDF.OperatorIndex);
 	READ_FLOAT3(BSDF.LuminanceWeightV);
+	READ_FLOAT(BSDF.CoverageAboveAlongN);
 	READ_FLOAT3(BSDF.TransmittanceAboveAlongN);
-	READ_FLOAT3(BSDF.TransmittanceBelowAlongN);
 	READ_UINT(BSDF.bIsBottom);
 	READ_UINT(BSDF.bIsTop);
 
@@ -3732,6 +3732,7 @@ FStrataBSDF UnpackFastPathStrataBSDFIn(FStrataMaterialContainer StrataBuffer, in
 	uint Data1 = StrataLoadUint1(StrataBuffer, StrataAddressing);
 	OutBSDF.State = 0;
 	OutBSDF.LuminanceWeightV = 1.0f;
+	OutBSDF.CoverageAboveAlongN = 0.0f;
 	OutBSDF.TransmittanceAboveAlongN = 1.0f;
 
 	const uint PackedRoughness8bits = 0xFF    & (Data0 >>  HEADER_FASTENCODING_BIT_COUNT);
@@ -3775,6 +3776,7 @@ FStrataBSDF UnpackFastWaterPathStrataBSDFIn(FStrataMaterialContainer StrataBuffe
 	uint Data1 = StrataLoadUint1(StrataBuffer, StrataAddressing);
 	OutBSDF.State = 0;
 	OutBSDF.LuminanceWeightV = 1.0f;
+	OutBSDF.CoverageAboveAlongN = 0.0f;
 	OutBSDF.TransmittanceAboveAlongN = 1.0f;
 
 	const uint PackedBaseColor8Bits					= 0xFF   & (Data0 >> (HEADER_FASTENCODING_BIT_COUNT + 8));
@@ -3827,6 +3829,7 @@ FStrataBSDF UnpackStrataBSDFIn(FStrataMaterialContainer StrataBuffer, inout FStr
 	}
 
 	FStrataBSDF OutBSDF = (FStrataBSDF)0;
+	OutBSDF.CoverageAboveAlongN = 0.0f;
 	OutBSDF.TransmittanceAboveAlongN = 1.0f;
 
 	float DummyFloat = 0.0f;
@@ -3858,6 +3861,7 @@ FStrataBSDF UnpackStrataBSDFIn(FStrataMaterialContainer StrataBuffer, inout FStr
 		{
 			OutBSDF.State = 0;
 			OutBSDF.LuminanceWeightV = 1.0f;
+			OutBSDF.CoverageAboveAlongN = 0.0f;
 			OutBSDF.TransmittanceAboveAlongN = 1.0f;
 			SLAB_F90(OutBSDF) = 1.0f;
 
@@ -4126,7 +4130,7 @@ FStrataBSDF UnpackStrataBSDFIn(FStrataMaterialContainer StrataBuffer, inout FStr
 
 		if (BSDF_GETHASTRANSABOVE(OutBSDF))
 		{
-			StrataLoad_R11G11B10F(StrataBuffer, StrataAddressing, OutBSDF.TransmittanceAboveAlongN);
+			StrataLoad_ColorGamma2ToLinearAlphaLinear(StrataBuffer, StrataAddressing, OutBSDF.TransmittanceAboveAlongN, OutBSDF.CoverageAboveAlongN);
 		}
 	}
 	break;
diff --git a/Engine/Shaders/Private/Strata/StrataEvaluation.ush b/Engine/Shaders/Private/Strata/StrataEvaluation.ush
--- a/Engine/Shaders/Private/Strata/StrataEvaluation.ush
+++ b/Engine/Shaders/Private/Strata/StrataEvaluation.ush
@@ -200,7 +200,7 @@ float3 LuminanceWeight(in float SatNoL, in FStrataBSDF InBSDF)
 	// TransmittanceAboveAlongL	= pow(TransmittanceAboveAlongN, 1/NoL)
 	const float DistanceL = rcp(max(SatNoL, 0.001f));
 	const float3 TransmittanceAboveAlongL = any(InBSDF.TransmittanceAboveAlongN < 1.f) ? pow(max(InBSDF.TransmittanceAboveAlongN, 0.0001f), DistanceL) : InBSDF.TransmittanceAboveAlongN;
-	return InBSDF.LuminanceWeightV * TransmittanceAboveAlongL;
+	return InBSDF.LuminanceWeightV * lerp(1.0f, TransmittanceAboveAlongL, InBSDF.CoverageAboveAlongN);
 #else
 	return InBSDF.LuminanceWeightV;
 #endif
diff --git a/Engine/Shaders/Private/Strata/StrataExport.ush b/Engine/Shaders/Private/Strata/StrataExport.ush
--- a/Engine/Shaders/Private/Strata/StrataExport.ush
+++ b/Engine/Shaders/Private/Strata/StrataExport.ush
@@ -774,7 +774,7 @@ void PackStrataOut(
 
 							if (BSDF_GETHASTRANSABOVE(BSDF))
 							{
-								STRATA_STORE_UINT1(PackR11G11B10F(BSDF.TransmittanceAboveAlongN));
+								STRATA_STORE_UINT1(PackColorLinearToGamma2AlphaLinear(float4(BSDF.TransmittanceAboveAlongN, BSDF.CoverageAboveAlongN)));
 							}
 
 							// 8-24 bytes
diff --git a/Engine/Shaders/Private/Strata/StrataTree.ush b/Engine/Shaders/Private/Strata/StrataTree.ush
--- a/Engine/Shaders/Private/Strata/StrataTree.ush
+++ b/Engine/Shaders/Private/Strata/StrataTree.ush
@@ -111,7 +111,9 @@ void UpdateSingleBSDFOperatorCoverageTransmittance(
 	//  - When coming from non parameter blending, this is equal identity 1 and the operator parsing will continue to correctly weight the BSDF according to the tree.
 	CurrentBSDF.LuminanceWeightV = float3(1.0f, 1.0f, 1.0f) * CurrentBSDFCoverage;
 
-	CurrentBSDF.TransmittanceAboveAlongN = float3(1.0f, 1.0f, 1.0f); // Initialise to "no matter above"
+	// Initialise to "no matter above"
+	CurrentBSDF.CoverageAboveAlongN = 0.0f;
+	CurrentBSDF.TransmittanceAboveAlongN = float3(1.0f, 1.0f, 1.0f);
 
 	// All BSDF are considered for contribution as the top BSDF can have their coverage dip to 0 and then another BSDF need to be considered for the top layer representation.
 	CurrentBSDF.TopLayerDataWeight = 1.0f;
@@ -147,6 +149,7 @@ void UpdateSingleBSDFOperatorCoverageTransmittance(
 #else // STRATA_COMPLEXPATH
 
 	CurrentBSDF.LuminanceWeightV		= CurrentBSDFCoverage;
+	CurrentBSDF.CoverageAboveAlongN		= 0.0f;
 	CurrentBSDF.TransmittanceAboveAlongN= float3(1.0f, 1.0f, 1.0f);
 	CurrentBSDF.TopLayerDataWeight		= 1.0f;
 
@@ -237,12 +240,12 @@ void UpdateSingleOperatorCoverageTransmittance(inout FStrataTree StrataTree, int
 
 		Op.Coverage				= Info.Coverage;
 		Op.ThroughputAlongV		= Info.TransmittanceOnlyTop * TopThroughputAlongV +		Info.TransmittanceOnlyBottom * BotThroughputAlongV +	Info.TransmittanceTopAndBottom * TopThroughputAlongV * BotThroughputAlongV;
-		Op.TransmittanceAlongN	= Info.TransmittanceOnlyTop * TopTransmittanceAlongN +	Info.TransmittanceOnlyBottom * BotTransmittanceAlongN + Info.TransmittanceTopAndBottom * TopTransmittanceAlongN * BotTransmittanceAlongN;
+		Op.TransmittanceAlongN  = Info.TransmittanceOnlyTop* TopTransmittanceAlongN + Info.TransmittanceOnlyBottom * BotTransmittanceAlongN + Info.TransmittanceTopAndBottom * TopTransmittanceAlongN * BotTransmittanceAlongN;
 #endif
 
 		Op.VerticalTop_Coverage				= OpA.Coverage;
 		Op.VerticalTop_ThroughputAlongV		= OpA.ThroughputAlongV;
-		Op.VerticalTop_TransmittanceAlongN	= OpA.TransmittanceAlongN;
+		Op.VerticalTop_TransmittanceAlongN  = OpA.TransmittanceAlongN;
 
 		Op.VerticalBot_Coverage				= OpB.Coverage;
 		Op.VerticalBot_ThroughputAlongV		= OpB.ThroughputAlongV;
@@ -383,9 +386,18 @@ void UpdateAllBSDFWithBottomUpOperatorVisit_Vertical(
 
 	// PreviousIsInputA > 0 means it is a BSDF affecting the top layer so we do not affect it by anything for this vertical layering.
 	// Otherise, the BSDF comes from the bottom part so we affect it weights from the gathered top layer coverage/transmittance
-	CurrentBSDF.LuminanceWeightV		*= PreviousIsInputA > 0 ? 1.0f : saturate(1.0f - Op.VerticalTop_Coverage) + saturate(Op.VerticalTop_Coverage) * Op.VerticalTop_ThroughputAlongV;
-	CurrentBSDF.TransmittanceAboveAlongN*= PreviousIsInputA > 0 ? 1.0f : saturate(1.0f - Op.VerticalTop_Coverage) + saturate(Op.VerticalTop_Coverage) * Op.VerticalTop_TransmittanceAlongN;
-	CurrentBSDF.TransmittanceBelowAlongN*= PreviousIsInputA ==0 ? 1.0f : saturate(1.0f - Op.VerticalBot_Coverage) + saturate(Op.VerticalBot_Coverage) * Op.VerticalBot_TransmittanceAlongN;
+	CurrentBSDF.LuminanceWeightV		 *= PreviousIsInputA > 0  ? 1.0f : saturate(1.0f - Op.VerticalTop_Coverage) + saturate(Op.VerticalTop_Coverage) * Op.VerticalTop_ThroughputAlongV;
+
+	// Now we are going to combine the coverage/transmittance above this slab into a new one according to the vertical operator we traverse now, if PreviousIsInputA > 0
+	const float BotCover  = CurrentBSDF.CoverageAboveAlongN;
+	const float TopCover  = Op.VerticalTop_Coverage;
+	const float3 BotTrans = CurrentBSDF.TransmittanceAboveAlongN;
+	const float3 TopTrans = Op.VerticalTop_TransmittanceAlongN;
+	FVerticalLayeringInfo Info = GetVerticalLayeringInfo(Op.VerticalTop_Coverage, CurrentBSDF.CoverageAboveAlongN);
+	// Pre coverage transmittance is the surface transmitance as if coverage=1, so the final combined transmittance is renormalized accordingly.
+	const float3 PreCoverageTransmittance = saturate((Info.TransmittanceOnlyTop * TopTrans + Info.TransmittanceOnlyBottom * BotTrans + Info.TransmittanceTopAndBottom * (TopTrans * BotTrans)) / max(STRATA_EPSILON, Info.Coverage));
+	CurrentBSDF.CoverageAboveAlongN		  = PreviousIsInputA  > 0 ? CurrentBSDF.CoverageAboveAlongN : Info.Coverage;
+	CurrentBSDF.TransmittanceAboveAlongN *= PreviousIsInputA  > 0 ? 1.0f							: PreCoverageTransmittance;
 
 	// If the BSDF is from the bottom, we reduce its contribution according to the top layer coverage only.
 	CurrentBSDF.TopLayerDataWeight *= PreviousIsInputA ? 1.0 : saturate(1.0f - Op.VerticalTop_Coverage);
diff --git a/Engine/Shaders/Private/Strata/StrataVisualize.usf b/Engine/Shaders/Private/Strata/StrataVisualize.usf
--- a/Engine/Shaders/Private/Strata/StrataVisualize.usf
+++ b/Engine/Shaders/Private/Strata/StrataVisualize.usf
@@ -159,11 +159,16 @@ void PrintBSDF(inout FShaderPrintContext Context, inout FStrataAddressing Strata
 			Print(Context, TEXT(" NormalID Aniso TopLayer Scattering SimpleVolume IsThin WeightV "), FontBSDFStateName);
 			if (!bGreyWeightV)
 			{
-				Print(Context, TEXT("           "), FontBSDFStateName);
+				Print(Context, TEXT("            "), FontBSDFStateName);
 			}
 			if (bWeightL)
 			{
-				Print(Context, TEXT("WeightL"), FontBSDFStateName);
+				Print(Context, TEXT("TransToL     "), FontBSDFStateName);
+				if (!bGreyTopTrans)
+				{
+					Print(Context, TEXT("       "), FontBSDFStateName);
+				}
+				Print(Context, TEXT("CoverToL"), FontBSDFStateName);
 			}
 			Newline(Context, RectMax);
 
@@ -191,7 +196,7 @@ void PrintBSDF(inout FShaderPrintContext Context, inout FStrataAddressing Strata
 			}
 
 			// Light Weight
-			if (bGreyTopTrans && bWeightL)
+			if (bWeightL)
 			{
 				Print(Context, TEXT(" "), FontSilver);
 				if (bGreyTopTrans)
@@ -204,6 +209,9 @@ void PrintBSDF(inout FShaderPrintContext Context, inout FStrataAddressing Strata
 					Print(Context, BSDF.TransmittanceAboveAlongN.g, FontLightGreen, 6, 4);
 					Print(Context, BSDF.TransmittanceAboveAlongN.b, FontLightBlue,  6, 4);
 				}
+
+				Print(Context, TEXT(" "), FontSilver);
+				Print(Context, BSDF.CoverageAboveAlongN, FontSilver, 6, 4);
 			}
 			Newline(Context, RectMax);
 
@@ -821,10 +829,9 @@ void MaterialDebugStrataTreeCS(uint3 DispatchThreadId : SV_DispatchThreadID)
 				Print(Context, BSDF.TransmittanceAboveAlongN.g, FontLightGreen, 5, 2);
 				Print(Context, BSDF.TransmittanceAboveAlongN.b, FontLightBlue, 5, 2);
 				Newline(Context);
-				Print(Context, TEXT("BotTrans = "), BSDFFont);
-				Print(Context, BSDF.TransmittanceBelowAlongN.r, FontLightRed, 5, 2);
-				Print(Context, BSDF.TransmittanceBelowAlongN.g, FontLightGreen, 5, 2);
-				Print(Context, BSDF.TransmittanceBelowAlongN.b, FontLightBlue, 5, 2);
+				Print(Context, TEXT("TopCover = "), BSDFFont);
+				Print(Context, BSDF.CoverageAboveAlongN, FontLightRed, 5, 2);
+				Newline(Context);
 				Newline(Context);
 				Print(Context, TEXT("Diffuse  = "), BSDFFont);
 				const float3 BSDFDiffuse = StrataGetBSDFDiffuseColor(BSDF);
@@ -1045,7 +1052,6 @@ void EvaluateLighting(
 	// Patch the BSDF to appear as a top layer
 	BSDF.LuminanceWeightV = 1.0f;
 	BSDF.TransmittanceAboveAlongN = 1.0f;
-	BSDF.TransmittanceBelowAlongN = 1.0f;
 	BSDF.bIsBottom = true;
 	BSDF.bIsTop = true;
 	BSDF.Coverage = 1.0f;
