Speculative fix for shader compilation errors on some platforms

- some platforms don't like an h suffix on a floating point literal to specify half type

#rnx

diff --git a/Engine/Shaders/Private/GBufferHelpers.ush b/Engine/Shaders/Private/GBufferHelpers.ush
--- a/Engine/Shaders/Private/GBufferHelpers.ush
+++ b/Engine/Shaders/Private/GBufferHelpers.ush
@@ -374,10 +374,10 @@ void AdjustBaseColorAndSpecularColorForSubsurfaceProfileLightingCopyHack(inout f
 
 void GBufferPostDecode(inout FGBufferData Ret, bool bChecker, bool bGetNormalizedNormal)
 {
-	Ret.CustomData = HasCustomGBufferData(Ret.ShadingModelID) ? Ret.CustomData : 0.0h;
+	Ret.CustomData = HasCustomGBufferData(Ret.ShadingModelID) ? Ret.CustomData : half(0.0f);
 
-	Ret.PrecomputedShadowFactors = !(Ret.SelectiveOutputMask & 0x2) ? Ret.PrecomputedShadowFactors : ((Ret.SelectiveOutputMask & 0x4) ? 0.0h : 1.0h);
-	Ret.Velocity = !(Ret.SelectiveOutputMask & 0x8) ? Ret.Velocity : 0.0h;
+	Ret.PrecomputedShadowFactors = !(Ret.SelectiveOutputMask & 0x2) ? Ret.PrecomputedShadowFactors : ((Ret.SelectiveOutputMask & 0x4) ? half(0.0f) : half(1.0f));
+	Ret.Velocity = !(Ret.SelectiveOutputMask & 0x8) ? Ret.Velocity : half(0.0f);
 	bool bHasAnisotropy = (Ret.SelectiveOutputMask & 0x1);
 
 	Ret.StoredBaseColor = Ret.BaseColor;
