Fix hair strands voxelization visualization by prunning raymarching with scene depth.

diff --git a/Engine/Shaders/Private/HairStrands/HairStrandsVoxelDebug.usf b/Engine/Shaders/Private/HairStrands/HairStrandsVoxelDebug.usf
--- a/Engine/Shaders/Private/HairStrands/HairStrandsVoxelDebug.usf
+++ b/Engine/Shaders/Private/HairStrands/HairStrandsVoxelDebug.usf
@@ -26,6 +26,7 @@ uint MacroGroupId;
 uint MacroGroupCount;
 uint MaxTotalPageIndexCount;
 
+Texture2D<float> SceneDepthBeforeCompositionTexture;
 RWTexture2D<float4> OutputTexture;
 Buffer<uint> TotalValidPageCounter;
 
@@ -270,8 +271,9 @@ void MainCS(uint3 DispatchThreadId : SV_DispatchThreadID)
 	if (bDrawPage)
 	{
 		// Reverse-Z
+		const float SceneDepth = ConvertFromDeviceZ(SceneDepthBeforeCompositionTexture.Load(uint3(PixelCoord, 0)));
 		const float3 O = PrimaryView.TranslatedWorldCameraOrigin;
-		const float3 E = ReconstructTranslatedWorldPositionFromDepth(UV, ConvertFromDeviceZ(0.0001f));
+		const float3 E = ReconstructTranslatedWorldPositionFromDepth(UV, SceneDepth);
 		const float3 D = (E - O);
 
 		float PixelRadius = -1;
@@ -330,10 +332,10 @@ void MainCS(uint3 DispatchThreadId : SV_DispatchThreadID)
 		const float DistanceThreshold = 10000;
 		const bool bDebugEnabled = false;
 		const float3 SampleRandom = GetHairVoxelJitter(PixelCoord, View.StateFrameIndexMod8, VirtualVoxel.JitterMode);
+		const float SceneDepth = ConvertFromDeviceZ(SceneDepthBeforeCompositionTexture.Load(uint3(PixelCoord, 0)));
 
 		const float3 TranslatedWP0 = PrimaryView.TranslatedWorldCameraOrigin;
-		const float3 E = ReconstructTranslatedWorldPositionFromDepth(UV, ConvertFromDeviceZ(0.0001f));
-		const float3 TranslatedWP1 = normalize(E - TranslatedWP0) * DistanceThreshold + TranslatedWP0;
+		const float3 TranslatedWP1 = ReconstructTranslatedWorldPositionFromDepth(UV, SceneDepth);
 
 		FVirtualVoxelCommonDesc CommonDesc;
 		CommonDesc.PageCountResolution	= VirtualVoxel.PageCountResolution;
diff --git a/Engine/Source/Runtime/Renderer/Private/HairStrands/HairStrandsData.h b/Engine/Source/Runtime/Renderer/Private/HairStrands/HairStrandsData.h
--- a/Engine/Source/Runtime/Renderer/Private/HairStrands/HairStrandsData.h
+++ b/Engine/Source/Runtime/Renderer/Private/HairStrands/HairStrandsData.h
@@ -450,6 +450,11 @@ struct FHairStrandsDebugData
 		FConvexVolume ViewFrustum;
 	} CullData;
 
+	struct FCommon
+	{
+		FRDGTextureRef SceneDepthTextureBeforeCompsition = nullptr;
+	} Common;
+
 };
 
 typedef TArray<FHairStrandsMacroGroupData, SceneRenderingAllocator> FHairStrandsMacroGroupDatas;
diff --git a/Engine/Source/Runtime/Renderer/Private/HairStrands/HairStrandsDebug.cpp b/Engine/Source/Runtime/Renderer/Private/HairStrands/HairStrandsDebug.cpp
--- a/Engine/Source/Runtime/Renderer/Private/HairStrands/HairStrandsDebug.cpp
+++ b/Engine/Source/Runtime/Renderer/Private/HairStrands/HairStrandsDebug.cpp
@@ -724,6 +724,7 @@ class FVoxelVirtualRaymarchingCS : public FGlobalShader
 		SHADER_PARAMETER(uint32, MacroGroupId)
 		SHADER_PARAMETER(uint32, MacroGroupCount)
 		SHADER_PARAMETER(uint32, MaxTotalPageIndexCount)
+		SHADER_PARAMETER_RDG_TEXTURE(Texture2D, SceneDepthBeforeCompositionTexture)
 		SHADER_PARAMETER_RDG_UNIFORM_BUFFER(FVirtualVoxelParameters, VirtualVoxel)
 		SHADER_PARAMETER_STRUCT_REF(FViewUniformShaderParameters, ViewUniformBuffer)
 		SHADER_PARAMETER_RDG_BUFFER_SRV(Buffer, TotalValidPageCounter)
@@ -771,6 +772,11 @@ static void AddVoxelPageRaymarchingPass(
 		Parameters->TotalValidPageCounter	= GraphBuilder.CreateSRV(VoxelResources.PageIndexGlobalCounter, PF_R32_UINT);
 		ShaderPrint::SetParameters(GraphBuilder, View.ShaderPrintData, Parameters->ShaderPrintParameters);
 		Parameters->OutputTexture			= GraphBuilder.CreateUAV(OutputTexture);
+		Parameters->SceneDepthBeforeCompositionTexture = View.HairStrandsViewData.DebugData.Common.SceneDepthTextureBeforeCompsition;
+		if (Parameters->SceneDepthBeforeCompositionTexture == nullptr)
+		{
+			Parameters->SceneDepthBeforeCompositionTexture = SceneTextures.SceneDepthTexture;
+		}
 
 		FVoxelVirtualRaymarchingCS::FPermutationDomain PermutationVector;
 		PermutationVector.Set<FVoxelVirtualRaymarchingCS::FTraversalType>(GHairVirtualVoxel_DebugTraversalType > 0 ? 1 : 0);
diff --git a/Engine/Source/Runtime/Renderer/Private/HairStrands/HairStrandsVisibility.cpp b/Engine/Source/Runtime/Renderer/Private/HairStrands/HairStrandsVisibility.cpp
--- a/Engine/Source/Runtime/Renderer/Private/HairStrands/HairStrandsVisibility.cpp
+++ b/Engine/Source/Runtime/Renderer/Private/HairStrands/HairStrandsVisibility.cpp
@@ -2563,6 +2563,13 @@ static void AddHairMaterialDataPatchPass(
 	const bool bLightingChannel = false;
 #endif
 
+	if (GetGroomViewMode(View) == EGroomViewMode::VoxelsDensity && OutDepthTexture)
+	{
+		FHairStrandsDebugData* DebugData = const_cast<FHairStrandsDebugData*>(&View.HairStrandsViewData.DebugData);
+		DebugData->Common.SceneDepthTextureBeforeCompsition = GraphBuilder.CreateTexture(OutDepthTexture->Desc, TEXT("Hair.SceneDepthBeforeCompositionForDebug"));
+		AddCopyTexturePass(GraphBuilder, OutDepthTexture, DebugData->Common.SceneDepthTextureBeforeCompsition);
+	}
+
 	AddHairAuxilaryPass(
 		GraphBuilder,
 		View,
