- Removed some NDA'd platform references in shader code

diff --git a/Engine/Shaders/Private/SceneTexturesCommon.ush b/Engine/Shaders/Private/SceneTexturesCommon.ush
--- a/Engine/Shaders/Private/SceneTexturesCommon.ush
+++ b/Engine/Shaders/Private/SceneTexturesCommon.ush
@@ -191,7 +191,7 @@ float LookupDeviceZ( float2 ScreenUV )
 {
 #if	SCENE_TEXTURES_DISABLED
 	return FarDepthValue;
-#elif FORCE_DEPTH_TEXTURE_READS || COMPILER_SWITCH
+#elif FORCE_DEPTH_TEXTURE_READS || PLATFORM_NEEDS_DEPTH_TEXTURE_READS
 	// native Depth buffer lookup
 	return Texture2DSampleLevel(MobileSceneTextures.SceneDepthTexture, MobileSceneTextures.SceneDepthTextureSampler, ScreenUV, 0).r;
 #elif MOBILE_DEPTHFECTH && COMPILER_GLSL_ES3_1
diff --git a/Engine/Shaders/Private/ShadowProjectionCommon.ush b/Engine/Shaders/Private/ShadowProjectionCommon.ush
--- a/Engine/Shaders/Private/ShadowProjectionCommon.ush
+++ b/Engine/Shaders/Private/ShadowProjectionCommon.ush
@@ -16,8 +16,9 @@ float4 ShadowBufferSize;
 /** Cube map texture. */
 TextureCube ShadowDepthCubeTexture;
 
-#if OPENGL_PROFILE || COMPILER_VULKAN || COMPILER_SWITCH
 // This is required on GLSL based languages as the standard does NOT allow having a texture being used with both a Comparison and regular SamplerState
+#define USE_SEPARATE_SHADOW_DEPTH_CUBE_TEXTURE (OPENGL_PROFILE || COMPILER_VULKAN || PLATFORM_NEEDS_SEPARATE_SHADOW_DEPTH_CUBE_TEXTURE)
+#if USE_SEPARATE_SHADOW_DEPTH_CUBE_TEXTURE
 TextureCube ShadowDepthCubeTexture2;
 #endif
 
diff --git a/Engine/Shaders/Private/ShadowProjectionPixelShader.usf b/Engine/Shaders/Private/ShadowProjectionPixelShader.usf
--- a/Engine/Shaders/Private/ShadowProjectionPixelShader.usf
+++ b/Engine/Shaders/Private/ShadowProjectionPixelShader.usf
@@ -421,7 +421,7 @@ void MainOnePassPointLightPS(
 	float3 WorldSampleToLightVec = LightPositionAndInvRadius.xyz - TranslateWorldPosition.xyz;
 #if OUTPUT_CUBE_SHADOW_DEPTH_NO_FILTERING
 	// Note: point light shadow depth is Z / W, not linear
-#if OPENGL_PROFILE || COMPILER_VULKAN || COMPILER_SWITCH
+#if USE_SEPARATE_SHADOW_DEPTH_CUBE_TEXTURE
 	float Shadow = TextureCubeSampleLevel(ShadowDepthCubeTexture2, ShadowDepthTextureSampler, WorldSampleToLightVec, 0);
 #else
 	float Shadow = TextureCubeSampleLevel(ShadowDepthCubeTexture, ShadowDepthTextureSampler, WorldSampleToLightVec, 0);
diff --git a/Engine/Shaders/Private/Strata/Strata.ush b/Engine/Shaders/Private/Strata/Strata.ush
--- a/Engine/Shaders/Private/Strata/Strata.ush
+++ b/Engine/Shaders/Private/Strata/Strata.ush
@@ -1723,7 +1723,7 @@ FRWStrataMaterialContainer InitialiseRWStrataMaterialContainer(RWTexture2DArray<
 
 void StrataStoreUint1(inout FRWStrataMaterialContainer StrataBuffer, RWTexture2DArray<uint> ExtraMaterialDataUAV, inout FStrataAddressing StrataAddressing, uint Data)
 {
-  #if !COMPILER_SWITCH
+  #if PLATFORM_SUPPORTS_STRATA_UINT1
 	// We heavily rely on the compiler to optimize out the if code
 	if (StrataAddressing.CurrentIndex < STRATA_BASE_PASS_MRT_OUTPUT_COUNT)
 	{
@@ -1739,7 +1739,7 @@ void StrataStoreUint1(inout FRWStrataMaterialContainer StrataBuffer, RWTexture2D
 		ExtraMaterialDataUAV[uint3(StrataAddressing.PixelCoords, StrataAddressing.CurrentIndex - STRATA_BASE_PASS_MRT_OUTPUT_COUNT)] = Data;
 	}
 	StrataAddressing.CurrentIndex++;
-  #endif // !COMPILER_SWITCH
+  #endif // PLATFORM_SUPPORTS_STRATA_UINT1
 }
 
 #define STRATA_STORE_UINT1(x)	StrataStoreUint1(StrataBuffer, ExtraMaterialDataUAV, StrataAddressing, x)
@@ -1797,7 +1797,7 @@ void FinalizeWrites(inout FRWStrataMaterialContainer StrataBuffer, inout FStrata
 uint StrataLoadUint1(FStrataMaterialContainer StrataBuffer, inout FStrataAddressing StrataAddressing)
 {
 	uint Data = 0;
-  #if !COMPILER_SWITCH
+  #if PLATFORM_SUPPORTS_STRATA_UINT1
 	#if STRATA_MATERIALCONTAINER_IS_WRITABLE
 	Data = StrataBuffer[uint3(StrataAddressing.PixelCoords, StrataAddressing.CurrentIndex)];
 	#else
@@ -1813,7 +1813,7 @@ uint StrataLoadUint1(FStrataMaterialContainer StrataBuffer, inout FStrataAddress
 uint StrataLoadUint1AtIndexOffset(in FStrataMaterialContainer StrataBuffer, in FStrataAddressing StrataAddressing, uint IndexOffset)
 {
 	uint Data = 0;
-  #if !COMPILER_SWITCH
+  #if PLATFORM_SUPPORTS_STRATA_UINT1
 	#if STRATA_MATERIALCONTAINER_IS_WRITABLE
 	Data = StrataBuffer[uint3(StrataAddressing.PixelCoords, IndexOffset)];
 	#else
diff --git a/Engine/Shaders/Private/TransmissionThickness.ush b/Engine/Shaders/Private/TransmissionThickness.ush
--- a/Engine/Shaders/Private/TransmissionThickness.ush
+++ b/Engine/Shaders/Private/TransmissionThickness.ush
@@ -205,7 +205,7 @@ float CalcTransmissionThickness(float3 WorldNormal, FTransmissionProfileParams T
 		{
 			float2 DeltUv = DiscSamples5[i] * ShadowFilterRadius;
 			float3 SamplePos = NormalizedLightVector + SideVector * DeltUv.x + UpVector * DeltUv.y;
-	#if OPENGL_PROFILE || COMPILER_VULKAN || COMPILER_SWITCH
+	#if USE_SEPARATE_SHADOW_DEPTH_CUBE_TEXTURE
 			float ShadowSample = TextureCubeSampleLevel(ShadowDepthCubeTexture2, ShadowDepthTextureSampler, SamplePos, 0.0).r;
 	#else
             float ShadowSample = TextureCubeSampleDepthLevel(ShadowDepthCubeTexture, ShadowDepthTextureSampler, SamplePos, 0.0);
diff --git a/Engine/Shaders/Public/Platform.ush b/Engine/Shaders/Public/Platform.ush
--- a/Engine/Shaders/Public/Platform.ush
+++ b/Engine/Shaders/Public/Platform.ush
@@ -277,6 +277,26 @@ float2 UnpackFloat2FromUInt(uint In)
 	#define UB_CB_NAME(X) X
 #endif
 
+#ifndef PLATFORM_NEEDS_DEPTH_TEXTURE_READS
+#define PLATFORM_NEEDS_DEPTH_TEXTURE_READS 0
+#endif
+
+#ifndef PLATFORM_NEEDS_SELECT_UINT
+#define PLATFORM_NEEDS_SELECT_UINT 0
+#endif
+
+#ifndef PLATFORM_NEEDS_SEPARATE_SHADOW_DEPTH_CUBE_TEXTURE
+#define PLATFORM_NEEDS_SEPARATE_SHADOW_DEPTH_CUBE_TEXTURE 0
+#endif
+
+#ifndef PLATFORM_SUPPORTS_STRATA_UINT1
+#define PLATFORM_SUPPORTS_STRATA_UINT1 1
+#endif
+
+#ifndef COMPILER_NEEDS_DETERMINANT
+#define COMPILER_NEEDS_DETERMINANT 0
+#endif
+
 // ---------------------------------------------------- Alternative floating point types
 
 #ifndef FORCE_FLOATS
@@ -803,7 +823,7 @@ void WriteDispatchIndirectArgs(RWBuffer<uint> RWIndirectDispatchArgsBuffer, in u
 
 // ---------------------------------------------------- Compiler missing implementations
 
-#if COMPILER_SWITCH
+#if COMPILER_NEEDS_DETERMINANT
 
 	float determinant(float3x3 M)
 	{
@@ -988,7 +1008,7 @@ half4 max3( half4 a, half4 b, half4 c )
 	TYPE##4 select_internal(bool4   c, TYPE##4 a, TYPE##4 b) { return TYPE##4(c.x ? a.x : b.x, c.y ? a.y : b.y, c.z ? a.z : b.z, c.w ? a.w : b.w); } \
 
 DEFINE_SELECT(bool)
-#if !(COMPILER_HLSL || COMPILER_DXC || COMPILER_HLSLCC) || COMPILER_SWITCH
+#if !(COMPILER_HLSL || COMPILER_DXC || COMPILER_HLSLCC) || PLATFORM_NEEDS_SELECT_UINT
 // @todo-lh: Ambiguous for DXC and HLSLcc if no suffix is provided for integer literals
 DEFINE_SELECT(uint)
 #endif
