Fixes TSR ghosting in the alpha channel

diff --git a/Engine/Shaders/Private/TemporalSuperResolution/TSRUpdateHistory.usf b/Engine/Shaders/Private/TemporalSuperResolution/TSRUpdateHistory.usf
--- a/Engine/Shaders/Private/TemporalSuperResolution/TSRUpdateHistory.usf
+++ b/Engine/Shaders/Private/TemporalSuperResolution/TSRUpdateHistory.usf
@@ -1395,10 +1395,13 @@ void MainCS(
 
 	{
 		tsr_halfCx2 ClampedPrevHighFrequencyColor = clamp(PrevHighFrequencyColor, InputMinColor, InputMaxColor);
-		tsr_halfCx2 BlendedPrevHighFrequencyColor = PrevHighFrequencyColor;
+		tsr_halfCx2 BlendedPrevHighFrequencyColor;
 		BlendedPrevHighFrequencyColor[0] = lerp(ClampedPrevHighFrequencyColor[0], PrevHighFrequencyColor[0], BlendClamp);
 		BlendedPrevHighFrequencyColor[1] = lerp(ClampedPrevHighFrequencyColor[1], PrevHighFrequencyColor[1], BlendClamp);
 		BlendedPrevHighFrequencyColor[2] = lerp(ClampedPrevHighFrequencyColor[2], PrevHighFrequencyColor[2], BlendClamp);
+		#if CONFIG_SCENE_COLOR_ALPHA
+			BlendedPrevHighFrequencyColor[3] = lerp(ClampedPrevHighFrequencyColor[3], PrevHighFrequencyColor[3], BlendClamp);
+		#endif
 		
 		// Clamp the validity due to motion to maintain better sharpness in history reprojection under motion.
 		#if !CONFIG_GRAND_REPROJECTION
@@ -1482,7 +1485,7 @@ void MainCS(
 				PrevResultant[0] = clamp(GrandPrevToPrevResultant[0] * SafeRcp(GrandPrevToPrev[0]), tsr_half(0.0), tsr_half(1.0));
 				PrevResultant[1] = clamp(GrandPrevToPrevResultant[1] * SafeRcp(GrandPrevToPrev[1]), tsr_half(0.0), tsr_half(1.0));
 				PrevResultant[2] = clamp(GrandPrevToPrevResultant[2] * SafeRcp(GrandPrevToPrev[2]), tsr_half(0.0), tsr_half(1.0));
-				#if CONFIG_SCENE_COLOR_APLHA
+				#if CONFIG_SCENE_COLOR_ALPHA
 					PrevResultant[3] = clamp(GrandPrevToPrevResultant[3] * SafeRcp(GrandPrevToPrev[3]), tsr_half(0.0), tsr_half(1.0));
 				#endif
 
