Eliminate loss precision warnings while compiling shaders.

diff --git a/Engine/Shaders/Private/DebugViewModeVertexShader.usf b/Engine/Shaders/Private/DebugViewModeVertexShader.usf
--- a/Engine/Shaders/Private/DebugViewModeVertexShader.usf
+++ b/Engine/Shaders/Private/DebugViewModeVertexShader.usf
@@ -61,7 +61,7 @@ void Main(
 	FVertexFactoryIntermediates VFIntermediates = GetVertexFactoryIntermediates(Input);
 	float4 WorldPos = VertexFactoryGetWorldPosition(Input, VFIntermediates);
 
-	float3x3 TangentToLocal = VertexFactoryGetTangentToLocal(Input, VFIntermediates);
+	half3x3 TangentToLocal = VertexFactoryGetTangentToLocal(Input, VFIntermediates);
 	FMaterialVertexParameters VertexParameters = GetMaterialVertexParameters(Input, VFIntermediates, WorldPos.xyz, TangentToLocal);
 
 	// Isolate instructions used for world position offset
diff --git a/Engine/Shaders/Private/DepthOnlyVertexShader.usf b/Engine/Shaders/Private/DepthOnlyVertexShader.usf
--- a/Engine/Shaders/Private/DepthOnlyVertexShader.usf
+++ b/Engine/Shaders/Private/DepthOnlyVertexShader.usf
@@ -50,7 +50,7 @@ void Main(
 	float4 WorldPos = VertexFactoryGetWorldPosition(Input, VFIntermediates);
 	float4 WorldPositionExcludingWPO = WorldPos;
 
-	float3x3 TangentToLocal = VertexFactoryGetTangentToLocal(Input, VFIntermediates);
+	half3x3 TangentToLocal = VertexFactoryGetTangentToLocal(Input, VFIntermediates);
 	FMaterialVertexParameters VertexParameters = GetMaterialVertexParameters(Input, VFIntermediates, WorldPos.xyz, TangentToLocal);
 
 	// Isolate instructions used for world position offset
diff --git a/Engine/Shaders/Private/HeightFogCommon.ush b/Engine/Shaders/Private/HeightFogCommon.ush
--- a/Engine/Shaders/Private/HeightFogCommon.ush
+++ b/Engine/Shaders/Private/HeightFogCommon.ush
@@ -217,7 +217,7 @@ half4 CalculateHeightFog(float3 WorldPositionRelativeToCamera)
 	ExcludeDistance = View.VolumetricFogMaxDistance * InvCosAngle;
 #endif
 
-	float4 FogInscatteringAndOpacity = GetExponentialHeightFog(WorldPositionRelativeToCamera, ExcludeDistance);
+	half4 FogInscatteringAndOpacity = GetExponentialHeightFog(WorldPositionRelativeToCamera, ExcludeDistance);
 	return FogInscatteringAndOpacity;
 }
 
diff --git a/Engine/Shaders/Private/HitProxyVertexShader.usf b/Engine/Shaders/Private/HitProxyVertexShader.usf
--- a/Engine/Shaders/Private/HitProxyVertexShader.usf
+++ b/Engine/Shaders/Private/HitProxyVertexShader.usf
@@ -59,7 +59,7 @@ void Main(
 	FVertexFactoryIntermediates VFIntermediates = GetVertexFactoryIntermediates(Input);
 	float4 WorldPosition = VertexFactoryGetWorldPosition(Input, VFIntermediates);
 	float4 WorldPositionExcludingWPO = WorldPosition;
-	float3x3 TangentToLocal = VertexFactoryGetTangentToLocal(Input, VFIntermediates);
+	half3x3 TangentToLocal = VertexFactoryGetTangentToLocal(Input, VFIntermediates);
 
 	FMaterialVertexParameters VertexParameters = GetMaterialVertexParameters(Input, VFIntermediates, WorldPosition.xyz, TangentToLocal);
 	WorldPosition.xyz += GetMaterialWorldPositionOffset(VertexParameters);
diff --git a/Engine/Shaders/Private/MobileBasePassVertexShader.usf b/Engine/Shaders/Private/MobileBasePassVertexShader.usf
--- a/Engine/Shaders/Private/MobileBasePassVertexShader.usf
+++ b/Engine/Shaders/Private/MobileBasePassVertexShader.usf
@@ -69,7 +69,7 @@ void Main(
 	half3x3 TangentToLocal = VertexFactoryGetTangentToLocal(Input, VFIntermediates);	
 	FMaterialVertexParameters VertexParameters = GetMaterialVertexParameters(Input, VFIntermediates, WorldPosition.xyz, TangentToLocal);
 
-	half3 WorldPositionOffset = GetMaterialWorldPositionOffset(VertexParameters);
+	float3 WorldPositionOffset = GetMaterialWorldPositionOffset(VertexParameters);
 	
 	WorldPosition.xyz += WorldPositionOffset;
 
@@ -86,7 +86,7 @@ void Main(
 #endif
 
 #if USE_VERTEX_FOG
-	float4 VertexFog = CalculateHeightFog(WorldPosition.xyz - ResolvedView.TranslatedWorldCameraOrigin);
+	half4 VertexFog = CalculateHeightFog(WorldPosition.xyz - ResolvedView.TranslatedWorldCameraOrigin);
 
 #if PROJECT_SUPPORT_SKY_ATMOSPHERE && MATERIAL_IS_SKY==0 // Do not apply aerial perpsective on sky materials
 	if (ResolvedView.SkyAtmosphereApplyCameraAerialPerspectiveVolume > 0.0f)
diff --git a/Engine/Shaders/Private/ShadowDepthVertexShader.usf b/Engine/Shaders/Private/ShadowDepthVertexShader.usf
--- a/Engine/Shaders/Private/ShadowDepthVertexShader.usf
+++ b/Engine/Shaders/Private/ShadowDepthVertexShader.usf
@@ -185,7 +185,7 @@ void Main(
 	
 	FVertexFactoryIntermediates VFIntermediates = GetVertexFactoryIntermediates(Input);
 	float4 WorldPos = VertexFactoryGetWorldPosition(Input, VFIntermediates);
-	float3x3 TangentToLocal = VertexFactoryGetTangentToLocal(Input, VFIntermediates);
+	half3x3 TangentToLocal = VertexFactoryGetTangentToLocal(Input, VFIntermediates);
 
 	FMaterialVertexParameters VertexParameters = GetMaterialVertexParameters(Input, VFIntermediates, WorldPos.xyz, TangentToLocal);
 	const float3 WorldNormal = VertexFactoryGetWorldNormal(Input, VFIntermediates);
diff --git a/Engine/Shaders/Private/VelocityShader.usf b/Engine/Shaders/Private/VelocityShader.usf
--- a/Engine/Shaders/Private/VelocityShader.usf
+++ b/Engine/Shaders/Private/VelocityShader.usf
@@ -57,7 +57,7 @@ void MainVertexShader(
 	
 	FVertexFactoryIntermediates VFIntermediates = GetVertexFactoryIntermediates(Input);
 	float4 TranslatedWorldPosition = VertexFactoryGetWorldPosition( Input, VFIntermediates );
-	float3x3 TangentToLocal = VertexFactoryGetTangentToLocal( Input, VFIntermediates );
+	half3x3 TangentToLocal = VertexFactoryGetTangentToLocal( Input, VFIntermediates );
 
 	FMaterialVertexParameters VertexParameters = GetMaterialVertexParameters(Input, VFIntermediates, TranslatedWorldPosition.xyz, TangentToLocal);
 	// Isolate instructions used for world position offset on xbox 360, 
