Fixed IBL reflections for Surface: LightingVolume(contributed by [REDACTED_PERSON])

#10256

diff --git a/Engine/Shaders/Private/MobileBasePassPixelShader.usf b/Engine/Shaders/Private/MobileBasePassPixelShader.usf
--- a/Engine/Shaders/Private/MobileBasePassPixelShader.usf
+++ b/Engine/Shaders/Private/MobileBasePassPixelShader.usf
@@ -803,7 +803,7 @@ void Main(
 	float2 LocalPosition = SvPosition.xy - ResolvedView.ViewRectMin.xy;
 	uint GridIndex = ComputeLightGridCellIndex(uint2(LocalPosition.x, LocalPosition.y), SvPosition.w);
 
-#if MATERIALBLENDING_MASKED || MATERIALBLENDING_SOLID || TRANSLUCENCY_LIGHTING_SURFACE_FORWARDSHADING || MATERIAL_SHADINGMODEL_SINGLELAYERWATER
+#if MATERIALBLENDING_MASKED || MATERIALBLENDING_SOLID || TRANSLUCENCY_LIGHTING_SURFACE_FORWARDSHADING || TRANSLUCENCY_LIGHTING_SURFACE_LIGHTINGVOLUME || MATERIAL_SHADINGMODEL_SINGLELAYERWATER
 	// Reflection IBL
 	AccumulateReflection(GBuffer
 		, CameraVector
