Implicitly passes down GroupThreadIndex through a static var in TSRRejectShading.usf

diff --git a/Engine/Shaders/Private/TemporalSuperResolution/TSRRejectShading.usf b/Engine/Shaders/Private/TemporalSuperResolution/TSRRejectShading.usf
--- a/Engine/Shaders/Private/TemporalSuperResolution/TSRRejectShading.usf
+++ b/Engine/Shaders/Private/TemporalSuperResolution/TSRRejectShading.usf
@@ -115,6 +115,11 @@ RWTexture2D<tsr_half2> HistoryRejectionOutput;
 RWTexture2D<tsr_half> InputSceneColorLdrLumaOutput;
 
 
+//------------------------------------------------------- GLOBAL VARIABLE
+
+static uint GGroupThreadIndex = 0;
+
+
 //------------------------------------------------------- GENERAL TEXEL SWIZZLING & NEIGHBOR ACCESS
 
 CALL_SITE_DEBUGLOC
@@ -139,69 +144,69 @@ tsr_short2 GetLaneSimdPixelOffset(const uint LaneIndex, const uint SimdIndex)
 
 groupshared tsr_half4 SharedData[GROUP_TILE_SIZE * GROUP_TILE_SIZE];
 
-void WriteAllElementToLDS(uint GroupThreadIndex, TLaneVector<tsr_half, 1, SIMD_SIZE> X)
+void WriteAllElementToLDS(TLaneVector<tsr_half, 1, SIMD_SIZE> X)
 {
 	GroupMemoryBarrierWithGroupSync();
 	UNROLL_N(SIMD_SIZE)
 	for (uint SimdIndex = 0; SimdIndex < SIMD_SIZE; SimdIndex++)
 	{
-		SharedData[GroupThreadIndex + SimdIndex * LANE_COUNT] = tsr_half4(X.GetElement(SimdIndex), 0.0, 0.0, 0.0);
+		SharedData[GGroupThreadIndex + SimdIndex * LANE_COUNT] = tsr_half4(X.GetElement(SimdIndex), 0.0, 0.0, 0.0);
 	}
 	GroupMemoryBarrierWithGroupSync();
 }
 
-void WriteAllElementToLDS(uint GroupThreadIndex, TLaneVector<tsr_half, 3, SIMD_SIZE> X)
+void WriteAllElementToLDS(TLaneVector<tsr_half, 3, SIMD_SIZE> X)
 {
 	GroupMemoryBarrierWithGroupSync();
 	UNROLL_N(SIMD_SIZE)
 	for (uint SimdIndex = 0; SimdIndex < SIMD_SIZE; SimdIndex++)
 	{
-		SharedData[GroupThreadIndex + SimdIndex * LANE_COUNT] = tsr_half4(X.GetElement(SimdIndex), 0.0);
+		SharedData[GGroupThreadIndex + SimdIndex * LANE_COUNT] = tsr_half4(X.GetElement(SimdIndex), 0.0);
 	}
 	GroupMemoryBarrierWithGroupSync();
 }
 
-void WriteAllElementToLDS(uint GroupThreadIndex, TLaneVector<tsr_half, 4, SIMD_SIZE> X)
+void WriteAllElementToLDS(TLaneVector<tsr_half, 4, SIMD_SIZE> X)
 {
 	GroupMemoryBarrierWithGroupSync();
 	UNROLL_N(SIMD_SIZE)
 	for (uint SimdIndex = 0; SimdIndex < SIMD_SIZE; SimdIndex++)
 	{
-		SharedData[GroupThreadIndex + SimdIndex * LANE_COUNT] = tsr_half4(X.GetElement(SimdIndex));
+		SharedData[GGroupThreadIndex + SimdIndex * LANE_COUNT] = tsr_half4(X.GetElement(SimdIndex));
 	}
 	GroupMemoryBarrierWithGroupSync();
 }
 
-uint GetElementLDSIndex(uint GroupThreadIndex, const uint NeigborSimdIndex, const uint LaneRotation)
+uint GetElementLDSIndex(const uint NeigborSimdIndex, const uint LaneRotation)
 {
-	return ((GroupThreadIndex + LaneRotation) % uint(LANE_COUNT)) + NeigborSimdIndex * LANE_COUNT;
+	return ((GGroupThreadIndex + LaneRotation) % uint(LANE_COUNT)) + NeigborSimdIndex * LANE_COUNT;
 }
 
-void ReadElementFromLDS(uint GroupThreadIndex, const uint NeigborSimdIndex, const uint LaneRotation, out tsr_half X)
+void ReadElementFromLDS(const uint NeigborSimdIndex, const uint LaneRotation, out tsr_half X)
 {
-	X = SharedData[GetElementLDSIndex(GroupThreadIndex, NeigborSimdIndex, LaneRotation)].x;
+	X = SharedData[GetElementLDSIndex(NeigborSimdIndex, LaneRotation)].x;
 }
 
-void ReadElementFromLDS(uint GroupThreadIndex, const uint NeigborSimdIndex, const uint LaneRotation, out tsr_half3 X)
+void ReadElementFromLDS(const uint NeigborSimdIndex, const uint LaneRotation, out tsr_half3 X)
 {
-	X = SharedData[GetElementLDSIndex(GroupThreadIndex, NeigborSimdIndex, LaneRotation)].xyz;
+	X = SharedData[GetElementLDSIndex(NeigborSimdIndex, LaneRotation)].xyz;
 }
 
-void ReadElementFromLDS(uint GroupThreadIndex, const uint NeigborSimdIndex, const uint LaneRotation, out tsr_half4 X)
+void ReadElementFromLDS(const uint NeigborSimdIndex, const uint LaneRotation, out tsr_half4 X)
 {
-	X = SharedData[GetElementLDSIndex(GroupThreadIndex, NeigborSimdIndex, LaneRotation)].xyzw;
+	X = SharedData[GetElementLDSIndex(NeigborSimdIndex, LaneRotation)].xyzw;
 }
 
 template<typename ScalarType, uint VectorSize, uint SimdSize>
-void WriteToLDSForAccessNeighborTexel(TLaneVector<ScalarType, VectorSize, SimdSize> Input, uint GroupThreadIndex)
+void WriteToLDSForAccessNeighborTexel(TLaneVector<ScalarType, VectorSize, SimdSize> Input)
 {
-	WriteAllElementToLDS(GroupThreadIndex, Input);
+	WriteAllElementToLDS(Input);
 }
 
 #else // !CONFIG_OPTIMIZED_LDS_ACCESS
 
 template<typename ScalarType, uint VectorSize, uint SimdSize>
-void WriteToLDSForAccessNeighborTexel(TLaneVector<ScalarType, VectorSize, SimdSize> Input, uint GroupThreadIndex)
+void WriteToLDSForAccessNeighborTexel(TLaneVector<ScalarType, VectorSize, SimdSize> Input)
 {
 	// NOP
 }
@@ -211,7 +216,7 @@ void WriteToLDSForAccessNeighborTexel(TLaneVector<ScalarType, VectorSize, SimdSi
 
 CALL_SITE_DEBUGLOC
 template<typename ScalarType, uint VectorSize, uint SimdSize>
-TLaneVector<ScalarType, VectorSize, SimdSize> AccessNeighborTexel(TLaneVector<ScalarType, VectorSize, SimdSize> Input, const tsr_short2 Offset, uint GroupThreadIndex)
+TLaneVector<ScalarType, VectorSize, SimdSize> AccessNeighborTexel(TLaneVector<ScalarType, VectorSize, SimdSize> Input, const tsr_short2 Offset)
 {
 	TLaneVector<ScalarType, VectorSize, SimdSize> Return;
 	
@@ -234,7 +239,7 @@ TLaneVector<ScalarType, VectorSize, SimdSize> AccessNeighborTexel(TLaneVector<Sc
 		if (LaneRotation != 0)
 		#if CONFIG_OPTIMIZED_LDS_ACCESS
 		{
-			ReadElementFromLDS(GroupThreadIndex, NeigborSimdIndex, LaneRotation, /* out */ ReturnSimdElement);
+			ReadElementFromLDS(NeigborSimdIndex, LaneRotation, /* out */ ReturnSimdElement);
 		}
 		#else
 		{
@@ -245,7 +250,7 @@ TLaneVector<ScalarType, VectorSize, SimdSize> AccessNeighborTexel(TLaneVector<Sc
 				#endif
 				const FWaveBroadcastSettings BroadcastSettings = InitWaveRotateLaneGroup(/* LaneGroupSize = */ LANE_COUNT, LaneRotation);
 			#else
-				const FWaveBroadcastSettings BroadcastSettings = ConvertWaveBroadcastToLDS(InitWaveRotateLaneGroup(/* LaneGroupSize = */ LANE_COUNT, LaneRotation), GroupThreadIndex);
+				const FWaveBroadcastSettings BroadcastSettings = ConvertWaveBroadcastToLDS(InitWaveRotateLaneGroup(/* LaneGroupSize = */ LANE_COUNT, LaneRotation));
 			#endif
 			ReturnSimdElement = WaveBroadcast(BroadcastSettings, SimdElement);
 		}
@@ -378,24 +383,24 @@ struct TAccumulatorMax
 
 CALL_SITE_DEBUGLOC
 template<class FAccumulator, typename FSampleType>
-FAccumulator Convole3x1(FSampleType Color, tsr_simd_half Weight, const tsr_short2 Offset, uint GroupThreadIndex)
+FAccumulator Convole3x1(FSampleType Color, tsr_simd_half Weight, const tsr_short2 Offset)
 #if CONFIG_OPTIMIZED_LDS_ACCESS
 {
 	const tsr_short2 OffsetN = tsr_short2(-int2(Offset)); // work arround a shader compiler bug
 	
 	FSampleType PreprocessedNeighborColor = FAccumulator::PreprocessNeighborSample(Color, Weight);
 
-	WriteToLDSForAccessNeighborTexel(PreprocessedNeighborColor, GroupThreadIndex);
-	FSampleType Sample0 = AccessNeighborTexel(PreprocessedNeighborColor, Offset, GroupThreadIndex);
-	FSampleType Sample1 = AccessNeighborTexel(PreprocessedNeighborColor, OffsetN, GroupThreadIndex);
+	WriteToLDSForAccessNeighborTexel(PreprocessedNeighborColor);
+	FSampleType Sample0 = AccessNeighborTexel(PreprocessedNeighborColor, Offset);
+	FSampleType Sample1 = AccessNeighborTexel(PreprocessedNeighborColor, OffsetN);
 	
 	tsr_simd_half Weight0 = tsr_simd_half::Const(tsr_half(1.0));
 	tsr_simd_half Weight1 = tsr_simd_half::Const(tsr_half(1.0));
 	if (FAccumulator::NeedWeightInAccumulate())
 	{
-		WriteToLDSForAccessNeighborTexel(Weight, GroupThreadIndex);
-		Weight0 = AccessNeighborTexel(Weight, Offset, GroupThreadIndex);
-		Weight1 = AccessNeighborTexel(Weight, OffsetN, GroupThreadIndex);
+		WriteToLDSForAccessNeighborTexel(Weight);
+		Weight0 = AccessNeighborTexel(Weight, Offset);
+		Weight1 = AccessNeighborTexel(Weight, OffsetN);
 	}
 	
 	FAccumulator Accumulator;
@@ -412,20 +417,20 @@ FAccumulator Convole3x1(FSampleType Color, tsr_simd_half Weight, const tsr_short
 	
 	FSampleType PreprocessedNeighborColor = FAccumulator::PreprocessNeighborSample(Color, Weight);
 
-	FSampleType Sample0 = AccessNeighborTexel(PreprocessedNeighborColor, Offset, GroupThreadIndex);
+	FSampleType Sample0 = AccessNeighborTexel(PreprocessedNeighborColor, Offset);
 	tsr_simd_half Weight0 = tsr_simd_half::Const(tsr_half(1.0));
 	if (FAccumulator::NeedWeightInAccumulate())
 	{
-		Weight0 = AccessNeighborTexel(Weight, Offset, GroupThreadIndex);
+		Weight0 = AccessNeighborTexel(Weight, Offset);
 	}
 	Accumulator.Accumulate(+Offset, Sample0, Weight0);
 	
 	const tsr_short2 OffsetN = tsr_short2(-int2(Offset)); // work arround a shader compiler bug
-	FSampleType Sample1 = AccessNeighborTexel(PreprocessedNeighborColor, OffsetN, GroupThreadIndex);
+	FSampleType Sample1 = AccessNeighborTexel(PreprocessedNeighborColor, OffsetN);
 	tsr_simd_half Weight1 = tsr_simd_half::Const(tsr_half(1.0));
 	if (FAccumulator::NeedWeightInAccumulate())
 	{
-		Weight1 = AccessNeighborTexel(Weight, OffsetN, GroupThreadIndex);
+		Weight1 = AccessNeighborTexel(Weight, OffsetN);
 	}
 	Accumulator.Accumulate(OffsetN, Sample1, Weight1);
 
@@ -436,10 +441,10 @@ FAccumulator Convole3x1(FSampleType Color, tsr_simd_half Weight, const tsr_short
 
 CALL_SITE_DEBUGLOC
 template<class FAccumulator, typename FSampleType>
-FSampleType Convolve3x3HV(FSampleType Color, tsr_simd_half Weight, uint GroupThreadIndex)
+FSampleType Convolve3x3HV(FSampleType Color, tsr_simd_half Weight)
 {
-	FSampleType HorizontalColor = Convole3x1<FAccumulator, FSampleType>(Color, Weight, tsr_short2(1, 0), GroupThreadIndex).Value;
-	return Convole3x1<FAccumulator, FSampleType>(HorizontalColor, Weight, tsr_short2(0, 1), GroupThreadIndex).Value;
+	FSampleType HorizontalColor = Convole3x1<FAccumulator, FSampleType>(Color, Weight, tsr_short2(1, 0)).Value;
+	return Convole3x1<FAccumulator, FSampleType>(HorizontalColor, Weight, tsr_short2(0, 1)).Value;
 }
 
 
@@ -447,28 +452,28 @@ FSampleType Convolve3x3HV(FSampleType Color, tsr_simd_half Weight, uint GroupThr
 
 CALL_SITE_DEBUGLOC
 template<typename FSampleType>
-FSampleType Min3x3(FSampleType Center, uint GroupThreadIndex)
+FSampleType Min3x3(FSampleType Center)
 {
-	return Convolve3x3HV<TAccumulatorMin<FSampleType>, FSampleType>(Center, tsr_simd_half::Const(1.0), GroupThreadIndex);
+	return Convolve3x3HV<TAccumulatorMin<FSampleType>, FSampleType>(Center, tsr_simd_half::Const(1.0));
 }
 
 CALL_SITE_DEBUGLOC
 template<typename FSampleType>
-FSampleType Max3x3(FSampleType Center, uint GroupThreadIndex)
+FSampleType Max3x3(FSampleType Center)
 {
-	return Convolve3x3HV<TAccumulatorMax<FSampleType>, FSampleType>(Center, tsr_simd_half::Const(1.0), GroupThreadIndex);
+	return Convolve3x3HV<TAccumulatorMax<FSampleType>, FSampleType>(Center, tsr_simd_half::Const(1.0));
 }
 
 CALL_SITE_DEBUGLOC
-tsr_simd_halfC ClampFireFliersWithGuide(tsr_simd_halfC ToClamp, tsr_simd_halfC Guide, uint GroupThreadIndex)
+tsr_simd_halfC ClampFireFliersWithGuide(tsr_simd_halfC ToClamp, tsr_simd_halfC Guide)
 {
-	tsr_simd_halfC ToClampMin = Min3x3(ToClamp, GroupThreadIndex);
-	tsr_simd_halfC ToClampMax = Max3x3(ToClamp, GroupThreadIndex);
+	tsr_simd_halfC ToClampMin = Min3x3(ToClamp);
+	tsr_simd_halfC ToClampMax = Max3x3(ToClamp);
 
 	tsr_simd_halfC ClampedGuide = clamp(Guide, ToClampMin, ToClampMax);
 	
-	tsr_simd_halfC ClampedGuideMin = Min3x3(ClampedGuide, GroupThreadIndex);
-	tsr_simd_halfC ClampedGuideMax = Max3x3(ClampedGuide, GroupThreadIndex);
+	tsr_simd_halfC ClampedGuideMin = Min3x3(ClampedGuide);
+	tsr_simd_halfC ClampedGuideMax = Max3x3(ClampedGuide);
 	
 	return clamp(ToClamp, ClampedGuideMin, ClampedGuideMax);
 }
@@ -530,31 +535,31 @@ FSampleType Median9(FSampleType In[9])
 
 CALL_SITE_DEBUGLOC
 template<typename FSampleType>
-FSampleType Median3x3(FSampleType Center, uint GroupThreadIndex)
+FSampleType Median3x3(FSampleType Center)
 #if 1
 {
-	WriteToLDSForAccessNeighborTexel(Center, GroupThreadIndex);
+	WriteToLDSForAccessNeighborTexel(Center);
 	
 	FSampleType C0 = Center;
-	FSampleType C1 = AccessNeighborTexel(Center, tsr_short2( 1, 0), GroupThreadIndex);
-	FSampleType C2 = AccessNeighborTexel(Center, tsr_short2(-1, 0), GroupThreadIndex);
+	FSampleType C1 = AccessNeighborTexel(Center, tsr_short2( 1, 0));
+	FSampleType C2 = AccessNeighborTexel(Center, tsr_short2(-1, 0));
 	
 	FSampleType M0L0;
 	FSampleType M0M0;
 	FSampleType M0H0;
 	LMHOperator(C0, C1, C2, M0L0, M0M0, M0H0);
 	
-	WriteToLDSForAccessNeighborTexel(M0L0, GroupThreadIndex);
-	FSampleType M0L1 = AccessNeighborTexel(M0L0, tsr_short2(0,  1), GroupThreadIndex);
-	FSampleType M0L2 = AccessNeighborTexel(M0L0, tsr_short2(0, -1), GroupThreadIndex);
+	WriteToLDSForAccessNeighborTexel(M0L0);
+	FSampleType M0L1 = AccessNeighborTexel(M0L0, tsr_short2(0,  1));
+	FSampleType M0L2 = AccessNeighborTexel(M0L0, tsr_short2(0, -1));
 	
-	WriteToLDSForAccessNeighborTexel(M0M0, GroupThreadIndex);
-	FSampleType M0M1 = AccessNeighborTexel(M0M0, tsr_short2(0,  1), GroupThreadIndex);
-	FSampleType M0M2 = AccessNeighborTexel(M0M0, tsr_short2(0, -1), GroupThreadIndex);
+	WriteToLDSForAccessNeighborTexel(M0M0);
+	FSampleType M0M1 = AccessNeighborTexel(M0M0, tsr_short2(0,  1));
+	FSampleType M0M2 = AccessNeighborTexel(M0M0, tsr_short2(0, -1));
 	
-	WriteToLDSForAccessNeighborTexel(M0H0, GroupThreadIndex);
-	FSampleType M0H1 = AccessNeighborTexel(M0H0, tsr_short2(0,  1), GroupThreadIndex);
-	FSampleType M0H2 = AccessNeighborTexel(M0H0, tsr_short2(0, -1), GroupThreadIndex);
+	WriteToLDSForAccessNeighborTexel(M0H0);
+	FSampleType M0H1 = AccessNeighborTexel(M0H0, tsr_short2(0,  1));
+	FSampleType M0H2 = AccessNeighborTexel(M0H0, tsr_short2(0, -1));
 	
 	// Second layer.
 	// TODO: min3, max3.
@@ -568,13 +573,13 @@ FSampleType Median3x3(FSampleType Center, uint GroupThreadIndex)
 }
 #else
 {
-	WriteToLDSForAccessNeighborTexel(Center, GroupThreadIndex);
+	WriteToLDSForAccessNeighborTexel(Center);
 	
 	FSampleType InSamples[9];
 	UNROLL
 	for (uint i = 0; i < 9; i++)
 	{
-		FSampleType Sample = AccessNeighborTexel(Center, tsr_short2(kOffsets3x3[i]), GroupThreadIndex);
+		FSampleType Sample = AccessNeighborTexel(Center, tsr_short2(kOffsets3x3[i]));
 		InSamples[i] = Sample;
 	}
 
@@ -593,6 +598,8 @@ void MainCS(
 	uint2 GroupId : SV_GroupID,
 	uint GroupThreadIndex : SV_GroupIndex)
 {
+	GGroupThreadIndex = GroupThreadIndex;
+
 	tsr_simd_halfC Debug[DEBUG_ARRAY_SIZE];
 	UNROLL_N(DEBUG_ARRAY_SIZE)
 	for (uint DebugId = 0; DebugId < DEBUG_ARRAY_SIZE; DebugId++)
@@ -666,7 +673,7 @@ void MainCS(
 
 	#if CONFIG_COMPOSE_TRANSLUCENCY
 	{
-		tsr_simd_half4 FilteredTranslucencyInput = Convolve3x3HV<TAccumulatorAvg<tsr_simd_half4>, tsr_simd_half4>(OriginalTranslucencyInput, bIsValid, GroupThreadIndex);
+		tsr_simd_half4 FilteredTranslucencyInput = Convolve3x3HV<TAccumulatorAvg<tsr_simd_half4>, tsr_simd_half4>(OriginalTranslucencyInput, bIsValid);
 
 		UNROLL_N(SIMD_SIZE)
 		for (uint SimdIndex = 0; SimdIndex < SIMD_SIZE; SimdIndex++)
@@ -727,7 +734,7 @@ void MainCS(
 		#endif
 		
 		// Convolve the input luma to have the exact same moire pattern as FilteredInput below.
-		InputLuma = Convolve3x3HV<TAccumulatorAvg<tsr_simd_half>, tsr_simd_half>(InputLuma, bIsValid, GroupThreadIndex);
+		InputLuma = Convolve3x3HV<TAccumulatorAvg<tsr_simd_half>, tsr_simd_half>(InputLuma, bIsValid);
 
 		// Unpack history
 		tsr_simd_half PrevFlickeringHistory = PrevMoire[0];
@@ -742,8 +749,8 @@ void MainCS(
 
 
 		// Compute clamping box of the limance.
-		tsr_simd_half InputLumaMin = Min3x3(InputLuma, GroupThreadIndex);
-		tsr_simd_half InputLumaMax = Max3x3(InputLuma, GroupThreadIndex);
+		tsr_simd_half InputLumaMin = Min3x3(InputLuma);
+		tsr_simd_half InputLumaMax = Max3x3(InputLuma);
 
 		// Clamp and contribue the history.		
 		tsr_simd_half ClampedPrevFlickeringHistory = clamp(PrevFlickeringHistory, InputLumaMin, InputLumaMax);
@@ -783,8 +790,8 @@ void MainCS(
 
 		// Dilates the total variation and contribution to keep stability despite inacuracies in the history reprojection
 		// of these need a nearest filter which isn't precise at all.
-		tsr_simd_half TotalVariationContribution = Max3x3(GradientVariation, GroupThreadIndex);
-		tsr_simd_half TotalVariationCountContribution = Max3x3(IsFlicker, GroupThreadIndex);
+		tsr_simd_half TotalVariationContribution = Max3x3(GradientVariation);
+		tsr_simd_half TotalVariationCountContribution = Max3x3(IsFlicker);
 
 		UNROLL_N(SIMD_SIZE)
 		for (uint SimdIndex = 0; SimdIndex < SIMD_SIZE; SimdIndex++)
@@ -851,14 +858,14 @@ void MainCS(
 	tsr_simd_halfC ClampedHistory = OriginalHistory;
 	#if CONFIG_FIREFLY_REMOVAL
 	{
-		ClampedInput   = ClampFireFliersWithGuide(OriginalInput  , /* Guide = */ OriginalHistory, GroupThreadIndex);
-		ClampedHistory = ClampFireFliersWithGuide(OriginalHistory, /* Guide = */ OriginalInput  , GroupThreadIndex);
+		ClampedInput   = ClampFireFliersWithGuide(OriginalInput  , /* Guide = */ OriginalHistory);
+		ClampedHistory = ClampFireFliersWithGuide(OriginalHistory, /* Guide = */ OriginalInput);
 	}
 	#endif
 
 	// Filter clamped input and history to increase quantization of the rejection detection
-	tsr_simd_halfC FilteredInput   = Convolve3x3HV<TAccumulatorAvg<tsr_simd_halfC>, tsr_simd_halfC>(ClampedInput  , bIsValid, GroupThreadIndex);
-	tsr_simd_halfC FilteredHistory = Convolve3x3HV<TAccumulatorAvg<tsr_simd_halfC>, tsr_simd_halfC>(ClampedHistory, bIsValid, GroupThreadIndex);
+	tsr_simd_halfC FilteredInput   = Convolve3x3HV<TAccumulatorAvg<tsr_simd_halfC>, tsr_simd_halfC>(ClampedInput  , bIsValid);
+	tsr_simd_halfC FilteredHistory = Convolve3x3HV<TAccumulatorAvg<tsr_simd_halfC>, tsr_simd_halfC>(ClampedHistory, bIsValid);
 	
 	// Measure rejection
 	tsr_simd_half RejectionBlendFinal = tsr_simd_half::Const(1.0);
@@ -869,8 +876,8 @@ void MainCS(
 
 		tsr_simd_halfC BackbufferError = tsr_simd_halfC::Const(MeasureBackbufferLDRQuantizationError());
 
-		tsr_simd_halfC FilteredBoxMin = Min3x3(FilteredInput, GroupThreadIndex) - BackbufferError;
-		tsr_simd_halfC FilteredBoxMax = Max3x3(FilteredInput, GroupThreadIndex) + BackbufferError;
+		tsr_simd_halfC FilteredBoxMin = Min3x3(FilteredInput) - BackbufferError;
+		tsr_simd_halfC FilteredBoxMax = Max3x3(FilteredInput) + BackbufferError;
 		
 		tsr_simd_halfC FlickeringClampedFilteredHistory = clamp(FilteredHistory, FilteredBoxMin, FilteredBoxMax);
 
@@ -904,8 +911,8 @@ void MainCS(
 			ClampedFilteredHistory = FlickeringClampedFilteredHistory;
 		#endif
 		
-		tsr_simd_halfC OriginalBoxMin = Min3x3(OriginalInput, GroupThreadIndex);
-		tsr_simd_halfC OriginalBoxMax = Max3x3(OriginalInput, GroupThreadIndex);
+		tsr_simd_halfC OriginalBoxMin = Min3x3(OriginalInput);
+		tsr_simd_halfC OriginalBoxMax = Max3x3(OriginalInput);
 		tsr_simd_halfC BoxSize = (OriginalBoxMax - OriginalBoxMin + (BackbufferError + BackbufferError)) * FilteringWeight;
 		
 		const tsr_simd_halfC BackbufferQuantizationErrorVector = tsr_simd_halfC::Const(MeasureBackbufferLDRQuantizationError());
@@ -913,8 +920,8 @@ void MainCS(
 
 		tsr_simd_halfC FilteredClampedEnergy;
 		{
-			FilteredClampedEnergy = Median3x3(RawClampedEnergy, GroupThreadIndex);
-			FilteredClampedEnergy = Max3x3(FilteredClampedEnergy, GroupThreadIndex);
+			FilteredClampedEnergy = Median3x3(RawClampedEnergy);
+			FilteredClampedEnergy = Max3x3(FilteredClampedEnergy);
 		}
 
 		#if CONFIG_MOIRE_DETECTION
@@ -944,9 +951,9 @@ void MainCS(
 		
 		RejectionBlendFinal = saturate(FilteredRejection + tsr_simd_half::Const(1.0) - bIsValid);
 		
-		RejectionClampBlend = Median3x3(RawRejection, GroupThreadIndex);
+		RejectionClampBlend = Median3x3(RawRejection);
 		RejectionClampBlend = saturate(RejectionClampBlend + tsr_simd_half::Const(1.0) - bIsValid);
-		RejectionClampBlend = Min3x3(RejectionClampBlend, GroupThreadIndex);
+		RejectionClampBlend = Min3x3(RejectionClampBlend);
 	}
 	#endif
 	
@@ -974,7 +981,7 @@ void MainCS(
 
 	// Mark the top left corner of the tile for debugging purposes.
 	#if DEBUG_OUTPUT
-	if (GroupThreadIndex == 0)
+	if (GGroupThreadIndex == 0)
 	{
 		UNROLL_N(DEBUG_ARRAY_SIZE)
 		for (uint DebugId = 0; DebugId < DEBUG_ARRAY_SIZE; DebugId++)
