Modularises TSR RejectShading's convolution network tools into standalone .ush file to be reused in other passes

diff --git a/Engine/Shaders/Private/TemporalSuperResolution/TSRCommon.ush b/Engine/Shaders/Private/TemporalSuperResolution/TSRCommon.ush
--- a/Engine/Shaders/Private/TemporalSuperResolution/TSRCommon.ush
+++ b/Engine/Shaders/Private/TemporalSuperResolution/TSRCommon.ush
@@ -1,5 +1,7 @@
 // Copyright Epic Games, Inc. All Rights Reserved.
 
+#pragma once
+
 #include "../Common.ush"
 #include "../Random.ush"
 #include "../TextureSampling.ush"
diff --git a/Engine/Shaders/Private/TemporalSuperResolution/TSRConvolutionNetwork.ush b/Engine/Shaders/Private/TemporalSuperResolution/TSRConvolutionNetwork.ush
new file mode 100644
--- /dev/null
+++ b/Engine/Shaders/Private/TemporalSuperResolution/TSRConvolutionNetwork.ush
@@ -0,0 +1,490 @@
+// Copyright Epic Games, Inc. All Rights Reserved.
+
+#pragma once
+
+#include "TSRCommon.ush"
+#include "/Engine/Public/LaneVectorization.ush"
+
+
+#ifndef LANE_COUNT
+	#error LANE_COUNT is undefined.
+#endif
+
+#ifndef LANE_COUNT_X
+	#error LANE_COUNT_X is undefined.
+#endif
+
+#ifndef LANE_COUNT_Y
+	#error LANE_COUNT_Y is undefined.
+#endif
+
+/**
+ * LANE_COUNT_X=2
+ * LANE_COUNT_Y=2
+ * LaneStride=(4, 2)
+ *
+ *    o   o   o   o | o   o   o   o
+ *	                |
+ *    o   o   o   o | o   o   o   o
+ *	  --------------+--------------
+ *    o   o   o   o | o   o   o   o
+ *	                |
+ *    o   o   o   o | o   o   o   o
+ */
+
+ 
+//------------------------------------------------------- GLOBAL VARIABLE
+
+static uint GGroupThreadIndex = 0;
+
+
+//------------------------------------------------------- VECTOR MEMORY ORDER
+
+/** Returns coordinate of a pixel within a lane. */
+CALL_SITE_DEBUGLOC
+template<uint LaneStrideX, uint LaneStrideY>
+tsr_short2 GetSimdIndexPixelCoordinateInLane(const uint SimdIndex)
+{
+	return tsr_short2(SimdIndex % LaneStrideX, SimdIndex / LaneStrideX);
+}
+
+/** Returns coordinate of a pixel within a wave. */
+CALL_SITE_DEBUGLOC
+template<uint LaneStrideX, uint LaneStrideY>
+tsr_short2 GetLaneSimdPixelOffset(const uint LaneIndex, const uint SimdIndex)
+{
+	return (
+		tsr_short2(LaneIndex % uint(LANE_COUNT_X), LaneIndex / uint(LANE_COUNT_X)) * tsr_short2(LaneStrideX, LaneStrideY) +
+		GetSimdIndexPixelCoordinateInLane<LaneStrideX, LaneStrideY>(SimdIndex)
+	);
+}
+
+
+//------------------------------------------------------- FUNCTIONS
+
+/** write out elements to LDS for use with AccessNeighborTexel().
+ * 
+ * WriteToLDSForAccessNeighborTexel(ToConvolve);
+ * AccessNeighborTexel(ToConvolve, tsr_short2(1, 0));
+ */
+#if CONFIG_OPTIMIZED_LDS_ACCESS
+
+groupshared tsr_half4 SharedData[GROUP_TILE_SIZE * GROUP_TILE_SIZE];
+
+uint GetElementLDSIndex(const uint NeigborSimdIndex, const uint LaneRotation)
+{
+	return ((GGroupThreadIndex + LaneRotation) % uint(LANE_COUNT)) + NeigborSimdIndex * LANE_COUNT;
+}
+
+void WriteElementToLDS(const uint SimdIndex, tsr_half X)
+{
+	SharedData[GGroupThreadIndex + SimdIndex * LANE_COUNT] = tsr_half4(X, 0.0, 0.0, 0.0);
+}
+
+void WriteElementToLDS(const uint SimdIndex, tsr_half3 X)
+{
+	SharedData[GGroupThreadIndex + SimdIndex * LANE_COUNT] = tsr_half4(X, 0.0);
+}
+
+void WriteElementToLDS(const uint SimdIndex, tsr_half4 X)
+{
+	SharedData[GGroupThreadIndex + SimdIndex * LANE_COUNT] = tsr_half4(X);
+}
+
+void ReadElementFromLDS(const uint NeigborSimdIndex, const uint LaneRotation, out tsr_half X)
+{
+	X = SharedData[GetElementLDSIndex(NeigborSimdIndex, LaneRotation)].x;
+}
+
+void ReadElementFromLDS(const uint NeigborSimdIndex, const uint LaneRotation, out tsr_half3 X)
+{
+	X = SharedData[GetElementLDSIndex(NeigborSimdIndex, LaneRotation)].xyz;
+}
+
+void ReadElementFromLDS(const uint NeigborSimdIndex, const uint LaneRotation, out tsr_half4 X)
+{
+	X = SharedData[GetElementLDSIndex(NeigborSimdIndex, LaneRotation)].xyzw;
+}
+
+template<typename ScalarType, uint VectorSize, uint LaneStrideX, uint LaneStrideY>
+void WriteToLDSForAccessNeighborTexel(TLaneVector2D<ScalarType, VectorSize, LaneStrideX, LaneStrideY> Input)
+{
+	GroupMemoryBarrierWithGroupSync();
+	UNROLL
+	for (uint SimdIndex = 0; SimdIndex < LaneStrideX * LaneStrideY; SimdIndex++)
+	{
+		WriteElementToLDS(SimdIndex, Input.GetElement(SimdIndex));
+	}
+	GroupMemoryBarrierWithGroupSync();
+}
+
+#else // !CONFIG_OPTIMIZED_LDS_ACCESS
+
+template<typename ScalarType, uint VectorSize, uint LaneStrideX, uint LaneStrideY>
+void WriteToLDSForAccessNeighborTexel(TLaneVector2D<ScalarType, VectorSize, LaneStrideX, LaneStrideY> Input)
+{
+	// NOP
+}
+
+#endif // !CONFIG_OPTIMIZED_LDS_ACCESS
+
+CALL_SITE_DEBUGLOC
+template<typename ScalarType, uint VectorSize, uint LaneStrideX, uint LaneStrideY>
+TLaneVector2D<ScalarType, VectorSize, LaneStrideX, LaneStrideY> AccessNeighborTexel(
+	TLaneVector2D<ScalarType, VectorSize, LaneStrideX, LaneStrideY> Input,
+	const tsr_short2 Offset)
+{
+	TLaneVector2D<ScalarType, VectorSize, LaneStrideX, LaneStrideY> Return;
+	
+	UNROLL
+	for (uint SimdIndex = 0; SimdIndex < LaneStrideX * LaneStrideY; SimdIndex++)
+	{
+		const uint2 SimdPixelCoord = uint2(GetSimdIndexPixelCoordinateInLane<LaneStrideX, LaneStrideY>(SimdIndex));
+
+		// Compute the simd index of the neighbor.
+		const uint2 NeigborSimdPixelCoord = (SimdPixelCoord + uint2(Offset)) % uint2(LaneStrideX, LaneStrideY);
+		const uint NeigborSimdIndex = dot(NeigborSimdPixelCoord, uint2(1, LaneStrideX));
+
+		// Compute the lane rotaton of the neighbor.
+		const bool2 bNeedsLaneRotation = (SimdPixelCoord + uint2(Offset)) >= uint2(LaneStrideX, LaneStrideY);
+		const int LaneRotation = dot(select(bNeedsLaneRotation, int2(Offset), 0), int2(1, LANE_COUNT_X));
+		
+		// Access the lement.
+		vector<ScalarType, VectorSize> SimdElement = Input.GetElement(NeigborSimdIndex);
+		vector<ScalarType, VectorSize> ReturnSimdElement = SimdElement;
+		if (LaneRotation != 0)
+		#if CONFIG_OPTIMIZED_LDS_ACCESS
+		{
+			ReadElementFromLDS(NeigborSimdIndex, LaneRotation, /* out */ ReturnSimdElement);
+		}
+		#else
+		{
+			// Only uses wave intrinsics on platforms that provides RDNA's ds_swizzle because we know the wave size of this platforms. 
+			#if DIM_WAVE_SIZE != 0
+				#if !PLATFORM_SUPPORTS_WAVE_ROTATE
+					//#error InitWaveRotateLaneGroup() isn't supported.
+				#endif
+				const FWaveBroadcastSettings BroadcastSettings = InitWaveRotateLaneGroup(/* LaneGroupSize = */ LANE_COUNT, LaneRotation);
+			#else
+				const FWaveBroadcastSettings BroadcastSettings = ConvertWaveBroadcastToLDS(InitWaveRotateLaneGroup(/* LaneGroupSize = */ LANE_COUNT, LaneRotation));
+			#endif
+			ReturnSimdElement = WaveBroadcast(BroadcastSettings, SimdElement);
+		}
+		#endif
+
+		Return.SetElement(SimdIndex, ReturnSimdElement);
+	}
+
+	return Return;
+}
+
+CALL_SITE_DEBUGLOC
+template<typename ScalarType, uint VectorSize, uint LaneStrideX, uint LaneStrideY>
+void AccessNeighborTexels3x1(
+	TLaneVector2D<ScalarType, VectorSize, LaneStrideX, LaneStrideY> CO,
+	const tsr_short2 OffsetP,
+	out TLaneVector2D<ScalarType, VectorSize, LaneStrideX, LaneStrideY> CP,
+	out TLaneVector2D<ScalarType, VectorSize, LaneStrideX, LaneStrideY> CN)
+{
+	const tsr_short2 OffsetN = tsr_short2(-int2(OffsetP)); // work arround a shader compiler bug
+	
+	WriteToLDSForAccessNeighborTexel(CO);
+	CP = AccessNeighborTexel(CO, OffsetP);
+	CN = AccessNeighborTexel(CO, OffsetN);
+}
+
+
+//------------------------------------------------------- CONVOLUTION ACCUMULATORS
+
+template<typename ScalarType, uint VectorSize, uint LaneStrideX, uint LaneStrideY>
+struct TAccumulatorAvg
+{
+	TLaneVector2D<ScalarType, VectorSize, LaneStrideX, LaneStrideY> Value;
+	TLaneVector2D<ScalarType, 1, LaneStrideX, LaneStrideY> Weight;
+	
+	CALL_SITE_DEBUGLOC
+	void Start(TLaneVector2D<ScalarType, VectorSize, LaneStrideX, LaneStrideY> CenterSample)
+	{
+		Value = CenterSample;
+		Weight = TLaneVector2D<ScalarType, 1, LaneStrideX, LaneStrideY>::Const(1.0);
+	}
+	
+	CALL_SITE_DEBUGLOC
+	static TLaneVector2D<ScalarType, VectorSize, LaneStrideX, LaneStrideY> PreprocessNeighborSample(
+		TLaneVector2D<ScalarType, VectorSize, LaneStrideX, LaneStrideY> Sample,
+		TLaneVector2D<ScalarType, 1, LaneStrideX, LaneStrideY> InWeight)
+	{
+		return Sample;
+	}
+	
+	CALL_SITE_DEBUGLOC
+	void Accumulate(
+		const tsr_short2 Offset,
+		TLaneVector2D<ScalarType, VectorSize, LaneStrideX, LaneStrideY> Sample,
+		TLaneVector2D<ScalarType, 1, LaneStrideX, LaneStrideY> InWeight)
+	{
+		const tsr_half SampleSpatialWeight = tsr_half(rcp(abs(tsr_half(Offset.x)) + 1.0) * rcp(abs(tsr_half(Offset.y)) + 1.0));
+		
+		Value = Value + Sample * TLaneVector2D<ScalarType, VectorSize, LaneStrideX, LaneStrideY>::Vectorize(InWeight * SampleSpatialWeight);
+		Weight = Weight + InWeight * TLaneVector2D<ScalarType, 1, LaneStrideX, LaneStrideY>::Const(SampleSpatialWeight);
+	}
+	
+	CALL_SITE_DEBUGLOC
+	void Finish(const uint SampleCount)
+	{
+		Value = Value * TLaneVector2D<ScalarType, VectorSize, LaneStrideX, LaneStrideY>::Vectorize(rcp(Weight));
+	}
+};
+
+
+//------------------------------------------------------- CONVOLUTIONS
+
+CALL_SITE_DEBUGLOC
+template<class FAccumulator, typename ScalarType, uint VectorSize, uint LaneStrideX, uint LaneStrideY>
+FAccumulator Convole3x1(
+	TLaneVector2D<ScalarType, VectorSize, LaneStrideX, LaneStrideY> Color,
+	TLaneVector2D<ScalarType, 1, LaneStrideX, LaneStrideY> Weight,
+	const tsr_short2 Offset)
+#if CONFIG_OPTIMIZED_LDS_ACCESS
+{
+	const tsr_short2 OffsetN = tsr_short2(-int2(Offset)); // work arround a shader compiler bug
+	
+	TLaneVector2D<ScalarType, VectorSize, LaneStrideX, LaneStrideY> PreprocessedNeighborColor = FAccumulator::PreprocessNeighborSample(Color, Weight);
+
+	WriteToLDSForAccessNeighborTexel(PreprocessedNeighborColor);
+	TLaneVector2D<ScalarType, VectorSize, LaneStrideX, LaneStrideY> Sample0 = AccessNeighborTexel(PreprocessedNeighborColor, Offset);
+	TLaneVector2D<ScalarType, VectorSize, LaneStrideX, LaneStrideY> Sample1 = AccessNeighborTexel(PreprocessedNeighborColor, OffsetN);
+	
+	WriteToLDSForAccessNeighborTexel(Weight);
+	TLaneVector2D<ScalarType, 1, LaneStrideX, LaneStrideY> Weight0 = AccessNeighborTexel(Weight, Offset);
+	TLaneVector2D<ScalarType, 1, LaneStrideX, LaneStrideY> Weight1 = AccessNeighborTexel(Weight, OffsetN);
+	
+	FAccumulator Accumulator;
+	Accumulator.Start(Color);
+	Accumulator.Accumulate(+Offset, Sample0, Weight0);
+	Accumulator.Accumulate(OffsetN, Sample1, Weight1);
+	Accumulator.Finish(3);
+	return Accumulator;
+}
+#else // !CONFIG_OPTIMIZED_LDS_ACCESS
+{
+	FAccumulator Accumulator;
+	Accumulator.Start(Color);
+	
+	TLaneVector2D<ScalarType, VectorSize, LaneStrideX, LaneStrideY> PreprocessedNeighborColor = FAccumulator::PreprocessNeighborSample(Color, Weight);
+
+	TLaneVector2D<ScalarType, VectorSize, LaneStrideX, LaneStrideY> Sample0 = AccessNeighborTexel(PreprocessedNeighborColor, Offset);
+	TLaneVector2D<ScalarType, 1, LaneStrideX, LaneStrideY>          Weight0 = AccessNeighborTexel(Weight, Offset);
+	Accumulator.Accumulate(+Offset, Sample0, Weight0);
+	
+	const tsr_short2 OffsetN = tsr_short2(-int2(Offset)); // work arround a shader compiler bug
+	TLaneVector2D<ScalarType, VectorSize, LaneStrideX, LaneStrideY> Sample1 = AccessNeighborTexel(PreprocessedNeighborColor, OffsetN);
+	TLaneVector2D<ScalarType, 1, LaneStrideX, LaneStrideY>          Weight1 = AccessNeighborTexel(Weight, OffsetN);
+	Accumulator.Accumulate(OffsetN, Sample1, Weight1);
+
+	Accumulator.Finish(3);
+	return Accumulator;
+}
+#endif // !CONFIG_OPTIMIZED_LDS_ACCESS
+
+CALL_SITE_DEBUGLOC
+template<class FAccumulator, typename ScalarType, uint VectorSize, uint LaneStrideX, uint LaneStrideY>
+TLaneVector2D<ScalarType, VectorSize, LaneStrideX, LaneStrideY> Convolve3x3HV(
+	TLaneVector2D<ScalarType, VectorSize, LaneStrideX, LaneStrideY> Color,
+	TLaneVector2D<ScalarType, 1, LaneStrideX, LaneStrideY> Weight)
+{
+	TLaneVector2D<ScalarType, VectorSize, LaneStrideX, LaneStrideY> HorizontalColor = Convole3x1<FAccumulator, ScalarType, VectorSize, LaneStrideX, LaneStrideY>(Color, Weight, tsr_short2(1, 0)).Value;
+	return Convole3x1<FAccumulator, ScalarType, VectorSize, LaneStrideX, LaneStrideY>(HorizontalColor, Weight, tsr_short2(0, 1)).Value;
+}
+
+
+//------------------------------------------------------- BLUR 3x3 CONVOLUTION
+
+CALL_SITE_DEBUGLOC
+template<typename ScalarType, uint VectorSize, uint LaneStrideX, uint LaneStrideY>
+TLaneVector2D<ScalarType, VectorSize, LaneStrideX, LaneStrideY> Blur3x3(
+	TLaneVector2D<ScalarType, VectorSize, LaneStrideX, LaneStrideY> Center,
+	TLaneVector2D<ScalarType, 1, LaneStrideX, LaneStrideY> Weight)
+{
+	return Convolve3x3HV<TAccumulatorAvg<ScalarType, VectorSize, LaneStrideX, LaneStrideY>, ScalarType, VectorSize, LaneStrideX, LaneStrideY>(
+		Center, Weight);
+}
+
+
+//------------------------------------------------------- MIN & MAX 3x3 CONVOLUTION
+
+CALL_SITE_DEBUGLOC
+template<typename ScalarType, uint VectorSize, uint LaneStrideX, uint LaneStrideY>
+TLaneVector2D<ScalarType, VectorSize, LaneStrideX, LaneStrideY> Min3x3(TLaneVector2D<ScalarType, VectorSize, LaneStrideX, LaneStrideY> Center)
+{
+	TLaneVector2D<ScalarType, VectorSize, LaneStrideX, LaneStrideY> CO = Center;
+	TLaneVector2D<ScalarType, VectorSize, LaneStrideX, LaneStrideY> CP;
+	TLaneVector2D<ScalarType, VectorSize, LaneStrideX, LaneStrideY> CN;
+	AccessNeighborTexels3x1(CO, tsr_short2(1, 0), /* out */ CP, /* out */ CN);
+	
+	TLaneVector2D<ScalarType, VectorSize, LaneStrideX, LaneStrideY> MinO = min(CN, min(CP, CO));
+	TLaneVector2D<ScalarType, VectorSize, LaneStrideX, LaneStrideY> MinP;
+	TLaneVector2D<ScalarType, VectorSize, LaneStrideX, LaneStrideY> MinN;
+	AccessNeighborTexels3x1(MinO, tsr_short2(0, 1), /* out */ MinP, /* out */ MinN);
+
+	return min(MinN, min(MinP, MinO));
+}
+
+CALL_SITE_DEBUGLOC
+template<typename ScalarType, uint VectorSize, uint LaneStrideX, uint LaneStrideY>
+TLaneVector2D<ScalarType, VectorSize, LaneStrideX, LaneStrideY> Max3x3(TLaneVector2D<ScalarType, VectorSize, LaneStrideX, LaneStrideY> Center)
+{
+	TLaneVector2D<ScalarType, VectorSize, LaneStrideX, LaneStrideY> CO = Center;
+	TLaneVector2D<ScalarType, VectorSize, LaneStrideX, LaneStrideY> CP;
+	TLaneVector2D<ScalarType, VectorSize, LaneStrideX, LaneStrideY> CN;
+	AccessNeighborTexels3x1(CO, tsr_short2(1, 0), /* out */ CP, /* out */ CN);
+	
+	TLaneVector2D<ScalarType, VectorSize, LaneStrideX, LaneStrideY> MaxO = max(CN, max(CP, CO));
+	TLaneVector2D<ScalarType, VectorSize, LaneStrideX, LaneStrideY> MaxP;
+	TLaneVector2D<ScalarType, VectorSize, LaneStrideX, LaneStrideY> MaxN;
+	AccessNeighborTexels3x1(MaxO, tsr_short2(0, 1), /* out */ MaxP, /* out */ MaxN);
+
+	return max(MaxN, max(MaxP, MaxO));
+}
+
+CALL_SITE_DEBUGLOC
+template<typename ScalarType, uint VectorSize, uint LaneStrideX, uint LaneStrideY>
+void MinMax3x3(
+	TLaneVector2D<ScalarType, VectorSize, LaneStrideX, LaneStrideY> Center,
+	out TLaneVector2D<ScalarType, VectorSize, LaneStrideX, LaneStrideY> OutMin,
+	out TLaneVector2D<ScalarType, VectorSize, LaneStrideX, LaneStrideY> OutMax)
+#if 1
+{
+	TLaneVector2D<ScalarType, VectorSize, LaneStrideX, LaneStrideY> CO = Center;
+	TLaneVector2D<ScalarType, VectorSize, LaneStrideX, LaneStrideY> CP;
+	TLaneVector2D<ScalarType, VectorSize, LaneStrideX, LaneStrideY> CN;
+	AccessNeighborTexels3x1(CO, tsr_short2(1, 0), /* out */ CP, /* out */ CN);
+	
+	TLaneVector2D<ScalarType, VectorSize, LaneStrideX, LaneStrideY> MinO = min(CN, min(CP, CO));
+	TLaneVector2D<ScalarType, VectorSize, LaneStrideX, LaneStrideY> MaxO = max(CN, max(CP, CO));
+
+	TLaneVector2D<ScalarType, VectorSize, LaneStrideX, LaneStrideY> MinP;
+	TLaneVector2D<ScalarType, VectorSize, LaneStrideX, LaneStrideY> MinN;
+	AccessNeighborTexels3x1(MinO, tsr_short2(0, 1), /* out */ MinP, /* out */ MinN);
+	OutMin = min(MinN, min(MinP, MinO));
+
+	TLaneVector2D<ScalarType, VectorSize, LaneStrideX, LaneStrideY> MaxP;
+	TLaneVector2D<ScalarType, VectorSize, LaneStrideX, LaneStrideY> MaxN;
+	AccessNeighborTexels3x1(MaxO, tsr_short2(0, 1), /* out */ MaxP, /* out */ MaxN);
+	OutMax = max(MaxN, max(MaxP, MaxO));
+}
+#else
+{
+	OutMin = Min3x3(Center);
+	OutMax = Max3x3(Center);
+}
+#endif
+
+
+//------------------------------------------------------- MEDIAN 3x3 CONVOLUTION
+
+// Operator that output, lowest, median and highest values from 3 input values.
+CALL_SITE_DEBUGLOC
+template<typename FSampleType>
+void LMHOperator(FSampleType A, FSampleType B, FSampleType C, out FSampleType L, out FSampleType M, out FSampleType H)
+{
+	// TODO: min3, max3, med3.
+	FSampleType X = min(B, C);
+	FSampleType Y = max(B, C);
+	
+	L = min(A, X);
+	FSampleType Z = max(A, X);
+
+	M = min(Z, Y);
+	H = max(Z, Y);
+}
+
+// 3 samples median.
+CALL_SITE_DEBUGLOC
+template<typename FSampleType>
+FSampleType Median(FSampleType A, FSampleType B, FSampleType C)
+{
+	// TODO: med3.
+	FSampleType L, M, H;
+	LMHOperator(A, B, C, L, M, H);
+	return M;
+}
+
+// 9 samples median using Smith1996.
+CALL_SITE_DEBUGLOC
+template<typename FSampleType>
+FSampleType Median9(FSampleType In[9])
+{
+	// First layer.
+	FSampleType L0[3];
+	FSampleType M0[3];
+	FSampleType H0[3];
+	for (uint j = 0; j < 3; j++)
+	{
+		LMHOperator(In[j * 3 + 0], In[j * 3 + 1], In[j * 3 + 2], L0[j], M0[j], H0[j]);
+	}
+	
+	// Second layer.
+	// TODO: min3, max3.
+	FSampleType M1[3];
+	M1[0] = max(max(L0[0], L0[1]), L0[2]);
+	M1[1] = Median(M0[0], M0[1], M0[2]);
+	M1[2] = min(min(H0[0], H0[1]), H0[2]);
+	
+	// Third layer.
+	return Median(M1[0], M1[1], M1[2]);
+}
+
+CALL_SITE_DEBUGLOC
+template<typename ScalarType, uint VectorSize, uint LaneStrideX, uint LaneStrideY>
+TLaneVector2D<ScalarType, VectorSize, LaneStrideX, LaneStrideY> Median3x3(TLaneVector2D<ScalarType, VectorSize, LaneStrideX, LaneStrideY> Center)
+#if 1
+{
+	WriteToLDSForAccessNeighborTexel(Center);
+	
+	TLaneVector2D<ScalarType, VectorSize, LaneStrideX, LaneStrideY> C0 = Center;
+	TLaneVector2D<ScalarType, VectorSize, LaneStrideX, LaneStrideY> C1 = AccessNeighborTexel(Center, tsr_short2( 1, 0));
+	TLaneVector2D<ScalarType, VectorSize, LaneStrideX, LaneStrideY> C2 = AccessNeighborTexel(Center, tsr_short2(-1, 0));
+	
+	TLaneVector2D<ScalarType, VectorSize, LaneStrideX, LaneStrideY> M0L0;
+	TLaneVector2D<ScalarType, VectorSize, LaneStrideX, LaneStrideY> M0M0;
+	TLaneVector2D<ScalarType, VectorSize, LaneStrideX, LaneStrideY> M0H0;
+	LMHOperator(C0, C1, C2, M0L0, M0M0, M0H0);
+	
+	WriteToLDSForAccessNeighborTexel(M0L0);
+	TLaneVector2D<ScalarType, VectorSize, LaneStrideX, LaneStrideY> M0L1 = AccessNeighborTexel(M0L0, tsr_short2(0,  1));
+	TLaneVector2D<ScalarType, VectorSize, LaneStrideX, LaneStrideY> M0L2 = AccessNeighborTexel(M0L0, tsr_short2(0, -1));
+	
+	WriteToLDSForAccessNeighborTexel(M0M0);
+	TLaneVector2D<ScalarType, VectorSize, LaneStrideX, LaneStrideY> M0M1 = AccessNeighborTexel(M0M0, tsr_short2(0,  1));
+	TLaneVector2D<ScalarType, VectorSize, LaneStrideX, LaneStrideY> M0M2 = AccessNeighborTexel(M0M0, tsr_short2(0, -1));
+	
+	WriteToLDSForAccessNeighborTexel(M0H0);
+	TLaneVector2D<ScalarType, VectorSize, LaneStrideX, LaneStrideY> M0H1 = AccessNeighborTexel(M0H0, tsr_short2(0,  1));
+	TLaneVector2D<ScalarType, VectorSize, LaneStrideX, LaneStrideY> M0H2 = AccessNeighborTexel(M0H0, tsr_short2(0, -1));
+	
+	// Second layer.
+	// TODO: min3, max3.
+	TLaneVector2D<ScalarType, VectorSize, LaneStrideX, LaneStrideY> M1[3];
+	M1[0] = max(max(M0L0, M0L1), M0L2);
+	M1[1] = Median(M0M0, M0M1, M0M2);
+	M1[2] = min(min(M0H0, M0H1), M0H2);
+	
+	// Third layer.
+	return Median(M1[0], M1[1], M1[2]);
+}
+#else
+{
+	WriteToLDSForAccessNeighborTexel(Center);
+	
+	TLaneVector2D<ScalarType, VectorSize, LaneStrideX, LaneStrideY> InSamples[9];
+	UNROLL
+	for (uint i = 0; i < 9; i++)
+	{
+		TLaneVector2D<ScalarType, VectorSize, LaneStrideX, LaneStrideY> Sample = AccessNeighborTexel(Center, tsr_short2(kOffsets3x3[i]));
+		InSamples[i] = Sample;
+	}
+
+	return Median9(InSamples);
+}
+#endif
diff --git a/Engine/Shaders/Private/TemporalSuperResolution/TSRRejectShading.usf b/Engine/Shaders/Private/TemporalSuperResolution/TSRRejectShading.usf
--- a/Engine/Shaders/Private/TemporalSuperResolution/TSRRejectShading.usf
+++ b/Engine/Shaders/Private/TemporalSuperResolution/TSRRejectShading.usf
@@ -78,13 +78,17 @@
 /** Size of the SIMD per lane that also number of input pixel loaded into each individual lanes. */
 #define SIMD_SIZE (LANE_STRIDE_X * LANE_STRIDE_Y)
 
+//------------------------------------------------------- INCLUDE CONVOLUTION NETWORK
+
+#include "TSRConvolutionNetwork.ush"
+
 
 //------------------------------------------------------- TYPEDEFS
 
-#define tsr_simd_half  TLaneVector<tsr_half, 1, SIMD_SIZE>
-#define tsr_simd_half4 TLaneVector<tsr_half, 4, SIMD_SIZE>
-#define tsr_simd_halfC TLaneVector<tsr_half, CONFIG_CHANNEL_COUNT, SIMD_SIZE>
-#define tsr_simd_bool  TLaneVector<bool, 1, SIMD_SIZE>
+#define tsr_simd_half  TLaneVector2D<tsr_half, 1, LANE_STRIDE_X, LANE_STRIDE_Y>
+#define tsr_simd_half4 TLaneVector2D<tsr_half, 4, LANE_STRIDE_X, LANE_STRIDE_Y>
+#define tsr_simd_halfC TLaneVector2D<tsr_half, CONFIG_CHANNEL_COUNT, LANE_STRIDE_X, LANE_STRIDE_Y>
+#define tsr_simd_bool  TLaneVector2D<bool, 1, LANE_STRIDE_X, LANE_STRIDE_Y>
 
 #ifndef PLATFORM_SPECIFIC_ISOLATE
 	#define PLATFORM_SPECIFIC_ISOLATE ISOLATE
@@ -115,479 +119,25 @@ RWTexture2D<tsr_half2> HistoryRejectionOutput;
 RWTexture2D<tsr_half> InputSceneColorLdrLumaOutput;
 
 
-//------------------------------------------------------- GLOBAL VARIABLE
-
-static uint GGroupThreadIndex = 0;
-
-
-//------------------------------------------------------- GENERAL TEXEL SWIZZLING & NEIGHBOR ACCESS
-
-CALL_SITE_DEBUGLOC
-tsr_short2 GetSimdIndexPixelCoordinateInLane(const uint SimdIndex)
-{
-	return tsr_short2(SimdIndex % uint(LANE_STRIDE_X), SimdIndex / uint(LANE_STRIDE_X));
-}
-
-CALL_SITE_DEBUGLOC
-tsr_short2 GetLaneSimdPixelOffset(const uint LaneIndex, const uint SimdIndex)
-{
-	return (
-		tsr_short2(LaneIndex % uint(LANE_COUNT_X), LaneIndex / uint(LANE_COUNT_X)) * tsr_short2(LANE_STRIDE_X, LANE_STRIDE_Y) +
-		GetSimdIndexPixelCoordinateInLane(SimdIndex)
-	);
-}
-
-
-//------------------------------------------------------- GENERAL TEXEL SWIZZLING & NEIGHBOR ACCESS
-
-#if CONFIG_OPTIMIZED_LDS_ACCESS
-
-groupshared tsr_half4 SharedData[GROUP_TILE_SIZE * GROUP_TILE_SIZE];
-
-void WriteAllElementToLDS(TLaneVector<tsr_half, 1, SIMD_SIZE> X)
-{
-	GroupMemoryBarrierWithGroupSync();
-	UNROLL_N(SIMD_SIZE)
-	for (uint SimdIndex = 0; SimdIndex < SIMD_SIZE; SimdIndex++)
-	{
-		SharedData[GGroupThreadIndex + SimdIndex * LANE_COUNT] = tsr_half4(X.GetElement(SimdIndex), 0.0, 0.0, 0.0);
-	}
-	GroupMemoryBarrierWithGroupSync();
-}
-
-void WriteAllElementToLDS(TLaneVector<tsr_half, 3, SIMD_SIZE> X)
-{
-	GroupMemoryBarrierWithGroupSync();
-	UNROLL_N(SIMD_SIZE)
-	for (uint SimdIndex = 0; SimdIndex < SIMD_SIZE; SimdIndex++)
-	{
-		SharedData[GGroupThreadIndex + SimdIndex * LANE_COUNT] = tsr_half4(X.GetElement(SimdIndex), 0.0);
-	}
-	GroupMemoryBarrierWithGroupSync();
-}
-
-void WriteAllElementToLDS(TLaneVector<tsr_half, 4, SIMD_SIZE> X)
-{
-	GroupMemoryBarrierWithGroupSync();
-	UNROLL_N(SIMD_SIZE)
-	for (uint SimdIndex = 0; SimdIndex < SIMD_SIZE; SimdIndex++)
-	{
-		SharedData[GGroupThreadIndex + SimdIndex * LANE_COUNT] = tsr_half4(X.GetElement(SimdIndex));
-	}
-	GroupMemoryBarrierWithGroupSync();
-}
-
-uint GetElementLDSIndex(const uint NeigborSimdIndex, const uint LaneRotation)
-{
-	return ((GGroupThreadIndex + LaneRotation) % uint(LANE_COUNT)) + NeigborSimdIndex * LANE_COUNT;
-}
-
-void ReadElementFromLDS(const uint NeigborSimdIndex, const uint LaneRotation, out tsr_half X)
-{
-	X = SharedData[GetElementLDSIndex(NeigborSimdIndex, LaneRotation)].x;
-}
-
-void ReadElementFromLDS(const uint NeigborSimdIndex, const uint LaneRotation, out tsr_half3 X)
-{
-	X = SharedData[GetElementLDSIndex(NeigborSimdIndex, LaneRotation)].xyz;
-}
-
-void ReadElementFromLDS(const uint NeigborSimdIndex, const uint LaneRotation, out tsr_half4 X)
-{
-	X = SharedData[GetElementLDSIndex(NeigborSimdIndex, LaneRotation)].xyzw;
-}
-
-template<typename ScalarType, uint VectorSize, uint SimdSize>
-void WriteToLDSForAccessNeighborTexel(TLaneVector<ScalarType, VectorSize, SimdSize> Input)
-{
-	WriteAllElementToLDS(Input);
-}
-
-#else // !CONFIG_OPTIMIZED_LDS_ACCESS
-
-template<typename ScalarType, uint VectorSize, uint SimdSize>
-void WriteToLDSForAccessNeighborTexel(TLaneVector<ScalarType, VectorSize, SimdSize> Input)
-{
-	// NOP
-}
-
-#endif // !CONFIG_OPTIMIZED_LDS_ACCESS
-
-
-CALL_SITE_DEBUGLOC
-template<typename ScalarType, uint VectorSize, uint SimdSize>
-TLaneVector<ScalarType, VectorSize, SimdSize> AccessNeighborTexel(TLaneVector<ScalarType, VectorSize, SimdSize> Input, const tsr_short2 Offset)
-{
-	TLaneVector<ScalarType, VectorSize, SimdSize> Return;
-	
-	UNROLL_N(SIMD_SIZE)
-	for (uint SimdIndex = 0; SimdIndex < SIMD_SIZE; SimdIndex++)
-	{
-		const uint2 SimdPixelCoord = uint2(GetSimdIndexPixelCoordinateInLane(SimdIndex));
-
-		// Compute the simd index of the neighbor.
-		const uint2 NeigborSimdPixelCoord = (SimdPixelCoord + uint2(Offset)) % uint2(LANE_STRIDE_X, LANE_STRIDE_Y);
-		const uint NeigborSimdIndex = dot(NeigborSimdPixelCoord, uint2(1, LANE_STRIDE_X));
-
-		// Compute the lane rotaton of the neighbor.
-		const bool2 bNeedsLaneRotation = (SimdPixelCoord + uint2(Offset)) >= uint2(LANE_STRIDE_X, LANE_STRIDE_Y);
-		const int LaneRotation = dot(select(bNeedsLaneRotation, int2(Offset), 0), int2(1, LANE_COUNT_X));
-		
-		// Access the lement.
-		vector<ScalarType, VectorSize> SimdElement = Input.GetElement(NeigborSimdIndex);
-		vector<ScalarType, VectorSize> ReturnSimdElement = SimdElement;
-		if (LaneRotation != 0)
-		#if CONFIG_OPTIMIZED_LDS_ACCESS
-		{
-			ReadElementFromLDS(NeigborSimdIndex, LaneRotation, /* out */ ReturnSimdElement);
-		}
-		#else
-		{
-			// Only uses wave intrinsics on platforms that provides RDNA's ds_swizzle because we know the wave size of this platforms. 
-			#if DIM_WAVE_SIZE != 0
-				#if !PLATFORM_SUPPORTS_WAVE_ROTATE
-					//#error InitWaveRotateLaneGroup() isn't supported.
-				#endif
-				const FWaveBroadcastSettings BroadcastSettings = InitWaveRotateLaneGroup(/* LaneGroupSize = */ LANE_COUNT, LaneRotation);
-			#else
-				const FWaveBroadcastSettings BroadcastSettings = ConvertWaveBroadcastToLDS(InitWaveRotateLaneGroup(/* LaneGroupSize = */ LANE_COUNT, LaneRotation));
-			#endif
-			ReturnSimdElement = WaveBroadcast(BroadcastSettings, SimdElement);
-		}
-		#endif
-
-		Return.SetElement(SimdIndex, ReturnSimdElement);
-	}
-
-	return Return;
-}
-
-
-//------------------------------------------------------- CONVOLUTION ACCUMULATORS
-
-template<typename Type>
-struct TAccumulatorAvg
-{
-	Type Value;
-	tsr_simd_half Weight;
-	
-	CALL_SITE_DEBUGLOC
-	static bool NeedWeightInAccumulate()
-	{
-		return true;
-	}
-
-	CALL_SITE_DEBUGLOC
-	void Start(Type CenterSample)
-	{
-		Value = CenterSample;
-		Weight = tsr_simd_half::Const(1.0);
-	}
-	
-	CALL_SITE_DEBUGLOC
-	static Type PreprocessNeighborSample(Type Sample, tsr_simd_half InWeight)
-	{
-		return Sample;
-	}
-	
-	CALL_SITE_DEBUGLOC
-	void Accumulate(const tsr_short2 Offset, Type Sample, tsr_simd_half InWeight)
-	{
-		const tsr_half SampleSpatialWeight = tsr_half(rcp(abs(tsr_half(Offset.x)) + 1.0) * rcp(abs(tsr_half(Offset.y)) + 1.0));
-		
-		Value = Value + Sample * Type::Vectorize(InWeight * SampleSpatialWeight);
-		Weight = Weight + InWeight * tsr_simd_half::Const(SampleSpatialWeight);
-	}
-	
-	CALL_SITE_DEBUGLOC
-	void Finish(const uint SampleCount)
-	{
-		Value = Value * Type::Vectorize(rcp(Weight));
-	}
-};
-
-template<typename Type>
-struct TAccumulatorMin
-{
-	Type Value;
-	
-	CALL_SITE_DEBUGLOC
-	static bool NeedWeightInAccumulate()
-	{
-		return false;
-	}
-
-	CALL_SITE_DEBUGLOC
-	void Start(Type CenterSample)
-	{
-		Value = CenterSample;
-	}
-	
-	CALL_SITE_DEBUGLOC
-	static Type PreprocessNeighborSample(Type Sample, tsr_simd_half InWeight)
-	{
-		return Sample;
-	}
-	
-	CALL_SITE_DEBUGLOC
-	void Accumulate(const tsr_short2 Offset, Type Sample, tsr_simd_half Weight)
-	{
-		Value = min(Value, Sample);
-	}
-	
-	CALL_SITE_DEBUGLOC
-	void Finish(const uint SampleCount)
-	{
-		// NOP
-	}
-};
-
-template<typename Type>
-struct TAccumulatorMax
-{
-	Type Value;
-	
-	CALL_SITE_DEBUGLOC
-	static bool NeedWeightInAccumulate()
-	{
-		return false;
-	}
-
-	CALL_SITE_DEBUGLOC
-	void Start(Type CenterSample)
-	{
-		Value = CenterSample;
-	}
-	
-	CALL_SITE_DEBUGLOC
-	static Type PreprocessNeighborSample(Type Sample, tsr_simd_half InWeight)
-	{
-		return Sample;
-	}
-	
-	CALL_SITE_DEBUGLOC
-	void Accumulate(const tsr_short2 Offset, Type Sample, tsr_simd_half Weight)
-	{
-		Value = max(Value, Sample);
-	}
-	
-	CALL_SITE_DEBUGLOC
-	void Finish(const uint SampleCount)
-	{
-		// NOP
-	}
-};
-
-
-//------------------------------------------------------- CONVOLUTIONS
-
-CALL_SITE_DEBUGLOC
-template<class FAccumulator, typename FSampleType>
-FAccumulator Convole3x1(FSampleType Color, tsr_simd_half Weight, const tsr_short2 Offset)
-#if CONFIG_OPTIMIZED_LDS_ACCESS
-{
-	const tsr_short2 OffsetN = tsr_short2(-int2(Offset)); // work arround a shader compiler bug
-	
-	FSampleType PreprocessedNeighborColor = FAccumulator::PreprocessNeighborSample(Color, Weight);
-
-	WriteToLDSForAccessNeighborTexel(PreprocessedNeighborColor);
-	FSampleType Sample0 = AccessNeighborTexel(PreprocessedNeighborColor, Offset);
-	FSampleType Sample1 = AccessNeighborTexel(PreprocessedNeighborColor, OffsetN);
-	
-	tsr_simd_half Weight0 = tsr_simd_half::Const(tsr_half(1.0));
-	tsr_simd_half Weight1 = tsr_simd_half::Const(tsr_half(1.0));
-	if (FAccumulator::NeedWeightInAccumulate())
-	{
-		WriteToLDSForAccessNeighborTexel(Weight);
-		Weight0 = AccessNeighborTexel(Weight, Offset);
-		Weight1 = AccessNeighborTexel(Weight, OffsetN);
-	}
-	
-	FAccumulator Accumulator;
-	Accumulator.Start(Color);
-	Accumulator.Accumulate(+Offset, Sample0, Weight0);
-	Accumulator.Accumulate(OffsetN, Sample1, Weight1);
-	Accumulator.Finish(3);
-	return Accumulator;
-}
-#else // !CONFIG_OPTIMIZED_LDS_ACCESS
-{
-	FAccumulator Accumulator;
-	Accumulator.Start(Color);
-	
-	FSampleType PreprocessedNeighborColor = FAccumulator::PreprocessNeighborSample(Color, Weight);
-
-	FSampleType Sample0 = AccessNeighborTexel(PreprocessedNeighborColor, Offset);
-	tsr_simd_half Weight0 = tsr_simd_half::Const(tsr_half(1.0));
-	if (FAccumulator::NeedWeightInAccumulate())
-	{
-		Weight0 = AccessNeighborTexel(Weight, Offset);
-	}
-	Accumulator.Accumulate(+Offset, Sample0, Weight0);
-	
-	const tsr_short2 OffsetN = tsr_short2(-int2(Offset)); // work arround a shader compiler bug
-	FSampleType Sample1 = AccessNeighborTexel(PreprocessedNeighborColor, OffsetN);
-	tsr_simd_half Weight1 = tsr_simd_half::Const(tsr_half(1.0));
-	if (FAccumulator::NeedWeightInAccumulate())
-	{
-		Weight1 = AccessNeighborTexel(Weight, OffsetN);
-	}
-	Accumulator.Accumulate(OffsetN, Sample1, Weight1);
-
-	Accumulator.Finish(3);
-	return Accumulator;
-}
-#endif // !CONFIG_OPTIMIZED_LDS_ACCESS
-
-CALL_SITE_DEBUGLOC
-template<class FAccumulator, typename FSampleType>
-FSampleType Convolve3x3HV(FSampleType Color, tsr_simd_half Weight)
-{
-	FSampleType HorizontalColor = Convole3x1<FAccumulator, FSampleType>(Color, Weight, tsr_short2(1, 0)).Value;
-	return Convole3x1<FAccumulator, FSampleType>(HorizontalColor, Weight, tsr_short2(0, 1)).Value;
-}
-
-
 //------------------------------------------------------- GENERAL OPERATIONS
 
-CALL_SITE_DEBUGLOC
-template<typename FSampleType>
-FSampleType Min3x3(FSampleType Center)
-{
-	return Convolve3x3HV<TAccumulatorMin<FSampleType>, FSampleType>(Center, tsr_simd_half::Const(1.0));
-}
-
-CALL_SITE_DEBUGLOC
-template<typename FSampleType>
-FSampleType Max3x3(FSampleType Center)
-{
-	return Convolve3x3HV<TAccumulatorMax<FSampleType>, FSampleType>(Center, tsr_simd_half::Const(1.0));
-}
-
 CALL_SITE_DEBUGLOC
 tsr_simd_halfC ClampFireFliersWithGuide(tsr_simd_halfC ToClamp, tsr_simd_halfC Guide)
 {
-	tsr_simd_halfC ToClampMin = Min3x3(ToClamp);
-	tsr_simd_halfC ToClampMax = Max3x3(ToClamp);
+	tsr_simd_halfC ToClampMin;
+	tsr_simd_halfC ToClampMax;
+	MinMax3x3(ToClamp, /* out */ ToClampMin, /* out */ ToClampMax);
 
 	tsr_simd_halfC ClampedGuide = clamp(Guide, ToClampMin, ToClampMax);
 	
-	tsr_simd_halfC ClampedGuideMin = Min3x3(ClampedGuide);
-	tsr_simd_halfC ClampedGuideMax = Max3x3(ClampedGuide);
+	tsr_simd_halfC ClampedGuideMin;
+	tsr_simd_halfC ClampedGuideMax;
+	MinMax3x3(ClampedGuide, /* out */ ClampedGuideMin, /* out */ ClampedGuideMax);
 	
 	return clamp(ToClamp, ClampedGuideMin, ClampedGuideMax);
 }
 
 
-//------------------------------------------------------- GENERAL OPERATIONS
-
-// Operator that output, lowest, median and highest values from 3 input values.
-CALL_SITE_DEBUGLOC
-template<typename FSampleType>
-void LMHOperator(FSampleType A, FSampleType B, FSampleType C, out FSampleType L, out FSampleType M, out FSampleType H)
-{
-	// TODO: min3, max3, med3.
-	FSampleType X = min(B, C);
-	FSampleType Y = max(B, C);
-	
-	L = min(A, X);
-	FSampleType Z = max(A, X);
-
-	M = min(Z, Y);
-	H = max(Z, Y);
-}
-
-// 3 samples median.
-CALL_SITE_DEBUGLOC
-template<typename FSampleType>
-FSampleType Median(FSampleType A, FSampleType B, FSampleType C)
-{
-	// TODO: med3.
-	FSampleType L, M, H;
-	LMHOperator(A, B, C, L, M, H);
-	return M;
-}
-
-// 9 samples median using Smith1996.
-CALL_SITE_DEBUGLOC
-template<typename FSampleType>
-FSampleType Median9(FSampleType In[9])
-{
-	// First layer.
-	FSampleType L0[3];
-	FSampleType M0[3];
-	FSampleType H0[3];
-	for (uint j = 0; j < 3; j++)
-	{
-		LMHOperator(In[j * 3 + 0], In[j * 3 + 1], In[j * 3 + 2], L0[j], M0[j], H0[j]);
-	}
-	
-	// Second layer.
-	// TODO: min3, max3.
-	FSampleType M1[3];
-	M1[0] = max(max(L0[0], L0[1]), L0[2]);
-	M1[1] = Median(M0[0], M0[1], M0[2]);
-	M1[2] = min(min(H0[0], H0[1]), H0[2]);
-	
-	// Third layer.
-	return Median(M1[0], M1[1], M1[2]);
-}
-
-CALL_SITE_DEBUGLOC
-template<typename FSampleType>
-FSampleType Median3x3(FSampleType Center)
-#if 1
-{
-	WriteToLDSForAccessNeighborTexel(Center);
-	
-	FSampleType C0 = Center;
-	FSampleType C1 = AccessNeighborTexel(Center, tsr_short2( 1, 0));
-	FSampleType C2 = AccessNeighborTexel(Center, tsr_short2(-1, 0));
-	
-	FSampleType M0L0;
-	FSampleType M0M0;
-	FSampleType M0H0;
-	LMHOperator(C0, C1, C2, M0L0, M0M0, M0H0);
-	
-	WriteToLDSForAccessNeighborTexel(M0L0);
-	FSampleType M0L1 = AccessNeighborTexel(M0L0, tsr_short2(0,  1));
-	FSampleType M0L2 = AccessNeighborTexel(M0L0, tsr_short2(0, -1));
-	
-	WriteToLDSForAccessNeighborTexel(M0M0);
-	FSampleType M0M1 = AccessNeighborTexel(M0M0, tsr_short2(0,  1));
-	FSampleType M0M2 = AccessNeighborTexel(M0M0, tsr_short2(0, -1));
-	
-	WriteToLDSForAccessNeighborTexel(M0H0);
-	FSampleType M0H1 = AccessNeighborTexel(M0H0, tsr_short2(0,  1));
-	FSampleType M0H2 = AccessNeighborTexel(M0H0, tsr_short2(0, -1));
-	
-	// Second layer.
-	// TODO: min3, max3.
-	FSampleType M1[3];
-	M1[0] = max(max(M0L0, M0L1), M0L2);
-	M1[1] = Median(M0M0, M0M1, M0M2);
-	M1[2] = min(min(M0H0, M0H1), M0H2);
-	
-	// Third layer.
-	return Median(M1[0], M1[1], M1[2]);
-}
-#else
-{
-	WriteToLDSForAccessNeighborTexel(Center);
-	
-	FSampleType InSamples[9];
-	UNROLL
-	for (uint i = 0; i < 9; i++)
-	{
-		FSampleType Sample = AccessNeighborTexel(Center, tsr_short2(kOffsets3x3[i]));
-		InSamples[i] = Sample;
-	}
-
-	return Median9(InSamples);
-}
-#endif
-
-
 //------------------------------------------------------- ENTRY POINT
 
 #if COMPILER_SUPPORTS_WAVE_SIZE && DIM_WAVE_SIZE > 0
@@ -630,7 +180,7 @@ void MainCS(
 		UNROLL_N(SIMD_SIZE)
 		for (uint SimdIndex = 0; SimdIndex < SIMD_SIZE; SimdIndex++)
 		{
-			tsr_short2 ClampedFetchPixelPos = GroupPixelOffset + GetLaneSimdPixelOffset(LaneIndex, SimdIndex);
+			tsr_short2 ClampedFetchPixelPos = GroupPixelOffset + GetLaneSimdPixelOffset<LANE_STRIDE_X, LANE_STRIDE_Y>(LaneIndex, SimdIndex);
 			float2 TranslucencyTextureUV = ApplyScreenTransform(float2(ClampedFetchPixelPos), InputPixelPosToTranslucencyTextureUV);
 
 			ClampedFetchPixelPos = ClampPixelOffset(ClampedFetchPixelPos, InputPixelPosMin, InputPixelPosMax);
@@ -673,7 +223,7 @@ void MainCS(
 
 	#if CONFIG_COMPOSE_TRANSLUCENCY
 	{
-		tsr_simd_half4 FilteredTranslucencyInput = Convolve3x3HV<TAccumulatorAvg<tsr_simd_half4>, tsr_simd_half4>(OriginalTranslucencyInput, bIsValid);
+		tsr_simd_half4 FilteredTranslucencyInput = Blur3x3(OriginalTranslucencyInput, bIsValid);
 
 		UNROLL_N(SIMD_SIZE)
 		for (uint SimdIndex = 0; SimdIndex < SIMD_SIZE; SimdIndex++)
@@ -734,7 +284,7 @@ void MainCS(
 		#endif
 		
 		// Convolve the input luma to have the exact same moire pattern as FilteredInput below.
-		InputLuma = Convolve3x3HV<TAccumulatorAvg<tsr_simd_half>, tsr_simd_half>(InputLuma, bIsValid);
+		InputLuma = Blur3x3(InputLuma, bIsValid);
 
 		// Unpack history
 		tsr_simd_half PrevFlickeringHistory = PrevMoire[0];
@@ -749,8 +299,9 @@ void MainCS(
 
 
 		// Compute clamping box of the limance.
-		tsr_simd_half InputLumaMin = Min3x3(InputLuma);
-		tsr_simd_half InputLumaMax = Max3x3(InputLuma);
+		tsr_simd_half InputLumaMin;
+		tsr_simd_half InputLumaMax;
+		MinMax3x3(InputLuma, /* out */ InputLumaMin, /* out */ InputLumaMax);
 
 		// Clamp and contribue the history.		
 		tsr_simd_half ClampedPrevFlickeringHistory = clamp(PrevFlickeringHistory, InputLumaMin, InputLumaMax);
@@ -864,8 +415,8 @@ void MainCS(
 	#endif
 
 	// Filter clamped input and history to increase quantization of the rejection detection
-	tsr_simd_halfC FilteredInput   = Convolve3x3HV<TAccumulatorAvg<tsr_simd_halfC>, tsr_simd_halfC>(ClampedInput  , bIsValid);
-	tsr_simd_halfC FilteredHistory = Convolve3x3HV<TAccumulatorAvg<tsr_simd_halfC>, tsr_simd_halfC>(ClampedHistory, bIsValid);
+	tsr_simd_halfC FilteredInput   = Blur3x3(ClampedInput  , bIsValid);
+	tsr_simd_halfC FilteredHistory = Blur3x3(ClampedHistory, bIsValid);
 	
 	// Measure rejection
 	tsr_simd_half RejectionBlendFinal = tsr_simd_half::Const(1.0);
@@ -876,8 +427,11 @@ void MainCS(
 
 		tsr_simd_halfC BackbufferError = tsr_simd_halfC::Const(MeasureBackbufferLDRQuantizationError());
 
-		tsr_simd_halfC FilteredBoxMin = Min3x3(FilteredInput) - BackbufferError;
-		tsr_simd_halfC FilteredBoxMax = Max3x3(FilteredInput) + BackbufferError;
+		tsr_simd_halfC FilteredBoxMin;
+		tsr_simd_halfC FilteredBoxMax;
+		MinMax3x3(FilteredInput, /* out */ FilteredBoxMin, /* out */ FilteredBoxMax);
+		FilteredBoxMin = FilteredBoxMin - BackbufferError;
+		FilteredBoxMax = FilteredBoxMax + BackbufferError;
 		
 		tsr_simd_halfC FlickeringClampedFilteredHistory = clamp(FilteredHistory, FilteredBoxMin, FilteredBoxMax);
 
@@ -911,8 +465,9 @@ void MainCS(
 			ClampedFilteredHistory = FlickeringClampedFilteredHistory;
 		#endif
 		
-		tsr_simd_halfC OriginalBoxMin = Min3x3(OriginalInput);
-		tsr_simd_halfC OriginalBoxMax = Max3x3(OriginalInput);
+		tsr_simd_halfC OriginalBoxMin;
+		tsr_simd_halfC OriginalBoxMax;
+		MinMax3x3(OriginalInput, /* out */ OriginalBoxMin, /* out */ OriginalBoxMax);
 		tsr_simd_halfC BoxSize = (OriginalBoxMax - OriginalBoxMin + (BackbufferError + BackbufferError)) * FilteringWeight;
 		
 		const tsr_simd_halfC BackbufferQuantizationErrorVector = tsr_simd_halfC::Const(MeasureBackbufferLDRQuantizationError());
@@ -997,7 +552,7 @@ void MainCS(
 		UNROLL_N(SIMD_SIZE)
 		for (uint SimdIndex = 0; SimdIndex < SIMD_SIZE; SimdIndex++)
 		{
-			tsr_short2 LaneSimdPixelOffset = GetLaneSimdPixelOffset(LaneIndex, SimdIndex);
+			tsr_short2 LaneSimdPixelOffset = GetLaneSimdPixelOffset<LANE_STRIDE_X, LANE_STRIDE_Y>(LaneIndex, SimdIndex);
 			bool bIsNonOverscanPixel = all(and(LaneSimdPixelOffset >= TILE_OVERSCAN_PADDING, LaneSimdPixelOffset < TILE_SIZE + TILE_OVERSCAN_PADDING));
 
 			tsr_short2 OutputPixelPos = InvalidateOutputPixelPos(GroupPixelOffset + LaneSimdPixelOffset, InputInfo_ViewportMax);
diff --git a/Engine/Shaders/Public/LaneVectorization.ush b/Engine/Shaders/Public/LaneVectorization.ush
--- a/Engine/Shaders/Public/LaneVectorization.ush
+++ b/Engine/Shaders/Public/LaneVectorization.ush
@@ -8,10 +8,12 @@
 
 #include "Platform.ush"
 
-
-template<typename ScalarType, uint VectorSize, uint SimdSize>
+template<typename ScalarType, uint InVectorSize, uint InSimdSize>
 struct TLaneVector
 {
+	static const uint VectorSize = InVectorSize;
+	static const uint SimdSize = InSimdSize;
+
 	vector<ScalarType, VectorSize> RegisterArray[SimdSize];
 
 	
@@ -172,7 +174,176 @@ struct TLaneVector
 
 };
 
+template<typename ScalarType, uint InVectorSize, uint InSimdSizeX, uint InSimdSizeY>
+struct TLaneVector2D
+{
+	static const uint VectorSize = InVectorSize;
+	static const uint SimdSize = InSimdSizeX * InSimdSizeY;
+	static const uint SimdSizeX = InSimdSizeX;
+	static const uint SimdSizeY = InSimdSizeY;
+
+	vector<ScalarType, VectorSize> RegisterArray[SimdSize];
+
+	
+	// Access a component of the vector
+	CALL_SITE_DEBUGLOC
+	TLaneVector2D<ScalarType, 1, SimdSizeX, SimdSizeY> operator [](const uint ComponentIndex)
+	{
+		TLaneVector2D<ScalarType, 1, SimdSizeX, SimdSizeY> R;
+		UNROLL
+		for (uint SimdIndex = 0; SimdIndex < SimdSize; SimdIndex++)
+		{
+			R.SetElement(SimdIndex, GetElement(SimdIndex)[ComponentIndex]);
+		}
+		return R;
+	}
+	
+	// Access an individual element of the vector.
+	CALL_SITE_DEBUGLOC
+	vector<ScalarType, VectorSize> GetElement(const uint SimdIndex)
+	{
+		return RegisterArray[SimdIndex];
+	}
+	
+	CALL_SITE_DEBUGLOC
+	void SetElement(const uint SimdIndex, vector<ScalarType, VectorSize> Element)
+	{
+		RegisterArray[SimdIndex] = Element;
+	}
+	
+	CALL_SITE_DEBUGLOC
+	void SetAllElements(vector<ScalarType, VectorSize> Element)
+	{
+		UNROLL
+		for (uint SimdIndex = 0; SimdIndex < SimdSize; SimdIndex++)
+		{
+			RegisterArray[SimdIndex] = Element;
+		}
+	}
+	
+	
+	// ------------- binary operator +
+	CALL_SITE_DEBUGLOC
+	TLaneVector2D<ScalarType, VectorSize, SimdSizeX, SimdSizeY> operator + (TLaneVector2D<ScalarType, VectorSize, SimdSizeX, SimdSizeY> B)
+	{
+		TLaneVector2D<ScalarType, VectorSize, SimdSizeX, SimdSizeY> R;
+		UNROLL
+		for (uint SimdIndex = 0; SimdIndex < SimdSize; SimdIndex++)
+		{
+			R.SetElement(SimdIndex, GetElement(SimdIndex) + B.GetElement(SimdIndex));
+		}
+		return R;
+	}
+	
+	CALL_SITE_DEBUGLOC
+	TLaneVector2D<ScalarType, VectorSize, SimdSizeX, SimdSizeY> operator + (vector<ScalarType, VectorSize> B)
+	{
+		TLaneVector2D<ScalarType, VectorSize, SimdSizeX, SimdSizeY> R;
+		UNROLL
+		for (uint SimdIndex = 0; SimdIndex < SimdSize; SimdIndex++)
+		{
+			R.SetElement(SimdIndex, GetElement(SimdIndex) + B);
+		}
+		return R;
+	}
+
+	
+	// ------------- unary operator -
+	CALL_SITE_DEBUGLOC
+	TLaneVector2D<ScalarType, VectorSize, SimdSizeX, SimdSizeY> operator - ()
+	{
+		TLaneVector2D<ScalarType, VectorSize, SimdSizeX, SimdSizeY> R;
+		UNROLL
+		for (uint SimdIndex = 0; SimdIndex < SimdSize; SimdIndex++)
+		{
+			R.SetElement(SimdIndex, -GetElement(SimdIndex));
+		}
+		return R;
+	}
+	
+	// ------------- binary operator -
+	CALL_SITE_DEBUGLOC
+	TLaneVector2D<ScalarType, VectorSize, SimdSizeX, SimdSizeY> operator - (TLaneVector2D<ScalarType, VectorSize, SimdSizeX, SimdSizeY> B)
+	{
+		TLaneVector2D<ScalarType, VectorSize, SimdSizeX, SimdSizeY> R;
+		UNROLL
+		for (uint SimdIndex = 0; SimdIndex < SimdSize; SimdIndex++)
+		{
+			R.SetElement(SimdIndex, GetElement(SimdIndex) - B.GetElement(SimdIndex));
+		}
+		return R;
+	}
+	
+	CALL_SITE_DEBUGLOC
+	TLaneVector2D<ScalarType, VectorSize, SimdSizeX, SimdSizeY> operator - (vector<ScalarType, VectorSize> B)
+	{
+		TLaneVector2D<ScalarType, VectorSize, SimdSizeX, SimdSizeY> R;
+		UNROLL
+		for (uint SimdIndex = 0; SimdIndex < SimdSize; SimdIndex++)
+		{
+			R.SetElement(SimdIndex, GetElement(SimdIndex) - B);
+		}
+		return R;
+	}
+
+	
+	// ------------- binary operator *
+	CALL_SITE_DEBUGLOC
+	TLaneVector2D<ScalarType, VectorSize, SimdSizeX, SimdSizeY> operator * (TLaneVector2D<ScalarType, VectorSize, SimdSizeX, SimdSizeY> B)
+	{
+		TLaneVector2D<ScalarType, VectorSize, SimdSizeX, SimdSizeY> R;
+		UNROLL
+		for (uint SimdIndex = 0; SimdIndex < SimdSize; SimdIndex++)
+		{
+			R.SetElement(SimdIndex, GetElement(SimdIndex) * B.GetElement(SimdIndex));
+		}
+		return R;
+	}
+
+	CALL_SITE_DEBUGLOC
+	TLaneVector2D<ScalarType, VectorSize, SimdSizeX, SimdSizeY> operator * (vector<ScalarType, VectorSize> B)
+	{
+		TLaneVector2D<ScalarType, VectorSize, SimdSizeX, SimdSizeY> R;
+		UNROLL
+		for (uint SimdIndex = 0; SimdIndex < SimdSize; SimdIndex++)
+		{
+			R.SetElement(SimdIndex, GetElement(SimdIndex) * B);
+		}
+		return R;
+	}
+	
+
+	// ------------- constructors
+	CALL_SITE_DEBUGLOC
+	static TLaneVector2D<ScalarType, VectorSize, SimdSizeX, SimdSizeY> Vectorize(
+		TLaneVector2D<ScalarType, 1, SimdSizeX, SimdSizeY> A)
+	{
+		TLaneVector2D<ScalarType, VectorSize, SimdSizeX, SimdSizeY> R;
+		UNROLL
+		for (uint SimdIndex = 0; SimdIndex < SimdSize; SimdIndex++)
+		{
+			R.SetElement(SimdIndex, A.GetElement(SimdIndex));
+		}
+		return R;
+	}
+	
+	CALL_SITE_DEBUGLOC
+	static TLaneVector2D<ScalarType, VectorSize, SimdSizeX, SimdSizeY> Const(
+		vector<ScalarType, VectorSize> A)
+	{
+		TLaneVector2D<ScalarType, VectorSize, SimdSizeX, SimdSizeY> R;
+		UNROLL
+		for (uint SimdIndex = 0; SimdIndex < SimdSize; SimdIndex++)
+		{
+			R.SetElement(SimdIndex, A);
+		}
+		return R;
+	}
+
+};
+
 // ------------- 1 parameter
+
 #define TVECTOR_FUNCTION_1PARAMS(FunctionName) \
 	CALL_SITE_DEBUGLOC \
 	template<typename ScalarType, uint VectorSize, uint SimdSize> \
@@ -186,13 +357,28 @@ struct TLaneVector
 			R.SetElement(SimdIndex, FunctionName(A.GetElement(SimdIndex))); \
 		} \
 		return R; \
-	}
+	} \
+	\
+	CALL_SITE_DEBUGLOC \
+	template<typename ScalarType, uint VectorSize, uint SimdSizeX, uint SimdSizeY> \
+	TLaneVector2D<ScalarType, VectorSize, SimdSizeX, SimdSizeY> FunctionName( \
+		TLaneVector2D<ScalarType, VectorSize, SimdSizeX, SimdSizeY> A) \
+	{ \
+		TLaneVector2D<ScalarType, VectorSize, SimdSizeX, SimdSizeY> R; \
+		UNROLL \
+		for (uint SimdIndex = 0; SimdIndex < SimdSizeX * SimdSizeY; SimdIndex++) \
+		{ \
+			R.SetElement(SimdIndex, FunctionName(A.GetElement(SimdIndex))); \
+		} \
+		return R; \
+	} \
 	
 TVECTOR_FUNCTION_1PARAMS(rcp);
 TVECTOR_FUNCTION_1PARAMS(saturate);
 TVECTOR_FUNCTION_1PARAMS(abs);
 
 // ------------- 2 parameters
+
 #define TVECTOR_FUNCTION_2PARAMS(ReturnVectorSize, FunctionName) \
 	CALL_SITE_DEBUGLOC \
 	template<typename ScalarType, uint VectorSize, uint SimdSize> \
@@ -205,7 +391,20 @@ TVECTOR_FUNCTION_1PARAMS(abs);
 			R.SetElement(SimdIndex, FunctionName(A.GetElement(SimdIndex), B.GetElement(SimdIndex))); \
 		} \
 		return R; \
-	}
+	} \
+	\
+	CALL_SITE_DEBUGLOC \
+	template<typename ScalarType, uint VectorSize, uint SimdSizeX, uint SimdSizeY> \
+	TLaneVector2D<ScalarType, ReturnVectorSize, SimdSizeX, SimdSizeY> FunctionName( \
+		TLaneVector2D<ScalarType, VectorSize, SimdSizeX, SimdSizeY> A, \
+		TLaneVector2D<ScalarType, VectorSize, SimdSizeX, SimdSizeY> B) \
+	{ \
+		TLaneVector2D<ScalarType, ReturnVectorSize, SimdSizeX, SimdSizeY> R; \
+		UNROLL for (uint SimdIndex = 0; SimdIndex < SimdSizeX * SimdSizeY; SimdIndex++) { \
+			R.SetElement(SimdIndex, FunctionName(A.GetElement(SimdIndex), B.GetElement(SimdIndex))); \
+		} \
+		return R; \
+	} \
 	
 TVECTOR_FUNCTION_2PARAMS(VectorSize, min);
 TVECTOR_FUNCTION_2PARAMS(VectorSize, max);
@@ -226,7 +425,21 @@ TVECTOR_FUNCTION_2PARAMS(/* VectorSize = */ 1, dot);
 			R.SetElement(SimdIndex, FunctionName(A.GetElement(SimdIndex), B.GetElement(SimdIndex), C.GetElement(SimdIndex))); \
 		} \
 		return R; \
-	}
+	} \
+	\
+	CALL_SITE_DEBUGLOC \
+	template<typename ScalarType, uint VectorSize, uint SimdSizeX, uint SimdSizeY> \
+	TLaneVector2D<ScalarType, ReturnVectorSize, SimdSizeX, SimdSizeY> FunctionName( \
+		TLaneVector2D<ScalarType, VectorSize, SimdSizeX, SimdSizeY> A, \
+		TLaneVector2D<ScalarType, VectorSize, SimdSizeX, SimdSizeY> B, \
+		TLaneVector2D<ScalarType, VectorSize, SimdSizeX, SimdSizeY> C) \
+	{ \
+		TLaneVector2D<ScalarType, ReturnVectorSize, SimdSizeX, SimdSizeY> R; \
+		UNROLL for (uint SimdIndex = 0; SimdIndex < SimdSizeX * SimdSizeY; SimdIndex++) { \
+			R.SetElement(SimdIndex, FunctionName(A.GetElement(SimdIndex), B.GetElement(SimdIndex), C.GetElement(SimdIndex))); \
+		} \
+		return R; \
+	} \
 	
 TVECTOR_FUNCTION_3PARAMS(VectorSize, clamp);
 TVECTOR_FUNCTION_3PARAMS(VectorSize, min3);
