Substrate - removed usless tree from legacy conversion functions.

diff --git a/Engine/Shaders/Private/Strata/StrataLegacyConversion.ush b/Engine/Shaders/Private/Strata/StrataLegacyConversion.ush
--- a/Engine/Shaders/Private/Strata/StrataLegacyConversion.ush
+++ b/Engine/Shaders/Private/Strata/StrataLegacyConversion.ush
@@ -85,7 +85,7 @@ FStrataLegacyParameters InitStrataLegacyParameters(FStrataPixelFootprint InPixel
 }
 
 FStrataData CreateLegacySlab(
-	FStrataLegacyParameters In, inout uint SharedLocalBasisTypes, inout FStrataTree StrataTree)
+	FStrataLegacyParameters In, inout uint SharedLocalBasisTypes)
 {
 	const float3 Zeros = float3(0, 0, 0);
 	const float3 Ones  = float3(1, 1, 1);
@@ -147,8 +147,7 @@ FStrataData StrataConvertLegacyMaterialDynamic(
 	float3 RawCustomTangent,
 	uint SharedLocalBasisIndex,
 	uint ClearCoatBottomNormal_SharedLocalBasisIndex,
-	inout uint SharedLocalBasisTypes,
-	inout FStrataTree StrataTree)
+	inout uint SharedLocalBasisTypes)
 {
 	const float DefaultThickness = STRATA_LAYER_DEFAULT_THICKNESS_CM;
 	const float3 Zeros = float3(0, 0, 0);
@@ -176,7 +175,7 @@ FStrataData StrataConvertLegacyMaterialDynamic(
 			ColorScaleBehindWater,	// ColorScaleBehindWater
 			SharedLocalBasisIndex);
 	}
-	#endif
+	#endif // MATERIAL_SHADINGMODEL_SINGLELAYERWATER
 	#if MATERIAL_SHADINGMODEL_HAIR
 	if (ShadingModel == SHADINGMODELID_HAIR)
 	{
@@ -189,7 +188,7 @@ FStrataData StrataConvertLegacyMaterialDynamic(
 			Emissive,	   // EmissiveColor
 			SharedLocalBasisIndex);
 	}
-	#endif
+	#endif // MATERIAL_SHADINGMODEL_HAIR
 	#if MATERIAL_SHADINGMODEL_EYE
 	if (ShadingModel == SHADINGMODELID_EYE)
 	{
@@ -211,7 +210,7 @@ FStrataData StrataConvertLegacyMaterialDynamic(
 				Emissive,	         // EmissiveColor
 				SharedLocalBasisIndex);
 	}
-	#endif
+	#endif // MATERIAL_SHADINGMODEL_EYE
 
 	FStrataLegacyParameters LegacySlab = InitStrataLegacyParameters(PixelFootprint, SharedLocalBasisIndex, 1.0f);
 
@@ -432,7 +431,7 @@ FStrataData StrataConvertLegacyMaterialDynamic(
 	}
 
 	// Fixed layer structure for helping compiler to unroll and optimize shader
-	FStrataData Slab = CreateLegacySlab(LegacySlab, SharedLocalBasisTypes, StrataTree);
+	FStrataData Slab = CreateLegacySlab(LegacySlab, SharedLocalBasisTypes);
 #if STRATA_INLINE_SHADING
 	Slab.InlinedBSDF.Coverage = LegacySlab.Weight * FinalWeight;	//////////////////// Remove FinalWeight? Do that in CreateLegacySlab?
 #endif
@@ -465,8 +464,7 @@ FStrataData StrataConvertLegacyMaterialStatic(
 	float3 RawCustomTangent,
 	uint SharedLocalBasisIndex,
 	uint ClearCoatBottomNormal_SharedLocalBasisIndex,
-	inout uint SharedLocalBasisTypes,
-	inout FStrataTree StrataTree)
+	inout uint SharedLocalBasisTypes)
 {
 	FStrataData Out = GetInitialisedStrataData();
 
@@ -574,8 +572,7 @@ FStrataData StrataConvertLegacyMaterialStatic(
 			RawCustomTangent,
 			SharedLocalBasisIndex,
 			ClearCoatBottomNormal_SharedLocalBasisIndex,
-			SharedLocalBasisTypes,
-			StrataTree);
+			SharedLocalBasisTypes);
 	}
 
 	return Out;
diff --git a/Engine/Source/Runtime/Engine/Private/Materials/HLSLMaterialTranslator.cpp b/Engine/Source/Runtime/Engine/Private/Materials/HLSLMaterialTranslator.cpp
--- a/Engine/Source/Runtime/Engine/Private/Materials/HLSLMaterialTranslator.cpp
+++ b/Engine/Source/Runtime/Engine/Private/Materials/HLSLMaterialTranslator.cpp
@@ -11550,7 +11550,7 @@ int32 FHLSLMaterialTranslator::StrataConversionFromLegacy(
 		}
 
 		return AddCodeChunk(
-			MCT_Strata, TEXT("PromoteParameterBlendedBSDFToOperator(StrataConvertLegacyMaterial%s(Parameters.StrataPixelFootprint, %s, %s, %s, %s, %s, %s, %s, %s, %s, %s, %s, %s, %s, %s, %s, %s, %s, %s, %s, %s, %s, %s, %s, Parameters.%s.Types, Parameters.%s) /* Normal = %s ; Tangent = %s ; ClearCoat_Normal = %s ; ClearCoat_Tangent = %s */, Parameters.%s, %u, %u, %u, %u)"),
+			MCT_Strata, TEXT("PromoteParameterBlendedBSDFToOperator(StrataConvertLegacyMaterial%s(Parameters.StrataPixelFootprint, %s, %s, %s, %s, %s, %s, %s, %s, %s, %s, %s, %s, %s, %s, %s, %s, %s, %s, %s, %s, %s, %s, %s, Parameters.%s.Types) /* Normal = %s ; Tangent = %s ; ClearCoat_Normal = %s ; ClearCoat_Tangent = %s */, Parameters.%s, %u, %u, %u, %u)"),
 			bHasDynamicShadingModels ? TEXT("Dynamic") : TEXT("Static"),
 			*StrataGetCastParameterCode(BaseColor,						MCT_Float3),
 			*StrataGetCastParameterCode(Specular,						MCT_Float),
@@ -11577,7 +11577,6 @@ int32 FHLSLMaterialTranslator::StrataConversionFromLegacy(
 			*SharedLocalBasisIndexMacro,
 			*ClearCoat_SharedLocalBasisIndexMacro,
 			*GetParametersSharedLocalBasesName(CurrentStrataCompilationContext),
-			*GetParametersStrataTreeName(CurrentStrataCompilationContext),
 			// Regular normal basis
 			*NormalCode,
 			*TangentCode,
@@ -11593,7 +11592,7 @@ int32 FHLSLMaterialTranslator::StrataConversionFromLegacy(
 	}
 
 	return AddCodeChunk(
-		MCT_Strata, TEXT("StrataConvertLegacyMaterial%s(Parameters.StrataPixelFootprint, %s, %s, %s, %s, %s, %s, %s, %s, %s, %s, %s, %s, %s, %s, %s, %s, %s, %s, %s, %s, %s, %s, %s, Parameters.%s.Types, Parameters.%s) /* Normal = %s ; Tangent = %s ; ClearCoat_Normal = %s ; ClearCoat_Tangent = %s */"),
+		MCT_Strata, TEXT("StrataConvertLegacyMaterial%s(Parameters.StrataPixelFootprint, %s, %s, %s, %s, %s, %s, %s, %s, %s, %s, %s, %s, %s, %s, %s, %s, %s, %s, %s, %s, %s, %s, %s, Parameters.%s.Types) /* Normal = %s ; Tangent = %s ; ClearCoat_Normal = %s ; ClearCoat_Tangent = %s */"),
 		bHasDynamicShadingModels ? TEXT("Dynamic") : TEXT("Static"),
 		*StrataGetCastParameterCode(BaseColor,						MCT_Float3),
 		*StrataGetCastParameterCode(Specular,						MCT_Float),
@@ -11620,7 +11619,6 @@ int32 FHLSLMaterialTranslator::StrataConversionFromLegacy(
 		*SharedLocalBasisIndexMacro,
 		*ClearCoat_SharedLocalBasisIndexMacro,
 		*GetParametersSharedLocalBasesName(CurrentStrataCompilationContext),
-		*GetParametersStrataTreeName(CurrentStrataCompilationContext),
 		// Regular normal basis
 		*NormalCode,
 		*TangentCode,
