Speeds up FTSRResolveHistoryCS with convolution network framework

diff --git a/Engine/Shaders/Private/TemporalSuperResolution/TSRConvolutionNetwork.ush b/Engine/Shaders/Private/TemporalSuperResolution/TSRConvolutionNetwork.ush
--- a/Engine/Shaders/Private/TemporalSuperResolution/TSRConvolutionNetwork.ush
+++ b/Engine/Shaders/Private/TemporalSuperResolution/TSRConvolutionNetwork.ush
@@ -48,6 +48,14 @@ tsr_short2 GetSimdIndexPixelCoordinateInLane(const uint SimdIndex)
 	return tsr_short2(SimdIndex % LaneStrideX, SimdIndex / LaneStrideX);
 }
 
+CALL_SITE_DEBUGLOC
+template<uint LaneStrideX, uint LaneStrideY>
+uint GetPixelCoordinateInLaneSimdIndex(const tsr_short2 PixelCoordinateInLane)
+{
+	return uint(dot(PixelCoordinateInLane, tsr_short2(1, LaneStrideX)));
+}
+
+
 /** Returns coordinate of a pixel within a wave. */
 CALL_SITE_DEBUGLOC
 template<uint LaneStrideX, uint LaneStrideY>
@@ -488,3 +496,75 @@ TLaneVector2D<ScalarType, VectorSize, LaneStrideX, LaneStrideY> Median3x3(TLaneV
 	return Median9(InSamples);
 }
 #endif
+
+
+//------------------------------------------------------- DOWNSAMPLE 2x2 CONVOLUTIONS
+
+CALL_SITE_DEBUGLOC
+template<typename ScalarType, uint VectorSize, uint LaneStrideX, uint LaneStrideY>
+TLaneVector2D<ScalarType, VectorSize, LaneStrideX / 2, LaneStrideY / 2> DownsampleMin2x2(
+	TLaneVector2D<ScalarType, VectorSize, LaneStrideX, LaneStrideY> Input)
+{
+	TLaneVector2D<ScalarType, VectorSize, LaneStrideX / 2, LaneStrideY / 2> OutMin;
+	UNROLL
+	for (uint OutputSimdIndex = 0; OutputSimdIndex < ((LaneStrideX * LaneStrideY) / 4); OutputSimdIndex++)
+	{
+		const uint2 OutputPos = GetSimdIndexPixelCoordinateInLane<LaneStrideX / 2, LaneStrideY / 2>(OutputSimdIndex);
+		vector<ScalarType, VectorSize> Input0 = Input.GetElement(GetPixelCoordinateInLaneSimdIndex<LaneStrideX, LaneStrideY>(tsr_short2(OutputPos * 2 + Offsets2x2[0])));
+		vector<ScalarType, VectorSize> Input1 = Input.GetElement(GetPixelCoordinateInLaneSimdIndex<LaneStrideX, LaneStrideY>(tsr_short2(OutputPos * 2 + Offsets2x2[1])));
+		vector<ScalarType, VectorSize> Input2 = Input.GetElement(GetPixelCoordinateInLaneSimdIndex<LaneStrideX, LaneStrideY>(tsr_short2(OutputPos * 2 + Offsets2x2[2])));
+		vector<ScalarType, VectorSize> Input3 = Input.GetElement(GetPixelCoordinateInLaneSimdIndex<LaneStrideX, LaneStrideY>(tsr_short2(OutputPos * 2 + Offsets2x2[3])));
+
+		OutMin.SetElement(OutputSimdIndex, min(min(Input0, Input1), min(Input2, Input3)));
+	}
+	return OutMin;
+}
+
+CALL_SITE_DEBUGLOC
+template<typename ScalarType, uint VectorSize, uint LaneStrideX, uint LaneStrideY>
+TLaneVector2D<ScalarType, VectorSize, LaneStrideX / 2, LaneStrideY / 2> DownsampleMax2x2(
+	TLaneVector2D<ScalarType, VectorSize, LaneStrideX, LaneStrideY> Input)
+{
+	TLaneVector2D<ScalarType, VectorSize, LaneStrideX / 2, LaneStrideY / 2> OutMax;
+	UNROLL
+	for (uint OutputSimdIndex = 0; OutputSimdIndex < ((LaneStrideX * LaneStrideY) / 4); OutputSimdIndex++)
+	{
+		const uint2 OutputPos = GetSimdIndexPixelCoordinateInLane<LaneStrideX / 2, LaneStrideY / 2>(OutputSimdIndex);
+		vector<ScalarType, VectorSize> Input0 = Input.GetElement(GetPixelCoordinateInLaneSimdIndex<LaneStrideX, LaneStrideY>(tsr_short2(OutputPos * 2 + Offsets2x2[0])));
+		vector<ScalarType, VectorSize> Input1 = Input.GetElement(GetPixelCoordinateInLaneSimdIndex<LaneStrideX, LaneStrideY>(tsr_short2(OutputPos * 2 + Offsets2x2[1])));
+		vector<ScalarType, VectorSize> Input2 = Input.GetElement(GetPixelCoordinateInLaneSimdIndex<LaneStrideX, LaneStrideY>(tsr_short2(OutputPos * 2 + Offsets2x2[2])));
+		vector<ScalarType, VectorSize> Input3 = Input.GetElement(GetPixelCoordinateInLaneSimdIndex<LaneStrideX, LaneStrideY>(tsr_short2(OutputPos * 2 + Offsets2x2[3])));
+
+		OutMax.SetElement(OutputSimdIndex, max(max(Input0, Input1), max(Input2, Input3)));
+	}
+	return OutMax;
+}
+
+CALL_SITE_DEBUGLOC
+template<typename ScalarType, uint VectorSize, uint LaneStrideX, uint LaneStrideY>
+TLaneVector2D<ScalarType, VectorSize, LaneStrideX / 2, LaneStrideY / 2> DownsampleDot2x2(
+	TLaneVector2D<ScalarType, VectorSize, LaneStrideX, LaneStrideY> Input, const ScalarType Weights[4])
+{
+	TLaneVector2D<ScalarType, VectorSize, LaneStrideX / 2, LaneStrideY / 2> OutMax;
+	UNROLL
+	for (uint OutputSimdIndex = 0; OutputSimdIndex < ((LaneStrideX * LaneStrideY) / 4); OutputSimdIndex++)
+	{
+		const uint2 OutputPos = GetSimdIndexPixelCoordinateInLane<LaneStrideX / 2, LaneStrideY / 2>(OutputSimdIndex);
+		vector<ScalarType, VectorSize> Input0 = Input.GetElement(GetPixelCoordinateInLaneSimdIndex<LaneStrideX, LaneStrideY>(tsr_short2(OutputPos * 2 + Offsets2x2[0])));
+		vector<ScalarType, VectorSize> Input1 = Input.GetElement(GetPixelCoordinateInLaneSimdIndex<LaneStrideX, LaneStrideY>(tsr_short2(OutputPos * 2 + Offsets2x2[1])));
+		vector<ScalarType, VectorSize> Input2 = Input.GetElement(GetPixelCoordinateInLaneSimdIndex<LaneStrideX, LaneStrideY>(tsr_short2(OutputPos * 2 + Offsets2x2[2])));
+		vector<ScalarType, VectorSize> Input3 = Input.GetElement(GetPixelCoordinateInLaneSimdIndex<LaneStrideX, LaneStrideY>(tsr_short2(OutputPos * 2 + Offsets2x2[3])));
+
+		OutMax.SetElement(OutputSimdIndex, Input0 * Weights[0] + Input1 * Weights[1] + Input2 * Weights[2] + Input3 * Weights[3]);
+	}
+	return OutMax;
+}
+
+CALL_SITE_DEBUGLOC
+template<typename ScalarType, uint VectorSize, uint LaneStrideX, uint LaneStrideY>
+TLaneVector2D<ScalarType, VectorSize, LaneStrideX / 2, LaneStrideY / 2> DownsampleAvg2x2(
+	TLaneVector2D<ScalarType, VectorSize, LaneStrideX, LaneStrideY> Input)
+{
+	const ScalarType Weights[4] = { ScalarType(0.25), ScalarType(0.25), ScalarType(0.25), ScalarType(0.25) };
+	return DownsampleDot2x2(Input, Weights);
+}
diff --git a/Engine/Shaders/Private/TemporalSuperResolution/TSRResolveHistory.usf b/Engine/Shaders/Private/TemporalSuperResolution/TSRResolveHistory.usf
--- a/Engine/Shaders/Private/TemporalSuperResolution/TSRResolveHistory.usf
+++ b/Engine/Shaders/Private/TemporalSuperResolution/TSRResolveHistory.usf
@@ -5,10 +5,48 @@
 
 //------------------------------------------------------- CONFIG
 
-#define TILE_SIZE 8
+#if DIM_NYQUIST
+	#define LANE_COUNT 32
+	#define DIM_WAVE_SIZE (LANE_COUNT)
+	#define DOWNSAMPLE_FACTOR 2u
+	
+	/** Size of the entire group loaded into memory including overscan. */
+	#define GROUP_TILE_SIZE 8u
+	#define TILE_SIZE (GROUP_TILE_SIZE - 2u)
+
+	#define LANE_OUTPUT_STRIDE_X 2u
+	#define LANE_OUTPUT_STRIDE_Y 1u
+	#define LANE_INPUT_STRIDE_X (LANE_OUTPUT_STRIDE_X * DOWNSAMPLE_FACTOR)
+	#define LANE_INPUT_STRIDE_Y (LANE_OUTPUT_STRIDE_Y * DOWNSAMPLE_FACTOR)
+	
+	/** Number of lane the GROUP_TILE_SIZE. */
+	#define LANE_COUNT_X (GROUP_TILE_SIZE / LANE_OUTPUT_STRIDE_X)
+	#define LANE_COUNT_Y (GROUP_TILE_SIZE / LANE_OUTPUT_STRIDE_Y)
+
+	/** Size of the SIMD per lane that also number of input pixel loaded into each individual lanes. */
+	#define INPUT_SIMD_SIZE (LANE_INPUT_STRIDE_X * LANE_INPUT_STRIDE_Y)
+	#define OUTPUT_SIMD_SIZE (LANE_OUTPUT_STRIDE_X * LANE_OUTPUT_STRIDE_Y)
+
+	#define tsr_input_halfC  TLaneVector2D<tsr_half, CONFIG_CHANNEL_COUNT, LANE_INPUT_STRIDE_X, LANE_INPUT_STRIDE_Y>
+	#define tsr_output_halfC TLaneVector2D<tsr_half, CONFIG_CHANNEL_COUNT, LANE_OUTPUT_STRIDE_X, LANE_OUTPUT_STRIDE_Y>
+
+#else
+	#define TILE_SIZE 8
+#endif
 
 #define KERNEL_SIZE 4
 
+#define CONFIG_CLAMP 1
+
+#define CONFIG_HDR_WEIGHT 1
+
+
+//------------------------------------------------------- CONVOLUTIONS
+
+#if DIM_NYQUIST
+	#include "TSRConvolutionNetwork.ush"
+#endif
+
 
 //------------------------------------------------------- CONSTANTS
 
@@ -33,6 +71,9 @@ groupshared tsr_halfC SharedArray0[TILE_SIZE * TILE_SIZE];
 
 //------------------------------------------------------- FUNCTIONS
 
+static const float MitchellNetravaliB = rcp(3.0);
+static const float MitchellNetravaliC = rcp(3.0);
+
 // Mitchell Netravali
 // B-spline:       B=1      C=0
 // Mitchell:       B=1/3    C=1/3
@@ -83,16 +124,153 @@ uint2 OutputViewRectMax;
 uint bGenerateOutputMip1;
 float HistoryValidityMultiply;
 
-Texture2D<tsr_halfC> History_HighFrequency;
-Texture2D<tsr_half2> History_Metadata;
-Texture2D<tsr_half3> History_Translucency;
-Texture2D<tsr_half>  History_TranslucencyAlpha;
+Texture2D<tsr_halfC> UpdateHistoryOutputTexture;
 
 RWTexture2D<tsr_halfC> SceneColorOutputMip0;
 RWTexture2D<tsr_halfC> SceneColorOutputMip1;
 
 
-//------------------------------------------------------- ENTRY POINT
+//------------------------------------------------------- ENTRY POINT NYQUIST DOWNSAMPLE
+
+#if DIM_NYQUIST
+
+#if COMPILER_SUPPORTS_WAVE_SIZE
+	WAVESIZE(LANE_COUNT)
+#endif
+[numthreads(LANE_COUNT, 1, 1)]
+void MainCS(
+	uint2 GroupId : SV_GroupID,
+	uint GroupThreadIndex : SV_GroupIndex)
+{
+	const uint LaneIndex = GroupThreadIndex;
+
+	// Load all the colors
+	tsr_input_halfC Colors;
+	ISOLATE
+	{
+		const tsr_short2 GroupPixelOffset = (
+			tsr_short2(HistoryInfo_ViewportMin) +
+			tsr_short2(GroupId) * tsr_short(TILE_SIZE * DOWNSAMPLE_FACTOR).xx +
+			tsr_short(-1).xx);
+		
+		UNROLL_N(INPUT_SIMD_SIZE)
+		for (uint SimdIndex = 0; SimdIndex < INPUT_SIMD_SIZE; SimdIndex++)
+		{
+			tsr_short2 PixelPos = GroupPixelOffset + GetLaneSimdPixelOffset<LANE_INPUT_STRIDE_X, LANE_INPUT_STRIDE_Y>(LaneIndex, SimdIndex);
+			PixelPos = clamp(PixelPos, tsr_short2(HistoryInfo_ViewportMin), tsr_short2(HistoryInfo_ViewportMax - 1));
+
+			Colors.SetElement(SimdIndex, UpdateHistoryOutputTexture[PixelPos]);
+		}
+	}
+
+	tsr_output_halfC Output;
+	{
+		#if CONFIG_CLAMP
+			tsr_output_halfC Min = DownsampleMin2x2(Colors);
+			tsr_output_halfC Max = DownsampleMax2x2(Colors);
+
+			Min = min(Min, AccessNeighborTexel(Min, tsr_short2(1, 0)));
+			Min = min(Min, AccessNeighborTexel(Min, tsr_short2(0, 1)));
+
+			Max = max(Max, AccessNeighborTexel(Max, tsr_short2(1, 0)));
+			Max = max(Max, AccessNeighborTexel(Max, tsr_short2(0, 1)));
+		#endif
+		
+		tsr_input_halfC LDRColors = Colors;
+		#if CONFIG_HDR_WEIGHT
+		{
+			UNROLL_N(INPUT_SIMD_SIZE)
+			for (uint SimdIndex = 0; SimdIndex < INPUT_SIMD_SIZE; SimdIndex++)
+			{
+				LDRColors.SetElement(SimdIndex, Colors.GetElement(SimdIndex) * HdrWeight4(Colors.GetElement(SimdIndex)));
+			}
+		}
+		#endif
+
+		/*
+		 *  ^
+		 *  |  b  c
+		 *  |  a  b
+		 *  o------>
+		 */
+		const float WeightA = MitchellNetravali(/* PixelOffset = */ 0.5, MitchellNetravaliB, MitchellNetravaliC);
+		const float WeightB = MitchellNetravali(/* PixelOffset = */ 1.5, MitchellNetravaliB, MitchellNetravaliC);
+		const float TotalWeight = 4.0 * (WeightA * WeightA + WeightB * WeightB + WeightA * WeightB * 2.0);
+
+		const tsr_half FinalWeightA = tsr_half(WeightA * WeightA / TotalWeight);
+		const tsr_half FinalWeightB = tsr_half(WeightA * WeightB / TotalWeight);
+		const tsr_half FinalWeightC = tsr_half(WeightB * WeightB / TotalWeight);
+	
+		/*
+		 *  c   b | b   c
+		 *  b   a | a   b
+		 *  ------o------
+		 *  b   a | a   b
+		 *  c   b | b   c
+		 */
+		const tsr_half Weights0[4] = { FinalWeightC, FinalWeightB, FinalWeightB, FinalWeightA };
+		const tsr_half Weights1[4] = { FinalWeightB, FinalWeightC, FinalWeightA, FinalWeightB };
+		const tsr_half Weights2[4] = { FinalWeightB, FinalWeightA, FinalWeightC, FinalWeightB };
+		const tsr_half Weights3[4] = { FinalWeightA, FinalWeightB, FinalWeightB, FinalWeightC };
+
+		tsr_output_halfC Out0 = DownsampleDot2x2(LDRColors, Weights0);
+		tsr_output_halfC Out1 = DownsampleDot2x2(LDRColors, Weights1);
+		tsr_output_halfC Out2 = DownsampleDot2x2(LDRColors, Weights2);
+		tsr_output_halfC Out3 = DownsampleDot2x2(LDRColors, Weights3);
+
+		Output = Out0 + AccessNeighborTexel(Out1, tsr_short2(1, 0)) + AccessNeighborTexel(Out2, tsr_short2(0, 1)) + AccessNeighborTexel(Out3, tsr_short2(1, 1));
+		#if CONFIG_HDR_WEIGHT
+		{
+			UNROLL_N(OUTPUT_SIMD_SIZE)
+			for (uint SimdIndex = 0; SimdIndex < OUTPUT_SIMD_SIZE; SimdIndex++)
+			{
+				Output.SetElement(SimdIndex, Output.GetElement(SimdIndex) * HdrWeightInvY(Luma4(Output.GetElement(SimdIndex))));
+			}
+		}
+		#endif
+		#if CONFIG_CLAMP
+			Output = clamp(Output, Min, Max);
+		#endif
+	}
+
+	ISOLATE
+	{
+		const tsr_short2 GroupPixelOffset = (
+			tsr_short2(OutputViewRectMin) +
+			tsr_short2(GroupId) * tsr_short(TILE_SIZE).xx);
+		
+		UNROLL_N(OUTPUT_SIMD_SIZE)
+		for (uint SimdIndex = 0; SimdIndex < OUTPUT_SIMD_SIZE; SimdIndex++)
+		{
+			tsr_short2 OutputPixelOffset = GetLaneSimdPixelOffset<LANE_OUTPUT_STRIDE_X, LANE_OUTPUT_STRIDE_Y>(LaneIndex, SimdIndex);
+			tsr_short2 OutputPixelPos = InvalidateOutputPixelPos(GroupPixelOffset + OutputPixelOffset, OutputViewRectMax);
+			
+			OutputPixelPos.x = select(all(OutputPixelOffset < TILE_SIZE), OutputPixelPos.x, ~tsr_short(0));
+
+			SceneColorOutputMip0[OutputPixelPos] = Output.GetElement(SimdIndex);
+		}
+		
+		{
+			tsr_output_halfC HalfResOutput = Output * tsr_half(0.25);
+			HalfResOutput = HalfResOutput + AccessNeighborTexel(HalfResOutput, tsr_short2(1, 0));
+			HalfResOutput = HalfResOutput + AccessNeighborTexel(HalfResOutput, tsr_short2(0, 1));
+		
+			tsr_short2 OutputPixelOffset = GetLaneSimdPixelOffset<LANE_OUTPUT_STRIDE_X, LANE_OUTPUT_STRIDE_Y>(LaneIndex, /* SimdIndex = */ 0);
+			tsr_short2 OutputPixelPos = InvalidateOutputPixelPos(GroupPixelOffset + OutputPixelOffset, OutputViewRectMax);
+			OutputPixelPos.x = select(all(OutputPixelOffset < TILE_SIZE), OutputPixelPos.x, ~tsr_short(0));
+
+			tsr_short2 HalfResOutputPixelPos = OutputPixelPos >> tsr_short(1);
+			HalfResOutputPixelPos[0] |= ((OutputPixelPos[0] & tsr_short(0x1)) | tsr_short(!bGenerateOutputMip1)) * tsr_short(~0);
+			HalfResOutputPixelPos[1] |= ((OutputPixelPos[1] & tsr_short(0x1)) | tsr_short(!bGenerateOutputMip1)) * tsr_short(~0);
+			
+			SceneColorOutputMip1[HalfResOutputPixelPos] = HalfResOutput.GetElement(/* SimdIndex = */ 0);
+		}
+	}
+}
+
+#else // !DIM_NYQUIST
+
+//------------------------------------------------------- ENTRY POINT REFERENCE
 
 [numthreads(TILE_SIZE * TILE_SIZE, 1, 1)]
 void MainCS(
@@ -109,15 +287,14 @@ void MainCS(
 	
 	// Downsample the history
 	{
-		const float B = rcp(3.0);
-		const float C = rcp(3.0);
-
 		// Position of the output pixel in the history.
 		float2 OutputPixelPos = ApplyScreenTransform(float2(DispatchThreadId), DispatchThreadToHistoryPixelPos);
 		float2 TopLeftKernel = floor(OutputPixelPos + (0.5 - 0.5 * KERNEL_SIZE)) + 0.5;
 		
-		tsr_halfC MinColor = tsr_half(POSITIVE_INFINITY);
-		tsr_halfC MaxColor = tsr_half(0.0);
+		#if CONFIG_CLAMP
+			tsr_halfC MinColor = tsr_half(POSITIVE_INFINITY);
+			tsr_halfC MaxColor = tsr_half(0.0);
+		#endif
 		tsr_half ColorWeight = tsr_half(0.0);
 
 		UNROLL_N(KERNEL_SIZE)
@@ -132,34 +309,26 @@ void MainCS(
 
 				SampleUV = clamp(SampleUV, HistoryInfo_UVViewportBilinearMin, HistoryInfo_UVViewportBilinearMax);
 
-				tsr_halfC SampleColor    = History_HighFrequency.SampleLevel(GlobalBilinearClampedSampler, SampleUV, 0);
-				tsr_half3 SampleTranslucencyColor = History_Translucency.SampleLevel(GlobalBilinearClampedSampler, SampleUV, 0);
-				tsr_half  SampleTranslucencyAlpha = History_TranslucencyAlpha.SampleLevel(GlobalBilinearClampedSampler, SampleUV, 0);
-				
-				#if DIM_SEPARATE_TRANSLUCENCY
-				{
-					SampleColor.rgb = SampleColor.rgb * SampleTranslucencyAlpha + SampleTranslucencyColor.rgb;
+				tsr_halfC SampleColor = UpdateHistoryOutputTexture.SampleLevel(GlobalBilinearClampedSampler, SampleUV, 0);
 				
-					#if CONFIG_SCENE_COLOR_ALPHA
-						SampleColor.a = saturate(tsr_half(1.0) - saturate(tsr_half(1.0) - SampleColor.a) * SampleTranslucencyAlpha);
-					#endif
-				}
-				#endif
-
 				float2 PixelOffset = abs(SamplePixelPos - OutputPixelPos);
 
-				tsr_half kernelWeight = tsr_half(MitchellNetravali(PixelOffset.x, B, C) * MitchellNetravali(PixelOffset.y, B, C));
+				tsr_half kernelWeight = tsr_half(MitchellNetravali(PixelOffset.x, MitchellNetravaliB, MitchellNetravaliC) * MitchellNetravali(PixelOffset.y, MitchellNetravaliB, MitchellNetravaliC));
 				tsr_half SampleHdrWeight = HdrWeight4(SampleColor);
-
-				MinColor = min(MinColor, SampleColor);
-				MaxColor = max(MaxColor, SampleColor);
+				
+				#if CONFIG_CLAMP
+					MinColor = min(MinColor, SampleColor);
+					MaxColor = max(MaxColor, SampleColor);
+				#endif
 				Color  += (kernelWeight * SampleHdrWeight) * SampleColor;
 				ColorWeight += (kernelWeight * SampleHdrWeight);
 			}
 		}
 
 		Color *= SafeRcp(ColorWeight);
-		Color = clamp(Color, MinColor, MaxColor);
+		#if CONFIG_CLAMP
+			Color = clamp(Color, MinColor, MaxColor);
+		#endif
 	}
 	
 	// Compute final output
@@ -218,3 +387,4 @@ void MainCS(
 	#endif
 }
 
+#endif // !DIM_NYQUIST
diff --git a/Engine/Source/Runtime/Renderer/Private/PostProcess/TemporalSuperResolution.cpp b/Engine/Source/Runtime/Renderer/Private/PostProcess/TemporalSuperResolution.cpp
--- a/Engine/Source/Runtime/Renderer/Private/PostProcess/TemporalSuperResolution.cpp
+++ b/Engine/Source/Runtime/Renderer/Private/PostProcess/TemporalSuperResolution.cpp
@@ -858,7 +858,9 @@ class FTSRResolveHistoryCS : public FTSRShader
 	DECLARE_GLOBAL_SHADER(FTSRResolveHistoryCS);
 	SHADER_USE_PARAMETER_STRUCT(FTSRResolveHistoryCS, FTSRShader);
 
-	using FPermutationDomain = TShaderPermutationDomain<FTSRUpdateHistoryCS::FSeparateTranslucencyDim>;
+	class FNyquistDim : SHADER_PERMUTATION_BOOL("DIM_NYQUIST");
+
+	using FPermutationDomain = TShaderPermutationDomain<FNyquistDim, FTSRShader::F16BitVALUDim>;
 
 	BEGIN_SHADER_PARAMETER_STRUCT(FParameters, )
 		SHADER_PARAMETER_STRUCT_INCLUDE(FTSRCommonParameters, CommonParameters)
@@ -868,17 +870,63 @@ class FTSRResolveHistoryCS : public FTSRShader
 		SHADER_PARAMETER(int32, bGenerateOutputMip1)
 		SHADER_PARAMETER(float, HistoryValidityMultiply)
 
-		SHADER_PARAMETER_STRUCT(FTSRHistorySRVs, History)
+		SHADER_PARAMETER_RDG_TEXTURE(Texture2D, UpdateHistoryOutputTexture)
 
 		SHADER_PARAMETER_RDG_TEXTURE_UAV(RWTexture2D, SceneColorOutputMip0)
 		SHADER_PARAMETER_RDG_TEXTURE_UAV(RWTexture2D, SceneColorOutputMip1)
 		SHADER_PARAMETER_RDG_TEXTURE_UAV(RWTexture2DArray, DebugOutput)
 	END_SHADER_PARAMETER_STRUCT()
 
+	static FPermutationDomain RemapPermutation(FPermutationDomain PermutationVector)
+	{
+		if (!PermutationVector.Get<FNyquistDim>())
+		{
+			PermutationVector.Set<FTSRShader::F16BitVALUDim>(false);
+		}
+
+		return PermutationVector;
+	}
+
+	static bool ShouldCompilePermutation(const FGlobalShaderPermutationParameters& Parameters)
+	{
+		FPermutationDomain PermutationVector(Parameters.PermutationId);
+
+		if (PermutationVector != RemapPermutation(PermutationVector))
+		{
+			return false;
+		}
+
+		if (PermutationVector.Get<FNyquistDim>())
+		{
+			if (FTSRShader::SupportsWaveOps(Parameters.Platform) != ERHIFeatureSupport::RuntimeGuaranteed)
+			{
+				return false;
+			}
+		}
+
+		if (!FTSRShader::ShouldCompile32or16BitPermutation(Parameters.Platform, PermutationVector.Get<FTSRShader::F16BitVALUDim>()))
+		{
+			return false;
+		}
+
+		return FTSRShader::ShouldCompilePermutation(Parameters);
+	}
+
 	static void ModifyCompilationEnvironment(const FGlobalShaderPermutationParameters& Parameters, FShaderCompilerEnvironment& OutEnvironment)
 	{
+		FPermutationDomain PermutationVector(Parameters.PermutationId);
+
 		FTSRShader::ModifyCompilationEnvironment(Parameters, OutEnvironment);
-		OutEnvironment.CompilerFlags.Add(CFLAG_ForceOptimization);
+
+		if (PermutationVector.Get<FNyquistDim>() != 0)
+		{
+			OutEnvironment.CompilerFlags.Add(CFLAG_WaveOperations);
+		}
+
+		if (PermutationVector.Get<FTSRShader::F16BitVALUDim>())
+		{
+			OutEnvironment.CompilerFlags.Add(CFLAG_AllowRealTypes);
+		}
 	}
 }; // class FTSRResolveHistoryCS
 
@@ -1218,13 +1266,13 @@ ITemporalUpscaler::FOutputs AddTemporalSuperResolutionPasses(
 		if (PassInputs.bGenerateSceneColorQuarterRes && !PassInputs.bGenerateSceneColorHalfRes && OutputRect.Size() == HistorySize)
 		{
 			FRDGTextureDesc QuarterResDesc = Desc;
-			QuarterResDesc.Extent /= 4;
+			QuarterResDesc.Extent = OutputExtent / 4;
 			SceneColorOutputQuarterResTexture = GraphBuilder.CreateTexture(QuarterResDesc, TEXT("TSR.QuarterResOutput"));
 		}
 		else if (PassInputs.bGenerateSceneColorHalfRes || PassInputs.bGenerateSceneColorQuarterRes)
 		{
 			FRDGTextureDesc HalfResDesc = Desc;
-			HalfResDesc.Extent /= 2;
+			HalfResDesc.Extent = OutputExtent / 2;
 			SceneColorOutputHalfResTexture = GraphBuilder.CreateTexture(HalfResDesc, TEXT("TSR.HalfResOutput"));
 		}
 	}
@@ -1981,19 +2029,12 @@ ITemporalUpscaler::FOutputs AddTemporalSuperResolutionPasses(
 		PassParameters->GrandPrevColorTexture = GrandPrevColorTexture;
 
 		PassParameters->HistoryOutput = CreateUAVs(GraphBuilder, History);
-		if (HistorySize != OutputRect.Size() && bIsOutputDifferentThanHighFrequency)
-		{
-			PassParameters->SceneColorOutputMip0 = CreateDummyUAV(GraphBuilder, PF_FloatR11G11B10);
-		}
-		else
-		{
-			PassParameters->SceneColorOutputMip0 = GraphBuilder.CreateUAV(FRDGTextureUAVDesc(UpdateHistoryOutputTexture, /* InMipLevel = */ 0));
-		}
+		PassParameters->SceneColorOutputMip0 = GraphBuilder.CreateUAV(FRDGTextureUAVDesc(UpdateHistoryOutputTexture, /* InMipLevel = */ 0));
 		
 		if (PassInputs.bGenerateOutputMip1 && HistorySize == OutputRect.Size())
 		{
 			PassParameters->bGenerateOutputMip1 = true;
-			PassParameters->SceneColorOutputMip1 = GraphBuilder.CreateUAV(FRDGTextureUAVDesc(SceneColorOutputTexture, /* InMipLevel = */ 1));
+			PassParameters->SceneColorOutputMip1 = GraphBuilder.CreateUAV(FRDGTextureUAVDesc(UpdateHistoryOutputTexture, /* InMipLevel = */ 1));
 		}
 		else if (SceneColorOutputHalfResTexture && HistorySize == OutputRect.Size())
 		{
@@ -2061,6 +2102,9 @@ ITemporalUpscaler::FOutputs AddTemporalSuperResolutionPasses(
 	if (HistorySize != OutputRect.Size())
 	{
 		check(!SceneColorOutputQuarterResTexture);
+
+		bool bNyquistHistory = HistorySize.X == 2 * OutputRect.Width() && HistorySize.Y == 2 * OutputRect.Height();
+
 		FTSRResolveHistoryCS::FParameters* PassParameters = GraphBuilder.AllocParameters<FTSRResolveHistoryCS::FParameters>();
 		PassParameters->CommonParameters = CommonParameters;
 		PassParameters->DispatchThreadToHistoryPixelPos = (
@@ -2073,7 +2117,7 @@ ITemporalUpscaler::FOutputs AddTemporalSuperResolutionPasses(
 		PassParameters->bGenerateOutputMip1 = (PassInputs.bGenerateOutputMip1 || PassInputs.bGenerateSceneColorHalfRes) ? 1 : 0;
 		PassParameters->HistoryValidityMultiply = float(HistorySize.X * HistorySize.Y) / float(OutputRect.Width() * OutputRect.Height());
 
-		PassParameters->History = CreateSRVs(GraphBuilder, History);
+		PassParameters->UpdateHistoryOutputTexture = UpdateHistoryOutputTexture;
 		
 		PassParameters->SceneColorOutputMip0 = GraphBuilder.CreateUAV(FRDGTextureUAVDesc(SceneColorOutputTexture, /* InMipLevel = */ 0));
 		if (PassInputs.bGenerateOutputMip1)
@@ -2091,16 +2135,21 @@ ITemporalUpscaler::FOutputs AddTemporalSuperResolutionPasses(
 		PassParameters->DebugOutput = CreateDebugUAV(HistoryExtent, TEXT("Debug.TSR.ResolveHistory"));
 
 		FTSRResolveHistoryCS::FPermutationDomain PermutationVector;
-		PermutationVector.Set<FTSRUpdateHistoryCS::FSeparateTranslucencyDim>(bAccumulateTranslucencySeparately);
+		PermutationVector.Set<FTSRResolveHistoryCS::FNyquistDim>(bNyquistHistory && bUseWaveOps && GRHIMaximumWaveSize >= 32 && GRHIMinimumWaveSize <= 32);
+		PermutationVector.Set<FTSRShader::F16BitVALUDim>(bUse16BitVALU);
+		PermutationVector = FTSRResolveHistoryCS::RemapPermutation(PermutationVector);
 
 		TShaderMapRef<FTSRResolveHistoryCS> ComputeShader(View.ShaderMap, PermutationVector);
 		FComputeShaderUtils::AddPass(
 			GraphBuilder,
-			RDG_EVENT_NAME("TSR ResolveHistory %dx%d", OutputRect.Width(), OutputRect.Height()),
+			RDG_EVENT_NAME("TSR ResolveHistory(%s%s) %dx%d", 
+				PermutationVector.Get<FTSRResolveHistoryCS::FNyquistDim>() ? TEXT("WaveOps") : TEXT(""),
+				PermutationVector.Get<FTSRShader::F16BitVALUDim>() ? TEXT(" 16bit") : TEXT(""),
+				OutputRect.Width(), OutputRect.Height()),
 			AsyncComputePasses >= 3 ? ERDGPassFlags::AsyncCompute : ERDGPassFlags::Compute,
 			ComputeShader,
 			PassParameters,
-			FComputeShaderUtils::GetGroupCount(OutputRect.Size(), 8));
+			FComputeShaderUtils::GetGroupCount(OutputRect.Size(), PermutationVector.Get<FTSRResolveHistoryCS::FNyquistDim>() ? 6 : 8));
 	}
 
 	// Extract all resources for next frame.
