Path Tracer: Fix fireflies that can occur from the rough diffuse shading model at grazing angles

The model is mostly energy conserving except at grazing angles where the weights can becomes slightly above 1.0. In the path tracing context this can become problematic in corners of models or when the shading normals diverge too much from the base surface (as can happen with some Nanite fallback meshes).

Clamping the weight to 1 fixes these issues without altering the overall shape of the BSDF

diff --git a/Engine/Shaders/Private/PathTracing/Material/PathTracingMaterialCommon.ush b/Engine/Shaders/Private/PathTracing/Material/PathTracingMaterialCommon.ush
--- a/Engine/Shaders/Private/PathTracing/Material/PathTracingMaterialCommon.ush
+++ b/Engine/Shaders/Private/PathTracing/Material/PathTracingMaterialCommon.ush
@@ -212,7 +212,8 @@ float ShadowTerminatorTerm(float3 L, float3 N, float3 Ns)
 float3 GetPathTracingDiffuseModel(float3 DiffuseColor, float Roughness, float NoV, float NoL, float VoH, float NoH)
 {
 #if MATERIAL_ROUGHDIFFUSE
-	return PI * Diffuse_Chan(DiffuseColor, Pow4(Roughness), NoV, NoL, VoH, NoH, 1.0f);
+	// Clamp to [0,1] as the model is ever so slightly non-energy conserving at grazing angles which can cause unwanted noise
+	return saturate(PI * Diffuse_Chan(DiffuseColor, Pow4(Roughness), NoV, NoL, VoH, NoH, 1.0f));
 #else
 	return DiffuseColor;
 #endif
