Change hair interpolation buffer to use byte address buffer to ease adaptive format with streaming.

diff --git a/Engine/Plugins/Runtime/HairStrands/Source/HairStrandsCore/Private/GroomBuilder.cpp b/Engine/Plugins/Runtime/HairStrands/Source/HairStrandsCore/Private/GroomBuilder.cpp
--- a/Engine/Plugins/Runtime/HairStrands/Source/HairStrandsCore/Private/GroomBuilder.cpp
+++ b/Engine/Plugins/Runtime/HairStrands/Source/HairStrandsCore/Private/GroomBuilder.cpp
@@ -40,7 +40,7 @@ static FAutoConsoleVariableRef CVarHairGroupIndexBuilder_MaxVoxelResolution(TEXT
 
 FString FGroomBuilder::GetVersion()
 {
-	return TEXT("v8p3");
+	return TEXT("v8q");
 }
 
 namespace FHairStrandsDecimation
@@ -1338,65 +1338,90 @@ namespace HairInterpolationBuilder
 		auto LowerPart = [](uint32 Index) -> uint16 { return uint16(Index & 0xFFFF); };
 		auto UpperPart = [](uint32 Index) -> uint8  { return uint8((Index >> 16) & 0xFF); };
 
-
 		OutBulkData.Header.Flags = FHairStrandsInterpolationBulkData::DataFlags_HasData;
 		OutBulkData.Header.PointCount = PointCount;
 
 		if (HairInterpolation.bUseUniqueGuide)
 		{
+			struct FHairInterpolation1GuideVertex
+			{
+				// Guide's vertex index are stored onto 24bits: 1) VertexGuideIndex0 is the lower part, 2) VertexGuideIndex1 is the upper part
+				uint16 VertexGuideIndex0;
+				uint8  VertexGuideIndex1;
+				uint8  VertexLerp;
+			};
+			static_assert(sizeof(FHairInterpolation1GuideVertex) == HAIR_INTERPOLATION_1GUIDE_STRIDE);
+			static_assert(sizeof(FHairStrandsInterpolationFormat::Type) == 4u);
+
 			OutBulkData.Header.Flags |= FHairStrandsInterpolationBulkData::DataFlags_HasSingleGuideData;
 
 			TArray<FHairStrandsInterpolationFormat::Type> OutPointsInterpolation;
-			OutPointsInterpolation.SetNum(PointCount * FHairStrandsInterpolationFormat::ComponentCount);
-
+			OutPointsInterpolation.SetNum(PointCount);
+			FHairInterpolation1GuideVertex* Out = (FHairInterpolation1GuideVertex*)OutPointsInterpolation.GetData();
 			for (uint32 PointIndex = 0; PointIndex < PointCount; ++PointIndex)
 			{
-				const FIntVector& Indices = HairInterpolation.PointsSimCurvesVertexIndex[PointIndex];
-				const FVector& Weights = (FVector)HairInterpolation.PointsSimCurvesVertexWeights[PointIndex];
-				const FVector& S = (FVector)HairInterpolation.PointsSimCurvesVertexLerp[PointIndex];
+				const FIntVector& Indices 	= HairInterpolation.PointsSimCurvesVertexIndex[PointIndex];
+				const FVector3f& Weights 	= (FVector3f)HairInterpolation.PointsSimCurvesVertexWeights[PointIndex];
+				const FVector3f& S 			= (FVector3f)HairInterpolation.PointsSimCurvesVertexLerp[PointIndex];
 
-				FHairStrandsInterpolationFormat::Type& OutInterp = OutPointsInterpolation[PointIndex];
+				FHairInterpolation1GuideVertex& OutInterp = Out[PointIndex];
 				OutInterp.VertexGuideIndex0 = LowerPart(Indices[0]);
 				OutInterp.VertexGuideIndex1 = UpperPart(Indices[0]);
-				OutInterp.VertexLerp		= S[0] * 255.f;
+				OutInterp.VertexLerp		 = FMath::Clamp(S[0] * 255.f, 0, 255);
 			}
 
 			HairStrandsBuilder::CopyToBulkData<FHairStrandsInterpolationFormat>(OutBulkData.Data.Interpolation, OutPointsInterpolation);
 		}
 		else
 		{
-			TArray<FHairStrandsInterpolation0Format::Type> OutPointsInterpolation0;
-			TArray<FHairStrandsInterpolation1Format::Type> OutPointsInterpolation1;
+			struct FHairInterpolation3GuidesVertex
+			{	
+				typedef FUintVector4 BulkType;
+
+				// 128 bits
+				uint16 Index0;       // X:16
+				uint16 Index1;       // X:16
+				uint16 Index2;       // Y:16
+				uint8  VertexWeight0;// Y:8
+				uint8  VertexWeight1;// Y:8
+				uint8  VertexIndex0; // Z:8
+				uint8  VertexIndex1; // Z:8
+				uint8  VertexIndex2; // Z:8
+				uint8  Pad0;         // Z:8
+				uint8  VertexLerp0;  // W:8
+				uint8  VertexLerp1;  // W:8
+				uint8  VertexLerp2;  // W:8
+				uint8  Pad1;         // W:8
+			};
+			static_assert(sizeof(FHairInterpolation3GuidesVertex) == HAIR_INTERPOLATION_3GUIDE_STRIDE);
+			static_assert(sizeof(FHairStrandsInterpolationFormat::Type) == 4u);
 
-			OutPointsInterpolation0.SetNum(PointCount * FHairStrandsInterpolation0Format::ComponentCount);
-			OutPointsInterpolation1.SetNum(PointCount * FHairStrandsInterpolation1Format::ComponentCount);
+			TArray<FHairStrandsInterpolationFormat::Type> OutPointsInterpolation;
+			OutPointsInterpolation.SetNum(PointCount * 4);
+			FHairInterpolation3GuidesVertex* Out = (FHairInterpolation3GuidesVertex*)OutPointsInterpolation.GetData();
 
 			for (uint32 PointIndex = 0; PointIndex < PointCount; ++PointIndex)
 			{
-				const FIntVector& Indices = HairInterpolation.PointsSimCurvesVertexIndex[PointIndex];
-				const FVector& Weights = (FVector)HairInterpolation.PointsSimCurvesVertexWeights[PointIndex];
-				const FVector& S = (FVector)HairInterpolation.PointsSimCurvesVertexLerp[PointIndex];
-
-				FHairStrandsInterpolation0Format::Type& OutInterp0 = OutPointsInterpolation0[PointIndex];
-				OutInterp0.Index0 = LowerPart(Indices[0]);
-				OutInterp0.Index1 = LowerPart(Indices[1]);
-				OutInterp0.Index2 = LowerPart(Indices[2]);
-				OutInterp0.VertexWeight0 = Weights[0] * 255.f;
-				OutInterp0.VertexWeight1 = Weights[1] * 255.f;
-
-				FHairStrandsInterpolation1Format::Type& OutInterp1 = OutPointsInterpolation1[PointIndex];
-				OutInterp1.VertexIndex0 = UpperPart(Indices[0]);
-				OutInterp1.VertexIndex1 = UpperPart(Indices[1]);
-				OutInterp1.VertexIndex2 = UpperPart(Indices[2]);
-				OutInterp1.VertexLerp0  = S[0] * 255.f;
-				OutInterp1.VertexLerp1  = S[1] * 255.f;
-				OutInterp1.VertexLerp2  = S[2] * 255.f;
-				OutInterp1.Pad0			= 0;
-				OutInterp1.Pad1			= 0;
-			}
-
-			HairStrandsBuilder::CopyToBulkData<FHairStrandsInterpolation0Format>(OutBulkData.Data.Interpolation0, OutPointsInterpolation0);
-			HairStrandsBuilder::CopyToBulkData<FHairStrandsInterpolation1Format>(OutBulkData.Data.Interpolation1, OutPointsInterpolation1);
+				const FIntVector& Indices 	= HairInterpolation.PointsSimCurvesVertexIndex[PointIndex];
+				const FVector3f& Weights 	= (FVector3f)HairInterpolation.PointsSimCurvesVertexWeights[PointIndex];
+				const FVector3f& S 			= (FVector3f)HairInterpolation.PointsSimCurvesVertexLerp[PointIndex];
+
+				FHairInterpolation3GuidesVertex& OutInterp = Out[PointIndex];
+				OutInterp.Index0 		= LowerPart(Indices[0]);
+				OutInterp.Index1 		= LowerPart(Indices[1]);
+				OutInterp.Index2 		= LowerPart(Indices[2]);
+				OutInterp.VertexWeight0 = FMath::Clamp(Weights[0] * 255.f, 0, 255);
+				OutInterp.VertexWeight1 = FMath::Clamp(Weights[1] * 255.f, 0, 255);
+				OutInterp.VertexIndex0  = UpperPart(Indices[0]);
+				OutInterp.VertexIndex1  = UpperPart(Indices[1]);
+				OutInterp.VertexIndex2  = UpperPart(Indices[2]);
+				OutInterp.Pad0			= 0;
+				OutInterp.VertexLerp0   = FMath::Clamp(S[0] * 255.f, 0, 255);
+				OutInterp.VertexLerp1   = FMath::Clamp(S[1] * 255.f, 0, 255);
+				OutInterp.VertexLerp2   = FMath::Clamp(S[2] * 255.f, 0, 255);
+				OutInterp.Pad1			= 0;
+			}
+			HairStrandsBuilder::CopyToBulkData<FHairStrandsInterpolationFormat>(OutBulkData.Data.Interpolation, OutPointsInterpolation);
 		}
 
 		{
diff --git a/Engine/Plugins/Runtime/HairStrands/Source/HairStrandsCore/Private/GroomResources.cpp b/Engine/Plugins/Runtime/HairStrands/Source/HairStrandsCore/Private/GroomResources.cpp
--- a/Engine/Plugins/Runtime/HairStrands/Source/HairStrandsCore/Private/GroomResources.cpp
+++ b/Engine/Plugins/Runtime/HairStrands/Source/HairStrandsCore/Private/GroomResources.cpp
@@ -1686,7 +1686,7 @@ void FHairStrandsRootData::Reset()
 
 FHairStrandsInterpolationResource::FHairStrandsInterpolationResource(FHairStrandsInterpolationBulkData& InBulkData, const FHairResourceName& InResourceName, const FName& InOwnerName) :
 	FHairCommonResource(EHairStrandsAllocationType::Deferred, InResourceName, InOwnerName),
-	InterpolationBuffer(), Interpolation0Buffer(), Interpolation1Buffer(), BulkData(InBulkData)
+	InterpolationBuffer(), BulkData(InBulkData)
 {
 	// Sanity check
 	check(!!(BulkData.Header.Flags & FHairStrandsInterpolationBulkData::DataFlags_HasData));
@@ -1704,25 +1704,13 @@ bool FHairStrandsInterpolationResource::InternalIsDataLoaded()
 void FHairStrandsInterpolationResource::InternalAllocate(FRDGBuilder& GraphBuilder)
 {
 	BulkDataRequest = FHairResourceRequest();
-
-	const bool bUseSingleGuide = !!(BulkData.Header.Flags & FHairStrandsInterpolationBulkData::DataFlags_HasSingleGuideData);
-	if (bUseSingleGuide)
-	{
-		InternalCreateVertexBufferRDG_FromBulkData<FHairStrandsInterpolationFormat>(GraphBuilder, BulkData.Data.Interpolation, BulkData.Header.PointCount, InterpolationBuffer, ToHairResourceDebugName(TEXT("Hair.StrandsInterpolation_InterpolationBuffer"), ResourceName), OwnerName, EHairResourceUsageType::Static);
-	}
-	else
-	{
-		InternalCreateVertexBufferRDG_FromBulkData<FHairStrandsInterpolation0Format>(GraphBuilder, BulkData.Data.Interpolation0, BulkData.Header.PointCount, Interpolation0Buffer, ToHairResourceDebugName(TEXT("Hair.StrandsInterpolation_Interpolation0Buffer"), ResourceName), OwnerName, EHairResourceUsageType::Static);
-		InternalCreateVertexBufferRDG_FromBulkData<FHairStrandsInterpolation1Format>(GraphBuilder, BulkData.Data.Interpolation1, BulkData.Header.PointCount, Interpolation1Buffer, ToHairResourceDebugName(TEXT("Hair.StrandsInterpolation_Interpolation1Buffer"), ResourceName), OwnerName, EHairResourceUsageType::Static);
-	}
+	InternalCreateByteAddressBufferRDG_FromBulkData(GraphBuilder, BulkData.Data.Interpolation, InterpolationBuffer, ToHairResourceDebugName(TEXT("Hair.StrandsInterpolation_InterpolationBuffer"), ResourceName), OwnerName, EHairResourceUsageType::Static);
 	InternalCreateVertexBufferRDG_FromBulkData<FHairStrandsRootIndexFormat>(GraphBuilder, BulkData.Data.SimRootPointIndex, BulkData.Header.SimPointCount, SimRootPointIndexBuffer, ToHairResourceDebugName(TEXT("Hair.StrandsInterpolation_SimRootPointIndex"), ResourceName), OwnerName, EHairResourceUsageType::Static);
 }
 
 void FHairStrandsInterpolationResource::InternalRelease()
 {
 	InterpolationBuffer.Release();
-	Interpolation0Buffer.Release();
-	Interpolation1Buffer.Release();
 	SimRootPointIndexBuffer.Release();
 	BulkDataRequest = FHairResourceRequest();
 }
diff --git a/Engine/Plugins/Runtime/HairStrands/Source/HairStrandsCore/Private/HairStrandsDatas.cpp b/Engine/Plugins/Runtime/HairStrands/Source/HairStrandsCore/Private/HairStrandsDatas.cpp
--- a/Engine/Plugins/Runtime/HairStrands/Source/HairStrandsCore/Private/HairStrandsDatas.cpp
+++ b/Engine/Plugins/Runtime/HairStrands/Source/HairStrandsCore/Private/HairStrandsDatas.cpp
@@ -111,63 +111,6 @@ FArchive& operator<<(FArchive& Ar, FPackedHairVertex& Vertex)
 	return Ar;
 }
 
-FArchive& operator<<(FArchive& Ar, FHairInterpolation0Vertex& Vertex)
-{
-	Ar << Vertex.Index0;
-	Ar << Vertex.Index1;
-	Ar << Vertex.Index2;
-	Ar << Vertex.VertexWeight0;
-	Ar << Vertex.VertexWeight1;
-
-	return Ar;
-}
-
-FArchive& operator<<(FArchive& Ar, FHairInterpolationVertex& Vertex)
-{
-	Ar << Vertex.VertexGuideIndex0;
-	Ar << Vertex.VertexGuideIndex1;
-	Ar << Vertex.VertexLerp;
-
-	return Ar;
-}
-
-FArchive& operator<<(FArchive& Ar, FHairInterpolation1Vertex& Vertex)
-{
-	Ar.UsingCustomVersion(FReleaseObjectVersion::GUID);
-
-	if (Ar.CustomVer(FReleaseObjectVersion::GUID) >= FReleaseObjectVersion::GroomAssetVersion3)
-	{
-		Ar << Vertex.VertexIndex0;
-		Ar << Vertex.VertexIndex1;
-		Ar << Vertex.VertexIndex2;
-
-		Ar << Vertex.VertexLerp0;
-		Ar << Vertex.VertexLerp1;
-		Ar << Vertex.VertexLerp2;
-
-		Ar << Vertex.Pad0;
-		Ar << Vertex.Pad1;
-	}
-	else
-	{
-		Ar << Vertex.VertexIndex0;
-		Ar << Vertex.VertexIndex1;
-		Ar << Vertex.VertexIndex2;
-
-		uint8 Pad0 = 0;
-		Ar << Pad0;
-
-		if (Ar.IsLoading())
-		{
-			Vertex.VertexLerp0 = 0;
-			Vertex.VertexLerp1 = 0;
-			Vertex.VertexLerp2 = 0;
-		}
-	}
-
-	return Ar;
-}
-
 void FHairStrandsInterpolationBulkData::Reset()
 {
 	Header.Flags = 0;
@@ -176,14 +119,10 @@ void FHairStrandsInterpolationBulkData::Reset()
 	
 	// Deallocate memory if needed
 	Data.Interpolation.RemoveBulkData();
-	Data.Interpolation0.RemoveBulkData();
-	Data.Interpolation1.RemoveBulkData();
 	Data.SimRootPointIndex.RemoveBulkData();
 
 	// Reset the bulk byte buffer to ensure the (serialize) data size is reset to 0
 	Data.Interpolation		= FByteBulkData();
-	Data.Interpolation0		= FByteBulkData();
-	Data.Interpolation1		= FByteBulkData();
 	Data.SimRootPointIndex	= FByteBulkData();
 }
 
@@ -202,25 +141,13 @@ void FHairStrandsInterpolationBulkData::SerializeHeader(FArchive& Ar, UObject* O
 
 void FHairStrandsInterpolationBulkData::SerializeData(FArchive& Ar, UObject* Owner)
 {
-	static_assert(sizeof(FHairInterpolationVertex::BulkType) == sizeof(FHairInterpolationVertex));
-	static_assert(sizeof(FHairInterpolation0Vertex::BulkType) == sizeof(FHairInterpolation0Vertex));
-	static_assert(sizeof(FHairInterpolation1Vertex::BulkType) == sizeof(FHairInterpolation1Vertex));
 	static_assert(sizeof(FHairStrandsRootIndexFormat::BulkType) == sizeof(FHairStrandsRootIndexFormat::Type));
 
 	if (!!(Header.Flags & DataFlags_HasData))
 	{
 		const int32 ChunkIndex = 0;
 		bool bAttemptFileMapping = false;
-
-		if (!!(Header.Flags & DataFlags_HasSingleGuideData))
-		{
-			Data.Interpolation.Serialize(Ar, Owner, ChunkIndex, bAttemptFileMapping);
-		}
-		else
-		{
-			Data.Interpolation0.Serialize(Ar, Owner, ChunkIndex, bAttemptFileMapping);
-			Data.Interpolation1.Serialize(Ar, Owner, ChunkIndex, bAttemptFileMapping);
-		}
+		Data.Interpolation.Serialize(Ar, Owner, ChunkIndex, bAttemptFileMapping);
 		Data.SimRootPointIndex.Serialize(Ar, Owner, ChunkIndex, bAttemptFileMapping);
 	}
 }
@@ -238,16 +165,8 @@ void FHairStrandsInterpolationBulkData::Request(FBulkDataBatchRequest& InRequest
 	}
 
 	check(InRequest.IsNone());
-	FBulkDataBatchRequest::FBatchBuilder Batch = FBulkDataBatchRequest::NewBatch(3);
-	if (!!(Header.Flags & DataFlags_HasSingleGuideData))
-	{
-		Batch.Read(Data.Interpolation);
-	}
-	else
-	{
-		Batch.Read(Data.Interpolation0);
-		Batch.Read(Data.Interpolation1);
-	}
+	FBulkDataBatchRequest::FBatchBuilder Batch = FBulkDataBatchRequest::NewBatch(2);
+	Batch.Read(Data.Interpolation);
 	Batch.Read(Data.SimRootPointIndex);
 	Batch.Issue(InRequest);
 }
diff --git a/Engine/Plugins/Runtime/HairStrands/Source/HairStrandsCore/Private/HairStrandsRendering.cpp b/Engine/Plugins/Runtime/HairStrands/Source/HairStrandsCore/Private/HairStrandsRendering.cpp
--- a/Engine/Plugins/Runtime/HairStrands/Source/HairStrandsCore/Private/HairStrandsRendering.cpp
+++ b/Engine/Plugins/Runtime/HairStrands/Source/HairStrandsCore/Private/HairStrandsRendering.cpp
@@ -507,9 +507,7 @@ class FHairInterpolationCS : public FGlobalShader
 		SHADER_PARAMETER_RDG_BUFFER_SRV(Buffer, SimRestPosePositionBuffer)
 		SHADER_PARAMETER_RDG_BUFFER_SRV(Buffer, SimDeformedPositionBuffer)
 
-		SHADER_PARAMETER_RDG_BUFFER_SRV(Buffer, InterpolationBuffer)
-		SHADER_PARAMETER_RDG_BUFFER_SRV(Buffer, Interpolation0Buffer)
-		SHADER_PARAMETER_RDG_BUFFER_SRV(Buffer, Interpolation1Buffer)
+		SHADER_PARAMETER_RDG_BUFFER_SRV(ByteAddressBuffer, InterpolationBuffer)
 
 		SHADER_PARAMETER_RDG_BUFFER_SRV(Buffer<float4>, RenRestPosition0Buffer)
 		SHADER_PARAMETER_RDG_BUFFER_SRV(Buffer<float4>, RenRestPosition1Buffer)
@@ -591,8 +589,6 @@ static void AddHairStrandsInterpolationPass(
 	const FRDGBufferSRVRef& RenPointToCurveBuffer,
 	const bool bUseSingleGuide,
 	const FRDGBufferSRVRef& InterpolationBuffer,
-	const FRDGBufferSRVRef& Interpolation0Buffer,
-	const FRDGBufferSRVRef& Interpolation1Buffer,
 	const FRDGBufferSRVRef& SimRestPosePositionBuffer,
 	const FRDGBufferSRVRef& SimDeformedPositionBuffer,
 	const FRDGBufferSRVRef& SimRootPointIndexBuffer,
@@ -605,15 +601,7 @@ static void AddHairStrandsInterpolationPass(
 	Parameters->RenRestPosePositionBuffer = RenRestPosePositionBuffer;
 	Parameters->SimRestPosePositionBuffer = SimRestPosePositionBuffer;
 	Parameters->SimDeformedPositionBuffer = SimDeformedPositionBuffer;
-	if (bUseSingleGuide)
-	{
-		Parameters->InterpolationBuffer = InterpolationBuffer;
-	}
-	else
-	{
-		Parameters->Interpolation0Buffer = Interpolation0Buffer;
-		Parameters->Interpolation1Buffer = Interpolation1Buffer;
-	}
+	Parameters->InterpolationBuffer = InterpolationBuffer;
 	Parameters->OutRenDeformedPositionBuffer = OutRenPositionBuffer.UAV;
 
 	Parameters->VertexCount = VertexCount;
@@ -768,11 +756,7 @@ class FHairPatchAttributeCS : public FGlobalShader
 		SHADER_PARAMETER(uint32, CurveCount)
 		SHADER_PARAMETER_ARRAY(FUintVector4, RenCurveAttributeOffsets, [HAIR_CURVE_ATTRIBUTE_OFFSET_COUNT])
 		SHADER_PARAMETER_RDG_BUFFER_SRV(Buffer, RenCurveBuffer)
-
-		SHADER_PARAMETER_RDG_BUFFER_SRV(Buffer, InterpolationBuffer)
-		SHADER_PARAMETER_RDG_BUFFER_SRV(Buffer, Interpolation0Buffer)
-		SHADER_PARAMETER_RDG_BUFFER_SRV(Buffer, Interpolation1Buffer)
-
+		SHADER_PARAMETER_RDG_BUFFER_SRV(ByteAddressBuffer, InterpolationBuffer)
 		SHADER_PARAMETER_RDG_BUFFER_SRV(Buffer, RenVertexToClusterIdBuffer)
 		SHADER_PARAMETER_RDG_BUFFER_UAV(RWByteAddressBuffer, OutRenAttributeBuffer)
 	END_SHADER_PARAMETER_STRUCT()
@@ -807,8 +791,6 @@ static void AddPatchAttributePass(
 	const FRDGBufferSRVRef& RenCurveBuffer,
 	const FRDGBufferSRVRef& RenVertexToClusterIdBuffer,
 	const FRDGBufferSRVRef& InterpolationBuffer,
-	const FRDGBufferSRVRef& Interpolation0Buffer,
-	const FRDGBufferSRVRef& Interpolation1Buffer,
 	FRDGImportedBuffer& OutRenAttributeBuffer)
 {
 	// Sanity check
@@ -829,8 +811,6 @@ static void AddPatchAttributePass(
 	if (bSimulation)
 	{
 		Parameters->InterpolationBuffer  = InterpolationBuffer;
-		Parameters->Interpolation0Buffer = Interpolation0Buffer;
-		Parameters->Interpolation1Buffer = Interpolation1Buffer;
 	}
 
 	Parameters->OutRenAttributeBuffer = OutRenAttributeBuffer.UAV;
@@ -1021,7 +1001,7 @@ class FHairCardsDeformationCS : public FGlobalShader
 		SHADER_PARAMETER_RDG_BUFFER_SRV(Buffer, GuideDeformedPositionOffsetBuffer)
 		SHADER_PARAMETER_SRV(Buffer, CardsRestPositionBuffer)
 		SHADER_PARAMETER_SRV(Buffer, CardsRestTangentBuffer)
-		SHADER_PARAMETER_RDG_BUFFER_SRV(Buffer, CardsInterpolationBuffer)
+		SHADER_PARAMETER_RDG_BUFFER_SRV(ByteAddressBuffer, CardsInterpolationBuffer)
 
 		SHADER_PARAMETER_RDG_BUFFER_SRV(Buffer<float4>, TriangleRestPosition0Buffer)
 		SHADER_PARAMETER_RDG_BUFFER_SRV(Buffer<float4>, TriangleRestPosition1Buffer)
@@ -2006,9 +1986,7 @@ void ComputeHairStrandsInterpolation(
 						RegisterAsSRV(GraphBuilder, Instance->Strands.RestResource->PositionBuffer),
 						RegisterAsSRV(GraphBuilder, Instance->Strands.RestResource->PointToCurveBuffer),
 						bUseSingleGuide,
-						bValidGuide &&  bUseSingleGuide ? RegisterAsSRV(GraphBuilder, Instance->Strands.InterpolationResource->InterpolationBuffer) : nullptr,
-						bValidGuide && !bUseSingleGuide ? RegisterAsSRV(GraphBuilder, Instance->Strands.InterpolationResource->Interpolation0Buffer) : nullptr,
-						bValidGuide && !bUseSingleGuide ? RegisterAsSRV(GraphBuilder, Instance->Strands.InterpolationResource->Interpolation1Buffer) : nullptr,
+						bValidGuide ? RegisterAsSRV(GraphBuilder, Instance->Strands.InterpolationResource->InterpolationBuffer) : nullptr,
 						bValidGuide ? RegisterAsSRV(GraphBuilder, Instance->Guides.RestResource->PositionBuffer) : nullptr,
 						bValidGuide ? RegisterAsSRV(GraphBuilder, Instance->Guides.DeformedResource->GetBuffer(FHairStrandsDeformedResource::Current)) : nullptr,
 						bValidGuide ? RegisterAsSRV(GraphBuilder, Instance->Strands.InterpolationResource->SimRootPointIndexBuffer) : nullptr,
@@ -2056,9 +2034,7 @@ void ComputeHairStrandsInterpolation(
 							RegisterAsSRV(GraphBuilder, Instance->Strands.RestResource->PositionBuffer),
 							RegisterAsSRV(GraphBuilder, Instance->Strands.RestResource->PointToCurveBuffer),
 							bUseSingleGuide,
-							bValidGuide &&  bUseSingleGuide ? RegisterAsSRV(GraphBuilder, Instance->Strands.InterpolationResource->InterpolationBuffer) : nullptr,
-							bValidGuide && !bUseSingleGuide ? RegisterAsSRV(GraphBuilder, Instance->Strands.InterpolationResource->Interpolation0Buffer) : nullptr,
-							bValidGuide && !bUseSingleGuide ? RegisterAsSRV(GraphBuilder, Instance->Strands.InterpolationResource->Interpolation1Buffer) : nullptr,
+							bValidGuide ? RegisterAsSRV(GraphBuilder, Instance->Strands.InterpolationResource->InterpolationBuffer) : nullptr,
 							bValidGuide ? RegisterAsSRV(GraphBuilder, Instance->Guides.RestResource->PositionBuffer) : nullptr,
 							bValidGuide ? RegisterAsSRV(GraphBuilder, Instance->Guides.DeformedResource->GetBuffer(FHairStrandsDeformedResource::Previous)) : nullptr,
 							bValidGuide ? RegisterAsSRV(GraphBuilder, Instance->Strands.InterpolationResource->SimRootPointIndexBuffer) : nullptr,
@@ -2152,9 +2128,7 @@ void ComputeHairStrandsInterpolation(
 					Register(GraphBuilder, Instance->Strands.RestResource->CurveAttributeBuffer, ERDGImportedBufferFlags::CreateSRV).Buffer,
 					RegisterAsSRV(GraphBuilder, Instance->Strands.RestResource->CurveBuffer),
 					RegisterAsSRV(GraphBuilder, Instance->Strands.ClusterCullingResource->VertexToClusterIdBuffer),
-					bValidGuide &&  bUseSingleGuide ? RegisterAsSRV(GraphBuilder, Instance->Strands.InterpolationResource->InterpolationBuffer) : nullptr,
-					bValidGuide && !bUseSingleGuide ? RegisterAsSRV(GraphBuilder, Instance->Strands.InterpolationResource->Interpolation0Buffer) : nullptr,
-					bValidGuide && !bUseSingleGuide ? RegisterAsSRV(GraphBuilder, Instance->Strands.InterpolationResource->Interpolation1Buffer) : nullptr,
+					bValidGuide ? RegisterAsSRV(GraphBuilder, Instance->Strands.InterpolationResource->InterpolationBuffer) : nullptr,
 					OutRenCurveAttributeBuffer);
 			}
 			#endif
@@ -2357,9 +2331,7 @@ void ComputeHairStrandsInterpolation(
 						RegisterAsSRV(GraphBuilder, LOD.Guides.RestResource->PositionBuffer),
 						RegisterAsSRV(GraphBuilder, LOD.Guides.RestResource->PointToCurveBuffer),
 						bUseSingleGuide,
-						bValidGuide &&  bUseSingleGuide ? RegisterAsSRV(GraphBuilder, LOD.Guides.InterpolationResource->InterpolationBuffer) : nullptr,
-						bValidGuide && !bUseSingleGuide ? RegisterAsSRV(GraphBuilder, LOD.Guides.InterpolationResource->Interpolation0Buffer) : nullptr,
-						bValidGuide && !bUseSingleGuide ? RegisterAsSRV(GraphBuilder, LOD.Guides.InterpolationResource->Interpolation1Buffer) : nullptr,
+						bValidGuide ? RegisterAsSRV(GraphBuilder, LOD.Guides.InterpolationResource->InterpolationBuffer) : nullptr,
 						bValidGuide ? RegisterAsSRV(GraphBuilder, Instance->Guides.RestResource->PositionBuffer) : nullptr,
 						bValidGuide ? RegisterAsSRV(GraphBuilder, Instance->Guides.DeformedResource->GetBuffer(FHairStrandsDeformedResource::Current)) : nullptr,
 						RegisterAsSRV(GraphBuilder, LOD.Guides.InterpolationResource->SimRootPointIndexBuffer),
diff --git a/Engine/Plugins/Runtime/HairStrands/Source/HairStrandsCore/Public/GroomAsset.h b/Engine/Plugins/Runtime/HairStrands/Source/HairStrandsCore/Public/GroomAsset.h
--- a/Engine/Plugins/Runtime/HairStrands/Source/HairStrandsCore/Public/GroomAsset.h
+++ b/Engine/Plugins/Runtime/HairStrands/Source/HairStrandsCore/Public/GroomAsset.h
@@ -133,8 +133,6 @@ struct HAIRSTRANDSCORE_API FHairGroupPlatformData
 			uint32 Total = 0;
 			Total += FBase::GetDataSize();
 			Total += InterpolationBulkData.Data.Interpolation.IsBulkDataLoaded()	? InterpolationBulkData.Data.Interpolation.GetBulkDataSize() : 0;
-			Total += InterpolationBulkData.Data.Interpolation0.IsBulkDataLoaded()	? InterpolationBulkData.Data.Interpolation0.GetBulkDataSize() : 0;
-			Total += InterpolationBulkData.Data.Interpolation1.IsBulkDataLoaded()	? InterpolationBulkData.Data.Interpolation1.GetBulkDataSize() : 0;
 			Total += InterpolationBulkData.Data.SimRootPointIndex.IsBulkDataLoaded()? InterpolationBulkData.Data.SimRootPointIndex.GetBulkDataSize() : 0;
 			return Total;
 		}
diff --git a/Engine/Plugins/Runtime/HairStrands/Source/HairStrandsCore/Public/GroomResources.h b/Engine/Plugins/Runtime/HairStrands/Source/HairStrandsCore/Public/GroomResources.h
--- a/Engine/Plugins/Runtime/HairStrands/Source/HairStrandsCore/Public/GroomResources.h
+++ b/Engine/Plugins/Runtime/HairStrands/Source/HairStrandsCore/Public/GroomResources.h
@@ -509,17 +509,13 @@ struct FHairStrandsInterpolationResource : public FHairCommonResource
 	{
 		uint32 Total = 0;
 		Total += GetBufferTotalNumBytes(InterpolationBuffer);
-		Total += GetBufferTotalNumBytes(Interpolation0Buffer);
-		Total += GetBufferTotalNumBytes(Interpolation1Buffer);
 		Total += GetBufferTotalNumBytes(SimRootPointIndexBuffer);
 		return Total;
 	}
 
-	bool UseSingleGuide() const { return InterpolationBuffer.Buffer != nullptr; }
+	bool UseSingleGuide() const { return (BulkData.Header.Flags & FHairStrandsInterpolationBulkData::DataFlags_HasSingleGuideData) != 0; }
 
 	FRDGExternalBuffer InterpolationBuffer;
-	FRDGExternalBuffer Interpolation0Buffer;
-	FRDGExternalBuffer Interpolation1Buffer;
 	FRDGExternalBuffer SimRootPointIndexBuffer;
 
 	/* Reference to the hair strands interpolation render data */
diff --git a/Engine/Plugins/Runtime/HairStrands/Source/HairStrandsCore/Public/HairStrandsDatas.h b/Engine/Plugins/Runtime/HairStrands/Source/HairStrandsCore/Public/HairStrandsDatas.h
--- a/Engine/Plugins/Runtime/HairStrands/Source/HairStrandsCore/Public/HairStrandsDatas.h
+++ b/Engine/Plugins/Runtime/HairStrands/Source/HairStrandsCore/Public/HairStrandsDatas.h
@@ -31,44 +31,6 @@ struct FPackedHairAttribute0Vertex
 	uint8 Seed;
 };
 
-struct FHairInterpolationVertex
-{
-	typedef uint32 BulkType;
-
-	// Guide's vertex index are stored onto 24bits: 1) VertexGuideIndex0 is the lower part, 2) VertexGuideIndex1 is the upper part
-	uint16 VertexGuideIndex0;
-	uint8  VertexGuideIndex1;
-	uint8  VertexLerp;
-};
-
-struct FHairInterpolation0Vertex
-{
-	typedef uint64 BulkType;
-
-	uint16 Index0;
-	uint16 Index1;
-	uint16 Index2;
-
-	uint8 VertexWeight0;
-	uint8 VertexWeight1;
-};
-
-struct FHairInterpolation1Vertex
-{
-	typedef uint64 BulkType;
-
-	uint8 VertexIndex0;
-	uint8 VertexIndex1;
-	uint8 VertexIndex2;
-
-	uint8 VertexLerp0;
-	uint8 VertexLerp1;
-	uint8 VertexLerp2;
-
-	uint8 Pad0;
-	uint8 Pad1;
-};
-
 struct FPackedHairCurve
 {
 	typedef uint32 BulkType;
@@ -148,37 +110,15 @@ struct FHairStrandsTangentFormat
 
 struct FHairStrandsInterpolationFormat
 {
-	typedef FHairInterpolationVertex Type;
-	typedef FHairInterpolationVertex::BulkType BulkType;
+	typedef uint32 Type;
+	typedef uint32 BulkType;
 
 	static const uint32 ComponentCount = 1;
 	static const uint32 SizeInByte = sizeof(Type);
-	static const EVertexElementType VertexElementType = VET_UShort2;
+	static const EVertexElementType VertexElementType = VET_UInt;
 	static const EPixelFormat Format = PF_R32_UINT;
 };
 
-struct FHairStrandsInterpolation0Format
-{
-	typedef FHairInterpolation0Vertex Type;
-	typedef FHairInterpolation0Vertex::BulkType BulkType;
-
-	static const uint32 ComponentCount = 1;
-	static const uint32 SizeInByte = sizeof(Type);
-	static const EVertexElementType VertexElementType = VET_UShort4;
-	static const EPixelFormat Format = PF_R16G16B16A16_UINT;
-};
-
-struct FHairStrandsInterpolation1Format
-{
-	typedef FHairInterpolation1Vertex Type;
-	typedef FHairInterpolation1Vertex::BulkType BulkType;
-
-	static const uint32 ComponentCount = 1;
-	static const uint32 SizeInByte = sizeof(Type);
-	static const EVertexElementType VertexElementType = VET_UShort4;
-	static const EPixelFormat Format = PF_R16G16B16A16_UINT;
-};
-
 struct FHairStrandsRootIndexFormat
 {
 	typedef uint32 Type;
@@ -358,10 +298,7 @@ struct HAIRSTRANDSCORE_API FHairStrandsInterpolationBulkData : FHairStrandsBulkC
 
 	struct FData
 	{
-
-		FByteBulkData Interpolation;	// FHairStrandsInterpolationFormat  - Per-rendering-vertex interpolation data (closest guides, weight factors, ...). Data for a single guide
-		FByteBulkData Interpolation0;	// FHairStrandsInterpolation0Format - Per-rendering-vertex interpolation data (closest guides, weight factors, ...). Data for up to 3 guides
-		FByteBulkData Interpolation1;	// FHairStrandsInterpolation1Format - Per-rendering-vertex interpolation data (closest guides, weight factors, ...). Data for up to 3 guides
+		FByteBulkData Interpolation;	// FHairStrandsInterpolationFormat  - Per-rendering-vertex interpolation data (closest guides, weight factors, ...). Data for a 1 or 3 guide(s))
 		FByteBulkData SimRootPointIndex;// FHairStrandsRootIndexFormat      - Per-rendering-vertex index of the sim-root vertex
 	} Data;
 };
diff --git a/Engine/Shaders/Private/HairStrands/HairCardsDeformation.usf b/Engine/Shaders/Private/HairStrands/HairCardsDeformation.usf
--- a/Engine/Shaders/Private/HairStrands/HairCardsDeformation.usf
+++ b/Engine/Shaders/Private/HairStrands/HairCardsDeformation.usf
@@ -21,7 +21,7 @@ Buffer<float4> GuideDeformedPositionOffsetBuffer;
 Buffer<uint4>  GuideRestPositionBuffer;
 Buffer<uint4>  GuideDeformedPositionBuffer;
 Buffer<float4> CardsRestPositionBuffer;
-Buffer<uint>   CardsInterpolationBuffer;
+ByteAddressBuffer CardsInterpolationBuffer;
 
 RWBuffer<float4> CardsDeformedPositionBuffer; 
 
@@ -92,11 +92,10 @@ void MainCS(uint DispatchThreadId : SV_DispatchThreadID)
 #if PERMUTATION_DYNAMIC_GEOMETRY == 0
 	// Simple deformation based on world space offset
 
-	const uint PackedInterpolationData = CardsInterpolationBuffer[VertexId];
-	const uint GuideIndex0 = PackedInterpolationData & 0xFFFFFF;
+	const FGuideData GuideData = UnpackGuideData(CardsInterpolationBuffer.Load(VertexId * HAIR_INTERPOLATION_1GUIDE_STRIDE));
+	const uint GuideIndex0 = GuideData.VertexIndices.x;
 	const uint GuideIndex1 = min(GuideIndex0+1, GuideVertexCount-1);
-
-	const float GuideLerp = ((PackedInterpolationData>>24) & 0xFF)/float(0xFF);
+	const float GuideLerp = GuideData.VertexLerps.x;
 
 	const FHairControlPoint RestGuideP0 = UnpackHairControlPoint(
 		GuideRestPositionBuffer[GuideIndex0],
@@ -140,10 +139,10 @@ void MainCS(uint DispatchThreadId : SV_DispatchThreadID)
 #elif PERMUTATION_DYNAMIC_GEOMETRY == 1
 	// Compute the deformation offset in triangle local space
 
-	const uint PackedInterpolationData = CardsInterpolationBuffer[VertexId];
-	const uint GuideIndex0 = PackedInterpolationData & 0xFFFFFF;
+	const FGuideData GuideData = UnpackGuideData(CardsInterpolationBuffer.Load(VertexId * HAIR_INTERPOLATION_1GUIDE_STRIDE));
+	const uint GuideIndex0 = GuideData.VertexIndices.x;
 	const uint GuideIndex1 = min(GuideIndex0 + 1, GuideVertexCount - 1);
-	const float GuideLerp = ((PackedInterpolationData >> 24) & 0xFF) / float(0xFF);
+	const float GuideLerp = GuideData.VertexLerps.x;
 
 	const uint RootIndex = GuideVertexToRootIndexBuffer[GuideIndex0];
 	const uint TriangleIndex = GuideRootToUniqueTriangleIndexBuffer[RootIndex];
diff --git a/Engine/Shaders/Private/HairStrands/HairStrandsInterpolation.usf b/Engine/Shaders/Private/HairStrands/HairStrandsInterpolation.usf
--- a/Engine/Shaders/Private/HairStrands/HairStrandsInterpolation.usf
+++ b/Engine/Shaders/Private/HairStrands/HairStrandsInterpolation.usf
@@ -22,9 +22,16 @@
 
 #pragma warning (default:7203) // downgrade 'not enough registers available for the entire program.' to a warning instead of error
 
+
+
 ///////////////////////////////////////////////////////////////////////////////////////////////////
 // Guides
-struct FGuideData
+
+#if SHADER_HAIRINTERPOLATION || SHADER_PATCHATTRIBUTE
+ByteAddressBuffer InterpolationBuffer;
+#endif
+
+struct FGuideDataWithOffset
 {
 	uint3  VertexIndices;
 	float3 VertexLerps;
@@ -32,34 +39,28 @@ struct FGuideData
 	float3 OutPositionOffset;
 };
 
-FGuideData GetGuideData(uint InVertexIndex, uint InInterpolationData, float3 OutPositionOffset)
+FGuideDataWithOffset GetGuideData(uint InVertexIndex, bool bSingleGuide, float3 OutPositionOffset)
 {
-	FGuideData Out = (FGuideData)0;
-	Out.OutPositionOffset	= OutPositionOffset;
-	Out.VertexWeights		= float3(1, 0, 0);
-	Out.VertexIndices		= 0;
-	Out.VertexLerps			= 0;
-	Out.VertexIndices.x		= InInterpolationData & 0xFFFFFF;
-	Out.VertexLerps.x		= saturate(((InInterpolationData >> 24) & 0xFF) / 255.f);
-	return Out;
-}
+#if PERMUTATION_SIMULATION > 0
+	FGuideData GuideData;
+	if (bSingleGuide)
+	{
+		GuideData = UnpackGuideData(InterpolationBuffer.Load(InVertexIndex * HAIR_INTERPOLATION_1GUIDE_STRIDE));
+	}
+	else
+	{
+		GuideData = UnpackGuideData(InterpolationBuffer.Load4(InVertexIndex * HAIR_INTERPOLATION_3GUIDE_STRIDE));
+	}
 
-FGuideData GetGuideData(uint InVertexIndex, uint4 InInterpolationData0, uint4 InInterpolationData1, float3 OutPositionOffset)
-{
-	FGuideData Out = (FGuideData)0;
+	FGuideDataWithOffset Out;
+	Out.VertexIndices     = GuideData.VertexIndices;
+	Out.VertexLerps       = GuideData.VertexLerps;
+	Out.VertexWeights     = GuideData.VertexWeights;
 	Out.OutPositionOffset = OutPositionOffset;
-
-	Out.VertexIndices.x = InInterpolationData0.x | (( InInterpolationData1.x & 0xFF) << 16);
-	Out.VertexIndices.y = InInterpolationData0.y | (((InInterpolationData1.x >> 8) & 0xFF) << 16);
-	Out.VertexIndices.z = InInterpolationData0.z | (( InInterpolationData1.y & 0xFF) << 16);
-
-	Out.VertexWeights	= float3(saturate((InInterpolationData0.w & 0xFF) / 255.f), saturate(((InInterpolationData0.w >> 8) & 0xFF) / 255.f), 0);
-	Out.VertexWeights.z = saturate(1.0f - (Out.VertexWeights.x + Out.VertexWeights.y));
-
-	Out.VertexLerps.x	= saturate(((InInterpolationData1.y >> 8) & 0xFF) / 255.f);
-	Out.VertexLerps.y	= saturate(( InInterpolationData1.z & 0xFF) / 255.f);
-	Out.VertexLerps.z	= saturate(((InInterpolationData1.z >> 8) & 0xFF) / 255.f);
 	return Out;
+#else
+	return (FGuideDataWithOffset)0;
+#endif
 }
 
 ///////////////////////////////////////////////////////////////////////////////////////////////////
@@ -80,22 +81,10 @@ float3 InRenHairPositionOffset;
 float3 InSimHairPositionOffset;
 float4x4 LocalToWorldMatrix;
 
-Buffer<uint4>  RenRestPosePositionBuffer;
-Buffer<uint4>  SimRestPosePositionBuffer;
-Buffer<uint4>  SimDeformedPositionBuffer;
-
-// Hair interpolation support either 
-// * A single guide, 
-// * Or up to 3 guides. 
-// This option is setup on interpolation data of the groom asset
-#if PERMUTATION_USE_SINGLE_GUIDE == 1
-Buffer<uint>	InterpolationBuffer;
-#else
-Buffer<uint4>	Interpolation0Buffer;
-Buffer<uint4>	Interpolation1Buffer;
-#endif
-
-Buffer<uint>   SimRootPointIndexBuffer;
+Buffer<uint4> RenRestPosePositionBuffer;
+Buffer<uint4> SimRestPosePositionBuffer;
+Buffer<uint4> SimDeformedPositionBuffer;
+Buffer<uint>  SimRootPointIndexBuffer;
 
 uint HairStrandsVFTODO_bCullingEnable;
 #if PERMUTATION_CULLING == 1
@@ -263,25 +252,9 @@ FRootTriangleData GetRootTriangleData(uint VertexIndex, bool bSim)
 	return Out;
 }
 
-///////////////////////////////////////////////////////////////////////////////////////////////////
-// Guides
-
-FGuideData GetGuideData(uint InVertexIndex)
-{
-#if PERMUTATION_SIMULATION == 1
-  #if PERMUTATION_USE_SINGLE_GUIDE
-	return GetGuideData(InVertexIndex, InterpolationBuffer[InVertexIndex], OutSimHairPositionOffsetBuffer[0].xyz);
-  #else
-	return GetGuideData(InVertexIndex, Interpolation0Buffer[InVertexIndex], Interpolation1Buffer[InVertexIndex], OutSimHairPositionOffsetBuffer[0].xyz);
-  #endif
-#else
-	return (FGuideData)0;
-#endif
-}
-
 ///////////////////////////////////////////////////////////////////////////////////////////////////
 // Rigid interpolation
-float3 ApplyRigidInterpolation(FGuideData GuideData, float3 InPosition)
+float3 ApplyRigidInterpolation(FGuideDataWithOffset GuideData, float3 InPosition)
 {
 	float3 CurrOffset = 0;
 #if PERMUTATION_SIMULATION == 1
@@ -307,7 +280,7 @@ float3 ApplyRigidInterpolation(FGuideData GuideData, float3 InPosition)
 
 ///////////////////////////////////////////////////////////////////////////////////////////////////
 // Skinning with offset
-float3 ApplySkinningWithOffset(const FGuideData GuideData, const FRootTriangleData RenTri, float3 InPosition)
+float3 ApplySkinningWithOffset(const FGuideDataWithOffset GuideData, const FRootTriangleData RenTri, float3 InPosition)
 {
 	float3 CurrOffset = 0;
 
@@ -340,7 +313,7 @@ float3 ApplySkinningWithOffset(const FGuideData GuideData, const FRootTriangleDa
 
 ///////////////////////////////////////////////////////////////////////////////////////////////////
 // Skinning with translantion
-float3 ApplySkinningWithTranslation(const FGuideData GuideData, const FRootTriangleData RenTri, float3 InPosition)
+float3 ApplySkinningWithTranslation(const FGuideDataWithOffset GuideData, const FRootTriangleData RenTri, float3 InPosition)
 {
 	InPosition -= RenTri.RootBarycentric.x * RenTri.RestTriangle.P0 + RenTri.RootBarycentric.y * RenTri.RestTriangle.P1 + RenTri.RootBarycentric.z * RenTri.RestTriangle.P2;
 
@@ -370,7 +343,7 @@ float3 ApplySkinningWithTranslation(const FGuideData GuideData, const FRootTrian
 
 ///////////////////////////////////////////////////////////////////////////////////////////////////
 // Skinning with rotation
-float3 ApplySkinningWithRotation(const FGuideData GuideData, const FRootTriangleData RenTri, float3 InPosition)
+float3 ApplySkinningWithRotation(const FGuideDataWithOffset GuideData, const FRootTriangleData RenTri, float3 InPosition)
 {
 	InPosition -= RenTri.RootBarycentric.x * RenTri.RestTriangle.P0 + RenTri.RootBarycentric.y * RenTri.RestTriangle.P1 + RenTri.RootBarycentric.z * RenTri.RestTriangle.P2;
 
@@ -404,7 +377,7 @@ float3 ApplySkinningWithRotation(const FGuideData GuideData, const FRootTriangle
 
 ///////////////////////////////////////////////////////////////////////////////////////////////////
 // Skinning with translation and rotation
-float3 ApplySkinningWithTranslationAndRotation(const FGuideData GuideData, const FRootTriangleData RenTri, float3 InPosition)
+float3 ApplySkinningWithTranslationAndRotation(const FGuideDataWithOffset GuideData, const FRootTriangleData RenTri, float3 InPosition)
 {
 	float3 CurrOffset = 0;
 
@@ -478,7 +451,7 @@ void MainCS(uint2 DispatchThreadId : SV_DispatchThreadID)
 	// 2. Appply deformation
 	if (VertexIndex < VertexCount)
 	{
-		const FGuideData GuideData = GetGuideData(VertexIndex);
+		const FGuideDataWithOffset GuideData = GetGuideData(VertexIndex, PERMUTATION_USE_SINGLE_GUIDE, OutSimHairPositionOffsetBuffer[0].xyz);
 		const FRootTriangleData RenTri = GetRootTriangleData(VertexIndex, false);
 
 		// Manual decoding of the rest position
@@ -554,26 +527,8 @@ uint				CurveCount;
 Buffer<uint>		RenCurveBuffer;
 uint4				RenCurveAttributeOffsets[HAIR_CURVE_ATTRIBUTE_OFFSET_COUNT];
 Buffer<int>			RenVertexToClusterIdBuffer;
-
-#if PERMUTATION_SIMULATION == 1
-Buffer<uint>		InterpolationBuffer;
-#elif PERMUTATION_SIMULATION == 2
-Buffer<uint4>		Interpolation0Buffer;
-Buffer<uint4>		Interpolation1Buffer;
-#endif
 RWByteAddressBuffer	OutRenAttributeBuffer;
 
-
-FGuideData GetGuideData(uint InVertexIndex)
-{
-#if PERMUTATION_SIMULATION == 1
-	return GetGuideData(InVertexIndex, InterpolationBuffer[InVertexIndex], 0 /*OutPositionOffset*/);
-#elif PERMUTATION_SIMULATION == 2
-	return GetGuideData(InVertexIndex, Interpolation0Buffer[InVertexIndex], Interpolation1Buffer[InVertexIndex], 0 /*OutPositionOffset*/);
-#else
-	return (FGuideData)0;
-#endif
-}
 groupshared uint SeedValues[1024];
 
 uint ToSeedHash(uint In)
@@ -609,7 +564,7 @@ void MainCS(uint2 DispatchThreadId : SV_DispatchThreadID, uint GroupThread1D : S
 		#else // if (PERMUTATION_SIMULATION > 0)
 		{
 			const uint VertexIndex0		= UnpackHairCurve(RenCurveBuffer[CurveIndex]).PointIndex;
-			const FGuideData GuideData	= GetGuideData(VertexIndex0);
+			const FGuideDataWithOffset GuideData = GetGuideData(VertexIndex0, PERMUTATION_SIMULATION == 1, 0 /*OutPositionOffset*/);
 			const uint GuideIndex		= GuideData.VertexIndices[0]; // Take the closest guide
 			SeedValues[GroupThread1D]	= ToSeedHash(GuideIndex);
 		}
diff --git a/Engine/Shaders/Private/HairStrands/HairStrandsVertexFactoryCommon.ush b/Engine/Shaders/Private/HairStrands/HairStrandsVertexFactoryCommon.ush
--- a/Engine/Shaders/Private/HairStrands/HairStrandsVertexFactoryCommon.ush
+++ b/Engine/Shaders/Private/HairStrands/HairStrandsVertexFactoryCommon.ush
@@ -125,6 +125,45 @@ float UnpackHairRoughness(uint In)
 ////////////////////////////////////////////////////////////////////////////////////////////////////////////////////////////////
 // Hair curve
 
+float UnpackR8(uint In) { return saturate((In & 0xFF) / 255.f); }
+
+struct FGuideData
+{
+	uint3  VertexIndices;
+	float3 VertexLerps;
+	float3 VertexWeights;
+};
+
+FGuideData UnpackGuideData(uint In)
+{
+	FGuideData Out = (FGuideData)0;
+	Out.VertexWeights		= float3(1, 0, 0);
+	Out.VertexIndices		= 0;
+	Out.VertexLerps			= 0;
+	Out.VertexIndices.x		= In & 0xFFFFFF;
+	Out.VertexLerps.x		= UnpackR8(In >> 24);
+	return Out;
+}
+
+FGuideData UnpackGuideData(uint4 In)
+{
+	FGuideData Out = (FGuideData)0;
+	Out.VertexIndices.x = (In.x & 0xFFFF) | (((In.z      ) & 0xFF) << 16);
+	Out.VertexIndices.y = (In.x >> 16   ) | (((In.z >> 8 ) & 0xFF) << 16);
+	Out.VertexIndices.z = (In.y & 0xFFFF) | (((In.z >> 16) & 0xFF) << 16);
+
+	Out.VertexWeights	= float3(UnpackR8(In.y>>16), UnpackR8(In.y>>24), 0);
+	Out.VertexWeights.z = saturate(1.0f - (Out.VertexWeights.x + Out.VertexWeights.y));
+
+	Out.VertexLerps.x	= UnpackR8(In.w);
+	Out.VertexLerps.y	= UnpackR8(In.w >> 8);
+	Out.VertexLerps.z	= UnpackR8(In.w >> 16);
+	return Out;
+}
+
+////////////////////////////////////////////////////////////////////////////////////////////////////////////////////////////////
+// Hair curve
+
 struct FHairCurve
 {
 	uint PointIndex;
diff --git a/Engine/Shaders/Shared/HairStrandsDefinitions.h b/Engine/Shaders/Shared/HairStrandsDefinitions.h
--- a/Engine/Shaders/Shared/HairStrandsDefinitions.h
+++ b/Engine/Shaders/Shared/HairStrandsDefinitions.h
@@ -31,6 +31,10 @@
 
 #define HAIR_ATTRIBUTE_INVALID_OFFSET 0xFFFFFFFF
 
+// Hair interpolation (in bytes)
+#define HAIR_INTERPOLATION_1GUIDE_STRIDE 4
+#define HAIR_INTERPOLATION_3GUIDE_STRIDE 16
+
 // Max number of discrete LOD that a hair group can have
 #define MAX_HAIR_LOD 8
 
