Made packing and unpacking code common in a single file always included in common.ush.
This should fix double defines of packing and unpacking while making common optimised packing functions.
Moved code from Platform.ush is only used in file including common.ush.

diff --git a/Engine/Shaders/Private/Common.ush b/Engine/Shaders/Private/Common.ush
--- a/Engine/Shaders/Private/Common.ush
+++ b/Engine/Shaders/Private/Common.ush
@@ -7,6 +7,9 @@
 #pragma once
 
 #include "/Engine/Public/Platform.ush"
+
+#include "PackUnpack.ush"
+
 // These types are used for material translator generated code, or any functions the translated code can call
 #if PIXELSHADER && !FORCE_MATERIAL_FLOAT_FULL_PRECISION
 	#define MaterialFloat half
@@ -898,34 +901,6 @@ float4 ConvertTangentSnormToUnorm16(float4 Input)
 #define TangentBias(X)		(X)
 #define TangentUnbias(X)	(X)
 
-uint PackUnorm2x16(float2 v)
-{
-	uint2 sv = uint2(round(clamp(v, 0.0, 1.0) * 65535.0));
-	return (sv.x | (sv.y << 16u));
-}
-
-uint PackSnorm2x16(float2 v)
-{
-	uint2 sv = uint2(round(clamp(v, -1.0, 1.0) * 32767.0) + 32767.0);
-	return (sv.x | (sv.y << 16u));
-}
-
-float2 UnpackUnorm2x16(uint p)
-{
-	float2 Ret;
-	Ret.x = (p & 0xffff) * rcp(65535.0f);
-	Ret.y = (p >> 16u) * rcp(65535.0f);
-	return Ret;
-}
-
-float2 UnpackSnorm2x16(uint p)
-{
-	float2 Ret;
-	Ret.x = clamp((float(p & 0xffff) - 32767.0f) * rcp(32767.0f), -1.0, 1.0);
-	Ret.y = clamp((float(p >> 16u) - 32767.0f) * rcp(32767.0f), -1.0, 1.0);
-	return Ret;
-}
-
 float Square( float x )
 {
 	return x*x;
diff --git a/Engine/Shaders/Private/HairStrands/HairStrandsVertexFactoryCommon.ush b/Engine/Shaders/Private/HairStrands/HairStrandsVertexFactoryCommon.ush
--- a/Engine/Shaders/Private/HairStrands/HairStrandsVertexFactoryCommon.ush
+++ b/Engine/Shaders/Private/HairStrands/HairStrandsVertexFactoryCommon.ush
@@ -125,8 +125,6 @@ float UnpackHairRoughness(uint In)
 ////////////////////////////////////////////////////////////////////////////////////////////////////////////////////////////////
 // Hair curve
 
-float UnpackR8(uint In) { return saturate((In & 0xFF) / 255.f); }
-
 struct FGuideData
 {
 	uint3  VertexIndices;
diff --git a/Engine/Shaders/Private/PackUnpack.ush b/Engine/Shaders/Private/PackUnpack.ush
new file mode 100644
--- /dev/null
+++ b/Engine/Shaders/Private/PackUnpack.ush
@@ -0,0 +1,108 @@
+// Copyright Epic Games, Inc. All Rights Reserved.
+
+/*=============================================================================
+	PackUnpack.ush: Common shader code for packing and unpacking data. 
+					Automatically included from Common.ush.
+=============================================================================*/
+
+#pragma once
+
+#include "/Engine/Public/Platform.ush"
+
+//
+//	float2 <-> Unorm
+//
+
+uint PackUnorm2x16(float2 v)
+{
+	uint2 sv = uint2(round(clamp(v, 0.0, 1.0) * 65535.0));
+	return (sv.x | (sv.y << 16u));
+}
+
+float2 UnpackUnorm2x16(uint p)
+{
+	float2 Ret;
+	Ret.x = (p & 0xffff) * rcp(65535.0f);
+	Ret.y = (p >> 16u) * rcp(65535.0f);
+	return Ret;
+}
+
+//
+//	float2 <-> Snorm
+//
+
+uint PackSnorm2x16(float2 v)
+{
+	uint2 sv = uint2(round(clamp(v, -1.0, 1.0) * 32767.0) + 32767.0);
+	return (sv.x | (sv.y << 16u));
+}
+
+float2 UnpackSnorm2x16(uint p)
+{
+	float2 Ret;
+	Ret.x = clamp((float(p & 0xffff) - 32767.0f) * rcp(32767.0f), -1.0, 1.0);
+	Ret.y = clamp((float(p >> 16u) - 32767.0f) * rcp(32767.0f), -1.0, 1.0);
+	return Ret;
+}
+
+//
+//	uint2 <-> uint
+//
+
+#ifndef COMPILER_SUPPORTS_PACK_INTRINSICS
+#define COMPILER_SUPPORTS_PACK_INTRINSICS 0
+
+uint PackUInt2ToUInt(uint X, uint Y)
+{
+	return X | (Y << 16u);
+}
+
+#endif
+
+uint PackUInt2ToUInt(uint2 XY)
+{
+	return PackUInt2ToUInt(XY.x, XY.y);
+}
+
+uint2 UnpackUInt2FromUInt(uint Packed)
+{
+	return uint2(Packed & 0xffff, Packed >> 16);
+}
+
+//
+//	float2 <-> uint
+//
+
+uint PackFloat2ToUInt(float X, float Y)
+{
+	return PackUInt2ToUInt(f32tof16(X), f32tof16(Y));
+}
+
+uint PackFloat2ToUInt(float2 XY)
+{
+	return PackFloat2ToUInt(XY.x, XY.y);
+}
+
+float2 UnpackFloat2FromUInt(uint In)
+{
+	return float2(f16tof32(In), f16tof32(In >> 16));
+}
+
+//
+//	float <-> u8
+//
+
+uint PackR8(float Value)
+{
+	return uint(saturate(Value) * 255.0f);
+}
+
+float UnpackR8(uint In)
+{
+#if COMPILER_SUPPORTS_UNPACKBYTEN
+	return UnpackByte0(In) * (1.0f / 255.0f);
+#else
+	return float(In & 0xFF) * (1.0f / 255.0f);
+#endif
+}
+
diff --git a/Engine/Shaders/Private/Strata/Strata.ush b/Engine/Shaders/Private/Strata/Strata.ush
--- a/Engine/Shaders/Private/Strata/Strata.ush
+++ b/Engine/Shaders/Private/Strata/Strata.ush
@@ -535,8 +535,6 @@ void   EnableSlabBSDFSimpleVolumetric(inout FStrataBSDF BSDF);
 void   StrataRequestSharedLocalBasisTangent(inout uint Types, uint Index);
 float3 StrataUnpackNormal(uint PackedNormal);
 void   StrataUnpackNormalAndTangent(inout float3 Normal, inout float3 Tangent, in uint InPacked);
-uint   PackR8(float Value);
-float  UnpackR8(uint Value);
 uint   PackRGBA8(float4 rgba);
 float4 UnpackRGBA8(uint rgba);
 uint   Pack10F(float Value);
@@ -3096,20 +3094,6 @@ float4 UnpackColorGamma2ToLinearAlphaLinear(uint In)
 	return float4(StrataSrgbToLinear(Out.rgb), Out.a);
 }
 
-uint PackR8(float Value)
-{
-	return uint(saturate(Value) * 255.0f);
-}
-
-float UnpackR8(uint In)
-{
-#if COMPILER_SUPPORTS_UNPACKBYTEN
-	return UnpackByte0(In) * (1.0f / 255.0f);
-#else
-	return float(In & 0xFF) * (1.0f / 255.0f);
-#endif
-}
-
 uint PackRGBA8(float4 In)
 {
 #if COMPILER_SUPPORTS_PACK_INTRINSICS
diff --git a/Engine/Shaders/Public/Platform.ush b/Engine/Shaders/Public/Platform.ush
--- a/Engine/Shaders/Public/Platform.ush
+++ b/Engine/Shaders/Public/Platform.ush
@@ -219,41 +219,6 @@
 #error Collision between FORCE_FLOATS and PLATFORM_SUPPORTS_REAL_TYPES
 #endif
 
-#ifndef COMPILER_SUPPORTS_PACK_INTRINSICS
-#define COMPILER_SUPPORTS_PACK_INTRINSICS 0
-
-uint PackUInt2ToUInt(uint X, uint Y)
-{
-	return X | (Y << 16u);
-}
-
-#endif
-
-uint PackUInt2ToUInt(uint2 XY)
-{
-	return PackUInt2ToUInt(XY.x, XY.y);
-}
-
-uint2 UnpackUInt2FromUInt(uint Packed)
-{
-	return uint2(Packed & 0xffff, Packed >> 16);
-}
-
-uint PackFloat2ToUInt(float X, float Y)
-{
-	return PackUInt2ToUInt(f32tof16(X), f32tof16(Y));
-}
-
-uint PackFloat2ToUInt(float2 XY)
-{
-	return PackFloat2ToUInt(XY.x, XY.y);
-}
-
-float2 UnpackFloat2FromUInt(uint In)
-{
-	return float2(f16tof32(In), f16tof32(In >> 16));
-}
-
 #ifndef COMPILER_SUPPORTS_WAVE_32_64_MODE
 #define COMPILER_SUPPORTS_WAVE_32_64_MODE 0
 #endif
