Mode some more common pack/unpack function into commone header.

diff --git a/Engine/Shaders/Private/HairStrands/HairStrandsForwardRaster.usf b/Engine/Shaders/Private/HairStrands/HairStrandsForwardRaster.usf
--- a/Engine/Shaders/Private/HairStrands/HairStrandsForwardRaster.usf
+++ b/Engine/Shaders/Private/HairStrands/HairStrandsForwardRaster.usf
@@ -25,22 +25,6 @@ uint2 UnpackTileCoord(uint In)
 	return uint2(In & 0xffff, (In >> 16) & 0xffff);
 }
 
-uint PackR11G11B10F(float3 rgb)
-{
-	uint r = (f32tof16(rgb.r) << 17) & 0xFFE00000;
-	uint g = (f32tof16(rgb.g) << 6) & 0x001FFC00;
-	uint b = (f32tof16(rgb.b) >> 5) & 0x000003FF;
-	return r | g | b;
-}
-
-float3 UnpackR11G11B10F(uint rgb)
-{
-	float r = f16tof32((rgb >> 17) & 0x7FF0);
-	float g = f16tof32((rgb >> 6) & 0x7FF0);
-	float b = f16tof32((rgb << 5) & 0x7FE0);
-	return float3(r, g, b);
-}
-
 uint PackDepth(float In)
 {
 	return asuint(In);
diff --git a/Engine/Shaders/Private/PackUnpack.ush b/Engine/Shaders/Private/PackUnpack.ush
--- a/Engine/Shaders/Private/PackUnpack.ush
+++ b/Engine/Shaders/Private/PackUnpack.ush
@@ -10,7 +10,7 @@
 #include "/Engine/Public/Platform.ush"
 
 //
-//	float2 <-> Unorm
+//	float2 <-> u32 as Unorm16.Unorm16
 //
 
 uint PackUnorm2x16(float2 v)
@@ -28,7 +28,7 @@ float2 UnpackUnorm2x16(uint p)
 }
 
 //
-//	float2 <-> Snorm
+//	float2 <-> u32 as Snorm16.Snorm16 
 //
 
 uint PackSnorm2x16(float2 v)
@@ -46,7 +46,7 @@ float2 UnpackSnorm2x16(uint p)
 }
 
 //
-//	uint2 <-> uint
+//	uint2 <-> u32 as u16.u16
 //
 
 #ifndef COMPILER_SUPPORTS_PACK_INTRINSICS
@@ -70,7 +70,7 @@ uint2 UnpackUInt2FromUInt(uint Packed)
 }
 
 //
-//	float2 <-> uint
+//	float2 <-> u32 as f16.f16
 //
 
 uint PackFloat2ToUInt(float X, float Y)
@@ -106,3 +106,106 @@ float UnpackR8(uint In)
 #endif
 }
 
+//
+//	float4 <-> u32 as u8.u8.u8.u8
+//
+
+uint PackRGBA8(float4 In)
+{
+#if COMPILER_SUPPORTS_PACK_INTRINSICS
+	return PackFloat4ToUInt(saturate(In) * 255.0f);
+#else
+	uint r = (uint(saturate(In.r) * 255.0f) << 0);
+	uint g = (uint(saturate(In.g) * 255.0f) << 8);
+	uint b = (uint(saturate(In.b) * 255.0f) << 16);
+	uint a = (uint(saturate(In.a) * 255.0f) << 24);
+	return r | g | b | a;
+#endif
+}
+
+float4 UnpackRGBA8(uint In)
+{
+#if COMPILER_SUPPORTS_UNPACKBYTEN
+	return float4(
+		UnpackByte0(In),
+		UnpackByte1(In),
+		UnpackByte2(In),
+		UnpackByte3(In)) * (1.0f / 255.0f);
+#else
+	float4 Out;
+	Out.r = float((In >> 0) & 0xFF) * (1.0f / 255.0f);
+	Out.g = float((In >> 8) & 0xFF) * (1.0f / 255.0f);
+	Out.b = float((In >> 16) & 0xFF) * (1.0f / 255.0f);
+	Out.a = float((In >> 24) & 0xFF) * (1.0f / 255.0f);
+	return Out;
+#endif
+}
+
+//
+//	float <-> f10
+//
+
+uint Pack10F(float Value)
+{
+	return (f32tof16(Value) >> 5) & 0x000003FF;
+}
+
+float Unpack10F(uint Value)
+{
+	return f16tof32((Value << 5) & 0x7FE0);
+}
+
+//
+//	float3 <-> u32 as f11.f11.f10
+//
+
+uint PackR11G11B10F(float3 rgb)
+{
+	uint r = (f32tof16(rgb.r) << 17) & 0xFFE00000;
+	uint g = (f32tof16(rgb.g) << 6) & 0x001FFC00;
+	uint b = (f32tof16(rgb.b) >> 5) & 0x000003FF;
+	return r | g | b;
+}
+
+float3 UnpackR11G11B10F(uint rgb)
+{
+	float r = f16tof32((rgb >> 17) & 0x7FF0);
+	float g = f16tof32((rgb >> 6) & 0x7FF0);
+	float b = f16tof32((rgb << 5) & 0x7FE0);
+	return float3(r, g, b);
+}
+
+//
+//	float3 <-> f10.f10.f10
+//
+
+uint PackR10G10B10F(float3 rgb)
+{
+	uint r = (f32tof16(rgb.r) << 15) & 0x3FF00000;	// 0011 1111 1111 0000 0000 0000 0000 0000 
+	uint g = (f32tof16(rgb.g) << 5) & 0x000FFC00;	// 0000 0000 0000 1111 1111 1100 0000 0000 
+	uint b = (f32tof16(rgb.b) >> 5) & 0x000003FF;	// 0000 0000 0000 0000 0000 0011 1111 1111
+	return r | g | b;
+}
+
+float3 UnpackR10G10B10F(uint rgb)
+{
+	float r = f16tof32((rgb >> 15) & 0x7FE0);
+	float g = f16tof32((rgb >> 5) & 0x7FE0);
+	float b = f16tof32((rgb << 5) & 0x7FE0);
+	return float3(r, g, b);
+}
+
+//
+//	float4 <-> uint2 as f16.f16.f16.f16
+//
+
+uint2 PackR16G16B16A16F(float4 In)
+{
+	return uint2(PackFloat2ToUInt(In.xy), PackFloat2ToUInt(In.zw));
+}
+
+float4 UnpackR16G16B16A16F(uint2 In)
+{
+	return float4(UnpackFloat2FromUInt(In.x), UnpackFloat2FromUInt(In.y));
+}
+
diff --git a/Engine/Shaders/Private/PathTracing/PathTracingCommon.ush b/Engine/Shaders/Private/PathTracing/PathTracingCommon.ush
--- a/Engine/Shaders/Private/PathTracing/PathTracingCommon.ush
+++ b/Engine/Shaders/Private/PathTracing/PathTracingCommon.ush
@@ -74,19 +74,12 @@ uint PayloadEncodeHDRColor(float3 rgb)
 {
 	// clamp to the largest representable value
 	rgb = clamp(rgb, 0.0, Max111110BitsFloat3);
-	// See PackR11G11B10F (duplicated to keep the code independent of Strata for now)
-	uint r = (f32tof16(rgb.r) << 17) & 0xFFE00000;
-	uint g = (f32tof16(rgb.g) <<  6) & 0x001FFC00;
-	uint b = (f32tof16(rgb.b) >>  5) & 0x000003FF;
-	return r | g | b;
+	return PackR11G11B10F(rgb);
 }
 
 float3 PayloadDecodeHDRColor(uint rgb)
 {
-	float r = f16tof32((rgb >> 17) & 0x7FF0);
-	float g = f16tof32((rgb >>  6) & 0x7FF0);
-	float b = f16tof32((rgb <<  5) & 0x7FE0);
-	return float3(r, g, b);
+	return UnpackR11G11B10F(rgb);
 }
 
 uint PayloadEncodeLDRColor(float3 rgb)
diff --git a/Engine/Shaders/Private/Strata/Strata.ush b/Engine/Shaders/Private/Strata/Strata.ush
--- a/Engine/Shaders/Private/Strata/Strata.ush
+++ b/Engine/Shaders/Private/Strata/Strata.ush
@@ -535,22 +535,16 @@ void   EnableSlabBSDFSimpleVolumetric(inout FStrataBSDF BSDF);
 void   StrataRequestSharedLocalBasisTangent(inout uint Types, uint Index);
 float3 StrataUnpackNormal(uint PackedNormal);
 void   StrataUnpackNormalAndTangent(inout float3 Normal, inout float3 Tangent, in uint InPacked);
-uint   PackRGBA8(float4 rgba);
-float4 UnpackRGBA8(uint rgba);
-uint   Pack10F(float Value);
-float  Unpack10F(uint Value);
 uint   PackR6(float Value, float Dither);
 float  UnpackR6(uint Value);
 uint   PackR7(float Value, float Dither);
 float  UnpackR7(uint Value);
-uint   PackR11G11B10F(float3 rgb);
-float3 UnpackR11G11B10F(uint rgb);
-uint   PackR10G10B10F(float3 rgb);
-float3 UnpackR10G10B10F(uint rgb);
 uint   PackR7G7B6Gamma2(float3 rgb, float Dither);
 float3 UnpackR7G7B6Gamma2(uint rgb);
-uint   PackColorLinearToGamma2AlphaLinear(float4 rgba);
-float4 UnpackColorGamma2ToLinearAlphaLinear(uint rgba);
+uint   PackColorLinearToGamma2(float3 In);
+float3 UnpackColorGamma2ToLinear(uint In);
+uint   PackColorLinearToGamma2AlphaLinear(float4 In);
+float4 UnpackColorGamma2ToLinearAlphaLinear(uint In);
 uint   StrataPackNormal(in float3 Normal);
 float3 StrataUnpackNormal(uint PackedNormal);
 uint   StrataPackNormal24(in float3 Normal);
@@ -3020,58 +3014,6 @@ half3 StrataSrgbToLinear(half3 In)
 #endif
 }
 
-uint Pack10F(float Value)
-{
-	return (f32tof16(Value) >> 5) & 0x000003FF;
-}
-
-float Unpack10F(uint Value)
-{
-	return f16tof32((Value << 5) & 0x7FE0);
-}
-
-uint PackR11G11B10F(float3 rgb)
-{
-	uint r = (f32tof16(rgb.r) << 17) & 0xFFE00000;
-	uint g = (f32tof16(rgb.g) << 6)  & 0x001FFC00;
-	uint b = (f32tof16(rgb.b) >> 5)  & 0x000003FF;
-	return r | g | b;
-}
-
-float3 UnpackR11G11B10F(uint rgb)
-{
-	float r = f16tof32((rgb >> 17) & 0x7FF0);
-	float g = f16tof32((rgb >> 6)  & 0x7FF0);
-	float b = f16tof32((rgb << 5)  & 0x7FE0);
-	return float3(r, g, b);
-}
-
-uint PackR10G10B10F(float3 rgb)
-{
-	uint r = (f32tof16(rgb.r) << 15) & 0x3FF00000;  // 0011 1111 1111 0000 0000 0000 0000 0000 
-	uint g = (f32tof16(rgb.g) << 5)  & 0x000FFC00;  // 0000 0000 0000 1111 1111 1100 0000 0000 
-	uint b = (f32tof16(rgb.b) >> 5)  & 0x000003FF;	// 0000 0000 0000 0000 0000 0011 1111 1111
-	return r | g | b;
-}
-
-float3 UnpackR10G10B10F(uint rgb)
-{
-	float r = f16tof32((rgb >> 15) & 0x7FE0);
-	float g = f16tof32((rgb >> 5)  & 0x7FE0);
-	float b = f16tof32((rgb << 5)  & 0x7FE0);
-	return float3(r, g, b);
-}
-
-uint2 PackR16G16B16A16F(float4 In)
-{
-	return uint2(PackFloat2ToUInt(In.xy), PackFloat2ToUInt(In.zw));
-}
-
-float4 UnpackR16G16B16A16F(uint2 In)
-{
-	return float4(UnpackFloat2FromUInt(In.x), UnpackFloat2FromUInt(In.y));
-}
-
 uint PackColorLinearToGamma2(float3 In)
 {
 	return PackRGBA8(float4(StrataLinearToSrgb(saturate(In.rgb)), 0));
@@ -3094,37 +3036,6 @@ float4 UnpackColorGamma2ToLinearAlphaLinear(uint In)
 	return float4(StrataSrgbToLinear(Out.rgb), Out.a);
 }
 
-uint PackRGBA8(float4 In)
-{
-#if COMPILER_SUPPORTS_PACK_INTRINSICS
-	return PackFloat4ToUInt(saturate(In) * 255.0f);
-#else
-	uint r = (uint(saturate(In.r) * 255.0f) << 0);
-	uint g = (uint(saturate(In.g) * 255.0f) << 8);
-	uint b = (uint(saturate(In.b) * 255.0f) << 16);
-	uint a = (uint(saturate(In.a) * 255.0f) << 24);
-	return r | g | b | a;
-#endif
-}
-
-float4 UnpackRGBA8(uint In)
-{
-#if COMPILER_SUPPORTS_UNPACKBYTEN
-	return float4(
-		UnpackByte0(In),
-		UnpackByte1(In),
-		UnpackByte2(In),
-		UnpackByte3(In)) * (1.0f / 255.0f);
-#else
-	float4 Out;
-	Out.r = float((In >> 0)  & 0xFF) * (1.0f/255.0f);
-	Out.g = float((In >> 8)  & 0xFF) * (1.0f/255.0f);
-	Out.b = float((In >> 16) & 0xFF) * (1.0f/255.0f);
-	Out.a = float((In >> 24) & 0xFF) * (1.0f/255.0f);
-	return Out;
-#endif
-}
-
 uint PackR7G7B6Gamma2(float3 rgb, float Dither)
 {
 	// We execute a dither that can go in positive/negative direction to avoid surfaces to looks overall brighter.
