Fix some comments and cvar descriptions.

diff --git a/Engine/Shaders/Private/AssertionMacros.ush b/Engine/Shaders/Private/AssertionMacros.ush
--- a/Engine/Shaders/Private/AssertionMacros.ush
+++ b/Engine/Shaders/Private/AssertionMacros.ush
@@ -23,13 +23,13 @@
 #if !defined(check)
 	/**
 	 * Check macro, intended to detect serious, execution-disrupting (not continuable) errors in the shader code.
-	 * Is controlled by the read-only CVar `r.Shader.CheckLevel`:
+	 * Is controlled by the read-only CVar `r.Shaders.CheckLevel`:
 	 *		0           => DO_CHECK=0 and DO_GUARD_SLOW=0
 	 *		1 (default) => DO_CHECK=1 and DO_GUARD_SLOW=0
 	 *		2           => DO_CHECK=1 and DO_GUARD_SLOW=1
 	 * NOTE: As check() is enabled by default this will cause increased register pressure and have some overhead. Therefore, be careful when using this
 	 * and if in doubt, default to using checkSlow et al.
-	 * When not doing GPU performance profiling it may be a good idea to keep r.Shader.CheckLevel=2 to catch errors early.
+	 * When not doing GPU performance profiling it may be a good idea to keep r.Shaders.CheckLevel=2 to catch errors early.
 	 */
 	#define check(expr)	 \
 		{ \
diff --git a/Engine/Shaders/Private/ShaderPrintCommon.ush b/Engine/Shaders/Private/ShaderPrintCommon.ush
--- a/Engine/Shaders/Private/ShaderPrintCommon.ush
+++ b/Engine/Shaders/Private/ShaderPrintCommon.ush
@@ -11,7 +11,7 @@
 // 
 // 0. Contex-based state API
 // -------------------------
-// FShaderPrintContext Ctx = InitShaderPrintContext(all(SvPosition.xy == float2(100, 100), uint2(50,50))
+// FShaderPrintContext Ctx = InitShaderPrintContext(all(SvPosition.xy == float2(100, 100)), uint2(50,50));
 //                             ------------------------------------   ------------
 //                                 Pixel/thread to filter             Coord at which 
 //                                                                       to draw
@@ -20,7 +20,7 @@
 // ----------------------
 // First create a context, then print scalar/vector/matrix/text...
 // 
-//   FShaderPrintContext Ctx = InitShaderPrintContext(all(PixelPos.xy == uint2(100, 100), uint2(50,50))
+//   FShaderPrintContext Ctx = InitShaderPrintContext(all(PixelPos.xy == uint2(100, 100)), uint2(50,50));
 // 
 //   Print(Ctx, MyFloat);
 //   Newline(Ctx);
@@ -66,7 +66,7 @@
 //    void MainCS(uint2 DispatchThreadId : SV_DispatchThreadID)
 //    {
 //       // Create context which 'add' the widget only on thread (0,0)
-//       FShaderPrintContext Ctx = InitShaderPrintContext(all(DispatchThreadId == uint2(0, 0), uint2(50,50))
+//       FShaderPrintContext Ctx = InitShaderPrintContext(all(DispatchThreadId == uint2(0, 0)), uint2(50,50));
 // 
 //		 // Any threads can get the slider value with this call
 //       const float ScaleValue = AddSlider(Ctx, 1.f, TEXT("Scale"));
@@ -216,7 +216,7 @@ FFontColor GetDefaultFontColor()							{ FFontColor Out; Out.Color = float3(1,1,
 #define ColorDarkBlue	float4(0.25f, 0.25f, 0.50f, 1.0f)
 
 ///////////////////////////////////////////////////////////////////////////////////////////////////
-// Util pack/unpac functions
+// Util pack/unpack functions
 
 struct FShaderPrintItem
 {
diff --git a/Engine/Source/Runtime/Core/Private/HAL/ConsoleManager.cpp b/Engine/Source/Runtime/Core/Private/HAL/ConsoleManager.cpp
--- a/Engine/Source/Runtime/Core/Private/HAL/ConsoleManager.cpp
+++ b/Engine/Source/Runtime/Core/Private/HAL/ConsoleManager.cpp
@@ -2605,7 +2605,7 @@ static TAutoConsoleVariable<int32> CVarTranslucentSortPolicy(
 	TEXT("r.TranslucentSortPolicy"),
 	0,
 	TEXT("0: Sort based on distance from camera centerpoint to bounding sphere centerpoint. (default, best for 3D games)\n"
-		 "1: Sort based on projected distance to camera."
+		 "1: Sort based on projected distance to camera.\n"
 		 "2: Sort based on the projection onto a fixed axis. (best for 2D games)"),
 	ECVF_RenderThreadSafe);
 
@@ -2619,15 +2619,15 @@ static TAutoConsoleVariable<int32> CVarMobileHDR(
 static TAutoConsoleVariable<int32> CVarMobileShadingPath(
 	TEXT("r.Mobile.ShadingPath"),
 	0,
-	TEXT("0: Forward shading (default)"
+	TEXT("0: Forward shading (default)\n"
 		 "1: Deferred shading"),
 	ECVF_RenderThreadSafe | ECVF_ReadOnly);
 
 static TAutoConsoleVariable<int32> CVarMobileAllowDeferredShadingOpenGL(
 	TEXT("r.Mobile.AllowDeferredShadingOpenGL"),
 	0,
-	TEXT("0: Do not Allow Deferred Shading on OpenGL (default)"
-		"1: Allow Deferred Shading on OpenGL"),
+	TEXT("0: Do not Allow Deferred Shading on OpenGL (default)\n"
+		 "1: Allow Deferred Shading on OpenGL"),
 	ECVF_RenderThreadSafe | ECVF_ReadOnly);
 
 static TAutoConsoleVariable<int32> CVarMobileEnableStaticAndCSMShadowReceivers(
