Fix NaN with cloth shading and rect lighting.

diff --git a/Engine/Shaders/Private/RectLightIntegrate.ush b/Engine/Shaders/Private/RectLightIntegrate.ush
--- a/Engine/Shaders/Private/RectLightIntegrate.ush
+++ b/Engine/Shaders/Private/RectLightIntegrate.ush
@@ -259,9 +259,7 @@ FDirectLighting IntegrateBxDF(FGBufferData GBuffer, half3 N, half3 V, FRect Rect
 	if (IsRectVisible(Rect))
 	{
 		FAreaLightIntegrateContext Context = CreateRectIntegrateContext(GBuffer.Roughness, N, V, Rect, SourceTexture);
-	#if CHEAP_RECT
 		GBuffer.Roughness = max(GBuffer.Roughness, 0.02);
-	#endif
 		Out = IntegrateBxDF(GBuffer, N, V, Context.L, Context.Falloff, Context.NoL, Context.AreaLight, Shadow);
 	}
 	return Out;
