Path Tracer: Fix incorrect pdfs returned by sample method

When refactoring the code for splitting diffuse from specular, a subtle error was introduced that could cause materials to render incorrectly under certain illumination conditions. This could only be seen when comparing the direct illumination output under MISMode 0,1,2. While modes 0 and 1 (corresponding to BSDF and Light sampling respetively) matched in all cases, the MIS combination of the two did not and frequently rendered darker.

The reason is that the eval function uses MIS across all lobes, while sample only returns the result of a single lobe. The weight would lead to the correct result, but the pdf did not compare correctly with the one returned by eval, leading to incorrect results when MIS favored BSDF sampling over light sampling.

To fix this, all multi-lobed shading models need to account for the probability of generating the same direction through all other lobes.

Along the way, change to extinction encoding for solid glass to avoid precision issues with transmission color in solidglass.

Also apply the RescaleRandomNumber function consistently across material sampling routines.

diff --git a/Engine/Shaders/Private/MonteCarlo.ush b/Engine/Shaders/Private/MonteCarlo.ush
--- a/Engine/Shaders/Private/MonteCarlo.ush
+++ b/Engine/Shaders/Private/MonteCarlo.ush
@@ -523,6 +523,15 @@ float MISWeightRobust(float Pdf, float OtherPdf) {
 	}
 }
 
+// When sampling from discrete CDFs, it can be convenient to re-use the random number by rescaling it
+// This function assumes that RandVal is in the interval: [LowerBound, UpperBound) and returns a value in [0,1)
+float RescaleRandomNumber(float RandVal, float LowerBound, float UpperBound)
+{
+	const float OneMinusEpsilon = 0.99999994; // 32-bit float just before 1.0
+	return min((RandVal - LowerBound) / (UpperBound - LowerBound), OneMinusEpsilon);
+}
+
+
 // Compute PDF of reflection (2 * dot( V, H ) * H - V).
 // [Heitz 2014, "Importance Sampling Microfacet-Based BSDFs using the Distribution of Visible Normals"]
 float RayPDFToReflectionRayPDF(float VoH, float RayPDF)
diff --git a/Engine/Shaders/Private/PathTracing/Material/PathTracingClearCoat.ush b/Engine/Shaders/Private/PathTracing/Material/PathTracingClearCoat.ush
--- a/Engine/Shaders/Private/PathTracing/Material/PathTracingClearCoat.ush
+++ b/Engine/Shaders/Private/PathTracing/Material/PathTracingClearCoat.ush
@@ -174,12 +174,11 @@ FMaterialSample ClearCoat_SampleMaterial(
 
 		if (bSampledCoat)
 		{
-			RandSample.x /= Data.LobeCdf.x;
+			RandSample.x = RescaleRandomNumber(RandSample.x, 0.0, Data.LobeCdf.x);
 		}
 		else
 		{
-			RandSample.x -= Data.LobeCdf.x;
-			RandSample.x /= Data.LobeCdf.y - Data.LobeCdf.x;
+			RandSample.x = RescaleRandomNumber(RandSample.x, Data.LobeCdf.x, Data.LobeCdf.y);
 		}
 		const float3 V = bSampledCoat ? Data.CoatV : Data.BaseV;
 		const float3 H = ImportanceSampleVisibleGGX_aniso(UniformSampleDisk(RandSample.xy), bSampledCoat ? Data.CoatAlpha : Data.BaseAlpha, V).xyz;
@@ -198,8 +197,7 @@ FMaterialSample ClearCoat_SampleMaterial(
 	else
 	{
 		// Sampled base diffuse
-		RandSample.x -= Data.LobeCdf.y;
-		RandSample.x /= 1.0 - Data.LobeCdf.y;
+		RandSample.x = RescaleRandomNumber(RandSample.x, Data.LobeCdf.y, 1.0);
 
 		const float3 L = CosineSampleHemisphere(RandSample.xy).xyz;
 		Result.Direction = mul(L, Data.BaseBasis);
@@ -227,30 +225,34 @@ FMaterialSample ClearCoat_SampleMaterial(
 	// Calculate transmission through substrate
 	const float3 Transmission = (1.0 - Data.CoatE) * lerp(1.0, SimpleClearCoatTransmittance(CoatNoL, CoatNoV, Payload.Metallic, Payload.BaseColor), Data.ClearCoat);
 
+	// compute coat and spec evals so we can re-use Pdfs across both branches
+	const float2 CoatGGXResult = GGXEvalReflection(CoatL, Data.CoatV, CoatH, Data.CoatAlpha);
+	const float2 SpecGGXResult = GGXEvalReflection(BaseL, Data.BaseV, BaseH, Data.BaseAlpha);
+	const float CoatPdf = CoatGGXResult.y;
+	const float SpecPdf = SpecGGXResult.y;
+	const float DiffPdf = BaseNoL / PI;
+
 	// evaluate the lobe we sampled
 	if (bSampledSpecular)
 	{
-		// Coat Specular lobe
-		const float2 CoatGGXResult = GGXEvalReflection(CoatL, Data.CoatV, CoatH, Data.CoatAlpha);
 		const float3 CoatWeight = Data.ClearCoat * ClearCoatFresnel(CoatVoH) * CoatGGXResult.x * Data.Coat.W;
-		const float CoatPdf = CoatGGXResult.y;
-
-		// Base Specular lobe
-		const float2 SpecGGXResult = GGXEvalReflection(BaseL, Data.BaseV, BaseH, Data.BaseAlpha);
 		const float3 SpecWeight = Transmission * F_Schlick(Data.SpecColor, BaseVoH) * SpecGGXResult.x * Data.Spec.W;
-		const float SpecPdf = SpecGGXResult.y;
 
 		Result.AddLobeWithMIS(CoatWeight, CoatPdf, Data.LobePdf.x);
 		Result.AddLobeWithMIS(SpecWeight, SpecPdf, Data.LobePdf.y);
+
+		Result.Pdf += Data.LobePdf.z * DiffPdf;
 	}
 	else
 	{
 		// Base Diffuse lobe
 		const float3 Diffuse = GetPathTracingDiffuseModel(Data.DiffColor, Payload.Roughness, BaseNoV, BaseNoL, BaseVoH, BaseNoH);
 		const float3 DiffWeight = Transmission * (1 - Data.Spec.E) * Diffuse * ShadowTerminatorTerm(L_World, Payload.WorldNormal, Payload.WorldSmoothNormal);
-		const float DiffPdf = BaseNoL / PI;
 
 		Result.AddLobeWithMIS(DiffWeight, DiffPdf, Data.LobePdf.z);
+
+		Result.Pdf += Data.LobePdf.x * CoatPdf;
+		Result.Pdf += Data.LobePdf.y * SpecPdf;
 	}
 
 
diff --git a/Engine/Shaders/Private/PathTracing/Material/PathTracingCloth.ush b/Engine/Shaders/Private/PathTracing/Material/PathTracingCloth.ush
--- a/Engine/Shaders/Private/PathTracing/Material/PathTracingCloth.ush
+++ b/Engine/Shaders/Private/PathTracing/Material/PathTracingCloth.ush
@@ -109,19 +109,18 @@ FMaterialSample Cloth_SampleMaterial(
 	// The cloth lobe is low enough frequency that sampling with cosine weighting is sufficient
 
 	float3 L = 0.0, H = 0.0;
+	float OutRoughness = 1.0;
 	const bool bSampledDiffuse = RandSample.x < Data.LobeCdf.x;
-	bool bSampledCloth = false;
-	if (RandSample.x < Data.LobeCdf.y)
+	const bool bSampledCloth   = RandSample.x < Data.LobeCdf.y;
+	if (bSampledDiffuse || bSampledCloth)
 	{
-		if (RandSample.x < Data.LobeCdf.x)
+		if (bSampledDiffuse)
 		{
-			RandSample.x /= Data.LobeCdf.x;
+			RandSample.x = RescaleRandomNumber(RandSample.x, 0.0, Data.LobeCdf.x);
 		}
 		else
 		{
-			RandSample.x -= Data.LobeCdf.x;
-			RandSample.x /= Data.LobeCdf.y - Data.LobeCdf.x;
-			bSampledCloth = true;
+			RandSample.x = RescaleRandomNumber(RandSample.x, Data.LobeCdf.x, Data.LobeCdf.y);
 		}
 
 		// Lambert
@@ -130,8 +129,7 @@ FMaterialSample Cloth_SampleMaterial(
 	}
 	else
 	{
-		RandSample.x -= Data.LobeCdf.y;
-		RandSample.x /= 1.0 - Data.LobeCdf.y;
+		RandSample.x = RescaleRandomNumber(RandSample.x, Data.LobeCdf.y, 1.0);
 
 		// #dxr_todo: Evaluate UniformSampleDiskConcentric
 		H = ImportanceSampleVisibleGGX_aniso(UniformSampleDisk(RandSample.xy), Data.Alpha, Data.V).xyz;
@@ -141,6 +139,7 @@ FMaterialSample Cloth_SampleMaterial(
 			// invalid output direction, exit early
 			return NullMaterialSample();
 		}
+		OutRoughness = Payload.Roughness;
 	}
 
 	// With a valid direction in hand -- now evaluate the BxDF (taking advantage of already computed terms)
@@ -153,33 +152,41 @@ FMaterialSample Cloth_SampleMaterial(
 	const float NoH = saturate(H.z);
 
 	const float ClothBlend = Payload.GetClothAmount();
+	const float2 GGXResult = GGXEvalReflection(L, V, H, Data.Alpha);
+	const float SpecPdf = GGXResult.y;
+	const float DiffPdf = NoL / PI;
+	const float ClothPdf = DiffPdf;
 
-	FMaterialSample Result = CreateMaterialSample(L_World, 0.0, 0.0, 1.0, 1.0, PATHTRACER_SCATTER_DIFFUSE);
+
+	FMaterialSample Result = CreateMaterialSample(L_World, 0.0, 0.0, 1.0, OutRoughness, PATHTRACER_SCATTER_DIFFUSE);
 
 	const float ShadowTerminator = ShadowTerminatorTerm(L_World, Payload.WorldNormal, Payload.WorldSmoothNormal);
 
 	if (bSampledDiffuse)
 	{
-		Result.AddLobeWithMIS(Data.DiffWeight * ShadowTerminator, NoL / PI, Data.LobePdf.x);
+		Result.AddLobeWithMIS(Data.DiffWeight * ShadowTerminator, DiffPdf, Data.LobePdf.x);
+
+		// Cloth and Spec PDF
+		Result.Pdf += Data.LobePdf.y * ClothPdf;
+		Result.Pdf += Data.LobePdf.z * SpecPdf;
 	}
-	else if (bSampledCloth)
+	else
 	{
+		Result.ScatterType = PATHTRACER_SCATTER_SPECULAR;
+
+		// Cloth Lobe
 		const float DCloth = D_InvGGX(Pow4(Payload.Roughness), NoH);
 		const float3 FCloth = F_Schlick(Payload.GetClothColor(), VoH);
 		const float3 ClothWeight = NoL > 0 && NoV > 0 ? (PI * ClothBlend * DCloth * Vis_Cloth(NoV, NoL)) * FCloth * ComputeEnergyConservation(Data.Fuzz) : 0.0;
-		Result.AddLobeWithMIS(ClothWeight * ShadowTerminator, NoL / PI, Data.LobePdf.y);
-		Result.ScatterType = PATHTRACER_SCATTER_SPECULAR;
-	}
-	else
-	{
+		Result.AddLobeWithMIS(ClothWeight * ShadowTerminator, ClothPdf, Data.LobePdf.y);
+
 		// Specular lobe
-		const float2 GGXResult = GGXEvalReflection(L, V, H, Data.Alpha);
 		const float3 F = F_Schlick(Payload.SpecularColor, VoH);
 		const float3 SpecWeight = (1 - ClothBlend) * F * GGXResult.x * Data.Spec.W;
-		const float SpecPdf = GGXResult.y;
 		Result.AddLobeWithMIS(SpecWeight, SpecPdf, Data.LobePdf.z);
-		Result.Roughness = Payload.Roughness;
-		Result.ScatterType = PATHTRACER_SCATTER_SPECULAR;
+
+		// Diffuse PDF
+		Result.Pdf += Data.LobePdf.x * DiffPdf;
 	}
 	Result.Weight *= Payload.BSDFOpacity;
 	return Result;
diff --git a/Engine/Shaders/Private/PathTracing/Material/PathTracingDefaultLit.ush b/Engine/Shaders/Private/PathTracing/Material/PathTracingDefaultLit.ush
--- a/Engine/Shaders/Private/PathTracing/Material/PathTracingDefaultLit.ush
+++ b/Engine/Shaders/Private/PathTracing/Material/PathTracingDefaultLit.ush
@@ -89,7 +89,7 @@ FMaterialSample DefaultLit_SampleMaterial(
 	const bool bSampledDiffuse = RandSample.x < Data.LobeProb;
 	if (bSampledDiffuse)
 	{
-		RandSample.x /= Data.LobeProb;
+		RandSample.x = RescaleRandomNumber(RandSample.x, 0.0, Data.LobeProb);
 
 		// Lambert
 		// TODO: evaluate CosineSampleHemisphereConcentric
@@ -98,8 +98,7 @@ FMaterialSample DefaultLit_SampleMaterial(
 	}
 	else
 	{
-		RandSample.x -= Data.LobeProb;
-		RandSample.x /= 1.0 - Data.LobeProb;
+		RandSample.x = RescaleRandomNumber(RandSample.x, Data.LobeProb, 1.0);
 
 		// #dxr_todo: Evaluate UniformSampleDiskConcentric
 		H = ImportanceSampleVisibleGGX_aniso(UniformSampleDisk(RandSample.xy), Data.Alpha, Data.V).xyz;
@@ -119,20 +118,24 @@ FMaterialSample DefaultLit_SampleMaterial(
 	const float VoH = saturate(dot(V, H));
 	const float NoH = saturate(H.z);
 
+	const float2 GGXResult = GGXEvalReflection(L, V, H, Data.Alpha);
+	const float SpecPdf = GGXResult.y;
+	const float DiffPdf = NoL / PI;
+
 	FMaterialSample Result = CreateMaterialSample(L_World, 0.0, 0.0, 1.0, 1.0, PATHTRACER_SCATTER_DIFFUSE);
 	if (bSampledDiffuse)
 	{
 		const float3 Diffuse = GetPathTracingDiffuseModel(Data.DiffuseWeight, Payload.Roughness, NoV, NoL, VoH, NoH);
-		Result.AddLobeWithMIS(Diffuse * ShadowTerminatorTerm(L_World, Payload.WorldNormal, Payload.WorldSmoothNormal), NoL / PI, Data.LobeProb);
+		Result.AddLobeWithMIS(Diffuse * ShadowTerminatorTerm(L_World, Payload.WorldNormal, Payload.WorldSmoothNormal), DiffPdf, Data.LobeProb);
+		Result.Pdf += (1.0 - Data.LobeProb) * SpecPdf;
 	}
 	else
 	{
-		const float2 GGXResult = GGXEvalReflection(L, V, H, Data.Alpha);
 		const float3 F = F_Schlick(Payload.SpecularColor, VoH);
 		const float3 SpecWeight = F * GGXResult.x * Data.Spec.W;
-		const float SpecPdf = GGXResult.y;
 
 		Result.AddLobeWithMIS(SpecWeight, SpecPdf, 1.0 - Data.LobeProb);
+		Result.Pdf += Data.LobeProb * DiffPdf;
 		Result.Roughness = Payload.Roughness;
 		Result.ScatterType = PATHTRACER_SCATTER_SPECULAR;
 	}
diff --git a/Engine/Shaders/Private/PathTracing/Material/PathTracingEye.ush b/Engine/Shaders/Private/PathTracing/Material/PathTracingEye.ush
--- a/Engine/Shaders/Private/PathTracing/Material/PathTracingEye.ush
+++ b/Engine/Shaders/Private/PathTracing/Material/PathTracingEye.ush
@@ -70,9 +70,10 @@ FMaterialEval Eye_EvalMaterial(
 	const float Iris = IrisNoL * Caustic;
 	const float Sclera = NoL;
 	const float EyeDiffuseTweak = 2 * lerp(Sclera, Iris, IrisMask);
+	const float DiffPdf = 1 / (2 * PI);
 
 	// Diffuse Lobe
-	Result.AddLobeWithMIS(Data.DiffWeight * EyeDiffuseTweak * DiffuseSpecularScale.x, 1 / (2 * PI), Data.LobeProb);
+	Result.AddLobeWithMIS(Data.DiffWeight * EyeDiffuseTweak * DiffuseSpecularScale.x, DiffPdf, Data.LobeProb);
 
 	// Specular lobe
 	const float2 GGXResult = GGXEvalReflection(L, V, H, Data.Alpha);
@@ -103,15 +104,14 @@ FMaterialSample Eye_SampleMaterial(
 	const bool bSampledDiffuse = RandSample.x < Data.LobeProb;
 	if (bSampledDiffuse)
 	{
-		RandSample.x /= Data.LobeProb;
+		RandSample.x = RescaleRandomNumber(RandSample.x, 0.0, Data.LobeProb);
 
 		L = UniformSampleHemisphere(RandSample.xy).xyz;
 		H = normalize(L + V);
 	}
 	else
 	{
-		RandSample.x -= Data.LobeProb;
-		RandSample.x /= 1.0 - Data.LobeProb;
+		RandSample.x = RescaleRandomNumber(RandSample.x, Data.LobeProb, 1.0);
 
 		H = ImportanceSampleVisibleGGX_aniso(UniformSampleDisk(RandSample.xy), Data.Alpha, V).xyz;
 		L = reflect(-V, H);
@@ -130,6 +130,10 @@ FMaterialSample Eye_SampleMaterial(
 	const float NoL = saturate(L.z);
 	const float VoH = saturate(dot(V, H));
 
+	const float2 GGXResult = GGXEvalReflection(L, V, H, Data.Alpha);
+	const float SpecPdf = GGXResult.y;
+	const float DiffPdf = 1 / (2 * PI);
+
 	FMaterialSample Result = CreateMaterialSample(L_World, 0.0, 0.0, 1.0, 1.0, PATHTRACER_SCATTER_DIFFUSE);
 	if (bSampledDiffuse)
 	{
@@ -145,16 +149,18 @@ FMaterialSample Eye_SampleMaterial(
 		const float EyeDiffuseTweak = 2 * lerp(Sclera, Iris, IrisMask);
 
 		// Diffuse Lobe
-		Result.AddLobeWithMIS(Data.DiffWeight * EyeDiffuseTweak, 1 / (2 * PI), Data.LobeProb);
+		Result.AddLobeWithMIS(Data.DiffWeight * EyeDiffuseTweak, DiffPdf, Data.LobeProb);
+
+		Result.Pdf += (1 - Data.LobeProb) * SpecPdf;
 	}
 	else
 	{
 		// Specular lobe
-		const float2 GGXResult = GGXEvalReflection(L, V, H, Data.Alpha);
 		const float3 F = F_Schlick(Payload.SpecularColor, VoH);
 		const float3 SpecWeight = F * GGXResult.x * Data.Spec.W;
-		const float SpecPdf = GGXResult.y;
+		
 		Result.AddLobeWithMIS(SpecWeight, SpecPdf, 1.0 - Data.LobeProb);
+		Result.Pdf += Data.LobeProb * DiffPdf;
 		Result.Roughness = Payload.Roughness;
 		Result.ScatterType = PATHTRACER_SCATTER_SPECULAR;
 	}
diff --git a/Engine/Shaders/Private/PathTracing/Material/PathTracingLambert.ush b/Engine/Shaders/Private/PathTracing/Material/PathTracingLambert.ush
--- a/Engine/Shaders/Private/PathTracing/Material/PathTracingLambert.ush
+++ b/Engine/Shaders/Private/PathTracing/Material/PathTracingLambert.ush
@@ -30,6 +30,7 @@ FMaterialEval Lambert_EvalMaterial(
 	return CreateMaterialEval(Payload.BSDFOpacity * Payload.BaseColor * ShadowTerminatorTerm(L_World, Payload.WorldNormal, Payload.WorldSmoothNormal) * DiffuseSpecularScale.x, NoL / PI);
 }
 
+#if SIMPLIFIED_MATERIAL_SHADER
 
 FMaterialSample Simplified_SampleMaterial(
 	FPathTracingPayload Payload,
@@ -73,3 +74,5 @@ FMaterialEval Simplified_EvalMaterial(
 	Weight *= Payload.BSDFOpacity;
 	return CreateMaterialEval(Weight, Pdf * saturate(abs(NoL)) / PI);
 }
+
+#endif // SIMPLIFIED_MATERIAL_SHADER
diff --git a/Engine/Shaders/Private/PathTracing/Material/PathTracingMaterialSampling.ush b/Engine/Shaders/Private/PathTracing/Material/PathTracingMaterialSampling.ush
--- a/Engine/Shaders/Private/PathTracing/Material/PathTracingMaterialSampling.ush
+++ b/Engine/Shaders/Private/PathTracing/Material/PathTracingMaterialSampling.ush
@@ -9,6 +9,7 @@ PathTracingMaterialSampling.usf: Material sampling functions
 #include "PathTracingMaterialCommon.ush"
 #include "PathTracingFresnel.ush"
 #include "PathTracingLambert.ush"
+#if !SIMPLIFIED_MATERIAL_SHADER
 #include "PathTracingGlossy.ush"
 #include "PathTracingSolidGlass.ush"
 #include "PathTracingThinGlass.ush"
@@ -20,12 +21,7 @@ PathTracingMaterialSampling.usf: Material sampling functions
 #include "PathTracingEye.ush"
 #include "PathTracingHair.ush"
 #include "PathTracingMedium.ush"
-
-// Simplified material shaders for GPU Lightmass, supporting Lambert, Two Sided Foliage and Thin Translucent
-// Make sure material evaluation code below matches with hit shader which fills the payload
-#ifndef SIMPLIFIED_MATERIAL_SHADER
-	#define SIMPLIFIED_MATERIAL_SHADER 0
-#endif
+#endif // !SIMPLIFIED_MATERIAL_SHADER
 
 float3 EstimateMaterialAlbedo(FPathTracingPayload Payload)
 {
diff --git a/Engine/Shaders/Private/PathTracing/Material/PathTracingSolidGlass.ush b/Engine/Shaders/Private/PathTracing/Material/PathTracingSolidGlass.ush
--- a/Engine/Shaders/Private/PathTracing/Material/PathTracingSolidGlass.ush
+++ b/Engine/Shaders/Private/PathTracing/Material/PathTracingSolidGlass.ush
@@ -75,7 +75,7 @@ FMaterialSample RoughGlass_SampleMaterial(
 	// Randomly choose to sample diffuse or specular
 	if (RandSample.x < Data.LobeCdf.x)
 	{
-		RandSample.x /= Data.LobeCdf.x;
+		RandSample.x = RescaleRandomNumber(RandSample.x, 0.0, Data.LobeCdf.x);
 
 		// Lambert
 		// TODO: evaluate CosineSampleHemisphereConcentric
@@ -103,13 +103,11 @@ FMaterialSample RoughGlass_SampleMaterial(
 
 		if (bDoSpec)
 		{
-			RandSample.x -= Data.LobeCdf.x;
-			RandSample.x /= Data.LobeCdf.y - Data.LobeCdf.x;
+			RandSample.x = RescaleRandomNumber(RandSample.x, Data.LobeCdf.x, Data.LobeCdf.y);
 		}
 		else
 		{
-			RandSample.x -= Data.LobeCdf.y;
-			RandSample.x /= 1.0 - Data.LobeCdf.y;
+			RandSample.x = RescaleRandomNumber(RandSample.x, Data.LobeCdf.y, 1.0);
 		}
 
 		// #dxr_todo: Evaluate UniformSampleDiskConcentric
@@ -151,7 +149,7 @@ FMaterialSample RoughGlass_SampleMaterial(
 
 			const float2 GGXResult = bRefract ? GGXEvalRefraction(L, V, H, Data.Alpha, Data.Eta) : GGXEvalReflection(L, V, H, Data.Alpha);
 
-			return CreateMaterialSample(L_World, (1.0 - Payload.BSDFOpacity) * GGXResult.x / Data.EGlass, F * GGXResult.y, bRefract ? -1.0 : 1.0, Payload.Roughness, bRefract ? PATHTRACER_SCATTER_REFRACT : PATHTRACER_SCATTER_SPECULAR);
+			return CreateMaterialSample(L_World, (1.0 - Payload.BSDFOpacity) * GGXResult.x / (Data.EGlass * Data.LobePdf.z), F * GGXResult.y * Data.LobePdf.z, bRefract ? -1.0 : 1.0, Payload.Roughness, bRefract ? PATHTRACER_SCATTER_REFRACT : PATHTRACER_SCATTER_SPECULAR);
 		}
 	}
 }
diff --git a/Engine/Shaders/Private/PathTracing/Material/PathTracingSubsurfaceProfile.ush b/Engine/Shaders/Private/PathTracing/Material/PathTracingSubsurfaceProfile.ush
--- a/Engine/Shaders/Private/PathTracing/Material/PathTracingSubsurfaceProfile.ush
+++ b/Engine/Shaders/Private/PathTracing/Material/PathTracingSubsurfaceProfile.ush
@@ -97,7 +97,7 @@ FMaterialSample SubsurfaceProfile_SampleMaterial(
 	float OutRoughness = 0;
 	if (bSampledDiffuse)
 	{
-		RandSample.x /= Data.LobeCdf.x;
+		RandSample.x = RescaleRandomNumber(RandSample.x, 0.0, Data.LobeCdf.x);
 
 		L = CosineSampleHemisphere(RandSample.xy).xyz;
 		H = normalize(L + V);
@@ -108,14 +108,12 @@ FMaterialSample SubsurfaceProfile_SampleMaterial(
 		const bool bUseSpec0 = RandSample.x < Data.LobeCdf.y;
 		if (bUseSpec0)
 		{
-			RandSample.x -= Data.LobeCdf.x;
-			RandSample.x /= Data.LobeCdf.y - Data.LobeCdf.x;
+			RandSample.x = RescaleRandomNumber(RandSample.x, Data.LobeCdf.x, Data.LobeCdf.y);
 			OutRoughness = DualRoughnessData.x;
 		}
 		else
 		{
-			RandSample.x -= Data.LobeCdf.y;
-			RandSample.x /= 1.0 - Data.LobeCdf.y;
+			RandSample.x = RescaleRandomNumber(RandSample.x, Data.LobeCdf.y, 1.0);
 			OutRoughness = DualRoughnessData.y;
 		}
 
@@ -133,22 +131,27 @@ FMaterialSample SubsurfaceProfile_SampleMaterial(
 	// transform to world space
 	const float3 L_World = normalize(mul(L, Data.Basis));
 
+	const float2 GGXResult0 = GGXEvalReflection(L, V, H, Data.Alpha0);
+	const float2 GGXResult1 = GGXEvalReflection(L, V, H, Data.Alpha1);
+	const float NoL = L.z;
+	const float DiffPdf = NoL / PI;
+
 	FMaterialSample Result = CreateMaterialSample(L_World, 0.0, 0.0, 1.0, OutRoughness, PATHTRACER_SCATTER_DIFFUSE);
 	if (bSampledDiffuse)
 	{
 		// Diffuse Lobe
-		const float NoL = L.z;
-		Result.AddLobeWithMIS(Data.DiffuseWeight * ShadowTerminatorTerm(L_World, Payload.WorldNormal, Payload.WorldSmoothNormal), NoL / PI, Data.LobePdf.x);
+		Result.AddLobeWithMIS(Data.DiffuseWeight * ShadowTerminatorTerm(L_World, Payload.WorldNormal, Payload.WorldSmoothNormal), DiffPdf, Data.LobePdf.x);
+		Result.Pdf += Data.LobePdf.y * GGXResult0.y;
+		Result.Pdf += Data.LobePdf.z * GGXResult1.y;
 	}
 	else 
 	{
 		const float VoH = saturate(dot(V, H));
 		// Specular lobes
-		const float2 GGXResult0 = GGXEvalReflection(L, V, H, Data.Alpha0);
-		const float2 GGXResult1 = GGXEvalReflection(L, V, H, Data.Alpha1);
 		const float3 F = F_Schlick(Payload.SpecularColor, VoH);
 		Result.AddLobeWithMIS((1 - DualRoughnessData.z) * F * GGXResult0.x * Data.Spec0.W, GGXResult0.y, Data.LobePdf.y);
 		Result.AddLobeWithMIS((    DualRoughnessData.z) * F * GGXResult1.x * Data.Spec1.W, GGXResult1.y, Data.LobePdf.z);
+		Result.Pdf += Data.LobePdf.z * DiffPdf;
 		Result.ScatterType = PATHTRACER_SCATTER_SPECULAR;
 	}
 	return Result;
diff --git a/Engine/Shaders/Private/PathTracing/Material/PathTracingThinGlass.ush b/Engine/Shaders/Private/PathTracing/Material/PathTracingThinGlass.ush
--- a/Engine/Shaders/Private/PathTracing/Material/PathTracingThinGlass.ush
+++ b/Engine/Shaders/Private/PathTracing/Material/PathTracingThinGlass.ush
@@ -191,6 +191,7 @@ FMaterialSample RoughThinGlass_SampleMaterial(
 
 	if (RandSample.x < Data.LobeCdf.x)
 	{
+		RandSample.x = RescaleRandomNumber(RandSample.x, 0.0, Data.LobeCdf.x);
 		// Diffuse lobe
 		const float3 L = CosineSampleHemisphere(RandSample.xy).xyz;
 		const float3 H = normalize(L + V);
@@ -209,6 +210,10 @@ FMaterialSample RoughThinGlass_SampleMaterial(
 
 		FMaterialSample Result = CreateMaterialSample(L_World, 0.0, 0.0, 1.0, 1.0, PATHTRACER_SCATTER_DIFFUSE);
 		Result.AddLobeWithMIS(DiffWeight, DiffPdf, Data.LobePdf.x);
+
+		const float SpecPdf = GGXEvalReflection(L, V, H, Data.AlphaR).y;
+		Result.Pdf += Data.LobePdf.y * SpecPdf;
+		Result.Pdf += Data.LobePdf.z * SpecPdf;
 		return Result;
 	}
 	else
@@ -217,18 +222,15 @@ FMaterialSample RoughThinGlass_SampleMaterial(
 		const bool bIsReflection = RandSample.x < Data.LobeCdf.z;
 		if (RandSample.x < Data.LobeCdf.y)
 		{
-			RandSample.x -= Data.LobeCdf.x;
-			RandSample.x /= Data.LobeCdf.y - Data.LobeCdf.x;
+			RandSample.x = RescaleRandomNumber(RandSample.x, Data.LobeCdf.x, Data.LobeCdf.y);
 		}
 		else if (RandSample.x < Data.LobeCdf.z)
 		{
-			RandSample.x -= Data.LobeCdf.y;
-			RandSample.x /= Data.LobeCdf.z - Data.LobeCdf.y;
+			RandSample.x = RescaleRandomNumber(RandSample.x, Data.LobeCdf.y, Data.LobeCdf.z);
 		}
 		else
 		{
-			RandSample.x -= Data.LobeCdf.z;
-			RandSample.x /= 1.0 - Data.LobeCdf.z;
+			RandSample.x = RescaleRandomNumber(RandSample.x, Data.LobeCdf.z, 1.0);
 		}
 		const float3 H = ImportanceSampleVisibleGGX_aniso(UniformSampleDisk(RandSample.xy), bIsReflection ? Data.AlphaR : Data.AlphaT, V).xyz;
 		const float NoH = saturate(H.z);
@@ -263,6 +265,8 @@ FMaterialSample RoughThinGlass_SampleMaterial(
 			const float GlassPdf = GGXResult.y;
 			Result.AddLobeWithMIS(GlassWeight, GlassPdf, Data.LobePdf.z);
 
+			Result.Pdf += Data.LobePdf.x * NoL / PI;
+
 			Result.Direction = mul(L, Data.Basis);
 			Result.PositionBiasSign = 1.0;
 			Result.ScatterType = PATHTRACER_SCATTER_SPECULAR;
diff --git a/Engine/Shaders/Private/PathTracing/Material/PathTracingTwoSidedFoliage.ush b/Engine/Shaders/Private/PathTracing/Material/PathTracingTwoSidedFoliage.ush
--- a/Engine/Shaders/Private/PathTracing/Material/PathTracingTwoSidedFoliage.ush
+++ b/Engine/Shaders/Private/PathTracing/Material/PathTracingTwoSidedFoliage.ush
@@ -41,7 +41,7 @@ FMaterialSample TwoSidedFoliage_SampleMaterial(
 	const FTwoSidedFoliageData Data = PrepareTwoSidedFoliageData(Payload, V_World);
 	if (RandSample.x < Data.ProbT)
 	{
-		RandSample.x /= Data.ProbT;
+		RandSample.x = RescaleRandomNumber(RandSample.x, 0.0, Data.ProbT);
 
 		const float3 N = Payload.WorldNormal;
 		const float4 SampledValue = CosineSampleHemisphere(RandSample.yz);
@@ -50,8 +50,7 @@ FMaterialSample TwoSidedFoliage_SampleMaterial(
 	}
 	else
 	{
-		RandSample.x -= Data.ProbT;
-		RandSample.x /= 1.0 - Data.ProbT;
+		RandSample.x = RescaleRandomNumber(RandSample.x, Data.ProbT, 1.0);
 
 		FMaterialSample Result = DefaultLit_SampleMaterial(V_World, Payload, RandSample);
 		Result.Weight /= 1.0 - Data.ProbT;
diff --git a/Engine/Shaders/Private/PathTracing/PathTracing.usf b/Engine/Shaders/Private/PathTracing/PathTracing.usf
--- a/Engine/Shaders/Private/PathTracing/PathTracing.usf
+++ b/Engine/Shaders/Private/PathTracing/PathTracing.usf
@@ -65,9 +65,9 @@ FPathState LoadPathStateData(uint Index)
 	Output.SigmaT.z = f16tof32(Data.PackedRoughnessSigma.y >> 16);
 	// extract scatter type from sign bits of SigmaT
 	Output.FirstScatterType = 0;
-	Output.FirstScatterType |= asuint(Output.SigmaT.x) >> 31;
-	Output.FirstScatterType |= asuint(Output.SigmaT.y) >> 30;
-	Output.FirstScatterType |= asuint(Output.SigmaT.z) >> 29;
+	Output.FirstScatterType |= (asuint(Output.SigmaT.x) >> 31) << 0;
+	Output.FirstScatterType |= (asuint(Output.SigmaT.y) >> 31) << 1;
+	Output.FirstScatterType |= (asuint(Output.SigmaT.z) >> 31) << 2;
 	// remove sign bits from SigmaT (must be a positive value)
 	Output.SigmaT = abs(Output.SigmaT);
 	return Output;
diff --git a/Engine/Shaders/Private/PathTracing/PathTracingCommon.ush b/Engine/Shaders/Private/PathTracing/PathTracingCommon.ush
--- a/Engine/Shaders/Private/PathTracing/PathTracingCommon.ush
+++ b/Engine/Shaders/Private/PathTracing/PathTracingCommon.ush
@@ -43,13 +43,6 @@ float3 PathTracingGlassTransmittanceToExtinction(float3 TransmittanceColor)
 	return -log(max(TransmittanceColor, 1e-8f)) * PATH_TRACING_ABSORPTION_SCALE;
 }
 
-float3 PathTracingGlassExtinctionToTransmittance(float3 ExtinctionCoefficient)
-{
-	// inverse of function above
-	return exp(-max(ExtinctionCoefficient, 0.0) / PATH_TRACING_ABSORPTION_SCALE);
-}
-
-
 uint PayloadEncodeUnitVector(float3 Normal)
 {
 	const int BITS = PATH_TRACING_PAYLOAD_NORMAL_BITS;
@@ -258,7 +251,7 @@ struct FPackedPathTracingPayload : FMinimalPayload
 		PackedData[13] = f32tof16(DBufferC.z);
 		PackedData[13] |= f32tof16(DBufferC.w) << 16;
 	}
-#endif
+#endif // IS_PAYLOAD_ENABLED(RT_PAYLOAD_TYPE_GPULIGHTMASS)
 };
 
 #if IS_PAYLOAD_ENABLED(RT_PAYLOAD_PATH_TRACING_MATERIAL)
@@ -374,10 +367,14 @@ struct FPathTracingPayload : FMinimalPayload
 	float2 GetHairPrimitiveUV() { return CustomData0.xy; }
 	void SetHairPrimitiveUV(float2 UV) { CustomData0.xy = UV; }
 
-	// SHADINGMODELID_DEFAULT_LIT, SHADINGMODELID_THIN_TRANSLUCENT
+	// SHADINGMODELID_THIN_TRANSLUCENT
 	float3 GetTransmittanceColor() { return CustomData0.xyz; }
 	void SetTransmittanceColor(float3 TransmittanceColor) { CustomData0.xyz = TransmittanceColor; }
 
+	// SHADINGMODELID_DEFAULT_LIT
+	float3 GetExtinction() { return CustomData0.xyz; }
+	void SetExtinction(float3 SigmaT) { CustomData0.xyz = SigmaT; }
+
 	// SHADINGMODELID_MEDIUM
 	float3 GetHGWeight() { return CustomData0.xyz; }
 	float GetHGPhaseG() { return CustomData0.w; }
diff --git a/Engine/Shaders/Private/PathTracing/PathTracingCore.ush b/Engine/Shaders/Private/PathTracing/PathTracingCore.ush
--- a/Engine/Shaders/Private/PathTracing/PathTracingCore.ush
+++ b/Engine/Shaders/Private/PathTracing/PathTracingCore.ush
@@ -853,12 +853,6 @@ float3 ComputeDwivediScale(float3 Albedo)
 	return rsqrt(1.0 - pow(ClampedAlbedo, 2.44294 - 0.0215813 * ClampedAlbedo + 0.578637 / ClampedAlbedo));
 }
 
-float RescaleRandomNumber(float RandVal, float LowerBound, float UpperBound)
-{
-	const float OneMinusEpsilon = 0.99999994; // 32-bit float just before 1.0
-	return min((RandVal - LowerBound) / (UpperBound - LowerBound), OneMinusEpsilon);
-}
-
 float SampleGuidedSpectralTransmittance(float RandValue, float SlabCosine, float3 DwivediScale, float GuidingFactor, float3 Sigma, float3 ProbT, float3 ColorChannelPdf)
 {
 	float3 ColorChannelCdf = float3(
@@ -1244,16 +1238,12 @@ bool ProcessSubsurfaceRandomWalk(inout FPathTracingPayload Payload, inout float3
 
 			// we hit the boundary! overwrite the current payload and exit the walk with a diffuse scattering event
 			Payload.TranslatedWorldPos = Ray.Origin + ProbeResult.HitT * Ray.Direction;
-			Payload.WorldNormal = ProbeResult.WorldNormal;
-			Payload.WorldSmoothNormal = ProbeResult.WorldSmoothNormal;
-			Payload.WorldGeoNormal = ProbeResult.WorldGeoNormal;
 			// make sure normal is pointed outward from the object so we capture illumination from the exterior
-			if (ProbeResult.FrontFace != Payload.IsFrontFace())
-			{
-				Payload.WorldNormal *= -1.0;
-				Payload.WorldSmoothNormal *= -1.0;
-				Payload.WorldGeoNormal *= -1.0;
-			}
+			float SignFlip = (ProbeResult.FrontFace != Payload.IsFrontFace()) ? -1.0 : 1.0;
+			Payload.WorldNormal       = SignFlip * ProbeResult.WorldNormal;
+			Payload.WorldSmoothNormal = SignFlip * ProbeResult.WorldSmoothNormal;
+			Payload.WorldGeoNormal    = SignFlip * ProbeResult.WorldGeoNormal;
+			
 			Payload.ShadingModelID = SHADINGMODELID_NUM; // invalid value so that we get diffuse shading
 			Payload.BSDFOpacity = 1;
 			Payload.BaseColor = 1;
@@ -1675,7 +1665,7 @@ bool PathTracingKernel(inout FPathState PathState, int Bounce)
 	// long term we will probably want a stack with priorities
 	if (MaterialSample.PositionBiasSign < 0 && Payload.IsMaterialSolidGlass())
 	{
-		const float3 LocalSigmaT = PathTracingGlassTransmittanceToExtinction(Payload.GetTransmittanceColor());
+		const float3 LocalSigmaT = Payload.GetExtinction();
 		if (Payload.IsFrontFace())
 		{
 			// entering
diff --git a/Engine/Shaders/Private/PathTracing/PathTracingInitExtinctionCoefficient.usf b/Engine/Shaders/Private/PathTracing/PathTracingInitExtinctionCoefficient.usf
--- a/Engine/Shaders/Private/PathTracing/PathTracingInitExtinctionCoefficient.usf
+++ b/Engine/Shaders/Private/PathTracing/PathTracingInitExtinctionCoefficient.usf
@@ -40,7 +40,7 @@ RAY_TRACING_ENTRY_RAYGEN(PathTracingInitExtinctionCoefficientRG)
 		if (HitPayload.IsMaterialSolidGlass())
 		{
 			// Found a solid transmissive medium -- pickup the extinction
-			float3 LocalSigmaT = PathTracingGlassTransmittanceToExtinction(HitPayload.GetTransmittanceColor());
+			float3 LocalSigmaT = HitPayload.GetExtinction();
 
 			// increase when we leave a medium (that means we must have been inside before)
 			// decrease when we enter a medium (will be cancelled out when we leave)
diff --git a/Engine/Shaders/Private/PathTracing/PathTracingMaterialHitShader.usf b/Engine/Shaders/Private/PathTracing/PathTracingMaterialHitShader.usf
--- a/Engine/Shaders/Private/PathTracing/PathTracingMaterialHitShader.usf
+++ b/Engine/Shaders/Private/PathTracing/PathTracingMaterialHitShader.usf
@@ -24,6 +24,7 @@
 #include "/Engine/Private/DeferredShadingCommon.ush"
 #include "/Engine/Private/SubsurfaceProfileCommon.ush"
 #include "/Engine/Private/BurleyNormalizedSSSCommon.ush"
+#include "/Engine/Private/PathTracing/PathTracingCommon.ush"
 #include "/Engine/Private/PathTracing/Material/PathTracingFresnel.ush"
 #include "/Engine/Private/PathTracing/Material/PathTracingThinGlass.ush"
 
@@ -126,6 +127,7 @@ RAY_TRACING_ENTRY_CLOSEST_HIT(PathTracingMaterialCHS,
 	const bool bIsFrontFace = HitKind() == HIT_KIND_TRIANGLE_FRONT_FACE;
 	{
 		const float4 ScreenPosition = SvPositionToResolvedScreenPosition(SvPosition);
+		MaterialParameters.CameraVector = -WorldRayDirection();
 		CalcMaterialParametersEx(MaterialParameters, PixelMaterialInputs, SvPosition, ScreenPosition, bIsFrontFace, TranslatedWorldPosition, TranslatedWorldPosition);
 	}
 
@@ -523,7 +525,7 @@ RAY_TRACING_ENTRY_CLOSEST_HIT(PathTracingMaterialCHS,
 
 		Payload.ShadingModelID = SHADINGMODELID_DEFAULT_LIT; // replace shading model to avoid extra cases in raygen
 		Payload.BaseColor = WaterAlbedo; // TODO: unused at the moment
-		Payload.SetTransmittanceColor(PathTracingGlassExtinctionToTransmittance(WaterExtinction));
+		Payload.SetExtinction(WaterExtinction);
 		Payload.Ior = DielectricF0ToIor(DielectricSpecularToF0(Specular));
 		// we are using actual refraction, so don't apply any transparency
 		Payload.TransparencyColor = 0.0;
@@ -644,10 +646,10 @@ RAY_TRACING_ENTRY_CLOSEST_HIT(PathTracingMaterialCHS,
 	{
 		// only allow refraction for default lit materials since we need space for the absorption amount
 #if HAVE_GetAbsorptionMediumMaterialOutput0
-		Payload.SetTransmittanceColor(GetAbsorptionMediumMaterialOutput0(MaterialParameters));
+		Payload.SetExtinction(PathTracingGlassTransmittanceToExtinction(GetAbsorptionMediumMaterialOutput0(MaterialParameters)));
 #else
 		// Make the glass totally clear if no custom medium is set
-		Payload.SetTransmittanceColor(1.0);
+		Payload.SetExtinction(0.0);
 #endif
 		Payload.Ior = Ior;
 		if (Ior > 0.0)
