Speculative fix to TSR's EngineTest intermitent failures

diff --git a/Engine/Shaders/Private/TemporalSuperResolution/TSRDilateVelocity.usf b/Engine/Shaders/Private/TemporalSuperResolution/TSRDilateVelocity.usf
--- a/Engine/Shaders/Private/TemporalSuperResolution/TSRDilateVelocity.usf
+++ b/Engine/Shaders/Private/TemporalSuperResolution/TSRDilateVelocity.usf
@@ -56,8 +56,8 @@ Texture2D<uint> PrevScatteredSubpixelDepthTexture;
 
 RWTexture2D<float2> DilatedVelocityOutput;
 RWTexture2D<float> ClosestDepthOutput;
-RWTexture2D<uint> PrevUseCountOutput;
-RWTexture2D<uint> PrevClosestDepthOutput;
+globallycoherent RWTexture2D<uint> PrevUseCountOutput;
+globallycoherent RWTexture2D<uint> PrevClosestDepthOutput;
 RWTexture2DArray<uint> R8Output;
 RWTexture2D<uint> SubpixelDepthOutput;
 
