Fixes a bug on TSR's history weight with CONFIG_ENABLE_DUAL_PIXEL_VECTORIZATION

diff --git a/Engine/Shaders/Private/TemporalSuperResolution/TSRUpdateHistory.usf b/Engine/Shaders/Private/TemporalSuperResolution/TSRUpdateHistory.usf
--- a/Engine/Shaders/Private/TemporalSuperResolution/TSRUpdateHistory.usf
+++ b/Engine/Shaders/Private/TemporalSuperResolution/TSRUpdateHistory.usf
@@ -989,7 +989,11 @@ void MainCS(
 			PrevHistoryMomentMax = HighFrequencyMax;
 
 			PrevHighFrequencyColor = AccumulateHighFrequency;
-			PrevHistoryValidity = AccumulateMetadata[0];
+			#if CONFIG_METADATA_CHANNELS == 1
+				PrevHistoryValidity = AccumulateMetadata;
+			#else
+				PrevHistoryValidity = AccumulateMetadata[0];
+			#endif
 
 			PrevHistoryTranslucency = AccumulateTranslucency;
 
