Fix a bug of rendering black color with thin translucent material on mobile

diff --git a/Engine/Shaders/Private/MobileBasePassPixelShader.usf b/Engine/Shaders/Private/MobileBasePassPixelShader.usf
--- a/Engine/Shaders/Private/MobileBasePassPixelShader.usf
+++ b/Engine/Shaders/Private/MobileBasePassPixelShader.usf
@@ -77,6 +77,7 @@
 #define FORCE_VERTEX_NORMAL (MOBILE_QL_DISABLE_MATERIAL_NORMAL)
 #define SUPPORT_SPOTLIGHTS_SHADOW (MATERIALBLENDING_SOLID || MATERIALBLENDING_MASKED) && PROJECT_MOBILE_ENABLE_MOVABLE_SPOTLIGHT_SHADOWS
 #define DEFERRED_SHADING_PATH (MOBILE_DEFERRED_SHADING && ((MATERIALBLENDING_SOLID || MATERIALBLENDING_MASKED) && !MATERIAL_SHADINGMODEL_SINGLELAYERWATER))
+#define MATERIAL_WORKS_WITH_DUAL_SOURCE_COLOR_BLENDING (MATERIAL_SHADINGMODEL_THIN_TRANSLUCENT || STRATA_BLENDING_TRANSLUCENT_COLOREDTRANSMITTANCE )
 
 // Force FORWARD_SHADING on mobile when MOBILE_DEFERRED_SHADING is disabled
 #ifndef FORWARD_SHADING
@@ -287,7 +288,7 @@ void Main(
 	, out HALF4_TYPE OutGBufferD : SV_Target4
 	#endif
 #else
-	#if DUAL_SOURCE_COLOR_BLENDING_ENABLED && !THIN_TRANSLUCENT_FALL_BACK
+	#if DUAL_SOURCE_COLOR_BLENDING_ENABLED
 	, out HALF4_TYPE OutColor DUAL_SOURCE_BLENDING_SLOT(0) : SV_Target0
 	, out HALF4_TYPE OutColor1 DUAL_SOURCE_BLENDING_SLOT(1) : SV_Target1
 	#else
@@ -914,20 +915,20 @@ void Main(
 #endif // MATERIAL_SHADINGMODEL_THIN_TRANSLUCENT
 	
 	// MATERIAL_SHADINGMODEL_THIN_TRANSLUCENT must come first because it also has MATERIALBLENDING_TRANSLUCENT defined
-	#if DUAL_SOURCE_COLOR_BLENDING_ENABLED
+	#if MATERIAL_WORKS_WITH_DUAL_SOURCE_COLOR_BLENDING
 
 		// Add fog luminance according to surfacecoverage and reduce surface luminance according to fog coverage.
 		half3 AdjustedDualBlendAdd = DualBlendSurfaceCoverage * VertexFog.rgb + VertexFog.a * DualBlendSurfaceLuminancePostCoverage;
 		// Fade the surface color transmittance out to 1 according to the surface coverage, and take into account the fog coverage to the surface.
 		half3 AdjustedDualBlendMul = lerp(1.0f, VertexFog.a * DualBlendSurfaceTransmittancePreCoverage, DualBlendSurfaceCoverage);
 
-		#if THIN_TRANSLUCENT_FALL_BACK
+		#if DUAL_SOURCE_COLOR_BLENDING_ENABLED
+			OutColor = half4(AdjustedDualBlendAdd,0.0);
+			OutColor1 = half4(AdjustedDualBlendMul,1.0);
+		#else
 			// In the fallback case, use standard alpha blending(grey scale transmittance)
 			half AdjustedAlpha = saturate(1 - dot(AdjustedDualBlendMul, half3(1.0f, 1.0f, 1.0f) / 3.0f));
 			OutColor = half4(AdjustedDualBlendAdd, AdjustedAlpha);
-		#else
-			OutColor = half4(AdjustedDualBlendAdd,0.0);
-			OutColor1 = half4(AdjustedDualBlendMul,1.0);
 		#endif
 	// On mobile, water (an opaque material) is rendered as trnaslucent with forced premultiplied alpha blending (see MobileBasePass::SetTranslucentRenderState)
 	#elif MATERIALBLENDING_ALPHACOMPOSITE || MATERIAL_SHADINGMODEL_SINGLELAYERWATER
diff --git a/Engine/Source/Runtime/Renderer/Private/MobileBasePassRendering.h b/Engine/Source/Runtime/Renderer/Private/MobileBasePassRendering.h
--- a/Engine/Source/Runtime/Renderer/Private/MobileBasePassRendering.h
+++ b/Engine/Source/Runtime/Renderer/Private/MobileBasePassRendering.h
@@ -397,11 +397,6 @@ public:
 		OutEnvironment.SetDefine(TEXT("ENABLE_CLUSTERED_LIGHTS"), bEnableLocalLights ? 1u : 0u);
 		OutEnvironment.SetDefine(TEXT("ENABLE_CLUSTERED_REFLECTION"), bEnableClusteredReflections ? 1u : 0u);
 		OutEnvironment.SetDefine(TEXT("USE_SHADOWMASKTEXTURE"), bMobileUsesShadowMaskTexture && !bTranslucentMaterial ? 1u : 0u);
-
-		if (!RHISupportsDualSourceBlending(Parameters.Platform))
-		{
-			OutEnvironment.SetDefine(TEXT("THIN_TRANSLUCENT_FALL_BACK"), TEXT("1"));
-		}
 	}
 	
 	/** Initialization constructor. */
