Substrate - GPULightmass support

diff --git a/Engine/Shaders/Private/PathTracing/PathTracingGPULightmassMaterialHitShader.usf b/Engine/Shaders/Private/PathTracing/PathTracingGPULightmassMaterialHitShader.usf
--- a/Engine/Shaders/Private/PathTracing/PathTracingGPULightmassMaterialHitShader.usf
+++ b/Engine/Shaders/Private/PathTracing/PathTracingGPULightmassMaterialHitShader.usf
@@ -27,6 +27,10 @@
 #include "/Engine/Private/PathTracing/Material/PathTracingFresnel.ush"
 #include "/Engine/Private/PathTracing/Material/PathTracingThinGlass.ush"
 
+#if STRATA_ENABLED
+#define STRATA_GPU_LIGHTMASS 1
+#include "/Engine/Private/Strata/StrataExport.ush"
+#endif
 
 RAY_TRACING_ENTRY_CLOSEST_HIT(GPULightmassMaterialCHS,
 	FPackedPathTracingPayload, PackedPayload,
@@ -78,6 +82,101 @@ RAY_TRACING_ENTRY_CLOSEST_HIT(GPULightmassMaterialCHS,
 		CalcMaterialParametersEx(MaterialParameters, PixelMaterialInputs, SvPosition, ScreenPosition, bIsFrontFace, TranslatedWorldPosition, TranslatedWorldPosition);
 	}
 
+#if STRATA_ENABLED
+
+
+
+	FStrataData StrataData = PixelMaterialInputs.GetFrontStrataData();
+	FStrataPixelHeader StrataPixelHeader = MaterialParameters.GetFrontStrataHeader();
+
+	float  Coverage = 0;
+	float3 TransmittancePreCoverage = 0;
+	float3 BaseColorPostCoverage = 0;
+	float3 WorldNormal = 0;
+	float3 EmissiveLuminance = 0;
+
+	FStrataIntegrationSettings Settings = InitStrataIntegrationSettings(true /*bForceFullyRough*/, false /*StrataStruct.bRoughDiffuse*/, -1/*StrataStruct.PeelLayersAboveDepth*/, false/*StrataStruct.bRoughnessTracking*/);
+
+	const float3 SurfaceWorldNormal = GeoNormal;
+	StrataMaterialExportOut(
+		Settings,
+		StrataPixelHeader,
+		StrataData,
+		SurfaceWorldNormal,
+		Coverage,
+		TransmittancePreCoverage,
+		BaseColorPostCoverage,
+		WorldNormal,
+		EmissiveLuminance);
+
+
+#if MATERIALBLENDING_ANY_TRANSLUCENT
+	if (PackedPayload.IsVisibilityRay())
+	{
+		// Update the ray throughput (it is packed simply into the payload since we don't need to carry any other information across hits)
+		float3 RayThroughput = PackedPayload.GetRayThroughput();
+		RayThroughput *= saturate((1.0f - Coverage) + Coverage * TransmittancePreCoverage);
+		PackedPayload.SetRayThroughput(RayThroughput);
+		return;
+	}
+#endif
+
+	FPathTracingPayload Payload = (FPathTracingPayload)0;
+
+	Payload.BaseColor = BaseColorPostCoverage;
+	Payload.SubsurfaceColor = 0;							// STRATA_TODO unused by GPU lightmass?
+	Payload.BSDFOpacity = Coverage;
+	Payload.TransparencyColor = 0;							// STRATA_TODO That is not important when not evaluating throughput for IsVisibilityRay() above.
+
+	Payload.ShadingModelID = SHADINGMODELID_DEFAULT_LIT;	// STRATA_TODO check that this is enough for GPULightmass
+	
+	const uint PrimitiveFlags = GetPrimitiveData(MaterialParameters.PrimitiveId).Flags;
+	Payload.PrimitiveLightingChannelMask = GetPrimitive_LightingChannelMask_FromFlags(PrimitiveFlags);
+	Payload.HitT = RayTCurrent();
+	if (HitKind() == HIT_KIND_TRIANGLE_FRONT_FACE)
+	{
+		Payload.SetFrontFace();
+	}
+
+#if MATERIAL_IS_SKY
+	if (!PackedPayload.IsCameraRay())
+	{
+		// avoid double counting what was captured by the skylight
+		// also avoid noise from hot spots (they can be properly
+		// importance sampled if a capturing skylight is added)
+		PackedPayload = PackPathTracingPayload(Payload);
+		return;
+	}
+#endif
+
+	Payload.TranslatedWorldPos = LWCToFloat(LWCAdd(MaterialParameters.AbsoluteWorldPosition, ResolvedView.PreViewTranslation));
+	float GeoNormalSign = MaterialParameters.TwoSidedSign;
+#if !VF_SUPPORTS_RAYTRACING_PREPARE_MATERIAL_PIXEL_PARAMETERS
+	// Because the geometric normal is computed directly in world space
+	// it doesn't reflect the sign flip from the object transform, so apply it here
+	GeoNormalSign *= GetPrimitive_DeterminantSign(MaterialParameters.PrimitiveId);
+#endif
+	Payload.WorldGeoNormal = GeoNormalSign * GeoNormal;
+	Payload.WorldNormal = WorldNormal;
+
+	Payload.Radiance = EmissiveLuminance;
+
+#if MATERIAL_TWOSIDED
+	Payload.SetMaterialTwoSided();
+#else
+	if (MaterialParameters.TwoSidedSign < 0)
+	{
+		// when viewing the surface from "inside", don't include emission
+		Payload.Radiance = 0;
+	}
+#endif
+
+
+
+#else // STRATA_ENABLED
+
+
+
 	// the following blend modes need to process shadows after having executed the material
 #if MATERIALBLENDING_MODULATE || (MATERIALBLENDING_MASKED && MATERIAL_DITHER_OPACITY_MASK) || MATERIALBLENDING_ALPHACOMPOSITE || MATERIALBLENDING_TRANSLUCENT
 	if (PackedPayload.IsVisibilityRay())
@@ -277,6 +376,12 @@ RAY_TRACING_ENTRY_CLOSEST_HIT(GPULightmassMaterialCHS,
 #endif
 	Payload.BaseColor = BaseColor;
 
+
+
+#endif // STRATA_ENABLED
+
+
+
 	PackedPayload = PackPathTracingPayload(Payload);
 }
 
diff --git a/Engine/Shaders/Private/Strata/StrataExport.ush b/Engine/Shaders/Private/Strata/StrataExport.ush
--- a/Engine/Shaders/Private/Strata/StrataExport.ush
+++ b/Engine/Shaders/Private/Strata/StrataExport.ush
@@ -829,7 +829,11 @@ void PackStrataOut(
 
 }
 
-#if STRATA_MATERIAL_EXPORT_TYPE > 0
+#ifndef STRATA_GPU_LIGHTMASS
+#define STRATA_GPU_LIGHTMASS 0
+#endif
+
+#if STRATA_MATERIAL_EXPORT_TYPE > 0 || STRATA_GPU_LIGHTMASS > 0
 
 void StrataMaterialExportOut(
 	in FStrataIntegrationSettings Settings,
diff --git a/Engine/Source/Runtime/RenderCore/Private/Shader.cpp b/Engine/Source/Runtime/RenderCore/Private/Shader.cpp
--- a/Engine/Source/Runtime/RenderCore/Private/Shader.cpp
+++ b/Engine/Source/Runtime/RenderCore/Private/Shader.cpp
@@ -451,7 +451,9 @@ void FShaderType::ModifyCompilationEnvironment(const FShaderPermutationParameter
 		OutEnvironment.SetDefine(TEXT("RT_PAYLOAD_TYPE"), static_cast<int32>(RayTracingPayloadType));
 		OutEnvironment.SetDefine(TEXT("RT_PAYLOAD_MAX_SIZE"), GetRayTracingPayloadTypeMaxSize(RayTracingPayloadType));
 
-		if (uint32(RayTracingPayloadType) & uint32(ERayTracingPayloadType::RayTracingMaterial))
+		if (   (uint32(RayTracingPayloadType) & uint32(ERayTracingPayloadType::RayTracingMaterial))
+			|| (uint32(RayTracingPayloadType) & uint32(ERayTracingPayloadType::GPULightmass))
+			)
 		{
 			// If any payload requires a fully simplified material, we force fully simplified material all the way.
 			// That is used to have material ray tracing shaders compressed to single slab.
