Fixed missing volumetric fog light function from atlas for spot light.

diff --git a/Engine/Shaders/Private/VolumetricFogLightFunction.usf b/Engine/Shaders/Private/VolumetricFogLightFunction.usf
--- a/Engine/Shaders/Private/VolumetricFogLightFunction.usf
+++ b/Engine/Shaders/Private/VolumetricFogLightFunction.usf
@@ -42,6 +42,10 @@ void Main(
 	ScreenPosition = (SvPosition.xy * LightFunctionTexelSize);
 	ScreenPosition.x = 1 - ScreenPosition.x;
 	ScreenPosition.y = 1 - ScreenPosition.y;
+
+	// Wrap ScreenPosition to within the atlas tile
+	ScreenPosition = frac(ScreenPosition);
+
 #else // LIGHT_TYPE == 2
 	// Directional light
 	ScreenPosition = (SvPosition.xy * LightFunctionTexelSize - float2(0.5f, 0.5f)) * float2(2, -2);
