Fix root binding triangle visualization.

diff --git a/Engine/Plugins/Runtime/HairStrands/Source/HairStrandsCore/Private/GroomDebug.cpp b/Engine/Plugins/Runtime/HairStrands/Source/HairStrandsCore/Private/GroomDebug.cpp
--- a/Engine/Plugins/Runtime/HairStrands/Source/HairStrandsCore/Private/GroomDebug.cpp
+++ b/Engine/Plugins/Runtime/HairStrands/Source/HairStrandsCore/Private/GroomDebug.cpp
@@ -298,14 +298,18 @@ static void AddDebugProjectionHairPass(
 	ShaderPrint::SetEnabled(true);
 
 	if (MeshLODIndex < 0 || MeshLODIndex >= RestRootResources->LODs.Num() || MeshLODIndex >= DeformedRootResources->LODs.Num())
+	{
 		return;
+	}
 	
 	const EPrimitiveType PrimitiveType = GeometryType == EDebugProjectionHairType::HairFrame ? PT_LineList : GeometryType == EDebugProjectionHairType::HairTriangle ? PT_TriangleList : PT_LineList;
 	const uint32 RootCount = EDebugProjectionHairType::HairSamples == GeometryType ? 3 * RestRootResources->LODs[MeshLODIndex].SampleCount : RestRootResources->BulkData.RootCount;
 	const uint32 PrimitiveCount = RootCount;
 
 	if (PrimitiveCount == 0)
+	{
 		return;
+	}
 
 	if (ShaderPrintData == nullptr || !ShaderPrint::IsEnabled(*ShaderPrintData))
 	{
@@ -317,11 +321,15 @@ static void AddDebugProjectionHairPass(
 
 	if (EDebugProjectionHairType::HairFrame == GeometryType &&
 		!RestRootResources->LODs[MeshLODIndex].RootBarycentricBuffer.Buffer)
+	{
 		return;
+	}
 
 	if (EDebugProjectionHairType::HairSamples == GeometryType &&
 		!RestRootResources->LODs[MeshLODIndex].RestSamplePositionsBuffer.Buffer)
+	{
 			return;
+	}
 
 	const FHairStrandsRestRootResource::FLOD& RestLODDatas = RestRootResources->LODs[MeshLODIndex];
 	const FHairStrandsDeformedRootResource::FLOD& DeformedLODDatas = DeformedRootResources->LODs[MeshLODIndex];
@@ -331,8 +339,10 @@ static void AddDebugProjectionHairPass(
 		!RestLODDatas.RestUniqueTrianglePosition2Buffer.Buffer ||
 		!DeformedLODDatas.DeformedUniqueTrianglePosition0Buffer[0].Buffer ||
 		!DeformedLODDatas.DeformedUniqueTrianglePosition1Buffer[0].Buffer ||
-		!DeformedLODDatas.DeformedUniqueTrianglePosition2Buffer[0].Buffer	)
+		!DeformedLODDatas.DeformedUniqueTrianglePosition2Buffer[0].Buffer)
+	{
 		return;
+	}
 
 	// Double buffering is disabled by default unless the read-only cvar r.HairStrands.ContinuousDecimationReordering is set
 	if (IsHairStrandContinuousDecimationReorderingEnabled() && (!DeformedLODDatas.DeformedUniqueTrianglePosition0Buffer[1].Buffer || !DeformedLODDatas.DeformedUniqueTrianglePosition1Buffer[1].Buffer || !DeformedLODDatas.DeformedUniqueTrianglePosition2Buffer[1].Buffer))
@@ -350,7 +360,7 @@ static void AddDebugProjectionHairPass(
 
 	if (EDebugProjectionHairType::HairFrame == GeometryType)
 	{
-		Parameters->RootBarycentricBuffer	= RegisterAsSRV(GraphBuilder, RestLODDatas.RootBarycentricBuffer);
+		Parameters->RootBarycentricBuffer = RegisterAsSRV(GraphBuilder, RestLODDatas.RootBarycentricBuffer);
 	}
 
 	Parameters->RootToUniqueTriangleIndexBuffer = RegisterAsSRV(GraphBuilder, RestLODDatas.RootToUniqueTriangleIndexBuffer);
diff --git a/Engine/Shaders/Private/HairStrands/HairStrandsDebug.usf b/Engine/Shaders/Private/HairStrands/HairStrandsDebug.usf
--- a/Engine/Shaders/Private/HairStrands/HairStrandsDebug.usf
+++ b/Engine/Shaders/Private/HairStrands/HairStrandsDebug.usf
@@ -332,8 +332,7 @@ void MainCS(uint3 DispatchThreadId : SV_DispatchThreadID)
 
 #if PERMUTATION_INPUT_TYPE == INPUT_ROOT_FRAME  
 	const uint RootIndex = DispatchThreadId.x;
-
-	if (RootIndex > MaxRootCount) return;
+	if (RootIndex >= MaxRootCount) return;
 	{
 		const uint TriangleIndex = RootToUniqueTriangleIndexBuffer[RootIndex];
 		const FHairMeshTriangleWithUVs RestTriangle		= GetTriangleWithUVsTransformation(TriangleIndex, RestPosition0Buffer, RestPosition1Buffer, RestPosition2Buffer);
@@ -375,62 +374,65 @@ void MainCS(uint3 DispatchThreadId : SV_DispatchThreadID)
 
 #if PERMUTATION_INPUT_TYPE == INPUT_ROOT_TRIANGLE
 	const uint TriangleIndex = DispatchThreadId.x;
-
-	if (TriangleIndex > MaxRootCount) return;
-
-	const FHairMeshTriangleWithUVs RestTriangle		= GetTriangleWithUVsTransformation(TriangleIndex, RestPosition0Buffer, RestPosition1Buffer, RestPosition2Buffer);
-	const FHairMeshTriangleWithUVs DeformedTriangle = GetTriangleWithUVsTransformation(TriangleIndex, DeformedPosition0Buffer, DeformedPosition1Buffer, DeformedPosition2Buffer);
-
-	FHairMeshTriangleWithUVs Tri;
-	if (DeformedFrameEnable > 0)
-		Tri = DeformedTriangle;
-	else
-		Tri = RestTriangle;
-	
-
-	// Debug: To test if the rest-triangle-to-deformed-triangle transformation works correctly
-	#if 0
-	if (DeformedFrameEnable > 0)
+	if (TriangleIndex < MaxRootCount)
 	{
-		float3 Barycentric = 0;
-		if (TriangleVertex == 0)
-		{
-			P = RestTriangle.P0;
-			Barycentric = float3(1, 0, 0);
-		}
-		if (TriangleVertex == 1)
-		{
-			P = RestTriangle.P1;
-			Barycentric = float3(0, 1, 0);
-		}
-		if (TriangleVertex == 2)
+		const FHairMeshTriangleWithUVs RestTriangle		= GetTriangleWithUVsTransformation(TriangleIndex, RestPosition0Buffer, RestPosition1Buffer, RestPosition2Buffer);
+		const FHairMeshTriangleWithUVs DeformedTriangle = GetTriangleWithUVsTransformation(TriangleIndex, DeformedPosition0Buffer, DeformedPosition1Buffer, DeformedPosition2Buffer);
+	
+		FHairMeshTriangleWithUVs Tri;
+		if (DeformedFrameEnable > 0)
+			Tri = DeformedTriangle;
+		else
+			Tri = RestTriangle;
+		
+	
+		// Debug: To test if the rest-triangle-to-deformed-triangle transformation works correctly
+		#if 0
+		if (DeformedFrameEnable > 0)
 		{
-			P = RestTriangle.P2;
-			Barycentric = float3(0, 0, 1);
+			float3 Barycentric = 0;
+			if (TriangleVertex == 0)
+			{
+				P = RestTriangle.P0;
+				Barycentric = float3(1, 0, 0);
+			}
+			if (TriangleVertex == 1)
+			{
+				P = RestTriangle.P1;
+				Barycentric = float3(0, 1, 0);
+			}
+			if (TriangleVertex == 2)
+			{
+				P = RestTriangle.P2;
+				Barycentric = float3(0, 0, 1);
+			}
+			P = TransformPoint(P, Barycentric, RestTriangle, DeformedTriangle);
 		}
-		P = TransformPoint(P, Barycentric, RestTriangle, DeformedTriangle);
+		#endif
+	
+		float4 Color = float4(Tri.UV0, 0, 1);
+		// Color = float4((Tri.N + 1) * 0.5f, 1);
+	
+		const float3 WP0 = mul(float4(Tri.P0, 1), RootLocalToWorld).xyz;
+		const float3 WP1 = mul(float4(Tri.P1, 1), RootLocalToWorld).xyz; 
+		const float3 WP2 = mul(float4(Tri.P2, 1), RootLocalToWorld).xyz;
+	
+		FShaderPrintContext Ctx = InitShaderPrintContext(true, uint2(0, 0));
+		AddLineWS(WP0, WP1, Color);
+		AddLineWS(WP1, WP2, Color);
+		AddLineWS(WP2, WP0, Color);
 	}
-	#endif
-
-	float4 Color = float4(Tri.UV0, 0, 1);
-	// Color = float4((Tri.N + 1) * 0.5f, 1);
-
-	AddFilledTriangleWS(
-		mul(float4(Tri.P0, 1), RootLocalToWorld).xyz, 
-		mul(float4(Tri.P1, 1), RootLocalToWorld).xyz, 
-		mul(float4(Tri.P2, 1), RootLocalToWorld).xyz, 
-		Color);
 #endif // INPUT_ROOT_TRIANGLE
 
 #if PERMUTATION_INPUT_TYPE == INPUT_ROOT_SAMPLES 
-	const uint PointIndex = DispatchThreadId.x;
-	
-	if (PointIndex > MaxRootCount) return;
-
-	float3 P = (DeformedFrameEnable > 0) ? DeformedSamplePositionsBuffer[PointIndex].xyz : RestSamplePositionsBuffer[PointIndex].xyz;
-	P = mul(float4(P, 1), RootLocalToWorld).xyz;
+	const uint PointIndex = DispatchThreadId.x;	
+	if (PointIndex < MaxRootCount)
+	{
+		float3 P = (DeformedFrameEnable > 0) ? DeformedSamplePositionsBuffer[PointIndex].xyz : RestSamplePositionsBuffer[PointIndex].xyz;
+		P = mul(float4(P, 1), RootLocalToWorld).xyz;
 
-	AddCrossWS(P, 1.f, ColorYellow);
+		AddCrossWS(P, 1.f, ColorYellow);
+	}
 #endif // INPUT_ROOT_SAMPLES
 }
 #endif // SHADER_MESH_PROJECTION_HAIR
