PR #10231: fix: backlit hair square artifact when contact shadow is on but cast ? (Contributed by [REDACTED_PERSON])

diff --git a/Engine/Shaders/Private/DeferredLightingCommon.ush b/Engine/Shaders/Private/DeferredLightingCommon.ush
--- a/Engine/Shaders/Private/DeferredLightingCommon.ush
+++ b/Engine/Shaders/Private/DeferredLightingCommon.ush
@@ -42,7 +42,7 @@ float DistanceFromCameraFade(float SceneDepth, FDeferredLightData LightData)
 // Sets bOutHitCastDynamicShadow if the hit point is marked as a dynamic shadow caster
 float ShadowRayCast(
 	float3 RayOriginTranslatedWorld, float3 RayDirection, float RayLength,
-	int NumSteps, float StepOffset, out bool bOutHitCastContactShadow )
+	int NumSteps, float StepOffset, bool bHairNoShadowLight, out bool bOutHitCastContactShadow )
 {
 	float4 RayStartClip	= mul( float4( RayOriginTranslatedWorld, 1 ), View.TranslatedWorldToClip );
 	float4 RayDirClip	= mul( float4( RayDirection * RayLength, 0 ), View.TranslatedWorldToClip );
@@ -79,7 +79,11 @@ float ShadowRayCast(
 		float DepthDiff = SampleUVz.z - SampleDepth;
 		bool Hit = abs( DepthDiff + CompareTolerance ) < CompareTolerance;
 		// Avoid self-intersection with the start pixel (exact comparison due to point sampling depth buffer)
-		Hit = Hit && ( SampleDepth != StartDepth );
+		// Exception is made for hair for occluding transmitted light with non-shadow casting light
+		if (!bHairNoShadowLight)
+		{
+			Hit = Hit && (SampleDepth != StartDepth);
+		}
 
 		FirstHitTime = (Hit && FirstHitTime < 0.0) ? SampleTime : FirstHitTime;
 
@@ -200,7 +204,8 @@ void GetShadowTerms(float SceneDepth, half4 PrecomputedShadowFactors, uint Shadi
 	{
 		float StepOffset = Dither - 0.5;
 		bool bHitCastContactShadow = false;
-		float HitDistance = ShadowRayCast( TranslatedWorldPosition, L, ContactShadowLength, 8, StepOffset, bHitCastContactShadow );
+		bool bHairNoShadowLight = ShadingModelID == SHADINGMODELID_HAIR && !LightData.ShadowedBits;
+		float HitDistance = ShadowRayCast( TranslatedWorldPosition, L, ContactShadowLength, 8, StepOffset, bHairNoShadowLight, bHitCastContactShadow );
 				
 		if ( HitDistance > 0.0 )
 		{
