* Share encoding/decoding of binding data
* Move common funciton into hairstrandsbindingcommon

diff --git a/Engine/Plugins/Runtime/HairStrands/Shaders/Private/NiagaraDataInterfaceHairStrandsTemplate.ush b/Engine/Plugins/Runtime/HairStrands/Shaders/Private/NiagaraDataInterfaceHairStrandsTemplate.ush
--- a/Engine/Plugins/Runtime/HairStrands/Shaders/Private/NiagaraDataInterfaceHairStrandsTemplate.ush
+++ b/Engine/Plugins/Runtime/HairStrands/Shaders/Private/NiagaraDataInterfaceHairStrandsTemplate.ush
@@ -3,6 +3,7 @@
 /*=============================================================================
 NiagaraDataInterfaceHairStrands.ush 
 =============================================================================*/
+#include "/Engine/Private/HairStrands/HairStrandsBindingCommon.ush"
 
 RWBuffer<uint4>		{ParameterName}_DeformedPositionBuffer;
 Buffer<uint4>		{ParameterName}_RestPositionBuffer;
@@ -651,7 +652,7 @@ void UpdateTriangleDisplace_{ParameterName}(in int StrandIndex, in float3 NodePo
 {
 	if( {ParameterName}_InterpolationMode >= 1)
 	{
-		const float2 ProjectionUV = ExtractBarycentrics({ParameterName}_RootBarycentricCoordinatesBuffer[StrandIndex]);
+		const float2 ProjectionUV = UnpackBarycentrics({ParameterName}_RootBarycentricCoordinatesBuffer[StrandIndex]).xy;
 		const uint TriangleIndex = {ParameterName}_RootToUniqueTriangleIndexBuffer[StrandIndex];
 
 		BuildRestTriangle_{ParameterName}(ProjectionUV, TriangleIndex, RestTrianglePosition, RestTriangleOrientation);
@@ -850,22 +851,14 @@ void GetLocalVector_{ParameterName}(in float3 WorldVector, in bool IsPosition, o
  * Attached the root to the skinned cache
  * -----------------------------------------------------------------
  */
-// Duplicate from #include "HairStrands/HairStrandsMeshProjectionCommon.ush" ApplyRBF()
+
 float3 ApplyRBF_{ParameterName}(in float3 RestControlPoint)
 {
-	float3 ControlPoint = RestControlPoint;
-	// Apply rbf interpolation from the samples set
-	for (int i = 0; i < {ParameterName}_SampleCount; ++i)
-	{
-		const float3 PositionDelta = RestControlPoint - {ParameterName}_RestSamplePositionsBuffer[i].xyz;
-		const float FunctionValue = sqrt(dot(PositionDelta, PositionDelta) + 1);
-		ControlPoint += FunctionValue * {ParameterName}_MeshSampleWeightsBuffer[i].xyz;
-	}
-	ControlPoint += {ParameterName}_MeshSampleWeightsBuffer[{ParameterName}_SampleCount].xyz;
-	ControlPoint += {ParameterName}_MeshSampleWeightsBuffer[{ParameterName}_SampleCount + 1].xyz * RestControlPoint.x;
-	ControlPoint += {ParameterName}_MeshSampleWeightsBuffer[{ParameterName}_SampleCount + 2].xyz * RestControlPoint.y;
-	ControlPoint += {ParameterName}_MeshSampleWeightsBuffer[{ParameterName}_SampleCount + 3].xyz * RestControlPoint.z;
-	return ControlPoint;
+	return ApplyRBF(
+		RestControlPoint, 
+		{ParameterName}_SampleCount, 
+		{ParameterName}_RestSamplePositionsBuffer,
+		{ParameterName}_MeshSampleWeightsBuffer);
 }
 
  // Compute the global position
@@ -898,7 +891,7 @@ void ComputeLocalState_{ParameterName}(in float3 RestPosition, in float4 RestOri
 	if( {ParameterName}_InterpolationMode >= 1)
 	{
 		const int StrandIndex = ExecIndex() / {ParameterName}_StrandSize;
-		const float2 ProjectionUV = ExtractBarycentrics({ParameterName}_RootBarycentricCoordinatesBuffer[StrandIndex]);
+		const float2 ProjectionUV = UnpackBarycentrics({ParameterName}_RootBarycentricCoordinatesBuffer[StrandIndex]);
 		const uint TriangleIndex = {ParameterName}_RootToUniqueTriangleIndexBuffer[StrandIndex];
 
 		float3 RestTrianglePosition = float3(0, 0, 0);
@@ -919,7 +912,7 @@ void AttachNodeState_{ParameterName}(in float3 LocalPosition, in float4 LocalOri
 	if( {ParameterName}_InterpolationMode >= 1)
 	{
 		const int StrandIndex = ExecIndex() / {ParameterName}_StrandSize;
-		const float2 ProjectionUV = ExtractBarycentrics({ParameterName}_RootBarycentricCoordinatesBuffer[StrandIndex]);
+		const float2 ProjectionUV = UnpackBarycentrics({ParameterName}_RootBarycentricCoordinatesBuffer[StrandIndex]);
 		const uint TriangleIndex = {ParameterName}_RootToUniqueTriangleIndexBuffer[StrandIndex];
 
 		float3 DeformedTrianglePosition = float3(0, 0, 0);
diff --git a/Engine/Plugins/Runtime/HairStrands/Source/HairStrandsCore/Private/GroomBindingBuilder.cpp b/Engine/Plugins/Runtime/HairStrands/Source/HairStrandsCore/Private/GroomBindingBuilder.cpp
--- a/Engine/Plugins/Runtime/HairStrands/Source/HairStrandsCore/Private/GroomBindingBuilder.cpp
+++ b/Engine/Plugins/Runtime/HairStrands/Source/HairStrandsCore/Private/GroomBindingBuilder.cpp
@@ -62,25 +62,25 @@ FString FGroomBindingBuilder::GetVersion()
 
 ///////////////////////////////////////////////////////////////////////////////////////////////////
 // Common utils functions
-// These utils function are a copy of function in HairStrandsMeshProjectionCommon.ush
-uint32 FHairStrandsRootUtils::EncodeTriangleIndex(uint32 TriangleIndex, uint32 SectionIndex)
+// These utils function are a copy of function in HairStrandsBindingCommon.ush
+uint32 FHairStrandsRootUtils::PackTriangleIndex(uint32 TriangleIndex, uint32 SectionIndex)
 {
 	return ((SectionIndex & 0xFF) << 24) | (TriangleIndex & 0xFFFFFF);
 }
 
-// This function is a copy of DecodeTriangleIndex in HairStrandsMeshProjectionCommon.ush
-void FHairStrandsRootUtils::DecodeTriangleIndex(uint32 Encoded, uint32& OutTriangleIndex, uint32& OutSectionIndex)
+// This function is a copy of UnpackTriangleIndex in HairStrandsBindingCommon.ush
+void FHairStrandsRootUtils::UnpackTriangleIndex(uint32 Encoded, uint32& OutTriangleIndex, uint32& OutSectionIndex)
 {
 	OutSectionIndex = (Encoded >> 24) & 0xFF;
 	OutTriangleIndex = Encoded & 0xFFFFFF;
 }
 
-uint32 FHairStrandsRootUtils::EncodeBarycentrics(const FVector2f& B)
+uint32 FHairStrandsRootUtils::PackBarycentrics(const FVector2f& B)
 {
 	return uint32(FFloat16(B.X).Encoded) | (uint32(FFloat16(B.Y).Encoded)<<16);
 }
 
-FVector2f FHairStrandsRootUtils::DecodeBarycentrics(uint32 B)
+FVector2f FHairStrandsRootUtils::UnpackBarycentrics(uint32 B)
 {
 	FFloat16 BX;
 	BX.Encoded = (B & 0xFFFF);
@@ -662,7 +662,7 @@ namespace GroomBinding_RBFWeighting
 		{
 			uint32 SectionIndex = 0;
 			uint32 TriangleIndex = 0;
-			FHairStrandsRootUtils::DecodeTriangleIndex(EncodedTriangleId, TriangleIndex, SectionIndex);
+			FHairStrandsRootUtils::UnpackTriangleIndex(EncodedTriangleId, TriangleIndex, SectionIndex);
 			if (!ValidSection || (ValidSection && (SectionIndex == TargetSection)))
 			{
 				const IMeshSectionData& Section = MeshLODData.GetSection(SectionIndex);
@@ -1294,8 +1294,8 @@ namespace GroomBinding_RootProjection
 				check(ClosestDistance < FLT_MAX);
 
 				// Record closest triangle and the root's barycentrics
-				const uint32 EncodedBarycentrics = FHairStrandsRootUtils::EncodeBarycentrics(FVector2f(ClosestBarycentrics));	// LWC_TODO: Precision loss
-				const uint32 EncodedTriangleIndex = FHairStrandsRootUtils::EncodeTriangleIndex(ClosestTriangle.TriangleIndex, ClosestTriangle.SectionIndex);
+				const uint32 EncodedBarycentrics = FHairStrandsRootUtils::PackBarycentrics(FVector2f(ClosestBarycentrics));	// LWC_TODO: Precision loss
+				const uint32 EncodedTriangleIndex = FHairStrandsRootUtils::PackTriangleIndex(ClosestTriangle.TriangleIndex, ClosestTriangle.SectionIndex);
 				OutRootData.MeshProjectionLODs[LODIt].RootBarycentricBuffer[CurveIndex] = EncodedBarycentrics;
 
 				RootTriangleIndexBuffer[CurveIndex] = EncodedTriangleIndex;
@@ -1327,7 +1327,7 @@ namespace GroomBinding_RootProjection
 					// Add unique section
 					uint32 TriangleIndex;
 					uint32 SectionIndex;
-					FHairStrandsRootUtils::DecodeTriangleIndex(EncodedTriangleId, TriangleIndex, SectionIndex);
+					FHairStrandsRootUtils::UnpackTriangleIndex(EncodedTriangleId, TriangleIndex, SectionIndex);
 					UniqueSectionId.AddUnique(SectionIndex);
 
 					// Add unique triangle
diff --git a/Engine/Plugins/Runtime/HairStrands/Source/HairStrandsCore/Private/GroomRBFDeformer.cpp b/Engine/Plugins/Runtime/HairStrands/Source/HairStrandsCore/Private/GroomRBFDeformer.cpp
--- a/Engine/Plugins/Runtime/HairStrands/Source/HairStrandsCore/Private/GroomRBFDeformer.cpp
+++ b/Engine/Plugins/Runtime/HairStrands/Source/HairStrandsCore/Private/GroomRBFDeformer.cpp
@@ -120,7 +120,7 @@ void DeformStrands(
 		const FVector3f& Deform_Position = OutDeformedPositionBuffer[VertexOffset];
 
 		const uint32 PackedBarycentric = RootBarycentricBuffer[RootIndex];
-		const FVector2f B0 = FVector2f(FHairStrandsRootUtils::DecodeBarycentrics(PackedBarycentric));
+		const FVector2f B0 = FVector2f(FHairStrandsRootUtils::UnpackBarycentrics(PackedBarycentric));
 		const FVector3f   B  = FVector3f(B0.X, B0.Y, 1.f - B0.X - B0.Y);
 
 		/* Strand hair roots translation and rotation in rest position relative to the bound triangle. Positions are relative to the rest root center */
@@ -170,7 +170,7 @@ void ExtractUniqueTrianglePositions(
 		const uint32 PackedTriangleIndex = RestLODData.UniqueTriangleIndexBuffer[UniqueTriangleIndex];
 		uint32 TriangleIndex = 0;
 		uint32 SectionIndex = 0;
-		FHairStrandsRootUtils::DecodeTriangleIndex(PackedTriangleIndex, TriangleIndex, SectionIndex);
+		FHairStrandsRootUtils::UnpackTriangleIndex(PackedTriangleIndex, TriangleIndex, SectionIndex);
 
 		check(SectionIndex < SectionCount)
 		const uint32 TriangleCount = InMeshRenderData->LODRenderData[MeshLODIndex].RenderSections[SectionIndex].NumTriangles;
diff --git a/Engine/Plugins/Runtime/HairStrands/Source/HairStrandsCore/Public/HairStrandsDatas.h b/Engine/Plugins/Runtime/HairStrands/Source/HairStrandsCore/Public/HairStrandsDatas.h
--- a/Engine/Plugins/Runtime/HairStrands/Source/HairStrandsCore/Public/HairStrandsDatas.h
+++ b/Engine/Plugins/Runtime/HairStrands/Source/HairStrandsCore/Public/HairStrandsDatas.h
@@ -224,10 +224,10 @@ struct FHairStrandsRootBarycentricFormat
 
 struct FHairStrandsRootUtils
 {
-	static uint32	 EncodeTriangleIndex(uint32 TriangleIndex, uint32 SectionIndex);
-	static void		 DecodeTriangleIndex(uint32 Encoded, uint32& OutTriangleIndex, uint32& OutSectionIndex);
-	static uint32	 EncodeBarycentrics(const FVector2f& B);
-	static FVector2f DecodeBarycentrics(uint32 B);
+	static uint32	 PackTriangleIndex(uint32 TriangleIndex, uint32 SectionIndex);
+	static void		 UnpackTriangleIndex(uint32 Encoded, uint32& OutTriangleIndex, uint32& OutSectionIndex);
+	static uint32	 PackBarycentrics(const FVector2f& B);
+	static FVector2f UnpackBarycentrics(uint32 B);
 	static uint32	 PackUVs(const FVector2f& UV);
 	static float	 PackUVsToFloat(const FVector2f& UV);
 };
diff --git a/Engine/Shaders/Private/HairStrands/HairCardsDeformation.usf b/Engine/Shaders/Private/HairStrands/HairCardsDeformation.usf
--- a/Engine/Shaders/Private/HairStrands/HairCardsDeformation.usf
+++ b/Engine/Shaders/Private/HairStrands/HairCardsDeformation.usf
@@ -2,7 +2,7 @@
 
 #include "../Common.ush"
 #include "HairStrandsVertexFactoryCommon.ush"
-#include "HairStrandsMeshProjectionCommon.ush"
+#include "HairStrandsBindingCommon.ush"
 
 #define DEBUG_ENABLE 0
 
@@ -143,7 +143,7 @@ void MainCS(uint DispatchThreadId : SV_DispatchThreadID)
 	const uint RootIndex = GuideVertexToRootIndexBuffer[GuideIndex0];
 	const uint TriangleIndex = GuideRootToUniqueTriangleIndexBuffer[RootIndex];
 
-	const float3 RootBarycentric = DecodeBarycentrics(GuideRootBarycentricBuffer[RootIndex]);
+	const float3 RootBarycentric = UnpackBarycentrics(GuideRootBarycentricBuffer[RootIndex]);
 
 	const FHairMeshTriangle RestTriangle = GetTriangleTransformation(TriangleIndex, TriangleRestPositionBuffer);
 	const FHairMeshTriangle DeformedTriangle = GetTriangleTransformation(TriangleIndex, TriangleDeformedPositionBuffer);
diff --git a/Engine/Shaders/Private/HairStrands/HairStrandsBindingCommon.ush b/Engine/Shaders/Private/HairStrands/HairStrandsBindingCommon.ush
new file mode 100644
--- /dev/null
+++ b/Engine/Shaders/Private/HairStrands/HairStrandsBindingCommon.ush
@@ -0,0 +1,249 @@
+// Copyright Epic Games, Inc. All Rights Reserved.
+
+struct FHairMeshTriangle
+{
+	float3 P0;
+	float3 P1;
+	float3 P2;
+
+	float3 E0;
+	float3 E1;
+
+	float3 N;
+};
+
+struct FHairMeshTriangleWithUVs : FHairMeshTriangle
+{
+	float2 UV0;
+	float2 UV1;
+	float2 UV2;
+};
+
+uint PackUVs(float2 UV)
+{
+	return (f32tof16(UV.x) & 0xFFFF) | ((f32tof16(UV.y) & 0xFFFF) << 16);
+}
+
+float PackUVsToFloat(float2 UV)
+{
+	return asfloat(PackUVs(UV));
+}
+
+float2 UnPackUVs(uint InUV)
+{
+	float2 Out;
+	Out.x = f16tof32(InUV & 0xFFFF);
+	Out.y = f16tof32((InUV >> 16) & 0xFFFF);
+	return Out;
+}
+
+float2 UnPackUVs(float InUV)
+{
+	return UnPackUVs(asuint(InUV));
+}
+
+float2 GetVertexUV(uint VId, Buffer<float2> InUVBuffer, Buffer<uint> InIndexBuffer, uint InIndexOffset, uint InMaxIndexCount, uint InMaxVertexCount, uint InMeshUVsChannelOffset, uint InMeshUVsChannelCount)
+{
+	const uint Index = min(InIndexOffset + VId, InMaxIndexCount - 1);
+	uint VertexIndex = InIndexBuffer.Load(Index);
+	VertexIndex = min(VertexIndex, InMaxVertexCount - 1);
+
+	const uint UVIndex = InMeshUVsChannelOffset + InMeshUVsChannelCount * VertexIndex;
+	return InUVBuffer.Load(UVIndex);
+}
+
+float3 GetVertexPosition(uint VId, Buffer<float> InPositionBuffer, Buffer<uint> InIndexBuffer, uint InIndexOffset, uint InMaxIndexCount, uint InMaxVertexCount)
+{
+	const uint Index = min(InIndexOffset + VId, InMaxIndexCount-1);
+	uint VertexIndex = InIndexBuffer.Load(Index);
+	VertexIndex = min(VertexIndex, InMaxVertexCount-1);
+
+	float3 P = 0;
+	P.x = InPositionBuffer.Load(VertexIndex*3 + 0);
+	P.y = InPositionBuffer.Load(VertexIndex*3 + 1);
+	P.z = InPositionBuffer.Load(VertexIndex*3 + 2);
+	return P;
+}
+
+FHairMeshTriangle GetTriangleData(uint TriangleIndex, Buffer<float> InPositionBuffer, Buffer<uint> InIndexBuffer, uint InIndexOffset, uint InMaxIndexCount, uint InMaxVertexCount)
+{
+	const uint BaseVertexId = TriangleIndex*3;
+	const uint VId0 = BaseVertexId;
+	const uint VId1 = BaseVertexId+1;
+	const uint VId2 = BaseVertexId+2;
+
+	FHairMeshTriangle Out;
+	Out.P0 = GetVertexPosition(VId0, InPositionBuffer, InIndexBuffer, InIndexOffset, InMaxIndexCount, InMaxVertexCount);
+	Out.P1 = GetVertexPosition(VId1, InPositionBuffer, InIndexBuffer, InIndexOffset, InMaxIndexCount, InMaxVertexCount);
+	Out.P2 = GetVertexPosition(VId2, InPositionBuffer, InIndexBuffer, InIndexOffset, InMaxIndexCount, InMaxVertexCount);
+	Out.E0 = Out.P1 - Out.P0;
+	Out.E1 = Out.P2 - Out.P0;
+	Out.N = normalize(cross(Out.E1, Out.E0));
+
+	return Out;
+}
+
+FHairMeshTriangleWithUVs GetTriangleData(uint TriangleIndex, Buffer<float> InPositionBuffer, Buffer<float2> InUVBuffer, Buffer<uint> InIndexBuffer, uint InIndexOffset, uint InMaxIndexCount, uint InMaxVertexCount, uint InMeshUVsChannelOffset, uint InMeshUVsChannelCount)
+{
+	const uint BaseVertexId = TriangleIndex * 3;
+	const uint VId0 = BaseVertexId;
+	const uint VId1 = BaseVertexId + 1;
+	const uint VId2 = BaseVertexId + 2;
+
+	FHairMeshTriangleWithUVs Out;
+	Out.P0 = GetVertexPosition(VId0, InPositionBuffer, InIndexBuffer, InIndexOffset, InMaxIndexCount, InMaxVertexCount);
+	Out.P1 = GetVertexPosition(VId1, InPositionBuffer, InIndexBuffer, InIndexOffset, InMaxIndexCount, InMaxVertexCount);
+	Out.P2 = GetVertexPosition(VId2, InPositionBuffer, InIndexBuffer, InIndexOffset, InMaxIndexCount, InMaxVertexCount);
+	Out.E0 = Out.P1 - Out.P0;
+	Out.E1 = Out.P2 - Out.P0;
+	Out.N = normalize(cross(Out.E1, Out.E0));
+
+	Out.UV0 = GetVertexUV(VId0, InUVBuffer, InIndexBuffer, InIndexOffset, InMaxIndexCount, InMaxVertexCount, InMeshUVsChannelOffset, InMeshUVsChannelCount);
+	Out.UV1 = GetVertexUV(VId1, InUVBuffer, InIndexBuffer, InIndexOffset, InMaxIndexCount, InMaxVertexCount, InMeshUVsChannelOffset, InMeshUVsChannelCount);
+	Out.UV2 = GetVertexUV(VId2, InUVBuffer, InIndexBuffer, InIndexOffset, InMaxIndexCount, InMaxVertexCount, InMeshUVsChannelOffset, InMeshUVsChannelCount);
+	
+	return Out;
+}
+
+FHairMeshTriangle GetTriangleTransformation(
+	uint TriangleIndex,
+	Buffer<float4> InPositionBuffer)
+{
+	FHairMeshTriangle Out;
+
+	// Build orthonormal frame
+	Out.P0 = InPositionBuffer[TriangleIndex * 3 + 0].xyz;
+	Out.P1 = InPositionBuffer[TriangleIndex * 3 + 1].xyz;
+	Out.P2 = InPositionBuffer[TriangleIndex * 3 + 2].xyz;
+	Out.E0 = Out.P1 - Out.P0;
+	Out.E1 = Out.P2 - Out.P0;
+
+	Out.N  = normalize(cross(Out.E1, Out.E0));
+	Out.E0 = normalize(Out.E0); 
+	Out.E1 = normalize(cross(Out.E0, Out.N));
+
+	return Out;
+}
+
+FHairMeshTriangleWithUVs GetTriangleWithUVsTransformation(
+	uint TriangleIndex,
+	Buffer<float4> InPositionBuffer)
+{
+	FHairMeshTriangleWithUVs Out;
+
+	const float4 V0 = InPositionBuffer[TriangleIndex * 3 + 0];
+	const float4 V1 = InPositionBuffer[TriangleIndex * 3 + 1];
+	const float4 V2 = InPositionBuffer[TriangleIndex * 3 + 2];
+
+	// Build orthonormal frame
+	Out.P0 = V0.xyz;
+	Out.P1 = V1.xyz;
+	Out.P2 = V2.xyz;
+	Out.E0 = Out.P1 - Out.P0;
+	Out.E1 = Out.P2 - Out.P0;
+
+	Out.N = normalize(cross(Out.E1, Out.E0));
+	Out.E0 = normalize(Out.E0);
+	Out.E1 = normalize(cross(Out.E0, Out.N));
+
+	Out.UV0 = UnPackUVs(V0.w);
+	Out.UV1 = UnPackUVs(V1.w);
+	Out.UV2 = UnPackUVs(V2.w);
+
+	return Out;
+}
+
+float3 ToTriangle(const float3 P, float3 B, const FHairMeshTriangle T)
+{
+	const float3 LocalP = P - (T.P0 * B.x + T.P1 * B.y + T.P2 * B.z);
+
+	float3 Out;
+	Out.x = dot(LocalP, T.E0);
+	Out.y = dot(LocalP, T.E1);
+	Out.z = dot(LocalP, T.N);
+	return Out;
+}
+
+float3 ToWorld(const float3 P, float3 B, const FHairMeshTriangle T)
+{
+	const float3 Rotation = P.x * T.E0 + P.y * T.E1 + P.z * T.N;
+	const float3 Translation = T.P0 * B.x + T.P1 * B.y + T.P2 * B.z;
+	return Rotation + Translation;
+
+}
+
+float3 VectorToWorld(const float3 P, const FHairMeshTriangle T)
+{
+	const float3 Rotation = P.x * T.E0 + P.y * T.E1 + P.z * T.N;
+	return Rotation;
+}
+
+float3 VectorToTriangle(const float3 P, const FHairMeshTriangle T)
+{
+	float3 Out;
+	Out.x = dot(P, T.E0);
+	Out.y = dot(P, T.E1);
+	Out.z = dot(P, T.N);
+	return Out;
+}
+
+float3 TransformPoint(float3 P, float3 B, const FHairMeshTriangle RestTriangle, const FHairMeshTriangle DeformedTriangle)
+{
+	float3 LocalP = ToTriangle(P, B, RestTriangle);
+	return ToWorld(LocalP, B, DeformedTriangle);
+}
+
+uint PackBarycentrics(float2 B)
+{
+	return f32tof16(B.x) | (f32tof16(B.y) << 16);
+}
+
+float3 UnpackBarycentrics(uint E)
+{
+	float3 Out;
+	Out.x = f16tof32(E & 0xFFFF);
+	Out.y = f16tof32((E >> 16) & 0xFFFF);
+	Out.z = 1 - Out.x - Out.y;
+	return Out;
+}
+
+struct FHairTriangleIndex
+{
+	uint TriangleIndex;
+	uint SectionIndex;
+};
+
+uint PackTriangleIndex(uint TriangleIndex, uint SectionIndex)
+{
+	return ((SectionIndex & 0xFF)<<24) | (TriangleIndex & 0xFFFFFF);
+}
+
+FHairTriangleIndex UnpackTriangleIndex(uint Encoded)
+{
+	FHairTriangleIndex Out;
+	Out.SectionIndex = (Encoded>>24) & 0xFF;
+	Out.TriangleIndex = Encoded & 0xFFFFFF;
+	return Out;
+}
+
+// Apply RBF weight onto a control point
+float3 ApplyRBF(
+	float3 RestControlPoint,
+	int InMaxSampleCount,
+	Buffer<float4> InRestSamplePositionsBuffer,
+	Buffer<float4> InMeshSampleWeightsBuffer)
+{
+	float3 ControlPoint = RestControlPoint;
+	// Apply rbf interpolation from the samples set
+	for (int i = 0; i < InMaxSampleCount; ++i)
+	{
+		const float3 PositionDelta = RestControlPoint - InRestSamplePositionsBuffer[i].xyz;
+		const float FunctionValue = sqrt(dot(PositionDelta, PositionDelta) + 1);
+		ControlPoint += FunctionValue * InMeshSampleWeightsBuffer[i].xyz;
+	}
+	ControlPoint += InMeshSampleWeightsBuffer[InMaxSampleCount].xyz;
+	ControlPoint += InMeshSampleWeightsBuffer[InMaxSampleCount + 1].xyz * RestControlPoint.x;
+	ControlPoint += InMeshSampleWeightsBuffer[InMaxSampleCount + 2].xyz * RestControlPoint.y;
+	ControlPoint += InMeshSampleWeightsBuffer[InMaxSampleCount + 3].xyz * RestControlPoint.z;
+	return ControlPoint;
+}
\ No newline at end of file
diff --git a/Engine/Shaders/Private/HairStrands/HairStrandsDebug.usf b/Engine/Shaders/Private/HairStrands/HairStrandsDebug.usf
--- a/Engine/Shaders/Private/HairStrands/HairStrandsDebug.usf
+++ b/Engine/Shaders/Private/HairStrands/HairStrandsDebug.usf
@@ -295,7 +295,7 @@ void MainCS(uint3 DispatchThreadId : SV_DispatchThreadID)
 
 #if SHADER_MESH_PROJECTION_HAIR
 
-#include "HairStrandsMeshProjectionCommon.ush" 
+#include "HairStrandsBindingCommon.ush" 
 #include "../ShaderPrint.ush"
 
 #define INPUT_ROOT_FRAME 0
@@ -339,7 +339,7 @@ void MainCS(uint3 DispatchThreadId : SV_DispatchThreadID)
 		else
 			Tri = RestTriangle;
 
-		const float3 B  = DecodeBarycentrics(RootBarycentricBuffer[RootIndex]);
+		const float3 B  = UnpackBarycentrics(RootBarycentricBuffer[RootIndex]);
 		float3 P0 = B.x * Tri.P0 + B.y * Tri.P1 + B.z * Tri.P2;
 		float3 P1 = P0 + Tri.N;
 
@@ -438,7 +438,7 @@ void MainCS(uint3 DispatchThreadId : SV_DispatchThreadID)
 
 #if SHADER_MESH_PROJECTION_SKIN_CACHE
 
-#include "HairStrandsMeshProjectionCommon.ush"
+#include "HairStrandsBindingCommon.ush"
 #include "../ShaderPrint.ush"
 
 float4x4 LocalToWorld;
diff --git a/Engine/Shaders/Private/HairStrands/HairStrandsFollicleMask.usf b/Engine/Shaders/Private/HairStrands/HairStrandsFollicleMask.usf
--- a/Engine/Shaders/Private/HairStrands/HairStrandsFollicleMask.usf
+++ b/Engine/Shaders/Private/HairStrands/HairStrandsFollicleMask.usf
@@ -1,7 +1,7 @@
 // Copyright Epic Games, Inc. All Rights Reserved.
 
 #include "../Common.ush"
-#include "HairStrandsMeshProjectionCommon.ush"
+#include "HairStrandsBindingCommon.ush"
 
 #if SHADER_FOLLICLE_MASK
 
@@ -51,7 +51,7 @@ void MainVS(
 	const FHairMeshTriangleWithUVs Tri = GetTriangleWithUVsTransformation(UniqueTriangleIndex, UniqueTrianglePositionBuffer);
 
 	const uint EncodedBarycentric = RootBarycentricBuffer[RootIndex];
-	const float3 Barycentric = DecodeBarycentrics(EncodedBarycentric);
+	const float3 Barycentric = UnpackBarycentrics(EncodedBarycentric);
 	float2 RootUV = Barycentric.x * Tri.UV0 + Barycentric.y * Tri.UV1 + Barycentric.z * Tri.UV2;
 #endif
 
diff --git a/Engine/Shaders/Private/HairStrands/HairStrandsGuideDeform.usf b/Engine/Shaders/Private/HairStrands/HairStrandsGuideDeform.usf
--- a/Engine/Shaders/Private/HairStrands/HairStrandsGuideDeform.usf
+++ b/Engine/Shaders/Private/HairStrands/HairStrandsGuideDeform.usf
@@ -2,7 +2,7 @@
 
 #include "../Common.ush"
 #include "HairStrandsClusterCommon.ush"
-#include "HairStrandsMeshProjectionCommon.ush"
+#include "HairStrandsBindingCommon.ush"
 
 #if SHADER_DEFORM_GUIDE
 
@@ -73,7 +73,7 @@ float3 DisplacePosition(uint VertexIndex, float3 Pos)
 
 #if PERMUTATION_DEFORMATION == DEFORMATION_OFFSET_GUIDE_DYNAMIC
 	const uint RootIndex = SimPointToCurveBuffer[VertexIndex];
-	const float3 RootBarycentric = DecodeBarycentrics(SimRootBarycentricBuffer[RootIndex]);
+	const float3 RootBarycentric = UnpackBarycentrics(SimRootBarycentricBuffer[RootIndex]);
 
 	const uint TriangleIndex = SimRootToUniqueTriangleIndexBuffer[RootIndex];
 	const FHairMeshTriangle RestTriangle = GetTriangleTransformation(TriangleIndex, SimRestPositionBuffer);
diff --git a/Engine/Shaders/Private/HairStrands/HairStrandsInterpolation.usf b/Engine/Shaders/Private/HairStrands/HairStrandsInterpolation.usf
--- a/Engine/Shaders/Private/HairStrands/HairStrandsInterpolation.usf
+++ b/Engine/Shaders/Private/HairStrands/HairStrandsInterpolation.usf
@@ -1,7 +1,7 @@
 // Copyright Epic Games, Inc. All Rights Reserved.
 
 #include "../Common.ush"
-#include "HairStrandsMeshProjectionCommon.ush"
+#include "HairStrandsBindingCommon.ush"
 #include "HairStrandsVisibilityCommon.ush"
 #include "HairStrandsAABBCommon.ush" 
 #include "HairStrandsClusterCommon.ush"
@@ -224,7 +224,7 @@ FRootTriangleData GetRootTriangleData(uint VertexIndex, bool bSim)
 	{
 		Out.RootIndex		 = SimPointToCurveBuffer[VertexIndex];
 		Out.TriangleIndex	 = SimRootToUniqueTriangleIndexBuffer[Out.RootIndex];
-		Out.RootBarycentric  = DecodeBarycentrics(SimRootBarycentricBuffer[Out.RootIndex]);
+		Out.RootBarycentric  = UnpackBarycentrics(SimRootBarycentricBuffer[Out.RootIndex]);
 		Out.RestTriangle	 = GetTriangleTransformation(Out.TriangleIndex, SimRootRestPositionBuffer);
 		Out.DeformedTriangle = GetTriangleTransformation(Out.TriangleIndex, SimRootDeformedPositionBuffer);
 	}
@@ -232,7 +232,7 @@ FRootTriangleData GetRootTriangleData(uint VertexIndex, bool bSim)
 	{
 		Out.RootIndex		 = RenPointToCurveBuffer[VertexIndex];
 		Out.TriangleIndex	 = RenRootToUniqueTriangleIndexBuffer[Out.RootIndex];
-		Out.RootBarycentric  = DecodeBarycentrics(RenRootBarycentricBuffer[Out.RootIndex]);
+		Out.RootBarycentric  = UnpackBarycentrics(RenRootBarycentricBuffer[Out.RootIndex]);
 		Out.RestTriangle	 = GetTriangleTransformation(Out.TriangleIndex, RenRootRestPositionBuffer);
 		Out.DeformedTriangle = GetTriangleTransformation(Out.TriangleIndex, RenRootDeformedPositionBuffer);
 	}
diff --git a/Engine/Shaders/Private/HairStrands/HairStrandsMesh.usf b/Engine/Shaders/Private/HairStrands/HairStrandsMesh.usf
--- a/Engine/Shaders/Private/HairStrands/HairStrandsMesh.usf
+++ b/Engine/Shaders/Private/HairStrands/HairStrandsMesh.usf
@@ -1,7 +1,7 @@
 // Copyright Epic Games, Inc. All Rights Reserved.
 
 #include "../Common.ush"
-#include "HairStrandsMeshProjectionCommon.ush"
+#include "HairStrandsBindingCommon.ush"
 
 ////////////////////////////////////////////////////////////////////////
 // Deform hair mesh basd on RBF weights
@@ -374,7 +374,7 @@ void MainCS(uint3 DispatchThreadId : SV_DispatchThreadID)
 	if (UniqueTriangleIndex >= MaxUniqueTriangleCount)
 		return;
 	
-	const FHairTriangleIndex TriangleIndex = DecodeTriangleIndex(UniqueTriangleIndices[UniqueTriangleIndex]);
+	const FHairTriangleIndex TriangleIndex = UnpackTriangleIndex(UniqueTriangleIndices[UniqueTriangleIndex]);
 	if (TriangleIndex.SectionIndex >= MaxSectionCount)
 		return;
 
diff --git a/Engine/Shaders/Private/HairStrands/HairStrandsMeshProjectionCommon.ush b/Engine/Shaders/Private/HairStrands/HairStrandsMeshProjectionCommon.ush
deleted file mode 100644
