Move hair packing functions into HairStrandPack for sharing with simulation.

diff --git a/Engine/Plugins/Runtime/HairStrands/Shaders/Private/NiagaraDataInterfaceHairStrandsTemplate.ush b/Engine/Plugins/Runtime/HairStrands/Shaders/Private/NiagaraDataInterfaceHairStrandsTemplate.ush
--- a/Engine/Plugins/Runtime/HairStrands/Shaders/Private/NiagaraDataInterfaceHairStrandsTemplate.ush
+++ b/Engine/Plugins/Runtime/HairStrands/Shaders/Private/NiagaraDataInterfaceHairStrandsTemplate.ush
@@ -113,20 +113,9 @@ void BuildBoundingBox_{ParameterName}(in float3 NodePosition, out bool FunctionS
  * Utilities regarding strands curve
  * -----------------------------------------------------------------
  */
-struct FHairCurve
-{
-	uint PointIndex;
-	uint PointCount;
-};
-
 FHairCurve UnpackCurve_{ParameterName}(in uint InCurveIndex)
 {
-	const uint In = {ParameterName}_CurvesOffsetsBuffer[InCurveIndex];
-
-	FHairCurve Out = (FHairCurve)0;
-	Out.PointIndex = In & 0xFFFFFF;
-	Out.PointCount = In>>24;
-	return Out;
+	return UnpackHairCurve({ParameterName}_CurvesOffsetsBuffer[InCurveIndex]);
 }
 
 /* -----------------------------------------------------------------
@@ -891,7 +880,7 @@ void ComputeLocalState_{ParameterName}(in float3 RestPosition, in float4 RestOri
 	if( {ParameterName}_InterpolationMode >= 1)
 	{
 		const int StrandIndex = ExecIndex() / {ParameterName}_StrandSize;
-		const float2 ProjectionUV = UnpackBarycentrics({ParameterName}_RootBarycentricCoordinatesBuffer[StrandIndex]);
+		const float2 ProjectionUV = UnpackBarycentrics({ParameterName}_RootBarycentricCoordinatesBuffer[StrandIndex]).xy;
 		const uint TriangleIndex = {ParameterName}_RootToUniqueTriangleIndexBuffer[StrandIndex];
 
 		float3 RestTrianglePosition = float3(0, 0, 0);
@@ -912,7 +901,7 @@ void AttachNodeState_{ParameterName}(in float3 LocalPosition, in float4 LocalOri
 	if( {ParameterName}_InterpolationMode >= 1)
 	{
 		const int StrandIndex = ExecIndex() / {ParameterName}_StrandSize;
-		const float2 ProjectionUV = UnpackBarycentrics({ParameterName}_RootBarycentricCoordinatesBuffer[StrandIndex]);
+		const float2 ProjectionUV = UnpackBarycentrics({ParameterName}_RootBarycentricCoordinatesBuffer[StrandIndex]).xy;
 		const uint TriangleIndex = {ParameterName}_RootToUniqueTriangleIndexBuffer[StrandIndex];
 
 		float3 DeformedTrianglePosition = float3(0, 0, 0);
diff --git a/Engine/Shaders/Private/HairStrands/HairStrandsAttributeCommon.ush b/Engine/Shaders/Private/HairStrands/HairStrandsAttributeCommon.ush
--- a/Engine/Shaders/Private/HairStrands/HairStrandsAttributeCommon.ush
+++ b/Engine/Shaders/Private/HairStrands/HairStrandsAttributeCommon.ush
@@ -2,7 +2,7 @@
 
 #pragma once 
 
-#include "/Engine/Private/HairStrands/HairStrandsVertexFactoryCommon.ush"
+#include "/Engine/Private/HairStrands/HairStrandsPack.ush"
 
 ////////////////////////////////////////////////////////////////////////////////
 // Data stride in bytes
@@ -89,6 +89,9 @@ bool IsAttributeValid(uint InOffset)
 
 ////////////////////////////////////////////////////////////////////////////////
 // Attribute access with explicit resources
+
+#include "/Engine/Private/HairStrands/HairStrandsVertexFactoryCommon.ush"
+
 #if HAIR_STRAND_MESH_FACTORY
 
 float InternalGetHairStrandsWorldRadius(uint HairPrimitiveId, float2 VertexUV)
diff --git a/Engine/Shaders/Private/HairStrands/HairStrandsBindingCommon.ush b/Engine/Shaders/Private/HairStrands/HairStrandsBindingCommon.ush
--- a/Engine/Shaders/Private/HairStrands/HairStrandsBindingCommon.ush
+++ b/Engine/Shaders/Private/HairStrands/HairStrandsBindingCommon.ush
@@ -1,5 +1,7 @@
 // Copyright Epic Games, Inc. All Rights Reserved.
 
+#include "HairStrandsPack.ush"
+
 struct FHairMeshTriangle
 {
 	float3 P0;
@@ -19,29 +21,6 @@ struct FHairMeshTriangleWithUVs : FHairMeshTriangle
 	float2 UV2;
 };
 
-uint PackUVs(float2 UV)
-{
-	return (f32tof16(UV.x) & 0xFFFF) | ((f32tof16(UV.y) & 0xFFFF) << 16);
-}
-
-float PackUVsToFloat(float2 UV)
-{
-	return asfloat(PackUVs(UV));
-}
-
-float2 UnPackUVs(uint InUV)
-{
-	float2 Out;
-	Out.x = f16tof32(InUV & 0xFFFF);
-	Out.y = f16tof32((InUV >> 16) & 0xFFFF);
-	return Out;
-}
-
-float2 UnPackUVs(float InUV)
-{
-	return UnPackUVs(asuint(InUV));
-}
-
 float2 GetVertexUV(uint VId, Buffer<float2> InUVBuffer, Buffer<uint> InIndexBuffer, uint InIndexOffset, uint InMaxIndexCount, uint InMaxVertexCount, uint InMeshUVsChannelOffset, uint InMeshUVsChannelCount)
 {
 	const uint Index = min(InIndexOffset + VId, InMaxIndexCount - 1);
@@ -193,39 +172,6 @@ float3 TransformPoint(float3 P, float3 B, const FHairMeshTriangle RestTriangle,
 	return ToWorld(LocalP, B, DeformedTriangle);
 }
 
-uint PackBarycentrics(float2 B)
-{
-	return f32tof16(B.x) | (f32tof16(B.y) << 16);
-}
-
-float3 UnpackBarycentrics(uint E)
-{
-	float3 Out;
-	Out.x = f16tof32(E & 0xFFFF);
-	Out.y = f16tof32((E >> 16) & 0xFFFF);
-	Out.z = 1 - Out.x - Out.y;
-	return Out;
-}
-
-struct FHairTriangleIndex
-{
-	uint TriangleIndex;
-	uint SectionIndex;
-};
-
-uint PackTriangleIndex(uint TriangleIndex, uint SectionIndex)
-{
-	return ((SectionIndex & 0xFF)<<24) | (TriangleIndex & 0xFFFFFF);
-}
-
-FHairTriangleIndex UnpackTriangleIndex(uint Encoded)
-{
-	FHairTriangleIndex Out;
-	Out.SectionIndex = (Encoded>>24) & 0xFF;
-	Out.TriangleIndex = Encoded & 0xFFFFFF;
-	return Out;
-}
-
 // Apply RBF weight onto a control point
 float3 ApplyRBF(
 	float3 RestControlPoint,
diff --git a/Engine/Shaders/Private/HairStrands/HairStrandsPack.ush b/Engine/Shaders/Private/HairStrands/HairStrandsPack.ush
new file mode 100644
--- /dev/null
+++ b/Engine/Shaders/Private/HairStrands/HairStrandsPack.ush
@@ -0,0 +1,242 @@
+// Copyright Epic Games, Inc. All Rights Reserved.
+
+#pragma once
+
+#include "/Engine/Shared/HairStrandsDefinitions.h"
+
+////////////////////////////////////////////////////////////////////////////////////////////////////////////////////////////////
+// Hair control points
+
+struct FHairControlPoint
+{
+	float3 Position;
+	float  WorldRadius;
+	float  UCoord;
+	uint   Type;
+};
+
+uint4 PackHairControlPoint(
+	FHairControlPoint CP,
+	float3 InVF_PositionOffset,
+	float InVF_Radius)
+{
+	uint PackedAttribute =  (uint(clamp(CP.UCoord * 0xFFu, 0u, 0xFFu)) << 8) |
+							(uint(clamp((CP.WorldRadius / InVF_Radius) * 63u, 0u, 63u)) << 2u) |
+							(CP.Type & 0x3);
+	return uint4(f32tof16(CP.Position - InVF_PositionOffset), PackedAttribute);
+}
+
+FHairControlPoint UnpackHairControlPoint(
+	uint4 InPackedData,
+	float3 InVF_PositionOffset,
+	float InVF_Radius,
+	float InVF_RootScale,
+	float InVF_TipScale)
+{
+	const uint PackedAlpha	= asuint(InPackedData.a);
+
+	FHairControlPoint Out = (FHairControlPoint)0;
+	Out.Position	= f16tof32(InPackedData.xyz) + InVF_PositionOffset;
+	Out.UCoord		= ((PackedAlpha & 0xFF00) >> 8) / 255.f;
+	Out.WorldRadius	= ((PackedAlpha & 0x00FC) >> 2) / 63.f;
+	Out.Type		= PackedAlpha & 0x3;
+
+	Out.WorldRadius *= InVF_Radius * lerp(InVF_RootScale, InVF_TipScale, Out.UCoord);
+
+	return Out;	
+}
+
+float UnpackHairControlPointCoordU(uint4 InPackedData)
+{
+	const uint PackedAlpha = asuint(InPackedData.w);
+	return ((PackedAlpha & 0xFF00) >> 8) / 255.f;
+}
+
+float3 UnpackHairControlPointPosition(uint4 InPackedData, float3 InVF_PositionOffset)
+{
+	return f16tof32(InPackedData.xyz) + InVF_PositionOffset;
+}
+
+uint4 PackHairControlPointPosition(uint4 InPackedData, float3 InNewPosition, float3 InPositionOffset)
+{
+	return uint4(f32tof16(InNewPosition.xyz - InPositionOffset), InPackedData.w);
+}
+
+bool IsValidHairStrandsSegment(const FHairControlPoint CP0, const FHairControlPoint CP1)
+{
+	return !(CP0.Type == HAIR_CONTROLPOINT_END && CP1.Type == HAIR_CONTROLPOINT_START);
+}
+
+////////////////////////////////////////////////////////////////////////////////////////////////////////////////////////////////
+// Hair attributes
+
+float2 UnpackHairRootUV(uint In)
+{
+	float2 RootUV, RootIndex;
+	RootUV.x	= ((In) & 0x7FF) / 2047.f;			// Coord are encoded on 11 bits
+	RootUV.y	= ((In >> 11) & 0x7FF) / 2047.f;	// Coord are encoded on 11 bits 
+	RootIndex.x = ((In >> 22) & 0x1F);				// UDIM tile are encoded on 5bits
+	RootIndex.y = ((In >> 27) & 0x1F);				// UDIM tile are encoded on 5bits
+
+	// Invert V to compensate image origin flip. Similar to GetHairStrandsRootUV in HairCardsAttributeCommon.ush
+	RootUV.y = 1.0f - RootUV.y;
+
+	return RootUV = RootUV + RootIndex;
+}
+
+float UnpackHairLength(uint In)
+{
+	return f16tof32(In & 0xFFFF);
+}
+
+float UnpackHairSeed(uint In)
+{
+	return (In & 0xFF) / 255.f;
+}
+
+uint3 UnpackHairClumpID(uint In)
+{
+	return (In & 0xFFFF).xxx;
+}
+
+uint3 UnpackHairClumpID(uint2 In)
+{
+	return uint3(In.x & 0xFFFF, (In.x>>16) & 0xFFFF, (In.x>>32) & 0xFFFF);
+}
+
+float3 UnpackHairColor(uint In)
+{
+	float3 Out;
+	Out.x = (In       & 0x7FF) / 2047.f;
+	Out.y = ((In>>11) & 0x7FF) / 2047.f;
+	Out.z = ((In>>22) & 0x3FF) / 1023.f;
+	return Out;
+}
+
+float UnpackHairRoughness(uint In)
+{
+	return (In & 0xFF) / 255.f;
+}
+
+////////////////////////////////////////////////////////////////////////////////////////////////////////////////////////////////
+// Guide
+
+struct FGuideData
+{
+	uint3  VertexIndices;
+	float3 VertexLerps;
+	float3 VertexWeights;
+};
+
+FGuideData UnpackGuideData(uint In)
+{
+	FGuideData Out = (FGuideData)0;
+	Out.VertexWeights		= float3(1, 0, 0);
+	Out.VertexIndices		= 0;
+	Out.VertexLerps			= 0;
+	Out.VertexIndices.x		= In & 0xFFFFFF;
+	Out.VertexLerps.x		= UnpackR8(In >> 24);
+	return Out;
+}
+
+FGuideData UnpackGuideData(uint4 In)
+{
+	FGuideData Out = (FGuideData)0;
+	Out.VertexIndices.x = (In.x & 0xFFFF) | (((In.z      ) & 0xFF) << 16);
+	Out.VertexIndices.y = (In.x >> 16   ) | (((In.z >> 8 ) & 0xFF) << 16);
+	Out.VertexIndices.z = (In.y & 0xFFFF) | (((In.z >> 16) & 0xFF) << 16);
+
+	Out.VertexWeights	= float3(UnpackR8(In.y>>16), UnpackR8(In.y>>24), 0);
+	Out.VertexWeights.z = saturate(1.0f - (Out.VertexWeights.x + Out.VertexWeights.y));
+
+	Out.VertexLerps.x	= UnpackR8(In.w);
+	Out.VertexLerps.y	= UnpackR8(In.w >> 8);
+	Out.VertexLerps.z	= UnpackR8(In.w >> 16);
+	return Out;
+}
+
+////////////////////////////////////////////////////////////////////////////////////////////////////////////////////////////////
+// Curve
+
+struct FHairCurve
+{
+	uint PointIndex;
+	uint PointCount;
+};
+
+FHairCurve UnpackHairCurve(uint In)
+{
+	#if HAIR_MAX_NUM_POINT_PER_CURVE != 0xFF || HAIR_MAX_NUM_POINT_PER_GROUP != 0xFFFFFFFF
+	#error Update hair curve ENCODING_TYPE
+	#endif
+
+	FHairCurve Out = (FHairCurve)0;
+	Out.PointIndex = In & HAIR_MAX_NUM_POINT_PER_GROUP;
+	Out.PointCount = In >> 24;
+	return Out;
+}
+
+////////////////////////////////////////////////////////////////////////////////////////////////////////////////////////////////
+// UVs
+
+uint PackUVs(float2 UV)
+{
+	return (f32tof16(UV.x) & 0xFFFF) | ((f32tof16(UV.y) & 0xFFFF) << 16);
+}
+
+float PackUVsToFloat(float2 UV)
+{
+	return asfloat(PackUVs(UV));
+}
+
+float2 UnPackUVs(uint InUV)
+{
+	float2 Out;
+	Out.x = f16tof32(InUV & 0xFFFF);
+	Out.y = f16tof32((InUV >> 16) & 0xFFFF);
+	return Out;
+}
+
+float2 UnPackUVs(float InUV)
+{
+	return UnPackUVs(asuint(InUV));
+}
+
+////////////////////////////////////////////////////////////////////////////////////////////////////////////////////////////////
+// Barycentrics
+
+uint PackBarycentrics(float2 B)
+{
+	return f32tof16(B.x) | (f32tof16(B.y) << 16);
+}
+
+float3 UnpackBarycentrics(uint E)
+{
+	float3 Out;
+	Out.x = f16tof32(E & 0xFFFF);
+	Out.y = f16tof32((E >> 16) & 0xFFFF);
+	Out.z = 1 - Out.x - Out.y;
+	return Out;
+}
+
+////////////////////////////////////////////////////////////////////////////////////////////////////////////////////////////////
+// Triangle
+
+struct FHairTriangleIndex
+{
+	uint TriangleIndex;
+	uint SectionIndex;
+};
+
+uint PackTriangleIndex(uint TriangleIndex, uint SectionIndex)
+{
+	return ((SectionIndex & 0xFF)<<24) | (TriangleIndex & 0xFFFFFF);
+}
+
+FHairTriangleIndex UnpackTriangleIndex(uint Encoded)
+{
+	FHairTriangleIndex Out;
+	Out.SectionIndex = (Encoded>>24) & 0xFF;
+	Out.TriangleIndex = Encoded & 0xFFFFFF;
+	return Out;
+}
diff --git a/Engine/Shaders/Private/HairStrands/HairStrandsVertexFactoryCommon.ush b/Engine/Shaders/Private/HairStrands/HairStrandsVertexFactoryCommon.ush
--- a/Engine/Shaders/Private/HairStrands/HairStrandsVertexFactoryCommon.ush
+++ b/Engine/Shaders/Private/HairStrands/HairStrandsVertexFactoryCommon.ush
@@ -3,183 +3,12 @@
 #pragma once
 
 #include "/Engine/Shared/HairStrandsDefinitions.h"
+#include "/Engine/Private/HairStrands/HairStrandsPack.ush"
 
 #ifndef HAIR_STRANDS_PARAMETERS
 #define HAIR_STRANDS_PARAMETERS 0
 #endif
 
-////////////////////////////////////////////////////////////////////////////////////////////////////////////////////////////////
-// Hair control points
-
-struct FHairControlPoint
-{
-	float3 Position;
-	float  WorldRadius;
-	float  UCoord;
-	uint   Type;
-};
-
-uint4 PackHairControlPoint(
-	FHairControlPoint CP,
-	float3 InVF_PositionOffset,
-	float InVF_Radius)
-{
-	uint PackedAttribute =  (uint(clamp(CP.UCoord * 0xFFu, 0u, 0xFFu)) << 8) |
-							(uint(clamp((CP.WorldRadius / InVF_Radius) * 63u, 0u, 63u)) << 2u) |
-							(CP.Type & 0x3);
-	return uint4(f32tof16(CP.Position - InVF_PositionOffset), PackedAttribute);
-}
-
-FHairControlPoint UnpackHairControlPoint(
-	uint4 InPackedData,
-	float3 InVF_PositionOffset,
-	float InVF_Radius,
-	float InVF_RootScale,
-	float InVF_TipScale)
-{
-	const uint PackedAlpha	= asuint(InPackedData.a);
-
-	FHairControlPoint Out = (FHairControlPoint)0;
-	Out.Position	= f16tof32(InPackedData.xyz) + InVF_PositionOffset;
-	Out.UCoord		= ((PackedAlpha & 0xFF00) >> 8) / 255.f;
-	Out.WorldRadius	= ((PackedAlpha & 0x00FC) >> 2) / 63.f;
-	Out.Type		= PackedAlpha & 0x3;
-
-	Out.WorldRadius *= InVF_Radius * lerp(InVF_RootScale, InVF_TipScale, Out.UCoord);
-
-	return Out;	
-}
-
-float UnpackHairControlPointCoordU(uint4 InPackedData)
-{
-	const uint PackedAlpha = asuint(InPackedData.w);
-	return ((PackedAlpha & 0xFF00) >> 8) / 255.f;
-}
-
-float3 UnpackHairControlPointPosition(uint4 InPackedData, float3 InVF_PositionOffset)
-{
-	return f16tof32(InPackedData.xyz) + InVF_PositionOffset;
-}
-
-uint4 PackHairControlPointPosition(uint4 InPackedData, float3 InNewPosition, float3 InPositionOffset)
-{
-	return uint4(f32tof16(InNewPosition.xyz - InPositionOffset), InPackedData.w);
-}
-
-bool IsValidHairStrandsSegment(const FHairControlPoint CP0, const FHairControlPoint CP1)
-{
-	return !(CP0.Type == HAIR_CONTROLPOINT_END && CP1.Type == HAIR_CONTROLPOINT_START);
-}
-
-////////////////////////////////////////////////////////////////////////////////////////////////////////////////////////////////
-// Hair attributes
-
-float2 UnpackHairRootUV(uint In)
-{
-	float2 RootUV, RootIndex;
-	RootUV.x	= ((In) & 0x7FF) / 2047.f;			// Coord are encoded on 11 bits
-	RootUV.y	= ((In >> 11) & 0x7FF) / 2047.f;	// Coord are encoded on 11 bits 
-	RootIndex.x = ((In >> 22) & 0x1F);				// UDIM tile are encoded on 5bits
-	RootIndex.y = ((In >> 27) & 0x1F);				// UDIM tile are encoded on 5bits
-
-	// Invert V to compensate image origin flip. Similar to GetHairStrandsRootUV in HairCardsAttributeCommon.ush
-	RootUV.y = 1.0f - RootUV.y;
-
-	return RootUV = RootUV + RootIndex;
-}
-
-float UnpackHairLength(uint In)
-{
-	return f16tof32(In & 0xFFFF);
-}
-
-float UnpackHairSeed(uint In)
-{
-	return (In & 0xFF) / 255.f;
-}
-
-uint3 UnpackHairClumpID(uint In)
-{
-	return (In & 0xFFFF).xxx;
-}
-
-uint3 UnpackHairClumpID(uint2 In)
-{
-	return uint3(In.x & 0xFFFF, (In.x>>16) & 0xFFFF, (In.x>>32) & 0xFFFF);
-}
-
-float3 UnpackHairColor(uint In)
-{
-	float3 Out;
-	Out.x = (In       & 0x7FF) / 2047.f;
-	Out.y = ((In>>11) & 0x7FF) / 2047.f;
-	Out.z = ((In>>22) & 0x3FF) / 1023.f;
-	return Out;
-}
-
-float UnpackHairRoughness(uint In)
-{
-	return (In & 0xFF) / 255.f;
-}
-
-////////////////////////////////////////////////////////////////////////////////////////////////////////////////////////////////
-// Hair curve
-
-struct FGuideData
-{
-	uint3  VertexIndices;
-	float3 VertexLerps;
-	float3 VertexWeights;
-};
-
-FGuideData UnpackGuideData(uint In)
-{
-	FGuideData Out = (FGuideData)0;
-	Out.VertexWeights		= float3(1, 0, 0);
-	Out.VertexIndices		= 0;
-	Out.VertexLerps			= 0;
-	Out.VertexIndices.x		= In & 0xFFFFFF;
-	Out.VertexLerps.x		= UnpackR8(In >> 24);
-	return Out;
-}
-
-FGuideData UnpackGuideData(uint4 In)
-{
-	FGuideData Out = (FGuideData)0;
-	Out.VertexIndices.x = (In.x & 0xFFFF) | (((In.z      ) & 0xFF) << 16);
-	Out.VertexIndices.y = (In.x >> 16   ) | (((In.z >> 8 ) & 0xFF) << 16);
-	Out.VertexIndices.z = (In.y & 0xFFFF) | (((In.z >> 16) & 0xFF) << 16);
-
-	Out.VertexWeights	= float3(UnpackR8(In.y>>16), UnpackR8(In.y>>24), 0);
-	Out.VertexWeights.z = saturate(1.0f - (Out.VertexWeights.x + Out.VertexWeights.y));
-
-	Out.VertexLerps.x	= UnpackR8(In.w);
-	Out.VertexLerps.y	= UnpackR8(In.w >> 8);
-	Out.VertexLerps.z	= UnpackR8(In.w >> 16);
-	return Out;
-}
-
-////////////////////////////////////////////////////////////////////////////////////////////////////////////////////////////////
-// Hair curve
-
-struct FHairCurve
-{
-	uint PointIndex;
-	uint PointCount;
-};
-
-FHairCurve UnpackHairCurve(uint In)
-{
-	#if HAIR_MAX_NUM_POINT_PER_CURVE != 0xFF || HAIR_MAX_NUM_POINT_PER_GROUP != 0xFFFFFFFF
-	#error Update hair curve ENCODING_TYPE
-	#endif
-
-	FHairCurve Out = (FHairCurve)0;
-	Out.PointIndex = In & HAIR_MAX_NUM_POINT_PER_GROUP;
-	Out.PointCount = In >> 24;
-	return Out;
-}
-
 ///////////////////////////////////////////////////////////////////////////////////////////////////
 // Hair strands instance parameters
 
