Makes min3(), med3(), max3() and fastClamp() first class citizen to speed up TSR convolutions

diff --git a/Engine/Shaders/Private/FastMathThirdParty.ush b/Engine/Shaders/Private/FastMathThirdParty.ush
--- a/Engine/Shaders/Private/FastMathThirdParty.ush
+++ b/Engine/Shaders/Private/FastMathThirdParty.ush
@@ -91,85 +91,11 @@ float3 normalizeFast( float3 v )
 	return v * rsqrtFast( LengthSqr );
 }
 
-float4 fastClamp(float4 x, float4 Min, float4 Max)
-{
-#if COMPILER_PSSL
-	// Warning: only correct if Min is smaller than Max
-	return med3(x, Min, Max);
-#else
-	return clamp(x, Min, Max);
-#endif
-}
-
-float3 fastClamp(float3 x, float3 Min, float3 Max)
-{
-#if COMPILER_PSSL
-	// Warning: only correct if Min is smaller than Max
-	return med3(x, Min, Max);
+#if COMPILER_SUPPORTS_MED3
+	#define fastClamp(x, Min, Max) med3(x, Min, Max)
 #else
-	return clamp(x, Min, Max);
+	#define fastClamp(x, Min, Max) clamp(x, Min, Max)
 #endif
-}
-
-float2 fastClamp(float2 x, float2 Min, float2 Max)
-{
-#if COMPILER_PSSL
-	// Warning: only correct if Min is smaller than Max
-	return med3(x, Min, Max);
-#else
-	return clamp(x, Min, Max);
-#endif
-}
-
-float fastClamp(float x, float Min, float Max)
-{
-#if COMPILER_PSSL
-	// Warning: only correct if Min is smaller than Max
-	return med3(x, Min, Max);
-#else
-	return clamp(x, Min, Max);
-#endif
-}
-
-int4 fastClamp(int4 x, int4 Min, int4 Max)
-{
-#if COMPILER_PSSL
-	// Warning: only correct if Min is smaller than Max
-	return med3(x, Min, Max);
-#else
-	return clamp(x, Min, Max);
-#endif
-}
-
-int3 fastClamp(int3 x, int3 Min, int3 Max)
-{
-#if COMPILER_PSSL
-	// Warning: only correct if Min is smaller than Max
-	return med3(x, Min, Max);
-#else
-	return clamp(x, Min, Max);
-#endif
-}
-
-int2 fastClamp(int2 x, int2 Min, int2 Max)
-{
-#if COMPILER_PSSL
-	// Warning: only correct if Min is smaller than Max
-	return med3(x, Min, Max);
-#else
-	return clamp(x, Min, Max);
-#endif
-}
-
-int fastClamp(int x, int Min, int Max)
-{
-#if COMPILER_PSSL
-	// Warning: only correct if Min is smaller than Max
-	return med3(x, Min, Max);
-#else
-	return clamp(x, Min, Max);
-#endif
-}
 
 //
 // Trigonometric functions
diff --git a/Engine/Shaders/Private/TemporalSuperResolution/TSRCommon.ush b/Engine/Shaders/Private/TemporalSuperResolution/TSRCommon.ush
--- a/Engine/Shaders/Private/TemporalSuperResolution/TSRCommon.ush
+++ b/Engine/Shaders/Private/TemporalSuperResolution/TSRCommon.ush
@@ -427,22 +427,9 @@ tsr_short2x2 ClampPixelOffset(
 	tsr_short2x2 SamplePixelPos,
 	int2 MinPixelPos, int2 MaxPixelPos)
 {
-	tsr_short2 SamplePixelPos0 = dpv_lo(SamplePixelPos);
-	tsr_short2 SamplePixelPos1 = dpv_hi(SamplePixelPos);
-
-	SamplePixelPos0.x = min(SamplePixelPos0.x, tsr_short(MaxPixelPos.x));
-	SamplePixelPos0.x = max(SamplePixelPos0.x, tsr_short(MinPixelPos.x));
-		
-	SamplePixelPos0.y = min(SamplePixelPos0.y, tsr_short(MaxPixelPos.y));
-	SamplePixelPos0.y = max(SamplePixelPos0.y, tsr_short(MinPixelPos.y));
-	
-	SamplePixelPos1.x = min(SamplePixelPos1.x, tsr_short(MaxPixelPos.x));
-	SamplePixelPos1.x = max(SamplePixelPos1.x, tsr_short(MinPixelPos.x));
-		
-	SamplePixelPos1.y = min(SamplePixelPos1.y, tsr_short(MaxPixelPos.y));
-	SamplePixelPos1.y = max(SamplePixelPos1.y, tsr_short(MinPixelPos.y));
-
-	return dpv_interleave_registers(SamplePixelPos0, SamplePixelPos1);
+	SamplePixelPos[0] = fastClamp(SamplePixelPos[0], tsr_short2(MinPixelPos.xx), tsr_short2(MaxPixelPos.xx));
+	SamplePixelPos[1] = fastClamp(SamplePixelPos[1], tsr_short2(MinPixelPos.yy), tsr_short2(MaxPixelPos.yy));
+	return SamplePixelPos;
 }
 
 CALL_SITE_DEBUGLOC
@@ -450,13 +437,7 @@ tsr_short2 ClampPixelOffset(
 	tsr_short2 SamplePixelPos,
 	int2 MinPixelPos, int2 MaxPixelPos)
 {
-	SamplePixelPos.x = min(SamplePixelPos.x, tsr_short(MaxPixelPos.x));
-	SamplePixelPos.x = max(SamplePixelPos.x, tsr_short(MinPixelPos.x));
-		
-	SamplePixelPos.y = min(SamplePixelPos.y, tsr_short(MaxPixelPos.y));
-	SamplePixelPos.y = max(SamplePixelPos.y, tsr_short(MinPixelPos.y));
-
-	return SamplePixelPos;
+	return fastClamp(SamplePixelPos, tsr_short2(MinPixelPos), tsr_short2(MaxPixelPos));
 }
 
 CALL_SITE_DEBUGLOC
@@ -1024,7 +1005,7 @@ tsr_subpixel_payload GetClosestSubpixelPayload(FSubpixelNeighborhood Neighborhoo
 CALL_SITE_DEBUGLOC
 tsr_subpixel_payload CompressParallaxFactor(tsr_half ParallaxFactor)
 {
-	return clamp(
+	return fastClamp(
 		tsr_subpixel_payload((ParallaxFactor - tsr_half(1.0)) * tsr_half(float(SUB_PIXEL_PARALLAX_FACTOR_BIT_MASK) / MAX_PARALLAX_FACTOR)),
 		tsr_subpixel_payload(0), tsr_subpixel_payload(SUB_PIXEL_PARALLAX_FACTOR_BIT_MASK));
 }
@@ -1106,8 +1087,8 @@ uint2 EncodeSpatialAntiAliasingOffset(tsr_half2x2 TexelOffset)
 	tsr_half Multiply = tsr_half(float(Zero) / SPATIAL_ANTI_ALIASING_OFFSET_RANGE);
 
 	uint2 EncodedTexelOffset = 0;
-	EncodedTexelOffset |= clamp(uint2(int2(round(TexelOffset[0] * Multiply)) + int(1 + Zero)), uint(1), Mask) << 0;
-	EncodedTexelOffset |= clamp(uint2(int2(round(TexelOffset[1] * Multiply)) + int(1 + Zero)), uint(1), Mask) << SPATIAL_ANTI_ALIASING_OFFSET_BITDEPTH_PER_CHANNEL;
+	EncodedTexelOffset |= fastClamp(uint2(int2(round(TexelOffset[0] * Multiply)) + int(1 + Zero)), uint(1), Mask) << 0;
+	EncodedTexelOffset |= fastClamp(uint2(int2(round(TexelOffset[1] * Multiply)) + int(1 + Zero)), uint(1), Mask) << SPATIAL_ANTI_ALIASING_OFFSET_BITDEPTH_PER_CHANNEL;
 	return EncodedTexelOffset;
 }
 
@@ -1155,7 +1136,7 @@ uint EncodeHoleFillingVelocity(float2 PixelVelocity)
 	uint iEncodedAngle = uint(round(EncodedAngle)) & (AngleBitDepth - 1);
 
 	float EncodedLength = length(PixelVelocity) * LengthPixelPrecision;
-	uint iEncodedLength = clamp(uint(ceil(EncodedLength)), 0u, LengthBitDepth - 1u);
+	uint iEncodedLength = fastClamp(uint(ceil(EncodedLength)), 0u, LengthBitDepth - 1u);
 
 	uint EncodedHoleFillingVelocity = iEncodedAngle * LengthBitDepth + iEncodedLength;
 	return EncodedHoleFillingVelocity;
@@ -1237,8 +1218,8 @@ uint CompressSubpixelDepth(FTSRSubpixelDepth SubpixelDepth)
 {
 	uint CompressedSubpixelDepth = 0u;
 	CompressedSubpixelDepth |= f32tof16(SubpixelDepth.DeviceZ) << 17u;
-	CompressedSubpixelDepth |= uint(clamp(int(round(SubpixelDepth.Offset.x * 16.00 + 16.0)), 0, 31)) << 12u;
-	CompressedSubpixelDepth |= uint(clamp(int(round(SubpixelDepth.Offset.y * 16.00 + 16.0)), 0, 31)) <<  7u;
+	CompressedSubpixelDepth |= uint(fastClamp(int(round(SubpixelDepth.Offset.x * 16.00 + 16.0)), 0, 31)) << 12u;
+	CompressedSubpixelDepth |= uint(fastClamp(int(round(SubpixelDepth.Offset.y * 16.00 + 16.0)), 0, 31)) <<  7u;
 	CompressedSubpixelDepth |= (SubpixelDepth.Age) << 0u;
 	return CompressedSubpixelDepth;
 }
diff --git a/Engine/Shaders/Private/TemporalSuperResolution/TSRConvolutionNetwork.ush b/Engine/Shaders/Private/TemporalSuperResolution/TSRConvolutionNetwork.ush
--- a/Engine/Shaders/Private/TemporalSuperResolution/TSRConvolutionNetwork.ush
+++ b/Engine/Shaders/Private/TemporalSuperResolution/TSRConvolutionNetwork.ush
@@ -330,12 +330,12 @@ TLaneVector2D<ScalarType, VectorSize, LaneStrideX, LaneStrideY> Min3x3(TLaneVect
 	TLaneVector2D<ScalarType, VectorSize, LaneStrideX, LaneStrideY> CN;
 	AccessNeighborTexels3x1(CO, tsr_short2(1, 0), /* out */ CP, /* out */ CN);
 	
-	TLaneVector2D<ScalarType, VectorSize, LaneStrideX, LaneStrideY> MinO = min(CN, min(CP, CO));
+	TLaneVector2D<ScalarType, VectorSize, LaneStrideX, LaneStrideY> MinO = min3(CN, CP, CO);
 	TLaneVector2D<ScalarType, VectorSize, LaneStrideX, LaneStrideY> MinP;
 	TLaneVector2D<ScalarType, VectorSize, LaneStrideX, LaneStrideY> MinN;
 	AccessNeighborTexels3x1(MinO, tsr_short2(0, 1), /* out */ MinP, /* out */ MinN);
 
-	return min(MinN, min(MinP, MinO));
+	return min3(MinN, MinP, MinO);
 }
 
 CALL_SITE_DEBUGLOC
@@ -347,12 +347,12 @@ TLaneVector2D<ScalarType, VectorSize, LaneStrideX, LaneStrideY> Max3x3(TLaneVect
 	TLaneVector2D<ScalarType, VectorSize, LaneStrideX, LaneStrideY> CN;
 	AccessNeighborTexels3x1(CO, tsr_short2(1, 0), /* out */ CP, /* out */ CN);
 	
-	TLaneVector2D<ScalarType, VectorSize, LaneStrideX, LaneStrideY> MaxO = max(CN, max(CP, CO));
+	TLaneVector2D<ScalarType, VectorSize, LaneStrideX, LaneStrideY> MaxO = max3(CN, CP, CO);
 	TLaneVector2D<ScalarType, VectorSize, LaneStrideX, LaneStrideY> MaxP;
 	TLaneVector2D<ScalarType, VectorSize, LaneStrideX, LaneStrideY> MaxN;
 	AccessNeighborTexels3x1(MaxO, tsr_short2(0, 1), /* out */ MaxP, /* out */ MaxN);
 
-	return max(MaxN, max(MaxP, MaxO));
+	return max3(MaxN, MaxP, MaxO);
 }
 
 CALL_SITE_DEBUGLOC
@@ -368,18 +368,18 @@ void MinMax3x3(
 	TLaneVector2D<ScalarType, VectorSize, LaneStrideX, LaneStrideY> CN;
 	AccessNeighborTexels3x1(CO, tsr_short2(1, 0), /* out */ CP, /* out */ CN);
 	
-	TLaneVector2D<ScalarType, VectorSize, LaneStrideX, LaneStrideY> MinO = min(CN, min(CP, CO));
-	TLaneVector2D<ScalarType, VectorSize, LaneStrideX, LaneStrideY> MaxO = max(CN, max(CP, CO));
+	TLaneVector2D<ScalarType, VectorSize, LaneStrideX, LaneStrideY> MinO = min3(CN, CP, CO);
+	TLaneVector2D<ScalarType, VectorSize, LaneStrideX, LaneStrideY> MaxO = max3(CN, CP, CO);
 
 	TLaneVector2D<ScalarType, VectorSize, LaneStrideX, LaneStrideY> MinP;
 	TLaneVector2D<ScalarType, VectorSize, LaneStrideX, LaneStrideY> MinN;
 	AccessNeighborTexels3x1(MinO, tsr_short2(0, 1), /* out */ MinP, /* out */ MinN);
-	OutMin = min(MinN, min(MinP, MinO));
+	OutMin = min3(MinN, MinP, MinO);
 
 	TLaneVector2D<ScalarType, VectorSize, LaneStrideX, LaneStrideY> MaxP;
 	TLaneVector2D<ScalarType, VectorSize, LaneStrideX, LaneStrideY> MaxN;
 	AccessNeighborTexels3x1(MaxO, tsr_short2(0, 1), /* out */ MaxP, /* out */ MaxN);
-	OutMax = max(MaxN, max(MaxP, MaxO));
+	OutMax = max3(MaxN, MaxP, MaxO);
 }
 #else
 {
@@ -395,8 +395,14 @@ void MinMax3x3(
 CALL_SITE_DEBUGLOC
 template<typename FSampleType>
 void LMHOperator(FSampleType A, FSampleType B, FSampleType C, out FSampleType L, out FSampleType M, out FSampleType H)
+#if COMPILER_SUPPORTS_MED3
+{
+	L = min3(A, B, C);
+	M = med3(A, B, C);
+	H = max3(A, B, C);
+}
+#else
 {
-	// TODO: min3, max3, med3.
 	FSampleType X = min(B, C);
 	FSampleType Y = max(B, C);
 	
@@ -406,17 +412,23 @@ void LMHOperator(FSampleType A, FSampleType B, FSampleType C, out FSampleType L,
 	M = min(Z, Y);
 	H = max(Z, Y);
 }
+#endif
 
 // 3 samples median.
 CALL_SITE_DEBUGLOC
 template<typename FSampleType>
 FSampleType Median(FSampleType A, FSampleType B, FSampleType C)
+#if COMPILER_SUPPORTS_MED3
+{
+	return med3(A, B, C);
+}
+#else
 {
-	// TODO: med3.
 	FSampleType L, M, H;
 	LMHOperator(A, B, C, L, M, H);
 	return M;
 }
+#endif
 
 // 9 samples median using Smith1996.
 CALL_SITE_DEBUGLOC
@@ -433,11 +445,10 @@ FSampleType Median9(FSampleType In[9])
 	}
 	
 	// Second layer.
-	// TODO: min3, max3.
 	FSampleType M1[3];
-	M1[0] = max(max(L0[0], L0[1]), L0[2]);
+	M1[0] = max3(L0[0], L0[1], L0[2]);
 	M1[1] = Median(M0[0], M0[1], M0[2]);
-	M1[2] = min(min(H0[0], H0[1]), H0[2]);
+	M1[2] = min3(H0[0], H0[1], H0[2]);
 	
 	// Third layer.
 	return Median(M1[0], M1[1], M1[2]);
@@ -448,35 +459,34 @@ template<typename ScalarType, uint VectorSize, uint LaneStrideX, uint LaneStride
 TLaneVector2D<ScalarType, VectorSize, LaneStrideX, LaneStrideY> Median3x3(TLaneVector2D<ScalarType, VectorSize, LaneStrideX, LaneStrideY> Center)
 #if 1
 {
-	WriteToLDSForAccessNeighborTexel(Center);
-	
 	TLaneVector2D<ScalarType, VectorSize, LaneStrideX, LaneStrideY> C0 = Center;
-	TLaneVector2D<ScalarType, VectorSize, LaneStrideX, LaneStrideY> C1 = AccessNeighborTexel(Center, tsr_short2( 1, 0));
-	TLaneVector2D<ScalarType, VectorSize, LaneStrideX, LaneStrideY> C2 = AccessNeighborTexel(Center, tsr_short2(-1, 0));
+	TLaneVector2D<ScalarType, VectorSize, LaneStrideX, LaneStrideY> C1;
+	TLaneVector2D<ScalarType, VectorSize, LaneStrideX, LaneStrideY> C2;
+	AccessNeighborTexels3x1(C0, tsr_short2(1, 0), /* out */ C1, /* out */ C2);
 	
+	// First layer.
 	TLaneVector2D<ScalarType, VectorSize, LaneStrideX, LaneStrideY> M0L0;
 	TLaneVector2D<ScalarType, VectorSize, LaneStrideX, LaneStrideY> M0M0;
 	TLaneVector2D<ScalarType, VectorSize, LaneStrideX, LaneStrideY> M0H0;
 	LMHOperator(C0, C1, C2, M0L0, M0M0, M0H0);
 	
-	WriteToLDSForAccessNeighborTexel(M0L0);
-	TLaneVector2D<ScalarType, VectorSize, LaneStrideX, LaneStrideY> M0L1 = AccessNeighborTexel(M0L0, tsr_short2(0,  1));
-	TLaneVector2D<ScalarType, VectorSize, LaneStrideX, LaneStrideY> M0L2 = AccessNeighborTexel(M0L0, tsr_short2(0, -1));
+	TLaneVector2D<ScalarType, VectorSize, LaneStrideX, LaneStrideY> M0L1;
+	TLaneVector2D<ScalarType, VectorSize, LaneStrideX, LaneStrideY> M0L2;
+	AccessNeighborTexels3x1(M0L0, tsr_short2(0, 1), /* out */ M0L1, /* out */ M0L2);
 	
-	WriteToLDSForAccessNeighborTexel(M0M0);
-	TLaneVector2D<ScalarType, VectorSize, LaneStrideX, LaneStrideY> M0M1 = AccessNeighborTexel(M0M0, tsr_short2(0,  1));
-	TLaneVector2D<ScalarType, VectorSize, LaneStrideX, LaneStrideY> M0M2 = AccessNeighborTexel(M0M0, tsr_short2(0, -1));
+	TLaneVector2D<ScalarType, VectorSize, LaneStrideX, LaneStrideY> M0M1;
+	TLaneVector2D<ScalarType, VectorSize, LaneStrideX, LaneStrideY> M0M2;
+	AccessNeighborTexels3x1(M0M0, tsr_short2(0, 1), /* out */ M0M1, /* out */ M0M2);
 	
-	WriteToLDSForAccessNeighborTexel(M0H0);
-	TLaneVector2D<ScalarType, VectorSize, LaneStrideX, LaneStrideY> M0H1 = AccessNeighborTexel(M0H0, tsr_short2(0,  1));
-	TLaneVector2D<ScalarType, VectorSize, LaneStrideX, LaneStrideY> M0H2 = AccessNeighborTexel(M0H0, tsr_short2(0, -1));
+	TLaneVector2D<ScalarType, VectorSize, LaneStrideX, LaneStrideY> M0H1;
+	TLaneVector2D<ScalarType, VectorSize, LaneStrideX, LaneStrideY> M0H2;
+	AccessNeighborTexels3x1(M0H0, tsr_short2(0, 1), /* out */ M0H1, /* out */ M0H2);
 	
 	// Second layer.
-	// TODO: min3, max3.
 	TLaneVector2D<ScalarType, VectorSize, LaneStrideX, LaneStrideY> M1[3];
-	M1[0] = max(max(M0L0, M0L1), M0L2);
+	M1[0] = max3(M0L0, M0L1, M0L2);
 	M1[1] = Median(M0M0, M0M1, M0M2);
-	M1[2] = min(min(M0H0, M0H1), M0H2);
+	M1[2] = min3(M0H0, M0H1, M0H2);
 	
 	// Third layer.
 	return Median(M1[0], M1[1], M1[2]);
diff --git a/Engine/Shaders/Private/TemporalSuperResolution/TSRRejectShading.usf b/Engine/Shaders/Private/TemporalSuperResolution/TSRRejectShading.usf
--- a/Engine/Shaders/Private/TemporalSuperResolution/TSRRejectShading.usf
+++ b/Engine/Shaders/Private/TemporalSuperResolution/TSRRejectShading.usf
@@ -128,13 +128,13 @@ tsr_simd_halfC ClampFireFliersWithGuide(tsr_simd_halfC ToClamp, tsr_simd_halfC G
 	tsr_simd_halfC ToClampMax;
 	MinMax3x3(ToClamp, /* out */ ToClampMin, /* out */ ToClampMax);
 
-	tsr_simd_halfC ClampedGuide = clamp(Guide, ToClampMin, ToClampMax);
+	tsr_simd_halfC ClampedGuide = fastClamp(Guide, ToClampMin, ToClampMax);
 	
 	tsr_simd_halfC ClampedGuideMin;
 	tsr_simd_halfC ClampedGuideMax;
 	MinMax3x3(ClampedGuide, /* out */ ClampedGuideMin, /* out */ ClampedGuideMax);
 	
-	return clamp(ToClamp, ClampedGuideMin, ClampedGuideMax);
+	return fastClamp(ToClamp, ClampedGuideMin, ClampedGuideMax);
 }
 
 
@@ -184,7 +184,7 @@ void MainCS(
 			float2 TranslucencyTextureUV = ApplyScreenTransform(float2(ClampedFetchPixelPos), InputPixelPosToTranslucencyTextureUV);
 
 			ClampedFetchPixelPos = ClampPixelOffset(ClampedFetchPixelPos, InputPixelPosMin, InputPixelPosMax);
-			TranslucencyTextureUV = clamp(TranslucencyTextureUV, TranslucencyTextureUVMin, TranslucencyTextureUVMax);
+			TranslucencyTextureUV = fastClamp(TranslucencyTextureUV, TranslucencyTextureUVMin, TranslucencyTextureUVMax);
 
 			OriginalOpaqueInput.SetElement(SimdIndex, InputTexture[ClampedFetchPixelPos]);
 			#if CONFIG_COMPOSE_TRANSLUCENCY
@@ -304,7 +304,7 @@ void MainCS(
 		MinMax3x3(InputLuma, /* out */ InputLumaMin, /* out */ InputLumaMax);
 
 		// Clamp and contribue the history.		
-		tsr_simd_half ClampedPrevFlickeringHistory = clamp(PrevFlickeringHistory, InputLumaMin, InputLumaMax);
+		tsr_simd_half ClampedPrevFlickeringHistory = fastClamp(PrevFlickeringHistory, InputLumaMin, InputLumaMax);
 		tsr_simd_half FinalFlickeringHistory = InputLuma * tsr_simd_half::Const(BlendFinal) - ClampedPrevFlickeringHistory * tsr_simd_half::Const(BlendFinal - tsr_half(1.0));
 	
 		// Measure the flickering gradient of the current frame.
@@ -433,7 +433,7 @@ void MainCS(
 		FilteredBoxMin = FilteredBoxMin - BackbufferError;
 		FilteredBoxMax = FilteredBoxMax + BackbufferError;
 		
-		tsr_simd_halfC FlickeringClampedFilteredHistory = clamp(FilteredHistory, FilteredBoxMin, FilteredBoxMax);
+		tsr_simd_halfC FlickeringClampedFilteredHistory = fastClamp(FilteredHistory, FilteredBoxMin, FilteredBoxMax);
 
 		tsr_simd_halfC ClampedFilteredHistory;
 		#if CONFIG_MOIRE_DETECTION
@@ -459,7 +459,7 @@ void MainCS(
 			tsr_simd_halfC StableFilteredBoxMin = min(FilteredBoxMin, FilteredBoxMax - MoireErrorSize);
 			tsr_simd_halfC StableFilteredBoxMax = max(FilteredBoxMax, FilteredBoxMin + MoireErrorSize);
 			
-			ClampedFilteredHistory = clamp(FilteredHistory, StableFilteredBoxMin, StableFilteredBoxMax);
+			ClampedFilteredHistory = fastClamp(FilteredHistory, StableFilteredBoxMin, StableFilteredBoxMax);
 		}
 		#else
 			ClampedFilteredHistory = FlickeringClampedFilteredHistory;
@@ -501,8 +501,8 @@ void MainCS(
 		tsr_simd_halfC RawFactor      = saturate(tsr_simd_halfC::Const(1.0) - max(RawClampedEnergy - ClampError, tsr_simd_halfC::Const(0.0)) * rcp(max(Delta, BoxSize)));
 		tsr_simd_halfC FilteredFactor = saturate(tsr_simd_halfC::Const(1.0) - max(FilteredClampedEnergy - ClampError, tsr_simd_halfC::Const(0.0)) * rcp(max(Delta, BoxSize)));
 		
-		tsr_simd_half RawRejection      = min(RawFactor[0], min(RawFactor[1], RawFactor[2]));
-		tsr_simd_half FilteredRejection = min(FilteredFactor[0], min(FilteredFactor[1], FilteredFactor[2]));
+		tsr_simd_half RawRejection      = min3(RawFactor[0], RawFactor[1], RawFactor[2]);
+		tsr_simd_half FilteredRejection = min3(FilteredFactor[0], FilteredFactor[1], FilteredFactor[2]);
 		
 		RejectionBlendFinal = saturate(FilteredRejection + tsr_simd_half::Const(1.0) - bIsValid);
 		
diff --git a/Engine/Shaders/Private/TemporalSuperResolution/TSRUpdateHistory.usf b/Engine/Shaders/Private/TemporalSuperResolution/TSRUpdateHistory.usf
--- a/Engine/Shaders/Private/TemporalSuperResolution/TSRUpdateHistory.usf
+++ b/Engine/Shaders/Private/TemporalSuperResolution/TSRUpdateHistory.usf
@@ -221,6 +221,29 @@ groupshared tsr_half4 SharedArray11[GROUP_SHARED_SIZE];
 
 //------------------------------------------------------- WAVE
 
+#if COMPILER_SUPPORTS_MED3
+
+CALL_SITE_DEBUGLOC
+tsr_half3x2 med3(tsr_half3x2 a, tsr_half3x2 b, tsr_half3x2 c)
+{
+	return tsr_half3x2(
+		med3(a[0], b[0], c[0]),
+		med3(a[1], b[1], c[1]),
+		med3(a[2], b[2], c[2]));
+} 
+
+CALL_SITE_DEBUGLOC
+tsr_half4x2 med3(tsr_half4x2 a, tsr_half4x2 b, tsr_half4x2 c)
+{
+	return tsr_half4x2(
+		med3(a[0], b[0], c[0]),
+		med3(a[1], b[1], c[1]),
+		med3(a[2], b[2], c[2]),
+		med3(a[3], b[3], c[3]));
+} 
+
+#endif
+
 #if PLATFORM_SUPPORTS_WAVE_BROADCAST
 
 CALL_SITE_DEBUGLOC
@@ -357,7 +380,7 @@ T BilinearSampleColorHistory(Texture2D<T> Texture, float2 SampleUV)
 	{
 		float BilinearWeight = GetSampleWeight(BilinearInter, i);
 		uint2 PixelPos = GetSamplePixelCoord(BilinearInter, i);
-		PixelPos = clamp(PixelPos, PrevHistoryInfo_ViewportMin, PrevHistoryInfo_ViewportMax - 1);
+		PixelPos = fastClamp(PixelPos, PrevHistoryInfo_ViewportMin, PrevHistoryInfo_ViewportMax - 1);
 
 		T RawSample = Texture[PixelPos];
 		
@@ -822,8 +845,8 @@ void MainCS(
 		// Sample raw history
 		{
 			float2x2 PrevHistoryBufferUV = ApplyScreenTransform(PrevScreenPos, ScreenPosToPrevHistoryBufferUV);
-			PrevHistoryBufferUV[0] = clamp(PrevHistoryBufferUV[0], PrevHistoryInfo_UVViewportBilinearMin[0], PrevHistoryInfo_UVViewportBilinearMax[0]);
-			PrevHistoryBufferUV[1] = clamp(PrevHistoryBufferUV[1], PrevHistoryInfo_UVViewportBilinearMin[1], PrevHistoryInfo_UVViewportBilinearMax[1]);
+			PrevHistoryBufferUV[0] = fastClamp(PrevHistoryBufferUV[0], PrevHistoryInfo_UVViewportBilinearMin[0], PrevHistoryInfo_UVViewportBilinearMax[0]);
+			PrevHistoryBufferUV[1] = fastClamp(PrevHistoryBufferUV[1], PrevHistoryInfo_UVViewportBilinearMin[1], PrevHistoryInfo_UVViewportBilinearMax[1]);
 
 			FCatmullRomSamples Samples0 = GetBicubic2DCatmullRomSamples(
 				dpv_lo(PrevHistoryBufferUV), PrevHistoryInfo_Extent, PrevHistoryInfo_ExtentInverse);
@@ -832,8 +855,8 @@ void MainCS(
 			
 			#if CONFIG_GRAND_REPROJECTION
 				float2x2 GrandPrevHistoryBufferUV = ApplyScreenTransform(GrandPrevScreenPos, ScreenPosToGrandPrevHistoryBufferUV);
-				GrandPrevHistoryBufferUV[0] = clamp(GrandPrevHistoryBufferUV[0], PrevHistoryInfo_UVViewportBilinearMin[0], PrevHistoryInfo_UVViewportBilinearMax[0]);
-				GrandPrevHistoryBufferUV[1] = clamp(GrandPrevHistoryBufferUV[1], PrevHistoryInfo_UVViewportBilinearMin[1], PrevHistoryInfo_UVViewportBilinearMax[1]);
+				GrandPrevHistoryBufferUV[0] = fastClamp(GrandPrevHistoryBufferUV[0], PrevHistoryInfo_UVViewportBilinearMin[0], PrevHistoryInfo_UVViewportBilinearMax[0]);
+				GrandPrevHistoryBufferUV[1] = fastClamp(GrandPrevHistoryBufferUV[1], PrevHistoryInfo_UVViewportBilinearMin[1], PrevHistoryInfo_UVViewportBilinearMax[1]);
 
 				FCatmullRomSamples GrandSamples0 = GetBicubic2DCatmullRomSamples(
 					dpv_lo(GrandPrevHistoryBufferUV), PrevHistoryInfo_Extent, PrevHistoryInfo_ExtentInverse);
@@ -865,8 +888,8 @@ void MainCS(
 				UNROLL_N(BICUBIC_CATMULL_ROM_SAMPLES)
 				for (uint i = 0; i < BICUBIC_CATMULL_ROM_SAMPLES; i++)
 				{
-					float2 SampleUV0 = clamp(Samples0.UV[i], PrevHistoryInfo_UVViewportBilinearMin, PrevHistoryInfo_UVViewportBilinearMax);
-					float2 SampleUV1 = clamp(Samples1.UV[i], PrevHistoryInfo_UVViewportBilinearMin, PrevHistoryInfo_UVViewportBilinearMax);
+					float2 SampleUV0 = fastClamp(Samples0.UV[i], PrevHistoryInfo_UVViewportBilinearMin, PrevHistoryInfo_UVViewportBilinearMax);
+					float2 SampleUV1 = fastClamp(Samples1.UV[i], PrevHistoryInfo_UVViewportBilinearMin, PrevHistoryInfo_UVViewportBilinearMax);
 					
 					tsr_half2 KernelWeight = dpv_force_interleave_registers(tsr_half(Samples0.Weight[i]), tsr_half(Samples1.Weight[i]));
 
@@ -898,8 +921,8 @@ void MainCS(
 							#endif
 
 							{
-								float2 GrandSampleUV0 = clamp(GrandSamples0.UV[i], PrevHistoryInfo_UVViewportBilinearMin, PrevHistoryInfo_UVViewportBilinearMax);
-								float2 GrandSampleUV1 = clamp(GrandSamples1.UV[i], PrevHistoryInfo_UVViewportBilinearMin, PrevHistoryInfo_UVViewportBilinearMax);
+								float2 GrandSampleUV0 = fastClamp(GrandSamples0.UV[i], PrevHistoryInfo_UVViewportBilinearMin, PrevHistoryInfo_UVViewportBilinearMax);
+								float2 GrandSampleUV1 = fastClamp(GrandSamples1.UV[i], PrevHistoryInfo_UVViewportBilinearMin, PrevHistoryInfo_UVViewportBilinearMax);
 					
 								float2 GrandSampleUV = PixelId == 0 ? GrandSampleUV0 : GrandSampleUV1;
 						
@@ -1398,7 +1421,7 @@ void MainCS(
 	#endif
 
 	{
-		tsr_halfCx2 ClampedPrevHighFrequencyColor = clamp(PrevHighFrequencyColor, InputMinColor, InputMaxColor);
+		tsr_halfCx2 ClampedPrevHighFrequencyColor = fastClamp(PrevHighFrequencyColor, InputMinColor, InputMaxColor);
 		tsr_halfCx2 BlendedPrevHighFrequencyColor;
 		BlendedPrevHighFrequencyColor[0] = lerp(ClampedPrevHighFrequencyColor[0], PrevHighFrequencyColor[0], BlendClamp);
 		BlendedPrevHighFrequencyColor[1] = lerp(ClampedPrevHighFrequencyColor[1], PrevHighFrequencyColor[1], BlendClamp);
@@ -1454,9 +1477,9 @@ void MainCS(
 
 			// Clamp error imprecision.
 			bool2 bIsOverblurLegal = (
-				and(clamp(PrevReprojectionOverblur[0], -BoxSize[0], BoxSize[0]) == PrevReprojectionOverblur[0],
-				and(clamp(PrevReprojectionOverblur[1], -BoxSize[1], BoxSize[1]) == PrevReprojectionOverblur[1],
-				    clamp(PrevReprojectionOverblur[2], -BoxSize[2], BoxSize[2]) == PrevReprojectionOverblur[2])));
+				and(fastClamp(PrevReprojectionOverblur[0], -BoxSize[0], BoxSize[0]) == PrevReprojectionOverblur[0],
+				and(fastClamp(PrevReprojectionOverblur[1], -BoxSize[1], BoxSize[1]) == PrevReprojectionOverblur[1],
+				    fastClamp(PrevReprojectionOverblur[2], -BoxSize[2], BoxSize[2]) == PrevReprojectionOverblur[2])));
 
 			tsr_half2 PrevReprojectionOverblurScale = select(bIsOverblurLegal, tsr_half(1.0), tsr_half(0.0));
 
@@ -1481,16 +1504,16 @@ void MainCS(
 
 				tsr_halfCx2 CorrectedFinalHighFrequencyColor = FinalHighFrequencyColor + Correction;
 
-				FinalHighFrequencyColor = clamp(CorrectedFinalHighFrequencyColor, min(FinalHighFrequencyColor, PrevHistoryMomentMin), max(FinalHighFrequencyColor, PrevHistoryMomentMax));
+				FinalHighFrequencyColor = fastClamp(CorrectedFinalHighFrequencyColor, min(FinalHighFrequencyColor, PrevHistoryMomentMin), max(FinalHighFrequencyColor, PrevHistoryMomentMax));
 
 			#elif CONFIG_GRAND_REPROJECTION == 2
 				// Compute MeasuredResultant of the previous frame
 				tsr_halfCx2 PrevResultant;
-				PrevResultant[0] = clamp(GrandPrevToPrevResultant[0] * SafeRcp(GrandPrevToPrev[0]), tsr_half(0.0), tsr_half(1.0));
-				PrevResultant[1] = clamp(GrandPrevToPrevResultant[1] * SafeRcp(GrandPrevToPrev[1]), tsr_half(0.0), tsr_half(1.0));
-				PrevResultant[2] = clamp(GrandPrevToPrevResultant[2] * SafeRcp(GrandPrevToPrev[2]), tsr_half(0.0), tsr_half(1.0));
+				PrevResultant[0] = fastClamp(GrandPrevToPrevResultant[0] * SafeRcp(GrandPrevToPrev[0]), tsr_half(0.0), tsr_half(1.0));
+				PrevResultant[1] = fastClamp(GrandPrevToPrevResultant[1] * SafeRcp(GrandPrevToPrev[1]), tsr_half(0.0), tsr_half(1.0));
+				PrevResultant[2] = fastClamp(GrandPrevToPrevResultant[2] * SafeRcp(GrandPrevToPrev[2]), tsr_half(0.0), tsr_half(1.0));
 				#if CONFIG_SCENE_COLOR_ALPHA
-					PrevResultant[3] = clamp(GrandPrevToPrevResultant[3] * SafeRcp(GrandPrevToPrev[3]), tsr_half(0.0), tsr_half(1.0));
+					PrevResultant[3] = fastClamp(GrandPrevToPrevResultant[3] * SafeRcp(GrandPrevToPrev[3]), tsr_half(0.0), tsr_half(1.0));
 				#endif
 
 				FinalPrevHighFrequencyColorResultant = dpv_scale(PrevHighFrequencyColor, MeasuredResultant);
@@ -1512,7 +1535,7 @@ void MainCS(
 		tsr_half2 TranslucencyLumaHistory  = Luma4(GetColorChannel(PrevHistoryTranslucency));
 		tsr_half2 TranslucencyLumaFiltered = Luma4(GetColorChannel(FilteredTranslucencyColor));
 
-		tsr_half4x2 ClampedPrevHistoryTranslucency = clamp(PrevHistoryTranslucency, TranslucencyMinColor, TranslucencyMaxColor);
+		tsr_half4x2 ClampedPrevHistoryTranslucency = fastClamp(PrevHistoryTranslucency, TranslucencyMinColor, TranslucencyMaxColor);
 		
 		#if CONFIG_SCENE_COLOR_ALPHA
 			tsr_half2 HighFrequencyTranslucencyRejection = MeasureRejectionFactor(
diff --git a/Engine/Shaders/Public/LaneVectorization.ush b/Engine/Shaders/Public/LaneVectorization.ush
--- a/Engine/Shaders/Public/LaneVectorization.ush
+++ b/Engine/Shaders/Public/LaneVectorization.ush
@@ -444,3 +444,6 @@ TVECTOR_FUNCTION_2PARAMS(/* VectorSize = */ 1, dot);
 TVECTOR_FUNCTION_3PARAMS(VectorSize, clamp);
 TVECTOR_FUNCTION_3PARAMS(VectorSize, min3);
 TVECTOR_FUNCTION_3PARAMS(VectorSize, max3);
+#if COMPILER_SUPPORTS_MED3
+	TVECTOR_FUNCTION_3PARAMS(VectorSize, med3);
+#endif
diff --git a/Engine/Shaders/Public/OverloadMacros.ush b/Engine/Shaders/Public/OverloadMacros.ush
new file mode 100644
--- /dev/null
+++ b/Engine/Shaders/Public/OverloadMacros.ush
@@ -0,0 +1,26 @@
+// Copyright Epic Games, Inc. All Rights Reserved.
+
+/*=============================================================================
+	OverloadMacros.ush: Macros to conveninently overload standard functions
+	for different types.
+=============================================================================*/
+
+#pragma once
+
+/** Overloads a function that takes 1 paramater for all vector types */
+#define DECLARE_VECTOR_FUNCTION_OVERLOAD_1_PARAM(FunctionName,Type) \
+	Type##2 FunctionName(Type##2 v) { return Type##2(FunctionName(v.x),  FunctionName(v.y)); } \
+	Type##3 FunctionName(Type##3 v) { return Type##3(FunctionName(v.xy), FunctionName(v.z)); } \
+	Type##4 FunctionName(Type##4 v) { return Type##4(FunctionName(v.xy), FunctionName(v.zw)); }
+	
+/** Overloads a function that takes 2 paramaters for all vector types */
+#define DECLARE_VECTOR_FUNCTION_OVERLOAD_2_PARAM(FunctionName,Type) \
+	Type##2 FunctionName(Type##2 a, Type##2 b) { return Type##2(FunctionName(a.x , b.x ), FunctionName(a.y , b.y )); } \
+	Type##3 FunctionName(Type##3 a, Type##3 b) { return Type##3(FunctionName(a.xy, b.xy), FunctionName(a.z , b.z )); } \
+	Type##4 FunctionName(Type##4 a, Type##4 b) { return Type##4(FunctionName(a.xy, b.xy), FunctionName(a.zw, b.zw)); }
+	
+/** Overloads a function that takes 3 paramaters for all vector types */
+#define DECLARE_VECTOR_FUNCTION_OVERLOAD_3_PARAM(FunctionName,Type) \
+	Type##2 FunctionName(Type##2 a, Type##2 b, Type##2 c) { return Type##2(FunctionName(a.x , b.x , c.x ), FunctionName(a.y , b.y , c.y )); } \
+	Type##3 FunctionName(Type##3 a, Type##3 b, Type##3 c) { return Type##3(FunctionName(a.xy, b.xy, c.xy), FunctionName(a.z , b.z , c.z )); } \
+	Type##4 FunctionName(Type##4 a, Type##4 b, Type##4 c) { return Type##4(FunctionName(a.xy, b.xy, c.xy), FunctionName(a.zw, b.zw, c.zw)); }
diff --git a/Engine/Shaders/Public/Platform.ush b/Engine/Shaders/Public/Platform.ush
--- a/Engine/Shaders/Public/Platform.ush
+++ b/Engine/Shaders/Public/Platform.ush
@@ -66,6 +66,8 @@
 #endif
 
 #include "/Engine/Public/BindlessResources.ush"
+#include "/Engine/Public/OverloadMacros.ush"
+
 
 // ---------------------------------------------------- DDC invalidation
 // to support the console command "r.InvalidateShaderCache"
@@ -571,6 +573,10 @@ void ShaderYield()
 #define COMPILER_SUPPORTS_GATHER_UINT 0
 #endif
 
+#ifndef COMPILER_SUPPORTS_MED3
+#define COMPILER_SUPPORTS_MED3 0
+#endif
+
 #if ES3_1_PROFILE && !METAL_PROFILE
 #define HALF_TYPE half
 #define HALF2_TYPE half2
@@ -815,139 +821,87 @@ void WriteDispatchIndirectArgs(RWBuffer<uint> RWIndirectDispatchArgsBuffer, in u
 
 #if !COMPILER_SUPPORTS_MINMAX3
 
-float min3( float a, float b, float c )
+float min3(float a, float b, float c)
 {
-	return min( a, min( b, c ) );
+	return min(a, min(b, c));
 }
 
-float max3( float a, float b, float c )
+int min3(int a, int b, int c)
 {
-	return max( a, max( b, c ) );
+	return min(a, min(b, c));
 }
 
-float2 min3( float2 a, float2 b, float2 c )
+uint min3(uint a, uint b, uint c)
 {
-	return float2(
-		min3( a.x, b.x, c.x ),
-		min3( a.y, b.y, c.y )
-	);
+	return min(a, min(b, c));
 }
 
-float2 max3( float2 a, float2 b, float2 c )
-{
-	return float2(
-		max3( a.x, b.x, c.x ),
-		max3( a.y, b.y, c.y )
-	);
-}
+DECLARE_VECTOR_FUNCTION_OVERLOAD_3_PARAM(min3, float)
+DECLARE_VECTOR_FUNCTION_OVERLOAD_3_PARAM(min3, int)
+DECLARE_VECTOR_FUNCTION_OVERLOAD_3_PARAM(min3, uint)
 
-float3 max3( float3 a, float3 b, float3 c )
+float max3(float a, float b, float c)
 {
-	return float3(
-		max3( a.x, b.x, c.x ),
-		max3( a.y, b.y, c.y ),
-		max3( a.z, b.z, c.z )
-	);
+	return max(a, max(b, c));
 }
 
-float3 min3( float3 a, float3 b, float3 c )
+int max3(int a, int b, int c)
 {
-	return float3(
-		min3( a.x, b.x, c.x ),
-		min3( a.y, b.y, c.y ),
-		min3( a.z, b.z, c.z )
-	);
+	return max(a, max(b, c));
 }
 
-float4 min3( float4 a, float4 b, float4 c )
+uint max3(uint a, uint b, uint c)
 {
-	return float4(
-		min3( a.x, b.x, c.x ),
-		min3( a.y, b.y, c.y ),
-		min3( a.z, b.z, c.z ),
-		min3( a.w, b.w, c.w )
-	);
+	return max(a, max(b, c));
 }
 
-float4 max3( float4 a, float4 b, float4 c )
-{
-	return float4(
-		max3( a.x, b.x, c.x ),
-		max3( a.y, b.y, c.y ),
-		max3( a.z, b.z, c.z ),
-		max3( a.w, b.w, c.w )
-	);
-}
+DECLARE_VECTOR_FUNCTION_OVERLOAD_3_PARAM(max3, float)
+DECLARE_VECTOR_FUNCTION_OVERLOAD_3_PARAM(max3, int)
+DECLARE_VECTOR_FUNCTION_OVERLOAD_3_PARAM(max3, uint)
 
 #if PLATFORM_SUPPORTS_REAL_TYPES
 
-half min3( half a, half b, half c )
+half min3(half a, half b, half c)
 {
-	return min( a, min( b, c ) );
+	return min(a, min(b, c));
 }
 
-half max3( half a, half b, half c )
+int16_t min3(int16_t a, int16_t b, int16_t c)
 {
-	return max( a, max( b, c ) );
+	return min(a, min(b, c));
 }
 
-half2 min3( half2 a, half2 b, half2 c )
+uint16_t min3(uint16_t a, uint16_t b, uint16_t c)
 {
-	return half2(
-		min3( a.x, b.x, c.x ),
-		min3( a.y, b.y, c.y )
-	);
+	return min(a, min(b, c));
 }
 
-half2 max3( half2 a, half2 b, half2 c )
-{
-	return half2(
-		max3( a.x, b.x, c.x ),
-		max3( a.y, b.y, c.y )
-	);
-}
+DECLARE_VECTOR_FUNCTION_OVERLOAD_3_PARAM(min3, half)
+DECLARE_VECTOR_FUNCTION_OVERLOAD_3_PARAM(min3, int16_t)
+DECLARE_VECTOR_FUNCTION_OVERLOAD_3_PARAM(min3, uint16_t)
 
-half3 max3( half3 a, half3 b, half3 c )
+half max3(half a, half b, half c)
 {
-	return half3(
-		max3( a.x, b.x, c.x ),
-		max3( a.y, b.y, c.y ),
-		max3( a.z, b.z, c.z )
-	);
+	return max(a, max(b, c));
 }
 
-half3 min3( half3 a, half3 b, half3 c )
+int16_t max3(int16_t a, int16_t b, int16_t c)
 {
-	return half3(
-		min3( a.x, b.x, c.x ),
-		min3( a.y, b.y, c.y ),
-		min3( a.z, b.z, c.z )
-	);
+	return max(a, max(b, c));
 }
 
-half4 min3( half4 a, half4 b, half4 c )
+uint16_t max3(uint16_t a, uint16_t b, uint16_t c)
 {
-	return half4(
-		min3( a.x, b.x, c.x ),
-		min3( a.y, b.y, c.y ),
-		min3( a.z, b.z, c.z ),
-		min3( a.w, b.w, c.w )
-	);
+	return max(a, max(b, c));
 }
 
-half4 max3( half4 a, half4 b, half4 c )
-{
-	return half4(
-		max3( a.x, b.x, c.x ),
-		max3( a.y, b.y, c.y ),
-		max3( a.z, b.z, c.z ),
-		max3( a.w, b.w, c.w )
-	);
-}
+DECLARE_VECTOR_FUNCTION_OVERLOAD_3_PARAM(max3, half)
+DECLARE_VECTOR_FUNCTION_OVERLOAD_3_PARAM(max3, int16_t)
+DECLARE_VECTOR_FUNCTION_OVERLOAD_3_PARAM(max3, uint16_t)
 
-#endif
+#endif // PLATFORM_SUPPORTS_REAL_TYPES
 
-#endif
+#endif // !COMPILER_SUPPORTS_MINMAX3
 
 // https://devblogs.microsoft.com/directx/announcing-hlsl-2021/
 // HLSL 2021 supports Logical Operator Short Circuiting. To do vector bool operations, need to use and() or() select()
