Fix incorrect UV mapping when optionally drawing out scene colour when upsampling depth in debug draw passes.

diff --git a/Engine/Shaders/Private/PostProcessCompositePrimitives.usf b/Engine/Shaders/Private/PostProcessCompositePrimitives.usf
--- a/Engine/Shaders/Private/PostProcessCompositePrimitives.usf
+++ b/Engine/Shaders/Private/PostProcessCompositePrimitives.usf
@@ -256,7 +256,7 @@ void MainPopulateSceneDepthPS(
 
 	OutDepth = Texture2DSampleLevel(DepthTexture, DepthSampler, DepthUV, 0).r;
 #if FORCE_DRAW_COLOR
-	OutColor = Texture2DSampleLevel(ColorTexture, ColorSampler, ViewportUV, 0).rgba;
+	OutColor = Texture2DSampleLevel(ColorTexture, ColorSampler, DepthUV, 0).rgba;
 #else
 	OutColor = 0.0;
 #endif
