Trivial fixes for misc shader compilation warnings

diff --git a/Engine/Shaders/Private/GBufferHelpers.ush b/Engine/Shaders/Private/GBufferHelpers.ush
--- a/Engine/Shaders/Private/GBufferHelpers.ush
+++ b/Engine/Shaders/Private/GBufferHelpers.ush
@@ -391,7 +391,7 @@ void GBufferPostDecode(inout FGBufferData Ret, bool bChecker, bool bGetNormalize
 #elif ALLOW_STATIC_LIGHTING
 	Ret.GBufferAO = 1;
 	Ret.DiffuseIndirectSampleOcclusion = 0x0;
-	Ret.IndirectIrradiance = DecodeIndirectIrradiance(Ret.GenericAO.x);
+	Ret.IndirectIrradiance = half(DecodeIndirectIrradiance(Ret.GenericAO.x));
 #else
 	Ret.GBufferAO = Ret.GenericAO;
 	Ret.DiffuseIndirectSampleOcclusion = 0x0;
diff --git a/Engine/Shaders/Private/HairStrands/HairStrandsInterpolation.usf b/Engine/Shaders/Private/HairStrands/HairStrandsInterpolation.usf
--- a/Engine/Shaders/Private/HairStrands/HairStrandsInterpolation.usf
+++ b/Engine/Shaders/Private/HairStrands/HairStrandsInterpolation.usf
@@ -421,7 +421,7 @@ float3 ApplyDeformer(const FRootTriangleData RenTri, float3 InRestPostion, float
 void MainCS(uint2 DispatchThreadId : SV_DispatchThreadID)
 {
 	// 1. Find vertex index
-	uint VertexIndex = DispatchThreadId;
+	uint VertexIndex = DispatchThreadId.x;
 #if PERMUTATION_CULLING
 	if (HairStrandsVFTODO_bCullingEnable)
 	{
diff --git a/Engine/Shaders/Private/HairStrands/HairStrandsRaytracingGeometry.usf b/Engine/Shaders/Private/HairStrands/HairStrandsRaytracingGeometry.usf
--- a/Engine/Shaders/Private/HairStrands/HairStrandsRaytracingGeometry.usf
+++ b/Engine/Shaders/Private/HairStrands/HairStrandsRaytracingGeometry.usf
@@ -322,12 +322,12 @@ void MainCS(uint2 DispatchThreadId : SV_DispatchThreadID, uint LinearIndex : SV_
 	if (!bIsAlreadyInvalid)
 	{
 		// Fetch vertex index if culling is enabled
-		uint VertexIndex0 = DispatchThreadId;
+		uint VertexIndex0 = DispatchThreadId.x;
 		bool bIsValid = VertexIndex0 < VertexCount;
 		if (bIsValid && HairStrandsVF_bCullingEnable)
 		{
 			const uint VertexCountAfterCulling = HairStrandsVF_CullingIndirectBuffer[3];
-			const uint FetchIndex0 = DispatchThreadId;
+			const uint FetchIndex0 = DispatchThreadId.x;
 			bIsValid = FetchIndex0 < VertexCountAfterCulling;
 			if (bIsValid)
 			{
diff --git a/Engine/Shaders/Private/HairStrands/HairStrandsVisibilityRasterCompute.usf b/Engine/Shaders/Private/HairStrands/HairStrandsVisibilityRasterCompute.usf
--- a/Engine/Shaders/Private/HairStrands/HairStrandsVisibilityRasterCompute.usf
+++ b/Engine/Shaders/Private/HairStrands/HairStrandsVisibilityRasterCompute.usf
@@ -1406,8 +1406,8 @@ void RasterMultiSampleCS(uint DispatchThreadID : SV_DispatchThreadID, uint Group
 				
 				if (bVisible)
 				{
-					const float2 SP0Clipped = lerp(SP0, SP1, T.x);
-					const float2 SP1Clipped = lerp(SP0, SP1, T.y);
+					const float2 SP0Clipped = lerp(SP0, SP1, T.x).xy;
+					const float2 SP1Clipped = lerp(SP0, SP1, T.y).xy;
 					const bool bIsSteep = abs(SP1.x - SP0.x) < abs(SP1.y - SP0.y);
 					const float X0 = bIsSteep ? min(SP0Clipped.y, SP1Clipped.y) : min(SP0Clipped.x, SP1Clipped.x);
 					const float X1 = bIsSteep ? max(SP0Clipped.y, SP1Clipped.y) : max(SP0Clipped.x, SP1Clipped.x);
diff --git a/Engine/Shaders/Private/HairStrands/HairStrandsVoxelRasterCompute.usf b/Engine/Shaders/Private/HairStrands/HairStrandsVoxelRasterCompute.usf
--- a/Engine/Shaders/Private/HairStrands/HairStrandsVoxelRasterCompute.usf
+++ b/Engine/Shaders/Private/HairStrands/HairStrandsVoxelRasterCompute.usf
@@ -66,8 +66,8 @@ void AllocatePageIndex(uint2 DispatchThreadId : SV_DispatchThreadID)
 		const bool bSupportVoxelization = (bDoesMacroGroupSupportVoxelization >> MacroGroupId) & 0x1;
 		if (CPU_bUseCPUData > 0)
 		{
-			Bound.Min = CPU_TranslatedWorldMinAABB[MacroGroupId];
-			Bound.Max = CPU_TranslatedWorldMaxAABB[MacroGroupId];
+			Bound.Min = CPU_TranslatedWorldMinAABB[MacroGroupId].xyz;
+			Bound.Max = CPU_TranslatedWorldMaxAABB[MacroGroupId].xyz;
 		}
 		else
 		{
@@ -99,7 +99,7 @@ void AllocatePageIndex(uint2 DispatchThreadId : SV_DispatchThreadID)
 			float3 MacroGroupSize = Bound.Max - Bound.Min;
 			if (CPU_bUseCPUData > 0) 
 			{ 
-				PageIndexResolution = CPU_PageIndexResolution[MacroGroupId]; 
+				PageIndexResolution = CPU_PageIndexResolution[MacroGroupId].xyz; 
 			}
 			else
 			{
diff --git a/Engine/Shaders/Private/HeightFogCommon.ush b/Engine/Shaders/Private/HeightFogCommon.ush
--- a/Engine/Shaders/Private/HeightFogCommon.ush
+++ b/Engine/Shaders/Private/HeightFogCommon.ush
@@ -292,7 +292,7 @@ float3 ComputeHistoryVolumeUVFromNDC(float4 NDCPosition)
 {
 	NDCPosition.xy /= NDCPosition.w;
 	float3 VolumeUV = float3(NDCPosition.xy * float2(.5f, -.5f) + .5f, ComputeNormalizedZSliceFromDepth(NDCPosition.w));
-	return min(VolumeUV * float3(View.VolumetricFogViewGridUVToPrevViewRectUV.xy * float3(View.VolumetricFogPrevViewGridRectUVToResourceUV.xy, 1.0), 1.0), float3(View.VolumetricFogPrevUVMax, 1.0));
+	return min(VolumeUV * float3(View.VolumetricFogViewGridUVToPrevViewRectUV.xy * View.VolumetricFogPrevViewGridRectUVToResourceUV.xy, 1.0), float3(View.VolumetricFogPrevUVMax, 1.0));
 }
 float3 ComputeHistoryVolumeUV_DEPRECATED(float3 WorldPosition, float4x4 WorldToClip)
 {
diff --git a/Engine/Shaders/Private/PostProcessVisualizeDOF.usf b/Engine/Shaders/Private/PostProcessVisualizeDOF.usf
--- a/Engine/Shaders/Private/PostProcessVisualizeDOF.usf
+++ b/Engine/Shaders/Private/PostProcessVisualizeDOF.usf
@@ -24,8 +24,8 @@ void VisualizeDOFPS(
 {
 	OutColor = 0;
 
-	float SceneDeviceZ = DepthTexture.SampleLevel(DepthSampler, ApplyScreenTransform(SvPosition, SvPositionToColorUV), 0).r;
-	half3 SceneColor = Texture2DSample(ColorTexture, ColorSampler, ApplyScreenTransform(SvPosition, SvPositionToDepthUV)).rgb;
+	float SceneDeviceZ = DepthTexture.SampleLevel(DepthSampler, ApplyScreenTransform(SvPosition.xy, SvPositionToColorUV), 0).r;
+	half3 SceneColor = Texture2DSample(ColorTexture, ColorSampler, ApplyScreenTransform(SvPosition.xy, SvPositionToDepthUV)).rgb;
 
 	float SceneDepth = ConvertFromDeviceZ(SceneDeviceZ);
 
