Substrate - update to in shader decals
- Now mfp fades out correctly to 0 (hidding bottom layers)
- Transition to no SSSProfile is now seemless
- Fuzz fades out with base color opacity
- Aniso and HAziness fades out with roughness and other specular feature
- To be discussed

diff --git a/Engine/Shaders/Private/DBufferDecalShared.ush b/Engine/Shaders/Private/DBufferDecalShared.ush
--- a/Engine/Shaders/Private/DBufferDecalShared.ush
+++ b/Engine/Shaders/Private/DBufferDecalShared.ush
@@ -183,14 +183,15 @@ void ApplyDBufferData(in FStrataDBuffer StrataDBufferData, inout FStrataPixelHea
 
 				#if STRATA_DBUFFER_RESPONSE_BASECOLOR
 					BaseColor = BaseColor * StrataDBufferData.OneMinusCoverage_BaseColor + StrataDBufferData.BaseColor;
-					if (BSDF_GETHASFUZZ(BSDF))
-					{
-						SLAB_FUZZ_COLOR(BSDF) = SLAB_FUZZ_COLOR(BSDF) * StrataDBufferData.OneMinusCoverage_BaseColor + StrataDBufferData.BaseColor;
-					}
 
-					// Disable SSS when the coverage reach one to preserve the decal color
-					SLAB_SSSMFP(BSDF) = lerp(0.f, SLAB_SSSMFP(BSDF), StrataDBufferData.OneMinusCoverage_BaseColor);
-					if (all(SLAB_SSSMFP(BSDF) == 0.0f))
+					// Fade out fuzz to avoid overbright decal.
+					SLAB_FUZZ_AMOUNT(BSDF) *= StrataDBufferData.OneMinusCoverage_BaseColor;
+
+					// Disable SSS when the coverage reach one to preserve the decal color. 
+					// That operation also progressively makes non top layer disapear.
+					// Since this is ran before the strata tree is processed, we must use TmpMFP and not SLAB_SSSMFP.
+					BSDF.TmpMFP = BSDF.TmpMFP * StrataDBufferData.OneMinusCoverage_BaseColor;
+					if (StrataDBufferData.OneMinusCoverage_BaseColor == 0.0f)
 					{
 						BSDF_SETSSSTYPE(BSDF, SSS_TYPE_INVALID);
 					}
@@ -200,6 +201,15 @@ void ApplyDBufferData(in FStrataDBuffer StrataDBufferData, inout FStrataPixelHea
 					Metallic  = Metallic  * StrataDBufferData.OneMinusCoverage_Roughness + StrataDBufferData.Metallic;
 					Specular  = Specular  * StrataDBufferData.OneMinusCoverage_Roughness + StrataDBufferData.Specular;
 					Roughness = Roughness * StrataDBufferData.OneMinusCoverage_Roughness + StrataDBufferData.Roughness;
+
+					// Also fade out other specular features since we consider decal to have not specular features.
+					if (BSDF_GETHASHAZINESS(BSDF))
+					{
+						FHaziness Haziness = UnpackHaziness(SLAB_HAZINESS(BSDF));
+						Haziness.Weight *= StrataDBufferData.OneMinusCoverage_Roughness;
+						SLAB_HAZINESS(BSDF) = PackHaziness(Haziness);
+					}
+					SLAB_ANISOTROPY(BSDF) *= StrataDBufferData.OneMinusCoverage_Roughness;
 				#endif
 
 					float3 DiffuseAlbedo = 0;
