Substrate - deferred DBuffer application now fades out mroe data to match in shader DBuffer application.

diff --git a/Engine/Shaders/Private/Strata/StrataDbuffer.usf b/Engine/Shaders/Private/Strata/StrataDbuffer.usf
--- a/Engine/Shaders/Private/Strata/StrataDbuffer.usf
+++ b/Engine/Shaders/Private/Strata/StrataDbuffer.usf
@@ -151,7 +151,6 @@ void MainCS(
 		// 2. Load DBuffer data and apply them
 		float3 F0 = 0;
 		float3 DiffuseAlbedo = 0;
-		float3 FuzzColor = 0;
 		float Roughness = 0;
 		{
 			Roughness = StrataGetBSDFRoughness(BSDF);
@@ -166,13 +165,6 @@ void MainCS(
 			Roughness = Roughness * DBufferData.OneMinusCoverage_Roughness + DBufferData.Roughness;
 			BaseColor = BaseColor * DBufferData.OneMinusCoverage_BaseColor + DBufferData.BaseColor;
 
-			#if !STRATA_FASTPATH
-			if (BSDF_GETHASFUZZ(BSDF))
-			{
-				FuzzColor = SLAB_FUZZ_COLOR(BSDF) * DBufferData.OneMinusCoverage_BaseColor + DBufferData.BaseColor;
-			}
-			#endif
-
 			ConvertFromMetalness(BaseColor,	Specular, Metallic, DiffuseAlbedo, F0);
 		}
 
@@ -217,6 +209,14 @@ void MainCS(
 
 				if (BSDF_GETHASF90(BSDF) || BSDF_GETHASHAZINESS(BSDF))
 				{
+					if (BSDF_GETHASHAZINESS(BSDF))
+					{
+						FHaziness Haziness = UnpackHaziness(SLAB_HAZINESS(BSDF));
+						Haziness.Weight *= DBufferData.OneMinusCoverage_Roughness;
+						const uint PackedHaziness = PackHaziness(Haziness);
+						MaterialTextureArrayUAV[uint3(PixelCoord, RWOffset)] = (MaterialTextureArrayUAV[uint3(PixelCoord, RWOffset)] & 0x0000FFFF) | (PackedHaziness << 16);
+					}
+
 					++RWOffset;
 				}
 
@@ -225,8 +225,8 @@ void MainCS(
 				if (bHasSSS || bIsSimpleVolume)
 				{
 					// Disable SSS when the coverage reach one to preserve the decal color
-					SLAB_SSSMFP(BSDF) = lerp(0.f, SLAB_SSSMFP(BSDF), DBufferData.OneMinusCoverage_BaseColor);
-					if (all(SLAB_SSSMFP(BSDF) == 0.0f))
+					SLAB_SSSMFP(BSDF) *= DBufferData.OneMinusCoverage_BaseColor;
+					if (DBufferData.OneMinusCoverage_BaseColor == 0.0f)
 					{
 						BSDF_SETSSSTYPE(BSDF, SSS_TYPE_INVALID);
 					}
@@ -254,7 +254,10 @@ void MainCS(
 
 				if (BSDF_GETHASFUZZ(BSDF))
 				{
-					MaterialTextureArrayUAV[uint3(PixelCoord, RWOffset)] = PackFuzz(FuzzColor, SLAB_FUZZ_AMOUNT(BSDF), SLAB_FUZZ_ROUGHNESS(BSDF), Dither);
+					MaterialTextureArrayUAV[uint3(PixelCoord, RWOffset)] = PackFuzz(
+						SLAB_FUZZ_COLOR(BSDF), 
+						SLAB_FUZZ_AMOUNT(BSDF) * DBufferData.OneMinusCoverage_BaseColor, 
+						SLAB_FUZZ_ROUGHNESS(BSDF), Dither);
 					RWOffset++;
 				}
 			}
@@ -263,29 +266,45 @@ void MainCS(
 				const uint PackedDiffuse20Bits = PackR7G7B6Gamma2(DiffuseAlbedo, Dither);
 				const uint PackedF024Bits = PackRGBA8(float4(F0, 0.0f));
 
-				// Roughnes
-				MaterialTextureArrayUAV[uint3(PixelCoord, 0)] = (MaterialTextureArrayUAV[uint3(PixelCoord, 0)] & 0x00FFFFFF) | (PackR8(Roughness) << 24);
+				const uint Data0 = MaterialTextureArrayUAV[uint3(PixelCoord, 0)];
+				const uint Data1 = MaterialTextureArrayUAV[uint3(PixelCoord, 1)];
+				const uint Data2 = MaterialTextureArrayUAV[uint3(PixelCoord, 2)];
+
+				// Clear coat disappear with coverage
+				const uint PackedClearCoat16Bits = ((Data1 >> 12) & 0xFF00) | ((Data2 >> 24) & 0x00FF);
+				FHaziness Haziness = UnpackHaziness(PackedClearCoat16Bits);
+				Haziness.Weight *= DBufferData.OneMinusCoverage_BaseColor;
+				const uint PackedHaziness = PackHaziness(Haziness);
+
+				// Roughness
+				MaterialTextureArrayUAV[uint3(PixelCoord, 0)] = (Data0 & 0x00FFFFFF) | (PackR8(Roughness) << 24);
 
 				// Diffuse Albedo & F0
-				MaterialTextureArrayUAV[uint3(PixelCoord, 1)] = (MaterialTextureArrayUAV[uint3(PixelCoord, 1)] & 0xFFF00000) | PackedDiffuse20Bits;
-				MaterialTextureArrayUAV[uint3(PixelCoord, 2)] = (MaterialTextureArrayUAV[uint3(PixelCoord, 2)] & 0xFF000000) | PackedF024Bits;
+				// MRT1 20 PackedDiffuse and 8 high bits from PackedClearCoat
+				MaterialTextureArrayUAV[uint3(PixelCoord, 1)] = PackedDiffuse20Bits	| ((PackedHaziness & 0xFF00) << 12);
+				// MRT2 24 PackedF0      and 8 low  bits from PackedClearCoat
+				MaterialTextureArrayUAV[uint3(PixelCoord, 2)] = PackedF024Bits		| (PackedHaziness << 24);
 			}
 			else if (OptimisedLegacyMode == SINGLE_OPTLEGACYMODE_CLOTH)
 			{
 				// Roughness is only stored in the TopLayer texture.
 
 				// Diffuse Albedo & F0
-				const uint PackedDiffuse20Bits = PackR7G7B6Gamma2(DiffuseAlbedo, Dither);
-				const uint PackedF020Bits = PackR7G7B6Gamma2(F0, Dither);
-				const uint PackedFuzzColor20bits = PackR7G7B6Gamma2(FuzzColor, Dither);
-
+				const uint PackedDiffuse20Bits   = PackR7G7B6Gamma2(DiffuseAlbedo, Dither);
+				const uint PackedF020Bits        = PackR7G7B6Gamma2(F0, Dither);
+				const uint PackedFuzzColor20bits = PackR7G7B6Gamma2(SLAB_FUZZ_COLOR(BSDF), Dither);
+				const uint PackedFuzzAmount8bits = PackR8(SLAB_FUZZ_AMOUNT(BSDF) * DBufferData.OneMinusCoverage_BaseColor);
+
+				const uint FuzzAmountOffset = HEADER_SINGLEENCODING_BIT_COUNT + HEADER_SINGLE_OPTLEGACYMODE_BIT_COUNT;
+				const uint FuzzAmountMask = (1u << FuzzAmountOffset) - 1u;
+				MaterialTextureArrayUAV[uint3(PixelCoord, 0)] = (MaterialTextureArrayUAV[uint3(PixelCoord, 0)] & FuzzAmountMask) | (PackedFuzzAmount8bits << FuzzAmountOffset);
 				MaterialTextureArrayUAV[uint3(PixelCoord, 1)] = (PackedDiffuse20Bits | ((PackedFuzzColor20bits & 0xFFC00) << 10));
 				MaterialTextureArrayUAV[uint3(PixelCoord, 2)] = (PackedF020Bits | ((PackedFuzzColor20bits & 0x3FF) << 20));
 			}
 			else if (OptimisedLegacyMode == SINGLE_OPTLEGACYMODE_SSSWRAP || OptimisedLegacyMode == SINGLE_OPTLEGACYMODE_SSSWRAP_THIN)
 			{
 				// Roughness is only stored in the TopLayer texture.
-				const float3 MFP = lerp(0.f, SLAB_SSSMFP(BSDF), DBufferData.OneMinusCoverage_BaseColor);
+				const float3 MFP = SLAB_SSSMFP(BSDF) * DBufferData.OneMinusCoverage_BaseColor;
 
 				// DiffuseAlbedo & F0
 				const uint PackedDiffuse20Bits = PackR7G7B6Gamma2(DiffuseAlbedo, Dither);
@@ -307,13 +326,27 @@ void MainCS(
 				const uint PackedDiffuse20Bits = PackR7G7B6Gamma2(DiffuseAlbedo, Dither);
 				const uint PackedF020Bits = PackR7G7B6Gamma2(F0, Dither);
 
+				FStrataSubsurfaceHeader SSSDataHeader;
+				SSSDataHeader.Bytes = MaterialTextureArrayUAV[uint3(PixelCoord, FirstSliceStoringStrataSSSData + 0)];
+
+				float ProfileScl= StrataSubSurfaceHeaderGetProfileRadiusScale(SSSDataHeader) * DBufferData.OneMinusCoverage_BaseColor;
+				uint  ProfileId = StrataSubSurfaceHeaderGetProfileId(SSSDataHeader);
+				StrataSubSurfaceHeaderSetProfile(SSSDataHeader, ProfileScl, DBufferData.OneMinusCoverage_BaseColor == 0.0f ? 0 : ProfileId);
+				// Note We cannot change the type using StrataSubSurfaceHeaderSetSSSType(SSSDataHeader, SSS_TYPE_INVALID);
+				// Because we are stuck using the SINGLE_OPTLEGACYMODE_SSSPROFILE mode.
+				// So some SSSProfile material with decal will look difference when using deferred DBuffer instead of in shader application.
+
 				MaterialTextureArrayUAV[uint3(PixelCoord, 0)] = (MaterialTextureArrayUAV[uint3(PixelCoord, 0)] & 0x00FFFFFF) | (PackR8(Roughness) << 24);
-				MaterialTextureArrayUAV[uint3(PixelCoord, 1)] = (MaterialTextureArrayUAV[uint3(PixelCoord, 1)] & 0xFF000000) | PackedDiffuse20Bits;
+				MaterialTextureArrayUAV[uint3(PixelCoord, 1)] = PackedDiffuse20Bits  | (PackR8(ProfileScl) << 24);
 				MaterialTextureArrayUAV[uint3(PixelCoord, 2)] = (MaterialTextureArrayUAV[uint3(PixelCoord, 2)] & 0xFF000000) | PackedF020Bits;
 				
-				FStrataSubsurfaceExtras SSSDataExtras;
-				StrataSubsurfaceExtrasSetBaseColor(SSSDataExtras, DiffuseAlbedo);
-				StrataStoreSubsurfaceExtras(MaterialTextureArrayUAV, FirstSliceStoringStrataSSSData, PixelCoord, SSSDataExtras.Bytes);
+				StrataStoreSubsurfaceHeader(MaterialTextureArrayUAV, FirstSliceStoringStrataSSSData, PixelCoord, SSSDataHeader.Bytes);
+				if (StrataSubSurfaceHeaderHasExtras(SSSDataHeader))
+				{
+					FStrataSubsurfaceExtras SSSDataExtras;
+					StrataSubsurfaceExtrasSetBaseColor(SSSDataExtras, DiffuseAlbedo);
+					StrataStoreSubsurfaceExtras(MaterialTextureArrayUAV, FirstSliceStoringStrataSSSData, PixelCoord, SSSDataExtras.Bytes);
+				}
 			}
 		}
 		#endif // !STRATA_FASTPATH
