Substrate - added comments about TmpMFP.

diff --git a/Engine/Shaders/Private/Strata/Strata.ush b/Engine/Shaders/Private/Strata/Strata.ush
--- a/Engine/Shaders/Private/Strata/Strata.ush
+++ b/Engine/Shaders/Private/Strata/Strata.ush
@@ -135,7 +135,7 @@ struct FStrataBSDF
 
 	float3				Emissive;		// All bsdf can potientially have emissive. This is never written out to the strata gbuffer but instead accumulated in the emissive buffer during the base pass.
 	float				ThicknessCm;	// This must be kept in order to be able to normalize the mean free path to comply with our volumetric material math all done in a normalised slab of medium.
-	float3				TmpMFP;
+	float3				TmpMFP;			// TmpMFP is used to store user specified MFP while we do not know if this strata BSDF is going to be using simple volume lighting. So we always store the MFP on the side to not override SSSPROFILE radius scale sharing the same register.
 
 	float				TopLayerDataWeight;// The weight this BSDF contributes to the top layer (normal and roughness for SSR and SSAO for instance)
 #endif
@@ -1498,7 +1498,7 @@ FStrataData GetStrataSlabBSDF(
 #endif // !STRATA_FASTPATH
 
 #if STRATA_INLINE_SHADING
-	// At this stage we do not know if this strata BSDF is going to be using simple volume lighting so always store the MFP on the side to not override SSSPROFILE data.
+	// At this stage we do not know if this strata BSDF is going to be using simple volume lighting. So we always store the MFP on the side to not override SSSPROFILE radius scale sharing the same register.
 	StrataData.InlinedBSDF.TmpMFP = SSSMFP * SSSMFPScale;
 
 	StrataData.InlinedBSDF.Coverage = 1.0f;
