Substrate - advanced vis debug now uses STRATA_CLAMPED_BSDF_COUNT and fixed debug view when there is more than 1 BSDF.

diff --git a/Engine/Shaders/Private/Strata/Strata.ush b/Engine/Shaders/Private/Strata/Strata.ush
--- a/Engine/Shaders/Private/Strata/Strata.ush
+++ b/Engine/Shaders/Private/Strata/Strata.ush
@@ -1129,7 +1129,7 @@ void StrataPackOutStrataTree(
 #define WRITE_UINT(Data) ExtraMaterialDataUAV[uint3(OutCoord++, OutCoordY, SliceStoringDebugStrataTreeData)] = Data;
 #define WRITE_FLOAT(Data) ExtraMaterialDataUAV[uint3(OutCoord++, OutCoordY, SliceStoringDebugStrataTreeData)] = asuint(Data);
 #define WRITE_FLOAT3(Data) WRITE_FLOAT(Data.x);WRITE_FLOAT(Data.y);WRITE_FLOAT(Data.z);
-#define WRITE_FLOAT4(Data) WRITE_FLOAT(Data.x);WRITE_FLOAT(Data.y);WRITE_FLOAT(Data.z);;WRITE_FLOAT(Data.w);
+#define WRITE_FLOAT4(Data) WRITE_FLOAT(Data.x);WRITE_FLOAT(Data.y);WRITE_FLOAT(Data.z);WRITE_FLOAT(Data.w);
 
 	// Write out header
 	int OutCoord = 0;
@@ -1147,8 +1147,8 @@ void StrataPackOutStrataTree(
 #if STRATA_COMPILER == 1
 	LOOP
 #else
-	StrataTree.BSDFCount = min(StrataTree.BSDFCount, STRATA_MAX_BSDF_COUNT);
-	STRATA_UNROLL_N(STRATA_MAX_BSDF_COUNT)
+	StrataTree.BSDFCount = min(StrataTree.BSDFCount, STRATA_CLAMPED_BSDF_COUNT);
+	STRATA_UNROLL_N(STRATA_CLAMPED_BSDF_COUNT)
 #endif
 	for(int i=0; i< StrataTree.BSDFCount; ++i)
 	{
@@ -1244,7 +1244,7 @@ FStrataBSDF StrataUnpackInStrataTreeBSDF(
 	inout FStrataTreeHeader StrataTreeHeader,
 	in int SliceStoringDebugStrataTreeData)
 {
-	int OutCoord = BSDFIndex * (STRATA_INLINE_SHADING ? 42 : 33);		// Skip BSDFs
+	int OutCoord = BSDFIndex * (STRATA_INLINE_SHADING ? 40 : 31);		// Skip BSDFs
 	int OutCoordY= 1;
 
 	FStrataBSDF BSDF = (FStrataBSDF)0;
