Path Tracer: Add basic Substrate support

Implement the core slab BSDF in the path tracer. When Strata is enabled, enlarge the payload slightly to be able to encode a slab without needing to rely on Substrate's packing routines. This allows the path tracer to be more independent of the project settings. Multiple slabs are handled by stochastically returning a single slab from the hit shader. This approaches greatly simplifies the handling of SSS slabs and refractive/translucent cases in the future.

Also added a cvar that lets you toggle (at runtime) to the fully simplified form of the material. This doesn't appear that much faster, but it could be interesting to validate the appearance of the simplification procedure. It is gated by an extra cvar to avoid compiling extra permutations by default.

diff --git a/Engine/Shaders/Private/PathTracing/Material/PathTracingEye.ush b/Engine/Shaders/Private/PathTracing/Material/PathTracingEye.ush
--- a/Engine/Shaders/Private/PathTracing/Material/PathTracingEye.ush
+++ b/Engine/Shaders/Private/PathTracing/Material/PathTracingEye.ush
@@ -26,13 +26,13 @@ FEyeData PrepareEyeData(FPathTracingPayload Payload, float3 V_World)
 	FEyeData Data = (FEyeData)0;
 
 	Data.Basis = GetTangentBasis(Payload.WorldNormal);
-	Data.Alpha = GetGGXAlpha(Payload.Roughness);
+	Data.Alpha = GetGGXAlpha(Payload.GetEyeRoughness());
 	Data.V = mul(Data.Basis, V_World);
 
 	const float NoV = saturate(Data.V.z);
-	Data.Spec = ComputeGGXSpecEnergyTermsRGB(Payload.Roughness, NoV, Payload.SpecularColor);
+	Data.Spec = ComputeGGXSpecEnergyTermsRGB(Payload.GetEyeRoughness(), NoV, Payload.SpecularColor);
 
-	Data.DiffWeight = (1 - Data.Spec.E) * Payload.DiffuseColor;
+	Data.DiffWeight = (1 - Luminance(Data.Spec.E)) * Payload.DiffuseColor;
 
 	// Probability of picking diffuse lobe vs. specular lobe
 	Data.LobeProb = LobeSelectionProb(Data.DiffWeight, Data.Spec.E);
@@ -161,7 +161,7 @@ FMaterialSample Eye_SampleMaterial(
 		
 		Result.AddLobeWithMIS(SpecWeight, SpecPdf, 1.0 - Data.LobeProb);
 		Result.Pdf += Data.LobeProb * DiffPdf;
-		Result.Roughness = Payload.Roughness;
+		Result.Roughness = Payload.GetEyeRoughness();
 		Result.ScatterType = PATHTRACER_SCATTER_SPECULAR;
 	}
 	Result.Weight *= Payload.BSDFOpacity;
diff --git a/Engine/Shaders/Private/PathTracing/Material/PathTracingHair.ush b/Engine/Shaders/Private/PathTracing/Material/PathTracingHair.ush
--- a/Engine/Shaders/Private/PathTracing/Material/PathTracingHair.ush
+++ b/Engine/Shaders/Private/PathTracing/Material/PathTracingHair.ush
@@ -2,6 +2,7 @@
 
 #pragma once
 
+#include "PathTracingMaterialCommon.ush"
 
 #define PATH_TRACER_HAIR_SHADING_MODEL				1 // 0: zinke diffuse (a simple, cheap model for performance testing)
                                                       // 1: chiang principled hair (a sophisticated model with specular reflection/transmission)
@@ -52,20 +53,20 @@ FMaterialSample Hair_SampleMaterial(
 	float3x3 HairBasis = GetHairBasis(Payload.WorldTangent, V_World);
 
 	// sample a random position across the width of the fiber
-	float h = 2 * RandSample.x - 1;
+	const float h = 2 * RandSample.x - 1;
 	// create a normal vector in the plane tangent to the fiber
-	float2 N = float2(h, sqrt(1 - h * h));
+	const float2 N = float2(h, sqrt(1 - h * h));
 	// sample a cosine weighted direction
-	float4 SampledValue = CosineSampleHemisphere(RandSample.yz);
+	const float4 SampledValue = CosineSampleHemisphere(RandSample.yz);
 	// rotate this sampled direction about the (random) normal vector
-	float3 L = float3(SampledValue.x * N.y + SampledValue.z * N.x,
-					  SampledValue.z * N.y - SampledValue.x * N.x,
-					  SampledValue.y);
+	const float3 L = float3(SampledValue.x * N.y + SampledValue.z * N.x,
+			                SampledValue.z * N.y - SampledValue.x * N.x,
+		                    SampledValue.y);
 	// finally transform into world space
-	float3 L_World = mul(L, HairBasis);
+	const float3 L_World = mul(L, HairBasis);
 
-	float Pdf = ZinkeHairDiffuse(L);
-	return CreateMaterialSample(L_World, Payload.BaseColor, Pdf, sign(dot(Payload.WorldGeoNormal, L_World)), 1.0, PATHTRACER_SCATTER_DIFFUSE);
+	const float Pdf = ZinkeHairDiffuse(L);
+	return CreateMaterialSample(L_World, Payload.GetBaseColor(), Pdf, sign(dot(Payload.WorldGeoNormal, L_World)), 1.0, PATHTRACER_SCATTER_DIFFUSE);
 }
 
 FMaterialEval Hair_EvalMaterial(
@@ -75,12 +76,12 @@ FMaterialEval Hair_EvalMaterial(
 	float2 DiffuseSpecularScale
 )
 {
-	float3x3 HairBasis = GetHairBasis(Payload.WorldTangent, V_World);
+	const float3x3 HairBasis = GetHairBasis(Payload.WorldTangent, V_World);
 
-	float3 L = mul(HairBasis, L_World);
+	const float3 L = mul(HairBasis, L_World);
 
-	float Pdf = ZinkeHairDiffuse(L);
-	return CreateMaterialEval(Payload.BaseColor * DiffuseSpecularScale.x, Pdf);
+	const float Pdf = ZinkeHairDiffuse(L);
+	return CreateMaterialEval(Payload.GetBaseColor() * DiffuseSpecularScale.x, Pdf);
 }
 
 
@@ -196,7 +197,7 @@ struct FHairData
 	float3 A2;
 	float3 A3;
 	float4 LobePdf;
-	float4 LobeCdf; // Normalized sum to help select a lobe
+	float3 LobeCdf; // Normalized sum to help select a lobe
 };
 
 FHairData PrepareHairData(float3 V, float3 Color, float LongitudinalRoughness, float AzimuthalRoughness, float h, float Specular)
@@ -280,19 +281,8 @@ FHairData PrepareHairData(float3 V, float3 Color, float LongitudinalRoughness, f
 	// Chiang et al. - Eq(6)
 	Result.A3 = Result.A2 * F * T * rcp(1.00001 - F * T); // avoid dividing by 0 when F*T=1
 
-	Result.LobePdf = float4(
-		Result.A0,
-		max3(Result.A1.x, Result.A1.y, Result.A1.z),
-		max3(Result.A2.x, Result.A2.y, Result.A2.z),
-		max3(Result.A3.x, Result.A3.y, Result.A3.z)
-	);
-
-	Result.LobeCdf.x = Result.LobePdf.x;
-	Result.LobeCdf.y = Result.LobeCdf.x + Result.LobePdf.y;
-	Result.LobeCdf.z = Result.LobeCdf.y + Result.LobePdf.z;
-	Result.LobeCdf.w = Result.LobeCdf.z + Result.LobePdf.w;
-
-	Result.LobePdf /= Result.LobeCdf.w + 1e-6;
+	Result.LobeCdf = LobeSelectionCdf(Result.A0, Result.A1, Result.A2, Result.A3);
+	Result.LobePdf = LobeSelectionPdf(Result.LobeCdf);
 
 	return Result;
 }
@@ -306,45 +296,42 @@ FMaterialSample Hair_SampleMaterial(
 
 	const float3 V = mul(HairBasis, V_World); // TODO: simplify this since HairBasis is already sort of aligned with V ?
 	const float h = 2 * Payload.GetHairPrimitiveUV().y - 1; // remap back to [-1,1]
-	const float AzimuthalRoughness = lerp(1.0, 0.1, Payload.Metallic);
-	const FHairData HairData = PrepareHairData(V, Payload.BaseColor, Payload.Roughness, AzimuthalRoughness, h, Payload.Specular);
+	const float LongitudinalRoughness = Payload.GetHairLongitudinalRoughness();
+	const float AzimuthalRoughness = lerp(1.0, 0.1, Payload.GetHairAzimuthalRoughness());
+	const FHairData HairData = PrepareHairData(V, Payload.GetBaseColor(), LongitudinalRoughness, AzimuthalRoughness, h, Payload.GetHairSpecular());
 
 	// 1) Decide which lobe to sample
-	float u = RandSample.z * HairData.LobeCdf.w;
+	float u = RandSample.z;
 	int p = 0;
 	float v = 0, SinThetaOp = 0;
 	if (u < HairData.LobeCdf.x)
 	{
-		u /= HairData.LobeCdf.x;
+		u = RescaleRandomNumber(u, 0.0, HairData.LobeCdf.x);
 		p = 0;
 		v = HairData.v0;
 		SinThetaOp = HairData.SinThetaO.x;
 	}
 	else if (u < HairData.LobeCdf.y)
 	{
-		u -= HairData.LobeCdf.x;
-		u /= HairData.LobeCdf.y - HairData.LobeCdf.x;
+		u = RescaleRandomNumber(u, HairData.LobeCdf.x, HairData.LobeCdf.y);
 		p = 1;
 		v = HairData.v1;
 		SinThetaOp = HairData.SinThetaO.y;
 	}
 	else if (u < HairData.LobeCdf.z)
 	{
-		u -= HairData.LobeCdf.y;
-		u /= HairData.LobeCdf.z - HairData.LobeCdf.y;
+		u = RescaleRandomNumber(u, HairData.LobeCdf.y, HairData.LobeCdf.z);
 		p = 2;
 		v = HairData.v2;
 		SinThetaOp = HairData.SinThetaO.z;
 	}
 	else
 	{
-		u -= HairData.LobeCdf.z;
-		u /= HairData.LobeCdf.w - HairData.LobeCdf.z;
+		u = RescaleRandomNumber(u, HairData.LobeCdf.z, 1.0);
 		p = 3;
 		v = HairData.v3;
 		SinThetaOp = HairData.SinThetaO.w;
 	}
-	u = saturate(u);
 
 	// 2) Sample longitudinal angle
 	const float CosTheta = 1 + v * log(RandSample.x + (1 - RandSample.x) * exp(-2 / v));
@@ -360,7 +347,7 @@ FMaterialSample Hair_SampleMaterial(
 	// transform direction into world space
 	const float3 L_World = mul(L, HairBasis);
 
-	FMaterialSample Result = CreateMaterialSample(L_World, 0.0, 0.0, sign(dot(Payload.WorldGeoNormal, L_World)), Payload.Roughness, PATHTRACER_SCATTER_SPECULAR);
+	FMaterialSample Result = CreateMaterialSample(L_World, 0.0, 0.0, sign(dot(Payload.WorldGeoNormal, L_World)), LongitudinalRoughness, PATHTRACER_SCATTER_SPECULAR);
 
 	const float InvS = rcp(HairData.s);
 	const float Pdf0 = Mp(SinThetaI, HairData.SinThetaO.x, HairData.v0) * Np(PhiI, 0, InvS, HairData.GammaO, HairData.GammaT);
@@ -387,10 +374,10 @@ FMaterialEval Hair_EvalMaterial(
 
 	const float3 V = mul(HairBasis, V_World);
 	const float3 L = mul(HairBasis, L_World);
-
 	const float h = 2 * Payload.GetHairPrimitiveUV().y - 1; // remap back to [-1,1]
-	const float AzimuthalRoughness = lerp(1.0, 0.1, Payload.Metallic);
-	const FHairData HairData = PrepareHairData(V, Payload.BaseColor, Payload.Roughness, AzimuthalRoughness, h, Payload.Specular);
+	const float LongitudinalRoughness = Payload.GetHairLongitudinalRoughness();
+	const float AzimuthalRoughness = lerp(1.0, 0.1, Payload.GetHairAzimuthalRoughness());
+	const FHairData HairData = PrepareHairData(V, Payload.GetBaseColor(), LongitudinalRoughness, AzimuthalRoughness, h, Payload.GetHairSpecular());
 
 	// precompute data based of lighting direction
 	const float SinThetaI = L.z;
diff --git a/Engine/Shaders/Private/PathTracing/Material/PathTracingLambert.ush b/Engine/Shaders/Private/PathTracing/Material/PathTracingLambert.ush
--- a/Engine/Shaders/Private/PathTracing/Material/PathTracingLambert.ush
+++ b/Engine/Shaders/Private/PathTracing/Material/PathTracingLambert.ush
@@ -9,13 +9,19 @@ FMaterialSample Lambert_SampleMaterial(
 	FPathTracingPayload Payload,
 	float3 RandSample)
 {
-	float3 N_World = Payload.WorldNormal;
+	const float3 N_World = Payload.WorldNormal;
+	const float3 BaseColor = Payload.GetBaseColor();
 
 	float4 SampledValue = CosineSampleHemisphere(RandSample.xy);
 
 	float3 L_World = TangentToWorld(SampledValue.xyz, N_World);
+#if SIMPLIFIED_MATERIAL_SHADER
+	float ShadowTerminator = 1;
+#else
+	float ShadowTerminator = ShadowTerminatorTerm(L_World, N_World, Payload.WorldSmoothNormal);
+#endif
 
-	return CreateMaterialSample(L_World, Payload.BSDFOpacity * Payload.BaseColor * ShadowTerminatorTerm(L_World, Payload.WorldNormal, Payload.WorldSmoothNormal), SampledValue.w, 1.0, 1.0, PATHTRACER_SCATTER_DIFFUSE);
+	return CreateMaterialSample(L_World, Payload.BSDFOpacity * BaseColor * ShadowTerminator, SampledValue.w, 1.0, 1.0, PATHTRACER_SCATTER_DIFFUSE);
 }
 
 FMaterialEval Lambert_EvalMaterial(
@@ -24,10 +30,17 @@ FMaterialEval Lambert_EvalMaterial(
 	float2 DiffuseSpecularScale
 )
 {
-	float3 N_World = Payload.WorldNormal;
+	const float3 N_World = Payload.WorldNormal;
+	const float3 BaseColor = Payload.GetBaseColor();
+
 	float NoL = saturate(dot(N_World, L_World));
+#if SIMPLIFIED_MATERIAL_SHADER
+	float ShadowTerminator = 1;
+#else
+	float ShadowTerminator = ShadowTerminatorTerm(L_World, N_World, Payload.WorldSmoothNormal);
+#endif
 
-	return CreateMaterialEval(Payload.BSDFOpacity * Payload.BaseColor * ShadowTerminatorTerm(L_World, Payload.WorldNormal, Payload.WorldSmoothNormal) * DiffuseSpecularScale.x, NoL / PI);
+	return CreateMaterialEval(Payload.BSDFOpacity * BaseColor * ShadowTerminator * DiffuseSpecularScale.x, NoL / PI);
 }
 
 #if SIMPLIFIED_MATERIAL_SHADER
diff --git a/Engine/Shaders/Private/PathTracing/Material/PathTracingMaterialCommon.ush b/Engine/Shaders/Private/PathTracing/Material/PathTracingMaterialCommon.ush
--- a/Engine/Shaders/Private/PathTracing/Material/PathTracingMaterialCommon.ush
+++ b/Engine/Shaders/Private/PathTracing/Material/PathTracingMaterialCommon.ush
@@ -12,7 +12,8 @@ PathTracingMaterialCommon.usf: Brdf utility functions
 float LobeColorToWeight(float3 C)
 {
 	// TODO: What is the best heuristic to use here?
-	// average seems to perform slightly better, but needs more careful testing
+	// Sum seems to perform slightly better sometimes, while max3 is better in others. Need to look at more cases to be sure
+	// Also tried Luminance, but the result was very close to the plain sum
 #if 1
 	return C.x + C.y + C.z;
 #else
diff --git a/Engine/Shaders/Private/PathTracing/Material/PathTracingMaterialSampling.ush b/Engine/Shaders/Private/PathTracing/Material/PathTracingMaterialSampling.ush
--- a/Engine/Shaders/Private/PathTracing/Material/PathTracingMaterialSampling.ush
+++ b/Engine/Shaders/Private/PathTracing/Material/PathTracingMaterialSampling.ush
@@ -11,6 +11,9 @@ PathTracingMaterialSampling.usf: Material sampling functions
 #include "PathTracingLambert.ush"
 #if !SIMPLIFIED_MATERIAL_SHADER
 #include "PathTracingGlossy.ush"
+#if STRATA_ENABLED
+#include "PathTracingSubstrate.ush"
+#else
 #include "PathTracingSolidGlass.ush"
 #include "PathTracingThinGlass.ush"
 #include "PathTracingDefaultLit.ush"
@@ -18,6 +21,7 @@ PathTracingMaterialSampling.usf: Material sampling functions
 #include "PathTracingTwoSidedFoliage.ush"
 #include "PathTracingCloth.ush"
 #include "PathTracingSubsurfaceProfile.ush"
+#endif // STRATA_ENABLED
 #include "PathTracingEye.ush"
 #include "PathTracingHair.ush"
 #include "PathTracingMedium.ush"
@@ -25,17 +29,36 @@ PathTracingMaterialSampling.usf: Material sampling functions
 
 float3 EstimateMaterialAlbedo(FPathTracingPayload Payload)
 {
+
+	if (Payload.ShadingModelID == SHADINGMODELID_UNLIT)
+	{
+		// this object does not scatter light
+		return 0;
+	}
+#if PATHTRACING_STRATA_PAYLOAD
+	const float3 MaxWeight = Payload.WeightV * Payload.BSDFOpacity * lerp(1.0, Payload.TransmittanceN, Payload.CoverageAboveAlongN);
+#endif
 	// TODO: make this more accurate by estimating the directional albedo
 	float3 Albedo = 0;
 	switch (Payload.ShadingModelID)
 	{
 		case SHADINGMODELID_UNLIT:
-			// this object does not scatter light
 			return 0;
 #if SIMPLIFIED_MATERIAL_SHADER
 		case SHADINGMODELID_TWOSIDED_FOLIAGE:
 			Albedo = Payload.BaseColor + Payload.SubsurfaceColor;
 			break;
+#elif PATHTRACING_STRATA_PAYLOAD
+		case SHADINGMODELID_STRATA:
+		{
+			Albedo = MaxWeight * (Payload.DiffuseColor + Payload.SpecularColor + Payload.FuzzAmount * Payload.FuzzColor);
+			break;
+		}
+		case SHADINGMODELID_EYE:
+			Albedo = MaxWeight * (Payload.DiffuseColor + Payload.SpecularColor);
+			break;
+		case SHADINGMODELID_HAIR:
+		// STRATA_TODO: Add more cases here?
 #else
 		case SHADINGMODELID_TWOSIDED_FOLIAGE:
 		case SHADINGMODELID_SUBSURFACE:
@@ -65,7 +88,7 @@ float3 EstimateMaterialAlbedo(FPathTracingPayload Payload)
 		case SHADINGMODELID_HAIR:
 #endif
 		default:
-			Albedo = Payload.BaseColor;
+			Albedo = Payload.GetBaseColor();
 			break;
 	}
 	Albedo *= Payload.BSDFOpacity;
@@ -78,6 +101,8 @@ float3 EstimateMaterialAlbedo(FPathTracingPayload Payload)
 	return Albedo;
 }
 
+#define ROUGHNESS_WEIGHT(r)		(saturate((r - max(MinRoughness, PathRoughness)) / (MaxRoughness - max(MinRoughness, PathRoughness))))
+
 void AccumulateAlbedoNormal(inout FPathTracingPayload Payload, float3 PathThroughput, float PathRoughness, inout float3 Albedo, inout float3 Normal)
 {
 	const float MaxRoughness = 0.15; // don't bother past this level
@@ -87,22 +112,56 @@ void AccumulateAlbedoNormal(inout FPathTracingPayload Payload, float3 PathThroug
 		// no point tracking past a certain roughness
 		return;
 	}
+	if (Payload.ShadingModelID == SHADINGMODELID_UNLIT)
+	{
+		return;
+	}
 	// figure out how much the current path should contribute to the denoiser aovs
 	PathThroughput *= 1 - saturate((PathRoughness - MinRoughness) / (MaxRoughness - MinRoughness));
+#if PATHTRACING_STRATA_PAYLOAD
+	PathThroughput *= Payload.WeightV * Payload.BSDFOpacity * lerp(1.0, Payload.TransmittanceN, Payload.CoverageAboveAlongN);
+#endif
 	float NormThroughput = Luminance(PathThroughput);
 	// figure out how much the specular lobes should contribute to the total (relative to the current PathRoughness)
 	// if we are above the MinRoughness, use the PathRoughness as are lower bound instead so that multiple bounces don't accumulate too much albedo
-	const float RoughnessWeight = saturate((Payload.Roughness - max(MinRoughness, PathRoughness)) / (MaxRoughness - max(MinRoughness, PathRoughness)));
+#if PATHTRACING_STRATA_PAYLOAD
+	switch (Payload.ShadingModelID)
+	{
+		case SHADINGMODELID_STRATA:
+		{
+			const float Roughness0 = Payload.RoughnessData.x;
+			const float Roughness1 = Payload.RoughnessData.y;
+			const float SpecBlend = Payload.RoughnessData.z;
+			const float Roughness0Weight = (1 - SpecBlend) * ROUGHNESS_WEIGHT(Roughness0);
+			const float Roughness1Weight = (    SpecBlend) * ROUGHNESS_WEIGHT(Roughness1);
+			Albedo += PathThroughput * (Payload.DiffuseColor + Payload.SpecularColor * (Roughness0Weight + Roughness1Weight) + Payload.FuzzAmount * Payload.FuzzColor);
+			Normal += NormThroughput * Payload.WorldNormal;
+			break;
+		}
+		case SHADINGMODELID_EYE:
+		{
+			const float RoughnessWeight = ROUGHNESS_WEIGHT(Payload.GetEyeRoughness());
+			Albedo += PathThroughput * (Payload.DiffuseColor + RoughnessWeight * Payload.SpecularColor + Payload.SubsurfaceColor);
+			Normal += NormThroughput * Payload.WorldNormal;
+			break;
+		}
+		default:
+		{
+			Albedo += PathThroughput * Payload.GetBaseColor();
+			Normal += NormThroughput * Payload.WorldNormal;
+			break;
+		}
+	}
+#else
+	const float RoughnessWeight = ROUGHNESS_WEIGHT(Payload.Roughness);
 	const float BrdfWeight = Payload.BSDFOpacity;
+
 	switch (Payload.ShadingModelID)
 	{
-	case SHADINGMODELID_UNLIT:
-		// this object does not scatter light
-		return;
 #if SIMPLIFIED_MATERIAL_SHADER
 	case SHADINGMODELID_TWOSIDED_FOLIAGE:
 		Albedo += BrdfWeight * PathThroughput * (Payload.BaseColor + Payload.SubsurfaceColor);
-		Normal += Luminance(BrdfWeight * PathThroughput) * Payload.WorldNormal;
+		Normal += BrdfWeight * NormThroughput * Payload.WorldNormal;
 		break;
 #else
 	case SHADINGMODELID_TWOSIDED_FOLIAGE:
@@ -141,14 +200,17 @@ void AccumulateAlbedoNormal(inout FPathTracingPayload Payload, float3 PathThroug
 		Normal += BrdfWeight * NormThroughput * Payload.WorldNormal;
 		break;
 	}
-#endif
+#endif // SIMPLIFIED_MATERIAL_SHADER
 	default:
 		Albedo += BrdfWeight * PathThroughput * Payload.BaseColor;
 		Normal += BrdfWeight * NormThroughput * Payload.WorldNormal;
 		break;
 	}
+#endif // PATHTRACING_STRATA_PAYLOAD
 }
 
+#undef ROUGHNESS_WEIGHT
+
 void AccumulateAlbedo(float3 SigmaS, float3 PathThroughput, float PathRoughness, inout float3 Albedo)
 {
 	const float MaxRoughness = 0.15; // don't bother past this level
@@ -160,12 +222,8 @@ void AccumulateAlbedo(float3 SigmaS, float3 PathThroughput, float PathRoughness,
 	}
 	// figure out how much the current path should contribute to the denoiser aovs
 	PathThroughput *= 1 - saturate((PathRoughness - MinRoughness) / (MaxRoughness - MinRoughness));
-	// figure out how much the specular lobes should contribute to the total (relative to the current PathRoughness)
-	// if we are above the MinRoughness, use the PathRoughness as are lower bound instead so that multiple bounces don't accumulate too much albedo
-	const float VolumeRoughness = 1.0; // assume volume scattering is diffuse
-	const float RoughnessWeight = saturate((VolumeRoughness - max(MinRoughness, PathRoughness)) / (MaxRoughness - max(MinRoughness, PathRoughness)));
-
-	Albedo += RoughnessWeight * SigmaS * PathThroughput;
+	// Volume acts as diffuse scattering, so make it contribute fully (like DiffuseColor above)
+	Albedo += SigmaS * PathThroughput;
 }
 
 // returns an average roughness for the whole material (for the path roughness of shadow rays from this material)
@@ -173,6 +231,21 @@ float GetAverageRoughness(FPathTracingPayload Payload)
 {
 #if SIMPLIFIED_MATERIAL_SHADER
 	return 1.0;
+#elif PATHTRACING_STRATA_PAYLOAD
+	float AvgRoughness = Payload.RoughnessData.x;
+	switch (Payload.ShadingModelID)
+	{
+		case SHADINGMODELID_STRATA:
+		{
+			// STRATA_TODO: revisit this when we support glass
+			AvgRoughness = max(AvgRoughness, Payload.RoughnessData.y);
+			float DiffuseFraction = LobeSelectionProb(Payload.SubsurfaceColor + Payload.DiffuseColor + Payload.FuzzAmount * Payload.FuzzColor, Payload.SpecularColor);
+			AvgRoughness = lerp(1.0, AvgRoughness, DiffuseFraction);
+			break;
+		}
+		// STRATA_TODO: do we need special cases for the other slab types?
+	}
+	return AvgRoughness;
 #else
 	const float GlassPercent = Payload.HasRefraction() ? 1 - Payload.BSDFOpacity : 0.0;
 	float Roughness = Payload.Roughness;
@@ -198,13 +271,23 @@ FPathTracingPayload UnpackPathTracingPayloadMaterial(FPackedPathTracingPayload P
 {
 	// This function can also do a bit of per-material post-processing to simplify the rest of the code
 	FPathTracingPayload HitPayload = UnpackPathTracingPayload(PackedInput, Ray);
-
 	if (bApplyOverrides)
 	{
-		HitPayload.DiffuseColor  = HitPayload.DiffuseColor  * View.DiffuseOverrideParameter.w  + View.DiffuseOverrideParameter.xyz;
+		HitPayload.DiffuseColor = HitPayload.DiffuseColor * View.DiffuseOverrideParameter.w + View.DiffuseOverrideParameter.xyz;
 		HitPayload.SpecularColor = HitPayload.SpecularColor * View.SpecularOverrideParameter.w + View.SpecularOverrideParameter.xyz;
 	}
-
+#if PATHTRACING_STRATA_PAYLOAD
+	switch (HitPayload.ShadingModelID)
+	{
+		case SHADINGMODELID_STRATA:
+		{
+			// NOTE: if radius is 0 the sss color will get flipped back into diffuse
+			HitPayload.SubsurfaceColor = HitPayload.DiffuseColor;
+			HitPayload.DiffuseColor = 0;
+			break;
+		}
+	}
+#else
 	switch (HitPayload.ShadingModelID)
 	{
 #if !SIMPLIFIED_MATERIAL_SHADER
@@ -286,7 +369,7 @@ FPathTracingPayload UnpackPathTracingPayloadMaterial(FPackedPathTracingPayload P
 		}
 #endif
 	}
-
+#endif // PATHTRACING_STRATA_PAYLOAD
 	return HitPayload;
 }
 
@@ -294,19 +377,46 @@ struct FSSSRandomWalkInfo
 {
 	float3 Color;
 	float3 Radius;
-
 	float3 Weight;  // multiplier for the path throughput to take into account energy conservation from specular above the SSS
 	float  Prob;    // probability of choosing to perform the random walk
+	float  G;
 };
 
 
 FSSSRandomWalkInfo GetMaterialSSSInfo(FPathTracingPayload Payload, float3 V_World)
 {
 	FSSSRandomWalkInfo Result = (FSSSRandomWalkInfo)0;
+#if PATHTRACING_STRATA_PAYLOAD
+	// STRATA_TODO: figure out how to handle the multi-slab case as well
+	if (Payload.ShadingModelID == SHADINGMODELID_STRATA)
+	{
+		if (any(Payload.MeanFreePath > 0.0))
+		{
+			Result.Color = Payload.SubsurfaceColor;
+			Result.Radius = Payload.MeanFreePath;
+
+			FSubstrateSingleSlabData Data = (FSubstrateSingleSlabData)0;
+			Data.PrepareBSDF(V_World, Payload);
+			const float3 SSSAlbedo = Payload.SubsurfaceColor * Data.DiffuseWeight;
+			const float3 DiffuseAlbedo = Payload.DiffuseColor * Data.DiffuseWeight;
+			const float3 Spec0Albedo = Data.FuzzAttenuation * (1.0 - Payload.RoughnessData.z) * Data.Spec0.E;
+			const float3 Spec1Albedo = Data.FuzzAttenuation * (Payload.RoughnessData.z) * Data.Spec1.E;
+			const float  FuzzAlbedo = Payload.FuzzAmount * Data.Fuzz.E;
+			const float3 MaxLobeWeight = Payload.GetMaxLobeWeight();
+
+			Result.Weight = Payload.WeightV * Payload.BSDFOpacity * Data.DiffuseWeight;
+			Result.Prob = LobeSelectionProb(
+				MaxLobeWeight * (Payload.SubsurfaceColor * Data.DiffuseWeight),
+				MaxLobeWeight * (Payload.DiffuseColor * Data.DiffuseWeight + Spec0Albedo + Spec1Albedo + FuzzAlbedo));
+			// STRATA_TODO: evaluate if we want to set Result.G in all cases?
+		}
+	}
+#else
 	Result.Color = Payload.SubsurfaceColor;
 	Result.Radius = Payload.SubsurfaceRadius;
 	Result.Weight = 0;
 	Result.Prob = 0;
+	Result.G = Payload.GetSubsurfacePhaseFunction();
 	if (Payload.IsSubsurfaceMaterial())
 	{
 		const float NoV = saturate(dot(Payload.WorldNormal, V_World));
@@ -331,6 +441,7 @@ FSSSRandomWalkInfo GetMaterialSSSInfo(FPathTracingPayload Payload, float3 V_Worl
 		Result.Prob = LobeSelectionProb(SSSLobeAlbedo, DiffLobeAlbedo + SpecLobeAlbedo);
 		Result.Weight = Payload.BSDFOpacity * (1 - SpecE);
 	}
+#endif // PATHTRACING_STRATA_PAYLOAD
 	return Result;
 }
 
@@ -359,6 +470,10 @@ FMaterialSample SampleMaterial(
 		case SHADINGMODELID_UNLIT:
 			// Nothing to do, material does not scatter light
 			return NullMaterialSample();
+#if PATHTRACING_STRATA_PAYLOAD
+		case SHADINGMODELID_STRATA:
+			return Substrate_SampleMaterial(V_World, Payload, RandSample);
+#else
 		case SHADINGMODELID_DEFAULT_LIT:
 		{
 			if (Payload.HasRefraction() && Payload.BSDFOpacity < 1)
@@ -380,6 +495,7 @@ FMaterialSample SampleMaterial(
 			return ClearCoat_SampleMaterial(V_World, Payload, RandSample);
 		case SHADINGMODELID_TWOSIDED_FOLIAGE:
 			return TwoSidedFoliage_SampleMaterial(V_World, Payload, RandSample);
+#endif // PATHTRACING_STRATA_PAYLOAD
 		case SHADINGMODELID_EYE:
 			return Eye_SampleMaterial(V_World, Payload, RandSample);
 		case SHADINGMODELID_HAIR:
@@ -408,6 +524,10 @@ FMaterialEval EvalMaterial(
 	{
 		case SHADINGMODELID_UNLIT:
 			break;
+#if PATHTRACING_STRATA_PAYLOAD
+		case SHADINGMODELID_STRATA:
+			return Substrate_EvalMaterial(V_World, L_World, Payload, DiffuseSpecularScale);
+#else
 		case SHADINGMODELID_DEFAULT_LIT:
 		{
 			if (Payload.HasRefraction() && Payload.BSDFOpacity < 1)
@@ -429,6 +549,7 @@ FMaterialEval EvalMaterial(
 			return ClearCoat_EvalMaterial(V_World, L_World, Payload, DiffuseSpecularScale);
 		case SHADINGMODELID_TWOSIDED_FOLIAGE:
 			return TwoSidedFoliage_EvalMaterial(V_World, L_World, Payload, DiffuseSpecularScale);
+#endif // PATHTRACING_STRATA_PAYLOAD
 		case SHADINGMODELID_EYE:
 			return Eye_EvalMaterial(V_World, L_World, Payload, DiffuseSpecularScale);
 		case SHADINGMODELID_HAIR:
diff --git a/Engine/Shaders/Private/PathTracing/Material/PathTracingMedium.ush b/Engine/Shaders/Private/PathTracing/Material/PathTracingMedium.ush
--- a/Engine/Shaders/Private/PathTracing/Material/PathTracingMedium.ush
+++ b/Engine/Shaders/Private/PathTracing/Material/PathTracingMedium.ush
@@ -16,7 +16,7 @@ FMaterialSample Medium_SampleMaterial(
 	FPathTracingPayload Payload,
 	float3 RandSample)
 {
-	const float3 RayleighWeight = Payload.BaseColor;
+	const float3 RayleighWeight = Payload.GetBaseColor();
 	const float3 HGWeight = Payload.GetHGWeight();
 	const float G = Payload.GetHGPhaseG();
 
@@ -28,8 +28,8 @@ FMaterialSample Medium_SampleMaterial(
 	const float LobeProb = LobeSelectionProb(RayleighWeight, HGWeight);
 
 	const float CosTheta = RandSample.x < LobeProb ?
-		RayleighPhaseInvertCdf(RandSample.x / LobeProb) :
-		HenyeyGreensteinPhaseInvertCDF((RandSample.x - LobeProb) / (1.0 - LobeProb), G);
+		RayleighPhaseInvertCdf(RescaleRandomNumber(RandSample.x, 0.0, LobeProb)) :
+		HenyeyGreensteinPhaseInvertCDF(RescaleRandomNumber(RandSample.x, LobeProb, 1.0), G);
 	const float SinTheta = sqrt(saturate(1.0 - CosTheta * CosTheta));
 	const float Phi = RandSample.y * (2.0 * PI);
 
@@ -49,7 +49,7 @@ FMaterialEval Medium_EvalMaterial(
 	float2 DiffuseSpecularScale
 )
 {
-	const float3 RayleighWeight = Payload.BaseColor;
+	const float3 RayleighWeight = Payload.GetBaseColor();
 	const float3 HGWeight = Payload.GetHGWeight();
 	const float G = Payload.GetHGPhaseG();
 
diff --git a/Engine/Shaders/Private/PathTracing/Material/PathTracingRadianceProbe.ush b/Engine/Shaders/Private/PathTracing/Material/PathTracingRadianceProbe.ush
--- a/Engine/Shaders/Private/PathTracing/Material/PathTracingRadianceProbe.ush
+++ b/Engine/Shaders/Private/PathTracing/Material/PathTracingRadianceProbe.ush
@@ -235,13 +235,6 @@ FMaterialSample RadianceProbe_SampleMaterial(
 	return CreateMaterialSample(TangentToWorld(SampledValue.xyz, N_World), 1.0, SampledValue.w, 1.0, 1.0, PATHTRACER_SCATTER_DIFFUSE);
 }
 
-float RadianceProbe_Pdf(float3 L_World, FPathTracingPayload Payload)
-{
-	float3 N_World = Payload.WorldNormal;
-	float NoL = saturate(dot(N_World, L_World));
-	return NoL / PI;
-}
-
 FMaterialEval RadianceProbe_EvalMaterial(
 	float3 L_World,
 	FPathTracingPayload Payload
diff --git a/Engine/Shaders/Private/PathTracing/Material/PathTracingSubstrate.ush b/Engine/Shaders/Private/PathTracing/Material/PathTracingSubstrate.ush
new file mode 100644
--- /dev/null
+++ b/Engine/Shaders/Private/PathTracing/Material/PathTracingSubstrate.ush
@@ -0,0 +1,257 @@
+// Copyright Epic Games, Inc. All Rights Reserved.
+
+#pragma once
+
+#ifndef STRATA_ENABLED
+#error "This header should only be included when Strata is enabled."
+#endif
+
+#include "PathTracingMaterialCommon.ush"
+#include "PathTracingGlossy.ush"
+#include "../../BRDF.ush"
+
+#define STRATA_FRESNEL_F82  0
+
+float3 SubstrateLobeWeight(FPathTracingPayload Payload, float SatNoL)
+{
+	// See reference in LuminanceWeight
+	const float DistanceL = rcp(max(SatNoL, 0.001f));
+	const float3 TransmittanceAboveAlongL = select(Payload.TransmittanceN < 1.f, pow(max(Payload.TransmittanceN, 0.0001f), DistanceL), Payload.TransmittanceN);
+	return Payload.BSDFOpacity * Payload.WeightV * lerp(1.0f, TransmittanceAboveAlongL, Payload.CoverageAboveAlongN);
+}
+
+struct FSubstrateSingleSlabData
+{
+	float3x3 Basis;
+	float3 V;
+
+	float  DiffuseWeight;
+	float3 F0;
+	float3 F90;
+	float2 Alpha0;
+	float2 Alpha1;
+	float FuzzRoughness;
+
+	float FuzzAttenuation;
+	FBxDFEnergyTermsRGB Spec0;
+	FBxDFEnergyTermsRGB Spec1;
+	FBxDFEnergyTermsA Fuzz;
+
+	float3 LobeCdf;
+	float4 LobePdf;
+
+
+	// NOTE: Must set Basis field before calling this
+	void PrepareBSDF(float3 V_World, FPathTracingPayload Payload)
+	{
+
+		Basis = GetGGXBasis(Payload.RoughnessData.x, Payload.Anisotropy, Payload.WorldNormal, Payload.WorldTangent, Alpha0);
+		Alpha1 = GetGGXAlpha(Payload.RoughnessData.y, Payload.Anisotropy);
+
+		V = mul(Basis, V_World);
+
+		F0 = Payload.SpecularColor;
+		F90 = Payload.SpecularEdgeColor;
+		{
+			// make sure F0=0.0 fades off the specular lobe completely
+			const float SpecularMicroOcclusion = F0RGBToMicroOcclusion(F0);
+			F0 *= SpecularMicroOcclusion;
+			F90 *= SpecularMicroOcclusion;
+		}
+
+		FuzzRoughness = MakeRoughnessSafe(Payload.FuzzRoughness, STRATA_MIN_FUZZ_ROUGHNESS);
+
+		const float NoV = saturate(V.z);
+		// STRATA_TODO: Add support for F82 model?
+		Spec0 = ComputeGGXSpecEnergyTermsRGB(Payload.RoughnessData.x, NoV, F0, F90);
+		Spec1 = ComputeGGXSpecEnergyTermsRGB(Payload.RoughnessData.y, NoV, F0, F90);
+		Fuzz = ComputeClothEnergyTermsA(FuzzRoughness, NoV);
+
+		const float3 SpecE = lerp(Spec0.E, Spec1.E, Payload.RoughnessData.z);
+
+		FuzzAttenuation = 1.0 - Payload.FuzzAmount * Fuzz.E;
+		DiffuseWeight = FuzzAttenuation * (1.0 - Luminance(SpecE));
+
+		const float3 Spec0Albedo = FuzzAttenuation * (1.0 - Payload.RoughnessData.z) * Spec0.E;
+		const float3 Spec1Albedo = FuzzAttenuation * (      Payload.RoughnessData.z) * Spec1.E;
+		const float  FuzzAlbedo = Payload.FuzzAmount * Fuzz.E;
+
+		// Now prepare a cdf/pdf for lobe selection
+		float3 MaxLobeWeight = Payload.GetMaxLobeWeight();
+		LobeCdf = LobeSelectionCdf(
+			MaxLobeWeight * DiffuseWeight * Payload.DiffuseColor,
+			MaxLobeWeight * Spec0Albedo,
+			MaxLobeWeight * Spec1Albedo,
+			MaxLobeWeight * FuzzAlbedo);
+		LobePdf = LobeSelectionPdf(LobeCdf);
+	}
+};
+
+FMaterialSample Substrate_SampleMaterial(
+	float3 V_World,
+	FPathTracingPayload Payload,
+	float3 RandSample
+)
+{
+	FSubstrateSingleSlabData Data = (FSubstrateSingleSlabData)0;
+	Data.PrepareBSDF(V_World, Payload);
+
+
+	float3 L = 0, H = 0, V = Data.V;
+	float OutRoughness = 1;
+
+	const bool bSampledDiffuse  = RandSample.x < Data.LobeCdf.x;
+	const bool bSampledSpecular = RandSample.x < Data.LobeCdf.z;
+	if (bSampledDiffuse)
+	{
+		RandSample.x = RescaleRandomNumber(RandSample.x, 0.0, Data.LobeCdf.x);
+		// diffuse lobe
+		L = CosineSampleHemisphere(RandSample.xy).xyz;
+		H = normalize(L + V);
+}
+	else if (bSampledSpecular)
+	{
+		// specular lobes
+		const bool bUseSpec0 = RandSample.x < Data.LobeCdf.y;
+		if (bUseSpec0)
+		{
+			RandSample.x = RescaleRandomNumber(RandSample.x, Data.LobeCdf.x, Data.LobeCdf.y);
+			OutRoughness = Payload.RoughnessData.x;
+		}
+		else
+		{
+			RandSample.x = RescaleRandomNumber(RandSample.x, Data.LobeCdf.y, Data.LobeCdf.z);
+			OutRoughness = Payload.RoughnessData.y;
+		}
+
+		// #dxr_todo: Evaluate UniformSampleDiskConcentric
+		H = ImportanceSampleVisibleGGX_aniso(UniformSampleDisk(RandSample.xy), bUseSpec0 ? Data.Alpha0 : Data.Alpha1, V).xyz;
+		L = reflect(-V, H);
+		if (L.z <= 0)
+		{
+			// invalid output direction, exit early
+			return NullMaterialSample();
+		}
+	}
+	else
+	{
+		// cloth lobe
+		RandSample.x = RescaleRandomNumber(RandSample.x, Data.LobeCdf.z, 1.0);
+		L = UniformSampleHemisphere(RandSample.xy).xyz;
+		H = normalize(L + V);
+	}
+
+	// transform to world space
+	const float3 L_World = normalize(mul(L, Data.Basis));
+	const float3 N_World = Payload.WorldNormal;
+	const float NoV = saturate(V.z);
+	const float NoL = saturate(L.z);
+	const float VoH = saturate(dot(V, H));
+	const float NoH = saturate(H.z);
+
+	FMaterialSample Result = CreateMaterialSample(L_World, 0.0, 0.0, 1.0, OutRoughness, PATHTRACER_SCATTER_DIFFUSE);
+
+	const float DiffPdf = NoL / PI;
+	const float2 GGXResult0 = GGXEvalReflection(L, V, H, Data.Alpha0);
+	const float2 GGXResult1 = GGXEvalReflection(L, V, H, Data.Alpha1);
+	const float Spec0Pdf = GGXResult0.y;
+	const float Spec1Pdf = GGXResult1.y;
+	const float ClothPdf = 1 / (2 * PI);
+
+	const float ShadowTerminator = ShadowTerminatorTerm(L_World, N_World, Payload.WorldSmoothNormal);
+	if (bSampledDiffuse)
+	{
+		const float3 Diffuse = Data.DiffuseWeight * GetPathTracingDiffuseModel(Payload.DiffuseColor, Payload.RoughnessData.x, NoV, NoL, VoH, NoH);
+
+		Result.AddLobeWithMIS(Diffuse * ShadowTerminator, NoL / PI, Data.LobePdf.x);
+
+		Result.Pdf += Data.LobePdf.y * Spec0Pdf;
+		Result.Pdf += Data.LobePdf.z * Spec1Pdf;
+		Result.Pdf += Data.LobePdf.w * ClothPdf;
+	}
+	else if (bSampledSpecular)
+	{
+#if STRATA_FRESNEL_F82
+		// STRATA_TODO: Need to tune the logic for specular micro-occlusion ...
+		const float3 F = Data.FuzzAttenuation * F_AdobeF82(Data.F0, Data.F90, VoH);
+#else
+		const float3 F = Data.FuzzAttenuation * F_Schlick(Data.F0, Data.F90, VoH);
+#endif
+
+		Result.AddLobeWithMIS((1.0 - Payload.RoughnessData.x) * F * GGXResult0.x * Data.Spec0.W, GGXResult0.y, Data.LobePdf.y);
+		Result.AddLobeWithMIS((      Payload.RoughnessData.x) * F * GGXResult1.x * Data.Spec1.W, GGXResult1.y, Data.LobePdf.z);
+		Result.Pdf += Data.LobePdf.x * DiffPdf;
+		Result.Pdf += Data.LobePdf.w * ClothPdf;
+
+		Result.ScatterType = PATHTRACER_SCATTER_SPECULAR;
+	}
+	else
+	{
+		// Cloth Lobe
+		const float DCloth = D_Charlie(Data.FuzzRoughness, NoH);
+		const float3 FCloth = F_Schlick(Payload.FuzzColor, VoH);
+		const float3 ClothWeight = NoL > 0 && NoV > 0 ? (NoL * Payload.FuzzAmount * DCloth * Vis_Ashikhmin(NoV, NoL)) * FCloth * Data.Fuzz.W * (2 * PI) : 0.0;
+		Result.AddLobeWithMIS(ClothWeight * ShadowTerminator, 1 / (2.0 * PI), Data.LobePdf.w);
+		Result.ScatterType = PATHTRACER_SCATTER_SPECULAR;
+
+		Result.Pdf += Data.LobePdf.x * DiffPdf;
+		Result.Pdf += Data.LobePdf.y * Spec0Pdf;
+		Result.Pdf += Data.LobePdf.z * Spec1Pdf;
+	}
+
+	Result.Weight *= SubstrateLobeWeight(Payload, NoL);
+
+	return Result;
+}
+
+FMaterialEval Substrate_EvalMaterial(
+	float3 V_World,
+	float3 L_World,
+	FPathTracingPayload Payload,
+	float2 DiffuseSpecularScale
+)
+{
+	FSubstrateSingleSlabData Data = (FSubstrateSingleSlabData)0;
+	Data.PrepareBSDF(V_World, Payload);
+
+	const float3 N_World = Payload.WorldNormal;
+
+	// move vectors into right shading frame
+	const float3 V = Data.V;
+	const float3 L = mul(Data.Basis, L_World);
+	const float3 H = normalize(V + L);
+	const float NoV = saturate(V.z);
+	const float NoL = saturate(L.z);
+	const float VoH = saturate(dot(V, H));
+	const float NoH = saturate(H.z);
+
+	FMaterialEval Result = NullMaterialEval();
+
+	const float ShadowTerminator = ShadowTerminatorTerm(L_World, N_World, Payload.WorldSmoothNormal);
+	{
+		const float3 Diffuse = Data.DiffuseWeight * GetPathTracingDiffuseModel(Payload.DiffuseColor, Payload.RoughnessData.x, NoV, NoL, VoH, NoH);
+		Result.AddLobeWithMIS(DiffuseSpecularScale.x * Diffuse * ShadowTerminator, NoL / PI, Data.LobePdf.x);
+	}
+	{
+		const float2 GGXResult0 = GGXEvalReflection(L, V, H, Data.Alpha0);
+		const float2 GGXResult1 = GGXEvalReflection(L, V, H, Data.Alpha1);
+#if STRATA_FRESNEL_F82
+		const float3 F = Data.FuzzAttenuation * F_AdobeF82(Data.F0, Data.F90, VoH);
+#else
+		const float3 F = Data.FuzzAttenuation * F_Schlick(Data.F0, Data.F90, VoH);
+#endif
+		Result.AddLobeWithMIS(DiffuseSpecularScale.y * (1.0 - Payload.RoughnessData.z) * F * GGXResult0.x * Data.Spec0.W, GGXResult0.y, Data.LobePdf.y);
+		Result.AddLobeWithMIS(DiffuseSpecularScale.y * (      Payload.RoughnessData.z) * F * GGXResult1.x * Data.Spec1.W, GGXResult1.y, Data.LobePdf.z);
+	}
+	{
+		// Cloth Lobe
+		const float DCloth = D_Charlie(Data.FuzzRoughness, NoH);
+		const float3 FCloth = F_Schlick(Payload.FuzzColor, VoH);
+		const float3 ClothWeight = NoL > 0 && NoV > 0 ? (NoL * Payload.FuzzAmount * DCloth * Vis_Ashikhmin(NoV, NoL)) * FCloth * Data.Fuzz.W * (2 * PI) : 0.0;
+		Result.AddLobeWithMIS(DiffuseSpecularScale.y * ClothWeight * ShadowTerminator, 1 / (2.0 * PI), Data.LobePdf.w);
+	}
+
+	Result.Weight *= SubstrateLobeWeight(Payload, NoL);
+
+	return Result;
+}
diff --git a/Engine/Shaders/Private/PathTracing/PathTracingCommon.ush b/Engine/Shaders/Private/PathTracing/PathTracingCommon.ush
--- a/Engine/Shaders/Private/PathTracing/PathTracingCommon.ush
+++ b/Engine/Shaders/Private/PathTracing/PathTracingCommon.ush
@@ -115,6 +115,12 @@ float3 AdjustShadingNormal(float3 ShadingNormal, float3 GeoNormal, float3 RayDir
 #error "Path tracing and GPULightmass payloads are mutually exclusive. They should not be both enabled at once."
 #endif
 
+#if STRATA_ENABLED && !IS_PAYLOAD_ENABLED(RT_PAYLOAD_TYPE_GPULIGHTMASS)
+#define PATHTRACING_STRATA_PAYLOAD 1
+#else
+#define PATHTRACING_STRATA_PAYLOAD 0
+#endif
+
 // This payload structure is what we transport between RGS/CHS/AHS programs
 struct FPackedPathTracingPayload : FMinimalPayload
 {
@@ -122,6 +128,10 @@ struct FPackedPathTracingPayload : FMinimalPayload
 #if IS_PAYLOAD_ENABLED(RT_PAYLOAD_TYPE_GPULIGHTMASS)
 	uint PackedData[7];                                  // 28 bytes (encoded data, depends on shading model and ray type)
 														 // 32 bytes total
+#elif PATHTRACING_STRATA_PAYLOAD
+	// Need a bit more space for Substrate slab data
+	uint PackedData[18];                                 // 72 bytes
+	                                                     // 76 bytes total
 #else
 	uint PackedData[15];								 // 60 bytes (encoded data, depends on shading model and ray type)
 														 // 64 bytes total
@@ -251,6 +261,19 @@ struct FPackedPathTracingPayload : FMinimalPayload
 		PackedData[13] = f32tof16(DBufferC.z);
 		PackedData[13] |= f32tof16(DBufferC.w) << 16;
 	}
+	
+#if PATHTRACING_STRATA_PAYLOAD
+	void SetStochasticSlabRand(float SlabRand)
+	{
+		PackedData[14] = asuint(SlabRand);
+	}
+	
+	float GetStochasticSlabRand()
+	{
+		return asfloat(PackedData[14]);
+	}
+#endif
+
 #endif // IS_PAYLOAD_ENABLED(RT_PAYLOAD_TYPE_GPULIGHTMASS)
 };
 
@@ -261,8 +284,6 @@ CHECK_RT_PAYLOAD_SIZE(FPackedPathTracingPayload)
 CHECK_RT_PAYLOAD_SIZE(FPackedPathTracingPayload)
 #endif
 
-
-
 // This payload structure is the expanded version of the above which is more convenient to work with
 struct FPathTracingPayload : FMinimalPayload
 {
@@ -274,12 +295,36 @@ struct FPathTracingPayload : FMinimalPayload
 #else
 	float3 WorldSmoothNormal;        // smooth normal before normal/bump is applied
 #endif
-	float3 BaseColor;
+
 	float  BSDFOpacity;              // how much should we weigh down the bsdf by? (for glass, controls the blend between solid and glass lobes)
 	float3 TransparencyColor;        // how much do we see straight through the surface?
 	uint   ShadingModelID;
 	uint   Flags;
 	uint   PrimitiveLightingChannelMask;
+
+#if PATHTRACING_STRATA_PAYLOAD
+	float3 DiffuseColor;
+	float3 SpecularColor;
+	float3 SpecularEdgeColor;
+	float3 RoughnessData; // r0,r1,rblend
+	float  Anisotropy;
+	float3 WorldTangent;
+	float3 MeanFreePath;
+	float PhaseG;
+	float Ior;
+	float FuzzAmount;
+	float FuzzRoughness;
+	float3 FuzzColor;
+	float3 WeightV;
+	float3 TransmittanceN;
+	float CoverageAboveAlongN;
+
+	// not transported through packed payload, but used in raygen
+	float3 SubsurfaceColor;
+
+#else // PATHTRACING_STRATA_PAYLOAD
+	
+	float3 BaseColor;
 #if IS_PAYLOAD_ENABLED(RT_PAYLOAD_TYPE_GPULIGHTMASS)
 	float3 SubsurfaceColor;
 #else
@@ -295,7 +340,8 @@ struct FPathTracingPayload : FMinimalPayload
 	float3 WorldTangent;
 	float3 SubsurfaceColor;
 	float3 SubsurfaceRadius;
-#endif
+#endif // IS_PAYLOAD_ENABLED(RT_PAYLOAD_TYPE_GPULIGHTMASS)
+#endif // PATHTRACING_STRATA_PAYLOAD
 	float3 TranslatedWorldPos;
 
 	void SetFrontFace() { Flags |= PATH_TRACING_PAYLOAD_OUTPUT_FLAG_FRONT_FACE; }
@@ -317,12 +363,56 @@ struct FPathTracingPayload : FMinimalPayload
 	bool IsMaterialTwoSided() { return (Flags & PATH_TRACING_PAYLOAD_OUTPUT_FLAG_TWO_SIDED_MATERIAL) != 0; }
 
 #if IS_PAYLOAD_ENABLED(RT_PAYLOAD_TYPE_GPULIGHTMASS)
+	float3 GetBaseColor() { return BaseColor; }
 	float3 GetSubsurfaceColor() { return SubsurfaceColor; }
 	bool IsMaterialTransmissive()
 	{
 		return ShadingModelID == SHADINGMODELID_TWOSIDED_FOLIAGE;
 	}
-#else
+#else // IS_PAYLOAD_ENABLED(RT_PAYLOAD_TYPE_GPULIGHTMASS)
+
+#if PATHTRACING_STRATA_PAYLOAD
+	// SHADINGMODELID_STRATA
+	float3 GetMaxLobeWeight()
+	{
+		return WeightV * lerp(1.0f, TransmittanceN, CoverageAboveAlongN); // largest value LobeWeight() could return
+	}
+
+
+	// SHADINGMODELID_NUM (Lambert diffuse) and SHADINGMODELID_MEDIUM
+	float3 GetBaseColor() { return DiffuseColor; }
+	void SetBaseColor(float3 BaseColor) { DiffuseColor = BaseColor; }
+
+	// SHADINGMODELID_MEDIUM
+	float3 GetHGWeight() { return SpecularColor; }
+	float GetHGPhaseG() { return PhaseG; }
+	void SetHG(float3 Weight, float G)
+	{
+		SpecularColor = Weight;
+		PhaseG = G;
+	}
+
+	// SHADINGMODELID_EYE
+	float GetEyeRoughness() { return RoughnessData.x; }
+	void SetEyeRoughness(float R) { RoughnessData.x = R; }
+	float3 GetEyeCausticNormal() { return 2.0 * FuzzColor - 1.0; } // STRATA_TODO: not optimal since packed as color
+	void SetEyeCausticNormal(float3 CausticNormal) { FuzzColor = 0.5 * CausticNormal + 0.5; }
+	float GetEyeIrisMask() { return Anisotropy; }
+	float3 GetEyeIrisNormal() { return WorldTangent; }
+	void SetEyeIrisMask(float IrisMask) { Anisotropy = IrisMask; }
+	void SetEyeIrisNormal(float3 IrisNormal) { WorldTangent = IrisNormal; }
+
+	// SHADINGMODELID_HAIR
+	float GetHairLongitudinalRoughness() { return RoughnessData.x; }
+	void SetHairLongitudinalRoughness(float Roughness) { RoughnessData.x = Roughness;}
+	float GetHairAzimuthalRoughness() { return RoughnessData.y; }
+	void SetHairAzimuthalRoughness(float AzimuthalRoughness) { RoughnessData.y = AzimuthalRoughness; }
+	float GetHairSpecular() { return RoughnessData.z; }
+	void SetHairSpecular(float Specular) { RoughnessData.z = Specular; }
+	float2 GetHairPrimitiveUV() { return SpecularEdgeColor.xy; }
+	void SetHairPrimitiveUV(float2 UV) { SpecularEdgeColor.xy = UV; }
+
+#else // PATHTRACING_STRATA_PAYLOAD
 	// Various ways to interpret CustomData (depending on ShadingModelID)
 	// NOTE: This is not always following the same conventions as GetMaterialCustomData0,1()
 
@@ -356,6 +446,8 @@ struct FPathTracingPayload : FMinimalPayload
 
 
 	// SHADINGMODELID_EYE
+	float GetEyeRoughness() { return Roughness; }
+	void SetEyeRoughness(float R) { Roughness = R; }
 	float3 GetEyeCausticNormal() { return CustomData1.xyz; }
 	void SetEyeCausticNormal(float3 CausticNormal) { CustomData1.xyz = CausticNormal; }
 	float GetEyeIrisMask() { return Anisotropy; }
@@ -364,6 +456,12 @@ struct FPathTracingPayload : FMinimalPayload
 	void SetEyeIrisNormal(float3 IrisNormal) { WorldTangent = IrisNormal; }
 
 	// SHADINGMODELID_HAIR
+	float GetHairLongitudinalRoughness() { return Roughness; }
+	void SetHairLongitudinalRoughness(float R) { Roughness = R; }
+	float GetHairAzimuthalRoughness() { return Metallic; }
+	void SetHairAzimuthalRoughness(float AzimuthalRoughness) { Metallic = AzimuthalRoughness; }
+	float GetHairSpecular() { return Specular; }
+	void SetHairSpecular(float S) { Specular = S; }
 	float2 GetHairPrimitiveUV() { return CustomData0.xy; }
 	void SetHairPrimitiveUV(float2 UV) { CustomData0.xy = UV; }
 
@@ -375,11 +473,18 @@ struct FPathTracingPayload : FMinimalPayload
 	float3 GetExtinction() { return CustomData0.xyz; }
 	void SetExtinction(float3 SigmaT) { CustomData0.xyz = SigmaT; }
 
+	// SHADINGMODELID_NUM (Lambert diffuse) and SHADINGMODELID_MEDIUM
+	float3 GetBaseColor() { return BaseColor; }
+	void SetBaseColor(float3 C)
+	{
+		BaseColor = C;
+	}
+
 	// SHADINGMODELID_MEDIUM
 	float3 GetHGWeight() { return CustomData0.xyz; }
 	float GetHGPhaseG() { return CustomData0.w; }
 	void SetHG(float3 Weight, float G) { CustomData0 = float4(Weight, G); }
-
+#endif // PATHTRACING_STRATA_PAYLOAD
 
 	// Methods used by the integrator
 
@@ -412,9 +517,13 @@ struct FPathTracingPayload : FMinimalPayload
 		return ShadingModelID == SHADINGMODELID_SUBSURFACE ||
 			   ShadingModelID == SHADINGMODELID_PREINTEGRATED_SKIN ||
 			   ShadingModelID == SHADINGMODELID_SUBSURFACE_PROFILE ||
-			   ShadingModelID == SHADINGMODELID_EYE;
-	}
+			   ShadingModelID == SHADINGMODELID_EYE
+#if PATHTRACING_STRATA_PAYLOAD
+			   || ShadingModelID == SHADINGMODELID_STRATA
 #endif
+				;
+	}
+#endif // IS_PAYLOAD_ENABLED(RT_PAYLOAD_TYPE_GPULIGHTMASS)
 };
 
 FPackedPathTracingPayload InitPathTracingPayload(uint ScatterType, float PathRoughness)
@@ -468,8 +577,28 @@ FPackedPathTracingPayload PackPathTracingPayload(FPathTracingPayload Input)
 	Output.PackedData[1] = PayloadEncodeUnitVector(Input.WorldGeoNormal);
 	Output.PackedData[2] = PayloadEncodeUnitVector(Input.WorldSmoothNormal);
 	Output.PackedData[3] = PayloadEncodeUnitVector(Input.WorldNormal);
-
 	Output.PackedData[4] = PayloadEncodeHDRColor(Input.Radiance);
+#if PATHTRACING_STRATA_PAYLOAD
+	Output.PackedData[5] = PayloadEncodeLDRColor(Input.DiffuseColor);
+	Output.PackedData[6] = PayloadEncodeLDRColor(Input.SpecularColor);
+	Output.PackedData[7] = PayloadEncodeLDRColor(Input.SpecularEdgeColor);
+	Output.PackedData[8] = PayloadEncodeLDRColor(Input.RoughnessData); // not a color, but encoding is still ok
+	Output.PackedData[9] = PayloadEncodeUnitVector(Input.WorldTangent);
+	Output.PackedData[10] = PayloadEncodeLDRColor(Input.TransparencyColor);
+	Output.PackedData[11] = f32tof16(Input.MeanFreePath.x);
+	Output.PackedData[11] |= f32tof16(Input.MeanFreePath.y) << 16;
+	Output.PackedData[12] = f32tof16(Input.MeanFreePath.z);
+	Output.PackedData[12] |= f32tof16(Input.PhaseG) << 16;
+	Output.PackedData[13] = PayloadEncodeLDRColor(Input.FuzzColor);
+    Output.PackedData[14] = PayloadEncodeLDRColor(float3(0.5 * Input.Anisotropy + 0.5, Input.FuzzAmount, Input.FuzzRoughness)); // NOTE: Aniso is [-1,1]
+	Output.PackedData[15] = PayloadEncodeLDRColor(Input.WeightV);
+	Output.PackedData[16] = PayloadEncodeLDRColor(Input.TransmittanceN);
+	Output.PackedData[17] = f32tof16(Input.Ior);
+	Output.PackedData[17] |= f32tof16(Input.CoverageAboveAlongN) << 16;
+#else // PATHTRACING_STRATA_PAYLOAD
+
+
+
 	Output.PackedData[5] = PayloadEncodeLDRColor(Input.BaseColor);
 	Output.PackedData[6] = PayloadEncodeLDRColor(Input.TransparencyColor);
 
@@ -492,7 +621,8 @@ FPackedPathTracingPayload PackPathTracingPayload(FPathTracingPayload Input)
 
 	Output.PackedData[14] = f32tof16(Input.Ior);
 	// NOTE: 16 bits left in PackedData[14]
-#endif
+#endif // PATHTRACING_STRATA_PAYLOAD
+#endif // IS_PAYLOAD_ENABLED(RT_PAYLOAD_TYPE_GPULIGHTMASS)
 	return Output;
 }
 
@@ -501,6 +631,7 @@ FPathTracingPayload UnpackPathTracingPayload(FPackedPathTracingPayload Input, FR
 	FPathTracingPayload Output = (FPathTracingPayload)0;
 
 	Output.HitT = Input.HitT;
+	Output.TranslatedWorldPos = Ray.Origin + Output.HitT * Ray.Direction;
 	Output.BSDFOpacity = f16tof32(Input.PackedData[0]);
 	Output.ShadingModelID = (Input.PackedData[0] >> 16) & 0xF;
 	Output.Flags = Input.GetFlags();
@@ -513,11 +644,32 @@ FPathTracingPayload UnpackPathTracingPayload(FPackedPathTracingPayload Input, FR
 	Output.TransparencyColor = PayloadDecodeLDRColor(Input.PackedData[5]);
 	Output.SubsurfaceColor = PayloadDecodeLDRColor(Input.PackedData[6]);
 #else
-	Output.WorldGeoNormal = PayloadDecodeUnitVector(Input.PackedData[1]);
+	Output.WorldGeoNormal    = PayloadDecodeUnitVector(Input.PackedData[1]);
 	Output.WorldSmoothNormal = PayloadDecodeUnitVector(Input.PackedData[2]);
-	Output.WorldNormal = PayloadDecodeUnitVector(Input.PackedData[3]);
-
-	Output.Radiance = PayloadDecodeHDRColor(Input.PackedData[4]);
+	Output.WorldNormal       = PayloadDecodeUnitVector(Input.PackedData[3]);
+	Output.Radiance          = PayloadDecodeHDRColor(Input.PackedData[4]);
+
+#if PATHTRACING_STRATA_PAYLOAD
+	Output.DiffuseColor        = PayloadDecodeLDRColor(Input.PackedData[5]);
+	Output.SpecularColor       = PayloadDecodeLDRColor(Input.PackedData[6]);
+	Output.SpecularEdgeColor   = PayloadDecodeLDRColor(Input.PackedData[7]);
+	Output.RoughnessData       = PayloadDecodeLDRColor(Input.PackedData[8]);
+	Output.WorldTangent        = PayloadDecodeUnitVector(Input.PackedData[9]);
+	Output.TransparencyColor   = PayloadDecodeLDRColor(Input.PackedData[10]);
+	Output.MeanFreePath.x      = f16tof32(Input.PackedData[11]);
+	Output.MeanFreePath.y      = f16tof32(Input.PackedData[11] >> 16);
+	Output.MeanFreePath.z      = f16tof32(Input.PackedData[12]);
+	Output.PhaseG              = f16tof32(Input.PackedData[12] >> 16);
+	Output.FuzzColor           = PayloadDecodeLDRColor(Input.PackedData[13]);
+	const float3 Data14        = PayloadDecodeLDRColor(Input.PackedData[14]);
+	Output.Anisotropy          = 2 * Data14.x - 1;
+	Output.FuzzAmount          = Data14.y;
+	Output.FuzzRoughness       = Data14.z;
+	Output.WeightV             = PayloadDecodeLDRColor(Input.PackedData[15]);
+	Output.TransmittanceN      = PayloadDecodeLDRColor(Input.PackedData[16]);
+	Output.Ior                 = f16tof32(Input.PackedData[17]);
+	Output.CoverageAboveAlongN = f16tof32(Input.PackedData[17] >> 16);
+#else // PATHTRACING_STRATA_PAYLOAD
 	Output.BaseColor = PayloadDecodeLDRColor(Input.PackedData[5]);
 	Output.TransparencyColor = PayloadDecodeLDRColor(Input.PackedData[6]);
 	
@@ -542,7 +694,8 @@ FPathTracingPayload UnpackPathTracingPayload(FPackedPathTracingPayload Input, FR
 
 	Output.DiffuseColor = Output.BaseColor - Output.BaseColor * Output.Metallic;
 	Output.SpecularColor = ComputeF0(Output.Specular, Output.BaseColor, Output.Metallic);
-#endif
+#endif // PATHTRACING_STRATA_PAYLOAD
+#endif // IS_PAYLOAD_ENABLED(RT_PAYLOAD_TYPE_GPULIGHTMASS)
 	Output.TranslatedWorldPos = Ray.Origin + Output.HitT * Ray.Direction;
 	return Output;
 }
diff --git a/Engine/Shaders/Private/PathTracing/PathTracingCore.ush b/Engine/Shaders/Private/PathTracing/PathTracingCore.ush
--- a/Engine/Shaders/Private/PathTracing/PathTracingCore.ush
+++ b/Engine/Shaders/Private/PathTracing/PathTracingCore.ush
@@ -192,6 +192,9 @@ void AccumulateRadiance(inout float3 TotalRadiance, float3 PathRadiance, bool bI
 	TotalRadiance += PathRadiance;
 }
 
+#if STRATA_ENABLED
+	// TODO: Figure out decals + Strata
+#else
 struct FDecalData
 {
 	float4 BaseColor;
@@ -322,6 +325,7 @@ void EvaluateDecals(FPathState PathState, uint InstanceInclusionMask, FPackedPat
 	}
 #endif
 }
+#endif // STRATA_ENABLED
 
 FPathTracingPayload TraceTransparentRay(inout FPathState PathState, bool bIsCameraRay, bool bLastBounce, bool bIncludeEmission, uint NumVisibleLights, inout FVolumeSegment VolumeSegment)
 {
@@ -358,6 +362,9 @@ FPathTracingPayload TraceTransparentRay(inout FPathState PathState, bool bIsCame
 		PackedPayload.SetDBufferA(float4(0, 0, 0, 1));
 		PackedPayload.SetDBufferB(float4(0, 0, 0, 1));
 		PackedPayload.SetDBufferC(float4(0, 0, 0, 1));
+#if PATHTRACING_STRATA_PAYLOAD
+		PackedPayload.SetStochasticSlabRand(RandomSequence_GenerateSample1D(PathState.RandSequence));
+#endif
 		// Trace the ray
 		TraceRay(
 			TLAS,
@@ -369,6 +376,9 @@ FPathTracingPayload TraceTransparentRay(inout FPathState PathState, bool bIsCame
 			PathState.Ray.GetNativeDesc(),
 			PackedPayload);
 
+#if STRATA_ENABLED
+		// STRATA_TODO: figure out decals
+#else
 		FDecalData DecalData;
 		DecalData.BaseColor = float4(0, 0, 0, 1);
 		DecalData.WorldNormal = float4(0, 0, 0, 1);
@@ -402,6 +412,7 @@ FPathTracingPayload TraceTransparentRay(inout FPathState PathState, bool bIsCame
 			}
 #endif
 		}
+#endif // STRATA_ENABLED
 
 		if (PackedPayload.IsMiss() && VolumeIsectList.HitBlocker())
 		{
@@ -577,6 +588,9 @@ FPathTracingPayload TraceTransparentRay(inout FPathState PathState, bool bIsCame
 		// Unpack the payload
 		FPathTracingPayload HitPayload = UnpackPathTracingPayloadMaterial(PackedPayload, PathState.Ray, bIsCameraRay && ApplyDiffuseSpecularOverrides);
 
+#if STRATA_ENABLED
+		// STRATA_TODO: figure out decals
+#else // STRATA_ENABLED
 		const float DecalThroughput = ApplyDecal(HitPayload, DecalData);
 
 		if (DecalThroughput < 1.0f)
@@ -595,6 +609,7 @@ FPathTracingPayload TraceTransparentRay(inout FPathState PathState, bool bIsCame
 				HitPayload.SpecularColor = HitPayload.SpecularColor * View.SpecularOverrideParameter.w + View.SpecularOverrideParameter.xyz;
 			}
 		}
+#endif // STRATA_ENABLED
 
 		// account for Beer's law through the currently active medium
 		// TODO: merge this with the volume raymarching?
@@ -831,16 +846,12 @@ FProbeResult TraceSSSProbeRay(FRayDesc Ray, inout int InterfaceCounter)
 
 void ApplyRayBias(inout FRayDesc Ray, float HitT, float3 Normal)
 {
-#if 1
 	// Take maximum of position or hit distance to determine "radius" of hit avoidance
 	const float RefValue = max(max(abs(Ray.Origin.x), abs(Ray.Origin.y)), max(abs(Ray.Origin.z), HitT));
 	const uint UlpRadius = 16; // number of floating point ulps to skip around an intersection
 	const float RelativeOffset = asfloat(asuint(RefValue) + UlpRadius) - RefValue;
 	const float BaseOffset = 0.001; // 0.01mm (avoid precision issues very close to the origin where ulps become tiny)
 	Ray.Origin += max(BaseOffset, RelativeOffset) * Normal;
-#else
-	ApplyPositionBias(Ray.Origin, Ray.Direction, Normal, MaxNormalBias);
-#endif
 }
 
 float3 ComputeDwivediScale(float3 Albedo)
@@ -1095,7 +1106,7 @@ bool ProcessSubsurfaceRandomWalk(inout FPathTracingPayload Payload, inout float3
 
 	// Revisiting Physically Based Shading at Imageworks.
 	// https://blog.selfshadow.com/publications/s2017-shading-course/imageworks/s2017_pbs_imageworks_slides_v2.pdf
-	float G = Payload.GetSubsurfacePhaseFunction();
+	float G = SSS.G;
 	Albedo = Albedo / (1 - G * (1 - Albedo));
 
 #if SSS_USE_DWIVEDI
@@ -1246,7 +1257,7 @@ bool ProcessSubsurfaceRandomWalk(inout FPathTracingPayload Payload, inout float3
 			
 			Payload.ShadingModelID = SHADINGMODELID_NUM; // invalid value so that we get diffuse shading
 			Payload.BSDFOpacity = 1;
-			Payload.BaseColor = 1;
+			Payload.SetBaseColor(1.0);
 			Payload.TransparencyColor = 0;
 			return true;
 		}
@@ -1660,6 +1671,9 @@ bool PathTracingKernel(inout FPathState PathState, int Bounce)
 	// enlarge roughness based on the chosen lobe roughness
 	PathState.PathRoughness = max(PathState.PathRoughness, MaterialSample.Roughness);
 
+#if STRATA_ENABLED
+	// STRATA_TODO: Account for medium absorption here once we have support for refractive glass
+#else
 	// update the current extinction if we are crossing a boundary on glass or water
 	// summing the local extinction gives a rudimentary way of dealing with overlapping regions
 	// long term we will probably want a stack with priorities
@@ -1678,6 +1692,7 @@ bool PathTracingKernel(inout FPathState PathState, int Bounce)
 			PathState.SigmaT = max(PathState.SigmaT, 0);
 		}
 	}
+#endif // STRATA_ENABLED
 
 	// If we are using Material sampling for lights
 	if (MISMode != 1 && bDoLightLoop)
diff --git a/Engine/Shaders/Private/PathTracing/PathTracingGPULightmassMaterialHitShader.usf b/Engine/Shaders/Private/PathTracing/PathTracingGPULightmassMaterialHitShader.usf
--- a/Engine/Shaders/Private/PathTracing/PathTracingGPULightmassMaterialHitShader.usf
+++ b/Engine/Shaders/Private/PathTracing/PathTracingGPULightmassMaterialHitShader.usf
@@ -12,6 +12,9 @@
 // This should be good enough for path tracing and avoids having to bind an extra buffer
 #define EYE_ADAPTATION_PREV_FRAME_EXPOSURE 1
 
+// Ensure that SSS albedo comes through in the material
+#define STRATA_SSS_MATERIAL_OVERRIDE 0
+
 #include "/Engine/Private/Common.ush"
 #include "/Engine/Private/RayTracing/RayTracingCommon.ush"
 #include "/Engine/Private/RayTracing/RayTracingHitGroupCommon.ush"
@@ -79,6 +82,7 @@ RAY_TRACING_ENTRY_CLOSEST_HIT(GPULightmassMaterialCHS,
 	const bool bIsFrontFace = HitKind() == HIT_KIND_TRIANGLE_FRONT_FACE;
 	{
 		const float4 ScreenPosition = SvPositionToResolvedScreenPosition(SvPosition);
+		MaterialParameters.CameraVector = -WorldRayDirection();
 		CalcMaterialParametersEx(MaterialParameters, PixelMaterialInputs, SvPosition, ScreenPosition, bIsFrontFace, TranslatedWorldPosition, TranslatedWorldPosition);
 	}
 
diff --git a/Engine/Shaders/Private/PathTracing/PathTracingInitExtinctionCoefficient.usf b/Engine/Shaders/Private/PathTracing/PathTracingInitExtinctionCoefficient.usf
--- a/Engine/Shaders/Private/PathTracing/PathTracingInitExtinctionCoefficient.usf
+++ b/Engine/Shaders/Private/PathTracing/PathTracingInitExtinctionCoefficient.usf
@@ -8,6 +8,10 @@ RWBuffer<float> RWStartingExtinctionCoefficient;
 
 RAY_TRACING_ENTRY_RAYGEN(PathTracingInitExtinctionCoefficientRG)
 {
+#if STRATA_ENABLED
+	// STRATA_TODO: Figure out medium parameters from strata data
+	float3 StartingSigmaT = 0.0;
+#else
 	const uint RayFlags = 0;
 	const uint MissShaderIndex = 0;
 	const uint InstanceInclusionMask = PATHTRACER_MASK_CAMERA; // Ignore hair strands, only trace what camera rays can see
@@ -49,7 +53,7 @@ RAY_TRACING_ENTRY_RAYGEN(PathTracingInitExtinctionCoefficientRG)
 		// keep tracing
 		UpRay.TMin = asfloat(asuint(PackedPayload.HitT) + 1);
 	}
-
+#endif // STRATA_ENABLED
 	// clamp against 0 in the unlikely case that we only hit front faces of glass on our way up
 	StartingSigmaT = max(StartingSigmaT, 0.0);
 
diff --git a/Engine/Shaders/Private/PathTracing/PathTracingMaterialHitShader.usf b/Engine/Shaders/Private/PathTracing/PathTracingMaterialHitShader.usf
--- a/Engine/Shaders/Private/PathTracing/PathTracingMaterialHitShader.usf
+++ b/Engine/Shaders/Private/PathTracing/PathTracingMaterialHitShader.usf
@@ -12,6 +12,9 @@
 // This should be good enough for path tracing and avoids having to bind an extra buffer
 #define EYE_ADAPTATION_PREV_FRAME_EXPOSURE 1
 
+// Ensure that SSS albedo comes through in the material
+#define STRATA_SSS_MATERIAL_OVERRIDE 0
+
 #include "/Engine/Private/Common.ush"
 #include "/Engine/Private/RayTracing/RayTracingCommon.ush"
 #include "/Engine/Private/RayTracing/RayTracingHitGroupCommon.ush"
@@ -26,8 +29,12 @@
 #include "/Engine/Private/BurleyNormalizedSSSCommon.ush"
 #include "/Engine/Private/PathTracing/PathTracingCommon.ush"
 #include "/Engine/Private/PathTracing/Material/PathTracingFresnel.ush"
+#if STRATA_ENABLED
+#define MATERIAL_STRATA_OPAQUE_PRECOMPUTED_LIGHTING 0
+#include "/Engine/Private/Strata/StrataExport.ush"
+#else
 #include "/Engine/Private/PathTracing/Material/PathTracingThinGlass.ush"
-
+#endif
 
 float AdjustMaterialRoughness(float Roughness, float PathRoughness)
 {
@@ -135,6 +142,7 @@ RAY_TRACING_ENTRY_CLOSEST_HIT(PathTracingMaterialCHS,
 #if MATERIALBLENDING_MODULATE || (MATERIALBLENDING_MASKED && MATERIAL_DITHER_OPACITY_MASK) || MATERIALBLENDING_ALPHACOMPOSITE || MATERIALBLENDING_TRANSLUCENT
 	if (PackedPayload.IsVisibilityRay())
 	{
+		// STRATA_TODO: Handle transparent shadows here ...
 #if MATERIALBLENDING_MODULATE
 		const float3 Transparency = GetMaterialEmissive(PixelMaterialInputs);
 #elif MATERIALBLENDING_ALPHACOMPOSITE
@@ -326,6 +334,317 @@ RAY_TRACING_ENTRY_CLOSEST_HIT(PathTracingMaterialCHS,
 	}
 #endif
 
+	float GeoNormalSign = MaterialParameters.TwoSidedSign;
+#if !VF_SUPPORTS_RAYTRACING_PREPARE_MATERIAL_PIXEL_PARAMETERS
+	// Because the geometric normal is computed directly in world space
+	// it doesn't reflect the sign flip from the object transform, so apply it here
+	GeoNormalSign *= GetPrimitive_DeterminantSign_FromFlags(PrimitiveFlags);
+#endif
+	Payload.WorldGeoNormal = GeoNormalSign * GeoNormal;
+	Payload.WorldSmoothNormal = MaterialParameters.TwoSidedSign * TransformTangentNormalToWorld(MaterialParameters.TangentToWorld, float3(0, 0, 1));
+	Payload.WorldSmoothNormal = AdjustShadingNormal(Payload.WorldSmoothNormal, Payload.WorldGeoNormal, WorldRayDirection());
+
+
+#if STRATA_ENABLED
+	// unless proven otherwise, we are a basic unlit material
+	Payload.ShadingModelID = SHADINGMODELID_UNLIT;
+
+#if STRATA_OPTIMIZED_UNLIT
+	FStrataBSDF UnlitBSDF = PixelMaterialInputs.GetFrontStrataData().InlinedBSDF;
+	Payload.Radiance = BSDF_GETEMISSIVE(UnlitBSDF);
+#if STRATA_OPAQUE_MATERIAL
+	Payload.BSDFOpacity = 1;
+	Payload.TransparencyColor = 0;
+#else
+	Payload.BSDFOpacity = UnlitBSDF.Coverage;
+	Payload.TransparencyColor = 1 - Payload.BSDFOpacity; // STRATA_TODO: is this right?
+#endif
+
+#else // STRATA_OPTIMIZED_UNLIT
+
+	// STRATA_TODO: Handle transparency blend modes ...
+	Payload.BSDFOpacity = 1;
+	Payload.TransparencyColor = 0;
+
+	const float3 V_World = MaterialParameters.CameraVector;
+	FStrataData StrataData= PixelMaterialInputs.GetFrontStrataData();
+	FStrataPixelHeader StrataPixelHeader = MaterialParameters.GetFrontStrataHeader();
+	FStrataIntegrationSettings IntegrationSettings = InitStrataIntegrationSettings();
+	FStrataAddressing NullStrataAddressing = (FStrataAddressing)0;	// Fake unused in StrataCreateBSDFContext when using Forward inline shading
+	FStrataTree StrataTree = StrataPixelHeader.StrataTree;
+	if (StrataTree.BSDFCount > 0)
+	{
+		StrataUpdateTree(NullStrataAddressing, StrataPixelHeader, StrataTree, V_World, IntegrationSettings);
+
+	
+		{
+#if HAIR_STRAND_MESH_FACTORY
+			Payload.WorldSmoothNormal = Payload.WorldGeoNormal;
+#else
+			// Nudge normal to avoid dark edges in bump
+			// Not that we must recompute tangent vector to keep the basis frame orthogonal
+			for (uint i = 0; i < StrataPixelHeader.SharedLocalBases.Count; ++i)
+			{
+				float3 N = StrataPixelHeader.SharedLocalBases.Normals[i];
+				float3 T = StrataPixelHeader.SharedLocalBases.Tangents[i];
+				N = AdjustShadingNormal(N, Payload.WorldGeoNormal, WorldRayDirection());
+				float3 B = normalize(cross(N, T)); // N and T are no longer orthogonal to each other, need to normalize
+				StrataPixelHeader.SharedLocalBases.Normals[i]  = N;
+				StrataPixelHeader.SharedLocalBases.Tangents[i] = cross(B, N); // no normalize needed here, will be orthogonal
+				
+			}
+#endif
+		}
+
+#define BSDF StrataTree.BSDFs[BSDFIdx]
+
+#if STRATA_CLAMPED_BSDF_COUNT > 1
+		// The material might be using more than one slab -- insert code to stochastically pick one 
+		int BSDFIdx = 0;
+		float SlabPdf = 1.0;
+		if (StrataTree.BSDFCount > 1)
+		{
+			float SlabCDF[STRATA_CLAMPED_BSDF_COUNT];
+			float SlabCDFSum = 0.0;
+			for (BSDFIdx = 0; BSDFIdx < StrataTree.BSDFCount; ++BSDFIdx)
+			{
+				float Pdf = 0.0; // The goal is to predict the weight of each BSDF
+				switch (BSDF_GETTYPE(BSDF))
+				{
+					case STRATA_BSDF_TYPE_SLAB:
+					{
+						const float3 WeightV = BSDF.LuminanceWeightV;
+						const float3 TransmittanceN = BSDF.TransmittanceAboveAlongN;
+						const float CoverageAboveAlongN = BSDF.CoverageAboveAlongN;
+						const float3 MaxLobeWeight = WeightV * lerp(1.0f, TransmittanceN, CoverageAboveAlongN); // largest value LobeWeight() could return
+
+						// accumulate visibility
+						Payload.Radiance += WeightV * BSDF_GETEMISSIVE(BSDF);
+
+						// Is the slab even visible at all?
+						if (any(MaxLobeWeight > 0.0))
+						{
+							const float3x3 Basis = StrataGetBSDFSharedBasis_InlineShading(StrataPixelHeader, BSDF, NullStrataAddressing);
+							const float3 V = mul(Basis, V_World);
+
+							float3 DiffuseColor = SLAB_DIFFUSEALBEDO(BSDF);
+							float3 F0 = SLAB_F0(BSDF);
+							float3 F90 = SLAB_F90(BSDF);
+							{
+								// make sure F0=0.0 fades off the specular lobe completely
+								const float SpecularMicroOcclusion = F0RGBToMicroOcclusion(F0);
+								F0 *= SpecularMicroOcclusion;
+								F90 *= SpecularMicroOcclusion;
+							}
+
+							float Roughness0 = MakeRoughnessSafe(SLAB_ROUGHNESS(BSDF), PackedPayload.GetPathRoughness());
+							float Roughness1 = Roughness0;
+							float SpecBlend = 0.0;
+							if (BSDF_GETHASHAZINESS(BSDF))
+							{
+								const FHaziness Haziness = UnpackHaziness(SLAB_HAZINESS(BSDF));
+								// STRATA_TODO: Ideally clearcoat in the path tracer would be handled via actual layering rather than this backwards compatibility trick
+								if (!Haziness.bSimpleClearCoat)
+								{
+									SpecBlend = Haziness.Weight;
+									Roughness1 = MakeRoughnessSafe(Haziness.Roughness, PackedPayload.GetPathRoughness());
+								}
+							}
+							float FuzzAmount = 0.0;
+							float FuzzRoughness = 1.0;
+							float3 FuzzColor = 0.0;
+							if (BSDF_GETHASFUZZ(BSDF))
+							{
+								FuzzAmount = SLAB_FUZZ_AMOUNT(BSDF);
+								FuzzColor = SLAB_FUZZ_COLOR(BSDF);
+								FuzzRoughness = SLAB_FUZZ_ROUGHNESS(BSDF);
+								FuzzRoughness = MakeRoughnessSafe(FuzzRoughness, STRATA_MIN_FUZZ_ROUGHNESS);
+							}
+
+							const float NoV = saturate(V.z);
+							// STRATA_TODO: Support Fresnel82?
+							const FBxDFEnergyTermsRGB Spec0 = ComputeGGXSpecEnergyTermsRGB(Roughness0, NoV, F0, F90);
+							const FBxDFEnergyTermsRGB Spec1 = ComputeGGXSpecEnergyTermsRGB(Roughness1, NoV, F0, F90);
+							const FBxDFEnergyTermsA   Fuzz  = ComputeClothEnergyTermsA(FuzzRoughness, NoV);
+							const float3 SpecE = lerp(Spec0.E, Spec1.E, SpecBlend);
+							const float FuzzAttenuation = 1.0 - FuzzAmount * Fuzz.E;
+							const float DiffuseWeight = FuzzAttenuation * (1.0 - Luminance(SpecE));
+							const float3 Spec0Albedo  = FuzzAttenuation * (1.0 - SpecBlend) * Spec0.E;
+							const float3 Spec1Albedo  = FuzzAttenuation * (      SpecBlend) * Spec1.E;
+							const float  FuzzAlbedo   = FuzzAmount * Fuzz.E;
+							const float3 SlabAlbedo = MaxLobeWeight * (DiffuseWeight * DiffuseColor + Spec0Albedo + Spec1Albedo + FuzzAlbedo);
+							Pdf = LobeColorToWeight(SlabAlbedo);
+						}
+						break;
+					}
+					default:
+					{
+						// In theory none of the other slab types can be layered
+						break;
+					}
+				}
+				SlabCDFSum += Pdf;
+				SlabCDF[BSDFIdx] = SlabCDFSum;
+			}
+			if (SlabCDFSum > 0.0)
+			{
+				// linear search
+				float PreviousCdfValue = 0;
+				float RandSample = SlabCDFSum * PackedPayload.GetStochasticSlabRand();
+				for (BSDFIdx = 0; BSDFIdx < StrataTree.BSDFCount - 1; ++BSDFIdx)
+				{
+					if (RandSample < SlabCDF[BSDFIdx])
+					{
+						break;
+					}
+					PreviousCdfValue = SlabCDF[BSDFIdx];
+				}
+				SlabPdf = (SlabCDF[BSDFIdx] - PreviousCdfValue) / SlabCDFSum;
+			}
+			else
+			{
+				BSDFIdx = 0;
+				SlabPdf = 0.0;
+			}
+		}
+		if (SlabPdf > 0)
+		{
+#elif STRATA_CLAMPED_BSDF_COUNT == 1
+		if (StrataTree.BSDFCount > 0)
+		{
+			// there is only one slab active
+			const int BSDFIdx = 0;
+#else
+// shouldn't be possible to reach this point
+#error Unexpected BSDF Count
+#endif // STRATA_CLAMPED_BSDF_COUNT 
+
+
+			Payload.WorldNormal = StrataPixelHeader.SharedLocalBases.Normals[BSDF_GETSHAREDLOCALBASISID(BSDF)];
+			Payload.WorldTangent = StrataPixelHeader.SharedLocalBases.Tangents[BSDF_GETSHAREDLOCALBASISID(BSDF)];
+
+			switch (BSDF_GETTYPE(BSDF))
+			{
+				case STRATA_BSDF_TYPE_SLAB:
+				{
+					Payload.DiffuseColor      = SLAB_DIFFUSEALBEDO(BSDF);
+					Payload.SpecularColor     = SLAB_F0(BSDF);
+					Payload.SpecularEdgeColor = SLAB_F90(BSDF);
+					Payload.RoughnessData.x   = MakeRoughnessSafe(SLAB_ROUGHNESS(BSDF), PackedPayload.GetPathRoughness());
+					Payload.RoughnessData.y = Payload.RoughnessData.x;
+					Payload.RoughnessData.z = 0;
+					if (BSDF_GETHASHAZINESS(BSDF))
+					{
+						const FHaziness Haziness = UnpackHaziness(SLAB_HAZINESS(BSDF));
+						// STRATA_TODO: detect this and use the old SHADINGMODELID_CLEARCOAT?
+						if (!Haziness.bSimpleClearCoat)
+						{
+							Payload.RoughnessData.y = MakeRoughnessSafe(Haziness.Roughness, PackedPayload.GetPathRoughness());
+							Payload.RoughnessData.z = Haziness.Weight;
+						}
+					}
+					Payload.Anisotropy = BSDF_GETHASANISOTROPY(BSDF) ? SLAB_ANISOTROPY(BSDF) : 0.0;
+					if (BSDF_GETHASFUZZ(BSDF))
+					{
+						Payload.FuzzAmount    = SLAB_FUZZ_AMOUNT(BSDF);
+						Payload.FuzzColor     = SLAB_FUZZ_COLOR(BSDF);
+						Payload.FuzzRoughness = MakeRoughnessSafe(SLAB_FUZZ_ROUGHNESS(BSDF), STRATA_MIN_FUZZ_ROUGHNESS);
+					}
+					else
+					{
+						Payload.FuzzAmount    = 0.0;
+						Payload.FuzzColor     = 0.0;
+						Payload.FuzzRoughness = 0.0;
+					}
+					if (BSDF_GETSSSTYPE(BSDF) == SSS_TYPE_DIFFUSION)
+					{
+						Payload.MeanFreePath = SLAB_SSSMFP(BSDF);
+						Payload.PhaseG       = SLAB_SSSPHASEANISOTROPY(BSDF); // STRATA_TODO: evaluate if we always want this or not ...
+					}
+					else
+					{
+						// STRATA_TODO: Do something meaningful for the other SSS types ...
+						Payload.MeanFreePath = 0;
+						Payload.PhaseG       = 0;
+					}
+					Payload.WeightV             = BSDF.LuminanceWeightV;
+					Payload.TransmittanceN      = BSDF.TransmittanceAboveAlongN;
+					Payload.CoverageAboveAlongN = BSDF.CoverageAboveAlongN;
+					Payload.Ior                 = 0.0; // STRATA_TODO: figure out when to use refraction ...
+					Payload.ShadingModelID      = SHADINGMODELID_STRATA;
+					break;
+				}
+				case STRATA_BSDF_TYPE_HAIR:
+				{
+					// clamp the roughness to whatever came along the path
+					Payload.SetBaseColor(HAIR_BASECOLOR(BSDF));
+					Payload.SetHairLongitudinalRoughness(max(HAIR_ROUGHNESS(BSDF), PackedPayload.GetPathRoughness()));
+					Payload.SetHairAzimuthalRoughness(HAIR_SCATTER(BSDF)); // STRATA_TODO: Map from scatter param to azimuthal roughness here instead of in the raygen
+					Payload.SetHairSpecular(HAIR_SPECULAR(BSDF));
+
+	#if HAIR_STRAND_MESH_FACTORY
+					Payload.WorldSmoothNormal = Payload.WorldNormal = Payload.WorldGeoNormal;
+					Payload.WorldTangent      = MaterialParameters.TangentToWorld[2];
+					Payload.SetHairPrimitiveUV(MaterialParameters.HairPrimitiveUV);
+	#else
+					Payload.SetHairPrimitiveUV(0.5);
+	#endif
+					Payload.ShadingModelID = SHADINGMODELID_HAIR;
+					break;
+				}
+				case STRATA_BSDF_TYPE_EYE:
+				{
+					Payload.DiffuseColor = EYE_DIFFUSEALBEDO(BSDF);
+					Payload.SpecularColor = EYE_F0(BSDF);
+					// NOTE: EYE_F90 is always 1.0 (does not exist in legacy model), so don't bother with it
+					Payload.SetEyeRoughness(max(EYE_ROUGHNESS(BSDF), PackedPayload.GetPathRoughness()));
+
+					// STRATA_TODO: handle EYE_SSSPROFILEID(BSDF)
+
+					float IrisMask = EYE_IRISMASK(BSDF);
+					float IrisDistance = EYE_IRISDISTANCE(BSDF);
+
+					// see logic in non-substrate code below
+					const float3 PlaneNormal = normalize(EYE_IRISPLANENORMAL(BSDF));
+					const float3 CausticNormal = normalize(lerp(PlaneNormal, -Payload.WorldNormal, IrisMask * IrisDistance));
+
+					Payload.SetEyeCausticNormal(CausticNormal);
+					Payload.SetEyeIrisMask(IrisMask);
+					Payload.SetEyeIrisNormal(normalize(EYE_IRISNORMAL(BSDF)));
+					Payload.ShadingModelID = SHADINGMODELID_EYE;
+					break;
+				}
+				default:
+				{
+					// STRATA_TODO: handle all other cases
+					break;
+				}
+			}
+#if STRATA_CLAMPED_BSDF_COUNT > 1
+			// account for the probability of having chosen this slab
+			Payload.BSDFOpacity /= SlabPdf;
+#endif
+		} // close brace opened by STRATA_CLAMPED_BSDF_COUNT > 1 or STRATA_CLAMPED_BSDF_COUNT == 1 ifdefs above
+#undef  BSDF
+	} // if BSDFCount > 0
+
+#endif // STRATA_OPTIMIZED_UNLIT
+
+
+#if TWO_SIDED_MATERIAL == 0
+	if (MaterialParameters.TwoSidedSign < 0)
+	{
+		// when viewing the surface from "inside", don't include emission
+		Payload.Radiance = 0;
+	}
+#endif
+
+#else // STRATA_ENABLED
+	Payload.Radiance = GetMaterialEmissive(PixelMaterialInputs);
+
+	Payload.WorldNormal = MaterialParameters.WorldNormal;
+	Payload.WorldNormal = AdjustShadingNormal(Payload.WorldNormal, Payload.WorldGeoNormal, WorldRayDirection());
+
 	// Store the results in local variables and reuse instead of calling the functions multiple times.
 	half3 BaseColor = GetMaterialBaseColor(PixelMaterialInputs);
 	half  Metallic = GetMaterialMetallic(PixelMaterialInputs);
@@ -338,22 +657,7 @@ RAY_TRACING_ENTRY_CLOSEST_HIT(PathTracingMaterialCHS,
 	Ior = GetRefractionIor(PixelMaterialInputs);
 #endif
 
-	Payload.TranslatedWorldPos = LWCToFloat(LWCAdd(MaterialParameters.AbsoluteWorldPosition, ResolvedView.PreViewTranslation));
-	float GeoNormalSign = MaterialParameters.TwoSidedSign;
-#if !VF_SUPPORTS_RAYTRACING_PREPARE_MATERIAL_PIXEL_PARAMETERS
-	// Because the geometric normal is computed directly in world space
-	// it doesn't reflect the sign flip from the object transform, so apply it here
-	GeoNormalSign *= GetPrimitive_DeterminantSign(MaterialParameters.PrimitiveId);
-#endif
-	Payload.WorldGeoNormal = GeoNormalSign * GeoNormal;
-	Payload.WorldNormal = MaterialParameters.WorldNormal;
-	Payload.WorldSmoothNormal = MaterialParameters.TwoSidedSign * TransformTangentNormalToWorld(MaterialParameters.TangentToWorld, float3(0, 0, 1));
-
-	// fix dark edges in reflections from shading normals
-	Payload.WorldNormal       = AdjustShadingNormal(Payload.WorldNormal      , Payload.WorldGeoNormal, WorldRayDirection());
-	Payload.WorldSmoothNormal = AdjustShadingNormal(Payload.WorldSmoothNormal, Payload.WorldGeoNormal, WorldRayDirection());
 
-	Payload.Radiance = GetMaterialEmissive(PixelMaterialInputs);
 	Payload.Radiance *= Payload.BSDFOpacity; // pre-multiply
 
 
@@ -707,6 +1011,9 @@ RAY_TRACING_ENTRY_CLOSEST_HIT(PathTracingMaterialCHS,
 	// adjust after everything else (because SSS profile case needs to decode the dual spec info from the unmodified roughness)
 	Payload.Roughness = AdjustMaterialRoughness(Payload.Roughness, PackedPayload.GetPathRoughness());
 
+#endif
+
+
 	PackedPayload = PackPathTracingPayload(Payload);
 }
 
diff --git a/Engine/Shaders/Private/PathTracing/Volume/PathTracingAtmosphere.ush b/Engine/Shaders/Private/PathTracing/Volume/PathTracingAtmosphere.ush
--- a/Engine/Shaders/Private/PathTracing/Volume/PathTracingAtmosphere.ush
+++ b/Engine/Shaders/Private/PathTracing/Volume/PathTracingAtmosphere.ush
@@ -147,7 +147,7 @@ FPackedPathTracingPayload AtmosphereGetBlockerHit(float3 Origin, float3 Directio
 	float3 WorldN = normalize(WorldP - Center);
 
 	FPathTracingPayload Payload = (FPathTracingPayload)0;
-	Payload.BaseColor = Atmosphere.GroundAlbedo.rgb;
+	Payload.SetBaseColor(Atmosphere.GroundAlbedo.rgb);
 	Payload.BSDFOpacity = 1.0;
 	Payload.WorldNormal = Payload.WorldGeoNormal = Payload.WorldSmoothNormal = WorldN;
 	Payload.ShadingModelID = SHADINGMODELID_NUM; // defaults to diffuse
diff --git a/Engine/Shaders/Private/PathTracing/Volume/PathTracingVolumeCommon.ush b/Engine/Shaders/Private/PathTracing/Volume/PathTracingVolumeCommon.ush
--- a/Engine/Shaders/Private/PathTracing/Volume/PathTracingVolumeCommon.ush
+++ b/Engine/Shaders/Private/PathTracing/Volume/PathTracingVolumeCommon.ush
@@ -207,7 +207,7 @@ FPathTracingPayload CreateMediumHitPayload(float HitT, float3 TranslatedWorldPos
 	const float3 SigmaS = ShadedResult.SigmaSRayleigh + ShadedResult.SigmaSHG;
 	const float3 RayleighWeight = select(SigmaS > 0.0, min(ShadedResult.SigmaSRayleigh, SigmaS) / SigmaS, 0.0);
 	const float3 HGWeight = select(SigmaS > 0, min(ShadedResult.SigmaSHG, SigmaS) / SigmaS, 0.0);
-	Result.BaseColor = RayleighWeight;
+	Result.SetBaseColor(RayleighWeight);
 	Result.SetHG(HGWeight, ShadedResult.PhaseG);
 	return Result;
 }
diff --git a/Engine/Shaders/Private/Strata/Strata.ush b/Engine/Shaders/Private/Strata/Strata.ush
--- a/Engine/Shaders/Private/Strata/Strata.ush
+++ b/Engine/Shaders/Private/Strata/Strata.ush
@@ -49,7 +49,7 @@
 
 // During the unpacking of the Strata material data, if a slab has some sub-surface scattering, its BaseColor/Specular value will be 
 // patched/overriden with white albedo (and optionally no specular) for computing the incoming irradiance without any albedo information. 
-// This is required by the post-process/screen-space sub-surface methods. This overriden can be opt-out. This is used during the classification 
+// This is required by the post-process/screen-space sub-surface methods. This override can be opt-out. This is used during the classification 
 // pass for writing out the correct base color value into the FStrataSubsurfaceData
 #ifndef STRATA_SSS_MATERIAL_OVERRIDE
 #define STRATA_SSS_MATERIAL_OVERRIDE 1
@@ -747,7 +747,7 @@ FStrataOpaqueRoughRefractionData StrataUnpackOpaqueRoughRefractionData(float3 In
 //			|
 //			|
 //			v
-//	------------------- ^	Surfcace of roughness R
+//	------------------- ^	Surface of roughness R
 //		   /|\			|
 //		  / | \			| Depth D
 //		 /  |  \		|
@@ -765,7 +765,7 @@ float GetVarianceFromRoughness(float Roughness)
 ///////////////////////////////////////////////////////////////////////////////
 // Sub-surface 
 
-// Stored in frist slice as UINT
+// Stored in first slice as UINT
 // This is hot data that is often accessed by intrested passes.
 struct FStrataSubsurfaceHeader
 {
diff --git a/Engine/Source/Runtime/Renderer/Private/PathTracing.cpp b/Engine/Source/Runtime/Renderer/Private/PathTracing.cpp
--- a/Engine/Source/Runtime/Renderer/Private/PathTracing.cpp
+++ b/Engine/Source/Runtime/Renderer/Private/PathTracing.cpp
@@ -337,6 +337,26 @@ TAutoConsoleVariable<int32> CVarPathTracingOutputPostProcessResources(
 	ECVF_RenderThreadSafe
 );
 
+TAutoConsoleVariable<int32> CVarPathTracingSubstrateUseSimplifiedMaterial(
+	TEXT("r.PathTracing.Substrate.UseSimplifiedMaterials"),
+	0,
+	TEXT("Instead of evaluating all layers, use an optimized material in which all slabs have been merged. This is mainly intended for debugging and requires r.PathTracing.Substrate.CompileSimplifiedMaterials to be true.\n")
+	TEXT("0: off (default)\n")
+	TEXT("1: on\n"),
+	ECVF_RenderThreadSafe
+);
+
+TAutoConsoleVariable<int32> CVarPathTracingSubstrateCompileSimplifiedMaterial(
+	TEXT("r.PathTracing.Substrate.CompileSimplifiedMaterials"),
+	0,
+	TEXT("Compile a simplified representation of Substrate materials which merges all slabs into one. This is mainly intended for debugging purposes. Enabling this double the number of path tracing shader permutations.\n")
+	TEXT("0: off (default)\n")
+	TEXT("1: on\n"),
+	ECVF_RenderThreadSafe | ECVF_ReadOnly
+);
+
+
+
 BEGIN_SHADER_PARAMETER_STRUCT(FPathTracingData, )
 	SHADER_PARAMETER(float, BlendFactor)
 	SHADER_PARAMETER(uint32, Iteration)
@@ -781,7 +801,13 @@ static FPathTracingFogParameters PrepareFogParameters(const FViewInfo& View, con
 	return Parameters;
 }
 
-IMPLEMENT_RT_PAYLOAD_TYPE(ERayTracingPayloadType::PathTracingMaterial, 64); // TODO: Expand this when strata is enabled (use GetStrataRaytracingMaterialPayload())
+static uint32 GetPathtracingMaterialPayloadSize()
+{
+	// Strata uses a slightly bigger payload as the basic slab contains more information
+	return Strata::IsStrataEnabled() ? 76u : 64u;
+}
+
+IMPLEMENT_RT_PAYLOAD_TYPE_FUNCTION(ERayTracingPayloadType::PathTracingMaterial, GetPathtracingMaterialPayloadSize);
 IMPLEMENT_RT_PAYLOAD_TYPE(ERayTracingPayloadType::GPULightmass, 32);
 
 class FPathTracingRG : public FGlobalShader
@@ -1211,7 +1237,7 @@ static bool NeedsAnyHitShader(const FMaterial& RESTRICT MaterialResource)
 	return NeedsAnyHitShader(MaterialResource.IsMasked(), MaterialResource.IsDitherMasked());
 }
 
-template<bool UseAnyHitShader, bool UseIntersectionShader, bool IsGPULightmass>
+template<bool UseAnyHitShader, bool UseIntersectionShader, bool IsGPULightmass, bool SimplifyStrata>
 class TPathTracingMaterial : public FMeshMaterialShader
 {
 	DECLARE_SHADER_TYPE(TPathTracingMaterial, MeshMaterial);
@@ -1249,6 +1275,13 @@ public:
 			// only need to compile the intersection shader permutation if the VF actually requires it
 			return false;
 		}
+		if (SimplifyStrata && (!Strata::IsStrataEnabled() || CVarPathTracingSubstrateCompileSimplifiedMaterial.GetValueOnAnyThread() == 0))
+		{
+			// don't compile the extra strata permutation if:
+			//    Substrate is not enabled on this project
+			// or the user did not request the extra permutations to be compiled (default)
+			return false;
+		}
 		if (IsGPULightmass)
 		{
 			return ShouldCompileGPULightmassShadersForProject(Parameters);
@@ -1267,6 +1300,7 @@ public:
 		OutEnvironment.SetDefine(TEXT("USE_RAYTRACED_TEXTURE_RAYCONE_LOD"), 0);
 		OutEnvironment.SetDefine(TEXT("SCENE_TEXTURES_DISABLED"), 1);
 		OutEnvironment.SetDefine(TEXT("SIMPLIFIED_MATERIAL_SHADER"), IsGPULightmass);
+		OutEnvironment.SetDefine(TEXT("STRATA_USE_FULLYSIMPLIFIED_MATERIAL"), IsGPULightmass || SimplifyStrata);
 		FMeshMaterialShader::ModifyCompilationEnvironment(Parameters, OutEnvironment);
 	}
 
@@ -1307,19 +1341,29 @@ public:
 
 
 // TODO: It would be nice to avoid this template boilerplate and just use ordinary permutations. This would require allowing the FunctionName for the material to be dependent on the permutation somehow
-using FPathTracingMaterialCHS        = TPathTracingMaterial<false, false, false>;
-using FPathTracingMaterialCHS_AHS    = TPathTracingMaterial<true , false, false>;
-using FPathTracingMaterialCHS_IS     = TPathTracingMaterial<false, true , false>;
-using FPathTracingMaterialCHS_AHS_IS = TPathTracingMaterial<true , true , false>;
+using FPathTracingMaterialCHS        = TPathTracingMaterial<false, false, false, false>;
+using FPathTracingMaterialCHS_AHS    = TPathTracingMaterial<true , false, false, false>;
+using FPathTracingMaterialCHS_IS     = TPathTracingMaterial<false, true , false, false>;
+using FPathTracingMaterialCHS_AHS_IS = TPathTracingMaterial<true , true , false, false>;
+using FPathTracingMaterialSimplifiedCHS        = TPathTracingMaterial<false, false, false, true>;
+using FPathTracingMaterialSimplifiedCHS_AHS    = TPathTracingMaterial<true , false, false, true>;
+using FPathTracingMaterialSimplifiedCHS_IS     = TPathTracingMaterial<false, true , false, true>;
+using FPathTracingMaterialSimplifiedCHS_AHS_IS = TPathTracingMaterial<true , true , false, true>;
+
 
 // NOTE: lightmass doesn't work with intersection shader VFs at the moment, so avoid instantiating permutations that will never generate any shaders
-using FGPULightmassCHS               = TPathTracingMaterial<false, false, true>;
-using FGPULightmassCHS_AHS           = TPathTracingMaterial<true , false, true>;
+using FGPULightmassCHS               = TPathTracingMaterial<false, false, true, true>;
+using FGPULightmassCHS_AHS           = TPathTracingMaterial<true , false, true, true>;
 
 IMPLEMENT_MATERIAL_SHADER_TYPE(template <>, FPathTracingMaterialCHS       , TEXT("/Engine/Private/PathTracing/PathTracingMaterialHitShader.usf"), TEXT("closesthit=PathTracingMaterialCHS"), SF_RayHitGroup);
 IMPLEMENT_MATERIAL_SHADER_TYPE(template <>, FPathTracingMaterialCHS_AHS   , TEXT("/Engine/Private/PathTracing/PathTracingMaterialHitShader.usf"), TEXT("closesthit=PathTracingMaterialCHS anyhit=PathTracingMaterialAHS"), SF_RayHitGroup);
 IMPLEMENT_MATERIAL_SHADER_TYPE(template <>, FPathTracingMaterialCHS_IS    , TEXT("/Engine/Private/PathTracing/PathTracingMaterialHitShader.usf"), TEXT("closesthit=PathTracingMaterialCHS intersection=MaterialIS"), SF_RayHitGroup);
 IMPLEMENT_MATERIAL_SHADER_TYPE(template <>, FPathTracingMaterialCHS_AHS_IS, TEXT("/Engine/Private/PathTracing/PathTracingMaterialHitShader.usf"), TEXT("closesthit=PathTracingMaterialCHS anyhit=PathTracingMaterialAHS intersection=MaterialIS"), SF_RayHitGroup);
+IMPLEMENT_MATERIAL_SHADER_TYPE(template <>, FPathTracingMaterialSimplifiedCHS       , TEXT("/Engine/Private/PathTracing/PathTracingMaterialHitShader.usf"), TEXT("closesthit=PathTracingMaterialCHS"), SF_RayHitGroup);
+IMPLEMENT_MATERIAL_SHADER_TYPE(template <>, FPathTracingMaterialSimplifiedCHS_AHS   , TEXT("/Engine/Private/PathTracing/PathTracingMaterialHitShader.usf"), TEXT("closesthit=PathTracingMaterialCHS anyhit=PathTracingMaterialAHS"), SF_RayHitGroup);
+IMPLEMENT_MATERIAL_SHADER_TYPE(template <>, FPathTracingMaterialSimplifiedCHS_IS    , TEXT("/Engine/Private/PathTracing/PathTracingMaterialHitShader.usf"), TEXT("closesthit=PathTracingMaterialCHS intersection=MaterialIS"), SF_RayHitGroup);
+IMPLEMENT_MATERIAL_SHADER_TYPE(template <>, FPathTracingMaterialSimplifiedCHS_AHS_IS, TEXT("/Engine/Private/PathTracing/PathTracingMaterialHitShader.usf"), TEXT("closesthit=PathTracingMaterialCHS anyhit=PathTracingMaterialAHS intersection=MaterialIS"), SF_RayHitGroup);
+
 IMPLEMENT_MATERIAL_SHADER_TYPE(template <>, FGPULightmassCHS              , TEXT("/Engine/Private/PathTracing/PathTracingGPULightmassMaterialHitShader.usf"), TEXT("closesthit=GPULightmassMaterialCHS"), SF_RayHitGroup);
 IMPLEMENT_MATERIAL_SHADER_TYPE(template <>, FGPULightmassCHS_AHS          , TEXT("/Engine/Private/PathTracing/PathTracingGPULightmassMaterialHitShader.usf"), TEXT("closesthit=GPULightmassMaterialCHS anyhit=GPULightmassMaterialAHS"), SF_RayHitGroup);
 
@@ -1407,19 +1451,43 @@ bool FRayTracingMeshProcessor::ProcessPathTracing(
 		{
 			case ERayTracingMeshCommandsMode::PATH_TRACING:
 			{
+				// In order to use Substrate simplified materials, strata has to be enabled, we have to have _compiled_ the extra permutations _and_ the runtime toggle must be true
+				const bool bUseSimplifiedMaterial = Strata::IsStrataEnabled() &&
+					CVarPathTracingSubstrateCompileSimplifiedMaterial.GetValueOnRenderThread() != 0 &&
+					CVarPathTracingSubstrateUseSimplifiedMaterial.GetValueOnRenderThread() != 0;
 				if (NeedsAnyHitShader(MaterialResource))
 				{
-					if (bUseProceduralPrimitive)
-						ShaderTypes.AddShaderType<FPathTracingMaterialCHS_AHS_IS>();
+					if (bUseSimplifiedMaterial)
+					{
+						if (bUseProceduralPrimitive)
+							ShaderTypes.AddShaderType<FPathTracingMaterialSimplifiedCHS_AHS_IS>();
+						else
+							ShaderTypes.AddShaderType<FPathTracingMaterialSimplifiedCHS_AHS>();
+					}
 					else
-						ShaderTypes.AddShaderType<FPathTracingMaterialCHS_AHS>();
+					{
+						if (bUseProceduralPrimitive)
+							ShaderTypes.AddShaderType<FPathTracingMaterialCHS_AHS_IS>();
+						else
+							ShaderTypes.AddShaderType<FPathTracingMaterialCHS_AHS>();
+					}
 				}
 				else
 				{
-					if (bUseProceduralPrimitive)
-						ShaderTypes.AddShaderType<FPathTracingMaterialCHS_IS>();
+					if (bUseSimplifiedMaterial)
+					{
+						if (bUseProceduralPrimitive)
+							ShaderTypes.AddShaderType<FPathTracingMaterialSimplifiedCHS_IS>();
+						else
+							ShaderTypes.AddShaderType<FPathTracingMaterialSimplifiedCHS>();
+					}
 					else
-						ShaderTypes.AddShaderType<FPathTracingMaterialCHS>();
+					{
+						if (bUseProceduralPrimitive)
+							ShaderTypes.AddShaderType<FPathTracingMaterialCHS_IS>();
+						else
+							ShaderTypes.AddShaderType<FPathTracingMaterialCHS>();
+					}
 				}
 				break;
 			}
@@ -2170,6 +2238,11 @@ void FDeferredShadingSceneRenderer::RenderPathTracing(
 	Config.LightShowFlags |= View.Family->EngineShowFlags.LightingOnlyOverride       ? 1 << 11 : 0;
 	Config.LightShowFlags |= View.Family->EngineShowFlags.ReflectionOverride         ? 1 << 12 : 0;
 	Config.LightShowFlags |= View.Family->EngineShowFlags.SubsurfaceScattering       ? 1 << 13 : 0;
+	// the following affects which material shaders get used and therefore change the image
+	if (Strata::IsStrataEnabled() && CVarPathTracingSubstrateCompileSimplifiedMaterial.GetValueOnRenderThread() != 0)
+	{
+		Config.LightShowFlags |= CVarPathTracingSubstrateUseSimplifiedMaterial.GetValueOnRenderThread() != 0 ? 1 << 14 : 0;
+	}
 
 	PreparePathTracingData(Scene, View, Config.PathTracingData);
 
