Path Tracer: Fix non-slab shading models with Substrate support, add support for SSS profile legacy shading model, fix one-sided emission support

diff --git a/Engine/Shaders/Private/PathTracing/Material/PathTracingMaterialSampling.ush b/Engine/Shaders/Private/PathTracing/Material/PathTracingMaterialSampling.ush
--- a/Engine/Shaders/Private/PathTracing/Material/PathTracingMaterialSampling.ush
+++ b/Engine/Shaders/Private/PathTracing/Material/PathTracingMaterialSampling.ush
@@ -35,10 +35,7 @@ float3 EstimateMaterialAlbedo(FPathTracingPayload Payload)
 		// this object does not scatter light
 		return 0;
 	}
-#if PATHTRACING_STRATA_PAYLOAD
-	const float3 MaxWeight = Payload.WeightV * Payload.BSDFOpacity * lerp(1.0, Payload.TransmittanceN, Payload.CoverageAboveAlongN);
-#endif
-	// TODO: make this more accurate by estimating the directional albedo
+	// TODO: make this more accurate by estimating the directional albedo?
 	float3 Albedo = 0;
 	switch (Payload.ShadingModelID)
 	{
@@ -51,11 +48,12 @@ float3 EstimateMaterialAlbedo(FPathTracingPayload Payload)
 #elif PATHTRACING_STRATA_PAYLOAD
 		case SHADINGMODELID_STRATA:
 		{
-			Albedo = MaxWeight * (Payload.DiffuseColor + Payload.SpecularColor + Payload.FuzzAmount * Payload.FuzzColor);
+			Albedo = Payload.DiffuseColor + Payload.SubsurfaceColor + Payload.SpecularColor + Payload.FuzzAmount * Payload.FuzzColor;
+			Albedo *= Payload.GetMaxLobeWeight();
 			break;
 		}
 		case SHADINGMODELID_EYE:
-			Albedo = MaxWeight * (Payload.DiffuseColor + Payload.SpecularColor);
+			Albedo = Payload.DiffuseColor + Payload.SubsurfaceColor + Payload.SpecularColor;
 			break;
 		case SHADINGMODELID_HAIR:
 		// STRATA_TODO: Add more cases here?
@@ -92,7 +90,8 @@ float3 EstimateMaterialAlbedo(FPathTracingPayload Payload)
 			break;
 	}
 	Albedo *= Payload.BSDFOpacity;
-#if !SIMPLIFIED_MATERIAL_SHADER
+
+#if !SIMPLIFIED_MATERIAL_SHADER && !PATHTRACING_STRATA_PAYLOAD
 	if (Payload.HasRefraction())
 	{
 		Albedo += 1 - Payload.BSDFOpacity; // glass albedo is 1.0
@@ -286,6 +285,14 @@ FPathTracingPayload UnpackPathTracingPayloadMaterial(FPackedPathTracingPayload P
 			HitPayload.DiffuseColor = 0;
 			break;
 		}
+		case SHADINGMODELID_EYE:
+		{
+			// eye still needs a diffuse component for the iris center
+			float IrisMask = HitPayload.GetEyeIrisMask();
+			HitPayload.SubsurfaceColor = HitPayload.DiffuseColor - IrisMask * HitPayload.DiffuseColor;
+			HitPayload.DiffuseColor *= IrisMask;
+			break;
+		}
 	}
 #else
 	switch (HitPayload.ShadingModelID)
@@ -387,7 +394,6 @@ FSSSRandomWalkInfo GetMaterialSSSInfo(FPathTracingPayload Payload, float3 V_Worl
 {
 	FSSSRandomWalkInfo Result = (FSSSRandomWalkInfo)0;
 #if PATHTRACING_STRATA_PAYLOAD
-	// STRATA_TODO: figure out how to handle the multi-slab case as well
 	if (Payload.ShadingModelID == SHADINGMODELID_STRATA)
 	{
 		if (any(Payload.MeanFreePath > 0.0))
@@ -411,6 +417,24 @@ FSSSRandomWalkInfo GetMaterialSSSInfo(FPathTracingPayload Payload, float3 V_Worl
 			// STRATA_TODO: evaluate if we want to set Result.G in all cases?
 		}
 	}
+	if (Payload.ShadingModelID == SHADINGMODELID_EYE)
+	{
+		if (any(Payload.MeanFreePath > 0.0))
+		{
+			Result.Color = Payload.SubsurfaceColor;
+			Result.Radius = Payload.MeanFreePath;
+			const float NoV = saturate(dot(Payload.WorldNormal, V_World));
+			const float3 SpecE = ComputeGGXSpecEnergyTermsRGB(Payload.GetEyeRoughness(), NoV, Payload.SpecularColor).E;
+
+			const float DiffWeight = (1 - Luminance(SpecE));
+			const float3 SSSAlbedo  = DiffWeight * Payload.SubsurfaceColor;
+			const float3 DiffAlbedo = DiffWeight * Payload.DiffuseColor;
+			const float3 SpecAlbedo = SpecE;
+			Result.Prob = LobeSelectionProb(SSSAlbedo, DiffAlbedo + SpecAlbedo);
+			Result.Weight = Payload.BSDFOpacity * DiffWeight;
+
+		}
+	}
 #else
 	Result.Color = Payload.SubsurfaceColor;
 	Result.Radius = Payload.SubsurfaceRadius;
diff --git a/Engine/Shaders/Private/PathTracing/PathTracingMaterialHitShader.usf b/Engine/Shaders/Private/PathTracing/PathTracingMaterialHitShader.usf
--- a/Engine/Shaders/Private/PathTracing/PathTracingMaterialHitShader.usf
+++ b/Engine/Shaders/Private/PathTracing/PathTracingMaterialHitShader.usf
@@ -367,7 +367,6 @@ RAY_TRACING_ENTRY_CLOSEST_HIT(PathTracingMaterialCHS,
 	Payload.TransparencyColor = 0;
 
 	const float3 V_World = MaterialParameters.CameraVector;
-	FStrataData StrataData= PixelMaterialInputs.GetFrontStrataData();
 	FStrataPixelHeader StrataPixelHeader = MaterialParameters.GetFrontStrataHeader();
 	FStrataIntegrationSettings IntegrationSettings = InitStrataIntegrationSettings();
 	FStrataAddressing NullStrataAddressing = (FStrataAddressing)0;	// Fake unused in StrataCreateBSDFContext when using Forward inline shading
@@ -385,7 +384,7 @@ RAY_TRACING_ENTRY_CLOSEST_HIT(PathTracingMaterialCHS,
 			// Not that we must recompute tangent vector to keep the basis frame orthogonal
 			for (uint i = 0; i < StrataPixelHeader.SharedLocalBases.Count; ++i)
 			{
-				float3 N = StrataPixelHeader.SharedLocalBases.Normals[i];
+				float3 N = normalize(StrataPixelHeader.SharedLocalBases.Normals[i]);
 				float3 T = StrataPixelHeader.SharedLocalBases.Tangents[i];
 				N = AdjustShadingNormal(N, Payload.WorldGeoNormal, WorldRayDirection());
 				float3 B = normalize(cross(N, T)); // N and T are no longer orthogonal to each other, need to normalize
@@ -398,6 +397,14 @@ RAY_TRACING_ENTRY_CLOSEST_HIT(PathTracingMaterialCHS,
 
 #define BSDF StrataTree.BSDFs[BSDFIdx]
 
+		// Accumulate radiance across all BSDFs
+		Payload.Radiance = 0.0;
+		for (int BSDFIdx = 0; BSDFIdx < StrataTree.BSDFCount; ++BSDFIdx)
+		{
+			Payload.Radiance += BSDF.LuminanceWeightV * BSDF_GETEMISSIVE(BSDF);
+		}
+
+
 #if STRATA_CLAMPED_BSDF_COUNT > 1
 		// The material might be using more than one slab -- insert code to stochastically pick one 
 		int BSDFIdx = 0;
@@ -424,8 +431,8 @@ RAY_TRACING_ENTRY_CLOSEST_HIT(PathTracingMaterialCHS,
 						// Is the slab even visible at all?
 						if (any(MaxLobeWeight > 0.0))
 						{
-							const float3x3 Basis = StrataGetBSDFSharedBasis_InlineShading(StrataPixelHeader, BSDF, NullStrataAddressing);
-							const float3 V = mul(Basis, V_World);
+							const float3 N = StrataPixelHeader.SharedLocalBases.Normals[BSDF_GETSHAREDLOCALBASISID(BSDF)];
+							const float NoV = saturate(dot(N, V_World));
 
 							float3 DiffuseColor = SLAB_DIFFUSEALBEDO(BSDF);
 							float3 F0 = SLAB_F0(BSDF);
@@ -461,7 +468,6 @@ RAY_TRACING_ENTRY_CLOSEST_HIT(PathTracingMaterialCHS,
 								FuzzRoughness = MakeRoughnessSafe(FuzzRoughness, STRATA_MIN_FUZZ_ROUGHNESS);
 							}
 
-							const float NoV = saturate(V.z);
 							// STRATA_TODO: Support Fresnel82?
 							const FBxDFEnergyTermsRGB Spec0 = ComputeGGXSpecEnergyTermsRGB(Roughness0, NoV, F0, F90);
 							const FBxDFEnergyTermsRGB Spec1 = ComputeGGXSpecEnergyTermsRGB(Roughness1, NoV, F0, F90);
@@ -561,6 +567,13 @@ RAY_TRACING_ENTRY_CLOSEST_HIT(PathTracingMaterialCHS,
 						Payload.MeanFreePath = SLAB_SSSMFP(BSDF);
 						Payload.PhaseG       = SLAB_SSSPHASEANISOTROPY(BSDF); // STRATA_TODO: evaluate if we always want this or not ...
 					}
+					else if (BSDF_GETSSSTYPE(BSDF) == SSS_TYPE_DIFFUSION_PROFILE)
+					{
+						const uint ProfileId = StrataSubsurfaceProfileIdTo8bits(SLAB_SSSPROFILEID(BSDF));
+						const float Scale = SLAB_SSSPROFILERADIUSSCALE(BSDF);
+						Payload.MeanFreePath = DecodeSSSProfileRadius(ProfileId, Payload.DiffuseColor, Scale);
+						Payload.PhaseG = 0.0;
+					}
 					else
 					{
 						// STRATA_TODO: Do something meaningful for the other SSS types ...
@@ -599,7 +612,13 @@ RAY_TRACING_ENTRY_CLOSEST_HIT(PathTracingMaterialCHS,
 					// NOTE: EYE_F90 is always 1.0 (does not exist in legacy model), so don't bother with it
 					Payload.SetEyeRoughness(max(EYE_ROUGHNESS(BSDF), PackedPayload.GetPathRoughness()));
 
-					// STRATA_TODO: handle EYE_SSSPROFILEID(BSDF)
+					if (BSDF_GETSSSTYPE(BSDF) == SSS_TYPE_DIFFUSION_PROFILE)
+					{
+						const uint ProfileId = StrataSubsurfaceProfileIdTo8bits(EYE_SSSPROFILEID(BSDF));
+						Payload.MeanFreePath = DecodeSSSProfileRadius(ProfileId, Payload.DiffuseColor, 1.0);
+						Payload.PhaseG = 0.0;
+					}
+
 
 					float IrisMask = EYE_IRISMASK(BSDF);
 					float IrisDistance = EYE_IRISDISTANCE(BSDF);
@@ -631,7 +650,7 @@ RAY_TRACING_ENTRY_CLOSEST_HIT(PathTracingMaterialCHS,
 #endif // STRATA_OPTIMIZED_UNLIT
 
 
-#if TWO_SIDED_MATERIAL == 0
+#if MATERIAL_TWOSIDED == 0
 	if (MaterialParameters.TwoSidedSign < 0)
 	{
 		// when viewing the surface from "inside", don't include emission
