Path Tracing: Fix dual specular blending in substrate slab

diff --git a/Engine/Shaders/Private/PathTracing/Material/PathTracingSubstrate.ush b/Engine/Shaders/Private/PathTracing/Material/PathTracingSubstrate.ush
--- a/Engine/Shaders/Private/PathTracing/Material/PathTracingSubstrate.ush
+++ b/Engine/Shaders/Private/PathTracing/Material/PathTracingSubstrate.ush
@@ -178,8 +178,8 @@ FMaterialSample Substrate_SampleMaterial(
 		const float3 F = Data.FuzzAttenuation * F_Schlick(Data.F0, Data.F90, VoH);
 #endif
 
-		Result.AddLobeWithMIS((1.0 - Payload.RoughnessData.x) * F * GGXResult0.x * Data.Spec0.W, GGXResult0.y, Data.LobePdf.y);
-		Result.AddLobeWithMIS((      Payload.RoughnessData.x) * F * GGXResult1.x * Data.Spec1.W, GGXResult1.y, Data.LobePdf.z);
+		Result.AddLobeWithMIS((1.0 - Payload.RoughnessData.z) * F * GGXResult0.x * Data.Spec0.W, Spec0Pdf, Data.LobePdf.y);
+		Result.AddLobeWithMIS((      Payload.RoughnessData.z) * F * GGXResult1.x * Data.Spec1.W, Spec1Pdf, Data.LobePdf.z);
 		Result.Pdf += Data.LobePdf.x * DiffPdf;
 		Result.Pdf += Data.LobePdf.w * ClothPdf;
 
