Substrata - Path tracer - fixed duplicated variable definition.

diff --git a/Engine/Shaders/Private/PathTracing/PathTracingMaterialHitShader.usf b/Engine/Shaders/Private/PathTracing/PathTracingMaterialHitShader.usf
--- a/Engine/Shaders/Private/PathTracing/PathTracingMaterialHitShader.usf
+++ b/Engine/Shaders/Private/PathTracing/PathTracingMaterialHitShader.usf
@@ -396,10 +396,11 @@ RAY_TRACING_ENTRY_CLOSEST_HIT(PathTracingMaterialCHS,
 		}
 
 #define BSDF StrataPixelHeader.StrataTree.BSDFs[BSDFIdx]
+		int BSDFIdx = 0;
 
 		// Accumulate radiance across all BSDFs
 		Payload.Radiance = 0.0;
-		for (int BSDFIdx = 0; BSDFIdx < BSDFCount; ++BSDFIdx)
+		for (BSDFIdx = 0; BSDFIdx < BSDFCount; ++BSDFIdx)
 		{
 			Payload.Radiance += BSDF.LuminanceWeightV * BSDF_GETEMISSIVE(BSDF);
 		}
@@ -407,7 +408,6 @@ RAY_TRACING_ENTRY_CLOSEST_HIT(PathTracingMaterialCHS,
 
 #if STRATA_CLAMPED_BSDF_COUNT > 1
 		// The material might be using more than one slab -- insert code to stochastically pick one 
-		int BSDFIdx = 0;
 		float SlabPdf = 1.0;
 		if (BSDFCount > 1)
 		{
@@ -518,8 +518,8 @@ RAY_TRACING_ENTRY_CLOSEST_HIT(PathTracingMaterialCHS,
 #elif STRATA_CLAMPED_BSDF_COUNT == 1
 		if (BSDFCount > 0)
 		{
-			// there is only one slab active
-			const int BSDFIdx = 0;
+			// There is only one slab active
+			BSDFIdx = 0;
 #else
 // shouldn't be possible to reach this point
 #error Unexpected BSDF Count
