Fix missing packing intrinsic on PS4/PS5.

diff --git a/Engine/Shaders/Private/HairStrands/HairStrandsForwardRaster.usf b/Engine/Shaders/Private/HairStrands/HairStrandsForwardRaster.usf
--- a/Engine/Shaders/Private/HairStrands/HairStrandsForwardRaster.usf
+++ b/Engine/Shaders/Private/HairStrands/HairStrandsForwardRaster.usf
@@ -43,7 +43,7 @@ struct FDepthRange
 
 uint PackDepthRange(FDepthRange In)
 {
-	return Pack2FloatsToUInt(In.MinZ, In.MaxZ);
+	return PackFloat2ToUInt(In.MinZ, In.MaxZ);
 }
 
 FDepthRange UnpackDepthRange(uint In)
diff --git a/Engine/Shaders/Private/PackUnpack.ush b/Engine/Shaders/Private/PackUnpack.ush
--- a/Engine/Shaders/Private/PackUnpack.ush
+++ b/Engine/Shaders/Private/PackUnpack.ush
@@ -73,14 +73,16 @@ uint2 UnpackUInt2FromUInt(uint Packed)
 //	float2 <-> u32 as f16.f16
 //
 
-uint Pack2FloatsToUInt(float X, float Y)
+#if !COMPILER_SUPPORTS_PACK_INTRINSICS
+uint PackFloat2ToUInt(float X, float Y)
 {
 	return PackUInt2ToUInt(f32tof16(X), f32tof16(Y));
 }
+#endif
 
 uint PackFloat2ToUInt(float2 XY)
 {
-	return Pack2FloatsToUInt(XY.x, XY.y);
+	return PackFloat2ToUInt(XY.x, XY.y);
 }
 
 float2 UnpackFloat2FromUInt(uint In)
