Path Tracer - Substrate - no need to modify tangent vector during shading as it is not used there. The tangent frame is ortho-normalized during BSDF evaluation if needed.

Also fix sidedness issues with normals when the material is not in tangent space.

diff --git a/Engine/Shaders/Private/PathTracing/PathTracingMaterialHitShader.usf b/Engine/Shaders/Private/PathTracing/PathTracingMaterialHitShader.usf
--- a/Engine/Shaders/Private/PathTracing/PathTracingMaterialHitShader.usf
+++ b/Engine/Shaders/Private/PathTracing/PathTracingMaterialHitShader.usf
@@ -385,12 +385,12 @@ RAY_TRACING_ENTRY_CLOSEST_HIT(PathTracingMaterialCHS,
 			for (uint i = 0; i < StrataPixelHeader.SharedLocalBases.Count; ++i)
 			{
 				float3 N = normalize(StrataPixelHeader.SharedLocalBases.Normals[i]);
-				float3 T = StrataPixelHeader.SharedLocalBases.Tangents[i];
+#if !MATERIAL_TANGENTSPACENORMAL
+                // already flipped if the material was in tangent space, so add the flip if it wasn't
+				N *= MaterialParameters.TwoSidedSign;
+#endif
 				N = AdjustShadingNormal(N, Payload.WorldGeoNormal, WorldRayDirection());
-				float3 B = normalize(cross(N, T)); // N and T are no longer orthogonal to each other, need to normalize
 				StrataPixelHeader.SharedLocalBases.Normals[i]  = N;
-				StrataPixelHeader.SharedLocalBases.Tangents[i] = cross(B, N); // no normalize needed here, will be orthogonal
-				
 			}
 #endif
 		}
