Increase scratch memory (160->190) for fixing Hit shader compilation with Substrate.

diff --git a/Engine/Shaders/Private/RayTracing/RayTracingLightingMS.usf b/Engine/Shaders/Private/RayTracing/RayTracingLightingMS.usf
--- a/Engine/Shaders/Private/RayTracing/RayTracingLightingMS.usf
+++ b/Engine/Shaders/Private/RayTracing/RayTracingLightingMS.usf
@@ -96,7 +96,6 @@ RAY_TRACING_ENTRY_MISS(RayTracingLightingMS, FPackedMaterialClosestHitPayload, P
 		GetShadowTerms(FakeSceneDepth, FakePrecomputedShadowFactors, FakeShadingModelID, FakeContactShadowOpacity, LightData, WorldPosition, L, LightAttenuation, Dither, ShadowTerms);
 
 		// RayTracing + Lighting + Strata is too expenssive (as in uses too many registers) to be supported on some platofrms
-		#if COMPILER_DXC
 		FStrataDeferredLighting Lighting = StrataDeferredLighting(
 			LightData,
 			V,
@@ -108,9 +107,6 @@ RAY_TRACING_ENTRY_MISS(RayTracingLightingMS, FPackedMaterialClosestHitPayload, P
 			StrataAddressing,
 			StrataPixelHeader);
 		LightContribution = Lighting.SceneColor.xyz;
-		#else
-		LightContribution = float4(1, 0, 1, 0.5); // Make it return purple so that it is noticeable.
-		#endif
 	}
 #else
 	// LWC_TODO - lighting wants translated world space
