Substrate - unified view memeber access and added comment.

diff --git a/Engine/Shaders/Private/Strata/Strata.ush b/Engine/Shaders/Private/Strata/Strata.ush
--- a/Engine/Shaders/Private/Strata/Strata.ush
+++ b/Engine/Shaders/Private/Strata/Strata.ush
@@ -2651,7 +2651,8 @@ float SanitizeRoughness(float Roughness)
 	return 1.0f;
 #else
 	#if (USE_EDITOR_SHADERS && !ES3_1_PROFILE) || MOBILE_EMULATION
-	// this feature is only needed for development/editor - we can compile it out for a shipping build (see r.CompileShadersForDevelopment cvar help)
+	// This feature is only needed for development/editor - we can compile it out for a shipping build (see r.CompileShadersForDevelopment cvar help)
+	// If more view members are accessed, make sure they are patched into the SharedLumenView from FDeferredShadingSceneRenderer::UpdateLumenScene.
 	Roughness = Roughness * ResolvedView.RoughnessOverrideParameter.y + ResolvedView.RoughnessOverrideParameter.x;
 	#endif
 	return saturate(max(View.MinRoughness, Roughness));
@@ -2662,7 +2663,8 @@ float3 SanitizeDiffuseAlbedo(float3 DiffuseAlbedo)
 {
 	#if USE_DEVELOPMENT_SHADERS
 	// This feature is only needed for development/editor - we can compile it out for a shipping build (see r.CompileShadersForDevelopment cvar help)
-	return saturate(DiffuseAlbedo) * View.DiffuseOverrideParameter.www + View.DiffuseOverrideParameter.xyz;
+	// If more view members are accessed, make sure they are patched into the SharedLumenView from FDeferredShadingSceneRenderer::UpdateLumenScene.
+	return saturate(DiffuseAlbedo) * ResolvedView.DiffuseOverrideParameter.www + ResolvedView.DiffuseOverrideParameter.xyz;
 	#else
 	return saturate(DiffuseAlbedo);
 	#endif
@@ -2672,7 +2674,8 @@ float3 SanitizeF0(float3 F0)
 {
 	#if USE_DEVELOPMENT_SHADERS
 	// This feature is only needed for development/editor - we can compile it out for a shipping build (see r.CompileShadersForDevelopment cvar help)
-	return saturate(F0) * View.SpecularOverrideParameter.w + View.SpecularOverrideParameter.xyz;
+	// If more view members are accessed, make sure they are patched into the SharedLumenView from FDeferredShadingSceneRenderer::UpdateLumenScene.
+	return saturate(F0) * ResolvedView.SpecularOverrideParameter.w + ResolvedView.SpecularOverrideParameter.xyz;
 	#else
 	return saturate(F0);
 	#endif
@@ -2682,7 +2685,8 @@ float SanitizeF0(float F0)
 {
 	#if USE_DEVELOPMENT_SHADERS
 	// This feature is only needed for development/editor - we can compile it out for a shipping build (see r.CompileShadersForDevelopment cvar help)
-	return saturate(F0) * View.SpecularOverrideParameter.w + View.SpecularOverrideParameter.x;
+	// If more view members are accessed, make sure they are patched into the SharedLumenView from FDeferredShadingSceneRenderer::UpdateLumenScene.
+	return saturate(F0) * ResolvedView.SpecularOverrideParameter.w + ResolvedView.SpecularOverrideParameter.x;
 	#else
 	return saturate(F0);
 	#endif
