Fixes FFT downsample pass on platforms that don't support global samplers

diff --git a/Engine/Shaders/Private/Bloom/BloomDownsampleKernel.usf b/Engine/Shaders/Private/Bloom/BloomDownsampleKernel.usf
--- a/Engine/Shaders/Private/Bloom/BloomDownsampleKernel.usf
+++ b/Engine/Shaders/Private/Bloom/BloomDownsampleKernel.usf
@@ -12,6 +12,7 @@
 
 uint2 KernelSpatialTextureSize;
 Texture2D KernelSpatialTexture;
+SamplerState KernelSpatialSampler;
 
 RWTexture2D<float4> DownsampleKernelSpatialOutput;
 
@@ -30,10 +31,10 @@ void MainCS(
 	float2 Offset = 0.7 / float2(KernelSpatialTextureSize);
 
 	float4 KernelColor = 0.0;
-	KernelColor += 0.25 * KernelSpatialTexture.SampleLevel(GlobalBilinearClampedSampler, TextureUV + float2(+1, +1) * Offset, 0);
-	KernelColor += 0.25 * KernelSpatialTexture.SampleLevel(GlobalBilinearClampedSampler, TextureUV + float2(+1, -1) * Offset, 0);
-	KernelColor += 0.25 * KernelSpatialTexture.SampleLevel(GlobalBilinearClampedSampler, TextureUV + float2(-1, -1) * Offset, 0);
-	KernelColor += 0.25 * KernelSpatialTexture.SampleLevel(GlobalBilinearClampedSampler, TextureUV + float2(-1, +1) * Offset, 0);
+	KernelColor += 0.25 * KernelSpatialTexture.SampleLevel(KernelSpatialSampler, TextureUV + float2(+1, +1) * Offset, 0);
+	KernelColor += 0.25 * KernelSpatialTexture.SampleLevel(KernelSpatialSampler, TextureUV + float2(+1, -1) * Offset, 0);
+	KernelColor += 0.25 * KernelSpatialTexture.SampleLevel(KernelSpatialSampler, TextureUV + float2(-1, -1) * Offset, 0);
+	KernelColor += 0.25 * KernelSpatialTexture.SampleLevel(KernelSpatialSampler, TextureUV + float2(-1, +1) * Offset, 0);
 
 	DownsampleKernelSpatialOutput[OutputPixelCoord] = KernelColor;
 }
diff --git a/Engine/Source/Runtime/Renderer/Private/PostProcess/PostProcessFFTBloom.cpp b/Engine/Source/Runtime/Renderer/Private/PostProcess/PostProcessFFTBloom.cpp
--- a/Engine/Source/Runtime/Renderer/Private/PostProcess/PostProcessFFTBloom.cpp
+++ b/Engine/Source/Runtime/Renderer/Private/PostProcess/PostProcessFFTBloom.cpp
@@ -168,6 +168,7 @@ public:
 	BEGIN_SHADER_PARAMETER_STRUCT(FParameters, )
 		SHADER_PARAMETER(FIntPoint, KernelSpatialTextureSize)
 		SHADER_PARAMETER_RDG_TEXTURE(Texture2D, KernelSpatialTexture)
+		SHADER_PARAMETER_SAMPLER(SamplerState, KernelSpatialSampler)
 		SHADER_PARAMETER_RDG_TEXTURE_UAV(RWTexture2D<float4>, DownsampleKernelSpatialOutput)
 	END_SHADER_PARAMETER_STRUCT()
 };
@@ -620,6 +621,7 @@ void InitDomainAndGetKernel(
 				FBloomDownsampleKernelCS::FParameters* PassParameters = GraphBuilder.AllocParameters<FBloomDownsampleKernelCS::FParameters>();
 				PassParameters->KernelSpatialTextureSize = ClampedKernelTexture->Desc.Extent;
 				PassParameters->KernelSpatialTexture = ClampedKernelTexture;
+				PassParameters->KernelSpatialSampler = TStaticSamplerState<SF_Bilinear>::GetRHI();
 				PassParameters->DownsampleKernelSpatialOutput = GraphBuilder.CreateUAV(NewClampedKernelTexture);
 
 				TShaderMapRef<FBloomDownsampleKernelCS> ComputeShader(View.ShaderMap);
