Fixed Substrate error in GetDiffuseIndirectSampleOcclusion (BenNormal -> BentNormal)

diff --git a/Engine/Shaders/Private/BasePassPixelShader.usf b/Engine/Shaders/Private/BasePassPixelShader.usf
--- a/Engine/Shaders/Private/BasePassPixelShader.usf
+++ b/Engine/Shaders/Private/BasePassPixelShader.usf
@@ -672,7 +672,7 @@ uint GetDiffuseIndirectSampleOcclusion(FSharedLocalBases SharedLocalBases, float
 		for (uint i = 0; i < SharedLocalBases.Count; ++i)
 		{
 			const float3 WorldNormal = SharedLocalBases.Normals[i];
-			const float3 BenNormal = WorldNormal; // STRATA_TODO: bent normal support - GetBentNormal(MaterialParameters)
+			const float3 BentNormal = WorldNormal; // STRATA_TODO: bent normal support - GetBentNormal(MaterialParameters)
 			const FSphericalGaussian HemisphereSG = Hemisphere_ToSphericalGaussian(WorldNormal);
 			const FSphericalGaussian VisibleSG = BentNormalAO_ToSphericalGaussian(BentNormal, MaterialAO);
 			const float VisibilityThreshold = InterleavedGradientNoise(SvPosition, View.StateFrameIndexMod8);
