Path Tracer: Implement basic blend mode support

This is incomplete and only covers the basic transparent cases for now.

diff --git a/Engine/Shaders/Private/PathTracing/PathTracingMaterialHitShader.usf b/Engine/Shaders/Private/PathTracing/PathTracingMaterialHitShader.usf
--- a/Engine/Shaders/Private/PathTracing/PathTracingMaterialHitShader.usf
+++ b/Engine/Shaders/Private/PathTracing/PathTracingMaterialHitShader.usf
@@ -86,7 +86,7 @@ RAY_TRACING_ENTRY_CLOSEST_HIT(PathTracingMaterialCHS,
 	FRayTracingIntersectionAttributes, Attributes)
 {
 	PackedPayload.HitT = RayTCurrent();
-#if MATERIALBLENDING_SOLID || MATERIALBLENDING_ALPHAHOLDOUT || (MATERIALBLENDING_MASKED && !MATERIAL_DITHER_OPACITY_MASK)
+#if MATERIALBLENDING_SOLID || MATERIALBLENDING_ALPHAHOLDOUT || (MATERIALBLENDING_MASKED && !MATERIAL_DITHER_OPACITY_MASK) || STRATA_OPAQUE_MATERIAL
 	if (PackedPayload.IsVisibilityRay())
 	{
 		// Shadow ray case -- all these blend modes get treated as fully opaque, so exit early.
@@ -138,11 +138,36 @@ RAY_TRACING_ENTRY_CLOSEST_HIT(PathTracingMaterialCHS,
 		CalcMaterialParametersEx(MaterialParameters, PixelMaterialInputs, SvPosition, ScreenPosition, bIsFrontFace, TranslatedWorldPosition, TranslatedWorldPosition);
 	}
 
+#if STRATA_ENABLED
+	FStrataData StrataData = PixelMaterialInputs.GetFrontStrataData();
+
+#if STRATA_OPTIMIZED_UNLIT
+	const float3 Transparency = UNLIT_TRANSMITTANCE(StrataData.InlinedBSDF);
+#else
+	FStrataPixelHeader StrataPixelHeader = MaterialParameters.GetFrontStrataHeader();
+	FStrataIntegrationSettings IntegrationSettings = InitStrataIntegrationSettings();
+	FStrataAddressing NullStrataAddressing = (FStrataAddressing)0;	// Fake unused in StrataCreateBSDFContext when using Forward inline shading
+	const uint BSDFCount = StrataPixelHeader.StrataTree.BSDFCount;
+	float TotalCoverage = 1;
+	float3 TotalTransmittancePreCoverage = 0;
+	const float3 V_World = MaterialParameters.CameraVector;
+	StrataPixelHeader.StrataUpdateTree(NullStrataAddressing, StrataData, V_World, IntegrationSettings, TotalCoverage, TotalTransmittancePreCoverage);
+	const float Transparency = 1 - TotalCoverage;
+#endif // STRATA_OPTIMIZED_UNLIT
+
+#endif // STRATA_ENABLED
+
 	// the following blend modes need to process shadows after having executed the material
 #if MATERIALBLENDING_MODULATE || (MATERIALBLENDING_MASKED && MATERIAL_DITHER_OPACITY_MASK) || MATERIALBLENDING_ALPHACOMPOSITE || MATERIALBLENDING_TRANSLUCENT
 	if (PackedPayload.IsVisibilityRay())
 	{
-		// STRATA_TODO: Handle transparent shadows here ...
+
+#if STRATA_ENABLED
+		// Transparency was already defined above
+		// STRATA_TODO: figure out refractive / thin translucent equivalents
+
+#else // STRATA_ENABLED
+
 #if MATERIALBLENDING_MODULATE
 		const float3 Transparency = GetMaterialEmissive(PixelMaterialInputs);
 #elif MATERIALBLENDING_ALPHACOMPOSITE
@@ -258,6 +283,8 @@ RAY_TRACING_ENTRY_CLOSEST_HIT(PathTracingMaterialCHS,
 #error Unhandled blending mode!
 #endif
 
+#endif // STRATA_ENABLED
+
 		// Update the ray throughput (it is packed simply into the payload since we don't need to carry any other information across hits)
 		float3 RayThroughput = PackedPayload.GetRayThroughput();
 		RayThroughput *= Transparency;
@@ -268,6 +295,7 @@ RAY_TRACING_ENTRY_CLOSEST_HIT(PathTracingMaterialCHS,
 
 	FPathTracingPayload Payload = (FPathTracingPayload)0;
 
+#if !STRATA_ENABLED
 	/**
 	 * Set common material attributes for both full and simplified materials
 	 **/
@@ -311,6 +339,8 @@ RAY_TRACING_ENTRY_CLOSEST_HIT(PathTracingMaterialCHS,
 #error Unknown material blending mode
 #endif
 
+#endif //  !STRATA_ENABLED
+
 	// fetch primitive flags only once
 	// TODO: would be nice to keep this inside MaterialParameters as it is also needed there as well
 	const uint PrimitiveFlags = GetPrimitiveData(MaterialParameters.PrimitiveId).Flags;
@@ -350,31 +380,24 @@ RAY_TRACING_ENTRY_CLOSEST_HIT(PathTracingMaterialCHS,
 	Payload.ShadingModelID = SHADINGMODELID_UNLIT;
 
 #if STRATA_OPTIMIZED_UNLIT
-	FStrataBSDF UnlitBSDF = PixelMaterialInputs.GetFrontStrataData().InlinedBSDF;
-	Payload.Radiance = BSDF_GETEMISSIVE(UnlitBSDF);
+	Payload.Radiance = BSDF_GETEMISSIVE(StrataData.InlinedBSDF);
 #if STRATA_OPAQUE_MATERIAL
 	Payload.BSDFOpacity = 1;
 	Payload.TransparencyColor = 0;
 #else
-	Payload.BSDFOpacity = UnlitBSDF.Coverage;
-	Payload.TransparencyColor = 1 - Payload.BSDFOpacity; // STRATA_TODO: is this right?
+	Payload.BSDFOpacity = StrataData.InlinedBSDF.Coverage;
+	Payload.TransparencyColor = UNLIT_TRANSMITTANCE(StrataData.InlinedBSDF);
 #endif
 
 #else // STRATA_OPTIMIZED_UNLIT
 
-	// STRATA_TODO: Handle transparency blend modes ...
 	Payload.BSDFOpacity = 1;
-	Payload.TransparencyColor = 0;
 
-	const float3 V_World = MaterialParameters.CameraVector;
-	FStrataPixelHeader StrataPixelHeader = MaterialParameters.GetFrontStrataHeader();
-	FStrataIntegrationSettings IntegrationSettings = InitStrataIntegrationSettings();
-	FStrataAddressing NullStrataAddressing = (FStrataAddressing)0;	// Fake unused in StrataCreateBSDFContext when using Forward inline shading
-	const uint BSDFCount = StrataPixelHeader.StrataTree.BSDFCount;
 	if (BSDFCount > 0)
 	{
-		StrataPixelHeader.StrataUpdateTree(NullStrataAddressing, V_World, IntegrationSettings);
-
+#if !STRATA_OPAQUE_MATERIAL
+		Payload.TransparencyColor = Transparency;
+#endif
 	
 		{
 #if HAIR_STRAND_MESH_FACTORY
