VulkanRHI:  Fix validation warning about 3 component write into 4 component texture.  Valid formats for the HistoryGuide are PF_FloatRGBA and PF_A2B10G10R10, so always write 4 components.

#rnx

diff --git a/Engine/Shaders/Private/TemporalSuperResolution/TSRDecimateHistory.usf b/Engine/Shaders/Private/TemporalSuperResolution/TSRDecimateHistory.usf
--- a/Engine/Shaders/Private/TemporalSuperResolution/TSRDecimateHistory.usf
+++ b/Engine/Shaders/Private/TemporalSuperResolution/TSRDecimateHistory.usf
@@ -65,7 +65,7 @@ FScreenTransform ScreenPosToPrevHistoryGuideBufferUV;
 float3 HistoryGuideQuantizationError;
 
 
-RWTexture2D<tsr_halfC> ReprojectedHistoryGuideOutput;
+RWTexture2D<tsr_half4> ReprojectedHistoryGuideOutput;
 RWTexture2D<tsr_half4> ReprojectedHistoryMoireOutput;
 
 RWTexture2D<float2> HoleFilledVelocityOutput;
@@ -657,7 +657,12 @@ void MainCS(
 		
 		tsr_short2 OutputPixelPos = InvalidateOutputPixelPos(InputPixelPos, InputInfo_ViewportMax);
 
-		ReprojectedHistoryGuideOutput[OutputPixelPos] = ReprojectedHistoryGuideColor;
+		#if CONFIG_SCENE_COLOR_ALPHA
+			ReprojectedHistoryGuideOutput[OutputPixelPos] = ReprojectedHistoryGuideColor;
+		#else
+			ReprojectedHistoryGuideOutput[OutputPixelPos] = tsr_half4(ReprojectedHistoryGuideColor.xyz, tsr_half(1.0));
+		#endif
+
 		#if CONFIG_MOIRE_REPROJECTION
 			ReprojectedHistoryMoireOutput[OutputPixelPos] = ReprojectedHistoryMoireColor;
 		#endif
diff --git a/Engine/Shaders/Private/TemporalSuperResolution/TSRRejectShading.usf b/Engine/Shaders/Private/TemporalSuperResolution/TSRRejectShading.usf
--- a/Engine/Shaders/Private/TemporalSuperResolution/TSRRejectShading.usf
+++ b/Engine/Shaders/Private/TemporalSuperResolution/TSRRejectShading.usf
@@ -113,7 +113,7 @@ Texture2D<tsr_half4> ReprojectedHistoryMoireTexture;
 Texture2D<tsr_half> ParallaxRejectionMaskTexture;
 Texture2D<tsr_ushort> IsMovingMaskTexture;
 
-RWTexture2D<tsr_halfC> HistoryGuideOutput;
+RWTexture2D<tsr_half4> HistoryGuideOutput;
 RWTexture2D<tsr_half4> HistoryMoireOutput;
 RWTexture2D<tsr_half2> HistoryRejectionOutput;
 RWTexture2D<tsr_half> InputSceneColorLdrLumaOutput;
@@ -573,7 +573,12 @@ void MainCS(
 			}
 			#endif
 
-			HistoryGuideOutput[OutputPixelPos] = FinalGuideColor;
+			#if CONFIG_SCENE_COLOR_ALPHA
+				HistoryGuideOutput[OutputPixelPos] = FinalGuideColor;
+			#else
+				HistoryGuideOutput[OutputPixelPos] = tsr_half4(FinalGuideColor.xyz, tsr_half(1.0));
+			#endif
+
 			#if CONFIG_MOIRE_DETECTION
 				HistoryMoireOutput[OutputPixelPos] = FinalMoire.GetElement(SimdIndex);
 			#endif
