Revert rect light fix for avoiding shader invaliation

diff --git a/Engine/Shaders/Private/RectLightIntegrate.ush b/Engine/Shaders/Private/RectLightIntegrate.ush
--- a/Engine/Shaders/Private/RectLightIntegrate.ush
+++ b/Engine/Shaders/Private/RectLightIntegrate.ush
@@ -259,7 +259,9 @@ FDirectLighting IntegrateBxDF(FGBufferData GBuffer, half3 N, half3 V, FRect Rect
 	if (IsRectVisible(Rect))
 	{
 		FAreaLightIntegrateContext Context = CreateRectIntegrateContext(GBuffer.Roughness, N, V, Rect, SourceTexture);
+	#if CHEAP_RECT
 		GBuffer.Roughness = max(GBuffer.Roughness, 0.02);
+	#endif
 		Out = IntegrateBxDF(GBuffer, N, V, Context.L, Context.Falloff, Context.NoL, Context.AreaLight, Shadow);
 	}
 	return Out;
