Path Tracer: Add support for legacy clearcoat shading model with Substrate

Re-use the shadingmodel id for clearcoat to signal we want the alternate interpretation of specular blending. Note that the resulting shading includes fuzz and subsurface which could arise from legacy clearcoat getting combined with newer features through layering or parameter blending.

Cleanup how SSS probabilities are calculated to re-use existing logic.

diff --git a/Engine/Shaders/Private/PathTracing/Material/PathTracingClearCoat.ush b/Engine/Shaders/Private/PathTracing/Material/PathTracingClearCoat.ush
--- a/Engine/Shaders/Private/PathTracing/Material/PathTracingClearCoat.ush
+++ b/Engine/Shaders/Private/PathTracing/Material/PathTracingClearCoat.ush
@@ -11,20 +11,11 @@
 
 #define GREY_INNER_LAYER 0
 
-#define CLEAR_COAT_F0	0.04 // From hard-coded value in ShadingModels.ush
-
 #include "PathTracingMaterialCommon.ush"
 #include "PathTracingGlossy.ush"
 #include "PathTracingFresnel.ush"
 #include "PathTracingEnergyConservation.ush"
 
-
-float ClearCoatFresnel(float VoH) {
-	// from hardcoded logic in ClearCoatBxDF
-	return CLEAR_COAT_F0 + (1.0 - CLEAR_COAT_F0) * Pow5(1 - VoH);
-}
-
-
 struct FClearCoatData
 {
 	float3x3 BaseBasis;
diff --git a/Engine/Shaders/Private/PathTracing/Material/PathTracingFresnel.ush b/Engine/Shaders/Private/PathTracing/Material/PathTracingFresnel.ush
--- a/Engine/Shaders/Private/PathTracing/Material/PathTracingFresnel.ush
+++ b/Engine/Shaders/Private/PathTracing/Material/PathTracingFresnel.ush
@@ -117,3 +117,12 @@ bool SampleRefraction(float3 RayDirection, float3 Normal, float Eta, float F0, f
 	OutDirection = RayDirection + (2 * CosI) * Normal;
 	return false;
 }
+
+
+#define CLEAR_COAT_F0	0.04 // From hard-coded value in ShadingModels.ush
+
+float ClearCoatFresnel(float VoH)
+{
+	// from hardcoded logic in ClearCoatBxDF
+	return CLEAR_COAT_F0 + (1.0 - CLEAR_COAT_F0) * Pow5(1 - VoH);
+}
diff --git a/Engine/Shaders/Private/PathTracing/Material/PathTracingMaterialSampling.ush b/Engine/Shaders/Private/PathTracing/Material/PathTracingMaterialSampling.ush
--- a/Engine/Shaders/Private/PathTracing/Material/PathTracingMaterialSampling.ush
+++ b/Engine/Shaders/Private/PathTracing/Material/PathTracingMaterialSampling.ush
@@ -52,6 +52,12 @@ float3 EstimateMaterialAlbedo(FPathTracingPayload Payload)
 			Albedo *= Payload.GetMaxLobeWeight();
 			break;
 		}
+		case SHADINGMODELID_CLEAR_COAT:
+		{
+			Albedo = Payload.DiffuseColor + Payload.SubsurfaceColor + Payload.SpecularColor + Payload.RoughnessData.z * CLEAR_COAT_F0 + Payload.FuzzAmount * Payload.FuzzColor;
+			Albedo *= Payload.GetMaxLobeWeight();
+			break;
+		}
 		case SHADINGMODELID_EYE:
 			Albedo = Payload.DiffuseColor + Payload.SubsurfaceColor + Payload.SpecularColor;
 			break;
@@ -66,10 +72,11 @@ float3 EstimateMaterialAlbedo(FPathTracingPayload Payload)
 			Albedo = Payload.DiffuseColor + Payload.SpecularColor + Payload.SubsurfaceColor;
 			break;
 		case SHADINGMODELID_DEFAULT_LIT:
-		case SHADINGMODELID_CLEAR_COAT:
-			// TODO: include clearcoat albedo as well
 			Albedo = Payload.DiffuseColor + Payload.SpecularColor;
 			break;
+		case SHADINGMODELID_CLEAR_COAT:
+			Albedo = Payload.DiffuseColor + Payload.SpecularColor + Payload.GetClearCoat() * CLEAR_COAT_F0;
+			break;
 		case SHADINGMODELID_THIN_TRANSLUCENT:
 		{
 			// glass portion is a mixture of pure specular and tinted transmission (when refracting)
@@ -137,6 +144,17 @@ void AccumulateAlbedoNormal(inout FPathTracingPayload Payload, float3 PathThroug
 			Normal += NormThroughput * Payload.WorldNormal;
 			break;
 		}
+		case SHADINGMODELID_CLEAR_COAT:
+		{
+			const float Roughness0 = Payload.RoughnessData.x;
+			const float Roughness1 = Payload.RoughnessData.y;
+			const float SpecBlend = Payload.RoughnessData.z;
+			const float Roughness0Weight = ROUGHNESS_WEIGHT(Roughness0);
+			const float Roughness1Weight = SpecBlend * ROUGHNESS_WEIGHT(Roughness1);
+			Albedo += PathThroughput * (Payload.DiffuseColor + Roughness0Weight * Payload.SpecularColor + Roughness1Weight * CLEAR_COAT_F0 + Payload.FuzzAmount * Payload.FuzzColor);
+			Normal += NormThroughput * Payload.WorldNormal;
+			break;
+		}
 		case SHADINGMODELID_EYE:
 		{
 			const float RoughnessWeight = ROUGHNESS_WEIGHT(Payload.GetEyeRoughness());
@@ -242,6 +260,19 @@ float GetAverageRoughness(FPathTracingPayload Payload)
 			AvgRoughness = lerp(1.0, AvgRoughness, DiffuseFraction);
 			break;
 		}
+		case SHADINGMODELID_CLEAR_COAT:
+		{
+			AvgRoughness = max(AvgRoughness, Payload.RoughnessData.y);
+			float DiffuseFraction = LobeSelectionProb(Payload.SubsurfaceColor + Payload.DiffuseColor + Payload.FuzzAmount * Payload.FuzzColor, Payload.SpecularColor + Payload.RoughnessData.z * CLEAR_COAT_F0);
+			AvgRoughness = lerp(1.0, AvgRoughness, DiffuseFraction);
+			break;
+		}
+		case SHADINGMODELID_EYE:
+		{
+			float DiffuseFraction = LobeSelectionProb(Payload.SubsurfaceColor + Payload.DiffuseColor, Payload.SpecularColor);
+			AvgRoughness = lerp(1.0, AvgRoughness, DiffuseFraction);
+			break;
+		}
 		// STRATA_TODO: do we need special cases for the other slab types?
 	}
 	return AvgRoughness;
@@ -272,7 +303,7 @@ FPathTracingPayload UnpackPathTracingPayloadMaterial(FPackedPathTracingPayload P
 	FPathTracingPayload HitPayload = UnpackPathTracingPayload(PackedInput, Ray);
 	if (bApplyOverrides)
 	{
-		HitPayload.DiffuseColor = HitPayload.DiffuseColor * View.DiffuseOverrideParameter.w + View.DiffuseOverrideParameter.xyz;
+		HitPayload.DiffuseColor  = HitPayload.DiffuseColor  * View.DiffuseOverrideParameter.w + View.DiffuseOverrideParameter.xyz;
 		HitPayload.SpecularColor = HitPayload.SpecularColor * View.SpecularOverrideParameter.w + View.SpecularOverrideParameter.xyz;
 	}
 #if PATHTRACING_STRATA_PAYLOAD
@@ -285,6 +316,14 @@ FPathTracingPayload UnpackPathTracingPayloadMaterial(FPackedPathTracingPayload P
 			HitPayload.DiffuseColor = 0;
 			break;
 		}
+		case SHADINGMODELID_CLEAR_COAT:
+		{
+			if (bApplyOverrides)
+			{
+				HitPayload.RoughnessData.z = HitPayload.RoughnessData.z * View.SpecularOverrideParameter.w + View.SpecularOverrideParameter.x;
+			}
+			break;
+		}
 		case SHADINGMODELID_EYE:
 		{
 			// eye still needs a diffuse component for the iris center
@@ -394,45 +433,35 @@ FSSSRandomWalkInfo GetMaterialSSSInfo(FPathTracingPayload Payload, float3 V_Worl
 {
 	FSSSRandomWalkInfo Result = (FSSSRandomWalkInfo)0;
 #if PATHTRACING_STRATA_PAYLOAD
-	if (Payload.ShadingModelID == SHADINGMODELID_STRATA)
+	if (any(Payload.MeanFreePath > 0.0))
 	{
-		if (any(Payload.MeanFreePath > 0.0))
+		Result.Color = Payload.SubsurfaceColor;
+		Result.Radius = Payload.MeanFreePath;
+		Result.G = Payload.PhaseG;
+		if (Payload.ShadingModelID == SHADINGMODELID_STRATA)
 		{
-			Result.Color = Payload.SubsurfaceColor;
-			Result.Radius = Payload.MeanFreePath;
-
 			FSubstrateSingleSlabData Data = (FSubstrateSingleSlabData)0;
-			Data.PrepareBSDF(V_World, Payload);
-			const float3 SSSAlbedo = Payload.SubsurfaceColor * Data.DiffuseWeight;
-			const float3 DiffuseAlbedo = Payload.DiffuseColor * Data.DiffuseWeight;
-			const float3 Spec0Albedo = Data.FuzzAttenuation * (1.0 - Payload.RoughnessData.z) * Data.Spec0.E;
-			const float3 Spec1Albedo = Data.FuzzAttenuation * (Payload.RoughnessData.z) * Data.Spec1.E;
-			const float  FuzzAlbedo = Payload.FuzzAmount * Data.Fuzz.E;
-			const float3 MaxLobeWeight = Payload.GetMaxLobeWeight();
-
+			Data.PrepareBSDF(V_World, Payload, true);
 			Result.Weight = Payload.WeightV * Payload.BSDFOpacity * Data.DiffuseWeight;
-			Result.Prob = LobeSelectionProb(
-				MaxLobeWeight * (Payload.SubsurfaceColor * Data.DiffuseWeight),
-				MaxLobeWeight * (Payload.DiffuseColor * Data.DiffuseWeight + Spec0Albedo + Spec1Albedo + FuzzAlbedo));
-			// STRATA_TODO: evaluate if we want to set Result.G in all cases?
+			Result.Prob = Data.LobePdf.x;
 		}
-	}
-	if (Payload.ShadingModelID == SHADINGMODELID_EYE)
-	{
-		if (any(Payload.MeanFreePath > 0.0))
+		else if (Payload.ShadingModelID == SHADINGMODELID_CLEAR_COAT)
+		{
+			FSubstrateClearCoatData Data = (FSubstrateClearCoatData)0;
+			Data.PrepareBSDF(V_World, Payload, true);
+			Result.Weight = Payload.WeightV * Payload.BSDFOpacity * Data.DiffuseWeight;
+			Result.Prob = Data.LobePdf.x;
+		}
+		else if (Payload.ShadingModelID == SHADINGMODELID_EYE)
 		{
-			Result.Color = Payload.SubsurfaceColor;
-			Result.Radius = Payload.MeanFreePath;
 			const float NoV = saturate(dot(Payload.WorldNormal, V_World));
 			const float3 SpecE = ComputeGGXSpecEnergyTermsRGB(Payload.GetEyeRoughness(), NoV, Payload.SpecularColor).E;
-
 			const float DiffWeight = (1 - Luminance(SpecE));
-			const float3 SSSAlbedo  = DiffWeight * Payload.SubsurfaceColor;
+			const float3 SSSAlbedo = DiffWeight * Payload.SubsurfaceColor;
 			const float3 DiffAlbedo = DiffWeight * Payload.DiffuseColor;
 			const float3 SpecAlbedo = SpecE;
 			Result.Prob = LobeSelectionProb(SSSAlbedo, DiffAlbedo + SpecAlbedo);
 			Result.Weight = Payload.BSDFOpacity * DiffWeight;
-
 		}
 	}
 #else
@@ -497,6 +526,8 @@ FMaterialSample SampleMaterial(
 #if PATHTRACING_STRATA_PAYLOAD
 		case SHADINGMODELID_STRATA:
 			return Substrate_SampleMaterial(V_World, Payload, RandSample);
+		case SHADINGMODELID_CLEAR_COAT:
+			return SubstrateClearCoat_SampleMaterial(V_World, Payload, RandSample);
 #else
 		case SHADINGMODELID_DEFAULT_LIT:
 		{
@@ -551,6 +582,8 @@ FMaterialEval EvalMaterial(
 #if PATHTRACING_STRATA_PAYLOAD
 		case SHADINGMODELID_STRATA:
 			return Substrate_EvalMaterial(V_World, L_World, Payload, DiffuseSpecularScale);
+		case SHADINGMODELID_CLEAR_COAT:
+			return SubstrateClearCoat_EvalMaterial(V_World, L_World, Payload, DiffuseSpecularScale);
 #else
 		case SHADINGMODELID_DEFAULT_LIT:
 		{
diff --git a/Engine/Shaders/Private/PathTracing/Material/PathTracingSubstrate.ush b/Engine/Shaders/Private/PathTracing/Material/PathTracingSubstrate.ush
--- a/Engine/Shaders/Private/PathTracing/Material/PathTracingSubstrate.ush
+++ b/Engine/Shaders/Private/PathTracing/Material/PathTracingSubstrate.ush
@@ -8,6 +8,7 @@
 
 #include "PathTracingMaterialCommon.ush"
 #include "PathTracingGlossy.ush"
+#include "PathTracingFresnel.ush"
 #include "../../BRDF.ush"
 
 #define STRATA_FRESNEL_F82  0
@@ -40,11 +41,8 @@ struct FSubstrateSingleSlabData
 	float3 LobeCdf;
 	float4 LobePdf;
 
-
-	// NOTE: Must set Basis field before calling this
-	void PrepareBSDF(float3 V_World, FPathTracingPayload Payload)
+	void PrepareBSDF(float3 V_World, FPathTracingPayload Payload, bool bUseSSSColor = false)
 	{
-
 		Basis = GetGGXBasis(Payload.RoughnessData.x, Payload.Anisotropy, Payload.WorldNormal, Payload.WorldTangent, Alpha0);
 		Alpha1 = GetGGXAlpha(Payload.RoughnessData.y, Payload.Anisotropy);
 
@@ -79,7 +77,7 @@ struct FSubstrateSingleSlabData
 		// Now prepare a cdf/pdf for lobe selection
 		float3 MaxLobeWeight = Payload.GetMaxLobeWeight();
 		LobeCdf = LobeSelectionCdf(
-			MaxLobeWeight * DiffuseWeight * Payload.DiffuseColor,
+			MaxLobeWeight * DiffuseWeight * (bUseSSSColor ? Payload.SubsurfaceColor : Payload.DiffuseColor),
 			MaxLobeWeight * Spec0Albedo,
 			MaxLobeWeight * Spec1Albedo,
 			MaxLobeWeight * FuzzAlbedo);
@@ -255,3 +253,253 @@ FMaterialEval Substrate_EvalMaterial(
 
 	return Result;
 }
+
+struct FSubstrateClearCoatData
+{
+	float3x3 Basis;
+	float3 V;
+
+	float  DiffuseWeight;
+	float3 F0;
+	float3 F90;
+	float2 BaseAlpha;
+	float2 CoatAlpha;
+	float FuzzRoughness;
+
+	float FuzzAttenuation;
+	float CoatAttenuation;
+	FBxDFEnergyTermsRGB BaseSpec;
+	FBxDFEnergyTermsA CoatSpec;
+	FBxDFEnergyTermsA Fuzz;
+
+	float3 LobeCdf;
+	float4 LobePdf;
+
+	void PrepareBSDF(float3 V_World, FPathTracingPayload Payload, bool bUseSSSColor = false)
+	{
+		Basis = GetGGXBasis(Payload.RoughnessData.x, Payload.Anisotropy, Payload.WorldNormal, Payload.WorldTangent, BaseAlpha);
+		CoatAlpha = GetGGXAlpha(Payload.RoughnessData.y, Payload.Anisotropy);
+
+		V = mul(Basis, V_World);
+
+		F0 = Payload.SpecularColor;
+		F90 = Payload.SpecularEdgeColor;
+		{
+			// make sure F0=0.0 fades off the specular lobe completely
+			const float SpecularMicroOcclusion = F0RGBToMicroOcclusion(F0);
+			F0 *= SpecularMicroOcclusion;
+			F90 *= SpecularMicroOcclusion;
+		}
+
+		FuzzRoughness = MakeRoughnessSafe(Payload.FuzzRoughness, STRATA_MIN_FUZZ_ROUGHNESS);
+
+		const float NoV = saturate(V.z);
+		// STRATA_TODO: Add support for F82 model?
+		BaseSpec = ComputeGGXSpecEnergyTermsRGB(Payload.RoughnessData.x, NoV, F0, F90);
+		CoatSpec = ComputeGGXSpecEnergyTermsA(Payload.RoughnessData.y, NoV, CLEAR_COAT_F0);
+		Fuzz = ComputeClothEnergyTermsA(FuzzRoughness, NoV);
+
+		FuzzAttenuation = 1.0 - Payload.FuzzAmount * Fuzz.E;
+		CoatAttenuation = 1.0 - Payload.RoughnessData.z * CoatSpec.E;
+		DiffuseWeight = FuzzAttenuation * CoatAttenuation * (1.0 - Luminance(BaseSpec.E));
+
+		const float3 BaseSpecAlbedo = FuzzAttenuation * CoatAttenuation * BaseSpec.E;
+		const float  CoatSpecAlbedo = FuzzAttenuation * Payload.RoughnessData.z * CoatSpec.E;
+		const float  FuzzAlbedo = Payload.FuzzAmount * Fuzz.E;
+
+		// Now prepare a cdf/pdf for lobe selection
+		float3 MaxLobeWeight = Payload.GetMaxLobeWeight();
+		LobeCdf = LobeSelectionCdf(
+			MaxLobeWeight * DiffuseWeight * (bUseSSSColor ? Payload.SubsurfaceColor : Payload.DiffuseColor),
+			MaxLobeWeight * BaseSpecAlbedo,
+			MaxLobeWeight * CoatSpecAlbedo,
+			MaxLobeWeight * FuzzAlbedo);
+		LobePdf = LobeSelectionPdf(LobeCdf);
+	}
+};
+
+
+
+FMaterialSample SubstrateClearCoat_SampleMaterial(
+	float3 V_World,
+	FPathTracingPayload Payload,
+	float3 RandSample
+)
+{
+	FSubstrateClearCoatData Data = (FSubstrateClearCoatData)0;
+	Data.PrepareBSDF(V_World, Payload);
+
+	float3 L = 0, H = 0, V = Data.V;
+	float OutRoughness = 1;
+
+	const bool bSampledDiffuse = RandSample.x < Data.LobeCdf.x;
+	const bool bSampledSpecular = RandSample.x < Data.LobeCdf.z;
+	if (bSampledDiffuse)
+	{
+		RandSample.x = RescaleRandomNumber(RandSample.x, 0.0, Data.LobeCdf.x);
+		// diffuse lobe
+		L = CosineSampleHemisphere(RandSample.xy).xyz;
+		H = normalize(L + V);
+	}
+	else if (bSampledSpecular)
+	{
+		// specular lobes
+		const bool bUseBaseSpec = RandSample.x < Data.LobeCdf.y;
+		if (bUseBaseSpec)
+		{
+			RandSample.x = RescaleRandomNumber(RandSample.x, Data.LobeCdf.x, Data.LobeCdf.y);
+			OutRoughness = Payload.RoughnessData.x;
+		}
+		else
+		{
+			RandSample.x = RescaleRandomNumber(RandSample.x, Data.LobeCdf.y, Data.LobeCdf.z);
+			OutRoughness = Payload.RoughnessData.y;
+		}
+
+		// #dxr_todo: Evaluate UniformSampleDiskConcentric
+		H = ImportanceSampleVisibleGGX_aniso(UniformSampleDisk(RandSample.xy), bUseBaseSpec ? Data.BaseAlpha : Data.CoatAlpha, V).xyz;
+		L = reflect(-V, H);
+		if (L.z <= 0)
+		{
+			// invalid output direction, exit early
+			return NullMaterialSample();
+		}
+	}
+	else
+	{
+		// cloth lobe
+		RandSample.x = RescaleRandomNumber(RandSample.x, Data.LobeCdf.z, 1.0);
+		L = UniformSampleHemisphere(RandSample.xy).xyz;
+		H = normalize(L + V);
+	}
+
+	// transform to world space
+	const float3 L_World = normalize(mul(L, Data.Basis));
+	const float3 N_World = Payload.WorldNormal;
+	const float NoV = saturate(V.z);
+	const float NoL = saturate(L.z);
+	const float VoH = saturate(dot(V, H));
+	const float NoH = saturate(H.z);
+
+	FMaterialSample Result = CreateMaterialSample(L_World, 0.0, 0.0, 1.0, OutRoughness, PATHTRACER_SCATTER_DIFFUSE);
+
+	const float Metallic = F0RGBToMetallic(Payload.SpecularColor);
+	const float3 BaseColor = lerp(Payload.DiffuseColor, Payload.SpecularColor, Metallic);
+	const float3 Transmission = lerp(1.0, SimpleClearCoatTransmittance(NoL, NoV, Metallic, BaseColor), Payload.RoughnessData.z);
+
+	const float DiffPdf = NoL / PI;
+	const float2 BaseGGXResult = GGXEvalReflection(L, V, H, Data.BaseAlpha);
+	const float2 CoatGGXResult = GGXEvalReflection(L, V, H, Data.CoatAlpha);
+	const float BaseSpecPdf = BaseGGXResult.y;
+	const float CoatSpecPdf = CoatGGXResult.y;
+	const float ClothPdf = 1 / (2 * PI);
+
+	const float ShadowTerminator = ShadowTerminatorTerm(L_World, N_World, Payload.WorldSmoothNormal);
+	if (bSampledDiffuse)
+	{
+		const float3 Diffuse = Transmission * Data.DiffuseWeight * GetPathTracingDiffuseModel(Payload.DiffuseColor, Payload.RoughnessData.x, NoV, NoL, VoH, NoH);
+
+		Result.AddLobeWithMIS(Diffuse * ShadowTerminator, NoL / PI, Data.LobePdf.x);
+
+		Result.Pdf += Data.LobePdf.y * BaseSpecPdf;
+		Result.Pdf += Data.LobePdf.z * CoatSpecPdf;
+		Result.Pdf += Data.LobePdf.w * ClothPdf;
+	}
+	else if (bSampledSpecular)
+	{
+#if STRATA_FRESNEL_F82
+		// STRATA_TODO: Need to tune the logic for specular micro-occlusion ...
+		const float3 F = Data.FuzzAttenuation * F_AdobeF82(Data.F0, Data.F90, VoH);
+#else
+		const float3 FBase = Data.FuzzAttenuation * F_Schlick(Data.F0, Data.F90, VoH);
+#endif
+		const float FCoat = Payload.RoughnessData.z * ClearCoatFresnel(VoH);
+
+		Result.AddLobeWithMIS(Data.CoatAttenuation * FBase * BaseGGXResult.x * Data.BaseSpec.W * Transmission, BaseSpecPdf, Data.LobePdf.y);
+		Result.AddLobeWithMIS(Data.FuzzAttenuation * FCoat * CoatGGXResult.x * Data.CoatSpec.W               , CoatSpecPdf, Data.LobePdf.z);
+		Result.Pdf += Data.LobePdf.x * DiffPdf;
+		Result.Pdf += Data.LobePdf.w * ClothPdf;
+
+		Result.ScatterType = PATHTRACER_SCATTER_SPECULAR;
+	}
+	else
+	{
+		// Cloth Lobe
+		const float DCloth = D_Charlie(Data.FuzzRoughness, NoH);
+		const float3 FCloth = F_Schlick(Payload.FuzzColor, VoH);
+		const float3 ClothWeight = NoL > 0 && NoV > 0 ? (NoL * Payload.FuzzAmount * DCloth * Vis_Ashikhmin(NoV, NoL)) * FCloth * Data.Fuzz.W * (2 * PI) : 0.0;
+		Result.AddLobeWithMIS(ClothWeight * ShadowTerminator, 1 / (2.0 * PI), Data.LobePdf.w);
+		Result.ScatterType = PATHTRACER_SCATTER_SPECULAR;
+
+		Result.Pdf += Data.LobePdf.x * DiffPdf;
+		Result.Pdf += Data.LobePdf.y * BaseSpecPdf;
+		Result.Pdf += Data.LobePdf.z * CoatSpecPdf;
+	}
+
+	Result.Weight *= SubstrateLobeWeight(Payload, NoL);
+
+	return Result;
+}
+
+
+FMaterialEval SubstrateClearCoat_EvalMaterial(
+	float3 V_World,
+	float3 L_World,
+	FPathTracingPayload Payload,
+	float2 DiffuseSpecularScale
+)
+{
+	FSubstrateClearCoatData Data = (FSubstrateClearCoatData)0;
+	Data.PrepareBSDF(V_World, Payload);
+
+	const float3 N_World = Payload.WorldNormal;
+
+	// move vectors into right shading frame
+	const float3 V = Data.V;
+	const float3 L = mul(Data.Basis, L_World);
+	const float3 H = normalize(V + L);
+	const float NoV = saturate(V.z);
+	const float NoL = saturate(L.z);
+	const float VoH = saturate(dot(V, H));
+	const float NoH = saturate(H.z);
+
+	FMaterialEval Result = NullMaterialEval();
+
+	const float Metallic = F0RGBToMetallic(Payload.SpecularColor);
+	const float3 BaseColor = lerp(Payload.DiffuseColor, Payload.SpecularColor, Metallic);
+	const float3 Transmission = lerp(1.0, SimpleClearCoatTransmittance(NoL, NoV, Metallic, BaseColor), Payload.RoughnessData.z);
+
+	const float ShadowTerminator = ShadowTerminatorTerm(L_World, N_World, Payload.WorldSmoothNormal);
+	{
+		const float3 Diffuse = Transmission * Data.DiffuseWeight * GetPathTracingDiffuseModel(Payload.DiffuseColor, Payload.RoughnessData.x, NoV, NoL, VoH, NoH);
+		Result.AddLobeWithMIS(DiffuseSpecularScale.x * Diffuse * ShadowTerminator, NoL / PI, Data.LobePdf.x);
+	}
+	{
+		const float2 BaseGGXResult = GGXEvalReflection(L, V, H, Data.BaseAlpha);
+		const float2 CoatGGXResult = GGXEvalReflection(L, V, H, Data.CoatAlpha);
+		const float BaseSpecPdf = BaseGGXResult.y;
+		const float CoatSpecPdf = CoatGGXResult.y;
+
+#if STRATA_FRESNEL_F82
+		// STRATA_TODO: Need to tune the logic for specular micro-occlusion ...
+		const float3 F = Data.FuzzAttenuation * F_AdobeF82(Data.F0, Data.F90, VoH);
+#else
+		const float3 FBase = Data.FuzzAttenuation * F_Schlick(Data.F0, Data.F90, VoH);
+#endif
+		const float FCoat = Payload.RoughnessData.z * ClearCoatFresnel(VoH);
+
+		Result.AddLobeWithMIS(Data.CoatAttenuation * FBase * BaseGGXResult.x * Data.BaseSpec.W * Transmission, BaseSpecPdf, Data.LobePdf.y);
+		Result.AddLobeWithMIS(Data.FuzzAttenuation * FCoat * CoatGGXResult.x * Data.CoatSpec.W               , CoatSpecPdf, Data.LobePdf.z);
+	}
+	{
+		// Cloth Lobe
+		const float DCloth = D_Charlie(Data.FuzzRoughness, NoH);
+		const float3 FCloth = F_Schlick(Payload.FuzzColor, VoH);
+		const float3 ClothWeight = NoL > 0 && NoV > 0 ? (NoL * Payload.FuzzAmount * DCloth * Vis_Ashikhmin(NoV, NoL)) * FCloth * Data.Fuzz.W * (2 * PI) : 0.0;
+		Result.AddLobeWithMIS(DiffuseSpecularScale.y * ClothWeight * ShadowTerminator, 1 / (2.0 * PI), Data.LobePdf.w);
+	}
+
+	Result.Weight *= SubstrateLobeWeight(Payload, NoL);
+
+	return Result;
+}
diff --git a/Engine/Shaders/Private/PathTracing/PathTracingMaterialHitShader.usf b/Engine/Shaders/Private/PathTracing/PathTracingMaterialHitShader.usf
--- a/Engine/Shaders/Private/PathTracing/PathTracingMaterialHitShader.usf
+++ b/Engine/Shaders/Private/PathTracing/PathTracingMaterialHitShader.usf
@@ -468,16 +468,25 @@ RAY_TRACING_ENTRY_CLOSEST_HIT(PathTracingMaterialCHS,
 							}
 
 							float Roughness0 = MakeRoughnessSafe(SLAB_ROUGHNESS(BSDF), PackedPayload.GetPathRoughness());
+							// STRATA_TODO: Support Fresnel82?
+							float3 Spec0E = ComputeGGXSpecEnergyTermsRGB(Roughness0, NoV, F0, F90).E;
+							float3 Spec1E = Spec0E;
 							float Roughness1 = Roughness0;
 							float SpecBlend = 0.0;
+							bool bUseClearCoat = false;
 							if (BSDF_GETHASHAZINESS(BSDF))
 							{
 								const FHaziness Haziness = UnpackHaziness(SLAB_HAZINESS(BSDF));
-								// STRATA_TODO: Ideally clearcoat in the path tracer would be handled via actual layering rather than this backwards compatibility trick
-								if (!Haziness.bSimpleClearCoat)
+								SpecBlend = Haziness.Weight;
+								Roughness1 = MakeRoughnessSafe(Haziness.Roughness, PackedPayload.GetPathRoughness());
+								if (Haziness.bSimpleClearCoat)
+								{
+									Spec1E = ComputeGGXSpecEnergyTermsRGB(Roughness1, NoV, CLEAR_COAT_F0).E;
+									bUseClearCoat = true;
+								}
+								else
 								{
-									SpecBlend = Haziness.Weight;
-									Roughness1 = MakeRoughnessSafe(Haziness.Roughness, PackedPayload.GetPathRoughness());
+									Spec1E = ComputeGGXSpecEnergyTermsRGB(Roughness1, NoV, F0, F90).E;
 								}
 							}
 							float FuzzAmount = 0.0;
@@ -490,17 +499,28 @@ RAY_TRACING_ENTRY_CLOSEST_HIT(PathTracingMaterialCHS,
 								FuzzRoughness = SLAB_FUZZ_ROUGHNESS(BSDF);
 								FuzzRoughness = MakeRoughnessSafe(FuzzRoughness, STRATA_MIN_FUZZ_ROUGHNESS);
 							}
-
-							// STRATA_TODO: Support Fresnel82?
-							const FBxDFEnergyTermsRGB Spec0 = ComputeGGXSpecEnergyTermsRGB(Roughness0, NoV, F0, F90);
-							const FBxDFEnergyTermsRGB Spec1 = ComputeGGXSpecEnergyTermsRGB(Roughness1, NoV, F0, F90);
-							const FBxDFEnergyTermsA   Fuzz  = ComputeClothEnergyTermsA(FuzzRoughness, NoV);
-							const float3 SpecE = lerp(Spec0.E, Spec1.E, SpecBlend);
-							const float FuzzAttenuation = 1.0 - FuzzAmount * Fuzz.E;
-							const float DiffuseWeight = FuzzAttenuation * (1.0 - Luminance(SpecE));
-							const float3 Spec0Albedo  = FuzzAttenuation * (1.0 - SpecBlend) * Spec0.E;
-							const float3 Spec1Albedo  = FuzzAttenuation * (      SpecBlend) * Spec1.E;
-							const float  FuzzAlbedo   = FuzzAmount * Fuzz.E;
+							
+							const float FuzzE = ComputeClothEnergyTermsA(FuzzRoughness, NoV).E;
+							const float FuzzAttenuation = 1.0 - FuzzAmount * FuzzE;
+							float DiffuseWeight = FuzzAttenuation;
+							float3 Spec0Albedo = FuzzAttenuation * Spec0E;
+							float3 Spec1Albedo = FuzzAttenuation * Spec1E;
+							if (bUseClearCoat)
+							{
+								// clearcoat slab
+								float CoatAttenuation = 1.0 - SpecBlend * Spec1E.x;
+								DiffuseWeight *= CoatAttenuation * (1.0 - Luminance(Spec0E));
+								Spec0Albedo *= CoatAttenuation;
+								Spec1Albedo *= SpecBlend;
+							}
+							else
+							{
+								// dual-specular slab
+								DiffuseWeight *= 1.0 - Luminance(lerp(Spec0E, Spec1E, SpecBlend));
+								Spec0Albedo *= 1.0 - SpecBlend;
+								Spec1Albedo *=       SpecBlend;
+							}
+							const float  FuzzAlbedo = FuzzAmount * FuzzE;
 							const float3 SlabAlbedo = MaxLobeWeight * (DiffuseWeight * DiffuseColor + Spec0Albedo + Spec1Albedo + FuzzAlbedo);
 							Pdf = LobeColorToWeight(SlabAlbedo);
 						}
@@ -556,6 +576,7 @@ RAY_TRACING_ENTRY_CLOSEST_HIT(PathTracingMaterialCHS,
 			{
 				case STRATA_BSDF_TYPE_SLAB:
 				{
+					Payload.ShadingModelID    = SHADINGMODELID_STRATA;
 					Payload.DiffuseColor      = SLAB_DIFFUSEALBEDO(BSDF);
 					Payload.SpecularColor     = SLAB_F0(BSDF);
 					Payload.SpecularEdgeColor = SLAB_F90(BSDF);
@@ -565,12 +586,12 @@ RAY_TRACING_ENTRY_CLOSEST_HIT(PathTracingMaterialCHS,
 					if (BSDF_GETHASHAZINESS(BSDF))
 					{
 						const FHaziness Haziness = UnpackHaziness(SLAB_HAZINESS(BSDF));
-						// STRATA_TODO: detect this and use the old SHADINGMODELID_CLEARCOAT?
-						if (!Haziness.bSimpleClearCoat)
+						if (Haziness.bSimpleClearCoat)
 						{
-							Payload.RoughnessData.y = MakeRoughnessSafe(Haziness.Roughness, PackedPayload.GetPathRoughness());
-							Payload.RoughnessData.z = Haziness.Weight;
+							Payload.ShadingModelID = SHADINGMODELID_CLEAR_COAT;
 						}
+						Payload.RoughnessData.y = MakeRoughnessSafe(Haziness.Roughness, PackedPayload.GetPathRoughness());
+						Payload.RoughnessData.z = Haziness.Weight;
 					}
 					Payload.Anisotropy = BSDF_GETHASANISOTROPY(BSDF) ? SLAB_ANISOTROPY(BSDF) : 0.0;
 					if (BSDF_GETHASFUZZ(BSDF))
@@ -588,7 +609,7 @@ RAY_TRACING_ENTRY_CLOSEST_HIT(PathTracingMaterialCHS,
 					if (BSDF_GETSSSTYPE(BSDF) == SSS_TYPE_DIFFUSION)
 					{
 						Payload.MeanFreePath = SLAB_SSSMFP(BSDF);
-						Payload.PhaseG       = SLAB_SSSPHASEANISOTROPY(BSDF); // STRATA_TODO: evaluate if we always want this or not ...
+						Payload.PhaseG       = SLAB_SSSPHASEANISOTROPY(BSDF);
 					}
 					else if (BSDF_GETSSSTYPE(BSDF) == SSS_TYPE_DIFFUSION_PROFILE)
 					{
@@ -607,7 +628,6 @@ RAY_TRACING_ENTRY_CLOSEST_HIT(PathTracingMaterialCHS,
 					Payload.TransmittanceN      = BSDF.TransmittanceAboveAlongN;
 					Payload.CoverageAboveAlongN = BSDF.CoverageAboveAlongN;
 					Payload.Ior                 = 0.0; // STRATA_TODO: figure out when to use refraction ...
-					Payload.ShadingModelID      = SHADINGMODELID_STRATA;
 					break;
 				}
 				case STRATA_BSDF_TYPE_HAIR:
