Substrate - Fixed some graph node preview math.

diff --git a/Engine/Shaders/Private/Strata/Strata.ush b/Engine/Shaders/Private/Strata/Strata.ush
--- a/Engine/Shaders/Private/Strata/Strata.ush
+++ b/Engine/Shaders/Private/Strata/Strata.ush
@@ -2454,7 +2454,7 @@ FStrataData StrataAddParameterBlending(FStrataData A, FStrataData B, float AMixF
 #endif
 
 #if USE_DEVELOPMENT_SHADERS
-	Result.PreviewColor = lerp(A.PreviewColor, B.PreviewColor, 0.5);
+	Result.PreviewColor = A.PreviewColor + B.PreviewColor;
 #endif
 
 #undef ResultBSDF
@@ -2597,6 +2597,10 @@ FStrataData StrataWeightParameterBlending(FStrataData A, float Weight)
 	StrataData.InlinedBSDF.Coverage *= SafeWeight;
 #endif
 
+#if USE_DEVELOPMENT_SHADERS
+	StrataData.PreviewColor *= SafeWeight;
+#endif
+
 	return StrataData;
 }
 
