Fix missing include causing shader compilation fail on certain platform.

diff --git a/Engine/Shaders/Private/HairStrands/HairStrandsPack.ush b/Engine/Shaders/Private/HairStrands/HairStrandsPack.ush
--- a/Engine/Shaders/Private/HairStrands/HairStrandsPack.ush
+++ b/Engine/Shaders/Private/HairStrands/HairStrandsPack.ush
@@ -2,6 +2,8 @@
 
 #pragma once
 
+#include "/Engine/Public/Platform.ush"
+#include "/Engine/Private/PackUnpack.ush"
 #include "/Engine/Shared/HairStrandsDefinitions.h"
 
 ////////////////////////////////////////////////////////////////////////////////////////////////////////////////////////////////
