Global Distance Field For Mobile
Remove bSupportsMobileDistanceField from DDSPI

diff --git a/Engine/Config/Windows/DataDrivenPlatformInfo.ini b/Engine/Config/Windows/DataDrivenPlatformInfo.ini
--- a/Engine/Config/Windows/DataDrivenPlatformInfo.ini
+++ b/Engine/Config/Windows/DataDrivenPlatformInfo.ini
@@ -147,7 +147,7 @@ MaxFeatureLevel=ES3_1
 ShaderFormat=PCD3D_ES31
 bIsPC=true
 bIsMobile=true
-bSupportsDistanceFields = true
+bSupportsDistanceFields = false
 bSupportsComputeFramework=true
 bSupportsByteBufferComputeShaders = true
 bSupportsInstancedStereo=true
diff --git a/Engine/Plugins/FX/Niagara/Source/NiagaraShader/Private/NiagaraDistanceFieldHelper.cpp b/Engine/Plugins/FX/Niagara/Source/NiagaraShader/Private/NiagaraDistanceFieldHelper.cpp
--- a/Engine/Plugins/FX/Niagara/Source/NiagaraShader/Private/NiagaraDistanceFieldHelper.cpp
+++ b/Engine/Plugins/FX/Niagara/Source/NiagaraShader/Private/NiagaraDistanceFieldHelper.cpp
@@ -63,4 +63,7 @@ void FNiagaraDistanceFieldHelper::SetGlobalDistanceFieldParameters(const FGlobal
 	ShaderParameters.NotCoveredMinStepScale = 0.0f;
 	ShaderParameters.DitheredTransparencyStepThreshold = 0.0f;
 	ShaderParameters.DitheredTransparencyTraceThreshold = 0.0f;
+	ShaderParameters.GlobalDistanceFieldCoverageAtlasTextureSampler = TStaticSamplerState<SF_Trilinear, AM_Wrap, AM_Wrap, AM_Wrap>::GetRHI();
+	ShaderParameters.GlobalDistanceFieldPageAtlasTextureSampler = TStaticSamplerState<SF_Trilinear, AM_Wrap, AM_Wrap, AM_Wrap>::GetRHI();
+	ShaderParameters.GlobalDistanceFieldMipTextureSampler = TStaticSamplerState<SF_Trilinear, AM_Clamp, AM_Clamp, AM_Clamp>::GetRHI();
 }
diff --git a/Engine/Shaders/Private/DistanceField/GlobalDistanceFieldShared.ush b/Engine/Shaders/Private/DistanceField/GlobalDistanceFieldShared.ush
--- a/Engine/Shaders/Private/DistanceField/GlobalDistanceFieldShared.ush
+++ b/Engine/Shaders/Private/DistanceField/GlobalDistanceFieldShared.ush
@@ -2,8 +2,6 @@
 
 #pragma once
 
-#if FEATURE_LEVEL >= FEATURE_LEVEL_SM5
-
 #if IS_MATERIAL_SHADER
 	#define DISTANCE_FIELD_IN_VIEW_UB 1
 #endif
@@ -71,6 +69,22 @@
 	float DitheredTransparencyTraceThreshold;
 #endif
 
+#if SUPPORTS_INDEPENDENT_SAMPLERS
+	#define GlobalDistanceFieldCoverageAtlasTextureSampler GlobalTrilinearWrappedSampler
+	#define GlobalDistanceFieldPageAtlasTextureSampler GlobalTrilinearWrappedSampler
+	#define GlobalDistanceFieldMipTextureSampler GlobalTrilinearClampedSampler
+#else
+	#if DISTANCE_FIELD_IN_VIEW_UB
+		#define GlobalDistanceFieldCoverageAtlasTextureSampler View.GlobalDistanceFieldCoverageAtlasTextureSampler
+		#define GlobalDistanceFieldPageAtlasTextureSampler View.GlobalDistanceFieldPageAtlasTextureSampler
+		#define GlobalDistanceFieldMipTextureSampler View.GlobalDistanceFieldMipTextureSampler
+	#else
+		SamplerState GlobalDistanceFieldCoverageAtlasTextureSampler;
+		SamplerState GlobalDistanceFieldPageAtlasTextureSampler;
+		SamplerState GlobalDistanceFieldMipTextureSampler;
+	#endif
+#endif
+
 // Must match GlobalDistanceField.cpp
 #define GLOBAL_DISTANCE_FIELD_PAGE_BORDER 0.5f
 #define GLOBAL_DISTANCE_FIELD_PAGE_RESOLUTION_IN_ATLAS 8 // Includes 0.5 texel trilinear filter margin
@@ -218,7 +232,7 @@ float GetDistanceToNearestSurfaceGlobal(float3 TranslatedWorldPosition)
 			{
 				float3 PageUV = ComputeGlobalDistanceFieldPageUV(ClipmapVolumeUV, Page);
 
-				float DistanceFieldValue = Texture3DSampleLevel(GlobalDistanceFieldPageAtlasTexture, GlobalTrilinearWrappedSampler, PageUV, 0).x;
+				float DistanceFieldValue = Texture3DSampleLevel(GlobalDistanceFieldPageAtlasTexture, GlobalDistanceFieldPageAtlasTextureSampler, PageUV, 0).x;
 
 				if (DistanceFieldValue < 1.0f)
 				{
@@ -247,7 +261,7 @@ float SampleGlobalDistanceField(float3 TranslatedWorldPosition, float MaxDistanc
 		{
 			float3 PageUV = ComputeGlobalDistanceFieldPageUV(ClipmapVolumeUV, Page);
 
-			float DistanceFieldValue = Texture3DSampleLevel(GlobalDistanceFieldPageAtlasTexture, GlobalTrilinearWrappedSampler, PageUV, 0).x;
+			float DistanceFieldValue = Texture3DSampleLevel(GlobalDistanceFieldPageAtlasTexture, GlobalDistanceFieldPageAtlasTextureSampler, PageUV, 0).x;
 
 			if (DistanceFieldValue < 1.0f)
 			{
@@ -271,7 +285,7 @@ float GlobalDistanceFieldSampleClipmap(float3 ClipmapVolumeUV, uint ClipmapIndex
 	if (Page.bValid)
 	{
 		float3 PageUV = ComputeGlobalDistanceFieldPageUV(ClipmapVolumeUV, Page);
-		DistanceFieldValue = Texture3DSampleLevel(GlobalDistanceFieldPageAtlasTexture, GlobalTrilinearWrappedSampler, PageUV, 0).x;
+		DistanceFieldValue = Texture3DSampleLevel(GlobalDistanceFieldPageAtlasTexture, GlobalDistanceFieldPageAtlasTextureSampler, PageUV, 0).x;
 	}
 
 	return DistanceFieldValue;
@@ -320,7 +334,7 @@ float3 GetDistanceFieldGradientGlobal(float3 TranslatedWorldPosition)
 			{
 				float3 PageUV = ComputeGlobalDistanceFieldPageUV(ClipmapVolumeUV, Page);
 
-				float DistanceFieldValue = Texture3DSampleLevel(GlobalDistanceFieldPageAtlasTexture, GlobalTrilinearWrappedSampler, PageUV, 0).x;
+				float DistanceFieldValue = Texture3DSampleLevel(GlobalDistanceFieldPageAtlasTexture, GlobalDistanceFieldPageAtlasTextureSampler, PageUV, 0).x;
 
 				// Find a page which where narrow band isn't clipped
 				if (DistanceFieldValue < 0.9f)
@@ -418,4 +432,3 @@ float3 GetDistanceFieldGradientGlobal(FLWCVector3 WorldPosition)
 	return GetDistanceFieldGradientGlobal(TranslatedWorldPosition);
 }
 
-#endif // FEATURE_LEVEL >= FEATURE_LEVEL_SM5
diff --git a/Engine/Shaders/Private/DistanceField/GlobalDistanceFieldUtils.ush b/Engine/Shaders/Private/DistanceField/GlobalDistanceFieldUtils.ush
--- a/Engine/Shaders/Private/DistanceField/GlobalDistanceFieldUtils.ush
+++ b/Engine/Shaders/Private/DistanceField/GlobalDistanceFieldUtils.ush
@@ -132,7 +132,7 @@ FGlobalSDFTraceResult RayTraceGlobalDistanceField(FGlobalSDFTraceInput TraceInpu
 				float3 ClipmapVolumeUV = ComputeGlobalUV(SampleTranslatedWorldPosition, ClipmapIndex);
 				float3 MipUV = ComputeGlobalMipUV(SampleTranslatedWorldPosition, ClipmapIndex);
 
-				float DistanceFieldMipValue = Texture3DSampleLevel(GlobalDistanceFieldMipTexture, GlobalTrilinearClampedSampler, MipUV, 0).x;
+				float DistanceFieldMipValue = Texture3DSampleLevel(GlobalDistanceFieldMipTexture, GlobalDistanceFieldMipTextureSampler, MipUV, 0).x;
 				float DistanceField = DecodeGlobalDistanceFieldPageDistance(DistanceFieldMipValue, GlobalDistanceFieldMipFactor * ClipmapInfluenceRange);
 				float Coverage = 1;
 				FGlobalDistanceFieldPage Page = GetGlobalDistanceFieldPage(ClipmapVolumeUV, ClipmapIndex);
@@ -146,11 +146,11 @@ FGlobalSDFTraceResult RayTraceGlobalDistanceField(FGlobalSDFTraceInput TraceInpu
 					{
 						float3 CoveragePageUV;
 						ComputeGlobalDistanceFieldPageUV(ClipmapVolumeUV, Page, PageUV, CoveragePageUV);
-						Coverage = Texture3DSampleLevel(GlobalDistanceFieldCoverageAtlasTexture, GlobalTrilinearWrappedSampler, CoveragePageUV, 0).x;
+						Coverage = Texture3DSampleLevel(GlobalDistanceFieldCoverageAtlasTexture, GlobalDistanceFieldCoverageAtlasTextureSampler, CoveragePageUV, 0).x;
 					}
 					#endif
 
-					float DistanceFieldValue = Texture3DSampleLevel(GlobalDistanceFieldPageAtlasTexture, GlobalTrilinearWrappedSampler, PageUV, 0).x;
+					float DistanceFieldValue = Texture3DSampleLevel(GlobalDistanceFieldPageAtlasTexture, GlobalDistanceFieldPageAtlasTextureSampler, PageUV, 0).x;
 					DistanceField = DecodeGlobalDistanceFieldPageDistance(DistanceFieldValue, ClipmapInfluenceRange);
 				}
 
diff --git a/Engine/Source/Runtime/Engine/Private/Materials/HLSLMaterialTranslator.cpp b/Engine/Source/Runtime/Engine/Private/Materials/HLSLMaterialTranslator.cpp
--- a/Engine/Source/Runtime/Engine/Private/Materials/HLSLMaterialTranslator.cpp
+++ b/Engine/Source/Runtime/Engine/Private/Materials/HLSLMaterialTranslator.cpp
@@ -3262,6 +3262,17 @@ int32 FHLSLMaterialTranslator::ErrorUnlessFeatureLevelSupported(ERHIFeatureLevel
 	return 0;
 }
 
+int32 FHLSLMaterialTranslator::ErrorUnlessPlatformSupports(const bool (*SupportFunction)(const FStaticShaderPlatform Platform), const TCHAR* ConditionString)
+{
+	if (!SupportFunction(Platform))
+	{
+		FString ShaderPlatformName = FDataDrivenShaderPlatformInfo::GetName(Platform).ToString();
+		return Errorf(TEXT("Node not supported in shader platform %s. The node requires %s support."), *ShaderPlatformName, ConditionString);
+	}
+
+	return 0;
+}
+
 int32 FHLSLMaterialTranslator::NonVertexShaderExpressionError()
 {
 	return Errorf(TEXT("Invalid node used in pixel/hull/domain shader input!"));
@@ -10528,7 +10539,7 @@ int32 FHLSLMaterialTranslator::GetHairColorFromMelanin(int32 Melanin, int32 Redn
 
 int32 FHLSLMaterialTranslator::DistanceToNearestSurface(int32 PositionArg)
 {
-	if (ErrorUnlessFeatureLevelSupported(ERHIFeatureLevel::SM5) == INDEX_NONE)
+	if (ErrorUnlessPlatformSupports(FDataDrivenShaderPlatformInfo::GetSupportsDistanceFields, TEXT("DistanceField")) == INDEX_NONE)
 	{
 		return INDEX_NONE;
 	}
@@ -10545,7 +10556,7 @@ int32 FHLSLMaterialTranslator::DistanceToNearestSurface(int32 PositionArg)
 
 int32 FHLSLMaterialTranslator::DistanceFieldGradient(int32 PositionArg)
 {
-	if (ErrorUnlessFeatureLevelSupported(ERHIFeatureLevel::SM5) == INDEX_NONE)
+	if (ErrorUnlessPlatformSupports(FDataDrivenShaderPlatformInfo::GetSupportsDistanceFields, TEXT("DistanceField")) == INDEX_NONE)
 	{
 		return INDEX_NONE;
 	}
@@ -10562,7 +10573,7 @@ int32 FHLSLMaterialTranslator::DistanceFieldGradient(int32 PositionArg)
 
 int32 FHLSLMaterialTranslator::DistanceFieldApproxAO(int32 PositionArg, int32 NormalArg, int32 BaseDistanceArg, int32 RadiusArg, uint32 NumSteps, float StepScale)
 {
-	if (ErrorUnlessFeatureLevelSupported(ERHIFeatureLevel::SM5) == INDEX_NONE)
+	if (ErrorUnlessPlatformSupports(FDataDrivenShaderPlatformInfo::GetSupportsDistanceFields, TEXT("DistanceField")) == INDEX_NONE)
 	{
 		return INDEX_NONE;
 	}
diff --git a/Engine/Source/Runtime/Engine/Private/Materials/HLSLMaterialTranslator.h b/Engine/Source/Runtime/Engine/Private/Materials/HLSLMaterialTranslator.h
--- a/Engine/Source/Runtime/Engine/Private/Materials/HLSLMaterialTranslator.h
+++ b/Engine/Source/Runtime/Engine/Private/Materials/HLSLMaterialTranslator.h
@@ -614,6 +614,7 @@ protected:
 	int32 NonPixelShaderExpressionError();
 
 	int32 ErrorUnlessFeatureLevelSupported(ERHIFeatureLevel::Type RequiredFeatureLevel);
+	int32 ErrorUnlessPlatformSupports(const bool (*SupportFunction)(const FStaticShaderPlatform Platform), const TCHAR* ConditionString);
 
 	int32 NonVertexShaderExpressionError();
 	int32 NonVertexOrPixelShaderExpressionError();
diff --git a/Engine/Source/Runtime/Engine/Private/SceneManagement.cpp b/Engine/Source/Runtime/Engine/Private/SceneManagement.cpp
--- a/Engine/Source/Runtime/Engine/Private/SceneManagement.cpp
+++ b/Engine/Source/Runtime/Engine/Private/SceneManagement.cpp
@@ -66,7 +66,7 @@ bool ShouldAllPrimitivesHaveDistanceField(EShaderPlatform ShaderPlatform)
 
 bool ShouldCompileDistanceFieldShaders(EShaderPlatform ShaderPlatform)
 {
-	return IsFeatureLevelSupported(ShaderPlatform, ERHIFeatureLevel::SM5) && DoesPlatformSupportDistanceFieldAO(ShaderPlatform) && IsUsingDistanceFields(ShaderPlatform);
+	return DoesPlatformSupportDistanceFieldAO(ShaderPlatform) && IsUsingDistanceFields(ShaderPlatform);
 }
 
 
@@ -956,6 +956,10 @@ FViewUniformShaderParameters::FViewUniformShaderParameters()
 	GlobalDistanceFieldPageTableTexture = BlackVolume;
 	GlobalDistanceFieldMipTexture = BlackVolume;
 
+	GlobalDistanceFieldPageAtlasTextureSampler = TStaticSamplerState<SF_Trilinear, AM_Wrap, AM_Wrap, AM_Wrap>::GetRHI();
+	GlobalDistanceFieldCoverageAtlasTextureSampler = TStaticSamplerState<SF_Trilinear, AM_Wrap, AM_Wrap, AM_Wrap>::GetRHI();
+	GlobalDistanceFieldMipTextureSampler = TStaticSamplerState<SF_Trilinear, AM_Clamp, AM_Clamp, AM_Clamp>::GetRHI();
+
 	SharedPointWrappedSampler = TStaticSamplerState<SF_Point, AM_Wrap, AM_Wrap, AM_Wrap>::GetRHI();
 	SharedPointClampedSampler = TStaticSamplerState<SF_Point, AM_Clamp, AM_Clamp, AM_Clamp>::GetRHI();
 	SharedBilinearWrappedSampler = TStaticSamplerState<SF_Bilinear, AM_Wrap, AM_Wrap, AM_Wrap>::GetRHI();
diff --git a/Engine/Source/Runtime/Engine/Public/SceneView.h b/Engine/Source/Runtime/Engine/Public/SceneView.h
--- a/Engine/Source/Runtime/Engine/Public/SceneView.h
+++ b/Engine/Source/Runtime/Engine/Public/SceneView.h
@@ -938,6 +938,9 @@ BEGIN_GLOBAL_SHADER_PARAMETER_STRUCT_WITH_CONSTRUCTOR(FViewUniformShaderParamete
 	SHADER_PARAMETER_TEXTURE(Texture3D, GlobalDistanceFieldCoverageAtlasTexture)
 	SHADER_PARAMETER_TEXTURE(Texture3D<uint>, GlobalDistanceFieldPageTableTexture)
 	SHADER_PARAMETER_TEXTURE(Texture3D, GlobalDistanceFieldMipTexture)
+	SHADER_PARAMETER_SAMPLER(SamplerState, GlobalDistanceFieldPageAtlasTextureSampler)
+	SHADER_PARAMETER_SAMPLER(SamplerState, GlobalDistanceFieldCoverageAtlasTextureSampler)
+	SHADER_PARAMETER_SAMPLER(SamplerState, GlobalDistanceFieldMipTextureSampler)
 
 	SHADER_PARAMETER_TEXTURE(Texture2D, AtmosphereTransmittanceTexture)
 	SHADER_PARAMETER_SAMPLER(SamplerState, AtmosphereTransmittanceTextureSampler)
@@ -1021,7 +1024,6 @@ BEGIN_GLOBAL_SHADER_PARAMETER_STRUCT_WITH_CONSTRUCTOR(FViewUniformShaderParamete
 	// Ray tracing
 	SHADER_PARAMETER(FVector3f, TLASRelativePreViewTranslation)
 	SHADER_PARAMETER(FVector3f, TLASViewTilePosition)
-
 END_GLOBAL_SHADER_PARAMETER_STRUCT()
 
 #undef VIEW_UNIFORM_BUFFER_MEMBER
diff --git a/Engine/Source/Runtime/RHI/Private/DataDrivenShaderPlatformInfo.cpp b/Engine/Source/Runtime/RHI/Private/DataDrivenShaderPlatformInfo.cpp
--- a/Engine/Source/Runtime/RHI/Private/DataDrivenShaderPlatformInfo.cpp
+++ b/Engine/Source/Runtime/RHI/Private/DataDrivenShaderPlatformInfo.cpp
@@ -258,7 +258,6 @@ void FGenericDataDrivenShaderPlatformInfo::ParseDataDrivenShaderInfo(const FConf
 	GET_SECTION_BOOL_HELPER(bSupportsManualVertexFetch);
 	GET_SECTION_BOOL_HELPER(bRequiresReverseCullingOnMobile);
 	GET_SECTION_BOOL_HELPER(bOverrideFMaterial_NeedsGBufferEnabled);
-	GET_SECTION_BOOL_HELPER(bSupportsMobileDistanceField);
 	GET_SECTION_BOOL_HELPER(bSupportsFFTBloom);
 	GET_SECTION_BOOL_HELPER(bSupportsVertexShaderLayer);
 	GET_SECTION_BINDLESS_SUPPORT_HELPER(BindlessSupport);
diff --git a/Engine/Source/Runtime/RHI/Public/DataDrivenShaderPlatformInfo.h b/Engine/Source/Runtime/RHI/Public/DataDrivenShaderPlatformInfo.h
--- a/Engine/Source/Runtime/RHI/Public/DataDrivenShaderPlatformInfo.h
+++ b/Engine/Source/Runtime/RHI/Public/DataDrivenShaderPlatformInfo.h
@@ -96,7 +96,6 @@ class FGenericDataDrivenShaderPlatformInfo
 	uint32 bSupportsManualVertexFetch : 1;
 	uint32 bRequiresReverseCullingOnMobile : 1;
 	uint32 bOverrideFMaterial_NeedsGBufferEnabled : 1;
-	uint32 bSupportsMobileDistanceField : 1;
 	uint32 bSupportsFFTBloom : 1;
 	uint32 bSupportsInlineRayTracing : 1;
 	uint32 bSupportsRayTracingShaders : 1;
@@ -676,10 +675,11 @@ public:
 		return Infos[Platform].bOverrideFMaterial_NeedsGBufferEnabled;
 	}
 
+	UE_DEPRECATED(5.3, "This function is no longer in use and will be removed. Please use GetSupportsDistanceFields instead")
 	static FORCEINLINE_DEBUGGABLE const bool GetSupportsMobileDistanceField(const FStaticShaderPlatform Platform)
 	{
 		check(IsValid(Platform));
-		return Infos[Platform].bSupportsMobileDistanceField;
+		return false;
 	}
 
 	static FORCEINLINE_DEBUGGABLE const bool GetSupportsFFTBloom(const FStaticShaderPlatform Platform)
diff --git a/Engine/Source/Runtime/RenderCore/Private/RenderUtils.cpp b/Engine/Source/Runtime/RenderCore/Private/RenderUtils.cpp
--- a/Engine/Source/Runtime/RenderCore/Private/RenderUtils.cpp
+++ b/Engine/Source/Runtime/RenderCore/Private/RenderUtils.cpp
@@ -584,7 +584,7 @@ RENDERCORE_API bool IsMobileAmbientOcclusionEnabled(const FStaticShaderPlatform
 
 RENDERCORE_API bool IsMobileDistanceFieldEnabled(const FStaticShaderPlatform Platform)
 {
-	return IsMobilePlatform(Platform) && (FDataDrivenShaderPlatformInfo::GetSupportsMobileDistanceField(Platform)/* || IsD3DPlatform(Platform)*/) && IsUsingDistanceFields(Platform);
+	return IsMobilePlatform(Platform) && FDataDrivenShaderPlatformInfo::GetSupportsDistanceFields(Platform) && IsUsingDistanceFields(Platform);
 }
 
 RENDERCORE_API bool IsMobileMovableSpotlightShadowsEnabled(const FStaticShaderPlatform Platform)
diff --git a/Engine/Source/Runtime/Renderer/Private/DistanceFieldAmbientOcclusion.cpp b/Engine/Source/Runtime/Renderer/Private/DistanceFieldAmbientOcclusion.cpp
--- a/Engine/Source/Runtime/Renderer/Private/DistanceFieldAmbientOcclusion.cpp
+++ b/Engine/Source/Runtime/Renderer/Private/DistanceFieldAmbientOcclusion.cpp
@@ -309,6 +309,11 @@ FAutoConsoleCommandWithWorld ListMeshDistanceFieldsMemoryConsoleCommand(
 	FConsoleCommandWithWorldDelegate::CreateStatic(OnListMeshDistanceFields)
 	);
 
+bool ShouldCompileDFNormalShaders(EShaderPlatform ShaderPlatform)
+{
+	return ShouldCompileDistanceFieldShaders(ShaderPlatform) && !IsMobilePlatform(ShaderPlatform);
+}
+
 class FComputeDistanceFieldNormalPS : public FGlobalShader
 {
 public:
@@ -325,7 +330,7 @@ public:
 
 	static bool ShouldCompilePermutation(const FGlobalShaderPermutationParameters& Parameters)
 	{
-		return ShouldCompileDistanceFieldShaders(Parameters.Platform);
+		return ShouldCompileDFNormalShaders(Parameters.Platform);
 	}
 
 	static void ModifyCompilationEnvironment(const FGlobalShaderPermutationParameters& Parameters, FShaderCompilerEnvironment& OutEnvironment)
@@ -354,7 +359,7 @@ public:
 
 	static bool ShouldCompilePermutation(const FGlobalShaderPermutationParameters& Parameters)
 	{
-		return ShouldCompileDistanceFieldShaders(Parameters.Platform);
+		return ShouldCompileDFNormalShaders(Parameters.Platform);
 	}
 
 	static void ModifyCompilationEnvironment(const FGlobalShaderPermutationParameters& Parameters, FShaderCompilerEnvironment& OutEnvironment)
@@ -672,7 +677,6 @@ bool SupportsDistanceFieldAO(ERHIFeatureLevel::Type FeatureLevel, EShaderPlatfor
 		&& !GRHIDeviceIsAMDPreGCNArchitecture
 		// Intel HD 4000 hangs in the RHICreateTexture3D call to allocate the large distance field atlas, and virtually no Intel cards can afford it anyway
 		&& !GRHIDeviceIsIntegrated
-		&& FeatureLevel >= ERHIFeatureLevel::SM5
 		&& DoesPlatformSupportDistanceFieldAO(ShaderPlatform)
 		&& IsUsingDistanceFields(ShaderPlatform);
 }
diff --git a/Engine/Source/Runtime/Renderer/Private/DistanceFieldLightingPost.cpp b/Engine/Source/Runtime/Renderer/Private/DistanceFieldLightingPost.cpp
--- a/Engine/Source/Runtime/Renderer/Private/DistanceFieldLightingPost.cpp
+++ b/Engine/Source/Runtime/Renderer/Private/DistanceFieldLightingPost.cpp
@@ -12,6 +12,7 @@
 #include "CompositionLighting/PostProcessAmbientOcclusion.h"
 #include "PipelineStateCache.h"
 #include "ScenePrivate.h"
+#include "DataDrivenShaderPlatformInfo.h"
 
 int32 GAOUseHistory = 1;
 FAutoConsoleVariableRef CVarAOUseHistory(
@@ -67,6 +68,11 @@ bool UseAOHistoryStabilityPass()
 	return GAOHistoryStabilityPass && GDistanceFieldAOQuality >= 2;
 }
 
+bool ShouldCompileDFLightingPostShaders(EShaderPlatform ShaderPlatform)
+{
+	return ShouldCompileDistanceFieldShaders(ShaderPlatform) && !IsMobilePlatform(ShaderPlatform);
+}
+
 BEGIN_SHADER_PARAMETER_STRUCT(FGeometryAwareUpsampleParameters, )
 	SHADER_PARAMETER_RDG_TEXTURE(Texture2D, DistanceFieldNormalTexture)
 	SHADER_PARAMETER_SAMPLER(SamplerState, DistanceFieldNormalSampler)
@@ -137,7 +143,7 @@ public:
 
 	static bool ShouldCompilePermutation(const FGlobalShaderPermutationParameters& Parameters)
 	{
-		return ShouldCompileDistanceFieldShaders(Parameters.Platform);
+		return ShouldCompileDFLightingPostShaders(Parameters.Platform);
 	}
 
 	static void ModifyCompilationEnvironment(const FGlobalShaderPermutationParameters& Parameters, FShaderCompilerEnvironment& OutEnvironment)
@@ -170,7 +176,7 @@ public:
 
 	static bool ShouldCompilePermutation(const FGlobalShaderPermutationParameters& Parameters)
 	{
-		return ShouldCompileDistanceFieldShaders(Parameters.Platform);
+		return ShouldCompileDFLightingPostShaders(Parameters.Platform);
 	}
 
 	static void ModifyCompilationEnvironment(const FGlobalShaderPermutationParameters& Parameters, FShaderCompilerEnvironment& OutEnvironment)
@@ -196,7 +202,7 @@ public:
 
 	static bool ShouldCompilePermutation(const FGlobalShaderPermutationParameters& Parameters)
 	{
-		return ShouldCompileDistanceFieldShaders(Parameters.Platform);
+		return ShouldCompileDFLightingPostShaders(Parameters.Platform);
 	}
 
 	static void ModifyCompilationEnvironment(const FGlobalShaderPermutationParameters& Parameters, FShaderCompilerEnvironment& OutEnvironment)
@@ -546,7 +552,7 @@ public:
 
 	static bool ShouldCompilePermutation(const FGlobalShaderPermutationParameters& Parameters)
 	{
-		return ShouldCompileDistanceFieldShaders(Parameters.Platform);
+		return ShouldCompileDFLightingPostShaders(Parameters.Platform);
 	}
 };
 
diff --git a/Engine/Source/Runtime/Renderer/Private/DistanceFieldScreenGridLighting.cpp b/Engine/Source/Runtime/Renderer/Private/DistanceFieldScreenGridLighting.cpp
--- a/Engine/Source/Runtime/Renderer/Private/DistanceFieldScreenGridLighting.cpp
+++ b/Engine/Source/Runtime/Renderer/Private/DistanceFieldScreenGridLighting.cpp
@@ -28,6 +28,7 @@
 #include "PipelineStateCache.h"
 #include "ClearQuad.h"
 #include "VisualizeTexture.h"
+#include "DataDrivenShaderPlatformInfo.h"
 
 /** Number of cone traced directions. */
 static const int32 NumConeSampleDirections = 9;
@@ -67,6 +68,11 @@ FVector2f JitterOffsets[4] =
 extern float GAOConeHalfAngle;
 extern bool DistanceFieldAOUseHistory(const FViewInfo& View);
 
+bool ShouldCompileDFScreenGridLightingShaders(EShaderPlatform ShaderPlatform)
+{
+	return ShouldCompileDistanceFieldShaders(ShaderPlatform) && !IsMobilePlatform(ShaderPlatform);
+}
+
 FVector2f GetJitterOffset(const FViewInfo& View)
 {
 	if (GAOUseJitter && DistanceFieldAOUseHistory(View))
@@ -157,7 +163,7 @@ public:
 
 	static bool ShouldCompilePermutation(const FGlobalShaderPermutationParameters& Parameters)
 	{
-		return ShouldCompileDistanceFieldShaders(Parameters.Platform);
+		return ShouldCompileDFScreenGridLightingShaders(Parameters.Platform);
 	}
 
 	static void ModifyCompilationEnvironment(const FGlobalShaderPermutationParameters& Parameters, FShaderCompilerEnvironment& OutEnvironment)
@@ -200,7 +206,7 @@ public:
 
 	static bool ShouldCompilePermutation(const FGlobalShaderPermutationParameters& Parameters)
 	{
-		return ShouldCompileDistanceFieldShaders(Parameters.Platform);
+		return ShouldCompileDFScreenGridLightingShaders(Parameters.Platform);
 	}
 
 	static void ModifyCompilationEnvironment(const FGlobalShaderPermutationParameters& Parameters, FShaderCompilerEnvironment& OutEnvironment)
@@ -238,7 +244,7 @@ public:
 
 	static bool ShouldCompilePermutation(const FGlobalShaderPermutationParameters& Parameters)
 	{
-		return ShouldCompileDistanceFieldShaders(Parameters.Platform);
+		return ShouldCompileDFScreenGridLightingShaders(Parameters.Platform);
 	}
 
 	static void ModifyCompilationEnvironment(const FGlobalShaderPermutationParameters& Parameters, FShaderCompilerEnvironment& OutEnvironment)
diff --git a/Engine/Source/Runtime/Renderer/Private/GlobalDistanceField.cpp b/Engine/Source/Runtime/Renderer/Private/GlobalDistanceField.cpp
--- a/Engine/Source/Runtime/Renderer/Private/GlobalDistanceField.cpp
+++ b/Engine/Source/Runtime/Renderer/Private/GlobalDistanceField.cpp
@@ -327,7 +327,10 @@ FGlobalDistanceFieldParameters2 SetupGlobalDistanceFieldParameters(const FGlobal
 	ShaderParameters.NotCoveredMinStepScale = GLumenSceneGlobalSDFNotCoveredMinStepScale;
 	ShaderParameters.DitheredTransparencyStepThreshold = GLumenSceneGlobalSDFDitheredTransparencyStepThreshold;
 	ShaderParameters.DitheredTransparencyTraceThreshold = GLumenSceneGlobalSDFDitheredTransparencyTraceThreshold;
-	
+	ShaderParameters.GlobalDistanceFieldCoverageAtlasTextureSampler = TStaticSamplerState<SF_Trilinear, AM_Wrap, AM_Wrap, AM_Wrap>::GetRHI();
+	ShaderParameters.GlobalDistanceFieldPageAtlasTextureSampler = TStaticSamplerState<SF_Trilinear, AM_Wrap, AM_Wrap, AM_Wrap>::GetRHI();
+	ShaderParameters.GlobalDistanceFieldMipTextureSampler = TStaticSamplerState<SF_Trilinear, AM_Clamp, AM_Clamp, AM_Clamp>::GetRHI();
+
 	return ShaderParameters;
 }
 
diff --git a/Engine/Source/Runtime/Renderer/Public/GlobalDistanceFieldParameters.h b/Engine/Source/Runtime/Renderer/Public/GlobalDistanceFieldParameters.h
--- a/Engine/Source/Runtime/Renderer/Public/GlobalDistanceFieldParameters.h
+++ b/Engine/Source/Runtime/Renderer/Public/GlobalDistanceFieldParameters.h
@@ -66,6 +66,9 @@ BEGIN_SHADER_PARAMETER_STRUCT(FGlobalDistanceFieldParameters2, )
 	SHADER_PARAMETER(float, NotCoveredMinStepScale)
 	SHADER_PARAMETER(float, DitheredTransparencyStepThreshold)
 	SHADER_PARAMETER(float, DitheredTransparencyTraceThreshold)
+	SHADER_PARAMETER_SAMPLER(SamplerState, GlobalDistanceFieldCoverageAtlasTextureSampler)
+	SHADER_PARAMETER_SAMPLER(SamplerState, GlobalDistanceFieldPageAtlasTextureSampler)
+	SHADER_PARAMETER_SAMPLER(SamplerState, GlobalDistanceFieldMipTextureSampler)
 END_SHADER_PARAMETER_STRUCT()
 
 FGlobalDistanceFieldParameters2 SetupGlobalDistanceFieldParameters(const FGlobalDistanceFieldParameterData& ParameterData);
