Fix screen space denoising shaders on AMD GPUs by removing CONFIG_VGPR_FREE_SAMPLE_TRACK_ID optimization

This optimization used to allow reusing VGPRs related to DispatchThreadId by computing SampleTrackId value (id of pixel in a quad) from GPU wave lane index (implemented in GetSampleTrackIdFromLaneIndex()).
The implementation of GetSampleTrackIdFromLaneIndex() relied on a linear/scanline order of pixel lanes within "quads" in a 2D compute dispatch, which is not guaranteed to be the case in D3D.

Linear order is only defined for a 1D dispatch: https://microsoft.github.io/DirectX-Specs/d3d/HLSL_SM_6_6_Derivatives.html#1d-quads
For 2D case, the spec contains the following language: There is no defined assignment of any quad within any wave and no relation between the values of SV_GroupIndex and the return value of WaveGetLaneIndex() should be assumed.

The optimization is no longer deemed critical for current UE5 rendering features and platforms. It is therefore decided to simply remove it instead of refactoring the denoising shaders to use 1D dispatch.

diff --git a/Engine/Shaders/Private/ScreenSpaceDenoise/SSDSpatialAccumulation.usf b/Engine/Shaders/Private/ScreenSpaceDenoise/SSDSpatialAccumulation.usf
--- a/Engine/Shaders/Private/ScreenSpaceDenoise/SSDSpatialAccumulation.usf
+++ b/Engine/Shaders/Private/ScreenSpaceDenoise/SSDSpatialAccumulation.usf
@@ -546,18 +546,6 @@
 	#error Unknown output mode.
 #endif
 
-
-//------------------------------------------------------- LATE SETUP
-
-#if COMPILER_SUPPORTS_WAVE_ONCE
-	#if CONFIG_UPSCALE
-		#define CONFIG_VGPR_FREE_SAMPLE_TRACK_ID 2
-	#else
-		#define CONFIG_VGPR_FREE_SAMPLE_TRACK_ID 1
-	#endif
-#endif
-
-
 //------------------------------------------------------- INCLUDES
 
 #include "SSDSignalFramework.ush"
@@ -1246,10 +1234,6 @@ void MainCS(
 	{
 		// No need to keep DispatchThreadId, can recompute the output pixel position based on information stored in VGPRs for spatial kernel.
 		#if CONFIG_SAMPLE_SET == SAMPLE_SET_STACKOWIAK_4_SETS
-			#if CONFIG_VGPR_FREE_SAMPLE_TRACK_ID
-				SampleTrackId = GetSampleTrackIdFromLaneIndex();
-			#endif
-
 			#if CONFIG_UPSCALE
 				SampleTrackId ^= 0x3;
 			#endif
diff --git a/Engine/Shaders/Private/ScreenSpaceDenoise/SSDSpatialKernel.ush b/Engine/Shaders/Private/ScreenSpaceDenoise/SSDSpatialKernel.ush
--- a/Engine/Shaders/Private/ScreenSpaceDenoise/SSDSpatialKernel.ush
+++ b/Engine/Shaders/Private/ScreenSpaceDenoise/SSDSpatialKernel.ush
@@ -30,11 +30,6 @@
 
 //------------------------------------------------------- CONFIG DISABLED DEFAULTS
 
-/** Compute FSSDKernelConfig::SampleTrackId directly based on lane id. */
-#ifndef CONFIG_VGPR_FREE_SAMPLE_TRACK_ID
-	#define CONFIG_VGPR_FREE_SAMPLE_TRACK_ID 0
-#endif
-
 #ifndef CONFIG_ACCUMULATOR_VGPR_COMPRESSION
 	#define CONFIG_ACCUMULATOR_VGPR_COMPRESSION ACCUMULATOR_COMPRESSION_DISABLED
 #endif
@@ -398,19 +393,6 @@ static const float kWaveletFilterWeights5x5[] = { 3.0 / 8.0, 1.0 / 4.0, 1.0 / 16
 
 //------------------------------------------------------- REDERIVE INFORMATION FOR LOWER VGPR OCCUPENCY
 
-#if CONFIG_VGPR_FREE_SAMPLE_TRACK_ID
-
-ISOLATE
-uint GetSampleTrackIdFromLaneIndex()
-{
-	uint LaneId = WaveGetLaneIndex();
-
-	return ((LaneId & 0x1) | ((LaneId >> 2) & 0x2)) ^ (CONFIG_VGPR_FREE_SAMPLE_TRACK_ID == 2 ? 0x3 : 0x0);
-}
-
-#endif // CONFIG_VGPR_FREE_SAMPLE_TRACK_ID
-
-
 /** Deduce the buffer UV of the output pixel this kernel has been configured for. */
 ISOLATE
 float2 ComputeRefBufferUV(FSSDKernelConfig KernelConfig)
@@ -429,10 +411,6 @@ float2 ComputeRefBufferUV(FSSDKernelConfig KernelConfig)
 	{
 		uint SampleTrackId = KernelConfig.SampleTrackId;
 
-		#if CONFIG_VGPR_FREE_SAMPLE_TRACK_ID
-			SampleTrackId = GetSampleTrackIdFromLaneIndex();
-		#endif
-
 		// Matches first line of kStackowiakSampleSet0
 		// TODO(Denoiser): could be optimised further by just setting sign bit on 0.5.
 		float2 SampleOffset = float2(
@@ -1734,9 +1712,6 @@ void ConvolveStackowiakKernel(
 				bool bIsKernelCenterSample = SampleId == 0 && (SampleBatchId == (kSamplingBatchSize - 1));
 				
 				uint SampleTrackId = KernelConfig.SampleTrackId;
-				#if CONFIG_VGPR_FREE_SAMPLE_TRACK_ID
-					SampleTrackId = GetSampleTrackIdFromLaneIndex();
-				#endif
 
 				float2 SampleOffset = kStackowiakSampleSet0[kStackowiakSampleSetCount * SampleId + SampleTrackId];
 				if (KernelConfig.SampleSubSetId == 1)
@@ -1821,9 +1796,6 @@ void ConvolveStackowiakKernel(
 		// Accumulate second sample to lattency hide with the center sample.
 		{
 			uint SampleTrackId = KernelConfig.SampleTrackId;
-			#if CONFIG_VGPR_FREE_SAMPLE_TRACK_ID
-				SampleTrackId = GetSampleTrackIdFromLaneIndex();
-			#endif
 
 			uint SampleId = 1;
 
@@ -1859,9 +1831,6 @@ void ConvolveStackowiakKernel(
 			float2 SampleBufferUV[2];
 			
 			uint SampleTrackId = KernelConfig.SampleTrackId;
-			#if CONFIG_VGPR_FREE_SAMPLE_TRACK_ID
-				SampleTrackId = GetSampleTrackIdFromLaneIndex();
-			#endif
 
 			UNROLL_N(2)
 			for (uint SampleBatchId = 0; SampleBatchId < 2; SampleBatchId++)
