Fix compilation error for hlslcc compiler

diff --git a/Engine/Shaders/Private/Bloom/BloomSurveyMaxScatterDispersion.usf b/Engine/Shaders/Private/Bloom/BloomSurveyMaxScatterDispersion.usf
--- a/Engine/Shaders/Private/Bloom/BloomSurveyMaxScatterDispersion.usf
+++ b/Engine/Shaders/Private/Bloom/BloomSurveyMaxScatterDispersion.usf
@@ -24,8 +24,7 @@ RWTexture2D<float4> DebugOutput;
 
 
 //------------------------------------------------------- LDS
-
-groupshared uint4 SharedMaxNeighborColor;
+groupshared uint SharedMaxNeighborColor[4];
 
 
 //------------------------------------------------------- ENTRY POINT
@@ -43,7 +42,10 @@ void MainCS(
 	// Init LDS
 	if (GroupThreadIndex == 0)
 	{
-		SharedMaxNeighborColor = 0;
+		SharedMaxNeighborColor[0] = 0;
+		SharedMaxNeighborColor[1] = 0;
+		SharedMaxNeighborColor[2] = 0;
+		SharedMaxNeighborColor[3] = 0;
 	}
 	
 	GroupMemoryBarrierWithGroupSync();
@@ -67,10 +69,10 @@ void MainCS(
 
 		TexelColor.rgb = max3(TexelColor.r, TexelColor.g, TexelColor.b).xxx;
 
-		InterlockedMax(SharedMaxNeighborColor.r, asuint(TexelColor.r));
-		InterlockedMax(SharedMaxNeighborColor.g, asuint(TexelColor.g));
-		InterlockedMax(SharedMaxNeighborColor.b, asuint(TexelColor.b));
-		InterlockedMax(SharedMaxNeighborColor.a, asuint(TexelColor.a));
+		InterlockedMax(SharedMaxNeighborColor[0], asuint(TexelColor.r));
+		InterlockedMax(SharedMaxNeighborColor[1], asuint(TexelColor.g));
+		InterlockedMax(SharedMaxNeighborColor[2], asuint(TexelColor.b));
+		InterlockedMax(SharedMaxNeighborColor[3], asuint(TexelColor.a));
 	}
 
 	GroupMemoryBarrierWithGroupSync();
@@ -78,10 +80,10 @@ void MainCS(
 	if (GroupThreadIndex == 0)
 	{
 		float4 MaxNeighborColor = float4(
-			asfloat(SharedMaxNeighborColor.r),
-			asfloat(SharedMaxNeighborColor.g),
-			asfloat(SharedMaxNeighborColor.b),
-			asfloat(SharedMaxNeighborColor.a));
+			asfloat(SharedMaxNeighborColor[0]),
+			asfloat(SharedMaxNeighborColor[1]),
+			asfloat(SharedMaxNeighborColor[2]),
+			asfloat(SharedMaxNeighborColor[3]));
 
 		SurveyOutput[LinearGroupId] = MaxNeighborColor;
 		
