Fix distant lights bug introduced by previous refactor

diff --git a/Engine/Shaders/Private/VirtualShadowMaps/VirtualShadowMapPageManagement.usf b/Engine/Shaders/Private/VirtualShadowMaps/VirtualShadowMapPageManagement.usf
--- a/Engine/Shaders/Private/VirtualShadowMaps/VirtualShadowMapPageManagement.usf
+++ b/Engine/Shaders/Private/VirtualShadowMaps/VirtualShadowMapPageManagement.usf
@@ -50,10 +50,14 @@ void GenerateHierarchicalPageFlags(uint ThreadId : SV_DispatchThreadID)
 #endif
 	if (Flag != 0)
 	{
-		// Note: we need to set the page rect bounds for the last mip level, since that is the only one that is valid, logically, for a single-page VSM.
-		//       This is important since this is what filters all the rendering that would otherwise try to draw stuff to the other levels.
 		uint MipLevel = MetaData.MipLevel;
 		uint2 PageAddress = MetaData.PageAddress;
+		if (bIsSinglePageSm)
+		{
+			// Note: we need to set the page rect bounds for the last mip level, since that is the only one that is valid, logically, for a single-page VSM.
+			//       This is important since this is what filters all the rendering that would otherwise try to draw stuff to the other levels.
+			MipLevel = VSM_MAX_MIP_LEVELS - 1U;
+		}
 
 		// Compute the min/max rect of active pages
 		uint PageBoundIndex = VirtualShadowMapId * VSM_MAX_MIP_LEVELS + MipLevel;
