Scalarized PixelIndex in Nanite CS shading

diff --git a/Engine/Shaders/Private/ComputeShaderOutputCommon.ush b/Engine/Shaders/Private/ComputeShaderOutputCommon.ush
--- a/Engine/Shaders/Private/ComputeShaderOutputCommon.ush
+++ b/Engine/Shaders/Private/ComputeShaderOutputCommon.ush
@@ -259,19 +259,13 @@ void ProcessPixel(bool bActiveLane, uint ShadingBin, const uint2 PixelPos, uint
 }
 
 [numthreads(64, 1, 1)]
-void MainCS(uint ThreadIndex : SV_GroupIndex, uint GroupID : SV_GroupID, uint3 GroupThreadID : SV_GroupThreadID, uint3 DispatchThreadID : SV_DispatchThreadID)
+void MainCS(uint ThreadIndex : SV_GroupIndex, uint GroupID : SV_GroupID, uint3 GroupThreadID : SV_GroupThreadID)
 {
 	const uint ShadingBin = PassData.x;
-
-	// Just like pixel shaders, compute 2x2 CS quads from thread index so we can leverage the 4 lane cross bar.
-	// [0 1][2 3] -> [0 1][4 5]
-	// [4 5][6 7] -> [2 3][6 7]
-	//const uint SwizzleX = BitFieldInsertU32(BitFieldMaskU32(1u, 0u), ThreadIndex, BitFieldExtractU32(ThreadIndex, 3u, 2u));
-	//const uint SwizzleY = BitFieldInsertU32(BitFieldMaskU32(2u, 0u), BitFieldExtractU32(ThreadIndex, 2u, 1u), BitFieldExtractU32(ThreadIndex, 3u, 3u));
-	//const uint SampleIndex	= SwizzleX * SwizzleY * 2u;
-
 	const bool bQuadBinning = PassData.z != 0u;
 
+	const uint PixelIndex = (GroupID * 64u) + ThreadIndex;
+
 	uint HelperLaneCount = 0u;
 
 	uint CornerIndex = 0;
@@ -309,7 +303,6 @@ void MainCS(uint ThreadIndex : SV_GroupIndex, uint GroupID : SV_GroupID, uint3 G
 	}
 	else
 	{
-		const uint PixelIndex = DispatchThreadID.x;
 		const uint PixelCount = Nanite.ShadingBinMeta[ShadingBin].y;
 		
 		if (PixelIndex >= PixelCount)
@@ -349,5 +342,5 @@ void MainCS(uint ThreadIndex : SV_GroupIndex, uint GroupID : SV_GroupID, uint3 G
 	const uint2 PixelPos = Quad.TopLeft + uint2(CornerIndex & 1u, (CornerIndex >> 1u) & 1u);
 
 	// Shade and export quad
-	ProcessPixel(bActiveLane, ShadingBin, PixelPos, QuadIndex, DispatchThreadID.x, PixelCopy, HelperLaneCount);
+	ProcessPixel(bActiveLane, ShadingBin, PixelPos, QuadIndex, PixelIndex, PixelCopy, HelperLaneCount);
 }
