Simplify and uniformize hair attribute access.

diff --git a/Engine/Plugins/Runtime/HairStrands/Shaders/Private/DeformerDataInterfaceGroom.ush b/Engine/Plugins/Runtime/HairStrands/Shaders/Private/DeformerDataInterfaceGroom.ush
--- a/Engine/Plugins/Runtime/HairStrands/Shaders/Private/DeformerDataInterfaceGroom.ush
+++ b/Engine/Plugins/Runtime/HairStrands/Shaders/Private/DeformerDataInterfaceGroom.ush
@@ -9,6 +9,9 @@
 
 HAIR_STRANDS_INSTANCE_PARAMETERS({DataInterfaceName})
 
+#define HAIR_STRANDS_ATTRIBUTE_ACCESSORS(Name) {DataInterfaceName}_##Name
+#include "/Engine/Private/HairStrands/HairStrandsAttributeTemplate.ush"
+
 ///////////////////////////////////////////////////////////////////////////////////////////////////
 // Helpers
 
@@ -70,19 +73,17 @@ float ReadCoordU_{DataInterfaceName}(uint ControlPointIndex)
 
 float ReadLength_{DataInterfaceName}(uint ControlPointIndex)
 {
-	const float CoordU      = ReadControlPointData(ControlPointIndex).UCoord;
-	const float CurveLength = UnpackHairLength(InternalGetHairStrandsAttributesPerCurve(ControlPointIndex, HAIR_CURVE_ATTRIBUTE_OFFSET_LENGTH({DataInterfaceName}_CurveAttributeOffsets), HAIR_CURVE_ATTRIBUTE_STRIDE_LENGTH, {DataInterfaceName}_CurveAttributeBuffer, {DataInterfaceName}_PointToCurveBuffer));
-	return CurveLength * CoordU;
+	return GetHairStrandsDimensions(ControlPointIndex, 0).x;
 }
 
 float2 ReadRootUV_{DataInterfaceName}(uint ControlPointIndex)
 {
-	return UnpackHairRootUV(InternalGetHairStrandsAttributesPerCurve(ControlPointIndex, HAIR_CURVE_ATTRIBUTE_OFFSET_ROOTUV({DataInterfaceName}_CurveAttributeOffsets), HAIR_CURVE_ATTRIBUTE_STRIDE_ROOTUV, {DataInterfaceName}_CurveAttributeBuffer, {DataInterfaceName}_PointToCurveBuffer));
+	return GetHairStrandsRootUV(ControlPointIndex);
 }
 
 float ReadSeed_{DataInterfaceName}(uint ControlPointIndex)
 {
-	return UnpackHairSeed(InternalGetHairStrandsAttributesPerCurve(ControlPointIndex, HAIR_CURVE_ATTRIBUTE_OFFSET_SEED({DataInterfaceName}_CurveAttributeOffsets), HAIR_CURVE_ATTRIBUTE_STRIDE_SEED, {DataInterfaceName}_CurveAttributeBuffer, {DataInterfaceName}_PointToCurveBuffer));
+	return GetHairStrandsSeed(ControlPointIndex);
 }
 
 float ReadCurveOffsetPoint_{DataInterfaceName}(uint CurveIndex)
diff --git a/Engine/Plugins/Runtime/HairStrands/Source/HairStrandsCore/Private/GroomTextureBuilder.cpp b/Engine/Plugins/Runtime/HairStrands/Source/HairStrandsCore/Private/GroomTextureBuilder.cpp
--- a/Engine/Plugins/Runtime/HairStrands/Source/HairStrandsCore/Private/GroomTextureBuilder.cpp
+++ b/Engine/Plugins/Runtime/HairStrands/Source/HairStrandsCore/Private/GroomTextureBuilder.cpp
@@ -1051,6 +1051,7 @@ static bool TraceTextures(
 			// Manually setting up parameters, since we don't have a real instance
 			FHairStrandsInstanceRawParameters Instance;
 			Instance.Common.GroupIndex = GroupIndex;
+			Instance.Common.GroupCount = GroupCount;
 			Instance.Common.PointCount = GroupData.Strands.RestResource->BulkData.GetNumPoints();
 			Instance.Common.CurveCount = GroupData.Strands.RestResource->BulkData.GetNumCurves();
 			Instance.Common.Radius = RenderingData.GeometrySettings.HairWidth * 0.5f;
diff --git a/Engine/Shaders/Private/HairStrands/HairStrandsAttributeTemplate.ush b/Engine/Shaders/Private/HairStrands/HairStrandsAttributeTemplate.ush
--- a/Engine/Shaders/Private/HairStrands/HairStrandsAttributeTemplate.ush
+++ b/Engine/Shaders/Private/HairStrands/HairStrandsAttributeTemplate.ush
@@ -10,120 +10,114 @@
 // Redefine macro for easing reability
 #define PREFIX(Name) HAIR_STRANDS_ATTRIBUTE_ACCESSORS(Name)
 
-float InternalGetHairStrandsWorldRadius(uint HairPrimitiveId, float2 VertexUV)
-{
-	const float Radius0 = InternalGetHairStrandsWorldRadius(PREFIX(PositionBuffer), HairPrimitiveId,   PREFIX(Radius), PREFIX(RootScale), PREFIX(TipScale));
-	const float Radius1 = InternalGetHairStrandsWorldRadius(PREFIX(PositionBuffer), HairPrimitiveId+1, PREFIX(Radius), PREFIX(RootScale), PREFIX(TipScale));
-	return lerp(Radius0, Radius1, VertexUV.x);
-}
-float InternalGetHairStrandsCoordU(uint HairPrimitiveId, float2 VertexUV)
-{
-	const float UCoord0 = InternalGetHairStrandsCoordU(PREFIX(PositionBuffer), HairPrimitiveId,   PREFIX(Radius), PREFIX(RootScale), PREFIX(TipScale));
-	const float UCoord1 = InternalGetHairStrandsCoordU(PREFIX(PositionBuffer), HairPrimitiveId+1, PREFIX(Radius), PREFIX(RootScale), PREFIX(TipScale));
-	return lerp(UCoord0, UCoord1, VertexUV.x);
-}
-float GetHairStrandsDepth(float2 InUV)
-{
-	return 0;
-}
-float GetHairStrandsCoverage(uint HairPrimitiveId, float2 HairPrimitiveUV)
-{
-	return 1;
-}
-float2 GetHairStrandsUV(uint HairPrimitiveId, float2 HairPrimitiveUV)
+float2 GetHairStrandsUV(uint InPrimitiveId, float2 InSegmentUV)
 {
+	const float CoordU0 = InternalGetHairStrandsCoordU(PREFIX(PositionBuffer), InPrimitiveId,   PREFIX(Radius), PREFIX(RootScale), PREFIX(TipScale));
+	const float CoordU1 = InternalGetHairStrandsCoordU(PREFIX(PositionBuffer), InPrimitiveId+1, PREFIX(Radius), PREFIX(RootScale), PREFIX(TipScale));
+
 	float2 Out;
-	Out.x = InternalGetHairStrandsCoordU(HairPrimitiveId, HairPrimitiveUV);
-	Out.y = HairPrimitiveUV.y;
+	Out.x = lerp(CoordU0, CoordU1, InSegmentUV.x);
+	Out.y = InSegmentUV.y;
 	return Out;
 }
-float2 GetHairStrandsDimensions(uint HairPrimitiveId, float2 HairPrimitiveUV)
+float2 GetHairStrandsDimensions(uint InPrimitiveId, float InSegmentU)
 {
-	const float CurveLength = UnpackHairLength(InternalGetHairStrandsAttributesPerCurve(HairPrimitiveId, HAIR_CURVE_ATTRIBUTE_OFFSET_LENGTH(PREFIX(CurveAttributeOffsets)), HAIR_CURVE_ATTRIBUTE_STRIDE_LENGTH, PREFIX(CurveAttributeBuffer), PREFIX(PointToCurveBuffer)));
-	const float CoordU = InternalGetHairStrandsCoordU(HairPrimitiveId, HairPrimitiveUV);
-	const float Radius = InternalGetHairStrandsWorldRadius(HairPrimitiveId, HairPrimitiveUV);
+	const float CurveLength = UnpackHairLength(InternalGetHairStrandsAttributesPerCurve(InPrimitiveId, HAIR_CURVE_ATTRIBUTE_OFFSET_LENGTH(PREFIX(CurveAttributeOffsets)), HAIR_CURVE_ATTRIBUTE_STRIDE_LENGTH, PREFIX(CurveAttributeBuffer), PREFIX(PointToCurveBuffer)));
+
+	const float CoordU0 = InternalGetHairStrandsCoordU(PREFIX(PositionBuffer), InPrimitiveId,   PREFIX(Radius), PREFIX(RootScale), PREFIX(TipScale));
+	const float CoordU1 = InternalGetHairStrandsCoordU(PREFIX(PositionBuffer), InPrimitiveId+1, PREFIX(Radius), PREFIX(RootScale), PREFIX(TipScale));
+	const float CoordU  = lerp(CoordU0, CoordU1, InSegmentU);
+
+	const float Radius0 = InternalGetHairStrandsWorldRadius(PREFIX(PositionBuffer), InPrimitiveId,   PREFIX(Radius), PREFIX(RootScale), PREFIX(TipScale));
+	const float Radius1 = InternalGetHairStrandsWorldRadius(PREFIX(PositionBuffer), InPrimitiveId+1, PREFIX(Radius), PREFIX(RootScale), PREFIX(TipScale));
+	const float Radius  = lerp(Radius0, Radius1, InSegmentU);
+
 	return float2(CurveLength * CoordU, Radius);
-}\
-float2 GetHairStrandsRootUV(uint HairPrimitiveId, float2 HairPrimitiveUV)
+}
+float2 GetHairStrandsRootUV(uint InPrimitiveId)
 {
 	const uint RootUVOffset = HAIR_CURVE_ATTRIBUTE_OFFSET_ROOTUV(PREFIX(CurveAttributeOffsets));
 	if (IsAttributeValid(RootUVOffset))
 	{
-		return UnpackHairRootUV(InternalGetHairStrandsAttributesPerCurve(HairPrimitiveId, RootUVOffset, HAIR_CURVE_ATTRIBUTE_STRIDE_ROOTUV, PREFIX(CurveAttributeBuffer), PREFIX(PointToCurveBuffer)));
+		return UnpackHairRootUV(InternalGetHairStrandsAttributesPerCurve(InPrimitiveId, RootUVOffset, HAIR_CURVE_ATTRIBUTE_STRIDE_ROOTUV, PREFIX(CurveAttributeBuffer), PREFIX(PointToCurveBuffer)));
 	}
 	else
 	{
 		return 0;
 	}
 }
-float GetHairStrandsSeed(uint HairPrimitiveId, float2 HairPrimitiveUV)
+float GetHairStrandsSeed(uint InPrimitiveId)
 {
-	return UnpackHairSeed(InternalGetHairStrandsAttributesPerCurve(HairPrimitiveId, HAIR_CURVE_ATTRIBUTE_OFFSET_SEED(PREFIX(CurveAttributeOffsets)), HAIR_CURVE_ATTRIBUTE_STRIDE_SEED, PREFIX(CurveAttributeBuffer), PREFIX(PointToCurveBuffer)));
+	return UnpackHairSeed(InternalGetHairStrandsAttributesPerCurve(InPrimitiveId, HAIR_CURVE_ATTRIBUTE_OFFSET_SEED(PREFIX(CurveAttributeOffsets)), HAIR_CURVE_ATTRIBUTE_STRIDE_SEED, PREFIX(CurveAttributeBuffer), PREFIX(PointToCurveBuffer)));
 }
-uint3 GetHairStrandsClumpID(uint HairPrimitiveId, float2 HairPrimitiveUV)
+uint3 GetHairStrandsClumpID(uint InPrimitiveId)
 {
 	const uint ClumpIdOffset  = HAIR_CURVE_ATTRIBUTE_OFFSET_CLUMPID(PREFIX(CurveAttributeOffsets));
 	const uint ClumpId3Offset = HAIR_CURVE_ATTRIBUTE_OFFSET_CLUMPID3(PREFIX(CurveAttributeOffsets));
 	if (IsAttributeValid(ClumpIdOffset))
 	{
-		return UnpackHairClumpID(InternalGetHairStrandsAttributesPerCurve(HairPrimitiveId, ClumpIdOffset, HAIR_CURVE_ATTRIBUTE_STRIDE_CLUMPID, PREFIX(CurveAttributeBuffer), PREFIX(PointToCurveBuffer)));
+		return UnpackHairClumpID(InternalGetHairStrandsAttributesPerCurve(InPrimitiveId, ClumpIdOffset, HAIR_CURVE_ATTRIBUTE_STRIDE_CLUMPID, PREFIX(CurveAttributeBuffer), PREFIX(PointToCurveBuffer)));
 	}
 	else if (IsAttributeValid(ClumpId3Offset))
 	{
-		return UnpackHairClumpID(InternalGetHairStrandsAttributesPerCurve64(HairPrimitiveId, ClumpId3Offset, HAIR_CURVE_ATTRIBUTE_STRIDE_CLUMPID3, PREFIX(CurveAttributeBuffer), PREFIX(PointToCurveBuffer)));
+		return UnpackHairClumpID(InternalGetHairStrandsAttributesPerCurve64(InPrimitiveId, ClumpId3Offset, HAIR_CURVE_ATTRIBUTE_STRIDE_CLUMPID3, PREFIX(CurveAttributeBuffer), PREFIX(PointToCurveBuffer)));
 	}
 	else
 	{
 		return 0;
 	}
 }
-float3 GetHairStrandsBaseColor(uint HairPrimitiveId, float2 HairPrimitiveUV)
+float3 GetHairStrandsColor(uint InPrimitiveId, float InSegmentU)
 {
 	const uint ColorOffset = HAIR_POINT_ATTRIBUTE_OFFSET_COLOR(PREFIX(PointAttributeOffsets));
 	if (IsAttributeValid(ColorOffset))
 	{
-		const float3 Color0 = UnpackHairColor(InternalGetHairStrandsAttributesPerVertex(HairPrimitiveId,   ColorOffset, HAIR_POINT_ATTRIBUTE_STRIDE_COLOR, PREFIX(PointAttributeBuffer)));
-		const float3 Color1 = UnpackHairColor(InternalGetHairStrandsAttributesPerVertex(HairPrimitiveId+1, ColorOffset, HAIR_POINT_ATTRIBUTE_STRIDE_COLOR, PREFIX(PointAttributeBuffer)));
-		return lerp(Color0, Color1, HairPrimitiveUV.x);
+		const float3 Color0 = UnpackHairColor(InternalGetHairStrandsAttributesPerVertex(InPrimitiveId,   ColorOffset, HAIR_POINT_ATTRIBUTE_STRIDE_COLOR, PREFIX(PointAttributeBuffer)));
+		const float3 Color1 = UnpackHairColor(InternalGetHairStrandsAttributesPerVertex(InPrimitiveId+1, ColorOffset, HAIR_POINT_ATTRIBUTE_STRIDE_COLOR, PREFIX(PointAttributeBuffer)));
+		return lerp(Color0, Color1, InSegmentU);
 	}
 	else
 	{
 		return 0;
 	}
 }
-float GetHairStrandsRoughness(uint HairPrimitiveId, float2 HairPrimitiveUV)
+float GetHairStrandsRoughness(uint InPrimitiveId, float InSegmentU)
 {
 	const uint RoughnessOffset = HAIR_POINT_ATTRIBUTE_OFFSET_ROUGHNESS(PREFIX(PointAttributeOffsets));
 	if (IsAttributeValid(RoughnessOffset))
 	{
-		const float Roughness0 = UnpackHairRoughness(InternalGetHairStrandsAttributesPerVertex(HairPrimitiveId,   RoughnessOffset, HAIR_POINT_ATTRIBUTE_STRIDE_ROUGHNESS, PREFIX(PointAttributeBuffer)));
-		const float Roughness1 = UnpackHairRoughness(InternalGetHairStrandsAttributesPerVertex(HairPrimitiveId+1, RoughnessOffset, HAIR_POINT_ATTRIBUTE_STRIDE_ROUGHNESS, PREFIX(PointAttributeBuffer)));
-		return lerp(Roughness0, Roughness1, HairPrimitiveUV.x);
+		const float Roughness0 = UnpackHairRoughness(InternalGetHairStrandsAttributesPerVertex(InPrimitiveId,   RoughnessOffset, HAIR_POINT_ATTRIBUTE_STRIDE_ROUGHNESS, PREFIX(PointAttributeBuffer)));
+		const float Roughness1 = UnpackHairRoughness(InternalGetHairStrandsAttributesPerVertex(InPrimitiveId+1, RoughnessOffset, HAIR_POINT_ATTRIBUTE_STRIDE_ROUGHNESS, PREFIX(PointAttributeBuffer)));
+		return lerp(Roughness0, Roughness1, InSegmentU);
 	}
 	else
 	{
 		return 0;
 	}
 }
-float GetHairStrandsAO(uint HairPrimitiveId, float2 HairPrimitiveUV)
+float GetHairStrandsAO(uint InPrimitiveId, float InSegmentU)
 {
 	const uint AoOffset = HAIR_POINT_ATTRIBUTE_OFFSET_AO(PREFIX(PointAttributeOffsets));
 	if (IsAttributeValid(AoOffset))
 	{
-		const float AO0 = UnpackHairRoughness(InternalGetHairStrandsAttributesPerVertex(HairPrimitiveId, AoOffset, HAIR_POINT_ATTRIBUTE_STRIDE_AO, PREFIX(PointAttributeBuffer)));
-		const float AO1 = UnpackHairRoughness(InternalGetHairStrandsAttributesPerVertex(HairPrimitiveId+1, AoOffset, HAIR_POINT_ATTRIBUTE_STRIDE_AO, PREFIX(PointAttributeBuffer)));
-		return lerp(AO0, AO1, HairPrimitiveUV.x);
+		const float AO0 = UnpackHairRoughness(InternalGetHairStrandsAttributesPerVertex(InPrimitiveId, AoOffset, HAIR_POINT_ATTRIBUTE_STRIDE_AO, PREFIX(PointAttributeBuffer)));
+		const float AO1 = UnpackHairRoughness(InternalGetHairStrandsAttributesPerVertex(InPrimitiveId+1, AoOffset, HAIR_POINT_ATTRIBUTE_STRIDE_AO, PREFIX(PointAttributeBuffer)));
+		return lerp(AO0, AO1, InSegmentU);
 	}
 	else
 	{
 		return 1.f;
 	}
 }
-float GetHairStrandsDepth(uint HairPrimitiveId, float2 InUV, float InDeviceZ) 
+float GetHairStrandsDepth() 
 { 
 	return 0.f; 
 } 
-float GetHairStrandsGroupIndex(float2 InAtlasUV) 
+float GetHairStrandsCoverage()
+{
+	return 1;
+}
+float GetHairStrandsGroupIndex() 
 { 
 	return PREFIX(GroupIndex); 
 }
diff --git a/Engine/Shaders/Private/HairStrands/HairStrandsDebug.usf b/Engine/Shaders/Private/HairStrands/HairStrandsDebug.usf
--- a/Engine/Shaders/Private/HairStrands/HairStrandsDebug.usf
+++ b/Engine/Shaders/Private/HairStrands/HairStrandsDebug.usf
@@ -2724,16 +2724,16 @@ void CSMain(
 			FFontColor FontName = FontWhite;
 
 			FShaderPrintContext Ctx = InitShaderPrintContext(true, uint2(PixelCoord) + uint2(20, 0));
-			Print(Ctx, TEXT("Group    : "), FontName); Print(Ctx, GetHairStrandsGroupIndex(0), FontValue); Newline(Ctx);
+			Print(Ctx, TEXT("Group    : "), FontName); Print(Ctx, GetHairStrandsGroupIndex(), FontValue); Newline(Ctx);
 			Print(Ctx, TEXT("Radius   : "), FontName); Print(Ctx, GetHairStrandsDimensions(VisNode.PrimitiveId, 0).y, FontValue); Newline(Ctx);
-			Print(Ctx, TEXT("Seed     : "), FontName); Print(Ctx, GetHairStrandsSeed(VisNode.PrimitiveId, 0), FontValue); Newline(Ctx);
-			Print(Ctx, TEXT("CoordU   : "), FontName); Print(Ctx, GetHairStrandsUV(VisNode.PrimitiveId, 0).x, FontValue); Newline(Ctx);
+			Print(Ctx, TEXT("Seed     : "), FontName); Print(Ctx, GetHairStrandsSeed(VisNode.PrimitiveId), FontValue); Newline(Ctx);
+			Print(Ctx, TEXT("CoordU   : "), FontName); Print(Ctx, GetHairStrandsUV(VisNode.PrimitiveId, float2(0, 0.5f)).x, FontValue); Newline(Ctx);
 			Newline(Ctx);
 
-			if (HasHairStrandsRootUV()) 	{ Print(Ctx, TEXT("RootUV   : "), FontName); Print(Ctx, GetHairStrandsRootUV(VisNode.PrimitiveId, 0),  FontValue); Newline(Ctx); }
-			if (HasHairStrandsClumpId()) 	{ Print(Ctx, TEXT("Clump ID : "), FontName); Print(Ctx, GetHairStrandsClumpID(VisNode.PrimitiveId, 0).x, FontValue); Newline(Ctx); }
-			if (HasHairStrandsClumpId3()) 	{ Print(Ctx, TEXT("Clump IDs: "), FontName); Print(Ctx, GetHairStrandsClumpID(VisNode.PrimitiveId, 0), FontValue); Newline(Ctx); }
-			if (HasHairStrandsColor()) 		{ Print(Ctx, TEXT("Color    : "), FontName); Print(Ctx, GetHairStrandsBaseColor(VisNode.PrimitiveId, 0), FontValue); Newline(Ctx); }
+			if (HasHairStrandsRootUV()) 	{ Print(Ctx, TEXT("RootUV   : "), FontName); Print(Ctx, GetHairStrandsRootUV(VisNode.PrimitiveId),  FontValue); Newline(Ctx); }
+			if (HasHairStrandsClumpId()) 	{ Print(Ctx, TEXT("Clump ID : "), FontName); Print(Ctx, GetHairStrandsClumpID(VisNode.PrimitiveId).x, FontValue); Newline(Ctx); }
+			if (HasHairStrandsClumpId3()) 	{ Print(Ctx, TEXT("Clump IDs: "), FontName); Print(Ctx, GetHairStrandsClumpID(VisNode.PrimitiveId), FontValue); Newline(Ctx); }
+			if (HasHairStrandsColor()) 		{ Print(Ctx, TEXT("Color    : "), FontName); Print(Ctx, GetHairStrandsColor(VisNode.PrimitiveId, 0), FontValue); Newline(Ctx); }
 			if (HasHairStrandsRoughness()) 	{ Print(Ctx, TEXT("Roughness: "), FontName); Print(Ctx, GetHairStrandsRoughness(VisNode.PrimitiveId, 0), FontValue); Newline(Ctx); }
 			if (HasHairStrandsAO()) 		{ Print(Ctx, TEXT("AO       : "), FontName); Print(Ctx, GetHairStrandsAO(VisNode.PrimitiveId, 0), FontValue); Newline(Ctx); }
 			Newline(Ctx);
diff --git a/Engine/Shaders/Private/HairStrands/HairStrandsTexturesGeneration.usf b/Engine/Shaders/Private/HairStrands/HairStrandsTexturesGeneration.usf
--- a/Engine/Shaders/Private/HairStrands/HairStrandsTexturesGeneration.usf
+++ b/Engine/Shaders/Private/HairStrands/HairStrandsTexturesGeneration.usf
@@ -64,6 +64,10 @@ void MainVS(
 
 HAIR_STRANDS_INSTANCE_PARAMETERS(InVF)
 
+#include "HairStrandsAttributeCommon.ush"
+#define HAIR_STRANDS_ATTRIBUTE_ACCESSORS(Name) InVF_##Name
+#include "/Engine/Private/HairStrands/HairStrandsAttributeTemplate.ush"
+
 float3 ComputeTangent(
 	const FHairControlPoint CP0,
 	const FHairControlPoint CP1,
@@ -284,29 +288,20 @@ bool Trace(
 
 						if (bIntersectAndClosest)
 						{
-							OutTangent = Tangent;
-							OutDepth = Distance;
-
-							OutSeed         = UnpackHairSeed(InternalGetHairStrandsAttributesPerCurve(P0, HAIR_CURVE_ATTRIBUTE_OFFSET_SEED(InVF_CurveAttributeOffsets), HAIR_CURVE_ATTRIBUTE_STRIDE_SEED, InVF_CurveAttributeBuffer, InVF_PointToCurveBuffer));
-							OutRootUV       = UnpackHairRootUV(InternalGetHairStrandsAttributesPerCurve(P0, HAIR_CURVE_ATTRIBUTE_OFFSET_ROOTUV(InVF_CurveAttributeOffsets), HAIR_CURVE_ATTRIBUTE_STRIDE_ROOTUV, InVF_CurveAttributeBuffer, InVF_PointToCurveBuffer));
-							OutU			= lerp(CP_Curr.UCoord, CP_Next.UCoord, VertexU);
-							OutMaterial.xyz = 0;
-							OutMaterial.w   = 0;
-
-							const uint ColorOffset = HAIR_POINT_ATTRIBUTE_OFFSET_COLOR(InVF_PointAttributeOffsets);
-							if (IsAttributeValid(ColorOffset))
+							OutTangent 	= Tangent;
+							OutDepth 	= Distance;
+							OutSeed		= GetHairStrandsSeed(P0);
+							OutRootUV	= GetHairStrandsRootUV(P0);
+							OutU		= GetHairStrandsDimensions(P0, VertexU).x;
+							OutMaterial	= 0;
+
+							if (HasHairStrandsColor())
 							{
-								const float3 C0 = UnpackHairColor(InternalGetHairStrandsAttributesPerVertex(P0, ColorOffset, HAIR_POINT_ATTRIBUTE_STRIDE_COLOR, InVF_PointAttributeBuffer));
-								const float3 C1 = UnpackHairColor(InternalGetHairStrandsAttributesPerVertex(P1, ColorOffset, HAIR_POINT_ATTRIBUTE_STRIDE_COLOR, InVF_PointAttributeBuffer));
-								OutMaterial.xyz = lerp(C0, C1, VertexU);
+								OutMaterial.xyz = GetHairStrandsColor(P0, VertexU);
 							}
-
-							const uint RoughnessOffset = HAIR_POINT_ATTRIBUTE_OFFSET_ROUGHNESS(InVF_PointAttributeOffsets);
-							if (IsAttributeValid(RoughnessOffset))
+							if (HasHairStrandsRoughness())
 							{
-								const float R0 = UnpackHairRoughness(InternalGetHairStrandsAttributesPerVertex(P0, RoughnessOffset, HAIR_POINT_ATTRIBUTE_STRIDE_ROUGHNESS, InVF_PointAttributeBuffer));
-								const float R1 = UnpackHairRoughness(InternalGetHairStrandsAttributesPerVertex(P1, RoughnessOffset, HAIR_POINT_ATTRIBUTE_STRIDE_ROUGHNESS, InVF_PointAttributeBuffer));
-								OutMaterial.w = lerp(R0, R1, VertexU);
+								OutMaterial.w = GetHairStrandsRoughness(P0, VertexU);
 							}
 
 							bIsValid = true;
diff --git a/Engine/Shaders/Private/MaterialTemplate.ush b/Engine/Shaders/Private/MaterialTemplate.ush
--- a/Engine/Shaders/Private/MaterialTemplate.ush
+++ b/Engine/Shaders/Private/MaterialTemplate.ush
@@ -1514,7 +1514,7 @@ float3 MaterialExpressionBlackBody( float Temp )
 float2 MaterialExpressionGetHairRootUV(FMaterialPixelParameters Parameters)
 {
 #if HAIR_STRAND_MESH_FACTORY
-	return GetHairStrandsRootUV(Parameters.HairPrimitiveId, Parameters.HairPrimitiveUV);
+	return GetHairStrandsRootUV(Parameters.HairPrimitiveId);
 #elif HAIR_CARD_MESH_FACTORY
 	return GetHairStrandsRootUV(Parameters.HairPrimitiveUV, Parameters.HairPrimitiveRootUV);
 #else
@@ -1536,7 +1536,7 @@ float2 MaterialExpressionGetHairUV(FMaterialPixelParameters Parameters)
 float2 MaterialExpressionGetHairDimensions(FMaterialPixelParameters Parameters)
 {
 #if HAIR_STRAND_MESH_FACTORY
-	return GetHairStrandsDimensions(Parameters.HairPrimitiveId, Parameters.HairPrimitiveUV);
+	return GetHairStrandsDimensions(Parameters.HairPrimitiveId, Parameters.HairPrimitiveUV.x);
 #elif HAIR_CARD_MESH_FACTORY
 	return GetHairStrandsDimensions(Parameters.HairPrimitiveUV, Parameters.HairPrimitiveLength);
 #else
@@ -1547,7 +1547,7 @@ float2 MaterialExpressionGetHairDimensions(FMaterialPixelParameters Parameters)
 float MaterialExpressionGetHairSeed(FMaterialPixelParameters Parameters)
 {
 #if HAIR_STRAND_MESH_FACTORY
-	return GetHairStrandsSeed(Parameters.HairPrimitiveId, Parameters.HairPrimitiveUV);
+	return GetHairStrandsSeed(Parameters.HairPrimitiveId);
 #elif HAIR_CARD_MESH_FACTORY
 	return GetHairStrandsSeed(Parameters.HairPrimitiveUV);
 #else
@@ -1558,7 +1558,7 @@ float MaterialExpressionGetHairSeed(FMaterialPixelParameters Parameters)
 float3 MaterialExpressionGetHairClumpID(FMaterialPixelParameters Parameters)
 {
 #if HAIR_STRAND_MESH_FACTORY
-	return GetHairStrandsClumpID(Parameters.HairPrimitiveId, Parameters.HairPrimitiveUV);
+	return GetHairStrandsClumpID(Parameters.HairPrimitiveId);
 #elif HAIR_CARD_MESH_FACTORY
 	return GetHairStrandsClumpID(Parameters.HairPrimitiveUV);
 #else
@@ -1569,7 +1569,7 @@ float3 MaterialExpressionGetHairClumpID(FMaterialPixelParameters Parameters)
 float3 MaterialExpressionGetHairBaseColor(FMaterialPixelParameters Parameters)
 {
 #if HAIR_STRAND_MESH_FACTORY
-	return GetHairStrandsBaseColor(Parameters.HairPrimitiveId, Parameters.HairPrimitiveUV);
+	return GetHairStrandsColor(Parameters.HairPrimitiveId, Parameters.HairPrimitiveUV.x);
 #elif HAIR_CARD_MESH_FACTORY
 	return GetHairStrandsBaseColor(Parameters.HairPrimitiveUV, Parameters.HairPrimitiveMaterial.xyz);
 #else
@@ -1580,7 +1580,7 @@ float3 MaterialExpressionGetHairBaseColor(FMaterialPixelParameters Parameters)
 float MaterialExpressionGetHairRoughness(FMaterialPixelParameters Parameters)
 {
 #if HAIR_STRAND_MESH_FACTORY
-	return GetHairStrandsRoughness(Parameters.HairPrimitiveId, Parameters.HairPrimitiveUV);
+	return GetHairStrandsRoughness(Parameters.HairPrimitiveId, Parameters.HairPrimitiveUV.x);
 #elif HAIR_CARD_MESH_FACTORY
 	return GetHairStrandsRoughness(Parameters.HairPrimitiveUV, Parameters.HairPrimitiveMaterial.w);
 #else
@@ -1591,7 +1591,7 @@ float MaterialExpressionGetHairRoughness(FMaterialPixelParameters Parameters)
 float MaterialExpressionGetHairAO(FMaterialPixelParameters Parameters)
 {
 #if HAIR_STRAND_MESH_FACTORY
-	return GetHairStrandsAO(Parameters.HairPrimitiveId, Parameters.HairPrimitiveUV);
+	return GetHairStrandsAO(Parameters.HairPrimitiveId, Parameters.HairPrimitiveUV.x);
 #elif HAIR_CARD_MESH_FACTORY
 	return GetHairStrandsAO(Parameters.HairPrimitiveUV, 1.0f);
 #else
@@ -1607,7 +1607,7 @@ float MaterialExpressionGetHairDepth(FMaterialVertexParameters Parameters)
 float MaterialExpressionGetHairDepth(FMaterialPixelParameters Parameters)
 {
 #if HAIR_STRAND_MESH_FACTORY
-	return GetHairStrandsDepth(Parameters.HairPrimitiveId, Parameters.HairPrimitiveUV, Parameters.SvPosition.z);
+	return GetHairStrandsDepth();
 #elif HAIR_CARD_MESH_FACTORY
 	return GetHairStrandsDepth(Parameters.HairPrimitiveUV, Parameters.SvPosition.z);
 #else
@@ -1618,7 +1618,7 @@ float MaterialExpressionGetHairDepth(FMaterialPixelParameters Parameters)
 float MaterialExpressionGetHairCoverage(FMaterialPixelParameters Parameters)
 {
 #if HAIR_STRAND_MESH_FACTORY
-	return GetHairStrandsCoverage(Parameters.HairPrimitiveId, Parameters.HairPrimitiveUV);
+	return GetHairStrandsCoverage();
 #elif HAIR_CARD_MESH_FACTORY
 	return GetHairStrandsCoverage(Parameters.HairPrimitiveUV);
 #else
@@ -1661,7 +1661,7 @@ float4 MaterialExpressionGetHairAuxilaryData(FMaterialPixelParameters Parameters
 float MaterialExpressionGetHairGroupIndex(FMaterialPixelParameters Parameters)
 {
 #if HAIR_STRAND_MESH_FACTORY
-	return GetHairStrandsGroupIndex(Parameters.HairPrimitiveId);
+	return GetHairStrandsGroupIndex();
 #elif HAIR_CARD_MESH_FACTORY
 	return GetHairStrandsGroupIndex(Parameters.HairPrimitiveUV, Parameters.HairPrimitiveGroupIndex);
 #else
