Orthographic rendering - Fix for general lighting/reflection/shadow passes which are not necessarily Lumen specific. Mainly alters all the relevant passes to calculate screen positions with or without depth considerations depending on whether perspective or orthographic projection (respectively) is used. Uses the function GetScreenPositionForProjectionType introduced in Common.ush in 24634679 to achieve this with minimal overhead.

Relies on initial changes of CLs 24634679 and 24886008

diff --git a/Engine/Shaders/Private/CapsuleShadowShaders.usf b/Engine/Shaders/Private/CapsuleShadowShaders.usf
--- a/Engine/Shaders/Private/CapsuleShadowShaders.usf
+++ b/Engine/Shaders/Private/CapsuleShadowShaders.usf
@@ -1011,7 +1011,7 @@ void CapsuleShadowingCS(
 	float SceneDepth = CalcSceneDepth(ScreenUV);
 #endif
 	
-	const float3 OpaqueTranslatedWorldPosition = mul(float4(ScreenPosition * SceneDepth, SceneDepth, 1), PrimaryView.ScreenToTranslatedWorld).xyz;
+	const float3 OpaqueTranslatedWorldPosition = mul(float4(GetScreenPositionForProjectionType(ScreenPosition, SceneDepth), SceneDepth, 1), PrimaryView.ScreenToTranslatedWorld).xyz;
 
 	uint CulledDataParameter = 0;
 	bool bTileShouldComputeShadowing = true;
diff --git a/Engine/Shaders/Private/ClusteredDeferredShadingPixelShader.usf b/Engine/Shaders/Private/ClusteredDeferredShadingPixelShader.usf
--- a/Engine/Shaders/Private/ClusteredDeferredShadingPixelShader.usf
+++ b/Engine/Shaders/Private/ClusteredDeferredShadingPixelShader.usf
@@ -374,7 +374,7 @@ float4 GetLightGridLocalLighting(
 		float2 ScreenPosition = UVAndScreenPos.zw;	\
 		FScreenSpaceData ScreenSpaceData = GetScreenSpaceData(ScreenUV);	\
 		float SceneDepth = CalcSceneDepth(ScreenUV);	\
-		float3 TranslatedWorldPosition = mul(float4(ScreenPosition * SceneDepth, SceneDepth, 1), PrimaryView.ScreenToTranslatedWorld).xyz;	\
+		float3 TranslatedWorldPosition = mul(float4(GetScreenPositionForProjectionType(ScreenPosition, SceneDepth), SceneDepth, 1), PrimaryView.ScreenToTranslatedWorld).xyz;	\
 		float3 CameraVector = normalize(TranslatedWorldPosition - PrimaryView.TranslatedWorldCameraOrigin);	\
 		ScreenSpaceData.GBuffer.ShadingModelID = ShadingModelId;	\
 		CompositedLighting += GetLightGridLocalLighting(ScreenSpaceData, LightGridData, TranslatedWorldPosition, CameraVector, ScreenUV, 0, Dither, FirstNonSimpleLightIndex); \
@@ -421,7 +421,7 @@ void ClusteredShadingPixelShader(
 	if (StrataPixelHeader.BSDFCount > 0 || CulledLightGridData.NumLocalLights == 0)
 	{
 		const float2 ScreenPosition = ViewportUVToScreenPos(ScreenUV);
-		const float3 TranslatedWorldPosition = mul(float4(ScreenPosition * SceneDepth, SceneDepth, 1), PrimaryView.ScreenToTranslatedWorld).xyz;
+		const float3 TranslatedWorldPosition = mul(float4(GetScreenPositionForProjectionType(ScreenPosition, SceneDepth), SceneDepth, 1), PrimaryView.ScreenToTranslatedWorld).xyz;
 		const float3 CameraVector = normalize(TranslatedWorldPosition - PrimaryView.TranslatedWorldCameraOrigin);
 		uint FirstNonSimpleLightIndex = 0;
 
@@ -477,7 +477,7 @@ void ClusteredShadingPixelShader(
 	}
 
 	float4 CompositedLighting = 0;
-	float3 TranslatedWorldPosition = mul(float4(ScreenPosition * SceneDepth, SceneDepth, 1), PrimaryView.ScreenToTranslatedWorld).xyz;
+	float3 TranslatedWorldPosition = mul(float4(GetScreenPositionForProjectionType(ScreenPosition, SceneDepth), SceneDepth, 1), PrimaryView.ScreenToTranslatedWorld).xyz;
 	float3 CameraVector = normalize(TranslatedWorldPosition - PrimaryView.TranslatedWorldCameraOrigin);
 	uint FirstNonSimpleLightIndex = 0;
 
@@ -551,7 +551,7 @@ void ClusteredShadingPixelShader(
 		return;
 	}
 
-	const float3 TranslatedWorldPosition = mul(float4(ScreenPosition * SceneDepth, SceneDepth, 1), PrimaryView.ScreenToTranslatedWorld).xyz;
+	const float3 TranslatedWorldPosition = mul(float4(GetScreenPositionForProjectionType(ScreenPosition, SceneDepth), SceneDepth, 1), PrimaryView.ScreenToTranslatedWorld).xyz;
 	const float3 CameraVector = normalize(TranslatedWorldPosition - PrimaryView.TranslatedWorldCameraOrigin);
 
 	float4 CompositedLighting = 0;
diff --git a/Engine/Shaders/Private/Common.ush b/Engine/Shaders/Private/Common.ush
--- a/Engine/Shaders/Private/Common.ush
+++ b/Engine/Shaders/Private/Common.ush
@@ -1201,7 +1201,7 @@ bool IsOrthoProjection()
 float ConvertToDeviceZ(float SceneDepth)
 {
 	FLATTEN
-	if (IsOrthoProjection(View.ViewToClip))
+	if (IsOrthoProjection())
 	{
 		// Ortho
 		return SceneDepth * View.ViewToClip[2][2] + View.ViewToClip[3][2];
diff --git a/Engine/Shaders/Private/DiffuseIndirectComposite.usf b/Engine/Shaders/Private/DiffuseIndirectComposite.usf
--- a/Engine/Shaders/Private/DiffuseIndirectComposite.usf
+++ b/Engine/Shaders/Private/DiffuseIndirectComposite.usf
@@ -352,7 +352,7 @@ void MainPS(
 	const float AOMask = IsValid(Material);
 	const float FinalAmbientOcclusion = lerp(1.0f, Material.MaterialAO * DynamicAmbientOcclusion, AOMask * AmbientOcclusionStaticFraction);
 
-	const float3 TranslatedWorldPosition = mul(float4(ScreenPosition * Material.SceneDepth, Material.SceneDepth, 1), View.ScreenToTranslatedWorld).xyz;
+	const float3 TranslatedWorldPosition = mul(float4(GetScreenPositionForProjectionType(ScreenPosition, Material.SceneDepth), Material.SceneDepth, 1), View.ScreenToTranslatedWorld).xyz;
 
 	const float3 N = Material.WorldNormal;
 	const float3 V = normalize(View.TranslatedWorldCameraOrigin - TranslatedWorldPosition);
diff --git a/Engine/Shaders/Private/DistanceFieldLightingPost.usf b/Engine/Shaders/Private/DistanceFieldLightingPost.usf
--- a/Engine/Shaders/Private/DistanceFieldLightingPost.usf
+++ b/Engine/Shaders/Private/DistanceFieldLightingPost.usf
@@ -26,8 +26,8 @@ SamplerState VelocityTextureSampler;
 
 float ComputeHistoryWeightBasedOnPosition(float2 ScreenPosition, float SceneDepth, float2 OldScreenPosition, float HistoryCameraDepth)
 {
-	float3 TranslatedWorldPosition = mul(float4(ScreenPosition * SceneDepth, SceneDepth, 1), View.ScreenToTranslatedWorld).xyz;
-	float3 PrevPositionTranslatedWorld = mul(float4(OldScreenPosition * HistoryCameraDepth, HistoryCameraDepth, 1), View.PrevScreenToTranslatedWorld).xyz;
+	float3 TranslatedWorldPosition = mul(float4(GetScreenPositionForProjectionType(ScreenPosition, SceneDepth), SceneDepth, 1), View.ScreenToTranslatedWorld).xyz;
+	float3 PrevPositionTranslatedWorld = mul(float4(GetScreenPositionForProjectionType(OldScreenPosition, HistoryCameraDepth), HistoryCameraDepth, 1), View.PrevScreenToTranslatedWorld).xyz;
 
 	float DistanceToHistoryValue = length(PrevPositionTranslatedWorld - TranslatedWorldPosition);
 	float RelativeHistoryDistanceThreshold = HistoryDistanceThreshold / 1000.0f;
@@ -79,7 +79,7 @@ float4 GetNormalWeights(float2 Corner00UV, float2 LowResTexelSize, float3 WorldN
 float ComputeSampleWeightBasedOnPosition(float4 ReferencePlane, float2 SampleScreenUV, float SampleDepth)
 {
 	float2 SampleScreenPosition = ApplyScreenTransform(SampleScreenUV, UVToScreenPos);
-	float3 SampleTranslatedWorldPosition = mul(float4(SampleScreenPosition * SampleDepth, SampleDepth, 1), PrimaryView.ScreenToTranslatedWorld).xyz;
+	float3 SampleTranslatedWorldPosition = mul(float4(GetScreenPositionForProjectionType(SampleScreenPosition, SampleDepth), SampleDepth, 1), PrimaryView.ScreenToTranslatedWorld).xyz;
 
 	float PlaneDistance = dot(ReferencePlane, float4(SampleTranslatedWorldPosition, 1));
 
@@ -94,7 +94,7 @@ void GeometryAwareUpsample(float4 UVAndScreenPos, out float4 OutBentNormal)
 	float SceneDepth;
 	GetDownsampledGBuffer(UVAndScreenPos.xy, WorldNormal, SceneDepth);
 
-	float3 TranslatedWorldPosition = mul(float4(UVAndScreenPos.zw * SceneDepth, SceneDepth, 1), PrimaryView.ScreenToTranslatedWorld).xyz;
+	float3 TranslatedWorldPosition = mul(float4(GetScreenPositionForProjectionType(UVAndScreenPos.zw, SceneDepth), SceneDepth, 1), PrimaryView.ScreenToTranslatedWorld).xyz;
 
 	float4 ReferencePlane = float4(WorldNormal, -dot(TranslatedWorldPosition, WorldNormal));
 
diff --git a/Engine/Shaders/Private/DistanceFieldScreenGridLighting.usf b/Engine/Shaders/Private/DistanceFieldScreenGridLighting.usf
--- a/Engine/Shaders/Private/DistanceFieldScreenGridLighting.usf
+++ b/Engine/Shaders/Private/DistanceFieldScreenGridLighting.usf
@@ -217,7 +217,7 @@ void ConeTraceObjectOcclusionCS(
 			float2 ScreenUV = GetScreenUVFromScreenGrid(OutputCoordinate);
 			float2 ScreenPosition = (ScreenUV.xy - View.ScreenPositionScaleBias.wz) / View.ScreenPositionScaleBias.xy;
 		
-			FLWCVector3 WorldShadingPosition = LWCMultiply(float3(ScreenPosition * SceneDepth, SceneDepth), PrimaryView.ScreenToWorld);
+			FLWCVector3 WorldShadingPosition = LWCMultiply(float3(GetScreenPositionForProjectionType(ScreenPosition, SceneDepth), SceneDepth), PrimaryView.ScreenToWorld);
 
 			HemisphereConeTraceAgainstTileCulledObjects(ObjectIndex, OutputBaseIndex, WorldShadingPosition, SceneDepth, WorldNormal, TangentX, TangentY);
 		}
@@ -344,7 +344,7 @@ void ConeTraceGlobalOcclusionCS(
 		float2 ScreenUV = GetScreenUVFromScreenGrid(OutputCoordinate);
 		float2 ScreenPosition = (ScreenUV.xy - View.ScreenPositionScaleBias.wz) / View.ScreenPositionScaleBias.xy;
 		
-		float3 TranslatedWorldPosition = mul(float4(ScreenPosition * SceneDepth, SceneDepth, 1), PrimaryView.ScreenToTranslatedWorld).xyz;
+		float3 TranslatedWorldPosition = mul(float4(GetScreenPositionForProjectionType(ScreenPosition, SceneDepth), SceneDepth, 1), PrimaryView.ScreenToTranslatedWorld).xyz;
 
 		uint OutputBaseIndex = OutputCoordinate.y * ScreenGridConeVisibilitySize.x + OutputCoordinate.x;
 
diff --git a/Engine/Shaders/Private/DistanceFieldShadowing.usf b/Engine/Shaders/Private/DistanceFieldShadowing.usf
--- a/Engine/Shaders/Private/DistanceFieldShadowing.usf
+++ b/Engine/Shaders/Private/DistanceFieldShadowing.usf
@@ -604,7 +604,7 @@ void DistanceFieldShadowingCS(
 	float2 ScreenPosition = (ScreenUV.xy - View.ScreenPositionScaleBias.wz) / View.ScreenPositionScaleBias.xy;
 
 	float SceneDepth = CalcSceneDepth(ScreenUV);
-	float3 OpaqueTranslatedWorldPosition = mul(float4(ScreenPosition * SceneDepth, SceneDepth, 1), PrimaryView.ScreenToTranslatedWorld).xyz;
+	float3 OpaqueTranslatedWorldPosition = mul(float4(GetScreenPositionForProjectionType(ScreenPosition, SceneDepth), SceneDepth, 1), PrimaryView.ScreenToTranslatedWorld).xyz;
 
 	// Distance for directional lights to trace
 	float TraceDistance = TanLightAngleAndNormalThreshold.z;
diff --git a/Engine/Shaders/Private/DistanceFieldVisualization.usf b/Engine/Shaders/Private/DistanceFieldVisualization.usf
--- a/Engine/Shaders/Private/DistanceFieldVisualization.usf
+++ b/Engine/Shaders/Private/DistanceFieldVisualization.usf
@@ -133,7 +133,7 @@ void VisualizeMeshDistanceFieldCS(
 	float2 ScreenPosition = (ScreenUV.xy - View.ScreenPositionScaleBias.wz) / View.ScreenPositionScaleBias.xy;
 
 	float SceneDepth = CalcSceneDepth(ScreenUV);
-	float3 TranslatedWorldPosition = mul(float4(ScreenPosition * SceneDepth, SceneDepth, 1), View.ScreenToTranslatedWorld).xyz;
+	float3 TranslatedWorldPosition = mul(float4(GetScreenPositionForProjectionType(ScreenPosition, SceneDepth), SceneDepth, 1), View.ScreenToTranslatedWorld).xyz;
 
 	float TraceDistance = 100000.0f;
 	float3 TranslatedWorldRayStart = PrimaryView.TranslatedWorldCameraOrigin;
diff --git a/Engine/Shaders/Private/MobileDeferredShading.usf b/Engine/Shaders/Private/MobileDeferredShading.usf
--- a/Engine/Shaders/Private/MobileDeferredShading.usf
+++ b/Engine/Shaders/Private/MobileDeferredShading.usf
@@ -153,7 +153,7 @@ void MobileDirectionalLightPS(
 
 	FGBufferData GBuffer = MobileFetchAndDecodeGBuffer(UVAndScreenPos.xy, UVAndScreenPos.zw);
 	float2 ScreenPos = UVAndScreenPos.zw;
-	float3 TranslatedWorldPosition = mul(float4(ScreenPos * GBuffer.Depth, GBuffer.Depth, 1), PrimaryView.ScreenToTranslatedWorld).xyz;
+	float3 TranslatedWorldPosition = mul(float4(GetScreenPositionForProjectionType(ScreenPosition, GBuffer.Depth), GBuffer.Depth, 1), PrimaryView.ScreenToTranslatedWorld).xyz;
 
 	half3 CameraVector = normalize(TranslatedWorldPosition);
 	half3 V = -CameraVector;
@@ -229,7 +229,7 @@ void MobileRadialLightPS(
 	FGBufferData GBuffer = MobileFetchAndDecodeGBuffer(ScreenUV, SVPos.xy);
 	// With a perspective projection, the clip space position is NDC * Clip.w
 	// With an orthographic projection, clip space is the same as NDC
-	float2 ClipPosition = InScreenPosition.xy / InScreenPosition.w * (View.ViewToClip[3][3] < 1.0f ? GBuffer.Depth : 1.0f);
+	float2 ClipPosition = GetScreenPositionForProjectionType(InScreenPosition.xy / InScreenPosition.w, GBuffer.Depth);
 	float3 TranslatedWorldPosition = mul(float4(ClipPosition, GBuffer.Depth, 1), PrimaryView.ScreenToTranslatedWorld).xyz;
 	half3 CameraVector = normalize(TranslatedWorldPosition);
 			
@@ -320,7 +320,7 @@ void MobileReflectionEnvironmentSkyLightingPS(
 	ResolvedView = ResolveView();
 	FGBufferData GBuffer = MobileFetchAndDecodeGBuffer(UVAndScreenPos.xy, UVAndScreenPos.zw);
 	float2 ScreenPos = UVAndScreenPos.zw;
-	float3 TranslatedWorldPosition = mul(float4(ScreenPos * GBuffer.Depth, GBuffer.Depth, 1), PrimaryView.ScreenToTranslatedWorld).xyz;
+	float3 TranslatedWorldPosition = mul(float4(GetScreenPositionForProjectionType(ScreenPos, GBuffer.Depth), GBuffer.Depth, 1), PrimaryView.ScreenToTranslatedWorld).xyz;
 	half3 CameraVector = normalize(TranslatedWorldPosition);
 	half3 V = -CameraVector;
 	half NoV = max(0, dot(GBuffer.WorldNormal, V));
diff --git a/Engine/Shaders/Private/ReflectionEnvironmentPixelShader.usf b/Engine/Shaders/Private/ReflectionEnvironmentPixelShader.usf
--- a/Engine/Shaders/Private/ReflectionEnvironmentPixelShader.usf
+++ b/Engine/Shaders/Private/ReflectionEnvironmentPixelShader.usf
@@ -117,7 +117,7 @@ float3 ReflectionEnvironment(FGBufferData GBuffer, float AmbientOcclusion, float
 	}
 #endif
 
-	float3 TranslatedWorldPosition = mul(float4(ScreenPosition * GBuffer.Depth, GBuffer.Depth, 1), View.ScreenToTranslatedWorld).xyz;
+	float3 TranslatedWorldPosition = mul(float4(GetScreenPositionForProjectionType(ScreenPosition, GBuffer.Depth), GBuffer.Depth, 1), View.ScreenToTranslatedWorld).xyz;
 	float3 CameraToPixel = normalize(TranslatedWorldPosition - View.TranslatedWorldCameraOrigin);
 	float3 ReflectionVector = reflect(CameraToPixel, GBuffer.WorldNormal);
     float3 V = -CameraToPixel;
@@ -292,8 +292,8 @@ void ReflectionEnvironmentSkyLighting(
 		float DeviceZ = SampleDeviceZFromSceneTextures(BufferUV);
 		float SceneDepth = ConvertFromDeviceZ(DeviceZ);
 
-		float3 WorldPosition = mul(float4(ScreenPosition * SceneDepth, SceneDepth, 1), LWCHackToFloat(PrimaryView.ScreenToWorld)).xyz;
-		float3 TranslatedWorldPosition = mul(float4(ScreenPosition * SceneDepth, SceneDepth, 1), View.ScreenToTranslatedWorld).xyz;
+		float3 WorldPosition = mul(float4(GetScreenPositionForProjectionType(ScreenPosition, SceneDepth), SceneDepth, 1), LWCHackToFloat(PrimaryView.ScreenToWorld)).xyz;
+		float3 TranslatedWorldPosition = mul(float4(GetScreenPositionForProjectionType(ScreenPosition, SceneDepth), SceneDepth, 1), View.ScreenToTranslatedWorld).xyz;
 		float3 CameraToPixel = normalize(WorldPosition - LWCHackToFloat(PrimaryView.WorldCameraOrigin));
 		float3 V = -CameraToPixel;
 
@@ -372,7 +372,7 @@ void ReflectionEnvironmentSkyLighting(
 	BRANCH
 	if (!bUnlitMaterial) // Only light pixels marked as lit
 	{
-		float3 TranslatedWorldPosition = mul(float4(ScreenPosition * GBuffer.Depth, GBuffer.Depth, 1), View.ScreenToTranslatedWorld).xyz;
+		float3 TranslatedWorldPosition = mul(float4(GetScreenPositionForProjectionType(ScreenPosition, GBuffer.Depth), GBuffer.Depth, 1), View.ScreenToTranslatedWorld).xyz;
 		const float CloudVolumetricAOShadow = GetCloudVolumetricAOShadow(TranslatedWorldPosition);
 
 		float3 SkyLighting = CloudVolumetricAOShadow * SkyLightDiffuse(GBuffer, AmbientOcclusion, BufferUV, ScreenPosition, BentNormal, DiffuseColor);
diff --git a/Engine/Shaders/Private/ReflectionEnvironmentShaders.usf b/Engine/Shaders/Private/ReflectionEnvironmentShaders.usf
--- a/Engine/Shaders/Private/ReflectionEnvironmentShaders.usf
+++ b/Engine/Shaders/Private/ReflectionEnvironmentShaders.usf
@@ -26,7 +26,7 @@ void CopyToCubeFaceVS(
 	)
 {	
 	DrawRectangle(float4(InPosition.xy, 0, 1), InUV, Out.Position, Out.UV);
-	Out.ScreenVector = mul(float4(Out.Position.xy, 1, 0), View.ScreenToTranslatedWorld).xyz;
+	Out.ScreenVector = ScreenVectorFromScreenRect(float4(Out.Position.xy, 1, 0));
 }
 
 int CubeFace;
@@ -861,7 +861,7 @@ void RenderRealTimeReflectionHeightFogVS(
 	UV = VertexId == 1 ? float2(-1.0f, 3.0f) : UV;
 	UV = VertexId == 2 ? float2(3.0f, -1.0f) : UV;
 	Position = float4(UV, 0.5f, 1.0f);
-	OutScreenVector = mul(float4(Position.xy, 1, 0), View.ScreenToTranslatedWorld).xyz;
+	OutScreenVector = ScreenVectorFromScreenRect(float4(Position.xy, 1, 0));
 }
 
 
diff --git a/Engine/Shaders/Private/ShadingFurnaceTest.usf b/Engine/Shaders/Private/ShadingFurnaceTest.usf
--- a/Engine/Shaders/Private/ShadingFurnaceTest.usf
+++ b/Engine/Shaders/Private/ShadingFurnaceTest.usf
@@ -83,7 +83,7 @@ void MainPS(
 
 	const float4 PixelLightAttenuation = 1.0f;
 	const float3 InRadiance = 0.5f;
-	const float3 WorldPosition = mul(float4(ScreenPosition * SceneDepth, SceneDepth, 1), PrimaryView.ScreenToTranslatedWorld).xyz;
+	const float3 WorldPosition = mul(float4(GetScreenPositionForProjectionType(ScreenPosition, SceneDepth), SceneDepth, 1), PrimaryView.ScreenToTranslatedWorld).xyz;
 	const float3 CameraVector = normalize(WorldPosition - PrimaryView.TranslatedWorldCameraOrigin);
 	const float3 TranslatedWorldPosition = WorldPosition;
 	const float3 V = -CameraVector;
diff --git a/Engine/Shaders/Private/ShadowProjectionPixelShader.usf b/Engine/Shaders/Private/ShadowProjectionPixelShader.usf
--- a/Engine/Shaders/Private/ShadowProjectionPixelShader.usf
+++ b/Engine/Shaders/Private/ShadowProjectionPixelShader.usf
@@ -123,7 +123,7 @@ void Main(
 	bIsSubsurfaceCompatible = false;
 #endif
 
-	float4 ScreenPosition = float4(((ScreenUV.xy - View.ScreenPositionScaleBias.wz ) / View.ScreenPositionScaleBias.xy) * SceneW, SceneW, 1);
+	float4 ScreenPosition = float4(GetScreenPositionForProjectionType(((ScreenUV.xy - View.ScreenPositionScaleBias.wz ) / View.ScreenPositionScaleBias.xy),SceneW), SceneW, 1);
 	float4 ShadowPosition = mul(ScreenPosition, ScreenToShadowMatrix);
 	float3 TranslateWorldPosition = mul(ScreenPosition, View.ScreenToTranslatedWorld).xyz;
 
@@ -414,7 +414,7 @@ void MainOnePassPointLightPS(
 #endif
 
 	float2 ScreenPosition = ( ScreenUV.xy - View.ScreenPositionScaleBias.wz ) / View.ScreenPositionScaleBias.xy;
-	float3 TranslateWorldPosition = mul(float4(ScreenPosition.xy * SceneW, SceneW, 1), View.ScreenToTranslatedWorld).xyz;
+	float3 TranslateWorldPosition = mul(float4(GetScreenPositionForProjectionType(ScreenPosition, SceneW), SceneW, 1), View.ScreenToTranslatedWorld).xyz;
 
 	// For debugging
 #define OUTPUT_CUBE_SHADOW_DEPTH_NO_FILTERING 0
