Fixes TSR ghosting on water when characters walk into waters with static camera

This changes moves the luminance measurement used for anti-flickering from the scene color a lit but further up in the renderer, before SLW.

To avoid starting ghost in case of high geometric density far in the distance hidden in fog, but also due to geometric details hidden underwater, the TSR anti-flickering heuristic ends up comparing the final scene color luminance versus the one measured earlier to accurately take into account what has been draw over since generally isn't the cause of the flickering.

diff --git a/Engine/Shaders/Private/TemporalSuperResolution/TSRCommon.ush b/Engine/Shaders/Private/TemporalSuperResolution/TSRCommon.ush
--- a/Engine/Shaders/Private/TemporalSuperResolution/TSRCommon.ush
+++ b/Engine/Shaders/Private/TemporalSuperResolution/TSRCommon.ush
@@ -236,6 +236,10 @@ static const tsr_half4 LargestSceneColorRGBA = (LargestNormalNumber * tsr_half(0
 static const tsr_half3 LargestSceneColorYCoCg = LargestNormalNumber.xxx;
 static const tsr_half4 LargestSceneColorYCoCgA = LargestNormalNumber.xxxx;
 
+// Weight of each channel for the computation of luminance.
+static const tsr_half kMoireLumaToChannel = tsr_half(3.0);
+static const tsr_half3 kMoireLumaWeights = rcp(kMoireLumaToChannel).xxx;
+
 
 //------------------------------------------------------- PARAMETERS
 
diff --git a/Engine/Shaders/Private/TemporalSuperResolution/TSRComputeMoireLuma.usf b/Engine/Shaders/Private/TemporalSuperResolution/TSRComputeMoireLuma.usf
--- a/Engine/Shaders/Private/TemporalSuperResolution/TSRComputeMoireLuma.usf
+++ b/Engine/Shaders/Private/TemporalSuperResolution/TSRComputeMoireLuma.usf
@@ -36,7 +36,7 @@ void MainCS(
 			UNROLL_N(TILE_ROWS)
 			for (uint x = 0; x < TILE_ROWS; x++)
 			{
-				uint2 PixelOffset = uint2(x, y) * TILE_SIZE;
+				const uint2 PixelOffset = uint2(x, y) * TILE_SIZE;
 				uint2 InputPixelPos = OutputPixelOrigin + PixelOffset;
 			
 				ColorArray[x + y * TILE_ROWS] = SceneColorTexture[InputPixelPos];
@@ -46,16 +46,13 @@ void MainCS(
 
 	ISOLATE
 	{
-		const tsr_half MoireLumaToChannel = tsr_half(3.0);
-		const tsr_half3 LumaWeights = rcp(MoireLumaToChannel.xxx);
-
 		UNROLL_N(TILE_ROWS)
 		for (uint y = 0; y < TILE_ROWS; y++)
 		{
 			UNROLL_N(TILE_ROWS)
 			for (uint x = 0; x < TILE_ROWS; x++)
 			{
-				uint2 PixelOffset = uint2(x, y) * TILE_SIZE;
+				const uint2 PixelOffset = uint2(x, y) * TILE_SIZE;
 				uint2 InputPixelPos = OutputPixelOrigin + PixelOffset;
 			
 				tsr_half3 Color = ColorArray[x + y * TILE_ROWS];
@@ -63,7 +60,7 @@ void MainCS(
 
 				InputPixelPos.x = select(all(InputPixelPos < InputInfo_ViewportMax), InputPixelPos.x, uint(~0));
 				
-				MoireLumaOutput[InputPixelPos] = dot(Color, LumaWeights);
+				MoireLumaOutput[InputPixelPos] = dot(Color, kMoireLumaWeights);
 			}
 		}
 	}
diff --git a/Engine/Shaders/Private/TemporalSuperResolution/TSRDebugHistory.usf b/Engine/Shaders/Private/TemporalSuperResolution/TSRDebugHistory.usf
--- a/Engine/Shaders/Private/TemporalSuperResolution/TSRDebugHistory.usf
+++ b/Engine/Shaders/Private/TemporalSuperResolution/TSRDebugHistory.usf
@@ -16,14 +16,12 @@ Texture2D<tsr_halfC> History_HighFrequency;
 Texture2D<tsr_half2> History_Metadata;
 Texture2D<tsr_half3> History_Translucency;
 Texture2D<tsr_half>  History_TranslucencyAlpha;
-Texture2D<tsr_subpixel_details> History_SubpixelDetails;
 
 Texture2D<tsr_halfC> PrevHistory_LowFrequency;
 Texture2D<tsr_halfC> PrevHistory_HighFrequency;
 Texture2D<tsr_half2> PrevHistory_Metadata;
 Texture2D<tsr_half3> PrevHistory_Translucency;
 Texture2D<tsr_half>  PrevHistory_TranslucencyAlpha;
-Texture2D<tsr_subpixel_details> PrevHistory_SubpixelDetails;
 
 
 //------------------------------------------------------- FUNCTION
@@ -49,20 +47,6 @@ float Median(float A, float B, float C)
 }
 
 
-tsr_subpixel_payload ReadSubpixelDetails(uint2 SubpixelPos)
-{
-	uint2 HistoryPixelPos = SubpixelPos / SUB_PIXEL_GRID_SIZE;
-
-	uint2 SubpixelRelPos = SubpixelPos % SUB_PIXEL_GRID_SIZE;
-	uint SubpixelId = dot(SubpixelRelPos, uint2(1, SUB_PIXEL_GRID_SIZE));
-	uint SubpixelShift = SubpixelId * SUB_PIXEL_BIT_COUNT;
-	
-	tsr_subpixel_details RawHistory3 = History_SubpixelDetails[HistoryPixelPos];
-
-	return tsr_subpixel_payload(RawHistory3 >> SubpixelShift) & SUB_PIXEL_BIT_MASK;
-}
-
-
 //------------------------------------------------------- ENTRY POINT
 
 [numthreads(TILE_SIZE, TILE_SIZE, 1)]
@@ -80,27 +64,8 @@ void MainCS(
 	float2 ViewportUV = (float2(HistoryPixelPos) + 0.5f) * HistoryInfo_ViewportSizeInverse;
 	float2 ScreenPos = ViewportUVToScreenPos(ViewportUV);
 
-	uint2 SubpixelPos = DispatchThreadId % SUB_PIXEL_GRID_SIZE;
-	uint SubpixelId = dot(SubpixelPos, uint2(1, SUB_PIXEL_GRID_SIZE));
-	uint SubpixelShift = SubpixelId * SUB_PIXEL_BIT_COUNT;
-	
 	{
-		#if 1
-			//tsr_half3 RawHistory0 = History_Textures_0[HistoryPixelPos];
-			//tsr_half3 RawHistory1 = History_Textures_1[HistoryPixelPos];
-			// RawHistory2 = History_Textures_2[HistoryPixelPos];
-			tsr_subpixel_details RawHistory3 = History_SubpixelDetails[HistoryPixelPos];
-		#else
-			//tsr_half3 RawHistory0 = PrevHistory_Textures_0[HistoryPixelPos];
-			//tsr_half3 RawHistory1 = PrevHistory_Textures_1[HistoryPixelPos];
-			//tsr_half2 RawHistory2 = PrevHistory_Textures_2[HistoryPixelPos];
-			tsr_subpixel_details RawHistory3 = PrevHistory_SubpixelDetails[HistoryPixelPos];
-		#endif
-
-		tsr_subpixel_payload SubpixelDetails = tsr_subpixel_payload(RawHistory3 >> SubpixelShift) & SUB_PIXEL_BIT_MASK;
-
-		Debug.r = UncompressParallaxFactor(SubpixelDetails);
-		//Debug.b = float(SubpixelId);
+
 	}
 
 	#if DEBUG_OUTPUT && 1
diff --git a/Engine/Shaders/Private/TemporalSuperResolution/TSRRejectShading.usf b/Engine/Shaders/Private/TemporalSuperResolution/TSRRejectShading.usf
--- a/Engine/Shaders/Private/TemporalSuperResolution/TSRRejectShading.usf
+++ b/Engine/Shaders/Private/TemporalSuperResolution/TSRRejectShading.usf
@@ -253,8 +253,6 @@ void MainCS(
 	
 	// Analyses the moire pattern overtime.
 	#if CONFIG_MOIRE_DETECTION
-	const tsr_half MoireLumaToChannel = tsr_half(3.0);
-
 	tsr_simd_half4 FinalMoire = tsr_simd_half4::Const(0.0);
 	tsr_simd_half MoireError = tsr_simd_half::Const(0.0);
 	{
@@ -265,26 +263,27 @@ void MainCS(
 		// Compute the threshold at which the variation count should quick in
 		const tsr_half TotalVariationCountThreshold = tsr_half(1.0 / (1.0 - pow(1.0 - BlendFinal, tsr_half(FlickeringFramePeriod))));
 
+		// Input luminance to monitor flickering from.
+		tsr_simd_half InputLuma = MoireInput;
 
-		tsr_simd_half InputLuma;
-		#if 1
-		{
-			InputLuma = MoireInput;
-		}
-		#else
+		// How much input luminance from the TSRComputeMoireLuma.usf measurement has been modified since for instance with fog, water.
+		tsr_simd_half InputLumaModification;
 		{
+			const tsr_half InputLumaEncodingError = tsr_half(0.5 / 255.0);
+
+			tsr_simd_half SceneColorLuma;
 			UNROLL_N(SIMD_SIZE)
 			for (uint SimdIndex = 0; SimdIndex < SIMD_SIZE; SimdIndex++)
 			{
-				const tsr_half3 LumaWeights = rcp(MoireLumaToChannel.xxx);
-				
-				InputLuma.SetElement(SimdIndex, dot(OriginalOpaqueInput.GetElement(SimdIndex).rgb, LumaWeights));
+				SceneColorLuma.SetElement(SimdIndex, dot(OriginalOpaqueInput.GetElement(SimdIndex).rgb, kMoireLumaWeights));
 			}
+
+			InputLumaModification = max(abs(InputLuma - SceneColorLuma) - tsr_simd_half::Const(InputLumaEncodingError), tsr_simd_half::Const(0.0));
 		}
-		#endif
 		
 		// Convolve the input luma to have the exact same moire pattern as FilteredInput below.
 		InputLuma = Blur3x3(InputLuma, bIsValid);
+		InputLumaModification = Blur3x3(InputLumaModification, bIsValid);
 
 		// Unpack history
 		tsr_simd_half PrevFlickeringHistory = PrevMoire[0];
@@ -311,6 +310,17 @@ void MainCS(
 		//tsr_simd_half CurrentGradient = FinalFlickeringHistory - PrevFlickeringHistory;
 		tsr_simd_half CurrentGradient = ClampedPrevFlickeringHistory - PrevFlickeringHistory;
 		
+		// Remove how much luma has changed from luma measurement to final scene color from the gradient.
+		for (uint SimdIndex = 0; SimdIndex < SIMD_SIZE; SimdIndex++)
+		{
+			tsr_half LocalCurrentGradient = CurrentGradient.GetElement(SimdIndex);
+			tsr_half LocalInputLumaModification = InputLumaModification.GetElement(SimdIndex);
+			
+			Debug[0].SetElement(SimdIndex, tsr_halfC(LocalInputLumaModification, 0, 0));
+
+			CurrentGradient.SetElement(SimdIndex, clamp(LocalCurrentGradient * tsr_half(POSITIVE_INFINITY), tsr_half(-1.0), tsr_half(1.0)) * max(abs(LocalCurrentGradient) - LocalInputLumaModification, tsr_half(0.0)));
+		}
+
 		#if 0
 			tsr_simd_half InputBoxSize = InputLumaMax - InputLumaMin;
 			CurrentGradient = sign(CurrentGradient) * max(abs(CurrentGradient) - tsr_simd_half::Const(0.5) * InputBoxSize, tsr_simd_half::Const(0.0));
@@ -448,11 +458,11 @@ void MainCS(
 				tsr_halfC ClampedEnergyAmount = ClampedEnergy * SafeRcp(ClampedEnergy[0] + ClampedEnergy[1] + ClampedEnergy[2]);
 
 				// Apply the moire error on each individual channel based on how much they need to clamp.
-				MoireErrorSize.SetElement(SimdIndex, ClampedEnergyAmount * (MoireLumaToChannel * MoireError.GetElement(SimdIndex)));
+				MoireErrorSize.SetElement(SimdIndex, ClampedEnergyAmount * (kMoireLumaToChannel * MoireError.GetElement(SimdIndex)));
 			}
 			#else
 			{
-				MoireErrorSize.SetElement(SimdIndex, MoireLumaToChannel * MoireError.GetElement(SimdIndex));
+				MoireErrorSize.SetElement(SimdIndex, kMoireLumaToChannel * MoireError.GetElement(SimdIndex));
 			}
 			#endif
 
@@ -574,7 +584,7 @@ void MainCS(
 			#endif
 
 			#if CONFIG_SCENE_COLOR_ALPHA
-				HistoryGuideOutput[OutputPixelPos] = FinalGuideColor;
+			HistoryGuideOutput[OutputPixelPos] = FinalGuideColor;
 			#else
 				HistoryGuideOutput[OutputPixelPos] = tsr_half4(FinalGuideColor.xyz, tsr_half(1.0));
 			#endif
diff --git a/Engine/Source/Runtime/Renderer/Private/DeferredShadingRenderer.cpp b/Engine/Source/Runtime/Renderer/Private/DeferredShadingRenderer.cpp
--- a/Engine/Source/Runtime/Renderer/Private/DeferredShadingRenderer.cpp
+++ b/Engine/Source/Runtime/Renderer/Private/DeferredShadingRenderer.cpp
@@ -3962,6 +3962,22 @@ void FDeferredShadingSceneRenderer::Render(FRDGBuilder& GraphBuilder)
 		ReconstructVolumetricRenderTarget(GraphBuilder, Views, SceneTextures.Depth.Resolve, HalfResolutionDepthCheckerboardMinMaxTexture, bAsyncComputeVolumetricCloud);
 	}
 
+	TArray<FScreenPassTexture, TInlineAllocator<4>> TSRMoireInputTextures;
+	// Extract TSR's moire heuristic luminance before renderering translucency into the scene color.
+	for (int32 ViewIndex = 0; ViewIndex < Views.Num(); ++ViewIndex)
+	{
+		FViewInfo& View = Views[ViewIndex];
+		if (NeedTSRMoireLuma(View))
+		{
+			if (TSRMoireInputTextures.Num() == 0)
+			{
+				TSRMoireInputTextures.SetNum(Views.Num());
+			}
+
+			TSRMoireInputTextures[ViewIndex] = AddTSRComputeMoireLuma(GraphBuilder, View.ShaderMap, FScreenPassTexture(SceneTextures.Color.Target, View.ViewRect));
+		}
+	}
+
 	const bool bShouldRenderTranslucency = !bHasRayTracedOverlay && ShouldRenderTranslucency();
 
 	// Union of all translucency view render flags.
@@ -4058,7 +4074,6 @@ void FDeferredShadingSceneRenderer::Render(FRDGBuilder& GraphBuilder)
 	}
 
 	// Draw translucency.
-	TArray<FScreenPassTexture> TSRMoireInputTextures;
 	if (!bHasRayTracedOverlay && TranslucencyViewsToRender != ETranslucencyView::None)
 	{
 		RDG_CSV_STAT_EXCLUSIVE_SCOPE(GraphBuilder, RenderTranslucency);
@@ -4088,21 +4103,6 @@ void FDeferredShadingSceneRenderer::Render(FRDGBuilder& GraphBuilder)
 			}
 		}
 
-		// Extract TSR's moire heuristic luminance before renderering translucency into the scene color.
-		for (int32 ViewIndex = 0; ViewIndex < Views.Num(); ++ViewIndex)
-		{
-			FViewInfo& View = Views[ViewIndex];
-			if (ITemporalUpscaler::GetMainTAAPassConfig(View) == EMainTAAPassConfig::TSR)
-			{
-				if (TSRMoireInputTextures.Num() == 0)
-				{
-					TSRMoireInputTextures.SetNum(Views.Num());
-				}
-
-				TSRMoireInputTextures[ViewIndex] = AddTSRComputeMoireLuma(GraphBuilder, View.ShaderMap, FScreenPassTexture(SceneTextures.Color.Target, View.ViewRect));
-			}
-		}
-
 		// Sort objects' triangles
 		for (FViewInfo& View : Views)
 		{
diff --git a/Engine/Source/Runtime/Renderer/Private/PostProcess/TemporalAA.h b/Engine/Source/Runtime/Renderer/Private/PostProcess/TemporalAA.h
--- a/Engine/Source/Runtime/Renderer/Private/PostProcess/TemporalAA.h
+++ b/Engine/Source/Runtime/Renderer/Private/PostProcess/TemporalAA.h
@@ -172,7 +172,11 @@ extern RENDERER_API FTAAOutputs AddTemporalAAPass(
 	const FTemporalAAHistory& InputHistory,
 	FTemporalAAHistory* OutputHistory);
 
-extern RENDERER_API FScreenPassTexture AddTSRComputeMoireLuma(FRDGBuilder& GraphBuilder, FGlobalShaderMap* ShaderMap, FScreenPassTexture SceneColor);
+/** Returns whether a given view need to measure luminance of the scene color for moire anti-flickering. */
+bool NeedTSRMoireLuma(const FViewInfo& View);
+
+/** Measure luminance of the scene color for moire anti-flickering. */
+FScreenPassTexture AddTSRComputeMoireLuma(FRDGBuilder& GraphBuilder, FGlobalShaderMap* ShaderMap, FScreenPassTexture SceneColor);
 
 /** Interface for the main temporal upscaling algorithm. */
 class RENDERER_API ITemporalUpscaler : public ISceneViewFamilyExtention
diff --git a/Engine/Source/Runtime/Renderer/Private/PostProcess/TemporalSuperResolution.cpp b/Engine/Source/Runtime/Renderer/Private/PostProcess/TemporalSuperResolution.cpp
--- a/Engine/Source/Runtime/Renderer/Private/PostProcess/TemporalSuperResolution.cpp
+++ b/Engine/Source/Runtime/Renderer/Private/PostProcess/TemporalSuperResolution.cpp
@@ -1063,6 +1063,11 @@ static FRDGTextureUAVRef CreateDummyUAV(FRDGBuilder& GraphBuilder, EPixelFormat
 	return GraphBuilder.CreateUAV(DummyTexture);
 };
 
+bool NeedTSRMoireLuma(const FViewInfo& View)
+{
+	return ITemporalUpscaler::GetMainTAAPassConfig(View) == EMainTAAPassConfig::TSR;
+}
+
 FScreenPassTexture AddTSRComputeMoireLuma(FRDGBuilder& GraphBuilder, FGlobalShaderMap* ShaderMap, FScreenPassTexture SceneColor)
 {
 	check(SceneColor.Texture)
