Fixed an issue where some Nanite meshes would appear corrupted because their vertex attributes were not decoded correctly

diff --git a/Engine/Shaders/Private/BitStreamReaderImplementation.ush b/Engine/Shaders/Private/BitStreamReaderImplementation.ush
--- a/Engine/Shaders/Private/BitStreamReaderImplementation.ush
+++ b/Engine/Shaders/Private/BitStreamReaderImplementation.ush
@@ -54,11 +54,16 @@ uint CONCAT(BitStreamReader_Read_, TYPE_SUFFIX)
 		// Bottom dword could be out of bits: Shift down.
 		State.BitOffsetFromAddress += State.BufferOffset;
 
-		State.BufferBits.x											= BitAlignU32(State.BufferBits.y,	State.BufferBits.x,	State.BufferOffset); // BufferOffset implicitly &31
-		if (State.CompileTimeMinBufferBits > 32) State.BufferBits.y	= BitAlignU32(State.BufferBits.z,	State.BufferBits.y,	State.BufferOffset); // BufferOffset implicitly &31
-		if (State.CompileTimeMinBufferBits > 64) State.BufferBits.z	= BitAlignU32(State.BufferBits.w,	State.BufferBits.z,	State.BufferOffset); // BufferOffset implicitly &31
-		if (State.CompileTimeMinBufferBits > 96) State.BufferBits.w	= BitAlignU32(0,					State.BufferBits.w, State.BufferOffset); // BufferOffset implicitly &31
-
+		// Workaround for BitAlignU32(x, y, 32) returning x instead of y.
+		// In the common case where State.CompileTimeMinDwordBits != 0, this will be optimized to just BitAlignU32.
+		// TODO: Can we get rid of this special case?
+		const bool bOffset32 = State.CompileTimeMinDwordBits == 0 && State.BufferOffset == 32;
+
+		State.BufferBits.x											= bOffset32 ? State.BufferBits.y :	BitAlignU32(State.BufferBits.y, State.BufferBits.x, State.BufferOffset);
+		if (State.CompileTimeMinBufferBits > 32) State.BufferBits.y	= bOffset32 ? State.BufferBits.z :	BitAlignU32(State.BufferBits.z, State.BufferBits.y, State.BufferOffset);
+		if (State.CompileTimeMinBufferBits > 64) State.BufferBits.z	= bOffset32 ? State.BufferBits.w :	BitAlignU32(State.BufferBits.w, State.BufferBits.z, State.BufferOffset);
+		if (State.CompileTimeMinBufferBits > 96) State.BufferBits.w	= bOffset32 ? 0u :					BitAlignU32(0,					State.BufferBits.w, State.BufferOffset);
+	
 		State.BufferOffset = 0;
 
 		State.CompileTimeMinDwordBits = min(32, State.CompileTimeMaxRemainingBits);
