Fixes shader compiler failures with r.PostProcessing.PropagateAlpha=1

diff --git a/Engine/Shaders/Private/TemporalSuperResolution/TSRRejectShading.usf b/Engine/Shaders/Private/TemporalSuperResolution/TSRRejectShading.usf
--- a/Engine/Shaders/Private/TemporalSuperResolution/TSRRejectShading.usf
+++ b/Engine/Shaders/Private/TemporalSuperResolution/TSRRejectShading.usf
@@ -316,8 +316,6 @@ void MainCS(
 			tsr_half LocalCurrentGradient = CurrentGradient.GetElement(SimdIndex);
 			tsr_half LocalInputLumaModification = InputLumaModification.GetElement(SimdIndex);
 			
-			Debug[0].SetElement(SimdIndex, tsr_halfC(LocalInputLumaModification, 0, 0));
-
 			CurrentGradient.SetElement(SimdIndex, clamp(LocalCurrentGradient * tsr_half(POSITIVE_INFINITY), tsr_half(-1.0), tsr_half(1.0)) * max(abs(LocalCurrentGradient) - LocalInputLumaModification, tsr_half(0.0)));
 		}
 
