Fix build error compiling MobileDeferredShading.

- Using wrong local variable: `ScreenPosition` instead of `ScreenPos`.

diff --git a/Engine/Shaders/Private/MobileDeferredShading.usf b/Engine/Shaders/Private/MobileDeferredShading.usf
--- a/Engine/Shaders/Private/MobileDeferredShading.usf
+++ b/Engine/Shaders/Private/MobileDeferredShading.usf
@@ -153,7 +153,7 @@ void MobileDirectionalLightPS(
 
 	FGBufferData GBuffer = MobileFetchAndDecodeGBuffer(UVAndScreenPos.xy, UVAndScreenPos.zw);
 	float2 ScreenPos = UVAndScreenPos.zw;
-	float3 TranslatedWorldPosition = mul(float4(GetScreenPositionForProjectionType(ScreenPosition, GBuffer.Depth), GBuffer.Depth, 1), PrimaryView.ScreenToTranslatedWorld).xyz;
+	float3 TranslatedWorldPosition = mul(float4(GetScreenPositionForProjectionType(ScreenPos, GBuffer.Depth), GBuffer.Depth, 1), PrimaryView.ScreenToTranslatedWorld).xyz;
 
 	half3 CameraVector = normalize(TranslatedWorldPosition);
 	half3 V = -CameraVector;
