Fixed Landscape Edge Normal calculation when using Global Merge.

It now flags valid regions during layer combination using the alpha channel, and masks out normal calculations that used invalid samples, removing their influence from the normal map results.

diff --git a/Engine/Shaders/Private/Landscape/LandscapeCommon.ush b/Engine/Shaders/Private/Landscape/LandscapeCommon.ush
--- a/Engine/Shaders/Private/Landscape/LandscapeCommon.ush
+++ b/Engine/Shaders/Private/Landscape/LandscapeCommon.ush
@@ -17,4 +17,14 @@ float UnpackHeight(float2 InPackedHeight)
 	// However, this isn't great if you just interpolate between packed heights because that truncates the interpolation
 	// of the higher order byte. So instead we do it by multiplication, which ends up simpler.
 	return (InPackedHeight.r * 256 + InPackedHeight.g) * 255;
-}
\ No newline at end of file
+}
+
+// Compute the normal of the triangle formed by the 3 points (in winding order).
+// .xyz of each point is its world space position
+// Additionally, the .w component of each point indicates its validity (valid: .w == 1, invalid: .w == 0) 
+// the resulting normal will be 0 if any of the points is invalid
+float3 ComputeNullableTriangleNormal(float4 InPoint0, float4 InPoint1, float4 InPoint2)
+{
+	float3 Normal = normalize(cross(InPoint0.xyz - InPoint1.xyz, InPoint1.xyz - InPoint2.xyz));
+	return Normal * InPoint0.w * InPoint1.w * InPoint2.w;
+}
diff --git a/Engine/Shaders/Private/Landscape/LandscapeLayersHeightmapsPS.usf b/Engine/Shaders/Private/Landscape/LandscapeLayersHeightmapsPS.usf
--- a/Engine/Shaders/Private/Landscape/LandscapeLayersHeightmapsPS.usf
+++ b/Engine/Shaders/Private/Landscape/LandscapeLayersHeightmapsPS.usf
@@ -220,13 +220,6 @@ float4 LoadNeighborHeightmapScaledPosition(bool2 bInIsMinComponentBorder, bool2
 	float2 NeighborHeight = LoadNeighborHeightmapUnpackedHeight(bInIsMinComponentBorder, bInIsMaxComponentBorder, InRelativeTextureCoordinates, InNeighborOffset, bSkipDuplicateRowColumn);
 	return float4(((int2) InRelativeTextureCoordinates + InNeighborOffset) * InLandscapeGridScale.xy, ScaleHeight(NeighborHeight.x), NeighborHeight.y);
 }
-
-// Compute the normal of the triangle formed by the 3 points (in winding order). Additionally, the .w component of each point indicates its validity and the result will be 0 if either of the points is invalid :
-float3 ComputeNullableTriangleNormal(float4 InPoint0, float4 InPoint1, float4 InPoint2)
-{
-	float3 Normal = normalize(cross(InPoint0.xyz - InPoint1.xyz, InPoint1.xyz - InPoint2.xyz));
-	return Normal * InPoint0.w * InPoint1.w * InPoint2.w;
-}
 #endif // FINALIZE_HEIGHTMAP
 
 
diff --git a/Engine/Shaders/Private/LandscapeLayersPS.usf b/Engine/Shaders/Private/LandscapeLayersPS.usf
--- a/Engine/Shaders/Private/LandscapeLayersPS.usf
+++ b/Engine/Shaders/Private/LandscapeLayersPS.usf
@@ -55,7 +55,7 @@ void PSHeightmapMain(float2 InTextureCoordinates : TEXCOORD0, out float4 OutColo
 		float LayerBlendMode = LayerInfo.z;
 		if (LayerBlendMode == 1 && LayerVisibility == 1)
 		{
-			// Alpha Blend Mode
+			// Pre-multiplied Alpha Blend Mode
 			float HeightAlphaBlend = ExtractAlpha(SourceColor.ba);
 			float NewHeight = Height + (HeightAlphaBlend * BaseHeight);
 			uint RaiseLower = floor(SourceColor.a * 255.f + 0.5f);
@@ -75,27 +75,43 @@ void PSHeightmapMain(float2 InTextureCoordinates : TEXCOORD0, out float4 OutColo
 			// Additive Blend Mode
 			Height += BaseHeight;
 		}
+
+		// clamp to valid range
 		Height = clamp(Height, -32768.0, 32767.0);
 	}
 
+	// convert back to positive range
 	Height += 32768.0f;
 
 	float2 PackedHeight = PackHeight(Height);
 
+	// prepare the output color
 	OutColor.r = PackedHeight.x;
 	OutColor.g = PackedHeight.y;
-	OutColor.b = SourceColor.b;
-	OutColor.a = SourceColor.a;
+
+	// once normals are calculated, we need to pass through the normal data when copying
+	bool bPassthroughHeightmap1BlueAlpha = (OutputConfig.y == 1.0);
+	if (bPassthroughHeightmap1BlueAlpha)
+	{
+		OutColor.b = SourceColor.b;
+		OutColor.a = SourceColor.a;
+	}
+	else
+	{
+		// otherwise we can use alpha to indicate valid pixels (needed by the normal calculation)
+		OutColor.b = 0.0f;
+		OutColor.a = 1.0f;
+	}
 
 	// Output normals
 	if (OutputConfig.w == 1.0)
 	{
 		float2 TexelSize = 1.0 / TextureSize;
 
-		bool bIsMinBorderTexelX = InTextureCoordinates.x <= TexelSize.x;
-		bool bIsMinBorderTexelY = InTextureCoordinates.y <= TexelSize.y;
-		bool bIsMaxBorderTexelX = InTextureCoordinates.x >= 1.0f - TexelSize.x;
-		bool bIsMaxBorderTexelY = InTextureCoordinates.y >= 1.0f - TexelSize.y;
+		float MaskMinBorderX = (InTextureCoordinates.x <= TexelSize.x) ? 0.0f : 1.0f;
+		float MaskMinBorderY = (InTextureCoordinates.y <= TexelSize.y) ? 0.0f : 1.0f;
+		float MaskMaxBorderX = (InTextureCoordinates.x >= 1.0f - TexelSize.x) ? 0.0f : 1.0f;
+		float MaskMaxBorderY = (InTextureCoordinates.y >= 1.0f - TexelSize.y) ? 0.0f : 1.0f;
 
 		// The triangle topology is the following (where C = center, T = top, B = bottom, L = left, R = right and Nx the normals we need to interpolate):
 		// TL ------ TT
@@ -117,45 +133,42 @@ void PSHeightmapMain(float2 InTextureCoordinates : TEXCOORD0, out float4 OutColo
 		//           BB ------ BR
 		// So we can retrieve all 7 samples using 4 gathers : 
 		// Make sure we gather the 4 pixels in the bottom right direction by sampling at (0.75, 0.75) of the pixel (gather retrieves the 4 samples that would be used for bilinear interpolation):
+
 		float2 GatherLocation = InTextureCoordinates + 0.25f * TexelSize; 
 		float4 Red0 = ReadTexture1.GatherRed(ReadTexture1Sampler, GatherLocation - TexelSize);
 		float4 Green0 = ReadTexture1.GatherGreen(ReadTexture1Sampler, GatherLocation - TexelSize);
+		float4 Mask0 = ReadTexture1.GatherAlpha(ReadTexture1Sampler, GatherLocation - TexelSize);
 		float4 Red1 = ReadTexture1.GatherRed(ReadTexture1Sampler, GatherLocation);
 		float4 Green1 = ReadTexture1.GatherGreen(ReadTexture1Sampler, GatherLocation);
-
-		float3 TL = float3((InTextureCoordinates * TextureSize + float2(-1, -1)) * LandscapeGridScale.xy, ExtractAndScaleHeight(float2(Red0.w, Green0.w)));
-		float3 TT = float3((InTextureCoordinates * TextureSize + float2(+0, -1)) * LandscapeGridScale.xy, ExtractAndScaleHeight(float2(Red0.z, Green0.z)));
-		float3 CC = float3((InTextureCoordinates * TextureSize + float2(+0, +0)) * LandscapeGridScale.xy, ExtractAndScaleHeight(float2(Red0.y, Green0.y)));
-		float3 LL = float3((InTextureCoordinates * TextureSize + float2(-1, +0)) * LandscapeGridScale.xy, ExtractAndScaleHeight(float2(Red0.x, Green0.x)));
-		float3 RR = float3((InTextureCoordinates * TextureSize + float2(+1, +0)) * LandscapeGridScale.xy, ExtractAndScaleHeight(float2(Red1.z, Green1.z)));
-		float3 BR = float3((InTextureCoordinates * TextureSize + float2(+1, +1)) * LandscapeGridScale.xy, ExtractAndScaleHeight(float2(Red1.y, Green1.y)));
-		float3 BB = float3((InTextureCoordinates * TextureSize + float2(+0, +1)) * LandscapeGridScale.xy, ExtractAndScaleHeight(float2(Red1.x, Green1.x)));
-
-		float3 N0 = normalize(cross(CC - LL, LL - TL));
-		float3 N1 = normalize(cross(TL - TT, TT - CC));
-		float3 N2 = normalize(cross(LL - CC, CC - BB));
-		float3 N3 = normalize(cross(RR - CC, CC - TT));
-		float3 N4 = normalize(cross(BR - BB, BB - CC));
-		float3 N5 = normalize(cross(CC - RR, RR - BR));
-		
-		// Handle edge cases by simply cancelling the normal so that it doesn't contribute : 
-		if (bIsMinBorderTexelX)
-		{
-			N0 = N1 = N2 = 0.0f;
-		}
-		if (bIsMinBorderTexelY)
-		{
-			N0 = N1 = N3 = 0.0f;
-		}
-		if (bIsMaxBorderTexelX)
-		{
-			N3 = N4 = N5 = 0.0f;
-		}
-		if (bIsMaxBorderTexelY)
-		{
-			N2 = N4 = N5 = 0.0f;
-		}
-		
+		float4 Mask1 = ReadTexture1.GatherAlpha(ReadTexture1Sampler, GatherLocation);
+
+		// [x,y,z] = World Position [w] = valid mask
+		float4 TL = float4((InTextureCoordinates * TextureSize + float2(-1, -1)) * LandscapeGridScale.xy, ExtractAndScaleHeight(float2(Red0.w, Green0.w)), Mask0.w);
+		float4 TT = float4((InTextureCoordinates * TextureSize + float2(+0, -1)) * LandscapeGridScale.xy, ExtractAndScaleHeight(float2(Red0.z, Green0.z)), Mask0.z);
+		float4 CC = float4((InTextureCoordinates * TextureSize + float2(+0, +0)) * LandscapeGridScale.xy, ExtractAndScaleHeight(float2(Red0.y, Green0.y)), Mask0.y);
+		float4 LL = float4((InTextureCoordinates * TextureSize + float2(-1, +0)) * LandscapeGridScale.xy, ExtractAndScaleHeight(float2(Red0.x, Green0.x)), Mask0.x);
+		float4 RR = float4((InTextureCoordinates * TextureSize + float2(+1, +0)) * LandscapeGridScale.xy, ExtractAndScaleHeight(float2(Red1.z, Green1.z)), Mask1.z);
+		float4 BR = float4((InTextureCoordinates * TextureSize + float2(+1, +1)) * LandscapeGridScale.xy, ExtractAndScaleHeight(float2(Red1.y, Green1.y)), Mask1.y);
+		float4 BB = float4((InTextureCoordinates * TextureSize + float2(+0, +1)) * LandscapeGridScale.xy, ExtractAndScaleHeight(float2(Red1.x, Green1.x)), Mask1.x);
+
+		// mask out samples that are off the edge of the texture
+		TL.w *= MaskMinBorderX * MaskMinBorderY;
+		TT.w *=                  MaskMinBorderY;
+		// CC should never be off the edge
+		LL.w *= MaskMinBorderX;
+		RR.w *= MaskMaxBorderX;
+		BR.w *= MaskMaxBorderX * MaskMaxBorderY;
+		BB.w *=                  MaskMaxBorderY;
+
+		// compute the normals (zeroing out those that use masked points, so they are ignored)
+		float3 N0 = ComputeNullableTriangleNormal(CC, LL, TL);
+		float3 N1 = ComputeNullableTriangleNormal(TL, TT, CC);
+		float3 N2 = ComputeNullableTriangleNormal(LL, CC, BB);
+		float3 N3 = ComputeNullableTriangleNormal(RR, CC, TT);
+		float3 N4 = ComputeNullableTriangleNormal(BR, BB, CC);
+		float3 N5 = ComputeNullableTriangleNormal(CC, RR, BR);
+
+		// average the normals
 		float3 FinalNormal = normalize(N0 + N1 + N2 + N3 + N4 + N5);
 
 		// Scale back to be 0-1 normal
diff --git a/Engine/Source/Runtime/Landscape/Private/LandscapeEditLayers.cpp b/Engine/Source/Runtime/Landscape/Private/LandscapeEditLayers.cpp
--- a/Engine/Source/Runtime/Landscape/Private/LandscapeEditLayers.cpp
+++ b/Engine/Source/Runtime/Landscape/Private/LandscapeEditLayers.cpp
@@ -474,6 +474,7 @@ struct FLandscapeLayersHeightmapShaderParameters
 		, ReadHeightmap2(nullptr)
 		, HeightmapSize(0, 0)
 		, ApplyLayerModifiers(false)
+		, PassthroughHeightmap1BlueAlpha(true)
 		, LayerAlpha(1.0f)
 		, LayerVisible(true)
 		, LayerBlendMode(LSBM_AdditiveBlend)
@@ -488,6 +489,7 @@ struct FLandscapeLayersHeightmapShaderParameters
 	UTexture* ReadHeightmap2;
 	FIntPoint HeightmapSize;
 	bool ApplyLayerModifiers;
+	bool PassthroughHeightmap1BlueAlpha;
 	float LayerAlpha;
 	bool LayerVisible;
 	ELandscapeBlendMode LayerBlendMode;
@@ -536,7 +538,7 @@ public:
 		SetTextureParameter(BatchedParameters, ReadTexture2Param, ReadTexture2SamplerParam, TStaticSamplerState<SF_Point, AM_Clamp, AM_Clamp, AM_Clamp>::GetRHI(), InParams.ReadHeightmap2 != nullptr ? InParams.ReadHeightmap2->GetResource()->TextureRHI : GWhiteTexture->TextureRHI);
 
 		FVector4f LayerInfo(InParams.LayerAlpha, InParams.LayerVisible ? 1.0f : 0.0f, InParams.LayerBlendMode == LSBM_AlphaBlend ? 1.0f : 0.f, 0.f);
-		FVector4f OutputConfig(InParams.ApplyLayerModifiers ? 1.0f : 0.0f, 0.0f /*unused*/, InParams.ReadHeightmap2 ? 1.0f : 0.0f, InParams.GenerateNormals ? 1.0f : 0.0f);
+		FVector4f OutputConfig(InParams.ApplyLayerModifiers ? 1.0f : 0.0f, InParams.PassthroughHeightmap1BlueAlpha ? 1.0f : 0.0f, InParams.ReadHeightmap2 ? 1.0f : 0.0f, InParams.GenerateNormals ? 1.0f : 0.0f);
 		FVector2f TextureSize(static_cast<float>(InParams.HeightmapSize.X), static_cast<float>(InParams.HeightmapSize.Y));
 
 		SetShaderValue(BatchedParameters, LayerInfoParam, LayerInfo);
@@ -4851,8 +4853,9 @@ int32 ALandscape::PerformLayersHeightmapsGlobalMerge(const FUpdateLayersContentC
 
 	for (FLandscapeLayer& Layer : LandscapeLayers)
 	{
-		//Draw Layer heightmap to Combined RT Atlas
+		// Draw each Layer's heightmaps to a Combined RT Atlas in LandscapeScratchRT1
 		ShaderParams.ApplyLayerModifiers = false;
+		ShaderParams.PassthroughHeightmap1BlueAlpha = false; // before we calculate normals, don't passthrough (instead mark valid using alpha)
 		ShaderParams.LayerVisible = Layer.bVisible;
 		ShaderParams.GenerateNormals = false;
 		ShaderParams.LayerBlendMode = Layer.BlendMode;
@@ -4886,13 +4889,14 @@ int32 ALandscape::PerformLayersHeightmapsGlobalMerge(const FUpdateLayersContentC
 			ExecuteCopyLayersTexture(MoveTemp(DeferredCopyTextures));
 		}
 
-		// NOTE: From this point on, we always work in non atlas, we'll convert back at the end to atlas only
+		// Convert Atlas LandscapeScratchRT1 to the world-projected NonAtlas in LandscapeScratchRT2
+		// TODO: just use this format from the beginning above...
 		DrawHeightmapComponentsToRenderTarget(FString::Printf(TEXT("LS Height: %s += -> NonAtlas %s"), *Layer.Name.ToString(), *LandscapeScratchRT1->GetName(), *LandscapeScratchRT2->GetName()),
 			InUpdateLayersContentContext.LandscapeComponentsHeightmapsToRender, LandscapeExtent.Min, LandscapeScratchRT1, nullptr, LandscapeScratchRT2, ERTDrawingType::RTAtlasToNonAtlas, true, ShaderParams);
 
 		ShaderParams.ApplyLayerModifiers = true;
 
-		// Combine Current layer with current result
+		// Combine Current layer NonAtlas LandscapeScratchRT2 with current result in LandscapeScratchRT3, writing final result to CombinedHeightmapNonAtlasRT
 		DrawHeightmapComponentsToRenderTarget(FString::Printf(TEXT("LS Height: %s += -> CombinedNonAtlas %s"), *Layer.Name.ToString(), *LandscapeScratchRT2->GetName(), *CombinedHeightmapNonAtlasRT->GetName()),
 			InUpdateLayersContentContext.LandscapeComponentsHeightmapsToRender, LandscapeExtent.Min, LandscapeScratchRT2, FirstLayer ? nullptr : LandscapeScratchRT3, CombinedHeightmapNonAtlasRT, ERTDrawingType::RTNonAtlas, FirstLayer, ShaderParams);
 
@@ -4933,6 +4937,8 @@ int32 ALandscape::PerformLayersHeightmapsGlobalMerge(const FUpdateLayersContentC
 			}
 		}
 
+		// copy CombinedHeightmapNonAtlasRT to LandscapeScratchRT3 (as a source for later layers... this is wasted on the last layer I think...)
+		// TODO: you can get the same effect for much cheaper by swapping these two pointers before the render above...
 		ExecuteCopyLayersTexture({ FLandscapeLayersCopyTextureParams(CombinedHeightmapNonAtlasRT, LandscapeScratchRT3) });
 		PrintLayersDebugRT(FString::Printf(TEXT("LS Height: %s Component %s += -> CombinedNonAtlas %s"), *Layer.Name.ToString(), *CombinedHeightmapNonAtlasRT->GetName(), *LandscapeScratchRT3->GetName()), LandscapeScratchRT3);
 
@@ -4948,14 +4954,18 @@ int32 ALandscape::PerformLayersHeightmapsGlobalMerge(const FUpdateLayersContentC
 		LandscapeFullHeightmapRenderDoneDelegate.Broadcast(LandscapeScratchRT3);
 	}
 
+	// compute Normals into LandscapeScratchRT1
 	DrawHeightmapComponentsToRenderTarget(FString::Printf(TEXT("LS Height: %s = -> CombinedNonAtlasNormals : %s"), *CombinedHeightmapNonAtlasRT->GetName(), *LandscapeScratchRT1->GetName()),
 		InUpdateLayersContentContext.LandscapeComponentsHeightmapsToRender, LandscapeExtent.Min, CombinedHeightmapNonAtlasRT, nullptr, LandscapeScratchRT1, ERTDrawingType::RTNonAtlas, true, ShaderParams);
 
 	ShaderParams.GenerateNormals = false;
+	ShaderParams.PassthroughHeightmap1BlueAlpha = true;		// now that we've generated normals, need to pass through the blue and alpha channels to preserve them
 
+	// convert back to atlas (TODO: we could do this on the first mip downsample instead...)
 	DrawHeightmapComponentsToRenderTarget(FString::Printf(TEXT("LS Height: %s = -> CombinedAtlasFinal : %s"), *LandscapeScratchRT1->GetName(), *CombinedHeightmapAtlasRT->GetName()),
 		InUpdateLayersContentContext.LandscapeComponentsHeightmapsToRender, LandscapeExtent.Min, LandscapeScratchRT1, nullptr, CombinedHeightmapAtlasRT, ERTDrawingType::RTNonAtlasToAtlas, true, ShaderParams);
 
+	// Downsample to generate mips...
 	DrawHeightmapComponentsToRenderTargetMips(InUpdateLayersContentContext.LandscapeComponentsHeightmapsToRender, LandscapeExtent.Min, CombinedHeightmapAtlasRT, true, ShaderParams);
 
 	// List of UTexture2D that we need to kick off readbacks for :
