Refactors the tonemapper's code to be stateless arround which exposure affecting terms is applied

We often have a PreExposure that make the linear scene color a lot darker than they should compared to what is displayed in the back buffer.
That is becasue PreExposure on CPU just do PreExposure =  last available GlobalExposure; but tonemaper is doing a lot more than that.
This problems causes compression of colors on fewer bits into different very low bitdepth resources.
This has consequence for TSR to loose precision and starting ghosting on content that should otherwise just work fine.
This create undesirable TSR support and give the wrong impression of the tech.

The goal of this change is to make the code of the tonemapper a lot more readable and maintainable and accessible for anyone to understand what affects the final exposure of the frame.
This will then allow to make the computation of the PreExposure on the CPU a lot more robust to user mistakes in post process volumes to match exactly what is now easily readable in the tonemaper's code.

This change is purposefully .usf only to verify without editor restart that certain post process settings do not change in look either in desktop or mobile renderer.
Tests carried on this change for both desktop and mobile renderer:
* vignette post process settings
* both gaussian, FFT and mobile bloom post process settings
* dirt mask post process settings
* film grain post process settings
* r.Tonemaper.Sharpen
* showflags using GammaOnly
* EngineTest's for post processing

To simplify even further the maintenance of the tonemapper, gamma only should be its own shader instead of permutation, but deferred in followup change in favors of minimizing risks of this change.

diff --git a/Engine/Shaders/Private/PostProcessTonemap.usf b/Engine/Shaders/Private/PostProcessTonemap.usf
--- a/Engine/Shaders/Private/PostProcessTonemap.usf
+++ b/Engine/Shaders/Private/PostProcessTonemap.usf
@@ -240,98 +240,11 @@ void MainVS(
     TonemapCommonVS(OutPosition, OutTexCoord, OutVignette, OutGrainUV, OutScreenPos, OutFullViewUV);
 }
 
-// Function graphing
-float F0( float x )
-{
-	return x*saturate( (x - 0.5)/2 );
-}
-
-float F1( float x )
-{
-	return x;
-}
-
-float F2( float x )
-{
-	return x;
-}
-
-float F3( float x )
-{
-	return x;
-}
-
-float LineShade( float fx, float y, float dydx, float LineWidth )
-{
-	return 1 - smoothstep( 0.5 * LineWidth, LineWidth, abs( fx - y ) / sqrt( 1 + Square( dydx ) ) );
-}
-
-float3 Graph( float2 ScreenSpacePos )
-{
-	float2 WindowMin = float2( 0, 0 );
-	float2 WindowMax = float2( 1, 1 );
-	
-	float2 p = ( (ScreenSpacePos + 1) * 0.5 - WindowMin ) * ( WindowMax - WindowMin );
-	float LineWidth = dot( WindowMax - WindowMin, 0.0005 );
-	
-	float3 Color;
-	Color  = float3( 1, 0, 0 ) * LineShade( F0(p.x), p.y, ( F0(p.x + LineWidth) - F0(p.x - LineWidth) ) / (2 * LineWidth), LineWidth );
-	Color += float3( 0, 1, 0 ) * LineShade( F1(p.x), p.y, ( F1(p.x + LineWidth) - F1(p.x - LineWidth) ) / (2 * LineWidth), LineWidth );
-	Color += float3( 0, 0, 1 ) * LineShade( F2(p.x), p.y, ( F2(p.x + LineWidth) - F2(p.x - LineWidth) ) / (2 * LineWidth), LineWidth );
-	Color += float3( 1, 1, 0 ) * LineShade( F3(p.x), p.y, ( F3(p.x + LineWidth) - F3(p.x - LineWidth) ) / (2 * LineWidth), LineWidth );
-	//Color += float3( 0, 1, 1 ) * LineShade( F4(p.x), p.y, ( F4(p.x + LineWidth) - F4(p.x - LineWidth) ) / (2 * LineWidth), LineWidth );
-	//Color += float3( 1, 1, 1 ) * LineShade( F5(p.x), p.y, ( F5(p.x + LineWidth) - F5(p.x - LineWidth) ) / (2 * LineWidth), LineWidth );
-	return Color;
-}
-
-// @return color in SRGB
-float3 SimpleToneMap(float3 HDRColor)
-{
-	// from http://filmicgames.com/archives/75
-//	float3 x = max(0, HDRColor - 0.004f);	return (x * (6.2f * x + 0.5f)) / (x * (6.2f * x + 1.7f) + 0.06f);
-	
-	// linear/no tonemapper 
-	return HDRColor;
-}
-
-float max3(float3 In)
-{
-	return max(In.x, max(In.y, In.z));
-}
-
-// @return 0 at low, 1 at high and linear interpolated inbetween
-float RemapScalar(float Low, float High, float x)
-{
-	return (x - Low) / (High - Low);
-}
-
 float max4(float4 x)
 {
 	return max(max(x.r, x.g), max(x.b, x.a));
 }
 
-// useful debug helper (not optimized), could be moved into a more common place
-// @param DebugTile 0,0 is center, -1,0 is left, 1,0 is right, ...
-float2 ComputeCounterStepForTileDebugging(float4 SvPosition, int2 TileSize, out int2 DebugTile, out int2 LocalDebugTilePosition)
-{
-	// ViewportRect is defined for postprocessing passes
-	float2 CenterPos = (Output_ViewportMin + Output_ViewportMax) / 2.0f;
-	float2 LeftTopPos = CenterPos - TileSize / 2;
-
-	float2 LocalPos = SvPosition.xy - LeftTopPos;
-
-	DebugTile = floor(LocalPos / TileSize);
-
-	LocalDebugTilePosition = LocalPos - DebugTile * TileSize;
-
-	float2 CounterStep = -DebugTile * TileSize;
-
-	return CounterStep;
-}
-// useful debug helper (works with ComputeCounterStepForTileDebugging()), could be moved into a more common place
-#define OffsetInterpolator(CounterStep, InterpolatorName) \
-		InterpolatorName += ddx(InterpolatorName) * CounterStep.x + ddy(InterpolatorName) * CounterStep.y;
-
 float4 SampleSceneColor(float2 SceneUV)
 {
 	SceneUV = clamp(SceneUV, Color_UVViewportBilinearMin, Color_UVViewportBilinearMax);
@@ -354,118 +267,131 @@ float4 TonemapCommonPS(
 	out float OutLuminance
 	)
 {
-	float4 OutColor = 0;
-	OutLuminance = 0;
-	
-
-#if DEBUG_TONEMAPPER
-	// @param DebugTile 0,0 is center, -1,0 is left, 1,0 is right, ...
-	int2 DebugTile;
-	int2 LocalDebugTilePosition;
-	bool bOutputDebugTile = false;
-	{
-		// split the screen in a 3x3 layout and add some border (0.3)
-		int2 TileSize = floor(Output_ViewportSize / float2(3.3f, 3.3f));
-		bool bBorder;
+	const float OneOverPreExposure = View.OneOverPreExposure;
+	const float2 EyeAdaptationData = EyeAdaptationBuffer[0].xw;
+	const float GlobalExposure = EyeAdaptationData.x;
 
-		float2 CounterStep = ComputeCounterStepForTileDebugging(SvPosition, TileSize, DebugTile, LocalDebugTilePosition);
+	float2 SceneUV = UV.xy;
 
-		OffsetInterpolator(CounterStep, UV);
-		OffsetInterpolator(CounterStep, GrainUV);
-		OffsetInterpolator(CounterStep, ScreenPos);
-		OffsetInterpolator(CounterStep, FullViewUV);
-	}
-	if (all(DebugTile == int2(0, 0)))
+	// Sample the scene color.
+	half4 SceneColor;
 	{
-		// center is unaltered tonemapper output
-		bOutputDebugTile = true;
-	}
-#endif
+		#if USE_COLOR_FRINGE
+			float2 SceneUVJitter = float2(0.0, 0.0);
+		
+				float2 CAScale = ChromaticAberrationParams.rg;
+				float StartOffset = ChromaticAberrationParams.z;
 
-	float2 SceneUV = UV.xy;
+				float2 LensUV = ConvertScreenViewportSpaceToLensViewportSpace(ScreenPos);
 
-	#if USE_COLOR_FRINGE
-		float2 SceneUVJitter = float2(0.0, 0.0);
-		
-	        float2 CAScale = ChromaticAberrationParams.rg;
-	        float StartOffset = ChromaticAberrationParams.z;
+				float4 CAUV;
+				CAUV = LensUV.xyxy - sign(LensUV).xyxy * saturate(abs(LensUV) - StartOffset).xyxy * CAScale.rrgg;
 
-			float2 LensUV = ConvertScreenViewportSpaceToLensViewportSpace(ScreenPos);
+				CAUV.xy = ConvertLensViewportSpaceToScreenViewportSpace(CAUV.xy);
+				CAUV.zw = ConvertLensViewportSpaceToScreenViewportSpace(CAUV.zw);
 
-	        float4 CAUV;
-	        CAUV = LensUV.xyxy - sign(LensUV).xyxy * saturate(abs(LensUV) - StartOffset).xyxy * CAScale.rrgg;
+				CAUV.xy = ScreenPosToUV(CAUV.xy, Color_ExtentInverse);
+				CAUV.zw = ScreenPosToUV(CAUV.zw, Color_ExtentInverse);
 
-	        CAUV.xy = ConvertLensViewportSpaceToScreenViewportSpace(CAUV.xy);
-	        CAUV.zw = ConvertLensViewportSpaceToScreenViewportSpace(CAUV.zw);
+			half4 SceneColorR = SampleSceneColor(CAUV.xy + SceneUVJitter.xy);
+			half SceneColorG = SampleSceneColor(CAUV.zw + SceneUVJitter.xy).g;
+			half SceneColorB = SampleSceneColor(SceneUV).b;
+			SceneColor = SceneColorR;
+			SceneColor.g = SceneColorG;
+			SceneColor.b = SceneColorB;
+		#else
+			SceneColor = SampleSceneColor(SceneUV);
+		#endif
 
-			CAUV.xy = ScreenPosToUV(CAUV.xy, Color_ExtentInverse);
-	        CAUV.zw = ScreenPosToUV(CAUV.zw, Color_ExtentInverse);
+		#if METAL_MSAA_HDR_DECODE && !USE_GAMMA_ONLY
+			// look for PreTonemapMSAA 
+			SceneColor.rgb *= rcp(SceneColor.r * (-0.299) + SceneColor.g * (-0.587) + SceneColor.b * (-0.114) + 1.0);
+			// Try to kill negatives and NaNs here
+			SceneColor.rgb = max(SceneColor.rgb, 0);
+		#endif
+	}
 
-		half4 SceneColor = SampleSceneColor(CAUV.xy + SceneUVJitter.xy);
-		half SceneColorG = SampleSceneColor(CAUV.zw + SceneUVJitter.xy).g;
-		half SceneColorB = SampleSceneColor(SceneUV).b;
-		SceneColor.g = SceneColorG;
-		SceneColor.b = SceneColorB;
+	// Sample bloom & mobile's vignette.
+	#if USE_BLOOM && FEATURE_LEVEL == FEATURE_LEVEL_ES3_1
+		float MobileVignette;
 	#else
-		half4 SceneColor = SampleSceneColor(SceneUV);
+		const float MobileVignette = 1.0;
 	#endif
+	#if USE_BLOOM
+	half3 Bloom;
+	{
+		float2 BloomUV;
+		#if ES3_1_PROFILE
+		{
+			BloomUV = FullViewUV.xy;
+		}
+		#else
+		{
+			BloomUV = ApplyScreenTransform(UV, ColorToBloom);
+			BloomUV = clamp(BloomUV, BloomUVViewportBilinearMin, BloomUVViewportBilinearMax);
+		}
+		#endif
 
-#if METAL_MSAA_HDR_DECODE && !USE_GAMMA_ONLY
-	// look for PreTonemapMSAA 
-	SceneColor.rgb *= rcp(SceneColor.r * (-0.299) + SceneColor.g * (-0.587) + SceneColor.b * (-0.114) + 1.0);
-	// Try to kill negatives and NaNs here
-	SceneColor.rgb = max(SceneColor.rgb, 0);
-#endif
-
-	SceneColor.rgb *= View.OneOverPreExposure;
+		float4 RawBloom = Texture2DSample(BloomTexture, BloomSampler, BloomUV);
 
-	float2 EyeAdaptationData = EyeAdaptationBuffer[0].xw;
+		#if FEATURE_LEVEL == FEATURE_LEVEL_ES3_1
+		{
+			// Support sunshaft and vignette for mobile, and we have applied the BloomIntensity and the BloomDirtMask at the sun merge pass.
+			Bloom = RawBloom.rgb;
+			MobileVignette = RawBloom.a;
+		}
+		#else
+		{
+			float2 DirtLensUV = ConvertScreenViewportSpaceToLensViewportSpace(ScreenPos) * float2(1.0f, -1.0f);
 
-	float ExposureScale = EyeAdaptationData.x;
+			float3 RawBloomDirt = Texture2DSample(BloomDirtMaskTexture, BloomDirtMaskSampler, DirtLensUV * .5f + .5f).rgb;
 
-	float LocalExposureOneOverPreExposure = View.OneOverPreExposure;
+			Bloom = RawBloom.rgb * (1.0 + RawBloomDirt * BloomDirtMaskTint.rgb);
+		}
+		#endif
+	}
+	#endif // !USE_BLOOM
+	
+	// Compute vignette.
+	#if USE_VIGNETTE
+		const float VignetteMask = MobileVignette * ComputeVignetteMask(Vignette, TonemapperParams.x);
+	#else
+		const float VignetteMask = MobileVignette;
+	#endif
+	
+	// Compute scene color tint.
+	#if USE_GAMMA_ONLY
+		const float3 SceneColorTint = 1.0;
+	#elif SUPPORTS_SCENE_COLOR_APPLY_PARAMETERS
+		// Note that SceneColorApplyParamaters[0] is only for FFT bloom energy conservativeness = doesn't need to be in PreExposure
+		const float3 SceneColorTint = ColorScale0 * SceneColorApplyParamaters[0].xyz;
+	#else
+		const float3 SceneColorTint = ColorScale0;
+	#endif
 
-#if USE_LOCAL_EXPOSURE
+	// Compute the local exposure based of the scene color.
+	#if USE_LOCAL_EXPOSURE
+	float LocalExposure = 1.0;
 	{
-		float LuminanceVal = CalculateEyeAdaptationLuminance(SceneColor.rgb);
+		float LuminanceVal = CalculateEyeAdaptationLuminance(SceneColor.rgb * OneOverPreExposure);
 		float LogLuminance = log2(LuminanceVal);
 
 		// Middle grey lum value adjusted by exposure compensation
 		float MiddleGreyLumValue = log2(0.18 * EyeAdaptationData.y * EyeAdaptation_LocalExposureMiddleGreyExposureCompensation);
 
-		float BaseLogLum = CalculateBaseLogLuminance(LogLuminance, EyeAdaptation_LocalExposureBlurredLuminanceBlend, ExposureScale, FullViewUV, LumBilateralGrid, BlurredLogLum, LumBilateralGridSampler, BlurredLogLumSampler);
-		float LocalExposure = CalculateLocalExposure(LogLuminance + log2(ExposureScale), BaseLogLum, MiddleGreyLumValue, EyeAdaptation_LocalExposureHighlightContrastScale, EyeAdaptation_LocalExposureShadowContrastScale, EyeAdaptation_LocalExposureDetailStrength);
-
-		// Apply before bloom
-		SceneColor.rgb *= LocalExposure;
-		LocalExposureOneOverPreExposure *= LocalExposure;
-	}
-#endif
-
-#if DEBUG_TONEMAPPER
-	if (all(DebugTile == int2(-1, 1)))
-	{
-		// left below to the center is: no ExposureScale
-		ExposureScale = 1.0f;
-		bOutputDebugTile = true;
+		float BaseLogLum = CalculateBaseLogLuminance(LogLuminance, EyeAdaptation_LocalExposureBlurredLuminanceBlend, GlobalExposure, FullViewUV, LumBilateralGrid, BlurredLogLum, LumBilateralGridSampler, BlurredLogLumSampler);
+		LocalExposure = CalculateLocalExposure(LogLuminance + log2(GlobalExposure), BaseLogLum, MiddleGreyLumValue, EyeAdaptation_LocalExposureHighlightContrastScale, EyeAdaptation_LocalExposureShadowContrastScale, EyeAdaptation_LocalExposureDetailStrength);
 	}
-#endif
-
-	// 0..x, 0:no effect .. 1:strong, from r.Tonemapper.Sharpen
-	// Div6 is to save one instruction
-	float SharpenMultiplierDiv6 = TonemapperParams.y;
-
-#if DEBUG_TONEMAPPER
-	if (all(DebugTile == int2(0, -1)) || all(DebugTile == int2(-1, -1)))
-	{
-		// top row: no sharpen
-		SharpenMultiplierDiv6 = 0;
-		bOutputDebugTile = true;
-	}
-#endif //DEBUG_TONEMAPPER
+	#else // !USE_LOCAL_EXPOSURE
+	const float LocalExposure = 1.0;
+	#endif
 
-#if USE_SHARPEN
+	// Apply sharpening onto SceneColor.
+	#if USE_SHARPEN
 	{
+		// 0..x, 0:no effect .. 1:strong, from r.Tonemapper.Sharpen
+		// Div6 is to save one instruction
+		float SharpenMultiplierDiv6 = TonemapperParams.y;
 
 		half A0 = Luminance(SceneColor.rgb);
 
@@ -475,8 +401,8 @@ float4 TonemapCommonPS(
 				
 		half4 LuminanceNeighbors;
 
-		half3 C1 = LookupSceneColor(SceneUV, float2(PQMP.QuadVector.x, 0)).rgb * LocalExposureOneOverPreExposure;
-		half3 C3 = LookupSceneColor(SceneUV, float2(0, PQMP.QuadVector.y)).rgb * LocalExposureOneOverPreExposure;
+		half3 C1 = LookupSceneColor(SceneUV, float2(PQMP.QuadVector.x, 0)).rgb;
+		half3 C3 = LookupSceneColor(SceneUV, float2(0, PQMP.QuadVector.y)).rgb;
 		LuminanceNeighbors.x = Luminance(C1);
 		LuminanceNeighbors.y = Luminance(C3);
 		
@@ -485,131 +411,78 @@ float4 TonemapCommonPS(
 		LuminanceNeighbors.z =  PQMPReadX(PQMP,A0);
 		LuminanceNeighbors.w =  PQMPReadY(PQMP,A0);
 	#else
-		half3 C1 = LookupSceneColor(SceneUV, float2(-1,  0)) * LocalExposureOneOverPreExposure;
-		half3 C2 = LookupSceneColor(SceneUV, float2( 1,  0)) * LocalExposureOneOverPreExposure;
-		half3 C3 = LookupSceneColor(SceneUV, float2( 0, -1)) * LocalExposureOneOverPreExposure;
-		half3 C4 = LookupSceneColor(SceneUV, float2( 0,  1)) * LocalExposureOneOverPreExposure;
+		half3 C1 = LookupSceneColor(SceneUV, float2(-1,  0));
+		half3 C2 = LookupSceneColor(SceneUV, float2( 1,  0));
+		half3 C3 = LookupSceneColor(SceneUV, float2( 0, -1));
+		half3 C4 = LookupSceneColor(SceneUV, float2( 0,  1));
 		half4 LuminanceNeighbors = half4(Luminance(C1), Luminance(C2), Luminance(C3), Luminance(C4));
 	#endif
 
 		// compute mask to not sharpen near very bright HDR content
 		// Note: using max instead of summming up saves 1 instruction
 		// Note: We could multiply this to tweak the edge weight but it might instroduce a cost
-		half HDREdge = ExposureScale * max4(abs(A0 - LuminanceNeighbors));
+		half HDREdge = (OneOverPreExposure * LocalExposure * GlobalExposure) * max4(abs(A0 - LuminanceNeighbors));
 	
 		// 0..1
 		half EdgeMask = saturate(1.0f - HDREdge);
 
-#if DEBUG_TONEMAPPER
-		if (all(DebugTile == int2(1, 0)))
-		{
-			// right to the center is: Sharpen EdgeMask
-			OutColor = EdgeMask; 
-			return OutColor;
-		}
-#endif //DEBUG_TONEMAPPER
-
 		// -1:sharpen, 0:no effect, 1:blur
 		float LerpFactor = -EdgeMask * SharpenMultiplierDiv6;
 
 		// optimized, Div6 went into the C++ code
 		half3 DeltaColor = (C1 + C2 + C3 + C4) - SceneColor.rgb * 4;
+
 		SceneColor.rgb += DeltaColor * LerpFactor;
 	}
-#endif
+	#endif // USE_SHARPEN
 
-	#if USE_GAMMA_ONLY
+	float4 OutColor = 0;
+#if USE_GAMMA_ONLY
+	{
+		#if USE_LOCAL_EXPOSURE
+			#error Local exposure is not supporte with gamma only
+		#endif
+		#if USE_SHARPEN
+			#error Sharpen is not supported with gamma only
+		#endif
+		#if USE_VIGNETTE
+			#error Vignette is not supported with gamma only
+		#endif
 
-		SceneColor.rgb *= ExposureScale;
+		OutColor.rgb = pow(SceneColor.rgb * (OneOverPreExposure * GlobalExposure), InverseGamma.x);
 
-		OutColor.rgb = pow(SceneColor.rgb, InverseGamma.x);
 		OutLuminance = Luminance(OutColor.rgb);
-
-	#else 
+	}
+#else // !USE_GAMMA_ONLY
 	{
-		half3 LinearColor = SceneColor.rgb * ColorScale0.rgb;
-
-		#if SUPPORTS_SCENE_COLOR_APPLY_PARAMETERS
-		{
-			float3 SceneColorMultiply = SceneColorApplyParamaters[0].xyz;
-			LinearColor *= SceneColorMultiply;
-		}
-		#endif
-
+		// Compute the final linear color space with explicitly all the terms that should affect the pre exposure.
+		half3 FinalLinearColor = SceneColor.rgb * SceneColorTint * (OneOverPreExposure * GlobalExposure * VignetteMask * LocalExposure);
+		
+		// Apply bloom.
 		#if USE_BLOOM
 		{
-			float2 BloomUV;
-			#if ES3_1_PROFILE
-			{
-				BloomUV = FullViewUV.xy;
-			}
-			#else
-			{
-				BloomUV = ApplyScreenTransform(UV, ColorToBloom);
-				BloomUV = clamp(BloomUV, BloomUVViewportBilinearMin, BloomUVViewportBilinearMax);
-			}
-			#endif
-
-			float4 CombinedBloom = Texture2DSample(BloomTexture, BloomSampler, BloomUV);
-			CombinedBloom.rgb *= View.OneOverPreExposure;
-
-			//float2 DirtViewportUV = (SvPosition.xy - float2(Output_ViewportMin)) * Output_ViewportSizeInverse; 
-
-#if DEBUG_TONEMAPPER
-			if (all(DebugTile == int2(-1, -1)))
-			{
-				// left to the center is: no bloom
-				CombinedBloom = 0;
-				bOutputDebugTile = true;
-			}
-			if (all(DebugTile == int2(1, -1)))
-			{
-				// right to the center is: bloom only
-				LinearColor = 0;
-				bOutputDebugTile = true;
-			}
-#endif //DEBUG_TONEMAPPER
-
-			#if FEATURE_LEVEL == FEATURE_LEVEL_ES3_1
-			{
-				// Support sunshaft and vignette for mobile, and we have applied the BloomIntensity and the BloomDirtMask at the sun merge pass.
-				LinearColor = LinearColor.rgb * CombinedBloom.a + CombinedBloom.rgb;
-			}
-			#else
-			{
-				float2 DirtLensUV = ConvertScreenViewportSpaceToLensViewportSpace(ScreenPos) * float2(1.0f, -1.0f);
-
-				float3 BloomDirtMaskColor = Texture2DSample(BloomDirtMaskTexture, BloomDirtMaskSampler, DirtLensUV * .5f + .5f).rgb * BloomDirtMaskTint.rgb;
-				
-				LinearColor += CombinedBloom.rgb * (1.0 + BloomDirtMaskColor);
-			}
-			#endif
+			FinalLinearColor += Bloom * (OneOverPreExposure * GlobalExposure * VignetteMask);
 		}
 		#endif
 		
-		LinearColor *= ExposureScale;
-
-		#if USE_VIGNETTE
-			LinearColor.rgb *= ComputeVignetteMask( Vignette, TonemapperParams.x );
-		#endif
-
+		// Film grain is energy conservative so doesn't need to be taken into PreExposure
 		#if USE_FILM_GRAIN
 		{
 			float3 FilmGrainDecodeMultiply = FilmGrainTextureConstants[0].rgb;
 
-			half FilmGrainIntensity = ComputeFilmGrainIntensity(LinearColor);
+			half FilmGrainIntensity = ComputeFilmGrainIntensity(FinalLinearColor);
 
 			float2 FilmGrainUV = ApplyScreenTransform(ScreenPos, ScreenPosToFilmGrainTextureUV);
 
 			half3 RawGrain = FilmGrainTexture.SampleLevel(FilmGrainSampler, FilmGrainUV, 0);
 			half3 FinalGrain = RawGrain * half3(FilmGrainDecodeMultiply);
 
-			LinearColor.rgb *= lerp(1.0, FinalGrain, FilmGrainIntensity);
+			FinalLinearColor.rgb *= lerp(1.0, FinalGrain, FilmGrainIntensity);
 		}
 		#endif
-
-		half3 OutDeviceColor = ColorLookupTable(LinearColor);
 		
+		half3 OutDeviceColor = ColorLookupTable(FinalLinearColor);
+
 		half LuminanceForPostProcessAA = dot(OutDeviceColor, half3(0.299f, 0.587f, 0.114f));
 		OutLuminance = LuminanceForPostProcessAA;
 
@@ -651,31 +524,12 @@ float4 TonemapCommonPS(
 			OutColor.xyz = OutDeviceColor;
 		#endif
 	}
-	#endif
+#endif // !USE_GAMMA_ONLY
 	
 	#if POST_PROCESS_ALPHA == 2 || (POST_PROCESS_ALPHA == 1 && DIM_OUTPUT_DEVICE == TONEMAPPER_OUTPUT_NoToneCurve)  || (POST_PROCESS_ALPHA == 1 && DIM_OUTPUT_DEVICE == TONEMAPPER_OUTPUT_WithToneCurve)
 		OutColor.a = SceneColor.a;
 	#endif
 
-#if DEBUG_TONEMAPPER
-	if (!bOutputDebugTile)
-	{
-		// the original tonemapped color is only in the center
-		// other tiles can output like this:
-		//  if(all(DebugTile == int2(1, 0)))
-		//    { OutColor = ..; return; }
-		// or
-		//  if(all(DebugTile == int2(1, 0)))
-		//    { ..; bOutputDebugTile = true; }
-		OutColor = 0;
-	}
-	if (any(LocalDebugTilePosition == int2(0, 0)))
-	{
-		// border grid
-		OutColor = 1;
-	}
-#endif //DEBUG_TONEMAPPER
-
 #if PC_D3D && !USE_GAMMA_ONLY
 	// If HDR in the editor need to convert from PQ space to linear sRGB
 	BRANCH
