Substrate - fixed bad transmittance when doing masking through a vertical layer operator and more than 2 layers.

diff --git a/Engine/Shaders/Private/Strata/StrataTree.ush b/Engine/Shaders/Private/Strata/StrataTree.ush
--- a/Engine/Shaders/Private/Strata/StrataTree.ush
+++ b/Engine/Shaders/Private/Strata/StrataTree.ush
@@ -345,8 +345,10 @@ void FStrataTree::UpdateAllBSDFWithBottomUpOperatorVisit_Vertical(
 #if STRATA_INLINE_SHADING
 
 	// PreviousIsInputA > 0 means it is a BSDF affecting the top layer so we do not affect it by anything for this vertical layering.
+	const bool bBSDFComesFromTopLayer = PreviousIsInputA > 0;
+
 	// Otherise, the BSDF comes from the bottom part so we affect it weights from the gathered top layer coverage/transmittance
-	CurrentBSDF.LuminanceWeightV		 *= PreviousIsInputA > 0  ? 1.0f : saturate(1.0f - Op.VerticalTop_Coverage) + saturate(Op.VerticalTop_Coverage) * Op.VerticalTop_ThroughputAlongV;
+	CurrentBSDF.LuminanceWeightV		 *= bBSDFComesFromTopLayer ? 1.0f : saturate(1.0f - Op.VerticalTop_Coverage) + saturate(Op.VerticalTop_Coverage) * Op.VerticalTop_ThroughputAlongV;
 
 	// Now we are going to combine the coverage/transmittance above this slab into a new one according to the vertical operator we traverse now, if PreviousIsInputA > 0
 	const float BotCover  = CurrentBSDF.CoverageAboveAlongN;
@@ -356,13 +358,13 @@ void FStrataTree::UpdateAllBSDFWithBottomUpOperatorVisit_Vertical(
 	FVerticalLayeringInfo Info = GetVerticalLayeringInfo(Op.VerticalTop_Coverage, CurrentBSDF.CoverageAboveAlongN);
 	// Pre coverage transmittance is the surface transmitance as if coverage=1, so the final combined transmittance is renormalized accordingly.
 	const float3 PreCoverageTransmittance = saturate((Info.TransmittanceOnlyTop * TopTrans + Info.TransmittanceOnlyBottom * BotTrans + Info.TransmittanceTopAndBottom * (TopTrans * BotTrans)) / max(STRATA_EPSILON, Info.Coverage));
-	CurrentBSDF.CoverageAboveAlongN		  = PreviousIsInputA  > 0 ? CurrentBSDF.CoverageAboveAlongN : Info.Coverage;
-	CurrentBSDF.TransmittanceAboveAlongN *= PreviousIsInputA  > 0 ? 1.0f							: PreCoverageTransmittance;
+	CurrentBSDF.CoverageAboveAlongN		  = bBSDFComesFromTopLayer ? CurrentBSDF.CoverageAboveAlongN		: Info.Coverage;			// Info.Coverage combines the incoming operator Top part coverage and the BSDF current CoverageAboveAlongN
+	CurrentBSDF.TransmittanceAboveAlongN  = bBSDFComesFromTopLayer ? CurrentBSDF.TransmittanceAboveAlongN	: PreCoverageTransmittance;	// PreCoverageTransmittance is the combined trnasmittance for the BSDF TransmittanceAboveAlongN computed so far and the top layer of this new vertical operator.
 
 	// If the BSDF is from the bottom, we reduce its contribution according to the top layer coverage only.
-	CurrentBSDF.TopLayerDataWeight *= PreviousIsInputA ? 1.0 : saturate(1.0f - Op.VerticalTop_Coverage);
+	CurrentBSDF.TopLayerDataWeight *= bBSDFComesFromTopLayer ? 1.0 : saturate(1.0f - Op.VerticalTop_Coverage);
 
-	if (PreviousIsInputA <= 0)	// Only apply to BSDF from the bottom layer
+	if (!bBSDFComesFromTopLayer)	// Only apply to BSDF from the bottom layer
 	{
 		// We propagate the bottom layer lob up after the top layer top layer simulate the fact that light rays are potentially heavily scattered/refracted by the top layer.
 		// The outcome is that we affect the bottom layer roughness to aggregate that fact.
