- Fix vector field visualization for LWC
- Fix LWC vector field for Cascade GPU simulation
- Fix LWC point attraction for Cascade GPU simulation

diff --git a/Engine/Shaders/Private/ParticleSimulationShader.usf b/Engine/Shaders/Private/ParticleSimulationShader.usf
--- a/Engine/Shaders/Private/ParticleSimulationShader.usf
+++ b/Engine/Shaders/Private/ParticleSimulationShader.usf
@@ -186,7 +186,7 @@ float3 ComputeOrbitVelocity(float Time, float RandomOrbit)
  * @param Position - Position of the particle.
  * @param PerParticleScale - Amount by which to scale the influence on this particle.
  */
-void EvaluateVectorFields(out float3 OutForce, out float4 OutVelocity, float3 Position, float PerParticleScale)
+void EvaluateVectorFields(out float3 OutForce, out float4 OutVelocity, FLWCVector3 LWCPosition, float PerParticleScale)
 {
 	float3 TotalForce = 0;
 	float3 WeightedVelocity = 0;
@@ -195,11 +195,14 @@ void EvaluateVectorFields(out float3 OutForce, out float4 OutVelocity, float3 Po
 
 	for (int VectorFieldIndex = 0; VectorFieldIndex < VectorFields.Count; ++VectorFieldIndex)
 	{
+		FLWCVector3 WorldToVolumeTile = MakeLWCVector3(VectorFields.WorldToVolumeTile[VectorFieldIndex].xyz, 0.0f);
+		float3 Position = LWCToFloat(LWCSubtract(LWCPosition, WorldToVolumeTile));
+
 		float2 IntensityAndTightness = VectorFields.IntensityAndTightness[VectorFieldIndex].xy;
 		float Intensity = IntensityAndTightness.x * PerParticleScale;
 		float Tightness = IntensityAndTightness.y;
 		float3 VolumeSize = VectorFields.VolumeSize[VectorFieldIndex].xyz;
-		float3 VolumeUV = mul(float4(Position.xyz,1), VectorFields.WorldToVolume[VectorFieldIndex]).xyz;
+		float3 VolumeUV = mul(float4(Position, 1.0f), VectorFields.WorldToVolume[VectorFieldIndex]).xyz;
 		//Tile the UVs if needed. TilingAxes will be 1.0 or 0.0 in each channel depending on which axes are being tiled, if any.
 		VolumeUV -= floor(VolumeUV * VectorFields.TilingAxes[VectorFieldIndex].xyz);
 
@@ -563,7 +566,7 @@ void PixelMain(
 		// Evalute vector fields.
 		float3 FieldForce = 0;
 		float4 FieldVelocity = 0;
-		EvaluateVectorFields(FieldForce, FieldVelocity, LWCHackToFloat(Position), PerParticleVectorFieldScale);		//LWC_TODO: This needs to be rebased
+		EvaluateVectorFields(FieldForce, FieldVelocity, Position, PerParticleVectorFieldScale);
 
 		// Add in force from vector fields.
 		Force += FieldForce;
@@ -576,7 +579,7 @@ void PixelMain(
 		Force += ComputeDrag(Velocity.xyz, DragCoefficient);
 
 		// Compute force to a point gravity source.
-		Force += ComputeAttractionForce(LWCHackToFloat(Position));	//LWC_TODO: This needs to be rebased
+		Force += ComputeAttractionForce(PositionSample);
 
 		// Compute the acceleration to apply to the particle this frame.
 		float3 Acceleration = Force * DeltaSeconds;
diff --git a/Engine/Shaders/Private/VectorFieldVisualizationVertexFactory.ush b/Engine/Shaders/Private/VectorFieldVisualizationVertexFactory.ush
--- a/Engine/Shaders/Private/VectorFieldVisualizationVertexFactory.ush
+++ b/Engine/Shaders/Private/VectorFieldVisualizationVertexFactory.ush
@@ -72,9 +72,10 @@ FVertexFactoryIntermediates GetVertexFactoryIntermediates( FVertexFactoryInput I
 	const float3 Vec = normalize(Force) * ForceScale;
 
 	// The base position for this vector.
-	const float3 BasePosition = mul(float4(VoxelUV.xyz,1), VectorFieldVis.VolumeToWorld).xyz - LWCHackToFloat(ResolvedView.WorldCameraOrigin);
-	const float3 Position = BasePosition + float(Input.VertexId) * 
-		mul(float4(Vec,0), VectorFieldVis.VolumeToWorldNoScale).xyz;
+	FLWCVector3 VolumeToWorldTile = MakeLWCVector3(VectorFieldVis.VolumeToWorldTile.xyz, 0.0f);
+	FLWCVector3 WorldPosition = LWCAdd(LWCPromote(mul(float4(VoxelUV.xyz,1), VectorFieldVis.VolumeToWorld).xyz), VolumeToWorldTile);
+	const float3 BasePosition = LWCToFloat(LWCSubtract(WorldPosition, ResolvedView.WorldCameraOrigin));
+	const float3 Position = BasePosition + float(Input.VertexId) * mul(float4(Vec,0), VectorFieldVis.VolumeToWorldNoScale).xyz;
 
 	// Build and return the set of intermediates.
 	FVertexFactoryIntermediates Intermediates;
diff --git a/Engine/Source/Runtime/Engine/Private/Particles/ParticleGpuSimulation.cpp b/Engine/Source/Runtime/Engine/Private/Particles/ParticleGpuSimulation.cpp
--- a/Engine/Source/Runtime/Engine/Private/Particles/ParticleGpuSimulation.cpp
+++ b/Engine/Source/Runtime/Engine/Private/Particles/ParticleGpuSimulation.cpp
@@ -858,6 +858,7 @@ FParticlePerFrameSimulationShaderParameters GetParticlePerFrameSimulationShaderP
  */
 BEGIN_GLOBAL_SHADER_PARAMETER_STRUCT( FVectorFieldUniformParameters,)
 	SHADER_PARAMETER( int32, Count )
+	SHADER_PARAMETER_ARRAY( FVector4f, WorldToVolumeTile, [MAX_VECTOR_FIELDS] )
 	SHADER_PARAMETER_ARRAY( FMatrix44f, WorldToVolume, [MAX_VECTOR_FIELDS] )
 	SHADER_PARAMETER_ARRAY( FMatrix44f, VolumeToWorld, [MAX_VECTOR_FIELDS] )
 	SHADER_PARAMETER_ARRAY( FVector4f, IntensityAndTightness, [MAX_VECTOR_FIELDS] )
@@ -3151,7 +3152,7 @@ FGPUSpriteParticleEmitterInstance(FFXSystem* InFXSystem, FGPUSpriteEmitterInfo&
 
 			}
 			
-			FVector3f PointAttractorPosition((FVector4f)ComponentToWorld.TransformPosition(EmitterInfo.PointAttractorPosition));	// LWC_TODO: Precision loss
+			const FVector3f PointAttractorPosition = FVector4f(ComponentToWorld.TransformPosition(EmitterInfo.PointAttractorPosition));
 			DynamicData->PerFrameSimulationParameters.PointAttractor = FVector4f(PointAttractorPosition, EmitterInfo.PointAttractorRadiusSq);
 			DynamicData->PerFrameSimulationParameters.PositionOffsetAndAttractorStrength = FVector4f(FVector3f(PositionOffsetThisTick), PointAttractorStrength);
 			DynamicData->PerFrameSimulationParameters.LocalToWorldScale = DynamicData->EmitterDynamicParameters.LocalToWorldScale;
@@ -4345,7 +4346,10 @@ static void SetParametersForVectorField(FVectorFieldUniformParameters& OutParame
 			Tightness = FMath::Clamp<float>(VectorFieldInstance->Tightness, 0.0f, 1.0f);
 		}
 
-		OutParameters.WorldToVolume[Index] = FMatrix44f(VectorFieldInstance->WorldToVolume);		// LWC_TODO: Precision loss
+		const FLargeWorldRenderPosition WorldToVolumeOrigin(VectorFieldInstance->VolumeToWorld.GetOrigin());
+
+		OutParameters.WorldToVolumeTile[Index] = WorldToVolumeOrigin.GetTile();
+		OutParameters.WorldToVolume[Index] = FLargeWorldRenderScalar::MakeToRelativeWorldMatrix(WorldToVolumeOrigin.GetTileOffset(), VectorFieldInstance->VolumeToWorld).Inverse();
 		OutParameters.VolumeToWorld[Index] = FMatrix44f(VectorFieldInstance->VolumeToWorldNoScale);
 		OutParameters.VolumeSize[Index] = FVector4f(Resource->SizeX, Resource->SizeY, Resource->SizeZ, 0);
 		OutParameters.IntensityAndTightness[Index] = FVector4f(Intensity, Tightness, 0, 0 );
diff --git a/Engine/Source/Runtime/Engine/Private/VectorFieldVisualization.cpp b/Engine/Source/Runtime/Engine/Private/VectorFieldVisualization.cpp
--- a/Engine/Source/Runtime/Engine/Private/VectorFieldVisualization.cpp
+++ b/Engine/Source/Runtime/Engine/Private/VectorFieldVisualization.cpp
@@ -243,8 +243,11 @@ void GetVectorFieldMesh(
 		Collector.RegisterOneFrameMaterialProxy(VisualizationMaterial);
 
 		// Set up parameters.
+		const FLargeWorldRenderPosition VolumeToWorldOrigin(VectorFieldInstance->VolumeToWorld.GetOrigin());
+
 		FVectorFieldVisualizationParameters UniformParameters;
-		UniformParameters.VolumeToWorld = FMatrix44f(VectorFieldInstance->VolumeToWorld);	// LWC_TODO: Precision loss
+		UniformParameters.VolumeToWorldTile = VolumeToWorldOrigin.GetTile();
+		UniformParameters.VolumeToWorld = FLargeWorldRenderScalar::MakeToRelativeWorldMatrix(VolumeToWorldOrigin.GetTileOffset(), VectorFieldInstance->VolumeToWorld);
 		UniformParameters.VolumeToWorldNoScale = FMatrix44f(VectorFieldInstance->VolumeToWorldNoScale);
 		UniformParameters.VoxelSize = FVector3f( 1.0f / Resource->SizeX, 1.0f / Resource->SizeY, 1.0f / Resource->SizeZ );
 		UniformParameters.Scale = VectorFieldInstance->Intensity * Resource->Intensity;
diff --git a/Engine/Source/Runtime/Engine/Private/VectorFieldVisualization.h b/Engine/Source/Runtime/Engine/Private/VectorFieldVisualization.h
--- a/Engine/Source/Runtime/Engine/Private/VectorFieldVisualization.h
+++ b/Engine/Source/Runtime/Engine/Private/VectorFieldVisualization.h
@@ -37,6 +37,7 @@ class FSceneView;
  * Uniform buffer to hold parameters for vector field visualization.
  */
 BEGIN_GLOBAL_SHADER_PARAMETER_STRUCT( FVectorFieldVisualizationParameters, )
+	SHADER_PARAMETER( FVector3f, VolumeToWorldTile )
 	SHADER_PARAMETER( FMatrix44f, VolumeToWorld )
 	SHADER_PARAMETER( FMatrix44f, VolumeToWorldNoScale )
 	SHADER_PARAMETER( FVector3f, VoxelSize )
