Removes bilateral weights in TSR shading rejection's Blur3x3

diff --git a/Engine/Shaders/Private/TemporalSuperResolution/TSRConvolutionNetwork.ush b/Engine/Shaders/Private/TemporalSuperResolution/TSRConvolutionNetwork.ush
--- a/Engine/Shaders/Private/TemporalSuperResolution/TSRConvolutionNetwork.ush
+++ b/Engine/Shaders/Private/TemporalSuperResolution/TSRConvolutionNetwork.ush
@@ -202,120 +202,30 @@ void AccessNeighborTexels3x1(
 }
 
 
-//------------------------------------------------------- CONVOLUTION ACCUMULATORS
-
-template<typename ScalarType, uint VectorSize, uint LaneStrideX, uint LaneStrideY>
-struct TAccumulatorAvg
-{
-	TLaneVector2D<ScalarType, VectorSize, LaneStrideX, LaneStrideY> Value;
-	TLaneVector2D<ScalarType, 1, LaneStrideX, LaneStrideY> Weight;
-	
-	CALL_SITE_DEBUGLOC
-	void Start(TLaneVector2D<ScalarType, VectorSize, LaneStrideX, LaneStrideY> CenterSample)
-	{
-		Value = CenterSample;
-		Weight = TLaneVector2D<ScalarType, 1, LaneStrideX, LaneStrideY>::Const(1.0);
-	}
-	
-	CALL_SITE_DEBUGLOC
-	static TLaneVector2D<ScalarType, VectorSize, LaneStrideX, LaneStrideY> PreprocessNeighborSample(
-		TLaneVector2D<ScalarType, VectorSize, LaneStrideX, LaneStrideY> Sample,
-		TLaneVector2D<ScalarType, 1, LaneStrideX, LaneStrideY> InWeight)
-	{
-		return Sample;
-	}
-	
-	CALL_SITE_DEBUGLOC
-	void Accumulate(
-		const tsr_short2 Offset,
-		TLaneVector2D<ScalarType, VectorSize, LaneStrideX, LaneStrideY> Sample,
-		TLaneVector2D<ScalarType, 1, LaneStrideX, LaneStrideY> InWeight)
-	{
-		const tsr_half SampleSpatialWeight = tsr_half(rcp(abs(tsr_half(Offset.x)) + 1.0) * rcp(abs(tsr_half(Offset.y)) + 1.0));
-		
-		Value = Value + Sample * TLaneVector2D<ScalarType, VectorSize, LaneStrideX, LaneStrideY>::Vectorize(InWeight * SampleSpatialWeight);
-		Weight = Weight + InWeight * TLaneVector2D<ScalarType, 1, LaneStrideX, LaneStrideY>::Const(SampleSpatialWeight);
-	}
-	
-	CALL_SITE_DEBUGLOC
-	void Finish(const uint SampleCount)
-	{
-		Value = Value * TLaneVector2D<ScalarType, VectorSize, LaneStrideX, LaneStrideY>::Vectorize(rcp(Weight));
-	}
-};
-
-
-//------------------------------------------------------- CONVOLUTIONS
+//------------------------------------------------------- BLUR 3x3 CONVOLUTION
 
 CALL_SITE_DEBUGLOC
-template<class FAccumulator, typename ScalarType, uint VectorSize, uint LaneStrideX, uint LaneStrideY>
-FAccumulator Convole3x1(
-	TLaneVector2D<ScalarType, VectorSize, LaneStrideX, LaneStrideY> Color,
-	TLaneVector2D<ScalarType, 1, LaneStrideX, LaneStrideY> Weight,
-	const tsr_short2 Offset)
-#if CONFIG_OPTIMIZED_LDS_ACCESS
-{
-	const tsr_short2 OffsetN = tsr_short2(-int2(Offset)); // work arround a shader compiler bug
-	
-	TLaneVector2D<ScalarType, VectorSize, LaneStrideX, LaneStrideY> PreprocessedNeighborColor = FAccumulator::PreprocessNeighborSample(Color, Weight);
-
-	WriteToLDSForAccessNeighborTexel(PreprocessedNeighborColor);
-	TLaneVector2D<ScalarType, VectorSize, LaneStrideX, LaneStrideY> Sample0 = AccessNeighborTexel(PreprocessedNeighborColor, Offset);
-	TLaneVector2D<ScalarType, VectorSize, LaneStrideX, LaneStrideY> Sample1 = AccessNeighborTexel(PreprocessedNeighborColor, OffsetN);
-	
-	WriteToLDSForAccessNeighborTexel(Weight);
-	TLaneVector2D<ScalarType, 1, LaneStrideX, LaneStrideY> Weight0 = AccessNeighborTexel(Weight, Offset);
-	TLaneVector2D<ScalarType, 1, LaneStrideX, LaneStrideY> Weight1 = AccessNeighborTexel(Weight, OffsetN);
-	
-	FAccumulator Accumulator;
-	Accumulator.Start(Color);
-	Accumulator.Accumulate(+Offset, Sample0, Weight0);
-	Accumulator.Accumulate(OffsetN, Sample1, Weight1);
-	Accumulator.Finish(3);
-	return Accumulator;
-}
-#else // !CONFIG_OPTIMIZED_LDS_ACCESS
+template<typename ScalarType, uint VectorSize, uint LaneStrideX, uint LaneStrideY>
+TLaneVector2D<ScalarType, VectorSize, LaneStrideX, LaneStrideY> Blur3x3(TLaneVector2D<ScalarType, VectorSize, LaneStrideX, LaneStrideY> Center)
 {
-	FAccumulator Accumulator;
-	Accumulator.Start(Color);
-	
-	TLaneVector2D<ScalarType, VectorSize, LaneStrideX, LaneStrideY> PreprocessedNeighborColor = FAccumulator::PreprocessNeighborSample(Color, Weight);
-
-	TLaneVector2D<ScalarType, VectorSize, LaneStrideX, LaneStrideY> Sample0 = AccessNeighborTexel(PreprocessedNeighborColor, Offset);
-	TLaneVector2D<ScalarType, 1, LaneStrideX, LaneStrideY>          Weight0 = AccessNeighborTexel(Weight, Offset);
-	Accumulator.Accumulate(+Offset, Sample0, Weight0);
+	TLaneVector2D<ScalarType, VectorSize, LaneStrideX, LaneStrideY> CO = Center;
+	TLaneVector2D<ScalarType, VectorSize, LaneStrideX, LaneStrideY> CP;
+	TLaneVector2D<ScalarType, VectorSize, LaneStrideX, LaneStrideY> CN;
+	AccessNeighborTexels3x1(CO, tsr_short2(1, 0), /* out */ CP, /* out */ CN);
 	
-	const tsr_short2 OffsetN = tsr_short2(-int2(Offset)); // work arround a shader compiler bug
-	TLaneVector2D<ScalarType, VectorSize, LaneStrideX, LaneStrideY> Sample1 = AccessNeighborTexel(PreprocessedNeighborColor, OffsetN);
-	TLaneVector2D<ScalarType, 1, LaneStrideX, LaneStrideY>          Weight1 = AccessNeighborTexel(Weight, OffsetN);
-	Accumulator.Accumulate(OffsetN, Sample1, Weight1);
-
-	Accumulator.Finish(3);
-	return Accumulator;
-}
-#endif // !CONFIG_OPTIMIZED_LDS_ACCESS
-
-CALL_SITE_DEBUGLOC
-template<class FAccumulator, typename ScalarType, uint VectorSize, uint LaneStrideX, uint LaneStrideY>
-TLaneVector2D<ScalarType, VectorSize, LaneStrideX, LaneStrideY> Convolve3x3HV(
-	TLaneVector2D<ScalarType, VectorSize, LaneStrideX, LaneStrideY> Color,
-	TLaneVector2D<ScalarType, 1, LaneStrideX, LaneStrideY> Weight)
-{
-	TLaneVector2D<ScalarType, VectorSize, LaneStrideX, LaneStrideY> HorizontalColor = Convole3x1<FAccumulator, ScalarType, VectorSize, LaneStrideX, LaneStrideY>(Color, Weight, tsr_short2(1, 0)).Value;
-	return Convole3x1<FAccumulator, ScalarType, VectorSize, LaneStrideX, LaneStrideY>(HorizontalColor, Weight, tsr_short2(0, 1)).Value;
-}
-
+	TLaneVector2D<ScalarType, VectorSize, LaneStrideX, LaneStrideY> BO = (
+		CO * TLaneVector2D<ScalarType, VectorSize, LaneStrideX, LaneStrideY>::Const(0.5) +
+		CP * TLaneVector2D<ScalarType, VectorSize, LaneStrideX, LaneStrideY>::Const(0.25) +
+		CN * TLaneVector2D<ScalarType, VectorSize, LaneStrideX, LaneStrideY>::Const(0.25));
 
-//------------------------------------------------------- BLUR 3x3 CONVOLUTION
+	TLaneVector2D<ScalarType, VectorSize, LaneStrideX, LaneStrideY> BP;
+	TLaneVector2D<ScalarType, VectorSize, LaneStrideX, LaneStrideY> BN;
+	AccessNeighborTexels3x1(BO, tsr_short2(0, 1), /* out */ BP, /* out */ BN);
 
-CALL_SITE_DEBUGLOC
-template<typename ScalarType, uint VectorSize, uint LaneStrideX, uint LaneStrideY>
-TLaneVector2D<ScalarType, VectorSize, LaneStrideX, LaneStrideY> Blur3x3(
-	TLaneVector2D<ScalarType, VectorSize, LaneStrideX, LaneStrideY> Center,
-	TLaneVector2D<ScalarType, 1, LaneStrideX, LaneStrideY> Weight)
-{
-	return Convolve3x3HV<TAccumulatorAvg<ScalarType, VectorSize, LaneStrideX, LaneStrideY>, ScalarType, VectorSize, LaneStrideX, LaneStrideY>(
-		Center, Weight);
+	return (
+		BO * TLaneVector2D<ScalarType, VectorSize, LaneStrideX, LaneStrideY>::Const(0.5) +
+		BP * TLaneVector2D<ScalarType, VectorSize, LaneStrideX, LaneStrideY>::Const(0.25) +
+		BN * TLaneVector2D<ScalarType, VectorSize, LaneStrideX, LaneStrideY>::Const(0.25));
 }
 
 
diff --git a/Engine/Shaders/Private/TemporalSuperResolution/TSRRejectShading.usf b/Engine/Shaders/Private/TemporalSuperResolution/TSRRejectShading.usf
--- a/Engine/Shaders/Private/TemporalSuperResolution/TSRRejectShading.usf
+++ b/Engine/Shaders/Private/TemporalSuperResolution/TSRRejectShading.usf
@@ -218,7 +218,7 @@ void MainCS(
 
 	#if CONFIG_COMPOSE_TRANSLUCENCY
 	{
-		tsr_simd_half4 FilteredTranslucencyInput = Blur3x3(OriginalTranslucencyInput, bIsValid);
+		tsr_simd_half4 FilteredTranslucencyInput = Blur3x3(OriginalTranslucencyInput);
 
 		UNROLL_N(SIMD_SIZE)
 		for (uint SimdIndex = 0; SimdIndex < SIMD_SIZE; SimdIndex++)
@@ -277,8 +277,8 @@ void MainCS(
 		}
 		
 		// Convolve the input luma to have the exact same moire pattern as FilteredInput below.
-		InputLuma = Blur3x3(InputLuma, bIsValid);
-		InputLumaModification = Blur3x3(InputLumaModification, bIsValid);
+		InputLuma = Blur3x3(InputLuma);
+		InputLumaModification = Blur3x3(InputLumaModification);
 
 		// Unpack history
 		tsr_simd_half PrevFlickeringHistory = PrevMoire[0];
@@ -418,8 +418,8 @@ void MainCS(
 	#endif
 
 	// Filter clamped input and history to increase quantization of the rejection detection
-	tsr_simd_halfC FilteredInput   = Blur3x3(ClampedInput  , bIsValid);
-	tsr_simd_halfC FilteredHistory = Blur3x3(ClampedHistory, bIsValid);
+	tsr_simd_halfC FilteredInput   = Blur3x3(ClampedInput);
+	tsr_simd_halfC FilteredHistory = Blur3x3(ClampedHistory);
 	
 	// Measure rejection
 	tsr_simd_half RejectionBlendFinal = tsr_simd_half::Const(1.0);
