Eye, foliage and subsurface was accidientally marked as needing backface cullling during refactoring for ray traced shadows.

diff --git a/Engine/Shaders/Private/RayTracing/RayTracingOcclusionRGS.usf b/Engine/Shaders/Private/RayTracing/RayTracingOcclusionRGS.usf
--- a/Engine/Shaders/Private/RayTracing/RayTracingOcclusionRGS.usf
+++ b/Engine/Shaders/Private/RayTracing/RayTracingOcclusionRGS.usf
@@ -694,10 +694,11 @@ RAY_TRACING_ENTRY_RAYGEN_OR_INLINE(Occlusion)
 			ShadingParameters.bIsHair = GBufferData.ShadingModelID == SHADINGMODELID_HAIR;
 			ShadingParameters.bIsSubsurface = GBufferData.ShadingModelID == SHADINGMODELID_SUBSURFACE;
 			ShadingParameters.bHasSubsurfaceProfile = GBufferData.ShadingModelID == SHADINGMODELID_SUBSURFACE_PROFILE;
+			// eye, two sided foliage and subsurface do not need backface culling.
 			ShadingParameters.bBackfaceCulling =
-				GBufferData.ShadingModelID == SHADINGMODELID_EYE ||
-				GBufferData.ShadingModelID == SHADINGMODELID_TWOSIDED_FOLIAGE ||
-				GBufferData.ShadingModelID == SHADINGMODELID_SUBSURFACE;
+				!(GBufferData.ShadingModelID == SHADINGMODELID_EYE ||
+				GBufferData.ShadingModelID == SHADINGMODELID_TWOSIDED_FOLIAGE || 
+				GBufferData.ShadingModelID == SHADINGMODELID_SUBSURFACE);
 
 			WorldNormal	= GBufferData.WorldNormal;
 
