Explicitly unroll loop in TAA shader to improve platform specific DXC performance. This improves TAA on certain platforms with DXC by roughly 0.5ms.

#rnx

diff --git a/Engine/Shaders/Private/TemporalAA.usf b/Engine/Shaders/Private/TemporalAA.usf
--- a/Engine/Shaders/Private/TemporalAA.usf
+++ b/Engine/Shaders/Private/TemporalAA.usf
@@ -1821,6 +1821,7 @@ FTAAHistoryPayload SampleHistory(in float2 HistoryScreenPosition)
 		#endif
 
 		FCatmullRomSamples Samples = GetBicubic2DCatmullRomSamples(HistoryBufferUV, HistoryBufferSize.xy, HistoryBufferSize.zw);
+		UNROLL
 		for (uint i = 0; i < Samples.Count; i++)
 		{
 			float2 SampleUV = Samples.UV[i];
