Implements triangular grain quantization

This also moves the grain quantization right before output write as it is designed to stocastically quantize based on output pixel format's bitdepth, to have correct behavior on HDR outputs.
This also automaticaly choose the quantization amplitude based of back buffer pixel format if FSceneViewFamily::RenderTarget.

All rendering EngineTest has been ran to verify the change of quantization did not cause test image reference invalidations.

diff --git a/Engine/Shaders/Private/PostProcessTonemap.usf b/Engine/Shaders/Private/PostProcessTonemap.usf
--- a/Engine/Shaders/Private/PostProcessTonemap.usf
+++ b/Engine/Shaders/Private/PostProcessTonemap.usf
@@ -23,8 +23,7 @@
 	#define DIM_OUTPUT_DEVICE (TONEMAPPER_OUTPUT_sRGB)
 #endif
 
-// 0 / 1(slower, visualize multiple internals side by side for the same image portion, see DebugTile)
-#define DEBUG_TONEMAPPER 0
+#define DEBUG_GRAIN_QUANTIZATION 0
 
 SCREEN_PASS_TEXTURE_VIEWPORT(Color)
 SCREEN_PASS_TEXTURE_VIEWPORT(Output)
@@ -74,6 +73,7 @@ float4 ScreenPosToFilmGrainTextureUV;
 	StructuredBuffer<float4> FilmGrainTextureConstants;
 #endif // SUPPORTS_FILM_GRAIN
 
+float GrainQuantization;
 float EditorNITLevel;
 
 uint bOutputInHDR;
@@ -218,7 +218,11 @@ void TonemapCommonVS(
 
 	// Grain
 	OutGrainUV.xy = TexCoord + Color_ExtentInverse * float2(-0.5,0.5);
-	OutGrainUV.zw = TexCoord + GrainRandomFull.xy;
+	#if DEBUG_GRAIN_QUANTIZATION
+		OutGrainUV.zw = TexCoord;
+	#else
+		OutGrainUV.zw = TexCoord + GrainRandomFull.xy;
+	#endif
 
 	// Fringe
 	OutScreenPos = UVToScreenPos(TexCoord, Color_Extent);
@@ -486,26 +490,6 @@ float4 TonemapCommonPS(
 		half LuminanceForPostProcessAA = dot(OutDeviceColor, half3(0.299f, 0.587f, 0.114f));
 		OutLuminance = LuminanceForPostProcessAA;
 
-		#if USE_GRAIN_QUANTIZATION
-			half Grain = GrainFromUV(GrainUV.zw);
-
-			#if DIM_OUTPUT_DEVICE == TONEMAPPER_OUTPUT_NoToneCurve
-				#error LinearNoToneCurve device requires no Grain Quantization
-			#endif
-			#if DIM_OUTPUT_DEVICE == TONEMAPPER_OUTPUT_WithToneCurve
-				#error LinearWithToneCurve device requires no Grain Quantization
-			#endif
-		// Needs to go after tonemapping. 
-			#if USE_LINEAR_FLOAT_RT
-				half GrainQuantization = 1.0/1024.0;
-			#else
-				half GrainQuantization = 1.0/256.0;
-			#endif
-
-			half GrainAdd = (Grain * GrainQuantization) + (-0.5 * GrainQuantization);
-			OutDeviceColor.rgb += GrainAdd;
-		#endif
-
 		// RETURN_COLOR not needed unless writing to SceneColor
 		OutColor.rgb = OutDeviceColor;
 
@@ -540,6 +524,54 @@ float4 TonemapCommonPS(
         OutColor.rgb = LinearToPostTonemapSpace(OutColor.rgb);
 	}
 #endif
+
+	#if USE_GRAIN_QUANTIZATION && !USE_GAMMA_ONLY
+	{
+		#if DIM_OUTPUT_DEVICE == TONEMAPPER_OUTPUT_NoToneCurve
+			#error LinearNoToneCurve device requires no Grain Quantization
+		#endif
+		#if DIM_OUTPUT_DEVICE == TONEMAPPER_OUTPUT_WithToneCurve
+			#error LinearWithToneCurve device requires no Grain Quantization
+		#endif
+
+		// Compute uniform grain on [0;1]
+		half UniformGrainOP = GrainFromUV(GrainUV.zw); //TODO: verify the quality of GrainFromUV()
+
+		// Compute uniform grain on [-1;1]
+		half UniformGrainNP = UniformGrainOP * 2.0 - 1.0;
+		half UniformGrainNPSign = clamp(UniformGrainNP * POSITIVE_INFINITY, half(-1.0), half(1.0));
+		
+		// Compute triangular grain on [-1;1]
+		// half TriangularGrainOP = UniformGrainNPSign * (1.0 - sqrt(1.0 - abs(UniformGrainNP)));
+		half TriangularGrainOP = UniformGrainNPSign - UniformGrainNPSign * sqrt(saturate(half(1.0) - abs(UniformGrainNP)));
+		
+		// Selects the grain to use.
+		half FinalGrain = 0.0;
+		#if DEBUG_GRAIN_QUANTIZATION
+			if (ScreenPos.y > 1.0 / 3.0)
+				FinalGrain = TriangularGrainOP;
+			else if (ScreenPos.y < -1.0 / 3.0)
+				FinalGrain = UniformGrainOP - 0.5;
+		#elif 0
+			FinalGrain = UniformGrainOP - 0.5;
+		#else
+			FinalGrain = TriangularGrainOP;
+		#endif
+
+		// Debug grain quantization.
+		#if DEBUG_GRAIN_QUANTIZATION
+		{
+			half Quantization = 10.0;
+			
+			OutColor.rgb = round(OutColor.rgb * Quantization + FinalGrain) * rcp(Quantization);
+			OutColor.rgb = OutColor.ggg;
+		}
+		#else
+			OutColor.rgb += FinalGrain * GrainQuantization;
+		#endif
+	}
+	#endif
+
 	return OutColor;
 }
 
diff --git a/Engine/Source/Runtime/Renderer/Private/PostProcess/PostProcessTonemap.cpp b/Engine/Source/Runtime/Renderer/Private/PostProcess/PostProcessTonemap.cpp
--- a/Engine/Source/Runtime/Renderer/Private/PostProcess/PostProcessTonemap.cpp
+++ b/Engine/Source/Runtime/Renderer/Private/PostProcess/PostProcessTonemap.cpp
@@ -337,6 +337,7 @@ BEGIN_SHADER_PARAMETER_STRUCT(FTonemapParameters, )
 	SHADER_PARAMETER(float, LUTScale)
 	SHADER_PARAMETER(float, LUTOffset)
 	SHADER_PARAMETER(float, EditorNITLevel)
+	SHADER_PARAMETER(float, GrainQuantization)
 	SHADER_PARAMETER(uint32, bOutputInHDR)
 END_SHADER_PARAMETER_STRUCT()
 
@@ -653,6 +654,16 @@ FScreenPassTexture AddTonemapPass(FRDGBuilder& GraphBuilder, const FViewInfo& Vi
 	}
 	#endif
 
+	float GrainQuantization = 1.0f / 1024.0f;
+	if (View.Family->RenderTarget && View.Family->RenderTarget->GetRenderTargetTexture())
+	{
+		EPixelFormat BackbufferPixelFormat = View.Family->RenderTarget->GetRenderTargetTexture()->GetFormat();
+		if (BackbufferPixelFormat == PF_B8G8R8A8 || BackbufferPixelFormat == PF_R8G8B8A8)
+		{
+			GrainQuantization = 1.0 / 256.0;
+		}
+	}
+
 	FTonemapParameters CommonParameters;
 	CommonParameters.View = View.ViewUniformBuffer;
 
@@ -768,6 +779,7 @@ FScreenPassTexture AddTonemapPass(FRDGBuilder& GraphBuilder, const FViewInfo& Vi
 	CommonParameters.ChromaticAberrationParams = ChromaticAberrationParams;
 	CommonParameters.TonemapperParams = FVector4f(PostProcessSettings.VignetteIntensity, SharpenDiv6, 0.0f, 0.0f);
 	CommonParameters.EditorNITLevel = EditorNITLevel;
+	CommonParameters.GrainQuantization = GrainQuantization;
 	CommonParameters.bOutputInHDR = ViewFamily.bIsHDR;
 	CommonParameters.LUTSize = LUTSize;
 	CommonParameters.InvLUTSize = 1.0f / LUTSize;
