Path Tracer: Add decal support with Substrate enabled

Refactor the decal apply logic so that we adjust diffuse/specular _after_ having applied the decals so that we can properly handle BaseColor even in the presence of subsurface scattering.

Refactor EnableDBuffer logic to remove the decal data after the second trace so we don't need to check it again. This keeps the dbuffer logic more self contained.

PT_Decals in EngineTest is mostly passing, but there are a few differences that need to be investigated more carefully. I suspect that some of them might be due to the round trip between diffuse/specular color and metallic/specular workflow.

diff --git a/Engine/Shaders/Private/PathTracing/Material/PathTracingMaterialSampling.ush b/Engine/Shaders/Private/PathTracing/Material/PathTracingMaterialSampling.ush
--- a/Engine/Shaders/Private/PathTracing/Material/PathTracingMaterialSampling.ush
+++ b/Engine/Shaders/Private/PathTracing/Material/PathTracingMaterialSampling.ush
@@ -356,10 +356,9 @@ float GetAverageRoughness(FPathTracingPayload Payload)
 #endif
 }
 
-FPathTracingPayload UnpackPathTracingPayloadMaterial(FPackedPathTracingPayload PackedInput, FRayDesc Ray, bool bApplyOverrides)
+void AdjustPayloadAfterUnpack(inout FPathTracingPayload HitPayload, bool bApplyOverrides)
 {
 	// This function can also do a bit of per-material post-processing to simplify the rest of the code
-	FPathTracingPayload HitPayload = UnpackPathTracingPayload(PackedInput, Ray);
 	if (bApplyOverrides)
 	{
 		HitPayload.DiffuseColor  = HitPayload.DiffuseColor  * View.DiffuseOverrideParameter.w + View.DiffuseOverrideParameter.xyz;
@@ -475,7 +474,6 @@ FPathTracingPayload UnpackPathTracingPayloadMaterial(FPackedPathTracingPayload P
 #endif
 	}
 #endif // PATHTRACING_STRATA_PAYLOAD
-	return HitPayload;
 }
 
 struct FSSSRandomWalkInfo
diff --git a/Engine/Shaders/Private/PathTracing/PathTracingCommon.ush b/Engine/Shaders/Private/PathTracing/PathTracingCommon.ush
--- a/Engine/Shaders/Private/PathTracing/PathTracingCommon.ush
+++ b/Engine/Shaders/Private/PathTracing/PathTracingCommon.ush
@@ -139,6 +139,7 @@ struct FPackedPathTracingPayload : FMinimalPayload
 
 	uint GetFlags() { return (PackedData[0] >> 20) & 0x1FF; }
 	void SetFlag(uint Flag) { PackedData[0] |= Flag << 20; }
+	void RemoveFlag(uint Flag) { PackedData[0] &= ~(Flag << 20); }
 	uint GetPrimitiveLightingChannelMask() { return (PackedData[0] >> 29) & 0x7; }
 
 	// Flag methods
diff --git a/Engine/Shaders/Private/PathTracing/PathTracingCore.ush b/Engine/Shaders/Private/PathTracing/PathTracingCore.ush
--- a/Engine/Shaders/Private/PathTracing/PathTracingCore.ush
+++ b/Engine/Shaders/Private/PathTracing/PathTracingCore.ush
@@ -216,41 +216,69 @@ void AccumulateRadiance(inout float3 TotalRadiance, float3 PathRadiance, bool bI
 	TotalRadiance += PathRadiance;
 }
 
-#if STRATA_ENABLED
-	// TODO: Figure out decals + Strata
-#else
 struct FDecalData
 {
 	float4 BaseColor;
 	float4 WorldNormal;
 	float4 MetallicSpecularRoughness;
 	float3 Emissive;
-};
 
-float ApplyDecal(inout FPathTracingPayload HitPayload, FDecalData DecalData)
-{
-	if (!EnableDBuffer || HitPayload.HasDecalResponseColor())
+	bool HasData()
 	{
-		HitPayload.BaseColor = HitPayload.BaseColor * DecalData.BaseColor.a + DecalData.BaseColor.rgb;
+		return min3(BaseColor.a, WorldNormal.a, MetallicSpecularRoughness.a) < 1.0;
 	}
 
-	if (!EnableDBuffer || HitPayload.HasDecalResponseNormal())
+	bool ApplyDecal(inout FPathTracingPayload HitPayload)
 	{
-		// We normalize the normal to get smoother visual result (it helps to avoid having D_GGX explodes toward infinity, and matches ApplyDBufferData(...))
-		HitPayload.WorldNormal = normalize(HitPayload.WorldNormal * DecalData.WorldNormal.a + DecalData.WorldNormal.rgb);
-	}
+#if PATHTRACING_STRATA_PAYLOAD
+		float Roughness = HitPayload.RoughnessData.x;
+		float Metallic = F0RGBToMetallic(HitPayload.SpecularColor);
+		float Specular = F0RGBToDielectricSpecular(HitPayload.SpecularColor);
+		float3 BaseColor = lerp(HitPayload.DiffuseColor, HitPayload.SpecularColor, Metallic);
+#else
+		float3 BaseColor = HitPayload.BaseColor;
+		float Metallic = HitPayload.Metallic;
+		float Specular = HitPayload.Specular;
+		float Roughness = HitPayload.Roughness;
+#endif // PATHTRACING_STRATA_PAYLOAD
+		if (HitPayload.HasDecalResponseColor())
+		{
+			BaseColor = BaseColor * this.BaseColor.a + this.BaseColor.rgb;
+		}
 
-	if (!EnableDBuffer || HitPayload.HasDecalResponseRoughness())
-	{
-		HitPayload.Metallic = HitPayload.Metallic * DecalData.MetallicSpecularRoughness.a + DecalData.MetallicSpecularRoughness.x;
-		HitPayload.Specular = HitPayload.Specular * DecalData.MetallicSpecularRoughness.a + DecalData.MetallicSpecularRoughness.y;
-		HitPayload.Roughness = HitPayload.Roughness * DecalData.MetallicSpecularRoughness.a + DecalData.MetallicSpecularRoughness.z;
-	}
+		if (HitPayload.HasDecalResponseNormal())
+		{
+			// We normalize the normal to get smoother visual result (it helps to avoid having D_GGX explodes toward infinity, and matches ApplyDBufferData(...))
+			HitPayload.WorldNormal = normalize(HitPayload.WorldNormal * this.WorldNormal.a + this.WorldNormal.rgb);
+		}
+
+		if (HitPayload.HasDecalResponseRoughness())
+		{
+			Metallic = Metallic * this.MetallicSpecularRoughness.a + this.MetallicSpecularRoughness.x;
+			Specular = Specular * this.MetallicSpecularRoughness.a + this.MetallicSpecularRoughness.y;
+			Roughness = Roughness * this.MetallicSpecularRoughness.a + this.MetallicSpecularRoughness.z;
+		}
 
-	HitPayload.Radiance += DecalData.Emissive;
+		HitPayload.Radiance += this.Emissive;
 
-	return min3(DecalData.BaseColor.a, DecalData.WorldNormal.a, DecalData.MetallicSpecularRoughness.a);
-}
+#if PATHTRACING_STRATA_PAYLOAD
+		if (HitPayload.HasDecalResponseColor() || HitPayload.HasDecalResponseRoughness())
+		{
+			// the encode-decode is potentially lossy, so only do it if we actually need to
+			HitPayload.DiffuseColor = BaseColor - BaseColor * Metallic;
+			HitPayload.SpecularColor = lerp(DielectricSpecularToF0(Specular), BaseColor, Metallic);
+			HitPayload.RoughnessData.x = Roughness;
+		}
+#else
+		// no need for an encode-decode here, so safe to overwrite even if nothing changes
+		HitPayload.BaseColor = BaseColor;
+		HitPayload.Metallic = Metallic;
+		HitPayload.Specular = Specular;
+		HitPayload.Roughness = Roughness;
+#endif
+		return HasData();
+	}
+};
 
 void CombineDecal(FDecalShaderPayload DecalPayload, inout FDecalData DecalData)
 {
@@ -349,7 +377,6 @@ void EvaluateDecals(FPathState PathState, uint InstanceInclusionMask, FPackedPat
 	}
 #endif
 }
-#endif // STRATA_ENABLED
 
 FPathTracingPayload TraceTransparentRay(inout FPathState PathState, bool bIsCameraRay, bool bLastBounce, bool bIncludeEmission, uint NumVisibleLights, inout FVolumeSegment VolumeSegment)
 {
@@ -401,9 +428,6 @@ FPathTracingPayload TraceTransparentRay(inout FPathState PathState, bool bIsCame
 			PathState.Ray.GetNativeDesc(),
 			PackedPayload);
 
-#if STRATA_ENABLED
-		// STRATA_TODO: figure out decals
-#else
 		FDecalData DecalData;
 		DecalData.BaseColor = float4(0, 0, 0, 1);
 		DecalData.WorldNormal = float4(0, 0, 0, 1);
@@ -415,29 +439,36 @@ FPathTracingPayload TraceTransparentRay(inout FPathState PathState, bool bIsCame
 			EvaluateDecals(PathState, InstanceInclusionMask, PackedPayload, DecalData);
 
 #if USE_DBUFFER
-			const float DecalThroughput = min3(DecalData.BaseColor.a, DecalData.WorldNormal.a, DecalData.MetallicSpecularRoughness.a);
-			if (EnableDBuffer && PackedPayload.UsesDBufferLookup() && DecalThroughput < 1.0f)
+			if (PackedPayload.UsesDBufferLookup() && DecalData.HasData())
 			{
-				// Retrace ray with DBuffer data
-
-				PackedPayload = InitPathTracingPayload(PathState.FirstScatterType, ApproximateCaustics ? PathState.PathRoughness : 0.0);
-
-				PackedPayload.SetDBufferA(DecalData.BaseColor);
-				PackedPayload.SetDBufferB(DecalData.WorldNormal);
-				PackedPayload.SetDBufferC(DecalData.MetallicSpecularRoughness);
-				TraceRay(
-					TLAS,
-					RayFlags,
-					InstanceInclusionMask,
-					RAY_TRACING_SHADER_SLOT_MATERIAL,
-					RAY_TRACING_NUM_SHADER_SLOTS,
-					MissShaderIndex,
-					PathState.Ray.GetNativeDesc(),
-					PackedPayload);
+				if (EnableDBuffer)
+				{
+					// Retrace ray with DBuffer data
+					PackedPayload = InitPathTracingPayload(PathState.FirstScatterType, ApproximateCaustics ? PathState.PathRoughness : 0.0);
+
+					PackedPayload.SetDBufferA(DecalData.BaseColor);
+					PackedPayload.SetDBufferB(DecalData.WorldNormal);
+					PackedPayload.SetDBufferC(DecalData.MetallicSpecularRoughness);
+					TraceRay(
+						TLAS,
+						RayFlags,
+						InstanceInclusionMask,
+						RAY_TRACING_SHADER_SLOT_MATERIAL,
+						RAY_TRACING_NUM_SHADER_SLOTS,
+						MissShaderIndex,
+						PathState.Ray.GetNativeDesc(),
+						PackedPayload);
+				}
+				else
+				{
+					// The user asked for a dbuffer lookup, but we have disabled the double-trace.
+					// Set all the decal response flags so we at least pickup some kind of response
+					// even if it won't be exactly what the user had in mind.
+					PackedPayload.SetFlag(PATH_TRACING_PAYLOAD_OUTPUT_FLAG_DECAL_RESPONSE_COLOR | PATH_TRACING_PAYLOAD_OUTPUT_FLAG_DECAL_RESPONSE_NORMAL | PATH_TRACING_PAYLOAD_OUTPUT_FLAG_DECAL_RESPONSE_ROUGHNESS);
+				}
 			}
 #endif
 		}
-#endif // STRATA_ENABLED
 
 		if (PackedPayload.IsMiss() && VolumeIsectList.HitBlocker())
 		{
@@ -465,7 +496,7 @@ FPathTracingPayload TraceTransparentRay(inout FPathState PathState, bool bIsCame
 				LightHitPayload.ShadingModelID = SHADINGMODELID_UNLIT;
 				// light hits are fully transparent to match how the light loop works (the light surface does not cast a shadow)
 				LightHitPayload.BSDFOpacity = 1.0;
-				LightHitPayload.TransparencyColor = 0.0;
+				LightHitPayload.TransparencyColor = 1.0;
 				LightHitPayload.PrimitiveLightingChannelMask = 7;
 				LightHitPayload.SetFrontFace();
 				PackedPayload = PackPathTracingPayload(LightHitPayload);
@@ -611,30 +642,26 @@ FPathTracingPayload TraceTransparentRay(inout FPathState PathState, bool bIsCame
 		}
 
 		// Unpack the payload
-		FPathTracingPayload HitPayload = UnpackPathTracingPayloadMaterial(PackedPayload, PathState.Ray, bIsCameraRay && ApplyDiffuseSpecularOverrides);
-
-#if STRATA_ENABLED
-		// STRATA_TODO: figure out decals
-#else // STRATA_ENABLED
-		const float DecalThroughput = ApplyDecal(HitPayload, DecalData);
-
-		if (DecalThroughput < 1.0f)
+		FPathTracingPayload HitPayload = UnpackPathTracingPayload(PackedPayload, PathState.Ray);
+		if (DecalData.ApplyDecal(HitPayload))
 		{
 			// if we ran any decal logic, update the payload
 			if (ApproximateCaustics)
 			{
+#if PATHTRACING_STRATA_PAYLOAD
+				Payload.RoughnessData.x = max(Payload.RoughnessData.x, PathState.PathRoughness);
+#else
 				Payload.Roughness = max(Payload.Roughness, PathState.PathRoughness);
+#endif
 			}
 			HitPayload.WorldNormal = AdjustShadingNormal(HitPayload.WorldNormal, HitPayload.WorldGeoNormal, PathState.Ray.Direction);
+#if !PATHTRACING_STRATA_PAYLOAD
 			HitPayload.DiffuseColor = HitPayload.BaseColor - HitPayload.BaseColor * HitPayload.Metallic;
 			HitPayload.SpecularColor = ComputeF0(HitPayload.Specular, HitPayload.BaseColor, HitPayload.Metallic);
-			if (bIsCameraRay)
-			{
-				HitPayload.DiffuseColor = HitPayload.DiffuseColor * View.DiffuseOverrideParameter.w + View.DiffuseOverrideParameter.xyz;
-				HitPayload.SpecularColor = HitPayload.SpecularColor * View.SpecularOverrideParameter.w + View.SpecularOverrideParameter.xyz;
-			}
+#endif
 		}
-#endif // STRATA_ENABLED
+		AdjustPayloadAfterUnpack(HitPayload, bIsCameraRay && ApplyDiffuseSpecularOverrides);
+
 
 		// account for Beer's law through the currently active medium
 		// TODO: merge this with the volume raymarching?
diff --git a/Engine/Shaders/Private/PathTracing/PathTracingMaterialHitShader.usf b/Engine/Shaders/Private/PathTracing/PathTracingMaterialHitShader.usf
--- a/Engine/Shaders/Private/PathTracing/PathTracingMaterialHitShader.usf
+++ b/Engine/Shaders/Private/PathTracing/PathTracingMaterialHitShader.usf
@@ -1030,6 +1030,26 @@ RAY_TRACING_ENTRY_CLOSEST_HIT(PathTracingMaterialCHS,
 		// terminal case for the conditionals above
 	}
 
+	// adjust after everything else (because SSS profile case needs to decode the dual spec info from the unmodified roughness)
+	Payload.Roughness = AdjustMaterialRoughness(Payload.Roughness, PackedPayload.GetPathRoughness());
+
+#endif // STRATA_ENABLED
+
+	// Only opaque surfaces can recieve decals. This is both for compatibility with raster and improving performance when we have stacks of transparent surfaces
+#if MATERIALBLENDING_SOLID || MATERIALBLENDING_MASKED
+	if ((PrimitiveFlags & PRIMITIVE_SCENE_DATA_FLAG_DECAL_RECEIVER) != 0)
+	{
+#if USE_DBUFFER
+		Payload.SetDecalReceiver(MATERIALDECALRESPONSEMASK);
+#	if MATERIAL_USES_DECAL_LOOKUP
+		Payload.SetUseDBufferLookup();
+#	endif
+#else
+		Payload.SetDecalReceiver(0x07);
+#endif
+	}
+#endif
+
 #if MATERIALBLENDING_ALPHAHOLDOUT
 	// the material is already a holdout, no need to do anything else
 #else
@@ -1055,27 +1075,6 @@ RAY_TRACING_ENTRY_CLOSEST_HIT(PathTracingMaterialCHS,
 	}
 #endif
 
-	// Only opaque surfaces can recieve decals. This is both for compatibility with raster and improving performance when we have stacks of transparent surfaces
-#if MATERIALBLENDING_SOLID || MATERIALBLENDING_MASKED
-	if ((PrimitiveFlags & PRIMITIVE_SCENE_DATA_FLAG_DECAL_RECEIVER) != 0)
-	{
-#if USE_DBUFFER
-		Payload.SetDecalReceiver(MATERIALDECALRESPONSEMASK);
-#	if MATERIAL_USES_DECAL_LOOKUP
-		Payload.SetUseDBufferLookup();
-#	endif
-#else
-		Payload.SetDecalReceiver(0x07);
-#endif
-	}
-#endif
-
-	// adjust after everything else (because SSS profile case needs to decode the dual spec info from the unmodified roughness)
-	Payload.Roughness = AdjustMaterialRoughness(Payload.Roughness, PackedPayload.GetPathRoughness());
-
-#endif
-
-
 	PackedPayload = PackPathTracingPayload(Payload);
 }
 
diff --git a/Engine/Shaders/Private/RayTracing/RayTracingDecalMaterialShader.usf b/Engine/Shaders/Private/RayTracing/RayTracingDecalMaterialShader.usf
--- a/Engine/Shaders/Private/RayTracing/RayTracingDecalMaterialShader.usf
+++ b/Engine/Shaders/Private/RayTracing/RayTracingDecalMaterialShader.usf
@@ -6,6 +6,8 @@
 // This should be good enough for ray tracing and avoids having to bind an extra buffer
 #define EYE_ADAPTATION_PREV_FRAME_EXPOSURE 1
 
+#define STRATA_INLINE_SHADING 1
+
 #include "/Engine/Private/Common.ush"
 #include "/Engine/Private/RayTracing/RayTracingCommon.ush"
 
@@ -21,6 +23,75 @@
 
 #include "/Engine/Generated/Material.ush"
 
+
+struct FDecalData
+{
+	float Opacity;
+	float3 BaseColor;
+	float3 WorldNormal;
+	float3 MetallicSpecularRoughness;
+	float3 Emissive;
+
+	void WriteToPayload(inout FDecalShaderPayload Payload, uint WriteFlags, float OpacityMult)
+	{
+		Opacity *= OpacityMult;
+		if (WriteFlags & DECAL_WRITE_BASE_COLOR_FLAG)
+		{
+			Payload.SetBaseColor(BaseColor * Opacity);
+		}
+		if (WriteFlags & DECAL_WRITE_NORMAL_FLAG)
+		{
+			Payload.SetWorldNormal(WorldNormal * Opacity);
+		}
+		if (WriteFlags & DECAL_WRITE_ROUGHNESS_SPECULAR_METALLIC_FLAG)
+		{
+			Payload.SetMetallicSpecularRoughness(MetallicSpecularRoughness * Opacity);
+		}
+		if (WriteFlags & DECAL_WRITE_EMISSIVE_FLAG)
+		{
+			Payload.SetEmissive(Emissive * Opacity);
+		}
+		Payload.SetTransparency(1.0f - Opacity);
+		Payload.SetFlags(WriteFlags);
+	}
+};
+
+
+
+FDecalData GetDecalData(FPixelMaterialInputs PixelMaterialInputs, FMaterialPixelParameters MaterialParameters)
+{
+	FDecalData Result = (FDecalData)0;
+#if STRATA_ENABLED
+	FStrataPixelHeader StrataHeader = MaterialParameters.GetFrontStrataHeader();
+
+	if (StrataHeader.StrataTree.BSDFCount > 0)
+	{
+		FStrataAddressing NullStrataAddressing = (FStrataAddressing)0;	// Fake unused in StrataCreateBSDFContext when using Forward inline shading
+		const float3 V = MaterialParameters.CameraVector;
+
+		// Update tree (coverage/transmittance/luminace weights)
+		const FStrataIntegrationSettings Settings = InitStrataIntegrationSettings(false /*bForceFullyRough*/, false /*bRoughDiffuseEnabled*/, 0, false/*bRoughnessTracking*/);
+		StrataHeader.StrataUpdateTree(NullStrataAddressing, V, Settings);
+
+		// Since the Convert-To-Decal node forces parameter blending for the entire Strata tree, the BSDF we are interested in will be in 0.
+		FStrataDBuffer DBuffer = StrataHeader.StrataConvertToDBuffer(StrataHeader.StrataTree.BSDFs[0]);
+		Result.Opacity = DBuffer.Coverage;
+		Result.BaseColor = DBuffer.BaseColor;
+		Result.WorldNormal = DBuffer.WorldNormal;
+		Result.MetallicSpecularRoughness = float3(DBuffer.Metallic, DBuffer.Specular, DBuffer.Roughness);
+		Result.Emissive = DBuffer.Emissive;
+	}
+#else // STRATA_ENABLED
+	Result.Opacity = GetMaterialOpacity(PixelMaterialInputs);
+	Result.BaseColor = GetMaterialBaseColor(PixelMaterialInputs);
+	Result.WorldNormal = MaterialParameters.WorldNormal;
+	Result.MetallicSpecularRoughness = float3(GetMaterialMetallic(PixelMaterialInputs), GetMaterialSpecular(PixelMaterialInputs), GetMaterialRoughness(PixelMaterialInputs));
+	Result.Emissive = GetMaterialEmissive(PixelMaterialInputs);
+#endif // STRATA_ENABLED
+	return Result;
+}
+
+
 #if RAYCALLABLESHADER
 RAY_TRACING_ENTRY_CALLABLE(RayTracingDecalMaterialShader, FDecalShaderPayload, Payload)
 {
@@ -69,30 +140,9 @@ RAY_TRACING_ENTRY_CALLABLE(RayTracingDecalMaterialShader, FDecalShaderPayload, P
 	CalcPixelMaterialInputs(MaterialParameters, PixelMaterialInputs);
 
 	const float DecalFading = saturate(4 - 4 * abs(SwizzlePos.z * 2 - 1)) * DecalParams.x;
-	const float Opacity = GetMaterialOpacity(PixelMaterialInputs) * DecalFading;
-	Payload.SetTransparency(1.0f - Opacity);
-
-	if (DecalWriteFlags & DECAL_WRITE_BASE_COLOR_FLAG)
-	{
-		Payload.SetBaseColor(GetMaterialBaseColor(PixelMaterialInputs) * Opacity);
-	}
-
-	if (DecalWriteFlags & DECAL_WRITE_NORMAL_FLAG)
-	{
-		Payload.SetWorldNormal(MaterialParameters.WorldNormal * Opacity);
-	}
 
-	if (DecalWriteFlags & DECAL_WRITE_ROUGHNESS_SPECULAR_METALLIC_FLAG)
-	{
-		Payload.SetMetallicSpecularRoughness(float3(GetMaterialMetallic(PixelMaterialInputs), GetMaterialSpecular(PixelMaterialInputs), GetMaterialRoughness(PixelMaterialInputs)) * Opacity);
-	}
-
-	if (DecalWriteFlags & DECAL_WRITE_EMISSIVE_FLAG)
-	{
-		Payload.SetEmissive(GetMaterialEmissive(PixelMaterialInputs) * Opacity);
-	}
-
-	Payload.SetFlags(DecalWriteFlags);
+	FDecalData Data = GetDecalData(PixelMaterialInputs, MaterialParameters);
+	Data.WriteToPayload(Payload, DecalWriteFlags, DecalFading);
 }
 #endif // RAYCALLABLESHADER
 
@@ -139,6 +189,8 @@ RAY_TRACING_ENTRY_CLOSEST_HIT(RayTracingDecalMaterialCHS,
 		const float4 ScreenPosition = SvPositionToResolvedScreenPosition(SvPosition);
 		const bool bIsFrontFace = HitKind() == HIT_KIND_TRIANGLE_FRONT_FACE;
 
+		MaterialParameters.CameraVector = -WorldRayDirection();
+
 		// #dxr_todo: UE-72130 support world position offset 
 		// #if USE_WORLD_POSITION_EXCLUDING_SHADER_OFFSETS
 		// CalcMaterialParametersEx(MaterialParameters, PixelMaterialInputs, SvPosition, ScreenPosition, bIsFrontFace, TranslatedWorldPosition, BasePassInterpolants.PixelPositionExcludingWPO);
@@ -147,32 +199,9 @@ RAY_TRACING_ENTRY_CLOSEST_HIT(RayTracingDecalMaterialCHS,
 		// #endif
 	}
 
-	const float Opacity = GetMaterialOpacity(PixelMaterialInputs);
-	Payload.SetTransparency(1.0f - Opacity);
-
-	if (DECAL_PAYLOAD_FLAGS & DECAL_WRITE_BASE_COLOR_FLAG)
-	{
-		Payload.SetBaseColor(GetMaterialBaseColor(PixelMaterialInputs) * Opacity);
-	}
-
-	if (DECAL_PAYLOAD_FLAGS & DECAL_WRITE_NORMAL_FLAG)
-	{
-		Payload.SetWorldNormal(MaterialParameters.WorldNormal * Opacity);
-	}
-
-	if (DECAL_PAYLOAD_FLAGS & DECAL_WRITE_ROUGHNESS_SPECULAR_METALLIC_FLAG)
-	{
-		Payload.SetMetallicSpecularRoughness(float3(GetMaterialMetallic(PixelMaterialInputs), GetMaterialSpecular(PixelMaterialInputs), GetMaterialRoughness(PixelMaterialInputs)) * Opacity);
-	}
-
-	if (DECAL_PAYLOAD_FLAGS & DECAL_WRITE_EMISSIVE_FLAG)
-	{
-		Payload.SetEmissive(GetMaterialEmissive(PixelMaterialInputs) * Opacity);
-	}
-
-	Payload.SetFlags(DECAL_PAYLOAD_FLAGS);
+	FDecalData Data = GetDecalData(PixelMaterialInputs, MaterialParameters);
+	Data.WriteToPayload(Payload, DECAL_PAYLOAD_FLAGS, 1.0);
 
-	return;
 }
 
 #if USE_MATERIAL_ANY_HIT_SHADER
@@ -203,7 +232,7 @@ RAY_TRACING_ENTRY_ANY_HIT(RayTracingDecalMaterialAHS,
 		CalcMaterialParametersEx(MaterialParameters, PixelMaterialInputs, SvPosition, ScreenPosition, bIsFrontFace, TranslatedWorldPosition, TranslatedWorldPosition);
 	}
 
-	float Opacity = GetMaterialOpacity(PixelMaterialInputs);
+	const float Opacity = GetDecalData(PixelMaterialInputs, MaterialParameters).Opacity;
 	if (Opacity <= 0.00001f)
 	{
 		IgnoreHit();
