Implements r.TSR.ShadingRejection.ExposureOffset

Turns out our VisualizeTexture tool didn't apply proper linear -> sRGB color transformation.
This hidden a long standing problem in TSR that had the wrong idea of the final backbuffer contributions.
This meant TSR's MeasureBackbufferLDRQuantizationError() was overly large in shadows, which caused
subtle VFX in shadows that have enough contribution to be visible in backbuffer to ghost with TSR thinking
their contribution was about or smaller than lowest significant bit in backbuffer.

diff --git a/Engine/Shaders/Private/TemporalSuperResolution/TSRComputeMoireLuma.usf b/Engine/Shaders/Private/TemporalSuperResolution/TSRComputeMoireLuma.usf
--- a/Engine/Shaders/Private/TemporalSuperResolution/TSRComputeMoireLuma.usf
+++ b/Engine/Shaders/Private/TemporalSuperResolution/TSRComputeMoireLuma.usf
@@ -12,6 +12,8 @@
 
 //------------------------------------------------------- PARAMETERS
 
+float PerceptionAdd;
+
 Texture2D<tsr_half3> SceneColorTexture;
 RWTexture2D<tsr_half> MoireLumaOutput;
 
@@ -56,7 +58,7 @@ void MainCS(
 				uint2 InputPixelPos = OutputPixelOrigin + PixelOffset;
 			
 				tsr_half3 Color = ColorArray[x + y * TILE_ROWS];
-				Color *= rcp(tsr_half(1.0) + Color);
+				Color *= rcp(tsr_half(PerceptionAdd) + Color);
 
 				InputPixelPos.x = select(all(InputPixelPos < InputInfo_ViewportMax), InputPixelPos.x, uint(~0));
 				
diff --git a/Engine/Shaders/Private/TemporalSuperResolution/TSRRejectShading.usf b/Engine/Shaders/Private/TemporalSuperResolution/TSRRejectShading.usf
--- a/Engine/Shaders/Private/TemporalSuperResolution/TSRRejectShading.usf
+++ b/Engine/Shaders/Private/TemporalSuperResolution/TSRRejectShading.usf
@@ -98,6 +98,7 @@ float3 HistoryGuideQuantizationError;
 float FlickeringFramePeriod;
 float TheoricBlendFactor;
 uint TileOverscan;
+float PerceptionAdd;
 
 
 Texture2D<tsr_halfC> InputTexture;
@@ -242,10 +243,6 @@ void MainCS(
 		OriginalInput = min(OriginalInput, tsr_simd_halfC::Const(LargestSceneColorRGB));
 	#endif
 
-	// Convert input to LDR
-	OriginalOpaqueInput = OriginalOpaqueInput * rcp(OriginalOpaqueInput + tsr_simd_halfC::Const(1.0));
-	OriginalInput = OriginalInput * rcp(OriginalInput + tsr_simd_halfC::Const(1.0));
-	
 	// Analyses the moire pattern overtime.
 	#if CONFIG_MOIRE_DETECTION
 	tsr_simd_half4 FinalMoire = tsr_simd_half4::Const(0.0);
@@ -265,6 +262,9 @@ void MainCS(
 		tsr_simd_half InputLumaModification;
 		{
 			const tsr_half InputLumaEncodingError = tsr_half(0.5 / 255.0);
+			
+			OriginalOpaqueInput = OriginalOpaqueInput * rcp(OriginalOpaqueInput + tsr_simd_halfC::Const(tsr_half(PerceptionAdd)));
+
 
 			tsr_simd_half SceneColorLuma;
 			UNROLL_N(SIMD_SIZE)
@@ -407,34 +407,43 @@ void MainCS(
 	}
 	#endif
 
-	// Top secret sauce with paper and salt.
-	tsr_simd_halfC ClampedInput = OriginalInput;
-	tsr_simd_halfC ClampedHistory = OriginalHistory;
-	#if CONFIG_FIREFLY_REMOVAL
-	{
-		ClampedInput   = ClampFireFliersWithGuide(OriginalInput  , /* Guide = */ OriginalHistory);
-		ClampedHistory = ClampFireFliersWithGuide(OriginalHistory, /* Guide = */ OriginalInput);
-	}
-	#endif
-
-	// Filter clamped input and history to increase quantization of the rejection detection
-	tsr_simd_halfC FilteredInput   = Blur3x3(ClampedInput);
-	tsr_simd_halfC FilteredHistory = Blur3x3(ClampedHistory);
-	
 	// Measure rejection
 	tsr_simd_half RejectionBlendFinal = tsr_simd_half::Const(1.0);
 	tsr_simd_half RejectionClampBlend = tsr_simd_half::Const(1.0);
 	#if 1
 	{
 		const tsr_half FilteringWeight = rcp(tsr_half(1.0 + 4 * 0.5 + 4 * 0.25));
+		const tsr_simd_halfC BackbufferQuantizationErrorVector = tsr_simd_halfC::Const(MeasureBackbufferLDRQuantizationError());
+		
+		// Convert input to LDR
+		tsr_simd_halfC ExposedInput = OriginalInput * rcp(OriginalInput + tsr_simd_halfC::Const(tsr_half(PerceptionAdd)));
+	
+		// Change the exposure of history for the shading rejection
+		tsr_simd_halfC ExposedHistory = OriginalHistory;
+		{
+			ExposedHistory = ExposedHistory * rcp(tsr_simd_halfC::Const(1.0) - ExposedHistory);
+			ExposedHistory = ExposedHistory * rcp(ExposedHistory + tsr_simd_halfC::Const(tsr_half(PerceptionAdd)));
+		}
 
-		tsr_simd_halfC BackbufferError = tsr_simd_halfC::Const(MeasureBackbufferLDRQuantizationError());
+		// Top secret sauce with paper and salt.
+		tsr_simd_halfC ClampedInput = ExposedInput;
+		tsr_simd_halfC ClampedHistory = ExposedHistory;
+		#if CONFIG_FIREFLY_REMOVAL
+		{
+			ClampedInput   = ClampFireFliersWithGuide(ClampedInput  , /* Guide = */ ExposedHistory);
+			ClampedHistory = ClampFireFliersWithGuide(ClampedHistory, /* Guide = */ ExposedInput);
+		}
+		#endif
 
+		// Filter clamped input and history to increase quantization of the rejection detection
+		tsr_simd_halfC FilteredInput   = Blur3x3(ClampedInput);
+		tsr_simd_halfC FilteredHistory = Blur3x3(ClampedHistory);
+	
 		tsr_simd_halfC FilteredBoxMin;
 		tsr_simd_halfC FilteredBoxMax;
 		MinMax3x3(FilteredInput, /* out */ FilteredBoxMin, /* out */ FilteredBoxMax);
-		FilteredBoxMin = FilteredBoxMin - BackbufferError;
-		FilteredBoxMax = FilteredBoxMax + BackbufferError;
+		FilteredBoxMin = FilteredBoxMin - BackbufferQuantizationErrorVector;
+		FilteredBoxMax = FilteredBoxMax + BackbufferQuantizationErrorVector;
 		
 		tsr_simd_halfC FlickeringClampedFilteredHistory = fastClamp(FilteredHistory, FilteredBoxMin, FilteredBoxMax);
 
@@ -470,10 +479,9 @@ void MainCS(
 		
 		tsr_simd_halfC OriginalBoxMin;
 		tsr_simd_halfC OriginalBoxMax;
-		MinMax3x3(OriginalInput, /* out */ OriginalBoxMin, /* out */ OriginalBoxMax);
-		tsr_simd_halfC BoxSize = (OriginalBoxMax - OriginalBoxMin + (BackbufferError + BackbufferError)) * FilteringWeight;
+		MinMax3x3(ExposedInput, /* out */ OriginalBoxMin, /* out */ OriginalBoxMax);
+		tsr_simd_halfC BoxSize = (OriginalBoxMax - OriginalBoxMin + (BackbufferQuantizationErrorVector + BackbufferQuantizationErrorVector)) * FilteringWeight;
 		
-		const tsr_simd_halfC BackbufferQuantizationErrorVector = tsr_simd_halfC::Const(MeasureBackbufferLDRQuantizationError());
 		tsr_simd_halfC RawClampedEnergy = abs(ClampedFilteredHistory - FilteredHistory);
 
 		tsr_simd_halfC FilteredClampedEnergy;
@@ -518,6 +526,11 @@ void MainCS(
 	// Update the guide for next frame.
 	tsr_simd_halfC FinalGuide;
 	{
+		
+		// Convert input to LDR
+		tsr_simd_halfC ExposedInput = OriginalInput * rcp(OriginalInput + tsr_simd_halfC::Const(tsr_half(1.0)));
+	
+
 		tsr_simd_half BlendFinal;
 		// Completely discard the history
 		UNROLL_N(SIMD_SIZE)
@@ -534,7 +547,7 @@ void MainCS(
 			BlendFinal.SetElement(SimdIndex, max(RejectedBlendFactor, TotalBlend));
 		}
 		
-		FinalGuide = OriginalInput * tsr_simd_halfC::Vectorize(BlendFinal) - OriginalHistory * tsr_simd_halfC::Vectorize(BlendFinal - 1.0);
+		FinalGuide = ExposedInput * tsr_simd_halfC::Vectorize(BlendFinal) - OriginalHistory * tsr_simd_halfC::Vectorize(BlendFinal - 1.0);
 	}
 
 	// Mark the top left corner of the tile for debugging purposes.
diff --git a/Engine/Source/Runtime/Renderer/Private/PostProcess/TemporalSuperResolution.cpp b/Engine/Source/Runtime/Renderer/Private/PostProcess/TemporalSuperResolution.cpp
--- a/Engine/Source/Runtime/Renderer/Private/PostProcess/TemporalSuperResolution.cpp
+++ b/Engine/Source/Runtime/Renderer/Private/PostProcess/TemporalSuperResolution.cpp
@@ -158,6 +158,22 @@ TAutoConsoleVariable<int32> CVarTSRShadingTileOverscan(
 	TEXT(" anti-aliasing scalability settings."),
 	ECVF_Scalability | ECVF_RenderThreadSafe);
 
+TAutoConsoleVariable<float> CVarTSRShadingExposureOffset(
+	TEXT("r.TSR.ShadingRejection.ExposureOffset"), 0.0, // TODO: update engine default with new EngineTest references
+	TEXT("The shading rejection needs to have a representative idea how bright a linear color pixel ends up displayed to the user. ")
+	TEXT("And the shading rejection detect if a color become to changed to be visible in the back buffer by comparing to MeasureBackbufferLDRQuantizationError().\n")
+	TEXT("\n")
+	TEXT("It is important to have TSR's MeasureBackbufferLDRQuantizationError() ends up distributed uniformly across ")
+	TEXT("the range of color intensity or it could otherwise disregard some subtle VFX causing ghostin.\n")
+	TEXT("\n")
+	TEXT("This controls adjusts the exposure of the linear color space solely in the TSR's rejection heuristic, such that higher value ")
+	TEXT("lifts the shadows's LDR intensity, meaning MeasureBackbufferLDRQuantizationError() is decreased in these shadows and increased in ")
+	TEXT("the highlights, control directly.\n")
+	TEXT("\n")
+	TEXT("The best TSR internal buffer to verify this is TSR.Moire.Luma but must be verified in DumpGPU ")
+	TEXT("with the RGB Linear[0;1] source color space against the Tonemaper's output in sRGB source color space.\n"),
+	ECVF_RenderThreadSafe);
+
 TAutoConsoleVariable<int32> CVarTSRRejectionAntiAliasingQuality(
 	TEXT("r.TSR.RejectionAntiAliasingQuality"), 3,
 	TEXT("Controls the quality of TSR's built-in spatial anti-aliasing technology when the history needs to be rejected. ")
@@ -520,6 +536,7 @@ class FTSRComputeMoireLumaCS : public FTSRShader
 
 	BEGIN_SHADER_PARAMETER_STRUCT(FParameters, )
 		SHADER_PARAMETER_STRUCT(FScreenPassTextureViewportParameters, InputInfo)
+		SHADER_PARAMETER(float, PerceptionAdd)
 		SHADER_PARAMETER_RDG_TEXTURE(Texture2D, SceneColorTexture)
 		SHADER_PARAMETER_RDG_TEXTURE_UAV(RWTexture2D, MoireLumaOutput)
 	END_SHADER_PARAMETER_STRUCT()
@@ -705,6 +722,7 @@ class FTSRRejectShadingCS : public FTSRShader
 		SHADER_PARAMETER(float, FlickeringFramePeriod)
 		SHADER_PARAMETER(float, TheoricBlendFactor)
 		SHADER_PARAMETER(int32, TileOverscan)
+		SHADER_PARAMETER(float, PerceptionAdd)
 
 		SHADER_PARAMETER_RDG_TEXTURE(Texture2D, InputTexture)
 		SHADER_PARAMETER_RDG_TEXTURE(Texture2D, InputMoireLumaTexture)
@@ -1098,6 +1116,7 @@ FScreenPassTexture AddTSRComputeMoireLuma(FRDGBuilder& GraphBuilder, FGlobalShad
 	FTSRComputeMoireLumaCS::FParameters* PassParameters = GraphBuilder.AllocParameters<FTSRComputeMoireLumaCS::FParameters>();
 	PassParameters->InputInfo = GetScreenPassTextureViewportParameters(FScreenPassTextureViewport(
 		SceneColor.Texture->Desc.Extent, SceneColor.ViewRect));
+	PassParameters->PerceptionAdd = FMath::Pow(0.5f, CVarTSRShadingExposureOffset.GetValueOnRenderThread());
 	PassParameters->SceneColorTexture = SceneColor.Texture;
 	PassParameters->MoireLumaOutput = GraphBuilder.CreateUAV(MoireLuma.Texture);
 
@@ -1952,6 +1971,7 @@ ITemporalUpscaler::FOutputs AddTemporalSuperResolutionPasses(
 		PassParameters->FlickeringFramePeriod = FlickeringFramePeriod;
 		PassParameters->TheoricBlendFactor = 1.0f / (1.0f + MaxHistorySampleCount / OutputToInputResolutionFractionSquare);
 		PassParameters->TileOverscan = FMath::Clamp(CVarTSRShadingTileOverscan.GetValueOnRenderThread(), 2, GroupTileSize / 2 - 1);
+		PassParameters->PerceptionAdd = FMath::Pow(0.5f, CVarTSRShadingExposureOffset.GetValueOnRenderThread());
 
 		PassParameters->InputTexture = PassInputs.SceneColorTexture;
 		if (PassInputs.MoireInputTexture.IsValid())
