Substrate - tested emissive on decals and it is all working correctly. Minor code clean up.

diff --git a/Engine/Shaders/Private/DecalCommon.ush b/Engine/Shaders/Private/DecalCommon.ush
--- a/Engine/Shaders/Private/DecalCommon.ush
+++ b/Engine/Shaders/Private/DecalCommon.ush
@@ -45,21 +45,20 @@ float DecalCommonOutput(
 	Color *= View.PreExposure;
 #endif
 
+	// Outputing color to MRT0 (can be emissive or any other data set above)
 	Out.MRT[0] = float4(Color, DBuffer.Coverage);
 
 #if DECAL_RENDERTARGETMODE == DECAL_RENDERTARGETMODE_DBUFFER
-
-	//StrataPackDBuffer(DBuffer, Out.MRT[0], Out.MRT[1], Out.MRT[2], Out.MRT[3]);
-	float4 PreMulEmissive_Coverage = 0;
-	StrataPackDBuffer(DBuffer, PreMulEmissive_Coverage, Out.MRT[0], Out.MRT[1], Out.MRT[2]);
+	// If we output to DBuffer, write to all the 3 RTs.
+	StrataPackDBuffer(DBuffer, Out.MRT[0], Out.MRT[1], Out.MRT[2]);
 
 	// This is to set the coverage over the background to an arbitrary value in case premultiplied alpha is used.
 	Out.MRT[0].a = CoverageOverBackground;
 	Out.MRT[1].a = CoverageOverBackground;
 	Out.MRT[2].a = CoverageOverBackground;
 	Out.MRT[3].a = CoverageOverBackground;
-
 #endif
+
 	// Allow shader compiler to cull outputs that we won't use.
 	// Blend state also prevents writing unused outputs, but we want to be sure that we also avoid shader cost.
 #if !DECAL_OUT_MRT0
diff --git a/Engine/Shaders/Private/Strata/Strata.ush b/Engine/Shaders/Private/Strata/Strata.ush
--- a/Engine/Shaders/Private/Strata/Strata.ush
+++ b/Engine/Shaders/Private/Strata/Strata.ush
@@ -3423,13 +3423,11 @@ struct FStrataDBuffer
 
 void StrataPackDBuffer(
 	in FStrataDBuffer In,
-	out float4 Out0,
 	out float4 Out1,
 	out float4 Out2,
 	out float4 Out3)
 {
 	// DBuffer packing needs to be blendable
-	Out0 = float4(In.Emissive,  In.Coverage);
 	Out1 = float4(In.BaseColor, In.Coverage);
 	Out2 = float4(In.WorldNormal * 0.5f + 128.0f / 255.0f, In.Coverage);
 	Out3 = float4(In.Metallic, In.Specular, In.Roughness, In.Coverage);
