Substrate - legacy converted material now weight emissive by opaicty. eve nwhen .

diff --git a/Engine/Shaders/Private/DecalCommon.ush b/Engine/Shaders/Private/DecalCommon.ush
--- a/Engine/Shaders/Private/DecalCommon.ush
+++ b/Engine/Shaders/Private/DecalCommon.ush
@@ -43,6 +43,11 @@ float DecalCommonOutput(
 
 #if DECAL_RENDERTARGETMODE == DECAL_RENDERTARGETMODE_SCENECOLOR
 	Color *= View.PreExposure;
+#if STRATA_USES_CONVERSION_FROM_LEGACY && MATERIALBLENDING_ALPHACOMPOSITE
+	// In legacy mode, emissive is always influenced by opacity, even when AlphaComposite blend mode is used.
+	// This is only here to maintain the legacy behavior after material are converted; not with Substrate itself.
+	Color *= CoverageOverBackground;
+#endif
 #endif
 
 	// Outputing color to MRT0 (can be emissive or any other data set above)
