Make GpuSkinVertexFactory work with the new HLSL generator

diff --git a/Engine/Shaders/Private/GpuSkinVertexFactory.ush b/Engine/Shaders/Private/GpuSkinVertexFactory.ush
--- a/Engine/Shaders/Private/GpuSkinVertexFactory.ush
+++ b/Engine/Shaders/Private/GpuSkinVertexFactory.ush
@@ -365,6 +365,9 @@ FMaterialVertexParameters GetMaterialVertexParameters(FVertexFactoryInput Input,
 
 #endif
 
+#if ENABLE_NEW_HLSL_GENERATOR
+	EvaluateVertexMaterialAttributes(Result);
+#endif
 	return Result;
 }
 
