Modify the path tracing depth write logic to match the rasterizer path. Write the depth when 1) solid or masked material. or 2) when `output depth and velocity` is set  and Opacity >= Opacity mask clip value for translucent, additive and modulate, blending, alpha holdout, alpha composite mode to match the behavior of the rasterizer path.

Rename `output depth` to `output depth and velocity`.

diff --git a/Engine/Shaders/Private/PathTracing/Material/PathTracingMaterialSampling.ush b/Engine/Shaders/Private/PathTracing/Material/PathTracingMaterialSampling.ush
--- a/Engine/Shaders/Private/PathTracing/Material/PathTracingMaterialSampling.ush
+++ b/Engine/Shaders/Private/PathTracing/Material/PathTracingMaterialSampling.ush
@@ -248,7 +248,6 @@ void AccumulateAlbedo(float3 SigmaS, float3 PathThroughput, float PathRoughness,
 struct FDepthContext
 {
 	float3 TranslatedWorldPos;
-	float Transparency;
 	uint HitCount;
 
 #if DEBUG_DEPTH_SAMPLING
@@ -267,23 +266,10 @@ struct FDepthContext
 			return;
 		}
 
-		
-		// The smaller the transparency color, the more opaque the material is.
-		// Ensure Transparency <= 0.5 specified in the material editor contributes to the depth.
-		//const float HalfTransparency = Luminance(PayloadDecodeLDRColor(PayloadEncodeLDRColor(0.5f)));
-		//  asfloat(0x3f0011c5); e.g., opacity >=0.4996 in the material editor for translucent material contributes to depth.
-		const float HalfTransparency = 0.500271142f;
-		const float HitTransparency = Luminance(Payload.TransparencyColor);
-
-		// ensure that 0.5 in the Material editor has a depth contribution.
-		const bool bShouldAcceptNewDepth =
-			Transparency > HalfTransparency && // current depth is not valid
-			HitTransparency < Transparency;    // since it's not valid, we should accept hit with smaller transparency color
-
-		if (bShouldAcceptNewDepth || HitCount == 0)
+		// 2. Matching depth output to the rasterizer.
+		if (Payload.ShouldOutputDepth() && HitCount == 0)
 		{
 			TranslatedWorldPos = Payload.TranslatedWorldPos;
-			Transparency = HitTransparency;
 			HitCount++;
 		}
 
diff --git a/Engine/Shaders/Private/PathTracing/PathTracingCommon.ush b/Engine/Shaders/Private/PathTracing/PathTracingCommon.ush
--- a/Engine/Shaders/Private/PathTracing/PathTracingCommon.ush
+++ b/Engine/Shaders/Private/PathTracing/PathTracingCommon.ush
@@ -20,7 +20,7 @@
 #define PATH_TRACING_PAYLOAD_OUTPUT_FLAG_DECAL_RESPONSE_NORMAL        (1 << 4) // Indicates that ray has hit a primitive that receives decals normal. Set by closest hit shader.
 #define PATH_TRACING_PAYLOAD_OUTPUT_FLAG_DECAL_RESPONSE_ROUGHNESS     (1 << 5) // Indicates that ray has hit a primitive that receives decals metallic/specular/roughness. Set by closest hit shader.
 #define PATH_TRACING_PAYLOAD_OUTPUT_FLAG_USE_DBUFFER_LOOKUP           (1 << 6) // Indicates that ray has hit a primitive that uses DBufferTextureLookup. Set by closest hit shader.
-#define PATH_TRACING_PAYLOAD_OUTPUT_FLAG_UNUSED7                      (1 << 7) // Free for future use
+#define PATH_TRACING_PAYLOAD_OUTPUT_FLAG_OUTOUT_DEPTH				  (1 << 7) // Indicates that ray has hit a primitive that needs to output depth. Set by closest hit shader and used in primary ray only.
 #define PATH_TRACING_PAYLOAD_OUTPUT_FLAG_UNUSED8                      (1 << 8) // Free for future use
 
 #define PATH_TRACING_PAYLOAD_INPUT_FLAG_RAY_TYPE_MASK                 (7 << 0) // Reserve 3 bits for ray-mask
@@ -162,6 +162,9 @@ struct FPackedPathTracingPayload : FMinimalPayload
 	}
 	bool UsesDBufferLookup() { return (GetFlags() & PATH_TRACING_PAYLOAD_OUTPUT_FLAG_USE_DBUFFER_LOOKUP) != 0; }
 
+	void SetOutputDepth() { SetFlag(PATH_TRACING_PAYLOAD_OUTPUT_FLAG_OUTOUT_DEPTH); }
+	bool ShouldOutputDepth() { return (GetFlags() & PATH_TRACING_PAYLOAD_OUTPUT_FLAG_OUTOUT_DEPTH) != 0; }
+
 	// These method are meant to be used only when IsVisibilityRay() is true
 	float3 GetRayThroughput()
 	{
@@ -363,6 +366,9 @@ struct FPathTracingPayload : FMinimalPayload
 	void SetMaterialTwoSided() { Flags |= PATH_TRACING_PAYLOAD_OUTPUT_FLAG_TWO_SIDED_MATERIAL; }
 	bool IsMaterialTwoSided() { return (Flags & PATH_TRACING_PAYLOAD_OUTPUT_FLAG_TWO_SIDED_MATERIAL) != 0; }
 
+	void SetOutputDepth() { Flags |= PATH_TRACING_PAYLOAD_OUTPUT_FLAG_OUTOUT_DEPTH; }
+	bool ShouldOutputDepth() { return Flags & PATH_TRACING_PAYLOAD_OUTPUT_FLAG_OUTOUT_DEPTH; }
+
 #if IS_PAYLOAD_ENABLED(RT_PAYLOAD_TYPE_GPULIGHTMASS)
 	float3 GetBaseColor() { return BaseColor; }
 	float3 GetSubsurfaceColor() { return SubsurfaceColor; }
diff --git a/Engine/Shaders/Private/PathTracing/PathTracingMaterialHitShader.usf b/Engine/Shaders/Private/PathTracing/PathTracingMaterialHitShader.usf
--- a/Engine/Shaders/Private/PathTracing/PathTracingMaterialHitShader.usf
+++ b/Engine/Shaders/Private/PathTracing/PathTracingMaterialHitShader.usf
@@ -1075,6 +1075,18 @@ RAY_TRACING_ENTRY_CLOSEST_HIT(PathTracingMaterialCHS,
 	}
 #endif
 
+	// Output depth for solid, masked 
+	// or when `output depth and velocity' is set and Opacity >= Opacity mask clip value to match the behavior of the rasterizer path. 
+#if (MATERIALBLENDING_SOLID || MATERIALBLENDING_MASKED)
+	Payload.SetOutputDepth();
+#elif (TRANSLUCENT_WRITING_VELOCITY && \
+	(MATERIALBLENDING_TRANSLUCENT || MATERIALBLENDING_ADDITIVE || MATERIALBLENDING_MODULATE || MATERIALBLENDING_ALPHACOMPOSITE || MATERIALBLENDING_ALPHAHOLDOUT))
+	if (GetMaterialOpacity(PixelMaterialInputs) >= GetMaterialOpacityMaskClipValue())
+	{
+		Payload.SetOutputDepth();
+	}
+#endif
+
 	PackedPayload = PackPathTracingPayload(Payload);
 }
 
diff --git a/Engine/Source/Runtime/Engine/Classes/Materials/Material.h b/Engine/Source/Runtime/Engine/Classes/Materials/Material.h
--- a/Engine/Source/Runtime/Engine/Classes/Materials/Material.h
+++ b/Engine/Source/Runtime/Engine/Classes/Materials/Material.h
@@ -492,7 +492,7 @@ public:
 
 	/**
 	 * If BlendMode is BLEND_Masked, the surface is not rendered where OpacityMask < OpacityMaskClipValue.
-	 * If BlendMode is BLEND_Translucent, BLEND_Additive, or BLEND_Modulate, and "Output Velocity" is enabled,
+	 * If BlendMode is BLEND_Translucent, BLEND_Additive, or BLEND_Modulate, and "Output Depth and Velocity" is enabled,
 	 * the object velocity is not rendered where Opacity < OpacityMaskClipValue.
 	 */
 	UPROPERTY(EditAnywhere, Category = Material, AdvancedDisplay)
@@ -946,7 +946,7 @@ public:
 	uint8 bComputeFogPerPixel : 1;
 
 	/** When true, translucent materials will output motion vectors and write to depth buffer in velocity pass. */
-	UPROPERTY(EditAnywhere, Category = Translucency, meta = (DisplayName = "Output Velocity"))
+	UPROPERTY(EditAnywhere, Category = Translucency, meta = (DisplayName = "Output Depth and Velocity"))
 	uint8 bOutputTranslucentVelocity : 1;
 
 	/** If true the compilation environment will be changed to remove the global COMPILE_SHADERS_FOR_DEVELOPMENT flag. */
