Path-tracing support for heterogeneous volumes

Uses the Voxel Grid Pipeline to create global acceleration structures to represent heterogeneous volumes in the path tracer. Creation of global acceleration structures is memory intensive. Please refer to the corresponding Voxel Grid Pipeline CL 25170595 for instructions on parameter tuning to control memory and quality.

r.PathTracing.HeterogeneousVolumes (default = 0) Enables heterogeneous volume rendering in the path tracer
r.PathTracing.HeterogeneousVolumes.RebuildEveryFrame (default = 1) Rebuilds the acceleration structures every frame.

diff --git a/Engine/Shaders/Private/PathTracing/PathTracingCore.ush b/Engine/Shaders/Private/PathTracing/PathTracingCore.ush
--- a/Engine/Shaders/Private/PathTracing/PathTracingCore.ush
+++ b/Engine/Shaders/Private/PathTracing/PathTracingCore.ush
@@ -11,6 +11,7 @@
 // In ordinary scenes, the most likely occurance of this problem is when combining hair and volumetric atmosphere as rays from several Km can travel back towards the hair and cause long delays or even TDRs
 #define PATH_TRACING_SKIP_HAIR_DISTANCE      1e6
 
+#include "../HeterogeneousVolumes/HeterogeneousVolumesVoxelGridTypes.ush"
 #include "../Common.ush"
 #include "../PostProcessCommon.ush"
 #include "../RectLight.ush"
@@ -378,7 +379,7 @@ void EvaluateDecals(FPathState PathState, uint InstanceInclusionMask, FPackedPat
 #endif
 }
 
-FPathTracingPayload TraceTransparentRay(inout FPathState PathState, bool bIsCameraRay, bool bLastBounce, bool bIncludeEmission, uint NumVisibleLights, inout FVolumeSegment VolumeSegment)
+FPathTracingPayload TraceTransparentRay(inout FPathState PathState, int Bounce, bool bIsCameraRay, bool bLastBounce, bool bIncludeEmission, uint NumVisibleLights, inout FVolumeSegment VolumeSegment)
 {
 	const uint RayFlags = bIsCameraRay && EnableCameraBackfaceCulling ? RAY_FLAG_CULL_BACK_FACING_TRIANGLES : 0;
 	uint InstanceInclusionMask = bIsCameraRay ? PATHTRACER_MASK_CAMERA : PATHTRACER_MASK_INDIRECT;
@@ -622,7 +623,7 @@ FPathTracingPayload TraceTransparentRay(inout FPathState PathState, bool bIsCame
 			}
 #if PATH_TRACING_USE_TRANSMITTANCE_FUNC
 			// analytically handle the transmittance to the next surface (which will be higher quality than the transmittance implicitly computed above and prevent noise with emissive surfaces in volumes)
-			PathState.PathThroughput = VolumeGetTransmittance(PathState.PathThroughput, PathState.Ray.Origin, PathState.Ray.Direction, Interval, PathState.RandSequence);
+			PathState.PathThroughput = VolumeGetTransmittance(PathState.PathThroughput, PathState.Ray.Origin, PathState.Ray.Direction, Interval, PathState.PathRoughness, PathState.RandSequence);
 #endif
 			// Update our current volume intersection list to reflect the fact that we have made it to the end of this interval
 			// This will either clip the finished portions of active segments, or remove the segments we have fully marched through
@@ -742,7 +743,7 @@ FPathTracingPayload TraceTransparentRay(inout FPathState PathState, bool bIsCame
 
 #define PATH_TRACING_DEBUG_TRANSPARENCY_RAY 0
 
-float3 TraceTransparentVisibilityRay(FRayDesc Ray, float PathRoughness, uint MissShaderIndex, bool bCastShadows, inout RandomSequence RandSequence)
+float3 TraceTransparentVisibilityRay(FRayDesc Ray, int Bounce, float PathRoughness, uint MissShaderIndex, bool bCastShadows, inout RandomSequence RandSequence)
 {
 	if (!bCastShadows && MissShaderIndex == 0)
 	{
@@ -824,7 +825,7 @@ float3 TraceTransparentVisibilityRay(FRayDesc Ray, float PathRoughness, uint Mis
 		}
 		if (VolumeIsectList.HitVolume())
 		{
-			Throughput = VolumeGetTransmittance(Throughput, Ray.Origin, Ray.Direction, VolumeIsectList, RandSequence);
+			Throughput = VolumeGetTransmittance(Throughput, Ray.Origin, Ray.Direction, VolumeIsectList, PathRoughness, RandSequence);
 		}
 		// add in the medium extinction
 		Throughput *= exp(-max(PackedPayload.GetTau(), 0.0));
@@ -1403,7 +1404,7 @@ bool PathTracingKernel(inout FPathState PathState, int Bounce)
 
 #if PATH_TRACING_DEBUG_TRANSPARENCY_RAY
 	// visualize the shadow ray handling directly
-	PathState.Radiance = TraceTransparentVisibilityRay(PathState.Ray, 0.0, 0, true, PathState.RandSequence);
+	PathState.Radiance = TraceTransparentVisibilityRay(PathState.Ray, Bounce, 0.0, 0, true, PathState.RandSequence);
 	return false;
 #endif
 
@@ -1419,6 +1420,7 @@ bool PathTracingKernel(inout FPathState PathState, int Bounce)
 	FVolumeSegment VolumeSegment = CreateEmptyVolumeSegment();
 	FPathTracingPayload Payload = TraceTransparentRay(
 		PathState,
+		Bounce,
 		bIsCameraRay,
 		bIsLastBounce,
 		bIncludeEmissive,
@@ -1473,7 +1475,7 @@ bool PathTracingKernel(inout FPathState PathState, int Bounce)
 			{
 				// Account for the transmittance up to the current point within the current slice of volume
 				VolumeSegment.Interval.VolumeTMax = EquiAngularT;
-				VolumeSegment.Throughput = VolumeGetTransmittance(VolumeSegment.Throughput, Ro, Rd, VolumeSegment.Interval, PathState.RandSequence);
+				VolumeSegment.Throughput = VolumeGetTransmittance(VolumeSegment.Throughput, Ro, Rd, VolumeSegment.Interval, PathState.PathRoughness, PathState.RandSequence);
 
 				float3 Contrib = VolumeSegment.Throughput * SigmaS;
 
@@ -1497,7 +1499,7 @@ bool PathTracingKernel(inout FPathState PathState, int Bounce)
 						LightRay.TMin = 0;
 						LightRay.Direction = LightSample.Direction;
 						LightRay.TMax = LightSample.Distance;
-						LightSample.RadianceOverPdf *= TraceTransparentVisibilityRay(LightRay, 1.0, MissShaderIndex, bCastShadows, PathState.RandSequence);
+						LightSample.RadianceOverPdf *= TraceTransparentVisibilityRay(LightRay, Bounce, 1.0, MissShaderIndex, bCastShadows, PathState.RandSequence);
 						// #dxr_todo: Is it cheaper to fire the ray first? Or eval the material first?
 						if (any(LightSample.RadianceOverPdf > 0))
 						{
@@ -1534,7 +1536,7 @@ bool PathTracingKernel(inout FPathState PathState, int Bounce)
 								LightContrib *= MISWeightRobust(MaterialSample.Pdf, LightResult.Pdf);
 							}
 							MaterialRay.TMax = LightResult.HitT;
-							LightContrib *= TraceTransparentVisibilityRay(MaterialRay, MaterialSample.Roughness, MissShaderIndex, bCastShadows, PathState.RandSequence);
+							LightContrib *= TraceTransparentVisibilityRay(MaterialRay, Bounce, MaterialSample.Roughness, MissShaderIndex, bCastShadows, PathState.RandSequence);
 							AccumulateRadiance(PathState.Radiance, LightContrib, bIsCameraRay);
 						}
 					}
@@ -1653,7 +1655,7 @@ bool PathTracingKernel(inout FPathState PathState, int Bounce)
 
 				const bool bCastShadows = CastsShadow(LightId);
 				const uint MissShaderIndex = GetLightMissShaderIndex(LightId);
-				LightSample.RadianceOverPdf *= TraceTransparentVisibilityRay(LightRay, AvgRoughness, MissShaderIndex, bCastShadows, PathState.RandSequence);
+				LightSample.RadianceOverPdf *= TraceTransparentVisibilityRay(LightRay, Bounce, AvgRoughness, MissShaderIndex, bCastShadows, PathState.RandSequence);
 
 				if (bIsVolumeSample)
 				{
@@ -1824,7 +1826,7 @@ bool PathTracingKernel(inout FPathState PathState, int Bounce)
 				LightRay.TMax = LightResult.HitT;
 				const bool bCastShadows = CastsShadow(LightId);
 				const uint MissShaderIndex = GetLightMissShaderIndex(LightId);
-				LightContrib *= TraceTransparentVisibilityRay(LightRay, ShadowRayRoughness, MissShaderIndex, bCastShadows, PathState.RandSequence);
+				LightContrib *= TraceTransparentVisibilityRay(LightRay, Bounce, ShadowRayRoughness, MissShaderIndex, bCastShadows, PathState.RandSequence);
 				// the light made some contribution, and there was nothing along the shadow ray
 				AccumulateRadiance(PathState.Radiance, LightContrib, bIsCameraRay);
 			}
diff --git a/Engine/Shaders/Private/PathTracing/Volume/PathTracingHeterogeneousVolumes.ush b/Engine/Shaders/Private/PathTracing/Volume/PathTracingHeterogeneousVolumes.ush
new file mode 100644
--- /dev/null
+++ b/Engine/Shaders/Private/PathTracing/Volume/PathTracingHeterogeneousVolumes.ush
@@ -0,0 +1,431 @@
+// Copyright Epic Games, Inc. All Rights Reserved.
+
+#pragma once
+
+#include "../../HeterogeneousVolumes/HeterogeneousVolumesTracingUtils.ush"
+#include "../../HeterogeneousVolumes/HeterogeneousVolumesVoxelGridTraversal.ush"
+
+float3 GetTranslatedWorldPos(float3 WorldPos)
+{
+	float3 TranslatedWorldPos = LWCHackToFloat(LWCAdd(WorldPos, PrimaryView.PreViewTranslation));
+	return TranslatedWorldPos;
+}
+
+float GetPhaseG()
+{
+	return 0.0;
+}
+
+FVolumeIntersection HeterogeneousVolumesIntersect(float3 Origin, float3 Direction, float TMin, float TMax)
+{
+	if (EnableHeterogeneousVolumes)
+	{
+		float3 WorldBoundsMin = FrustumGridUniformBuffer.TopLevelGridWorldBoundsMin;
+		float3 WorldBoundsMax = FrustumGridUniformBuffer.TopLevelGridWorldBoundsMax;
+		if (OrthoGridUniformBuffer.bUseOrthoGrid)
+		{
+			WorldBoundsMin = OrthoGridUniformBuffer.TopLevelGridWorldBoundsMin;
+			WorldBoundsMax = OrthoGridUniformBuffer.TopLevelGridWorldBoundsMax;
+		}
+
+		if (OrthoGridUniformBuffer.bUseOrthoGrid || FrustumGridUniformBuffer.bUseFrustumGrid)
+		{
+			float3 TranslatedWorldBoundsMin = GetTranslatedWorldPos(WorldBoundsMin);
+			float3 TranslatedWorldBoundsMax = GetTranslatedWorldPos(WorldBoundsMax);
+
+			float2 RayHitT = IntersectAABB(
+				Origin,
+				Direction,
+				TMin,
+				TMax,
+				TranslatedWorldBoundsMin,
+				TranslatedWorldBoundsMax
+			);
+
+			if (RayHitT.y > RayHitT.x)
+			{
+				return CreateVolumeIntersection(RayHitT.x, RayHitT.y);
+			}
+		}
+	}
+
+	return CreateEmptyVolumeIntersection();
+}
+
+void MarchMajorantGrid(
+	float3 VoxelRayOrigin,
+	float3 VoxelRayDirection,
+	float VoxelRayTMin,
+	float VoxelRayTMax,
+	float WorldRayHitSpan,
+	inout FMajorantData MajorantData
+)
+{
+	float3 VoxelRayDirectionInv = 1.0 / VoxelRayDirection;
+
+	// Grid intersection
+	float VoxelRayMarchT = VoxelRayTMin;
+	float3 VoxelRayMarchPos = VoxelRayOrigin;
+
+	float3 VoxelBoundsPos;
+	VoxelBoundsPos.x = sign(VoxelRayDirection.x) > 0 ? floor(VoxelRayMarchPos.x) + 1 : ceil(VoxelRayMarchPos.x) - 1;
+	VoxelBoundsPos.y = sign(VoxelRayDirection.y) > 0 ? floor(VoxelRayMarchPos.y) + 1 : ceil(VoxelRayMarchPos.y) - 1;
+	VoxelBoundsPos.z = sign(VoxelRayDirection.z) > 0 ? floor(VoxelRayMarchPos.z) + 1 : ceil(VoxelRayMarchPos.z) - 1;
+
+	float3 VoxelBoundsHitT = (VoxelBoundsPos - VoxelRayMarchPos) * VoxelRayDirectionInv;
+
+	// Remove floating-point rounding error
+	float Epsilon = 1.0e-4;
+	if (VoxelBoundsHitT.x <= Epsilon) VoxelBoundsHitT.x += abs(VoxelRayDirectionInv.x);
+	if (VoxelBoundsHitT.y <= Epsilon) VoxelBoundsHitT.y += abs(VoxelRayDirectionInv.y);
+	if (VoxelBoundsHitT.z <= Epsilon) VoxelBoundsHitT.z += abs(VoxelRayDirectionInv.z);
+
+	while ((VoxelRayMarchT < VoxelRayTMax))
+	{
+		float VoxelRayMarchDeltaT = 0.0;
+		if (VoxelBoundsHitT.x < VoxelBoundsHitT.y)
+		{
+			if (VoxelBoundsHitT.x < VoxelBoundsHitT.z)
+			{
+				VoxelRayMarchDeltaT = VoxelBoundsHitT.x - VoxelRayMarchT;
+				VoxelBoundsHitT.x += abs(VoxelRayDirectionInv.x);
+			}
+			else
+			{
+				VoxelRayMarchDeltaT = VoxelBoundsHitT.z - VoxelRayMarchT;
+				VoxelBoundsHitT.z += abs(VoxelRayDirectionInv.z);
+			}
+		}
+		else
+		{
+			if (VoxelBoundsHitT.y < VoxelBoundsHitT.z)
+			{
+				VoxelRayMarchDeltaT = VoxelBoundsHitT.y - VoxelRayMarchT;
+				VoxelBoundsHitT.y += abs(VoxelRayDirectionInv.y);
+			}
+			else
+			{
+				VoxelRayMarchDeltaT = VoxelBoundsHitT.z - VoxelRayMarchT;
+				VoxelBoundsHitT.z += abs(VoxelRayDirectionInv.z);
+			}
+		}
+
+		// Clip voxel delta-t by overall voxel ray-length.
+		if (VoxelRayMarchT + VoxelRayMarchDeltaT > VoxelRayTMax)
+		{
+			VoxelRayMarchDeltaT = VoxelRayTMax - VoxelRayMarchT;
+		}
+
+		// Sample at the midpoint of the voxel ray
+		float3 VoxelEndPos = VoxelRayMarchPos + VoxelRayDirection * VoxelRayMarchDeltaT;
+		float3 VoxelSamplePos = (VoxelRayMarchPos + VoxelEndPos) * 0.5;
+		uint VoxelSampleLinearPos = GetLinearIndex(VoxelSamplePos, OrthoGridUniformBuffer.TopLevelGridResolution);
+
+		int3 BottomLevelVoxelResolution = GetBottomLevelVoxelResolution(OrthoGridUniformBuffer.TopLevelGridBuffer[VoxelSampleLinearPos]);
+		if (IsBottomLevelAllocated(OrthoGridUniformBuffer.TopLevelGridBuffer[VoxelSampleLinearPos]))
+		{
+			UpdateMajorantData(MajorantData, GetMajorantData(OrthoGridUniformBuffer.MajorantGridBuffer[VoxelSampleLinearPos]));
+		}
+
+		VoxelRayMarchT += VoxelRayMarchDeltaT;
+		VoxelRayMarchPos = VoxelEndPos;
+	}
+}
+
+FMajorantData TraceMajorantGrid(float3 Origin, float3 Direction, float TMin, float TMax)
+{
+	FMajorantData MajorantData = CreateMajorantData();
+
+	float3 WorldBoundsMin = OrthoGridUniformBuffer.TopLevelGridWorldBoundsMin;
+	float3 WorldBoundsMax = OrthoGridUniformBuffer.TopLevelGridWorldBoundsMax;
+
+	// Intersect world-space AABB to clip ray
+	float3 WorldRayOrigin = LWCHackToFloat(LWCSubtract(Origin, PrimaryView.PreViewTranslation));
+	float2 ClipHitT = IntersectAABB(WorldRayOrigin, Direction, TMin, TMax, WorldBoundsMin, WorldBoundsMax);
+	if (ClipHitT.x > ClipHitT.y)
+	{
+		return MajorantData;
+	}
+
+	float3 WorldRayBegin = WorldRayOrigin + Direction * ClipHitT.x;
+	float3 WorldRayEnd = WorldRayOrigin + Direction * ClipHitT.y;
+	float WorldRayLength = ClipHitT.y - ClipHitT.x;
+
+	// Transform to voxel-space
+	float3 TopLevelGridWorldBoundsExtent = WorldBoundsMax - WorldBoundsMin;
+	float3 VoxelRayBegin = (WorldRayBegin - WorldBoundsMin) / TopLevelGridWorldBoundsExtent * OrthoGridUniformBuffer.TopLevelGridResolution;
+	float3 VoxelRayEnd = (WorldRayEnd - WorldBoundsMin) / TopLevelGridWorldBoundsExtent * OrthoGridUniformBuffer.TopLevelGridResolution;
+	float3 VoxelRayDirection = VoxelRayEnd - VoxelRayBegin;
+	float VoxelRayTMin = 0.0;
+	float VoxelRayTMax = length(VoxelRayDirection);
+	VoxelRayDirection /= VoxelRayTMax;
+
+	MarchMajorantGrid(
+		VoxelRayBegin,
+		VoxelRayDirection,
+		VoxelRayTMin,
+		VoxelRayTMax,
+		WorldRayLength,
+		MajorantData
+	);
+
+	return MajorantData;
+}
+
+FVolumeDensityBounds HeterogeneousVolumesGetDensityBounds(float3 Origin, float3 Direction, float TMin, float TMax)
+{
+	if (EnableHeterogeneousVolumes)
+	{
+		FMajorantData MajorantData = (FMajorantData) 0;
+		if (OrthoGridUniformBuffer.bUseMajorantGrid)
+		{
+			MajorantData = TraceMajorantGrid(Origin, Direction, TMin, TMax);
+		}
+		else
+		{
+			MajorantData.Minorant = 0;
+			MajorantData.Majorant = 0.1;
+		}
+
+		return CreateVolumeDensityBound(MajorantData.Minorant, MajorantData.Majorant);
+	}
+
+	return CreateVolumeDensityBound(0, 0);
+
+}
+
+void EvalWithLinearInterpolation(uint BottomLevelIndex, float3 VoxelPos,
+	StructuredBuffer<FScalarGridData> ExtinctionGridBuffer,
+	StructuredBuffer<FVectorGridData> ScatteringGridBuffer,
+	inout float3 SigmaT, inout float3 Albedo)
+{
+	uint Voxel000 = BottomLevelIndex + MortonEncode3(VoxelPos + uint3(0, 0, 0));
+	uint Voxel100 = BottomLevelIndex + MortonEncode3(VoxelPos + uint3(1, 0, 0));
+	uint Voxel010 = BottomLevelIndex + MortonEncode3(VoxelPos + uint3(0, 1, 0));
+	uint Voxel110 = BottomLevelIndex + MortonEncode3(VoxelPos + uint3(1, 1, 0));
+	uint Voxel001 = BottomLevelIndex + MortonEncode3(VoxelPos + uint3(0, 0, 1));
+	uint Voxel101 = BottomLevelIndex + MortonEncode3(VoxelPos + uint3(1, 0, 1));
+	uint Voxel011 = BottomLevelIndex + MortonEncode3(VoxelPos + uint3(0, 1, 1));
+	uint Voxel111 = BottomLevelIndex + MortonEncode3(VoxelPos + uint3(1, 1, 1));
+
+	// SigmaT
+	{
+		float3 SigmaT_000 = GetExtinction(ExtinctionGridBuffer[Voxel000]);
+		float3 SigmaT_100 = GetExtinction(ExtinctionGridBuffer[Voxel100]);
+		float3 SigmaT_010 = GetExtinction(ExtinctionGridBuffer[Voxel010]);
+		float3 SigmaT_110 = GetExtinction(ExtinctionGridBuffer[Voxel110]);
+		float3 SigmaT_001 = GetExtinction(ExtinctionGridBuffer[Voxel001]);
+		float3 SigmaT_101 = GetExtinction(ExtinctionGridBuffer[Voxel101]);
+		float3 SigmaT_011 = GetExtinction(ExtinctionGridBuffer[Voxel011]);
+		float3 SigmaT_111 = GetExtinction(ExtinctionGridBuffer[Voxel111]);
+
+		float3 Weight = frac(VoxelPos);
+		float3 Edge00 = (1.0 - Weight.x) * SigmaT_000 + Weight.x * SigmaT_100;
+		float3 Edge10 = (1.0 - Weight.x) * SigmaT_010 + Weight.x * SigmaT_110;
+		float3 Face0 = (1.0 - Weight.y) * Edge00 + Weight.y * Edge10;
+
+		float3 Edge01 = (1.0 - Weight.x) * SigmaT_001 + Weight.x * SigmaT_101;
+		float3 Edge11 = (1.0 - Weight.x) * SigmaT_011 + Weight.x * SigmaT_111;
+		float3 Face1 = (1.0 - Weight.y) * Edge01 + Weight.y * Edge11;
+
+		SigmaT = (1.0 - Weight.z) * Face0 + Weight.z * Face1;
+	}
+
+	// Albedo
+	{
+		float3 Albedo_000 = GetAlbedo(ScatteringGridBuffer[Voxel000]);
+		float3 Albedo_100 = GetAlbedo(ScatteringGridBuffer[Voxel100]);
+		float3 Albedo_010 = GetAlbedo(ScatteringGridBuffer[Voxel010]);
+		float3 Albedo_110 = GetAlbedo(ScatteringGridBuffer[Voxel110]);
+		float3 Albedo_001 = GetAlbedo(ScatteringGridBuffer[Voxel001]);
+		float3 Albedo_101 = GetAlbedo(ScatteringGridBuffer[Voxel101]);
+		float3 Albedo_011 = GetAlbedo(ScatteringGridBuffer[Voxel011]);
+		float3 Albedo_111 = GetAlbedo(ScatteringGridBuffer[Voxel111]);
+
+		float3 Weight = frac(VoxelPos);
+		float3 Edge00 = (1.0 - Weight.x) * Albedo_000 + Weight.x * Albedo_100;
+		float3 Edge10 = (1.0 - Weight.x) * Albedo_010 + Weight.x * Albedo_110;
+		float3 Face0 = (1.0 - Weight.y) * Edge00 + Weight.y * Edge10;
+
+		float3 Edge01 = (1.0 - Weight.x) * Albedo_001 + Weight.x * Albedo_101;
+		float3 Edge11 = (1.0 - Weight.x) * Albedo_011 + Weight.x * Albedo_111;
+		float3 Face1 = (1.0 - Weight.y) * Edge01 + Weight.y * Edge11;
+
+		Albedo = (1.0 - Weight.z) * Face0 + Weight.z * Face1;
+	}
+}
+
+FVolumeShadedResult GetOrthoVoxelGridDensity(float3 TranslatedWorldPos)
+{
+	float3 SigmaT = 0.0;
+	float3 Albedo = 0.0;
+
+	float3 TranslatedWorldBoundsMin = GetTranslatedWorldPos(OrthoGridUniformBuffer.TopLevelGridWorldBoundsMin);
+	float3 TranslatedWorldBoundsMax = GetTranslatedWorldPos(OrthoGridUniformBuffer.TopLevelGridWorldBoundsMax);
+	float3 TranslatedWorldBoundsExtent = TranslatedWorldBoundsMax - TranslatedWorldBoundsMin;
+	float3 GridUV = (TranslatedWorldPos - TranslatedWorldBoundsMin) / TranslatedWorldBoundsExtent;
+	if (all(GridUV >= 0.0) && all(GridUV <= 1.0))
+	{
+		float3 TopLevelVoxelPos = GridUV * OrthoGridUniformBuffer.TopLevelGridResolution;
+		uint LinearTopLevelVoxelPos = GetLinearIndex(TopLevelVoxelPos, OrthoGridUniformBuffer.TopLevelGridResolution);
+
+		if (OrthoGridUniformBuffer.bEnableIndirectionGrid)
+		{
+			FTopLevelGridData TopLevelGridData = OrthoGridUniformBuffer.TopLevelGridBuffer[LinearTopLevelVoxelPos];
+			if (IsBottomLevelAllocated(TopLevelGridData))
+			{
+				uint IndirectionGridIndex = GetBottomLevelIndex(TopLevelGridData);
+				uint3 IndirectionGridVoxelResolution = GetBottomLevelVoxelResolution(TopLevelGridData);
+
+				// Convert to Indirection voxel-space
+				float3 IndirectionGridVoxelPos = frac(TopLevelVoxelPos) * IndirectionGridVoxelResolution;
+				FTopLevelGridData IndirectionData = OrthoGridUniformBuffer.IndirectionGridBuffer[IndirectionGridIndex + MortonEncode3(IndirectionGridVoxelPos)];
+				if (IsBottomLevelAllocated(IndirectionData))
+				{
+					uint BottomLevelIndex = GetBottomLevelIndex(IndirectionData);
+					uint3 BottomLevelVoxelResolution = GetBottomLevelVoxelResolution(IndirectionData);
+
+					// Constant interpolation
+					float3 BottomLevelVoxelPos = frac(IndirectionGridVoxelPos) * BottomLevelVoxelResolution;
+					uint LinearBottomLevelVoxelPos = BottomLevelIndex + MortonEncode3(uint3(BottomLevelVoxelPos));
+					SigmaT = GetExtinction(OrthoGridUniformBuffer.ExtinctionGridBuffer[LinearBottomLevelVoxelPos]);
+					Albedo = GetAlbedo(OrthoGridUniformBuffer.ScatteringGridBuffer[LinearBottomLevelVoxelPos]);
+				}
+			}
+		}
+		else
+		{
+			FTopLevelGridData TopLevelData = OrthoGridUniformBuffer.TopLevelGridBuffer[LinearTopLevelVoxelPos];
+			if (IsBottomLevelAllocated(TopLevelData))
+			{
+				uint BottomLevelIndex = GetBottomLevelIndex(OrthoGridUniformBuffer.TopLevelGridBuffer[LinearTopLevelVoxelPos]);
+				uint3 BottomLevelVoxelResolution = GetBottomLevelVoxelResolution(OrthoGridUniformBuffer.TopLevelGridBuffer[LinearTopLevelVoxelPos]);
+
+				// Constant Interpolation
+				float3 BottomLevelVoxelPos = frac(TopLevelVoxelPos) * BottomLevelVoxelResolution;
+				uint LinearBottomLevelVoxelPos = BottomLevelIndex + MortonEncode3(uint3(BottomLevelVoxelPos));
+
+				SigmaT = GetExtinction(OrthoGridUniformBuffer.ExtinctionGridBuffer[LinearBottomLevelVoxelPos]);
+				Albedo = GetAlbedo(OrthoGridUniformBuffer.ScatteringGridBuffer[LinearBottomLevelVoxelPos]);
+			}
+		}
+	}
+
+	FVolumeShadedResult Result = (FVolumeShadedResult)0;
+	Result.SigmaT = SigmaT;
+	//Result.SigmaSHG = Albedo * Result.SigmaT;
+	Result.SigmaSHG = Albedo;
+	Result.PhaseG = GetPhaseG();
+	// TODO: Emission?
+
+	return Result;
+}
+
+FVolumeShadedResult GetFrustumVoxelGridDensity(float3 TranslatedWorldPos, inout bool bInFrustum)
+{
+	float3 SigmaT = 0.0;
+	float3 Albedo = 0.0;
+
+	// Convert TranslatedWorldPos to voxel space
+	float3 WorldPos = LWCHackToFloat(LWCSubtract(TranslatedWorldPos, PrimaryView.PreViewTranslation));
+	float3 ViewPos = mul(float4(WorldPos, 1), FrustumGridUniformBuffer.WorldToView).xyz;
+
+	int3 VoxelDimensions = FrustumGridUniformBuffer.VoxelDimensions;
+	float NearPlaneDepth = FrustumGridUniformBuffer.NearPlaneDepth;
+	float FarPlaneDepth = FrustumGridUniformBuffer.FarPlaneDepth;
+	float3 VoxelPos = ViewToVoxel(ViewPos, VoxelDimensions, NearPlaneDepth, FarPlaneDepth);
+
+	bInFrustum = all(VoxelPos > 0) && all(VoxelPos < FrustumGridUniformBuffer.TopLevelFroxelGridResolution);
+	if (bInFrustum)
+	{
+		uint LinearTopLevelVoxelPos = GetLinearIndex(VoxelPos, FrustumGridUniformBuffer.TopLevelFroxelGridResolution);
+		FTopLevelGridData TopLevelGridData = FrustumGridUniformBuffer.TopLevelFroxelGridBuffer[LinearTopLevelVoxelPos];
+		if (IsBottomLevelAllocated(TopLevelGridData))
+		{
+			uint BottomLevelIndex = GetBottomLevelIndex(TopLevelGridData);
+			uint3 BottomLevelVoxelResolution = GetBottomLevelVoxelResolution(TopLevelGridData);
+
+			float3 BottomLevelVoxelPos = frac(VoxelPos) * BottomLevelVoxelResolution;
+			uint LinearBottomLevelVoxelPos = BottomLevelIndex + MortonEncode3(uint3(BottomLevelVoxelPos));
+			SigmaT = GetExtinction(FrustumGridUniformBuffer.ExtinctionFroxelGridBuffer[LinearBottomLevelVoxelPos]);
+			Albedo = GetAlbedo(FrustumGridUniformBuffer.ScatteringFroxelGridBuffer[LinearBottomLevelVoxelPos]);
+		}
+	}
+
+	// Return struct
+	FVolumeShadedResult Result = (FVolumeShadedResult) 0;
+	Result.SigmaT = SigmaT;
+	//Result.SigmaSHG = Albedo * Result.SigmaT;
+	Result.SigmaSHG = Albedo;
+	Result.PhaseG = GetPhaseG();
+	// TODO: Emission?
+	return Result;
+}
+
+FVolumeShadedResult HeterogeneousVolumesGetDensity(float3 TranslatedWorldPos)
+{
+	FVolumeShadedResult Result = (FVolumeShadedResult) 0;
+
+	if (EnableHeterogeneousVolumes)
+	{
+		bool bUseBothGrids = OrthoGridUniformBuffer.bUseOrthoGrid && FrustumGridUniformBuffer.bUseFrustumGrid;
+		if (bUseBothGrids)
+		{
+			bool bInFrustum = false;
+			Result = GetFrustumVoxelGridDensity(TranslatedWorldPos, bInFrustum);
+			if (!bInFrustum)
+			{
+				Result = GetOrthoVoxelGridDensity(TranslatedWorldPos);
+			}
+		}
+		else if (FrustumGridUniformBuffer.bUseFrustumGrid)
+		{
+			bool bInFrustum = false;
+			Result = GetFrustumVoxelGridDensity(TranslatedWorldPos, bInFrustum);
+		}
+		else if (OrthoGridUniformBuffer.bUseOrthoGrid)
+		{
+			Result = GetOrthoVoxelGridDensity(TranslatedWorldPos);
+		}
+	}
+
+	return Result;
+}
+
+float3 HeterogeneousVolumesGetTransmittance(float3 Origin, float3 Direction, float TMin, float TMax, float PathRoughness)
+{
+	float3 Transmittance = 1.0;
+
+	if (EnableHeterogeneousVolumes)
+	{
+		bool bCombinedTraversal = OrthoGridUniformBuffer.bUseOrthoGrid && FrustumGridUniformBuffer.bUseFrustumGrid;
+		//if (bCombinedTraversal && Bounce < 1)
+		if (bCombinedTraversal)
+		{
+			float FrustumTMin = TMin;
+			float FrustumTMax = (PathRoughness < 0.1) ? TMax : min(TMax, 10.0);
+			Transmittance = TraceFrustumVoxelGrid(Origin, Direction, FrustumTMin, FrustumTMax);
+
+			if (FrustumTMin < FrustumTMax)
+			{
+				Transmittance *= TraceOrthoVoxelGrid(Origin, Direction, TMin, FrustumTMin);
+				Transmittance *= TraceOrthoVoxelGrid(Origin, Direction, FrustumTMax, TMax);
+			}
+			else
+			{
+				Transmittance *= TraceOrthoVoxelGrid(Origin, Direction, TMin, TMax);
+			}
+		}
+		// NOTE: Ordering of ortho traversal is important for multi-bounce fallback
+		else if (OrthoGridUniformBuffer.bUseOrthoGrid)
+		{
+			Transmittance = TraceOrthoVoxelGrid(Origin, Direction, TMin, TMax);
+		}
+		else if (FrustumGridUniformBuffer.bUseFrustumGrid)
+		{
+			Transmittance = TraceFrustumVoxelGrid(Origin, Direction, TMin, TMax);
+		}
+	}
+
+	return Transmittance;
+}
diff --git a/Engine/Shaders/Private/PathTracing/Volume/PathTracingVolume.ush b/Engine/Shaders/Private/PathTracing/Volume/PathTracingVolume.ush
--- a/Engine/Shaders/Private/PathTracing/Volume/PathTracingVolume.ush
+++ b/Engine/Shaders/Private/PathTracing/Volume/PathTracingVolume.ush
@@ -5,12 +5,14 @@
 #include "PathTracingVolumeCommon.ush"
 #include "PathTracingAtmosphere.ush"
 #include "PathTracingFog.ush"
+#include "PathTracingHeterogeneousVolumes.ush"
 #include "PathTracingVolumeSampling.ush"
 
 #define PATH_TRACING_USE_TRANSMITTANCE_FUNC 1 // 0: null tracking transmittance (reference) 1: analytic transmittance
 
 uint EnableAtmosphere;
 uint EnableFog;
+uint EnableHeterogeneousVolumes;
 int MaxRaymarchSteps;
 
 
@@ -26,6 +28,10 @@ FVolumeIntersectionList VolumeIntersect(float3 Origin, float3 Direction, float T
 	{
 		Result.Add(VOLUMEID_FOG, FogIntersect(Origin, Direction, TMin, TMax));
 	}
+	if (EnableHeterogeneousVolumes)
+	{
+		Result.Add(VOLUMEID_HETEROGENEOUS_VOLUMES, HeterogeneousVolumesIntersect(Origin, Direction, TMin, TMax));
+	}
 	return Result;
 }
 
@@ -43,8 +49,9 @@ FVolumeDensityBounds VolumeGetDensityBounds(float3 Origin, float3 Direction, FVo
 	{
 		switch (Interval.VolumeID[Index])
 		{
-			case VOLUMEID_ATMOSPHERE: MergeVolumeDensityBounds(Result, AtmosphereGetDensityBounds(Origin, Direction, Interval.VolumeTMin, Interval.VolumeTMax)); break;
-			case VOLUMEID_FOG:        MergeVolumeDensityBounds(Result, FogGetDensityBounds(Origin, Direction, Interval.VolumeTMin, Interval.VolumeTMax)); break;
+			case VOLUMEID_ATMOSPHERE:            MergeVolumeDensityBounds(Result, AtmosphereGetDensityBounds(Origin, Direction, Interval.VolumeTMin, Interval.VolumeTMax)); break;
+			case VOLUMEID_FOG:                   MergeVolumeDensityBounds(Result, FogGetDensityBounds(Origin, Direction, Interval.VolumeTMin, Interval.VolumeTMax)); break;
+			case VOLUMEID_HETEROGENEOUS_VOLUMES: MergeVolumeDensityBounds(Result, HeterogeneousVolumesGetDensityBounds(Origin, Direction, Interval.VolumeTMin, Interval.VolumeTMax)); break;
 		}
 	}
 	return Result;
@@ -58,8 +65,9 @@ FVolumeShadedResult VolumeGetDensity(float3 TranslatedWorldPos, FVolumeIntersect
 	{
 		switch (Interval.VolumeID[Index])
 		{
-			case VOLUMEID_ATMOSPHERE: MergeVolumeShadedResult(Result, AtmosphereGetDensity(TranslatedWorldPos)); break;
-			case VOLUMEID_FOG:        MergeVolumeShadedResult(Result, FogGetDensity(TranslatedWorldPos)); break;
+			case VOLUMEID_ATMOSPHERE:            MergeVolumeShadedResult(Result, AtmosphereGetDensity(TranslatedWorldPos)); break;
+			case VOLUMEID_FOG:                   MergeVolumeShadedResult(Result, FogGetDensity(TranslatedWorldPos)); break;
+			case VOLUMEID_HETEROGENEOUS_VOLUMES: MergeVolumeShadedResult(Result, HeterogeneousVolumesGetDensity(TranslatedWorldPos)); break;
 		}
 	}
 	return Result;
@@ -68,22 +76,23 @@ FVolumeShadedResult VolumeGetDensity(float3 TranslatedWorldPos, FVolumeIntersect
 #if PATH_TRACING_USE_TRANSMITTANCE_FUNC
 
 // Return the transmittance along a ray segment
-float3 VolumeGetTransmittance(float3 StartThroughput, float3 Origin, float3 Direction, FVolumeIntersectionInterval Interval, inout RandomSequence RandSequence)
+float3 VolumeGetTransmittance(float3 StartThroughput, float3 Origin, float3 Direction, FVolumeIntersectionInterval Interval, float PathRoughness, inout RandomSequence RandSequence)
 {
 	float3 Throughput = StartThroughput;
 	for (int Index = 0; Index < Interval.Num; Index++)
 	{
 		switch (Interval.VolumeID[Index])
 		{
-			case VOLUMEID_ATMOSPHERE: Throughput *= AtmosphereGetTransmittance(Origin, Direction, Interval.VolumeTMin, Interval.VolumeTMax); break;
-			case VOLUMEID_FOG:        Throughput *= FogGetTransmittance(Origin, Direction, Interval.VolumeTMin, Interval.VolumeTMax); break;
+			case VOLUMEID_ATMOSPHERE:            Throughput *= AtmosphereGetTransmittance(Origin, Direction, Interval.VolumeTMin, Interval.VolumeTMax); break;
+			case VOLUMEID_FOG:                   Throughput *= FogGetTransmittance(Origin, Direction, Interval.VolumeTMin, Interval.VolumeTMax); break;
+			case VOLUMEID_HETEROGENEOUS_VOLUMES: Throughput *= HeterogeneousVolumesGetTransmittance(Origin, Direction, Interval.VolumeTMin, Interval.VolumeTMax, PathRoughness); break;
 		}
 	}
 	return Throughput;
 }
 
 // Return the transmittance along a whole intersection list
-float3 VolumeGetTransmittance(float3 StartThroughput, float3 Origin, float3 Direction, FVolumeIntersectionList IntersectionList, inout RandomSequence RandSequence)
+float3 VolumeGetTransmittance(float3 StartThroughput, float3 Origin, float3 Direction, FVolumeIntersectionList IntersectionList, float PathRoughness, inout RandomSequence RandSequence)
 {
 	float3 Throughput = StartThroughput;
 	for (int Index = 0; Index < IntersectionList.Num; Index++)
@@ -92,8 +101,9 @@ float3 VolumeGetTransmittance(float3 StartThroughput, float3 Origin, float3 Dire
 		float TMax = IntersectionList.VolumeTMax[Index];
 		switch (IntersectionList.VolumeID[Index])
 		{
-			case VOLUMEID_ATMOSPHERE: Throughput *= AtmosphereGetTransmittance(Origin, Direction, TMin, TMax); break;
-			case VOLUMEID_FOG:        Throughput *= FogGetTransmittance(Origin, Direction, TMin, TMax); break;
+			case VOLUMEID_ATMOSPHERE:            Throughput *= AtmosphereGetTransmittance(Origin, Direction, TMin, TMax); break;
+			case VOLUMEID_FOG:                   Throughput *= FogGetTransmittance(Origin, Direction, TMin, TMax); break;
+			case VOLUMEID_HETEROGENEOUS_VOLUMES: Throughput *= HeterogeneousVolumesGetTransmittance(Origin, Direction, TMin, TMax, PathRoughness); break;
 		}
 	}
 	return Throughput;
@@ -142,7 +152,7 @@ float3 VolumeGetTransmittance(float3 StartThroughput, float3 Origin, float3 Dire
 
 
 // The returned value factors in the initial throughput. This is used to avoid sampling too much if throughput reaches 0.
-float3 VolumeGetTransmittance(float3 StartThroughput, float3 Origin, float3 Direction, FVolumeIntersectionInterval Interval, inout RandomSequence RandSequence)
+float3 VolumeGetTransmittance(float3 StartThroughput, float3 Origin, float3 Direction, FVolumeIntersectionInterval Interval, int Bounce, inout RandomSequence RandSequence)
 {
 	float3 Throughput = StartThroughput;
 	for (int Index = 0; Index < Interval.Num; Index++)
@@ -166,7 +176,7 @@ float3 VolumeGetTransmittance(float3 StartThroughput, float3 Origin, float3 Dire
 	return Throughput;
 
 }
-float3 VolumeGetTransmittance(float3 StartThroughput, float3 Origin, float3 Direction, FVolumeIntersectionList IntersectionList, inout RandomSequence RandSequence)
+float3 VolumeGetTransmittance(float3 StartThroughput, float3 Origin, float3 Direction, FVolumeIntersectionList IntersectionList, int Bounce, inout RandomSequence RandSequence)
 {
 	float3 Throughput = StartThroughput;
 	for (int Index = 0; Index < IntersectionList.Num; Index++)
diff --git a/Engine/Shaders/Private/PathTracing/Volume/PathTracingVolumeCommon.ush b/Engine/Shaders/Private/PathTracing/Volume/PathTracingVolumeCommon.ush
--- a/Engine/Shaders/Private/PathTracing/Volume/PathTracingVolumeCommon.ush
+++ b/Engine/Shaders/Private/PathTracing/Volume/PathTracingVolumeCommon.ush
@@ -33,9 +33,10 @@ FVolumeIntersection CreateEmptyVolumeIntersection()
 	return (FVolumeIntersection)0;
 }
 
-#define VOLUMEID_ATMOSPHERE		0
-#define VOLUMEID_FOG			1
-#define PATH_TRACER_MAX_VOLUMES 2
+#define VOLUMEID_ATMOSPHERE				0
+#define VOLUMEID_FOG					1
+#define VOLUMEID_HETEROGENEOUS_VOLUMES	2
+#define PATH_TRACER_MAX_VOLUMES			3
 
 // representes an interval along the ray over which the specified volumes are known to exist
 struct FVolumeIntersectionInterval
diff --git a/Engine/Source/Runtime/Renderer/Private/PathTracing.cpp b/Engine/Source/Runtime/Renderer/Private/PathTracing.cpp
--- a/Engine/Source/Runtime/Renderer/Private/PathTracing.cpp
+++ b/Engine/Source/Runtime/Renderer/Private/PathTracing.cpp
@@ -29,6 +29,7 @@ TAutoConsoleVariable<int32> CVarPathTracing(
 #include "FogRendering.h"
 #include "GenerateMips.h"
 #include "HairStrands/HairStrandsData.h"
+#include "HeterogeneousVolumes/HeterogeneousVolumes.h"
 #include "Modules/ModuleManager.h"
 #include "SkyAtmosphereRendering.h"
 #include <limits>
@@ -319,6 +320,24 @@ TAutoConsoleVariable<int32> CVarPathTracingLightFunctionColor(
 	ECVF_RenderThreadSafe
 );
 
+TAutoConsoleVariable<int32> CVarPathTracingHeterogeneousVolumes(
+	TEXT("r.PathTracing.HeterogeneousVolumes"),
+	0,
+	TEXT("Enables heterogeneous volumes (default = 0)\n")
+	TEXT("0: off (default)\n")
+	TEXT("1: on (combined frustum + ortho grids)\n")
+	TEXT("2: frustum-only grid\n")
+	TEXT("3: ortho-only grid\n"),
+	ECVF_RenderThreadSafe
+);
+
+TAutoConsoleVariable<int32> CVarPathTracingHeterogeneousVolumesRebuildEveryFrame(
+	TEXT("r.PathTracing.HeterogeneousVolumes.RebuildEveryFrame"),
+	1,
+	TEXT("Rebuilds volumetric acceleration structures every frame (default = 1)\n"),
+	ECVF_RenderThreadSafe
+);
+
 TAutoConsoleVariable<int32> CVarPathTracingCameraMediumTracking(
 	TEXT("r.PathTracing.CameraMediumTracking"),
 	1,
@@ -383,6 +402,7 @@ BEGIN_SHADER_PARAMETER_STRUCT(FPathTracingData, )
 	SHADER_PARAMETER(uint32, EnableDBuffer)
 	SHADER_PARAMETER(uint32, EnableAtmosphere)
 	SHADER_PARAMETER(uint32, EnableFog)
+	SHADER_PARAMETER(uint32, EnableHeterogeneousVolumes)
 	SHADER_PARAMETER(uint32, EnabledDirectLightingContributions)   // PATHTRACER_CONTRIBUTION_*
 	SHADER_PARAMETER(uint32, EnabledIndirectLightingContributions) // PATHTRACER_CONTRIBUTION_*
 	SHADER_PARAMETER(uint32, ApplyDiffuseSpecularOverrides)
@@ -433,6 +453,7 @@ struct FPathTracingConfig
 			PathTracingData.FilterWidth != Other.PathTracingData.FilterWidth ||
 			PathTracingData.EnableAtmosphere != Other.PathTracingData.EnableAtmosphere ||
 			PathTracingData.EnableFog != Other.PathTracingData.EnableFog ||
+			PathTracingData.EnableHeterogeneousVolumes != Other.PathTracingData.EnableHeterogeneousVolumes ||
 			PathTracingData.ApplyDiffuseSpecularOverrides != Other.PathTracingData.ApplyDiffuseSpecularOverrides ||
 			PathTracingData.EnabledDirectLightingContributions != Other.PathTracingData.EnabledDirectLightingContributions ||
 			PathTracingData.EnabledIndirectLightingContributions != Other.PathTracingData.EnabledIndirectLightingContributions ||
@@ -512,6 +533,10 @@ struct FPathTracingState {
 	TRefCountPtr<IPooledRenderTarget> LastAlbedoRT;
 	TRefCountPtr<FRDGPooledBuffer> LastVarianceBuffer;
 
+	// Volume acceleration structures
+	FAdaptiveOrthoGridParameterCache AdaptiveOrthoGridParameterCache;
+	FAdaptiveFrustumGridParameterCache AdaptiveFrustumGridParameterCache;
+
 	// Texture holding onto the precomputed atmosphere data
 	TRefCountPtr<IPooledRenderTarget> AtmosphereOpticalDepthLUT;
 	FAtmosphereConfig LastAtmosphereConfig;
@@ -595,6 +620,7 @@ static void PreparePathTracingData(const FScene* Scene, const FViewInfo& View, F
 		&& (Scene->ExponentialFogs[0].FogData[0].Density > 0 ||
 			Scene->ExponentialFogs[0].FogData[1].Density > 0);
 
+	PathTracingData.EnableHeterogeneousVolumes = CVarPathTracingHeterogeneousVolumes.GetValueOnRenderThread();
 	PathTracingData.EnableDBuffer = CVarPathTracingUseDBuffer.GetValueOnRenderThread();
 
 	PathTracingData.DecalRoughnessCutoff = PathTracing::UsesDecals(*View.Family) && View.bHasRayTracingDecals ? CVarPathTracingDecalRoughnessCutoff.GetValueOnRenderThread() : -1.0f;
@@ -875,6 +901,10 @@ class FPathTracingRG : public FGlobalShader
 		// exponential height fog
 		SHADER_PARAMETER_STRUCT_INCLUDE(FPathTracingFogParameters, FogParameters)
 
+		// Heterogeneous volumes adaptive voxel grid
+		SHADER_PARAMETER_RDG_UNIFORM_BUFFER(FOrthoVoxelGridUniformBufferParameters, OrthoGridUniformBuffer)
+		SHADER_PARAMETER_RDG_UNIFORM_BUFFER(FFrustumVoxelGridUniformBufferParameters, FrustumGridUniformBuffer)
+
 		// scene decals
 		SHADER_PARAMETER_RDG_UNIFORM_BUFFER(FRayTracingDecals, DecalParameters)
 
@@ -2156,6 +2186,8 @@ void FSceneViewState::PathTracingInvalidate(bool InvalidateAnimationStates)
 		State->NormalRT.SafeRelease();
 		State->VarianceBuffer.SafeRelease();
 		State->SampleIndex = 0;
+
+		State->AdaptiveFrustumGridParameterCache.TopLevelGridBuffer.SafeRelease();
 	}
 }
 
@@ -2412,6 +2444,31 @@ void FDeferredShadingSceneRenderer::RenderPathTracing(
 		View.ViewState->PathTracingInvalidate();
 	}
 
+
+	TRDGUniformBufferRef<FOrthoVoxelGridUniformBufferParameters> OrthoGridUniformBuffer;
+	TRDGUniformBufferRef<FFrustumVoxelGridUniformBufferParameters> FrustumGridUniformBuffer;
+	bool bForceRebuild = CVarPathTracingHeterogeneousVolumesRebuildEveryFrame.GetValueOnRenderThread() != 0;
+	bool bCreateVolumeGrids = bForceRebuild ||
+		!PathTracingState->AdaptiveFrustumGridParameterCache.TopLevelGridBuffer ||
+		!PathTracingState->AdaptiveOrthoGridParameterCache.TopLevelGridBuffer;
+	if (bCreateVolumeGrids)
+	{
+		BuildOrthoVoxelGrid(GraphBuilder, Scene, Views, OrthoGridUniformBuffer);
+		BuildFrustumVoxelGrid(GraphBuilder, Scene, Views[0], FrustumGridUniformBuffer);
+	}
+	else
+	{
+		RegisterExternalOrthoVoxelGridUniformBuffer(GraphBuilder,
+			PathTracingState->AdaptiveOrthoGridParameterCache,
+			OrthoGridUniformBuffer
+		);
+
+		RegisterExternalFrustumVoxelGridUniformBuffer(GraphBuilder,
+			PathTracingState->AdaptiveFrustumGridParameterCache,
+			FrustumGridUniformBuffer
+		);
+	}
+
 	// Prepare radiance buffer (will be shared with display pass)
 	FRDGTexture* RadianceTexture = nullptr;
 	FRDGTexture* AlbedoTexture = nullptr;
@@ -2598,6 +2655,10 @@ void FDeferredShadingSceneRenderer::RenderPathTracing(
 								PassParameters->FogParameters = {};
 							}
 
+							// Heterogeneous volume bindings
+							PassParameters->OrthoGridUniformBuffer = OrthoGridUniformBuffer;
+							PassParameters->FrustumGridUniformBuffer = FrustumGridUniformBuffer;
+
 							PassParameters->TilePixelOffset.X = TileX;
 							PassParameters->TilePixelOffset.Y = TileY + CurrentGPU;
 							PassParameters->TileTextureOffset.X = TileX;
@@ -2781,6 +2842,12 @@ void FDeferredShadingSceneRenderer::RenderPathTracing(
 		GraphBuilder.QueueTextureExtraction(NormalTexture, &PathTracingState->NormalRT);
 	}
 
+	if (bCreateVolumeGrids)
+	{
+		ExtractOrthoVoxelGridUniformBuffer(GraphBuilder, OrthoGridUniformBuffer, PathTracingState->AdaptiveOrthoGridParameterCache);
+		ExtractFrustumVoxelGridUniformBuffer(GraphBuilder, FrustumGridUniformBuffer, PathTracingState->AdaptiveFrustumGridParameterCache);
+	}
+
 	RDG_GPU_MASK_SCOPE(GraphBuilder, View.GPUMask);
 	RDG_GPU_STAT_SCOPE(GraphBuilder, PathTracingPost);
 
