Fix invalid valid for HAIR_MAX_NUM_POINT_PER_GROUP
Add PackHairCurve function

diff --git a/Engine/Shaders/Private/HairStrands/HairStrandsPack.ush b/Engine/Shaders/Private/HairStrands/HairStrandsPack.ush
--- a/Engine/Shaders/Private/HairStrands/HairStrandsPack.ush
+++ b/Engine/Shaders/Private/HairStrands/HairStrandsPack.ush
@@ -212,9 +212,18 @@ struct FHairCurve
 	uint PointCount;
 };
 
+uint PackHairCurve(FHairCurve In)
+{
+	#if HAIR_MAX_NUM_POINT_PER_CURVE != 0xFF || HAIR_MAX_NUM_POINT_PER_GROUP != 0xFFFFFF
+	#error Update hair curve ENCODING_TYPE
+	#endif
+
+	return (In.PointIndex & HAIR_MAX_NUM_POINT_PER_GROUP) | (In.PointCount << 24);
+}
+
 FHairCurve UnpackHairCurve(uint In)
 {
-	#if HAIR_MAX_NUM_POINT_PER_CURVE != 0xFF || HAIR_MAX_NUM_POINT_PER_GROUP != 0xFFFFFFFF
+	#if HAIR_MAX_NUM_POINT_PER_CURVE != 0xFF || HAIR_MAX_NUM_POINT_PER_GROUP != 0xFFFFFF
 	#error Update hair curve ENCODING_TYPE
 	#endif
 
