Fix for GammaCorrectionCommon.ush compile error due to non-scalar argument pass to short circuit ? operator. Now using "select" intrinsic.

diff --git a/Engine/Shaders/Private/GammaCorrectionCommon.ush b/Engine/Shaders/Private/GammaCorrectionCommon.ush
--- a/Engine/Shaders/Private/GammaCorrectionCommon.ush
+++ b/Engine/Shaders/Private/GammaCorrectionCommon.ush
@@ -184,6 +184,5 @@ LinearToST2084(float3 lin)
 float3
 SLog3ToLinear(float3 Value)
 {
-	return (Value >= 171.2102946929f / 1023.0f) ? (pow(10.0f, (Value * 1023.0f - 420.f) / 261.5f)) * 0.19f - 0.01f
-												: (Value * 1023.0f - 95.0f) * 0.01125000f / (171.2102946929f - 95.0f);
+	return select((Value >= 171.2102946929f / 1023.0f), (pow(10.0f, (Value * 1023.0f - 420.f) / 261.5f)) * 0.19f - 0.01f, (Value * 1023.0f - 95.0f) * 0.01125000f / (171.2102946929f - 95.0f));
 }
