Hair: Fix #error in shader header

GROUP_SIZE is already defined for this shader, so there is no need to check PERMUTATION_GROUP_SIZE

diff --git a/Engine/Shaders/Private/HairStrands/HairStrandsInterpolation.usf b/Engine/Shaders/Private/HairStrands/HairStrandsInterpolation.usf
--- a/Engine/Shaders/Private/HairStrands/HairStrandsInterpolation.usf
+++ b/Engine/Shaders/Private/HairStrands/HairStrandsInterpolation.usf
@@ -121,11 +121,8 @@ Buffer<uint>	SimRootToUniqueTriangleIndexBuffer;
 Buffer<float4> OutSimHairPositionOffsetBuffer;
 Buffer<float4> OutRenHairPositionOffsetBuffer;
 
-#if PERMUTATION_GROUP_SIZE == 64
-	#define GROUP_SIZE 64
-#elif PERMUTATION_GROUP_SIZE == 32
-	#define GROUP_SIZE 32
-#else
+
+#ifndef GROUP_SIZE
 	#error Unknown group size	
 #endif
 
