Path Tracer: Fix corner cases with invalid Strata materials

There were rare cases where we could generate invalid code.

Also comment out some common methods that are unused when compiling GPU Lightmass

diff --git a/Engine/Shaders/Private/PathTracing/Material/PathTracingMaterialSampling.ush b/Engine/Shaders/Private/PathTracing/Material/PathTracingMaterialSampling.ush
--- a/Engine/Shaders/Private/PathTracing/Material/PathTracingMaterialSampling.ush
+++ b/Engine/Shaders/Private/PathTracing/Material/PathTracingMaterialSampling.ush
@@ -107,6 +107,8 @@ float3 EstimateMaterialAlbedo(FPathTracingPayload Payload)
 	return Albedo;
 }
 
+#if !SIMPLIFIED_MATERIAL_SHADER
+
 #define ROUGHNESS_WEIGHT(r)		(saturate((r - max(MinRoughness, PathRoughness)) / (MaxRoughness - max(MinRoughness, PathRoughness))))
 
 void AccumulateAlbedoNormal(inout FPathTracingPayload Payload, float3 PathThroughput, float PathRoughness, inout float3 Albedo, inout float3 Normal)
@@ -170,17 +172,12 @@ void AccumulateAlbedoNormal(inout FPathTracingPayload Payload, float3 PathThroug
 		}
 	}
 #else
+
 	const float RoughnessWeight = ROUGHNESS_WEIGHT(Payload.Roughness);
 	const float BrdfWeight = Payload.BSDFOpacity;
 
 	switch (Payload.ShadingModelID)
 	{
-#if SIMPLIFIED_MATERIAL_SHADER
-	case SHADINGMODELID_TWOSIDED_FOLIAGE:
-		Albedo += BrdfWeight * PathThroughput * (Payload.BaseColor + Payload.SubsurfaceColor);
-		Normal += BrdfWeight * NormThroughput * Payload.WorldNormal;
-		break;
-#else
 	case SHADINGMODELID_TWOSIDED_FOLIAGE:
 	case SHADINGMODELID_SUBSURFACE:
 	case SHADINGMODELID_PREINTEGRATED_SKIN:
@@ -217,7 +214,6 @@ void AccumulateAlbedoNormal(inout FPathTracingPayload Payload, float3 PathThroug
 		Normal += BrdfWeight * NormThroughput * Payload.WorldNormal;
 		break;
 	}
-#endif // SIMPLIFIED_MATERIAL_SHADER
 	default:
 		Albedo += BrdfWeight * PathThroughput * Payload.BaseColor;
 		Normal += BrdfWeight * NormThroughput * Payload.WorldNormal;
@@ -225,7 +221,6 @@ void AccumulateAlbedoNormal(inout FPathTracingPayload Payload, float3 PathThroug
 	}
 #endif // PATHTRACING_STRATA_PAYLOAD
 }
-
 #undef ROUGHNESS_WEIGHT
 
 void AccumulateAlbedo(float3 SigmaS, float3 PathThroughput, float PathRoughness, inout float3 Albedo)
@@ -243,6 +238,9 @@ void AccumulateAlbedo(float3 SigmaS, float3 PathThroughput, float PathRoughness,
 	Albedo += SigmaS * PathThroughput;
 }
 
+#endif // SIMPLIFIED_MATERIAL_SHADER
+
+
 #define DEBUG_DEPTH_SAMPLING 0
 
 struct FDepthContext
@@ -257,15 +255,12 @@ struct FDepthContext
 
 	void AddHit(FPathTracingPayload Payload)
 	{
-		if (false
 #if DEBUG_DEPTH_SAMPLING
-			|| DebugHitCount == 0
-#endif
-			)
+		if (DebugHitCount == 0)
 		{
 			return;
 		}
-
+#endif
 		// 2. Matching depth output to the rasterizer.
 		if (Payload.ShouldOutputDepth() && HitCount == 0)
 		{
@@ -342,6 +337,7 @@ float GetAverageRoughness(FPathTracingPayload Payload)
 #endif
 }
 
+#if !SIMPLIFIED_MATERIAL_SHADER
 void AdjustPayloadAfterUnpack(inout FPathTracingPayload HitPayload, bool bApplyOverrides)
 {
 	// This function can also do a bit of per-material post-processing to simplify the rest of the code
@@ -350,6 +346,7 @@ void AdjustPayloadAfterUnpack(inout FPathTracingPayload HitPayload, bool bApplyO
 		HitPayload.DiffuseColor  = HitPayload.DiffuseColor  * View.DiffuseOverrideParameter.w + View.DiffuseOverrideParameter.xyz;
 		HitPayload.SpecularColor = HitPayload.SpecularColor * View.SpecularOverrideParameter.w + View.SpecularOverrideParameter.xyz;
 	}
+
 #if PATHTRACING_STRATA_PAYLOAD
 	switch (HitPayload.ShadingModelID)
 	{
@@ -380,14 +377,12 @@ void AdjustPayloadAfterUnpack(inout FPathTracingPayload HitPayload, bool bApplyO
 #else
 	switch (HitPayload.ShadingModelID)
 	{
-#if !SIMPLIFIED_MATERIAL_SHADER
 		case SHADINGMODELID_SUBSURFACE:
 		case SHADINGMODELID_PREINTEGRATED_SKIN:
 		{
 			HitPayload.SubsurfaceRadius = HitPayload.GetSubsurfaceRadius();
 			// fallthrough to case below to avoid over-bright SSS
 		}
-#endif
 		case SHADINGMODELID_TWOSIDED_FOLIAGE:
 		{
 			HitPayload.SubsurfaceColor = HitPayload.GetSubsurfaceColor();
@@ -408,13 +403,6 @@ void AdjustPayloadAfterUnpack(inout FPathTracingPayload HitPayload, bool bApplyO
 			}
 			break;
 		}
-#if SIMPLIFIED_MATERIAL_SHADER
-		default:
-		{
-			HitPayload.SubsurfaceColor = 0.0;
-			break;
-		}
-#else
 		case SHADINGMODELID_SUBSURFACE_PROFILE:
 		{
 			HitPayload.SubsurfaceRadius = HitPayload.GetSubsurfaceRadius();
@@ -457,7 +445,6 @@ void AdjustPayloadAfterUnpack(inout FPathTracingPayload HitPayload, bool bApplyO
 			}
 			break;
 		}
-#endif
 	}
 #endif // PATHTRACING_STRATA_PAYLOAD
 }
@@ -550,6 +537,8 @@ void RemoveMaterialSSS(inout FPathTracingPayload Payload)
 	Payload.SubsurfaceColor = 0.0;
 }
 
+#endif // SIMPLIFIED_MATERIAL_SHADER
+
 
 FMaterialSample SampleMaterial(
 	float3 V_World,
diff --git a/Engine/Shaders/Private/PathTracing/Material/PathTracingThinGlass.ush b/Engine/Shaders/Private/PathTracing/Material/PathTracingThinGlass.ush
--- a/Engine/Shaders/Private/PathTracing/Material/PathTracingThinGlass.ush
+++ b/Engine/Shaders/Private/PathTracing/Material/PathTracingThinGlass.ush
@@ -54,6 +54,7 @@ FThinSlabWeights ComputeThinSlabWeights(float3 SlabColor, float CosV, float Eta,
 	return Result;
 }
 
+#if !SIMPLIFIED_MATERIAL_SHADER
 
 struct FRoughThinGlassData
 {
@@ -286,3 +287,5 @@ FMaterialSample RoughThinGlass_SampleMaterial(
 		return Result;
 	}
 }
+
+#endif
diff --git a/Engine/Shaders/Private/PathTracing/PathTracingGPULightmassMaterialHitShader.usf b/Engine/Shaders/Private/PathTracing/PathTracingGPULightmassMaterialHitShader.usf
--- a/Engine/Shaders/Private/PathTracing/PathTracingGPULightmassMaterialHitShader.usf
+++ b/Engine/Shaders/Private/PathTracing/PathTracingGPULightmassMaterialHitShader.usf
@@ -15,6 +15,9 @@
 // Ensure that SSS albedo comes through in the material
 #define STRATA_SSS_MATERIAL_OVERRIDE 0
 
+// Not using precomputed lighting
+#define MATERIAL_STRATA_OPAQUE_PRECOMPUTED_LIGHTING  0
+
 #include "/Engine/Private/Common.ush"
 #include "/Engine/Private/RayTracing/RayTracingCommon.ush"
 #include "/Engine/Private/RayTracing/RayTracingHitGroupCommon.ush"
@@ -32,6 +35,7 @@
 
 #if STRATA_ENABLED
 #define STRATA_GPU_LIGHTMASS 1
+#include "/Engine/Private/Strata/Strata.ush"
 #include "/Engine/Private/Strata/StrataExport.ush"
 #endif
 
@@ -88,17 +92,17 @@ RAY_TRACING_ENTRY_CLOSEST_HIT(GPULightmassMaterialCHS,
 
 #if STRATA_ENABLED
 
-
-
-	FStrataData StrataData = PixelMaterialInputs.GetFrontStrataData();
-	FStrataPixelHeader StrataPixelHeader = MaterialParameters.GetFrontStrataHeader();
-
 	float  Coverage = 0;
 	float3 TransmittancePreCoverage = 0;
 	float3 BaseColorPostCoverage = 0;
 	float3 WorldNormal = 0;
 	float3 EmissiveLuminance = 0;
 
+#if MATERIAL_IS_STRATA
+	FStrataData StrataData = PixelMaterialInputs.GetFrontStrataData();
+	FStrataPixelHeader StrataPixelHeader = MaterialParameters.GetFrontStrataHeader();
+
+
 	FStrataIntegrationSettings Settings = InitStrataIntegrationSettings(true /*bForceFullyRough*/, false /*StrataStruct.bRoughDiffuse*/, -1/*StrataStruct.PeelLayersAboveDepth*/, false/*StrataStruct.bRoughnessTracking*/);
 
 	const float3 SurfaceWorldNormal = GeoNormal;
@@ -112,6 +116,7 @@ RAY_TRACING_ENTRY_CLOSEST_HIT(GPULightmassMaterialCHS,
 		BaseColorPostCoverage,
 		WorldNormal,
 		EmissiveLuminance);
+#endif // MATERIAL_IS_STRATA
 
 
 #if MATERIALBLENDING_ANY_TRANSLUCENT
@@ -123,17 +128,20 @@ RAY_TRACING_ENTRY_CLOSEST_HIT(GPULightmassMaterialCHS,
 		PackedPayload.SetRayThroughput(RayThroughput);
 		return;
 	}
+	// STRATA_TODO: Make sure this is correct for all blend modes/glass cases
+	float3 TransparencyColor = saturate((1.0f - Coverage) + Coverage * TransmittancePreCoverage);
+#else
+	float3 TransparencyColor = 0.0;
 #endif
 
 	FPathTracingPayload Payload = (FPathTracingPayload)0;
 
-	Payload.BaseColor = BaseColorPostCoverage;
-	Payload.SubsurfaceColor = 0;							// STRATA_TODO unused by GPU lightmass?
 	Payload.BSDFOpacity = Coverage;
-	Payload.TransparencyColor = 0;							// STRATA_TODO That is not important when not evaluating throughput for IsVisibilityRay() above.
+	Payload.BaseColor = BaseColorPostCoverage;
+	Payload.SubsurfaceColor = 0;							// STRATA_TODO: Figure out how to support old two-sided foliage approximation
+	Payload.ShadingModelID = any(BaseColorPostCoverage > 0.0) ? SHADINGMODELID_DEFAULT_LIT : SHADINGMODELID_UNLIT;
+	Payload.TransparencyColor = TransparencyColor;
 
-	Payload.ShadingModelID = SHADINGMODELID_DEFAULT_LIT;	// STRATA_TODO check that this is enough for GPULightmass
-	
 	const uint PrimitiveFlags = GetPrimitiveData(MaterialParameters.PrimitiveId).Flags;
 	Payload.PrimitiveLightingChannelMask = GetPrimitive_LightingChannelMask_FromFlags(PrimitiveFlags);
 	Payload.HitT = RayTCurrent();
diff --git a/Engine/Shaders/Private/PathTracing/PathTracingMaterialHitShader.usf b/Engine/Shaders/Private/PathTracing/PathTracingMaterialHitShader.usf
--- a/Engine/Shaders/Private/PathTracing/PathTracingMaterialHitShader.usf
+++ b/Engine/Shaders/Private/PathTracing/PathTracingMaterialHitShader.usf
@@ -46,9 +46,23 @@ float AdjustMaterialRoughness(float Roughness, float PathRoughness)
 	return max(Roughness, PathRoughness);
 }
 
+#if STRATA_ENABLED
+
+float GetRefractionIor(FPixelMaterialInputs PixelMaterialInputs, float3 F0)
+{
+#if REFRACTION_ROOT_NODE_OVERRIDES_DEFAULT
+	float Ior = GetMaterialRefractionIOR(GetMaterialRefraction(PixelMaterialInputs));
+#else
+	float Ior = DielectricF0RGBToIor(F0);
+#endif
+	return Ior;
+}
+
+#else // STRATA_ENABLED
+
 float GetRefractionIor(FPixelMaterialInputs PixelMaterialInputs)
 {
-#if REFRACTION_ROOT_NODE_OVERRIDES_DEFAULT || STRATA_OPTIMIZED_UNLIT
+#if REFRACTION_ROOT_NODE_OVERRIDES_DEFAULT
 	float Ior = GetMaterialRefractionIOR(GetMaterialRefraction(PixelMaterialInputs));
 #else
 	float3 BaseColor = GetMaterialBaseColor(PixelMaterialInputs);
@@ -60,6 +74,8 @@ float GetRefractionIor(FPixelMaterialInputs PixelMaterialInputs)
 	return Ior;
 }
 
+#endif // STRATA_ENABLED
+
 float3 DecodeSSSProfileRadius(uint ProfileId, float3 DiffuseColor, float Opacity)
 {
 	// Burley parameterization
@@ -139,6 +155,8 @@ RAY_TRACING_ENTRY_CLOSEST_HIT(PathTracingMaterialCHS,
 	}
 
 #if STRATA_ENABLED
+
+#if MATERIAL_IS_STRATA
 	FStrataData StrataData = PixelMaterialInputs.GetFrontStrataData();
 
 #if STRATA_OPTIMIZED_UNLIT
@@ -152,9 +170,14 @@ RAY_TRACING_ENTRY_CLOSEST_HIT(PathTracingMaterialCHS,
 	float3 TotalTransmittancePreCoverage = 0;
 	const float3 V_World = MaterialParameters.CameraVector;
 	StrataPixelHeader.StrataUpdateTree(NullStrataAddressing, StrataData, V_World, IntegrationSettings, TotalCoverage, TotalTransmittancePreCoverage);
-	const float Transparency = 1 - TotalCoverage;
+	const float3 Transparency = 1 - TotalCoverage * TotalTransmittancePreCoverage;
 #endif // STRATA_OPTIMIZED_UNLIT
 
+#else // MATERIAL_IS_STRATA
+	// This only happens for invalid materials
+    const float3 Transparency = 0;
+#endif
+
 #endif // STRATA_ENABLED
 
 	// the following blend modes need to process shadows after having executed the material
@@ -379,6 +402,8 @@ RAY_TRACING_ENTRY_CLOSEST_HIT(PathTracingMaterialCHS,
 	// unless proven otherwise, we are a basic unlit material
 	Payload.ShadingModelID = SHADINGMODELID_UNLIT;
 
+#if MATERIAL_IS_STRATA // This can become false if the material is invalid
+
 #if STRATA_OPTIMIZED_UNLIT
 	Payload.Radiance = BSDF_GETEMISSIVE(StrataData.InlinedBSDF);
 #if STRATA_OPAQUE_MATERIAL
@@ -569,6 +594,7 @@ RAY_TRACING_ENTRY_CLOSEST_HIT(PathTracingMaterialCHS,
 #endif // STRATA_CLAMPED_BSDF_COUNT 
 
 
+#if STRATA_CLAMPED_BSDF_COUNT > 0
 			Payload.WorldNormal = StrataPixelHeader.SharedLocalBases.Normals[BSDF_GETSHAREDLOCALBASISID(BSDF)];
 			Payload.WorldTangent = StrataPixelHeader.SharedLocalBases.Tangents[BSDF_GETSHAREDLOCALBASISID(BSDF)];
 
@@ -627,7 +653,7 @@ RAY_TRACING_ENTRY_CLOSEST_HIT(PathTracingMaterialCHS,
 					Payload.WeightV             = BSDF.LuminanceWeightV;
 					Payload.TransmittanceN      = BSDF.TransmittanceAboveAlongN;
 					Payload.CoverageAboveAlongN = BSDF.CoverageAboveAlongN;
-					Payload.Ior                 = 0.0; // STRATA_TODO: figure out when to use refraction ...
+					Payload.Ior                 = GetRefractionIor(PixelMaterialInputs, Payload.SpecularColor);
 					break;
 				}
 				case STRATA_BSDF_TYPE_HAIR:
@@ -687,9 +713,12 @@ RAY_TRACING_ENTRY_CLOSEST_HIT(PathTracingMaterialCHS,
 			Payload.BSDFOpacity /= SlabPdf;
 #endif
 		} // close brace opened by STRATA_CLAMPED_BSDF_COUNT > 1 or STRATA_CLAMPED_BSDF_COUNT == 1 ifdefs above
+#endif // STRATA_CLAMPED_BSDF_COUNT > 0
+
 #undef  BSDF
 	} // if BSDFCount > 0
 
+
 #endif // STRATA_OPTIMIZED_UNLIT
 
 
@@ -701,6 +730,8 @@ RAY_TRACING_ENTRY_CLOSEST_HIT(PathTracingMaterialCHS,
 	}
 #endif
 
+#endif // MATERIAL_IS_STRATA
+
 #else // STRATA_ENABLED
 	Payload.Radiance = GetMaterialEmissive(PixelMaterialInputs);
 
diff --git a/Engine/Shaders/Private/Strata/StrataExport.ush b/Engine/Shaders/Private/Strata/StrataExport.ush
--- a/Engine/Shaders/Private/Strata/StrataExport.ush
+++ b/Engine/Shaders/Private/Strata/StrataExport.ush
@@ -872,9 +872,6 @@ void StrataMaterialExportOut(
 
 			if (StrataIsBSDFVisible(CurrentBSDF))
 			{
-				// Starting with environment lighting, the context does not specify a light
-				FStrataBSDFContext StrataBSDFContext = StrataCreateBSDFContext(StrataPixelHeader, CurrentBSDF, NullStrataAddressing, FakeV);
-
 				const float3x3 TangentBasis = StrataGetBSDFSharedBasis_InlineShading(StrataPixelHeader, CurrentBSDF, NullStrataAddressing);
 				const float3 N = TangentBasis[2];
 
