Fixing dx11 shader compilation error

diff --git a/Engine/Shaders/Private/HeterogeneousVolumes/HeterogeneousVolumesMaterialBakingPipeline.usf b/Engine/Shaders/Private/HeterogeneousVolumes/HeterogeneousVolumesMaterialBakingPipeline.usf
--- a/Engine/Shaders/Private/HeterogeneousVolumes/HeterogeneousVolumesMaterialBakingPipeline.usf
+++ b/Engine/Shaders/Private/HeterogeneousVolumes/HeterogeneousVolumesMaterialBakingPipeline.usf
@@ -377,7 +377,10 @@ void RasterizeBottomLevelGridCS(
 	#else
 		AllocatedVoxelResolution = BottomLevelVoxelResolution;
 	#endif
-		AllocatedVoxelCount = AllocatedVoxelResolution.x * AllocatedVoxelResolution.y * AllocatedVoxelResolution.z;
+		if (all(GroupThreadId == 0))
+		{
+			AllocatedVoxelCount = AllocatedVoxelResolution.x * AllocatedVoxelResolution.y * AllocatedVoxelResolution.z;
+		}
 
 		GroupMemoryBarrier();
 
@@ -648,6 +651,10 @@ void RasterizeBottomLevelFroxelGridCS(
 	// Aggregate in group-shared memory
 	uint LinearThreadIndex = MortonEncode3(GroupThreadId);
 	GSExtinctionSum[LinearThreadIndex] = Luminance(Extinction);
+	if (all(GroupThreadId == 0))
+	{
+		AllocatedVoxelCount = BottomLevelVoxelResolution.x * BottomLevelVoxelResolution.y * BottomLevelVoxelResolution.z;
+	}
 	GroupMemoryBarrierWithGroupSync();
 
 	// Parallel Sum
@@ -661,7 +668,6 @@ void RasterizeBottomLevelFroxelGridCS(
 
 	// Allocate if non-zero voxel data
 	bool bWasAlreadyAllocated = IsBottomLevelAllocated(TopLevelGridData);
-	AllocatedVoxelCount = BottomLevelVoxelResolution.x * BottomLevelVoxelResolution.y * BottomLevelVoxelResolution.z;
 	float ZeroThreshold = 1.0e-6;
 	if (all(GroupThreadId == 0) && any(GSExtinctionSum[0] > ZeroThreshold) && !bWasAlreadyAllocated)
 	{
