Fix shader compilation issue.

diff --git a/Engine/Shaders/Private/HairStrands/HairStrandsPack.ush b/Engine/Shaders/Private/HairStrands/HairStrandsPack.ush
--- a/Engine/Shaders/Private/HairStrands/HairStrandsPack.ush
+++ b/Engine/Shaders/Private/HairStrands/HairStrandsPack.ush
@@ -82,8 +82,8 @@ uint PackHairRootUV(float2 In)
 {
 	In.y = 1.0f - In.y;
 
-	return	(min(uint(saturate(In.x) * 0x7FF), 0x7FF)      )|
-			(min(uint(saturate(In.y) * 0x7FF), 0x7FF) << 11);
+	return	(min(uint(saturate(In.x) * 0x7FFu), 0x7FFu)      )|
+			(min(uint(saturate(In.y) * 0x7FFu), 0x7FFu) << 11);
 }
 
 float2 UnpackHairRootUV(uint In)
@@ -117,7 +117,7 @@ float UnpackHairSeed(uint In)
 
 uint PackHairClumpID(uint In)
 {
-	return min(In, 0xFFFF);
+	return min(In, 0xFFFFu);
 }
 
 uint3 UnpackHairClumpID(uint In)
@@ -132,17 +132,17 @@ uint3 UnpackHairClumpID(uint2 In)
 
 uint PackHairColor(float3 In)
 {
-	return	(min(uint(In.x * 0x7FF), 0x7FF)    )|
-			(min(uint(In.y * 0x7FF), 0x7FF)<<11)|
-			(min(uint(In.z * 0x3FF), 0x3FF)<<22);
+	return	(min(uint(In.x * 0x7FFu), 0x7FFu)    )|
+			(min(uint(In.y * 0x7FFu), 0x7FFu)<<11)|
+			(min(uint(In.z * 0x3FFu), 0x3FFu)<<22);
 }
 
 float3 UnpackHairColor(uint In)
 {
 	float3 Out;
-	Out.x = (In       & 0x7FF) / 2047.f;
-	Out.y = ((In>>11) & 0x7FF) / 2047.f;
-	Out.z = ((In>>22) & 0x3FF) / 1023.f;
+	Out.x = (In       & 0x7FFu) / 2047.f;
+	Out.y = ((In>>11) & 0x7FFu) / 2047.f;
+	Out.z = ((In>>22) & 0x3FFu) / 1023.f;
 	return Out;
 }
 
