Path Tracing: Implement basic support for transparency for Substrate materials

Add basic handling of plain transparency, and refactor some code related to thin translucent to cleanup header includes.

diff --git a/Engine/Shaders/Private/PathTracing/Material/PathTracingFresnel.ush b/Engine/Shaders/Private/PathTracing/Material/PathTracingFresnel.ush
--- a/Engine/Shaders/Private/PathTracing/Material/PathTracingFresnel.ush
+++ b/Engine/Shaders/Private/PathTracing/Material/PathTracingFresnel.ush
@@ -126,3 +126,47 @@ float ClearCoatFresnel(float VoH)
 	// from hardcoded logic in ClearCoatBxDF
 	return CLEAR_COAT_F0 + (1.0 - CLEAR_COAT_F0) * Pow5(1 - VoH);
 }
+
+
+struct FThinSlabWeights
+{
+	float3 Reflected;
+	float3 Transmitted;
+};
+
+// return the amount of reflected and transmitted energy for an infinite slab viewed from a particular direction
+// taking into account the (repeated) Fresnel terms and the transmittance through the slab
+FThinSlabWeights ComputeThinSlabWeights(float3 SlabColor, float CosV, float Eta, float F0)
+{
+	// Compute Fresnel reflection off the top of the slab
+	// The Fresnel refletion at the bottom layer is exactly equal so we can assume energy is split into R and T at each bounce
+	// NOTE: we don't need to pass Eta through here because TIR (which could only occur in contrived scenarios) isn't really supported
+	const float R = FresnelReflectance(CosV, 1.0, F0);
+
+	// The slab color is normalized to unit thickness, so we only need to account for the refracted viewing angle
+	// Note that we only allow Eta > 1.0 here so that we don't need to worry about TIR
+	const float3 A = pow(SlabColor, Eta > 1.0 ? Eta * rsqrt(Eta * Eta - 1.0 + CosV * CosV) : rcp(abs(CosV)));
+
+	// figure out how much is reflected/transmitted from a single interaction
+	float3 SlabR = R;
+	float3 SlabT = A - R * A; // (1 - R) * A
+
+#if 1
+	// Enable bounces within the slab
+	if (R < 1.0)
+	{
+		// If there was some transmission we must also account for the infinite series of bounces within the slab
+		// which can be computed exactly with a geometric series. Note that this causes a bit of energy loss because we don't
+		// account for this non-linear fresnel term in the energy conservation calculation. However, this is fairly minimal in
+		// practice, and the added visual richness seems worth it.
+		float3 InvDenom = rcp(1.0 - Pow2(R * A));
+		SlabR += R * Pow2(SlabT) * InvDenom;
+		SlabT *= InvDenom - R * InvDenom; // (1 - R) * InvDenom
+	}
+#endif
+
+	FThinSlabWeights Result = (FThinSlabWeights)0;
+	Result.Reflected = SlabR;
+	Result.Transmitted = SlabT;
+	return Result;
+}
diff --git a/Engine/Shaders/Private/PathTracing/Material/PathTracingThinGlass.ush b/Engine/Shaders/Private/PathTracing/Material/PathTracingThinGlass.ush
--- a/Engine/Shaders/Private/PathTracing/Material/PathTracingThinGlass.ush
+++ b/Engine/Shaders/Private/PathTracing/Material/PathTracingThinGlass.ush
@@ -11,51 +11,6 @@ PathTracingThinGlass.ush: Microfacet BSDF for thin glass
 #include "PathTracingGlossy.ush"
 #include "PathTracingEnergyConservation.ush"
 
-struct FThinSlabWeights
-{
-	float3 Reflected;
-	float3 Transmitted;
-};
-
-// return the amount of reflected and transmitted energy for an infinite slab viewed from a particular direction
-// taking into account the (repeated) Fresnel terms and the transmittance through the slab
-FThinSlabWeights ComputeThinSlabWeights(float3 SlabColor, float CosV, float Eta, float F0)
-{
-	// Compute Fresnel reflection off the top of the slab
-	// The Fresnel refletion at the bottom layer is exactly equal so we can assume energy is split into R and T at each bounce
-	// NOTE: we don't need to pass Eta through here because TIR (which could only occur in contrived scenarios) isn't really supported
-	const float R = FresnelReflectance(CosV, 1.0, F0);
-
-	// The slab color is normalized to unit thickness, so we only need to account for the refracted viewing angle
-	// Note that we only allow Eta > 1.0 here so that we don't need to worry about TIR
-	const float3 A = pow(SlabColor, Eta > 1.0 ? Eta * rsqrt(Eta * Eta - 1.0 + CosV * CosV) : rcp(abs(CosV)));
-
-	// figure out how much is reflected/transmitted from a single interaction
-	float3 SlabR = R;
-	float3 SlabT = A - R * A; // (1 - R) * A
-
-#if 1
-	// Enable bounces within the slab
-	if (R < 1.0)
-	{
-		// If there was some transmission we must also account for the infinite series of bounces within the slab
-		// which can be computed exactly with a geometric series. Note that this causes a bit of energy loss because we don't
-		// account for this non-linear fresnel term in the energy conservation calculation. However, this is fairly minimal in
-		// practice, and the added visual richness seems worth it.
-		float3 InvDenom = rcp(1.0 - Pow2(R * A));
-		SlabR += R * Pow2(SlabT) * InvDenom;
-		SlabT *= InvDenom - R * InvDenom; // (1 - R) * InvDenom
-	}
-#endif
-
-	FThinSlabWeights Result = (FThinSlabWeights) 0;
-	Result.Reflected = SlabR;
-	Result.Transmitted = SlabT;
-	return Result;
-}
-
-#if !SIMPLIFIED_MATERIAL_SHADER
-
 struct FRoughThinGlassData
 {
 	float F0, F90;
@@ -287,5 +242,3 @@ FMaterialSample RoughThinGlass_SampleMaterial(
 		return Result;
 	}
 }
-
-#endif
diff --git a/Engine/Shaders/Private/PathTracing/PathTracingGPULightmassMaterialHitShader.usf b/Engine/Shaders/Private/PathTracing/PathTracingGPULightmassMaterialHitShader.usf
--- a/Engine/Shaders/Private/PathTracing/PathTracingGPULightmassMaterialHitShader.usf
+++ b/Engine/Shaders/Private/PathTracing/PathTracingGPULightmassMaterialHitShader.usf
@@ -31,7 +31,6 @@
 #include "/Engine/Private/SubsurfaceProfileCommon.ush"
 #include "/Engine/Private/BurleyNormalizedSSSCommon.ush"
 #include "/Engine/Private/PathTracing/Material/PathTracingFresnel.ush"
-#include "/Engine/Private/PathTracing/Material/PathTracingThinGlass.ush"
 
 #if STRATA_ENABLED
 #define STRATA_GPU_LIGHTMASS 1
diff --git a/Engine/Shaders/Private/PathTracing/PathTracingMaterialHitShader.usf b/Engine/Shaders/Private/PathTracing/PathTracingMaterialHitShader.usf
--- a/Engine/Shaders/Private/PathTracing/PathTracingMaterialHitShader.usf
+++ b/Engine/Shaders/Private/PathTracing/PathTracingMaterialHitShader.usf
@@ -32,8 +32,6 @@
 #if STRATA_ENABLED
 #define MATERIAL_STRATA_OPAQUE_PRECOMPUTED_LIGHTING 0
 #include "/Engine/Private/Strata/StrataExport.ush"
-#else
-#include "/Engine/Private/PathTracing/Material/PathTracingThinGlass.ush"
 #endif
 
 float AdjustMaterialRoughness(float Roughness, float PathRoughness)
@@ -165,12 +163,16 @@ RAY_TRACING_ENTRY_CLOSEST_HIT(PathTracingMaterialCHS,
 	FStrataPixelHeader StrataPixelHeader = MaterialParameters.GetFrontStrataHeader();
 	FStrataIntegrationSettings IntegrationSettings = InitStrataIntegrationSettings();
 	FStrataAddressing NullStrataAddressing = (FStrataAddressing)0;	// Fake unused in StrataCreateBSDFContext when using Forward inline shading
+#if STRATA_BLENDING_COLOREDTRANSMITTANCEONLY
+	const uint BSDFCount = 0;
+#else
 	const uint BSDFCount = StrataPixelHeader.StrataTree.BSDFCount;
+#endif
 	float TotalCoverage = 1;
 	float3 TotalTransmittancePreCoverage = 0;
 	const float3 V_World = MaterialParameters.CameraVector;
 	StrataPixelHeader.StrataUpdateTree(NullStrataAddressing, StrataData, V_World, IntegrationSettings, TotalCoverage, TotalTransmittancePreCoverage);
-	const float3 Transparency = 1 - TotalCoverage * TotalTransmittancePreCoverage;
+	const float3 Transparency = saturate((1 - TotalCoverage) + TotalCoverage * TotalTransmittancePreCoverage);
 #endif // STRATA_OPTIMIZED_UNLIT
 
 #else // MATERIAL_IS_STRATA
@@ -417,13 +419,12 @@ RAY_TRACING_ENTRY_CLOSEST_HIT(PathTracingMaterialCHS,
 #else // STRATA_OPTIMIZED_UNLIT
 
 	Payload.BSDFOpacity = 1;
+#if !STRATA_OPAQUE_MATERIAL
+	Payload.TransparencyColor = Transparency;
+#endif
 
 	if (BSDFCount > 0)
 	{
-#if !STRATA_OPAQUE_MATERIAL
-		Payload.TransparencyColor = Transparency;
-#endif
-	
 		{
 #if HAIR_STRAND_MESH_FACTORY
 			Payload.WorldSmoothNormal = Payload.WorldGeoNormal;
