Fixing shader warnings

diff --git a/Engine/Shaders/Private/HeterogeneousVolumes/HeterogeneousVolumesMaterialBakingPipeline.usf b/Engine/Shaders/Private/HeterogeneousVolumes/HeterogeneousVolumesMaterialBakingPipeline.usf
--- a/Engine/Shaders/Private/HeterogeneousVolumes/HeterogeneousVolumesMaterialBakingPipeline.usf
+++ b/Engine/Shaders/Private/HeterogeneousVolumes/HeterogeneousVolumesMaterialBakingPipeline.usf
@@ -394,9 +394,9 @@ void RasterizeBottomLevelGridCS(
 #else
 		ExtinctionSum[LinearThreadIndex] = Extinction[LinearThreadIndex];
 		// Parallel Sum
-		for (int NeighborOffset = 1; NeighborOffset < THREADGROUP_SIZE_1D; NeighborOffset = (NeighborOffset << 1))
+		for (int NeighborOffsetSum = 1; NeighborOffsetSum < THREADGROUP_SIZE_1D; NeighborOffsetSum = (NeighborOffsetSum << 1))
 		{
-			int NeighborIndex = LinearThreadIndex + NeighborOffset;
+			int NeighborIndex = LinearThreadIndex + NeighborOffsetSum;
 			ExtinctionSum[LinearThreadIndex] += (NeighborIndex < THREADGROUP_SIZE_1D) ? ExtinctionSum[NeighborIndex] : 0.0;
 
 			//GroupMemoryBarrierWithGroupSync();
@@ -414,7 +414,7 @@ void RasterizeBottomLevelGridCS(
 		ZeroCrossings[LinearThreadIndex] = CalcZeroCrossing(Extinction[LinearThreadIndex]);
 
 		// Parallel Sum
-		for (NeighborOffset = 1; NeighborOffset < AllocatedVoxelCount; NeighborOffset = (NeighborOffset << 1))
+		for (int NeighborOffset = 1; NeighborOffset < AllocatedVoxelCount; NeighborOffset = (NeighborOffset << 1))
 		{
 			int NeighborIndex = LinearThreadIndex + NeighborOffset;
 			ExtinctionVariance[LinearThreadIndex] += (NeighborIndex < AllocatedVoxelCount) ? ExtinctionVariance[NeighborIndex] : 0.0;
@@ -434,13 +434,13 @@ void RasterizeBottomLevelGridCS(
 		GroupMemoryBarrierWithGroupSync();
 		ExtinctionVariance[0] *= rcp(AllocatedVoxelCount - 1);
 
-		float StdRelError = sqrt(ExtinctionVariance[0] * rcp(AllocatedVoxelCount)) / ExtinctionMean;
+		float StdRelError = sqrt(ExtinctionVariance[0] * rcp(AllocatedVoxelCount)) / Luminance(ExtinctionMean);
 
 		// Declare low-variance to constitute a region of homogeneity
 		if (all(GroupThreadId == 0) && (StdRelError < HomogeneousThreshold) && (ZeroCrossings[0] < 4))
 		{
 			Extinction[0] = ExtinctionMean;
-			AllocatedVoxelCount = AllocatedVoxelResolution = 1;
+			AllocatedVoxelResolution = AllocatedVoxelCount = 1;
 		}
 #endif // #if DIM_ENABLE_HOMOGENEOUS_AGGREGATION
 
@@ -492,7 +492,7 @@ void RasterizeBottomLevelGridCS(
 					InterlockedExchange(RWBottomLevelGridAllocatorBuffer[0], BottomLevelGridBufferSize, Dummy);
 
 					BottomLevelIndex = EMPTY_VOXEL_INDEX;
-					AllocatedVoxelCount = AllocatedVoxelResolution = 0;
+					AllocatedVoxelResolution = AllocatedVoxelCount = 0;
 				}
 			}
 	#if DIM_ENABLE_VOXEL_HASHING
