Adding HeterogeneousVolume support for reference tracking mode.

diff --git a/Engine/Shaders/Private/PathTracing/Volume/PathTracingHeterogeneousVolumes.ush b/Engine/Shaders/Private/PathTracing/Volume/PathTracingHeterogeneousVolumes.ush
--- a/Engine/Shaders/Private/PathTracing/Volume/PathTracingHeterogeneousVolumes.ush
+++ b/Engine/Shaders/Private/PathTracing/Volume/PathTracingHeterogeneousVolumes.ush
@@ -279,7 +279,8 @@ FVolumeShadedResult GetOrthoVoxelGridDensity(float3 TranslatedWorldPos)
 
 				// Convert to Indirection voxel-space
 				float3 IndirectionGridVoxelPos = frac(TopLevelVoxelPos) * IndirectionGridVoxelResolution;
-				FTopLevelGridData IndirectionData = OrthoGridUniformBuffer.IndirectionGridBuffer[IndirectionGridIndex + MortonEncode3(IndirectionGridVoxelPos)];
+				int3 IndirectionGridVoxelPosAsInt = clamp(IndirectionGridVoxelPos, 0, IndirectionGridVoxelResolution - 1);
+				FTopLevelGridData IndirectionData = OrthoGridUniformBuffer.IndirectionGridBuffer[IndirectionGridIndex + MortonEncode3(IndirectionGridVoxelPosAsInt)];
 				if (IsBottomLevelAllocated(IndirectionData))
 				{
 					uint BottomLevelIndex = GetBottomLevelIndex(IndirectionData);
@@ -287,7 +288,8 @@ FVolumeShadedResult GetOrthoVoxelGridDensity(float3 TranslatedWorldPos)
 
 					// Constant interpolation
 					float3 BottomLevelVoxelPos = frac(IndirectionGridVoxelPos) * BottomLevelVoxelResolution;
-					uint LinearBottomLevelVoxelPos = BottomLevelIndex + MortonEncode3(uint3(BottomLevelVoxelPos));
+					int3 BottomLevelVoxelPosAsInt = clamp(BottomLevelVoxelPos, 0, BottomLevelVoxelResolution - 1);
+					uint LinearBottomLevelVoxelPos = BottomLevelIndex + MortonEncode3(BottomLevelVoxelPosAsInt);
 					SigmaT = GetExtinction(OrthoGridUniformBuffer.ExtinctionGridBuffer[LinearBottomLevelVoxelPos]);
 					Albedo = GetAlbedo(OrthoGridUniformBuffer.ScatteringGridBuffer[LinearBottomLevelVoxelPos]);
 				}
@@ -303,7 +305,8 @@ FVolumeShadedResult GetOrthoVoxelGridDensity(float3 TranslatedWorldPos)
 
 				// Constant Interpolation
 				float3 BottomLevelVoxelPos = frac(TopLevelVoxelPos) * BottomLevelVoxelResolution;
-				uint LinearBottomLevelVoxelPos = BottomLevelIndex + MortonEncode3(uint3(BottomLevelVoxelPos));
+				int3 BottomLevelVoxelPosAsInt = clamp(BottomLevelVoxelPos, 0, BottomLevelVoxelResolution - 1);
+				uint LinearBottomLevelVoxelPos = BottomLevelIndex + MortonEncode3(BottomLevelVoxelPosAsInt);
 
 				SigmaT = GetExtinction(OrthoGridUniformBuffer.ExtinctionGridBuffer[LinearBottomLevelVoxelPos]);
 				Albedo = GetAlbedo(OrthoGridUniformBuffer.ScatteringGridBuffer[LinearBottomLevelVoxelPos]);
@@ -313,7 +316,6 @@ FVolumeShadedResult GetOrthoVoxelGridDensity(float3 TranslatedWorldPos)
 
 	FVolumeShadedResult Result = (FVolumeShadedResult)0;
 	Result.SigmaT = SigmaT;
-	//Result.SigmaSHG = Albedo * Result.SigmaT;
 	Result.SigmaSHG = Albedo;
 	Result.PhaseG = GetPhaseG();
 	// TODO: Emission?
@@ -355,7 +357,6 @@ FVolumeShadedResult GetFrustumVoxelGridDensity(float3 TranslatedWorldPos, inout
 	// Return struct
 	FVolumeShadedResult Result = (FVolumeShadedResult) 0;
 	Result.SigmaT = SigmaT;
-	//Result.SigmaSHG = Albedo * Result.SigmaT;
 	Result.SigmaSHG = Albedo;
 	Result.PhaseG = GetPhaseG();
 	// TODO: Emission?
diff --git a/Engine/Shaders/Private/PathTracing/Volume/PathTracingVolume.ush b/Engine/Shaders/Private/PathTracing/Volume/PathTracingVolume.ush
--- a/Engine/Shaders/Private/PathTracing/Volume/PathTracingVolume.ush
+++ b/Engine/Shaders/Private/PathTracing/Volume/PathTracingVolume.ush
@@ -171,6 +171,11 @@ float3 VolumeGetTransmittance(float3 StartThroughput, float3 Origin, float3 Dire
 				PATH_TRACER_REFERENCE_TRANSMITTANCE_LOOP(FogGetDensityBounds, FogGetDensity);
 				break;
 			}
+			case VOLUMEID_HETEROGENEOUS_VOLUMES:
+			{
+				PATH_TRACER_REFERENCE_TRANSMITTANCE_LOOP(HeterogeneousVolumesGetDensityBounds, HeterogeneousVolumesGetDensity);
+				break;
+			}
 		}
 	}
 	return Throughput;
@@ -195,6 +200,11 @@ float3 VolumeGetTransmittance(float3 StartThroughput, float3 Origin, float3 Dire
 				PATH_TRACER_REFERENCE_TRANSMITTANCE_LOOP(FogGetDensityBounds, FogGetDensity);
 				break;
 			}
+			case VOLUMEID_HETEROGENEOUS_VOLUMES:
+			{
+				PATH_TRACER_REFERENCE_TRANSMITTANCE_LOOP(HeterogeneousVolumesGetDensityBounds, HeterogeneousVolumesGetDensity);
+				break;
+			}
 		}
 	}
 	return Throughput;
