Heterogeneous Volumes VoxelGrid Pipeline fixes

Voxelization refactor
Temporarily remove hashing deduplication

diff --git a/Engine/Shaders/Private/HeterogeneousVolumes/HeterogeneousVolumesMaterialBakingPipeline.usf b/Engine/Shaders/Private/HeterogeneousVolumes/HeterogeneousVolumesMaterialBakingPipeline.usf
--- a/Engine/Shaders/Private/HeterogeneousVolumes/HeterogeneousVolumesMaterialBakingPipeline.usf
+++ b/Engine/Shaders/Private/HeterogeneousVolumes/HeterogeneousVolumesMaterialBakingPipeline.usf
@@ -113,32 +113,29 @@ struct FRasterTileData
 Buffer<uint> RasterTileAllocatorBuffer;
 StructuredBuffer<FRasterTileData> RasterTileBuffer;
 
-RWStructuredBuffer<uint> RWTopLevelGridBuffer;
-RWBuffer<uint> RWBottomLevelGridAllocatorBuffer;
+RWStructuredBuffer<FTopLevelGridData> RWTopLevelGridBuffer;
+RWStructuredBuffer<FTopLevelGridData> RWIndirectionGridBuffer;
 
+RWBuffer<uint> RWBottomLevelGridAllocatorBuffer;
 RWStructuredBuffer<FScalarGridData> RWExtinctionGridBuffer;
 RWStructuredBuffer<FVectorGridData> RWEmissionGridBuffer;
 RWStructuredBuffer<FVectorGridData> RWScatteringGridBuffer;
 
-// Indirection Grid
 int FixedBottomLevelResolution;
-RWStructuredBuffer<FTopLevelGridData> RWIndirectionGridBuffer;
 
 // Hash index to voxel index
-RWBuffer<uint> RWHashToVoxelBuffer;
+//RWBuffer<uint> RWHashToVoxelBuffer;
 
 // Sampling mode
 int bSampleAtVertices;
 
 float HomogeneousThreshold;
 
-groupshared float3 Extinction[THREADGROUP_SIZE_1D];
-groupshared float3 ExtinctionSum[THREADGROUP_SIZE_1D];
-groupshared float ExtinctionVariance[THREADGROUP_SIZE_1D];
-groupshared float ZeroCrossings[THREADGROUP_SIZE_1D];
-groupshared FTopLevelGridData AllocatedGridData;
-groupshared int3 AllocatedVoxelResolution;
-groupshared int AllocatedVoxelCount;
+groupshared float3 GSExtinctionSum[THREADGROUP_SIZE_1D];
+groupshared float GSExtinctionVariance[THREADGROUP_SIZE_1D];
+groupshared float GSZeroCrossings[THREADGROUP_SIZE_1D];
+groupshared int3 GSAllocatedVoxelResolution;
+groupshared int GSAllocatedVoxelCount;
 
 float GetZeroThreshold()
 {
@@ -157,6 +154,7 @@ float CalcZeroCrossing(float3 SampleValue)
 	return all(SampleValue < EmptyVoxelThreshold) ? 1.0 : 0.0;
 }
 
+#if 0
 uint2 CalcStencil(uint VoxelCount)
 {
 	//int x = i % 4;
@@ -168,135 +166,18 @@ uint2 CalcStencil(uint VoxelCount)
 	for (int i = 0; i < min(VoxelCount, 32); ++i)
 	{
 		int bit = i;
-		Stencil[0] |= any(Extinction[i] > ZeroThreshold) ? 1 << bit : 0;
+		Stencil[0] |= any(GSExtinction[i] > ZeroThreshold) ? 1 << bit : 0;
 	}
 	
 	for (i = 32; i < min(VoxelCount, 64); ++i)
 	{
 		int bit = i - 32;
-		Stencil[1] |= any(Extinction[i] > ZeroThreshold) ? 1 << bit : 0;
+		Stencil[1] |= any(GSExtinction[i] > ZeroThreshold) ? 1 << bit : 0;
 	}
 
 	return Stencil;
 }
-
-float ExtinctionMax()
-{
-	return 0.01;
-}
-
-float ExtinctionFixedPointCeiling()
-{
-	//int QuantizeBits = 16;
-	int QuantizeBits = 8;
-	return 1 << QuantizeBits;
-}
-
-// TODO: Parallel
-uint QuantizeOpticalThickness(uint VoxelCount)
-{
-	float ZeroThreshold = GetZeroThreshold();
-
-	float3 ExtinctionMean = 0;
-	uint NonZeroCount = 0;
-	for (uint i = 0; i < VoxelCount; ++i)
-	{
-		if (any(Extinction[i] > ZeroThreshold))
-		{
-			ExtinctionMean += Extinction[i];
-			NonZeroCount++;
-		}
-	}
-
-	if (NonZeroCount > 0)
-	{
-		ExtinctionMean /= NonZeroCount;
-	}
-
-	return Luminance(ExtinctionMean) / ExtinctionMax() * ExtinctionFixedPointCeiling();
-}
-
-// TODO: Parallel
-float3 QuantizeExtinction(uint VoxelCount)
-{
-	float ZeroThreshold = GetZeroThreshold();
-
-	float3 ExtinctionMean = 0;
-	uint NonZeroCount = 0;
-	for (uint i = 0; i < VoxelCount; ++i)
-	{
-		if (any(Extinction[i] > ZeroThreshold))
-		{
-			ExtinctionMean += Extinction[i];
-			NonZeroCount++;
-		}
-	}
-
-	if (NonZeroCount > 0)
-	{
-		ExtinctionMean /= NonZeroCount;
-	}
-
-	uint3 QuantizedValue = ExtinctionMean / ExtinctionMax() * ExtinctionFixedPointCeiling();
-	return QuantizedValue / ExtinctionFixedPointCeiling() * ExtinctionMax();
-}
-
-uint CalcBottomLevelHashKey(uint VoxelCount)
-{
-	uint HashKey = 0;
-
-	// Hash Resolution
-	HashKey = MurmurMix(VoxelCount);
-
-	// Hash quantized average optical thickness
-	//HashKey = MurmurAdd(HashKey, QuantizeOpticalThickness(VoxelCount));
-	float3 QuantizedExtinction = QuantizeExtinction(VoxelCount);
-	uint QuantizedValue = Luminance(QuantizedExtinction / ExtinctionMax()) * ExtinctionFixedPointCeiling();
-	HashKey = MurmurAdd(HashKey, QuantizedValue);
-
-	// Hash bottom-voxel stencil
-	uint2 Stencil = CalcStencil(VoxelCount);
-	HashKey = MurmurAdd(HashKey, Stencil[0]);
-	HashKey = MurmurAdd(HashKey, Stencil[1]);
-
-	return HashKey;
-}
-
-bool HashTableFind2(uint Key, out uint Index)
-{
-	// Zero is reserved as invalid
-	Key++;
-
-	uint NumLoops = 0;
-	LOOP
-		[allow_uav_condition]
-	for (Index = MurmurMix(Key);; Index++)
-	{
-		//Belt and braces safety code to prevent inf loops if tables are malformed.
-		if (++NumLoops > HashTableSize)
-		{
-			break;
-		}
-
-		//Index &= HashTableSize - 1;
-		Index = Index % HashTableSize;
-
-		uint StoredKey = RWHashTable[Index];
-		if (StoredKey != Key)
-		{
-			if (StoredKey != 0)
-				continue;
-		}
-		else
-		{
-			return true;
-		}
-
-		break;
-	}
-
-	return false;
-}
+#endif
 
 [numthreads(THREADGROUP_SIZE_3D, THREADGROUP_SIZE_3D, THREADGROUP_SIZE_3D)]
 void RasterizeBottomLevelGridCS(
@@ -310,9 +191,6 @@ void RasterizeBottomLevelGridCS(
 		return;
 	}
 
-	uint LinearThreadIndex = MortonEncode3(GroupThreadId);
-	Extinction[LinearThreadIndex] = 0.0;
-
 	uint TopLevelGridLinearIndex = RasterTileBuffer[RasterTileIndex].TopLevelGridLinearIndex;
 	uint3 TopLevelGridVoxelIndex = GetVoxelIndex(TopLevelGridLinearIndex, TopLevelGridResolution);
 
@@ -327,21 +205,21 @@ void RasterizeBottomLevelGridCS(
 		// Map thread id to voxel-space position
 		uint BottomLevelGridLinearOffset = RasterTileBuffer[RasterTileIndex].BottomLevelGridLinearOffset;
 
-		FTopLevelGridData TopLevelGridData = CreateTopLevelGridData(RWTopLevelGridBuffer[TopLevelGridLinearIndex]);
+		FTopLevelGridData TopLevelGridData = RWTopLevelGridBuffer[TopLevelGridLinearIndex];
 		int3 BottomLevelVoxelResolution = GetBottomLevelVoxelResolution(TopLevelGridData);
 
 	#if DIM_ENABLE_INDIRECTION_GRID
 		int3 BottomLevelVoxelIndex = MortonDecode3(BottomLevelGridLinearOffset) * THREADGROUP_SIZE_3D + GroupThreadId;
 		BottomLevelVoxelResolution *= FixedBottomLevelResolution;
 	#else
-		int3 BottomLevelVoxelIndex = GetVoxelIndex(BottomLevelGridLinearOffset * THREADGROUP_SIZE_1D + MortonEncode3(GroupThreadId), BottomLevelVoxelResolution);
+		int3 BottomLevelVoxelIndex = MortonDecode3(BottomLevelGridLinearOffset) * THREADGROUP_SIZE_3D + GroupThreadId;
 	#endif
 
+		float3 Extinction = 0;
 		float3 Emission = 0;
 		float3 Albedo = 0;
 		if (all(BottomLevelVoxelIndex < BottomLevelVoxelResolution))
 		{
-#if 1
 			float3 UVW;
 			if (bSampleAtVertices)
 			{
@@ -363,166 +241,112 @@ void RasterizeBottomLevelGridCS(
 			CalcPixelMaterialInputs(MaterialParameters, PixelMaterialInputs);
 
 			// Extract volume rendering coefficients
-			Extinction[LinearThreadIndex] = SampleExtinctionCoefficients(PixelMaterialInputs);
+			Extinction = SampleExtinctionCoefficients(PixelMaterialInputs);
 			Emission = SampleEmissive(PixelMaterialInputs);
-			//Albedo = SampleAlbedo(PixelMaterialInputs);
-			Albedo = SampleAlbedo(PixelMaterialInputs) * Extinction[LinearThreadIndex];
-#else
-			Extinction[LinearThreadIndex] = 0.01;
-#endif
+			Albedo = SampleAlbedo(PixelMaterialInputs) * Extinction;
 		}
 
+		// Aggregate in group-shared memory
+		uint LinearThreadIndex = MortonEncode3(GroupThreadId);
+		GSExtinctionSum[LinearThreadIndex] = Extinction;
+
+		if (all(GroupThreadId == 0))
+		{
 	#if DIM_ENABLE_INDIRECTION_GRID
-		AllocatedVoxelResolution = FixedBottomLevelResolution;
+			GSAllocatedVoxelResolution = FixedBottomLevelResolution;
 	#else
-		AllocatedVoxelResolution = BottomLevelVoxelResolution;
+			GSAllocatedVoxelResolution = BottomLevelVoxelResolution;
 	#endif
-		if (all(GroupThreadId == 0))
-		{
-			AllocatedVoxelCount = AllocatedVoxelResolution.x * AllocatedVoxelResolution.y * AllocatedVoxelResolution.z;
+			GSAllocatedVoxelCount = GSAllocatedVoxelResolution.x * GSAllocatedVoxelResolution.y * GSAllocatedVoxelResolution.z;
 		}
+		GroupMemoryBarrierWithGroupSync();
 
-		GroupMemoryBarrier();
-
-		// Sum contribution
-#if 0
-		ExtinctionSum[0] = 0;
-		for (uint Index = 0; Index < THREADGROUP_SIZE_1D; ++Index)
-		{
-			ExtinctionSum[0] += Extinction[Index];
-		}
-#else
-		ExtinctionSum[LinearThreadIndex] = Extinction[LinearThreadIndex];
 		// Parallel Sum
 		for (int NeighborOffsetSum = 1; NeighborOffsetSum < THREADGROUP_SIZE_1D; NeighborOffsetSum = (NeighborOffsetSum << 1))
 		{
 			int NeighborIndex = LinearThreadIndex + NeighborOffsetSum;
-			ExtinctionSum[LinearThreadIndex] += (NeighborIndex < THREADGROUP_SIZE_1D) ? ExtinctionSum[NeighborIndex] : 0.0;
+			GSExtinctionSum[LinearThreadIndex] += (NeighborIndex < THREADGROUP_SIZE_1D) ? GSExtinctionSum[NeighborIndex] : 0.0;
 
-			//GroupMemoryBarrierWithGroupSync();
-			GroupMemoryBarrier();
+			GroupMemoryBarrierWithGroupSync();
 		}
-#endif
-
-		GroupMemoryBarrierWithGroupSync();
 
 #if DIM_ENABLE_HOMOGENEOUS_AGGREGATION
 		// Calculate variance
-		float3 ExtinctionMean = ExtinctionSum[0] * rcp(AllocatedVoxelCount);
-		ExtinctionVariance[LinearThreadIndex] = CalcSampleVariance(ExtinctionMean, Extinction[LinearThreadIndex]);
-		// Reuse ExtinctionSum..
-		ZeroCrossings[LinearThreadIndex] = CalcZeroCrossing(Extinction[LinearThreadIndex]);
+		float3 ExtinctionMean = GSExtinctionSum[0] * rcp(GSAllocatedVoxelCount);
+		GSExtinctionVariance[LinearThreadIndex] = CalcSampleVariance(ExtinctionMean, Extinction);
+		GSZeroCrossings[LinearThreadIndex] = CalcZeroCrossing(Extinction);
 
 		// Parallel Sum
-		for (int NeighborOffset = 1; NeighborOffset < AllocatedVoxelCount; NeighborOffset = (NeighborOffset << 1))
+		for (int NeighborOffset = 1; NeighborOffset < THREADGROUP_SIZE_1D; NeighborOffset = (NeighborOffset << 1))
 		{
 			int NeighborIndex = LinearThreadIndex + NeighborOffset;
-			ExtinctionVariance[LinearThreadIndex] += (NeighborIndex < AllocatedVoxelCount) ? ExtinctionVariance[NeighborIndex] : 0.0;
-			ZeroCrossings[LinearThreadIndex] += (NeighborIndex < AllocatedVoxelCount) ? ZeroCrossings[NeighborIndex] : 0.0;
+			GSExtinctionVariance[LinearThreadIndex] += (NeighborIndex < THREADGROUP_SIZE_1D) ? GSExtinctionVariance[NeighborIndex] : 0.0;
+			GSZeroCrossings[LinearThreadIndex] += (NeighborIndex < THREADGROUP_SIZE_1D) ? GSZeroCrossings[NeighborIndex] : 0.0;
 
-			//GroupMemoryBarrierWithGroupSync();
-			GroupMemoryBarrier();
-		}
-#if 0	
-		// Linear Sum
-		for (uint Index = 0; Index < AllocatedVoxelCount; ++Index)
-		{
-			ExtinctionVariance[0] += ExtinctionVariance[Index];
+			GroupMemoryBarrierWithGroupSync();
 		}
-#endif
 
-		GroupMemoryBarrierWithGroupSync();
-		ExtinctionVariance[0] *= rcp(AllocatedVoxelCount - 1);
-
-		float StdRelError = sqrt(ExtinctionVariance[0] * rcp(AllocatedVoxelCount)) / Luminance(ExtinctionMean);
+		float ExtinctionVariance = GSExtinctionVariance[0] * rcp(GSAllocatedVoxelCount - 1);
+		float StdRelError = sqrt(ExtinctionVariance * rcp(GSAllocatedVoxelCount)) * rcp(ExtinctionMean);
 
 		// Declare low-variance to constitute a region of homogeneity
-		if (all(GroupThreadId == 0) && (StdRelError < HomogeneousThreshold) && (ZeroCrossings[0] < 4))
+		if (all(GroupThreadId == 0) && (StdRelError < HomogeneousThreshold) && (GSZeroCrossings[0] < 4))
 		{
-			Extinction[0] = ExtinctionMean;
-			AllocatedVoxelResolution = AllocatedVoxelCount = 1;
+			Extinction = ExtinctionMean;
+			GSAllocatedVoxelCount = GSAllocatedVoxelResolution = 1;
 		}
 #endif // #if DIM_ENABLE_HOMOGENEOUS_AGGREGATION
 
-		bool bAllocate = true;
-		uint BottomLevelIndex = EMPTY_VOXEL_INDEX;
-
 	#if DIM_ENABLE_INDIRECTION_GRID
 		uint IndirectionIndexOffset = GetBottomLevelIndex(TopLevelGridData);
 		uint IndirectionGridLinearOffset = IndirectionIndexOffset + BottomLevelGridLinearOffset;
 
 		FTopLevelGridData IndirectionData = RWIndirectionGridBuffer[IndirectionGridLinearOffset];
+		AllMemoryBarrierWithGroupSync();
 		bool bWasAlreadyAllocated = IsBottomLevelAllocated(IndirectionData);
-		if (all(GroupThreadId == 0) && any(ExtinctionSum[0] > 1.0e-4) && !bWasAlreadyAllocated)
 	#else
 		bool bWasAlreadyAllocated = IsBottomLevelAllocated(TopLevelGridData);
-		if (all(GroupThreadId == 0) && any(ExtinctionSum[0] > GetZeroThreshold()) && !bWasAlreadyAllocated)
 	#endif
+		if (all(GroupThreadId == 0) && !bWasAlreadyAllocated)
 		{
-			// Voxel Hashing
-			uint HashIndex = 0;
-	#if DIM_ENABLE_VOXEL_HASHING
-			if (AllocatedVoxelCount == 1)
-			{
-				uint HashKey = CalcBottomLevelHashKey(AllocatedVoxelCount);
-				bAllocate = !HashTableFind2(HashKey, HashIndex);
-			}
-	#endif // DIM_ENABLE_VOXEL_HASHING
+			uint BottomLevelIndex = EMPTY_VOXEL_INDEX;
 
-#if 0
-			// DEBUG: Only process aggregated voxels..
-			if (AllocatedVoxelCount != 1)
+			bool bShouldAllocate = any(GSExtinctionSum[0] > GetZeroThreshold());
+			if (bShouldAllocate)
 			{
-				BottomLevelIndex = EMPTY_VOXEL_INDEX;
-				AllocatedVoxelCount = AllocatedVoxelResolution = 0;
-				bAllocate = false;
-			}
-#endif
-
-			// Otherwise allocate
-			if (bAllocate)
-			{
-				InterlockedAdd(RWBottomLevelGridAllocatorBuffer[0], AllocatedVoxelCount, BottomLevelIndex);
+				InterlockedAdd(RWBottomLevelGridAllocatorBuffer[0], GSAllocatedVoxelCount, BottomLevelIndex);
 
 				// Guard against over allocation
-				if (BottomLevelIndex + AllocatedVoxelCount > BottomLevelGridBufferSize)
+				if (BottomLevelIndex + GSAllocatedVoxelCount > BottomLevelGridBufferSize)
 				{
 					// Declare the buffer to be filled
 					uint Dummy;
 					InterlockedExchange(RWBottomLevelGridAllocatorBuffer[0], BottomLevelGridBufferSize, Dummy);
 
 					BottomLevelIndex = EMPTY_VOXEL_INDEX;
-					AllocatedVoxelResolution = AllocatedVoxelCount = 0;
+					GSAllocatedVoxelCount = GSAllocatedVoxelResolution = 0;
 				}
+
+				// Update the index with the properly allocated index
+				FTopLevelGridData AllocatedGridData = (FTopLevelGridData)0;
+				SetBottomLevelIndex(AllocatedGridData, BottomLevelIndex);
+				SetBottomLevelVoxelResolution(AllocatedGridData, GSAllocatedVoxelResolution);
+
+#if DIM_ENABLE_INDIRECTION_GRID
+				RWIndirectionGridBuffer[IndirectionGridLinearOffset] = AllocatedGridData;
+#else
+				RWTopLevelGridBuffer[TopLevelGridLinearIndex] = AllocatedGridData;
+#endif
 			}
-	#if DIM_ENABLE_VOXEL_HASHING
 			else
 			{
-				BottomLevelIndex = RWHashToVoxelBuffer[HashIndex];
+				GSAllocatedVoxelCount = GSAllocatedVoxelResolution = 0;
 			}
-	#endif // DIM_ENABLE_VOXEL_HASHING
-
-			// Update the index with the properly allocated index
-			AllocatedGridData = TopLevelGridData;
-			SetBottomLevelIndex(AllocatedGridData, BottomLevelIndex);
-			SetBottomLevelVoxelResolution(AllocatedGridData, AllocatedVoxelResolution);
-
-	#if DIM_ENABLE_INDIRECTION_GRID
-			RWIndirectionGridBuffer[IndirectionGridLinearOffset] = AllocatedGridData;
-	#else
-			InterlockedCompareStore(RWTopLevelGridBuffer[TopLevelGridLinearIndex], TopLevelGridData.PackedData[0], AllocatedGridData.PackedData[0]);
-	#endif
 		}
 
 		AllMemoryBarrierWithGroupSync();
-		//AllMemoryBarrier();
 
-#if DIM_ENABLE_VOXEL_HASHING
-		// Compositing with hashed voxels should require copy-on-write, but it is not currently supported
-		if (bAllocate && LinearThreadIndex < AllocatedVoxelCount)
-#else
-		if (LinearThreadIndex < AllocatedVoxelCount)
-#endif
+		if (LinearThreadIndex < GSAllocatedVoxelCount)
 		{
 	#if DIM_ENABLE_INDIRECTION_GRID
 			FTopLevelGridData IndirectionData = RWIndirectionGridBuffer[IndirectionGridLinearOffset];
@@ -530,20 +354,19 @@ void RasterizeBottomLevelGridCS(
 			if (IsBottomLevelAllocated(IndirectionData))
 			{
 				uint BottomLevelVoxelLinearIndex = GetBottomLevelIndex(IndirectionData) + LinearThreadIndex;
-
 				if (bWasAlreadyAllocated)
 				{
-					Extinction[LinearThreadIndex] += GetExtinction(RWExtinctionGridBuffer[BottomLevelVoxelLinearIndex]);
+					Extinction += GetExtinction(RWExtinctionGridBuffer[BottomLevelVoxelLinearIndex]);
 					Emission += GetEmission(RWEmissionGridBuffer[BottomLevelVoxelLinearIndex]);
 					Albedo += GetAlbedo(RWScatteringGridBuffer[BottomLevelVoxelLinearIndex]);
 				}
 
-				SetExtinction(RWExtinctionGridBuffer[BottomLevelVoxelLinearIndex], Extinction[LinearThreadIndex]);
+				SetExtinction(RWExtinctionGridBuffer[BottomLevelVoxelLinearIndex], Extinction);
 				SetEmission(RWEmissionGridBuffer[BottomLevelVoxelLinearIndex], Emission);
 				SetAlbedo(RWScatteringGridBuffer[BottomLevelVoxelLinearIndex], Albedo);
 			}
 	#else
-			FTopLevelGridData TopLevelGridData = CreateTopLevelGridData(RWTopLevelGridBuffer[TopLevelGridLinearIndex]);
+			FTopLevelGridData TopLevelGridData = RWTopLevelGridBuffer[TopLevelGridLinearIndex];
 
 			if (IsBottomLevelAllocated(TopLevelGridData))
 			{
@@ -551,32 +374,17 @@ void RasterizeBottomLevelGridCS(
 
 				if (bWasAlreadyAllocated)
 				{
-					Extinction[LinearThreadIndex] += GetExtinction(RWExtinctionGridBuffer[BottomLevelVoxelLinearIndex]);
+					Extinction += GetExtinction(RWExtinctionGridBuffer[BottomLevelVoxelLinearIndex]);
 					Emission += GetEmission(RWEmissionGridBuffer[BottomLevelVoxelLinearIndex]);
 					Albedo += GetAlbedo(RWScatteringGridBuffer[BottomLevelVoxelLinearIndex]);
 				}
 
-				SetExtinction(RWExtinctionGridBuffer[BottomLevelVoxelLinearIndex], Extinction[LinearThreadIndex]);
+				SetExtinction(RWExtinctionGridBuffer[BottomLevelVoxelLinearIndex], Extinction);
 				SetEmission(RWEmissionGridBuffer[BottomLevelVoxelLinearIndex], Emission);
 				SetAlbedo(RWScatteringGridBuffer[BottomLevelVoxelLinearIndex], Albedo);
 			}
 	#endif
 		}
-
-	#if DIM_ENABLE_VOXEL_HASHING
-		// Update deduplication structure
-		if (all(GroupThreadId == 0) && AllocatedVoxelCount == 1)
-		//if (all(GroupThreadId == 0) && any(ExtinctionSum[0] > GetZeroThreshold()) && bAllocate)
-		{
-			uint HashKey = CalcBottomLevelHashKey(AllocatedVoxelCount);
-			uint HashIndex = 0;
-			if (HashTableAdd(HashKey, HashIndex))
-			{
-				uint Dummy;
-				InterlockedExchange(RWHashToVoxelBuffer[HashIndex], BottomLevelIndex, Dummy);
-			}
-		}
-	#endif // DIM_ENABLE_VOXEL_HASHING
 	}
 }
 
@@ -590,7 +398,7 @@ float FarPlaneDepth;
 #include "HeterogeneousVolumesFrustumVoxelGridUtils.ush"
 #include "../BlueNoise.ush"
 
-groupshared float GSExtinctionSum[THREADGROUP_SIZE_1D];
+groupshared float GSExtinctionSumScalar[THREADGROUP_SIZE_1D];
 
 bool EvaluationPositionIntersectsPrimitive(float3 WorldPosition)
 {
@@ -610,7 +418,7 @@ void RasterizeBottomLevelFroxelGridCS(
 	}
 
 	uint TopLevelGridLinearIndex = RasterTileBuffer[RasterTileIndex].TopLevelGridLinearIndex;
-	FTopLevelGridData TopLevelGridData = CreateTopLevelGridData(RWTopLevelGridBuffer[TopLevelGridLinearIndex]);
+	FTopLevelGridData TopLevelGridData = RWTopLevelGridBuffer[TopLevelGridLinearIndex];
 	int3 BottomLevelVoxelResolution = GetBottomLevelVoxelResolution(TopLevelGridData);
 
 	// Setup evaluation context
@@ -650,10 +458,10 @@ void RasterizeBottomLevelFroxelGridCS(
 
 	// Aggregate in group-shared memory
 	uint LinearThreadIndex = MortonEncode3(GroupThreadId);
-	GSExtinctionSum[LinearThreadIndex] = Luminance(Extinction);
+	GSExtinctionSumScalar[LinearThreadIndex] = Luminance(Extinction);
 	if (all(GroupThreadId == 0))
 	{
-		AllocatedVoxelCount = BottomLevelVoxelResolution.x * BottomLevelVoxelResolution.y * BottomLevelVoxelResolution.z;
+		GSAllocatedVoxelCount = BottomLevelVoxelResolution.x * BottomLevelVoxelResolution.y * BottomLevelVoxelResolution.z;
 	}
 	GroupMemoryBarrierWithGroupSync();
 
@@ -661,46 +469,54 @@ void RasterizeBottomLevelFroxelGridCS(
 	for (int NeighborOffset = 1; NeighborOffset < THREADGROUP_SIZE_1D; NeighborOffset = (NeighborOffset << 1))
 	{
 		int NeighborIndex = LinearThreadIndex + NeighborOffset;
-		GSExtinctionSum[LinearThreadIndex] += (NeighborIndex < THREADGROUP_SIZE_1D) ? GSExtinctionSum[NeighborIndex] : 0.0;
+		GSExtinctionSumScalar[LinearThreadIndex] += (NeighborIndex < THREADGROUP_SIZE_1D) ? GSExtinctionSumScalar[NeighborIndex] : 0.0;
 
 		GroupMemoryBarrierWithGroupSync();
 	}
 
+	uint BottomLevelIndex = EMPTY_VOXEL_INDEX;
+
 	// Allocate if non-zero voxel data
 	bool bWasAlreadyAllocated = IsBottomLevelAllocated(TopLevelGridData);
-	float ZeroThreshold = 1.0e-6;
-	if (all(GroupThreadId == 0) && any(GSExtinctionSum[0] > ZeroThreshold) && !bWasAlreadyAllocated)
+	AllMemoryBarrierWithGroupSync();
+	if (all(GroupThreadId == 0) && !bWasAlreadyAllocated)
 	{
 		uint3 AllocatedVoxelResolution = BottomLevelVoxelResolution;
-		AllocatedVoxelCount = AllocatedVoxelResolution.x * AllocatedVoxelResolution.y * AllocatedVoxelResolution.z;
-
-		uint BottomLevelIndex = 0;
-		InterlockedAdd(RWBottomLevelGridAllocatorBuffer[0], AllocatedVoxelCount, BottomLevelIndex);
-
-		// Guard against over allocation
-		if (BottomLevelIndex + AllocatedVoxelCount > BottomLevelGridBufferSize)
+		bool bShouldAllocate = GSExtinctionSumScalar[0] > GetZeroThreshold();
+		if (bShouldAllocate)
 		{
-			// Declare the buffer to be filled
-			uint Dummy;
-			InterlockedExchange(RWBottomLevelGridAllocatorBuffer[0], BottomLevelGridBufferSize, Dummy);
+			GSAllocatedVoxelCount = AllocatedVoxelResolution.x * AllocatedVoxelResolution.y * AllocatedVoxelResolution.z;
 
-			BottomLevelIndex = EMPTY_VOXEL_INDEX;
-			AllocatedVoxelResolution = AllocatedVoxelCount = 0;
-		}
+			InterlockedAdd(RWBottomLevelGridAllocatorBuffer[0], GSAllocatedVoxelCount, BottomLevelIndex);
 
-		// Update the index with the properly allocated index
-		FTopLevelGridData AllocatedGridData = TopLevelGridData;
-		SetBottomLevelIndex(AllocatedGridData, BottomLevelIndex);
-		SetBottomLevelVoxelResolution(AllocatedGridData, AllocatedVoxelResolution);
+			// Guard against over allocation
+			if (BottomLevelIndex + GSAllocatedVoxelCount > BottomLevelGridBufferSize)
+			{
+				// Declare the buffer to be filled
+				uint Dummy;
+				InterlockedExchange(RWBottomLevelGridAllocatorBuffer[0], BottomLevelGridBufferSize, Dummy);
+
+				BottomLevelIndex = EMPTY_VOXEL_INDEX;
+				GSAllocatedVoxelCount = AllocatedVoxelResolution = 0;
+			}
 
-		InterlockedCompareStore(RWTopLevelGridBuffer[TopLevelGridLinearIndex], TopLevelGridData.PackedData[0], AllocatedGridData.PackedData[0]);
+			// Update the index with the properly allocated index
+			FTopLevelGridData AllocatedGridData = (FTopLevelGridData)0;
+			SetBottomLevelIndex(AllocatedGridData, BottomLevelIndex);
+			SetBottomLevelVoxelResolution(AllocatedGridData, AllocatedVoxelResolution);
+			RWTopLevelGridBuffer[TopLevelGridLinearIndex] = AllocatedGridData;
+		}
+		else
+		{
+			GSAllocatedVoxelCount = AllocatedVoxelResolution = 0;
+		}
 	}
 
 	AllMemoryBarrierWithGroupSync();
 
-	if (LinearThreadIndex < AllocatedVoxelCount)
+	if (LinearThreadIndex < GSAllocatedVoxelCount)
 	{
-		FTopLevelGridData TopLevelGridData = CreateTopLevelGridData(RWTopLevelGridBuffer[TopLevelGridLinearIndex]);
+		FTopLevelGridData TopLevelGridData = RWTopLevelGridBuffer[TopLevelGridLinearIndex];
 		if (IsBottomLevelAllocated(TopLevelGridData))
 		{
 			uint BottomLevelVoxelLinearIndex = GetBottomLevelIndex(TopLevelGridData) + LinearThreadIndex;
diff --git a/Engine/Shaders/Private/HeterogeneousVolumes/HeterogeneousVolumesVoxelGridPipeline.usf b/Engine/Shaders/Private/HeterogeneousVolumes/HeterogeneousVolumesVoxelGridPipeline.usf
--- a/Engine/Shaders/Private/HeterogeneousVolumes/HeterogeneousVolumesVoxelGridPipeline.usf
+++ b/Engine/Shaders/Private/HeterogeneousVolumes/HeterogeneousVolumesVoxelGridPipeline.usf
@@ -209,7 +209,7 @@ void TopLevelGridCalculateVoxelSize(
 
 		// Store minimum voxel rate temporarily as the bottom-level index
 		uint LinearIndex = GetLinearIndex(VoxelIndex, TopLevelGridResolution);
-		if (!IsBottomLevelEmpty(RWTopLevelGridBuffer[LinearIndex]))
+		if (IsBottomLevelAllocated(RWTopLevelGridBuffer[LinearIndex]))
 		{
 			float PrevVoxelSize = GetVoxelSize(RWTopLevelGridBuffer[LinearIndex]);
 			VoxelSize = min(VoxelSize, PrevVoxelSize);
@@ -282,22 +282,17 @@ void AllocateBottomLevelGrid(
 		// Allocate indirection voxels
 		int3 BottomLevelVoxelResolution = MaxVoxelResolution;
 		int3 IndirectionGridResolution = ceil(VoxelResolution / BottomLevelVoxelResolution);
-		// Encoding restricts dimension to be a power-of-2
-		//IndirectionGridResolution = NextPowerOfTwo(IndirectionGridResolution);
-		IndirectionGridResolution.x = clamp(IndirectionGridResolution.x, 1, 4);
-		IndirectionGridResolution.y = clamp(IndirectionGridResolution.y, 1, 4);
-		IndirectionGridResolution.z = clamp(IndirectionGridResolution.z, 1, 4);
-		// Force regular sized voxels so the delineation is tractable for the allocator
+		IndirectionGridResolution = clamp(IndirectionGridResolution, 1, 4);
+		// Force regular sized dimensions so the permutation count is tractable for the allocator
 		IndirectionGridResolution = max(IndirectionGridResolution.x, max(IndirectionGridResolution.y, IndirectionGridResolution.z));
-
 		int IndirectionVoxelCount = IndirectionGridResolution.x * IndirectionGridResolution.y * IndirectionGridResolution.z;
 
 		uint IndirectionGridIndex = EMPTY_VOXEL_INDEX;
 		InterlockedAdd(RWIndirectionGridAllocatorBuffer[0], IndirectionVoxelCount, IndirectionGridIndex);
 
+		// If we over-allocate, eclare the allocator to be full and the voxel contents to be empty
 		if (IndirectionGridIndex + IndirectionVoxelCount > IndirectionGridBufferSize)
 		{
-			// Declare the buffer to be filled
 			uint Dummy;
 			InterlockedExchange(RWIndirectionGridAllocatorBuffer[0], IndirectionGridBufferSize, Dummy);
 
@@ -305,6 +300,7 @@ void AllocateBottomLevelGrid(
 			IndirectionGridResolution = 0;
 		}
 
+		// Write the voxel contents
 		FTopLevelGridData TopLevelGridData;
 		{
 			SetBottomLevelIndex(TopLevelGridData, IndirectionGridIndex);
@@ -315,12 +311,11 @@ void AllocateBottomLevelGrid(
 #else // DIM_ENABLE_INDIRECTION_GRID
 		uint BottomLevelIndex = EMPTY_VOXEL_INDEX;
 
-		VoxelResolution.x = clamp(VoxelResolution.x, 1, MaxVoxelResolution);
-		VoxelResolution.y = clamp(VoxelResolution.y, 1, MaxVoxelResolution);
-		VoxelResolution.z = clamp(VoxelResolution.z, 1, MaxVoxelResolution);
-		// Force regular sized voxels so the delineation is tractable for the allocator
+		VoxelResolution = clamp(VoxelResolution, 1, MaxVoxelResolution);
+		// Force regular sized dimensions so the permutation count is tractable for the allocator
 		VoxelResolution = max(VoxelResolution.x, max(VoxelResolution.y, VoxelResolution.z));
 
+		// Write the voxel contents
 		FTopLevelGridData TopLevelGridData = RWTopLevelGridBuffer[LinearIndex];
 		{
 			SetBottomLevelIndex(TopLevelGridData, BottomLevelIndex);
@@ -360,16 +355,16 @@ void GenerateRasterTiles(
 	if (!IsBottomLevelEmpty(TopLevelGridBuffer[LinearIndex]))
 	{
 		int3 BottomLevelResolution = GetBottomLevelVoxelResolution(TopLevelGridBuffer[LinearIndex]);
-		int BottomLevelLinearResolution = BottomLevelResolution.x * BottomLevelResolution.y * BottomLevelResolution.z;
 
 #if DIM_ENABLE_INDIRECTION_GRID
-		int RasterTileCount = BottomLevelLinearResolution;
-
+		int3 RasterTileCount3 = BottomLevelResolution;
 #else // DIM_ENABLE_INDIRECTION_GRID
+		// NOTE: Always 1, due to limitations of bottom-level grid resolution
 		int3 RasterTileCount3 = ceil(float3(BottomLevelResolution) / RasterTileVoxelResolution);
-		int RasterTileCount = RasterTileCount3.x * RasterTileCount3.y * RasterTileCount3.z;
 #endif // DIM_ENABLE_INDIRECTION_GRID
 
+		int RasterTileCount = RasterTileCount3.x * RasterTileCount3.y * RasterTileCount3.z;
+
 		if (RasterTileCount > 0)
 		{
 			uint RasterTileOffset;
@@ -648,7 +643,7 @@ void RenderTransmittanceWithVoxelGridCS(
 	}
 
 	// Debug
-	Radiance = Transmittance;
+	Radiance = 1.0 - Transmittance;
 
 	// Output..
 	float3 Opacity = 1.0 - Transmittance;
diff --git a/Engine/Source/Runtime/Renderer/Private/HeterogeneousVolumes/HeterogeneousVolumesVoxelGridPipeline.cpp b/Engine/Source/Runtime/Renderer/Private/HeterogeneousVolumes/HeterogeneousVolumesVoxelGridPipeline.cpp
--- a/Engine/Source/Runtime/Renderer/Private/HeterogeneousVolumes/HeterogeneousVolumesVoxelGridPipeline.cpp
+++ b/Engine/Source/Runtime/Renderer/Private/HeterogeneousVolumes/HeterogeneousVolumesVoxelGridPipeline.cpp
@@ -388,11 +388,9 @@ class FRasterizeBottomLevelFroxelGridCS : public FMeshMaterialShader
 		RDG_BUFFER_ACCESS(IndirectArgs, ERHIAccess::IndirectArgs)
 
 		// Grid data
-		//SHADER_PARAMETER_RDG_BUFFER_SRV(StructuredBuffer<FTopLevelGridData>, TopLevelGridBuffer)
+		SHADER_PARAMETER_RDG_BUFFER_UAV(RWStructuredBuffer<FTopLevelGridData>, RWTopLevelGridBuffer)
 
 		SHADER_PARAMETER_RDG_BUFFER_UAV(RWBuffer<uint>, RWBottomLevelGridAllocatorBuffer)
-		SHADER_PARAMETER_RDG_BUFFER_UAV(RWStructuredBuffer<uint>, RWTopLevelGridBuffer)
-
 		SHADER_PARAMETER_RDG_BUFFER_UAV(RWStructuredBuffer<FScalarGridData>, RWExtinctionGridBuffer)
 		SHADER_PARAMETER_RDG_BUFFER_UAV(RWStructuredBuffer<FVectorGridData>, RWEmissionGridBuffer)
 		SHADER_PARAMETER_RDG_BUFFER_UAV(RWStructuredBuffer<FVectorGridData>, RWScatteringGridBuffer)
@@ -686,7 +684,7 @@ class FRasterizeBottomLevelGrid : public FMeshMaterialShader
 		SHADER_PARAMETER_RDG_BUFFER_SRV(StructuredBuffer<FTopLevelGridData>, TopLevelGridBuffer)
 
 		SHADER_PARAMETER_RDG_BUFFER_UAV(RWBuffer<uint>, RWBottomLevelGridAllocatorBuffer)
-		SHADER_PARAMETER_RDG_BUFFER_UAV(RWStructuredBuffer<uint>, RWTopLevelGridBuffer)
+		SHADER_PARAMETER_RDG_BUFFER_UAV(RWStructuredBuffer<FTopLevelGridData>, RWTopLevelGridBuffer)
 		SHADER_PARAMETER_RDG_BUFFER_UAV(RWStructuredBuffer<FScalarGridData>, RWExtinctionGridBuffer)
 		SHADER_PARAMETER_RDG_BUFFER_UAV(RWStructuredBuffer<FVectorGridData>, RWEmissionGridBuffer)
 		SHADER_PARAMETER_RDG_BUFFER_UAV(RWStructuredBuffer<FVectorGridData>, RWScatteringGridBuffer)
@@ -1691,7 +1689,7 @@ void RasterizeVolumesIntoFrustumVoxelGrid(
 
 			// Grid data
 			PassParameters->RWBottomLevelGridAllocatorBuffer = GraphBuilder.CreateUAV(BottomLevelGridAllocatorBuffer, PF_R32_UINT);
-			PassParameters->RWTopLevelGridBuffer = GraphBuilder.CreateUAV(TopLevelGridBuffer, PF_R32_UINT);
+			PassParameters->RWTopLevelGridBuffer = GraphBuilder.CreateUAV(TopLevelGridBuffer);
 
 			PassParameters->RWExtinctionGridBuffer = GraphBuilder.CreateUAV(ExtinctionGridBuffer);
 			PassParameters->RWEmissionGridBuffer = GraphBuilder.CreateUAV(EmissionGridBuffer);
@@ -2139,7 +2137,7 @@ void RasterizeVolumesIntoOrthoVoxelGrid(
 				// Grid data
 				PassParameters->TopLevelGridBuffer = GraphBuilder.CreateSRV(TopLevelGridBuffer);
 				PassParameters->RWBottomLevelGridAllocatorBuffer = GraphBuilder.CreateUAV(BottomLevelGridAllocatorBuffer, PF_R32_UINT);
-				PassParameters->RWTopLevelGridBuffer = GraphBuilder.CreateUAV(TopLevelGridBuffer, PF_R32_UINT);
+				PassParameters->RWTopLevelGridBuffer = GraphBuilder.CreateUAV(TopLevelGridBuffer);
 				PassParameters->RWExtinctionGridBuffer = GraphBuilder.CreateUAV(ExtinctionGridBuffer);
 				PassParameters->RWEmissionGridBuffer = GraphBuilder.CreateUAV(EmissionGridBuffer);
 				PassParameters->RWScatteringGridBuffer = GraphBuilder.CreateUAV(ScatteringGridBuffer);
