Just ignore geometry cache texture coordinates after the 4th one as showing an errro prevents compilation of materials with more TextureCoordinate nodes.

It only really uses the first UV set from the mesh anyway.

diff --git a/Engine/Shaders/Private/GeometryCacheVertexFactory.ush b/Engine/Shaders/Private/GeometryCacheVertexFactory.ush
--- a/Engine/Shaders/Private/GeometryCacheVertexFactory.ush
+++ b/Engine/Shaders/Private/GeometryCacheVertexFactory.ush
@@ -64,10 +64,6 @@ struct FVertexFactoryInput
 	#if NUM_MATERIAL_TEXCOORDS_VERTEX > 3
 		float2	TexCoords3 : ATTRIBUTE8;
 	#endif
-
-	#if NUM_MATERIAL_TEXCOORDS_VERTEX > 4
-		#error Too many texture coordinate sets defined on Geometry cache vertex input. Max: 4.
-	#endif
 #endif
 
 	VF_GPUSCENE_DECLARE_INPUT_BLOCK(13)
