Fixed: Artifacts on masked material with virtual texture lightmap enabled on mobile

diff --git a/Engine/Shaders/Private/MobileBasePassPixelShader.usf b/Engine/Shaders/Private/MobileBasePassPixelShader.usf
--- a/Engine/Shaders/Private/MobileBasePassPixelShader.usf
+++ b/Engine/Shaders/Private/MobileBasePassPixelShader.usf
@@ -164,8 +164,9 @@ void ApplyPixelDepthOffsetForMobileBasePass(inout FMaterialPixelParameters Mater
     float PixelDepthOffset = ApplyPixelDepthOffsetToMaterialParameters(MaterialParameters, PixelMaterialInputs, OutDepth);
 }
 
-// Force early depth_stencil for materials that use VT feedback and don't modify the pixel depth
-#if (NUM_VIRTUALTEXTURE_SAMPLES || LIGHTMAP_VT_ENABLED) && !(OUTPUT_PIXEL_DEPTH_OFFSET)
+#define USES_PIXEL_DISCARD ((MATERIALBLENDING_MASKED || USE_DITHERED_LOD_TRANSITION) && !EARLY_Z_PASS_ONLY_MATERIAL_MASKING)
+// Force early depth_stencil for materials that use VT feedback and don't use pixel discard or modify the pixel depth
+#if (NUM_VIRTUALTEXTURE_SAMPLES || LIGHTMAP_VT_ENABLED) && !(OUTPUT_PIXEL_DEPTH_OFFSET || USES_PIXEL_DISCARD)
 	#define PIXELSHADER_EARLYDEPTHSTENCIL EARLYDEPTHSTENCIL	
 #else
 	#define PIXELSHADER_EARLYDEPTHSTENCIL 	
