Fix PrevLocalToWorld() for non-instanced primitives.
Broken in 21757684
The fix in 24111068 only applied to instanced things.

diff --git a/Engine/Shaders/Private/MaterialTemplate.ush b/Engine/Shaders/Private/MaterialTemplate.ush
--- a/Engine/Shaders/Private/MaterialTemplate.ush
+++ b/Engine/Shaders/Private/MaterialTemplate.ush
@@ -1304,7 +1304,11 @@ FLWCMatrix GetLocalToWorld(FMaterialPixelParameters Parameters)
 
 FLWCMatrix GetPrevLocalToWorld(FMaterialVertexParameters Parameters)
 {
+#if USE_INSTANCING || USE_INSTANCE_CULLING || IS_MESHPARTICLE_FACTORY
 	return Parameters.PrevFrameLocalToWorld;
+#else
+	return GetPrimitiveData(Parameters).PreviousLocalToWorld;
+#endif
 }
 
 /** Transforms a vector from local space to world space (PS version) */
