Add SLW 16/12bits tile coord support.

diff --git a/Engine/Shaders/Private/PackUnpack.ush b/Engine/Shaders/Private/PackUnpack.ush
--- a/Engine/Shaders/Private/PackUnpack.ush
+++ b/Engine/Shaders/Private/PackUnpack.ush
@@ -211,3 +211,17 @@ float4 UnpackR16G16B16A16F(uint2 In)
 	return float4(UnpackFloat2FromUInt(In.x), UnpackFloat2FromUInt(In.y));
 }
 
+
+//
+// Common tile coord encoding
+//
+
+#define PackTileCoordXXbits(In, InMask, InBits)  ((In.x & InMask) | ((In.y & InMask) << InBits))
+#define UnpackTileCoordXXbits(In,InMask, InBits) uint2(In & InMask, (In >> InBits) & InMask)
+
+uint  PackTileCoord10bits(uint2 In) { return PackTileCoordXXbits  (In, 0x3FF, 10u); }
+uint2 UnpackTileCoord10bits(uint In){ return UnpackTileCoordXXbits(In, 0x3FF, 10u); }
+uint  PackTileCoord12bits(uint2 In) { return PackTileCoordXXbits  (In, 0xFFF, 12u); }
+uint2 UnpackTileCoord12bits(uint In){ return UnpackTileCoordXXbits(In, 0xFFF, 12u); }
+uint  PackTileCoord16bits(uint2 In) { return PackTileCoordXXbits  (In, 0xFFFF,16u); }
+uint2 UnpackTileCoord16bits(uint In){ return UnpackTileCoordXXbits(In, 0xFFFF,16u); }
diff --git a/Engine/Shaders/Private/SingleLayerWaterComposite.usf b/Engine/Shaders/Private/SingleLayerWaterComposite.usf
--- a/Engine/Shaders/Private/SingleLayerWaterComposite.usf
+++ b/Engine/Shaders/Private/SingleLayerWaterComposite.usf
@@ -328,7 +328,7 @@ void WaterTileCatergorisationMarkCS(uint3 ThreadId : SV_DispatchThreadID, uint3
 			uint WriteToIndex;
 			InterlockedAdd(DrawIndirectDataUAV[1], 1, WriteToIndex);
 			InterlockedAdd(DispatchIndirectDataUAV[0], 1);
-			WaterTileListDataUAV[WriteToIndex] = GroupId.x | (GroupId.y << 16); // assumes 16bit is enought to represent a tiled resolution up to 65,535 :)
+			WaterTileListDataUAV[WriteToIndex] = PackTileCoord16bits(GroupId.xy);
 		}
 	}
 
@@ -509,8 +509,13 @@ void WaterTileVS(
 	in uint VertexId : SV_VertexID,
 	out float4 Position : SV_POSITION) 
 {
-	uint TileData = TileListData[InstanceId.x];
-	uint2 TileOrigin = uint2((TileData & 0x0000FFFF) * SLW_TILE_SIZE_XY, (TileData >> 16) * SLW_TILE_SIZE_XY);
+	const uint PackedTile = TileListData[InstanceId.x];
+#if PERMUTATION_TILE_ENCODING == 0 		// 16bits coord
+	const uint2 TileCoord = UnpackTileCoord16bits(PackedTile);
+#elif PERMUTATION_TILE_ENCODING == 1 	// 12bits coord
+	const uint2 TileCoord = UnpackTileCoord12bits(PackedTile);
+#endif
+	const uint2 TileOrigin = TileCoord * SLW_TILE_SIZE_XY;
 
 	uint2 TileVertex = TileOrigin;
 	TileVertex.x += VertexId == 1 || VertexId == 2 || VertexId == 4 ? SLW_TILE_SIZE_XY : 0;
diff --git a/Engine/Shaders/Private/VirtualShadowMaps/VirtualShadowMapProjection.usf b/Engine/Shaders/Private/VirtualShadowMaps/VirtualShadowMapProjection.usf
--- a/Engine/Shaders/Private/VirtualShadowMaps/VirtualShadowMapProjection.usf
+++ b/Engine/Shaders/Private/VirtualShadowMaps/VirtualShadowMapProjection.usf
@@ -560,12 +560,12 @@ void VirtualShadowMapProjection(
 	uint3	DispatchThreadId	: SV_DispatchThreadID )
 {
 #if USE_TILE_LIST
-	uint PackedTile = TileListData.Load(GroupId.x);
-	uint2 LocalPixelPos = WORK_TILE_SIZE * uint2(PackedTile & 0xFFFF, PackedTile >> 16u) + MortonDecode(GroupIndex);
+	const uint2 TileCoord = UnpackTileCoord16bits(TileListData.Load(GroupId.x)); // Need to match SingleLayerWater tile encoding
 #else
-	// Morton order within a group so page access/atomics are more coherent and wave-swizzled gradients are possible.
-	uint2 LocalPixelPos = WORK_TILE_SIZE * GroupId.xy + MortonDecode( GroupIndex );
+	const uint2 TileCoord = GroupId.xy;
 #endif // USE_TILE_LIST
+	// Morton order within a group so page access/atomics are more coherent and wave-swizzled gradients are possible.
+	uint2 LocalPixelPos = WORK_TILE_SIZE * TileCoord + MortonDecode(GroupIndex);
 	uint2 PixelPos = LocalPixelPos + uint2( ProjectionRect.xy );
 	if ( any( PixelPos >= uint2( ProjectionRect.zw ) ) )
 	{
diff --git a/Engine/Source/Runtime/Renderer/Private/ScreenSpaceRayTracing.cpp b/Engine/Source/Runtime/Renderer/Private/ScreenSpaceRayTracing.cpp
--- a/Engine/Source/Runtime/Renderer/Private/ScreenSpaceRayTracing.cpp
+++ b/Engine/Source/Runtime/Renderer/Private/ScreenSpaceRayTracing.cpp
@@ -524,6 +524,7 @@ class FScreenSpaceReflectionsTileVS : public FGlobalShader
 	{
 		OutEnvironment.SetDefine(TEXT("TILE_VERTEX_SHADER"), 1.0f);
 		OutEnvironment.SetDefine(TEXT("WORK_TILE_SIZE"), 8);
+		OutEnvironment.SetDefine(TEXT("PERMUTATION_TILE_ENCODING"), 0); // Use generic 16 bits tile coords
 		FGlobalShader::ModifyCompilationEnvironment(Parameters, OutEnvironment);
 	}
 };
diff --git a/Engine/Source/Runtime/Renderer/Private/SingleLayerWaterRendering.cpp b/Engine/Source/Runtime/Renderer/Private/SingleLayerWaterRendering.cpp
--- a/Engine/Source/Runtime/Renderer/Private/SingleLayerWaterRendering.cpp
+++ b/Engine/Source/Runtime/Renderer/Private/SingleLayerWaterRendering.cpp
@@ -1175,7 +1175,10 @@ void FDeferredShadingSceneRenderer::RenderSingleLayerWaterReflections(
 
 			if (bRunTiled)
 			{
+				const bool bUseTileCoord16bitsEncoding = true;
+
 				FWaterTileVS::FPermutationDomain VsPermutationVector;
+				VsPermutationVector.Set<FWaterTileVS::ETileEncoding>(bUseTileCoord16bitsEncoding ? 0 : 1);
 				TShaderMapRef<FWaterTileVS> VertexShader(View.ShaderMap, VsPermutationVector);
 				ValidateShaderParameters(VertexShader, PassParameters->VS);
 				ClearUnusedGraphResources(VertexShader, &PassParameters->VS);
diff --git a/Engine/Source/Runtime/Renderer/Private/SingleLayerWaterRendering.h b/Engine/Source/Runtime/Renderer/Private/SingleLayerWaterRendering.h
--- a/Engine/Source/Runtime/Renderer/Private/SingleLayerWaterRendering.h
+++ b/Engine/Source/Runtime/Renderer/Private/SingleLayerWaterRendering.h
@@ -48,7 +48,8 @@ class FWaterTileVS : public FGlobalShader
 	DECLARE_GLOBAL_SHADER(FWaterTileVS);
 	SHADER_USE_PARAMETER_STRUCT(FWaterTileVS, FGlobalShader);
 
-	using FPermutationDomain = TShaderPermutationDomain<>;
+	class ETileEncoding : SHADER_PERMUTATION_INT("PERMUTATION_TILE_ENCODING", 2); // 16 bits or 12 bits tile coord encoding
+	using FPermutationDomain = TShaderPermutationDomain<ETileEncoding>;
 
 	BEGIN_SHADER_PARAMETER_STRUCT(FParameters, )
 		SHADER_PARAMETER_STRUCT_REF(FViewUniformShaderParameters, ViewUniformBuffer)
