Substrate - Custom decals now works

diff --git a/Engine/Shaders/Private/MaterialTemplate.ush b/Engine/Shaders/Private/MaterialTemplate.ush
--- a/Engine/Shaders/Private/MaterialTemplate.ush
+++ b/Engine/Shaders/Private/MaterialTemplate.ush
@@ -2611,8 +2611,10 @@ float4 MaterialExpressionDBufferTextureLookup(FMaterialPixelParameters Parameter
 {
 	if ((GetPrimitiveData(Parameters).Flags & PRIMITIVE_SCENE_DATA_FLAG_DECAL_RECEIVER) != 0 && View.ShowDecalsMask > 0)
 	{
+		uint2 PixelPos = uint2(BufferUV * View.BufferSizeAndInvSize.xy);
+		uint ValidDBufferTargetMask = GetDBufferTargetMask(PixelPos) & (1u << DBufferTextureIndex);
 #if STRATA_ENABLED
-		const FStrataDBuffer DBufferData = GetDBufferData(BufferUV, DBufferTextureIndex);
+		const FStrataDBuffer DBufferData = GetDBufferData(BufferUV, ValidDBufferTargetMask);
 		switch (DBufferTextureIndex)
 		{
 		// All data (BaseColor, WorldNormal, Roughness) are pre-multiplied by their coverage
@@ -2621,8 +2623,6 @@ float4 MaterialExpressionDBufferTextureLookup(FMaterialPixelParameters Parameter
 		case 2:	return float4(DBufferData.Roughness,	DBufferData.Metallic, DBufferData.Specular, DBufferData.OneMinusCoverage_Roughness);
 		}
 #else
-		uint2 PixelPos = uint2(BufferUV * View.BufferSizeAndInvSize.xy);
-		uint ValidDBufferTargetMask = GetDBufferTargetMask(PixelPos) & (1u << DBufferTextureIndex);
 		FDBufferData DBufferData = GetDBufferData(BufferUV, ValidDBufferTargetMask);
 		switch (DBufferTextureIndex)
 		{
