Substrate - Fix for material looking weird in some cases due to from/to metalness conversion used for decals not being identity when no decals are visible. To have an identity transform:
- Computations are now done in Diffuse/F0 space (decal converted into that space)
- It is still not perfect but the "no decal present" condition now do not morph the materials
- Made F0ToMetallic started metallic transitions at F0=0.08 similarly to the legacy specular/metallic model. (to be discussed)

diff --git a/Engine/Shaders/Private/DBufferDecalShared.ush b/Engine/Shaders/Private/DBufferDecalShared.ush
--- a/Engine/Shaders/Private/DBufferDecalShared.ush
+++ b/Engine/Shaders/Private/DBufferDecalShared.ush
@@ -133,6 +133,7 @@ void ApplyDBufferData(in FStrataDBuffer StrataDBufferData, inout FStrataPixelHea
 
 	// For normal, apply decal data onto all basis, without consideration of top layers
 	#if STRATA_DBUFFER_RESPONSE_NORMAL && STRATA_INLINE_SHADING
+	if(any(StrataDBufferData.WorldNormal != 0.0) || StrataDBufferData.OneMinusCoverage_WorldNormal < 1.0)
 	{
 		UNROLL
 		for (uint i = 0; i < StrataPixelHeader.SharedLocalBases.Count; ++i)
@@ -154,6 +155,13 @@ void ApplyDBufferData(in FStrataDBuffer StrataDBufferData, inout FStrataPixelHea
 	}
 	#endif
 
+	const bool bAffectDiffuse = any(StrataDBufferData.BaseColor > 0) || StrataDBufferData.OneMinusCoverage_BaseColor < 1.0;
+	const bool bAffectF0Rough = any(StrataDBufferData.Metallic > 0) || any(StrataDBufferData.Specular > 0) || any(StrataDBufferData.Roughness > 0) || StrataDBufferData.OneMinusCoverage_Roughness < 1.0;
+	if (!bAffectDiffuse && !bAffectF0Rough)
+	{
+		return;
+	}
+
 	// For now apply decal data onto all layers, without consideration of top layers
 	STRATA_UNROLL_N(STRATA_CLAMPED_BSDF_COUNT)
 	for (int BSDFIdx = 0; BSDFIdx < StrataPixelHeader.StrataTree.BSDFCount; ++BSDFIdx)
@@ -166,54 +174,88 @@ void ApplyDBufferData(in FStrataDBuffer StrataDBufferData, inout FStrataPixelHea
 			{	
 				// Note:
 				// * For Slab, since the DBuffer is stored with a Metalness parameterization for storage-space reason, 
-				//   we convert back&forth the Slab data between DiffuseAlbedo/F0 and Metalness. This induces some shift 
-				//   since the conversion is lossy. In the future, we could change the DBuffer storage layout (would require 
-				//   more space) or do all the blending in DiffuseAlbedo/F0 space
+				//   we were convert back&forth the Slab data between DiffuseAlbedo/F0 and Metalness. 
+				// * However, this is a non identy operation even when non decal so material look could be changed.
+				//   Furthermore this resulted in some transition edges being visible where the decal was being applied
+				//=> We now do computation in Diffuse/F0 space. It is still not perfectly matching legacy 
+				//   but at least the null operation is identity and no more transition edges are visible.
 					
 				// Apply DBuffer as per-component parameter blending (legacy mode)
 				#if STRATA_DBUFFER_BLENDING_MODE == STRATA_DBUFFER_LEGACY
 				{
-					float3 BaseColor = 0;
-					float Metallic = 0;
-					float Specular = 0;
-					ConvertToMetalness(SLAB_DIFFUSEALBEDO(BSDF), SLAB_F0(BSDF), BaseColor, Specular, Metallic);
-					float Roughness = SLAB_ROUGHNESS(BSDF);
+					// The original metallic, before the decal touched the slab.
+					// We use MIN on F0 to nor remove too much energy when metallic is lowered and have smoother metallic gradient with low F0.
+					float SlabMetallic = F0ToMetallic(min(SLAB_F0(BSDF).r, min(SLAB_F0(BSDF).g, SLAB_F0(BSDF).b)));
 
-				#if STRATA_DBUFFER_RESPONSE_BASECOLOR
-					BaseColor = BaseColor * StrataDBufferData.OneMinusCoverage_BaseColor + StrataDBufferData.BaseColor;
+				#if STRATA_DBUFFER_RESPONSE_BASECOLOR || STRATA_DBUFFER_RESPONSE_ROUGHNESS
+					float3 DecalDiffuse;
+					float3 DecalF0;
+					ConvertFromMetalness(StrataDBufferData.BaseColor, StrataDBufferData.Specular, StrataDBufferData.Metallic, DecalDiffuse, DecalF0);
 
-					// Fade out fuzz to avoid overbright decal.
-					SLAB_FUZZ_AMOUNT(BSDF) *= StrataDBufferData.OneMinusCoverage_BaseColor;
+					if (bAffectDiffuse || bAffectF0Rough)
+					{
+						// Affect F0 if basecolor or metallic are changed.
+						SLAB_F0(BSDF) = SLAB_F0(BSDF) * StrataDBufferData.OneMinusCoverage_BaseColor + DecalF0;
+					}
+				#endif
 
-					// Disable SSS when the coverage reach one to preserve the decal color. 
-					// That operation also progressively makes non top layer disapear.
-					// Since this is ran before the strata tree is processed, we must use TmpMFP and not SLAB_SSSMFP.
-					BSDF.TmpMFP = BSDF.TmpMFP * StrataDBufferData.OneMinusCoverage_BaseColor;
-					if (StrataDBufferData.OneMinusCoverage_BaseColor == 0.0f)
+				#if STRATA_DBUFFER_RESPONSE_BASECOLOR
+					if (bAffectDiffuse)
 					{
-						BSDF_SETSSSTYPE(BSDF, SSS_TYPE_INVALID);
+						// Diffuse albedo is affected by throughpu and decal diffuse contribution
+						SLAB_DIFFUSEALBEDO(BSDF) = SLAB_DIFFUSEALBEDO(BSDF) * StrataDBufferData.OneMinusCoverage_BaseColor	+ DecalDiffuse;
+
+						// Fade out fuzz to avoid overbright decal.
+						SLAB_FUZZ_AMOUNT(BSDF) *= StrataDBufferData.OneMinusCoverage_BaseColor;
+
+						// Disable SSS when the coverage reach one to preserve the decal color. 
+						// That operation also progressively makes non top layer disapear.
+						// Since this is ran before the strata tree is processed, we must use TmpMFP and not SLAB_SSSMFP.
+						BSDF.TmpMFP = BSDF.TmpMFP * StrataDBufferData.OneMinusCoverage_BaseColor;
+						if (StrataDBufferData.OneMinusCoverage_BaseColor == 0.0f)
+						{
+							BSDF_SETSSSTYPE(BSDF, SSS_TYPE_INVALID);
+						}
 					}
 				#endif
 
 				#if STRATA_DBUFFER_RESPONSE_ROUGHNESS
-					Metallic  = Metallic  * StrataDBufferData.OneMinusCoverage_Roughness + StrataDBufferData.Metallic;
-					Specular  = Specular  * StrataDBufferData.OneMinusCoverage_Roughness + StrataDBufferData.Specular;
-					Roughness = Roughness * StrataDBufferData.OneMinusCoverage_Roughness + StrataDBufferData.Roughness;
-
-					// Also fade out other specular features since we consider decal to have not specular features.
-					if (BSDF_GETHASHAZINESS(BSDF))
+					if (bAffectF0Rough)
 					{
-						FHaziness Haziness = UnpackHaziness(SLAB_HAZINESS(BSDF));
-						Haziness.Weight *= StrataDBufferData.OneMinusCoverage_Roughness;
-						SLAB_HAZINESS(BSDF) = PackHaziness(Haziness);
+						// Now also operate a change to reflect the change of metallic
+						float NewMetallic = SlabMetallic * StrataDBufferData.OneMinusCoverage_Roughness + StrataDBufferData.Metallic;
+						
+						if (NewMetallic >= SlabMetallic)
+						{
+							float TransferDelta		 = (NewMetallic - SlabMetallic);	// Positive or zero
+							SLAB_F0(BSDF)			+= TransferDelta * SLAB_DIFFUSEALBEDO(BSDF);
+							SLAB_DIFFUSEALBEDO(BSDF)-= TransferDelta * SLAB_DIFFUSEALBEDO(BSDF);
+						}
+						else
+						{
+							float TransferDelta =-(NewMetallic - SlabMetallic);			// Positive
+
+							float CurrentSpecular		= F0ToDielectricSpecular(min(SLAB_F0(BSDF).r, min(SLAB_F0(BSDF).g, SLAB_F0(BSDF).b)));
+							float F0ThatCannotTransfer	= DielectricSpecularToF0(CurrentSpecular);
+							float3 F0ThatCanTransfer	= saturate(SLAB_F0(BSDF) - F0ThatCannotTransfer);
+
+							SLAB_DIFFUSEALBEDO(BSDF)+= TransferDelta * F0ThatCanTransfer;
+							SLAB_F0(BSDF)			-= TransferDelta * F0ThatCanTransfer;
+						}
+					
+						SLAB_ROUGHNESS(BSDF)	= SLAB_ROUGHNESS(BSDF) * StrataDBufferData.OneMinusCoverage_Roughness + StrataDBufferData.Roughness;
+					
+						// Also fade out other specular features since we consider decal to have not specular features.
+						if (BSDF_GETHASHAZINESS(BSDF))
+						{
+							FHaziness Haziness = UnpackHaziness(SLAB_HAZINESS(BSDF));
+							Haziness.Weight *= StrataDBufferData.OneMinusCoverage_Roughness;
+							SLAB_HAZINESS(BSDF) = PackHaziness(Haziness);
+						}
+					
+						SLAB_ANISOTROPY(BSDF) *= StrataDBufferData.OneMinusCoverage_Roughness;
 					}
-					SLAB_ANISOTROPY(BSDF) *= StrataDBufferData.OneMinusCoverage_Roughness;
 				#endif
-
-					float3 DiffuseAlbedo = 0;
-					float3 F0 = 0;
-					ConvertFromMetalness(BaseColor, Specular, Metallic, SLAB_DIFFUSEALBEDO(BSDF), SLAB_F0(BSDF));
-					SLAB_ROUGHNESS(BSDF) = Roughness;
 				}
 				// Apply DBuffer as a coating layer (using parameter blending for efficiency purpose)
 				#elif STRATA_DBUFFER_BLENDING_MODE == STRATA_DBUFFER_COATING
diff --git a/Engine/Shaders/Private/ShadingCommon.ush b/Engine/Shaders/Private/ShadingCommon.ush
--- a/Engine/Shaders/Private/ShadingCommon.ush
+++ b/Engine/Shaders/Private/ShadingCommon.ush
@@ -166,11 +166,11 @@ float MakeRoughnessSafe(float Roughness, float MinRoughness=0.001f)
 float F0ToMetallic(float F0)
 {
 	// Approximate the metallic input from F0 with a small lerp region
-	const float DiamondF0 = 0.24f;
-	const float CondutorMinF0 = 0.4f;	// roughly the end of semi-conductor
+	const float FullMetalBeginF0 = 0.08f;	// Instead of DiamondF0 = 0.24, the metallic region starts right after metallic >0 and specular=1 to match with legacy.
+	const float FullMetalEndF0   = 0.4f;	// roughly the end of semi-conductor
 	// This is compatible with UE shading model mapping allowing F0 to take a value up to 0.08 for dielectric.
 
-	return saturate((F0 - DiamondF0) / (CondutorMinF0 - DiamondF0));
+	return saturate((F0 - FullMetalBeginF0) / (FullMetalEndF0 - FullMetalBeginF0));
 }
 
 float F0RGBToMetallic(float3 F0)
