Fix dxc errors

diff --git a/Engine/Shaders/Private/WorkDistribution.ush b/Engine/Shaders/Private/WorkDistribution.ush
--- a/Engine/Shaders/Private/WorkDistribution.ush
+++ b/Engine/Shaders/Private/WorkDistribution.ush
@@ -352,8 +352,8 @@ struct FInputQueue
 
 struct FGlobalWorkQueue
 {
-	RWCoherentByteAddressBuffer						DataBuffer;
-	RWCoherentStructuredBuffer( FWorkQueueState )	StateBuffer;	// Ideally this was GDS but we don't have that API control.
+	RWByteAddressBuffer						DataBuffer;
+	RWStructuredBuffer< FWorkQueueState >	StateBuffer;	// Ideally this was GDS but we don't have that API control.
 
 	uint	StateIndex;
 	uint	Size;
