Path Tracing: Fix shader errors when rebuilding shaders with GPU Lightmass

It sometimes happens that MATERIAL_IS_STRATA becomes false, which causes the contents of StrataExport.ush to become invalid. Simply skip this include in this case for now.

diff --git a/Engine/Shaders/Private/PathTracing/PathTracingGPULightmassMaterialHitShader.usf b/Engine/Shaders/Private/PathTracing/PathTracingGPULightmassMaterialHitShader.usf
--- a/Engine/Shaders/Private/PathTracing/PathTracingGPULightmassMaterialHitShader.usf
+++ b/Engine/Shaders/Private/PathTracing/PathTracingGPULightmassMaterialHitShader.usf
@@ -35,9 +35,13 @@
 #if STRATA_ENABLED
 #define STRATA_GPU_LIGHTMASS 1
 #include "/Engine/Private/Strata/Strata.ush"
+
+#if MATERIAL_IS_STRATA // STRATA_TODO: Should this header be guarded from inside instead?
 #include "/Engine/Private/Strata/StrataExport.ush"
 #endif
 
+#endif
+
 RAY_TRACING_ENTRY_CLOSEST_HIT(GPULightmassMaterialCHS,
 	FPackedPathTracingPayload, PackedPayload,
 	FRayTracingIntersectionAttributes, Attributes)
