Path Tracing: Implement refraction for Substrate materials

When a material is flagged to used refraction, compute the IOR and set it on the slab (for the bottom slabs only). For now this case is split into its own logic with

A different implementation is used for Thin vs Solid cases for now.

This implementation is still a work in progress. In particular, the treatment of the diffuse color (which controls the medium albedo) is not complete.

diff --git a/Engine/Shaders/Private/PathTracing/Material/PathTracingFresnel.ush b/Engine/Shaders/Private/PathTracing/Material/PathTracingFresnel.ush
--- a/Engine/Shaders/Private/PathTracing/Material/PathTracingFresnel.ush
+++ b/Engine/Shaders/Private/PathTracing/Material/PathTracingFresnel.ush
@@ -170,3 +170,9 @@ FThinSlabWeights ComputeThinSlabWeights(float3 SlabColor, float CosV, float Eta,
 	Result.Transmitted = SlabT;
 	return Result;
 }
+
+float ComputeThinTransmittedRoughness(float Roughness, float Ior)
+{
+	// Transmitted lobe should have slightly less roughness due to the squeeze implied by the IOR (eventually converging to 0.0 as Ior approaches 1.0)
+	return Ior > 0.0 ? Roughness * sqrt(1.0 - min(Ior, rcp(Ior))) : Roughness;
+}
diff --git a/Engine/Shaders/Private/PathTracing/Material/PathTracingMaterialSampling.ush b/Engine/Shaders/Private/PathTracing/Material/PathTracingMaterialSampling.ush
--- a/Engine/Shaders/Private/PathTracing/Material/PathTracingMaterialSampling.ush
+++ b/Engine/Shaders/Private/PathTracing/Material/PathTracingMaterialSampling.ush
@@ -13,6 +13,8 @@ PathTracingMaterialSampling.usf: Material sampling functions
 #include "PathTracingGlossy.ush"
 #if STRATA_ENABLED
 #include "PathTracingSubstrate.ush"
+#include "PathTracingSubstrateSolidGlass.ush"
+#include "PathTracingSubstrateThinGlass.ush"
 #else
 #include "PathTracingSolidGlass.ush"
 #include "PathTracingThinGlass.ush"
@@ -48,7 +50,16 @@ float3 EstimateMaterialAlbedo(FPathTracingPayload Payload)
 #elif PATHTRACING_STRATA_PAYLOAD
 		case SHADINGMODELID_STRATA:
 		{
-			Albedo = Payload.DiffuseColor + Payload.SubsurfaceColor + Payload.SpecularColor + Payload.FuzzAmount * Payload.FuzzColor;
+			if (Payload.HasRefraction())
+				Albedo = 1.0 + Payload.FuzzAmount * Payload.FuzzColor;
+			else
+				Albedo = Payload.DiffuseColor + Payload.SubsurfaceColor + Payload.SpecularColor + Payload.FuzzAmount * Payload.FuzzColor;
+			Albedo *= Payload.GetMaxLobeWeight();
+			break;
+		}
+		case SHADINGMODELID_THIN_TRANSLUCENT:
+		{
+			Albedo = 1.0 + Payload.FuzzAmount * Payload.FuzzColor;
 			Albedo *= Payload.GetMaxLobeWeight();
 			break;
 		}
@@ -140,12 +151,32 @@ void AccumulateAlbedoNormal(inout FPathTracingPayload Payload, float3 PathThroug
 			const float Roughness0 = Payload.RoughnessData.x;
 			const float Roughness1 = Payload.RoughnessData.y;
 			const float SpecBlend = Payload.RoughnessData.z;
-			const float Roughness0Weight = (1 - SpecBlend) * ROUGHNESS_WEIGHT(Roughness0);
-			const float Roughness1Weight = (    SpecBlend) * ROUGHNESS_WEIGHT(Roughness1);
-			Albedo += PathThroughput * (Payload.DiffuseColor + Payload.SpecularColor * (Roughness0Weight + Roughness1Weight) + Payload.FuzzAmount * Payload.FuzzColor);
+			const float RoughnessWeight = lerp(ROUGHNESS_WEIGHT(Roughness0), ROUGHNESS_WEIGHT(Roughness1), SpecBlend);
+			if (Payload.HasRefraction())
+			{
+				float GlassAlbedo = 1.0;
+				Albedo += PathThroughput * (GlassAlbedo * RoughnessWeight + Payload.FuzzAmount * Payload.FuzzColor);
+			}
+			else
+			{
+				Albedo += PathThroughput * (Payload.DiffuseColor + Payload.SubsurfaceColor + Payload.SpecularColor * RoughnessWeight + Payload.FuzzAmount * Payload.FuzzColor);
+			}
+			
 			Normal += NormThroughput * Payload.WorldNormal;
 			break;
 		}
+		case SHADINGMODELID_THIN_TRANSLUCENT:
+		{
+			const float F0 = F0RGBToF0(Payload.SpecularColor);
+			const float3 GlassAlbedo = Payload.HasRefraction() ? lerp(Payload.GetTransmittanceColor(), 1.0, F0) : F0;
+			const float Roughness0 = Payload.RoughnessData.x;
+			const float Roughness1 = Payload.RoughnessData.y;
+			const float SpecBlend = Payload.RoughnessData.z;
+			const float RoughnessWeight = lerp(ROUGHNESS_WEIGHT(Roughness0), ROUGHNESS_WEIGHT(Roughness1), SpecBlend);
+			Albedo += PathThroughput * (RoughnessWeight * GlassAlbedo + Payload.FuzzAmount * Payload.FuzzColor);
+			Normal += NormThroughput * RoughnessWeight * Payload.WorldNormal;
+
+		}
 		case SHADINGMODELID_CLEAR_COAT:
 		{
 			const float Roughness0 = Payload.RoughnessData.x;
@@ -293,10 +324,13 @@ float GetAverageRoughness(FPathTracingPayload Payload)
 	switch (Payload.ShadingModelID)
 	{
 		case SHADINGMODELID_STRATA:
+		case SHADINGMODELID_THIN_TRANSLUCENT:
 		{
-			// STRATA_TODO: revisit this when we support glass
+
 			AvgRoughness = max(AvgRoughness, Payload.RoughnessData.y);
-			float DiffuseFraction = LobeSelectionProb(Payload.SubsurfaceColor + Payload.DiffuseColor + Payload.FuzzAmount * Payload.FuzzColor, Payload.SpecularColor);
+			// Only non-glass slabs contain diffuse
+			float3 DiffuseColor = Payload.HasRefraction() ? 0.0 : Payload.SubsurfaceColor + Payload.DiffuseColor;
+			float DiffuseFraction = LobeSelectionProb(DiffuseColor + Payload.FuzzAmount * Payload.FuzzColor, Payload.SpecularColor);
 			AvgRoughness = lerp(1.0, AvgRoughness, DiffuseFraction);
 			break;
 		}
@@ -352,9 +386,12 @@ void AdjustPayloadAfterUnpack(inout FPathTracingPayload HitPayload, bool bApplyO
 	{
 		case SHADINGMODELID_STRATA:
 		{
-			// NOTE: if radius is 0 the sss color will get flipped back into diffuse
-			HitPayload.SubsurfaceColor = HitPayload.DiffuseColor;
-			HitPayload.DiffuseColor = 0;
+			if (!HitPayload.HasRefraction())
+			{
+				// NOTE: if radius is 0 the sss color will get flipped back into diffuse
+				HitPayload.SubsurfaceColor = HitPayload.DiffuseColor;
+				HitPayload.DiffuseColor = 0;
+			}
 			break;
 		}
 		case SHADINGMODELID_CLEAR_COAT:
@@ -557,7 +594,12 @@ FMaterialSample SampleMaterial(
 			return NullMaterialSample();
 #if PATHTRACING_STRATA_PAYLOAD
 		case SHADINGMODELID_STRATA:
-			return Substrate_SampleMaterial(V_World, Payload, RandSample);
+			if (Payload.HasRefraction())
+				return SubstrateSolidGlass_SampleMaterial(V_World, Payload, RandSample);
+			else
+				return Substrate_SampleMaterial(V_World, Payload, RandSample);
+		case SHADINGMODELID_THIN_TRANSLUCENT:
+			return SubstrateThinGlass_SampleMaterial(V_World, Payload, RandSample);
 		case SHADINGMODELID_CLEAR_COAT:
 			return SubstrateClearCoat_SampleMaterial(V_World, Payload, RandSample);
 #else
@@ -613,7 +655,12 @@ FMaterialEval EvalMaterial(
 			break;
 #if PATHTRACING_STRATA_PAYLOAD
 		case SHADINGMODELID_STRATA:
-			return Substrate_EvalMaterial(V_World, L_World, Payload, DiffuseSpecularScale);
+			if (Payload.HasRefraction())
+				return SubstrateSolidGlass_EvalMaterial(V_World, L_World, Payload, DiffuseSpecularScale);
+			else
+				return Substrate_EvalMaterial(V_World, L_World, Payload, DiffuseSpecularScale);
+		case SHADINGMODELID_THIN_TRANSLUCENT:
+			return SubstrateThinGlass_EvalMaterial(V_World, L_World, Payload, DiffuseSpecularScale);
 		case SHADINGMODELID_CLEAR_COAT:
 			return SubstrateClearCoat_EvalMaterial(V_World, L_World, Payload, DiffuseSpecularScale);
 #else
diff --git a/Engine/Shaders/Private/PathTracing/Material/PathTracingSubstrate.ush b/Engine/Shaders/Private/PathTracing/Material/PathTracingSubstrate.ush
--- a/Engine/Shaders/Private/PathTracing/Material/PathTracingSubstrate.ush
+++ b/Engine/Shaders/Private/PathTracing/Material/PathTracingSubstrate.ush
@@ -6,21 +6,12 @@
 #error "This header should only be included when Strata is enabled."
 #endif
 
+#include "PathTracingSubstrateCommon.ush"
 #include "PathTracingMaterialCommon.ush"
 #include "PathTracingGlossy.ush"
 #include "PathTracingFresnel.ush"
 #include "../../BRDF.ush"
 
-#define STRATA_FRESNEL_F82  0
-
-float3 SubstrateLobeWeight(FPathTracingPayload Payload, float SatNoL)
-{
-	// See reference in LuminanceWeight
-	const float DistanceL = rcp(max(SatNoL, 0.001f));
-	const float3 TransmittanceAboveAlongL = select(Payload.TransmittanceN < 1.f, pow(max(Payload.TransmittanceN, 0.0001f), DistanceL), Payload.TransmittanceN);
-	return Payload.BSDFOpacity * Payload.WeightV * lerp(1.0f, TransmittanceAboveAlongL, Payload.CoverageAboveAlongN);
-}
-
 struct FSubstrateSingleSlabData
 {
 	float3x3 Basis;
diff --git a/Engine/Shaders/Private/PathTracing/Material/PathTracingSubstrateCommon.ush b/Engine/Shaders/Private/PathTracing/Material/PathTracingSubstrateCommon.ush
new file mode 100644
--- /dev/null
+++ b/Engine/Shaders/Private/PathTracing/Material/PathTracingSubstrateCommon.ush
@@ -0,0 +1,19 @@
+// Copyright Epic Games, Inc. All Rights Reserved.
+
+#pragma once
+
+#ifndef STRATA_ENABLED
+#error "This header should only be included when Strata is enabled."
+#endif
+
+#include "../PathTracingCommon.ush"
+
+#define STRATA_FRESNEL_F82  0
+
+float3 SubstrateLobeWeight(FPathTracingPayload Payload, float SatNoL)
+{
+	// See reference in LuminanceWeight
+	const float DistanceL = rcp(max(SatNoL, 0.001f));
+	const float3 TransmittanceAboveAlongL = select(Payload.TransmittanceN < 1.f, pow(max(Payload.TransmittanceN, 0.0001f), DistanceL), Payload.TransmittanceN);
+	return Payload.BSDFOpacity * Payload.WeightV * lerp(1.0f, TransmittanceAboveAlongL, Payload.CoverageAboveAlongN);
+}
diff --git a/Engine/Shaders/Private/PathTracing/Material/PathTracingSubstrateSolidGlass.ush b/Engine/Shaders/Private/PathTracing/Material/PathTracingSubstrateSolidGlass.ush
new file mode 100644
--- /dev/null
+++ b/Engine/Shaders/Private/PathTracing/Material/PathTracingSubstrateSolidGlass.ush
@@ -0,0 +1,248 @@
+// Copyright Epic Games, Inc. All Rights Reserved.
+
+#pragma once
+
+#ifndef STRATA_ENABLED
+#error "This header should only be included when Strata is enabled."
+#endif
+
+#include "PathTracingSubstrateCommon.ush"
+#include "PathTracingMaterialCommon.ush"
+#include "PathTracingGlossy.ush"
+#include "PathTracingFresnel.ush"
+#include "../../BRDF.ush"
+
+
+struct FSubstrateSolidGlassSlabData
+{
+	float3x3 Basis;
+	float3 V;
+
+	float  F0;
+	float  Eta;
+	float2 Alpha0;
+	float2 Alpha1;
+	float  FuzzRoughness;
+
+	float  FuzzAttenuation;
+	float  EGlass0;
+	float  EGlass1;
+	FBxDFEnergyTermsA Fuzz;
+
+	float2 LobeCdf;
+	float3 LobePdf;
+
+	void PrepareBSDF(float3 V_World, FPathTracingPayload Payload)
+	{
+		Basis = GetGGXBasis(Payload.RoughnessData.x, Payload.Anisotropy, Payload.WorldNormal, Payload.WorldTangent, Alpha0);
+		Alpha1 = GetGGXAlpha(Payload.RoughnessData.y, Payload.Anisotropy);
+
+		V = mul(Basis, V_World);
+
+		// #dxr_todo: Maintain a refraction stack on the path tracing payload
+		F0 = F0RGBToF0(Payload.SpecularColor);
+		// NOTE: IsFrontFace() determines polygon orientation, because the normal is always flipped towards in the incoming ray
+		Eta = Payload.IsFrontFace() ? Payload.Ior : rcp(Payload.Ior);
+		
+
+		FuzzRoughness = MakeRoughnessSafe(Payload.FuzzRoughness, STRATA_MIN_FUZZ_ROUGHNESS);
+
+		const float NoV = saturate(V.z);
+		// correct for energy loss by scaling the whole BSDF
+		EGlass0 = GGXEnergyLookup(Payload.RoughnessData.x, NoV, Eta);
+		EGlass1 = GGXEnergyLookup(Payload.RoughnessData.y, NoV, Eta);
+		Fuzz = ComputeClothEnergyTermsA(FuzzRoughness, NoV);
+
+		FuzzAttenuation = 1.0 - Payload.FuzzAmount * Fuzz.E;
+
+		const float3 Glass0Albedo = FuzzAttenuation * (1.0 - Payload.RoughnessData.z);
+		const float3 Glass1Albedo = FuzzAttenuation * (      Payload.RoughnessData.z);
+		const float  FuzzAlbedo = Payload.FuzzAmount * Fuzz.E;
+
+		// Now prepare a cdf/pdf for lobe selection
+		float3 MaxLobeWeight = Payload.GetMaxLobeWeight();
+		LobeCdf = LobeSelectionCdf(
+			MaxLobeWeight * Glass0Albedo,
+			MaxLobeWeight * Glass1Albedo,
+			MaxLobeWeight * FuzzAlbedo);
+		LobePdf = LobeSelectionPdf(LobeCdf);
+	}
+};
+
+
+FMaterialSample SubstrateSolidGlass_SampleMaterial(
+	float3 V_World,
+	FPathTracingPayload Payload,
+	float3 RandSample
+)
+{
+	FSubstrateSolidGlassSlabData Data = (FSubstrateSolidGlassSlabData)0;
+	Data.PrepareBSDF(V_World, Payload);
+
+	float3 L = 0, H = 0, V = Data.V;
+	float OutRoughness = 1;
+
+	const bool bSampledSpecular = RandSample.x < Data.LobeCdf.y;
+	if (bSampledSpecular)
+	{
+		// specular lobes
+		const bool bUseSpec0 = RandSample.x < Data.LobeCdf.x;
+		if (bUseSpec0)
+		{
+			RandSample.x = RescaleRandomNumber(RandSample.x, 0.0, Data.LobeCdf.x);
+			OutRoughness = Payload.RoughnessData.x;
+		}
+		else
+		{
+			RandSample.x = RescaleRandomNumber(RandSample.x, Data.LobeCdf.x, Data.LobeCdf.y);
+			OutRoughness = Payload.RoughnessData.y;
+		}
+
+		// #dxr_todo: Evaluate UniformSampleDiskConcentric
+		H = ImportanceSampleVisibleGGX_aniso(UniformSampleDisk(RandSample.xy), bUseSpec0 ? Data.Alpha0 : Data.Alpha1, V).xyz;
+
+		// Glass lobe (reflection and refraction)
+		float3 L = 0;
+		float F = 0;
+
+		const bool bRefract = SampleRefraction(-V, H, Data.Eta, Data.F0, RandSample.z, L, F);
+		// transform to world space
+		const float3 L_World = normalize(mul(L, Data.Basis));
+
+		FMaterialSample Result = CreateMaterialSample(L_World, 0.0, 0.0, bRefract ? -1.0 : 1.0, OutRoughness, bRefract ? PATHTRACER_SCATTER_REFRACT : PATHTRACER_SCATTER_SPECULAR);
+
+		const float2 GGXResult0 = bRefract ? GGXEvalRefraction(L, V, H, Data.Alpha0, Data.Eta) : GGXEvalReflection(L, V, H, Data.Alpha0);
+		const float2 GGXResult1 = bRefract ? GGXEvalRefraction(L, V, H, Data.Alpha1, Data.Eta) : GGXEvalReflection(L, V, H, Data.Alpha1);
+
+		Result.AddLobeWithMIS((1.0 - Payload.RoughnessData.z) * GGXResult0.x / Data.EGlass0, F * GGXResult0.y, Data.LobePdf.x);
+		Result.AddLobeWithMIS((      Payload.RoughnessData.z) * GGXResult1.x / Data.EGlass1, F * GGXResult1.y, Data.LobePdf.y);
+
+		if (bRefract)
+		{
+			// only need to account for overall BSDF scaling
+			Result.Weight *= Data.FuzzAttenuation * Payload.BSDFOpacity * Payload.WeightV;
+		}
+		else
+		{
+			// reflection side, apply full lobe weight and attenuation along L
+			// and account for probability of picking the fuzz layer
+			const float ClothPdf = 1 / (2 * PI);
+			Result.Pdf += Data.LobePdf.z * ClothPdf;
+			Result.Weight *= Data.FuzzAttenuation * SubstrateLobeWeight(Payload, saturate(L.z));
+		}
+
+		return Result;
+	}
+	else
+	{
+		// cloth lobe
+		RandSample.x = RescaleRandomNumber(RandSample.x, Data.LobeCdf.y, 1.0);
+		L = UniformSampleHemisphere(RandSample.xy).xyz;
+		H = normalize(L + V);
+
+		// transform to world space
+		const float3 L_World = normalize(mul(L, Data.Basis));
+		const float3 N_World = Payload.WorldNormal;
+		const float NoV = saturate(V.z);
+		const float NoL = saturate(L.z);
+		const float VoH = saturate(dot(V, H));
+		const float NoH = saturate(H.z);
+
+
+		const float2 GGXResult0 = GGXEvalReflection(L, V, H, Data.Alpha0);
+		const float2 GGXResult1 = GGXEvalReflection(L, V, H, Data.Alpha1);
+		const float Fg = FresnelReflectance(VoH, Data.Eta, Data.F0);
+		const float Spec0Pdf = Fg * GGXResult0.y;
+		const float Spec1Pdf = Fg * GGXResult1.y;
+		const float ClothPdf = 1 / (2 * PI);
+
+		FMaterialSample Result = CreateMaterialSample(L_World, 0.0, 0.0, 1.0, 1.0, PATHTRACER_SCATTER_SPECULAR);
+
+		// Cloth Lobe
+		const float ShadowTerminator = ShadowTerminatorTerm(L_World, N_World, Payload.WorldSmoothNormal);
+		const float DCloth = D_Charlie(Data.FuzzRoughness, NoH);
+		const float3 FCloth = F_Schlick(Payload.FuzzColor, VoH);
+		const float3 ClothWeight = NoL > 0 && NoV > 0 ? (NoL * Payload.FuzzAmount * DCloth * Vis_Ashikhmin(NoV, NoL)) * FCloth * Data.Fuzz.W * (2 * PI) : 0.0;
+		Result.AddLobeWithMIS(ClothWeight * ShadowTerminator, 1 / (2.0 * PI), Data.LobePdf.z);
+		Result.ScatterType = PATHTRACER_SCATTER_SPECULAR;
+
+		Result.Pdf += Data.LobePdf.y * Spec0Pdf;
+		Result.Pdf += Data.LobePdf.z * Spec1Pdf;
+
+		Result.Weight *= SubstrateLobeWeight(Payload, NoL);
+		return Result;
+	}
+}
+
+
+FMaterialEval SubstrateSolidGlass_EvalMaterial(
+	float3 V_World,
+	float3 L_World,
+	FPathTracingPayload Payload,
+	float2 DiffuseSpecularScale
+)
+{
+	FSubstrateSolidGlassSlabData Data = (FSubstrateSolidGlassSlabData)0;
+	Data.PrepareBSDF(V_World, Payload);
+
+	const float3 N_World = Payload.WorldNormal;
+
+	// move vectors into right shading frame
+	const float3 V = Data.V;
+	const float3 L = mul(Data.Basis, L_World);
+
+	FMaterialEval Result = NullMaterialEval();
+	if (L.z > 0.0)
+	{
+		// Evaluate reflection lobes
+		const float3 H = normalize(V + L);
+		const float NoV = saturate(V.z);
+		const float NoL = saturate(L.z);
+		const float VoH = saturate(dot(V, H));
+		const float NoH = saturate(H.z);
+
+
+		{
+			// Specular Lobes
+			const float2 GGXResult0 = GGXEvalReflection(L, V, H, Data.Alpha0);
+			const float2 GGXResult1 = GGXEvalReflection(L, V, H, Data.Alpha1);
+
+			const float Fg = FresnelReflectance(VoH, Data.Eta, Data.F0);
+
+			Result.AddLobeWithMIS(DiffuseSpecularScale.y * Data.FuzzAttenuation * (1.0 - Payload.RoughnessData.z) * GGXResult0.x / Data.EGlass0, Fg * GGXResult0.y, Data.LobePdf.x);
+			Result.AddLobeWithMIS(DiffuseSpecularScale.y * Data.FuzzAttenuation * (      Payload.RoughnessData.z) * GGXResult1.x / Data.EGlass1, Fg * GGXResult1.y, Data.LobePdf.y);
+		}
+		{
+			// Cloth Lobe
+			const float ShadowTerminator = ShadowTerminatorTerm(L_World, N_World, Payload.WorldSmoothNormal);
+			const float DCloth = D_Charlie(Data.FuzzRoughness, NoH);
+			const float3 FCloth = F_Schlick(Payload.FuzzColor, VoH);
+			const float3 ClothWeight = NoL > 0 && NoV > 0 ? (NoL * Payload.FuzzAmount * DCloth * Vis_Ashikhmin(NoV, NoL)) * FCloth * Data.Fuzz.W * (2 * PI) : 0.0;
+			Result.AddLobeWithMIS(DiffuseSpecularScale.y * ClothWeight * ShadowTerminator, 1 / (2.0 * PI), Data.LobePdf.z);
+		}
+
+		Result.Weight *= SubstrateLobeWeight(Payload, NoL);
+	}
+	else if (L.z < 0)
+	{
+		// Evaluate refracted lobes
+		// refraction side
+		const float NoL = saturate(-L.z);
+		float3 Ht = -(Data.Eta * L + V);
+		Ht = normalize((Data.Eta < 1.0f) ? -Ht : Ht);
+
+		const float VoH = dot(V, Ht);
+		const float Fg = 1.0f - FresnelReflectance(VoH, Data.Eta, Data.F0);
+		if (Fg > 0)
+		{
+			const float2 GGXResult0 = GGXEvalRefraction(L, V, Ht, Data.Alpha0, Data.Eta);
+			const float2 GGXResult1 = GGXEvalRefraction(L, V, Ht, Data.Alpha1, Data.Eta);
+
+			Result.AddLobeWithMIS((1.0 - Payload.RoughnessData.z) * GGXResult0.x / Data.EGlass0 * DiffuseSpecularScale.y, GGXResult0.y * Fg, Data.LobePdf.x);
+			Result.AddLobeWithMIS((      Payload.RoughnessData.z) * GGXResult1.x / Data.EGlass1 * DiffuseSpecularScale.y, GGXResult1.y * Fg, Data.LobePdf.y);
+		}
+
+		Result.Weight *= Payload.BSDFOpacity * Payload.WeightV * Data.FuzzAttenuation;
+	}
+	return Result;
+}
diff --git a/Engine/Shaders/Private/PathTracing/Material/PathTracingSubstrateThinGlass.ush b/Engine/Shaders/Private/PathTracing/Material/PathTracingSubstrateThinGlass.ush
new file mode 100644
--- /dev/null
+++ b/Engine/Shaders/Private/PathTracing/Material/PathTracingSubstrateThinGlass.ush
@@ -0,0 +1,259 @@
+// Copyright Epic Games, Inc. All Rights Reserved.
+
+#pragma once
+
+#ifndef STRATA_ENABLED
+#error "This header should only be included when Strata is enabled."
+#endif
+
+#include "PathTracingSubstrateCommon.ush"
+#include "PathTracingMaterialCommon.ush"
+#include "PathTracingGlossy.ush"
+#include "PathTracingFresnel.ush"
+#include "../../BRDF.ush"
+
+
+struct FSubstrateThinGlassSlabData
+{
+	float3x3 Basis;
+	float3 V;
+
+	float2 RoughnessT;
+
+	float  F0, F90;
+	float  Eta;
+	float3 TransmissionColor;
+	float2 AlphaR0, AlphaT0;
+	float2 AlphaR1, AlphaT1;
+	float  FuzzRoughness;
+
+	float  FuzzAttenuation;
+	float2  NudgeE;
+	FBxDFEnergyTermsA SpecR0, SpecR1;
+	FBxDFEnergyTermsA SpecT0, SpecT1;
+	FBxDFEnergyTermsA Fuzz;
+
+	float2 LobeCdf;
+	float3 LobePdf;
+
+	void PrepareBSDF(float3 V_World, FPathTracingPayload Payload)
+	{
+		float2 RoughnessR = Payload.RoughnessData.xy;
+		RoughnessT.x = ComputeThinTransmittedRoughness(RoughnessR.x, Payload.Ior);
+		RoughnessT.y = ComputeThinTransmittedRoughness(RoughnessR.y, Payload.Ior);
+
+		Basis   = GetGGXBasis(RoughnessR.x, Payload.Anisotropy, Payload.WorldNormal, Payload.WorldTangent, AlphaR0);
+		AlphaR1 = GetGGXAlpha(RoughnessR.y, Payload.Anisotropy);
+
+		AlphaT0 = GetGGXAlpha(RoughnessT.x, Payload.Anisotropy);
+		AlphaT1 = GetGGXAlpha(RoughnessT.y, Payload.Anisotropy);
+
+		V = mul(Basis, V_World);
+
+		F0 = F0RGBToF0(Payload.SpecularColor);
+		F90 = saturate(50.0 * F0);
+		TransmissionColor = Payload.GetTransmittanceColor();
+		FuzzRoughness = MakeRoughnessSafe(Payload.FuzzRoughness, STRATA_MIN_FUZZ_ROUGHNESS);
+
+		const float NoV = saturate(V.z);
+		SpecR0 = ComputeGGXSpecEnergyTermsA(RoughnessR.x, NoV, F0, F90);
+		SpecR1 = ComputeGGXSpecEnergyTermsA(RoughnessR.y, NoV, F0, F90);
+		SpecT0 = ComputeGGXSpecEnergyTermsA(RoughnessT.x, NoV, 1.0 - F0, 1.0 - F90);
+		SpecT1 = ComputeGGXSpecEnergyTermsA(RoughnessT.y, NoV, 1.0 - F0, 1.0 - F90);
+		NudgeE.x = rcp(SpecR0.E + SpecT0.E);
+		NudgeE.y = rcp(SpecR1.E + SpecT1.E);
+
+		Fuzz = ComputeClothEnergyTermsA(FuzzRoughness, NoV);
+
+		FuzzAttenuation = 1.0 - Payload.FuzzAmount * Fuzz.E;
+
+		const float3 GlassR = FuzzAttenuation * lerp(SpecR0.E * NudgeE.x, SpecR1.E * NudgeE.y, Payload.RoughnessData.z);
+		const float3 GlassT = FuzzAttenuation * lerp(SpecT0.E * NudgeE.x, SpecT1.E * NudgeE.y, Payload.RoughnessData.z);
+		const float  FuzzAlbedo = Payload.FuzzAmount * Fuzz.E;
+
+		// Now prepare a cdf/pdf for lobe selection
+		// NOTE: we put T first so that the reflected lobe (R and fuzz) are consecutive in the PDF
+		float3 MaxLobeWeight = Payload.GetMaxLobeWeight();
+		LobeCdf = LobeSelectionCdf(
+			Payload.Ior > 0.0 ? MaxLobeWeight * GlassT : 0.0,
+			MaxLobeWeight * GlassR,
+			MaxLobeWeight * FuzzAlbedo);
+		LobePdf = LobeSelectionPdf(LobeCdf);
+	}
+};
+
+
+FMaterialSample SubstrateThinGlass_SampleMaterial(
+	float3 V_World,
+	FPathTracingPayload Payload,
+	float3 RandSample
+)
+{
+	FSubstrateThinGlassSlabData Data = (FSubstrateThinGlassSlabData)0;
+	Data.PrepareBSDF(V_World, Payload);
+
+	float3 L = 0, H = 0, V = Data.V;
+	float OutRoughness = 1;
+
+	const bool bIsReflection = RandSample.x >= Data.LobeCdf.x;
+	const bool bSampleGlass = RandSample.x < Data.LobeCdf.y;
+	if (bSampleGlass)
+	{
+		// Two-phase lobe selection:
+		// First use the overall CDF to pick between transmission and reflection ...
+		RandSample.x = RescaleRandomNumber(RandSample.x,
+			bIsReflection ? Data.LobeCdf.x : 0.0,
+			bIsReflection ? Data.LobeCdf.y : Data.LobeCdf.x);
+
+		// Then use the roughness blend to pick one of the two possible roughness values
+		const bool bUseRoughness1 = RandSample.x < Payload.RoughnessData.z;
+		RandSample.x = RescaleRandomNumber(RandSample.x,
+			bUseRoughness1 ? 0.0 : Payload.RoughnessData.z,
+			bUseRoughness1 ? Payload.RoughnessData.z : 1.0);
+
+		// 
+		float2 Alpha = bIsReflection
+			? (bUseRoughness1 ? Data.AlphaR1 : Data.AlphaR0)
+			: (bUseRoughness1 ? Data.AlphaT1 : Data.AlphaT0);
+		OutRoughness = bIsReflection
+			? (bUseRoughness1 ? Payload.RoughnessData.y : Payload.RoughnessData.x)
+			: (bUseRoughness1 ? Data.RoughnessT.y : Data.RoughnessT.x);
+
+		// #dxr_todo: Evaluate UniformSampleDiskConcentric
+		H = ImportanceSampleVisibleGGX_aniso(UniformSampleDisk(RandSample.xy), Alpha, V).xyz;
+		L = reflect(-V, H);
+		if (L.z <= 0)
+		{
+			// invalid output direction, exit early
+			return NullMaterialSample();
+		}
+	}
+	else
+	{
+		// cloth lobe
+		RandSample.x = RescaleRandomNumber(RandSample.x, Data.LobeCdf.y, 1.0);
+		L = UniformSampleHemisphere(RandSample.xy).xyz;
+		H = normalize(L + V);
+	}
+	// transform to world space
+	const float3 N_World = Payload.WorldNormal;
+	const float NoV = saturate(V.z);
+	const float NoL = saturate(L.z);
+	const float VoH = saturate(dot(V, H));
+	const float NoH = saturate(H.z);
+
+	const float PositionBiasSign = bIsReflection ? 1.0 : -1.0;
+	const float3 L_World = normalize(mul(float3(L.xy, PositionBiasSign * L.z), Data.Basis)); // flip reflection to other side if we are refracting
+	FMaterialSample Result = CreateMaterialSample(L_World, 0.0, 0.0, PositionBiasSign, OutRoughness, bIsReflection ? PATHTRACER_SCATTER_SPECULAR : PATHTRACER_SCATTER_REFRACT);
+
+
+
+	// Specular Lobes (either reflected or transmitted)
+	const float2 GGXResult0 = GGXEvalReflection(L, V, H, bIsReflection ? Data.AlphaR0 : Data.AlphaT0);
+	const float2 GGXResult1 = GGXEvalReflection(L, V, H, bIsReflection ? Data.AlphaR1 : Data.AlphaT1);
+	const float LobePdf = bIsReflection ? Data.LobePdf.y : Data.LobePdf.x;
+	const float ClothPdf = 1.0 / (2 * PI);
+	if (bSampleGlass)
+	{
+		const FThinSlabWeights SlabResult = ComputeThinSlabWeights(Data.TransmissionColor, VoH, Payload.Ior, Data.F0);
+
+		const float3 SlabTint = bIsReflection ? SlabResult.Reflected : SlabResult.Transmitted;
+		const float SlabW0 = (bIsReflection ? Data.SpecR0.W : Data.SpecT0.W) * Data.NudgeE.x;
+		const float SlabW1 = (bIsReflection ? Data.SpecR1.W : Data.SpecT1.W) * Data.NudgeE.y;
+
+
+		Result.AddLobeWithMIS(Data.FuzzAttenuation * (1 - Payload.RoughnessData.z) * GGXResult0.x * SlabW0 * SlabTint, GGXResult0.y, LobePdf * (1 - Payload.RoughnessData.z));
+		Result.AddLobeWithMIS(Data.FuzzAttenuation * (    Payload.RoughnessData.z) * GGXResult1.x * SlabW1 * SlabTint, GGXResult1.y, LobePdf * (    Payload.RoughnessData.z));
+
+		if (bIsReflection)
+		{
+			Result.Pdf += Data.LobePdf.z * ClothPdf;
+		}
+	}
+	else
+	{
+		// Cloth Lobe
+		const float ShadowTerminator = ShadowTerminatorTerm(L_World, N_World, Payload.WorldSmoothNormal);
+		const float DCloth = D_Charlie(Data.FuzzRoughness, NoH);
+		const float3 FCloth = F_Schlick(Payload.FuzzColor, VoH);
+		const float3 ClothWeight = NoL > 0 && NoV > 0 ? (NoL * Payload.FuzzAmount * DCloth * Vis_Ashikhmin(NoV, NoL)) * FCloth * Data.Fuzz.W * (2 * PI) : 0.0;
+		Result.AddLobeWithMIS(ClothWeight * ShadowTerminator, ClothPdf, Data.LobePdf.z);
+
+		Result.Pdf += LobePdf * (1.0 - Payload.RoughnessData.z) * GGXResult0.y;
+		Result.Pdf += LobePdf * (      Payload.RoughnessData.z) * GGXResult1.y;
+	}
+
+
+	if (bIsReflection)
+	{
+		Result.Weight *= SubstrateLobeWeight(Payload, NoL);
+	}
+	else
+	{
+		Result.Weight *= Payload.BSDFOpacity * Payload.WeightV;
+	}
+	return Result;
+}
+
+
+FMaterialEval SubstrateThinGlass_EvalMaterial(
+	float3 V_World,
+	float3 L_World,
+	FPathTracingPayload Payload,
+	float2 DiffuseSpecularScale
+)
+{
+	FSubstrateThinGlassSlabData Data = (FSubstrateThinGlassSlabData)0;
+	Data.PrepareBSDF(V_World, Payload);
+
+	const float3 N_World = Payload.WorldNormal;
+
+	// move vectors into right shading frame
+	const float3 V = Data.V;
+	float3 L = mul(Data.Basis, L_World);
+
+	const bool bIsReflection = L.z >= 0;
+	L.z = abs(L.z); // push L to the same side as V
+
+	const float NoL = saturate(L.z);
+	const float NoV = saturate(V.z);
+	const float3 H = normalize(L + V);
+	const float VoH = saturate(dot(V, H));
+	const float NoH = saturate(H.z);
+
+
+	// Specular Lobes (either reflected or transmitted)
+	const float2 GGXResult0 = GGXEvalReflection(L, V, H, bIsReflection ? Data.AlphaR0 : Data.AlphaT0);
+	const float2 GGXResult1 = GGXEvalReflection(L, V, H, bIsReflection ? Data.AlphaR1 : Data.AlphaT1);
+	const FThinSlabWeights SlabResult = ComputeThinSlabWeights(Data.TransmissionColor, VoH, Payload.Ior, Data.F0);
+
+	const float3 SlabTint = bIsReflection ? SlabResult.Reflected : SlabResult.Transmitted;
+	const float SlabW0 = (bIsReflection ? Data.SpecR0.W : Data.SpecT0.W) * Data.NudgeE.x;
+	const float SlabW1 = (bIsReflection ? Data.SpecR1.W : Data.SpecT1.W) * Data.NudgeE.y;
+
+	const float LobePdf = bIsReflection ? Data.LobePdf.y : Data.LobePdf.x;
+
+	FMaterialEval Result = NullMaterialEval();
+
+	Result.AddLobeWithMIS(DiffuseSpecularScale.y * Data.FuzzAttenuation * (1 - Payload.RoughnessData.z) * GGXResult0.x * SlabW0 * SlabTint, GGXResult0.y, LobePdf * (1 - Payload.RoughnessData.z));
+	Result.AddLobeWithMIS(DiffuseSpecularScale.y * Data.FuzzAttenuation * (    Payload.RoughnessData.z) * GGXResult1.x * SlabW1 * SlabTint, GGXResult1.y, LobePdf * (    Payload.RoughnessData.z));
+
+
+	if (bIsReflection)
+	{
+		{
+			// Cloth Lobe
+			const float ShadowTerminator = ShadowTerminatorTerm(L_World, N_World, Payload.WorldSmoothNormal);
+			const float DCloth = D_Charlie(Data.FuzzRoughness, NoH);
+			const float3 FCloth = F_Schlick(Payload.FuzzColor, VoH);
+			const float3 ClothWeight = NoL > 0 && NoV > 0 ? (NoL * Payload.FuzzAmount * DCloth * Vis_Ashikhmin(NoV, NoL)) * FCloth * Data.Fuzz.W * (2 * PI) : 0.0;
+			Result.AddLobeWithMIS(DiffuseSpecularScale.y * ClothWeight * ShadowTerminator, 1 / (2.0 * PI), Data.LobePdf.z);
+		}
+		Result.Weight *= SubstrateLobeWeight(Payload, NoL);
+	}
+	else
+	{
+		Result.Weight *= Payload.BSDFOpacity * Payload.WeightV;
+	}
+	return Result;
+}
diff --git a/Engine/Shaders/Private/PathTracing/Material/PathTracingThinGlass.ush b/Engine/Shaders/Private/PathTracing/Material/PathTracingThinGlass.ush
--- a/Engine/Shaders/Private/PathTracing/Material/PathTracingThinGlass.ush
+++ b/Engine/Shaders/Private/PathTracing/Material/PathTracingThinGlass.ush
@@ -35,9 +35,8 @@ FRoughThinGlassData PrepareRoughThinGlassData(FPathTracingPayload Payload, float
 	Data.F0 = F0RGBToF0(Payload.SpecularColor);
 	Data.F90 = saturate(50.0 * Data.F0);
 
-	// transmitted lobe should have slightly less roughness due to the squeeze implied by the IOR (eventually converging to 0.0 as Ior approaches 1.0)
 	const float RoughnessR = Payload.Roughness;
-	const float RoughnessT = Payload.Roughness * sqrt(1.0 - min(Payload.Ior, rcp(Payload.Ior)));
+	const float RoughnessT = ComputeThinTransmittedRoughness(Payload.Roughness, Payload.Ior);
 
 	Data.Basis = GetGGXBasis(RoughnessR, Payload.Anisotropy, Payload.WorldNormal, Payload.WorldTangent, Data.AlphaR);
 	Data.AlphaT = GetGGXAlpha(RoughnessT, Payload.Anisotropy);
diff --git a/Engine/Shaders/Private/PathTracing/PathTracingCommon.ush b/Engine/Shaders/Private/PathTracing/PathTracingCommon.ush
--- a/Engine/Shaders/Private/PathTracing/PathTracingCommon.ush
+++ b/Engine/Shaders/Private/PathTracing/PathTracingCommon.ush
@@ -419,6 +419,12 @@ struct FPathTracingPayload : FMinimalPayload
 	float2 GetHairPrimitiveUV() { return SpecularEdgeColor.xy; }
 	void SetHairPrimitiveUV(float2 UV) { SpecularEdgeColor.xy = UV; }
 
+	float3 GetExtinction() { return rcp(max(MeanFreePath, 0.0001)); }
+
+	// SHADINGMODELID_THIN_TRANSLUCENT
+	float3 GetTransmittanceColor() { return exp(-GetExtinction() * STRATA_SIMPLEVOLUME_THICKNESS_CM); }
+
+
 #else // PATHTRACING_STRATA_PAYLOAD
 	// Various ways to interpret CustomData (depending on ShadingModelID)
 	// NOTE: This is not always following the same conventions as GetMaterialCustomData0,1()
@@ -502,12 +508,20 @@ struct FPathTracingPayload : FMinimalPayload
 
 	bool IsMaterialSolidGlass()
 	{
+#if PATHTRACING_STRATA_PAYLOAD
+		return ShadingModelID == SHADINGMODELID_STRATA && HasRefraction();
+#else
 		return ShadingModelID == SHADINGMODELID_DEFAULT_LIT && HasRefraction() && BSDFOpacity < 1;
+#endif
 	}
 
 	bool IsMaterialThinGlass()
 	{
+#if PATHTRACING_STRATA_PAYLOAD
+		return ShadingModelID == SHADINGMODELID_THIN_TRANSLUCENT && HasRefraction();
+#else
 		return ShadingModelID == SHADINGMODELID_THIN_TRANSLUCENT && HasRefraction() && BSDFOpacity < 1;
+#endif
 	}
 
 	bool IsMaterialTransmissive()
@@ -521,14 +535,16 @@ struct FPathTracingPayload : FMinimalPayload
 
 	bool IsSubsurfaceMaterial()
 	{
+#if PATHTRACING_STRATA_PAYLOAD
+		return ShadingModelID == SHADINGMODELID_STRATA ||
+		       ShadingModelID == SHADINGMODELID_EYE;
+#else
 		return ShadingModelID == SHADINGMODELID_SUBSURFACE ||
 			   ShadingModelID == SHADINGMODELID_PREINTEGRATED_SKIN ||
 			   ShadingModelID == SHADINGMODELID_SUBSURFACE_PROFILE ||
-			   ShadingModelID == SHADINGMODELID_EYE
-#if PATHTRACING_STRATA_PAYLOAD
-			   || ShadingModelID == SHADINGMODELID_STRATA
+			   ShadingModelID == SHADINGMODELID_EYE;
 #endif
-				;
+
 	}
 #endif // IS_PAYLOAD_ENABLED(RT_PAYLOAD_TYPE_GPULIGHTMASS)
 };
diff --git a/Engine/Shaders/Private/PathTracing/PathTracingCore.ush b/Engine/Shaders/Private/PathTracing/PathTracingCore.ush
--- a/Engine/Shaders/Private/PathTracing/PathTracingCore.ush
+++ b/Engine/Shaders/Private/PathTracing/PathTracingCore.ush
@@ -1740,9 +1740,6 @@ bool PathTracingKernel(inout FPathState PathState, int Bounce)
 	// enlarge roughness based on the chosen lobe roughness
 	PathState.PathRoughness = max(PathState.PathRoughness, MaterialSample.Roughness);
 
-#if STRATA_ENABLED
-	// STRATA_TODO: Account for medium absorption here once we have support for refractive glass
-#else
 	// update the current extinction if we are crossing a boundary on glass or water
 	// summing the local extinction gives a rudimentary way of dealing with overlapping regions
 	// long term we will probably want a stack with priorities
@@ -1761,7 +1758,6 @@ bool PathTracingKernel(inout FPathState PathState, int Bounce)
 			PathState.SigmaT = max(PathState.SigmaT, 0);
 		}
 	}
-#endif // STRATA_ENABLED
 
 	// If we are using Material sampling for lights
 	if (MISMode != 1 && bDoLightLoop)
diff --git a/Engine/Shaders/Private/PathTracing/PathTracingInitExtinctionCoefficient.usf b/Engine/Shaders/Private/PathTracing/PathTracingInitExtinctionCoefficient.usf
--- a/Engine/Shaders/Private/PathTracing/PathTracingInitExtinctionCoefficient.usf
+++ b/Engine/Shaders/Private/PathTracing/PathTracingInitExtinctionCoefficient.usf
@@ -8,10 +8,6 @@ RWBuffer<float> RWStartingExtinctionCoefficient;
 
 RAY_TRACING_ENTRY_RAYGEN(PathTracingInitExtinctionCoefficientRG)
 {
-#if STRATA_ENABLED
-	// STRATA_TODO: Figure out medium parameters from strata data
-	float3 StartingSigmaT = 0.0;
-#else
 	const uint RayFlags = 0;
 	const uint MissShaderIndex = 0;
 	const uint InstanceInclusionMask = PATHTRACER_MASK_CAMERA; // Ignore hair strands, only trace what camera rays can see
@@ -26,6 +22,8 @@ RAY_TRACING_ENTRY_RAYGEN(PathTracingInitExtinctionCoefficientRG)
 	for (;;)
 	{
 		FPackedPathTracingPayload PackedPayload = InitPathTracingPayload(PATHTRACER_SCATTER_CAMERA, 0.0);
+		// STRATA_TODO: Set the stochastic slab random value? Add special flag to only return extinction of solid glass?
+		// The current behavior will only work correctly for single-slab materials
 		TraceRay(
 			TLAS,
 			RayFlags,
@@ -53,7 +51,6 @@ RAY_TRACING_ENTRY_RAYGEN(PathTracingInitExtinctionCoefficientRG)
 		// keep tracing
 		UpRay.TMin = asfloat(asuint(PackedPayload.HitT) + 1);
 	}
-#endif // STRATA_ENABLED
 	// clamp against 0 in the unlikely case that we only hit front faces of glass on our way up
 	StartingSigmaT = max(StartingSigmaT, 0.0);
 
diff --git a/Engine/Shaders/Private/PathTracing/PathTracingMaterialHitShader.usf b/Engine/Shaders/Private/PathTracing/PathTracingMaterialHitShader.usf
--- a/Engine/Shaders/Private/PathTracing/PathTracingMaterialHitShader.usf
+++ b/Engine/Shaders/Private/PathTracing/PathTracingMaterialHitShader.usf
@@ -44,6 +44,8 @@ float AdjustMaterialRoughness(float Roughness, float PathRoughness)
 	return max(Roughness, PathRoughness);
 }
 
+#if REFRACTION_USE_INDEX_OF_REFRACTION
+
 #if STRATA_ENABLED
 
 float GetRefractionIor(FPixelMaterialInputs PixelMaterialInputs, float3 F0)
@@ -74,6 +76,45 @@ float GetRefractionIor(FPixelMaterialInputs PixelMaterialInputs)
 
 #endif // STRATA_ENABLED
 
+
+float3 FakeCausticsPathFactor(float Roughness, float PathRoughness)
+{
+	// The heuristic used here is inspired by the following presentations:
+	//   - Kulla & Conty: Revisiting Physically Based Shading at Imageworks
+	//       https://blog.selfshadow.com/publications/s2017-shading-course/imageworks/s2017_pbs_imageworks_slides_v2.pdf
+	//   - Colin Barré-Brisebois - Pica Pica & Nvidia Turing
+	//       https://www.ea.com/seed/news/siggraph-2018-picapica-nv-turing
+
+	// NOTE: this returns 0.0 for PathRoughness <= Roughness which accounts for when we are not using ApproximateCaustics
+	return (1 - Roughness * Roughness) * saturate(PathRoughness - Roughness);
+}
+
+
+float FakeCaustics(float F0, float Ior, float NoV)
+{
+	// fake caustics for solid glass, approximates the effect of the boundaries' fresnel
+	float Fr = FresnelReflectance(abs(NoV), Ior, F0);
+	return Pow2(1 - Fr);
+}
+
+float3 ComputeOpticalDepthForHit(float3 LocalSigmaT, bool bIsFrontFace)
+{
+	// Does the material have any kind of volumetric absorption to apply?
+	// Now track the optical thickness so that we can account for Beer's law along the shadow ray
+	// TODO: to support lights inside glass we would need to track an extra offset, but this is hopefully an uncommon scenario
+	// TODO2: SingleLayerWater needs something similar, but it generally is setup with opaque or masked blend mode and water geometry does not cast shadows ...
+	float Distance = RayTCurrent();
+
+	// Front Face: remove contribution from the ray origin to the current hit (assuming the backside will be hit)
+	// Back Face:  add contribution from ray origin to current hit (backside) - excess portion not covered by the ray will be removed by the front face hit
+	Distance *= bIsFrontFace ? -1.0 : +1.0;
+
+	return LocalSigmaT * Distance;
+}
+
+
+#endif // REFRACTION_USE_INDEX_OF_REFRACTION
+
 float3 DecodeSSSProfileRadius(uint ProfileId, float3 DiffuseColor, float Opacity)
 {
 	// Burley parameterization
@@ -172,7 +213,14 @@ RAY_TRACING_ENTRY_CLOSEST_HIT(PathTracingMaterialCHS,
 	float3 TotalTransmittancePreCoverage = 0;
 	const float3 V_World = MaterialParameters.CameraVector;
 	StrataPixelHeader.StrataUpdateTree(NullStrataAddressing, StrataData, V_World, IntegrationSettings, TotalCoverage, TotalTransmittancePreCoverage);
-	const float3 Transparency = saturate((1 - TotalCoverage) + TotalCoverage * TotalTransmittancePreCoverage);
+#if REFRACTION_USE_INDEX_OF_REFRACTION
+	// transparency is only used for coverage, transmittance happens through refraction (except for shadows potentially, which is handled below)
+	float3 Transparency = 1 - TotalCoverage; 
+#else
+	// transparency is used for both coverage and transmittance
+	// STRATA_TODO: Thin case does not seem to be properly handled here, so the TotalTransmittancePreCoverage is recomputed later
+	float3 Transparency = saturate((1 - TotalCoverage) + TotalCoverage * TotalTransmittancePreCoverage);
+#endif // REFRACTION_USE_INDEX_OF_REFRACTION
 #endif // STRATA_OPTIMIZED_UNLIT
 
 #else // MATERIAL_IS_STRATA
@@ -188,8 +236,83 @@ RAY_TRACING_ENTRY_CLOSEST_HIT(PathTracingMaterialCHS,
 	{
 
 #if STRATA_ENABLED
-		// Transparency was already defined above
-		// STRATA_TODO: figure out refractive / thin translucent equivalents
+		// We already figured out how much Transparency there is above
+		// So only do extra logic for refraction and thin refraction ases
+#if MATERIALBLENDING_TRANSLUCENT && (REFRACTION_USE_INDEX_OF_REFRACTION || MATERIAL_ISTHINSURFACE)
+
+		// STRATA_TODO: refraction handling
+		float PathRoughness = PackedPayload.GetPathRoughness();
+
+		// STRATA_TODO: This does not match the logic in DistortAccumulatePS.usf
+		//  The idea here is to take F0 from the bottom refraction layers since this matches most closely what the refraction logic below does (only bottom slabs can refract).
+		//  The distortion approximation on the raster side instead averages F0 across layers which may produce different results.
+		//  In both cases - we get a potentially strange answer when there is more than one bottom slab that produces refraction
+
+		// STRATA_TODO: How should we handle MFP when there are multiple bottom slabs?
+		float F0 = 0.0f;
+		float3 RoughnessData = 0.0;
+		float NoV = 0.0f;
+		float3 LocalSigmaT = 0.0f;
+		STRATA_UNROLL_N(STRATA_CLAMPED_BSDF_COUNT)
+		for (int BSDFIdx = 0; BSDFIdx < StrataPixelHeader.StrataTree.BSDFCount; ++BSDFIdx)
+		{
+#define BSDF StrataPixelHeader.StrataTree.BSDFs[BSDFIdx]
+			if (BSDF.bIsBottom)
+			{
+				// STRATA_TODO: Also consider the MFP to see how much this slab is contributing to refraction ....
+				F0 = F0RGBToF0(StrataGetBSDFSpecularF0(BSDF));
+				RoughnessData.x = SLAB_ROUGHNESS(BSDF);
+				if (BSDF_GETHASHAZINESS(BSDF))
+				{
+					// NOTE: It should be safe to assume that the clearcoat case cannot happen here as it is only triggered by legacy clearcoat materials which are unlikely to have refraction
+					const FHaziness Haziness = UnpackHaziness(SLAB_HAZINESS(BSDF));
+					RoughnessData.y = Haziness.Roughness;
+					RoughnessData.z = Haziness.Weight;
+				}
+
+				float3 N = StrataPixelHeader.SharedLocalBases.Normals[BSDF_GETSHAREDLOCALBASISID(BSDF)];
+				LocalSigmaT = rcp(max(SLAB_SSSMFP(BSDF), 0.0001));
+				NoV = dot(N, V_World);
+				break;
+			}
+#undef BSDF
+		}
+
+
+#if REFRACTION_USE_INDEX_OF_REFRACTION
+		float Ior = GetRefractionIor(PixelMaterialInputs, F0);
+#else
+		float Ior = 0.0;
+#endif
+
+#if MATERIAL_ISTHINSURFACE
+		// Thin Glass Fake Caustics
+
+		// STRATA_TODO: Revisit this -- TotalTransmittancePreCoverage should probably account for this in the thin case
+		float3 Transmission = exp(-LocalSigmaT * STRATA_SIMPLEVOLUME_THICKNESS_CM);
+		float3 Tr = ComputeThinSlabWeights(Transmission, NoV, Ior, F0).Transmitted;
+#else
+		HLSL_STATIC_ASSERT(REFRACTION_USE_INDEX_OF_REFRACTION == 1, "Expect refraction to be true here");
+		// Solid Glass Fake Caustics
+		float Tr = FakeCaustics(F0, Ior, NoV);
+		// Account for transmission 
+		PackedPayload.SetTau(PackedPayload.GetTau() + ComputeOpticalDepthForHit(LocalSigmaT, bIsFrontFace));
+#endif
+
+		// If material is doing refraction, apply an extra factor to account for ApproximateCaustics
+#if REFRACTION_USE_INDEX_OF_REFRACTION
+		if (Ior > 0.0)
+		{
+			Tr *= lerp(FakeCausticsPathFactor(RoughnessData.x, PathRoughness),
+					   FakeCausticsPathFactor(RoughnessData.y, PathRoughness), RoughnessData.z);
+		}
+#endif
+		Transparency += TotalCoverage * Tr;
+
+		// STRATA_TODO: Is this sanity check necessary?
+		Transparency = saturate(Transparency);
+#endif // MATERIALBLENDING_TRANSLUCENT && (REFRACTION_USE_INDEX_OF_REFRACTION || MATERIAL_ISTHINSURFACE)
+
 
 #else // STRATA_ENABLED
 
@@ -242,10 +365,12 @@ RAY_TRACING_ENTRY_CLOSEST_HIT(PathTracingMaterialCHS,
 				Transparency *= ComputeThinSlabWeights(Transmission, VoN, Ior, F0).Transmitted;
 
 				// fake caustic approximation (see comments below)
+#if REFRACTION_USE_INDEX_OF_REFRACTION
 				if (Ior > 0.0)
 				{
-					Transparency *= (1 - Roughness * Roughness) * saturate(PathRoughness - Roughness);
+					Transparency *= FakeCausticsPathFactor(Roughness, PathRoughness);
 				}
+#endif
 			}
 		}
 		else
@@ -260,42 +385,23 @@ RAY_TRACING_ENTRY_CLOSEST_HIT(PathTracingMaterialCHS,
 				// current hit has some refraction
 				float PathRoughness = PackedPayload.GetPathRoughness();
 				float Roughness = GetMaterialRoughness(PixelMaterialInputs);
-				if (PathRoughness <= Roughness)
-				{
-					// not using fast caustic approximation - treat as opaque
-					PackedPayload.SetRayThroughput(0.0);
-					return;
-				}
-				// The heuristic used here is inspired by the following presentations:
-				//   - Kulla & Conty: Revisiting Physically Based Shading at Imageworks
-				//       https://blog.selfshadow.com/publications/s2017-shading-course/imageworks/s2017_pbs_imageworks_slides_v2.pdf
-				//   - Colin Barré-Brisebois - Pica Pica & Nvidia Turing
-				//       https://www.ea.com/seed/news/siggraph-2018-picapica-nv-turing
 
-				// fake caustics
-				float3 N = normalize(MaterialParameters.WorldNormal);
-				float Eta = Ior;
-				float VoN = abs(dot(WorldRayDirection(), N));
 				float3 BaseColor = GetMaterialBaseColor(PixelMaterialInputs);
 				float  Metallic = GetMaterialMetallic(PixelMaterialInputs);
 				float  Specular = GetMaterialSpecular(PixelMaterialInputs);
 				float  F0 = F0RGBToF0(ComputeF0(Specular, BaseColor, Metallic));
-				float Fr = FresnelReflectance(VoN, Eta, F0);
-				Transparency *= pow(1 - Fr, 2) * (1 - Roughness * Roughness) * saturate(PathRoughness - Roughness);
+
+				float3 N = normalize(MaterialParameters.WorldNormal);
+				float NoV = dot(WorldRayDirection(), N);
+
+				Transparency *= FakeCaustics(F0, Ior, NoV) * FakeCausticsPathFactor(Roughness, PathRoughness);
 
 #if HAVE_GetAbsorptionMediumMaterialOutput0
 				// Does the material have any kind of volumetric absorption to apply?
-				// Now track the optical thickness so that we can account for Beer's law along the shadow ray
-				// TODO: to support lights inside glass we would need to track an extra offset, but this is hopefully an uncommon scenario
-				// TODO2: SingleLayerWater needs something similar, but it generally is setup with opaque or masked blend mode and water geometry does not cast shadows ...
-				float Distance = RayTCurrent();
 				float3 TransmittanceColor = GetAbsorptionMediumMaterialOutput0(MaterialParameters);
 				float3 LocalSigmaT = PathTracingGlassTransmittanceToExtinction(TransmittanceColor);
 				
-				// Front Face: remove contribution from the ray origin to the current hit (assuming the backside will be hit)
-				// Back Face:  add contribution from ray origin to current hit (backside) - excess portion not covered by the ray will be removed by the front face hit
-				Distance *= bIsFrontFace ? -1.0 : +1.0;
-				PackedPayload.SetTau(PackedPayload.GetTau() + LocalSigmaT * Distance);
+				PackedPayload.SetTau(PackedPayload.GetTau() + ComputeOpticalDepthForHit(LocalSigmaT, bIsFrontFace));
 #endif // HAVE_GetAbsorptionMediumMaterialOutput0
 			}
 		}
@@ -304,7 +410,7 @@ RAY_TRACING_ENTRY_CLOSEST_HIT(PathTracingMaterialCHS,
 		{
 			// base case for shadingmodel if/else
 		}
-#else // MATERIALBLEINDG_*
+#else // MATERIALBLENDING_*
 #error Unhandled blending mode!
 #endif
 
@@ -546,6 +652,7 @@ RAY_TRACING_ENTRY_CLOSEST_HIT(PathTracingMaterialCHS,
 								Spec0Albedo *= 1.0 - SpecBlend;
 								Spec1Albedo *=       SpecBlend;
 							}
+							// STRATA_TODO: Need to account for glass case here for lobe selection ....
 							const float  FuzzAlbedo = FuzzAmount * FuzzE;
 							const float3 SlabAlbedo = MaxLobeWeight * (DiffuseWeight * DiffuseColor + Spec0Albedo + Spec1Albedo + FuzzAlbedo);
 							Pdf = LobeColorToWeight(SlabAlbedo);
@@ -654,7 +761,30 @@ RAY_TRACING_ENTRY_CLOSEST_HIT(PathTracingMaterialCHS,
 					Payload.WeightV             = BSDF.LuminanceWeightV;
 					Payload.TransmittanceN      = BSDF.TransmittanceAboveAlongN;
 					Payload.CoverageAboveAlongN = BSDF.CoverageAboveAlongN;
-					Payload.Ior                 = GetRefractionIor(PixelMaterialInputs, Payload.SpecularColor);
+#if MATERIAL_ISTHINSURFACE && MATERIALBLENDING_ANY_TRANSLUCENT
+					Payload.ShadingModelID      = SHADINGMODELID_THIN_TRANSLUCENT;
+#endif
+
+#if MATERIALBLENDING_ANY_TRANSLUCENT && REFRACTION_USE_INDEX_OF_REFRACTION
+					// NOTE: only the bottom most layer can refract, coatings should never include a refraction lobe
+					Payload.Ior					= BSDF.bIsBottom ? GetRefractionIor(PixelMaterialInputs, Payload.SpecularColor) : 0.0;
+					Payload.MeanFreePath        = BSDF.bIsBottom ? SLAB_SSSMFP(BSDF) : Payload.MeanFreePath;
+#else
+					Payload.Ior                 = 0.0;
+#endif
+
+#if MATERIAL_ISTHINSURFACE && MATERIALBLENDING_ANY_TRANSLUCENT
+					// Material is configured as a thin surface, and refraction is disabled
+					if (BSDF.bIsBottom && !(Payload.Ior > 0.0))
+					{
+						Payload.MeanFreePath = SLAB_SSSMFP(BSDF);
+						float3 Transmission = Payload.GetTransmittanceColor();
+						const float NoV = saturate(dot(Payload.WorldNormal, V_World));
+						float3 Tr = ComputeThinSlabWeights(Transmission, NoV, Payload.Ior, F0RGBToF0(Payload.SpecularColor)).Transmitted;
+						Payload.TransparencyColor += TotalCoverage * Tr;
+					}
+#endif
+
 					break;
 				}
 				case STRATA_BSDF_TYPE_HAIR:
