Don't set EarlyZ if we are sampling VT but not doing any VT feedback

diff --git a/Engine/Shaders/Private/BasePassPixelShader.usf b/Engine/Shaders/Private/BasePassPixelShader.usf
--- a/Engine/Shaders/Private/BasePassPixelShader.usf
+++ b/Engine/Shaders/Private/BasePassPixelShader.usf
@@ -2084,7 +2084,7 @@ void FPixelShaderInOut_MainPS(
 	// We also half that range to also make sure we have room for other additive elements such as bloom, clouds or particle visual effects.
 	Out.MRT[0].xyz = min(Out.MRT[0].xyz, Max10BitsFloat.xxx * 0.5f);
 
-#if NUM_VIRTUALTEXTURE_SAMPLES || LIGHTMAP_VT_ENABLED
+#if MATERIAL_VIRTUALTEXTURE_FEEDBACK || LIGHTMAP_VT_ENABLED
 #if OUTPUT_PIXEL_DEPTH_OFFSET
 	// When in deferred, opaque materials with pixel depth offset must execute a custom depth test in order to avoid conflicting UAV writes.
 	if (ManualDepthTestEqual(In.SvPosition, Out.Depth))
@@ -2137,8 +2137,8 @@ void FPixelShaderInOut_MainPS(
 
 #else // STRATA_OPAQUE_DEFERRED
 
-// If virtual texture is enabled then use early depth test so that UAV feedback buffer writes respect the depth test
-#if NUM_VIRTUALTEXTURE_SAMPLES || LIGHTMAP_VT_ENABLED || OIT_ENABLED
+// If virtual texture feedback is enabled then use early depth test so that UAV feedback buffer writes respect the depth test
+#if MATERIAL_VIRTUALTEXTURE_FEEDBACK || LIGHTMAP_VT_ENABLED || OIT_ENABLED
 	#if IS_NANITE_PASS
 		// Even with masking and/or pixel depth offset, Nanite never uses actual clip() or SV_Depth
 		#define PIXELSHADER_EARLYDEPTHSTENCIL EARLYDEPTHSTENCIL
diff --git a/Engine/Shaders/Private/DebugViewModePixelShader.usf b/Engine/Shaders/Private/DebugViewModePixelShader.usf
--- a/Engine/Shaders/Private/DebugViewModePixelShader.usf
+++ b/Engine/Shaders/Private/DebugViewModePixelShader.usf
@@ -444,7 +444,7 @@ float4 VisualizeRequiredTextureResolution(in FDebugPSIn DebugInputs, in bool bIs
 		}
 	}
 
-#if NUM_VIRTUALTEXTURE_SAMPLES || LIGHTMAP_VT_ENABLED
+#if MATERIAL_VIRTUALTEXTURE_FEEDBACK || LIGHTMAP_VT_ENABLED
 	FinalizeVirtualTextureFeedback(
 		MaterialParameters.VirtualTextureFeedback,
 		MaterialParameters.SvPosition,
@@ -516,7 +516,7 @@ float4 ApplyDebugColor(in FDebugPSIn DebugInputs, in bool bIsFrontFace, in float
 
 	Result *= (.05 + .95 * Luminance(BaseColor));
 
-#if NUM_VIRTUALTEXTURE_SAMPLES || LIGHTMAP_VT_ENABLED
+#if MATERIAL_VIRTUALTEXTURE_FEEDBACK || LIGHTMAP_VT_ENABLED
 	FinalizeVirtualTextureFeedback(
 		MaterialParameters.VirtualTextureFeedback,
 		MaterialParameters.SvPosition,
diff --git a/Engine/Shaders/Private/DeferredDecal.usf b/Engine/Shaders/Private/DeferredDecal.usf
--- a/Engine/Shaders/Private/DeferredDecal.usf
+++ b/Engine/Shaders/Private/DeferredDecal.usf
@@ -57,8 +57,8 @@ void MainVS(
 #define DECAL_PRIMITIVE 1
 #include "/Engine/Generated/Material.ush"
 
-// If virtual texture is enabled then use early depth test so that UAV feedback buffer writes respect the depth test
-#if NUM_VIRTUALTEXTURE_SAMPLES
+// If virtual texture feedback is enabled then use early depth test so that UAV feedback buffer writes respect the depth test
+#if MATERIAL_VIRTUALTEXTURE_FEEDBACK
 	#if COMPILER_SUPPORTS_DEPTHSTENCIL_EARLYTEST_LATEWRITE
 		// If we support early depth test with late write behaviour then use it since we may be using discard, or modifying depth
 		#define PIXELSHADER_EARLYDEPTHSTENCIL DEPTHSTENCIL_EARLYTEST_LATEWRITE
@@ -136,7 +136,7 @@ void FPixelShaderInOut_MainPS(inout FPixelShaderIn In, inout FPixelShaderOut Out
 	MaterialParameters.AbsoluteWorldPosition = MaterialParameters.WorldPosition_NoOffsets = LWCSubtract(TranslatedWorldPosition, PrimaryView.PreViewTranslation);
 	MaterialParameters.WorldPosition_CamRelative = MaterialParameters.WorldPosition_NoOffsets_CamRelative = TranslatedWorldPosition;
 
-#if NUM_VIRTUALTEXTURE_SAMPLES
+#if MATERIAL_VIRTUALTEXTURE_FEEDBACK
 	InitializeVirtualTextureFeedback(MaterialParameters.VirtualTextureFeedback, (uint2)In.SvPosition.xy, View.FrameNumber);
 #endif
 
@@ -194,7 +194,7 @@ void FPixelShaderInOut_MainPS(inout FPixelShaderIn In, inout FPixelShaderOut Out
 	DecalCommonOutput(In, Out, Color, Opacity, GBufferData);
 #endif // STRATA_ENABLED
 
-#if NUM_VIRTUALTEXTURE_SAMPLES
+#if MATERIAL_VIRTUALTEXTURE_FEEDBACK
 	FinalizeVirtualTextureFeedback(
 		MaterialParameters.VirtualTextureFeedback,
 		MaterialParameters.SvPosition,
diff --git a/Engine/Shaders/Private/MaterialTemplate.ush b/Engine/Shaders/Private/MaterialTemplate.ush
--- a/Engine/Shaders/Private/MaterialTemplate.ush
+++ b/Engine/Shaders/Private/MaterialTemplate.ush
@@ -3594,7 +3594,7 @@ void CalcMaterialParametersEx(
 	#endif
 #endif
 
-#if NUM_VIRTUALTEXTURE_SAMPLES || LIGHTMAP_VT_ENABLED
+#if MATERIAL_VIRTUALTEXTURE_FEEDBACK || LIGHTMAP_VT_ENABLED
 	InitializeVirtualTextureFeedback(Parameters.VirtualTextureFeedback, (uint2)SvPosition.xy, View.FrameNumber);
 #endif
 
diff --git a/Engine/Shaders/Private/MeshDecals.usf b/Engine/Shaders/Private/MeshDecals.usf
--- a/Engine/Shaders/Private/MeshDecals.usf
+++ b/Engine/Shaders/Private/MeshDecals.usf
@@ -21,8 +21,8 @@
 #include "/Engine/Generated/VertexFactory.ush"
 
 
-// If virtual texture is enabled then use early depth test so that UAV feedback buffer writes respect the depth test
-#if NUM_VIRTUALTEXTURE_SAMPLES
+// If virtual texture feedback is enabled then use early depth test so that UAV feedback buffer writes respect the depth test
+#if MATERIAL_VIRTUALTEXTURE_FEEDBACK
 	#if COMPILER_SUPPORTS_DEPTHSTENCIL_EARLYTEST_LATEWRITE
 		// If we support early depth test with late write behaviour then use it since we may be using discard, or modifying depth
 		#define PIXELSHADER_EARLYDEPTHSTENCIL DEPTHSTENCIL_EARLYTEST_LATEWRITE
@@ -180,7 +180,7 @@ void FPixelShaderInOut_MainPS(
 	DecalCommonOutput(In, Out, Color, Opacity, GBufferData);
 #endif // STRATA_ENABLED
 
-#if NUM_VIRTUALTEXTURE_SAMPLES
+#if MATERIAL_VIRTUALTEXTURE_FEEDBACK
 	FinalizeVirtualTextureFeedback(
 		MaterialParameters.VirtualTextureFeedback,
 		MaterialParameters.SvPosition,
diff --git a/Engine/Shaders/Private/MobileBasePassPixelShader.usf b/Engine/Shaders/Private/MobileBasePassPixelShader.usf
--- a/Engine/Shaders/Private/MobileBasePassPixelShader.usf
+++ b/Engine/Shaders/Private/MobileBasePassPixelShader.usf
@@ -166,7 +166,7 @@ void ApplyPixelDepthOffsetForMobileBasePass(inout FMaterialPixelParameters Mater
 
 #define USES_PIXEL_DISCARD ((MATERIALBLENDING_MASKED || USE_DITHERED_LOD_TRANSITION) && !EARLY_Z_PASS_ONLY_MATERIAL_MASKING)
 // Force early depth_stencil for materials that use VT feedback and don't use pixel discard or modify the pixel depth
-#if (NUM_VIRTUALTEXTURE_SAMPLES || LIGHTMAP_VT_ENABLED) && !(OUTPUT_PIXEL_DEPTH_OFFSET || USES_PIXEL_DISCARD)
+#if (MATERIAL_VIRTUALTEXTURE_FEEDBACK || LIGHTMAP_VT_ENABLED) && !(OUTPUT_PIXEL_DEPTH_OFFSET || USES_PIXEL_DISCARD)
 	#define PIXELSHADER_EARLYDEPTHSTENCIL EARLYDEPTHSTENCIL	
 #else
 	#define PIXELSHADER_EARLYDEPTHSTENCIL 	
@@ -984,7 +984,7 @@ void Main(
 	#endif
 #endif
 
-#if NUM_VIRTUALTEXTURE_SAMPLES || LIGHTMAP_VT_ENABLED
+#if MATERIAL_VIRTUALTEXTURE_FEEDBACK || LIGHTMAP_VT_ENABLED
 	FinalizeVirtualTextureFeedback(
 		MaterialParameters.VirtualTextureFeedback,
 		MaterialParameters.SvPosition,
diff --git a/Engine/Source/Runtime/Engine/Private/Materials/HLSLMaterialTranslator.cpp b/Engine/Source/Runtime/Engine/Private/Materials/HLSLMaterialTranslator.cpp
--- a/Engine/Source/Runtime/Engine/Private/Materials/HLSLMaterialTranslator.cpp
+++ b/Engine/Source/Runtime/Engine/Private/Materials/HLSLMaterialTranslator.cpp
@@ -2054,6 +2054,14 @@ void FHLSLMaterialTranslator::GetMaterialEnvironment(EShaderPlatform InPlatform,
 	// Count the number of VTStacks (each stack will allocate a feedback slot)
 	OutEnvironment.SetDefine(TEXT("NUM_VIRTUALTEXTURE_SAMPLES"), VTStacks.Num());
 
+	// Check if any feedback slots are in use. We can simplify shader and remove EARLYZ optimizations if none are.
+	bool bGenerateFeedback = false;
+	for (int i = 0; i < VTStacks.Num() && !bGenerateFeedback; ++i)
+	{
+		bGenerateFeedback |= VTStacks[i].bGenerateFeedback;
+	}
+	OutEnvironment.SetDefine(TEXT("MATERIAL_VIRTUALTEXTURE_FEEDBACK"), bGenerateFeedback);
+
 	// Setup defines to map each VT stack to either 1 or 2 page table textures, depending on how many layers it uses
 	for (int i = 0; i < VTStacks.Num(); ++i)
 	{
