* Fix increase of energy when light transmission is enabled due to incorrectly normalize phase function.
* Add MFP/SSS aniso value for SSS profile in Substrate material info.

diff --git a/Engine/Shaders/Private/Strata/StrataEvaluation.ush b/Engine/Shaders/Private/Strata/StrataEvaluation.ush
--- a/Engine/Shaders/Private/Strata/StrataEvaluation.ush
+++ b/Engine/Shaders/Private/Strata/StrataEvaluation.ush
@@ -754,7 +754,7 @@ FStrataEvaluateResult StrataEvaluateBSDFCommon(FStrataBSDFContext BSDFContext, F
 				}
 
 				const float3 RefracV = refract(BSDFContext.V, -BSDFContext.N, OneOverIOR);
-				const float PhaseFunction = ApproximateHG(dot(-BSDFContext.L, RefracV), PhaseFunctionAnisotropy);
+				const float PhaseFunction = HenyeyGreensteinPhase(PhaseFunctionAnisotropy, dot(BSDFContext.L, RefracV));
 
 				Sample.ThroughputV				= OpaqueBSDFThroughput; // SSS is not translucent as of today
 				Sample.TransmittanceAlongN		= OpaqueBSDFThroughput; // idem
diff --git a/Engine/Shaders/Private/Strata/StrataVisualize.usf b/Engine/Shaders/Private/Strata/StrataVisualize.usf
--- a/Engine/Shaders/Private/Strata/StrataVisualize.usf
+++ b/Engine/Shaders/Private/Strata/StrataVisualize.usf
@@ -59,10 +59,8 @@ void AddDrawPixelFootprint(float3 P, float3 dPdx, float3 dPdy, float2 Scale, boo
 void DrawPixelFootprint(float3 P, float3 dPdx, float3 dPdy, uint2 PixelCoord)
 {
 	const FStrataSubsurfaceHeader SSSHeader = StrataLoadSubsurfaceHeader(Strata.MaterialTextureArray, Strata.FirstSliceStoringStrataSSSData, PixelCoord);
-	const uint SubsurfaceProfileInt = StrataSubSurfaceHeaderGetProfileId(SSSHeader);
 	const bool bIsValid = StrataSubSurfaceHeaderGetIsValid(SSSHeader);
 	const bool bIsProfile = StrataSubSurfaceHeaderGetIsProfile(SSSHeader);
-	const float PixelRadiusScale = StrataSubSurfaceHeaderGetProfileRadiusScale(SSSHeader);
 
 	if (bIsValid)
 	{
@@ -266,21 +264,25 @@ void PrintBSDF(inout FShaderPrintContext Context, inout FStrataAddressing Strata
 				{
 					const float Opacity = StrataSubSurfaceHeaderGetWrapOpacity(SSSHeader);
 					Print(Context, TEXT(" SSS Opacity "), FontBSDFPropName);	Print(Context, Opacity, FontBSDFPropValu);	Newline(Context, RectMax);
+					Print(Context, TEXT(" SSS Phase   "), FontBSDFPropName);	Print(Context, SLAB_SSSPHASEANISOTROPY(BSDF), FontBSDFPropValu);		Newline(Context, RectMax);
 				}
 				if (BSDF_GETSSSTYPE(BSDF) == SSS_TYPE_DIFFUSION_PROFILE && BSDF_GETISSIMPLEVOLUME(BSDF) == 0)
 				{
 					const uint ProfileId = StrataSubSurfaceHeaderGetProfileId(SSSHeader);
 					const float RadiusScale = StrataSubSurfaceHeaderGetProfileRadiusScale(SSSHeader);
+					const float3 DiffuseMFP = GetSubsurfaceProfileMFPInCm(ProfileId) * RadiusScale;
+					const float Phase = GetTransmissionProfileParams(ProfileId).ScatteringDistribution;
 
 					Print(Context, TEXT(" SSS ID.     "), FontBSDFPropName);	Print(Context, ProfileId, FontBSDFPropValu);	Newline(Context, RectMax);
 					Print(Context, TEXT(" SSS Radius  "), FontBSDFPropName);	Print(Context, RadiusScale, FontBSDFPropValu);	Newline(Context, RectMax);
+					Print(Context, TEXT(" SSS MFP     "), FontBSDFPropName);	Print(Context, DiffuseMFP, FontBSDFPropValu);	Newline(Context, RectMax);
+					Print(Context, TEXT(" SSS Phase   "), FontBSDFPropName);	Print(Context, Phase, FontBSDFPropValu);		Newline(Context, RectMax);
 				}
 				else
 				{
 					Print(Context, TEXT(" SSS MFP     "), FontBSDFPropName);	Print(Context, SLAB_SSSMFP(BSDF), FontBSDFPropValu);				Newline(Context, RectMax);
+					Print(Context, TEXT(" SSS Phase   "), FontBSDFPropName);	Print(Context, SLAB_SSSPHASEANISOTROPY(BSDF), FontBSDFPropValu);		Newline(Context, RectMax);
 				}
-
-				Print(Context, TEXT(" SSS Phase   "), FontBSDFPropName);	Print(Context, SLAB_SSSPHASEANISOTROPY(BSDF), FontBSDFPropValu);		Newline(Context, RectMax);
 			} // SSS
 
 			if (BSDF_GETHASFUZZ(BSDF))
