Proper fix for globallycoherent qualifiers being used in FGlobalWorkQueue, and also some possibly uninitialized data errors on some platforms with tessellation tasks.

diff --git a/Engine/Shaders/Private/WorkDistribution.ush b/Engine/Shaders/Private/WorkDistribution.ush
--- a/Engine/Shaders/Private/WorkDistribution.ush
+++ b/Engine/Shaders/Private/WorkDistribution.ush
@@ -352,8 +352,8 @@ struct FInputQueue
 
 struct FGlobalWorkQueue
 {
-	RWByteAddressBuffer						DataBuffer;
-	RWStructuredBuffer< FWorkQueueState >	StateBuffer;	// Ideally this was GDS but we don't have that API control.
+	RWCoherentByteAddressBufferRef						DataBuffer;
+	RWCoherentStructuredBufferRef( FWorkQueueState )	StateBuffer;	// Ideally this was GDS but we don't have that API control.
 
 	uint	StateIndex;
 	uint	Size;
@@ -469,7 +469,7 @@ void GlobalTaskLoopVariable( FGlobalWorkQueue GlobalWorkQueue )
 				break;
 		}
 
-		FTask Task;
+		FTask Task = (FTask)0;
 		bool bTaskReady = false;
 		if( !bTaskComplete )
 		{
diff --git a/Engine/Shaders/Public/Platform.ush b/Engine/Shaders/Public/Platform.ush
--- a/Engine/Shaders/Public/Platform.ush
+++ b/Engine/Shaders/Public/Platform.ush
@@ -697,10 +697,20 @@ void ShaderYield()
 #define RWCoherentStructuredBuffer(TYPE) globallycoherent RWStructuredBuffer<TYPE>
 #endif
 
+// Drops globallycoherent qualifier when used in a struct
+#ifndef RWCoherentStructuredBufferRef
+#define RWCoherentStructuredBufferRef(TYPE) RWStructuredBuffer<TYPE>
+#endif
+
 #ifndef RWCoherentByteAddressBuffer
 #define RWCoherentByteAddressBuffer globallycoherent RWByteAddressBuffer
 #endif
 
+// Drops globallycoherent qualifier when used in a struct
+#ifndef RWCoherentByteAddressBufferRef
+#define RWCoherentByteAddressBufferRef RWByteAddressBuffer
+#endif
+
 #ifndef ISOLATE
 #define ISOLATE
 #endif
