Initialize SSS anisotropy factor for SSS profile with Substrate.

diff --git a/Engine/Shaders/Private/PathTracing/PathTracingMaterialHitShader.usf b/Engine/Shaders/Private/PathTracing/PathTracingMaterialHitShader.usf
--- a/Engine/Shaders/Private/PathTracing/PathTracingMaterialHitShader.usf
+++ b/Engine/Shaders/Private/PathTracing/PathTracingMaterialHitShader.usf
@@ -750,7 +750,7 @@ RAY_TRACING_ENTRY_CLOSEST_HIT(PathTracingMaterialCHS,
 						const uint ProfileId = StrataSubsurfaceProfileIdTo8bits(SLAB_SSSPROFILEID(BSDF));
 						const float Scale = SLAB_SSSPROFILERADIUSSCALE(BSDF);
 						Payload.MeanFreePath = DecodeSSSProfileRadius(ProfileId, Payload.DiffuseColor, Scale);
-						Payload.PhaseG = 0.0;
+						Payload.PhaseG = DecodeSSSProfileScatteringDistribution(ProfileId);
 					}
 					else
 					{
