Conditionally skip 0 CS writes to gbuffer (off by default for now)

diff --git a/Engine/Shaders/Private/ComputeShaderOutputCommon.ush b/Engine/Shaders/Private/ComputeShaderOutputCommon.ush
--- a/Engine/Shaders/Private/ComputeShaderOutputCommon.ush
+++ b/Engine/Shaders/Private/ComputeShaderOutputCommon.ush
@@ -130,6 +130,45 @@ FPixelShaderOut ShadePixel(const uint2 PixelPos, uint QuadIndex)
 	return PixelShaderOut;
 }
 
+#define CONDITIONAL_EXPORT 0
+
+#if STRATA_OPAQUE_DEFERRED
+
+void ConditionalExport(uint2 PixelPos, FPixelShaderOut ShadedPixel, uint Index)
+{
+#if CONDITIONAL_EXPORT
+	BRANCH
+	if (ShadedPixel.StrataOutput[Index] != 0)
+#endif
+	{
+		OutTargets[uint3(PixelPos, Index)] = ShadedPixel.StrataOutput[Index];
+	}
+}
+
+void ConditionalExport(RWTexture2D<STRATA_TOP_LAYER_TYPE> OutTarget, uint2 PixelPos, STRATA_TOP_LAYER_TYPE TopLayerData)
+{
+#if CONDITIONAL_EXPORT
+	BRANCH
+	if (TopLayerData != 0)
+#endif
+	{
+		OutTarget[PixelPos] = TopLayerData;
+	}
+}
+
+#endif
+
+void ConditionalExport(RWTexture2D<float4> OutTarget, uint2 PixelPos, float4 Value)
+{
+#if CONDITIONAL_EXPORT
+	BRANCH
+	if(any(Value != 0.0f))
+#endif
+	{
+		OutTarget[PixelPos] = Value;
+	}
+}
+
 void ExportPixel(const uint2 PixelPos, FPixelShaderOut ShadedPixel)
 {
 #if PIXELSHADEROUTPUT_COVERAGE || PIXELSHADEROUTPUT_A2C
@@ -141,19 +180,19 @@ void ExportPixel(const uint2 PixelPos, FPixelShaderOut ShadedPixel)
 #endif 
 
 #if PIXELSHADEROUTPUT_MRT0
-	OutTarget0[PixelPos] = ShadedPixel.MRT[0];
+	ConditionalExport(OutTarget0, PixelPos, ShadedPixel.MRT[0]);
 #endif
 
 #if PIXELSHADEROUTPUT_MRT1
-	OutTarget1[PixelPos] = ShadedPixel.MRT[1];
+	ConditionalExport(OutTarget1, PixelPos, ShadedPixel.MRT[1]);
 #endif
 
 #if PIXELSHADEROUTPUT_MRT2
-	OutTarget2[PixelPos] = ShadedPixel.MRT[2];
+	ConditionalExport(OutTarget2, PixelPos, ShadedPixel.MRT[2]);
 #endif
 
 #if PIXELSHADEROUTPUT_MRT3
-	OutTarget3[PixelPos] = ShadedPixel.MRT[3];
+	ConditionalExport(OutTarget3, PixelPos, ShadedPixel.MRT[3]);
 #endif
 
 #if STRATA_OPAQUE_DEFERRED
@@ -164,44 +203,36 @@ void ExportPixel(const uint2 PixelPos, FPixelShaderOut ShadedPixel)
 	// MRT2 is precomputed shadow if both velocity and prec shadow are enabled.
 	// After, Substrate top layer data appended. Remaining Substrate outputs are in the 2d array UAV
 
+	ConditionalExport(PixelPos, ShadedPixel, 0);
+	ConditionalExport(PixelPos, ShadedPixel, 1);
+	ConditionalExport(PixelPos, ShadedPixel, 2);
+
 #if PIXELSHADEROUTPUT_MRT3
-	OutTargets[uint3(PixelPos, 0)]	= ShadedPixel.StrataOutput[0];
-	OutTargets[uint3(PixelPos, 1)]	= ShadedPixel.StrataOutput[1];
-	OutTargets[uint3(PixelPos, 2)]	= ShadedPixel.StrataOutput[2];
-	OutTarget4[PixelPos]			= ShadedPixel.StrataTopLayerData;
+	ConditionalExport(OutTarget4, PixelPos, ShadedPixel.StrataTopLayerData);
 #elif PIXELSHADEROUTPUT_MRT2
-	OutTargets[uint3(PixelPos, 0)]	= ShadedPixel.StrataOutput[0];
-	OutTargets[uint3(PixelPos, 1)]	= ShadedPixel.StrataOutput[1];
-	OutTargets[uint3(PixelPos, 2)]	= ShadedPixel.StrataOutput[2];
-	OutTarget3[PixelPos]			= ShadedPixel.StrataTopLayerData;
+	ConditionalExport(OutTarget3, PixelPos, ShadedPixel.StrataTopLayerData);
 #elif PIXELSHADEROUTPUT_MRT1
-	OutTargets[uint3(PixelPos, 0)]	= ShadedPixel.StrataOutput[0];
-	OutTargets[uint3(PixelPos, 1)]	= ShadedPixel.StrataOutput[1];
-	OutTargets[uint3(PixelPos, 2)]	= ShadedPixel.StrataOutput[2];
-	OutTarget2[PixelPos]			= ShadedPixel.StrataTopLayerData;
+	ConditionalExport(OutTarget2, PixelPos, ShadedPixel.StrataTopLayerData);
 #else
-	OutTargets[uint3(PixelPos, 0)]	= ShadedPixel.StrataOutput[0];
-	OutTargets[uint3(PixelPos, 1)]	= ShadedPixel.StrataOutput[1];
-	OutTargets[uint3(PixelPos, 2)]	= ShadedPixel.StrataOutput[2];
-	OutTarget1[PixelPos]			= ShadedPixel.StrataTopLayerData;
+	ConditionalExport(OutTarget1, PixelPos, ShadedPixel.StrataTopLayerData);
 #endif
 
 #else // STRATA_OPAQUE_DEFERRED
 
 #if PIXELSHADEROUTPUT_MRT4
-	OutTarget4[PixelPos] = ShadedPixel.MRT[4];
+	ConditionalExport(OutTarget4, PixelPos, ShadedPixel.MRT[4]);
 #endif
 
 #if PIXELSHADEROUTPUT_MRT5
-	OutTarget5[PixelPos] = ShadedPixel.MRT[5];
+	ConditionalExport(OutTarget5, PixelPos, ShadedPixel.MRT[5]);
 #endif
 
 #if PIXELSHADEROUTPUT_MRT6
-	OutTarget6[PixelPos] = ShadedPixel.MRT[6];
+	ConditionalExport(OutTarget6, PixelPos, ShadedPixel.MRT[6]);
 #endif
 
 #if PIXELSHADEROUTPUT_MRT7
-	OutTarget7[PixelPos] = ShadedPixel.MRT[7];
+	ConditionalExport(OutTarget7, PixelPos, ShadedPixel.MRT[7]);
 #endif
 
 #endif // STRATA_OPAQUE_DEFERRED
