Moved EyeAdaptation code to a common function
- Both EyeAdaptationCS and BasicEyeAdaptationCS call EyeAdaptationCommon

diff --git a/Engine/Shaders/Private/PostProcessEyeAdaptation.usf b/Engine/Shaders/Private/PostProcessEyeAdaptation.usf
--- a/Engine/Shaders/Private/PostProcessEyeAdaptation.usf
+++ b/Engine/Shaders/Private/PostProcessEyeAdaptation.usf
@@ -115,20 +115,13 @@ float2 ComputeWeightedTextureAverageAlphaSubRegion(
 	return float2(Value, max(WeightTotal,0.0000000001f));
 }
 
-#if COMPUTESHADER
-[numthreads(1, 1, 1)]
-void EyeAdaptationCS(uint2 DispatchThreadId : SV_DispatchThreadID)
+float4 EyeAdaptationCommon(float AverageSceneLuminance, float OldExposureScale)
 {
-	float4 OutColor = 0;
-
-	const float AverageSceneLuminance = ComputeEyeAdaptationExposure(HistogramTexture);
-
 	const float TargetAverageLuminance = clamp(AverageSceneLuminance, EyeAdaptation_MinAverageLuminance, EyeAdaptation_MaxAverageLuminance);
-	
+
 	// White point luminance is target luminance divided by 0.18 (18% grey).
 	const float TargetExposure = TargetAverageLuminance / 0.18;
 
-	const float OldExposureScale = HistogramTexture.Load(int3(0, 1, 0)).x;
 	const float MiddleGreyExposureCompensation = EyeAdaptation_ExposureCompensationSettings * EyeAdaptation_ExposureCompensationCurve; // we want the average luminance remapped to 0.18, not 1.0
 	const float OldExposure = MiddleGreyExposureCompensation / (OldExposureScale != 0 ? OldExposureScale : 1.0f);
 
@@ -136,17 +129,32 @@ void EyeAdaptationCS(uint2 DispatchThreadId : SV_DispatchThreadID)
 	const float EstimatedExposure = ComputeEyeAdaptation(OldExposure, TargetExposure, EyeAdaptation_DeltaWorldTime);
 
 	// maybe make this an option to avoid hard clamping when transitioning between different exposure volumes?
-	const float SmoothedExposure = clamp(EstimatedExposure, EyeAdaptation_MinAverageLuminance/.18f, EyeAdaptation_MaxAverageLuminance/.18f);
+	const float SmoothedExposure = clamp(EstimatedExposure, EyeAdaptation_MinAverageLuminance / .18f, EyeAdaptation_MaxAverageLuminance / .18f);
 
 	const float SmoothedExposureScale = 1.0f / max(0.0001f, SmoothedExposure);
-	const float TargetExposureScale =   1.0f / max(0.0001f, TargetExposure);
+	const float TargetExposureScale = 1.0f / max(0.0001f, TargetExposure);
 
-	OutColor.x = MiddleGreyExposureCompensation * SmoothedExposureScale;
-	OutColor.y = MiddleGreyExposureCompensation * TargetExposureScale;
-	OutColor.z = AverageSceneLuminance;
-	OutColor.w = MiddleGreyExposureCompensation;
-	
-	RWEyeAdaptationBuffer[0] = OutColor;
+	float4 OutData;
+	// Output the number that will rescale the image intensity
+	OutData.x = MiddleGreyExposureCompensation * SmoothedExposureScale;
+	// Output the target value
+	OutData.y = MiddleGreyExposureCompensation * TargetExposureScale;
+	OutData.z = AverageSceneLuminance;
+	OutData.w = MiddleGreyExposureCompensation;
+
+	return OutData;
+}
+
+#if COMPUTESHADER
+[numthreads(1, 1, 1)]
+void EyeAdaptationCS(uint2 DispatchThreadId : SV_DispatchThreadID)
+{
+	const float AverageSceneLuminance = ComputeEyeAdaptationExposure(HistogramTexture);
+	const float OldExposureScale = HistogramTexture.Load(int3(0, 1, 0)).x;
+
+	const float4 OutData0 = EyeAdaptationCommon(AverageSceneLuminance, OldExposureScale);
+
+	RWEyeAdaptationBuffer[0] = OutData0;
 }
 
 RWTexture2D<float4> RWEyeAdaptationTexture;
@@ -179,94 +187,78 @@ void BasicEyeAdaptationSetupPS(
 
 groupshared float2 SubRectValueWeight[TGSIZE*TGSIZE];
 
-[numthreads(TGSIZE, TGSIZE, 1)]
-void BasicEyeAdaptationCS(uint GIndex : SV_GroupIndex, uint2 GTId : SV_GroupThreadID)
+float2 Reduce(uint GIndex, float2 Value)
 {
-	float4 OutColor = 0;
-	
-	// Compute scaled Log Luminance Average
-
-	// There are TGSIZE*TGSIZE threads. Each thread will calculate the luminance for its own subregions from a TGSIZE*TGSIZE screen grid
-	const uint2 SubRectMin = uint2(
-		((Color_ViewportMax.x - Color_ViewportMin.x) * GTId.x) / TGSIZE,
-		((Color_ViewportMax.y - Color_ViewportMin.y) * GTId.y) / TGSIZE);
-
-	const uint2 SubRectMax = uint2(
-		((Color_ViewportMax.x - Color_ViewportMin.x) * (GTId.x + 1)) / TGSIZE,
-		((Color_ViewportMax.y - Color_ViewportMin.y) * (GTId.y + 1)) / TGSIZE);
-
-	const float2 LogLumAveWithWeight = ComputeWeightedTextureAverageAlphaSubRegion(ColorTexture, SubRectMin, SubRectMax, Color_ViewportMin, Color_ViewportMax);
-
 	// Store in LDS
-	SubRectValueWeight[GIndex] = LogLumAveWithWeight;
+	SubRectValueWeight[GIndex] = Value;
 	GroupMemoryBarrierWithGroupSync();
 
 	// Merge the ValueWeight from all threads
-	SubRectValueWeight[GIndex] =  SubRectValueWeight[GIndex] + SubRectValueWeight[GIndex ^ 1];
+	SubRectValueWeight[GIndex] = SubRectValueWeight[GIndex] + SubRectValueWeight[GIndex ^ 1];
 	GroupMemoryBarrierWithGroupSync();
 
-	SubRectValueWeight[GIndex] =  SubRectValueWeight[GIndex] + SubRectValueWeight[GIndex ^ 2];
+	SubRectValueWeight[GIndex] = SubRectValueWeight[GIndex] + SubRectValueWeight[GIndex ^ 2];
 	GroupMemoryBarrierWithGroupSync();
 
-	SubRectValueWeight[GIndex] =  SubRectValueWeight[GIndex] + SubRectValueWeight[GIndex ^ 4];
+	SubRectValueWeight[GIndex] = SubRectValueWeight[GIndex] + SubRectValueWeight[GIndex ^ 4];
 	GroupMemoryBarrierWithGroupSync();
 
-	SubRectValueWeight[GIndex] =  SubRectValueWeight[GIndex] + SubRectValueWeight[GIndex ^ 8];
+	SubRectValueWeight[GIndex] = SubRectValueWeight[GIndex] + SubRectValueWeight[GIndex ^ 8];
 	GroupMemoryBarrierWithGroupSync();
 
-	SubRectValueWeight[GIndex] =  SubRectValueWeight[GIndex] + SubRectValueWeight[GIndex ^ 16];
+	SubRectValueWeight[GIndex] = SubRectValueWeight[GIndex] + SubRectValueWeight[GIndex ^ 16];
 	GroupMemoryBarrierWithGroupSync();
 
-	SubRectValueWeight[GIndex] =  SubRectValueWeight[GIndex] + SubRectValueWeight[GIndex ^ 32];
+	SubRectValueWeight[GIndex] = SubRectValueWeight[GIndex] + SubRectValueWeight[GIndex ^ 32];
 	GroupMemoryBarrierWithGroupSync();
 
-	SubRectValueWeight[GIndex] =  SubRectValueWeight[GIndex] + SubRectValueWeight[GIndex ^ 64];
+	SubRectValueWeight[GIndex] = SubRectValueWeight[GIndex] + SubRectValueWeight[GIndex ^ 64];
 	GroupMemoryBarrierWithGroupSync();
 
-	SubRectValueWeight[GIndex] =  SubRectValueWeight[GIndex] + SubRectValueWeight[GIndex ^ 128];
+	SubRectValueWeight[GIndex] = SubRectValueWeight[GIndex] + SubRectValueWeight[GIndex ^ 128];
 	GroupMemoryBarrierWithGroupSync();
 
-	float LogLumAve = SubRectValueWeight[0].x /  SubRectValueWeight[0].y;
+	return SubRectValueWeight[0];
+}
+
+[numthreads(TGSIZE, TGSIZE, 1)]
+void BasicEyeAdaptationCS(uint GIndex : SV_GroupIndex, uint2 GTId : SV_GroupThreadID)
+{	
+	// Compute scaled Log Luminance Average
+
+	// There are TGSIZE*TGSIZE threads. Each thread will calculate the luminance for its own subregions from a TGSIZE*TGSIZE screen grid
+	const uint2 SubRectMin = uint2(
+		((Color_ViewportMax.x - Color_ViewportMin.x) * GTId.x) / TGSIZE,
+		((Color_ViewportMax.y - Color_ViewportMin.y) * GTId.y) / TGSIZE);
+
+	const uint2 SubRectMax = uint2(
+		((Color_ViewportMax.x - Color_ViewportMin.x) * (GTId.x + 1)) / TGSIZE,
+		((Color_ViewportMax.y - Color_ViewportMin.y) * (GTId.y + 1)) / TGSIZE);
+
+	const float2 LogLumAveWithWeight = ComputeWeightedTextureAverageAlphaSubRegion(ColorTexture, SubRectMin, SubRectMax, Color_ViewportMin, Color_ViewportMax);
+
+	const float2 ReducedLogLumAveWithWeight = Reduce(GIndex, LogLumAveWithWeight);
+
+	float LogLumAve = ReducedLogLumAveWithWeight.x / ReducedLogLumAveWithWeight.y;
 
 	// Correct for [0,1] scaling
 	LogLumAve = (LogLumAve - EyeAdaptation_HistogramBias) / EyeAdaptation_HistogramScale;
 	
 	// Convert LogLuminanceAverage to Average Intensity
 	const float AverageSceneLuminance = View.OneOverPreExposure * exp2(LogLumAve);
-
-	const float MiddleGreyExposureCompensation = EyeAdaptation_ExposureCompensationSettings * EyeAdaptation_ExposureCompensationCurve * EyeAdaptation_GreyMult;// we want the average luminance remapped to 0.18, not 1.0
-
-	const float LumAve = AverageSceneLuminance; 
-
-	const float ClampedLumAve = clamp(LumAve, EyeAdaptation_MinAverageLuminance, EyeAdaptation_MaxAverageLuminance);
 	
 	// The Exposure Scale (and thus intensity) used in the previous frame
-	const float ExposureScaleOld = EyeAdaptationBuffer[0].x;
-	const float LuminanceAveOld = MiddleGreyExposureCompensation / (ExposureScaleOld != 0 ? ExposureScaleOld : 1.0f);
-
-	// Time-based expoential blend of the intensity to allow the eye adaptation to ramp up over a few frames.
-	const float EstimatedLuminance = ComputeEyeAdaptation(LuminanceAveOld, ClampedLumAve, EyeAdaptation_DeltaWorldTime);
-
-	// maybe make this an option to avoid hard clamping when transitioning between different exposure volumes?
-	const float SmoothedLuminance = clamp(EstimatedLuminance, EyeAdaptation_MinAverageLuminance, EyeAdaptation_MaxAverageLuminance);
+	const float OldExposureScale = EyeAdaptationBuffer[0].x;
 
-	const float SmoothedExposureScale = 1.0f / max(0.0001f, SmoothedLuminance);
-	const float TargetExposureScale   = 1.0f / max(0.0001f, ClampedLumAve);
-	
-	// Output the number that will rescale the image intensity
-	OutColor.x = MiddleGreyExposureCompensation * SmoothedExposureScale;
-	// Output the target value
-	OutColor.y = MiddleGreyExposureCompensation * TargetExposureScale;
-	OutColor.z = AverageSceneLuminance;
-	OutColor.w = MiddleGreyExposureCompensation / EyeAdaptation_GreyMult;
+	const float4 OutData0 = EyeAdaptationCommon(AverageSceneLuminance, OldExposureScale);
 	
 #if XBOXONE_PROFILE
-	OutColor = !all(IsFinite(OutColor)) ? float4(1, 1, 1, 0) : OutColor;
+	OutData0 = !all(IsFinite(OutData0)) ? float4(1, 1, 1, 0) : OutData0;
 #endif
 	
-	if(GIndex==0)
+	if(GIndex == 0)
 	{
-		RWEyeAdaptationBuffer[0] = OutColor;
+		RWEyeAdaptationBuffer[0] = OutData0;
 	}
 }
 
