Path tracer: Gate the logic for using the SSS profile anistropy behind a define

Keep this off for now for backwards compatibility, but keep the code around with a comment so we have some trace of why it was done this way.

diff --git a/Engine/Shaders/Private/PathTracing/PathTracingMaterialHitShader.usf b/Engine/Shaders/Private/PathTracing/PathTracingMaterialHitShader.usf
--- a/Engine/Shaders/Private/PathTracing/PathTracingMaterialHitShader.usf
+++ b/Engine/Shaders/Private/PathTracing/PathTracingMaterialHitShader.usf
@@ -3,6 +3,12 @@
 #define PATH_TRACING 1
 #define ENABLE_SKY_LIGHT 1      
 #define NEEDS_LIGHTMAP_COORDINATE 0
+
+// This define controls if we use the "g" factor from the SSS profile or not. This is disabled by default because
+// we can get a better match to how Metahumans are currently dialed with an isotropic phase function. It may be possible to revisit this
+// if we get a better alignment between raster SSS techniques and the PT random walk.
+#define USE_SSS_PROFILE_ANISOTROPY 0
+
 #ifdef NEEDS_VERTEX_FACTORY_INTERPOLATION
 #undef NEEDS_VERTEX_FACTORY_INTERPOLATION
 #endif
@@ -750,7 +756,12 @@ RAY_TRACING_ENTRY_CLOSEST_HIT(PathTracingMaterialCHS,
 						const uint ProfileId = StrataSubsurfaceProfileIdTo8bits(SLAB_SSSPROFILEID(BSDF));
 						const float Scale = SLAB_SSSPROFILERADIUSSCALE(BSDF);
 						Payload.MeanFreePath = DecodeSSSProfileRadius(ProfileId, Payload.DiffuseColor, Scale);
+#if USE_SSS_PROFILE_ANISOTROPY
 						Payload.PhaseG = DecodeSSSProfileScatteringDistribution(ProfileId);
+						Payload.MeanFreePath *= 1 - G; // Compensate the extra translucency by scaling down the MFP
+#else
+						Payload.PhaseG = 0.0;
+#endif
 					}
 					else
 					{
@@ -1017,6 +1028,10 @@ RAY_TRACING_ENTRY_CLOSEST_HIT(PathTracingMaterialCHS,
 #if HAVE_GetSubsurfaceMediumMaterialOutput1
 		float ScatteringDistribution = clamp(GetSubsurfaceMediumMaterialOutput1(MaterialParameters), -0.99f, 0.99f);
 		Payload.SetSubsurfacePhaseFunction(ScatteringDistribution);
+#elif USE_SSS_PROFILE_ANISOTROPY
+		float G = DecodeSSSProfileScatteringDistribution(ProfileId);
+		Payload.SetSubsurfacePhaseFunction(G);
+		Payload.SetSubsurfaceRadius(Payload.GetSubsurfaceRadius() * (1 - G));
 #else
 		Payload.SetSubsurfacePhaseFunction(0.0);
 #endif
