Path Tracer: Improve handling of Substrate "simple volume" layers

The raster implementation remaps diffuse color as a function of MFP to simulate the effect an actual slab of scattering medium would have. Correct the diffuse color according to this logic.

Also fix a corner case that was leading to compiler errors (unlit materials with refraction enabled).

diff --git a/Engine/Shaders/Private/PathTracing/PathTracingMaterialHitShader.usf b/Engine/Shaders/Private/PathTracing/PathTracingMaterialHitShader.usf
--- a/Engine/Shaders/Private/PathTracing/PathTracingMaterialHitShader.usf
+++ b/Engine/Shaders/Private/PathTracing/PathTracingMaterialHitShader.usf
@@ -142,6 +142,20 @@ float DecodeSSSProfileScatteringDistribution(uint ProfileId)
 	return DecodeScatteringDistribution(EncodedScatteringDistribution);
 }
 
+#if STRATA_ENABLED
+
+float3 GetSimpleVolumeDiffuseColor(float3 DiffuseColor, float3 MeanFreePath)
+{
+	// See reference in StrataEvaluation.ush
+	// STRATA_TODO: In the case of thin materials, we should also include a backscattering diffuse portion
+	FParticipatingMedia PM = StrataSlabCreateParticipatingMedia(DiffuseColor, MeanFreePath);
+	const float DiffuseToVolumeBlend = StrataSlabDiffuseToVolumeBlend(PM);
+	const float3 SlabDirectionalAlbedo = IsotropicMediumSlabEnvDirectionalAlbedo(PM);
+	return lerp(DiffuseColor, SlabDirectionalAlbedo, DiffuseToVolumeBlend);
+}
+
+#endif
+
 RAY_TRACING_ENTRY_CLOSEST_HIT(PathTracingMaterialCHS,
 	FPackedPathTracingPayload, PackedPayload,
 	FRayTracingIntersectionAttributes, Attributes)
@@ -205,7 +219,7 @@ RAY_TRACING_ENTRY_CLOSEST_HIT(PathTracingMaterialCHS,
 	FStrataData StrataData = PixelMaterialInputs.GetFrontStrataData();
 
 #if STRATA_OPTIMIZED_UNLIT
-	const float3 Transparency = UNLIT_TRANSMITTANCE(StrataData.InlinedBSDF);
+	const float3 Transparency = saturate(StrataData.InlinedBSDF.Coverage + (1.0 - StrataData.InlinedBSDF.Coverage) * UNLIT_TRANSMITTANCE(StrataData.InlinedBSDF));
 #else
 	FStrataPixelHeader StrataPixelHeader = MaterialParameters.GetFrontStrataHeader();
 	FStrataIntegrationSettings IntegrationSettings = InitStrataIntegrationSettings();
@@ -242,6 +256,8 @@ RAY_TRACING_ENTRY_CLOSEST_HIT(PathTracingMaterialCHS,
 	{
 
 #if STRATA_ENABLED
+
+#if MATERIAL_IS_STRATA && !STRATA_OPTIMIZED_UNLIT
 		// We already figured out how much Transparency there is above
 		// So only do extra logic for refraction and thin refraction ases
 #if MATERIALBLENDING_TRANSLUCENT && (REFRACTION_USE_INDEX_OF_REFRACTION || MATERIAL_ISTHINSURFACE)
@@ -319,6 +335,7 @@ RAY_TRACING_ENTRY_CLOSEST_HIT(PathTracingMaterialCHS,
 		Transparency = saturate(Transparency);
 #endif // MATERIALBLENDING_TRANSLUCENT && (REFRACTION_USE_INDEX_OF_REFRACTION || MATERIAL_ISTHINSURFACE)
 
+#endif // MATERIAL_IS_STRATA && !STRATA_OPTIMIZED_UNLIT
 
 #else // STRATA_ENABLED
 
@@ -627,6 +644,11 @@ RAY_TRACING_ENTRY_CLOSEST_HIT(PathTracingMaterialCHS,
 									Spec1E = ComputeGGXSpecEnergyTermsRGB(Roughness1, NoV, F0, F90).E;
 								}
 							}
+							if (BSDF_GETISSIMPLEVOLUME(BSDF))
+							{
+								DiffuseColor = GetSimpleVolumeDiffuseColor(DiffuseColor, SLAB_SSSMFP(BSDF));
+							}
+
 							float FuzzAmount = 0.0;
 							float FuzzRoughness = 1.0;
 							float3 FuzzColor = 0.0;
@@ -746,7 +768,11 @@ RAY_TRACING_ENTRY_CLOSEST_HIT(PathTracingMaterialCHS,
 						Payload.FuzzColor     = 0.0;
 						Payload.FuzzRoughness = 0.0;
 					}
-					if (BSDF_GETSSSTYPE(BSDF) == SSS_TYPE_DIFFUSION)
+					if (BSDF_GETISSIMPLEVOLUME(BSDF))
+					{
+						Payload.DiffuseColor = GetSimpleVolumeDiffuseColor(Payload.DiffuseColor, SLAB_SSSMFP(BSDF));
+					}
+					else if (BSDF_GETSSSTYPE(BSDF) == SSS_TYPE_DIFFUSION)
 					{
 						Payload.MeanFreePath = SLAB_SSSMFP(BSDF);
 						Payload.PhaseG       = SLAB_SSSPHASEANISOTROPY(BSDF);
