* Fix IsThin support with SSS diffusion and SSS profile diffusion (now it works as expected, and diffusion is no longer scaled incorrectly)
* Fix IsThin support surface for translucent case
* Clarify the behavior for each case in code and add details to material property visualization

diff --git a/Engine/Shaders/Private/BurleyNormalizedSSSCommon.ush b/Engine/Shaders/Private/BurleyNormalizedSSSCommon.ush
--- a/Engine/Shaders/Private/BurleyNormalizedSSSCommon.ush
+++ b/Engine/Shaders/Private/BurleyNormalizedSSSCommon.ush
@@ -191,6 +191,15 @@ float3 InternalGetBurleyTransmissionProfile(
 	return Output;
 }
 
+float3 GetBurleyTransmission(
+	float3 SubsurfaceAlebdo,
+	float3 MeanFreePathInCm,
+	float ThicknessInCm)
+{
+	const float TransmissionMFPScaleFactor = 100.f; // This factor compensate a mistake in SSS profile transmission computation. See BurleyNormalizedSSS.cpp for more details
+	return InternalGetBurleyTransmissionProfile(SubsurfaceAlebdo, MeanFreePathInCm * TransmissionMFPScaleFactor, ThicknessInCm);
+}
+
 float3 GetBurleyTransmissionProfile(
 	float3 SubsurfaceAlebdo,
 	float3 MeanFreePathInCm,
diff --git a/Engine/Shaders/Private/Strata/Strata.ush b/Engine/Shaders/Private/Strata/Strata.ush
--- a/Engine/Shaders/Private/Strata/Strata.ush
+++ b/Engine/Shaders/Private/Strata/Strata.ush
@@ -169,8 +169,6 @@ struct FStrataBSDF
 
 	void EnableSlabBSDFSimpleVolumetric();
 
-	void EnableSlabBSDFIsThin();
-
 	void StrataSanitizeBSDF();
 
 	void StrataSetBSDFRoughness(in float Roughness);
@@ -214,8 +212,8 @@ struct FStrataBSDF
 //            BSDF(1)       ?? |  Adaptive
 //                             |
 // * LAYOUT 1 -------------------------------------------------------------------------------------------------------------------------------------------------------------------------
-//  (Single)  HEADER        27 |  State | AO | Aniso | TopLayer | SSSType | GreyWeighV | Haze | F90 | Simple vol. | MFP Plugged | Fuzz | Thin
-//                             |  8       8    1       1          2         1            1      1     1             1             1      1            
+//  (Single)  HEADER        27 |  State | AO | Opt Layout | Aniso | TopLayer | SSSType | GreyWeighV | Haze | F90 | Simple vol. | MFP Plugged | Fuzz | Thin
+//                             |  8       8    3            1       1          2         1            1      1     1             1             1      1            
 //            BASIS(0)       0 | (Use top layer normal)
 //            BSDF(0)       ?? |  Adaptive
 // * LAYOUT 0 -------------------------------------------------------------------------------------------------------------------------------------------------------------------------          
@@ -511,6 +509,7 @@ struct FStrataBSDF
 // Either with SSS profile
 #define SLAB_SSSPROFILEID(X)			X.VGPRs[3].x
 #define SLAB_SSSPROFILERADIUSSCALE(X)	X.VGPRs[3].y
+#define SLAB_SSSPROFILETHICKNESSCM(X)	X.VGPRs[3].z
 // Or explicit MFP in centimeter
 #define SLAB_SSSMFP(X)					X.VGPRs[3].xyz
 #define SLAB_SSSPHASEANISOTROPY(X)		X.VGPRs[3].w
@@ -1006,6 +1005,25 @@ uint StrataSubsurfaceProfileIdTo8bits(float In)
 	return uint(In * 255.0f + 0.5f);
 }
 
+///////////////////////////////////////////////////////////////////////////////
+// PackSSSProfile
+
+uint PackSSSProfile(float InProfileId, float InRadius, half InThickness)
+{
+	uint Out = 0; 	
+	Out |= (PackR8(InProfileId));
+	Out |= (PackR8(InRadius)<<8);
+	Out |= (Pack10F(InThickness)<<16);
+	return Out;
+}
+
+void UnpackSSSProfile(uint In, inout float OutProfileId, inout float OutRadius, inout float OutThickness)
+{
+	OutProfileId = UnpackR8(In);
+	OutRadius = UnpackR8(In >> 8);
+	OutThickness = Unpack10F((In >> 16) & 0x3FF);
+}
+
 ///////////////////////////////////////////////////////////////////////////////
 // Fuzz
 
@@ -1456,7 +1474,7 @@ FStrataData GetStrataSlabBSDF(
 	const bool bMFPPlugged = any(SSSMFP > 0);
 	bool bIsThin = bMFPPlugged && bIsThinSurface; // Enable thin only if MFP is plugged, as otherwise there will be no back-tramission
 	bool bHasSSSWrap = bMFPPlugged && !UseSSSDiffusion;
-	bool bHasSSSProfile = !bIsThin && bSSSProfilePlugged && SSSMFPScale > 0.f;	// When no profile is provide, we use Burley. STRATA_TODO we should clarify as burley is also an option of the profile.
+	bool bHasSSSProfile = bSSSProfilePlugged && SSSMFPScale > 0.f;
 	bool bHasSSS = bMFPPlugged || bHasSSSProfile || bHasSSSWrap;
 
 	FHaziness Haziness = InitialiseHaziness();
@@ -1559,6 +1577,7 @@ FStrataData GetStrataSlabBSDF(
 	{
 		SLAB_SSSPROFILEID	(StrataData.InlinedBSDF)		= SSSProfileID;
 		SLAB_SSSPROFILERADIUSSCALE(StrataData.InlinedBSDF)	= SSSMFPScale;
+		SLAB_SSSPROFILETHICKNESSCM(StrataData.InlinedBSDF)	= ThicknessCm;
 	}
 	else
 	{
@@ -2357,6 +2376,8 @@ FStrataData StrataHorizontalMixingParameterBlending(FStrataData Background, FStr
 			bForegroundHasSSSProfile ? SLAB_SSSPROFILERADIUSSCALE(ForegroundBSDF) : 0.0f,
 			ForegroundOtherMixFactor);
 
+		SLAB_SSSPROFILETHICKNESSCM(ResultBSDF) = bForegroundHasSSSProfile ? BSDF_GETTHICKNESSCM(ForegroundBSDF) : BSDF_GETTHICKNESSCM(BackgroundBSDF);
+
 		if (SLAB_SSSPROFILERADIUSSCALE(ResultBSDF) == 0.0f)
 		{
 			SLAB_SSSPROFILEID(ResultBSDF) = SSS_PROFILE_ID_INVALID;
@@ -2458,6 +2479,7 @@ FStrataData StrataAddParameterBlending(FStrataData A, FStrataData B, float AMixF
 		{
 			SLAB_SSSPROFILEID(ResultBSDF) = SLAB_SSSPROFILEID(ABSDF);
 			SLAB_SSSPROFILERADIUSSCALE(ResultBSDF) = SLAB_SSSPROFILERADIUSSCALE(ABSDF);
+			SLAB_SSSPROFILETHICKNESSCM(ResultBSDF) = BSDF_GETTHICKNESSCM(ABSDF);
 		}
 		else
 		{
@@ -2598,7 +2620,7 @@ FStrataBSDF StrataVerticalLayeringParameterBlending(in FStrataBSDF TopBSDF, in F
 	const float3 NewMFP = Info2.TransmittanceOnlyBottom * BaseBSDF.TmpMFP + Info2.TransmittanceOnlyTop * TopBSDF.TmpMFP + Info2.TransmittanceTopAndBottom * min(TopBSDF.TmpMFP, BaseBSDF.TmpMFP);
 #endif
 
-	ResultBSDF.TmpMFP					= NewMFP;
+	ResultBSDF.TmpMFP				= NewMFP;
 	SLAB_SSSMFP(ResultBSDF)			= NewMFP;
 
 	ResultBSDF.Coverage				= MaxCoverage;
@@ -2809,11 +2831,36 @@ void FStrataBSDF::StrataSanitizeBSDF()
 	}
 }
 
+// Description of SSS conversions
+//
+// Opaque
+// ======
+//               |                     Bottom                      |         !Bottom          |
+//               |          Thin          |          Thick         |                          |
+//               |----------------------------------------------------------------------------|
+// SSS Diffusion |        SSS / MFP       |       SSS / MFP        |  SimpleVolume / Res.MFP  | 
+// SSS Profile   |        SSS / MFP       |       SSS / MFP        |  SimpleVolume / Res.MFP  | 
+// SSS Wrap      |      SSS / Res.MFP     |     SSS / Res.MFP      |  SimpleVolume / Res.MFP  | 
+// SSS Invalid   |                        |                        |                          | 
+//
+// Translucent
+// ===========
+//               |                     Bottom                      |         !Bottom          |
+//               |          Thin          |          Thick         |                          |
+//               |----------------------------------------------------------------------------|
+// SSS Diffusion | SimpleVolume / Res.MFP | SimpleVolume / Res.MFP |  SimpleVolume / Res.MFP  | 
+// SSS Profile   | SimpleVolume / Res.MFP | SimpleVolume / Res.MFP |  SimpleVolume / Res.MFP  | 
+// SSS Wrap      | SimpleVolume / Res.MFP | SimpleVolume / Res.MFP |  SimpleVolume / Res.MFP  | 
+// SSS Invalid   |                        |                        |                          | 
+
+//  Res.MFP : Rescaled MFP in computation space i.e., MFP is rescaled for a thickness of 1m
+
+// Convert to simple volume to get some transmittance
 void FStrataBSDF::PostProcessSlabBeforeLighting(bool bIsStrataOpaqueMaterial)
 {
 	// Two sided lighting can only be evaluated on the bottom layer, together with option SSS. Otherwise, the simple volume will be used.
 	// At the moment, a material can't be thin & simple volume at the same time. This should be fixed by unifying SSSDiffusion|SSSWrap|SimpleVolume into a single evaluation function
-	const bool bIsThin = this.bIsBottom && BSDF_GETISTHIN(this);
+	const bool bIsThin = this.bIsBottom && BSDF_GETISTHIN(this) && bIsStrataOpaqueMaterial;
 
 	// We want simple volumetric only if the MFP input is pluged in, otherwise we avoid dealing with simple volumetric.
 	// We also do not want to use the simple volumetric path if we are dealing with an opaque material and this slab is not at the bottom.
@@ -2822,12 +2869,22 @@ void FStrataBSDF::PostProcessSlabBeforeLighting(bool bIsStrataOpaqueMaterial)
 
 	if (bIsSimpleVolume)
 	{
+		// * Rescale MFP to computation space
+		// * Set Invalid SSS type
+		// * Set IsSimpleVolume
 		this.EnableSlabBSDFSimpleVolumetric();
 	}
 
 	if (bIsThin)
 	{
-		this.EnableSlabBSDFIsThin();
+		BSDF_SETISSIMPLEVOLUME(this, 0);
+
+		// Only rescale MFP to computation space for SSS wrap. 
+		// SSS Diffuse types use their 'original' MFP for diffusion
+		if (BSDF_GETSSSTYPE(this) == SSS_TYPE_WRAP)
+		{
+			this.EnableSlabMFPComputationSpace();
+		}
 	}
 	else
 	{
@@ -2932,13 +2989,6 @@ void FStrataBSDF::EnableSlabBSDFSimpleVolumetric()
 	BSDF_SETISSIMPLEVOLUME(this, 1);
 }
 
-void FStrataBSDF::EnableSlabBSDFIsThin()
-{
-	// BSDF_SETISTHIN is already set when creating the Slab.
-	BSDF_SETISSIMPLEVOLUME(this, 0);
-	this.EnableSlabMFPComputationSpace();
-}
-
 ///////////////////////////////////////////////////////////////////////////////
 // Generic BSDF accessors
 
@@ -4167,9 +4217,8 @@ FStrataBSDF UnpackStrataBSDFIn(FStrataMaterialContainer StrataBuffer, inout FStr
 		{
 			if (BSDF_GETSSSTYPE(OutBSDF) == SSS_TYPE_DIFFUSION_PROFILE && !bIsSimpleVolume)
 			{
-				float4 Data = UnpackRGBA8(StrataLoadUint1(StrataBuffer, StrataAddressing));
-				SLAB_SSSPROFILEID(OutBSDF) = Data.r;
-				SLAB_SSSPROFILERADIUSSCALE(OutBSDF) = Data.g;
+				uint RawData = StrataLoadUint1(StrataBuffer, StrataAddressing);
+				UnpackSSSProfile(RawData, SLAB_SSSPROFILEID(OutBSDF), SLAB_SSSPROFILERADIUSSCALE(OutBSDF), SLAB_SSSPROFILETHICKNESSCM(OutBSDF));
 			}
 			else
 			{
diff --git a/Engine/Shaders/Private/Strata/StrataEvaluation.ush b/Engine/Shaders/Private/Strata/StrataEvaluation.ush
--- a/Engine/Shaders/Private/Strata/StrataEvaluation.ush
+++ b/Engine/Shaders/Private/Strata/StrataEvaluation.ush
@@ -809,17 +809,18 @@ FStrataEvaluateResult StrataEvaluateBSDFCommon(FStrataBSDFContext BSDFContext, F
 			////
 			//// Evaluate transmitted light through a mesh due to sub surface scattering.
 			////
-
 			#if STRATA_FASTPATH==0 && STRATA_SSS_TRANSMISSION
 			if (BSDF_GETSSSTYPE(BSDFContext.BSDF) == SSS_TYPE_DIFFUSION_PROFILE || BSDF_GETSSSTYPE(BSDFContext.BSDF) == SSS_TYPE_DIFFUSION)
 			{
-				const float ThicknessInCm = DecodeThickness(ShadowTerms.TransmissionThickness) * SSSS_MAX_TRANSMISSION_PROFILE_DISTANCE;
+				float ThicknessInCm = DecodeThickness(ShadowTerms.TransmissionThickness) * SSSS_MAX_TRANSMISSION_PROFILE_DISTANCE;
 
 				float OneOverIOR = 1.0f;
 				float PhaseFunctionAnisotropy = 0.0f;
 				float3 TransmissionThroughput = 1.0f;
 				if (BSDF_GETSSSTYPE(BSDFContext.BSDF) == SSS_TYPE_DIFFUSION_PROFILE)
 				{
+					ThicknessInCm = BSDF_GETISTHIN(BSDFContext.BSDF) ? SLAB_SSSPROFILETHICKNESSCM(BSDFContext.BSDF) : ThicknessInCm;
+
 					const uint ProfileId = StrataSubsurfaceProfileIdTo8bits(SLAB_SSSPROFILEID(BSDFContext.BSDF)); // TODO move this into the PackStrataOut( function, to avoid this decode here
 					const FTransmissionProfileParams TransmissionParams = GetTransmissionProfileParams(ProfileId);
 					TransmissionThroughput	= GetTransmissionProfile(ProfileId, ThicknessInCm).rgb;
@@ -828,9 +829,13 @@ FStrataEvaluateResult StrataEvaluateBSDFCommon(FStrataBSDFContext BSDFContext, F
 				}
 				else
 				{
-					const float3 SubsurfaceAlebdo = SLAB_DIFFUSEALBEDO(BSDFContext.BSDF);
+					// When the surface 'IsThin' MFP is rescaled to STRATA_SIMPLEVOLUME_THICKNESS_CM
+					ThicknessInCm = BSDF_GETISTHIN(BSDFContext.BSDF) ? STRATA_SIMPLEVOLUME_THICKNESS_CM : ThicknessInCm;
+
 					const float3 MeanFreePathInCm = SLAB_SSSMFP(BSDFContext.BSDF);
-					TransmissionThroughput	= GetBurleyTransmissionProfile(SubsurfaceAlebdo, MeanFreePathInCm, ThicknessInCm);
+					const float3 SubsurfaceAlebdo = SLAB_DIFFUSEALBEDO(BSDFContext.BSDF);
+					TransmissionThroughput = GetBurleyTransmission(SubsurfaceAlebdo, MeanFreePathInCm, ThicknessInCm).xyz;
+
 					PhaseFunctionAnisotropy = SLAB_SSSPHASEANISOTROPY(BSDFContext.BSDF);
 					OneOverIOR				= 1.f / DielectricF0ToIor(F0.y);
 				}
diff --git a/Engine/Shaders/Private/Strata/StrataExport.ush b/Engine/Shaders/Private/Strata/StrataExport.ush
--- a/Engine/Shaders/Private/Strata/StrataExport.ush
+++ b/Engine/Shaders/Private/Strata/StrataExport.ush
@@ -371,7 +371,7 @@ void PackStrataOut(
 					OptimisedLegacyMode = SINGLE_OPTLEGACYMODE_SSSWRAP;
 					SSSData.Header.Bytes = 0; // Force to skip SSSData writing
 				}
-				else if ((OptimisedModeMask & SingleIsSSSProfile) == SingleIsSSSProfile)
+				else if ((OptimisedModeMask & (SingleIsSSSProfile|SingleIsThin)) == SingleIsSSSProfile) // Force to skip SINGLE_OPTLEGACYMODE_SSSPROFILE when set to ThinSurface
 				{
 					// SingleIsSSSProfile appears critically after SingleIsWrapThin because thin will disable subsurface profile.
 					OptimisedLegacyMode = SINGLE_OPTLEGACYMODE_SSSPROFILE;
@@ -764,10 +764,19 @@ void PackStrataOut(
 								if (BSDF_GETSSSTYPE(BSDF) == SSS_TYPE_DIFFUSION_PROFILE && !bIsSimpleVolume)
 								{
 									// Simple volume overrides SSS profile because it takes over for when not at the bottom of the BSDF layer, or during forward rendering.
-									STRATA_STORE_UINT1(PackRGBA8(float4(SLAB_SSSPROFILEID(BSDF), SLAB_SSSPROFILERADIUSSCALE(BSDF), 0, 0)));
+									// Note: Since it is a profile, we can't rescale its MFP value. So we store the 'thin' thickness value for computing the correct transmission during evaluation
+									STRATA_STORE_UINT1(PackSSSProfile(SLAB_SSSPROFILEID(BSDF), SLAB_SSSPROFILERADIUSSCALE(BSDF), BSDF_GETTHICKNESSCM(BSDF)));
 								}
 								else
 								{
+									// If using non-profile SSS diffusion, rescale SSS so that we take into account the slab thickness. 
+									// Note: This is done 'after' storing the SSSHeader/SSSData, as the original MFP needs to be preserved for correct post-process diffusion.
+									//       This rescaled MFP value is only used for transmission evaluation, not the diffusion
+									if (BSDF_GETISTHIN(BSDF) && BSDF_GETSSSTYPE(BSDF) == SSS_TYPE_DIFFUSION)
+									{
+										SLAB_SSSMFP(BSDF)= RescaleMFPToComputationSpace(SLAB_SSSMFP(BSDF), BSDF_GETTHICKNESSCM(BSDF), STRATA_SIMPLEVOLUME_THICKNESS_CM);
+									}
+
 									// Path used for bottom most layer SSS, simple volume and two sided lighting.
 									STRATA_STORE_UINT1(PackR11G11B10F(SLAB_SSSMFP(BSDF)));
 								}
diff --git a/Engine/Shaders/Private/Strata/StrataVisualize.usf b/Engine/Shaders/Private/Strata/StrataVisualize.usf
--- a/Engine/Shaders/Private/Strata/StrataVisualize.usf
+++ b/Engine/Shaders/Private/Strata/StrataVisualize.usf
@@ -268,30 +268,58 @@ void PrintBSDF(inout FShaderPrintContext Context, inout FStrataAddressing Strata
 
 			if (BSDF_GETSSSTYPE(BSDF) != SSS_TYPE_INVALID || BSDF_GETISSIMPLEVOLUME(BSDF) || BSDF_GETISTHIN(BSDF))
 			{
+				const bool bIsThin = BSDF_GETISTHIN(BSDF);
 				const FStrataSubsurfaceHeader SSSHeader = StrataLoadSubsurfaceHeader(Strata.MaterialTextureArray, Strata.FirstSliceStoringStrataSSSData, StrataAddressing.PixelCoords);
 				Print(Context, TEXT(" SSS Type    "), FontBSDFPropName);	PrintSSSType(Context, BSDF_GETSSSTYPE(BSDF), BSDF_GETISSIMPLEVOLUME(BSDF), BSDF_GETISTHIN(BSDF), FontYellow);	Newline(Context, RectMax);
 				if (BSDF_GETSSSTYPE(BSDF) == SSS_TYPE_WRAP)
 				{
 					const float Opacity = StrataSubSurfaceHeaderGetWrapOpacity(SSSHeader);
-					Print(Context, TEXT(" SSS Opacity "), FontBSDFPropName);	Print(Context, Opacity, FontBSDFPropValu);	Newline(Context, RectMax);
-					Print(Context, TEXT(" SSS Phase   "), FontBSDFPropName);	Print(Context, SLAB_SSSPHASEANISOTROPY(BSDF), FontBSDFPropValu);		Newline(Context, RectMax);
+					Print(Context, TEXT(" SSS Opacity "), FontBSDFPropName);	Print(Context, Opacity, FontBSDFPropValu);						Newline(Context, RectMax);
+					Print(Context, TEXT(" SSS MFP     "), FontBSDFPropName);	Print(Context, SLAB_SSSMFP(BSDF), FontBSDFPropValu);			Newline(Context, RectMax);
+					Print(Context, TEXT(" SSS Phase   "), FontBSDFPropName);	Print(Context, SLAB_SSSPHASEANISOTROPY(BSDF), FontBSDFPropValu);Newline(Context, RectMax);
+					Print(Context, TEXT(" SSS Thicknes"), FontBSDFPropName);	Print(Context, BSDF_GETTHICKNESSCM(BSDF), FontBSDFPropValu);	Newline(Context, RectMax);
 				}
-				if (BSDF_GETSSSTYPE(BSDF) == SSS_TYPE_DIFFUSION_PROFILE && BSDF_GETISSIMPLEVOLUME(BSDF) == 0)
+				else if (BSDF_GETISSIMPLEVOLUME(BSDF))
+				{
+					const float Opacity = StrataSubSurfaceHeaderGetWrapOpacity(SSSHeader);
+					Print(Context, TEXT(" SSS Opacity "), FontBSDFPropName);	Print(Context, Opacity, FontBSDFPropValu);						Newline(Context, RectMax);
+					Print(Context, TEXT(" SSS MFP     "), FontBSDFPropName);	Print(Context, SLAB_SSSMFP(BSDF), FontBSDFPropValu);			Newline(Context, RectMax);
+					Print(Context, TEXT(" SSS Phase   "), FontBSDFPropName);	Print(Context, SLAB_SSSPHASEANISOTROPY(BSDF), FontBSDFPropValu);Newline(Context, RectMax);
+					Print(Context, TEXT(" SSS Thicknes"), FontBSDFPropName);	Print(Context, BSDF_GETTHICKNESSCM(BSDF), FontBSDFPropValu);	Newline(Context, RectMax);
+				}
+				else if (BSDF_GETSSSTYPE(BSDF) == SSS_TYPE_DIFFUSION_PROFILE)
 				{
 					const uint ProfileId = StrataSubSurfaceHeaderGetProfileId(SSSHeader);
 					const float RadiusScale = StrataSubSurfaceHeaderGetProfileRadiusScale(SSSHeader);
-					const float3 DiffuseMFP = GetSubsurfaceProfileMFPInCm(ProfileId) * RadiusScale;
+					const float3 DiffuseMFP = GetSubsurfaceProfileMFPInCm(ProfileId).xyz * RadiusScale;
 					const float Phase = GetTransmissionProfileParams(ProfileId).ScatteringDistribution;
 
 					Print(Context, TEXT(" SSS ID.     "), FontBSDFPropName);	Print(Context, ProfileId, FontBSDFPropValu);	Newline(Context, RectMax);
 					Print(Context, TEXT(" SSS Radius  "), FontBSDFPropName);	Print(Context, RadiusScale, FontBSDFPropValu);	Newline(Context, RectMax);
 					Print(Context, TEXT(" SSS MFP     "), FontBSDFPropName);	Print(Context, DiffuseMFP, FontBSDFPropValu);	Newline(Context, RectMax);
 					Print(Context, TEXT(" SSS Phase   "), FontBSDFPropName);	Print(Context, Phase, FontBSDFPropValu);		Newline(Context, RectMax);
+					if (bIsThin)
+					{
+					Print(Context, TEXT(" SSS Thicknes"), FontBSDFPropName);	Print(Context, SLAB_SSSPROFILETHICKNESSCM(BSDF), FontBSDFPropValu);		Newline(Context, RectMax);
+					}
 				}
-				else
+				else if (BSDF_GETSSSTYPE(BSDF) == SSS_TYPE_DIFFUSION)
 				{
-					Print(Context, TEXT(" SSS MFP     "), FontBSDFPropName);	Print(Context, SLAB_SSSMFP(BSDF), FontBSDFPropValu);				Newline(Context, RectMax);
-					Print(Context, TEXT(" SSS Phase   "), FontBSDFPropName);	Print(Context, SLAB_SSSPHASEANISOTROPY(BSDF), FontBSDFPropValu);		Newline(Context, RectMax);
+					const float3 OriginalMFP = StrataSubSurfaceHeaderGetMFP(SSSHeader);
+					const float3 RescaledMFP = SLAB_SSSMFP(BSDF);
+
+					if (bIsThin)
+					{
+						Print(Context, TEXT(" SSS MFP     "), FontBSDFPropName);	Print(Context, OriginalMFP, FontBSDFPropValu);				Newline(Context, RectMax);
+						Print(Context, TEXT(" SSS MFP Norm"), FontBSDFPropName);	Print(Context, RescaledMFP, FontBSDFPropValu);				Newline(Context, RectMax);
+						Print(Context, TEXT(" SSS Thicknes"), FontBSDFPropName);	Print(Context, BSDF_GETTHICKNESSCM(BSDF), FontBSDFPropValu);Newline(Context, RectMax);
+					}
+					else
+					{
+						Print(Context, TEXT(" SSS MFP     "), FontBSDFPropName);	Print(Context, OriginalMFP, FontBSDFPropValu);				Newline(Context, RectMax);
+					}
+
+					Print(Context, TEXT(" SSS Phase   "), FontBSDFPropName);	Print(Context, SLAB_SSSPHASEANISOTROPY(BSDF), FontBSDFPropValu);Newline(Context, RectMax);
 				}
 			} // SSS
 
diff --git a/Engine/Source/Runtime/Engine/Private/Materials/Material.cpp b/Engine/Source/Runtime/Engine/Private/Materials/Material.cpp
--- a/Engine/Source/Runtime/Engine/Private/Materials/Material.cpp
+++ b/Engine/Source/Runtime/Engine/Private/Materials/Material.cpp
@@ -4526,15 +4526,7 @@ bool UMaterial::CanEditChange(const FProperty* InProperty) const
 
 		if (PropertyName == GET_MEMBER_NAME_STRING_CHECKED(UMaterial, SubsurfaceProfile))
 		{
-			const bool bCommonTest = MaterialDomain == MD_Surface && UseSubsurfaceProfile(ShadingModels) && IsOpaqueOrMaskedBlendMode(BlendMode);
-			if (bStrataEnabled)
-			{
-				return !IsThinSurface() && bCommonTest;
-			}
-			else
-			{
-				return bCommonTest;
-			}
+			return MaterialDomain == MD_Surface && UseSubsurfaceProfile(ShadingModels) && IsOpaqueOrMaskedBlendMode(BlendMode);
 		}
 
 		if (PropertyName == GET_MEMBER_NAME_STRING_CHECKED(FLightmassMaterialInterfaceSettings, bCastShadowAsMasked))
@@ -5059,8 +5051,7 @@ void UMaterial::RebuildShadingModelField()
 
 		// Now, reset the subsurface profile (in case it has been removed from any slab before) and set it only if needed.
 		SubsurfaceProfile = nullptr;
-		if ((StrataMaterialInfo.HasOnlyShadingModel(SSM_Eye) || StrataMaterialInfo.HasOnlyShadingModel(SSM_SubsurfaceLit)) && StrataMaterialInfo.CountSubsurfaceProfiles() > 0 
-			&& !IsThinSurface())	// Thin surfaces disable subsurface profile because this is not compatible with fast path.
+		if ((StrataMaterialInfo.HasOnlyShadingModel(SSM_Eye) || StrataMaterialInfo.HasOnlyShadingModel(SSM_SubsurfaceLit)) && StrataMaterialInfo.CountSubsurfaceProfiles() > 0)
 		{
 			if (StrataMaterialInfo.CountSubsurfaceProfiles() > 1)
 			{
