Fix binocular disparity in parallax corrected box reflections when ISR is enabled.

diff --git a/Engine/Shaders/Private/ReflectionEnvironmentComposite.ush b/Engine/Shaders/Private/ReflectionEnvironmentComposite.ush
--- a/Engine/Shaders/Private/ReflectionEnvironmentComposite.ush
+++ b/Engine/Shaders/Private/ReflectionEnvironmentComposite.ush
@@ -60,7 +60,7 @@ float3 CompositeReflectionCapturesAndSkylightTWS(
 #endif
 
 		FLWCVector3 CaptureWorldPosition = MakeLWCVector3(GetReflectionTilePosition(CaptureIndex).xyz, GetReflectionPositionAndRadius(CaptureIndex).xyz);
-		float3 CaptureTranslatedWorldPosition = LWCToFloat(LWCAdd(CaptureWorldPosition, PrimaryView.PreViewTranslation));
+		float3 CaptureTranslatedWorldPosition = LWCToFloat(LWCAdd(CaptureWorldPosition, ResolvedView.PreViewTranslation));
 		float CaptureRadius = GetReflectionPositionAndRadius(CaptureIndex).w;
 
 		float4 CaptureProperties = GetReflectionCaptureProperties(CaptureIndex);
@@ -128,7 +128,7 @@ float3 CompositeReflectionCapturesAndSkylightTWS(
 	float3 ProjectedCaptureVector = RayDirection;
 
 	FLWCVector3 SingleCaptureWorldPosition = MakeLWCVector3(GetReflectionTilePosition(SingleCaptureIndex).xyz, GetReflectionPositionAndRadius(SingleCaptureIndex).xyz);
-	float3 SingleCaptureTranslatedWorldPosition = LWCToFloat(LWCAdd(SingleCaptureWorldPosition, PrimaryView.PreViewTranslation));
+	float3 SingleCaptureTranslatedWorldPosition = LWCToFloat(LWCAdd(SingleCaptureWorldPosition, ResolvedView.PreViewTranslation));
 	float SingleCaptureRadius = GetReflectionPositionAndRadius(SingleCaptureIndex).w;
 
 	float4 SingleCaptureOffsetAndAverageBrightness = GetReflectionCaptureOffsetAndAverageBrightness(SingleCaptureIndex);
