Path Tracer: Fix darker specular colors becoming grey in Substrate

The decal logic was always running, even when no decals were present.

diff --git a/Engine/Shaders/Private/PathTracing/PathTracingCore.ush b/Engine/Shaders/Private/PathTracing/PathTracingCore.ush
--- a/Engine/Shaders/Private/PathTracing/PathTracingCore.ush
+++ b/Engine/Shaders/Private/PathTracing/PathTracingCore.ush
@@ -242,18 +242,23 @@ struct FDecalData
 		float Specular = HitPayload.Specular;
 		float Roughness = HitPayload.Roughness;
 #endif // PATHTRACING_STRATA_PAYLOAD
-		if (HitPayload.HasDecalResponseColor())
+
+		bool bHasColorResponse = HitPayload.HasDecalResponseColor() && this.BaseColor.a < 1.0;
+		bool bHasNormalResponse = HitPayload.HasDecalResponseNormal() && this.WorldNormal.a < 1.0;
+		bool bHasRoughnessResponse = HitPayload.HasDecalResponseRoughness() && this.MetallicSpecularRoughness.a < 1.0;
+
+		if (bHasColorResponse)
 		{
 			BaseColor = BaseColor * this.BaseColor.a + this.BaseColor.rgb;
 		}
 
-		if (HitPayload.HasDecalResponseNormal())
+		if (bHasNormalResponse)
 		{
 			// We normalize the normal to get smoother visual result (it helps to avoid having D_GGX explodes toward infinity, and matches ApplyDBufferData(...))
 			HitPayload.WorldNormal = normalize(HitPayload.WorldNormal * this.WorldNormal.a + this.WorldNormal.rgb);
 		}
 
-		if (HitPayload.HasDecalResponseRoughness())
+		if (bHasRoughnessResponse)
 		{
 			Metallic = Metallic * this.MetallicSpecularRoughness.a + this.MetallicSpecularRoughness.x;
 			Specular = Specular * this.MetallicSpecularRoughness.a + this.MetallicSpecularRoughness.y;
@@ -263,7 +268,7 @@ struct FDecalData
 		HitPayload.Radiance += this.Emissive;
 
 #if PATHTRACING_STRATA_PAYLOAD
-		if (HitPayload.HasDecalResponseColor() || HitPayload.HasDecalResponseRoughness())
+		if (bHasColorResponse || bHasRoughnessResponse)
 		{
 			// the encode-decode is potentially lossy, so only do it if we actually need to
 			HitPayload.DiffuseColor = BaseColor - BaseColor * Metallic;
