Add comment for HasScatteringData  and HasBackScattering.

diff --git a/Engine/Shaders/Private/Strata/Strata.ush b/Engine/Shaders/Private/Strata/Strata.ush
--- a/Engine/Shaders/Private/Strata/Strata.ush
+++ b/Engine/Shaders/Private/Strata/Strata.ush
@@ -179,8 +179,10 @@ struct FStrataBSDF
 
 	void PostProcessBSDFBeforeLighting(bool bIsStrataOpaqueMaterial);
 
+	// Return true if the BSDF has side SSS data (i.e., SSSHeader/SSSData)
 	bool HasScatteringData();
 
+	// Return true if the BSDF transmits and scatters light on the opposite side
 	bool HasBackScattering();
 };
 
