Removed unused GroupThreadID parameter from ComputeShaderOutputCommon.ush

diff --git a/Engine/Shaders/Private/ComputeShaderOutputCommon.ush b/Engine/Shaders/Private/ComputeShaderOutputCommon.ush
--- a/Engine/Shaders/Private/ComputeShaderOutputCommon.ush
+++ b/Engine/Shaders/Private/ComputeShaderOutputCommon.ush
@@ -290,7 +290,7 @@ void ProcessPixel(uint ShadingBin, const uint2 PixelPos, uint QuadIndex, uint Di
 }
 
 [numthreads(64, 1, 1)]
-void MainCS(uint ThreadIndex : SV_GroupIndex, uint GroupID : SV_GroupID, uint3 GroupThreadID : SV_GroupThreadID)
+void MainCS(uint ThreadIndex : SV_GroupIndex, uint GroupID : SV_GroupID)
 {
 	const uint ShadingBin = PassData.x;
 	const bool bQuadBinning = PassData.z != 0u;
