Substrate - Glint - disabled ewhen derivative are not supported.

diff --git a/Engine/Shaders/Private/Strata/StrataEvaluation.ush b/Engine/Shaders/Private/Strata/StrataEvaluation.ush
--- a/Engine/Shaders/Private/Strata/StrataEvaluation.ush
+++ b/Engine/Shaders/Private/Strata/StrataEvaluation.ush
@@ -12,7 +12,11 @@
 
 #include "Strata.ush"
 
-#if STRATA_COMPLEXPATH && STRATA_GLINTS
+// Disable glint when in tracing mode for now since derivative are not available.
+#define STRATA_GLINTS_DERIVATIVE_SUPPORTED (!RAYCALLABLESHADER && !RAYHITGROUPSHADER && !RAYGENSHADER && !RAYMISSSHADER)
+#define STRATA_GLINTS_ENABLED (STRATA_COMPLEXPATH && STRATA_GLINTS && STRATA_GLINTS_DERIVATIVE_SUPPORTED)
+
+#if STRATA_GLINTS_ENABLED
 #include "Glint/GlintThirdParty.ush"
 #endif
 
@@ -410,7 +414,7 @@ FStrataEvaluateResult StrataEvaluateBSDFCommon(FStrataBSDFContext BSDFContext, F
 						float2 AlphaXY = 0;
 						GetAnisotropicRoughness(Alpha, SLAB_ANISOTROPY(BSDFContext.BSDF), AlphaXY.x, AlphaXY.y);
 
-					#if STRATA_GLINTS
+					#if STRATA_GLINTS_ENABLED
 						// Glint shading only in the complex path
 						BRANCH
 						if (BSDF_GETHASGLINT(BSDFContext.BSDF))
@@ -445,7 +449,7 @@ FStrataEvaluateResult StrataEvaluateBSDFCommon(FStrataBSDFContext BSDFContext, F
 							RectLightSpec = RectGGXApproxLTC(SafeRoughness, F0, BSDFContext.N, BSDFContext.V, AreaLightContext.AreaLight.Rect, AreaLightContext.AreaLight.Texture, MeanLightWorldDirection);
 							
 							// Now combine the rectangular area light with Glints.
-						#if STRATA_COMPLEXPATH && STRATA_GLINTS
+						#if STRATA_GLINTS_ENABLED
 							BRANCH
 							if (BSDF_GETHASGLINT(BSDFContext.BSDF))
 							{
@@ -474,7 +478,7 @@ FStrataEvaluateResult StrataEvaluateBSDFCommon(FStrataBSDFContext BSDFContext, F
 						{
 							if(IntegrationType == INTEGRATION_PUNCTUAL_LIGHT)
 							{
-							#if STRATA_COMPLEXPATH && STRATA_GLINTS
+							#if STRATA_GLINTS_ENABLED
 								BRANCH
 								if (BSDF_GETHASGLINT(BSDFContext.BSDF))
 								{
@@ -494,7 +498,7 @@ FStrataEvaluateResult StrataEvaluateBSDFCommon(FStrataBSDFContext BSDFContext, F
 							{
 								const float Energy = EnergyNormalization(Alpha2Spec, BSDFContext.Context.VoH, AreaLightContext.AreaLight);
 
-							#if STRATA_COMPLEXPATH && STRATA_GLINTS
+							#if STRATA_GLINTS_ENABLED
 								BRANCH
 								if (BSDF_GETHASGLINT(BSDFContext.BSDF))
 								{
@@ -1465,7 +1469,7 @@ FStrataEnvLightResult StrataEvaluateForEnvLight(FStrataBSDFContext BSDFContext,
 				StrataEnvLightResult.SpecularColor = F0;
 
 				// STRATA_TODO environment light importance sampling
-		//	#if STRATA_COMPLEXPATH && STRATA_GLINTS
+		//	#if STRATA_GLINTS_ENABLED
 		//		BRANCH
 		//		if (BSDF_GETHASGLINT(BSDFContext.BSDF))
 		//		{
