Substrate - Glint - made some code common.

diff --git a/Engine/Shaders/Private/Strata/StrataEvaluation.ush b/Engine/Shaders/Private/Strata/StrataEvaluation.ush
--- a/Engine/Shaders/Private/Strata/StrataEvaluation.ush
+++ b/Engine/Shaders/Private/Strata/StrataEvaluation.ush
@@ -101,8 +101,9 @@ struct FStrataBSDFContext
 	float3x3 TangentBasis;
 	float3 TangentV;
 	float3 TangentH;
-	float3 TangentB;
-
+	float3 TangentB; 
+	float3 TangentL;
+	
 	void StrataUpdateBSDFContext(float3 NewL);
 };
 
@@ -137,6 +138,7 @@ FStrataBSDFContext StrataCreateBSDFContext(float3x3 TangentBasis, FStrataBSDF BS
 	BSDFContext.TangentV = mul(BSDFContext.TangentBasis, BSDFContext.V);
 	BSDFContext.TangentH = mul(BSDFContext.TangentBasis, BSDFContext.H);
 	BSDFContext.TangentB = mul(BSDFContext.TangentBasis, BSDFContext.B);
+	BSDFContext.TangentL = normalize(mul(BSDFContext.TangentBasis, BSDFContext.L));
 
 	BSDFContext.SatNoL = saturate(BSDFContext.Context.NoL);
 	BSDFContext.SatNoV = saturate(BSDFContext.Context.NoV);
@@ -259,6 +261,24 @@ float2 GGXRoughnessToBeckmannRoughness(float2 Roughness)
 	return float2(GGXRoughnessToBeckmannRoughness(Roughness.x), GGXRoughnessToBeckmannRoughness(Roughness.y));
 }
 
+struct FBeckmannDesc
+{
+	float2 Roughness;	// Roughness
+	float2 Alpha;		// Roughness squared
+	float2 Sigma;		// RMS slope of the microfacets
+	float  Rho;			// Slope correlation factor to be used with LEAN mapping. See http://igg.unistra.fr/People/chermain/assets/pdf/Chermain2021RealTime.pdf.
+};
+FBeckmannDesc GGXToBeckmann(float2 GGXSafeRoughness)
+{
+	FBeckmannDesc Beckmann;
+	Beckmann.Roughness	= GGXRoughnessToBeckmannRoughness(GGXSafeRoughness);
+	Beckmann.Alpha		= Square(Beckmann.Roughness);
+	Beckmann.Sigma		= Beckmann.Alpha * InvSqrt2;
+	Beckmann.Rho		= 0.0f;
+	return Beckmann;
+}
+
+
 ///////////////////////////////////////////////////////////////////////////////////////////////////
 // BSDF evaluation (punctual/area)
 
@@ -401,8 +421,6 @@ FStrataEvaluateResult StrataEvaluateBSDFCommon(FStrataBSDFContext BSDFContext, F
 				float D = 0;
 				float Vis = 0;
 
-				float3 TangentL = normalize(mul(BSDFContext.TangentBasis, BSDFContext.L));
-
 				float3 RectLightSpec = 0;
 				#if STRATA_COMPLEXPATH
 				BRANCH
@@ -420,11 +438,8 @@ FStrataEvaluateResult StrataEvaluateBSDFCommon(FStrataBSDFContext BSDFContext, F
 						if (BSDF_GETHASGLINT(BSDFContext.BSDF))
 						{
 							const float2 GGXRoughnessXY = float2(sqrtFast(AlphaXY.x), sqrtFast(AlphaXY.y));
-							const float2 BeckmannRoughnessXY = GGXRoughnessToBeckmannRoughness(GGXRoughnessXY);
-							const float2 AlphaXY = Square(BeckmannRoughnessXY);
-							float2 SigmaXY = AlphaXY * InvSqrt2;
-							const float Rho = 0.0f;
-							D = f_P(BSDFContext.TangentV, TangentL, float3(SigmaXY, Rho), SLAB_GLINT_VALUE(BSDFContext.BSDF), SLAB_GLINT_UV(BSDFContext.BSDF));
+							FBeckmannDesc Beckmann = GGXToBeckmann(GGXRoughnessXY);
+							D = f_P(BSDFContext.TangentV, BSDFContext.TangentL, float3(Beckmann.Sigma, Beckmann.Rho), SLAB_GLINT_VALUE(BSDFContext.BSDF), SLAB_GLINT_UV(BSDFContext.BSDF));
 						}
 						else
 					#endif
@@ -454,15 +469,11 @@ FStrataEvaluateResult StrataEvaluateBSDFCommon(FStrataBSDFContext BSDFContext, F
 							if (BSDF_GETHASGLINT(BSDFContext.BSDF))
 							{
 								float3 MeanLightLocalDirection = normalize(mul(BSDFContext.TangentBasis, MeanLightWorldDirection));
-
-								const float BeckmannRoughness = GGXRoughnessToBeckmannRoughness(SafeRoughness);
-								const float Alpha = Square(BeckmannRoughness);
-								float Sigma = Alpha * InvSqrt2;
-								const float Rho = 0.0f;
+								FBeckmannDesc Beckmann = GGXToBeckmann(SafeRoughness);
 
 								// It is wrong to multiply again the GGXevaluation together with the glint function (representing Beckmann D function).
 								// But visually it works and that is an acceptable solution in the meantime we find something more correct.
-								RectLightSpec *= f_P(BSDFContext.TangentV, MeanLightLocalDirection, float3(Sigma, Sigma, Rho), SLAB_GLINT_VALUE(BSDFContext.BSDF), SLAB_GLINT_UV(BSDFContext.BSDF));
+								RectLightSpec *= f_P(BSDFContext.TangentV, MeanLightLocalDirection, float3(Beckmann.Sigma.xx, Beckmann.Rho), SLAB_GLINT_VALUE(BSDFContext.BSDF), SLAB_GLINT_UV(BSDFContext.BSDF));
 							}
 						#endif // STRATA_COMPLEXPATH
 
@@ -482,11 +493,8 @@ FStrataEvaluateResult StrataEvaluateBSDFCommon(FStrataBSDFContext BSDFContext, F
 								BRANCH
 								if (BSDF_GETHASGLINT(BSDFContext.BSDF))
 								{
-									const float BeckmannRoughness = GGXRoughnessToBeckmannRoughness(SafeRoughness);
-									const float Alpha = Square(BeckmannRoughness);
-									float Sigma = Alpha * InvSqrt2;
-									const float Rho = 0.0f;
-									D = f_P(BSDFContext.TangentV, TangentL, float3(Sigma, Sigma, Rho), SLAB_GLINT_VALUE(BSDFContext.BSDF), SLAB_GLINT_UV(BSDFContext.BSDF));
+									FBeckmannDesc Beckmann = GGXToBeckmann(SafeRoughness);
+									D = f_P(BSDFContext.TangentV, BSDFContext.TangentL, float3(Beckmann.Sigma.xx, Beckmann.Rho), SLAB_GLINT_VALUE(BSDFContext.BSDF), SLAB_GLINT_UV(BSDFContext.BSDF));
 								}
 								else
 							#endif // STRATA_COMPLEXPATH
@@ -502,11 +510,8 @@ FStrataEvaluateResult StrataEvaluateBSDFCommon(FStrataBSDFContext BSDFContext, F
 								BRANCH
 								if (BSDF_GETHASGLINT(BSDFContext.BSDF))
 								{
-									const float BeckmannRoughness = GGXRoughnessToBeckmannRoughness(SafeRoughness);
-									const float Alpha = Square(BeckmannRoughness);
-									float Sigma = Alpha * InvSqrt2;
-									const float Rho = 0.0f;
-									D = Energy * f_P(BSDFContext.TangentV, TangentL, float3(Sigma, Sigma, Rho), SLAB_GLINT_VALUE(BSDFContext.BSDF), SLAB_GLINT_UV(BSDFContext.BSDF));
+									FBeckmannDesc Beckmann = GGXToBeckmann(SafeRoughness);
+									D = Energy * f_P(BSDFContext.TangentV, BSDFContext.TangentL, float3(Beckmann.Sigma.xx, Beckmann.Rho), SLAB_GLINT_VALUE(BSDFContext.BSDF), SLAB_GLINT_UV(BSDFContext.BSDF));
 								}
 								else
 							#endif // STRATA_COMPLEXPATH
@@ -1475,14 +1480,11 @@ FStrataEnvLightResult StrataEvaluateForEnvLight(FStrataBSDFContext BSDFContext,
 		//		{
 		//			float3 MeanLightLocalDirection = normalize(mul(BSDFContext.TangentBasis, StrataEnvLightResult.SpecularDirection));
 		//
-		//			const float BeckmannRoughness = GGXRoughnessToBeckmannRoughness(SafeRoughness);
-		//			const float Alpha = Square(BeckmannRoughness);
-		//			float Sigma = Alpha * InvSqrt2;
-		//			const float Rho = 0.0f;
+		//			FBeckmannDesc Beckmann = GGXToBeckmann(SafeRoughness);
 		//
 		//			// It is wrong to multiply again the split env lighting (for split-sum approxiation) together with the glint function (representing Beckmann D function).
 		//			// But visually it works and that is an acceptable solution in the meantime we find something more correct.
-		//			float GlintD = f_P(BSDFContext.TangentV, MeanLightLocalDirection, float3(Sigma, Sigma, Rho), SLAB_GLINT_VALUE(BSDFContext.BSDF), SLAB_GLINT_UV(BSDFContext.BSDF));
+		//			float GlintD = f_P(BSDFContext.TangentV, MeanLightLocalDirection, float3(Beckmann.Sigma, Beckmann.Rho), SLAB_GLINT_VALUE(BSDFContext.BSDF), SLAB_GLINT_UV(BSDFContext.BSDF));
 		//
 		//			// We need a way to lerp to the glint direction with a roughness of 0. 
 		//			// STRATA_TODO recover glint direction.
