Substrate - GLint - added support for parameter blending.

diff --git a/Engine/Shaders/Private/Strata/Strata.ush b/Engine/Shaders/Private/Strata/Strata.ush
--- a/Engine/Shaders/Private/Strata/Strata.ush
+++ b/Engine/Shaders/Private/Strata/Strata.ush
@@ -2354,10 +2354,6 @@ FStrataData StrataHorizontalMixingParameterBlending(FStrataData Background, FStr
 	BSDF_SETHASFUZZ(ResultBSDF,			BSDF_GETHASFUZZ(ForegroundBSDF)					| BSDF_GETHASFUZZ(BackgroundBSDF));
 	BSDF_SETHASTRANSABOVE(ResultBSDF,	BSDF_GETHASTRANSABOVE(ForegroundBSDF)			| BSDF_GETHASTRANSABOVE(BackgroundBSDF));
 
-#if STRATA_GLINTS
-	BSDF_SETHASGLINT(ResultBSDF, 0); // Disable for now
-#endif
-
 	BSDF_SETEMISSIVE(ResultBSDF,		lerp(BSDF_GETEMISSIVE(BackgroundBSDF),			BSDF_GETEMISSIVE(ForegroundBSDF),			ForegroundOtherMixFactor));
 	SLAB_DIFFUSEALBEDO(ResultBSDF) =	lerp(SLAB_DIFFUSEALBEDO(BackgroundBSDF),		SLAB_DIFFUSEALBEDO(ForegroundBSDF),			ForegroundOtherMixFactor);
 	SLAB_F0(ResultBSDF) =				lerp(SLAB_F0(BackgroundBSDF),					SLAB_F0(ForegroundBSDF),					ForegroundOtherMixFactor);
@@ -2370,6 +2366,12 @@ FStrataData StrataHorizontalMixingParameterBlending(FStrataData Background, FStr
 	SLAB_FUZZ_AMOUNT(ResultBSDF) =		lerp(SLAB_FUZZ_AMOUNT(BackgroundBSDF),			SLAB_FUZZ_AMOUNT(ForegroundBSDF),			ForegroundOtherMixFactor);
 	SLAB_FUZZ_ROUGHNESS(ResultBSDF)	=	lerp(SLAB_FUZZ_ROUGHNESS(BackgroundBSDF),		SLAB_FUZZ_ROUGHNESS(ForegroundBSDF),		ForegroundOtherMixFactor);
 	SLAB_HAZINESS(ResultBSDF) =			LerpHazinessParameterBlending(BackgroundBSDF,	ForegroundBSDF,								ForegroundOtherMixFactor);
+
+#if STRATA_GLINTS
+	BSDF_SETHASGLINT(ResultBSDF,		BSDF_GETHASGLINT(ForegroundBSDF)				| BSDF_GETHASGLINT(BackgroundBSDF));
+	SLAB_GLINT_UV(ResultBSDF) =			lerp(SLAB_GLINT_UV(BackgroundBSDF),				SLAB_GLINT_UV(ForegroundBSDF),				ForegroundOtherMixFactor);
+	SLAB_GLINT_VALUE(ResultBSDF) =		lerp(SLAB_GLINT_VALUE(BackgroundBSDF),			SLAB_GLINT_VALUE(ForegroundBSDF),			ForegroundOtherMixFactor);
+#endif
 	
 
 	if (BSDF_GETSSSTYPE(ResultBSDF) == SSS_TYPE_DIFFUSION_PROFILE)
@@ -2463,10 +2465,6 @@ FStrataData StrataAddParameterBlending(FStrataData A, FStrataData B, float AMixF
 	BSDF_SETHASF90(ResultBSDF,			BSDF_GETHASF90(BBSDF)							| BSDF_GETHASF90(ABSDF));
 	BSDF_SETHASFUZZ(ResultBSDF,			BSDF_GETHASFUZZ(BBSDF)							| BSDF_GETHASFUZZ(ABSDF));
 
-#if STRATA_GLINTS
-	BSDF_SETHASGLINT(ResultBSDF, 0); // Disable for now
-#endif
-
 	// Some parameters will contribute to added luminance: in this case we add them together and saturate to not go out of the safe range
 	BSDF_SETEMISSIVE(ResultBSDF,		BSDF_GETEMISSIVE(ABSDF)								+ BSDF_GETEMISSIVE(ABSDF));
 	SLAB_DIFFUSEALBEDO(ResultBSDF) =	saturate(SLAB_DIFFUSEALBEDO(BBSDF)					+ SLAB_DIFFUSEALBEDO(ABSDF));
@@ -2482,6 +2480,12 @@ FStrataData StrataAddParameterBlending(FStrataData A, FStrataData B, float AMixF
 	SLAB_FUZZ_ROUGHNESS(ResultBSDF) =	lerp(SLAB_FUZZ_ROUGHNESS(BBSDF),				SLAB_FUZZ_ROUGHNESS(ABSDF),					AMixFactor);
 	SLAB_HAZINESS(ResultBSDF) =			LerpHazinessParameterBlending(BBSDF,			ABSDF,										AMixFactor);
 
+#if STRATA_GLINTS
+	BSDF_SETHASGLINT(ResultBSDF,		BSDF_GETHASGLINT(BBSDF)							| BSDF_GETHASGLINT(ABSDF));
+	SLAB_GLINT_UV(ResultBSDF) =			lerp(SLAB_GLINT_UV(BBSDF),						SLAB_GLINT_UV(ABSDF), AMixFactor);
+	SLAB_GLINT_VALUE(ResultBSDF) =		lerp(SLAB_GLINT_VALUE(BBSDF),					SLAB_GLINT_VALUE(ABSDF), AMixFactor);
+#endif
+
 	if (BSDF_GETSSSTYPE(ResultBSDF) == SSS_TYPE_DIFFUSION_PROFILE)
 	{
 		// As mentioned above, we cannot blend SSSProfiles. We do select the appropriate profile, ABSDF being selected if both have SSSProfile.
@@ -2600,6 +2604,12 @@ FStrataBSDF StrataVerticalLayeringParameterBlending(in FStrataBSDF TopBSDF, in F
 	SLAB_FUZZ_ROUGHNESS(ResultBSDF) =	lerp(SLAB_FUZZ_ROUGHNESS(TopBSDF),				SLAB_FUZZ_ROUGHNESS(BaseBSDF),TopToBottomLerp);
 	SLAB_HAZINESS(ResultBSDF) =			LerpHazinessParameterBlending(TopBSDF,			BaseBSDF,					TopToBottomLerp);
 
+#if STRATA_GLINTS
+	BSDF_SETHASGLINT(ResultBSDF,		BSDF_GETHASGLINT(TopBSDF)						| BSDF_GETHASGLINT(BaseBSDF));
+	SLAB_GLINT_UV(ResultBSDF) =			lerp(SLAB_GLINT_UV(TopBSDF),					SLAB_GLINT_UV(BaseBSDF), TopToBottomLerp);
+	SLAB_GLINT_VALUE(ResultBSDF) =		lerp(SLAB_GLINT_VALUE(TopBSDF),					SLAB_GLINT_VALUE(BaseBSDF), TopToBottomLerp);
+#endif
+
 	// Anisotropy is disabled with vertical blending. 
 	SLAB_ANISOTROPY(ResultBSDF) =		0.f;
 
