Substrate: Fix shader warnings in EngineTest

The unused functions should be guarded by STRATA_GLINTS_ENABLED so we can build without the minifier enabled.

diff --git a/Engine/Shaders/Private/Strata/StrataEvaluation.ush b/Engine/Shaders/Private/Strata/StrataEvaluation.ush
--- a/Engine/Shaders/Private/Strata/StrataEvaluation.ush
+++ b/Engine/Shaders/Private/Strata/StrataEvaluation.ush
@@ -250,6 +250,7 @@ half3 Strata_F_GGX(half3 F0, half3 F90, half VoH)
 //#endif
 }
 
+#if STRATA_GLINTS_ENABLED
 // Conversion of GGX roughness to Beckmann roughness
 // Using Microfacet Distribution Function: To Change or Not to Change, That Is the Question.
 float GGXRoughnessToBeckmannRoughness(float Roughness)
@@ -277,7 +278,7 @@ FBeckmannDesc GGXToBeckmann(float2 GGXSafeRoughness)
 	Beckmann.Rho		= 0.0f;
 	return Beckmann;
 }
-
+#endif
 
 ///////////////////////////////////////////////////////////////////////////////////////////////////
 // BSDF evaluation (punctual/area)
