Path Tracer: Enable roughness tracking between slabs for Substrate materials

The integration settings were not initialized to include this by default. Left a TODO to implement support for the actual cvar.

diff --git a/Engine/Shaders/Private/PathTracing/PathTracingMaterialHitShader.usf b/Engine/Shaders/Private/PathTracing/PathTracingMaterialHitShader.usf
--- a/Engine/Shaders/Private/PathTracing/PathTracingMaterialHitShader.usf
+++ b/Engine/Shaders/Private/PathTracing/PathTracingMaterialHitShader.usf
@@ -222,7 +222,14 @@ RAY_TRACING_ENTRY_CLOSEST_HIT(PathTracingMaterialCHS,
 	const float3 Transparency = saturate(StrataData.InlinedBSDF.Coverage + (1.0 - StrataData.InlinedBSDF.Coverage) * UNLIT_TRANSMITTANCE(StrataData.InlinedBSDF));
 #else
 	FStrataPixelHeader StrataPixelHeader = MaterialParameters.GetFrontStrataHeader();
-	FStrataIntegrationSettings IntegrationSettings = InitStrataIntegrationSettings();
+	// STRATA_TODO: These should actually be initialized via the corresponding cvars. To do this properly would require binding FStrataSceneData in the path tracing shaders
+	// For now, simply match the default substrate settings
+	FStrataIntegrationSettings IntegrationSettings = InitStrataIntegrationSettings(
+		false /* bForceFullyRough */,
+		true  /* StrataStruct.bRoughDiffuse */,
+		-1    /* StrataStruct.PeelLayersAboveDepth */,
+		true  /* StrataStruct.bRoughnessTracking */
+	);
 	FStrataAddressing NullStrataAddressing = (FStrataAddressing)0;	// Fake unused in StrataCreateBSDFContext when using Forward inline shading
 #if STRATA_BLENDING_COLOREDTRANSMITTANCEONLY
 	const uint BSDFCount = 0;
