Change glint UV to be stored onto 24bits for improving precision and removig derivative artifacts.

diff --git a/Engine/Shaders/Private/PackUnpack.ush b/Engine/Shaders/Private/PackUnpack.ush
--- a/Engine/Shaders/Private/PackUnpack.ush
+++ b/Engine/Shaders/Private/PackUnpack.ush
@@ -211,6 +211,19 @@ float4 UnpackR16G16B16A16F(uint2 In)
 	return float4(UnpackFloat2FromUInt(In.x), UnpackFloat2FromUInt(In.y));
 }
 
+//
+//	float <-> 24bit float (8bit exp, 15 bit mantissa)
+//
+
+uint PackR24F(float In)
+{
+	return asuint(In) >> 8;
+}
+
+float UnpackR24F(uint In)
+{
+	return asfloat(In<<8);
+}
 
 //
 // Common tile coord encoding
diff --git a/Engine/Shaders/Private/Strata/Strata.ush b/Engine/Shaders/Private/Strata/Strata.ush
--- a/Engine/Shaders/Private/Strata/Strata.ush
+++ b/Engine/Shaders/Private/Strata/Strata.ush
@@ -4278,9 +4278,10 @@ FStrataBSDF UnpackStrataBSDFIn(FStrataMaterialContainer StrataBuffer, inout FStr
 			uint2 GlintPackedData;
 			GlintPackedData.x = StrataLoadUint1(StrataBuffer, StrataAddressing);
 			GlintPackedData.y = StrataLoadUint1(StrataBuffer, StrataAddressing);
-			float4 GlintUnpackedData	= UnpackR16G16B16A16F(GlintPackedData);
-			SLAB_GLINT_UV(OutBSDF)		= GlintUnpackedData.xy;
-			SLAB_GLINT_VALUE(OutBSDF)	= GlintUnpackedData.z;
+
+			SLAB_GLINT_UV(OutBSDF).x	= UnpackR24F(GlintPackedData.x);
+			SLAB_GLINT_UV(OutBSDF).y	= UnpackR24F(GlintPackedData.y);
+			SLAB_GLINT_VALUE(OutBSDF)	= f16tof32((GlintPackedData.x>>24) | ((GlintPackedData.y>>16) & 0xFF00));
 		}
 	#endif
 	}
diff --git a/Engine/Shaders/Private/Strata/StrataExport.ush b/Engine/Shaders/Private/Strata/StrataExport.ush
--- a/Engine/Shaders/Private/Strata/StrataExport.ush
+++ b/Engine/Shaders/Private/Strata/StrataExport.ush
@@ -794,9 +794,11 @@ void PackStrataOut(
 							#if STRATA_GLINTS
 							if (BSDF_GETHASGLINT(BSDF))
 							{
-								uint2 GlintData = PackR16G16B16A16F(float4(SLAB_GLINT_UV(BSDF), SLAB_GLINT_VALUE(BSDF), 0.0f));
-								STRATA_STORE_UINT1(GlintData.x);
-								STRATA_STORE_UINT1(GlintData.y);
+								const uint Packed16 = f32tof16(SLAB_GLINT_VALUE(BSDF));
+								const uint Packed8L = (Packed16 & 0xFF) << 24;
+								const uint Packed8H = (Packed16 >> 8  ) << 24;
+								STRATA_STORE_UINT1(PackR24F(SLAB_GLINT_UV(BSDF).x) | Packed8L);
+								STRATA_STORE_UINT1(PackR24F(SLAB_GLINT_UV(BSDF).y) | Packed8H);
 							}
 							#endif // STRATA_GLINTS
 
