Change Glint roughness mapping.

diff --git a/Engine/Shaders/Private/Strata/StrataEvaluation.ush b/Engine/Shaders/Private/Strata/StrataEvaluation.ush
--- a/Engine/Shaders/Private/Strata/StrataEvaluation.ush
+++ b/Engine/Shaders/Private/Strata/StrataEvaluation.ush
@@ -255,7 +255,10 @@ half3 Strata_F_GGX(half3 F0, half3 F90, half VoH)
 // Using Microfacet Distribution Function: To Change or Not to Change, That Is the Question.
 float GGXRoughnessToBeckmannRoughness(float Roughness)
 {
-	return saturate(0.8388 * Roughness + 0.7 * Pow4(Roughness));
+	// Use a custom remapping to better match the overall highlight shape, in particular for lower roughness (i.e., <0.5). 
+	// The original mapping was Ctr=Roughness based on the paper
+	const float Ctr = pow(Roughness, 0.7f);
+	return saturate(0.8388 * Ctr+ 0.7 * Pow4(Ctr));
 }
 float2 GGXRoughnessToBeckmannRoughness(float2 Roughness)
 {
