Substrate - Glint lut parameterisation from imported source. To be able to easily swap and test more LUTs later.

diff --git a/Engine/Shaders/Private/Strata/Glint/GlintThirdParty.ush b/Engine/Shaders/Private/Strata/Glint/GlintThirdParty.ush
--- a/Engine/Shaders/Private/Strata/Glint/GlintThirdParty.ush
+++ b/Engine/Shaders/Private/Strata/Glint/GlintThirdParty.ush
@@ -19,13 +19,13 @@
 #define InvSqrt2						0.707106f
 #define MaxAnisotropy					4.0f
 
-// The following variables are linked to the GlintShadingLUT
-#define NumberOfLevels					16
-#define NumberOfDistributionsPerChannel	64
-#define	Dictionary_Alpha				0.5f
-#define	Dictionary_N					int(NumberOfDistributionsPerChannel * 3)	
-#define	Dictionary_NLevels				int(NumberOfLevels)	
-#define	Dictionary_Pyramid0Size			int(1u << (numberOfLevels - 1u))					
+// FGlintLutParams contains variables linked to the GlintShadingLUT
+struct FGlintLutParams
+{
+	float Alpha;
+	int N;
+	int NLevels;
+};
 
 // Geometric Glint Anti-Aliasing parameters
 #define KernelSize						0.5
@@ -105,9 +105,9 @@ float hashIQ(uint n)
 //=============================================================================
 //================== Pyramid size at LOD level ================================
 //=============================================================================
-int pyramidSize(int level)
+int pyramidSize(FGlintLutParams GlintLutParams, int level)
 {
-	return int(pow(2., float(Dictionary_NLevels - 1 - level)));
+	return int(pow(2., float(GlintLutParams.NLevels - 1 - level)));
 }
 
 //=============================================================================
@@ -124,7 +124,9 @@ float sampleNormalDistribution(float U, float mu, float sigma)
 //=============== and transformed slope distribution function  ================
 //=============================== Equation 4 ==================================
 //=============================================================================
-float P22_M(float2 slope_h, int l, int s0, int t0,
+float P22_M(
+	FGlintLutParams GlintLutParams,
+	float2 slope_h, int l, int s0, int t0,
 	float2 slope_dx, float2 slope_dy,
 	float3 sigma_x_y_rho, float l_dist)
 {
@@ -146,7 +148,7 @@ float P22_M(float2 slope_h, int l, int s0, int t0,
 	l_dist = sampleNormalDistribution(uDensityRandomisation, l_dist,
 		densityRandomisation);
 
-	l_dist = clamp(int(round(l_dist)), 0, Dictionary_NLevels);
+	l_dist = clamp(int(round(l_dist)), 0, GlintLutParams.NLevels);
 
 	// Recover roughness and slope correlation factor
 	float sigma_x = sigma_x_y_rho.x;
@@ -154,7 +156,7 @@ float P22_M(float2 slope_h, int l, int s0, int t0,
 	float rho = sigma_x_y_rho.z;
 
 	// If we are too far from the surface, the SDF is a gaussian.
-	if (l_dist == Dictionary_NLevels) {
+	if (l_dist == GlintLutParams.NLevels) {
 		return non_axis_aligned_anisotropic_beckmann(slope_h.x, slope_h.y,
 			sigma_x, sigma_y, rho);
 	}
@@ -170,7 +172,7 @@ float P22_M(float2 slope_h, int l, int s0, int t0,
 	//========= Linearly transformed isotropic Beckmann distribution ==========
 	//=========================================================================
 
-	float SIGMA_DICT = Dictionary_Alpha * InvSqrt2;
+	float SIGMA_DICT = GlintLutParams.Alpha * InvSqrt2;
 	float tmp1 = SIGMA_DICT / (sigma_x * sqrt(1. - rho * rho));
 	float tmp2 =-SIGMA_DICT * rho / (sigma_y * sqrt(1. - rho * rho));
 	float tmp3 = SIGMA_DICT / sigma_y;
@@ -203,8 +205,8 @@ float P22_M(float2 slope_h, int l, int s0, int t0,
 	// The SDF is an even function
 	float2 abs_slope_h_o = float2(abs(slope_h_o.x), abs(slope_h_o.y));
 
-	int distPerChannel = Dictionary_N / 3;
-	float alpha_dist_isqrt2_4 = Dictionary_Alpha * InvSqrt2 * 4.f;
+	int distPerChannel = GlintLutParams.N / 3;
+	float alpha_dist_isqrt2_4 = GlintLutParams.Alpha * InvSqrt2 * 4.f;
 
 	// After 4 standard deviations, the SDF equals zero
 	if (abs_slope_h_o.x > alpha_dist_isqrt2_4
@@ -214,8 +216,8 @@ float P22_M(float2 slope_h, int l, int s0, int t0,
 	float u1 = hashIQ(rngSeed * 16807U);
 	float u2 = hashIQ(rngSeed * 48271U);
 
-	int i = int(u1 * float(Dictionary_N));
-	int j = int(u2 * float(Dictionary_N));
+	int i = int(u1 * float(GlintLutParams.N));
+	int j = int(u2 * float(GlintLutParams.N));
 
 	// 3 distributions values in one texel
 	int distIdxXOver3 = i / 3;
@@ -285,12 +287,14 @@ float P22_M(float2 slope_h, int l, int s0, int t0,
 // which itself is similar to Heckbert 1889 algorithm, 
 // http://www.cs.cmu.edu/~ph/texfund/texfund.pdf, Section 3.5.9.
 // Go through cells within the pixel footprint for a givin LOD
-float P22__glint_discrete_LOD(int l, float2 slope_h, float2 st, float2 dst0,
+float P22__glint_discrete_LOD(
+	FGlintLutParams GlintLutParams,
+	int l, float2 slope_h, float2 st, float2 dst0,
 	float2 dst1, float2 slope_dx, float2 slope_dy,
 	float3 sigma_x_y_rho, float l_dist)
 {
 	// Convert surface coordinates to appropriate scale for level
-	float pyrSize = pyramidSize(l);
+	float pyrSize = pyramidSize(GlintLutParams, l);
 	st[0] = st[0] * pyrSize - 0.5f;
 	st[1] = st[1] * pyrSize - 0.5f;
 	dst0[0] *= pyrSize;
@@ -338,7 +342,9 @@ float P22__glint_discrete_LOD(int l, float2 slope_h, float2 st, float2 dst0,
 				// Weighting function used in pbrt-v3 EWA function
 				float alpha = 2;
 				float W_P = exp(-alpha * r2) - exp(-alpha);
-				sum += P22_M(slope_h, l, is, it,
+				sum += P22_M(
+					GlintLutParams,
+					slope_h, l, is, it,
 					slope_dx, slope_dy,
 					sigma_x_y_rho, l_dist) * W_P;
 				sumWts += W_P;
@@ -373,6 +379,11 @@ float3 f_P(float3 wo, float3 wi, float3 sigmas_rho, float LogMicrofacetDensity,
 	if (dot(wo, wh) <= 0. || dot(wi, wh) <= 0.)
 		return float3(0., 0., 0.);
 
+	FGlintLutParams GlintLutParams;
+	GlintLutParams.Alpha	= View.GlintLUTParameters.x;
+	GlintLutParams.N		= asint(View.GlintLUTParameters.y);
+	GlintLutParams.NLevels	= asint(View.GlintLUTParameters.z);
+
 	// Compute texture derivatives
 	float2 texCoord = TexCoord;
 	float2 dst0 = ddx(texCoord);
@@ -438,32 +449,32 @@ float3 f_P(float3 wo, float3 wi, float3 sigmas_rho, float LogMicrofacetDensity,
 	{
 		// Choose LOD
 		float l =
-			max(0., Dictionary_NLevels - 1. + log2(minorLength));
+			max(0., GlintLutParams.NLevels - 1. + log2(minorLength));
 		int il = int(floor(l));
 
 		float w = l - float(il);
 
 		// Number of microfacets in a cell at level il
-		float n_il = pow(2., float(2 * il - (2 * (Dictionary_NLevels - 1))));
+		float n_il = pow(2., float(2 * il - (2 * (GlintLutParams.NLevels - 1))));
 		n_il *= exp(LogMicrofacetDensity);
 		float LOD_dist_il = log(n_il) / 1.38629; // 2. * log(2) = 1.38629
 
 		// Number of microfacets in a cell at level il + 1
-		float n_ilp1 = pow(2., float(2 * (il + 1) - (2 * (Dictionary_NLevels - 1))));
+		float n_ilp1 = pow(2., float(2 * (il + 1) - (2 * (GlintLutParams.NLevels - 1))));
 		n_ilp1 *= exp(LogMicrofacetDensity);
 		float LOD_dist_ilp1 = log(n_ilp1) / 1.38629; // 2. * log(2) = 1.38629
 
 		float P22_P_il = non_axis_aligned_anisotropic_beckmann(slope_h.x, slope_h.y, sigmas_rho.x, sigmas_rho.y, sigmas_rho.z);
 		float P22_P_ilp1 = P22_P_il;
 
-		if (int(round(LOD_dist_il)) < Dictionary_NLevels)
+		if (int(round(LOD_dist_il)) < GlintLutParams.NLevels)
 		{
-			P22_P_il = P22__glint_discrete_LOD(il, slope_h, texCoord, dst0, dst1, slope_dx, slope_dy, sigmas_rho, LOD_dist_il);
+			P22_P_il = P22__glint_discrete_LOD(GlintLutParams, il, slope_h, texCoord, dst0, dst1, slope_dx, slope_dy, sigmas_rho, LOD_dist_il);
 		}
 
-		if (int(round(LOD_dist_ilp1)) < Dictionary_NLevels)
+		if (int(round(LOD_dist_ilp1)) < GlintLutParams.NLevels)
 		{
-			P22_P_ilp1 = P22__glint_discrete_LOD(il + 1, slope_h, texCoord, dst0, dst1, slope_dx, slope_dy, sigmas_rho, LOD_dist_ilp1);
+			P22_P_ilp1 = P22__glint_discrete_LOD(GlintLutParams, il + 1, slope_h, texCoord, dst0, dst1, slope_dx, slope_dy, sigmas_rho, LOD_dist_ilp1);
 		}
 
 		P22_P = lerp(P22_P_il, P22_P_ilp1, w);
diff --git a/Engine/Source/Runtime/Engine/Public/SceneView.h b/Engine/Source/Runtime/Engine/Public/SceneView.h
--- a/Engine/Source/Runtime/Engine/Public/SceneView.h
+++ b/Engine/Source/Runtime/Engine/Public/SceneView.h
@@ -893,6 +893,7 @@ enum ETranslucencyVolumeCascade
 	VIEW_UNIFORM_BUFFER_MEMBER(float, SubSurfaceColorAsTransmittanceAtDistanceInMeters) \
 	VIEW_UNIFORM_BUFFER_MEMBER_PER_VIEW(FVector4f, TanAndInvTanHalfFOV) \
 	VIEW_UNIFORM_BUFFER_MEMBER_PER_VIEW(FVector4f, PrevTanAndInvTanHalfFOV) \
+	VIEW_UNIFORM_BUFFER_MEMBER_PER_VIEW(FVector4f, GlintLUTParameters) \
 
 /** The uniform shader parameters associated with a view. */
 #define VIEW_UNIFORM_BUFFER_MEMBER(type, identifier) \
diff --git a/Engine/Source/Runtime/Renderer/Private/SceneRendering.cpp b/Engine/Source/Runtime/Renderer/Private/SceneRendering.cpp
--- a/Engine/Source/Runtime/Renderer/Private/SceneRendering.cpp
+++ b/Engine/Source/Runtime/Renderer/Private/SceneRendering.cpp
@@ -1952,6 +1952,11 @@ void FViewInfo::SetupUniformBufferParameters(
 	if (ViewState)
 	{
 		ViewUniformShaderParameters.GlintTexture = ViewState->GlintShadingLUTsData.RHIGlintShadingLUTs ? ViewState->GlintShadingLUTsData.RHIGlintShadingLUTs : nullptr;
+		ViewUniformShaderParameters.GlintLUTParameters = FVector4f(
+			ViewState->GlintShadingLUTsData.Dictionary_Alpha,
+			*reinterpret_cast<float*>(&ViewState->GlintShadingLUTsData.Dictionary_N),
+			*reinterpret_cast<float*>(&ViewState->GlintShadingLUTsData.Dictionary_NLevels),
+			0.0f);
 	}
 	ViewUniformShaderParameters.GlintTexture = OrBlack2DArrayIfNull(ViewUniformShaderParameters.GlintTexture);
 
diff --git a/Engine/Source/Runtime/Renderer/Private/Strata/Glint/GlintShadingLUTs.cpp b/Engine/Source/Runtime/Renderer/Private/Strata/Glint/GlintShadingLUTs.cpp
--- a/Engine/Source/Runtime/Renderer/Private/Strata/Glint/GlintShadingLUTs.cpp
+++ b/Engine/Source/Runtime/Renderer/Private/Strata/Glint/GlintShadingLUTs.cpp
@@ -22,10 +22,12 @@ void FGlintShadingLUTsStateData::Init(FRDGBuilder& GraphBuilder, FViewInfo& View
 	{
 		if (View.ViewState && View.ViewState->GlintShadingLUTsData.GlintShadingLUTs == nullptr)
 		{
+			FGlintShadingLUTsStateData& GlintShadingLUTsStateData = View.ViewState->GlintShadingLUTsData;
+
 			const uint32 MipCount = 7;
 			const EPixelFormat PixelFormat = PF_FloatR11G11B10;
 
-			View.ViewState->GlintShadingLUTsData.GlintShadingLUTs = GRenderTargetPool.FindFreeElement(
+			GlintShadingLUTsStateData.GlintShadingLUTs = GRenderTargetPool.FindFreeElement(
 				FRDGTextureDesc::Create2DArray(FIntPoint(GlintLut1DArrayMip0Width, 1), PixelFormat, FClearValueBinding::None, TexCreate_ShaderResource, GlintLut1DArrayMip0Height, MipCount), TEXT("Material.Glints"));
 
 			auto UploadGlintTextureArrayMipData = [&](uint32 LutArrayIndex, uint32 LutMipLevel, uint32 LutMipWidth, uint8* SrcData)
@@ -48,13 +50,13 @@ void FGlintShadingLUTsStateData::Init(FRDGBuilder& GraphBuilder, FViewInfo& View
 				uint32 SrcStride = SrcBytesPerPixel * LutMipWidth;
 
 				uint32 DstStride;
-				uint8* Dst = (uint8*)GraphBuilder.RHICmdList.LockTexture2DArray(View.ViewState->GlintShadingLUTsData.GlintShadingLUTs->GetRHI()->GetTexture2DArray(), LutArrayIndex, LutMipLevel, RLM_WriteOnly, DstStride, false);
+				uint8* Dst = (uint8*)GraphBuilder.RHICmdList.LockTexture2DArray(GlintShadingLUTsStateData.GlintShadingLUTs->GetRHI()->GetTexture2DArray(), LutArrayIndex, LutMipLevel, RLM_WriteOnly, DstStride, false);
 
 				const uint8* Src = SrcData +LutArrayIndex * SrcStride;
 
 				FMemory::Memcpy(Dst, Src, SrcStride);
 
-				GraphBuilder.RHICmdList.UnlockTexture2DArray(View.ViewState->GlintShadingLUTsData.GlintShadingLUTs->GetRHI()->GetTexture2DArray(), LutArrayIndex, LutMipLevel, false);
+				GraphBuilder.RHICmdList.UnlockTexture2DArray(GlintShadingLUTsStateData.GlintShadingLUTs->GetRHI()->GetTexture2DArray(), LutArrayIndex, LutMipLevel, false);
 			};
 
 			const uint32 GlintLut1DArraySize = GlintLut1DArrayMip0Height;
@@ -86,6 +88,8 @@ void FGlintShadingLUTsStateData::Init(FRDGBuilder& GraphBuilder, FViewInfo& View
 			{
 				UploadGlintTextureArrayMipData(LutArrayIndex, 6, GlintLut1DArrayMip6Width, (uint8*)GlintLut1DArrayMip6RGBA);
 			}
+
+			GlintShadingLUTsStateData.SetDictionaryParameter(16, 64, 0.5);	// Match the embedded LUT
 		}
 	}
 	else
@@ -108,14 +112,27 @@ void FGlintShadingLUTsStateData::Init(FRDGBuilder& GraphBuilder, FViewInfo& View
 	// Use the conditionally loaded asset
 	if (View.ViewState)
 	{
-		View.ViewState->GlintShadingLUTsData.GlintShadingLUTs = nullptr;
-		View.ViewState->GlintShadingLUTsData.RHIGlintShadingLUTs = nullptr;
+		FGlintShadingLUTsStateData& GlintShadingLUTsStateData = View.ViewState->GlintShadingLUTsData;
+
+		GlintShadingLUTsStateData.GlintShadingLUTs = nullptr;
+		GlintShadingLUTsStateData.RHIGlintShadingLUTs = nullptr;
 		if (GEngine->GlintTexture)
 		{
-			View.ViewState->GlintShadingLUTsData.RHIGlintShadingLUTs = GEngine->GlintTexture->GetResource()->TextureRHI->GetTexture2DArray();
+			GlintShadingLUTsStateData.RHIGlintShadingLUTs = GEngine->GlintTexture->GetResource()->TextureRHI->GetTexture2DArray();
+
+			GlintShadingLUTsStateData.SetDictionaryParameter(16, 64, 0.5);	// Match the embedded LUT
 		}
 	}
 
 #endif
+}
+
+void FGlintShadingLUTsStateData::SetDictionaryParameter(int32 InNumberOfLevels, int32 InNumberOfDistributionsPerChannel, float InDictionaryAlpha)
+{
+	Dictionary_Alpha					= InDictionaryAlpha;
+	Dictionary_NLevels					= InNumberOfLevels;
+	Dictionary_NDistributionsPerChannel = InNumberOfDistributionsPerChannel;
+	Dictionary_N						= Dictionary_NDistributionsPerChannel * 3;
+	Dictionary_Pyramid0Size				= 1u << (Dictionary_NLevels - 1u);
+}
 
-}
\ No newline at end of file
diff --git a/Engine/Source/Runtime/Renderer/Private/Strata/Glint/GlintShadingLUTs.h b/Engine/Source/Runtime/Renderer/Private/Strata/Glint/GlintShadingLUTs.h
--- a/Engine/Source/Runtime/Renderer/Private/Strata/Glint/GlintShadingLUTs.h
+++ b/Engine/Source/Runtime/Renderer/Private/Strata/Glint/GlintShadingLUTs.h
@@ -16,7 +16,16 @@ struct FGlintShadingLUTsStateData
 	TRefCountPtr<IPooledRenderTarget> GlintShadingLUTs = nullptr;
 	FRHITexture2DArray* RHIGlintShadingLUTs = nullptr;
 
+	float Dictionary_Alpha = 0.0f;
+	int32 Dictionary_NDistributionsPerChannel = 0;
+	int32 Dictionary_N = 0;
+	int32 Dictionary_NLevels = 0;
+	int32 Dictionary_Pyramid0Size = 0;
+
 	uint64 GetGPUSizeBytes(bool bLogSizes) const;
 
 	static void Init(FRDGBuilder& GraphBuilder, FViewInfo& View);
+
+private:
+	void SetDictionaryParameter(int32 NumberOfLevels, int32 NumberOfDistributionsPerChannel, float Dictionary_Alpha);
 };
