Fixed inconsistancies in hair lighting due to pow(0,0) returning nan or 1.0 on different platforms

diff --git a/Engine/Shaders/Private/HairBsdf.ush b/Engine/Shaders/Private/HairBsdf.ush
--- a/Engine/Shaders/Private/HairBsdf.ush
+++ b/Engine/Shaders/Private/HairBsdf.ush
@@ -34,7 +34,8 @@ float3 KajiyaKayDiffuseAttenuation(FGBufferData GBuffer, float3 L, float3 V, hal
 	float NoL = saturate((dot(N, L) + Wrap) / Square(1 + Wrap));
 	float DiffuseScatter = (1 / PI) * lerp(NoL, KajiyaDiffuse, 0.33) * GBuffer.Metallic;
 	float Luma = Luminance(GBuffer.BaseColor);
-	float3 ScatterTint = pow(abs(GBuffer.BaseColor / max(Luma, MinValue)), 1 - Shadow);
+    float3 BaseOverLuma = abs(GBuffer.BaseColor / max(Luma, MinValue));
+	float3 ScatterTint = Shadow < 1 ? pow(BaseOverLuma, 1 - Shadow) : BaseOverLuma;
 	return sqrt(abs(GBuffer.BaseColor)) * DiffuseScatter * ScatterTint;
 }
 
