Move shader swap function to Common.ush and add variants.

diff --git a/Engine/Shaders/Private/BCCompressionCommon.ush b/Engine/Shaders/Private/BCCompressionCommon.ush
--- a/Engine/Shaders/Private/BCCompressionCommon.ush
+++ b/Engine/Shaders/Private/BCCompressionCommon.ush
@@ -487,27 +487,6 @@ uint4 CompressBC3BlockYCoCg(in float3 Block[16])
 	return uint4((MinLumaUint << 8) | MaxLumaUint | Indices.x, Indices.y, (MinColor565 << 16) | MaxColor565, CoCgIndices);
 }
 
-void Swap(inout uint3 A, inout uint3 B)
-{
-	uint3 Temp = A;
-	A = B;
-	B = Temp;
-}
-
-void Swap(inout float3 A, inout float3 B)
-{
-	float3 Temp = A;
-	A = B;
-	B = Temp;
-}
-
-void Swap(inout float A, inout float B)
-{
-	float Temp = A;
-	A = B;
-	B = Temp;
-}
-
 float3 Quantize10(float3 X)
 {
 	return (f32tof16(X) * 1024.0f) / (0x7bff + 1.0f);
diff --git a/Engine/Shaders/Private/Common.ush b/Engine/Shaders/Private/Common.ush
--- a/Engine/Shaders/Private/Common.ush
+++ b/Engine/Shaders/Private/Common.ush
@@ -2060,6 +2060,17 @@ struct FPixelShaderOut
 };
 
 
+#define SwapGeneric(T) void Swap(inout T A, inout T B) { T Temp = A; A = B; B = Temp; }
+SwapGeneric(uint)
+SwapGeneric(uint2)
+SwapGeneric(uint3)
+SwapGeneric(uint4)
+
+SwapGeneric(float)
+SwapGeneric(float2)
+SwapGeneric(float3)
+SwapGeneric(float4)
+
 // ---------------------------------------------------- Global samplers.
 
 // If GetGlobalSampler() was not implemented in Platform.ush, provide a default one from View uniform buffer.
