Substrate - Glint - fixed glint aliasing issue due to bd position of derivatives.

diff --git a/Engine/Shaders/Private/Strata/Glint/GlintThirdParty.ush b/Engine/Shaders/Private/Strata/Glint/GlintThirdParty.ush
--- a/Engine/Shaders/Private/Strata/Glint/GlintThirdParty.ush
+++ b/Engine/Shaders/Private/Strata/Glint/GlintThirdParty.ush
@@ -364,7 +364,7 @@ float P22__glint_discrete_LOD(
 //=============================================================================
 //======== Evaluation of the procedural physically based glinty BRDF ==========
 //=============================================================================
-float3 f_P(float3 wo, float3 wi, float3 sigmas_rho, float LogMicrofacetDensity, float2 TexCoord)
+float3 f_P(float3 wo, float3 wi, float3 sigmas_rho, float LogMicrofacetDensity, float2 TexCoord, float2 TexCoordDDX, float2 TexCoordDDY)
 {
 	if (wo.z <= 0.)
 		return float3(0., 0., 0.);
@@ -385,9 +385,8 @@ float3 f_P(float3 wo, float3 wi, float3 sigmas_rho, float LogMicrofacetDensity,
 	GlintLutParams.NLevels	= asint(View.GlintLUTParameters.z);
 
 	// Compute texture derivatives
-	float2 texCoord = TexCoord;
-	float2 dst0 = ddx(texCoord);
-	float2 dst1 = ddy(texCoord);
+	float2 dst0 = TexCoordDDX;
+	float2 dst1 = TexCoordDDY;
 
 	// Normal to slope
 	float2 slope_h = float2(-wh.x / wh.z, -wh.y / wh.z);
@@ -469,12 +468,12 @@ float3 f_P(float3 wo, float3 wi, float3 sigmas_rho, float LogMicrofacetDensity,
 
 		if (int(round(LOD_dist_il)) < GlintLutParams.NLevels)
 		{
-			P22_P_il = P22__glint_discrete_LOD(GlintLutParams, il, slope_h, texCoord, dst0, dst1, slope_dx, slope_dy, sigmas_rho, LOD_dist_il);
+			P22_P_il = P22__glint_discrete_LOD(GlintLutParams, il, slope_h, TexCoord, dst0, dst1, slope_dx, slope_dy, sigmas_rho, LOD_dist_il);
 		}
 
 		if (int(round(LOD_dist_ilp1)) < GlintLutParams.NLevels)
 		{
-			P22_P_ilp1 = P22__glint_discrete_LOD(GlintLutParams, il + 1, slope_h, texCoord, dst0, dst1, slope_dx, slope_dy, sigmas_rho, LOD_dist_ilp1);
+			P22_P_ilp1 = P22__glint_discrete_LOD(GlintLutParams, il + 1, slope_h, TexCoord, dst0, dst1, slope_dx, slope_dy, sigmas_rho, LOD_dist_ilp1);
 		}
 
 		P22_P = lerp(P22_P_il, P22_P_ilp1, w);
diff --git a/Engine/Shaders/Private/Strata/Strata.ush b/Engine/Shaders/Private/Strata/Strata.ush
--- a/Engine/Shaders/Private/Strata/Strata.ush
+++ b/Engine/Shaders/Private/Strata/Strata.ush
@@ -97,6 +97,10 @@
 
 #define STRATA_COMPLEXPATH (!STRATA_FASTPATH && !STRATA_SINGLEPATH)
 
+// Disable glint when in tracing mode for now since derivative are not available.
+#define STRATA_GLINTS_DERIVATIVE_SUPPORTED (!RAYCALLABLESHADER && !RAYHITGROUPSHADER && !RAYGENSHADER && !RAYMISSSHADER)
+#define STRATA_GLINTS_ENABLED (STRATA_COMPLEXPATH && STRATA_GLINTS && STRATA_GLINTS_DERIVATIVE_SUPPORTED)
+
 #define STRATA_LEGACY_MATERIAL_APPLIES_FINAL_WEIGHT (!MATERIALBLENDING_MASKED && !MATERIALBLENDING_SOLID && STRATA_LEGACY_PREMULT_ALPHA_OVERRIDE==0 || STRATA_MATERIAL_EXPORT_FROM_TRANSLUCENT)
 
 #define STRATA_INLINE_SINGLELAYERWATER (STRATA_ENABLED && MATERIAL_IS_STRATA && STRATA_INLINE_SHADING && MATERIAL_SHADINGMODEL_SINGLELAYERWATER)
@@ -157,8 +161,8 @@ struct FStrataBSDF
 	float4 InlineVGPRs[3];
 #endif
 
-#if STRATA_GLINTS
-	float4 VGPRs[6];
+#if STRATA_GLINTS_ENABLED
+	float4 VGPRs[7];
 #else
 	float4 VGPRs[5];
 #endif
@@ -525,6 +529,8 @@ struct FStrataBSDF
 // Glints
 #define SLAB_GLINT_UV(X)				X.VGPRs[5].xy
 #define SLAB_GLINT_VALUE(X)				X.VGPRs[5].z
+#define SLAB_GLINT_UVDDX(X)				X.VGPRs[6].xy
+#define SLAB_GLINT_UVDDY(X)				X.VGPRs[6].zw
 #endif
 
 ///////////////////////////////////////////////////////////////////////////////
@@ -1558,12 +1564,14 @@ FStrataData GetStrataSlabBSDF(
 	BSDF_SETTHICKNESSCM		(StrataData.InlinedBSDF,		ThicknessCm);
 #endif
 
-#if STRATA_COMPLEXPATH && STRATA_GLINTS
+#if STRATA_GLINTS_ENABLED
 	if (GlintValue > 0.0f)
 	{
 		BSDF_SETHASGLINT(StrataData.InlinedBSDF, 1);
 		SLAB_GLINT_VALUE(StrataData.InlinedBSDF)		= GlintValue;
 		SLAB_GLINT_UV(StrataData.InlinedBSDF)			= GlintUV;
+		SLAB_GLINT_UVDDX(StrataData.InlinedBSDF)		= ddx(GlintUV);
+		SLAB_GLINT_UVDDY(StrataData.InlinedBSDF)		= ddy(GlintUV);
 	}
 #endif
 
@@ -2367,10 +2375,12 @@ FStrataData StrataHorizontalMixingParameterBlending(FStrataData Background, FStr
 	SLAB_FUZZ_ROUGHNESS(ResultBSDF)	=	lerp(SLAB_FUZZ_ROUGHNESS(BackgroundBSDF),		SLAB_FUZZ_ROUGHNESS(ForegroundBSDF),		ForegroundOtherMixFactor);
 	SLAB_HAZINESS(ResultBSDF) =			LerpHazinessParameterBlending(BackgroundBSDF,	ForegroundBSDF,								ForegroundOtherMixFactor);
 
-#if STRATA_GLINTS
+#if STRATA_GLINTS_ENABLED
 	BSDF_SETHASGLINT(ResultBSDF,		BSDF_GETHASGLINT(ForegroundBSDF)				| BSDF_GETHASGLINT(BackgroundBSDF));
 	SLAB_GLINT_UV(ResultBSDF) =			lerp(SLAB_GLINT_UV(BackgroundBSDF),				SLAB_GLINT_UV(ForegroundBSDF),				ForegroundOtherMixFactor);
 	SLAB_GLINT_VALUE(ResultBSDF) =		lerp(SLAB_GLINT_VALUE(BackgroundBSDF),			SLAB_GLINT_VALUE(ForegroundBSDF),			ForegroundOtherMixFactor);
+	SLAB_GLINT_UVDDX(ResultBSDF) =		lerp(SLAB_GLINT_UVDDX(BackgroundBSDF),			SLAB_GLINT_UVDDX(ForegroundBSDF),			ForegroundOtherMixFactor);
+	SLAB_GLINT_UVDDY(ResultBSDF) =		lerp(SLAB_GLINT_UVDDY(BackgroundBSDF),			SLAB_GLINT_UVDDY(ForegroundBSDF),			ForegroundOtherMixFactor);
 #endif
 	
 
@@ -2480,10 +2490,12 @@ FStrataData StrataAddParameterBlending(FStrataData A, FStrataData B, float AMixF
 	SLAB_FUZZ_ROUGHNESS(ResultBSDF) =	lerp(SLAB_FUZZ_ROUGHNESS(BBSDF),				SLAB_FUZZ_ROUGHNESS(ABSDF),					AMixFactor);
 	SLAB_HAZINESS(ResultBSDF) =			LerpHazinessParameterBlending(BBSDF,			ABSDF,										AMixFactor);
 
-#if STRATA_GLINTS
+#if STRATA_GLINTS_ENABLED
 	BSDF_SETHASGLINT(ResultBSDF,		BSDF_GETHASGLINT(BBSDF)							| BSDF_GETHASGLINT(ABSDF));
 	SLAB_GLINT_UV(ResultBSDF) =			lerp(SLAB_GLINT_UV(BBSDF),						SLAB_GLINT_UV(ABSDF), AMixFactor);
 	SLAB_GLINT_VALUE(ResultBSDF) =		lerp(SLAB_GLINT_VALUE(BBSDF),					SLAB_GLINT_VALUE(ABSDF), AMixFactor);
+	SLAB_GLINT_UVDDX(ResultBSDF) =		lerp(SLAB_GLINT_UVDDX(BBSDF),					SLAB_GLINT_UVDDX(ABSDF), AMixFactor);
+	SLAB_GLINT_UVDDY(ResultBSDF) =		lerp(SLAB_GLINT_UVDDY(BBSDF),					SLAB_GLINT_UVDDY(ABSDF), AMixFactor);
 #endif
 
 	if (BSDF_GETSSSTYPE(ResultBSDF) == SSS_TYPE_DIFFUSION_PROFILE)
@@ -2564,7 +2576,7 @@ FStrataBSDF StrataVerticalLayeringParameterBlending(in FStrataBSDF TopBSDF, in F
 	BSDF_SETHASF90(ResultBSDF,			BSDF_GETHASF90(BaseBSDF)						| BSDF_GETHASF90(TopBSDF));			// F90: keep union of both even though it will be hard to get a match
 	BSDF_SETHASFUZZ(ResultBSDF,			BSDF_GETHASFUZZ(BaseBSDF)						| BSDF_GETHASFUZZ(TopBSDF));		// Fuzz: keep union of both even though it will be hard to get a match
 
-#if STRATA_GLINTS
+#if STRATA_GLINTS_ENABLED
 	BSDF_SETHASGLINT(ResultBSDF, 0); // Disable for now
 #endif
 
@@ -2604,10 +2616,12 @@ FStrataBSDF StrataVerticalLayeringParameterBlending(in FStrataBSDF TopBSDF, in F
 	SLAB_FUZZ_ROUGHNESS(ResultBSDF) =	lerp(SLAB_FUZZ_ROUGHNESS(TopBSDF),				SLAB_FUZZ_ROUGHNESS(BaseBSDF),TopToBottomLerp);
 	SLAB_HAZINESS(ResultBSDF) =			LerpHazinessParameterBlending(TopBSDF,			BaseBSDF,					TopToBottomLerp);
 
-#if STRATA_GLINTS
+#if STRATA_GLINTS_ENABLED
 	BSDF_SETHASGLINT(ResultBSDF,		BSDF_GETHASGLINT(TopBSDF)						| BSDF_GETHASGLINT(BaseBSDF));
 	SLAB_GLINT_UV(ResultBSDF) =			lerp(SLAB_GLINT_UV(TopBSDF),					SLAB_GLINT_UV(BaseBSDF), TopToBottomLerp);
 	SLAB_GLINT_VALUE(ResultBSDF) =		lerp(SLAB_GLINT_VALUE(TopBSDF),					SLAB_GLINT_VALUE(BaseBSDF), TopToBottomLerp);
+	SLAB_GLINT_UVDDX(ResultBSDF) =		lerp(SLAB_GLINT_UVDDX(TopBSDF),					SLAB_GLINT_UVDDX(BaseBSDF), TopToBottomLerp);
+	SLAB_GLINT_UVDDY(ResultBSDF) =		lerp(SLAB_GLINT_UVDDY(TopBSDF),					SLAB_GLINT_UVDDY(BaseBSDF), TopToBottomLerp);
 #endif
 
 	// Anisotropy is disabled with vertical blending. 
@@ -4272,7 +4286,7 @@ FStrataBSDF UnpackStrataBSDFIn(FStrataMaterialContainer StrataBuffer, inout FStr
 		}
 
 
-	#if STRATA_GLINTS
+	#if STRATA_GLINTS_ENABLED
 		if (BSDF_GETHASGLINT(OutBSDF))
 		{
 			uint2 GlintPackedData;
@@ -4281,7 +4295,10 @@ FStrataBSDF UnpackStrataBSDFIn(FStrataMaterialContainer StrataBuffer, inout FStr
 
 			SLAB_GLINT_UV(OutBSDF).x	= UnpackR24F(GlintPackedData.x);
 			SLAB_GLINT_UV(OutBSDF).y	= UnpackR24F(GlintPackedData.y);
-			SLAB_GLINT_VALUE(OutBSDF)	= f16tof32((GlintPackedData.x>>24) | ((GlintPackedData.y>>16) & 0xFF00));
+			SLAB_GLINT_VALUE(OutBSDF)	= f16tof32((GlintPackedData.x>>24) | ((GlintPackedData.y>>16) & 0xFF00)); 
+
+			SLAB_GLINT_UVDDX(OutBSDF)	= UnpackFloat2FromUInt(StrataLoadUint1(StrataBuffer, StrataAddressing));
+			SLAB_GLINT_UVDDY(OutBSDF)	= UnpackFloat2FromUInt(StrataLoadUint1(StrataBuffer, StrataAddressing));
 		}
 	#endif
 	}
diff --git a/Engine/Shaders/Private/Strata/StrataEvaluation.ush b/Engine/Shaders/Private/Strata/StrataEvaluation.ush
--- a/Engine/Shaders/Private/Strata/StrataEvaluation.ush
+++ b/Engine/Shaders/Private/Strata/StrataEvaluation.ush
@@ -12,10 +12,6 @@
 
 #include "Strata.ush"
 
-// Disable glint when in tracing mode for now since derivative are not available.
-#define STRATA_GLINTS_DERIVATIVE_SUPPORTED (!RAYCALLABLESHADER && !RAYHITGROUPSHADER && !RAYGENSHADER && !RAYMISSSHADER)
-#define STRATA_GLINTS_ENABLED (STRATA_COMPLEXPATH && STRATA_GLINTS && STRATA_GLINTS_DERIVATIVE_SUPPORTED)
-
 #if STRATA_GLINTS_ENABLED
 #include "Glint/GlintThirdParty.ush"
 #endif
@@ -443,7 +439,8 @@ FStrataEvaluateResult StrataEvaluateBSDFCommon(FStrataBSDFContext BSDFContext, F
 						{
 							const float2 GGXRoughnessXY = float2(sqrtFast(AlphaXY.x), sqrtFast(AlphaXY.y));
 							FBeckmannDesc Beckmann = GGXToBeckmann(GGXRoughnessXY);
-							D = f_P(BSDFContext.TangentV, BSDFContext.TangentL, float3(Beckmann.Sigma, Beckmann.Rho), SLAB_GLINT_VALUE(BSDFContext.BSDF), SLAB_GLINT_UV(BSDFContext.BSDF));
+							D = f_P(BSDFContext.TangentV, BSDFContext.TangentL, float3(Beckmann.Sigma, Beckmann.Rho),
+								SLAB_GLINT_VALUE(BSDFContext.BSDF), SLAB_GLINT_UV(BSDFContext.BSDF), SLAB_GLINT_UVDDX(BSDFContext.BSDF), SLAB_GLINT_UVDDY(BSDFContext.BSDF));
 						}
 						else
 					#endif
@@ -477,7 +474,8 @@ FStrataEvaluateResult StrataEvaluateBSDFCommon(FStrataBSDFContext BSDFContext, F
 
 								// It is wrong to multiply again the GGXevaluation together with the glint function (representing Beckmann D function).
 								// But visually it works and that is an acceptable solution in the meantime we find something more correct.
-								RectLightSpec *= f_P(BSDFContext.TangentV, MeanLightLocalDirection, float3(Beckmann.Sigma.xx, Beckmann.Rho), SLAB_GLINT_VALUE(BSDFContext.BSDF), SLAB_GLINT_UV(BSDFContext.BSDF));
+								RectLightSpec *= f_P(BSDFContext.TangentV, MeanLightLocalDirection, float3(Beckmann.Sigma.xx, Beckmann.Rho), 
+									SLAB_GLINT_VALUE(BSDFContext.BSDF), SLAB_GLINT_UV(BSDFContext.BSDF), SLAB_GLINT_UVDDX(BSDFContext.BSDF), SLAB_GLINT_UVDDY(BSDFContext.BSDF));
 							}
 						#endif // STRATA_COMPLEXPATH
 
@@ -498,7 +496,8 @@ FStrataEvaluateResult StrataEvaluateBSDFCommon(FStrataBSDFContext BSDFContext, F
 								if (BSDF_GETHASGLINT(BSDFContext.BSDF))
 								{
 									FBeckmannDesc Beckmann = GGXToBeckmann(SafeRoughness);
-									D = f_P(BSDFContext.TangentV, BSDFContext.TangentL, float3(Beckmann.Sigma.xx, Beckmann.Rho), SLAB_GLINT_VALUE(BSDFContext.BSDF), SLAB_GLINT_UV(BSDFContext.BSDF));
+									D = f_P(BSDFContext.TangentV, BSDFContext.TangentL, float3(Beckmann.Sigma.xx, Beckmann.Rho),
+										SLAB_GLINT_VALUE(BSDFContext.BSDF), SLAB_GLINT_UV(BSDFContext.BSDF), SLAB_GLINT_UVDDX(BSDFContext.BSDF), SLAB_GLINT_UVDDY(BSDFContext.BSDF));
 								}
 								else
 							#endif // STRATA_COMPLEXPATH
@@ -515,7 +514,8 @@ FStrataEvaluateResult StrataEvaluateBSDFCommon(FStrataBSDFContext BSDFContext, F
 								if (BSDF_GETHASGLINT(BSDFContext.BSDF))
 								{
 									FBeckmannDesc Beckmann = GGXToBeckmann(SafeRoughness);
-									D = Energy * f_P(BSDFContext.TangentV, BSDFContext.TangentL, float3(Beckmann.Sigma.xx, Beckmann.Rho), SLAB_GLINT_VALUE(BSDFContext.BSDF), SLAB_GLINT_UV(BSDFContext.BSDF));
+									D = Energy * f_P(BSDFContext.TangentV, BSDFContext.TangentL, float3(Beckmann.Sigma.xx, Beckmann.Rho),
+										SLAB_GLINT_VALUE(BSDFContext.BSDF), SLAB_GLINT_UV(BSDFContext.BSDF), SLAB_GLINT_UVDDX(BSDFContext.BSDF), SLAB_GLINT_UVDDY(BSDFContext.BSDF));
 								}
 								else
 							#endif // STRATA_COMPLEXPATH
@@ -1488,7 +1488,8 @@ FStrataEnvLightResult StrataEvaluateForEnvLight(FStrataBSDFContext BSDFContext,
 		//
 		//			// It is wrong to multiply again the split env lighting (for split-sum approxiation) together with the glint function (representing Beckmann D function).
 		//			// But visually it works and that is an acceptable solution in the meantime we find something more correct.
-		//			float GlintD = f_P(BSDFContext.TangentV, MeanLightLocalDirection, float3(Beckmann.Sigma, Beckmann.Rho), SLAB_GLINT_VALUE(BSDFContext.BSDF), SLAB_GLINT_UV(BSDFContext.BSDF));
+		//			float GlintD = f_P(BSDFContext.TangentV, MeanLightLocalDirection, float3(Beckmann.Sigma, Beckmann.Rho), 
+		//								SLAB_GLINT_VALUE(BSDFContext.BSDF), SLAB_GLINT_UV(BSDFContext.BSDF), SLAB_GLINT_UVDDX(BSDFContext.BSDF), SLAB_GLINT_UVDDY(BSDFContext.BSDF));
 		//
 		//			// We need a way to lerp to the glint direction with a roughness of 0. 
 		//			// STRATA_TODO recover glint direction.
diff --git a/Engine/Shaders/Private/Strata/StrataExport.ush b/Engine/Shaders/Private/Strata/StrataExport.ush
--- a/Engine/Shaders/Private/Strata/StrataExport.ush
+++ b/Engine/Shaders/Private/Strata/StrataExport.ush
@@ -799,6 +799,11 @@ void PackStrataOut(
 								const uint Packed8H = (Packed16 >> 8  ) << 24;
 								STRATA_STORE_UINT1(PackR24F(SLAB_GLINT_UV(BSDF).x) | Packed8L);
 								STRATA_STORE_UINT1(PackR24F(SLAB_GLINT_UV(BSDF).y) | Packed8H);
+
+								// We have to store derivative when rasterizing the mesh. 
+								// That cannot be reconstructed in a deferred fashion (e.g. edge of a mesh overlaping with backgroud would look wrongly aliased)
+								STRATA_STORE_UINT1(PackFloat2ToUInt(SLAB_GLINT_UVDDX(BSDF)));
+								STRATA_STORE_UINT1(PackFloat2ToUInt(SLAB_GLINT_UVDDY(BSDF)));
 							}
 							#endif // STRATA_GLINTS
 
diff --git a/Engine/Source/Runtime/Engine/Private/Materials/HLSLMaterialTranslator.cpp b/Engine/Source/Runtime/Engine/Private/Materials/HLSLMaterialTranslator.cpp
--- a/Engine/Source/Runtime/Engine/Private/Materials/HLSLMaterialTranslator.cpp
+++ b/Engine/Source/Runtime/Engine/Private/Materials/HLSLMaterialTranslator.cpp
@@ -11720,6 +11720,8 @@ bool FHLSLMaterialTranslator::FStrataCompilationContext::StrataGenerateDerivedMa
 						{
 							StrataMaterialRequestedSizeByte += UintByteSize;
 							StrataMaterialRequestedSizeByte += UintByteSize;
+							StrataMaterialRequestedSizeByte += UintByteSize;
+							StrataMaterialRequestedSizeByte += UintByteSize;
 						}
 						break;
 					}
