Substrate - Glint - Fixed bad glint look with translucent forward shading.

diff --git a/Engine/Shaders/Private/Strata/StrataEvaluation.ush b/Engine/Shaders/Private/Strata/StrataEvaluation.ush
--- a/Engine/Shaders/Private/Strata/StrataEvaluation.ush
+++ b/Engine/Shaders/Private/Strata/StrataEvaluation.ush
@@ -173,6 +173,7 @@ void FStrataBSDFContext::StrataUpdateBSDFContext(float3 NewL)
 
 	this.TangentH = mul(this.TangentBasis, this.H);
 	this.TangentB = mul(this.TangentBasis, this.B);
+	this.TangentL = normalize(mul(this.TangentBasis, this.L));
 
 	this.SatNoL	 = saturate(this.Context.NoL);
 	this.SatNoV	 = saturate(this.Context.NoV);
