Enable env. evaluation for glints.

diff --git a/Engine/Shaders/Private/Strata/StrataEvaluation.ush b/Engine/Shaders/Private/Strata/StrataEvaluation.ush
--- a/Engine/Shaders/Private/Strata/StrataEvaluation.ush
+++ b/Engine/Shaders/Private/Strata/StrataEvaluation.ush
@@ -1479,25 +1479,25 @@ FStrataEnvLightResult StrataEvaluateForEnvLight(FStrataBSDFContext BSDFContext,
 				StrataEnvLightResult.SpecularColor = F0;
 
 				// STRATA_TODO environment light importance sampling
-		//	#if STRATA_GLINTS_ENABLED
-		//		BRANCH
-		//		if (BSDF_GETHASGLINT(BSDFContext.BSDF))
-		//		{
-		//			float3 MeanLightLocalDirection = normalize(mul(BSDFContext.TangentBasis, StrataEnvLightResult.SpecularDirection));
-		//
-		//			FBeckmannDesc Beckmann = GGXToBeckmann(SafeRoughness);
-		//
-		//			// It is wrong to multiply again the split env lighting (for split-sum approxiation) together with the glint function (representing Beckmann D function).
-		//			// But visually it works and that is an acceptable solution in the meantime we find something more correct.
-		//			float GlintD = f_P(BSDFContext.TangentV, MeanLightLocalDirection, float3(Beckmann.Sigma, Beckmann.Rho), 
-		//								SLAB_GLINT_VALUE(BSDFContext.BSDF), SLAB_GLINT_UV(BSDFContext.BSDF), SLAB_GLINT_UVDDX(BSDFContext.BSDF), SLAB_GLINT_UVDDY(BSDFContext.BSDF));
-		//
-		//			// We need a way to lerp to the glint direction with a roughness of 0. 
-		//			// STRATA_TODO recover glint direction.
-		//			StrataEnvLightResult.SpecularWeight *= saturate(GlintD);
-		//			StrataEnvLightResult.SpecularSafeRoughness = lerp(StrataEnvLightResult.SpecularSafeRoughness, 0.001f, saturate(GlintD));
-		//		}
-		//	#endif // STRATA_COMPLEXPATH
+			#if STRATA_GLINTS_ENABLED
+				BRANCH
+				if (BSDF_GETHASGLINT(BSDFContext.BSDF))
+				{
+					float3 MeanLightLocalDirection = normalize(mul(BSDFContext.TangentBasis, StrataEnvLightResult.SpecularDirection));
+		
+					FBeckmannDesc Beckmann = GGXToBeckmann(SafeRoughness);
+		
+					// It is wrong to multiply again the split env lighting (for split-sum approxiation) together with the glint function (representing Beckmann D function).
+					// But visually it works and that is an acceptable solution in the meantime we find something more correct.
+					float GlintD = f_P(BSDFContext.TangentV, MeanLightLocalDirection, float3(Beckmann.Sigma, Beckmann.Rho), 
+										SLAB_GLINT_VALUE(BSDFContext.BSDF), SLAB_GLINT_UV(BSDFContext.BSDF), SLAB_GLINT_UVDDX(BSDFContext.BSDF), SLAB_GLINT_UVDDY(BSDFContext.BSDF));
+		
+					// We need a way to lerp to the glint direction with a roughness of 0. 
+					// STRATA_TODO recover glint direction.
+					StrataEnvLightResult.SpecularWeight *= saturate(GlintD);
+					StrataEnvLightResult.SpecularSafeRoughness = lerp(StrataEnvLightResult.SpecularSafeRoughness, 0.001f, saturate(GlintD));
+				}
+			#endif // STRATA_COMPLEXPATH
 
 				StrataEnvLightResult.SSRSpecularWeight = StrataEnvLightResult.SpecularWeight;
 
