Add Glint importance sampling code.

diff --git a/Engine/Shaders/Private/Strata/Glint/GlintThirdParty.ush b/Engine/Shaders/Private/Strata/Glint/GlintThirdParty.ush
--- a/Engine/Shaders/Private/Strata/Glint/GlintThirdParty.ush
+++ b/Engine/Shaders/Private/Strata/Glint/GlintThirdParty.ush
@@ -25,6 +25,7 @@ struct FGlintLutParams
 	float Alpha;
 	int N;
 	int NLevels;
+	float MinLDist;
 };
 
 // Geometric Glint Anti-Aliasing parameters
@@ -119,6 +120,28 @@ float sampleNormalDistribution(float U, float mu, float sigma)
 	return x;
 }
 
+float2 sampleBivariateNormalDistribution(float2 u, float sigmaX, float sigmaY, float rho) 
+{
+	// Sample an original 2D centered unit-variance Gaussian
+	float sampleXO = sampleNormalDistribution(u.x, 0., 1.);
+	float sampleYO = sampleNormalDistribution(u.y, 0., 1.);
+	float2 sampleO = float2(sampleXO, sampleYO);
+
+	// Matrix M transformed the original unit-variance Gaussian
+	// into a non-axis aligned Gaussian with standard deviations sigmaX and
+	// sigmaY and correlation factor rho
+	// See Chermain et al. 2021 for more information
+	float M00 = sigmaX * sqrt(1. - rho * rho);
+	float M01 = rho * sigmaX;
+	float M10 = 0.;
+	float M11 = sigmaY;
+
+	// Transform the sample using the linear transformation M
+	return float2(
+		sampleO.x * M00 + sampleO.y * M01,
+		sampleO.x * M10 + sampleO.y * M11);
+}
+
 //=============================================================================
 //===================== Spatially-varying, multiscale, ========================
 //=============== and transformed slope distribution function  ================
@@ -148,7 +171,7 @@ float P22_M(
 	l_dist = sampleNormalDistribution(uDensityRandomisation, l_dist,
 		densityRandomisation);
 
-	l_dist = clamp(int(round(l_dist)), 0, GlintLutParams.NLevels);
+	l_dist = clamp(int(round(l_dist)), GlintLutParams.MinLDist, GlintLutParams.NLevels);
 
 	// Recover roughness and slope correlation factor
 	float sigma_x = sigma_x_y_rho.x;
@@ -496,3 +519,370 @@ float3 f_P(float3 wo, float3 wi, float3 sigmas_rho, float LogMicrofacetDensity,
 
 	return D_P / (4. * wo.z);
 }
+
+// Glint Sampling
+
+// Returns the ith 1D distribution at level l. It uses two random
+// numbers: one for sampling the positive 1D distribution, the other to
+// inverse the sign of the sampled value.
+float Sample_P_1D(FGlintLutParams GlintLutParams, float u, int ij, int l)
+{
+	// The distribution equals 0 after 4 standard deviations.
+	// Slope standard deviation = alpha / sqrt(2)
+
+	// 3 distributions values in one texel
+	int distIdxOver3 = ij / 3;
+	int distPerChannel = GlintLutParams.N / 3;
+	float alpha_dist_isqrt2_4 = GlintLutParams.Alpha * InvSqrt2 * 4.f;
+
+	// The 1D distribution is an even function
+	// [-1, 1]
+	float SignScale = 1;
+	if (u < 0.5)
+	{
+		SignScale *= -1.;
+		u = RescaleRandomNumber(u, 0.0, 0.5);
+	}
+	else
+	{
+		u = RescaleRandomNumber(u, 0.5, 1.f);
+	}
+
+	// Equation 15
+	const int Idx = int(fmod(ij, 3)); // ij % 3; // int(mod(ij, 3))
+
+	// Slow version - the CDF is computed inline within the shader (TODO precompute CDF textures)
+	const uint EntryCount = 64;
+	float Cdfs[EntryCount+1];
+	float Acc = 0;
+	LOOP
+	for (uint it=0;it<EntryCount;++it)
+	{
+		const float texCoordX = saturate(float(it) / float(EntryCount-1));
+		float V = View.GlintTexture.SampleLevel(View.GlintSampler, float3(texCoordX, 0.5, l * distPerChannel + distIdxOver3), 0.0f)[Idx];
+
+		Cdfs[it] = Acc;
+		Acc += V;
+	}
+	Acc = max(0.0001f, Acc);
+
+	LOOP
+	for (uint it = 0; it < EntryCount; ++it)
+	{
+		Cdfs[it] /= Acc;
+	}
+	Cdfs[EntryCount] = 1.f;
+
+	// Samples the positive part of the 1D distribution
+	// Sampled value is in [0, 1)
+	float sample = 0;
+	LOOP
+	for (uint it = 0; it < EntryCount; ++it)
+	{
+		if (Cdfs[it] <= u && u <= Cdfs[it + 1])
+		{
+			const float s = (u - Cdfs[it]) / max(0.0001f, (Cdfs[it+1]-Cdfs[it]));
+			sample = saturate((it + s) / float(EntryCount));
+			break;
+		}
+	}
+
+	// Output is in [-alpha_dist_isqrt2_4, alpha_dist_isqrt2_4]
+	return SignScale * sample * alpha_dist_isqrt2_4;
+}
+
+//=============================================================================
+//===================== Spatially-varying, multiscale, ========================
+//=============== and transformed slope distribution function  ================
+//=============================== Equation 4 ==================================
+//=============================================================================
+float2 Sample_P22_M(FGlintLutParams GlintLutParams, float2 u01, int l, int s0, int t0, float3 sigma_x_y_rho, float l_dist)
+{
+	// Coherent index
+	int twoToTheL = int(pow(2., float(l)));
+	s0 *= twoToTheL;
+	t0 *= twoToTheL;
+
+	// Seed pseudo random generator
+	uint rngSeed = uint(s0 + 1549 * t0);
+
+	float uDensityRandomisation = hashIQ(rngSeed * 2171U);
+
+	// Fix density randomisation to 2 to have better appearance
+	float densityRandomisation = 2.;
+
+	// Sample a Gaussian to randomise the distribution LOD around the 
+	// distribution level l_dist
+	l_dist = sampleNormalDistribution(uDensityRandomisation, l_dist, densityRandomisation);
+	l_dist = clamp(int(round(l_dist)), 0, GlintLutParams.NLevels);
+
+	// Recover roughness and slope correlation factor
+	float sigma_x = sigma_x_y_rho.x;
+	float sigma_y = sigma_x_y_rho.y;
+	float rho = sigma_x_y_rho.z;
+
+	// If we are too far from the surface, the SDF is a gaussian.
+	if (l_dist == GlintLutParams.NLevels) 
+	{
+		return sampleBivariateNormalDistribution(u01, sigma_x, sigma_y, rho);
+	}
+
+	// Random rotations to remove glint alignment
+	float uTheta = hashIQ(rngSeed);
+	float theta = 2.0 * PI * uTheta;
+
+	float cosTheta = cos(theta);
+	float sinTheta = sin(theta);
+
+	// Linear transformation represented by matrix M
+	float SIGMA_DICT = GlintLutParams.Alpha * InvSqrt2;
+	float M00 = 1. / SIGMA_DICT * (sigma_x * sqrt(1. - rho * rho));
+	float M01 = 1. / SIGMA_DICT * rho * sigma_y;
+	float M10 = 0.;
+	float M11 = 1. / SIGMA_DICT * sigma_y;
+
+	float MR00 =  cosTheta * M00 + sinTheta * M01;
+	float MR01 = -sinTheta * M00 + cosTheta * M01;
+	float MR10 =  cosTheta * M10 + sinTheta * M11;
+	float MR11 = -sinTheta * M10 + cosTheta * M11;
+
+	float u1 = hashIQ(rngSeed * 16807U);
+	float u2 = hashIQ(rngSeed * 48271U);
+
+	int i = int(u1 * float(GlintLutParams.N));
+	int j = int(u2 * float(GlintLutParams.N));
+
+	int il_dist = l_dist;
+
+	// Sample original (dictionary) slope space
+	const float slopeHXO = Sample_P_1D(GlintLutParams, u01.x, i, il_dist);
+	const float slopeHYO = Sample_P_1D(GlintLutParams, u01.y, j, il_dist);
+
+	float2 slope_h = float2(
+		slopeHXO * MR00 + slopeHYO * MR01,
+		slopeHXO * MR10 + slopeHYO * MR11);
+
+	return slope_h;
+}
+
+// Most of this function is similar to pbrt-v3 EWA function,
+// which itself is similar to Heckbert 1889 algorithm, 
+// http://www.cs.cmu.edu/~ph/texfund/texfund.pdf, Section 3.5.9.
+// Go through cells within the pixel footprint for a givin LOD
+float2 Sample_P22_LOD(
+	FGlintLutParams GlintLutParams,
+	float u0, float u1, float u2,
+	int l, float2 st, 
+	float2 dst0, float2 dst1, 
+	float3 sigma_x_y_rho, 
+	float l_dist)
+{
+	// Convert surface coordinates to appropriate scale for level
+	float pyrSize = pyramidSize(GlintLutParams, l);
+	st[0] = st[0] * pyrSize - 0.5f;
+	st[1] = st[1] * pyrSize - 0.5f;
+	dst0[0] *= pyrSize;
+	dst0[1] *= pyrSize;
+	dst1[0] *= pyrSize;
+	dst1[1] *= pyrSize;
+
+	// Compute ellipse coefficients to bound filter region
+	float A = dst0[1] * dst0[1] + dst1[1] * dst1[1] + 1.;
+	float B = -2. * (dst0[0] * dst0[1] + dst1[0] * dst1[1]);
+	float C = dst0[0] * dst0[0] + dst1[0] * dst1[0] + 1.;
+
+	float invF = 1. / (A * C - B * B * 0.25f);
+	A *= invF;
+	B *= invF;
+	C *= invF;
+
+	// Compute the ellipse's bounding box in texture space
+	float det = -B * B + 4 * A * C;
+	float invDet = 1 / det;
+	float uSqrt = sqrt(det * C), vSqrt = sqrt(A * det);
+	int s0 = int(ceil(st[0] - 2. * invDet * uSqrt));
+	int s1 = int(floor(st[0] + 2. * invDet * uSqrt));
+	int t0 = int(ceil(st[1] - 2. * invDet * vSqrt));
+	int t1 = int(floor(st[1] + 2. * invDet * vSqrt));
+
+	// The size is in [1-4]x[1-4]
+	int SizeT = t1-t0;
+	int SizeS = s1-s0;
+
+	// Scan over ellipse bound and compute quadratic equation
+
+	// 1. Count valid cells
+	// Cell picking is done over valid cell rather actual EWA weight
+	// The foot print is at max 4x4
+	uint2 selected_st = 0;
+	{
+		#define MAX_CELL_COUNT 32
+		float Cdfs[MAX_CELL_COUNT + 1];
+		uint2 Indices[MAX_CELL_COUNT];
+		float Acc = 0;
+		uint ValidCellCount = 0;
+		const int SafeMaxIterationCount = 100;
+
+		// 1.1 Build CDF
+		int nbrOfIter = 0;
+		uint ValidCount = 0;
+		LOOP
+		for (int it = t0; it <= t1; ++it)
+		{
+			float tt = it - st[1];
+			LOOP
+			for (int is = s0; is <= s1; ++is)
+			{
+				float ss = is - st[0];
+				// Compute squared radius 
+				// and filter SDF if inside ellipse
+				float r2 = A * ss * ss + B * ss * tt + C * tt * tt;
+				if (r2 < 1)
+				{
+					// Weighting function used in pbrt-v3 EWA function
+					float alpha = 2;
+					float W_P = exp(-alpha * r2) - exp(-alpha);
+					Cdfs[ValidCellCount] = Acc;
+					Indices[ValidCellCount] = uint2(is, it);
+					Acc += W_P;
+					++ValidCellCount;
+				}
+				nbrOfIter++;
+				// Guardrail (Extremely rare case.)
+				if (nbrOfIter > SafeMaxIterationCount)
+					break;
+			}
+			// Guardrail (Extremely rare case.)
+			if (nbrOfIter > SafeMaxIterationCount)
+				break;
+		}
+
+		// 1.2 Normalized CDF
+		LOOP
+		for (uint it = 0; it < ValidCellCount; ++it)
+		{
+			Cdfs[it] /= max(Acc, 0.0001f);
+		}
+		Cdfs[ValidCellCount] = 1;
+
+		// 1.3 Sample CDF
+		LOOP
+		for (uint it = 0; it < ValidCellCount; ++it)
+		{
+			if (Cdfs[it] <= u0 && u0 <= Cdfs[it + 1])
+			{
+				selected_st = Indices[it];
+				break;
+			}
+		}
+	}
+
+	// 2. Sample valid cell and return slope_h
+	return Sample_P22_M(GlintLutParams, float2(u0, u1), l, selected_st.x, selected_st.y, sigma_x_y_rho, l_dist);
+}
+
+float3 Sample_P(float3 u, float3 wo, float3 sigmas_rho, float LogMicrofacetDensity, float2 TexCoord, float2 TexCoordDDX, float2 TexCoordDDY)
+{
+	FGlintLutParams GlintLutParams;
+	GlintLutParams.Alpha	= View.GlintLUTParameters.x;
+	GlintLutParams.N		= asint(View.GlintLUTParameters.y);
+	GlintLutParams.NLevels	= asint(View.GlintLUTParameters.z);
+
+	float U0 = u.x;
+	float U1 = u.y;
+	float U2 = u.z;
+
+	// Compute texture derivatives
+	float2 texCoord = TexCoord;
+	float2 dst0 = TexCoordDDX;
+	float2 dst1 = TexCoordDDY;
+
+	float3 OutWm = float3(0,0,1); // z+
+	if (wo.z <= 0.)
+	{
+		return OutWm;
+	}
+
+	float3 D_P = float3(0.0f, 0.0f, 0.0f);
+	float P22_P = 0.;
+
+	//=========================================================================
+	// Similar to pbrt-v3 MIPMap::Lookup function, 
+	// http://www.pbr-book.org/3ed-2018/Texture/Image_Texture.html#EllipticallyWeightedAverage
+
+	// Compute ellipse minor and major axes
+	float dst0LengthSquared = dot(dst0, dst0);
+	float dst1LengthSquared = dot(dst1, dst1);
+
+	if (dst0LengthSquared < dst1LengthSquared)
+	{
+		Swap(dst0, dst1);
+	}
+	float majorLength = length(dst0);
+	float minorLength = length(dst1);
+
+	// Clamp ellipse eccentricity if too large
+	if (minorLength * MaxAnisotropy < majorLength && minorLength > 0.)
+	{
+		float scale = majorLength / (minorLength * MaxAnisotropy);
+		dst1 *= scale;
+		minorLength *= scale;
+	}
+	//=========================================================================
+	
+	// Compute LOD
+	if (minorLength > 0)
+	{
+		const float l = max(0., GlintLutParams.NLevels - 1. + log2(minorLength));
+		const int il = int(floor(l));
+		const float wLOD = l - float(il);
+
+		// Number of microfacets in a cell at level il
+		float n_il = pow(2., float(2 * il - (2 * (GlintLutParams.NLevels - 1))));
+		n_il *= exp(LogMicrofacetDensity);
+		float LOD_dist_il = log(n_il) / 1.38629; // 2. * log(2) = 1.38629
+
+		// Number of microfacets in a cell at level il + 1
+		float n_ilp1 = pow(2., float(2 * (il + 1) - (2 * (GlintLutParams.NLevels - 1))));
+		n_ilp1 *= exp(LogMicrofacetDensity);
+		float LOD_dist_ilp1 = log(n_ilp1) / 1.38629; // 2. * log(2) = 1.38629
+
+		// 1. Select LOD (N or N+1)
+		uint LOD = il;
+		float LODDist = LOD_dist_il;
+		if (U2 < wLOD)
+		{
+			U2 = RescaleRandomNumber(U2, 0.0, wLOD);
+			LOD = il;
+			LODDist = LOD_dist_il;
+		}
+		else
+		{
+			U2 = RescaleRandomNumber(U2, wLOD, 1.0);
+			LOD = il + 1;
+			LODDist = LOD_dist_ilp1;  //LOD_dist_il + 1;
+		}
+
+		// 2. Select cell and sample distribution
+		if (int(round(LODDist)) < GlintLutParams.NLevels)
+		{
+			const float2 slope_h = Sample_P22_LOD(GlintLutParams, U0, U1, U2, LOD, texCoord, dst0, dst1, sigmas_rho, LODDist);
+
+			// Build micro normal and its V-cavity dual
+			const float3 Wm = SlopeToNormal(slope_h);
+			const float3 WmP = float3(-Wm.x, -Wm.y, Wm.z);
+
+			// Projected area of the other facet of the V-cavity
+			const float dotWOWMP = max(dot(wo, WmP), 0);
+			const float dotWOWM = max(dot(wo, Wm), 0);
+			const float projectedAreaWMP = dotWOWMP / max(0.0001f, dotWOWM + dotWOWMP);
+
+			// Samples a facet of the V-cavity
+			if (U2 < projectedAreaWMP)
+				OutWm = WmP;
+			else
+				OutWm = Wm;
+		}
+	}
+	return OutWm;
+}
