Fix overflow issue and uninitialized value.

diff --git a/Engine/Shaders/Private/Strata/Glint/GlintThirdParty.ush b/Engine/Shaders/Private/Strata/Glint/GlintThirdParty.ush
--- a/Engine/Shaders/Private/Strata/Glint/GlintThirdParty.ush
+++ b/Engine/Shaders/Private/Strata/Glint/GlintThirdParty.ush
@@ -406,6 +406,7 @@ float3 f_P(float3 wo, float3 wi, float3 sigmas_rho, float LogMicrofacetDensity,
 	GlintLutParams.Alpha	= View.GlintLUTParameters.x;
 	GlintLutParams.N		= asint(View.GlintLUTParameters.y);
 	GlintLutParams.NLevels	= asint(View.GlintLUTParameters.z);
+	GlintLutParams.MinLDist = 0;
 
 	// Compute texture derivatives
 	float2 dst0 = TexCoordDDX;
@@ -716,7 +717,7 @@ float2 Sample_P22_LOD(
 	// The foot print is at max 4x4
 	uint2 selected_st = 0;
 	{
-		#define MAX_CELL_COUNT 32
+		#define MAX_CELL_COUNT 64
 		float Cdfs[MAX_CELL_COUNT + 1];
 		uint2 Indices[MAX_CELL_COUNT];
 		float Acc = 0;
@@ -739,11 +740,13 @@ float2 Sample_P22_LOD(
 				float r2 = A * ss * ss + B * ss * tt + C * tt * tt;
 				if (r2 < 1)
 				{
+					const uint Index = min(ValidCellCount, MAX_CELL_COUNT);
+
 					// Weighting function used in pbrt-v3 EWA function
 					float alpha = 2;
 					float W_P = exp(-alpha * r2) - exp(-alpha);
-					Cdfs[ValidCellCount] = Acc;
-					Indices[ValidCellCount] = uint2(is, it);
+					Cdfs[Index] = Acc;
+					Indices[Index] = uint2(is, it);
 					Acc += W_P;
 					++ValidCellCount;
 				}
@@ -757,6 +760,8 @@ float2 Sample_P22_LOD(
 				break;
 		}
 
+		ValidCellCount = min(ValidCellCount, MAX_CELL_COUNT);
+
 		// 1.2 Normalized CDF
 		LOOP
 		for (uint it = 0; it < ValidCellCount; ++it)
@@ -787,6 +792,7 @@ float3 Sample_P(float3 u, float3 wo, float3 sigmas_rho, float LogMicrofacetDensi
 	GlintLutParams.Alpha	= View.GlintLUTParameters.x;
 	GlintLutParams.N		= asint(View.GlintLUTParameters.y);
 	GlintLutParams.NLevels	= asint(View.GlintLUTParameters.z);
+	GlintLutParams.MinLDist = 0;
 
 	float U0 = u.x;
 	float U1 = u.y;
@@ -860,7 +866,7 @@ float3 Sample_P(float3 u, float3 wo, float3 sigmas_rho, float LogMicrofacetDensi
 		{
 			U2 = RescaleRandomNumber(U2, wLOD, 1.0);
 			LOD = il + 1;
-			LODDist = LOD_dist_ilp1;  //LOD_dist_il + 1;
+			LODDist = LOD_dist_il + 1;
 		}
 
 		// 2. Select cell and sample distribution
