Add optional glint importance sampling for sky evaluation.

diff --git a/Engine/Shaders/Private/Strata/StrataEvaluation.ush b/Engine/Shaders/Private/Strata/StrataEvaluation.ush
--- a/Engine/Shaders/Private/Strata/StrataEvaluation.ush
+++ b/Engine/Shaders/Private/Strata/StrataEvaluation.ush
@@ -12,8 +12,16 @@
 
 #include "Strata.ush"
 
+#ifndef STRATA_GLINT_IS
+#define STRATA_GLINT_IS 0 
+#endif
+
 #if STRATA_GLINTS_ENABLED
 #include "Glint/GlintThirdParty.ush"
+
+#if STRATA_GLINT_IS
+#include "../BlueNoise.ush"
+#endif 
 #endif
 
 #include "../ParticipatingMediaCommon.ush"
@@ -100,10 +108,12 @@ struct FStrataBSDFContext
 	float3 TangentB; 
 	float3 TangentL;
 	
+	uint2 PixelCoord; 
+
 	void StrataUpdateBSDFContext(float3 NewL);
 };
 
-FStrataBSDFContext StrataCreateBSDFContext(float3x3 TangentBasis, FStrataBSDF BSDF, float3 V, float3 L, bool bHasValidL=true)
+FStrataBSDFContext StrataCreateBSDFContext(float3x3 TangentBasis, FStrataBSDF BSDF, float3 V, float3 L, bool bHasValidL=true, uint2 InPixelCoord=0)
 {
 	FStrataBSDFContext BSDFContext = (FStrataBSDFContext)0;
 
@@ -146,6 +156,8 @@ FStrataBSDFContext StrataCreateBSDFContext(float3x3 TangentBasis, FStrataBSDF BS
 	BSDFContext.Context.NoH = clamp(BSDFContext.Context.NoH, 0.0f, 0.9999f);
 #endif
 
+	BSDFContext.PixelCoord = InPixelCoord;
+
 	return BSDFContext;
 }
 
@@ -153,13 +165,13 @@ FStrataBSDFContext StrataCreateBSDFContext(FStrataPixelHeader StrataPixelHeader,
 {
 	float3 UnusedL = float3(0, 0, 1);
 	float3x3 TangentBasis = StrataGetBSDFSharedBasis(StrataPixelHeader, BSDF, StrataAddressing);
-	return StrataCreateBSDFContext(TangentBasis, BSDF, V, UnusedL, false);
+	return StrataCreateBSDFContext(TangentBasis, BSDF, V, UnusedL, false, StrataAddressing.PixelCoords);
 }
 
 FStrataBSDFContext StrataCreateBSDFContext(FStrataPixelHeader StrataPixelHeader, FStrataBSDF BSDF, const FStrataAddressing StrataAddressing, float3 V, float3 L)
 {
 	float3x3 TangentBasis = StrataGetBSDFSharedBasis(StrataPixelHeader, BSDF, StrataAddressing);
-	return StrataCreateBSDFContext(TangentBasis, BSDF, V, L, true);
+	return StrataCreateBSDFContext(TangentBasis, BSDF, V, L, true, StrataAddressing.PixelCoords);
 }
 
 void FStrataBSDFContext::StrataUpdateBSDFContext(float3 NewL)
@@ -1483,10 +1495,21 @@ FStrataEnvLightResult StrataEvaluateForEnvLight(FStrataBSDFContext BSDFContext,
 				BRANCH
 				if (BSDF_GETHASGLINT(BSDFContext.BSDF))
 				{
-					float3 MeanLightLocalDirection = normalize(mul(BSDFContext.TangentBasis, StrataEnvLightResult.SpecularDirection));
-		
 					FBeckmannDesc Beckmann = GGXToBeckmann(SafeRoughness);
-		
+
+					float3 MeanLightLocalDirection = normalize(mul(BSDFContext.TangentBasis, StrataEnvLightResult.SpecularDirection));
+					#if STRATA_GLINT_IS
+					{
+						const float3 Random = float3(BlueNoiseVec2(BSDFContext.PixelCoord, View.StateFrameIndex), BlueNoiseScalar(BSDFContext.PixelCoord, View.StateFrameIndex));
+					
+						// Sample and refect view according the sampled m
+						const float3 m_local = Sample_P(Random, BSDFContext.TangentV, float3(Beckmann.Sigma, Beckmann.Rho), \
+										SLAB_GLINT_VALUE(BSDFContext.BSDF), SLAB_GLINT_UV(BSDFContext.BSDF), SLAB_GLINT_UVDDX(BSDFContext.BSDF), SLAB_GLINT_UVDDY(BSDFContext.BSDF));
+						MeanLightLocalDirection =  normalize(2 * dot(BSDFContext.TangentV, m_local) * m_local - BSDFContext.TangentV);	
+						StrataEnvLightResult.SpecularDirection = normalize(mul(transpose(BSDFContext.TangentBasis), MeanLightLocalDirection));
+					}
+					#endif
+
 					// It is wrong to multiply again the split env lighting (for split-sum approxiation) together with the glint function (representing Beckmann D function).
 					// But visually it works and that is an acceptable solution in the meantime we find something more correct.
 					float GlintD = f_P(BSDFContext.TangentV, MeanLightLocalDirection, float3(Beckmann.Sigma, Beckmann.Rho), 
diff --git a/Engine/Source/Runtime/Renderer/Private/IndirectLightRendering.cpp b/Engine/Source/Runtime/Renderer/Private/IndirectLightRendering.cpp
--- a/Engine/Source/Runtime/Renderer/Private/IndirectLightRendering.cpp
+++ b/Engine/Source/Runtime/Renderer/Private/IndirectLightRendering.cpp
@@ -366,6 +366,7 @@ class FReflectionEnvironmentSkyLightingPS : public FGlobalShader
 		OutEnvironment.SetDefine(TEXT("MAX_CAPTURES"), GMaxNumReflectionCaptures);
 		OutEnvironment.SetDefine(TEXT("SUPPORTS_ANISOTROPIC_MATERIALS"), FDataDrivenShaderPlatformInfo::GetSupportsAnisotropicMaterials(Parameters.Platform));
 		OutEnvironment.SetDefine(TEXT("USE_HAIR_COMPLEX_TRANSMITTANCE"), IsHairStrandsSupported(EHairStrandsShaderType::All, Parameters.Platform) ? 1u : 0u);
+		OutEnvironment.SetDefine(TEXT("STRATA_GLINT_IS"), 0);
 		OutEnvironment.CompilerFlags.Add(CFLAG_StandardOptimization);
 		FForwardLightingParameters::ModifyCompilationEnvironment(Parameters.Platform, OutEnvironment);
 	}
@@ -405,7 +406,7 @@ class FReflectionEnvironmentSkyLightingPS : public FGlobalShader
 
 		SHADER_PARAMETER_RDG_UNIFORM_BUFFER(FStrataGlobalUniformParameters, Strata)
 		SHADER_PARAMETER_STRUCT_INCLUDE(LumenReflections::FCompositeParameters, ReflectionsCompositeParameters)
-
+		SHADER_PARAMETER_STRUCT_REF(FBlueNoise, BlueNoise)
 	END_SHADER_PARAMETER_STRUCT()
 }; // FReflectionEnvironmentSkyLightingPS
 
@@ -1715,6 +1716,11 @@ static void AddSkyReflectionPass(
 		PassParameters->PS.ForwardLightData = View.ForwardLightingResources.ForwardLightUniformBuffer;
 
 		PassParameters->PS.Strata = Strata::BindStrataGlobalUniformParameters(View);
+		if (Strata::IsGlintEnabled())
+		{
+			FBlueNoise BlueNoise = GetBlueNoiseGlobalParameters();
+			PassParameters->PS.BlueNoise = CreateUniformBufferImmediate(BlueNoise, EUniformBufferUsage::UniformBuffer_SingleDraw);
+		}
 
 		LumenReflections::SetupCompositeParameters(PassParameters->PS.ReflectionsCompositeParameters);
 	}
diff --git a/Engine/Source/Runtime/Renderer/Private/SceneRendering.cpp b/Engine/Source/Runtime/Renderer/Private/SceneRendering.cpp
--- a/Engine/Source/Runtime/Renderer/Private/SceneRendering.cpp
+++ b/Engine/Source/Runtime/Renderer/Private/SceneRendering.cpp
@@ -2510,7 +2510,11 @@ FSceneRenderer::FSceneRenderer(const FSceneViewFamily* InViewFamily, FHitProxyCo
 		}
 #endif
 
-		if (ShouldRenderLumenDiffuseGI(Scene, *ViewInfo) || ShouldRenderLumenReflections(*ViewInfo) || ShouldRenderVolumetricCloudWithBlueNoise_GameThread(Scene, *ViewInfo) || UseVirtualShadowMaps(Scene->GetShaderPlatform(), Scene->GetFeatureLevel()))
+		if (ShouldRenderLumenDiffuseGI(Scene, *ViewInfo) || 
+			ShouldRenderLumenReflections(*ViewInfo) || 
+			ShouldRenderVolumetricCloudWithBlueNoise_GameThread(Scene, *ViewInfo) || 
+			UseVirtualShadowMaps(Scene->GetShaderPlatform(), Scene->GetFeatureLevel()) ||
+			Strata::IsGlintEnabled())
 		{
 			GEngine->LoadBlueNoiseTexture();
 		}
