Fix shader warning.

diff --git a/Engine/Shaders/Private/Strata/Glint/GlintThirdParty.ush b/Engine/Shaders/Private/Strata/Glint/GlintThirdParty.ush
--- a/Engine/Shaders/Private/Strata/Glint/GlintThirdParty.ush
+++ b/Engine/Shaders/Private/Strata/Glint/GlintThirdParty.ush
@@ -387,20 +387,20 @@ float P22__glint_discrete_LOD(
 //=============================================================================
 //======== Evaluation of the procedural physically based glinty BRDF ==========
 //=============================================================================
-float3 f_P(float3 wo, float3 wi, float3 sigmas_rho, float LogMicrofacetDensity, float2 TexCoord, float2 TexCoordDDX, float2 TexCoordDDY)
+float f_P(float3 wo, float3 wi, float3 sigmas_rho, float LogMicrofacetDensity, float2 TexCoord, float2 TexCoordDDX, float2 TexCoordDDY)
 {
 	if (wo.z <= 0.)
-		return float3(0., 0., 0.);
+		return 0;
 	if (wi.z <= 0.)
-		return float3(0., 0., 0.);
+		return 0;
 
 	float3 wh = normalize(wo + wi);
 	if (wh.z <= 0.)
-		return float3(0., 0., 0.);
+		return 0;
 
 	// Local masking shadowing
 	if (dot(wo, wh) <= 0. || dot(wi, wh) <= 0.)
-		return float3(0., 0., 0.);
+		return 0;
 
 	FGlintLutParams GlintLutParams;
 	GlintLutParams.Alpha	= View.GlintLUTParameters.x;
@@ -430,7 +430,7 @@ float3 f_P(float3 wo, float3 wi, float3 sigmas_rho, float LogMicrofacetDensity,
 		slope_dy = ddy(slope_h);
 	}
 
-	float3 D_P = float3(0.0f, 0.0f, 0.0f);
+	float D_P = 0;
 	float P22_P = 0.;
 
 	//=========================================================================
@@ -466,7 +466,7 @@ float3 f_P(float3 wo, float3 wi, float3 sigmas_rho, float LogMicrofacetDensity,
 	// Without footprint -> no reflection
 	if (minorLength == 0)
 	{
-		D_P = float3(0.0, 0.0, 0.0);
+		D_P = 0;
 	}
 	else
 	{
@@ -502,7 +502,7 @@ float3 f_P(float3 wo, float3 wi, float3 sigmas_rho, float LogMicrofacetDensity,
 
 		P22_P = lerp(P22_P_il, P22_P_ilp1, w);
 
-		D_P = float3(P22_P.xxx / (wh.z * wh.z * wh.z * wh.z));
+		D_P = P22_P / (wh.z * wh.z * wh.z * wh.z);
 	}
 
 //	// V-cavity masking shadowing
