Local hieght fog volume - preparing for proper integration accounting for transmittance.
+ code clean up.

diff --git a/Engine/Shaders/Private/LocalHeightFog.usf b/Engine/Shaders/Private/LocalHeightFog.usf
--- a/Engine/Shaders/Private/LocalHeightFog.usf
+++ b/Engine/Shaders/Private/LocalHeightFog.usf
@@ -53,38 +53,56 @@ void LocalHeightFogSplatVS(
 	OtherVertexOutput.InstanceId = InstanceIndex;
 } 
 
-// This is an integral of height fog where density is defined as "DensityScale * Exp[-VerticalFallOff * z]"
-float HeightFogRayIntegral(FLocalHeightFogGPUInstanceData FogInstance, float3 CameraPos, float3 RayStart, float3 RayDir, float RayLength)
-{
-	float StartHeight = RayStart.z - FogInstance.HeightOffset;
-#if 0
-	float OpticalDepth = (FogInstance.Density / FogInstance.HeightFalloff) * exp(-StartHeight * FogInstance.HeightFalloff) * (1.0 - exp(-RayDir.z * RayLength * FogInstance.HeightFalloff)) / RayDir.z;
-#else
-	float OpticalDepth = (FogInstance.Density / FogInstance.HeightFalloff) *(exp(-StartHeight * FogInstance.HeightFalloff) - exp(-StartHeight * FogInstance.HeightFalloff -RayDir.z * RayLength * FogInstance.HeightFalloff)) / RayDir.z;
-#endif
-
-	return OpticalDepth;
-}
-
 // CameraPos can be behind of RayStart if the camera is out of the volume
 struct FFogData
 {
+	float IntegratedLuminanceFactor;
 	float Coverage;
 };
-FFogData ComputeFogCoverage(FLocalHeightFogGPUInstanceData FogInstance,  float3 CameraPos, float3 RayStart, float3 RayDir, float RayLength, float RadialAttenuation)
+
+FFogData HeightFogRayIntegral(FLocalHeightFogGPUInstanceData FogInstance, float3 CameraPos, float3 RayStart, float3 RayDir, float RayLength)
+{
+
+	return OpticalDepth;
+}
+
+// Height fog is evaluated according to spatially varying density defined as "DensityScale * Exp[-VerticalFallOff * z]"
+FFogData EvaluateHeightFogIntegral(FLocalHeightFogGPUInstanceData FogInstance,  float3 CameraPos, float3 RayStart, float3 RayDir, float RayLength, float RadialAttenuation)
 {
 	FFogData FogData;
+
+	float StartHeight = RayStart.z - FogInstance.HeightOffset; // accoutn for the fog offset.
+
+	float OpticalDepth = 0.0f;
+	float IntegratedLuminanceFactor = 0.0f;
+#if 0
+	// Original integral, less artefact
+	OpticalDepth = (FogInstance.Density / FogInstance.HeightFalloff) * exp(-StartHeight * FogInstance.HeightFalloff) * (1.0 - exp(-RayDir.z * RayLength * FogInstance.HeightFalloff)) / RayDir.z;
+#elif 1
+	// Reworked to avoid large value as input to exp
+	OpticalDepth = (FogInstance.Density / FogInstance.HeightFalloff) * (exp(-StartHeight * FogInstance.HeightFalloff) - exp(-StartHeight * FogInstance.HeightFalloff - RayDir.z * RayLength * FogInstance.HeightFalloff)) / RayDir.z;
+#else
+	// Accounting for transmittance towards the view
+//	OpticalDepth = (FogInstance.Density * exp(-FogInstance.HeightFalloff * (StartHeight + RayDir.z * RayLength)) * (-1.0 + exp(FogInstance.HeightFalloff * RayDir.z * RayLength))) / (FogInstance.HeightFalloff * RayDir.z);
+	// Now simplified for less artefact:
+	OpticalDepth = (FogInstance.Density/(FogInstance.HeightFalloff * RayDir.z)) 
+		* ( -exp(-FogInstance.HeightFalloff * (StartHeight + RayDir.z * RayLength)) + exp(-FogInstance.HeightFalloff * StartHeight) );
+
+	// TODO evaluate this integral properly, there is a mistake on the transmittance.
+	//IntegratedLuminanceFactor = FogInstance.Density * exp(-OpticalDepth - FogInstance.HeightFalloff * (StartHeight + RayDir.z * RayLength)) * (-1 + exp(FogInstance.HeightFalloff * RayDir.z * RayLength)) / (FogInstance.HeightFalloff * RayDir.z);
+#endif
+	
 	// Integral of density as extinction over a distance is optical depth.
-	float OpticalDepth	= HeightFogRayIntegral(FogInstance, CameraPos, RayStart, RayDir, RayLength);
 	float Transmittance	= exp(-OpticalDepth);
-	FogData.Coverage	= (1.0 - Transmittance) * RadialAttenuation;	// Only coverage is affected 
+	FogData.Coverage = (1.0 - Transmittance) * RadialAttenuation;
+	FogData.IntegratedLuminanceFactor = FogData.Coverage;	// TODO should be IntegratedLuminanceFactor * RadialAttenuation 
 	return FogData;
 }
 
 
 float3 ComputeFogInscattering(in FLocalHeightFogGPUInstanceData FogInstance, in FFogData FogData, in float3 RayDir)
 {
-	half3 InScattering = 0;// float3(0.1, 0.5, 0.5);// FogStruct.ExponentialFogColorParameter.xyz;
+	half3 InScattering = 0;
 
 #if PROJECT_SUPPORT_SKY_ATMOSPHERE_AFFECTS_HEIGHFOG
 	InScattering +=  Texture2DSampleLevel(View.DistantSkyLightLutTexture, View.DistantSkyLightLutTextureSampler, float2(0.5f, 0.5f), 0.0f).rgb;
@@ -103,10 +121,10 @@ float3 ComputeFogInscattering(in FLocalHeightFogGPUInstanceData FogInstance, in
 	// Now account for emissive luminance after albedo and heighfog contribution has been accounted for.
 	InScattering += FogInstance.Emissive;
 
-	return InScattering * FogData.Coverage;
+	return InScattering * FogData.IntegratedLuminanceFactor;
 }
 
-/** . */
+
 void LocalHeightFogSplatPS(
 	in float4 SVPos : SV_POSITION,
 	in FLocalHeightFogVertexOutput OtherVertexOutput,
@@ -133,18 +151,10 @@ void LocalHeightFogSplatPS(
 	float2 Ts = RayIntersectSphere(RayPosU.xyz, RayDirU.xyz, float4(0.0, 0.0, 0.0, 1.0));
 	if (any(Ts > 0.0))
 	{
-#if 1
 		// World space computations
 		float3 P0 = mul(float4(RayPosU.xyz + max(0.0, Ts.x) * RayDirU.xyz, 1.0), Transform).xyz;
 		float3 P1 = mul(float4(RayPosU.xyz + max(0.0, Ts.y) * RayDirU.xyz, 1.0), Transform).xyz;
 
-		float RadialAttenuation = 1.0f;
-#if 0	// attenuation with normalised sphere. But has artefact or vanishes to early
-		float USphereRayThickness = max(0.0, max(Ts.x,Ts.y)) - max(0.0, min(Ts.x, Ts.y));
-		float USphereLinearAttenuation = asinFast(USphereRayThickness * 0.5) * (1.0f / (PI * 0.5));
-		RadialAttenuation = pow(USphereLinearAttenuation, FogInstance.RadialAttenuationPower);
-#endif
-
 		float Length0 = length(P0 - CamRayWorldOrigin);
 		float Length1 = length(P1 - CamRayWorldOrigin);
 
@@ -152,15 +162,10 @@ void LocalHeightFogSplatPS(
 		float3 DepthBufferTranslatedWorldPos = SvPositionToTranslatedWorld(float4(SVPos.xy, DeviceZ, 1.0));
 		float LengthD = length(DepthBufferTranslatedWorldPos - CamRayTranslatedWorldOrigin);
 
-#if 1
 		// Attenuation based on world distance travelled in sphere, ignoring depth intersection
 		float RayLengthInSphere = max(0.0, max(Length0, Length1)) - max(0.0, min(Length0, Length1));
-		RadialAttenuation = saturate(RayLengthInSphere / FogInstance.RadialAttenuationDepth);
+		float RadialAttenuation = saturate(RayLengthInSphere / FogInstance.RadialAttenuationDepth);
 		RadialAttenuation = pow(RadialAttenuation, FogInstance.RadialAttenuationPower);
-		//RadialAttenuation = 1.0 - RadialAttenuation;
-		//RadialAttenuation = 1.0 / (1.0 + exp(-RadialAttenuation * FogInstance.RadialAttenuationPower));	// Sigmoid: do better with range over -x/+x
-		//RadialAttenuation = 1.0 - RadialAttenuation;
-#endif
 
 		Length0 = min(LengthD, Length0);
 		Length1 = min(LengthD, Length1);
@@ -174,35 +179,10 @@ void LocalHeightFogSplatPS(
 
 		if (WorldLength > 0.0)
 		{
-			FFogData FogData = ComputeFogCoverage(FogInstance, CamRayWorldOrigin, CamRayWorldOrigin + CamRayWorldDir * Length0, CamRayWorldDir, WorldLength, RadialAttenuation);
-
-			OutColor.xyz = ComputeFogInscattering(FogInstance, FogData, CamRayWorldDir);	// = Luminance
-			OutColor.a		 = 1.0 - FogData.Coverage;										// = Transmittance
-		}
-		else
-		{
-			clip(-1.0);
-		}
-
-		OutColor.xyz *= PrimaryView.PreExposure;
-#else
-		// Local space compuatations
-		const float DeviceZ = LookupDeviceZ(uint2(SVPos.xy));
-		float3 DepthBufferTranslatedWorldPos = SvPositionToTranslatedWorld(float4(SVPos.xy, DeviceZ, 1.0));
-		float3 DepthBufferWorldPos = DepthBufferTranslatedWorldPos - LWCHackToFloat(PrimaryView.PreViewTranslation).xyz;
-		float3 DepthBufferPosU = mul(float4(DepthBufferWorldPos, 1.0), InvTransform);
-		float3 LengthD = length(DepthBufferPosU - RayPosU);								// TODO account for scale!
-
-		float3 CamPosU = mul(float4(CamRayWorldOrigin, 1.0), InvTransform).xyz;
-
-		float Length0 = length((RayPosU.xyz + max(0.0, Ts.x) * RayDirU.xyz) - CamPosU);
-		float Length1 = length((RayPosU.xyz + max(0.0, Ts.y) * RayDirU.xyz) - CamPosU);
+			FFogData FogData = EvaluateHeightFogIntegral(FogInstance, CamRayWorldOrigin, CamRayWorldOrigin + CamRayWorldDir * Length0, CamRayWorldDir, WorldLength, RadialAttenuation);
 
-		float UnitVolumeLength = max(0.0, abs(Length1 - Length0));
-
-		if (UnitVolumeLength > 0.0)
-		{
-			OutColor = ComputeFog(FogInstance, 10.0*CamPosU, 10.0*(CamPosU + RayDirU * Length0), RayDirU, 10.0*UnitVolumeLength);
+			OutColor.xyz	= ComputeFogInscattering(FogInstance, FogData, CamRayWorldDir);	// = Luminance
+			OutColor.a		= 1.0 - FogData.Coverage;										// = Transmittance
 		}
 		else
 		{
@@ -210,9 +190,6 @@ void LocalHeightFogSplatPS(
 		}
 
 		OutColor.xyz *= PrimaryView.PreExposure;
-
-#endif
-
 	}
 	else
 	{
