Local height fog - removing invalid code left behind.

diff --git a/Engine/Shaders/Private/LocalHeightFog.usf b/Engine/Shaders/Private/LocalHeightFog.usf
--- a/Engine/Shaders/Private/LocalHeightFog.usf
+++ b/Engine/Shaders/Private/LocalHeightFog.usf
@@ -60,12 +60,6 @@ struct FFogData
 	float Coverage;
 };
 
-FFogData HeightFogRayIntegral(FLocalHeightFogGPUInstanceData FogInstance, float3 CameraPos, float3 RayStart, float3 RayDir, float RayLength)
-{
-
-	return OpticalDepth;
-}
-
 // Height fog is evaluated according to spatially varying density defined as "DensityScale * Exp[-VerticalFallOff * z]"
 FFogData EvaluateHeightFogIntegral(FLocalHeightFogGPUInstanceData FogInstance,  float3 CameraPos, float3 RayStart, float3 RayDir, float RayLength, float RadialAttenuation)
 {
