Local height fog new fading

diff --git a/Engine/Shaders/Private/LocalHeightFog.usf b/Engine/Shaders/Private/LocalHeightFog.usf
--- a/Engine/Shaders/Private/LocalHeightFog.usf
+++ b/Engine/Shaders/Private/LocalHeightFog.usf
@@ -14,12 +14,12 @@ struct FLocalHeightFogGPUInstanceData
 	float Density;
 	float HeightFalloff;
 	float HeightOffset;
-	float RadialAttenuationPower;
+	float RadialAttenuation;
 
 	float3 Albedo;
 	float  PhaseG;
 	float3 Emissive;
-	float RadialAttenuationDepth;
+	float RadialAttenuationHermiteTangent;
 };
 
 StructuredBuffer<FLocalHeightFogGPUInstanceData> LocalHeightFogInstances;
@@ -156,10 +156,19 @@ void LocalHeightFogSplatPS(
 		float3 DepthBufferTranslatedWorldPos = SvPositionToTranslatedWorld(float4(SVPos.xy, DeviceZ, 1.0));
 		float LengthD = length(DepthBufferTranslatedWorldPos - CamRayTranslatedWorldOrigin);
 
-		// Attenuation based on world distance travelled in sphere, ignoring depth intersection
-		float RayLengthInSphere = max(0.0, max(Length0, Length1)) - max(0.0, min(Length0, Length1));
-		float RadialAttenuation = saturate(RayLengthInSphere / FogInstance.RadialAttenuationDepth);
-		RadialAttenuation = pow(RadialAttenuation, FogInstance.RadialAttenuationPower);
+		// Radial attenuation based on the distance travelled in the unit sphere.
+		float RadialAttenuation = saturate(abs(max(0.0, Ts.y) - max(0.0, Ts.x)) / FogInstance.RadialAttenuation);
+		// The attenuation is linear but it can be made soft using an Hermit interpolation by controling the tangents.
+		{
+			const float t = RadialAttenuation;
+			const float t2 = t * t;
+			const float t3 = t2 * t;
+			const float HP0 = 0;
+			const float HM0 = saturate(FogInstance.RadialAttenuationHermiteTangent);
+			const float HP1 = 1;
+			const float HM1 = saturate(FogInstance.RadialAttenuationHermiteTangent);
+			RadialAttenuation = saturate((2 * t3 - 3 * t2 + 1) * HP0 + (t3 - 2 * t2 + t) * HM0 + (-2 * t3 + 3 * t2) * HP1 + (t3 - t2) * HM1);
+		}
 
 		Length0 = min(LengthD, Length0);
 		Length1 = min(LengthD, Length1);
diff --git a/Engine/Source/Runtime/Engine/Classes/Components/LocalHeightFogComponent.h b/Engine/Source/Runtime/Engine/Classes/Components/LocalHeightFogComponent.h
--- a/Engine/Source/Runtime/Engine/Classes/Components/LocalHeightFogComponent.h
+++ b/Engine/Source/Runtime/Engine/Classes/Components/LocalHeightFogComponent.h
@@ -28,13 +28,13 @@ class ULocalHeightFogComponent : public USceneComponent
 	UPROPERTY(EditAnywhere, BlueprintReadOnly, interp, Category = "Fog Distribution", meta = (UIMin = "-5000.0", UIMax = "5000.0"))
 	float FogHeightOffset = 0.0f;
 
-	/** Controls the radial attenuation of this fog volume. */
-	UPROPERTY(EditAnywhere, BlueprintReadOnly, interp, Category = "Fog Distribution", meta = (UIMin = "0.01", UIMax = "5.0", ClampMin = 0.01))
-	float FogRadialAttenuationPower = 1.0f;
+	/** Controls the softness of the transition region when the volume is fading out. */
+	UPROPERTY(EditAnywhere, BlueprintReadOnly, interp, Category = "Fog Distribution", meta = (UIMin = "0.0", UIMax = "1.0", ClampMin = 0.0, ClampMax = 1.0))
+	float FogRadialAttenuationSoftness = 1.0f;
 
-	/** Controls the radial attenuation of this fog volume. */
-	UPROPERTY(EditAnywhere, BlueprintReadOnly, interp, Category = "Fog Distribution", meta = (UIMin = "0.01", UIMax = "5000.0", SliderExponent = 4.0, ClampMin = 0.01))
-	float FogRadialAttenuationDepth = 1000.0f;
+	/** Controls how strong the radial attenuation of this fog volume is. */
+	UPROPERTY(EditAnywhere, BlueprintReadOnly, interp, Category = "Fog Distribution", meta = (UIMin = "0.001", UIMax = "4.0", SliderExponent = 3.0, ClampMin = 0.001))
+	float FogRadialAttenuation = 0.0f;
 
 	/** Controls the directionality of the scattering within this fog volume. */
 	UPROPERTY(EditAnywhere, BlueprintReadOnly, interp, Category = "Fog Shading", meta = (UIMin = "0.0", UIMax = "0.999", ClampMin = 0.0, ClampMax = 0.999))
diff --git a/Engine/Source/Runtime/Engine/Private/LocalHeightFogSceneProxy.cpp b/Engine/Source/Runtime/Engine/Private/LocalHeightFogSceneProxy.cpp
--- a/Engine/Source/Runtime/Engine/Private/LocalHeightFogSceneProxy.cpp
+++ b/Engine/Source/Runtime/Engine/Private/LocalHeightFogSceneProxy.cpp
@@ -14,8 +14,8 @@ FLocalHeightFogSceneProxy::FLocalHeightFogSceneProxy(const ULocalHeightFogCompon
 	, FogDensity(InComponent->FogDensity)
 	, FogHeightFalloff(InComponent->FogHeightFalloff)
 	, FogHeightOffset(InComponent->FogHeightOffset + InComponent->GetComponentTransform().GetLocation().Z)
-	, FogRadialAttenuationPower(InComponent->FogRadialAttenuationPower)
-	, FogRadialAttenuationDepth(InComponent->FogRadialAttenuationDepth)
+	, FogRadialAttenuation(InComponent->FogRadialAttenuation)
+	, FogRadialAttenuationSoftness(InComponent->FogRadialAttenuationSoftness)
 	, FogPhaseG(InComponent->FogPhaseG)
 	, FogAlbedo(InComponent->FogAlbedo)
 	, FogEmissive(InComponent->FogEmissive)
diff --git a/Engine/Source/Runtime/Engine/Public/LocalHeightFogSceneProxy.h b/Engine/Source/Runtime/Engine/Public/LocalHeightFogSceneProxy.h
--- a/Engine/Source/Runtime/Engine/Public/LocalHeightFogSceneProxy.h
+++ b/Engine/Source/Runtime/Engine/Public/LocalHeightFogSceneProxy.h
@@ -21,8 +21,8 @@ public:
 	float FogDensity;
 	float FogHeightFalloff;
 	float FogHeightOffset;
-	float FogRadialAttenuationPower;
-	float FogRadialAttenuationDepth;
+	float FogRadialAttenuation;
+	float FogRadialAttenuationSoftness;
 
 	float FogPhaseG;
 	FLinearColor FogAlbedo;
diff --git a/Engine/Source/Runtime/Renderer/Private/LocalHeightFogRendering.cpp b/Engine/Source/Runtime/Renderer/Private/LocalHeightFogRendering.cpp
--- a/Engine/Source/Runtime/Renderer/Private/LocalHeightFogRendering.cpp
+++ b/Engine/Source/Runtime/Renderer/Private/LocalHeightFogRendering.cpp
@@ -75,12 +75,12 @@ public:
 	float Density;
 	float HeightFalloff;
 	float HeightOffset;
-	float RadialAttenuationPower;
+	float RadialAttenuation;
 
 	FVector3f Albedo;
 	float PhaseG;
 	FVector3f Emissive;
-	float RadialAttenuationDepth;
+	float RadialAttenuationHermiteTangent;
 };
 
 class FLocalHeightFogVS : public FGlobalShader
@@ -155,8 +155,8 @@ void RenderLocalHeightFog(
 		LocalHeightFogGPUInstanceDataIt->Density = LHF->FogDensity;
 		LocalHeightFogGPUInstanceDataIt->HeightFalloff = LHF->FogHeightFalloff * 0.01f;	// This scale is used to have artist author reasonable range.
 		LocalHeightFogGPUInstanceDataIt->HeightOffset = LHF->FogHeightOffset;
-		LocalHeightFogGPUInstanceDataIt->RadialAttenuationPower = LHF->FogRadialAttenuationPower;
-		LocalHeightFogGPUInstanceDataIt->RadialAttenuationDepth = LHF->FogRadialAttenuationDepth;
+		LocalHeightFogGPUInstanceDataIt->RadialAttenuation = LHF->FogRadialAttenuation;
+		LocalHeightFogGPUInstanceDataIt->RadialAttenuationHermiteTangent = 1.0f - LHF->FogRadialAttenuationSoftness;
 
 		LocalHeightFogGPUInstanceDataIt->Albedo = FVector3f(LHF->FogAlbedo);
 		LocalHeightFogGPUInstanceDataIt->PhaseG = LHF->FogPhaseG;
