Path Tracing: Adding UseAnalyticTransmittance CVar to determine method for transmittance estimation.

r.PathTracing.UseAnalyticTransmittance (default = -1)
-1: uses null-tracking estimation if heterogeneous volumes are present, or analytical estimation otherwise
0: off, uses null-tracking estimation
1: on, uses analytical estimation

diff --git a/Engine/Shaders/Private/PathTracing/PathTracingCore.ush b/Engine/Shaders/Private/PathTracing/PathTracingCore.ush
--- a/Engine/Shaders/Private/PathTracing/PathTracingCore.ush
+++ b/Engine/Shaders/Private/PathTracing/PathTracingCore.ush
@@ -542,10 +542,9 @@ FPathTracingPayload TraceTransparentRay(inout FPathState PathState, int Bounce,
 				}
 			}
 
-#if PATH_TRACING_USE_TRANSMITTANCE_FUNC
 			// NOTE: when using analytic raymarching, we can skip all raymarching on the last bounce because we know we don't need a scattering point
-			if (!bLastBounce)
-#endif
+			const uint UseNullTrackingTransmittance = !UseAnalyticTransmittance;
+			if (UseNullTrackingTransmittance || (UseAnalyticTransmittance && !bLastBounce))
 			{
 				// remember the start of our interval before ray marching messes with it
 				float StartVolumeTMin = Interval.VolumeTMin;
@@ -556,12 +555,13 @@ FPathTracingPayload TraceTransparentRay(inout FPathState PathState, int Bounce,
 				{
 					// figure out optimal bounds for the current ray segment
 					float3 SigmaBar = VolumeGetDensityBounds(PathState.Ray.Origin, PathState.Ray.Direction, Interval).SigmaMax;
-#if PATH_TRACING_USE_TRANSMITTANCE_FUNC == 0
-					// if we are not using analytical transmittance, a tight majorant could prevent us from "seeing" hits that match the density exactly
-					// leading to heavy noise on bright objects embedded in the volume
-					// however if we can track transmittance analytically, this workaround is not needed as we will re-compute a precise answer after ray marching
-					SigmaBar *= bIsCameraRay ? 2 : 1;
-#endif
+					if (UseNullTrackingTransmittance)
+					{
+						// if we are not using analytical transmittance, a tight majorant could prevent us from "seeing" hits that match the density exactly
+						// leading to heavy noise on bright objects embedded in the volume
+						// however if we can track transmittance analytically, this workaround is not needed as we will re-compute a precise answer after ray marching
+						SigmaBar *= bIsCameraRay ? 2 : 1;
+					}
 
 					// there is some atmosphere present along the current ray, ray march until we get to the current hit (or exit the atmosphere)
 					float3 ColorChannelPdf = PathThroughput;
@@ -591,9 +591,7 @@ FPathTracingPayload TraceTransparentRay(inout FPathState PathState, int Bounce,
 						float3 SigmaT = min(Result.SigmaT, SigmaBar);
 						float3 SigmaN = SigmaBar - SigmaT;
 
-#if PATH_TRACING_USE_TRANSMITTANCE_FUNC == 0
-						if (!bLastBounce) // check now if this loop is also computing transmittance
-#endif
+						if (UseAnalyticTransmittance || (UseNullTrackingTransmittance && !bLastBounce)) // check now if this loop is also computing transmittance
 						{
 							float3 SigmaS = min(Result.SigmaSRayleigh + Result.SigmaSHG, SigmaT);
 							// accumulate a signal for the denoiser
@@ -621,15 +619,19 @@ FPathTracingPayload TraceTransparentRay(inout FPathState PathState, int Bounce,
 				}
 				// restore interval start now that we are finished with ray marching
 				Interval.VolumeTMin = StartVolumeTMin;
-#if PATH_TRACING_USE_TRANSMITTANCE_FUNC == 0
-				// record the change in transmittance since we already computed it
-				PathState.PathThroughput = PathThroughput;
-#endif
+				if (UseNullTrackingTransmittance)
+				{
+					// record the change in transmittance since we already computed it
+					PathState.PathThroughput = PathThroughput;
+				}
 			}
-#if PATH_TRACING_USE_TRANSMITTANCE_FUNC
-			// analytically handle the transmittance to the next surface (which will be higher quality than the transmittance implicitly computed above and prevent noise with emissive surfaces in volumes)
-			PathState.PathThroughput = VolumeGetTransmittance(PathState.PathThroughput, PathState.Ray.Origin, PathState.Ray.Direction, Interval, PathState.PathRoughness, PathState.RandSequence);
-#endif
+			
+			if (UseAnalyticTransmittance)
+			{
+				// analytically handle the transmittance to the next surface (which will be higher quality than the transmittance implicitly computed above and prevent noise with emissive surfaces in volumes)
+				PathState.PathThroughput = VolumeGetTransmittance(PathState.PathThroughput, PathState.Ray.Origin, PathState.Ray.Direction, Interval, PathState.PathRoughness, PathState.RandSequence);
+			}
+
 			// Update our current volume intersection list to reflect the fact that we have made it to the end of this interval
 			// This will either clip the finished portions of active segments, or remove the segments we have fully marched through
 			VolumeIsectList = VolumeIsectList.Update(Interval.VolumeTMax);
diff --git a/Engine/Shaders/Private/PathTracing/Volume/PathTracingVolume.ush b/Engine/Shaders/Private/PathTracing/Volume/PathTracingVolume.ush
--- a/Engine/Shaders/Private/PathTracing/Volume/PathTracingVolume.ush
+++ b/Engine/Shaders/Private/PathTracing/Volume/PathTracingVolume.ush
@@ -8,7 +8,7 @@
 #include "PathTracingHeterogeneousVolumes.ush"
 #include "PathTracingVolumeSampling.ush"
 
-#define PATH_TRACING_USE_TRANSMITTANCE_FUNC 1 // 0: null tracking transmittance (reference) 1: analytic transmittance
+uint UseAnalyticTransmittance;
 
 uint EnableAtmosphere;
 uint EnableFog;
@@ -73,10 +73,9 @@ FVolumeShadedResult VolumeGetDensity(float3 TranslatedWorldPos, FVolumeIntersect
 	return Result;
 }
 
-#if PATH_TRACING_USE_TRANSMITTANCE_FUNC
 
 // Return the transmittance along a ray segment
-float3 VolumeGetTransmittance(float3 StartThroughput, float3 Origin, float3 Direction, FVolumeIntersectionInterval Interval, float PathRoughness, inout RandomSequence RandSequence)
+float3 VolumeGetTransmittanceAnalytical(float3 StartThroughput, float3 Origin, float3 Direction, FVolumeIntersectionInterval Interval, float PathRoughness, inout RandomSequence RandSequence)
 {
 	float3 Throughput = StartThroughput;
 	for (int Index = 0; Index < Interval.Num; Index++)
@@ -92,7 +91,7 @@ float3 VolumeGetTransmittance(float3 StartThroughput, float3 Origin, float3 Dire
 }
 
 // Return the transmittance along a whole intersection list
-float3 VolumeGetTransmittance(float3 StartThroughput, float3 Origin, float3 Direction, FVolumeIntersectionList IntersectionList, float PathRoughness, inout RandomSequence RandSequence)
+float3 VolumeGetTransmittanceAnalytical(float3 StartThroughput, float3 Origin, float3 Direction, FVolumeIntersectionList IntersectionList, float PathRoughness, inout RandomSequence RandSequence)
 {
 	float3 Throughput = StartThroughput;
 	for (int Index = 0; Index < IntersectionList.Num; Index++)
@@ -109,7 +108,6 @@ float3 VolumeGetTransmittance(float3 StartThroughput, float3 Origin, float3 Dire
 	return Throughput;
 }
 
-#else
 
 // Reference implementation of transmittance that uses only GetDensity and GetDensityBounds
 // express the inner loop here as a macro so we can stamp down a separate copy per volume ID and avoid a nested loop during ray marching
@@ -152,7 +150,7 @@ float3 VolumeGetTransmittance(float3 StartThroughput, float3 Origin, float3 Dire
 
 
 // The returned value factors in the initial throughput. This is used to avoid sampling too much if throughput reaches 0.
-float3 VolumeGetTransmittance(float3 StartThroughput, float3 Origin, float3 Direction, FVolumeIntersectionInterval Interval, int Bounce, inout RandomSequence RandSequence)
+float3 VolumeGetTransmittanceNullTracking(float3 StartThroughput, float3 Origin, float3 Direction, FVolumeIntersectionInterval Interval, int Bounce, inout RandomSequence RandSequence)
 {
 	float3 Throughput = StartThroughput;
 	for (int Index = 0; Index < Interval.Num; Index++)
@@ -181,7 +179,7 @@ float3 VolumeGetTransmittance(float3 StartThroughput, float3 Origin, float3 Dire
 	return Throughput;
 
 }
-float3 VolumeGetTransmittance(float3 StartThroughput, float3 Origin, float3 Direction, FVolumeIntersectionList IntersectionList, int Bounce, inout RandomSequence RandSequence)
+float3 VolumeGetTransmittanceNullTracking(float3 StartThroughput, float3 Origin, float3 Direction, FVolumeIntersectionList IntersectionList, int Bounce, inout RandomSequence RandSequence)
 {
 	float3 Throughput = StartThroughput;
 	for (int Index = 0; Index < IntersectionList.Num; Index++)
@@ -212,4 +210,22 @@ float3 VolumeGetTransmittance(float3 StartThroughput, float3 Origin, float3 Dire
 #undef PATH_TRACER_REFERENCE_TRANSMITTANCE_LOOP
 }
 
-#endif
+float3 VolumeGetTransmittance(float3 StartThroughput, float3 Origin, float3 Direction, FVolumeIntersectionInterval Interval, int Bounce, inout RandomSequence RandSequence)
+{
+	if (UseAnalyticTransmittance)
+	{
+		return VolumeGetTransmittanceAnalytical(StartThroughput, Origin, Direction, Interval, Bounce, RandSequence);
+	}
+
+	return VolumeGetTransmittanceNullTracking(StartThroughput, Origin, Direction, Interval, Bounce, RandSequence);
+}
+
+float3 VolumeGetTransmittance(float3 StartThroughput, float3 Origin, float3 Direction, FVolumeIntersectionList IntersectionList, int Bounce, inout RandomSequence RandSequence)
+{
+	if (UseAnalyticTransmittance)
+	{
+		return VolumeGetTransmittanceAnalytical(StartThroughput, Origin, Direction, IntersectionList, Bounce, RandSequence);
+	}
+
+	return VolumeGetTransmittanceNullTracking(StartThroughput, Origin, Direction, IntersectionList, Bounce, RandSequence);
+}
diff --git a/Engine/Source/Runtime/Renderer/Private/PathTracing.cpp b/Engine/Source/Runtime/Renderer/Private/PathTracing.cpp
--- a/Engine/Source/Runtime/Renderer/Private/PathTracing.cpp
+++ b/Engine/Source/Runtime/Renderer/Private/PathTracing.cpp
@@ -383,6 +383,15 @@ TAutoConsoleVariable<int32> CVarpathTracingOverrideDepth(
 	ECVF_RenderThreadSafe
 );
 
+TAutoConsoleVariable<int32> CVarPathTracingUseAnalyticTransmittance(
+	TEXT("r.PathTracing.UseAnalyticTransmittance"),
+	-1,
+	TEXT("Determines use of analytical or null-tracking estimation when evaluating transmittance\n")
+	TEXT("-1: uses null-tracking estimation if heterogeneous volumes are present, or analytical estimation otherwise (default)\n")
+	TEXT("0: off (uses null-tracking estimation, instead)\n")
+	TEXT("1: on (uses analytical estimation)\n"),
+	ECVF_RenderThreadSafe
+);
 
 BEGIN_SHADER_PARAMETER_STRUCT(FPathTracingData, )
 	SHADER_PARAMETER(float, BlendFactor)
@@ -406,6 +415,7 @@ BEGIN_SHADER_PARAMETER_STRUCT(FPathTracingData, )
 	SHADER_PARAMETER(uint32, EnabledDirectLightingContributions)   // PATHTRACER_CONTRIBUTION_*
 	SHADER_PARAMETER(uint32, EnabledIndirectLightingContributions) // PATHTRACER_CONTRIBUTION_*
 	SHADER_PARAMETER(uint32, ApplyDiffuseSpecularOverrides)
+	SHADER_PARAMETER(uint32, UseAnalyticTransmittance)
 	SHADER_PARAMETER(int32, MaxRaymarchSteps)
 	SHADER_PARAMETER(float, MaxPathIntensity)
 	SHADER_PARAMETER(float, MaxNormalBias)
@@ -454,6 +464,7 @@ struct FPathTracingConfig
 			PathTracingData.EnableAtmosphere != Other.PathTracingData.EnableAtmosphere ||
 			PathTracingData.EnableFog != Other.PathTracingData.EnableFog ||
 			PathTracingData.EnableHeterogeneousVolumes != Other.PathTracingData.EnableHeterogeneousVolumes ||
+			PathTracingData.UseAnalyticTransmittance != Other.PathTracingData.UseAnalyticTransmittance ||
 			PathTracingData.ApplyDiffuseSpecularOverrides != Other.PathTracingData.ApplyDiffuseSpecularOverrides ||
 			PathTracingData.EnabledDirectLightingContributions != Other.PathTracingData.EnabledDirectLightingContributions ||
 			PathTracingData.EnabledIndirectLightingContributions != Other.PathTracingData.EnabledIndirectLightingContributions ||
@@ -560,6 +571,17 @@ namespace PathTracing
 	}
 }
 
+static uint32 EvalUseAnalyticTransmittance(const FViewInfo& View)
+{
+	int32 UseAnalyticTransmittance = CVarPathTracingUseAnalyticTransmittance.GetValueOnRenderThread();
+	if (UseAnalyticTransmittance < 0)
+	{
+		UseAnalyticTransmittance = !ShouldRenderHeterogeneousVolumesForView(View);
+	}
+
+	return uint32(UseAnalyticTransmittance);
+}
+
 // This function prepares the portion of shader arguments that may involve invalidating the path traced state
 static void PreparePathTracingData(const FScene* Scene, const FViewInfo& View, FPathTracingData& PathTracingData)
 {
@@ -621,6 +643,7 @@ static void PreparePathTracingData(const FScene* Scene, const FViewInfo& View, F
 			Scene->ExponentialFogs[0].FogData[1].Density > 0);
 
 	PathTracingData.EnableHeterogeneousVolumes = CVarPathTracingHeterogeneousVolumes.GetValueOnRenderThread();
+	PathTracingData.UseAnalyticTransmittance = EvalUseAnalyticTransmittance(View);
 	PathTracingData.EnableDBuffer = CVarPathTracingUseDBuffer.GetValueOnRenderThread();
 
 	PathTracingData.DecalRoughnessCutoff = PathTracing::UsesDecals(*View.Family) && View.bHasRayTracingDecals ? CVarPathTracingDecalRoughnessCutoff.GetValueOnRenderThread() : -1.0f;
