LocalHeightFogVolume - justified proper integration of luminance and emissive from the volume.

diff --git a/Engine/Shaders/Private/LocalHeightFog.usf b/Engine/Shaders/Private/LocalHeightFog.usf
--- a/Engine/Shaders/Private/LocalHeightFog.usf
+++ b/Engine/Shaders/Private/LocalHeightFog.usf
@@ -60,7 +60,7 @@ struct FFogData
 	float Coverage;
 };
 
-// Height fog is evaluated according to spatially varying density defined as "DensityScale * Exp[-VerticalFallOff * z]"
+// Height fog is evaluated according to spatially varying density defined as "Density * Exp[-HeightFalloff * z]"
 FFogData EvaluateHeightFogIntegral(FLocalHeightFogGPUInstanceData FogInstance,  float3 CameraPos, float3 RayStart, float3 RayDir, float RayLength, float RadialAttenuation)
 {
 	FFogData FogData;
@@ -68,28 +68,22 @@ FFogData EvaluateHeightFogIntegral(FLocalHeightFogGPUInstanceData FogInstance,
 	float StartHeight = RayStart.z - FogInstance.HeightOffset; // accoutn for the fog offset.
 
 	float OpticalDepth = 0.0f;
-	float IntegratedLuminanceFactor = 0.0f;
-#if 0
+#if 1
 	// Original integral, less artefact
 	OpticalDepth = (FogInstance.Density / FogInstance.HeightFalloff) * exp(-StartHeight * FogInstance.HeightFalloff) * (1.0 - exp(-RayDir.z * RayLength * FogInstance.HeightFalloff)) / RayDir.z;
-#elif 1
+#elif 0
 	// Reworked to avoid large value as input to exp
 	OpticalDepth = (FogInstance.Density / FogInstance.HeightFalloff) * (exp(-StartHeight * FogInstance.HeightFalloff) - exp(-StartHeight * FogInstance.HeightFalloff - RayDir.z * RayLength * FogInstance.HeightFalloff)) / RayDir.z;
-#else
-	// Accounting for transmittance towards the view
-//	OpticalDepth = (FogInstance.Density * exp(-FogInstance.HeightFalloff * (StartHeight + RayDir.z * RayLength)) * (-1.0 + exp(FogInstance.HeightFalloff * RayDir.z * RayLength))) / (FogInstance.HeightFalloff * RayDir.z);
-	// Now simplified for less artefact:
-	OpticalDepth = (FogInstance.Density/(FogInstance.HeightFalloff * RayDir.z)) 
-		* ( -exp(-FogInstance.HeightFalloff * (StartHeight + RayDir.z * RayLength)) + exp(-FogInstance.HeightFalloff * StartHeight) );
-
-	// TODO evaluate this integral properly, there is a mistake on the transmittance.
-	//IntegratedLuminanceFactor = FogInstance.Density * exp(-OpticalDepth - FogInstance.HeightFalloff * (StartHeight + RayDir.z * RayLength)) * (-1 + exp(FogInstance.HeightFalloff * RayDir.z * RayLength)) / (FogInstance.HeightFalloff * RayDir.z);
 #endif
-	
-	// Integral of density as extinction over a distance is optical depth.
-	float Transmittance	= exp(-OpticalDepth);
+
+	// Transmittance is the 
+	float Transmittance = exp(-OpticalDepth);
+	// Assuming extinction is grey scale, we can compute a single coverage value from transmittance
 	FogData.Coverage = (1.0 - Transmittance) * RadialAttenuation;
-	FogData.IntegratedLuminanceFactor = FogData.Coverage;	// TODO should be IntegratedLuminanceFactor * RadialAttenuation 
+	// Assuming extinction==scattering, i.e. albedo=1, we can compute the integrale for each point on a ray while accouting for transmittance back to the view.
+	// It turns out that, under these conditions, the luminance and emissive factor is simply coverage. => see Mathematica details at the bottom of this page.
+	FogData.IntegratedLuminanceFactor = FogData.Coverage;
+
 	return FogData;
 }
 
@@ -201,4 +195,34 @@ void LocalHeightFogSplatPS(
 }
 
 
- 
\ No newline at end of file
+/*
+
+Proof of the height fog integral equations using Mathematica
+
+(*Density of matter according to height Y*)
+Density [y_] := A * Exp[-B * y]
+ 
+(*Ray position*)
+Ray[t_] := Oz + t * Rz
+
+(*Integrate Optical Depth for a distance T*)
+Fog[T_] := Integrate[Density[Ray[t]], {t, 0, T}]
+
+(*Optical Depth integration equation*)
+Fog[T]
+
+(*Transmittance from Optical Depth*)
+Exp[-Integrate[Density[Ray[t]], {t, 0, T}]]
+
+(*Transmittance from a position to the origin of tracing*)
+Trans[T_] := Exp[-Integrate[Density[Ray[t]], {t, 0, T}]]
+
+(*Integrate emissive and scattering from height fog, assuming density \
+is extinction and extinction=scattering, thus albedo=1*)
+Integrate[Density[Ray[t]]*EmSc *Trans[t], {t, 0, T}]
+
+=> this leads to EmSc - Transmittance * EmSc
+				 EmSc * (1 - Transmittance)
+				 EmSc * Coverage
+
+*/
\ No newline at end of file
