Local Height Fog Volume now rotates fog direction.

diff --git a/Engine/Shaders/Private/LocalHeightFog.usf b/Engine/Shaders/Private/LocalHeightFog.usf
--- a/Engine/Shaders/Private/LocalHeightFog.usf
+++ b/Engine/Shaders/Private/LocalHeightFog.usf
@@ -10,6 +10,10 @@ struct FLocalHeightFogGPUInstanceData
 {
 	float4x4 Transform;
 	float4x4 InvTransform;
+	float4x4 FogTranformNoScale;
+	float4x4 FogTranformNoScaleInv;
+	float4x4 FogTransformScaleOnly;
+	float4x4 FogTransformScaleOnlyInv;
 
 	float Density;
 	float HeightFalloff;
@@ -88,7 +92,7 @@ FFogData EvaluateHeightFogIntegral(FLocalHeightFogGPUInstanceData FogInstance,
 }
 
 
-float3 ComputeFogInscattering(in FLocalHeightFogGPUInstanceData FogInstance, in FFogData FogData, in float3 RayDir)
+float3 ComputeFogInscattering(in FLocalHeightFogGPUInstanceData FogInstance, in FFogData FogData, in float3 RayDirWorld)
 {
 	half3 InScattering = 0;
 
@@ -96,11 +100,11 @@ float3 ComputeFogInscattering(in FLocalHeightFogGPUInstanceData FogInstance, in
 	InScattering +=  Texture2DSampleLevel(View.DistantSkyLightLutTexture, View.DistantSkyLightLutTextureSampler, float2(0.5f, 0.5f), 0.0f).rgb;
 
 	// No need to test View.AtmosphereLightIlluminanceOnGroundPostTransmittance[0].a because InscatteringLightDirection.w above is doing the same test already.
-	InScattering += View.AtmosphereLightIlluminanceOnGroundPostTransmittance[0].rgb * SchlickPhase(-FogInstance.PhaseG, dot(RayDir, View.AtmosphereLightDirection[0].xyz));
+	InScattering += View.AtmosphereLightIlluminanceOnGroundPostTransmittance[0].rgb * SchlickPhase(-FogInstance.PhaseG, dot(RayDirWorld, View.AtmosphereLightDirection[0].xyz));
 
 	if (View.AtmosphereLightIlluminanceOnGroundPostTransmittance[1].a > 0.0f) // Skip the second light when disabled.
 	{
-		InScattering += View.AtmosphereLightIlluminanceOnGroundPostTransmittance[0].rgb * SchlickPhase(-FogInstance.PhaseG, dot(RayDir, View.AtmosphereLightDirection[1].xyz));
+		InScattering += View.AtmosphereLightIlluminanceOnGroundPostTransmittance[0].rgb * SchlickPhase(-FogInstance.PhaseG, dot(RayDirWorld, View.AtmosphereLightDirection[1].xyz));
 	}
 #endif
 
@@ -112,6 +116,7 @@ float3 ComputeFogInscattering(in FLocalHeightFogGPUInstanceData FogInstance, in
 	return InScattering * FogData.IntegratedLuminanceFactor;
 }
 
+#define LHFV_FOLLOWS_ROTATIONS 1
 
 void LocalHeightFogSplatPS(
 	in float4 SVPos : SV_POSITION,
@@ -132,26 +137,42 @@ void LocalHeightFogSplatPS(
 	OutColor = float4(0, 0, 0, 1);
 
 	float3 CamRayWorldOrigin = CamRayTranslatedWorldOrigin - LWCHackToFloat(PrimaryView.PreViewTranslation);
-	float4 RayPosU = mul(float4(CamRayWorldOrigin, 1.0), InvTransform);
-	float4 RayDirU = mul(float4(CamRayTranslatedWorldDir, 0.0), InvTransform);
+
+	// The "U" prefix is for all compuation done in the Unit Sphere space.
+	float4 RayPosU = mul(float4(CamRayWorldOrigin, 1.0), FogInstance.InvTransform);
+	float4 RayDirU = mul(float4(CamRayTranslatedWorldDir, 0.0), FogInstance.InvTransform);	// TODO This potentially contains a non uniform scale so check corectness
 	RayDirU.xyz = normalize(RayDirU.xyz);
 
-	float2 Ts = RayIntersectSphere(RayPosU.xyz, RayDirU.xyz, float4(0.0, 0.0, 0.0, 1.0));
-	if (any(Ts > 0.0))
+	float2 TsU = RayIntersectSphere(RayPosU.xyz, RayDirU.xyz, float4(0.0, 0.0, 0.0, 1.0));
+	if (any(TsU > 0.0))
 	{
-		// World space computations
-		float3 P0 = mul(float4(RayPosU.xyz + max(0.0, Ts.x) * RayDirU.xyz, 1.0), Transform).xyz;
-		float3 P1 = mul(float4(RayPosU.xyz + max(0.0, Ts.y) * RayDirU.xyz, 1.0), Transform).xyz;
-
-		float Length0 = length(P0 - CamRayWorldOrigin);
-		float Length1 = length(P1 - CamRayWorldOrigin);
-
 		const float DeviceZ = LookupDeviceZ(uint2(SVPos.xy));
 		float3 DepthBufferTranslatedWorldPos = SvPositionToTranslatedWorld(float4(SVPos.xy, DeviceZ, 1.0));
 		float LengthD = length(DepthBufferTranslatedWorldPos - CamRayTranslatedWorldOrigin);
 
-		// Radial attenuation based on the distance travelled in the unit sphere.
-		float RadialAttenuation = saturate(abs(max(0.0, Ts.y) - max(0.0, Ts.x)) / FogInstance.RadialAttenuation);
+		float3 P0U = RayPosU.xyz + max(0.0, TsU.x) * RayDirU.xyz;
+		float3 P1U = RayPosU.xyz + max(0.0, TsU.y) * RayDirU.xyz;
+#if LHFV_FOLLOWS_ROTATIONS
+		// Computations are all done transformed with Space = local fog volume but ignoring scaling.
+		float3 P0Space = mul(float4(P0U, 1.0), FogInstance.FogTransformScaleOnly).xyz;
+		float3 P1Space = mul(float4(P1U, 1.0), FogInstance.FogTransformScaleOnly).xyz;
+		float3 CamRayOriginSpace = mul(float4(CamRayWorldOrigin, 1.0), FogInstance.FogTranformNoScaleInv);
+		float3 CamRayDirSpace    = mul(float4(CamRayWorldDir,    0.0), FogInstance.FogTranformNoScaleInv);
+		CamRayDirSpace = normalize(CamRayDirSpace);
+#else
+		// Computations are all done with Space = world space
+		float3 P0Space = mul(float4(P0U, 1.0), FogInstance.Transform).xyz;
+		float3 P1Space = mul(float4(P1U, 1.0), FogInstance.Transform).xyz;
+		float3 CamRayOriginSpace = CamRayWorldOrigin;
+		float3 CamRayDirSpace    = CamRayWorldDir;
+#endif
+
+		float Length0 = length(P0Space - CamRayOriginSpace);
+		float Length1 = length(P1Space - CamRayOriginSpace);
+
+		// Radial attenuation based on the distance travelled in the unit sphere and ignoring depth buffer intersection. 
+		// This is only to soften sphere edge.
+		float RadialAttenuation = saturate(abs(max(0.0, TsU.y) - max(0.0, TsU.x)) / FogInstance.RadialAttenuation);
 		// The attenuation is linear but it can be made soft using an Hermit interpolation by controling the tangents.
 		{
 			const float t = RadialAttenuation;
@@ -164,19 +185,14 @@ void LocalHeightFogSplatPS(
 			RadialAttenuation = saturate((2 * t3 - 3 * t2 + 1) * HP0 + (t3 - 2 * t2 + t) * HM0 + (-2 * t3 + 3 * t2) * HP1 + (t3 - t2) * HM1);
 		}
 
+		// Clamp the traced distance to depth and compute world length.
 		Length0 = min(LengthD, Length0);
 		Length1 = min(LengthD, Length1);
-		if (Length1 < Length0)
-		{
-			float LengthX = Length1;
-			Length1 = Length0;
-			Length0 = LengthX;
-		}
-		float WorldLength = max(0.0, Length1 - Length0);
+		float RayTracedLength = max(0.0, abs(Length1 - Length0));
 
-		if (WorldLength > 0.0)
+		if (RayTracedLength > 0.0)
 		{
-			FFogData FogData = EvaluateHeightFogIntegral(FogInstance, CamRayWorldOrigin, CamRayWorldOrigin + CamRayWorldDir * Length0, CamRayWorldDir, WorldLength, RadialAttenuation);
+			FFogData FogData = EvaluateHeightFogIntegral(FogInstance, CamRayOriginSpace, CamRayOriginSpace + CamRayDirSpace * Length0, CamRayDirSpace, RayTracedLength, RadialAttenuation);
 
 			OutColor.xyz	= ComputeFogInscattering(FogInstance, FogData, CamRayWorldDir);	// = Luminance
 			OutColor.a		= 1.0 - FogData.Coverage;										// = Transmittance
diff --git a/Engine/Source/Runtime/Renderer/Private/LocalHeightFogRendering.cpp b/Engine/Source/Runtime/Renderer/Private/LocalHeightFogRendering.cpp
--- a/Engine/Source/Runtime/Renderer/Private/LocalHeightFogRendering.cpp
+++ b/Engine/Source/Runtime/Renderer/Private/LocalHeightFogRendering.cpp
@@ -72,6 +72,10 @@ class FLocalHeightFogGPUInstanceData
 public:
 	FMatrix44f Transform;
 	FMatrix44f InvTransform;
+	FMatrix44f FogTranformNoScale;
+	FMatrix44f FogTranformNoScaleInv;
+	FMatrix44f FogTransformScaleOnly;
+	FMatrix44f FogTransformScaleOnlyInv;
 
 	float Density;
 	float HeightFalloff;
@@ -149,8 +153,15 @@ void RenderLocalHeightFog(
 		FLocalHeightFogGPUInstanceData* LocalHeightFogGPUInstanceDataIt = LocalHeightFogGPUInstanceData;
 		for (FLocalHeightFogSceneProxy* LHF : Scene->LocalHeightFogs)
 		{
-			LocalHeightFogGPUInstanceDataIt->Transform = FMatrix44f(LHF->FogTransform.ToMatrixWithScale());
-			LocalHeightFogGPUInstanceDataIt->InvTransform = LocalHeightFogGPUInstanceDataIt->Transform.Inverse();
+			FTransform FogTransformScaleOnly;
+			FogTransformScaleOnly.SetScale3D(LHF->FogTransform.GetScale3D());
+
+			LocalHeightFogGPUInstanceDataIt->Transform					= FMatrix44f(LHF->FogTransform.ToMatrixWithScale());
+			LocalHeightFogGPUInstanceDataIt->InvTransform				= LocalHeightFogGPUInstanceDataIt->Transform.Inverse();
+			LocalHeightFogGPUInstanceDataIt->FogTranformNoScale			= FMatrix44f(LHF->FogTransform.ToMatrixNoScale());
+			LocalHeightFogGPUInstanceDataIt->FogTranformNoScaleInv		= LocalHeightFogGPUInstanceDataIt->FogTranformNoScale.Inverse();
+			LocalHeightFogGPUInstanceDataIt->FogTransformScaleOnly		= FMatrix44f(FogTransformScaleOnly.ToMatrixWithScale());
+			LocalHeightFogGPUInstanceDataIt->FogTransformScaleOnlyInv	= LocalHeightFogGPUInstanceDataIt->FogTransformScaleOnly.Inverse();
 
 			LocalHeightFogGPUInstanceDataIt->Density = LHF->FogDensity;
 			LocalHeightFogGPUInstanceDataIt->HeightFalloff = LHF->FogHeightFalloff * 0.01f;	// This scale is used to have artist author reasonable range.
