Add 3D circles and spheres to ShaderPrint

diff --git a/Engine/Shaders/Private/ShaderPrintCommon.ush b/Engine/Shaders/Private/ShaderPrintCommon.ush
--- a/Engine/Shaders/Private/ShaderPrintCommon.ush
+++ b/Engine/Shaders/Private/ShaderPrintCommon.ush
@@ -1175,16 +1175,15 @@ void AddCircleSS(FShaderPrintContext Ctx, float2 Center, float Radius, float4 Co
 {
 	if (Ctx.bIsActive)
 	{
-		const float TStep = 0.05f;
-
-		float S, C;
-		sincos(0.0, S, C);
-		float2 PrevP = Center + float2(C, S) * Radius;
-
-		for (uint SegIt = 1; SegIt <= SegmentCount; ++SegIt)
+		const float AngleStep = 2.0f * PI * rcp(float(SegmentCount-1));
+		float Angle = 0.0f;
+		float2 PrevP = Center + float2(Radius, 0.0f);
+		for (uint SegIt = 0; SegIt < SegmentCount; ++SegIt)
 		{
-			const float t = float(SegIt) / float(SegmentCount-1);
-			sincos(t * 2 * PI, S, C);
+			Angle += AngleStep;
+
+			float S, C;
+			sincos(Angle, S, C);
 			float2 P = Center + float2(C, S) * Radius;
 
 			AddLineSS(Ctx, PrevP, P, Color, Color);
@@ -1193,6 +1192,50 @@ void AddCircleSS(FShaderPrintContext Ctx, float2 Center, float Radius, float4 Co
 	}
 }
 
+void AddCircleTWS(FShaderPrintContext Ctx, float3 Center, float3 Normal, float Radius, float4 Color, uint SegmentCount = 16)
+{
+	if (Ctx.bIsActive)
+	{
+		// Calculate two perpendicular unit vectors to the normal
+		const float3 RefUp = abs(Normal.z) < 0.999f ? float3(0, 0, 1) : float3(1, 0, 0);
+		float3 X = normalize(cross(RefUp, Normal));
+		float3 Y = cross(Normal, X);
+		X *= Radius;
+		Y *= Radius;
+
+		const float AngleStep = 2.0f * PI * rcp(float(SegmentCount-1));
+		float Angle = 0.0f;
+		float3 PrevP = Center + X;
+		for (uint SegIt = 0; SegIt < SegmentCount; ++SegIt)
+		{
+			Angle += AngleStep;
+
+			float S, C;
+			sincos(Angle, S, C);
+			float3 P = Center + X * C + Y * S;
+
+			AddLineTWS(Ctx, PrevP, P, Color, Color);
+			PrevP = P;
+		}
+	}
+}
+
+void AddCircleWS (FShaderPrintContext Ctx, float3 Center, float3 Normal, float Radius, float4 Color, uint SegmentCount = 16) { AddCircleTWS(Ctx, Center + Ctx.Config.TranslatedWorldOffset, Normal, Radius, Color, SegmentCount); }
+void AddCircle   (FShaderPrintContext Ctx, float3 Center, float3 Normal, float Radius, float4 Color, uint SegmentCount = 16) { AddCircleTWS(Ctx, Center + Ctx.Config.TranslatedWorldOffset, Normal, Radius, Color, SegmentCount); }
+
+///////////////////////////////////////////////////////////////////////////////////////////////////
+// Spheres
+
+void AddSphereTWS(FShaderPrintContext Ctx, float3 Center, float Radius, float4 Color, uint SegmentCount = 16)
+{
+	AddCircleTWS(Ctx, Center, float3(1, 0, 0), Radius, Color, SegmentCount);
+	AddCircleTWS(Ctx, Center, float3(0, 1, 0), Radius, Color, SegmentCount);
+	AddCircleTWS(Ctx, Center, float3(0, 0, 1), Radius, Color, SegmentCount);
+}
+
+void AddSphereWS (FShaderPrintContext Ctx, float3 Center, float Radius, float4 Color, uint SegmentCount = 16) { AddSphereTWS(Ctx, Center + Ctx.Config.TranslatedWorldOffset, Radius, Color, SegmentCount); }
+void AddSphere   (FShaderPrintContext Ctx, float3 Center, float Radius, float4 Color, uint SegmentCount = 16) { AddSphereTWS(Ctx, Center + Ctx.Config.TranslatedWorldOffset, Radius, Color, SegmentCount); }
+
 ///////////////////////////////////////////////////////////////////////////////////////////////////
 // Referential
 
