HeterogeneousVolumes: Refactoring voxel grid debug rendering. Enumerating volume tracking algorithms.

diff --git a/Engine/Shaders/Private/HeterogeneousVolumes/HeterogeneousVolumesMaterialBakingPipeline.usf b/Engine/Shaders/Private/HeterogeneousVolumes/HeterogeneousVolumesMaterialBakingPipeline.usf
--- a/Engine/Shaders/Private/HeterogeneousVolumes/HeterogeneousVolumesMaterialBakingPipeline.usf
+++ b/Engine/Shaders/Private/HeterogeneousVolumes/HeterogeneousVolumesMaterialBakingPipeline.usf
@@ -5,6 +5,7 @@
 #define SUPPORT_CONTACT_SHADOWS 0
 
 #include "/Engine/Generated/Material.ush"
+#include "../BlueNoise.ush"
 #include "../PositionReconstructionCommon.ush"
 #include "../MortonCode.ush"
 #include "../ComputeShaderUtils.ush"
@@ -180,7 +181,7 @@ uint2 CalcStencil(uint VoxelCount)
 #endif
 
 [numthreads(THREADGROUP_SIZE_3D, THREADGROUP_SIZE_3D, THREADGROUP_SIZE_3D)]
-void RasterizeBottomLevelGridCS(
+void RasterizeBottomLevelOrthoGridCS(
 	uint3 GroupId : SV_GroupID,
 	uint3 GroupThreadId : SV_GroupThreadID
 )
@@ -215,6 +216,13 @@ void RasterizeBottomLevelGridCS(
 		int3 BottomLevelVoxelIndex = MortonDecode3(BottomLevelGridLinearOffset) * THREADGROUP_SIZE_3D + GroupThreadId;
 	#endif
 
+		// Unwrap the 3D index as a 2D index..
+		int2 BottomLevelVoxelIndexAs2D = int2(BottomLevelVoxelIndex.y * BottomLevelVoxelResolution.x + BottomLevelVoxelIndex.x, BottomLevelVoxelIndex.z);
+
+		float3 Jitter = 0.5;
+		Jitter.xy = BlueNoiseVec2(BottomLevelVoxelIndexAs2D, View.StateFrameIndex);
+		Jitter.z = BlueNoiseScalar(BottomLevelVoxelIndexAs2D, View.StateFrameIndex);
+
 		float3 Extinction = 0;
 		float3 Emission = 0;
 		float3 Albedo = 0;
@@ -227,7 +235,7 @@ void RasterizeBottomLevelGridCS(
 			}
 			else
 			{
-				UVW = float3(BottomLevelVoxelIndex + 0.5) / float3(BottomLevelVoxelResolution);
+				UVW = float3(BottomLevelVoxelIndex + Jitter) / float3(BottomLevelVoxelResolution);
 			}
 			float3 WorldPosition = UVW * TopLevelGridVoxelWorldBoundsExtent + TopLevelGridVoxelWorldBoundsMin;
 
@@ -396,7 +404,6 @@ float NearPlaneDepth;
 float FarPlaneDepth;
 
 #include "HeterogeneousVolumesFrustumVoxelGridUtils.ush"
-#include "../BlueNoise.ush"
 
 groupshared float GSExtinctionSumScalar[THREADGROUP_SIZE_1D];
 
@@ -406,7 +413,7 @@ bool EvaluationPositionIntersectsPrimitive(float3 WorldPosition)
 }
 
 [numthreads(THREADGROUP_SIZE_3D, THREADGROUP_SIZE_3D, THREADGROUP_SIZE_3D)]
-void RasterizeBottomLevelFroxelGridCS(
+void RasterizeBottomLevelFrustumGridCS(
 	uint3 GroupId : SV_GroupID,
 	uint3 GroupThreadId : SV_GroupThreadID
 )
diff --git a/Engine/Shaders/Private/HeterogeneousVolumes/HeterogeneousVolumesVoxelGridPipeline.usf b/Engine/Shaders/Private/HeterogeneousVolumes/HeterogeneousVolumesVoxelGridPipeline.usf
--- a/Engine/Shaders/Private/HeterogeneousVolumes/HeterogeneousVolumesVoxelGridPipeline.usf
+++ b/Engine/Shaders/Private/HeterogeneousVolumes/HeterogeneousVolumesVoxelGridPipeline.usf
@@ -20,40 +20,20 @@
 #define THREADGROUP_SIZE_3D 1
 #endif // THREADGROUP_SIZE_3D
 
-#ifndef ENABLE_ADAPTIVE_VOXEL_GRID_UNIFORM_BUFFER
-#define ENABLE_ADAPTIVE_VOXEL_GRID_UNIFORM_BUFFER 0
-#endif // ENABLE_ADAPTIVE_VOXEL_GRID_UNIFORM_BUFFER
-
 #ifndef DIM_ENABLE_INDIRECTION_GRID
 #define DIM_ENABLE_INDIRECTION_GRID 0
 #endif // DIM_ENABLE_INDIRECTION_GRID
 
-#if ENABLE_ADAPTIVE_VOXEL_GRID_UNIFORM_BUFFER
-	#define TopLevelGridResolution OrthoGridUniformBuffer.TopLevelGridResolution
-	#define TopLevelGridWorldBoundsMin OrthoGridUniformBuffer.TopLevelGridWorldBoundsMin
-	#define TopLevelGridWorldBoundsMax OrthoGridUniformBuffer.TopLevelGridWorldBoundsMax
-
-	#define TopLevelGridBitmaskBuffer OrthoGridUniformBuffer.TopLevelGridBitmaskBuffer
-	#define TopLevelGridBuffer OrthoGridUniformBuffer.TopLevelGridBuffer
-	#define IndirectionGridBuffer OrthoGridUniformBuffer.IndirectionGridBuffer
-
-	#define ExtinctionGridBuffer OrthoGridUniformBuffer.ExtinctionGridBuffer
-	#define EmissionGridBuffer OrthoGridUniformBuffer.EmissionGridBuffer
-	#define ScatteringGridBuffer OrthoGridUniformBuffer.ScatteringGridBuffer
-
-#else
-	int3 TopLevelGridResolution;
-	float3 TopLevelGridWorldBoundsMin;
-	float3 TopLevelGridWorldBoundsMax;
-
-	StructuredBuffer<FTopLevelGridBitmaskData> TopLevelGridBitmaskBuffer;
-	StructuredBuffer<FTopLevelGridData> TopLevelGridBuffer;
-	StructuredBuffer<FTopLevelGridData> IndirectionGridBuffer;
-	StructuredBuffer<FScalarGridData> ExtinctionGridBuffer;
-	StructuredBuffer<FVectorGridData> EmissionGridBuffer;
-	StructuredBuffer<FVectorGridData> ScatteringGridBuffer;
+int3 TopLevelGridResolution;
+float3 TopLevelGridWorldBoundsMin;
+float3 TopLevelGridWorldBoundsMax;
 
-#endif // !ENABLE_ADAPTIVE_VOXEL_GRID_UNIFORM_BUFFER
+StructuredBuffer<FTopLevelGridBitmaskData> TopLevelGridBitmaskBuffer;
+StructuredBuffer<FTopLevelGridData> TopLevelGridBuffer;
+StructuredBuffer<FTopLevelGridData> IndirectionGridBuffer;
+StructuredBuffer<FScalarGridData> ExtinctionGridBuffer;
+StructuredBuffer<FVectorGridData> EmissionGridBuffer;
+StructuredBuffer<FVectorGridData> ScatteringGridBuffer;
 
 float3 PrimitiveWorldBoundsMin;
 float3 PrimitiveWorldBoundsMax;
@@ -494,7 +474,7 @@ void BuildMajorantVoxelGridCS(
 			{
 				uint BottomLevelIndex = IndirectionIndex + VoxelIndex;
 				float3 Extinction = GetExtinction(ExtinctionGridBuffer[BottomLevelIndex]);
-				UpdateMajorantData(MajorantData, Luminance(Extinction));
+				UpdateMajorantData(MajorantData, Extinction);
 			}
 		}
 #else // DIM_ENABLE_INDIRECTION_GRID
@@ -502,7 +482,7 @@ void BuildMajorantVoxelGridCS(
 		{
 			int BottomLevelIndex = GetBottomLevelIndex(TopLevelGridData) + VoxelIndex;
 			float3 Extinction = GetExtinction(ExtinctionGridBuffer[BottomLevelIndex]);
-			UpdateMajorantData(MajorantData, Luminance(Extinction));
+			UpdateMajorantData(MajorantData, Extinction);
 		}
 #endif // DIM_ENABLE_INDIRECTION_GRID
 	}
@@ -510,142 +490,50 @@ void BuildMajorantVoxelGridCS(
 	SetMajorantData(RWMajorantVoxelGridBuffer[LinearIndex], MajorantData);
 }
 
-// Ray data
-int bJitter;
-float MaxTraceDistance;
-float StepSize;
-int MaxStepCount;
+StructuredBuffer<FScalarGridData> InputBuffer;
+RWStructuredBuffer<FScalarGridData> RWOutputBuffer;
 
-// Dispatch data
-int3 GroupCount;
+int3 InputDimensions;
 
-// Output
-RWTexture2D<float4> RWLightingTexture;
+groupshared float GSMajorant[THREADGROUP_SIZE_1D];
 
-[numthreads(THREADGROUP_SIZE_2D, THREADGROUP_SIZE_2D, 1)]
-void RenderTransmittanceWithVoxelGridCS(
-	uint2 GroupThreadId : SV_GroupThreadID,
-	uint2 DispatchThreadId : SV_DispatchThreadID
+[numthreads(THREADGROUP_SIZE_3D, THREADGROUP_SIZE_3D, THREADGROUP_SIZE_3D)]
+void DownsampleMajorantVoxelGridCS(
+	uint3 GroupId : SV_GroupID,
+	uint3 GroupThreadId : SV_GroupThreadID,
+	uint3 DispatchThreadId : SV_DispatchThreadID
 )
 {
-	float3 Radiance = 0.0;
-	float3 Transmittance = 1.0;
-
-	// Create screen ray
-	if (any(DispatchThreadId.xy >= View.ViewSizeAndInvSize.xy))
+	uint LinearThreadIndex = MortonEncode3(GroupThreadId);
+	uint LinearDispatchIndex = GetLinearIndex(DispatchThreadId, InputDimensions);
+	GSMajorant[LinearThreadIndex] = 0;
+	if (all(DispatchThreadId < InputDimensions))
 	{
-		return;
+		FMajorantData MajorantData = GetMajorantData(InputBuffer[LinearDispatchIndex]);
+		GSMajorant[LinearThreadIndex] = MajorantData.Majorant;
 	}
-	uint LinearIndex = DispatchThreadId.y * GroupCount.x * THREADGROUP_SIZE_2D + DispatchThreadId.x;
-	uint2 PixelCoord = DispatchThreadId.xy + View.ViewRectMin.xy;
-
-	// Extract depth
-	float DeviceZ = SceneDepthTexture.Load(int3(PixelCoord, 0)).r;
-#if HAS_INVERTED_Z_BUFFER
-	DeviceZ = max(0.000000000001, DeviceZ);
-#endif // HAS_INVERTED_Z_BUFFER
 
-	// Clip trace distance
-	float SceneDepth = min(ConvertFromDeviceZ(DeviceZ), MaxTraceDistance);
-	DeviceZ = ConvertToDeviceZ(SceneDepth);
-
-	float3 RayOrigin = View.TranslatedWorldCameraOrigin;
-	float3 RayEnd = SvPositionToTranslatedWorld(float4(PixelCoord + 0.5, DeviceZ, 1));
-	float3 RayDirection = RayEnd - RayOrigin;
-	float RayTMin = 0.0;
-	float RayTMax = length(RayDirection);
-	RayDirection /= RayTMax;
-
-	float FrustumRayTMin = RayTMin;
-	float FrustumRayTMax = RayTMax;
-	if (FrustumGridUniformBuffer.bUseFrustumGrid)
+	// Parallel Sum
+	for (int NeighborOffsetSum = 1; NeighborOffsetSum < THREADGROUP_SIZE_1D; NeighborOffsetSum = (NeighborOffsetSum << 1))
 	{
-		float ViewNearZ = min(ConvertFromDeviceZ(DeviceZ), FrustumGridUniformBuffer.NearPlaneDepth);
-		float DeviceNearZ = ConvertToDeviceZ(ViewNearZ);
-		float ViewFarZ = min(ConvertFromDeviceZ(DeviceZ), FrustumGridUniformBuffer.FarPlaneDepth);
-		float DeviceFarZ = ConvertToDeviceZ(ViewFarZ);
-
-		// TODO: Shared code is imprecise, in comparison with commented out frustum-trimmed ray. Investigate...
-#if 0
-		float3 RayBegin = SvPositionToTranslatedWorld(float4(PixelCoord + 0.5, DeviceNearZ, 1));
-		float3 RayEnd = SvPositionToTranslatedWorld(float4(PixelCoord + 0.5, DeviceFarZ, 1));
-		float3 RayDirection = RayEnd - RayBegin;
-		FrustumRayTMin = 0.0;
-		FrustumRayTMax = length(RayDirection);
-		RayDirection /= FrustumRayTMax;
-
-		Transmittance = TraceFrustumVoxelGrid(RayBegin, RayDirection, FrustumRayTMin, FrustumRayTMax);
-		float RayTOffset = length(RayBegin - RayOrigin);
-		FrustumRayTMin += RayTOffset;
-		FrustumRayTMax += RayTOffset;
-#else
-		float3 WorldRayOrigin = LWCHackToFloat(SvPositionToWorld(float4(PixelCoord + 0.5, DeviceNearZ, 1)));
-		float3 WorldRayEnd = LWCHackToFloat(SvPositionToWorld(float4(PixelCoord + 0.5, DeviceFarZ, 1)));
-		float3 WorldRayDirection = WorldRayEnd - WorldRayOrigin;
-		float WorldRayLength = length(WorldRayDirection);
-		WorldRayDirection /= WorldRayLength;
-
-
-		// Convert to view-space
-		float3 ViewRayOrigin = mul(float4(WorldRayOrigin, 1), FrustumGridUniformBuffer.WorldToView).xyz;
-		float3 ViewRayEnd = mul(float4(WorldRayEnd, 1), FrustumGridUniformBuffer.WorldToView).xyz;
-		float3 ViewRayDirection = ViewRayEnd - ViewRayOrigin;
-		float ViewRayLength = length(ViewRayDirection);
-		ViewRayDirection /= ViewRayLength;
-
-		// Convert to voxel-space
-		int3 VoxelDimensions = FrustumGridUniformBuffer.VoxelDimensions;
-		float NearPlaneDepth = FrustumGridUniformBuffer.NearPlaneDepth;
-		float FarPlaneDepth = FrustumGridUniformBuffer.FarPlaneDepth;
-		float3 VoxelRayOrigin = ViewToVoxel(ViewRayOrigin, VoxelDimensions, NearPlaneDepth, FarPlaneDepth);
-		float3 VoxelRayEnd = ViewToVoxel(ViewRayOrigin + ViewRayDirection * ViewRayLength, VoxelDimensions, NearPlaneDepth, FarPlaneDepth);
-
-		float3 VoxelRayDirection = VoxelRayEnd - VoxelRayOrigin;
-		float VoxelRayLength = length(VoxelRayDirection);
-		VoxelRayDirection /= VoxelRayLength;
-
-		float VoxelRayTMin = 0.0;
-		float VoxelRayTMax = VoxelRayLength;
-
-		// Traversal
-		{
-			MarchTopLevelFrustumVoxelGrid(
-				VoxelRayOrigin,
-				VoxelRayDirection,
-				VoxelRayTMin,
-				VoxelRayTMax,
-				WorldRayLength,
-				FrustumGridUniformBuffer.TopLevelFroxelGridResolution,
-				FrustumGridUniformBuffer.TopLevelFroxelGridBuffer,
-				FrustumGridUniformBuffer.ExtinctionFroxelGridBuffer,
-				Transmittance
-			);
-		}
+		int NeighborIndex = LinearThreadIndex + NeighborOffsetSum;
+		float NeighborMajorant = (NeighborIndex < THREADGROUP_SIZE_1D) ? GSMajorant[NeighborIndex] : 0.0;
+		GSMajorant[LinearThreadIndex] = max(GSMajorant[LinearThreadIndex], NeighborMajorant);
 
-		float3 RayBegin = SvPositionToTranslatedWorld(float4(PixelCoord + 0.5, DeviceNearZ, 1));
-		float RayTOffset = length(RayBegin - RayOrigin);
-		FrustumRayTMin += RayTOffset;
-		FrustumRayTMax += RayTOffset;
-#endif
+		GroupMemoryBarrierWithGroupSync();
 	}
 
-	if (OrthoGridUniformBuffer.bUseOrthoGrid)
-	{
-		if (FrustumGridUniformBuffer.bUseFrustumGrid && (FrustumRayTMin < FrustumRayTMax))
-		{
-			Transmittance *= TraceOrthoVoxelGrid(RayOrigin, RayDirection, RayTMin, FrustumRayTMin);
-			Transmittance *= TraceOrthoVoxelGrid(RayOrigin, RayDirection, FrustumRayTMax, RayTMax);
-		}
-		else
-		{
-			Transmittance = TraceOrthoVoxelGrid(RayOrigin, RayDirection, RayTMin, RayTMax);
-		}
-	}
+	int3 OutputDimensions = int3((InputDimensions + 3) / 4);
+	uint LinearOuputIndex = GetLinearIndex(GroupId, OutputDimensions);
 
-	// Debug
-	Radiance = 1.0 - Transmittance;
+	FMajorantData MajorantData;
+	MajorantData.Minorant = 0;
+	MajorantData.Majorant = GSMajorant[0];
+	SetMajorantData(RWOutputBuffer[LinearOuputIndex], MajorantData);
+}
 
-	// Output..
-	float3 Opacity = 1.0 - Transmittance;
-	RWLightingTexture[PixelCoord] += float4(Radiance, Luminance(Opacity));
+[numthreads(1, 1, 1)]
+void CopyMaximumIntoMajorantVoxelGridCS()
+{
+	SetMajorantData(RWMajorantVoxelGridBuffer[0], GetMajorantData(InputBuffer[0]));
 }
\ No newline at end of file
diff --git a/Engine/Shaders/Private/HeterogeneousVolumes/HeterogeneousVolumesVoxelGridRendering.usf b/Engine/Shaders/Private/HeterogeneousVolumes/HeterogeneousVolumesVoxelGridRendering.usf
new file mode 100644
--- /dev/null
+++ b/Engine/Shaders/Private/HeterogeneousVolumes/HeterogeneousVolumesVoxelGridRendering.usf
@@ -0,0 +1,823 @@
+// Copyright Epic Games, Inc. All Rights Reserved.
+
+#include "HeterogeneousVolumesVoxelGridTypes.ush"
+
+#include "../Common.ush"
+#include "../DeferredShadingCommon.ush"
+#include "../ComputeShaderUtils.ush"
+#include "HeterogeneousVolumesTracingUtils.ush"
+#include "HeterogeneousVolumesVoxelGridTraversal.ush"
+
+#include "../PathTracing/Utilities/PathTracingRandomSequence.ush"
+
+#define DEBUG_MODE_DDA_MARCHING 0
+
+#define DEBUG_MODE_SAMPLE_BASED_MARCHING 1
+#define DEBUG_MODE_SAMPLE_BASED_MAJORANT 2
+#define DEBUG_MODE_SAMPLE_BASED_MAX_MAJORANT 3
+#define DEBUG_MODE_SAMPLE_BASED_MAJORANT_DDA 4
+
+#define DEBUG_MODE_TRACKING_BASED_DELTA_TRACKING 5
+#define DEBUG_MODE_TRACKING_BASED_RATIO_TRACKING 6
+#define DEBUG_MODE_TRACKING_BASED_SPECTRAL_RATIO_TRACKING 7
+#define DEBUG_MODE_TRACKING_BASED_MAJORANT_DDA 8
+
+#define DEBUG_MODE_SAMPLING_BASED \
+	(DEBUG_MODE == DEBUG_MODE_SAMPLE_BASED_MARCHING) || \
+	(DEBUG_MODE == DEBUG_MODE_SAMPLE_BASED_MAJORANT) || \
+	(DEBUG_MODE == DEBUG_MODE_SAMPLE_BASED_MAX_MAJORANT) || \
+	(DEBUG_MODE == DEBUG_MODE_SAMPLE_BASED_MAJORANT_DDA)
+
+#define DEBUG_MODE_TRACKING_BASED \
+	(DEBUG_MODE == DEBUG_MODE_TRACKING_BASED_DELTA_TRACKING) || \
+	(DEBUG_MODE == DEBUG_MODE_TRACKING_BASED_RATIO_TRACKING) || \
+	(DEBUG_MODE == DEBUG_MODE_TRACKING_BASED_SPECTRAL_RATIO_TRACKING) || \
+	(DEBUG_MODE == DEBUG_MODE_TRACKING_BASED_MAJORANT_DDA)
+
+#define COLLISION_MODE_RATIO_TRACKING 0
+#define COLLISION_MODE_MAX_COMPONENT 1
+#define COLLISION_MODE_AVG_COMPONENT 2
+
+struct FRenderDebugData
+{
+	float3 RayOrigin;
+	float3 RayDirection;
+	float TMin;
+	float TMax;
+	float3 Estimate;
+	float Pdf;
+	float Distance;
+};
+
+// Ray data
+int bJitter;
+float MaxTraceDistance;
+float StepSize;
+int MaxStepCount;
+
+// Dispatch data
+int3 GroupCount;
+int DebugMode;
+
+// Output
+RWTexture2D<float4> RWLightingTexture;
+RWStructuredBuffer<FRenderDebugData> RWDebugBuffer;
+
+struct FTrackingResult
+{
+	float3 Radiance;
+	float3 Throughput;
+	float Pdf;
+	float Distance;
+	bool bIsScattering;
+};
+
+FTrackingResult CreateTrackingResult(float Distance)
+{
+	FTrackingResult Result = (FTrackingResult)0;
+
+	Result.Radiance = 0;
+	Result.Throughput = 1;
+	Result.Pdf = 0;
+	Result.Distance = Distance;
+	Result.bIsScattering = false;
+
+	return Result;
+}
+
+// Concatenates the result of rhs onto lhs
+FTrackingResult ConcatenateTrackingResult(FTrackingResult lhs, FTrackingResult rhs)
+{
+	FTrackingResult Result;
+	//Result.Radiance = lhs.Radiance + rhs.Radiance * lhs.Throughput;
+	Result.Radiance = rhs.Radiance * lhs.Throughput;
+	Result.Throughput = lhs.Throughput * rhs.Throughput;
+	Result.Pdf = lhs.Pdf * rhs.Pdf;
+	Result.Distance = lhs.Distance + rhs.Distance;
+	Result.bIsScattering = rhs.bIsScattering;
+	return Result;
+}
+
+FTrackingResult AddTrackingResult(FTrackingResult lhs, FTrackingResult rhs)
+{
+	FTrackingResult Sum;
+	Sum.Radiance = lhs.Radiance + rhs.Radiance;
+	Sum.Throughput = lhs.Throughput + rhs.Throughput;
+	Sum.Pdf = lhs.Pdf + rhs.Pdf;
+	Sum.Distance = lhs.Distance + rhs.Distance;
+	return Sum;
+}
+
+FTrackingResult DivideTrackingResult(FTrackingResult Value, int SPP)
+{
+	FTrackingResult Quotient;
+	Quotient.Radiance = Value.Radiance / SPP;
+	Quotient.Throughput = Value.Throughput / SPP;
+	Quotient.Pdf = Value.Pdf / SPP;
+	Quotient.Distance = Value.Distance / SPP;
+	return Quotient;
+}
+
+float3 GetTranslatedWorldPos(float3 WorldPos)
+{
+	float3 TranslatedWorldPos = LWCHackToFloat(LWCAdd(WorldPos, PrimaryView.PreViewTranslation));
+	return TranslatedWorldPos;
+}
+
+FMajorantData EvalMajorant(float3 TranslatedWorldPos)
+{
+	FMajorantData MajorantData = CreateMajorantData();
+
+	float3 TranslatedWorldBoundsMin = GetTranslatedWorldPos(OrthoGridUniformBuffer.TopLevelGridWorldBoundsMin);
+	float3 TranslatedWorldBoundsMax = GetTranslatedWorldPos(OrthoGridUniformBuffer.TopLevelGridWorldBoundsMax);
+	float3 TranslatedWorldBoundsExtent = TranslatedWorldBoundsMax - TranslatedWorldBoundsMin;
+	float3 GridUV = (TranslatedWorldPos - TranslatedWorldBoundsMin) / TranslatedWorldBoundsExtent;
+	if (all(GridUV >= 0.0) && all(GridUV <= 1.0))
+	{
+		float3 TopLevelVoxelPos = GridUV * OrthoGridUniformBuffer.TopLevelGridResolution;
+		uint LinearTopLevelVoxelPos = GetLinearIndex(TopLevelVoxelPos, OrthoGridUniformBuffer.TopLevelGridResolution);
+
+		UpdateMajorantData(MajorantData, GetMajorantData(OrthoGridUniformBuffer.MajorantGridBuffer[LinearTopLevelVoxelPos]));
+	}
+
+	return MajorantData;
+}
+
+float3 EvalExtinction(float3 TranslatedWorldPos, inout int LinearBottomLevelVoxelPos)
+{
+	float3 Extinction = 0;
+
+	float3 TranslatedWorldBoundsMin = GetTranslatedWorldPos(OrthoGridUniformBuffer.TopLevelGridWorldBoundsMin);
+	float3 TranslatedWorldBoundsMax = GetTranslatedWorldPos(OrthoGridUniformBuffer.TopLevelGridWorldBoundsMax);
+	float3 TranslatedWorldBoundsExtent = TranslatedWorldBoundsMax - TranslatedWorldBoundsMin;
+	float3 GridUV = (TranslatedWorldPos - TranslatedWorldBoundsMin) / TranslatedWorldBoundsExtent;
+	if (all(GridUV >= 0.0) && all(GridUV <= 1.0))
+	{
+		float3 TopLevelVoxelPos = GridUV * OrthoGridUniformBuffer.TopLevelGridResolution;
+		uint LinearTopLevelVoxelPos = GetLinearIndex(TopLevelVoxelPos, OrthoGridUniformBuffer.TopLevelGridResolution);
+
+		if (OrthoGridUniformBuffer.bEnableIndirectionGrid)
+		{
+			FTopLevelGridData TopLevelGridData = OrthoGridUniformBuffer.TopLevelGridBuffer[LinearTopLevelVoxelPos];
+			if (IsBottomLevelAllocated(TopLevelGridData))
+			{
+				uint IndirectionGridIndex = GetBottomLevelIndex(TopLevelGridData);
+				uint3 IndirectionGridVoxelResolution = GetBottomLevelVoxelResolution(TopLevelGridData);
+
+				// Convert to Indirection voxel-space
+				float3 IndirectionGridVoxelPos = frac(TopLevelVoxelPos) * IndirectionGridVoxelResolution;
+				int3 IndirectionGridVoxelPosAsInt = clamp(IndirectionGridVoxelPos, 0, IndirectionGridVoxelResolution - 1);
+				FTopLevelGridData IndirectionData = OrthoGridUniformBuffer.IndirectionGridBuffer[IndirectionGridIndex + MortonEncode3(IndirectionGridVoxelPosAsInt)];
+				if (IsBottomLevelAllocated(IndirectionData))
+				{
+					uint BottomLevelIndex = GetBottomLevelIndex(IndirectionData);
+					uint3 BottomLevelVoxelResolution = GetBottomLevelVoxelResolution(IndirectionData);
+
+					// Constant interpolation
+					float3 BottomLevelVoxelPos = frac(IndirectionGridVoxelPos) * BottomLevelVoxelResolution;
+					int3 BottomLevelVoxelPosAsInt = clamp(BottomLevelVoxelPos, 0, BottomLevelVoxelResolution - 1);
+					LinearBottomLevelVoxelPos = BottomLevelIndex + MortonEncode3(BottomLevelVoxelPosAsInt);
+					Extinction = GetExtinction(OrthoGridUniformBuffer.ExtinctionGridBuffer[LinearBottomLevelVoxelPos]);
+					//Albedo = GetAlbedo(OrthoGridUniformBuffer.ScatteringGridBuffer[LinearBottomLevelVoxelPos]);
+				}
+			}
+		}
+		else
+		{
+			FTopLevelGridData TopLevelData = OrthoGridUniformBuffer.TopLevelGridBuffer[LinearTopLevelVoxelPos];
+			if (IsBottomLevelAllocated(TopLevelData))
+			{
+				uint BottomLevelIndex = GetBottomLevelIndex(OrthoGridUniformBuffer.TopLevelGridBuffer[LinearTopLevelVoxelPos]);
+				uint3 BottomLevelVoxelResolution = GetBottomLevelVoxelResolution(OrthoGridUniformBuffer.TopLevelGridBuffer[LinearTopLevelVoxelPos]);
+
+				// Constant Interpolation
+				float3 BottomLevelVoxelPos = frac(TopLevelVoxelPos) * BottomLevelVoxelResolution;
+				int3 BottomLevelVoxelPosAsInt = clamp(BottomLevelVoxelPos, 0, BottomLevelVoxelResolution - 1);
+				LinearBottomLevelVoxelPos = BottomLevelIndex + MortonEncode3(BottomLevelVoxelPosAsInt);
+
+				Extinction = GetExtinction(OrthoGridUniformBuffer.ExtinctionGridBuffer[LinearBottomLevelVoxelPos]);
+				//Albedo = GetAlbedo(OrthoGridUniformBuffer.ScatteringGridBuffer[LinearBottomLevelVoxelPos]);
+			}
+		}
+	}
+
+	return Extinction;
+}
+
+float CalcStepSize(float3 WorldDirection)
+{
+	float3 TranslatedWorldBoundsMin = GetTranslatedWorldPos(OrthoGridUniformBuffer.TopLevelGridWorldBoundsMin);
+	float3 TranslatedWorldBoundsMax = GetTranslatedWorldPos(OrthoGridUniformBuffer.TopLevelGridWorldBoundsMax);
+	float3 TranslatedWorldBoundsExtent = TranslatedWorldBoundsMax - TranslatedWorldBoundsMin;
+	float3 BottomLevelVoxelResolution = OrthoGridUniformBuffer.TopLevelGridResolution * 4;
+	// Conditionally apply indirection dimensions
+	BottomLevelVoxelResolution *= 4;
+
+	float3 VoxelSize = TranslatedWorldBoundsExtent / BottomLevelVoxelResolution;
+	float3 VoxelHitT = VoxelSize / abs(WorldDirection);
+	float StepSize = min(VoxelHitT.x, min(VoxelHitT.y, VoxelHitT.z));
+	return StepSize;
+}
+
+FTrackingResult SampleBasedMarch(float3 WorldRayOrigin, float3 WorldRayDirection, float TMin, float TMax, float WorldStepSize)
+{
+	FTrackingResult Result = CreateTrackingResult(TMin);
+
+	float HitT = TMin;
+	const float Epsilon = 1.0e-4;
+	while (HitT < TMax && any(Result.Throughput > Epsilon))
+	{
+		float3 WorldPosition = WorldRayOrigin + WorldRayDirection * HitT;
+
+		int LinearVoxelPos = -1;
+		float3 Extinction = EvalExtinction(WorldPosition, LinearVoxelPos);
+		if (any(Extinction > 0))
+		{
+			Result.Radiance += GetEmission(OrthoGridUniformBuffer.EmissionGridBuffer[LinearVoxelPos]) * Result.Throughput;
+		}
+		Result.Throughput *= exp(-Extinction * WorldStepSize);
+
+		HitT += WorldStepSize;
+	}
+	Result.Distance = HitT;
+
+	return Result;
+}
+
+FTrackingResult SampleBasedMajorantMarch(float3 WorldRayOrigin, float3 WorldRayDirection, float TMin, float TMax, float WorldStepSize)
+{
+	FTrackingResult Result = (FTrackingResult)0;
+	Result.Throughput = 1.0;
+
+	float HitT = TMin;
+	const float Epsilon = 1.0e-3;
+	while (HitT < TMax && any(Result.Throughput > Epsilon))
+	{
+		float3 WorldPosition = WorldRayOrigin + WorldRayDirection * HitT;
+
+		FMajorantData MajorantData = EvalMajorant(WorldPosition);
+		float3 Extinction = MajorantData.Majorant;
+		Result.Throughput *= exp(-Extinction * WorldStepSize);
+
+		HitT += WorldStepSize;
+	}
+	Result.Distance = HitT;
+	Result.Pdf = Result.Throughput;
+
+	return Result;
+}
+
+FTrackingResult SampleBasedMaxMajorantMarch(float3 WorldRayOrigin, float3 WorldRayDirection, float TMin, float TMax, float WorldStepSize)
+{
+	FTrackingResult Result = (FTrackingResult)0;
+	Result.Throughput = 1.0;
+
+	float HitT = TMin;
+	const float Epsilon = 1.0e-3;
+	while (HitT < TMax && any(Result.Throughput > Epsilon))
+	{
+		float3 WorldPosition = WorldRayOrigin + WorldRayDirection * HitT;
+
+		FMajorantData MajorantData = GetMajorantData(OrthoGridUniformBuffer.MajorantGridBuffer[0]);
+		float3 Extinction = MajorantData.Majorant;
+		Result.Throughput *= exp(-Extinction * WorldStepSize);
+
+		HitT += WorldStepSize;
+	}
+	Result.Distance = HitT;
+	Result.Pdf = Result.Throughput;
+
+	return Result;
+}
+
+float MaxComponent(float3 Value)
+{
+	return max(Value.x, max(Value.y, Value.z));
+}
+
+FTrackingResult DeltaTracking(float3 WorldRayOrigin, float3 WorldRayDirection, float TMin, float TMax, float SigmaBar, inout RandomSequence RandSequence)
+{
+	FTrackingResult Result = CreateTrackingResult(TMin);
+
+	float HitT = 0;
+	while (Result.Distance < TMax)
+	{
+		// Pick a distance, based on SigmaBar
+		float RandValue = RandomSequence_GenerateSample1D(RandSequence);
+		HitT = -log(RandValue) / SigmaBar;
+		Result.Distance += HitT;
+
+		// Evaluate SigmaT at that location
+		if (Result.Distance < TMax)
+		{
+			float3 WorldPosition = WorldRayOrigin + WorldRayDirection * Result.Distance;
+			int LinearVoxelPos = -1;
+			float3 Extinction = EvalExtinction(WorldPosition, LinearVoxelPos);
+			float SigmaT = Luminance(Extinction);
+
+			//float AcceptanceProbability = MaxComponent(SigmaT) / SigmaBar;
+			float AcceptanceProbability = min(SigmaT / SigmaBar, 1);
+			if (RandomSequence_GenerateSample1D(RandSequence) < AcceptanceProbability)
+			{
+				Result.Throughput = 0;
+				Result.Radiance = GetEmission(OrthoGridUniformBuffer.EmissionGridBuffer[LinearVoxelPos]) / Luminance(SigmaT);
+				break;
+			}
+		}
+	}
+	Result.Distance = min(Result.Distance, TMax);
+	Result.Pdf = SigmaBar * exp(-SigmaBar * (Result.Distance - TMin));
+
+	return Result;
+}
+
+FTrackingResult SpectralWeightedDeltaTracking(float3 WorldRayOrigin, float3 WorldRayDirection, float TMin, float TMax, float SigmaBar, int CollisionMode, inout RandomSequence RandSequence)
+{
+	FTrackingResult Result = CreateTrackingResult(TMin);
+
+	float HitT = 0;
+	const float Epsilon = 1.0e-4;
+	while (Result.Distance < TMax && all(Result.Throughput > Epsilon))
+	{
+		// Distance-sampling, based on SigmaBar
+		float RandValue = RandomSequence_GenerateSample1D(RandSequence);
+		HitT = -log(RandValue) / SigmaBar;
+		Result.Distance += HitT;
+
+		if (Result.Distance < TMax)
+		{
+			float3 WorldPosition = WorldRayOrigin + WorldRayDirection * Result.Distance;
+			int LinearVoxelPos = -1;
+			float3 Extinction = EvalExtinction(WorldPosition, LinearVoxelPos);
+
+			// Pick a volume coefficient, based on proportional weight
+			float3 Albedo = GetAlbedo(OrthoGridUniformBuffer.ScatteringGridBuffer[LinearVoxelPos]) / Extinction;
+			float3 SpectralAbsorptionCollisionWeight = Extinction * (1 - Albedo) / SigmaBar;
+			float3 SpectralScatteringCollisionWeight = Extinction * Albedo / SigmaBar;
+			float3 SpectralNullCollisionWeight = 1.0 - Extinction / SigmaBar;
+
+			// Set collision probabilities
+			float AbsorptionProbability = 0;
+			float ScatteringProbability = 0;
+			float NullCollisionProbability = 0;
+			if (CollisionMode == COLLISION_MODE_MAX_COMPONENT)
+			{
+				AbsorptionProbability = MaxComponent(SpectralAbsorptionCollisionWeight);
+				ScatteringProbability = MaxComponent(SpectralScatteringCollisionWeight);
+				NullCollisionProbability = MaxComponent(SpectralNullCollisionWeight);
+			}
+			else if (CollisionMode == COLLISION_MODE_AVG_COMPONENT)
+			{
+				AbsorptionProbability = dot(SpectralAbsorptionCollisionWeight, Result.Throughput);
+				ScatteringProbability = dot(SpectralScatteringCollisionWeight, Result.Throughput);
+				NullCollisionProbability = dot(SpectralNullCollisionWeight, Result.Throughput);
+			}
+			else if (CollisionMode == COLLISION_MODE_RATIO_TRACKING)
+			{
+				AbsorptionProbability = 0;
+				ScatteringProbability = 0;
+				NullCollisionProbability = 1;
+			}
+
+			// Collision-sampling
+			float3 CollisionCdf = float3(AbsorptionProbability,
+				AbsorptionProbability + ScatteringProbability,
+				AbsorptionProbability + ScatteringProbability + NullCollisionProbability
+			);
+			float CollisionSample = RandomSequence_GenerateSample1D(RandSequence) * CollisionCdf.z;
+
+			// Normalize probabilities
+			AbsorptionProbability *= rcp(CollisionCdf.z);
+			ScatteringProbability *= rcp(CollisionCdf.z);
+			NullCollisionProbability *= rcp(CollisionCdf.z);
+
+			// Emission
+			if (CollisionSample < CollisionCdf.x)
+			{
+				Result.Throughput *= SpectralAbsorptionCollisionWeight / AbsorptionProbability;
+				Result.Radiance = GetEmission(OrthoGridUniformBuffer.EmissionGridBuffer[LinearVoxelPos]) * Result.Throughput;
+				Result.Throughput = 0;
+				break;
+			}
+			// Scattering
+			else if (CollisionSample < CollisionCdf.y)
+			{
+				// Compute in-scattering
+				Result.Throughput *= SpectralScatteringCollisionWeight / ScatteringProbability;
+				Result.bIsScattering = true;
+				break;
+			}
+			// Transmittance
+			else
+			{
+				Result.Throughput *= SpectralNullCollisionWeight / NullCollisionProbability;
+			}
+		}
+	}
+	Result.Distance = min(Result.Distance, TMax);
+	Result.Pdf = SigmaBar * exp(-SigmaBar * (Result.Distance - TMin));
+
+	return Result;
+}
+
+FTrackingResult RatioTracking(float3 WorldRayOrigin, float3 WorldRayDirection, float TMin, float TMax, float SigmaBar, inout RandomSequence RandSequence)
+{
+	FTrackingResult Result = CreateTrackingResult(TMin);
+
+	const float Epsilon = 1.0e-5;
+	float HitT = 0;
+	while (Result.Distance < TMax && all(Result.Throughput > Epsilon))
+	{
+		// Pick a distance, based on SigmaBar
+		float RandValue = RandomSequence_GenerateSample1D(RandSequence);
+		HitT = -log(RandValue) / SigmaBar;
+		Result.Distance += HitT;
+
+		// Evaluate SigmaT at that location
+		if (Result.Distance < TMax)
+		{
+			float3 WorldPosition = WorldRayOrigin + WorldRayDirection * Result.Distance;
+			int LinearVoxelPos = -1;
+			float3 Extinction = EvalExtinction(WorldPosition, LinearVoxelPos);
+			float SigmaT = Luminance(Extinction);
+
+			float AcceptanceProbability = min(SigmaT / SigmaBar, 1);
+			if (RandomSequence_GenerateSample1D(RandSequence) < AcceptanceProbability)
+			{
+				//Result.Emission = GetEmission(OrthoGridUniformBuffer.EmissionGridBuffer[LinearVoxelPos]) / Luminance(SigmaT);
+			}
+			Result.Throughput *= (1 - AcceptanceProbability);
+		}
+	}
+	Result.Distance = min(Result.Distance, TMax);
+	Result.Pdf = SigmaBar * exp(-SigmaBar * (Result.Distance - TMin));
+
+	return Result;
+}
+
+FTrackingResult SpectralRatioTracking(float3 WorldRayOrigin, float3 WorldRayDirection, float TMin, float TMax, float SigmaBar, inout RandomSequence RandSequence)
+{
+	FTrackingResult Result = CreateTrackingResult(TMin);
+
+	const float Epsilon = 1.0e-5;
+	float HitT = 0;
+	while (Result.Distance < TMax && all(Result.Throughput > Epsilon))
+	{
+		// Pick a distance, based on SigmaBar
+		float RandValue = RandomSequence_GenerateSample1D(RandSequence);
+		HitT = -log(RandValue) / SigmaBar;
+		Result.Distance += HitT;
+
+		// Evaluate SigmaT at that location
+		if (Result.Distance < TMax)
+		{
+			float3 WorldPosition = WorldRayOrigin + WorldRayDirection * Result.Distance;
+			int LinearVoxelPos = -1;
+			float3 Extinction = EvalExtinction(WorldPosition, LinearVoxelPos);
+
+			// Pick a component, based on proportional weight
+			float3 SpectralCdf = float3(Extinction.x, Extinction.x + Extinction.y, Extinction.x + Extinction.y + Extinction.z);
+			float RandSpectra = RandomSequence_GenerateSample1D(RandSequence) * SpectralCdf.z;
+			float SigmaT = (RandSpectra < SpectralCdf.x) ? Extinction.x : ((RandSpectra < SpectralCdf.y) ? Extinction.y : Extinction.z);
+
+			if (SigmaT > 0)
+			{
+				float3 SpectralWeight = Extinction / SigmaT;
+				float AcceptanceProbability = min(SigmaT / SigmaBar, 1);
+				float3 SpectralAcceptanceProbability = min(SpectralWeight * SigmaT / SigmaBar, 1);
+				if (RandomSequence_GenerateSample1D(RandSequence) < AcceptanceProbability)
+				{
+					//Result.Emission = GetEmission(OrthoGridUniformBuffer.EmissionGridBuffer[LinearVoxelPos]) / Luminance(SigmaT);
+				}
+				Result.Throughput *= (1 - SpectralAcceptanceProbability);
+			}
+		}
+	}
+	Result.Distance = min(Result.Distance, TMax);
+	Result.Pdf = SigmaBar * exp(-SigmaBar * (Result.Distance - TMin));
+
+	return Result;
+}
+
+FTrackingResult DDABasedMajorantMarch(float3 WorldRayOrigin, float3 Direction, float TMin, float TMax, int DistanceSamplingMethod, int CollisionSamplingMethod, inout RandomSequence RandSequence)
+{
+	FTrackingResult TrackingResult = CreateTrackingResult(TMin);
+
+	float3 WorldRayBegin = WorldRayOrigin + Direction * TMin;
+	float3 WorldRayEnd = WorldRayOrigin + Direction * TMax;
+	float WorldRayTMax = length(WorldRayEnd - WorldRayBegin);
+
+	// Transform to voxel-space
+	float3 WorldBoundsMin = GetTranslatedWorldPos(OrthoGridUniformBuffer.TopLevelGridWorldBoundsMin);
+	float3 WorldBoundsMax = GetTranslatedWorldPos(OrthoGridUniformBuffer.TopLevelGridWorldBoundsMax);
+	float3 TopLevelGridWorldBoundsExtent = WorldBoundsMax - WorldBoundsMin;
+	float3 VoxelRayBegin = (WorldRayBegin - WorldBoundsMin) / TopLevelGridWorldBoundsExtent * OrthoGridUniformBuffer.TopLevelGridResolution;
+	float3 VoxelRayEnd = (WorldRayEnd - WorldBoundsMin) / TopLevelGridWorldBoundsExtent * OrthoGridUniformBuffer.TopLevelGridResolution;
+	float3 VoxelRayDirection = VoxelRayEnd - VoxelRayBegin;
+	float VoxelRayTMin = 0.0;
+	float VoxelRayTMax = length(VoxelRayDirection);
+	VoxelRayDirection /= VoxelRayTMax;
+
+	float VoxelToWorldScale = WorldRayTMax / VoxelRayTMax;
+
+	// March majorant grid via DDA
+	float3 VoxelRayDirectionInv = 1.0 / VoxelRayDirection;
+
+	float VoxelRayMarchT = VoxelRayTMin;
+	float3 VoxelRayMarchPos = VoxelRayBegin;
+
+	float3 VoxelBoundsPos;
+	VoxelBoundsPos.x = sign(VoxelRayDirection.x) > 0 ? floor(VoxelRayMarchPos.x) + 1 : ceil(VoxelRayMarchPos.x) - 1;
+	VoxelBoundsPos.y = sign(VoxelRayDirection.y) > 0 ? floor(VoxelRayMarchPos.y) + 1 : ceil(VoxelRayMarchPos.y) - 1;
+	VoxelBoundsPos.z = sign(VoxelRayDirection.z) > 0 ? floor(VoxelRayMarchPos.z) + 1 : ceil(VoxelRayMarchPos.z) - 1;
+
+	float3 VoxelBoundsHitT = (VoxelBoundsPos - VoxelRayMarchPos) * VoxelRayDirectionInv;
+
+	// Remove floating-point rounding error
+	float Epsilon = 1.0e-4;
+	if (VoxelBoundsHitT.x <= Epsilon) VoxelBoundsHitT.x += abs(VoxelRayDirectionInv.x);
+	if (VoxelBoundsHitT.y <= Epsilon) VoxelBoundsHitT.y += abs(VoxelRayDirectionInv.y);
+	if (VoxelBoundsHitT.z <= Epsilon) VoxelBoundsHitT.z += abs(VoxelRayDirectionInv.z);
+
+	while ((VoxelRayMarchT < VoxelRayTMax) && any(TrackingResult.Throughput > Epsilon) && !TrackingResult.bIsScattering)
+	{
+		float VoxelRayMarchDeltaT = 0.0;
+		if (VoxelBoundsHitT.x < VoxelBoundsHitT.y)
+		{
+			if (VoxelBoundsHitT.x < VoxelBoundsHitT.z)
+			{
+				VoxelRayMarchDeltaT = VoxelBoundsHitT.x - VoxelRayMarchT;
+				VoxelBoundsHitT.x += abs(VoxelRayDirectionInv.x);
+			}
+			else
+			{
+				VoxelRayMarchDeltaT = VoxelBoundsHitT.z - VoxelRayMarchT;
+				VoxelBoundsHitT.z += abs(VoxelRayDirectionInv.z);
+			}
+		}
+		else
+		{
+			if (VoxelBoundsHitT.y < VoxelBoundsHitT.z)
+			{
+				VoxelRayMarchDeltaT = VoxelBoundsHitT.y - VoxelRayMarchT;
+				VoxelBoundsHitT.y += abs(VoxelRayDirectionInv.y);
+			}
+			else
+			{
+				VoxelRayMarchDeltaT = VoxelBoundsHitT.z - VoxelRayMarchT;
+				VoxelBoundsHitT.z += abs(VoxelRayDirectionInv.z);
+			}
+		}
+
+		// Clip voxel delta-t by overall voxel ray-length.
+		if (VoxelRayMarchT + VoxelRayMarchDeltaT > VoxelRayTMax)
+		{
+			VoxelRayMarchDeltaT = VoxelRayTMax - VoxelRayMarchT;
+		}
+
+		// Sample at the midpoint of the voxel ray
+		float3 VoxelEndPos = VoxelRayMarchPos + VoxelRayDirection * VoxelRayMarchDeltaT;
+		float3 VoxelSamplePos = (VoxelRayMarchPos + VoxelEndPos) * 0.5;
+		uint VoxelSampleLinearPos = GetLinearIndex(VoxelSamplePos, OrthoGridUniformBuffer.TopLevelGridResolution);
+		FMajorantData MajorantData = GetMajorantData(OrthoGridUniformBuffer.MajorantGridBuffer[VoxelSampleLinearPos]);
+
+		float3 TranslatedWorldPos = WorldRayBegin + Direction * VoxelRayMarchT * VoxelToWorldScale;
+		if (DistanceSamplingMethod == DEBUG_MODE_SAMPLE_BASED_MARCHING)
+		{
+			float StepSize = CalcStepSize(Direction);
+			FTrackingResult Sample = SampleBasedMarch(TranslatedWorldPos, Direction, 0.0, VoxelRayMarchDeltaT * VoxelToWorldScale, StepSize);
+			TrackingResult = ConcatenateTrackingResult(TrackingResult, Sample);
+		}
+		else if (DistanceSamplingMethod == DEBUG_MODE_TRACKING_BASED_SPECTRAL_RATIO_TRACKING)
+		{
+			FTrackingResult Sample = SpectralRatioTracking(TranslatedWorldPos, Direction, 0.0, VoxelRayMarchDeltaT * VoxelToWorldScale, MajorantData.Majorant, RandSequence);
+			TrackingResult = ConcatenateTrackingResult(TrackingResult, Sample);
+		}
+		else if (DistanceSamplingMethod == DEBUG_MODE_TRACKING_BASED_DELTA_TRACKING)
+		{
+			FTrackingResult Sample = SpectralWeightedDeltaTracking(TranslatedWorldPos, Direction, 0.0, VoxelRayMarchDeltaT * VoxelToWorldScale, MajorantData.Majorant, CollisionSamplingMethod, RandSequence);
+			TrackingResult = ConcatenateTrackingResult(TrackingResult, Sample);
+		}
+		else
+		{
+			FTrackingResult Sample = CreateTrackingResult(0);
+			Sample.Distance += VoxelRayMarchDeltaT * VoxelToWorldScale;
+			Sample.Throughput *= exp(-MajorantData.Majorant * Sample.Distance);
+			TrackingResult = ConcatenateTrackingResult(TrackingResult, Sample);
+		}
+
+		VoxelRayMarchT += VoxelRayMarchDeltaT;
+		VoxelRayMarchPos = VoxelEndPos;
+	}
+
+	return TrackingResult;
+}
+
+[numthreads(THREADGROUP_SIZE_2D, THREADGROUP_SIZE_2D, 1)]
+void RenderTransmittanceWithVoxelGridCS(
+	uint2 GroupThreadId : SV_GroupThreadID,
+	uint2 DispatchThreadId : SV_DispatchThreadID
+)
+{
+	float3 Radiance = 0.0;
+
+	// Create screen ray
+	if (any(DispatchThreadId.xy >= View.ViewSizeAndInvSize.xy))
+	{
+		return;
+	}
+	uint LinearIndex = DispatchThreadId.y * GroupCount.x * THREADGROUP_SIZE_2D + DispatchThreadId.x;
+	uint2 PixelCoord = DispatchThreadId.xy + View.ViewRectMin.xy;
+
+	// Extract depth
+	float DeviceZ = SceneDepthTexture.Load(int3(PixelCoord, 0)).r;
+#if HAS_INVERTED_Z_BUFFER
+	DeviceZ = max(0.000000000001, DeviceZ);
+#endif // HAS_INVERTED_Z_BUFFER
+
+	// Clip trace distance
+	float SceneDepth = min(ConvertFromDeviceZ(DeviceZ), MaxTraceDistance);
+	DeviceZ = ConvertToDeviceZ(SceneDepth);
+
+	float3 RayOrigin = View.TranslatedWorldCameraOrigin;
+	float3 RayEnd = SvPositionToTranslatedWorld(float4(PixelCoord + 0.5, DeviceZ, 1));
+	float3 RayDirection = RayEnd - RayOrigin;
+	float RayTMin = 0.0;
+	float RayTMax = length(RayDirection);
+	RayDirection /= RayTMax;
+
+	RandomSequence RandSequence;
+	RandomSequence_Initialize(RandSequence, LinearIndex, View.StateFrameIndex);
+	float Jitter = RandomSequence_GenerateSample1D(RandSequence);
+
+	float3 Emission = 0;
+	float3 Transmittance = 1;
+
+	float3 BoundsMin = GetTranslatedWorldPos(OrthoGridUniformBuffer.TopLevelGridWorldBoundsMin);
+	float3 BoundsMax = GetTranslatedWorldPos(OrthoGridUniformBuffer.TopLevelGridWorldBoundsMax);
+	float2 ClipHitT = IntersectAABB(RayOrigin, RayDirection, RayTMin, RayTMax, BoundsMin, BoundsMax);
+
+#if DEBUG_MODE_SAMPLING_BASED
+	if (ClipHitT.x < ClipHitT.y)
+	{
+		float StepSize = CalcStepSize(RayDirection);
+		Jitter *= StepSize;
+#if DEBUG_MODE == DEBUG_MODE_SAMPLE_BASED_MARCHING
+		FTrackingResult Result = SampleBasedMarch(RayOrigin, RayDirection, ClipHitT.x + Jitter, ClipHitT.y, StepSize);
+#elif DEBUG_MODE == DEBUG_MODE_SAMPLE_BASED_MAJORANT
+		FTrackingResult Result = SampleBasedMajorantMarch(RayOrigin, RayDirection, ClipHitT.x + Jitter, ClipHitT.y, StepSize);
+#elif DEBUG_MODE == DEBUG_MODE_SAMPLE_BASED_MAX_MAJORANT
+		FTrackingResult Result = SampleBasedMaxMajorantMarch(RayOrigin, RayDirection, ClipHitT.x + Jitter, ClipHitT.y, StepSize);
+#else // DEBUG_MODE == DEBUG_MODE_SAMPLE_BASED_MAJORANT_DDA
+		int DistanceSamplingMethod = DEBUG_MODE_SAMPLE_BASED_MARCHING;
+		int CollisionSamplingMethod = 0;
+		FTrackingResult Result = DDABasedMajorantMarch(RayOrigin, RayDirection, ClipHitT.x, ClipHitT.y, DistanceSamplingMethod, CollisionSamplingMethod, RandSequence);
+#endif
+		Transmittance = Result.Throughput;
+	}
+
+	//Radiance = Result.Emission;
+	Radiance = 1.0 - Transmittance;
+
+#elif DEBUG_MODE_TRACKING_BASED
+	if (ClipHitT.x < ClipHitT.y)
+	{
+		FTrackingResult Result = (FTrackingResult) 0;
+
+		FMajorantData MajorantData = GetMajorantData(OrthoGridUniformBuffer.MajorantGridBuffer[0]);
+		float SigmaBar = MajorantData.Majorant;
+
+		int SPPCount = 1;
+		for (int Index = 0; Index < SPPCount; ++Index)
+		{
+			RandomSequence_Initialize(RandSequence, LinearIndex, View.StateFrameIndex * SPPCount + Index);
+#if DEBUG_MODE == DEBUG_MODE_TRACKING_BASED_DELTA_TRACKING
+			FTrackingResult Sample = SpectralWeightedDeltaTracking(RayOrigin, RayDirection, ClipHitT.x, ClipHitT.y, SigmaBar, COLLISION_MODE_AVG_COMPONENT, RandSequence);
+#elif DEBUG_MODE == DEBUG_MODE_TRACKING_BASED_RATIO_TRACKING
+			FTrackingResult Sample = RatioTracking(RayOrigin, RayDirection, ClipHitT.x, ClipHitT.y, SigmaBar, RandSequence);
+#elif DEBUG_MODE == DEBUG_MODE_TRACKING_BASED_SPECTRAL_RATIO_TRACKING
+			FTrackingResult Sample = SpectralRatioTracking(RayOrigin, RayDirection, ClipHitT.x, ClipHitT.y, SigmaBar, RandSequence);
+#else // DEBUG_MODE == DEBUG_MODE_TRACKING_BASED_MAJORANT_DDA
+			//int DistanceSamplingMethod = DEBUG_MODE_TRACKING_BASED_SPECTRAL_RATIO_TRACKING;
+			int DistanceSamplingMethod = DEBUG_MODE_TRACKING_BASED_DELTA_TRACKING;
+			int CollisionSamplingMethod = COLLISION_MODE_AVG_COMPONENT;
+			FTrackingResult Sample = DDABasedMajorantMarch(RayOrigin, RayDirection, ClipHitT.x, ClipHitT.y, DistanceSamplingMethod, CollisionSamplingMethod, RandSequence);
+#endif
+#if 0
+			// Faux direct-lighting ray
+			if (Sample.bIsScattering)
+			{
+				float3 ShadowRayOrigin = RayOrigin + RayDirection * Sample.Distance;
+				float3 LightDirection = normalize(float3(0, 0, 1));
+				float2 ClipShadowHitT = IntersectAABB(ShadowRayOrigin, LightDirection, 0, 10000, BoundsMin, BoundsMax);
+				if (ClipShadowHitT.y > ClipShadowHitT.x)
+				{
+					float3 Le = 10000;
+					float Phase = 1.0 / (4 * PI);
+					FTrackingResult Lighting = DDABasedMajorantMarch(ShadowRayOrigin, LightDirection, ClipShadowHitT.x, ClipShadowHitT.y, DEBUG_MODE_TRACKING_BASED_DELTA_TRACKING, COLLISION_MODE_RATIO_TRACKING, RandSequence);
+					Sample.Radiance = Le * Phase * Lighting.Throughput * Sample.Throughput;// / (Sample.Pdf * Lighting.Pdf);
+				}
+				Sample.Throughput = 0;
+			}
+#endif
+			Result = AddTrackingResult(Result, Sample);
+		}
+
+		Result = DivideTrackingResult(Result, SPPCount);
+		float3 Le = 1;
+		Radiance = (Result.Radiance + (1.0 - Result.Throughput) * Le);
+		Transmittance = Result.Throughput;
+	}
+
+#elif DEBUG_MODE == DEBUG_MODE_DDA_MARCHING
+	float FrustumRayTMin = RayTMin;
+	float FrustumRayTMax = RayTMax;
+	if (FrustumGridUniformBuffer.bUseFrustumGrid)
+	{
+		float ViewNearZ = min(ConvertFromDeviceZ(DeviceZ), FrustumGridUniformBuffer.NearPlaneDepth);
+		float DeviceNearZ = ConvertToDeviceZ(ViewNearZ);
+		float ViewFarZ = min(ConvertFromDeviceZ(DeviceZ), FrustumGridUniformBuffer.FarPlaneDepth);
+		float DeviceFarZ = ConvertToDeviceZ(ViewFarZ);
+
+		// TODO: Shared code is imprecise, in comparison with commented out frustum-trimmed ray. Investigate...
+#if 0
+		float3 RayBegin = SvPositionToTranslatedWorld(float4(PixelCoord + 0.5, DeviceNearZ, 1));
+		float3 RayEnd = SvPositionToTranslatedWorld(float4(PixelCoord + 0.5, DeviceFarZ, 1));
+		float3 RayDirection = RayEnd - RayBegin;
+		FrustumRayTMin = 0.0;
+		FrustumRayTMax = length(RayDirection);
+		RayDirection /= FrustumRayTMax;
+
+		Transmittance = TraceFrustumVoxelGrid(RayBegin, RayDirection, FrustumRayTMin, FrustumRayTMax);
+		float RayTOffset = length(RayBegin - RayOrigin);
+		FrustumRayTMin += RayTOffset;
+		FrustumRayTMax += RayTOffset;
+#else
+		float3 WorldRayOrigin = LWCHackToFloat(SvPositionToWorld(float4(PixelCoord + 0.5, DeviceNearZ, 1)));
+		float3 WorldRayEnd = LWCHackToFloat(SvPositionToWorld(float4(PixelCoord + 0.5, DeviceFarZ, 1)));
+		float3 WorldRayDirection = WorldRayEnd - WorldRayOrigin;
+		float WorldRayLength = length(WorldRayDirection);
+		WorldRayDirection /= WorldRayLength;
+
+		// Convert to view-space
+		float3 ViewRayOrigin = mul(float4(WorldRayOrigin, 1), FrustumGridUniformBuffer.WorldToView).xyz;
+		float3 ViewRayEnd = mul(float4(WorldRayEnd, 1), FrustumGridUniformBuffer.WorldToView).xyz;
+		float3 ViewRayDirection = ViewRayEnd - ViewRayOrigin;
+		float ViewRayLength = length(ViewRayDirection);
+		ViewRayDirection /= ViewRayLength;
+
+		// Convert to voxel-space
+		int3 VoxelDimensions = FrustumGridUniformBuffer.VoxelDimensions;
+		float NearPlaneDepth = FrustumGridUniformBuffer.NearPlaneDepth;
+		float FarPlaneDepth = FrustumGridUniformBuffer.FarPlaneDepth;
+		float3 VoxelRayOrigin = ViewToVoxel(ViewRayOrigin, VoxelDimensions, NearPlaneDepth, FarPlaneDepth);
+		float3 VoxelRayEnd = ViewToVoxel(ViewRayOrigin + ViewRayDirection * ViewRayLength, VoxelDimensions, NearPlaneDepth, FarPlaneDepth);
+
+		float3 VoxelRayDirection = VoxelRayEnd - VoxelRayOrigin;
+		float VoxelRayLength = length(VoxelRayDirection);
+		VoxelRayDirection /= VoxelRayLength;
+
+		float VoxelRayTMin = 0.0;
+		float VoxelRayTMax = VoxelRayLength;
+
+		// Traversal
+		{
+			MarchTopLevelFrustumVoxelGrid(
+				VoxelRayOrigin,
+				VoxelRayDirection,
+				VoxelRayTMin,
+				VoxelRayTMax,
+				WorldRayLength,
+				FrustumGridUniformBuffer.TopLevelFroxelGridResolution,
+				FrustumGridUniformBuffer.TopLevelFroxelGridBuffer,
+				FrustumGridUniformBuffer.ExtinctionFroxelGridBuffer,
+				Transmittance
+			);
+		}
+
+		float3 RayBegin = SvPositionToTranslatedWorld(float4(PixelCoord + 0.5, DeviceNearZ, 1));
+		float RayTOffset = length(RayBegin - RayOrigin);
+		FrustumRayTMin += RayTOffset;
+		FrustumRayTMax += RayTOffset;
+#endif
+	}
+
+	if (OrthoGridUniformBuffer.bUseOrthoGrid)
+	{
+		if (FrustumGridUniformBuffer.bUseFrustumGrid && (FrustumRayTMin < FrustumRayTMax))
+		{
+			Transmittance *= TraceOrthoVoxelGrid(RayOrigin, RayDirection, RayTMin, FrustumRayTMin);
+			Transmittance *= TraceOrthoVoxelGrid(RayOrigin, RayDirection, FrustumRayTMax, RayTMax);
+		}
+		else
+		{
+			Transmittance = TraceOrthoVoxelGrid(RayOrigin, RayDirection, RayTMin, RayTMax);
+		}
+	}
+
+#endif // DEBUG_MODE_DDA_MAJORANT
+
+	// Output..
+	float3 Opacity = 1.0 - Transmittance;
+	RWLightingTexture[PixelCoord] += float4(Radiance, Luminance(Opacity));
+}
\ No newline at end of file
diff --git a/Engine/Shaders/Private/HeterogeneousVolumes/HeterogeneousVolumesVoxelGridUtils.ush b/Engine/Shaders/Private/HeterogeneousVolumes/HeterogeneousVolumesVoxelGridUtils.ush
--- a/Engine/Shaders/Private/HeterogeneousVolumes/HeterogeneousVolumesVoxelGridUtils.ush
+++ b/Engine/Shaders/Private/HeterogeneousVolumes/HeterogeneousVolumesVoxelGridUtils.ush
@@ -18,10 +18,12 @@ FMajorantData CreateMajorantData()
 	return MajorantData;
 }
 
-void UpdateMajorantData(inout FMajorantData MajorantData, float Value)
+void UpdateMajorantData(inout FMajorantData MajorantData, float3 Value)
 {
-	MajorantData.Minorant = min(MajorantData.Minorant, Value);
-	MajorantData.Majorant = max(MajorantData.Majorant, Value);
+	float MinValue = min(Value.x, min(Value.y, Value.z));
+	float MaxValue = max(Value.x, max(Value.y, Value.z));
+	MajorantData.Minorant = min(MajorantData.Minorant, MinValue);
+	MajorantData.Majorant = max(MajorantData.Majorant, MaxValue);
 }
 
 void UpdateMajorantData(inout FMajorantData MajorantData, FMajorantData Value)
@@ -33,14 +35,15 @@ void UpdateMajorantData(inout FMajorantData MajorantData, FMajorantData Value)
 FMajorantData GetMajorantData(FScalarGridData MajorantGridData)
 {
 	FMajorantData Data;
-	Data.Majorant = f16tof32(MajorantGridData.PackedData[0] & 0xFFFF);
-	Data.Minorant = f16tof32(MajorantGridData.PackedData[0] >> 16);
+	Data.Minorant = asfloat(MajorantGridData.PackedData[0]);
+	Data.Majorant = asfloat(MajorantGridData.PackedData[1]);
 	return Data;
 }
 
 void SetMajorantData(inout FScalarGridData MajorantGridData, FMajorantData MajorantData)
 {
-	MajorantGridData.PackedData[0] = (f32tof16(MajorantData.Minorant) << 16) | f32tof16(MajorantData.Majorant);
+	MajorantGridData.PackedData[0] = asuint(MajorantData.Minorant);
+	MajorantGridData.PackedData[1] = asuint(MajorantData.Majorant);
 }
 
 float3 GetExtinction(FScalarGridData ExtinctionGridData)
@@ -156,7 +159,7 @@ bool IsBottomLevelEmpty(FTopLevelGridData TopLevelGridData)
 int3 GetBottomLevelVoxelResolution(FTopLevelGridData TopLevelGridData)
 {
 	uint ResolutionExponent = TopLevelGridData.PackedData[0] & 0x7;
-	int3 Resolution = ResolutionExponent > 0 ? (1u << (ResolutionExponent - 1)) : 0;
+	int3 Resolution = ResolutionExponent;
 
 	return Resolution;
 }
@@ -164,7 +167,7 @@ int3 GetBottomLevelVoxelResolution(FTopLevelGridData TopLevelGridData)
 void SetBottomLevelVoxelResolution(inout FTopLevelGridData TopLevelGridData, int3 BottomLevelVoxelResolution)
 {
 	uint Index = GetBottomLevelIndex(TopLevelGridData);
-	uint ResolutionExponent = int(BottomLevelVoxelResolution.x > 0 ? uint(log2(BottomLevelVoxelResolution.x)) + 1 : 0) & 0x7;
+	uint ResolutionExponent = BottomLevelVoxelResolution.x & 0x7;
 	TopLevelGridData.PackedData[0] = (Index << 3) | ResolutionExponent;
 }
 
diff --git a/Engine/Source/Runtime/Renderer/Private/HeterogeneousVolumes/HeterogeneousVolumes.cpp b/Engine/Source/Runtime/Renderer/Private/HeterogeneousVolumes/HeterogeneousVolumes.cpp
--- a/Engine/Source/Runtime/Renderer/Private/HeterogeneousVolumes/HeterogeneousVolumes.cpp
+++ b/Engine/Source/Runtime/Renderer/Private/HeterogeneousVolumes/HeterogeneousVolumes.cpp
@@ -414,7 +414,7 @@ void FDeferredShadingSceneRenderer::RenderHeterogeneousVolumes(
 
 		if (ShouldRenderHeterogeneousVolumesForView(View))
 		{
-			if (HeterogeneousVolumes::GetDebugMode() == 1)
+			if (HeterogeneousVolumes::GetDebugMode() != 0)
 			{
 				RenderTransmittanceWithVoxelGrid(
 					GraphBuilder,
diff --git a/Engine/Source/Runtime/Renderer/Private/HeterogeneousVolumes/HeterogeneousVolumesVoxelGridPipeline.cpp b/Engine/Source/Runtime/Renderer/Private/HeterogeneousVolumes/HeterogeneousVolumesVoxelGridPipeline.cpp
--- a/Engine/Source/Runtime/Renderer/Private/HeterogeneousVolumes/HeterogeneousVolumesVoxelGridPipeline.cpp
+++ b/Engine/Source/Runtime/Renderer/Private/HeterogeneousVolumes/HeterogeneousVolumesVoxelGridPipeline.cpp
@@ -166,6 +166,13 @@ static TAutoConsoleVariable<int32> CVarHeterogeneousVolumesEnableTopLevelBitmask
 
 static TAutoConsoleVariable<int32> CVarHeterogeneousVolumesEnableMajorantGrid(
 	TEXT("r.HeterogeneousVolumes.Tessellation.MajorantGrid"),
+	1,
+	TEXT("Enables building majorant grids to accelerate volume tracking (Default = 0)\n"),
+	ECVF_RenderThreadSafe
+);
+
+static TAutoConsoleVariable<int32> CVarHeterogeneousVolumesEnableMajorantGridMax(
+	TEXT("r.HeterogeneousVolumes.Tessellation.MajorantGrid.Max"),
 	0,
 	TEXT("Enables building majorant grids to accelerate volume tracking (Default = 0)\n"),
 	ECVF_RenderThreadSafe
@@ -347,9 +354,9 @@ class FMarkTopLevelGridVoxelsForFrustumGrid : public FGlobalShader
 
 IMPLEMENT_MATERIAL_SHADER_TYPE(, FMarkTopLevelGridVoxelsForFrustumGrid, TEXT("/Engine/Private/HeterogeneousVolumes/HeterogeneousVolumesFrustumVoxelGrid.usf"), TEXT("MarkTopLevelGridVoxelsForFrustumGridCS"), SF_Compute);
 
-class FRasterizeBottomLevelFroxelGridCS : public FMeshMaterialShader
+class FRasterizeBottomLevelFrustumGridCS : public FMeshMaterialShader
 {
-	DECLARE_SHADER_TYPE(FRasterizeBottomLevelFroxelGridCS, MeshMaterial);
+	DECLARE_SHADER_TYPE(FRasterizeBottomLevelFrustumGridCS, MeshMaterial);
 
 	using FPermutationDomain = TShaderPermutationDomain<>;
 
@@ -396,9 +403,9 @@ class FRasterizeBottomLevelFroxelGridCS : public FMeshMaterialShader
 		SHADER_PARAMETER_RDG_BUFFER_UAV(RWStructuredBuffer<FVectorGridData>, RWScatteringGridBuffer)
 	END_SHADER_PARAMETER_STRUCT()
 
-	FRasterizeBottomLevelFroxelGridCS() = default;
+	FRasterizeBottomLevelFrustumGridCS() = default;
 
-	FRasterizeBottomLevelFroxelGridCS(
+	FRasterizeBottomLevelFrustumGridCS(
 		const FMeshMaterialShaderType::CompiledShaderInitializerType& Initializer
 	)
 		: FMeshMaterialShader(Initializer)
@@ -458,7 +465,7 @@ class FRasterizeBottomLevelFroxelGridCS : public FMeshMaterialShader
 	static int32 GetThreadGroupSize3D() { return 4; }
 };
 
-IMPLEMENT_MATERIAL_SHADER_TYPE(, FRasterizeBottomLevelFroxelGridCS, TEXT("/Engine/Private/HeterogeneousVolumes/HeterogeneousVolumesMaterialBakingPipeline.usf"), TEXT("RasterizeBottomLevelFroxelGridCS"), SF_Compute);
+IMPLEMENT_MATERIAL_SHADER_TYPE(, FRasterizeBottomLevelFrustumGridCS, TEXT("/Engine/Private/HeterogeneousVolumes/HeterogeneousVolumesMaterialBakingPipeline.usf"), TEXT("RasterizeBottomLevelFrustumGridCS"), SF_Compute);
 
 class FTopLevelGridCalculateVoxelSize : public FGlobalShader
 {
@@ -640,9 +647,9 @@ class FSetRasterizeBottomLevelGridIndirectArgs : public FGlobalShader
 
 IMPLEMENT_MATERIAL_SHADER_TYPE(, FSetRasterizeBottomLevelGridIndirectArgs, TEXT("/Engine/Private/HeterogeneousVolumes/HeterogeneousVolumesVoxelGridPipeline.usf"), TEXT("SetRasterizeBottomLevelGridIndirectArgs"), SF_Compute);
 
-class FRasterizeBottomLevelGrid : public FMeshMaterialShader
+class FRasterizeBottomLevelOrthoGridCS : public FMeshMaterialShader
 {
-	DECLARE_SHADER_TYPE(FRasterizeBottomLevelGrid, MeshMaterial);
+	DECLARE_SHADER_TYPE(FRasterizeBottomLevelOrthoGridCS, MeshMaterial);
 
 	class FEnableIndirectionGrid : SHADER_PERMUTATION_BOOL("DIM_ENABLE_INDIRECTION_GRID");
 	//class FEnableVoxelHashing : SHADER_PERMUTATION_BOOL("DIM_ENABLE_VOXEL_HASHING");
@@ -670,6 +677,9 @@ class FRasterizeBottomLevelGrid : public FMeshMaterialShader
 
 		SHADER_PARAMETER(int, BottomLevelGridBufferSize)
 
+		// Sampling data
+		SHADER_PARAMETER_STRUCT_REF(FBlueNoise, BlueNoise)
+
 		// Volume sample mode
 		SHADER_PARAMETER(int, bSampleAtVertices)
 
@@ -704,9 +714,9 @@ class FRasterizeBottomLevelGrid : public FMeshMaterialShader
 		SHADER_PARAMETER(float, HomogeneousThreshold)
 	END_SHADER_PARAMETER_STRUCT()
 	
-	FRasterizeBottomLevelGrid() = default;
+	FRasterizeBottomLevelOrthoGridCS() = default;
 
-	FRasterizeBottomLevelGrid(
+	FRasterizeBottomLevelOrthoGridCS(
 		const FMeshMaterialShaderType::CompiledShaderInitializerType & Initializer
 	)
 		: FMeshMaterialShader(Initializer)
@@ -766,7 +776,7 @@ class FRasterizeBottomLevelGrid : public FMeshMaterialShader
 	static int32 GetThreadGroupSize3D() { return 4; }
 };
 
-IMPLEMENT_MATERIAL_SHADER_TYPE(, FRasterizeBottomLevelGrid, TEXT("/Engine/Private/HeterogeneousVolumes/HeterogeneousVolumesMaterialBakingPipeline.usf"), TEXT("RasterizeBottomLevelGridCS"), SF_Compute);
+IMPLEMENT_MATERIAL_SHADER_TYPE(, FRasterizeBottomLevelOrthoGridCS, TEXT("/Engine/Private/HeterogeneousVolumes/HeterogeneousVolumesMaterialBakingPipeline.usf"), TEXT("RasterizeBottomLevelOrthoGridCS"), SF_Compute);
 
 class FTopLevelCreateBitmaskCS : public FGlobalShader
 {
@@ -899,6 +909,78 @@ class FBuildMajorantVoxelGridCS : public FGlobalShader
 
 IMPLEMENT_MATERIAL_SHADER_TYPE(, FBuildMajorantVoxelGridCS, TEXT("/Engine/Private/HeterogeneousVolumes/HeterogeneousVolumesVoxelGridPipeline.usf"), TEXT("BuildMajorantVoxelGridCS"), SF_Compute);
 
+class FDownsampleMajorantVoxelGridCS : public FGlobalShader
+{
+	DECLARE_GLOBAL_SHADER(FDownsampleMajorantVoxelGridCS);
+	SHADER_USE_PARAMETER_STRUCT(FDownsampleMajorantVoxelGridCS, FGlobalShader);
+
+	using FPermutationDomain = TShaderPermutationDomain<>;
+
+	BEGIN_SHADER_PARAMETER_STRUCT(FParameters, )
+		SHADER_PARAMETER(FIntVector, InputDimensions)
+		SHADER_PARAMETER_RDG_BUFFER_SRV(RWStructuredBuffer<FScalarGridData>, InputBuffer)
+		SHADER_PARAMETER_RDG_BUFFER_UAV(RWStructuredBuffer<FScalarGridData>, RWOutputBuffer)
+	END_SHADER_PARAMETER_STRUCT()
+
+	static bool ShouldCompilePermutation(
+		const FGlobalShaderPermutationParameters& Parameters
+	)
+	{
+		return DoesPlatformSupportHeterogeneousVolumes(Parameters.Platform);
+	}
+
+	static void ModifyCompilationEnvironment(
+		const FGlobalShaderPermutationParameters& Parameters,
+		FShaderCompilerEnvironment& OutEnvironment
+	)
+	{
+		FMaterialShader::ModifyCompilationEnvironment(Parameters, OutEnvironment);
+		OutEnvironment.SetDefine(TEXT("THREADGROUP_SIZE_1D"), GetThreadGroupSize1D());
+		OutEnvironment.SetDefine(TEXT("THREADGROUP_SIZE_3D"), GetThreadGroupSize3D());
+	}
+
+	static int32 GetThreadGroupSize1D() { return GetThreadGroupSize3D() * GetThreadGroupSize3D() * GetThreadGroupSize3D(); }
+	static int32 GetThreadGroupSize3D() { return 4; }
+};
+
+IMPLEMENT_MATERIAL_SHADER_TYPE(, FDownsampleMajorantVoxelGridCS, TEXT("/Engine/Private/HeterogeneousVolumes/HeterogeneousVolumesVoxelGridPipeline.usf"), TEXT("DownsampleMajorantVoxelGridCS"), SF_Compute);
+
+class FCopyMaxIntoMajorantVoxelGridCS : public FGlobalShader
+{
+	DECLARE_GLOBAL_SHADER(FCopyMaxIntoMajorantVoxelGridCS);
+	SHADER_USE_PARAMETER_STRUCT(FCopyMaxIntoMajorantVoxelGridCS, FGlobalShader);
+
+	using FPermutationDomain = TShaderPermutationDomain<>;
+
+	BEGIN_SHADER_PARAMETER_STRUCT(FParameters, )
+		SHADER_PARAMETER_RDG_BUFFER_SRV(RWStructuredBuffer<FScalarGridData>, InputBuffer)
+		SHADER_PARAMETER_RDG_BUFFER_UAV(RWStructuredBuffer<FScalarGridData>, RWMajorantVoxelGridBuffer)
+	END_SHADER_PARAMETER_STRUCT()
+
+	static bool ShouldCompilePermutation(
+		const FGlobalShaderPermutationParameters& Parameters
+	)
+	{
+		return DoesPlatformSupportHeterogeneousVolumes(Parameters.Platform);
+	}
+
+	static void ModifyCompilationEnvironment(
+		const FGlobalShaderPermutationParameters& Parameters,
+		FShaderCompilerEnvironment& OutEnvironment
+	)
+	{
+		FMaterialShader::ModifyCompilationEnvironment(Parameters, OutEnvironment);
+		OutEnvironment.SetDefine(TEXT("THREADGROUP_SIZE_1D"), GetThreadGroupSize1D());
+		OutEnvironment.SetDefine(TEXT("THREADGROUP_SIZE_3D"), GetThreadGroupSize3D());
+	}
+
+	static int32 GetThreadGroupSize1D() { return GetThreadGroupSize3D() * GetThreadGroupSize3D() * GetThreadGroupSize3D(); }
+	static int32 GetThreadGroupSize3D() { return 4; }
+};
+
+IMPLEMENT_MATERIAL_SHADER_TYPE(, FCopyMaxIntoMajorantVoxelGridCS, TEXT("/Engine/Private/HeterogeneousVolumes/HeterogeneousVolumesVoxelGridPipeline.usf"), TEXT("CopyMaximumIntoMajorantVoxelGridCS"), SF_Compute);
+
+
 void BuildMajorantVoxelGrid(
 	FRDGBuilder& GraphBuilder,
 	const FScene* Scene,
@@ -921,38 +1003,115 @@ void BuildMajorantVoxelGrid(
 			TEXT("HeterogeneousVolumes.MajorantVoxelGridBuffer")
 		);
 
-		FBuildMajorantVoxelGridCS::FParameters* PassParameters = GraphBuilder.AllocParameters<FBuildMajorantVoxelGridCS::FParameters>();
+		const FGlobalShaderMap* GlobalShaderMap = GetGlobalShaderMap(Scene->GetFeatureLevel());
+
+		// Build Majorant Grid
 		{
-			PassParameters->TopLevelGridResolution = TopLevelGridResolution;
-			PassParameters->TopLevelGridBuffer = GraphBuilder.CreateSRV(TopLevelGridBuffer);
-			PassParameters->IndirectionGridBuffer = GraphBuilder.CreateSRV(IndirectionGridBuffer);
-			PassParameters->ExtinctionGridBuffer = GraphBuilder.CreateSRV(ExtinctionGridBuffer);
-			PassParameters->RWMajorantVoxelGridBuffer = GraphBuilder.CreateUAV(MajorantVoxelGridBuffer);
+			FBuildMajorantVoxelGridCS::FParameters* PassParameters = GraphBuilder.AllocParameters<FBuildMajorantVoxelGridCS::FParameters>();
+			{
+				PassParameters->TopLevelGridResolution = TopLevelGridResolution;
+				PassParameters->TopLevelGridBuffer = GraphBuilder.CreateSRV(TopLevelGridBuffer);
+				PassParameters->IndirectionGridBuffer = GraphBuilder.CreateSRV(IndirectionGridBuffer);
+				PassParameters->ExtinctionGridBuffer = GraphBuilder.CreateSRV(ExtinctionGridBuffer);
+				PassParameters->RWMajorantVoxelGridBuffer = GraphBuilder.CreateUAV(MajorantVoxelGridBuffer);
+			}
+
+			FBuildMajorantVoxelGridCS::FPermutationDomain PermutationVector;
+			PermutationVector.Set<FBuildMajorantVoxelGridCS::FEnableIndirectionGrid>(HeterogeneousVolumes::EnableIndirectionGrid());
+			TShaderRef<FBuildMajorantVoxelGridCS> ComputeShader = GlobalShaderMap->GetShader<FBuildMajorantVoxelGridCS>(PermutationVector);
+
+			FComputeShaderUtils::AddPass(
+				GraphBuilder,
+				RDG_EVENT_NAME("RenderTransmittanceTopLevelGridCS"),
+				ERDGPassFlags::Compute | ERDGPassFlags::NeverCull,
+				ComputeShader,
+				PassParameters,
+				TopLevelGridResolution
+			);
 		}
 
-		const FGlobalShaderMap* GlobalShaderMap = GetGlobalShaderMap(Scene->GetFeatureLevel());
+		bool bUseMax = CVarHeterogeneousVolumesEnableMajorantGridMax.GetValueOnRenderThread() != 0;
+		if (bUseMax)
+		{
 
-		FBuildMajorantVoxelGridCS::FPermutationDomain PermutationVector;
-		PermutationVector.Set<FBuildMajorantVoxelGridCS::FEnableIndirectionGrid>(HeterogeneousVolumes::EnableIndirectionGrid());
-		TShaderRef<FBuildMajorantVoxelGridCS> ComputeShader = GlobalShaderMap->GetShader<FBuildMajorantVoxelGridCS>(PermutationVector);
+			// Downsample majorant grid
+			FIntVector InputDimensions = TopLevelGridResolution;
+			FRDGBufferRef InputBuffer = MajorantVoxelGridBuffer;
+			FRDGBufferRef OutputBuffer = nullptr;
+			while (InputDimensions.X > 1 && InputDimensions.Y && InputDimensions.Z)
+			{
+				FIntVector OutputDimensions = FIntVector(
+					FMath::DivideAndRoundUp(InputDimensions.X, FDownsampleMajorantVoxelGridCS::GetThreadGroupSize3D()),
+					FMath::DivideAndRoundUp(InputDimensions.Y, FDownsampleMajorantVoxelGridCS::GetThreadGroupSize3D()),
+					FMath::DivideAndRoundUp(InputDimensions.Z, FDownsampleMajorantVoxelGridCS::GetThreadGroupSize3D())
+				);
+				uint32 OutputBufferSize = OutputDimensions.X * OutputDimensions.Y * OutputDimensions.Z;
+				OutputBuffer = GraphBuilder.CreateBuffer(
+					FRDGBufferDesc::CreateStructuredDesc(sizeof(FScalarGridData), OutputBufferSize),
+					TEXT("HeterogeneousVolumes.DownsampledBuffer")
+				);
+
+				FDownsampleMajorantVoxelGridCS::FParameters* PassParameters = GraphBuilder.AllocParameters<FDownsampleMajorantVoxelGridCS::FParameters>();
+				{
+					PassParameters->InputDimensions = InputDimensions;
+					PassParameters->InputBuffer = GraphBuilder.CreateSRV(InputBuffer);
+					PassParameters->RWOutputBuffer = GraphBuilder.CreateUAV(OutputBuffer);
+				}
 
-		FComputeShaderUtils::AddPass(
-			GraphBuilder,
-			RDG_EVENT_NAME("RenderTransmittanceTopLevelGridCS"),
-			ERDGPassFlags::Compute | ERDGPassFlags::NeverCull,
-			ComputeShader,
-			PassParameters,
-			TopLevelGridResolution
-		);
+				TShaderRef<FDownsampleMajorantVoxelGridCS> ComputeShader = GlobalShaderMap->GetShader<FDownsampleMajorantVoxelGridCS>();
+				FComputeShaderUtils::AddPass(
+					GraphBuilder,
+					RDG_EVENT_NAME("FDownsampleMajorantVoxelGridCS"),
+					ERDGPassFlags::Compute | ERDGPassFlags::NeverCull,
+					ComputeShader,
+					PassParameters,
+					TopLevelGridResolution
+				);
+
+				InputBuffer = OutputBuffer;
+				InputDimensions = OutputDimensions;
+			}
+
+			// Copy maximum value into the first index of the majorant grid
+			if (OutputBuffer)
+			{
+				FCopyMaxIntoMajorantVoxelGridCS::FParameters* PassParameters = GraphBuilder.AllocParameters<FCopyMaxIntoMajorantVoxelGridCS::FParameters>();
+				{
+					PassParameters->InputBuffer = GraphBuilder.CreateSRV(OutputBuffer);
+					PassParameters->RWMajorantVoxelGridBuffer = GraphBuilder.CreateUAV(MajorantVoxelGridBuffer);
+				}
+
+				TShaderRef<FCopyMaxIntoMajorantVoxelGridCS> ComputeShader = GlobalShaderMap->GetShader<FCopyMaxIntoMajorantVoxelGridCS>();
+				FComputeShaderUtils::AddPass(
+					GraphBuilder,
+					RDG_EVENT_NAME("FCopyMaxIntoMajorantVoxelGridCS"),
+					ERDGPassFlags::Compute | ERDGPassFlags::NeverCull,
+					ComputeShader,
+					PassParameters,
+					FIntVector(1)
+				);
+			}
+		}
 	}
 }
 
+struct FRenderDebugData
+{
+	FVector4f RayOrigin;
+	FVector4f RayDirection;
+	float TMin;
+	float TMax;
+	float Distance;
+	FVector4f EstimateAndPdf;
+};
+
 class FRenderTransmittanceWithVoxelGridCS : public FGlobalShader
 {
 	DECLARE_GLOBAL_SHADER(FRenderTransmittanceWithVoxelGridCS);
 	SHADER_USE_PARAMETER_STRUCT(FRenderTransmittanceWithVoxelGridCS, FGlobalShader);
 
-	using FPermutationDomain = TShaderPermutationDomain<>;
+	class FDebugMode : SHADER_PERMUTATION_INT("DEBUG_MODE", 9);
+	using FPermutationDomain = TShaderPermutationDomain<FDebugMode>;
 
 	BEGIN_SHADER_PARAMETER_STRUCT(FParameters, )
 		// Scene data
@@ -974,6 +1133,8 @@ class FRenderTransmittanceWithVoxelGridCS : public FGlobalShader
 
 		// Output
 		SHADER_PARAMETER_RDG_TEXTURE_UAV(RWTexture2D<float4>, RWLightingTexture)
+		// Debug
+		SHADER_PARAMETER_RDG_BUFFER_UAV(RWStructuredBuffer<FRenderDebugData>, RWDebugBuffer)
 	END_SHADER_PARAMETER_STRUCT()
 
 	static bool ShouldCompilePermutation(
@@ -994,8 +1155,6 @@ class FRenderTransmittanceWithVoxelGridCS : public FGlobalShader
 		OutEnvironment.SetDefine(TEXT("THREADGROUP_SIZE_3D"), GetThreadGroupSize3D());
 
 		OutEnvironment.CompilerFlags.Add(CFLAG_AllowTypedUAVLoads);
-
-		OutEnvironment.SetDefine(TEXT("ENABLE_ADAPTIVE_VOXEL_GRID_UNIFORM_BUFFER"), 1);
 	}
 
 	static int32 GetThreadGroupSize1D() { return GetThreadGroupSize3D() * GetThreadGroupSize3D() * GetThreadGroupSize3D(); }
@@ -1003,7 +1162,7 @@ class FRenderTransmittanceWithVoxelGridCS : public FGlobalShader
 	static int32 GetThreadGroupSize3D() { return 4; }
 };
 
-IMPLEMENT_MATERIAL_SHADER_TYPE(, FRenderTransmittanceWithVoxelGridCS, TEXT("/Engine/Private/HeterogeneousVolumes/HeterogeneousVolumesVoxelGridPipeline.usf"), TEXT("RenderTransmittanceWithVoxelGridCS"), SF_Compute);
+IMPLEMENT_MATERIAL_SHADER_TYPE(, FRenderTransmittanceWithVoxelGridCS, TEXT("/Engine/Private/HeterogeneousVolumes/HeterogeneousVolumesVoxelGridRendering.usf"), TEXT("RenderTransmittanceWithVoxelGridCS"), SF_Compute);
 
 void CalcViewBoundsAndMinimumVoxelSize(
 	FRDGBuilder& GraphBuilder,
@@ -1502,9 +1661,9 @@ void CalculateTopLevelGridResolutionForFrustumGrid(
 	int32 Depth = FMath::Min(FMath::CeilToInt32(HeterogeneousVolumes::GetDepthSliceCountForFrustumGrid() / ShadingRate), MaxDepth);
 
 	TopLevelGridResolution = FIntVector(Width, Height, Depth);
-	TopLevelGridResolution.X = FMath::DivideAndRoundUp(TopLevelGridResolution.X, BottomLevelGridResolution);
-	TopLevelGridResolution.Y = FMath::DivideAndRoundUp(TopLevelGridResolution.Y, BottomLevelGridResolution);
-	TopLevelGridResolution.Z = FMath::DivideAndRoundUp(TopLevelGridResolution.Z, BottomLevelGridResolution);
+	TopLevelGridResolution.X = FMath::Max(FMath::DivideAndRoundUp(TopLevelGridResolution.X, BottomLevelGridResolution), 1);
+	TopLevelGridResolution.Y = FMath::Max(FMath::DivideAndRoundUp(TopLevelGridResolution.Y, BottomLevelGridResolution), 1);
+	TopLevelGridResolution.Z = FMath::Max(FMath::DivideAndRoundUp(TopLevelGridResolution.Z, BottomLevelGridResolution), 1);
 }
 
 void MarkTopLevelGridVoxelsForFrustumGrid(
@@ -1652,7 +1811,7 @@ void RasterizeVolumesIntoFrustumVoxelGrid(
 		const FBoxSphereBounds PrimitiveBounds = PrimitiveSceneProxy->GetBounds();
 		const FMaterial& Material = MaterialRenderProxy->GetMaterialWithFallback(View.GetFeatureLevel(), MaterialRenderProxy);
 
-		FRasterizeBottomLevelFroxelGridCS::FParameters* PassParameters = GraphBuilder.AllocParameters<FRasterizeBottomLevelFroxelGridCS::FParameters>();
+		FRasterizeBottomLevelFrustumGridCS::FParameters* PassParameters = GraphBuilder.AllocParameters<FRasterizeBottomLevelFrustumGridCS::FParameters>();
 		{
 			// Scene data
 			PassParameters->View = View.ViewUniformBuffer;
@@ -1704,8 +1863,8 @@ void RasterizeVolumesIntoFrustumVoxelGrid(
 			// Why is scene explicitly copied??
 			[PassParameters, LocalScene = Scene, &View, MaterialRenderProxy, &Material](FRHIComputeCommandList& RHICmdList)
 			{
-				FRasterizeBottomLevelFroxelGridCS::FPermutationDomain PermutationVector;
-				TShaderRef<FRasterizeBottomLevelFroxelGridCS> ComputeShader = Material.GetShader<FRasterizeBottomLevelFroxelGridCS>(&FLocalVertexFactory::StaticType, PermutationVector, false);
+				FRasterizeBottomLevelFrustumGridCS::FPermutationDomain PermutationVector;
+				TShaderRef<FRasterizeBottomLevelFrustumGridCS> ComputeShader = Material.GetShader<FRasterizeBottomLevelFrustumGridCS>(&FLocalVertexFactory::StaticType, PermutationVector, false);
 
 				if (!ComputeShader.IsNull())
 				{
@@ -2096,7 +2255,7 @@ void RasterizeVolumesIntoOrthoVoxelGrid(
 			const FBoxSphereBounds PrimitiveBounds = PrimitiveSceneProxy->GetBounds();
 			const FMaterial& Material = MaterialRenderProxy->GetMaterialWithFallback(View.GetFeatureLevel(), MaterialRenderProxy);
 
-			FRasterizeBottomLevelGrid::FParameters* PassParameters = GraphBuilder.AllocParameters<FRasterizeBottomLevelGrid::FParameters>();
+			FRasterizeBottomLevelOrthoGridCS::FParameters* PassParameters = GraphBuilder.AllocParameters<FRasterizeBottomLevelOrthoGridCS::FParameters>();
 			{
 				// Scene data
 				PassParameters->View = View.ViewUniformBuffer;
@@ -2118,6 +2277,10 @@ void RasterizeVolumesIntoOrthoVoxelGrid(
 				PassParameters->TopLevelGridWorldBoundsMin = FVector3f(TopLevelGridBounds.Origin - TopLevelGridBounds.BoxExtent);
 				PassParameters->TopLevelGridWorldBoundsMax = FVector3f(TopLevelGridBounds.Origin + TopLevelGridBounds.BoxExtent);
 
+				// Sampling data
+				FBlueNoise BlueNoise = GetBlueNoiseGlobalParameters();
+				PassParameters->BlueNoise = CreateUniformBufferImmediate(BlueNoise, EUniformBufferUsage::UniformBuffer_SingleDraw);
+
 				// Unify with "object" definition??
 				PassParameters->PrimitiveWorldBoundsMin = FVector3f(PrimitiveBounds.Origin - PrimitiveBounds.BoxExtent);
 				PassParameters->PrimitiveWorldBoundsMax = FVector3f(PrimitiveBounds.Origin + PrimitiveBounds.BoxExtent);
@@ -2165,11 +2328,10 @@ void RasterizeVolumesIntoOrthoVoxelGrid(
 				// Why is scene explicitly copied?
 				[PassParameters, LocalScene = Scene, &View, MaterialRenderProxy, &Material](FRHIComputeCommandList& RHICmdList)
 				{
-					FRasterizeBottomLevelGrid::FPermutationDomain PermutationVector;
-					PermutationVector.Set<FRasterizeBottomLevelGrid::FEnableIndirectionGrid>(HeterogeneousVolumes::EnableIndirectionGrid());
-					//PermutationVector.Set<FRasterizeBottomLevelGrid::FEnableVoxelHashing>(HeterogeneousVolumes::EnableVoxelHashing());
-					PermutationVector.Set<FRasterizeBottomLevelGrid::FEnableHomogeneousAggregation>(HeterogeneousVolumes::EnableHomogeneousAggregation());
-					TShaderRef<FRasterizeBottomLevelGrid> ComputeShader = Material.GetShader<FRasterizeBottomLevelGrid>(&FLocalVertexFactory::StaticType, PermutationVector, false);
+					FRasterizeBottomLevelOrthoGridCS::FPermutationDomain PermutationVector;
+					PermutationVector.Set<FRasterizeBottomLevelOrthoGridCS::FEnableIndirectionGrid>(HeterogeneousVolumes::EnableIndirectionGrid());
+					PermutationVector.Set<FRasterizeBottomLevelOrthoGridCS::FEnableHomogeneousAggregation>(HeterogeneousVolumes::EnableHomogeneousAggregation());
+					TShaderRef<FRasterizeBottomLevelOrthoGridCS> ComputeShader = Material.GetShader<FRasterizeBottomLevelOrthoGridCS>(&FLocalVertexFactory::StaticType, PermutationVector, false);
 
 
 					if (!ComputeShader.IsNull())
@@ -2349,6 +2511,12 @@ void RenderTransmittanceWithVoxelGrid(
 	uint32 GroupCountY = FMath::DivideAndRoundUp(View.ViewRect.Size().Y, FRenderTransmittanceWithVoxelGridCS::GetThreadGroupSize2D());
 	FIntVector GroupCount = FIntVector(GroupCountX, GroupCountY, 1);
 
+	int32 BufferSize = View.ViewRect.Width() * View.ViewRect.Height();
+	FRDGBufferRef DebugBuffer = GraphBuilder.CreateBuffer(
+		FRDGBufferDesc::CreateStructuredDesc(sizeof(FRenderDebugData), BufferSize),
+		TEXT("HeterogeneousVolumes.RenderDebugData")
+	);
+
 	FRenderTransmittanceWithVoxelGridCS::FParameters* PassParameters = GraphBuilder.AllocParameters<FRenderTransmittanceWithVoxelGridCS::FParameters>();
 	{
 		// Scene data
@@ -2368,9 +2536,12 @@ void RenderTransmittanceWithVoxelGrid(
 
 		// Output
 		PassParameters->RWLightingTexture = GraphBuilder.CreateUAV(HeterogeneousVolumeRadiance);
+		PassParameters->RWDebugBuffer = GraphBuilder.CreateUAV(DebugBuffer);
 	}
 
 	FRenderTransmittanceWithVoxelGridCS::FPermutationDomain PermutationVector;
+	int32 DebugMode = FMath::Clamp(HeterogeneousVolumes::GetDebugMode() - 1, 0, FRenderTransmittanceWithVoxelGridCS::FDebugMode::MaxValue);
+	PermutationVector.Set<FRenderTransmittanceWithVoxelGridCS::FDebugMode>(DebugMode);
 	TShaderRef<FRenderTransmittanceWithVoxelGridCS> ComputeShader = View.ShaderMap->GetShader<FRenderTransmittanceWithVoxelGridCS>(PermutationVector);
 	FComputeShaderUtils::AddPass(
 		GraphBuilder,
