Local Height Fog Volume - Clean up unused matrices and code simplification.

diff --git a/Engine/Shaders/Private/LocalHeightFog.usf b/Engine/Shaders/Private/LocalHeightFog.usf
--- a/Engine/Shaders/Private/LocalHeightFog.usf
+++ b/Engine/Shaders/Private/LocalHeightFog.usf
@@ -10,10 +10,9 @@ struct FLocalHeightFogGPUInstanceData
 {
 	float4x4 Transform;
 	float4x4 InvTransform;
-	float4x4 FogTranformNoScale;
-	float4x4 FogTranformNoScaleInv;
-	float4x4 FogTransformScaleOnly;
-	float4x4 FogTransformScaleOnlyInv;
+
+	float4x4 InvTranformNoScale;
+	float4x4 TransformScaleOnly;
 
 	float Density;
 	float HeightFalloff;
@@ -139,25 +138,25 @@ void LocalHeightFogSplatPS(
 	float3 CamRayWorldOrigin = CamRayTranslatedWorldOrigin - LWCHackToFloat(PrimaryView.PreViewTranslation);
 
 	// The "U" prefix is for all compuation done in the Unit Sphere space.
-	float4 RayPosU = mul(float4(CamRayWorldOrigin, 1.0), FogInstance.InvTransform);
-	float4 RayDirU = mul(float4(CamRayTranslatedWorldDir, 0.0), FogInstance.InvTransform);	// TODO This potentially contains a non uniform scale so check corectness
-	RayDirU.xyz = normalize(RayDirU.xyz);
+	float3 RayPosU = mul(float4(CamRayWorldOrigin,        1.0f), FogInstance.InvTransform).xyz;
+	float3 RayDirU = mul(float4(CamRayTranslatedWorldDir, 0.0f), FogInstance.InvTransform).xyz; // Try to use InverseTranspose to correctly handle direction transform with scaling.
+	RayDirU.xyz = normalize(RayDirU);
 
-	float2 TsU = RayIntersectSphere(RayPosU.xyz, RayDirU.xyz, float4(0.0, 0.0, 0.0, 1.0));
+	float2 TsU = RayIntersectSphere(RayPosU, RayDirU, float4(0.0, 0.0, 0.0, 1.0));
 	if (any(TsU > 0.0))
 	{
 		const float DeviceZ = LookupDeviceZ(uint2(SVPos.xy));
 		float3 DepthBufferTranslatedWorldPos = SvPositionToTranslatedWorld(float4(SVPos.xy, DeviceZ, 1.0));
 		float LengthD = length(DepthBufferTranslatedWorldPos - CamRayTranslatedWorldOrigin);
 
-		float3 P0U = RayPosU.xyz + max(0.0, TsU.x) * RayDirU.xyz;
-		float3 P1U = RayPosU.xyz + max(0.0, TsU.y) * RayDirU.xyz;
+		float3 P0U = RayPosU + max(0.0, TsU.x) * RayDirU;
+		float3 P1U = RayPosU + max(0.0, TsU.y) * RayDirU;
 #if LHFV_FOLLOWS_ROTATIONS
 		// Computations are all done transformed with Space = local fog volume but ignoring scaling.
-		float3 P0Space = mul(float4(P0U, 1.0), FogInstance.FogTransformScaleOnly).xyz;
-		float3 P1Space = mul(float4(P1U, 1.0), FogInstance.FogTransformScaleOnly).xyz;
-		float3 CamRayOriginSpace = mul(float4(CamRayWorldOrigin, 1.0), FogInstance.FogTranformNoScaleInv);
-		float3 CamRayDirSpace    = mul(float4(CamRayWorldDir,    0.0), FogInstance.FogTranformNoScaleInv);
+		float3 P0Space = mul(float4(P0U, 1.0), FogInstance.TransformScaleOnly).xyz;
+		float3 P1Space = mul(float4(P1U, 1.0), FogInstance.TransformScaleOnly).xyz;
+		float3 CamRayOriginSpace = mul(float4(CamRayWorldOrigin, 1.0), FogInstance.InvTranformNoScale).xyz;
+		float3 CamRayDirSpace    = mul(float4(CamRayWorldDir,    0.0), FogInstance.InvTranformNoScale).xyz;	// No scale so no need for a special InverseTranspose matrix.
 		CamRayDirSpace = normalize(CamRayDirSpace);
 #else
 		// Computations are all done with Space = world space
diff --git a/Engine/Source/Runtime/Renderer/Private/LocalHeightFogRendering.cpp b/Engine/Source/Runtime/Renderer/Private/LocalHeightFogRendering.cpp
--- a/Engine/Source/Runtime/Renderer/Private/LocalHeightFogRendering.cpp
+++ b/Engine/Source/Runtime/Renderer/Private/LocalHeightFogRendering.cpp
@@ -72,10 +72,9 @@ class FLocalHeightFogGPUInstanceData
 public:
 	FMatrix44f Transform;
 	FMatrix44f InvTransform;
-	FMatrix44f FogTranformNoScale;
-	FMatrix44f FogTranformNoScaleInv;
-	FMatrix44f FogTransformScaleOnly;
-	FMatrix44f FogTransformScaleOnlyInv;
+
+	FMatrix44f InvTranformNoScale;
+	FMatrix44f TransformScaleOnly;
 
 	float Density;
 	float HeightFalloff;
@@ -153,15 +152,13 @@ void RenderLocalHeightFog(
 		FLocalHeightFogGPUInstanceData* LocalHeightFogGPUInstanceDataIt = LocalHeightFogGPUInstanceData;
 		for (FLocalHeightFogSceneProxy* LHF : Scene->LocalHeightFogs)
 		{
-			FTransform FogTransformScaleOnly;
-			FogTransformScaleOnly.SetScale3D(LHF->FogTransform.GetScale3D());
+			FTransform TransformScaleOnly;
+			TransformScaleOnly.SetScale3D(LHF->FogTransform.GetScale3D());
 
 			LocalHeightFogGPUInstanceDataIt->Transform					= FMatrix44f(LHF->FogTransform.ToMatrixWithScale());
 			LocalHeightFogGPUInstanceDataIt->InvTransform				= LocalHeightFogGPUInstanceDataIt->Transform.Inverse();
-			LocalHeightFogGPUInstanceDataIt->FogTranformNoScale			= FMatrix44f(LHF->FogTransform.ToMatrixNoScale());
-			LocalHeightFogGPUInstanceDataIt->FogTranformNoScaleInv		= LocalHeightFogGPUInstanceDataIt->FogTranformNoScale.Inverse();
-			LocalHeightFogGPUInstanceDataIt->FogTransformScaleOnly		= FMatrix44f(FogTransformScaleOnly.ToMatrixWithScale());
-			LocalHeightFogGPUInstanceDataIt->FogTransformScaleOnlyInv	= LocalHeightFogGPUInstanceDataIt->FogTransformScaleOnly.Inverse();
+			LocalHeightFogGPUInstanceDataIt->InvTranformNoScale			= FMatrix44f(LHF->FogTransform.ToMatrixNoScale()).Inverse();
+			LocalHeightFogGPUInstanceDataIt->TransformScaleOnly			= FMatrix44f(TransformScaleOnly.ToMatrixWithScale());
 
 			LocalHeightFogGPUInstanceDataIt->Density = LHF->FogDensity;
 			LocalHeightFogGPUInstanceDataIt->HeightFalloff = LHF->FogHeightFalloff * 0.01f;	// This scale is used to have artist author reasonable range.
