Fix shader compilation with Strata.

diff --git a/Engine/Shaders/Private/Strata/StrataExport.ush b/Engine/Shaders/Private/Strata/StrataExport.ush
--- a/Engine/Shaders/Private/Strata/StrataExport.ush
+++ b/Engine/Shaders/Private/Strata/StrataExport.ush
@@ -791,7 +791,7 @@ void PackStrataOut(
 								STRATA_STORE_UINT1(PackColorLinearToGamma2AlphaLinear(float4(BSDF.TransmittanceAboveAlongN, BSDF.CoverageAboveAlongN)));
 							}
 
-							#if STRATA_GLINTS
+							#if STRATA_GLINTS_ENABLED
 							if (BSDF_GETHASGLINT(BSDF))
 							{
 								const uint Packed16 = f32tof16(SLAB_GLINT_VALUE(BSDF));
