Add SpecularProfile LUT half-angle parameterization.

diff --git a/Engine/Shaders/Private/Strata/Strata.ush b/Engine/Shaders/Private/Strata/Strata.ush
--- a/Engine/Shaders/Private/Strata/Strata.ush
+++ b/Engine/Shaders/Private/Strata/Strata.ush
@@ -1048,6 +1048,30 @@ void UnpackSSSProfile(uint In, inout float OutProfileId, inout float OutRadius,
 	OutThickness = Unpack10F((In >> 16) & 0x3FF);
 }
 
+///////////////////////////////////////////////////////////////////////////////
+// Specular profile
+
+uint PackSpecularProfile(float InProfileId)
+{
+	return asuint(InProfileId);
+}
+
+void UnpackSpecularProfile(uint In, inout float OutProfileId)
+{
+	OutProfileId = asfloat(In);
+}
+
+uint GetSpecularProfileId(float In)
+{
+	return abs(In);
+}
+
+// Return 0: View/Light angles, 1: Half angles
+uint GetSpecularProfileParameterization(float In)
+{
+	return In > 0.f ? 0u : 1u;
+}
+
 ///////////////////////////////////////////////////////////////////////////////
 // Fuzz
 
@@ -1583,10 +1607,10 @@ FStrataData GetStrataSlabBSDF(
 		SLAB_GLINT_UVDDY(StrataData.InlinedBSDF)		= ddy(GlintUV);
 	}
 	
-	if (SpecularProfileId > 0.0f)
+	if (SpecularProfileId != 0.0f)
 	{
 		BSDF_SETHASSPECPROFILE	(StrataData.InlinedBSDF, 1);
-		SLAB_SPECPROFILEID		(StrataData.InlinedBSDF) = SpecularProfileId * 255; // Stored as 0-255 slice index
+		SLAB_SPECPROFILEID		(StrataData.InlinedBSDF) = SpecularProfileId * 255; // Stored as 0-255 index, and sign store parametrization type (View/Light or Half-Angle)
 	}
 #endif
 
@@ -4329,7 +4353,7 @@ FStrataBSDF UnpackStrataBSDFIn(FStrataMaterialContainer StrataBuffer, inout FStr
 
 		if (BSDF_GETHASSPECPROFILE(OutBSDF))
 		{
-			SLAB_SPECPROFILEID(OutBSDF)	= StrataLoadUint1(StrataBuffer, StrataAddressing);
+			UnpackSpecularProfile(StrataLoadUint1(StrataBuffer, StrataAddressing), SLAB_SPECPROFILEID(OutBSDF));
 		}
 	#endif
 	}
diff --git a/Engine/Shaders/Private/Strata/StrataEvaluation.ush b/Engine/Shaders/Private/Strata/StrataEvaluation.ush
--- a/Engine/Shaders/Private/Strata/StrataEvaluation.ush
+++ b/Engine/Shaders/Private/Strata/StrataEvaluation.ush
@@ -686,8 +686,10 @@ FStrataEvaluateResult StrataEvaluateBSDFCommon(FStrataBSDFContext BSDFContext, F
 			BRANCH
 			if (BSDF_GETHASSPECPROFILE(BSDFContext.BSDF))
 			{
-				const uint ProfileID = SLAB_SPECPROFILEID(BSDFContext.BSDF);
-				Sample.IntegratedSpecularValue *= View.SpecularProfileTexture.SampleLevel(View.SpecularProfileSampler, float3(BSDFContext.SatNoV, BSDFContext.SatNoL, ProfileID), 0);
+				const uint ProfileID = GetSpecularProfileId(SLAB_SPECPROFILEID(BSDFContext.BSDF));
+				const bool bUseHalfVector = GetSpecularProfileParameterization(SLAB_SPECPROFILEID(BSDFContext.BSDF)) == 1;
+				const float3 ProfileCoord = bUseHalfVector ? float3(BSDFContext.Context.VoH, BSDFContext.Context.NoH, ProfileID) : float3(BSDFContext.SatNoV, BSDFContext.SatNoL, ProfileID);
+				Sample.IntegratedSpecularValue *= View.SpecularProfileTexture.SampleLevel(View.SpecularProfileSampler, ProfileCoord, 0);
 			}
 			#endif
 
@@ -1536,7 +1538,7 @@ FStrataEnvLightResult StrataEvaluateForEnvLight(FStrataBSDFContext BSDFContext,
 				BRANCH
 				if (BSDF_GETHASSPECPROFILE(BSDFContext.BSDF))
 				{
-					const uint ProfileID = SLAB_SPECPROFILEID(BSDFContext.BSDF);
+					const uint ProfileID = GetSpecularProfileId(SLAB_SPECPROFILEID(BSDFContext.BSDF));
 					const float NoR = saturate(dot(BSDFContext.R, BSDFContext.N));
 					StrataEnvLightResult.SpecularWeight *= View.SpecularProfileTexture.SampleLevel(View.SpecularProfileSampler, float3(BSDFContext.SatNoV, NoR, ProfileID), 0);
 				}
diff --git a/Engine/Shaders/Private/Strata/StrataExport.ush b/Engine/Shaders/Private/Strata/StrataExport.ush
--- a/Engine/Shaders/Private/Strata/StrataExport.ush
+++ b/Engine/Shaders/Private/Strata/StrataExport.ush
@@ -808,7 +808,7 @@ void PackStrataOut(
 
 							if (BSDF_GETHASSPECPROFILE(BSDF))
 							{
-								STRATA_STORE_UINT1(uint(SLAB_SPECPROFILEID(BSDF)));
+								STRATA_STORE_UINT1(PackSpecularProfile(SLAB_SPECPROFILEID(BSDF)));
 							}
 							#endif // STRATA_GLINTS
 
diff --git a/Engine/Shaders/Private/Strata/StrataVisualize.usf b/Engine/Shaders/Private/Strata/StrataVisualize.usf
--- a/Engine/Shaders/Private/Strata/StrataVisualize.usf
+++ b/Engine/Shaders/Private/Strata/StrataVisualize.usf
@@ -336,7 +336,15 @@ void PrintBSDF(inout FShaderPrintContext Context, inout FStrataAddressing Strata
 			}
 			if (BSDF_GETHASSPECPROFILE(BSDF))
 			{
-				Print(Context, TEXT(" SpecID      "), FontBSDFPropName);	Print(Context, SLAB_SPECPROFILEID(BSDF), FontBSDFPropValu);				Newline(Context, RectMax);
+				Print(Context, TEXT(" SpecID      "), FontBSDFPropName);	Print(Context, GetSpecularProfileId(SLAB_SPECPROFILEID(BSDF)), FontBSDFPropValu);Newline(Context, RectMax);
+				if (GetSpecularProfileParameterization(SLAB_SPECPROFILEID(BSDF)) == 0)
+				{
+				Print(Context, TEXT(" Spec Param. "), FontBSDFPropName);	Print(Context, TEXT("View/Light angles"), FontBSDFPropValu);			Newline(Context, RectMax);
+				}
+				else
+				{
+				Print(Context, TEXT(" Spec Param. "), FontBSDFPropName);	Print(Context, TEXT("Half angles"), FontBSDFPropValu);					Newline(Context, RectMax);
+				}
 			}
 #endif
 		}
diff --git a/Engine/Source/Runtime/Engine/Classes/Engine/SpecularProfile.h b/Engine/Source/Runtime/Engine/Classes/Engine/SpecularProfile.h
--- a/Engine/Source/Runtime/Engine/Classes/Engine/SpecularProfile.h
+++ b/Engine/Source/Runtime/Engine/Classes/Engine/SpecularProfile.h
@@ -14,12 +14,26 @@ class UTexture2D;
 class FRDGBuilder;
 class FTextureReference;
 
+/** List of niagara solvers */
+UENUM(BlueprintType)
+enum class ESpecularProfileFormat : uint8
+{
+	ViewLightVector UMETA(DisplayName = "View / Light"),
+	HalfVector UMETA(DisplayName = "Half Angle"),
+};
+
 // struct with all the settings we want in USpecularProfile, separate to make it easer to pass this data around in the engine.
 USTRUCT(BlueprintType)
 struct FSpecularProfileStruct
 {
 	GENERATED_USTRUCT_BODY()
 	
+	/**
+	 * Define the texture used as a specular profile
+	 */
+	UPROPERTY(Category = "Common", EditAnywhere, BlueprintReadOnly, meta=(DisplayName="LUT Format"))
+	ESpecularProfileFormat Format;
+
 	/**
 	* Define the view facing color
 	*/
@@ -38,9 +52,9 @@ struct FSpecularProfileStruct
 	UPROPERTY(Category  ="Texture", EditAnywhere, BlueprintReadOnly, meta=(DisplayName="Texture"))
 	TObjectPtr<UTexture2D> Texture;
 
-	// constructor
 	FSpecularProfileStruct()
 	{
+		Format = ESpecularProfileFormat::ViewLightVector;
 		ViewColor = FLinearColor(1.0f, 1.0f, 1.0f);
 		LightColor = FLinearColor(1.0f, 1.0f, 1.0f);
 		Texture = nullptr;
diff --git a/Engine/Source/Runtime/Engine/Private/Rendering/SpecularProfile.cpp b/Engine/Source/Runtime/Engine/Private/Rendering/SpecularProfile.cpp
--- a/Engine/Source/Runtime/Engine/Private/Rendering/SpecularProfile.cpp
+++ b/Engine/Source/Runtime/Engine/Private/Rendering/SpecularProfile.cpp
@@ -109,6 +109,9 @@ public:
 	// return the parameter name for a given profile
 	FName GetParameterName(const USpecularProfile* InProfile) const;
 
+	// Return the parameterization sign
+	float GetParameterizationSign(int32 AllocationId) const;
+
 	// @return can be nullptr if there is no SpecularProfile
 	IPooledRenderTarget* GetAtlasTexture(FRDGBuilder& GraphBuilder, EShaderPlatform ShaderPlatform);
 	IPooledRenderTarget* GetAtlasTexture();
@@ -267,6 +270,16 @@ FName FSpecularProfileTextureManager::GetParameterName(const USpecularProfile* I
 	return FName();
 }
 
+float FSpecularProfileTextureManager::GetParameterizationSign(int32 AllocationId) const
+{
+	float OutSign = 1.f;
+	if (AllocationId != INDEX_NONE)
+	{
+		OutSign = SpecularProfileEntries[AllocationId].Settings.Format == ESpecularProfileFormat::HalfVector ? -1.f : 1.f;
+	}
+	return OutSign;
+}
+
 IPooledRenderTarget* FSpecularProfileTextureManager::GetAtlasTexture()
 {
 	return GSpecularProfileTextureAtlas;
@@ -426,7 +439,7 @@ namespace SpecularProfileAtlas
 			// can be optimized (cached)
 			AllocationId = GSpecularProfileTextureManager.FindAllocationId(In);
 		}
-		return AllocationId / 255.0f;
+		return GSpecularProfileTextureManager.GetParameterizationSign(AllocationId) * AllocationId / 255.0f;
 	}
 	
 	int32 AddOrUpdateProfile(const USpecularProfile* InProfile, const FGuid& InGuid, const FSpecularProfileStruct InSettings, const FTextureReference* InTexture)
