Materials: Implement GetShadowReplaceState for raytraced shadows

Raytracing shaders are shared across raytraced effects, so introduce a global variable that gets set from the ray payload flags to implement this function. This is similar to how the path tracing ray specific behavior is implemented.

diff --git a/Engine/Shaders/Private/Common.ush b/Engine/Shaders/Private/Common.ush
--- a/Engine/Shaders/Private/Common.ush
+++ b/Engine/Shaders/Private/Common.ush
@@ -1916,29 +1916,35 @@ IsAndMakeFinite_Def(float2, bool2)
 IsAndMakeFinite_Def(float3, bool3)
 IsAndMakeFinite_Def(float4, bool4)
 
+// NOTE: The raytraced implementation of the ShadowPassSwitch node is kept in RayTracingShaderUtils.ush as it needs to access per ray information.
+#if RAYHITGROUPSHADER == 0
 // Experimental way to allow adjusting the OpacityMask for shadow map rendering of masked materials.
-// Can be accessed with a Custom material node. If this turns out to be very useful we can expose as MaterialFunction
+// This is exposed via the ShadowPassSwitch material node. This can also be accessed with a Custom
+// material node. If this turns out to be very useful we can expose as MaterialFunction
 // and potentially expose other queries as well (e.g. SkeletalMesh, HitProxy, ).
 // @return 0:no, 1:yes
 bool GetShadowReplaceState()
 {
-#ifdef SHADOW_DEPTH_SHADER
+#if SHADOW_DEPTH_SHADER
 	return true;
 #else
 	return false;
 #endif
 }
 
-bool GetReflectionCapturePassSwitchState()
+float IsShadowDepthShader()
 {
-	return View.RenderingReflectionCaptureMask > 0.0f;
+	return GetShadowReplaceState() ? 1.0f : 0.0f;
 }
 
-float IsShadowDepthShader()
+#endif // RAYHITGROUPSHADER == 0
+
+bool GetReflectionCapturePassSwitchState()
 {
-	return GetShadowReplaceState() ? 1.0f : 0.0f;
+	return View.RenderingReflectionCaptureMask > 0.0f;
 }
 
+
 #define TERRAIN_ZSCALE (1.0f/128.0f)
 
 // Decodes a value which was packed into two 8 bit channels
diff --git a/Engine/Shaders/Private/PathTracing/PathTracingShaderUtils.ush b/Engine/Shaders/Private/PathTracing/PathTracingShaderUtils.ush
--- a/Engine/Shaders/Private/PathTracing/PathTracingShaderUtils.ush
+++ b/Engine/Shaders/Private/PathTracing/PathTracingShaderUtils.ush
@@ -41,3 +41,16 @@ bool GetPathTracingIsIndirectVolume()
 {
 	return (CurrentPayloadInputFlags & PATH_TRACING_PAYLOAD_INPUT_FLAG_RAY_TYPE_MASK) == PATH_TRACING_PAYLOAD_INPUT_FLAG_RAY_TYPE_INDIRECT_VOLUME;
 }
+
+#if RAYHITGROUPSHADER
+// In the CHS/AHS case, also define these (which will be missing from Common.ush)
+bool GetShadowReplaceState()
+{
+	return GetPathTracingIsShadow();
+}
+
+float IsShadowDepthShader()
+{
+	return GetShadowReplaceState() ? 1.0f : 0.0f;
+}
+#endif
diff --git a/Engine/Shaders/Private/RayTracing/RayTracingCommon.ush b/Engine/Shaders/Private/RayTracing/RayTracingCommon.ush
--- a/Engine/Shaders/Private/RayTracing/RayTracingCommon.ush
+++ b/Engine/Shaders/Private/RayTracing/RayTracingCommon.ush
@@ -159,7 +159,7 @@ void name##_INTERNAL(uint3 DispatchThreadIndex)
 #define RAY_TRACING_BLEND_MODE_MODULATE				4
 #define RAY_TRACING_BLEND_MODE_ALPHA_HOLDOUT		5
 
-// Check how Flags are packed in the FPackedMaterialClosestHitPayload before adding to ensure bits are free
+// Check how Flags are packed in the FPackedMaterialClosestHitPayload before adding to ensure bits are free (current packing gives us 6 bits for input and output)
 // There are two sets of flags:
 //   - "inputs" are set before calling TraceRay
 //   - "outputs" are set during execution of the CHS and are valid after TraceRay returns
@@ -177,7 +177,7 @@ void name##_INTERNAL(uint3 DispatchThreadIndex)
 #define RAY_TRACING_PAYLOAD_INPUT_FLAG_ENABLE_SKY_LIGHT_CONTRIBUTION (1 << 1) // Indicates that hit shaders should add in the Sky Light contribution in their output. Set by RayGen shader, before TraceRay().
 #define RAY_TRACING_PAYLOAD_INPUT_FLAG_LUMEN_PAYLOAD		         (1 << 2) // Indicates that closest hit shader should evaluate geometry normal and instanceid.
 #define RAY_TRACING_PAYLOAD_INPUT_FLAG_IGNORE_TRANSLUCENT            (1 << 3) // Indicates that translucent materials should be ignored by IgnoreHit() in anyhit shaders.
-#define RAY_TRACING_PAYLOAD_INPUT_FLAG_UNUSED4                       (1 << 4) // Free for future use
+#define RAY_TRACING_PAYLOAD_INPUT_FLAG_SHADOW_RAY                    (1 << 4) // Indicates this ray is being traced as part of a raytraced shadow calculation (see TraceVisibilityRay).
 #define RAY_TRACING_PAYLOAD_INPUT_FLAG_UNUSED5                       (1 << 5) // Free for future use
 
 ///////////////////////////////////////////////////////////////////////////////////////////////////
@@ -850,6 +850,7 @@ struct FPackedMaterialClosestHitPayload : FMinimalPayload
 	uint GetShadingModelID() { return (IorAndShadingModelIDAndBlendingModeAndFlagsAndPrimitiveLightingChannelMask >> 16) & 0xF; }
 	uint GetBlendingMode()   { return (IorAndShadingModelIDAndBlendingModeAndFlagsAndPrimitiveLightingChannelMask >> 20) & 0x7; }
 	uint GetFlags()          { return (IorAndShadingModelIDAndBlendingModeAndFlagsAndPrimitiveLightingChannelMask >> 23) & 0x3F; }
+	void SetFlags(uint Flags) { IorAndShadingModelIDAndBlendingModeAndFlagsAndPrimitiveLightingChannelMask |= (Flags & 0x3F) << 23; }
 	bool IsTwoSided()		 { return (GetFlags() & RAY_TRACING_PAYLOAD_OUTPUT_FLAG_TWO_SIDED) != 0; }
 	bool IsFrontFace()		 { return (GetFlags() & RAY_TRACING_PAYLOAD_OUTPUT_FLAG_FRONT_FACE) != 0; }
 	uint GetPrimitiveLightingChannelMask() { return (IorAndShadingModelIDAndBlendingModeAndFlagsAndPrimitiveLightingChannelMask >> 29) & 0x7; }
@@ -1076,16 +1077,12 @@ void TraceVisibilityRayPacked(
 FMinimalPayload TraceVisibilityRay(
 	in RaytracingAccelerationStructure TLAS,
 	in uint RayFlags,
-	in uint PayloadFlags,
 	in uint InstanceInclusionMask,
 	in FRayDesc Ray)
 {
 	FPackedMaterialClosestHitPayload PackedPayload = (FPackedMaterialClosestHitPayload)0;
 	
-	if ((PayloadFlags & RAY_TRACING_PAYLOAD_INPUT_FLAG_IGNORE_TRANSLUCENT) != 0)
-	{
-		PackedPayload.SetIgnoreTranslucentMaterials();
-	}
+	PackedPayload.SetFlags(RAY_TRACING_PAYLOAD_INPUT_FLAG_SHADOW_RAY);
 
 	TraceVisibilityRayPacked(PackedPayload, TLAS, RayFlags, InstanceInclusionMask, Ray);
 
@@ -1101,20 +1098,6 @@ FMinimalPayload TraceVisibilityRay(
 	return MinimalPayload;
 }
 
-FMinimalPayload TraceVisibilityRay(
-	in RaytracingAccelerationStructure TLAS,
-	in uint RayFlags,
-	in uint InstanceInclusionMask,
-	in FRayDesc Ray)
-{
-	return TraceVisibilityRay(
-		TLAS,
-		RayFlags,
-		RAY_TRACING_PAYLOAD_INPUT_FLAG_NONE,
-		InstanceInclusionMask,
-		Ray);
-}
-
 #if !COMPUTESHADER
 
 void TraceMaterialRayPacked(
diff --git a/Engine/Shaders/Private/RayTracing/RayTracingMaterialHitShaders.usf b/Engine/Shaders/Private/RayTracing/RayTracingMaterialHitShaders.usf
--- a/Engine/Shaders/Private/RayTracing/RayTracingMaterialHitShaders.usf
+++ b/Engine/Shaders/Private/RayTracing/RayTracingMaterialHitShaders.usf
@@ -21,6 +21,7 @@
 #include "/Engine/Private/Common.ush"
 #include "RayTracingCommon.ush"
 #include "RayTracingHitGroupCommon.ush"
+#include "RayTracingShaderUtils.ush"
 
 #include "/Engine/Generated/Material.ush"   
 #include "/Engine/Generated/VertexFactory.ush"
@@ -497,6 +498,8 @@ RAY_TRACING_ENTRY_CLOSEST_HIT(MaterialCHS,
 		return;
 	}
 
+	CurrentPayloadInputFlags = PackedPayload.GetFlags();
+
 	FPixelMaterialInputs PixelMaterialInputs;
 
 	{
@@ -616,6 +619,9 @@ RAY_TRACING_ENTRY_ANY_HIT(MaterialAHS,
 #else
 	FMaterialPixelParameters MaterialParameters = GetMaterialPixelParameters(Interpolants, SvPosition);
 #endif
+
+	CurrentPayloadInputFlags = PackedPayload.GetFlags();
+
 	FPixelMaterialInputs PixelMaterialInputs;
 
 	{
diff --git a/Engine/Shaders/Private/RayTracing/RayTracingShaderUtils.ush b/Engine/Shaders/Private/RayTracing/RayTracingShaderUtils.ush
new file mode 100644
--- /dev/null
+++ b/Engine/Shaders/Private/RayTracing/RayTracingShaderUtils.ush
@@ -0,0 +1,27 @@
+// Copyright Epic Games, Inc. All Rights Reserved.
+
+#pragma once
+
+// This file contains utility methods that can be used by ray tracing shaders (such as the CHS/AHS pairs, light miss shaders, etc ...)
+
+#ifndef RAYHITGROUPSHADER
+#error "This header should only be included in raytracing contexts"
+#endif
+
+#ifndef PATH_TRACING // Path Tracing has a similar implemental with a slightly different set of flags
+
+#include "RayTracingCommon.ush"
+
+static int CurrentPayloadInputFlags = 0;
+
+bool GetShadowReplaceState()
+{
+	return (CurrentPayloadInputFlags & RAY_TRACING_PAYLOAD_INPUT_FLAG_SHADOW_RAY) != 0;
+}
+
+float IsShadowDepthShader()
+{
+	return GetShadowReplaceState() ? 1.0f : 0.0f;
+}
+
+#endif
\ No newline at end of file
