Implemented COMPILER_ALLOW_CS_DERIVATIVES on applicable platforms, and decorated Nanite CS materials with it so we efficient deriv tex2d samples (not just mip0)

diff --git a/Engine/Shaders/Private/ComputeShaderOutputCommon.ush b/Engine/Shaders/Private/ComputeShaderOutputCommon.ush
--- a/Engine/Shaders/Private/ComputeShaderOutputCommon.ush
+++ b/Engine/Shaders/Private/ComputeShaderOutputCommon.ush
@@ -6,6 +6,8 @@ through a single struct, allowing for a more readable code
 (less #ifdefs, reducing the boolean hell)
 =============================================================================*/
 
+COMPILER_ALLOW_CS_DERIVATIVES
+
 #include "ShaderOutputCommon.ush"
 #include "GammaCorrectionCommon.ush"
 #include "VariableRateShading/VRSShadingRateCommon.ush"
diff --git a/Engine/Shaders/Public/Platform.ush b/Engine/Shaders/Public/Platform.ush
--- a/Engine/Shaders/Public/Platform.ush
+++ b/Engine/Shaders/Public/Platform.ush
@@ -133,6 +133,10 @@
 #define COMPILER_FORCE_WAVE32_MODE
 #endif
 
+#ifndef COMPILER_ALLOW_CS_DERIVATIVES
+#define COMPILER_ALLOW_CS_DERIVATIVES
+#endif
+
 #ifndef PLATFORM_SUPPORTS_ROV
 #define PLATFORM_SUPPORTS_ROV 0
 #define RasterizerOrderedTexture2D RWTexture2D
