Substrate - Added special path for more expensive features such as glints or SpecularLUT.
This path is taken in the shader only when needed using #define.

diff --git a/Engine/Shaders/Private/ClusteredDeferredShadingPixelShader.usf b/Engine/Shaders/Private/ClusteredDeferredShadingPixelShader.usf
--- a/Engine/Shaders/Private/ClusteredDeferredShadingPixelShader.usf
+++ b/Engine/Shaders/Private/ClusteredDeferredShadingPixelShader.usf
@@ -21,7 +21,10 @@
   #elif STRATA_TILETYPE == 1
 	#define STRATA_SINGLEPATH 1
   #elif STRATA_TILETYPE == 2
-  // COMPLEX PATH
+	// COMPLEX PATH
+  #elif STRATA_TILETYPE == 3
+	// COMPLEX PATH
+	#define STRATA_COMPLEX_SPECIALPATH_ENABLED 1
   #else
 	#error Strata tile type non-implemented
   #endif
diff --git a/Engine/Shaders/Private/DeferredLightPixelShaders.usf b/Engine/Shaders/Private/DeferredLightPixelShaders.usf
--- a/Engine/Shaders/Private/DeferredLightPixelShaders.usf
+++ b/Engine/Shaders/Private/DeferredLightPixelShaders.usf
@@ -14,6 +14,9 @@
 	#define STRATA_SINGLEPATH 1
   #elif STRATA_TILETYPE == 2
 	// COMPLEX PATH
+  #elif STRATA_TILETYPE == 3
+	// COMPLEX PATH
+	#define STRATA_COMPLEX_SPECIALPATH_ENABLED 1
   #else
 	#error Strata tile type non-implemented
   #endif
diff --git a/Engine/Shaders/Private/Definitions.usf b/Engine/Shaders/Private/Definitions.usf
--- a/Engine/Shaders/Private/Definitions.usf
+++ b/Engine/Shaders/Private/Definitions.usf
@@ -100,10 +100,6 @@
 #error Uknown STRATA_NORMAL_QUALITY
 #endif
 
-#ifndef STRATA_GLINTS
-#define STRATA_GLINTS									0
-#endif
-
 #ifndef MATERIAL_SHADINGMODEL_DEFAULT_LIT
 #define MATERIAL_SHADINGMODEL_DEFAULT_LIT				0
 #endif
diff --git a/Engine/Shaders/Private/DiffuseIndirectComposite.usf b/Engine/Shaders/Private/DiffuseIndirectComposite.usf
--- a/Engine/Shaders/Private/DiffuseIndirectComposite.usf
+++ b/Engine/Shaders/Private/DiffuseIndirectComposite.usf
@@ -52,6 +52,9 @@
 	#define STRATA_SINGLEPATH 1
   #elif STRATA_TILETYPE == 2
 	// COMPLEX PATH
+  #elif STRATA_TILETYPE == 3
+	// COMPLEX PATH
+	#define STRATA_COMPLEX_SPECIALPATH_ENABLED 1
   #else
 	#error Strata tile type non-implemented
   #endif
diff --git a/Engine/Shaders/Private/ReflectionEnvironmentPixelShader.usf b/Engine/Shaders/Private/ReflectionEnvironmentPixelShader.usf
--- a/Engine/Shaders/Private/ReflectionEnvironmentPixelShader.usf
+++ b/Engine/Shaders/Private/ReflectionEnvironmentPixelShader.usf
@@ -6,15 +6,18 @@
   
 #if STRATA_ENABLED
 #define STRATA_INLINE_SHADING 0
-#if STRATA_TILETYPE == 0
+  #if STRATA_TILETYPE == 0
 	#define STRATA_FASTPATH 1
-#elif STRATA_TILETYPE == 1
+  #elif STRATA_TILETYPE == 1
 	#define STRATA_SINGLEPATH 1
-#elif STRATA_TILETYPE == 2
+  #elif STRATA_TILETYPE == 2
 	// COMPLEX PATH
-#else
+  #elif STRATA_TILETYPE == 3
+	// COMPLEX PATH
+	#define STRATA_COMPLEX_SPECIALPATH_ENABLED 1
+  #else
 	#error Strata tile type non-implemented
-#endif
+  #endif
 #endif
 
 #include "Common.ush"
diff --git a/Engine/Shaders/Private/Strata/Strata.ush b/Engine/Shaders/Private/Strata/Strata.ush
--- a/Engine/Shaders/Private/Strata/Strata.ush
+++ b/Engine/Shaders/Private/Strata/Strata.ush
@@ -95,11 +95,16 @@
 #define STRATA_SINGLEPATH 0
 #endif
 
+#ifndef STRATA_COMPLEX_SPECIALPATH_ENABLED 
+#define STRATA_COMPLEX_SPECIALPATH_ENABLED 0
+#endif
+
 #define STRATA_COMPLEXPATH (!STRATA_FASTPATH && !STRATA_SINGLEPATH)
+#define STRATA_COMPLEXPATH_SPECIAL (STRATA_COMPLEXPATH && STRATA_COMPLEX_SPECIALPATH_ENABLED)
 
 // Disable glint when in tracing mode for now since derivative are not available.
 #define STRATA_GLINTS_DERIVATIVE_SUPPORTED (!RAYCALLABLESHADER && !RAYHITGROUPSHADER && !RAYGENSHADER && !RAYMISSSHADER)
-#define STRATA_GLINTS_ENABLED (STRATA_COMPLEXPATH && STRATA_GLINTS && STRATA_GLINTS_DERIVATIVE_SUPPORTED)
+#define STRATA_GLINTS_ENABLED (STRATA_COMPLEXPATH_SPECIAL && STRATA_GLINTS_DERIVATIVE_SUPPORTED)
 
 #define STRATA_LEGACY_MATERIAL_APPLIES_FINAL_WEIGHT (!MATERIALBLENDING_MASKED && !MATERIALBLENDING_SOLID && STRATA_LEGACY_PREMULT_ALPHA_OVERRIDE==0 || STRATA_MATERIAL_EXPORT_FROM_TRANSLUCENT)
 
@@ -161,7 +166,7 @@ struct FStrataBSDF
 	float4 InlineVGPRs[3];
 #endif
 
-#if STRATA_GLINTS_ENABLED
+#if STRATA_COMPLEXPATH_SPECIAL	
 	float4 VGPRs[7];
 #else
 	float4 VGPRs[5];
@@ -290,13 +295,20 @@ struct FStrataBSDF
 #define HEADER_BIT_COUNT_BSDF_COUNT					4
 #define HEADER_BIT_COUNT_SHAREDLOCALBASES_COUNT		3
 #define HEADER_BIT_COUNT_SHAREDLOCALBASES_TYPE		4
-#define HEADER_BIT_COUNT___UNUSED___				5
+#define HEADER_BIT_COUNT_COMPLEXPATH_SPECIAL		1
+#define HEADER_BIT_COUNT___UNUSED___				4
 
 #define HEADER_BIT_OFFSET_STATE						0
 #define HEADER_BIT_OFFSET_IRRADIANCE_AO				(HEADER_BIT_OFFSET_STATE					+ HEADER_BIT_COUNT_STATE)
 #define HEADER_BIT_OFFSET_BSDF_COUNT				(HEADER_BIT_OFFSET_IRRADIANCE_AO			+ HEADER_BIT_COUNT_IRRADIANCE_AO)
 #define HEADER_BIT_OFFSET_SHAREDLOCALBASES_COUNT	(HEADER_BIT_OFFSET_BSDF_COUNT				+ HEADER_BIT_COUNT_BSDF_COUNT)
 #define HEADER_BIT_OFFSET_SHAREDLOCALBASES_TYPE		(HEADER_BIT_OFFSET_SHAREDLOCALBASES_COUNT	+ HEADER_BIT_COUNT_SHAREDLOCALBASES_COUNT)
+#define HEADER_BIT_OFFSET_COMPLEXPATH_SPECIAL		(HEADER_BIT_OFFSET_SHAREDLOCALBASES_TYPE	+ HEADER_BIT_COUNT_SHAREDLOCALBASES_TYPE)
+
+// Sanity check
+#if (HEADER_BIT_OFFSET_COMPLEXPATH_SPECIAL + HEADER_BIT_COUNT_COMPLEXPATH_SPECIAL) > 32
+#error Strata Header is larger than 32 bits
+#endif
 
 // The different layout 'aliases' differently the material header
 // 
@@ -319,11 +331,6 @@ struct FStrataBSDF
 #define HEADER_SLWENCODING_BIT_COUNT				(HEADER_BIT_COUNT_STATE + HEADER_BIT_COUNT_IRRADIANCE_AO)
 #define HEADER_SLWENCODING_MASK					((1u<<HEADER_SLWENCODING_BIT_COUNT)-1u)
 
-// Sanity check
-#if (HEADER_BIT_OFFSET_SHAREDLOCALBASES_TYPE + HEADER_BIT_COUNT_SHAREDLOCALBASES_TYPE) > 32
-#error Strata Header is larger than 32 bits
-#endif
-
 // A header always has a common state represented by stat bits (low 8 bits) and Irradiance/AO bits (higher 8 bits)
 #define HEADER_GETCOMMONSTATES(X)					READ_BITS( X, (HEADER_BIT_COUNT_STATE + HEADER_BIT_COUNT_IRRADIANCE_AO), HEADER_BIT_OFFSET_STATE)
 #define HEADER_SETCOMMONSTATES(X, NewValue)			WRITE_BITS(X, (HEADER_BIT_COUNT_STATE + HEADER_BIT_COUNT_IRRADIANCE_AO), HEADER_BIT_OFFSET_STATE, NewValue)
@@ -343,6 +350,9 @@ struct FStrataBSDF
 #define HEADER_GETSHAREDLOCALBASESTYPE(X)			READ_BITS( X, HEADER_BIT_COUNT_SHAREDLOCALBASES_TYPE, HEADER_BIT_OFFSET_SHAREDLOCALBASES_TYPE)
 #define HEADER_SETSHAREDLOCALBASESTYPE(X, NewValue)	WRITE_BITS(X, HEADER_BIT_COUNT_SHAREDLOCALBASES_TYPE, HEADER_BIT_OFFSET_SHAREDLOCALBASES_TYPE, NewValue)
 
+#define HEADER_GETCOMPLEXSPECIALPATH(X)				READ_BITS( X, HEADER_BIT_COUNT_COMPLEXPATH_SPECIAL, HEADER_BIT_OFFSET_COMPLEXPATH_SPECIAL)
+#define HEADER_SETCOMPLEXSPECIALPATH(X, NewValue)	WRITE_BITS(X, HEADER_BIT_COUNT_COMPLEXPATH_SPECIAL, HEADER_BIT_OFFSET_COMPLEXPATH_SPECIAL, NewValue)
+
 #define HEADER_MATERIALMODE_NONE				0
 #define HEADER_MATERIALMODE_SLAB_SIMPLE			1
 #define HEADER_MATERIALMODE_SLAB_SINGLE			2
@@ -531,7 +541,7 @@ struct FStrataBSDF
 #define SLAB_FUZZ_COLOR(X)				X.VGPRs[4].xyz
 #define SLAB_FUZZ_AMOUNT(X)				X.VGPRs[4].w
 #define SLAB_FUZZ_ROUGHNESS(X)			X.VGPRs[2].w
-#if STRATA_GLINTS
+#if STRATA_COMPLEXPATH_SPECIAL 
 // Glints
 #define SLAB_GLINT_UV(X)				X.VGPRs[5].xy
 #define SLAB_GLINT_VALUE(X)				X.VGPRs[5].z
@@ -1597,6 +1607,7 @@ FStrataData GetStrataSlabBSDF(
 	BSDF_SETTHICKNESSCM		(StrataData.InlinedBSDF,		ThicknessCm);
 #endif
 
+#if STRATA_COMPLEXPATH_SPECIAL
 #if STRATA_GLINTS_ENABLED
 	if (GlintValue > 0.0f)
 	{
@@ -1606,6 +1617,7 @@ FStrataData GetStrataSlabBSDF(
 		SLAB_GLINT_UVDDX(StrataData.InlinedBSDF)		= ddx(GlintUV);
 		SLAB_GLINT_UVDDY(StrataData.InlinedBSDF)		= ddy(GlintUV);
 	}
+#endif
 	
 	if (SpecularProfileId != 0.0f)
 	{
@@ -2071,6 +2083,8 @@ struct FStrataPixelHeader
 
 	uint State;
 
+	bool bUsesComplexSpecialRenderPath; // STRATA_TODO pack with someting else, in State maybe?
+
 #if STRATA_INLINE_SHADING
 	// All the shared local bases in VGPRS. This is only used in basepass and during forward rendering.
 	FSharedLocalBases SharedLocalBases;
@@ -2149,7 +2163,9 @@ struct FStrataPixelHeader
 	void SetDynamicIndirectShadowCasterRepresentation(bool bIn)	{ this.State |= (bIn ? HEADER_MASK_HASDYNINDIRECTSHADOWCASTER : 0u); }
 	void SetHasSubsurface(bool bIn)								{ this.State |= (bIn ? HEADER_MASK_HASSUBSURFACE : 0u); }
 
-	void SetMaterialMode(uint MaterialMode)							{ this.State = (this.State & ~HEADER_MASK_MATERIALMODE) | (MaterialMode & HEADER_MASK_MATERIALMODE); }
+	void SetMaterialMode(uint MaterialMode)						{ this.State = (this.State & ~HEADER_MASK_MATERIALMODE) | (MaterialMode & HEADER_MASK_MATERIALMODE); }
+
+	void SetUsesComplexSpecialRenderPath(bool bIn)				{ this.bUsesComplexSpecialRenderPath |= bIn; }
 };
 
 #if STRATA_DEFERRED_SHADING && STRATA_MATERIALCONTAINER_IS_VIEWRESOURCE
@@ -2161,6 +2177,7 @@ FStrataPixelHeader InitialiseStrataPixelHeader()
 {
 	FStrataPixelHeader Out;
 	Out.BSDFCount = 0;
+	Out.bUsesComplexSpecialRenderPath = false;
 #if STRATA_INLINE_SHADING
 	Out.SharedLocalBases = (FSharedLocalBases)0;
 	Out.StrataTree = GetInitialisedStrataTree();
@@ -2414,12 +2431,14 @@ FStrataData StrataHorizontalMixingParameterBlending(FStrataData Background, FStr
 	SLAB_FUZZ_ROUGHNESS(ResultBSDF)	=	lerp(SLAB_FUZZ_ROUGHNESS(BackgroundBSDF),		SLAB_FUZZ_ROUGHNESS(ForegroundBSDF),		ForegroundOtherMixFactor);
 	SLAB_HAZINESS(ResultBSDF) =			LerpHazinessParameterBlending(BackgroundBSDF,	ForegroundBSDF,								ForegroundOtherMixFactor);
 
+#if STRATA_COMPLEXPATH_SPECIAL
 #if STRATA_GLINTS_ENABLED
 	BSDF_SETHASGLINT(ResultBSDF,		BSDF_GETHASGLINT(ForegroundBSDF)				| BSDF_GETHASGLINT(BackgroundBSDF));
 	SLAB_GLINT_UV(ResultBSDF) =			lerp(SLAB_GLINT_UV(BackgroundBSDF),				SLAB_GLINT_UV(ForegroundBSDF),				ForegroundOtherMixFactor);
 	SLAB_GLINT_VALUE(ResultBSDF) =		lerp(SLAB_GLINT_VALUE(BackgroundBSDF),			SLAB_GLINT_VALUE(ForegroundBSDF),			ForegroundOtherMixFactor);
 	SLAB_GLINT_UVDDX(ResultBSDF) =		lerp(SLAB_GLINT_UVDDX(BackgroundBSDF),			SLAB_GLINT_UVDDX(ForegroundBSDF),			ForegroundOtherMixFactor);
 	SLAB_GLINT_UVDDY(ResultBSDF) =		lerp(SLAB_GLINT_UVDDY(BackgroundBSDF),			SLAB_GLINT_UVDDY(ForegroundBSDF),			ForegroundOtherMixFactor);
+#endif
 
 	BSDF_SETHASSPECPROFILE(ResultBSDF,	BSDF_GETHASSPECPROFILE(ForegroundBSDF)			| BSDF_GETHASSPECPROFILE(BackgroundBSDF));
 	SLAB_SPECPROFILEID(ResultBSDF) = BSDF_GETHASSPECPROFILE(ForegroundBSDF) ? SLAB_SPECPROFILEID(ForegroundBSDF) : SLAB_SPECPROFILEID(BackgroundBSDF);
@@ -2532,12 +2551,14 @@ FStrataData StrataAddParameterBlending(FStrataData A, FStrataData B, float AMixF
 	SLAB_FUZZ_ROUGHNESS(ResultBSDF) =	lerp(SLAB_FUZZ_ROUGHNESS(BBSDF),				SLAB_FUZZ_ROUGHNESS(ABSDF),					AMixFactor);
 	SLAB_HAZINESS(ResultBSDF) =			LerpHazinessParameterBlending(BBSDF,			ABSDF,										AMixFactor);
 
+#if STRATA_COMPLEXPATH_SPECIAL
 #if STRATA_GLINTS_ENABLED
 	BSDF_SETHASGLINT(ResultBSDF,		BSDF_GETHASGLINT(BBSDF)							| BSDF_GETHASGLINT(ABSDF));
 	SLAB_GLINT_UV(ResultBSDF) =			lerp(SLAB_GLINT_UV(BBSDF),						SLAB_GLINT_UV(ABSDF), AMixFactor);
 	SLAB_GLINT_VALUE(ResultBSDF) =		lerp(SLAB_GLINT_VALUE(BBSDF),					SLAB_GLINT_VALUE(ABSDF), AMixFactor);
 	SLAB_GLINT_UVDDX(ResultBSDF) =		lerp(SLAB_GLINT_UVDDX(BBSDF),					SLAB_GLINT_UVDDX(ABSDF), AMixFactor);
 	SLAB_GLINT_UVDDY(ResultBSDF) =		lerp(SLAB_GLINT_UVDDY(BBSDF),					SLAB_GLINT_UVDDY(ABSDF), AMixFactor);
+#endif
 
 	BSDF_SETHASSPECPROFILE(ResultBSDF,	BSDF_GETHASSPECPROFILE(BBSDF) | BSDF_GETHASSPECPROFILE(ABSDF));
 	SLAB_SPECPROFILEID(ResultBSDF) = BSDF_GETHASSPECPROFILE(BBSDF) ? SLAB_SPECPROFILEID(BBSDF) : SLAB_SPECPROFILEID(ABSDF);
@@ -2621,7 +2642,7 @@ FStrataBSDF StrataVerticalLayeringParameterBlending(in FStrataBSDF TopBSDF, in F
 	BSDF_SETHASF90(ResultBSDF,			BSDF_GETHASF90(BaseBSDF)						| BSDF_GETHASF90(TopBSDF));			// F90: keep union of both even though it will be hard to get a match
 	BSDF_SETHASFUZZ(ResultBSDF,			BSDF_GETHASFUZZ(BaseBSDF)						| BSDF_GETHASFUZZ(TopBSDF));		// Fuzz: keep union of both even though it will be hard to get a match
 
-#if STRATA_GLINTS_ENABLED
+#if STRATA_COMPLEXPATH_SPECIAL
 	BSDF_SETHASGLINT(ResultBSDF, 0); // Disable for now
 	BSDF_SETHASSPECPROFILE(ResultBSDF,	0); // Disable for now
 #endif
@@ -2662,12 +2683,14 @@ FStrataBSDF StrataVerticalLayeringParameterBlending(in FStrataBSDF TopBSDF, in F
 	SLAB_FUZZ_ROUGHNESS(ResultBSDF) =	lerp(SLAB_FUZZ_ROUGHNESS(TopBSDF),				SLAB_FUZZ_ROUGHNESS(BaseBSDF),TopToBottomLerp);
 	SLAB_HAZINESS(ResultBSDF) =			LerpHazinessParameterBlending(TopBSDF,			BaseBSDF,					TopToBottomLerp);
 
+#if STRATA_COMPLEXPATH_SPECIAL
 #if STRATA_GLINTS_ENABLED
 	BSDF_SETHASGLINT(ResultBSDF,		BSDF_GETHASGLINT(TopBSDF)						| BSDF_GETHASGLINT(BaseBSDF));
 	SLAB_GLINT_UV(ResultBSDF) =			lerp(SLAB_GLINT_UV(TopBSDF),					SLAB_GLINT_UV(BaseBSDF), TopToBottomLerp);
 	SLAB_GLINT_VALUE(ResultBSDF) =		lerp(SLAB_GLINT_VALUE(TopBSDF),					SLAB_GLINT_VALUE(BaseBSDF), TopToBottomLerp);
 	SLAB_GLINT_UVDDX(ResultBSDF) =		lerp(SLAB_GLINT_UVDDX(TopBSDF),					SLAB_GLINT_UVDDX(BaseBSDF), TopToBottomLerp);
 	SLAB_GLINT_UVDDY(ResultBSDF) =		lerp(SLAB_GLINT_UVDDY(TopBSDF),					SLAB_GLINT_UVDDY(BaseBSDF), TopToBottomLerp);
+#endif
 
 	BSDF_SETHASSPECPROFILE(ResultBSDF,	BSDF_GETHASSPECPROFILE(TopBSDF) | BSDF_GETHASSPECPROFILE(BaseBSDF));
 	SLAB_SPECPROFILEID(ResultBSDF) = 	BSDF_GETHASSPECPROFILE(TopBSDF) ? SLAB_SPECPROFILEID(TopBSDF) : SLAB_SPECPROFILEID(BaseBSDF);
@@ -3840,6 +3863,10 @@ uint PackStrataHeader(uint InBSDFCount, FStrataPixelHeader InHeader)
 	HEADER_SETSHAREDLOCALBASESTYPE(Out, InHeader.SharedLocalBases.Types);
 #endif
 
+#if STRATA_COMPLEXPATH
+	HEADER_SETCOMPLEXSPECIALPATH(Out, InHeader.bUsesComplexSpecialRenderPath ? 1 : 0);
+#endif
+
 	HEADER_SETCOMMONSTATES(Out, InHeader.State);
 
 	return Out;
@@ -3862,7 +3889,7 @@ FStrataPixelHeader UnpackStrataHeaderIn(uint PackedHeader, inout FStrataAddressi
 		Out.BSDFCount = Out.IsStrataMaterial() ? 1 : 0; 
 
 	  #if STRATA_DEFERRED_SHADING
-		Out.SharedLocalBasesTypes_PackedHeader = PackedHeader;
+		Out.SharedLocalBasesTypes_PackedHeader = PackedHeader; // ?
 		Out.SharedLocalBasesIndexOffset = 0; // Unused
 
 		Out.PackedTopLayerData = StrataLoadTopLayerData(InStrataTopLayerTexture, StrataAddressing);
@@ -3872,9 +3899,11 @@ FStrataPixelHeader UnpackStrataHeaderIn(uint PackedHeader, inout FStrataAddressi
 	#if STRATA_COMPLEXPATH
 	else
 	{
-		uint SharedLocalBasesCount = HEADER_GETSHAREDLOCALBASESCOUNT(PackedHeader);
+		uint SharedLocalBasesCount = HEADER_GETSHAREDLOCALBASESCOUNT(PackedHeader); // ?
 		uint SharedLocalBasesTypes = HEADER_GETSHAREDLOCALBASESTYPE(PackedHeader);
 
+		Out.bUsesComplexSpecialRenderPath = HEADER_GETCOMPLEXSPECIALPATH(PackedHeader) > 0;
+
 	#if STRATA_DEFERRED_SHADING
 		Out.SharedLocalBasesTypes_PackedHeader = SharedLocalBasesTypes;
 		Out.PackedTopLayerData = 0;
@@ -4336,6 +4365,7 @@ FStrataBSDF UnpackStrataBSDFIn(FStrataMaterialContainer StrataBuffer, inout FStr
 		}
 
 
+	#if STRATA_COMPLEXPATH_SPECIAL
 	#if STRATA_GLINTS_ENABLED
 		if (BSDF_GETHASGLINT(OutBSDF))
 		{
@@ -4350,6 +4380,7 @@ FStrataBSDF UnpackStrataBSDFIn(FStrataMaterialContainer StrataBuffer, inout FStr
 			SLAB_GLINT_UVDDX(OutBSDF)	= UnpackFloat2FromUInt(StrataLoadUint1(StrataBuffer, StrataAddressing));
 			SLAB_GLINT_UVDDY(OutBSDF)	= UnpackFloat2FromUInt(StrataLoadUint1(StrataBuffer, StrataAddressing));
 		}
+	#endif //  // STRATA_GLINTS_ENABLED
 
 		if (BSDF_GETHASSPECPROFILE(OutBSDF))
 		{
diff --git a/Engine/Shaders/Private/Strata/StrataEvaluation.ush b/Engine/Shaders/Private/Strata/StrataEvaluation.ush
--- a/Engine/Shaders/Private/Strata/StrataEvaluation.ush
+++ b/Engine/Shaders/Private/Strata/StrataEvaluation.ush
@@ -16,7 +16,7 @@
 #define STRATA_GLINT_IS 0 
 #endif
 
-#if STRATA_GLINTS_ENABLED
+#if STRATA_COMPLEXPATH_SPECIAL
 #include "Glint/GlintThirdParty.ush"
 
 #if STRATA_GLINT_IS
@@ -259,7 +259,7 @@ half3 Strata_F_GGX(half3 F0, half3 F90, half VoH)
 //#endif
 }
 
-#if STRATA_GLINTS_ENABLED
+#if STRATA_COMPLEXPATH_SPECIAL
 // Conversion of GGX roughness to Beckmann roughness
 // Using Microfacet Distribution Function: To Change or Not to Change, That Is the Question.
 float GGXRoughnessToBeckmannRoughness(float Roughness)
@@ -445,7 +445,7 @@ FStrataEvaluateResult StrataEvaluateBSDFCommon(FStrataBSDFContext BSDFContext, F
 						float2 AlphaXY = 0;
 						GetAnisotropicRoughness(Alpha, SLAB_ANISOTROPY(BSDFContext.BSDF), AlphaXY.x, AlphaXY.y);
 
-					#if STRATA_GLINTS_ENABLED
+					#if STRATA_COMPLEXPATH_SPECIAL
 						// Glint shading only in the complex path
 						BRANCH
 						if (BSDF_GETHASGLINT(BSDFContext.BSDF))
@@ -478,7 +478,7 @@ FStrataEvaluateResult StrataEvaluateBSDFCommon(FStrataBSDFContext BSDFContext, F
 							RectLightSpec = RectGGXApproxLTC(SafeRoughness, F0, BSDFContext.N, BSDFContext.V, AreaLightContext.AreaLight.Rect, AreaLightContext.AreaLight.Texture, MeanLightWorldDirection);
 							
 							// Now combine the rectangular area light with Glints.
-						#if STRATA_GLINTS_ENABLED
+						#if STRATA_COMPLEXPATH_SPECIAL
 							BRANCH
 							if (BSDF_GETHASGLINT(BSDFContext.BSDF))
 							{
@@ -504,7 +504,7 @@ FStrataEvaluateResult StrataEvaluateBSDFCommon(FStrataBSDFContext BSDFContext, F
 						{
 							if(IntegrationType == INTEGRATION_PUNCTUAL_LIGHT)
 							{
-							#if STRATA_GLINTS_ENABLED
+							#if STRATA_COMPLEXPATH_SPECIAL
 								BRANCH
 								if (BSDF_GETHASGLINT(BSDFContext.BSDF))
 								{
@@ -522,7 +522,7 @@ FStrataEvaluateResult StrataEvaluateBSDFCommon(FStrataBSDFContext BSDFContext, F
 							{
 								const float Energy = EnergyNormalization(Alpha2Spec, BSDFContext.Context.VoH, AreaLightContext.AreaLight);
 
-							#if STRATA_GLINTS_ENABLED
+							#if STRATA_COMPLEXPATH_SPECIAL
 								BRANCH
 								if (BSDF_GETHASGLINT(BSDFContext.BSDF))
 								{
@@ -682,7 +682,7 @@ FStrataEvaluateResult StrataEvaluateBSDFCommon(FStrataBSDFContext BSDFContext, F
 			#endif
 
 			// Lighting LUT
-			#if STRATA_FASTPATH==0 && STRATA_GLINTS_ENABLED
+			#if STRATA_FASTPATH==0 && STRATA_COMPLEXPATH_SPECIAL
 			BRANCH
 			if (BSDF_GETHASSPECPROFILE(BSDFContext.BSDF))
 			{
@@ -1503,7 +1503,7 @@ FStrataEnvLightResult StrataEvaluateForEnvLight(FStrataBSDFContext BSDFContext,
 				StrataEnvLightResult.SpecularColor = F0;
 
 				// STRATA_TODO environment light importance sampling
-			#if STRATA_GLINTS_ENABLED
+			#if STRATA_COMPLEXPATH_SPECIAL
 				// Glints
 				BRANCH
 				if (BSDF_GETHASGLINT(BSDFContext.BSDF))
diff --git a/Engine/Shaders/Private/Strata/StrataExport.ush b/Engine/Shaders/Private/Strata/StrataExport.ush
--- a/Engine/Shaders/Private/Strata/StrataExport.ush
+++ b/Engine/Shaders/Private/Strata/StrataExport.ush
@@ -190,11 +190,16 @@ void PackStrataOut(
 						// Simple and single bsdf do not allow weights other than 1
 						const bool LuminanceWeightEqualOne = all(BSDF.LuminanceWeightV == 1.0f);
 
-						#if STRATA_GLINTS
-						const bool bForceComplexMaterial = BSDF_GETHASANISOTROPY(BSDF) || BSDF_GETHASGLINT(BSDF) || BSDF_GETHASSPECPROFILE(BSDF);
+						#if STRATA_COMPLEXPATH_SPECIAL
+						#if STRATA_GLINTS_ENABLED
+						const bool bForceComplexSpecialPath = BSDF_GETHASSPECPROFILE(BSDF) || BSDF_GETHASGLINT(BSDF);
 						#else
-						const bool bForceComplexMaterial = BSDF_GETHASANISOTROPY(BSDF);
+						const bool bForceComplexSpecialPath = BSDF_GETHASSPECPROFILE(BSDF);
 						#endif
+						#else
+						const bool bForceComplexSpecialPath = false;
+						#endif
+						const bool bForceComplexMaterial = BSDF_GETHASANISOTROPY(BSDF) || bForceComplexSpecialPath;
 
 						// Update simple material compatibility
 						bIsSimpleMaterial = bIsSimpleMaterial && IsStrataSlabFastPathCompatible(BSDF) && LuminanceWeightEqualOne && !bForceComplexMaterial;
@@ -226,6 +231,8 @@ void PackStrataOut(
 
 						const bool bSSSMask = BSDF.bIsBottom;	// SSS and Thin can only be on the bottom layer, so ignoring bIsSimpleVolume here.
 						StrataPixelHeader.SetHasSubsurface(bSSSMask && bHasSSS && !bIsSimpleVolume);
+
+						StrataPixelHeader.SetUsesComplexSpecialRenderPath(bForceComplexSpecialPath);
 						break;
 					}
 
@@ -791,6 +798,7 @@ void PackStrataOut(
 								STRATA_STORE_UINT1(PackColorLinearToGamma2AlphaLinear(float4(BSDF.TransmittanceAboveAlongN, BSDF.CoverageAboveAlongN)));
 							}
 
+							#if STRATA_COMPLEXPATH_SPECIAL
 							#if STRATA_GLINTS_ENABLED
 							if (BSDF_GETHASGLINT(BSDF))
 							{
@@ -805,12 +813,13 @@ void PackStrataOut(
 								STRATA_STORE_UINT1(PackFloat2ToUInt(SLAB_GLINT_UVDDX(BSDF)));
 								STRATA_STORE_UINT1(PackFloat2ToUInt(SLAB_GLINT_UVDDY(BSDF)));
 							}
+							#endif // STRATA_GLINTS_ENABLED
 
 							if (BSDF_GETHASSPECPROFILE(BSDF))
 							{
 								STRATA_STORE_UINT1(PackSpecularProfile(SLAB_SPECPROFILEID(BSDF)));
 							}
-							#endif // STRATA_GLINTS
+							#endif // STRATA_COMPLEXPATH_SPECIAL
 
 							// 8-28 bytes
 						}
diff --git a/Engine/Shaders/Private/Strata/StrataMaterialClassification.usf b/Engine/Shaders/Private/Strata/StrataMaterialClassification.usf
--- a/Engine/Shaders/Private/Strata/StrataMaterialClassification.usf
+++ b/Engine/Shaders/Private/Strata/StrataMaterialClassification.usf
@@ -108,6 +108,7 @@ void TileMainCS(uint2 DispatchThreadId : SV_DispatchThreadID, uint LinearIndex :
 #endif
 
 	// Pixels outside of the view area are considered simple to enable screen borders to receive the simple permutation when not aligned to shader group size.
+	bool bContainsComplexSpecialMaterial			= false;
 	bool bContainsComplexMaterial					= false;
 	bool bContainsSimpleMaterial					= false;
 	bool bContainsSingleMaterial					= false;
@@ -130,7 +131,9 @@ void TileMainCS(uint2 DispatchThreadId : SV_DispatchThreadID, uint LinearIndex :
 		bContainsStrataMaterial  = StrataPixelHeader.BSDFCount > 0;
 		bContainsSimpleMaterial  = bIsSimple;
 		bContainsSingleMaterial  = bIsSingle;
-		bContainsComplexMaterial = !bIsSingle && !bIsSimple;
+		bContainsComplexMaterial		= !bIsSingle && !bIsSimple && !StrataPixelHeader.bUsesComplexSpecialRenderPath;
+		bContainsComplexSpecialMaterial = !bIsSingle && !bIsSimple &&  StrataPixelHeader.bUsesComplexSpecialRenderPath;
+
 		bContainsScreenSpaceSubsurfaceScattering = StrataPixelHeader.HasSubsurface();
 
 	#if PERMUTATION_DECAL
@@ -212,6 +215,7 @@ void TileMainCS(uint2 DispatchThreadId : SV_DispatchThreadID, uint LinearIndex :
 	const bool bTileContainsSimple							= WaveActiveAnyTrue(bContainsSimpleMaterial);
 	const bool bTileContainsSingle							= WaveActiveAnyTrue(bContainsSingleMaterial);
 	const bool bTileContainsComplex							= WaveActiveAnyTrue(bContainsComplexMaterial);
+	const bool bTileContainsComplexSpecial					= WaveActiveAnyTrue(bContainsComplexSpecialMaterial);
 	const bool bTileContainsOpaqueRoughRefraction			= WaveActiveAnyTrue(bContainsOpaqueRoughRefraction);
 	const bool bTileContainsScreenSpaceSubsurfaceScattering	= WaveActiveAnyTrue(bContainsScreenSpaceSubsurfaceScattering);
 	const bool bTileContainsDecals							= WaveActiveAnyTrue(bContainsDecals);
@@ -222,10 +226,10 @@ void TileMainCS(uint2 DispatchThreadId : SV_DispatchThreadID, uint LinearIndex :
 		| (bContainsSimpleMaterial					? 0x2u	: 0u)
 		| (bContainsSingleMaterial					? 0x4u	: 0u) 
 		| (bContainsComplexMaterial					? 0x8u	: 0u) 
-		| (bContainsOpaqueRoughRefraction			? 0x10u : 0u)
-		| (bContainsScreenSpaceSubsurfaceScattering	? 0x20u : 0u)
-		| (bContainsScreenSpaceSubsurfaceScattering	? 0x20u : 0u)
-		| (bContainsDecals							? 0x40u : 0u);
+		| (bContainsComplexSpecialMaterial			? 0x10u : 0u)
+		| (bContainsOpaqueRoughRefraction			? 0x20u : 0u)
+		| (bContainsScreenSpaceSubsurfaceScattering	? 0x40u : 0u)
+		| (bContainsDecals							? 0x80u : 0u);
 
 	GroupMemoryBarrierWithGroupSync();
 	if (LinearIndex < 32)
@@ -261,16 +265,23 @@ void TileMainCS(uint2 DispatchThreadId : SV_DispatchThreadID, uint LinearIndex :
 	const bool bTileContainsSimple							= (FinalTileFlags & 0x2u) > 0;
 	const bool bTileContainsSingle							= (FinalTileFlags & 0x4u) > 0;
 	const bool bTileContainsComplex							= (FinalTileFlags & 0x8u) > 0;
-	const bool bTileContainsOpaqueRoughRefraction			= (FinalTileFlags & 0x10u) > 0;
-	const bool bTileContainsScreenSpaceSubsurfaceScattering	= (FinalTileFlags & 0x20u) > 0;
-	const bool bTileContainsDecals							= (FinalTileFlags & 0x40u) > 0;
+	const bool bTileContainsComplexSpecial					= (FinalTileFlags & 0x10u) > 0;
+	const bool bTileContainsOpaqueRoughRefraction			= (FinalTileFlags & 0x20u) > 0;
+	const bool bTileContainsScreenSpaceSubsurfaceScattering	= (FinalTileFlags & 0x40u) > 0;
+	const bool bTileContainsDecals							= (FinalTileFlags & 0x80u) > 0;
 #endif // PERMUTATION_WAVE_OPS
 
 	if (LinearIndex < 1 && bTileContainsStrata)
 	{
 		uint EncodedTile = StrataPackTile(GroupId.xy);
 
-		if (bTileContainsComplex)
+		if (bTileContainsComplexSpecial)
+		{
+			uint WriteToIndex;
+			InterlockedAdd(TileDrawIndirectDataBufferUAV[GetStrataTileTypeDrawIndirectArgOffset_DWord(STRATA_TILE_TYPE_COMPLEX_SPECIAL) + 1], 1, WriteToIndex);
+			TileListBufferUAV[GetTileListBufferOffsets(STRATA_TILE_TYPE_COMPLEX_SPECIAL) + WriteToIndex] = EncodedTile;
+		}
+		else  if (bTileContainsComplex)
 		{
 			uint WriteToIndex;
 			InterlockedAdd(TileDrawIndirectDataBufferUAV[GetStrataTileTypeDrawIndirectArgOffset_DWord(STRATA_TILE_TYPE_COMPLEX) + 1], 1, WriteToIndex);
diff --git a/Engine/Shaders/Private/Strata/StrataVisualize.usf b/Engine/Shaders/Private/Strata/StrataVisualize.usf
--- a/Engine/Shaders/Private/Strata/StrataVisualize.usf
+++ b/Engine/Shaders/Private/Strata/StrataVisualize.usf
@@ -10,7 +10,8 @@
 #endif
 
 #define STRATA_SSS_MATERIAL_OVERRIDE 0
-#define STRATA_SSS_TRANSMISSION 0  // Disable for now, as the SSS profile texture need to be bound
+#define STRATA_SSS_TRANSMISSION 0				// Disable for now, as the SSS profile texture need to be bound
+#define STRATA_COMPLEX_SPECIALPATH_ENABLED 1	// Allow to read data from complex special path
 #include "/Engine/Private/Strata/StrataEvaluation.ush"
 #include "/Engine/Private/Strata/StrataTile.ush"
 
@@ -91,13 +92,14 @@ RWBuffer<uint> RWPositionOffsetBuffer;
 
 void PrintPixelType(inout FShaderPrintContext Ctx, in FStrataPixelHeader Header, FFontColor InColor)
 {
-	if (Header.IsSimpleMaterial())			Print(Ctx, TEXT("Simple  "), InColor);
-	else if (Header.IsSingleMaterial())		Print(Ctx, TEXT("Single  "), InColor);
-	else if (Header.IsComplexMaterial())	Print(Ctx, TEXT("Complex "), InColor);
-	else if (Header.IsSingleLayerWater())	Print(Ctx, TEXT("SLWater "), InColor);
-	else if (Header.IsHair())				Print(Ctx, TEXT("Hair    "), InColor);
-	else if (Header.IsEye())				Print(Ctx, TEXT("Eye     "), InColor);
-	else									Print(Ctx, TEXT("Error - Unkown pixel type"), FontRed);
+	if (Header.IsSimpleMaterial())					Print(Ctx, TEXT("Simple  "), InColor);
+	else if (Header.IsSingleMaterial())				Print(Ctx, TEXT("Single  "), InColor);
+	else if (Header.bUsesComplexSpecialRenderPath)	Print(Ctx, TEXT("Complex Special "), InColor);
+	else if (Header.IsComplexMaterial())			Print(Ctx, TEXT("Complex "), InColor);
+	else if (Header.IsSingleLayerWater())			Print(Ctx, TEXT("SLWater "), InColor);
+	else if (Header.IsHair())						Print(Ctx, TEXT("Hair    "), InColor);
+	else if (Header.IsEye())						Print(Ctx, TEXT("Eye     "), InColor);
+	else											Print(Ctx, TEXT("Error - Unkown pixel type"), FontRed);
 }
 
 void PrintSSSType(inout FShaderPrintContext Ctx, uint SSSType, bool bThin, FFontColor InColor)
@@ -328,7 +330,7 @@ void PrintBSDF(inout FShaderPrintContext Context, inout FStrataAddressing Strata
 				Print(Context, TEXT(" Fuzz Color  "), FontBSDFPropName);	Print(Context, SLAB_FUZZ_COLOR(BSDF), FontBSDFPropValu);				Newline(Context, RectMax);
 				Print(Context, TEXT(" Fuzz Rough. "), FontBSDFPropName);	Print(Context, SLAB_FUZZ_ROUGHNESS(BSDF), FontBSDFPropValu);			Newline(Context, RectMax);
 			}
-#if STRATA_GLINTS
+#if STRATA_COMPLEXPATH_SPECIAL
 			if (BSDF_GETHASGLINT(BSDF))
 			{
 				Print(Context, TEXT(" Glint Value "), FontBSDFPropName);	Print(Context, SLAB_GLINT_VALUE(BSDF), FontBSDFPropValu);				Newline(Context, RectMax);
@@ -1670,6 +1672,10 @@ void MainCS(uint3 DispatchThreadId : SV_DispatchThreadID)
 		Print(Context, GetTileCount(STRATA_TILE_TYPE_COMPLEX), ValueColor);
 		Newline(Context);
 
+		Print(Context, TEXT(" # Complex Special tiles: "), LegendColor);
+		Print(Context, GetTileCount(STRATA_TILE_TYPE_COMPLEX_SPECIAL), ValueColor);
+		Newline(Context);
+
 		Newline(Context);
 
 		// Rough tiles
diff --git a/Engine/Shaders/Shared/StrataDefinitions.h b/Engine/Shaders/Shared/StrataDefinitions.h
--- a/Engine/Shaders/Shared/StrataDefinitions.h
+++ b/Engine/Shaders/Shared/StrataDefinitions.h
@@ -8,7 +8,7 @@
 #pragma once
 
 // Change this to force recompilation of all strata dependent shaders (use https://www.random.org/cgi-bin/randbyte?nbytes=4&format=h)
-#define STRATA_SHADER_VERSION 0xbd4a70e7
+#define STRATA_SHADER_VERSION 0x932681EC
 
 
 
@@ -57,12 +57,13 @@
 #define STRATA_TILE_TYPE_SIMPLE						0
 #define STRATA_TILE_TYPE_SINGLE						1
 #define STRATA_TILE_TYPE_COMPLEX					2
-#define STRATA_TILE_TYPE_ROUGH_REFRACT				3
-#define STRATA_TILE_TYPE_ROUGH_REFRACT_SSS_WITHOUT	4
-#define STRATA_TILE_TYPE_DECAL_SIMPLE				5
-#define STRATA_TILE_TYPE_DECAL_SINGLE				6
-#define STRATA_TILE_TYPE_DECAL_COMPLEX				7
-#define STRATA_TILE_TYPE_COUNT						8
+#define STRATA_TILE_TYPE_COMPLEX_SPECIAL			3
+#define STRATA_TILE_TYPE_ROUGH_REFRACT				4
+#define STRATA_TILE_TYPE_ROUGH_REFRACT_SSS_WITHOUT	5
+#define STRATA_TILE_TYPE_DECAL_SIMPLE				6
+#define STRATA_TILE_TYPE_DECAL_SINGLE				7
+#define STRATA_TILE_TYPE_DECAL_COMPLEX				8
+#define STRATA_TILE_TYPE_COUNT						9
 
 // sizeof(FRHIDrawIndirectParameters) = 4 uints = 16 bytes
 #define GetStrataTileTypeDrawIndirectArgOffset_Byte(x)  (x * 16)
diff --git a/Engine/Source/Editor/MaterialEditor/Private/SMaterialEditorStrataWidget.cpp b/Engine/Source/Editor/MaterialEditor/Private/SMaterialEditorStrataWidget.cpp
--- a/Engine/Source/Editor/MaterialEditor/Private/SMaterialEditorStrataWidget.cpp
+++ b/Engine/Source/Editor/MaterialEditor/Private/SMaterialEditorStrataWidget.cpp
@@ -224,6 +224,9 @@ void SMaterialEditorStrataWidget::Tick(const FGeometry& AllottedGeometry, const
 				case 2:
 					MaterialDescription += FString::Printf(TEXT("Material complexity          = COMPLEX\r\n"));
 					break;
+				case 3:
+					MaterialDescription += FString::Printf(TEXT("Material complexity          = COMPLEX SPECIAL\r\n"));
+					break;
 				default:
 					MaterialDescription += FString::Printf(TEXT("Material complexity          = UNKOWN => ERROR!\r\n"));
 				}
diff --git a/Engine/Source/Runtime/Engine/Classes/Materials/MaterialRelevance.h b/Engine/Source/Runtime/Engine/Classes/Materials/MaterialRelevance.h
--- a/Engine/Source/Runtime/Engine/Classes/Materials/MaterialRelevance.h
+++ b/Engine/Source/Runtime/Engine/Classes/Materials/MaterialRelevance.h
@@ -19,6 +19,7 @@ struct ENGINE_API FMaterialRelevance
 			uint16 ShadingModelMask;
 			uint8 StrataUintPerPixel;
 			uint8 StrataBSDFCountMask;
+			uint8 bUsesComplexSpecialRenderPath : 1;
 			uint8 bOpaque : 1;
 			uint8 bMasked : 1;
 			uint8 bDistortion : 1;
diff --git a/Engine/Source/Runtime/Engine/Private/Materials/HLSLMaterialTranslator.cpp b/Engine/Source/Runtime/Engine/Private/Materials/HLSLMaterialTranslator.cpp
--- a/Engine/Source/Runtime/Engine/Private/Materials/HLSLMaterialTranslator.cpp
+++ b/Engine/Source/Runtime/Engine/Private/Materials/HLSLMaterialTranslator.cpp
@@ -2296,6 +2296,7 @@ void FHLSLMaterialTranslator::GetMaterialEnvironment(EShaderPlatform InPlatform,
 			OutEnvironment.SetDefine(TEXT("STRATA_SINGLEPATH"), StrataCtx.bStrataMaterialIsSingle ? TEXT("1") : TEXT("0"));
 			OutEnvironment.SetDefine(TEXT("STRATA_FASTPATH"), StrataCtx.bStrataMaterialIsSingle ? TEXT("0") : (StrataCtx.bStrataMaterialIsSimple ? TEXT("1") : TEXT("0")));
 			OutEnvironment.SetDefine(TEXT("STRATA_CLAMPED_BSDF_COUNT"), StrataCtx.StrataMaterialBSDFCount);
+			OutEnvironment.SetDefine(TEXT("STRATA_COMPLEX_SPECIALPATH_ENABLED"), MaterialCompilationOutput.StrataMaterialCompilationOutput.bUsesComplexSpecialRenderPath ? TEXT("1") : TEXT("0"));
 		}
 
 
@@ -11559,6 +11560,7 @@ bool FHLSLMaterialTranslator::FStrataCompilationContext::StrataGenerateDerivedMa
 			bStrataMaterialIsSingle = StrataMaterialBSDFCount == 1;
 			bool bIsFastWaterPath = false;
 			bool bCustomEncoding = false;
+			bool bUsesComplexSpecialRenderPath = false;
 			for (auto& It : StrataMaterialExpressionRegisteredOperators)
 			{
 				if (It.IsDiscarded())
@@ -11579,6 +11581,7 @@ bool FHLSLMaterialTranslator::FStrataCompilationContext::StrataGenerateDerivedMa
 					{
 						bStrataMaterialIsSimple = bStrataMaterialIsSimple && !bMayHaveColoredWeight && !It.bBSDFHasAnisotropy && !It.bBSDFHasEdgeColor && !It.bBSDFHasFuzz && !It.bBSDFHasSecondRoughnessOrSimpleClearCoat && !It.bBSDFHasMFPPluggedIn && !It.bBSDFHasSSS && !It.bBSDFHasGlint && !It.bBSDFHasSpecularProfile;
 						bStrataMaterialIsSingle = bStrataMaterialIsSingle && !bMayHaveColoredWeight && !It.bBSDFHasAnisotropy && !It.bBSDFHasGlint && !It.bBSDFHasSpecularProfile;
+						bUsesComplexSpecialRenderPath |= It.bBSDFHasGlint || It.bBSDFHasSpecularProfile;
 						break;
 					}
 					case STRATA_BSDF_TYPE_HAIR:
@@ -11790,9 +11793,10 @@ bool FHLSLMaterialTranslator::FStrataCompilationContext::StrataGenerateDerivedMa
 			if (CompilationContextIndex == EStrataCompilationContext::SCC_Default)
 			{
 				// Only write those data for the default material
-				Compiler->MaterialCompilationOutput.StrataMaterialCompilationOutput.StrataMaterialType = bStrataMaterialIsSimple ? 0 : bStrataMaterialIsSingle ? 1 : 2;
+				Compiler->MaterialCompilationOutput.StrataMaterialCompilationOutput.StrataMaterialType = bStrataMaterialIsSimple ? 0 : (bStrataMaterialIsSingle ? 1 : (bUsesComplexSpecialRenderPath ? 3 : 2));
 				Compiler->MaterialCompilationOutput.StrataMaterialCompilationOutput.StrataBSDFCount = StrataMaterialBSDFCount;
 				Compiler->MaterialCompilationOutput.StrataMaterialCompilationOutput.StrataUintPerPixel = uint8(FMath::Clamp(RequestedSizeInUint, 0u, 0xFF));
+				Compiler->MaterialCompilationOutput.StrataMaterialCompilationOutput.bUsesComplexSpecialRenderPath = bUsesComplexSpecialRenderPath;
 
 #if WITH_EDITOR
 				Compiler->MaterialCompilationOutput.StrataMaterialCompilationOutput.SharedLocalBasesCount = 0; // FinalUsedSharedLocalBasesCount is not valid yet
diff --git a/Engine/Source/Runtime/Engine/Private/Materials/MaterialInterface.cpp b/Engine/Source/Runtime/Engine/Private/Materials/MaterialInterface.cpp
--- a/Engine/Source/Runtime/Engine/Private/Materials/MaterialInterface.cpp
+++ b/Engine/Source/Runtime/Engine/Private/Materials/MaterialInterface.cpp
@@ -481,6 +481,7 @@ FMaterialRelevance UMaterialInterface::GetRelevance_Internal(const UMaterial* Ma
 		const uint8 StrataBSDFCount = FMath::Max(MaterialResource->MaterialGetStrataBSDFCount_GameThread(), uint8(1u));
 		const uint8 StrataBSDFCountMask = 1u << uint8(FMath::Min(StrataBSDFCount - 1, 8));
 		const uint8 StrataUintPerPixel = FMath::Max(MaterialResource->MaterialGetStrataUintPerPixel_GameThread(), uint8(1u));
+		const uint8 bUsesComplexSpecialRenderPath = MaterialResource->MaterialGetStrataUsesComplexSpecialRenderPath_GameThread();
 
 		MaterialRelevance.bOpaque = !bIsTranslucent;
 		MaterialRelevance.bMasked = IsMasked();
@@ -508,6 +509,8 @@ FMaterialRelevance UMaterialInterface::GetRelevance_Internal(const UMaterial* Ma
 		MaterialRelevance.bUsesAnisotropy = bUsesAnisotropy;
 		MaterialRelevance.StrataBSDFCountMask = StrataBSDFCountMask;
 		MaterialRelevance.StrataUintPerPixel = StrataUintPerPixel;
+		MaterialRelevance.bUsesComplexSpecialRenderPath = bUsesComplexSpecialRenderPath;
+
 		return MaterialRelevance;
 	}
 	else
diff --git a/Engine/Source/Runtime/Engine/Private/Materials/MaterialShared.cpp b/Engine/Source/Runtime/Engine/Private/Materials/MaterialShared.cpp
--- a/Engine/Source/Runtime/Engine/Private/Materials/MaterialShared.cpp
+++ b/Engine/Source/Runtime/Engine/Private/Materials/MaterialShared.cpp
@@ -1155,6 +1155,17 @@ uint8 FMaterial::MaterialGetStrataUintPerPixel_RenderThread() const
 	return RenderingThreadShaderMap ? RenderingThreadShaderMap->GetStrataUintPerPixel() : false;
 }
 
+bool FMaterial::MaterialGetStrataUsesComplexSpecialRenderPath_GameThread() const
+{
+	return GameThreadShaderMap ? GameThreadShaderMap->GetStrataUsesComplexSpecialRenderPath() : false;
+}
+
+bool FMaterial::MaterialGetStrataUsesComplexSpecialRenderPath_RenderThread() const
+{
+	check(IsInParallelRenderingThread());
+	return RenderingThreadShaderMap ? RenderingThreadShaderMap->GetStrataUsesComplexSpecialRenderPath() : false;
+}
+
 void FMaterial::SetGameThreadShaderMap(FMaterialShaderMap* InMaterialShaderMap)
 {
 	checkSlow(IsInGameThread() || IsInAsyncLoadingThread());
diff --git a/Engine/Source/Runtime/Engine/Private/ShaderCompiler/ShaderCompiler.cpp b/Engine/Source/Runtime/Engine/Private/ShaderCompiler/ShaderCompiler.cpp
--- a/Engine/Source/Runtime/Engine/Private/ShaderCompiler/ShaderCompiler.cpp
+++ b/Engine/Source/Runtime/Engine/Private/ShaderCompiler/ShaderCompiler.cpp
@@ -6662,9 +6662,6 @@ void GlobalBeginCompileShader(
 
 			const bool bStrataDBufferPass = Strata::IsDBufferPassEnabled(Target.GetPlatform());
 			Input.Environment.SetDefine(TEXT("STRATA_USE_DBUFFER_PASS"), bStrataDBufferPass ? 1 : 0);
-
-			const bool bStrataGlints = Strata::IsGlintEnabled();
-			Input.Environment.SetDefine(TEXT("STRATA_GLINTS"), bStrataGlints ? 1 : 0);
 		}
 
 		const bool bStrataBackCompatibility = bStrata && Strata::IsBackCompatibilityEnabled();
diff --git a/Engine/Source/Runtime/Engine/Public/MaterialShared.h b/Engine/Source/Runtime/Engine/Public/MaterialShared.h
--- a/Engine/Source/Runtime/Engine/Public/MaterialShared.h
+++ b/Engine/Source/Runtime/Engine/Public/MaterialShared.h
@@ -1371,6 +1371,7 @@ public:
 	uint8 GetStrataMaterialType() const { return GetStrataMaterialCompilationOutput().StrataMaterialType; }
 	uint8 GetStrataBSDFCount() const { return GetStrataMaterialCompilationOutput().StrataBSDFCount; }
 	uint8 GetStrataUintPerPixel() const { return GetStrataMaterialCompilationOutput().StrataUintPerPixel; }
+	bool GetStrataUsesComplexSpecialRenderPath() const { return GetStrataMaterialCompilationOutput().bUsesComplexSpecialRenderPath; }
 	
 #if WITH_EDITOR
 	uint32 GetNumUsedUVScalars() const { return GetContent()->MaterialCompilationOutput.NumUsedUVScalars; }
@@ -2126,10 +2127,14 @@ public:
 	ENGINE_API uint8 MaterialGetStrataBSDFCount_GameThread() const;
 	ENGINE_API uint8 MaterialGetStrataBSDFCount_RenderThread() const;
 
-	/** Get Strata material uint count oer pixel. */
+	/** Get Strata material uint count per pixel. */
 	ENGINE_API uint8 MaterialGetStrataUintPerPixel_GameThread() const;
 	ENGINE_API uint8 MaterialGetStrataUintPerPixel_RenderThread() const;
 
+	/** Get Strata material special path requirement (for more expenssive features such as Glints or SpecularLUT). */
+	ENGINE_API bool MaterialGetStrataUsesComplexSpecialRenderPath_GameThread() const;
+	ENGINE_API bool MaterialGetStrataUsesComplexSpecialRenderPath_RenderThread() const;
+
 	class FMaterialShaderMap* GetGameThreadShaderMap() const 
 	{ 
 		checkSlow(IsInGameThread() || IsInParallelGameThread() || IsInAsyncLoadingThread());
diff --git a/Engine/Source/Runtime/Engine/Public/Rendering/StrataMaterialShared.cpp b/Engine/Source/Runtime/Engine/Public/Rendering/StrataMaterialShared.cpp
--- a/Engine/Source/Runtime/Engine/Public/Rendering/StrataMaterialShared.cpp
+++ b/Engine/Source/Runtime/Engine/Public/Rendering/StrataMaterialShared.cpp
@@ -9,6 +9,7 @@ FStrataMaterialCompilationOutput::FStrataMaterialCompilationOutput()
 	: StrataMaterialType(0)
 	, StrataBSDFCount(0)
 	, StrataUintPerPixel(0)
+	, bUsesComplexSpecialRenderPath(0)
 #if WITH_EDITOR
 	, StrataMaterialDescription()
 	, SharedLocalBasesCount(0)
diff --git a/Engine/Source/Runtime/Engine/Public/Rendering/StrataMaterialShared.h b/Engine/Source/Runtime/Engine/Public/Rendering/StrataMaterialShared.h
--- a/Engine/Source/Runtime/Engine/Public/Rendering/StrataMaterialShared.h
+++ b/Engine/Source/Runtime/Engine/Public/Rendering/StrataMaterialShared.h
@@ -108,6 +108,9 @@ public:
 	/** Strata uint per pixel, at compile time (0-255) */
 	LAYOUT_FIELD(uint8, StrataUintPerPixel);
 
+	/** If Glints or SpecularLUT are used, we need to go the special more expenssive render path */
+	LAYOUT_BITFIELD(uint8, bUsesComplexSpecialRenderPath, 1);
+
 	////
 	//// The following data is only needed when compiling with the editor.
 	////
diff --git a/Engine/Source/Runtime/Renderer/Private/ClusteredDeferredShadingPass.cpp b/Engine/Source/Runtime/Renderer/Private/ClusteredDeferredShadingPass.cpp
--- a/Engine/Source/Runtime/Renderer/Private/ClusteredDeferredShadingPass.cpp
+++ b/Engine/Source/Runtime/Renderer/Private/ClusteredDeferredShadingPass.cpp
@@ -95,7 +95,7 @@ class FClusteredShadingPS : public FGlobalShader
 
 	class FVisualizeLightCullingDim : SHADER_PERMUTATION_BOOL("VISUALIZE_LIGHT_CULLING");
 	class FHairStrandsLighting : SHADER_PERMUTATION_BOOL("USE_HAIR_LIGHTING");
-	class FStrataTileType : SHADER_PERMUTATION_INT("STRATA_TILETYPE", 3);
+	class FStrataTileType : SHADER_PERMUTATION_INT("STRATA_TILETYPE", 4);
 	using FPermutationDomain = TShaderPermutationDomain<FVisualizeLightCullingDim, FHairStrandsLighting, FStrataTileType>;
 
 	BEGIN_SHADER_PARAMETER_STRUCT(FParameters, )
@@ -290,6 +290,18 @@ void FDeferredShadingSceneRenderer::AddClusteredDeferredShadingPass(
 			{
 				StrataSceneData = &Scene->StrataSceneData;
 
+				InternalAddClusteredDeferredShadingPass(
+					GraphBuilder,
+					View,
+					SceneTextures,
+					SortedLightsSet,
+					EClusterPassInputType::Strata,
+					EStrataTileType::EComplexSpecial,
+					ShadowMaskBits,
+					VirtualShadowMapArray,
+					nullptr,
+					StrataSceneData);
+
 				InternalAddClusteredDeferredShadingPass(
 					GraphBuilder,
 					View,
diff --git a/Engine/Source/Runtime/Renderer/Private/IndirectLightRendering.cpp b/Engine/Source/Runtime/Renderer/Private/IndirectLightRendering.cpp
--- a/Engine/Source/Runtime/Renderer/Private/IndirectLightRendering.cpp
+++ b/Engine/Source/Runtime/Renderer/Private/IndirectLightRendering.cpp
@@ -123,7 +123,7 @@ class FDiffuseIndirectCompositePS : public FGlobalShader
 	class FApplyDiffuseIndirectDim : SHADER_PERMUTATION_INT("DIM_APPLY_DIFFUSE_INDIRECT", 4);
 	class FUpscaleDiffuseIndirectDim : SHADER_PERMUTATION_BOOL("DIM_UPSCALE_DIFFUSE_INDIRECT");
 	class FScreenBentNormal : SHADER_PERMUTATION_BOOL("DIM_SCREEN_BENT_NORMAL");
-	class FStrataTileType : SHADER_PERMUTATION_INT("STRATA_TILETYPE", 3);
+	class FStrataTileType : SHADER_PERMUTATION_INT("STRATA_TILETYPE", 4);
 	class FEnableDualSrcBlending : SHADER_PERMUTATION_BOOL("ENABLE_DUAL_SRC_BLENDING");
 
 	using FPermutationDomain = TShaderPermutationDomain<FApplyDiffuseIndirectDim, FUpscaleDiffuseIndirectDim, FScreenBentNormal, FStrataTileType, FEnableDualSrcBlending>;
@@ -283,7 +283,7 @@ class FReflectionEnvironmentSkyLightingPS : public FGlobalShader
 	class FSkyShadowing : SHADER_PERMUTATION_BOOL("APPLY_SKY_SHADOWING");
 	class FRayTracedReflections : SHADER_PERMUTATION_BOOL("RAY_TRACED_REFLECTIONS");
 	class FLumenStandaloneReflections : SHADER_PERMUTATION_BOOL("LUMEN_STANDALONE_REFLECTIONS");
-	class FStrataTileType : SHADER_PERMUTATION_INT("STRATA_TILETYPE", 3);
+	class FStrataTileType : SHADER_PERMUTATION_INT("STRATA_TILETYPE", 4);
 
 	using FPermutationDomain = TShaderPermutationDomain<
 		FHasBoxCaptures,
@@ -343,7 +343,7 @@ class FReflectionEnvironmentSkyLightingPS : public FGlobalShader
 		PermutationVector.Set<FStrataTileType>(0);
 		if (Strata::IsStrataEnabled())
 		{
-			check(TileType <= EStrataTileType::EComplex);
+			check(TileType <= EStrataTileType::EComplexSpecial);
 			PermutationVector.Set<FStrataTileType>(TileType);
 		}
 		return RemapPermutation(PermutationVector);
@@ -1506,6 +1506,10 @@ void FDeferredShadingSceneRenderer::RenderDiffuseIndirectAndAmbientOcclusion(
 		{
 			if (Strata::IsStrataEnabled())
 			{
+				if (Strata::GetStrataUsesComplexSpecialPath(View))
+				{
+					ApplyDiffuseIndirect(EStrataTileType::EComplexSpecial);
+				}
 				ApplyDiffuseIndirect(EStrataTileType::EComplex);
 				ApplyDiffuseIndirect(EStrataTileType::ESingle);
 				ApplyDiffuseIndirect(EStrataTileType::ESimple);
@@ -1615,6 +1619,10 @@ void FDeferredShadingSceneRenderer::RenderDiffuseIndirectAndAmbientOcclusion(
 
 			if (Strata::IsStrataEnabled())
 			{
+				if (Strata::GetStrataUsesComplexSpecialPath(View))
+				{
+					ApplyAmbientCubemapComposite(EStrataTileType::EComplexSpecial);
+				}
 				ApplyAmbientCubemapComposite(EStrataTileType::EComplex);
 				ApplyAmbientCubemapComposite(EStrataTileType::ESingle);
 				ApplyAmbientCubemapComposite(EStrataTileType::ESimple);
@@ -1760,7 +1768,7 @@ static void AddSkyReflectionPass(
 	const bool bStrataEnabled = Strata::IsStrataEnabled();
 	if (bStrataEnabled)
 	{
-		check(StrataTileMaterialType <= EStrataTileType::EComplex);
+		check(StrataTileMaterialType <= EStrataTileType::EComplexSpecial);
 		PassParameters->VS = Strata::SetTileParameters(GraphBuilder, View, StrataTileMaterialType, StrataTilePrimitiveType);
 		ClearUnusedGraphResources(StrataTilePassVertexShader, &PassParameters->VS);
 	}
@@ -2086,6 +2094,21 @@ void FDeferredShadingSceneRenderer::RenderDeferredReflectionsAndSkyLighting(
 
 			if (Strata::IsStrataEnabled())
 			{
+				AddSkyReflectionPass(
+					GraphBuilder,
+					View,
+					Scene,
+					SceneTextures,
+					DynamicBentNormalAOTexture,
+					ReflectionsColor,
+					RayTracingReflectionOptions,
+					SceneTextureParameters,
+					bSkyLight,
+					bDynamicSkyLight,
+					bApplySkyShadowing,
+					bLumenStandaloneReflections,
+					EStrataTileType::EComplexSpecial);
+
 				AddSkyReflectionPass(
 					GraphBuilder,
 					View,
diff --git a/Engine/Source/Runtime/Renderer/Private/LightRendering.cpp b/Engine/Source/Runtime/Renderer/Private/LightRendering.cpp
--- a/Engine/Source/Runtime/Renderer/Private/LightRendering.cpp
+++ b/Engine/Source/Runtime/Renderer/Private/LightRendering.cpp
@@ -730,7 +730,7 @@ class FDeferredLightPS : public FGlobalShader
 	class FAtmosphereTransmittance : SHADER_PERMUTATION_BOOL("USE_ATMOSPHERE_TRANSMITTANCE");
 	class FCloudTransmittance 	: SHADER_PERMUTATION_BOOL("USE_CLOUD_TRANSMITTANCE");
 	class FAnistropicMaterials 	: SHADER_PERMUTATION_BOOL("SUPPORTS_ANISOTROPIC_MATERIALS");
-	class FStrataTileType		: SHADER_PERMUTATION_INT("STRATA_TILETYPE", 3);
+	class FStrataTileType		: SHADER_PERMUTATION_INT("STRATA_TILETYPE", 4);
 	class FVirtualShadowMapMask : SHADER_PERMUTATION_BOOL("USE_VIRTUAL_SHADOW_MAP_MASK");
 
 	using FPermutationDomain = TShaderPermutationDomain<
@@ -2071,9 +2071,10 @@ static uint32 InternalSetBoundingGeometryDepthState(FGraphicsPipelineStateInitia
 		check(Strata::IsStrataEnabled());
 		switch (TileType)
 		{
-		case EStrataTileType::ESimple : StencilRef = Strata::StencilBit_Fast;    GraphicsPSOInit.DepthStencilState = TStaticDepthStencilState<false, CompareFunction, true, CF_Equal, SO_Keep, SO_Keep, SO_Keep, true, CF_Equal, SO_Keep, SO_Keep, SO_Keep, Strata::StencilBit_Fast, 0x0>::GetRHI(); break;
-		case EStrataTileType::ESingle : StencilRef = Strata::StencilBit_Single;  GraphicsPSOInit.DepthStencilState = TStaticDepthStencilState<false, CompareFunction, true, CF_Equal, SO_Keep, SO_Keep, SO_Keep, true, CF_Equal, SO_Keep, SO_Keep, SO_Keep, Strata::StencilBit_Single, 0x0>::GetRHI(); break;
-		case EStrataTileType::EComplex: StencilRef = Strata::StencilBit_Complex; GraphicsPSOInit.DepthStencilState = TStaticDepthStencilState<false, CompareFunction, true, CF_Equal, SO_Keep, SO_Keep, SO_Keep, true, CF_Equal, SO_Keep, SO_Keep, SO_Keep, Strata::StencilBit_Complex, 0x0>::GetRHI(); break;
+		case EStrataTileType::ESimple :			StencilRef = Strata::StencilBit_Fast;			GraphicsPSOInit.DepthStencilState = TStaticDepthStencilState<false, CompareFunction, true, CF_Equal, SO_Keep, SO_Keep, SO_Keep, true, CF_Equal, SO_Keep, SO_Keep, SO_Keep, Strata::StencilBit_Fast, 0x0>::GetRHI(); break;
+		case EStrataTileType::ESingle :			StencilRef = Strata::StencilBit_Single;			GraphicsPSOInit.DepthStencilState = TStaticDepthStencilState<false, CompareFunction, true, CF_Equal, SO_Keep, SO_Keep, SO_Keep, true, CF_Equal, SO_Keep, SO_Keep, SO_Keep, Strata::StencilBit_Single, 0x0>::GetRHI(); break;
+		case EStrataTileType::EComplex:			StencilRef = Strata::StencilBit_Complex;		GraphicsPSOInit.DepthStencilState = TStaticDepthStencilState<false, CompareFunction, true, CF_Equal, SO_Keep, SO_Keep, SO_Keep, true, CF_Equal, SO_Keep, SO_Keep, SO_Keep, Strata::StencilBit_Complex, 0x0>::GetRHI(); break;
+		case EStrataTileType::EComplexSpecial:	StencilRef = Strata::StencilBit_ComplexSpecial; GraphicsPSOInit.DepthStencilState = TStaticDepthStencilState<false, CompareFunction, true, CF_Equal, SO_Keep, SO_Keep, SO_Keep, true, CF_Equal, SO_Keep, SO_Keep, SO_Keep, Strata::StencilBit_ComplexSpecial, 0x0>::GetRHI(); break;
 		default: check(false);
 		}
 	}
@@ -2506,6 +2507,14 @@ static void RenderLight(
 		//   so that the geometry get correctly stencil culled on complex/simple part of the screen
 		if (Strata::IsStrataEnabled())
 		{
+			// Complex Special tiles
+			if (Strata::GetStrataUsesComplexSpecialPath(View))
+			{
+				const EStrataTileType TileType = EStrataTileType::EComplexSpecial;
+				PermutationVector.Set<FDeferredLightPS::FStrataTileType>(TileType);
+				TShaderMapRef< FDeferredLightPS > PixelShader(View.ShaderMap, PermutationVector);
+				InternalRenderLight(GraphBuilder, Scene, View, LightSceneInfo, PixelShader, PassParameters, TileType, TEXT("Light::StandardDeferred(ComplexSpecial)"));
+			}
 			// Complex tiles
 			{
 				const EStrataTileType TileType = EStrataTileType::EComplex;
@@ -2874,6 +2883,10 @@ void FDeferredShadingSceneRenderer::RenderSimpleLightsStandardDeferred(
 
 		if (Strata::IsStrataEnabled())
 		{
+			if (Strata::GetStrataUsesComplexSpecialPath(View))
+			{
+				InternalRenderSimpleLightsStandardDeferred(GraphBuilder, Scene, View, ViewIndex, NumViews, SceneTextures, SimpleLights, EStrataTileType::EComplexSpecial);
+			}
 			InternalRenderSimpleLightsStandardDeferred(GraphBuilder, Scene, View, ViewIndex, NumViews, SceneTextures, SimpleLights, EStrataTileType::EComplex);
 			InternalRenderSimpleLightsStandardDeferred(GraphBuilder, Scene, View, ViewIndex, NumViews, SceneTextures, SimpleLights, EStrataTileType::ESingle);
 			InternalRenderSimpleLightsStandardDeferred(GraphBuilder, Scene, View, ViewIndex, NumViews, SceneTextures, SimpleLights, EStrataTileType::ESimple);
diff --git a/Engine/Source/Runtime/Renderer/Private/PostProcess/SceneRenderTargets.h b/Engine/Source/Runtime/Renderer/Private/PostProcess/SceneRenderTargets.h
--- a/Engine/Source/Runtime/Renderer/Private/PostProcess/SceneRenderTargets.h
+++ b/Engine/Source/Runtime/Renderer/Private/PostProcess/SceneRenderTargets.h
@@ -30,6 +30,7 @@
 #define STENCIL_STRATA_FASTPATH							4 
 #define STENCIL_STRATA_SINGLEPATH						5
 #define STENCIL_STRATA_COMPLEX							6
+#define STENCIL_STRATA_COMPLEX_SPECIAL					0
 // Used only by Strata during the base pass when bUseDBufferPass is enabled (to mark material STRATA_DBUFFER_RESPONSE_xxx Normal/BaseColor/Roughness)
 #define STENCIL_STRATA_RECEIVE_DBUFFER_NORMAL_BIT_ID	1
 #define STENCIL_STRATA_RECEIVE_DBUFFER_DIFFUSE_BIT_ID	3
diff --git a/Engine/Source/Runtime/Renderer/Private/SceneVisibility.cpp b/Engine/Source/Runtime/Renderer/Private/SceneVisibility.cpp
--- a/Engine/Source/Runtime/Renderer/Private/SceneVisibility.cpp
+++ b/Engine/Source/Runtime/Renderer/Private/SceneVisibility.cpp
@@ -1153,6 +1153,7 @@ void FRelevancePacket::Finalize()
 	WriteView.bUsesCustomStencil |= bUsesCustomStencil;
 	WriteView.StrataViewData.MaxBSDFCount = FMath::Max(WriteView.StrataViewData.MaxBSDFCount, 8u - FMath::CountLeadingZeros8(StrataBSDFCountMask));
 	WriteView.StrataViewData.MaxBytesPerPixel = FMath::Max(WriteView.StrataViewData.MaxBytesPerPixel, StrataUintPerPixel * 4u);
+	WriteView.StrataViewData.bUsesComplexSpecialRenderPath |= bUsesComplexSpecialRenderPath;
 	DirtyIndirectLightingCacheBufferPrimitives.AppendTo(WriteView.DirtyIndirectLightingCacheBufferPrimitives);
 
 	WriteView.MeshDecalBatches.Append(MeshDecalBatches);
@@ -1210,6 +1211,7 @@ void FRelevancePacket::ComputeRelevance(FDynamicPrimitiveIndexList& DynamicPrimi
 
 	CombinedShadingModelMask = 0;
 	StrataUintPerPixel = 0;
+	bUsesComplexSpecialRenderPath = false;
 	StrataBSDFCountMask = 0;
 	bSceneHasSkyMaterial = 0;
 	bHasSingleLayerWaterMaterial = 0;
@@ -1767,6 +1769,7 @@ void FRelevancePacket::ComputeRelevance(FDynamicPrimitiveIndexList& DynamicPrimi
 
 		CombinedShadingModelMask |= ViewRelevance.ShadingModelMask;
 		StrataUintPerPixel = FMath::Max(StrataUintPerPixel, ViewRelevance.StrataUintPerPixel);
+		bUsesComplexSpecialRenderPath |= ViewRelevance.bUsesComplexSpecialRenderPath;
 		StrataBSDFCountMask |= ViewRelevance.StrataBSDFCountMask;
 		bUsesLightingChannels |= ViewRelevance.bUsesLightingChannels;
 		bTranslucentSurfaceLighting |= ViewRelevance.bTranslucentSurfaceLighting;
diff --git a/Engine/Source/Runtime/Renderer/Private/SceneVisibilityPrivate.h b/Engine/Source/Runtime/Renderer/Private/SceneVisibilityPrivate.h
--- a/Engine/Source/Runtime/Renderer/Private/SceneVisibilityPrivate.h
+++ b/Engine/Source/Runtime/Renderer/Private/SceneVisibilityPrivate.h
@@ -656,6 +656,7 @@ private:
 	uint16 CombinedShadingModelMask = 0;
 	uint8 StrataUintPerPixel = 0;
 	uint8 StrataBSDFCountMask = 0;
+	bool bUsesComplexSpecialRenderPath = false;
 	bool bHasDistortionPrimitives = false;
 	bool bHasCustomDepthPrimitives = false;
 	bool bUsesLightingChannels = false;
diff --git a/Engine/Source/Runtime/Renderer/Private/Strata/Strata.cpp b/Engine/Source/Runtime/Renderer/Private/Strata/Strata.cpp
--- a/Engine/Source/Runtime/Renderer/Private/Strata/Strata.cpp
+++ b/Engine/Source/Runtime/Renderer/Private/Strata/Strata.cpp
@@ -76,9 +76,11 @@ void FStrataViewData::Reset()
 	// Preserve old max BSDF count which is set prior to the reset operation
 	const uint32 OldMaxBSDFCount = MaxBSDFCount;
 	const uint32 OldMaxBytesPerPixel = MaxBytesPerPixel;
+	const bool   OldUsesComplexSpecialRenderPath = bUsesComplexSpecialRenderPath;
 	*this = FStrataViewData();
 	MaxBSDFCount = OldMaxBSDFCount;
 	MaxBytesPerPixel = OldMaxBytesPerPixel;
+	bUsesComplexSpecialRenderPath = OldUsesComplexSpecialRenderPath;
 
 	ClassificationTileListBuffer = nullptr;
 	ClassificationTileListBufferUAV = nullptr;
@@ -92,6 +94,7 @@ const TCHAR* ToString(EStrataTileType Type)
 	case EStrataTileType::ESimple:							return TEXT("Simple");
 	case EStrataTileType::ESingle:							return TEXT("Single");
 	case EStrataTileType::EComplex:							return TEXT("Complex");
+	case EStrataTileType::EComplexSpecial:					return TEXT("ComplexSpecial");
 	case EStrataTileType::EOpaqueRoughRefraction:			return TEXT("Opaque/RoughRefraction");
 	case EStrataTileType::EOpaqueRoughRefractionSSSWithout:	return TEXT("Opaque/RoughRefraction/SSSWithout");
 	case EStrataTileType::EDecalSimple:						return TEXT("Decal/Simple");
@@ -163,6 +166,15 @@ FIntPoint GetStrataTextureResolution(const FViewInfo& View, const FIntPoint& InR
 	}
 }
 
+bool GetStrataUsesComplexSpecialPath(const FViewInfo& View)
+{
+	if (Strata::IsStrataEnabled())
+	{
+		return View.StrataViewData.bUsesComplexSpecialRenderPath;
+	}
+	return false;
+}
+
 static void BindStrataGlobalUniformParameters(FRDGBuilder& GraphBuilder, FStrataViewData* StrataViewData, FStrataGlobalUniformParameters& OutStrataUniformParameters);
 
 bool SupportsCMask(const FStaticShaderPlatform InPlatform)
@@ -200,18 +212,6 @@ static void InitialiseStrataViewData(FRDGBuilder& GraphBuilder, FViewInfo& View,
 	{
 		const FIntPoint TileResolution(FMath::DivideAndRoundUp(ViewResolution.X, STRATA_TILE_SIZE), FMath::DivideAndRoundUp(ViewResolution.Y, STRATA_TILE_SIZE));
 
-		const TCHAR* StrataTileListBufferNames[EStrataTileType::ECount] =
-		{
-			TEXT("Substrate.StrataTileListBuffer(Simple)"),
-			TEXT("Substrate.StrataTileListBuffer(Single)"),
-			TEXT("Substrate.StrataTileListBuffer(Complex)"),
-			TEXT("Substrate.StrataTileListBuffer(OpaqueRoughRefraction)"),
-			TEXT("Substrate.StrataTileListBuffer(OpaqueRoughRefractionSSSWithout)"),
-			TEXT("Substrate.StrataTileListBuffer(DecalSimple)"),
-			TEXT("Substrate.StrataTileListBuffer(DecalSingle)"),
-			TEXT("Substrate.StrataTileListBuffer(DecalComplex)"),
-		};
-
 		// Tile classification buffers
 		{
 			// Indirect draw
@@ -232,12 +232,13 @@ static void InitialiseStrataViewData(FRDGBuilder& GraphBuilder, FViewInfo& View,
 			Out.ClassificationTileListBufferOffset[EStrataTileType::ESimple]							= RegularTileCount * 0 + RoughTileCount * 0 + DecalTileCount * 0;
 			Out.ClassificationTileListBufferOffset[EStrataTileType::ESingle]							= RegularTileCount * 1 + RoughTileCount * 0 + DecalTileCount * 0;
 			Out.ClassificationTileListBufferOffset[EStrataTileType::EComplex]							= RegularTileCount * 2 + RoughTileCount * 0 + DecalTileCount * 0;
-			Out.ClassificationTileListBufferOffset[EStrataTileType::EOpaqueRoughRefraction]				= RegularTileCount * 3 + RoughTileCount * 0 + DecalTileCount * 0;
-			Out.ClassificationTileListBufferOffset[EStrataTileType::EOpaqueRoughRefractionSSSWithout]	= RegularTileCount * 3 + RoughTileCount * 1 + DecalTileCount * 0;
-			Out.ClassificationTileListBufferOffset[EStrataTileType::EDecalSimple]						= RegularTileCount * 3 + RoughTileCount * 2 + DecalTileCount * 0;
-			Out.ClassificationTileListBufferOffset[EStrataTileType::EDecalSingle]						= RegularTileCount * 3 + RoughTileCount * 2 + DecalTileCount * 1;
-			Out.ClassificationTileListBufferOffset[EStrataTileType::EDecalComplex]						= RegularTileCount * 3 + RoughTileCount * 2 + DecalTileCount * 2;
-			uint32 TotalTileCount										 								= RegularTileCount * 3 + RoughTileCount * 2 + DecalTileCount * 3;
+			Out.ClassificationTileListBufferOffset[EStrataTileType::EComplexSpecial]					= RegularTileCount * 3 + RoughTileCount * 0 + DecalTileCount * 0;
+			Out.ClassificationTileListBufferOffset[EStrataTileType::EOpaqueRoughRefraction]				= RegularTileCount * 4 + RoughTileCount * 0 + DecalTileCount * 0;
+			Out.ClassificationTileListBufferOffset[EStrataTileType::EOpaqueRoughRefractionSSSWithout]	= RegularTileCount * 4 + RoughTileCount * 1 + DecalTileCount * 0;
+			Out.ClassificationTileListBufferOffset[EStrataTileType::EDecalSimple]						= RegularTileCount * 4 + RoughTileCount * 2 + DecalTileCount * 0;
+			Out.ClassificationTileListBufferOffset[EStrataTileType::EDecalSingle]						= RegularTileCount * 4 + RoughTileCount * 2 + DecalTileCount * 1;
+			Out.ClassificationTileListBufferOffset[EStrataTileType::EDecalComplex]						= RegularTileCount * 4 + RoughTileCount * 2 + DecalTileCount * 2;
+			uint32 TotalTileCount										 								= RegularTileCount * 4 + RoughTileCount * 2 + DecalTileCount * 3;
 
 			check(TotalTileCount > 0);
 
@@ -1002,6 +1003,7 @@ uint32 TileTypeDispatchIndirectArgOffset(const EStrataTileType Type)
 constexpr uint32 StencilBit_Fast_1	  = StencilBit_Fast;
 constexpr uint32 StencilBit_Single_1  = StencilBit_Single;
 constexpr uint32 StencilBit_Complex_1 = StencilBit_Complex; 
+constexpr uint32 StencilBit_ComplexSpecial_1 = StencilBit_ComplexSpecial; 
 
 void AddStrataInternalClassificationTilePass(
 	FRDGBuilder& GraphBuilder,
@@ -1041,6 +1043,7 @@ void AddStrataInternalClassificationTilePass(
 		case EStrataTileType::ESimple:							ParametersPS->DebugTileColor = FVector4f(0.0f, 1.0f, 0.0f, 1.0); break;
 		case EStrataTileType::ESingle:							ParametersPS->DebugTileColor = FVector4f(1.0f, 1.0f, 0.0f, 1.0); break;
 		case EStrataTileType::EComplex:							ParametersPS->DebugTileColor = FVector4f(1.0f, 0.0f, 0.0f, 1.0); break;
+		case EStrataTileType::EComplexSpecial:					ParametersPS->DebugTileColor = FVector4f(0.3f, 0.0f, 0.0f, 1.0); break;
 
 		case EStrataTileType::EOpaqueRoughRefraction:			ParametersPS->DebugTileColor = FVector4f(0.0f, 1.0f, 1.0f, 1.0); break;
 		case EStrataTileType::EOpaqueRoughRefractionSSSWithout:	ParametersPS->DebugTileColor = FVector4f(0.0f, 0.0f, 1.0f, 1.0); break;
@@ -1118,6 +1121,16 @@ void AddStrataInternalClassificationTilePass(
 					StencilRef = StencilBit_Complex_1;
 				}
 				break;
+				case EStrataTileType::EComplexSpecial:
+				{
+					GraphicsPSOInit.DepthStencilState = TStaticDepthStencilState<
+						false, CF_Always,
+						true, CF_Always, SO_Keep, SO_Keep, SO_Replace,
+						false, CF_Always, SO_Keep, SO_Keep, SO_Keep,
+						0xFF, StencilBit_ComplexSpecial_1>::GetRHI();
+					StencilRef = StencilBit_ComplexSpecial_1;
+				}
+				break;
 				}
 			}
 			GraphicsPSOInit.BoundShaderState.VertexDeclarationRHI = GFilterVertexDeclaration.VertexDeclarationRHI;
@@ -1150,6 +1163,10 @@ void AddStrataStencilPass(
 		RDG_EVENT_SCOPE_CONDITIONAL(GraphBuilder, Views.Num() > 1, "View%d", i);
 
 		const FViewInfo& View = Views[i];
+		if (GetStrataUsesComplexSpecialPath(View))
+		{
+			AddStrataInternalClassificationTilePass(GraphBuilder, View, &SceneTextures.Depth.Target, nullptr, EStrataTileType::EComplexSpecial);
+		}
 		AddStrataInternalClassificationTilePass(GraphBuilder, View, &SceneTextures.Depth.Target, nullptr, EStrataTileType::EComplex);
 		AddStrataInternalClassificationTilePass(GraphBuilder, View, &SceneTextures.Depth.Target, nullptr, EStrataTileType::ESingle);
 		AddStrataInternalClassificationTilePass(GraphBuilder, View, &SceneTextures.Depth.Target, nullptr, EStrataTileType::ESimple);
@@ -1310,36 +1327,44 @@ void AddStrataMaterialClassificationPass(FRDGBuilder& GraphBuilder, const FMinim
 			FRDGBufferUAVRef RWBSDFTileCountBuffer = GraphBuilder.CreateUAV(StrataViewData->BSDFTileCountBuffer, PF_R32_UINT);
 			AddClearUAVPass(GraphBuilder, RWBSDFTileCountBuffer, 0u);
 
-			FStrataBSDFTilePassCS::FPermutationDomain PermutationVector;
-			PermutationVector.Set< FStrataBSDFTilePassCS::FWaveOps >(bWaveOps);
-			TShaderMapRef<FStrataBSDFTilePassCS> ComputeShader(View.ShaderMap, PermutationVector);
-			FStrataBSDFTilePassCS::FParameters* PassParameters = GraphBuilder.AllocParameters<FStrataBSDFTilePassCS::FParameters>();
-			PassParameters->ViewUniformBuffer = View.ViewUniformBuffer;
-			PassParameters->TileSizeLog2 = STRATA_TILE_SIZE_DIV_AS_SHIFT;
-			PassParameters->TileCount_Primary = StrataViewData->TileCount;
-			PassParameters->TileOffset_Primary = StrataViewData->TileOffset;
-			PassParameters->OverflowTileCount = StrataViewData->OverflowTileCount;
-			PassParameters->OverflowTileOffset = StrataViewData->OverflowTileOffset;
-			PassParameters->ViewResolution = View.ViewRect.Size();
-			PassParameters->MaxBytesPerPixel = StrataSceneData->MaxBytesPerPixel;
-			PassParameters->TopLayerTexture = StrataSceneData->TopLayerTexture;
-			PassParameters->MaterialTextureArray = StrataSceneData->MaterialTextureArraySRV;
-			PassParameters->TileListBuffer = StrataViewData->ClassificationTileListBufferSRV;
-			PassParameters->TileListBufferOffset = StrataViewData->ClassificationTileListBufferOffset[EStrataTileType::EComplex];
-			PassParameters->TileIndirectBuffer = StrataViewData->ClassificationTileDispatchIndirectBuffer;
-
-			PassParameters->RWBSDFOffsetTexture = GraphBuilder.CreateUAV(StrataSceneData->BSDFOffsetTexture);
-			PassParameters->RWBSDFTileTexture = GraphBuilder.CreateUAV(StrataViewData->BSDFTileTexture);
-			PassParameters->RWBSDFTileCountBuffer = RWBSDFTileCountBuffer;
-
-			FComputeShaderUtils::AddPass(
-				GraphBuilder,
-				RDG_EVENT_NAME("Substrate::BSDFTileAndOffsets(%s)", bWaveOps ? TEXT("Wave") : TEXT("SharedMemory")),
-				PassFlags,
-				ComputeShader,
-				PassParameters,
-				PassParameters->TileIndirectBuffer,
-				TileTypeDispatchIndirectArgOffset(EStrataTileType::EComplex));
+			auto MarkBSDFTilePass = [&](EStrataTileType TileType)
+			{
+				FStrataBSDFTilePassCS::FPermutationDomain PermutationVector;
+				PermutationVector.Set< FStrataBSDFTilePassCS::FWaveOps >(bWaveOps);
+				TShaderMapRef<FStrataBSDFTilePassCS> ComputeShader(View.ShaderMap, PermutationVector);
+				FStrataBSDFTilePassCS::FParameters* PassParameters = GraphBuilder.AllocParameters<FStrataBSDFTilePassCS::FParameters>();
+				PassParameters->ViewUniformBuffer = View.ViewUniformBuffer;
+				PassParameters->TileSizeLog2 = STRATA_TILE_SIZE_DIV_AS_SHIFT;
+				PassParameters->TileCount_Primary = StrataViewData->TileCount;
+				PassParameters->TileOffset_Primary = StrataViewData->TileOffset;
+				PassParameters->OverflowTileCount = StrataViewData->OverflowTileCount;
+				PassParameters->OverflowTileOffset = StrataViewData->OverflowTileOffset;
+				PassParameters->ViewResolution = View.ViewRect.Size();
+				PassParameters->MaxBytesPerPixel = StrataSceneData->MaxBytesPerPixel;
+				PassParameters->TopLayerTexture = StrataSceneData->TopLayerTexture;
+				PassParameters->MaterialTextureArray = StrataSceneData->MaterialTextureArraySRV;
+				PassParameters->TileListBuffer = StrataViewData->ClassificationTileListBufferSRV;
+				PassParameters->TileListBufferOffset = StrataViewData->ClassificationTileListBufferOffset[TileType];
+				PassParameters->TileIndirectBuffer = StrataViewData->ClassificationTileDispatchIndirectBuffer;
+
+				PassParameters->RWBSDFOffsetTexture = GraphBuilder.CreateUAV(StrataSceneData->BSDFOffsetTexture);
+				PassParameters->RWBSDFTileTexture = GraphBuilder.CreateUAV(StrataViewData->BSDFTileTexture);
+				PassParameters->RWBSDFTileCountBuffer = RWBSDFTileCountBuffer;
+
+				FComputeShaderUtils::AddPass(
+					GraphBuilder,
+					RDG_EVENT_NAME("Substrate::BSDFTileAndOffsets(%s - %s)", ToString(TileType), bWaveOps ? TEXT("Wave") : TEXT("SharedMemory")),
+					PassFlags,
+					ComputeShader,
+					PassParameters,
+					PassParameters->TileIndirectBuffer,
+					TileTypeDispatchIndirectArgOffset(TileType));
+			};
+			if (GetStrataUsesComplexSpecialPath(View))
+			{
+				MarkBSDFTilePass(EStrataTileType::EComplexSpecial);
+			}
+			MarkBSDFTilePass(EStrataTileType::EComplex);
 		}
 
 		// Tile indirect dispatch args conversion
diff --git a/Engine/Source/Runtime/Renderer/Private/Strata/Strata.h b/Engine/Source/Runtime/Renderer/Private/Strata/Strata.h
--- a/Engine/Source/Runtime/Renderer/Private/Strata/Strata.h
+++ b/Engine/Source/Runtime/Renderer/Private/Strata/Strata.h
@@ -79,6 +79,7 @@ enum EStrataTileType : uint32
 	ESimple								= STRATA_TILE_TYPE_SIMPLE,
 	ESingle								= STRATA_TILE_TYPE_SINGLE,
 	EComplex							= STRATA_TILE_TYPE_COMPLEX,
+	EComplexSpecial						= STRATA_TILE_TYPE_COMPLEX_SPECIAL,
 	EOpaqueRoughRefraction				= STRATA_TILE_TYPE_ROUGH_REFRACT,
 	EOpaqueRoughRefractionSSSWithout	= STRATA_TILE_TYPE_ROUGH_REFRACT_SSS_WITHOUT,
 	EDecalSimple						= STRATA_TILE_TYPE_DECAL_SIMPLE,
@@ -134,6 +135,9 @@ struct FStrataViewData
 	uint32 MaxBSDFCount = 0;
 	uint32 MaxBytesPerPixel = 0;
 
+	// True if any material requires the complex special path (e.g. glints or SpecularLUT)
+	bool bUsesComplexSpecialRenderPath = 0;
+
 	FIntPoint TileCount  = FIntPoint(0, 0);
 	FIntPoint TileOffset = FIntPoint(0, 0);
 	FIntPoint OverflowTileCount = FIntPoint(0, 0);
@@ -164,11 +168,13 @@ struct FStrataViewData
 
 namespace Strata
 {
-constexpr uint32 StencilBit_Fast   = 0x08; // In sync with SceneRenderTargets.h - GET_STENCIL_BIT_MASK(STENCIL_STRATA_FASTPATH)
-constexpr uint32 StencilBit_Single = 0x10; // In sync with SceneRenderTargets.h - GET_STENCIL_BIT_MASK(STENCIL_STRATA_SINGLEPATH)
-constexpr uint32 StencilBit_Complex= 0x20; // In sync with SceneRenderTargets.h - GET_STENCIL_BIT_MASK(STENCIL_STRATA_COMPLEX)
+constexpr uint32 StencilBit_Fast			= 0x08; // In sync with SceneRenderTargets.h - GET_STENCIL_BIT_MASK(STENCIL_STRATA_FASTPATH)
+constexpr uint32 StencilBit_Single			= 0x10; // In sync with SceneRenderTargets.h - GET_STENCIL_BIT_MASK(STENCIL_STRATA_SINGLEPATH)
+constexpr uint32 StencilBit_Complex			= 0x20; // In sync with SceneRenderTargets.h - GET_STENCIL_BIT_MASK(STENCIL_STRATA_COMPLEX)
+constexpr uint32 StencilBit_ComplexSpecial	= 0x01; // In sync with SceneRenderTargets.h - GET_STENCIL_BIT_MASK(STENCIL_STRATA_COMPLEX_SPECIAL)	
 
 FIntPoint GetStrataTextureResolution(const FViewInfo& View, const FIntPoint& InResolution);
+bool GetStrataUsesComplexSpecialPath(const FViewInfo& View);
 
 void InitialiseStrataFrameSceneData(FRDGBuilder& GraphBuilder, FSceneRenderer& SceneRenderer);
 
diff --git a/Engine/Source/Runtime/Renderer/Private/Strata/StrataVisualize.cpp b/Engine/Source/Runtime/Renderer/Private/Strata/StrataVisualize.cpp
--- a/Engine/Source/Runtime/Renderer/Private/Strata/StrataVisualize.cpp
+++ b/Engine/Source/Runtime/Renderer/Private/Strata/StrataVisualize.cpp
@@ -446,6 +446,10 @@ FScreenPassTexture AddStrataDebugPasses(FRDGBuilder& GraphBuilder, const FViewIn
 		}
 		else if (DebugMode == FStrataViewMode::MaterialClassification)
 		{
+			if (GetStrataUsesComplexSpecialPath(View))
+			{
+				AddStrataInternalClassificationTilePass(GraphBuilder, View, nullptr, &ScreenPassSceneColor.Texture, EStrataTileType::EComplexSpecial, bDebugPass);
+			}
 			AddStrataInternalClassificationTilePass(GraphBuilder, View, nullptr, &ScreenPassSceneColor.Texture, EStrataTileType::EComplex, bDebugPass);
 			AddStrataInternalClassificationTilePass(GraphBuilder, View, nullptr, &ScreenPassSceneColor.Texture, EStrataTileType::ESingle, bDebugPass);
 			AddStrataInternalClassificationTilePass(GraphBuilder, View, nullptr, &ScreenPassSceneColor.Texture, EStrataTileType::ESimple, bDebugPass);
