Substrate
- glint can be enabled/disabled per platofm still.
- Glints are enabled by default now.
- Simplified the complexspecial path defines to only have one.

diff --git a/Engine/Shaders/Private/ClusteredDeferredShadingPixelShader.usf b/Engine/Shaders/Private/ClusteredDeferredShadingPixelShader.usf
--- a/Engine/Shaders/Private/ClusteredDeferredShadingPixelShader.usf
+++ b/Engine/Shaders/Private/ClusteredDeferredShadingPixelShader.usf
@@ -24,7 +24,7 @@
 	// COMPLEX PATH
   #elif STRATA_TILETYPE == 3
 	// COMPLEX PATH
-	#define STRATA_COMPLEX_SPECIALPATH_ENABLED 1
+	#define STRATA_COMPLEXSPECIALPATH 1
   #else
 	#error Strata tile type non-implemented
   #endif
diff --git a/Engine/Shaders/Private/DeferredLightPixelShaders.usf b/Engine/Shaders/Private/DeferredLightPixelShaders.usf
--- a/Engine/Shaders/Private/DeferredLightPixelShaders.usf
+++ b/Engine/Shaders/Private/DeferredLightPixelShaders.usf
@@ -16,7 +16,7 @@
 	// COMPLEX PATH
   #elif STRATA_TILETYPE == 3
 	// COMPLEX PATH
-	#define STRATA_COMPLEX_SPECIALPATH_ENABLED 1
+	#define STRATA_COMPLEXSPECIALPATH 1
   #else
 	#error Strata tile type non-implemented
   #endif
diff --git a/Engine/Shaders/Private/DiffuseIndirectComposite.usf b/Engine/Shaders/Private/DiffuseIndirectComposite.usf
--- a/Engine/Shaders/Private/DiffuseIndirectComposite.usf
+++ b/Engine/Shaders/Private/DiffuseIndirectComposite.usf
@@ -54,7 +54,7 @@
 	// COMPLEX PATH
   #elif STRATA_TILETYPE == 3
 	// COMPLEX PATH
-	#define STRATA_COMPLEX_SPECIALPATH_ENABLED 1
+	#define STRATA_COMPLEXSPECIALPATH 1
   #else
 	#error Strata tile type non-implemented
   #endif
diff --git a/Engine/Shaders/Private/ReflectionEnvironmentPixelShader.usf b/Engine/Shaders/Private/ReflectionEnvironmentPixelShader.usf
--- a/Engine/Shaders/Private/ReflectionEnvironmentPixelShader.usf
+++ b/Engine/Shaders/Private/ReflectionEnvironmentPixelShader.usf
@@ -14,7 +14,7 @@
 	// COMPLEX PATH
   #elif STRATA_TILETYPE == 3
 	// COMPLEX PATH
-	#define STRATA_COMPLEX_SPECIALPATH_ENABLED 1
+	#define STRATA_COMPLEXSPECIALPATH 1
   #else
 	#error Strata tile type non-implemented
   #endif
diff --git a/Engine/Shaders/Private/Strata/Strata.ush b/Engine/Shaders/Private/Strata/Strata.ush
--- a/Engine/Shaders/Private/Strata/Strata.ush
+++ b/Engine/Shaders/Private/Strata/Strata.ush
@@ -95,16 +95,19 @@
 #define STRATA_SINGLEPATH 0
 #endif
 
-#ifndef STRATA_COMPLEX_SPECIALPATH_ENABLED 
-#define STRATA_COMPLEX_SPECIALPATH_ENABLED 0
-#endif
-
 #define STRATA_COMPLEXPATH (!STRATA_FASTPATH && !STRATA_SINGLEPATH)
-#define STRATA_COMPLEXPATH_SPECIAL (STRATA_COMPLEXPATH && STRATA_COMPLEX_SPECIALPATH_ENABLED)
+
+// STRATA_COMPLEXSPECIALPATH is defined by the material or lighting shaders as they see fit.
+// However, this should only be done for STRATA_COMPLEXPATH path.
+#if STRATA_COMPLEXSPECIALPATH
+#ifndef STRATA_COMPLEXPATH
+#error STRATA_COMPLEXSPECIALPATH is defined while is not defined STRATA_COMPLEXPATH
+#endif
+#endif
 
 // Disable glint when in tracing mode for now since derivative are not available.
 #define STRATA_GLINTS_DERIVATIVE_SUPPORTED (!RAYCALLABLESHADER && !RAYHITGROUPSHADER && !RAYGENSHADER && !RAYMISSSHADER)
-#define STRATA_GLINTS_ENABLED (STRATA_COMPLEXPATH_SPECIAL && STRATA_GLINTS_DERIVATIVE_SUPPORTED)
+#define STRATA_GLINTS_ENABLED (STRATA_COMPLEXSPECIALPATH && PLATFORM_ENABLES_STRATA_GLINTS && STRATA_GLINTS_DERIVATIVE_SUPPORTED)
 
 #define STRATA_LEGACY_MATERIAL_APPLIES_FINAL_WEIGHT (!MATERIALBLENDING_MASKED && !MATERIALBLENDING_SOLID && STRATA_LEGACY_PREMULT_ALPHA_OVERRIDE==0 || STRATA_MATERIAL_EXPORT_FROM_TRANSLUCENT)
 
@@ -166,7 +169,7 @@ struct FStrataBSDF
 	float4 InlineVGPRs[3];
 #endif
 
-#if STRATA_COMPLEXPATH_SPECIAL	
+#if STRATA_COMPLEXSPECIALPATH	
 	float4 VGPRs[7];
 #else
 	float4 VGPRs[5];
@@ -541,7 +544,7 @@ struct FStrataBSDF
 #define SLAB_FUZZ_COLOR(X)				X.VGPRs[4].xyz
 #define SLAB_FUZZ_AMOUNT(X)				X.VGPRs[4].w
 #define SLAB_FUZZ_ROUGHNESS(X)			X.VGPRs[2].w
-#if STRATA_COMPLEXPATH_SPECIAL 
+#if STRATA_COMPLEXSPECIALPATH 
 // Glints
 #define SLAB_GLINT_UV(X)				X.VGPRs[5].xy
 #define SLAB_GLINT_VALUE(X)				X.VGPRs[5].z
@@ -1607,7 +1610,7 @@ FStrataData GetStrataSlabBSDF(
 	BSDF_SETTHICKNESSCM		(StrataData.InlinedBSDF,		ThicknessCm);
 #endif
 
-#if STRATA_COMPLEXPATH_SPECIAL
+#if STRATA_COMPLEXSPECIALPATH
 #if STRATA_GLINTS_ENABLED
 	if (GlintValue > 0.0f)
 	{
@@ -2431,7 +2434,7 @@ FStrataData StrataHorizontalMixingParameterBlending(FStrataData Background, FStr
 	SLAB_FUZZ_ROUGHNESS(ResultBSDF)	=	lerp(SLAB_FUZZ_ROUGHNESS(BackgroundBSDF),		SLAB_FUZZ_ROUGHNESS(ForegroundBSDF),		ForegroundOtherMixFactor);
 	SLAB_HAZINESS(ResultBSDF) =			LerpHazinessParameterBlending(BackgroundBSDF,	ForegroundBSDF,								ForegroundOtherMixFactor);
 
-#if STRATA_COMPLEXPATH_SPECIAL
+#if STRATA_COMPLEXSPECIALPATH
 #if STRATA_GLINTS_ENABLED
 	BSDF_SETHASGLINT(ResultBSDF,		BSDF_GETHASGLINT(ForegroundBSDF)				| BSDF_GETHASGLINT(BackgroundBSDF));
 	SLAB_GLINT_UV(ResultBSDF) =			lerp(SLAB_GLINT_UV(BackgroundBSDF),				SLAB_GLINT_UV(ForegroundBSDF),				ForegroundOtherMixFactor);
@@ -2551,7 +2554,7 @@ FStrataData StrataAddParameterBlending(FStrataData A, FStrataData B, float AMixF
 	SLAB_FUZZ_ROUGHNESS(ResultBSDF) =	lerp(SLAB_FUZZ_ROUGHNESS(BBSDF),				SLAB_FUZZ_ROUGHNESS(ABSDF),					AMixFactor);
 	SLAB_HAZINESS(ResultBSDF) =			LerpHazinessParameterBlending(BBSDF,			ABSDF,										AMixFactor);
 
-#if STRATA_COMPLEXPATH_SPECIAL
+#if STRATA_COMPLEXSPECIALPATH
 #if STRATA_GLINTS_ENABLED
 	BSDF_SETHASGLINT(ResultBSDF,		BSDF_GETHASGLINT(BBSDF)							| BSDF_GETHASGLINT(ABSDF));
 	SLAB_GLINT_UV(ResultBSDF) =			lerp(SLAB_GLINT_UV(BBSDF),						SLAB_GLINT_UV(ABSDF), AMixFactor);
@@ -2642,7 +2645,7 @@ FStrataBSDF StrataVerticalLayeringParameterBlending(in FStrataBSDF TopBSDF, in F
 	BSDF_SETHASF90(ResultBSDF,			BSDF_GETHASF90(BaseBSDF)						| BSDF_GETHASF90(TopBSDF));			// F90: keep union of both even though it will be hard to get a match
 	BSDF_SETHASFUZZ(ResultBSDF,			BSDF_GETHASFUZZ(BaseBSDF)						| BSDF_GETHASFUZZ(TopBSDF));		// Fuzz: keep union of both even though it will be hard to get a match
 
-#if STRATA_COMPLEXPATH_SPECIAL
+#if STRATA_COMPLEXSPECIALPATH
 	BSDF_SETHASGLINT(ResultBSDF, 0); // Disable for now
 	BSDF_SETHASSPECPROFILE(ResultBSDF,	0); // Disable for now
 #endif
@@ -2683,7 +2686,7 @@ FStrataBSDF StrataVerticalLayeringParameterBlending(in FStrataBSDF TopBSDF, in F
 	SLAB_FUZZ_ROUGHNESS(ResultBSDF) =	lerp(SLAB_FUZZ_ROUGHNESS(TopBSDF),				SLAB_FUZZ_ROUGHNESS(BaseBSDF),TopToBottomLerp);
 	SLAB_HAZINESS(ResultBSDF) =			LerpHazinessParameterBlending(TopBSDF,			BaseBSDF,					TopToBottomLerp);
 
-#if STRATA_COMPLEXPATH_SPECIAL
+#if STRATA_COMPLEXSPECIALPATH
 #if STRATA_GLINTS_ENABLED
 	BSDF_SETHASGLINT(ResultBSDF,		BSDF_GETHASGLINT(TopBSDF)						| BSDF_GETHASGLINT(BaseBSDF));
 	SLAB_GLINT_UV(ResultBSDF) =			lerp(SLAB_GLINT_UV(TopBSDF),					SLAB_GLINT_UV(BaseBSDF), TopToBottomLerp);
@@ -4365,7 +4368,7 @@ FStrataBSDF UnpackStrataBSDFIn(FStrataMaterialContainer StrataBuffer, inout FStr
 		}
 
 
-	#if STRATA_COMPLEXPATH_SPECIAL
+	#if STRATA_COMPLEXSPECIALPATH
 	#if STRATA_GLINTS_ENABLED
 		if (BSDF_GETHASGLINT(OutBSDF))
 		{
diff --git a/Engine/Shaders/Private/Strata/StrataEvaluation.ush b/Engine/Shaders/Private/Strata/StrataEvaluation.ush
--- a/Engine/Shaders/Private/Strata/StrataEvaluation.ush
+++ b/Engine/Shaders/Private/Strata/StrataEvaluation.ush
@@ -16,7 +16,7 @@
 #define STRATA_GLINT_IS 0 
 #endif
 
-#if STRATA_COMPLEXPATH_SPECIAL
+#if STRATA_COMPLEXSPECIALPATH
 #include "Glint/GlintThirdParty.ush"
 
 #if STRATA_GLINT_IS
@@ -259,7 +259,7 @@ half3 Strata_F_GGX(half3 F0, half3 F90, half VoH)
 //#endif
 }
 
-#if STRATA_COMPLEXPATH_SPECIAL
+#if STRATA_COMPLEXSPECIALPATH
 // Conversion of GGX roughness to Beckmann roughness
 // Using Microfacet Distribution Function: To Change or Not to Change, That Is the Question.
 float GGXRoughnessToBeckmannRoughness(float Roughness)
@@ -445,7 +445,7 @@ FStrataEvaluateResult StrataEvaluateBSDFCommon(FStrataBSDFContext BSDFContext, F
 						float2 AlphaXY = 0;
 						GetAnisotropicRoughness(Alpha, SLAB_ANISOTROPY(BSDFContext.BSDF), AlphaXY.x, AlphaXY.y);
 
-					#if STRATA_COMPLEXPATH_SPECIAL
+					#if STRATA_COMPLEXSPECIALPATH
 						// Glint shading only in the complex path
 						BRANCH
 						if (BSDF_GETHASGLINT(BSDFContext.BSDF))
@@ -478,7 +478,7 @@ FStrataEvaluateResult StrataEvaluateBSDFCommon(FStrataBSDFContext BSDFContext, F
 							RectLightSpec = RectGGXApproxLTC(SafeRoughness, F0, BSDFContext.N, BSDFContext.V, AreaLightContext.AreaLight.Rect, AreaLightContext.AreaLight.Texture, MeanLightWorldDirection);
 							
 							// Now combine the rectangular area light with Glints.
-						#if STRATA_COMPLEXPATH_SPECIAL
+						#if STRATA_COMPLEXSPECIALPATH
 							BRANCH
 							if (BSDF_GETHASGLINT(BSDFContext.BSDF))
 							{
@@ -504,7 +504,7 @@ FStrataEvaluateResult StrataEvaluateBSDFCommon(FStrataBSDFContext BSDFContext, F
 						{
 							if(IntegrationType == INTEGRATION_PUNCTUAL_LIGHT)
 							{
-							#if STRATA_COMPLEXPATH_SPECIAL
+							#if STRATA_COMPLEXSPECIALPATH
 								BRANCH
 								if (BSDF_GETHASGLINT(BSDFContext.BSDF))
 								{
@@ -522,7 +522,7 @@ FStrataEvaluateResult StrataEvaluateBSDFCommon(FStrataBSDFContext BSDFContext, F
 							{
 								const float Energy = EnergyNormalization(Alpha2Spec, BSDFContext.Context.VoH, AreaLightContext.AreaLight);
 
-							#if STRATA_COMPLEXPATH_SPECIAL
+							#if STRATA_COMPLEXSPECIALPATH
 								BRANCH
 								if (BSDF_GETHASGLINT(BSDFContext.BSDF))
 								{
@@ -682,7 +682,7 @@ FStrataEvaluateResult StrataEvaluateBSDFCommon(FStrataBSDFContext BSDFContext, F
 			#endif
 
 			// Lighting LUT
-			#if STRATA_FASTPATH==0 && STRATA_COMPLEXPATH_SPECIAL
+			#if STRATA_FASTPATH==0 && STRATA_COMPLEXSPECIALPATH
 			BRANCH
 			if (BSDF_GETHASSPECPROFILE(BSDFContext.BSDF))
 			{
@@ -1503,7 +1503,7 @@ FStrataEnvLightResult StrataEvaluateForEnvLight(FStrataBSDFContext BSDFContext,
 				StrataEnvLightResult.SpecularColor = F0;
 
 				// STRATA_TODO environment light importance sampling
-			#if STRATA_COMPLEXPATH_SPECIAL
+			#if STRATA_COMPLEXSPECIALPATH
 				// Glints
 				BRANCH
 				if (BSDF_GETHASGLINT(BSDFContext.BSDF))
diff --git a/Engine/Shaders/Private/Strata/StrataExport.ush b/Engine/Shaders/Private/Strata/StrataExport.ush
--- a/Engine/Shaders/Private/Strata/StrataExport.ush
+++ b/Engine/Shaders/Private/Strata/StrataExport.ush
@@ -190,7 +190,7 @@ void PackStrataOut(
 						// Simple and single bsdf do not allow weights other than 1
 						const bool LuminanceWeightEqualOne = all(BSDF.LuminanceWeightV == 1.0f);
 
-						#if STRATA_COMPLEXPATH_SPECIAL
+						#if STRATA_COMPLEXSPECIALPATH
 						#if STRATA_GLINTS_ENABLED
 						const bool bForceComplexSpecialPath = BSDF_GETHASSPECPROFILE(BSDF) || BSDF_GETHASGLINT(BSDF);
 						#else
@@ -798,7 +798,7 @@ void PackStrataOut(
 								STRATA_STORE_UINT1(PackColorLinearToGamma2AlphaLinear(float4(BSDF.TransmittanceAboveAlongN, BSDF.CoverageAboveAlongN)));
 							}
 
-							#if STRATA_COMPLEXPATH_SPECIAL
+							#if STRATA_COMPLEXSPECIALPATH
 							#if STRATA_GLINTS_ENABLED
 							if (BSDF_GETHASGLINT(BSDF))
 							{
@@ -819,7 +819,7 @@ void PackStrataOut(
 							{
 								STRATA_STORE_UINT1(PackSpecularProfile(SLAB_SPECPROFILEID(BSDF)));
 							}
-							#endif // STRATA_COMPLEXPATH_SPECIAL
+							#endif // STRATA_COMPLEXSPECIALPATH
 
 							// 8-28 bytes
 						}
diff --git a/Engine/Shaders/Private/Strata/StrataVisualize.usf b/Engine/Shaders/Private/Strata/StrataVisualize.usf
--- a/Engine/Shaders/Private/Strata/StrataVisualize.usf
+++ b/Engine/Shaders/Private/Strata/StrataVisualize.usf
@@ -11,7 +11,7 @@
 
 #define STRATA_SSS_MATERIAL_OVERRIDE 0
 #define STRATA_SSS_TRANSMISSION 0				// Disable for now, as the SSS profile texture need to be bound
-#define STRATA_COMPLEX_SPECIALPATH_ENABLED 1	// Allow to read data from complex special path
+#define STRATA_COMPLEXSPECIALPATH 1				// Allow to read data from complex special path
 #include "/Engine/Private/Strata/StrataEvaluation.ush"
 #include "/Engine/Private/Strata/StrataTile.ush"
 
@@ -330,7 +330,7 @@ void PrintBSDF(inout FShaderPrintContext Context, inout FStrataAddressing Strata
 				Print(Context, TEXT(" Fuzz Color  "), FontBSDFPropName);	Print(Context, SLAB_FUZZ_COLOR(BSDF), FontBSDFPropValu);				Newline(Context, RectMax);
 				Print(Context, TEXT(" Fuzz Rough. "), FontBSDFPropName);	Print(Context, SLAB_FUZZ_ROUGHNESS(BSDF), FontBSDFPropValu);			Newline(Context, RectMax);
 			}
-#if STRATA_COMPLEXPATH_SPECIAL
+#if STRATA_COMPLEXSPECIALPATH
 			if (BSDF_GETHASGLINT(BSDF))
 			{
 				Print(Context, TEXT(" Glint Value "), FontBSDFPropName);	Print(Context, SLAB_GLINT_VALUE(BSDF), FontBSDFPropValu);				Newline(Context, RectMax);
diff --git a/Engine/Shaders/Shared/StrataDefinitions.h b/Engine/Shaders/Shared/StrataDefinitions.h
--- a/Engine/Shaders/Shared/StrataDefinitions.h
+++ b/Engine/Shaders/Shared/StrataDefinitions.h
@@ -8,7 +8,7 @@
 #pragma once
 
 // Change this to force recompilation of all strata dependent shaders (use https://www.random.org/cgi-bin/randbyte?nbytes=4&format=h)
-#define STRATA_SHADER_VERSION 0x932681EC
+#define STRATA_SHADER_VERSION 0x56F06029 
 
 
 
diff --git a/Engine/Source/Runtime/Engine/Private/Materials/HLSLMaterialTranslator.cpp b/Engine/Source/Runtime/Engine/Private/Materials/HLSLMaterialTranslator.cpp
--- a/Engine/Source/Runtime/Engine/Private/Materials/HLSLMaterialTranslator.cpp
+++ b/Engine/Source/Runtime/Engine/Private/Materials/HLSLMaterialTranslator.cpp
@@ -2296,7 +2296,7 @@ void FHLSLMaterialTranslator::GetMaterialEnvironment(EShaderPlatform InPlatform,
 			OutEnvironment.SetDefine(TEXT("STRATA_SINGLEPATH"), StrataCtx.bStrataMaterialIsSingle ? TEXT("1") : TEXT("0"));
 			OutEnvironment.SetDefine(TEXT("STRATA_FASTPATH"), StrataCtx.bStrataMaterialIsSingle ? TEXT("0") : (StrataCtx.bStrataMaterialIsSimple ? TEXT("1") : TEXT("0")));
 			OutEnvironment.SetDefine(TEXT("STRATA_CLAMPED_BSDF_COUNT"), StrataCtx.StrataMaterialBSDFCount);
-			OutEnvironment.SetDefine(TEXT("STRATA_COMPLEX_SPECIALPATH_ENABLED"), MaterialCompilationOutput.StrataMaterialCompilationOutput.bUsesComplexSpecialRenderPath ? TEXT("1") : TEXT("0"));
+			OutEnvironment.SetDefine(TEXT("STRATA_COMPLEXSPECIALPATH"), MaterialCompilationOutput.StrataMaterialCompilationOutput.bUsesComplexSpecialRenderPath ? TEXT("1") : TEXT("0"));
 		}
 
 
@@ -11555,6 +11555,9 @@ bool FHLSLMaterialTranslator::FStrataCompilationContext::StrataGenerateDerivedMa
 			const uint32 UintByteSize = sizeof(uint32);
 			StrataMaterialRequestedSizeByte = 0;
 
+			// Used to not promote to ComplexSpecial render path if glint rendering is not enabled and avoid registering such possibility at runtime.
+			const bool bGlintEnabled = Strata::IsGlintEnabled();
+
 			// 1. Evaluate simple/single BSDF
 			bStrataMaterialIsSimple = StrataMaterialBSDFCount == 1;
 			bStrataMaterialIsSingle = StrataMaterialBSDFCount == 1;
@@ -11581,7 +11584,7 @@ bool FHLSLMaterialTranslator::FStrataCompilationContext::StrataGenerateDerivedMa
 					{
 						bStrataMaterialIsSimple = bStrataMaterialIsSimple && !bMayHaveColoredWeight && !It.bBSDFHasAnisotropy && !It.bBSDFHasEdgeColor && !It.bBSDFHasFuzz && !It.bBSDFHasSecondRoughnessOrSimpleClearCoat && !It.bBSDFHasMFPPluggedIn && !It.bBSDFHasSSS && !It.bBSDFHasGlint && !It.bBSDFHasSpecularProfile;
 						bStrataMaterialIsSingle = bStrataMaterialIsSingle && !bMayHaveColoredWeight && !It.bBSDFHasAnisotropy && !It.bBSDFHasGlint && !It.bBSDFHasSpecularProfile;
-						bUsesComplexSpecialRenderPath |= It.bBSDFHasGlint || It.bBSDFHasSpecularProfile;
+						bUsesComplexSpecialRenderPath |= (bGlintEnabled && It.bBSDFHasGlint) || It.bBSDFHasSpecularProfile;
 						break;
 					}
 					case STRATA_BSDF_TYPE_HAIR:
diff --git a/Engine/Source/Runtime/Engine/Private/ShaderCompiler/ShaderCompiler.cpp b/Engine/Source/Runtime/Engine/Private/ShaderCompiler/ShaderCompiler.cpp
--- a/Engine/Source/Runtime/Engine/Private/ShaderCompiler/ShaderCompiler.cpp
+++ b/Engine/Source/Runtime/Engine/Private/ShaderCompiler/ShaderCompiler.cpp
@@ -6662,6 +6662,9 @@ void GlobalBeginCompileShader(
 
 			const bool bStrataDBufferPass = Strata::IsDBufferPassEnabled(Target.GetPlatform());
 			Input.Environment.SetDefine(TEXT("STRATA_USE_DBUFFER_PASS"), bStrataDBufferPass ? 1 : 0);
+
+			const bool bStrataGlints = Strata::IsGlintEnabled();
+			Input.Environment.SetDefine(TEXT("PLATFORM_ENABLES_STRATA_GLINTS"), bStrataGlints ? 1 : 0);
 		}
 
 		const bool bStrataBackCompatibility = bStrata && Strata::IsBackCompatibilityEnabled();
diff --git a/Engine/Source/Runtime/RenderCore/Private/RenderUtils.cpp b/Engine/Source/Runtime/RenderCore/Private/RenderUtils.cpp
--- a/Engine/Source/Runtime/RenderCore/Private/RenderUtils.cpp
+++ b/Engine/Source/Runtime/RenderCore/Private/RenderUtils.cpp
@@ -1599,8 +1599,8 @@ static TAutoConsoleVariable<int32> CVarSubstrateRoughDiffuse(
 
 static TAutoConsoleVariable<int32> CVarSubstrateGlints(
 	TEXT("r.Substrate.Glints"),
-	0,
-	TEXT("Enable Glint support for Strata slabs."),
+	1,
+	TEXT("Enable Glint support for Strata slabs. If changed, shaders needs to be recompiled."),
 	ECVF_ReadOnly | ECVF_RenderThreadSafe);
 
 static TAutoConsoleVariable<int32> CVarSubstrateGlintsLUT(
