Substrate - added platofrm control for specular profile through r.Substrate.SpecularProfile.

diff --git a/Engine/Shaders/Private/Strata/Strata.ush b/Engine/Shaders/Private/Strata/Strata.ush
--- a/Engine/Shaders/Private/Strata/Strata.ush
+++ b/Engine/Shaders/Private/Strata/Strata.ush
@@ -109,6 +109,8 @@
 #define STRATA_GLINTS_DERIVATIVE_SUPPORTED (!RAYCALLABLESHADER && !RAYHITGROUPSHADER && !RAYGENSHADER && !RAYMISSSHADER)
 #define STRATA_GLINTS_ENABLED (STRATA_COMPLEXSPECIALPATH && PLATFORM_ENABLES_STRATA_GLINTS && STRATA_GLINTS_DERIVATIVE_SUPPORTED)
 
+#define STRATA_SPECPROFILE_ENABLED (STRATA_COMPLEXSPECIALPATH && PLATFORM_ENABLES_STRATA_SPECULAR_PROFILE)
+
 #define STRATA_LEGACY_MATERIAL_APPLIES_FINAL_WEIGHT (!MATERIALBLENDING_MASKED && !MATERIALBLENDING_SOLID && STRATA_LEGACY_PREMULT_ALPHA_OVERRIDE==0 || STRATA_MATERIAL_EXPORT_FROM_TRANSLUCENT)
 
 #define STRATA_INLINE_SINGLELAYERWATER (STRATA_ENABLED && MATERIAL_IS_STRATA && STRATA_INLINE_SHADING && MATERIAL_SHADINGMODEL_SINGLELAYERWATER)
@@ -1620,14 +1622,15 @@ FStrataData GetStrataSlabBSDF(
 		SLAB_GLINT_UVDDX(StrataData.InlinedBSDF)		= ddx(GlintUV);
 		SLAB_GLINT_UVDDY(StrataData.InlinedBSDF)		= ddy(GlintUV);
 	}
-#endif
-	
+#endif // STRATA_GLINTS_ENABLED
+#if STRATA_SPECPROFILE_ENABLED
 	if (SpecularProfileId != 0.0f)
 	{
 		BSDF_SETHASSPECPROFILE	(StrataData.InlinedBSDF, 1);
 		SLAB_SPECPROFILEID		(StrataData.InlinedBSDF) = SpecularProfileId * 255; // Stored as 0-255 index, and sign store parametrization type (View/Light or Half-Angle)
 	}
-#endif
+#endif // STRATA_SPECPROFILE_ENABLED
+#endif // STRATA_COMPLEXSPECIALPATH
 
 #if !STRATA_FASTPATH
 	// See how UMaterialExpressionStrataSlabBSDF node interface is changed
@@ -2441,11 +2444,12 @@ FStrataData StrataHorizontalMixingParameterBlending(FStrataData Background, FStr
 	SLAB_GLINT_VALUE(ResultBSDF) =		lerp(SLAB_GLINT_VALUE(BackgroundBSDF),			SLAB_GLINT_VALUE(ForegroundBSDF),			ForegroundOtherMixFactor);
 	SLAB_GLINT_UVDDX(ResultBSDF) =		lerp(SLAB_GLINT_UVDDX(BackgroundBSDF),			SLAB_GLINT_UVDDX(ForegroundBSDF),			ForegroundOtherMixFactor);
 	SLAB_GLINT_UVDDY(ResultBSDF) =		lerp(SLAB_GLINT_UVDDY(BackgroundBSDF),			SLAB_GLINT_UVDDY(ForegroundBSDF),			ForegroundOtherMixFactor);
-#endif
-
+#endif // STRATA_GLINTS_ENABLED
+#if STRATA_SPECPROFILE_ENABLED
 	BSDF_SETHASSPECPROFILE(ResultBSDF,	BSDF_GETHASSPECPROFILE(ForegroundBSDF)			| BSDF_GETHASSPECPROFILE(BackgroundBSDF));
 	SLAB_SPECPROFILEID(ResultBSDF) = BSDF_GETHASSPECPROFILE(ForegroundBSDF) ? SLAB_SPECPROFILEID(ForegroundBSDF) : SLAB_SPECPROFILEID(BackgroundBSDF);
-#endif
+#endif // STRATA_SPECPROFILE_ENABLED
+#endif // STRATA_COMPLEXSPECIALPATH
 	
 
 	if (BSDF_GETSSSTYPE(ResultBSDF) == SSS_TYPE_DIFFUSION_PROFILE)
@@ -2561,11 +2565,12 @@ FStrataData StrataAddParameterBlending(FStrataData A, FStrataData B, float AMixF
 	SLAB_GLINT_VALUE(ResultBSDF) =		lerp(SLAB_GLINT_VALUE(BBSDF),					SLAB_GLINT_VALUE(ABSDF), AMixFactor);
 	SLAB_GLINT_UVDDX(ResultBSDF) =		lerp(SLAB_GLINT_UVDDX(BBSDF),					SLAB_GLINT_UVDDX(ABSDF), AMixFactor);
 	SLAB_GLINT_UVDDY(ResultBSDF) =		lerp(SLAB_GLINT_UVDDY(BBSDF),					SLAB_GLINT_UVDDY(ABSDF), AMixFactor);
-#endif
-
+#endif // STRATA_GLINTS_ENABLED
+#if STRATA_SPECPROFILE_ENABLED
 	BSDF_SETHASSPECPROFILE(ResultBSDF,	BSDF_GETHASSPECPROFILE(BBSDF) | BSDF_GETHASSPECPROFILE(ABSDF));
 	SLAB_SPECPROFILEID(ResultBSDF) = BSDF_GETHASSPECPROFILE(BBSDF) ? SLAB_SPECPROFILEID(BBSDF) : SLAB_SPECPROFILEID(ABSDF);
-#endif
+#endif // STRATA_SPECPROFILE_ENABLED
+#endif // STRATA_COMPLEXSPECIALPATH
 
 	if (BSDF_GETSSSTYPE(ResultBSDF) == SSS_TYPE_DIFFUSION_PROFILE)
 	{
@@ -2693,11 +2698,12 @@ FStrataBSDF StrataVerticalLayeringParameterBlending(in FStrataBSDF TopBSDF, in F
 	SLAB_GLINT_VALUE(ResultBSDF) =		lerp(SLAB_GLINT_VALUE(TopBSDF),					SLAB_GLINT_VALUE(BaseBSDF), TopToBottomLerp);
 	SLAB_GLINT_UVDDX(ResultBSDF) =		lerp(SLAB_GLINT_UVDDX(TopBSDF),					SLAB_GLINT_UVDDX(BaseBSDF), TopToBottomLerp);
 	SLAB_GLINT_UVDDY(ResultBSDF) =		lerp(SLAB_GLINT_UVDDY(TopBSDF),					SLAB_GLINT_UVDDY(BaseBSDF), TopToBottomLerp);
-#endif
-
+#endif // STRATA_GLINTS_ENABLED
+#if STRATA_SPECPROFILE_ENABLED
 	BSDF_SETHASSPECPROFILE(ResultBSDF,	BSDF_GETHASSPECPROFILE(TopBSDF) | BSDF_GETHASSPECPROFILE(BaseBSDF));
 	SLAB_SPECPROFILEID(ResultBSDF) = 	BSDF_GETHASSPECPROFILE(TopBSDF) ? SLAB_SPECPROFILEID(TopBSDF) : SLAB_SPECPROFILEID(BaseBSDF);
-#endif
+#endif // STRATA_SPECPROFILE_ENABLED
+#endif // STRATA_COMPLEXSPECIALPATH
 
 	// Anisotropy is disabled with vertical blending. 
 	SLAB_ANISOTROPY(ResultBSDF) =		0.f;
@@ -4383,13 +4389,14 @@ FStrataBSDF UnpackStrataBSDFIn(FStrataMaterialContainer StrataBuffer, inout FStr
 			SLAB_GLINT_UVDDX(OutBSDF)	= UnpackFloat2FromUInt(StrataLoadUint1(StrataBuffer, StrataAddressing));
 			SLAB_GLINT_UVDDY(OutBSDF)	= UnpackFloat2FromUInt(StrataLoadUint1(StrataBuffer, StrataAddressing));
 		}
-	#endif //  // STRATA_GLINTS_ENABLED
-
+	#endif // STRATA_GLINTS_ENABLED
+	#if STRATA_SPECPROFILE_ENABLED
 		if (BSDF_GETHASSPECPROFILE(OutBSDF))
 		{
 			UnpackSpecularProfile(StrataLoadUint1(StrataBuffer, StrataAddressing), SLAB_SPECPROFILEID(OutBSDF));
 		}
-	#endif
+	#endif // STRATA_SPECPROFILE_ENABLED
+	#endif // STRATA_COMPLEXSPECIALPATH
 	}
 	break;
 	case STRATA_BSDF_TYPE_HAIR:
diff --git a/Engine/Shaders/Private/Strata/StrataExport.ush b/Engine/Shaders/Private/Strata/StrataExport.ush
--- a/Engine/Shaders/Private/Strata/StrataExport.ush
+++ b/Engine/Shaders/Private/Strata/StrataExport.ush
@@ -191,11 +191,7 @@ void PackStrataOut(
 						const bool LuminanceWeightEqualOne = all(BSDF.LuminanceWeightV == 1.0f);
 
 						#if STRATA_COMPLEXSPECIALPATH
-						#if STRATA_GLINTS_ENABLED
-						const bool bForceComplexSpecialPath = BSDF_GETHASSPECPROFILE(BSDF) || BSDF_GETHASGLINT(BSDF);
-						#else
-						const bool bForceComplexSpecialPath = BSDF_GETHASSPECPROFILE(BSDF);
-						#endif
+						const bool bForceComplexSpecialPath = (STRATA_SPECPROFILE_ENABLED && BSDF_GETHASSPECPROFILE(BSDF)) || (STRATA_GLINTS_ENABLED && BSDF_GETHASGLINT(BSDF));
 						#else
 						const bool bForceComplexSpecialPath = false;
 						#endif
@@ -814,11 +810,12 @@ void PackStrataOut(
 								STRATA_STORE_UINT1(PackFloat2ToUInt(SLAB_GLINT_UVDDY(BSDF)));
 							}
 							#endif // STRATA_GLINTS_ENABLED
-
+							#if STRATA_SPECPROFILE_ENABLED
 							if (BSDF_GETHASSPECPROFILE(BSDF))
 							{
 								STRATA_STORE_UINT1(PackSpecularProfile(SLAB_SPECPROFILEID(BSDF)));
 							}
+							#endif // STRATA_SPECPROFILE_ENABLED
 							#endif // STRATA_COMPLEXSPECIALPATH
 
 							// 8-28 bytes
diff --git a/Engine/Source/Runtime/Engine/Private/Materials/HLSLMaterialTranslator.cpp b/Engine/Source/Runtime/Engine/Private/Materials/HLSLMaterialTranslator.cpp
--- a/Engine/Source/Runtime/Engine/Private/Materials/HLSLMaterialTranslator.cpp
+++ b/Engine/Source/Runtime/Engine/Private/Materials/HLSLMaterialTranslator.cpp
@@ -11555,9 +11555,6 @@ bool FHLSLMaterialTranslator::FStrataCompilationContext::StrataGenerateDerivedMa
 			const uint32 UintByteSize = sizeof(uint32);
 			StrataMaterialRequestedSizeByte = 0;
 
-			// Used to not promote to ComplexSpecial render path if glint rendering is not enabled and avoid registering such possibility at runtime.
-			const bool bGlintEnabled = Strata::IsGlintEnabled();
-
 			// 1. Evaluate simple/single BSDF
 			bStrataMaterialIsSimple = StrataMaterialBSDFCount == 1;
 			bStrataMaterialIsSingle = StrataMaterialBSDFCount == 1;
@@ -11584,7 +11581,7 @@ bool FHLSLMaterialTranslator::FStrataCompilationContext::StrataGenerateDerivedMa
 					{
 						bStrataMaterialIsSimple = bStrataMaterialIsSimple && !bMayHaveColoredWeight && !It.bBSDFHasAnisotropy && !It.bBSDFHasEdgeColor && !It.bBSDFHasFuzz && !It.bBSDFHasSecondRoughnessOrSimpleClearCoat && !It.bBSDFHasMFPPluggedIn && !It.bBSDFHasSSS && !It.bBSDFHasGlint && !It.bBSDFHasSpecularProfile;
 						bStrataMaterialIsSingle = bStrataMaterialIsSingle && !bMayHaveColoredWeight && !It.bBSDFHasAnisotropy && !It.bBSDFHasGlint && !It.bBSDFHasSpecularProfile;
-						bUsesComplexSpecialRenderPath |= (bGlintEnabled && It.bBSDFHasGlint) || It.bBSDFHasSpecularProfile;
+						bUsesComplexSpecialRenderPath |= It.bBSDFHasGlint || It.bBSDFHasSpecularProfile;
 						break;
 					}
 					case STRATA_BSDF_TYPE_HAIR:
diff --git a/Engine/Source/Runtime/Engine/Private/Materials/MaterialExpressions.cpp b/Engine/Source/Runtime/Engine/Private/Materials/MaterialExpressions.cpp
--- a/Engine/Source/Runtime/Engine/Private/Materials/MaterialExpressions.cpp
+++ b/Engine/Source/Runtime/Engine/Private/Materials/MaterialExpressions.cpp
@@ -23709,12 +23709,14 @@ bool UMaterialExpressionStrataSlabBSDF::HasAnisotropy() const
 
 bool UMaterialExpressionStrataSlabBSDF::HasGlint() const
 {
-	return GlintValue.IsConnected();
+	// Use IsGlintEnabled to not promote to ComplexSpecial render path if glint rendering is not enabled and avoid registering such possibility at runtime.
+	return GlintValue.IsConnected() && Strata::IsGlintEnabled();
 }
 
 bool UMaterialExpressionStrataSlabBSDF::HasSpecularProfile() const
 {
-	return SpecularProfile != nullptr;
+	// Use IsSpecularLUTEnabled to not promote to ComplexSpecial render path if glint rendering is not enabled and avoid registering such possibility at runtime.
+	return SpecularProfile != nullptr && Strata::IsSpecularProfileEnabled();
 }
 
 #endif // WITH_EDITOR
diff --git a/Engine/Source/Runtime/Engine/Private/ShaderCompiler/ShaderCompiler.cpp b/Engine/Source/Runtime/Engine/Private/ShaderCompiler/ShaderCompiler.cpp
--- a/Engine/Source/Runtime/Engine/Private/ShaderCompiler/ShaderCompiler.cpp
+++ b/Engine/Source/Runtime/Engine/Private/ShaderCompiler/ShaderCompiler.cpp
@@ -6665,6 +6665,9 @@ void GlobalBeginCompileShader(
 
 			const bool bStrataGlints = Strata::IsGlintEnabled();
 			Input.Environment.SetDefine(TEXT("PLATFORM_ENABLES_STRATA_GLINTS"), bStrataGlints ? 1 : 0);
+
+			const bool bSpecularProfileEnabled = Strata::IsSpecularProfileEnabled();
+			Input.Environment.SetDefine(TEXT("PLATFORM_ENABLES_STRATA_SPECULAR_PROFILE"), bSpecularProfileEnabled ? 1 : 0);
 		}
 
 		const bool bStrataBackCompatibility = bStrata && Strata::IsBackCompatibilityEnabled();
diff --git a/Engine/Source/Runtime/RenderCore/Private/RenderUtils.cpp b/Engine/Source/Runtime/RenderCore/Private/RenderUtils.cpp
--- a/Engine/Source/Runtime/RenderCore/Private/RenderUtils.cpp
+++ b/Engine/Source/Runtime/RenderCore/Private/RenderUtils.cpp
@@ -1604,11 +1604,17 @@ static TAutoConsoleVariable<int32> CVarSubstrateGlints(
 	ECVF_ReadOnly | ECVF_RenderThreadSafe);
 
 static TAutoConsoleVariable<int32> CVarSubstrateGlintsLUT(
-	TEXT("r.Substrate.GlintsLUT"),
+	TEXT("r.Substrate.Glints.LUT"),
 	1,
 	TEXT("Select one of the glint rendering LUT for testing."),
 	ECVF_RenderThreadSafe);
 
+static TAutoConsoleVariable<int32> CVarSubstrateSpecularProfile(
+	TEXT("r.Substrate.SpecularProfile"),
+	1,
+	TEXT("Enable Specular Profile support for Strata slabs. If changed, shaders needs to be recompiled."),
+	ECVF_ReadOnly | ECVF_RenderThreadSafe);
+
 static TAutoConsoleVariable<int32> CVarSubstrateDebugAdvancedVisualizationShaders(
 	TEXT("r.Substrate.Debug.AdvancedVisualizationShaders"),
 	0,
@@ -1718,6 +1724,11 @@ namespace Strata
 		return CVarSubstrateGlintsLUT.GetValueOnAnyThread()<= 0 ? 0u : 1u;
 	}
 
+	bool IsSpecularProfileEnabled()
+	{
+		return IsStrataEnabled() && CVarSubstrateSpecularProfile.GetValueOnAnyThread() > 0;
+	}
+
 	bool Is8bitTileCoordEnabled()
 	{
 		return CVarSubstrateTileCoord8Bits.GetValueOnAnyThread() > 0 ? 1 : 0;
diff --git a/Engine/Source/Runtime/RenderCore/Private/Shader.cpp b/Engine/Source/Runtime/RenderCore/Private/Shader.cpp
--- a/Engine/Source/Runtime/RenderCore/Private/Shader.cpp
+++ b/Engine/Source/Runtime/RenderCore/Private/Shader.cpp
@@ -2019,6 +2019,11 @@ void ShaderMapAppendKeyString(EShaderPlatform Platform, FString& KeyString)
 		{
 			KeyString += FString::Printf(TEXT("_STRTGLT"));
 		}
+
+		if (Strata::IsSpecularProfileEnabled())
+		{
+			KeyString += FString::Printf(TEXT("_STRTSP"));
+		}
 	}
 
 	{
diff --git a/Engine/Source/Runtime/RenderCore/Public/RenderUtils.h b/Engine/Source/Runtime/RenderCore/Public/RenderUtils.h
--- a/Engine/Source/Runtime/RenderCore/Public/RenderUtils.h
+++ b/Engine/Source/Runtime/RenderCore/Public/RenderUtils.h
@@ -469,6 +469,7 @@ namespace Strata
 	RENDERCORE_API bool IsRoughDiffuseEnabled();
 	RENDERCORE_API bool IsGlintEnabled();
 	RENDERCORE_API uint32 GlintLUTIndex();
+	RENDERCORE_API bool IsSpecularProfileEnabled();
 	RENDERCORE_API bool IsBackCompatibilityEnabled();
 	RENDERCORE_API bool IsDBufferPassEnabled(EShaderPlatform InPlatform);
 	RENDERCORE_API bool IsOpaqueRoughRefractionEnabled();
