Substrate - remove the bUsesComplexSpecialRenderPath memeber from the header and simply reuse the State uint.

diff --git a/Engine/Shaders/Private/Strata/Strata.ush b/Engine/Shaders/Private/Strata/Strata.ush
--- a/Engine/Shaders/Private/Strata/Strata.ush
+++ b/Engine/Shaders/Private/Strata/Strata.ush
@@ -219,8 +219,8 @@ struct FStrataBSDF
 // 
 // Layouts:
 // * LAYOUT 2 -------------------------------------------------------------------------------------------------------------------------------------------------------------------------
-//  (Complex) HEADER        27 |  State | AO | BSDF Count | Basis count | Basis type |
-//                             |  8       8    4            3             4           
+//  (Complex) HEADER        28 |  State | AO | BSDF Count | Basis count | Basis type | ComplexSpecialPath
+//                             |  8       8    4            3             4            1
 //            BASIS(0)      32 |  Normal: Oct 16/16 | Tangent: 11/11/10
 //            BASIS(1)      32 |  Normal: Oct 16/16 | Tangent: 11/11/10
 //                             |
@@ -232,7 +232,7 @@ struct FStrataBSDF
 //            BSDF(1)       ?? |  Adaptive
 //                             |
 // * LAYOUT 1 -------------------------------------------------------------------------------------------------------------------------------------------------------------------------
-//  (Single)  HEADER        27 |  State | AO | Opt Layout | Aniso | TopLayer | SSSType | GreyWeighV | Haze | F90 | Simple vol. | MFP Plugged | Fuzz | Thin
+//  (Single)  HEADER        30 |  State | AO | Opt Layout | Aniso | TopLayer | SSSType | GreyWeighV | Haze | F90 | Simple vol. | MFP Plugged | Fuzz | Thin
 //                             |  8       8    3            1       1          2         1            1      1     1             1             1      1            
 //            BASIS(0)       0 | (Use top layer normal)
 //            BSDF(0)       ?? |  Adaptive
@@ -297,11 +297,13 @@ struct FStrataBSDF
 // For HEADER_BIT_BSDF_COUNT and  HEADER_BIT_SHAREDLOCALBASES_X, please refer to values in StrataDefinitions.h
 #define HEADER_BIT_COUNT_STATE						HEADER_MASK_BIT_COUNT
 #define HEADER_BIT_COUNT_IRRADIANCE_AO				8
+// The remaining header bits are only available for the complex path.
 #define HEADER_BIT_COUNT_BSDF_COUNT					4
 #define HEADER_BIT_COUNT_SHAREDLOCALBASES_COUNT		3
 #define HEADER_BIT_COUNT_SHAREDLOCALBASES_TYPE		4
 #define HEADER_BIT_COUNT_COMPLEXPATH_SPECIAL		1
 #define HEADER_BIT_COUNT___UNUSED___				4
+// End of header bits.
 
 #define HEADER_BIT_OFFSET_STATE						0
 #define HEADER_BIT_OFFSET_IRRADIANCE_AO				(HEADER_BIT_OFFSET_STATE					+ HEADER_BIT_COUNT_STATE)
@@ -334,7 +336,7 @@ struct FStrataBSDF
 #define HEADER_HAIRENCODING_MASK					((1u<<HEADER_HAIRENCODING_BIT_COUNT)-1u)
 // Layout Eye
 #define HEADER_SLWENCODING_BIT_COUNT				(HEADER_BIT_COUNT_STATE + HEADER_BIT_COUNT_IRRADIANCE_AO)
-#define HEADER_SLWENCODING_MASK					((1u<<HEADER_SLWENCODING_BIT_COUNT)-1u)
+#define HEADER_SLWENCODING_MASK						((1u<<HEADER_SLWENCODING_BIT_COUNT)-1u)
 
 // A header always has a common state represented by stat bits (low 8 bits) and Irradiance/AO bits (higher 8 bits)
 #define HEADER_GETCOMMONSTATES(X)					READ_BITS( X, (HEADER_BIT_COUNT_STATE + HEADER_BIT_COUNT_IRRADIANCE_AO), HEADER_BIT_OFFSET_STATE)
@@ -357,6 +359,7 @@ struct FStrataBSDF
 
 #define HEADER_GETCOMPLEXSPECIALPATH(X)				READ_BITS( X, HEADER_BIT_COUNT_COMPLEXPATH_SPECIAL, HEADER_BIT_OFFSET_COMPLEXPATH_SPECIAL)
 #define HEADER_SETCOMPLEXSPECIALPATH(X, NewValue)	WRITE_BITS(X, HEADER_BIT_COUNT_COMPLEXPATH_SPECIAL, HEADER_BIT_OFFSET_COMPLEXPATH_SPECIAL, NewValue)
+#define HEADER_COMPLEXSPECIALPATH_MASK				TO_OFFSET_BITMASK(HEADER_BIT_COUNT_COMPLEXPATH_SPECIAL, HEADER_BIT_OFFSET_COMPLEXPATH_SPECIAL)
 
 #define HEADER_MATERIALMODE_NONE				0
 #define HEADER_MATERIALMODE_SLAB_SIMPLE			1
@@ -2089,8 +2092,6 @@ struct FStrataPixelHeader
 
 	uint State;
 
-	bool bUsesComplexSpecialRenderPath; // STRATA_TODO pack with someting else, in State maybe?
-
 #if STRATA_INLINE_SHADING
 	// All the shared local bases in VGPRS. This is only used in basepass and during forward rendering.
 	FSharedLocalBases SharedLocalBases;
@@ -2159,19 +2160,20 @@ struct FStrataPixelHeader
 	bool HasPrecShadowMask()									{ return (this.State & HEADER_MASK_HASPRECSHADOWMASK) != 0; }
 	bool HasZeroPrecShadowMask()								{ return (this.State & HEADER_MASK_ZEROPRECSHADOWMASK) != 0; }
 	bool DoesCastContactShadow()								{ return (this.State & HEADER_MASK_CASTCONTACTSHADOW) != 0; }
-	bool IsStrataMaterial()										{ return this.GetMaterialMode() != HEADER_MATERIALMODE_NONE; }
 	bool HasSubsurface()										{ return (this.State & HEADER_MASK_HASSUBSURFACE) != 0; }
 	bool HasDynamicIndirectShadowCasterRepresentation()			{ return (this.State & HEADER_MASK_HASDYNINDIRECTSHADOWCASTER) != 0; }
 
 	void SetHasPrecShadowMask(bool bIn)							{ this.State |= (bIn ? HEADER_MASK_HASPRECSHADOWMASK : 0u); }
 	void SetZeroPrecShadowMask(bool bIn)						{ this.State |= (bIn ? HEADER_MASK_ZEROPRECSHADOWMASK : 0u); }
 	void SetCastContactShadow(bool bIn)							{ this.State |= (bIn ? HEADER_MASK_CASTCONTACTSHADOW : 0u); }
-	void SetDynamicIndirectShadowCasterRepresentation(bool bIn)	{ this.State |= (bIn ? HEADER_MASK_HASDYNINDIRECTSHADOWCASTER : 0u); }
 	void SetHasSubsurface(bool bIn)								{ this.State |= (bIn ? HEADER_MASK_HASSUBSURFACE : 0u); }
+	void SetDynamicIndirectShadowCasterRepresentation(bool bIn)	{ this.State |= (bIn ? HEADER_MASK_HASDYNINDIRECTSHADOWCASTER : 0u); }
 
+	bool IsStrataMaterial()										{ return this.GetMaterialMode() != HEADER_MATERIALMODE_NONE; }
 	void SetMaterialMode(uint MaterialMode)						{ this.State = (this.State & ~HEADER_MASK_MATERIALMODE) | (MaterialMode & HEADER_MASK_MATERIALMODE); }
 
-	void SetUsesComplexSpecialRenderPath(bool bIn)				{ this.bUsesComplexSpecialRenderPath |= bIn; }
+	bool GetUsesComplexSpecialRenderPath()						{ return (this.State & HEADER_COMPLEXSPECIALPATH_MASK) != 0; }
+	void SetUsesComplexSpecialRenderPath(bool bIn)				{ this.State |= (bIn ? HEADER_COMPLEXSPECIALPATH_MASK : 0u); }
 };
 
 #if STRATA_DEFERRED_SHADING && STRATA_MATERIALCONTAINER_IS_VIEWRESOURCE
@@ -2183,7 +2185,6 @@ FStrataPixelHeader InitialiseStrataPixelHeader()
 {
 	FStrataPixelHeader Out;
 	Out.BSDFCount = 0;
-	Out.bUsesComplexSpecialRenderPath = false;
 #if STRATA_INLINE_SHADING
 	Out.SharedLocalBases = (FSharedLocalBases)0;
 	Out.StrataTree = GetInitialisedStrataTree();
@@ -3873,7 +3874,7 @@ uint PackStrataHeader(uint InBSDFCount, FStrataPixelHeader InHeader)
 #endif
 
 #if STRATA_COMPLEXPATH
-	HEADER_SETCOMPLEXSPECIALPATH(Out, InHeader.bUsesComplexSpecialRenderPath ? 1 : 0);
+	Out |= InHeader.State & HEADER_COMPLEXSPECIALPATH_MASK;
 #endif
 
 	HEADER_SETCOMMONSTATES(Out, InHeader.State);
@@ -3911,7 +3912,7 @@ FStrataPixelHeader UnpackStrataHeaderIn(uint PackedHeader, inout FStrataAddressi
 		uint SharedLocalBasesCount = HEADER_GETSHAREDLOCALBASESCOUNT(PackedHeader); // ?
 		uint SharedLocalBasesTypes = HEADER_GETSHAREDLOCALBASESTYPE(PackedHeader);
 
-		Out.bUsesComplexSpecialRenderPath = HEADER_GETCOMPLEXSPECIALPATH(PackedHeader) > 0;
+		Out.SetUsesComplexSpecialRenderPath((PackedHeader & HEADER_COMPLEXSPECIALPATH_MASK) > 0);
 
 	#if STRATA_DEFERRED_SHADING
 		Out.SharedLocalBasesTypes_PackedHeader = SharedLocalBasesTypes;
diff --git a/Engine/Shaders/Private/Strata/StrataMaterialClassification.usf b/Engine/Shaders/Private/Strata/StrataMaterialClassification.usf
--- a/Engine/Shaders/Private/Strata/StrataMaterialClassification.usf
+++ b/Engine/Shaders/Private/Strata/StrataMaterialClassification.usf
@@ -128,11 +128,12 @@ void TileMainCS(uint2 DispatchThreadId : SV_DispatchThreadID, uint LinearIndex :
 
 		const bool bIsSimple =  StrataPixelHeader.IsSimpleMaterial() || StrataPixelHeader.BSDFCount == 0; // BSDFCount == 0 ensures that non-strata pixel, like sky pixels, won't make a simple tile flagged as complex
 		const bool bIsSingle = !StrataPixelHeader.IsSimpleMaterial() && StrataPixelHeader.IsSingleMaterial();
+		const bool bIsComplexSpecial = StrataPixelHeader.GetUsesComplexSpecialRenderPath();
 		bContainsStrataMaterial  = StrataPixelHeader.BSDFCount > 0;
 		bContainsSimpleMaterial  = bIsSimple;
 		bContainsSingleMaterial  = bIsSingle;
-		bContainsComplexMaterial		= !bIsSingle && !bIsSimple && !StrataPixelHeader.bUsesComplexSpecialRenderPath;
-		bContainsComplexSpecialMaterial = !bIsSingle && !bIsSimple &&  StrataPixelHeader.bUsesComplexSpecialRenderPath;
+		bContainsComplexMaterial		= !bIsSingle && !bIsSimple && !bIsComplexSpecial;
+		bContainsComplexSpecialMaterial = !bIsSingle && !bIsSimple &&  bIsComplexSpecial;
 
 		bContainsScreenSpaceSubsurfaceScattering = StrataPixelHeader.HasSubsurface();
 
diff --git a/Engine/Shaders/Private/Strata/StrataVisualize.usf b/Engine/Shaders/Private/Strata/StrataVisualize.usf
--- a/Engine/Shaders/Private/Strata/StrataVisualize.usf
+++ b/Engine/Shaders/Private/Strata/StrataVisualize.usf
@@ -92,14 +92,14 @@ RWBuffer<uint> RWPositionOffsetBuffer;
 
 void PrintPixelType(inout FShaderPrintContext Ctx, in FStrataPixelHeader Header, FFontColor InColor)
 {
-	if (Header.IsSimpleMaterial())					Print(Ctx, TEXT("Simple  "), InColor);
-	else if (Header.IsSingleMaterial())				Print(Ctx, TEXT("Single  "), InColor);
-	else if (Header.bUsesComplexSpecialRenderPath)	Print(Ctx, TEXT("Complex Special "), InColor);
-	else if (Header.IsComplexMaterial())			Print(Ctx, TEXT("Complex "), InColor);
-	else if (Header.IsSingleLayerWater())			Print(Ctx, TEXT("SLWater "), InColor);
-	else if (Header.IsHair())						Print(Ctx, TEXT("Hair    "), InColor);
-	else if (Header.IsEye())						Print(Ctx, TEXT("Eye     "), InColor);
-	else											Print(Ctx, TEXT("Error - Unkown pixel type"), FontRed);
+	if (Header.IsSimpleMaterial())						Print(Ctx, TEXT("Simple  "), InColor);
+	else if (Header.IsSingleMaterial())					Print(Ctx, TEXT("Single  "), InColor);
+	else if (Header.GetUsesComplexSpecialRenderPath())	Print(Ctx, TEXT("Complex Special "), InColor);
+	else if (Header.IsComplexMaterial())				Print(Ctx, TEXT("Complex "), InColor);
+	else if (Header.IsSingleLayerWater())				Print(Ctx, TEXT("SLWater "), InColor);
+	else if (Header.IsHair())							Print(Ctx, TEXT("Hair    "), InColor);
+	else if (Header.IsEye())							Print(Ctx, TEXT("Eye     "), InColor);
+	else												Print(Ctx, TEXT("Error - Unkown pixel type"), FontRed);
 }
 
 void PrintSSSType(inout FShaderPrintContext Ctx, uint SSSType, bool bThin, FFontColor InColor)
