Fix typo use of DECAL_RENDERTARGET instead of DECAL_RENDERTARGETMODE.
Fixes issue with oversaturated emissive GBuffer decals.

diff --git a/Engine/Shaders/Private/DecalCommon.ush b/Engine/Shaders/Private/DecalCommon.ush
--- a/Engine/Shaders/Private/DecalCommon.ush
+++ b/Engine/Shaders/Private/DecalCommon.ush
@@ -95,7 +95,7 @@ void DecalCommonOutput(inout FPixelShaderIn In, inout FPixelShaderOut Out, float
 	Color = lerp(1, Data.GBufferAO, Opacity);
 #endif
 
-	#if DECAL_RENDERTARGETMODE == DECAL_RENDERTARGETMODE_SCENECOLOR || DECAL_RENDERTARGET == DECAL_RENDERTARGETMODE_GBUFFER || DECAL_RENDERTARGET == DECAL_RENDERTARGETMODE_GBUFFER_NONORMAL
+	#if DECAL_RENDERTARGETMODE == DECAL_RENDERTARGETMODE_SCENECOLOR || DECAL_RENDERTARGETMODE == DECAL_RENDERTARGETMODE_GBUFFER || DECAL_RENDERTARGETMODE == DECAL_RENDERTARGETMODE_GBUFFER_NONORMAL
 		Color *= View.PreExposure;
 	#endif
 
