Path Tracing: Accelerating heterogeneous volume tracking with the majorant grid

diff --git a/Engine/Shaders/Private/HeterogeneousVolumes/HeterogeneousVolumesVoxelGridUtils.ush b/Engine/Shaders/Private/HeterogeneousVolumes/HeterogeneousVolumesVoxelGridUtils.ush
--- a/Engine/Shaders/Private/HeterogeneousVolumes/HeterogeneousVolumesVoxelGridUtils.ush
+++ b/Engine/Shaders/Private/HeterogeneousVolumes/HeterogeneousVolumesVoxelGridUtils.ush
@@ -26,12 +26,28 @@ void UpdateMajorantData(inout FMajorantData MajorantData, float3 Value)
 	MajorantData.Majorant = max(MajorantData.Majorant, MaxValue);
 }
 
+FMajorantData CreateMajorantData(float3 MinValue, float3 MaxValue)
+{
+	FMajorantData MajorantData = CreateMajorantData();
+	UpdateMajorantData(MajorantData, MinValue);
+	UpdateMajorantData(MajorantData, MaxValue);
+	return MajorantData;
+}
+
+// Updating majorant data should occur for a single media
 void UpdateMajorantData(inout FMajorantData MajorantData, FMajorantData Value)
 {
 	MajorantData.Minorant = min(MajorantData.Minorant, Value.Minorant);
 	MajorantData.Majorant = max(MajorantData.Majorant, Value.Majorant);
 }
 
+// Merging majorant data should occur when mixing media
+void MergeMajorantData(inout FMajorantData MajorantData, FMajorantData OverlappingMajorant)
+{
+	MajorantData.Minorant += OverlappingMajorant.Minorant;
+	MajorantData.Majorant += OverlappingMajorant.Majorant;
+}
+
 FMajorantData GetMajorantData(FScalarGridData MajorantGridData)
 {
 	FMajorantData Data;
diff --git a/Engine/Shaders/Private/PathTracing/PathTracingCore.ush b/Engine/Shaders/Private/PathTracing/PathTracingCore.ush
--- a/Engine/Shaders/Private/PathTracing/PathTracingCore.ush
+++ b/Engine/Shaders/Private/PathTracing/PathTracingCore.ush
@@ -553,33 +553,16 @@ FPathTracingPayload TraceTransparentRay(inout FPathState PathState, int Bounce,
 				// Limit number of steps to prevent timeouts // FIXME: This biases the result!
 				for (int Step = 0; Step < MaxRaymarchSteps; Step++)
 				{
-					// figure out optimal bounds for the current ray segment
-					float3 SigmaBar = VolumeGetDensityBounds(PathState.Ray.Origin, PathState.Ray.Direction, Interval).SigmaMax;
-					if (UseNullTrackingTransmittance)
+					FVolumeTrackingResult TrackingResult = VolumeSampleDistance(PathThroughput, PathState.Ray.Origin, PathState.Ray.Direction, Interval, bIsCameraRay, PathState.RandSequence);
+					if (TrackingResult.Distance < 0)
 					{
-						// if we are not using analytical transmittance, a tight majorant could prevent us from "seeing" hits that match the density exactly
-						// leading to heavy noise on bright objects embedded in the volume
-						// however if we can track transmittance analytically, this workaround is not needed as we will re-compute a precise answer after ray marching
-						SigmaBar *= bIsCameraRay ? 2 : 1;
-					}
-
-					// there is some atmosphere present along the current ray, ray march until we get to the current hit (or exit the atmosphere)
-					float3 ColorChannelPdf = PathThroughput;
-
-					// Sample the distance to the next interaction
-					float RandValue = RandomSequence_GenerateSample1D(PathState.RandSequence);
-					float DeltaT = SampleSpectralTransmittance(RandValue, SigmaBar, ColorChannelPdf);
-					if (DeltaT < 0.0)
-					{
-						// no more energy left in the path
 						break;
 					}
 
-					if (Interval.VolumeTMin + DeltaT < Interval.VolumeTMax)
+					if (TrackingResult.bIsCollision)
 					{
-						Interval.VolumeTMin += DeltaT;
-						// our ray marching step stayed inside the atmo and is still in front of the next hit
-						PathThroughput *= EvaluateSpectralTransmittanceHit(DeltaT, SigmaBar, ColorChannelPdf).xyz;
+						Interval.VolumeTMin = TrackingResult.Distance;
+						PathThroughput = TrackingResult.Throughput;
 
 						// find out how much volume exists at the current point
 						float3 Ro = PathState.Ray.Origin;
@@ -588,8 +571,8 @@ FPathTracingPayload TraceTransparentRay(inout FPathState PathState, int Bounce,
 						FVolumeShadedResult Result = VolumeGetDensity(TranslatedWorldPos, Interval);
 
 						// clamp to make sure we never exceed the majorant (should not be the case, but need to avoid any possible numerical issues)
-						float3 SigmaT = min(Result.SigmaT, SigmaBar);
-						float3 SigmaN = SigmaBar - SigmaT;
+						float3 SigmaT = min(Result.SigmaT, TrackingResult.SigmaBar);
+						float3 SigmaN = TrackingResult.SigmaBar - SigmaT;
 
 						if (UseAnalyticTransmittance || (UseNullTrackingTransmittance && !bLastBounce)) // check now if this loop is also computing transmittance
 						{
@@ -612,7 +595,7 @@ FPathTracingPayload TraceTransparentRay(inout FPathState PathState, int Bounce,
 					else
 					{
 						// update the path throughput, knowing that we escaped the medium
-						PathThroughput *= EvaluateSpectralTransmittanceMiss(Interval.VolumeTMax - Interval.VolumeTMin, SigmaBar, ColorChannelPdf).xyz;
+						PathThroughput = TrackingResult.Throughput;
 						// exit the ray marching loop
 						break;
 					}
diff --git a/Engine/Shaders/Private/PathTracing/Volume/PathTracingHeterogeneousVolumes.ush b/Engine/Shaders/Private/PathTracing/Volume/PathTracingHeterogeneousVolumes.ush
--- a/Engine/Shaders/Private/PathTracing/Volume/PathTracingHeterogeneousVolumes.ush
+++ b/Engine/Shaders/Private/PathTracing/Volume/PathTracingHeterogeneousVolumes.ush
@@ -4,6 +4,8 @@
 
 #include "../../HeterogeneousVolumes/HeterogeneousVolumesTracingUtils.ush"
 #include "../../HeterogeneousVolumes/HeterogeneousVolumesVoxelGridTraversal.ush"
+#include "../Utilities/PathTracingRandomSequence.ush"
+#include "PathTracingVolumeSampling.ush"
 
 float3 GetTranslatedWorldPos(float3 WorldPos)
 {
@@ -393,40 +395,261 @@ FVolumeShadedResult HeterogeneousVolumesGetDensity(float3 TranslatedWorldPos)
 	return Result;
 }
 
-float3 HeterogeneousVolumesGetTransmittance(float3 Origin, float3 Direction, float TMin, float TMax, float PathRoughness)
+struct FVolumeTrackingResult
 {
-	float3 Transmittance = 1.0;
+	float3 Throughput;
+	float Pdf;
+	float Distance;
+	float SigmaBar;
+	bool bIsCollision;
+};
+
+FVolumeTrackingResult CreateVolumeTrackingResult(float Distance)
+{
+	FVolumeTrackingResult Result = (FVolumeTrackingResult)0;
+	Result.Throughput = 1;
+	Result.Pdf = 0;
+	Result.Distance = Distance;
+	Result.SigmaBar = 0;
+	Result.bIsCollision = false;
 
-	if (EnableHeterogeneousVolumes)
+	return Result;
+}
+
+// Concatenates the result of rhs onto lhs
+FVolumeTrackingResult ConcatenateVolumeTrackingResult(FVolumeTrackingResult lhs, FVolumeTrackingResult rhs)
+{
+	FVolumeTrackingResult Result;
+	Result.Throughput = lhs.Throughput * rhs.Throughput;
+	Result.Pdf = lhs.Pdf * rhs.Pdf;
+	Result.Distance = lhs.Distance + rhs.Distance;
+	Result.SigmaBar = rhs.SigmaBar;
+	Result.bIsCollision = rhs.bIsCollision;
+	return Result;
+}
+
+FVolumeTrackingResult HeterogeneousVolumesDeltaTracking(float3 PathThroughput, float3 Origin, float3 Direction, float TMax, FMajorantData MajorantData, inout RandomSequence RandSequence)
+{
+	FVolumeTrackingResult Sample = CreateVolumeTrackingResult(0);
+	Sample.SigmaBar = max(MajorantData.Majorant, 0);
+
+	while ((Sample.Distance < TMax) && !Sample.bIsCollision)
+	{
+		float RandValue = RandomSequence_GenerateSample1D(RandSequence);
+#define SAMPLE_ACHROMATIC 0
+#if SAMPLE_ACHROMATIC
+		float DeltaT = -log(RandValue) / Sample.SigmaBar;
+#else
+		float DeltaT = SampleSpectralTransmittance(RandValue, Sample.SigmaBar, PathThroughput);
+		if (DeltaT < 0)
+		{
+			Sample.Throughput = 0;
+			Sample.Distance = -1;
+			return Sample;
+		}
+#endif
+
+		Sample.Distance = min(Sample.Distance + DeltaT, TMax);
+		float Transmittance = exp(-Sample.SigmaBar * Sample.Distance);
+
+		if (Sample.Distance < TMax)
+		{
+			Sample.bIsCollision = true;
+			float CollisionWeight = 1.0;
+#if SAMPLE_ACHROMATIC
+			Sample.Throughput *= CollisionWeight / Sample.SigmaBar;
+			Sample.Pdf = Transmittance * Sample.SigmaBar;
+#else
+			float4 EvalResult = EvaluateSpectralTransmittanceHit(Sample.Distance, Sample.SigmaBar, PathThroughput);
+			Sample.Throughput = CollisionWeight * EvalResult.xyz;
+			Sample.Pdf = EvalResult.w;
+		}
+		else
+		{
+			float4 EvalResult = EvaluateSpectralTransmittanceMiss(Sample.Distance, Sample.SigmaBar, PathThroughput);
+			Sample.Throughput = EvalResult.xyz;
+			Sample.Pdf = EvalResult.w;
+#endif
+		}
+	}
+
+	return Sample;
+}
+
+FVolumeTrackingResult HeterogeneousVolumesRatioTracking(float3 PathThroughput, float3 Origin, float3 Direction, float TMax, FMajorantData MajorantData, inout RandomSequence RandSequence)
+{
+	FVolumeTrackingResult Sample = CreateVolumeTrackingResult(0);
+	Sample.SigmaBar = max(MajorantData.Majorant, 0);
+
+	while ((Sample.Distance < TMax) && !Sample.bIsCollision)
 	{
-		bool bCombinedTraversal = OrthoGridUniformBuffer.bUseOrthoGrid && FrustumGridUniformBuffer.bUseFrustumGrid;
-		//if (bCombinedTraversal && Bounce < 1)
-		if (bCombinedTraversal)
+		float RandValue = RandomSequence_GenerateSample1D(RandSequence);
+#if SAMPLE_ACHROMATIC
+		float DeltaT = -log(RandValue) / Sample.SigmaBar;
+#else
+		float DeltaT = SampleSpectralTransmittance(RandValue, Sample.SigmaBar, PathThroughput);
+		if (DeltaT < 0)
+		{
+			Sample.Throughput = 0;
+			Sample.Distance = -1;
+			return Sample;
+		}
+#endif
+
+		Sample.Distance = min(Sample.Distance + DeltaT, TMax);
+		float Transmittance = exp(-Sample.SigmaBar * Sample.Distance);
+
+		if (Sample.Distance < TMax)
 		{
-			float FrustumTMin = TMin;
-			float FrustumTMax = (PathRoughness < 0.1) ? TMax : min(TMax, 10.0);
-			Transmittance = TraceFrustumVoxelGrid(Origin, Direction, FrustumTMin, FrustumTMax);
+			float3 SamplePos = Origin + Direction * Sample.Distance;
+			float3 SigmaT = HeterogeneousVolumesGetDensity(SamplePos).SigmaT;
+			float3 SpectralNullCollisionWeight = max(Sample.SigmaBar - SigmaT, 0);
+#if SAMPLE_ACHROMATIC
+			Sample.Throughput *= SpectralNullCollisionWeight / Sample.SigmaBar;
+			Sample.Pdf *= Transmittance * Sample.SigmaBar;
+#else
+			float4 EvalResult = EvaluateSpectralTransmittanceHit(Sample.Distance, Sample.SigmaBar, PathThroughput);
+			Sample.Throughput *= SpectralNullCollisionWeight * EvalResult.xyz;
+			Sample.Pdf *= EvalResult.w;
+		}
+		else
+		{
+			float4 EvalResult = EvaluateSpectralTransmittanceMiss(Sample.Distance, Sample.SigmaBar, PathThroughput);
+			Sample.Throughput *= EvalResult.xyz;
+			Sample.Pdf *= EvalResult.w;
+#endif
+		}
+	}
+
+	return Sample;
+}
+
+FVolumeTrackingResult HeterogeneousVolumesMajorantBasedTracking(float3 PathThroughput, float3 WorldRayOrigin, float3 Direction, float TMin, float TMax, int bTransmittanceOnly, FMajorantData OverlappingMajorant, inout RandomSequence RandSequence)
+{
+	FVolumeTrackingResult Result = CreateVolumeTrackingResult(TMin);
+	Result.Throughput = PathThroughput;
+
+	float3 WorldRayBegin = WorldRayOrigin + Direction * TMin;
+	float3 WorldRayEnd = WorldRayOrigin + Direction * TMax;
+	float WorldRayTMax = length(WorldRayEnd - WorldRayBegin);
+
+	// Transform to voxel-space
+	float3 WorldBoundsMin = GetTranslatedWorldPos(OrthoGridUniformBuffer.TopLevelGridWorldBoundsMin);
+	float3 WorldBoundsMax = GetTranslatedWorldPos(OrthoGridUniformBuffer.TopLevelGridWorldBoundsMax);
+	float3 TopLevelGridWorldBoundsExtent = WorldBoundsMax - WorldBoundsMin;
+	float3 VoxelRayBegin = (WorldRayBegin - WorldBoundsMin) / TopLevelGridWorldBoundsExtent * OrthoGridUniformBuffer.TopLevelGridResolution;
+	float3 VoxelRayEnd = (WorldRayEnd - WorldBoundsMin) / TopLevelGridWorldBoundsExtent * OrthoGridUniformBuffer.TopLevelGridResolution;
+	float3 VoxelRayDirection = VoxelRayEnd - VoxelRayBegin;
+	float VoxelRayTMin = 0.0;
+	float VoxelRayTMax = length(VoxelRayDirection);
+	VoxelRayDirection /= VoxelRayTMax;
+
+	float VoxelToWorldScale = WorldRayTMax / VoxelRayTMax;
+
+	// March majorant grid via DDA
+	float3 VoxelRayDirectionInv = 1.0 / VoxelRayDirection;
+
+	float VoxelRayMarchT = VoxelRayTMin;
+	float3 VoxelRayMarchPos = VoxelRayBegin;
+
+	float3 VoxelBoundsPos;
+	VoxelBoundsPos.x = sign(VoxelRayDirection.x) > 0 ? floor(VoxelRayMarchPos.x) + 1 : ceil(VoxelRayMarchPos.x) - 1;
+	VoxelBoundsPos.y = sign(VoxelRayDirection.y) > 0 ? floor(VoxelRayMarchPos.y) + 1 : ceil(VoxelRayMarchPos.y) - 1;
+	VoxelBoundsPos.z = sign(VoxelRayDirection.z) > 0 ? floor(VoxelRayMarchPos.z) + 1 : ceil(VoxelRayMarchPos.z) - 1;
+
+	float3 VoxelBoundsHitT = (VoxelBoundsPos - VoxelRayMarchPos) * VoxelRayDirectionInv;
 
-			if (FrustumTMin < FrustumTMax)
+	// Remove floating-point rounding error
+	float Epsilon = 1.0e-4;
+	if (VoxelBoundsHitT.x <= Epsilon) VoxelBoundsHitT.x += abs(VoxelRayDirectionInv.x);
+	if (VoxelBoundsHitT.y <= Epsilon) VoxelBoundsHitT.y += abs(VoxelRayDirectionInv.y);
+	if (VoxelBoundsHitT.z <= Epsilon) VoxelBoundsHitT.z += abs(VoxelRayDirectionInv.z);
+
+	while ((VoxelRayMarchT < VoxelRayTMax) && any(Result.Throughput > Epsilon) && !Result.bIsCollision)
+	{
+		float VoxelRayMarchDeltaT = 0.0;
+		if (VoxelBoundsHitT.x < VoxelBoundsHitT.y)
+		{
+			if (VoxelBoundsHitT.x < VoxelBoundsHitT.z)
 			{
-				Transmittance *= TraceOrthoVoxelGrid(Origin, Direction, TMin, FrustumTMin);
-				Transmittance *= TraceOrthoVoxelGrid(Origin, Direction, FrustumTMax, TMax);
+				VoxelRayMarchDeltaT = VoxelBoundsHitT.x - VoxelRayMarchT;
+				VoxelBoundsHitT.x += abs(VoxelRayDirectionInv.x);
 			}
 			else
 			{
-				Transmittance *= TraceOrthoVoxelGrid(Origin, Direction, TMin, TMax);
+				VoxelRayMarchDeltaT = VoxelBoundsHitT.z - VoxelRayMarchT;
+				VoxelBoundsHitT.z += abs(VoxelRayDirectionInv.z);
 			}
 		}
-		// NOTE: Ordering of ortho traversal is important for multi-bounce fallback
-		else if (OrthoGridUniformBuffer.bUseOrthoGrid)
+		else
 		{
-			Transmittance = TraceOrthoVoxelGrid(Origin, Direction, TMin, TMax);
+			if (VoxelBoundsHitT.y < VoxelBoundsHitT.z)
+			{
+				VoxelRayMarchDeltaT = VoxelBoundsHitT.y - VoxelRayMarchT;
+				VoxelBoundsHitT.y += abs(VoxelRayDirectionInv.y);
+			}
+			else
+			{
+				VoxelRayMarchDeltaT = VoxelBoundsHitT.z - VoxelRayMarchT;
+				VoxelBoundsHitT.z += abs(VoxelRayDirectionInv.z);
+			}
 		}
-		else if (FrustumGridUniformBuffer.bUseFrustumGrid)
+
+		// Clip voxel delta-t by overall voxel ray-length.
+		if (VoxelRayMarchT + VoxelRayMarchDeltaT > VoxelRayTMax)
 		{
-			Transmittance = TraceFrustumVoxelGrid(Origin, Direction, TMin, TMax);
+			VoxelRayMarchDeltaT = VoxelRayTMax - VoxelRayMarchT;
 		}
+
+		// Sample at the midpoint of the voxel ray
+		float3 VoxelEndPos = VoxelRayMarchPos + VoxelRayDirection * VoxelRayMarchDeltaT;
+		float3 VoxelSamplePos = (VoxelRayMarchPos + VoxelEndPos) * 0.5;
+		uint VoxelSampleLinearPos = GetLinearIndex(VoxelSamplePos, OrthoGridUniformBuffer.TopLevelGridResolution);
+		FMajorantData MajorantData = GetMajorantData(OrthoGridUniformBuffer.MajorantGridBuffer[VoxelSampleLinearPos]);
+		MergeMajorantData(MajorantData, OverlappingMajorant);
+
+		float3 TrackingOrigin = WorldRayBegin + Direction * VoxelRayMarchT * VoxelToWorldScale;
+		float WorldRayMarchDeltaT = VoxelRayMarchDeltaT * VoxelToWorldScale;
+
+		FVolumeTrackingResult Sample;
+		if (bTransmittanceOnly)
+		{
+			Sample = HeterogeneousVolumesRatioTracking(Result.Throughput, TrackingOrigin, Direction, WorldRayMarchDeltaT, MajorantData, RandSequence);
+		}
+		else
+		{
+			Sample = HeterogeneousVolumesDeltaTracking(Result.Throughput, TrackingOrigin, Direction, WorldRayMarchDeltaT, MajorantData, RandSequence);
+		}
+
+		if (Sample.Distance < 0)
+		{
+			return Sample;
+		}
+
+		Result = ConcatenateVolumeTrackingResult(Result, Sample);
+
+		VoxelRayMarchT += VoxelRayMarchDeltaT;
+		VoxelRayMarchPos = VoxelEndPos;
+	}
+
+	return Result;
+}
+
+float3 HeterogeneousVolumesGetTransmittance(float3 PathThroughput, float3 Origin, float3 Direction, float TMin, float TMax, float PathRoughness, inout RandomSequence RandSequence)
+{
+	float3 Transmittance = 1.0;
+	if (EnableHeterogeneousVolumes)
+	{
+		bool bTransmittanceOnly = true;
+		FMajorantData EmptyOverlappingMajorant = CreateMajorantData();
+		Transmittance = HeterogeneousVolumesMajorantBasedTracking(PathThroughput, Origin, Direction, TMin, TMax, bTransmittanceOnly, EmptyOverlappingMajorant, RandSequence).Throughput;
 	}
 
 	return Transmittance;
 }
+
+FVolumeTrackingResult HeterogeneousVolumesSampleDistance(float3 PathThroughput, float3 Origin, float3 Direction, float TMin, float TMax, FMajorantData OverlappingMajorant, inout RandomSequence RandSequence)
+{
+	bool bTransmittanceOnly = false;
+	return HeterogeneousVolumesMajorantBasedTracking(PathThroughput, Origin, Direction, TMin, TMax, bTransmittanceOnly, OverlappingMajorant, RandSequence);
+}
diff --git a/Engine/Shaders/Private/PathTracing/Volume/PathTracingVolume.ush b/Engine/Shaders/Private/PathTracing/Volume/PathTracingVolume.ush
--- a/Engine/Shaders/Private/PathTracing/Volume/PathTracingVolume.ush
+++ b/Engine/Shaders/Private/PathTracing/Volume/PathTracingVolume.ush
@@ -51,7 +51,7 @@ FVolumeDensityBounds VolumeGetDensityBounds(float3 Origin, float3 Direction, FVo
 		{
 			case VOLUMEID_ATMOSPHERE:            MergeVolumeDensityBounds(Result, AtmosphereGetDensityBounds(Origin, Direction, Interval.VolumeTMin, Interval.VolumeTMax)); break;
 			case VOLUMEID_FOG:                   MergeVolumeDensityBounds(Result, FogGetDensityBounds(Origin, Direction, Interval.VolumeTMin, Interval.VolumeTMax)); break;
-			case VOLUMEID_HETEROGENEOUS_VOLUMES: MergeVolumeDensityBounds(Result, HeterogeneousVolumesGetDensityBounds(Origin, Direction, Interval.VolumeTMin, Interval.VolumeTMax)); break;
+			//case VOLUMEID_HETEROGENEOUS_VOLUMES: MergeVolumeDensityBounds(Result, HeterogeneousVolumesGetDensityBounds(Origin, Direction, Interval.VolumeTMin, Interval.VolumeTMax)); break;
 		}
 	}
 	return Result;
@@ -84,7 +84,7 @@ float3 VolumeGetTransmittanceAnalytical(float3 StartThroughput, float3 Origin, f
 		{
 			case VOLUMEID_ATMOSPHERE:            Throughput *= AtmosphereGetTransmittance(Origin, Direction, Interval.VolumeTMin, Interval.VolumeTMax); break;
 			case VOLUMEID_FOG:                   Throughput *= FogGetTransmittance(Origin, Direction, Interval.VolumeTMin, Interval.VolumeTMax); break;
-			case VOLUMEID_HETEROGENEOUS_VOLUMES: Throughput *= HeterogeneousVolumesGetTransmittance(Origin, Direction, Interval.VolumeTMin, Interval.VolumeTMax, PathRoughness); break;
+			case VOLUMEID_HETEROGENEOUS_VOLUMES: Throughput *= HeterogeneousVolumesGetTransmittance(StartThroughput, Origin, Direction, Interval.VolumeTMin, Interval.VolumeTMax, PathRoughness, RandSequence); break;
 		}
 	}
 	return Throughput;
@@ -102,7 +102,7 @@ float3 VolumeGetTransmittanceAnalytical(float3 StartThroughput, float3 Origin, f
 		{
 			case VOLUMEID_ATMOSPHERE:            Throughput *= AtmosphereGetTransmittance(Origin, Direction, TMin, TMax); break;
 			case VOLUMEID_FOG:                   Throughput *= FogGetTransmittance(Origin, Direction, TMin, TMax); break;
-			case VOLUMEID_HETEROGENEOUS_VOLUMES: Throughput *= HeterogeneousVolumesGetTransmittance(Origin, Direction, TMin, TMax, PathRoughness); break;
+			case VOLUMEID_HETEROGENEOUS_VOLUMES: Throughput *= HeterogeneousVolumesGetTransmittance(StartThroughput, Origin, Direction, TMin, TMax, PathRoughness, RandSequence); break;
 		}
 	}
 	return Throughput;
@@ -171,7 +171,8 @@ float3 VolumeGetTransmittanceNullTracking(float3 StartThroughput, float3 Origin,
 			}
 			case VOLUMEID_HETEROGENEOUS_VOLUMES:
 			{
-				PATH_TRACER_REFERENCE_TRANSMITTANCE_LOOP(HeterogeneousVolumesGetDensityBounds, HeterogeneousVolumesGetDensity);
+				float PathRoughness = 0;
+				Throughput *= HeterogeneousVolumesGetTransmittance(StartThroughput, Origin, Direction, TMin, TMax, PathRoughness, RandSequence);
 				break;
 			}
 		}
@@ -179,6 +180,7 @@ float3 VolumeGetTransmittanceNullTracking(float3 StartThroughput, float3 Origin,
 	return Throughput;
 
 }
+
 float3 VolumeGetTransmittanceNullTracking(float3 StartThroughput, float3 Origin, float3 Direction, FVolumeIntersectionList IntersectionList, int Bounce, inout RandomSequence RandSequence)
 {
 	float3 Throughput = StartThroughput;
@@ -200,7 +202,8 @@ float3 VolumeGetTransmittanceNullTracking(float3 StartThroughput, float3 Origin,
 			}
 			case VOLUMEID_HETEROGENEOUS_VOLUMES:
 			{
-				PATH_TRACER_REFERENCE_TRANSMITTANCE_LOOP(HeterogeneousVolumesGetDensityBounds, HeterogeneousVolumesGetDensity);
+				float PathRoughness = 0;
+				Throughput *= HeterogeneousVolumesGetTransmittance(StartThroughput, Origin, Direction, TMin, TMax, PathRoughness, RandSequence);
 				break;
 			}
 		}
@@ -229,3 +232,37 @@ float3 VolumeGetTransmittance(float3 StartThroughput, float3 Origin, float3 Dire
 
 	return VolumeGetTransmittanceNullTracking(StartThroughput, Origin, Direction, IntersectionList, Bounce, RandSequence);
 }
+
+FVolumeTrackingResult VolumeSampleDistance(float3 PathThroughput, float3 Origin, float3 Direction, FVolumeIntersectionInterval Interval, bool bIsCameraRay, inout RandomSequence RandSequence)
+{
+	FVolumeDensityBounds VolumeDensityBounds = VolumeGetDensityBounds(Origin, Direction, Interval);
+	float3 SigmaBar = VolumeDensityBounds.SigmaMax;
+	if (!UseAnalyticTransmittance)
+	{
+		// if we are not using analytical transmittance, a tight majorant could prevent us from "seeing" hits that match the density exactly
+		// leading to heavy noise on bright objects embedded in the volume
+		// however if we can track transmittance analytically, this workaround is not needed as we will re-compute a precise answer after ray marching
+		SigmaBar *= bIsCameraRay ? 2 : 1;
+	}
+
+	FVolumeTrackingResult TrackingResult;
+	if (EnableHeterogeneousVolumes)
+	{
+		FMajorantData OverlappingMajorant = CreateMajorantData(VolumeDensityBounds.SigmaMin, VolumeDensityBounds.SigmaMax);
+		TrackingResult = HeterogeneousVolumesSampleDistance(PathThroughput, Origin, Direction, Interval.VolumeTMin, Interval.VolumeTMax, OverlappingMajorant, RandSequence);
+	}
+	else
+	{
+		float RandValue = RandomSequence_GenerateSample1D(RandSequence);
+		TrackingResult.Distance = SampleSpectralTransmittance(RandValue, SigmaBar, PathThroughput);
+		TrackingResult.Throughput = EvaluateSpectralTransmittanceHit(TrackingResult.Distance, SigmaBar, PathThroughput).xyz;
+		float3 SpectralPdf = PathThroughput * rcp(PathThroughput.x + PathThroughput.y + PathThroughput.z);
+		float3 DistancePdf = TrackingResult.Throughput;
+		if (TrackingResult.Distance < Interval.VolumeTMax)
+		{
+			DistancePdf *= SigmaBar;
+		}
+		TrackingResult.Pdf = dot(DistancePdf, SpectralPdf);
+	}
+	return TrackingResult;
+}
