Removed 0.5 / 255.0 bias when packing DiffuseIndirectSampleOcclusion into the gbuffer as the hardware format conversion already rounds to nearest.
The bias made it so the conversion to UINT8 was hardware dependent as both rounding up and down would be within the D3D format conversion error tolerance of 0.6ULP.
This fixes a performance regression where enabling Nanite compute materials would cause DiffuseIndirectSampleOcclusion to be rounded up to 1, causing Lumen to do more expensive sampling.

diff --git a/Engine/Shaders/Private/BasePassPixelShader.usf b/Engine/Shaders/Private/BasePassPixelShader.usf
--- a/Engine/Shaders/Private/BasePassPixelShader.usf
+++ b/Engine/Shaders/Private/BasePassPixelShader.usf
@@ -1915,7 +1915,7 @@ void FPixelShaderInOut_MainPS(
 		{
 			// change this so that we can pack everything into the gbuffer, but leave this for now
 			#if GBUFFER_HAS_DIFFUSE_SAMPLE_OCCLUSION
-				GBuffer.GenericAO = float(GBuffer.DiffuseIndirectSampleOcclusion) * rcp(255) + (0.5 / 255.0);
+				GBuffer.GenericAO = float(GBuffer.DiffuseIndirectSampleOcclusion) * (1.0f / 255.0f);
 			#elif ALLOW_STATIC_LIGHTING
 				// No space for AO. Multiply IndirectIrradiance by AO instead of storing.
 				GBuffer.GenericAO = EncodeIndirectIrradiance(GBuffer.IndirectIrradiance * GBuffer.GBufferAO) + QuantizationBias * (1.0 / 255.0);	// Stationary sky light path
diff --git a/Engine/Shaders/Private/DeferredShadingCommon.ush b/Engine/Shaders/Private/DeferredShadingCommon.ush
--- a/Engine/Shaders/Private/DeferredShadingCommon.ush
+++ b/Engine/Shaders/Private/DeferredShadingCommon.ush
@@ -869,7 +869,7 @@ void EncodeGBuffer(
 		OutGBufferC.rgb = EncodeBaseColor( GBuffer.BaseColor );
 
 #if GBUFFER_HAS_DIFFUSE_SAMPLE_OCCLUSION
-		OutGBufferC.a = float(GBuffer.DiffuseIndirectSampleOcclusion) * rcp(255) + (0.5 / 255.0);
+		OutGBufferC.a = float(GBuffer.DiffuseIndirectSampleOcclusion) * (1.0f / 255.0f);
 #elif ALLOW_STATIC_LIGHTING
 		// No space for AO. Multiply IndirectIrradiance by AO instead of storing.
 		OutGBufferC.a = EncodeIndirectIrradiance(GBuffer.IndirectIrradiance * GBuffer.GBufferAO) + QuantizationBias * (1.0 / 255.0);
diff --git a/Engine/Shaders/Private/GBufferHelpers.ush b/Engine/Shaders/Private/GBufferHelpers.ush
--- a/Engine/Shaders/Private/GBufferHelpers.ush
+++ b/Engine/Shaders/Private/GBufferHelpers.ush
@@ -332,7 +332,7 @@ void GBufferPreEncode(inout FGBufferData GBuffer, bool bChecker, float Geometric
 
 // should move this
 #if GBUFFER_HAS_DIFFUSE_SAMPLE_OCCLUSION
-		GBuffer.GenericAO = float(GBuffer.DiffuseIndirectSampleOcclusion) * rcp(255) + (0.5 / 255.0);
+		GBuffer.GenericAO = float(GBuffer.DiffuseIndirectSampleOcclusion) * (1.0f / 255.0f);
 #elif ALLOW_STATIC_LIGHTING
 		// No space for AO. Multiply IndirectIrradiance by AO instead of storing.
 		GBuffer.GenericAO = EncodeIndirectIrradiance(GBuffer.IndirectIrradiance * GBuffer.GBufferAO) + QuantizationBias * (1.0 / 255.0);
