Manually round off alpha channel in PostProcessUpscale.usf to mitigate floating point rounding issues

diff --git a/Engine/Shaders/Private/PostProcessUpscale.usf b/Engine/Shaders/Private/PostProcessUpscale.usf
--- a/Engine/Shaders/Private/PostProcessUpscale.usf
+++ b/Engine/Shaders/Private/PostProcessUpscale.usf
@@ -290,7 +290,17 @@ void MainPS(noperspective float4 UVAndScreenPos : TEXCOORD0, float4 SvPosition :
 	}
 #endif
 
-#if !POST_PROCESS_ALPHA
+#if POST_PROCESS_ALPHA
+	FLATTEN
+	if (OutColor.a > 0.999)
+	{
+		OutColor.a = 1;
+	}
+	else if (OutColor.a < 0.001)
+	{
+		OutColor.a = 0.0;
+	}
+#else
 	OutColor.a = 0; // Skip all computations related to alpha
 #endif
 }
\ No newline at end of file
