Fix shader compile CIS issues

diff --git a/Engine/Shaders/Private/SplineMeshCommon.ush b/Engine/Shaders/Private/SplineMeshCommon.ush
--- a/Engine/Shaders/Private/SplineMeshCommon.ush
+++ b/Engine/Shaders/Private/SplineMeshCommon.ush
@@ -159,11 +159,11 @@ FSplineMeshDeformBoundsContext SplineMeshInitializeDeformBoundsContext(
 	FSplineMeshShaderParams Params,
 	float3 MeshBoundsCenter,
 	float3 MeshBoundsExtent,
-	FShaderPrintContext ShaderPrint = (FShaderPrintContext)0,
-	float4x4 LocalToTWS = (float4x4)0
+	FShaderPrintContext ShaderPrint,
+	float4x4 LocalToTWS
 )
 {
-	FSplineMeshDeformBoundsContext Result = (FSplineMeshDeformBoundsContext)0;
+	FSplineMeshDeformBoundsContext Result;
 	Result.Params				= Params;
 	Result.MeshMinBounds		= MeshBoundsCenter - MeshBoundsExtent;
 	Result.MeshMaxBounds		= MeshBoundsCenter + MeshBoundsExtent;
@@ -177,6 +177,27 @@ FSplineMeshDeformBoundsContext SplineMeshInitializeDeformBoundsContext(
 	return Result;
 }
 
+FSplineMeshDeformBoundsContext SplineMeshInitializeDeformBoundsContext(
+	FSplineMeshShaderParams Params,
+	float3 MeshBoundsCenter,
+	float3 MeshBoundsExtent
+)
+{
+	FSplineMeshDeformBoundsContext Result;
+	Result.Params				= Params;
+	Result.MeshMinBounds		= MeshBoundsCenter - MeshBoundsExtent;
+	Result.MeshMaxBounds		= MeshBoundsCenter + MeshBoundsExtent;
+	Result.MeshBoundsExtent		= MeshBoundsExtent;
+	Result.DeformedMinBounds	= (float3)POSITIVE_INFINITY;
+	Result.DeformedMaxBounds	= (float3)NEGATIVE_INFINITY;
+	Result.CurSplineLength		= -1.0f; // Negative to mark as not started
+
+	Result.ShaderPrint.bIsActive	= false;
+	Result.LocalToTWS				= (float4x4)0;
+
+	return Result;
+}
+
 /** Solves post-deformed bounds of a slice of a spline mesh given the mesh-local bounds (iterative step) */
 void SplineMeshDeformLocalBoundsStep(inout FSplineMeshDeformBoundsContext Context, float SplineDist)
 {
