Updated distance field related compute shaders to use wrapped group count in order to support more than 4M instances.

diff --git a/Engine/Shaders/Private/DistanceField/GlobalDistanceField.usf b/Engine/Shaders/Private/DistanceField/GlobalDistanceField.usf
--- a/Engine/Shaders/Private/DistanceField/GlobalDistanceField.usf
+++ b/Engine/Shaders/Private/DistanceField/GlobalDistanceField.usf
@@ -3,6 +3,7 @@
 #define USE_DISTANCE_FIELD_SAMPLER 1
 
 #include "../Common.ush"
+#include "../ComputeShaderUtils.ush"
 #include "../DeferredShadingCommon.ush"
 #include "../DistanceFieldLightingShared.ush"
 #include "../DistanceFieldAOShared.ush"
@@ -54,13 +55,12 @@ float3 RelativePreViewTranslation;
 
 [numthreads(CULLOBJECTS_THREADGROUP_SIZE, 1, 1)]
 void CullObjectsToClipmapCS(
-	uint3 GroupId : SV_GroupID,
-	uint3 DispatchThreadId : SV_DispatchThreadID,
-	uint3 GroupThreadId : SV_GroupThreadID)
+	uint GroupIndex : SV_GroupIndex,
+	uint3 GroupId : SV_GroupID)
 {
 	const FLWCVector3 PreViewTranslation = MakeLWCVector3(-ViewTilePosition, RelativePreViewTranslation);
 
-	const uint ObjectIndex = DispatchThreadId.x;
+	const uint ObjectIndex = GetUnWrappedDispatchThreadId(GroupId, GroupIndex, CULLOBJECTS_THREADGROUP_SIZE);
 	if (ObjectIndex < NumSceneObjects)
 	{
 		const FDFObjectBounds DFObjectBounds = LoadDFObjectBounds(ObjectIndex);
diff --git a/Engine/Shaders/Private/DistanceFieldObjectCulling.usf b/Engine/Shaders/Private/DistanceFieldObjectCulling.usf
--- a/Engine/Shaders/Private/DistanceFieldObjectCulling.usf
+++ b/Engine/Shaders/Private/DistanceFieldObjectCulling.usf
@@ -5,6 +5,7 @@
 =============================================================================*/
 
 #include "Common.ush"
+#include "ComputeShaderUtils.ush"
 #include "DeferredShadingCommon.ush"
 #include "DistanceFieldLightingShared.ush"
 #include "DistanceFieldAOShared.ush"
@@ -19,22 +20,22 @@ groupshared uint GroupObjectIndices[UPDATEOBJECTS_THREADGROUP_SIZE];
 
 [numthreads(UPDATEOBJECTS_THREADGROUP_SIZE, 1, 1)]
 void CullObjectsForViewCS(
-	uint3 GroupId : SV_GroupID,
-	uint3 DispatchThreadId : SV_DispatchThreadID,
-    uint3 GroupThreadId : SV_GroupThreadID) 
+	uint GroupIndex : SV_GroupIndex,
+	uint3 GroupId : SV_GroupID)
 {
-	uint ObjectIndex = DispatchThreadId.x;
+	const uint ThreadIndex = GetUnWrappedDispatchThreadId(GroupId, GroupIndex, UPDATEOBJECTS_THREADGROUP_SIZE);
+	const uint ObjectIndex = ThreadIndex;
 
 #define USE_FRUSTUM_CULLING 1
 #if USE_FRUSTUM_CULLING
-	if (DispatchThreadId.x == 0)
+	if (ThreadIndex == 0)
 	{
 		// RWObjectIndirectArguments is zeroed by a clear before this shader, only need to set things that are non-zero (and are not read by this shader as that would be a race condition)
 		// IndexCount, NumInstances, StartIndex, BaseVertexIndex, FirstInstance
 		RWObjectIndirectArguments[0] = ObjectBoundingGeometryIndexCount;
 	}
 
-	if (GroupThreadId.x == 0)
+	if (GroupIndex == 0)
 	{
 		NumGroupObjects = 0;
 	}
@@ -68,23 +69,23 @@ void CullObjectsForViewCS(
 
 	GroupMemoryBarrierWithGroupSync();
 
-	if (GroupThreadId.x == 0)
+	if (GroupIndex == 0)
 	{
 		InterlockedAdd(RWObjectIndirectArguments[1], NumGroupObjects, GroupBaseIndex);
 	}
 
 	GroupMemoryBarrierWithGroupSync();
 
-	if (GroupThreadId.x < NumGroupObjects)
+	if (GroupIndex < NumGroupObjects)
 	{
-		uint SourceIndex = GroupObjectIndices[GroupThreadId.x];
-		uint DestIndex = GroupBaseIndex + GroupThreadId.x;
+		uint SourceIndex = GroupObjectIndices[GroupIndex];
+		uint DestIndex = GroupBaseIndex + GroupIndex;
 		RWCulledObjectIndices[DestIndex] = SourceIndex;
 	}
 
 #else
 
-	if (DispatchThreadId.x == 0)
+	if (ThreadIndex == 0)
 	{
 		// IndexCount, NumInstances, StartIndex, BaseVertexIndex, FirstInstance
 		RWObjectIndirectArguments[0] = ObjectBoundingGeometryIndexCount;
@@ -384,12 +385,13 @@ StructuredBuffer<uint> NumCulledTilesArray;
 
 [numthreads(COMPUTE_START_OFFSET_GROUP_SIZE, 1, 1)]
 void ComputeCulledTilesStartOffsetCS(
-	uint3 GroupId : SV_GroupID,
-	uint3 DispatchThreadId : SV_DispatchThreadID,
-	uint3 GroupThreadId : SV_GroupThreadID)
+	uint GroupIndex : SV_GroupIndex,
+	uint3 GroupId : SV_GroupID)
 {
-	uint ObjectIndex = DispatchThreadId.x;
-	uint NumCulledObjects = GetCulledNumObjects();
+	const uint NumCulledObjects = GetCulledNumObjects();
+
+	const uint ThreadIndex = GetUnWrappedDispatchThreadId(GroupId, GroupIndex, COMPUTE_START_OFFSET_GROUP_SIZE);
+	const uint ObjectIndex = ThreadIndex;
 
 	if (ObjectIndex < NumCulledObjects)
 	{
@@ -409,7 +411,7 @@ void ComputeCulledTilesStartOffsetCS(
 		}
 	}
 
-	if (DispatchThreadId.x == 0)
+	if (ThreadIndex == 0)
 	{
 		RWObjectTilesIndirectArguments[1] = 1;
 		RWObjectTilesIndirectArguments[2] = 1;
diff --git a/Engine/Shaders/Private/DistanceFieldStreaming.usf b/Engine/Shaders/Private/DistanceFieldStreaming.usf
--- a/Engine/Shaders/Private/DistanceFieldStreaming.usf
+++ b/Engine/Shaders/Private/DistanceFieldStreaming.usf
@@ -5,6 +5,7 @@
 =============================================================================*/
 
 #include "Common.ush"
+#include "ComputeShaderUtils.ush"
 #include "DistanceFieldLightingShared.ush"
 
 RWTexture3D<UNORM float> RWDistanceFieldBrickAtlas;
@@ -111,16 +112,17 @@ uint DebugForceNumMips;
 
 [numthreads(THREADGROUP_SIZE, 1, 1)]
 void ComputeDistanceFieldAssetWantedMipsCS(
-	uint3 GroupId : SV_GroupID,
-	uint3 DispatchThreadId : SV_DispatchThreadID,
-	uint3 GroupThreadId : SV_GroupThreadID)
+	uint GroupIndex : SV_GroupIndex,
+	uint3 GroupId : SV_GroupID)
 {
-	if (DispatchThreadId.x == 0)
+	const uint Index = GetUnWrappedDispatchThreadId(GroupId, GroupIndex, THREADGROUP_SIZE);
+
+	if (Index == 0)
 	{
 		RWDistanceFieldAssetStreamingRequests[0] = 0;
 	}
 
-	uint ObjectIndex = DispatchThreadId.x;
+	const uint ObjectIndex = Index;
 
 	if (ObjectIndex < NumSceneObjects)
 	{
diff --git a/Engine/Source/Runtime/Renderer/Private/DistanceFieldObjectCulling.cpp b/Engine/Source/Runtime/Renderer/Private/DistanceFieldObjectCulling.cpp
--- a/Engine/Source/Runtime/Renderer/Private/DistanceFieldObjectCulling.cpp
+++ b/Engine/Source/Runtime/Renderer/Private/DistanceFieldObjectCulling.cpp
@@ -95,14 +95,14 @@ void CullObjectsToView(FRDGBuilder& GraphBuilder, FScene* Scene, const FViewInfo
 		PassParameters->AOMaxViewDistance = GetMaxAOViewDistance();
 
 		auto ComputeShader = View.ShaderMap->GetShader<FCullObjectsForViewCS>();
-		const int32 GroupSize = FMath::DivideAndRoundUp<uint32>(NumObjectsInBuffer, UpdateObjectsGroupSize);
+		const FIntVector GroupCount = FComputeShaderUtils::GetGroupCountWrapped(NumObjectsInBuffer, UpdateObjectsGroupSize);
 
 		FComputeShaderUtils::AddPass(
 			GraphBuilder,
 			RDG_EVENT_NAME("ObjectFrustumCulling"),
 			ComputeShader,
 			PassParameters,
-			FIntVector(GroupSize, 1, 1));
+			GroupCount);
 	}
 }
 
@@ -399,14 +399,14 @@ void BuildTileObjectLists(
 			PassParameters->SceneTextures = SceneTexturesUniformBuffer;
 
 			auto ComputeShader = View.ShaderMap->GetShader<FComputeCulledTilesStartOffsetCS>();
-			const int32 GroupSize = FMath::DivideAndRoundUp<uint32>(Scene->DistanceFieldSceneData.NumObjectsInBuffer, ComputeStartOffsetGroupSize);
+			const FIntVector GroupCount = FComputeShaderUtils::GetGroupCountWrapped(Scene->DistanceFieldSceneData.NumObjectsInBuffer, ComputeStartOffsetGroupSize);
 
 			FComputeShaderUtils::AddPass(
 				GraphBuilder,
 				RDG_EVENT_NAME("ComputeStartOffsets"),
 				ComputeShader,
 				PassParameters,
-				FIntVector(GroupSize, 1, 1));
+				GroupCount);
 		}
 
 		// Start at 0 tiles per object
diff --git a/Engine/Source/Runtime/Renderer/Private/DistanceFieldStreaming.cpp b/Engine/Source/Runtime/Renderer/Private/DistanceFieldStreaming.cpp
--- a/Engine/Source/Runtime/Renderer/Private/DistanceFieldStreaming.cpp
+++ b/Engine/Source/Runtime/Renderer/Private/DistanceFieldStreaming.cpp
@@ -1152,7 +1152,7 @@ void FDistanceFieldSceneData::GenerateStreamingRequests(
 				RDG_EVENT_NAME("ComputeWantedMips"),
 				ComputeShader,
 				PassParameters,
-				FComputeShaderUtils::GetGroupCount(NumObjectsInBuffer, FComputeDistanceFieldAssetWantedMipsCS::GetGroupSize()));
+				FComputeShaderUtils::GetGroupCountWrapped(NumObjectsInBuffer, FComputeDistanceFieldAssetWantedMipsCS::GetGroupSize()));
 		}
 
 		{
diff --git a/Engine/Source/Runtime/Renderer/Private/GlobalDistanceField.cpp b/Engine/Source/Runtime/Renderer/Private/GlobalDistanceField.cpp
--- a/Engine/Source/Runtime/Renderer/Private/GlobalDistanceField.cpp
+++ b/Engine/Source/Runtime/Renderer/Private/GlobalDistanceField.cpp
@@ -2371,14 +2371,14 @@ void UpdateGlobalDistanceFieldVolume(
 								PermutationVector.Set<FCullObjectsToClipmapCS::FReadbackHasPendingStreaming>(bRecacheClipmapsWithPendingStreaming);
 								auto ComputeShader = View.ShaderMap->GetShader<FCullObjectsToClipmapCS>(PermutationVector);
 
-								const FIntVector GroupSize = FComputeShaderUtils::GetGroupCount(DistanceFieldSceneData.NumObjectsInBuffer, FCullObjectsToClipmapCS::GetGroupSize());
+								const FIntVector GroupCount = FComputeShaderUtils::GetGroupCountWrapped(DistanceFieldSceneData.NumObjectsInBuffer, FCullObjectsToClipmapCS::GetGroupSize());
 
 								FComputeShaderUtils::AddPass(
 									GraphBuilder,
 									RDG_EVENT_NAME("CullToClipmap"),
 									ComputeShader,
 									PassParameters,
-									GroupSize);
+									GroupCount);
 							}
 
 							if (bRecacheClipmapsWithPendingStreaming)
