Path tracing: adding volume emission

diff --git a/Engine/Shaders/Private/PathTracing/PathTracingCore.ush b/Engine/Shaders/Private/PathTracing/PathTracingCore.ush
--- a/Engine/Shaders/Private/PathTracing/PathTracingCore.ush
+++ b/Engine/Shaders/Private/PathTracing/PathTracingCore.ush
@@ -589,6 +589,7 @@ FPathTracingPayload TraceTransparentRay(inout FPathState PathState, int Bounce,
 							}
 						}
 
+						PathState.Radiance += Result.Emission * PathThroughput;
 						// keep tracing through the volume
 						PathThroughput *= SigmaN;
 					}
diff --git a/Engine/Shaders/Private/PathTracing/Volume/PathTracingHeterogeneousVolumes.ush b/Engine/Shaders/Private/PathTracing/Volume/PathTracingHeterogeneousVolumes.ush
--- a/Engine/Shaders/Private/PathTracing/Volume/PathTracingHeterogeneousVolumes.ush
+++ b/Engine/Shaders/Private/PathTracing/Volume/PathTracingHeterogeneousVolumes.ush
@@ -261,6 +261,7 @@ FVolumeShadedResult GetOrthoVoxelGridDensity(float3 TranslatedWorldPos)
 {
 	float3 SigmaT = 0.0;
 	float3 Albedo = 0.0;
+	float3 Emission = 0.0;
 
 	float3 TranslatedWorldBoundsMin = GetTranslatedWorldPos(OrthoGridUniformBuffer.TopLevelGridWorldBoundsMin);
 	float3 TranslatedWorldBoundsMax = GetTranslatedWorldPos(OrthoGridUniformBuffer.TopLevelGridWorldBoundsMax);
@@ -294,6 +295,7 @@ FVolumeShadedResult GetOrthoVoxelGridDensity(float3 TranslatedWorldPos)
 					uint LinearBottomLevelVoxelPos = BottomLevelIndex + MortonEncode3(BottomLevelVoxelPosAsInt);
 					SigmaT = GetExtinction(OrthoGridUniformBuffer.ExtinctionGridBuffer[LinearBottomLevelVoxelPos]);
 					Albedo = GetAlbedo(OrthoGridUniformBuffer.ScatteringGridBuffer[LinearBottomLevelVoxelPos]);
+					Emission = GetEmission(OrthoGridUniformBuffer.EmissionGridBuffer[LinearBottomLevelVoxelPos]);
 				}
 			}
 		}
@@ -312,6 +314,7 @@ FVolumeShadedResult GetOrthoVoxelGridDensity(float3 TranslatedWorldPos)
 
 				SigmaT = GetExtinction(OrthoGridUniformBuffer.ExtinctionGridBuffer[LinearBottomLevelVoxelPos]);
 				Albedo = GetAlbedo(OrthoGridUniformBuffer.ScatteringGridBuffer[LinearBottomLevelVoxelPos]);
+				Emission = GetEmission(OrthoGridUniformBuffer.EmissionGridBuffer[LinearBottomLevelVoxelPos]);
 			}
 		}
 	}
@@ -320,7 +323,7 @@ FVolumeShadedResult GetOrthoVoxelGridDensity(float3 TranslatedWorldPos)
 	Result.SigmaT = SigmaT;
 	Result.SigmaSHG = Albedo;
 	Result.PhaseG = GetPhaseG();
-	// TODO: Emission?
+	Result.Emission = Emission;
 
 	return Result;
 }
@@ -329,6 +332,7 @@ FVolumeShadedResult GetFrustumVoxelGridDensity(float3 TranslatedWorldPos, inout
 {
 	float3 SigmaT = 0.0;
 	float3 Albedo = 0.0;
+	float3 Emission = 0.0;
 
 	// Convert TranslatedWorldPos to voxel space
 	float3 WorldPos = LWCHackToFloat(LWCSubtract(TranslatedWorldPos, PrimaryView.PreViewTranslation));
@@ -353,6 +357,7 @@ FVolumeShadedResult GetFrustumVoxelGridDensity(float3 TranslatedWorldPos, inout
 			uint LinearBottomLevelVoxelPos = BottomLevelIndex + MortonEncode3(uint3(BottomLevelVoxelPos));
 			SigmaT = GetExtinction(FrustumGridUniformBuffer.ExtinctionFroxelGridBuffer[LinearBottomLevelVoxelPos]);
 			Albedo = GetAlbedo(FrustumGridUniformBuffer.ScatteringFroxelGridBuffer[LinearBottomLevelVoxelPos]);
+			Emission = GetEmission(FrustumGridUniformBuffer.EmissionFroxelGridBuffer[LinearBottomLevelVoxelPos]);
 		}
 	}
 
@@ -361,7 +366,7 @@ FVolumeShadedResult GetFrustumVoxelGridDensity(float3 TranslatedWorldPos, inout
 	Result.SigmaT = SigmaT;
 	Result.SigmaSHG = Albedo;
 	Result.PhaseG = GetPhaseG();
-	// TODO: Emission?
+	Result.Emission = Emission;
 	return Result;
 }
 
diff --git a/Engine/Shaders/Private/PathTracing/Volume/PathTracingVolumeCommon.ush b/Engine/Shaders/Private/PathTracing/Volume/PathTracingVolumeCommon.ush
--- a/Engine/Shaders/Private/PathTracing/Volume/PathTracingVolumeCommon.ush
+++ b/Engine/Shaders/Private/PathTracing/Volume/PathTracingVolumeCommon.ush
@@ -190,10 +190,9 @@ struct FVolumeShadedResult
 	float3 SigmaT;         // extinction
 	float3 SigmaSRayleigh; // Rayleight scattering coefficient
 	float3 SigmaSHG;       // Henyey-Greenstein scattering coefficient
+	float3 Emission;
 	float PhaseG;          // 'g' term for HG lobe
-
 	// TODO: second HG lobe?
-	// TODO: emission?
 };
 
 FPathTracingPayload CreateMediumHitPayload(float HitT, float3 TranslatedWorldPos, FVolumeShadedResult ShadedResult)
@@ -218,6 +217,7 @@ void MergeVolumeShadedResult(inout FVolumeShadedResult Merged, FVolumeShadedResu
 	Merged.SigmaT += Other.SigmaT;
 	Merged.SigmaSRayleigh += Other.SigmaSRayleigh;
 	Merged.SigmaSHG += Other.SigmaSHG;
+	Merged.Emission += Other.Emission;
 	// blend phase function anisotropy based on amount of each medium
 	float Qa = Other.SigmaSHG.x + Other.SigmaSHG.y + Other.SigmaSHG.z;
 	float Qs = Merged.SigmaSHG.x + Merged.SigmaSHG.y + Merged.SigmaSHG.z;
