LocalHeightFogVolume - fix exploding values due to exp.

diff --git a/Engine/Shaders/Private/LocalHeightFog.usf b/Engine/Shaders/Private/LocalHeightFog.usf
--- a/Engine/Shaders/Private/LocalHeightFog.usf
+++ b/Engine/Shaders/Private/LocalHeightFog.usf
@@ -71,12 +71,29 @@ FFogData EvaluateHeightFogIntegral(FLocalHeightFogGPUInstanceData FogInstance,
 	float StartHeight = RayStart.z - FogInstance.HeightOffset; // accoutn for the fog offset.
 
 	float OpticalDepth = 0.0f;
-#if 1
+#if 0
 	// Original integral, less artefact
 	OpticalDepth = (FogInstance.Density / FogInstance.HeightFalloff) * exp(-StartHeight * FogInstance.HeightFalloff) * (1.0 - exp(-RayDir.z * RayLength * FogInstance.HeightFalloff)) / RayDir.z;
-#elif 0
+#elif 1
 	// Reworked to avoid large value as input to exp
-	OpticalDepth = (FogInstance.Density / FogInstance.HeightFalloff) * (exp(-StartHeight * FogInstance.HeightFalloff) - exp(-StartHeight * FogInstance.HeightFalloff - RayDir.z * RayLength * FogInstance.HeightFalloff)) / RayDir.z;
+	//OpticalDepth = (FogInstance.Density / FogInstance.HeightFalloff) * (exp(-StartHeight * FogInstance.HeightFalloff) - exp(-StartHeight * FogInstance.HeightFalloff - RayDir.z * RayLength * FogInstance.HeightFalloff)) / RayDir.z;
+
+	// But we use an updated formulation to avoid large input to exp leading to inf and nan.
+	float Factor0 = StartHeight * FogInstance.HeightFalloff;
+	float Factor1 = RayDir.z * RayLength * FogInstance.HeightFalloff;
+//	if (Factor0 < 0	// If we start below the height fog horizon, where it is already thick according to falloff...
+//		&& Factor1 < 0)	// ... and that the difference between entry and exit point is large...
+//	{
+//		// ... then exp is likely going to explode.
+//		OpticalDepth = FogInstance.Density * exp(-Factor0) * RayLength;	// Force to a large optical depth matching the extinction at ray originand according to ray segment length.
+//	}
+//	else
+//	{
+//		OpticalDepth = (FogInstance.Density / (FogInstance.HeightFalloff * RayDir.z)) * (exp(-Factor0) - exp(-(Factor0 + Factor1)));
+//	}
+	Factor0 = max(-80.0f, Factor0); // Simply clamp Factor0 to a valid range in order to not have exp explode.
+
+	OpticalDepth = (FogInstance.Density / (FogInstance.HeightFalloff * RayDir.z)) * (exp(-Factor0) - exp(-(Factor0 + Factor1)));
 #endif
 
 	// Transmittance is the 
