Substrate - fixed microoclusion/F0/F90 usage.

diff --git a/Engine/Shaders/Private/Strata/StrataEvaluation.ush b/Engine/Shaders/Private/Strata/StrataEvaluation.ush
--- a/Engine/Shaders/Private/Strata/StrataEvaluation.ush
+++ b/Engine/Shaders/Private/Strata/StrataEvaluation.ush
@@ -350,6 +350,11 @@ FStrataEvaluateResult StrataEvaluateBSDFCommon(FStrataBSDFContext BSDFContext, F
 			const bool bHasAnisotropy	= BSDF_GETHASANISOTROPY(BSDFContext.BSDF);
 			const bool bHaziness		= BSDF_GETHASHAZINESS(BSDFContext.BSDF);
 
+			// Specular occlusion is only used here once to affect the F90 source parameters. F0 will decrease naturally to 0.
+			const float SpecularMicroOcclusion = F0RGBToMicroOcclusion(F0);
+			F0  *= SpecularMicroOcclusion;
+			F90 *= SpecularMicroOcclusion;
+
 			if (Settings.bForceFullyRough)
 			{
 				// When rendering reflection captures, the BSDF roughness has already been forced to 1 using View.RoughnessOverrideParameter (see StrataSanitizeBSDF).
@@ -544,18 +549,14 @@ FStrataEvaluateResult StrataEvaluateBSDFCommon(FStrataBSDFContext BSDFContext, F
 					}
 				}
 
-					
-
-				const float SpecularMicroOcclusion = F0RGBToMicroOcclusion(F0);
-
 				if (bIsRectLight)
 				{
-					Sample.SpecularPathValue = (SpecularMicroOcclusion * MSScale) * RectLightSpec;
+					Sample.SpecularPathValue = MSScale * RectLightSpec;
 				}
 				else
 				{
 					const float3 FresnelTerm = Strata_F_GGX(F0, F90, BSDFContext.Context.VoH);
-					Sample.SpecularPathValue = ((SpecularMicroOcclusion * D * Vis) * MSScale) * FresnelTerm;
+					Sample.SpecularPathValue = D * Vis * MSScale * FresnelTerm;
 				}
 
 				Sample.SpecularPathProbability     = 1.0f;
@@ -569,6 +570,7 @@ FStrataEvaluateResult StrataEvaluateBSDFCommon(FStrataBSDFContext BSDFContext, F
 			BRANCH
 			if (bHaziness)
 			{
+				// F0 / F90 have already been affected by specular micro occlusion.
 				float3 HazeF0 = F0;
 				float3 HazeF90= F90;
 			
@@ -581,6 +583,7 @@ FStrataEvaluateResult StrataEvaluateBSDFCommon(FStrataBSDFContext BSDFContext, F
 				{
 					HazeF0 = 0.04f;
 					HazeF90 = 1.0f;
+					// In this case, no specular micro occlusion happens.
 				}
 
 				FBxDFEnergyTerms EnergyTerms = ComputeGGXSpecEnergyTerms(HazeSafeRoughness, BSDFContext.Context.NoV, HazeF0, HazeF90);
@@ -637,16 +640,14 @@ FStrataEvaluateResult StrataEvaluateBSDFCommon(FStrataBSDFContext BSDFContext, F
 					}
 				}
 
-				
-				const float HazeSpecularMicroOcclusion = F0RGBToMicroOcclusion(HazeF0);
 				const float3 HazeFresnelTerm = Strata_F_GGX(HazeF0, HazeF90, BSDFContext.Context.VoH);
 				if (bIsRectLight)
 				{					
-					Sample.SpecularHazePathValue = RectLightSpecHaze * (HazeMSScale * HazeSpecularMicroOcclusion);
+					Sample.SpecularHazePathValue = RectLightSpecHaze * HazeMSScale;
 				}
 				else
 				{
-					Sample.SpecularHazePathValue = ((HazeD * HazeVis * HazeSpecularMicroOcclusion) * HazeMSScale) * HazeFresnelTerm;
+					Sample.SpecularHazePathValue = HazeD * HazeVis * HazeMSScale * HazeFresnelTerm;
 				}
 
 				// Reminder: 
@@ -659,7 +660,7 @@ FStrataEvaluateResult StrataEvaluateBSDFCommon(FStrataBSDFContext BSDFContext, F
 					const float3 HazeClearCoatTransmittance = SimpleClearCoatTransmittance(ClearCoatContext.NoL, ClearCoatContext.NoV, StrataGetBSDFMetallic(BSDFContext.BSDF), StrataGetBSDFBaseColor(BSDFContext.BSDF));
 
 					const float TopLayerCoverage = HazeWeight;
-					const float TopLayerSpecularTransmittionApprox = saturate(1.0f - HazeSpecularMicroOcclusion * HazeFresnelTerm.x); // Simple and do not use the LUT for now. With bHazeAsSimpleClearCoat, Fresnel is achromatic
+					const float TopLayerSpecularTransmittionApprox = saturate(1.0f - HazeFresnelTerm.x); // Simple and do not use the LUT for now. With bHazeAsSimpleClearCoat, Fresnel is achromatic
 					const float3 TopLayerThrouput = lerp(1.0f, HazeClearCoatTransmittance * TopLayerSpecularTransmittionApprox, TopLayerCoverage);
 
 					const float TopLayerThrouputGrey = dot(TopLayerThrouput, (1.0 / 3.0).xxx);
@@ -978,6 +979,11 @@ FStrataEvaluateResult StrataEvaluateBSDFCommon(FStrataBSDFContext BSDFContext, F
 			float3 F90			= EYE_F90(BSDFContext.BSDF);
 			const float SafeRoughness	= MakeRoughnessSafe(EYE_ROUGHNESS(BSDFContext.BSDF));
 
+			// Specular occlusion is only used here once to affect the F90 source parameters. F0 will decrease naturally to 0.
+			const float SpecularMicroOcclusion = F0RGBToMicroOcclusion(F0);
+			F0  *= SpecularMicroOcclusion;
+			F90 *= SpecularMicroOcclusion;
+
 			if (Settings.bForceFullyRough)
 			{
 				// When rendering reflection captures, the BSDF roughness has already been forced to 1 using View.RoughnessOverrideParameter (see StrataSanitizeBSDF).
@@ -1073,20 +1079,19 @@ FStrataEvaluateResult StrataEvaluateBSDFCommon(FStrataBSDFContext BSDFContext, F
 				PDF = RayPDFToReflectionRayPDF(BSDFContext.Context.VoH, H_PDF);
 			}
 
-			const float SpecularMicroOcclusion = F0RGBToMicroOcclusion(F0);
-			const float3 FresnelTerm = SpecularMicroOcclusion * Strata_F_GGX(F0, F90, BSDFContext.Context.VoH);
+			const float3 FresnelTerm = Strata_F_GGX(F0, F90, BSDFContext.Context.VoH);
 
 			Sample.SpecularPathProbability	= 1.0f;
 			Sample.SpecularPDF				= PDF * Sample.SpecularPathProbability;
 
 			if (bIsRectLight)
 			{
-				Sample.SpecularPathValue       = RectLightSpec * (MSScale * SpecularMicroOcclusion);
+				Sample.SpecularPathValue       = RectLightSpec * MSScale;
 				Sample.IntegratedSpecularValue = (ShadowTerms.SurfaceShadow * /* AreaLightContext.NoL and falloff is part of the LTC integration already */ Sample.SpecularPathProbability) * Sample.SpecularPathValue;
 			}
 			else
 			{
-				Sample.SpecularPathValue        = ((D * Vis) * MSScale) * FresnelTerm;
+				Sample.SpecularPathValue        = D * Vis * MSScale * FresnelTerm;
 				Sample.IntegratedSpecularValue += (ShadowTerms.SurfaceShadow * AreaLightContext.NoL * AreaLightContext.Falloff) * AreaLightContext.AreaLight.FalloffColor * Sample.SpecularPathValue * Sample.SpecularPathProbability;
 			}
 
@@ -1416,6 +1421,11 @@ FStrataEnvLightResult StrataEvaluateForEnvLight(FStrataBSDFContext BSDFContext,
 			float3 F0 				= SLAB_F0(BSDFContext.BSDF);
 			float3 F90 				= SLAB_F90(BSDFContext.BSDF);
 
+			// Specular occlusion is only used here once to affect the F90 source parameters. F0 will decrease naturally to 0.
+			const float SpecularMicroOcclusion = F0RGBToMicroOcclusion(F0);
+			F0  *= SpecularMicroOcclusion;
+			F90 *= SpecularMicroOcclusion;
+
 			if (Settings.bForceFullyRough)
 			{
 				// When rendering reflection captures, the BSDF roughness has already been forced to 1 using View.RoughnessOverrideParameter (see StrataSanitizeBSDF).
@@ -1480,8 +1490,6 @@ FStrataEnvLightResult StrataEvaluateForEnvLight(FStrataBSDFContext BSDFContext,
 			BRANCH
 			if (bEnableSpecular)
 			{
-				const float SpecularMicroOcclusion = F0RGBToMicroOcclusion(F0);
-
 				#if STRATA_COMPLEXPATH
 				const bool bHasAnisotropy = BSDF_GETHASANISOTROPY(BSDFContext.BSDF);
 				if (bHasAnisotropy)
@@ -1498,7 +1506,7 @@ FStrataEnvLightResult StrataEvaluateForEnvLight(FStrataBSDFContext BSDFContext,
 				#endif
 
 				StrataEnvLightResult.SpecularDirection = StrataGetOffSpecularPeakReflectionDir(BSDFContext.N, BSDFContext.R, SafeRoughness);
-				StrataEnvLightResult.SpecularWeight = EvalEnvBRDF * SpecularMicroOcclusion;
+				StrataEnvLightResult.SpecularWeight = EvalEnvBRDF;
 				StrataEnvLightResult.SpecularSafeRoughness = SafeRoughness;
 				StrataEnvLightResult.SpecularColor = F0;
 
@@ -1565,7 +1573,6 @@ FStrataEnvLightResult StrataEvaluateForEnvLight(FStrataBSDFContext BSDFContext,
 
 						const float HazeF0 = 0.04f;
 						const float HazeF90= 1.0f;
-						const float HazeSpecularMicroOcclusion = F0RGBToMicroOcclusion(HazeF0);
 						FBxDFEnergyTermsRGB EnergyTerms = ComputeGGXSpecEnergyTermsRGB(HazeSafeRoughness, BSDFContext.Context.NoV, HazeF0, HazeF90);
 						const float3 EvalEnvBRDFHaze = EnergyTerms.E;
 						const float3 TopLayerSpecularTransmittionApprox = saturate(1.0f - EvalEnvBRDFHaze); // Simple and do not use extra LUT for now
@@ -1576,7 +1583,7 @@ FStrataEnvLightResult StrataEvaluateForEnvLight(FStrataBSDFContext BSDFContext,
 						StrataEnvLightResult.SpecularWeight *= BottomLayerFactors;
 						StrataEnvLightResult.DiffuseWeight  *= BottomLayerFactors;
 
-						StrataEnvLightResult.SpecularHazeWeight = (TopLayerCoverage * HazeSpecularMicroOcclusion) * EvalEnvBRDFHaze;
+						StrataEnvLightResult.SpecularHazeWeight = TopLayerCoverage * EvalEnvBRDFHaze;
 						StrataEnvLightResult.SpecularHazeSafeRoughness = HazeSafeRoughness;
 
 						StrataEnvLightResult.SSRSpecularWeight = lerp(StrataEnvLightResult.SpecularWeight, StrataEnvLightResult.SpecularHazeWeight, TopLayerCoverage);
@@ -1595,7 +1602,7 @@ FStrataEnvLightResult StrataEvaluateForEnvLight(FStrataBSDFContext BSDFContext,
 						// Compute the directional albedo for hazy specular
 						FBxDFEnergyTermsRGB EnergyTerms = ComputeGGXSpecEnergyTermsRGB(HazeSafeRoughness, BSDFContext.Context.NoV, F0, F90);
 						const float3 EvalEnvBRDFHaze = EnergyTerms.E;
-						StrataEnvLightResult.SpecularHazeWeight = (lerp(0.0f, HazeWeight, HazinessFadeIn) * SpecularMicroOcclusion) * EvalEnvBRDFHaze;
+						StrataEnvLightResult.SpecularHazeWeight = lerp(0.0f, HazeWeight, HazinessFadeIn) * EvalEnvBRDFHaze;
 						StrataEnvLightResult.SpecularHazeSafeRoughness = HazeSafeRoughness;
 					}
 				}
@@ -1686,6 +1693,11 @@ FStrataEnvLightResult StrataEvaluateForEnvLight(FStrataBSDFContext BSDFContext,
 			float3 F0 				= EYE_F0(BSDFContext.BSDF);
 			float3 F90 				= EYE_F90(BSDFContext.BSDF);
 
+			// Specular occlusion is only used here once to affect the F90 source parameters. F0 will decrease naturally to 0.
+			const float SpecularMicroOcclusion = F0RGBToMicroOcclusion(F0);
+			F0  *= SpecularMicroOcclusion;
+			F90 *= SpecularMicroOcclusion;
+
 			if (Settings.bForceFullyRough)
 			{
 				// When rendering reflection captures, the BSDF roughness has already been forced to 1 using View.RoughnessOverrideParameter (see StrataSanitizeBSDF).
@@ -1711,9 +1723,8 @@ FStrataEnvLightResult StrataEvaluateForEnvLight(FStrataBSDFContext BSDFContext,
 			BRANCH
 			if (bEnableSpecular)
 			{
-				const float SpecularMicroOcclusion = F0RGBToMicroOcclusion(F0);
 				StrataEnvLightResult.SpecularDirection = StrataGetOffSpecularPeakReflectionDir(BSDFContext.N, BSDFContext.R, SafeRoughness);
-				StrataEnvLightResult.SpecularWeight = EvalEnvBRDF * SpecularMicroOcclusion;
+				StrataEnvLightResult.SpecularWeight = EvalEnvBRDF;
 				StrataEnvLightResult.SpecularSafeRoughness = SafeRoughness;
 				StrataEnvLightResult.SpecularColor = F0;
 			}
