LocalHeightFogVolume
- removed softness control.
- Added fog mode for height or radial density.
- code cleaning

diff --git a/Engine/Shaders/Private/LocalHeightFog.usf b/Engine/Shaders/Private/LocalHeightFog.usf
--- a/Engine/Shaders/Private/LocalHeightFog.usf
+++ b/Engine/Shaders/Private/LocalHeightFog.usf
@@ -22,7 +22,7 @@ struct FLocalHeightFogGPUInstanceData
 	float3 Albedo;
 	float  PhaseG;
 	float3 Emissive;
-	float RadialAttenuationHermiteTangent;
+	uint   FogMode;
 };
 
 StructuredBuffer<FLocalHeightFogGPUInstanceData> LocalHeightFogInstances;
@@ -108,6 +108,53 @@ FFogData EvaluateHeightFogIntegral(FLocalHeightFogGPUInstanceData FogInstance,
 }
 
 
+FFogData EvaluateRadialFogIntegral(FLocalHeightFogGPUInstanceData FogInstance, float3 CameraPos, float3 RayStart, float3 RayDir, float RayLength, float RadialAttenuation)
+{
+	FFogData FogData = (FFogData)0;
+
+	float3 RayOriU = mul(float4(RayStart, 1.0f),                      FogInstance.InvTransform).xyz;
+	float3 RayIntU = mul(float4(RayStart + RayDir * RayLength, 1.0f), FogInstance.InvTransform).xyz;
+	float3 RayDirU = mul(float4(RayDir  , 0.0f),                      FogInstance.InvTransform).xyz;
+	RayDirU = normalize(RayDirU);
+
+	float RayLengthU = length(RayIntU - RayOriU);
+
+	float3 SphereCenter = 0.0f;
+	float3 VolumeCenterToRayO = (RayOriU - SphereCenter);
+
+	float b = dot(RayDirU,            VolumeCenterToRayO);
+	float c = dot(VolumeCenterToRayO, VolumeCenterToRayO) - 1.0f;
+	float h = b * b - c;
+
+	float OpticalDepth = 0.0f;
+	if (h >= 0.0)
+	{
+		h = sqrt(h);
+		float Length0 = -b - h;
+		float Length1 = -b + h;
+
+		Length0 = max(Length0, 0.0);
+		Length1 = max(Length1, 0.0);
+		Length1 = min(Length1, RayLengthU);
+
+		// Integral of a density that descreases accordingf to square distance from the center of the volume.
+		float Integral0 = -(c * Length0 + b * Length0 * Length0 + Length0 * Length0 * Length0 / 3.0f);
+		float Integral1 = -(c * Length1 + b * Length1 * Length1 + Length1 * Length1 * Length1 / 3.0f);
+		OpticalDepth = max(0.0, FogInstance.Density * (Integral1 - Integral0) * (3.0f / 4.0f));
+	}
+
+	// Transmittance is the 
+	float Transmittance = exp(-OpticalDepth);
+	// Assuming extinction is grey scale, we can compute a single coverage value from transmittance
+	FogData.Coverage = (1.0 - Transmittance) * RadialAttenuation;
+	// Assuming extinction==scattering, i.e. albedo=1, we can compute the integrale for each point on a ray while accouting for transmittance back to the view.
+	// It turns out that, under these conditions, the luminance and emissive factor is simply coverage. => see Mathematica details at the bottom of this page.
+	FogData.IntegratedLuminanceFactor = FogData.Coverage;
+
+	return FogData;
+}
+
+
 float3 ComputeFogInscattering(in FLocalHeightFogGPUInstanceData FogInstance, in FFogData FogData, in float3 RayDirWorld)
 {
 	half3 InScattering = 0;
@@ -147,8 +194,6 @@ void LocalHeightFogSplatPS(
 	float3 CamRayWorldDir = CamRayTranslatedWorldDir;
 
 	FLocalHeightFogGPUInstanceData FogInstance = LocalHeightFogInstances[InstanceIndex];
-	float4x4 Transform = FogInstance.Transform;
-	float4x4 InvTransform = FogInstance.InvTransform;
 
 	OutColor = float4(0, 0, 0, 1);
 
@@ -194,10 +239,10 @@ void LocalHeightFogSplatPS(
 			const float t = RadialAttenuation;
 			const float t2 = t * t;
 			const float t3 = t2 * t;
-			const float HP0 = 0;
-			const float HM0 = saturate(FogInstance.RadialAttenuationHermiteTangent);
-			const float HP1 = 1;
-			const float HM1 = saturate(FogInstance.RadialAttenuationHermiteTangent);
+			const float HP0 = 0.0f;
+			const float HP1 = 1.0f;
+			const float HM0 = 0.0f;	// tangent = 0 for soft transition
+			const float HM1 = 0.0f; // idem
 			RadialAttenuation = saturate((2 * t3 - 3 * t2 + 1) * HP0 + (t3 - 2 * t2 + t) * HM0 + (-2 * t3 + 3 * t2) * HP1 + (t3 - t2) * HM1);
 		}
 
@@ -208,7 +253,17 @@ void LocalHeightFogSplatPS(
 
 		if (RayTracedLength > 0.0)
 		{
-			FFogData FogData = EvaluateHeightFogIntegral(FogInstance, CamRayOriginSpace, CamRayOriginSpace + CamRayDirSpace * Length0, CamRayDirSpace, RayTracedLength, RadialAttenuation);
+			FFogData FogData = (FFogData)0;
+
+			if (FogInstance.FogMode == 0)
+			{
+				FogData = EvaluateHeightFogIntegral(FogInstance, CamRayOriginSpace, CamRayOriginSpace + CamRayDirSpace * Length0, CamRayDirSpace, RayTracedLength, RadialAttenuation);
+			}
+			else // (FogInstance.FogMode == 1)
+			{
+				// TODO EvaluateRadialFogIntegral needs to be optimised because it is redoing lots of computations. Maybe move HeightFogIntegral to also be in a local normalised space.
+				FogData = EvaluateRadialFogIntegral(FogInstance, CamRayWorldOrigin, CamRayWorldOrigin + CamRayWorldDir * Length0, CamRayWorldDir, RayTracedLength, RadialAttenuation);
+			}
 
 			OutColor.xyz	= ComputeFogInscattering(FogInstance, FogData, CamRayWorldDir);	// = Luminance
 			OutColor.a		= 1.0 - FogData.Coverage;										// = Transmittance
diff --git a/Engine/Source/Runtime/Engine/Classes/Components/LocalHeightFogComponent.h b/Engine/Source/Runtime/Engine/Classes/Components/LocalHeightFogComponent.h
--- a/Engine/Source/Runtime/Engine/Classes/Components/LocalHeightFogComponent.h
+++ b/Engine/Source/Runtime/Engine/Classes/Components/LocalHeightFogComponent.h
@@ -9,6 +9,13 @@
 
 class FLocalHeightFogSceneProxy;
 
+UENUM()
+enum class ELocalFogMode : uint8
+{
+	LocalHeightFog = 0,
+	LocalSphereFog = 1,
+};
+
 UCLASS(ClassGroup = Rendering, collapsecategories, hidecategories = (Object, Mobility, Activation, "Components|Activation"), editinlinenew, meta = (BlueprintSpawnableComponent), MinimalAPI)
 class ULocalHeightFogComponent : public USceneComponent
 {
@@ -16,6 +23,10 @@ class ULocalHeightFogComponent : public USceneComponent
 
 	~ULocalHeightFogComponent();
 
+	/** Controls the softness of the transition region when the volume is fading out. */
+	UPROPERTY(EditAnywhere, BlueprintReadOnly, Category = "Fog Mode")
+	ELocalFogMode FogMode = ELocalFogMode::LocalHeightFog;
+
 	/** Global density factor for this fog. */
 	UPROPERTY(EditAnywhere, BlueprintReadOnly, interp, Category = "Fog Distribution", meta = (UIMin = "0", UIMax = ".05", SliderExponent = 4.0, ClampMin = 0.0))
 	float FogDensity = 0.02f;
@@ -28,12 +39,8 @@ class ULocalHeightFogComponent : public USceneComponent
 	UPROPERTY(EditAnywhere, BlueprintReadOnly, interp, Category = "Fog Distribution", meta = (UIMin = "-5000.0", UIMax = "5000.0"))
 	float FogHeightOffset = 0.0f;
 
-	/** Controls the softness of the transition region when the volume is fading out. */
-	UPROPERTY(EditAnywhere, BlueprintReadOnly, interp, Category = "Fog Distribution", meta = (UIMin = "0.0", UIMax = "1.0", ClampMin = 0.0, ClampMax = 1.0))
-	float FogRadialAttenuationSoftness = 1.0f;
-
 	/** Controls how strong the radial attenuation of this fog volume is. */
-	UPROPERTY(EditAnywhere, BlueprintReadOnly, interp, Category = "Fog Distribution", meta = (UIMin = "0.001", UIMax = "4.0", SliderExponent = 3.0, ClampMin = 0.001))
+	UPROPERTY(EditAnywhere, BlueprintReadOnly, interp, Category = "Fog Distribution", meta = (UIMin = "0.0", UIMax = "4.0", SliderExponent = 3.0, ClampMin = 0.001))
 	float FogRadialAttenuation = 0.0f;
 
 	/** Controls the directionality of the scattering within this fog volume. */
@@ -52,6 +59,7 @@ public:
 
 	//~ Begin UObject Interface
 #if WITH_EDITOR
+	virtual bool CanEditChange(const FProperty* InProperty) const override;
 	virtual void PostEditChangeProperty(FPropertyChangedEvent& PropertyChangedEvent) override;
 #endif // WITH_EDITOR
 	//~ End UObject Interface
diff --git a/Engine/Source/Runtime/Engine/Private/Components/LocalHeightFogComponent.cpp b/Engine/Source/Runtime/Engine/Private/Components/LocalHeightFogComponent.cpp
--- a/Engine/Source/Runtime/Engine/Private/Components/LocalHeightFogComponent.cpp
+++ b/Engine/Source/Runtime/Engine/Private/Components/LocalHeightFogComponent.cpp
@@ -93,6 +93,25 @@ void ULocalHeightFogComponent::DestroyRenderState_Concurrent()
 
 #if WITH_EDITOR
 
+bool ULocalHeightFogComponent::CanEditChange(const FProperty* InProperty) const
+{
+	if (InProperty)
+	{
+		FString PropertyName = InProperty->GetName();
+
+		if (FogMode != ELocalFogMode::LocalHeightFog)
+		{
+			if (PropertyName == GET_MEMBER_NAME_STRING_CHECKED(ULocalHeightFogComponent, FogHeightFalloff)
+				|| PropertyName == GET_MEMBER_NAME_STRING_CHECKED(ULocalHeightFogComponent, FogHeightOffset))
+			{
+				return false;
+			}
+		}
+	}
+
+	return true;
+}
+
 void ULocalHeightFogComponent::PostEditChangeProperty(FPropertyChangedEvent& PropertyChangedEvent)
 {
 	Super::PostEditChangeProperty(PropertyChangedEvent);
diff --git a/Engine/Source/Runtime/Engine/Private/LocalHeightFogSceneProxy.cpp b/Engine/Source/Runtime/Engine/Private/LocalHeightFogSceneProxy.cpp
--- a/Engine/Source/Runtime/Engine/Private/LocalHeightFogSceneProxy.cpp
+++ b/Engine/Source/Runtime/Engine/Private/LocalHeightFogSceneProxy.cpp
@@ -15,7 +15,7 @@ FLocalHeightFogSceneProxy::FLocalHeightFogSceneProxy(const ULocalHeightFogCompon
 	, FogHeightFalloff(InComponent->FogHeightFalloff)
 	, FogHeightOffset(InComponent->FogHeightOffset + InComponent->GetComponentTransform().GetLocation().Z)
 	, FogRadialAttenuation(InComponent->FogRadialAttenuation)
-	, FogRadialAttenuationSoftness(InComponent->FogRadialAttenuationSoftness)
+	, FogMode((uint8)InComponent->FogMode)
 	, FogPhaseG(InComponent->FogPhaseG)
 	, FogAlbedo(InComponent->FogAlbedo)
 	, FogEmissive(InComponent->FogEmissive)
diff --git a/Engine/Source/Runtime/Engine/Public/LocalHeightFogSceneProxy.h b/Engine/Source/Runtime/Engine/Public/LocalHeightFogSceneProxy.h
--- a/Engine/Source/Runtime/Engine/Public/LocalHeightFogSceneProxy.h
+++ b/Engine/Source/Runtime/Engine/Public/LocalHeightFogSceneProxy.h
@@ -22,7 +22,8 @@ public:
 	float FogHeightFalloff;
 	float FogHeightOffset;
 	float FogRadialAttenuation;
-	float FogRadialAttenuationSoftness;
+
+	uint8 FogMode;
 
 	float FogPhaseG;
 	FLinearColor FogAlbedo;
diff --git a/Engine/Source/Runtime/Renderer/Private/LocalHeightFogRendering.cpp b/Engine/Source/Runtime/Renderer/Private/LocalHeightFogRendering.cpp
--- a/Engine/Source/Runtime/Renderer/Private/LocalHeightFogRendering.cpp
+++ b/Engine/Source/Runtime/Renderer/Private/LocalHeightFogRendering.cpp
@@ -84,7 +84,7 @@ public:
 	FVector3f Albedo;
 	float PhaseG;
 	FVector3f Emissive;
-	float RadialAttenuationHermiteTangent;
+	uint32 FogMode;
 };
 
 class FLocalHeightFogVS : public FGlobalShader
@@ -170,7 +170,8 @@ void RenderLocalHeightFog(
 			LocalHeightFogGPUInstanceDataIt->HeightFalloff = LHF->FogHeightFalloff * 0.01f;	// This scale is used to have artist author reasonable range.
 			LocalHeightFogGPUInstanceDataIt->HeightOffset = LHF->FogHeightOffset;
 			LocalHeightFogGPUInstanceDataIt->RadialAttenuation = LHF->FogRadialAttenuation;
-			LocalHeightFogGPUInstanceDataIt->RadialAttenuationHermiteTangent = 1.0f - LHF->FogRadialAttenuationSoftness;
+
+			LocalHeightFogGPUInstanceDataIt->FogMode = LHF->FogMode;
 
 			LocalHeightFogGPUInstanceDataIt->Albedo = FVector3f(LHF->FogAlbedo);
 			LocalHeightFogGPUInstanceDataIt->PhaseG = LHF->FogPhaseG;
