LocalHeightFogVolume - height fog mode now uses local space computations

diff --git a/Engine/Shaders/Private/LocalHeightFog.usf b/Engine/Shaders/Private/LocalHeightFog.usf
--- a/Engine/Shaders/Private/LocalHeightFog.usf
+++ b/Engine/Shaders/Private/LocalHeightFog.usf
@@ -67,33 +67,30 @@ struct FFogData
 FFogData EvaluateHeightFogIntegral(FLocalHeightFogGPUInstanceData FogInstance,  float3 CameraPos, float3 RayStart, float3 RayDir, float RayLength, float RadialAttenuation)
 {
 	FFogData FogData;
+	
+	float3 RayOriU = mul(float4(RayStart, 1.0f),                      FogInstance.InvTransform).xyz;
+	float3 RayIntU = mul(float4(RayStart + RayDir * RayLength, 1.0f), FogInstance.InvTransform).xyz;
+	float3 RayDirU = mul(float4(RayDir  , 0.0f),                      FogInstance.InvTransform).xyz;
+	RayDirU = normalize(RayDirU);
+
+	float RayLengthU = length(RayIntU - RayOriU);
 
-	float StartHeight = RayStart.z - FogInstance.HeightOffset; // accoutn for the fog offset.
+	float StartHeight = RayOriU.z - FogInstance.HeightOffset; // accoutn for the fog offset.
 
 	float OpticalDepth = 0.0f;
 #if 0
 	// Original integral, less artefact
-	OpticalDepth = (FogInstance.Density / FogInstance.HeightFalloff) * exp(-StartHeight * FogInstance.HeightFalloff) * (1.0 - exp(-RayDir.z * RayLength * FogInstance.HeightFalloff)) / RayDir.z;
+	OpticalDepth = (FogInstance.Density / FogInstance.HeightFalloff) * exp(-StartHeight * FogInstance.HeightFalloff) * (1.0 - exp(-RayDirU.z * RayLengthU * FogInstance.HeightFalloff)) / RayDirU.z;
 #elif 1
 	// Reworked to avoid large value as input to exp
 	//OpticalDepth = (FogInstance.Density / FogInstance.HeightFalloff) * (exp(-StartHeight * FogInstance.HeightFalloff) - exp(-StartHeight * FogInstance.HeightFalloff - RayDir.z * RayLength * FogInstance.HeightFalloff)) / RayDir.z;
 
 	// But we use an updated formulation to avoid large input to exp leading to inf and nan.
 	float Factor0 = StartHeight * FogInstance.HeightFalloff;
-	float Factor1 = RayDir.z * RayLength * FogInstance.HeightFalloff;
-//	if (Factor0 < 0	// If we start below the height fog horizon, where it is already thick according to falloff...
-//		&& Factor1 < 0)	// ... and that the difference between entry and exit point is large...
-//	{
-//		// ... then exp is likely going to explode.
-//		OpticalDepth = FogInstance.Density * exp(-Factor0) * RayLength;	// Force to a large optical depth matching the extinction at ray originand according to ray segment length.
-//	}
-//	else
-//	{
-//		OpticalDepth = (FogInstance.Density / (FogInstance.HeightFalloff * RayDir.z)) * (exp(-Factor0) - exp(-(Factor0 + Factor1)));
-//	}
+	float Factor1 = RayDirU.z * RayLengthU * FogInstance.HeightFalloff;
 	Factor0 = max(-80.0f, Factor0); // Simply clamp Factor0 to a valid range in order to not have exp explode.
 
-	OpticalDepth = (FogInstance.Density / (FogInstance.HeightFalloff * RayDir.z)) * (exp(-Factor0) - exp(-(Factor0 + Factor1)));
+	OpticalDepth = (FogInstance.Density / (FogInstance.HeightFalloff * RayDirU.z)) * (exp(-Factor0) - exp(-(Factor0 + Factor1)));
 #endif
 
 	// Transmittance is the 
@@ -257,11 +254,14 @@ void LocalHeightFogSplatPS(
 
 			if (FogInstance.FogMode == 0)
 			{
-				FogData = EvaluateHeightFogIntegral(FogInstance, CamRayOriginSpace, CamRayOriginSpace + CamRayDirSpace * Length0, CamRayDirSpace, RayTracedLength, RadialAttenuation);
+				FogData = EvaluateHeightFogIntegral(FogInstance, CamRayWorldOrigin, CamRayWorldOrigin + CamRayWorldDir * Length0, CamRayWorldDir, RayTracedLength, RadialAttenuation);
 			}
 			else // (FogInstance.FogMode == 1)
 			{
+				// TODO
 				// TODO EvaluateRadialFogIntegral needs to be optimised because it is redoing lots of computations. Maybe move HeightFogIntegral to also be in a local normalised space.
+				// TODO It is like that for now while iterating with artists.
+				// TODO
 				FogData = EvaluateRadialFogIntegral(FogInstance, CamRayWorldOrigin, CamRayWorldOrigin + CamRayWorldDir * Length0, CamRayWorldDir, RayTracedLength, RadialAttenuation);
 			}
 
diff --git a/Engine/Source/Runtime/Engine/Classes/Components/LocalHeightFogComponent.h b/Engine/Source/Runtime/Engine/Classes/Components/LocalHeightFogComponent.h
--- a/Engine/Source/Runtime/Engine/Classes/Components/LocalHeightFogComponent.h
+++ b/Engine/Source/Runtime/Engine/Classes/Components/LocalHeightFogComponent.h
@@ -28,15 +28,15 @@ class ULocalHeightFogComponent : public USceneComponent
 	ELocalFogMode FogMode = ELocalFogMode::LocalHeightFog;
 
 	/** Global density factor for this fog. */
-	UPROPERTY(EditAnywhere, BlueprintReadOnly, interp, Category = "Fog Distribution", meta = (UIMin = "0", UIMax = ".05", SliderExponent = 4.0, ClampMin = 0.0))
-	float FogDensity = 0.02f;
+	UPROPERTY(EditAnywhere, BlueprintReadOnly, interp, Category = "Fog Distribution", meta = (UIMin = "0", UIMax = "10.0", SliderExponent = 2.0, ClampMin = 0.0))
+	float FogDensity = 5.0f;
 
 	/** Controls how the density decreases as height increases. Smaller values make the visible transition larger. */
-	UPROPERTY(EditAnywhere, BlueprintReadOnly, interp, Category = "Fog Distribution", meta = (UIMin = "0.001", UIMax = "10.0", SliderExponent = 4.0, ClampMin = 0.001))
-	float FogHeightFalloff = 3.0f;
+	UPROPERTY(EditAnywhere, BlueprintReadOnly, interp, Category = "Fog Distribution", meta = (UIMin = "0.001", UIMax = "5000", SliderExponent = 2.0, ClampMin = 0.001))
+	float FogHeightFalloff = 1000.0f;
 
 	/** Height offset, relative to the actor Z position. */
-	UPROPERTY(EditAnywhere, BlueprintReadOnly, interp, Category = "Fog Distribution", meta = (UIMin = "-5000.0", UIMax = "5000.0"))
+	UPROPERTY(EditAnywhere, BlueprintReadOnly, interp, Category = "Fog Distribution", meta = (UIMin = "-2.0", UIMax = "2.0"))
 	float FogHeightOffset = 0.0f;
 
 	/** Controls how strong the radial attenuation of this fog volume is. */
