Path Tracer: Simplify handling of "micro occlusion" which lowers F90 color as F0 tends to zero

There is no need to multiply F0 by the microocclusion factor as it is already becoming 0.

Also make the equivalent change on the substrate side.

diff --git a/Engine/Shaders/Private/PathTracing/Material/PathTracingSubstrate.ush b/Engine/Shaders/Private/PathTracing/Material/PathTracingSubstrate.ush
--- a/Engine/Shaders/Private/PathTracing/Material/PathTracingSubstrate.ush
+++ b/Engine/Shaders/Private/PathTracing/Material/PathTracingSubstrate.ush
@@ -39,14 +39,9 @@ struct FSubstrateSingleSlabData
 
 		V = mul(Basis, V_World);
 
+		// make sure F0=0.0 fades off the specular lobe completely
 		F0 = Payload.SpecularColor;
-		F90 = Payload.SpecularEdgeColor;
-		{
-			// make sure F0=0.0 fades off the specular lobe completely
-			const float SpecularMicroOcclusion = F0RGBToMicroOcclusion(F0);
-			F0 *= SpecularMicroOcclusion;
-			F90 *= SpecularMicroOcclusion;
-		}
+		F90 = Payload.SpecularEdgeColor * F0RGBToMicroOcclusion(F0);
 
 		FuzzRoughness = MakeRoughnessSafe(Payload.FuzzRoughness, STRATA_MIN_FUZZ_ROUGHNESS);
 
@@ -273,14 +268,9 @@ struct FSubstrateClearCoatData
 
 		V = mul(Basis, V_World);
 
+		// make sure F0=0.0 fades off the specular lobe completely
 		F0 = Payload.SpecularColor;
-		F90 = Payload.SpecularEdgeColor;
-		{
-			// make sure F0=0.0 fades off the specular lobe completely
-			const float SpecularMicroOcclusion = F0RGBToMicroOcclusion(F0);
-			F0 *= SpecularMicroOcclusion;
-			F90 *= SpecularMicroOcclusion;
-		}
+		F90 = Payload.SpecularEdgeColor * F0RGBToMicroOcclusion(F0);
 
 		FuzzRoughness = MakeRoughnessSafe(Payload.FuzzRoughness, STRATA_MIN_FUZZ_ROUGHNESS);
 
diff --git a/Engine/Shaders/Private/PathTracing/Material/PathTracingSubstrateThinGlass.ush b/Engine/Shaders/Private/PathTracing/Material/PathTracingSubstrateThinGlass.ush
--- a/Engine/Shaders/Private/PathTracing/Material/PathTracingSubstrateThinGlass.ush
+++ b/Engine/Shaders/Private/PathTracing/Material/PathTracingSubstrateThinGlass.ush
@@ -51,7 +51,7 @@ struct FSubstrateThinGlassSlabData
 		V = mul(Basis, V_World);
 
 		F0 = F0RGBToF0(Payload.SpecularColor);
-		F90 = saturate(50.0 * F0);
+		F90 = F0ToMicroOcclusion(F0);
 		TransmissionColor = Payload.GetTransmittanceColor();
 		FuzzRoughness = MakeRoughnessSafe(Payload.FuzzRoughness, STRATA_MIN_FUZZ_ROUGHNESS);
 
diff --git a/Engine/Shaders/Private/PathTracing/PathTracingMaterialHitShader.usf b/Engine/Shaders/Private/PathTracing/PathTracingMaterialHitShader.usf
--- a/Engine/Shaders/Private/PathTracing/PathTracingMaterialHitShader.usf
+++ b/Engine/Shaders/Private/PathTracing/PathTracingMaterialHitShader.usf
@@ -621,13 +621,7 @@ RAY_TRACING_ENTRY_CLOSEST_HIT(PathTracingMaterialCHS,
 
 							float3 DiffuseColor = SLAB_DIFFUSEALBEDO(BSDF);
 							float3 F0 = SLAB_F0(BSDF);
-							float3 F90 = SLAB_F90(BSDF);
-							{
-								// make sure F0=0.0 fades off the specular lobe completely
-								const float SpecularMicroOcclusion = F0RGBToMicroOcclusion(F0);
-								F0 *= SpecularMicroOcclusion;
-								F90 *= SpecularMicroOcclusion;
-							}
+							float3 F90 = SLAB_F90(BSDF) * F0RGBToMicroOcclusion(F0);
 
 							float Roughness0 = MakeRoughnessSafe(SLAB_ROUGHNESS(BSDF), PackedPayload.GetPathRoughness());
 							// STRATA_TODO: Support Fresnel82?
diff --git a/Engine/Shaders/Private/Strata/StrataEvaluation.ush b/Engine/Shaders/Private/Strata/StrataEvaluation.ush
--- a/Engine/Shaders/Private/Strata/StrataEvaluation.ush
+++ b/Engine/Shaders/Private/Strata/StrataEvaluation.ush
@@ -351,9 +351,7 @@ FStrataEvaluateResult StrataEvaluateBSDFCommon(FStrataBSDFContext BSDFContext, F
 			const bool bHaziness		= BSDF_GETHASHAZINESS(BSDFContext.BSDF);
 
 			// Specular occlusion is only used here once to affect the F90 source parameters. F0 will decrease naturally to 0.
-			const float SpecularMicroOcclusion = F0RGBToMicroOcclusion(F0);
-			F0  *= SpecularMicroOcclusion;
-			F90 *= SpecularMicroOcclusion;
+			F90 *= F0RGBToMicroOcclusion(F0);
 
 			if (Settings.bForceFullyRough)
 			{
@@ -980,9 +978,7 @@ FStrataEvaluateResult StrataEvaluateBSDFCommon(FStrataBSDFContext BSDFContext, F
 			const float SafeRoughness	= MakeRoughnessSafe(EYE_ROUGHNESS(BSDFContext.BSDF));
 
 			// Specular occlusion is only used here once to affect the F90 source parameters. F0 will decrease naturally to 0.
-			const float SpecularMicroOcclusion = F0RGBToMicroOcclusion(F0);
-			F0  *= SpecularMicroOcclusion;
-			F90 *= SpecularMicroOcclusion;
+			F90 *= F0RGBToMicroOcclusion(F0);
 
 			if (Settings.bForceFullyRough)
 			{
@@ -1422,9 +1418,7 @@ FStrataEnvLightResult StrataEvaluateForEnvLight(FStrataBSDFContext BSDFContext,
 			float3 F90 				= SLAB_F90(BSDFContext.BSDF);
 
 			// Specular occlusion is only used here once to affect the F90 source parameters. F0 will decrease naturally to 0.
-			const float SpecularMicroOcclusion = F0RGBToMicroOcclusion(F0);
-			F0  *= SpecularMicroOcclusion;
-			F90 *= SpecularMicroOcclusion;
+			F90 *= F0RGBToMicroOcclusion(F0);
 
 			if (Settings.bForceFullyRough)
 			{
@@ -1694,9 +1688,7 @@ FStrataEnvLightResult StrataEvaluateForEnvLight(FStrataBSDFContext BSDFContext,
 			float3 F90 				= EYE_F90(BSDFContext.BSDF);
 
 			// Specular occlusion is only used here once to affect the F90 source parameters. F0 will decrease naturally to 0.
-			const float SpecularMicroOcclusion = F0RGBToMicroOcclusion(F0);
-			F0  *= SpecularMicroOcclusion;
-			F90 *= SpecularMicroOcclusion;
+			F90 *= F0RGBToMicroOcclusion(F0);
 
 			if (Settings.bForceFullyRough)
 			{
