Path Tracer: Remove unused functions

diff --git a/Engine/Shaders/Private/PathTracing/PathTracingCompositingPixelShader.usf b/Engine/Shaders/Private/PathTracing/PathTracingCompositingPixelShader.usf
--- a/Engine/Shaders/Private/PathTracing/PathTracingCompositingPixelShader.usf
+++ b/Engine/Shaders/Private/PathTracing/PathTracingCompositingPixelShader.usf
@@ -12,29 +12,6 @@ int ProgressDisplayEnabled;
 Texture2D<float4> RadianceTexture;
 Texture2D<float4> NormalDepthTexture;
 
-float TonemapGammaAndHDR(float Value)
-{
-	if (Value > 1.0)
-	{
-		Value = 1.0 + 2.0 * log(Value);
-	}
-	else
-	{
-		Value = pow(Value, 2.0);
-	}
-	return Value;
-}
-
-float NoTonemap(float Value)
-{
-	return Value;
-}
-
-float Tonemap(float Value)
-{
-	return TonemapGammaAndHDR(Value);
-}
-
 void CompositeMain(
 	in noperspective float2 UV : TEXCOORD0,
 	out float4 OutColor : SV_Target0,
