Workaround shader compile error for a platform that doesn't support the typedef keyword

diff --git a/Engine/Shaders/Private/Quaternion.ush b/Engine/Shaders/Private/Quaternion.ush
--- a/Engine/Shaders/Private/Quaternion.ush
+++ b/Engine/Shaders/Private/Quaternion.ush
@@ -1,6 +1,6 @@
 // Copyright Epic Games, Inc. All Rights Reserved.
 
-typedef half4 FQuat;
+#define FQuat half4
 
 FQuat QuatFromAxisAngle(in half3 AxisN, in float AngleRad)
 {
