Remove ClusterCullingData from public interface, as now all culling is done within the HairStrands plugin.

diff --git a/Engine/Plugins/Runtime/HairStrands/Source/HairStrandsCore/Private/GroomDebug.cpp b/Engine/Plugins/Runtime/HairStrands/Source/HairStrandsCore/Private/GroomDebug.cpp
--- a/Engine/Plugins/Runtime/HairStrands/Source/HairStrandsCore/Private/GroomDebug.cpp
+++ b/Engine/Plugins/Runtime/HairStrands/Source/HairStrandsCore/Private/GroomDebug.cpp
@@ -59,9 +59,6 @@ static FAutoConsoleVariableRef CVarHairDebugMeshProjection_Sim_HairDeformedSampl
 static int32 GHairCardsAtlasDebug = 0;
 static FAutoConsoleVariableRef CVarHairCardsAtlasDebug(TEXT("r.HairStrands.Cards.DebugAtlas"), GHairCardsAtlasDebug, TEXT("Draw debug hair cards atlas."));
 
-static int32 GHairCardsVoxelDebug = 0;
-static FAutoConsoleVariableRef CVarHairCardsVoxelDebug(TEXT("r.HairStrands.Cards.DebugVoxel"), GHairCardsVoxelDebug, TEXT("Draw debug hair cards voxel datas."));
-
 ///////////////////////////////////////////////////////////////////////////////////////////////////
 
 FCachedGeometry GetCacheGeometryForHair(
@@ -377,87 +374,6 @@ static void AddDebugProjectionHairPass(
 		FIntVector(FMath::DivideAndRoundUp(PrimitiveCount, 256u), 1, 1));
 }	
 
-///////////////////////////////////////////////////////////////////////////////////////////////////
-class FVoxelPlainRaymarchingCS : public FGlobalShader
-{
-	DECLARE_GLOBAL_SHADER(FVoxelPlainRaymarchingCS);
-	SHADER_USE_PARAMETER_STRUCT(FVoxelPlainRaymarchingCS, FGlobalShader);
-
-	BEGIN_SHADER_PARAMETER_STRUCT(FParameters, )
-		SHADER_PARAMETER_STRUCT_INCLUDE(ShaderPrint::FShaderParameters, ShaderPrintParameters)
-		SHADER_PARAMETER(FVector2f, OutputResolution)		
-		SHADER_PARAMETER(FIntVector, Voxel_Resolution)
-		SHADER_PARAMETER(float, Voxel_VoxelSize)
-		SHADER_PARAMETER(FVector3f, Voxel_MinBound)
-		SHADER_PARAMETER(FVector3f, Voxel_MaxBound)
-		SHADER_PARAMETER_SRV(Buffer, Voxel_TangentBuffer)
-		SHADER_PARAMETER_SRV(Buffer, Voxel_NormalBuffer)
-		SHADER_PARAMETER_SRV(Buffer, Voxel_DensityBuffer)
-		SHADER_PARAMETER_RDG_BUFFER_SRV(Buffer, Voxel_ProcessedDensityBuffer)
-		SHADER_PARAMETER_STRUCT_REF(FViewUniformShaderParameters, ViewUniformBuffer)
-		SHADER_PARAMETER_RDG_TEXTURE_UAV(RWTexture2D, OutputTexture)
-		END_SHADER_PARAMETER_STRUCT()
-
-public:
-	static bool ShouldCompilePermutation(const FGlobalShaderPermutationParameters& Parameters) { return IsHairStrandsSupported(EHairStrandsShaderType::Strands, Parameters.Platform); }
-	static void ModifyCompilationEnvironment(const FGlobalShaderPermutationParameters& Parameters, FShaderCompilerEnvironment& OutEnvironment)
-	{
-		FGlobalShader::ModifyCompilationEnvironment(Parameters, OutEnvironment);
-		OutEnvironment.SetDefine(TEXT("SHADER_CARDS_VOXEL"), 1);
-	}	
-};
-
-IMPLEMENT_GLOBAL_SHADER(FVoxelPlainRaymarchingCS, "/Engine/Private/HairStrands/HairStrandsDebug.usf", "MainCS", SF_Compute);
-
-static void AddVoxelPlainRaymarchingPass(
-	FRDGBuilder& GraphBuilder,
-	const FSceneView& View,
-	FGlobalShaderMap* ShaderMap,
-	const FHairGroupInstance* Instance,
-	const FShaderPrintData* ShaderPrintData,
-	FRDGTextureRef OutputTexture)
-{
-#if 0 // #hair_todo: renable if needed
-	const FHairStrandClusterData::FHairGroup& HairGroupClusters = HairClusterData.HairGroups[DataIndex];
-
-	FViewInfo& View = Views[ViewIndex];
-	if (ShaderPrint::IsEnabled(View.ShaderPrintData))
-	{
-		if (Instance->HairGroupPublicData->VFInput.GeometryType != EHairGeometryType::Cards)
-			return;
-
-		FSceneTextureParameters SceneTextures;
-		SetupSceneTextureParameters(GraphBuilder, &SceneTextures);
-
-		const FIntPoint OutputResolution(OutputTexture->Desc.Extent);
-		const FHairCardsVoxel& CardsVoxel = Instance->HairGroupPublicData->VFInput.Cards.Voxel;
-
-		FRDGBufferRef VoxelDensityBuffer2 = nullptr;
-		AddVoxelProcessPass(GraphBuilder, View, CardsVoxel, VoxelDensityBuffer2);
-
-		FVoxelPlainRaymarchingCS::FParameters* Parameters = GraphBuilder.AllocParameters<FVoxelPlainRaymarchingCS::FParameters>();
-		Parameters->ViewUniformBuffer	= View.ViewUniformBuffer;
-		Parameters->OutputResolution	= OutputResolution;
-		Parameters->Voxel_Resolution	= CardsVoxel.Resolution;
-		Parameters->Voxel_VoxelSize		= CardsVoxel.VoxelSize;
-		Parameters->Voxel_MinBound		= CardsVoxel.MinBound;
-		Parameters->Voxel_MaxBound		= CardsVoxel.MaxBound;
-		Parameters->Voxel_TangentBuffer	= CardsVoxel.TangentBuffer.SRV;
-		Parameters->Voxel_NormalBuffer	= CardsVoxel.NormalBuffer.SRV;
-		Parameters->Voxel_DensityBuffer = CardsVoxel.DensityBuffer.SRV;
-		Parameters->Voxel_ProcessedDensityBuffer = GraphBuilder.CreateSRV(VoxelDensityBuffer2, PF_R32_UINT);
-
-		ShaderPrint::SetParameters(GraphBuilder, ShaderPrintData, Parameters->ShaderPrintParameters);
-		//ShaderPrint::SetParameters(View, Parameters->ShaderPrintParameters);
-		Parameters->OutputTexture = GraphBuilder.CreateUAV(OutputTexture);
-
-		TShaderMapRef<FVoxelPlainRaymarchingCS> ComputeShader(ShaderMap);
-		const FIntVector DispatchCount = DispatchCount.DivideAndRoundUp(FIntVector(OutputTexture->Desc.Extent.X, OutputTexture->Desc.Extent.Y, 1), FIntVector(8, 8, 1));
-		FComputeShaderUtils::AddPass(GraphBuilder, RDG_EVENT_NAME("HairStrands::VoxelPlainRaymarching"), ComputeShader, Parameters, DispatchCount);
-	}
-#endif
-}
-
 ///////////////////////////////////////////////////////////////////////////////////////////////////
 
 class FDrawDebugClusterAABBCS : public FGlobalShader
@@ -1308,16 +1224,6 @@ void RunHairStrandsDebug(
 		}
 	}
 
-	if (GHairCardsVoxelDebug > 0)
-	{
-		for (FHairStrandsInstance* AbstractInstance : Instances)
-		{
-			FHairGroupInstance* Instance = static_cast<FHairGroupInstance*>(AbstractInstance);
-
-			AddVoxelPlainRaymarchingPass(GraphBuilder, View, ShaderMap, Instance, ShaderPrintData, SceneColorTexture);
-		}
-	}
-
 	if (GHairCardsAtlasDebug > 0)
 	{
 		for (FHairStrandsInstance* AbstractInstance : Instances)
diff --git a/Engine/Plugins/Runtime/HairStrands/Source/HairStrandsCore/Private/HairStrandsClusterCulling.h b/Engine/Plugins/Runtime/HairStrands/Source/HairStrandsCore/Private/HairStrandsClusterCulling.h
--- a/Engine/Plugins/Runtime/HairStrands/Source/HairStrandsCore/Private/HairStrandsClusterCulling.h
+++ b/Engine/Plugins/Runtime/HairStrands/Source/HairStrandsCore/Private/HairStrandsClusterCulling.h
@@ -6,17 +6,57 @@
 #include "RendererInterface.h"
 #include "HairStrandsInterface.h"
 
-struct FHairCullingParams
-{
-	bool bCullingProcessSkipped	= false;
-};
-
 struct FHairGroupInstance;
-struct FHairStrandClusterData;
 class FViewInfo;
 class FGlobalShaderMap;
 class FRDGBuilder;
 
+struct FHairStrandClusterData
+{
+	struct FHairGroup
+	{
+		uint32 ClusterCount = 0;
+		uint32 VertexCount = 0;
+		uint32 MaxPointPerCurve = 0;
+
+		float LODIndex = -1;
+		float LODBias = 0.0f;
+		bool bVisible = false;
+
+		FRDGExternalBuffer* CurveBuffer = nullptr;
+		FRDGExternalBuffer* PointLODBuffer = nullptr;
+
+		// See FHairStrandsClusterCullingResource fro details about those buffers.
+		FRDGExternalBuffer* GroupAABBBuffer = nullptr;
+		FRDGExternalBuffer* ClusterAABBBuffer = nullptr;
+		FRDGExternalBuffer* ClusterInfoBuffer = nullptr; // SRV
+		FRDGExternalBuffer* ClusterLODInfoBuffer = nullptr; // SRV
+		FRDGExternalBuffer* CurveToClusterIdBuffer = nullptr; // SRV
+		FRDGExternalBuffer* ClusterVertexIdBuffer = nullptr;
+
+		TRefCountPtr<FRDGPooledBuffer> ClusterIdBuffer;
+		TRefCountPtr<FRDGPooledBuffer> ClusterIndexOffsetBuffer;
+		TRefCountPtr<FRDGPooledBuffer> ClusterIndexCountBuffer;
+
+		// Culling & LOD output
+		FRDGExternalBuffer* GetCulledCurveBuffer() const				{ return HairGroupPublicPtr ? &HairGroupPublicPtr->GetCulledCurveBuffer() : nullptr; }
+		FRDGExternalBuffer* GetCulledVertexIdBuffer() const				{ return HairGroupPublicPtr ? &HairGroupPublicPtr->GetCulledVertexIdBuffer() : nullptr; }
+		FRDGExternalBuffer* GetCulledVertexRadiusScaleBuffer() const	{ return HairGroupPublicPtr ? &HairGroupPublicPtr->GetCulledVertexRadiusScaleBuffer() : nullptr; }
+		bool GetCullingResultAvailable() const							{ return HairGroupPublicPtr ? HairGroupPublicPtr->GetCullingResultAvailable() : false; }
+		void SetCullingResultAvailable(bool b)							{ if (HairGroupPublicPtr) HairGroupPublicPtr->SetCullingResultAvailable(b); }
+
+		TRefCountPtr<FRDGPooledBuffer> ClusterDebugInfoBuffer;	// Null if this debug is not enabled.
+		FRDGBufferRef CulledClusterCountBuffer = nullptr;
+		FRDGBufferRef CulledCluster1DIndirectArgsBuffer = nullptr;
+		FRDGBufferRef CulledCluster2DIndirectArgsBuffer = nullptr;
+		uint32 GroupSize1D = 0;
+
+		FHairGroupPublicData* HairGroupPublicPtr = nullptr;
+	};
+
+	TArray<FHairGroup> HairGroups;
+};
+
 void AddInstanceToClusterData(
 	FHairGroupInstance* In,
 	FHairStrandClusterData& Out);
diff --git a/Engine/Plugins/Runtime/HairStrands/Source/HairStrandsCore/Private/HairStrandsInterpolation.cpp b/Engine/Plugins/Runtime/HairStrands/Source/HairStrandsCore/Private/HairStrandsInterpolation.cpp
--- a/Engine/Plugins/Runtime/HairStrands/Source/HairStrandsCore/Private/HairStrandsInterpolation.cpp
+++ b/Engine/Plugins/Runtime/HairStrands/Source/HairStrandsCore/Private/HairStrandsInterpolation.cpp
@@ -28,6 +28,7 @@
 #include "GroomInstance.h"
 #include "SystemTextures.h"
 #include "ShaderPrintParameters.h"
+#include "HairStrandsClusterCulling.h"
 
 static float GHairRaytracingRadiusScale = 0;
 static FAutoConsoleVariableRef CVarHairRaytracingRadiusScale(TEXT("r.HairStrands.RaytracingRadiusScale"), GHairRaytracingRadiusScale, TEXT("Override the per instance scale factor for raytracing hair strands geometry (0: disabled, >0:enabled)"));
diff --git a/Engine/Shaders/Private/HairStrands/HairStrandsDebug.usf b/Engine/Shaders/Private/HairStrands/HairStrandsDebug.usf
--- a/Engine/Shaders/Private/HairStrands/HairStrandsDebug.usf
+++ b/Engine/Shaders/Private/HairStrands/HairStrandsDebug.usf
@@ -1949,223 +1949,6 @@ void MainCS(uint3 DispatchThreadId : SV_DispatchThreadID)
 
 #endif // SHADER_CARDS_GUIDE
 
-///////////////////////////////////////////////////////////////////////////////////////////////////
-// Visualize voxelized hair strands (used for cards generation)
-
-#if SHADER_CARDS_VOXEL
-
-#include "../CommonViewUniformBuffer.ush"
-#include "../SceneTextureParameters.ush"
-#include "../DeferredShadingCommon.ush"
-#include "../ShaderPrint.ush"
-#include "../PositionReconstructionCommon.ush" 
-
-uint3 PositionToCoord(float3 P, float3 InMinAABB, float3 InMaxAABB, uint3 InResolution)
-{
-	return clamp(
-		saturate((P - InMinAABB) / (InMaxAABB - InMinAABB)) * InResolution,
-		uint3(0, 0, 0),
-		InResolution - 1);
-}
-
-uint3 PositionToCoordUnclampled(float3 P, float3 InMinAABB, float3 InMaxAABB, uint3 InResolution)
-{
-	return saturate((P - InMinAABB) / (InMaxAABB - InMinAABB)) * InResolution;
-}
-
-uint CoordToIndex(uint3 InCoord, uint3 InResolution, uint LinearOffset)
-{
-	return
-		InCoord.x +
-		InCoord.y * InResolution.x +
-		InCoord.z * InResolution.x * InResolution.y +
-		LinearOffset;
-}
-
-uint3 IndexToCoord(uint InIndex, uint3 InResolution)
-{
-	const uint SliceSize = (InResolution.x * InResolution.y);
-	const uint SliceIndex = InIndex % SliceSize;
-
-	uint3 OutCoord = 0;
-	OutCoord.x = SliceIndex % InResolution.x;
-	OutCoord.y = SliceIndex / InResolution.x;
-	OutCoord.z = InIndex / SliceSize;
-
-	return OutCoord;
-}
-
-bool IsInVoxelBounds(float3 P, float3 MinP, float3 MaxP)
-{
-	return
-		P.x >= MinP.x && P.y >= MinP.y && P.z >= MinP.z &&
-		P.x <= MaxP.x && P.y <= MaxP.y && P.z <= MaxP.z;
-}
-
-float3 EncodeTangent(float3 N)
-{
-	return N * 0.5 + 0.5;
-}
-
-float3 DecodeTangent(float3 N)
-{
-	return N * 2 - 1;
-}
-
-uint PackDirection(float3 In)
-{
-	float3 T = EncodeTangent(In);
-	return
-		((0xFF & uint(T.x * 255))) |
-		((0xFF & uint(T.y * 255)) << 8) |
-		((0xFF & uint(T.z * 255)) << 16);
-}
-
-float3 UnpackDirection(uint In)
-{
-	return DecodeTangent(float3(
-		float((In)       & 0xFF) / 255.f,
-		float((In >> 8)  & 0xFF) / 255.f,
-		float((In >> 16) & 0xFF) / 255.f));
-}
-
-
-float2 OutputResolution;
-int3 Voxel_Resolution;
-float Voxel_VoxelSize;
-float3 Voxel_MinBound;
-float3 Voxel_MaxBound;
-Buffer<float4> Voxel_TangentBuffer;
-Buffer<float4> Voxel_NormalBuffer;
-Buffer<uint>   Voxel_DensityBuffer;
-Buffer<uint>   Voxel_ProcessedDensityBuffer;
-
-RWTexture2D<float4> OutputTexture;
-
-[numthreads(8, 8, 1)]
-void MainCS(uint3 DispatchThreadId : SV_DispatchThreadID)
-{
-	const float2 PixelCoord = DispatchThreadId.xy + 0.5f;
-	const float2 UV = PixelCoord / float2(OutputResolution);
-
-	const bool bDebugEnabled = all(int2(PixelCoord) == GetCursorPos());
-
-	float  WorldDepth = 0;
-	float  ClipDepth = 0;
-
-
-	if (uint(PixelCoord.x) == 0 && uint(PixelCoord.y) == 0)
-	{
-		AddAABBWS(Voxel_MinBound, Voxel_MaxBound, float4(1, 0.5f, 0, 1));
-	}
-
-	bool bIsValid = false;
-	float3 HitColor = 0;
-	{
-		// Reverse-Z
-		const float3 O = LWCHackToFloat(PrimaryView.WorldCameraOrigin);
-		const float3 E = ReconstructWorldPositionFromDepth(UV, ConvertFromDeviceZ(0.0001f));
-		const float3 D = (E - O);
-
-		float PixelRadius = -1;
-		const float2 HitT = LineBoxIntersect(O, E, Voxel_MinBound, Voxel_MaxBound);
-		float AccDensity = 0;
-		if (HitT.x < HitT.y)
-		{
-			// Use different offseting to see either: 
-			// * Individual voxel (for debugging mainly)
-			// * Smoothed voxel (as used for transmittance computation)
-			const float3 Offset = 0;
-			const float Dither = InterleavedGradientNoise(PixelCoord.xy, 1);
-
-			// Intersect the indirection table to see if a brick is empty or not
-			const uint MaxStep = 256;
-			const float Delta = (HitT.y - HitT.x) / float(MaxStep);
-			for (uint StepIt = 0; StepIt < MaxStep; ++StepIt)
-			{
-				const float3 HitP = Offset + O + ((Dither + StepIt) * Delta + HitT.x) * D;
-				const int3 VoxelCoord = PositionToCoord(HitP, Voxel_MinBound, Voxel_MaxBound, Voxel_Resolution);
-				const uint VoxelLinearCoord = CoordToIndex(VoxelCoord, Voxel_Resolution, 0);
-
-				const uint Density = Voxel_DensityBuffer.Load(VoxelLinearCoord);
-
-				if (Density > 0)
-				{
-					float3 UnoccludedDirection = UnpackDirection(Density);
-
-
-					// Filtering
-					if (0)
-					{
-						int KernelSize = 1;
-						uint ValidCount = 0;
-						UnoccludedDirection = 0;
-						for (int Z = -KernelSize; Z <= KernelSize; ++Z)
-						for (int Y = -KernelSize; Y <= KernelSize; ++Y)
-						for (int X = -KernelSize; X <= KernelSize; ++X)
-						{
-							const int3 Offset = int3(X, Y, Z);
-							const int3 C = VoxelCoord + Offset;
-							if (all(C >= 0) && all(C < Voxel_Resolution))
-							{
-								const float w = all(Offset == VoxelCoord) ? 2 : 1;
-								const uint LinearC = CoordToIndex(C, Voxel_Resolution, 0);
-								const uint PackedDir = Voxel_ProcessedDensityBuffer[LinearC];
-								if (PackedDir > 0)
-								{
-									const float3 Dir = UnpackDirection(PackedDir);
-									UnoccludedDirection += normalize(Dir) * w;
-									ValidCount++;
-								}
-							}
-						}
-
-						if (ValidCount>0)
-						{
-							UnoccludedDirection = normalize(UnoccludedDirection);
-
-						}						
-
-						// CPU tangent
-						//UnoccludedDirection = Voxel_TangentBuffer.Load(VoxelLinearCoord) / Density;
-						//UnoccludedDirection = normalize(UnoccludedDirection);
-					}
-
-					const float3 Tangent = Voxel_TangentBuffer.Load(VoxelLinearCoord).xyz;
-					const float3 Normal  = Voxel_NormalBuffer.Load(VoxelLinearCoord).xyz;
-					//UnoccludedDirection = Tangent;
-					UnoccludedDirection = Normal;
-
-					AccDensity = 1;
-					bIsValid = true;
-					HitColor = VoxelCoord / 12.f;
-
-					HitColor = UnoccludedDirection * 2 + 1;
-
-					//HitColor = UnoccludedDirection; // (UnoccludedDirection + 1) * 0.5f;
-					//HitColor = float(Density) / 10.f;
-					if (bDebugEnabled)
-					{
-						//AddLineWS(HitP, HitP + UnoccludedDirection * 3, float4(1, 0, 0, 1), float4(1, 1, 0, 1));
-
-						AddLineWS(HitP, HitP + Tangent * 3, float4(0, 0, 1, 1), float4(0, 0, 1, 1));
-						AddLineWS(HitP, HitP + Normal * 3, float4(1, 0, 0, 1), float4(1, 0, 0, 1));
-					}
-
-					break;
-				}
-			}
-		}
-	}
-	
-	if (bIsValid)
-	{
-		OutputTexture[DispatchThreadId.xy] = float4(HitColor, 1);
-	}
-}
-
-#endif //SHADER_CARDS_VOXEL
-
 ///////////////////////////////////////////////////////////////////////////////////////////////////
 // Visualize directional occlusion from voxelized hair strands (used for cards generation)
 
diff --git a/Engine/Source/Runtime/Renderer/Public/HairStrandsInterface.h b/Engine/Source/Runtime/Renderer/Public/HairStrandsInterface.h
--- a/Engine/Source/Runtime/Renderer/Public/HairStrandsInterface.h
+++ b/Engine/Source/Runtime/Renderer/Public/HairStrandsInterface.h
@@ -423,55 +423,6 @@ public:
 	FLinearColor DebugGroupColor;
 };
 
-////////////////////////////////////////////////////////////////////////////////////////////////////
-// Cluster information exchanged between renderer and the hair strand plugin. 
-
-struct FHairStrandClusterData
-{
-	struct FHairGroup
-	{
-		uint32 ClusterCount = 0;
-		uint32 VertexCount = 0;
-		uint32 MaxPointPerCurve = 0;
-
-		float LODIndex = -1;
-		float LODBias = 0.0f;
-		bool bVisible = false;
-
-		FRDGExternalBuffer* CurveBuffer = nullptr;
-		FRDGExternalBuffer* PointLODBuffer = nullptr;
-
-		// See FHairStrandsClusterCullingResource fro details about those buffers.
-		FRDGExternalBuffer* GroupAABBBuffer = nullptr;
-		FRDGExternalBuffer* ClusterAABBBuffer = nullptr;
-		FRDGExternalBuffer* ClusterInfoBuffer = nullptr; // SRV
-		FRDGExternalBuffer* ClusterLODInfoBuffer = nullptr; // SRV
-		FRDGExternalBuffer* CurveToClusterIdBuffer = nullptr; // SRV
-		FRDGExternalBuffer* ClusterVertexIdBuffer = nullptr;
-
-		TRefCountPtr<FRDGPooledBuffer> ClusterIdBuffer;
-		TRefCountPtr<FRDGPooledBuffer> ClusterIndexOffsetBuffer;
-		TRefCountPtr<FRDGPooledBuffer> ClusterIndexCountBuffer;
-
-		// Culling & LOD output
-		FRDGExternalBuffer* GetCulledCurveBuffer() const				{ return HairGroupPublicPtr ? &HairGroupPublicPtr->GetCulledCurveBuffer() : nullptr; }
-		FRDGExternalBuffer* GetCulledVertexIdBuffer() const				{ return HairGroupPublicPtr ? &HairGroupPublicPtr->GetCulledVertexIdBuffer() : nullptr; }
-		FRDGExternalBuffer* GetCulledVertexRadiusScaleBuffer() const	{ return HairGroupPublicPtr ? &HairGroupPublicPtr->GetCulledVertexRadiusScaleBuffer() : nullptr; }
-		bool GetCullingResultAvailable() const							{ return HairGroupPublicPtr ? HairGroupPublicPtr->GetCullingResultAvailable() : false; }
-		void SetCullingResultAvailable(bool b)							{ if (HairGroupPublicPtr) HairGroupPublicPtr->SetCullingResultAvailable(b); }
-
-		TRefCountPtr<FRDGPooledBuffer> ClusterDebugInfoBuffer;	// Null if this debug is not enabled.
-		FRDGBufferRef CulledClusterCountBuffer = nullptr;
-		FRDGBufferRef CulledCluster1DIndirectArgsBuffer = nullptr;
-		FRDGBufferRef CulledCluster2DIndirectArgsBuffer = nullptr;
-		uint32 GroupSize1D = 0;
-
-		FHairGroupPublicData* HairGroupPublicPtr = nullptr;
-	};
-
-	TArray<FHairGroup> HairGroups;
-};
-
 ////////////////////////////////////////////////////////////////////////////////////////////////////
 
 // API for enabling/disabling the various geometry representation
