Rename HairPrimitiveId to ControlPointId to avoid confusion with Scene PrimitiveId.

diff --git a/Engine/Shaders/Private/HairStrands/HairStrandsAttributeCommon.ush b/Engine/Shaders/Private/HairStrands/HairStrandsAttributeCommon.ush
--- a/Engine/Shaders/Private/HairStrands/HairStrandsAttributeCommon.ush
+++ b/Engine/Shaders/Private/HairStrands/HairStrandsAttributeCommon.ush
@@ -64,21 +64,21 @@ uint2 InternalGetHairStrandsAttributes64(uint Index, uint AttributeOffsetInBytes
 	return OutPacked;
 }
 
-uint InternalGetHairStrandsAttributesPerCurve(uint HairPrimitiveId, uint AttributeOffsetInBytes, uint AttributeStrideInBytes, ByteAddressBuffer InAttributeBuffer, Buffer<uint> InPointToCurveBuffer, uint InIndexToChunkDivAsShift, uint InChunkElementCount, uint InChunkStrideInBytes)
+uint InternalGetHairStrandsAttributesPerCurve(uint HairControlPointId, uint AttributeOffsetInBytes, uint AttributeStrideInBytes, ByteAddressBuffer InAttributeBuffer, Buffer<uint> InPointToCurveBuffer, uint InIndexToChunkDivAsShift, uint InChunkElementCount, uint InChunkStrideInBytes)
 {
-	const uint CurveIndex = InPointToCurveBuffer[HairPrimitiveId];
+	const uint CurveIndex = InPointToCurveBuffer[HairControlPointId];
 	return InternalGetHairStrandsAttributes(CurveIndex, AttributeOffsetInBytes, AttributeStrideInBytes, InAttributeBuffer, InIndexToChunkDivAsShift, InChunkElementCount, InChunkStrideInBytes);
 }
 
-uint2 InternalGetHairStrandsAttributesPerCurve64(uint HairPrimitiveId, uint AttributeOffsetInBytes, uint AttributeStrideInBytes, ByteAddressBuffer InAttributeBuffer, Buffer<uint> InPointToCurveBuffer, uint InIndexToChunkDivAsShift, uint InChunkElementCount, uint InChunkStrideInBytes)
+uint2 InternalGetHairStrandsAttributesPerCurve64(uint HairControlPointId, uint AttributeOffsetInBytes, uint AttributeStrideInBytes, ByteAddressBuffer InAttributeBuffer, Buffer<uint> InPointToCurveBuffer, uint InIndexToChunkDivAsShift, uint InChunkElementCount, uint InChunkStrideInBytes)
 {
-	const uint CurveIndex = InPointToCurveBuffer[HairPrimitiveId];
+	const uint CurveIndex = InPointToCurveBuffer[HairControlPointId];
 	return InternalGetHairStrandsAttributes64(CurveIndex, AttributeOffsetInBytes, AttributeStrideInBytes, InAttributeBuffer, InIndexToChunkDivAsShift, InChunkElementCount, InChunkStrideInBytes);
 }
 
-uint InternalGetHairStrandsAttributesPerVertex(uint HairPrimitiveId, uint AttributeOffsetInBytes, uint AttributeStrideInBytes, ByteAddressBuffer InAttributeBuffer, uint InIndexToChunkDivAsShift, uint InChunkElementCount, uint InChunkStrideInBytes)
+uint InternalGetHairStrandsAttributesPerVertex(uint HairControlPointId, uint AttributeOffsetInBytes, uint AttributeStrideInBytes, ByteAddressBuffer InAttributeBuffer, uint InIndexToChunkDivAsShift, uint InChunkElementCount, uint InChunkStrideInBytes)
 {
-	return InternalGetHairStrandsAttributes(HairPrimitiveId, AttributeOffsetInBytes, AttributeStrideInBytes, InAttributeBuffer, InIndexToChunkDivAsShift, InChunkElementCount, InChunkStrideInBytes);
+	return InternalGetHairStrandsAttributes(HairControlPointId, AttributeOffsetInBytes, AttributeStrideInBytes, InAttributeBuffer, InIndexToChunkDivAsShift, InChunkElementCount, InChunkStrideInBytes);
 }
 
 float InternalGetHairStrandsCoordU(Buffer<uint4> InPositionBuffer, uint VertexId, float InRadius, float InRootScale, float InTipScale)
diff --git a/Engine/Shaders/Private/HairStrands/HairStrandsAttributeTemplate.ush b/Engine/Shaders/Private/HairStrands/HairStrandsAttributeTemplate.ush
--- a/Engine/Shaders/Private/HairStrands/HairStrandsAttributeTemplate.ush
+++ b/Engine/Shaders/Private/HairStrands/HairStrandsAttributeTemplate.ush
@@ -10,70 +10,70 @@
 // Redefine macro for easing reability
 #define PREFIX(Name) HAIR_STRANDS_ATTRIBUTE_ACCESSORS(Name)
 
-float2 GetHairStrandsUV(uint InPrimitiveId, float2 InSegmentUV)
+float2 GetHairStrandsUV(uint InControlPointId, float2 InSegmentUV)
 {
-	const float CoordU0 = InternalGetHairStrandsCoordU(PREFIX(PositionBuffer), InPrimitiveId,   PREFIX(Radius), PREFIX(RootScale), PREFIX(TipScale));
-	const float CoordU1 = InternalGetHairStrandsCoordU(PREFIX(PositionBuffer), InPrimitiveId+1, PREFIX(Radius), PREFIX(RootScale), PREFIX(TipScale));
+	const float CoordU0 = InternalGetHairStrandsCoordU(PREFIX(PositionBuffer), InControlPointId,   PREFIX(Radius), PREFIX(RootScale), PREFIX(TipScale));
+	const float CoordU1 = InternalGetHairStrandsCoordU(PREFIX(PositionBuffer), InControlPointId+1, PREFIX(Radius), PREFIX(RootScale), PREFIX(TipScale));
 
 	float2 Out;
 	Out.x = lerp(CoordU0, CoordU1, InSegmentUV.x);
 	Out.y = InSegmentUV.y;
 	return Out;
 }
-float2 GetHairStrandsDimensions(uint InPrimitiveId, float InSegmentU)
+float2 GetHairStrandsDimensions(uint InControlPointId, float InSegmentU)
 {
-	const float CurveLength = UnpackHairLength(InternalGetHairStrandsAttributesPerCurve(InPrimitiveId, HAIR_CURVE_ATTRIBUTE_OFFSET_LENGTH(PREFIX(CurveAttributeOffsets)), HAIR_CURVE_ATTRIBUTE_STRIDE_LENGTH, PREFIX(CurveAttributeBuffer), PREFIX(PointToCurveBuffer), PREFIX(CurveAttributeIndexToChunkDivAsShift), PREFIX(CurveAttributeChunkElementCount), PREFIX(CurveAttributeChunkStrideInBytes)));
+	const float CurveLength = UnpackHairLength(InternalGetHairStrandsAttributesPerCurve(InControlPointId, HAIR_CURVE_ATTRIBUTE_OFFSET_LENGTH(PREFIX(CurveAttributeOffsets)), HAIR_CURVE_ATTRIBUTE_STRIDE_LENGTH, PREFIX(CurveAttributeBuffer), PREFIX(PointToCurveBuffer), PREFIX(CurveAttributeIndexToChunkDivAsShift), PREFIX(CurveAttributeChunkElementCount), PREFIX(CurveAttributeChunkStrideInBytes)));
 
-	const float CoordU0 = InternalGetHairStrandsCoordU(PREFIX(PositionBuffer), InPrimitiveId,   PREFIX(Radius), PREFIX(RootScale), PREFIX(TipScale));
-	const float CoordU1 = InternalGetHairStrandsCoordU(PREFIX(PositionBuffer), InPrimitiveId+1, PREFIX(Radius), PREFIX(RootScale), PREFIX(TipScale));
+	const float CoordU0 = InternalGetHairStrandsCoordU(PREFIX(PositionBuffer), InControlPointId,   PREFIX(Radius), PREFIX(RootScale), PREFIX(TipScale));
+	const float CoordU1 = InternalGetHairStrandsCoordU(PREFIX(PositionBuffer), InControlPointId+1, PREFIX(Radius), PREFIX(RootScale), PREFIX(TipScale));
 	const float CoordU  = lerp(CoordU0, CoordU1, InSegmentU);
 
-	const float Radius0 = InternalGetHairStrandsWorldRadius(PREFIX(PositionBuffer), InPrimitiveId,   PREFIX(Radius), PREFIX(RootScale), PREFIX(TipScale));
-	const float Radius1 = InternalGetHairStrandsWorldRadius(PREFIX(PositionBuffer), InPrimitiveId+1, PREFIX(Radius), PREFIX(RootScale), PREFIX(TipScale));
+	const float Radius0 = InternalGetHairStrandsWorldRadius(PREFIX(PositionBuffer), InControlPointId,   PREFIX(Radius), PREFIX(RootScale), PREFIX(TipScale));
+	const float Radius1 = InternalGetHairStrandsWorldRadius(PREFIX(PositionBuffer), InControlPointId+1, PREFIX(Radius), PREFIX(RootScale), PREFIX(TipScale));
 	const float Radius  = lerp(Radius0, Radius1, InSegmentU);
 
 	return float2(CurveLength * CoordU, Radius);
 }
-float2 GetHairStrandsRootUV(uint InPrimitiveId)
+float2 GetHairStrandsRootUV(uint InControlPointId)
 {
 	const uint RootUVOffset = HAIR_CURVE_ATTRIBUTE_OFFSET_ROOTUV(PREFIX(CurveAttributeOffsets));
 	if (IsAttributeValid(RootUVOffset))
 	{
-		return UnpackHairRootUV(InternalGetHairStrandsAttributesPerCurve(InPrimitiveId, RootUVOffset, HAIR_CURVE_ATTRIBUTE_STRIDE_ROOTUV, PREFIX(CurveAttributeBuffer), PREFIX(PointToCurveBuffer), PREFIX(CurveAttributeIndexToChunkDivAsShift), PREFIX(CurveAttributeChunkElementCount), PREFIX(CurveAttributeChunkStrideInBytes)));
+		return UnpackHairRootUV(InternalGetHairStrandsAttributesPerCurve(InControlPointId, RootUVOffset, HAIR_CURVE_ATTRIBUTE_STRIDE_ROOTUV, PREFIX(CurveAttributeBuffer), PREFIX(PointToCurveBuffer), PREFIX(CurveAttributeIndexToChunkDivAsShift), PREFIX(CurveAttributeChunkElementCount), PREFIX(CurveAttributeChunkStrideInBytes)));
 	}
 	else
 	{
 		return 0;
 	}
 }
-float GetHairStrandsSeed(uint InPrimitiveId)
+float GetHairStrandsSeed(uint InControlPointId)
 {
-	return UnpackHairSeed(InternalGetHairStrandsAttributesPerCurve(InPrimitiveId, HAIR_CURVE_ATTRIBUTE_OFFSET_SEED(PREFIX(CurveAttributeOffsets)), HAIR_CURVE_ATTRIBUTE_STRIDE_SEED, PREFIX(CurveAttributeBuffer), PREFIX(PointToCurveBuffer), PREFIX(CurveAttributeIndexToChunkDivAsShift), PREFIX(CurveAttributeChunkElementCount), PREFIX(CurveAttributeChunkStrideInBytes)));
+	return UnpackHairSeed(InternalGetHairStrandsAttributesPerCurve(InControlPointId, HAIR_CURVE_ATTRIBUTE_OFFSET_SEED(PREFIX(CurveAttributeOffsets)), HAIR_CURVE_ATTRIBUTE_STRIDE_SEED, PREFIX(CurveAttributeBuffer), PREFIX(PointToCurveBuffer), PREFIX(CurveAttributeIndexToChunkDivAsShift), PREFIX(CurveAttributeChunkElementCount), PREFIX(CurveAttributeChunkStrideInBytes)));
 }
-uint3 GetHairStrandsClumpID(uint InPrimitiveId)
+uint3 GetHairStrandsClumpID(uint InControlPointId)
 {
 	const uint ClumpIdOffset  = HAIR_CURVE_ATTRIBUTE_OFFSET_CLUMPID(PREFIX(CurveAttributeOffsets));
 	const uint ClumpId3Offset = HAIR_CURVE_ATTRIBUTE_OFFSET_CLUMPID3(PREFIX(CurveAttributeOffsets));
 	if (IsAttributeValid(ClumpIdOffset))
 	{
-		return UnpackHairClumpID(InternalGetHairStrandsAttributesPerCurve(InPrimitiveId, ClumpIdOffset, HAIR_CURVE_ATTRIBUTE_STRIDE_CLUMPID, PREFIX(CurveAttributeBuffer), PREFIX(PointToCurveBuffer), PREFIX(CurveAttributeIndexToChunkDivAsShift), PREFIX(CurveAttributeChunkElementCount), PREFIX(CurveAttributeChunkStrideInBytes)));
+		return UnpackHairClumpID(InternalGetHairStrandsAttributesPerCurve(InControlPointId, ClumpIdOffset, HAIR_CURVE_ATTRIBUTE_STRIDE_CLUMPID, PREFIX(CurveAttributeBuffer), PREFIX(PointToCurveBuffer), PREFIX(CurveAttributeIndexToChunkDivAsShift), PREFIX(CurveAttributeChunkElementCount), PREFIX(CurveAttributeChunkStrideInBytes)));
 	}
 	else if (IsAttributeValid(ClumpId3Offset))
 	{
-		return UnpackHairClumpID(InternalGetHairStrandsAttributesPerCurve64(InPrimitiveId, ClumpId3Offset, HAIR_CURVE_ATTRIBUTE_STRIDE_CLUMPID3, PREFIX(CurveAttributeBuffer), PREFIX(PointToCurveBuffer), PREFIX(CurveAttributeIndexToChunkDivAsShift), PREFIX(CurveAttributeChunkElementCount), PREFIX(CurveAttributeChunkStrideInBytes)));
+		return UnpackHairClumpID(InternalGetHairStrandsAttributesPerCurve64(InControlPointId, ClumpId3Offset, HAIR_CURVE_ATTRIBUTE_STRIDE_CLUMPID3, PREFIX(CurveAttributeBuffer), PREFIX(PointToCurveBuffer), PREFIX(CurveAttributeIndexToChunkDivAsShift), PREFIX(CurveAttributeChunkElementCount), PREFIX(CurveAttributeChunkStrideInBytes)));
 	}
 	else
 	{
 		return 0;
 	}
 }
-float3 GetHairStrandsColor(uint InPrimitiveId, float InSegmentU)
+float3 GetHairStrandsColor(uint InControlPointId, float InSegmentU)
 {
 	const uint ColorOffset = HAIR_POINT_ATTRIBUTE_OFFSET_COLOR(PREFIX(PointAttributeOffsets));
 	if (IsAttributeValid(ColorOffset))
 	{
-		const float3 Color0 = UnpackHairColor(InternalGetHairStrandsAttributesPerVertex(InPrimitiveId,   ColorOffset, HAIR_POINT_ATTRIBUTE_STRIDE_COLOR, PREFIX(PointAttributeBuffer), PREFIX(PointAttributeIndexToChunkDivAsShift), PREFIX(PointAttributeChunkElementCount), PREFIX(PointAttributeChunkStrideInBytes)));
-		const float3 Color1 = UnpackHairColor(InternalGetHairStrandsAttributesPerVertex(InPrimitiveId+1, ColorOffset, HAIR_POINT_ATTRIBUTE_STRIDE_COLOR, PREFIX(PointAttributeBuffer), PREFIX(PointAttributeIndexToChunkDivAsShift), PREFIX(PointAttributeChunkElementCount), PREFIX(PointAttributeChunkStrideInBytes)));
+		const float3 Color0 = UnpackHairColor(InternalGetHairStrandsAttributesPerVertex(InControlPointId,   ColorOffset, HAIR_POINT_ATTRIBUTE_STRIDE_COLOR, PREFIX(PointAttributeBuffer), PREFIX(PointAttributeIndexToChunkDivAsShift), PREFIX(PointAttributeChunkElementCount), PREFIX(PointAttributeChunkStrideInBytes)));
+		const float3 Color1 = UnpackHairColor(InternalGetHairStrandsAttributesPerVertex(InControlPointId+1, ColorOffset, HAIR_POINT_ATTRIBUTE_STRIDE_COLOR, PREFIX(PointAttributeBuffer), PREFIX(PointAttributeIndexToChunkDivAsShift), PREFIX(PointAttributeChunkElementCount), PREFIX(PointAttributeChunkStrideInBytes)));
 		return lerp(Color0, Color1, InSegmentU);
 	}
 	else
@@ -81,13 +81,13 @@ float3 GetHairStrandsColor(uint InPrimitiveId, float InSegmentU)
 		return 0;
 	}
 }
-float GetHairStrandsRoughness(uint InPrimitiveId, float InSegmentU)
+float GetHairStrandsRoughness(uint InControlPointId, float InSegmentU)
 {
 	const uint RoughnessOffset = HAIR_POINT_ATTRIBUTE_OFFSET_ROUGHNESS(PREFIX(PointAttributeOffsets));
 	if (IsAttributeValid(RoughnessOffset))
 	{
-		const float Roughness0 = UnpackHairRoughness(InternalGetHairStrandsAttributesPerVertex(InPrimitiveId,   RoughnessOffset, HAIR_POINT_ATTRIBUTE_STRIDE_ROUGHNESS, PREFIX(PointAttributeBuffer), PREFIX(PointAttributeIndexToChunkDivAsShift), PREFIX(PointAttributeChunkElementCount), PREFIX(PointAttributeChunkStrideInBytes)));
-		const float Roughness1 = UnpackHairRoughness(InternalGetHairStrandsAttributesPerVertex(InPrimitiveId+1, RoughnessOffset, HAIR_POINT_ATTRIBUTE_STRIDE_ROUGHNESS, PREFIX(PointAttributeBuffer), PREFIX(PointAttributeIndexToChunkDivAsShift), PREFIX(PointAttributeChunkElementCount), PREFIX(PointAttributeChunkStrideInBytes)));
+		const float Roughness0 = UnpackHairRoughness(InternalGetHairStrandsAttributesPerVertex(InControlPointId,   RoughnessOffset, HAIR_POINT_ATTRIBUTE_STRIDE_ROUGHNESS, PREFIX(PointAttributeBuffer), PREFIX(PointAttributeIndexToChunkDivAsShift), PREFIX(PointAttributeChunkElementCount), PREFIX(PointAttributeChunkStrideInBytes)));
+		const float Roughness1 = UnpackHairRoughness(InternalGetHairStrandsAttributesPerVertex(InControlPointId+1, RoughnessOffset, HAIR_POINT_ATTRIBUTE_STRIDE_ROUGHNESS, PREFIX(PointAttributeBuffer), PREFIX(PointAttributeIndexToChunkDivAsShift), PREFIX(PointAttributeChunkElementCount), PREFIX(PointAttributeChunkStrideInBytes)));
 		return lerp(Roughness0, Roughness1, InSegmentU);
 	}
 	else
@@ -95,13 +95,13 @@ float GetHairStrandsRoughness(uint InPrimitiveId, float InSegmentU)
 		return 0;
 	}
 }
-float GetHairStrandsAO(uint InPrimitiveId, float InSegmentU)
+float GetHairStrandsAO(uint InControlPointId, float InSegmentU)
 {
 	const uint AoOffset = HAIR_POINT_ATTRIBUTE_OFFSET_AO(PREFIX(PointAttributeOffsets));
 	if (IsAttributeValid(AoOffset))
 	{
-		const float AO0 = UnpackHairAO(InternalGetHairStrandsAttributesPerVertex(InPrimitiveId, AoOffset, HAIR_POINT_ATTRIBUTE_STRIDE_AO, PREFIX(PointAttributeBuffer), PREFIX(PointAttributeIndexToChunkDivAsShift), PREFIX(PointAttributeChunkElementCount), PREFIX(PointAttributeChunkStrideInBytes)));
-		const float AO1 = UnpackHairAO(InternalGetHairStrandsAttributesPerVertex(InPrimitiveId+1, AoOffset, HAIR_POINT_ATTRIBUTE_STRIDE_AO, PREFIX(PointAttributeBuffer), PREFIX(PointAttributeIndexToChunkDivAsShift), PREFIX(PointAttributeChunkElementCount), PREFIX(PointAttributeChunkStrideInBytes)));
+		const float AO0 = UnpackHairAO(InternalGetHairStrandsAttributesPerVertex(InControlPointId, AoOffset, HAIR_POINT_ATTRIBUTE_STRIDE_AO, PREFIX(PointAttributeBuffer), PREFIX(PointAttributeIndexToChunkDivAsShift), PREFIX(PointAttributeChunkElementCount), PREFIX(PointAttributeChunkStrideInBytes)));
+		const float AO1 = UnpackHairAO(InternalGetHairStrandsAttributesPerVertex(InControlPointId+1, AoOffset, HAIR_POINT_ATTRIBUTE_STRIDE_AO, PREFIX(PointAttributeBuffer), PREFIX(PointAttributeIndexToChunkDivAsShift), PREFIX(PointAttributeChunkElementCount), PREFIX(PointAttributeChunkStrideInBytes)));
 		return lerp(AO0, AO1, InSegmentU);
 	}
 	else
diff --git a/Engine/Shaders/Private/HairStrands/HairStrandsDebug.usf b/Engine/Shaders/Private/HairStrands/HairStrandsDebug.usf
--- a/Engine/Shaders/Private/HairStrands/HairStrandsDebug.usf
+++ b/Engine/Shaders/Private/HairStrands/HairStrandsDebug.usf
@@ -1131,7 +1131,7 @@ void MainCS(uint InGroupIndex : SV_GroupIndex, uint3 DispatchThreadId : SV_Dispa
 		Newline(Context2);
 
 		Print(Context2, Sample.MacroGroupId, FontYellow);						Newline(Context2);
-		Print(Context2, Sample.PrimitiveId, FontEmerald);						Newline(Context2);
+		Print(Context2, Sample.ControlPointId, FontEmerald);						Newline(Context2);
 	#endif
 		Print(Context2, Sample.Coverage8bit, FontOrange);						Newline(Context2);
 	}
@@ -2497,7 +2497,7 @@ void CSMain(
 		if (VisNode.MaterialId == MaterialId)
 		{
 			FHairControlPoint Out = UnpackHairControlPoint(
-				HairInstance_PositionBuffer[VisNode.PrimitiveId],
+				HairInstance_PositionBuffer[VisNode.ControlPointId],
 				HairInstance_PositionOffsetBuffer[0].xyz,
 				HairInstance_Radius,
 				HairInstance_RootScale,
@@ -2508,17 +2508,17 @@ void CSMain(
 
 			FShaderPrintContext Ctx = InitShaderPrintContext(true, uint2(PixelCoord) + uint2(20, 0));
 			Print(Ctx, TEXT("Group    : "), FontName); Print(Ctx, GetHairStrandsGroupIndex(), FontValue); Newline(Ctx);
-			Print(Ctx, TEXT("Radius   : "), FontName); Print(Ctx, GetHairStrandsDimensions(VisNode.PrimitiveId, 0).y, FontValue); Newline(Ctx);
-			Print(Ctx, TEXT("Seed     : "), FontName); Print(Ctx, GetHairStrandsSeed(VisNode.PrimitiveId), FontValue); Newline(Ctx);
-			Print(Ctx, TEXT("CoordU   : "), FontName); Print(Ctx, GetHairStrandsUV(VisNode.PrimitiveId, float2(0, 0.5f)).x, FontValue); Newline(Ctx);
+			Print(Ctx, TEXT("Radius   : "), FontName); Print(Ctx, GetHairStrandsDimensions(VisNode.ControlPointId, 0).y, FontValue); Newline(Ctx);
+			Print(Ctx, TEXT("Seed     : "), FontName); Print(Ctx, GetHairStrandsSeed(VisNode.ControlPointId), FontValue); Newline(Ctx);
+			Print(Ctx, TEXT("CoordU   : "), FontName); Print(Ctx, GetHairStrandsUV(VisNode.ControlPointId, float2(0, 0.5f)).x, FontValue); Newline(Ctx);
 			Newline(Ctx);
 
-			if (HasHairStrandsRootUV()) 	{ Print(Ctx, TEXT("RootUV   : "), FontName); Print(Ctx, GetHairStrandsRootUV(VisNode.PrimitiveId),  FontValue); Newline(Ctx); }
-			if (HasHairStrandsClumpId()) 	{ Print(Ctx, TEXT("Clump ID : "), FontName); Print(Ctx, GetHairStrandsClumpID(VisNode.PrimitiveId).x, FontValue); Newline(Ctx); }
-			if (HasHairStrandsClumpId3()) 	{ Print(Ctx, TEXT("Clump IDs: "), FontName); Print(Ctx, GetHairStrandsClumpID(VisNode.PrimitiveId), FontValue); Newline(Ctx); }
-			if (HasHairStrandsColor()) 		{ Print(Ctx, TEXT("Color    : "), FontName); Print(Ctx, GetHairStrandsColor(VisNode.PrimitiveId, 0), FontValue); Newline(Ctx); }
-			if (HasHairStrandsRoughness()) 	{ Print(Ctx, TEXT("Roughness: "), FontName); Print(Ctx, GetHairStrandsRoughness(VisNode.PrimitiveId, 0), FontValue); Newline(Ctx); }
-			if (HasHairStrandsAO()) 		{ Print(Ctx, TEXT("AO       : "), FontName); Print(Ctx, GetHairStrandsAO(VisNode.PrimitiveId, 0), FontValue); Newline(Ctx); }
+			if (HasHairStrandsRootUV()) 	{ Print(Ctx, TEXT("RootUV   : "), FontName); Print(Ctx, GetHairStrandsRootUV(VisNode.ControlPointId),  FontValue); Newline(Ctx); }
+			if (HasHairStrandsClumpId()) 	{ Print(Ctx, TEXT("Clump ID : "), FontName); Print(Ctx, GetHairStrandsClumpID(VisNode.ControlPointId).x, FontValue); Newline(Ctx); }
+			if (HasHairStrandsClumpId3()) 	{ Print(Ctx, TEXT("Clump IDs: "), FontName); Print(Ctx, GetHairStrandsClumpID(VisNode.ControlPointId), FontValue); Newline(Ctx); }
+			if (HasHairStrandsColor()) 		{ Print(Ctx, TEXT("Color    : "), FontName); Print(Ctx, GetHairStrandsColor(VisNode.ControlPointId, 0), FontValue); Newline(Ctx); }
+			if (HasHairStrandsRoughness()) 	{ Print(Ctx, TEXT("Roughness: "), FontName); Print(Ctx, GetHairStrandsRoughness(VisNode.ControlPointId, 0), FontValue); Newline(Ctx); }
+			if (HasHairStrandsAO()) 		{ Print(Ctx, TEXT("AO       : "), FontName); Print(Ctx, GetHairStrandsAO(VisNode.ControlPointId, 0), FontValue); Newline(Ctx); }
 			Newline(Ctx);
 		}
 	}
diff --git a/Engine/Shaders/Private/HairStrands/HairStrandsEnvironmentLighting.usf b/Engine/Shaders/Private/HairStrands/HairStrandsEnvironmentLighting.usf
--- a/Engine/Shaders/Private/HairStrands/HairStrandsEnvironmentLighting.usf
+++ b/Engine/Shaders/Private/HairStrands/HairStrandsEnvironmentLighting.usf
@@ -287,7 +287,7 @@ float MaxComponent3(float3 In)
 	return max(In.x, max(In.y, In.z));
 }
 
-float3 ReflectionEnvironment(FGBufferData GBuffer, float2 ScreenPosition, float2 SvPosition, uint PrimitiveId, uint MacroGroupId)
+float3 ReflectionEnvironment(FGBufferData GBuffer, float2 ScreenPosition, float2 SvPosition, uint ControlPointId, uint MacroGroupId)
 {
 	const float3 TranslatedWorldPosition = mul(float4(ScreenPosition * GBuffer.Depth, GBuffer.Depth, 1), PrimaryView.ScreenToTranslatedWorld).xyz;
 	const float3 CameraToPixel = normalize(TranslatedWorldPosition - PrimaryView.TranslatedWorldCameraOrigin);
@@ -295,8 +295,8 @@ float3 ReflectionEnvironment(FGBufferData GBuffer, float2 ScreenPosition, float2
 	const float  PixelSceneDepth = GBuffer.Depth;
 	const float  PixelRadius = VirtualVoxel.HairCoveragePixelRadiusAtDepth1 * PixelSceneDepth;
 
-	const uint PrimitiveIdMod8 = (PrimitiveId % 8);
-	const float3 VoxelOffset =  GetHairVoxelJitter(SvPosition, View.StateFrameIndexMod8 + PrimitiveIdMod8, VirtualVoxel.JitterMode);
+	const uint ControlPointIdMod8 = (ControlPointId % 8);
+	const float3 VoxelOffset =  GetHairVoxelJitter(SvPosition, View.StateFrameIndexMod8 + ControlPointIdMod8, VirtualVoxel.JitterMode);
 
 	const float3 V = -CameraToPixel;
 	const float3 T = GBuffer.WorldNormal;
@@ -602,7 +602,7 @@ float3 ReflectionEnvironmentUniformSampling(
 	FGBufferData GBuffer, 
 	float2 ScreenPosition, 
 	float2 SvPosition, 
-	uint PrimitiveId, 
+	uint HairControlPointId, 
 	uint MacroGroupId,
 	bool bDebugEnabled)
 {
@@ -612,8 +612,6 @@ float3 ReflectionEnvironmentUniformSampling(
 	const float  PixelSceneDepth = GBuffer.Depth;
 	const float  PixelRadius = VirtualVoxel.HairCoveragePixelRadiusAtDepth1 * PixelSceneDepth;
 
-	const uint PrimitiveIdMod8 = (PrimitiveId % 8);
-
 	const float3 V = -CameraToPixel;
 	const float3 T = GBuffer.WorldNormal;
 
@@ -705,7 +703,7 @@ float3 ReflectionEnvironmentSH(
 	FGBufferData GBuffer,
 	float2 ScreenPosition,
 	float2 SvPosition,
-	uint PrimitiveId,
+	uint HairControlPointId,
 	uint MacroGroupId,
 	bool bDebugEnabled)
 {
@@ -773,7 +771,7 @@ float3 ReflectionEnvironmentSH(
 #if PERMUTATION_INTEGRATION_TYPE == INTEGRATION_SCENECOLOR
 Texture2D<float4> SceneColorTexture;
 
-float3 ComputeThinTransmission(FGBufferData GBuffer, float2 ScreenPosition, float2 SvPosition, uint PrimitiveId)
+float3 ComputeThinTransmission(FGBufferData GBuffer, float2 ScreenPosition, float2 SvPosition, uint HairControlPointId)
 {
 	const float3 TranslatedWorldPosition = mul(float4(ScreenPosition * GBuffer.Depth, GBuffer.Depth, 1), PrimaryView.ScreenToTranslatedWorld).xyz;
 	const float3 CameraToPixel = normalize(TranslatedWorldPosition - PrimaryView.TranslatedWorldCameraOrigin);
@@ -867,13 +865,13 @@ void MainPS(
 	{
 		const float3 SkyLighting =
 		#if PERMUTATION_INTEGRATION_TYPE == INTEGRATION_SCENECOLOR
-			ComputeThinTransmission(GBuffer, ScreenPosition, SvPosition, Sample.PrimitiveId);
+			ComputeThinTransmission(GBuffer, ScreenPosition, SvPosition, Sample.ControlPointId);
 		#elif PERMUTATION_INTEGRATION_TYPE == INTEGRATION_ADHOC
-			ReflectionEnvironment(GBuffer, ScreenPosition, SvPosition, Sample.PrimitiveId, Sample.MacroGroupId);
+			ReflectionEnvironment(GBuffer, ScreenPosition, SvPosition, Sample.ControlPointId, Sample.MacroGroupId);
 		#elif PERMUTATION_INTEGRATION_TYPE == INTEGRATION_UNIFORM
-			ReflectionEnvironmentUniformSampling(GBuffer, ScreenPosition, SvPosition, Sample.PrimitiveId, Sample.MacroGroupId, bDebugEnabled);
+			ReflectionEnvironmentUniformSampling(GBuffer, ScreenPosition, SvPosition, Sample.ControlPointId, Sample.MacroGroupId, bDebugEnabled);
 		#elif PERMUTATION_INTEGRATION_TYPE == INTEGRATION_SH
-			ReflectionEnvironmentSH(GBuffer, ScreenPosition, SvPosition, Sample.PrimitiveId, Sample.MacroGroupId, bDebugEnabled);
+			ReflectionEnvironmentSH(GBuffer, ScreenPosition, SvPosition, Sample.ControlPointId, Sample.MacroGroupId, bDebugEnabled);
 		#endif
 
 		const float SampleCoverage = saturate(From8bitCoverage(Sample.Coverage8bit));
diff --git a/Engine/Shaders/Private/HairStrands/HairStrandsForwardCulling.usf b/Engine/Shaders/Private/HairStrands/HairStrandsForwardCulling.usf
--- a/Engine/Shaders/Private/HairStrands/HairStrandsForwardCulling.usf
+++ b/Engine/Shaders/Private/HairStrands/HairStrandsForwardCulling.usf
@@ -211,7 +211,7 @@ void CSMain(uint DispatchThreadID : SV_DispatchThreadID, uint GroupThreadID : SV
 					Out.Depth			= DeviceZ;
 					Out.Coverage		= 1;
 					Out.Coverage8bit	= 255;
-					Out.PrimitiveId		= PrimID;
+					Out.ControlPointId		= PrimID;
 					Out.MaterialId		= HairMaterialId;
 
 					uint PackedW = f32tof16(Point0.WorldRadius);
@@ -248,7 +248,7 @@ void CSMain(uint DispatchThreadID : SV_DispatchThreadID, uint GroupThreadID : SV
 	
 	GroupMemoryBarrierWithGroupSync();
 
-	uint   PrimID = DispatchThreadID;
+	uint   ControlPointId = DispatchThreadID;
 	bool   bIsPrimVisibile = false;
 	uint   PrimWriteOffset = 0;
 	float  DeviceZ = 0;
@@ -257,14 +257,14 @@ void CSMain(uint DispatchThreadID : SV_DispatchThreadID, uint GroupThreadID : SV
 	FHairControlPoint Point0 = (FHairControlPoint)0;
 
 	// Count valid segments
-	bool bIsPrimValid = PrimID < ControlPointCount;
+	bool bIsPrimValid = ControlPointId < ControlPointCount;
 	if (bIsPrimValid)
 	{
 		const float3 PositionOffset = HairStrandsVF_PositionOffsetBuffer[0].xyz;
 
 		// Fetch both control points
 		FHairControlPoint CP0 = UnpackHairControlPoint(
-			HairStrandsVF_PositionBuffer[PrimID],
+			HairStrandsVF_PositionBuffer[ControlPointId],
 			PositionOffset,
 			HairStrandsVF_Radius,
 			HairStrandsVF_RootScale,
@@ -274,7 +274,7 @@ void CSMain(uint DispatchThreadID : SV_DispatchThreadID, uint GroupThreadID : SV
 		if (CP0.Type != HAIR_CONTROLPOINT_END)
 		{
 			CP1 = UnpackHairControlPoint(
-				HairStrandsVF_PositionBuffer[PrimID + 1],
+				HairStrandsVF_PositionBuffer[ControlPointId + 1],
 				PositionOffset,
 				HairStrandsVF_Radius,
 				HairStrandsVF_RootScale,
@@ -319,7 +319,7 @@ void CSMain(uint DispatchThreadID : SV_DispatchThreadID, uint GroupThreadID : SV
 		Out.Depth			= DeviceZ;
 		Out.Coverage		= 1;
 		Out.Coverage8bit	= 255;
-		Out.PrimitiveId		= PrimID;
+		Out.ControlPointId	= ControlPointId;
 		Out.MaterialId		= HairMaterialId;
 
 		uint PackedW = f32tof16(Point0.WorldRadius);
diff --git a/Engine/Shaders/Private/HairStrands/HairStrandsMaterialCommon.ush b/Engine/Shaders/Private/HairStrands/HairStrandsMaterialCommon.ush
--- a/Engine/Shaders/Private/HairStrands/HairStrandsMaterialCommon.ush
+++ b/Engine/Shaders/Private/HairStrands/HairStrandsMaterialCommon.ush
@@ -56,7 +56,7 @@ struct FHairVertexData
 };
 
 // GPUCULL_TODO: this should be switched over to use Instance ID instead of Primitive ID.
-FVertexFactoryInterpolantsVSToPS GetInterpolants(uint PrimitiveId, uint HairPrimitiveId, inout FHairVertexData OutHairVertex)
+FVertexFactoryInterpolantsVSToPS GetInterpolants(uint PrimitiveId, uint HairControlPointId, inout FHairVertexData OutHairVertex)
 {
 	FMaterialVertexParameters Empty = MakeInitializedMaterialVertexParameters();
 
@@ -65,7 +65,7 @@ FVertexFactoryInterpolantsVSToPS GetInterpolants(uint PrimitiveId, uint HairPrim
 	VF_GPUSCENE_SET_INPUT_FOR_RT(Input, PrimitiveId, 0U);
 	// * USE_HAIR_TRIANGLE_STRIP=0: the vertex factory expect triangles strip index. A quad is made of two triangles with indices: 0,1.
 	// * USE_HAIR_TRIANGLE_STRIP=1: the vertex factory expect triangles list index. A quad is made of two triangles with indices: 0,1,2,3,4,5.
-	Input.VertexId = HairPrimitiveId * HAIR_POINT_TO_VERTEX; 
+	Input.VertexId = HairControlPointId * HAIR_POINT_TO_VERTEX; 
 
 	FVertexFactoryIntermediates	Intermediates = GetVertexFactoryIntermediates(Input);
 	FVertexFactoryInterpolantsVSToPS Interpolants = VertexFactoryGetInterpolantsVSToPS(Input, Intermediates, Empty);
@@ -83,7 +83,7 @@ FVertexFactoryInterpolantsVSToPS GetInterpolants(uint PrimitiveId, uint HairPrim
 FEvaluateOutput Evaluate(
 	in float2 SamplePixelCoord,
 	in float SampleDepth,
-	in uint SamplePrimitiveId,
+	in uint SampleControlPointId,
 	in uint SampleCoverage8bit,
 	in uint PrimitiveId,
 	in bool bUpdateSampleCoverage,
@@ -100,20 +100,20 @@ FEvaluateOutput Evaluate(
 	float WorldStrandRadius = 0;
 	FVertexFactoryInterpolantsVSToPS Interpolants;
 	{
-		const uint HairPrimitiveId0 = SamplePrimitiveId;
+		const uint HairControlPointId0 = SampleControlPointId;
 
 		FHairVertexData Vertex0;
-		FVertexFactoryInterpolantsVSToPS Interpolants0 = GetInterpolants(PrimitiveId, HairPrimitiveId0, Vertex0);
+		FVertexFactoryInterpolantsVSToPS Interpolants0 = GetInterpolants(PrimitiveId, HairControlPointId0, Vertex0);
 		Interpolants = Interpolants0;
 
 		float SegmentU = 0.0f;
 		float SegmentV = 0.5f;
 		if (bInterpolationEnabled)
 		{
-			const uint HairPrimitiveId1 = HairPrimitiveId0 + 1; //TODO: min(HairPrimitiveId0 + 1, MaxVertexCount-1);
+			const uint HairControlPointId1 = HairControlPointId0 + 1; //TODO: min(HairControlPointId0 + 1, MaxVertexCount-1);
 
 			FHairVertexData Vertex1;
-			FVertexFactoryInterpolantsVSToPS Interpolants1 = GetInterpolants(PrimitiveId, HairPrimitiveId1, Vertex1);
+			FVertexFactoryInterpolantsVSToPS Interpolants1 = GetInterpolants(PrimitiveId, HairControlPointId1, Vertex1);
 
 			// Compute U
 			// Use the projection of the current sample point onto the hair segment (which is in the middle of the hair strands)
@@ -156,7 +156,7 @@ FEvaluateOutput Evaluate(
 #if VF_USE_PRIMITIVE_SCENE_DATA
 		Interpolants.PrimitiveId	 = PrimitiveId;
 #endif // VF_USE_PRIMITIVE_SCENE_DATA
-		Interpolants.HairPrimitiveId = HairPrimitiveId0;
+		Interpolants.HairControlPointId = HairControlPointId0;
 		Interpolants.HairPrimitiveUV = float2(SegmentU, SegmentV);
 	}
 	
@@ -194,7 +194,7 @@ FEvaluateOutput Evaluate(
 	FHairSample OutSample		= (FHairSample)0;
 	OutSample.Depth				= SampleDepth;
 	OutSample.Coverage8bit		= SampleCoverage8bit;
-	OutSample.PrimitiveId		= SamplePrimitiveId;
+	OutSample.ControlPointId	= SampleControlPointId;
 	OutSample.MacroGroupId		= MaterialPass_MacroGroupId;
 	OutSample.LightChannelMask	= MaterialPass_LightChannelMask;
 	OutSample.bScatterSceneLighting = UseScatterSceneLighting();
diff --git a/Engine/Shaders/Private/HairStrands/HairStrandsMaterialPS.usf b/Engine/Shaders/Private/HairStrands/HairStrandsMaterialPS.usf
--- a/Engine/Shaders/Private/HairStrands/HairStrandsMaterialPS.usf
+++ b/Engine/Shaders/Private/HairStrands/HairStrandsMaterialPS.usf
@@ -25,7 +25,7 @@ void Main(
 
 		float SampleDepth = InNode.Depth;
 		uint SampleCoverage8bit = InNode.Coverage8bit;
-		uint SamplePrimitiveId = InNode.PrimitiveId;
+		uint SampleControlPointId = InNode.ControlPointId;
 		uint SampleMaterialId = InNode.MaterialId;
 
 		if (SampleMaterialId == MaterialPass_MaterialId)
@@ -33,7 +33,7 @@ void Main(
 			FEvaluateOutput Out = (FEvaluateOutput)0;
 			if (SampleDepth > 0)
 			{
-				Out = Evaluate(SamplePixelCoord, SampleDepth, SamplePrimitiveId, SampleCoverage8bit, InPrimitiveId, MaterialPassParameters.bUpdateSampleCoverage > 0, MaterialPassParameters.bInterpolationEnabled);
+				Out = Evaluate(SamplePixelCoord, SampleDepth, SampleControlPointId, SampleCoverage8bit, InPrimitiveId, MaterialPassParameters.bUpdateSampleCoverage > 0, MaterialPassParameters.bInterpolationEnabled);
 			}
 
 			// Note: Use the syntax _ over the syntax . as on certain platform the UAV buffers are 
diff --git a/Engine/Shaders/Private/HairStrands/HairStrandsRaytracingGeometry.usf b/Engine/Shaders/Private/HairStrands/HairStrandsRaytracingGeometry.usf
--- a/Engine/Shaders/Private/HairStrands/HairStrandsRaytracingGeometry.usf
+++ b/Engine/Shaders/Private/HairStrands/HairStrandsRaytracingGeometry.usf
@@ -127,7 +127,7 @@ void SetBodySegment(uint InVertexIndex)
 	{
 		// Note that this vertex ordering is carefully chosen to work with the Vertex Factory in RT mode
 		// In particular, we _must_ put one of the bottom point as the last vertex of the triangle
-		// so that the interpolation step will copy the right HairPrimitiveId onwards
+		// so that the interpolation step will copy the right HairControlPointId onwards
 		OutputIndexBuffer[OutIndex + 0] = TopPoint + PrevIdx;
 		OutputIndexBuffer[OutIndex + 1] = TopPoint + Idx;
 		OutputIndexBuffer[OutIndex + 2] = BotPoint + PrevIdx;
diff --git a/Engine/Shaders/Private/HairStrands/HairStrandsVertexFactory.ush b/Engine/Shaders/Private/HairStrands/HairStrandsVertexFactory.ush
--- a/Engine/Shaders/Private/HairStrands/HairStrandsVertexFactory.ush
+++ b/Engine/Shaders/Private/HairStrands/HairStrandsVertexFactory.ush
@@ -102,7 +102,7 @@ struct FVertexFactoryInput
 struct FVertexInfo
 {
 	uint VertexIndex;
-	uint HairPrimitiveId;
+	uint HairControlPointId;
 	uint IsTip;				// The vertex is on the quad side toward the tip of the strand.
 #if !RAYHITGROUPSHADER
 	uint IsLeft;
@@ -130,7 +130,7 @@ FVertexInfo GetVertexInfo(FVertexFactoryInput Input)
 	VertexInfo.IsTip		= QuadIndex == 0 || QuadIndex == 2 || QuadIndex == 4 ? 0 : 1;
 	VertexInfo.IsLeft		= QuadIndex == 0 || QuadIndex == 1 || QuadIndex == 5 ? 1 : 0;
 #endif
-	VertexInfo.HairPrimitiveId = BaseIndex;
+	VertexInfo.HairControlPointId = BaseIndex;
 	VertexInfo.VertexIndex = BaseIndex + VertexInfo.IsTip;
 #if !RAYHITGROUPSHADER
 	VertexInfo.LodRadiusScale = 1.0f;
@@ -145,7 +145,7 @@ FVertexInfo GetVertexInfo(FVertexFactoryInput Input)
 		VertexInfo.LodRadiusScale = HairStrandsVF.CullingRadiusScaleBuffer[BaseIndex + VertexInfo.IsTip];
 		VertexInfo.bForceInvalidQuad = VertexInfo.LodRadiusScale <= 0.0f;
 #endif
-		VertexInfo.HairPrimitiveId = VertexInfo.VertexIndex - VertexInfo.IsTip;
+		VertexInfo.HairControlPointId = VertexInfo.VertexIndex - VertexInfo.IsTip;
 	}
 #endif
 
@@ -181,7 +181,7 @@ struct FVertexFactoryIntermediates
 	half3x3 TangentToLocal;
 	half3x3 TangentToWorld;
 	half TangentToWorldSign;
-	uint HairPrimitiveId;
+	uint HairControlPointId;
 	float2 HairPrimitiveUV;
 	float2 HairDimensions;  // This is kept on the intermediate as this is used in several vertex shader for the actual coverage.
 	float  HairDensity;
@@ -189,9 +189,9 @@ struct FVertexFactoryIntermediates
 	FSceneDataIntermediates SceneData;
 };
 
-uint GetHairPrimitiveId(FVertexFactoryInput Input)
+uint GetHairControlPointId(FVertexFactoryInput Input)
 {
-	return GetVertexInfo(Input).HairPrimitiveId;
+	return GetVertexInfo(Input).HairControlPointId;
 }
 
 FHairControlPoint GetVertexPosition(const FVertexInfo VertexInfo, bool bInvalidJointVertex=true)
@@ -299,7 +299,7 @@ FMaterialPixelParameters GetMaterialPixelParameters(FVertexFactoryInterpolantsVS
 
 	Result.TwoSidedSign = 1;
 	Result.PrimitiveId = ToScalarMemory(GetPrimitiveId(Interpolants));
-	Result.HairPrimitiveId = Interpolants.HairPrimitiveId;
+	Result.HairControlPointId = Interpolants.HairControlPointId;
 	Result.HairPrimitiveUV = float2(Interpolants.HairPrimitiveUV);
 
 	return Result;
@@ -577,7 +577,7 @@ FMaterialPixelParameters GetMaterialPixelParameters(float3 RayOrigin, float3 Ray
 
 
 #endif // ENABLE_PROCEDURAL_INTERSECTOR
-	Result.HairPrimitiveId = VertexIndex0;
+	Result.HairControlPointId = VertexIndex0;
 	Result.HairPrimitiveUV = float2(UCoord, VCoord);
 
 	// Required for previewing materials that use ParticleColor
@@ -614,7 +614,7 @@ FVertexFactoryIntermediates GetVertexFactoryIntermediates(FVertexFactoryInput In
 {
 	FVertexFactoryIntermediates Intermediates;
 	Intermediates.SceneData			= VF_GPUSCENE_GET_INTERMEDIATES(Input);
-	Intermediates.HairPrimitiveId	= GetHairPrimitiveId(Input);
+	Intermediates.HairControlPointId= GetHairControlPointId(Input);
 	Intermediates.HairPrimitiveUV	= GetSegmentUVCoord(Input);
 	Intermediates.HairDensity		= HairStrandsVF.Density;
 	Intermediates.HairDimensions	= float2(0.f, GetWorldStrandRadius(Input));
@@ -663,10 +663,10 @@ float4 ComputeViewAlignedWorldPosition(FVertexFactoryInput Input, float3 WorldTa
 
 // @return translated world position 
 // Specialized version of VertexFactoryGetWorldPosition & ComputeViewAlignedWorldPosition, for visibility VS shader.
-float4 VertexFactoryGetWorldPosition_Visibility(FVertexFactoryInput Input, FHairViewInfo HairViewInfo, inout uint OutHairPrimitiveId)
+float4 VertexFactoryGetWorldPosition_Visibility(FVertexFactoryInput Input, FHairViewInfo HairViewInfo, inout uint OutHairControlPointId)
 {
 	const FVertexInfo VertexInfo = GetVertexInfo(Input);
-	OutHairPrimitiveId = VertexInfo.HairPrimitiveId;
+	OutHairControlPointId = VertexInfo.HairControlPointId;
 
 	FHairControlPoint ControlPoint = GetVertexPosition(VertexInfo);
 
@@ -757,7 +757,7 @@ FVertexFactoryInterpolantsVSToPS VertexFactoryGetInterpolantsVSToPS(FVertexFacto
 	Interpolants.EyeIndex = 0;
 #endif
 	
-	Interpolants.HairPrimitiveId = Intermediates.HairPrimitiveId;
+	Interpolants.HairControlPointId = Intermediates.HairControlPointId;
 	Interpolants.HairPrimitiveUV = Intermediates.HairPrimitiveUV;
 	SetPrimitiveId(Interpolants, Intermediates.SceneData.PrimitiveId);
 
@@ -893,7 +893,7 @@ FVertexFactoryRayTracingInterpolants VertexFactoryInterpolate(FVertexFactoryRayT
 	INTERPOLATE_MEMBER(InterpolantsVSToPS.TexCoords[0]);
 #endif
 
-	O.InterpolantsVSToPS.HairPrimitiveId = a.InterpolantsVSToPS.HairPrimitiveId;
+	O.InterpolantsVSToPS.HairControlPointId = a.InterpolantsVSToPS.HairControlPointId;
 	INTERPOLATE_MEMBER(InterpolantsVSToPS.HairPrimitiveUV);
 	return O;
 }
diff --git a/Engine/Shaders/Private/HairStrands/HairStrandsVisibilityCommon.ush b/Engine/Shaders/Private/HairStrands/HairStrandsVisibilityCommon.ush
--- a/Engine/Shaders/Private/HairStrands/HairStrandsVisibilityCommon.ush
+++ b/Engine/Shaders/Private/HairStrands/HairStrandsVisibilityCommon.ush
@@ -148,12 +148,12 @@ uint PackHairVisDepthCoverage(float InDepth, uint InCoverage8bits)
 	return (uint(InDepth * 0x00FFFFFF) << 8) | min(0xFFu, InCoverage8bits);
 }
 
-uint PackHairVisPrimitiveMaterialId(uint PrimitiveId, uint MaterialId)
+uint PackHairVisControlPointMaterialId(uint ControlPointId, uint MaterialId)
 {
-	return (PrimitiveId & 0x00FFFFFF) | ((MaterialId & 0xFF) << 24);
+	return (ControlPointId & 0x00FFFFFF) | ((MaterialId & 0xFF) << 24);
 }
 
-uint UnpackHairVisPrimitiveId(uint RawVis)
+uint UnpackHairVisControlPointId(uint RawVis)
 {
 	return RawVis & 0x00FFFFFF;
 }
@@ -178,7 +178,7 @@ uint UnpackHairVisCoverage8bits(uint InPacked)
 	return InPacked & 0xFF;
 }
 
-uint GetInvalidHairPrimitiveId()
+uint GetInvalidHairControlPointId()
 {
 	return 0xFFFFFFFF;
 }
@@ -193,7 +193,7 @@ struct FHairVis
 	float	Depth;
 	float   Coverage;	  // Coverage in float point, but with 8-bits precision. For convenience.
 	uint	Coverage8bit;
-	uint	PrimitiveId;
+	uint	ControlPointId;
 	uint	MaterialId;
 };
 
@@ -201,7 +201,7 @@ FPackedHairVis PackHairVis(const FHairVis In)
 {
 	FPackedHairVis Out;
 	Out.Depth_Coverage8bit	   = PackHairVisDepthCoverage(In.Depth, In.Coverage8bit);
-	Out.PrimitiveID_MaterialID = PackHairVisPrimitiveMaterialId(In.PrimitiveId, In.MaterialId);
+	Out.ControlPointID_MaterialID = PackHairVisControlPointMaterialId(In.ControlPointId, In.MaterialId);
 	return Out;
 }
 
@@ -211,8 +211,8 @@ FHairVis UnpackHairVis(const FPackedHairVis In)
 	Out.Depth		 = UnpackHairVisDepth(In.Depth_Coverage8bit);
 	Out.Coverage8bit = UnpackHairVisCoverage8bits(In.Depth_Coverage8bit);
 	Out.Coverage	 = float(Out.Coverage8bit) / float(0xFF);
-	Out.PrimitiveId  = UnpackHairVisPrimitiveId(In.PrimitiveID_MaterialID);
-	Out.MaterialId	 = UnpackHairVisMaterialId(In.PrimitiveID_MaterialID);
+	Out.ControlPointId= UnpackHairVisControlPointId(In.ControlPointID_MaterialID);
+	Out.MaterialId	 = UnpackHairVisMaterialId(In.ControlPointID_MaterialID);
 
 	return Out;
 }
@@ -222,14 +222,14 @@ FHairVis UnpackHairVis(const FPackedHairVis In)
 
 // Footprint: 128bits | 16bytes
 // Option to reduce the footprint: 
-// * we don't use PrimitiveID at the moment which is encoded on 3bytes
+// * we don't use ControlPointId at the moment which is encoded on 3bytes
 // * Deferred the evaluation of Coverage & material properties into some other pass?
 struct FHairSample
 {
 	float	Depth;
 	float3	Tangent;
 	uint	Coverage8bit;
-	uint	PrimitiveId;
+	uint	ControlPointId;
 	uint	MacroGroupId;
 	float3	BaseColor;
 	float	Roughness;
@@ -254,11 +254,11 @@ FPackedHairSample PackHairSample(const FHairSample In)
 		((0xFF & uint(T.z * 255)) << 16)| 
 		(min(uint(0xFF), In.Coverage8bit) << 24);
 
-	// PrimitiveId and MacroGroupId
+	// ControlPointId and MacroGroupId
 	// Encode into 1 x 32bits uint
-	// PrimitiveId is on 28bits | MacroGroupId is on 4bits
-	Out.PrimitiveID_MacroGroupID = 
-		((In.PrimitiveId  & 0x0FFFFFFF)) |
+	// ControlPointId is on 28bits | MacroGroupId is on 4bits
+	Out.ControlPointID_MacroGroupID = 
+		((In.ControlPointId  & 0x0FFFFFFF)) |
 		((In.MacroGroupId & 0xF) << 28); // If this encoding change, ensure it conforms to MAX_HAIR_MACROGROUP_COUNT
 
 	Out.Depth					= In.Depth;																			// 32bits float
@@ -284,8 +284,8 @@ FHairSample UnpackHairSample(const FPackedHairSample In)
 		float((In.Tangent_Coverage8bit >> 16) & 0xFF) / 255.f));
 	Out.Coverage8bit = (In.Tangent_Coverage8bit >> 24) & 0xFF;
 
-	Out.PrimitiveId  = In.PrimitiveID_MacroGroupID & 0x0FFFFFFF;
-	Out.MacroGroupId = (In.PrimitiveID_MacroGroupID>>28) & 0xF; // If this encoding change, ensure it conforms to MAX_HAIR_MACROGROUP_COUNT
+	Out.ControlPointId  = In.ControlPointID_MacroGroupID & 0x0FFFFFFF;
+	Out.MacroGroupId = (In.ControlPointID_MacroGroupID>>28) & 0xF; // If this encoding change, ensure it conforms to MAX_HAIR_MACROGROUP_COUNT
 
 	const float4 BaseColorAndRoughness = UintToFloat4(In.BaseColor_Roughness);
 	Out.BaseColor			= BaseColorAndRoughness.xyz * BaseColorAndRoughness.xyz;
@@ -379,20 +379,20 @@ float TransmittanceToCoverage(float InTransmittance, float InCoverageThreashold)
 FPackedHairVisPPLL PackHairVisPPLL(
 	float	Depth, 
 	float	Coverage,
-	uint	PrimitiveId,
+	uint	ControlPointId,
 	uint	MaterialId,
 	uint	NextNodeIndex)
 {
 	FHairVis S;
 	S.Depth = Depth;
 	S.Coverage8bit = To8bitCoverage(Coverage); // To16bitCoverage(CoverageFullPrecision)
-	S.PrimitiveId = PrimitiveId;
+	S.ControlPointId = ControlPointId;
 	S.MaterialId = MaterialId;
 	const FPackedHairVis PackedS = PackHairVis(S);
 
 	FPackedHairVisPPLL Out;
 	Out.Depth_Coverage8bit = PackedS.Depth_Coverage8bit;
-	Out.PrimitiveID_MaterialID = PackedS.PrimitiveID_MaterialID;
+	Out.ControlPointID_MaterialID = PackedS.ControlPointID_MaterialID;
 	Out.NextNodeIndex = NextNodeIndex;
 	return Out;
 }
@@ -401,7 +401,7 @@ FHairVis UnpackHairVisPPLL(FPackedHairVisPPLL In)
 {
 	FPackedHairVis Temp;
 	Temp.Depth_Coverage8bit = In.Depth_Coverage8bit;
-	Temp.PrimitiveID_MaterialID = In.PrimitiveID_MaterialID;
+	Temp.ControlPointID_MaterialID = In.ControlPointID_MaterialID;
 	return UnpackHairVis(Temp);
 }
 
@@ -409,7 +409,7 @@ FPackedHairVis ConvertToPackedHairVis(FPackedHairVisPPLL In)
 {
 	FPackedHairVis Out;
 	Out.Depth_Coverage8bit = In.Depth_Coverage8bit;
-	Out.PrimitiveID_MaterialID = In.PrimitiveID_MaterialID;
+	Out.ControlPointID_MaterialID = In.ControlPointID_MaterialID;
 	return Out;
 }
 
diff --git a/Engine/Shaders/Private/HairStrands/HairStrandsVisibilityCommonStruct.ush b/Engine/Shaders/Private/HairStrands/HairStrandsVisibilityCommonStruct.ush
--- a/Engine/Shaders/Private/HairStrands/HairStrandsVisibilityCommonStruct.ush
+++ b/Engine/Shaders/Private/HairStrands/HairStrandsVisibilityCommonStruct.ush
@@ -5,7 +5,7 @@
 struct FPackedHairSample
 {
 	float Depth;
-	uint  PrimitiveID_MacroGroupID;
+	uint  ControlPointID_MacroGroupID;
 	uint  Tangent_Coverage8bit;
 	uint  BaseColor_Roughness;
 	uint  Specular_LightChannels_Backlit;
@@ -14,12 +14,12 @@ struct FPackedHairSample
 struct FPackedHairVis
 {
 	uint Depth_Coverage8bit;
-	uint PrimitiveID_MaterialID;
+	uint ControlPointID_MaterialID;
 };
 
 struct FPackedHairVisPPLL
 {
 	uint Depth_Coverage8bit;
-	uint PrimitiveID_MaterialID;
+	uint ControlPointID_MaterialID;
 	uint NextNodeIndex;
 };
\ No newline at end of file
diff --git a/Engine/Shaders/Private/HairStrands/HairStrandsVisibilityCompaction.usf b/Engine/Shaders/Private/HairStrands/HairStrandsVisibilityCompaction.usf
--- a/Engine/Shaders/Private/HairStrands/HairStrandsVisibilityCompaction.usf
+++ b/Engine/Shaders/Private/HairStrands/HairStrandsVisibilityCompaction.usf
@@ -227,8 +227,8 @@ void ComputeUniqueSamplesWithoutDepthTest(
 		// Note: InMSAA_IDTexture contains both the primitive ID and tha material ID. However 
 		// the material ID is constant along the primitive, so it is correct to use this as a 
 		// sorting/deduplication key
-		const uint PrimitiveId = InMSAA_IDTexture.Load(PixelCoord, SampleIt);
-		const bool bIsValid = PrimitiveId != GetInvalidHairPrimitiveId();
+		const uint HairControlPointId = InMSAA_IDTexture.Load(PixelCoord, SampleIt);
+		const bool bIsValid = HairControlPointId != GetInvalidHairControlPointId();
 		if (!bIsValid)
 			continue;
 
@@ -248,10 +248,10 @@ void ComputeUniqueSamplesWithoutDepthTest(
 			const float UniqueDepth = asfloat(OutSamples[UniqueIt].w);
 			const float UniqueSceneDepth = ConvertFromDeviceZ(UniqueDepth);
 			const bool bIsSimilar = 
-				PrimitiveId == OutSamples[UniqueIt].x ||
+				HairControlPointId == OutSamples[UniqueIt].x ||
 				abs(UniqueSceneDepth - SceneDepth) < DepthTheshold;
 		#else
-			const bool bIsSimilar = PrimitiveId == OutSamples[UniqueIt].x;
+			const bool bIsSimilar = HairControlPointId == OutSamples[UniqueIt].x;
 		#endif
 			if (bIsSimilar)
 			{
@@ -262,7 +262,7 @@ void ComputeUniqueSamplesWithoutDepthTest(
 				if (IntDepth > OutSamples[UniqueIt].w)
 				{
 				#if MERGE_SAMPLE
-					OutSamples[UniqueIt].x = PrimitiveId;
+					OutSamples[UniqueIt].x = HairControlPointId;
 				#endif
 					OutSamples[UniqueIt].z = SampleIt;
 					OutSamples[UniqueIt].w = asuint(SampleDepth);
@@ -275,7 +275,7 @@ void ComputeUniqueSamplesWithoutDepthTest(
 		
 		if (!bAlreadyExist)
 		{
-			OutSamples[OutSet.UniqueSampleCount].x = PrimitiveId;
+			OutSamples[OutSet.UniqueSampleCount].x = HairControlPointId;
 			OutSamples[OutSet.UniqueSampleCount].y = 1;
 			OutSamples[OutSet.UniqueSampleCount].z = SampleIt;
 			OutSamples[OutSet.UniqueSampleCount].w = asuint(SampleDepth);
@@ -320,8 +320,8 @@ void ComputeUniqueSamplesWithoutDepthTest(
 	OutSet.ValidSampleCount = 0;
 	OutSet.HairSampleCount = PERMUTATION_MSAACOUNT;
 
-	const uint PrimitiveId = InMSAA_IDTexture.Load(uint3(PixelCoord, 0));
-	const bool bIsValid = PrimitiveId != GetInvalidHairPrimitiveId();
+	const uint HairControlPointId = InMSAA_IDTexture.Load(uint3(PixelCoord, 0));
+	const bool bIsValid = HairControlPointId != GetInvalidHairControlPointId();
 	if (!bIsValid)
 		return;
 
@@ -330,7 +330,7 @@ void ComputeUniqueSamplesWithoutDepthTest(
 		return;
 
 	++OutSet.ValidSampleCount;
-	OutSamples[0].x = PrimitiveId;
+	OutSamples[0].x = HairControlPointId;
 	OutSamples[0].y = 1;
 	OutSamples[0].z = 0;
 	OutSamples[0].w = asuint(SampleDepth);
@@ -383,7 +383,7 @@ void MainCS(
 	FSampleSetDesc SampleDesc;
 
 	const float OpaqueDepth = SceneDepthTexture.Load(uint3(PixelCoord, 0)).r;
-	uint4 Samples[PERMUTATION_MSAACOUNT];		// x:PrimitiveId|MaterialId, y:Weight, z:SampleIt, w:Depth (as uint)
+	uint4 Samples[PERMUTATION_MSAACOUNT];		// x:ControlPointIdId|MaterialId, y:Weight, z:SampleIt, w:Depth (as uint)
 	ComputeUniqueSamplesWithoutDepthTest(
 		Samples,
 		SampleDesc,
@@ -421,7 +421,7 @@ void MainCS(
 			{
 				// VisibilityData.Coverage8bit is a weight normalising to 1 the contribution of all the compacted samples. Because later it is weighted by PixelCoverage.
 				FHairVis OutNodeVis;
-				OutNodeVis.PrimitiveId	= UnpackHairVisPrimitiveId(Samples[OutIndex].x);
+				OutNodeVis.ControlPointId= UnpackHairVisControlPointId(Samples[OutIndex].x);
 				OutNodeVis.Depth		= asfloat(Samples[OutIndex].w); // MSAA_DepthTexture.Load(PixelCoord, Samples[OutIndex].z);
 				OutNodeVis.Coverage8bit = To8bitCoverage(Samples[OutIndex].y / float(SampleDesc.ValidSampleCount));
 				OutNodeVis.MaterialId	= UnpackHairVisMaterialId(Samples[OutIndex].x);
diff --git a/Engine/Shaders/Private/HairStrands/HairStrandsVisibilityCompactionComputeRaster.usf b/Engine/Shaders/Private/HairStrands/HairStrandsVisibilityCompactionComputeRaster.usf
--- a/Engine/Shaders/Private/HairStrands/HairStrandsVisibilityCompactionComputeRaster.usf
+++ b/Engine/Shaders/Private/HairStrands/HairStrandsVisibilityCompactionComputeRaster.usf
@@ -88,8 +88,8 @@ void ComputeUniqueSamples(const uint2 PixelCoord, out uint4 OutSamples[PERMUTATI
 		// Note: PrimMatTexture contains both the primitive ID and the material ID. However 
 		// the material ID is constant along the primitive, so it is correct to use this as a 
 		// sorting/deduplication key
-		const uint PrimitiveId = PrimMatTexture.Load(uint4(PixelCoord, SampleIt, 0));
-		const bool bIsValid = PrimitiveId != GetInvalidHairPrimitiveId();
+		const uint HairControlPointId = PrimMatTexture.Load(uint4(PixelCoord, SampleIt, 0));
+		const bool bIsValid = HairControlPointId != GetInvalidHairControlPointId();
 		if (!bIsValid)
 		{
 			continue;
@@ -109,10 +109,10 @@ void ComputeUniqueSamples(const uint2 PixelCoord, out uint4 OutSamples[PERMUTATI
 			const float UniqueDepth = asfloat(OutSamples[UniqueIt].w);
 			const float UniqueSceneDepth = ConvertFromDeviceZ(UniqueDepth);
 			const bool bIsSimilar =
-				PrimitiveId == OutSamples[UniqueIt].x ||
+				HairControlPointId == OutSamples[UniqueIt].x ||
 				abs(UniqueSceneDepth - SceneDepth) < DepthTheshold;
 #else
-			const bool bIsSimilar = PrimitiveId == OutSamples[UniqueIt].x;
+			const bool bIsSimilar = HairControlPointId == OutSamples[UniqueIt].x;
 #endif
 			if (bIsSimilar)
 			{
@@ -123,7 +123,7 @@ void ComputeUniqueSamples(const uint2 PixelCoord, out uint4 OutSamples[PERMUTATI
 				if (IntDepth > OutSamples[UniqueIt].w)
 				{
 #if MERGE_SAMPLE
-					OutSamples[UniqueIt].x = PrimitiveId;
+					OutSamples[UniqueIt].x = HairControlPointId;
 #endif
 					OutSamples[UniqueIt].z = SampleIt;
 					OutSamples[UniqueIt].w = asuint(SampleDepth);
@@ -136,7 +136,7 @@ void ComputeUniqueSamples(const uint2 PixelCoord, out uint4 OutSamples[PERMUTATI
 
 		if (!bAlreadyExist)
 		{
-			OutSamples[OutSet.UniqueSampleCount].x = PrimitiveId;
+			OutSamples[OutSet.UniqueSampleCount].x = HairControlPointId;
 			OutSamples[OutSet.UniqueSampleCount].y = 1;
 			OutSamples[OutSet.UniqueSampleCount].z = SampleIt;
 			OutSamples[OutSet.UniqueSampleCount].w = asuint(SampleDepth);
@@ -189,7 +189,7 @@ void MainCS(uint GroupIndex : SV_GroupIndex, uint3 GroupId : SV_GroupID, uint3 G
 	}
 
 	FSampleSetDesc SampleDesc;
-	uint4 Samples[PERMUTATION_MULTI_SAMPLE_COUNT]; // x:PrimitiveId|MaterialId, y:Weight, z:SampleIt, w:Depth (as uint)
+	uint4 Samples[PERMUTATION_MULTI_SAMPLE_COUNT]; // x:ControlPointId|MaterialId, y:Weight, z:SampleIt, w:Depth (as uint)
 
 	ComputeUniqueSamples(PixelCoord, Samples, SampleDesc);
 
@@ -221,7 +221,7 @@ void MainCS(uint GroupIndex : SV_GroupIndex, uint3 GroupId : SV_GroupID, uint3 G
 			{
 				// VisibilityData.Coverage8bit is a weight normalising to 1 the contribution of all the compacted samples. Because later it is weighted by Categorization.PixelCoverage.
 				FHairVis OutNodeVis;
-				OutNodeVis.PrimitiveId	= UnpackHairVisPrimitiveId(Samples[OutIndex].x);
+				OutNodeVis.ControlPointId	= UnpackHairVisControlPointId(Samples[OutIndex].x);
 				OutNodeVis.MaterialId	= UnpackHairVisMaterialId(Samples[OutIndex].x);
 				OutNodeVis.Depth		= asfloat(Samples[OutIndex].w);
 				OutNodeVis.Coverage8bit = To8bitCoverage(Samples[OutIndex].y / float(SampleDesc.ValidSampleCount)); // 0xff;
diff --git a/Engine/Shaders/Private/HairStrands/HairStrandsVisibilityPS.usf b/Engine/Shaders/Private/HairStrands/HairStrandsVisibilityPS.usf
--- a/Engine/Shaders/Private/HairStrands/HairStrandsVisibilityPS.usf
+++ b/Engine/Shaders/Private/HairStrands/HairStrandsVisibilityPS.usf
@@ -40,7 +40,7 @@ EARLYDEPTHSTENCIL	// Required for UAV operations to not happen.
 #endif
 void MainVisibility(
 #if HAIR_RENDER_MODE == RENDER_MODE_MSAA_VISIBILITY
-	nointerpolation uint HairPrimitiveId : HAIR_PRIMITIVE_ID 
+	nointerpolation uint HairControlPointId : HAIR_PRIMITIVE_ID 
 #else
 	centroid FVertexFactoryInterpolantsVSToPS Interpolants
 #endif
@@ -94,7 +94,7 @@ void MainVisibility(
 #endif
 
 #if HAIR_RENDER_MODE == RENDER_MODE_PPLL && HAIR_STRAND_MESH_FACTORY
-	const uint HairPrimitiveId = Interpolants.HairPrimitiveId;
+	const uint HairControlPointId = Interpolants.HairControlPointId;
 	const uint2 PixelCoord = SvPosition.xy;
 	const float2 NDC = SvPosition.xy / SvPosition.w;
 	const float2 ScreenUV = NDC * ResolvedView.ScreenPositionScaleBias.xy + ResolvedView.ScreenPositionScaleBias.wz;
@@ -111,7 +111,7 @@ void MainVisibility(
 		const FPackedHairVisPPLL Node = PackHairVisPPLL(
 			SvPosition.z, 
 			Coverage, 
-			HairPrimitiveId,
+			HairControlPointId,
 			HairVisibilityPass_HairMaterialId,
 			NextNodeIndex);
 		HairVisibilityPass.PPLLNodeData[NodeIndex] = Node;
@@ -120,7 +120,7 @@ void MainVisibility(
 
 #if HAIR_RENDER_MODE == RENDER_MODE_MSAA_VISIBILITY
   #if HAIR_STRAND_MESH_FACTORY
-	OutColor0 = PackHairVisPrimitiveMaterialId(HairPrimitiveId, HairVisibilityPass_HairMaterialId);
+	OutColor0 = PackHairVisControlPointMaterialId(HairControlPointId, HairVisibilityPass_HairMaterialId);
   #else
 	OutColor0 = 0;
   #endif
diff --git a/Engine/Shaders/Private/HairStrands/HairStrandsVisibilityRasterCompute.usf b/Engine/Shaders/Private/HairStrands/HairStrandsVisibilityRasterCompute.usf
--- a/Engine/Shaders/Private/HairStrands/HairStrandsVisibilityRasterCompute.usf
+++ b/Engine/Shaders/Private/HairStrands/HairStrandsVisibilityRasterCompute.usf
@@ -956,7 +956,7 @@ void RasterCS(uint DispatchThreadID : SV_DispatchThreadID, uint GroupThreadID :
 			uint Prim = GroupThreadID;
 			uint PrimID = VisTilePrims[PrimOffset + Prim];
 
-			group_PrimMatID[Prim] = PackHairVisPrimitiveMaterialId(PrimID, HairMaterialId);
+			group_PrimMatID[Prim] = PackHairVisControlPointMaterialId(PrimID, HairMaterialId);
 
 			uint TypeDummy;
 			CalcHomogenousPosAndRad(HairStrandsVF_PositionBuffer[PrimID], group_PositionOffset, group_SP0[Prim], group_Rad0[Prim], TypeDummy);
@@ -975,7 +975,7 @@ void RasterCS(uint DispatchThreadID : SV_DispatchThreadID, uint GroupThreadID :
 				Coord += uint2(((PackedTileMin >> 0) & 0xffff), ((PackedTileMin >> 16) & 0xffff));
 
 				group_SubTile[LinearIndex].x = OutDepthCovTexture[uint3(Coord, 0)];
-				group_SubTile[LinearIndex].y = GetInvalidHairPrimitiveId();
+				group_SubTile[LinearIndex].y = GetInvalidHairControlPointId();
 				group_SubTile[LinearIndex].z = 0;
 				group_SubTile[LinearIndex].w = PackHairVisDepthCoverage(SceneDepthTexture.Load(uint3(Coord, 0)), 1.0f);
 			}
@@ -1005,7 +1005,7 @@ void RasterCS(uint DispatchThreadID : SV_DispatchThreadID, uint GroupThreadID :
 			else
 			{
 				uint PrimID = VisTilePrims[PrimOffset + Prim];
-				PrimMatID = PackHairVisPrimitiveMaterialId(PrimID, HairMaterialId);
+				PrimMatID = PackHairVisControlPointMaterialId(PrimID, HairMaterialId);
 
 				uint TypeDummy;
 				CalcHomogenousPosAndRad(HairStrandsVF_PositionBuffer[PrimID], group_PositionOffset, SP0, Rad0, TypeDummy);
@@ -1139,7 +1139,7 @@ void RasterCS(uint DispatchThreadID : SV_DispatchThreadID, uint GroupThreadID :
 
 				Coord += uint2(((PackedTileMin >> 0) & 0xffff), ((PackedTileMin >> 16) & 0xffff));
 
-				if (group_SubTile[LinearIndex].y != GetInvalidHairPrimitiveId())
+				if (group_SubTile[LinearIndex].y != GetInvalidHairControlPointId())
 				{
 					uint oldValue;
 					InterlockedMax(OutDepthCovTexture[uint3(Coord, 0)], group_SubTile[LinearIndex].x, oldValue);
@@ -1354,7 +1354,7 @@ void RasterMultiSampleCS(uint DispatchThreadID : SV_DispatchThreadID, uint Group
 		if (bThreadValid)
 		{
 			const uint PrimID = VisTilePrims[PrimOffset + Prim];
-			PrimMatID = PackHairVisPrimitiveMaterialId(PrimID, HairMaterialId);
+			PrimMatID = PackHairVisControlPointMaterialId(PrimID, HairMaterialId);
 
 			uint Type;
 			CalcHomogenousPosAndRad(HairStrandsVF_PositionBuffer[PrimID], group_PositionOffset, SP0, Rad0, Type);
@@ -1388,7 +1388,7 @@ void RasterMultiSampleCS(uint DispatchThreadID : SV_DispatchThreadID, uint Group
 				{
 					const uint HairDepth = OutDepthCovTexture[uint3(Coord, SampleIdx)];
 					group_SubTileHairDepth[SampleIdx][GroupThreadID] = HairDepth;
-					group_SubTilePrimMatID[SampleIdx][GroupThreadID] = GetInvalidHairPrimitiveId();
+					group_SubTilePrimMatID[SampleIdx][GroupThreadID] = GetInvalidHairControlPointId();
 				}
 			}
 
@@ -1447,7 +1447,7 @@ void RasterMultiSampleCS(uint DispatchThreadID : SV_DispatchThreadID, uint Group
 					const uint3 SampleCoord = uint3(Coord, SampleIdx);
 					
 					const uint PrimMatID = group_SubTilePrimMatID[SampleIdx][GroupThreadID];
-					if (PrimMatID != GetInvalidHairPrimitiveId())
+					if (PrimMatID != GetInvalidHairControlPointId())
 					{
 						const uint HairDepth = group_SubTileHairDepth[SampleIdx][GroupThreadID];
 
diff --git a/Engine/Shaders/Private/HairStrands/HairStrandsVisibilityRasterComputeNaive.usf b/Engine/Shaders/Private/HairStrands/HairStrandsVisibilityRasterComputeNaive.usf
--- a/Engine/Shaders/Private/HairStrands/HairStrandsVisibilityRasterComputeNaive.usf
+++ b/Engine/Shaders/Private/HairStrands/HairStrandsVisibilityRasterComputeNaive.usf
@@ -221,7 +221,7 @@ void CSMain(uint DispatchThreadID : SV_DispatchThreadID, uint GroupThreadID : SV
 		const uint NumSteps = (uint)ceil(max(DeltaAbs.x, DeltaAbs.y)) + 1;
 		const float RcpNumSteps = NumSteps > 1 ? (1.0f / (NumSteps - 1)) : 1.0f;
 		const bool bIsSteep = DeltaAbs.y > DeltaAbs.x;
-		const uint PrimMatID = PackHairVisPrimitiveMaterialId(PrimID, HairMaterialId);
+		const uint PrimMatID = PackHairVisControlPointMaterialId(PrimID, HairMaterialId);
 
 		LOOP for (int J = 0; J < NumSteps; ++J)
 		{
diff --git a/Engine/Shaders/Private/HairStrands/HairStrandsVisibilityRasterHW.usf b/Engine/Shaders/Private/HairStrands/HairStrandsVisibilityRasterHW.usf
--- a/Engine/Shaders/Private/HairStrands/HairStrandsVisibilityRasterHW.usf
+++ b/Engine/Shaders/Private/HairStrands/HairStrandsVisibilityRasterHW.usf
@@ -127,7 +127,7 @@ VSOutput VSMain(VSInput Input)
 	
 	Output.Position = mul(float4(WP, 1.0f), LWCHackToFloat(PrimaryView.WorldToClip));
 	Output.Radius = CPWorldRadius * 2000.0f;
-	Output.PrimMatID = PackHairVisPrimitiveMaterialId(PrimID, HairMaterialId);
+	Output.PrimMatID = PackHairVisControlPointMaterialId(PrimID, HairMaterialId);
 
 	return Output;
 }
diff --git a/Engine/Shaders/Private/HairStrands/HairStrandsVisibilityVS.usf b/Engine/Shaders/Private/HairStrands/HairStrandsVisibilityVS.usf
--- a/Engine/Shaders/Private/HairStrands/HairStrandsVisibilityVS.usf
+++ b/Engine/Shaders/Private/HairStrands/HairStrandsVisibilityVS.usf
@@ -9,7 +9,7 @@
 struct FHairVisibilityAccumulateVSToPS
 {
 #if HAIR_RENDER_MODE == RENDER_MODE_MSAA_VISIBILITY
-	nointerpolation uint HairPrimitiveId : HAIR_PRIMITIVE_ID;
+	nointerpolation uint HairControlPointId : HAIR_PRIMITIVE_ID;
 #else
 	FVertexFactoryInterpolantsVSToPS FactoryInterpolants;
 #endif
@@ -45,10 +45,10 @@ void Main(
 	HairViewInfo.RadiusAtDepth1 = lerp(HairRenderInfo.RadiusAtDepth1Primary, HairRenderInfo.RadiusAtDepth1Velocity, VelocityScale);
 	HairViewInfo.bIsOrthoView = HairRenderInfo.bIsOrthoView;
 
-	uint HairPrimitiveId = 0;
-	const float4 WorldPosition = VertexFactoryGetWorldPosition_Visibility(Input, HairViewInfo, HairPrimitiveId);
+	uint HairControlPointId = 0;
+	const float4 WorldPosition = VertexFactoryGetWorldPosition_Visibility(Input, HairViewInfo, HairControlPointId);
 	Output.Position = mul(WorldPosition, View.TranslatedWorldToClip);
-	Output.HairPrimitiveId = HairPrimitiveId;
+	Output.HairControlPointId = HairControlPointId;
 }
 #else // HAIR_RENDER_MODE == RENDER_MODE_MSAA_VISIBILITY
 {
@@ -115,7 +115,7 @@ void Main(
 	Output.Position = mul(WorldPosition, View.TranslatedWorldToClip);
 
 #if HAIR_RENDER_MODE == RENDER_MODE_MSAA_VISIBILITY
-	Output.HairPrimitiveId = VFIntermediates.HairPrimitiveId;
+	Output.HairControlPointId = VFIntermediates.HairControlPointId;
 #else
 	Output.FactoryInterpolants = VertexFactoryGetInterpolantsVSToPS(Input, VFIntermediates, VertexParameters);
 
diff --git a/Engine/Shaders/Private/LocalVertexFactoryCommon.ush b/Engine/Shaders/Private/LocalVertexFactoryCommon.ush
--- a/Engine/Shaders/Private/LocalVertexFactoryCommon.ush
+++ b/Engine/Shaders/Private/LocalVertexFactoryCommon.ush
@@ -46,8 +46,8 @@ struct FVertexFactoryInterpolantsVSToPS
 
 
 #if HAIR_STRAND_MESH_FACTORY
-		nointerpolation uint HairPrimitiveId		: HAIR_PRIMITIVE_ID;	// Control point ID
-		float2 HairPrimitiveUV						: HAIR_PRIMITIVE_UV;	// U: parameteric distance between the two surrounding control points. V: parametric distance along the width.
+	nointerpolation uint HairControlPointId			: HAIR_PRIMITIVE_ID;	// Control point ID
+	float2 HairPrimitiveUV							: HAIR_PRIMITIVE_UV;	// U: parameteric distance between the two surrounding control points. V: parametric distance along the width.
 #endif
 #if HAIR_CARD_MESH_FACTORY
 	float2 HairPrimitiveUV							: HAIR_PRIMITIVE_UV;	// Atlas UV
diff --git a/Engine/Shaders/Private/MaterialTemplate.ush b/Engine/Shaders/Private/MaterialTemplate.ush
--- a/Engine/Shaders/Private/MaterialTemplate.ush
+++ b/Engine/Shaders/Private/MaterialTemplate.ush
@@ -449,7 +449,7 @@ struct FMaterialPixelParameters
 
 	// Actual primitive Id
 #if HAIR_STRAND_MESH_FACTORY
-	uint	HairPrimitiveId;		// Control point ID
+	uint	HairControlPointId;		// Control point ID
 	float2	HairPrimitiveUV;		// U: parametric distance between the two surrounding control point. V: parametric distance along hair width
 #endif
 #if HAIR_CARD_MESH_FACTORY
@@ -1524,7 +1524,7 @@ float3 MaterialExpressionBlackBody( float Temp )
 float2 MaterialExpressionGetHairRootUV(FMaterialPixelParameters Parameters)
 {
 #if HAIR_STRAND_MESH_FACTORY
-	return GetHairStrandsRootUV(Parameters.HairPrimitiveId);
+	return GetHairStrandsRootUV(Parameters.HairControlPointId);
 #elif HAIR_CARD_MESH_FACTORY
 	return GetHairStrandsRootUV(Parameters.HairPrimitiveUV, Parameters.HairPrimitiveRootUV);
 #else
@@ -1535,7 +1535,7 @@ float2 MaterialExpressionGetHairRootUV(FMaterialPixelParameters Parameters)
 float2 MaterialExpressionGetHairUV(FMaterialPixelParameters Parameters)
 {	
 #if HAIR_STRAND_MESH_FACTORY
-	return GetHairStrandsUV(Parameters.HairPrimitiveId, Parameters.HairPrimitiveUV);
+	return GetHairStrandsUV(Parameters.HairControlPointId, Parameters.HairPrimitiveUV);
 #elif HAIR_CARD_MESH_FACTORY
 	return GetHairStrandsUV(Parameters.HairPrimitiveUV);
 #else
@@ -1546,7 +1546,7 @@ float2 MaterialExpressionGetHairUV(FMaterialPixelParameters Parameters)
 float2 MaterialExpressionGetHairDimensions(FMaterialPixelParameters Parameters)
 {
 #if HAIR_STRAND_MESH_FACTORY
-	return GetHairStrandsDimensions(Parameters.HairPrimitiveId, Parameters.HairPrimitiveUV.x);
+	return GetHairStrandsDimensions(Parameters.HairControlPointId, Parameters.HairPrimitiveUV.x);
 #elif HAIR_CARD_MESH_FACTORY
 	return GetHairStrandsDimensions(Parameters.HairPrimitiveUV, Parameters.HairPrimitiveLength);
 #else
@@ -1557,7 +1557,7 @@ float2 MaterialExpressionGetHairDimensions(FMaterialPixelParameters Parameters)
 float MaterialExpressionGetHairSeed(FMaterialPixelParameters Parameters)
 {
 #if HAIR_STRAND_MESH_FACTORY
-	return GetHairStrandsSeed(Parameters.HairPrimitiveId);
+	return GetHairStrandsSeed(Parameters.HairControlPointId);
 #elif HAIR_CARD_MESH_FACTORY
 	return GetHairStrandsSeed(Parameters.HairPrimitiveUV);
 #else
@@ -1568,7 +1568,7 @@ float MaterialExpressionGetHairSeed(FMaterialPixelParameters Parameters)
 float3 MaterialExpressionGetHairClumpID(FMaterialPixelParameters Parameters)
 {
 #if HAIR_STRAND_MESH_FACTORY
-	return GetHairStrandsClumpID(Parameters.HairPrimitiveId);
+	return GetHairStrandsClumpID(Parameters.HairControlPointId);
 #elif HAIR_CARD_MESH_FACTORY
 	return GetHairStrandsClumpID(Parameters.HairPrimitiveUV);
 #else
@@ -1579,7 +1579,7 @@ float3 MaterialExpressionGetHairClumpID(FMaterialPixelParameters Parameters)
 float3 MaterialExpressionGetHairBaseColor(FMaterialPixelParameters Parameters)
 {
 #if HAIR_STRAND_MESH_FACTORY
-	return GetHairStrandsColor(Parameters.HairPrimitiveId, Parameters.HairPrimitiveUV.x);
+	return GetHairStrandsColor(Parameters.HairControlPointId, Parameters.HairPrimitiveUV.x);
 #elif HAIR_CARD_MESH_FACTORY
 	return GetHairStrandsBaseColor(Parameters.HairPrimitiveUV, Parameters.HairPrimitiveMaterial.xyz);
 #else
@@ -1590,7 +1590,7 @@ float3 MaterialExpressionGetHairBaseColor(FMaterialPixelParameters Parameters)
 float MaterialExpressionGetHairRoughness(FMaterialPixelParameters Parameters)
 {
 #if HAIR_STRAND_MESH_FACTORY
-	return GetHairStrandsRoughness(Parameters.HairPrimitiveId, Parameters.HairPrimitiveUV.x);
+	return GetHairStrandsRoughness(Parameters.HairControlPointId, Parameters.HairPrimitiveUV.x);
 #elif HAIR_CARD_MESH_FACTORY
 	return GetHairStrandsRoughness(Parameters.HairPrimitiveUV, Parameters.HairPrimitiveMaterial.w);
 #else
@@ -1601,7 +1601,7 @@ float MaterialExpressionGetHairRoughness(FMaterialPixelParameters Parameters)
 float MaterialExpressionGetHairAO(FMaterialPixelParameters Parameters)
 {
 #if HAIR_STRAND_MESH_FACTORY
-	return GetHairStrandsAO(Parameters.HairPrimitiveId, Parameters.HairPrimitiveUV.x);
+	return GetHairStrandsAO(Parameters.HairControlPointId, Parameters.HairPrimitiveUV.x);
 #elif HAIR_CARD_MESH_FACTORY
 	return GetHairStrandsAO(Parameters.HairPrimitiveUV, 1.0f);
 #else
diff --git a/Engine/Source/Runtime/Renderer/Private/HairStrands/HairStrandsVisibility.cpp b/Engine/Source/Runtime/Renderer/Private/HairStrands/HairStrandsVisibility.cpp
--- a/Engine/Source/Runtime/Renderer/Private/HairStrands/HairStrandsVisibility.cpp
+++ b/Engine/Source/Runtime/Renderer/Private/HairStrands/HairStrandsVisibility.cpp
@@ -132,7 +132,7 @@ namespace HairStrandsVisibilityInternal
 	struct NodeData
 	{
 		uint32 Depth;
-		uint32 PrimitiveId_MacroGroupId;
+		uint32 ControlPointId_MacroGroupId;
 		uint32 Tangent_Coverage;
 		uint32 BaseColor_Roughness;
 		uint32 Specular;
@@ -142,7 +142,7 @@ namespace HairStrandsVisibilityInternal
 	struct NodeVis
 	{
 		uint32 Depth_Coverage;
-		uint32 PrimitiveId_MacroGroupId;
+		uint32 ControlPointId_MacroGroupId;
 	};
 }
 
@@ -603,7 +603,7 @@ private:
 		const FMaterial& RESTRICT MaterialResource,
 		const int32 MacroGroupId,
 		const int32 HairMaterialId,
-		const int32 HairPrimitiveId,
+		const int32 HairControlPointId,
 		const uint32 HairPrimitiveLightChannelMask,
 		const float HairCoverageScale);
 
@@ -686,7 +686,7 @@ bool FHairMaterialProcessor::Process(
 	const FMaterial& RESTRICT MaterialResource,
 	const int32 MacroGroupId,
 	const int32 HairMaterialId,
-	const int32 HairPrimitiveId,
+	const int32 HairControlPointId,
 	const uint32 HairPrimitiveLightChannelMask,
 	const float HairCoverageScale)
 {
@@ -713,7 +713,7 @@ bool FHairMaterialProcessor::Process(
 	}
 
 	FMeshPassProcessorRenderState DrawRenderState(PassDrawRenderState);
-	FHairMaterialShaderElementData ShaderElementData(MacroGroupId, HairMaterialId, HairPrimitiveId, HairPrimitiveLightChannelMask, HairCoverageScale);
+	FHairMaterialShaderElementData ShaderElementData(MacroGroupId, HairMaterialId, HairControlPointId, HairPrimitiveLightChannelMask, HairCoverageScale);
 	ShaderElementData.InitializeMeshMaterialData(ViewIfDynamicMeshCommand, PrimitiveSceneProxy, MeshBatch, StaticMeshId, false);
 
 	const bool bReverseCulling = ViewIfDynamicMeshCommand ? ViewIfDynamicMeshCommand->bIsPlanarReflection : false;
@@ -1230,7 +1230,7 @@ END_SHADER_PARAMETER_STRUCT()
 struct FPackedHairVisPPLL
 {
 	uint32 Depth;
-	uint32 PrimitiveId_MacroGroupId;
+	uint32 ControlPointId_MacroGroupId;
 	uint32 NextNodeIndex;
 };
 
@@ -1716,10 +1716,10 @@ static FRDGBufferRef AddCopyIndirectArgPass(
 }
 
 ///////////////////////////////////////////////////////////////////////////////////////////////////
-class FHairVisibilityPrimitiveIdCompactionCS : public FGlobalShader
+class FHairVisibilityControlPointIdCompactionCS : public FGlobalShader
 {
-	DECLARE_GLOBAL_SHADER(FHairVisibilityPrimitiveIdCompactionCS);
-	SHADER_USE_PARAMETER_STRUCT(FHairVisibilityPrimitiveIdCompactionCS, FGlobalShader);
+	DECLARE_GLOBAL_SHADER(FHairVisibilityControlPointIdCompactionCS);
+	SHADER_USE_PARAMETER_STRUCT(FHairVisibilityControlPointIdCompactionCS, FGlobalShader);
 
 	class FGroupSize	: SHADER_PERMUTATION_SPARSE_INT("PERMUTATION_GROUPSIZE", 32, 64);
 	class FPPLL 		: SHADER_PERMUTATION_SPARSE_INT("PERMUTATION_PPLL", 0, 8, 16, 32); // See GetPPLLMaxRenderNodePerPixel
@@ -1789,9 +1789,9 @@ public:
 	}
 };
 
-IMPLEMENT_GLOBAL_SHADER(FHairVisibilityPrimitiveIdCompactionCS, "/Engine/Private/HairStrands/HairStrandsVisibilityCompaction.usf", "MainCS", SF_Compute);
+IMPLEMENT_GLOBAL_SHADER(FHairVisibilityControlPointIdCompactionCS, "/Engine/Private/HairStrands/HairStrandsVisibilityCompaction.usf", "MainCS", SF_Compute);
 
-static void AddHairVisibilityPrimitiveIdCompactionPass(
+static void AddHairVisibilityControlPointIdCompactionPass(
 	const bool bUsePPLL,
 	FRDGBuilder& GraphBuilder,
 	const FViewInfo& View,
@@ -1799,7 +1799,7 @@ static void AddHairVisibilityPrimitiveIdCompactionPass(
 	const FHairStrandsMacroGroupDatas& MacroGroupDatas,
 	const uint32 NodeGroupSize,
 	const FHairStrandsTiles& TileData,
-	FHairVisibilityPrimitiveIdCompactionCS::FParameters* PassParameters,
+	FHairVisibilityControlPointIdCompactionCS::FParameters* PassParameters,
 	FRDGTextureRef& OutCompactCounter,
 	FRDGTextureRef& OutCompactNodeIndex,
 	FRDGBufferRef& OutCompactNodeVis,
@@ -1859,11 +1859,11 @@ static void AddHairVisibilityPrimitiveIdCompactionPass(
 	// Pixel coord of the node. Stored as 2*R16_UINT
 	OutCompactNodeCoord = GraphBuilder.CreateBuffer(FRDGBufferDesc::CreateBufferDesc(sizeof(uint32), MaxRenderNodeCount), TEXT("Hair.VisibilityNodeCoord"));
 
-	FHairVisibilityPrimitiveIdCompactionCS::FPermutationDomain PermutationVector;
-	PermutationVector.Set<FHairVisibilityPrimitiveIdCompactionCS::FGroupSize>(FHairStrandsTiles::GroupSize);
-	PermutationVector.Set<FHairVisibilityPrimitiveIdCompactionCS::FPPLL>(bUsePPLL ? PPLLMaxRenderNodePerPixel : 0);
-	PermutationVector.Set<FHairVisibilityPrimitiveIdCompactionCS::FMSAACount>(MSAASampleCount);
-	PermutationVector = FHairVisibilityPrimitiveIdCompactionCS::RemapPermutation(PermutationVector);
+	FHairVisibilityControlPointIdCompactionCS::FPermutationDomain PermutationVector;
+	PermutationVector.Set<FHairVisibilityControlPointIdCompactionCS::FGroupSize>(FHairStrandsTiles::GroupSize);
+	PermutationVector.Set<FHairVisibilityControlPointIdCompactionCS::FPPLL>(bUsePPLL ? PPLLMaxRenderNodePerPixel : 0);
+	PermutationVector.Set<FHairVisibilityControlPointIdCompactionCS::FMSAACount>(MSAASampleCount);
+	PermutationVector = FHairVisibilityControlPointIdCompactionCS::RemapPermutation(PermutationVector);
 
 	PassParameters->ResolutionOffset = FIntPoint(0,0);
 	PassParameters->OutputResolution = Resolution;
@@ -1887,7 +1887,7 @@ static void AddHairVisibilityPrimitiveIdCompactionPass(
 	PassParameters->TileDataBuffer		= TileData.GetTileBufferSRV(TileType);
 	PassParameters->IndirectBufferArgs	= TileData.TileIndirectDispatchBuffer;
  
-	TShaderMapRef<FHairVisibilityPrimitiveIdCompactionCS> ComputeShader(View.ShaderMap, PermutationVector);
+	TShaderMapRef<FHairVisibilityControlPointIdCompactionCS> ComputeShader(View.ShaderMap, PermutationVector);
 	FComputeShaderUtils::AddPass(
 		GraphBuilder,
 		RDG_EVENT_NAME("HairStrands::VisibilityCompaction"),
@@ -1997,7 +1997,7 @@ static void AddHairVisibilityCompactionComputeRasterPass(
 	const uint32 MSAASampleCount = GetHairVisibilityRenderMode(View.GetShaderPlatform()) == HairVisibilityRenderMode_MSAA_Visibility ? MaxSamplePerPixel : 1;
 	const uint32 PPLLMaxRenderNodePerPixel = MaxSamplePerPixel;
 	const uint32 MaxRenderNodeCount = GetTotalSampleCountForAllocation(Resolution, View.GetShaderPlatform());
-	OutCompactNodeVis = GraphBuilder.CreateBuffer(FRDGBufferDesc::CreateStructuredDesc(sizeof(HairStrandsVisibilityInternal::NodeVis), MaxRenderNodeCount), TEXT("Hair.VisibilityPrimitiveIdCompactNodeData"));
+	OutCompactNodeVis = GraphBuilder.CreateBuffer(FRDGBufferDesc::CreateStructuredDesc(sizeof(HairStrandsVisibilityInternal::NodeVis), MaxRenderNodeCount), TEXT("Hair.VisibilityControlPointIdCompactNodeData"));
 	OutCompactNodeCoord = GraphBuilder.CreateBuffer(FRDGBufferDesc::CreateBufferDesc(sizeof(uint32), MaxRenderNodeCount), TEXT("Hair.VisibilityPrimitiveIdCompactNodeCoord"));
 
 	FRDGTextureRef DefaultTexture = GSystemTextures.GetBlackDummy(GraphBuilder);
@@ -4188,12 +4188,12 @@ void DrawHitProxies(
 			VisIdTexture,
 			VisDepthTexture);
 
-		FHairVisibilityPrimitiveIdCompactionCS::FParameters* PassParameters = GraphBuilder.AllocParameters<FHairVisibilityPrimitiveIdCompactionCS::FParameters>();
+		FHairVisibilityControlPointIdCompactionCS::FParameters* PassParameters = GraphBuilder.AllocParameters<FHairVisibilityControlPointIdCompactionCS::FParameters>();
 		PassParameters->MSAA_DepthTexture = VisDepthTexture;
 		PassParameters->MSAA_IDTexture = VisIdTexture;
 		PassParameters->ViewTransmittanceTexture = ViewTransmittance.TransmittanceTexture;
 
-		AddHairVisibilityPrimitiveIdCompactionPass(
+		AddHairVisibilityControlPointIdCompactionPass(
 			false, // bUsePPLL
 			GraphBuilder,
 			View,
@@ -4637,7 +4637,7 @@ void RenderHairStrandsVisibilityBuffer(
 				VisibilityData.HairOnlyDepthTexture = HairOnlyDepthTexture;
 				
 				{
-					FHairVisibilityPrimitiveIdCompactionCS::FParameters* PassParameters = GraphBuilder.AllocParameters<FHairVisibilityPrimitiveIdCompactionCS::FParameters>();
+					FHairVisibilityControlPointIdCompactionCS::FParameters* PassParameters = GraphBuilder.AllocParameters<FHairVisibilityControlPointIdCompactionCS::FParameters>();
 					PassParameters->MSAA_DepthTexture = MsaaVisibilityResources.DepthTexture;
 					PassParameters->MSAA_IDTexture = MsaaVisibilityResources.IdTexture;
 					PassParameters->ViewTransmittanceTexture = ViewTransmittance.TransmittanceTexture;
@@ -4645,7 +4645,7 @@ void RenderHairStrandsVisibilityBuffer(
 					FRDGBufferRef CompactNodeCoord;
 					FRDGBufferRef IndirectArgsBuffer;
 					FRDGTextureRef ResolveMaskTexture = nullptr;
-					AddHairVisibilityPrimitiveIdCompactionPass(
+					AddHairVisibilityControlPointIdCompactionPass(
 						false, // bUsePPLL
 						GraphBuilder,
 						View,
@@ -4792,7 +4792,7 @@ void RenderHairStrandsVisibilityBuffer(
 
 				// Linked list sorting pass and compaction into common representation
 				{
-					FHairVisibilityPrimitiveIdCompactionCS::FParameters* PassParameters = GraphBuilder.AllocParameters<FHairVisibilityPrimitiveIdCompactionCS::FParameters>();
+					FHairVisibilityControlPointIdCompactionCS::FParameters* PassParameters = GraphBuilder.AllocParameters<FHairVisibilityControlPointIdCompactionCS::FParameters>();
 					PassParameters->PPLLCounter  = PPLLNodeCounterTexture;
 					PassParameters->PPLLNodeIndex= PPLLNodeIndexTexture;
 					PassParameters->PPLLNodeData = GraphBuilder.CreateSRV(PPLLNodeDataBuffer);
@@ -4801,7 +4801,7 @@ void RenderHairStrandsVisibilityBuffer(
 					FRDGBufferRef CompactNodeCoord;
 					FRDGBufferRef IndirectArgsBuffer;
 					FRDGTextureRef ResolveMaskTexture = nullptr;
-					AddHairVisibilityPrimitiveIdCompactionPass(
+					AddHairVisibilityControlPointIdCompactionPass(
 						true, // bUsePPLL
 						GraphBuilder,
 						View,
