Fix some Substrate RT occlusion behavior.

This matches better legacy behavior and makes SSS diffusion works similarly to SSS profile. There are still some discrepancy.

diff --git a/Engine/Shaders/Private/RayTracing/RayTracingOcclusionRGS.usf b/Engine/Shaders/Private/RayTracing/RayTracingOcclusionRGS.usf
--- a/Engine/Shaders/Private/RayTracing/RayTracingOcclusionRGS.usf
+++ b/Engine/Shaders/Private/RayTracing/RayTracingOcclusionRGS.usf
@@ -664,21 +664,30 @@ RAY_TRACING_ENTRY_RAYGEN_OR_INLINE(Occlusion)
 			const FStrataTopLayerData TopLayerData = StrataUnpackTopLayerData(Strata.TopLayerTexture.Load(uint3(PixelCoord, 0)));
 
 			const FStrataSubsurfaceHeader SSSHeader = StrataLoadSubsurfaceHeader(Strata.MaterialTextureArray, Strata.FirstSliceStoringStrataSSSData, PixelCoord);
-			const bool bSSSIsValid = StrataSubSurfaceHeaderGetUseDiffusion(SSSHeader);
-			const bool bSSSIsProfile = StrataSubSurfaceHeaderGetIsProfile(SSSHeader);
+			
+			const uint SSSType = SSSHEADER_TYPE(SSSHeader);
+			const bool bHasSSS = StrataPixelHeader.HasSubsurface();
+			const uint BSDFType = StrataPixelHeader.StrataGetBSDFType();
 
 			ShadingParameters.bIsValid = StrataPixelHeader.IsStrataMaterial();
-			ShadingParameters.bIsHair = StrataPixelHeader.StrataGetBSDFType() == STRATA_BSDF_TYPE_HAIR;
-			ShadingParameters.bIsSubsurface = bSSSIsValid && !bSSSIsProfile;
-			ShadingParameters.bHasSubsurfaceProfile = bSSSIsValid && bSSSIsProfile;
-			ShadingParameters.bBackfaceCulling = any(StrataSubSurfaceHeaderGetMFP(SSSHeader) > 0); // STRATA_TODO
+			ShadingParameters.bIsHair = BSDFType == STRATA_BSDF_TYPE_HAIR;
+			ShadingParameters.bIsSubsurface = bHasSSS ? SSSType == SSS_TYPE_WRAP : false;
+			ShadingParameters.bHasSubsurfaceProfile = bHasSSS ? (SSSType == SSS_TYPE_DIFFUSION_PROFILE || SSSType == SSS_TYPE_DIFFUSION) : false;
+			ShadingParameters.bBackfaceCulling = !(BSDFType == STRATA_BSDF_TYPE_EYE || (bHasSSS && SSSType == SSS_TYPE_WRAP)); // isthin?
 
 			WorldNormal = TopLayerData.WorldNormal;
 
 			if (ENABLE_TRANSMISSION && ShadingParameters.bHasSubsurfaceProfile)
 			{
-				TransmissionParameters = GetTransmissionProfileParams(StrataSubSurfaceHeaderGetProfileId(SSSHeader));
-				//@TODO: Read the max MFP from STRATA?
+				if (SSSType == SSS_TYPE_DIFFUSION_PROFILE)
+				{
+					TransmissionParameters = GetTransmissionProfileParams(StrataSubSurfaceHeaderGetProfileId(SSSHeader));
+				}
+				else if (SSSType == SSS_TYPE_DIFFUSION)
+				{
+					TransmissionParameters = InitTransmissionProfileParams();
+					TransmissionParameters.ExtinctionScale = StrataSubSurfaceHeaderGetMFP(SSSHeader);
+				}
 				SubsurfaceSamplingMeanFreePath = TransmissionParameters.ExtinctionScale;
 			}
 
diff --git a/Engine/Shaders/Private/ShadowProjectionPixelShader.usf b/Engine/Shaders/Private/ShadowProjectionPixelShader.usf
--- a/Engine/Shaders/Private/ShadowProjectionPixelShader.usf
+++ b/Engine/Shaders/Private/ShadowProjectionPixelShader.usf
@@ -321,6 +321,11 @@ void Main(
 				{
 					TransmissionParams = GetTransmissionProfileParams(StrataSubSurfaceHeaderGetProfileId(SSSHeader));
 				}
+				else if (SSSType == SSS_TYPE_DIFFUSION)
+				{
+					TransmissionParams = InitTransmissionProfileParams(); // IOR and phase are incorrect. but not used
+					TransmissionParams.ExtinctionScale = StrataSubSurfaceHeaderGetMFP(SSSHeader);
+				}
 			#else // STRATA_ENABLED
 				float3 WorldNormal = GBufferData.WorldNormal;
 				FTransmissionProfileParams TransmissionParams = GetTransmissionProfileParams(ExtractSubsurfaceProfileInt(GBufferData));
