Use GetGeometryNormalFromTriangleBaseAttributes(...) in RayTracingDebugMainCHS.
- instead of reimplementing the same logic in ComputeWorldNormal()

diff --git a/Engine/Shaders/Private/RayTracing/RayTracingDebugCHS.usf b/Engine/Shaders/Private/RayTracing/RayTracingDebugCHS.usf
--- a/Engine/Shaders/Private/RayTracing/RayTracingDebugCHS.usf
+++ b/Engine/Shaders/Private/RayTracing/RayTracingDebugCHS.usf
@@ -7,23 +7,6 @@
 #include "RayTracingDebugUtils.ush"
 #include "RayTracingHitGroupCommon.ush"
 
-
-float3 ComputeWorldNormal()
-{
-	FTriangleBaseAttributes Tri = LoadTriangleBaseAttributes(PrimitiveIndex());
-
-	float3 LocalEdges[2] = {
-		Tri.LocalPositions[1] - Tri.LocalPositions[0],
-		Tri.LocalPositions[2] - Tri.LocalPositions[0]
-	};
-	float3 LocalNormal = cross(LocalEdges[1], LocalEdges[0]);
-
-	float3x3 InverseTranspose3x3 = transpose((float3x3)WorldToObject4x3());
-	float3 WorldNormal = normalize(mul(LocalNormal, InverseTranspose3x3));
-
-	return WorldNormal;
-}
-
 RAY_TRACING_ENTRY_CLOSEST_HIT(RayTracingDebugMainCHS,
 	FRayTracingDebugPayload, Payload,
 	FRayTracingIntersectionAttributes, Attributes)
@@ -36,7 +19,7 @@ RAY_TRACING_ENTRY_CLOSEST_HIT(RayTracingDebugMainCHS,
 
 	Payload.TriangleIndex = PrimitiveIndex();
 
-	Payload.WorldNormal = ComputeWorldNormal();
+	Payload.WorldNormal = GetGeometryNormalFromTriangleBaseAttributes(PrimitiveIndex());
 	Payload.GeometryInstanceIndex = GetHitGroupUserData();
 	Payload.InstanceIndex = InstanceIndex();
 	Payload.GeometryIndex = 0; // GeometryIndex();
