Workaround for a ShaderCompiler issue in the DFAO

#rnx

diff --git a/Engine/Shaders/Private/DistanceField/GlobalDistanceFieldShared.ush b/Engine/Shaders/Private/DistanceField/GlobalDistanceFieldShared.ush
--- a/Engine/Shaders/Private/DistanceField/GlobalDistanceFieldShared.ush
+++ b/Engine/Shaders/Private/DistanceField/GlobalDistanceFieldShared.ush
@@ -248,9 +248,7 @@ float GetDistanceToNearestSurfaceGlobal(float3 TranslatedWorldPosition)
 
 float SampleGlobalDistanceField(float3 TranslatedWorldPosition, float MaxDistance, int MinClipmapIndex)
 {
-	float Distance = MaxDistance;
-
-	for (uint ClipmapIndex = MinClipmapIndex; Distance <= MaxDistance && ClipmapIndex < NumGlobalSDFClipmaps; ++ClipmapIndex)
+	for (uint ClipmapIndex = MinClipmapIndex; ClipmapIndex < NumGlobalSDFClipmaps; ++ClipmapIndex)
 	{
 		const float ClipmapInfluenceRange = GLOBAL_DISTANCE_FIELD_INFLUENCE_RANGE_IN_VOXELS * 2.0f * GlobalVolumeTranslatedCenterAndExtent[ClipmapIndex].w * GlobalVolumeTexelSize;
 		
@@ -265,13 +263,12 @@ float SampleGlobalDistanceField(float3 TranslatedWorldPosition, float MaxDistanc
 
 			if (DistanceFieldValue < 1.0f)
 			{
-				Distance = DecodeGlobalDistanceFieldPageDistance(DistanceFieldValue, ClipmapInfluenceRange);
-				break;
+				return DecodeGlobalDistanceFieldPageDistance(DistanceFieldValue, ClipmapInfluenceRange);
 			}
 		}
 	}
 
-	return Distance;
+	return MaxDistance;
 }
 
 float GlobalDistanceFieldSampleClipmap(float3 ClipmapVolumeUV, uint ClipmapIndex)
