Updated FCullObjectsForShadowCS to use wrapped group count in order to support more than 4M instances.

diff --git a/Engine/Shaders/Private/DistanceFieldShadowing.usf b/Engine/Shaders/Private/DistanceFieldShadowing.usf
--- a/Engine/Shaders/Private/DistanceFieldShadowing.usf
+++ b/Engine/Shaders/Private/DistanceFieldShadowing.usf
@@ -5,6 +5,7 @@
 =============================================================================*/
 
 #include "Common.ush"
+#include "ComputeShaderUtils.ush"
 #include "DeferredShadingCommon.ush"
 #include "DistanceFieldLightingShared.ush"
 #include "DistanceFieldShadowingShared.ush"
@@ -76,16 +77,16 @@ uint GetNumSceneObjects()
 
 [numthreads(UPDATEOBJECTS_THREADGROUP_SIZE, 1, 1)]
 void CullObjectsForShadowCS(
-	uint3 GroupId : SV_GroupID,
-	uint3 DispatchThreadId : SV_DispatchThreadID,
-    uint3 GroupThreadId : SV_GroupThreadID) 
+	uint GroupIndex : SV_GroupIndex,
+	uint3 GroupId : SV_GroupID)
 {
-	uint ObjectIndex = DispatchThreadId.x;
+	const uint ThreadIndex = GetUnWrappedDispatchThreadId(GroupId, GroupIndex, UPDATEOBJECTS_THREADGROUP_SIZE);
+	const uint ObjectIndex = ThreadIndex;
 
 #define USE_FRUSTUM_CULLING 1
 #if USE_FRUSTUM_CULLING
 
-	if (DispatchThreadId.x == 0)
+	if (ThreadIndex == 0)
 	{
 		// RWObjectIndirectArguments is zeroed by a clear before this shader, only need to set things that are non-zero (and are not read by this shader as that would be a race condition)
 		// IndexCount, NumInstances, StartIndex, BaseVertexIndex, FirstInstance
@@ -136,7 +137,7 @@ void CullObjectsForShadowCS(
 
 #else
 
-	if (DispatchThreadId.x == 0)
+	if (ThreadIndex == 0)
 	{
 		// IndexCount, NumInstances, StartIndex, BaseVertexIndex, FirstInstance
 		RWObjectIndirectArguments[0] = ObjectBoundingGeometryIndexCount;
diff --git a/Engine/Source/Runtime/Renderer/Private/DistanceFieldShadowing.cpp b/Engine/Source/Runtime/Renderer/Private/DistanceFieldShadowing.cpp
--- a/Engine/Source/Runtime/Renderer/Private/DistanceFieldShadowing.cpp
+++ b/Engine/Source/Runtime/Renderer/Private/DistanceFieldShadowing.cpp
@@ -592,14 +592,15 @@ void CullDistanceFieldObjectsForLight(
 		FCullObjectsForShadowCS::FPermutationDomain PermutationVector;
 		PermutationVector.Set< FCullObjectsForShadowCS::FPrimitiveType >(PrimitiveType);
 		auto ComputeShader = View.ShaderMap->GetShader<FCullObjectsForShadowCS>(PermutationVector);
-		const int32 GroupSize = FMath::DivideAndRoundUp<uint32>(NumObjectsInBuffer, UpdateObjectsGroupSize);
+
+		const FIntVector GroupCount = FComputeShaderUtils::GetGroupCountWrapped(NumObjectsInBuffer, UpdateObjectsGroupSize);
 
 		FComputeShaderUtils::AddPass(
 			GraphBuilder,
 			RDG_EVENT_NAME("CullMeshSDFObjectsToFrustum"),
 			ComputeShader,
 			PassParameters,
-			FIntVector(GroupSize, 1, 1));
+			GroupCount);
 	}
 
 	// Allocate tile resolution based on world space size
