Don't include NaniteVertexFactory.ush in RayTracingHitGroupCommon.ush
- not necessary and causes compilation errors.

diff --git a/Engine/Shaders/Private/RayTracing/RayTracingHitGroupCommon.ush b/Engine/Shaders/Private/RayTracing/RayTracingHitGroupCommon.ush
--- a/Engine/Shaders/Private/RayTracing/RayTracingHitGroupCommon.ush
+++ b/Engine/Shaders/Private/RayTracing/RayTracingHitGroupCommon.ush
@@ -18,7 +18,6 @@
 
 	#include "/Engine/Private/Nanite/NaniteDataDecode.ush"
 	#include "/Engine/Private/Nanite/NaniteAttributeDecode.ush"
-	#include "/Engine/Private/Nanite/NaniteVertexFactory.ush"
 #endif
 
 uint3 LoadIndices16Bit(ByteAddressBuffer SourceBuffer, uint OffsetInBytes)
