Path Tracing: Fixing recent regressions in spectral tracking, where collisions were not correctly reported.

diff --git a/Engine/Shaders/Private/PathTracing/Volume/PathTracingHeterogeneousVolumes.ush b/Engine/Shaders/Private/PathTracing/Volume/PathTracingHeterogeneousVolumes.ush
--- a/Engine/Shaders/Private/PathTracing/Volume/PathTracingHeterogeneousVolumes.ush
+++ b/Engine/Shaders/Private/PathTracing/Volume/PathTracingHeterogeneousVolumes.ush
@@ -405,7 +405,7 @@ struct FVolumeTrackingResult
 	float3 Throughput;
 	float Pdf;
 	float Distance;
-	float SigmaBar;
+	float3 SigmaBar;
 	bool bIsCollision;
 };
 
@@ -443,7 +443,7 @@ FVolumeTrackingResult HeterogeneousVolumesDeltaTracking(float3 PathThroughput, f
 		float RandValue = RandomSequence_GenerateSample1D(RandSequence);
 #define SAMPLE_ACHROMATIC 0
 #if SAMPLE_ACHROMATIC
-		float DeltaT = -log(RandValue) / Sample.SigmaBar;
+		float DeltaT = -log(RandValue) / Sample.SigmaBar.x;
 #else
 		float DeltaT = SampleSpectralTransmittance(RandValue, Sample.SigmaBar, PathThroughput);
 		if (DeltaT < 0)
@@ -455,15 +455,15 @@ FVolumeTrackingResult HeterogeneousVolumesDeltaTracking(float3 PathThroughput, f
 #endif
 
 		Sample.Distance = min(Sample.Distance + DeltaT, TMax);
-		float Transmittance = exp(-Sample.SigmaBar * Sample.Distance);
+		float Transmittance = exp(-Sample.SigmaBar.x * Sample.Distance);
 
 		if (Sample.Distance < TMax)
 		{
 			Sample.bIsCollision = true;
 			float CollisionWeight = 1.0;
 #if SAMPLE_ACHROMATIC
-			Sample.Throughput *= CollisionWeight / Sample.SigmaBar;
-			Sample.Pdf = Transmittance * Sample.SigmaBar;
+			Sample.Throughput *= CollisionWeight / Sample.SigmaBar.x;
+			Sample.Pdf = Transmittance * Sample.SigmaBar.x;
 #else
 			float4 EvalResult = EvaluateSpectralTransmittanceHit(Sample.Distance, Sample.SigmaBar, PathThroughput);
 			Sample.Throughput = CollisionWeight * EvalResult.xyz;
@@ -490,7 +490,7 @@ FVolumeTrackingResult HeterogeneousVolumesRatioTracking(float3 PathThroughput, f
 	{
 		float RandValue = RandomSequence_GenerateSample1D(RandSequence);
 #if SAMPLE_ACHROMATIC
-		float DeltaT = -log(RandValue) / Sample.SigmaBar;
+		float DeltaT = -log(RandValue) / Sample.SigmaBar.x;
 #else
 		float DeltaT = SampleSpectralTransmittance(RandValue, Sample.SigmaBar, PathThroughput);
 		if (DeltaT < 0)
@@ -502,16 +502,16 @@ FVolumeTrackingResult HeterogeneousVolumesRatioTracking(float3 PathThroughput, f
 #endif
 
 		Sample.Distance = min(Sample.Distance + DeltaT, TMax);
-		float Transmittance = exp(-Sample.SigmaBar * Sample.Distance);
+		float Transmittance = exp(-Sample.SigmaBar.x * Sample.Distance);
 
 		if (Sample.Distance < TMax)
 		{
 			float3 SamplePos = Origin + Direction * Sample.Distance;
 			float3 SigmaT = HeterogeneousVolumesGetDensity(SamplePos).SigmaT;
-			float3 SpectralNullCollisionWeight = max(Sample.SigmaBar - SigmaT, 0);
+			float3 SpectralNullCollisionWeight = max(Sample.SigmaBar.x - SigmaT, 0);
 #if SAMPLE_ACHROMATIC
-			Sample.Throughput *= SpectralNullCollisionWeight / Sample.SigmaBar;
-			Sample.Pdf *= Transmittance * Sample.SigmaBar;
+			Sample.Throughput *= SpectralNullCollisionWeight / Sample.SigmaBar.x;
+			Sample.Pdf *= Transmittance * Sample.SigmaBar.x;
 #else
 			float4 EvalResult = EvaluateSpectralTransmittanceHit(Sample.Distance, Sample.SigmaBar, PathThroughput);
 			Sample.Throughput *= SpectralNullCollisionWeight * EvalResult.xyz;
diff --git a/Engine/Shaders/Private/PathTracing/Volume/PathTracingVolume.ush b/Engine/Shaders/Private/PathTracing/Volume/PathTracingVolume.ush
--- a/Engine/Shaders/Private/PathTracing/Volume/PathTracingVolume.ush
+++ b/Engine/Shaders/Private/PathTracing/Volume/PathTracingVolume.ush
@@ -255,14 +255,27 @@ FVolumeTrackingResult VolumeSampleDistance(float3 PathThroughput, float3 Origin,
 	{
 		float RandValue = RandomSequence_GenerateSample1D(RandSequence);
 		TrackingResult.Distance = SampleSpectralTransmittance(RandValue, SigmaBar, PathThroughput);
-		TrackingResult.Throughput = EvaluateSpectralTransmittanceHit(TrackingResult.Distance, SigmaBar, PathThroughput).xyz;
-		float3 SpectralPdf = PathThroughput * rcp(PathThroughput.x + PathThroughput.y + PathThroughput.z);
-		float3 DistancePdf = TrackingResult.Throughput;
-		if (TrackingResult.Distance < Interval.VolumeTMax)
+		if (TrackingResult.Distance < 0.0)
 		{
-			DistancePdf *= SigmaBar;
+			return TrackingResult;
+		}
+
+		TrackingResult.SigmaBar = SigmaBar;
+		TrackingResult.Throughput = PathThroughput;
+		TrackingResult.bIsCollision = Interval.VolumeTMin + TrackingResult.Distance < Interval.VolumeTMax;
+		if (TrackingResult.bIsCollision)
+		{
+			float4 Evaluation = EvaluateSpectralTransmittanceHit(TrackingResult.Distance, SigmaBar, PathThroughput);
+			TrackingResult.Throughput *= Evaluation.xyz;
+			TrackingResult.Pdf = Evaluation.w;
+			TrackingResult.Distance += Interval.VolumeTMin;
+		}
+		else
+		{
+			float4 Evaluation = EvaluateSpectralTransmittanceMiss(Interval.VolumeTMax - Interval.VolumeTMin, SigmaBar, PathThroughput);
+			TrackingResult.Throughput *= Evaluation.xyz;
+			TrackingResult.Pdf = Evaluation.w;
 		}
-		TrackingResult.Pdf = dot(DistancePdf, SpectralPdf);
 	}
 	return TrackingResult;
 }
