Fix multi-view distance field normals computation.

Both the pixel and compute shader do not offset properly back into the downscaled normal texture, e.g. rendering a stereo view means that the right eye computation overwrites the left eye computation instead of being offset.

The pixel shader requires rendering to the appropriate viewport with the correct UV offsets, while the compute shader requires writing back to the correct pixel index.

diff --git a/Engine/Shaders/Private/DistanceFieldScreenGridLighting.usf b/Engine/Shaders/Private/DistanceFieldScreenGridLighting.usf
--- a/Engine/Shaders/Private/DistanceFieldScreenGridLighting.usf
+++ b/Engine/Shaders/Private/DistanceFieldScreenGridLighting.usf
@@ -17,7 +17,7 @@ void ComputeDistanceFieldNormalPS(
 	out float4 OutColor : SV_Target0)
 {
 	// Sample from the center of the top left full resolution texel
-	float2 PixelCoord = floor(SVPos.xy) * DOWNSAMPLE_FACTOR + View.ViewRectMin.xy + .5f;
+	float2 PixelCoord = floor(SVPos.xy) * DOWNSAMPLE_FACTOR + .5f;
 	float2 ScreenUV = float2(PixelCoord * View.BufferSizeAndInvSize.zw);
 
 	float SceneDepth = CalcSceneDepth(ScreenUV);
@@ -50,7 +50,8 @@ void ComputeDistanceFieldNormalCS(
 	float3 WorldNormal = GetGBufferData(ScreenUV).WorldNormal.xyz;
 #endif
 
-	RWDistanceFieldNormal[DispatchThreadId.xy] = EncodeDownsampledGBuffer(WorldNormal, SceneDepth);
+	const uint2 DownsampledPixelCoords = DispatchThreadId.xy + uint2(floor(View.ViewRectMin.xy / DOWNSAMPLE_FACTOR));
+	RWDistanceFieldNormal[DownsampledPixelCoords] = EncodeDownsampledGBuffer(WorldNormal, SceneDepth);
 }
 
 Buffer<float4> TileConeDepthRanges;
diff --git a/Engine/Source/Runtime/Renderer/Private/DistanceFieldAmbientOcclusion.cpp b/Engine/Source/Runtime/Renderer/Private/DistanceFieldAmbientOcclusion.cpp
--- a/Engine/Source/Runtime/Renderer/Private/DistanceFieldAmbientOcclusion.cpp
+++ b/Engine/Source/Runtime/Renderer/Private/DistanceFieldAmbientOcclusion.cpp
@@ -421,7 +421,11 @@ void ComputeDistanceFieldNormal(
 				ERDGPassFlags::Raster,
 				[&View, PassParameters](FRHICommandList& RHICmdList)
 			{
-				RHICmdList.SetViewport(0, 0, 0.0f, View.ViewRect.Width() / GAODownsampleFactor, View.ViewRect.Height() / GAODownsampleFactor, 1.0f);
+				const float ViewportMinX = View.ViewRect.Min.X / GAODownsampleFactor;
+				const float ViewportMinY = View.ViewRect.Min.Y / GAODownsampleFactor;
+				const float ViewportMaxX = View.ViewRect.Max.X / GAODownsampleFactor;
+				const float ViewportMaxY = View.ViewRect.Max.Y / GAODownsampleFactor;
+				RHICmdList.SetViewport(ViewportMinX, ViewportMinY, 0.0f, ViewportMaxX, ViewportMaxY, 1.0f);
 
 				TShaderMapRef<FPostProcessVS> VertexShader(View.ShaderMap);
 
