Add Substrate specular profile to the path tracer.

diff --git a/Engine/Shaders/Private/PathTracing/Material/PathTracingSubstrate.ush b/Engine/Shaders/Private/PathTracing/Material/PathTracingSubstrate.ush
--- a/Engine/Shaders/Private/PathTracing/Material/PathTracingSubstrate.ush
+++ b/Engine/Shaders/Private/PathTracing/Material/PathTracingSubstrate.ush
@@ -155,11 +155,12 @@ FMaterialSample Substrate_SampleMaterial(
 	}
 	else if (bSampledSpecular)
 	{
+		const float3 FTint = SubstrateSpecularTint(Payload, NoV, NoL, VoH, NoH);
 #if STRATA_FRESNEL_F82
 		// STRATA_TODO: Need to tune the logic for specular micro-occlusion ...
-		const float3 F = Data.FuzzAttenuation * F_AdobeF82(Data.F0, Data.F90, VoH);
+		const float3 F = Data.FuzzAttenuation * F_AdobeF82(Data.F0, Data.F90, VoH) * FTint;
 #else
-		const float3 F = Data.FuzzAttenuation * F_Schlick(Data.F0, Data.F90, VoH);
+		const float3 F = Data.FuzzAttenuation * F_Schlick(Data.F0, Data.F90, VoH) * FTint;
 #endif
 
 		Result.AddLobeWithMIS((1.0 - Payload.RoughnessData.z) * F * GGXResult0.x * Data.Spec0.W, Spec0Pdf, Data.LobePdf.y);
@@ -219,10 +220,11 @@ FMaterialEval Substrate_EvalMaterial(
 	{
 		const float2 GGXResult0 = GGXEvalReflection(L, V, H, Data.Alpha0);
 		const float2 GGXResult1 = GGXEvalReflection(L, V, H, Data.Alpha1);
+		const float3 FTint = SubstrateSpecularTint(Payload, NoV, NoL, VoH, NoH);
 #if STRATA_FRESNEL_F82
-		const float3 F = Data.FuzzAttenuation * F_AdobeF82(Data.F0, Data.F90, VoH);
+		const float3 F = Data.FuzzAttenuation * F_AdobeF82(Data.F0, Data.F90, VoH) * FTint;
 #else
-		const float3 F = Data.FuzzAttenuation * F_Schlick(Data.F0, Data.F90, VoH);
+		const float3 F = Data.FuzzAttenuation * F_Schlick(Data.F0, Data.F90, VoH) * FTint;
 #endif
 		Result.AddLobeWithMIS(DiffuseSpecularScale.y * (1.0 - Payload.RoughnessData.z) * F * GGXResult0.x * Data.Spec0.W, GGXResult0.y, Data.LobePdf.y);
 		Result.AddLobeWithMIS(DiffuseSpecularScale.y * (      Payload.RoughnessData.z) * F * GGXResult1.x * Data.Spec1.W, GGXResult1.y, Data.LobePdf.z);
@@ -388,13 +390,14 @@ FMaterialSample SubstrateClearCoat_SampleMaterial(
 	}
 	else if (bSampledSpecular)
 	{
+		const float3 FTint = SubstrateSpecularTint(Payload, NoV, NoL, VoH, NoH);
 #if STRATA_FRESNEL_F82
 		// STRATA_TODO: Need to tune the logic for specular micro-occlusion ...
-		const float3 F = Data.FuzzAttenuation * F_AdobeF82(Data.F0, Data.F90, VoH);
+		const float3 F = Data.FuzzAttenuation * F_AdobeF82(Data.F0, Data.F90, VoH) * FTint;
 #else
-		const float3 FBase = Data.FuzzAttenuation * F_Schlick(Data.F0, Data.F90, VoH);
+		const float3 FBase = Data.FuzzAttenuation * F_Schlick(Data.F0, Data.F90, VoH) * FTint;
 #endif
-		const float FCoat = Payload.RoughnessData.z * ClearCoatFresnel(VoH);
+		const float3 FCoat = Payload.RoughnessData.z * ClearCoatFresnel(VoH) * FTint;
 
 		Result.AddLobeWithMIS(Data.CoatAttenuation * FBase * BaseGGXResult.x * Data.BaseSpec.W * Transmission, BaseSpecPdf, Data.LobePdf.y);
 		Result.AddLobeWithMIS(Data.FuzzAttenuation * FCoat * CoatGGXResult.x * Data.CoatSpec.W               , CoatSpecPdf, Data.LobePdf.z);
@@ -461,13 +464,14 @@ FMaterialEval SubstrateClearCoat_EvalMaterial(
 		const float BaseSpecPdf = BaseGGXResult.y;
 		const float CoatSpecPdf = CoatGGXResult.y;
 
+		const float3 FTint = SubstrateSpecularTint(Payload, NoV, NoL, VoH, NoH);
 #if STRATA_FRESNEL_F82
 		// STRATA_TODO: Need to tune the logic for specular micro-occlusion ...
-		const float3 F = Data.FuzzAttenuation * F_AdobeF82(Data.F0, Data.F90, VoH);
+		const float3 F = Data.FuzzAttenuation * F_AdobeF82(Data.F0, Data.F90, VoH) * FTint;
 #else
-		const float3 FBase = Data.FuzzAttenuation * F_Schlick(Data.F0, Data.F90, VoH);
+		const float3 FBase = Data.FuzzAttenuation * F_Schlick(Data.F0, Data.F90, VoH) * FTint;
 #endif
-		const float FCoat = Payload.RoughnessData.z * ClearCoatFresnel(VoH);
+		const float3 FCoat = Payload.RoughnessData.z * ClearCoatFresnel(VoH) * FTint;
 
 		Result.AddLobeWithMIS(Data.CoatAttenuation * FBase * BaseGGXResult.x * Data.BaseSpec.W * Transmission, BaseSpecPdf, Data.LobePdf.y);
 		Result.AddLobeWithMIS(Data.FuzzAttenuation * FCoat * CoatGGXResult.x * Data.CoatSpec.W               , CoatSpecPdf, Data.LobePdf.z);
diff --git a/Engine/Shaders/Private/PathTracing/Material/PathTracingSubstrateCommon.ush b/Engine/Shaders/Private/PathTracing/Material/PathTracingSubstrateCommon.ush
--- a/Engine/Shaders/Private/PathTracing/Material/PathTracingSubstrateCommon.ush
+++ b/Engine/Shaders/Private/PathTracing/Material/PathTracingSubstrateCommon.ush
@@ -17,3 +17,16 @@ float3 SubstrateLobeWeight(FPathTracingPayload Payload, float SatNoL)
 	const float3 TransmittanceAboveAlongL = select(Payload.TransmittanceN < 1.f, pow(max(Payload.TransmittanceN, 0.0001f), DistanceL), Payload.TransmittanceN);
 	return Payload.BSDFOpacity * Payload.WeightV * lerp(1.0f, TransmittanceAboveAlongL, Payload.CoverageAboveAlongN);
 }
+
+float3 SubstrateSpecularTint(FPathTracingPayload Payload, float NoV, float NoL, float VoH, float NoH)
+{
+	float3 Out = 1.f;
+	if (Payload.SpecularProfileId != 0)
+	{
+		const uint ProfileID = GetSpecularProfileId(Payload.SpecularProfileId);
+		const bool bUseHalfVector = GetSpecularProfileParameterization(Payload.SpecularProfileId) == 1;
+		const float3 ProfileCoord = bUseHalfVector ? float3(VoH, NoH, ProfileID) : float3(NoV, NoL, ProfileID);
+		Out = View.SpecularProfileTexture.SampleLevel(View.SpecularProfileSampler, ProfileCoord, 0).xyz;
+	}
+	return Out;
+}
\ No newline at end of file
diff --git a/Engine/Shaders/Private/PathTracing/Material/PathTracingSubstrateSolidGlass.ush b/Engine/Shaders/Private/PathTracing/Material/PathTracingSubstrateSolidGlass.ush
--- a/Engine/Shaders/Private/PathTracing/Material/PathTracingSubstrateSolidGlass.ush
+++ b/Engine/Shaders/Private/PathTracing/Material/PathTracingSubstrateSolidGlass.ush
@@ -124,11 +124,18 @@ FMaterialSample SubstrateSolidGlass_SampleMaterial(
 		}
 		else
 		{
+			// Specular profile tinting
+			const float NoL = saturate(L.z);
+			const float NoV = saturate(V.z);
+			const float VoH = saturate(dot(V, H));
+			const float NoH = saturate(H.z);
+			const float3 FTint = SubstrateSpecularTint(Payload, NoV, NoL, VoH, NoH);
+
 			// reflection side, apply full lobe weight and attenuation along L
 			// and account for probability of picking the fuzz layer
 			const float ClothPdf = 1 / (2 * PI);
 			Result.Pdf += Data.LobePdf.z * ClothPdf;
-			Result.Weight *= Data.FuzzAttenuation * SubstrateLobeWeight(Payload, saturate(L.z));
+			Result.Weight *= Data.FuzzAttenuation * SubstrateLobeWeight(Payload, saturate(L.z)) * FTint;
 		}
 
 		return Result;
@@ -208,9 +215,10 @@ FMaterialEval SubstrateSolidGlass_EvalMaterial(
 			const float2 GGXResult1 = GGXEvalReflection(L, V, H, Data.Alpha1);
 
 			const float Fg = FresnelReflectance(VoH, Data.Eta, Data.F0);
+			const float3 FTint = SubstrateSpecularTint(Payload, NoV, NoL, VoH, NoH);
 
-			Result.AddLobeWithMIS(DiffuseSpecularScale.y * Data.FuzzAttenuation * (1.0 - Payload.RoughnessData.z) * GGXResult0.x / Data.EGlass0, Fg * GGXResult0.y, Data.LobePdf.x);
-			Result.AddLobeWithMIS(DiffuseSpecularScale.y * Data.FuzzAttenuation * (      Payload.RoughnessData.z) * GGXResult1.x / Data.EGlass1, Fg * GGXResult1.y, Data.LobePdf.y);
+			Result.AddLobeWithMIS(FTint * DiffuseSpecularScale.y * Data.FuzzAttenuation * (1.0 - Payload.RoughnessData.z) * GGXResult0.x / Data.EGlass0, Fg * GGXResult0.y, Data.LobePdf.x);
+			Result.AddLobeWithMIS(FTint * DiffuseSpecularScale.y * Data.FuzzAttenuation * (      Payload.RoughnessData.z) * GGXResult1.x / Data.EGlass1, Fg * GGXResult1.y, Data.LobePdf.y);
 		}
 		{
 			// Cloth Lobe
diff --git a/Engine/Shaders/Private/PathTracing/Material/PathTracingSubstrateThinGlass.ush b/Engine/Shaders/Private/PathTracing/Material/PathTracingSubstrateThinGlass.ush
--- a/Engine/Shaders/Private/PathTracing/Material/PathTracingSubstrateThinGlass.ush
+++ b/Engine/Shaders/Private/PathTracing/Material/PathTracingSubstrateThinGlass.ush
@@ -155,7 +155,8 @@ FMaterialSample SubstrateThinGlass_SampleMaterial(
 	const float ClothPdf = 1.0 / (2 * PI);
 	if (bSampleGlass)
 	{
-		const FThinSlabWeights SlabResult = ComputeThinSlabWeights(Data.TransmissionColor, VoH, Payload.Ior, Data.F0);
+		FThinSlabWeights SlabResult = ComputeThinSlabWeights(Data.TransmissionColor, VoH, Payload.Ior, Data.F0);
+		SlabResult.Reflected *= SubstrateSpecularTint(Payload, NoV, NoL, VoH, NoH);
 
 		const float3 SlabTint = bIsReflection ? SlabResult.Reflected : SlabResult.Transmitted;
 		const float SlabW0 = (bIsReflection ? Data.SpecR0.W : Data.SpecT0.W) * Data.NudgeE.x;
@@ -225,7 +226,8 @@ FMaterialEval SubstrateThinGlass_EvalMaterial(
 	// Specular Lobes (either reflected or transmitted)
 	const float2 GGXResult0 = GGXEvalReflection(L, V, H, bIsReflection ? Data.AlphaR0 : Data.AlphaT0);
 	const float2 GGXResult1 = GGXEvalReflection(L, V, H, bIsReflection ? Data.AlphaR1 : Data.AlphaT1);
-	const FThinSlabWeights SlabResult = ComputeThinSlabWeights(Data.TransmissionColor, VoH, Payload.Ior, Data.F0);
+	FThinSlabWeights SlabResult = ComputeThinSlabWeights(Data.TransmissionColor, VoH, Payload.Ior, Data.F0);
+	SlabResult.Reflected *= SubstrateSpecularTint(Payload, NoV, NoL, VoH, NoH);
 
 	const float3 SlabTint = bIsReflection ? SlabResult.Reflected : SlabResult.Transmitted;
 	const float SlabW0 = (bIsReflection ? Data.SpecR0.W : Data.SpecT0.W) * Data.NudgeE.x;
diff --git a/Engine/Shaders/Private/PathTracing/PathTracingCommon.ush b/Engine/Shaders/Private/PathTracing/PathTracingCommon.ush
--- a/Engine/Shaders/Private/PathTracing/PathTracingCommon.ush
+++ b/Engine/Shaders/Private/PathTracing/PathTracingCommon.ush
@@ -322,6 +322,7 @@ struct FPathTracingPayload : FMinimalPayload
 	float3 WeightV;
 	float3 TransmittanceN;
 	float CoverageAboveAlongN;
+	float SpecularProfileId;
 
 	// not transported through packed payload, but used in raygen
 	float3 SubsurfaceColor;
@@ -609,15 +610,16 @@ FPackedPathTracingPayload PackPathTracingPayload(FPathTracingPayload Input)
 	Output.PackedData[9] = PayloadEncodeUnitVector(Input.WorldTangent);
 	Output.PackedData[10] = PayloadEncodeLDRColor(Input.TransparencyColor);
 	Output.PackedData[11] = f32tof16(Input.MeanFreePath.x);
-	Output.PackedData[11] |= f32tof16(Input.MeanFreePath.y) << 16;
+	Output.PackedData[11]|= f32tof16(Input.MeanFreePath.y) << 16;
 	Output.PackedData[12] = f32tof16(Input.MeanFreePath.z);
-	Output.PackedData[12] |= f32tof16(Input.PhaseG) << 16;
+	Output.PackedData[12]|= f32tof16(Input.PhaseG) << 16;
 	Output.PackedData[13] = PayloadEncodeLDRColor(Input.FuzzColor);
     Output.PackedData[14] = PayloadEncodeLDRColor(float3(0.5 * Input.Anisotropy + 0.5, Input.FuzzAmount, Input.FuzzRoughness)); // NOTE: Aniso is [-1,1]
 	Output.PackedData[15] = PayloadEncodeLDRColor(Input.WeightV);
 	Output.PackedData[16] = PayloadEncodeLDRColor(Input.TransmittanceN);
-	Output.PackedData[17] = f32tof16(Input.Ior);
-	Output.PackedData[17] |= f32tof16(Input.CoverageAboveAlongN) << 16;
+	Output.PackedData[17] = Pack10F(Input.Ior);
+	Output.PackedData[17]|= Pack10F(Input.CoverageAboveAlongN) << 10;
+	Output.PackedData[17]|= Pack10F(Input.SpecularProfileId) << 20;
 #else // PATHTRACING_STRATA_PAYLOAD
 
 
@@ -690,8 +692,9 @@ FPathTracingPayload UnpackPathTracingPayload(FPackedPathTracingPayload Input, FR
 	Output.FuzzRoughness       = Data14.z;
 	Output.WeightV             = PayloadDecodeLDRColor(Input.PackedData[15]);
 	Output.TransmittanceN      = PayloadDecodeLDRColor(Input.PackedData[16]);
-	Output.Ior                 = f16tof32(Input.PackedData[17]);
-	Output.CoverageAboveAlongN = f16tof32(Input.PackedData[17] >> 16);
+	Output.Ior                 = Unpack10F(Input.PackedData[17]);
+	Output.CoverageAboveAlongN = Unpack10F(Input.PackedData[17] >> 10);
+	Output.SpecularProfileId   = Unpack10F(Input.PackedData[17] >> 20);
 #else // PATHTRACING_STRATA_PAYLOAD
 	Output.BaseColor = PayloadDecodeLDRColor(Input.PackedData[5]);
 	Output.TransparencyColor = PayloadDecodeLDRColor(Input.PackedData[6]);
diff --git a/Engine/Shaders/Private/PathTracing/PathTracingMaterialHitShader.usf b/Engine/Shaders/Private/PathTracing/PathTracingMaterialHitShader.usf
--- a/Engine/Shaders/Private/PathTracing/PathTracingMaterialHitShader.usf
+++ b/Engine/Shaders/Private/PathTracing/PathTracingMaterialHitShader.usf
@@ -21,6 +21,9 @@
 // Ensure that SSS albedo comes through in the material
 #define STRATA_SSS_MATERIAL_OVERRIDE 0
 
+// Ensure specular profile and glints are supported
+#define STRATA_COMPLEXSPECIALPATH 1
+
 #include "/Engine/Private/Common.ush"
 #include "/Engine/Private/RayTracing/RayTracingCommon.ush"
 #include "/Engine/Private/RayTracing/RayTracingHitGroupCommon.ush"
@@ -796,6 +799,11 @@ RAY_TRACING_ENTRY_CLOSEST_HIT(PathTracingMaterialCHS,
 						Payload.MeanFreePath = 0;
 						Payload.PhaseG       = 0;
 					}
+					Payload.SpecularProfileId = 0;
+					if (BSDF_GETHASSPECPROFILE(BSDF))
+					{
+						Payload.SpecularProfileId = SLAB_SPECPROFILEID(BSDF);
+					}
 					Payload.WeightV             = BSDF.LuminanceWeightV;
 					Payload.TransmittanceN      = BSDF.TransmittanceAboveAlongN;
 					Payload.CoverageAboveAlongN = BSDF.CoverageAboveAlongN;
