Fix for VSM shadow and SSS.

diff --git a/Engine/Shaders/Private/RayTracing/RayTracingOcclusionRGS.usf b/Engine/Shaders/Private/RayTracing/RayTracingOcclusionRGS.usf
--- a/Engine/Shaders/Private/RayTracing/RayTracingOcclusionRGS.usf
+++ b/Engine/Shaders/Private/RayTracing/RayTracingOcclusionRGS.usf
@@ -677,7 +677,8 @@ RAY_TRACING_ENTRY_RAYGEN_OR_INLINE(Occlusion)
 
 			WorldNormal = TopLayerData.WorldNormal;
 
-			if (ENABLE_TRANSMISSION && ShadingParameters.bHasSubsurfaceProfile)
+
+			if (ENABLE_TRANSMISSION)
 			{
 				if (SSSType == SSS_TYPE_DIFFUSION_PROFILE)
 				{
@@ -685,15 +686,31 @@ RAY_TRACING_ENTRY_RAYGEN_OR_INLINE(Occlusion)
 				}
 				else if (SSSType == SSS_TYPE_DIFFUSION)
 				{
+					FSubsurfaceOpacityMFP SubsurfaceOpacityMFP = StrataGetOpacityFromSubSurface(SSSHeader);
+
+					// We know that or SSS_TYPE_DIFFUSION we have access to the mean free path. 
+					const float MFP = SubsurfaceOpacityMFP.Data;
+					const float Extinction = 1.0f / max(0.00001, MFP);
+
 					TransmissionParameters = InitTransmissionProfileParams();
-					TransmissionParameters.ExtinctionScale = StrataSubSurfaceHeaderGetMFP(SSSHeader);
+					TransmissionParameters.ExtinctionScale = Extinction;
+					
+					SubsurfaceSamplingMeanFreePath = 1.0;
+					BRANCH
+					if (TransmissionMeanFreePathType == TransmissionMeanFreePathType_MaxMeanFreePath)
+					{
+						SubsurfaceSamplingMeanFreePath = MFP;
+					}
 				}
-				SubsurfaceSamplingMeanFreePath = TransmissionParameters.ExtinctionScale;
 			}
 
-			if (ShadingParameters.bIsSubsurface)
+			if (ShadingParameters.bIsSubsurface) // Legacy Subsurface shading model == Substrate Wrap SSS.
 			{
-				TransmissionParameters.ExtinctionScale = StrataGetOpacityFromSubSurface(SSSHeader);
+				FSubsurfaceOpacityMFP SubsurfaceOpacityMFP = StrataGetOpacityFromSubSurface(SSSHeader);
+				if (SubsurfaceOpacityMFP.bDataIsOpacity)
+				{
+					TransmissionParameters.ExtinctionScale = SubsurfaceOpacityMFP.Data;
+				}
 			}
 		}	
 		#else
diff --git a/Engine/Shaders/Private/ShadowProjectionPixelShader.usf b/Engine/Shaders/Private/ShadowProjectionPixelShader.usf
--- a/Engine/Shaders/Private/ShadowProjectionPixelShader.usf
+++ b/Engine/Shaders/Private/ShadowProjectionPixelShader.usf
@@ -311,11 +311,11 @@ void Main(
 			#if USE_TRANSMISSION
 			if (bHasSubsurfaceTransmission)
 			{
+				// This branch is only taken by SSS profiles and SSS per pixel diffusion methods.
 			#if STRATA_ENABLED
 				const FStrataTopLayerData TopLayerData = StrataUnpackTopLayerData(Strata.TopLayerTexture.Load(uint3(SVPos.xy, 0)));
 				const float3 WorldNormal = TopLayerData.WorldNormal;
 
-				// For non-profile version, use default transmission parameters (only NormalScale & ExtinctionScale are used)
 				FTransmissionProfileParams TransmissionParams = InitTransmissionProfileParams();
 				if (SSSType == SSS_TYPE_DIFFUSION_PROFILE)
 				{
@@ -323,9 +323,12 @@ void Main(
 				}
 				else if (SSSType == SSS_TYPE_DIFFUSION)
 				{
-					TransmissionParams = InitTransmissionProfileParams(); // IOR and phase are incorrect. but not used
-					TransmissionParams.ExtinctionScale = StrataSubSurfaceHeaderGetMFP(SSSHeader);
+					//For per pixel diffusion, use default transmission parameters (only NormalScale & ExtinctionScale are used)
+					const float MFP = dot(StrataSubSurfaceHeaderGetMFP(SSSHeader), (1.0/3.0).xxx);
+					const float Extinction = 1.0f / max(0.00001, MFP);
+					TransmissionParams.ExtinctionScale = Extinction;
 				}
+
 			#else // STRATA_ENABLED
 				float3 WorldNormal = GBufferData.WorldNormal;
 				FTransmissionProfileParams TransmissionParams = GetTransmissionProfileParams(ExtractSubsurfaceProfileInt(GBufferData));
@@ -335,6 +338,7 @@ void Main(
 			else
 			#endif
 			{
+				// This branch is taken by Subsurface shading model (Wrap for Substrate)
 				// ideally we use a larger filter kernel for SSSbut as Gather4 makes that harder
 				SSSTransmission = ManualPCF(ShadowPosition.xy, Settings);
 			}
@@ -482,10 +486,23 @@ void MainOnePassPointLightPS(
 	const FStrataTopLayerData TopLayerData = StrataUnpackTopLayerData(Strata.TopLayerTexture.Load(uint3(SVPos.xy, 0)));
 
 	const FStrataSubsurfaceHeader SSSHeader = StrataLoadSubsurfaceHeader(Strata.MaterialTextureArray, Strata.FirstSliceStoringStrataSSSData, SVPos.xy);
+	const uint SSSType = StrataSubSurfaceHeaderGetSSSType(SSSHeader);
 
 	const bool IsSubsurfaceProfile = StrataSubSurfaceHeaderGetUseDiffusion(SSSHeader);
 	const float3 WorldNormal = TopLayerData.WorldNormal;
-	const FTransmissionProfileParams TransmissionParams = GetTransmissionProfileParams(StrataSubSurfaceHeaderGetProfileId(SSSHeader));
+
+	FTransmissionProfileParams TransmissionParams = InitTransmissionProfileParams();
+	if (SSSType == SSS_TYPE_DIFFUSION_PROFILE)
+	{
+		TransmissionParams = GetTransmissionProfileParams(StrataSubSurfaceHeaderGetProfileId(SSSHeader));
+	}
+	else if (SSSType == SSS_TYPE_DIFFUSION)
+	{
+		//For per pixel diffusion, use default transmission parameters (only NormalScale & ExtinctionScale are used)
+		const float MFP = dot(StrataSubSurfaceHeaderGetMFP(SSSHeader), (1.0 / 3.0).xxx);
+		const float Extinction = 1.0f / max(0.00001, MFP);
+		TransmissionParams.ExtinctionScale = Extinction;
+	}
 
 #else // STRATA_ENABLED
 
diff --git a/Engine/Shaders/Private/Strata/Strata.ush b/Engine/Shaders/Private/Strata/Strata.ush
--- a/Engine/Shaders/Private/Strata/Strata.ush
+++ b/Engine/Shaders/Private/Strata/Strata.ush
@@ -8,6 +8,25 @@
 	#error STRATA_ENABLED needs to be defined
 #endif
 
+
+// A structure that can be used to transfert Legacy/substrate data for light transmission.
+struct FSubsurfaceOpacityMFP
+{
+	bool  bDataIsOpacity;	// If false, Data is a mean free path
+	float Data;				// MFP or Opacity
+	float Density;			// Must be stored separately from Opacity, because Opacity is used for some tests
+};
+
+FSubsurfaceOpacityMFP GetInitialisedSubsurfaceOpacityMFP()
+{
+	FSubsurfaceOpacityMFP SubsurfaceOpacityMFP;
+	SubsurfaceOpacityMFP.bDataIsOpacity = true;
+	SubsurfaceOpacityMFP.Data = 1.0;
+	SubsurfaceOpacityMFP.Density = 0.0f;
+	return SubsurfaceOpacityMFP;
+}
+
+
 #if STRATA_ENABLED
 
 #include "../ShadingEnergyConservation.ush"
@@ -929,6 +948,11 @@ bool StrataSubSurfaceHeaderGetIsProfile(in FStrataSubsurfaceHeader SSSHeader)
 	return SSSHEADER_TYPE(SSSHeader) == SSS_TYPE_DIFFUSION_PROFILE;
 }
 
+bool StrataSubSurfaceHeaderGetIsWrap(in FStrataSubsurfaceHeader SSSHeader)
+{
+	return SSSHEADER_TYPE(SSSHeader) == SSS_TYPE_WRAP;
+}
+
 uint StrataSubSurfaceHeaderGetSSSType(in FStrataSubsurfaceHeader SSSHeader)
 {
 	return SSSHEADER_TYPE(SSSHeader);
@@ -975,20 +999,34 @@ float StrataSubSurfaceHeaderGetOpacity(in FStrataSubsurfaceHeader SSSHeader)
 }
 
 // Legacy conversion function, which translated Strata sub-surface data into 'opacity' for shadow transmission purpose. 
-// Note: When SSS data uses DFMP data, the conversion to opacity is ad-hoc/hacky
-// STRATA_TODO: merge this function with StrataSubSurfaceHeaderGetOpacity()
-float StrataGetOpacityFromSubSurface(FStrataSubsurfaceHeader SSSHeader)
+// STRATA_TODO: merge this function with StrataSubSurfaceHeaderGetOpacity() and GetSubsurfaceOpacityFromGbuffer() in a common Strata header.
+FSubsurfaceOpacityMFP StrataGetOpacityFromSubSurface(FStrataSubsurfaceHeader SSSHeader)
 {
+	FSubsurfaceOpacityMFP SubsurfaceOpacityMFP = GetInitialisedSubsurfaceOpacityMFP();
+
 	const uint SSSType = SSSHEADER_TYPE(SSSHeader);
-	const float ProfileRadiusScale = StrataSubSurfaceHeaderGetProfileRadiusScale(SSSHeader);
-	const float3 MFP = StrataSubSurfaceHeaderGetMFP(SSSHeader);
-	const float WrapOpacity = StrataSubSurfaceHeaderGetWrapOpacity(SSSHeader);
 
-	float Opacity = 1.0f;
-	Opacity = SSSType == SSS_TYPE_DIFFUSION ? saturate(1.0f - max3(MFP.x, MFP.y, MFP.z)) : Opacity;
-	Opacity = SSSType == SSS_TYPE_DIFFUSION_PROFILE ? ProfileRadiusScale : Opacity;
-	Opacity = SSSType == SSS_TYPE_WRAP ? WrapOpacity : Opacity;
-	return Opacity;
+	if (SSSType == SSS_TYPE_DIFFUSION)
+	{
+		SubsurfaceOpacityMFP.bDataIsOpacity = false;
+		const float3 MFP = StrataSubSurfaceHeaderGetMFP(SSSHeader);
+		SubsurfaceOpacityMFP.Data = dot(MFP, (1.0f/3.0f).xxx);
+		SubsurfaceOpacityMFP.Density = 0; // Not fed
+	}
+	else if (SSSType == SSS_TYPE_DIFFUSION_PROFILE)
+	{
+		SubsurfaceOpacityMFP.bDataIsOpacity = true;
+		SubsurfaceOpacityMFP.Data = min(StrataSubSurfaceHeaderGetProfileRadiusScale(SSSHeader), 0.99);
+		SubsurfaceOpacityMFP.Density = 0; // Not fed
+	}
+	else if (SSSType == SSS_TYPE_WRAP)
+	{
+		SubsurfaceOpacityMFP.bDataIsOpacity = true;
+		SubsurfaceOpacityMFP.Data = min(StrataSubSurfaceHeaderGetWrapOpacity(SSSHeader), 0.99);
+		SubsurfaceOpacityMFP.Density = 0; // Not fed
+	}
+
+	return SubsurfaceOpacityMFP;
 }
 
 // Stored in second slice as UINT
diff --git a/Engine/Shaders/Private/VirtualShadowMaps/VirtualShadowMapProjection.usf b/Engine/Shaders/Private/VirtualShadowMaps/VirtualShadowMapProjection.usf
--- a/Engine/Shaders/Private/VirtualShadowMaps/VirtualShadowMapProjection.usf
+++ b/Engine/Shaders/Private/VirtualShadowMaps/VirtualShadowMapProjection.usf
@@ -56,11 +56,16 @@ FProjectionShadingInfo GetProjectionShadingInfo(uint2 PixelPos)
 	const FStrataSubsurfaceHeader SSSHeader = StrataLoadSubsurfaceHeader(Strata.MaterialTextureArray, Strata.FirstSliceStoringStrataSSSData, PixelPos);
 	const bool bIsValid = StrataSubSurfaceHeaderGetIsValid(SSSHeader);
 	const bool bIsProfile = StrataSubSurfaceHeaderGetIsProfile(SSSHeader);
-	Out.bIsSubsurface = bIsValid && !bIsProfile;
+
+	const uint SSSType = SSSHEADER_TYPE(SSSHeader);
+	const bool bHasSSS = StrataPixelHeader.HasSubsurface();
+	const uint BSDFType = StrataPixelHeader.StrataGetBSDFType();
 
 	Out.bIsValid = StrataPixelHeader.IsStrataMaterial();
-	Out.bIsHair = StrataPixelHeader.StrataGetBSDFType() == STRATA_BSDF_TYPE_HAIR;
+	Out.bIsHair = BSDFType == STRATA_BSDF_TYPE_HAIR;
 	Out.WorldNormal = TopLayerData.WorldNormal;
+	Out.bIsSubsurface = SSSType != SSS_TYPE_INVALID; // legacy IsSubsurfaceModel is true for subsurface profile, eyes, etc.
+
 #else
 	const FGBufferData GBufferData = GetGBufferDataUint(PixelPos, true);
 	Out.bIsValid = GBufferData.ShadingModelID != SHADINGMODELID_UNLIT;
@@ -205,7 +210,7 @@ FVirtualShadowMapSampleResult ProjectLight(
 	float ScreenRayLengthWorld,
 	float3 TranslatedWorldPosition,
 	const float Noise,
-	float SubsurfaceOpacity)
+	const FSubsurfaceOpacityMFP SubsurfaceOpacityMFP)
 {
 	const bool bIsHairInput = InputType == INPUT_TYPE_HAIRSTRANDS;
 
@@ -304,15 +309,18 @@ FVirtualShadowMapSampleResult ProjectLight(
 				SMRTRayOffset,
 				WorldNormal);
 		}
-
-		if (SubsurfaceOpacity < 1.0f)
+		
+		// Explanation:
+		//	- If bDataIsOpacity = 1: Data stores the surface's SSS opacity. We compute a transmitted color when the surface is not fully opaque (opacity == 1).
+		//	- If bDataIsOpacity = 0: Data stores the surface's MFP.         We compute a transmitted color when the surface is not fully opaque (MFP == 0)
+		if ((SubsurfaceOpacityMFP.bDataIsOpacity && SubsurfaceOpacityMFP.Data < 1.0f) || (!SubsurfaceOpacityMFP.bDataIsOpacity && SubsurfaceOpacityMFP.Data > 0.0f))
 		{
 			// We use a single term for both the shadow and "transmission shadow" in the VSM path to provide a
 			// higher quality filtered result on back faces and avoid artifacts at the point where normals flip.
 			// This is especially important for foliage as normals are often "warped" to provide softer shading
 			// and do not represent the underlying geometry accurately, so having a continuous shadow function
 			// is important.
-			Result.ShadowFactor = GetShadowFactorSubsurface(Result.ShadowFactor, Result.OccluderDistance, SubsurfaceOpacity);
+			Result.ShadowFactor = GetShadowFactorSubsurface(Result.ShadowFactor, Result.OccluderDistance, SubsurfaceOpacityMFP);
 		}
 
 		//Result.GeneralDebug = GreenToRedTurbo(1.0f - (SMRTRayOffset / ScreenRayLengthWorld));
@@ -593,7 +601,12 @@ void VirtualShadowMapProjection(
 	const float Noise = InterleavedGradientNoise( SvPosition.xy, View.StateFrameIndexMod8 );
 
 	const FProjectionShadingInfo ShadingInfo = GetProjectionShadingInfo(PixelPos);
-	const float SubsurfaceOpacity = bIsHairInput ? 1.0f : GetSubsurfaceOpacityFromGbuffer(PixelPos);
+
+	FSubsurfaceOpacityMFP SubsurfaceOpacityMFP = GetInitialisedSubsurfaceOpacityMFP();
+	if (!bIsHairInput)
+	{
+		SubsurfaceOpacityMFP = GetSubsurfaceOpacityFromGbuffer(PixelPos);
+	}
 
 	FVirtualShadowMapSampleResult VisualizeResult = InitVirtualShadowMapSampleResult();
 
@@ -626,7 +639,7 @@ void VirtualShadowMapProjection(
 				ScreenRayLengthWorld,
 				TranslatedWorldPosition,
 				Noise,
-				SubsurfaceOpacity);
+				SubsurfaceOpacityMFP);
 
 			PackShadowMask(ShadowMask, Result.ShadowFactor, Index);
 
@@ -644,13 +657,13 @@ void VirtualShadowMapProjection(
 		int VirtualShadowMapId = LightUniformVirtualShadowMapId;
 		FLightShaderParameters Light = GetRootLightShaderParameters();
 
-		// Subsurface approximation
-		const float SubsurfaceOpacity = bIsHairInput ? 1.0f : GetSubsurfaceOpacityFromGbuffer(PixelPos);
 #if DIRECTIONAL_LIGHT
 		// Increases light source radius for subsurface material as the opacity reduces to emulates light diffusion.
 		// TODO: Does something similar make sense for local lights?
-		if (SubsurfaceOpacity < 1.0f)
+		// Substrate: this approximation has no concrete fundation so it is currently skipped for Substrate using PerPixel SSS, e.g. when bDataIsOpacity==false and data is MFP.
+		if (SubsurfaceOpacityMFP.bDataIsOpacity && SubsurfaceOpacityMFP.Data < 1.0f)
 		{
+			const float SubsurfaceOpacity = SubsurfaceOpacityMFP.Data;
 			Light.SourceRadius = max(Light.SourceRadius, (1.0f - SubsurfaceOpacity) * SubsurfaceMinSourceRadius);
 		}
 #endif
@@ -664,7 +677,7 @@ void VirtualShadowMapProjection(
 			ScreenRayLengthWorld,
 			TranslatedWorldPosition,
 			Noise,
-			SubsurfaceOpacity);
+			SubsurfaceOpacityMFP);
 
 		OutShadowFactor[ PixelPos ] = Result.ShadowFactor.xx;
 
diff --git a/Engine/Shaders/Private/VirtualShadowMaps/VirtualShadowMapTransmissionCommon.ush b/Engine/Shaders/Private/VirtualShadowMaps/VirtualShadowMapTransmissionCommon.ush
--- a/Engine/Shaders/Private/VirtualShadowMaps/VirtualShadowMapTransmissionCommon.ush
+++ b/Engine/Shaders/Private/VirtualShadowMaps/VirtualShadowMapTransmissionCommon.ush
@@ -11,12 +11,53 @@ VirtualShadowMaps/VirtualShadowMapTransmissionCommon.ush:
 #include "../TransmissionCommon.ush"
 #include "VirtualShadowMapProjectionCommon.ush"
 
-float GetSubsurfaceOpacityFromGbuffer(uint2 PixelPos)
+FSubsurfaceOpacityMFP GetSubsurfaceOpacityFromGbuffer(uint2 PixelPos)
 {
-	float Opacity = 1.0;
+	FSubsurfaceOpacityMFP SubsurfaceOpacityMFP = GetInitialisedSubsurfaceOpacityMFP();
 
 #if STRATA_ENABLED
-	Opacity = StrataGetOpacityFromSubSurface(StrataLoadSubsurfaceHeader(Strata.MaterialTextureArray, Strata.FirstSliceStoringStrataSSSData, PixelPos));
+
+	FStrataSubsurfaceData LoadStrat = StrataLoadSubsurfaceData(Strata.MaterialTextureArray, Strata.FirstSliceStoringStrataSSSData, PixelPos);
+
+	const uint SSSType = SSSHEADER_TYPE(LoadStrat.Header);
+	
+	if (SSSType == SSS_TYPE_DIFFUSION)
+	{
+		SubsurfaceOpacityMFP.bDataIsOpacity = false;
+		const float3 MFP = StrataSubSurfaceHeaderGetMFP(LoadStrat.Header);
+		SubsurfaceOpacityMFP.Data = dot(MFP, (1.0f / 3.0f).xxx);
+		SubsurfaceOpacityMFP.Density = 0.0f;
+	}
+	else if (SSSType == SSS_TYPE_DIFFUSION_PROFILE)
+	{
+		SubsurfaceOpacityMFP.bDataIsOpacity = true;
+ 	   
+		const float ProfileRadiusScale = StrataSubSurfaceHeaderGetProfileRadiusScale(LoadStrat.Header);
+		// This clamp aligns with SubsurfaceDensityFromOpacity
+		// Various engine paths treat these subsurface materials differently
+		// even when they have Opacity = 1 in the material shader, so this is
+		// important to avoid things like backface transmission being shadowed by
+		// contact shadows and so on.
+		const float Opacity = min(ProfileRadiusScale, 0.99);
+		SubsurfaceOpacityMFP.Data = Opacity;
+
+		FTransmissionProfileParams TransmissionParams = GetTransmissionProfileParams(StrataSubSurfaceHeaderGetProfileId(LoadStrat.Header));
+		SubsurfaceOpacityMFP.Density = SubsurfaceDensityFromOpacity(SubsurfaceOpacityMFP.Data) * TransmissionParams.ExtinctionScale * 3.1f;	// This matches the weird computations done in CalcTransmissionThickness
+	}
+	else if (SSSType == SSS_TYPE_WRAP)
+	{
+		SubsurfaceOpacityMFP.bDataIsOpacity = true;
+ 	   
+		// This clamp aligns with SubsurfaceDensityFromOpacity
+		// Various engine paths treat these subsurface materials differently
+		// even when they have Opacity = 1 in the material shader, so this is
+		// important to avoid things like backface transmission being shadowed by
+		// contact shadows and so on.
+		const float Opacity = min(StrataSubSurfaceHeaderGetWrapOpacity(LoadStrat.Header), 0.99);
+		SubsurfaceOpacityMFP.Data = Opacity;
+		SubsurfaceOpacityMFP.Density = SubsurfaceDensityFromOpacity(SubsurfaceOpacityMFP.Data);
+	}
+
 #else // STRATA_ENABLED
 	FGBufferData GBufferData = GetGBufferDataUint(PixelPos);
 
@@ -30,19 +71,36 @@ float GetSubsurfaceOpacityFromGbuffer(uint2 PixelPos)
 		// even when they have Opacity = 1 in the material shader, so this is
 		// important to avoid things like backface transmission being shadowed by
 		// contact shadows and so on.
-		Opacity = min(GBufferData.CustomData.a, 0.99f);
+		SubsurfaceOpacityMFP.Data = min(GBufferData.CustomData.a, 0.99f);
+		SubsurfaceOpacityMFP.Density = SubsurfaceDensityFromOpacity(SubsurfaceOpacityMFP.Data);
 	}
 #endif // STRATA_ENABLED
 
-	return Opacity;
+	return SubsurfaceOpacityMFP;
 }
 
 /* Compute Subsurface transmission */
-float GetShadowFactorSubsurface(float ShadowFactor, float OccluderDistance, float SubsurfaceOpacity)
+float GetShadowFactorSubsurface(float ShadowFactor, float OccluderDistance, FSubsurfaceOpacityMFP SubsurfaceOpacityMFP)
 {
 	if (ShadowFactor < 1.0f)
 	{
-		float SSSTransmission = GetSubSurfaceTransmission(OccluderDistance, SubsurfaceDensityFromOpacity(SubsurfaceOpacity));
+#if STRATA_ENABLED
+		float SSSTransmission = 1.0f;
+		if (SubsurfaceOpacityMFP.bDataIsOpacity)
+		{
+			// Opacity from subsurface profile and subsurface legacy shading model
+			SSSTransmission = GetSubSurfaceTransmission(OccluderDistance, SubsurfaceOpacityMFP.Density);
+		}
+		else
+		{
+			const float MFP = SubsurfaceOpacityMFP.Data;
+			const float Extinction = 1.0f / max(0.00001, MFP);
+			SSSTransmission = saturate(FastExp(-OccluderDistance * Extinction));
+		}
+#else
+		// Only opacity in this path
+		float SSSTransmission = GetSubSurfaceTransmission(OccluderDistance, SubsurfaceOpacityMFP.Density);
+#endif
 
 		// NOTE: This 'square' is to compensate the ShadowMask encoding (i.e., sqrt) applied when compositing VSM to shadow mask. 
 		//       This allows to be consistent with the PCF path
