Moved Translucency Volume sampling functions to TranslucencyVolumeCommon.ush

diff --git a/Engine/Shaders/Private/BasePassCommon.ush b/Engine/Shaders/Private/BasePassCommon.ush
--- a/Engine/Shaders/Private/BasePassCommon.ush
+++ b/Engine/Shaders/Private/BasePassCommon.ush
@@ -107,45 +107,14 @@ struct FSharedBasePassInterpolants
 
 #define FBasePassInterpolantsVSToPS FSharedBasePassInterpolants
 
-#define SharedAmbientInnerSampler View.SharedBilinearClampedSampler
-#define SharedAmbientOuterSampler View.SharedBilinearClampedSampler
-#define SharedDirectionalInnerSampler View.SharedBilinearClampedSampler
-#define SharedDirectionalOuterSampler View.SharedBilinearClampedSampler
-
-void ComputeVolumeUVs(float3 TranslatedWorldPosition, float3 LightingPositionOffset, out float3 InnerVolumeUVs, out float3 OuterVolumeUVs, out float FinalLerpFactor)
-{
-	// Apply a stable offset to the world position used for lighting, which breaks up artifacts from using a low res volume texture
-	InnerVolumeUVs = (TranslatedWorldPosition + LightingPositionOffset - View.TranslucencyLightingVolumeMin[0].xyz) * View.TranslucencyLightingVolumeInvSize[0].xyz;
-	OuterVolumeUVs = (TranslatedWorldPosition + LightingPositionOffset - View.TranslucencyLightingVolumeMin[1].xyz) * View.TranslucencyLightingVolumeInvSize[1].xyz;
-
-	// Controls how fast the lerp between the inner and outer cascades happens
-	// Larger values result in a shorter transition distance
-	float TransitionScale = 6;
-	// Setup a 3d lerp factor going to 0 at the edge of the inner volume
-	float3 LerpFactors = saturate((.5f - abs(InnerVolumeUVs - .5f)) * TransitionScale);
-	FinalLerpFactor = LerpFactors.x * LerpFactors.y * LerpFactors.z;
-}
-
-float4 GetAmbientLightingVectorFromTranslucentLightingVolume(float3 InnerVolumeUVs, float3 OuterVolumeUVs, float FinalLerpFactor)
-{
-	// Lookup the inner and outer cascade ambient lighting values
-	float4 InnerLighting = Texture3DSampleLevel(TranslucentBasePass.TranslucencyLightingVolumeAmbientInner, SharedAmbientInnerSampler, InnerVolumeUVs, 0);
-	float4 OuterLighting = Texture3DSampleLevel(TranslucentBasePass.TranslucencyLightingVolumeAmbientOuter, SharedAmbientOuterSampler, OuterVolumeUVs, 0);
+#endif
 
-	// Lerp between cascades
-	return lerp(OuterLighting, InnerLighting, FinalLerpFactor);
-}
+#define TranslucencyLightingVolumeAmbientInner TranslucentBasePass.TranslucencyLightingVolumeAmbientInner
+#define TranslucencyLightingVolumeAmbientOuter TranslucentBasePass.TranslucencyLightingVolumeAmbientOuter
+#define TranslucencyLightingVolumeDirectionalInner TranslucentBasePass.TranslucencyLightingVolumeDirectionalInner
+#define TranslucencyLightingVolumeDirectionalOuter TranslucentBasePass.TranslucencyLightingVolumeDirectionalOuter
 
-float3 GetDirectionalLightingVectorFromTranslucentLightingVolume(float3 InnerVolumeUVs, float3 OuterVolumeUVs, float FinalLerpFactor)
-{
-	// Fetch both the ambient and directional values for both cascades
-	float3 InnerVector1 = Texture3DSampleLevel(TranslucentBasePass.TranslucencyLightingVolumeDirectionalInner, SharedDirectionalInnerSampler, InnerVolumeUVs, 0).rgb;
-	float3 OuterVector1 = Texture3DSampleLevel(TranslucentBasePass.TranslucencyLightingVolumeDirectionalOuter, SharedDirectionalOuterSampler, OuterVolumeUVs, 0).rgb;
-
-	// Lerp between cascades
-	return lerp(OuterVector1, InnerVector1, FinalLerpFactor);
-}
-#endif
+#include "TranslucencyVolumeCommon.ush"
 
 #include "SphericalGaussian.ush"
 
@@ -284,43 +253,3 @@ FShadingOcclusion ApplyBentNormal(
 #endif
 	return Out;
 }
-
-#include "SHCommon.ush"
-
-void GetVolumeLightingNonDirectional(float4 AmbientLightingVector, float3 DiffuseColor, inout float3 InterpolatedLighting, out float4 VolumeLighting)
-{
-	// Normal is not taken into account with non directional lighting, and only the ambient term of the SH coefficients are needed
-	FOneBandSHVectorRGB TranslucentLighting;
-	TranslucentLighting.R.V.x = AmbientLightingVector.r;
-	TranslucentLighting.G.V.x = AmbientLightingVector.g;
-	TranslucentLighting.B.V.x = AmbientLightingVector.b;
-
-	FOneBandSHVector DiffuseTransferSH = CalcDiffuseTransferSH1(1);
-	VolumeLighting = float4(DotSH1(TranslucentLighting, DiffuseTransferSH), AmbientLightingVector.a);
-	InterpolatedLighting = DiffuseColor * VolumeLighting.rgb;
-}
-
-void GetVolumeLightingDirectional(float4 AmbientLightingVector, float3 DirectionalLightingVector, float3 WorldNormal, float3 DiffuseColor, inout float3 InterpolatedLighting, out float4 VolumeLighting)
-{
-	float DirectionalLightingIntensity = GetMaterialTranslucencyDirectionalLightingIntensity();
-
-	AmbientLightingVector.rgb /= DirectionalLightingIntensity;
-	DirectionalLightingVector.rgb *= DirectionalLightingIntensity;
-
-	// Reconstruct the SH coefficients based on what was encoded
-	FTwoBandSHVectorRGB TranslucentLighting;
-	TranslucentLighting.R.V.x = AmbientLightingVector.r;
-	TranslucentLighting.G.V.x = AmbientLightingVector.g;
-	TranslucentLighting.B.V.x = AmbientLightingVector.b;
-	float3 NormalizedAmbientColor = AmbientLightingVector.rgb / (Luminance(AmbientLightingVector.rgb) + 0.00001f);
-
-	// Scale the monocrome directional coefficients with the normalzed ambient color as an approximation to the uncompressed values
-	TranslucentLighting.R.V.yzw = DirectionalLightingVector.rgb * NormalizedAmbientColor.r;
-	TranslucentLighting.G.V.yzw = DirectionalLightingVector.rgb * NormalizedAmbientColor.g;
-	TranslucentLighting.B.V.yzw = DirectionalLightingVector.rgb * NormalizedAmbientColor.b;
-
-	// Compute diffuse lighting which takes the normal into account
-	FTwoBandSHVector DiffuseTransferSH = CalcDiffuseTransferSH(WorldNormal, 1);
-	VolumeLighting = float4(max(half3(0, 0, 0), DotSH(TranslucentLighting, DiffuseTransferSH)), AmbientLightingVector.a);
-	InterpolatedLighting += DiffuseColor * VolumeLighting.rgb;
-}
\ No newline at end of file
diff --git a/Engine/Shaders/Private/PostProcessEyeAdaptation.usf b/Engine/Shaders/Private/PostProcessEyeAdaptation.usf
--- a/Engine/Shaders/Private/PostProcessEyeAdaptation.usf
+++ b/Engine/Shaders/Private/PostProcessEyeAdaptation.usf
@@ -391,72 +391,7 @@ Texture3D TranslucencyLightingVolumeAmbientOuter;
 Texture3D TranslucencyLightingVolumeDirectionalInner;
 Texture3D TranslucencyLightingVolumeDirectionalOuter;
 
-#define SharedAmbientInnerSampler View.SharedBilinearClampedSampler
-#define SharedAmbientOuterSampler View.SharedBilinearClampedSampler
-#define SharedDirectionalInnerSampler View.SharedBilinearClampedSampler
-#define SharedDirectionalOuterSampler View.SharedBilinearClampedSampler
-
-// TODO: Move these functions to header file shared with BasePass/Strata
-
-void ComputeVolumeUVs(float3 TranslatedWorldPosition, float3 LightingPositionOffset, out float3 InnerVolumeUVs, out float3 OuterVolumeUVs, out float FinalLerpFactor)
-{
-	// Apply a stable offset to the world position used for lighting, which breaks up artifacts from using a low res volume texture
-	InnerVolumeUVs = (TranslatedWorldPosition + LightingPositionOffset - View.TranslucencyLightingVolumeMin[0].xyz) * View.TranslucencyLightingVolumeInvSize[0].xyz;
-	OuterVolumeUVs = (TranslatedWorldPosition + LightingPositionOffset - View.TranslucencyLightingVolumeMin[1].xyz) * View.TranslucencyLightingVolumeInvSize[1].xyz;
-
-	// Controls how fast the lerp between the inner and outer cascades happens
-	// Larger values result in a shorter transition distance
-	float TransitionScale = 6;
-	// Setup a 3d lerp factor going to 0 at the edge of the inner volume
-	float3 LerpFactors = saturate((.5f - abs(InnerVolumeUVs - .5f)) * TransitionScale);
-	FinalLerpFactor = LerpFactors.x * LerpFactors.y * LerpFactors.z;
-}
-
-float4 GetAmbientLightingVectorFromTranslucentLightingVolume(float3 InnerVolumeUVs, float3 OuterVolumeUVs, float FinalLerpFactor)
-{
-	// Lookup the inner and outer cascade ambient lighting values
-	float4 InnerLighting = Texture3DSampleLevel(TranslucencyLightingVolumeAmbientInner, SharedAmbientInnerSampler, InnerVolumeUVs, 0);
-	float4 OuterLighting = Texture3DSampleLevel(TranslucencyLightingVolumeAmbientOuter, SharedAmbientOuterSampler, OuterVolumeUVs, 0);
-
-	// Lerp between cascades
-	return lerp(OuterLighting, InnerLighting, FinalLerpFactor);
-}
-
-float3 GetDirectionalLightingVectorFromTranslucentLightingVolume(float3 InnerVolumeUVs, float3 OuterVolumeUVs, float FinalLerpFactor)
-{
-	// Fetch both the ambient and directional values for both cascades
-	float3 InnerVector1 = Texture3DSampleLevel(TranslucencyLightingVolumeDirectionalInner, SharedDirectionalInnerSampler, InnerVolumeUVs, 0).rgb;
-	float3 OuterVector1 = Texture3DSampleLevel(TranslucencyLightingVolumeDirectionalOuter, SharedDirectionalOuterSampler, OuterVolumeUVs, 0).rgb;
-
-	// Lerp between cascades
-	return lerp(OuterVector1, InnerVector1, FinalLerpFactor);
-}
-
-float4 GetVolumeLightingDirectional(float4 AmbientLightingVector, float3 DirectionalLightingVector, float3 WorldNormal)
-{
-	//const float DirectionalLightingIntensity = GetMaterialTranslucencyDirectionalLightingIntensity();
-	const float DirectionalLightingIntensity = 1.0f;
-
-	AmbientLightingVector.rgb /= DirectionalLightingIntensity;
-	DirectionalLightingVector.rgb *= DirectionalLightingIntensity;
-
-	// Reconstruct the SH coefficients based on what was encoded
-	FTwoBandSHVectorRGB TranslucentLighting;
-	TranslucentLighting.R.V.x = AmbientLightingVector.r;
-	TranslucentLighting.G.V.x = AmbientLightingVector.g;
-	TranslucentLighting.B.V.x = AmbientLightingVector.b;
-	float3 NormalizedAmbientColor = AmbientLightingVector.rgb / (Luminance(AmbientLightingVector.rgb) + 0.00001f);
-
-	// Scale the monocrome directional coefficients with the normalzed ambient color as an approximation to the uncompressed values
-	TranslucentLighting.R.V.yzw = DirectionalLightingVector.rgb * NormalizedAmbientColor.r;
-	TranslucentLighting.G.V.yzw = DirectionalLightingVector.rgb * NormalizedAmbientColor.g;
-	TranslucentLighting.B.V.yzw = DirectionalLightingVector.rgb * NormalizedAmbientColor.b;
-
-	// Compute diffuse lighting which takes the normal into account
-	FTwoBandSHVector DiffuseTransferSH = CalcDiffuseTransferSH(WorldNormal, 1);
-	
-	return float4(max(half3(0, 0, 0), DotSH(TranslucentLighting, DiffuseTransferSH)), AmbientLightingVector.a);
-}
+#include "TranslucencyVolumeCommon.ush"
 
 struct FExposureMaterial
 {
@@ -614,11 +549,11 @@ void CalculateExposureIlluminance(uint2 ThreadId : SV_DispatchThreadID)
 			float4 AmbientLightingVector = GetAmbientLightingVectorFromTranslucentLightingVolume(InnerVolumeUVs, OuterVolumeUVs, FinalLerpFactor);
 			float3 DirectionalLightingVector = GetDirectionalLightingVectorFromTranslucentLightingVolume(InnerVolumeUVs, OuterVolumeUVs, FinalLerpFactor);
 			
-			float3 VolumeLighting = GetVolumeLightingDirectional(AmbientLightingVector, DirectionalLightingVector, Material.WorldNormal).rgb;
+			float3 VolumeLighting = GetVolumeLightingDirectional(AmbientLightingVector, DirectionalLightingVector, Material.WorldNormal, 1.0f).rgb;
 			
 			if (Material.bHasBackfaceLighting)
 			{
-				VolumeLighting += GetVolumeLightingDirectional(AmbientLightingVector, DirectionalLightingVector, -Material.WorldNormal).rgb;
+				VolumeLighting += GetVolumeLightingDirectional(AmbientLightingVector, DirectionalLightingVector, -Material.WorldNormal, 1.0f).rgb;
 			}
 
 		#if PROJECT_SUPPORTS_LUMEN && !FORWARD_SHADING
diff --git a/Engine/Shaders/Private/TranslucencyVolumeCommon.ush b/Engine/Shaders/Private/TranslucencyVolumeCommon.ush
new file mode 100644
--- /dev/null
+++ b/Engine/Shaders/Private/TranslucencyVolumeCommon.ush
@@ -0,0 +1,91 @@
+// Copyright Epic Games, Inc. All Rights Reserved.
+
+#pragma once
+
+#define SharedAmbientInnerSampler View.SharedBilinearClampedSampler
+#define SharedAmbientOuterSampler View.SharedBilinearClampedSampler
+#define SharedDirectionalInnerSampler View.SharedBilinearClampedSampler
+#define SharedDirectionalOuterSampler View.SharedBilinearClampedSampler
+
+void ComputeVolumeUVs(float3 TranslatedWorldPosition, float3 LightingPositionOffset, out float3 InnerVolumeUVs, out float3 OuterVolumeUVs, out float FinalLerpFactor)
+{
+	// Apply a stable offset to the world position used for lighting, which breaks up artifacts from using a low res volume texture
+	InnerVolumeUVs = (TranslatedWorldPosition + LightingPositionOffset - View.TranslucencyLightingVolumeMin[0].xyz) * View.TranslucencyLightingVolumeInvSize[0].xyz;
+	OuterVolumeUVs = (TranslatedWorldPosition + LightingPositionOffset - View.TranslucencyLightingVolumeMin[1].xyz) * View.TranslucencyLightingVolumeInvSize[1].xyz;
+
+	// Controls how fast the lerp between the inner and outer cascades happens
+	// Larger values result in a shorter transition distance
+	float TransitionScale = 6;
+	// Setup a 3d lerp factor going to 0 at the edge of the inner volume
+	float3 LerpFactors = saturate((.5f - abs(InnerVolumeUVs - .5f)) * TransitionScale);
+	FinalLerpFactor = LerpFactors.x * LerpFactors.y * LerpFactors.z;
+}
+
+float4 GetAmbientLightingVectorFromTranslucentLightingVolume(float3 InnerVolumeUVs, float3 OuterVolumeUVs, float FinalLerpFactor)
+{
+	// Lookup the inner and outer cascade ambient lighting values
+	float4 InnerLighting = Texture3DSampleLevel(TranslucencyLightingVolumeAmbientInner, SharedAmbientInnerSampler, InnerVolumeUVs, 0);
+	float4 OuterLighting = Texture3DSampleLevel(TranslucencyLightingVolumeAmbientOuter, SharedAmbientOuterSampler, OuterVolumeUVs, 0);
+
+	// Lerp between cascades
+	return lerp(OuterLighting, InnerLighting, FinalLerpFactor);
+}
+
+float3 GetDirectionalLightingVectorFromTranslucentLightingVolume(float3 InnerVolumeUVs, float3 OuterVolumeUVs, float FinalLerpFactor)
+{
+	// Fetch both the ambient and directional values for both cascades
+	float3 InnerVector1 = Texture3DSampleLevel(TranslucencyLightingVolumeDirectionalInner, SharedDirectionalInnerSampler, InnerVolumeUVs, 0).rgb;
+	float3 OuterVector1 = Texture3DSampleLevel(TranslucencyLightingVolumeDirectionalOuter, SharedDirectionalOuterSampler, OuterVolumeUVs, 0).rgb;
+
+	// Lerp between cascades
+	return lerp(OuterVector1, InnerVector1, FinalLerpFactor);
+}
+
+#include "SHCommon.ush"
+
+void GetVolumeLightingNonDirectional(float4 AmbientLightingVector, float3 DiffuseColor, inout float3 InterpolatedLighting, out float4 VolumeLighting)
+{
+	// Normal is not taken into account with non directional lighting, and only the ambient term of the SH coefficients are needed
+	FOneBandSHVectorRGB TranslucentLighting;
+	TranslucentLighting.R.V.x = AmbientLightingVector.r;
+	TranslucentLighting.G.V.x = AmbientLightingVector.g;
+	TranslucentLighting.B.V.x = AmbientLightingVector.b;
+
+	FOneBandSHVector DiffuseTransferSH = CalcDiffuseTransferSH1(1);
+	VolumeLighting = float4(DotSH1(TranslucentLighting, DiffuseTransferSH), AmbientLightingVector.a);
+	InterpolatedLighting = DiffuseColor * VolumeLighting.rgb;
+}
+
+void GetVolumeLightingDirectional(float4 AmbientLightingVector, float3 DirectionalLightingVector, float3 WorldNormal, float3 DiffuseColor, float DirectionalLightingIntensity, inout float3 InterpolatedLighting, out float4 VolumeLighting)
+{
+	AmbientLightingVector.rgb /= DirectionalLightingIntensity;
+	DirectionalLightingVector.rgb *= DirectionalLightingIntensity;
+
+	// Reconstruct the SH coefficients based on what was encoded
+	FTwoBandSHVectorRGB TranslucentLighting;
+	TranslucentLighting.R.V.x = AmbientLightingVector.r;
+	TranslucentLighting.G.V.x = AmbientLightingVector.g;
+	TranslucentLighting.B.V.x = AmbientLightingVector.b;
+	float3 NormalizedAmbientColor = AmbientLightingVector.rgb / (Luminance(AmbientLightingVector.rgb) + 0.00001f);
+
+	// Scale the monocrome directional coefficients with the normalzed ambient color as an approximation to the uncompressed values
+	TranslucentLighting.R.V.yzw = DirectionalLightingVector.rgb * NormalizedAmbientColor.r;
+	TranslucentLighting.G.V.yzw = DirectionalLightingVector.rgb * NormalizedAmbientColor.g;
+	TranslucentLighting.B.V.yzw = DirectionalLightingVector.rgb * NormalizedAmbientColor.b;
+
+	// Compute diffuse lighting which takes the normal into account
+	FTwoBandSHVector DiffuseTransferSH = CalcDiffuseTransferSH(WorldNormal, 1);
+	VolumeLighting = float4(max(half3(0, 0, 0), DotSH(TranslucentLighting, DiffuseTransferSH)), AmbientLightingVector.a);
+	InterpolatedLighting += DiffuseColor * VolumeLighting.rgb;
+}
+
+float4 GetVolumeLightingDirectional(float4 AmbientLightingVector, float3 DirectionalLightingVector, float3 WorldNormal, float DirectionalLightingIntensity)
+{
+	float3 DiffuseColor = 1.0f;
+	float3 InterpolatedLighting = 0.0f;
+	float4 VolumeLighting = 0.0f;
+
+	GetVolumeLightingDirectional(AmbientLightingVector, DirectionalLightingVector, WorldNormal, DiffuseColor, DirectionalLightingIntensity, InterpolatedLighting, VolumeLighting);
+
+	return VolumeLighting;
+}
\ No newline at end of file
