Substrate - made some translucency shadow code common.

diff --git a/Engine/Shaders/Private/RayTracing/RayTracingOcclusionRGS.usf b/Engine/Shaders/Private/RayTracing/RayTracingOcclusionRGS.usf
--- a/Engine/Shaders/Private/RayTracing/RayTracingOcclusionRGS.usf
+++ b/Engine/Shaders/Private/RayTracing/RayTracingOcclusionRGS.usf
@@ -15,6 +15,7 @@
 #include "../TransmissionCommon.ush"
 #include "../HairStrands/HairStrandsVisibilityCommon.ush"
 #include "../Strata/Strata.ush"
+#include "../Strata/StrataSubsurface.ush"
 #include "RayTracingCommon.ush"
 #include "RayTracingDeferredShadingCommon.ush"
 #include "RayTracingDirectionalLight.ush"
@@ -686,11 +687,11 @@ RAY_TRACING_ENTRY_RAYGEN_OR_INLINE(Occlusion)
 				}
 				else if (SSSType == SSS_TYPE_DIFFUSION)
 				{
-					FSubsurfaceOpacityMFP SubsurfaceOpacityMFP = StrataGetOpacityFromSubSurface(SSSHeader);
+					FSubsurfaceOpacityMFP SubsurfaceOpacityMFP = StrataGetSubsurfaceOpacityMFP(SSSHeader);
 
 					// We know that or SSS_TYPE_DIFFUSION we have access to the mean free path. 
 					const float MFP = SubsurfaceOpacityMFP.Data;
-					const float Extinction = 1.0f / max(0.00001, MFP);
+					const float Extinction = StrataShadowMFPToExtinction(MFP);
 
 					TransmissionParameters = InitTransmissionProfileParams();
 					TransmissionParameters.ExtinctionScale = Extinction;
@@ -706,7 +707,7 @@ RAY_TRACING_ENTRY_RAYGEN_OR_INLINE(Occlusion)
 
 			if (ShadingParameters.bIsSubsurface) // Legacy Subsurface shading model == Substrate Wrap SSS.
 			{
-				FSubsurfaceOpacityMFP SubsurfaceOpacityMFP = StrataGetOpacityFromSubSurface(SSSHeader);
+				FSubsurfaceOpacityMFP SubsurfaceOpacityMFP = StrataGetSubsurfaceOpacityMFP(SSSHeader);
 				if (SubsurfaceOpacityMFP.bDataIsOpacity)
 				{
 					TransmissionParameters.ExtinctionScale = SubsurfaceOpacityMFP.Data;
diff --git a/Engine/Shaders/Private/ShadowProjectionPixelShader.usf b/Engine/Shaders/Private/ShadowProjectionPixelShader.usf
--- a/Engine/Shaders/Private/ShadowProjectionPixelShader.usf
+++ b/Engine/Shaders/Private/ShadowProjectionPixelShader.usf
@@ -324,8 +324,9 @@ void Main(
 				else if (SSSType == SSS_TYPE_DIFFUSION)
 				{
 					//For per pixel diffusion, use default transmission parameters (only NormalScale & ExtinctionScale are used)
-					const float MFP = dot(StrataSubSurfaceHeaderGetMFP(SSSHeader), (1.0/3.0).xxx);
-					const float Extinction = 1.0f / max(0.00001, MFP);
+					const float3 ColoredMFP = StrataSubSurfaceHeaderGetMFP(SSSHeader);
+					const float3 GreyMFP    = StrataShadowColoredMFPToGreyScaleMFP(ColoredMFP);
+					const float Extinction  = StrataShadowMFPToExtinction(GreyMFP);
 					TransmissionParams.ExtinctionScale = Extinction;
 				}
 
@@ -499,8 +500,9 @@ void MainOnePassPointLightPS(
 	else if (SSSType == SSS_TYPE_DIFFUSION)
 	{
 		//For per pixel diffusion, use default transmission parameters (only NormalScale & ExtinctionScale are used)
-		const float MFP = dot(StrataSubSurfaceHeaderGetMFP(SSSHeader), (1.0 / 3.0).xxx);
-		const float Extinction = 1.0f / max(0.00001, MFP);
+		const float3 ColoredMFP = StrataSubSurfaceHeaderGetMFP(SSSHeader);
+		const float3 GreyMFP    = StrataShadowColoredMFPToGreyScaleMFP(ColoredMFP);
+		const float Extinction  = StrataShadowMFPToExtinction(GreyMFP);
 		TransmissionParams.ExtinctionScale = Extinction;
 	}
 
diff --git a/Engine/Shaders/Private/Strata/Strata.ush b/Engine/Shaders/Private/Strata/Strata.ush
--- a/Engine/Shaders/Private/Strata/Strata.ush
+++ b/Engine/Shaders/Private/Strata/Strata.ush
@@ -9,24 +9,6 @@
 #endif
 
 
-// A structure that can be used to transfert Legacy/substrate data for light transmission.
-struct FSubsurfaceOpacityMFP
-{
-	bool  bDataIsOpacity;	// If false, Data is a mean free path
-	float Data;				// MFP or Opacity
-	float Density;			// Must be stored separately from Opacity, because Opacity is used for some tests
-};
-
-FSubsurfaceOpacityMFP GetInitialisedSubsurfaceOpacityMFP()
-{
-	FSubsurfaceOpacityMFP SubsurfaceOpacityMFP;
-	SubsurfaceOpacityMFP.bDataIsOpacity = true;
-	SubsurfaceOpacityMFP.Data = 1.0;
-	SubsurfaceOpacityMFP.Density = 0.0f;
-	return SubsurfaceOpacityMFP;
-}
-
-
 #if STRATA_ENABLED
 
 #include "../ShadingEnergyConservation.ush"
@@ -43,6 +25,7 @@ FSubsurfaceOpacityMFP GetInitialisedSubsurfaceOpacityMFP()
 #include "../GammaCorrectionCommon.ush"
 #include "/Engine/Shared/StrataDefinitions.h"
 #include "/Engine/Private/Strata/StrataStatisticalOperators.ush"
+#include "/Engine/Private/Strata/StrataSubsurface.ush"
 
 // DXC has issue with unrolling certain part of the code. To address temporarly the warning-as-error issues, 
 // we redefine the unroll function and fall back to loop on DXC
@@ -448,14 +431,6 @@ struct FStrataBSDF
 #define BSDF_COMPLEXCODING_BIT_COUNT			(32u)
 #define BSDF_COMPLEXFASTENCODING_MASK			(0xFFFFFFFFu)
 
-// Sub-surface scattering types
-#define SSS_TYPE_INVALID						(0u)
-#define SSS_TYPE_WRAP							(1u)
-#define SSS_TYPE_DIFFUSION						(2u)
-#define SSS_TYPE_DIFFUSION_PROFILE				(3u)
-#define SSS_TYPE_SIMPLEVOLUME					(4u)
-#define SSS_TYPE_COUNT							(5u)
-
 // Sanity check
 #if (BSDF_BIT_OFFSET_WEIGHT10F + BSDF_BIT_COUNT_WEIGHT10F) > 32
 #error Strata BSDF state is larger than 32 bits
@@ -831,7 +806,6 @@ FStrataOpaqueRoughRefractionData StrataUnpackOpaqueRoughRefractionData(float3 In
 	return Out;
 }
 
-
 // Returns the variance of refracted rays intersection with on the XY plane after scattering at the interface of a slab of matter,
 // having the given Roughness and a normal along Z; after traveling along -Z from air to water (eta = 1.333).
 // That variance has been generated from unit vectors, so it is a measure of the spread of light on the XY plane after a depth of 1.
@@ -858,233 +832,6 @@ float GetVarianceFromRoughness(float Roughness)
 	return saturate(             - 0.000332564 * Roughness - 0.0168059 * Roughness2 + 0.115569 * Roughness3 - 0.0526336 * Roughness4);	// Removing the 0.0000226575 to get GetVarianceFromRoughness(0)=0, to not get any blurred refraction when roughness==0.
 }
 
-///////////////////////////////////////////////////////////////////////////////
-// Sub-surface 
-
-// Stored in first slice as UINT
-// This is hot data that is often accessed by intrested passes.
-struct FStrataSubsurfaceHeader
-{
-	// 2 bits
-	// SSS type:
-	// 0: Invalid
-	// 1: Wrap
-	// 2: Diffusion
-	// 3: Diffusion with profile
-
-	// 30 bits of data.
-	// If  bIsProfile: 8bit profile radius, 8bits ProfileId, 16bits unused
-	// If !bIsProfile: 101010 mean free path
-	uint Bytes;
-};
-
-// Note: SSSHeader only stores INVALID, WRAP, DIFFUSION, DIFFUSION_PROFILE type. This is because the type is represented onto 2bits and 
-//       other SSS types don't have/need side payload. This means SIMPLEVOLUME is trunked.
-#if SSS_TYPE_COUNT != 5
-#error Update this code to ensure all SSS types requiering side payload can be represented in the 2bits allocated for types.
-#endif
-#define SSSHEADER_TYPE(Header)	(Header.Bytes & 0x3)
-#define SSSHEADER_TYPE_MASK		(0x00000003)
-
-void StrataSubSurfaceHeaderSetSSSType(inout FStrataSubsurfaceHeader SSSHeader, uint SSSType)
-{
-	SSSHeader.Bytes &= (~SSSHEADER_TYPE_MASK);
-	SSSHeader.Bytes |= SSSType & SSSHEADER_TYPE_MASK;
-}
-
-void StrataSubSurfaceHeaderSetProfile(inout FStrataSubsurfaceHeader SSSHeader, float RadiusScale, uint ProfileId)
-{
-	SSSHeader.Bytes &= SSSHEADER_TYPE_MASK;
-	SSSHeader.Bytes |= ProfileId << 24;
-	SSSHeader.Bytes |= PackR8(RadiusScale) << 16;
-}
-
-void StrataSubSurfaceHeaderSetNonProfile(inout FStrataSubsurfaceHeader SSSHeader, float3 MeanFreePath)
-{
-	SSSHeader.Bytes &= SSSHEADER_TYPE_MASK;
-	SSSHeader.Bytes |= (Pack10F(MeanFreePath.x) << 22) | (Pack10F(MeanFreePath.y) << 12) | (Pack10F(MeanFreePath.z) << 2);
-}
-
-float StrataSubSurfaceGetWrapOpacityFromAnisotropy(float PhaseAnisotropy)
-{
-	// Reinterpret the phase function anisotropy as 'opacity' value
-	// * Forward/Backward phase function == Thin surface (i.e., opacity => 0)
-	// * Isotropic phase function == Thick surface (i.e., opacity => 1)
-	const float Opacity = 1.f - abs(PhaseAnisotropy);
-	return Opacity;
-}
-
-void StrataSubSurfaceHeaderSetWrapOpacity(inout FStrataSubsurfaceHeader SSSHeader, float Opacity)
-{
-	SSSHeader.Bytes &= SSSHEADER_TYPE_MASK;
-	SSSHeader.Bytes |= PackR8(Opacity) << 2;
-}
-
-void StrataSubSurfaceHeaderSetWrap(inout FStrataSubsurfaceHeader SSSHeader, float PhaseAnisotropy)
-{
-	const float Opacity = StrataSubSurfaceGetWrapOpacityFromAnisotropy(PhaseAnisotropy);
-	StrataSubSurfaceHeaderSetWrapOpacity(SSSHeader, Opacity);
-}
-
-bool StrataSubSurfaceHeaderGetIsValid(in FStrataSubsurfaceHeader SSSHeader)
-{
-	return SSSHEADER_TYPE(SSSHeader) != SSS_TYPE_INVALID;
-}
-
-bool StrataSubSurfaceHeaderHasExtras(in FStrataSubsurfaceHeader SSSHeader)
-{
-	const uint SSSType = SSSHEADER_TYPE(SSSHeader);
-	return SSSType == SSS_TYPE_DIFFUSION || SSSType == SSS_TYPE_DIFFUSION_PROFILE;
-}
-
-bool StrataSubSurfaceHeaderGetUseDiffusion(in FStrataSubsurfaceHeader SSSHeader)
-{
-	const uint SSSType = SSSHEADER_TYPE(SSSHeader);
-	return SSSType == SSS_TYPE_DIFFUSION || SSSType == SSS_TYPE_DIFFUSION_PROFILE;
-}
-
-bool StrataSubSurfaceHeaderGetIsProfile(in FStrataSubsurfaceHeader SSSHeader)
-{
-	return SSSHEADER_TYPE(SSSHeader) == SSS_TYPE_DIFFUSION_PROFILE;
-}
-
-bool StrataSubSurfaceHeaderGetIsWrap(in FStrataSubsurfaceHeader SSSHeader)
-{
-	return SSSHEADER_TYPE(SSSHeader) == SSS_TYPE_WRAP;
-}
-
-uint StrataSubSurfaceHeaderGetSSSType(in FStrataSubsurfaceHeader SSSHeader)
-{
-	return SSSHEADER_TYPE(SSSHeader);
-}
-
-uint StrataSubSurfaceHeaderGetProfileId(in FStrataSubsurfaceHeader SSSHeader)
-{
-	const uint SSSType = SSSHEADER_TYPE(SSSHeader);
-	const uint ProfileId = (SSSType == SSS_TYPE_DIFFUSION_PROFILE) ? ((SSSHeader.Bytes >> 24) & 0xFF) : (SSSType == SSS_TYPE_DIFFUSION ? SSS_PROFILE_ID_PERPIXEL : SSS_PROFILE_ID_INVALID);
-	return ProfileId;
-}
-
-float StrataSubSurfaceHeaderGetProfileRadiusScale(in FStrataSubsurfaceHeader SSSHeader)
-{
-	const bool bIsValidAndProfile = SSSHEADER_TYPE(SSSHeader) == SSS_TYPE_DIFFUSION_PROFILE;
-	return bIsValidAndProfile ? UnpackR8(SSSHeader.Bytes >> 16) : 1.f;
-}
-
-float3 StrataSubSurfaceHeaderGetMFP(in FStrataSubsurfaceHeader SSSHeader)				// Only when !IsProfile
-{
-	return float3(
-		Unpack10F(SSSHeader.Bytes >> 22),
-		Unpack10F(SSSHeader.Bytes >> 12),
-		Unpack10F(SSSHeader.Bytes >>  2));
-}
-
-float StrataSubSurfaceHeaderGetWrapOpacity(in FStrataSubsurfaceHeader SSSHeader)
-{
-	return SSSHEADER_TYPE(SSSHeader) == SSS_TYPE_WRAP ? UnpackR8(SSSHeader.Bytes >> 2) : 1.0f;
-}
-
-float StrataSubSurfaceHeaderGetOpacity(in FStrataSubsurfaceHeader SSSHeader)
-{
-	const uint SSSType = SSSHEADER_TYPE(SSSHeader);
-
-	// Shadow 'Opacity' is based on legacy shading code
-	// * SSS Wrap              -> StrataSubSurfaceHeaderGetWrapOpacity()
-	// * SSS Diffusion         -> 1
-	// * SSS Diffusion Profile -> StrataSubSurfaceHeaderGetProfileRadiusScale()
-	float Opacity = 1.0f;
-	Opacity = SSSType == SSS_TYPE_DIFFUSION_PROFILE ? StrataSubSurfaceHeaderGetProfileRadiusScale(SSSHeader) : Opacity;
-	Opacity = SSSType == SSS_TYPE_WRAP ? StrataSubSurfaceHeaderGetWrapOpacity(SSSHeader) : Opacity;
-	return Opacity;
-}
-
-// Legacy conversion function, which translated Strata sub-surface data into 'opacity' for shadow transmission purpose. 
-// STRATA_TODO: merge this function with StrataSubSurfaceHeaderGetOpacity() and GetSubsurfaceOpacityFromGbuffer() in a common Strata header.
-FSubsurfaceOpacityMFP StrataGetOpacityFromSubSurface(FStrataSubsurfaceHeader SSSHeader)
-{
-	FSubsurfaceOpacityMFP SubsurfaceOpacityMFP = GetInitialisedSubsurfaceOpacityMFP();
-
-	const uint SSSType = SSSHEADER_TYPE(SSSHeader);
-
-	if (SSSType == SSS_TYPE_DIFFUSION)
-	{
-		SubsurfaceOpacityMFP.bDataIsOpacity = false;
-		const float3 MFP = StrataSubSurfaceHeaderGetMFP(SSSHeader);
-		SubsurfaceOpacityMFP.Data = dot(MFP, (1.0f/3.0f).xxx);
-		SubsurfaceOpacityMFP.Density = 0; // Not fed
-	}
-	else if (SSSType == SSS_TYPE_DIFFUSION_PROFILE)
-	{
-		SubsurfaceOpacityMFP.bDataIsOpacity = true;
-		SubsurfaceOpacityMFP.Data = min(StrataSubSurfaceHeaderGetProfileRadiusScale(SSSHeader), 0.99);
-		SubsurfaceOpacityMFP.Density = 0; // Not fed
-	}
-	else if (SSSType == SSS_TYPE_WRAP)
-	{
-		SubsurfaceOpacityMFP.bDataIsOpacity = true;
-		SubsurfaceOpacityMFP.Data = min(StrataSubSurfaceHeaderGetWrapOpacity(SSSHeader), 0.99);
-		SubsurfaceOpacityMFP.Density = 0; // Not fed
-	}
-
-	return SubsurfaceOpacityMFP;
-}
-
-// Stored in second slice as UINT
-// This is cold data, rarely accessed
-struct FStrataSubsurfaceExtras
-{
-//	float3 BaseColor; // alpha unused
-	uint Bytes;
-};
-void StrataSubsurfaceExtrasSetBaseColor(inout FStrataSubsurfaceExtras SSSExtras, float3 BaseColor)
-{
-	SSSExtras.Bytes = PackColorLinearToGamma2AlphaLinear(float4(BaseColor, 0.0f));
-}
-float3 StrataSubsurfaceExtrasGetBaseColor(in FStrataSubsurfaceExtras SSSExtras)
-{
-	return UnpackColorGamma2ToLinearAlphaLinear(SSSExtras.Bytes).rgb;
-}
-
-struct FStrataSubsurfaceData
-{
-	FStrataSubsurfaceHeader Header;
-	FStrataSubsurfaceExtras Extras;
-};
-FStrataSubsurfaceHeader StrataLoadSubsurfaceHeader(Texture2DArray<uint> StrataBuffer, uint FirstSliceStoringStrataSSSData, uint2 PixelPos)
-{
-	FStrataSubsurfaceHeader Header;
-	Header.Bytes = StrataBuffer.Load(uint4(PixelPos, FirstSliceStoringStrataSSSData + 0, 0));
-	return Header;
-}
-FStrataSubsurfaceExtras StrataLoadSubsurfaceExtras(Texture2DArray<uint> StrataBuffer, uint FirstSliceStoringStrataSSSData, uint2 PixelPos)
-{
-	FStrataSubsurfaceExtras Extras;
-	Extras.Bytes = StrataBuffer.Load(uint4(PixelPos, FirstSliceStoringStrataSSSData + 1, 0));
-	return Extras;
-}
-FStrataSubsurfaceData StrataLoadSubsurfaceData(Texture2DArray<uint> StrataBuffer, uint FirstSliceStoringStrataSSSData, uint2 PixelPos)
-{
-	FStrataSubsurfaceData SSSData;
-	SSSData.Header = StrataLoadSubsurfaceHeader(StrataBuffer, FirstSliceStoringStrataSSSData, PixelPos);
-	SSSData.Extras = StrataLoadSubsurfaceExtras(StrataBuffer, FirstSliceStoringStrataSSSData, PixelPos);
-	return SSSData;
-}
-
-void StrataStoreSubsurfaceHeader(RWTexture2DArray<uint> StrataBuffer, uint FirstSliceStoringStrataSSSData, uint2 PixelPos, uint HeaderBytes)
-{
-	StrataBuffer[uint3(PixelPos, FirstSliceStoringStrataSSSData + 0)] = HeaderBytes;
-}
-void StrataStoreSubsurfaceExtras(RWTexture2DArray<uint> StrataBuffer, uint FirstSliceStoringStrataSSSData, uint2 PixelPos, uint ExtraBytes)
-{
-	StrataBuffer[uint3(PixelPos, FirstSliceStoringStrataSSSData + 1)] = ExtraBytes;
-}
-
-uint StrataSubsurfaceProfileIdTo8bits(float In)
-{
-	// Similar encoding than ExtractSubsurfaceProfileInt. Valid profile ID start at 1.
-	return uint(In * 255.0f + 0.5f);
-}
-
 ///////////////////////////////////////////////////////////////////////////////
 // PackSSSProfile
 
diff --git a/Engine/Shaders/Private/Strata/StrataSubsurface.ush b/Engine/Shaders/Private/Strata/StrataSubsurface.ush
new file mode 100644
--- /dev/null
+++ b/Engine/Shaders/Private/Strata/StrataSubsurface.ush
@@ -0,0 +1,294 @@
+// Copyright Epic Games, Inc. All Rights Reserved.
+
+#pragma once
+
+#include "../TransmissionCommon.ush"
+
+// Forward declarations
+uint PackColorLinearToGamma2AlphaLinear(float4 In);
+float4 UnpackColorGamma2ToLinearAlphaLinear(uint In);
+
+///////////////////////////////////////////////////////////////////////////////
+// Sub-surface 
+// 
+// Stored in first slice as UINT
+// This is hot data that is often accessed by SSS related passes.
+struct FStrataSubsurfaceHeader
+{
+	// 2 bits
+	// SSS type:
+	// 0: Invalid
+	// 1: Wrap
+	// 2: Diffusion
+	// 3: Diffusion with profile
+
+	// 30 bits of data.
+	// If  bIsProfile: 8bit profile radius, 8bits ProfileId, 16bits unused
+	// If !bIsProfile: 101010 mean free path
+	uint Bytes;
+};
+
+// Sub-surface scattering types
+#define SSS_TYPE_INVALID						(0u)
+#define SSS_TYPE_WRAP							(1u)
+#define SSS_TYPE_DIFFUSION						(2u)
+#define SSS_TYPE_DIFFUSION_PROFILE				(3u)
+#define SSS_TYPE_SIMPLEVOLUME					(4u)
+#define SSS_TYPE_COUNT							(5u)
+
+// Note: SSSHeader only stores INVALID, WRAP, DIFFUSION, DIFFUSION_PROFILE type. This is because the type is represented onto 2bits and 
+//       other SSS types don't have/need side payload. This means SIMPLEVOLUME is trunked.
+#if SSS_TYPE_COUNT != 5
+#error Update this code to ensure all SSS types requiering side payload can be represented in the 2bits allocated for types.
+#endif
+#define SSSHEADER_TYPE(Header)	(Header.Bytes & 0x3)
+#define SSSHEADER_TYPE_MASK		(0x00000003)
+
+void StrataSubSurfaceHeaderSetSSSType(inout FStrataSubsurfaceHeader SSSHeader, uint SSSType)
+{
+	SSSHeader.Bytes &= (~SSSHEADER_TYPE_MASK);
+	SSSHeader.Bytes |= SSSType & SSSHEADER_TYPE_MASK;
+}
+
+void StrataSubSurfaceHeaderSetProfile(inout FStrataSubsurfaceHeader SSSHeader, float RadiusScale, uint ProfileId)
+{
+	SSSHeader.Bytes &= SSSHEADER_TYPE_MASK;
+	SSSHeader.Bytes |= ProfileId << 24;
+	SSSHeader.Bytes |= PackR8(RadiusScale) << 16;
+}
+
+void StrataSubSurfaceHeaderSetNonProfile(inout FStrataSubsurfaceHeader SSSHeader, float3 MeanFreePath)
+{
+	SSSHeader.Bytes &= SSSHEADER_TYPE_MASK;
+	SSSHeader.Bytes |= (Pack10F(MeanFreePath.x) << 22) | (Pack10F(MeanFreePath.y) << 12) | (Pack10F(MeanFreePath.z) << 2);
+}
+
+float StrataSubSurfaceGetWrapOpacityFromAnisotropy(float PhaseAnisotropy)
+{
+	// Reinterpret the phase function anisotropy as 'opacity' value
+	// * Forward/Backward phase function == Thin surface (i.e., opacity => 0)
+	// * Isotropic phase function == Thick surface (i.e., opacity => 1)
+	const float Opacity = 1.f - abs(PhaseAnisotropy);
+	return Opacity;
+}
+
+void StrataSubSurfaceHeaderSetWrapOpacity(inout FStrataSubsurfaceHeader SSSHeader, float Opacity)
+{
+	SSSHeader.Bytes &= SSSHEADER_TYPE_MASK;
+	SSSHeader.Bytes |= PackR8(Opacity) << 2;
+}
+
+void StrataSubSurfaceHeaderSetWrap(inout FStrataSubsurfaceHeader SSSHeader, float PhaseAnisotropy)
+{
+	const float Opacity = StrataSubSurfaceGetWrapOpacityFromAnisotropy(PhaseAnisotropy);
+	StrataSubSurfaceHeaderSetWrapOpacity(SSSHeader, Opacity);
+}
+
+bool StrataSubSurfaceHeaderGetIsValid(in FStrataSubsurfaceHeader SSSHeader)
+{
+	return SSSHEADER_TYPE(SSSHeader) != SSS_TYPE_INVALID;
+}
+
+bool StrataSubSurfaceHeaderHasExtras(in FStrataSubsurfaceHeader SSSHeader)
+{
+	const uint SSSType = SSSHEADER_TYPE(SSSHeader);
+	return SSSType == SSS_TYPE_DIFFUSION || SSSType == SSS_TYPE_DIFFUSION_PROFILE;
+}
+
+bool StrataSubSurfaceHeaderGetUseDiffusion(in FStrataSubsurfaceHeader SSSHeader)
+{
+	const uint SSSType = SSSHEADER_TYPE(SSSHeader);
+	return SSSType == SSS_TYPE_DIFFUSION || SSSType == SSS_TYPE_DIFFUSION_PROFILE;
+}
+
+bool StrataSubSurfaceHeaderGetIsProfile(in FStrataSubsurfaceHeader SSSHeader)
+{
+	return SSSHEADER_TYPE(SSSHeader) == SSS_TYPE_DIFFUSION_PROFILE;
+}
+
+bool StrataSubSurfaceHeaderGetIsWrap(in FStrataSubsurfaceHeader SSSHeader)
+{
+	return SSSHEADER_TYPE(SSSHeader) == SSS_TYPE_WRAP;
+}
+
+uint StrataSubSurfaceHeaderGetSSSType(in FStrataSubsurfaceHeader SSSHeader)
+{
+	return SSSHEADER_TYPE(SSSHeader);
+}
+
+uint StrataSubSurfaceHeaderGetProfileId(in FStrataSubsurfaceHeader SSSHeader)
+{
+	const uint SSSType = SSSHEADER_TYPE(SSSHeader);
+	const uint ProfileId = (SSSType == SSS_TYPE_DIFFUSION_PROFILE) ? ((SSSHeader.Bytes >> 24) & 0xFF) : (SSSType == SSS_TYPE_DIFFUSION ? SSS_PROFILE_ID_PERPIXEL : SSS_PROFILE_ID_INVALID);
+	return ProfileId;
+}
+
+float StrataSubSurfaceHeaderGetProfileRadiusScale(in FStrataSubsurfaceHeader SSSHeader)
+{
+	const bool bIsValidAndProfile = SSSHEADER_TYPE(SSSHeader) == SSS_TYPE_DIFFUSION_PROFILE;
+	return bIsValidAndProfile ? UnpackR8(SSSHeader.Bytes >> 16) : 1.f;
+}
+
+float3 StrataSubSurfaceHeaderGetMFP(in FStrataSubsurfaceHeader SSSHeader)				// Only when !IsProfile
+{
+	return float3(
+		Unpack10F(SSSHeader.Bytes >> 22),
+		Unpack10F(SSSHeader.Bytes >> 12),
+		Unpack10F(SSSHeader.Bytes >> 2));
+}
+
+float StrataSubSurfaceHeaderGetWrapOpacity(in FStrataSubsurfaceHeader SSSHeader)
+{
+	return SSSHEADER_TYPE(SSSHeader) == SSS_TYPE_WRAP ? UnpackR8(SSSHeader.Bytes >> 2) : 1.0f;
+}
+
+float StrataSubSurfaceHeaderGetOpacity(in FStrataSubsurfaceHeader SSSHeader)
+{
+	const uint SSSType = SSSHEADER_TYPE(SSSHeader);
+
+	// Shadow 'Opacity' is based on legacy shading code
+	// * SSS Wrap              -> StrataSubSurfaceHeaderGetWrapOpacity()
+	// * SSS Diffusion         -> 1
+	// * SSS Diffusion Profile -> StrataSubSurfaceHeaderGetProfileRadiusScale()
+	float Opacity = 1.0f;
+	Opacity = SSSType == SSS_TYPE_DIFFUSION_PROFILE ? StrataSubSurfaceHeaderGetProfileRadiusScale(SSSHeader) : Opacity;
+	Opacity = SSSType == SSS_TYPE_WRAP ? StrataSubSurfaceHeaderGetWrapOpacity(SSSHeader) : Opacity;
+	return Opacity;
+}
+
+// Stored in second slice as UINT
+// This is cold data, rarely accessed
+struct FStrataSubsurfaceExtras
+{
+	//	float3 BaseColor; // alpha unused
+	uint Bytes;
+};
+void StrataSubsurfaceExtrasSetBaseColor(inout FStrataSubsurfaceExtras SSSExtras, float3 BaseColor)
+{
+	SSSExtras.Bytes = PackColorLinearToGamma2AlphaLinear(float4(BaseColor, 0.0f));
+}
+float3 StrataSubsurfaceExtrasGetBaseColor(in FStrataSubsurfaceExtras SSSExtras)
+{
+	return UnpackColorGamma2ToLinearAlphaLinear(SSSExtras.Bytes).rgb;
+}
+
+struct FStrataSubsurfaceData
+{
+	FStrataSubsurfaceHeader Header;
+	FStrataSubsurfaceExtras Extras;
+};
+FStrataSubsurfaceHeader StrataLoadSubsurfaceHeader(Texture2DArray<uint> StrataBuffer, uint FirstSliceStoringStrataSSSData, uint2 PixelPos)
+{
+	FStrataSubsurfaceHeader Header;
+	Header.Bytes = StrataBuffer.Load(uint4(PixelPos, FirstSliceStoringStrataSSSData + 0, 0));
+	return Header;
+}
+FStrataSubsurfaceExtras StrataLoadSubsurfaceExtras(Texture2DArray<uint> StrataBuffer, uint FirstSliceStoringStrataSSSData, uint2 PixelPos)
+{
+	FStrataSubsurfaceExtras Extras;
+	Extras.Bytes = StrataBuffer.Load(uint4(PixelPos, FirstSliceStoringStrataSSSData + 1, 0));
+	return Extras;
+}
+FStrataSubsurfaceData StrataLoadSubsurfaceData(Texture2DArray<uint> StrataBuffer, uint FirstSliceStoringStrataSSSData, uint2 PixelPos)
+{
+	FStrataSubsurfaceData SSSData;
+	SSSData.Header = StrataLoadSubsurfaceHeader(StrataBuffer, FirstSliceStoringStrataSSSData, PixelPos);
+	SSSData.Extras = StrataLoadSubsurfaceExtras(StrataBuffer, FirstSliceStoringStrataSSSData, PixelPos);
+	return SSSData;
+}
+
+void StrataStoreSubsurfaceHeader(RWTexture2DArray<uint> StrataBuffer, uint FirstSliceStoringStrataSSSData, uint2 PixelPos, uint HeaderBytes)
+{
+	StrataBuffer[uint3(PixelPos, FirstSliceStoringStrataSSSData + 0)] = HeaderBytes;
+}
+void StrataStoreSubsurfaceExtras(RWTexture2DArray<uint> StrataBuffer, uint FirstSliceStoringStrataSSSData, uint2 PixelPos, uint ExtraBytes)
+{
+	StrataBuffer[uint3(PixelPos, FirstSliceStoringStrataSSSData + 1)] = ExtraBytes;
+}
+
+uint StrataSubsurfaceProfileIdTo8bits(float In)
+{
+	// Similar encoding than ExtractSubsurfaceProfileInt. Valid profile ID start at 1.
+	return uint(In * 255.0f + 0.5f);
+}
+
+
+
+///////////////////////////////////////////////////////////////////////////////
+// Shadow and transmission 
+
+// A structure that can be used to transfert Legacy/substrate data for light transmission.
+struct FSubsurfaceOpacityMFP
+{
+	bool  bDataIsOpacity;	// If false, Data is a mean free path
+	float Data;				// MFP or Opacity
+	float Density;			// Must be stored separately from Opacity, because Opacity is used for some tests
+};
+
+FSubsurfaceOpacityMFP GetInitialisedSubsurfaceOpacityMFP()
+{
+	FSubsurfaceOpacityMFP SubsurfaceOpacityMFP;
+	SubsurfaceOpacityMFP.bDataIsOpacity = true;
+	SubsurfaceOpacityMFP.Data = 1.0;
+	SubsurfaceOpacityMFP.Density = 0.0f;
+	return SubsurfaceOpacityMFP;
+}
+
+float StrataShadowMFPToExtinction(float MFP)
+{
+	const float Extinction = 1.0f / max(0.00001f, MFP);
+	return Extinction;
+}
+
+float StrataShadowColoredMFPToGreyScaleMFP(float3 MFP)
+{
+	return dot(MFP, (1.0f / 3.0f).xxx);
+}
+
+// Legacy conversion function, which translated Strata sub-surface data into 'opacity' for shadow transmission purpose. 
+FSubsurfaceOpacityMFP StrataGetSubsurfaceOpacityMFP(FStrataSubsurfaceHeader SSSHeader)
+{
+	FSubsurfaceOpacityMFP SubsurfaceOpacityMFP = GetInitialisedSubsurfaceOpacityMFP();
+
+	const uint SSSType = StrataSubSurfaceHeaderGetSSSType(SSSHeader);
+
+	if (SSSType == SSS_TYPE_DIFFUSION)
+	{
+		SubsurfaceOpacityMFP.bDataIsOpacity = false;
+		const float3 ColoredMFP = StrataSubSurfaceHeaderGetMFP(SSSHeader);
+		SubsurfaceOpacityMFP.Data = StrataShadowColoredMFPToGreyScaleMFP(ColoredMFP);
+		SubsurfaceOpacityMFP.Density = 0.0f;
+	}
+	else if (SSSType == SSS_TYPE_DIFFUSION_PROFILE)
+	{
+		SubsurfaceOpacityMFP.bDataIsOpacity = true;
+
+		const float ProfileRadiusScale = StrataSubSurfaceHeaderGetProfileRadiusScale(SSSHeader);
+		// This clamp aligns with SubsurfaceDensityFromOpacity
+		// Various engine paths treat these subsurface materials differently
+		// even when they have Opacity = 1 in the material shader, so this is
+		// important to avoid things like backface transmission being shadowed by
+		// contact shadows and so on.
+		const float Opacity = min(ProfileRadiusScale, 0.99f);
+		SubsurfaceOpacityMFP.Data = Opacity;
+
+		FTransmissionProfileParams TransmissionParams = GetTransmissionProfileParams(StrataSubSurfaceHeaderGetProfileId(SSSHeader));
+		SubsurfaceOpacityMFP.Density = SubsurfaceDensityFromOpacity(SubsurfaceOpacityMFP.Data) * TransmissionParams.ExtinctionScale * 3.1f;	// This matches the weird computations done in CalcTransmissionThickness
+	}
+	else if (SSSType == SSS_TYPE_WRAP)
+	{
+		SubsurfaceOpacityMFP.bDataIsOpacity = true;
+
+		// This clamp aligns with SubsurfaceDensityFromOpacity
+		// Various engine paths treat these subsurface materials differently
+		// even when they have Opacity = 1 in the material shader, so this is
+		// important to avoid things like backface transmission being shadowed by
+		// contact shadows and so on.
+		const float Opacity = min(StrataSubSurfaceHeaderGetWrapOpacity(SSSHeader), 0.99f);
+		SubsurfaceOpacityMFP.Data = Opacity;
+		SubsurfaceOpacityMFP.Density = SubsurfaceDensityFromOpacity(SubsurfaceOpacityMFP.Data);
+	}
+
+	return SubsurfaceOpacityMFP;
+}
+
diff --git a/Engine/Shaders/Private/VirtualShadowMaps/VirtualShadowMapTransmissionCommon.ush b/Engine/Shaders/Private/VirtualShadowMaps/VirtualShadowMapTransmissionCommon.ush
--- a/Engine/Shaders/Private/VirtualShadowMaps/VirtualShadowMapTransmissionCommon.ush
+++ b/Engine/Shaders/Private/VirtualShadowMaps/VirtualShadowMapTransmissionCommon.ush
@@ -10,6 +10,7 @@ VirtualShadowMaps/VirtualShadowMapTransmissionCommon.ush:
 #include "../DeferredShadingCommon.ush"
 #include "../TransmissionCommon.ush"
 #include "VirtualShadowMapProjectionCommon.ush"
+#include "../Strata/StrataSubsurface.ush"
 
 FSubsurfaceOpacityMFP GetSubsurfaceOpacityFromGbuffer(uint2 PixelPos)
 {
@@ -19,44 +20,7 @@ FSubsurfaceOpacityMFP GetSubsurfaceOpacityFromGbuffer(uint2 PixelPos)
 
 	FStrataSubsurfaceData LoadStrat = StrataLoadSubsurfaceData(Strata.MaterialTextureArray, Strata.FirstSliceStoringStrataSSSData, PixelPos);
 
-	const uint SSSType = SSSHEADER_TYPE(LoadStrat.Header);
-	
-	if (SSSType == SSS_TYPE_DIFFUSION)
-	{
-		SubsurfaceOpacityMFP.bDataIsOpacity = false;
-		const float3 MFP = StrataSubSurfaceHeaderGetMFP(LoadStrat.Header);
-		SubsurfaceOpacityMFP.Data = dot(MFP, (1.0f / 3.0f).xxx);
-		SubsurfaceOpacityMFP.Density = 0.0f;
-	}
-	else if (SSSType == SSS_TYPE_DIFFUSION_PROFILE)
-	{
-		SubsurfaceOpacityMFP.bDataIsOpacity = true;
- 	   
-		const float ProfileRadiusScale = StrataSubSurfaceHeaderGetProfileRadiusScale(LoadStrat.Header);
-		// This clamp aligns with SubsurfaceDensityFromOpacity
-		// Various engine paths treat these subsurface materials differently
-		// even when they have Opacity = 1 in the material shader, so this is
-		// important to avoid things like backface transmission being shadowed by
-		// contact shadows and so on.
-		const float Opacity = min(ProfileRadiusScale, 0.99);
-		SubsurfaceOpacityMFP.Data = Opacity;
-
-		FTransmissionProfileParams TransmissionParams = GetTransmissionProfileParams(StrataSubSurfaceHeaderGetProfileId(LoadStrat.Header));
-		SubsurfaceOpacityMFP.Density = SubsurfaceDensityFromOpacity(SubsurfaceOpacityMFP.Data) * TransmissionParams.ExtinctionScale * 3.1f;	// This matches the weird computations done in CalcTransmissionThickness
-	}
-	else if (SSSType == SSS_TYPE_WRAP)
-	{
-		SubsurfaceOpacityMFP.bDataIsOpacity = true;
- 	   
-		// This clamp aligns with SubsurfaceDensityFromOpacity
-		// Various engine paths treat these subsurface materials differently
-		// even when they have Opacity = 1 in the material shader, so this is
-		// important to avoid things like backface transmission being shadowed by
-		// contact shadows and so on.
-		const float Opacity = min(StrataSubSurfaceHeaderGetWrapOpacity(LoadStrat.Header), 0.99);
-		SubsurfaceOpacityMFP.Data = Opacity;
-		SubsurfaceOpacityMFP.Density = SubsurfaceDensityFromOpacity(SubsurfaceOpacityMFP.Data);
-	}
+	SubsurfaceOpacityMFP = StrataGetSubsurfaceOpacityMFP(LoadStrat.Header);
 
 #else // STRATA_ENABLED
 	FGBufferData GBufferData = GetGBufferDataUint(PixelPos);
@@ -71,6 +35,7 @@ FSubsurfaceOpacityMFP GetSubsurfaceOpacityFromGbuffer(uint2 PixelPos)
 		// even when they have Opacity = 1 in the material shader, so this is
 		// important to avoid things like backface transmission being shadowed by
 		// contact shadows and so on.
+		SubsurfaceOpacityMFP.bDataIsOpacity = true;
 		SubsurfaceOpacityMFP.Data = min(GBufferData.CustomData.a, 0.99f);
 		SubsurfaceOpacityMFP.Density = SubsurfaceDensityFromOpacity(SubsurfaceOpacityMFP.Data);
 	}
@@ -94,8 +59,9 @@ float GetShadowFactorSubsurface(float ShadowFactor, float OccluderDistance, FSub
 		else
 		{
 			const float MFP = SubsurfaceOpacityMFP.Data;
-			const float Extinction = 1.0f / max(0.00001, MFP);
-			SSSTransmission = saturate(FastExp(-OccluderDistance * Extinction));
+			const float Extinction = StrataShadowMFPToExtinction(MFP);
+			const float OpticalDepth = OccluderDistance * Extinction;
+			SSSTransmission = saturate(FastExp(-OpticalDepth));
 		}
 #else
 		// Only opacity in this path
