Define SUPPORT_CONTACT_SHADOWS to 0 in MobileLocalLightPrepass to solve shader compile error

diff --git a/Engine/Shaders/Private/MobileLocalLightPrepass.usf b/Engine/Shaders/Private/MobileLocalLightPrepass.usf
--- a/Engine/Shaders/Private/MobileLocalLightPrepass.usf
+++ b/Engine/Shaders/Private/MobileLocalLightPrepass.usf
@@ -1,5 +1,7 @@
 // Copyright Epic Games, Inc. All Rights Reserved.
 
+#define SUPPORT_CONTACT_SHADOWS 0
+
 #include "Common.ush"
 #include "SceneTextureParameters.ush"
 #include "LightGridCommon.ush"
