Set ResolvedView = ResolveView() in ReflectionEnvironmentSkyLighting PS and SingleLayerWaterComposite PS. Fixes reflection issues introduced by CL 25461034.

diff --git a/Engine/Shaders/Private/ReflectionEnvironmentPixelShader.usf b/Engine/Shaders/Private/ReflectionEnvironmentPixelShader.usf
--- a/Engine/Shaders/Private/ReflectionEnvironmentPixelShader.usf
+++ b/Engine/Shaders/Private/ReflectionEnvironmentPixelShader.usf
@@ -275,6 +275,8 @@ void ReflectionEnvironmentSkyLighting(
 #endif
 	)
 {
+	ResolvedView = ResolveView();
+
 	uint2 PixelPos = SvPosition.xy;
 	float2 BufferUV = SvPositionToBufferUV(SvPosition);
 	float2 ScreenPosition = SvPositionToScreenPosition(SvPosition).xy;
diff --git a/Engine/Shaders/Private/SingleLayerWaterComposite.usf b/Engine/Shaders/Private/SingleLayerWaterComposite.usf
--- a/Engine/Shaders/Private/SingleLayerWaterComposite.usf
+++ b/Engine/Shaders/Private/SingleLayerWaterComposite.usf
@@ -534,6 +534,8 @@ void SingleLayerWaterCompositePS(
 	in float4 SvPosition : SV_Position,
 	out float4 OutColor : SV_Target0)
 {
+	ResolvedView = ResolveView();
+
 	float2 BufferUV = SvPositionToBufferUV(SvPosition);
 	float2 ScreenPosition = SvPositionToScreenPosition(SvPosition).xy;
 
