Add initial glint support for PT.

Current limitations:
* The GlintUV derivatives are incorrect/hardcoded for now.
* Sampling is done as a single lobe, without knowledge of the glints.

diff --git a/Engine/Shaders/Private/BRDF.ush b/Engine/Shaders/Private/BRDF.ush
--- a/Engine/Shaders/Private/BRDF.ush
+++ b/Engine/Shaders/Private/BRDF.ush
@@ -743,3 +743,36 @@ float3 SimpleClearCoatTransmittance(float NoL, float NoV, float Metallic, float3
 
 	return Transmittance;
 }
+
+// Conversion of GGX roughness to Beckmann roughness
+// Using Microfacet Distribution Function: To Change or Not to Change, That Is the Question.
+float ConvertRoughnessGGXToBeckmann(float Roughness)
+{
+	// Use a custom remapping to better match the overall highlight shape, in particular for lower roughness (i.e., <0.5). 
+	// The original mapping was Ctr=Roughness based on the paper
+	const float Ctr = pow(Roughness, 0.7f);
+	return saturate(0.8388 * Ctr+ 0.7 * Pow4(Ctr));
+}
+float2 ConvertRoughnessGGXToBeckmann(float2 Roughness)
+{
+	return float2(ConvertRoughnessGGXToBeckmann(Roughness.x), ConvertRoughnessGGXToBeckmann(Roughness.y));
+}
+
+struct FBeckmannDesc
+{
+	float2 Roughness;	// Roughness
+	float2 Alpha;		// Roughness squared
+	float2 Sigma;		// RMS slope of the microfacets
+	float  Rho;			// Slope correlation factor to be used with LEAN mapping. See http://igg.unistra.fr/People/chermain/assets/pdf/Chermain2021RealTime.pdf.
+};
+FBeckmannDesc GGXToBeckmann(float2 GGXSafeRoughness)
+{
+	const float InvSqrt2 = 0.707106f;
+
+	FBeckmannDesc Out;
+	Out.Roughness	= ConvertRoughnessGGXToBeckmann(GGXSafeRoughness);
+	Out.Alpha		= Square(Out.Roughness);
+	Out.Sigma		= Out.Alpha * InvSqrt2;
+	Out.Rho			= 0.0f;
+	return Out;
+}
\ No newline at end of file
diff --git a/Engine/Shaders/Private/PathTracing/Material/PathTracingSubstrate.ush b/Engine/Shaders/Private/PathTracing/Material/PathTracingSubstrate.ush
--- a/Engine/Shaders/Private/PathTracing/Material/PathTracingSubstrate.ush
+++ b/Engine/Shaders/Private/PathTracing/Material/PathTracingSubstrate.ush
@@ -136,7 +136,7 @@ FMaterialSample Substrate_SampleMaterial(
 	FMaterialSample Result = CreateMaterialSample(L_World, 0.0, 0.0, 1.0, OutRoughness, PATHTRACER_SCATTER_DIFFUSE);
 
 	const float DiffPdf = NoL / PI;
-	const float2 GGXResult0 = GGXEvalReflection(L, V, H, Data.Alpha0);
+	float2 GGXResult0 = GGXEvalReflection(L, V, H, Data.Alpha0);
 	const float2 GGXResult1 = GGXEvalReflection(L, V, H, Data.Alpha1);
 	const float Spec0Pdf = GGXResult0.y;
 	const float Spec1Pdf = GGXResult1.y;
@@ -155,6 +155,8 @@ FMaterialSample Substrate_SampleMaterial(
 	}
 	else if (bSampledSpecular)
 	{
+		SubstrateGlint(Payload, V, L, Data.Alpha0, GGXResult0.y /*Pdf*/, GGXResult0.x /*Weight*/);
+		
 		const float3 FTint = SubstrateSpecularTint(Payload, NoV, NoL, VoH, NoH);
 #if STRATA_FRESNEL_F82
 		// STRATA_TODO: Need to tune the logic for specular micro-occlusion ...
@@ -218,7 +220,7 @@ FMaterialEval Substrate_EvalMaterial(
 		Result.AddLobeWithMIS(DiffuseSpecularScale.x * Diffuse * ShadowTerminator, NoL / PI, Data.LobePdf.x);
 	}
 	{
-		const float2 GGXResult0 = GGXEvalReflection(L, V, H, Data.Alpha0);
+		float2 GGXResult0 = GGXEvalReflection(L, V, H, Data.Alpha0);
 		const float2 GGXResult1 = GGXEvalReflection(L, V, H, Data.Alpha1);
 		const float3 FTint = SubstrateSpecularTint(Payload, NoV, NoL, VoH, NoH);
 #if STRATA_FRESNEL_F82
@@ -226,6 +228,8 @@ FMaterialEval Substrate_EvalMaterial(
 #else
 		const float3 F = Data.FuzzAttenuation * F_Schlick(Data.F0, Data.F90, VoH) * FTint;
 #endif
+		SubstrateGlint(Payload, V, L, Data.Alpha0, GGXResult0.y /*Pdf*/, GGXResult0.x /*Weight*/);
+
 		Result.AddLobeWithMIS(DiffuseSpecularScale.y * (1.0 - Payload.RoughnessData.z) * F * GGXResult0.x * Data.Spec0.W, GGXResult0.y, Data.LobePdf.y);
 		Result.AddLobeWithMIS(DiffuseSpecularScale.y * (      Payload.RoughnessData.z) * F * GGXResult1.x * Data.Spec1.W, GGXResult1.y, Data.LobePdf.z);
 	}
diff --git a/Engine/Shaders/Private/PathTracing/Material/PathTracingSubstrateCommon.ush b/Engine/Shaders/Private/PathTracing/Material/PathTracingSubstrateCommon.ush
--- a/Engine/Shaders/Private/PathTracing/Material/PathTracingSubstrateCommon.ush
+++ b/Engine/Shaders/Private/PathTracing/Material/PathTracingSubstrateCommon.ush
@@ -10,6 +10,12 @@
 
 #define STRATA_FRESNEL_F82  0
 
+#define STRATA_PATH_TRACING_GLINTS (STRATA_COMPLEXSPECIALPATH && PLATFORM_ENABLES_STRATA_GLINTS)
+
+#if STRATA_PATH_TRACING_GLINTS
+#include "/Engine/Private/Strata/Glint/GlintThirdParty.ush"
+#endif
+
 float3 SubstrateLobeWeight(FPathTracingPayload Payload, float SatNoL)
 {
 	// See reference in LuminanceWeight
@@ -21,6 +27,7 @@ float3 SubstrateLobeWeight(FPathTracingPayload Payload, float SatNoL)
 float3 SubstrateSpecularTint(FPathTracingPayload Payload, float NoV, float NoL, float VoH, float NoH)
 {
 	float3 Out = 1.f;
+	BRANCH
 	if (Payload.SpecularProfileId != 0)
 	{
 		const uint ProfileID = GetSpecularProfileId(Payload.SpecularProfileId);
@@ -29,4 +36,19 @@ float3 SubstrateSpecularTint(FPathTracingPayload Payload, float NoV, float NoL,
 		Out = View.SpecularProfileTexture.SampleLevel(View.SpecularProfileSampler, ProfileCoord, 0).xyz;
 	}
 	return Out;
+}
+
+void SubstrateGlint(FPathTracingPayload Payload, float3 InTangentV, float3 InTangentL, float2 InAlphaXY, in float InPdf, inout float OutWeight)
+{
+	// For now use the GGX lobe sampling. This will be changed with glint importance sampling later on.
+#if STRATA_PATH_TRACING_GLINTS
+	BRANCH
+	if (Payload.GlintValue > 0)
+	{
+		const float2 GGXRoughnessXY = float2(sqrtFast(InAlphaXY.x), sqrtFast(InAlphaXY.y));
+		FBeckmannDesc Beckmann = GGXToBeckmann(GGXRoughnessXY);
+		const float Fs = f_P(InTangentV, InTangentL, float3(Beckmann.Sigma, Beckmann.Rho), Payload.GlintValue, Payload.GlintUV, Payload.GlintUVdx, Payload.GlintUVdy, true /*bIncludeGeometryTerm*/);
+		OutWeight = Fs / max(0.0001f, InPdf);
+	}
+#endif
 }
\ No newline at end of file
diff --git a/Engine/Shaders/Private/PathTracing/PathTracing.usf b/Engine/Shaders/Private/PathTracing/PathTracing.usf
--- a/Engine/Shaders/Private/PathTracing/PathTracing.usf
+++ b/Engine/Shaders/Private/PathTracing/PathTracing.usf
@@ -8,6 +8,11 @@ PathTracing.usf: Reference path tracing
 
 #define PATH_TRACING_SHADER_VERSION 0x51478b68 // Bump to force-compile this shader
 
+// Ensure specular profile and glints are supported
+#if STRATA_ENABLED && PATH_TRACER_USE_STRATA_SPECIAL_COMPLEX_MATERIAL
+#define STRATA_COMPLEXSPECIALPATH 1
+#endif
+
 #include "PathTracingCore.ush"
 
 // describe how we split up the image between tiles and amongst GPUs
diff --git a/Engine/Shaders/Private/PathTracing/PathTracingCommon.ush b/Engine/Shaders/Private/PathTracing/PathTracingCommon.ush
--- a/Engine/Shaders/Private/PathTracing/PathTracingCommon.ush
+++ b/Engine/Shaders/Private/PathTracing/PathTracingCommon.ush
@@ -323,6 +323,10 @@ struct FPathTracingPayload : FMinimalPayload
 	float3 TransmittanceN;
 	float CoverageAboveAlongN;
 	float SpecularProfileId;
+	float GlintValue;
+	float2 GlintUV;
+	float2 GlintUVdx;
+	float2 GlintUVdy;
 
 	// not transported through packed payload, but used in raygen
 	float3 SubsurfaceColor;
@@ -606,7 +610,7 @@ FPackedPathTracingPayload PackPathTracingPayload(FPathTracingPayload Input)
 	Output.PackedData[5] = PayloadEncodeLDRColor(Input.DiffuseColor);
 	Output.PackedData[6] = PayloadEncodeLDRColor(Input.SpecularColor);
 	Output.PackedData[7] = PayloadEncodeLDRColor(Input.SpecularEdgeColor);
-	Output.PackedData[8] = PayloadEncodeLDRColor(Input.RoughnessData); // not a color, but encoding is still ok
+	Output.PackedData[8] = PackRGBA8(float4(Input.RoughnessData, 0.5 * Input.Anisotropy + 0.5)); // not a color, but encoding is still ok. NOTE: Aniso is [-1,1]
 	Output.PackedData[9] = PayloadEncodeUnitVector(Input.WorldTangent);
 	Output.PackedData[10] = PayloadEncodeLDRColor(Input.TransparencyColor);
 	Output.PackedData[11] = f32tof16(Input.MeanFreePath.x);
@@ -614,12 +618,24 @@ FPackedPathTracingPayload PackPathTracingPayload(FPathTracingPayload Input)
 	Output.PackedData[12] = f32tof16(Input.MeanFreePath.z);
 	Output.PackedData[12]|= f32tof16(Input.PhaseG) << 16;
 	Output.PackedData[13] = PayloadEncodeLDRColor(Input.FuzzColor);
-    Output.PackedData[14] = PayloadEncodeLDRColor(float3(0.5 * Input.Anisotropy + 0.5, Input.FuzzAmount, Input.FuzzRoughness)); // NOTE: Aniso is [-1,1]
+    Output.PackedData[14] = PayloadEncodeLDRColor(float3(0, Input.FuzzAmount, Input.FuzzRoughness));
 	Output.PackedData[15] = PayloadEncodeLDRColor(Input.WeightV);
 	Output.PackedData[16] = PayloadEncodeLDRColor(Input.TransmittanceN);
 	Output.PackedData[17] = Pack10F(Input.Ior);
 	Output.PackedData[17]|= Pack10F(Input.CoverageAboveAlongN) << 10;
 	Output.PackedData[17]|= Pack10F(Input.SpecularProfileId) << 20;
+	Output.PackedData[17]|= (Input.GlintValue > 0 ? 1u : 0u) << 30; // Bit flag for glint. 
+
+	// In order to not increase the data payload size, we alias SSS/Fuzz parameters when glint are enabled
+	if (Input.GlintValue > 0)
+	{
+		const uint2 PackedGlintsData = PackGlints(Input.GlintValue, Input.GlintUV);
+		Output.PackedData[11] = PackedGlintsData.x;
+		Output.PackedData[12] = PackedGlintsData.y;
+		Output.PackedData[13] = PackFloat2ToUInt(Input.GlintUVdx);
+		Output.PackedData[14] = PackFloat2ToUInt(Input.GlintUVdx);
+	}
+
 #else // PATHTRACING_STRATA_PAYLOAD
 
 
@@ -678,7 +694,8 @@ FPathTracingPayload UnpackPathTracingPayload(FPackedPathTracingPayload Input, FR
 	Output.DiffuseColor        = PayloadDecodeLDRColor(Input.PackedData[5]);
 	Output.SpecularColor       = PayloadDecodeLDRColor(Input.PackedData[6]);
 	Output.SpecularEdgeColor   = PayloadDecodeLDRColor(Input.PackedData[7]);
-	Output.RoughnessData       = PayloadDecodeLDRColor(Input.PackedData[8]);
+	Output.RoughnessData       = UnpackRGBA8(Input.PackedData[8]).xyz;
+	Output.Anisotropy          = 2 * UnpackRGBA8(Input.PackedData[8]).w - 1;
 	Output.WorldTangent        = PayloadDecodeUnitVector(Input.PackedData[9]);
 	Output.TransparencyColor   = PayloadDecodeLDRColor(Input.PackedData[10]);
 	Output.MeanFreePath.x      = f16tof32(Input.PackedData[11]);
@@ -687,7 +704,6 @@ FPathTracingPayload UnpackPathTracingPayload(FPackedPathTracingPayload Input, FR
 	Output.PhaseG              = f16tof32(Input.PackedData[12] >> 16);
 	Output.FuzzColor           = PayloadDecodeLDRColor(Input.PackedData[13]);
 	const float3 Data14        = PayloadDecodeLDRColor(Input.PackedData[14]);
-	Output.Anisotropy          = 2 * Data14.x - 1;
 	Output.FuzzAmount          = Data14.y;
 	Output.FuzzRoughness       = Data14.z;
 	Output.WeightV             = PayloadDecodeLDRColor(Input.PackedData[15]);
@@ -695,6 +711,24 @@ FPathTracingPayload UnpackPathTracingPayload(FPackedPathTracingPayload Input, FR
 	Output.Ior                 = Unpack10F(Input.PackedData[17]);
 	Output.CoverageAboveAlongN = Unpack10F(Input.PackedData[17] >> 10);
 	Output.SpecularProfileId   = Unpack10F(Input.PackedData[17] >> 20);
+	
+	// In order to not increase the data payload size, we alias SSS/Fuzz parameters when glint are enabled
+	const bool bHasGlints = ((Input.PackedData[17] >> 30u) & 0x1u) > 0u;
+	if (bHasGlints)
+	{
+		Output.MeanFreePath.x      = 0.f;
+		Output.MeanFreePath.y      = 0.f;
+		Output.MeanFreePath.z      = 0.f;
+		Output.PhaseG              = 0.f;
+		Output.FuzzColor           = 0.f;
+		Output.FuzzAmount          = 0.f;
+		Output.FuzzRoughness       = 0.f;
+
+		UnpackGlints(uint2(Input.PackedData[11], Input.PackedData[12]), Output.GlintValue, Output.GlintUV);
+		Output.GlintUVdx = UnpackFloat2FromUInt(Input.PackedData[13]);
+		Output.GlintUVdy = UnpackFloat2FromUInt(Input.PackedData[14]);
+	}
+
 #else // PATHTRACING_STRATA_PAYLOAD
 	Output.BaseColor = PayloadDecodeLDRColor(Input.PackedData[5]);
 	Output.TransparencyColor = PayloadDecodeLDRColor(Input.PackedData[6]);
diff --git a/Engine/Shaders/Private/PathTracing/PathTracingMaterialHitShader.usf b/Engine/Shaders/Private/PathTracing/PathTracingMaterialHitShader.usf
--- a/Engine/Shaders/Private/PathTracing/PathTracingMaterialHitShader.usf
+++ b/Engine/Shaders/Private/PathTracing/PathTracingMaterialHitShader.usf
@@ -804,6 +804,16 @@ RAY_TRACING_ENTRY_CLOSEST_HIT(PathTracingMaterialCHS,
 					{
 						Payload.SpecularProfileId = SLAB_SPECPROFILEID(BSDF);
 					}
+					Payload.GlintValue = 0;
+					if (BSDF_GETHASGLINT(BSDF))
+					{
+						Payload.GlintValue = SLAB_GLINT_VALUE(BSDF);
+						Payload.GlintUV = SLAB_GLINT_UV(BSDF);
+						
+						const float PlaceholderEpsilon = 0.001f;
+						Payload.GlintUVdx = float2(PlaceholderEpsilon, 0); // STRATA_TODO: SLAB_GLINT_UVDDX(BSDF)
+						Payload.GlintUVdy = float2(0, PlaceholderEpsilon); // STRATA_TODO: SLAB_GLINT_UVDDY(BSDF)
+					}
 					Payload.WeightV             = BSDF.LuminanceWeightV;
 					Payload.TransmittanceN      = BSDF.TransmittanceAboveAlongN;
 					Payload.CoverageAboveAlongN = BSDF.CoverageAboveAlongN;
diff --git a/Engine/Shaders/Private/Strata/Glint/GlintThirdParty.ush b/Engine/Shaders/Private/Strata/Glint/GlintThirdParty.ush
--- a/Engine/Shaders/Private/Strata/Glint/GlintThirdParty.ush
+++ b/Engine/Shaders/Private/Strata/Glint/GlintThirdParty.ush
@@ -28,6 +28,8 @@ struct FGlintLutParams
 	float MinLDist;
 };
 
+#define USE_GLINT_AA (!RAYCALLABLESHADER && !RAYHITGROUPSHADER && !RAYGENSHADER && !RAYMISSSHADER)
+
 // Geometric Glint Anti-Aliasing parameters
 #define KernelSize						0.5
 #define GGAAFilter						1
@@ -259,6 +261,7 @@ float P22_M(
 	//=========================== Contribution 1 ==============================
 	//================================ GGAA ===================================
 	//=========================================================================
+#if USE_GLINT_AA
 	if (GGAAFilter)
 	{
 		// As for the distribution, we transform the filtering kernel
@@ -285,6 +288,7 @@ float P22_M(
 	//========================= END Contribution 1 ============================
 	//=========================================================================
 	else // Without geometric glint anti-aliasing
+#endif
 	{
 		P_20_o = View.GlintTexture.SampleLevel(
 			View.GlintSampler,
@@ -387,7 +391,7 @@ float P22__glint_discrete_LOD(
 //=============================================================================
 //======== Evaluation of the procedural physically based glinty BRDF ==========
 //=============================================================================
-float f_P(float3 wo, float3 wi, float3 sigmas_rho, float LogMicrofacetDensity, float2 TexCoord, float2 TexCoordDDX, float2 TexCoordDDY)
+float f_P(float3 wo, float3 wi, float3 sigmas_rho, float LogMicrofacetDensity, float2 TexCoord, float2 TexCoordDDX, float2 TexCoordDDY, bool bIncludeGeometryTerm=false)
 {
 	if (wo.z <= 0.)
 		return 0;
@@ -415,8 +419,9 @@ float f_P(float3 wo, float3 wi, float3 sigmas_rho, float LogMicrofacetDensity, f
 	// Normal to slope
 	float2 slope_h = float2(-wh.x / wh.z, -wh.y / wh.z);
 
-	float2 slope_dx;
-	float2 slope_dy;
+	float2 slope_dx = 0;
+	float2 slope_dy = 0;
+#if USE_GLINT_AA
 	if (UseHemisDerivatives)
 	{
 		// Derivatives in the projected hemispherical domain
@@ -429,6 +434,7 @@ float f_P(float3 wo, float3 wi, float3 sigmas_rho, float LogMicrofacetDensity, f
 		slope_dx = ddx(slope_h);
 		slope_dy = ddy(slope_h);
 	}
+#endif
 
 	float D_P = 0;
 	float P22_P = 0.;
@@ -518,7 +524,14 @@ float f_P(float3 wo, float3 wi, float3 sigmas_rho, float LogMicrofacetDensity, f
 //	// the cosine weight in the rendering equation
 //	return (F * G * D_P) / (4. * wo.z);
 
-	return D_P / (4. * wo.z);
+	float G = 1.f;
+	if (bIncludeGeometryTerm)
+	{
+		float G1wowh = min(1., 2. * wh.z * wo.z / dot(wo, wh));
+		float G1wiwh = min(1., 2. * wh.z * wi.z / dot(wi, wh));
+		G = G1wowh * G1wiwh;																						// TODO check 
+	}
+	return G * D_P / (4. * wo.z);
 }
 
 // Glint Sampling
diff --git a/Engine/Shaders/Private/Strata/Strata.ush b/Engine/Shaders/Private/Strata/Strata.ush
--- a/Engine/Shaders/Private/Strata/Strata.ush
+++ b/Engine/Shaders/Private/Strata/Strata.ush
@@ -8,7 +8,6 @@
 	#error STRATA_ENABLED needs to be defined
 #endif
 
-
 #if STRATA_ENABLED
 
 #include "../ShadingEnergyConservation.ush"
@@ -107,9 +106,7 @@
 #endif
 #endif
 
-// Disable glint when in tracing mode for now since derivative are not available.
-#define STRATA_GLINTS_DERIVATIVE_SUPPORTED (!RAYCALLABLESHADER && !RAYHITGROUPSHADER && !RAYGENSHADER && !RAYMISSSHADER)
-#define STRATA_GLINTS_ENABLED (STRATA_COMPLEXSPECIALPATH && PLATFORM_ENABLES_STRATA_GLINTS && STRATA_GLINTS_DERIVATIVE_SUPPORTED)
+#define STRATA_GLINTS_ENABLED (STRATA_COMPLEXSPECIALPATH && PLATFORM_ENABLES_STRATA_GLINTS)
 
 #define STRATA_SPECPROFILE_ENABLED (STRATA_COMPLEXSPECIALPATH && PLATFORM_ENABLES_STRATA_SPECULAR_PROFILE)
 
@@ -875,6 +872,28 @@ uint GetSpecularProfileParameterization(float In)
 	return In > 0.f ? 0u : 1u;
 }
 
+///////////////////////////////////////////////////////////////////////////////
+// Glint
+
+uint2 PackGlints(float InValue, float2 InUV)
+{
+	const uint Packed16 = f32tof16(InValue);
+	const uint Packed8L = (Packed16 & 0xFF) << 24;
+	const uint Packed8H = (Packed16 >> 8  ) << 24;
+
+	uint2 Out;
+	Out.x = PackR24F(InUV.x) | Packed8L;
+	Out.y = PackR24F(InUV.y) | Packed8H;
+	return Out;
+}
+
+void UnpackGlints(in uint2 InPacked, inout float OutValue, inout float2 OutUV)
+{
+	OutUV.x	= UnpackR24F(InPacked.x);
+	OutUV.y	= UnpackR24F(InPacked.y);
+	OutValue= f16tof32((InPacked.x>>24) | ((InPacked.y>>16) & 0xFF00)); 
+}
+
 ///////////////////////////////////////////////////////////////////////////////
 // Fuzz
 
@@ -4167,10 +4186,7 @@ FStrataBSDF UnpackStrataBSDFIn(FStrataMaterialContainer StrataBuffer, inout FStr
 			uint2 GlintPackedData;
 			GlintPackedData.x = StrataLoadUint1(StrataBuffer, StrataAddressing);
 			GlintPackedData.y = StrataLoadUint1(StrataBuffer, StrataAddressing);
-
-			SLAB_GLINT_UV(OutBSDF).x	= UnpackR24F(GlintPackedData.x);
-			SLAB_GLINT_UV(OutBSDF).y	= UnpackR24F(GlintPackedData.y);
-			SLAB_GLINT_VALUE(OutBSDF)	= f16tof32((GlintPackedData.x>>24) | ((GlintPackedData.y>>16) & 0xFF00)); 
+			UnpackGlints(GlintPackedData, SLAB_GLINT_VALUE(OutBSDF), SLAB_GLINT_UV(OutBSDF));
 
 			SLAB_GLINT_UVDDX(OutBSDF)	= UnpackFloat2FromUInt(StrataLoadUint1(StrataBuffer, StrataAddressing));
 			SLAB_GLINT_UVDDY(OutBSDF)	= UnpackFloat2FromUInt(StrataLoadUint1(StrataBuffer, StrataAddressing));
diff --git a/Engine/Shaders/Private/Strata/StrataEvaluation.ush b/Engine/Shaders/Private/Strata/StrataEvaluation.ush
--- a/Engine/Shaders/Private/Strata/StrataEvaluation.ush
+++ b/Engine/Shaders/Private/Strata/StrataEvaluation.ush
@@ -259,39 +259,6 @@ half3 Strata_F_GGX(half3 F0, half3 F90, half VoH)
 //#endif
 }
 
-#if STRATA_COMPLEXSPECIALPATH
-// Conversion of GGX roughness to Beckmann roughness
-// Using Microfacet Distribution Function: To Change or Not to Change, That Is the Question.
-float GGXRoughnessToBeckmannRoughness(float Roughness)
-{
-	// Use a custom remapping to better match the overall highlight shape, in particular for lower roughness (i.e., <0.5). 
-	// The original mapping was Ctr=Roughness based on the paper
-	const float Ctr = pow(Roughness, 0.7f);
-	return saturate(0.8388 * Ctr+ 0.7 * Pow4(Ctr));
-}
-float2 GGXRoughnessToBeckmannRoughness(float2 Roughness)
-{
-	return float2(GGXRoughnessToBeckmannRoughness(Roughness.x), GGXRoughnessToBeckmannRoughness(Roughness.y));
-}
-
-struct FBeckmannDesc
-{
-	float2 Roughness;	// Roughness
-	float2 Alpha;		// Roughness squared
-	float2 Sigma;		// RMS slope of the microfacets
-	float  Rho;			// Slope correlation factor to be used with LEAN mapping. See http://igg.unistra.fr/People/chermain/assets/pdf/Chermain2021RealTime.pdf.
-};
-FBeckmannDesc GGXToBeckmann(float2 GGXSafeRoughness)
-{
-	FBeckmannDesc Beckmann;
-	Beckmann.Roughness	= GGXRoughnessToBeckmannRoughness(GGXSafeRoughness);
-	Beckmann.Alpha		= Square(Beckmann.Roughness);
-	Beckmann.Sigma		= Beckmann.Alpha * InvSqrt2;
-	Beckmann.Rho		= 0.0f;
-	return Beckmann;
-}
-#endif
-
 ///////////////////////////////////////////////////////////////////////////////////////////////////
 // BSDF evaluation (punctual/area)
 
@@ -688,7 +655,7 @@ FStrataEvaluateResult StrataEvaluateBSDFCommon(FStrataBSDFContext BSDFContext, F
 				const uint ProfileID = GetSpecularProfileId(SLAB_SPECPROFILEID(BSDFContext.BSDF));
 				const bool bUseHalfVector = GetSpecularProfileParameterization(SLAB_SPECPROFILEID(BSDFContext.BSDF)) == 1;
 				const float3 ProfileCoord = bUseHalfVector ? float3(BSDFContext.Context.VoH, BSDFContext.Context.NoH, ProfileID) : float3(BSDFContext.SatNoV, BSDFContext.SatNoL, ProfileID);
-				Sample.IntegratedSpecularValue *= View.SpecularProfileTexture.SampleLevel(View.SpecularProfileSampler, ProfileCoord, 0);
+				Sample.IntegratedSpecularValue *= View.SpecularProfileTexture.SampleLevel(View.SpecularProfileSampler, ProfileCoord, 0).xyz;
 			}
 			#endif
 
diff --git a/Engine/Shaders/Private/Strata/StrataExport.ush b/Engine/Shaders/Private/Strata/StrataExport.ush
--- a/Engine/Shaders/Private/Strata/StrataExport.ush
+++ b/Engine/Shaders/Private/Strata/StrataExport.ush
@@ -798,11 +798,9 @@ void PackStrataOut(
 							#if STRATA_GLINTS_ENABLED
 							if (BSDF_GETHASGLINT(BSDF))
 							{
-								const uint Packed16 = f32tof16(SLAB_GLINT_VALUE(BSDF));
-								const uint Packed8L = (Packed16 & 0xFF) << 24;
-								const uint Packed8H = (Packed16 >> 8  ) << 24;
-								STRATA_STORE_UINT1(PackR24F(SLAB_GLINT_UV(BSDF).x) | Packed8L);
-								STRATA_STORE_UINT1(PackR24F(SLAB_GLINT_UV(BSDF).y) | Packed8H);
+								const uint2 PackedGlintData = PackGlints(SLAB_GLINT_VALUE(BSDF), SLAB_GLINT_UV(BSDF));
+								STRATA_STORE_UINT1(PackedGlintData.x);
+								STRATA_STORE_UINT1(PackedGlintData.y);
 
 								// We have to store derivative when rasterizing the mesh. 
 								// That cannot be reconstructed in a deferred fashion (e.g. edge of a mesh overlaping with backgroud would look wrongly aliased)
diff --git a/Engine/Source/Runtime/Renderer/Private/DeferredShadingRenderer.cpp b/Engine/Source/Runtime/Renderer/Private/DeferredShadingRenderer.cpp
--- a/Engine/Source/Runtime/Renderer/Private/DeferredShadingRenderer.cpp
+++ b/Engine/Source/Runtime/Renderer/Private/DeferredShadingRenderer.cpp
@@ -1875,7 +1875,7 @@ bool FDeferredShadingSceneRenderer::SetupRayTracingPipelineStates(FRDGBuilder& G
 		{
 			// This view only needs the path tracing raygen shaders as all other
 			// passes should be disabled.
-			PreparePathTracing(ViewFamily, RayGenShaders);
+			PreparePathTracing(ViewFamily, *Scene, RayGenShaders);
 		}
 		else
 		{
diff --git a/Engine/Source/Runtime/Renderer/Private/DeferredShadingRenderer.h b/Engine/Source/Runtime/Renderer/Private/DeferredShadingRenderer.h
--- a/Engine/Source/Runtime/Renderer/Private/DeferredShadingRenderer.h
+++ b/Engine/Source/Runtime/Renderer/Private/DeferredShadingRenderer.h
@@ -1172,7 +1172,7 @@ private:
 	static void PrepareRayTracingTranslucency(const FViewInfo& View, TArray<FRHIRayTracingShader*>& OutRayGenShaders);
 	static void PrepareRayTracingVolumetricFogShadows(const FViewInfo& View, const FScene& Scene, TArray<FRHIRayTracingShader*>& OutRayGenShaders);
 	static void PrepareRayTracingDebug(const FSceneViewFamily& ViewFamily, TArray<FRHIRayTracingShader*>& OutRayGenShaders);
-	static void PreparePathTracing(const FSceneViewFamily& ViewFamily, TArray<FRHIRayTracingShader*>& OutRayGenShaders);
+	static void PreparePathTracing(const FSceneViewFamily& ViewFamily, const FScene& Scene, TArray<FRHIRayTracingShader*>& OutRayGenShaders);
 	static void PrepareLumenHardwareRayTracingScreenProbeGather(const FViewInfo& View, TArray<FRHIRayTracingShader*>& OutRayGenShaders);
 	static void PrepareLumenHardwareRayTracingShortRangeAO(const FViewInfo& View, TArray<FRHIRayTracingShader*>& OutRayGenShaders);
 	static void PrepareLumenHardwareRayTracingScreenProbeGatherDeferredMaterial(const FViewInfo& View, TArray<FRHIRayTracingShader*>& OutRayGenShaders);
diff --git a/Engine/Source/Runtime/Renderer/Private/PathTracing.cpp b/Engine/Source/Runtime/Renderer/Private/PathTracing.cpp
--- a/Engine/Source/Runtime/Renderer/Private/PathTracing.cpp
+++ b/Engine/Source/Runtime/Renderer/Private/PathTracing.cpp
@@ -868,8 +868,8 @@ class FPathTracingRG : public FGlobalShader
 	SHADER_USE_ROOT_PARAMETER_STRUCT(FPathTracingRG, FGlobalShader)
 
 	class FCompactionType : SHADER_PERMUTATION_INT("PATH_TRACER_USE_COMPACTION", 2);
-
-	using FPermutationDomain = TShaderPermutationDomain<FCompactionType>;
+	class FStrataComplexSpecialMaterial : SHADER_PERMUTATION_BOOL("PATH_TRACER_USE_STRATA_SPECIAL_COMPLEX_MATERIAL");
+	using FPermutationDomain = TShaderPermutationDomain<FCompactionType, FStrataComplexSpecialMaterial>;
 
 	static bool ShouldCompilePermutation(const FGlobalShaderPermutationParameters& Parameters)
 	{
@@ -2161,15 +2161,25 @@ class FPathTracingCompositorPS : public FGlobalShader
 };
 IMPLEMENT_SHADER_TYPE(, FPathTracingCompositorPS, TEXT("/Engine/Private/PathTracing/PathTracingCompositingPixelShader.usf"), TEXT("CompositeMain"), SF_Pixel);
 
-void FDeferredShadingSceneRenderer::PreparePathTracing(const FSceneViewFamily& ViewFamily, TArray<FRHIRayTracingShader*>& OutRayGenShaders)
+static FPathTracingRG::FPermutationDomain GetPathTracingRGPermutation(const FScene& Scene)
+{
+	const int CompactionType = CVarPathTracingCompaction.GetValueOnRenderThread();
+	const bool bHasComplexSpecialRenderPath = Strata::IsStrataEnabled() && Scene.StrataSceneData.bUsesComplexSpecialRenderPath;
+
+	FPathTracingRG::FPermutationDomain Out;
+	Out.Set<FPathTracingRG::FCompactionType>(CompactionType);
+	Out.Set<FPathTracingRG::FStrataComplexSpecialMaterial>(bHasComplexSpecialRenderPath);
+	return Out;
+}
+
+void FDeferredShadingSceneRenderer::PreparePathTracing(const FSceneViewFamily& ViewFamily, const FScene& Scene, TArray<FRHIRayTracingShader*>& OutRayGenShaders)
 {
 	if (ViewFamily.EngineShowFlags.PathTracing
 		&& ShouldCompilePathTracingShadersForProject(ViewFamily.GetShaderPlatform()))
 	{
 		// Declare all RayGen shaders that require material closest hit shaders to be bound
 		const int CompactionType = CVarPathTracingCompaction.GetValueOnRenderThread();
-		FPathTracingRG::FPermutationDomain PermutationVector;
-		PermutationVector.Set<FPathTracingRG::FCompactionType>(CompactionType);
+		FPathTracingRG::FPermutationDomain PermutationVector = GetPathTracingRGPermutation(Scene);
 		{
 			auto RayGenShader = GetGlobalShaderMap(ViewFamily.GetShaderPlatform())->GetShader<FPathTracingRG>(PermutationVector);
 			OutRayGenShaders.Add(RayGenShader.GetRayTracingShader());
@@ -2577,9 +2587,8 @@ void FDeferredShadingSceneRenderer::RenderPathTracing(
 				}
 				PathStateData = GraphBuilder.CreateBuffer(FRDGBufferDesc::CreateStructuredDesc(sizeof(FPathTracingPackedPathState), NumPaths), TEXT("PathTracer.PathStateData"));
 			}
-			FPathTracingRG::FPermutationDomain PermutationVector;
-			PermutationVector.Set<FPathTracingRG::FCompactionType>(CompactionType);
-			TShaderMapRef<FPathTracingRG> RayGenShader(View.ShaderMap, PermutationVector);
+
+			TShaderMapRef<FPathTracingRG> RayGenShader(View.ShaderMap, GetPathTracingRGPermutation(*Scene));
 			FPathTracingRG::FParameters* PreviousPassParameters = nullptr;
 			// Divide each tile among all the active GPUs (interleaving scanlines)
 			// The assumption is that the tiles are as big as possible, hopefully covering the entire screen
