Substrate - Fixed unlit to also work with refraction.

diff --git a/Engine/Shaders/Private/DistortAccumulatePS.usf b/Engine/Shaders/Private/DistortAccumulatePS.usf
--- a/Engine/Shaders/Private/DistortAccumulatePS.usf
+++ b/Engine/Shaders/Private/DistortAccumulatePS.usf
@@ -96,6 +96,29 @@ void Main(
 	// Retrieve rough refraction variance
 	float Sigma = 0.0f;
 	float3 RefractionWorldNormal = MaterialParameters.CameraVector;
+
+#if STRATA_OPTIMIZED_UNLIT
+
+	// Unlit forces a single BSDF
+	FStrataBSDF UnlitBSDF = PixelMaterialInputs.FrontMaterial.InlinedBSDF;
+	UnlitBSDF.StrataSanitizeBSDF();
+
+	RefractionWorldNormal = UNLIT_NORMAL(UnlitBSDF);
+	const float MaterialIOR = GetMaterialRefractionIOR(RefractionData); // Unlit does not have any surface F0 so we only use the root node input.
+
+	const bool TryToClip = true; // In this case rough refractions are not supported, we can thus always clip.
+
+#if DISTORTION_ACCOUNT_FOR_COVERAGE
+	const float RefractionCoverage = UnlitBSDF.Coverage;	// Coverage of the material. STRATA_TODO allow coverage operator with unlit material to make this not 1 (allowed by legacy)
+#else
+	const float RefractionCoverage = 1.0;					// Behavior where coverage is always ignored (legacy support)
+#endif
+
+	const float RefractionLobeVariance = 0.0f;
+	OutVarianceCoverage = float2(RefractionLobeVariance, RefractionCoverage);
+
+#else
+
 	if (StrataPixelHeader.StrataTree.BSDFCount > 0)
 	{
 		const float3 V = MaterialParameters.CameraVector;
@@ -126,26 +149,31 @@ void Main(
 	}
 	TopF0 = SanitizeF0(TopF0);	// Sanitize, in case a Add node has been used
 
+	const float MaterialIOR = DielectricF0RGBToIor(TopF0);
+
 	const FStrataOperator RootOperator	= StrataPixelHeader.StrataTree.Operators[StrataData.OperatorIndex];
 #if DISTORTION_ACCOUNT_FOR_COVERAGE
 	const float RefractionCoverage		= RootOperator.Coverage;														// Coverage of the material
 #else
 	const float RefractionCoverage		= 1.0;																			// Behavior where coverage is always ignored (legacy support)
 #endif
-	const float RefractionDepthMeter	= length(MaterialParameters.WorldPosition_CamRelative) * CENTIMETER_TO_METER;	// Distance for camera in meters
 	const float RefractionLobeVariance	= Sigma;																		// The lobe variance that can be added to in the en recover a single lobe.
 
 	// Write out variance (as this is more linear to accumulate than roughness).
 	OutVarianceCoverage = float2(RefractionLobeVariance, RefractionCoverage);
-	OutClosestDepthMeter = RefractionDepthMeter;
 
 	// Only allow clipping if roughness is close to 0.
 	const bool TryToClip = false;  //OutVarianceCoverage.x < 0.01f; // STRATA_TODO || RefractionCoverage <= 0.0f;	// Have an option to account for coverage or not?, ignored by default
 
+#endif
+
+	const float RefractionDepthMeter = length(MaterialParameters.WorldPosition_CamRelative) * CENTIMETER_TO_METER;	// Distance for camera in meters
+	OutClosestDepthMeter = RefractionDepthMeter;
+
 	// Compute UV distortion
 	float2 BufferUVDistortion = ComputeBufferUVDistortion(
 		MaterialParameters, PixelMaterialInputs, ResolvedView,
-		RefractionWorldNormal, DielectricF0RGBToIor(TopF0),
+		RefractionWorldNormal, MaterialIOR,
 		DistortionParams, ScreenUV, RefractionData, TryToClip, EyeIndex);
 
 	// Sample depth at distortion offset
diff --git a/Engine/Shaders/Private/Strata/Strata.ush b/Engine/Shaders/Private/Strata/Strata.ush
--- a/Engine/Shaders/Private/Strata/Strata.ush
+++ b/Engine/Shaders/Private/Strata/Strata.ush
@@ -559,6 +559,7 @@ struct FStrataBSDF
 ///////////////////////////////////////////////////////////////////////////////
 // Unlit BSDF members
 #define UNLIT_TRANSMITTANCE(X)			X.VGPRs[0].xyz
+#define UNLIT_NORMAL(X)					X.VGPRs[1].xyz
 
 ///////////////////////////////////////////////////////////////////////////////
 // Hair BSDF members
@@ -1509,13 +1510,14 @@ FStrataData GetStrataVolumeFogCloudBSDF(float3 Albedo, float3 Extinction, float3
 	return StrataData;
 }
 
-FStrataData GetStrataUnlitBSDF(float3 Emissive, float3 TransmittanceColor)
+FStrataData GetStrataUnlitBSDF(float3 Emissive, float3 TransmittanceColor, float3 Normal)
 {
 	FStrataData StrataData = GetInitialisedStrataData();
 
 	BSDF_SETTYPE					(StrataData.InlinedBSDF, STRATA_BSDF_TYPE_UNLIT);
 	BSDF_SETEMISSIVE				(StrataData.InlinedBSDF, Emissive);
 	UNLIT_TRANSMITTANCE				(StrataData.InlinedBSDF) = TransmittanceColor;
+	UNLIT_NORMAL					(StrataData.InlinedBSDF) = Normal;
 
 #if STRATA_INLINE_SHADING
 	StrataData.InlinedBSDF.Coverage = 1.0f;
@@ -2747,6 +2749,7 @@ void FStrataBSDF::StrataSanitizeBSDF()
 			BSDF_SETEMISSIVE(this, max(BSDF_GETEMISSIVE(this), 0.0));
 		#endif
 			UNLIT_TRANSMITTANCE(this) = saturate(UNLIT_TRANSMITTANCE(this));
+			UNLIT_NORMAL(this) = normalize(UNLIT_NORMAL(this));
 		}
 	}
 }
diff --git a/Engine/Shaders/Private/Strata/StrataLegacyConversion.ush b/Engine/Shaders/Private/Strata/StrataLegacyConversion.ush
--- a/Engine/Shaders/Private/Strata/StrataLegacyConversion.ush
+++ b/Engine/Shaders/Private/Strata/StrataLegacyConversion.ush
@@ -481,31 +481,31 @@ FStrataData StrataConvertLegacyMaterialStatic(
 		const float GreyTransmittance = saturate(1.0f - Opacity);
 	#if MATERIALBLENDING_SOLID || MATERIALBLENDING_MASKED
 		// Opaque materials only write emissive.
-		Out = GetStrataUnlitBSDF(Emissive, 0.0f);
+		Out = GetStrataUnlitBSDF(Emissive, 0.0f, RawNormal);
 	#elif MATERIALBLENDING_ALPHACOMPOSITE
 		// Uses (Luminance, black transmittance) for premultiplied alpha blending. Opacity is controled using coverage
-		Out = GetStrataUnlitBSDF(Emissive, 0.0f);
+		Out = GetStrataUnlitBSDF(Emissive, 0.0f, RawNormal);
 		#if STRATA_INLINE_SHADING
 			Out.InlinedBSDF.Coverage = Opacity;
 		#endif
 	#elif MATERIALBLENDING_ALPHAHOLDOUT
 		// The MATERIALBLENDING_ALPHAHOLDOUT case needs to be before the MATERIALBLENDING_TRANSLUCENT case, as alphaholdout material have both cases defined
-		Out = GetStrataUnlitBSDF(0.0f, GreyTransmittance);
+		Out = GetStrataUnlitBSDF(0.0f, GreyTransmittance, RawNormal);
 	#elif MATERIALBLENDING_TRANSLUCENT
 		// Uses (Luminance, black transmittance) for translucency blending. Opacity is controled using coverage
-		Out = GetStrataUnlitBSDF(Emissive * Opacity, 0.0f);
+		Out = GetStrataUnlitBSDF(Emissive * Opacity, 0.0f, RawNormal);
 		#if STRATA_INLINE_SHADING
 			Out.InlinedBSDF.Coverage = Opacity;
 		#endif
 	#elif MATERIALBLENDING_ADDITIVE
 		// Emissive is weighted by opacity in the legacy material.
-		Out = GetStrataUnlitBSDF(Emissive * Opacity, 1.0f);
+		Out = GetStrataUnlitBSDF(Emissive * Opacity, 1.0f, RawNormal);
 	#elif MATERIALBLENDING_MODULATE
 		// Setting up emissive as the transmittance color. It is not clamped in case brightening is required.
-		Out = GetStrataUnlitBSDF(0.0f, Emissive);
+		Out = GetStrataUnlitBSDF(0.0f, Emissive, RawNormal);
 	#else
 		// Required default for some materials such as Editor UI
-		Out = GetStrataUnlitBSDF(Emissive, 0.0f);
+		Out = GetStrataUnlitBSDF(Emissive, 0.0f, RawNormal);
 	#endif
 	}
 	else if (ShadingModel == SHADINGMODELID_HAIR)
diff --git a/Engine/Source/Runtime/Engine/Classes/Materials/MaterialExpressionStrata.h b/Engine/Source/Runtime/Engine/Classes/Materials/MaterialExpressionStrata.h
--- a/Engine/Source/Runtime/Engine/Classes/Materials/MaterialExpressionStrata.h
+++ b/Engine/Source/Runtime/Engine/Classes/Materials/MaterialExpressionStrata.h
@@ -465,6 +465,12 @@ class UMaterialExpressionStrataUnlitBSDF : public UMaterialExpressionStrataBSDF
 	UPROPERTY()
 	FExpressionInput TransmittanceColor;
 
+	/**
+	 * The surface normal. Only used for refraction effects when `IOR` or `pixel normal offset` modes are selected.
+	 */
+	UPROPERTY()
+	FExpressionInput Normal;
+
 	//~ Begin UMaterialExpression Interface
 #if WITH_EDITOR
 	virtual int32 Compile(class FMaterialCompiler* Compiler, int32 OutputIndex) override;
diff --git a/Engine/Source/Runtime/Engine/Private/Materials/HLSLMaterialTranslator.cpp b/Engine/Source/Runtime/Engine/Private/Materials/HLSLMaterialTranslator.cpp
--- a/Engine/Source/Runtime/Engine/Private/Materials/HLSLMaterialTranslator.cpp
+++ b/Engine/Source/Runtime/Engine/Private/Materials/HLSLMaterialTranslator.cpp
@@ -12316,12 +12316,13 @@ int32 FHLSLMaterialTranslator::StrataVolumetricFogCloudBSDF(int32 Albedo, int32
 	);
 }
 
-int32 FHLSLMaterialTranslator::StrataUnlitBSDF(int32 EmissiveColor, int32 TransmittanceColor)
+int32 FHLSLMaterialTranslator::StrataUnlitBSDF(int32 EmissiveColor, int32 TransmittanceColor, int32 Normal)
 {
 	return AddCodeChunk(
-		MCT_Strata, TEXT("GetStrataUnlitBSDF(%s, %s)"),
+		MCT_Strata, TEXT("GetStrataUnlitBSDF(%s, %s, %s)"),
 		*StrataGetCastParameterCode(EmissiveColor,		MCT_Float3),
-		*StrataGetCastParameterCode(TransmittanceColor,	MCT_Float3)
+		*StrataGetCastParameterCode(TransmittanceColor, MCT_Float3),
+		*StrataGetCastParameterCode(Normal,				MCT_Float3)
 	);
 }
 
diff --git a/Engine/Source/Runtime/Engine/Private/Materials/HLSLMaterialTranslator.h b/Engine/Source/Runtime/Engine/Private/Materials/HLSLMaterialTranslator.h
--- a/Engine/Source/Runtime/Engine/Private/Materials/HLSLMaterialTranslator.h
+++ b/Engine/Source/Runtime/Engine/Private/Materials/HLSLMaterialTranslator.h
@@ -1189,7 +1189,7 @@ protected:
 		int32 CustomTangent_Tangent,
 		FStrataOperator* PromoteToOperator) override;
 	virtual int32 StrataVolumetricFogCloudBSDF(int32 Albedo, int32 Extinction, int32 EmissiveColor, int32 AmbientOcclusion) override;
-	virtual int32 StrataUnlitBSDF(int32 EmissiveColor, int32 TransmittanceColor) override;
+	virtual int32 StrataUnlitBSDF(int32 EmissiveColor, int32 TransmittanceColor, int32 Normal) override;
 	virtual int32 StrataHairBSDF(int32 BaseColor, int32 Scatter, int32 Specular, int32 Roughness, int32 Backlit, int32 EmissiveColor, int32 Tangent, const FString& SharedLocalBasisIndexMacro, FStrataOperator* PromoteToOperator) override;
 	virtual int32 StrataEyeBSDF(int32 DiffuseColor, int32 Roughness, int32 IrisMask, int32 IrisDistance, int32 EmissiveColor, int32 CorneaNormal, int32 IrisNormal, int32 IrisPlaneNormal, int32 SSSProfileId, const FString& SharedLocalBasisIndexMacro, FStrataOperator* PromoteToOperator) override;
 	virtual int32 StrataSingleLayerWaterBSDF(
diff --git a/Engine/Source/Runtime/Engine/Private/Materials/MaterialExpressions.cpp b/Engine/Source/Runtime/Engine/Private/Materials/MaterialExpressions.cpp
--- a/Engine/Source/Runtime/Engine/Private/Materials/MaterialExpressions.cpp
+++ b/Engine/Source/Runtime/Engine/Private/Materials/MaterialExpressions.cpp
@@ -23871,7 +23871,8 @@ int32 UMaterialExpressionStrataLightFunction::Compile(class FMaterialCompiler* C
 {
 	int32 OutputCodeChunk = Compiler->StrataUnlitBSDF(
 		CompileWithDefaultFloat3(Compiler, Color, 0.0f, 0.0f, 0.0f),
-		Compiler->Constant(1.0f));	// Opacity / Transmittance is ignored by light functions.
+		Compiler->Constant(1.0f),				// Opacity / Transmittance is ignored by light functions.
+		Compiler->Constant3(0.0f, 0.0f, 1.0f));	// place holder normal
 
 	return OutputCodeChunk;
 }
@@ -23942,7 +23943,8 @@ int32 UMaterialExpressionStrataPostProcess::Compile(class FMaterialCompiler* Com
 
 	int32 OutputCodeChunk = Compiler->StrataUnlitBSDF(
 		CompileWithDefaultFloat3(Compiler, Color, 0.0f, 0.0f, 0.0f),
-		TransmittanceCodeChunk);
+		TransmittanceCodeChunk,
+		Compiler->Constant3(0.0f, 0.0f, 1.0f));	// place holder normal
 
 	return OutputCodeChunk;
 }
@@ -24016,7 +24018,8 @@ int32 UMaterialExpressionStrataUI::Compile(class FMaterialCompiler* Compiler, in
 
 	int32 OutputCodeChunk = Compiler->StrataUnlitBSDF(
 		CompileWithDefaultFloat3(Compiler, Color, 0.0f, 0.0f, 0.0f),
-		TransmittanceCodeChunk);
+		TransmittanceCodeChunk,
+		Compiler->Constant3(0.0f, 0.0f, 1.0f));	// place holder normal
 
 	return OutputCodeChunk;
 }
@@ -24211,7 +24214,8 @@ int32 UMaterialExpressionStrataUnlitBSDF::Compile(class FMaterialCompiler* Compi
 {
 	int32 OutputCodeChunk = Compiler->StrataUnlitBSDF(
 		CompileWithDefaultFloat3(Compiler, EmissiveColor, 0.0f, 0.0f, 0.0f),
-		CompileWithDefaultFloat3(Compiler, TransmittanceColor, 1.0f, 1.0f, 1.0f));
+		CompileWithDefaultFloat3(Compiler, TransmittanceColor, 1.0f, 1.0f, 1.0f),
+		CompileWithDefaultNormalWS(Compiler, Normal));
 
 	return OutputCodeChunk;
 }
@@ -24236,6 +24240,9 @@ uint32 UMaterialExpressionStrataUnlitBSDF::GetInputType(int32 InputIndex)
 	case 1:
 		return MCT_Float3;
 		break;
+	case 2:
+		return MCT_Float3;
+		break;
 	}
 
 	check(false);
diff --git a/Engine/Source/Runtime/Engine/Public/MaterialCompiler.h b/Engine/Source/Runtime/Engine/Public/MaterialCompiler.h
--- a/Engine/Source/Runtime/Engine/Public/MaterialCompiler.h
+++ b/Engine/Source/Runtime/Engine/Public/MaterialCompiler.h
@@ -577,7 +577,7 @@ public:
 		int32 CustomTangent_Tangent,
 		FStrataOperator* PromoteToOperator) = 0;
 	virtual int32 StrataVolumetricFogCloudBSDF(int32 Albedo, int32 Extinction, int32 EmissiveColor, int32 AmbientOcclusion) = 0;
-	virtual int32 StrataUnlitBSDF(int32 EmissiveColor, int32 TransmittanceColor) = 0;
+	virtual int32 StrataUnlitBSDF(int32 EmissiveColor, int32 TransmittanceColor, int32 Normal) = 0;
 	virtual int32 StrataHairBSDF(int32 BaseColor, int32 Scatter, int32 Specular, int32 Roughness, int32 Backlit, int32 EmissiveColor, int32 Tangent, const FString& SharedLocalBasisIndexMacro, FStrataOperator* PromoteToOperator) = 0;
 	virtual int32 StrataEyeBSDF(int32 DiffuseAlbedo, int32 Roughness, int32 IrisMask, int32 IrisDistance, int32 IrisNormal, int32 IrisPlaneNormal, int32 SSSProfileId, int32 EmissiveColor, int32 CorneaNormal, const FString& SharedLocalBasisIndexMacro, FStrataOperator* PromoteToOperator) = 0;
 	virtual int32 StrataSingleLayerWaterBSDF(
@@ -1235,9 +1235,9 @@ public:
 		return Compiler->StrataVolumetricFogCloudBSDF(Albedo, Extinction, EmissiveColor, AmbientOcclusion);
 	}
 
-	virtual int32 StrataUnlitBSDF(int32 EmissiveColor, int32 TransmittanceColor) override
+	virtual int32 StrataUnlitBSDF(int32 EmissiveColor, int32 TransmittanceColor, int32 Normal) override
 	{
-		return Compiler->StrataUnlitBSDF(EmissiveColor, TransmittanceColor);
+		return Compiler->StrataUnlitBSDF(EmissiveColor, TransmittanceColor, Normal);
 	}
 
 	virtual int32 StrataHairBSDF(int32 BaseColor, int32 Scatter, int32 Specular, int32 Roughness, int32 Backlit, int32 EmissiveColor, int32 Tangent, const FString& SharedLocalBasisIndexMacro, FStrataOperator* PromoteToOperator) override
