Fix missing haziness roughness sanitization, causing NaN.

diff --git a/Engine/Shaders/Private/Strata/Strata.ush b/Engine/Shaders/Private/Strata/Strata.ush
--- a/Engine/Shaders/Private/Strata/Strata.ush
+++ b/Engine/Shaders/Private/Strata/Strata.ush
@@ -2682,6 +2682,13 @@ void FStrataBSDF::StrataSanitizeBSDF()
 			SLAB_F0(this) = SanitizeF0(SLAB_F0(this));
 			SLAB_F90(this) = saturate(SLAB_F90(this));
 			SLAB_ROUGHNESS(this) = SanitizeRoughness(SLAB_ROUGHNESS(this));
+			if (BSDF_GETHASHAZINESS(this))
+			{
+				FHaziness Haziness = UnpackHaziness(SLAB_HAZINESS(this));
+				Haziness.Roughness = SanitizeRoughness(Haziness.Roughness);
+				SLAB_HAZINESS(this) = PackHaziness(Haziness);
+			}
+			SLAB_ROUGHNESS(this) = SanitizeRoughness(SLAB_ROUGHNESS(this)); // 
 			SLAB_ANISOTROPY(this) = clamp(SLAB_ANISOTROPY(this), -1.f, 1.f);
 			// SLAB_HAZINESS is sanitize when packed, see PackHaziness.
 		#if STRATA_INLINE_SHADING && STRATA_OPAQUE_MATERIAL
