VolumetricCloud - fix tracing cut when traveling in and out of the cloud layer.
There are still reprojection issues.

diff --git a/Engine/Shaders/Private/VolumetricCloud.usf b/Engine/Shaders/Private/VolumetricCloud.usf
--- a/Engine/Shaders/Private/VolumetricCloud.usf
+++ b/Engine/Shaders/Private/VolumetricCloud.usf
@@ -486,10 +486,10 @@ void MainCommon(in FMaterialPixelParameters MaterialParameters, in float4 SvPosi
 		//
 		// Check tracing start and end position within the cloud layer
 		//
-		float2 tTop2 = 0.0f;
+		float2 tTop2 = -1.0f;
+		float2 tBottom2 = -1.0f;
 		if (RayIntersectSphereSolution(RayWorldOriginKm, Raydir, float4(RenderVolumetricCloudParameters.CloudLayerCenterKm, RenderVolumetricCloudParameters.TopRadiusKm), tTop2))
 		{
-			float2 tBottom2 = 0.0f;
 			if (RayIntersectSphereSolution(RayWorldOriginKm, Raydir, float4(RenderVolumetricCloudParameters.CloudLayerCenterKm, RenderVolumetricCloudParameters.BottomRadiusKm), tBottom2))
 			{
 				// If we see both intersection in front of us, keep the min/closest, otherwise the max/furthest
@@ -528,7 +528,7 @@ void MainCommon(in FMaterialPixelParameters MaterialParameters, in float4 SvPosi
 
 		if (TMax <= TMin || TMin > RenderVolumetricCloudParameters.TracingStartMaxDistance)
 		{
-			//OutColor0 = float4(1.0f, 0.0f, 0.0f, 1.0f);
+			//OutColor0 = float4(1.0f, 0.0f, 1.0f, 1.0f);
 			return;
 		}
 
@@ -575,7 +575,7 @@ void MainCommon(in FMaterialPixelParameters MaterialParameters, in float4 SvPosi
 
 		if (bDoDepthInteresectionAndTest && (TDepthBufferKm * KILOMETER_TO_CENTIMETER) < TMin)
 		{
-			//OutColor0 = float4(1.0f, 1.0f, 0.0f, 1.0f);
+			//OutColor0 = float4(0.0f, 0.0f, 1.0f, 1.0f);
 			return;
 		}
 
@@ -613,7 +613,7 @@ void MainCommon(in FMaterialPixelParameters MaterialParameters, in float4 SvPosi
 #endif // CLOUD_SAMPLE_SECOND_LIGHT
 
 
-		const bool SkyAtmospherePresentInScene = ResolvedView.SkyAtmospherePresentInScene > 0.0f;
+		const bool bTracingMaxDistanceModeFromCamera = RenderVolumetricCloudParameters.TracingMaxDistanceMode == 1;
 
 		// Clamp to tracing max distance according to the selected mode
 		if (RenderVolumetricCloudParameters.TracingMaxDistanceMode == 0)
@@ -625,12 +625,24 @@ void MainCommon(in FMaterialPixelParameters MaterialParameters, in float4 SvPosi
 		else // if (RenderVolumetricCloudParameters.TracingMaxDistanceMode == 1)
 		{
 			// Distance from the point of view
-			TMin = min(TMin, RenderVolumetricCloudParameters.TracingMaxDistance);
-			TMax = min(TMax, RenderVolumetricCloudParameters.TracingMaxDistance);
+			TMin = min(RenderVolumetricCloudParameters.TracingMaxDistance, TMin);
+
+			const float MaxBottom = max(tBottom2.x, tBottom2.y) * KILOMETER_TO_CENTIMETER;
+			if (all(tBottom2 > 0.0f) && (MaxBottom <= RenderVolumetricCloudParameters.TracingMaxDistance))
+			{
+				// Make sure we trace when traveling out of the bottom entry point but travel back again in for the specified range of distance.
+				// But we try to not trace outside of the cloud layer still for highly curved planets.
+				TMax = min(RenderVolumetricCloudParameters.TracingMaxDistance, max(tTop2.x, tTop2.y) * KILOMETER_TO_CENTIMETER);
+			}
+			else
+			{
+				TMax = min(TMax, RenderVolumetricCloudParameters.TracingMaxDistance);
+			}
 
 			// Exit if tracing a 0 length segment.
 			if ((TMax - TMin) <= 0.0f)
 			{
+				//OutColor0 = float4(1.0f, 1.0f, 0.0f, 1.0f);
 				return;
 			}
 		}
@@ -647,7 +659,8 @@ void MainCommon(in FMaterialPixelParameters MaterialParameters, in float4 SvPosi
 		float3 DistantSkyLightLuminance = 0.0f;
 		if (RenderVolumetricCloudParameters.EnableDistantSkyLightSampling)
 		{
-			if (SkyAtmospherePresentInScene)
+			const bool bSkyAtmospherePresentInScene = ResolvedView.SkyAtmospherePresentInScene > 0.0f;
+			if (bSkyAtmospherePresentInScene)
 			{
 				// TODO Only works for ground view. This should be spatially varying for space viewsspace view. 
 				// TODO Have a illuminance texture to get that overall ambient contribution? 
@@ -802,6 +815,14 @@ void MainCommon(in FMaterialPixelParameters MaterialParameters, in float4 SvPosi
 		#endif
 #endif
 
+			if (MaterialParameters.CloudSampleNormAltitudeInLayer <= 0.0f || MaterialParameters.CloudSampleNormAltitudeInLayer >= 1.0f)
+			{
+				// If we are out of the cloud volume, for instance when bTracingMaxDistanceModeFromCamera is true, 
+				// we can and must ignore all material to not change the look in case the material graph is not fading out the clouds below the layer.
+				t += StepT;
+				continue;
+			}
+
 			//////////////////////////////
 			// Sample the participating media material at the sample position.
 			//////////////////////////////
@@ -1531,6 +1552,8 @@ void MainCommon(in FMaterialPixelParameters MaterialParameters, in float4 SvPosi
 	{
 		OutDepth.x = ((GrayScaleTransmittance > 0.99) ? NoCloudDepth : tAP) * CENTIMETER_TO_KILOMETER; // using a small threshold on transmittance
 	}
+
+	//OutColor0.g = 0.3f; // Debug revealing reveal pixels that have been traced.
 }
 
 #endif // SHADER_RENDERVIEW_PS || SHADER_RENDERVIEW_CS
