VolumetricCloud - optimised shadow by using CloudSampleNormAltitudeInLayer correctly. Can result in performance win between 5 top 10%

diff --git a/Engine/Shaders/Private/VolumetricCloud.usf b/Engine/Shaders/Private/VolumetricCloud.usf
--- a/Engine/Shaders/Private/VolumetricCloud.usf
+++ b/Engine/Shaders/Private/VolumetricCloud.usf
@@ -1049,7 +1049,7 @@ void MainCommon(in FMaterialPixelParameters MaterialParameters, in float4 SvPosi
 					#endif // CLOUD_DEBUG_SAMPLES
 #endif
 
-					if (MaterialParameters.CloudSampleNormAltitudeInLayer < 0.0f || MaterialParameters.CloudSampleNormAltitudeInLayer > 1.0f)
+					if (MaterialParameters.CloudSampleNormAltitudeInLayer <= 0.0f || MaterialParameters.CloudSampleNormAltitudeInLayer >= 1.0f)
 					{
 						break; // Ignore remaining samples since we have just traveld out of the cloud layer.
 					}
@@ -1114,7 +1114,7 @@ void MainCommon(in FMaterialPixelParameters MaterialParameters, in float4 SvPosi
 						}
 #endif
 
-						if (MaterialParameters.CloudSampleNormAltitudeInLayer < 0.0f || MaterialParameters.CloudSampleNormAltitudeInLayer > 1.0f)
+						if (MaterialParameters.CloudSampleNormAltitudeInLayer <= 0.0f || MaterialParameters.CloudSampleNormAltitudeInLayer >= 1.0f)
 						{
 							break; // Ignore remaining samples since we have just traveld out of the cloud layer.
 						}
@@ -1247,7 +1247,7 @@ void MainCommon(in FMaterialPixelParameters MaterialParameters, in float4 SvPosi
 								ConvertCloudPixelMaterialInputsToWorkingColorSpace(PixelMaterialInputs);
 								MATVEC ShadowExtinctionCoefficients = SampleExtinctionCoefficients(PixelMaterialInputs);
 
-								if (MaterialParameters.CloudSampleNormAltitudeInLayer < 0.0f || MaterialParameters.CloudSampleNormAltitudeInLayer > 1.0f)
+								if (MaterialParameters.CloudSampleNormAltitudeInLayer <= 0.0f || MaterialParameters.CloudSampleNormAltitudeInLayer >= 1.0f)
 								{
 									ShadowExtinctionCoefficients = 0.0f; // Ignore remaining samples since we have just traveld out of the cloud layer.
 								}
