Path Tracing: Fix darkening in SSS

When the subsurface scattering radius become small, a fair amount of energy is lost which darkens the result in an un-desireable way. To avoid this, clamp the radius to a safe value and blend the subsurface albedo back towards diffuse instead. This keeps the overall energy contribution while avoiding small radii. It also improves the transition between radius 0 and radius slightly above 0.

diff --git a/Engine/Shaders/Private/PathTracing/Material/PathTracingMaterialSampling.ush b/Engine/Shaders/Private/PathTracing/Material/PathTracingMaterialSampling.ush
--- a/Engine/Shaders/Private/PathTracing/Material/PathTracingMaterialSampling.ush
+++ b/Engine/Shaders/Private/PathTracing/Material/PathTracingMaterialSampling.ush
@@ -371,6 +371,16 @@ float GetAverageRoughness(FPathTracingPayload Payload)
 #endif
 }
 
+void AdjustDiffuseSSSContribution(inout float3 DiffColor, inout float3 SSSColor, inout float3 Radius)
+{
+ 	// Below this radius, lerp to diffuse so we can avoid numerical issues in the random walk itself which lead to darkening
+	const float MinRadius = 0.02;
+	float3 Blend = saturate(Radius / MinRadius);
+	DiffColor += SSSColor * (1 - Blend);
+	SSSColor *= Blend;
+	Radius = max(Radius, MinRadius);
+}
+
 #if !SIMPLIFIED_MATERIAL_SHADER
 void AdjustPayloadAfterUnpack(inout FPathTracingPayload HitPayload, bool bApplyOverrides)
 {
@@ -391,6 +401,7 @@ void AdjustPayloadAfterUnpack(inout FPathTracingPayload HitPayload, bool bApplyO
 				// NOTE: if radius is 0 the sss color will get flipped back into diffuse
 				HitPayload.SubsurfaceColor = HitPayload.DiffuseColor;
 				HitPayload.DiffuseColor = 0;
+				AdjustDiffuseSSSContribution(HitPayload.DiffuseColor, HitPayload.SubsurfaceColor, HitPayload.SubsurfaceRadius);
 			}
 			break;
 		}
@@ -434,6 +445,10 @@ void AdjustPayloadAfterUnpack(inout FPathTracingPayload HitPayload, bool bApplyO
 				HitPayload.DiffuseColor = Front / S;
 				HitPayload.SubsurfaceColor = Back / S;
 			}
+			if (HitPayload.ShadingModelID != SHADINGMODELID_TWOSIDED_FOLIAGE)
+			{
+				AdjustDiffuseSSSContribution(HitPayload.DiffuseColor, HitPayload.SubsurfaceColor, HitPayload.SubsurfaceRadius);
+			}
 			if (bApplyOverrides)
 			{
 				HitPayload.SubsurfaceColor = HitPayload.SubsurfaceColor * View.DiffuseOverrideParameter.w + View.DiffuseOverrideParameter.xyz;
@@ -446,6 +461,7 @@ void AdjustPayloadAfterUnpack(inout FPathTracingPayload HitPayload, bool bApplyO
 			// NOTE: if radius is 0 the sss color will get flipped back into diffuse
 			HitPayload.SubsurfaceColor = HitPayload.DiffuseColor;
 			HitPayload.DiffuseColor = 0;
+			AdjustDiffuseSSSContribution(HitPayload.DiffuseColor, HitPayload.SubsurfaceColor, HitPayload.SubsurfaceRadius);
 			break;
 		}
 		case SHADINGMODELID_EYE:
@@ -455,6 +471,7 @@ void AdjustPayloadAfterUnpack(inout FPathTracingPayload HitPayload, bool bApplyO
 			float IrisMask = HitPayload.GetEyeIrisMask();
 			HitPayload.SubsurfaceColor = HitPayload.DiffuseColor - IrisMask * HitPayload.DiffuseColor;
 			HitPayload.DiffuseColor *= IrisMask;
+			AdjustDiffuseSSSContribution(HitPayload.DiffuseColor, HitPayload.SubsurfaceColor, HitPayload.SubsurfaceRadius);
 			break;
 		}
 		case SHADINGMODELID_HAIR:
