fix for primitive data in mobile preview

diff --git a/Engine/Shaders/Private/InstanceCulling/BuildInstanceDrawCommands.usf b/Engine/Shaders/Private/InstanceCulling/BuildInstanceDrawCommands.usf
--- a/Engine/Shaders/Private/InstanceCulling/BuildInstanceDrawCommands.usf
+++ b/Engine/Shaders/Private/InstanceCulling/BuildInstanceDrawCommands.usf
@@ -135,7 +135,13 @@ void WritePackedInstanceData(uint Offset, uint InstanceId, FPrimitiveSceneData P
 
 	// This will only write out the relative instance transform, but that's fine; the tile coordinate of the instance transform comes from the primitive data
 	float4 InstanceOriginAndId = InstanceData.LocalToWorld.M[3];
-	InstanceOriginAndId.w = asfloat(InstanceData.PrimitiveId | (CullingFlags << PRIMITIVE_ID_NUM_BITS));
+
+#if INSTANCE_CULLING_FLAGS_NUM_BITS != 4 || PRIMITIVE_ID_NUM_BITS != 20
+#error "Fix packing below as it assumes we can pack PrimitiveId and CullingFlags in 32 bits"
+#endif
+
+	// <<4 as a workaround for FXC turning asfloat(x)=0 if x is known to be denormalized
+	InstanceOriginAndId.w = asfloat((InstanceData.PrimitiveId | (CullingFlags << PRIMITIVE_ID_NUM_BITS)) << 4);
 
 	float4 InstanceTransform1 = InstanceData.LocalToWorld.M[0];
 	float4 InstanceTransform2 = InstanceData.LocalToWorld.M[1];
diff --git a/Engine/Shaders/Private/SceneData.ush b/Engine/Shaders/Private/SceneData.ush
--- a/Engine/Shaders/Private/SceneData.ush
+++ b/Engine/Shaders/Private/SceneData.ush
@@ -1161,8 +1161,11 @@ FSceneDataIntermediates GetSceneDataIntermediates(uint DrawInstanceId, float4 In
 	const uint PackedPrimitiveFlags = asuint(InstanceTransform1.w);
 	const uint PrimitiveFlags = (PackedPrimitiveFlags & 0x00FFFFFF);
 	const uint InstanceFlags = (PackedPrimitiveFlags >> 24u);
-	const uint PrimitiveId = asuint(InstanceOrigin.w) & PRIMITIVE_ID_MASK;
-	const uint CullingFlags = asuint(InstanceOrigin.w) >> PRIMITIVE_ID_NUM_BITS;
+	
+	// See WritePackedInstanceData
+	uint PrimitiveIdAndCullingFlags = asuint(InstanceOrigin.w) >> 4;
+	const uint PrimitiveId = PrimitiveIdAndCullingFlags & PRIMITIVE_ID_MASK;
+	const uint CullingFlags = PrimitiveIdAndCullingFlags >> PRIMITIVE_ID_NUM_BITS;
 
 	// Reconstruct InstanceData from a packed data
 	FInstanceSceneData InstanceData = (FInstanceSceneData)0;
