Path Tracer: Fix compilation error when substrate is enabled

diff --git a/Engine/Shaders/Private/PathTracing/Material/PathTracingMaterialSampling.ush b/Engine/Shaders/Private/PathTracing/Material/PathTracingMaterialSampling.ush
--- a/Engine/Shaders/Private/PathTracing/Material/PathTracingMaterialSampling.ush
+++ b/Engine/Shaders/Private/PathTracing/Material/PathTracingMaterialSampling.ush
@@ -401,7 +401,7 @@ void AdjustPayloadAfterUnpack(inout FPathTracingPayload HitPayload, bool bApplyO
 				// NOTE: if radius is 0 the sss color will get flipped back into diffuse
 				HitPayload.SubsurfaceColor = HitPayload.DiffuseColor;
 				HitPayload.DiffuseColor = 0;
-				AdjustDiffuseSSSContribution(HitPayload.DiffuseColor, HitPayload.SubsurfaceColor, HitPayload.SubsurfaceRadius);
+				AdjustDiffuseSSSContribution(HitPayload.DiffuseColor, HitPayload.SubsurfaceColor, HitPayload.MeanFreePath);
 			}
 			break;
 		}
