Fix mobile GPU scene shader compilation error

diff --git a/Engine/Shaders/Private/SceneData.ush b/Engine/Shaders/Private/SceneData.ush
--- a/Engine/Shaders/Private/SceneData.ush
+++ b/Engine/Shaders/Private/SceneData.ush
@@ -1206,7 +1206,6 @@ FSceneDataIntermediates GetSceneDataIntermediates(uint DrawInstanceId, float4 In
 	Intermediates.Primitive = Primitive;
 	// Primitive data that comes from per-instance vertex data
 	Intermediates.Primitive.Flags					= PrimitiveFlags;
-	Intermediates.Primitive.VisibilityFlags			= VisibilityFlags;
 	Intermediates.Primitive.LocalToWorld			= Intermediates.InstanceData.LocalToWorld;
 	Intermediates.Primitive.InvNonUniformScale		= Intermediates.InstanceData.InvNonUniformScale;
 	Intermediates.Primitive.WorldToLocal			= Intermediates.InstanceData.WorldToLocal;
