* Update hair BSDF
* Update engine test results.

diff --git a/Engine/Shaders/Private/HairBsdf.ush b/Engine/Shaders/Private/HairBsdf.ush
--- a/Engine/Shaders/Private/HairBsdf.ush
+++ b/Engine/Shaders/Private/HairBsdf.ush
@@ -35,7 +35,7 @@ float3 KajiyaKayDiffuseAttenuation(FGBufferData GBuffer, float3 L, float3 V, hal
 	float DiffuseScatter = (1 / PI) * lerp(NoL, KajiyaDiffuse, 0.33) * GBuffer.Metallic;
 	float Luma = Luminance(GBuffer.BaseColor);
     float3 BaseOverLuma = abs(GBuffer.BaseColor / max(Luma, MinValue));
-	float3 ScatterTint = Shadow < 1 ? pow(BaseOverLuma, 1 - Shadow) : BaseOverLuma;
+	float3 ScatterTint = Shadow < 1 ? pow(BaseOverLuma, 1 - Shadow) : 1;
 	return sqrt(abs(GBuffer.BaseColor)) * DiffuseScatter * ScatterTint;
 }
 
