Path Tracer: Change refraction rays to behave more like camera rays with respect to affect indirect

This makes refraction and transparency behave more consistently with one another.

diff --git a/Engine/Shaders/Private/PathTracing/PathTracingCore.ush b/Engine/Shaders/Private/PathTracing/PathTracingCore.ush
--- a/Engine/Shaders/Private/PathTracing/PathTracingCore.ush
+++ b/Engine/Shaders/Private/PathTracing/PathTracingCore.ush
@@ -387,10 +387,11 @@ void EvaluateDecals(FPathState PathState, uint InstanceInclusionMask, FPackedPat
 FPathTracingPayload TraceTransparentRay(inout FPathState PathState, int Bounce, bool bIsCameraRay, bool bLastBounce, bool bIncludeEmission, uint NumVisibleLights, inout FVolumeSegment VolumeSegment)
 {
 	const uint RayFlags = bIsCameraRay && EnableCameraBackfaceCulling ? RAY_FLAG_CULL_BACK_FACING_TRIANGLES : 0;
-	uint InstanceInclusionMask = bIsCameraRay ? PATHTRACER_MASK_CAMERA : PATHTRACER_MASK_INDIRECT;
+	const bool bUseCameraRayType = !PathState.HasMadeContributionScatter();
+	uint InstanceInclusionMask = bUseCameraRayType ? PATHTRACER_MASK_CAMERA : PATHTRACER_MASK_INDIRECT;
 	if (length2(PathState.Ray.Origin) < Pow2(PATH_TRACING_SKIP_HAIR_DISTANCE))
 	{
-		InstanceInclusionMask |= bIsCameraRay ? PATHTRACER_MASK_HAIR_CAMERA : PATHTRACER_MASK_HAIR_INDIRECT;
+		InstanceInclusionMask |= bUseCameraRayType ? PATHTRACER_MASK_HAIR_CAMERA : PATHTRACER_MASK_HAIR_INDIRECT;
 	}
 	
 	const uint MissShaderIndex = 0;
