Change SpecularProfile evalution to be shared.

diff --git a/Engine/Shaders/Private/PathTracing/Material/PathTracingSubstrateCommon.ush b/Engine/Shaders/Private/PathTracing/Material/PathTracingSubstrateCommon.ush
--- a/Engine/Shaders/Private/PathTracing/Material/PathTracingSubstrateCommon.ush
+++ b/Engine/Shaders/Private/PathTracing/Material/PathTracingSubstrateCommon.ush
@@ -30,10 +30,7 @@ float3 SubstrateSpecularTint(FPathTracingPayload Payload, float NoV, float NoL,
 	BRANCH
 	if (Payload.SpecularProfileId != 0)
 	{
-		const uint ProfileID = GetSpecularProfileId(Payload.SpecularProfileId);
-		const bool bUseHalfVector = GetSpecularProfileParameterization(Payload.SpecularProfileId) == 1;
-		const float3 ProfileCoord = bUseHalfVector ? float3(VoH, NoH, ProfileID) : float3(NoV, NoL, ProfileID);
-		Out = View.SpecularProfileTexture.SampleLevel(View.SpecularProfileSampler, ProfileCoord, 0).xyz;
+		Out = EvaluateSpecularProfile(Payload.SpecularProfileId, NoV, NoL, VoH, NoH);
 	}
 	return Out;
 }
diff --git a/Engine/Shaders/Private/Strata/Strata.ush b/Engine/Shaders/Private/Strata/Strata.ush
--- a/Engine/Shaders/Private/Strata/Strata.ush
+++ b/Engine/Shaders/Private/Strata/Strata.ush
@@ -873,6 +873,14 @@ uint GetSpecularProfileParameterization(float In)
 	return In > 0.f ? 0u : 1u;
 }
 
+float3 EvaluateSpecularProfile(float InSpecularProfileId, float NoV, float NoL, float VoH, float NoH)
+{
+	const uint ProfileID = GetSpecularProfileId(InSpecularProfileId);
+	const bool bUseHalfVector = GetSpecularProfileParameterization(InSpecularProfileId) == 1;
+	const float3 ProfileCoord = bUseHalfVector ? float3(VoH, NoH, ProfileID) : float3(NoV, NoL, ProfileID);
+	return View.SpecularProfileTexture.SampleLevel(View.SpecularProfileSampler, ProfileCoord, 0).xyz;
+}
+
 ///////////////////////////////////////////////////////////////////////////////
 // Glint
 
diff --git a/Engine/Shaders/Private/Strata/StrataEvaluation.ush b/Engine/Shaders/Private/Strata/StrataEvaluation.ush
--- a/Engine/Shaders/Private/Strata/StrataEvaluation.ush
+++ b/Engine/Shaders/Private/Strata/StrataEvaluation.ush
@@ -652,10 +652,7 @@ FStrataEvaluateResult StrataEvaluateBSDFCommon(FStrataBSDFContext BSDFContext, F
 			BRANCH
 			if (BSDF_GETHASSPECPROFILE(BSDFContext.BSDF))
 			{
-				const uint ProfileID = GetSpecularProfileId(SLAB_SPECPROFILEID(BSDFContext.BSDF));
-				const bool bUseHalfVector = GetSpecularProfileParameterization(SLAB_SPECPROFILEID(BSDFContext.BSDF)) == 1;
-				const float3 ProfileCoord = bUseHalfVector ? float3(BSDFContext.Context.VoH, BSDFContext.Context.NoH, ProfileID) : float3(BSDFContext.SatNoV, BSDFContext.SatNoL, ProfileID);
-				Sample.IntegratedSpecularValue *= View.SpecularProfileTexture.SampleLevel(View.SpecularProfileSampler, ProfileCoord, 0).xyz;
+				Sample.IntegratedSpecularValue *= EvaluateSpecularProfile(SLAB_SPECPROFILEID(BSDFContext.BSDF), BSDFContext.SatNoV, BSDFContext.SatNoL, BSDFContext.Context.VoH, BSDFContext.Context.NoH);
 			}
 			#endif
 
@@ -1507,9 +1504,8 @@ FStrataEnvLightResult StrataEvaluateForEnvLight(FStrataBSDFContext BSDFContext,
 				BRANCH
 				if (BSDF_GETHASSPECPROFILE(BSDFContext.BSDF))
 				{
-					const uint ProfileID = GetSpecularProfileId(SLAB_SPECPROFILEID(BSDFContext.BSDF));
-					const float NoR = saturate(dot(BSDFContext.R, BSDFContext.N));
-					StrataEnvLightResult.SpecularWeight *= View.SpecularProfileTexture.SampleLevel(View.SpecularProfileSampler, float3(BSDFContext.SatNoV, NoR, ProfileID), 0);
+					const float NoR = saturate(dot(StrataEnvLightResult.SpecularDirection, BSDFContext.N));
+					StrataEnvLightResult.SpecularWeight *= EvaluateSpecularProfile(SLAB_SPECPROFILEID(BSDFContext.BSDF), BSDFContext.SatNoV, NoR, BSDFContext.SatNoV, NoR);
 				}
 			#endif // STRATA_COMPLEXPATH
 
