Fix fuzz flickering at grazing angle.

diff --git a/Engine/Shaders/Private/BRDF.ush b/Engine/Shaders/Private/BRDF.ush
--- a/Engine/Shaders/Private/BRDF.ush
+++ b/Engine/Shaders/Private/BRDF.ush
@@ -515,7 +515,7 @@ float SheenLTC_DirectionalAlbedo(float NoV, float Roughness, Texture2D InSheenLT
 }
 
 // [Zeltner et al. 2022, "Practical Multiple-Scattering Sheen Using Linearly Transformed Cosines"]
-float3 SheenLTC_Sample(float3 CosSample, float3 V, float3 L, half3 N, float NoV, float Roughness, Texture2D InSheenLTCTexture, SamplerState InSheenSampler)
+float3 SheenLTC_Sample(float3 CosSample, float3 V, half3 N, float NoV, float Roughness, Texture2D InSheenLTCTexture, SamplerState InSheenSampler)
 {
 	const float4 LTC = SheenLTC_Cofficients(NoV, Roughness, InSheenLTCTexture, InSheenSampler);
 	const float aInv = LTC.x;
@@ -531,7 +531,7 @@ float3 SheenLTC_Sample(float3 CosSample, float3 V, float3 L, half3 N, float NoV,
 }
 
 // [Zeltner et al. 2022, "Practical Multiple-Scattering Sheen Using Linearly Transformed Cosines"]
-float SheenLTC_Eval(float3 V, float3 L, half3 N, float NoV, float Roughness, Texture2D InSheenLTCTexture, SamplerState InSheenSampler, inout float OutDirectionalAlbedo, inout float OutPdf)
+float SheenLTC_Eval(float3 V, float3 L, half3 N, float NoV, float Roughness, Texture2D InSheenLTCTexture, SamplerState InSheenSampler, inout float OutDirectionalAlbedo)
 {
 	// Rotate L into tangent space
 	const float3 T1 = normalize(V - N * NoV);
@@ -563,9 +563,7 @@ float SheenLTC_Eval(float3 V, float3 L, half3 N, float NoV, float Roughness, Tex
 	const float CosDistribution = WrappedLocalL.z / PI;
 	const float Out = OutDirectionalAlbedo * CosDistribution * Jacobian; // Fs * CosTheta
 
-	// The function has perfect important sampling PDF=F
-	OutPdf = Out;
-	return Out;
+	return max(Out, 0.f);
 }
 
 //---------------
diff --git a/Engine/Shaders/Private/Strata/StrataEvaluation.ush b/Engine/Shaders/Private/Strata/StrataEvaluation.ush
--- a/Engine/Shaders/Private/Strata/StrataEvaluation.ush
+++ b/Engine/Shaders/Private/Strata/StrataEvaluation.ush
@@ -713,8 +713,7 @@ FStrataEvaluateResult StrataEvaluateBSDFCommon(FStrataBSDFContext BSDFContext, F
 					}
 					else
 					{
-						float SheenPDF = 1;
-						ClothSpecularPathValueNoL = FuzzF0 * SheenLTC_Eval(BSDFContext.V, AreaLightContext.L, BSDFContext.N, BSDFContext.Context.NoV, FuzzRoughness, View.SheenLTCTexture, View.SheenLTCSampler, DirectionalAlbedo, SheenPDF);
+						ClothSpecularPathValueNoL = FuzzF0 * SheenLTC_Eval(BSDFContext.V, AreaLightContext.L, BSDFContext.N, BSDFContext.Context.NoV, FuzzRoughness, View.SheenLTCTexture, View.SheenLTCSampler, DirectionalAlbedo);
 					}
 					EnergyTerms.E = DirectionalAlbedo;
 					EnergyTerms.W = 1.f;
