Add Disney LTC cloth model to the pather tracer.

diff --git a/Engine/Shaders/Private/PathTracing/Material/PathTracingSubstrate.ush b/Engine/Shaders/Private/PathTracing/Material/PathTracingSubstrate.ush
--- a/Engine/Shaders/Private/PathTracing/Material/PathTracingSubstrate.ush
+++ b/Engine/Shaders/Private/PathTracing/Material/PathTracingSubstrate.ush
@@ -6,6 +6,10 @@
 #error "This header should only be included when Strata is enabled."
 #endif
 
+#ifndef STRATA_SHEEN_QUALITY
+#define STRATA_SHEEN_QUALITY 0
+#endif
+
 #include "PathTracingSubstrateCommon.ush"
 #include "PathTracingMaterialCommon.ush"
 #include "PathTracingGlossy.ush"
@@ -49,7 +53,16 @@ struct FSubstrateSingleSlabData
 		// STRATA_TODO: Add support for F82 model?
 		Spec0 = ComputeGGXSpecEnergyTermsRGB(Payload.RoughnessData.x, NoV, F0, F90);
 		Spec1 = ComputeGGXSpecEnergyTermsRGB(Payload.RoughnessData.y, NoV, F0, F90);
-		Fuzz = ComputeClothEnergyTermsA(FuzzRoughness, NoV);
+		#if STRATA_SHEEN_QUALITY == 1
+		{
+			Fuzz.E = SheenLTC_DirectionalAlbedo(NoV, FuzzRoughness, View.SheenLTCTexture, View.SheenLTCSampler);
+			Fuzz.W = 1;
+		}
+		#else
+		{
+			Fuzz = ComputeClothEnergyTermsA(FuzzRoughness, NoV);
+		}
+		#endif
 
 		const float3 SpecE = lerp(Spec0.E, Spec1.E, Payload.RoughnessData.z);
 
@@ -92,7 +105,7 @@ FMaterialSample Substrate_SampleMaterial(
 		// diffuse lobe
 		L = CosineSampleHemisphere(RandSample.xy).xyz;
 		H = normalize(L + V);
-}
+	}
 	else if (bSampledSpecular)
 	{
 		// specular lobes
@@ -175,9 +188,22 @@ FMaterialSample Substrate_SampleMaterial(
 	else
 	{
 		// Cloth Lobe
-		const float DCloth = D_Charlie(Data.FuzzRoughness, NoH);
-		const float3 FCloth = F_Schlick(Payload.FuzzColor, VoH);
-		const float3 ClothWeight = NoL > 0 && NoV > 0 ? (NoL * Payload.FuzzAmount * DCloth * Vis_Ashikhmin(NoV, NoL)) * FCloth * Data.Fuzz.W * (2 * PI) : 0.0;
+		float3 ClothWeight = 0;
+		#if STRATA_SHEEN_QUALITY == 1
+		{
+			// Disney's Cloth Lobe
+			float FuzzDirectionalAlbedo = 1;
+			float3 FuzzSpecularPathValueNoL = Payload.FuzzColor * SheenLTC_Eval(V_World, L_World, N_World, NoV, Data.FuzzRoughness, View.SheenLTCTexture, View.SheenLTCSampler, FuzzDirectionalAlbedo);
+			ClothWeight = NoL > 0 && NoV > 0 ? (Payload.FuzzAmount * FuzzSpecularPathValueNoL) * (2 * PI) : 0.0;
+		}
+		#else
+		{
+			// Charlie's Cloth Lobe
+			const float DCloth = D_Charlie(Data.FuzzRoughness, NoH);
+			const float3 FCloth = F_Schlick(Payload.FuzzColor, VoH);
+			ClothWeight = NoL > 0 && NoV > 0 ? (NoL * Payload.FuzzAmount * DCloth * Vis_Ashikhmin(NoV, NoL)) * FCloth * Data.Fuzz.W * (2 * PI) : 0.0;
+		}
+		#endif
 		Result.AddLobeWithMIS(ClothWeight * ShadowTerminator, 1 / (2.0 * PI), Data.LobePdf.w);
 		Result.ScatterType = PATHTRACER_SCATTER_SPECULAR;
 
@@ -235,9 +261,20 @@ FMaterialEval Substrate_EvalMaterial(
 	}
 	{
 		// Cloth Lobe
-		const float DCloth = D_Charlie(Data.FuzzRoughness, NoH);
-		const float3 FCloth = F_Schlick(Payload.FuzzColor, VoH);
-		const float3 ClothWeight = NoL > 0 && NoV > 0 ? (NoL * Payload.FuzzAmount * DCloth * Vis_Ashikhmin(NoV, NoL)) * FCloth * Data.Fuzz.W * (2 * PI) : 0.0;
+		float3 ClothWeight = 0;
+		#if STRATA_SHEEN_QUALITY == 1
+		{
+			float FuzzDirectionalAlbedo = 1;
+			float3 FuzzSpecularPathValueNoL = Payload.FuzzColor * SheenLTC_Eval(V_World, L_World, N_World, NoV, Data.FuzzRoughness, View.SheenLTCTexture, View.SheenLTCSampler, FuzzDirectionalAlbedo);
+			ClothWeight = NoL > 0 && NoV > 0 ? (Payload.FuzzAmount * FuzzSpecularPathValueNoL) * (2 * PI) : 0.0;
+		}
+		#else
+		{
+			const float DCloth = D_Charlie(Data.FuzzRoughness, NoH);
+			const float3 FCloth = F_Schlick(Payload.FuzzColor, VoH);
+			ClothWeight = NoL > 0 && NoV > 0 ? (NoL * Payload.FuzzAmount * DCloth * Vis_Ashikhmin(NoV, NoL)) * FCloth * Data.Fuzz.W * (2 * PI) : 0.0;
+		}
+		#endif
 		Result.AddLobeWithMIS(DiffuseSpecularScale.y * ClothWeight * ShadowTerminator, 1 / (2.0 * PI), Data.LobePdf.w);
 	}
 
