Offset into bent normals texture correctly for multiple views.

UpsampleDFAO is called when the ReflectionEnvironmentAndSky pass is executing. Since the result of the DFAO pass is being composited into a texture atlas (combined for all views), we don't need to subtract the ViewRect offset anymore.

diff --git a/Engine/Shaders/Private/DistanceFieldAOShared.ush b/Engine/Shaders/Private/DistanceFieldAOShared.ush
--- a/Engine/Shaders/Private/DistanceFieldAOShared.ush
+++ b/Engine/Shaders/Private/DistanceFieldAOShared.ush
@@ -155,9 +155,7 @@ FUpsampleDFAOOutput UpsampleDFAO(float2 BufferUV, float SceneDepth)
 {
 	FUpsampleDFAOOutput UpsampleDFAOOutput;
 
-	// Distance field AO was computed at 0,0 regardless of viewrect min
-	float2 DistanceFieldUVs = BufferUV - View.ViewRectMin.xy * View.BufferSizeAndInvSize.zw;
-	DistanceFieldUVs = min(DistanceFieldUVs, AOBufferBilinearUVMax);
+	float2 DistanceFieldUVs = min(BufferUV, AOBufferBilinearUVMax);
 
 #define BILATERAL_UPSAMPLE 1
 #if BILATERAL_UPSAMPLE
