Respect r.RayTracing.DebugVisualizationMode.OpaqueOnly in DebugRayTracingInstance(...)
- particularly important so that RayTracing Picker results match.

diff --git a/Engine/Shaders/Private/RayTracing/RayTracingDebug.usf b/Engine/Shaders/Private/RayTracing/RayTracingDebug.usf
--- a/Engine/Shaders/Private/RayTracing/RayTracingDebug.usf
+++ b/Engine/Shaders/Private/RayTracing/RayTracingDebug.usf
@@ -168,10 +168,13 @@ float4 DebugRayTracingInstance(FRayDesc Ray, uint2 PixelCoord, out float OutDept
 	FRayTracingDebugPayload Payload = (FRayTracingDebugPayload)0;
 	Payload.SetMiss();
 
+	const uint RayFlags = RAY_FLAG_CULL_BACK_FACING_TRIANGLES;
+	const uint InstanceInclusionMask = (OpaqueOnly ? RAY_TRACING_MASK_OPAQUE : RAY_TRACING_MASK_ALL);
+
 	TraceRay(
 		TLAS,
-		RAY_FLAG_CULL_BACK_FACING_TRIANGLES  /*RayFlags*/,
-		0xff /*InstanceInclusionMask*/,
+		RayFlags  /*RayFlags*/,
+		InstanceInclusionMask /*InstanceInclusionMask*/,
 		0 /*RayContributionToHitGroupIndex*/,
 		RAY_TRACING_NUM_SHADER_SLOTS /*MultiplierForGeometryContributionToShaderIndex*/,
 		0 /*MissShaderIndex*/,
