Changes FTSRSubpixelDepth::Age to FTSRSubpixelDepth::RemainingLifetime to support any arbitrary lifetimes

diff --git a/Engine/Shaders/Private/TemporalSuperResolution/TSRCommon.ush b/Engine/Shaders/Private/TemporalSuperResolution/TSRCommon.ush
--- a/Engine/Shaders/Private/TemporalSuperResolution/TSRCommon.ush
+++ b/Engine/Shaders/Private/TemporalSuperResolution/TSRCommon.ush
@@ -1214,8 +1214,8 @@ struct FTSRSubpixelDepth
 	/** Offset of the depth from the center of pixel. */
 	float2 Offset;
 	
-	/** Depth of the pixel in number of frames. */
-	uint Age;
+	/** Number of reprojected frame remaining before this depth is to be discarded. */
+	uint RemainingLifetime;
 };
 
 uint CompressSubpixelDepth(FTSRSubpixelDepth SubpixelDepth)
@@ -1224,7 +1224,7 @@ uint CompressSubpixelDepth(FTSRSubpixelDepth SubpixelDepth)
 	CompressedSubpixelDepth |= f32tof16(SubpixelDepth.DeviceZ) << 17u;
 	CompressedSubpixelDepth |= uint(fastClamp(int(round(SubpixelDepth.Offset.x * 16.00 + 16.0)), 0, 31)) << 12u;
 	CompressedSubpixelDepth |= uint(fastClamp(int(round(SubpixelDepth.Offset.y * 16.00 + 16.0)), 0, 31)) <<  7u;
-	CompressedSubpixelDepth |= (SubpixelDepth.Age) << 0u;
+	CompressedSubpixelDepth |= (SubpixelDepth.RemainingLifetime) << 0u;
 	return CompressedSubpixelDepth;
 }
 
@@ -1234,6 +1234,21 @@ FTSRSubpixelDepth UncompressSubpixelDepth(uint CompressedSubpixelDepth)
 	SubpixelDepth.DeviceZ  = f16tof32(CompressedSubpixelDepth >> 17u);
 	SubpixelDepth.Offset.x = float((CompressedSubpixelDepth >> 12u) & 0x1Fu) * rcp(16.0) - 1.0;
 	SubpixelDepth.Offset.y = float((CompressedSubpixelDepth >>  7u) & 0x1Fu) * rcp(16.0) - 1.0;
-	SubpixelDepth.Age      = CompressedSubpixelDepth & 0x7Fu;
+	SubpixelDepth.RemainingLifetime = CompressedSubpixelDepth & 0x7Fu;
 	return SubpixelDepth;
 }
+
+void ForwardReprojectStaticDepth(float2 PrevScreenPosition, float PrevDeviceZ, out float2 ScreenPos, out float SceneDepth)
+{
+	float3 PreViewTranslationOffset = LWCToFloat(LWCSubtract(PrimaryView.PreViewTranslation, PrimaryView.PrevPreViewTranslation));
+ 
+	// Compute position of the history in world space.
+	float PrevSceneDepth = ConvertFromDeviceZ(PrevDeviceZ);
+	float3 TranslatedWorldPosition = mul(float4(PrevScreenPosition * PrevSceneDepth, PrevSceneDepth, 1), View.PrevScreenToTranslatedWorld).xyz + PreViewTranslationOffset;
+	
+	// Reproject the world position into new viewport.
+	float4 ClipPos = mul(float4(TranslatedWorldPosition, 1.0), View.TranslatedWorldToClip);
+	
+	ScreenPos = ClipPos.xy / ClipPos.w;
+	SceneDepth = ClipPos.w;
+}
diff --git a/Engine/Shaders/Private/TemporalSuperResolution/TSRDilateVelocity.usf b/Engine/Shaders/Private/TemporalSuperResolution/TSRDilateVelocity.usf
--- a/Engine/Shaders/Private/TemporalSuperResolution/TSRDilateVelocity.usf
+++ b/Engine/Shaders/Private/TemporalSuperResolution/TSRDilateVelocity.usf
@@ -49,6 +49,7 @@ float VelocityExtrapolationMultiplier;
 float InvFlickeringMaxParralaxVelocity;
 uint bOutputIsMovingTexture;
 uint bIncludeDynamicDepthInSubpixelDetails;
+uint SubpixelDepthLifetime;
 
 Texture2D<float> SceneDepthTexture;
 Texture2D SceneVelocityTexture;
@@ -376,7 +377,7 @@ void MainCS(
 			FTSRSubpixelDepth FinalSubpixelDepth;
 			FinalSubpixelDepth.DeviceZ = DeviceZ;
 			FinalSubpixelDepth.Offset = 0.0;
-			FinalSubpixelDepth.Age = 0;
+			FinalSubpixelDepth.RemainingLifetime = SubpixelDepthLifetime;
 
 			CompressedSubpixelDepth = CompressSubpixelDepth(FinalSubpixelDepth);
 
@@ -392,8 +393,6 @@ void MainCS(
 		ReduceVelocityFlattenTile(GroupThreadIndex, FinalScreenVelocity.xy, /* out */ VelocityPolarRange, /* out */ VelocityPolar);
 	#endif
 
-	Debug.rg = VelocityPixelOffset * 0.5 + 0.5;
-
 	{
 		bool bValidOutputPixel = all(InputPixelPos < InputInfo_ViewportMax);
 		
diff --git a/Engine/Shaders/Private/TemporalSuperResolution/TSRForwardScatterDepth.usf b/Engine/Shaders/Private/TemporalSuperResolution/TSRForwardScatterDepth.usf
--- a/Engine/Shaders/Private/TemporalSuperResolution/TSRForwardScatterDepth.usf
+++ b/Engine/Shaders/Private/TemporalSuperResolution/TSRForwardScatterDepth.usf
@@ -18,7 +18,6 @@ uint2 PrevSubpixelDepthViewportMin;
 uint2 PrevSubpixelDepthViewportMax;
 FScreenTransform InputPixelPosToPrevScreenPosition;
 FScreenTransform ScreenPosToOutputPixelPos;
-uint MaxDepthAge;
 
 Texture2D<uint> PrevSubpixelDepthTexture;
 
@@ -43,7 +42,7 @@ void MainCS(
 	uint CompressedPrevSubpixelDepth = PrevSubpixelDepthTexture[InputPixelPos];
 	FTSRSubpixelDepth PrevScatteredSubpixelDepth = UncompressSubpixelDepth(CompressedPrevSubpixelDepth);
 		
-	const bool bEarlyReturn = PrevScatteredSubpixelDepth.DeviceZ == 0.0 || PrevScatteredSubpixelDepth.Age >= MaxDepthAge || any(InputPixelPos > PrevSubpixelDepthViewportMax);
+	const bool bEarlyReturn = PrevScatteredSubpixelDepth.DeviceZ == 0.0 || PrevScatteredSubpixelDepth.RemainingLifetime == 0 || any(InputPixelPos > PrevSubpixelDepthViewportMax);
 #if CONFIG_EARLY_RETURN && !DEBUG_INTERLOCK_ALL
 	BRANCH
 	if (bEarlyReturn)
@@ -52,23 +51,19 @@ void MainCS(
 	}
 #endif
 	
-	float3 PreViewTranslationOffset = LWCToFloat(LWCSubtract(PrimaryView.PreViewTranslation, PrimaryView.PrevPreViewTranslation));
- 
-	// Compute position of the history in world space.
-	float PrevSceneDepth = ConvertFromDeviceZ(PrevScatteredSubpixelDepth.DeviceZ);
 	float2 PrevScreenPosition = ApplyScreenTransform(float2(InputPixelPos), InputPixelPosToPrevScreenPosition);
-	float3 TranslatedWorldPosition = mul(float4(PrevScreenPosition * PrevSceneDepth, PrevSceneDepth, 1), View.PrevScreenToTranslatedWorld).xyz + PreViewTranslationOffset;
 	
-	// Reproject the world position into new viewport.
-	float4 ClipPos = mul(float4(TranslatedWorldPosition, 1.0), View.TranslatedWorldToClip);
-	float2 ScreenPos = ClipPos.xy / ClipPos.w;
-	float2 DesiredPixelPos = ApplyScreenTransform(ScreenPos, ScreenPosToOutputPixelPos);
+	// Reproject static depth from previous frame to current frame.
+	float2 ScreenPos;
+	float SceneDepth;
+	ForwardReprojectStaticDepth(PrevScreenPosition, PrevScatteredSubpixelDepth.DeviceZ, /* out */ ScreenPos, /* out */ SceneDepth);
 
+	float2 DesiredPixelPos = ApplyScreenTransform(ScreenPos, ScreenPosToOutputPixelPos);
 	float2 ActualPixelPos = floor(DesiredPixelPos);
 	float2 NewOffset = DesiredPixelPos - (ActualPixelPos + 0.5);
 
 	int2 OutputPixelPos = int2(ActualPixelPos);
-	bool bIsValidOutputPixelPos = all(and(OutputPixelPos >= InputInfo_ViewportMin, OutputPixelPos < InputInfo_ViewportMax)) && ClipPos.w > 0.0;
+	bool bIsValidOutputPixelPos = all(and(OutputPixelPos >= InputInfo_ViewportMin, OutputPixelPos < InputInfo_ViewportMax)) && SceneDepth > 0.0;
 
 #if !CONFIG_EARLY_RETURN
 	const bool bScatter = !bEarlyReturn && bIsValidOutputPixelPos;
@@ -82,9 +77,9 @@ void MainCS(
 #endif
 	{
 		FTSRSubpixelDepth FinalSubpixelDepth;
-		FinalSubpixelDepth.DeviceZ = ConvertToDeviceZ(ClipPos.w);
+		FinalSubpixelDepth.DeviceZ = ConvertToDeviceZ(SceneDepth);
 		FinalSubpixelDepth.Offset = NewOffset;
-		FinalSubpixelDepth.Age = PrevScatteredSubpixelDepth.Age + 1;
+		FinalSubpixelDepth.RemainingLifetime = PrevScatteredSubpixelDepth.RemainingLifetime - 1;
 
 		uint CompressedSubpixelDepth = CompressSubpixelDepth(FinalSubpixelDepth);
 	
diff --git a/Engine/Source/Runtime/Renderer/Private/PostProcess/TemporalSuperResolution.cpp b/Engine/Source/Runtime/Renderer/Private/PostProcess/TemporalSuperResolution.cpp
--- a/Engine/Source/Runtime/Renderer/Private/PostProcess/TemporalSuperResolution.cpp
+++ b/Engine/Source/Runtime/Renderer/Private/PostProcess/TemporalSuperResolution.cpp
@@ -555,7 +555,6 @@ class FTSRForwardScatterDepthCS : public FTSRShader
 		SHADER_PARAMETER(FIntPoint, PrevSubpixelDepthViewportMax)
 		SHADER_PARAMETER(FScreenTransform, InputPixelPosToPrevScreenPosition)
 		SHADER_PARAMETER(FScreenTransform, ScreenPosToOutputPixelPos)
-		SHADER_PARAMETER(int32, MaxDepthAge)
 
 		SHADER_PARAMETER_RDG_TEXTURE(Texture2D, PrevSubpixelDepthTexture)
 
@@ -585,6 +584,7 @@ class FTSRDilateVelocityCS : public FTSRShader
 		SHADER_PARAMETER(float, InvFlickeringMaxParralaxVelocity)
 		SHADER_PARAMETER(int32, bOutputIsMovingTexture)
 		SHADER_PARAMETER(int32, bIncludeDynamicDepthInSubpixelDetails)
+		SHADER_PARAMETER(int32, SubpixelDepthLifetime)
 
 		SHADER_PARAMETER_RDG_TEXTURE(Texture2D, SceneDepthTexture)
 		SHADER_PARAMETER_RDG_TEXTURE(Texture2D, SceneVelocityTexture)
@@ -1099,6 +1099,9 @@ FDefaultTemporalUpscaler::FOutputs AddTemporalSuperResolutionPasses(
 {
 	const FTSRHistory& InputHistory = View.PrevViewInfo.TSRHistory;
 
+	// Maximum depth lifetime encodable in FTSRSubpixelDepth::RemainingLifetime
+	const int32 kMaxDepthLifetime = (1 << 7) - 1;
+
 #if COMPILE_TSR_DEBUG_PASSES
 	const bool bSetupDebugPasses = CVarTSRSetupDebugPasses.GetValueOnRenderThread() != 0;
 #endif
@@ -1562,7 +1565,6 @@ FDefaultTemporalUpscaler::FOutputs AddTemporalSuperResolutionPasses(
 		PassParameters->ScreenPosToOutputPixelPos = FScreenTransform::ChangeTextureBasisFromTo(
 			InputExtent, InputRect,
 			FScreenTransform::ETextureBasis::ScreenPosition, FScreenTransform::ETextureBasis::TexelPosition);
-		PassParameters->MaxDepthAge = FMath::Clamp(CVarTSRSubpixelDepthMaxAge.GetValueOnRenderThread(), 1, (1 << 7) - 1);
 
 		PassParameters->PrevSubpixelDepthTexture = PrevHistory.SubpixelDepth;
 		PassParameters->SubpixelDepthOutput = GraphBuilder.CreateUAV(PrevScatteredSubpixelDepthTexture);
@@ -1641,6 +1643,7 @@ FDefaultTemporalUpscaler::FOutputs AddTemporalSuperResolutionPasses(
 		}
 		PassParameters->bOutputIsMovingTexture = bOutputIsMovingTexture;
 		PassParameters->bIncludeDynamicDepthInSubpixelDetails = SubpixelMethod == ETSRSubpixelMethod::ClosestDepth && CVarTSRSubpixelIncludeMovingDepth.GetValueOnRenderThread();
+		PassParameters->SubpixelDepthLifetime = FMath::Clamp(CVarTSRSubpixelDepthMaxAge.GetValueOnRenderThread(), 1, kMaxDepthLifetime);
 
 		PassParameters->SceneDepthTexture = PassInputs.SceneDepth.Texture;
 		PassParameters->SceneVelocityTexture = PassInputs.SceneVelocity.Texture;
