HeterogeneousVolumes: Apply object-space culling when voxelizing the world-space grid.

diff --git a/Engine/Shaders/Private/HeterogeneousVolumes/HeterogeneousVolumesMaterialBakingPipeline.usf b/Engine/Shaders/Private/HeterogeneousVolumes/HeterogeneousVolumesMaterialBakingPipeline.usf
--- a/Engine/Shaders/Private/HeterogeneousVolumes/HeterogeneousVolumesMaterialBakingPipeline.usf
+++ b/Engine/Shaders/Private/HeterogeneousVolumes/HeterogeneousVolumesMaterialBakingPipeline.usf
@@ -105,6 +105,18 @@ bool PrimitiveIntersectsVoxel(float3 VoxelBoundsMin, float3 VoxelBoundsMax)
 	return true;
 }
 
+bool WorldPositionIntersectsPrimitive(float3 WorldPosition)
+{
+	return all(WorldPosition >= PrimitiveWorldBoundsMin) && all(WorldPosition <= PrimitiveWorldBoundsMax);
+}
+
+bool LocalPositionIntersectsPrimitive(float3 LocalPosition)
+{
+	float3 LocalMinPosition = LocalBoundsOrigin - LocalBoundsExtent;
+	float3 LocalMaxPosition = LocalBoundsOrigin + LocalBoundsExtent;
+	return all(LocalPosition >= LocalMinPosition) && all(LocalPosition <= LocalMaxPosition);
+}
+
 struct FRasterTileData
 {
 	uint TopLevelGridLinearIndex;
@@ -237,21 +249,25 @@ void RasterizeBottomLevelOrthoGridCS(
 			{
 				UVW = float3(BottomLevelVoxelIndex + Jitter) / float3(BottomLevelVoxelResolution);
 			}
-			float3 WorldPosition = UVW * TopLevelGridVoxelWorldBoundsExtent + TopLevelGridVoxelWorldBoundsMin;
-
-			// Setup evaluation context
-			FMaterialPixelParameters MaterialParameters = MakeInitializedMaterialPixelParameters();
-			MaterialParameters.PrimitiveId = PrimitiveId;
-			MaterialParameters.AbsoluteWorldPosition = LWCPromote(WorldPosition);
 
-			// Evaluate material graph
-			FPixelMaterialInputs PixelMaterialInputs;
-			CalcPixelMaterialInputs(MaterialParameters, PixelMaterialInputs);
-
-			// Extract volume rendering coefficients
-			Extinction = SampleExtinctionCoefficients(PixelMaterialInputs);
-			Emission = SampleEmissive(PixelMaterialInputs);
-			Albedo = SampleAlbedo(PixelMaterialInputs) * Extinction;
+			float3 WorldPosition = UVW * TopLevelGridVoxelWorldBoundsExtent + TopLevelGridVoxelWorldBoundsMin;
+			float3 LocalPosition = mul(float4(WorldPosition, 1), WorldToLocal).xyz;
+			if (LocalPositionIntersectsPrimitive(LocalPosition))
+			{
+				// Setup evaluation context
+				FMaterialPixelParameters MaterialParameters = MakeInitializedMaterialPixelParameters();
+				MaterialParameters.PrimitiveId = PrimitiveId;
+				MaterialParameters.AbsoluteWorldPosition = LWCPromote(WorldPosition);
+
+				// Evaluate material graph
+				FPixelMaterialInputs PixelMaterialInputs;
+				CalcPixelMaterialInputs(MaterialParameters, PixelMaterialInputs);
+
+				// Extract volume rendering coefficients
+				Extinction = SampleExtinctionCoefficients(PixelMaterialInputs);
+				Emission = SampleEmissive(PixelMaterialInputs);
+				Albedo = SampleAlbedo(PixelMaterialInputs) * Extinction;
+			}
 		}
 
 		// Aggregate in group-shared memory
@@ -407,11 +423,6 @@ float FarPlaneDepth;
 
 groupshared float GSExtinctionSumScalar[THREADGROUP_SIZE_1D];
 
-bool EvaluationPositionIntersectsPrimitive(float3 WorldPosition)
-{
-	return all(WorldPosition > PrimitiveWorldBoundsMin) && all(WorldPosition < PrimitiveWorldBoundsMax);
-}
-
 [numthreads(THREADGROUP_SIZE_3D, THREADGROUP_SIZE_3D, THREADGROUP_SIZE_3D)]
 void RasterizeBottomLevelFrustumGridCS(
 	uint3 GroupId : SV_GroupID,
@@ -447,7 +458,7 @@ void RasterizeBottomLevelFrustumGridCS(
 	float3 Emission = 0;
 	float3 Albedo = 0;
 
-	if (EvaluationPositionIntersectsPrimitive(WorldPosition))
+	if (WorldPositionIntersectsPrimitive(WorldPosition))
 	{
 		FMaterialPixelParameters MaterialParameters = MakeInitializedMaterialPixelParameters();
 		MaterialParameters.PrimitiveId = PrimitiveId;
