Fix multi-view DFAO object cone tracing.

Since the bent normal texture is now being created per-view, the UV coordinates should be calculated based on the size of the ViewRect.

diff --git a/Engine/Shaders/Private/DistanceFieldObjectCulling.usf b/Engine/Shaders/Private/DistanceFieldObjectCulling.usf
--- a/Engine/Shaders/Private/DistanceFieldObjectCulling.usf
+++ b/Engine/Shaders/Private/DistanceFieldObjectCulling.usf
@@ -129,12 +129,10 @@ void BuildTileConesMain(
 {
     uint ThreadIndex = GroupThreadId.y * THREADGROUP_SIZEX + GroupThreadId.x;
 
-	float2 BaseLevelScreenUV = (DispatchThreadId.xy + float2(.5f, .5f)) * DOWNSAMPLE_FACTOR * View.BufferSizeAndInvSize.zw;
+	// Sampling from the texture based off of the ViewRect size because the texture is created on a per-view basis
+	float2 BaseLevelScreenUV = (DispatchThreadId.xy + float2(.5f, .5f)) * DOWNSAMPLE_FACTOR * View.ViewSizeAndInvSize.zw;
 	float SceneDepth = GetDownsampledDepth(BaseLevelScreenUV);
 
-	//float2 ScreenUV = (DispatchThreadId.xy * DOWNSAMPLE_FACTOR + View.ViewRectMin.xy + float2(.5f, .5f)) * View.BufferSizeAndInvSize.zw;
-	//float SceneDepth = CalcSceneDepth(ScreenUV);
-
 	// Initialize per-tile variables
     if (ThreadIndex == 0) 
 	{
diff --git a/Engine/Source/Runtime/Renderer/Private/DistanceFieldObjectCulling.cpp b/Engine/Source/Runtime/Renderer/Private/DistanceFieldObjectCulling.cpp
--- a/Engine/Source/Runtime/Renderer/Private/DistanceFieldObjectCulling.cpp
+++ b/Engine/Source/Runtime/Renderer/Private/DistanceFieldObjectCulling.cpp
@@ -337,9 +337,10 @@ void ScatterTilesToObjects(
 
 FIntPoint GetTileListGroupSizeForView(const FViewInfo& View)
 {
+	const FIntPoint AOBufferSize = GetBufferSizeForAO(View);
 	return FIntPoint(
-		FMath::DivideAndRoundUp(FMath::Max(View.ViewRect.Size().X / GAODownsampleFactor, 1), GDistanceFieldAOTileSizeX),
-		FMath::DivideAndRoundUp(FMath::Max(View.ViewRect.Size().Y / GAODownsampleFactor, 1), GDistanceFieldAOTileSizeY));
+		FMath::DivideAndRoundUp(FMath::Max(AOBufferSize.X, 1), GDistanceFieldAOTileSizeX),
+		FMath::DivideAndRoundUp(FMath::Max(AOBufferSize.Y, 1), GDistanceFieldAOTileSizeY));
 }
 
 void BuildTileObjectLists(
