Deletes r.TSR.History.GrandReprojection=1 and r.TSR.Subpixel.Method=1

r.TSR.History.GrandReprojection=1 doesn't looks nearly as good as r.TSR.History.ScreenPercentage=200 and cause high maintenance burden;
r.TSR.Subpixel.Method=1 is just always worst looking due to the introduce instability on the image.

diff --git a/Engine/Config/BaseScalability.ini b/Engine/Config/BaseScalability.ini
--- a/Engine/Config/BaseScalability.ini
+++ b/Engine/Config/BaseScalability.ini
@@ -42,7 +42,6 @@ r.TemporalAA.Quality=0
 r.TSR.History.R11G11B10=1
 r.TSR.History.ScreenPercentage=100
 r.TSR.History.UpdateQuality=0
-r.TSR.History.GrandReprojection=0
 r.TSR.ShadingRejection.Flickering=0
 r.TSR.ShadingRejection.TileOverscan=2
 r.TSR.Velocity.Extrapolation=0
@@ -54,7 +53,6 @@ r.TemporalAA.Quality=1
 r.TSR.History.R11G11B10=1
 r.TSR.History.ScreenPercentage=100
 r.TSR.History.UpdateQuality=1
-r.TSR.History.GrandReprojection=0
 r.TSR.ShadingRejection.Flickering=0
 r.TSR.ShadingRejection.TileOverscan=2
 r.TSR.Velocity.Extrapolation=0
@@ -66,7 +64,6 @@ r.TemporalAA.Quality=1
 r.TSR.History.R11G11B10=1
 r.TSR.History.ScreenPercentage=100
 r.TSR.History.UpdateQuality=2
-r.TSR.History.GrandReprojection=0
 r.TSR.ShadingRejection.Flickering=1
 r.TSR.ShadingRejection.TileOverscan=3
 r.TSR.Velocity.Extrapolation=0
@@ -78,7 +75,6 @@ r.TemporalAA.Quality=2
 r.TSR.History.R11G11B10=1
 r.TSR.History.ScreenPercentage=200
 r.TSR.History.UpdateQuality=3
-r.TSR.History.GrandReprojection=0
 r.TSR.ShadingRejection.Flickering=1
 r.TSR.ShadingRejection.TileOverscan=3
 r.TSR.Velocity.Extrapolation=1
@@ -90,7 +86,6 @@ r.TemporalAA.Quality=2
 r.TSR.History.R11G11B10=0
 r.TSR.History.ScreenPercentage=200
 r.TSR.History.UpdateQuality=3
-r.TSR.History.GrandReprojection=1
 r.TSR.ShadingRejection.Flickering=1
 r.TSR.ShadingRejection.TileOverscan=3
 r.TSR.Velocity.Extrapolation=1
diff --git a/Engine/Shaders/Private/TemporalSuperResolution/TSRCommon.ush b/Engine/Shaders/Private/TemporalSuperResolution/TSRCommon.ush
--- a/Engine/Shaders/Private/TemporalSuperResolution/TSRCommon.ush
+++ b/Engine/Shaders/Private/TemporalSuperResolution/TSRCommon.ush
@@ -222,11 +222,6 @@ static const uint kPlusIndexes3x3[5] = { 4, 1, 3, 7, 5 };
 #endif
 
 
-#define tsr_subpixel_details uint
-#define tsr_subpixel_details2 uint2
-#define tsr_subpixel_payload uint // TODO: uint16_t
-#define tsr_subpixel_payload2 uint2 // TODO: uint16_t
-
 // Largest encodable normal number in a half used on console.
 static const tsr_half LargestNormalNumber = tsr_half(65504.0);
 
@@ -290,8 +285,6 @@ float2 PrevHistoryInfo_UVViewportSizeInverse;
 float2 PrevHistoryInfo_UVViewportBilinearMin;
 float2 PrevHistoryInfo_UVViewportBilinearMax;
 FScreenTransform ScreenPosToPrevHistoryBufferUV;
-FScreenTransform ScreenPosToPrevSubpixelDetails;
-float2 PrevSubpixelDetailsExtent;
 float HistoryPreExposureCorrection;
 
 uint bCameraCut;
@@ -927,102 +920,6 @@ tsr_half MeasureRejectionFactor(
 		dpv_interleave_mono_registers(BackbufferQuantizationError)));
 }
 
-struct FSubpixelNeighborhood
-{
-	// 2x2 quad of subpixels grid
-	tsr_subpixel_details SampleArray[4];
-
-	// Offset of subpixel
-	uint2 SubpixelOffset;
-
-	// Offset of the closest subpixel
-	bool2 ClosestSubpixelOffset;
-};
-
-FSubpixelNeighborhood GatherPrevSubpixelNeighborhood(Texture2D<tsr_subpixel_details> PrevSubPixelDetailTexture, float2 PrevHistoryBufferUV)
-{
-	float2 PixelCoord = PrevHistoryBufferUV * PrevSubpixelDetailsExtent;
-	float2 TopLeftPixelCoord = floor(PixelCoord - 0.5) + 0.5;
-	
-	FSubpixelNeighborhood Neighborhood;
-	#if 1 // TODO
-	{
-		Neighborhood.SampleArray[0] = PrevSubPixelDetailTexture[uint2(TopLeftPixelCoord + float2(0.0, 0.0))];
-		Neighborhood.SampleArray[1] = PrevSubPixelDetailTexture[uint2(TopLeftPixelCoord + float2(1.0, 0.0))];
-		Neighborhood.SampleArray[2] = PrevSubPixelDetailTexture[uint2(TopLeftPixelCoord + float2(0.0, 1.0))];
-		Neighborhood.SampleArray[3] = PrevSubPixelDetailTexture[uint2(TopLeftPixelCoord + float2(1.0, 1.0))];
-	}
-	#else
-	{
-		// (-,+),(+,+),(+,-),(-,-),
-		/**
-		 *  3 2
-		 *  0 1
-		 */
-		uint4 Samples = PrevSubPixelDetailTexture.Gather(GlobalPointClampedSampler, PrevHistoryBufferUV);
-
-		Neighborhood.SampleArray[0] = Samples[3]; // 00
-		Neighborhood.SampleArray[1] = Samples[2]; // 10
-		Neighborhood.SampleArray[2] = Samples[0]; // 01
-		Neighborhood.SampleArray[3] = Samples[1]; // 11
-	}
-	#endif
-
-	tsr_half2 Interp = tsr_half2(PixelCoord - TopLeftPixelCoord);
-	Neighborhood.SubpixelOffset = uint2(round(Interp * 2.0));
-	Neighborhood.ClosestSubpixelOffset = Interp > 0.5;
-
-	return Neighborhood;
-}
-
-tsr_subpixel_payload GetSubpixelPayload(FSubpixelNeighborhood Neighborhood, uint SubpixelId)
-{
-	const uint2 SubpixelCoord = uint2(SubpixelId % SUB_PIXEL_GRID_SIZE, SubpixelId / SUB_PIXEL_GRID_SIZE);
-
-	uint2 InputSubpixelCoord = Neighborhood.SubpixelOffset + SubpixelCoord;
-	uint2 InputTexelCoord = InputSubpixelCoord / SUB_PIXEL_GRID_SIZE;
-	uint2 QuadSubpixelCoord = InputSubpixelCoord % SUB_PIXEL_GRID_SIZE;
-
-
-	uint QuadSampleId = dot(InputTexelCoord, uint2(1, 2));
-	tsr_subpixel_details InputSample = Neighborhood.SampleArray[QuadSampleId];
-
-	uint BitShift = SUB_PIXEL_BIT_COUNT * dot(QuadSubpixelCoord, uint2(1, SUB_PIXEL_GRID_SIZE));
-
-	tsr_subpixel_payload SubpixelPayload = tsr_subpixel_payload(InputSample >> BitShift) & SUB_PIXEL_BIT_MASK;
-
-	return SubpixelPayload;
-}
-
-tsr_subpixel_payload GetClosestSubpixelPayload(FSubpixelNeighborhood Neighborhood)
-{
-	uint QuadSampleId = dot(uint2(Neighborhood.ClosestSubpixelOffset), uint2(1, 2));
-	tsr_subpixel_details InputSample = Neighborhood.SampleArray[QuadSampleId];
-
-	uint BitShift = SUB_PIXEL_BIT_COUNT * (3 - QuadSampleId);
-
-	tsr_subpixel_payload SubpixelPayload = tsr_subpixel_payload(InputSample >> BitShift) & SUB_PIXEL_BIT_MASK;
-
-	return SubpixelPayload;
-}
-
-CALL_SITE_DEBUGLOC
-tsr_subpixel_payload CompressParallaxFactor(tsr_half ParallaxFactor)
-{
-	return fastClamp(
-		tsr_subpixel_payload((ParallaxFactor - tsr_half(1.0)) * tsr_half(float(SUB_PIXEL_PARALLAX_FACTOR_BIT_MASK) / MAX_PARALLAX_FACTOR)),
-		tsr_subpixel_payload(0), tsr_subpixel_payload(SUB_PIXEL_PARALLAX_FACTOR_BIT_MASK));
-}
-
-CALL_SITE_DEBUGLOC
-tsr_half UncompressParallaxFactor(tsr_subpixel_payload CompressedParallaxFactor)
-{
-	return 
-		tsr_half(CompressedParallaxFactor & SUB_PIXEL_PARALLAX_FACTOR_BIT_MASK) *
-		tsr_half(MAX_PARALLAX_FACTOR / float(SUB_PIXEL_PARALLAX_FACTOR_BIT_MASK)) +
-		tsr_half(1.0);
-}
-
 CALL_SITE_DEBUGLOC
 tsr_half3x2 QuantizeForUNormRenderTarget(tsr_half3x2 Color, tsr_half E, const float3 QuantizationError)
 {
diff --git a/Engine/Shaders/Private/TemporalSuperResolution/TSRDebugHistory.usf b/Engine/Shaders/Private/TemporalSuperResolution/TSRDebugHistory.usf
deleted file mode 100644
diff --git a/Engine/Shaders/Private/TemporalSuperResolution/TSRDecimateHistory.usf b/Engine/Shaders/Private/TemporalSuperResolution/TSRDecimateHistory.usf
--- a/Engine/Shaders/Private/TemporalSuperResolution/TSRDecimateHistory.usf
+++ b/Engine/Shaders/Private/TemporalSuperResolution/TSRDecimateHistory.usf
@@ -11,8 +11,6 @@
 
 #define CONFIG_MANUAL_LDS_SPILL 1
 
-#define CONFIG_GRAND_REPROJECTION (DIM_GRAND_REPROJECTION)
-
 #define CONFIG_VELOCITY_SAMPLES 9
 #if CONFIG_VELOCITY_SAMPLES != 1
 	#define CONFIG_VELOCITY_SAMPLE_PAIRS ((CONFIG_VELOCITY_SAMPLES - 1) / 2)
@@ -31,8 +29,6 @@
 /** Whether to compile the moire reprojection. */
 #define CONFIG_MOIRE_REPROJECTION (DIM_MOIRE_REPROJECTION)
 
-#define CONFIG_SUBPIXEL_PARALLAX_FACTOR (DIM_SUBPIXEL_PARALLAX_FACTOR)
-
 
 //------------------------------------------------------- PARAMETERS
 
@@ -44,11 +40,6 @@ Texture2D<float> ClosestDepthTexture;
 Texture2D<uint> PrevUseCountTexture;
 Texture2D<uint> PrevClosestDepthTexture;
 
-#if CONFIG_SUBPIXEL_PARALLAX_FACTOR
-Texture2D<tsr_subpixel_payload> ParallaxFactorTexture;
-Texture2D<tsr_subpixel_details> PrevHistorySubpixelDetails;
-#endif
-
 Texture2D<tsr_halfC> PrevHistoryGuide;
 Texture2D<tsr_half4> PrevHistoryMoire;
 Texture2D<float2> PrevVelocityTexture;
@@ -72,9 +63,6 @@ RWTexture2D<float2> HoleFilledVelocityOutput;
 RWTexture2D<float2> GrandHoleFilledVelocityOutput;
 RWTexture2D<tsr_half2> HoleFilledVelocityMaskOutput;
 RWTexture2D<tsr_half> ParallaxRejectionMaskOutput;
-#if CONFIG_SUBPIXEL_PARALLAX_FACTOR
-	RWTexture2D<tsr_subpixel_payload> HistorySubpixelDetailsOutput;
-#endif
 
 
 //------------------------------------------------------- LDS
@@ -85,6 +73,293 @@ RWTexture2D<tsr_half> ParallaxRejectionMaskOutput;
 #endif
 
 
+//------------------------------------------------------- FUNCTIONS
+
+#if CONFIG_VELOCITY_SAMPLES != 1
+void LoadVelocityNeighborhood(
+	tsr_short2 InputPixelPos,
+	out float2 EncodedVelocityNeighborhood[CONFIG_VELOCITY_PRE_FETCH_SAMPLES])
+{
+	tsr_short2 LaneOffsetSign = GetLaneOffsetSign();
+
+	UNROLL_N(CONFIG_VELOCITY_PRE_FETCH_SAMPLES)
+	for (uint i = 0; i < CONFIG_VELOCITY_PRE_FETCH_SAMPLES; i++)
+	{
+		tsr_short2 SampleInputPixelPos;
+		#if CONFIG_BUTTERFLY_KERNEL
+		{
+			tsr_short2 Offset = tsr_short2(Offsets2x2[i]) * LaneOffsetSign;
+				
+			SampleInputPixelPos = InputPixelPos + Offset;
+		}
+		#else
+		{
+			tsr_short2 Offset;
+			#if CONFIG_VELOCITY_SAMPLES == 9
+				Offset = tsr_short2(kOffsets3x3[kSquareIndexes3x3[i]]);
+			#elif CONFIG_VELOCITY_SAMPLES == 5
+				Offset = tsr_short2(kOffsets3x3[kPlusIndexes3x3[i]]);
+			#endif
+					
+			SampleInputPixelPos = InputPixelPos + ClampPixelOffset(
+				InputPixelPos,
+				Offset, Offset,
+				InputPixelPosMin, InputPixelPosMax);
+		}
+		#endif
+
+		EncodedVelocityNeighborhood[i] = DilatedVelocityTexture[SampleInputPixelPos];
+	}
+} // LoadVelocityNeighborhood()
+
+void ProcessVelocityNeighborhood(
+	float2 EncodedVelocityNeighborhood[CONFIG_VELOCITY_PRE_FETCH_SAMPLES],
+	out float2 EncodedScreenVelocity,
+	out tsr_half RejectHighFrequency)
+{
+	// Preprocessed the prefetched samples
+	float2 PixelVelocityNeighborhood[CONFIG_VELOCITY_PRE_FETCH_SAMPLES];
+	{
+		UNROLL_N(CONFIG_VELOCITY_PRE_FETCH_SAMPLES)
+		for (uint i = 0; i < CONFIG_VELOCITY_PRE_FETCH_SAMPLES; i++)
+		{
+			PixelVelocityNeighborhood[i] = DecodeVelocityFromTexture(float4(EncodedVelocityNeighborhood[i], 0.0, 0.0)).xy * ScreenVelocityToInputPixelVelocity;
+		}
+	}
+	
+	RejectHighFrequency = tsr_half(0.0);
+
+	// Sample
+	float2 ScreenPixelVelocity = PixelVelocityNeighborhood[0];
+	tsr_short2 LaneOffsetSign = GetLaneOffsetSign();
+
+	UNROLL_N(CONFIG_VELOCITY_SAMPLE_PAIRS)
+	for (uint i = 0; i < CONFIG_VELOCITY_SAMPLE_PAIRS; i++)
+	ISOLATE
+	{
+		tsr_short2 Offset0;
+		tsr_short2 Offset1;
+			
+		float2 NeighborPixelVelocity0;
+		float2 NeighborPixelVelocity1;
+
+		#if CONFIG_BUTTERFLY_KERNEL
+		{
+			const uint PairIndex = i;
+				
+			const uint ButterflyArray0[] = { 0x0, 0x0, 0x0, 0x1 };
+			const uint ButterflyArray1[] = { 0x1, 0x2, 0x3, 0x2 };
+				
+			uint FetchIndex[2];
+			FWaveBroadcastSettings BroadcastSettings[2];
+			BroadcastSettings[0] = InitWaveXorButterfly(ButterflyArray0[PairIndex]);
+			BroadcastSettings[1] = InitWaveXorButterfly(ButterflyArray1[PairIndex]);
+				
+			if (PairIndex == 3)
+			{
+				FetchIndex[0] = 2;
+				FetchIndex[1] = 1;
+				Offset0 = tsr_short2(-1, 1);
+			}
+			else
+			{
+				FetchIndex[0] = PairIndex + 1;
+				FetchIndex[1] = PairIndex + 1;
+				Offset0 = tsr_short2(Offsets2x2[PairIndex + 1]);
+			}
+
+			Offset0 *= LaneOffsetSign;
+			Offset1 = -Offset0;
+
+			NeighborPixelVelocity0 = WaveBroadcast(BroadcastSettings[0], PixelVelocityNeighborhood[FetchIndex[0]]);
+			NeighborPixelVelocity1 = WaveBroadcast(BroadcastSettings[1], PixelVelocityNeighborhood[FetchIndex[1]]);
+		}
+		#else
+		{
+			const uint NeightbordId0 = i + 1;
+			const uint NeightbordId1 = i + 1 + CONFIG_VELOCITY_SAMPLE_PAIRS;
+			
+			#if CONFIG_VELOCITY_SAMPLES == 9
+				Offset0 = tsr_short2(kOffsets3x3[kSquareIndexes3x3[NeightbordId0]]);
+				Offset1 = tsr_short2(kOffsets3x3[kSquareIndexes3x3[NeightbordId1]]);
+			#elif CONFIG_VELOCITY_SAMPLES == 5
+				Offset0 = tsr_short2(kOffsets3x3[kPlusIndexes3x3[NeightbordId0]]);
+				Offset1 = tsr_short2(kOffsets3x3[kPlusIndexes3x3[NeightbordId1]]);
+			#endif
+					
+			NeighborPixelVelocity0 = PixelVelocityNeighborhood[NeightbordId0];
+			NeighborPixelVelocity1 = PixelVelocityNeighborhood[NeightbordId1];
+		}
+		#endif
+
+		tsr_half2 NeighborDeltaPixelVelocity0 = tsr_half2(ScreenPixelVelocity - NeighborPixelVelocity0);
+		tsr_half2 NeighborDeltaPixelVelocity1 = tsr_half2(ScreenPixelVelocity - NeighborPixelVelocity1);
+			
+		// Reject high frequency on velocity edge discontinuity.
+		{
+			tsr_half IsOverlapingEdge0 = tsr_half(1.0) - dot(saturate(NeighborDeltaPixelVelocity0 * Offset0), tsr_half(1.0).xx); // * saturate(ScreenPixelVelocityLength - NeighborPixelVelocityLength0);
+			tsr_half IsOverlapingEdge1 = tsr_half(1.0) - dot(saturate(NeighborDeltaPixelVelocity1 * Offset1), tsr_half(1.0).xx); // * saturate(ScreenPixelVelocityLength - NeighborPixelVelocityLength1);
+				
+			tsr_half IsOverlapingEdge = min(IsOverlapingEdge0, IsOverlapingEdge1);
+
+			if (i == 0)
+				RejectHighFrequency = IsOverlapingEdge;
+			else
+				RejectHighFrequency = min(RejectHighFrequency, IsOverlapingEdge);
+		}
+	}
+
+	// Remove noise from high frequency rejections.
+	RejectHighFrequency = saturate(tsr_half(1.0) - RejectHighFrequency * tsr_half(1.1));
+
+	// Do not reject high frequency additionally where hole filling happen
+	// RejectHighFrequency = max(RejectHighFrequency, HoleFilledVelocityMask); TODO
+
+	EncodedScreenVelocity = EncodedVelocityNeighborhood[0];
+} // ProcessVelocityNeighborhood()
+
+#endif // CONFIG_VELOCITY_SAMPLES != 1
+
+void LoadPrevAtomicTexturesSamples(
+	float2 PrevScreenPos,
+	out uint HistoryUseCountSamples[4],
+	out uint HistoryClosestDeviceZSamples[4])
+{
+	float2 PrevInputBufferUV = (InputInfo_ScreenPosToViewportScale * PrevScreenPos + InputInfo_ScreenPosToViewportBias) * InputInfo_ExtentInverse;
+	FBilinearSampleInfos BilinearInter = GetBilinearSampleLevelInfos(PrevInputBufferUV, InputInfo_Extent, InputInfo_ExtentInverse);
+
+	UNROLL_N(4)
+	for (uint i = 0; i < 4; i++)
+	{
+		float BilinearWeight = GetSampleWeight(BilinearInter, i);
+		uint2 PixelPos = GetSamplePixelCoord(BilinearInter, i);
+
+		bool bValidPixelPos = all(and(PixelPos >= InputInfo_ViewportMin, PixelPos < InputInfo_ViewportMax));
+
+		#if COMPRESS_PREV_USE_COUNT
+			uint2 CompressedUseCountPixelPos;
+			uint Shift;
+			ComputeCompressedUseCountPixelCoordinates(PixelPos, /* out */ CompressedUseCountPixelPos, /* out */ Shift);
+
+			uint HistoryUseCount = PrevUseCountTexture[CompressedUseCountPixelPos];
+		#else
+			uint HistoryUseCount = PrevUseCountTexture[PixelPos];
+		#endif
+			
+		HistoryUseCountSamples[i] = HistoryUseCount;
+		HistoryClosestDeviceZSamples[i] = PrevClosestDepthTexture[PixelPos];
+	}
+} // LoadPrevAtomicTexturesSamples()
+
+void ProcessPrevAtomicTexturesSamples(
+	uint HistoryUseCountSamples[4],
+	uint HistoryClosestDeviceZSamples[4],
+	float DeviceZ,
+	float2 PrevScreenPos,
+	out tsr_half ParallaxRejectionMask,
+	out bool bIsValidHoleFillingPixelVelocity,
+	out uint EncodedHoleFillingVelocity)
+{
+	ParallaxRejectionMask = tsr_half(0.0);
+	bIsValidHoleFillingPixelVelocity = false;
+	EncodedHoleFillingVelocity = 0;
+
+	float WorldDepth = ConvertFromDeviceZ(DeviceZ);
+	
+	float2 PrevInputBufferUV = (InputInfo_ScreenPosToViewportScale * PrevScreenPos + InputInfo_ScreenPosToViewportBias) * InputInfo_ExtentInverse;
+	FBilinearSampleInfos BilinearInter = GetBilinearSampleLevelInfos(PrevInputBufferUV, InputInfo_Extent, InputInfo_ExtentInverse);
+
+	UNROLL_N(4)
+	for (uint i = 0; i < 4; i++)
+	{
+		float BilinearWeight = GetSampleWeight(BilinearInter, i);
+		uint2 PixelPos = GetSamplePixelCoord(BilinearInter, i);
+
+		bool bValidPixelPos = all(and(PixelPos >= InputInfo_ViewportMin, PixelPos < InputInfo_ViewportMax));
+				
+		#if COMPRESS_PREV_USE_COUNT
+			uint2 CompressedUseCountPixelPos;
+			uint Shift;
+			ComputeCompressedUseCountPixelCoordinates(PixelPos, /* out */ CompressedUseCountPixelPos, /* out */ Shift);
+
+			uint HistoryUseCount = (HistoryUseCountSamples[i] >> (8 * Shift)) & 0xFF;
+		#else
+			uint HistoryUseCount = HistoryUseCountSamples[i];
+		#endif
+			
+		uint SampleHistoryClosestDeviceZ = HistoryClosestDeviceZSamples[i];
+		
+		float HistoryClosestDeviceZ = f16tof32(SampleHistoryClosestDeviceZ >> VELOCITY_HOLE_FILLING_BITS);
+		float HistoryClosestWorldDepth = ConvertFromDeviceZ(HistoryClosestDeviceZ);
+
+		tsr_half DepthRejection;
+		{
+			const float WorldDepthToPixelWorldRadius = GetTanHalfFieldOfView().x * InputInfo_ViewportSizeInverse.x;
+			const float PixelDepthError = 3.0;
+			const float RadiusToDiamater = 2.0;
+			
+			float WorldDepthEpsilon = (WorldDepthToPixelWorldRadius * RadiusToDiamater * PixelDepthError) * HistoryClosestWorldDepth; 
+			float DeltaDepth = abs(HistoryClosestWorldDepth - WorldDepth);
+
+			DepthRejection = saturate(tsr_half(2.0) - tsr_half(DeltaDepth / WorldDepthEpsilon));
+		}
+
+		tsr_half HistoryGhostingRejection = saturate(tsr_half(1.0) - tsr_half(3.0) * abs(tsr_half(HistoryUseCount) * rcp(tsr_half(PREV_USE_COUNT_QUANTIZATION)) - tsr_half(1.0)));
+			
+		bIsValidHoleFillingPixelVelocity = bIsValidHoleFillingPixelVelocity || bValidPixelPos;
+
+		EncodedHoleFillingVelocity = select(bValidPixelPos, max(EncodedHoleFillingVelocity, SampleHistoryClosestDeviceZ), EncodedHoleFillingVelocity);
+		ParallaxRejectionMask = select(bValidPixelPos, ParallaxRejectionMask + tsr_half(BilinearWeight) * max(DepthRejection, HistoryGhostingRejection), ParallaxRejectionMask);
+	}
+} // ProcessPrevAtomicTexturesSamples()
+
+// Sample a texture with a catrom kernel
+template<typename Type>
+Type SampleTextureCatrom(Texture2D<Type> InTexture, float2 TextureUV)
+{
+	FCatmullRomSamples Samples = GetBicubic2DCatmullRomSamples(TextureUV, PrevGuideInfo_Extent, PrevGuideInfo_ExtentInverse);
+
+	Type Color = tsr_half(0.0);
+	for (uint i = 0; i < Samples.Count; i++)
+	{
+		float2 SampleUV = Samples.UV[i];
+		SampleUV = clamp(
+			SampleUV,
+			PrevGuideInfo_UVViewportBilinearMin,
+			PrevGuideInfo_UVViewportBilinearMax);
+
+		Color += InTexture.SampleLevel(GlobalBilinearClampedSampler, SampleUV, 0) * tsr_half(Samples.Weight[i]);
+	}
+	Color *= tsr_half(Samples.FinalMultiplier);
+
+	return Color;
+}
+
+// Correct pre exposure of the guide color.
+tsr_halfC CorrectGuideColorExposure(tsr_halfC ReprojectedHistoryGuideColor)
+{
+	tsr_halfC LinearReprojectedHistoryGuideColor = ReprojectedHistoryGuideColor * rcp(tsr_half(1.0) - ReprojectedHistoryGuideColor) * tsr_half(HistoryPreExposureCorrection);
+				
+	ReprojectedHistoryGuideColor = LinearReprojectedHistoryGuideColor * rcp(LinearReprojectedHistoryGuideColor + tsr_half(1.0));
+
+	// Removes NaN
+	ReprojectedHistoryGuideColor = -min(-ReprojectedHistoryGuideColor, tsr_half(0.0));
+
+	return ReprojectedHistoryGuideColor;
+}
+
+// Sample previous moire history
+#if CONFIG_MOIRE_REPROJECTION
+tsr_half4 ReprojectPrevHistoryMoire(tsr_ushort2 InputPixelPos, float2 TextureUV)
+{
+	uint2 Random = Rand3DPCG16(int3(InputPixelPos - InputInfo_ViewportMin, View.StateFrameIndexMod8)).xy;
+	float2 E = Hammersley16(0, 1, Random).xy;
+
+	return PrevHistoryMoire.SampleLevel(GlobalPointClampedSampler, TextureUV + (E - 0.5) * PrevGuideInfo_ExtentInverse, 0);
+}
+#endif
+
+
 //------------------------------------------------------- ENTRY POINT
 
 [numthreads(TILE_SIZE * TILE_SIZE, 1, 1)]
@@ -97,10 +372,6 @@ void MainCS(
 	// Sample current frame informations.
 	float2 EncodedScreenVelocity;
 	float DeviceZ;
-	#if CONFIG_SUBPIXEL_PARALLAX_FACTOR
-		tsr_subpixel_payload ParallaxFactorBits;
-		tsr_half ParallaxFactor;
-	#endif
 	tsr_half RejectHighFrequency = tsr_half(0.0);
 	ISOLATE
 	{
@@ -115,50 +386,12 @@ void MainCS(
 		// Issue texture fetechs
 		ISOLATE
 		{
-			#if CONFIG_SUBPIXEL_PARALLAX_FACTOR
-				tsr_subpixel_payload RawParallaxFactor = ParallaxFactorTexture[InputPixelPos];
-			#endif
-
 			DeviceZ = ClosestDepthTexture[InputPixelPos];
 
-			#if CONFIG_SUBPIXEL_PARALLAX_FACTOR
-				ParallaxFactor = UncompressParallaxFactor(RawParallaxFactor & tsr_subpixel_payload(SUB_PIXEL_PARALLAX_FACTOR_BIT_MASK));
-				ParallaxFactorBits = RawParallaxFactor & tsr_subpixel_payload(SUB_PIXEL_PARALLAX_FACTOR_BIT_MASK);
-			#endif
-
 			#if CONFIG_VELOCITY_SAMPLES == 1
 				EncodedScreenVelocity = DilatedVelocityTexture[InputPixelPos];
 			#else // CONFIG_VELOCITY_SAMPLES != 1
-				tsr_short2 LaneOffsetSign = GetLaneOffsetSign();
-
-				UNROLL_N(CONFIG_VELOCITY_PRE_FETCH_SAMPLES)
-				for (uint i = 0; i < CONFIG_VELOCITY_PRE_FETCH_SAMPLES; i++)
-				{
-					tsr_short2 SampleInputPixelPos;
-					#if CONFIG_BUTTERFLY_KERNEL
-					{
-						tsr_short2 Offset = tsr_short2(Offsets2x2[i]) * LaneOffsetSign;
-				
-						SampleInputPixelPos = InputPixelPos + Offset;
-					}
-					#else
-					{
-						tsr_short2 Offset;
-						#if CONFIG_VELOCITY_SAMPLES == 9
-							Offset = tsr_short2(kOffsets3x3[kSquareIndexes3x3[i]]);
-						#elif CONFIG_VELOCITY_SAMPLES == 5
-							Offset = tsr_short2(kOffsets3x3[kPlusIndexes3x3[i]]);
-						#endif
-					
-						SampleInputPixelPos = InputPixelPos + ClampPixelOffset(
-							InputPixelPos,
-							Offset, Offset,
-							InputPixelPosMin, InputPixelPosMax);
-					}
-					#endif
-
-					EncodedVelocityNeighborhood[i] = DilatedVelocityTexture[SampleInputPixelPos];
-				}
+				LoadVelocityNeighborhood(InputPixelPos, /* out */ EncodedVelocityNeighborhood);
 			#endif // CONFIG_VELOCITY_SAMPLES != 1
 		}
 
@@ -179,103 +412,10 @@ void MainCS(
 		#if CONFIG_VELOCITY_SAMPLES != 1
 		ISOLATE
 		{
-			// Preprocessed the prefetched samples
-			float2 PixelVelocityNeighborhood[CONFIG_VELOCITY_PRE_FETCH_SAMPLES];
-			{
-				UNROLL_N(CONFIG_VELOCITY_PRE_FETCH_SAMPLES)
-				for (uint i = 0; i < CONFIG_VELOCITY_PRE_FETCH_SAMPLES; i++)
-				{
-					PixelVelocityNeighborhood[i] = DecodeVelocityFromTexture(float4(EncodedVelocityNeighborhood[i], 0.0, 0.0)).xy * ScreenVelocityToInputPixelVelocity;
-				}
-			}
-
-			// Sample
-			float2 ScreenPixelVelocity = PixelVelocityNeighborhood[0];
-			tsr_short2 LaneOffsetSign = GetLaneOffsetSign();
-
-			UNROLL_N(CONFIG_VELOCITY_SAMPLE_PAIRS)
-			for (uint i = 0; i < CONFIG_VELOCITY_SAMPLE_PAIRS; i++)
-			ISOLATE
-			{
-				tsr_short2 Offset0;
-				tsr_short2 Offset1;
-			
-				float2 NeighborPixelVelocity0;
-				float2 NeighborPixelVelocity1;
-
-				#if CONFIG_BUTTERFLY_KERNEL
-				{
-					const uint PairIndex = i;
-				
-					const uint ButterflyArray0[] = { 0x0, 0x0, 0x0, 0x1 };
-					const uint ButterflyArray1[] = { 0x1, 0x2, 0x3, 0x2 };
-				
-					uint FetchIndex[2];
-					FWaveBroadcastSettings BroadcastSettings[2];
-					BroadcastSettings[0] = InitWaveXorButterfly(ButterflyArray0[PairIndex]);
-					BroadcastSettings[1] = InitWaveXorButterfly(ButterflyArray1[PairIndex]);
-				
-					if (PairIndex == 3)
-					{
-						FetchIndex[0] = 2;
-						FetchIndex[1] = 1;
-						Offset0 = tsr_short2(-1, 1);
-					}
-					else
-					{
-						FetchIndex[0] = PairIndex + 1;
-						FetchIndex[1] = PairIndex + 1;
-						Offset0 = tsr_short2(Offsets2x2[PairIndex + 1]);
-					}
-
-					Offset0 *= LaneOffsetSign;
-					Offset1 = -Offset0;
-
-					NeighborPixelVelocity0 = WaveBroadcast(BroadcastSettings[0], PixelVelocityNeighborhood[FetchIndex[0]]);
-					NeighborPixelVelocity1 = WaveBroadcast(BroadcastSettings[1], PixelVelocityNeighborhood[FetchIndex[1]]);
-				}
-				#else
-				{
-					const uint NeightbordId0 = i + 1;
-					const uint NeightbordId1 = i + 1 + CONFIG_VELOCITY_SAMPLE_PAIRS;
-			
-					#if CONFIG_VELOCITY_SAMPLES == 9
-						Offset0 = tsr_short2(kOffsets3x3[kSquareIndexes3x3[NeightbordId0]]);
-						Offset1 = tsr_short2(kOffsets3x3[kSquareIndexes3x3[NeightbordId1]]);
-					#elif CONFIG_VELOCITY_SAMPLES == 5
-						Offset0 = tsr_short2(kOffsets3x3[kPlusIndexes3x3[NeightbordId0]]);
-						Offset1 = tsr_short2(kOffsets3x3[kPlusIndexes3x3[NeightbordId1]]);
-					#endif
-					
-					NeighborPixelVelocity0 = PixelVelocityNeighborhood[NeightbordId0];
-					NeighborPixelVelocity1 = PixelVelocityNeighborhood[NeightbordId1];
-				}
-				#endif
-
-				tsr_half2 NeighborDeltaPixelVelocity0 = tsr_half2(ScreenPixelVelocity - NeighborPixelVelocity0);
-				tsr_half2 NeighborDeltaPixelVelocity1 = tsr_half2(ScreenPixelVelocity - NeighborPixelVelocity1);
-			
-				// Reject high frequency on velocity edge discontinuity.
-				{
-					tsr_half IsOverlapingEdge0 = tsr_half(1.0) - dot(saturate(NeighborDeltaPixelVelocity0 * Offset0), tsr_half(1.0).xx); // * saturate(ScreenPixelVelocityLength - NeighborPixelVelocityLength0);
-					tsr_half IsOverlapingEdge1 = tsr_half(1.0) - dot(saturate(NeighborDeltaPixelVelocity1 * Offset1), tsr_half(1.0).xx); // * saturate(ScreenPixelVelocityLength - NeighborPixelVelocityLength1);
-				
-					tsr_half IsOverlapingEdge = min(IsOverlapingEdge0, IsOverlapingEdge1);
-
-					if (i == 0)
-						RejectHighFrequency = IsOverlapingEdge;
-					else
-						RejectHighFrequency = min(RejectHighFrequency, IsOverlapingEdge);
-				}
-			}
-
-			// Remove noise from high frequency rejections.
-			RejectHighFrequency = saturate(tsr_half(1.0) - RejectHighFrequency * tsr_half(1.1));
-
-			// Do not reject high frequency additionally where hole filling happen
-			// RejectHighFrequency = max(RejectHighFrequency, HoleFilledVelocityMask); TODO
-
-			EncodedScreenVelocity = EncodedVelocityNeighborhood[0];
+			ProcessVelocityNeighborhood(
+				EncodedVelocityNeighborhood,
+				/* out */ EncodedScreenVelocity,
+				/* out */ RejectHighFrequency);
 		}
 		#endif // CONFIG_VELOCITY_SAMPLES != 1
 	}
@@ -288,9 +428,6 @@ void MainCS(
 	#if CONFIG_MOIRE_REPROJECTION
 		tsr_half4 ReprojectedHistoryMoireColor;
 	#endif
-	#if CONFIG_SUBPIXEL_PARALLAX_FACTOR
-		FSubpixelNeighborhood SubpixelNeigborhood;
-	#endif
 	bool bIsOffScreen;
 
 	// Issues all the history texture fetches
@@ -302,33 +439,9 @@ void MainCS(
 			
 		float2 ScreenPos = ApplyScreenTransform(float2(InputPixelPos), InputPixelPosToScreenPos);
 		float2 PrevScreenPos = ScreenPos - DecodeVelocityFromTexture(float4(EncodedScreenVelocity, 0.0, 0.0)).xy;
-		{
-			float2 PrevInputBufferUV = (InputInfo_ScreenPosToViewportScale * PrevScreenPos + InputInfo_ScreenPosToViewportBias) * InputInfo_ExtentInverse;
-		
-			FBilinearSampleInfos BilinearInter = GetBilinearSampleLevelInfos(PrevInputBufferUV, InputInfo_Extent, InputInfo_ExtentInverse);
-
-			UNROLL_N(4)
-			for (uint i = 0; i < 4; i++)
-			{
-				float BilinearWeight = GetSampleWeight(BilinearInter, i);
-				uint2 PixelPos = GetSamplePixelCoord(BilinearInter, i);
-
-				bool bValidPixelPos = all(and(PixelPos >= InputInfo_ViewportMin, PixelPos < InputInfo_ViewportMax));
-
-				#if COMPRESS_PREV_USE_COUNT
-					uint2 CompressedUseCountPixelPos;
-					uint Shift;
-					ComputeCompressedUseCountPixelCoordinates(PixelPos, /* out */ CompressedUseCountPixelPos, /* out */ Shift);
-
-					uint HistoryUseCount = PrevUseCountTexture[CompressedUseCountPixelPos];
-				#else
-					uint HistoryUseCount = PrevUseCountTexture[PixelPos];
-				#endif
-			
-				HistoryUseCountSamples[i] = HistoryUseCount;
-				HistoryClosestDeviceZSamples[i] = PrevClosestDepthTexture[PixelPos];
-			}
-		}
+		LoadPrevAtomicTexturesSamples(PrevScreenPos,
+			/* out */ HistoryUseCountSamples,
+			/* out */ HistoryClosestDeviceZSamples);
 			
 		{
 			float2 OutputScreenPos = ApplyScreenTransform(float2(InputPixelPos), InputPixelPosToReprojectScreenPos);
@@ -341,79 +454,20 @@ void MainCS(
 				(PrevOutputScreenPos.y > 1.0));
 				
 			float2 KernelHistoryBufferUV = ApplyScreenTransform(PrevOutputScreenPos, ScreenPosToPrevHistoryGuideBufferUV);
-			FCatmullRomSamples Samples = GetBicubic2DCatmullRomSamples(KernelHistoryBufferUV, PrevGuideInfo_Extent, PrevGuideInfo_ExtentInverse);
-
-			ReprojectedHistoryGuideColor = tsr_half(0.0);
-			for (uint i = 0; i < Samples.Count; i++)
-			{
-				float2 SampleUV = Samples.UV[i];
-				SampleUV = clamp(
-					SampleUV,
-					PrevGuideInfo_UVViewportBilinearMin,
-					PrevGuideInfo_UVViewportBilinearMax);
-
-				ReprojectedHistoryGuideColor += PrevHistoryGuide.SampleLevel(GlobalBilinearClampedSampler, SampleUV, 0) * tsr_half(Samples.Weight[i]);
-			}
-			ReprojectedHistoryGuideColor *= tsr_half(Samples.FinalMultiplier);
-
-			// Correct pre exposure of the guide color.
-			{
-				tsr_halfC LinearReprojectedHistoryGuideColor = ReprojectedHistoryGuideColor * rcp(tsr_half(1.0) - ReprojectedHistoryGuideColor) * tsr_half(HistoryPreExposureCorrection);
-				
-				ReprojectedHistoryGuideColor = LinearReprojectedHistoryGuideColor * rcp(LinearReprojectedHistoryGuideColor + tsr_half(1.0));
-
-				// Removes NaN
-				ReprojectedHistoryGuideColor = -min(-ReprojectedHistoryGuideColor, tsr_half(0.0));
-			}
 
+			ReprojectedHistoryGuideColor = SampleTextureCatrom(PrevHistoryGuide, KernelHistoryBufferUV);
+			ReprojectedHistoryGuideColor = CorrectGuideColorExposure(ReprojectedHistoryGuideColor);
 			#if CONFIG_MOIRE_REPROJECTION
 			{
-				uint2 Random = Rand3DPCG16(int3(InputPixelPos - InputInfo_ViewportMin, View.StateFrameIndexMod8)).xy;
-				float2 E = Hammersley16(0, 1, Random).xy;
-
-				ReprojectedHistoryMoireColor = PrevHistoryMoire.SampleLevel(GlobalPointClampedSampler, KernelHistoryBufferUV + (E - 0.5) * PrevGuideInfo_ExtentInverse, 0);
+				ReprojectedHistoryMoireColor = ReprojectPrevHistoryMoire(InputPixelPos, KernelHistoryBufferUV);
 			}
 			#endif
 		}
-
-		#if CONFIG_SUBPIXEL_PARALLAX_FACTOR
-		{
-			// Pixel coordinate of the input pixel in the predictor viewport.
-			float2 InputPPCk = float2(InputPixelPos - InputInfo_ViewportMin) + 0.5f - InputJitter;
-			float2 OutputViewportUV = InputPPCk * InputInfo_ViewportSizeInverse;
-			float2 OutputScreenPos = ViewportUVToScreenPos(OutputViewportUV);
-	
-			float2 PrevOutputScreenPos = OutputScreenPos - DecodeVelocityFromTexture(float4(EncodedScreenVelocity, 0.0, 0.0)).xy;
-			bIsOffScreen = (
-				(PrevOutputScreenPos.x < -1.0) ||
-				(PrevOutputScreenPos.y < -1.0) ||
-				(PrevOutputScreenPos.x > 1.0) ||
-				(PrevOutputScreenPos.y > 1.0));
-
-			float2 PrevOutputViewportUV = ScreenPosToViewportUV(PrevOutputScreenPos);
-
-			float2 HistoryPPCk = PrevOutputViewportUV * PrevHistoryInfo_ViewportSize;
-
-			float2 KernelHistoryBufferUV = clamp(
-				HistoryPPCk * PrevHistoryInfo_ExtentInverse,
-				PrevHistoryInfo_UVViewportBilinearMin,
-				PrevHistoryInfo_UVViewportBilinearMax);
-		
-			float2 SubpixelBufferUV = ApplyScreenTransform(PrevScreenPos, ScreenPosToPrevSubpixelDetails);
-
-			SubpixelNeigborhood = GatherPrevSubpixelNeighborhood(PrevHistorySubpixelDetails, SubpixelBufferUV);
-		}
-		#endif
 	}
 	
-
 	// Process the history samples
-	tsr_half ParallaxRejectionMask = 0.0;
+	tsr_half ParallaxRejectionMask;
 	
-	#if CONFIG_SUBPIXEL_PARALLAX_FACTOR
-		tsr_subpixel_details FinalSubpixelDetails;
-	#endif
-
 	ISOLATE
 	{
 		// LDS despill
@@ -431,107 +485,18 @@ void MainCS(
 		float2 ScreenPos = ApplyScreenTransform(float2(InputPixelPos), InputPixelPosToScreenPos);
 		float2 ScreenVelocity = DecodeVelocityFromTexture(float4(EncodedScreenVelocity, 0.0, 0.0)).xy;
 		float2 PrevScreenPos = ScreenPos - ScreenVelocity;
-	
-		#if CONFIG_SUBPIXEL_PARALLAX_FACTOR
-		// Compute HistoryMaxRejectionFactor
-		tsr_half HistoryMaxRejectionFactor;
-		tsr_subpixel_details PrevSubpixelDetails;
-		{
-			uint EncodedHistoryMaxRejectionFactor = 0;
-			
-			PrevSubpixelDetails = 0;
-
-			UNROLL_N(SUB_PIXEL_COUNT)
-			for (uint SubpixelId = 0; SubpixelId < SUB_PIXEL_COUNT; SubpixelId++)
-			{
-				tsr_subpixel_payload SubpixelPayload = GetSubpixelPayload(SubpixelNeigborhood, SubpixelId);
-				tsr_subpixel_payload EncodedRejectionFactor = (SubpixelPayload >> SUB_PIXEL_PARALLAX_FACTOR_BIT_OFFSET) & SUB_PIXEL_PARALLAX_FACTOR_BIT_MASK;
-				
-				PrevSubpixelDetails |= SubpixelPayload << tsr_subpixel_payload(SUB_PIXEL_BIT_COUNT * SubpixelId);
-
-				EncodedHistoryMaxRejectionFactor = max(EncodedHistoryMaxRejectionFactor, EncodedRejectionFactor);
-			}
-
-			HistoryMaxRejectionFactor = UncompressParallaxFactor(EncodedHistoryMaxRejectionFactor);
-		}
-		#endif
 		
 		// Compute the parralax rejection mask
 		bool bIsValidHoleFillingPixelVelocity = false;
 		uint EncodedHoleFillingVelocity = 0;
-		{
-			float WorldDepth = ConvertFromDeviceZ(DeviceZ);
-	
-			float2 PrevInputBufferUV = (InputInfo_ScreenPosToViewportScale * PrevScreenPos + InputInfo_ScreenPosToViewportBias) * InputInfo_ExtentInverse;
-			FBilinearSampleInfos BilinearInter = GetBilinearSampleLevelInfos(PrevInputBufferUV, InputInfo_Extent, InputInfo_ExtentInverse);
-
-			UNROLL_N(4)
-			for (uint i = 0; i < 4; i++)
-			{
-				float BilinearWeight = GetSampleWeight(BilinearInter, i);
-				uint2 PixelPos = GetSamplePixelCoord(BilinearInter, i);
-
-				bool bValidPixelPos = all(and(PixelPos >= InputInfo_ViewportMin, PixelPos < InputInfo_ViewportMax));
-				
-				#if COMPRESS_PREV_USE_COUNT
-					uint2 CompressedUseCountPixelPos;
-					uint Shift;
-					ComputeCompressedUseCountPixelCoordinates(PixelPos, /* out */ CompressedUseCountPixelPos, /* out */ Shift);
-
-					uint HistoryUseCount = (HistoryUseCountSamples[i] >> (8 * Shift)) & 0xFF;
-				#else
-					uint HistoryUseCount = HistoryUseCountSamples[i];
-				#endif
-			
-				uint SampleHistoryClosestDeviceZ = HistoryClosestDeviceZSamples[i];
-		
-				float HistoryClosestDeviceZ = f16tof32(SampleHistoryClosestDeviceZ >> VELOCITY_HOLE_FILLING_BITS);
-				float HistoryClosestWorldDepth = ConvertFromDeviceZ(HistoryClosestDeviceZ);
-
-				tsr_half DepthRejection;
-				{
-					const float WorldDepthToPixelWorldRadius = GetTanHalfFieldOfView().x * InputInfo_ViewportSizeInverse.x;
-					const float PixelDepthError = 3.0;
-					const float RadiusToDiamater = 2.0;
-			
-					float WorldDepthEpsilon = (WorldDepthToPixelWorldRadius * RadiusToDiamater * PixelDepthError) * HistoryClosestWorldDepth; 
-					float DeltaDepth = abs(HistoryClosestWorldDepth - WorldDepth);
-
-					DepthRejection = saturate(tsr_half(2.0) - tsr_half(DeltaDepth / WorldDepthEpsilon));
-				}
-
-				tsr_half HistoryGhostingRejection = saturate(tsr_half(1.0) - tsr_half(3.0) * abs(tsr_half(HistoryUseCount) * rcp(tsr_half(PREV_USE_COUNT_QUANTIZATION)) - tsr_half(1.0)));
-			
-				bIsValidHoleFillingPixelVelocity = bIsValidHoleFillingPixelVelocity || bValidPixelPos;
-
-				EncodedHoleFillingVelocity = select(bValidPixelPos, max(EncodedHoleFillingVelocity, SampleHistoryClosestDeviceZ), EncodedHoleFillingVelocity);
-				ParallaxRejectionMask = select(bValidPixelPos, ParallaxRejectionMask + tsr_half(BilinearWeight) * max(DepthRejection, HistoryGhostingRejection), ParallaxRejectionMask);
-			}
-		}
-
-		// Reject subpixel details that might be in history but not input just based on their parrallax
-		#if CONFIG_SUBPIXEL_PARALLAX_FACTOR
-		{
-			tsr_half HistorySubdetailRejectionFactor = max(HistoryMaxRejectionFactor / ParallaxFactor, tsr_half(1.0));
-
-			float4 ThisClip = float4(ScreenPos, DeviceZ, 1);
-			float4 PrevClip = mul(ThisClip, View.ClipToPrevClip); 
-			float2 PrevScreen = PrevClip.xy / PrevClip.w;
-			
-			float4 RotationalPrevClip = mul(ThisClip, RotationalClipToPrevClip);
-			float2 RotationalPrevScreen = RotationalPrevClip.xy / RotationalPrevClip.w;
-			
-			tsr_half PixelParallaxVelocity = tsr_half(0.5) * length(tsr_half2(RotationalPrevScreen - PrevScreen) * tsr_half2(InputInfo_ViewportSize));
-			tsr_half HistorySubdetailAdditionalParralax = (HistorySubdetailRejectionFactor - tsr_half(1.0)) * PixelParallaxVelocity;
-
-			// PixelParallaxVelocity can be +inf and generate NaN if HistorySubdetailRejectionFactor = 1.0;
-			HistorySubdetailAdditionalParralax = -min(-HistorySubdetailAdditionalParralax, tsr_half(0.0));
-
-			tsr_half HistoryParallaxRejectionMask = saturate(tsr_half(1.0) - HistorySubdetailAdditionalParralax);
-
-			ParallaxRejectionMask = min(ParallaxRejectionMask, HistoryParallaxRejectionMask);
-		}
-		#endif
+		ProcessPrevAtomicTexturesSamples(
+			HistoryUseCountSamples,
+			HistoryClosestDeviceZSamples,
+			DeviceZ,
+			PrevScreenPos,
+			/* out */ ParallaxRejectionMask,
+			/* out */ bIsValidHoleFillingPixelVelocity,
+			/* out */ EncodedHoleFillingVelocity);
 
 		// Finish off the ParallaxRejectionMask
 		ParallaxRejectionMask = select(bIsOffScreen, tsr_half(0.0), ParallaxRejectionMask);
@@ -561,84 +526,8 @@ void MainCS(
 				//HoleFilledVelocityMask = tsr_half(HoleFillingOutOfRange) * tsr_half(saturate(MaxHollFillPixelVelocity - VelocityDelta));
 			}
 		}
-		
-		// Update subpixel detail
-		#if CONFIG_SUBPIXEL_PARALLAX_FACTOR
-		{
-			FinalSubpixelDetails = PrevSubpixelDetails;
-
-			{
-				const tsr_subpixel_details ParallaxFactorMask = tsr_subpixel_details(
-					(SUB_PIXEL_PARALLAX_FACTOR_BIT_MASK << (SUB_PIXEL_PARALLAX_FACTOR_BIT_OFFSET + 0 * SUB_PIXEL_BIT_COUNT)) | 
-					(SUB_PIXEL_PARALLAX_FACTOR_BIT_MASK << (SUB_PIXEL_PARALLAX_FACTOR_BIT_OFFSET + 1 * SUB_PIXEL_BIT_COUNT)) | 
-					(SUB_PIXEL_PARALLAX_FACTOR_BIT_MASK << (SUB_PIXEL_PARALLAX_FACTOR_BIT_OFFSET + 2 * SUB_PIXEL_BIT_COUNT)) | 
-					(SUB_PIXEL_PARALLAX_FACTOR_BIT_MASK << (SUB_PIXEL_PARALLAX_FACTOR_BIT_OFFSET + 3 * SUB_PIXEL_BIT_COUNT)) | 
-					0x0);
-			
-				// Reset parallax factor
-				FinalSubpixelDetails = select(bIsOffscreenOrParallaxRejected, FinalSubpixelDetails & ~ParallaxFactorMask, FinalSubpixelDetails);
-			}
-		
-			// TODO(TSR): Move SubpixelId CPU
-			bool2 bKOSign = InputJitter < 0.0;
-			uint SubpixelId = dot(uint2(bKOSign), uint2(1, SUB_PIXEL_GRID_SIZE));
-
-			tsr_subpixel_payload SubpixelShift = tsr_subpixel_payload(SubpixelId) * tsr_subpixel_payload(SUB_PIXEL_BIT_COUNT);
-
-			tsr_subpixel_payload SubpixelPayload = (ParallaxFactorBits << tsr_subpixel_payload(SUB_PIXEL_PARALLAX_FACTOR_BIT_OFFSET));
-
-			tsr_subpixel_details UpdatedFinalSubpixelDetails = (FinalSubpixelDetails & (~(tsr_subpixel_payload(SUB_PIXEL_BIT_MASK) << SubpixelShift))) | (SubpixelPayload << SubpixelShift);
-
-			FinalSubpixelDetails = UpdatedFinalSubpixelDetails;
-		}
-		#endif
 	}
 	
-	#if CONFIG_GRAND_REPROJECTION
-	float2 EncodedGrandScreenVelocity;
-	{
-		tsr_ushort2 InputPixelPos = (
-			Map8x8Tile2x2Lane(GroupThreadIndex) +
-			(tsr_ushort2(InputInfo_ViewportMin) + tsr_ushort2(GroupId) * tsr_ushort2(TILE_SIZE, TILE_SIZE)));
-			
-		float2 ScreenPos = ApplyScreenTransform(float2(InputPixelPos), InputPixelPosToScreenPos);
-		float2 ScreenVelocity = DecodeVelocityFromTexture(float4(EncodedScreenVelocity, 0.0, 0.0)).xy;
-
-		float2 PrevScreenPos = ScreenPos - ScreenVelocity;
-		float2 KernelHistoryBufferUV = ApplyScreenTransform(PrevScreenPos, ScreenPosToPrevHistoryGuideBufferUV) + CurrentToPrevInputJitter * PrevGuideInfo_ExtentInverse;
-		//float2 KernelHistoryBufferUV = ApplyScreenTransform(PrevScreenPos, ScreenPosToPrevHistoryGuideBufferUV);
-		
-		FBilinearSampleInfos BilinearInter = GetBilinearSampleLevelInfos(KernelHistoryBufferUV, InputInfo_Extent, InputInfo_ExtentInverse);
-		
-		float2 PrevVelocity = 0.0;
-		float PrevVelocityWeight = 0.0;
-
-		UNROLL_N(4)
-		for (uint i = 0; i < 4; i++)
-		{
-			float BilinearWeight = GetSampleWeight(BilinearInter, i);
-			uint2 PixelPos = GetSamplePixelCoord(BilinearInter, i);
-
-			bool bValidPixelPos = all(and(PixelPos >= InputInfo_ViewportMin, PixelPos < InputInfo_ViewportMax));
-			BilinearWeight = bValidPixelPos ? BilinearWeight : 0.0;
-
-			float2 EncodedPrevVelocity = PrevVelocityTexture[PixelPos];
-
-			float2 PrevScreenVelocity = DecodeVelocityFromTexture(float4(EncodedPrevVelocity, 0.0, 0.0)).xy;
-
-			PrevVelocity += BilinearWeight * PrevScreenVelocity;
-			PrevVelocityWeight += BilinearWeight;
-		}
-
-		PrevVelocity *= SafeRcp(PrevVelocityWeight);
-
-		EncodedGrandScreenVelocity = EncodeVelocityToTexture(float3(ScreenVelocity + PrevVelocity, 0.0)).xy;
-
-		//Debug.r = abs(ScreenVelocity.x) * ScreenVelocityToInputPixelVelocity.x;
-		Debug.g = abs(PrevVelocity.x) * ScreenVelocityToInputPixelVelocity.x;
-	}
-	#endif // CONFIG_GRAND_REPROJECTION
-
 	// Output full res
 	ISOLATE
 	{
@@ -667,14 +556,8 @@ void MainCS(
 			ReprojectedHistoryMoireOutput[OutputPixelPos] = ReprojectedHistoryMoireColor;
 		#endif
 		HoleFilledVelocityOutput[OutputPixelPos] = EncodedScreenVelocity;
-		#if CONFIG_GRAND_REPROJECTION
-			GrandHoleFilledVelocityOutput[OutputPixelPos] = EncodedGrandScreenVelocity;
-		#endif
 		HoleFilledVelocityMaskOutput[InputPixelPos] = tsr_half2(0.0, RejectHighFrequency);
 		ParallaxRejectionMaskOutput[OutputPixelPos] = ParallaxRejectionMask;
-		#if CONFIG_SUBPIXEL_PARALLAX_FACTOR
-			HistorySubpixelDetailsOutput[OutputPixelPos] = FinalSubpixelDetails;
-		#endif
 
 		#if DEBUG_OUTPUT
 		{
diff --git a/Engine/Shaders/Private/TemporalSuperResolution/TSRDilateVelocity.usf b/Engine/Shaders/Private/TemporalSuperResolution/TSRDilateVelocity.usf
--- a/Engine/Shaders/Private/TemporalSuperResolution/TSRDilateVelocity.usf
+++ b/Engine/Shaders/Private/TemporalSuperResolution/TSRDilateVelocity.usf
@@ -150,8 +150,6 @@ void MainCS(
 	
 	float ClosestDeviceZ = DeviceZ;
 	int2 VelocityPixelOffset = int2(0, 0);
-	float ParallaxFactorForHistory = 1.0;
-	float ParallaxFactorForHistoryRejection = 1.0;
 	#if CONFIG_AA_CROSS
 	ISOLATE
 	{
@@ -176,9 +174,6 @@ void MainCS(
 
 		uint PackedDepthOffset = PackDepthAndOffset(/* Offset = */ int2(0, 0), CenterSubpixelDeviceZ);
 
-		float FurthestPairDeviceZ = DeviceZ;
-		float FurthestDeviceZ = DeviceZ;
-
 		#if CONFIG_DEPTH_EDGE_DETECTION
 		float DeviceZError;
 		{
@@ -252,14 +247,8 @@ void MainCS(
 				PackedDepthOffset = max(PackedDepthOffset, PackDepthAndOffset(-Offset, CombinedSampleSubpixelDeviceZ1));
 			}
 			#endif
-
-			FurthestPairDeviceZ = min(FurthestPairDeviceZ, max(SampleDeviceZ0, SampleDeviceZ1));
-			FurthestDeviceZ = min(FurthestDeviceZ, min(SampleDeviceZ0, SampleDeviceZ1));
 		}
 
-		ParallaxFactorForHistory = clamp(DeviceZ / FurthestPairDeviceZ, 0.0, MAX_PARALLAX_FACTOR);
-		ParallaxFactorForHistoryRejection = clamp(DeviceZ / FurthestDeviceZ, 0.0, MAX_PARALLAX_FACTOR);
-
 		UnpackDepthAndOffset(PackedDepthOffset, /* out */ VelocityPixelOffset, /* out */ ClosestDeviceZ);
 	}
 	#endif
@@ -371,18 +360,6 @@ void MainCS(
 		{
 			CompressedSubpixelDepth = 0;
 		}
-
-		if (ParallaxFactorForHistory > 1.005 && (!bIsMovingPixel || bIncludeDynamicDepthInSubpixelDetails))
-		{
-			FTSRSubpixelDepth FinalSubpixelDepth;
-			FinalSubpixelDepth.DeviceZ = DeviceZ;
-			FinalSubpixelDepth.Offset = 0.0;
-			FinalSubpixelDepth.RemainingLifetime = SubpixelDepthLifetime;
-
-			CompressedSubpixelDepth = CompressSubpixelDepth(FinalSubpixelDepth);
-
-			ParallaxFactorForHistory = 0.0;
-		}
 	}
 	#endif
 
@@ -447,19 +424,13 @@ void MainCS(
 		#if CONFIG_SUBPIXEL_DEPTH
 		{
 			SubpixelDepthOutput[OutputPixelPos] = CompressedSubpixelDepth;
-			R8Output[uint3(OutputPixelPos, 0)] = 0; // TODO save some memory bandwidth.
-		}
-		#else
-		{
-			// TODO save some memory bandwidth.
-			R8Output[uint3(OutputPixelPos, 0)] = CompressParallaxFactor(tsr_half(ParallaxFactorForHistory)) | (CompressParallaxFactor(tsr_half(ParallaxFactorForHistoryRejection)) << 4);
 		}
 		#endif
 
 		BRANCH
 		if (bOutputIsMovingTexture)
 		{
-			R8Output[uint3(OutputPixelPos, 1)] = uint(round(IsMovingMask * 255.0));
+			R8Output[uint3(OutputPixelPos, 0)] = uint(round(IsMovingMask * 255.0));
 		}
 
 		// Output motion blur velocity flatten
diff --git a/Engine/Shaders/Private/TemporalSuperResolution/TSRUpdateHistory.usf b/Engine/Shaders/Private/TemporalSuperResolution/TSRUpdateHistory.usf
--- a/Engine/Shaders/Private/TemporalSuperResolution/TSRUpdateHistory.usf
+++ b/Engine/Shaders/Private/TemporalSuperResolution/TSRUpdateHistory.usf
@@ -54,24 +54,10 @@
 	#define CONFIG_TRANSLUCENCY_REJECTION 1
 #endif
 
-#if DIM_GRAND_REPROJECTION
-	#define CONFIG_GRAND_REPROJECTION 1
-#else
-	#define CONFIG_GRAND_REPROJECTION 0
-#endif
-
 #define CONFIG_RESPONSIVE_STENCIL 1
 #define CONFIG_TAA_RESPONSIVE_BEHAVIOR 1
 #define CONFIG_LOAD_TRANSLUCENCY_REJECTION 0
 
-#define CONFIG_PREMULTIPLY_HISTORY 0
-
-#if CONFIG_COMPILE_FP16 && !CONFIG_SCENE_COLOR_ALPHA
-	#define CONFIG_MANUAL_LDS_SPILL 0 // TODO
-#else
-	#define CONFIG_MANUAL_LDS_SPILL 0
-#endif
-
 #if CONFIG_COMPILE_FP16
 	// Take advantage of RDNA's v_pk_*_{uif}16 instructions
 	#define CONFIG_ENABLE_DUAL_PIXEL_VECTORIZATION 1
@@ -87,7 +73,7 @@
 	#define CONFIG_INTERPOLATE_VELOCITY 0
 #endif
 
-#if CONFIG_GRAND_REPROJECTION == 2 || CONFIG_ACCUMULATE_TRANSLUCENCY_SEPARATELY || CONFIG_PREMULTIPLY_HISTORY
+#if CONFIG_ACCUMULATE_TRANSLUCENCY_SEPARATELY
 	#define CONFIG_OUTPUT_HIGH_FREQUENCY 1
 #else
 	#define CONFIG_OUTPUT_HIGH_FREQUENCY 0
@@ -98,13 +84,7 @@
 
 #define TILE_SIZE 8
 
-#if CONFIG_GRAND_REPROJECTION == 1
-	#define CONFIG_METADATA_CHANNELS 2
-#elif CONFIG_PREMULTIPLY_HISTORY
-	#define CONFIG_METADATA_CHANNELS 2
-#else
-	#define CONFIG_METADATA_CHANNELS 1
-#endif
+#define CONFIG_METADATA_CHANNELS 1
 
 #if CONFIG_METADATA_CHANNELS == 2
 	#define tsr_halfM tsr_half2
@@ -126,22 +106,17 @@ Texture2D<tsr_half4> InputSceneTranslucencyTexture;
 Texture2D<tsr_half2> HistoryRejectionTexture;
 
 Texture2D<float2> DilatedVelocityTexture;
-Texture2D<uint> ParallaxFactorTexture;
 Texture2D<tsr_half> ParallaxRejectionMaskTexture;
 Texture2D<tsr_ushort> AntiAliasingTexture;
 Texture2D<tsr_half> NoiseFilteringTexture;
 Texture2D<tsr_half2> HoleFilledVelocityMaskTexture;
 
-Texture2D<float2> GrandDilatedVelocityTexture;
-Texture2D<tsr_halfC> GrandPrevColorTexture;
-
 int2   TranslucencyPixelPosMin;
 int2   TranslucencyPixelPosMax;
 
 FScreenTransform HistoryPixelPosToScreenPos;
 FScreenTransform HistoryPixelPosToInputPPCo;
 FScreenTransform HistoryPixelPosToTranslucencyPPCo;
-FScreenTransform ScreenPosToGrandPrevHistoryBufferUV;
 float3 HistoryQuantizationError;
 float HistorySampleCount;
 float HistoryHisteresis;
@@ -149,7 +124,6 @@ float WeightClampingRejection;
 float WeightClampingPixelSpeedAmplitude;
 float InvWeightClampingPixelSpeed;
 float InputToHistoryFactor;
-float GrandPrevPreExposureCorrection;
 uint ResponsiveStencilMask;
 uint bGenerateOutputMip1;
 uint bGenerateOutputMip2;
@@ -159,25 +133,16 @@ Texture2D<tsr_halfC> PrevHistory_HighFrequency;
 Texture2D<tsr_halfM> PrevHistory_Metadata;
 Texture2D<tsr_half3> PrevHistory_Translucency;
 Texture2D<tsr_half>  PrevHistory_TranslucencyAlpha;
-Texture2D<tsr_subpixel_details> PrevHistory_SubpixelDetails;
-
-Texture2D<tsr_halfC> PrevHistory_PrevHighFrequency;
-Texture2D<tsr_halfC> PrevHistory_PrevHighFrequencyResultant;
-Texture2D<tsr_halfC> PrevHistory_HighFrequencyOverblur;
 
 RWTexture2D<tsr_halfC> SceneColorOutputMip0;
 RWTexture2D<tsr_halfC> SceneColorOutputMip1;
 
 uint HistoryOutput_ArrayIndices_HighFrequency;
 uint HistoryOutput_ArrayIndices_Translucency;
-uint HistoryOutput_ArrayIndices_PrevHighFrequency;
-uint HistoryOutput_ArrayIndices_PrevHighFrequencyResultant;
-uint HistoryOutput_ArrayIndices_HighFrequencyOverblur;
 
 RWTexture2DArray<tsr_halfC> HistoryOutput_ColorArray;
 RWTexture2D<tsr_halfM> HistoryOutput_Metadata;
 RWTexture2D<tsr_half>  HistoryOutput_TranslucencyAlpha;
-RWTexture2D<tsr_subpixel_details> HistoryOutput_SubpixelDetails;
 
 
 //------------------------------------------------------- LDS
@@ -510,9 +475,6 @@ void MainCS(
 
 	// Fetch reprojection-related information.
 	float2x2 PrevScreenPos = ScreenPos;
-	#if CONFIG_GRAND_REPROJECTION
-		float2x2 GrandPrevScreenPos = ScreenPos;
-	#endif
 	tsr_half2 LowFrequencyRejection = tsr_half(1.0).xx;
 	tsr_half2 LowFrequencyClamp = tsr_half(1.0).xx;
 	tsr_half2 OutputPixelVelocity = tsr_half(0.0).xx;
@@ -533,17 +495,8 @@ void MainCS(
 		#else
 			float2 RawEncodedVelocity[DPV_PIXEL_PER_LANE];
 		#endif
-		
-		#if CONFIG_GRAND_REPROJECTION
-			#if CONFIG_INTERPOLATE_VELOCITY
-				float2 RawEncodedGrandVelocityNeighborhood[4][DPV_PIXEL_PER_LANE];
-			#else
-				float2 RawEncodedGrandVelocity[DPV_PIXEL_PER_LANE];
-			#endif
-		#endif
 
 		tsr_half RawParallaxRejectionMask[DPV_PIXEL_PER_LANE];
-		uint RawParallaxFactorBits[DPV_PIXEL_PER_LANE];
 		tsr_half2 RawLowFrequencyRejection[DPV_PIXEL_PER_LANE];
 		
 		uint RawSceneStencilRef[DPV_PIXEL_PER_LANE];
@@ -565,13 +518,9 @@ void MainCS(
 				
 				#if !CONFIG_INTERPOLATE_VELOCITY
 					RawEncodedVelocity[PixelId]   = DilatedVelocityTexture[LocalInputPixelPos];
-					#if CONFIG_GRAND_REPROJECTION
-						RawEncodedGrandVelocity[PixelId] = GrandDilatedVelocityTexture[LocalInputPixelPos];
-					#endif
 				#endif
 
 				RawParallaxRejectionMask[PixelId] = ParallaxRejectionMaskTexture[LocalInputPixelPos];
-				RawParallaxFactorBits[PixelId]    = ParallaxFactorTexture[LocalInputPixelPos];
 				RawSceneStencilRef[PixelId]       = InputSceneStencilTexture.Load(int3(LocalInputPixelPos, 0)) STENCIL_COMPONENT_SWIZZLE;
 				
 				RawLowFrequencyRejection[PixelId] = HistoryRejectionTexture[LocalInputPixelPos];
@@ -619,13 +568,6 @@ void MainCS(
 				RawEncodedVelocityNeighborhood[1][PixelId2] = DilatedVelocityTexture[tsr_short2(LocalInputPixelPos1.x, LocalInputPixelPos0.y)];
 				RawEncodedVelocityNeighborhood[2][PixelId2] = DilatedVelocityTexture[tsr_short2(LocalInputPixelPos0.x, LocalInputPixelPos1.y)];
 				RawEncodedVelocityNeighborhood[3][PixelId2] = DilatedVelocityTexture[tsr_short2(LocalInputPixelPos1.x, LocalInputPixelPos1.y)];
-				
-				#if CONFIG_GRAND_REPROJECTION
-					RawEncodedGrandVelocityNeighborhood[0][PixelId2] = GrandDilatedVelocityTexture[tsr_short2(LocalInputPixelPos0.x, LocalInputPixelPos0.y)];
-					RawEncodedGrandVelocityNeighborhood[1][PixelId2] = GrandDilatedVelocityTexture[tsr_short2(LocalInputPixelPos1.x, LocalInputPixelPos0.y)];
-					RawEncodedGrandVelocityNeighborhood[2][PixelId2] = GrandDilatedVelocityTexture[tsr_short2(LocalInputPixelPos0.x, LocalInputPixelPos1.y)];
-					RawEncodedGrandVelocityNeighborhood[3][PixelId2] = GrandDilatedVelocityTexture[tsr_short2(LocalInputPixelPos1.x, LocalInputPixelPos1.y)];
-				#endif
 			}
 			#endif
 		}
@@ -669,51 +611,11 @@ void MainCS(
 					ScreenVelocity[0][PixelId] = LocalScreenVelocity.x;
 					ScreenVelocity[1][PixelId] = LocalScreenVelocity.y;
 				}
-				#if CONFIG_GRAND_REPROJECTION
-				for (uint PixelId = 0; PixelId < DPV_PIXEL_PER_LANE; PixelId++)
-				{
-					float2 V00 = DecodeVelocityFromTexture(float4(RawEncodedGrandVelocityNeighborhood[0][PixelId], 0.0, 0.0)).xy;
-					float2 V10 = DecodeVelocityFromTexture(float4(RawEncodedGrandVelocityNeighborhood[1][PixelId], 0.0, 0.0)).xy;
-					float2 V01 = DecodeVelocityFromTexture(float4(RawEncodedGrandVelocityNeighborhood[2][PixelId], 0.0, 0.0)).xy;
-					float2 V11 = DecodeVelocityFromTexture(float4(RawEncodedGrandVelocityNeighborhood[3][PixelId], 0.0, 0.0)).xy;
-
-					tsr_half2 BilinearInterp = abs(dpv_access_pixel(dInputKO, PixelId));
-					
-					float Bilinear00 = saturate(1.0 - BilinearInterp.x) * saturate(1.0 - BilinearInterp.y);
-					float Bilinear10 = saturate(0.0 + BilinearInterp.x) * saturate(1.0 - BilinearInterp.y);
-					float Bilinear01 = saturate(1.0 - BilinearInterp.x) * saturate(0.0 + BilinearInterp.y);
-					float Bilinear11 = saturate(0.0 + BilinearInterp.x) * saturate(0.0 + BilinearInterp.y);
-					
-					float Bilateral00 = Bilinear00;
-					float Bilateral10 = Bilinear10 * ComputeScreenVelocityBilateralWeight(V10 - V00);
-					float Bilateral01 = Bilinear01 * ComputeScreenVelocityBilateralWeight(V01 - V00);
-					float Bilateral11 = Bilinear11 * ComputeScreenVelocityBilateralWeight(V11 - V00);
-
-					float NormalizeBilateral = rcp(Bilateral00 + Bilateral10 + Bilateral01 + Bilateral11);
-
-					float2 LocalScreenVelocity = (V00 * Bilateral00 + V10 * Bilateral10 + V01 * Bilateral01 + V11 * Bilateral11) * NormalizeBilateral;
-					
-					float2 LocalPrevScreenPos = dpv_access_pixel(ScreenPos, PixelId) - LocalScreenVelocity;
-
-					GrandPrevScreenPos[0][PixelId] = LocalPrevScreenPos.x;
-					GrandPrevScreenPos[1][PixelId] = LocalPrevScreenPos.y;
-				}
-				#endif // CONFIG_GRAND_REPROJECTION
 				#else // !CONFIG_INTERPOLATE_VELOCITY
 				{
 					ScreenVelocity = dpv_interleave_registers(
 						DecodeVelocityFromTexture(float4(RawEncodedVelocity[0], 0.0, 0.0)).xy,
 						DecodeVelocityFromTexture(float4(RawEncodedVelocity[1], 0.0, 0.0)).xy);
-						
-					#if CONFIG_GRAND_REPROJECTION
-					{
-						float2x2 GrandScreenVelocity = dpv_interleave_registers(
-							DecodeVelocityFromTexture(float4(RawEncodedGrandVelocity[0], 0.0, 0.0)).xy,
-							DecodeVelocityFromTexture(float4(RawEncodedGrandVelocity[1], 0.0, 0.0)).xy);
-					
-						GrandPrevScreenPos = ScreenPos - GrandScreenVelocity;
-					}
-					#endif
 				}
 				#endif
 
@@ -781,17 +683,6 @@ void MainCS(
 
 	tsr_half4x2 PrevHistoryTranslucency;
 	
-	#if CONFIG_GRAND_REPROJECTION
-		tsr_halfCx2 GrandPrevToPrev;
-		#if CONFIG_GRAND_REPROJECTION == 1
-			tsr_half2 GrandPrevToPrevResultant;
-		#elif CONFIG_GRAND_REPROJECTION == 2
-			tsr_halfCx2 GrandPrevToPrevResultant;
-			tsr_halfCx2 PrevHighFrequencyOverblur;
-		#endif
-		tsr_halfCx2 GrandReprojection;
-	#endif
-
 	PLATFORM_SPECIFIC_ISOLATE
 	{
 		// Reproject history
@@ -801,15 +692,6 @@ void MainCS(
 
 		tsr_halfCx2 HighFrequencyMin = INFINITE_FLOAT;
 		tsr_halfCx2 HighFrequencyMax = -INFINITE_FLOAT;
-		
-		#if CONFIG_GRAND_REPROJECTION
-			tsr_halfCx2 AccumulateGrandPrevToPrev = tsr_half(0);
-			#if CONFIG_GRAND_REPROJECTION == 2
-				tsr_halfCx2 AccumulateGrandPrevToPrevResultant = tsr_half(0);
-				tsr_halfCx2 AccumulatePrevHighFrequencyOverblur = tsr_half(0);
-			#endif
-			tsr_halfCx2 AccumulateGrandReprojection = tsr_half(0);
-		#endif
 
 		// Sample raw history
 		{
@@ -821,17 +703,6 @@ void MainCS(
 				dpv_lo(PrevHistoryBufferUV), PrevHistoryInfo_Extent, PrevHistoryInfo_ExtentInverse);
 			FCatmullRomSamples Samples1 = GetBicubic2DCatmullRomSamples(
 				dpv_hi(PrevHistoryBufferUV), PrevHistoryInfo_Extent, PrevHistoryInfo_ExtentInverse);
-			
-			#if CONFIG_GRAND_REPROJECTION
-				float2x2 GrandPrevHistoryBufferUV = ApplyScreenTransform(GrandPrevScreenPos, ScreenPosToGrandPrevHistoryBufferUV);
-				GrandPrevHistoryBufferUV[0] = fastClamp(GrandPrevHistoryBufferUV[0], PrevHistoryInfo_UVViewportBilinearMin[0], PrevHistoryInfo_UVViewportBilinearMax[0]);
-				GrandPrevHistoryBufferUV[1] = fastClamp(GrandPrevHistoryBufferUV[1], PrevHistoryInfo_UVViewportBilinearMin[1], PrevHistoryInfo_UVViewportBilinearMax[1]);
-
-				FCatmullRomSamples GrandSamples0 = GetBicubic2DCatmullRomSamples(
-					dpv_lo(GrandPrevHistoryBufferUV), PrevHistoryInfo_Extent, PrevHistoryInfo_ExtentInverse);
-				FCatmullRomSamples GrandSamples1 = GetBicubic2DCatmullRomSamples(
-					dpv_hi(GrandPrevHistoryBufferUV), PrevHistoryInfo_Extent, PrevHistoryInfo_ExtentInverse);
-			#endif
 
 			tsr_halfC RawHighFrequency[BICUBIC_CATMULL_ROM_SAMPLES][DPV_PIXEL_PER_LANE];
 			tsr_halfM RawMetadata[BICUBIC_CATMULL_ROM_SAMPLES][DPV_PIXEL_PER_LANE];
@@ -839,17 +710,6 @@ void MainCS(
 			tsr_half  RawTranslucencyAlpha[BICUBIC_CATMULL_ROM_SAMPLES][DPV_PIXEL_PER_LANE];
 
 			tsr_half2 RawKernelWeight[BICUBIC_CATMULL_ROM_SAMPLES];
-			
-			#if CONFIG_GRAND_REPROJECTION
-				tsr_halfC RawGrandPrevToPrev[BICUBIC_CATMULL_ROM_SAMPLES][DPV_PIXEL_PER_LANE];
-				#if CONFIG_GRAND_REPROJECTION == 2
-					tsr_halfC RawGrandPrevToPrevResultant[BICUBIC_CATMULL_ROM_SAMPLES][DPV_PIXEL_PER_LANE];
-					tsr_halfC RawPrevHighFrequencyOverblur[BICUBIC_CATMULL_ROM_SAMPLES][DPV_PIXEL_PER_LANE];
-				#endif
-				tsr_halfC RawGrandReprojection[BICUBIC_CATMULL_ROM_SAMPLES][DPV_PIXEL_PER_LANE];
-
-				tsr_half2 RawGrandKernelWeight[BICUBIC_CATMULL_ROM_SAMPLES];
-			#endif
 
 			// Issues texture fetches.
 			PLATFORM_SPECIFIC_ISOLATE
@@ -863,10 +723,6 @@ void MainCS(
 					tsr_half2 KernelWeight = dpv_force_interleave_registers(tsr_half(Samples0.Weight[i]), tsr_half(Samples1.Weight[i]));
 
 					RawKernelWeight[i] = KernelWeight;
-					
-					#if CONFIG_GRAND_REPROJECTION
-						RawGrandKernelWeight[i] = dpv_force_interleave_registers(tsr_half(GrandSamples0.Weight[i]), tsr_half(GrandSamples1.Weight[i]));
-					#endif
 
 					UNROLL_N(2)
 					for (uint PixelId = 0; PixelId < DPV_PIXEL_PER_LANE; PixelId++)
@@ -882,23 +738,6 @@ void MainCS(
 							RawTranslucencyAlpha[i][PixelId] = PrevHistory_TranslucencyAlpha.SampleLevel(GlobalBilinearClampedSampler, SampleUV, 0);
 						#endif
 						
-						#if CONFIG_GRAND_REPROJECTION
-							RawGrandPrevToPrev[i][PixelId] = BilinearSampleColorHistory(PrevHistory_PrevHighFrequency, SampleUV);
-							#if CONFIG_GRAND_REPROJECTION == 2
-								RawGrandPrevToPrevResultant[i][PixelId] = BilinearSampleColorHistory(PrevHistory_PrevHighFrequencyResultant, SampleUV);
-								RawPrevHighFrequencyOverblur[i][PixelId] = BilinearSampleColorHistory(PrevHistory_HighFrequencyOverblur, SampleUV);
-							#endif
-
-							{
-								float2 GrandSampleUV0 = fastClamp(GrandSamples0.UV[i], PrevHistoryInfo_UVViewportBilinearMin, PrevHistoryInfo_UVViewportBilinearMax);
-								float2 GrandSampleUV1 = fastClamp(GrandSamples1.UV[i], PrevHistoryInfo_UVViewportBilinearMin, PrevHistoryInfo_UVViewportBilinearMax);
-					
-								float2 GrandSampleUV = PixelId == 0 ? GrandSampleUV0 : GrandSampleUV1;
-						
-								RawGrandReprojection[i][PixelId] = BilinearSampleColorHistory(GrandPrevColorTexture, GrandSampleUV);
-							
-							}
-						#endif
 					}
 
 				} // for (uint i = 0; i < BICUBIC_CATMULL_ROM_SAMPLES; i++)
@@ -916,15 +755,6 @@ void MainCS(
 					tsr_half4x2 SampleTranslucency = dpv_interleave_registers(
 						tsr_half4(RawTranslucencyColor[i][0], RawTranslucencyAlpha[i][0]),
 						tsr_half4(RawTranslucencyColor[i][1], RawTranslucencyAlpha[i][1]));
-					
-					#if CONFIG_GRAND_REPROJECTION
-						tsr_halfCx2 SampleGrandPrevious = dpv_interleave_registers_array(RawGrandPrevToPrev[i]);
-						#if CONFIG_GRAND_REPROJECTION == 2
-							tsr_halfCx2 SampleGrandPreviousResultant = dpv_interleave_registers_array(RawGrandPrevToPrevResultant[i]);
-							tsr_halfCx2 SamplePrevHighFrequencyOverblur = dpv_interleave_registers_array(RawPrevHighFrequencyOverblur[i]) - SampleHighFrequency;
-						#endif
-						tsr_halfCx2 SampleGrandReprojection = dpv_interleave_registers_array(RawGrandReprojection[i]);
-					#endif
 
 					tsr_half2 KernelWeight = RawKernelWeight[i];
 
@@ -937,15 +767,6 @@ void MainCS(
 					#if CONFIG_ACCUMULATE_TRANSLUCENCY_SEPARATELY
 						AccumulateTranslucency  += dpv_scale(SampleTranslucency, KernelWeight);
 					#endif
-
-					#if CONFIG_GRAND_REPROJECTION
-						AccumulateGrandPrevToPrev  += dpv_scale(SampleGrandPrevious, KernelWeight);
-						#if CONFIG_GRAND_REPROJECTION == 2
-							AccumulateGrandPrevToPrevResultant  += dpv_scale(SampleGrandPreviousResultant, KernelWeight);
-							AccumulatePrevHighFrequencyOverblur  += dpv_scale(SamplePrevHighFrequencyOverblur, KernelWeight);
-						#endif
-						AccumulateGrandReprojection  += dpv_scale(SampleGrandReprojection, RawGrandKernelWeight[i]);
-					#endif
 				}
 				
 				tsr_half2 FinalMultiplier = dpv_force_interleave_registers(tsr_half(Samples0.FinalMultiplier), tsr_half(Samples1.FinalMultiplier));
@@ -953,17 +774,6 @@ void MainCS(
 				AccumulateHighFrequency = dpv_scale(AccumulateHighFrequency, FinalMultiplier);
 				AccumulateMetadata  = dpv_scale(AccumulateMetadata, FinalMultiplier);
 				AccumulateTranslucency  = dpv_scale(AccumulateTranslucency, FinalMultiplier);
-
-				#if CONFIG_GRAND_REPROJECTION
-					tsr_half2 GrandFinalMultiplier = dpv_interleave_registers(tsr_half(GrandSamples0.FinalMultiplier), tsr_half(GrandSamples1.FinalMultiplier));
-
-					AccumulateGrandPrevToPrev          = dpv_scale(AccumulateGrandPrevToPrev, FinalMultiplier);
-					#if CONFIG_GRAND_REPROJECTION == 2
-						AccumulateGrandPrevToPrevResultant = dpv_scale(AccumulateGrandPrevToPrevResultant, FinalMultiplier);
-						AccumulatePrevHighFrequencyOverblur = dpv_scale(AccumulatePrevHighFrequencyOverblur, FinalMultiplier);
-					#endif
-					AccumulateGrandReprojection        = dpv_scale(AccumulateGrandReprojection, GrandFinalMultiplier);
-				#endif
 			}
 			
 			#if CONFIG_SCENE_COLOR_ALPHA
@@ -988,17 +798,6 @@ void MainCS(
 			#endif
 
 			PrevHistoryTranslucency = AccumulateTranslucency;
-
-			#if CONFIG_GRAND_REPROJECTION
-				GrandPrevToPrev          = AccumulateGrandPrevToPrev;
-				#if CONFIG_GRAND_REPROJECTION == 1
-					GrandPrevToPrevResultant = AccumulateMetadata[1];
-				#elif CONFIG_GRAND_REPROJECTION == 2
-					GrandPrevToPrevResultant = AccumulateGrandPrevToPrevResultant;
-					PrevHighFrequencyOverblur = AccumulatePrevHighFrequencyOverblur;
-				#endif
-				GrandReprojection        = AccumulateGrandReprojection;
-			#endif
 		}
 
 		// Correct history
@@ -1012,24 +811,6 @@ void MainCS(
 			CorrectExposure(PrevHighFrequencyColor[0]);
 			CorrectExposure(PrevHighFrequencyColor[1]);
 			CorrectExposure(PrevHighFrequencyColor[2]);
-			
-			#if CONFIG_GRAND_REPROJECTION
-				CorrectExposure(GrandPrevToPrev[0]);
-				CorrectExposure(GrandPrevToPrev[1]);
-				CorrectExposure(GrandPrevToPrev[2]);
-				#if CONFIG_GRAND_REPROJECTION == 2
-					CorrectExposure(GrandPrevToPrevResultant[0]);
-					CorrectExposure(GrandPrevToPrevResultant[1]);
-					CorrectExposure(GrandPrevToPrevResultant[2]);
-					CorrectExposure(PrevHighFrequencyOverblur[0]);
-					CorrectExposure(PrevHighFrequencyOverblur[1]);
-					CorrectExposure(PrevHighFrequencyOverblur[2]);
-				#endif
-				
-				GrandReprojection[0] *= tsr_half(GrandPrevPreExposureCorrection);
-				GrandReprojection[1] *= tsr_half(GrandPrevPreExposureCorrection);
-				GrandReprojection[2] *= tsr_half(GrandPrevPreExposureCorrection);
-			#endif
 
 			// Fixes negative validity that can happen due to negative lobes of Catrom.
 			PrevHistoryValidity = max(PrevHistoryValidity, tsr_half(0.0).xx);
@@ -1290,104 +1071,9 @@ void MainCS(
 		#endif
 	}
 		
-	// TODO: CONFIG_MANUAL_LDS_SPILL
-
-
-	// Spills to LDS to make room in VGPR for history sampling.
-	#if CONFIG_MANUAL_LDS_SPILL
-	PLATFORM_SPECIFIC_ISOLATE
-	{
-		uint LocalGroupThreadIndex = GetGroupThreadIndex(GroupThreadIndex, GroupWaveIndex);
-	
-		#if CONFIG_ENABLE_DUAL_PIXEL_VECTORIZATION
-			SharedArray0[LocalGroupThreadIndex] = tsr_half4(FilteredInputColor[0], FilteredInputColor[1]);
-			SharedArray1[LocalGroupThreadIndex] = tsr_half4(FilteredInputColor[2], LowFrequencyRejection);
-			SharedArray2[LocalGroupThreadIndex] = tsr_half4(InputMinColor[0], InputMinColor[1]);
-			SharedArray3[LocalGroupThreadIndex] = tsr_half4(InputMinColor[2], InputPixelAlignement);
-			SharedArray4[LocalGroupThreadIndex] = tsr_half4(InputMaxColor[0], InputMaxColor[1]);
-			SharedArray5[LocalGroupThreadIndex] = tsr_half4(InputMaxColor[2], OutputPixelVelocity);
-	
-			#if CONFIG_ACCUMULATE_TRANSLUCENCY_SEPARATELY
-				SharedArray6[LocalGroupThreadIndex] = tsr_half4(FilteredTranslucencyColor[0], FilteredTranslucencyColor[1]);
-				SharedArray7[LocalGroupThreadIndex] = tsr_half4(FilteredTranslucencyColor[2], FilteredTranslucencyColor[3]);
-				SharedArray8[LocalGroupThreadIndex] = tsr_half4(TranslucencyMinColor[0], TranslucencyMinColor[1]);
-				SharedArray9[LocalGroupThreadIndex] = tsr_half4(TranslucencyMinColor[2], TranslucencyMinColor[3]);
-				SharedArray10[LocalGroupThreadIndex] = tsr_half4(TranslucencyMaxColor[0], TranslucencyMaxColor[1]);
-				SharedArray11[LocalGroupThreadIndex] = tsr_half4(TranslucencyMaxColor[2], TranslucencyMaxColor[3]);
-			#endif
-		#else
-			SharedArray0[LocalGroupThreadIndex] = tsr_half4(dpv_lo(FilteredInputColor), dpv_lo(LowFrequencyRejection));
-			SharedArray1[LocalGroupThreadIndex] = tsr_half4(dpv_lo(InputMinColor), dpv_lo(InputPixelAlignement));
-			SharedArray2[LocalGroupThreadIndex] = tsr_half4(dpv_lo(InputMaxColor), dpv_lo(OutputPixelVelocity));
-		#endif
-	}
-	#endif
-		
-	#if CONFIG_MANUAL_LDS_SPILL
-	PLATFORM_SPECIFIC_ISOLATE
-	{
-		uint LocalGroupThreadIndex = GetGroupThreadIndex(GroupThreadIndex, GroupWaveIndex);
-	
-		#if CONFIG_ENABLE_DUAL_PIXEL_VECTORIZATION
-			tsr_half4 RawLDS0 = SharedArray0[LocalGroupThreadIndex];
-			tsr_half4 RawLDS1 = SharedArray1[LocalGroupThreadIndex];
-			tsr_half4 RawLDS2 = SharedArray2[LocalGroupThreadIndex];
-			tsr_half4 RawLDS3 = SharedArray3[LocalGroupThreadIndex];
-			tsr_half4 RawLDS4 = SharedArray4[LocalGroupThreadIndex];
-			tsr_half4 RawLDS5 = SharedArray5[LocalGroupThreadIndex];
-			
-			#if CONFIG_ACCUMULATE_TRANSLUCENCY_SEPARATELY
-				tsr_half4 RawLDS6 = SharedArray6[LocalGroupThreadIndex];
-				tsr_half4 RawLDS7 = SharedArray7[LocalGroupThreadIndex];
-				tsr_half4 RawLDS8 = SharedArray8[LocalGroupThreadIndex];
-				tsr_half4 RawLDS9 = SharedArray9[LocalGroupThreadIndex];
-				tsr_half4 RawLDS10 = SharedArray10[LocalGroupThreadIndex];
-				tsr_half4 RawLDS11 = SharedArray11[LocalGroupThreadIndex];
-			#endif
-
-			FilteredInputColor = tsr_half3x2(RawLDS0.rg, RawLDS0.ba, RawLDS1.rg);
-			InputMinColor = tsr_half3x2(RawLDS2.rg, RawLDS2.ba, RawLDS3.rg);
-			InputMaxColor = tsr_half3x2(RawLDS4.rg, RawLDS4.ba, RawLDS5.rg);
-		
-			LowFrequencyRejection = RawLDS1.ba;
-			InputPixelAlignement = RawLDS3.ba;
-			OutputPixelVelocity = RawLDS5.ba;
-			
-			#if CONFIG_ACCUMULATE_TRANSLUCENCY_SEPARATELY
-				FilteredTranslucencyColor = tsr_half4x2(RawLDS6.rg, RawLDS6.ba, RawLDS7.rg, RawLDS7.ba);
-				TranslucencyMinColor = tsr_half4x2(RawLDS8.rg, RawLDS8.ba, RawLDS9.rg, RawLDS9.ba);
-				TranslucencyMaxColor = tsr_half4x2(RawLDS10.rg, RawLDS10.ba, RawLDS11.rg, RawLDS11.ba);
-			#endif
-		#else
-			tsr_half4 RawLDS0 = SharedArray0[LocalGroupThreadIndex];
-			tsr_half4 RawLDS1 = SharedArray1[LocalGroupThreadIndex];
-			tsr_half4 RawLDS2 = SharedArray2[LocalGroupThreadIndex];
-
-			FilteredInputColor = dpv_interleave_mono_registers(RawLDS0.rgb);
-			InputMinColor = dpv_interleave_mono_registers(RawLDS1.rgb);
-			InputMaxColor = dpv_interleave_mono_registers(RawLDS2.rgb);
-		
-			LowFrequencyRejection = dpv_interleave_mono_registers(RawLDS0.a);
-			InputPixelAlignement = dpv_interleave_mono_registers(RawLDS1.a);
-			OutputPixelVelocity = dpv_interleave_mono_registers(RawLDS2.a);
-		#endif
-	}
-	#endif
-
 	// Contribute current frame input into the predictor for next frame.
 	tsr_halfCx2 FinalHighFrequencyColor;
 	tsr_half2 FinalHistoryValidity;
-	
-	#if CONFIG_GRAND_REPROJECTION
-		tsr_halfCx2 FinalPrevHighFrequencyColor;
-		#if CONFIG_GRAND_REPROJECTION == 1
-			tsr_half2 FinalPrevHighFrequencyResultant;
-		#elif CONFIG_GRAND_REPROJECTION == 2
-			tsr_halfCx2 FinalPrevHighFrequencyColorResultant;
-			tsr_halfCx2 FinalOverblur = tsr_half(0.0);
-		#endif
-		tsr_halfCx2 FinalHighFrequencyColorOverblurCorrection = tsr_half(0.0);
-	#endif
 
 	{
 		tsr_halfCx2 ClampedPrevHighFrequencyColor = fastClamp(PrevHighFrequencyColor, InputMinColor, InputMaxColor);
@@ -1400,7 +1086,7 @@ void MainCS(
 		#endif
 		
 		// Clamp the validity due to motion to maintain better sharpness in history reprojection under motion.
-		#if !CONFIG_GRAND_REPROJECTION
+		#if 1
 		{
 			tsr_half2 MaxValidity = tsr_half(1.0) - tsr_half(WeightClampingPixelSpeedAmplitude) * saturate(OutputPixelVelocity * tsr_half(InvWeightClampingPixelSpeed));
 
@@ -1436,62 +1122,6 @@ void MainCS(
 				
 		// Quantize validity for the 8bit encoding to avoid numerical shift between color and validity.
 		FinalHistoryValidity = ceil(tsr_half(255.0) * OutputValidity) * rcp(tsr_half(255.0));
-
-		#if CONFIG_GRAND_REPROJECTION
-		{
-			tsr_halfCx2 BoxSize = dpv_scale(PrevHistoryMomentMax - PrevHistoryMomentMin, tsr_half(0.5));
-
-			// Compute how much frame N-2 got overblur by reprojecting into N-1 and N versus directly to N.
-			tsr_halfCx2 PrevReprojectionOverblur = GrandReprojection - GrandPrevToPrev;
-
-			// Clamp error imprecision.
-			bool2 bIsOverblurLegal = (
-				and(fastClamp(PrevReprojectionOverblur[0], -BoxSize[0], BoxSize[0]) == PrevReprojectionOverblur[0],
-				and(fastClamp(PrevReprojectionOverblur[1], -BoxSize[1], BoxSize[1]) == PrevReprojectionOverblur[1],
-				    fastClamp(PrevReprojectionOverblur[2], -BoxSize[2], BoxSize[2]) == PrevReprojectionOverblur[2])));
-
-			tsr_half2 PrevReprojectionOverblurScale = select(bIsOverblurLegal, tsr_half(1.0), tsr_half(0.0));
-
-			tsr_half2 HistoryClampResultant = MeasureRejectionFactor(
-				PrevHighFrequencyColor, BlendedPrevHighFrequencyColor,
-				FilteredInputColor, InputMinColor, InputMaxColor,
-				MeasureBackbufferLDRQuantizationError());
-
-			tsr_half2 MeasuredResultant = saturate(CommonWeight * BlendPrevHistory * HistoryClampResultant);
-
-			FinalPrevHighFrequencyColor = PrevHighFrequencyColor;
-
-			#if CONFIG_GRAND_REPROJECTION == 1
-				FinalPrevHighFrequencyResultant = MeasuredResultant;
-
-				const tsr_half CorrectionError = tsr_half(1.0 / sqrt(PI * 0.5));
-				MeasuredResultant = min(MeasuredResultant, saturate(tsr_half(1.0) - tsr_half(HistoryHisteresis) * rcp(FinalHistoryValidity)));
-
-				tsr_half2 CorrectionResultant = GrandPrevToPrevResultant * MeasuredResultant * PrevReprojectionOverblurScale * CorrectionError;
-
-				tsr_halfCx2 Correction = dpv_scale(PrevReprojectionOverblur, CorrectionResultant);
-
-				tsr_halfCx2 CorrectedFinalHighFrequencyColor = FinalHighFrequencyColor + Correction;
-
-				FinalHighFrequencyColor = fastClamp(CorrectedFinalHighFrequencyColor, min(FinalHighFrequencyColor, PrevHistoryMomentMin), max(FinalHighFrequencyColor, PrevHistoryMomentMax));
-
-			#elif CONFIG_GRAND_REPROJECTION == 2
-				// Compute MeasuredResultant of the previous frame
-				tsr_halfCx2 PrevResultant;
-				PrevResultant[0] = fastClamp(GrandPrevToPrevResultant[0] * SafeRcp(GrandPrevToPrev[0]), tsr_half(0.0), tsr_half(1.0));
-				PrevResultant[1] = fastClamp(GrandPrevToPrevResultant[1] * SafeRcp(GrandPrevToPrev[1]), tsr_half(0.0), tsr_half(1.0));
-				PrevResultant[2] = fastClamp(GrandPrevToPrevResultant[2] * SafeRcp(GrandPrevToPrev[2]), tsr_half(0.0), tsr_half(1.0));
-				#if CONFIG_SCENE_COLOR_ALPHA
-					PrevResultant[3] = fastClamp(GrandPrevToPrevResultant[3] * SafeRcp(GrandPrevToPrev[3]), tsr_half(0.0), tsr_half(1.0));
-				#endif
-
-				FinalPrevHighFrequencyColorResultant = dpv_scale(PrevHighFrequencyColor, MeasuredResultant);
-
-				FinalOverblur = dpv_scale(PrevResultant, HistoryClampResultant * TotalRejection) * (PrevHighFrequencyOverblur + dpv_scale(PrevReprojectionOverblur, PrevReprojectionOverblurScale));
-				FinalHighFrequencyColorOverblurCorrection = dpv_scale(FinalOverblur, MeasuredResultant);
-			#endif
-		}
-		#endif
 	}
 	
 	#if CONFIG_ACCUMULATE_TRANSLUCENCY_SEPARATELY
@@ -1571,11 +1201,7 @@ void MainCS(
 	// Compute final output
 	tsr_halfCx2 FinalOutputColor;
 	{
-		#if CONFIG_GRAND_REPROJECTION == 2
-			FinalOutputColor = FinalHighFrequencyColor + FinalHighFrequencyColorOverblurCorrection;
-		#else
-			FinalOutputColor = FinalHighFrequencyColor;
-		#endif
+		FinalOutputColor = FinalHighFrequencyColor;
 
 		#if CONFIG_ACCUMULATE_TRANSLUCENCY_SEPARATELY
 			FinalOutputColor[0] = FinalOutputColor[0] * FinalTranslucencyColor[3] + FinalTranslucencyColor[0];
@@ -1636,13 +1262,6 @@ void MainCS(
 			FinalMetadata[0] = FinalHistoryValidity;
 		#endif
 
-		#if CONFIG_GRAND_REPROJECTION == 1
-			FinalMetadata[1] = FinalPrevHighFrequencyResultant;
-		#elif CONFIG_GRAND_REPROJECTION == 2
-			// Encode the overblur with the high frequency color to avoid negative numbers not encodable in a R11G11B10.
-			FinalOverblur = FinalOverblur + FinalHighFrequencyColor;
-		#endif
-
 		// Stochastically round up or down using the hardware RWTexture2D truncation unit to take into precision
 		// loss due to pixel format encoding.
 		#if CONFIG_ENABLE_STOCASTIC_QUANTIZATION
@@ -1652,14 +1271,6 @@ void MainCS(
 			
 			FinalOutputColor = QuantizeForFloatRenderTarget(FinalOutputColor, E, HistoryQuantizationError);
 			FinalHighFrequencyColor = QuantizeForFloatRenderTarget(FinalHighFrequencyColor, E, HistoryQuantizationError);
-
-			#if CONFIG_GRAND_REPROJECTION
-				FinalPrevHighFrequencyColor = QuantizeForFloatRenderTarget(FinalPrevHighFrequencyColor, E, HistoryQuantizationError);
-				#if CONFIG_GRAND_REPROJECTION == 2
-					FinalPrevHighFrequencyColorResultant = QuantizeForFloatRenderTarget(FinalPrevHighFrequencyColorResultant, E, HistoryQuantizationError);
-					FinalOverblur = QuantizeForFloatRenderTarget(FinalOverblur, E, HistoryQuantizationError);
-				#endif
-			#endif
 		}
 		#endif
 
@@ -1690,13 +1301,6 @@ void MainCS(
 					#endif 
 					HistoryOutput_TranslucencyAlpha[dpv_lo(LocalHistoryPixelPos)] = dpv_lo(FinalTranslucencyColor).a;
 				#endif
-				#if CONFIG_GRAND_REPROJECTION
-					HistoryOutput_ColorArray[tsr_short3(dpv_lo(LocalHistoryPixelPos), HistoryOutput_ArrayIndices_PrevHighFrequency)] = dpv_lo(FinalPrevHighFrequencyColor);
-					#if CONFIG_GRAND_REPROJECTION == 2
-						HistoryOutput_ColorArray[tsr_short3(dpv_lo(LocalHistoryPixelPos), HistoryOutput_ArrayIndices_PrevHighFrequencyResultant)] = dpv_lo(FinalPrevHighFrequencyColorResultant);
-						HistoryOutput_ColorArray[tsr_short3(dpv_lo(LocalHistoryPixelPos), HistoryOutput_ArrayIndices_HighFrequencyOverblur)] = dpv_lo(FinalOverblur);
-					#endif
-				#endif
 			}
 		
 			// Output final history hi pixel.
@@ -1714,13 +1318,6 @@ void MainCS(
 					#endif
 					HistoryOutput_TranslucencyAlpha[dpv_hi(LocalHistoryPixelPos)] = dpv_hi(FinalTranslucencyColor).a;
 				#endif
-				#if CONFIG_GRAND_REPROJECTION
-					HistoryOutput_ColorArray[tsr_short3(dpv_hi(LocalHistoryPixelPos), HistoryOutput_ArrayIndices_PrevHighFrequency)] = dpv_hi(FinalPrevHighFrequencyColor);
-					#if CONFIG_GRAND_REPROJECTION == 2
-						HistoryOutput_ColorArray[tsr_short3(dpv_hi(LocalHistoryPixelPos), HistoryOutput_ArrayIndices_PrevHighFrequencyResultant)] = dpv_hi(FinalPrevHighFrequencyColorResultant);
-						HistoryOutput_ColorArray[tsr_short3(dpv_hi(LocalHistoryPixelPos), HistoryOutput_ArrayIndices_HighFrequencyOverblur)] = dpv_hi(FinalOverblur);
-					#endif
-				#endif
 			}
 			#endif // !CONFIG_ENABLE_DUAL_PIXEL_VECTORIZATION
 		}
diff --git a/Engine/Source/Runtime/Renderer/Private/PostProcess/TemporalSuperResolution.cpp b/Engine/Source/Runtime/Renderer/Private/PostProcess/TemporalSuperResolution.cpp
--- a/Engine/Source/Runtime/Renderer/Private/PostProcess/TemporalSuperResolution.cpp
+++ b/Engine/Source/Runtime/Renderer/Private/PostProcess/TemporalSuperResolution.cpp
@@ -56,11 +56,6 @@ TAutoConsoleVariable<int32> CVarTSRHistorySeparateTranslucency(
 	TEXT("Whether separate translucency should be accumulated separatly."),
 	ECVF_RenderThreadSafe);
 
-TAutoConsoleVariable<int32> CVarTSRHistoryGrandReprojection(
-	TEXT("r.TSR.History.GrandReprojection"), 1,
-	TEXT("Experimental functionality to keep higher sharpness in TSR's history in motion."),
-	ECVF_Scalability | ECVF_RenderThreadSafe);
-
 TAutoConsoleVariable<int32> CVarTSRWaveOps(
 	TEXT("r.TSR.WaveOps"), 1,
 	TEXT("Whether to use wave ops in the shading rejection heuristics to speeds up convolutions.\n")
@@ -236,8 +231,6 @@ TAutoConsoleVariable<int32> CVarTSRSubpixelMethod(
 	TEXT("\n")
 	TEXT("This settings control the method to reproject and/or discard the subpixel details.\n")
 	TEXT(" 0: disable subpixel details accumulation in history which means all this these subpixel features may ghost; \n")
-	TEXT(" 1: accumulate how much subpixel detail can paralax under movement of their background for history rejection which allow minimise ")
-	TEXT("amount of ghosting but at the expense of some image stability on these very thin geometry.\n")
 	TEXT(" 2: accumulate subpixel details' closest depth to be able to reproject them even when not drawing in depth/velocity buffer which works great for static geometry but not so much for any moving geometry (default)\n"),
 	ECVF_RenderThreadSafe);
 
@@ -259,11 +252,6 @@ TAutoConsoleVariable<int32> CVarTSRDebugArraySize(
 	TEXT("Size of array for the TSR.Debug.* RDG textures"),
 	ECVF_RenderThreadSafe);
 
-TAutoConsoleVariable<int32> CVarTSRSetupDebugPasses(
-	TEXT("r.TSR.Debug.SetupExtraPasses"), 0,
-	TEXT("Whether to enable the debug passes"),
-	ECVF_RenderThreadSafe);
-
 #endif
 
 BEGIN_SHADER_PARAMETER_STRUCT(FTSRCommonParameters, )
@@ -286,10 +274,6 @@ BEGIN_SHADER_PARAMETER_STRUCT(FTSRHistoryArrayIndices, )
 	SHADER_PARAMETER(int32, HighFrequency)
 	SHADER_PARAMETER(int32, Translucency)
 
-	SHADER_PARAMETER(int32, PrevHighFrequency)
-	SHADER_PARAMETER(int32, PrevHighFrequencyResultant)
-	SHADER_PARAMETER(int32, HighFrequencyOverblur)
-
 	SHADER_PARAMETER(int32, Size)
 END_SHADER_PARAMETER_STRUCT()
 
@@ -298,7 +282,6 @@ BEGIN_SHADER_PARAMETER_STRUCT(FTSRHistoryTextures, )
 	SHADER_PARAMETER_RDG_TEXTURE(Texture2D, Output)
 	SHADER_PARAMETER_RDG_TEXTURE(Texture2DArray, ColorArray)
 	SHADER_PARAMETER_RDG_TEXTURE(Texture2D, Metadata)
-	SHADER_PARAMETER_RDG_TEXTURE(Texture2D, SubpixelDetails)
 	SHADER_PARAMETER_RDG_TEXTURE(Texture2D, SubpixelDepth)
 	SHADER_PARAMETER_RDG_TEXTURE(Texture2D, TranslucencyAlpha)
 	SHADER_PARAMETER_RDG_TEXTURE(Texture2D, Guide)
@@ -309,10 +292,6 @@ BEGIN_SHADER_PARAMETER_STRUCT(FTSRHistorySRVs, )
 	SHADER_PARAMETER_RDG_TEXTURE_SRV(Texture2D, HighFrequency)
 	SHADER_PARAMETER_RDG_TEXTURE_SRV(Texture2D, Metadata)
 
-	SHADER_PARAMETER_RDG_TEXTURE_SRV(Texture2D, PrevHighFrequency)
-	SHADER_PARAMETER_RDG_TEXTURE_SRV(Texture2D, PrevHighFrequencyResultant)
-	SHADER_PARAMETER_RDG_TEXTURE_SRV(Texture2D, HighFrequencyOverblur)
-
 	SHADER_PARAMETER_RDG_TEXTURE_SRV(Texture2D, Translucency)
 	SHADER_PARAMETER_RDG_TEXTURE_SRV(Texture2D, TranslucencyAlpha)
 	SHADER_PARAMETER_RDG_TEXTURE(Texture2D, Guide)
@@ -329,9 +308,6 @@ END_SHADER_PARAMETER_STRUCT()
 BEGIN_SHADER_PARAMETER_STRUCT(FTSRPrevHistoryParameters, )
 	SHADER_PARAMETER_STRUCT(FScreenPassTextureViewportParameters, PrevHistoryInfo)
 	SHADER_PARAMETER(FScreenTransform, ScreenPosToPrevHistoryBufferUV)
-	SHADER_PARAMETER(FScreenTransform, ScreenPosToPrevSubpixelDetails)
-	SHADER_PARAMETER(FScreenTransform, ScreenPosToGrandPrevHistoryBufferUV)
-	SHADER_PARAMETER(FVector2f, PrevSubpixelDetailsExtent)
 	SHADER_PARAMETER(float, HistoryPreExposureCorrection)
 END_SHADER_PARAMETER_STRUCT()
 
@@ -340,7 +316,6 @@ enum class ETSRHistoryFormatBits : uint32
 	None = 0,
 	Translucency = 1 << 0,
 	Moire = 1 << 1,
-	GrandReprojection = 1 << 2,
 };
 ENUM_CLASS_FLAGS(ETSRHistoryFormatBits);
 
@@ -362,9 +337,6 @@ FTSRHistoryArrayIndices TranslateHistoryFormatBitsToArrayIndices(ETSRHistoryForm
 	ArrayIndices.Size = 0;
 	ArrayIndices.HighFrequency = -1;
 	ArrayIndices.Translucency = -1;
-	ArrayIndices.PrevHighFrequency = -1;
-	ArrayIndices.PrevHighFrequencyResultant = -1;
-	ArrayIndices.HighFrequencyOverblur = -1;
 
 	if (IsOutputDifferentThanHighFrequency(HistoryFormatBits))
 	{
@@ -376,13 +348,6 @@ FTSRHistoryArrayIndices TranslateHistoryFormatBitsToArrayIndices(ETSRHistoryForm
 		ArrayIndices.Translucency = ArrayIndices.Size++;
 	}
 
-	if (EnumHasAnyFlags(HistoryFormatBits, ETSRHistoryFormatBits::GrandReprojection))
-	{
-		ArrayIndices.PrevHighFrequency = ArrayIndices.Size++;
-		//ArrayIndices.PrevHighFrequencyResultant = ArrayIndices.Size++;
-		//ArrayIndices.HighFrequencyOverblur = ArrayIndices.Size++;
-	}
-
 	return ArrayIndices;
 }
 
@@ -403,12 +368,6 @@ FTSRHistorySRVs CreateSRVs(FRDGBuilder& GraphBuilder, const FTSRHistoryTextures&
 		SRVs.Translucency = GraphBuilder.CreateSRV(FRDGTextureSRVDesc::CreateForSlice(Textures.ColorArray, Textures.ArrayIndices.Translucency));
 		SRVs.TranslucencyAlpha = GraphBuilder.CreateSRV(FRDGTextureSRVDesc::Create(Textures.TranslucencyAlpha));
 	}
-	if (Textures.ArrayIndices.PrevHighFrequency >= 0)
-	{
-		SRVs.PrevHighFrequency = GraphBuilder.CreateSRV(FRDGTextureSRVDesc::CreateForSlice(Textures.ColorArray, Textures.ArrayIndices.PrevHighFrequency));
-		//SRVs.PrevHighFrequencyResultant = GraphBuilder.CreateSRV(FRDGTextureSRVDesc::CreateForSlice(Textures.ColorArray, Textures.ArrayIndices.PrevHighFrequencyResultant));
-		//SRVs.HighFrequencyOverblur = GraphBuilder.CreateSRV(FRDGTextureSRVDesc::CreateForSlice(Textures.ColorArray, Textures.ArrayIndices.HighFrequencyOverblur));
-	}
 
 	SRVs.Guide = Textures.Guide;
 	SRVs.Moire = Textures.Moire;
@@ -626,9 +585,7 @@ class FTSRDecimateHistoryCS : public FTSRShader
 	SHADER_USE_PARAMETER_STRUCT(FTSRDecimateHistoryCS, FTSRShader);
 
 	class FMoireReprojectionDim : SHADER_PERMUTATION_BOOL("DIM_MOIRE_REPROJECTION");
-	class FGrandReprojectionDim : SHADER_PERMUTATION_BOOL("DIM_GRAND_REPROJECTION");
-	class FSubpixelParallaxFactorDim : SHADER_PERMUTATION_BOOL("DIM_SUBPIXEL_PARALLAX_FACTOR");
-	using FPermutationDomain = TShaderPermutationDomain<FMoireReprojectionDim, FGrandReprojectionDim, FSubpixelParallaxFactorDim>;
+	using FPermutationDomain = TShaderPermutationDomain<FMoireReprojectionDim>;
 
 	BEGIN_SHADER_PARAMETER_STRUCT(FParameters, )
 		SHADER_PARAMETER_STRUCT_INCLUDE(FTSRCommonParameters, CommonParameters)
@@ -641,14 +598,10 @@ class FTSRDecimateHistoryCS : public FTSRShader
 		SHADER_PARAMETER_RDG_TEXTURE(Texture2D, ClosestDepthTexture)
 		SHADER_PARAMETER_RDG_TEXTURE(Texture2D, PrevUseCountTexture)
 		SHADER_PARAMETER_RDG_TEXTURE(Texture2D, PrevClosestDepthTexture)
-		SHADER_PARAMETER_RDG_TEXTURE_SRV(Texture2D, ParallaxFactorTexture)
 
 		SHADER_PARAMETER_STRUCT_INCLUDE(FTSRPrevHistoryParameters, PrevHistoryParameters)
-		SHADER_PARAMETER_RDG_TEXTURE(Texture2D, PrevHistorySubpixelDetails)
-
 		SHADER_PARAMETER_RDG_TEXTURE(Texture2D, PrevHistoryGuide)
 		SHADER_PARAMETER_RDG_TEXTURE(Texture2D, PrevHistoryMoire)
-		SHADER_PARAMETER_RDG_TEXTURE(Texture2D, PrevVelocityTexture)
 		SHADER_PARAMETER_STRUCT(FScreenPassTextureViewportParameters, PrevGuideInfo)
 		SHADER_PARAMETER(FScreenTransform, InputPixelPosToReprojectScreenPos)
 		SHADER_PARAMETER(FScreenTransform, ScreenPosToPrevHistoryGuideBufferUV)
@@ -658,10 +611,8 @@ class FTSRDecimateHistoryCS : public FTSRShader
 		SHADER_PARAMETER_RDG_TEXTURE_UAV(RWTexture2D, ReprojectedHistoryMoireOutput)
 		SHADER_PARAMETER_RDG_TEXTURE_UAV(RWTexture2D, HoleFilledVelocityOutput)
 		SHADER_PARAMETER_RDG_TEXTURE_UAV(RWTexture2D, HoleFilledVelocityMaskOutput)
-		SHADER_PARAMETER_RDG_TEXTURE_UAV(RWTexture2D, GrandHoleFilledVelocityOutput)
 
 		SHADER_PARAMETER_RDG_TEXTURE_UAV(RWTexture2D, ParallaxRejectionMaskOutput)
-		SHADER_PARAMETER_RDG_TEXTURE_UAV(RWTexture2D, HistorySubpixelDetailsOutput)
 
 		SHADER_PARAMETER_RDG_TEXTURE_UAV(RWTexture2DArray, DebugOutput)
 	END_SHADER_PARAMETER_STRUCT()
@@ -842,9 +793,8 @@ class FTSRUpdateHistoryCS : public FTSRShader
 
 	class FQualityDim : SHADER_PERMUTATION_ENUM_CLASS("DIM_UPDATE_QUALITY", EQuality);
 	class FSeparateTranslucencyDim : SHADER_PERMUTATION_BOOL("DIM_SEPARATE_TRANSLUCENCY");
-	class FGrandReprojectionDim : SHADER_PERMUTATION_BOOL("DIM_GRAND_REPROJECTION");
 
-	using FPermutationDomain = TShaderPermutationDomain<FQualityDim, FSeparateTranslucencyDim, FGrandReprojectionDim, FTSRShader::F16BitVALUDim>;
+	using FPermutationDomain = TShaderPermutationDomain<FQualityDim, FSeparateTranslucencyDim, FTSRShader::F16BitVALUDim>;
 
 	BEGIN_SHADER_PARAMETER_STRUCT(FParameters, )
 		SHADER_PARAMETER_STRUCT_INCLUDE(FTSRCommonParameters, CommonParameters)
@@ -855,8 +805,6 @@ class FTSRUpdateHistoryCS : public FTSRShader
 		SHADER_PARAMETER_RDG_TEXTURE(Texture2D, HistoryRejectionTexture)
 
 		SHADER_PARAMETER_RDG_TEXTURE(Texture2D, DilatedVelocityTexture)
-		SHADER_PARAMETER_RDG_TEXTURE(Texture2D, GrandDilatedVelocityTexture)
-		SHADER_PARAMETER_RDG_TEXTURE_SRV(Texture2D, ParallaxFactorTexture)
 		SHADER_PARAMETER_RDG_TEXTURE(Texture2D, ParallaxRejectionMaskTexture)
 		SHADER_PARAMETER_RDG_TEXTURE(Texture2D, AntiAliasingTexture)
 		SHADER_PARAMETER_RDG_TEXTURE(Texture2D, NoiseFilteringTexture)
@@ -878,7 +826,6 @@ class FTSRUpdateHistoryCS : public FTSRShader
 		SHADER_PARAMETER(float, InvWeightClampingPixelSpeed)
 		SHADER_PARAMETER(float, InputToHistoryFactor)
 		SHADER_PARAMETER(float, InputContributionMultiplier)
-		SHADER_PARAMETER(float, GrandPrevPreExposureCorrection)
 		SHADER_PARAMETER(int32, ResponsiveStencilMask)
 		SHADER_PARAMETER(int32, bGenerateOutputMip1)
 		SHADER_PARAMETER(int32, bGenerateOutputMip2)
@@ -886,7 +833,6 @@ class FTSRUpdateHistoryCS : public FTSRShader
 
 		SHADER_PARAMETER_STRUCT_INCLUDE(FTSRPrevHistoryParameters, PrevHistoryParameters)
 		SHADER_PARAMETER_STRUCT(FTSRHistorySRVs, PrevHistory)
-		SHADER_PARAMETER_RDG_TEXTURE_SRV(Texture2D, GrandPrevColorTexture)
 
 		SHADER_PARAMETER_RDG_TEXTURE_UAV(RWTexture2D, SceneColorOutputMip0)
 		SHADER_PARAMETER_RDG_TEXTURE_UAV(RWTexture2D, SceneColorOutputMip1)
@@ -898,12 +844,6 @@ class FTSRUpdateHistoryCS : public FTSRShader
 	{
 		FPermutationDomain PermutationVector(Parameters.PermutationId);
 
-		// Only compile the grand reprojection when translucency isn't accumulated separately.
-		if (PermutationVector.Get<FGrandReprojectionDim>() && PermutationVector.Get<FSeparateTranslucencyDim>())
-		{
-			return false;
-		}
-
 		if (!FTSRShader::ShouldCompile32or16BitPermutation(Parameters.Platform, PermutationVector.Get<FTSRShader::F16BitVALUDim>()))
 		{
 			return false;
@@ -1002,19 +942,6 @@ class FTSRResolveHistoryCS : public FTSRShader
 	}
 }; // class FTSRResolveHistoryCS
 
-class FTSRDebugHistoryCS : public FTSRShader
-{
-	DECLARE_GLOBAL_SHADER(FTSRDebugHistoryCS);
-	SHADER_USE_PARAMETER_STRUCT(FTSRDebugHistoryCS, FTSRShader);
-
-	BEGIN_SHADER_PARAMETER_STRUCT(FParameters, )
-		SHADER_PARAMETER_STRUCT_INCLUDE(FTSRCommonParameters, CommonParameters)
-		SHADER_PARAMETER_STRUCT(FTSRHistorySRVs, History)
-		SHADER_PARAMETER_STRUCT(FTSRHistorySRVs, PrevHistory)
-		SHADER_PARAMETER_RDG_TEXTURE_UAV(RWTexture2DArray, DebugOutput)
-	END_SHADER_PARAMETER_STRUCT()
-}; // class FTSRDebugHistoryCS
-
 IMPLEMENT_GLOBAL_SHADER(FTSRComputeMoireLumaCS,      "/Engine/Private/TemporalSuperResolution/TSRComputeMoireLuma.usf",      "MainCS", SF_Compute);
 IMPLEMENT_GLOBAL_SHADER(FTSRClearPrevTexturesCS,     "/Engine/Private/TemporalSuperResolution/TSRClearPrevTextures.usf",     "MainCS", SF_Compute);
 IMPLEMENT_GLOBAL_SHADER(FTSRForwardScatterDepthCS,   "/Engine/Private/TemporalSuperResolution/TSRForwardScatterDepth.usf",   "MainCS", SF_Compute);
@@ -1025,7 +952,6 @@ IMPLEMENT_GLOBAL_SHADER(FTSRSpatialAntiAliasingCS,   "/Engine/Private/TemporalSu
 IMPLEMENT_GLOBAL_SHADER(FTSRFilterAntiAliasingCS,    "/Engine/Private/TemporalSuperResolution/TSRFilterAntiAliasing.usf",    "MainCS", SF_Compute);
 IMPLEMENT_GLOBAL_SHADER(FTSRUpdateHistoryCS,         "/Engine/Private/TemporalSuperResolution/TSRUpdateHistory.usf",         "MainCS", SF_Compute);
 IMPLEMENT_GLOBAL_SHADER(FTSRResolveHistoryCS,        "/Engine/Private/TemporalSuperResolution/TSRResolveHistory.usf",        "MainCS", SF_Compute);
-IMPLEMENT_GLOBAL_SHADER(FTSRDebugHistoryCS,          "/Engine/Private/TemporalSuperResolution/TSRDebugHistory.usf",          "MainCS", SF_Compute);
 
 DECLARE_GPU_STAT(TemporalSuperResolution)
 
@@ -1102,10 +1028,6 @@ FDefaultTemporalUpscaler::FOutputs AddTemporalSuperResolutionPasses(
 	// Maximum depth lifetime encodable in FTSRSubpixelDepth::RemainingLifetime
 	const int32 kMaxDepthLifetime = (1 << 7) - 1;
 
-#if COMPILE_TSR_DEBUG_PASSES
-	const bool bSetupDebugPasses = CVarTSRSetupDebugPasses.GetValueOnRenderThread() != 0;
-#endif
-
 	// Whether to use wave ops optimizations.
 	const ERHIFeatureSupport WaveOpsSupport = FTSRShader::SupportsWaveOps(View.GetShaderPlatform());
 	const bool bSupportsLDS = FTSRShader::SupportsLDS(View.GetShaderPlatform());
@@ -1124,8 +1046,6 @@ FDefaultTemporalUpscaler::FOutputs AddTemporalSuperResolutionPasses(
 	// Whether separate translucency should be temporarily accumulated separatly.
 	const bool bAccumulateTranslucencySeparately = CVarTSRHistorySeparateTranslucency.GetValueOnRenderThread() != 0;
 
-	const bool bGrandReprojection = !bAccumulateTranslucencySeparately && CVarTSRHistoryGrandReprojection.GetValueOnRenderThread() != 0;
-
 	const float RefreshRateToFrameRateCap = (View.Family->Time.GetDeltaRealTimeSeconds() > 0.0f && CVarTSRFlickeringAdjustToFrameRate.GetValueOnRenderThread())
 		? View.Family->Time.GetDeltaRealTimeSeconds() * CVarTSRFlickeringFrameRateCap.GetValueOnRenderThread() : 1.0f;
 
@@ -1149,12 +1069,8 @@ FDefaultTemporalUpscaler::FOutputs AddTemporalSuperResolutionPasses(
 		{
 			HistoryFormatBits |= ETSRHistoryFormatBits::Moire;
 		}
-
-		if (bGrandReprojection)
-		{
-			HistoryFormatBits |= ETSRHistoryFormatBits::GrandReprojection;
-		}
 	}
+	FTSRHistoryArrayIndices HistoryArrayIndices = TranslateHistoryFormatBitsToArrayIndices(HistoryFormatBits);
 
 	const bool bIsOutputDifferentThanHighFrequency = IsOutputDifferentThanHighFrequency(HistoryFormatBits);
 
@@ -1375,14 +1291,7 @@ FDefaultTemporalUpscaler::FOutputs AddTemporalSuperResolutionPasses(
 			FClearValueBinding::None,
 			TexCreate_ShaderResource | TexCreate_UAV);
 
-		if (bGrandReprojection && History.ArrayIndices.PrevHighFrequencyResultant == -1)
-		{
-			Desc.Format = PF_R8G8;
-		}
-		else
-		{
-			Desc.Format = PF_R8;
-		}
+		Desc.Format = PF_R8;
 		History.Metadata = GraphBuilder.CreateTexture(Desc, TEXT("TSR.History.Metadata"));
 
 		Desc.Format = PF_R8;
@@ -1392,16 +1301,7 @@ FDefaultTemporalUpscaler::FOutputs AddTemporalSuperResolutionPasses(
 		}
 	}
 
-	if (SubpixelMethod == ETSRSubpixelMethod::ParallaxFactor)
-	{
-		FRDGTextureDesc SubpixelDetailsDesc = FRDGTextureDesc::Create2D(
-			InputExtent,
-			PF_R16_UINT,
-			FClearValueBinding::None,
-			TexCreate_ShaderResource | TexCreate_UAV);
-		History.SubpixelDetails = GraphBuilder.CreateTexture(SubpixelDetailsDesc, TEXT("TSR.History.SubpixelInfo"));
-	}
-	else if (SubpixelMethod == ETSRSubpixelMethod::ClosestDepth)
+	if (SubpixelMethod == ETSRSubpixelMethod::ClosestDepth)
 	{
 		FRDGTextureDesc SubpixelDepthDesc = FRDGTextureDesc::Create2D(
 			InputExtent,
@@ -1430,16 +1330,11 @@ FDefaultTemporalUpscaler::FOutputs AddTemporalSuperResolutionPasses(
 		DummyHistorySRVs.Metadata = GraphBuilder.CreateSRV(FRDGTextureSRVDesc::Create(BlackDummy));
 		DummyHistorySRVs.Translucency = GraphBuilder.CreateSRV(FRDGTextureSRVDesc::Create(BlackDummy));
 		DummyHistorySRVs.TranslucencyAlpha = GraphBuilder.CreateSRV(FRDGTextureSRVDesc::Create(WhiteDummy));
-		DummyHistorySRVs.PrevHighFrequency = GraphBuilder.CreateSRV(FRDGTextureSRVDesc::Create(BlackDummy));
-		DummyHistorySRVs.PrevHighFrequencyResultant = GraphBuilder.CreateSRV(FRDGTextureSRVDesc::Create(BlackDummy));
-		DummyHistorySRVs.HighFrequencyOverblur = GraphBuilder.CreateSRV(FRDGTextureSRVDesc::Create(BlackDummy));
 		DummyHistorySRVs.Guide = BlackDummy;
 		DummyHistorySRVs.Moire = BlackDummy;
 	}
 
 	// Setup the previous frame history
-	FRDGTextureRef PrevVelocity = BlackDummy;
-	FRDGTextureSRVRef GrandPrevColorTexture = GraphBuilder.CreateSRV(BlackDummy);
 	FTSRHistoryTextures PrevHistory;
 	FTSRHistorySRVs PrevHistorySRVs;
 	if (!bCameraCut)
@@ -1455,23 +1350,9 @@ FDefaultTemporalUpscaler::FOutputs AddTemporalSuperResolutionPasses(
 		PrevHistory.Moire = InputHistory.Moire.IsValid() ? GraphBuilder.RegisterExternalTexture(InputHistory.Moire, TEXT("TSR.PrevHistory.Moire")) : DummyHistorySRVs.Moire;
 
 		// Register non-filterable history
-		PrevHistory.SubpixelDetails = InputHistory.SubpixelDetails.IsValid() ? GraphBuilder.RegisterExternalTexture(InputHistory.SubpixelDetails, TEXT("TSR.PrevHistory.SubpixelInfo")) : nullptr;
 		PrevHistory.SubpixelDepth = InputHistory.SubpixelDepth.IsValid() ? GraphBuilder.RegisterExternalTexture(InputHistory.SubpixelDepth, TEXT("TSR.PrevHistory.SubpixelDepth")) : nullptr;
 
 		PrevHistorySRVs = CreateSRVs(GraphBuilder, PrevHistory);
-
-		if (bGrandReprojection)
-		{
-			PrevVelocity = GraphBuilder.RegisterExternalTexture(InputHistory.Velocity);
-			if (InputHistory.PrevColorArray.IsValid() && bIsOutputDifferentThanHighFrequency)
-			{
-				GrandPrevColorTexture = GraphBuilder.CreateSRV(FRDGTextureSRVDesc::CreateForSlice(GraphBuilder.RegisterExternalTexture(InputHistory.PrevColorArray, TEXT("TSR.GrandPrevHistory.ColorArray")), PrevHistory.ArrayIndices.HighFrequency));
-			}
-			else if (InputHistory.PrevOutput.IsValid() && !bIsOutputDifferentThanHighFrequency)
-			{
-				GrandPrevColorTexture = GraphBuilder.CreateSRV(GraphBuilder.RegisterExternalTexture(InputHistory.PrevOutput, TEXT("TSR.GrandPrevHistory.Output")));
-			}
-		}
 	}
 	else
 	{
@@ -1483,27 +1364,15 @@ FDefaultTemporalUpscaler::FOutputs AddTemporalSuperResolutionPasses(
 	{
 		// Setup prev history parameters.
 		FScreenPassTextureViewport PrevHistoryViewport(PrevHistorySRVs.HighFrequency->Desc.Texture->Desc.Extent, InputHistory.OutputViewportRect);
-		FScreenPassTextureViewport PrevSubpixelDetailsViewport(PrevHistorySRVs.Guide->Desc.Extent, InputHistory.InputViewportRect);
-		FScreenPassTextureViewport GrandPrevSubpixelDetailsViewport(GrandPrevColorTexture->Desc.Texture->Desc.Extent, InputHistory.PrevOutputViewportRect);
-
 		if (bCameraCut)
 		{
 			PrevHistoryViewport.Extent = FIntPoint(1, 1);
 			PrevHistoryViewport.Rect = FIntRect(FIntPoint(0, 0), FIntPoint(1, 1));
-			PrevSubpixelDetailsViewport.Extent = FIntPoint(1, 1);
-			PrevSubpixelDetailsViewport.Rect = FIntRect(FIntPoint(0, 0), FIntPoint(1, 1));
-			GrandPrevSubpixelDetailsViewport.Extent = FIntPoint(1, 1);
-			GrandPrevSubpixelDetailsViewport.Rect = FIntRect(FIntPoint(0, 0), FIntPoint(1, 1));
 		}
 
 		PrevHistoryParameters.PrevHistoryInfo = GetScreenPassTextureViewportParameters(PrevHistoryViewport);
 		PrevHistoryParameters.ScreenPosToPrevHistoryBufferUV = FScreenTransform::ChangeTextureBasisFromTo(
 			PrevHistoryViewport, FScreenTransform::ETextureBasis::ScreenPosition, FScreenTransform::ETextureBasis::TextureUV);
-		PrevHistoryParameters.ScreenPosToPrevSubpixelDetails = FScreenTransform::ChangeTextureBasisFromTo(
-			PrevSubpixelDetailsViewport, FScreenTransform::ETextureBasis::ScreenPosition, FScreenTransform::ETextureBasis::TextureUV);
-		PrevHistoryParameters.ScreenPosToGrandPrevHistoryBufferUV = FScreenTransform::ChangeTextureBasisFromTo(
-			GrandPrevSubpixelDetailsViewport, FScreenTransform::ETextureBasis::ScreenPosition, FScreenTransform::ETextureBasis::TextureUV);
-		PrevHistoryParameters.PrevSubpixelDetailsExtent = PrevSubpixelDetailsViewport.Extent;
 		PrevHistoryParameters.HistoryPreExposureCorrection = View.PreExposure / View.PrevViewInfo.SceneColorPreExposure;
 	}
 
@@ -1583,7 +1452,6 @@ FDefaultTemporalUpscaler::FOutputs AddTemporalSuperResolutionPasses(
 	// Dilate the velocity texture & scatter reprojection into previous frame
 	FRDGTextureRef DilatedVelocityTexture;
 	FRDGTextureRef ClosestDepthTexture;
-	FRDGTextureSRVRef ParallaxFactorTexture;
 	FRDGTextureSRVRef IsMovingMaskTexture = nullptr;
 	FVelocityFlattenTextures VelocityFlattenTextures;
 	{
@@ -1609,13 +1477,12 @@ FDefaultTemporalUpscaler::FOutputs AddTemporalSuperResolutionPasses(
 				PF_R8_UINT,
 				FClearValueBinding::None,
 				/* InFlags = */ TexCreate_ShaderResource | TexCreate_UAV,
-				bOutputIsMovingTexture ? 2 : 1);
+				1);
 
 			R8OutputTexture = GraphBuilder.CreateTexture(Desc, TEXT("TSR.ParallaxFactor"));
-			ParallaxFactorTexture = GraphBuilder.CreateSRV(FRDGTextureSRVDesc::CreateForSlice(R8OutputTexture, 0));
 			if (bOutputIsMovingTexture)
 			{
-				IsMovingMaskTexture = GraphBuilder.CreateSRV(FRDGTextureSRVDesc::CreateForSlice(R8OutputTexture, 1));
+				IsMovingMaskTexture = GraphBuilder.CreateSRV(FRDGTextureSRVDesc::CreateForSlice(R8OutputTexture, 0));
 			}
 		}
 
@@ -1723,7 +1590,6 @@ FDefaultTemporalUpscaler::FOutputs AddTemporalSuperResolutionPasses(
 	FRDGTextureRef ReprojectedHistoryMoireTexture = nullptr;
 	FRDGTextureRef ParallaxRejectionMaskTexture = nullptr;
 	FRDGTextureRef HoleFilledVelocityMaskTexture = nullptr;
-	FRDGTextureRef GrandHoleFilledVelocityTexture = BlackDummy;
 	{
 		FRDGTextureRef HoleFilledVelocityTexture;
 		{
@@ -1737,7 +1603,6 @@ FDefaultTemporalUpscaler::FOutputs AddTemporalSuperResolutionPasses(
 
 			Desc.Format = PF_G16R16;
 			HoleFilledVelocityTexture = GraphBuilder.CreateTexture(Desc, TEXT("TSR.Velocity.HoleFilled"));
-			GrandHoleFilledVelocityTexture = GraphBuilder.CreateTexture(Desc, TEXT("TSR.Velocity.GrandVelocity"));
 
 			Desc.Format = PF_R8G8;
 			HoleFilledVelocityMaskTexture = GraphBuilder.CreateTexture(Desc, TEXT("TSR.velocity.HoleFillMask"));
@@ -1779,11 +1644,8 @@ FDefaultTemporalUpscaler::FOutputs AddTemporalSuperResolutionPasses(
 		PassParameters->ClosestDepthTexture = ClosestDepthTexture;
 		PassParameters->PrevUseCountTexture = PrevUseCountTexture;
 		PassParameters->PrevClosestDepthTexture = PrevClosestDepthTexture;
-		PassParameters->ParallaxFactorTexture = ParallaxFactorTexture;
 
 		PassParameters->PrevHistoryParameters = PrevHistoryParameters;
-		PassParameters->PrevHistorySubpixelDetails = PrevHistory.SubpixelDetails ? PrevHistory.SubpixelDetails : BlackUintDummy;
-		PassParameters->PrevVelocityTexture = PrevVelocity;
 
 		{
 			FScreenPassTextureViewport PrevHistoryColorViewport(PrevHistorySRVs.Guide->Desc.Extent, InputHistory.InputViewportRect);
@@ -1802,30 +1664,17 @@ FDefaultTemporalUpscaler::FOutputs AddTemporalSuperResolutionPasses(
 		PassParameters->ReprojectedHistoryMoireOutput = GraphBuilder.CreateUAV(ReprojectedHistoryMoireTexture);
 		PassParameters->HoleFilledVelocityOutput = GraphBuilder.CreateUAV(HoleFilledVelocityTexture);
 		PassParameters->HoleFilledVelocityMaskOutput = GraphBuilder.CreateUAV(HoleFilledVelocityMaskTexture);
-		PassParameters->GrandHoleFilledVelocityOutput = GraphBuilder.CreateUAV(GrandHoleFilledVelocityTexture);
 		PassParameters->ParallaxRejectionMaskOutput = GraphBuilder.CreateUAV(ParallaxRejectionMaskTexture);
-		if (History.SubpixelDetails)
-		{
-			PassParameters->HistorySubpixelDetailsOutput = GraphBuilder.CreateUAV(History.SubpixelDetails);
-		}
-		else
-		{
-			PassParameters->HistorySubpixelDetailsOutput = CreateDummyUAV(GraphBuilder, PF_R16_UINT);
-		}
 		PassParameters->DebugOutput = CreateDebugUAV(InputExtent, TEXT("Debug.TSR.DecimateHistory"));
 
 		FTSRDecimateHistoryCS::FPermutationDomain PermutationVector;
 		PermutationVector.Set<FTSRDecimateHistoryCS::FMoireReprojectionDim>(FlickeringFramePeriod > 0.0f);
-		PermutationVector.Set<FTSRDecimateHistoryCS::FGrandReprojectionDim>(bGrandReprojection);
-		PermutationVector.Set<FTSRDecimateHistoryCS::FSubpixelParallaxFactorDim>(SubpixelMethod == ETSRSubpixelMethod::ParallaxFactor);
 
 		TShaderMapRef<FTSRDecimateHistoryCS> ComputeShader(View.ShaderMap, PermutationVector);
 		FComputeShaderUtils::AddPass(
 			GraphBuilder,
-			RDG_EVENT_NAME("TSR DecimateHistory(%s%s%s) %dx%d",
+			RDG_EVENT_NAME("TSR DecimateHistory(%s) %dx%d",
 				PermutationVector.Get<FTSRDecimateHistoryCS::FMoireReprojectionDim>() ? TEXT("ReprojectMoire") : TEXT(""),
-				PermutationVector.Get<FTSRDecimateHistoryCS::FGrandReprojectionDim>() ? TEXT(" GrandReprojection") : TEXT(""),
-				PermutationVector.Get<FTSRDecimateHistoryCS::FSubpixelParallaxFactorDim>() ? TEXT(" SubpixelParallaxFactor") : TEXT(""),
 				InputRect.Width(), InputRect.Height()),
 			AsyncComputePasses >= 2 ? ERDGPassFlags::AsyncCompute : ERDGPassFlags::Compute,
 			ComputeShader,
@@ -2012,8 +1861,6 @@ FDefaultTemporalUpscaler::FOutputs AddTemporalSuperResolutionPasses(
 		PassParameters->HistoryRejectionTexture = HistoryRejectionTexture;
 
 		PassParameters->DilatedVelocityTexture = DilatedVelocityTexture;
-		PassParameters->GrandDilatedVelocityTexture = GrandHoleFilledVelocityTexture;
-		PassParameters->ParallaxFactorTexture = ParallaxFactorTexture;
 		PassParameters->ParallaxRejectionMaskTexture = ParallaxRejectionMaskTexture;
 		PassParameters->AntiAliasingTexture = AntiAliasingTexture;
 		PassParameters->NoiseFilteringTexture = NoiseFilteringTexture;
@@ -2037,7 +1884,6 @@ FDefaultTemporalUpscaler::FOutputs AddTemporalSuperResolutionPasses(
 		
 		PassParameters->InputToHistoryFactor = float(HistorySize.X) / float(InputRect.Width());
 		PassParameters->InputContributionMultiplier = OutputToHistoryResolutionFractionSquare; 
-		PassParameters->GrandPrevPreExposureCorrection = bCameraCut ? 1.0f : View.PreExposure / InputHistory.PrevSceneColorPreExposure;
 		//PassParameters->ResponsiveStencilMask = CVarTSREnableResponiveAA.GetValueOnRenderThread() ? (STENCIL_TEMPORAL_RESPONSIVE_AA_MASK) : 0;
 		PassParameters->ResponsiveStencilMask = 0;
 		PassParameters->bGenerateOutputMip1 = false;
@@ -2046,7 +1892,6 @@ FDefaultTemporalUpscaler::FOutputs AddTemporalSuperResolutionPasses(
 
 		PassParameters->PrevHistoryParameters = PrevHistoryParameters;
 		PassParameters->PrevHistory = PrevHistorySRVs;
-		PassParameters->GrandPrevColorTexture = GrandPrevColorTexture;
 
 		PassParameters->HistoryOutput = CreateUAVs(GraphBuilder, History);
 			PassParameters->SceneColorOutputMip0 = GraphBuilder.CreateUAV(FRDGTextureUAVDesc(UpdateHistoryOutputTexture, /* InMipLevel = */ 0));
@@ -2075,17 +1920,15 @@ FDefaultTemporalUpscaler::FOutputs AddTemporalSuperResolutionPasses(
 		FTSRUpdateHistoryCS::FPermutationDomain PermutationVector;
 		PermutationVector.Set<FTSRUpdateHistoryCS::FQualityDim>(UpdateHistoryQuality);
 		PermutationVector.Set<FTSRUpdateHistoryCS::FSeparateTranslucencyDim>(bAccumulateTranslucencySeparately);
-		PermutationVector.Set<FTSRUpdateHistoryCS::FGrandReprojectionDim>(bGrandReprojection);
 		PermutationVector.Set<FTSRShader::F16BitVALUDim>(bUse16BitVALU);
 
 		TShaderMapRef<FTSRUpdateHistoryCS> ComputeShader(View.ShaderMap, PermutationVector);
 		FComputeShaderUtils::AddPass(
 			GraphBuilder,
-			RDG_EVENT_NAME("TSR UpdateHistory(Quality=%s%s%s%s%s%s) %dx%d",
+			RDG_EVENT_NAME("TSR UpdateHistory(Quality=%s%s%s%s%s) %dx%d",
 				kUpdateQualityNames[int32(PermutationVector.Get<FTSRUpdateHistoryCS::FQualityDim>())],
 				PermutationVector.Get<FTSRShader::F16BitVALUDim>() ? TEXT(" 16bit") : TEXT(""),
 				PermutationVector.Get<FTSRUpdateHistoryCS::FSeparateTranslucencyDim>() ? TEXT(" SeparateTranslucency") : TEXT(""),
-				PermutationVector.Get<FTSRUpdateHistoryCS::FGrandReprojectionDim>() ? TEXT(" GrandReprojection") : TEXT(""),
 				(History.ColorArray ? History.ColorArray : History.Output)->Desc.Format == PF_FloatR11G11B10 ? TEXT(" R11G11B10") : TEXT(""),
 				PassParameters->bGenerateOutputMip2 ? TEXT(" OutputMip2") : (PassParameters->bGenerateOutputMip1 ? TEXT(" OutputMip1") : TEXT("")),
 				HistorySize.X, HistorySize.Y),
@@ -2095,29 +1938,6 @@ FDefaultTemporalUpscaler::FOutputs AddTemporalSuperResolutionPasses(
 			FComputeShaderUtils::GetGroupCount(HistorySize, 8));
 	}
 
-	// Debug the history.
-	#if COMPILE_TSR_DEBUG_PASSES
-	if (bSetupDebugPasses)
-	{
-		const int32 kHistoryUpscalingFactor = 2;
-
-		FTSRDebugHistoryCS::FParameters* PassParameters = GraphBuilder.AllocParameters<FTSRDebugHistoryCS::FParameters>();
-		PassParameters->CommonParameters = CommonParameters;
-		PassParameters->History = CreateSRVs(GraphBuilder, History);
-		PassParameters->PrevHistory = PrevHistorySRVs;
-		PassParameters->DebugOutput = CreateDebugUAV(HistoryExtent * kHistoryUpscalingFactor, TEXT("Debug.TSR.History"));
-
-		TShaderMapRef<FTSRDebugHistoryCS> ComputeShader(View.ShaderMap);
-		FComputeShaderUtils::AddPass(
-			GraphBuilder,
-			RDG_EVENT_NAME("TSR DebugHistory %dx%d", HistorySize.X, HistorySize.Y),
-			AsyncComputePasses >= 3 ? ERDGPassFlags::AsyncCompute : ERDGPassFlags::Compute,
-			ComputeShader,
-			PassParameters,
-			FComputeShaderUtils::GetGroupCount(HistorySize * kHistoryUpscalingFactor, 8));
-	}
-	#endif
-
 	// If we upscaled the history buffer, downsize back to the secondary screen percentage size.
 	if (HistorySize != OutputRect.Size())
 	{
@@ -2201,11 +2021,7 @@ FDefaultTemporalUpscaler::FOutputs AddTemporalSuperResolutionPasses(
 		}
 
 		// Extract non-filterable history
-		if (History.SubpixelDetails)
-		{
-			GraphBuilder.QueueTextureExtraction(History.SubpixelDetails, &OutputHistory.SubpixelDetails);
-		}
-		else if (History.SubpixelDepth)
+		if (History.SubpixelDepth)
 		{
 			GraphBuilder.QueueTextureExtraction(History.SubpixelDepth, &OutputHistory.SubpixelDepth);
 		}
@@ -2218,33 +2034,6 @@ FDefaultTemporalUpscaler::FOutputs AddTemporalSuperResolutionPasses(
 			GraphBuilder.QueueTextureExtraction(History.Moire, &OutputHistory.Moire);
 		}
 
-		if (bGrandReprojection)
-		{
-			GraphBuilder.QueueTextureExtraction(DilatedVelocityTexture, &OutputHistory.Velocity);
-			if (!bIsOutputDifferentThanHighFrequency)
-			{
-				if (PrevHistory.Output)
-				{
-					GraphBuilder.QueueTextureExtraction(PrevHistory.Output, &OutputHistory.PrevOutput);
-				}
-				else
-				{
-					GraphBuilder.QueueTextureExtraction(History.Output, &OutputHistory.PrevOutput);
-				}
-			}
-			else
-			{
-				if (PrevHistory.ColorArray)
-				{
-					GraphBuilder.QueueTextureExtraction(PrevHistory.ColorArray, &OutputHistory.PrevColorArray);
-				}
-				else
-				{
-					GraphBuilder.QueueTextureExtraction(History.ColorArray, &OutputHistory.PrevColorArray);
-				}
-			}
-		}
-
 		// Extract the output for next frame SSR so that separate translucency shows up in SSR.
 		{
 			// Output in TemporalAAHistory and not CustomSSR so Lumen can pick up ScreenSpaceRayTracingInput in priority to ensure consistent behavior between TAA and TSR.
diff --git a/Engine/Source/Runtime/Renderer/Private/RendererScene.cpp b/Engine/Source/Runtime/Renderer/Private/RendererScene.cpp
--- a/Engine/Source/Runtime/Renderer/Private/RendererScene.cpp
+++ b/Engine/Source/Runtime/Renderer/Private/RendererScene.cpp
@@ -853,13 +853,9 @@ uint64 FTSRHistory::GetGPUSizeBytes(bool bLogSizes) const
 		GetRenderTargetGPUSizeBytes(ColorArray, bLogSizes) +
 		GetRenderTargetGPUSizeBytes(Metadata, bLogSizes) +
 		GetRenderTargetGPUSizeBytes(TranslucencyAlpha, bLogSizes) +
-		GetRenderTargetGPUSizeBytes(SubpixelDetails, bLogSizes) +
 		GetRenderTargetGPUSizeBytes(SubpixelDepth, bLogSizes) +
 		GetRenderTargetGPUSizeBytes(Guide, bLogSizes) +
-		GetRenderTargetGPUSizeBytes(Moire, bLogSizes) +
-		GetRenderTargetGPUSizeBytes(Velocity, bLogSizes) +
-		GetRenderTargetGPUSizeBytes(PrevOutput, bLogSizes) +
-		GetRenderTargetGPUSizeBytes(PrevColorArray, bLogSizes);
+		GetRenderTargetGPUSizeBytes(Moire, bLogSizes);
 }
 
 uint64 FScreenSpaceDenoiserHistory::GetGPUSizeBytes(bool bLogSizes) const
diff --git a/Engine/Source/Runtime/Renderer/Private/SceneRendering.h b/Engine/Source/Runtime/Renderer/Private/SceneRendering.h
--- a/Engine/Source/Runtime/Renderer/Private/SceneRendering.h
+++ b/Engine/Source/Runtime/Renderer/Private/SceneRendering.h
@@ -997,15 +997,9 @@ struct FTSRHistory
 	TRefCountPtr<IPooledRenderTarget> TranslucencyAlpha;
 
 	// Input resolution representation of the output
-	TRefCountPtr<IPooledRenderTarget> SubpixelDetails;
 	TRefCountPtr<IPooledRenderTarget> SubpixelDepth;
 	TRefCountPtr<IPooledRenderTarget> Guide;
 	TRefCountPtr<IPooledRenderTarget> Moire;
-	TRefCountPtr<IPooledRenderTarget> Velocity;
-
-	// Previous frame's history.
-	TRefCountPtr<IPooledRenderTarget> PrevOutput;
-	TRefCountPtr<IPooledRenderTarget> PrevColorArray;
 
 	// Frame's input and output resolution.
 	FIntRect InputViewportRect;
