Fix ScreenPosition material expression not rendering correctly in secondary views with Instanced Stereo Rendering enabled.

diff --git a/Engine/Shaders/Private/Common.ush b/Engine/Shaders/Private/Common.ush
--- a/Engine/Shaders/Private/Common.ush
+++ b/Engine/Shaders/Private/Common.ush
@@ -1318,9 +1318,9 @@ void SvPositionToResolvedScreenPositionDeriv(float4 SvPosition, float2 PPZ_DDX_D
 float2 SvPositionToViewportUV(float4 SvPosition)
 {
 	// can be optimized from 2SUB+2MUL to 2MAD
-	float2 PixelPos = SvPosition.xy - View.ViewRectMin.xy;	
+	float2 PixelPos = SvPosition.xy - ResolvedView.ViewRectMin.xy;	
 
-	return PixelPos.xy * View.ViewSizeAndInvSize.zw;
+	return PixelPos.xy * ResolvedView.ViewSizeAndInvSize.zw;
 }
 
 float2 BufferUVToViewportUV(float2 BufferUV)
diff --git a/Engine/Shaders/Private/MaterialTemplate.ush b/Engine/Shaders/Private/MaterialTemplate.ush
--- a/Engine/Shaders/Private/MaterialTemplate.ush
+++ b/Engine/Shaders/Private/MaterialTemplate.ush
@@ -898,7 +898,7 @@ float2 GetViewportUV(FMaterialPixelParameters Parameters)
 
 float2 GetPixelPosition(FMaterialPixelParameters Parameters)
 {
-	return Parameters.SvPosition.xy - float2(View.ViewRectMin.xy);
+	return Parameters.SvPosition.xy - ResolvedView.ViewRectMin.xy;
 }
 
 float2 GetViewportUV(FMaterialPixelParameters Parameters)
