Substrate - added reference model for simple volume.

diff --git a/Engine/Config/BaseEngine.ini b/Engine/Config/BaseEngine.ini
--- a/Engine/Config/BaseEngine.ini
+++ b/Engine/Config/BaseEngine.ini
@@ -163,6 +163,8 @@ BlueNoiseScalarTextureName=/Engine/EngineMaterials/STBlueNoise_scalar_128x128x64
 BlueNoiseVec2TextureName=/Engine/EngineMaterials/STBlueNoise_vec2_128x128x64.STBlueNoise_vec2_128x128x64
 GlintTextureName=/Engine/EngineMaterials/Substrate/GLints/AGlintTex2dArray.AGlintTex2dArray
 GlintTexture2Name=/Engine/EngineMaterials/Substrate/GLints2/AGlintTex2dArray2.AGlintTex2dArray2
+SimpleVolumeTextureName=/Engine/EngineMaterials/Substrate/Volume/SubstrateSimpleVolumeLUT3D.SubstrateSimpleVolumeLUT3D
+SimpleVolumeEnvTextureName=/Engine/EngineMaterials/Substrate/Volume/SubstrateSimpleVolumeEnvLUT3D.SubstrateSimpleVolumeEnvLUT3D
 MiniFontTextureName=/Engine/EngineMaterials/MiniFont.MiniFont
 WeightMapPlaceholderTextureName=/Engine/EngineMaterials/WeightMapPlaceholderTexture.WeightMapPlaceholderTexture
 LightMapDensityTextureName=/Engine/EngineMaterials/DefaultWhiteGrid.DefaultWhiteGrid
diff --git a/Engine/Shaders/Private/ParticipatingMediaCommon.ush b/Engine/Shaders/Private/ParticipatingMediaCommon.ush
--- a/Engine/Shaders/Private/ParticipatingMediaCommon.ush
+++ b/Engine/Shaders/Private/ParticipatingMediaCommon.ush
@@ -5,6 +5,7 @@
 #pragma once
 
 #include "Common.ush"
+#include "FastMath.ush"
 
 #define PARTICIPATING_MEDIA_MIN_MFP_METER		0.000000000001f
 #define PARTICIPATING_MEDIA_MIN_EXTINCTION		0.000000000001f
@@ -209,6 +210,75 @@ float3 IsotropicMediumSlabPunctualDirectionalAlbedo(FParticipatingMedia PM)
 	return FinalEvaluate * IsotropicMediumSlabDirectionalAlbedoFade(PM.BaseColor, MFP);
 }
 
+// Wait before activating that feature for now.
+#define USE_NEW_SIMPLEVOLUME 0
+
+float IsotropicMediumSlabDirectionalAlbedoFadeLUT(float MFP)
+{
+	// Fade out directional albedo to zero from MFP = 10meters (the last measured MFP for which we have fit the curve) to a large MFP=1000 meters
+	const float FitLastMeasuredSampleMFP = 10.0f;
+	const float AlbedoIsZeroForMFP = 640.0f;	// This limit is realted to the maximum value MFP that can be stored in the material buffer.
+	const float Fade = saturate(1.0f - (MFP - FitLastMeasuredSampleMFP) / (AlbedoIsZeroForMFP - FitLastMeasuredSampleMFP));
+	return Fade;
+}
+
+float IsotropicMediumSlabPunctualDirectionalAlbedoLUT(float MFP, float BaseColor, float NoV, float NoL)
+{
+	// LUT content
+	//		View angle	from 0 to 7/8	with steps of 1/8 => 8  texels
+	//		Light angle	from 0 to 1  	with steps of 1/8 => 9  texels
+	//		BaseColor	from 0 to 1		with steps of 0.1 => 11 texels
+	//		MFP			from 0 to 10	with steps of 1   => 11 texels
+	// 
+	// Indexing  LUTINDEX(V, L, B, M) (V * DimL * DimBC * DimMFP + L * DimBC * DimMFP  + B * DimMFP  + M)
+
+	// LUT dimension for input parameters: MFP, BaseColor, NoL, NoV
+	const float4 LUTDim = float4(11.0f, 11.0, 9.0f, 8.0f);
+	const float4 LUTDimInv = 1.0f / LUTDim;
+	const float4 LUTMin = 0.5f * LUTDimInv;
+	const float4 LUTMax = (LUTDim - 0.5f) * LUTDimInv;
+	const float4 LUTLen = (LUTDim - 1.0f) * LUTDimInv;
+
+	const float PIOver2 = PI / 2.0f;
+	const float LightAngleNorm = saturate(acosFast(saturate(NoL)) / PIOver2);
+	const float ViewAngleNorm  = saturate(acosFast(saturate(NoV)) / PIOver2);
+
+	float4 UVs = float4(MFP / 10.0f, BaseColor, LightAngleNorm, ViewAngleNorm);	// Input parameters in [0,1]
+	UVs = LUTMin + LUTLen * UVs;												// Correct in order to ignore texel border
+
+	float FourthCoord = ViewAngleNorm * LUTDim.w;
+	float FourthCoord0 = floor(FourthCoord);
+	      FourthCoord0 = min(FourthCoord0, LUTDim.w - 1.0);
+	float FourthCoord1 = FourthCoord0 + 1.0;
+	      FourthCoord1 = min(FourthCoord1, LUTDim.w - 1.0);
+	float FourthCoordLerp = saturate(FourthCoord - FourthCoord0);
+
+	const float LUT3DDimZ    = LUTDim.z * LUTDim.w;
+	const float LUT4DUvDimZ  = (LUTDim.z - 1.0) / LUT3DDimZ;	// Size of third dimension along Z
+	const float LUT4DUvStepW = (LUTDim.z      ) / LUT3DDimZ; // Size of step to sample the 4th dimensions along Z
+	const float EdgeOffset = (0.5) / LUT3DDimZ;
+
+	float3 UVW0 = float3(UVs.x, UVs.y, EdgeOffset + FourthCoord0 * LUT4DUvStepW + LightAngleNorm * LUT4DUvDimZ);
+	float3 UVW1 = float3(UVs.x, UVs.y, EdgeOffset + FourthCoord1 * LUT4DUvStepW + LightAngleNorm * LUT4DUvDimZ);
+
+	float Value0 = View.SimpleVolumeTexture.SampleLevel(View.SimpleVolumeTextureSampler, UVW0, 0);
+	float Value1 = View.SimpleVolumeTexture.SampleLevel(View.SimpleVolumeTextureSampler, UVW1, 0);
+
+	float DirectionalAlbedo = lerp(Value0, Value1, FourthCoordLerp);
+
+	const float Fade = IsotropicMediumSlabDirectionalAlbedoFadeLUT(MFP);
+	return DirectionalAlbedo * Fade;
+}
+
+float3 IsotropicMediumSlabPunctualDirectionalAlbedoLUT(FParticipatingMedia PM, float NoV, float3 NoL)
+{
+	return float3(
+		IsotropicMediumSlabPunctualDirectionalAlbedoLUT(PM.MeanFreePath.r, PM.BaseColor.r, NoV, NoL),
+		IsotropicMediumSlabPunctualDirectionalAlbedoLUT(PM.MeanFreePath.g, PM.BaseColor.g, NoV, NoL),
+		IsotropicMediumSlabPunctualDirectionalAlbedoLUT(PM.MeanFreePath.b, PM.BaseColor.b, NoV, NoL)
+	);
+}
+
 float3 IsotropicMediumSlabEnvDirectionalAlbedo(FParticipatingMedia PM)
 {
 	// Our measurement are only valid for a MFP of 0.01 meter so we clamp input MFP to that.
@@ -234,3 +304,40 @@ float3 IsotropicMediumSlabTransmittance(FParticipatingMedia PM, float SlabThickn
 	return SafePathSegmentTransmittance;
 }
 
+float IsotropicMediumSlabEnvDirectionalAlbedoLUT(float MFP, float BaseColor, float NoV)
+{
+	// LUT content
+	//		View angle	from 0 to 7/8	with steps of 1/8 => 8  texels
+	//		Light angle	none
+	//		BaseColor	from 0 to 1		with steps of 0.1 => 11 texels
+	//		MFP			from 0 to 10	with steps of 1   => 11 texels
+	// 
+	// Indexing  LUTINDEX(V, L, B, M) (V * DimL * DimBC * DimMFP + L * DimBC * DimMFP  + B * DimMFP  + M)
+
+	// LUT dimension for input parameters: MFP, BaseColor, NoL, NoV
+	const float3 LUTDim = float3(11.0f, 11.0, 8.0f);
+	const float3 LUTDimInv = 1.0f / LUTDim;
+	const float3 LUTMin = 0.5f * LUTDimInv;
+	const float3 LUTMax = (LUTDim - 0.5f) * LUTDimInv;
+	const float3 LUTLen = (LUTDim - 1.0f) * LUTDimInv;
+
+	const float PIOver2 = PI / 2.0f;
+	const float ViewAngleNorm = saturate(acosFast(saturate(NoV)) / PIOver2);
+
+	float3 UVs = float3(MFP / 10.0f, BaseColor, ViewAngleNorm);	// Input parameters in [0,1]
+	UVs = LUTMin + LUTLen * UVs;								// Correct in order to ignore texel border
+
+	float EnvDirectionalAlbedo = View.SimpleVolumeEnvTexture.SampleLevel(View.SimpleVolumeEnvTextureSampler, UVs, 0);
+
+	const float Fade = IsotropicMediumSlabDirectionalAlbedoFadeLUT(MFP);
+	return EnvDirectionalAlbedo * Fade;
+}
+
+float3 IsotropicMediumSlabEnvDirectionalAlbedoLUT(FParticipatingMedia PM, float NoV)
+{
+	return float3(
+		IsotropicMediumSlabEnvDirectionalAlbedoLUT(PM.MeanFreePath.r, PM.BaseColor.r, NoV),
+		IsotropicMediumSlabEnvDirectionalAlbedoLUT(PM.MeanFreePath.g, PM.BaseColor.g, NoV),
+		IsotropicMediumSlabEnvDirectionalAlbedoLUT(PM.MeanFreePath.b, PM.BaseColor.b, NoV)
+	);
+}
diff --git a/Engine/Shaders/Private/Strata/StrataEvaluation.ush b/Engine/Shaders/Private/Strata/StrataEvaluation.ush
--- a/Engine/Shaders/Private/Strata/StrataEvaluation.ush
+++ b/Engine/Shaders/Private/Strata/StrataEvaluation.ush
@@ -890,7 +890,12 @@ FStrataEvaluateResult StrataEvaluateBSDFCommon(FStrataBSDFContext BSDFContext, F
 			{
 				FParticipatingMedia PM = StrataSlabCreateParticipatingMedia(DiffuseColor, SLAB_SSSMFP(BSDFContext.BSDF));
 				const float DiffuseToVolumeBlend = StrataSlabDiffuseToVolumeBlend(PM);
+
+			#if USE_NEW_SIMPLEVOLUME
+				const float3 SlabDirectionalAlbedo = IsotropicMediumSlabPunctualDirectionalAlbedoLUT(PM, BSDFContext.Context.NoV, BSDFContext.Context.NoL)
+			#else
 				const float3 SlabDirectionalAlbedo = IsotropicMediumSlabPunctualDirectionalAlbedo(PM);
+			#endif
 
 				// Shading point <-> View throughput
 				const float3 SlabTransmittanceV = IsotropicMediumSlabTransmittance(PM, STRATA_SIMPLEVOLUME_THICKNESS_M, BSDFContext.Context.NoV);
@@ -1439,6 +1444,9 @@ FStrataEnvLightResult StrataEvaluateForEnvLight(FStrataBSDFContext BSDFContext,
 				FParticipatingMedia PM = StrataSlabCreateParticipatingMedia(DiffuseColor, SLAB_SSSMFP(BSDFContext.BSDF));
 				const float DiffuseToVolumeBlend = StrataSlabDiffuseToVolumeBlend(PM);
 
+			#if USE_NEW_SIMPLEVOLUME
+				float3 SlabDirectionalAlbedo = IsotropicMediumSlabEnvDirectionalAlbedoLUT(PM, BSDFContext.Context.NoV);
+			#else
 				// The environment response of has been measured for an hemisphere having a uniform radiance of 1. 
 				// This DirectionalAlbedo also contains the integral over the hemisphere according to the phase function.
 				// The spherical harmonic used for environment lighting contains the diffuse integral over the hemisphere as DiffLambert = int_{\omega} {1 * NoL d\Omega} = PI.
@@ -1447,6 +1455,7 @@ FStrataEnvLightResult StrataEvaluateForEnvLight(FStrataBSDFContext BSDFContext,
 				// So, as an approximation, we recover the uniform hemisphere luminance as UniformEnvL = DiffLambert * (1 / PI) * (2 * PI).
 				// We consider that the diffuse integration of the environment SH over the hemisphere multiplied with the light response will be a good approximation similar to the "split sum integral".
 				float3 SlabDirectionalAlbedo = IsotropicMediumSlabEnvDirectionalAlbedo(PM) * ((1 / PI)* (2 * PI));
+			#endif
 
 			#if 1
 				// Because the default EnvBRDF does not account for multiple scattering, we rescale the participating media response according to EnvBRDF.
diff --git a/Engine/Source/Runtime/Engine/Classes/Engine/Engine.h b/Engine/Source/Runtime/Engine/Classes/Engine/Engine.h
--- a/Engine/Source/Runtime/Engine/Classes/Engine/Engine.h
+++ b/Engine/Source/Runtime/Engine/Classes/Engine/Engine.h
@@ -1368,6 +1368,22 @@ public:
 	/** Path of the glint BSDF texture 2 */
 	UPROPERTY(globalconfig)
 	FSoftObjectPath GlintTexture2Name;
+	
+	/** Simple volume LUT texture */
+	UPROPERTY()
+	TObjectPtr<class UVolumeTexture> SimpleVolumeTexture;
+
+	/** Path of the simple volume LUT texture */
+	UPROPERTY(globalconfig)
+	FSoftObjectPath SimpleVolumeTextureName;
+	
+	/** Simple volume environment LUT texture */
+	UPROPERTY()
+	TObjectPtr<class UVolumeTexture> SimpleVolumeEnvTexture;
+
+	/** Path of the simple volume environment LUT texture */
+	UPROPERTY(globalconfig)
+	FSoftObjectPath SimpleVolumeEnvTextureName;
 
 	/** Texture used to do font rendering in shaders */
 	UPROPERTY()
@@ -1914,6 +1930,12 @@ public:
 	*/
 	ENGINE_API void LoadGlintTextures();
 
+
+	/** Delay loading the SimpleVolume texture until it is needed by the renderer.
+	* This texture is not going to be streamed to be available right away.
+	*/
+	void LoadSimpleVolumeTextures();
+
 private:
 	#if WITH_DYNAMIC_RESOLUTION
 		/** Last dynamic resolution event. */
diff --git a/Engine/Source/Runtime/Engine/Private/SceneManagement.cpp b/Engine/Source/Runtime/Engine/Private/SceneManagement.cpp
--- a/Engine/Source/Runtime/Engine/Private/SceneManagement.cpp
+++ b/Engine/Source/Runtime/Engine/Private/SceneManagement.cpp
@@ -1006,6 +1006,12 @@ FViewUniformShaderParameters::FViewUniformShaderParameters()
 	GlintTexture = GBlackArrayTexture->TextureRHI;
 	GlintSampler = TStaticSamplerState<SF_Trilinear, AM_Clamp, AM_Clamp, AM_Clamp>::GetRHI();
 
+	// SimpleVolume
+	SimpleVolumeTexture = GBlackVolumeTexture->TextureRHI;
+	SimpleVolumeTextureSampler = TStaticSamplerState<SF_Trilinear, AM_Clamp, AM_Clamp, AM_Clamp>::GetRHI();
+	SimpleVolumeEnvTexture = GBlackVolumeTexture->TextureRHI;
+	SimpleVolumeEnvTextureSampler = TStaticSamplerState<SF_Trilinear, AM_Clamp, AM_Clamp, AM_Clamp>::GetRHI();
+
 	// Rect light atlas
 	RectLightAtlasMaxMipLevel = 1;
 	RectLightAtlasSizeAndInvSize = FVector4f(1, 1, 1, 1);
diff --git a/Engine/Source/Runtime/Engine/Private/UnrealEngine.cpp b/Engine/Source/Runtime/Engine/Private/UnrealEngine.cpp
--- a/Engine/Source/Runtime/Engine/Private/UnrealEngine.cpp
+++ b/Engine/Source/Runtime/Engine/Private/UnrealEngine.cpp
@@ -3328,6 +3328,18 @@ void UEngine::LoadGlintTextures()
 	}
 }
 
+void UEngine::LoadSimpleVolumeTextures()
+{
+	if (SimpleVolumeTexture == nullptr && SimpleVolumeTextureName.IsValid())
+	{
+		LoadEngineTexture(SimpleVolumeTexture, *SimpleVolumeTextureName.ToString());
+	}
+	if (SimpleVolumeEnvTexture == nullptr && SimpleVolumeEnvTextureName.IsValid())
+	{
+		LoadEngineTexture(SimpleVolumeEnvTexture, *SimpleVolumeEnvTextureName.ToString());
+	}
+}
+
 void UEngine::InitializePortalServices()
 {
 	IMessagingRpcModule* MessagingRpcModule = nullptr;
diff --git a/Engine/Source/Runtime/Engine/Public/SceneView.h b/Engine/Source/Runtime/Engine/Public/SceneView.h
--- a/Engine/Source/Runtime/Engine/Public/SceneView.h
+++ b/Engine/Source/Runtime/Engine/Public/SceneView.h
@@ -1014,6 +1014,11 @@ BEGIN_GLOBAL_SHADER_PARAMETER_STRUCT_WITH_CONSTRUCTOR(FViewUniformShaderParamete
 	// Glints
 	SHADER_PARAMETER_TEXTURE(Texture2DArray<float4>, GlintTexture)
 	SHADER_PARAMETER_SAMPLER(SamplerState, GlintSampler)
+	// Simple volume texture
+	SHADER_PARAMETER_TEXTURE(Texture3D<float>, SimpleVolumeTexture)
+	SHADER_PARAMETER_SAMPLER(SamplerState, SimpleVolumeTextureSampler)
+	SHADER_PARAMETER_TEXTURE(Texture3D<float>, SimpleVolumeEnvTexture)
+	SHADER_PARAMETER_SAMPLER(SamplerState, SimpleVolumeEnvTextureSampler)
 	// SSS
 	SHADER_PARAMETER_TEXTURE(Texture2D, SSProfilesTexture)
 	SHADER_PARAMETER_SAMPLER(SamplerState, SSProfilesSampler)
diff --git a/Engine/Source/Runtime/Renderer/Private/SceneRendering.cpp b/Engine/Source/Runtime/Renderer/Private/SceneRendering.cpp
--- a/Engine/Source/Runtime/Renderer/Private/SceneRendering.cpp
+++ b/Engine/Source/Runtime/Renderer/Private/SceneRendering.cpp
@@ -90,6 +90,7 @@
 #include "VolumetricFog.h"
 #include "PrimitiveSceneShaderData.h"
 #include "Engine/SpecularProfile.h"
+#include "Engine/VolumeTexture.h"
 
 /*-----------------------------------------------------------------------------
 	Globals
@@ -1975,6 +1976,29 @@ void FViewInfo::SetupUniformBufferParameters(
 	}
 	ViewUniformShaderParameters.GlintTexture = OrBlack2DArrayIfNull(ViewUniformShaderParameters.GlintTexture);
 
+	ViewUniformShaderParameters.SimpleVolumeTextureSampler = TStaticSamplerState<SF_Trilinear, AM_Clamp, AM_Clamp, AM_Clamp>::GetRHI();
+	ViewUniformShaderParameters.SimpleVolumeEnvTextureSampler = TStaticSamplerState<SF_Trilinear, AM_Clamp, AM_Clamp, AM_Clamp>::GetRHI();
+
+	if (GEngine->SimpleVolumeTexture)
+	{
+		const FTextureResource* SimpleVolumeTextureResources = GEngine->SimpleVolumeTexture->GetResource();
+		if (SimpleVolumeTextureResources)
+		{
+			ViewUniformShaderParameters.SimpleVolumeTexture = SimpleVolumeTextureResources->TextureRHI->GetTexture3D();
+		}
+	}
+	ViewUniformShaderParameters.SimpleVolumeTexture = OrBlack3DIfNull(ViewUniformShaderParameters.SimpleVolumeTexture);
+
+	if (GEngine->SimpleVolumeEnvTexture)
+	{
+		const FTextureResource* SimpleVolumeEnvTextureResources = GEngine->SimpleVolumeEnvTexture->GetResource();
+		if (SimpleVolumeEnvTextureResources)
+		{
+			ViewUniformShaderParameters.SimpleVolumeEnvTexture = SimpleVolumeEnvTextureResources->TextureRHI->GetTexture3D();
+		}
+	}
+	ViewUniformShaderParameters.SimpleVolumeEnvTexture = OrBlack3DIfNull(ViewUniformShaderParameters.SimpleVolumeEnvTexture);
+
 	// Water global resources
 	if (WaterDataBuffer.IsValid() && WaterIndirectionBuffer.IsValid())
 	{
@@ -2539,6 +2563,11 @@ FSceneRenderer::FSceneRenderer(const FSceneViewFamily* InViewFamily, FHitProxyCo
 			GEngine->LoadGlintTextures();
 		}
 
+		if (Strata::IsStrataEnabled())
+		{
+			GEngine->LoadSimpleVolumeTextures();
+		}
+
 		// Handle the FFT bloom kernel textire
 		if (ViewInfo->FinalPostProcessSettings.BloomMethod == EBloomMethod::BM_FFT && ViewInfo->ViewState != nullptr)
 		{
