Support sampling Lumen Translucency in RayTracingPrimaryRaysRGS
- improves indirect lighting in Raytraced Translucency.

- ideally this would be done in the GetSkyLighting(...) of RayTracingMaterialHitShaders.usf similar to the logic in BasePassPixelShader.usf, however we are not binding the necessary resources in raytracing hitgroups.
- instead sampling Lumen Translucency Volume for each hit in the RGS.

diff --git a/Engine/Shaders/Private/RayTracing/RayTracingMaterialHitShaders.usf b/Engine/Shaders/Private/RayTracing/RayTracingMaterialHitShaders.usf
--- a/Engine/Shaders/Private/RayTracing/RayTracingMaterialHitShaders.usf
+++ b/Engine/Shaders/Private/RayTracing/RayTracingMaterialHitShaders.usf
@@ -54,6 +54,10 @@ void GetSkyLighting(
 	OutDiffuseLighting = float3(0,0,0);
 	OutSubsurfaceLighting = float3(0,0,0);
 
+	// TODO:
+	// Should apply Lumen Dynamic GI + shadowed Skylight based on IsLumenTranslucencyGIEnabled() here
+	// however we are not binding the unnecessary shader parameters in hitgroups at the moment
+
 #if ENABLE_SKY_LIGHT
 
 	// Check if the Sky Light contribution should be ignored
diff --git a/Engine/Shaders/Private/RayTracing/RayTracingPrimaryRays.usf b/Engine/Shaders/Private/RayTracing/RayTracingPrimaryRays.usf
--- a/Engine/Shaders/Private/RayTracing/RayTracingPrimaryRays.usf
+++ b/Engine/Shaders/Private/RayTracing/RayTracingPrimaryRays.usf
@@ -56,6 +56,11 @@ RWTexture2D<float> RayHitDistanceOutput;
 
 #include "RayTracingLightingCommon.ush"
 
+#define FrontLayerTranslucencyReflectionsStruct LumenGIVolumeStruct
+#define RadianceCacheInterpolation LumenGIVolumeStruct
+
+#include "../Lumen/LumenTranslucencyVolumeShared.ush"
+
 float CalcNoT(float CosTheta1, float N1, float N2)
 {
 	float SinTheta1_Squared = 1.0 - CosTheta1 * CosTheta1;
@@ -140,7 +145,11 @@ RAY_TRACING_ENTRY_RAYGEN(RayTracingPrimaryRaysRGS)
 		const uint RefractionInstanceInclusionMask = (!bPrimaryView && RefractionRayIndex == 0) || !TranslucencyRefraction ? RAY_TRACING_MASK_TRANSLUCENT : RAY_TRACING_MASK_ALL;
 		const bool bRefractionRayTraceSkyLightContribution = false;
 		const bool bRefractionDecoupleSampleGeneration = true;
+#if PROJECT_SUPPORTS_LUMEN
+		const bool bRefractionEnableSkyLightContribution = !IsLumenTranslucencyGIEnabled();
+#else
 		const bool bRefractionEnableSkyLightContribution = true;
+#endif
 		float3 PathVertexRadiance = float3(0, 0, 0);
 
 		RandomSequence RandSequence;
@@ -283,6 +292,39 @@ RAY_TRACING_ENTRY_RAYGEN(RayTracingPrimaryRaysRGS)
 			PathRadiance += PathThroughput * vertexRadianceWeight * ViewFogInScatteringAndTransmittance.a * ReflectionThroughput * ReflectionRadiance;
 		}
 
+		{
+			float3 VolumeLighting = 0;
+
+			const float3 TranslucencyEvaluationTranslatedPosition = HitPoint;
+
+#if PROJECT_SUPPORTS_LUMEN
+			if (IsLumenTranslucencyGIEnabled() && ShouldDoEmissiveAndIndirectLighting)
+			{
+				// Lumen Dynamic GI + shadowed Skylight
+				FLWCVector3 TranslucencyEvaluationPosition = LWCSubtract(TranslucencyEvaluationTranslatedPosition, PrimaryView.PreViewTranslation);
+
+				FTwoBandSHVectorRGB TranslucencyGISH = GetTranslucencyGIVolumeLighting(TranslucencyEvaluationPosition, PrimaryView.WorldToClip, true);
+
+				// Diffuse convolution
+				FTwoBandSHVector DiffuseTransferSH = CalcDiffuseTransferSH(Payload.WorldNormal, 1);
+				VolumeLighting.rgb += max(half3(0, 0, 0), DotSH(TranslucencyGISH, DiffuseTransferSH)) / PI;
+
+				const bool bEvaluateBackface = GetShadingModelRequiresBackfaceLighting(Payload.ShadingModelID);
+				if (bEvaluateBackface)
+				{
+					FTwoBandSHVector SubsurfaceTransferSH = CalcDiffuseTransferSH(-Payload.WorldNormal, 1);
+					VolumeLighting.rgb += max(half3(0, 0, 0), DotSH(TranslucencyGISH, SubsurfaceTransferSH)) / PI;
+				}
+			}
+#endif
+			else
+			{
+				// noop => if lumen is not enabled: bRefractionEnableSkyLightContribution == true
+			}
+
+			PathRadiance += PathThroughput * vertexRadianceWeight * VolumeLighting * Payload.DiffuseColor;
+		}
+
 		// Update the refraction/transparency path transmittance and check stop condition
 		switch (Payload.BlendingMode)
 		{
diff --git a/Engine/Source/Runtime/Renderer/Private/RayTracing/RayTracingPrimaryRays.cpp b/Engine/Source/Runtime/Renderer/Private/RayTracing/RayTracingPrimaryRays.cpp
--- a/Engine/Source/Runtime/Renderer/Private/RayTracing/RayTracingPrimaryRays.cpp
+++ b/Engine/Source/Runtime/Renderer/Private/RayTracing/RayTracingPrimaryRays.cpp
@@ -51,6 +51,8 @@ class FRayTracingPrimaryRaysRGS : public FGlobalShader
 		SHADER_PARAMETER_RDG_UNIFORM_BUFFER(FReflectionUniformParameters, ReflectionStruct)
 		SHADER_PARAMETER_RDG_UNIFORM_BUFFER(FFogUniformParameters, FogUniformParameters)
 
+		SHADER_PARAMETER_RDG_UNIFORM_BUFFER(FLumenTranslucencyLightingUniforms, LumenGIVolumeStruct)
+
 		SHADER_PARAMETER_STRUCT_INCLUDE(FSceneTextureParameters, SceneTextures)
 
 		SHADER_PARAMETER_RDG_TEXTURE(Texture2D, SceneColorTexture)
@@ -131,9 +133,9 @@ void FDeferredShadingSceneRenderer::RenderRayTracingPrimaryRaysView(
 		}
 	}
 
-	FRayTracingPrimaryRaysRGS::FParameters* PassParameters = GraphBuilder.AllocParameters<FRayTracingPrimaryRaysRGS::FParameters>();
+	const FRayTracingPrimaryRaysOptions TranslucencyOptions = GetRayTracingTranslucencyOptions(View);
 
-	FRayTracingPrimaryRaysOptions TranslucencyOptions = GetRayTracingTranslucencyOptions(View);
+	FRayTracingPrimaryRaysRGS::FParameters* PassParameters = GraphBuilder.AllocParameters<FRayTracingPrimaryRaysRGS::FParameters>();
 	PassParameters->SamplesPerPixel = SamplePerPixel;
 	PassParameters->MaxRefractionRays = TranslucencyOptions.MaxRefractionRays > -1 ? TranslucencyOptions.MaxRefractionRays : View.FinalPostProcessSettings.RayTracingTranslucencyRefractionRays;
 	PassParameters->HeightFog = HeightFog;
@@ -152,6 +154,10 @@ void FDeferredShadingSceneRenderer::RenderRayTracingPrimaryRaysView(
 	PassParameters->ViewUniformBuffer = View.ViewUniformBuffer;
 	PassParameters->LightDataPacked = View.RayTracingLightDataUniformBuffer;
 
+	auto* LumenUniforms = GraphBuilder.AllocParameters<FLumenTranslucencyLightingUniforms>();
+	LumenUniforms->Parameters = GetLumenTranslucencyLightingParameters(GraphBuilder, View.GetLumenTranslucencyGIVolume(), View.LumenFrontLayerTranslucency);
+	PassParameters->LumenGIVolumeStruct = GraphBuilder.CreateUniformBuffer(LumenUniforms);
+
 	PassParameters->SceneTextures = SceneTextureParameters;
 
 	PassParameters->SceneColorTexture = SceneTextures.Color.Resolve;
