Fix invalid cluster bounding computation with groom having LOD datas.

diff --git a/Engine/Plugins/Runtime/HairStrands/Source/HairStrandsCore/Private/HairStrandsInterpolation.cpp b/Engine/Plugins/Runtime/HairStrands/Source/HairStrandsCore/Private/HairStrandsInterpolation.cpp
--- a/Engine/Plugins/Runtime/HairStrands/Source/HairStrandsCore/Private/HairStrandsInterpolation.cpp
+++ b/Engine/Plugins/Runtime/HairStrands/Source/HairStrandsCore/Private/HairStrandsInterpolation.cpp
@@ -827,12 +827,14 @@ class FHairClusterAABBCS : public FGlobalShader
 
 	BEGIN_SHADER_PARAMETER_STRUCT(FParameters, )
 		SHADER_PARAMETER_STRUCT_INCLUDE(ShaderPrint::FShaderParameters, ShaderDrawParameters)
+		SHADER_PARAMETER(float, LODIndex)
 		SHADER_PARAMETER(uint32, ClusterCount)
 		SHADER_PARAMETER(uint32, CurveCount)
 		SHADER_PARAMETER(FVector3f, CPUBoundMin)
 		SHADER_PARAMETER(FVector3f, CPUBoundMax)
 		SHADER_PARAMETER(FMatrix44f, LocalToTranslatedWorldMatrix)
 		SHADER_PARAMETER_RDG_BUFFER_SRV(Buffer<uint>, RenCurveBuffer)
+		SHADER_PARAMETER_RDG_BUFFER_SRV(Buffer<uint>, RenPointLODBuffer)
 		SHADER_PARAMETER_RDG_BUFFER_SRV(Buffer, RenderDeformedPositionBuffer)
 		SHADER_PARAMETER_RDG_BUFFER_SRV(Buffer, RenderDeformedOffsetBuffer)
 		SHADER_PARAMETER_RDG_BUFFER_UAV(RWBuffer, OutClusterAABBBuffer)
@@ -880,6 +882,7 @@ static void AddHairClusterAABBPass(
 	Instance->HairGroupPublicData->bClusterAABBValid = UpdateType == EHairAABBUpdateType::UpdateClusterAABB;
 
 	FHairClusterAABBCS::FParameters* Parameters = GraphBuilder.AllocParameters<FHairClusterAABBCS::FParameters>();
+	Parameters->LODIndex = ClusterData ? ClusterData->LODIndex : 1;;
 	Parameters->CPUBoundMin = (FVector3f)TransformedBounds.GetBox().Min;
 	Parameters->CPUBoundMax = (FVector3f)TransformedBounds.GetBox().Max;
 	Parameters->LocalToTranslatedWorldMatrix = FMatrix44f(InRenLocalToTranslatedWorld.ToMatrixWithScale());
@@ -888,6 +891,7 @@ static void AddHairClusterAABBPass(
 	Parameters->CurveCount = ActiveCurveCount;
 	Parameters->ClusterCount = ClusterData ? ClusterData->ClusterCount : 1;
 	Parameters->RenCurveBuffer = RegisterAsSRV(GraphBuilder, Instance->Strands.RestResource->CurveBuffer);
+	Parameters->RenPointLODBuffer = ClusterData ? RegisterAsSRV(GraphBuilder, *ClusterData->PointLODBuffer) : nullptr;
 	Parameters->OutClusterAABBBuffer = ClusterAABBData.ClusterAABBBuffer.UAV;
 	Parameters->OutGroupAABBBuffer = ClusterAABBData.GroupAABBBuffer.UAV;
 
diff --git a/Engine/Shaders/Private/HairStrands/HairStrandsClusterCommon.ush b/Engine/Shaders/Private/HairStrands/HairStrandsClusterCommon.ush
--- a/Engine/Shaders/Private/HairStrands/HairStrandsClusterCommon.ush
+++ b/Engine/Shaders/Private/HairStrands/HairStrandsClusterCommon.ush
@@ -1,6 +1,7 @@
 // Copyright Epic Games, Inc. All Rights Reserved.
 
 #include "../Common.ush"
+#include "/Engine/Shared/HairStrandsDefinitions.h"
 
 ////////////////////////////////////////////////////////////////////////////////
 // Defines
@@ -216,3 +217,31 @@ float ComputeBoundsScreenSize(float3 InSphereOrigin, float InSphereRadius, float
 	// For clarity, we end up comparing the diameter
 	return ScreenRadius * 2.0f;
 }
+
+///////////////////////////////////////////////////////////////////////////////////////////////////
+// Point LOD
+
+uint UnpackPointLOD(uint InPacked, uint InIndex)
+{
+	return (InPacked >> (InIndex * HAIR_POINT_LOD_BIT_COUNT)) & 0xF;
+}
+
+uint GetHairControlPointMinLOD(uint InPointIndex, Buffer<uint> InPointLODBuffer)
+{
+	const uint BlockIt      = InPointIndex >> HAIR_POINT_LOD_COUNT_PER_UINT_DIV_AS_SHIFT;
+	const uint LocalPointIt = InPointIndex - BlockIt * HAIR_POINT_LOD_COUNT_PER_UINT;
+	const uint Packed 		= InPointLODBuffer[BlockIt];
+	const uint MinLOD 		= UnpackPointLOD(Packed, LocalPointIt);
+	return MinLOD;
+}
+
+bool IsHairControlPointActive(uint InMinLOD, float InLODIndex)
+{
+	return floor(InLODIndex) <= InMinLOD;
+}
+
+bool IsHairControlPointActive(uint InPointIndex, Buffer<uint> InPointLODBuffer, float InLODIndex)
+{
+	const uint MinLOD = GetHairControlPointMinLOD(InPointIndex, InPointLODBuffer);
+	return IsHairControlPointActive(MinLOD, InLODIndex);
+}
\ No newline at end of file
diff --git a/Engine/Shaders/Private/HairStrands/HairStrandsClusterCulling.usf b/Engine/Shaders/Private/HairStrands/HairStrandsClusterCulling.usf
--- a/Engine/Shaders/Private/HairStrands/HairStrandsClusterCulling.usf
+++ b/Engine/Shaders/Private/HairStrands/HairStrandsClusterCulling.usf
@@ -321,11 +321,6 @@ RWBuffer<uint>  OutCulledCurveBuffer;
 #define MAX_CURVE CURVE_PER_GROUP
 #define MAX_POINT PERMUTATION_POINT_PER_CURVE
 
-uint UnpackPointLOD(uint InPacked, uint InIndex)
-{
-	return (InPacked >> (InIndex * HAIR_POINT_LOD_BIT_COUNT)) & 0xF;
-}
-
 uint GetGlobalCurveIndex(uint ThreadIndex, uint2 GroupIndex)
 {
 	const uint LinearGroupIndex = GroupIndex.x + GroupIndex.y * DispatchCountX;
@@ -403,17 +398,14 @@ void Main(uint3 DispatchThreadId : SV_DispatchThreadID, uint ThreadIndex : SV_Gr
 		{	
 			if (CurrentPointIndex < CurvePointCount[CurveIndex])
 			{
-				const uint PointIndex   = CurrentPointIndex + CurvePointIndex[CurveIndex];
-				const uint BlockIt      = PointIndex >> HAIR_POINT_LOD_COUNT_PER_UINT_DIV_AS_SHIFT;
-				const uint LocalPointIt = PointIndex - BlockIt * HAIR_POINT_LOD_COUNT_PER_UINT;
-				const uint Packed 		= PointLODBuffer[BlockIt];
-				const uint MinLOD 		= UnpackPointLOD(Packed, LocalPointIt);
+				const uint PointIndex = CurrentPointIndex + CurvePointIndex[CurveIndex];				
+				const uint MinLOD = GetHairControlPointMinLOD(PointIndex, PointLODBuffer);
 
 				// When a curve reach its last LOD, fade its radius progressively
 				RadiusFade = ceil(ClusterLOD[CurveIndex]) >= LODCount[CurveIndex]-1 ? 1.f : saturate(MinLOD+1 - ClusterLOD[CurveIndex]);
 
 				// Prune control point based on their MinLOD value
-				if (floor(ClusterLOD[CurveIndex]) <= MinLOD)
+				if (IsHairControlPointActive(MinLOD, ClusterLOD[CurveIndex]))
 				{
 					ValidPoints[CurveIndex][CurrentPointIndex] = 1;
 				}
@@ -537,17 +529,14 @@ void Main(uint3 DispatchThreadId : SV_DispatchThreadID, uint ThreadIndex : SV_Gr
 			{	
 				if (CurrentPointIndex < CurvePointCount)
 				{
-					const uint PointIndex   = CurrentPointIndex + CurvePointIndex;
-					const uint BlockIt      = PointIndex >> HAIR_POINT_LOD_COUNT_PER_UINT_DIV_AS_SHIFT;
-					const uint LocalPointIt = PointIndex - BlockIt * HAIR_POINT_LOD_COUNT_PER_UINT;
-					const uint Packed 		= PointLODBuffer[BlockIt];
-					const uint MinLOD 		= UnpackPointLOD(Packed, LocalPointIt);
+					const uint PointIndex = CurrentPointIndex + CurvePointIndex;
+					const uint MinLOD = GetHairControlPointMinLOD(PointIndex, PointLODBuffer);
 
 					// When a curve reach its last LOD, fade its radius progressively
 					RadiusFade = ceil(ClusterLOD) >= LODCount-1 ? 1.f : saturate(MinLOD+1 - ClusterLOD);
 
 					// Prune control point based on their MinLOD value
-					if (floor(ClusterLOD) <= MinLOD)
+					if (IsHairControlPointActive(MinLOD, ClusterLOD))
 					{
 						ValidPoints[CurrentPointIndex] = 1;
 					}
diff --git a/Engine/Shaders/Private/HairStrands/HairStrandsInterpolation.usf b/Engine/Shaders/Private/HairStrands/HairStrandsInterpolation.usf
--- a/Engine/Shaders/Private/HairStrands/HairStrandsInterpolation.usf
+++ b/Engine/Shaders/Private/HairStrands/HairStrandsInterpolation.usf
@@ -594,9 +594,11 @@ uint CurveCount;
 uint ClusterCount;
 float3 CPUBoundMin;
 float3 CPUBoundMax;
+float LODIndex;
 float4x4 LocalToTranslatedWorldMatrix;
 
 Buffer<uint> RenCurveBuffer;
+Buffer<uint> RenPointLODBuffer;
 Buffer<uint4> RenderDeformedPositionBuffer;
 Buffer<float4> RenderDeformedOffsetBuffer;
 
@@ -645,7 +647,6 @@ void ClusterAABBEvaluationCS(uint3 DispatchThreadId : SV_DispatchThreadID, uint3
 	int3 ClusterAABBMin = 9999999.0f;
 	int3 ClusterAABBMax =-9999999.0f;
 
-	// TODO culled point have invalid position
 	const uint LoopCount = DivideAndRoundUp(CurvePointCount, GROUP_SIZE);
 	for (uint LoopIt = 0; LoopIt < LoopCount; ++LoopIt)
 	{
@@ -653,9 +654,9 @@ void ClusterAABBEvaluationCS(uint3 DispatchThreadId : SV_DispatchThreadID, uint3
 		if (LocalPointIndex < CurvePointCount)
 		{
 			const uint GlobalPointIndex = CurvePointIndex + LocalPointIndex;
+			const bool bIsActive = IsHairControlPointActive(GlobalPointIndex, RenPointLODBuffer, LODIndex);
 			const FHairControlPoint CP = UnpackHairControlPoint(RenderDeformedPositionBuffer[GlobalPointIndex], OutHairPositionOffset, 1, 1, 1);
-
-			if (all(IsFinite(CP.Position)))
+			if (bIsActive && all(IsFinite(CP.Position)))
 			{
 				const int3 WorldSpaceCentimeters = int3(mul(float4(CP.Position, 1.0f), LocalToTranslatedWorldMatrix).xyz);
 				const int3 WorldSpaceCentimetersMin = WorldSpaceCentimeters - 1;
@@ -735,7 +736,7 @@ void ClusterAABBEvaluationCS(uint3 DispatchThreadId : SV_DispatchThreadID, uint3
 	#if DEBUG_ENABLE
 	if (ThreadIndex == 0)
 	{
-		const float4 DebugColor = float4(ColorMapViridis(ClusterIndex / float(TotalClusterCount)),1);
+		const float4 DebugColor = float4(ColorMapViridis(ClusterIndex / float(ClusterCount)),1);
 		AddAABBTWS(SharedClusterAABBMin[0], SharedClusterAABBMax[0], DebugColor);
 		AddReferentialWS(LocalToWorldMatrix, 10);
 	}
