Substrate - fix compilation issue

diff --git a/Engine/Shaders/Private/ParticipatingMediaCommon.ush b/Engine/Shaders/Private/ParticipatingMediaCommon.ush
--- a/Engine/Shaders/Private/ParticipatingMediaCommon.ush
+++ b/Engine/Shaders/Private/ParticipatingMediaCommon.ush
@@ -270,7 +270,7 @@ float IsotropicMediumSlabPunctualDirectionalAlbedoLUT(float MFP, float BaseColor
 	return DirectionalAlbedo * Fade;
 }
 
-float3 IsotropicMediumSlabPunctualDirectionalAlbedoLUT(FParticipatingMedia PM, float NoV, float3 NoL)
+float3 IsotropicMediumSlabPunctualDirectionalAlbedoLUT(FParticipatingMedia PM, float NoV, float NoL)
 {
 	return float3(
 		IsotropicMediumSlabPunctualDirectionalAlbedoLUT(PM.MeanFreePath.r, PM.BaseColor.r, NoV, NoL),
