Fix overflow in curve culling for curve with control point count < 64.

diff --git a/Engine/Shaders/Private/HairStrands/HairStrandsClusterCulling.usf b/Engine/Shaders/Private/HairStrands/HairStrandsClusterCulling.usf
--- a/Engine/Shaders/Private/HairStrands/HairStrandsClusterCulling.usf
+++ b/Engine/Shaders/Private/HairStrands/HairStrandsClusterCulling.usf
@@ -359,7 +359,7 @@ groupshared uint WriteOffsetN[MAX_CURVE];
 void Main(uint3 DispatchThreadId : SV_DispatchThreadID, uint ThreadIndex : SV_GroupIndex, uint2 GroupIndex : SV_GroupID)
 {
 	// 1. Read the Curve data
-	if (ThreadIndex <= CURVE_PER_GROUP)
+	if (ThreadIndex < CURVE_PER_GROUP)
 	{
 		const uint LocalCurveIndex  = ThreadIndex;
 		const uint GlobalCurveIndex = GetGlobalCurveIndex(ThreadIndex, GroupIndex);
