Move Cluster AABB computation within HairStrandsClusterCulling.usf

diff --git a/Engine/Plugins/Runtime/HairStrands/Source/HairStrandsCore/Private/HairStrandsInterpolation.cpp b/Engine/Plugins/Runtime/HairStrands/Source/HairStrandsCore/Private/HairStrandsInterpolation.cpp
--- a/Engine/Plugins/Runtime/HairStrands/Source/HairStrandsCore/Private/HairStrandsInterpolation.cpp
+++ b/Engine/Plugins/Runtime/HairStrands/Source/HairStrandsCore/Private/HairStrandsInterpolation.cpp
@@ -852,7 +852,7 @@ public:
 	}
 };
 
-IMPLEMENT_GLOBAL_SHADER(FHairClusterAABBCS, "/Engine/Private/HairStrands/HairStrandsInterpolation.usf", "ClusterAABBEvaluationCS", SF_Compute);
+IMPLEMENT_GLOBAL_SHADER(FHairClusterAABBCS, "/Engine/Private/HairStrands/HairStrandsClusterCulling.usf", "ClusterAABBEvaluationCS", SF_Compute);
 
 enum class EHairAABBUpdateType
 {
diff --git a/Engine/Shaders/Private/HairStrands/HairStrandsClusterCulling.usf b/Engine/Shaders/Private/HairStrands/HairStrandsClusterCulling.usf
--- a/Engine/Shaders/Private/HairStrands/HairStrandsClusterCulling.usf
+++ b/Engine/Shaders/Private/HairStrands/HairStrandsClusterCulling.usf
@@ -22,8 +22,6 @@
 
 // Cluster data in ClusterAABB buffers is packed as {uint3 AABBMin, uint3 AABBMax}
 
-uint ClusterCount;
-
 /////////////////////////////////////////////////////////////////////////////////////////////////////////////////////////////
 // Utility functions
 
@@ -39,6 +37,7 @@ uint3 ComputeDispatchGroupCount2D(uint GroupCount)
 /// Clear clusters AABB
 #ifdef SHADER_CLEARCLUSTERAABB
 
+uint ClusterCount;
 uint bClearClusterAABBs;
 
 RWBuffer<int> OutClusterAABBBuffer; // Cluster data packed as {uint3 AABBMin, uint3 AABBMax}
@@ -83,6 +82,7 @@ void MainClearClusterAABBCS(uint2 DispatchThreadId : SV_DispatchThreadID)
 #define MODE_LOD		3
 #define MODE_LOD_SIMPLE	4
 
+uint ClusterCount;
 uint HairGroupId;
 uint TriangleCount;
 int ClusterDebugMode;
@@ -629,3 +629,166 @@ void Main(uint DispatchThreadId : SV_DispatchThreadID, uint ThreadIndex : SV_Gro
 }
 
 #endif // SHADER_CLUSTER_CULL_ARGS
+
+///////////////////////////////////////////////////////////////////////////////////////////////////
+// Hair clusters update 
+
+#if SHADER_CLUSTERAABB
+
+uint CurveCount;
+uint ClusterCount;
+float3 CPUBoundMin;
+float3 CPUBoundMax;
+float LODIndex;
+float4x4 LocalToTranslatedWorldMatrix;
+
+Buffer<uint> RenCurveBuffer;
+Buffer<uint> RenPointLODBuffer;
+Buffer<uint4> RenderDeformedPositionBuffer;
+Buffer<float4> RenderDeformedOffsetBuffer;
+
+RWBuffer<int> OutClusterAABBBuffer; 	// Cluster data packed as {uint3 AABBMin, uint3 AABBMax}
+RWBuffer<int> OutGroupAABBBuffer; 		// Group data packed as {uint3 AABBMin, uint3 AABBMax}
+
+groupshared float3 SharedClusterAABBMin[GROUP_SIZE];
+groupshared float3 SharedClusterAABBMax[GROUP_SIZE];
+groupshared uint CurvePointIndex;
+groupshared uint CurvePointCount;
+groupshared float3 OutHairPositionOffset;
+
+[numthreads(GROUP_SIZE, 1, 1)]
+	void ClusterAABBEvaluationCS(uint3 DispatchThreadId : SV_DispatchThreadID, uint3 GroupId : SV_GroupID, uint3 GroupThreadId : SV_GroupThreadID, uint ThreadIndex : SV_GroupIndex)
+{
+#if PERMUTATION_USE_CPU_AABB
+	// Only update the full group AABB
+	if (all(DispatchThreadId == 0))
+	{
+		FHairAABB Bound;
+		Bound.Min = CPUBoundMin;
+		Bound.Max = CPUBoundMax;
+		WriteHairAABB(0, Bound, OutGroupAABBBuffer);
+	}
+#else // PERMUTATION_USE_CPU_AABB
+
+	// Use the first K curves to to compute cluster bound. This is an approximation, 
+	// but since curves are now shuffled, it covers most of the groom
+	const uint ClusterIndex = GroupId.x;
+	if (ThreadIndex == 0)
+	{
+		CurvePointIndex = 0;
+		CurvePointCount = 0;
+
+		const uint CurveIndex = GroupId.x;
+		if (CurveIndex < CurveCount)
+		{
+			const FHairCurve RenCurve = UnpackHairCurve(RenCurveBuffer[CurveIndex]);
+			CurvePointIndex = RenCurve.PointIndex;
+			CurvePointCount = RenCurve.PointCount;
+		}
+		OutHairPositionOffset = RenderDeformedOffsetBuffer[0].xyz;
+	}
+	GroupMemoryBarrierWithGroupSync();
+
+	int3 ClusterAABBMin = 9999999.0f;
+	int3 ClusterAABBMax =-9999999.0f;
+
+	const uint LoopCount = DivideAndRoundUp(CurvePointCount, GROUP_SIZE);
+	for (uint LoopIt = 0; LoopIt < LoopCount; ++LoopIt)
+	{
+		const uint LocalPointIndex = ThreadIndex + LoopIt * GROUP_SIZE;
+		if (LocalPointIndex < CurvePointCount)
+		{
+			const uint GlobalPointIndex = CurvePointIndex + LocalPointIndex;
+			const bool bIsActive = IsHairControlPointActive(GlobalPointIndex, RenPointLODBuffer, LODIndex);
+			const FHairControlPoint CP = UnpackHairControlPoint(RenderDeformedPositionBuffer[GlobalPointIndex], OutHairPositionOffset, 1, 1, 1);
+			if (bIsActive && all(IsFinite(CP.Position)))
+			{
+				const int3 WorldSpaceCentimeters = int3(mul(float4(CP.Position, 1.0f), LocalToTranslatedWorldMatrix).xyz);
+				const int3 WorldSpaceCentimetersMin = WorldSpaceCentimeters - 1;
+				const int3 WorldSpaceCentimetersMax = WorldSpaceCentimeters + 1;
+
+				ClusterAABBMin = min(ClusterAABBMin, WorldSpaceCentimetersMin);
+				ClusterAABBMax = max(ClusterAABBMax, WorldSpaceCentimetersMax);
+			}
+		}
+	}
+
+	// Write each thread result to shared memory
+	SharedClusterAABBMin[ThreadIndex] = ClusterAABBMin;
+	SharedClusterAABBMax[ThreadIndex] = ClusterAABBMax;
+
+	// Do a local reduction in shared memory. Assumes ClusterLOD0.VertexCount>64 to have all min&max values valid.
+	GroupMemoryBarrierWithGroupSync();
+	if (ThreadIndex < 32)
+	{
+		SharedClusterAABBMin[ThreadIndex] = min(SharedClusterAABBMin[ThreadIndex], SharedClusterAABBMin[ThreadIndex + 32]);
+		SharedClusterAABBMax[ThreadIndex] = max(SharedClusterAABBMax[ThreadIndex], SharedClusterAABBMax[ThreadIndex + 32]);
+	}
+	GroupMemoryBarrierWithGroupSync();
+	if (ThreadIndex < 16)
+	{
+		SharedClusterAABBMin[ThreadIndex] = min(SharedClusterAABBMin[ThreadIndex], SharedClusterAABBMin[ThreadIndex + 16]);
+		SharedClusterAABBMax[ThreadIndex] = max(SharedClusterAABBMax[ThreadIndex], SharedClusterAABBMax[ThreadIndex + 16]);
+	}
+	GroupMemoryBarrierWithGroupSync();
+	if (ThreadIndex < 8)
+	{
+		SharedClusterAABBMin[ThreadIndex] = min(SharedClusterAABBMin[ThreadIndex], SharedClusterAABBMin[ThreadIndex + 8]);
+		SharedClusterAABBMax[ThreadIndex] = max(SharedClusterAABBMax[ThreadIndex], SharedClusterAABBMax[ThreadIndex + 8]);
+	}
+	// No hardware has SIMD Vector unit operating in sync with less than 16 threads per group. So can skip sync now.
+	//GroupMemoryBarrierWithGroupSync();
+	if (ThreadIndex < 4)
+	{
+		SharedClusterAABBMin[ThreadIndex] = min(SharedClusterAABBMin[ThreadIndex], SharedClusterAABBMin[ThreadIndex + 4]);
+		SharedClusterAABBMax[ThreadIndex] = max(SharedClusterAABBMax[ThreadIndex], SharedClusterAABBMax[ThreadIndex + 4]);
+	}
+	//GroupMemoryBarrierWithGroupSync();
+	if (ThreadIndex < 2)
+	{
+		SharedClusterAABBMin[ThreadIndex] = min(SharedClusterAABBMin[ThreadIndex], SharedClusterAABBMin[ThreadIndex + 2]);
+		SharedClusterAABBMax[ThreadIndex] = max(SharedClusterAABBMax[ThreadIndex], SharedClusterAABBMax[ThreadIndex + 2]);
+	}
+	//GroupMemoryBarrierWithGroupSync();
+	if (ThreadIndex < 1)
+	{
+		SharedClusterAABBMin[ThreadIndex] = min(SharedClusterAABBMin[ThreadIndex], SharedClusterAABBMin[ThreadIndex + 1]);
+		SharedClusterAABBMax[ThreadIndex] = max(SharedClusterAABBMax[ThreadIndex], SharedClusterAABBMax[ThreadIndex + 1]);
+	}
+
+	// Write out hair cluster AABB
+	if (ThreadIndex == 0)
+	{
+		const uint ClusterIdx6 = ClusterIndex * 6;
+		OutClusterAABBBuffer[ClusterIdx6 + 0] = SharedClusterAABBMin[0].x;
+		OutClusterAABBBuffer[ClusterIdx6 + 1] = SharedClusterAABBMin[0].y;
+		OutClusterAABBBuffer[ClusterIdx6 + 2] = SharedClusterAABBMin[0].z;
+		OutClusterAABBBuffer[ClusterIdx6 + 3] = SharedClusterAABBMax[0].x;
+		OutClusterAABBBuffer[ClusterIdx6 + 4] = SharedClusterAABBMax[0].y;
+		OutClusterAABBBuffer[ClusterIdx6 + 5] = SharedClusterAABBMax[0].z;
+
+		// And contribute to the group full AABB
+		FHairAABB Bound;
+		Bound.Min.x = SharedClusterAABBMin[0].x;
+		Bound.Min.y = SharedClusterAABBMin[0].y;
+		Bound.Min.z = SharedClusterAABBMin[0].z;
+		Bound.Max.x = SharedClusterAABBMax[0].x;
+		Bound.Max.y = SharedClusterAABBMax[0].y;
+		Bound.Max.z = SharedClusterAABBMax[0].z;
+		InterlockHairAABB(0, Bound, OutGroupAABBBuffer);
+	}
+
+#if DEBUG_ENABLE
+	if (ThreadIndex == 0)
+	{
+		const float4 DebugColor = float4(ColorMapViridis(ClusterIndex / float(ClusterCount)),1);
+		AddAABBTWS(SharedClusterAABBMin[0], SharedClusterAABBMax[0], DebugColor);
+		AddReferentialWS(LocalToWorldMatrix, 10);
+	}
+#endif
+
+	// Min/Max group AABB is done later by another pass
+#endif // PERMUTATION_USE_CPU_AABB
+}
+
+#endif // SHADER_CLUSTERAABB
diff --git a/Engine/Shaders/Private/HairStrands/HairStrandsInterpolation.usf b/Engine/Shaders/Private/HairStrands/HairStrandsInterpolation.usf
--- a/Engine/Shaders/Private/HairStrands/HairStrandsInterpolation.usf
+++ b/Engine/Shaders/Private/HairStrands/HairStrandsInterpolation.usf
@@ -585,169 +585,6 @@ void MainCS(uint2 DispatchThreadId : SV_DispatchThreadID, uint GroupThread1D : S
 }
 #endif // SHADER_PATCHATTRIBUTE
 
-///////////////////////////////////////////////////////////////////////////////////////////////////
-// Hair clusters update 
-
-#if SHADER_CLUSTERAABB
-
-uint CurveCount;
-uint ClusterCount;
-float3 CPUBoundMin;
-float3 CPUBoundMax;
-float LODIndex;
-float4x4 LocalToTranslatedWorldMatrix;
-
-Buffer<uint> RenCurveBuffer;
-Buffer<uint> RenPointLODBuffer;
-Buffer<uint4> RenderDeformedPositionBuffer;
-Buffer<float4> RenderDeformedOffsetBuffer;
-
-RWBuffer<int> OutClusterAABBBuffer; 	// Cluster data packed as {uint3 AABBMin, uint3 AABBMax}
-RWBuffer<int> OutGroupAABBBuffer; 		// Group data packed as {uint3 AABBMin, uint3 AABBMax}
-
-groupshared float3 SharedClusterAABBMin[GROUP_SIZE];
-groupshared float3 SharedClusterAABBMax[GROUP_SIZE];
-groupshared uint CurvePointIndex;
-groupshared uint CurvePointCount;
-groupshared float3 OutHairPositionOffset;
-
-[numthreads(GROUP_SIZE, 1, 1)]
-void ClusterAABBEvaluationCS(uint3 DispatchThreadId : SV_DispatchThreadID, uint3 GroupId : SV_GroupID, uint3 GroupThreadId : SV_GroupThreadID, uint ThreadIndex : SV_GroupIndex)
-{
-#if PERMUTATION_USE_CPU_AABB
-	// Only update the full group AABB
-	if (all(DispatchThreadId == 0))
-	{
-		FHairAABB Bound;
-		Bound.Min = CPUBoundMin;
-		Bound.Max = CPUBoundMax;
-		WriteHairAABB(0, Bound, OutGroupAABBBuffer);
-	}
-#else // PERMUTATION_USE_CPU_AABB
-
-	// Use the first K curves to to compute cluster bound. This is an approximation, 
-	// but since curves are now shuffled, it covers most of the groom
-	const uint ClusterIndex = GroupId.x;
-	if (ThreadIndex == 0)
-	{
-		CurvePointIndex = 0;
-		CurvePointCount = 0;
-
-		const uint CurveIndex = GroupId.x;
-		if (CurveIndex < CurveCount)
-		{
-			const FHairCurve RenCurve = UnpackHairCurve(RenCurveBuffer[CurveIndex]);
-			CurvePointIndex = RenCurve.PointIndex;
-			CurvePointCount = RenCurve.PointCount;
-		}
-		OutHairPositionOffset = RenderDeformedOffsetBuffer[0].xyz;
-	}
-	GroupMemoryBarrierWithGroupSync();
-
-	int3 ClusterAABBMin = 9999999.0f;
-	int3 ClusterAABBMax =-9999999.0f;
-
-	const uint LoopCount = DivideAndRoundUp(CurvePointCount, GROUP_SIZE);
-	for (uint LoopIt = 0; LoopIt < LoopCount; ++LoopIt)
-	{
-		const uint LocalPointIndex = ThreadIndex + LoopIt * GROUP_SIZE;
-		if (LocalPointIndex < CurvePointCount)
-		{
-			const uint GlobalPointIndex = CurvePointIndex + LocalPointIndex;
-			const bool bIsActive = IsHairControlPointActive(GlobalPointIndex, RenPointLODBuffer, LODIndex);
-			const FHairControlPoint CP = UnpackHairControlPoint(RenderDeformedPositionBuffer[GlobalPointIndex], OutHairPositionOffset, 1, 1, 1);
-			if (bIsActive && all(IsFinite(CP.Position)))
-			{
-				const int3 WorldSpaceCentimeters = int3(mul(float4(CP.Position, 1.0f), LocalToTranslatedWorldMatrix).xyz);
-				const int3 WorldSpaceCentimetersMin = WorldSpaceCentimeters - 1;
-				const int3 WorldSpaceCentimetersMax = WorldSpaceCentimeters + 1;
-
-				ClusterAABBMin = min(ClusterAABBMin, WorldSpaceCentimetersMin);
-				ClusterAABBMax = max(ClusterAABBMax, WorldSpaceCentimetersMax);
-			}
-		}
-	}
-
-	// Write each thread result to shared memory
-	SharedClusterAABBMin[ThreadIndex] = ClusterAABBMin;
-	SharedClusterAABBMax[ThreadIndex] = ClusterAABBMax;
-
-	// Do a local reduction in shared memory. Assumes ClusterLOD0.VertexCount>64 to have all min&max values valid.
-	GroupMemoryBarrierWithGroupSync();
-	if (ThreadIndex < 32)
-	{
-		SharedClusterAABBMin[ThreadIndex] = min(SharedClusterAABBMin[ThreadIndex], SharedClusterAABBMin[ThreadIndex + 32]);
-		SharedClusterAABBMax[ThreadIndex] = max(SharedClusterAABBMax[ThreadIndex], SharedClusterAABBMax[ThreadIndex + 32]);
-	}
-	GroupMemoryBarrierWithGroupSync();
-	if (ThreadIndex < 16)
-	{
-		SharedClusterAABBMin[ThreadIndex] = min(SharedClusterAABBMin[ThreadIndex], SharedClusterAABBMin[ThreadIndex + 16]);
-		SharedClusterAABBMax[ThreadIndex] = max(SharedClusterAABBMax[ThreadIndex], SharedClusterAABBMax[ThreadIndex + 16]);
-	}
-	GroupMemoryBarrierWithGroupSync();
-	if (ThreadIndex < 8)
-	{
-		SharedClusterAABBMin[ThreadIndex] = min(SharedClusterAABBMin[ThreadIndex], SharedClusterAABBMin[ThreadIndex + 8]);
-		SharedClusterAABBMax[ThreadIndex] = max(SharedClusterAABBMax[ThreadIndex], SharedClusterAABBMax[ThreadIndex + 8]);
-	}
-	// No hardware has SIMD Vector unit operating in sync with less than 16 threads per group. So can skip sync now.
-	//GroupMemoryBarrierWithGroupSync();
-	if (ThreadIndex < 4)
-	{
-		SharedClusterAABBMin[ThreadIndex] = min(SharedClusterAABBMin[ThreadIndex], SharedClusterAABBMin[ThreadIndex + 4]);
-		SharedClusterAABBMax[ThreadIndex] = max(SharedClusterAABBMax[ThreadIndex], SharedClusterAABBMax[ThreadIndex + 4]);
-	}
-	//GroupMemoryBarrierWithGroupSync();
-	if (ThreadIndex < 2)
-	{
-		SharedClusterAABBMin[ThreadIndex] = min(SharedClusterAABBMin[ThreadIndex], SharedClusterAABBMin[ThreadIndex + 2]);
-		SharedClusterAABBMax[ThreadIndex] = max(SharedClusterAABBMax[ThreadIndex], SharedClusterAABBMax[ThreadIndex + 2]);
-	}
-	//GroupMemoryBarrierWithGroupSync();
-	if (ThreadIndex < 1)
-	{
-		SharedClusterAABBMin[ThreadIndex] = min(SharedClusterAABBMin[ThreadIndex], SharedClusterAABBMin[ThreadIndex + 1]);
-		SharedClusterAABBMax[ThreadIndex] = max(SharedClusterAABBMax[ThreadIndex], SharedClusterAABBMax[ThreadIndex + 1]);
-	}
-
-	// Write out hair cluster AABB
-	if (ThreadIndex == 0)
-	{
-		const uint ClusterIdx6 = ClusterIndex * 6;
-		OutClusterAABBBuffer[ClusterIdx6 + 0] = SharedClusterAABBMin[0].x;
-		OutClusterAABBBuffer[ClusterIdx6 + 1] = SharedClusterAABBMin[0].y;
-		OutClusterAABBBuffer[ClusterIdx6 + 2] = SharedClusterAABBMin[0].z;
-		OutClusterAABBBuffer[ClusterIdx6 + 3] = SharedClusterAABBMax[0].x;
-		OutClusterAABBBuffer[ClusterIdx6 + 4] = SharedClusterAABBMax[0].y;
-		OutClusterAABBBuffer[ClusterIdx6 + 5] = SharedClusterAABBMax[0].z;
-
-		// And contribute to the group full AABB
-		FHairAABB Bound;
-		Bound.Min.x = SharedClusterAABBMin[0].x;
-		Bound.Min.y = SharedClusterAABBMin[0].y;
-		Bound.Min.z = SharedClusterAABBMin[0].z;
-		Bound.Max.x = SharedClusterAABBMax[0].x;
-		Bound.Max.y = SharedClusterAABBMax[0].y;
-		Bound.Max.z = SharedClusterAABBMax[0].z;
-		InterlockHairAABB(0, Bound, OutGroupAABBBuffer);
-	}
-
-	#if DEBUG_ENABLE
-	if (ThreadIndex == 0)
-	{
-		const float4 DebugColor = float4(ColorMapViridis(ClusterIndex / float(ClusterCount)),1);
-		AddAABBTWS(SharedClusterAABBMin[0], SharedClusterAABBMax[0], DebugColor);
-		AddReferentialWS(LocalToWorldMatrix, 10);
-	}
-	#endif
-
-	// Min/Max group AABB is done later by another pass
-#endif // PERMUTATION_USE_CPU_AABB
-}
-
-#endif // SHADER_CLUSTERAABB
-
 ///////////////////////////////////////////////////////////////////////////////////////////////////
 // Hair transfer prev. position
 
