Substrate - improvements and bug fixes related to BSDFOFFset

diff --git a/Engine/Shaders/Private/Strata/StrataMaterialClassification.usf b/Engine/Shaders/Private/Strata/StrataMaterialClassification.usf
--- a/Engine/Shaders/Private/Strata/StrataMaterialClassification.usf
+++ b/Engine/Shaders/Private/Strata/StrataMaterialClassification.usf
@@ -4,6 +4,7 @@
 
 #define STRATA_INLINE_SHADING 0
 #define STRATA_SSS_MATERIAL_OVERRIDE 0
+#define STRATA_COMPLEXSPECIALPATH 1
 #include "/Engine/Private/Strata/Strata.ush"
 #include "StrataTile.ush"
 
diff --git a/Engine/Shaders/Private/Strata/StrataTile.ush b/Engine/Shaders/Private/Strata/StrataTile.ush
--- a/Engine/Shaders/Private/Strata/StrataTile.ush
+++ b/Engine/Shaders/Private/Strata/StrataTile.ush
@@ -90,7 +90,9 @@ FStrataBSDFTile UnpackBSDFTile(uint In)
 // Strata BSDF offset
 
 // BSDF offsets are packed into 32bits to max. BSDF is < STRATA_MAX_BSDF_COUNT, which could leads to some issue.
-#define STRATA_MAX_BSDF_COUNT_FOR_BDSFOFFSET 10u
+#define STRATA_MAX_BSDF_COUNT_FOR_BDSFOFFSET 8u
+#define STRATA_BSDF_OFFSET_BIT_COUNT 4u
+#define STRATA_BSDF_OFFSET_BIT_MASK  0xF
 
 // Store the material slice index offset for each BSDF, for pixels containing several BSDFs
 struct FStrataBSDFOffset
@@ -101,21 +103,21 @@ struct FStrataBSDFOffset
 
 uint PackBSDFOffset(FStrataBSDFOffset In)
 {
-	// Packed as relative cumulative offsets, stored into 32 bits
-	// | 3 bits | 3 bits | 3 bits | 3 bits | ... | 3 bits |
-	//   BSDF 0   BSDF 1   BSDF 2   BSDF 3  ...   BSDF 9
+	// Packed as relative cumulative offsets of N bits, stored into 32 bits
+	// | N bits | N bits | N bits | N bits | ... | N bits |
+	//   BSDF 0   BSDF 1   BSDF 2   BSDF 3  ...   BSDF (2^N - 1)
 
 	uint Out = 0;
 	const uint BSDFCount = min(In.BSDFCount, STRATA_MAX_BSDF_COUNT_FOR_BDSFOFFSET);
 	if (BSDFCount > 0)
 	{
-		Out = In.BSDFOffsets[0] & 0x7;
+		Out = In.BSDFOffsets[0] & STRATA_BSDF_OFFSET_BIT_MASK;
 		UNROLL_N(STRATA_MAX_BSDF_COUNT_FOR_BDSFOFFSET-1)
 		for (uint i = 1; i < BSDFCount; i++)
 		{
 			const uint Curr = In.BSDFOffsets[i];
 			const uint Prev = In.BSDFOffsets[i-1];
-			Out = Out | (((Curr - Prev) & 0x7) << (3u*i));
+			Out = Out | (((Curr - Prev) & STRATA_BSDF_OFFSET_BIT_MASK) << (STRATA_BSDF_OFFSET_BIT_COUNT * i));
 		}
 	}
 	return Out;
@@ -127,12 +129,12 @@ FStrataBSDFOffset UnpackBSDFOffset(uint In, uint InBSDFCount)
 	if (InBSDFCount > 0)
 	{
 		Out.BSDFCount = min(InBSDFCount, STRATA_MAX_BSDF_COUNT_FOR_BDSFOFFSET);
-		Out.BSDFOffsets[0] = (In & 0x7);
+		Out.BSDFOffsets[0] = (In & STRATA_BSDF_OFFSET_BIT_MASK);
 
 		UNROLL_N(STRATA_MAX_BSDF_COUNT_FOR_BDSFOFFSET-1)
 		for (uint i = 1; i < Out.BSDFCount; i++)
 		{
-			Out.BSDFOffsets[i] = Out.BSDFOffsets[i-1] + ((In>>(3u*i)) & 0x7);
+			Out.BSDFOffsets[i] = Out.BSDFOffsets[i-1] + ((In>>(STRATA_BSDF_OFFSET_BIT_COUNT * i)) & STRATA_BSDF_OFFSET_BIT_MASK);
 		}
 	}
 	return Out;
@@ -147,7 +149,7 @@ uint UnpackBSDFOffsetAtIndex(uint In, uint InIndex, uint InBSDFCount)
 		UNROLL_N(STRATA_MAX_BSDF_COUNT_FOR_BDSFOFFSET)
 		for (uint i = 0; i <= InIndex; i++)
 		{
-			Out += ((In >> (3u * i)) & 0x7);
+			Out += ((In >> (STRATA_BSDF_OFFSET_BIT_COUNT * i)) & STRATA_BSDF_OFFSET_BIT_MASK);
 		}
 	}
 	return Out;
diff --git a/Engine/Shaders/Private/Strata/StrataVisualize.usf b/Engine/Shaders/Private/Strata/StrataVisualize.usf
--- a/Engine/Shaders/Private/Strata/StrataVisualize.usf
+++ b/Engine/Shaders/Private/Strata/StrataVisualize.usf
@@ -137,6 +137,10 @@ void PrintBSDF(inout FShaderPrintContext Context, inout FStrataAddressing Strata
 	const FFontColor FontBSDFPropName = FontOrange;
 	const FFontColor FontBSDFPropValu = FontWhite;
 
+	Print(Context, TEXT("Address = "), FontEmerald);
+	Print(Context, StrataAddressing.CurrentIndex, FontEmerald);
+	Newline(Context, RectMax);
+
 	const FStrataBSDF BSDF = UnpackStrataBSDFIn(Strata.MaterialTextureArray, StrataAddressing, Header);
 
 	// Draw Referential
