Substrate - fixed glint popping when using secondary roughness and  roughness0 and roughness1 are the same but glint is enabled.

diff --git a/Engine/Shaders/Private/Strata/Strata.ush b/Engine/Shaders/Private/Strata/Strata.ush
--- a/Engine/Shaders/Private/Strata/Strata.ush
+++ b/Engine/Shaders/Private/Strata/Strata.ush
@@ -1398,7 +1398,8 @@ FStrataData GetStrataSlabBSDF(
 		ThicknessCm = STRATA_SIMPLEVOLUME_THICKNESS_CM;
 	}
 
-	const uint bHasHaziness = ((Haziness.Roughness != Roughness  && Haziness.Weight > 0.0f) || Haziness.bSimpleClearCoat) ? 1 : 0;
+	const bool bHasGlint = GlintValue > 0.0f;
+	const uint bHasHaziness = (((Haziness.Roughness != Roughness || bHasGlint)  && Haziness.Weight > 0.0f) || Haziness.bSimpleClearCoat) ? 1 : 0;
 	const uint bHasAnisotropy = Anisotropy != 0 ? 1 : 0;
 	const uint bHasFuzz = FuzzAmount > 0.0f ? 1 : 0;
 	if (!STRATA_FASTPATH && bHasAnisotropy)
@@ -1431,7 +1432,7 @@ FStrataData GetStrataSlabBSDF(
 
 #if STRATA_COMPLEXSPECIALPATH
 #if STRATA_GLINTS_ENABLED
-	if (GlintValue > 0.0f)
+	if (bHasGlint)
 	{
 		BSDF_SETHASGLINT(StrataData.InlinedBSDF, 1);
 		SLAB_GLINT_VALUE(StrataData.InlinedBSDF)		= GlintValue;
