ColorManagement: Enhancing ColorSpace.ush for variable working color spaces.

#rnx

diff --git a/Engine/Shaders/Private/ColorSpace.ush b/Engine/Shaders/Private/ColorSpace.ush
--- a/Engine/Shaders/Private/ColorSpace.ush
+++ b/Engine/Shaders/Private/ColorSpace.ush
@@ -37,28 +37,47 @@ float3 LMS_2_sRGB( float3 LMS )
 // https://en.wikipedia.org/wiki/CIE_1931_color_space
 //------------------------------------------------------------------
 
-static const float3x3 XYZ_2_LinearRGB_MAT =
+static const float3x3 XYZ_2_Linear_sRGB_MAT =
 {
 	 3.2409699419, -1.5373831776, -0.4986107603,
 	-0.9692436363,  1.8759675015,  0.0415550574,
 	 0.0556300797, -0.2039769589,  1.0569715142,
 };
 
-static const float3x3 LinearRGB_2_XYZ_MAT =
+static const float3x3 Linear_sRGB_2_XYZ_MAT =
 {
-	0.4124564, 0.3575761, 0.1804375,
-	0.2126729, 0.7151522, 0.0721750,
-	0.0193339, 0.1191920, 0.9503041,
+	0.4123907993, 0.3575843394, 0.1804807884,
+	0.2126390059, 0.7151686788, 0.0721923154,
+	0.0193308187, 0.1191947798, 0.9505321522,
 };
 
 float3 LinearRGB_2_XYZ( float3 LinearRGB )
 {
-	return mul(LinearRGB_2_XYZ_MAT, LinearRGB);
+#if WORKING_COLOR_SPACE_IS_SRGB
+	return mul(Linear_sRGB_2_XYZ_MAT, LinearRGB);
+#else
+	return mul(WORKING_COLOR_SPACE_RGB_TO_XYZ_MAT, LinearRGB);
+#endif
 }
 
 float3 XYZ_2_LinearRGB( float3 XYZ )
 {
-	return mul(XYZ_2_LinearRGB_MAT, XYZ);
+#if WORKING_COLOR_SPACE_IS_SRGB
+	return mul(XYZ_2_Linear_sRGB_MAT, XYZ);
+#else
+	return mul(XYZ_TO_RGB_WORKING_COLOR_SPACE_MAT, XYZ);
+#endif
+}
+
+#if WORKING_COLOR_SPACE_IS_SRGB
+static const float3 WorkingColorSpaceLuminanceFactors = Linear_sRGB_2_XYZ_MAT._m10_m11_m12;
+#else
+static const float3 WorkingColorSpaceLuminanceFactors = WORKING_COLOR_SPACE_RGB_TO_XYZ_MAT._m10_m11_m12;
+#endif
+
+float3 Luminance( float3 LinearColor, float3 LuminanceFactors )
+{
+	return dot( LinearColor, LuminanceFactors );
 }
 
 //------------------------------------------------------------------
@@ -66,17 +85,16 @@ float3 XYZ_2_LinearRGB( float3 XYZ )
 // https://en.wikipedia.org/wiki/Lab_color_space
 //------------------------------------------------------------------
 // Assume XYZ is in nominal range to reduce the operation on a constant coefficient 100.
-// E.g., whilte RGB (1, 1, 1) maps to XYZ (0.95047, 1.00, 1.08883)
+// E.g., white RGB (1, 1, 1) maps to XYZ (0.9504559271, 1.0, 1.0890577508)
 // for a reference white of D65. LAB is still in the original range.
-// Use XYZ_WHITE_REF_DXX_NOMINAL instead for LinearRGB_2_LAB() and LAB_2_LinearRGB().
-// E.g.,
-// LinearRGB_2_LAB((1.0f, 1.0f, 1.0f), XYZ_WHITE_REF_D65_NOMINAL) = (100.0f, 0.0f, 0.0f)
-
-static const float3 XYZ_WHITE_REF_D65 = float3(95.047, 100.0, 108.883);
-static const float3 XYZ_WHITE_REF_D50 = float3(96.6797, 100.0, 82.5188);
-static const float3 XYZ_WHITE_REF_D65_NOMINAL = float3(0.95047, 1.00, 1.08883);
-static const float3 XYZ_WHITE_REF_D50_NOMINAL = float3(0.966797, 1.00, 0.825188);
-static const float XYZ_2_LAB_DELTA_SQUARED = 0.04280618311; // (6/29)^3
+
+#if WORKING_COLOR_SPACE_IS_SRGB
+static const float3 XYZ_WHITE_REF_NOMINAL = float3(0.9504559271, 1.0, 1.0890577508);
+#else
+static const float3 XYZ_WHITE_REF_NOMINAL = mul(WORKING_COLOR_SPACE_RGB_TO_XYZ_MAT, float3(1.0, 1.0, 1.0));
+#endif
+
+static const float XYZ_2_LAB_DELTA_SQUARED = 0.04280618311; // (6/29)^2
 static const float XYZ_2_LAB_DELTA_CUBED = 0.00885645167; // (6/29)^3
 
 float xyz_otherwise(float t)
@@ -84,13 +102,13 @@ float xyz_otherwise(float t)
 	return (t / (3.0 * XYZ_2_LAB_DELTA_SQUARED)) + 4.0 / 29.0;
 }
 
-float3 LinearRGB_2_LAB( float3 LinearRGB, float3 ReferenceWhite )
+float3 LinearRGB_2_LAB( float3 LinearRGB )
 {
 	float3 XYZ = LinearRGB_2_XYZ(LinearRGB);
 
-	float t_X = XYZ.x / ReferenceWhite.x;
-	float t_Y = XYZ.y / ReferenceWhite.y;
-	float t_Z = XYZ.z / ReferenceWhite.z;
+	float t_X = XYZ.x / XYZ_WHITE_REF_NOMINAL.x;
+	float t_Y = XYZ.y / XYZ_WHITE_REF_NOMINAL.y;
+	float t_Z = XYZ.z / XYZ_WHITE_REF_NOMINAL.z;
 
 	float f_X = (t_X > XYZ_2_LAB_DELTA_CUBED) ? pow(max(t_X, XYZ_2_LAB_DELTA_CUBED), 1.0 / 3.0) : xyz_otherwise(t_X);
 	float f_Y = (t_Y > XYZ_2_LAB_DELTA_CUBED) ? pow(max(t_Y, XYZ_2_LAB_DELTA_CUBED), 1.0 / 3.0) : xyz_otherwise(t_Y);
@@ -108,7 +126,7 @@ float lab_otherwise(float t)
 	return (3.0 * XYZ_2_LAB_DELTA_SQUARED) * (t - (4.0 / 29.0));
 }
 
-float3 LAB_2_LinearRGB( float3 LAB, float3 ReferenceWhite)
+float3 LAB_2_LinearRGB( float3 LAB )
 {
 	float L = LAB.x;
 	float a = LAB.y;
@@ -137,9 +155,9 @@ float3 LAB_2_LinearRGB( float3 LAB, float3 ReferenceWhite)
 		f_z = lab_otherwise(t_z);
 	}
 
-	float X = ReferenceWhite.x * f_x;
-	float Y = ReferenceWhite.y * f_y;
-	float Z = ReferenceWhite.z * f_z;
+	float X = XYZ_WHITE_REF_NOMINAL.x * f_x;
+	float Y = XYZ_WHITE_REF_NOMINAL.y * f_y;
+	float Z = XYZ_WHITE_REF_NOMINAL.z * f_z;
 
 	return XYZ_2_LinearRGB(float3(X, Y, Z));
 }
diff --git a/Engine/Shaders/Private/PathTracing/PathTracingSpatialTemporalDenoising.usf b/Engine/Shaders/Private/PathTracing/PathTracingSpatialTemporalDenoising.usf
--- a/Engine/Shaders/Private/PathTracing/PathTracingSpatialTemporalDenoising.usf
+++ b/Engine/Shaders/Private/PathTracing/PathTracingSpatialTemporalDenoising.usf
@@ -643,8 +643,8 @@ static float NormalizeDeltaEAndEuclideanDistanceToSameMean = 0.0141f;
 float GetDeltaE(float3 LinearRGB0, float3 LinearRGB1, bool bNormalizeToEuclideanMean = true)
 {
 
-	float3 Lab0 = LinearRGB_2_LAB(LinearRGB0, XYZ_WHITE_REF_D65_NOMINAL);
-	float3 Lab1 = LinearRGB_2_LAB(LinearRGB1, XYZ_WHITE_REF_D65_NOMINAL);
+	float3 Lab0 = LinearRGB_2_LAB(LinearRGB0);
+	float3 Lab1 = LinearRGB_2_LAB(LinearRGB1);
 
 	float DeltaE = DeltaE_CIE2000(Lab0, Lab1);
 
diff --git a/Engine/Shaders/Private/Strata/StrataLegacyConversion.ush b/Engine/Shaders/Private/Strata/StrataLegacyConversion.ush
--- a/Engine/Shaders/Private/Strata/StrataLegacyConversion.ush
+++ b/Engine/Shaders/Private/Strata/StrataLegacyConversion.ush
@@ -398,10 +398,10 @@ FStrataData StrataConvertLegacyMaterialDynamic(
 		// Because we are using a single slab for simulating a 'coverage mix' between two BSDFs (a opaque metallic/conductor and a non-scattering translucent medium), 
 		// we need to do some approximation. In order to reduce avoid hue shift due to difference between 'solid' base-color/diffuse-albedo and 'transparent' transmission 
 		// color, we lerp BaseColor with transmittance color in LAB space so that the interpolated color minized its 'distance' 
-		const float3 BaseColor_LAB			= LinearRGB_2_LAB(BaseColor, XYZ_WHITE_REF_D65_NOMINAL);
-		const float3 TransmittanceColor_LAB = LinearRGB_2_LAB(TransmittanceColor, XYZ_WHITE_REF_D65_NOMINAL);
+		const float3 BaseColor_LAB			= LinearRGB_2_LAB(BaseColor);
+		const float3 TransmittanceColor_LAB = LinearRGB_2_LAB(TransmittanceColor);
 		const float3 Color_LAB			    = lerp(TransmittanceColor_LAB, BaseColor_LAB, Opacity);
-		const float3 Color					= LAB_2_LinearRGB(Color_LAB, XYZ_WHITE_REF_D65_NOMINAL);
+		const float3 Color					= LAB_2_LinearRGB(Color_LAB);
 
 		// Lerp the transmittance color rather than MFP as it is more perceptually linear
 		const float3 MFP = TransmittanceToMeanFreePath(lerp(Color, Zeros, Opacity), Thickness * CENTIMETER_TO_METER) * METER_TO_CENTIMETER;
diff --git a/Engine/Source/Runtime/Engine/Private/ShaderCompiler/ShaderCompiler.cpp b/Engine/Source/Runtime/Engine/Private/ShaderCompiler/ShaderCompiler.cpp
--- a/Engine/Source/Runtime/Engine/Private/ShaderCompiler/ShaderCompiler.cpp
+++ b/Engine/Source/Runtime/Engine/Private/ShaderCompiler/ShaderCompiler.cpp
@@ -6816,13 +6816,30 @@ void GlobalBeginCompileShader(
 		const bool bWorkingColorSpaceIsSRGB = FColorSpace::GetWorking().IsSRGB();
 		Input.Environment.SetDefine(TEXT("WORKING_COLOR_SPACE_IS_SRGB"), bWorkingColorSpaceIsSRGB ? 1 : 0);
 		
-		// We limit matrix definition below when WORKING_COLOR_SPACE_IS_SRGB == 0.
+		// We limit matrix definitions below to WORKING_COLOR_SPACE_IS_SRGB == 0.
 		if (!bWorkingColorSpaceIsSRGB)
 		{
-			TCHAR MatrixFormat[] = TEXT("float3x3(%0.10f, %0.10f, %0.10f, %0.10f, %0.10f, %0.10f, %0.10f, %0.10f, %0.10f)");
+			const TCHAR MatrixFormat[] = TEXT("float3x3(%0.10f, %0.10f, %0.10f, %0.10f, %0.10f, %0.10f, %0.10f, %0.10f, %0.10f)");
+			const FColorSpace& WorkingColorSpace = FColorSpace::GetWorking();
 
-			const FColorSpaceTransform FromSRGB(FColorSpace(EColorSpace::sRGB), FColorSpace::GetWorking());
+			// Note that we transpose the matrices during print since color matrices are usually pre-multiplied.
+			const FMatrix44d& ToXYZ = WorkingColorSpace.GetRgbToXYZ();
+			Input.Environment.SetDefine(
+				TEXT("WORKING_COLOR_SPACE_RGB_TO_XYZ_MAT"),
+				FString::Printf(MatrixFormat,
+					ToXYZ.M[0][0], ToXYZ.M[1][0], ToXYZ.M[2][0],
+					ToXYZ.M[0][1], ToXYZ.M[1][1], ToXYZ.M[2][1],
+					ToXYZ.M[0][2], ToXYZ.M[1][2], ToXYZ.M[2][2]));
+
+			const FMatrix44d& FromXYZ = WorkingColorSpace.GetXYZToRgb();
+			Input.Environment.SetDefine(
+				TEXT("XYZ_TO_RGB_WORKING_COLOR_SPACE_MAT"),
+				FString::Printf(MatrixFormat,
+					FromXYZ.M[0][0], FromXYZ.M[1][0], FromXYZ.M[2][0],
+					FromXYZ.M[0][1], FromXYZ.M[1][1], FromXYZ.M[2][1],
+					FromXYZ.M[0][2], FromXYZ.M[1][2], FromXYZ.M[2][2]));
 
+			const FColorSpaceTransform FromSRGB(FColorSpace(EColorSpace::sRGB), WorkingColorSpace);
 			Input.Environment.SetDefine(
 				TEXT("SRGB_TO_WORKING_COLOR_SPACE_MAT"),
 				FString::Printf(MatrixFormat,
