Substrate - fixed nans from glints.

diff --git a/Engine/Shaders/Private/Strata/Glint/GlintThirdParty.ush b/Engine/Shaders/Private/Strata/Glint/GlintThirdParty.ush
--- a/Engine/Shaders/Private/Strata/Glint/GlintThirdParty.ush
+++ b/Engine/Shaders/Private/Strata/Glint/GlintThirdParty.ush
@@ -19,6 +19,8 @@
 #define InvSqrt2						0.707106f
 #define MaxAnisotropy					4.0f
 
+#define GlintSafeThreshold				0.000001f
+
 // FGlintLutParams contains variables linked to the GlintShadingLUT
 struct FGlintLutParams
 {
@@ -385,7 +387,7 @@ float P22__glint_discrete_LOD(
 		if (nbrOfIter > SafeMaxIterationCount)
 			break;
 	}
-	return sum / sumWts;
+	return sum / max(GlintSafeThreshold, sumWts);
 }
 
 //=============================================================================
