Substrate - added glint global bias and min level cvars.
ShaderVersion bumped since View data has been changed.

diff --git a/Engine/Shaders/Private/Strata/Glint/GlintThirdParty.ush b/Engine/Shaders/Private/Strata/Glint/GlintThirdParty.ush
--- a/Engine/Shaders/Private/Strata/Glint/GlintThirdParty.ush
+++ b/Engine/Shaders/Private/Strata/Glint/GlintThirdParty.ush
@@ -27,9 +27,23 @@ struct FGlintLutParams
 	float Alpha;
 	int N;
 	int NLevels;
-	float MinLDist;
+	float LevelBias;
+	float LevelMin;
 };
 
+FGlintLutParams GetGlintParams()
+{
+	FGlintLutParams GlintLutParams;
+	const float4 Params0 = View.GlintLUTParameters0;
+	const float4 Params1 = View.GlintLUTParameters1;
+	GlintLutParams.Alpha		= Params0.x;
+	GlintLutParams.N			= asint(Params0.y);
+	GlintLutParams.NLevels		= asint(Params0.z);
+	GlintLutParams.LevelBias	= Params0.w;
+	GlintLutParams.LevelMin		= Params1.x;
+	return GlintLutParams;
+}
+
 #define USE_GLINT_AA (!RAYCALLABLESHADER && !RAYHITGROUPSHADER && !RAYGENSHADER && !RAYMISSSHADER)
 
 // Geometric Glint Anti-Aliasing parameters
@@ -175,7 +189,7 @@ float P22_M(
 	l_dist = sampleNormalDistribution(uDensityRandomisation, l_dist,
 		densityRandomisation);
 
-	l_dist = clamp(int(round(l_dist)), GlintLutParams.MinLDist, GlintLutParams.NLevels);
+	l_dist = clamp(int(round(l_dist)), 0, GlintLutParams.NLevels);
 
 	// Recover roughness and slope correlation factor
 	float sigma_x = sigma_x_y_rho.x;
@@ -408,11 +422,7 @@ float f_P(float3 wo, float3 wi, float3 sigmas_rho, float LogMicrofacetDensity, f
 	if (dot(wo, wh) <= 0. || dot(wi, wh) <= 0.)
 		return 0;
 
-	FGlintLutParams GlintLutParams;
-	GlintLutParams.Alpha	= View.GlintLUTParameters.x;
-	GlintLutParams.N		= asint(View.GlintLUTParameters.y);
-	GlintLutParams.NLevels	= asint(View.GlintLUTParameters.z);
-	GlintLutParams.MinLDist = 0;
+	FGlintLutParams GlintLutParams = GetGlintParams();
 
 	// Compute texture derivatives
 	float2 dst0 = TexCoordDDX;
@@ -479,8 +489,7 @@ float f_P(float3 wo, float3 wi, float3 sigmas_rho, float LogMicrofacetDensity, f
 	else
 	{
 		// Choose LOD
-		float l =
-			max(0., GlintLutParams.NLevels - 1. + log2(minorLength));
+		float l = max(0.0, GlintLutParams.NLevels - 1. + log2(minorLength));
 		int il = int(floor(l));
 
 		float w = l - float(il);
@@ -498,6 +507,10 @@ float f_P(float3 wo, float3 wi, float3 sigmas_rho, float LogMicrofacetDensity, f
 		float P22_P_il = non_axis_aligned_anisotropic_beckmann(slope_h.x, slope_h.y, sigmas_rho.x, sigmas_rho.y, sigmas_rho.z);
 		float P22_P_ilp1 = P22_P_il;
 
+		// Apply here level space bias and clamp to a minimum. This before we evaluate each glint level.
+		LOD_dist_il   = clamp(LOD_dist_il   + GlintLutParams.LevelBias, 0, max(0, GlintLutParams.NLevels - GlintLutParams.LevelMin));
+		LOD_dist_ilp1 = clamp(LOD_dist_ilp1 + GlintLutParams.LevelBias, 0, max(0, GlintLutParams.NLevels - GlintLutParams.LevelMin));
+
 		if (int(round(LOD_dist_il)) < GlintLutParams.NLevels)
 		{
 			P22_P_il = P22__glint_discrete_LOD(GlintLutParams, il, slope_h, TexCoord, dst0, dst1, slope_dx, slope_dy, sigmas_rho, LOD_dist_il);
@@ -803,11 +816,7 @@ float2 Sample_P22_LOD(
 
 float3 Sample_P(float3 u, float3 wo, float3 sigmas_rho, float LogMicrofacetDensity, float2 TexCoord, float2 TexCoordDDX, float2 TexCoordDDY)
 {
-	FGlintLutParams GlintLutParams;
-	GlintLutParams.Alpha	= View.GlintLUTParameters.x;
-	GlintLutParams.N		= asint(View.GlintLUTParameters.y);
-	GlintLutParams.NLevels	= asint(View.GlintLUTParameters.z);
-	GlintLutParams.MinLDist = 0;
+	FGlintLutParams GlintLutParams = GetGlintParams();
 
 	float U0 = u.x;
 	float U1 = u.y;
@@ -854,7 +863,7 @@ float3 Sample_P(float3 u, float3 wo, float3 sigmas_rho, float LogMicrofacetDensi
 	// Compute LOD
 	if (minorLength > 0)
 	{
-		const float l = max(0., GlintLutParams.NLevels - 1. + log2(minorLength));
+		const float l = max(0.0, GlintLutParams.NLevels - 1. + log2(minorLength));
 		const int il = int(floor(l));
 		const float wLOD = l - float(il);
 
@@ -868,6 +877,10 @@ float3 Sample_P(float3 u, float3 wo, float3 sigmas_rho, float LogMicrofacetDensi
 		n_ilp1 *= exp(LogMicrofacetDensity);
 		float LOD_dist_ilp1 = log(n_ilp1) / 1.38629; // 2. * log(2) = 1.38629
 
+		// Apply here level space bias and clamp to a minimum. This before we evaluate each glint level.
+		LOD_dist_il   = clamp(LOD_dist_il   + GlintLutParams.LevelBias, 0, max(0, GlintLutParams.NLevels - GlintLutParams.LevelMin));
+		LOD_dist_ilp1 = clamp(LOD_dist_ilp1 + GlintLutParams.LevelBias, 0, max(0, GlintLutParams.NLevels - GlintLutParams.LevelMin));
+
 		// 1. Select LOD (N or N+1)
 		uint LOD = il;
 		float LODDist = LOD_dist_il;
diff --git a/Engine/Shaders/Public/ShaderVersion.ush b/Engine/Shaders/Public/ShaderVersion.ush
--- a/Engine/Shaders/Public/ShaderVersion.ush
+++ b/Engine/Shaders/Public/ShaderVersion.ush
@@ -2,5 +2,5 @@
 // Each time the console command is executed it generates a new GUID. As this file is included
 // in Platform.ush (which should be included in any shader) it allows to invalidate the shader DDC.
 // 
-// GUID = 7CC2CDB6-7689-4413-A8FD-E89BB19E2EDD
+// GUID = F649466F-64CA-4F7C-896D-6557FAEE1130
 
diff --git a/Engine/Source/Runtime/Engine/Public/SceneView.h b/Engine/Source/Runtime/Engine/Public/SceneView.h
--- a/Engine/Source/Runtime/Engine/Public/SceneView.h
+++ b/Engine/Source/Runtime/Engine/Public/SceneView.h
@@ -901,7 +901,8 @@ enum ETranslucencyVolumeCascade
 	VIEW_UNIFORM_BUFFER_MEMBER(float, SubSurfaceColorAsTransmittanceAtDistanceInMeters) \
 	VIEW_UNIFORM_BUFFER_MEMBER_PER_VIEW(FVector4f, TanAndInvTanHalfFOV) \
 	VIEW_UNIFORM_BUFFER_MEMBER_PER_VIEW(FVector4f, PrevTanAndInvTanHalfFOV) \
-	VIEW_UNIFORM_BUFFER_MEMBER_PER_VIEW(FVector4f, GlintLUTParameters) \
+	VIEW_UNIFORM_BUFFER_MEMBER_PER_VIEW(FVector4f, GlintLUTParameters0) \
+	VIEW_UNIFORM_BUFFER_MEMBER_PER_VIEW(FVector4f, GlintLUTParameters1) \
 
 /** The uniform shader parameters associated with a view. */
 #define VIEW_UNIFORM_BUFFER_MEMBER(type, identifier) \
diff --git a/Engine/Source/Runtime/RenderCore/Private/RenderUtils.cpp b/Engine/Source/Runtime/RenderCore/Private/RenderUtils.cpp
--- a/Engine/Source/Runtime/RenderCore/Private/RenderUtils.cpp
+++ b/Engine/Source/Runtime/RenderCore/Private/RenderUtils.cpp
@@ -1638,6 +1638,18 @@ static TAutoConsoleVariable<int32> CVarSubstrateGlintsLUT(
 	TEXT("Select one of the glint rendering LUT for testing."),
 	ECVF_RenderThreadSafe);
 
+static TAutoConsoleVariable<float> CVarSubstrateGlintsLevelBias(
+	TEXT("r.Substrate.Glints.LevelBias"),
+	0.0f,
+	TEXT("Constantly bias the glint level. A negative value will reinforce the glint effect while a positive value will attenuate it. Be careful as this might cause aliasing."),
+	ECVF_RenderThreadSafe);
+
+static TAutoConsoleVariable<float> CVarSubstrateGlintsLevelMin(
+	TEXT("r.Substrate.Glints.LevelMin"),
+	0.0f,
+	TEXT("The minimum glint level used to render glints. Mostly affect the look of distance surfaces. A value greater than 0 will force glints to always be visible instead of using the analytic BSDF. Be careful as this might cause aliasing."),
+	ECVF_RenderThreadSafe);
+
 static TAutoConsoleVariable<int32> CVarSubstrateSpecularProfile(
 	TEXT("r.Substrate.SpecularProfile"),
 	1,
@@ -1753,6 +1765,17 @@ namespace Strata
 		return CVarSubstrateGlintsLUT.GetValueOnAnyThread()<= 0 ? 0u : 1u;
 	}
 
+	float GlintLevelBias()
+	{
+		return CVarSubstrateGlintsLevelBias.GetValueOnAnyThread();
+	}
+
+	float GlintLevelMin()
+	{
+		const float MinLevel = CVarSubstrateGlintsLevelMin.GetValueOnAnyThread();
+		return MinLevel < 0.0f ? 0.0f : MinLevel;
+	}
+
 	bool IsSpecularProfileEnabled()
 	{
 		return IsStrataEnabled() && CVarSubstrateSpecularProfile.GetValueOnAnyThread() > 0;
diff --git a/Engine/Source/Runtime/RenderCore/Public/RenderUtils.h b/Engine/Source/Runtime/RenderCore/Public/RenderUtils.h
--- a/Engine/Source/Runtime/RenderCore/Public/RenderUtils.h
+++ b/Engine/Source/Runtime/RenderCore/Public/RenderUtils.h
@@ -480,6 +480,8 @@ namespace Strata
 	RENDERCORE_API bool IsRoughDiffuseEnabled();
 	RENDERCORE_API bool IsGlintEnabled();
 	RENDERCORE_API uint32 GlintLUTIndex();
+	RENDERCORE_API float GlintLevelBias();
+	RENDERCORE_API float GlintLevelMin();
 	RENDERCORE_API bool IsSpecularProfileEnabled();
 	RENDERCORE_API bool IsBackCompatibilityEnabled();
 	RENDERCORE_API bool IsDBufferPassEnabled(EShaderPlatform InPlatform);
diff --git a/Engine/Source/Runtime/Renderer/Private/SceneRendering.cpp b/Engine/Source/Runtime/Renderer/Private/SceneRendering.cpp
--- a/Engine/Source/Runtime/Renderer/Private/SceneRendering.cpp
+++ b/Engine/Source/Runtime/Renderer/Private/SceneRendering.cpp
@@ -1968,11 +1968,14 @@ void FViewInfo::SetupUniformBufferParameters(
 	if (ViewState)
 	{
 		ViewUniformShaderParameters.GlintTexture = ViewState->GlintShadingLUTsData.RHIGlintShadingLUTs ? ViewState->GlintShadingLUTsData.RHIGlintShadingLUTs : nullptr;
-		ViewUniformShaderParameters.GlintLUTParameters = FVector4f(
+		ViewUniformShaderParameters.GlintLUTParameters0 = FVector4f(
 			ViewState->GlintShadingLUTsData.Dictionary_Alpha,
 			*reinterpret_cast<float*>(&ViewState->GlintShadingLUTsData.Dictionary_N),
 			*reinterpret_cast<float*>(&ViewState->GlintShadingLUTsData.Dictionary_NLevels),
-			0.0f);
+			Strata::GlintLevelBias());
+		ViewUniformShaderParameters.GlintLUTParameters1 = FVector4f(
+			Strata::GlintLevelMin(),
+			0.0f, 0.0f, 0.0f);
 	}
 	ViewUniformShaderParameters.GlintTexture = OrBlack2DArrayIfNull(ViewUniformShaderParameters.GlintTexture);
 
