Fix RT shadow dithering with Substrate and special complex material.

Stencil bit 0 is overriden by the RTShadowDitheringLOD path, which is used by Substrate to flag special complex material. This CL changes the special complex material to 1 instead which is aliased by the DBuffer during the based pass (and no longer used after).

diff --git a/Engine/Shaders/Private/Strata/StrataVisualize.usf b/Engine/Shaders/Private/Strata/StrataVisualize.usf
--- a/Engine/Shaders/Private/Strata/StrataVisualize.usf
+++ b/Engine/Shaders/Private/Strata/StrataVisualize.usf
@@ -129,6 +129,15 @@ void Newline(inout FShaderPrintContext Context, inout float2 MaxRect)
 	Newline(Context);
 }
 
+void PrintAddress(inout FShaderPrintContext Context, inout FStrataAddressing StrataAddressing)
+{
+	FFontColor Font;
+	Font.Color = lerp(FontEmerald.Color, FontSilver.Color, 0.75f);
+	Print(Context, TEXT("[Address="), Font);
+	Print(Context, StrataAddressing.CurrentIndex, Font, 2, 0);
+	Print(Context, TEXT("]"), Font);
+}
+
 void PrintBSDF(inout FShaderPrintContext Context, inout FStrataAddressing StrataAddressing, FStrataPixelHeader Header, float3 WorldPosition, float3 V, inout float2 RectMax)
 {
 	const FFontColor FontBSDF         = FontSilver;
@@ -137,8 +146,6 @@ void PrintBSDF(inout FShaderPrintContext Context, inout FStrataAddressing Strata
 	const FFontColor FontBSDFPropName = FontOrange;
 	const FFontColor FontBSDFPropValu = FontWhite;
 
-	Print(Context, TEXT("Address = "), FontEmerald);
-	Print(Context, StrataAddressing.CurrentIndex, FontEmerald);
 	Newline(Context, RectMax);
 
 	const FStrataBSDF BSDF = UnpackStrataBSDFIn(Strata.MaterialTextureArray, StrataAddressing, Header);
@@ -159,6 +166,7 @@ void PrintBSDF(inout FShaderPrintContext Context, inout FStrataAddressing Strata
 			const bool bGreyTopTrans = BSDF.TransmittanceAboveAlongN.x == BSDF.TransmittanceAboveAlongN.y && BSDF.TransmittanceAboveAlongN.x == BSDF.TransmittanceAboveAlongN.z;
 
 			Print(Context, TEXT("Slab "), FontBSDFType);
+			PrintAddress(Context, StrataAddressing);
 			Newline(Context, RectMax);
 			Print(Context, TEXT(" NormalID BasisType Aniso TopLayer Scattering IsThin WeightV "), FontBSDFStateName);
 			if (!bGreyWeightV)
@@ -359,6 +367,7 @@ void PrintBSDF(inout FShaderPrintContext Context, inout FStrataAddressing Strata
 		case STRATA_BSDF_TYPE_HAIR:
 		{
 			Print(Context, TEXT("Hair"), FontBSDFType);
+			PrintAddress(Context, StrataAddressing);
 			Newline(Context, RectMax);
 			Print(Context, TEXT(" NormalID Aniso IsTopLayer LuminanceWeight"), FontBSDFStateName);
 			Newline(Context, RectMax);
@@ -391,6 +400,7 @@ void PrintBSDF(inout FShaderPrintContext Context, inout FStrataAddressing Strata
 		case STRATA_BSDF_TYPE_EYE:
 		{
 			Print(Context, TEXT("Eye"), FontBSDFType);
+			PrintAddress(Context, StrataAddressing);
 			Newline(Context, RectMax);
 			Print(Context, TEXT(" NormalID IsTopLayer LuminanceWeight"), FontBSDFStateName);
 			Newline(Context, RectMax);
@@ -439,8 +449,9 @@ void PrintBSDF(inout FShaderPrintContext Context, inout FStrataAddressing Strata
 		case STRATA_BSDF_TYPE_SINGLELAYERWATER:
 		{
 			Print(Context, TEXT("Water"), FontBSDFType);
-
+			PrintAddress(Context, StrataAddressing);
 			Newline(Context);
+
 			Print(Context, TEXT(" BaseColor            "), FontBSDFPropName);	Print(Context, SLW_BASECOLOR(BSDF), FontBSDFPropValu);			Newline(Context, RectMax);
 			Print(Context, TEXT(" Metallic             "), FontBSDFPropName);	Print(Context, SLW_METALLIC(BSDF), FontBSDFPropValu);			Newline(Context, RectMax);
 			Print(Context, TEXT(" Specular             "), FontBSDFPropName);	Print(Context, SLW_SPECULAR(BSDF), FontBSDFPropValu);			Newline(Context, RectMax);
diff --git a/Engine/Source/Runtime/Renderer/Private/PostProcess/SceneRenderTargets.h b/Engine/Source/Runtime/Renderer/Private/PostProcess/SceneRenderTargets.h
--- a/Engine/Source/Runtime/Renderer/Private/PostProcess/SceneRenderTargets.h
+++ b/Engine/Source/Runtime/Renderer/Private/PostProcess/SceneRenderTargets.h
@@ -19,6 +19,7 @@
 *		[7]       | primitive receive decal bit
 *
 * After deferred decals, stencil is cleared to 0 and no longer packed in this way, to ensure use of fast hardware clears and HiStencil.
+* Stencil [0] is used by RT shadow for shadow LOD disthering (see: DitheredLODFadingOutMaskPass)
 */
 #define STENCIL_SANDBOX_BIT_ID							0
 // Must match usf
@@ -30,7 +31,7 @@
 #define STENCIL_STRATA_FASTPATH							4 
 #define STENCIL_STRATA_SINGLEPATH						5
 #define STENCIL_STRATA_COMPLEX							6
-#define STENCIL_STRATA_COMPLEX_SPECIAL					0
+#define STENCIL_STRATA_COMPLEX_SPECIAL					1
 // Used only by Strata during the base pass when bUseDBufferPass is enabled (to mark material STRATA_DBUFFER_RESPONSE_xxx Normal/BaseColor/Roughness)
 #define STENCIL_STRATA_RECEIVE_DBUFFER_NORMAL_BIT_ID	1
 #define STENCIL_STRATA_RECEIVE_DBUFFER_DIFFUSE_BIT_ID	3
diff --git a/Engine/Source/Runtime/Renderer/Private/Strata/Strata.h b/Engine/Source/Runtime/Renderer/Private/Strata/Strata.h
--- a/Engine/Source/Runtime/Renderer/Private/Strata/Strata.h
+++ b/Engine/Source/Runtime/Renderer/Private/Strata/Strata.h
@@ -176,7 +176,7 @@ namespace Strata
 constexpr uint32 StencilBit_Fast			= 0x08; // In sync with SceneRenderTargets.h - GET_STENCIL_BIT_MASK(STENCIL_STRATA_FASTPATH)
 constexpr uint32 StencilBit_Single			= 0x10; // In sync with SceneRenderTargets.h - GET_STENCIL_BIT_MASK(STENCIL_STRATA_SINGLEPATH)
 constexpr uint32 StencilBit_Complex			= 0x20; // In sync with SceneRenderTargets.h - GET_STENCIL_BIT_MASK(STENCIL_STRATA_COMPLEX)
-constexpr uint32 StencilBit_ComplexSpecial	= 0x01; // In sync with SceneRenderTargets.h - GET_STENCIL_BIT_MASK(STENCIL_STRATA_COMPLEX_SPECIAL)	
+constexpr uint32 StencilBit_ComplexSpecial	= 0x02; // In sync with SceneRenderTargets.h - GET_STENCIL_BIT_MASK(STENCIL_STRATA_COMPLEX_SPECIAL)	
 
 FIntPoint GetStrataTextureResolution(const FViewInfo& View, const FIntPoint& InResolution);
 bool GetStrataUsesComplexSpecialPath(const FViewInfo& View);
