Add dither for Substrate specular profile

diff --git a/Engine/Shaders/Private/Strata/Strata.ush b/Engine/Shaders/Private/Strata/Strata.ush
--- a/Engine/Shaders/Private/Strata/Strata.ush
+++ b/Engine/Shaders/Private/Strata/Strata.ush
@@ -873,11 +873,20 @@ uint GetSpecularProfileParameterization(float In)
 	return In > 0.f ? 0u : 1u;
 }
 
-float3 EvaluateSpecularProfile(float InSpecularProfileId, float NoV, float NoL, float VoH, float NoH)
+float IsNonZeroFast(float v);
+float3 EvaluateSpecularProfile(float InSpecularProfileId, float NoV, float NoL, float VoH, float NoH, uint2 PixelCoord = 0)
 {
+	// Dither sampling for balancing low-resolution SpecularProfile
 	const uint ProfileID = GetSpecularProfileId(InSpecularProfileId);
 	const bool bUseHalfVector = GetSpecularProfileParameterization(InSpecularProfileId) == 1;
-	const float3 ProfileCoord = bUseHalfVector ? float3(VoH, NoH, ProfileID) : float3(NoV, NoL, ProfileID);
+	const float Dither = InterleavedGradientNoise(PixelCoord, View.StateFrameIndexMod8);
+	const float CenteredDither = Dither - 0.5f;
+	const float ScaledDither = CenteredDither * View.SpecularProfileTextureSizeAndInvSize.z;
+
+	NoV = NoV + ScaledDither * IsNonZeroFast(NoV);
+	NoL = NoL + ScaledDither * IsNonZeroFast(NoL);
+
+	const float3 ProfileCoord = float3(NoV, NoL, ProfileID);
 	return View.SpecularProfileTexture.SampleLevel(View.SpecularProfileSampler, ProfileCoord, 0).xyz;
 }
 
diff --git a/Engine/Shaders/Private/Strata/StrataEvaluation.ush b/Engine/Shaders/Private/Strata/StrataEvaluation.ush
--- a/Engine/Shaders/Private/Strata/StrataEvaluation.ush
+++ b/Engine/Shaders/Private/Strata/StrataEvaluation.ush
@@ -655,7 +655,7 @@ FStrataEvaluateResult StrataEvaluateBSDFCommon(FStrataBSDFContext BSDFContext, F
 			BRANCH
 			if (BSDF_GETHASSPECPROFILE(BSDFContext.BSDF))
 			{
-				Sample.IntegratedSpecularValue *= EvaluateSpecularProfile(SLAB_SPECPROFILEID(BSDFContext.BSDF), BSDFContext.SatNoV, BSDFContext.SatNoL, BSDFContext.Context.VoH, BSDFContext.Context.NoH);
+				Sample.IntegratedSpecularValue *= EvaluateSpecularProfile(SLAB_SPECPROFILEID(BSDFContext.BSDF), BSDFContext.SatNoV, BSDFContext.SatNoL, BSDFContext.Context.VoH, BSDFContext.Context.NoH, BSDFContext.PixelCoord);
 			}
 			#endif
 
@@ -1539,7 +1539,7 @@ FStrataEnvLightResult StrataEvaluateForEnvLight(FStrataBSDFContext BSDFContext,
 				if (BSDF_GETHASSPECPROFILE(BSDFContext.BSDF))
 				{
 					const float NoR = saturate(dot(StrataEnvLightResult.SpecularDirection, BSDFContext.N));
-					StrataEnvLightResult.SpecularWeight *= EvaluateSpecularProfile(SLAB_SPECPROFILEID(BSDFContext.BSDF), BSDFContext.SatNoV, NoR, BSDFContext.SatNoV, NoR);
+					StrataEnvLightResult.SpecularWeight *= EvaluateSpecularProfile(SLAB_SPECPROFILEID(BSDFContext.BSDF), BSDFContext.SatNoV, NoR, BSDFContext.SatNoV, NoR, BSDFContext.PixelCoord);
 				}
 			#endif // STRATA_COMPLEXPATH
 
