Fix incorrect point sampling for Specular LUT and and revert dithering.

diff --git a/Engine/Shaders/Private/Strata/Strata.ush b/Engine/Shaders/Private/Strata/Strata.ush
--- a/Engine/Shaders/Private/Strata/Strata.ush
+++ b/Engine/Shaders/Private/Strata/Strata.ush
@@ -873,20 +873,11 @@ uint GetSpecularProfileParameterization(float In)
 	return In > 0.f ? 0u : 1u;
 }
 
-float IsNonZeroFast(float v);
-float3 EvaluateSpecularProfile(float InSpecularProfileId, float NoV, float NoL, float VoH, float NoH, uint2 PixelCoord = 0)
+float3 EvaluateSpecularProfile(float InSpecularProfileId, float NoV, float NoL, float VoH, float NoH)
 {
-	// Dither sampling for balancing low-resolution SpecularProfile
 	const uint ProfileID = GetSpecularProfileId(InSpecularProfileId);
 	const bool bUseHalfVector = GetSpecularProfileParameterization(InSpecularProfileId) == 1;
-	const float Dither = InterleavedGradientNoise(PixelCoord, View.StateFrameIndexMod8);
-	const float CenteredDither = Dither - 0.5f;
-	const float ScaledDither = CenteredDither * View.SpecularProfileTextureSizeAndInvSize.z;
-
-	NoV = NoV + ScaledDither * IsNonZeroFast(NoV);
-	NoL = NoL + ScaledDither * IsNonZeroFast(NoL);
-
-	const float3 ProfileCoord = float3(NoV, NoL, ProfileID);
+	const float3 ProfileCoord = bUseHalfVector ? float3(VoH, NoH, ProfileID) : float3(NoV, NoL, ProfileID);
 	return View.SpecularProfileTexture.SampleLevel(View.SpecularProfileSampler, ProfileCoord, 0).xyz;
 }
 
diff --git a/Engine/Shaders/Private/Strata/StrataEvaluation.ush b/Engine/Shaders/Private/Strata/StrataEvaluation.ush
--- a/Engine/Shaders/Private/Strata/StrataEvaluation.ush
+++ b/Engine/Shaders/Private/Strata/StrataEvaluation.ush
@@ -655,7 +655,7 @@ FStrataEvaluateResult StrataEvaluateBSDFCommon(FStrataBSDFContext BSDFContext, F
 			BRANCH
 			if (BSDF_GETHASSPECPROFILE(BSDFContext.BSDF))
 			{
-				Sample.IntegratedSpecularValue *= EvaluateSpecularProfile(SLAB_SPECPROFILEID(BSDFContext.BSDF), BSDFContext.SatNoV, BSDFContext.SatNoL, BSDFContext.Context.VoH, BSDFContext.Context.NoH, BSDFContext.PixelCoord);
+				Sample.IntegratedSpecularValue *= EvaluateSpecularProfile(SLAB_SPECPROFILEID(BSDFContext.BSDF), BSDFContext.SatNoV, BSDFContext.SatNoL, BSDFContext.Context.VoH, BSDFContext.Context.NoH);
 			}
 			#endif
 
@@ -1539,7 +1539,7 @@ FStrataEnvLightResult StrataEvaluateForEnvLight(FStrataBSDFContext BSDFContext,
 				if (BSDF_GETHASSPECPROFILE(BSDFContext.BSDF))
 				{
 					const float NoR = saturate(dot(StrataEnvLightResult.SpecularDirection, BSDFContext.N));
-					StrataEnvLightResult.SpecularWeight *= EvaluateSpecularProfile(SLAB_SPECPROFILEID(BSDFContext.BSDF), BSDFContext.SatNoV, NoR, BSDFContext.SatNoV, NoR, BSDFContext.PixelCoord);
+					StrataEnvLightResult.SpecularWeight *= EvaluateSpecularProfile(SLAB_SPECPROFILEID(BSDFContext.BSDF), BSDFContext.SatNoV, NoR, BSDFContext.SatNoV, NoR);
 				}
 			#endif // STRATA_COMPLEXPATH
 
diff --git a/Engine/Source/Runtime/Renderer/Private/SceneRendering.cpp b/Engine/Source/Runtime/Renderer/Private/SceneRendering.cpp
--- a/Engine/Source/Runtime/Renderer/Private/SceneRendering.cpp
+++ b/Engine/Source/Runtime/Renderer/Private/SceneRendering.cpp
@@ -1670,7 +1670,7 @@ void FViewInfo::SetupUniformBufferParameters(
 		FIntVector TextureSize = Texture->GetSizeXYZ();
 		ViewUniformShaderParameters.SpecularProfileTextureSizeAndInvSize = FVector4f(TextureSize.X, TextureSize.Y, 1.0f / TextureSize.X, 1.0f / TextureSize.Y);
 		ViewUniformShaderParameters.SpecularProfileTexture = Texture;
-		ViewUniformShaderParameters.SpecularProfileSampler = TStaticSamplerState<SF_Point, AM_Clamp, AM_Clamp, AM_Clamp>::GetRHI();
+		ViewUniformShaderParameters.SpecularProfileSampler = TStaticSamplerState<SF_Bilinear, AM_Clamp, AM_Clamp, AM_Clamp>::GetRHI();
 	}
 
 	{
