Add ShaderPrint GPU string, which allow to upload string (e.g. asset name)

diff --git a/Engine/Shaders/Private/ShaderPrintCommon.ush b/Engine/Shaders/Private/ShaderPrintCommon.ush
--- a/Engine/Shaders/Private/ShaderPrintCommon.ush
+++ b/Engine/Shaders/Private/ShaderPrintCommon.ush
@@ -1342,3 +1342,73 @@ void AddOBBTWS(FShaderPrintContext Ctx, float3 Min, float3 Max, float4 Color, fl
 
 void AddOBBWS (FShaderPrintContext Ctx, float3 Min, float3 Max, float4 Color, float4x4 Transform)	{ Transform[3].xyz += Ctx.Config.TranslatedWorldOffset; AddOBBTWS(Ctx, Min, Max, Color, Transform); }
 void AddOBB   (FShaderPrintContext Ctx, float3 Min, float3 Max, float4 Color, float4x4 Transform)	{ Transform[3].xyz += Ctx.Config.TranslatedWorldOffset; AddOBBTWS(Ctx, Min, Max, Color, Transform); }
+
+///////////////////////////////////////////////////////////////////////////////////////////////////
+// Strings
+struct FStringInfo
+{
+	uint EntryID;
+	uint Length;
+	uint Offset;
+	uint Pad;
+};
+
+FStringInfo UnpackStringInfo(uint2 In)
+{
+	FStringInfo Out = (FStringInfo)0;
+	Out.EntryID = In.x;
+	Out.Offset  = (In.y) & 0xFFFF;
+	Out.Length  = (In.y >> 16) & 0xFF;
+	return Out;
+}
+
+FStringInfo FindStringInfo(uint InEntryID, uint InInfoCount, StructuredBuffer<uint2> InInfoBuffer)
+{
+	FStringInfo Out = (FStringInfo)0;
+	Out.EntryID = ~0;
+	for (uint It = 0; It < InInfoCount; ++It)
+	{
+		if (InInfoBuffer[It].x == InEntryID)
+		{
+			return UnpackStringInfo(InInfoBuffer[It]);
+		}
+	}
+	return Out;
+}
+
+void PrintString(inout FShaderPrintContext Ctx, uint InEntryID, FFontColor InColor, uint InInfoCount, uint InCharCount, StructuredBuffer<uint2> InInfoBuffer, Buffer<uint> InCharBuffer)
+{
+	const FStringInfo Info = FindStringInfo(InEntryID, InInfoCount, InInfoBuffer);
+	if (Info.EntryID != ~0 && (Info.Length + Info.Offset) <= InCharCount)
+	{
+		for (uint It = 0; It < Info.Length; ++It)
+		{
+			const uint Char = InCharBuffer[It + Info.Offset];
+			PrintSymbol(Ctx, Char, InColor);
+		}
+	}
+}
+
+uint LenString(uint InEntryID, uint InInfoCount, StructuredBuffer<uint2> InInfoBuffer)
+{
+	const FStringInfo Info = FindStringInfo(InEntryID, InInfoCount, InInfoBuffer);
+	return Info.EntryID != ~0 ? Info.Length : 0;
+}
+
+#define SHADERPRINT_STRINGS(N) \
+	uint					N##_InfoCount;\
+	uint					N##_CharCount;\
+	StructuredBuffer<uint2> N##_InfoBuffer;\
+	Buffer<uint>			N##_CharBuffer;\
+	void Print##N(inout FShaderPrintContext Ctx, uint InEntryID, FFontColor InColor)\
+	{\
+		PrintString(Ctx, InEntryID, InColor, N##_InfoCount, N##_CharCount, N##_InfoBuffer, N##_CharBuffer);\
+	}\
+	uint Len##N(uint InEntryID)\
+	{\
+		return LenString(InEntryID, N##_InfoCount, N##_InfoBuffer);\
+	}\
+	uint Num##N()\
+	{\
+		return N##_InfoCount;\
+	}
\ No newline at end of file
diff --git a/Engine/Source/Runtime/Renderer/Private/ShaderPrint.cpp b/Engine/Source/Runtime/Renderer/Private/ShaderPrint.cpp
--- a/Engine/Source/Runtime/Renderer/Private/ShaderPrint.cpp
+++ b/Engine/Source/Runtime/Renderer/Private/ShaderPrint.cpp
@@ -15,6 +15,7 @@
 #include "SceneRendering.h"
 #include "ScreenPass.h"
 #include "SystemTextures.h"
+#include "RenderGraphUtils.h"
 
 namespace ShaderPrint
 {
@@ -1117,4 +1118,54 @@ namespace ShaderPrint
 
 		GDefaultView = nullptr;
 	}
+
+
+	FStrings::FStrings(uint32 InMaxEntryCount, uint32 InMaxStringLength)
+	{
+		static_assert(sizeof(FEntryInfo) == 8);
+		Chars.Reserve(InMaxEntryCount * InMaxStringLength);
+	}
+	
+	void FStrings::Add(const FString& In, uint32 EntryID)
+	{
+		const uint32 Offset = Chars.Num();
+		const uint32 Length = In.Len();
+		for (TCHAR C : In)
+		{
+			Chars.Add(uint8(C));
+		}
+
+		FEntryInfo& Info = Infos.AddDefaulted_GetRef();
+		Info.EntryID = EntryID;
+		Info.Length = Length;
+		Info.Offset = Offset;
+	}
+
+	void FStrings::Add(const FString& In)
+	{
+		Add(In, Infos.Num());
+	}
+	
+	FStrings::FShaderParameters FStrings::GetParameters(FRDGBuilder& GraphBuilder)
+	{
+		if (Infos.IsEmpty())
+		{
+			FEntryInfo& Info = Infos.AddDefaulted_GetRef();
+			Info.EntryID = ~0;
+			Info.Length = 4;
+			Info.Offset = 0;
+			Chars.Add(uint8('N'));
+			Chars.Add(uint8('o'));
+			Chars.Add(uint8('n'));
+			Chars.Add(uint8('e'));
+		}
+	
+		FShaderParameters Out;
+		Out.InfoCount = Infos.Num();
+		Out.CharCount = Chars.Num();
+		Out.InfoBuffer = GraphBuilder.CreateSRV(CreateStructuredBuffer(GraphBuilder, TEXT("ShaderPrint.Strings.Infos"), Infos));
+		Out.CharBuffer = GraphBuilder.CreateSRV(CreateVertexBuffer(GraphBuilder, TEXT("ShaderPrint.Strings.Chars"), FRDGBufferDesc::CreateBufferDesc(1, Chars.Num()), Chars.GetData(), Chars.Num()), PF_R8_UINT);
+
+		return Out;
+	}
 }
\ No newline at end of file
diff --git a/Engine/Source/Runtime/Renderer/Public/ShaderPrintParameters.h b/Engine/Source/Runtime/Renderer/Public/ShaderPrintParameters.h
--- a/Engine/Source/Runtime/Renderer/Public/ShaderPrintParameters.h
+++ b/Engine/Source/Runtime/Renderer/Public/ShaderPrintParameters.h
@@ -141,6 +141,46 @@ namespace ShaderPrint
 	
 	UE_DEPRECATED(5.1, "Use one of the other implementations of SetParameters()")
 	void SetParameters(FRDGBuilder& GraphBuilder, const FViewInfo& View, FShaderParameters& OutParameters);
+
+	// Experimental GPU string
+	//
+	// In Cpp:
+	//   SHADER_PARAMETER_STRUCT(ShaderPrint::FStrings::FShaderPararameters, MyVariable)
+	//   ShaderPrint::FStrings MyVariable;
+	//	 MyVariable.Add(FString(...), StringId);
+	//
+	// In shader:
+	//   SHADERPRINT_STRINGS(MyVariable)
+	//   void foo()
+	//   {  
+	//     InitShaderPrintContext Ctx = InitShaderPrintContext(...);
+	//     PrintMyVariable(Ctx, StringId, FontWhite);
+	//   }
+	struct RENDERER_API FStrings
+	{
+	public:
+		BEGIN_SHADER_PARAMETER_STRUCT(FShaderParameters, )
+			SHADER_PARAMETER(uint32, InfoCount)
+			SHADER_PARAMETER(uint32, CharCount)
+			SHADER_PARAMETER_RDG_BUFFER_SRV(StructuredBuffer<uint2>, InfoBuffer)
+			SHADER_PARAMETER_RDG_BUFFER_SRV(Buffer<uint8>, CharBuffer)
+		END_SHADER_PARAMETER_STRUCT()
+
+		FStrings(uint32 InAvgEntryCount=128u, uint32 InAvgStringLength=32u);
+		void Add(const FString& In, uint32 EntryID);
+		void Add(const FString& In);
+		FShaderParameters GetParameters(FRDGBuilder& GraphBuilder);
+	private:
+		struct FEntryInfo
+		{
+			uint32 EntryID;
+			uint16 Offset;
+			uint8  Length;
+			uint8  Pad0;
+		};
+		TArray<FEntryInfo> Infos;
+		TArray<uint8> Chars;
+	};
 }
 
 /** 
@@ -172,4 +212,4 @@ struct FFrozenShaderPrintData
 
 	TRefCountPtr<FRDGPooledBuffer> ShaderPrintEntryBuffer;
 	TRefCountPtr<FRDGPooledBuffer> ShaderPrintStateBuffer;
-};
+};
\ No newline at end of file
