Fix /Engine/Private/ShadingModelsSampling.ush(98,21-62): error X3206: 'ImportanceSampleVisibleGGX': implicit truncation of vector type error

diff --git a/Engine/Shaders/Private/ShadingModelsSampling.ush b/Engine/Shaders/Private/ShadingModelsSampling.ush
--- a/Engine/Shaders/Private/ShadingModelsSampling.ush
+++ b/Engine/Shaders/Private/ShadingModelsSampling.ush
@@ -95,7 +95,7 @@ FBxDFSample SampleDefaultLitBxDF(uint TermMask, float3 WorldNormal, float Roughn
 #if 1
 		// PDF = G_SmithV * VoH * D / NoV / (4 * VoH)
 		// PDF = G_SmithV * D / (4 * NoV);
-		float4 TangentH = ImportanceSampleVisibleGGX(E, a, TangentV);
+		float4 TangentH = ImportanceSampleVisibleGGX(E.xy, a, TangentV);
 #else
 		// PDF = D * NoH / (4 * VoH),
 		float4 TangentH = ImportanceSampleGGX(E.xy, a * a);
