Add Groom memory stats.

diff --git a/Engine/Plugins/Runtime/HairStrands/Source/HairStrandsCore/Private/GroomDebug.cpp b/Engine/Plugins/Runtime/HairStrands/Source/HairStrandsCore/Private/GroomDebug.cpp
--- a/Engine/Plugins/Runtime/HairStrands/Source/HairStrandsCore/Private/GroomDebug.cpp
+++ b/Engine/Plugins/Runtime/HairStrands/Source/HairStrandsCore/Private/GroomDebug.cpp
@@ -19,6 +19,7 @@
 #include "GroomCacheData.h"
 #include "HairStrandsClusterCulling.h"
 #include "ShaderPrintParameters.h"
+#include "GroomComponent.h"
 
 ///////////////////////////////////////////////////////////////////////////////////////////////////
 
@@ -1043,6 +1044,265 @@ static void AddHairDebugPrintInstancePass(
 	}
 }
 
+///////////////////////////////////////////////////////////////////////////////////////////////////
+
+class FHairDebugPrintMemoryCS : public FGlobalShader
+{
+	DECLARE_GLOBAL_SHADER(FHairDebugPrintMemoryCS);
+	SHADER_USE_PARAMETER_STRUCT(FHairDebugPrintMemoryCS, FGlobalShader);
+
+	BEGIN_SHADER_PARAMETER_STRUCT(FParameters, )
+		SHADER_PARAMETER(uint32, ComponentCount)
+		SHADER_PARAMETER(uint32, GroomCount)
+		SHADER_PARAMETER(uint32, BindingCount)
+		SHADER_PARAMETER_STRUCT(ShaderPrint::FStrings::FShaderParameters, ComponentNames)
+		SHADER_PARAMETER_STRUCT(ShaderPrint::FStrings::FShaderParameters, GroomNames)
+		SHADER_PARAMETER_STRUCT(ShaderPrint::FStrings::FShaderParameters, BindingNames)
+		SHADER_PARAMETER_RDG_BUFFER_SRV(Buffer<uint>, ComponentBuffer)
+		SHADER_PARAMETER_RDG_BUFFER_SRV(Buffer<uint>, GroomBuffer)
+		SHADER_PARAMETER_RDG_BUFFER_SRV(Buffer<uint>, BindingBuffer)
+		SHADER_PARAMETER_STRUCT_INCLUDE(ShaderPrint::FShaderParameters, ShaderPrintUniformBuffer)
+	END_SHADER_PARAMETER_STRUCT()
+
+public:
+	static bool ShouldCompilePermutation(const FGlobalShaderPermutationParameters& Parameters) { return IsHairStrandsSupported(EHairStrandsShaderType::Tool, Parameters.Platform); }
+	static void ModifyCompilationEnvironment(const FGlobalShaderPermutationParameters& Parameters, FShaderCompilerEnvironment& OutEnvironment)
+	{
+		// Skip optimization for avoiding long compilation time due to large UAV writes
+		FGlobalShader::ModifyCompilationEnvironment(Parameters, OutEnvironment);
+		OutEnvironment.CompilerFlags.Add(CFLAG_Debug);
+		OutEnvironment.SetDefine(TEXT("SHADER_PRINT_MEMORY"), 1);
+	}
+};
+
+IMPLEMENT_GLOBAL_SHADER(FHairDebugPrintMemoryCS, "/Engine/Private/HairStrands/HairStrandsDebug.usf", "MainCS", SF_Compute);
+
+const TCHAR* GetHairAttributeText(EHairAttribute In, uint32 InFlags);
+uint32 GetHairAttributeIndex(EHairAttribute In);
+EGroomCacheType GetHairInstanceCacheType(const FHairGroupInstance* Instance);
+
+static void AddHairDebugPrintMemoryPass(
+	FRDGBuilder& GraphBuilder, 
+	FGlobalShaderMap* ShaderMap,
+	const FSceneView& View,
+	const FShaderPrintData* ShaderPrintData)
+{
+
+	
+#if 0
+	// Streaming enabled
+	LOD Index | Auto LOD
+		rest
+		deformed
+		cluster
+		root
+
+		Total CPU | Total GPU
+#endif
+
+	if (!ShaderPrintData) { return; }
+
+	struct FInfos
+	{
+		FUintVector4 Data0 = { 0,0,0,0 };
+		FUintVector4 Data1 = { 0,0,0,0 };
+		FUintVector4 Data2 = { 0,0,0,0 };
+		FUintVector4 Data3 = { 0,0,0,0 };
+		FUintVector4 Data4 = { 0,0,0,0 };
+	};
+	const uint32 InfoInBytes  = sizeof(FInfos);
+	const uint32 InfoInUints = sizeof(FInfos) / sizeof(uint32);
+
+	// Component
+	uint32 ComponentCount = 0;
+	for (TObjectIterator<UGroomComponent> ComponentIt; ComponentIt; ++ComponentIt)
+	{
+		if (ComponentIt)
+		{
+			ComponentCount += ComponentIt->GetGroupCount();
+		}
+	}
+	ShaderPrint::FStrings ComponentNames(ComponentCount);
+	TArray<FInfos> ComponentBuffer;
+	ComponentBuffer.Reserve(ComponentCount);
+	for (TObjectIterator<UGroomComponent> ComponentIt; ComponentIt; ++ComponentIt)
+	{
+		if (ComponentIt)
+		{
+			const uint32 GroupCount = ComponentIt->GetGroupCount();
+			for (uint32 GroupIt = 0; GroupIt < GroupCount; ++GroupIt)
+			{
+				const FHairGroupInstance* Instance = ComponentIt->GetGroupInstance(GroupIt);
+				const FGroomComponentMemoryStats MemoryStats = FGroomComponentMemoryStats::Get(Instance);
+		
+				ComponentNames.Add(*Instance->Debug.GroomAssetName);
+		
+				FInfos& D = ComponentBuffer.AddDefaulted_GetRef();
+				D.Data0.X = Instance->Debug.GroupIndex;
+				D.Data0.Y = Instance->Debug.GroupCount;
+				D.Data0.Z = Instance->HairGroupPublicData->GetActiveStrandsCurveCount();
+				D.Data0.W = Instance->Strands.RestResource ? Instance->Strands.RestResource->BulkData.Header.CurveCount : 0;
+		
+				D.Data1.X = MemoryStats.Guides;
+				D.Data1.Y = MemoryStats.Strands;
+				D.Data1.Z = MemoryStats.Cards;
+				D.Data1.W = MemoryStats.Meshes;
+		
+				D.Data2.X = MemoryStats.Guides;
+				D.Data2.Y = MemoryStats.Strands;
+				D.Data2.Z = MemoryStats.Cards;
+				D.Data2.W = MemoryStats.Meshes;
+
+				D.Data3.X = MemoryStats.Common;
+				D.Data3.Y = 0;
+				D.Data3.Z = 0;
+				D.Data3.W = 0;
+			}
+		}
+	}
+
+	// Groom
+	uint32 GroomCount = 0;
+	for (TObjectIterator<UGroomAsset> AssetIt; AssetIt; ++AssetIt)
+	{
+		if (AssetIt)
+		{
+			GroomCount += AssetIt->GetHairGroupsPlatformData().Num();
+		}
+	}
+	ShaderPrint::FStrings GroomNames(GroomCount);
+	TArray<FInfos> GroomBuffer;
+	GroomBuffer.Reserve(GroomCount);
+	for (uint32 It = 0; It < GroomCount; ++It)
+	{
+		for (TObjectIterator<UGroomAsset> AssetIt; AssetIt; ++AssetIt)
+		{
+			if (AssetIt)
+			{
+				const uint32 GroupCount = AssetIt->GetHairGroupsPlatformData().Num();
+				for (uint32 GroupIt = 0; GroupIt < GroupCount; ++GroupIt)
+				{					
+					const FHairGroupPlatformData& Data = AssetIt->GetHairGroupsPlatformData()[GroupIt];
+					const FGroomAssetMemoryStats MemoryStats = FGroomAssetMemoryStats::Get(Data);
+
+					GroomNames.Add(AssetIt->GetName());
+
+					FInfos& D = GroomBuffer.AddDefaulted_GetRef();
+					D.Data0.X = GroupIt;
+					D.Data0.Y = GroupCount;
+					D.Data0.Z = Data.Strands.RestResource ? Data.Strands.RestResource->MaxAvailableCurveCount : 0;
+					D.Data0.W = Data.Strands.BulkData.IsValid() ? Data.Strands.BulkData.GetNumCurves() : 0;
+			
+					D.Data1.X = MemoryStats.CPU.Guides;
+					D.Data1.Y = MemoryStats.CPU.Strands;
+					D.Data1.Z = MemoryStats.CPU.Cards;
+					D.Data1.W = MemoryStats.CPU.Meshes;
+			
+					D.Data2.X = MemoryStats.GPU.Guides;
+					D.Data2.Y = MemoryStats.GPU.Strands;
+					D.Data2.Z = MemoryStats.GPU.Cards;
+					D.Data2.W = MemoryStats.GPU.Meshes;
+
+					D.Data3.X = MemoryStats.Memory.Rest;
+					D.Data3.Y = MemoryStats.Memory.Interpolation;
+					D.Data3.Z = MemoryStats.Memory.Cluster;
+					D.Data3.W = MemoryStats.Memory.Raytracing;
+
+					D.Data4.X = MemoryStats.Curves.Rest;
+					D.Data4.Y = MemoryStats.Curves.Interpolation;
+					D.Data4.Z = MemoryStats.Curves.Cluster;
+					D.Data4.W = MemoryStats.Curves.Raytracing;
+				}
+			}
+		}
+	}
+
+	// Binding
+	uint32 BindingCount = 0;
+	for (TObjectIterator<UGroomBindingAsset> AssetIt; AssetIt; ++AssetIt)
+	{
+		if (AssetIt)
+		{
+			BindingCount += AssetIt->GetHairGroupsPlatformData().Num();
+		}
+	}
+	ShaderPrint::FStrings BindingNames(BindingCount);
+	TArray<FInfos> BindingBuffer;
+	BindingBuffer.Reserve(BindingCount);
+	for (uint32 It = 0; It < BindingCount; ++It)
+	{
+		for (TObjectIterator<UGroomBindingAsset> AssetIt; AssetIt; ++AssetIt)
+		{
+			if (AssetIt)
+			{
+				const uint32 GroupCount = AssetIt->GetHairGroupsPlatformData().Num();
+				for (uint32 GroupIt = 0; GroupIt < GroupCount; ++GroupIt)
+				{
+					const UGroomBindingAsset::FHairGroupPlatformData& Data = AssetIt->GetHairGroupsPlatformData()[GroupIt];
+					const UGroomBindingAsset::FHairGroupResource& Resource = AssetIt->GetHairGroupResources()[GroupIt];
+					const FGroomBindingAssetMemoryStats MemoryStats = FGroomBindingAssetMemoryStats::Get(Data, Resource);
+			
+
+					BindingNames.Add(*AssetIt->GetName(), It);
+			
+					FInfos& D = BindingBuffer.AddDefaulted_GetRef();
+					D.Data0.X = GroupIt;
+					D.Data0.Y = GroupCount;
+					D.Data0.Z = Resource.RenRootResources ? Resource.RenRootResources->MaxAvailableCurveCount : 0;
+					D.Data0.W = Resource.RenRootResources ? Resource.RenRootResources->GetRootCount() : 0;
+			
+					D.Data1.X = MemoryStats.CPU.Guides;
+					D.Data1.Y = MemoryStats.CPU.Strands;
+					D.Data1.Z = MemoryStats.CPU.Cards;
+					D.Data1.W = 0;
+			
+					D.Data2.X = MemoryStats.GPU.Guides;
+					D.Data2.Y = MemoryStats.GPU.Strands;
+					D.Data2.Z = MemoryStats.GPU.Cards;
+					D.Data2.W = 0;
+				}
+			}
+		}
+	}
+
+	// Empty cases
+	if (ComponentCount == 0 && GroomCount == 0 && BindingCount == 0) return;
+	if (ComponentBuffer.IsEmpty())	{ ComponentBuffer.AddDefaulted(); }
+	if (GroomBuffer.IsEmpty())		{ GroomBuffer.AddDefaulted(); }
+	if (BindingBuffer.IsEmpty())	{ BindingBuffer.AddDefaulted(); }
+
+	// Force ShaderPrint on.
+	ShaderPrint::SetEnabled(true);
+	ShaderPrint::RequestSpaceForCharacters((ComponentCount + GroomCount + BindingCount) * 256 + 512);
+
+	FHairDebugPrintMemoryCS::FParameters* Parameters = GraphBuilder.AllocParameters<FHairDebugPrintMemoryCS::FParameters>();
+	Parameters->ComponentCount  = ComponentCount;
+	Parameters->ComponentNames  = ComponentNames.GetParameters(GraphBuilder);
+	Parameters->ComponentBuffer = GraphBuilder.CreateSRV(CreateVertexBuffer(GraphBuilder, TEXT("Hair.Debug.Component"), FRDGBufferDesc::CreateBufferDesc(4, InfoInUints * ComponentBuffer.Num()), ComponentBuffer.GetData(), InfoInBytes * ComponentBuffer.Num()), PF_R32_UINT);
+
+	Parameters->GroomCount		= GroomCount;
+	Parameters->GroomNames  	= GroomNames.GetParameters(GraphBuilder);
+	Parameters->GroomBuffer 	= GraphBuilder.CreateSRV(CreateVertexBuffer(GraphBuilder, TEXT("Hair.Debug.Groom"), FRDGBufferDesc::CreateBufferDesc(4, InfoInUints * GroomBuffer.Num()), GroomBuffer.GetData(), InfoInBytes * GroomBuffer.Num()), PF_R32_UINT);
+
+	Parameters->BindingCount  	= BindingCount;
+	Parameters->BindingNames  	= BindingNames.GetParameters(GraphBuilder);
+	Parameters->BindingBuffer 	= GraphBuilder.CreateSRV(CreateVertexBuffer(GraphBuilder, TEXT("Hair.Debug.Binding"), FRDGBufferDesc::CreateBufferDesc(4, InfoInUints * BindingBuffer.Num()), BindingBuffer.GetData(), InfoInBytes * BindingBuffer.Num()), PF_R32_UINT);
+
+	ShaderPrint::SetParameters(GraphBuilder, *ShaderPrintData, Parameters->ShaderPrintUniformBuffer);
+	TShaderMapRef<FHairDebugPrintMemoryCS> ComputeShader(ShaderMap);
+
+	ClearUnusedGraphResources(ComputeShader, Parameters);
+
+	FComputeShaderUtils::AddPass(
+		GraphBuilder,
+		RDG_EVENT_NAME("HairStrands::DebugMemory(Info,Instances:%d, Groom:%d, Binding:%d)", ComponentCount, GroomCount, BindingCount),
+		ComputeShader,
+		Parameters,
+		FIntVector(1, 1, 1));
+}
+
+///////////////////////////////////////////////////////////////////////////////////////////////////
+
 void RunHairStrandsDebug(
 	FRDGBuilder& GraphBuilder,
 	FGlobalShaderMap* ShaderMap,
@@ -1183,4 +1443,9 @@ void RunHairStrandsDebug(
 			}
 		}
 	}
+
+	if (ViewMode == EGroomViewMode::Memory)
+	{
+		AddHairDebugPrintMemoryPass(GraphBuilder, ShaderMap, View, ShaderPrintData);
+	}
 }
diff --git a/Engine/Shaders/Private/HairStrands/HairStrandsDebug.usf b/Engine/Shaders/Private/HairStrands/HairStrandsDebug.usf
--- a/Engine/Shaders/Private/HairStrands/HairStrandsDebug.usf
+++ b/Engine/Shaders/Private/HairStrands/HairStrandsDebug.usf
@@ -4,7 +4,7 @@
 #include "../HairShadingCommon.ush"
 #include "../ColorMap.ush"
 
-#if SHADER_PRINT_INSTANCE || SHADER_PRINT
+#if SHADER_PRINT_INSTANCE || SHADER_PRINT || SHADER_PRINT_MEMORY
 /////////////////////////////////////////////////////////////////////////////////////////
 // Utils function
 
@@ -2479,3 +2479,360 @@ void CSMain(
 }
 
 #endif // SHADER_ATTRIBUTE_DEBUG
+
+
+////////////////////////////////////////////////////////////////////////////////////////////////
+// Print memory
+
+#if SHADER_PRINT_MEMORY
+
+FSTRINGS(GroomNames)
+FSTRINGS(BindingNames)
+FSTRINGS(ComponentNames)
+
+uint ComponentCount;
+uint GroomCount;
+uint BindingCount;
+
+Buffer<uint> ComponentBuffer;
+Buffer<uint> GroomBuffer;
+Buffer<uint> BindingBuffer;
+
+struct FInfo
+{
+	uint4 Data0;
+	uint4 Data1;
+	uint4 Data2;
+	uint4 Data3;
+	uint4 Data4;
+};
+
+FInfo ReadInfo(uint Index, Buffer<uint> InBuffer)
+{
+	const uint Index20 = Index * 20u;
+	FInfo Out;
+	Out.Data0.x = InBuffer[Index20 + 0];
+	Out.Data0.y = InBuffer[Index20 + 1];
+	Out.Data0.z = InBuffer[Index20 + 2];
+	Out.Data0.w = InBuffer[Index20 + 3];
+	Out.Data1.x = InBuffer[Index20 + 4];
+	Out.Data1.y = InBuffer[Index20 + 5];
+	Out.Data1.z = InBuffer[Index20 + 6];
+	Out.Data1.w = InBuffer[Index20 + 7];
+	Out.Data2.x = InBuffer[Index20 + 8];
+	Out.Data2.y = InBuffer[Index20 + 9];
+	Out.Data2.z = InBuffer[Index20 + 10];
+	Out.Data2.w = InBuffer[Index20 + 11];
+	Out.Data3.x = InBuffer[Index20 + 12];
+	Out.Data3.y = InBuffer[Index20 + 13];
+	Out.Data3.z = InBuffer[Index20 + 14];
+	Out.Data3.w = InBuffer[Index20 + 15];
+	Out.Data4.x = InBuffer[Index20 + 16];
+	Out.Data4.y = InBuffer[Index20 + 17];
+	Out.Data4.z = InBuffer[Index20 + 18];
+	Out.Data4.w = InBuffer[Index20 + 19];
+	return Out;
+}
+
+struct FComponent
+{
+	uint GroupIndex;
+	uint GroupCount;
+	uint RestCurve;
+	uint ActiveCurve;
+
+	uint GPU_Common;
+	uint GPU_Guides;
+	uint GPU_Strands;
+	uint GPU_Cards;
+	uint GPU_Meshes;
+};
+
+FComponent GetComponent(uint Index)
+{	
+	const FInfo In = ReadInfo(Index, ComponentBuffer);
+
+	FComponent Out;
+	Out.GroupIndex 			=In.Data0.x;
+	Out.GroupCount 			=In.Data0.y;
+	Out.ActiveCurve			=In.Data0.z;
+	Out.RestCurve			=In.Data0.w;
+	
+	Out.GPU_Guides 			=In.Data1.x;
+	Out.GPU_Strands 		=In.Data1.y;
+	Out.GPU_Cards 			=In.Data1.z;
+	Out.GPU_Meshes 			=In.Data1.w;
+	Out.GPU_Common			=In.Data2.x;
+
+	return Out;
+}
+
+struct FGroom
+{
+	uint GroupIndex;
+	uint GroupCount;
+	uint RestCurve;
+	uint ActiveCurve;
+
+	uint CPU_Guides;
+	uint CPU_Strands;
+	uint CPU_Cards;
+	uint CPU_Meshes;
+
+	uint GPU_Guides;
+	uint GPU_Strands;
+	uint GPU_Cards;
+	uint GPU_Meshes;
+
+	uint GPU_Rest;
+	uint GPU_Interpolation;
+	uint GPU_Cluster;
+	uint GPU_Raytracing;
+
+	uint Curve_Rest;
+	uint Curve_Interpolation;
+	uint Curve_Cluster;
+	uint Curve_Raytracing;
+};
+
+FGroom GetGroom(uint Index)
+{	
+	const FInfo In = ReadInfo(Index, GroomBuffer);
+
+	FGroom Out;
+	Out.GroupIndex 			= In.Data0.x;
+	Out.GroupCount 			= In.Data0.y;
+	Out.ActiveCurve			= In.Data0.z;
+	Out.RestCurve			= In.Data0.w;
+
+	Out.CPU_Guides 			= In.Data1.x;
+	Out.CPU_Strands 		= In.Data1.y;
+	Out.CPU_Cards 			= In.Data1.z;
+	Out.CPU_Meshes 			= In.Data1.w;
+
+	Out.GPU_Guides 			= In.Data2.x;
+	Out.GPU_Strands 		= In.Data2.y;
+	Out.GPU_Cards 			= In.Data2.z;
+	Out.GPU_Meshes 			= In.Data2.w;
+
+	Out.GPU_Rest			= In.Data3.x;
+	Out.GPU_Interpolation	= In.Data3.y;
+	Out.GPU_Cluster			= In.Data3.z;
+	Out.GPU_Raytracing		= In.Data3.w;
+
+	Out.Curve_Rest			= In.Data4.x;
+	Out.Curve_Interpolation	= In.Data4.x;
+	Out.Curve_Cluster     	= In.Data4.x;
+	Out.Curve_Raytracing  	= In.Data4.x;
+
+	return Out;
+}
+
+struct FBinding
+{
+	uint GroupIndex;
+	uint GroupCount;
+	uint RestCurve;
+	uint ActiveCurve;
+
+	uint CPU_Guides;
+	uint CPU_Strands;
+	uint CPU_Cards;
+	uint CPU_Meshes;
+
+	uint GPU_Guides;
+	uint GPU_Strands;
+	uint GPU_Cards;
+	uint GPU_Meshes;
+};
+
+FBinding GetBinding(uint Index)
+{	
+	const FInfo In = ReadInfo(Index, BindingBuffer);
+
+	FBinding Out;
+	Out.GroupIndex 	= In.Data0.x;
+	Out.GroupCount 	= In.Data0.y;
+	Out.ActiveCurve	= In.Data0.z;
+	Out.RestCurve	= In.Data0.w;
+
+	Out.CPU_Guides 	= In.Data1.x;
+	Out.CPU_Strands = In.Data1.y;
+	Out.CPU_Cards 	= In.Data1.z;
+	Out.CPU_Meshes 	= In.Data1.w;
+
+	Out.GPU_Guides 	= In.Data2.x;
+	Out.GPU_Strands = In.Data2.y;
+	Out.GPU_Cards 	= In.Data2.z;
+	Out.GPU_Meshes 	= In.Data2.w;
+	return Out;
+}
+
+void PrintGroup(inout FShaderPrintContext Ctx, uint GroupIndex, uint GroupCount)
+{
+	Print(Ctx, GroupIndex,	FontEmerald, 2, 0);	
+	Print(Ctx, TEXT("/"), FontWhite);
+	Print(Ctx, GroupCount,	FontEmerald, 2, 0);
+}
+
+FFontColor GetOccupancyColor(float In)
+{
+	float3 Color = lerp(float3(0, 1, 0), float3(1, 0, 0), saturate(In));
+	return InitFontColor(Color);
+}
+
+void PrintCurve(inout FShaderPrintContext Ctx, uint ActiveCurve, uint RestCurve)
+{
+	const float Curve = float(ActiveCurve) / float(RestCurve);
+	PrintPercentage(Ctx, Curve, GetOccupancyColor(Curve));
+	//PrintLargeNumber(Ctx, Info.ActiveCurve,	FontOrange);	Print(Ctx, TEXT(" "), FontWhite);
+	//Print(Ctx, TEXT("/"), FontWhite);
+	//PrintLargeNumber(Ctx, Info.RestCurve,	FontOrange);	Print(Ctx, TEXT(" "), FontWhite);
+}
+
+[numthreads(1, 1, 1)]
+void MainCS(uint GroupIndex : SV_GroupIndex, uint3 DispatchThreadId : SV_DispatchThreadID)
+{
+	// Draw AABB of the instance groups
+	if (any(DispatchThreadId != 0))
+	{
+		return;
+	}
+
+	// Widgets
+	FShaderPrintContext Ctx2 = InitShaderPrintContext(true, uint2(900, 100));
+	Print(Ctx2, TEXT("Controls"), FontYellow); Newline(Ctx2);
+	const bool bComponent 	= AddCheckbox(Ctx2, TEXT("Components"), false, FontWhite); 	Newline(Ctx2);	
+	const bool bGroom     	= AddCheckbox(Ctx2, TEXT("Grooms"), true, FontWhite);     	Newline(Ctx2);	
+	const bool bBinding   	= AddCheckbox(Ctx2, TEXT("Bindings"), false, FontWhite);   	Newline(Ctx2);
+	const bool bStrandsOnly = AddCheckbox(Ctx2, TEXT("Strands Only"), true, FontWhite); Newline(Ctx2);
+
+	// Component
+	FShaderPrintContext Ctx = InitShaderPrintContext(true, uint2(50, 100));
+	uint Component_TotalGPU = 0;
+	if (bComponent)
+	{
+		Print(Ctx, TEXT("Components"), FontYellow); Newline(Ctx);
+		{
+			Print(Ctx, TEXT("Grp. Common Guides Strands Cards Meshes Total  Curv Name"), FontSilver); Newline(Ctx);
+			for (uint It = 0; It < ComponentCount; ++It)
+			{
+				const FComponent Info = GetComponent(It);
+				const uint Total
+					= Info.GPU_Common
+					+ Info.GPU_Guides
+					+ Info.GPU_Strands
+					+ Info.GPU_Cards;
+
+
+				PrintGroup(Ctx, Info.GroupIndex, Info.GroupCount);
+	
+				PrintLargeNumber(Ctx, Info.GPU_Common,	FontWhite);		Print(Ctx, TEXT(" "), FontWhite);
+				PrintLargeNumber(Ctx, Info.GPU_Guides,	FontWhite);		Print(Ctx, TEXT(" "), FontWhite);
+				PrintLargeNumber(Ctx, Info.GPU_Strands,	FontWhite);		Print(Ctx, TEXT(" "), FontWhite);
+				PrintLargeNumber(Ctx, Info.GPU_Cards,	FontWhite);		Print(Ctx, TEXT(" "), FontWhite);
+				PrintLargeNumber(Ctx, Info.GPU_Meshes,	FontWhite);		Print(Ctx, TEXT(" "), FontWhite);
+				PrintLargeNumber(Ctx, Total,			FontOrange);	Print(Ctx, TEXT(" "), FontWhite);
+				PrintCurve(Ctx, Info.ActiveCurve, Info.RestCurve);
+	
+				PrintComponentNames(Ctx, It, FontWhite);
+				Newline(Ctx);
+			}
+		}
+		Newline(Ctx);
+	}
+	// Total
+
+	// Grooms
+	if (bGroom)
+	{
+		Print(Ctx, TEXT("Grooms"), FontYellow); Newline(Ctx);
+		if (!bStrandsOnly)
+		{
+			Print(Ctx, TEXT("Grp. Guides Strands Cards Meshes Total Curv Name"), FontSilver); Newline(Ctx);
+			for (uint It = 0; It < GroomCount; ++It)
+			{
+				const FGroom Info = GetGroom(It);
+				PrintGroup(Ctx, Info.GroupIndex, Info.GroupCount);
+	
+				PrintLargeNumber(Ctx, Info.GPU_Guides, FontWhite);	Print(Ctx, TEXT(" "), FontWhite);
+				PrintLargeNumber(Ctx, Info.GPU_Strands, FontWhite);Print(Ctx, TEXT(" "), FontWhite);
+				PrintLargeNumber(Ctx, Info.GPU_Cards, FontWhite);	Print(Ctx, TEXT(" "), FontWhite);
+				PrintLargeNumber(Ctx, Info.GPU_Meshes, FontWhite);	Print(Ctx, TEXT(" "), FontWhite);
+				PrintCurve(Ctx, Info.ActiveCurve, Info.RestCurve);
+				PrintGroomNames(Ctx, It, FontWhite);
+				Newline(Ctx);
+			}
+		}
+		else
+		{
+			FFontColor FontMemory = FontWhite;
+			FFontColor FontCurve = FontSilver;
+
+			Print(Ctx, TEXT("Grp. Rest.  Intrp. Cluster  RT.  Total  Curv Name"), FontSilver); Newline(Ctx);
+			for (uint It = 0; It < GroomCount; ++It)
+			{
+				const FGroom Info = GetGroom(It);
+				const uint GPUStrandsTotal 
+					= Info.GPU_Rest
+					+ Info.GPU_Interpolation
+					+ Info.GPU_Cluster
+					+ Info.GPU_Raytracing;
+
+				PrintGroup(Ctx, Info.GroupIndex, Info.GroupCount);
+
+				PrintLargeNumber(Ctx, Info.GPU_Rest, FontMemory);	Print(Ctx, TEXT(" "), FontMemory);
+				PrintLargeNumber(Ctx, Info.GPU_Interpolation, FontMemory);Print(Ctx, TEXT(" "), FontMemory);
+				PrintLargeNumber(Ctx, Info.GPU_Cluster, FontMemory);	Print(Ctx, TEXT(" "), FontMemory);
+				PrintLargeNumber(Ctx, Info.GPU_Raytracing, FontMemory);	Print(Ctx, TEXT(" "), FontMemory);
+				PrintLargeNumber(Ctx, GPUStrandsTotal, FontOrange);	Print(Ctx, TEXT(" "), FontMemory);
+				PrintCurve(Ctx, Info.ActiveCurve, Info.RestCurve);
+				PrintGroomNames(Ctx, It, FontMemory);
+				Newline(Ctx);
+
+				// Curve count detail per resource type
+				#if 0
+				Print(Ctx, TEXT("     "), FontCurve);
+				PrintLargeNumber(Ctx, Info.Curve_Rest, FontCurve);	Print(Ctx, TEXT(" "), FontCurve);
+				PrintLargeNumber(Ctx, Info.Curve_Interpolation, FontCurve);Print(Ctx, TEXT(" "), FontCurve);
+				PrintLargeNumber(Ctx, Info.Curve_Cluster, FontCurve);	Print(Ctx, TEXT(" "), FontCurve);
+				PrintLargeNumber(Ctx, Info.Curve_Raytracing, FontCurve);	Print(Ctx, TEXT(" "), FontCurve);
+				Newline(Ctx);
+				#endif
+			}
+		}
+		Newline(Ctx);
+	}
+	// Total
+
+	// Bindings
+	if (bBinding)
+	{
+		Print(Ctx, TEXT("Bindings"), FontYellow); Newline(Ctx);
+		Print(Ctx, TEXT("Grp. Guides Strands Cards Meshes Total Curv Name"), FontSilver); Newline(Ctx);
+		for (uint It = 0; It < BindingCount; ++It)
+		{
+			const FBinding Info = GetBinding(It);
+			const uint Total
+				= Info.GPU_Guides
+				+ Info.GPU_Strands
+				+ Info.GPU_Cards
+				+ Info.GPU_Meshes;
+
+			PrintGroup(Ctx, Info.GroupIndex, Info.GroupCount);
+
+			PrintLargeNumber(Ctx, Info.GPU_Guides, FontWhite);	Print(Ctx, TEXT(" "), FontWhite);
+			PrintLargeNumber(Ctx, Info.GPU_Strands, FontWhite);	Print(Ctx, TEXT(" "), FontWhite);
+			PrintLargeNumber(Ctx, Info.GPU_Cards, FontWhite);	Print(Ctx, TEXT(" "), FontWhite);
+			PrintLargeNumber(Ctx, Info.GPU_Meshes, FontWhite);	Print(Ctx, TEXT(" "), FontWhite);
+			PrintLargeNumber(Ctx, Total, FontOrange);			Print(Ctx, TEXT(" "), FontWhite);
+			PrintCurve(Ctx, Info.ActiveCurve, Info.RestCurve);
+			PrintBindingNames(Ctx, It, FontWhite);
+			Newline(Ctx);
+		}
+		Newline(Ctx);
+	}
+	// Total
+}
+
+#endif // SHADER_PRINT_MEMORY
diff --git a/Engine/Source/Editor/CommonMenuExtensions/Private/GroomVisualizationMenuCommands.cpp b/Engine/Source/Editor/CommonMenuExtensions/Private/GroomVisualizationMenuCommands.cpp
--- a/Engine/Source/Editor/CommonMenuExtensions/Private/GroomVisualizationMenuCommands.cpp
+++ b/Engine/Source/Editor/CommonMenuExtensions/Private/GroomVisualizationMenuCommands.cpp
@@ -128,6 +128,7 @@ void FGroomVisualizationMenuCommands::InternalBuildVisualisationSubMenu(FMenuBui
 				Commands.AddCommandTypeToMenu(Menu, EGroomViewMode::Coverage);
 				Commands.AddCommandTypeToMenu(Menu, EGroomViewMode::Tile);
 			}
+			Commands.AddCommandTypeToMenu(Menu, EGroomViewMode::Memory);
 			Menu.EndSection();
 		}
 	}
diff --git a/Engine/Source/Runtime/Engine/Private/GroomVisualizationData.cpp b/Engine/Source/Runtime/Engine/Private/GroomVisualizationData.cpp
--- a/Engine/Source/Runtime/Engine/Private/GroomVisualizationData.cpp
+++ b/Engine/Source/Runtime/Engine/Private/GroomVisualizationData.cpp
@@ -73,6 +73,7 @@ EGroomViewMode GetGroomViewMode(const FSceneView& View)
 		case 37: return EGroomViewMode::CardGuides;
 		case 38: return EGroomViewMode::AO;
 		case 39: return EGroomViewMode::ClumpID;
+		case 40: return EGroomViewMode::Memory;
 		default: break;
 		}
 
@@ -129,6 +130,7 @@ const TCHAR* GetGroomViewModeName(EGroomViewMode In)
 	case EGroomViewMode::Group:						return TEXT("Group");
 	case EGroomViewMode::LODColoration:				return TEXT("LODColoration");
 	case EGroomViewMode::CardGuides:				return TEXT("CardGuides");
+	case EGroomViewMode::Memory:					return TEXT("Memory");
 	}
 	return TEXT("None");
 }
@@ -217,6 +219,7 @@ void FGroomVisualizationData::Initialize()
 		AddVisualizationMode(ModeMap, true, EGroomViewMode::LODColoration,				LOCTEXT("LODColoration", "LOD Color"),						LOCTEXT("LODColorationDesc", "Hair LOD coloring"));
 		AddVisualizationMode(ModeMap, true, EGroomViewMode::AO,							LOCTEXT("AO", "Per-CV AO"),									LOCTEXT("AODesc", "Control point AO  (optional)"));
 		AddVisualizationMode(ModeMap, true, EGroomViewMode::ClumpID,					LOCTEXT("ClumpID", "Clump ID"),								LOCTEXT("ClumpIDDesc", "Curve clumpID  (optional)"));
+		AddVisualizationMode(ModeMap, true, EGroomViewMode::Memory,						LOCTEXT("Memory", "Memory"),								LOCTEXT("MemoryDesc", "Memory"));
 
 		ConsoleDocumentationVisualizationMode = ConfigureConsoleCommand(ModeMap);
 	}
diff --git a/Engine/Source/Runtime/Engine/Public/GroomVisualizationData.h b/Engine/Source/Runtime/Engine/Public/GroomVisualizationData.h
--- a/Engine/Source/Runtime/Engine/Public/GroomVisualizationData.h
+++ b/Engine/Source/Runtime/Engine/Public/GroomVisualizationData.h
@@ -44,6 +44,7 @@ enum class EGroomViewMode : uint8
 	MaterialTangent,
 	CardGuides,
 	Tile,
+	Memory,
 	Count
 };
 
