Fix incorrect NumFreePages calculation in GlobalDistanceFieldDebugCS
- also print NumAllocPages

diff --git a/Engine/Shaders/Private/DistanceField/GlobalDistanceFieldDebug.usf b/Engine/Shaders/Private/DistanceField/GlobalDistanceFieldDebug.usf
--- a/Engine/Shaders/Private/DistanceField/GlobalDistanceFieldDebug.usf
+++ b/Engine/Shaders/Private/DistanceField/GlobalDistanceFieldDebug.usf
@@ -72,10 +72,11 @@ void GlobalDistanceFieldDebugCS(
 {
 	if (all(DispatchThreadId.xyz == 0))
 	{
-		int NumFreePages = -1;
+		int NumFreePages = GlobalDistanceFieldPageFreeListAllocatorBuffer[0];
+		int NumAllocPages = -1;
 		if (GlobalDistanceFieldMaxPageNum > 0)
 		{
-			NumFreePages = GlobalDistanceFieldMaxPageNum - GlobalDistanceFieldPageFreeListAllocatorBuffer[0];
+			NumAllocPages = GlobalDistanceFieldMaxPageNum - GlobalDistanceFieldPageFreeListAllocatorBuffer[0];
 		}
 
 		FShaderPrintContext Context = InitShaderPrintContext(true, float2(0.8, 0.2));
@@ -83,14 +84,18 @@ void GlobalDistanceFieldDebugCS(
 		Print(Context, TEXT("GDF Pages"));
 		Newline(Context);
 
-		Print(Context, TEXT(" Alloc "));
-		Print(Context, NumFreePages);
-		Newline(Context);
-
 		Print(Context, TEXT(" Max   "));
 		Print(Context, GlobalDistanceFieldMaxPageNum);
 		Newline(Context);
 
+		Print(Context, TEXT(" Alloc "), FontRed);
+		Print(Context, NumAllocPages, FontRed);
+		Newline(Context);
+
+		Print(Context, TEXT(" Free  "), FontGreen);
+		Print(Context, NumFreePages, FontGreen);
+		Newline(Context);
+
 		Print(Context, TEXT(" Combi "));
 		Print(Context, PageStatsBuffer[0]);
 		Newline(Context);
@@ -98,4 +103,4 @@ void GlobalDistanceFieldDebugCS(
 		Print(Context, TEXT(" Empty "));
 		Print(Context, PageStatsBuffer[1]);
 	}
-}
\ No newline at end of file
+}
