Corrected EncodeTileStencilZBaseDelta HTile encoder to properly use ZMax as the ZBase in the Hi-S path instead of ZMin, due to reverse Z. Also optimized the encoding logic.

diff --git a/Engine/Shaders/Private/HTileEncoding.ush b/Engine/Shaders/Private/HTileEncoding.ush
--- a/Engine/Shaders/Private/HTileEncoding.ush
+++ b/Engine/Shaders/Private/HTileEncoding.ush
@@ -41,86 +41,45 @@ uint EncodeTileMinMaxDepth(float MinDepth, float MaxDepth)
 	return BitFieldInsertU32(BitFieldMaskU32(14, 18), HTileValue, HTileValue << 2) | 0xF /* expanded tile ZMask */;
 }
 
-// TODO: Optimize (branchless, bit scan, etc.)
-uint EncodeTileStencilZBaseDelta(float MinDepth, float MaxDepth, uint XX, uint SMem, uint SR0, uint SR1)
+uint EncodeZRangeDelta(uint ZBase, uint ZOther)
 {
-	const float MaxDepthRange = float((1 << 14) - 1);
-	
-	const uint ZMask = 0xF; // Expanded
-	const uint ZMin  = uint(floor(MinDepth * MaxDepthRange));
-	const uint ZMax  = uint( ceil(MaxDepth * MaxDepthRange));
-	const uint ZDiff = ZMax - ZMin;
-	
-	uint ZDelta = 0;
-
-	if ((ZDiff & BitFieldMaskU32(1, 13)) != 0)
-	{
-		ZDelta = (0xFu << 2u) | BitFieldExtractU32(ZDiff, 2, 11);
-	}
-	else if ((ZDiff & BitFieldMaskU32(1, 12)) != 0)
-	{
-		ZDelta = (0xEu << 2u) | BitFieldExtractU32(ZDiff, 2, 10);
-	}
-	else if ((ZDiff & BitFieldMaskU32(1, 11)) != 0)
-	{
-		ZDelta = (0xDu << 2u) | BitFieldExtractU32(ZDiff, 2, 9);
-	}
-	else if ((ZDiff & BitFieldMaskU32(1, 10)) != 0)
-	{
-		ZDelta = (0xCu << 2u) | BitFieldExtractU32(ZDiff, 2, 8);
-	}
-	else if ((ZDiff & BitFieldMaskU32(1, 9)) != 0)
-	{
-		ZDelta = (0xBu << 2u) | BitFieldExtractU32(ZDiff, 2, 7);
-	}
-	else if ((ZDiff & BitFieldMaskU32(1, 8)) != 0)
-	{
-		ZDelta = (0xAu << 2u) | BitFieldExtractU32(ZDiff, 2, 6);
-	}
-	else if ((ZDiff & BitFieldMaskU32(1, 7)) != 0)
-	{
-		ZDelta = (0x9u << 2u) | BitFieldExtractU32(ZDiff, 2, 5);
-	}
-	else if ((ZDiff & BitFieldMaskU32(1, 6)) != 0)
+	uint Delta = abs(int(ZBase) - int(ZOther));
+	if (Delta < 16u)
 	{
-		ZDelta = (0x8u << 2u) | BitFieldExtractU32(ZDiff, 2, 4);
+		return Delta;
 	}
-	else if ((ZDiff & BitFieldMaskU32(1, 5)) != 0)
-	{
-		ZDelta = (0x3u << 3u) | BitFieldExtractU32(ZDiff, 3, 2);
-	}
-	else if ((ZDiff & BitFieldMaskU32(1, 4)) != 0)
+
+	const uint LeadingZeros = (31u - firstbithigh(Delta)) - 18u;
+	if (LeadingZeros >= 8)
 	{
-		ZDelta = (0x2u << 3u) | BitFieldExtractU32(ZDiff, 3, 1);
+		return ((11u - LeadingZeros) << 3u) | ((Delta >> (10u - LeadingZeros)) & 7u);
 	}
 	else
 	{
-		ZDelta = (0x0u << 4u) | BitFieldExtractU32(ZDiff, 4, 0);
+		return ((15u - LeadingZeros) << 2u) | ((Delta >> (11u - LeadingZeros)) & 3u);
 	}
+}
 
+uint EncodeTileStencilZBaseDelta(float MinDepth, float MaxDepth, uint XX, uint SMem, uint SR0, uint SR1)
+{
+	const float MaxDepthRange = float((1 << 14) - 1);
+	
+	const uint ZMask = 0xF; // Expanded
+	const uint ZMin  = uint(floor(MinDepth * MaxDepthRange));
+	const uint ZMax  = uint( ceil(MaxDepth * MaxDepthRange));
+	
+	// Assumes reverse Z
+	const uint ZBase = ZMax;
+	const uint ZDelta = EncodeZRangeDelta(ZMin, ZMax);
+	
 	uint HTileValue = 0;
-
-	// 0:3
-	HTileValue |= ZMask;
-
-	// 4:5
-	HTileValue |= (SR0 << 4);
-
-	// 6:7
-	HTileValue |= (SR1 << 6);
-
-	// 8:9
-	HTileValue |= (SMem << 8);
-
-	// 10:11
-	HTileValue |= (XX << 10);
-
-	// 12:17
-	HTileValue |= (ZDelta << 12);
-
-	// 18:31
-	HTileValue |= (ZMin << 18);
-
+	HTileValue |= ZMask;				//  0:3
+	HTileValue |= (SR0  << 4);			//  4:5
+	HTileValue |= (SR1  << 6);			//  6:7
+	HTileValue |= (SMem << 8);			//  8:9
+	HTileValue |= (XX << 10);			// 10:11
+	HTileValue |= (ZDelta << 12);		// 12:17
+	HTileValue |= (ZBase << 18);		// 18:31
 	return HTileValue;
 }
 
