Fixes reprojection of super bright highlights

diff --git a/Engine/Shaders/Private/TemporalSuperResolution/TSRUpdateHistory.usf b/Engine/Shaders/Private/TemporalSuperResolution/TSRUpdateHistory.usf
--- a/Engine/Shaders/Private/TemporalSuperResolution/TSRUpdateHistory.usf
+++ b/Engine/Shaders/Private/TemporalSuperResolution/TSRUpdateHistory.usf
@@ -775,7 +775,17 @@ void MainCS(
 				AccumulateMetadata  = dpv_scale(AccumulateMetadata, FinalMultiplier);
 				AccumulateTranslucency  = dpv_scale(AccumulateTranslucency, FinalMultiplier);
 			}
-			
+
+			// Super bight highlights have a tendy to generate negatives
+			bool2 HasAnyNegative = or(AccumulateHighFrequency[0] < 0.0, or(AccumulateHighFrequency[1] < 0.0, AccumulateHighFrequency[2] < 0.0));
+
+			AccumulateHighFrequency[0] = select(HasAnyNegative, dpv_interleave_registers_array(RawHighFrequency[2])[0], AccumulateHighFrequency[0]);
+			AccumulateHighFrequency[1] = select(HasAnyNegative, dpv_interleave_registers_array(RawHighFrequency[2])[1], AccumulateHighFrequency[1]);
+			AccumulateHighFrequency[2] = select(HasAnyNegative, dpv_interleave_registers_array(RawHighFrequency[2])[2], AccumulateHighFrequency[2]);
+			#if CONFIG_SCENE_COLOR_ALPHA
+				AccumulateHighFrequency[3] = select(HasAnyNegative, dpv_interleave_registers_array(RawHighFrequency[2])[3], AccumulateHighFrequency[3]);
+			#endif
+	
 			#if CONFIG_SCENE_COLOR_ALPHA
 				AccumulateHighFrequency = -dpv_min(-AccumulateHighFrequency, tsr_half(0.0).xxxx);
 			#else
