Fix voxelization issue missing pages.
* Change curve-as-cluster selection to encompass better the entire groom.
* This fixes issue with voxel page allocation. The cluster are optionally scaled.

diff --git a/Engine/Plugins/Runtime/HairStrands/Source/HairStrandsCore/Private/GroomBuilder.cpp b/Engine/Plugins/Runtime/HairStrands/Source/HairStrandsCore/Private/GroomBuilder.cpp
--- a/Engine/Plugins/Runtime/HairStrands/Source/HairStrandsCore/Private/GroomBuilder.cpp
+++ b/Engine/Plugins/Runtime/HairStrands/Source/HairStrandsCore/Private/GroomBuilder.cpp
@@ -40,7 +40,7 @@ static FAutoConsoleVariableRef CVarHairGroupIndexBuilder_MaxVoxelResolution(TEXT
 
 FString FGroomBuilder::GetVersion()
 {
-	return TEXT("v8r41");
+	return TEXT("v8r41f");
 }
 
 namespace FHairStrandsDecimation
@@ -2729,7 +2729,7 @@ struct FClusterGrid
 		TArray<FCurve> ClusterCurves;
 	};
 
-	FClusterGrid(const FIntVector& InResolution, const FVector& InMinBound, const FVector& InMaxBound)
+	FClusterGrid(const FIntVector& InResolution, const FVector3f& InMinBound, const FVector3f& InMaxBound)
 	{
 		MinBound = InMinBound;
 		MaxBound = InMaxBound;
@@ -2753,10 +2753,10 @@ struct FClusterGrid
 			FMath::Clamp(CellCoord.Z, 0, GridResolution.Z - 1));
 	}
 
-	FORCEINLINE FIntVector ToCellCoord(const FVector& P) const
+	FORCEINLINE FIntVector ToCellCoord(const FVector3f& P) const
 	{
 		bool bIsValid = false;
-		const FVector F = ((P - MinBound) / (MaxBound - MinBound));
+		const FVector3f F = ((P - MinBound) / (MaxBound - MinBound));
 		const FIntVector CellCoord = FIntVector(FMath::FloorToInt(F.X * GridResolution.X), FMath::FloorToInt(F.Y * GridResolution.Y), FMath::FloorToInt(F.Z * GridResolution.Z));
 		return ClampToVolume(CellCoord, bIsValid);
 	}
@@ -2768,7 +2768,7 @@ struct FClusterGrid
 		return CellIndex;
 	}
 
-	void InsertRenderingCurve(FCurve& Curve, const FVector& Root)
+	void InsertRenderingCurve(FCurve& Curve, const FVector3f& Root)
 	{
 		FIntVector CellCoord = ToCellCoord(Root);
 		uint32 Index = ToIndex(CellCoord);
@@ -2776,8 +2776,8 @@ struct FClusterGrid
 		Cluster.ClusterCurves.Add(Curve);
 	}
 
-	FVector MinBound;
-	FVector MaxBound;
+	FVector3f MinBound;
+	FVector3f MaxBound;
 	FIntVector GridResolution;
 	TArray<FCluster> Clusters;
 };
@@ -3038,8 +3038,93 @@ static void BuildClusterData(
 		}
 	}
 
+	// Compute the number of cluster based on a voxelization of the groom.
+	uint32 ClusterCount = 0;
+	{
+		// 1. Allocate cluster per voxel containing contains >=1 render curve root
+		FVector3f GroupMinBound(InRenStrandsData.BoundingBox.Min);
+		FVector3f GroupMaxBound(InRenStrandsData.BoundingBox.Max);
+
+		// Compute the voxel volume resolution, and snap the max bound to the voxel grid
+		// Iterate until voxel size are below max resolution, so that computation is not too long
+		FIntVector VoxelResolution = FIntVector::ZeroValue;
+		{
+			const int32 MaxResolution = FMath::Max(GHairClusterBuilder_MaxVoxelResolution, 2);
+
+			float ClusterWorldSize = 1.f; //cm
+			bool bIsValid = false;
+			while (!bIsValid)
+			{
+				FVector3f VoxelResolutionF = (GroupMaxBound - GroupMinBound) / ClusterWorldSize;
+				VoxelResolution = FIntVector(FMath::CeilToInt(VoxelResolutionF.X), FMath::CeilToInt(VoxelResolutionF.Y), FMath::CeilToInt(VoxelResolutionF.Z));
+				bIsValid = VoxelResolution.X <= MaxResolution && VoxelResolution.Y <= MaxResolution && VoxelResolution.Z <= MaxResolution;
+				if (!bIsValid)
+				{
+					ClusterWorldSize *= 2;
+				}
+			}
+			GroupMaxBound = GroupMinBound + FVector3f(VoxelResolution) * ClusterWorldSize;
+		}
+
+		// 2. Insert all rendering curves into the voxel structure
+		FClusterGrid ClusterGrid(VoxelResolution, GroupMinBound, GroupMaxBound);
+		for (uint32 RenCurveIndex = 0; RenCurveIndex < RenCurveCount; ++RenCurveIndex)
+		{
+			FClusterGrid::FCurve RCurve;
+			RCurve.CurveIndex = RenCurveIndex;
+			RCurve.Count = InRenStrandsData.StrandsCurves.CurvesCount[RenCurveIndex];
+			RCurve.Offset = InRenStrandsData.StrandsCurves.CurvesOffset[RenCurveIndex];
+			RCurve.Area = 0.0f;
+			RCurve.AvgRadius = 0;
+			RCurve.MaxRadius = 0;
+
+			// Compute area of each curve to later compute area correction
+			for (uint32 RenPointIndex = 0; RenPointIndex < RCurve.Count; ++RenPointIndex)
+			{
+				uint32 PointGlobalIndex = RenPointIndex + RCurve.Offset;
+				const FVector3f& V0 = InRenStrandsData.StrandsPoints.PointsPosition[PointGlobalIndex];
+				if (RenPointIndex > 0)
+				{
+					const FVector3f& V1 = InRenStrandsData.StrandsPoints.PointsPosition[PointGlobalIndex - 1];
+					FVector3f OutDir;
+					float OutLength;
+					(V1 - V0).ToDirectionAndLength(OutDir, OutLength);
+					RCurve.Area += InRenStrandsData.StrandsPoints.PointsRadius[PointGlobalIndex] * OutLength;
+				}
+
+				const float PointRadius = InRenStrandsData.StrandsPoints.PointsRadius[PointGlobalIndex];
+				RCurve.AvgRadius += PointRadius;
+				RCurve.MaxRadius = FMath::Max(RCurve.MaxRadius, PointRadius);
+			}
+			RCurve.AvgRadius /= FMath::Max(1u, RCurve.Count);
+
+			const FVector3f Root = InRenStrandsData.StrandsPoints.PointsPosition[RCurve.Offset];
+			ClusterGrid.InsertRenderingCurve(RCurve, Root);
+		}
+
+		// 3. Count non-empty clusters
+		{
+			for (FClusterGrid::FCluster& Cluster : ClusterGrid.Clusters)
+			{
+				if (Cluster.ClusterCurves.Num() > 0)
+				{
+					++ClusterCount;
+				}
+			}
+		}
+	}
+
 	// 4. Fill in cluster Info for legacy compatibility
-	Out.ClusterCount = FMath::Min(FMath::DivideAndRoundUp(InRenStrandsData.GetNumCurves(), 100u), 1024u);
+	const uint32 MinClusterCount = FMath::Max(FMath::DivideAndRoundUp(InRenStrandsData.GetNumCurves(), 100u), 128u); // 1%
+	const uint32 MaxClusterCount = 4096;
+	Out.ClusterCount = FMath::Max(MinClusterCount, ClusterCount);
+	Out.ClusterCount = FMath::Min(Out.ClusterCount, InRenStrandsData.GetNumCurves()); // Can't have more clusters than curve, since we use curves as proxy for cluster AABB
+	Out.ClusterScale = 1.f;
+	if (ClusterCount > MaxClusterCount)
+	{
+		Out.ClusterScale = float(ClusterCount) / float(MaxClusterCount);
+		Out.ClusterCount = MaxClusterCount;
+	}
 	Out.CurveToClusterIds.SetNum(InRenStrandsData.GetNumCurves());
 	Out.ClusterInfos.SetNum(Out.ClusterCount);
 	const uint32 CurvePerCluster = FMath::DivideAndRoundUp(InRenStrandsData.GetNumCurves(), Out.ClusterCount);
@@ -3070,6 +3155,7 @@ static void BuildClusterBulkData(
 {
 	Out.Reset();
 	Out.Header.ClusterCount		= In.ClusterCount;
+	Out.Header.ClusterScale		= In.ClusterScale;
 	Out.Header.PointCount		= In.PointCount;
 	Out.Header.CurveCount		= In.CurveCount;
 	Out.Header.LODInfos			= In.LODInfos;
diff --git a/Engine/Plugins/Runtime/HairStrands/Source/HairStrandsCore/Private/GroomComponent.cpp b/Engine/Plugins/Runtime/HairStrands/Source/HairStrandsCore/Private/GroomComponent.cpp
--- a/Engine/Plugins/Runtime/HairStrands/Source/HairStrandsCore/Private/GroomComponent.cpp
+++ b/Engine/Plugins/Runtime/HairStrands/Source/HairStrandsCore/Private/GroomComponent.cpp
@@ -2659,6 +2659,7 @@ void UGroomComponent::InitResources(bool bIsBindingReloading)
 			HairGroupInstance->HairGroupPublicData->RestPointCount = GroupData.Strands.BulkData.GetNumPoints();
 			HairGroupInstance->HairGroupPublicData->RestCurveCount = GroupData.Strands.BulkData.GetNumCurves();
 			HairGroupInstance->HairGroupPublicData->ClusterCount = 0u;
+			HairGroupInstance->HairGroupPublicData->ClusterScale = 0.f;
 
 			// (Lazy) Allocate interpolation resources, only if guides are required
 			if (bNeedGuides)
@@ -2774,6 +2775,7 @@ void UGroomComponent::InitResources(bool bIsBindingReloading)
 				check(StrandsLODCount <= LODScreenSizes.Num());
 				check(StrandsLODCount <= LODVisibilities.Num());
 				HairGroupInstance->HairGroupPublicData->ClusterCount = HairGroupInstance->Strands.ClusterResource->BulkData.Header.ClusterCount;
+				HairGroupInstance->HairGroupPublicData->ClusterScale = HairGroupInstance->Strands.ClusterResource->BulkData.Header.ClusterScale;
 			}
 
 			HairGroupInstance->Strands.CullingResource = new FHairStrandsCullingResource(
diff --git a/Engine/Plugins/Runtime/HairStrands/Source/HairStrandsCore/Private/GroomDebug.cpp b/Engine/Plugins/Runtime/HairStrands/Source/HairStrandsCore/Private/GroomDebug.cpp
--- a/Engine/Plugins/Runtime/HairStrands/Source/HairStrandsCore/Private/GroomDebug.cpp
+++ b/Engine/Plugins/Runtime/HairStrands/Source/HairStrandsCore/Private/GroomDebug.cpp
@@ -383,36 +383,32 @@ class FDrawDebugClusterAABBCS : public FGlobalShader
 	DECLARE_GLOBAL_SHADER(FDrawDebugClusterAABBCS);
 	SHADER_USE_PARAMETER_STRUCT(FDrawDebugClusterAABBCS, FGlobalShader);
 
-	class FDebugAABBBuffer : SHADER_PERMUTATION_INT("PERMUTATION_DEBUGAABBBUFFER", 2);
-	using FPermutationDomain = TShaderPermutationDomain<FDebugAABBBuffer>;
-
-	BEGIN_SHADER_PARAMETER_STRUCT(FParameters, )
+	BEGIN_SHADER_PARAMETER_STRUCT(FParameters,)
 		SHADER_PARAMETER_STRUCT_REF(FViewUniformShaderParameters, ViewUniformBuffer)
 		SHADER_PARAMETER_SRV(Buffer, ClusterAABBBuffer)
 		SHADER_PARAMETER_SRV(Buffer, GroupAABBBuffer)
-		SHADER_PARAMETER_RDG_BUFFER_SRV(Buffer, CulledDispatchIndirectParametersClusterCountBuffer)
-		SHADER_PARAMETER_RDG_BUFFER_SRV(StructuredBuffer, ClusterDebugInfoBuffer)
-		SHADER_PARAMETER_SRV(Buffer, CulledDrawIndirectParameters)
 		SHADER_PARAMETER(uint32, ClusterCount)
-		SHADER_PARAMETER(uint32, TriangleCount)
+		SHADER_PARAMETER(uint32, PointCount)
+		SHADER_PARAMETER(uint32, CurveCount)
 		SHADER_PARAMETER(uint32, HairGroupId)
-		SHADER_PARAMETER(int32, ClusterDebugMode)
+		SHADER_PARAMETER(uint32, bDrawAABB)
 		SHADER_PARAMETER_STRUCT_INCLUDE(ShaderPrint::FShaderParameters, ShaderPrintParameters)
 		END_SHADER_PARAMETER_STRUCT()
 
 public:
+	static uint32 GetGroupSize() { return 64; }
 	static bool ShouldCompilePermutation(const FGlobalShaderPermutationParameters& Parameters) { return IsHairStrandsSupported(EHairStrandsShaderType::Tool, Parameters.Platform); }
 	static void ModifyCompilationEnvironment(const FGlobalShaderPermutationParameters& Parameters, FShaderCompilerEnvironment& OutEnvironment)
 	{
 		FGlobalShader::ModifyCompilationEnvironment(Parameters, OutEnvironment);
-		OutEnvironment.SetDefine(TEXT("SHADER_DRAWDEBUGAABB"), 1);
-
+		OutEnvironment.SetDefine(TEXT("SHADER_CLUSTERAABB"), 1);
+		OutEnvironment.SetDefine(TEXT("GROUP_SIZE"), GetGroupSize());
 		// Skip optimization for avoiding long compilation time due to large UAV writes
 		OutEnvironment.CompilerFlags.Add(CFLAG_Debug);
 	}
 };
 
-IMPLEMENT_GLOBAL_SHADER(FDrawDebugClusterAABBCS, "/Engine/Private/HairStrands/HairStrandsClusterCulling.usf", "MainDrawDebugAABBCS", SF_Compute);
+IMPLEMENT_GLOBAL_SHADER(FDrawDebugClusterAABBCS, "/Engine/Private/HairStrands/HairStrandsDebug.usf", "MainDrawDebugAABBCS", SF_Compute);
 
 void AddDrawDebugClusterPass(
 	FRDGBuilder& GraphBuilder,
@@ -422,7 +418,6 @@ void AddDrawDebugClusterPass(
 	EGroomViewMode ViewMode,
 	FHairStrandClusterData& HairClusterData)
 {
-	// Force ShaderPrint on.
 	ShaderPrint::SetEnabled(true);
 	ShaderPrint::RequestSpaceForLines(64000u);
 	ShaderPrint::RequestSpaceForCharacters(2000);
@@ -433,38 +428,23 @@ void AddDrawDebugClusterPass(
 	uint32 DataIndex = 0;
 	for (const FHairStrandClusterData::FHairGroup& HairGroupClusters : HairClusterData.HairGroups)
 	{
-		if (HairGroupClusters.CulledClusterCountBuffer)
-		{
-			FRDGExternalBuffer& DrawIndirectBuffer = HairGroupClusters.HairGroupPublicPtr->GetDrawIndirectBuffer();
-
-			FDrawDebugClusterAABBCS::FPermutationDomain Permutation;
-			Permutation.Set<FDrawDebugClusterAABBCS::FDebugAABBBuffer>(bDebugAABB ? 1 : 0);
-			TShaderMapRef<FDrawDebugClusterAABBCS> ComputeShader(ShaderMap, Permutation);
-
-			FDrawDebugClusterAABBCS::FParameters* Parameters = GraphBuilder.AllocParameters<FDrawDebugClusterAABBCS::FParameters>();
-			Parameters->ViewUniformBuffer = View.ViewUniformBuffer;
-			Parameters->ClusterCount = HairGroupClusters.ClusterCount;
-			Parameters->TriangleCount = HairGroupClusters.VertexCount * 2; // VertexCount is actually the number of control points
-			Parameters->HairGroupId = DataIndex++;
-			Parameters->ClusterDebugMode = ViewMode == EGroomViewMode::Cluster ? 1 : (ViewMode == EGroomViewMode::ClusterAABB ? 2 : 0);
-			Parameters->ClusterAABBBuffer = HairGroupClusters.ClusterAABBBuffer->SRV;
-			Parameters->CulledDispatchIndirectParametersClusterCountBuffer = GraphBuilder.CreateSRV(HairGroupClusters.CulledClusterCountBuffer, EPixelFormat::PF_R32_UINT);
-			Parameters->CulledDrawIndirectParameters = DrawIndirectBuffer.SRV;
-			Parameters->GroupAABBBuffer = HairGroupClusters.GroupAABBBuffer->SRV;
-
-			if (HairGroupClusters.ClusterDebugInfoBuffer && bDebugAABB)
-			{
-				FRDGBufferRef ClusterDebugInfoBuffer = GraphBuilder.RegisterExternalBuffer(HairGroupClusters.ClusterDebugInfoBuffer);
-				Parameters->ClusterDebugInfoBuffer = GraphBuilder.CreateSRV(ClusterDebugInfoBuffer);
-			}
-			//ShaderPrint::SetParameters(GraphBuilder, ShaderPrintData, Parameters->ShaderPrintParameters);
-			ShaderPrint::SetParameters(GraphBuilder, *ShaderPrintData, Parameters->ShaderPrintParameters);
+		TShaderMapRef<FDrawDebugClusterAABBCS> ComputeShader(ShaderMap);
 
-			const FIntVector DispatchCount = DispatchCount.DivideAndRoundUp(FIntVector(Parameters->ClusterCount, 1, 1), FIntVector(64, 1, 1));// FIX ME, this could get over 65535
-			FComputeShaderUtils::AddPass(
-				GraphBuilder, RDG_EVENT_NAME("DrawDebugClusterAABB"),
-				ComputeShader, Parameters, DispatchCount);
-		}
+		FDrawDebugClusterAABBCS::FParameters* Parameters = GraphBuilder.AllocParameters<FDrawDebugClusterAABBCS::FParameters>();
+		Parameters->ViewUniformBuffer = View.ViewUniformBuffer;
+		Parameters->ClusterCount = HairGroupClusters.ClusterCount;
+		Parameters->PointCount = HairGroupClusters.HairGroupPublicPtr->GetActiveStrandsPointCount();
+		Parameters->CurveCount = HairGroupClusters.HairGroupPublicPtr->GetActiveStrandsCurveCount();
+		Parameters->HairGroupId = DataIndex++;
+		Parameters->bDrawAABB = ViewMode == EGroomViewMode::ClusterAABB ? 1 : 0;
+		Parameters->ClusterAABBBuffer = HairGroupClusters.ClusterAABBBuffer->SRV;
+		Parameters->GroupAABBBuffer = HairGroupClusters.GroupAABBBuffer->SRV;
+		ShaderPrint::SetParameters(GraphBuilder, *ShaderPrintData, Parameters->ShaderPrintParameters);
+
+		const FIntVector DispatchCount = DispatchCount.DivideAndRoundUp(FIntVector(Parameters->ClusterCount, 1, 1), FIntVector(FDrawDebugClusterAABBCS::GetGroupSize(), 1, 1));
+		FComputeShaderUtils::AddPass(
+			GraphBuilder, RDG_EVENT_NAME("HairStrands::DrawDebugClusterAABB"),
+			ComputeShader, Parameters, DispatchCount);
 	}
 }
 
diff --git a/Engine/Plugins/Runtime/HairStrands/Source/HairStrandsCore/Private/GroomManager.cpp b/Engine/Plugins/Runtime/HairStrands/Source/HairStrandsCore/Private/GroomManager.cpp
--- a/Engine/Plugins/Runtime/HairStrands/Source/HairStrandsCore/Private/GroomManager.cpp
+++ b/Engine/Plugins/Runtime/HairStrands/Source/HairStrandsCore/Private/GroomManager.cpp
@@ -823,11 +823,13 @@ static float ComputeActiveCurveCoverageScale(uint32 InAvailableCurveCount, uint3
 	return 1.f/FMath::Max(CurveRatio, 0.01f);
 }
 
-static uint32 ComputeActiveCurveCount(float InScreenSize, uint32 InCurveCount)
+static uint32 ComputeActiveCurveCount(float InScreenSize, uint32 InCurveCount, uint32 InClusterCount)
 {
 	const float Power = FMath::Max(0.1f, GHairStrands_AutoLOD_Scale);
 	const float ScreenSizeBias = FMath::Clamp(GHairStrands_AutoLOD_Bias, 0.f, 1.f);
-	const uint32 OutCurveCount = InCurveCount * FMath::Pow(FMath::Clamp(InScreenSize + ScreenSizeBias, 0.f, 1.0f), Power);
+	uint32 OutCurveCount = InCurveCount * FMath::Pow(FMath::Clamp(InScreenSize + ScreenSizeBias, 0.f, 1.0f), Power);
+	// Ensure there is at least 1 curve per cluster
+	OutCurveCount = FMath::Max(InClusterCount, OutCurveCount);
 	return FMath::Clamp(OutCurveCount, 1, InCurveCount);
 }
 
@@ -930,7 +932,7 @@ static void RunHairLODSelection(
 				uint32 EffectiveCurveCount = 0;
 				if (Instance->HairGroupPublicData->bAutoLOD || IsHairStrandsForceAutoLODEnabled())
 				{
-					EffectiveCurveCount = ComputeActiveCurveCount(Instance->HairGroupPublicData->ContinuousLODScreenSize, Instance->HairGroupPublicData->RestCurveCount);
+					EffectiveCurveCount = ComputeActiveCurveCount(Instance->HairGroupPublicData->ContinuousLODScreenSize, Instance->HairGroupPublicData->RestCurveCount, Instance->HairGroupPublicData->ClusterCount);
 					LODIndex = 0;
 				}
 				else
diff --git a/Engine/Plugins/Runtime/HairStrands/Source/HairStrandsCore/Private/HairStrandsClusterCulling.cpp b/Engine/Plugins/Runtime/HairStrands/Source/HairStrandsCore/Private/HairStrandsClusterCulling.cpp
--- a/Engine/Plugins/Runtime/HairStrands/Source/HairStrandsCore/Private/HairStrandsClusterCulling.cpp
+++ b/Engine/Plugins/Runtime/HairStrands/Source/HairStrandsCore/Private/HairStrandsClusterCulling.cpp
@@ -273,8 +273,8 @@ void AddInstanceToClusterData(FHairGroupInstance* In, FHairStrandClusterData& Ou
 	const int32 ClusterDataGroupIndex = Out.HairGroups.Num();
 	FHairStrandClusterData::FHairGroup& HairGroupCluster = Out.HairGroups.Emplace_GetRef();
 	HairGroupCluster.MaxPointPerCurve = In->Strands.Data ? In->Strands.Data->Header.MaxPointPerCurve : 0;
+	HairGroupCluster.ClusterScale = In->HairGroupPublicData->GetClusterScale();
 	HairGroupCluster.ClusterCount = In->HairGroupPublicData->GetClusterCount();
-	HairGroupCluster.VertexCount = In->HairGroupPublicData->RestPointCount * HAIR_POINT_TO_VERTEX; // Instance->HairGroupPublicData->GetActiveStrandsPointCount() ?
 	HairGroupCluster.GroupAABBBuffer = &In->HairGroupPublicData->GetGroupAABBBuffer();
 	HairGroupCluster.ClusterAABBBuffer = &In->HairGroupPublicData->GetClusterAABBBuffer();
 
diff --git a/Engine/Plugins/Runtime/HairStrands/Source/HairStrandsCore/Private/HairStrandsClusterCulling.h b/Engine/Plugins/Runtime/HairStrands/Source/HairStrandsCore/Private/HairStrandsClusterCulling.h
--- a/Engine/Plugins/Runtime/HairStrands/Source/HairStrandsCore/Private/HairStrandsClusterCulling.h
+++ b/Engine/Plugins/Runtime/HairStrands/Source/HairStrandsCore/Private/HairStrandsClusterCulling.h
@@ -16,7 +16,7 @@ struct FHairStrandClusterData
 	struct FHairGroup
 	{
 		uint32 ClusterCount = 0;
-		uint32 VertexCount = 0;
+		float ClusterScale = 0;
 		uint32 MaxPointPerCurve = 0;
 		FVector4f ClusterInfoParameters = FVector4f::Zero();
 		float LODIndex = -1;
@@ -43,12 +43,6 @@ struct FHairStrandClusterData
 		bool GetCullingResultAvailable() const							{ return HairGroupPublicPtr ? HairGroupPublicPtr->GetCullingResultAvailable() : false; }
 		void SetCullingResultAvailable(bool b)							{ if (HairGroupPublicPtr) HairGroupPublicPtr->SetCullingResultAvailable(b); }
 
-		TRefCountPtr<FRDGPooledBuffer> ClusterDebugInfoBuffer;	// Null if this debug is not enabled.
-		FRDGBufferRef CulledClusterCountBuffer = nullptr;
-		FRDGBufferRef CulledCluster1DIndirectArgsBuffer = nullptr;
-		FRDGBufferRef CulledCluster2DIndirectArgsBuffer = nullptr;
-		uint32 GroupSize1D = 0;
-
 		FHairGroupPublicData* HairGroupPublicPtr = nullptr;
 	};
 
diff --git a/Engine/Plugins/Runtime/HairStrands/Source/HairStrandsCore/Private/HairStrandsDatas.cpp b/Engine/Plugins/Runtime/HairStrands/Source/HairStrandsCore/Private/HairStrandsDatas.cpp
--- a/Engine/Plugins/Runtime/HairStrands/Source/HairStrandsCore/Private/HairStrandsDatas.cpp
+++ b/Engine/Plugins/Runtime/HairStrands/Source/HairStrandsCore/Private/HairStrandsDatas.cpp
@@ -693,6 +693,7 @@ void FHairStrandsInterpolationBulkData::ResetLoadedSize()
 void FHairStrandsClusterBulkData::Reset()
 {
 	Header.ClusterCount = 0;
+	Header.ClusterScale = 0;
 	Header.PointCount = 0;
 	Header.CurveCount = 0;
 	Header.ClusterInfoParameters = FVector4f::Zero();
@@ -718,6 +719,7 @@ void FHairStrandsClusterBulkData::ResetLoadedSize()
 void FHairStrandsClusterBulkData::SerializeHeader(FArchive& Ar, UObject* Owner)
 {
 	Ar << Header.ClusterCount;
+	Ar << Header.ClusterScale;
 	Ar << Header.PointCount;
 	Ar << Header.CurveCount;
 	Ar << Header.ClusterInfoParameters;
diff --git a/Engine/Plugins/Runtime/HairStrands/Source/HairStrandsCore/Private/HairStrandsInterpolation.cpp b/Engine/Plugins/Runtime/HairStrands/Source/HairStrandsCore/Private/HairStrandsInterpolation.cpp
--- a/Engine/Plugins/Runtime/HairStrands/Source/HairStrandsCore/Private/HairStrandsInterpolation.cpp
+++ b/Engine/Plugins/Runtime/HairStrands/Source/HairStrandsCore/Private/HairStrandsInterpolation.cpp
@@ -826,6 +826,7 @@ class FHairClusterAABBCS : public FGlobalShader
 	BEGIN_SHADER_PARAMETER_STRUCT(FParameters, )
 		SHADER_PARAMETER_STRUCT_INCLUDE(ShaderPrint::FShaderParameters, ShaderDrawParameters)
 		SHADER_PARAMETER(float, LODIndex)
+		SHADER_PARAMETER(float, ClusterScale)
 		SHADER_PARAMETER(uint32, ClusterCount)
 		SHADER_PARAMETER(uint32, CurveCount)
 		SHADER_PARAMETER(FVector3f, CPUBoundMin)
@@ -888,6 +889,7 @@ static void AddHairClusterAABBPass(
 	Parameters->RenderDeformedOffsetBuffer = RenderDeformedOffsetBuffer;
 	Parameters->CurveCount = ActiveCurveCount;
 	Parameters->ClusterCount = ClusterData ? ClusterData->ClusterCount : 1;
+	Parameters->ClusterScale = ClusterData ? ClusterData->ClusterScale : 1.f;
 	Parameters->RenCurveBuffer = RegisterAsSRV(GraphBuilder, Instance->Strands.RestResource->CurveBuffer);
 	Parameters->RenPointLODBuffer = ClusterData ? RegisterAsSRV(GraphBuilder, *ClusterData->PointLODBuffer) : nullptr;
 	Parameters->OutClusterAABBBuffer = ClusterAABBData.ClusterAABBBuffer.UAV;
diff --git a/Engine/Plugins/Runtime/HairStrands/Source/HairStrandsCore/Public/HairStrandsDatas.h b/Engine/Plugins/Runtime/HairStrands/Source/HairStrandsCore/Public/HairStrandsDatas.h
--- a/Engine/Plugins/Runtime/HairStrands/Source/HairStrandsCore/Public/HairStrandsDatas.h
+++ b/Engine/Plugins/Runtime/HairStrands/Source/HairStrandsCore/Public/HairStrandsDatas.h
@@ -706,6 +706,7 @@ struct HAIRSTRANDSCORE_API FHairStrandsClusterData
 	TArray<uint32>				PointLODs;
 
 	uint32 ClusterCount = 0;		// Number of clusters
+	float  ClusterScale = 0;		// Cluster scale factor
 	uint32 PointCount = 0;			// Number of points
 	uint32 CurveCount = 0;			// Number of curves
 };
@@ -728,6 +729,7 @@ struct HAIRSTRANDSCORE_API FHairStrandsClusterBulkData : FHairStrandsBulkCommon
 		TArray<FHairLODInfo> LODInfos;
 	
 		uint32 ClusterCount = 0;
+		float  ClusterScale = 0;
 		uint32 PointCount = 0;
 		uint32 CurveCount = 0;
 		FVector4f ClusterInfoParameters = FVector4f::Zero(); // xy:Scale/Offset for ScreenSize, zw:Scale/Offset for Radius
diff --git a/Engine/Shaders/Private/HairStrands/HairStrandsClusterCulling.usf b/Engine/Shaders/Private/HairStrands/HairStrandsClusterCulling.usf
--- a/Engine/Shaders/Private/HairStrands/HairStrandsClusterCulling.usf
+++ b/Engine/Shaders/Private/HairStrands/HairStrandsClusterCulling.usf
@@ -70,189 +70,6 @@ void MainClearClusterAABBCS(uint2 DispatchThreadId : SV_DispatchThreadID)
 
 #endif
 
-/////////////////////////////////////////////////////////////////////////////////////////////////////////////////////////////
-/// Debug draw clusters
-#ifdef SHADER_DRAWDEBUGAABB
-
-#include "../Common.ush"
-#include "../ShaderPrint.ush"
-
-#define MODE_CLUSTER	1
-#define MODE_CULLED		2
-#define MODE_LOD		3
-#define MODE_LOD_SIMPLE	4
-
-uint ClusterCount;
-uint HairGroupId;
-uint TriangleCount;
-int ClusterDebugMode;
-
-Buffer<int> GroupAABBBuffer;
-Buffer<int> ClusterAABBBuffer;
-StructuredBuffer<FHairClusterDebugInfo> ClusterDebugInfoBuffer;
-
-Buffer<uint> CulledDrawIndirectParameters;
-Buffer<uint> CulledDispatchIndirectParametersClusterCountBuffer;
-
-[numthreads(64, 1, 1)]
-void MainDrawDebugAABBCS(uint3 DispatchThreadId : SV_DispatchThreadID)
-{
-	if (DispatchThreadId.x == 0)
-	{
-		const bool bDrawClusterInfo = ClusterDebugMode == MODE_CLUSTER || ClusterDebugMode == MODE_CULLED || ClusterDebugMode == MODE_LOD;
-		const uint HairGroupHeight = 15; //bDrawClusterInfo ? 120u : 60u;
-
-		FShaderPrintContext Context = InitShaderPrintContext(true, uint2(50, 50 + (HairGroupId == 0 ? 0 : HairGroupHeight) + HairGroupId * HairGroupHeight));
-
-		FFontColor FontValue = FontWhite;
-		FFontColor FontPercent = FontSilver;
-
-		if (bDrawClusterInfo)
-		{
-			if (HairGroupId == 0)
-			{
-				Print(Context, TEXT("HairGroup   Clusters      Triangles"), FontYellow);
-			#if PERMUTATION_DEBUGAABBBUFFER
-				Print(Context, TEXT("         ScreenSize LOD"), FontYellow);
-			#endif
-				Newline(Context);
-			}
-
-			const uint VisibleClusterCount  = CulledDispatchIndirectParametersClusterCountBuffer[0];
-			const uint VisibleTriangleCount = CulledDrawIndirectParameters[0] / 3; // CulledDrawIndirectParameters contains the total number of triangle indices (3 per triangle)
-
-			Print(Context, HairGroupId, FontOrange);
-
-			Print(Context, VisibleClusterCount, FontValue, 5, 0);
-			Print(Context, TEXT(" ["), FontPercent);
-			Print(Context, 100.0f * float(VisibleClusterCount) / float(ClusterCount), FontPercent, 4, 1);
-			Print(Context, TEXT("%] "), FontPercent);
-
-			Print(Context, VisibleTriangleCount, FontValue, 9, 0);
-			Print(Context, TEXT(" ["), FontPercent);
-			Print(Context, 100.0f * float(VisibleTriangleCount) / float(TriangleCount), FontPercent, 4, 1);
-			Print(Context, TEXT("%] "), FontPercent);
-
-		}
-
-		#if PERMUTATION_DEBUGAABBBUFFER
-		{
-			float ScreenSize = 1;
-			{
-				const FHairAABB Bound = ReadHairAABB(0, GroupAABBBuffer);
-
-				const float3 SphereOrigin = (Bound.Min + Bound.Max) * 0.5f;
-				const float  SphereRadius = length(Bound.Max - SphereOrigin);
-				ScreenSize = ComputeBoundsScreenSize(SphereOrigin, SphereRadius, LWCHackToFloat(PrimaryView.WorldCameraOrigin), View.ViewToClip);
-			}
-			Print(Context, ScreenSize, FontValue, 8, 2);
-			Print(Context, TEXT("   "));
-
-			// Show actual LOD distribution for all clusters. Otherwise show the average LOD index
-			const bool bShowLODHistogram = false;
-			// Aggregate all the cluster results
-			uint LODHistogram[MAX_HAIR_LOD] = { 0,0,0,0,0,0,0,0 };
-			uint TotalVertexCount = 0;
-			uint TotalVisibleCluster = 0;
-			float AvgLOD = 0;
-			for (uint ClusterIt = 0; ClusterIt < ClusterCount; ++ClusterIt)
-			{
-				const FHairClusterDebugInfo DebugInfo = ClusterDebugInfoBuffer[ClusterIt];
-				if (DebugInfo.LOD >= 0) // Not culled
-				{
-					const uint LOD = uint(DebugInfo.LOD + 0.5f);
-					LODHistogram[LOD] += 1;
-					AvgLOD += DebugInfo.LOD;
-					TotalVertexCount += DebugInfo.VertexCount;
-					TotalVisibleCluster++;
-				}
-			}		
-			AvgLOD /= max(1u, TotalVisibleCluster);
-
-			if (bShowLODHistogram)
-			{		
-				Print(Context, uint(LODHistogram[0] / float(max(1u, TotalVisibleCluster)) * 100.f));
-				Print(Context, uint(LODHistogram[1] / float(max(1u, TotalVisibleCluster)) * 100.f));
-				Print(Context, uint(LODHistogram[2] / float(max(1u, TotalVisibleCluster)) * 100.f));
-				Print(Context, uint(LODHistogram[3] / float(max(1u, TotalVisibleCluster)) * 100.f));
-				Print(Context, uint(LODHistogram[4] / float(max(1u, TotalVisibleCluster)) * 100.f));
-				Print(Context, uint(LODHistogram[5] / float(max(1u, TotalVisibleCluster)) * 100.f));
-				Print(Context, uint(LODHistogram[6] / float(max(1u, TotalVisibleCluster)) * 100.f));
-				Print(Context, uint(LODHistogram[7] / float(max(1u, TotalVisibleCluster)) * 100.f));
-			}
-			else
-			{
-				if (TotalVisibleCluster > 0)
-				{
-					Print(Context, AvgLOD, FontValue, 4, 2);
-				}
-				else
-				{
-					Print(Context, TEXT("Culled"), FontLightRed);
-				}
-			}
-		}
-		#endif
-
-	}
-
-	if (DispatchThreadId.x < ClusterCount)
-	{
-		const bool bDrawClusterAABB = ClusterDebugMode == MODE_CLUSTER || ClusterDebugMode == MODE_CULLED || ClusterDebugMode == MODE_LOD;
-		if (DispatchThreadId.x == 0 && bDrawClusterAABB)
-		{
-			const float4 GroupColor = float4(1.0f, 0.5f, 0.0f, 1.0f);
-			const FHairAABB Bound = ReadHairAABB(0, GroupAABBBuffer);
-			AddAABBTWS(Bound.Min, Bound.Max, GroupColor);
-		}
-
-		const uint ClusterId = DispatchThreadId;
-		const uint ClusterId6 = DispatchThreadId.x * 6;
-		float3 Min = float3(int3(ClusterAABBBuffer[ClusterId6    ], ClusterAABBBuffer[ClusterId6 + 1], ClusterAABBBuffer[ClusterId6 + 2]));
-		float3 Max = float3(int3(ClusterAABBBuffer[ClusterId6 + 3], ClusterAABBBuffer[ClusterId6 + 4], ClusterAABBBuffer[ClusterId6 + 5]));
-	
-		#if PERMUTATION_DEBUGAABBBUFFER
-		const FHairClusterDebugInfo DebugInfo = ClusterDebugInfoBuffer[ClusterId];
-		const bool bIsCulledState = DebugInfo.LOD < 0;
-		#endif
-
-		float4 Color = 0.0f;
-		if (ClusterDebugMode == MODE_CLUSTER)
-		{
-			Color = float4(0.5f, 1.0f, 0.0f, 0.1f);
-		}
-		#if PERMUTATION_DEBUGAABBBUFFER
-		else if (ClusterDebugMode == MODE_CULLED)
-		{
-			Color = float4(1.0f, 0.1f, 0.1f, 1.0f);
-			if (!bIsCulledState)
-			{
-				return;
-			}
-		}
-		else if (ClusterDebugMode == MODE_LOD_SIMPLE)
-		{
-			// No cluster AABB drawing
-			return;
-		}
-		else if (ClusterDebugMode == MODE_LOD)
-		{
-			if (bIsCulledState)
-			{
-				return;
-			}
-
-			float SelectedLod = float(DebugInfo.LOD) / D_LOD_DEBUG_SCALE;
-			Color = float4(GetHSVDebugColor(SelectedLod), 0.05f);
-		}
-		#endif
-
-		AddAABBTWS(Min, Max, Color);
-	}
-}
-
-#endif
-
 /////////////////////////////////////////////////////////////////////////////////////////////////////////////////////////////
 /// Reset indirect buffer args
 #ifdef SHADER_CLUSTERCULLINGINDCLEAR  
@@ -637,6 +454,7 @@ void Main(uint DispatchThreadId : SV_DispatchThreadID, uint ThreadIndex : SV_Gro
 
 uint CurveCount;
 uint ClusterCount;
+float ClusterScale;
 float3 CPUBoundMin;
 float3 CPUBoundMax;
 float LODIndex;
@@ -712,6 +530,11 @@ groupshared float3 OutHairPositionOffset;
 			}
 		}
 	}
+	// Scaling
+	const float3 ClusterAABBCenter = float3(ClusterAABBMin + ClusterAABBMax) * 0.5f;
+	const float3 ClusterAABBExtent = float3(ClusterAABBMax - ClusterAABBMin) * 0.5f;
+	ClusterAABBMin = ClusterAABBCenter - ClusterAABBExtent * ClusterScale;
+	ClusterAABBMax = ClusterAABBCenter + ClusterAABBExtent * ClusterScale;
 
 	// Write each thread result to shared memory
 	SharedClusterAABBMin[ThreadIndex] = ClusterAABBMin;
diff --git a/Engine/Shaders/Private/HairStrands/HairStrandsDebug.usf b/Engine/Shaders/Private/HairStrands/HairStrandsDebug.usf
--- a/Engine/Shaders/Private/HairStrands/HairStrandsDebug.usf
+++ b/Engine/Shaders/Private/HairStrands/HairStrandsDebug.usf
@@ -4,7 +4,7 @@
 #include "../HairShadingCommon.ush"
 #include "../ColorMap.ush"
 
-#if SHADER_PRINT_INSTANCE || SHADER_PRINT || SHADER_PRINT_MEMORY
+#if SHADER_PRINT_INSTANCE || SHADER_PRINT || SHADER_PRINT_MEMORY || SHADER_CLUSTERAABB
 /////////////////////////////////////////////////////////////////////////////////////////
 // Utils function
 
@@ -2842,3 +2842,66 @@ void MainCS(uint GroupIndex : SV_GroupIndex, uint3 DispatchThreadId : SV_Dispatc
 }
 
 #endif // SHADER_PRINT_MEMORY
+
+/////////////////////////////////////////////////////////////////////////////////////////////////////////////////////////////
+/// Debug draw clusters
+#ifdef SHADER_CLUSTERAABB
+
+#include "HairStrandsAABBCommon.ush"
+
+uint ClusterCount;
+uint HairGroupId;
+uint PointCount;
+uint CurveCount;
+uint bDrawAABB;
+
+Buffer<int> GroupAABBBuffer;
+Buffer<int> ClusterAABBBuffer;
+
+[numthreads(GROUP_SIZE, 1, 1)]
+void MainDrawDebugAABBCS(uint3 DispatchThreadId : SV_DispatchThreadID)
+{
+	if (DispatchThreadId.x == 0)
+	{
+		const bool bDrawClusterInfo = true;
+		const uint HairGroupHeight = 15;
+
+		FShaderPrintContext Context = InitShaderPrintContext(true, uint2(50, 50 + (HairGroupId == 0 ? 0 : HairGroupHeight) + HairGroupId * HairGroupHeight));
+
+		FFontColor FontValue = FontWhite;
+		FFontColor FontPercent = FontSilver;
+
+		if (DispatchThreadId.x == 0)
+		{
+			if (HairGroupId == 0)
+			{
+				Print(Context, TEXT("Group Clusters Curves Points"), FontYellow);
+				Newline(Context);
+			}
+			Print(Context, HairGroupId, FontOrange, 6, 0);
+			Print(Context, ClusterCount, FontValue, 9, 0);
+			PrintLargeNumber(Context, CurveCount, FontValue);
+			Print(Context, TEXT(" "), FontYellow);
+			PrintLargeNumber(Context, PointCount, FontValue);
+		}
+	}
+
+	if (bDrawAABB)
+	{
+		// Group
+		if (DispatchThreadId.x == 0)
+		{
+			const FHairAABB GroupBound = ReadHairAABB(0, GroupAABBBuffer);
+			AddAABBTWS(GroupBound.Min, GroupBound.Max, float4(1.0f, 0.5f, 0.0f, 1.0f));
+		}
+	
+		// Cluster
+		if (DispatchThreadId.x < ClusterCount)
+		{
+			const FHairAABB ClusterBound = ReadHairAABB(DispatchThreadId, ClusterAABBBuffer, 1);
+			AddAABBTWS(ClusterBound.Min, ClusterBound.Max, float4(0.5f, 1.0f, 0.0f, 0.3f));
+		}
+	}
+}
+
+#endif // SHADER_CLUSTERAABB
\ No newline at end of file
diff --git a/Engine/Source/Editor/CommonMenuExtensions/Private/GroomVisualizationMenuCommands.cpp b/Engine/Source/Editor/CommonMenuExtensions/Private/GroomVisualizationMenuCommands.cpp
--- a/Engine/Source/Editor/CommonMenuExtensions/Private/GroomVisualizationMenuCommands.cpp
+++ b/Engine/Source/Editor/CommonMenuExtensions/Private/GroomVisualizationMenuCommands.cpp
@@ -114,7 +114,7 @@ void FGroomVisualizationMenuCommands::InternalBuildVisualisationSubMenu(FMenuBui
 		{
 			Menu.BeginSection("LevelViewportGroomVisualizationAdvanced", LOCTEXT("GroomVisualizationAdvanced", "Groom Strands Advanced"));
 			Commands.AddCommandTypeToMenu(Menu, EGroomViewMode::VoxelsDensity);
-			Commands.AddCommandTypeToMenu(Menu, EGroomViewMode::Cluster);
+			//Commands.AddCommandTypeToMenu(Menu, EGroomViewMode::Cluster);
 			Commands.AddCommandTypeToMenu(Menu, EGroomViewMode::ClusterAABB);
 			if (!bIsGroomEditor)
 			{
diff --git a/Engine/Source/Runtime/Renderer/Public/HairStrandsInterface.h b/Engine/Source/Runtime/Renderer/Public/HairStrandsInterface.h
--- a/Engine/Source/Runtime/Renderer/Public/HairStrandsInterface.h
+++ b/Engine/Source/Runtime/Renderer/Public/HairStrandsInterface.h
@@ -299,6 +299,7 @@ public:
 	bool GetLODVisibility() const { return bLODVisibility; }
 
 	uint32 GetClusterCount() const { return ClusterCount;  }
+	float GetClusterScale() const { return ClusterScale;  }
 
 	// Return the number of active point/curve for strand geometry
 	RENDERER_API uint32 GetActiveStrandsPointCount() const;
@@ -390,6 +391,7 @@ public:
 	uint32 RestPointCount = 0;
 	uint32 RestCurveCount = 0;
 	uint32 ClusterCount = 0;
+	float ClusterScale = 0;
 
 	bool bSupportVoxelization = true;
 
