Apply random color to instances when visualizing mesh distance fields.
- plus a couple of small changes.

diff --git a/Engine/Shaders/Private/DistanceField/GlobalDistanceField.usf b/Engine/Shaders/Private/DistanceField/GlobalDistanceField.usf
--- a/Engine/Shaders/Private/DistanceField/GlobalDistanceField.usf
+++ b/Engine/Shaders/Private/DistanceField/GlobalDistanceField.usf
@@ -349,7 +349,7 @@ void AllocatePagesCS(
 		uint CullGridObjectNum = CullGridObjectHeader[CullGridLinearIndex * 2 + 0];
 		uint CullGridObjectOffset = CullGridObjectHeader[CullGridLinearIndex * 2 + 1];
 
-		for (uint IndexInCullCell = ThreadIndex; IndexInCullCell < CullGridObjectNum; IndexInCullCell += 64)
+		for (uint IndexInCullCell = ThreadIndex; IndexInCullCell < CullGridObjectNum; IndexInCullCell += THREADGROUP_SIZE_X)
 		{
 			uint ObjectIndex = CullGridObjectArray[CullGridObjectOffset + IndexInCullCell];
 
diff --git a/Engine/Shaders/Private/DistanceFieldVisualization.usf b/Engine/Shaders/Private/DistanceFieldVisualization.usf
--- a/Engine/Shaders/Private/DistanceFieldVisualization.usf
+++ b/Engine/Shaders/Private/DistanceFieldVisualization.usf
@@ -1,6 +1,7 @@
 // Copyright Epic Games, Inc. All Rights Reserved.
 
 #include "Common.ush"
+#include "Visualization.ush"
 #include "DeferredShadingCommon.ush"
 #include "DistanceFieldLightingShared.ush"
 #include "DistanceFieldAOShared.ush"
@@ -248,6 +249,9 @@ void VisualizeMeshDistanceFieldCS(
 
 		Lighting += GetSkySHDiffuse(WorldHitNormal) * View.SkyLightColor.rgb;
 		float3 Albedo = .05f;
+#if !USE_GLOBAL_DISTANCE_FIELD
+		Albedo = 0.1f * IntToColor(HitCulledObjectIndex + 1);
+#endif
 		Lighting *= Albedo;
 	}
 	else if (ReflectionStruct.SkyLightParameters.y > 0)
