Path Tracer: Fix dark refraction with SingleLayerWater materials

This was actually caused by the Tangent Space Normal setting not being respected properly.

Also cleanup unused code for SLW for the substrate case.

diff --git a/Engine/Shaders/Private/PathTracing/PathTracingGPULightmassMaterialHitShader.usf b/Engine/Shaders/Private/PathTracing/PathTracingGPULightmassMaterialHitShader.usf
--- a/Engine/Shaders/Private/PathTracing/PathTracingGPULightmassMaterialHitShader.usf
+++ b/Engine/Shaders/Private/PathTracing/PathTracingGPULightmassMaterialHitShader.usf
@@ -331,6 +331,10 @@ RAY_TRACING_ENTRY_CLOSEST_HIT(GPULightmassMaterialCHS,
 #endif
 	Payload.WorldGeoNormal = GeoNormalSign * GeoNormal;
 	Payload.WorldNormal = MaterialParameters.WorldNormal;
+#if !MATERIAL_TANGENTSPACENORMAL
+	// already flipped if the material was in tangent space, so add the flip if it wasn't
+	Payload.WorldNormal *= MaterialParameters.TwoSidedSign;
+#endif
 
 	/**
 	 * Set material attributes for simplified materials
diff --git a/Engine/Shaders/Private/PathTracing/PathTracingMaterialHitShader.usf b/Engine/Shaders/Private/PathTracing/PathTracingMaterialHitShader.usf
--- a/Engine/Shaders/Private/PathTracing/PathTracingMaterialHitShader.usf
+++ b/Engine/Shaders/Private/PathTracing/PathTracingMaterialHitShader.usf
@@ -923,6 +923,10 @@ RAY_TRACING_ENTRY_CLOSEST_HIT(PathTracingMaterialCHS,
 	Payload.Radiance = GetMaterialEmissive(PixelMaterialInputs);
 
 	Payload.WorldNormal = MaterialParameters.WorldNormal;
+#if !MATERIAL_TANGENTSPACENORMAL
+	// already flipped if the material was in tangent space, so add the flip if it wasn't
+	Payload.WorldNormal *= MaterialParameters.TwoSidedSign;
+#endif
 	Payload.WorldNormal = AdjustShadingNormal(Payload.WorldNormal, Payload.WorldGeoNormal, WorldRayDirection());
 
 	// Store the results in local variables and reuse instead of calling the functions multiple times.
@@ -1086,18 +1090,6 @@ RAY_TRACING_ENTRY_CLOSEST_HIT(PathTracingMaterialCHS,
 #if MATERIAL_SHADINGMODEL_SINGLELAYERWATER
 	if (Payload.ShadingModelID == SHADINGMODELID_SINGLELAYERWATER)
 	{
-#if STRATA_INLINE_SINGLELAYERWATER
-		SanitizeStrataSingleLayerWater(MaterialParameters.StrataTree.BSDFs[0]);
-		FStrataBSDF SLWBSDF = MaterialParameters.StrataTree.BSDFs[0];
-		const float3 Albedo = SLW_WATERALBEDO(SLWBSDF);
-		const float3 ExtinctionCoeff = SLW_WATEREXTINCTION(SLWBSDF);
-		const float3 ScatteringCoeff = Albedo * ExtinctionCoeff;
-		const float3 AbsorptionCoeff = max(0.0, ExtinctionCoeff - ScatteringCoeff);
-		const float PhaseG = SLW_WATERPHASEG(SLWBSDF);
-
-		const float3 WaterExtinction = ScatteringCoeff + AbsorptionCoeff;
-		const float3 WaterAlbedo = ScatteringCoeff / WaterExtinction;
-#else
 		// For single layer water we always want to treat it as front facing when evaluating parameters 
 		// as hitting back face is treated as being underwater and we get completely different parameters.
 		FMaterialPixelParameters WaterMaterialParameters = MaterialParameters;
@@ -1108,8 +1100,7 @@ RAY_TRACING_ENTRY_CLOSEST_HIT(PathTracingMaterialCHS,
 		const float PhaseG = clamp(LWCToFloat(GetSingleLayerWaterMaterialOutput2(WaterMaterialParameters)), -0.99f, 0.99f);
 
 		const float3 WaterExtinction = ScatteringCoeff + AbsorptionCoeff;
-		const float3 WaterAlbedo = ScatteringCoeff / WaterExtinction;
-#endif
+		const float3 WaterAlbedo = select(WaterExtinction > 0.0, ScatteringCoeff / WaterExtinction, float3(0.0, 0.0, 0.0));
 
 		Payload.ShadingModelID = SHADINGMODELID_DEFAULT_LIT; // replace shading model to avoid extra cases in raygen
 		Payload.BaseColor = WaterAlbedo; // TODO: unused at the moment
