Support min/max depth bounds when culling objects for directional light distance field shadows without scatter tile culling.

diff --git a/Engine/Shaders/Private/DistanceFieldShadowing.usf b/Engine/Shaders/Private/DistanceFieldShadowing.usf
--- a/Engine/Shaders/Private/DistanceFieldShadowing.usf
+++ b/Engine/Shaders/Private/DistanceFieldShadowing.usf
@@ -433,9 +433,9 @@ void CullObjectsToTileWithGather(
 	uint ThreadIndex, 
 	uint2 GroupId, 
 	float TraceDistance, 
-	float MinDepth, 
+	float MinDepth,
+	float MaxDepth,
 	out uint NumIntersectingObjects, 
-	out bool bTileShouldComputeShadowing,
 	out uint GroupIndex)
 {
 	// Initialize per-tile variables
@@ -451,7 +451,7 @@ void CullObjectsToTileWithGather(
 
     GroupMemoryBarrierWithGroupSync();
     
-	if (SceneDepth > MinDepth)
+	if (SceneDepth > MinDepth && SceneDepth < MaxDepth)
 	{
 		// Use shared memory atomics to build the depth bounds for this tile
 		// Each thread is assigned to a pixel at this point
@@ -469,7 +469,7 @@ void CullObjectsToTileWithGather(
 
 	// Compute a second min and max Z, clipped by HalfZ, so that we get two depth bounds per tile
 	// This results in more conservative tile depth bounds and fewer intersections
-	if (SceneDepth >= HalfZ && SceneDepth > MinDepth)
+	if (SceneDepth >= HalfZ && SceneDepth < MaxDepth)
 	{
 		InterlockedMin(IntegerTileMinZ2, asuint(SceneDepth));
 	}
@@ -530,17 +530,16 @@ void CullObjectsToTileWithGather(
 	float RayLength1 = LightVector1Length;
 
 	// Don't operate on tiles completely outside of the light's influence
-	bTileShouldComputeShadowing = LightVector0Length < 1.0f / LightTranslatedPositionAndInvRadius.w + Group0BoundingRadius
+	bool bTileShouldComputeShadowing = LightVector0Length < 1.0f / LightTranslatedPositionAndInvRadius.w + Group0BoundingRadius
 		|| LightVector1Length < 1.0f / LightTranslatedPositionAndInvRadius.w + Group1BoundingRadius;
-
 #else
 	float3 LightDirection0 = LightDirection;
 	float3 LightDirection1 = LightDirection;
 	float RayLength0 = TraceDistance;
 	float RayLength1 = TraceDistance;
 
-	bTileShouldComputeShadowing = SceneDepth > MinDepth;
-
+	// Don't operate on tiles completely outside of the [MinDepth, MaxDepth] range
+	bool bTileShouldComputeShadowing = MaxTileZ > MinDepth && MinTileZ < MaxDepth;
 #endif
 
 	BRANCH
@@ -612,23 +611,21 @@ void DistanceFieldShadowingCS(
 
 	uint NumIntersectingObjects = GetCulledNumObjects();
 	uint CulledDataParameter = 0;
-	bool bTileShouldComputeShadowing = SceneDepth > MinDepth && SceneDepth < MaxDepth;
+	bool bShouldComputeShadowing = SceneDepth > MinDepth && SceneDepth < MaxDepth;
 
 #define USE_CULLING 1
 #if USE_CULLING
 
 #if SCATTER_TILE_CULLING
 
-	if (bTileShouldComputeShadowing)
+	if (bShouldComputeShadowing)
 	{
 		GetShadowTileCulledData(OpaqueTranslatedWorldPosition, CulledDataParameter, NumIntersectingObjects);
 	}
 
 #else
 
-	//@todo - support MinDepth for tile culling
-	float MinDepthForTileCulling = 0;
-	CullObjectsToTileWithGather(SceneDepth, ThreadIndex, GroupId.xy, TraceDistance, MinDepthForTileCulling, NumIntersectingObjects, bTileShouldComputeShadowing, CulledDataParameter);
+	CullObjectsToTileWithGather(SceneDepth, ThreadIndex, GroupId.xy, TraceDistance, MinDepth, MaxDepth, NumIntersectingObjects, CulledDataParameter);
 
 #endif
 #endif // USE_CULLING
@@ -638,7 +635,7 @@ void DistanceFieldShadowingCS(
 #define COMPUTE_SHADOWING 1
 #if COMPUTE_SHADOWING
 	BRANCH
-	if (bTileShouldComputeShadowing)
+	if (bShouldComputeShadowing && NumIntersectingObjects > 0)
 	{
 		// Keeps result from going all the way sharp
 		float MinSphereRadius = .4f;
@@ -760,8 +757,8 @@ void DistanceFieldShadowingCS(
 		}
 	}
 #else
-	//Result = bTileShouldComputeShadowing;
-	Result = NumIntersectingObjects / 500.0f;
+	//Result = bShouldComputeShadowing;
+	Result = bShouldComputeShadowing ? NumIntersectingObjects / 256.0f : 0.0f;
 #endif
 
 	#if METAL_PROFILE 
