Fix for Hi-S zcompare base on some platforms

diff --git a/Engine/Shaders/Private/HTileEncoding.ush b/Engine/Shaders/Private/HTileEncoding.ush
--- a/Engine/Shaders/Private/HTileEncoding.ush
+++ b/Engine/Shaders/Private/HTileEncoding.ush
@@ -60,7 +60,7 @@ uint EncodeZRangeDelta(uint ZBase, uint ZOther)
 	}
 }
 
-uint EncodeTileStencilZBaseDelta(float MinDepth, float MaxDepth, uint XX, uint SMem, uint SR0, uint SR1)
+uint EncodeTileStencilZBaseDelta(float MinDepth, float MaxDepth, uint XX, uint SMem, uint SR0, uint SR1, bool bZCompareGE)
 {
 	const float MaxDepthRange = float((1 << 14) - 1);
 	
@@ -69,7 +69,7 @@ uint EncodeTileStencilZBaseDelta(float MinDepth, float MaxDepth, uint XX, uint S
 	const uint ZMax  = uint( ceil(MaxDepth * MaxDepthRange));
 	
 	// Assumes reverse Z
-	const uint ZBase = ZMax;
+	const uint ZBase = select(bZCompareGE, ZMax, ZMin);
 	const uint ZDelta = EncodeZRangeDelta(ZMin, ZMax);
 	
 	uint HTileValue = 0;
