Changed Base comparison value in PositiveClampedPow from 0 to 2.980233e-8. Very small values for base could still generate floating point specials on some platforms.

diff --git a/Engine/Shaders/Private/Common.ush b/Engine/Shaders/Private/Common.ush
--- a/Engine/Shaders/Private/Common.ush
+++ b/Engine/Shaders/Private/Common.ush
@@ -657,10 +657,11 @@ MaterialFloat4 ClampedPow(MaterialFloat4 X,MaterialFloat4 Y)
 	return pow(max(abs(X),MaterialFloat4(POW_CLAMP,POW_CLAMP,POW_CLAMP,POW_CLAMP)),Y);
 } 
 
-/* Pow function that will return 0 if Base is <=0. This ensures that no compiler expands pow into exp(Exponent * log(Base)) with Base=0 */
+// Pow function that will return 0 if Base is <= 0 (or small enough to yield a floating point special).
+// This is done to prevent floating point specials when compilers expands pow into exp(Exponent * log(Base)).
 MaterialFloat PositiveClampedPow(MaterialFloat Base, MaterialFloat Exponent)
 {
-	return (Base <= 0.0f) ? 0.0f : pow(Base, Exponent);
+	return (Base <= 2.980233e-8f) ? 0.0f : pow(Base, Exponent);
 }
 MaterialFloat2 PositiveClampedPow(MaterialFloat2 Base, MaterialFloat2 Exponent)
 {
