Compiles TSR wave ops on SPIRV backends

This changes FTSRRejectShadingCS to use wave ops on row convolutions and LDS on column convolutions on desktop to reduce VGPR pressure and bring compilation time under 30s in DXC's SPIRV backend.

diff --git a/Engine/Shaders/Private/TemporalSuperResolution/TSRConvolutionNetwork.ush b/Engine/Shaders/Private/TemporalSuperResolution/TSRConvolutionNetwork.ush
--- a/Engine/Shaders/Private/TemporalSuperResolution/TSRConvolutionNetwork.ush
+++ b/Engine/Shaders/Private/TemporalSuperResolution/TSRConvolutionNetwork.ush
@@ -6,8 +6,12 @@
 #include "/Engine/Public/LaneVectorization.ush"
 
 
-#ifndef LANE_COUNT
-	#error LANE_COUNT is undefined.
+#ifndef WAVE_COUNT_X
+	#error WAVE_COUNT_X is undefined.
+#endif
+
+#ifndef WAVE_COUNT_Y
+	#error WAVE_COUNT_Y is undefined.
 #endif
 
 #ifndef LANE_COUNT_X
@@ -18,6 +22,13 @@
 	#error LANE_COUNT_Y is undefined.
 #endif
 
+/** Total number of waves. */
+#define WAVE_COUNT (WAVE_COUNT_X * WAVE_COUNT_Y)
+
+/** Total number of lane per wave. */
+#define LANE_COUNT (LANE_COUNT_X * LANE_COUNT_Y)
+
+
 /**
  * LANE_COUNT_X=2
  * LANE_COUNT_Y=2
@@ -62,83 +73,58 @@ template<uint LaneStrideX, uint LaneStrideY>
 tsr_short2 GetLaneSimdPixelOffset(const uint LaneIndex, const uint SimdIndex)
 {
 	return (
-		tsr_short2(LaneIndex % uint(LANE_COUNT_X), LaneIndex / uint(LANE_COUNT_X)) * tsr_short2(LaneStrideX, LaneStrideY) +
+		tsr_short2(LaneIndex % uint(LANE_COUNT_X * WAVE_COUNT_X), LaneIndex / uint(LANE_COUNT_X * WAVE_COUNT_X)) * tsr_short2(LaneStrideX, LaneStrideY) +
 		GetSimdIndexPixelCoordinateInLane<LaneStrideX, LaneStrideY>(SimdIndex)
 	);
 }
 
 
-//------------------------------------------------------- FUNCTIONS
-
-/** write out elements to LDS for use with AccessNeighborTexel().
- * 
- * WriteToLDSForAccessNeighborTexel(ToConvolve);
- * AccessNeighborTexel(ToConvolve, tsr_short2(1, 0));
- */
-#if CONFIG_OPTIMIZED_LDS_ACCESS
+//------------------------------------------------------- GROUP SHARED COMMUNICATION
 
-groupshared tsr_half4 SharedData[GROUP_TILE_SIZE * GROUP_TILE_SIZE];
+#if defined(LDS_SIZE)
 
-uint GetElementLDSIndex(const uint NeigborSimdIndex, const uint LaneRotation)
-{
-	return ((GGroupThreadIndex + LaneRotation) % uint(LANE_COUNT)) + NeigborSimdIndex * LANE_COUNT;
-}
+groupshared tsr_half4 SharedData[LDS_SIZE];
 
-void WriteElementToLDS(const uint SimdIndex, tsr_half X)
+void WriteElementToLDS(const uint SharedIndex, tsr_half X)
 {
-	SharedData[GGroupThreadIndex + SimdIndex * LANE_COUNT] = tsr_half4(X, 0.0, 0.0, 0.0);
+	SharedData[SharedIndex] = tsr_half4(X, 0.0, 0.0, 0.0);
 }
 
-void WriteElementToLDS(const uint SimdIndex, tsr_half3 X)
+void WriteElementToLDS(const uint SharedIndex, tsr_half3 X)
 {
-	SharedData[GGroupThreadIndex + SimdIndex * LANE_COUNT] = tsr_half4(X, 0.0);
+	SharedData[SharedIndex] = tsr_half4(X, 0.0);
 }
 
-void WriteElementToLDS(const uint SimdIndex, tsr_half4 X)
+void WriteElementToLDS(const uint SharedIndex, tsr_half4 X)
 {
-	SharedData[GGroupThreadIndex + SimdIndex * LANE_COUNT] = tsr_half4(X);
+	SharedData[SharedIndex] = tsr_half4(X);
 }
 
-void ReadElementFromLDS(const uint NeigborSimdIndex, const uint LaneRotation, out tsr_half X)
+void ReadElementFromLDS(const uint SharedIndex, out tsr_half X)
 {
-	X = SharedData[GetElementLDSIndex(NeigborSimdIndex, LaneRotation)].x;
+	X = SharedData[SharedIndex].x;
 }
 
-void ReadElementFromLDS(const uint NeigborSimdIndex, const uint LaneRotation, out tsr_half3 X)
+void ReadElementFromLDS(const uint SharedIndex, out tsr_half3 X)
 {
-	X = SharedData[GetElementLDSIndex(NeigborSimdIndex, LaneRotation)].xyz;
+	X = SharedData[SharedIndex].xyz;
 }
 
-void ReadElementFromLDS(const uint NeigborSimdIndex, const uint LaneRotation, out tsr_half4 X)
+void ReadElementFromLDS(const uint SharedIndex, out tsr_half4 X)
 {
-	X = SharedData[GetElementLDSIndex(NeigborSimdIndex, LaneRotation)].xyzw;
+	X = SharedData[SharedIndex].xyzw;
 }
 
-template<typename ScalarType, uint VectorSize, uint LaneStrideX, uint LaneStrideY>
-void WriteToLDSForAccessNeighborTexel(TLaneVector2D<ScalarType, VectorSize, LaneStrideX, LaneStrideY> Input)
-{
-	GroupMemoryBarrierWithGroupSync();
-	UNROLL
-	for (uint SimdIndex = 0; SimdIndex < LaneStrideX * LaneStrideY; SimdIndex++)
-	{
-		WriteElementToLDS(SimdIndex, Input.GetElement(SimdIndex));
-	}
-	GroupMemoryBarrierWithGroupSync();
-}
+#endif // defined(LDS_SIZE)
 
-#else // !CONFIG_OPTIMIZED_LDS_ACCESS
 
-template<typename ScalarType, uint VectorSize, uint LaneStrideX, uint LaneStrideY>
-void WriteToLDSForAccessNeighborTexel(TLaneVector2D<ScalarType, VectorSize, LaneStrideX, LaneStrideY> Input)
-{
-	// NOP
-}
+//------------------------------------------------------- FUNCTIONS
 
-#endif // !CONFIG_OPTIMIZED_LDS_ACCESS
+#if WAVE_COUNT_X == 1 || WAVE_COUNT_Y == 1
 
 CALL_SITE_DEBUGLOC
 template<typename ScalarType, uint VectorSize, uint LaneStrideX, uint LaneStrideY>
-TLaneVector2D<ScalarType, VectorSize, LaneStrideX, LaneStrideY> AccessNeighborTexel(
+TLaneVector2D<ScalarType, VectorSize, LaneStrideX, LaneStrideY> WaveAccessNeighborTexel(
 	TLaneVector2D<ScalarType, VectorSize, LaneStrideX, LaneStrideY> Input,
 	const tsr_short2 Offset)
 {
@@ -161,24 +147,14 @@ TLaneVector2D<ScalarType, VectorSize, LaneStrideX, LaneStrideY> AccessNeighborTe
 		vector<ScalarType, VectorSize> SimdElement = Input.GetElement(NeigborSimdIndex);
 		vector<ScalarType, VectorSize> ReturnSimdElement = SimdElement;
 		if (LaneRotation != 0)
-		#if CONFIG_OPTIMIZED_LDS_ACCESS
-		{
-			ReadElementFromLDS(NeigborSimdIndex, LaneRotation, /* out */ ReturnSimdElement);
-		}
-		#else
 		{
 			// Only uses wave intrinsics on platforms that provides RDNA's ds_swizzle because we know the wave size of this platforms. 
-			#if DIM_WAVE_SIZE != 0
-				#if !PLATFORM_SUPPORTS_WAVE_ROTATE
-					//#error InitWaveRotateLaneGroup() isn't supported.
-				#endif
-				const FWaveBroadcastSettings BroadcastSettings = InitWaveRotateLaneGroup(/* LaneGroupSize = */ LANE_COUNT, LaneRotation);
-			#else
-				const FWaveBroadcastSettings BroadcastSettings = ConvertWaveBroadcastToLDS(InitWaveRotateLaneGroup(/* LaneGroupSize = */ LANE_COUNT, LaneRotation));
+			#if !PLATFORM_SUPPORTS_WAVE_ROTATE
+				//#error InitWaveRotateLaneGroup() isn't supported.
 			#endif
+			const FWaveBroadcastSettings BroadcastSettings = InitWaveRotateLaneGroup(/* LaneGroupSize = */ LANE_COUNT, LaneRotation);
 			ReturnSimdElement = WaveBroadcast(BroadcastSettings, SimdElement);
 		}
-		#endif
 
 		Return.SetElement(SimdIndex, ReturnSimdElement);
 	}
@@ -186,20 +162,127 @@ TLaneVector2D<ScalarType, VectorSize, LaneStrideX, LaneStrideY> AccessNeighborTe
 	return Return;
 }
 
+#endif // WAVE_COUNT_X == 1 || WAVE_COUNT_Y == 1
+
 CALL_SITE_DEBUGLOC
 template<typename ScalarType, uint VectorSize, uint LaneStrideX, uint LaneStrideY>
 void AccessNeighborTexels3x1(
 	TLaneVector2D<ScalarType, VectorSize, LaneStrideX, LaneStrideY> CO,
-	const tsr_short2 OffsetP,
 	out TLaneVector2D<ScalarType, VectorSize, LaneStrideX, LaneStrideY> CP,
 	out TLaneVector2D<ScalarType, VectorSize, LaneStrideX, LaneStrideY> CN)
+#if WAVE_COUNT_X > 1
 {
-	const tsr_short2 OffsetN = tsr_short2(-int2(OffsetP)); // work arround a shader compiler bug
+	const uint TotalNumberOfLaneInGroup = LANE_COUNT_Y * WAVE_COUNT;
+
+	if (LaneStrideX == 1)
+	{
+		GroupMemoryBarrierWithGroupSync();
+		UNROLL
+		for (uint SimdIndex = 0; SimdIndex < LaneStrideY; SimdIndex++)
+		{
+			WriteElementToLDS(GGroupThreadIndex + SimdIndex * TotalNumberOfLaneInGroup, CO.GetElement(SimdIndex));
+		}
+		GroupMemoryBarrierWithGroupSync();
+		
+		UNROLL
+		for (uint SimdIndex = 0; SimdIndex < LaneStrideY; SimdIndex++)
+		{
+			vector<ScalarType, VectorSize> FirstElement;
+			vector<ScalarType, VectorSize> LastElement;
+			ReadElementFromLDS(((GGroupThreadIndex - 1) % TotalNumberOfLaneInGroup) + SimdIndex * TotalNumberOfLaneInGroup, /* out */ FirstElement);
+			ReadElementFromLDS(((GGroupThreadIndex + 1) % TotalNumberOfLaneInGroup) + SimdIndex * TotalNumberOfLaneInGroup, /* out */ LastElement);
 	
-	WriteToLDSForAccessNeighborTexel(CO);
-	CP = AccessNeighborTexel(CO, OffsetP);
-	CN = AccessNeighborTexel(CO, OffsetN);
+			CP.SetElement(SimdIndex, LastElement);
+			CN.SetElement(SimdIndex, FirstElement);
+		}
+	}
+	else
+	{
+		// TODO
+		CP = CN = CO;
+	}
+}
+#else
+{
+	CP = WaveAccessNeighborTexel(CO, tsr_short2(1, 0));
+	CN = WaveAccessNeighborTexel(CO, tsr_short2(-1, 0));
 }
+#endif
+
+
+CALL_SITE_DEBUGLOC
+ISOLATE
+template<typename ScalarType, uint VectorSize, uint LaneStrideX, uint LaneStrideY>
+void AccessNeighborTexels1x3(
+	TLaneVector2D<ScalarType, VectorSize, LaneStrideX, LaneStrideY> CO,
+	out TLaneVector2D<ScalarType, VectorSize, LaneStrideX, LaneStrideY> CP,
+	out TLaneVector2D<ScalarType, VectorSize, LaneStrideX, LaneStrideY> CN)
+#if WAVE_COUNT_Y > 1
+{
+	const uint TotalNumberOfLaneInGroup = LANE_COUNT_X * WAVE_COUNT;
+
+	if (LaneStrideY == 1)
+	{
+		GroupMemoryBarrierWithGroupSync();
+		UNROLL
+		for (uint SimdIndex = 0; SimdIndex < LaneStrideX; SimdIndex++)
+		{
+			WriteElementToLDS(GGroupThreadIndex + SimdIndex * TotalNumberOfLaneInGroup, CO.GetElement(SimdIndex));
+		}
+		GroupMemoryBarrierWithGroupSync();
+		
+		UNROLL
+		for (uint SimdIndex = 0; SimdIndex < LaneStrideX; SimdIndex++)
+		{
+			vector<ScalarType, VectorSize> FirstElement;
+			vector<ScalarType, VectorSize> LastElement;
+			ReadElementFromLDS(((GGroupThreadIndex - LANE_COUNT_X * WAVE_COUNT_X) % TotalNumberOfLaneInGroup) + SimdIndex * TotalNumberOfLaneInGroup, /* out */ FirstElement);
+			ReadElementFromLDS(((GGroupThreadIndex + LANE_COUNT_X * WAVE_COUNT_X) % TotalNumberOfLaneInGroup) + SimdIndex * TotalNumberOfLaneInGroup, /* out */ LastElement);
+	
+			CP.SetElement(SimdIndex, LastElement);
+			CN.SetElement(SimdIndex, FirstElement);
+		}
+	}
+	else
+	{
+		const uint FirstSimdIndex = 0;
+		const uint LastSimdIndex = LaneStrideX * LaneStrideY - LaneStrideX;
+
+		GroupMemoryBarrierWithGroupSync();
+		UNROLL
+		for (uint SimdIndex = 0; SimdIndex < LaneStrideX; SimdIndex++)
+		{
+			WriteElementToLDS(GGroupThreadIndex + (0 + SimdIndex) * TotalNumberOfLaneInGroup, CO.GetElement(FirstSimdIndex + SimdIndex));
+			WriteElementToLDS(GGroupThreadIndex + (2 + SimdIndex) * TotalNumberOfLaneInGroup, CO.GetElement(LastSimdIndex + SimdIndex));
+		}
+		GroupMemoryBarrierWithGroupSync();
+	
+		UNROLL
+		for (uint SimdIndex = 0; SimdIndex < LastSimdIndex; SimdIndex++)
+		{
+			CP.SetElement(SimdIndex + 0, CO.GetElement(SimdIndex + LaneStrideX));
+			CN.SetElement(SimdIndex + LaneStrideX, CO.GetElement(SimdIndex + 0));
+		}
+	
+		UNROLL
+		for (uint SimdIndex = 0; SimdIndex < LaneStrideX; SimdIndex++)
+		{
+			vector<ScalarType, VectorSize> FirstElement;
+			vector<ScalarType, VectorSize> LastElement;
+			ReadElementFromLDS(((GGroupThreadIndex - LANE_COUNT_X * WAVE_COUNT_X) % TotalNumberOfLaneInGroup) + (2 + SimdIndex) * TotalNumberOfLaneInGroup, /* out */ FirstElement);
+			ReadElementFromLDS(((GGroupThreadIndex + LANE_COUNT_X * WAVE_COUNT_X) % TotalNumberOfLaneInGroup) + (0 + SimdIndex) * TotalNumberOfLaneInGroup, /* out */ LastElement);
+	
+			CP.SetElement(LastSimdIndex + SimdIndex, LastElement);
+			CN.SetElement(FirstSimdIndex + SimdIndex, FirstElement);
+		}
+	}
+}
+#else
+{
+	CP = WaveAccessNeighborTexel(CO, tsr_short2(0, 1));
+	CN = WaveAccessNeighborTexel(CO, tsr_short2(0, -1));
+}
+#endif
 
 
 //------------------------------------------------------- BLUR 3x3 CONVOLUTION
@@ -211,7 +294,7 @@ TLaneVector2D<ScalarType, VectorSize, LaneStrideX, LaneStrideY> Blur3x3(TLaneVec
 	TLaneVector2D<ScalarType, VectorSize, LaneStrideX, LaneStrideY> CO = Center;
 	TLaneVector2D<ScalarType, VectorSize, LaneStrideX, LaneStrideY> CP;
 	TLaneVector2D<ScalarType, VectorSize, LaneStrideX, LaneStrideY> CN;
-	AccessNeighborTexels3x1(CO, tsr_short2(1, 0), /* out */ CP, /* out */ CN);
+	AccessNeighborTexels1x3(CO, /* out */ CP, /* out */ CN);
 	
 	TLaneVector2D<ScalarType, VectorSize, LaneStrideX, LaneStrideY> BO = (
 		CO * TLaneVector2D<ScalarType, VectorSize, LaneStrideX, LaneStrideY>::Const(0.5) +
@@ -220,7 +303,7 @@ TLaneVector2D<ScalarType, VectorSize, LaneStrideX, LaneStrideY> Blur3x3(TLaneVec
 
 	TLaneVector2D<ScalarType, VectorSize, LaneStrideX, LaneStrideY> BP;
 	TLaneVector2D<ScalarType, VectorSize, LaneStrideX, LaneStrideY> BN;
-	AccessNeighborTexels3x1(BO, tsr_short2(0, 1), /* out */ BP, /* out */ BN);
+	AccessNeighborTexels3x1(BO, /* out */ BP, /* out */ BN);
 
 	return (
 		BO * TLaneVector2D<ScalarType, VectorSize, LaneStrideX, LaneStrideY>::Const(0.5) +
@@ -238,12 +321,12 @@ TLaneVector2D<ScalarType, VectorSize, LaneStrideX, LaneStrideY> Min3x3(TLaneVect
 	TLaneVector2D<ScalarType, VectorSize, LaneStrideX, LaneStrideY> CO = Center;
 	TLaneVector2D<ScalarType, VectorSize, LaneStrideX, LaneStrideY> CP;
 	TLaneVector2D<ScalarType, VectorSize, LaneStrideX, LaneStrideY> CN;
-	AccessNeighborTexels3x1(CO, tsr_short2(1, 0), /* out */ CP, /* out */ CN);
+	AccessNeighborTexels1x3(CO, /* out */ CP, /* out */ CN);
 	
 	TLaneVector2D<ScalarType, VectorSize, LaneStrideX, LaneStrideY> MinO = min3(CN, CP, CO);
 	TLaneVector2D<ScalarType, VectorSize, LaneStrideX, LaneStrideY> MinP;
 	TLaneVector2D<ScalarType, VectorSize, LaneStrideX, LaneStrideY> MinN;
-	AccessNeighborTexels3x1(MinO, tsr_short2(0, 1), /* out */ MinP, /* out */ MinN);
+	AccessNeighborTexels3x1(MinO, /* out */ MinP, /* out */ MinN);
 
 	return min3(MinN, MinP, MinO);
 }
@@ -255,12 +338,12 @@ TLaneVector2D<ScalarType, VectorSize, LaneStrideX, LaneStrideY> Max3x3(TLaneVect
 	TLaneVector2D<ScalarType, VectorSize, LaneStrideX, LaneStrideY> CO = Center;
 	TLaneVector2D<ScalarType, VectorSize, LaneStrideX, LaneStrideY> CP;
 	TLaneVector2D<ScalarType, VectorSize, LaneStrideX, LaneStrideY> CN;
-	AccessNeighborTexels3x1(CO, tsr_short2(1, 0), /* out */ CP, /* out */ CN);
+	AccessNeighborTexels1x3(CO, /* out */ CP, /* out */ CN);
 	
 	TLaneVector2D<ScalarType, VectorSize, LaneStrideX, LaneStrideY> MaxO = max3(CN, CP, CO);
 	TLaneVector2D<ScalarType, VectorSize, LaneStrideX, LaneStrideY> MaxP;
 	TLaneVector2D<ScalarType, VectorSize, LaneStrideX, LaneStrideY> MaxN;
-	AccessNeighborTexels3x1(MaxO, tsr_short2(0, 1), /* out */ MaxP, /* out */ MaxN);
+	AccessNeighborTexels3x1(MaxO, /* out */ MaxP, /* out */ MaxN);
 
 	return max3(MaxN, MaxP, MaxO);
 }
@@ -276,19 +359,19 @@ void MinMax3x3(
 	TLaneVector2D<ScalarType, VectorSize, LaneStrideX, LaneStrideY> CO = Center;
 	TLaneVector2D<ScalarType, VectorSize, LaneStrideX, LaneStrideY> CP;
 	TLaneVector2D<ScalarType, VectorSize, LaneStrideX, LaneStrideY> CN;
-	AccessNeighborTexels3x1(CO, tsr_short2(1, 0), /* out */ CP, /* out */ CN);
+	AccessNeighborTexels1x3(CO, /* out */ CP, /* out */ CN);
 	
 	TLaneVector2D<ScalarType, VectorSize, LaneStrideX, LaneStrideY> MinO = min3(CN, CP, CO);
 	TLaneVector2D<ScalarType, VectorSize, LaneStrideX, LaneStrideY> MaxO = max3(CN, CP, CO);
 
 	TLaneVector2D<ScalarType, VectorSize, LaneStrideX, LaneStrideY> MinP;
 	TLaneVector2D<ScalarType, VectorSize, LaneStrideX, LaneStrideY> MinN;
-	AccessNeighborTexels3x1(MinO, tsr_short2(0, 1), /* out */ MinP, /* out */ MinN);
+	AccessNeighborTexels3x1(MinO, /* out */ MinP, /* out */ MinN);
 	OutMin = min3(MinN, MinP, MinO);
 
 	TLaneVector2D<ScalarType, VectorSize, LaneStrideX, LaneStrideY> MaxP;
 	TLaneVector2D<ScalarType, VectorSize, LaneStrideX, LaneStrideY> MaxN;
-	AccessNeighborTexels3x1(MaxO, tsr_short2(0, 1), /* out */ MaxP, /* out */ MaxN);
+	AccessNeighborTexels3x1(MaxO, /* out */ MaxP, /* out */ MaxN);
 	OutMax = max3(MaxN, MaxP, MaxO);
 }
 #else
@@ -340,39 +423,14 @@ FSampleType Median(FSampleType A, FSampleType B, FSampleType C)
 }
 #endif
 
-// 9 samples median using Smith1996.
-CALL_SITE_DEBUGLOC
-template<typename FSampleType>
-FSampleType Median9(FSampleType In[9])
-{
-	// First layer.
-	FSampleType L0[3];
-	FSampleType M0[3];
-	FSampleType H0[3];
-	for (uint j = 0; j < 3; j++)
-	{
-		LMHOperator(In[j * 3 + 0], In[j * 3 + 1], In[j * 3 + 2], L0[j], M0[j], H0[j]);
-	}
-	
-	// Second layer.
-	FSampleType M1[3];
-	M1[0] = max3(L0[0], L0[1], L0[2]);
-	M1[1] = Median(M0[0], M0[1], M0[2]);
-	M1[2] = min3(H0[0], H0[1], H0[2]);
-	
-	// Third layer.
-	return Median(M1[0], M1[1], M1[2]);
-}
-
 CALL_SITE_DEBUGLOC
 template<typename ScalarType, uint VectorSize, uint LaneStrideX, uint LaneStrideY>
 TLaneVector2D<ScalarType, VectorSize, LaneStrideX, LaneStrideY> Median3x3(TLaneVector2D<ScalarType, VectorSize, LaneStrideX, LaneStrideY> Center)
-#if 1
 {
 	TLaneVector2D<ScalarType, VectorSize, LaneStrideX, LaneStrideY> C0 = Center;
 	TLaneVector2D<ScalarType, VectorSize, LaneStrideX, LaneStrideY> C1;
 	TLaneVector2D<ScalarType, VectorSize, LaneStrideX, LaneStrideY> C2;
-	AccessNeighborTexels3x1(C0, tsr_short2(1, 0), /* out */ C1, /* out */ C2);
+	AccessNeighborTexels1x3(C0, /* out */ C1, /* out */ C2);
 	
 	// First layer.
 	TLaneVector2D<ScalarType, VectorSize, LaneStrideX, LaneStrideY> M0L0;
@@ -382,15 +440,15 @@ TLaneVector2D<ScalarType, VectorSize, LaneStrideX, LaneStrideY> Median3x3(TLaneV
 	
 	TLaneVector2D<ScalarType, VectorSize, LaneStrideX, LaneStrideY> M0L1;
 	TLaneVector2D<ScalarType, VectorSize, LaneStrideX, LaneStrideY> M0L2;
-	AccessNeighborTexels3x1(M0L0, tsr_short2(0, 1), /* out */ M0L1, /* out */ M0L2);
+	AccessNeighborTexels3x1(M0L0, /* out */ M0L1, /* out */ M0L2);
 	
 	TLaneVector2D<ScalarType, VectorSize, LaneStrideX, LaneStrideY> M0M1;
 	TLaneVector2D<ScalarType, VectorSize, LaneStrideX, LaneStrideY> M0M2;
-	AccessNeighborTexels3x1(M0M0, tsr_short2(0, 1), /* out */ M0M1, /* out */ M0M2);
+	AccessNeighborTexels3x1(M0M0, /* out */ M0M1, /* out */ M0M2);
 	
 	TLaneVector2D<ScalarType, VectorSize, LaneStrideX, LaneStrideY> M0H1;
 	TLaneVector2D<ScalarType, VectorSize, LaneStrideX, LaneStrideY> M0H2;
-	AccessNeighborTexels3x1(M0H0, tsr_short2(0, 1), /* out */ M0H1, /* out */ M0H2);
+	AccessNeighborTexels3x1(M0H0, /* out */ M0H1, /* out */ M0H2);
 	
 	// Second layer.
 	TLaneVector2D<ScalarType, VectorSize, LaneStrideX, LaneStrideY> M1[3];
@@ -401,21 +459,6 @@ TLaneVector2D<ScalarType, VectorSize, LaneStrideX, LaneStrideY> Median3x3(TLaneV
 	// Third layer.
 	return Median(M1[0], M1[1], M1[2]);
 }
-#else
-{
-	WriteToLDSForAccessNeighborTexel(Center);
-	
-	TLaneVector2D<ScalarType, VectorSize, LaneStrideX, LaneStrideY> InSamples[9];
-	UNROLL
-	for (uint i = 0; i < 9; i++)
-	{
-		TLaneVector2D<ScalarType, VectorSize, LaneStrideX, LaneStrideY> Sample = AccessNeighborTexel(Center, tsr_short2(kOffsets3x3[i]));
-		InSamples[i] = Sample;
-	}
-
-	return Median9(InSamples);
-}
-#endif
 
 
 //------------------------------------------------------- DOWNSAMPLE 2x2 CONVOLUTIONS
diff --git a/Engine/Shaders/Private/TemporalSuperResolution/TSRRejectShading.usf b/Engine/Shaders/Private/TemporalSuperResolution/TSRRejectShading.usf
--- a/Engine/Shaders/Private/TemporalSuperResolution/TSRRejectShading.usf
+++ b/Engine/Shaders/Private/TemporalSuperResolution/TSRRejectShading.usf
@@ -17,9 +17,6 @@
 
 #define DEBUG_ARRAY_SIZE 8
 
-/** Size of the entire group loaded into memory including overscan. */
-#define GROUP_TILE_SIZE 16
-
 /** Whether should remove firefliers from input and history before comparison to improve stability. */
 #define CONFIG_FIREFLY_REMOVAL 1
 
@@ -29,46 +26,62 @@
 /** Whether to detect flickering. */
 #define CONFIG_MOIRE_DETECTION (DIM_FLICKERING_DETECTION)
 
-#if DIM_WAVE_SIZE == 0
-	#define CONFIG_OPTIMIZED_LDS_ACCESS 1
-#else
-	#define CONFIG_OPTIMIZED_LDS_ACCESS 0
-#endif
-
 #define CONFIG_MANUAL_LDS_SPILL 0
 
 
 //------------------------------------------------------- DERIVED
 
 /** Pixel coordinate stride between each lane in the GROUP_TILE_SIZE. */
-#if DIM_WAVE_SIZE == 64
-	#define LANE_COUNT 64
-	#define LANE_STRIDE_X 2
-	#define LANE_STRIDE_Y 2
+#if DIM_WAVE_SIZE == 32 && PLATFORM_GPU_ARCH == PLATFORM_GPU_ARCH_AMD_RDNA_2
+	/** Load an entire 16x16 tile in a single wave32 with 4x2 worth of data per lane and just 100% waveops all communications. */
+	#define GROUP_TILE_SIZE 16
+
+	#define WAVE_COUNT_X 1
+	#define WAVE_COUNT_Y 1
+
+	#define LANE_COUNT_X 4
+	#define LANE_COUNT_Y 8
 
-#elif DIM_WAVE_SIZE == 32
-	#define LANE_COUNT 32
 	#define LANE_STRIDE_X 4
 	#define LANE_STRIDE_Y 2
+	
+#elif DIM_WAVE_SIZE == 64 || DIM_WAVE_SIZE == 32
+	/** Layouts wave in row major so an entire row of pixels fits within a single wave, and waveops horizontally, but LDS vertically.
+	 * Each lane worth of a 2x2 of data per lane to minimise to trade VGPR pressure for both number of wave ops and LDS.
+	 */
+	#define GROUP_TILE_SIZE 16
 
-#elif DIM_WAVE_SIZE == 0 && 0
-	#define LANE_COUNT 32
-	#define LANE_STRIDE_X 4
+	#define WAVE_COUNT_X 1
+	#define WAVE_COUNT_Y 8
+
+	#define LANE_COUNT_X 8
+	#define LANE_COUNT_Y 1
+
+	#define LANE_STRIDE_X 2
 	#define LANE_STRIDE_Y 2
 
-#elif DIM_WAVE_SIZE == 0 && 1
-	#define LANE_COUNT (GROUP_TILE_SIZE * GROUP_TILE_SIZE)
+	// Each wave needs to output top and bottom rows of pixels into LDS;
+	#define LDS_SIZE (2 * (LANE_STRIDE_X * LANE_COUNT_X) * (WAVE_COUNT_X * WAVE_COUNT_Y))
+
+#elif DIM_WAVE_SIZE == 0
+	/** Load an entire 16x16 with 1x1 worth of data by lane, only using LDS. */
+	#define GROUP_TILE_SIZE 16
+
+	#define WAVE_COUNT_X 16
+	#define WAVE_COUNT_Y 16
+
+	#define LANE_COUNT_X 1
+	#define LANE_COUNT_Y 1
+
 	#define LANE_STRIDE_X 1
 	#define LANE_STRIDE_Y 1
 
+	#define LDS_SIZE (WAVE_COUNT_X * WAVE_COUNT_Y)
+
 #else
 	#error Unknown wave size
 #endif
 
-/** Number of lane the GROUP_TILE_SIZE. */
-#define LANE_COUNT_X (GROUP_TILE_SIZE / LANE_STRIDE_X)
-#define LANE_COUNT_Y (GROUP_TILE_SIZE / LANE_STRIDE_Y)
-
 /** Size of the SIMD per lane that also number of input pixel loaded into each individual lanes. */
 #define SIMD_SIZE (LANE_STRIDE_X * LANE_STRIDE_Y)
 
@@ -137,9 +150,9 @@ tsr_simd_halfC ClampFireFliersWithGuide(tsr_simd_halfC ToClamp, tsr_simd_halfC G
 //------------------------------------------------------- ENTRY POINT
 
 #if COMPILER_SUPPORTS_WAVE_SIZE && DIM_WAVE_SIZE > 0
-	WAVESIZE(LANE_COUNT)
+WAVESIZE(DIM_WAVE_SIZE)
 #endif
-[numthreads(LANE_COUNT, 1, 1)]
+[numthreads(LANE_COUNT * WAVE_COUNT, 1, 1)]
 void MainCS(
 	uint2 GroupId : SV_GroupID,
 	uint GroupThreadIndex : SV_GroupIndex)
@@ -589,7 +602,7 @@ void MainCS(
 			#endif
 
 			#if CONFIG_SCENE_COLOR_ALPHA
-			HistoryGuideOutput[OutputPixelPos] = FinalGuideColor;
+				HistoryGuideOutput[OutputPixelPos] = FinalGuideColor;
 			#else
 				HistoryGuideOutput[OutputPixelPos] = tsr_half4(FinalGuideColor.xyz, tsr_half(1.0));
 			#endif
diff --git a/Engine/Shaders/Private/TemporalSuperResolution/TSRResolveHistory.usf b/Engine/Shaders/Private/TemporalSuperResolution/TSRResolveHistory.usf
--- a/Engine/Shaders/Private/TemporalSuperResolution/TSRResolveHistory.usf
+++ b/Engine/Shaders/Private/TemporalSuperResolution/TSRResolveHistory.usf
@@ -6,8 +6,9 @@
 //------------------------------------------------------- CONFIG
 
 #if DIM_NYQUIST
-	#define LANE_COUNT 32
-	#define DIM_WAVE_SIZE (LANE_COUNT)
+	#define WAVE_COUNT_X 1
+	#define WAVE_COUNT_Y 1
+
 	#define DOWNSAMPLE_FACTOR 2u
 	
 	/** Size of the entire group loaded into memory including overscan. */
@@ -135,7 +136,7 @@ RWTexture2D<tsr_halfC> SceneColorOutputMip1;
 #if DIM_NYQUIST
 
 #if COMPILER_SUPPORTS_WAVE_SIZE
-	WAVESIZE(LANE_COUNT)
+	WAVESIZE(32)
 #endif
 [numthreads(LANE_COUNT, 1, 1)]
 void MainCS(
@@ -169,11 +170,11 @@ void MainCS(
 			tsr_output_halfC Min = DownsampleMin2x2(Colors);
 			tsr_output_halfC Max = DownsampleMax2x2(Colors);
 
-			Min = min(Min, AccessNeighborTexel(Min, tsr_short2(1, 0)));
-			Min = min(Min, AccessNeighborTexel(Min, tsr_short2(0, 1)));
+			Min = min(Min, WaveAccessNeighborTexel(Min, tsr_short2(1, 0)));
+			Min = min(Min, WaveAccessNeighborTexel(Min, tsr_short2(0, 1)));
 
-			Max = max(Max, AccessNeighborTexel(Max, tsr_short2(1, 0)));
-			Max = max(Max, AccessNeighborTexel(Max, tsr_short2(0, 1)));
+			Max = max(Max, WaveAccessNeighborTexel(Max, tsr_short2(1, 0)));
+			Max = max(Max, WaveAccessNeighborTexel(Max, tsr_short2(0, 1)));
 		#endif
 		
 		tsr_input_halfC LDRColors = Colors;
@@ -218,7 +219,7 @@ void MainCS(
 		tsr_output_halfC Out2 = DownsampleDot2x2(LDRColors, Weights2);
 		tsr_output_halfC Out3 = DownsampleDot2x2(LDRColors, Weights3);
 
-		Output = Out0 + AccessNeighborTexel(Out1, tsr_short2(1, 0)) + AccessNeighborTexel(Out2, tsr_short2(0, 1)) + AccessNeighborTexel(Out3, tsr_short2(1, 1));
+		Output = Out0 + WaveAccessNeighborTexel(Out1, tsr_short2(1, 0)) + WaveAccessNeighborTexel(Out2, tsr_short2(0, 1)) + WaveAccessNeighborTexel(Out3, tsr_short2(1, 1));
 		#if CONFIG_HDR_WEIGHT
 		{
 			UNROLL_N(OUTPUT_SIMD_SIZE)
@@ -252,8 +253,8 @@ void MainCS(
 		
 		{
 			tsr_output_halfC HalfResOutput = Output * tsr_half(0.25);
-			HalfResOutput = HalfResOutput + AccessNeighborTexel(HalfResOutput, tsr_short2(1, 0));
-			HalfResOutput = HalfResOutput + AccessNeighborTexel(HalfResOutput, tsr_short2(0, 1));
+			HalfResOutput = HalfResOutput + WaveAccessNeighborTexel(HalfResOutput, tsr_short2(1, 0));
+			HalfResOutput = HalfResOutput + WaveAccessNeighborTexel(HalfResOutput, tsr_short2(0, 1));
 		
 			tsr_short2 OutputPixelOffset = GetLaneSimdPixelOffset<LANE_OUTPUT_STRIDE_X, LANE_OUTPUT_STRIDE_Y>(LaneIndex, /* SimdIndex = */ 0);
 			tsr_short2 OutputPixelPos = InvalidateOutputPixelPos(GroupPixelOffset + OutputPixelOffset, OutputViewRectMax);
diff --git a/Engine/Source/Runtime/Renderer/Private/PostProcess/TemporalSuperResolution.cpp b/Engine/Source/Runtime/Renderer/Private/PostProcess/TemporalSuperResolution.cpp
--- a/Engine/Source/Runtime/Renderer/Private/PostProcess/TemporalSuperResolution.cpp
+++ b/Engine/Source/Runtime/Renderer/Private/PostProcess/TemporalSuperResolution.cpp
@@ -395,12 +395,6 @@ public:
 
 	static ERHIFeatureSupport SupportsWaveOps(EShaderPlatform Platform)
 	{
-		// UE-161125: SPIRV ends up with 5min+ compilation times, and D3D11 needs SPIRV backend for HLSL2021 features on FXC.
-		if (FDataDrivenShaderPlatformInfo::GetIsSPIRV(Platform) || GetEditorShaderPlatform(Platform) == SP_PCD3D_SM5)
-		{
-			return ERHIFeatureSupport::Unsupported;
-		}
-
 		return FDataDrivenShaderPlatformInfo::GetSupportsWaveOperations(Platform);
 	}
 
