* Change ShaderPrint PrintBool to a macro to avoid unecesary processing at shader compilation time.
* Add verify() shader macro with automatic file/line for in-shader validation.

diff --git a/Engine/Shaders/Private/ShaderPrint.ush b/Engine/Shaders/Private/ShaderPrint.ush
--- a/Engine/Shaders/Private/ShaderPrint.ush
+++ b/Engine/Shaders/Private/ShaderPrint.ush
@@ -202,18 +202,19 @@ void AddOBB(float3 Min, float3 Max, float4 Color, float4x4 Transform)
 	AddOBB(InitShaderPrintContext(), Min, Max, Color, Transform);
 }
 
-// Assert (experimental)
+// Verify/Assert (experimental)
 // Can be used in any shader binding ShaderPrint uniform buffer. It will print the 
 // assert for all pixels/threads not complying to the condition. All print will be 
 // stacked correctly.
 // 
 // Example:
-//  Assert(DispatchThreadId < MaxThread, TEXT("Error: ThreadId higher than MaxThread");
-// 
+//  verify(DispatchThreadId < MaxThread);
+//  verifyf(DispatchThreadId < MaxThread, TEXT("Error: ThreadId higher than MaxThread");
+//
 // This function also returns a FShaderPrintContext so that one can add more details if needed.
 //
 // Example:
-//  FShaderPrintContext Ctx = Assert(DispatchThreadId < MaxThread, TEXT("Error: thread higher MaxThread");
+//  FShaderPrintContext Ctx = verify(DispatchThreadId < MaxThread, TEXT("Error: thread higher MaxThread");
 //  Print(Ctx, DispatchThreadId, FontWhite);
 //
 FShaderPrintContext Assert(bool Condition, FShaderPrintText Text, uint2 InWritePos = uint2(50, 150))
@@ -231,3 +232,27 @@ FShaderPrintContext Assert(bool Condition, FShaderPrintText Text, uint2 InWriteP
 
 	return OutCtx;
 }
+
+FShaderPrintContext Assert(bool Condition, FShaderPrintText Text, FShaderPrintText File, uint Line, FShaderPrintText Separator, uint2 InWritePos = uint2(50, 150))
+{
+	FShaderPrintContext OutCtx = InitShaderPrintContext(false, InWritePos);
+	if (!Condition)
+	{
+		uint LineOffset = 0;
+		InterlockedAdd(ShaderPrint_RWEntryBuffer[3] /* Free counter */, 1, LineOffset);
+		InWritePos.y += LineOffset * 12;
+
+		OutCtx = InitShaderPrintContext(true, InWritePos);
+		Print(OutCtx, Text, FontRed);
+		Print(OutCtx, Separator, FontRed);
+		Print(OutCtx, File, FontRed);
+		Print(OutCtx, Separator, FontRed);
+		Print(OutCtx, Line, FontRed);
+	}
+
+	return OutCtx;
+}
+
+// Non-blocking shader check (unlike check()/checkf() which terminate the GPU execution)
+#define verify(Condition) 			Assert(Condition, TEXT(#Condition), TEXT(__FILE__), __LINE__, TEXT(" - "));
+#define verifyf(Condition, Message) Assert(Condition, Message, TEXT(__FILE__), __LINE__, TEXT(" - "));
\ No newline at end of file
diff --git a/Engine/Shaders/Private/ShaderPrintCommon.ush b/Engine/Shaders/Private/ShaderPrintCommon.ush
--- a/Engine/Shaders/Private/ShaderPrintCommon.ush
+++ b/Engine/Shaders/Private/ShaderPrintCommon.ush
@@ -508,29 +508,8 @@ SHADER_PRINT_OVERLOAD_NN(4, 4, float4x4)
 ///////////////////////////////////////////////////////////////////////////////////////////////////
 // Bool printing
 
-void PrintBool(inout FShaderPrintContext Ctx, bool In)
-{
-	if (In)
-	{
-		Print(Ctx, TEXT("Yes "), FontGreen);
-	}
-	else
-	{
-		Print(Ctx, TEXT("No  "), FontRed);
-	}
-}
-
-void PrintBool(inout FShaderPrintContext Ctx, bool In, FFontColor InColor)
-{
-	if (In)
-	{
-		Print(Ctx, TEXT("Yes "), InColor);
-	}
-	else
-	{
-		Print(Ctx, TEXT("No  "), InColor);
-	}
-}
+// Note: Use a define rather than a function to avoid unecessary TEXT() parsing during shader compilation, even for shader not using ShaderPrint.
+#define PrintBool(Ctx, In) { if (In) Print(Ctx, TEXT("Yes "), FontGreen); else Print(Ctx, TEXT("No  "), FontRed); }
 
 ///////////////////////////////////////////////////////////////////////////////////////////////////
 // Widgets (experimental)
