Fix a minor shader compiling bug on mobile OpenGL.

diff --git a/Engine/Shaders/Private/PostProcessSelectionOutline.usf b/Engine/Shaders/Private/PostProcessSelectionOutline.usf
--- a/Engine/Shaders/Private/PostProcessSelectionOutline.usf
+++ b/Engine/Shaders/Private/PostProcessSelectionOutline.usf
@@ -85,7 +85,7 @@ FPixelInfo GetPixelInfo(int2 PixelPos, int SampleID, int OffsetX, int OffsetY, f
 
 #if (COMPILER_GLSL == 1 && FEATURE_LEVEL < FEATURE_LEVEL_SM4) || (COMPILER_METAL && FEATURE_LEVEL < FEATURE_LEVEL_SM4)
 	// Stencil read in opengl is not supported on older versions
-	Ret.ObjectId = Ret.bObjectMask ? 2 : 0;
+	Ret.ObjectId = (Ret.fDepth != 0.0f) ? 2 : 0;
 #endif
 
 	// Soft Bias with DeviceZ for best quality (larger bias than usual because SceneDeviceZ is only approximated)
