Optimises TSR history reprojections

diff --git a/Engine/Shaders/Private/TemporalSuperResolution/TSRCommon.ush b/Engine/Shaders/Private/TemporalSuperResolution/TSRCommon.ush
--- a/Engine/Shaders/Private/TemporalSuperResolution/TSRCommon.ush
+++ b/Engine/Shaders/Private/TemporalSuperResolution/TSRCommon.ush
@@ -1150,3 +1150,46 @@ void ForwardReprojectStaticDepth(float2 PrevScreenPosition, float PrevDeviceZ, o
 	ScreenPos = ClipPos.xy / ClipPos.w;
 	SceneDepth = ClipPos.w;
 }
+
+FCatmullRomSamples GetBicubic2DCatmullRomSamples_Stubbe(float2 UV, float2 Size, in float2 InvSize)
+{
+	float2 PixelCoord = UV * Size;
+
+	float2 iuv = floor(PixelCoord - 0.5) + 1.0;
+	tsr_half2 f = tsr_half2(PixelCoord - iuv);
+	tsr_half2 f2 = f * f;
+	tsr_half2 f_2 = f + f;
+
+	tsr_half2 uv = (tsr_half(1.25) - f2) * f + tsr_half(0.5);
+	tsr_half2 t0 = (tsr_half(0.25) * f2 - tsr_half(0.0625)) * (tsr_half(1.125) - tsr_half(0.5) * f2.yx);
+	tsr_half2 w0 = t0 - f_2 * t0;
+	tsr_half2 w3 = t0 + f_2 * t0;
+
+	float2 cuv = (iuv - 0.5) * InvSize;
+	uv *= tsr_half2(InvSize);
+
+	FCatmullRomSamples Samples;
+	Samples.Count = BICUBIC_CATMULL_ROM_SAMPLES;
+
+	Samples.UV[0] = cuv + tsr_half2(               -InvSize.x,                      uv.y);
+	Samples.UV[1] = cuv + tsr_half2(tsr_half(2.0) * InvSize.x,                      uv.y);
+	Samples.UV[2] = cuv + uv;
+	Samples.UV[3] = cuv + tsr_half2(                     uv.x,                -InvSize.y);
+	Samples.UV[4] = cuv + tsr_half2(                     uv.x, tsr_half(2.0) * InvSize.y);
+
+	Samples.Weight[0] = w0.x;
+	Samples.Weight[1] = w3.x;
+	Samples.Weight[2] = tsr_half(1.0) - w0.x - w3.x - w0.y - w3.y;
+	Samples.Weight[3] = w0.y;
+	Samples.Weight[4] = w3.y;
+
+	Samples.UVDir[0] = int2(-1,  0);
+	Samples.UVDir[1] = int2( 1,  0);
+	Samples.UVDir[2] = int2( 0,  0);
+	Samples.UVDir[3] = int2( 0, -1);
+	Samples.UVDir[4] = int2( 0,  1);
+	
+	Samples.FinalMultiplier = 1.0;
+	return Samples;
+}
+
diff --git a/Engine/Shaders/Private/TemporalSuperResolution/TSRDecimateHistory.usf b/Engine/Shaders/Private/TemporalSuperResolution/TSRDecimateHistory.usf
--- a/Engine/Shaders/Private/TemporalSuperResolution/TSRDecimateHistory.usf
+++ b/Engine/Shaders/Private/TemporalSuperResolution/TSRDecimateHistory.usf
@@ -331,7 +331,7 @@ void ProcessPrevAtomicTexturesSamples(
 template<typename Type>
 Type SampleTextureCatrom(Texture2DArray<Type> InTexture, float2 TextureUV, float FrameIndex, float2 MinUV, float2 MaxUV)
 {
-	FCatmullRomSamples Samples = GetBicubic2DCatmullRomSamples(TextureUV, PrevGuideInfo_Extent, PrevGuideInfo_ExtentInverse);
+	FCatmullRomSamples Samples = GetBicubic2DCatmullRomSamples_Stubbe(TextureUV, PrevGuideInfo_Extent, PrevGuideInfo_ExtentInverse);
 
 	Type Color = tsr_half(0.0);
 	for (uint i = 0; i < Samples.Count; i++)
@@ -344,7 +344,6 @@ Type SampleTextureCatrom(Texture2DArray<Type> InTexture, float2 TextureUV, float
 
 		Color += InTexture.SampleLevel(GlobalBilinearClampedSampler, float3(SampleUV, FrameIndex), 0) * tsr_half(Samples.Weight[i]);
 	}
-	Color *= tsr_half(Samples.FinalMultiplier);
 
 	return Color;
 }
diff --git a/Engine/Shaders/Private/TemporalSuperResolution/TSRUpdateHistory.usf b/Engine/Shaders/Private/TemporalSuperResolution/TSRUpdateHistory.usf
--- a/Engine/Shaders/Private/TemporalSuperResolution/TSRUpdateHistory.usf
+++ b/Engine/Shaders/Private/TemporalSuperResolution/TSRUpdateHistory.usf
@@ -679,114 +679,118 @@ void MainCS(
 	// Reproject history
 	tsr_halfCx2 PrevHighFrequencyColor;
 	tsr_half2 PrevHistoryValidity;
-	
-	tsr_halfCx2 PrevHistoryMomentMin;
-	tsr_halfCx2 PrevHistoryMomentMax;
 
 	PLATFORM_SPECIFIC_ISOLATE
 	{
-		// Reproject history
-		tsr_halfCx2 AccumulateHighFrequency = tsr_half(0);
-		tsr_halfMx2 AccumulateMetadata = tsr_half(0);
-
-		tsr_halfCx2 HighFrequencyMin = INFINITE_FLOAT;
-		tsr_halfCx2 HighFrequencyMax = -INFINITE_FLOAT;
+		tsr_halfC RawHighFrequency[BICUBIC_CATMULL_ROM_SAMPLES][DPV_PIXEL_PER_LANE];
+		tsr_halfM RawMetadata[BICUBIC_CATMULL_ROM_SAMPLES][DPV_PIXEL_PER_LANE];
+		tsr_half2 RawKernelWeight[BICUBIC_CATMULL_ROM_SAMPLES];
 
-		// Sample raw history
+		// Issues texture fetches.
+		PLATFORM_SPECIFIC_ISOLATE
 		{
 			float2x2 PrevHistoryBufferUV = ApplyScreenTransform(PrevScreenPos, ScreenPosToPrevHistoryBufferUV);
 			PrevHistoryBufferUV[0] = fastClamp(PrevHistoryBufferUV[0], PrevHistoryInfo_UVViewportBilinearMin[0], PrevHistoryInfo_UVViewportBilinearMax[0]);
 			PrevHistoryBufferUV[1] = fastClamp(PrevHistoryBufferUV[1], PrevHistoryInfo_UVViewportBilinearMin[1], PrevHistoryInfo_UVViewportBilinearMax[1]);
 
-			FCatmullRomSamples Samples0 = GetBicubic2DCatmullRomSamples(
+
+			FCatmullRomSamples Samples0 = GetBicubic2DCatmullRomSamples_Stubbe(
 				dpv_lo(PrevHistoryBufferUV), PrevHistoryInfo_Extent, PrevHistoryInfo_ExtentInverse);
-			FCatmullRomSamples Samples1 = GetBicubic2DCatmullRomSamples(
+			FCatmullRomSamples Samples1 = GetBicubic2DCatmullRomSamples_Stubbe(
 				dpv_hi(PrevHistoryBufferUV), PrevHistoryInfo_Extent, PrevHistoryInfo_ExtentInverse);
 
-			tsr_halfC RawHighFrequency[BICUBIC_CATMULL_ROM_SAMPLES][DPV_PIXEL_PER_LANE];
-			tsr_halfM RawMetadata[BICUBIC_CATMULL_ROM_SAMPLES][DPV_PIXEL_PER_LANE];
-
-			tsr_half2 RawKernelWeight[BICUBIC_CATMULL_ROM_SAMPLES];
-
-			// Issues texture fetches.
-			PLATFORM_SPECIFIC_ISOLATE
+			UNROLL_N(BICUBIC_CATMULL_ROM_SAMPLES)
+			for (uint i = 0; i < BICUBIC_CATMULL_ROM_SAMPLES; i++)
 			{
-				UNROLL_N(BICUBIC_CATMULL_ROM_SAMPLES)
-				for (uint i = 0; i < BICUBIC_CATMULL_ROM_SAMPLES; i++)
-				{
-					float2 SampleUV0 = fastClamp(Samples0.UV[i], PrevHistoryInfo_UVViewportBilinearMin, PrevHistoryInfo_UVViewportBilinearMax);
-					float2 SampleUV1 = fastClamp(Samples1.UV[i], PrevHistoryInfo_UVViewportBilinearMin, PrevHistoryInfo_UVViewportBilinearMax);
+				float2 SampleUV0 = fastClamp(Samples0.UV[i], PrevHistoryInfo_UVViewportBilinearMin, PrevHistoryInfo_UVViewportBilinearMax);
+				float2 SampleUV1 = fastClamp(Samples1.UV[i], PrevHistoryInfo_UVViewportBilinearMin, PrevHistoryInfo_UVViewportBilinearMax);
 					
-					tsr_half2 KernelWeight = dpv_force_interleave_registers(tsr_half(Samples0.Weight[i]), tsr_half(Samples1.Weight[i]));
+				tsr_half2 KernelWeight = dpv_interleave_registers(tsr_half(Samples0.Weight[i]), tsr_half(Samples1.Weight[i]));
 
-					RawKernelWeight[i] = KernelWeight;
+				RawKernelWeight[i] = KernelWeight;
 
-					UNROLL_N(2)
-					for (uint PixelId = 0; PixelId < DPV_PIXEL_PER_LANE; PixelId++)
-					{
-						float2 SampleUV = PixelId == 0 ? SampleUV0 : SampleUV1;
+				UNROLL_N(2)
+				for (uint PixelId = 0; PixelId < DPV_PIXEL_PER_LANE; PixelId++)
+				{
+					float2 SampleUV = PixelId == 0 ? SampleUV0 : SampleUV1;
 						
-						#if CONFIG_HISTORY_RESURRECTION
-							float FrameIndex = select(bIsHistoryResurrection[PixelId], ResurrectionFrameIndex, PrevFrameIndex);
-						#else
-							float FrameIndex = PrevFrameIndex;
-						#endif
-
-						RawHighFrequency[i][PixelId] = PrevHistoryColorTexture.SampleLevel(GlobalBilinearClampedSampler, float3(SampleUV, FrameIndex), 0);
-						RawMetadata[i][PixelId]      = PrevHistoryMetadataTexture.SampleLevel(GlobalBilinearClampedSampler, float3(SampleUV, FrameIndex), 0);
-					}
+					#if CONFIG_HISTORY_RESURRECTION
+						float FrameIndex = select(bIsHistoryResurrection[PixelId], ResurrectionFrameIndex, PrevFrameIndex);
+					#else
+						float FrameIndex = PrevFrameIndex;
+					#endif
 
-				} // for (uint i = 0; i < BICUBIC_CATMULL_ROM_SAMPLES; i++)
-			}
+					RawHighFrequency[i][PixelId] = PrevHistoryColorTexture.SampleLevel(GlobalBilinearClampedSampler, float3(SampleUV, FrameIndex), 0);
+					RawMetadata[i][PixelId]      = PrevHistoryMetadataTexture.SampleLevel(GlobalBilinearClampedSampler, float3(SampleUV, FrameIndex), 0);
+				}
 
-			// Process history texture fetches.
-			PLATFORM_SPECIFIC_ISOLATE
-			{
-				UNROLL_N(BICUBIC_CATMULL_ROM_SAMPLES)
-				for (uint i = 0; i < BICUBIC_CATMULL_ROM_SAMPLES; i++)
-				{
-					// TODO: Should use dpv_force_interleave_registers_array() but there is a shader compiler bug
-					tsr_halfCx2 SampleHighFrequency = dpv_interleave_registers_array(RawHighFrequency[i]);
-					tsr_halfMx2 SampleMetadata = dpv_interleave_registers_array(RawMetadata[i]);
+			} // for (uint i = 0; i < BICUBIC_CATMULL_ROM_SAMPLES; i++)
+		}
+		
+		#if CONFIG_HISTORY_RESURRECTION
+			const tsr_half2 PreExposureCorrection = select(
+				bIsHistoryResurrection,
+				tsr_half(ResurrectionPreExposureCorrection),
+				tsr_half(HistoryPreExposureCorrection));
+		#else
+			const tsr_half2 PreExposureCorrection = tsr_half(HistoryPreExposureCorrection);
+		#endif
 
-					tsr_half2 KernelWeight = RawKernelWeight[i];
+		// Process history texture fetches.
+		tsr_halfCx2 AccumulateHighFrequency = tsr_half(0);
+		tsr_halfMx2 AccumulateMetadata = tsr_half(0);
 
-					HighFrequencyMin = min(HighFrequencyMin, SampleHighFrequency);
-					HighFrequencyMax = max(HighFrequencyMax, SampleHighFrequency);
+		PLATFORM_SPECIFIC_ISOLATE
+		{
+			UNROLL_N(BICUBIC_CATMULL_ROM_SAMPLES)
+			for (uint i = 0; i < BICUBIC_CATMULL_ROM_SAMPLES; i++)
+			{
+				// TODO: Should use dpv_force_interleave_registers_array() but there is a shader compiler bug
+				tsr_halfCx2 SampleHighFrequency = dpv_interleave_registers_array(RawHighFrequency[i]);
+				tsr_halfMx2 SampleMetadata = dpv_interleave_registers_array(RawMetadata[i]);
 
-					AccumulateHighFrequency += dpv_scale(SampleHighFrequency, KernelWeight);
-					AccumulateMetadata  += dpv_scale(SampleMetadata, KernelWeight);
-				}
-				
-				tsr_half2 FinalMultiplier = dpv_force_interleave_registers(tsr_half(Samples0.FinalMultiplier), tsr_half(Samples1.FinalMultiplier));
+				tsr_half2 KernelWeight = RawKernelWeight[i];
 
-				AccumulateHighFrequency = dpv_scale(AccumulateHighFrequency, FinalMultiplier);
-				AccumulateMetadata  = dpv_scale(AccumulateMetadata, FinalMultiplier);
+				AccumulateHighFrequency[0] += SampleHighFrequency[0] * (KernelWeight * PreExposureCorrection);
+				AccumulateHighFrequency[1] += SampleHighFrequency[1] * (KernelWeight * PreExposureCorrection);
+				AccumulateHighFrequency[2] += SampleHighFrequency[2] * (KernelWeight * PreExposureCorrection);
+				#if CONFIG_SCENE_COLOR_ALPHA
+					AccumulateHighFrequency[3] += SampleHighFrequency[3] * KernelWeight;
+				#endif
+				AccumulateMetadata  += dpv_scale(SampleMetadata, KernelWeight);
 			}
+		}
 
+		// Corrects history.
+		{
 			// Super bight highlights have a tendy to generate negatives
-			bool2 HasAnyNegative = or(AccumulateHighFrequency[0] < 0.0, or(AccumulateHighFrequency[1] < 0.0, AccumulateHighFrequency[2] < 0.0));
-
-			AccumulateHighFrequency[0] = select(HasAnyNegative, dpv_interleave_registers_array(RawHighFrequency[2])[0], AccumulateHighFrequency[0]);
-			AccumulateHighFrequency[1] = select(HasAnyNegative, dpv_interleave_registers_array(RawHighFrequency[2])[1], AccumulateHighFrequency[1]);
-			AccumulateHighFrequency[2] = select(HasAnyNegative, dpv_interleave_registers_array(RawHighFrequency[2])[2], AccumulateHighFrequency[2]);
-			#if CONFIG_SCENE_COLOR_ALPHA
-				AccumulateHighFrequency[3] = select(HasAnyNegative, dpv_interleave_registers_array(RawHighFrequency[2])[3], AccumulateHighFrequency[3]);
-			#endif
-	
+			bool2 bHasAnyNegativeOrNaN = or(AccumulateHighFrequency[0] < 0.0, or(AccumulateHighFrequency[1] < 0.0, AccumulateHighFrequency[2] < 0.0));
 			#if CONFIG_SCENE_COLOR_ALPHA
-				AccumulateHighFrequency = -dpv_min(-AccumulateHighFrequency, tsr_half(0.0).xxxx);
-			#else
-				AccumulateHighFrequency = -dpv_min(-AccumulateHighFrequency, tsr_half(0.0).xxx);
+				bHasAnyNegativeOrNaN = or(bHasAnyNegativeOrNaN, AccumulateHighFrequency[3] < 0.0);
 			#endif
-			AccumulateMetadata      = -dpv_min(-AccumulateMetadata, tsr_half(0.0).xx);
+
+			BRANCH
+			if (any(bHasAnyNegativeOrNaN))
+			{
+				AccumulateHighFrequency[0] = select(bHasAnyNegativeOrNaN, dpv_interleave_registers_array(RawHighFrequency[2])[0] * PreExposureCorrection, AccumulateHighFrequency[0]);
+				AccumulateHighFrequency[1] = select(bHasAnyNegativeOrNaN, dpv_interleave_registers_array(RawHighFrequency[2])[1] * PreExposureCorrection, AccumulateHighFrequency[1]);
+				AccumulateHighFrequency[2] = select(bHasAnyNegativeOrNaN, dpv_interleave_registers_array(RawHighFrequency[2])[2] * PreExposureCorrection, AccumulateHighFrequency[2]);
+				#if CONFIG_SCENE_COLOR_ALPHA
+					AccumulateHighFrequency[3] = select(bHasAnyNegativeOrNaN, dpv_interleave_registers_array(RawHighFrequency[2])[3], AccumulateHighFrequency[3]);
+				#endif
+			
+				#if CONFIG_SCENE_COLOR_ALPHA
+					AccumulateHighFrequency = -dpv_min(-AccumulateHighFrequency, tsr_half(0.0).xxxx);
+				#else
+					AccumulateHighFrequency = -dpv_min(-AccumulateHighFrequency, tsr_half(0.0).xxx);
+				#endif
+			}
+
+			AccumulateMetadata = -dpv_min(-AccumulateMetadata, tsr_half(0.0).xx);
 		}
 		
 		// Unpack history.
 		{
-			PrevHistoryMomentMin = HighFrequencyMin;
-			PrevHistoryMomentMax = HighFrequencyMax;
-
 			PrevHighFrequencyColor = AccumulateHighFrequency;
 			#if CONFIG_METADATA_CHANNELS == 1
 				PrevHistoryValidity = AccumulateMetadata;
@@ -794,31 +798,6 @@ void MainCS(
 				PrevHistoryValidity = AccumulateMetadata[0];
 			#endif
 		}
-
-		// Correct history
-		{
-			#if CONFIG_HISTORY_RESURRECTION
-				const tsr_half2 PreExposureCorrection = select(
-					bIsHistoryResurrection,
-					tsr_half(ResurrectionPreExposureCorrection),
-					tsr_half(HistoryPreExposureCorrection));
-			#else
-				const tsr_half2 PreExposureCorrection = tsr_half(HistoryPreExposureCorrection);
-			#endif
-
-			PrevHistoryMomentMin[0] *= PreExposureCorrection;
-			PrevHistoryMomentMin[1] *= PreExposureCorrection;
-			PrevHistoryMomentMin[2] *= PreExposureCorrection;
-			PrevHistoryMomentMax[0] *= PreExposureCorrection;
-			PrevHistoryMomentMax[1] *= PreExposureCorrection;
-			PrevHistoryMomentMax[2] *= PreExposureCorrection;
-			PrevHighFrequencyColor[0] *= PreExposureCorrection;
-			PrevHighFrequencyColor[1] *= PreExposureCorrection;
-			PrevHighFrequencyColor[2] *= PreExposureCorrection;
-
-			// Fixes negative validity that can happen due to negative lobes of Catrom.
-			PrevHistoryValidity = max(PrevHistoryValidity, tsr_half(0.0).xx);
-		}
 	}
 	
 	// TODO: CONFIG_MANUAL_LDS_SPILL
