Composes PostDOF translucency solely in TSR's reject shading to save perf in the UpdatePass

diff --git a/Engine/Shaders/Private/TemporalSuperResolution/TSRRejectShading.usf b/Engine/Shaders/Private/TemporalSuperResolution/TSRRejectShading.usf
--- a/Engine/Shaders/Private/TemporalSuperResolution/TSRRejectShading.usf
+++ b/Engine/Shaders/Private/TemporalSuperResolution/TSRRejectShading.usf
@@ -125,6 +125,7 @@ Texture2D<tsr_ushort> IsMovingMaskTexture;
 RWTexture2D<tsr_half4> HistoryGuideOutput;
 RWTexture2D<tsr_half4> HistoryMoireOutput;
 RWTexture2D<tsr_half2> HistoryRejectionOutput;
+RWTexture2D<tsr_halfC> InputSceneColorOutput;
 RWTexture2D<tsr_half> InputSceneColorLdrLumaOutput;
 
 
@@ -146,6 +147,22 @@ tsr_simd_halfC ClampFireFliersWithGuide(tsr_simd_halfC ToClamp, tsr_simd_halfC G
 	return fastClamp(ToClamp, ClampedGuideMin, ClampedGuideMax);
 }
 
+tsr_short2 ComputeInputPixelPos(const tsr_short2 GroupPixelOffset, const uint LaneIndex, const uint SimdIndex)
+{
+	return GroupPixelOffset + GetLaneSimdPixelOffset<LANE_STRIDE_X, LANE_STRIDE_Y>(LaneIndex, SimdIndex);
+}
+
+tsr_short2 ComputeOutputPixelPos(const tsr_short2 GroupPixelOffset, const uint LaneIndex, const uint SimdIndex)
+{
+	tsr_short2 LaneSimdPixelOffset = GetLaneSimdPixelOffset<LANE_STRIDE_X, LANE_STRIDE_Y>(LaneIndex, SimdIndex);
+	bool bIsNonOverscanPixel = all(and(LaneSimdPixelOffset >= TileOverscan, LaneSimdPixelOffset < (GROUP_TILE_SIZE - TileOverscan)));
+
+	tsr_short2 OutputPixelPos = InvalidateOutputPixelPos(GroupPixelOffset + LaneSimdPixelOffset, InputInfo_ViewportMax);
+	OutputPixelPos.x = select(bIsNonOverscanPixel, OutputPixelPos.x, ~tsr_short(0));
+	
+	return OutputPixelPos;
+}
+
 
 //------------------------------------------------------- ENTRY POINT
 
@@ -178,7 +195,7 @@ void MainCS(
 		tsr_simd_half4 OriginalTranslucencyInput;
 	#endif
 	tsr_simd_halfC OriginalHistory;
-	tsr_simd_half bIsValid;
+	tsr_simd_bool bIsValid;
 	#if CONFIG_MOIRE_DETECTION
 		tsr_simd_half MoireInput;
 		tsr_simd_half4 PrevMoire;
@@ -189,7 +206,7 @@ void MainCS(
 		UNROLL_N(SIMD_SIZE)
 		for (uint SimdIndex = 0; SimdIndex < SIMD_SIZE; SimdIndex++)
 		{
-			tsr_short2 ClampedFetchPixelPos = GroupPixelOffset + GetLaneSimdPixelOffset<LANE_STRIDE_X, LANE_STRIDE_Y>(LaneIndex, SimdIndex);
+			tsr_short2 ClampedFetchPixelPos = ComputeInputPixelPos(GroupPixelOffset, LaneIndex, SimdIndex);
 			float2 TranslucencyTextureUV = ApplyScreenTransform(float2(ClampedFetchPixelPos), InputPixelPosToTranslucencyTextureUV);
 
 			ClampedFetchPixelPos = ClampPixelOffset(ClampedFetchPixelPos, InputPixelPosMin, InputPixelPosMax);
@@ -200,7 +217,7 @@ void MainCS(
 				OriginalTranslucencyInput.SetElement(SimdIndex, InputSceneTranslucencyTexture.SampleLevel(GlobalBilinearClampedSampler, TranslucencyTextureUV, 0));
 			#endif
 			OriginalHistory.SetElement(SimdIndex, ReprojectedHistoryGuideTexture[ClampedFetchPixelPos]);
-			bIsValid.SetElement(SimdIndex, ParallaxRejectionMaskTexture[ClampedFetchPixelPos] > tsr_half(PARALLAX_REJECTION_MASK_THRESHOLD) ? tsr_half(1.0) : tsr_half(0.0));
+			bIsValid.SetElement(SimdIndex, ParallaxRejectionMaskTexture[ClampedFetchPixelPos] > tsr_half(PARALLAX_REJECTION_MASK_THRESHOLD));
 			#if CONFIG_MOIRE_DETECTION
 				MoireInput.SetElement(SimdIndex, InputMoireLumaTexture[ClampedFetchPixelPos]);
 				PrevMoire.SetElement(SimdIndex, ReprojectedHistoryMoireTexture[ClampedFetchPixelPos]);
@@ -210,22 +227,43 @@ void MainCS(
 	}
 	
 	// Compute AA Luma
-	tsr_simd_half CenterLumaLDR;
 	{
 		const tsr_half SpatialAAExposure = tsr_half(0.5);
+		
+		tsr_simd_halfC CenterFinalSceneColor;
+		tsr_simd_half CenterLumaLDR;
+
 		UNROLL_N(SIMD_SIZE)
 		for (uint SimdIndex = 0; SimdIndex < SIMD_SIZE; SimdIndex++)
 		{
-			tsr_half3 Color = OriginalOpaqueInput.GetElement(SimdIndex).rgb;
+			tsr_halfC Color = OriginalOpaqueInput.GetElement(SimdIndex);
 			
 			#if CONFIG_COMPOSE_TRANSLUCENCY
 				tsr_half4 Translucency = OriginalTranslucencyInput.GetElement(SimdIndex);
 				Color.rgb = Color.rgb * Translucency.a + Translucency.rgb;
+				#if CONFIG_SCENE_COLOR_ALPHA
+					Color.a = Color.a * Translucency.a;
+				#endif
 			#endif
+			
+			CenterFinalSceneColor.SetElement(SimdIndex, Color);
 
 			tsr_half CenterLuma = dot(Color, tsr_half3(0.299f, 0.587f, 0.114f));
 			CenterLumaLDR.SetElement(SimdIndex, CenterLuma / (SpatialAAExposure + CenterLuma));
 		}
+
+		// Output some data early.
+		PLATFORM_SPECIFIC_ISOLATE
+		{
+			UNROLL_N(SIMD_SIZE)
+			for (uint SimdIndex = 0; SimdIndex < SIMD_SIZE; SimdIndex++)
+			{
+				const tsr_short2 OutputPixelPos = ComputeOutputPixelPos(GroupPixelOffset, LaneIndex, SimdIndex);
+		
+				InputSceneColorOutput[OutputPixelPos] = CenterFinalSceneColor.GetElement(SimdIndex);
+				InputSceneColorLdrLumaOutput[OutputPixelPos] = CenterLumaLDR.GetElement(SimdIndex);
+			}
+		}
 	}
 	
 	tsr_simd_halfC OriginalInput = OriginalOpaqueInput;
@@ -368,7 +406,7 @@ void MainCS(
 			tsr_half LocalPrevGradient = PrevGradient.GetElement(SimdIndex);
 			tsr_half LocalPrevTotalVariation = PrevTotalVariation.GetElement(SimdIndex);
 			tsr_half LocalPrevTotalVariationCount = PrevTotalVariationCount.GetElement(SimdIndex);
-			tsr_half LocalIsValid = bIsValid.GetElement(SimdIndex);
+			bool LocalIsValid = bIsValid.GetElement(SimdIndex);
 
 			tsr_half LocalIsFlicker = IsFlicker.GetElement(SimdIndex);
 
@@ -390,9 +428,9 @@ void MainCS(
 			tsr_half LocalNewTotalVariationCount = LocalPrevTotalVariationCount + LocalTotalVariationCountContribution;
 
 			// Discard all moire history on parallax disocclusion.
-			LocalNewGradient *= LocalIsValid;
-			LocalNewTotalVariation *= LocalIsValid;
-			LocalNewTotalVariationCount *= LocalIsValid;
+			LocalNewGradient = select(LocalIsValid, LocalNewGradient, tsr_half(0.0));
+			LocalNewTotalVariation = select(LocalIsValid, LocalNewTotalVariation, tsr_half(0.0));
+			LocalNewTotalVariationCount = select(LocalIsValid, LocalNewTotalVariationCount, tsr_half(0.0));
 
 			// Quantise total variation and count to history bit depth to variation drift in the history over time.
 			tsr_half LocalQuantizedNewTotalVariationCount = floor(LocalNewTotalVariationCount * (tsr_half(255.0) / MaxPrevTotalVariationCount)) * (MaxPrevTotalVariationCount / tsr_half(255.0));
@@ -420,6 +458,19 @@ void MainCS(
 	}
 	#endif
 
+	// Output some data early.
+	PLATFORM_SPECIFIC_ISOLATE
+	{
+		UNROLL_N(SIMD_SIZE)
+		for (uint SimdIndex = 0; SimdIndex < SIMD_SIZE; SimdIndex++)
+		{
+			const tsr_short2 OutputPixelPos = ComputeOutputPixelPos(GroupPixelOffset, LaneIndex, SimdIndex);
+			#if CONFIG_MOIRE_DETECTION
+				HistoryMoireOutput[OutputPixelPos] = FinalMoire.GetElement(SimdIndex);
+			#endif
+		}
+	}
+
 	// Measure rejection
 	tsr_simd_half RejectionBlendFinal = tsr_simd_half::Const(1.0);
 	tsr_simd_half RejectionClampBlend = tsr_simd_half::Const(1.0);
@@ -549,7 +600,7 @@ void MainCS(
 		for (uint SimdIndex = 0; SimdIndex < SIMD_SIZE; SimdIndex++)
 		{
 			bool bIsOffScreen = false;
-			bool bIsDisoccluded = bIsValid.GetElement(SimdIndex) == tsr_half(0.0);
+			bool bIsDisoccluded = !bIsValid.GetElement(SimdIndex);
 			tsr_half LocalRejectionBlendFinal = RejectionBlendFinal.GetElement(SimdIndex);
 
 			//tsr_half TotalBlend = select(or(bIsOffScreen, bIsDisoccluded), tsr_half(1.0), tsr_half(0.0));
@@ -580,11 +631,8 @@ void MainCS(
 		UNROLL_N(SIMD_SIZE)
 		for (uint SimdIndex = 0; SimdIndex < SIMD_SIZE; SimdIndex++)
 		{
-			tsr_short2 LaneSimdPixelOffset = GetLaneSimdPixelOffset<LANE_STRIDE_X, LANE_STRIDE_Y>(LaneIndex, SimdIndex);
-			bool bIsNonOverscanPixel = all(and(LaneSimdPixelOffset >= TileOverscan, LaneSimdPixelOffset < (GROUP_TILE_SIZE - TileOverscan)));
-
-			tsr_short2 OutputPixelPos = InvalidateOutputPixelPos(GroupPixelOffset + LaneSimdPixelOffset, InputInfo_ViewportMax);
-			OutputPixelPos.x = select(bIsNonOverscanPixel, OutputPixelPos.x, ~tsr_short(0));
+			const tsr_short2 LaneSimdPixelOffset = GetLaneSimdPixelOffset<LANE_STRIDE_X, LANE_STRIDE_Y>(LaneIndex, SimdIndex);
+			const tsr_short2 OutputPixelPos = ComputeOutputPixelPos(GroupPixelOffset, LaneIndex, SimdIndex);
 		
 			tsr_half LocalRejectionBlendFinal = RejectionBlendFinal.GetElement(SimdIndex);
 			tsr_half LocalRejectionClampBlend = RejectionClampBlend.GetElement(SimdIndex);
@@ -606,12 +654,7 @@ void MainCS(
 			#else
 				HistoryGuideOutput[OutputPixelPos] = tsr_half4(FinalGuideColor.xyz, tsr_half(1.0));
 			#endif
-
-			#if CONFIG_MOIRE_DETECTION
-				HistoryMoireOutput[OutputPixelPos] = FinalMoire.GetElement(SimdIndex);
-			#endif
 			HistoryRejectionOutput[OutputPixelPos] = tsr_half2(LocalRejectionBlendFinal, LocalRejectionClampBlend);
-			InputSceneColorLdrLumaOutput[OutputPixelPos] = CenterLumaLDR.GetElement(SimdIndex);
 			
 			#if DEBUG_OUTPUT
 			UNROLL_N(DEBUG_ARRAY_SIZE)
diff --git a/Engine/Shaders/Private/TemporalSuperResolution/TSRUpdateHistory.usf b/Engine/Shaders/Private/TemporalSuperResolution/TSRUpdateHistory.usf
--- a/Engine/Shaders/Private/TemporalSuperResolution/TSRUpdateHistory.usf
+++ b/Engine/Shaders/Private/TemporalSuperResolution/TSRUpdateHistory.usf
@@ -31,7 +31,7 @@
 	#error Unknown history update quality
 #endif
 
-#define CONFIG_COMPOSE_TRANSLUCENCY 1
+#define CONFIG_COMPOSE_TRANSLUCENCY 0
 #define CONFIG_REJECT_TRANSLUCENCY_ON_RESPONSIVE_AA_ONLY 0
 #define CONFIG_HISTORY_RESURRECTION 1
 
diff --git a/Engine/Source/Runtime/Renderer/Private/PostProcess/TemporalSuperResolution.cpp b/Engine/Source/Runtime/Renderer/Private/PostProcess/TemporalSuperResolution.cpp
--- a/Engine/Source/Runtime/Renderer/Private/PostProcess/TemporalSuperResolution.cpp
+++ b/Engine/Source/Runtime/Renderer/Private/PostProcess/TemporalSuperResolution.cpp
@@ -148,6 +148,11 @@ TAutoConsoleVariable<int32> CVarTSRShadingTileOverscan(
 	TEXT(" anti-aliasing scalability settings."),
 	ECVF_Scalability | ECVF_RenderThreadSafe);
 
+TAutoConsoleVariable<int32> CVarTSRShadingTileSize(
+	TEXT("r.TSR.ShadingRejection.TileSize"), 16,
+	TEXT(""),
+	ECVF_Scalability | ECVF_RenderThreadSafe);
+
 TAutoConsoleVariable<float> CVarTSRShadingExposureOffset(
 	TEXT("r.TSR.ShadingRejection.ExposureOffset"), 3.0,
 	TEXT("The shading rejection needs to have a representative idea how bright a linear color pixel ends up displayed to the user. ")
@@ -615,6 +620,7 @@ class FTSRRejectShadingCS : public FTSRShader
 		SHADER_PARAMETER_RDG_TEXTURE_UAV(RWTexture2D, HistoryGuideOutput)
 		SHADER_PARAMETER_RDG_TEXTURE_UAV(RWTexture2D, HistoryMoireOutput)
 		SHADER_PARAMETER_RDG_TEXTURE_UAV(RWTexture2D, HistoryRejectionOutput)
+		SHADER_PARAMETER_RDG_TEXTURE_UAV(RWTexture2D, InputSceneColorOutput)
 		SHADER_PARAMETER_RDG_TEXTURE_UAV(RWTexture2D, InputSceneColorLdrLumaOutput)
 		SHADER_PARAMETER_RDG_TEXTURE_UAV(RWTexture2D, DebugOutput)
 	END_SHADER_PARAMETER_STRUCT()
@@ -729,7 +735,6 @@ class FTSRSpatialAntiAliasingCS : public FTSRShader
 
 	BEGIN_SHADER_PARAMETER_STRUCT(FParameters, )
 		SHADER_PARAMETER_STRUCT_INCLUDE(FTSRCommonParameters, CommonParameters)
-		SHADER_PARAMETER_RDG_TEXTURE(Texture2D, InputSceneColorTexture)
 		SHADER_PARAMETER_RDG_TEXTURE(Texture2D, InputSceneColorLdrLumaTexture)
 		SHADER_PARAMETER_RDG_TEXTURE_UAV(RWTexture2D, AntiAliasingOutput)
 		SHADER_PARAMETER_RDG_TEXTURE_UAV(RWTexture2D, NoiseFilteringOutput)
@@ -1907,6 +1912,13 @@ FDefaultTemporalUpscaler::FOutputs AddTemporalSuperResolutionPasses(
 		DilatedVelocityTexture = HoleFilledVelocityTexture;
 	}
 
+	// Merge PostDOF translucency within same scene color.
+	FRDGTextureRef InputSceneColorTexture = nullptr;
+	if (!bHasSeparateTranslucency)
+	{
+		InputSceneColorTexture = PassInputs.SceneColor.Texture;
+	}
+
 	// Perform a history reject the history.
 	FRDGTextureRef InputSceneColorLdrLumaTexture = nullptr;
 	auto RejectReprojectedGuide = [&](
@@ -1915,10 +1927,22 @@ FDefaultTemporalUpscaler::FOutputs AddTemporalSuperResolutionPasses(
 		FRDGTextureUAVRef HistoryGuideOutput,
 		FRDGTextureUAVRef HistoryMoireOutput) -> FRDGTextureRef
 	{
-		const int32 GroupTileSize = 16;
+		const int32 GroupTileSize = CVarTSRShadingTileSize.GetValueOnRenderThread();
+
+		const bool bComputeInputSceneColorTexture = InputSceneColorLdrLumaTexture == nullptr;
+		if (bComputeInputSceneColorTexture)
+		{
+			FRDGTextureDesc Desc = FRDGTextureDesc::Create2D(
+				InputExtent,
+				HistoryColorFormat,
+				FClearValueBinding::None,
+				/* InFlags = */ TexCreate_ShaderResource | TexCreate_UAV);
+
+			InputSceneColorTexture = GraphBuilder.CreateTexture(Desc, TEXT("TSR.SceneColor"));
+		}
 
-		const bool ComputeLdrLuma = InputSceneColorLdrLumaTexture == nullptr;
-		if (ComputeLdrLuma)
+		const bool bComputeLdrLuma = InputSceneColorLdrLumaTexture == nullptr && RejectionAntiAliasingQuality > 0;
+		if (bComputeLdrLuma)
 		{
 			FRDGTextureDesc Desc = FRDGTextureDesc::Create2D(
 				InputExtent,
@@ -1977,10 +2001,12 @@ FDefaultTemporalUpscaler::FOutputs AddTemporalSuperResolutionPasses(
 		PassParameters->HistoryGuideOutput = HistoryGuideOutput;
 		PassParameters->HistoryMoireOutput = HistoryMoireOutput;
 		PassParameters->HistoryRejectionOutput = GraphBuilder.CreateUAV(HistoryRejectionTexture);
-		//PassParameters->InputSceneColorLdrLumaOutput = ComputeLdrLuma
-		//	? GraphBuilder.CreateUAV(InputSceneColorLdrLumaTexture)
-		//	: CreateDummyUAV(GraphBuilder, PF_R8);
-		PassParameters->InputSceneColorLdrLumaOutput = GraphBuilder.CreateUAV(InputSceneColorLdrLumaTexture);
+		PassParameters->InputSceneColorOutput = bComputeInputSceneColorTexture
+			? GraphBuilder.CreateUAV(InputSceneColorTexture)
+			: CreateDummyUAV(GraphBuilder, HistoryColorFormat);
+		PassParameters->InputSceneColorLdrLumaOutput = bComputeLdrLuma
+			? GraphBuilder.CreateUAV(InputSceneColorLdrLumaTexture)
+			: CreateDummyUAV(GraphBuilder, PF_R8);
 		PassParameters->DebugOutput = CreateDebugUAV(InputExtent, TEXT("Debug.TSR.RejectShading"));
 
 		int32 TileSize = GroupTileSize - 2 * PassParameters->TileOverscan;
@@ -2142,7 +2168,6 @@ FDefaultTemporalUpscaler::FOutputs AddTemporalSuperResolutionPasses(
 		{
 			FTSRSpatialAntiAliasingCS::FParameters* PassParameters = GraphBuilder.AllocParameters<FTSRSpatialAntiAliasingCS::FParameters>();
 			PassParameters->CommonParameters = CommonParameters;
-			PassParameters->InputSceneColorTexture = PassInputs.SceneColor.Texture;
 			PassParameters->InputSceneColorLdrLumaTexture = InputSceneColorLdrLumaTexture;
 			PassParameters->AntiAliasingOutput = GraphBuilder.CreateUAV(RawAntiAliasingTexture);
 			PassParameters->NoiseFilteringOutput = GraphBuilder.CreateUAV(NoiseFilteringTexture);
@@ -2194,7 +2219,7 @@ FDefaultTemporalUpscaler::FOutputs AddTemporalSuperResolutionPasses(
 
 		FTSRUpdateHistoryCS::FParameters* PassParameters = GraphBuilder.AllocParameters<FTSRUpdateHistoryCS::FParameters>();
 		PassParameters->CommonParameters = CommonParameters;
-		PassParameters->InputSceneColorTexture = PassInputs.SceneColor.Texture;
+		PassParameters->InputSceneColorTexture = InputSceneColorTexture;
 		PassParameters->InputSceneStencilTexture = GraphBuilder.CreateSRV(
 			FRDGTextureSRVDesc::CreateWithPixelFormat(PassInputs.SceneDepth.Texture, PF_X24_G8));
 		PassParameters->InputSceneTranslucencyTexture = SeparateTranslucencyTexture;
