Path Tracing: Fix anisotropy encoding in Substrate payloads

The payload encoding for anisotropy did not allow representing 0.0 anisotropy exactly. This in turn could lead to NaNs when the world tangent aligned with the shading normal.

diff --git a/Engine/Shaders/Private/PathTracing/PathTracingCommon.ush b/Engine/Shaders/Private/PathTracing/PathTracingCommon.ush
--- a/Engine/Shaders/Private/PathTracing/PathTracingCommon.ush
+++ b/Engine/Shaders/Private/PathTracing/PathTracingCommon.ush
@@ -93,6 +93,24 @@ float3 PayloadDecodeLDRColor(uint rgb)
 	return float3(r, g, b);
 }
 
+uint PayloadEncodeRoughnessAniso(float3 RoughnessData, float Anisotropy)
+{
+	uint r = uint(floor(saturate(RoughnessData.x) * 255.0 + 0.5));
+	uint g = uint(floor(saturate(RoughnessData.y) * 255.0 + 0.5));
+	uint b = uint(floor(saturate(RoughnessData.z) * 255.0 + 0.5));
+	uint a = uint(floor(clamp(Anisotropy, -1.0, 1.0) * 127.0 + 127.0)); // Range is [-1,1] so use an SNorm encoding here so we can represent 0.0 exactly
+	return (r << 24) | (g << 16) | (b << 8) | a;
+}
+
+float4 PayloadDecodeRoughnessAniso(uint rgba)
+{
+	float r = float((rgba >> 24)      ) * (1.0 / 255.0);
+	float g = float((rgba >> 16) & 255) * (1.0 / 255.0);
+	float b = float((rgba >>  8) & 255) * (1.0 / 255.0);
+	float a = float((rgba      ) & 255) * (1.0 / 127.0) - 1.0;
+	return float4(r, g, b, a);
+}
+
 float3 AdjustShadingNormal(float3 ShadingNormal, float3 GeoNormal, float3 RayDirection)
 {
 	// Clip shading normal in a view dependent way such that the reflection stays above the geometric normal
@@ -610,7 +628,7 @@ FPackedPathTracingPayload PackPathTracingPayload(FPathTracingPayload Input)
 	Output.PackedData[5] = PayloadEncodeLDRColor(Input.DiffuseColor);
 	Output.PackedData[6] = PayloadEncodeLDRColor(Input.SpecularColor);
 	Output.PackedData[7] = PayloadEncodeLDRColor(Input.SpecularEdgeColor);
-	Output.PackedData[8] = PackRGBA8(float4(Input.RoughnessData, 0.5 * Input.Anisotropy + 0.5)); // not a color, but encoding is still ok. NOTE: Aniso is [-1,1]
+	Output.PackedData[8] = PayloadEncodeRoughnessAniso(Input.RoughnessData, Input.Anisotropy);
 	Output.PackedData[9] = PayloadEncodeUnitVector(Input.WorldTangent);
 	Output.PackedData[10] = PayloadEncodeLDRColor(Input.TransparencyColor);
 	Output.PackedData[11] = f32tof16(Input.MeanFreePath.x);
@@ -694,8 +712,9 @@ FPathTracingPayload UnpackPathTracingPayload(FPackedPathTracingPayload Input, FR
 	Output.DiffuseColor        = PayloadDecodeLDRColor(Input.PackedData[5]);
 	Output.SpecularColor       = PayloadDecodeLDRColor(Input.PackedData[6]);
 	Output.SpecularEdgeColor   = PayloadDecodeLDRColor(Input.PackedData[7]);
-	Output.RoughnessData       = UnpackRGBA8(Input.PackedData[8]).xyz;
-	Output.Anisotropy          = 2 * UnpackRGBA8(Input.PackedData[8]).w - 1;
+	float4 RoughnessAniso      = PayloadDecodeRoughnessAniso(Input.PackedData[8]); 
+	Output.RoughnessData       = RoughnessAniso.xyz;
+	Output.Anisotropy          = RoughnessAniso.w;
 	Output.WorldTangent        = PayloadDecodeUnitVector(Input.PackedData[9]);
 	Output.TransparencyColor   = PayloadDecodeLDRColor(Input.PackedData[10]);
 	Output.MeanFreePath.x      = f16tof32(Input.PackedData[11]);
