Fixes shader compilation failure with r.PostProcessing.PropagateAlpha=1 in TSR

diff --git a/Engine/Shaders/Private/TemporalSuperResolution/TSRRejectShading.usf b/Engine/Shaders/Private/TemporalSuperResolution/TSRRejectShading.usf
--- a/Engine/Shaders/Private/TemporalSuperResolution/TSRRejectShading.usf
+++ b/Engine/Shaders/Private/TemporalSuperResolution/TSRRejectShading.usf
@@ -248,7 +248,7 @@ void MainCS(
 			
 			CenterFinalSceneColor.SetElement(SimdIndex, Color);
 
-			tsr_half CenterLuma = dot(Color, tsr_half3(0.299f, 0.587f, 0.114f));
+			tsr_half CenterLuma = dot(Color.rgb, tsr_half3(0.299f, 0.587f, 0.114f));
 			CenterLumaLDR.SetElement(SimdIndex, CenterLuma / (SpatialAAExposure + CenterLuma));
 		}
 
