Remove samplers from LocalLightPrepass textures in MobileBasePass

diff --git a/Engine/Shaders/Private/MobileBasePassPixelShader.usf b/Engine/Shaders/Private/MobileBasePassPixelShader.usf
--- a/Engine/Shaders/Private/MobileBasePassPixelShader.usf
+++ b/Engine/Shaders/Private/MobileBasePassPixelShader.usf
@@ -819,9 +819,9 @@ void Main(
 #if  PREPASS_LOCAL_LIGHTS_MOBILE
 	{
 		FDeferredLightData MergedLightData = (FDeferredLightData)0;
-		half3 LocalLightA = Texture2DSample(MobileSceneTextures.LocalLightTextureA, MobileSceneTextures.LocalLightTextureASampler,SvPositionToBufferUV(SvPosition)); 
-		half4 LocalLightB = Texture2DSample(MobileSceneTextures.LocalLightTextureB, MobileSceneTextures.LocalLightTextureBSampler,SvPositionToBufferUV(SvPosition));
-	
+		half3 LocalLightA = MobileSceneTextures.LocalLightTextureA.Load(int3(SvPosition.xy, 0)).xyz; 
+		half4 LocalLightB = MobileSceneTextures.LocalLightTextureB.Load(int3(SvPosition.xy, 0));
+
 		MergedLightData.SourceRadius	= 0;
 		MergedLightData.SourceLength	= 0;
 		MergedLightData.Tangent			= 0;
diff --git a/Engine/Source/Runtime/Engine/Public/SceneTexturesConfig.h b/Engine/Source/Runtime/Engine/Public/SceneTexturesConfig.h
--- a/Engine/Source/Runtime/Engine/Public/SceneTexturesConfig.h
+++ b/Engine/Source/Runtime/Engine/Public/SceneTexturesConfig.h
@@ -62,8 +62,6 @@ BEGIN_GLOBAL_SHADER_PARAMETER_STRUCT(FMobileSceneTextureUniformParameters, ENGIN
 	SHADER_PARAMETER_SAMPLER(SamplerState, GBufferCTextureSampler)
 	SHADER_PARAMETER_SAMPLER(SamplerState, GBufferDTextureSampler)
 	SHADER_PARAMETER_SAMPLER(SamplerState, SceneDepthAuxTextureSampler)
-	SHADER_PARAMETER_SAMPLER(SamplerState, LocalLightTextureASampler)
-	SHADER_PARAMETER_SAMPLER(SamplerState, LocalLightTextureBSampler)
 END_GLOBAL_SHADER_PARAMETER_STRUCT()
 
 BEGIN_SHADER_PARAMETER_STRUCT(FSceneTextureShaderParameters, ENGINE_API)
diff --git a/Engine/Source/Runtime/Renderer/Private/SceneTextures.cpp b/Engine/Source/Runtime/Renderer/Private/SceneTextures.cpp
--- a/Engine/Source/Runtime/Renderer/Private/SceneTextures.cpp
+++ b/Engine/Source/Runtime/Renderer/Private/SceneTextures.cpp
@@ -965,8 +965,6 @@ void SetupMobileSceneTextureUniformParameters(
 	SceneTextureParameters.GBufferCTextureSampler = TStaticSamplerState<>::GetRHI();
 	SceneTextureParameters.GBufferDTextureSampler = TStaticSamplerState<>::GetRHI();
 	SceneTextureParameters.SceneDepthAuxTextureSampler = TStaticSamplerState<>::GetRHI();
-	SceneTextureParameters.LocalLightTextureASampler = TStaticSamplerState<>::GetRHI();
-	SceneTextureParameters.LocalLightTextureBSampler = TStaticSamplerState<>::GetRHI();
 
 	if (SceneTextures)
 	{
