Make VSM cache more persistent and robust against multiple renders in a single "frame" (scene captures, etc).
Remove multiple cache managesr that were implemented to work around issues with the previous naive implementation
Fixes some invalidation issues in cases where multiple managers existed previously

NOTE: This change brings pages from scene captures and ndisplay and similar systems back into the main physical page pool, so it may be necessary in some of these cases to increase the page pool size.

diff --git a/Engine/Plugins/Runtime/nDisplay/Source/DisplayCluster/Private/Game/EngineClasses/Basics/DisplayClusterViewportClient.cpp b/Engine/Plugins/Runtime/nDisplay/Source/DisplayCluster/Private/Game/EngineClasses/Basics/DisplayClusterViewportClient.cpp
--- a/Engine/Plugins/Runtime/nDisplay/Source/DisplayCluster/Private/Game/EngineClasses/Basics/DisplayClusterViewportClient.cpp
+++ b/Engine/Plugins/Runtime/nDisplay/Source/DisplayCluster/Private/Game/EngineClasses/Basics/DisplayClusterViewportClient.cpp
@@ -618,9 +618,6 @@ void UDisplayClusterViewportClient::Draw(FViewport* InViewport, FCanvas* SceneCa
 
 					View->CameraConstrainedViewRect = View->UnscaledViewRect;
 
-					// Enable per-view virtual shadow map caching
-					View->State->AddVirtualShadowMapCache(MyWorld->Scene);
-
 					// Enable per-view Lumen scene
 					if (GDisplayClusterLumenPerView)
 					{
diff --git a/Engine/Shaders/Private/VirtualShadowMaps/VirtualShadowMapPageAccessCommon.ush b/Engine/Shaders/Private/VirtualShadowMaps/VirtualShadowMapPageAccessCommon.ush
--- a/Engine/Shaders/Private/VirtualShadowMaps/VirtualShadowMapPageAccessCommon.ush
+++ b/Engine/Shaders/Private/VirtualShadowMaps/VirtualShadowMapPageAccessCommon.ush
@@ -24,8 +24,8 @@
 // Stored in the physical meta data flags and must be higher than the  VSM_*_FLAG
 // The light was marked to not be cached (even though caching is enabled)
 #define VSM_PHYSICAL_FLAG_VIEW_UNCACHED		(1U << 4) 
-#define VSM_PHYSICAL_FLAG_DIRTY				(1U << 5) 
-#define VSM_PHYSICAL_FLAG_USED_THIS_FRAME	(1U << 6) 
+#define VSM_PHYSICAL_FLAG_DIRTY				(1U << 5)
+#define VSM_PHYSICAL_FLAG_UNREFERENCED		(1U << 6)
 
 
 // bit-offset for storing another set of page flags used to track invalidation in the physical page metadata flags.
@@ -48,9 +48,8 @@ bool IsSinglePageVirtualShadowMap(int VirtualShadowMapId)
 struct FPhysicalPageMetaData
 {
 	uint Flags;	// VSM_*_FLAG, not all relevant to physical pages + VSM_PHYSICAL_FLAG_*
-	uint Age;
+	uint LastRequestedSceneFrameNumber;
 	// Link back to the virtual shadow map & page table slot that references this physical page
-	// Useful to track data on a virtual address space frequency & vice versa
 	uint VirtualShadowMapId;
 	uint MipLevel;
 	uint2 PageAddress;
diff --git a/Engine/Shaders/Private/VirtualShadowMaps/VirtualShadowMapPageMarking.usf b/Engine/Shaders/Private/VirtualShadowMaps/VirtualShadowMapPageMarking.usf
--- a/Engine/Shaders/Private/VirtualShadowMaps/VirtualShadowMapPageMarking.usf
+++ b/Engine/Shaders/Private/VirtualShadowMaps/VirtualShadowMapPageMarking.usf
@@ -501,12 +501,16 @@ void MarkCoarsePages(uint DispatchThreadId : SV_DispatchThreadID)
 
 	FVirtualShadowMapProjectionShaderData ProjectionData = GetVirtualShadowMapProjectionData(VirtualShadowMapId);
 	checkSlow(ProjectionData.bCurrentDistantLight == IsSinglePageVirtualShadowMap(VirtualShadowMapId));
-
+		
 	// NOTE: Coarse pages are very large and tend to get invalidated a lot due to anything in the scene moving
 	// Rendering non-nanite geometry into these pages can be very expensive and thus isn't always desirable.
 	uint Flags = VSM_ALLOCATED_FLAG;
-
-	if (ProjectionData.LightType == LIGHT_TYPE_DIRECTIONAL)
+				
+	if (ProjectionData.bUnreferenced)
+	{
+		// Don't mark any pages for lights not referenced this render
+	}
+	else if (ProjectionData.LightType == LIGHT_TYPE_DIRECTIONAL)
 	{
 		// Idea here is the clipmaps already cover supersets of lower levels
 		// Thus to get coarser pages we can just mark the center page(s) offset by a level/LOD bias
diff --git a/Engine/Shaders/Private/VirtualShadowMaps/VirtualShadowMapPhysicalPageManagement.usf b/Engine/Shaders/Private/VirtualShadowMaps/VirtualShadowMapPhysicalPageManagement.usf
--- a/Engine/Shaders/Private/VirtualShadowMaps/VirtualShadowMapPhysicalPageManagement.usf
+++ b/Engine/Shaders/Private/VirtualShadowMaps/VirtualShadowMapPhysicalPageManagement.usf
@@ -8,6 +8,7 @@
 #include "../WaveOpUtil.ush"
 #include "../ReductionCommon.ush"
 #include "../GPUMessaging.ush"
+#include "../ShaderPrint.ush"
 #include "VirtualShadowMapProjectionStructs.ush"
 #include "VirtualShadowMapProjectionCommon.ush"
 #include "VirtualShadowMapPageAccessCommon.ush"
@@ -51,6 +52,7 @@ int GetPhysicalPageListItem(uint PageList, int Index)
 
 void SetPhysicalPageListItem(uint PageList, int Index, int Value)
 {
+	//check(Index < VirtualShadowMap.MaxPhysicalPages);
 	OutPhysicalPageLists[GetPhysicalPageListStart(PageList) + Index] = Value;
 }
 
@@ -100,20 +102,8 @@ int PopPhysicalPageList(uint PageList)
 	return Offset < 0 ? INDEX_NONE : OutPhysicalPageLists[PageListStart + Offset];
 }
 
-// Parameters for mapping to previous cached data
-#if HAS_CACHE_DATA
-
-// Mapping of previous frame/update data to current frame
-StructuredBuffer<FNextVirtualShadowMapData> NextVirtualShadowMapData;
-uint NextVirtualShadowMapDataCount;
-
 StructuredBuffer<int> PrevPhysicalPageLists;
 
-int bDynamicPageInvalidation;
-int bKeepOnlyRequestedPages;
-int bAllocateViaLRU;
-#endif
-
 RWStructuredBuffer<uint4> OutPageRectBounds;
 
 // This is admitadly a weird fusion of several initializations but it is the first thing
@@ -139,6 +129,68 @@ void InitPageRectBounds(uint3 Index : SV_DispatchThreadID)
 	}
 }
 
+
+// Mapping of previous frame/update data to current frame
+StructuredBuffer<FNextVirtualShadowMapData> NextVirtualShadowMapData;
+uint NextVirtualShadowMapDataCount;
+
+[numthreads(VSM_DEFAULT_CS_GROUP_X, 1, 1)]
+void UpdatePhysicalPageAddresses(uint3 Index : SV_DispatchThreadID)
+{	
+	// TODO: Make this a loose constant probably to remove dependency of this shader on VSM UB
+	// Still needs the VSM defines for IsVirtualShadowMapPageAddressValid addressing math though!
+	if (Index.x >= VirtualShadowMap.MaxPhysicalPages)
+	{
+		return;
+	}
+
+	// Use identity mapping by default
+	int PhysicalPageIndex = Index.x;
+	checkStructuredBufferAccessSlow(OutPhysicalPageMetaData, PhysicalPageIndex);
+	
+	FPhysicalPageMetaData PrevMetaData = OutPhysicalPageMetaData[PhysicalPageIndex];
+
+	bool bKeepPage = false;
+	if (PrevMetaData.Flags != 0)
+	{
+		// Update virtual shadow map ID to the equivalent one this frame if present
+		// NOTE: We need a range check as we only add elements to this mapping if they exist this frame
+		int PrevVirtualShadowMapId = PrevMetaData.VirtualShadowMapId;
+		if (PrevVirtualShadowMapId != INDEX_NONE && PrevVirtualShadowMapId < NextVirtualShadowMapDataCount)
+		{
+			FNextVirtualShadowMapData NextData = NextVirtualShadowMapData[PrevVirtualShadowMapId];
+			int VirtualShadowMapId = NextData.NextVirtualShadowMapId;
+
+			// Check if it maps to a valid virtual shadow map this frame
+			if (VirtualShadowMapId != INDEX_NONE)
+			{
+				// Clipmap panning; zeroed otherwise so safe
+				int2 TestPageAddress = int2(PrevMetaData.PageAddress) + NextData.PageAddressOffset;
+				if (IsVirtualShadowMapPageAddressValid(TestPageAddress, PrevMetaData.MipLevel))
+				{
+					// Valid physical page in the cache!
+					// It may still be invalidated by flags or over-written by new requests this frame, but for now we will maintain it
+					OutPhysicalPageMetaData[PhysicalPageIndex].VirtualShadowMapId = VirtualShadowMapId;
+					OutPhysicalPageMetaData[PhysicalPageIndex].PageAddress = uint2(TestPageAddress);
+					// No changes to other fields
+					bKeepPage = true;
+				}
+			}
+		}
+	}
+
+	if (!bKeepPage)
+	{
+		// Only need to zero out flags for it to be considered invalid
+		OutPhysicalPageMetaData[PhysicalPageIndex].Flags = 0;
+	}
+}
+
+
+int bDynamicPageInvalidation;
+int bAllocateViaLRU;
+int MaxPageAgeSinceLastRequest;
+
 [numthreads(VSM_DEFAULT_CS_GROUP_X, 1, 1)]
 void UpdatePhysicalPages(uint3 Index : SV_DispatchThreadID)
 {	
@@ -172,6 +224,7 @@ void UpdatePhysicalPages(uint3 Index : SV_DispatchThreadID)
 
 	// 1:1 read modify write is safe
 	uint NextPhysicalFlags = 0;
+
 #if HAS_CACHE_DATA
 	{
 		FPhysicalPageMetaData PrevMetaData = OutPhysicalPageMetaData[PhysicalPageIndex];
@@ -179,126 +232,115 @@ void UpdatePhysicalPages(uint3 Index : SV_DispatchThreadID)
 			
 		if (PrevMetaData.Flags != 0)
 		{
-			// Update virtual shadow map ID to the equivalent one this frame if present
-			// NOTE: We need a range check as we only add elements to this mapping if they exist this frame
-			int PrevVirtualShadowMapId = PrevMetaData.VirtualShadowMapId;				
-			if (PrevVirtualShadowMapId != INDEX_NONE && PrevVirtualShadowMapId < NextVirtualShadowMapDataCount)
-			{
-				FNextVirtualShadowMapData NextData = NextVirtualShadowMapData[PrevVirtualShadowMapId];
-				int VirtualShadowMapId = NextData.NextVirtualShadowMapId;
+			// Convenience
+			const int VirtualShadowMapId = PrevMetaData.VirtualShadowMapId;
+			const uint2 PageAddress = PrevMetaData.PageAddress;
 
-				// Check if it maps to a valid virtual shadow map this frame
-				if (VirtualShadowMapId != INDEX_NONE)
-				{
-					const uint PrevGlobalOffset = CalcPageOffset(PrevVirtualShadowMapId, MipLevel, PrevMetaData.PageAddress);
-					const uint PrevPhysicalFlags = PrevMetaData.Flags;
-						
-					// Clipmap panning; zeroed otherwise so safe
-					int2 TestPageAddress = int2(PrevMetaData.PageAddress) + NextData.PageAddressOffset;
-					if (IsVirtualShadowMapPageAddressValid(TestPageAddress, MipLevel))
-					{
-						// Valid physical page in the cache!
-						// It may still be invalidated by flags or over-written by new requests this frame, but for now we will maintain it
-						uint2 PageAddress = uint2(TestPageAddress);
-						uint GlobalPageOffset = CalcPageOffset(VirtualShadowMapId, MipLevel, PageAddress);
+			// Look up the request flags for this frame to see if this page was requested again
+			const uint GlobalPageOffset = CalcPageOffset(VirtualShadowMapId, MipLevel, PageAddress);
+			const uint RequestFlags = PageRequestFlags[GlobalPageOffset];
+			const bool bRequestedThisFrame = RequestFlags != 0;
 
-						// Look up the request flags for this frame to see if this page was requested again
-						const uint RequestFlags = PageRequestFlags[GlobalPageOffset];
-						const bool bRequestedThisFrame = RequestFlags != 0;
+			const int PhysicalPageRequestedAge = int(VirtualShadowMap.SceneFrameNumber - PrevMetaData.LastRequestedSceneFrameNumber);
 
-						if (bKeepOnlyRequestedPages == 0 || bRequestedThisFrame)
-						{
-							const FVirtualShadowMapProjectionShaderData Projection = GetVirtualShadowMapProjectionData(VirtualShadowMapId);
-							const bool bCacheStaticSeparately = VirtualShadowMapShouldCacheStaticSeparately();
+			// If the light is unreferenced we also allow its pages to live (unless reallocated) regardless of age for now
+			// since we won't be rendering into them so they don't do a lot of harm being present.
+			const FVirtualShadowMapProjectionShaderData Projection = GetVirtualShadowMapProjectionData(VirtualShadowMapId);
+			if (bRequestedThisFrame || Projection.bUnreferenced || PhysicalPageRequestedAge <= MaxPageAgeSinceLastRequest)
+			{
+				const uint PrevPhysicalFlags = PrevMetaData.Flags;
 
-							bool bDynamicValid = true;
-							bool bStaticValid  = true;
-							// Distant lights block GPU-GPU invalidations (and all others) as they are round-robin invalidated
-							if (bDynamicPageInvalidation && !Projection.bCurrentDistantLight)
-							{
-								uint InvalidationFlags = (PrevPhysicalFlags >> VSM_PHYSICAL_PAGE_INVALIDATION_FLAGS_SHIFT);
-								if (bCacheStaticSeparately)
-								{
-									bDynamicValid = (InvalidationFlags & VSM_DYNAMIC_UNCACHED_FLAG) == 0;
-									bStaticValid  = (InvalidationFlags & VSM_STATIC_UNCACHED_FLAG ) == 0;
-								}
-								else
-								{
-									bDynamicValid = InvalidationFlags == 0;
-									bStaticValid  = InvalidationFlags == 0;
-								}
-							}
+				// Update the mapping data for any valid cached pages so we don't lose it
+				OutPhysicalPageMetaData[PhysicalPageIndex].VirtualShadowMapId = VirtualShadowMapId;
+				OutPhysicalPageMetaData[PhysicalPageIndex].PageAddress = PageAddress;
 
-							// If the page was requested this frame, but doesn't match the detail geometry flag we treat this
-							// as a full invalidation.
-							const uint PhysicalPageDetailGeometryFlag = (PrevPhysicalFlags & VSM_DETAIL_GEOMETRY_FLAG);
-							const uint RequestDetailGeometryFlag = (RequestFlags & VSM_DETAIL_GEOMETRY_FLAG);
+				// Map the page to the physical page
+				// Note that we do this even if it was invalidated though as we might as well reuse the same slot and clear/overwrite it this frame.
+				// We also do this even if the page was not requested this frame because we still need the mapping for invalidations to be triggered, etc.
+				// If we later allocate over top of this page (for one requested this frame), we will zero this out again. See AllocateNewPageMappings
+				checkStructuredBufferAccessSlow(OutPageTable, GlobalPageOffset);
+				checkStructuredBufferAccessSlow(OutPageFlags, GlobalPageOffset);
+				OutPageTable[GlobalPageOffset] = ShadowEncodePageTable(VSMPhysicalIndexToPageAddress(PhysicalPageIndex));
+				// NOTE: PageFlags get written below
 
-							uint NextPageFlags = VSM_ALLOCATED_FLAG;
-							if (bRequestedThisFrame && RequestDetailGeometryFlag != PhysicalPageDetailGeometryFlag)
-							{
-								bDynamicValid = false;
-								bStaticValid  = false;
-								NextPageFlags |= RequestDetailGeometryFlag;
-							}
-							else
-							{
-								NextPageFlags |= PhysicalPageDetailGeometryFlag;
-							}
+				if (Projection.bUnreferenced)
+				{
+					// If the light is unrefereced (i.e. we are not going to render to it this frame) we want to leave the physical
+					// metadata alone, *specifically* the invalidation flags. Since an unreferenced shadow map will not get
+					// rendered to this frame, we can't clear these flags and instead want to maintain them until a potential
+					// future frame when this light might be referenced again.
+					// Tag the page so we can skip it in rendering-related tasks like clearing and merging
+					NextPhysicalFlags = PrevPhysicalFlags | VSM_PHYSICAL_FLAG_UNREFERENCED;
+
+					// NOTE: This is unused during this render, but may be used by invalidation between frames/renders
+					OutPageFlags[GlobalPageOffset] = PrevPhysicalFlags & VSM_PAGE_FLAGS_BITS_MASK;
+				}
+				else
+				{
+					uint NextPageFlags = VSM_ALLOCATED_FLAG;
 
-							// Note: stats are now total physical pages, not just ones requested/used this frame!
-							if (bStaticValid)
-							{
-								StatsBufferInterlockedInc(VSM_STAT_STATIC_CACHED_PAGES);	// Static cached pages
-							}
-							if (bDynamicValid)
-							{
-								StatsBufferInterlockedInc(VSM_STAT_DYNAMIC_CACHED_PAGES);	// Dynamic cached pages
-							}
+					// Distant lights block GPU-GPU invalidations (and all others) as they are round-robin invalidated
+					if (bDynamicPageInvalidation && !Projection.bCurrentDistantLight)
+					{
+						uint InvalidationFlags = (PrevPhysicalFlags >> VSM_PHYSICAL_PAGE_INVALIDATION_FLAGS_SHIFT);
+						if (VirtualShadowMapShouldCacheStaticSeparately())
+						{
+							NextPageFlags |= InvalidationFlags;
 
-							if (!bStaticValid)
-							{
-								StatsBufferInterlockedInc(VSM_STAT_STATIC_INVALIDATED_PAGES);
-								NextPageFlags |= VSM_DYNAMIC_UNCACHED_FLAG | VSM_STATIC_UNCACHED_FLAG;
-							}
-							else if (!bDynamicValid)
+							// Since we merge the static into the dynamic page, if static is invalidated, force dynamic invalid too
+							if (InvalidationFlags & VSM_STATIC_UNCACHED_FLAG)
 							{
-								StatsBufferInterlockedInc(VSM_STAT_DYNAMIC_INVALIDATED_PAGES);
 								NextPageFlags |= VSM_DYNAMIC_UNCACHED_FLAG;
 							}
+						}
+						else if (InvalidationFlags != 0)
+						{
+							NextPageFlags |= (VSM_STATIC_UNCACHED_FLAG | VSM_DYNAMIC_UNCACHED_FLAG);
+						}
+					}
 
-							uint2 PhysicalAddress = VSMPhysicalIndexToPageAddress(PhysicalPageIndex);
-							uint PhysicalFlags = (Projection.bUnCached ? VSM_PHYSICAL_FLAG_VIEW_UNCACHED : 0U);
-							if (bRequestedThisFrame)
-							{
-								StatsBufferInterlockedInc(VSM_STAT_REQUESTED_THIS_FRAME_PAGES);
-								PhysicalFlags |= VSM_PHYSICAL_FLAG_USED_THIS_FRAME;
-
-								// Remove from LRU list and add to requested list
-								PushPhysicalPageList(PHYSICAL_PAGE_LIST_REQUESTED, PhysicalPageIndex);
-								bRemovedPageFromList = true;
-							}
+					uint PhysicalPageDetailGeometryFlag = (PrevPhysicalFlags & VSM_DETAIL_GEOMETRY_FLAG);
+					if (bRequestedThisFrame)
+					{
+						StatsBufferInterlockedInc(VSM_STAT_REQUESTED_THIS_FRAME_PAGES);
+						// Remove from LRU list and add to requested list
+						PushPhysicalPageList(PHYSICAL_PAGE_LIST_REQUESTED, PhysicalPageIndex);
+						OutPhysicalPageMetaData[PhysicalPageIndex].LastRequestedSceneFrameNumber = VirtualShadowMap.SceneFrameNumber;
+						bRemovedPageFromList = true;
+
+						// If the detail geometry flag doesn't match the cached page we treat it as a full invalidation.
+						// TODO: This could potentially be a problem for interleaved multiview rendering;
+						// If the flag differs in the two views it will cause cache thrashing.									
+						const uint RequestDetailGeometryFlag = (RequestFlags & VSM_DETAIL_GEOMETRY_FLAG);
+						if (RequestDetailGeometryFlag != PhysicalPageDetailGeometryFlag)
+						{
+							NextPageFlags |= (VSM_STATIC_UNCACHED_FLAG | VSM_DYNAMIC_UNCACHED_FLAG);
+							PhysicalPageDetailGeometryFlag = RequestDetailGeometryFlag;
+						}
 
-							{
-								// Map the page to the physical page
-								// Note that we do this even if it was invalidated though as we might as well reuse the same slot
-								// and clear/overwrite it this frame.
-								// We also do this even if the page was not requested this frame because we still need the mapping
-								// for invalidations to be triggered, etc.
-								checkStructuredBufferAccessSlow(OutPageTable, GlobalPageOffset);
-								checkStructuredBufferAccessSlow(OutPageFlags, GlobalPageOffset);
-								OutPageTable[GlobalPageOffset] = ShadowEncodePageTable(PhysicalAddress);
-								OutPageFlags[GlobalPageOffset] = NextPageFlags;
-							}
-							
-							// Update the mapping data for any valid cached pages so we don't lose it
-							OutPhysicalPageMetaData[PhysicalPageIndex].Age = bStaticValid ? (PrevMetaData.Age + 1) : 0;
-							OutPhysicalPageMetaData[PhysicalPageIndex].VirtualShadowMapId = VirtualShadowMapId;
-							OutPhysicalPageMetaData[PhysicalPageIndex].PageAddress = PageAddress;
-							// Flags set below
-							NextPhysicalFlags = NextPageFlags | PhysicalFlags;
+						// Only increment the stats for pages requested this render, otherwise it gets confusing
+						if (NextPageFlags & VSM_STATIC_UNCACHED_FLAG)
+						{
+							StatsBufferInterlockedInc(VSM_STAT_STATIC_INVALIDATED_PAGES);
+						}
+						else
+						{
+							StatsBufferInterlockedInc(VSM_STAT_STATIC_CACHED_PAGES);
+						}
+						if (NextPageFlags & VSM_DYNAMIC_UNCACHED_FLAG)
+						{
+							StatsBufferInterlockedInc(VSM_STAT_DYNAMIC_INVALIDATED_PAGES);
+						}
+						else
+						{
+							StatsBufferInterlockedInc(VSM_STAT_DYNAMIC_CACHED_PAGES);
 						}
 					}
+					NextPageFlags |= PhysicalPageDetailGeometryFlag;
+								
+					const uint PhysicalFlags = (Projection.bUnCached ? VSM_PHYSICAL_FLAG_VIEW_UNCACHED : 0U);
+					OutPageFlags[GlobalPageOffset] = NextPageFlags;
+					NextPhysicalFlags = NextPageFlags | PhysicalFlags;
 				}
 			}
 		}
@@ -317,9 +359,7 @@ void UpdatePhysicalPages(uint3 Index : SV_DispatchThreadID)
 	OutPhysicalPageMetaData[PhysicalPageIndex].Flags = NextPhysicalFlags;
 
 	// Write out the LRU list while maintaining order, with anything we removed marked as INDEX_NONE
-	{
-		SetPhysicalPageListItem(PHYSICAL_PAGE_LIST_LRU, PhysicalPageListIndex, bRemovedPageFromList ? INDEX_NONE : PhysicalPageIndex);
-	}
+	SetPhysicalPageListItem(PHYSICAL_PAGE_LIST_LRU, PhysicalPageListIndex, bRemovedPageFromList ? INDEX_NONE : PhysicalPageIndex);
 }
 
 // Utility to map from PageIndex -> PageX, PageY, MipLevel
@@ -388,11 +428,10 @@ void AllocateNewPageMappings(uint VirtualShadowMapId, uint GlobalPageOffset, uin
 				OutPageFlags[GlobalPageOffset] = Flags;
 
 				const FVirtualShadowMapProjectionShaderData Projection = GetVirtualShadowMapProjectionData(VirtualShadowMapId);
-				const uint PhysicalFlags = VSM_PHYSICAL_FLAG_USED_THIS_FRAME |
-					(Projection.bUnCached ? VSM_PHYSICAL_FLAG_VIEW_UNCACHED : 0U);
+				const uint PhysicalFlags = (Projection.bUnCached ? VSM_PHYSICAL_FLAG_VIEW_UNCACHED : 0U);
 
-				OutPhysicalPageMetaData[PhysicalPageIndex].Age = 0;
 				OutPhysicalPageMetaData[PhysicalPageIndex].Flags = Flags | PhysicalFlags;
+				OutPhysicalPageMetaData[PhysicalPageIndex].LastRequestedSceneFrameNumber = VirtualShadowMap.SceneFrameNumber;
 				OutPhysicalPageMetaData[PhysicalPageIndex].VirtualShadowMapId = VirtualShadowMapId;
 
 				// Compute page address info from the offset
@@ -486,6 +525,10 @@ void PackAvailablePages(uint GroupIndex : SV_GroupIndex)
 			SetPhysicalPageListItem(PHYSICAL_PAGE_LIST_AVAILABLE, TotalCount + Offset, PhysicalPageIndex);
 		}
 		TotalCount += GroupCount;
+
+		// This should already be accounted for internally by ThreadGroupPrefixSum, but putting one here
+		// to be absolutely sure.
+		GroupMemoryBarrierWithGroupSync();
 	}
 
 	// Set total number
@@ -531,8 +574,35 @@ void AppendPhysicalPageLists(uint ThreadId : SV_DispatchThreadID)
 			// list must end up as the right size after the final append!
 			if (!bAppendEmptyToAvailable)
 			{
+				/*
+				PLATFORM_ASSERT4(
+					NewOutputCount == VirtualShadowMap.MaxPhysicalPages,
+					0xCECC,
+					__LINE__,
+					InputCount,
+					OutputCount,
+					NewOutputCount);
+				*/
 				checkSlow(NewOutputCount == VirtualShadowMap.MaxPhysicalPages);
 			}
+			else
+			{
+				/*
+				// All pages should now be in PHYSICAL_PAGE_LIS_AVAILABLE or REQUESTED
+				int AvailableCount = GetPhysicalPageListCount(PHYSICAL_PAGE_LIST_AVAILABLE);
+				int RequestedCount = GetPhysicalPageListCount(PHYSICAL_PAGE_LIST_REQUESTED);
+				int EmptyCount = GetPhysicalPageListCount(PHYSICAL_PAGE_LIST_EMPTY);
+				int TotalPages = AvailableCount + RequestedCount;
+
+				PLATFORM_ASSERT4(
+					TotalPages == VirtualShadowMap.MaxPhysicalPages,
+					0xCECC,
+					__LINE__,
+					AvailableCount,
+					RequestedCount,
+					EmptyCount);
+				*/
+			}
 		}
 	}
 	else
@@ -615,6 +685,7 @@ void SelectPagesToInitializeCS(uint PhysicalPageIndex : SV_DispatchThreadID)
 
 	FPhysicalPageMetaData MetaData = PhysicalPageMetaData[PhysicalPageIndex];
 
+	bool bUnreferenced = (MetaData.Flags & VSM_PHYSICAL_FLAG_UNREFERENCED) != 0;
 	bool bFullyCached = (MetaData.Flags & (VSM_DYNAMIC_UNCACHED_FLAG | VSM_STATIC_UNCACHED_FLAG)) == 0;
 	bool bStaticUncached = (MetaData.Flags & VSM_STATIC_UNCACHED_FLAG) != 0;
 
@@ -622,9 +693,9 @@ void SelectPagesToInitializeCS(uint PhysicalPageIndex : SV_DispatchThreadID)
 	{
 		// Page not used, we're done
 	}
-	else if (bFullyCached)
+	else if (bUnreferenced || bFullyCached)
 	{
-		// Page fully cached. Leave the data alone.
+		// Page fully cached or unreferenced. Leave the data alone.
 	}
 	else
 	{
@@ -770,8 +841,6 @@ bool UpdateAndClearDirtyFlags(uint PhysicalPageIndex, inout FPhysicalPageMetaDat
 	bool bInvalidatesDynamic = DirtyPageFlagsInOut[PhysicalPageIndex + VirtualShadowMap.MaxPhysicalPages] != 0U;
 	bool bInvalidatesStatic = DirtyPageFlagsInOut[PhysicalPageIndex + 2U * VirtualShadowMap.MaxPhysicalPages] != 0U;
 
-	MetaData = OutPhysicalPageMetaData[PhysicalPageIndex];
-
 	// clear the dirty/invalidation flags
 	DirtyPageFlagsInOut[PhysicalPageIndex] = 0U;
 	DirtyPageFlagsInOut[PhysicalPageIndex + VirtualShadowMap.MaxPhysicalPages] = 0U;
@@ -781,8 +850,9 @@ bool UpdateAndClearDirtyFlags(uint PhysicalPageIndex, inout FPhysicalPageMetaDat
 
 	// Update the page metadata to mark the pages as uncached, which allows the page merging to pick them up.
 	// Also mark bits for invalidation (need own set of bits to avoid being merged with the input state at the start of the frame - e.g., fresh pages are already uncached)
+	MetaData = OutPhysicalPageMetaData[PhysicalPageIndex];
 	MetaData.Flags |= (bPageDirty ? VSM_PHYSICAL_FLAG_DIRTY : 0U) | (CacheFlags << VSM_PHYSICAL_PAGE_INVALIDATION_FLAGS_SHIFT);
-	OutPhysicalPageMetaData[PhysicalPageIndex] = MetaData;
+	OutPhysicalPageMetaData[PhysicalPageIndex].Flags = MetaData.Flags;
 
 	return bPageDirty;
 }
@@ -816,9 +886,9 @@ void SelectPagesForHZBAndUpdateDirtyFlagsCS(uint PhysicalPageIndex : SV_Dispatch
 
 	if ((MetaData.Flags & VSM_ALLOCATED_FLAG) != 0)
 	{
-		// rebuild the HZB if the page is rendered to or it is freshly allocated, but only the first time for a given frame.
-		// TODO: Review this "age" logic and potentially look at metadata flags instead
-		if (bForceFullHZBUpdate || bPageDirty || (bFirstBuildThisFrame != 0 && MetaData.Age == 0U))
+		// rebuild the HZB if the page is rendered to or it is freshly allocated/cleared, but only the first time for a given frame.
+		bool bStaticUncached = (MetaData.Flags & VSM_STATIC_UNCACHED_FLAG) != 0;
+		if (bForceFullHZBUpdate || bPageDirty || (bFirstBuildThisFrame != 0 && bStaticUncached))
 		{
 			StatsBufferInterlockedInc(VSM_STAT_NUM_HZB_PAGES_BUILT);
 
@@ -1042,7 +1112,30 @@ void FeedbackStatusCS()
 {
 	{
 		FGPUMessageWriter Mw = GPUMessageBegin(StatusMessageId, 1U);
-		// 2. Write out the free list status
-		GPUMessageWriteItem(Mw, PhysicalPageLists[VirtualShadowMap.MaxPhysicalPages]);
+
+		// Write out how many pages are still available
+		int CountIndex = GetPhysicalPageListStart(PHYSICAL_PAGE_LIST_AVAILABLE) + VirtualShadowMap.MaxPhysicalPages;		
+		GPUMessageWriteItem(Mw, PhysicalPageLists[CountIndex]);
+	}
+}
+
+
+
+int PageListStatsRow = 0; 
+
+[numthreads(1, 1, 1)]
+void LogPageListStatsCS()
+{
+	float TopMargin = 0.5f;
+	float ItemX = 0.05f;
+	FShaderPrintContext Ctx = InitShaderPrintContext(true, float2(ItemX, TopMargin));
+
+	Ctx.Pos.y += PageListStatsRow * 0.02f;
+
+	Print(Ctx, GetPhysicalPageListCount(0));
+	for (int i = 1; i < PHYSICAL_PAGE_LIST_COUNT; ++i)
+	{
+		Print(Ctx, TEXT(", "));
+		Print(Ctx, GetPhysicalPageListCount(i));
 	}
 }
diff --git a/Engine/Shaders/Private/VirtualShadowMaps/VirtualShadowMapProjectionSpot.ush b/Engine/Shaders/Private/VirtualShadowMaps/VirtualShadowMapProjectionSpot.ush
--- a/Engine/Shaders/Private/VirtualShadowMaps/VirtualShadowMapProjectionSpot.ush
+++ b/Engine/Shaders/Private/VirtualShadowMaps/VirtualShadowMapProjectionSpot.ush
@@ -13,9 +13,6 @@ VirtualShadowMapProjectionSpot.ush:
 #include "VirtualShadowMapProjectionCommon.ush"
 #include "VirtualShadowMapSMRTCommon.ush"
 
-// Possibly parameterize this at some point, but tradeoffs are not always obvious
-#define VIRTUAL_SHADOW_MAP_MAX_SLOPE_BIAS 100.0f
-
 float2 ComputeDepthSlopeLocalUV(
 	int VirtualShadowMapId,
 	float3 ShadowTranslatedWorldPosition,
diff --git a/Engine/Shaders/Private/VirtualShadowMaps/VirtualShadowMapProjectionStructs.ush b/Engine/Shaders/Private/VirtualShadowMaps/VirtualShadowMapProjectionStructs.ush
--- a/Engine/Shaders/Private/VirtualShadowMaps/VirtualShadowMapProjectionStructs.ush
+++ b/Engine/Shaders/Private/VirtualShadowMaps/VirtualShadowMapProjectionStructs.ush
@@ -32,6 +32,7 @@ struct FVirtualShadowMapProjectionShaderData
 	int VirtualShadowMapId;	
 	bool bCurrentDistantLight; // Low detail == single page only
 	bool bUnCached; // see VSM_PROJ_FLAG_UNCACHED
+	bool bUnreferenced; // see VSM_PROJ_FLAG_UNREFERENCED
 };
 
 FVirtualShadowMapProjectionShaderData DecodeVirtualShadowMapProjectionData(ByteAddressBuffer ProjectionData, int VirtualShadowMapId)
@@ -84,6 +85,7 @@ FVirtualShadowMapProjectionShaderData DecodeVirtualShadowMapProjectionData(ByteA
 
 	Result.bCurrentDistantLight = (Result.Flags & VSM_PROJ_FLAG_CURRENT_DISTANT_LIGHT) != 0U;
 	Result.bUnCached = (Result.Flags & VSM_PROJ_FLAG_UNCACHED) != 0U;
+	Result.bUnreferenced = (Result.Flags & VSM_PROJ_FLAG_UNREFERENCED) != 0U;
 
 	return Result;
 }
diff --git a/Engine/Shaders/Shared/VirtualShadowMapDefinitions.h b/Engine/Shaders/Shared/VirtualShadowMapDefinitions.h
--- a/Engine/Shaders/Shared/VirtualShadowMapDefinitions.h
+++ b/Engine/Shaders/Shared/VirtualShadowMapDefinitions.h
@@ -22,6 +22,7 @@ VirtualShadowMapDefinitions.h: used in virtual shadow map shaders and C++ code t
 
 #define VSM_PROJ_FLAG_CURRENT_DISTANT_LIGHT (1U << 0)
 #define VSM_PROJ_FLAG_UNCACHED (1U << 1) // Used to indicate that the light is uncached and should only render to dynamic pages
+#define VSM_PROJ_FLAG_UNREFERENCED (1U << 2) // Used to indicate that the light is not referenced/rendered to this render
 
 // Hard limit for max distant lights supported 8k for now - we may revise later. We need to keep them in a fixed range for now to make allocation easy and minimize overhead for indexing.
 #define VSM_MAX_SINGLE_PAGE_SHADOW_MAPS (1024U * 8U)
diff --git a/Engine/Source/Runtime/Engine/Public/SceneManagement.h b/Engine/Source/Runtime/Engine/Public/SceneManagement.h
--- a/Engine/Source/Runtime/Engine/Public/SceneManagement.h
+++ b/Engine/Source/Runtime/Engine/Public/SceneManagement.h
@@ -245,14 +245,6 @@ PRAGMA_ENABLE_DEPRECATION_WARNINGS
 	virtual FLandscapeRayTracingStateList* GetLandscapeRayTracingStates() const = 0;
 #endif
 
-	/**
-	* Adds a per-view virtual shadow map cache, which can help performance, at a cost in memory.  Does nothing if one is already present.
-	* The cache only works for the first Scene this function is called for.
-	*/
-	virtual void AddVirtualShadowMapCache(FSceneInterface* InScene) {}
-	virtual void RemoveVirtualShadowMapCache(FSceneInterface* InScene) {}
-	virtual bool HasVirtualShadowMapCache() const = 0;
-
 	/** Similar to above, but adds Lumen Scene Data */
 	virtual void AddLumenSceneData(FSceneInterface* InScene, float SurfaceCacheResolution = 1.0f) {}
 	virtual void RemoveLumenSceneData(FSceneInterface* InScene) {}
@@ -282,7 +274,6 @@ protected:
 	uint8 bValidEyeAdaptationBuffer : 1;
 
 private:
-	virtual FVirtualShadowMapArrayCacheManager* GetVirtualShadowMapCache(const FScene* InScene) const { return nullptr; }
 	friend class FScene;
 };
 
diff --git a/Engine/Source/Runtime/Renderer/Private/DeferredShadingRenderer.cpp b/Engine/Source/Runtime/Renderer/Private/DeferredShadingRenderer.cpp
--- a/Engine/Source/Runtime/Renderer/Private/DeferredShadingRenderer.cpp
+++ b/Engine/Source/Runtime/Renderer/Private/DeferredShadingRenderer.cpp
@@ -2524,7 +2524,7 @@ void FDeferredShadingSceneRenderer::Render(FRDGBuilder& GraphBuilder)
 			// NOTE: Must be done after  system texture initialization
 			// TODO: This doesn't take into account the potential for split screen views with separate shadow caches
 			const bool bEnableVirtualShadowMaps = UseVirtualShadowMaps(ShaderPlatform, FeatureLevel) && ViewFamily.EngineShowFlags.DynamicShadows;
-			VirtualShadowMapArray.Initialize(GraphBuilder, Scene->GetVirtualShadowMapCache(Views[0]), bEnableVirtualShadowMaps, Views[0].bIsSceneCapture);
+			VirtualShadowMapArray.Initialize(GraphBuilder, Scene->GetVirtualShadowMapCache(), bEnableVirtualShadowMaps);
 
 			if (InitViewTaskDatas.LumenFrameTemporaries)
 			{
diff --git a/Engine/Source/Runtime/Renderer/Private/GPUScene.cpp b/Engine/Source/Runtime/Renderer/Private/GPUScene.cpp
--- a/Engine/Source/Runtime/Renderer/Private/GPUScene.cpp
+++ b/Engine/Source/Runtime/Renderer/Private/GPUScene.cpp
@@ -1566,7 +1566,7 @@ void FGPUScene::UploadDynamicPrimitiveShaderDataForViewInternal(FRDGBuilder& Gra
 		ensure(UploadIdStart < DynamicPrimitivesOffset);
 		ensure(InstanceIdStart != INDEX_NONE);
 
-		if (bIsShadowView && Scene.GetVirtualShadowMapCache(View) != nullptr)
+		if (bIsShadowView)
 		{
 			// Enqueue cache invalidations for all dynamic primitives' instances, as they will be removed this frame and are not associated
 			// with any particular FPrimitiveSceneInfo. Will occur on the next call to UpdateAllPrimitiveSceneInfos
diff --git a/Engine/Source/Runtime/Renderer/Private/MobileShadingRenderer.cpp b/Engine/Source/Runtime/Renderer/Private/MobileShadingRenderer.cpp
--- a/Engine/Source/Runtime/Renderer/Private/MobileShadingRenderer.cpp
+++ b/Engine/Source/Runtime/Renderer/Private/MobileShadingRenderer.cpp
@@ -900,7 +900,7 @@ void FMobileSceneRenderer::Render(FRDGBuilder& GraphBuilder)
 	// Important that this uses consistent logic throughout the frame, so evaluate once and pass in the flag from here
 	// NOTE: Must be done after  system texture initialization
 	// TODO: This doesn't take into account the potential for split screen views with separate shadow caches
-	VirtualShadowMapArray.Initialize(GraphBuilder, Scene->GetVirtualShadowMapCache(Views[0]), UseVirtualShadowMaps(ShaderPlatform, FeatureLevel), Views[0].bIsSceneCapture);
+	VirtualShadowMapArray.Initialize(GraphBuilder, Scene->GetVirtualShadowMapCache(), UseVirtualShadowMaps(ShaderPlatform, FeatureLevel));
 
 	FInitViewTaskDatas InitViewTaskDatas(VisibilityTaskData);
 
diff --git a/Engine/Source/Runtime/Renderer/Private/PlanarReflectionRendering.cpp b/Engine/Source/Runtime/Renderer/Private/PlanarReflectionRendering.cpp
--- a/Engine/Source/Runtime/Renderer/Private/PlanarReflectionRendering.cpp
+++ b/Engine/Source/Runtime/Renderer/Private/PlanarReflectionRendering.cpp
@@ -674,17 +674,6 @@ void FScene::UpdatePlanarReflectionContents(UPlanarReflectionComponent* CaptureC
 			Extension->SetupViewFamily(SceneRenderer->ViewFamily);
 		}
 
-		FSceneViewStateInterface* ViewStateInterface = CaptureComponent->GetViewState(0);
-
-		if (UseVirtualShadowMaps(SceneRenderer->ShaderPlatform, FeatureLevel) && ViewStateInterface)
-		{
-			ViewStateInterface->AddVirtualShadowMapCache(this);
-		}
-		else if (ViewStateInterface)
-		{
-			ViewStateInterface->RemoveVirtualShadowMapCache(this);
-		}
-
 		for (int32 ViewIndex = 0; ViewIndex < SceneCaptureViewInfo.Num(); ++ViewIndex)
 		{
 			FViewInfo& ViewInfo = SceneRenderer->Views[ViewIndex];
diff --git a/Engine/Source/Runtime/Renderer/Private/RendererScene.cpp b/Engine/Source/Runtime/Renderer/Private/RendererScene.cpp
--- a/Engine/Source/Runtime/Renderer/Private/RendererScene.cpp
+++ b/Engine/Source/Runtime/Renderer/Private/RendererScene.cpp
@@ -391,12 +391,9 @@ FSceneViewState::FSceneViewState(ERHIFeatureLevel::Type FeatureLevel, FSceneView
 	GatherPointsResolution = FIntVector(0, 0, 0);
 #endif
 
-	bVirtualShadowMapCacheAdded = false;
 	bLumenSceneDataAdded = false;
 	LumenSurfaceCacheResolution = 1.0f;
 
-	ViewVirtualShadowMapCache = nullptr;
-
 	// OcclusionFeedback works only with mobile rendering atm
 	if (FeatureLevel == ERHIFeatureLevel::ES3_1)
 	{
@@ -443,13 +440,6 @@ FSceneViewState::~FSceneViewState()
 	HairStrandsViewStateData.Release();
 	ShaderPrintStateData.Release();
 
-	if (ViewVirtualShadowMapCache)
-	{
-		delete ViewVirtualShadowMapCache;
-		ViewVirtualShadowMapCache = nullptr;
-		bVirtualShadowMapCacheAdded = false;
-	}
-
 	if (Scene)
 	{
 		Scene->RemoveViewState_RenderThread(this);
@@ -1160,10 +1150,6 @@ uint64 FSceneViewState::GetGPUSizeBytes(bool bLogSizes) const
 	TotalSize += GetBufferGPUSizeBytes(ShaderPrintStateData.StateBuffer, bLogSizes);
 	TotalSize += ShadingEnergyConservationData.GetGPUSizeBytes(bLogSizes);
 	TotalSize += GlintShadingLUTsData.GetGPUSizeBytes(bLogSizes);
-	if (ViewVirtualShadowMapCache)
-	{
-		TotalSize += ViewVirtualShadowMapCache->GetGPUSizeBytes(bLogSizes);
-	}
 
 	// Per-view Lumen scene data is stored in a map in the FScene
 	if (Scene && bLumenSceneDataAdded)
@@ -1180,82 +1166,6 @@ uint64 FSceneViewState::GetGPUSizeBytes(bool bLogSizes) const
 	return TotalSize;
 }
 
-void FSceneViewState::AddVirtualShadowMapCache(FSceneInterface* InScene)
-{
-	check(InScene);
-	if (!Scene)
-	{
-		Scene = (FScene*)InScene;
-
-		// Modification of scene structure needs to happen on render thread
-		ENQUEUE_RENDER_COMMAND(SceneViewStateAdd)(
-			[RenderScene = Scene, RenderViewState = this](FRHICommandList&)
-			{
-				RenderScene->ViewStates.Add(RenderViewState);
-			});
-	}
-
-	if (Scene == InScene)
-	{
-		// Don't add the cache if one has already been added.  We have a separate bool since the write to the cache
-		// pointer is deferred to the render thread.
-		if (!bVirtualShadowMapCacheAdded)
-		{
-			bVirtualShadowMapCacheAdded = true;
-
-			FVirtualShadowMapArrayCacheManager* ViewCache = new FVirtualShadowMapArrayCacheManager(Scene);
-
-			// Need to add reference to virtual shadow map cache in render thread
-			ENQUEUE_RENDER_COMMAND(LinkVirtualShadowMapCache)(
-				[this, ViewCache](FRHICommandListImmediate& RHICmdList)
-				{
-					this->ViewVirtualShadowMapCache = ViewCache;
-				});
-		} //-V773
-	}
-}
-
-void FSceneViewState::RemoveVirtualShadowMapCache(FSceneInterface* InScene)
-{
-	check(InScene);
-	if (Scene == InScene && bVirtualShadowMapCacheAdded)
-	{
-		bVirtualShadowMapCacheAdded = false;
-
-		ENQUEUE_RENDER_COMMAND(RemoveVirtualShadowMapCache)(
-			[this](FRHICommandListImmediate& RHICmdList)
-			{
-				delete ViewVirtualShadowMapCache;
-				ViewVirtualShadowMapCache = nullptr;
-			});
-	}
-}
-
-bool FSceneViewState::HasVirtualShadowMapCache() const
-{
-	return bVirtualShadowMapCacheAdded;
-}
-
-FVirtualShadowMapArrayCacheManager* FSceneViewState::GetVirtualShadowMapCache(const FScene* InScene) const
-{
-	// Per-view VSM cache can only be used for the Scene the view state was previously linked to
-	return Scene == InScene ? ViewVirtualShadowMapCache : nullptr;
-}
-
-FVirtualShadowMapArrayCacheManager* FScene::GetVirtualShadowMapCache(FSceneView& View) const
-{
-	FVirtualShadowMapArrayCacheManager* Result = DefaultVirtualShadowMapCache;
-	if (View.State)
-	{
-		FVirtualShadowMapArrayCacheManager* ViewCache = View.State->GetVirtualShadowMapCache(this);
-		if (ViewCache)
-		{
-			Result = ViewCache;
-		}
-	}
-	return Result;
-}
-
 void FSceneViewState::AddLumenSceneData(FSceneInterface* InScene, float InSurfaceCacheResolution)
 {
 	check(InScene);
@@ -1362,23 +1272,6 @@ FLumenSceneData* FScene::FindLumenSceneData(uint32 ViewKey, uint32 GPUIndex) con
 	return DefaultLumenSceneData;
 }
 
-void FScene::GetAllVirtualShadowMapCacheManagers(TArray<FVirtualShadowMapArrayCacheManager*, SceneRenderingAllocator>& OutCacheManagers) const
-{
-	OutCacheManagers.Empty();
-	if (DefaultVirtualShadowMapCache)
-	{
-		OutCacheManagers.Add(DefaultVirtualShadowMapCache);
-	}
-	for (const FSceneViewState* ViewState : ViewStates)
-	{
-		// Per-view VSM cache can only be used for the Scene the view state was previously linked to
-		if (ViewState->ViewVirtualShadowMapCache && ViewState->Scene == this)
-		{
-			OutCacheManagers.Add(ViewState->ViewVirtualShadowMapCache);
-		}
-	}
-}
-
 void FScene::UpdateParameterCollections(const TArray<FMaterialParameterCollectionInstanceResource*>& InParameterCollections)
 {
 	ENQUEUE_RENDER_COMMAND(UpdateParameterCollectionsCommand)(
@@ -1831,9 +1724,7 @@ FScene::FScene(UWorld* InWorld, bool bInRequiresHitProxies, bool bInIsEditorScen
 
 	DefaultLumenSceneData = new FLumenSceneData(GShaderPlatformForFeatureLevel[InFeatureLevel], InWorld->WorldType);
 
-	// We use a default Virtual Shadow Map cache, if one hasn't been allocated for a specific view.  GPU resources for
-	// a cache aren't allocated until the cache is actually used, so this shouldn't waste any GPU memory when not in use.
-	DefaultVirtualShadowMapCache = new FVirtualShadowMapArrayCacheManager(this);
+	VirtualShadowMapCache = new FVirtualShadowMapArrayCacheManager(this);
 
 	SceneLightInfoUpdates = new FSceneLightInfoUpdates;
 
@@ -1867,20 +1758,14 @@ FScene::~FScene()
 	{
 		check(ViewState->Scene == this);
 		ViewState->Scene = nullptr;
-
-		if (ViewState->ViewVirtualShadowMapCache)
-		{
-			delete ViewState->ViewVirtualShadowMapCache;
-			ViewState->ViewVirtualShadowMapCache = nullptr;
-		}
 	}
 	ViewStates.Empty();
 
 	// Delete default cache
-	if (DefaultVirtualShadowMapCache)
+	if (VirtualShadowMapCache)
 	{
-		delete DefaultVirtualShadowMapCache;
-		DefaultVirtualShadowMapCache = nullptr;
+		delete VirtualShadowMapCache;
+		VirtualShadowMapCache = nullptr;
 	}
 	if (SceneCulling)
 	{
@@ -3939,15 +3824,7 @@ void FScene::RemoveLightSceneInfo_RenderThread(FLightSceneInfo* LightSceneInfo)
 	Lights.RemoveAt(LightSceneInfo->Id);
 
 	// TODO: move this work to FShadowScene & batch the light removals
-	{
-		TArray<FVirtualShadowMapArrayCacheManager*, SceneRenderingAllocator> VirtualShadowCacheManagers;
-		GetAllVirtualShadowMapCacheManagers(VirtualShadowCacheManagers);
-
-		for (FVirtualShadowMapArrayCacheManager* CacheManager : VirtualShadowCacheManagers)
-		{
-			CacheManager->OnLightRemoved(LightSceneInfo->Id);
-		}
-	}
+	GetVirtualShadowMapCache()->OnLightRemoved(LightSceneInfo->Id);
 
 	if (!LightSceneInfo->Proxy->HasStaticShadowing()
 		&& LightSceneInfo->Proxy->CastsDynamicShadow()
@@ -5433,47 +5310,42 @@ void FScene::UpdateAllPrimitiveSceneInfos(FRDGBuilder& GraphBuilder, EUpdateAllP
 
 	SceneCullingUpdater.OnPreSceneUpdate(GraphBuilder, SceneUpdateChangeSetStorage.GetPreUpdateSet());
 
-	TArray<FVirtualShadowMapArrayCacheManager*, SceneRenderingAllocator> VirtualShadowCacheManagers;
-
 	// Create a SceneUB that permits access to the scene for invalidation processing.
 	FSceneUniformBuffer SceneUB;
 	GPUScene.FillSceneUniformBuffer(GraphBuilder, SceneUB);
 
 	{
 		SCOPED_NAMED_EVENT(FScene_VirtualShadowCacheUpdate, FColor::Orange);
-		GetAllVirtualShadowMapCacheManagers(VirtualShadowCacheManagers);
+		FVirtualShadowMapArrayCacheManager* CacheManager = GetVirtualShadowMapCache();
 
-		for (FVirtualShadowMapArrayCacheManager* CacheManager : VirtualShadowCacheManagers)
-		{
-			FVirtualShadowMapArrayCacheManager::FInvalidatingPrimitiveCollector InvalidatingPrimitiveCollector(CacheManager);
-			InvalidatingPrimitiveCollector.AddDynamicAndGPUPrimitives();
+		FVirtualShadowMapArrayCacheManager::FInvalidatingPrimitiveCollector InvalidatingPrimitiveCollector(CacheManager);
+		InvalidatingPrimitiveCollector.AddDynamicAndGPUPrimitives();
 
-			// Primitives that are tracked as always invalidating shadows, pipe through as transform updates
-			for (FPrimitiveSceneInfo* PrimitiveSceneInfo : ShadowScene->GetAlwaysInvalidatingPrimitives())
-			{
-				InvalidatingPrimitiveCollector.UpdatedTransform(PrimitiveSceneInfo);
-			}
+		// Primitives that are tracked as always invalidating shadows, pipe through as transform updates
+		for (FPrimitiveSceneInfo* PrimitiveSceneInfo : ShadowScene->GetAlwaysInvalidatingPrimitives())
+		{
+			InvalidatingPrimitiveCollector.UpdatedTransform(PrimitiveSceneInfo);
+		}
 
-			// All removed primitives must invalidate their footprints in the VSM before leaving
-			for (FPrimitiveSceneInfo* PrimitiveSceneInfo : RemovedLocalPrimitiveSceneInfos)
-			{
-				InvalidatingPrimitiveCollector.Removed(PrimitiveSceneInfo);
-			}
-			// All updated instances must also before moving or re-allocating (TODO: filter out only those actually updated)
-			for (const auto& Instance : UpdatedInstances)
-			{
-				InvalidatingPrimitiveCollector.UpdatedInstances(Instance.Key->GetPrimitiveSceneInfo());
-			}
-			// As must all primitive updates, 
-			for (const auto& Transform : UpdatedTransforms)
-			{
-				InvalidatingPrimitiveCollector.UpdatedTransform(Transform.Key->GetPrimitiveSceneInfo());
-			}
+		// All removed primitives must invalidate their footprints in the VSM before leaving
+		for (FPrimitiveSceneInfo* PrimitiveSceneInfo : RemovedLocalPrimitiveSceneInfos)
+		{
+			InvalidatingPrimitiveCollector.Removed(PrimitiveSceneInfo);
+		}
+		// All updated instances must also before moving or re-allocating (TODO: filter out only those actually updated)
+		for (const auto& Instance : UpdatedInstances)
+		{
+			InvalidatingPrimitiveCollector.UpdatedInstances(Instance.Key->GetPrimitiveSceneInfo());
+		}
+		// As must all primitive updates, 
+		for (const auto& Transform : UpdatedTransforms)
+		{
+			InvalidatingPrimitiveCollector.UpdatedTransform(Transform.Key->GetPrimitiveSceneInfo());
+		}
 
-			InvalidatingPrimitiveCollector.Finalize();
+		InvalidatingPrimitiveCollector.Finalize();
 
-			CacheManager->ProcessInvalidations(GraphBuilder, SceneUB, InvalidatingPrimitiveCollector);
-		}
+		CacheManager->ProcessInvalidations(GraphBuilder, SceneUB, InvalidatingPrimitiveCollector);
 	}
 
 	TArray<FPrimitiveSceneInfo*> AddedLocalPrimitiveSceneInfos;
@@ -6205,10 +6077,7 @@ void FScene::UpdateAllPrimitiveSceneInfos(FRDGBuilder& GraphBuilder, EUpdateAllP
 	FPrimitiveSceneInfo::AllocateGPUSceneInstances(this, PendingAllocateInstanceIds);
 
 	// handle scene changes
-	for (FVirtualShadowMapArrayCacheManager* CacheManager : VirtualShadowCacheManagers)
-	{
-		CacheManager->OnSceneChange();
-	}
+	GetVirtualShadowMapCache()->OnSceneChange();
 
 	if (SceneInfosWithAddToScene.Num() > 0)
 	{
diff --git a/Engine/Source/Runtime/Renderer/Private/SceneCaptureRendering.cpp b/Engine/Source/Runtime/Renderer/Private/SceneCaptureRendering.cpp
--- a/Engine/Source/Runtime/Renderer/Private/SceneCaptureRendering.cpp
+++ b/Engine/Source/Runtime/Renderer/Private/SceneCaptureRendering.cpp
@@ -895,25 +895,6 @@ void FScene::UpdateSceneCaptureContents(USceneCaptureComponent2D* CaptureCompone
 			ViewStateInterface->RemoveLumenSceneData(this);
 		}
 
-		if(UseVirtualShadowMaps(SceneRenderer->ShaderPlatform, FeatureLevel))
-		{
-			// A separate VSM cache is added to the SceneCapture views. 
-			// This is needed to prevent the cache being invalidated on every frame.
-			// In a future iteration, it may be more efficient to share this cache with the main viewfamily.
-			if (ViewStateInterface)
-			{
-				ViewStateInterface->AddVirtualShadowMapCache(this);
-			}
-			else
-			{
-				GEngine->AddOnScreenDebugMessage((uint64)this, 10.0f, FColor::Yellow, TEXT("SceneCapture has no viewstate. This may affect Virtual Shadow Map caching performance. Consider enabling bAlwaysPersistRenderingState or bCaptureEveryFrame."));
-			}
-		}
-		else if (ViewStateInterface)
-		{
-			ViewStateInterface->RemoveVirtualShadowMapCache(this);
-		}
-
 		// Ensure that the views for this scene capture reflect any simulated camera motion for this frame
 		TOptional<FTransform> PreviousTransform = FMotionVectorSimulation::Get().GetPreviousTransform(CaptureComponent);
 
diff --git a/Engine/Source/Runtime/Renderer/Private/ScenePrivate.h b/Engine/Source/Runtime/Renderer/Private/ScenePrivate.h
--- a/Engine/Source/Runtime/Renderer/Private/ScenePrivate.h
+++ b/Engine/Source/Runtime/Renderer/Private/ScenePrivate.h
@@ -726,9 +726,8 @@ public:
 
 	/**
 	 * The scene pointer may be NULL -- it's filled in by certain API calls that require a FSceneViewState and FScene to know about each other,
-	 * such as FSceneViewState::AddVirtualShadowMapCache.  Whenever a ViewState and Scene get linked, this pointer is set, and a pointer
-	 * to the ViewState is added to an array in the Scene.  The linking is necessary in cases where incremental FScene updates need to be
-	 * reflected in cached data stored in FSceneViewState.
+	 * Whenever a ViewState and Scene get linked, this pointer is set, and a pointer to the ViewState is added to an array in the Scene.
+	 * The linking is necessary in cases where incremental FScene updates need to be reflected in cached data stored in FSceneViewState.
 	 */
 	FScene* Scene;
 
@@ -1083,12 +1082,9 @@ public:
 
 	FGlintShadingLUTsStateData GlintShadingLUTsData;
 
-	bool bVirtualShadowMapCacheAdded;
 	bool bLumenSceneDataAdded;
 	float LumenSurfaceCacheResolution;
 
-	FVirtualShadowMapArrayCacheManager* ViewVirtualShadowMapCache;
-
 	// call after OnFrameRenderingSetup()
 	virtual uint32 GetCurrentTemporalAASampleIndex() const
 	{
@@ -1545,11 +1541,6 @@ PRAGMA_ENABLE_DEPRECATION_WARNINGS
 		return SequencerState;
 	}
 
-	virtual void AddVirtualShadowMapCache(FSceneInterface* InScene) override;
-	virtual void RemoveVirtualShadowMapCache(FSceneInterface* InScene) override;
-	virtual bool HasVirtualShadowMapCache() const override;
-	virtual FVirtualShadowMapArrayCacheManager* GetVirtualShadowMapCache(const FScene* InScene) const override;
-
 	virtual void AddLumenSceneData(FSceneInterface* InScene, float SurfaceCacheResolution) override;
 	virtual void RemoveLumenSceneData(FSceneInterface* InScene) override;
 	virtual bool HasLumenSceneData() const override;
@@ -2982,13 +2973,8 @@ public:
 	/** Preshadows that are currently cached in the PreshadowCache render target. */
 	TArray<TRefCountPtr<FProjectedShadowInfo> > CachedPreshadows;
 
-	/**
-	 * Virtual shadow maps can use a default cache stored in the FScene, or define a separate cache per view (AddVirtualShadowMapCache)
-	 * for improved performance.  A linked list of caches for the FScene is maintained to allow scene updates to be propagated to all
-	 * caches as needed.  If a cache is never used, the GPU resources for it are never allocated, so the overhead of the default cache
-	 * should be minimal when not in use.
-	 */
-	FVirtualShadowMapArrayCacheManager* DefaultVirtualShadowMapCache;
+	/**	Stores persistent virtual shadow map data */
+	FVirtualShadowMapArrayCacheManager* VirtualShadowMapCache;
 
 	/**
 	 * Stores scene-aspects needed for shadow rendering.
@@ -3230,8 +3216,7 @@ public:
 
 	virtual void RefreshNaniteRasterBins(FPrimitiveSceneInfo& PrimitiveSceneInfo) override;
 
-	FVirtualShadowMapArrayCacheManager* GetVirtualShadowMapCache(FSceneView& View) const;
-	void GetAllVirtualShadowMapCacheManagers(TArray<FVirtualShadowMapArrayCacheManager*, SceneRenderingAllocator>& OutCacheManagers) const;
+	FVirtualShadowMapArrayCacheManager* GetVirtualShadowMapCache() const { return VirtualShadowMapCache; }
 
 	FLumenSceneData* FindLumenSceneData(uint32 ViewKey, uint32 GPUIndex) const;
 	inline FLumenSceneData* GetLumenSceneData(const FViewInfo& View) const
diff --git a/Engine/Source/Runtime/Renderer/Private/ShadowSetup.cpp b/Engine/Source/Runtime/Renderer/Private/ShadowSetup.cpp
--- a/Engine/Source/Runtime/Renderer/Private/ShadowSetup.cpp
+++ b/Engine/Source/Runtime/Renderer/Private/ShadowSetup.cpp
@@ -6221,13 +6221,9 @@ void FSceneRenderer::FinishInitDynamicShadows(FRDGBuilder& GraphBuilder, FDynami
 
 	FinishGatherShadowPrimitives(TaskData);
 
-	if (!SortedShadowsForShadowDepthPass.VirtualShadowMapShadows.IsEmpty())
-	{
-		ShadowSceneRenderer->DispatchVirtualShadowMapViewAndCullingSetup(GraphBuilder, SortedShadowsForShadowDepthPass.VirtualShadowMapShadows);
-	}
-
 	if (ShadowSceneRenderer)
 	{
+		ShadowSceneRenderer->DispatchVirtualShadowMapViewAndCullingSetup(GraphBuilder, SortedShadowsForShadowDepthPass.VirtualShadowMapShadows);
 		ShadowSceneRenderer->PostSetupDebugRender();
 	}
 
diff --git a/Engine/Source/Runtime/Renderer/Private/Shadows/ShadowSceneRenderer.cpp b/Engine/Source/Runtime/Renderer/Private/Shadows/ShadowSceneRenderer.cpp
--- a/Engine/Source/Runtime/Renderer/Private/Shadows/ShadowSceneRenderer.cpp
+++ b/Engine/Source/Runtime/Renderer/Private/Shadows/ShadowSceneRenderer.cpp
@@ -118,7 +118,7 @@ FVirtualShadowMapProjectionShaderData FShadowSceneRenderer::GetLocalLightProject
 
 	int32 VirtualShadowMapId = ProjectedShadowInfo->VirtualShadowMapId + MapIndex;
 	bool bIsSinglePageSM = FVirtualShadowMapArray::IsSinglePage(VirtualShadowMapId);
-	check(VirtualShadowMapId != INDEX_NONE && CacheEntry->Current.bIsDistantLight == bIsSinglePageSM);
+	check(VirtualShadowMapId != INDEX_NONE && CacheEntry->bIsDistantLight == bIsSinglePageSM);
 
 	uint32 Flags = bIsSinglePageSM ? VSM_PROJ_FLAG_CURRENT_DISTANT_LIGHT : 0U;
 	Flags |= CacheEntry->IsUncached() ? VSM_PROJ_FLAG_UNCACHED : 0U;
@@ -336,14 +336,20 @@ void FShadowSceneRenderer::DispatchVirtualShadowMapViewAndCullingSetup(FRDGBuild
 	// Don't want to run this more than once in a given frame.
 	check(SceneInstanceCullingQuery == nullptr);
 
-	// Set up view array and collect culling work at the same time.
-	SceneInstanceCullingQuery = SceneRenderer.SceneCullingRenderer.CreateInstanceQuery(GraphBuilder);
-	VirtualShadowMapViews = VirtualShadowMapArray.CreateVirtualShadowMapNaniteViews(GraphBuilder, SceneRenderer.Views, VirtualShadowMapShadows, ComputeNaniteShadowsLODScaleFactor(), SceneInstanceCullingQuery);
+	// Unconditionally update GPU physical pages (on all GPUs) with new VSM IDs/addresses
+	VirtualShadowMapArray.UpdatePhysicalPageAddresses(GraphBuilder);
 
-	// Dispatch collected query 
-	if (SceneInstanceCullingQuery)
+	if (!VirtualShadowMapShadows.IsEmpty())
 	{
-		SceneInstanceCullingQuery->Dispatch(GraphBuilder);
+		// Set up view array and collect culling work at the same time.
+		SceneInstanceCullingQuery = SceneRenderer.SceneCullingRenderer.CreateInstanceQuery(GraphBuilder);
+		VirtualShadowMapViews = VirtualShadowMapArray.CreateVirtualShadowMapNaniteViews(GraphBuilder, SceneRenderer.Views, VirtualShadowMapShadows, ComputeNaniteShadowsLODScaleFactor(), SceneInstanceCullingQuery);
+
+		// Dispatch collected query 
+		if (SceneInstanceCullingQuery)
+		{
+			SceneInstanceCullingQuery->Dispatch(GraphBuilder);
+		}
 	}
 }
 
@@ -361,7 +367,7 @@ void FShadowSceneRenderer::PostSetupDebugRender()
 			for (const FLocalLightShadowFrameSetup& LightSetup : LocalLights)
 			{			
 				FLinearColor Color = FLinearColor(FColor::Blue);
-				if (LightSetup.PerLightCacheEntry && LightSetup.PerLightCacheEntry->Current.bIsDistantLight)
+				if (LightSetup.PerLightCacheEntry && LightSetup.PerLightCacheEntry->bIsDistantLight)
 				{
 					++NumDistant;
 					int32 FramesSinceLastRender = int32(Scene.GetFrameNumber()) - int32(LightSetup.PerLightCacheEntry->GetLastScheduledFrameNumber());
diff --git a/Engine/Source/Runtime/Renderer/Private/VirtualShadowMaps/VirtualShadowMapArray.cpp b/Engine/Source/Runtime/Renderer/Private/VirtualShadowMaps/VirtualShadowMapArray.cpp
--- a/Engine/Source/Runtime/Renderer/Private/VirtualShadowMaps/VirtualShadowMapArray.cpp
+++ b/Engine/Source/Runtime/Renderer/Private/VirtualShadowMaps/VirtualShadowMapArray.cpp
@@ -36,6 +36,7 @@ DECLARE_DWORD_COUNTER_STAT(TEXT("VSM Single Page Count"), STAT_VSMSinglePageCoun
 DECLARE_DWORD_COUNTER_STAT(TEXT("VSM Full Count"), STAT_VSMFullCount, STATGROUP_ShadowRendering);
 
 extern int32 GForceInvalidateDirectionalVSM;
+extern int32 GVSMMaxPageAgeSinceLastRequest;
 extern TAutoConsoleVariable<float> CVarNaniteMaxPixelsPerEdge;
 extern TAutoConsoleVariable<float> CVarNaniteMinPixelsPerEdgeHW;
 
@@ -66,13 +67,6 @@ TAutoConsoleVariable<int32> CVarMaxPhysicalPages(
 	ECVF_Scalability | ECVF_RenderThreadSafe
 );
 
-TAutoConsoleVariable<int32> CVarMaxPhysicalPagesSceneCapture(
-	TEXT("r.Shadow.Virtual.MaxPhysicalPagesSceneCapture"),
-	512,
-	TEXT("Maximum number of physical pages in the pool for each scene capture component."),
-	ECVF_Scalability | ECVF_RenderThreadSafe
-);
-
 TAutoConsoleVariable<int32> CVarCacheStaticSeparate(
 	TEXT("r.Shadow.Virtual.Cache.StaticSeparate"),
 	1,
@@ -214,13 +208,6 @@ static TAutoConsoleVariable<float> CVarCoarsePagePixelThresholdDynamicNanite(
 	ECVF_Scalability | ECVF_RenderThreadSafe
 );
 
-static TAutoConsoleVariable<int32> CVarCacheKeepOnlyRequestedPages(
-	TEXT("r.Shadow.Virtual.Cache.KeepOnlyRequestedPages"),
-	1,		// Temporarily
-	TEXT("When set cache pages that are not requested in a given frame will be immediately dropped.\n"),
-	ECVF_RenderThreadSafe
-);
-
 static TAutoConsoleVariable<int32> CVarCacheAllocateViaLRU(
 	TEXT("r.Shadow.Virtual.Cache.AllocateViaLRU"),
 	0,
@@ -390,16 +377,12 @@ void FVirtualShadowMapArray::UpdateNextData(int32 PrevVirtualShadowMapId, int32
 	NextData[PrevVirtualShadowMapId].PageAddressOffset = FIntVector2(PageOffset.X, PageOffset.Y);
 }
 
-void FVirtualShadowMapArray::Initialize(FRDGBuilder& GraphBuilder, FVirtualShadowMapArrayCacheManager* InCacheManager, bool bInEnabled, bool bIsSceneCapture)
+void FVirtualShadowMapArray::Initialize(FRDGBuilder& GraphBuilder, FVirtualShadowMapArrayCacheManager* InCacheManager, bool bInEnabled)
 {
 	bInitialized = true;
 	bEnabled = bInEnabled;
 	CacheManager = InCacheManager;
 
-#if WITH_MGPU
-	CacheManager->UpdateGPUMask(GraphBuilder.RHICmdList.GetGPUMask());
-#endif
-
 	bCullBackfacingPixels = CVarCullBackfacingPixels.GetValueOnRenderThread() != 0;
 	bUseHzbOcclusion = CVarShadowsVirtualUseHZB.GetValueOnRenderThread() != 0;
 	bUseTwoPassHzbOcclusion = CVarShadowsVirtualUseHZB.GetValueOnRenderThread() == 2;
@@ -416,6 +399,8 @@ void FVirtualShadowMapArray::Initialize(FRDGBuilder& GraphBuilder, FVirtualShado
 	UniformParameters.CoarsePagePixelThresholdStatic = CVarCoarsePagePixelThresholdStatic.GetValueOnRenderThread();
 	UniformParameters.CoarsePagePixelThresholdDynamicNanite = CVarCoarsePagePixelThresholdDynamicNanite.GetValueOnRenderThread();
 
+	UniformParameters.SceneFrameNumber = Scene.GetFrameNumberRenderThread();
+
 	// Reference dummy data in the UB initially
 	const uint32 DummyPageTableElement = 0xFFFFFFFF;
 	UniformParameters.PageTable = GraphBuilder.CreateSRV(GSystemTextures.GetDefaultStructuredBuffer(GraphBuilder, sizeof(DummyPageTableElement), DummyPageTableElement));
@@ -433,7 +418,7 @@ void FVirtualShadowMapArray::Initialize(FRDGBuilder& GraphBuilder, FVirtualShado
 		// NOTE: This assumes GetMax2DTextureDimension() is a power of two on supported platforms
 		const uint32 PhysicalPagesX = FMath::DivideAndRoundDown(GetMax2DTextureDimension(), FVirtualShadowMap::PageSize);
 		check(FMath::IsPowerOfTwo(PhysicalPagesX));
-		const int32 MaxPhysicalPages = (bIsSceneCapture ? CVarMaxPhysicalPagesSceneCapture : CVarMaxPhysicalPages).GetValueOnRenderThread();
+		const int32 MaxPhysicalPages = CVarMaxPhysicalPages.GetValueOnRenderThread();
 		uint32 PhysicalPagesY = FMath::DivideAndRoundUp((uint32)FMath::Max(1, MaxPhysicalPages), PhysicalPagesX);	
 
 		UniformParameters.MaxPhysicalPages = PhysicalPagesX * PhysicalPagesY;
@@ -477,7 +462,7 @@ void FVirtualShadowMapArray::Initialize(FRDGBuilder& GraphBuilder, FVirtualShado
 	}
 	else
 	{
-		CacheManager->FreePhysicalPool();
+		CacheManager->FreePhysicalPool(GraphBuilder);
 		UniformParameters.PhysicalPagePool = GSystemTextures.GetZeroUIntArrayDummy(GraphBuilder);
 	}
 
@@ -488,7 +473,7 @@ void FVirtualShadowMapArray::Initialize(FRDGBuilder& GraphBuilder, FVirtualShado
 	}
 	else
 	{
-		CacheManager->FreeHZBPhysicalPool();
+		CacheManager->FreeHZBPhysicalPool(GraphBuilder);
 		HZBPhysical = nullptr;
 		HZBPhysicalRDG = nullptr;
 	}
@@ -748,6 +733,26 @@ class FGenerateHierarchicalPageFlagsCS : public FVirtualPageManagementShader
 IMPLEMENT_GLOBAL_SHADER(FGenerateHierarchicalPageFlagsCS, "/Engine/Private/VirtualShadowMaps/VirtualShadowMapPageManagement.usf", "GenerateHierarchicalPageFlags", SF_Compute);
 
 
+class FUpdatePhysicalPageAddresses : public FVirtualPageManagementShader
+{
+	DECLARE_GLOBAL_SHADER(FUpdatePhysicalPageAddresses);
+	SHADER_USE_PARAMETER_STRUCT(FUpdatePhysicalPageAddresses, FVirtualPageManagementShader )
+
+	class FGenerateStatsDim : SHADER_PERMUTATION_BOOL("VSM_GENERATE_STATS");
+	using FPermutationDomain = TShaderPermutationDomain<FGenerateStatsDim>;
+
+	BEGIN_SHADER_PARAMETER_STRUCT( FParameters, )
+		SHADER_PARAMETER_RDG_UNIFORM_BUFFER(FVirtualShadowMapUniformParameters, VirtualShadowMap)
+		SHADER_PARAMETER_RDG_BUFFER_UAV( RWStructuredBuffer< FPhysicalPageMetaData >,	OutPhysicalPageMetaData)		
+		SHADER_PARAMETER_RDG_BUFFER_SRV( StructuredBuffer< FNextVirtualShadowMapData >,	NextVirtualShadowMapData )
+		SHADER_PARAMETER( uint32,														NextVirtualShadowMapDataCount )
+		// Required if using FGenerateStatsDim
+		SHADER_PARAMETER_RDG_BUFFER_UAV( RWStructuredBuffer< uint >,					OutStatsBuffer )
+	END_SHADER_PARAMETER_STRUCT()
+};
+IMPLEMENT_GLOBAL_SHADER(FUpdatePhysicalPageAddresses, "/Engine/Private/VirtualShadowMaps/VirtualShadowMapPhysicalPageManagement.usf", "UpdatePhysicalPageAddresses", SF_Compute );
+
+
 class FUpdatePhysicalPages : public FVirtualPageManagementShader
 {
 	DECLARE_GLOBAL_SHADER(FUpdatePhysicalPages);
@@ -766,11 +771,9 @@ class FUpdatePhysicalPages : public FVirtualPageManagementShader
 		SHADER_PARAMETER_RDG_BUFFER_UAV( RWStructuredBuffer< uint >,					OutPageFlags )
 		SHADER_PARAMETER_RDG_BUFFER_UAV( RWStructuredBuffer< uint >,					OutPageTable )
 		SHADER_PARAMETER_RDG_BUFFER_SRV( StructuredBuffer< int >,						PrevPhysicalPageLists )
-		SHADER_PARAMETER_RDG_BUFFER_SRV( StructuredBuffer< FNextVirtualShadowMapData >,	NextVirtualShadowMapData )
-		SHADER_PARAMETER( uint32,														NextVirtualShadowMapDataCount )
+		SHADER_PARAMETER( uint32,														MaxPageAgeSinceLastRequest )
 		// TODO: encode into options bitfield?
 		SHADER_PARAMETER( int32,														bDynamicPageInvalidation )
-		SHADER_PARAMETER( int32,														bKeepOnlyRequestedPages )
 		SHADER_PARAMETER( int32,														bAllocateViaLRU )
 		// Required if using FGenerateStatsDim
 		SHADER_PARAMETER_RDG_BUFFER_UAV( RWStructuredBuffer< uint >,					OutStatsBuffer )
@@ -1151,7 +1154,7 @@ void FVirtualShadowMapArray::AppendPhysicalPageList(FRDGBuilder& GraphBuilder, b
 		RDG_EVENT_NAME("AppendPhysicalPageList"),
 		ComputeShader,
 		PassParameters,
-		FIntVector(FMath::DivideAndRoundUp(GetMaxPhysicalPages(), FUpdatePhysicalPages::DefaultCSGroupX), 1, 1)
+		FIntVector(FMath::DivideAndRoundUp(GetMaxPhysicalPages(), FAppendPhysicalPageListsCS::DefaultCSGroupX), 1, 1)
 	);
 
 	FAppendPhysicalPageListsCS::FParameters* CountsParameters = GraphBuilder.AllocParameters<FAppendPhysicalPageListsCS::FParameters>();
@@ -1185,6 +1188,44 @@ void FVirtualShadowMapArray::UploadProjectionData(FRDGBuilder& GraphBuilder)
 	Uploader.ResourceUploadTo(GraphBuilder, ProjectionDataRDG);
 }
 
+void FVirtualShadowMapArray::UpdatePhysicalPageAddresses(FRDGBuilder& GraphBuilder)
+{
+	if (!IsEnabled())
+	{
+		return;
+	}
+
+	// First, let the cache manager update any that may not be referenced this frame but may still have cached pages.
+	// TODO: Store the number of active lights we have first this frame for GPU looping purposes.
+	// By construction unreferenced lights are at the end.
+	CacheManager->UpdateUnreferencedCacheEntries(*this);
+
+	// NOTE: This past MUST run on all GPUs, as we still need to propogate changes to the VSM IDs even if
+	// a given GPU may not do any rendering during this phase.
+	RDG_GPU_MASK_SCOPE(GraphBuilder, FRHIGPUMask::All());
+
+	FUpdatePhysicalPageAddresses::FParameters* PassParameters = GraphBuilder.AllocParameters<FUpdatePhysicalPageAddresses::FParameters>();
+	PassParameters->VirtualShadowMap		= GetUniformBuffer();
+	PassParameters->OutPhysicalPageMetaData = GraphBuilder.CreateUAV(PhysicalPageMetaDataRDG);
+
+	// Upload our prev -> next shadow data mapping (FNextVirtualShadowMapData) to the GPU
+	FRDGBufferRef NextVirtualShadowMapData = CreateStructuredBuffer(GraphBuilder, TEXT("Shadow.Virtual.NextVirtualShadowMapData"), NextData);
+	PassParameters->NextVirtualShadowMapData = GraphBuilder.CreateSRV(NextVirtualShadowMapData);
+	PassParameters->NextVirtualShadowMapDataCount = NextData.Num();
+
+	FUpdatePhysicalPageAddresses::FPermutationDomain PermutationVector;
+	SetStatsArgsAndPermutation<FUpdatePhysicalPageAddresses>(StatsBufferUAV, PassParameters, PermutationVector);
+	auto ComputeShader = GetGlobalShaderMap(Scene.GetFeatureLevel())->GetShader<FUpdatePhysicalPageAddresses>(PermutationVector);
+
+	FComputeShaderUtils::AddPass(
+		GraphBuilder,
+		RDG_EVENT_NAME("FVirtualShadowMapArray::UpdatePhysicalPageAddresses"),
+		ComputeShader,
+		PassParameters,
+		FIntVector(FMath::DivideAndRoundUp(GetMaxPhysicalPages(), FUpdatePhysicalPageAddresses::DefaultCSGroupX), 1, 1)
+	);
+}
+
 void FVirtualShadowMapArray::BuildPageAllocations(
 	FRDGBuilder& GraphBuilder,
 	const FMinimalSceneTextures& SceneTextures,
@@ -1197,10 +1238,6 @@ void FVirtualShadowMapArray::BuildPageAllocations(
 {
 	check(IsEnabled());
 
-	// Before we check on what shadow maps are present, let the cache manager update any that may not be referenced this frame
-	// but may still have cached pages.
-	CacheManager->UpdateUnreferencedCacheEntries(*this);
-
 	if (GetNumShadowMaps() == 0 || Views.Num() == 0)
 	{
 		// Nothing to do
@@ -1291,7 +1328,7 @@ void FVirtualShadowMapArray::BuildPageAllocations(
 		// Even though this is not a high performance path with stats enabled, we don't want to change behavior.
 		StatsBufferUAV = GraphBuilder.CreateUAV(StatsBufferRDG, ERDGUnorderedAccessViewFlags::SkipBarrier);
 	}
-	
+
 	// We potentially over-allocate these to avoid too many different allocation sizes each frame
 	const int32 NumPageFlagsToAllocate = FMath::RoundUpToPowerOfTwo(FMath::Max(128 * 1024, GetNumFullShadowMaps() * int32(FVirtualShadowMap::PageTableSize) + int32(VSM_MAX_SINGLE_PAGE_SHADOW_MAPS)));
 
@@ -1549,18 +1586,16 @@ void FVirtualShadowMapArray::BuildPageAllocations(
 			// Upload our prev -> next shadow data mapping (FNextVirtualShadowMapData) to the GPU
 			FRDGBufferRef NextVirtualShadowMapData = CreateStructuredBuffer(GraphBuilder, TEXT("Shadow.Virtual.NextVirtualShadowMapData"), NextData);
 
-			PassParameters->PageRequestFlags		 = GraphBuilder.CreateSRV(PageRequestFlagsRDG);
-			PassParameters->OutPageTable			 = GraphBuilder.CreateUAV(PageTableRDG);
-			PassParameters->OutPageFlags			 = GraphBuilder.CreateUAV(PageFlagsRDG);
-			PassParameters->PrevPhysicalPageLists    = GraphBuilder.CreateSRV(GraphBuilder.RegisterExternalBuffer(CacheManager->GetPrevBuffers().PhysicalPageLists));
-			PassParameters->NextVirtualShadowMapData = GraphBuilder.CreateSRV(NextVirtualShadowMapData);
-			PassParameters->NextVirtualShadowMapDataCount = NextData.Num();
-			PassParameters->bKeepOnlyRequestedPages  = CVarCacheKeepOnlyRequestedPages.GetValueOnRenderThread();
-			PassParameters->bDynamicPageInvalidation = 1;
+			PassParameters->PageRequestFlags		   = GraphBuilder.CreateSRV(PageRequestFlagsRDG);
+			PassParameters->OutPageTable			   = GraphBuilder.CreateUAV(PageTableRDG);
+			PassParameters->OutPageFlags			   = GraphBuilder.CreateUAV(PageFlagsRDG);
+			PassParameters->PrevPhysicalPageLists      = GraphBuilder.CreateSRV(GraphBuilder.RegisterExternalBuffer(CacheManager->GetPrevBuffers().PhysicalPageLists));
+			PassParameters->MaxPageAgeSinceLastRequest = GVSMMaxPageAgeSinceLastRequest;
+			PassParameters->bDynamicPageInvalidation   = 1;
 #if !UE_BUILD_SHIPPING
-			PassParameters->bDynamicPageInvalidation = CVarDebugSkipDynamicPageInvalidation.GetValueOnRenderThread() == 0 ? 1 : 0;
+			PassParameters->bDynamicPageInvalidation   = CVarDebugSkipDynamicPageInvalidation.GetValueOnRenderThread() == 0 ? 1 : 0;
 #endif
-			PassParameters->bAllocateViaLRU			 = CVarCacheAllocateViaLRU.GetValueOnRenderThread();
+			PassParameters->bAllocateViaLRU			   = CVarCacheAllocateViaLRU.GetValueOnRenderThread();
 		}
 
 		FUpdatePhysicalPages::FPermutationDomain PermutationVector;
@@ -1619,9 +1654,6 @@ void FVirtualShadowMapArray::BuildPageAllocations(
 		);
 	}
 
-	// Put any remaining available pages back into the sorted list for next frame
-	AppendPhysicalPageList(GraphBuilder, false);
-
 	{
 		// Run pass building hierarchical page flags to make culling acceptable performance.
 		FGenerateHierarchicalPageFlagsCS::FParameters* PassParameters = GraphBuilder.AllocParameters<FGenerateHierarchicalPageFlagsCS::FParameters>();
@@ -1734,10 +1766,11 @@ void FVirtualShadowMapArray::BuildPageAllocations(
 		);
 	}
 
-	UpdateCachedUniformBuffer(GraphBuilder);
+	// Put any remaining available pages back into the sorted list for next frame
+	// NOTE: Must do this *after* feedback status pass
+	AppendPhysicalPageList(GraphBuilder, false);
 
-	// Mark the cache data as valid if it wasn't already
-	CacheManager->MarkCacheDataValid();
+	UpdateCachedUniformBuffer(GraphBuilder);
 
 #if !UE_BUILD_SHIPPING
 	// Only dump one frame of light data
@@ -1897,6 +1930,28 @@ void FVirtualShadowMapArray::LogStats(FRDGBuilder& GraphBuilder, const FViewInfo
 }
 
 
+class FVirtualSmLogPageListStatsCS : public FVirtualPageManagementShader
+{
+	DECLARE_GLOBAL_SHADER(FVirtualSmLogPageListStatsCS);
+	SHADER_USE_PARAMETER_STRUCT(FVirtualSmLogPageListStatsCS, FVirtualPageManagementShader)
+
+	BEGIN_SHADER_PARAMETER_STRUCT(FParameters, )
+		SHADER_PARAMETER_RDG_UNIFORM_BUFFER(FVirtualShadowMapUniformParameters, VirtualShadowMap)
+		SHADER_PARAMETER_STRUCT_INCLUDE(ShaderPrint::FShaderParameters, ShaderPrintStruct)
+		SHADER_PARAMETER_RDG_BUFFER_UAV(RWStructuredBuffer< int >, OutPhysicalPageLists)
+		SHADER_PARAMETER(int, PageListStatsRow)
+	END_SHADER_PARAMETER_STRUCT()
+
+	static void ModifyCompilationEnvironment(const FGlobalShaderPermutationParameters& Parameters, FShaderCompilerEnvironment& OutEnvironment)
+	{
+		FVirtualPageManagementShader::ModifyCompilationEnvironment(Parameters, OutEnvironment);
+		// Disable optimizations as shader print causes long compile times
+		OutEnvironment.CompilerFlags.Add(CFLAG_SkipOptimizations);
+	}
+};
+IMPLEMENT_GLOBAL_SHADER(FVirtualSmLogPageListStatsCS, "/Engine/Private/VirtualShadowMaps/VirtualShadowMapPhysicalPageManagement.usf", "LogPageListStatsCS", SF_Compute);
+
+
 void FVirtualShadowMapArray::CreateMipViews( TArray<Nanite::FPackedView, SceneRenderingAllocator>& Views ) const
 {
 	TRACE_CPUPROFILER_EVENT_SCOPE(CreateMipViews);
diff --git a/Engine/Source/Runtime/Renderer/Private/VirtualShadowMaps/VirtualShadowMapArray.h b/Engine/Source/Runtime/Renderer/Private/VirtualShadowMaps/VirtualShadowMapArray.h
--- a/Engine/Source/Runtime/Renderer/Private/VirtualShadowMaps/VirtualShadowMapArray.h
+++ b/Engine/Source/Runtime/Renderer/Private/VirtualShadowMaps/VirtualShadowMapArray.h
@@ -142,6 +142,9 @@ BEGIN_GLOBAL_SHADER_PARAMETER_STRUCT(FVirtualShadowMapUniformParameters, )
 	SHADER_PARAMETER(float, CoarsePagePixelThresholdStatic)
 	SHADER_PARAMETER(float, CoarsePagePixelThresholdDynamicNanite)
 
+	// For shadow page age calculations
+	SHADER_PARAMETER(uint32, SceneFrameNumber)
+
 	SHADER_PARAMETER_RDG_BUFFER_SRV(ByteAddressBuffer, ProjectionData)
 	SHADER_PARAMETER_RDG_BUFFER_SRV(StructuredBuffer<uint>, PageTable)
 	SHADER_PARAMETER_RDG_BUFFER_SRV(StructuredBuffer<uint>, PageFlags)
@@ -208,7 +211,7 @@ public:
 	FVirtualShadowMapArray(FScene& InScene);
 	~FVirtualShadowMapArray();
 
-	void Initialize(FRDGBuilder& GraphBuilder, FVirtualShadowMapArrayCacheManager* InCacheManager, bool bInEnabled, bool bIsSceneCapture);
+	void Initialize(FRDGBuilder& GraphBuilder, FVirtualShadowMapArrayCacheManager* InCacheManager, bool bInEnabled);
 
 	// Returns true if virtual shadow maps are enabled
 	bool IsEnabled() const
@@ -268,6 +271,8 @@ public:
 
 	void MergeStaticPhysicalPages(FRDGBuilder& GraphBuilder);
 
+	void UpdatePhysicalPageAddresses(FRDGBuilder& GraphBuilder);
+
 	void BuildPageAllocations(
 		FRDGBuilder& GraphBuilder,
 		const FMinimalSceneTextures& SceneTextures,
diff --git a/Engine/Source/Runtime/Renderer/Private/VirtualShadowMaps/VirtualShadowMapCacheManager.cpp b/Engine/Source/Runtime/Renderer/Private/VirtualShadowMaps/VirtualShadowMapCacheManager.cpp
--- a/Engine/Source/Runtime/Renderer/Private/VirtualShadowMaps/VirtualShadowMapCacheManager.cpp
+++ b/Engine/Source/Runtime/Renderer/Private/VirtualShadowMaps/VirtualShadowMapCacheManager.cpp
@@ -17,7 +17,6 @@
 #define VSM_LOG_STATIC_CACHING 0
 
 CSV_DECLARE_CATEGORY_EXTERN(VSM);
-DECLARE_DWORD_COUNTER_STAT(TEXT("VSM Unreferenced Lights"), STAT_VSMUnreferencedLights, STATGROUP_ShadowRendering);
 
 static TAutoConsoleVariable<int32> CVarAccumulateStats(
 	TEXT("r.Shadow.Virtual.AccumulateStats"),
@@ -71,11 +70,12 @@ static FAutoConsoleVariableRef  CVarForceInvalidateDirectionalVSM(
 	TEXT("Forces the clipmap to always invalidate, useful to emulate a moving sun to avoid misrepresenting cache performance."),
 	ECVF_RenderThreadSafe);
 
-static TAutoConsoleVariable<int32> CVarCacheMaxUnreferencedLightAge(
-	TEXT("r.Shadow.Virtual.Cache.MaxUnreferencedLightAge"),
-	0,
-	TEXT("The number of frames to keep around cached pages from lights that are unreferenced (usually due to being offscreen or otherwise culled). 0=disabled.\n")
-	TEXT("Higher values potentially allow for more physical page reuse when cache space is available, but setting this too high can add overhead due to maintaining the extra entries."),
+// We give a little leeway here as occasionally the scene frame number is incremented multiple times between frames
+int32 GVSMMaxPageAgeSinceLastRequest = 3;
+FAutoConsoleVariableRef CVarVSMMaxPageAgeSinceLastRequest(
+	TEXT("r.Shadow.Virtual.Cache.MaxPageAgeSinceLastRequest"),
+	GVSMMaxPageAgeSinceLastRequest,
+	TEXT("The maximum number of frames to allow cached pages that aren't requested in the current frame to live. 0=disabled."),
 	ECVF_RenderThreadSafe
 );
 
@@ -192,7 +192,6 @@ void FVirtualShadowMapCacheEntry::SetHZBViewParams(Nanite::FPackedViewParams& Ou
 
 void FVirtualShadowMapPerLightCacheEntry::UpdateClipmap(const FVector& LightDirection, int FirstLevel)
 {
-	Prev.bIsUncached = Current.bIsUncached;
 	Prev.RenderedFrameNumber = FMath::Max(Prev.RenderedFrameNumber, Current.RenderedFrameNumber);
 	Current.RenderedFrameNumber = -1;
 
@@ -205,19 +204,20 @@ void FVirtualShadowMapPerLightCacheEntry::UpdateClipmap(const FVector& LightDire
 	ClipmapCacheKey.LightDirection = LightDirection;
 	ClipmapCacheKey.FirstLevel = FirstLevel;
 
-	Current.bIsUncached = GForceInvalidateDirectionalVSM != 0 || Prev.RenderedFrameNumber < 0;
+	bool bNewIsUncached = GForceInvalidateDirectionalVSM != 0 || Prev.RenderedFrameNumber < 0;
 
 	// On transition between uncached <-> cached we must invalidate since the static pages may not be initialized
-	if (Current.bIsUncached != Prev.bIsUncached)
+	if (bNewIsUncached != bIsUncached)
 	{
 		Prev.RenderedFrameNumber = -1;
 	}
+	bIsUncached = bNewIsUncached;
 }
 
-bool FVirtualShadowMapPerLightCacheEntry::UpdateLocal(const FProjectedShadowInitializer& InCacheKey, bool bIsDistantLight, bool bCacheEnabled, bool bAllowInvalidation)
+bool FVirtualShadowMapPerLightCacheEntry::UpdateLocal(const FProjectedShadowInitializer& InCacheKey, bool bNewIsDistantLight, bool bCacheEnabled, bool bAllowInvalidation)
 {
-	Prev.bIsUncached = Current.bIsUncached;
-	Prev.bIsDistantLight = Current.bIsDistantLight;
+	// TODO: The logic in this function is needlessly convoluted... clean up
+
 	Prev.RenderedFrameNumber = FMath::Max(Prev.RenderedFrameNumber, Current.RenderedFrameNumber);
 	Prev.ScheduledFrameNumber = FMath::Max(Prev.ScheduledFrameNumber, Current.ScheduledFrameNumber);
 
@@ -232,17 +232,24 @@ bool FVirtualShadowMapPerLightCacheEntry::UpdateLocal(const FProjectedShadowInit
 	}
 	LocalCacheKey = InCacheKey;
 
-	Current.bIsDistantLight = bIsDistantLight;
-	Current.RenderedFrameNumber = -1;
-	Current.ScheduledFrameNumber = -1;
-	Current.bIsUncached = Prev.RenderedFrameNumber < 0;
-
 	// On transition between uncached <-> cached we must invalidate since the static pages may not be initialized
-	if (Current.bIsUncached != Prev.bIsUncached)
+	bool bNewIsUncached = Prev.RenderedFrameNumber < 0;
+	if (bNewIsUncached != bIsUncached)
+	{
+		Prev.RenderedFrameNumber = -1;
+	}
+
+	// On transition between distant <-> regular we must invalidate
+	if (bNewIsDistantLight != bIsDistantLight)
 	{
 		Prev.RenderedFrameNumber = -1;
 	}
 
+	Current.RenderedFrameNumber = -1;
+	Current.ScheduledFrameNumber = -1;
+	bIsDistantLight = bNewIsDistantLight;
+	bIsUncached = bNewIsUncached;
+
 	return Prev.RenderedFrameNumber >= 0;
 }
 
@@ -673,16 +680,16 @@ void FVirtualShadowMapArrayCacheManager::SetPhysicalPoolSize(FRDGBuilder& GraphB
 		// Persistent, so we extract it immediately
 		PhysicalPageMetaData = GraphBuilder.ConvertToExternalBuffer(PhysicalPageMetaDataRDG);
 
-		Invalidate();
+		Invalidate(GraphBuilder);
 	}
 }
 
-void FVirtualShadowMapArrayCacheManager::FreePhysicalPool()
+void FVirtualShadowMapArrayCacheManager::FreePhysicalPool(FRDGBuilder& GraphBuilder)
 {
 	if (PhysicalPagePool)
 	{
 		PhysicalPagePool = nullptr;
-		Invalidate();
+		Invalidate(GraphBuilder);
 	}
 }
 
@@ -708,29 +715,29 @@ TRefCountPtr<IPooledRenderTarget> FVirtualShadowMapArrayCacheManager::SetHZBPhys
 	return HZBPhysicalPagePool;
 }
 
-void FVirtualShadowMapArrayCacheManager::FreeHZBPhysicalPool()
+void FVirtualShadowMapArrayCacheManager::FreeHZBPhysicalPool(FRDGBuilder& GraphBuilder)
 {
 	if (HZBPhysicalPagePool)
 	{
 		HZBPhysicalPagePool = nullptr;
-		Invalidate();
+		Invalidate(GraphBuilder);
 	}
 }
 
-void FVirtualShadowMapArrayCacheManager::Invalidate()
+void FVirtualShadowMapArrayCacheManager::Invalidate(FRDGBuilder& GraphBuilder)
 {
 	// Clear the cache
 	CacheEntries.Reset();
 
-	// Mark globally invalid until the next GPU allocation/metadata update
-	bCacheDataValid = false;
-
 	//UE_LOG(LogRenderer, Display, TEXT("Virtual shadow map cache invalidated."));
-}
 
-void FVirtualShadowMapArrayCacheManager::MarkCacheDataValid()
-{
-	bCacheDataValid = true;
+	// Clear the physical page metadata (on all GPUs)
+	if (PhysicalPageMetaData)
+	{
+		RDG_GPU_MASK_SCOPE(GraphBuilder, FRHIGPUMask::All());
+		FRDGBufferRef PhysicalPageMetaDataRDG = GraphBuilder.RegisterExternalBuffer(PhysicalPageMetaData);
+		AddClearUAVPass(GraphBuilder, GraphBuilder.CreateUAV(PhysicalPageMetaDataRDG), 0);
+	}
 }
 
 bool FVirtualShadowMapArrayCacheManager::IsCacheEnabled()
@@ -741,7 +748,6 @@ bool FVirtualShadowMapArrayCacheManager::IsCacheEnabled()
 bool FVirtualShadowMapArrayCacheManager::IsCacheDataAvailable()
 {
 	return IsCacheEnabled() &&
-		bCacheDataValid &&
 		PrevBuffers.PageTable &&
 		PrevBuffers.PageFlags &&
 		PrevBuffers.PageRectBounds &&
@@ -768,7 +774,8 @@ TSharedPtr<FVirtualShadowMapPerLightCacheEntry> FVirtualShadowMapArrayCacheManag
 
 		if (LightEntry->ShadowMapEntries.Num() == NumShadowMaps)
 		{
-			LightEntry->ReferencedRenderSequenceNumber = RenderSequenceNumber;
+			LightEntry->bReferencedThisRender = true;
+			LightEntry->LastReferencedFrameNumber = Scene->GetFrameNumberRenderThread();
 			return LightEntry;
 		}
 		else
@@ -782,7 +789,8 @@ TSharedPtr<FVirtualShadowMapPerLightCacheEntry> FVirtualShadowMapArrayCacheManag
 
 	// Make new entry for this light
 	TSharedPtr<FVirtualShadowMapPerLightCacheEntry> LightEntry = MakeShared<FVirtualShadowMapPerLightCacheEntry>(Scene->GetMaxPersistentPrimitiveIndex(), NumShadowMaps);
-	LightEntry->ReferencedRenderSequenceNumber = RenderSequenceNumber;	
+	LightEntry->bReferencedThisRender = true;
+	LightEntry->LastReferencedFrameNumber = Scene->GetFrameNumberRenderThread();
 	CacheEntries.Add(CacheKey, LightEntry);
 
 	return LightEntry;
@@ -810,28 +818,21 @@ void FVirtualShadowMapPerLightCacheEntry::OnPrimitiveRendered(const FPrimitiveSc
 void FVirtualShadowMapArrayCacheManager::UpdateUnreferencedCacheEntries(
 	FVirtualShadowMapArray& VirtualShadowMapArray)
 {
-	const int64 MaxAge = CVarCacheMaxUnreferencedLightAge.GetValueOnRenderThread();
-	const bool bAllowUnreferencedEntries = IsCacheEnabled() && MaxAge > 0;
+	const uint32 SceneFrameNumber = Scene->GetFrameNumberRenderThread();
 
 	TArray<uint64, SceneRenderingAllocator> EntriesToRemove;
 	for (auto& LightEntry : CacheEntries)
 	{
 		TSharedPtr<FVirtualShadowMapPerLightCacheEntry> CacheEntry = LightEntry.Value;
-		// NOTE: We probably want to decouple the age from render calls at some point, but regardless we need to know
-		// in the first branch that the given light was referenced *this* render, and therefore has already been
-		// (re)allocated a current VSM ID.
-		int64 Age = RenderSequenceNumber - CacheEntry->ReferencedRenderSequenceNumber;
-		check(CacheEntry->ReferencedRenderSequenceNumber >= 0);
-		check(Age >= 0);
-		if (Age == 0)
+		// For this test we care if it is active *this render*, not just this scene frame number (which can include multiple renders)
+		if (CacheEntry->bReferencedThisRender)
 		{
-			// Referenced this frame still, leave it alone
+			// Active this render, leave it alone
 			check(CacheEntry->ShadowMapEntries.Last().CurrentVirtualShadowMapId < VirtualShadowMapArray.GetNumShadowMapSlots());
 		}
-		else if (bAllowUnreferencedEntries && Age < MaxAge)
+		else if (int32(SceneFrameNumber - CacheEntry->LastReferencedFrameNumber) <= GVSMMaxPageAgeSinceLastRequest)
 		{
-			INC_DWORD_STAT(STAT_VSMUnreferencedLights);
-
+			// Not active this render, but still recent enough to keep it and its pages alive
 			int PrevBaseVirtualShadowMapId = CacheEntry->ShadowMapEntries[0].CurrentVirtualShadowMapId;
 			bool bIsSinglePage = FVirtualShadowMapArray::IsSinglePage(PrevBaseVirtualShadowMapId);
 
@@ -841,9 +842,10 @@ void FVirtualShadowMapArrayCacheManager::UpdateUnreferencedCacheEntries(
 			for (int32 Map = 0; Map < NumMaps; ++Map)
 			{
 				CacheEntry->ShadowMapEntries[Map].Update(VirtualShadowMapArray, *CacheEntry, VirtualShadowMapId + Map);
-				// NOTE: Leave the ProjectionData as whatever it was before
-				// TODO: We may want to add a flag that this is unreferenced so we can prune it from the light grid and skip it in page marking, etc...?
-				// Except in theory if we are marking things from onscreen pixels then we wouldn't have culled it... (small light culling though)?
+				// Mark it as inactive for this frame/render
+				// NOTE: We currently recompute/overwrite the whole ProjectionData structure for referenced lights, but if that changes we
+				// will need to clear this flag again when they become referenced.
+				CacheEntry->ShadowMapEntries[Map].ProjectionData.Flags |= VSM_PROJ_FLAG_UNREFERENCED;
 			}
 		}
 		else
@@ -928,7 +930,7 @@ void FVirtualShadowMapArrayCacheManager::ExtractFrameData(
 		// We drop the physical page pool here as well to ensure that it disappears in the case where
 		// thumbnail rendering or similar creates multiple FSceneRenderers that never get deleted.
 		// Caching is disabled on these contexts intentionally to avoid these issues.
-		FreePhysicalPool();
+		FreePhysicalPool(GraphBuilder);
 	}
 	else if (bNewShadowData)
 	{
@@ -965,14 +967,11 @@ void FVirtualShadowMapArrayCacheManager::ExtractFrameData(
 
 		ExtractStats(GraphBuilder, VirtualShadowMapArray);
 	}
-	else
+	
+	// Clear out the referenced light flags since this render is finishing
+	for (auto& LightEntry : CacheEntries)
 	{
-		// Do nothing; maintain the data that we had
-		// This allows us to work around some cases where the renderer gets called multiple times in a given frame
-		// - such as scene captures - but does no shadow-related work in all but one of them. We do not want to drop
-		// all the cached data in this case otherwise we effectively get no caching at all.
-		// Ideally in the long run we want the cache itself to be more robust against rendering multiple views. but
-		// for now this at least provides a work-around for some common cases where only one view is rendering VSMs.
+		LightEntry.Value->bReferencedThisRender = false;
 	}
 
 	// Every once in a while zero out our recently removed primitive flags. This lets us ignore slots that
@@ -984,8 +983,6 @@ void FVirtualShadowMapArrayCacheManager::ExtractFrameData(
 		RecentlyRemovedReadIndex = 1 - RecentlyRemovedReadIndex;
 		RecentlyRemovedFrameCounter = 0;
 	}
-
-	++RenderSequenceNumber;
 }
 
 void FVirtualShadowMapArrayCacheManager::ExtractStats(FRDGBuilder& GraphBuilder, FVirtualShadowMapArray &VirtualShadowMapArray)
@@ -1235,17 +1232,6 @@ TRDGUniformBufferRef<FVirtualShadowMapUniformParameters> FVirtualShadowMapArrayC
 	return GraphBuilder.CreateUniformBuffer(VersionedParameters);
 }
 
-#if WITH_MGPU
-void FVirtualShadowMapArrayCacheManager::UpdateGPUMask(FRHIGPUMask GPUMask)
-{
-	if (LastGPUMask != GPUMask)
-	{
-		LastGPUMask = GPUMask;
-		Invalidate();
-	}
-}
-#endif  // WITH_MGPU
-
 
 void FVirtualShadowMapArrayCacheManager::ProcessInvalidations(
 	FRDGBuilder& GraphBuilder,
diff --git a/Engine/Source/Runtime/Renderer/Private/VirtualShadowMaps/VirtualShadowMapCacheManager.h b/Engine/Source/Runtime/Renderer/Private/VirtualShadowMaps/VirtualShadowMapCacheManager.h
--- a/Engine/Source/Runtime/Renderer/Private/VirtualShadowMaps/VirtualShadowMapCacheManager.h
+++ b/Engine/Source/Runtime/Renderer/Private/VirtualShadowMaps/VirtualShadowMapCacheManager.h
@@ -20,7 +20,7 @@ namespace Nanite { struct FPackedViewParams; }
 class FVirtualShadowMapCacheEntry
 {
 public:
-	// Generic version used for local lights but also unreferenced lights
+	// Generic version used for local lights but also inactive lights
 	// Updates the VSM ID
 	void Update(
 		FVirtualShadowMapArray& VirtualShadowMapArray,
@@ -46,7 +46,7 @@ public:
 	int32 CurrentVirtualShadowMapId = INDEX_NONE;
 	FVirtualShadowMapHZBMetadata CurrentHZBMetadata;
 
-	// Stores the projection shader data. This is needed for cached entries that may be unreferenced on the current frame
+	// Stores the projection shader data. This is needed for cached entries that may be inactive in the current frame/render
 	// and also avoids recomputing it every frame.
 	FVirtualShadowMapProjectionShaderData ProjectionData;
 
@@ -72,14 +72,14 @@ public:
 
 	void OnPrimitiveRendered(const FPrimitiveSceneInfo* PrimitiveSceneInfo);
 	/**
-	 * The (local) VSM is fully cached if the previous frame if was distant and is distant this frame also.
+	 * The (local) VSM is fully cached if it is distant and has been rendered to previously
 	 * "Fully" implies that we know all pages are mapped as well as rendered to (ignoring potential CPU-side object culling).
 	 */
-	inline bool IsFullyCached() const { return Current.bIsDistantLight && Prev.bIsDistantLight && Prev.RenderedFrameNumber >= 0; }
+	inline bool IsFullyCached() const { return bIsDistantLight && Prev.RenderedFrameNumber >= 0; }
 
 	/**
 	 */
-	inline bool IsUncached() const { return Current.bIsUncached; }
+	inline bool IsUncached() const { return bIsUncached; }
 
 	void MarkRendered(int32 FrameIndex) { Current.RenderedFrameNumber = FrameIndex; }
 	int32 GetLastScheduledFrameNumber() const { return Prev.ScheduledFrameNumber; }
@@ -87,7 +87,7 @@ public:
 	/**
 	 * Returns true if the cache entry is valid (has previous state).
 	 */
-	bool UpdateLocal(const FProjectedShadowInitializer &InCacheKey, bool bIsDistantLight, bool bCacheEnabled, bool bAllowInvalidation);
+	bool UpdateLocal(const FProjectedShadowInitializer &InCacheKey, bool bNewIsDistantLight, bool bCacheEnabled, bool bAllowInvalidation);
 
 	/**
 	 * Mark as invalid, i.e., needing rendering.
@@ -97,20 +97,22 @@ public:
 	// TODO: We probably don't need the prev/next thing anymore
 	struct FFrameState
 	{
-		bool bIsUncached = false;
-		bool bIsDistantLight = false;
 		int32 RenderedFrameNumber = -1;
-		int32 ScheduledFrameNumber = -1;		
-
+		int32 ScheduledFrameNumber = -1;
 	};
 	FFrameState Prev;
 	FFrameState Current;
 
-	// Tracks the last time we saw this light referenced.
-	// This allows us to keep entries alive even when they are temporarily not visible or similar, but age them out over time
-	// NOTE: This is not exactly a frame number since multiple "Render" calls may happen per frame, and we need to track
-	// them separately for the purposes of determining whether a given light is actually referenced in the current Render.
-	int64 ReferencedRenderSequenceNumber = -1;
+	bool bIsUncached = false;
+	bool bIsDistantLight = false;
+
+	// Tracks if this cache entry is being used "this render", i.e. "active". Note that there may be multiple renders per frame in the case of
+	// scene captures or similar, so unlike the RenderedFrameNumber we don't use the scene frame number, but instead mark this
+	// when a light is set up, and clear it when extracting frame data.
+	bool bReferencedThisRender = false;
+
+	// This tracks the last "rendered frame" the light was active
+	uint32 LastReferencedFrameNumber = 0;
 
 	// Primitives that have been rendered (not culled) the previous frame, when a primitive transitions from being culled to not it must be rendered into the VSM
 	// Key culling reasons are small size or distance cutoff.
@@ -180,7 +182,7 @@ struct FVirtualShadowMapArrayFrameData
 struct FPhysicalPageMetaData
 {	
 	uint32 Flags;
-	uint32 Age;
+	uint32 LastRequestedSceneFrameNumber;
 	uint32 VirtualShadowMapId;
 	uint32 MipLevel;
 	FUintPoint PageAddress;
@@ -198,20 +200,13 @@ public:
 	// Called by VirtualShadowMapArray to potentially resize the physical pool
 	// If the requested size is not already the size, all cache data is dropped and the pool is resized.
 	void SetPhysicalPoolSize(FRDGBuilder& GraphBuilder, FIntPoint RequestedSize, int RequestedArraySize, uint32 MaxPhysicalPages);
-	void FreePhysicalPool();
+	void FreePhysicalPool(FRDGBuilder& GraphBuilder);
 	TRefCountPtr<IPooledRenderTarget> GetPhysicalPagePool() const { return PhysicalPagePool; }
 	TRefCountPtr<FRDGPooledBuffer> GetPhysicalPageMetaData() const { return PhysicalPageMetaData; }
 
 	// Called by VirtualShadowMapArray to potentially resize the HZB physical pool
 	TRefCountPtr<IPooledRenderTarget> SetHZBPhysicalPoolSize(FRDGBuilder& GraphBuilder, FIntPoint RequestedSize, const EPixelFormat Format);
-	void FreeHZBPhysicalPool();
-
-	// Set the cache as valid (called after allocation/analysis/metadata update)
-	// NOTE: Could be called after rendering instead, but the distinction is not currently meaningful
-	void MarkCacheDataValid();
-	
-	// Invalidate the cache for all shadows, causing any pages to be rerendered
-	void Invalidate();
+	void FreeHZBPhysicalPool(FRDGBuilder& GraphBuilder);
 
 	/**
 	 * Called before VSM builds page allocations to reallocate any lights that may not be visible this frame
@@ -313,10 +308,6 @@ public:
 	const FVirtualShadowMapUniformParameters& GetPreviousUniformParameters() const { return PrevUniformParameters; }
 	TRDGUniformBufferRef<FVirtualShadowMapUniformParameters> GetPreviousUniformBuffer(FRDGBuilder& GraphBuilder) const;
 
-#if WITH_MGPU
-	void UpdateGPUMask(FRHIGPUMask GPUMask);
-#endif
-
 	uint64 GetGPUSizeBytes(bool bLogSizes) const;
 
 	const FVirtualShadowMapArrayFrameData& GetPrevBuffers() const { return PrevBuffers; }
@@ -328,6 +319,9 @@ public:
 #endif
 
 private:
+	// Invalidate the cache for all shadows, causing any pages to be rerendered
+	void Invalidate(FRDGBuilder& GraphBuilder);
+
 	void ProcessInvalidations(FRDGBuilder& GraphBuilder, FSceneUniformBuffer &SceneUniformBuffer, const FInstanceGPULoadBalancer& Instances) const;
 
 	void ExtractStats(FRDGBuilder& GraphBuilder, FVirtualShadowMapArray &VirtualShadowMapArray);
@@ -363,10 +357,6 @@ private:
 	// Index the Cache entries by the light ID
 	TMap< uint64, TSharedPtr<FVirtualShadowMapPerLightCacheEntry> > CacheEntries;
 
-	// Marked after successfully allocating new physical pages
-	// Cleared if any global invalidation happens, in which case the next VSM update will not consider cached pages
-	bool bCacheDataValid = false;
-
 	// Tracks primitives (by persistent primitive index) that have been removed recently
 	// This allows us to ignore feedback from previous frames in the case of persistent primitive indices being
 	// reused after being removed. We mark bits in two bitfields, then zero out one of them and switch
@@ -376,9 +366,6 @@ private:
 	int32 RecentlyRemovedReadIndex = 0;
 	int32 RecentlyRemovedFrameCounter = 0;
 
-	// Used for tracking which cache entries are referenced in a given render call
-	int64 RenderSequenceNumber = 0;
-
 	// Stores stats over frames when activated.
 	TRefCountPtr<FRDGPooledBuffer> AccumulatedStatsBuffer;
 	bool bAccumulatingStats = false;
@@ -408,9 +395,5 @@ private:
 	TMap<uint32, FLargePageAreaItem> LargePageAreaItems;
 #endif // UE_BUILD_SHIPPING
 
-#if WITH_MGPU
-	FRHIGPUMask LastGPUMask;
-#endif
-
 	FScene* Scene;
 };
