Added support for SpecularScale and SoftSourceRadius on forward directional lights.

#rnx

diff --git a/Engine/Shaders/Private/LightData.ush b/Engine/Shaders/Private/LightData.ush
--- a/Engine/Shaders/Private/LightData.ush
+++ b/Engine/Shaders/Private/LightData.ush
@@ -65,6 +65,8 @@ struct FDirectionalLightData
 	float3 DirectionalLightColor;
 	float3 DirectionalLightDirection;
 	float  DirectionalLightSourceRadius;
+	float  DirectionalLightSoftSourceRadius;
+	float  DirectionalLightSpecularScale;
 };
 
 ///////////////////////////////////////////////////////////////////////////////////////////////////
@@ -411,10 +413,11 @@ FDeferredLightData ConvertToDeferredLight(
 	Out.Direction			= In.DirectionalLightDirection;
 	Out.DistanceFadeMAD		= In.DirectionalLightDistanceFadeMAD;
 	Out.bRadialLight		= false;
-	Out.SpecularScale		= InSpecularScale;
+	Out.SpecularScale		= In.DirectionalLightSpecularScale * InSpecularScale;
 	Out.ShadowedBits		= (In.DirectionalLightShadowMapChannelMask & 0xFF) != 0 ? 1 : 0;
 	Out.HairTransmittance	= InitHairTransmittanceData();
 	Out.SourceRadius		= In.DirectionalLightSourceRadius;
+	Out.SoftSourceRadius	= In.DirectionalLightSoftSourceRadius;
 
 	// Shadow
 	// * Static shadowing uses ShadowMapChannel, 
diff --git a/Engine/Shaders/Private/LightGridCommon.ush b/Engine/Shaders/Private/LightGridCommon.ush
--- a/Engine/Shaders/Private/LightGridCommon.ush
+++ b/Engine/Shaders/Private/LightGridCommon.ush
@@ -192,6 +192,8 @@ FDirectionalLightData GetDirectionalLightData(uint EyeIndex)
 		Result.DirectionalLightColor				= ForwardLightData.DirectionalLightColor;
 		Result.DirectionalLightDirection			= ForwardLightData.DirectionalLightDirection;
 		Result.DirectionalLightSourceRadius			= ForwardLightData.DirectionalLightSourceRadius;
+		Result.DirectionalLightSoftSourceRadius		= ForwardLightData.DirectionalLightSoftSourceRadius;
+		Result.DirectionalLightSpecularScale		= ForwardLightData.DirectionalLightSpecularScale;
 
 #if (INSTANCED_STEREO || MOBILE_MULTI_VIEW)
 	}
@@ -203,6 +205,8 @@ FDirectionalLightData GetDirectionalLightData(uint EyeIndex)
 		Result.DirectionalLightColor				= ForwardLightDataISR.DirectionalLightColor;
 		Result.DirectionalLightDirection			= ForwardLightDataISR.DirectionalLightDirection;
 		Result.DirectionalLightSourceRadius			= ForwardLightDataISR.DirectionalLightSourceRadius;
+		Result.DirectionalLightSoftSourceRadius		= ForwardLightDataISR.DirectionalLightSoftSourceRadius;
+		Result.DirectionalLightSpecularScale		= ForwardLightDataISR.DirectionalLightSpecularScale;
 	}
 #endif
 
diff --git a/Engine/Shaders/Private/Strata/StrataVisualize.usf b/Engine/Shaders/Private/Strata/StrataVisualize.usf
--- a/Engine/Shaders/Private/Strata/StrataVisualize.usf
+++ b/Engine/Shaders/Private/Strata/StrataVisualize.usf
@@ -1211,7 +1211,7 @@ void EvaluateLighting(
 		DirLightData.Direction = ForwardLightData.DirectionalLightDirection;
 		DirLightData.DistanceFadeMAD = ForwardLightData.DirectionalLightDistanceFadeMAD;
 		DirLightData.bRadialLight = false;
-		DirLightData.SpecularScale = 1;
+		DirLightData.SpecularScale = ForwardLightData.DirectionalLightSpecularScale;
 		float4 DirLightAttenuation = float4(1, 1, 1, 1);
 		DirLightData.ShadowedBits = 1;
 		DirLightData.ShadowMapChannelMask.x = 1;
diff --git a/Engine/Source/Runtime/Renderer/Private/LightGridInjection.cpp b/Engine/Source/Runtime/Renderer/Private/LightGridInjection.cpp
--- a/Engine/Source/Runtime/Renderer/Private/LightGridInjection.cpp
+++ b/Engine/Source/Runtime/Renderer/Private/LightGridInjection.cpp
@@ -614,8 +614,10 @@ FComputeLightGridOutput FSceneRenderer::ComputeLightGrid(FRDGBuilder& GraphBuild
 							ForwardLightData->HasDirectionalLight = 1;
 							ForwardLightData->DirectionalLightColor = FVector3f(LightParameters.Color);
 							ForwardLightData->DirectionalLightVolumetricScatteringIntensity = LightProxy->GetVolumetricScatteringIntensity();
+							ForwardLightData->DirectionalLightSpecularScale = LightProxy->GetSpecularScale();
 							ForwardLightData->DirectionalLightDirection = LightParameters.Direction;
 							ForwardLightData->DirectionalLightSourceRadius = LightParameters.SourceRadius;
+							ForwardLightData->DirectionalLightSoftSourceRadius = LightParameters.SoftSourceRadius;
 							ForwardLightData->DirectionalLightShadowMapChannelMask = LightTypeAndShadowMapChannelMaskPacked;
 							ForwardLightData->DirectionalLightVSM = INDEX_NONE;
 							ForwardLightData->DirectionalLightSMRTSettings = GetVirtualShadowMapSMRTSettings(true);
diff --git a/Engine/Source/Runtime/Renderer/Private/SceneRendering.h b/Engine/Source/Runtime/Renderer/Private/SceneRendering.h
--- a/Engine/Source/Runtime/Renderer/Private/SceneRendering.h
+++ b/Engine/Source/Runtime/Renderer/Private/SceneRendering.h
@@ -787,8 +787,10 @@ BEGIN_GLOBAL_SHADER_PARAMETER_STRUCT_WITH_CONSTRUCTOR(FForwardLightData, )
 	SHADER_PARAMETER(FVector3f, LightGridZParams)
 	SHADER_PARAMETER(FVector3f, DirectionalLightDirection)
 	SHADER_PARAMETER(float, DirectionalLightSourceRadius)
+	SHADER_PARAMETER(float, DirectionalLightSoftSourceRadius)
 	SHADER_PARAMETER(FVector3f, DirectionalLightColor)
 	SHADER_PARAMETER(float, DirectionalLightVolumetricScatteringIntensity)
+	SHADER_PARAMETER(float, DirectionalLightSpecularScale)
 	SHADER_PARAMETER(uint32, DirectionalLightShadowMapChannelMask)
 	SHADER_PARAMETER(FVector2f, DirectionalLightDistanceFadeMAD)
 	SHADER_PARAMETER(uint32, NumDirectionalLightCascades)
