Fix missing specular lighting with rect light when slab having anistropy.

Proper GGX anisotropy lighting with rect light is not supported at the moment. In such a case we fallback onto isotropic case.

diff --git a/Engine/Shaders/Private/Strata/StrataEvaluation.ush b/Engine/Shaders/Private/Strata/StrataEvaluation.ush
--- a/Engine/Shaders/Private/Strata/StrataEvaluation.ush
+++ b/Engine/Shaders/Private/Strata/StrataEvaluation.ush
@@ -313,11 +313,13 @@ FStrataEvaluateResult StrataEvaluateBSDFCommon(FStrataBSDFContext BSDFContext, F
 	{
 		case STRATA_BSDF_TYPE_SLAB:
 		{
+			const bool bIsRectLight 	= IntegrationType == INTEGRATION_AREA_LIGHT_RECT;
+
 			float3 DiffuseColor			= SLAB_DIFFUSEALBEDO(BSDFContext.BSDF);
 			float3 F0					= SLAB_F0(BSDFContext.BSDF);
 			float3 F90					= SLAB_F90(BSDFContext.BSDF);
 			const float SafeRoughness	= MakeRoughnessSafe(SLAB_ROUGHNESS(BSDFContext.BSDF));
-			const bool bHasAnisotropy	= BSDF_GETHASANISOTROPY(BSDFContext.BSDF);
+			const bool bHasAnisotropy	= BSDF_GETHASANISOTROPY(BSDFContext.BSDF) && !bIsRectLight;
 			const bool bHaziness		= BSDF_GETHASHAZINESS(BSDFContext.BSDF);
 
 			// Specular occlusion is only used here once to affect the F90 source parameters. F0 will decrease naturally to 0.
@@ -386,8 +388,6 @@ FStrataEvaluateResult StrataEvaluateBSDFCommon(FStrataBSDFContext BSDFContext, F
 			//// Note: anisotropy completely disables area integrations. Lights fall back to punctual.
 			////
 
-			const bool bIsRectLight = IntegrationType == INTEGRATION_AREA_LIGHT_RECT;
-
 			// Primary highlight
 			float PDF = 0;
 			float DirectionalAlbedo_SpecularTransmission = 1.0f;
