Silenced shader compiler warning

#rnx

diff --git a/Engine/Shaders/Private/Common.ush b/Engine/Shaders/Private/Common.ush
--- a/Engine/Shaders/Private/Common.ush
+++ b/Engine/Shaders/Private/Common.ush
@@ -661,7 +661,9 @@ MaterialFloat4 ClampedPow(MaterialFloat4 X,MaterialFloat4 Y)
 // This is done to prevent floating point specials when compilers expands pow into exp(Exponent * log(Base)).
 MaterialFloat PositiveClampedPow(MaterialFloat Base, MaterialFloat Exponent)
 {
+#pragma warning (disable:3571) // pow with a negative base is explicitly handled here
 	return (Base <= 2.980233e-8f) ? 0.0f : pow(Base, Exponent);
+#pragma warning (enable:3571)
 }
 MaterialFloat2 PositiveClampedPow(MaterialFloat2 Base, MaterialFloat2 Exponent)
 {
