Fixing alpha-compositing issues with Heterogeneous Volumes

diff --git a/Engine/Shaders/Private/HeterogeneousVolumes/HeterogeneousVolumesComposite.usf b/Engine/Shaders/Private/HeterogeneousVolumes/HeterogeneousVolumesComposite.usf
--- a/Engine/Shaders/Private/HeterogeneousVolumes/HeterogeneousVolumesComposite.usf
+++ b/Engine/Shaders/Private/HeterogeneousVolumes/HeterogeneousVolumesComposite.usf
@@ -18,9 +18,7 @@ void HeterogeneousVolumesCompositeCS(
 	float4 Result;
 	float HeterogeneousVolumeTransmittance = saturate(1.0 - HeterogeneousVolumeRadiance[PixelCoord].a);
 	Result.rgb = HeterogeneousVolumeRadiance[PixelCoord].rgb + RWColorTexture[PixelCoord].rgb * HeterogeneousVolumeTransmittance;
-
-	float AccumulatedTransmittance = (1.0 - RWColorTexture[PixelCoord].a) * HeterogeneousVolumeTransmittance;
-	Result.a = 1.0 - AccumulatedTransmittance;
+	Result.a = RWColorTexture[PixelCoord].a * HeterogeneousVolumeTransmittance;
 
 	RWColorTexture[PixelCoord] = Result;
 }
