Re-arrange instance data on mobile to avoid S10 Adreno GL crash on CS linking

diff --git a/Engine/Shaders/Private/SceneDataMobileLoader.ush b/Engine/Shaders/Private/SceneDataMobileLoader.ush
--- a/Engine/Shaders/Private/SceneDataMobileLoader.ush
+++ b/Engine/Shaders/Private/SceneDataMobileLoader.ush
@@ -20,18 +20,24 @@ FPrimitiveSceneData LoadPrimitiveDataUBO(uint DrawInstanceId)
 	FPrimitiveSceneData PrimitiveData = (FPrimitiveSceneData)0;
 
 	uint DataIdx = 0;
+	// TilePosition, Flags
+	{
+		float4 Data = LoadPrimitiveDataElementUBO(DrawInstanceId, DataIdx + 0);
+		PrimitiveData.TilePosition = Data.xyz;
+		PrimitiveData.Flags = asuint(Data.w);
+		DataIdx += 1u;
+	}
 
-	// LocalToWorld, TilePosition, Primitive.Flags
+	// LocalToWorld
 	{
-		float4 LocalToRelativeWorld0 = LoadPrimitiveDataElementUBO(DrawInstanceId, DataIdx + 0);
-		float4 LocalToRelativeWorld1 = LoadPrimitiveDataElementUBO(DrawInstanceId, DataIdx + 1);
-		float4 LocalToRelativeWorld2 = LoadPrimitiveDataElementUBO(DrawInstanceId, DataIdx + 2);
-		float4 LocalToRelativeWorld3 = LoadPrimitiveDataElementUBO(DrawInstanceId, DataIdx + 3);
-		
-		PrimitiveData.TilePosition = float3(LocalToRelativeWorld0.w, LocalToRelativeWorld1.w, LocalToRelativeWorld2.w);
-		PrimitiveData.LocalToWorld = MakeLWCMatrix4x3(PrimitiveData.TilePosition, float4x4(LocalToRelativeWorld0, LocalToRelativeWorld1, LocalToRelativeWorld2, LocalToRelativeWorld3));
-		PrimitiveData.Flags = asuint(LocalToRelativeWorld3.w);
-		DataIdx += 4u;
+		float4x4 LocalToRelativeWorld = transpose(float4x4(
+			LoadPrimitiveDataElementUBO(DrawInstanceId, DataIdx+0),
+			LoadPrimitiveDataElementUBO(DrawInstanceId, DataIdx+1),
+			LoadPrimitiveDataElementUBO(DrawInstanceId, DataIdx+2),
+			float4(0, 0, 0, 1)
+		));
+		PrimitiveData.LocalToWorld = MakeLWCMatrix4x3(PrimitiveData.TilePosition, LocalToRelativeWorld);
+		DataIdx += 3u;
 	}
 
 	// InvNonUniformScale, TODO .w
diff --git a/Engine/Shaders/Private/SceneDataMobileWriter.ush b/Engine/Shaders/Private/SceneDataMobileWriter.ush
--- a/Engine/Shaders/Private/SceneDataMobileWriter.ush
+++ b/Engine/Shaders/Private/SceneDataMobileWriter.ush
@@ -16,24 +16,21 @@ void WriteDataElementUBO(RWStructuredBuffer<float4> Output, uint Offset, uint Da
 void WritePrimitiveDataUBO(RWStructuredBuffer<float4> Output, uint Offset, FPrimitiveSceneData PrimitiveData, uint MeshLODIndex)
 {
 	uint DataIdx = 0;
-	
-	// LocalToWorld, TilePosition, Primitive.Flags
-	{
-		float4 LocalToRelativeWorld0 = PrimitiveData.LocalToWorld.M[0];
-		float4 LocalToRelativeWorld1 = PrimitiveData.LocalToWorld.M[1];
-		float4 LocalToRelativeWorld2 = PrimitiveData.LocalToWorld.M[2];
-		float4 LocalToRelativeWorld3 = PrimitiveData.LocalToWorld.M[3];
-		LocalToRelativeWorld0.w = PrimitiveData.TilePosition.x;
-		LocalToRelativeWorld1.w = PrimitiveData.TilePosition.y;
-		LocalToRelativeWorld2.w = PrimitiveData.TilePosition.z;
-		LocalToRelativeWorld3.w = asfloat(PrimitiveData.Flags);
-		WriteDataElementUBO(Output, Offset, DataIdx + 0, LocalToRelativeWorld0);
-		WriteDataElementUBO(Output, Offset, DataIdx + 1, LocalToRelativeWorld1);
-		WriteDataElementUBO(Output, Offset, DataIdx + 2, LocalToRelativeWorld2);
-		WriteDataElementUBO(Output, Offset, DataIdx + 3, LocalToRelativeWorld3);
-		DataIdx += 4u;
+	// TilePosition, Flags
+	{
+		WriteDataElementUBO(Output, Offset, DataIdx + 0, float4(LWCGetTile(PrimitiveData.LocalToWorld), asfloat(PrimitiveData.Flags)));
+		DataIdx += 1u;
 	}
 	
+	// LocalToWorld
+	{
+		float4x4 LocalToWorld = transpose(PrimitiveData.LocalToWorld.M);
+		WriteDataElementUBO(Output, Offset, DataIdx + 0, LocalToWorld[0]);
+		WriteDataElementUBO(Output, Offset, DataIdx + 1, LocalToWorld[1]);
+		WriteDataElementUBO(Output, Offset, DataIdx + 2, LocalToWorld[2]);
+		DataIdx += 3u;
+	}
+
 	// InvNonUniformScale, TODO .w
 	{
 		WriteDataElementUBO(Output, Offset, DataIdx + 0, float4(PrimitiveData.InvNonUniformScale, 0.0));
