Fix CIS (AddAABBWS getting wrong float4)

diff --git a/Engine/Shaders/Private/PhysicsFieldVisualizer.usf b/Engine/Shaders/Private/PhysicsFieldVisualizer.usf
--- a/Engine/Shaders/Private/PhysicsFieldVisualizer.usf
+++ b/Engine/Shaders/Private/PhysicsFieldVisualizer.usf
@@ -167,7 +167,7 @@ void MainCS(uint3 DispatchThreadId : SV_DispatchThreadID,
 		{
 			for(int BoundsIndex = BoundsOffset, BoundsEnd = BoundsOffset + BoundsSize; BoundsIndex < BoundsEnd; ++BoundsIndex)
 			{
-				AddAABBWS(BoundsMin[BoundsIndex].xyz, BoundsMax[BoundsIndex].xyz, float4(AABBColor, 1));
+				AddAABBWS(BoundsMin[BoundsIndex].xyz, BoundsMax[BoundsIndex].xyz, float4(AABBColor.Color, 1));
 			}
 			PrintVectorField();
 		}
