Add support for enabling occlusion query based instance culling only for specific scene primitive types

* Add FPrimitiveSceneProxy::AllowInstanceCullingOcclusionQueries() method to allow occlusion queries (disallowed by default)
* Add PRIMITIVE_SCENE_DATA_FLAG_INSTANCE_CULLING_OCCLUSION_QUERIES
* Allow OQ for static meshes only
* Incompatible instances are treated as visible during OQ instance culling pass
* Tweak help text forVisualizeInstanceOcclusionQueries visualize menu entry

Only certain objects may benefit from per-pixel occlusion culling and otherwise HZB-based culling is sufficient.

diff --git a/Engine/Shaders/Private/InstanceCulling/InstanceCullingOcclusionQuery.usf b/Engine/Shaders/Private/InstanceCulling/InstanceCullingOcclusionQuery.usf
--- a/Engine/Shaders/Private/InstanceCulling/InstanceCullingOcclusionQuery.usf
+++ b/Engine/Shaders/Private/InstanceCulling/InstanceCullingOcclusionQuery.usf
@@ -34,10 +34,23 @@ int4 ViewRect;
 #define IVS_Hidden          0 // Instance is hidden or invalid, no occlusion query needed
 #define IVS_Visible         1 // Instance is definitely visible, no occlusion query needed
 #define IVS_PossiblyVisible 2 // Occlusion query is needed to determine visibility
+#define IVS_Incompatible    3 // Instance can't use occlusion queries
+
+// TODO: Compute appropriate mask value based on current view index and previous mask value
+#define MaskValueVisible 1
+#define MaskValueHidden  0
 
 EInstanceVisibilityStatus GetInstanceDataAndVisibility(uint InstanceId, out FInstanceSceneData OutInstanceData)
 {
 	const FInstanceSceneData InstanceData = GetInstanceSceneData(InstanceId, InstanceSceneDataSOAStride);
+	const FPrimitiveSceneData PrimitiveData = GetPrimitiveData(InstanceData.PrimitiveId);
+
+	if ((PrimitiveData.Flags & PRIMITIVE_SCENE_DATA_FLAG_INSTANCE_CULLING_OCCLUSION_QUERIES) == 0u)
+	{
+		// If instance culling occlusion queries are not allowed for this primitive, it is marked as "visible"
+		return IVS_Incompatible;
+	}
+
 	OutInstanceData = InstanceData;
 	
 	if (!InstanceData.ValidInstance)
@@ -143,19 +156,19 @@ void MainCS(uint3 GroupId : SV_GroupID, uint GroupThreadIndex : SV_GroupIndex, u
 	{
 		Status = GetInstanceDataAndVisibility(InstanceId, /*out*/ InstanceData);
 
-		if (Status == IVS_Visible)
+		if (Status == IVS_Visible || Status == IVS_Incompatible)
 		{
 			// YURIY_TODO: multi-view support
-			OutVisibilityMask[InstanceId] = 1; // this is the final value
+			OutVisibilityMask[InstanceId] = MaskValueVisible; // this is the final value
 		}
 		else if (Status == IVS_Hidden)
 		{
 			// YURIY_TODO: multi-view support
-			OutVisibilityMask[InstanceId] = 0; // this is the final value
+			OutVisibilityMask[InstanceId] = MaskValueHidden; // this is the final value
 		}
 		else // Status == IVS_PossiblyVisible
 		{
-			OutVisibilityMask[InstanceId] = 0; // clear the value which will be updated by pixel shader later
+			OutVisibilityMask[InstanceId] = MaskValueHidden; // clear the value which will be updated by pixel shader later
 
 			uint LocalOutputIndex = 0;
 			InterlockedAdd(GroupNumVisibleInstances, 1, LocalOutputIndex);
@@ -186,6 +199,7 @@ void MainVS(
 	float3 InPosition : ATTRIBUTE0,
 	uint InstanceId : SV_InstanceID,
 	out nointerpolation uint OutInstanceId : INSTANCE_ID,
+	out nointerpolation uint OutVisibilityMask : VISIBILITY_MASK,
 	out float4 OutPosition : SV_POSITION
 )
 {
@@ -198,7 +212,7 @@ void MainVS(
 	GetInstanceDataAndVisibility(InstanceId, /*out*/ InstanceData);
 #else
 	EInstanceVisibilityStatus Status = GetInstanceDataAndVisibility(InstanceId, /*out*/ InstanceData);
-	if (Status == IVS_Visible)
+	if (Status == IVS_Visible || Status == IVS_Incompatible)
 	{
 		// Render a box that covers just one pixel at the near plane which will mark the instance visible
 		float2 ViewRectSize = float2(ViewRect.z - ViewRect.x, ViewRect.w - ViewRect.y);
@@ -218,17 +232,18 @@ void MainVS(
 	}
 
 	OutInstanceId = InstanceId;
+	OutVisibilityMask = MaskValueVisible;
 }
 
 
 EARLYDEPTHSTENCIL
 void MainPS(
 	in nointerpolation uint InstanceId : INSTANCE_ID,
+	in nointerpolation uint VisibilityMask : VISIBILITY_MASK,
 	out float4 OutColor : SV_Target0
 )
 {
-	// YURIY_TODO: multi-view support
-	OutVisibilityMask[InstanceId] = 1;
+	OutVisibilityMask[InstanceId] = VisibilityMask;
 	OutColor = float4(1, 0, 0, 1);
 }
 
@@ -248,7 +263,7 @@ void DebugMainVS(
 	FInstanceSceneData InstanceData;
 	EInstanceVisibilityStatus VisibilityStatus = GetInstanceDataAndVisibility(InstanceId, /*out*/ InstanceData);
 	
-	if (InstanceData.ValidInstance)
+	if (InstanceData.ValidInstance && VisibilityStatus != IVS_Incompatible)
 	{
 		float3 LocalPosition = (InPosition * InstanceData.LocalBoundsExtent * InstanceData.DeterminantSign) + InstanceData.LocalBoundsCenter;
 		float4 WorldPosition = TransformLocalToTranslatedWorld(LocalPosition, InstanceData.LocalToWorld);
diff --git a/Engine/Shaders/Shared/SceneDefinitions.h b/Engine/Shaders/Shared/SceneDefinitions.h
--- a/Engine/Shaders/Shared/SceneDefinitions.h
+++ b/Engine/Shaders/Shared/SceneDefinitions.h
@@ -50,6 +50,8 @@
 #define PRIMITIVE_SCENE_DATA_FLAG_DISABLE_MATERIAL_INVALIDATIONS		0x200000
 #define PRIMITIVE_SCENE_DATA_FLAG_HOLDOUT								0x400000
 #define PRIMITIVE_SCENE_DATA_FLAG_SPLINE_MESH							0x800000
+#define PRIMITIVE_SCENE_DATA_FLAG_INSTANCE_CULLING_OCCLUSION_QUERIES	0x1000000 // Whether per-instance occlusion queries are allowed when performing GPUScene instance culling
+
 
 // Primitive visibility flags
 #define PRIMITIVE_VISIBILITY_FLAG_FORCE_HIDDEN							0x1
diff --git a/Engine/Source/Runtime/Engine/Private/PrimitiveSceneProxy.cpp b/Engine/Source/Runtime/Engine/Private/PrimitiveSceneProxy.cpp
--- a/Engine/Source/Runtime/Engine/Private/PrimitiveSceneProxy.cpp
+++ b/Engine/Source/Runtime/Engine/Private/PrimitiveSceneProxy.cpp
@@ -640,6 +640,7 @@ void FPrimitiveSceneProxy::BuildUniformShaderParameters(FPrimitiveUniformShaderP
 			.CastShadow(CastsDynamicShadow())
 			.Holdout(Holdout())
 			.DisableMaterialInvalidations(ShadowCacheInvalidationBehavior == EShadowCacheInvalidationBehavior::Rigid || ShadowCacheInvalidationBehavior == EShadowCacheInvalidationBehavior::Static)
+			.AllowInstanceCullingOcclusionQueries(AllowInstanceCullingOcclusionQueries())
 			.VisibleInGame(IsDrawnInGame())
 			.VisibleInEditor(IsDrawnInEditor())
 			.VisibleInReflectionCaptures(IsVisibleInReflectionCaptures())
diff --git a/Engine/Source/Runtime/Engine/Public/PrimitiveSceneProxy.h b/Engine/Source/Runtime/Engine/Public/PrimitiveSceneProxy.h
--- a/Engine/Source/Runtime/Engine/Public/PrimitiveSceneProxy.h
+++ b/Engine/Source/Runtime/Engine/Public/PrimitiveSceneProxy.h
@@ -476,6 +476,15 @@ public:
 		return true;
 	}
 
+	/**
+	* @return true if per-pixel visibility tests are allowed for instances of this primitive during GPU instance culling.
+	* This may yield significant performance benefit for some cases, such as relatively high-poly instanced static meshes.
+	*/
+	virtual bool AllowInstanceCullingOcclusionQueries() const
+	{
+		return false;
+	}
+
 	/**
 	* @return true if the proxy can skip redundant transform updates where applicable.
 	*/
diff --git a/Engine/Source/Runtime/Engine/Public/PrimitiveUniformShaderParametersBuilder.h b/Engine/Source/Runtime/Engine/Public/PrimitiveUniformShaderParametersBuilder.h
--- a/Engine/Source/Runtime/Engine/Public/PrimitiveUniformShaderParametersBuilder.h
+++ b/Engine/Source/Runtime/Engine/Public/PrimitiveUniformShaderParametersBuilder.h
@@ -59,6 +59,7 @@ public:
 		bHoldout									= false;
 		bDisableMaterialInvalidations				= false;
 		bSplineMesh									= false;
+		bAllowInstanceCullingOcclusionQueries		= false;
 
 		Parameters.MaxWPOExtent						= 0.0f;
 		Parameters.MinMaterialDisplacement			= 0.0f;
@@ -124,6 +125,7 @@ public:
 	PRIMITIVE_UNIFORM_BUILDER_FLAG_METHOD(bool,			Holdout);
 	PRIMITIVE_UNIFORM_BUILDER_FLAG_METHOD(bool,			DisableMaterialInvalidations);
 	PRIMITIVE_UNIFORM_BUILDER_FLAG_METHOD(bool,			SplineMesh);
+	PRIMITIVE_UNIFORM_BUILDER_FLAG_METHOD(bool,			AllowInstanceCullingOcclusionQueries);
 	PRIMITIVE_UNIFORM_BUILDER_FLAG_METHOD(bool,			HasAlwaysEvaluateWPOMaterials);
 
 	PRIMITIVE_UNIFORM_BUILDER_METHOD(uint32,			InstanceSceneDataOffset);
@@ -393,6 +395,7 @@ public:
 		Parameters.Flags |= bHoldout ? PRIMITIVE_SCENE_DATA_FLAG_HOLDOUT : 0u;
 		Parameters.Flags |= bDisableMaterialInvalidations ? PRIMITIVE_SCENE_DATA_FLAG_DISABLE_MATERIAL_INVALIDATIONS : 0u;
 		Parameters.Flags |= bSplineMesh ? PRIMITIVE_SCENE_DATA_FLAG_SPLINE_MESH : 0u;
+		Parameters.Flags |= bAllowInstanceCullingOcclusionQueries ? PRIMITIVE_SCENE_DATA_FLAG_INSTANCE_CULLING_OCCLUSION_QUERIES: 0u;
 		
 		Parameters.VisibilityFlags = 0;
 		Parameters.VisibilityFlags |= bCastHiddenShadow ? PRIMITIVE_VISIBILITY_FLAG_CAST_HIDDEN_SHADOW : 0u;
@@ -455,4 +458,5 @@ private:
 	uint32 bHoldout : 1;
 	uint32 bDisableMaterialInvalidations : 1;
 	uint32 bSplineMesh : 1;
+	uint32 bAllowInstanceCullingOcclusionQueries : 1;
 };
diff --git a/Engine/Source/Runtime/Engine/Public/ShowFlagsValues.inl b/Engine/Source/Runtime/Engine/Public/ShowFlagsValues.inl
--- a/Engine/Source/Runtime/Engine/Public/ShowFlagsValues.inl
+++ b/Engine/Source/Runtime/Engine/Public/ShowFlagsValues.inl
@@ -380,7 +380,7 @@ SHOWFLAG_FIXED_IN_SHIPPING(0, VREditing, SFG_Hidden, NSLOCTEXT("UnrealEd", "VREd
 /** Visualize Occlusion Query bounding meshes */
 SHOWFLAG_FIXED_IN_SHIPPING(0, OcclusionMeshes, SFG_Visualize, NSLOCTEXT("UnrealEd", "VisualizeOcclusionQueries", "Visualize Occlusion Queries"))
 /** Visualize Occlusion Query bounding meshes */
-SHOWFLAG_FIXED_IN_SHIPPING(0, VisualizeInstanceOcclusionQueries, SFG_Visualize, NSLOCTEXT("UnrealEd", "VisualizeInstanceOcclusionQueriesSF", "Visualize Occlusion Queries Used By GPUScene Instance Culling"))
+SHOWFLAG_FIXED_IN_SHIPPING(0, VisualizeInstanceOcclusionQueries, SFG_Visualize, NSLOCTEXT("UnrealEd", "VisualizeInstanceOcclusionQueriesSF", "Visualize Instance Culling Occlusion Queries"))
 /** Disable hardware occlusion queries, similar to setting r.AllowOcclusionQueries=0, but just for this scene. */
 SHOWFLAG_FIXED_IN_SHIPPING(0, DisableOcclusionQueries, SFG_Developer, NSLOCTEXT("UnrealEd", "DisableOcclusionQueries", "Disable Hardware Occlusion Queries"))
 
diff --git a/Engine/Source/Runtime/Engine/Public/StaticMeshSceneProxy.h b/Engine/Source/Runtime/Engine/Public/StaticMeshSceneProxy.h
--- a/Engine/Source/Runtime/Engine/Public/StaticMeshSceneProxy.h
+++ b/Engine/Source/Runtime/Engine/Public/StaticMeshSceneProxy.h
@@ -112,6 +112,7 @@ public:
 	ENGINE_API virtual int32 GetLOD(const FSceneView* View) const override;
 	ENGINE_API virtual FPrimitiveViewRelevance GetViewRelevance(const FSceneView* View) const override;
 	ENGINE_API virtual bool CanBeOccluded() const override;
+	ENGINE_API virtual bool AllowInstanceCullingOcclusionQueries() const override { return true; }
 	ENGINE_API virtual bool IsUsingDistanceCullFade() const override;
 	ENGINE_API virtual void GetLightRelevance(const FLightSceneProxy* LightSceneProxy, bool& bDynamic, bool& bRelevant, bool& bLightMapped, bool& bShadowMapped) const override;
 	ENGINE_API virtual void GetDistanceFieldAtlasData(const FDistanceFieldVolumeData*& OutDistanceFieldData, float& SelfShadowBias) const override;
diff --git a/Engine/Source/Runtime/Experimental/GeometryCollectionEngine/Private/GeometryCollection/GeometryCollectionSceneProxy.h b/Engine/Source/Runtime/Experimental/GeometryCollectionEngine/Private/GeometryCollection/GeometryCollectionSceneProxy.h
--- a/Engine/Source/Runtime/Experimental/GeometryCollectionEngine/Private/GeometryCollection/GeometryCollectionSceneProxy.h
+++ b/Engine/Source/Runtime/Experimental/GeometryCollectionEngine/Private/GeometryCollection/GeometryCollectionSceneProxy.h
@@ -254,6 +254,8 @@ public:
 	FPrimitiveViewRelevance GetViewRelevance(const FSceneView* View) const override;
 	void GetDynamicMeshElements(const TArray<const FSceneView*>& Views, const FSceneViewFamily& ViewFamily, uint32 VisibilityMap, FMeshElementCollector& Collector) const override;
 
+	virtual bool AllowInstanceCullingOcclusionQueries() const override { return true; }
+
 #if GEOMETRYCOLLECTION_EDITOR_SELECTION
 	virtual HHitProxy* CreateHitProxies(UPrimitiveComponent* Component, TArray<TRefCountPtr<HHitProxy> >& OutHitProxies) override;
  	virtual const FColorVertexBuffer* GetCustomHitProxyIdBuffer() const override
