Fix a sky capture issue in OpenGL that occurred when binding a mipmap as a RenderTarget and sampling from a different mipmap.
Fix a artifacts caused by half precision in Pixel Shader on GL.

diff --git a/Engine/Shaders/Private/ReflectionEnvironmentShaders.usf b/Engine/Shaders/Private/ReflectionEnvironmentShaders.usf
--- a/Engine/Shaders/Private/ReflectionEnvironmentShaders.usf
+++ b/Engine/Shaders/Private/ReflectionEnvironmentShaders.usf
@@ -480,6 +480,13 @@ void DownsamplePS_Mobile(
 	OutColor = SampleCubemap(CubeCoordinates);
 }
 
+#if SHADING_PATH_MOBILE
+// There's a precision problem in OpenGL because the default precision in the Pixel Shader is set to half, even for constant values, use Hammersley16 instead.
+#define HammersleyDistribution Hammersley16
+#else
+#define HammersleyDistribution Hammersley
+#endif
+
 #ifdef USE_COMPUTE
 int CubeFaceOffset;
 [numthreads(THREADGROUP_SIZE, THREADGROUP_SIZE, 1)]
@@ -524,7 +531,7 @@ void FilterPS(
 	}
 
 	uint CubeSize = 1u << ( NumMips - 1 );
-	const float SolidAngleTexel = 4.0f * PI / float(6 * CubeSize * CubeSize) * 2.0f;
+	const float SolidAngleTexel = 4.0f * PI / float(6.0f * CubeSize * CubeSize) * 2.0f;
 
 	//const uint NumSamples = 1024;
 	const uint NumSamples = Roughness < 0.1 ? 32 : 64;
@@ -538,7 +545,7 @@ void FilterPS(
 		LOOP
 		for( uint i = 0; i < NumSamples; i++ )
 		{
-			float2 E = Hammersley( i, NumSamples, 0 );
+			float2 E = HammersleyDistribution( i, NumSamples, 0 );
 
 			float3 L = CosineSampleHemisphere( E ).xyz;
 
@@ -561,7 +568,7 @@ void FilterPS(
 		LOOP
 		for( uint i = 0; i < NumSamples; i++ )
 		{
-			float2 E = Hammersley( i, NumSamples, 0 );
+			float2 E = HammersleyDistribution( i, NumSamples, 0 );
 
 			// 6x6 Offset rows. Forms uniform star pattern
 			//uint2 Index = uint2( i % 6, i / 6 );
diff --git a/Engine/Source/Runtime/OpenGLDrv/Private/OpenGLUAV.cpp b/Engine/Source/Runtime/OpenGLDrv/Private/OpenGLUAV.cpp
--- a/Engine/Source/Runtime/OpenGLDrv/Private/OpenGLUAV.cpp
+++ b/Engine/Source/Runtime/OpenGLDrv/Private/OpenGLUAV.cpp
@@ -115,7 +115,7 @@ void FOpenGLShaderResourceView::UpdateView()
 		auto const Info = ViewDesc.Texture.SRV.GetViewInfo(Texture);
 
 		Target = Texture->Target;
-		LimitMip = Info.MipRange.First;
+		LimitMip = Info.MipRange.Num == 1 ? Info.MipRange.First : -1;
 
 		ensureMsgf(Info.MipRange.ExclusiveLast() == Texture->GetDesc().NumMips, TEXT("OpenGL RHI does not currently implement creation of SRVs that don't include all mips in the mip tail."));
 		ensureMsgf(Info.bAllSlices, TEXT("SRV array slice selection is currently unimplemented in OpenGL RHI."));
