Substrate - fixed compilation when new simple volume is enabled.

diff --git a/Engine/Shaders/Private/Strata/StrataEvaluation.ush b/Engine/Shaders/Private/Strata/StrataEvaluation.ush
--- a/Engine/Shaders/Private/Strata/StrataEvaluation.ush
+++ b/Engine/Shaders/Private/Strata/StrataEvaluation.ush
@@ -892,7 +892,7 @@ FStrataEvaluateResult StrataEvaluateBSDFCommon(FStrataBSDFContext BSDFContext, F
 				const float DiffuseToVolumeBlend = StrataSlabDiffuseToVolumeBlend(PM);
 
 			#if USE_NEW_SIMPLEVOLUME
-				const float3 SlabDirectionalAlbedo = IsotropicMediumSlabPunctualDirectionalAlbedoLUT(PM, BSDFContext.Context.NoV, BSDFContext.Context.NoL)
+				const float3 SlabDirectionalAlbedo = IsotropicMediumSlabPunctualDirectionalAlbedoLUT(PM, BSDFContext.Context.NoV, BSDFContext.Context.NoL);
 			#else
 				const float3 SlabDirectionalAlbedo = IsotropicMediumSlabPunctualDirectionalAlbedo(PM);
 			#endif
