Path Tracer: Fix PathTracingQualitySwitch not working in decal materials

diff --git a/Engine/Shaders/Private/RayTracing/RayTracingDecalMaterialShader.usf b/Engine/Shaders/Private/RayTracing/RayTracingDecalMaterialShader.usf
--- a/Engine/Shaders/Private/RayTracing/RayTracingDecalMaterialShader.usf
+++ b/Engine/Shaders/Private/RayTracing/RayTracingDecalMaterialShader.usf
@@ -10,6 +10,7 @@
 
 #include "/Engine/Private/Common.ush"
 #include "/Engine/Private/RayTracing/RayTracingCommon.ush"
+#include "/Engine/Private/PathTracing/PathTracingShaderUtils.ush"
 
 #define DecalTilePosition RayTracingDecalParameters.DecalTilePosition
 #define WorldToDecal RayTracingDecalParameters.WorldToDecal
