Shaders: Tonemap: Fixed broken split-screen rendering when using compute postprocessing.

The original issue was reported for the Lyra sample. It was caused by the UV/coordinates for the secondary view being expressed relative to the origin of the whole render target rather than the origin of the view area.

#rnx

diff --git a/Engine/Shaders/Private/PostProcessTonemap.usf b/Engine/Shaders/Private/PostProcessTonemap.usf
--- a/Engine/Shaders/Private/PostProcessTonemap.usf
+++ b/Engine/Shaders/Private/PostProcessTonemap.usf
@@ -597,14 +597,16 @@ RWTexture2D<float> RWOutputLuminance;
 void MainCS(uint2 DispatchThreadId : SV_DispatchThreadID)
 {
 	float4 SvPosition = float4((float2)DispatchThreadId + Output_ViewportMin + 0.5f, 0.0f, 1.0f);
-	float2 UV = SvPosition.xy * Output_ExtentInverse;
-	float4 InScreenPos = float4(UV * float2(2, -2) - float2(1, -1), 0, 1);
+	float2 UVScreen = SvPosition.xy * Output_ExtentInverse;
 
-	if (IsComputeUVOutOfBounds(UV))
+	if (IsComputeUVOutOfBounds(UVScreen))
 	{
 		return;
 	}
 
+	float2 UV = ((float2)DispatchThreadId + 0.5f) * Output_ExtentInverse;
+	float4 InScreenPos = float4(UV * float2(2, -2) - float2(1, -1), 0, 1);
+
 	float2 Vignette;
 	float4 GrainUV;
 	float2 FullViewUV, ScreenPos;	
