Fixed issues with pre-computed lighting and a mobile GPUScene

diff --git a/Engine/Shaders/Private/LandscapeVertexFactory.ush b/Engine/Shaders/Private/LandscapeVertexFactory.ush
--- a/Engine/Shaders/Private/LandscapeVertexFactory.ush
+++ b/Engine/Shaders/Private/LandscapeVertexFactory.ush
@@ -566,7 +566,11 @@ void GetLightMapCoordinates(FVertexFactoryInterpolantsVSToPS Interpolants, out f
 #endif
 
 #if VF_USE_PRIMITIVE_SCENE_DATA
-	LightmapDataIndex = GetPrimitiveData(Interpolants.PrimitiveId).LightmapDataIndex;
+	#if SHADER_USES_PRIMITIVE_UBO
+		LightmapDataIndex = Interpolants.PrimitiveId;
+	#else
+		LightmapDataIndex = GetPrimitiveData(Interpolants.PrimitiveId).LightmapDataIndex;
+	#endif
 #else
 	LightmapDataIndex = 0;
 #endif
@@ -577,7 +581,11 @@ void GetShadowMapCoordinate(FVertexFactoryInterpolantsVSToPS Interpolants, out f
 	ShadowMapCoordinate = Interpolants.ShadowMapCoordinate;
 	
 #if VF_USE_PRIMITIVE_SCENE_DATA
-	LightmapDataIndex = GetPrimitiveData(Interpolants.PrimitiveId).LightmapDataIndex;
+	#if SHADER_USES_PRIMITIVE_UBO
+		LightmapDataIndex = Interpolants.PrimitiveId;
+	#else
+		LightmapDataIndex = GetPrimitiveData(Interpolants.PrimitiveId).LightmapDataIndex;
+	#endif
 #else
 	LightmapDataIndex = 0;
 #endif
diff --git a/Engine/Shaders/Private/LightmapData.ush b/Engine/Shaders/Private/LightmapData.ush
--- a/Engine/Shaders/Private/LightmapData.ush
+++ b/Engine/Shaders/Private/LightmapData.ush
@@ -44,15 +44,16 @@ float4 LoadLightmapDataElement(uint Index)
 }
 
 // Fetch from scene lightmap data buffer
-FLightmapSceneData GetLightmapData(uint LightmapDataIndex, uint DrawInstanceId) 
+FLightmapSceneData GetLightmapData(uint LightmapDataIndex) 
 {
 	FLightmapSceneData LightmapData;
 
-#if PLATFORM_USES_PRIMITIVE_UBO
-	// FIXME: unpack relevant ligtmap data from Primitive UBO
+#if SHADER_USES_PRIMITIVE_UBO
 	LightmapData = (FLightmapSceneData)0;
-	LightmapData.LightMapCoordinateScaleBias = float4(1,1,0,0);
-	LightmapData.ShadowMapCoordinateScaleBias = float4(1,1,0,0);
+	LightmapData.StaticShadowMapMasks = LoadPrimitiveDataElementUBO(LightmapDataIndex, 0u);
+	LightmapData.InvUniformPenumbraSizes = LoadPrimitiveDataElementUBO(LightmapDataIndex, 1u);
+	LightmapData.LightMapCoordinateScaleBias = LoadPrimitiveDataElementUBO(LightmapDataIndex, 2u);
+	LightmapData.ShadowMapCoordinateScaleBias = LoadPrimitiveDataElementUBO(LightmapDataIndex, 3u);
 #else	
 	// Note: layout must match FLightmapSceneShaderData in C++
 	// Relying on optimizer to remove unused loads
@@ -72,14 +73,18 @@ FLightmapSceneData GetLightmapData(uint LightmapDataIndex, uint DrawInstanceId)
 	{
 		LightmapData.LightmapVTPackedUniform[i] = asuint(LoadLightmapDataElement(LightmapDataBaseOffset + 10 + i));
 	}
-#endif // PLATFORM_USES_PRIMITIVE_UBO
+#endif // SHADER_USES_PRIMITIVE_UBO
 
 	return LightmapData;
 }
 
-FLightmapSceneData GetLightmapData(uint LightmapDataIndex) 
+FLightmapSceneData GetLightmapData(uint LightmapDataIndex, uint DrawInstanceId)
 {
-	return GetLightmapData(LightmapDataIndex, 0); // FIXME: UBO path
+#if SHADER_USES_PRIMITIVE_UBO
+	return GetLightmapData(DrawInstanceId);
+#else
+	return GetLightmapData(LightmapDataIndex);
+#endif
 }
 
 #else 
diff --git a/Engine/Shaders/Private/LocalVertexFactoryCommon.ush b/Engine/Shaders/Private/LocalVertexFactoryCommon.ush
--- a/Engine/Shaders/Private/LocalVertexFactoryCommon.ush
+++ b/Engine/Shaders/Private/LocalVertexFactoryCommon.ush
@@ -109,7 +109,11 @@ void GetLightMapCoordinates(FVertexFactoryInterpolantsVSToPS Interpolants, out f
 	LightmapUV1 = LightmapUV0 + float2( 0, 0.5 );
 
 #if VF_USE_PRIMITIVE_SCENE_DATA && NEEDS_LIGHTMAP_COORDINATE
-	LightmapDataIndex = Interpolants.LightmapDataIndex;
+	#if SHADER_USES_PRIMITIVE_UBO
+		LightmapDataIndex = Interpolants.PrimitiveId;
+	#else
+		LightmapDataIndex = Interpolants.LightmapDataIndex;
+	#endif
 #else
 	LightmapDataIndex = 0;
 #endif
@@ -118,7 +122,11 @@ void GetLightMapCoordinates(FVertexFactoryInterpolantsVSToPS Interpolants, out f
 void GetShadowMapCoordinate(FVertexFactoryInterpolantsVSToPS Interpolants, out float2 ShadowMapCoordinate, out uint LightmapDataIndex)
 {
 #if VF_USE_PRIMITIVE_SCENE_DATA && NEEDS_LIGHTMAP_COORDINATE
-	LightmapDataIndex = Interpolants.LightmapDataIndex;
+	#if SHADER_USES_PRIMITIVE_UBO
+		LightmapDataIndex = Interpolants.PrimitiveId;
+	#else
+		LightmapDataIndex = Interpolants.LightmapDataIndex;
+	#endif
 #else
 	LightmapDataIndex = 0;
 #endif
diff --git a/Engine/Shaders/Private/SceneDataMobileLoader.ush b/Engine/Shaders/Private/SceneDataMobileLoader.ush
--- a/Engine/Shaders/Private/SceneDataMobileLoader.ush
+++ b/Engine/Shaders/Private/SceneDataMobileLoader.ush
@@ -18,8 +18,12 @@ float4 LoadPrimitiveDataElementUBO(uint DrawInstanceId, uint DataIdx)
 FPrimitiveSceneData LoadPrimitiveDataUBO(uint DrawInstanceId)
 {
 	FPrimitiveSceneData PrimitiveData = (FPrimitiveSceneData)0;
-
+	
 	uint DataIdx = 0;
+#if ALLOW_STATIC_LIGHTING
+	DataIdx += 4u; // 4 float4 reservered for pre-computed lighting data
+#endif
+
 	// TilePosition, Flags
 	{
 		float4 Data = LoadPrimitiveDataElementUBO(DrawInstanceId, DataIdx + 0);
@@ -126,25 +130,18 @@ FPrimitiveSceneData LoadPrimitiveDataUBO(uint DrawInstanceId)
 		DataIdx += 3u;
 	}
 
-	// WorldToPreviousWorld, TODO make optional
-	{
-		PrimitiveData.WorldToPreviousWorld = transpose(float4x4(
-			LoadPrimitiveDataElementUBO(DrawInstanceId, DataIdx+0),
-			LoadPrimitiveDataElementUBO(DrawInstanceId, DataIdx+1),
-			LoadPrimitiveDataElementUBO(DrawInstanceId, DataIdx+2),
-			float4(0, 0, 0, 1)
-		));
-		DataIdx += 3u;
-	}
+	// Should not be used with regular primitives
+	PrimitiveData.WorldToPreviousWorld = float4x4(1,0,0,0, 0,1,0,0, 0,0,1,0, 0,0,0,1);
 
 	// CustomData, TODO make optional
 	{
+		uint NumCustomData = min(NUM_CUSTOM_PRIMITIVE_DATA, 32u - DataIdx); //23u
 		UNROLL
-		for (uint i = 0; i < NUM_CUSTOM_PRIMITIVE_DATA; ++i)
+		for (uint i = 0; i < NumCustomData; ++i)
 		{
 			PrimitiveData.CustomPrimitiveData[i] = LoadPrimitiveDataElementUBO(DrawInstanceId,  DataIdx+i);
 		}
-		DataIdx += NUM_CUSTOM_PRIMITIVE_DATA;
+		DataIdx += NumCustomData;
 	}
 	return PrimitiveData;
 }
diff --git a/Engine/Shaders/Private/SceneDataMobileWriter.ush b/Engine/Shaders/Private/SceneDataMobileWriter.ush
--- a/Engine/Shaders/Private/SceneDataMobileWriter.ush
+++ b/Engine/Shaders/Private/SceneDataMobileWriter.ush
@@ -16,12 +16,39 @@ void WriteDataElementUBO(RWStructuredBuffer<float4> Output, uint Offset, uint Da
 void WritePrimitiveDataUBO(RWStructuredBuffer<float4> Output, uint Offset, FPrimitiveSceneData PrimitiveData, uint MeshLODIndex)
 {
 	uint DataIdx = 0;
+
+#if ALLOW_STATIC_LIGHTING
+	// Pre-computed lighting data 
+	{
+		float4 StaticShadowMapMasks = float4(1, 1, 1, 1);
+		float4 InvUniformPenumbraSizes = float4(0, 0, 0, 0);
+		float4 LightMapUVScaleBias = float4(1, 1, 0, 0);
+		float4 ShadowMapUVScaleBias = float4(1, 1, 0, 0);
+
+		uint LightmapDataIndex = PrimitiveData.LightmapDataIndex + MeshLODIndex;
+		if (LightmapDataIndex < GPUSceneNumLightmapDataItems)
+		{
+			FLightmapSceneData LightmapData = GetLightmapData(LightmapDataIndex);
+			
+			StaticShadowMapMasks = LightmapData.StaticShadowMapMasks;
+			InvUniformPenumbraSizes = LightmapData.InvUniformPenumbraSizes;
+			LightMapUVScaleBias = LightmapData.LightMapCoordinateScaleBias;
+			ShadowMapUVScaleBias = LightmapData.ShadowMapCoordinateScaleBias;
+		}
+		WriteDataElementUBO(Output, Offset, DataIdx + 0, StaticShadowMapMasks);
+		WriteDataElementUBO(Output, Offset, DataIdx + 1, InvUniformPenumbraSizes);
+		WriteDataElementUBO(Output, Offset, DataIdx + 2, LightMapUVScaleBias);
+		WriteDataElementUBO(Output, Offset, DataIdx + 3, ShadowMapUVScaleBias);
+		DataIdx += 4u;
+	}
+#endif
+
 	// TilePosition, Flags
 	{
 		WriteDataElementUBO(Output, Offset, DataIdx + 0, float4(LWCGetTile(PrimitiveData.LocalToWorld), asfloat(PrimitiveData.Flags)));
 		DataIdx += 1u;
 	}
-	
+
 	// LocalToWorld
 	{
 		float4x4 LocalToWorld = transpose(PrimitiveData.LocalToWorld.M);
@@ -94,24 +121,15 @@ void WritePrimitiveDataUBO(RWStructuredBuffer<float4> Output, uint Offset, FPrim
 		DataIdx += 3u;
 	}
 
-	// WorldToPreviousWorld, TODO make optional
-	{
-		float4x4 WorldToPreviousWorld = transpose(PrimitiveData.WorldToPreviousWorld);
-		WriteDataElementUBO(Output, Offset, DataIdx + 0, WorldToPreviousWorld[0]);
-		WriteDataElementUBO(Output, Offset, DataIdx + 1, WorldToPreviousWorld[1]);
-		WriteDataElementUBO(Output, Offset, DataIdx + 2, WorldToPreviousWorld[2]);
-		DataIdx += 3u;
-	}
-	
-	// 22u
 	// CustomData, TODO make optional
 	{
+		uint NumCustomData = min(NUM_CUSTOM_PRIMITIVE_DATA, 32u - DataIdx); //23u
 		UNROLL
-		for (uint i = 0; i < NUM_CUSTOM_PRIMITIVE_DATA; ++i)
+		for (uint i = 0; i < NumCustomData; ++i)
 		{
 			WriteDataElementUBO(Output, Offset, DataIdx + i, PrimitiveData.CustomPrimitiveData[i]);
 		}
-		DataIdx += NUM_CUSTOM_PRIMITIVE_DATA;
+		DataIdx += NumCustomData;
 	}
 }
 
