Fix indirect lighting of dynamic skylight on mobile

diff --git a/Engine/Shaders/Private/MobileLightingCommon.ush b/Engine/Shaders/Private/MobileLightingCommon.ush
--- a/Engine/Shaders/Private/MobileLightingCommon.ush
+++ b/Engine/Shaders/Private/MobileLightingCommon.ush
@@ -343,17 +343,24 @@ half3 GetImageBasedReflectionLighting_Mobile(
 
 #elif MOBILE_DEFERRED_LIGHTING
 
+	// Normalize for static skylight types which mix with lightmaps.
+	bool bNormalize = ReflectionStruct.SkyLightParameters.z < 1;
 	float SkyAverageBrightness = 1.0f;
 	SpecularIBL = GetSkyLightReflection(ReflectionVector, Roughness, SkyAverageBrightness);
 	
-	#if ALLOW_STATIC_LIGHTING
-	SpecularIBL *= ComputeMixingWeight(IndirectIrradiance, SkyAverageBrightness, Roughness);
-	#endif
+	if (bNormalize)
+	{
+		#if ALLOW_STATIC_LIGHTING
+		SpecularIBL *= ComputeMixingWeight(IndirectIrradiance, SkyAverageBrightness, Roughness);
+		#endif
+	}
 
 	SpecularIBL *= CompositeAlpha;
 #else // !(ENABLE_CLUSTERED_REFLECTION || MOBILE_DEFERRED_LIGHTING)
 
-	half UsingSkyReflection = MobileReflectionCapture.Params.y > 0.0f;
+	bool UsingSkyReflection = MobileReflectionCapture.Params.y > 0.0f;
+	// Normalize for static skylight types which mix with lightmaps.
+	bool bNormalize = UsingSkyReflection && MobileReflectionCapture.Params.z < 1;
 	if (UsingSkyReflection)
 	{
 		// Apply sky colour if the reflection map is the sky.
@@ -368,9 +375,12 @@ half3 GetImageBasedReflectionLighting_Mobile(
 		SpecularIBL = SpecularIBL * ReflectionCaptureBrightness;
 	}
 
-	#if ALLOW_STATIC_LIGHTING
-	SpecularIBL *= ComputeMixingWeight(IndirectIrradiance, MobileReflectionCapture.Params.x, Roughness);
-	#endif
+	if (bNormalize)
+	{
+		#if ALLOW_STATIC_LIGHTING
+		SpecularIBL *= ComputeMixingWeight(IndirectIrradiance, MobileReflectionCapture.Params.x, Roughness);
+		#endif
+	}
 
 	SpecularIBL *= CompositeAlpha;
 #endif
diff --git a/Engine/Source/Runtime/Renderer/Private/MobileBasePassRendering.cpp b/Engine/Source/Runtime/Renderer/Private/MobileBasePassRendering.cpp
--- a/Engine/Source/Runtime/Renderer/Private/MobileBasePassRendering.cpp
+++ b/Engine/Source/Runtime/Renderer/Private/MobileBasePassRendering.cpp
@@ -358,16 +358,18 @@ void SetupMobileSkyReflectionUniformParameters(FSkyLightSceneProxy* SkyLight, FM
 	float SkyMaxMipIndex = 0.f;
 	FTexture* CaptureTexture = GBlackTextureCube;
 
+	bool bSkyLightIsDynamic = false;
 	if (SkyLight && SkyLight->ProcessedTexture)
 	{
 		check(SkyLight->ProcessedTexture->IsInitialized());
 		CaptureTexture = SkyLight->ProcessedTexture;
 		SkyMaxMipIndex = FMath::Log2(static_cast<float>(CaptureTexture->GetSizeX()));
 		Brightness = SkyLight->AverageBrightness;
+		bSkyLightIsDynamic = !SkyLight->bHasStaticLighting && !SkyLight->bWantsStaticShadowing;
 	}
 	
 	//To keep ImageBasedReflectionLighting coherence with PC, use AverageBrightness instead of InvAverageBrightness to calculate the IBL contribution
-	Parameters.Params = FVector4f(Brightness, SkyMaxMipIndex, 0.f, 0.f);
+	Parameters.Params = FVector4f(Brightness, SkyMaxMipIndex, bSkyLightIsDynamic ? 1.0f : 0.0f, 0.f);
 	Parameters.Texture = CaptureTexture->TextureRHI;
 	Parameters.TextureSampler = CaptureTexture->SamplerStateRHI;
 }
