Selectively not reproject TSR's history if not needed

diff --git a/Engine/Shaders/Private/TemporalSuperResolution/TSRUpdateHistory.usf b/Engine/Shaders/Private/TemporalSuperResolution/TSRUpdateHistory.usf
--- a/Engine/Shaders/Private/TemporalSuperResolution/TSRUpdateHistory.usf
+++ b/Engine/Shaders/Private/TemporalSuperResolution/TSRUpdateHistory.usf
@@ -44,6 +44,7 @@
 	#define CONFIG_TRANSLUCENCY_REJECTION 1
 #endif
 
+#define CONFIG_SELECTIVE_REPROJECTION 1
 #define CONFIG_RESPONSIVE_STENCIL 1
 #define CONFIG_TAA_RESPONSIVE_BEHAVIOR 1
 #define CONFIG_LOAD_TRANSLUCENCY_REJECTION 0
@@ -680,6 +681,15 @@ void MainCS(
 	tsr_halfCx2 PrevHighFrequencyColor;
 	tsr_half2 PrevHistoryValidity;
 
+	#if CONFIG_SELECTIVE_REPROJECTION
+	BRANCH
+	if (all(or(LowFrequencyRejection <= tsr_half(0.0).xx, bIsOffScreen)))
+	{
+		PrevHighFrequencyColor = tsr_half(0.0);
+		PrevHistoryValidity    = tsr_half(0.0);
+	}
+	else
+	#endif
 	PLATFORM_SPECIFIC_ISOLATE
 	{
 		tsr_halfC RawHighFrequency[BICUBIC_CATMULL_ROM_SAMPLES][DPV_PIXEL_PER_LANE];
