Optimized bitstream writer, so it doesn't use atomic writes when writing whole DWORDs

diff --git a/Engine/Shaders/Private/BitPacking.ush b/Engine/Shaders/Private/BitPacking.ush
--- a/Engine/Shaders/Private/BitPacking.ush
+++ b/Engine/Shaders/Private/BitPacking.ush
@@ -135,59 +135,73 @@ void PutBits(RWByteAddressBuffer Output, uint AlignedBaseAddress, uint BitOffset
 
 struct FBitStreamWriterState
 {
-	uint AlignedByteAddress;
-   	uint BufferBits;
-    uint BufferOffset;
-    uint BufferMask;
+	uint StartAlignedByteAddress;
+	uint StartBitOffset;   	
+	uint StartBufferBits;
+
+	uint NextAlignedByteAddress;
+	uint BitOffset;	
+	uint BufferBits;
 };
 
 FBitStreamWriterState BitStreamWriter_Create_Aligned(uint AlignedBaseAddressInBytes, uint BitOffset)
 {
 	FBitStreamWriterState State;
 
-	State.AlignedByteAddress = AlignedBaseAddressInBytes + ((BitOffset >> 5) << 2);
+	State.StartAlignedByteAddress = AlignedBaseAddressInBytes + ((BitOffset >> 5) << 2);
 	BitOffset &= 31u;
+	State.StartBitOffset = BitOffset;
+	State.StartBufferBits = 0;		
 
+	State.NextAlignedByteAddress = State.StartAlignedByteAddress + 4u;
+	State.BitOffset = BitOffset;	
 	State.BufferBits = 0;
-	State.BufferOffset = BitOffset;
-	State.BufferMask = BitFieldMaskU32(BitOffset, 0);
-
+	
 	return State;
 }
 
 void BitStreamWriter_Writer(RWByteAddressBuffer Output, inout FBitStreamWriterState State, uint Value, int NumBits, int CompileTimeMaxBits)
 {
-    State.BufferBits |= Value << State.BufferOffset;
+	const uint Tmp = Value << (State.BitOffset & 31u);	// & 31u is implicit
 
-	// State.BufferOffset <= 31
-    uint NextBufferOffset = State.BufferOffset + NumBits;
-    
-    if (NextBufferOffset >= 32)
+    if(State.BitOffset >= 32)
+        State.BufferBits |= Tmp;
+    else
+        State.StartBufferBits |= Tmp;
+
+    const uint NextBitOffset = State.BitOffset + NumBits;
+	
+	// Overflow to next DWORD?
+    if ((State.BitOffset ^ NextBitOffset) >= 32)
     {
-        Output.InterlockedAnd(State.AlignedByteAddress, State.BufferMask);
-        Output.InterlockedOr(State.AlignedByteAddress, State.BufferBits);
-		State.BufferMask = 0;
-		
-		// Shifts are mod 32, so we need special handling when shift could be >= 32.
-		// State.BufferOffset can only be 0 here if NumBits >= 32 and therefore CompileTimeMaxBits >= 32.
-		if(CompileTimeMaxBits >= 32)
-			State.BufferBits = State.BufferOffset ? (Value >> (32 - State.BufferOffset)) : 0u;
-		else
-			State.BufferBits = Value >> (32 - State.BufferOffset);
-        State.AlignedByteAddress += 4;
+		// Keep first DWORD in register, so we can merge it in with atomic later
+        if(State.BitOffset >= 32)
+        {
+            Output.Store(State.NextAlignedByteAddress, State.BufferBits);
+            State.NextAlignedByteAddress += 4;
+        }
+        State.BufferBits = (CompileTimeMaxBits < 32 || (State.BitOffset & 31)) ? (Value >> ((32u - State.BitOffset) & 31u)) : 0u;
     }
 
-	State.BufferOffset = NextBufferOffset & 31;
+	State.BitOffset = NextBitOffset;
 }
 
 void BitStreamWriter_Flush(RWByteAddressBuffer Output, inout FBitStreamWriterState State)
 {
-    if (State.BufferOffset > 0)
-    {
-        uint Mask = State.BufferMask | ~BitFieldMaskU32(State.BufferOffset, 0);
-        Output.InterlockedAnd(State.AlignedByteAddress, Mask);
-        Output.InterlockedOr(State.AlignedByteAddress, State.BufferBits);
-    }
+	// Start
+	const uint NumBits = State.BitOffset - State.StartBitOffset;
+	uint StartMask =	NumBits >= 32 ? 0xFFFFFFFFu :
+						BitFieldMaskU32(NumBits, 0);
+	StartMask <<= State.StartBitOffset;
+	Output.InterlockedAnd(State.StartAlignedByteAddress, ~StartMask);
+	Output.InterlockedOr(State.StartAlignedByteAddress, State.StartBufferBits);
+
+	if (State.BitOffset > 32)
+	{
+		const uint Mask = BitFieldMaskU32(State.BitOffset & 31u, 0);
+		Output.InterlockedAnd(State.NextAlignedByteAddress, ~Mask);
+		Output.InterlockedOr(State.NextAlignedByteAddress, State.BufferBits);
+	}
 }
 
 // Utility functions for packing bits into uints.
