SkyAtmosphere - added multiple scattering support on the directional second light for parity.
This won't increase the cost of single light games.
EngineTest for Sky&Cloud is green.

diff --git a/Engine/Shaders/Private/SkyAtmosphere.usf b/Engine/Shaders/Private/SkyAtmosphere.usf
--- a/Engine/Shaders/Private/SkyAtmosphere.usf
+++ b/Engine/Shaders/Private/SkyAtmosphere.usf
@@ -654,6 +654,12 @@ SingleScatteringResult IntegrateSingleScatteredLuminance(
 #if MULTISCATTERING_APPROX_SAMPLING_ENABLED
 		MultiScatteredLuminance0 = GetMultipleScattering(P, Light0ZenithCosAngle);
 #endif
+#if SECOND_ATMOSPHERE_LIGHT_ENABLED
+		float3 MultiScatteredLuminance1 = 0.0f;
+	#if MULTISCATTERING_APPROX_SAMPLING_ENABLED
+		MultiScatteredLuminance1 = GetMultipleScattering(P, Light1ZenithCosAngle);
+	#endif
+#endif
 
 		// Planet shadow
 		float tPlanet0 = RaySphereIntersectNearest(P, Light0Dir, PlanetO + PLANET_RADIUS_OFFSET * UpVector, Atmosphere.BottomRadiusKm);
@@ -719,9 +725,9 @@ SingleScatteringResult IntegrateSingleScatteredLuminance(
 			VolumetricCloudCommonParameters.CloudShadowmapFarDepthKm[1].x, VolumetricCloudShadowMapTexture1, VolumetricCloudShadowMapTexture1Sampler, OutOpticalDepth3), VolumetricCloudShadowStrength1));
 #endif
 		//  Multi-scattering can work for the second light but it is disabled for the sake of performance.
-		S		+= ExposedLight1Illuminance * PlanetShadow1 * TransmittanceToLight1 * PhaseTimesScattering1			;// +MultiScatteredLuminance * Medium.Scattering);
-		SMieOnly+= ExposedLight1Illuminance * PlanetShadow1 * TransmittanceToLight1 * PhaseTimesScattering1MieOnly	;// +MultiScatteredLuminance * Medium.Scattering);
-		SRayOnly+= ExposedLight1Illuminance * PlanetShadow1 * TransmittanceToLight1 * PhaseTimesScattering1RayOnly	;// +MultiScatteredLuminance * Medium.Scattering);
+		S		 += ExposedLight1Illuminance * (PlanetShadow1 * TransmittanceToLight1 * PhaseTimesScattering1		 + MultiScatteredLuminance1 * Medium.Scattering);
+		SMieOnly += ExposedLight1Illuminance * (PlanetShadow1 * TransmittanceToLight1 * PhaseTimesScattering1MieOnly + MultiScatteredLuminance1 * Medium.ScatteringMie);
+		SRayOnly += ExposedLight1Illuminance * (PlanetShadow1 * TransmittanceToLight1 * PhaseTimesScattering1RayOnly + MultiScatteredLuminance1 * Medium.ScatteringRay);
 #endif
 
 		// When using the power serie to accumulate all sattering order, serie r must be <1 for a serie to converge. 
