Fixed several issues with mobile GPUScene related to static lighting and custom per-instance data

diff --git a/Engine/Shaders/Private/LandscapeVertexFactory.ush b/Engine/Shaders/Private/LandscapeVertexFactory.ush
--- a/Engine/Shaders/Private/LandscapeVertexFactory.ush
+++ b/Engine/Shaders/Private/LandscapeVertexFactory.ush
@@ -566,11 +566,7 @@ void GetLightMapCoordinates(FVertexFactoryInterpolantsVSToPS Interpolants, out f
 #endif
 
 #if VF_USE_PRIMITIVE_SCENE_DATA
-	#if SHADER_USES_PRIMITIVE_UBO
-		LightmapDataIndex = Interpolants.PrimitiveId;
-	#else
-		LightmapDataIndex = GetPrimitiveData(Interpolants.PrimitiveId).LightmapDataIndex;
-	#endif
+	LightmapDataIndex = GetPrimitiveData(Interpolants.PrimitiveId).LightmapDataIndex;
 #else
 	LightmapDataIndex = 0;
 #endif
@@ -581,11 +577,7 @@ void GetShadowMapCoordinate(FVertexFactoryInterpolantsVSToPS Interpolants, out f
 	ShadowMapCoordinate = Interpolants.ShadowMapCoordinate;
 	
 #if VF_USE_PRIMITIVE_SCENE_DATA
-	#if SHADER_USES_PRIMITIVE_UBO
-		LightmapDataIndex = Interpolants.PrimitiveId;
-	#else
-		LightmapDataIndex = GetPrimitiveData(Interpolants.PrimitiveId).LightmapDataIndex;
-	#endif
+	LightmapDataIndex = GetPrimitiveData(Interpolants.PrimitiveId).LightmapDataIndex;
 #else
 	LightmapDataIndex = 0;
 #endif
diff --git a/Engine/Shaders/Private/LightmapData.ush b/Engine/Shaders/Private/LightmapData.ush
--- a/Engine/Shaders/Private/LightmapData.ush
+++ b/Engine/Shaders/Private/LightmapData.ush
@@ -10,9 +10,7 @@
 	#define USE_GLOBAL_GPU_LIGHTMAP_DATA 0
 #endif
 
-#if VF_USE_PRIMITIVE_SCENE_DATA
-
-// Must match FPrecomputedLightingUniformParameters in C++
+ // Must match FPrecomputedLightingUniformParameters in C++
 struct FLightmapSceneData
 {
 	float4 StaticShadowMapMasks;
@@ -25,6 +23,29 @@ struct FLightmapSceneData
 	uint4 LightmapVTPackedUniform[5];
 };
 
+FLightmapSceneData GetLightmapDataFromUniformBuffer()
+{
+	FLightmapSceneData LightmapData;
+	LightmapData.StaticShadowMapMasks = PrecomputedLightingBuffer.StaticShadowMapMasks;
+	LightmapData.InvUniformPenumbraSizes = PrecomputedLightingBuffer.InvUniformPenumbraSizes;
+	LightmapData.LightMapCoordinateScaleBias = PrecomputedLightingBuffer.LightMapCoordinateScaleBias;
+	LightmapData.ShadowMapCoordinateScaleBias = PrecomputedLightingBuffer.ShadowMapCoordinateScaleBias;
+	LightmapData.LightMapScale[0] = PrecomputedLightingBuffer.LightMapScale[0];
+	LightmapData.LightMapScale[1] = PrecomputedLightingBuffer.LightMapScale[1];
+	LightmapData.LightMapAdd[0] = PrecomputedLightingBuffer.LightMapAdd[0];
+	LightmapData.LightMapAdd[1] = PrecomputedLightingBuffer.LightMapAdd[1];
+	LightmapData.LightmapVTPackedPageTableUniform[0] = PrecomputedLightingBuffer.LightmapVTPackedPageTableUniform[0];
+	LightmapData.LightmapVTPackedPageTableUniform[1] = PrecomputedLightingBuffer.LightmapVTPackedPageTableUniform[1];
+	UNROLL
+	for (uint i = 0u; i < 5u; ++i)
+	{
+		LightmapData.LightmapVTPackedUniform[i] = PrecomputedLightingBuffer.LightmapVTPackedUniform[i];
+	}
+	return LightmapData;
+}
+
+#if VF_USE_PRIMITIVE_SCENE_DATA
+
 // Stride of a single lightmap data entry in float4's, must match C++
 #define LIGHTMAP_SCENE_DATA_STRIDE 15
 
@@ -48,13 +69,6 @@ FLightmapSceneData GetLightmapData(uint LightmapDataIndex)
 {
 	FLightmapSceneData LightmapData;
 
-#if SHADER_USES_PRIMITIVE_UBO
-	LightmapData = (FLightmapSceneData)0;
-	LightmapData.StaticShadowMapMasks = LoadPrimitiveDataElementUBO(LightmapDataIndex, 0u);
-	LightmapData.InvUniformPenumbraSizes = LoadPrimitiveDataElementUBO(LightmapDataIndex, 1u);
-	LightmapData.LightMapCoordinateScaleBias = LoadPrimitiveDataElementUBO(LightmapDataIndex, 2u);
-	LightmapData.ShadowMapCoordinateScaleBias = LoadPrimitiveDataElementUBO(LightmapDataIndex, 3u);
-#else	
 	// Note: layout must match FLightmapSceneShaderData in C++
 	// Relying on optimizer to remove unused loads
 	uint LightmapDataBaseOffset = LightmapDataIndex * LIGHTMAP_SCENE_DATA_STRIDE;
@@ -73,7 +87,6 @@ FLightmapSceneData GetLightmapData(uint LightmapDataIndex)
 	{
 		LightmapData.LightmapVTPackedUniform[i] = asuint(LoadLightmapDataElement(LightmapDataBaseOffset + 10 + i));
 	}
-#endif // SHADER_USES_PRIMITIVE_UBO
 
 	return LightmapData;
 }
@@ -81,8 +94,15 @@ FLightmapSceneData GetLightmapData(uint LightmapDataIndex)
 FLightmapSceneData GetLightmapData(uint LightmapDataIndex, uint DrawInstanceId)
 {
 #if SHADER_USES_PRIMITIVE_UBO
-	return GetLightmapData(DrawInstanceId);
-#else
+	#if USE_INSTANCE_CULLING
+		return GetLightmapDataFromUniformBuffer();
+	#else
+		FLightmapSceneData LightmapData = (FLightmapSceneData)0;
+		LightmapData.LightMapCoordinateScaleBias = LoadPrimitiveDataElementUBO(DrawInstanceId, 0u);
+		LightmapData.ShadowMapCoordinateScaleBias = LoadPrimitiveDataElementUBO(DrawInstanceId, 1u);
+		return LightmapData;
+	#endif
+#else	
 	return GetLightmapData(LightmapDataIndex);
 #endif
 }
@@ -90,6 +110,9 @@ FLightmapSceneData GetLightmapData(uint LightmapDataIndex, uint DrawInstanceId)
 #else 
 
 // Route to uniform buffer
-#define GetLightmapData(x) PrecomputedLightingBuffer
+FLightmapSceneData GetLightmapData(uint LightmapDataIndex)
+{
+	return GetLightmapDataFromUniformBuffer();
+}
 
 #endif
\ No newline at end of file
diff --git a/Engine/Shaders/Private/LocalVertexFactoryCommon.ush b/Engine/Shaders/Private/LocalVertexFactoryCommon.ush
--- a/Engine/Shaders/Private/LocalVertexFactoryCommon.ush
+++ b/Engine/Shaders/Private/LocalVertexFactoryCommon.ush
@@ -109,11 +109,7 @@ void GetLightMapCoordinates(FVertexFactoryInterpolantsVSToPS Interpolants, out f
 	LightmapUV1 = LightmapUV0 + float2( 0, 0.5 );
 
 #if VF_USE_PRIMITIVE_SCENE_DATA && NEEDS_LIGHTMAP_COORDINATE
-	#if SHADER_USES_PRIMITIVE_UBO
-		LightmapDataIndex = Interpolants.PrimitiveId;
-	#else
-		LightmapDataIndex = Interpolants.LightmapDataIndex;
-	#endif
+	LightmapDataIndex = Interpolants.LightmapDataIndex;
 #else
 	LightmapDataIndex = 0;
 #endif
@@ -122,11 +118,7 @@ void GetLightMapCoordinates(FVertexFactoryInterpolantsVSToPS Interpolants, out f
 void GetShadowMapCoordinate(FVertexFactoryInterpolantsVSToPS Interpolants, out float2 ShadowMapCoordinate, out uint LightmapDataIndex)
 {
 #if VF_USE_PRIMITIVE_SCENE_DATA && NEEDS_LIGHTMAP_COORDINATE
-	#if SHADER_USES_PRIMITIVE_UBO
-		LightmapDataIndex = Interpolants.PrimitiveId;
-	#else
-		LightmapDataIndex = Interpolants.LightmapDataIndex;
-	#endif
+	LightmapDataIndex = Interpolants.LightmapDataIndex;
 #else
 	LightmapDataIndex = 0;
 #endif
diff --git a/Engine/Shaders/Private/SceneData.ush b/Engine/Shaders/Private/SceneData.ush
--- a/Engine/Shaders/Private/SceneData.ush
+++ b/Engine/Shaders/Private/SceneData.ush
@@ -45,7 +45,6 @@
 	#define BATCHED_INSTANCE_DATA_STRIDE 256u
 	#define BATCHED_PRIMITIVE_DATA_STRIDE_FLOAT4 (BATCHED_PRIMITIVE_DATA_STRIDE/16u)
 	#define BATCHED_INSTANCE_DATA_STRIDE_FLOAT4 (BATCHED_INSTANCE_DATA_STRIDE/16u)
-	#define NUM_BATCHED_INSTANCE_CUSTOM_DATA 8u // Num float4s used for the instance custom data when using UBO path
 #endif
 
 #ifndef SHADER_USES_PRIMITIVE_UBO
diff --git a/Engine/Shaders/Private/SceneDataMobileLoader.ush b/Engine/Shaders/Private/SceneDataMobileLoader.ush
--- a/Engine/Shaders/Private/SceneDataMobileLoader.ush
+++ b/Engine/Shaders/Private/SceneDataMobileLoader.ush
@@ -21,7 +21,10 @@ FPrimitiveSceneData LoadPrimitiveDataUBO(uint DrawInstanceId)
 	
 	uint DataIdx = 0;
 #if ALLOW_STATIC_LIGHTING
-	DataIdx += 4u; // 4 float4 reservered for pre-computed lighting data
+	// 2 float4 reservered for pre-computed lighting data
+	float4 Data = LoadPrimitiveDataElementUBO(DrawInstanceId, DataIdx+2);
+	PrimitiveData.LightmapDataIndex = asuint(Data.x);
+	DataIdx += 3u; 
 #endif
 
 	// TilePosition, Flags
@@ -135,7 +138,7 @@ FPrimitiveSceneData LoadPrimitiveDataUBO(uint DrawInstanceId)
 
 	// CustomData, TODO make optional
 	{
-		uint NumCustomData = min(NUM_CUSTOM_PRIMITIVE_DATA, 32u - DataIdx); //23u
+		uint NumCustomData = min(NUM_CUSTOM_PRIMITIVE_DATA, BATCHED_PRIMITIVE_DATA_STRIDE_FLOAT4 - DataIdx); //23u
 		UNROLL
 		for (uint i = 0; i < NumCustomData; ++i)
 		{
@@ -208,18 +211,43 @@ FInstanceSceneData LoadInstanceDataUBO(uint DrawInstanceId)
 		InstanceData.PrevLocalToWorld = MakeLWCMatrix4x3(TilePosition, PreviousLocalToRelativeWorld);
 		DataIdx += 3u;
 	}
+#if ALLOW_STATIC_LIGHTING
+	const bool bHasLightShadowUVBias = (InstanceData.Flags & INSTANCE_SCENE_DATA_FLAG_HAS_LIGHTSHADOW_UV_BIAS) != 0u;
+	if (bHasLightShadowUVBias)
+	{
+		InstanceData.LightMapAndShadowMapUVBias = LoadInstanceDataElementUBO(DrawInstanceId, DataIdx+0);
+	}
+	DataIdx += 1u;
+#endif
+
+	// CustomDataCount, TODO: yzw
+	{
+		float4 Data = LoadInstanceDataElementUBO(DrawInstanceId, DataIdx+0);
+#if ENABLE_PER_INSTANCE_CUSTOM_DATA
+		InstanceData.CustomDataCount = asuint(Data.x);
+#endif
+		DataIdx += 1u;
+	}
 
 	return InstanceData;
 }
 
 float LoadInstanceCustomDataElementUBO(uint DrawInstanceId, uint CustomDataIdx, float DefaultValue)
 {
-	const uint CustomDataFloat4Offset = 8u; // where custom data starts in an Instance UBO
+#if USE_INSTANCE_CULLING
+	uint CustomDataFloat4Offset = 8u; // where custom data starts in an Instance UBO
+	#if ALLOW_STATIC_LIGHTING
+		CustomDataFloat4Offset += 1u; // one float4 reserved for pre-computed lighting
+	#endif
+
 	const uint Float4Index = (CustomDataIdx >> 2u);
-	if (Float4Index < NUM_BATCHED_INSTANCE_CUSTOM_DATA)
+	const uint CustomDataCount = asuint(LoadInstanceDataElementUBO(DrawInstanceId, CustomDataFloat4Offset).x);
+	CustomDataFloat4Offset += 1u;
+	if (Float4Index < CustomDataCount)
 	{
 		const float4 Float4Packed = LoadInstanceDataElementUBO(DrawInstanceId, Float4Index + CustomDataFloat4Offset);
 		return Float4Packed[CustomDataIdx & 0x3u];
 	}
+#endif
 	return DefaultValue;
 }
\ No newline at end of file
diff --git a/Engine/Shaders/Private/SceneDataMobileWriter.ush b/Engine/Shaders/Private/SceneDataMobileWriter.ush
--- a/Engine/Shaders/Private/SceneDataMobileWriter.ush
+++ b/Engine/Shaders/Private/SceneDataMobileWriter.ush
@@ -20,8 +20,6 @@ void WritePrimitiveDataUBO(RWStructuredBuffer<float4> Output, uint Offset, FPrim
 #if ALLOW_STATIC_LIGHTING
 	// Pre-computed lighting data 
 	{
-		float4 StaticShadowMapMasks = float4(1, 1, 1, 1);
-		float4 InvUniformPenumbraSizes = float4(0, 0, 0, 0);
 		float4 LightMapUVScaleBias = float4(1, 1, 0, 0);
 		float4 ShadowMapUVScaleBias = float4(1, 1, 0, 0);
 
@@ -30,16 +28,13 @@ void WritePrimitiveDataUBO(RWStructuredBuffer<float4> Output, uint Offset, FPrim
 		{
 			FLightmapSceneData LightmapData = GetLightmapData(LightmapDataIndex);
 			
-			StaticShadowMapMasks = LightmapData.StaticShadowMapMasks;
-			InvUniformPenumbraSizes = LightmapData.InvUniformPenumbraSizes;
 			LightMapUVScaleBias = LightmapData.LightMapCoordinateScaleBias;
 			ShadowMapUVScaleBias = LightmapData.ShadowMapCoordinateScaleBias;
 		}
-		WriteDataElementUBO(Output, Offset, DataIdx + 0, StaticShadowMapMasks);
-		WriteDataElementUBO(Output, Offset, DataIdx + 1, InvUniformPenumbraSizes);
-		WriteDataElementUBO(Output, Offset, DataIdx + 2, LightMapUVScaleBias);
-		WriteDataElementUBO(Output, Offset, DataIdx + 3, ShadowMapUVScaleBias);
-		DataIdx += 4u;
+		WriteDataElementUBO(Output, Offset, DataIdx + 0, LightMapUVScaleBias);
+		WriteDataElementUBO(Output, Offset, DataIdx + 1, ShadowMapUVScaleBias);
+		WriteDataElementUBO(Output, Offset, DataIdx + 2, float4(asfloat(PrimitiveData.LightmapDataIndex), 0.f, 0.f, 0.f));
+		DataIdx += 3u;
 	}
 #endif
 
@@ -123,7 +118,7 @@ void WritePrimitiveDataUBO(RWStructuredBuffer<float4> Output, uint Offset, FPrim
 
 	// CustomData, TODO make optional
 	{
-		uint NumCustomData = min(NUM_CUSTOM_PRIMITIVE_DATA, 32u - DataIdx); //23u
+		uint NumCustomData = min(NUM_CUSTOM_PRIMITIVE_DATA, BATCHED_PRIMITIVE_DATA_STRIDE_FLOAT4 - DataIdx); //23u
 		UNROLL
 		for (uint i = 0; i < NumCustomData; ++i)
 		{
@@ -163,15 +158,31 @@ void WriteInstanceDataUBO(RWStructuredBuffer<float4> Output, uint Offset, uint I
 		DataIdx += 3u;
 	}
 
-	// TODO: make it optional, doubles instance data size
 	uint CustomDataCount = 0;
 	uint InstanceRelativeId = 0;
 	LoadInstanceRelativeIdAndCustomDataCount(InstanceId, InstanceSceneDataSOAStride, InstanceRelativeId, CustomDataCount);
-	if (CustomDataCount > 0)
+	FInstancePayloadDataOffsets Offsets = GetInstancePayloadDataOffsets(InstanceData.PrimitiveId, InstanceData.Flags, InstanceRelativeId);
+
+#if ALLOW_STATIC_LIGHTING
+	{
+		if (Offsets.LightShadowUVBias != INVALID_INSTANCE_PAYLOAD_OFFSET)
+		{
+			float4 LightMapAndShadowMapUVBias = LoadInstancePayloadDataElement(Offsets.LightShadowUVBias);
+			WriteDataElementUBO(Output, Offset, DataIdx + 0, LightMapAndShadowMapUVBias);
+		}
+		DataIdx += 1u;
+	}
+#endif
+	// CustomDataCount, TODO: yzw
+	{
+		CustomDataCount = min(CustomDataCount, (BATCHED_INSTANCE_DATA_STRIDE_FLOAT4 - DataIdx - 1u) * 4u);
+		WriteDataElementUBO(Output, Offset, DataIdx + 0, float4(asfloat(CustomDataCount), 0.f, 0.f, 0.f));
+		DataIdx += 1u;
+	}
+
+	// TODO: make it optional, doubles instance data size
 	{
 		// limited space for instance custom data
-		CustomDataCount = min(CustomDataCount, NUM_BATCHED_INSTANCE_CUSTOM_DATA);
-		FInstancePayloadDataOffsets Offsets	= GetInstancePayloadDataOffsets(InstanceData.PrimitiveId, InstanceData.Flags, InstanceRelativeId);
 		InstanceData.CustomDataCount = CustomDataCount;
 		InstanceData.CustomDataOffset = Offsets.CustomData;
 		UNROLL
diff --git a/Engine/Source/Runtime/Engine/Private/InstancedStaticMesh.cpp b/Engine/Source/Runtime/Engine/Private/InstancedStaticMesh.cpp
--- a/Engine/Source/Runtime/Engine/Private/InstancedStaticMesh.cpp
+++ b/Engine/Source/Runtime/Engine/Private/InstancedStaticMesh.cpp
@@ -1814,6 +1814,10 @@ void FInstancedStaticMeshSceneProxy::SetupProxy(UInstancedStaticMeshComponent* I
 		{
 			// Only instance data comes from GPUScene on platforms that use uniform buffer views
 			bVFRequiresPrimitiveUniformBuffer = true;
+			for (int32 LODIdx = 0; LODIdx < LODs.Num(); ++LODIdx)
+			{
+				LODs[LODIdx].bCanUsePrecomputedLightingParametersFromGPUScene = false;
+			}
 		}
 		
 		const TArray<int32>& InstanceReorderTable = InComponent->InstanceReorderTable;
diff --git a/Engine/Source/Runtime/Engine/Public/SceneManagement.h b/Engine/Source/Runtime/Engine/Public/SceneManagement.h
--- a/Engine/Source/Runtime/Engine/Public/SceneManagement.h
+++ b/Engine/Source/Runtime/Engine/Public/SceneManagement.h
@@ -884,7 +884,7 @@ public:
 
 	ENGINE_API FShadowMapInteraction GetShadowMapInteraction(ERHIFeatureLevel::Type InFeatureLevel) const;
 
-protected:
+public:
 	// Load parameters from GPUScene when possible
 	// Basically this is the same as VF_SUPPORTS_PRIMITIVE_SCENE_DATA on the vertex factory, but we can't deduce automatically
 	// because we don't know about VF type until we see the actual mesh batch
