Replug manual ThicknessScale per LOD so that user can scale strands geometry as desired.

diff --git a/Engine/Plugins/Runtime/HairStrands/Source/HairStrandsCore/Private/GroomBuilder.cpp b/Engine/Plugins/Runtime/HairStrands/Source/HairStrandsCore/Private/GroomBuilder.cpp
--- a/Engine/Plugins/Runtime/HairStrands/Source/HairStrandsCore/Private/GroomBuilder.cpp
+++ b/Engine/Plugins/Runtime/HairStrands/Source/HairStrandsCore/Private/GroomBuilder.cpp
@@ -40,7 +40,7 @@ static FAutoConsoleVariableRef CVarHairGroupIndexBuilder_MaxVoxelResolution(TEXT
 
 FString FGroomBuilder::GetVersion()
 {
-	return TEXT("v8r41f");
+	return TEXT("v8r43");
 }
 
 namespace FHairStrandsDecimation
@@ -3073,8 +3073,9 @@ static void BuildClusterData(
 		for (uint8 LODIt = 0; LODIt < LODCount; ++LODIt)
 		{
 			const float Lambda = InSettings.LODs[LODIt].CurveDecimation;
-			Out.LODInfos[LODIt].RadiusScale  = (1.f - FMath::Pow(1.f - r, 1.f / Lambda)) / r;
-			Out.LODInfos[LODIt].RadiusScale *= InSettings.LODs[LODIt].ThicknessScale;
+			// For now, for legacy reason, do not use Stochastic simplification, but rather use manual thickness compensation if provided (default=1.f)
+			//Out.LODInfos[LODIt].RadiusScale = (1.f - FMath::Pow(1.f - r, 1.f / Lambda)) / r;
+			Out.LODInfos[LODIt].RadiusScale = InSettings.LODs[LODIt].ThicknessScale;
 		}
 	}
 
diff --git a/Engine/Shaders/Private/HairStrands/HairStrandsClusterCulling.usf b/Engine/Shaders/Private/HairStrands/HairStrandsClusterCulling.usf
--- a/Engine/Shaders/Private/HairStrands/HairStrandsClusterCulling.usf
+++ b/Engine/Shaders/Private/HairStrands/HairStrandsClusterCulling.usf
@@ -187,7 +187,7 @@ void Main(uint3 DispatchThreadId : SV_DispatchThreadID, uint ThreadIndex : SV_Gr
 		// 1.3 Get cluster info (i.e. map all info per LOD for a given cluster ID) and derive effective LOD data (radius scale)
 		const FHairClusterLOD ClusterLODInfo = GetHairClusterLOD(ClusterInfoBuffer[ClusterId], ClusterInfoParameters, LODIndex);
 		ClusterLOD[LocalCurveIndex] 		 = ClusterLODInfo.LOD;
-		ClusterRadiusScale[LocalCurveIndex]  = 1.f; // Curve are no longer scaled individually by ClusterLODInfo.RadiusScale, but scale globally based on CPU value
+		ClusterRadiusScale[LocalCurveIndex]  = ClusterLODInfo.RadiusScale; // Curve are no longer scaled individually by ClusterLODInfo.RadiusScale, but scale globally based on CPU value. This is for legacy path, using discrete LOD with manual ThicknessScale
 		LODCount[LocalCurveIndex]			 = ClusterLODInfo.LODCount;
 
 		// 1.4 Curve's point count/offset
@@ -264,7 +264,7 @@ void Main(uint3 DispatchThreadId : SV_DispatchThreadID, uint ThreadIndex : SV_Gr
 	if (ThreadIndex < AllocationCount)
 	{
 		OutIndexBuffer[AllocationOffset + ThreadIndex]       = CulledPointIndices[ThreadIndex];
-		OutRadiusScaleBuffer[AllocationOffset + ThreadIndex] = 1.f;  // For fading out the curve which are about to be culled, one can use CulledPointFadeScales[ThreadIndex]
+		OutRadiusScaleBuffer[AllocationOffset + ThreadIndex] = CulledPointRadiusScales[ThreadIndex];  // For fading out the curve which are about to be culled, one can use CulledPointFadeScales[ThreadIndex]
 	}
 
 	// 6. Write out curve data with visible points for subsequent passes
@@ -318,7 +318,7 @@ void Main(uint3 DispatchThreadId : SV_DispatchThreadID, uint ThreadIndex : SV_Gr
 		// 1.2 Get cluster info (i.e. map all info per LOD for a given cluster ID) and derive effective LOD data (radius scale)
 		const FHairClusterLOD ClusterLODInfo = GetHairClusterLOD(ClusterInfoBuffer[ClusterId], ClusterInfoParameters, LODIndex);
 		ClusterLOD = ClusterLODInfo.LOD;
-		ClusterRadiusScale = 1.f; // Curve are no longer scaled individually by ClusterLODInfo.RadiusScale, but scale globally based on CPU value
+		ClusterRadiusScale = ClusterLODInfo.RadiusScale; // Curve are no longer scaled individually by ClusterLODInfo.RadiusScale, but scale globally based on CPU value. This is for legacy path, using discrete LOD with manual ThicknessScale
 		LODCount = ClusterLODInfo.LODCount;
 
 		// 1.4 Curve's point count/offset
@@ -400,7 +400,7 @@ void Main(uint3 DispatchThreadId : SV_DispatchThreadID, uint ThreadIndex : SV_Gr
 				const float FadeScale = CulledPointFadeScales[CurrentPointIndex];
 
 				OutIndexBuffer[AllocationOffset + CurrentPointIndex]       = CurvePointIndex + CullPointIndex;
-				OutRadiusScaleBuffer[AllocationOffset + CurrentPointIndex] = 1.f;  // For fading out the curve which are about to be culled, one can use FadeScale
+				OutRadiusScaleBuffer[AllocationOffset + CurrentPointIndex] = ClusterRadiusScale;  // For fading out the curve which are about to be culled, one can use FadeScale
 			}
 		}
 	}
