Refactors TSR RejectShading code into modular functions

diff --git a/Engine/Shaders/Private/TemporalSuperResolution/TSRConvolutionNetwork.ush b/Engine/Shaders/Private/TemporalSuperResolution/TSRConvolutionNetwork.ush
--- a/Engine/Shaders/Private/TemporalSuperResolution/TSRConvolutionNetwork.ush
+++ b/Engine/Shaders/Private/TemporalSuperResolution/TSRConvolutionNetwork.ush
@@ -78,6 +78,94 @@ tsr_short2 GetLaneSimdPixelOffset(const uint LaneIndex, const uint SimdIndex)
 	);
 }
 
+/** Returns coordinate of a element within a lane. */
+CALL_SITE_DEBUGLOC
+tsr_short2 GetElementPixelPosInLane(const uint LaneStrideX, const uint LaneStrideY, const uint ElementIndex)
+{
+	return tsr_short2(ElementIndex % LaneStrideX, ElementIndex / LaneStrideX);
+}
+
+/** Returns coordinate of a element within a tile. */
+CALL_SITE_DEBUGLOC
+tsr_short2 GetElementPixelOffsetInTile(const uint LaneStrideX, const uint LaneStrideY, const uint LaneIndex, const uint SimdIndex)
+{
+	return (
+		tsr_short2(LaneIndex % uint(LANE_COUNT_X * WAVE_COUNT_X), LaneIndex / uint(LANE_COUNT_X * WAVE_COUNT_X)) * tsr_short2(LaneStrideX, LaneStrideY) +
+		GetElementPixelPosInLane(LaneStrideX, LaneStrideY, SimdIndex)
+	);
+}
+
+CALL_SITE_DEBUGLOC
+tsr_short2 ComputeElementInputPixelPos(
+	const uint2 LaneStride,
+	const uint2 TileOverscan,
+	const uint2 PixelViewportMin,
+	const uint2 PixelViewportMax,
+	const uint2 GroupId,
+	const uint  GroupThreadIndex,
+	const uint  InputElementIndex,
+	const int2  InputDataOffset = int2(0, 0),
+	const uint  InputResDivisor = 1)
+{
+	const uint2 TileSize = uint2(LaneStride.x * WAVE_COUNT_X * LANE_COUNT_X, LaneStride.y * WAVE_COUNT_Y * LANE_COUNT_Y);
+	const uint2 TileStride = uint2(TileSize - 2 * TileOverscan);
+	
+	tsr_short2 InputPixelOffsetMin = tsr_short2(PixelViewportMin / InputResDivisor);
+	tsr_short2 InputPixelOffsetMax = tsr_short2((PixelViewportMax + (InputResDivisor - 1)) / InputResDivisor - 1);
+	
+	tsr_short2 InputTileOffset = tsr_short2(TileStride / InputResDivisor) * tsr_short2(GroupId) - tsr_short2(TileOverscan / InputResDivisor);
+	tsr_short2 InputPixelOffsetInTile = GetElementPixelOffsetInTile(
+		LaneStride.x / InputResDivisor,
+		LaneStride.y / InputResDivisor,
+		GroupThreadIndex,
+		InputElementIndex);
+
+	return fastClamp(InputTileOffset + InputPixelOffsetInTile, InputPixelOffsetMin, InputPixelOffsetMax);
+}
+
+CALL_SITE_DEBUGLOC
+tsr_short2 ComputeElementOutputPixelPos(
+	const uint2 LaneStride,
+	const uint2 TileOverscan,
+	const uint2 PixelViewportMin,
+	const uint2 PixelViewportMax,
+	const uint2 GroupId,
+	const uint  GroupThreadIndex,
+	const uint  OutputElementIndex,
+	const int2  OutputDataOffset = int2(0, 0),
+	const uint  OutputResDivisor = 1,
+	const bool  bDebugDisablePadding = false)
+{
+	const uint2 TileSize = uint2(LaneStride.x * WAVE_COUNT_X * LANE_COUNT_X, LaneStride.y * WAVE_COUNT_Y * LANE_COUNT_Y);
+	const uint2 TileStride = bDebugDisablePadding ? TileSize : uint2(TileSize - 2 * TileOverscan);
+	
+	tsr_short2 OutputPixelOffsetMin = tsr_short2(PixelViewportMin / OutputResDivisor);
+	tsr_short2 OutputPixelOffsetMax = tsr_short2((PixelViewportMax + (OutputResDivisor - 1)) / OutputResDivisor);
+
+	const uint2 DataMinInTile = TileOverscan - uint2(OutputDataOffset);
+	const uint2 DataMaxInTile = TileSize - uint2(OutputDataOffset) - TileOverscan;
+	
+	tsr_short2 OutputTileOffset = tsr_short2(TileStride / OutputResDivisor) * tsr_short2(GroupId) - tsr_short2((bDebugDisablePadding ? 0u : TileOverscan) / OutputResDivisor);
+	tsr_short2 OutputPixelOffsetInTile = GetElementPixelOffsetInTile(
+		LaneStride.x / OutputResDivisor,
+		LaneStride.y / OutputResDivisor,
+		GroupThreadIndex,
+		OutputElementIndex);
+
+	tsr_short2 OutputPixelOffsetInOutputTexture = OutputTileOffset + OutputPixelOffsetInTile;
+	
+	bool2 bIsWithinOutputViewport = OutputPixelOffsetInOutputTexture < OutputPixelOffsetMax;
+	bool2 bIsNotPadding = or(bDebugDisablePadding.xx, and(
+		OutputPixelOffsetInTile * tsr_short(OutputResDivisor) >= tsr_short2(DataMinInTile),
+		OutputPixelOffsetInTile * tsr_short(OutputResDivisor) <  tsr_short2(DataMaxInTile)));
+
+	bool2 bIsValid = and(bIsNotPadding, bIsWithinOutputViewport);
+
+	OutputPixelOffsetInOutputTexture = select(bIsValid, OutputPixelOffsetInOutputTexture, ~tsr_short(0).xx);
+
+	return OutputPixelOffsetInOutputTexture;
+}
+
 
 //------------------------------------------------------- GROUP SHARED COMMUNICATION
 
diff --git a/Engine/Shaders/Private/TemporalSuperResolution/TSRRejectShading.usf b/Engine/Shaders/Private/TemporalSuperResolution/TSRRejectShading.usf
--- a/Engine/Shaders/Private/TemporalSuperResolution/TSRRejectShading.usf
+++ b/Engine/Shaders/Private/TemporalSuperResolution/TSRRejectShading.usf
@@ -17,6 +17,8 @@
 
 #define DEBUG_ARRAY_SIZE 8
 
+#define DEBUG_DISABLE_PADDING 0
+
 /** Whether to compose translucency to the input. */
 #define CONFIG_COMPOSE_TRANSLUCENCY 1
 
@@ -126,23 +128,285 @@ RWTexture2D<tsr_halfC> InputSceneColorOutput;
 RWTexture2D<tsr_half> InputSceneColorLdrLumaOutput;
 
 
+//------------------------------------------------------- DEBUG
+
+static tsr_simd_halfC Debug[DEBUG_ARRAY_SIZE];
+
+
 //------------------------------------------------------- GENERAL OPERATIONS
 
-tsr_short2 ComputeInputPixelPos(const tsr_short2 GroupPixelOffset, const uint LaneIndex, const uint SimdIndex)
+void ComputeMoireError(
+	tsr_simd_bool  bIsValid,
+	tsr_simd_half  bIsStatic,
+	tsr_simd_half  MoireInput,
+	tsr_simd_halfC OriginalOpaqueInput,
+	tsr_simd_half4 PrevMoireHistory,
+	out tsr_simd_half4 OutMoireHistory,
+	out tsr_simd_half  OutMoireError)
 {
-	return GroupPixelOffset + GetLaneSimdPixelOffset<LANE_STRIDE_X, LANE_STRIDE_Y>(LaneIndex, SimdIndex);
-}
+	const tsr_half BlendFinal = tsr_half(0.05);
+	const tsr_half MoireEncodingError = tsr_half(rcp(127.0));
+	const tsr_half MaxPrevTotalVariationCount = tsr_half(20.0);
+
+	// Compute the threshold at which the variation count should quick in
+	const tsr_half TotalVariationCountThreshold = tsr_half(1.0 / (1.0 - pow(1.0 - BlendFinal, tsr_half(FlickeringFramePeriod))));
+
+	// Input luminance to monitor flickering from.
+	tsr_simd_half InputLuma = MoireInput;
+
+	// How much input luminance from the TSRComputeMoireLuma.usf measurement has been modified since for instance with fog, water.
+	tsr_simd_half InputLumaModification;
+	{
+		const tsr_half InputLumaEncodingError = tsr_half(0.5 / 255.0);
+			
+		OriginalOpaqueInput = OriginalOpaqueInput * rcp(OriginalOpaqueInput + tsr_simd_halfC::Const(tsr_half(PerceptionAdd)));
+
+
+		tsr_simd_half SceneColorLuma;
+		UNROLL_N(SIMD_SIZE)
+		for (uint SimdIndex = 0; SimdIndex < SIMD_SIZE; SimdIndex++)
+		{
+			SceneColorLuma.SetElement(SimdIndex, dot(OriginalOpaqueInput.GetElement(SimdIndex).rgb, kMoireLumaWeights));
+		}
+
+		InputLumaModification = max(abs(InputLuma - SceneColorLuma) - tsr_simd_half::Const(InputLumaEncodingError), tsr_simd_half::Const(0.0));
+	}
+		
+	// Convolve the input luma to have the exact same moire pattern as FilteredInput below.
+	InputLuma = Blur3x3(InputLuma);
+	InputLumaModification = Blur3x3(InputLumaModification);
+
+	// Unpack history
+	tsr_simd_half PrevFlickeringHistory = PrevMoireHistory[0];
+	tsr_simd_half PrevGradient = PrevMoireHistory[1] * tsr_half(255.0 / 127.0) - tsr_half(1.0);
+	//tsr_simd_half PrevTotalVariation = PrevMoireHistory[2] * tsr_simd_half::Const(1.0 - BlendFinal);
+	tsr_simd_half PrevTotalVariation = PrevMoireHistory[2] * tsr_simd_half::Const(tsr_half(1.0) - BlendFinal);
+	tsr_simd_half PrevTotalVariationCount = PrevMoireHistory[3] * MaxPrevTotalVariationCount * tsr_simd_half::Const(tsr_half(1.0) - BlendFinal);
+		
+	// Discard history on moving objects.
+	PrevTotalVariation = PrevTotalVariation * bIsStatic;
+	PrevTotalVariationCount = PrevTotalVariationCount * bIsStatic;
+
+
+	// Compute clamping box of the limance.
+	tsr_simd_half InputLumaMin;
+	tsr_simd_half InputLumaMax;
+	MinMax3x3(InputLuma, /* out */ InputLumaMin, /* out */ InputLumaMax);
+
+	// Clamp and contribue the history.		
+	tsr_simd_half ClampedPrevFlickeringHistory = fastClamp(PrevFlickeringHistory, InputLumaMin, InputLumaMax);
+	tsr_simd_half FinalFlickeringHistory = InputLuma * tsr_simd_half::Const(BlendFinal) - ClampedPrevFlickeringHistory * tsr_simd_half::Const(BlendFinal - tsr_half(1.0));
+	
+	// Measure the flickering gradient of the current frame.
+	//tsr_simd_half CurrentGradient = FinalFlickeringHistory - PrevFlickeringHistory;
+	tsr_simd_half CurrentGradient = ClampedPrevFlickeringHistory - PrevFlickeringHistory;
+		
+	// Remove how much luma has changed from luma measurement to final scene color from the gradient.
+	for (uint SimdIndex = 0; SimdIndex < SIMD_SIZE; SimdIndex++)
+	{
+		tsr_half LocalCurrentGradient = CurrentGradient.GetElement(SimdIndex);
+		tsr_half LocalInputLumaModification = InputLumaModification.GetElement(SimdIndex);
+			
+		CurrentGradient.SetElement(SimdIndex, clamp(LocalCurrentGradient * tsr_half(POSITIVE_INFINITY), tsr_half(-1.0), tsr_half(1.0)) * max(abs(LocalCurrentGradient) - LocalInputLumaModification, tsr_half(0.0)));
+	}
+
+	#if 0
+		tsr_simd_half InputBoxSize = InputLumaMax - InputLumaMin;
+		CurrentGradient = sign(CurrentGradient) * max(abs(CurrentGradient) - tsr_simd_half::Const(0.5) * InputBoxSize, tsr_simd_half::Const(0.0));
+	#endif
+
+	tsr_simd_half IsFlicker;
+	UNROLL_N(SIMD_SIZE)
+	for (uint SimdIndex = 0; SimdIndex < SIMD_SIZE; SimdIndex++)
+	{
+		tsr_half LocalCurrentGradient = CurrentGradient.GetElement(SimdIndex);
+		tsr_half LocalPrevGradient = PrevGradient.GetElement(SimdIndex);
+
+		bool bGradientsAreSameSign = LocalCurrentGradient * LocalPrevGradient > tsr_half(0.0);
+		bool bGradientsAreWithinError = abs(LocalCurrentGradient) < MoireEncodingError || abs(LocalPrevGradient) < MoireEncodingError;
+		//bool bCurrentGradientSubstentialEnough = abs(LocalCurrentGradient) > View.GeneralPurposeTweak * abs(LocalPrevGradient); // - 2.0 * MoireEncodingError;
+
+		tsr_half LocalIsFlicker = select((
+			bGradientsAreSameSign ||
+			bGradientsAreWithinError ||
+			//!bCurrentGradientSubstentialEnough ||
+			bCameraCut), tsr_half(0.0), tsr_half(1.0));
+
+		IsFlicker.SetElement(SimdIndex, LocalIsFlicker);
+	}
+
+	// Compute the total variation of the gradient over time.
+	tsr_simd_half GradientVariation = min(abs(PrevGradient), abs(CurrentGradient)) * IsFlicker;
+
+	// Dilates the total variation and contribution to keep stability despite inacuracies in the history reprojection
+	// of these need a nearest filter which isn't precise at all.
+	tsr_simd_half TotalVariationContribution = Max3x3(GradientVariation);
+	tsr_simd_half TotalVariationCountContribution = Max3x3(IsFlicker);
+
+	UNROLL_N(SIMD_SIZE)
+	for (uint SimdIndex = 0; SimdIndex < SIMD_SIZE; SimdIndex++)
+	{
+		tsr_half LocalCurrentGradient = CurrentGradient.GetElement(SimdIndex);
+
+		tsr_half LocalPrevGradient = PrevGradient.GetElement(SimdIndex);
+		tsr_half LocalPrevTotalVariation = PrevTotalVariation.GetElement(SimdIndex);
+		tsr_half LocalPrevTotalVariationCount = PrevTotalVariationCount.GetElement(SimdIndex);
+		bool LocalIsValid = bIsValid.GetElement(SimdIndex);
 
-tsr_short2 ComputeOutputPixelPos(const tsr_short2 GroupPixelOffset, const uint LaneIndex, const uint SimdIndex)
+		tsr_half LocalIsFlicker = IsFlicker.GetElement(SimdIndex);
+
+		tsr_half LocalTotalVariationContribution = TotalVariationContribution.GetElement(SimdIndex);
+		tsr_half LocalTotalVariationCountContribution = TotalVariationCountContribution.GetElement(SimdIndex);
+			
+		// Blend out previous gradient
+		tsr_half LocalPrevBlendedGradient = LocalPrevGradient * (tsr_half(1.0) - BlendFinal) * (tsr_half(1.0) - LocalIsFlicker);
+
+		// Compute new gradient.
+		#if 0
+			tsr_half LocalNewGradient = LocalPrevBlendedGradient + LocalCurrentGradient;
+		#else
+			tsr_half LocalNewGradient = abs(LocalPrevBlendedGradient) > abs(LocalCurrentGradient) ? LocalPrevBlendedGradient : LocalCurrentGradient;
+		#endif
+			
+		// Accumulate current frame variation.
+		tsr_half LocalNewTotalVariation = LocalPrevTotalVariation + LocalTotalVariationContribution;
+		tsr_half LocalNewTotalVariationCount = LocalPrevTotalVariationCount + LocalTotalVariationCountContribution;
+
+		// Discard all moire history on parallax disocclusion.
+		LocalNewGradient = select(LocalIsValid, LocalNewGradient, tsr_half(0.0));
+		LocalNewTotalVariation = select(LocalIsValid, LocalNewTotalVariation, tsr_half(0.0));
+		LocalNewTotalVariationCount = select(LocalIsValid, LocalNewTotalVariationCount, tsr_half(0.0));
+
+		// Quantise total variation and count to history bit depth to variation drift in the history over time.
+		tsr_half LocalQuantizedNewTotalVariationCount = floor(LocalNewTotalVariationCount * (tsr_half(255.0) / MaxPrevTotalVariationCount)) * (MaxPrevTotalVariationCount / tsr_half(255.0));
+		LocalNewTotalVariation *= LocalQuantizedNewTotalVariationCount * SafeRcp(LocalNewTotalVariationCount);
+
+		// Compute the final luminance 
+		#if 1
+			tsr_half LocalCountFadeIn = saturate(LocalNewTotalVariationCount / TotalVariationCountThreshold - tsr_half(0.5));
+		#else
+			tsr_half LocalCountFadeIn = saturate(LocalNewTotalVariationCount - TotalVariationCountThreshold);
+		#endif
+		tsr_half LocalMoireError = (abs(LocalNewTotalVariation * SafeRcp(LocalNewTotalVariationCount)) + LocalNewTotalVariationCount * MoireEncodingError) * LocalCountFadeIn;
+
+		OutMoireHistory.RegisterArray[SimdIndex][0] = FinalFlickeringHistory.GetElement(SimdIndex);
+		OutMoireHistory.RegisterArray[SimdIndex][1] = LocalNewGradient * tsr_half(127.0 / 255.0) + tsr_half(127.0 / 255.0);
+		OutMoireHistory.RegisterArray[SimdIndex][2] = LocalNewTotalVariation;
+		OutMoireHistory.RegisterArray[SimdIndex][3] = LocalQuantizedNewTotalVariationCount * rcp(MaxPrevTotalVariationCount);
+
+		OutMoireError.SetElement(SimdIndex, LocalMoireError);
+	}
+		
+	// Moving object in the scene have very low chance of creating moire pattern with TAA jitter, so discard the moire error to still
+	// have character animation ghost free evem as environement shadow project different things on character.
+	OutMoireError = OutMoireError * bIsStatic;
+} // ComputeMoireError()
+
+template<bool ApplyMoireError>
+void MeasureRejection(
+	tsr_simd_bool  bIsValid,
+	tsr_simd_halfC OriginalInput,
+	tsr_simd_halfC OriginalHistory,
+	tsr_simd_half  MoireError,
+	out tsr_simd_half OutRejectionBlendFinal,
+	out tsr_simd_half OutRejectionClampBlend)
 {
-	tsr_short2 LaneSimdPixelOffset = GetLaneSimdPixelOffset<LANE_STRIDE_X, LANE_STRIDE_Y>(LaneIndex, SimdIndex);
-	bool bIsNonOverscanPixel = all(and(LaneSimdPixelOffset >= TileOverscan, LaneSimdPixelOffset < (GROUP_TILE_SIZE - TileOverscan)));
+	const tsr_half FilteringWeight = rcp(tsr_half(1.0 + 4 * 0.5 + 4 * 0.25));
+	tsr_simd_halfC BackbufferQuantizationErrorVector = tsr_simd_halfC::Const(MeasureBackbufferLDRQuantizationError());
+		
+	// Convert input to LDR
+	tsr_simd_halfC ExposedInput = OriginalInput * rcp(OriginalInput + tsr_simd_halfC::Const(tsr_half(PerceptionAdd)));
+	
+	// Change the exposure of history for the shading rejection
+	tsr_simd_halfC ExposedHistory = OriginalHistory;
+	{
+		ExposedHistory = ExposedHistory * rcp(tsr_simd_halfC::Const(1.0) - ExposedHistory);
+		ExposedHistory = ExposedHistory * rcp(ExposedHistory + tsr_simd_halfC::Const(tsr_half(PerceptionAdd)));
+	}
 
-	tsr_short2 OutputPixelPos = InvalidateOutputPixelPos(GroupPixelOffset + LaneSimdPixelOffset, InputInfo_ViewportMax);
-	OutputPixelPos.x = select(bIsNonOverscanPixel, OutputPixelPos.x, ~tsr_short(0));
+	// Mutual Annihilation
+	tsr_simd_halfC ClampedInput   = AnnihilateToGuide3x3(ExposedInput  , /* Guide = */ ExposedHistory);
+	tsr_simd_halfC ClampedHistory = AnnihilateToGuide3x3(ExposedHistory, /* Guide = */ ExposedInput);
+		
+	// Filter clamped input and history to increase quantization of the rejection detection
+	tsr_simd_halfC FilteredInput   = Blur3x3(ClampedInput);
+	tsr_simd_halfC FilteredHistory = Blur3x3(ClampedHistory);
 	
-	return OutputPixelPos;
-}
+	// Build the clamping box of the filtered signal
+	tsr_simd_halfC FilteredBoxMin;
+	tsr_simd_halfC FilteredBoxMax;
+	MinMax3x3(FilteredInput, /* out */ FilteredBoxMin, /* out */ FilteredBoxMax);
+	FilteredBoxMin = FilteredBoxMin - BackbufferQuantizationErrorVector;
+	FilteredBoxMax = FilteredBoxMax + BackbufferQuantizationErrorVector;
+		
+	// Clamp filtered signal
+	tsr_simd_halfC ClampedFilteredHistory = fastClamp(FilteredHistory, FilteredBoxMin, FilteredBoxMax);
+	if (ApplyMoireError)
+	{
+		tsr_simd_halfC MoireErrorSize;
+			
+		UNROLL_N(SIMD_SIZE)
+		for (uint SimdIndex = 0; SimdIndex < SIMD_SIZE; SimdIndex++)
+		#if 0
+		{
+			tsr_halfC ClampedEnergy = abs(FlickeringClampedFilteredHistory.GetElement(SimdIndex) - FilteredHistory.GetElement(SimdIndex));
+			tsr_halfC ClampedEnergyAmount = ClampedEnergy * SafeRcp(ClampedEnergy[0] + ClampedEnergy[1] + ClampedEnergy[2]);
+
+			// Apply the moire error on each individual channel based on how much they need to clamp.
+			MoireErrorSize.SetElement(SimdIndex, ClampedEnergyAmount * (kMoireLumaToChannel * MoireError.GetElement(SimdIndex)));
+		}
+		#else
+		{
+			MoireErrorSize.SetElement(SimdIndex, kMoireLumaToChannel * MoireError.GetElement(SimdIndex));
+		}
+		#endif
+
+		tsr_simd_halfC StableFilteredBoxMin = min(FilteredBoxMin, FilteredBoxMax - MoireErrorSize);
+		tsr_simd_halfC StableFilteredBoxMax = max(FilteredBoxMax, FilteredBoxMin + MoireErrorSize);
+			
+		ClampedFilteredHistory = fastClamp(FilteredHistory, StableFilteredBoxMin, StableFilteredBoxMax);
+	}
+		
+	// Compute the box size.
+	tsr_simd_halfC BoxSize;
+	{
+		BoxSize = (
+			MaxMinusMin3x3(ClampedInput) * tsr_simd_halfC::Const(FilteringWeight) +
+			(BackbufferQuantizationErrorVector * tsr_half(2 * FilteringWeight)));
+
+		if (ApplyMoireError)
+		{
+			tsr_simd_halfC MoireErrorSize;
+			UNROLL_N(SIMD_SIZE)
+			for (uint SimdIndex = 0; SimdIndex < SIMD_SIZE; SimdIndex++)
+			{
+				tsr_halfC ClampedEnergy = abs(ClampedFilteredHistory.GetElement(SimdIndex) - FilteredHistory.GetElement(SimdIndex));
+				tsr_halfC ClampedEnergyAmount = ClampedEnergy * SafeRcp(ClampedEnergy[0] + ClampedEnergy[1] + ClampedEnergy[2]);
+
+				// Apply the moire error on each individual channel based on how much they need to clamp.
+				MoireErrorSize.SetElement(SimdIndex, ClampedEnergyAmount * ((FilteringWeight * tsr_half(3.0)) * MoireError.GetElement(SimdIndex)));
+			}
+			BoxSize = max(BoxSize, MoireErrorSize);
+		}
+	}
+
+	tsr_simd_halfC RawClampedEnergy      = abs(ClampedFilteredHistory - FilteredHistory);
+	tsr_simd_halfC FilteredClampedEnergy = Max3x3(Median3x3(RawClampedEnergy));
+
+	tsr_simd_halfC Delta = abs(FilteredInput - FilteredHistory);
+	tsr_simd_halfC ClampError = max(
+		BoxSize * tsr_half(0.25),
+		BackbufferQuantizationErrorVector);
+			
+	tsr_simd_halfC RawFactor      = saturate(tsr_simd_halfC::Const(1.0) - max(RawClampedEnergy - ClampError, tsr_simd_halfC::Const(0.0)) * rcp(max(Delta, BoxSize)));
+	tsr_simd_halfC FilteredFactor = saturate(tsr_simd_halfC::Const(1.0) - max(FilteredClampedEnergy - ClampError, tsr_simd_halfC::Const(0.0)) * rcp(max(Delta, BoxSize)));
+		
+	tsr_simd_half RawRejection      = min3(RawFactor[0], RawFactor[1], RawFactor[2]);
+	tsr_simd_half FilteredRejection = min3(FilteredFactor[0], FilteredFactor[1], FilteredFactor[2]);
+
+	OutRejectionBlendFinal = saturate(FilteredRejection);
+	OutRejectionClampBlend = Min3x3(Median3x3(RawRejection));
+} // MeasureRejection()
 
 
 //------------------------------------------------------- ENTRY POINT
@@ -157,7 +421,6 @@ void MainCS(
 {
 	GGroupThreadIndex = GroupThreadIndex;
 
-	tsr_simd_halfC Debug[DEBUG_ARRAY_SIZE];
 	UNROLL_N(DEBUG_ARRAY_SIZE)
 	for (uint DebugId = 0; DebugId < DEBUG_ARRAY_SIZE; DebugId++)
 	{
@@ -179,7 +442,7 @@ void MainCS(
 	tsr_simd_bool bIsValid;
 	#if CONFIG_MOIRE_DETECTION
 		tsr_simd_half MoireInput;
-		tsr_simd_half4 PrevMoire;
+		tsr_simd_half4 PrevMoireHistory;
 		tsr_simd_half bIsStatic;
 	#endif
 	PLATFORM_SPECIFIC_ISOLATE
@@ -187,10 +450,16 @@ void MainCS(
 		UNROLL_N(SIMD_SIZE)
 		for (uint SimdIndex = 0; SimdIndex < SIMD_SIZE; SimdIndex++)
 		{
-			tsr_short2 ClampedFetchPixelPos = ComputeInputPixelPos(GroupPixelOffset, LaneIndex, SimdIndex);
-			float2 TranslucencyTextureUV = ApplyScreenTransform(float2(ClampedFetchPixelPos), InputPixelPosToTranslucencyTextureUV);
+			tsr_short2 ClampedFetchPixelPos = ComputeElementInputPixelPos(
+				/* LaneStride = */ uint2(LANE_STRIDE_X, LANE_STRIDE_Y),
+				TileOverscan,
+				InputInfo_ViewportMin,
+				InputInfo_ViewportMax,
+				GroupId,
+				GroupThreadIndex,
+				SimdIndex);
 
-			ClampedFetchPixelPos = ClampPixelOffset(ClampedFetchPixelPos, InputPixelPosMin, InputPixelPosMax);
+			float2 TranslucencyTextureUV = ApplyScreenTransform(float2(ClampedFetchPixelPos), InputPixelPosToTranslucencyTextureUV);
 			TranslucencyTextureUV = fastClamp(TranslucencyTextureUV, TranslucencyTextureUVMin, TranslucencyTextureUVMax);
 
 			OriginalOpaqueInput.SetElement(SimdIndex, InputTexture[ClampedFetchPixelPos]);
@@ -201,7 +470,7 @@ void MainCS(
 			bIsValid.SetElement(SimdIndex, ParallaxRejectionMaskTexture[ClampedFetchPixelPos] > tsr_half(PARALLAX_REJECTION_MASK_THRESHOLD));
 			#if CONFIG_MOIRE_DETECTION
 				MoireInput.SetElement(SimdIndex, InputMoireLumaTexture[ClampedFetchPixelPos]);
-				PrevMoire.SetElement(SimdIndex, ReprojectedHistoryMoireTexture[ClampedFetchPixelPos]);
+				PrevMoireHistory.SetElement(SimdIndex, ReprojectedHistoryMoireTexture[ClampedFetchPixelPos]);
 				bIsStatic.SetElement(SimdIndex, tsr_half(1.0) - tsr_half(IsMovingMaskTexture[ClampedFetchPixelPos]) * rcp(tsr_half(255.0)));
 			#endif
 		}
@@ -239,7 +508,17 @@ void MainCS(
 			UNROLL_N(SIMD_SIZE)
 			for (uint SimdIndex = 0; SimdIndex < SIMD_SIZE; SimdIndex++)
 			{
-				const tsr_short2 OutputPixelPos = ComputeOutputPixelPos(GroupPixelOffset, LaneIndex, SimdIndex);
+				const tsr_short2 OutputPixelPos = ComputeElementOutputPixelPos(
+					/* LaneStride = */ uint2(LANE_STRIDE_X, LANE_STRIDE_Y),
+					TileOverscan,
+					InputInfo_ViewportMin,
+					InputInfo_ViewportMax,
+					GroupId,
+					GroupThreadIndex,
+					SimdIndex,
+					/* OutputDataOffset = */ int2(0, 0),
+					/* OutputResDivisor = */ 1,
+					/* bDebugDisablePadding = */ DEBUG_DISABLE_PADDING != 0);
 		
 				InputSceneColorOutput[OutputPixelPos] = CenterFinalSceneColor.GetElement(SimdIndex);
 				InputSceneColorLdrLumaOutput[OutputPixelPos] = CenterLumaLDR.GetElement(SimdIndex);
@@ -276,178 +555,39 @@ void MainCS(
 	#endif
 
 	// Analyses the moire pattern overtime.
+	tsr_simd_half  MoireError = tsr_simd_half::Const(0.0);
 	#if CONFIG_MOIRE_DETECTION
-	tsr_simd_half4 FinalMoire = tsr_simd_half4::Const(0.0);
-	tsr_simd_half MoireError = tsr_simd_half::Const(0.0);
-	{
-		const tsr_half BlendFinal = tsr_half(0.05);
-		const tsr_half MoireEncodingError = tsr_half(rcp(127.0));
-		const tsr_half MaxPrevTotalVariationCount = tsr_half(20.0);
-
-		// Compute the threshold at which the variation count should quick in
-		const tsr_half TotalVariationCountThreshold = tsr_half(1.0 / (1.0 - pow(1.0 - BlendFinal, tsr_half(FlickeringFramePeriod))));
-
-		// Input luminance to monitor flickering from.
-		tsr_simd_half InputLuma = MoireInput;
-
-		// How much input luminance from the TSRComputeMoireLuma.usf measurement has been modified since for instance with fog, water.
-		tsr_simd_half InputLumaModification;
-		{
-			const tsr_half InputLumaEncodingError = tsr_half(0.5 / 255.0);
-			
-			OriginalOpaqueInput = OriginalOpaqueInput * rcp(OriginalOpaqueInput + tsr_simd_halfC::Const(tsr_half(PerceptionAdd)));
-
-
-			tsr_simd_half SceneColorLuma;
-			UNROLL_N(SIMD_SIZE)
-			for (uint SimdIndex = 0; SimdIndex < SIMD_SIZE; SimdIndex++)
-			{
-				SceneColorLuma.SetElement(SimdIndex, dot(OriginalOpaqueInput.GetElement(SimdIndex).rgb, kMoireLumaWeights));
-			}
-
-			InputLumaModification = max(abs(InputLuma - SceneColorLuma) - tsr_simd_half::Const(InputLumaEncodingError), tsr_simd_half::Const(0.0));
-		}
-		
-		// Convolve the input luma to have the exact same moire pattern as FilteredInput below.
-		InputLuma = Blur3x3(InputLuma);
-		InputLumaModification = Blur3x3(InputLumaModification);
-
-		// Unpack history
-		tsr_simd_half PrevFlickeringHistory = PrevMoire[0];
-		tsr_simd_half PrevGradient = PrevMoire[1] * tsr_half(255.0 / 127.0) - tsr_half(1.0);
-		//tsr_simd_half PrevTotalVariation = PrevMoire[2] * tsr_simd_half::Const(1.0 - BlendFinal);
-		tsr_simd_half PrevTotalVariation = PrevMoire[2] * tsr_simd_half::Const(tsr_half(1.0) - BlendFinal);
-		tsr_simd_half PrevTotalVariationCount = PrevMoire[3] * MaxPrevTotalVariationCount * tsr_simd_half::Const(tsr_half(1.0) - BlendFinal);
-		
-		// Discard history on moving objects.
-		PrevTotalVariation = PrevTotalVariation * bIsStatic;
-		PrevTotalVariationCount = PrevTotalVariationCount * bIsStatic;
-
-
-		// Compute clamping box of the limance.
-		tsr_simd_half InputLumaMin;
-		tsr_simd_half InputLumaMax;
-		MinMax3x3(InputLuma, /* out */ InputLumaMin, /* out */ InputLumaMax);
-
-		// Clamp and contribue the history.		
-		tsr_simd_half ClampedPrevFlickeringHistory = fastClamp(PrevFlickeringHistory, InputLumaMin, InputLumaMax);
-		tsr_simd_half FinalFlickeringHistory = InputLuma * tsr_simd_half::Const(BlendFinal) - ClampedPrevFlickeringHistory * tsr_simd_half::Const(BlendFinal - tsr_half(1.0));
-	
-		// Measure the flickering gradient of the current frame.
-		//tsr_simd_half CurrentGradient = FinalFlickeringHistory - PrevFlickeringHistory;
-		tsr_simd_half CurrentGradient = ClampedPrevFlickeringHistory - PrevFlickeringHistory;
-		
-		// Remove how much luma has changed from luma measurement to final scene color from the gradient.
-		for (uint SimdIndex = 0; SimdIndex < SIMD_SIZE; SimdIndex++)
-		{
-			tsr_half LocalCurrentGradient = CurrentGradient.GetElement(SimdIndex);
-			tsr_half LocalInputLumaModification = InputLumaModification.GetElement(SimdIndex);
-			
-			CurrentGradient.SetElement(SimdIndex, clamp(LocalCurrentGradient * tsr_half(POSITIVE_INFINITY), tsr_half(-1.0), tsr_half(1.0)) * max(abs(LocalCurrentGradient) - LocalInputLumaModification, tsr_half(0.0)));
-		}
-
-		#if 0
-			tsr_simd_half InputBoxSize = InputLumaMax - InputLumaMin;
-			CurrentGradient = sign(CurrentGradient) * max(abs(CurrentGradient) - tsr_simd_half::Const(0.5) * InputBoxSize, tsr_simd_half::Const(0.0));
-		#endif
-
-		tsr_simd_half IsFlicker;
-		UNROLL_N(SIMD_SIZE)
-		for (uint SimdIndex = 0; SimdIndex < SIMD_SIZE; SimdIndex++)
-		{
-			tsr_half LocalCurrentGradient = CurrentGradient.GetElement(SimdIndex);
-			tsr_half LocalPrevGradient = PrevGradient.GetElement(SimdIndex);
-
-			bool bGradientsAreSameSign = LocalCurrentGradient * LocalPrevGradient > tsr_half(0.0);
-			bool bGradientsAreWithinError = abs(LocalCurrentGradient) < MoireEncodingError || abs(LocalPrevGradient) < MoireEncodingError;
-			//bool bCurrentGradientSubstentialEnough = abs(LocalCurrentGradient) > View.GeneralPurposeTweak * abs(LocalPrevGradient); // - 2.0 * MoireEncodingError;
-
-			tsr_half LocalIsFlicker = select((
-				bGradientsAreSameSign ||
-				bGradientsAreWithinError ||
-				//!bCurrentGradientSubstentialEnough ||
-				bCameraCut), tsr_half(0.0), tsr_half(1.0));
-
-			IsFlicker.SetElement(SimdIndex, LocalIsFlicker);
-		}
-
-		// Compute the total variation of the gradient over time.
-		tsr_simd_half GradientVariation = min(abs(PrevGradient), abs(CurrentGradient)) * IsFlicker;
-
-		// Dilates the total variation and contribution to keep stability despite inacuracies in the history reprojection
-		// of these need a nearest filter which isn't precise at all.
-		tsr_simd_half TotalVariationContribution = Max3x3(GradientVariation);
-		tsr_simd_half TotalVariationCountContribution = Max3x3(IsFlicker);
-
-		UNROLL_N(SIMD_SIZE)
-		for (uint SimdIndex = 0; SimdIndex < SIMD_SIZE; SimdIndex++)
-		{
-			tsr_half LocalCurrentGradient = CurrentGradient.GetElement(SimdIndex);
-
-			tsr_half LocalPrevGradient = PrevGradient.GetElement(SimdIndex);
-			tsr_half LocalPrevTotalVariation = PrevTotalVariation.GetElement(SimdIndex);
-			tsr_half LocalPrevTotalVariationCount = PrevTotalVariationCount.GetElement(SimdIndex);
-			bool LocalIsValid = bIsValid.GetElement(SimdIndex);
-
-			tsr_half LocalIsFlicker = IsFlicker.GetElement(SimdIndex);
-
-			tsr_half LocalTotalVariationContribution = TotalVariationContribution.GetElement(SimdIndex);
-			tsr_half LocalTotalVariationCountContribution = TotalVariationCountContribution.GetElement(SimdIndex);
-			
-			// Blend out previous gradient
-			tsr_half LocalPrevBlendedGradient = LocalPrevGradient * (tsr_half(1.0) - BlendFinal) * (tsr_half(1.0) - LocalIsFlicker);
-
-			// Compute new gradient.
-			#if 0
-				tsr_half LocalNewGradient = LocalPrevBlendedGradient + LocalCurrentGradient;
-			#else
-				tsr_half LocalNewGradient = abs(LocalPrevBlendedGradient) > abs(LocalCurrentGradient) ? LocalPrevBlendedGradient : LocalCurrentGradient;
-			#endif
-			
-			// Accumulate current frame variation.
-			tsr_half LocalNewTotalVariation = LocalPrevTotalVariation + LocalTotalVariationContribution;
-			tsr_half LocalNewTotalVariationCount = LocalPrevTotalVariationCount + LocalTotalVariationCountContribution;
-
-			// Discard all moire history on parallax disocclusion.
-			LocalNewGradient = select(LocalIsValid, LocalNewGradient, tsr_half(0.0));
-			LocalNewTotalVariation = select(LocalIsValid, LocalNewTotalVariation, tsr_half(0.0));
-			LocalNewTotalVariationCount = select(LocalIsValid, LocalNewTotalVariationCount, tsr_half(0.0));
-
-			// Quantise total variation and count to history bit depth to variation drift in the history over time.
-			tsr_half LocalQuantizedNewTotalVariationCount = floor(LocalNewTotalVariationCount * (tsr_half(255.0) / MaxPrevTotalVariationCount)) * (MaxPrevTotalVariationCount / tsr_half(255.0));
-			LocalNewTotalVariation *= LocalQuantizedNewTotalVariationCount * SafeRcp(LocalNewTotalVariationCount);
-
-			// Compute the final luminance 
-			#if 1
-				tsr_half LocalCountFadeIn = saturate(LocalNewTotalVariationCount / TotalVariationCountThreshold - tsr_half(0.5));
-			#else
-				tsr_half LocalCountFadeIn = saturate(LocalNewTotalVariationCount - TotalVariationCountThreshold);
-			#endif
-			tsr_half LocalMoireError = (abs(LocalNewTotalVariation * SafeRcp(LocalNewTotalVariationCount)) + LocalNewTotalVariationCount * MoireEncodingError) * LocalCountFadeIn;
-
-			FinalMoire.RegisterArray[SimdIndex][0] = FinalFlickeringHistory.GetElement(SimdIndex);
-			FinalMoire.RegisterArray[SimdIndex][1] = LocalNewGradient * tsr_half(127.0 / 255.0) + tsr_half(127.0 / 255.0);
-			FinalMoire.RegisterArray[SimdIndex][2] = LocalNewTotalVariation;
-			FinalMoire.RegisterArray[SimdIndex][3] = LocalQuantizedNewTotalVariationCount * rcp(MaxPrevTotalVariationCount);
-
-			MoireError.SetElement(SimdIndex, LocalMoireError);
-		}
-		
-		// Moving object in the scene have very low chance of creating moire pattern with TAA jitter, so discard the moire error to still
-		// have character animation ghost free evem as environement shadow project different things on character.
-		MoireError = MoireError * bIsStatic;
-	}
+	tsr_simd_half4 MoireHistory;
+	ComputeMoireError(
+		bIsValid,
+		bIsStatic,
+		MoireInput,
+		OriginalOpaqueInput,
+		PrevMoireHistory,
+		/* out */ MoireHistory,
+		/* out */ MoireError);
 	#endif
-
+	
 	// Output some data early.
 	PLATFORM_SPECIFIC_ISOLATE
 	{
 		UNROLL_N(SIMD_SIZE)
 		for (uint SimdIndex = 0; SimdIndex < SIMD_SIZE; SimdIndex++)
 		{
-			const tsr_short2 OutputPixelPos = ComputeOutputPixelPos(GroupPixelOffset, LaneIndex, SimdIndex);
+			const tsr_short2 OutputPixelPos = ComputeElementOutputPixelPos(
+				/* LaneStride = */ uint2(LANE_STRIDE_X, LANE_STRIDE_Y),
+				TileOverscan,
+				InputInfo_ViewportMin,
+				InputInfo_ViewportMax,
+				GroupId,
+				GroupThreadIndex,
+				SimdIndex,
+				/* OutputDataOffset = */ int2(0, 0),
+				/* OutputResDivisor = */ 1,
+				/* bDebugDisablePadding = */ DEBUG_DISABLE_PADDING != 0);
+
 			#if CONFIG_MOIRE_DETECTION
-				HistoryMoireOutput[OutputPixelPos] = FinalMoire.GetElement(SimdIndex);
+				HistoryMoireOutput[OutputPixelPos] = MoireHistory.GetElement(SimdIndex);
 			#endif
 		}
 	}
@@ -455,107 +595,13 @@ void MainCS(
 	// Measure rejection
 	tsr_simd_half RejectionBlendFinal = tsr_simd_half::Const(1.0);
 	tsr_simd_half RejectionClampBlend = tsr_simd_half::Const(1.0);
-	#if 1
-	{
-		const tsr_half FilteringWeight = rcp(tsr_half(1.0 + 4 * 0.5 + 4 * 0.25));
-		tsr_simd_halfC BackbufferQuantizationErrorVector = tsr_simd_halfC::Const(MeasureBackbufferLDRQuantizationError());
-		
-		// Convert input to LDR
-		tsr_simd_halfC ExposedInput = OriginalInput * rcp(OriginalInput + tsr_simd_halfC::Const(tsr_half(PerceptionAdd)));
-	
-		// Change the exposure of history for the shading rejection
-		tsr_simd_halfC ExposedHistory = OriginalHistory;
-		{
-			ExposedHistory = ExposedHistory * rcp(tsr_simd_halfC::Const(1.0) - ExposedHistory);
-			ExposedHistory = ExposedHistory * rcp(ExposedHistory + tsr_simd_halfC::Const(tsr_half(PerceptionAdd)));
-		}
-
-		// Mutual Annihilation
-		tsr_simd_halfC ClampedInput   = AnnihilateToGuide3x3(ExposedInput  , /* Guide = */ ExposedHistory);
-		tsr_simd_halfC ClampedHistory = AnnihilateToGuide3x3(ExposedHistory, /* Guide = */ ExposedInput);
-		
-		// Filter clamped input and history to increase quantization of the rejection detection
-		tsr_simd_halfC FilteredInput   = Blur3x3(ClampedInput);
-		tsr_simd_halfC FilteredHistory = Blur3x3(ClampedHistory);
-	
-		// Build the clamping box of the filtered signal
-		tsr_simd_halfC FilteredBoxMin;
-		tsr_simd_halfC FilteredBoxMax;
-		MinMax3x3(FilteredInput, /* out */ FilteredBoxMin, /* out */ FilteredBoxMax);
-		FilteredBoxMin = FilteredBoxMin - BackbufferQuantizationErrorVector;
-		FilteredBoxMax = FilteredBoxMax + BackbufferQuantizationErrorVector;
-		
-		// Clamp filtered signal
-		tsr_simd_halfC ClampedFilteredHistory;
-		#if CONFIG_MOIRE_DETECTION
-		{
-			tsr_simd_halfC MoireErrorSize;
-			
-			UNROLL_N(SIMD_SIZE)
-			for (uint SimdIndex = 0; SimdIndex < SIMD_SIZE; SimdIndex++)
-			#if 0
-			{
-				tsr_halfC ClampedEnergy = abs(FlickeringClampedFilteredHistory.GetElement(SimdIndex) - FilteredHistory.GetElement(SimdIndex));
-				tsr_halfC ClampedEnergyAmount = ClampedEnergy * SafeRcp(ClampedEnergy[0] + ClampedEnergy[1] + ClampedEnergy[2]);
-
-				// Apply the moire error on each individual channel based on how much they need to clamp.
-				MoireErrorSize.SetElement(SimdIndex, ClampedEnergyAmount * (kMoireLumaToChannel * MoireError.GetElement(SimdIndex)));
-			}
-			#else
-			{
-				MoireErrorSize.SetElement(SimdIndex, kMoireLumaToChannel * MoireError.GetElement(SimdIndex));
-			}
-			#endif
-
-			tsr_simd_halfC StableFilteredBoxMin = min(FilteredBoxMin, FilteredBoxMax - MoireErrorSize);
-			tsr_simd_halfC StableFilteredBoxMax = max(FilteredBoxMax, FilteredBoxMin + MoireErrorSize);
-			
-			ClampedFilteredHistory = fastClamp(FilteredHistory, StableFilteredBoxMin, StableFilteredBoxMax);
-		}
-		#else
-			ClampedFilteredHistory = fastClamp(FilteredHistory, FilteredBoxMin, FilteredBoxMax);
-		#endif
-		
-		// Compute the box size.
-		tsr_simd_halfC BoxSize;
-		{
-			BoxSize = (
-				MaxMinusMin3x3(ClampedInput) * tsr_simd_halfC::Const(FilteringWeight) +
-				(BackbufferQuantizationErrorVector * tsr_half(2 * FilteringWeight)));
-
-			#if CONFIG_MOIRE_DETECTION
-			tsr_simd_halfC MoireErrorSize;
-			UNROLL_N(SIMD_SIZE)
-			for (uint SimdIndex = 0; SimdIndex < SIMD_SIZE; SimdIndex++)
-			{
-				tsr_halfC ClampedEnergy = abs(ClampedFilteredHistory.GetElement(SimdIndex) - FilteredHistory.GetElement(SimdIndex));
-				tsr_halfC ClampedEnergyAmount = ClampedEnergy * SafeRcp(ClampedEnergy[0] + ClampedEnergy[1] + ClampedEnergy[2]);
-
-				// Apply the moire error on each individual channel based on how much they need to clamp.
-				MoireErrorSize.SetElement(SimdIndex, ClampedEnergyAmount * ((FilteringWeight * tsr_half(3.0)) * MoireError.GetElement(SimdIndex)));
-			}
-			BoxSize = max(BoxSize, MoireErrorSize);
-			#endif
-		}
-
-		tsr_simd_halfC RawClampedEnergy      = abs(ClampedFilteredHistory - FilteredHistory);
-		tsr_simd_halfC FilteredClampedEnergy = Max3x3(Median3x3(RawClampedEnergy));
-
-		tsr_simd_halfC Delta = abs(FilteredInput - FilteredHistory);
-		tsr_simd_halfC ClampError = max(
-			BoxSize * tsr_half(0.25),
-			BackbufferQuantizationErrorVector);
-			
-		tsr_simd_halfC RawFactor      = saturate(tsr_simd_halfC::Const(1.0) - max(RawClampedEnergy - ClampError, tsr_simd_halfC::Const(0.0)) * rcp(max(Delta, BoxSize)));
-		tsr_simd_halfC FilteredFactor = saturate(tsr_simd_halfC::Const(1.0) - max(FilteredClampedEnergy - ClampError, tsr_simd_halfC::Const(0.0)) * rcp(max(Delta, BoxSize)));
-		
-		tsr_simd_half RawRejection      = min3(RawFactor[0], RawFactor[1], RawFactor[2]);
-		tsr_simd_half FilteredRejection = min3(FilteredFactor[0], FilteredFactor[1], FilteredFactor[2]);
-
-		RejectionBlendFinal = saturate(FilteredRejection);
-		RejectionClampBlend = Min3x3(Median3x3(RawRejection));
-	}
-	#endif
+	MeasureRejection</* ApplyMoireError = */ CONFIG_MOIRE_DETECTION != 0>(
+		bIsValid,
+		OriginalInput,
+		OriginalHistory,
+		MoireError,
+		/* out */ RejectionBlendFinal,
+		/* out */ RejectionClampBlend);
 	
 	// Update the guide for next frame.
 	tsr_simd_halfC FinalGuide;
@@ -602,9 +648,19 @@ void MainCS(
 		UNROLL_N(SIMD_SIZE)
 		for (uint SimdIndex = 0; SimdIndex < SIMD_SIZE; SimdIndex++)
 		{
-			const tsr_short2 LaneSimdPixelOffset = GetLaneSimdPixelOffset<LANE_STRIDE_X, LANE_STRIDE_Y>(LaneIndex, SimdIndex);
-			const tsr_short2 OutputPixelPos = ComputeOutputPixelPos(GroupPixelOffset, LaneIndex, SimdIndex);
-		
+			const tsr_short2 LaneSimdPixelOffset = GetElementPixelOffsetInTile(LANE_STRIDE_X, LANE_STRIDE_Y, LaneIndex, SimdIndex);
+			const tsr_short2 OutputPixelPos = ComputeElementOutputPixelPos(
+				/* LaneStride = */ uint2(LANE_STRIDE_X, LANE_STRIDE_Y),
+				TileOverscan,
+				InputInfo_ViewportMin,
+				InputInfo_ViewportMax,
+				GroupId,
+				GroupThreadIndex,
+				SimdIndex,
+				/* OutputDataOffset = */ int2(0, 0),
+				/* OutputResDivisor = */ 1,
+				/* bDebugDisablePadding = */ DEBUG_DISABLE_PADDING != 0);
+
 			tsr_half LocalRejectionBlendFinal = RejectionBlendFinal.GetElement(SimdIndex);
 			tsr_half LocalRejectionClampBlend = RejectionClampBlend.GetElement(SimdIndex);
 			tsr_halfC FinalGuideColor = FinalGuide.GetElement(SimdIndex);
