Speeds up the decimation of TSR's resurrected guide

diff --git a/Engine/Shaders/Private/TemporalSuperResolution/TSRDecimateHistory.usf b/Engine/Shaders/Private/TemporalSuperResolution/TSRDecimateHistory.usf
--- a/Engine/Shaders/Private/TemporalSuperResolution/TSRDecimateHistory.usf
+++ b/Engine/Shaders/Private/TemporalSuperResolution/TSRDecimateHistory.usf
@@ -446,9 +446,6 @@ void MainCS(
 	#endif
 	#if CONFIG_RESURRECTION_REPROJECTION
 		tsr_halfC ResurrectedHistoryGuideColor;
-		#if CONFIG_MOIRE_REPROJECTION
-			tsr_half4 ResurrectedHistoryMoireColor;
-		#endif
 	#endif
 	bool bIsOffScreen;
 
@@ -501,18 +498,18 @@ void MainCS(
 
 			float2 KernelHistoryBufferUV = ApplyScreenTransform(ResurrectionScreen, ScreenPosToResurrectionGuideBufferUV);
 
-			ResurrectedHistoryGuideColor = SampleTextureCatrom(
-				PrevHistoryGuide,
-				KernelHistoryBufferUV,
-				ResurrectionFrameIndex,
-				ResurrectionGuideUVViewportBilinearMin,
-				ResurrectionGuideUVViewportBilinearMax);
-			ResurrectedHistoryGuideColor = CorrectGuideColorExposure(ResurrectedHistoryGuideColor, ResurrectionPreExposureCorrection);
-			#if CONFIG_MOIRE_REPROJECTION
-			{
-				ResurrectedHistoryMoireColor = ReprojectPrevHistoryMoire(InputPixelPos, KernelHistoryBufferUV, ResurrectionFrameIndex);
-			}
+			#if 1
+				KernelHistoryBufferUV = fastClamp(KernelHistoryBufferUV, ResurrectionGuideUVViewportBilinearMin, ResurrectionGuideUVViewportBilinearMax);
+				ResurrectedHistoryGuideColor = PrevHistoryGuide.SampleLevel(GlobalBilinearClampedSampler, float3(KernelHistoryBufferUV, ResurrectionFrameIndex), 0);
+			#else
+				ResurrectedHistoryGuideColor = SampleTextureCatrom(
+					PrevHistoryGuide,
+					KernelHistoryBufferUV,
+					ResurrectionFrameIndex,
+					ResurrectionGuideUVViewportBilinearMin,
+					ResurrectionGuideUVViewportBilinearMax);
 			#endif
+			ResurrectedHistoryGuideColor = CorrectGuideColorExposure(ResurrectedHistoryGuideColor, ResurrectionPreExposureCorrection);
 		}
 		#endif
 	}
@@ -616,9 +613,6 @@ void MainCS(
 			#else
 				ReprojectedHistoryGuideOutput[tsr_short3(OutputPixelPos, 1)] = tsr_half4(ResurrectedHistoryGuideColor.xyz, tsr_half(1.0));
 			#endif
-			#if CONFIG_MOIRE_REPROJECTION
-				ReprojectedHistoryMoireOutput[tsr_short3(OutputPixelPos, 1)] = ResurrectedHistoryMoireColor;
-			#endif
 		}
 		#endif
 
diff --git a/Engine/Source/Runtime/Renderer/Private/PostProcess/TemporalSuperResolution.cpp b/Engine/Source/Runtime/Renderer/Private/PostProcess/TemporalSuperResolution.cpp
--- a/Engine/Source/Runtime/Renderer/Private/PostProcess/TemporalSuperResolution.cpp
+++ b/Engine/Source/Runtime/Renderer/Private/PostProcess/TemporalSuperResolution.cpp
@@ -213,12 +213,6 @@ TAutoConsoleVariable<int32> CVarTSRResurrectionPersistentFrameInterval(
 	TEXT("(default=31)"),
 	ECVF_RenderThreadSafe);
 
-TAutoConsoleVariable<int32> CVarTSRResurrectionFlickeringHistory(
-	TEXT("r.TSR.Resurrection.Flickering"), 0,
-	TEXT("Configures whether flickering history of r.TSR.ShadingRejection.Flickering should be resurrected. ")
-	TEXT("(Enabled by default)"),
-	ECVF_RenderThreadSafe);
-
 TAutoConsoleVariable<int32> CVarTSRAsyncCompute(
 	TEXT("r.TSR.AsyncCompute"), 2,
 	TEXT("Controls how TSR run on async compute. Some TSR passes can overlap with previous passes.\n")
@@ -1145,9 +1139,6 @@ FDefaultTemporalUpscaler::FOutputs AddTemporalSuperResolutionPasses(
 	}
 	check(HistorySliceSequence.Check());
 
-	// Whether the flickering history should be resurrected.
-	const bool bResurrectFlickeringHistory = HistorySliceSequence.FrameStorageCount > 1 && CVarTSRResurrectionFlickeringHistory.GetValueOnRenderThread() != 0;
-
 	// Whether to use wave ops optimizations.
 	const ERHIFeatureSupport WaveOpsSupport = FTSRShader::SupportsWaveOps(View.GetShaderPlatform());
 	const bool bSupportsLDS = FTSRShader::SupportsLDS(View.GetShaderPlatform());
@@ -1823,7 +1814,7 @@ FDefaultTemporalUpscaler::FOutputs AddTemporalSuperResolutionPasses(
 				History.MoireArray->Desc.Format,
 				FClearValueBinding::None,
 				/* InFlags = */ TexCreate_ShaderResource | TexCreate_UAV,
-				/* InArraySize = */ bCanResurrectHistory ? 2 : 1);
+				/* InArraySize = */ 1);
 
 			ReprojectedHistoryMoireTexture = GraphBuilder.CreateTexture(Desc, TEXT("TSR.ReprojectedHistoryMoire"));
 		}
