Substrate - enabled new simple volume.
Validated by project  K.

diff --git a/Engine/Shaders/Private/ParticipatingMediaCommon.ush b/Engine/Shaders/Private/ParticipatingMediaCommon.ush
--- a/Engine/Shaders/Private/ParticipatingMediaCommon.ush
+++ b/Engine/Shaders/Private/ParticipatingMediaCommon.ush
@@ -211,7 +211,7 @@ float3 IsotropicMediumSlabPunctualDirectionalAlbedo(FParticipatingMedia PM)
 }
 
 // Wait before activating that feature for now.
-#define USE_NEW_SIMPLEVOLUME 0
+#define USE_NEW_SIMPLEVOLUME 1
 
 float IsotropicMediumSlabDirectionalAlbedoFadeLUT(float MFP)
 {
