Forced shader recompile to fix crash due to BuildInstanceIdBufferAndCommandsFromPrimitiveIdsCs shader not being rebuilt after a permutation was removed in code.

#rnx

diff --git a/Engine/Shaders/Private/InstanceCulling/BuildInstanceDrawCommands.usf b/Engine/Shaders/Private/InstanceCulling/BuildInstanceDrawCommands.usf
--- a/Engine/Shaders/Private/InstanceCulling/BuildInstanceDrawCommands.usf
+++ b/Engine/Shaders/Private/InstanceCulling/BuildInstanceDrawCommands.usf
@@ -1,5 +1,8 @@
 // Copyright Epic Games, Inc. All Rights Reserved.
 
+// Just used to trigger shader recomplies, can be anything but should be regenerated on merge conflicts
+// Version GUID = C869E4C3-E041-4849-A611-6C11A6392EBC
+
 // TODO: should this be 1 by default so compute always has access to it?
 #ifndef USES_PER_INSTANCE_RANDOM
 	#define USES_PER_INSTANCE_RANDOM 1
