Fix for incorrect directional light light-function broken in volumetric fog.

diff --git a/Engine/Shaders/Private/VolumetricFogLightFunction.usf b/Engine/Shaders/Private/VolumetricFogLightFunction.usf
--- a/Engine/Shaders/Private/VolumetricFogLightFunction.usf
+++ b/Engine/Shaders/Private/VolumetricFogLightFunction.usf
@@ -53,7 +53,7 @@ void Main(
 	float4 Hom = mul(float4(TranslatedWorldPosition.xyz, 1), LightFunctionTranslatedWorldToLight);
 	LightVector = Hom.xyz / Hom.w;
 
-#else //LIGHT_TYPE == 2	// Directional and Rect light
+#elif LIGHT_TYPE == 2		// Rect light
 
 	ScreenPosition = (SvPosition.xy * LightFunctionTexelSize - float2(0.5f, 0.5f)) * float2(2, 2);
 	ScreenPosition.xy = frac(ScreenPosition.xy);
@@ -69,6 +69,18 @@ void Main(
 
 	LightVector.xyz = LightVector.zyx;	// Swap to match UV convention
 
+#else //LIGHT_TYPE == 3		// Directional light
+
+	// Directional light
+	ScreenPosition = (SvPosition.xy * LightFunctionTexelSize - float2(0.5f, 0.5f)) * float2(2, -2);
+
+	// make SvPosition appear to be rasterized with the depth from the depth buffer
+	TranslatedWorldPosition = mul(float4(ScreenPosition, Depth, 1), ShadowToTranslatedWorld);
+	TranslatedWorldPosition.xyz /= TranslatedWorldPosition.w;
+
+	float4 Hom = mul(float4(TranslatedWorldPosition.xyz, 1), LightFunctionTranslatedWorldToLight);
+	LightVector = Hom.xyz / Hom.w;
+
 #endif
 	
 	float3 LightFunction = GetLightFunctionColor(LightVector, TranslatedWorldPosition.xyz);
diff --git a/Engine/Source/Runtime/Renderer/Private/VolumetricFogLightFunction.cpp b/Engine/Source/Runtime/Renderer/Private/VolumetricFogLightFunction.cpp
--- a/Engine/Source/Runtime/Renderer/Private/VolumetricFogLightFunction.cpp
+++ b/Engine/Source/Runtime/Renderer/Private/VolumetricFogLightFunction.cpp
@@ -54,7 +54,7 @@ class FVolumetricFogLightFunctionPS : public FMaterialShader
 {
 	DECLARE_SHADER_TYPE(FVolumetricFogLightFunctionPS, Material);
 
-	class FLightType : SHADER_PERMUTATION_INT("LIGHT_TYPE", 3);
+	class FLightType : SHADER_PERMUTATION_INT("LIGHT_TYPE", 4);
 	using FPermutationDomain = TShaderPermutationDomain< FLightType >;
 
 	BEGIN_SHADER_PARAMETER_STRUCT(FParameters, )
@@ -315,7 +315,7 @@ void FSceneRenderer::RenderLightFunctionForVolumetricFog(
 
 					check(DirectionalLightSceneInfo->Proxy->GetLightType() == LightType_Directional)
 					FVolumetricFogLightFunctionPS::FPermutationDomain PermutationVector;
-					PermutationVector.Set<FVolumetricFogLightFunctionPS::FLightType>(2);
+					PermutationVector.Set<FVolumetricFogLightFunctionPS::FLightType>(3);
 					TShaderRef<FVolumetricFogLightFunctionPS> PixelShader = MaterialShaderMap->GetShader<FVolumetricFogLightFunctionPS>(PermutationVector);
 
 					GraphicsPSOInit.BoundShaderState.VertexDeclarationRHI = GFilterVertexDeclaration.VertexDeclarationRHI;
