Work around shader codegen bug on some platforms

diff --git a/Engine/Shaders/Private/InstanceCulling/InstanceCullingOcclusionQuery.usf b/Engine/Shaders/Private/InstanceCulling/InstanceCullingOcclusionQuery.usf
--- a/Engine/Shaders/Private/InstanceCulling/InstanceCullingOcclusionQuery.usf
+++ b/Engine/Shaders/Private/InstanceCulling/InstanceCullingOcclusionQuery.usf
@@ -295,7 +295,6 @@ void DebugMainPS(
 {
 	switch (DebugVisibility)
 	{
-	default:
 	case 0:
 		// Definitely invisible due to HZB test.
 		// No occlusion query performed.
@@ -314,6 +313,8 @@ void DebugMainPS(
 		// Invisible due to occlusion query.
 		OutColor = float4(0.25, 0.0, 0.25, 0.0);
 		break;
+	default:
+		OutColor = (float4)0;
 	}
 }
 
