Fix: Subsurface Profile material in scene capture with r.SubsurfaceScattering set to 0 causes wrong value in alpha channel. Add subsurface scattering check while accumulating the alpha for subsurface profile material.

diff --git a/Engine/Shaders/Private/LightAccumulator.ush b/Engine/Shaders/Private/LightAccumulator.ush
--- a/Engine/Shaders/Private/LightAccumulator.ush
+++ b/Engine/Shaders/Private/LightAccumulator.ush
@@ -102,7 +102,15 @@ float4 LightAccumulator_GetResult(FLightAccumulator In)
 
 		if (SUBSURFACE_CHANNEL_MODE == 1 )
 		{
-			if (View.bCheckerboardSubsurfaceProfileRendering == 0)
+			// bSubsurfacePostprocessEnabled	bCheckerboardSubsurfaceProfileRendering  OutputLuma
+			// 0								0											0
+			// 0								1											0
+			// 1								0											1
+			// 1								1											0
+			// The alpha channel will not be used in the post process subsurface pass if View.bSubsurfacePostprocessEnabled
+			// is not enabled. Add this gate (View.bSubsurfacePostprocessEnabled) so that when subsurface is disabled, scene 
+			// captures will not pick up ScatterableLightLuma to pollute the alpha channel.
+			if (View.bCheckerboardSubsurfaceProfileRendering == 0 && View.bSubsurfacePostprocessEnabled)
 			{
 				// RGB accumulated RGB HDR color, A: specular luminance for screenspace subsurface scattering
 				Ret.a = In.ScatterableLightLuma;
@@ -139,7 +147,7 @@ FDeferredLightingSplit LightAccumulator_GetResultSplit(FLightAccumulator In)
 
 		if (SUBSURFACE_CHANNEL_MODE == 1 )
 		{
-			if (View.bCheckerboardSubsurfaceProfileRendering == 0)
+			if (View.bCheckerboardSubsurfaceProfileRendering == 0 && View.bSubsurfacePostprocessEnabled)
 			{
 				// RGB accumulated RGB HDR color, A: specular luminance for screenspace subsurface scattering
 				RetDiffuse.a = In.ScatterableLightLuma;
