Implement support for LWC in InstanceCullingOcclusionQuery shaders

diff --git a/Engine/Shaders/Private/InstanceCulling/InstanceCullingOcclusionQuery.usf b/Engine/Shaders/Private/InstanceCulling/InstanceCullingOcclusionQuery.usf
--- a/Engine/Shaders/Private/InstanceCulling/InstanceCullingOcclusionQuery.usf
+++ b/Engine/Shaders/Private/InstanceCulling/InstanceCullingOcclusionQuery.usf
@@ -64,11 +64,13 @@ EInstanceVisibilityStatus GetInstanceDataAndVisibility(uint InstanceId, out FIns
 
 	const bool bIsOrtho = IsOrthoProjection(ResolvedView.ViewToClip);
 
+	float4x4 LocalToTranslatedWorld = LWCMultiplyTranslation(InstanceData.LocalToWorld, ResolvedView.PreViewTranslation);
+
 	FFrustumCullData Cull = BoxCullFrustum(
 		BoxOrigin.xyz, 
 		BoxExtent.xyz,
-		LWCHackToFloat(InstanceData.LocalToWorld),
-		ResolvedView.RelativeWorldToClip,
+		LocalToTranslatedWorld,
+		ResolvedView.TranslatedWorldToClip,
 		ResolvedView.ViewToClip, 
 		bIsOrtho, true /* near clip */, false /* skip culling */
 	);
