Works arround shader compilation failure to Metal SM6 that have COMPILER_SUPPORTS_WAVE_VOTE == 0

diff --git a/Engine/Shaders/Private/TemporalSuperResolution/TSRConvolutionNetwork.ush b/Engine/Shaders/Private/TemporalSuperResolution/TSRConvolutionNetwork.ush
--- a/Engine/Shaders/Private/TemporalSuperResolution/TSRConvolutionNetwork.ush
+++ b/Engine/Shaders/Private/TemporalSuperResolution/TSRConvolutionNetwork.ush
@@ -361,7 +361,7 @@ void ReadElementFromLDS(const uint SharedIndex, out tsr_half4 X)
 
 CALL_SITE_DEBUGLOC
 bool AnyThread(bool X)
-#if !COMPILER_SUPPORTS_WAVE_VOTE
+#if !COMPILER_SUPPORTS_WAVE_VOTE && defined(LDS_SIZE)
 {
 	GroupMemoryBarrierWithGroupSync();
 	
@@ -378,11 +378,19 @@ bool AnyThread(bool X)
 	
 	return ReadAtomic(/* SharedIndex = */ 0) > 0;
 }
+#elif !COMPILER_SUPPORTS_WAVE_VOTE && !defined(LDS_SIZE)
+{
+	// Metal SM6 have COMPILER_SUPPORTS_WAVE_VOTE=0 due to shader compiler bug
+	#if !METAL_PROFILE
+		#error Requires to specify LDS_SIZE
+	#endif
+	return true;
+}
 #elif WAVE_COUNT_X == 1 && WAVE_COUNT_Y == 1
 {
 	return WaveActiveAnyTrue(X);
 }
-#elif COMPILER_SUPPORTS_WAVE_VOTE
+#else
 {
 	const uint kLaneCount = WaveGetLaneCount();
 	const uint kThreadCount = LANE_COUNT_X * LANE_COUNT_Y * WAVE_COUNT_X * WAVE_COUNT_Y;
