Fixed error compiling RayTracingPrimaryRays.usf if !PROJECT_SUPPORTS_LUMEN.
- also fixed some typos.

diff --git a/Engine/Shaders/Private/RayTracing/RayTracingPrimaryRays.usf b/Engine/Shaders/Private/RayTracing/RayTracingPrimaryRays.usf
--- a/Engine/Shaders/Private/RayTracing/RayTracingPrimaryRays.usf
+++ b/Engine/Shaders/Private/RayTracing/RayTracingPrimaryRays.usf
@@ -194,7 +194,7 @@ RAY_TRACING_ENTRY_RAYGEN(RayTracingPrimaryRaysRGS)
 		//
 
 		float vertexRadianceWeight = Payload.Opacity;	// Opacity as coverage. This works for RAY_TRACING_BLEND_MODE_OPAQUE and RAY_TRACING_BLEND_MODE_TRANSLUCENT.
-		// It is also needed for RAY_TRACING_BLEND_MODE_ADDITIVE and  RAY_TRACING_BLEND_MODE_ALPHA_COMPOSITE: radiance continbution is alway weighted by coverage.
+		// It is also needed for RAY_TRACING_BLEND_MODE_ADDITIVE and  RAY_TRACING_BLEND_MODE_ALPHA_COMPOSITE: radiance contribution is alway weighted by coverage.
 		// #dxr_todo: I have not been able to setup a material using RAY_TRACING_BLEND_MODE_MODULATE.
 
 		//
@@ -206,7 +206,7 @@ RAY_TRACING_ENTRY_RAYGEN(RayTracingPrimaryRaysRGS)
 		{
 			float3 WorldPositionRelativeToCameraPrimaryRay = PrimaryView.TranslatedWorldCameraOrigin - (Ray.Origin + Ray.Direction * Payload.HitT);
 
-			// We always samplg the height fog to make sure all the intersected transparent surface get affected by it.
+			// We always sample the height fog to make sure all the intersected transparent surface get affected by it.
 			// We cannot only do it for the primary ray otherwise refelction on secondary surface will look too bright. 
 			// For that sample the fog for the hit position and view assuming no refraction. This seems acceptable.
 			ViewFogInScatteringAndTransmittance = CalculateHeightFog(WorldPositionRelativeToCameraPrimaryRay);
@@ -316,8 +316,8 @@ RAY_TRACING_ENTRY_RAYGEN(RayTracingPrimaryRaysRGS)
 					VolumeLighting.rgb += max(half3(0, 0, 0), DotSH(TranslucencyGISH, SubsurfaceTransferSH)) / PI;
 				}
 			}
-#endif
 			else
+#endif
 			{
 				// noop => if lumen is not enabled: bRefractionEnableSkyLightContribution == true
 			}
@@ -420,7 +420,7 @@ RAY_TRACING_ENTRY_RAYGEN(RayTracingPrimaryRaysRGS)
 
 	if (!bHasScattered)
 	{
-		// Use the scene radiance for ray that has not been scattered/refracted (no surface or IORin=IORout). Still apply the throughtput in case we have traversed surfaces with opacity>0.
+		// Use the scene radiance for ray that has not been scattered/refracted (no surface or IORin=IORout). Still apply the throughput in case we have traversed surfaces with opacity>0.
 		PathRadiance += PathThroughput * SceneColorTexture.SampleLevel(GlobalPointClampedSampler, UV, 0).xyz / View.PreExposure;
 	}
 
