Visualize whether Local Exposure Bilateral Grid lookup was valid.

diff --git a/Engine/Shaders/Private/PostProcessVisualizeLocalExposure.usf b/Engine/Shaders/Private/PostProcessVisualizeLocalExposure.usf
--- a/Engine/Shaders/Private/PostProcessVisualizeLocalExposure.usf
+++ b/Engine/Shaders/Private/PostProcessVisualizeLocalExposure.usf
@@ -52,6 +52,21 @@ float4 MainPS(noperspective float4 UVAndScreenPos : TEXCOORD0, float4 SvPosition
 	{
 		OutValue = exp2(DetailLogLum);
 	}
+	else if (DebugMode == 3)
+	{
+		float LumW = (ComputeHistogramPositionFromLogLuminance(LogLuminance) * (BILATERAL_GRID_DEPTH - 1) + 0.5f) / BILATERAL_GRID_DEPTH;
+		float2 BilateralGridLum = Texture3DSample(LumBilateralGrid, TextureSampler, float3(ViewportUV, LumW)).xy;
+
+		if (BilateralGridLum.y < 0.001)
+		{
+			// bilateral grid doesn't have valid data
+			OutValue = float3(1,0,0);
+		}
+		else
+		{
+			OutValue = float3(0,1,0);
+		}
+	}
 
 	return float4(OutValue, 1.0f);
 }
\ No newline at end of file
diff --git a/Engine/Source/Runtime/Renderer/Private/PostProcess/PostProcessVisualizeLocalExposure.cpp b/Engine/Source/Runtime/Renderer/Private/PostProcess/PostProcessVisualizeLocalExposure.cpp
--- a/Engine/Source/Runtime/Renderer/Private/PostProcess/PostProcessVisualizeLocalExposure.cpp
+++ b/Engine/Source/Runtime/Renderer/Private/PostProcess/PostProcessVisualizeLocalExposure.cpp
@@ -13,7 +13,8 @@ TAutoConsoleVariable<int> CVarLocalExposureVisualizeDebugMode(
 	TEXT("When enabling Show->Visualize->Local Exposure is enabled, this flag controls which mode to use.\n")
 	TEXT("    0: Local Exposure\n")
 	TEXT("    1: Base Luminance\n")
-	TEXT("    2: Detail Luminance\n"),
+	TEXT("    2: Detail Luminance\n")
+	TEXT("    3: Valid Bilateral Grid Lookup\n"),
 	ECVF_RenderThreadSafe);
 
 class FVisualizeLocalExposurePS : public FGlobalShader
