Fix documentation of ShaderPrintCommon.ush.

#rnx

diff --git a/Engine/Shaders/Private/ShaderPrintCommon.ush b/Engine/Shaders/Private/ShaderPrintCommon.ush
--- a/Engine/Shaders/Private/ShaderPrintCommon.ush
+++ b/Engine/Shaders/Private/ShaderPrintCommon.ush
@@ -12,9 +12,9 @@
 // 0. Contex-based state API
 // -------------------------
 // FShaderPrintContext Ctx = InitShaderPrintContext(all(SvPosition.xy == float2(100, 100)), uint2(50,50));
-//                             ------------------------------------   ------------
-//                                 Pixel/thread to filter             Coord at which 
-//                                                                       to draw
+//                                                      ---------------------------------   ------------
+//                                                           Pixel/thread to filter        Coord at which 
+//                                                                                            to draw
 // 
 // 1. How to print value?
 // ----------------------
@@ -52,11 +52,11 @@
 // Widget can be added from shader. 
 // For slider:
 // 
-//   const float MySliderValue  = AddSlider(Ctx, DefaultValue, TEXT("My Slider"));
+//   const float MySliderValue  = AddSlider(Ctx, TEXT("My Slider"), DefaultValue, GetDefaultFontColor());
 // 
 // For checkbox:
 // 
-//   const bool MyCheckboxValue = AddCheckbox(Ctx, DefaultValue, TEXT("My Checkbox"));
+//   const bool MyCheckboxValue = AddCheckbox(Ctx, TEXT("My Checkbox"), DefaultValue, GetDefaultFontColor());
 //
 // The value of these widgets can be retrieved/evaluated from *any* thread/pixels, but the 
 // drawing need to be done from a single thread/pixel. This is done through the ShaderPrintContext, 
@@ -69,7 +69,7 @@
 //       FShaderPrintContext Ctx = InitShaderPrintContext(all(DispatchThreadId == uint2(0, 0)), uint2(50,50));
 // 
 //		 // Any threads can get the slider value with this call
-//       const float ScaleValue = AddSlider(Ctx, 1.f, TEXT("Scale"));
+//       const float ScaleValue = AddSlider(Ctx, TEXT("Scale"), 1.f, GetDefaultFontColor());
 //       ...
 //    }
 // 
