Moves TSR's CONFIG_SCENE_COLOR_OVERFLOW into the RejectShading

diff --git a/Engine/Shaders/Private/TemporalSuperResolution/TSRRejectShading.usf b/Engine/Shaders/Private/TemporalSuperResolution/TSRRejectShading.usf
--- a/Engine/Shaders/Private/TemporalSuperResolution/TSRRejectShading.usf
+++ b/Engine/Shaders/Private/TemporalSuperResolution/TSRRejectShading.usf
@@ -513,6 +513,14 @@ void MainCS(
 				#endif
 			#endif
 			
+			#if CONFIG_SCENE_COLOR_OVERFLOW
+				#if CONFIG_SCENE_COLOR_ALPHA
+					Color = min(Color, LargestSceneColorRGBA);
+				#else
+					Color = min(Color, LargestSceneColorRGB);
+				#endif
+			#endif
+
 			CenterFinalSceneColor.SetElement(SimdIndex, Color);
 
 			tsr_half CenterLuma = dot(Color.rgb, tsr_half3(0.299f, 0.587f, 0.114f));
diff --git a/Engine/Shaders/Private/TemporalSuperResolution/TSRUpdateHistory.usf b/Engine/Shaders/Private/TemporalSuperResolution/TSRUpdateHistory.usf
--- a/Engine/Shaders/Private/TemporalSuperResolution/TSRUpdateHistory.usf
+++ b/Engine/Shaders/Private/TemporalSuperResolution/TSRUpdateHistory.usf
@@ -761,7 +761,6 @@ void MainCS(
 	tsr_halfCx2 InputMaxColor;
 
 	tsr_half2 InputPixelAlignement;
-	tsr_half2 ClosestInputLuma4;
 	
 	tsr_half2 BlendClamp;
 	tsr_half2 PrevWeight;
@@ -939,14 +938,6 @@ void MainCS(
 			}
 			#endif
 
-			#if CONFIG_SCENE_COLOR_OVERFLOW
-				#if CONFIG_SCENE_COLOR_ALPHA
-					InputColor = min(InputColor, dpv_interleave_mono_registers(LargestSceneColorRGBA));
-				#else
-					InputColor = min(InputColor, dpv_interleave_mono_registers(LargestSceneColorRGB));
-				#endif
-			#endif
-
 			tsr_half2 ToneWeight = HdrWeight4(InputColor);
 
 			FilteredInputColor       += dpv_scale((SampleSpatialWeight * ToneWeight), InputColor);
@@ -965,7 +956,6 @@ void MainCS(
 			#else
 				FilteredInputColor = min(FilteredInputColor, dpv_interleave_mono_registers(LargestSceneColorRGB));
 			#endif
-			ClosestInputLuma4 = min(ClosestInputLuma4, LargestSceneColorYCoCg.x);
 		}
 		#endif
 	}
