Converts DOF's translucency to opacity into DOFReduce pass

diff --git a/Engine/Shaders/Private/DiaphragmDOF/DOFReduce.usf b/Engine/Shaders/Private/DiaphragmDOF/DOFReduce.usf
--- a/Engine/Shaders/Private/DiaphragmDOF/DOFReduce.usf
+++ b/Engine/Shaders/Private/DiaphragmDOF/DOFReduce.usf
@@ -294,6 +294,11 @@ void ReduceCS(
 	{
 		GatheredColor = GatherInput_SceneColor.SampleLevel(GlobalPointClampedSampler, BufferUV, 0);
 		CocRadius = GatherInput_SeparateCoc.SampleLevel(GlobalPointClampedSampler, BufferUV, 0).r;
+
+		// Convert from translucency to opacity.
+		#if !DIM_DISABLE_OUTPUT_MIP0
+			GatheredColor.a = 1.0 - GatheredColor.a;
+		#endif
 	}
 	#else
 	{
diff --git a/Engine/Shaders/Private/DiaphragmDOF/DOFSetup.usf b/Engine/Shaders/Private/DiaphragmDOF/DOFSetup.usf
--- a/Engine/Shaders/Private/DiaphragmDOF/DOFSetup.usf
+++ b/Engine/Shaders/Private/DiaphragmDOF/DOFSetup.usf
@@ -79,9 +79,6 @@ void SetupCS(
 	float DeviceZ = SceneDepthTexture.SampleLevel(GlobalPointClampedSampler, SceneBufferUV, 0).r;
 	float CocRadius = SceneDepthToCocRadius(ConvertFromDeviceZ(DeviceZ));
 	
-	// Convert from translucency to opacity.
-	SceneColor.a = 1 - SceneColor.a;
-
 	// Output full resolution.
 	#if CONFIG_DOF_ALPHA
 	{
@@ -229,9 +226,6 @@ void SetupCS(
 		{
 			float2 SampleUV = SceneBufferUV + SceneBufferUVOffset * float2(kOffsetsCross3x3[i]);
 			SceneColors[i] = SceneColorTexture.SampleLevel(GlobalPointClampedSampler, SampleUV, 0);
-
-			// Convert from translucency to opacity.
-			SceneColors[i].a = 1 - SceneColors[i].a;
 		}
 	
 		// Downsample.
@@ -266,9 +260,6 @@ void SetupCS(
 
 			SceneColors[i] = SceneColorTexture.SampleLevel(GlobalPointClampedSampler, SampleUV, 0);
 			
-			// Convert from translucency to opacity.
-			SceneColors[i].a = 1 - SceneColors[i].a;
-
 			float SampleDeviceZ = SceneDepthTexture.SampleLevel(GlobalPointClampedSampler, SampleUV, 0).r;
 			CocRadii[i] = SceneDepthToCocRadius(ConvertFromDeviceZ(SampleDeviceZ));
 		}
