Renamed HitPoint to TranslatedHitPoint in RayTracingPrimaryRays.usf

diff --git a/Engine/Shaders/Private/RayTracing/RayTracingPrimaryRays.usf b/Engine/Shaders/Private/RayTracing/RayTracingPrimaryRays.usf
--- a/Engine/Shaders/Private/RayTracing/RayTracingPrimaryRays.usf
+++ b/Engine/Shaders/Private/RayTracing/RayTracingPrimaryRays.usf
@@ -186,7 +186,7 @@ RAY_TRACING_ENTRY_RAYGEN(RayTracingPrimaryRaysRGS)
 			}
 			break;
 		}
-		float3 HitPoint = Ray.Origin + Ray.Direction * Payload.HitT;
+		float3 TranslatedHitPoint = Ray.Origin + Ray.Direction * Payload.HitT;
 		float NextMaxRayDistance = Ray.TMax - Payload.HitT;
 
 		//
@@ -204,7 +204,7 @@ RAY_TRACING_ENTRY_RAYGEN(RayTracingPrimaryRaysRGS)
 		float4 ViewFogInScatteringAndTransmittance = float4(0.0, 0.0, 0.0, 1.0);
 		if (HeightFog > 0)
 		{
-			float3 WorldPositionRelativeToCameraPrimaryRay = PrimaryView.TranslatedWorldCameraOrigin - (Ray.Origin + Ray.Direction * Payload.HitT);
+			float3 WorldPositionRelativeToCameraPrimaryRay = PrimaryView.TranslatedWorldCameraOrigin - TranslatedHitPoint;
 
 			// We always sample the height fog to make sure all the intersected transparent surface get affected by it.
 			// We cannot only do it for the primary ray otherwise refelction on secondary surface will look too bright. 
@@ -214,7 +214,7 @@ RAY_TRACING_ENTRY_RAYGEN(RayTracingPrimaryRaysRGS)
 			if (FogStruct.ApplyVolumetricFog > 0
 				&& !bHasScattered) // We can only sample the volumetric fog contribution (always camera aligned) if the ray is not scattered around
 			{
-				float4 ClipPos = mul(float4(HitPoint, 1.0f), PrimaryView.TranslatedWorldToClip); // HitPoint is translated world space
+				float4 ClipPos = mul(float4(TranslatedHitPoint, 1.0f), PrimaryView.TranslatedWorldToClip);
 				float3 VolumeUV = ComputeVolumeUVFromNDC(ClipPos);
 
 				ViewFogInScatteringAndTransmittance = CombineVolumetricFog(ViewFogInScatteringAndTransmittance, VolumeUV, 0/*EyeIndex*/, GBufferData.Depth);
@@ -240,7 +240,7 @@ RAY_TRACING_ENTRY_RAYGEN(RayTracingPrimaryRaysRGS)
 			FRayDesc ReflectionRay;
 			ReflectionRay.TMin = 0.01;
 			ReflectionRay.TMax = LocalMaxRayDistance;
-			ReflectionRay.Origin = HitPoint;
+			ReflectionRay.Origin = TranslatedHitPoint;
 
 			ModifyGGXAnisotropicNormalRoughness(Payload.WorldTangent, Payload.Anisotropy, Payload.Roughness, Payload.WorldNormal, Ray.Direction);
 
@@ -295,7 +295,7 @@ RAY_TRACING_ENTRY_RAYGEN(RayTracingPrimaryRaysRGS)
 		{
 			float3 VolumeLighting = 0;
 
-			const float3 TranslucencyEvaluationTranslatedPosition = HitPoint;
+			const float3 TranslucencyEvaluationTranslatedPosition = TranslatedHitPoint;
 
 #if PROJECT_SUPPORTS_LUMEN
 			if (IsLumenTranslucencyGIEnabled() && ShouldDoEmissiveAndIndirectLighting)
@@ -410,7 +410,7 @@ RAY_TRACING_ENTRY_RAYGEN(RayTracingPrimaryRaysRGS)
 		// Setup refracted ray to be traced
 		//
 
-		Ray.Origin = HitPoint;
+		Ray.Origin = TranslatedHitPoint;
 		Ray.TMin = 0.01;
 		Ray.TMax = NextMaxRayDistance;
 		Ray.Direction = RefractedDirection;
