Substrate - fixed decal fading from screen area.

diff --git a/Engine/Shaders/Private/DeferredDecal.usf b/Engine/Shaders/Private/DeferredDecal.usf
--- a/Engine/Shaders/Private/DeferredDecal.usf
+++ b/Engine/Shaders/Private/DeferredDecal.usf
@@ -159,7 +159,7 @@ void FPixelShaderInOut_MainPS(inout FPixelShaderIn In, inout FPixelShaderOut Out
 		StrataHeader.StrataUpdateTree(V, Settings);
 
 		// Either propagate the BSDF opacity, or use the one for AlphaComposite (premultiplied alpha)
-		float CoverageOverBackground = StrataHeader.StrataTree.BSDFs[0].Coverage;
+		float CoverageOverBackground = StrataHeader.StrataTree.BSDFs[0].Coverage * DecalFading;
 #if STRATA_USE_PREMULTALPHA_OVERRIDE // AlphaComposite - Premultiplied alpha blending
 		CoverageOverBackground = GetMaterialOpacity(PixelMaterialInputs);
 #endif 
