Color: Fix hard-coded ACES matrices.

Improved precision such that multiplication of inverses better yield identity.

#rnx

diff --git a/Engine/Shaders/Private/ACES.ush b/Engine/Shaders/Private/ACES.ush
--- a/Engine/Shaders/Private/ACES.ush
+++ b/Engine/Shaders/Private/ACES.ush
@@ -148,54 +148,54 @@ static const float3x3 XYZ_2_sRGB_MAT =
 
 static const float3x3 sRGB_2_XYZ_MAT =
 {
-	0.4124564, 0.3575761, 0.1804375,
-	0.2126729, 0.7151522, 0.0721750,
-	0.0193339, 0.1191920, 0.9503041,
+	0.4123907993, 0.3575843394, 0.1804807884,
+	0.2126390059, 0.7151686788, 0.0721923154,
+	0.0193308187, 0.1191947798, 0.9505321522,
 };
 
 // REC 2020 primaries
 static const float3x3 XYZ_2_Rec2020_MAT =
 {
-	 1.7166084, -0.3556621, -0.2533601,
-	-0.6666829,  1.6164776,  0.0157685,
-	 0.0176422, -0.0427763,  0.94222867
+	 1.7166511880, -0.3556707838, -0.2533662814,
+	-0.6666843518,  1.6164812366,  0.0157685458,
+	 0.0176398574, -0.0427706133,  0.9421031212,
 };
 
 static const float3x3 Rec2020_2_XYZ_MAT =
 {
-	0.6369736, 0.1446172, 0.1688585,
-	0.2627066, 0.6779996, 0.0592938,
-	0.0000000, 0.0280728, 1.0608437
+	0.6369580483, 0.1446169036, 0.1688809752,
+	0.2627002120, 0.6779980715, 0.0593017165,
+	0.0000000000, 0.0280726930, 1.0609850577,
 };
 
 // P3, D65 primaries
 static const float3x3 XYZ_2_P3D65_MAT =
 {
-	 2.4933963, -0.9313459, -0.4026945,
-	-0.8294868,  1.7626597,  0.0236246,
-	 0.0358507, -0.0761827,  0.9570140
+	 2.4934969119, -0.9313836179, -0.4027107845,
+	-0.8294889696,  1.7626640603,  0.0236246858,
+	 0.0358458302, -0.0761723893,  0.9568845240,
 };
 
 static const float3x3 P3D65_2_XYZ_MAT =
 {
-	0.4865906, 0.2656683, 0.1981905,
-	0.2289838, 0.6917402, 0.0792762,
-	0.0000000, 0.0451135, 1.0438031
+	0.4865709486, 0.2656676932, 0.1982172852,
+	0.2289745641, 0.6917385218, 0.0792869141,
+	0.0000000000, 0.0451133819, 1.0439443689,
 };
 
 // Bradford chromatic adaptation transforms between ACES white point (D60) and sRGB white point (D65)
 static const float3x3 D65_2_D60_CAT =
 {
-	 1.01303,    0.00610531, -0.014971,
-	 0.00769823, 0.998165,   -0.00503203,
-	-0.00284131, 0.00468516,  0.924507,
+	 1.0130349146, 0.0061052578, -0.0149709436,
+	 0.0076982301, 0.9981633521, -0.0050320385,
+	-0.0028413174, 0.0046851567,  0.9245061375,
 };
 
 static const float3x3 D60_2_D65_CAT =
 {
-	 0.987224,   -0.00611327, 0.0159533,
-	-0.00759836,  1.00186,    0.00533002,
-	 0.00307257, -0.00509595, 1.08168,
+	 0.9872240087, -0.0061132286, 0.0159532883,
+	-0.0075983718,  1.0018614847, 0.0053300358,
+	 0.0030725771, -0.0050959615, 1.0816806031,
 };
 
 static const float HALF_MAX = 65504.0;
