Optionally compute per-instance params (fade amount, show/hide, dither) depending on whether material actually needs it

diff --git a/Engine/Shaders/Private/GeometryCollectionVertexFactory.ush b/Engine/Shaders/Private/GeometryCollectionVertexFactory.ush
--- a/Engine/Shaders/Private/GeometryCollectionVertexFactory.ush
+++ b/Engine/Shaders/Private/GeometryCollectionVertexFactory.ush
@@ -580,7 +580,7 @@ FVertexFactoryRayTracingInterpolants VertexFactoryInterpolate(FVertexFactoryRayT
 #if INTERPOLATE_VERTEX_COLOR
 	INTERPOLATE_MEMBER(InterpolantsVSToPS.Color);
 #endif
-#if USE_INSTANCING || USE_INSTANCE_CULLING
+#if	NEEDS_PER_INSTANCE_PARAMS
 	INTERPOLATE_MEMBER(InterpolantsVSToPS.PerInstanceParams);
 #endif
 
diff --git a/Engine/Shaders/Private/HairStrands/HairCardsVertexFactory.ush b/Engine/Shaders/Private/HairStrands/HairCardsVertexFactory.ush
--- a/Engine/Shaders/Private/HairStrands/HairCardsVertexFactory.ush
+++ b/Engine/Shaders/Private/HairStrands/HairCardsVertexFactory.ush
@@ -433,7 +433,7 @@ FVertexFactoryRayTracingInterpolants VertexFactoryInterpolate(FVertexFactoryRayT
 #if INTERPOLATE_VERTEX_COLOR
 	INTERPOLATE_MEMBER(InterpolantsVSToPS.Color);
 #endif
-#if USE_INSTANCING
+#if	NEEDS_PER_INSTANCE_PARAMS
 	INTERPOLATE_MEMBER(InterpolantsVSToPS.PerInstanceParams);
 #endif
 
diff --git a/Engine/Shaders/Private/HairStrands/HairStrandsVertexFactory.ush b/Engine/Shaders/Private/HairStrands/HairStrandsVertexFactory.ush
--- a/Engine/Shaders/Private/HairStrands/HairStrandsVertexFactory.ush
+++ b/Engine/Shaders/Private/HairStrands/HairStrandsVertexFactory.ush
@@ -875,7 +875,7 @@ FVertexFactoryRayTracingInterpolants VertexFactoryInterpolate(FVertexFactoryRayT
 #if INTERPOLATE_VERTEX_COLOR
 	INTERPOLATE_MEMBER(InterpolantsVSToPS.Color);
 #endif
-#if USE_INSTANCING
+#if	NEEDS_PER_INSTANCE_PARAMS
 	INTERPOLATE_MEMBER(InterpolantsVSToPS.PerInstanceParams);
 #endif
 
diff --git a/Engine/Shaders/Private/LocalVertexFactory.ush b/Engine/Shaders/Private/LocalVertexFactory.ush
--- a/Engine/Shaders/Private/LocalVertexFactory.ush
+++ b/Engine/Shaders/Private/LocalVertexFactory.ush
@@ -541,7 +541,11 @@ FMaterialPixelParameters GetMaterialPixelParameters(FVertexFactoryInterpolantsVS
 	// Required for previewing materials that use ParticleColor
 	Result.Particle.Color = half4(1,1,1,1);
 #if USE_INSTANCING || USE_INSTANCE_CULLING 
-	Result.PerInstanceParams = Interpolants.PerInstanceParams;
+	#if	NEEDS_PER_INSTANCE_PARAMS
+		Result.PerInstanceParams = Interpolants.PerInstanceParams;
+	#else
+		Result.PerInstanceParams = float4(1, 1, 0, 0);
+	#endif
 #endif
 
 	Result.TangentToWorld = AssembleTangentToWorld( TangentToWorld0, TangentToWorld2 );
@@ -627,7 +631,9 @@ FMaterialVertexParameters GetMaterialVertexParameters(FVertexFactoryInput Input,
 	float4x4 InstanceToLocal = GetInstanceTransform(Intermediates);
 	Result.InstanceLocalToWorld = LWCMultiply(InstanceToLocal, PrimitiveData.LocalToWorld);
 	Result.InstanceLocalPosition = Input.Position.xyz;
+	#if	NEEDS_PER_INSTANCE_PARAMS
 	Result.PerInstanceParams = Intermediates.PerInstanceParams;
+	#endif
 	Result.InstanceId = GetInstanceId(Input.InstanceId); 
 	Result.InstanceOffset = InstanceOffset;
 	Result.PrevFrameLocalToWorld = LWCMultiply(GetInstancePrevTransform(Intermediates), PrimitiveData.PreviousLocalToWorld);
@@ -657,7 +663,9 @@ FMaterialVertexParameters GetMaterialVertexParameters(FVertexFactoryInput Input,
 		Result.InstanceLocalPosition = Input.Position.xyz;
 		Result.InstanceLocalToWorld = InstanceData.LocalToWorld;
 		Result.InstanceWorldToLocal = InstanceData.WorldToLocal;
+		#if	NEEDS_PER_INSTANCE_PARAMS
 		Result.PerInstanceParams = Intermediates.PerInstanceParams;
+		#endif
 	#endif
 	#if USES_PER_INSTANCE_CUSTOM_DATA && VF_USE_PRIMITIVE_SCENE_DATA
 		Result.CustomDataOffset = InstanceData.CustomDataOffset;
@@ -856,6 +864,7 @@ FVertexFactoryIntermediates GetVertexFactoryIntermediates(FVertexFactoryInput In
 	Intermediates.HitProxyId = float4(0.0f, 0.0f, 0.0f, 0.0f);
 #endif // USE_EDITOR_SHADERS
 
+#if NEEDS_PER_INSTANCE_PARAMS
 	{
 		// This repackaging is stupid-seeming but these params are propably packed off to some interpolator or something
 		
@@ -889,6 +898,10 @@ FVertexFactoryIntermediates GetVertexFactoryIntermediates(FVertexFactoryInput In
 			Intermediates.bEvaluateWorldPositionOffset = false;
 		}
 	}
+#else
+	Intermediates.PerInstanceParams = float4(0.f,Intermediates.IsVisible,0.f,0.f);
+#endif // NEEDS_PER_INSTANCE_PARAMS
+
 #endif // USE_INSTANCE_CULLING
 #endif // VF_USE_PRIMITIVE_SCENE_DATA
 
@@ -911,38 +924,41 @@ FVertexFactoryIntermediates GetVertexFactoryIntermediates(FVertexFactoryInput In
 	Intermediates.TangentToWorldSign = TangentSign * GetInstanceData(Intermediates).DeterminantSign;
 
 #if USE_INSTANCING
-	// x = per-instance fade out factor, y = zero or one depending of if it is shown at all, z is dither cutoff 
-
-	// PerInstanceParams.y stores a hide/show flag for this instance
-	float SelectedValue = GetInstanceSelected(Intermediates);
 	// GPUCULL_TODO: This can't be right
 	float3 InstanceTranslatedLocation = TransformLocalToTranslatedWorld(GetInstanceOrigin(Intermediates), PrimitiveData.LocalToWorld).xyz;
 	float InstanceViewDistance = length(InstanceTranslatedLocation);
-	Intermediates.PerInstanceParams.x = 1.0 - saturate((InstanceViewDistance - InstancedVFLooseParameters.InstancingFadeOutParams.x) * InstancedVFLooseParameters.InstancingFadeOutParams.y);
-	// InstancingFadeOutParams.z,w are RenderSelected and RenderDeselected respectively.
-	Intermediates.PerInstanceParams.y = InstancedVFLooseParameters.InstancingFadeOutParams.z * SelectedValue + InstancedVFLooseParameters.InstancingFadeOutParams.w * (1-SelectedValue);
-	#if USE_DITHERED_LOD_TRANSITION
-		float RandomLOD = InstancedVFLooseParameters.InstancingViewZCompareZero.w * InstanceData.RandomID;
-		float ViewZZero = length(InstanceTranslatedLocation - InstancedVFLooseParameters.InstancingTranslatedWorldViewOriginZero.xyz) + RandomLOD;
-		float ViewZOne = length(InstanceTranslatedLocation - InstancedVFLooseParameters.InstancingTranslatedWorldViewOriginOne.xyz) + RandomLOD;
-		Intermediates.PerInstanceParams.z = 
-			dot(float3(ViewZZero.xxx > InstancedVFLooseParameters.InstancingViewZCompareZero.xyz), InstancedVFLooseParameters.InstancingViewZConstant.xyz) * InstancedVFLooseParameters.InstancingTranslatedWorldViewOriginZero.w +
-			dot(float3(ViewZOne.xxx > InstancedVFLooseParameters.InstancingViewZCompareOne.xyz), InstancedVFLooseParameters.InstancingViewZConstant.xyz) * InstancedVFLooseParameters.InstancingTranslatedWorldViewOriginOne.w;
-		Intermediates.PerInstanceParams.y *= abs(Intermediates.PerInstanceParams.z) < .999;
-	#else
-		Intermediates.PerInstanceParams.z = 0;
-	#endif
+	#if	NEEDS_PER_INSTANCE_PARAMS	
+		// x = per-instance fade out factor, y = zero or one depending of if it is shown at all, z is dither cutoff 
+		// PerInstanceParams.y stores a hide/show flag for this instance
+		float SelectedValue = GetInstanceSelected(Intermediates);
+		Intermediates.PerInstanceParams.x = 1.0 - saturate((InstanceViewDistance - InstancedVFLooseParameters.InstancingFadeOutParams.x) * InstancedVFLooseParameters.InstancingFadeOutParams.y);
+		// InstancingFadeOutParams.z,w are RenderSelected and RenderDeselected respectively.
+		Intermediates.PerInstanceParams.y = InstancedVFLooseParameters.InstancingFadeOutParams.z * SelectedValue + InstancedVFLooseParameters.InstancingFadeOutParams.w * (1-SelectedValue);
+		#if USE_DITHERED_LOD_TRANSITION
+			float RandomLOD = InstancedVFLooseParameters.InstancingViewZCompareZero.w * InstanceData.RandomID;
+			float ViewZZero = length(InstanceTranslatedLocation - InstancedVFLooseParameters.InstancingTranslatedWorldViewOriginZero.xyz) + RandomLOD;
+			float ViewZOne = length(InstanceTranslatedLocation - InstancedVFLooseParameters.InstancingTranslatedWorldViewOriginOne.xyz) + RandomLOD;
+			Intermediates.PerInstanceParams.z = 
+				dot(float3(ViewZZero.xxx > InstancedVFLooseParameters.InstancingViewZCompareZero.xyz), InstancedVFLooseParameters.InstancingViewZConstant.xyz) * InstancedVFLooseParameters.InstancingTranslatedWorldViewOriginZero.w +
+				dot(float3(ViewZOne.xxx > InstancedVFLooseParameters.InstancingViewZCompareOne.xyz), InstancedVFLooseParameters.InstancingViewZConstant.xyz) * InstancedVFLooseParameters.InstancingTranslatedWorldViewOriginOne.w;
+			Intermediates.PerInstanceParams.y *= abs(Intermediates.PerInstanceParams.z) < .999;
+		#else
+			Intermediates.PerInstanceParams.z = 0;
+		#endif
 	
-	#if USES_PER_INSTANCE_CUSTOM_DATA
-		// index into instance CustomData
-		Intermediates.PerInstanceParams.w = asfloat(GetInstanceId(Input.InstanceId) + InstanceOffset);
-	#endif
+		#if USES_PER_INSTANCE_CUSTOM_DATA
+			// index into instance CustomData
+			Intermediates.PerInstanceParams.w = asfloat(GetInstanceId(Input.InstanceId) + InstanceOffset);
+		#endif
 
-	// Disable WPO if this instance is not visible for dithered LOD transition or instance fade reasons
-	if (Intermediates.PerInstanceParams.y == 0.0f)
-	{
-		Intermediates.bEvaluateWorldPositionOffset = false;
-	}
+		// Disable WPO if this instance is not visible for dithered LOD transition or instance fade reasons
+		if (Intermediates.PerInstanceParams.y == 0.0f)
+		{
+			Intermediates.bEvaluateWorldPositionOffset = false;
+		}
+	#else
+		Intermediates.PerInstanceParams = float4(1.f,Intermediates.IsVisible,0.f,0.f);
+	#endif // NEEDS_PER_INSTANCE_PARAMS
 
 	const float InstanceViewDistSq = Square(InstanceViewDistance);
 #else
@@ -1082,7 +1098,7 @@ FVertexFactoryInterpolantsVSToPS VertexFactoryGetInterpolantsVSToPS(FVertexFacto
 	SetTangents(Interpolants, Intermediates.TangentToWorld[0], Intermediates.TangentToWorld[2], Intermediates.TangentToWorldSign);
 	SetColor(Interpolants, Intermediates.Color);
 
-#if USE_INSTANCING || USE_INSTANCE_CULLING
+#if	NEEDS_PER_INSTANCE_PARAMS
 	Interpolants.PerInstanceParams = Intermediates.PerInstanceParams;
 #endif
 
@@ -1242,7 +1258,7 @@ float4 VertexFactoryGetPreviousWorldPosition(FVertexFactoryInput Input, FVertexF
 	#if INTERPOLATE_VERTEX_COLOR
 		INTERPOLATE_MEMBER(InterpolantsVSToPS.Color);
 	#endif
-	#if USE_INSTANCING || USE_INSTANCE_CULLING
+	#if	NEEDS_PER_INSTANCE_PARAMS
 		INTERPOLATE_MEMBER(InterpolantsVSToPS.PerInstanceParams);
 	#endif
 
diff --git a/Engine/Shaders/Private/LocalVertexFactoryCommon.ush b/Engine/Shaders/Private/LocalVertexFactoryCommon.ush
--- a/Engine/Shaders/Private/LocalVertexFactoryCommon.ush
+++ b/Engine/Shaders/Private/LocalVertexFactoryCommon.ush
@@ -4,6 +4,8 @@
 	LocalVertexFactoryCommon.usf: Local vertex factory common functionality
 =============================================================================*/
 
+#define NEEDS_PER_INSTANCE_PARAMS ((USE_INSTANCING || USE_INSTANCE_CULLING) && (USES_PER_INSTANCE_FADE_AMOUNT || USES_SPEEDTREE || USE_DITHERED_LOD_TRANSITION || (!VF_USE_PRIMITIVE_SCENE_DATA && USES_PER_INSTANCE_CUSTOM_DATA)))
+
 struct FVertexFactoryInterpolantsVSToPS
 {
 	TANGENTTOWORLD_INTERPOLATOR_BLOCK
@@ -12,7 +14,7 @@ struct FVertexFactoryInterpolantsVSToPS
 	half4	Color : COLOR0;
 #endif
 
-#if USE_INSTANCING || USE_INSTANCE_CULLING
+#if NEEDS_PER_INSTANCE_PARAMS
 	// x = per-instance fade out amount, y = hide/show flag, z dither fade cutoff, w - CustomData index
 	nointerpolation float4  PerInstanceParams : COLOR1;
 #endif
diff --git a/Engine/Source/Runtime/Engine/Private/Materials/HLSLMaterialTranslator.cpp b/Engine/Source/Runtime/Engine/Private/Materials/HLSLMaterialTranslator.cpp
--- a/Engine/Source/Runtime/Engine/Private/Materials/HLSLMaterialTranslator.cpp
+++ b/Engine/Source/Runtime/Engine/Private/Materials/HLSLMaterialTranslator.cpp
@@ -377,6 +377,7 @@ FHLSLMaterialTranslator::FHLSLMaterialTranslator(FMaterial* InMaterial,
 ,	bUsesAnisotropy(false)
 ,	bMaterialIsStrata(false)
 ,	bUsesCurvature(false)
+,	bUsesPerInstanceFadeAmount(false)
 ,	AllocatedUserTexCoords()
 ,	AllocatedUserVertexTexCoords()
 ,	DynamicParticleParameterMask(0)
@@ -2060,6 +2061,7 @@ void FHLSLMaterialTranslator::GetMaterialEnvironment(EShaderPlatform InPlatform,
 
 	OutEnvironment.SetDefine(TEXT("USES_PER_INSTANCE_CUSTOM_DATA"), MaterialCompilationOutput.bUsesPerInstanceCustomData && Material->IsUsedWithInstancedStaticMeshes());
 	OutEnvironment.SetDefine(TEXT("USES_PER_INSTANCE_RANDOM"), MaterialCompilationOutput.bUsesPerInstanceRandom && Material->IsUsedWithInstancedStaticMeshes());
+	OutEnvironment.SetDefine(TEXT("USES_PER_INSTANCE_FADE_AMOUNT"), bUsesPerInstanceFadeAmount && Material->IsUsedWithInstancedStaticMeshes());
 	OutEnvironment.SetDefine(TEXT("USES_VERTEX_INTERPOLATOR"), MaterialCompilationOutput.bUsesVertexInterpolator);
 
 	OutEnvironment.SetDefine(TEXT("MATERIAL_SKY_ATMOSPHERE"), bUsesSkyAtmosphere);
@@ -13721,6 +13723,7 @@ int32 FHLSLMaterialTranslator::PerInstanceFadeAmount()
 	}
 	else
 	{
+		bUsesPerInstanceFadeAmount = true;
 		return AddInlinedCodeChunkZeroDeriv(MCT_Float, TEXT("GetPerInstanceFadeAmount(Parameters)"));
 	}
 }
diff --git a/Engine/Source/Runtime/Engine/Private/Materials/HLSLMaterialTranslator.h b/Engine/Source/Runtime/Engine/Private/Materials/HLSLMaterialTranslator.h
--- a/Engine/Source/Runtime/Engine/Private/Materials/HLSLMaterialTranslator.h
+++ b/Engine/Source/Runtime/Engine/Private/Materials/HLSLMaterialTranslator.h
@@ -418,6 +418,8 @@ protected:
 	uint32 bEnableExecutionFlow : 1;
 
 	uint32 bUsesCurvature : 1;
+	/** true if PerInstanceFadeAmount expression is used */
+	uint32 bUsesPerInstanceFadeAmount : 1;
 
 	/** Tracks the texture coordinates used by this material. */
 	TBitArray<> AllocatedUserTexCoords;
