Properly initialize PreSkinnedLocalBounds with mobile GPUScene enabled

diff --git a/Engine/Shaders/Private/SceneDataMobileLoader.ush b/Engine/Shaders/Private/SceneDataMobileLoader.ush
--- a/Engine/Shaders/Private/SceneDataMobileLoader.ush
+++ b/Engine/Shaders/Private/SceneDataMobileLoader.ush
@@ -96,6 +96,12 @@ FPrimitiveSceneData LoadPrimitiveDataUBO(uint DrawInstanceId)
 		PrimitiveData.ObjectBoundsZ = Data.w;
 		DataIdx += 1u;
 	}
+	
+	// PreSkinnedLocalBounds, TODO: this works until we actually support GPUScene for skinned primitives
+	{
+		PrimitiveData.PreSkinnedLocalBoundsMin = PrimitiveData.LocalObjectBoundsMin;
+		PrimitiveData.PreSkinnedLocalBoundsMax = PrimitiveData.LocalObjectBoundsMax;
+	}
 
 	// WorldToLocal, TODO compute from LocalToWorld?
 	{
