Clamps TSR history on velocity edges of the previous frame to clean up some ghosting due to velocity buffer aliasing

diff --git a/Engine/Shaders/Private/TemporalSuperResolution/TSRDecimateHistory.usf b/Engine/Shaders/Private/TemporalSuperResolution/TSRDecimateHistory.usf
--- a/Engine/Shaders/Private/TemporalSuperResolution/TSRDecimateHistory.usf
+++ b/Engine/Shaders/Private/TemporalSuperResolution/TSRDecimateHistory.usf
@@ -42,7 +42,7 @@ Texture2D<float> ClosestDepthTexture;
 Texture2D<uint> PrevUseCountTexture;
 Texture2D<uint> PrevClosestDepthTexture;
 
-Texture2DArray<tsr_halfC> PrevHistoryGuide;
+Texture2DArray<tsr_half4> PrevHistoryGuide;
 Texture2DArray<tsr_half4> PrevHistoryMoire;
 Texture2D<float2> PrevVelocityTexture;
 float2 PrevGuideInfo_Extent;
@@ -437,11 +437,13 @@ void MainCS(
 	uint HistoryClosestDeviceZSamples[4];
 	
 	tsr_halfC ReprojectedHistoryGuideColor;
+	tsr_half  ReprojectedHistoryGuideUncertainty;
 	#if CONFIG_MOIRE_REPROJECTION
 		tsr_half4 ReprojectedHistoryMoireColor;
 	#endif
 	#if CONFIG_RESURRECTION_REPROJECTION
 		tsr_halfC ResurrectedHistoryGuideColor;
+		tsr_half  ResurrectedHistoryGuideUncertainty;
 	#endif
 	bool bIsOffScreen;
 
@@ -470,12 +472,30 @@ void MainCS(
 				
 			float2 KernelHistoryBufferUV = ApplyScreenTransform(PrevOutputScreenPos, ScreenPosToPrevHistoryGuideBufferUV);
 
-			ReprojectedHistoryGuideColor = SampleTextureCatrom(
-				PrevHistoryGuide,
-				KernelHistoryBufferUV,
-				PrevFrameIndex,
-				PrevGuideInfo_UVViewportBilinearMin,
-				PrevGuideInfo_UVViewportBilinearMax);
+			#if CONFIG_SCENE_COLOR_ALPHA
+				ReprojectedHistoryGuideColor = SampleTextureCatrom(
+					PrevHistoryGuide,
+					KernelHistoryBufferUV,
+					PrevFrameIndex * 2.0 + 0.0,
+					PrevGuideInfo_UVViewportBilinearMin,
+					PrevGuideInfo_UVViewportBilinearMax);
+					
+				ReprojectedHistoryGuideUncertainty = SampleTextureCatrom(
+					PrevHistoryGuide,
+					KernelHistoryBufferUV,
+					PrevFrameIndex * 2.0 + 1.0,
+					PrevGuideInfo_UVViewportBilinearMin,
+					PrevGuideInfo_UVViewportBilinearMax).r;
+			#else
+				tsr_half4 RawGuide = SampleTextureCatrom(
+					PrevHistoryGuide,
+					KernelHistoryBufferUV,
+					PrevFrameIndex,
+					PrevGuideInfo_UVViewportBilinearMin,
+					PrevGuideInfo_UVViewportBilinearMax);
+				ReprojectedHistoryGuideColor = RawGuide.rgb;
+				ReprojectedHistoryGuideUncertainty = RawGuide.a;
+			#endif
 			ReprojectedHistoryGuideColor = CorrectGuideColorExposure(ReprojectedHistoryGuideColor, HistoryPreExposureCorrection);
 			#if CONFIG_MOIRE_REPROJECTION
 			{
@@ -493,17 +513,15 @@ void MainCS(
 			float2 ResurrectionScreen = ResurrectionClip.xy / ResurrectionClip.w;
 
 			float2 KernelHistoryBufferUV = ApplyScreenTransform(ResurrectionScreen, ScreenPosToResurrectionGuideBufferUV);
+			KernelHistoryBufferUV = fastClamp(KernelHistoryBufferUV, ResurrectionGuideUVViewportBilinearMin, ResurrectionGuideUVViewportBilinearMax);
 
-			#if 1
-				KernelHistoryBufferUV = fastClamp(KernelHistoryBufferUV, ResurrectionGuideUVViewportBilinearMin, ResurrectionGuideUVViewportBilinearMax);
-				ResurrectedHistoryGuideColor = PrevHistoryGuide.SampleLevel(GlobalBilinearClampedSampler, float3(KernelHistoryBufferUV, ResurrectionFrameIndex), 0);
+			#if CONFIG_SCENE_COLOR_ALPHA
+				ResurrectedHistoryGuideColor       = PrevHistoryGuide.SampleLevel(GlobalBilinearClampedSampler, float3(KernelHistoryBufferUV, ResurrectionFrameIndex * 2.0 + 0.0), 0);
+				ResurrectedHistoryGuideUncertainty = PrevHistoryGuide.SampleLevel(GlobalBilinearClampedSampler, float3(KernelHistoryBufferUV, ResurrectionFrameIndex * 2.0 + 1.0), 0).r;
 			#else
-				ResurrectedHistoryGuideColor = SampleTextureCatrom(
-					PrevHistoryGuide,
-					KernelHistoryBufferUV,
-					ResurrectionFrameIndex,
-					ResurrectionGuideUVViewportBilinearMin,
-					ResurrectionGuideUVViewportBilinearMax);
+				tsr_half4 RawGuide                 = PrevHistoryGuide.SampleLevel(GlobalBilinearClampedSampler, float3(KernelHistoryBufferUV, ResurrectionFrameIndex), 0);
+				ResurrectedHistoryGuideColor = RawGuide.rgb;
+				ResurrectedHistoryGuideUncertainty = RawGuide.a;
 			#endif
 			ResurrectedHistoryGuideColor = CorrectGuideColorExposure(ResurrectedHistoryGuideColor, ResurrectionPreExposureCorrection);
 		}
@@ -594,8 +612,9 @@ void MainCS(
 		{
 			#if CONFIG_SCENE_COLOR_ALPHA
 				ReprojectedHistoryGuideOutput[tsr_short3(OutputPixelPos, 0)] = ReprojectedHistoryGuideColor;
+				ReprojectedHistoryGuideOutput[tsr_short3(OutputPixelPos, 1)] = tsr_half4(ReprojectedHistoryGuideUncertainty, 0, 0, 0);
 			#else
-				ReprojectedHistoryGuideOutput[tsr_short3(OutputPixelPos, 0)] = tsr_half4(ReprojectedHistoryGuideColor.xyz, tsr_half(1.0));
+				ReprojectedHistoryGuideOutput[tsr_short3(OutputPixelPos, 0)] = tsr_half4(ReprojectedHistoryGuideColor.xyz, ReprojectedHistoryGuideUncertainty);
 			#endif
 			#if CONFIG_MOIRE_REPROJECTION
 				ReprojectedHistoryMoireOutput[tsr_short3(OutputPixelPos, 0)] = ReprojectedHistoryMoireColor;
@@ -605,9 +624,10 @@ void MainCS(
 		#if CONFIG_RESURRECTION_REPROJECTION
 		{
 			#if CONFIG_SCENE_COLOR_ALPHA
-				ReprojectedHistoryGuideOutput[tsr_short3(OutputPixelPos, 1)] = ResurrectedHistoryGuideColor;
+				ReprojectedHistoryGuideOutput[tsr_short3(OutputPixelPos, 2)] = ResurrectedHistoryGuideColor;
+				ReprojectedHistoryGuideOutput[tsr_short3(OutputPixelPos, 3)] = tsr_half4(ResurrectedHistoryGuideUncertainty, 0, 0, 0);
 			#else
-				ReprojectedHistoryGuideOutput[tsr_short3(OutputPixelPos, 1)] = tsr_half4(ResurrectedHistoryGuideColor.xyz, tsr_half(1.0));
+				ReprojectedHistoryGuideOutput[tsr_short3(OutputPixelPos, 1)] = tsr_half4(ResurrectedHistoryGuideColor.xyz, ResurrectedHistoryGuideUncertainty);
 			#endif
 		}
 		#endif
diff --git a/Engine/Shaders/Private/TemporalSuperResolution/TSRRejectShading.usf b/Engine/Shaders/Private/TemporalSuperResolution/TSRRejectShading.usf
--- a/Engine/Shaders/Private/TemporalSuperResolution/TSRRejectShading.usf
+++ b/Engine/Shaders/Private/TemporalSuperResolution/TSRRejectShading.usf
@@ -121,13 +121,15 @@ float PerceptionAdd;
 Texture2D<tsr_halfC> InputTexture;
 Texture2D<tsr_half>  InputMoireLumaTexture;
 Texture2D<tsr_half4> InputSceneTranslucencyTexture;
-Texture2D<tsr_halfC> ReprojectedHistoryGuideTexture;
+Texture2D<tsr_half4> ReprojectedHistoryGuideTexture;
+Texture2D<tsr_half>  ReprojectedHistoryGuideMetadataTexture;
 Texture2D<tsr_half4> ReprojectedHistoryMoireTexture;
-Texture2D<tsr_halfC> ResurrectedHistoryGuideTexture;
+Texture2D<tsr_half4> ResurrectedHistoryGuideTexture;
+Texture2D<tsr_half>  ResurrectedHistoryGuideMetadataTexture;
 Texture2D<tsr_half2> DecimateMaskTexture;
 Texture2D<tsr_ushort> IsMovingMaskTexture;
 
-RWTexture2D<tsr_half4> HistoryGuideOutput;
+RWTexture2DArray<tsr_half4> HistoryGuideOutput;
 RWTexture2D<tsr_half4> HistoryMoireOutput;
 RWTexture2D<tsr_half4> HistoryRejectionOutput;
 RWTexture2D<tsr_halfC> InputSceneColorOutput;
@@ -444,6 +446,7 @@ void MainCS(
 		tsr_simd_half4 OriginalTranslucencyInput;
 	#endif
 	tsr_simd_halfC OriginalHistory;
+	tsr_simd_half OriginalHistoryUncertainty;
 	tsr_simd_half VelocityEdge;
 	tsr_simd_bool bIsValid;
 	#if CONFIG_MOIRE_DETECTION
@@ -453,6 +456,7 @@ void MainCS(
 	#endif
 	#if CONFIG_HISTORY_RESURRECTION
 		tsr_simd_halfC OriginalResurrectedHistory;
+		tsr_simd_half OriginalResurrectedHistoryUncertainty;
 	#endif
 	PLATFORM_SPECIFIC_ISOLATE
 	{
@@ -472,12 +476,28 @@ void MainCS(
 			TranslucencyTextureUV = fastClamp(TranslucencyTextureUV, TranslucencyTextureUVMin, TranslucencyTextureUVMax);
 
 			tsr_half2 DecimateMask = DecimateMaskTexture[ClampedFetchPixelPos];
+			tsr_half4 ReprojectedGuide = ReprojectedHistoryGuideTexture[ClampedFetchPixelPos];
+			#if CONFIG_SCENE_COLOR_ALPHA
+				tsr_half ReprojectedGuideMetadata = ReprojectedHistoryGuideMetadataTexture[ClampedFetchPixelPos];
+			#endif
+			#if CONFIG_HISTORY_RESURRECTION
+				tsr_half4 ResurrectedGuide = ResurrectedHistoryGuideTexture[ClampedFetchPixelPos];
+				#if CONFIG_SCENE_COLOR_ALPHA
+					tsr_half ResurrectedGuideMetadata = ReprojectedHistoryGuideMetadataTexture[ClampedFetchPixelPos];
+				#endif
+			#endif
 
 			OriginalOpaqueInput.SetElement(SimdIndex, InputTexture[ClampedFetchPixelPos]);
 			#if CONFIG_COMPOSE_TRANSLUCENCY
 				OriginalTranslucencyInput.SetElement(SimdIndex, InputSceneTranslucencyTexture.SampleLevel(GlobalBilinearClampedSampler, TranslucencyTextureUV, 0));
 			#endif
-			OriginalHistory.SetElement(SimdIndex, ReprojectedHistoryGuideTexture[ClampedFetchPixelPos]);
+			#if CONFIG_SCENE_COLOR_ALPHA
+				OriginalHistory.SetElement(SimdIndex, ReprojectedGuide);
+				OriginalHistoryUncertainty.SetElement(SimdIndex, ReprojectedGuideMetadata);
+			#else
+				OriginalHistory.SetElement(SimdIndex, ReprojectedGuide.rgb);
+				OriginalHistoryUncertainty.SetElement(SimdIndex, ReprojectedGuide.a);
+			#endif
 			VelocityEdge.SetElement(SimdIndex, DecimateMask.g);
 			bIsValid.SetElement(SimdIndex, DecimateMask.r > tsr_half(PARALLAX_REJECTION_MASK_THRESHOLD));
 			#if CONFIG_MOIRE_DETECTION
@@ -486,7 +506,13 @@ void MainCS(
 				bIsStatic.SetElement(SimdIndex, tsr_half(1.0) - tsr_half(IsMovingMaskTexture[ClampedFetchPixelPos]) * rcp(tsr_half(255.0)));
 			#endif
 			#if CONFIG_HISTORY_RESURRECTION
-				OriginalResurrectedHistory.SetElement(SimdIndex, ResurrectedHistoryGuideTexture[ClampedFetchPixelPos]);
+				#if CONFIG_SCENE_COLOR_ALPHA
+					OriginalResurrectedHistory.SetElement(SimdIndex, ResurrectedGuide);
+					OriginalResurrectedHistoryUncertainty.SetElement(SimdIndex, ResurrectedGuideMetadata);
+				#else
+					OriginalResurrectedHistory.SetElement(SimdIndex, ResurrectedGuide.rgb);
+					OriginalResurrectedHistoryUncertainty.SetElement(SimdIndex, ResurrectedGuide.a);
+				#endif
 			#endif
 		}
 	}
@@ -689,6 +715,7 @@ void MainCS(
 				
 				RejectionBlendFinal = select(bResurrectHistory, ResurrectionRejectionBlendFinal, RejectionBlendFinal);
 				RejectionClampBlend = select(bResurrectHistory, ResurrectionRejectionClampBlend, RejectionClampBlend);
+				OriginalHistoryUncertainty = select(bResurrectHistory, OriginalResurrectedHistoryUncertainty, OriginalHistoryUncertainty);
 			}
 		}
 		#endif
@@ -725,11 +752,13 @@ void MainCS(
 		FinalGuide = ExposedInput * tsr_simd_halfC::Vectorize(BlendFinal) - OriginalHistory * tsr_simd_halfC::Vectorize(BlendFinal - 1.0);
 	}
 	
+	tsr_simd_half InputUncertainty = VelocityEdge;
+	
 	// Controls how the history and its validity is clamped.
 	tsr_simd_half DisableHistoryClamp;
 	tsr_simd_half IncreaseValidityMultiplier;
 	{
-		DisableHistoryClamp        = select(and(bIsDisoccluded, not(bResurrectHistory)), tsr_simd_half::Const(0.0), min(RejectionClampBlend, VelocityEdge));
+		DisableHistoryClamp        = select(and(bIsDisoccluded, not(bResurrectHistory)), tsr_simd_half::Const(0.0), min(RejectionClampBlend, min(InputUncertainty, Min3x3(OriginalHistoryUncertainty))));
 		IncreaseValidityMultiplier = select(and(bIsDisoccluded, not(bResurrectHistory)), tsr_simd_half::Const(0.0), RejectionClampBlend);
 	}
 
@@ -780,9 +809,10 @@ void MainCS(
 			#endif
 
 			#if CONFIG_SCENE_COLOR_ALPHA
-				HistoryGuideOutput[OutputPixelPos] = FinalGuideColor;
+				HistoryGuideOutput[tsr_short3(OutputPixelPos, 0)] = FinalGuideColor;
+				HistoryGuideOutput[tsr_short3(OutputPixelPos, 1)] = tsr_half4(InputUncertainty.GetElement(SimdIndex), 0, 0, 0);
 			#else
-				HistoryGuideOutput[OutputPixelPos] = tsr_half4(FinalGuideColor.xyz, tsr_half(1.0));
+				HistoryGuideOutput[tsr_short3(OutputPixelPos, 0)] = tsr_half4(FinalGuideColor.xyz, InputUncertainty.GetElement(SimdIndex));
 			#endif
 			HistoryRejectionOutput[OutputPixelPos] = tsr_half4(
 				RejectionBlendFinal.GetElement(SimdIndex) - tsr_half(0.5 * rcp(256.0)),
diff --git a/Engine/Shaders/Public/LaneVectorization.ush b/Engine/Shaders/Public/LaneVectorization.ush
--- a/Engine/Shaders/Public/LaneVectorization.ush
+++ b/Engine/Shaders/Public/LaneVectorization.ush
@@ -1010,6 +1010,7 @@ TVECTOR_FUNCTION_2PARAMS_DIFF_RETURN(/* VectorSize = */ 1, dot);
 	
 TVECTOR_FUNCTION_3PARAMS(select_internal, bool);
 TVECTOR_FUNCTION_3PARAMS(clamp, ScalarType);
+TVECTOR_FUNCTION_3PARAMS(lerp, ScalarType);
 TVECTOR_FUNCTION_3PARAMS(min3, ScalarType);
 TVECTOR_FUNCTION_3PARAMS(max3, ScalarType);
 #if COMPILER_SUPPORTS_MED3
diff --git a/Engine/Source/Runtime/Renderer/Private/PostProcess/TemporalSuperResolution.cpp b/Engine/Source/Runtime/Renderer/Private/PostProcess/TemporalSuperResolution.cpp
--- a/Engine/Source/Runtime/Renderer/Private/PostProcess/TemporalSuperResolution.cpp
+++ b/Engine/Source/Runtime/Renderer/Private/PostProcess/TemporalSuperResolution.cpp
@@ -612,12 +612,14 @@ class FTSRRejectShadingCS : public FTSRShader
 		SHADER_PARAMETER_RDG_TEXTURE(Texture2D, InputMoireLumaTexture)
 		SHADER_PARAMETER_RDG_TEXTURE(Texture2D, InputSceneTranslucencyTexture)
 		SHADER_PARAMETER_RDG_TEXTURE_SRV(Texture2D, ReprojectedHistoryGuideTexture)
+		SHADER_PARAMETER_RDG_TEXTURE_SRV(Texture2D, ReprojectedHistoryGuideMetadataTexture)
 		SHADER_PARAMETER_RDG_TEXTURE_SRV(Texture2D, ReprojectedHistoryMoireTexture)
 		SHADER_PARAMETER_RDG_TEXTURE_SRV(Texture2D, ResurrectedHistoryGuideTexture)
+		SHADER_PARAMETER_RDG_TEXTURE_SRV(Texture2D, ResurrectedHistoryGuideMetadataTexture)
 		SHADER_PARAMETER_RDG_TEXTURE(Texture2D, DecimateMaskTexture)
 		SHADER_PARAMETER_RDG_TEXTURE_SRV(Texture2D, IsMovingMaskTexture)
 
-		SHADER_PARAMETER_RDG_TEXTURE_UAV(RWTexture2D, HistoryGuideOutput)
+		SHADER_PARAMETER_RDG_TEXTURE_UAV(RWTexture2DArray, HistoryGuideOutput)
 		SHADER_PARAMETER_RDG_TEXTURE_UAV(RWTexture2D, HistoryMoireOutput)
 		SHADER_PARAMETER_RDG_TEXTURE_UAV(RWTexture2D, HistoryRejectionOutput)
 		SHADER_PARAMETER_RDG_TEXTURE_UAV(RWTexture2D, InputSceneColorOutput)
@@ -1404,6 +1406,7 @@ FDefaultTemporalUpscaler::FOutputs AddTemporalSuperResolutionPasses(
 
 	// Allocate a new history
 	FTSRHistoryTextures History;
+	const int32 HistoryColorGuideSliceCountWithoutResurrection = bSupportsAlpha ? 2 : 1;
 	{
 		{
 			bool bRequires2Mips = HistorySize == OutputRect.Size() && !bIsOutputDifferentThanHighFrequency && PassInputs.bGenerateOutputMip1;
@@ -1434,10 +1437,17 @@ FDefaultTemporalUpscaler::FOutputs AddTemporalSuperResolutionPasses(
 				bSupportsAlpha ? PF_FloatRGBA : PF_A2B10G10R10,
 				FClearValueBinding::None,
 				TexCreate_ShaderResource | TexCreate_UAV,
-				HistorySliceSequence.FrameStorageCount);
-
+				HistorySliceSequence.FrameStorageCount * HistoryColorGuideSliceCountWithoutResurrection);
 			History.GuideArray = GraphBuilder.CreateTexture(Desc, TEXT("TSR.History.Guide"));
-			Desc.Format = PF_R8G8B8A8;
+		}
+
+		{
+			FRDGTextureDesc Desc = FRDGTextureDesc::Create2DArray(
+				InputExtent,
+				PF_R8G8B8A8,
+				FClearValueBinding::None,
+				TexCreate_ShaderResource | TexCreate_UAV,
+				HistorySliceSequence.FrameStorageCount);
 			History.MoireArray = GraphBuilder.CreateTexture(Desc, TEXT("TSR.History.Moire"));
 		}
 	}
@@ -1819,8 +1829,7 @@ FDefaultTemporalUpscaler::FOutputs AddTemporalSuperResolutionPasses(
 				History.GuideArray->Desc.Format,
 				FClearValueBinding::None,
 				/* InFlags = */ TexCreate_ShaderResource | TexCreate_UAV,
-				/* InArraySize = */ bCanResurrectHistory ? 2 : 1);
-
+				/* InArraySize = */ (bCanResurrectHistory ? 2 : 1) * HistoryColorGuideSliceCountWithoutResurrection);
 			ReprojectedHistoryGuideTexture = GraphBuilder.CreateTexture(Desc, TEXT("TSR.ReprojectedHistoryGuide"));
 		}
 
@@ -1831,7 +1840,6 @@ FDefaultTemporalUpscaler::FOutputs AddTemporalSuperResolutionPasses(
 				FClearValueBinding::None,
 				/* InFlags = */ TexCreate_ShaderResource | TexCreate_UAV,
 				/* InArraySize = */ 1);
-
 			ReprojectedHistoryMoireTexture = GraphBuilder.CreateTexture(Desc, TEXT("TSR.ReprojectedHistoryMoire"));
 		}
 
@@ -1997,13 +2005,24 @@ FDefaultTemporalUpscaler::FOutputs AddTemporalSuperResolutionPasses(
 		}
 		PassParameters->InputSceneTranslucencyTexture = SeparateTranslucencyTexture;
 		PassParameters->ReprojectedHistoryGuideTexture = GraphBuilder.CreateSRV(FRDGTextureSRVDesc::CreateForSlice(
-			ReprojectedHistoryGuideTexture, /* SliceIndex = */ 0));
+			ReprojectedHistoryGuideTexture, /* SliceIndex = */ 0 * HistoryColorGuideSliceCountWithoutResurrection));
+		if (bSupportsAlpha)
+		{
+			PassParameters->ReprojectedHistoryGuideMetadataTexture = GraphBuilder.CreateSRV(FRDGTextureSRVDesc::CreateForSlice(
+				ReprojectedHistoryGuideTexture, /* SliceIndex = */ 0 * HistoryColorGuideSliceCountWithoutResurrection + 1));
+		}
 		PassParameters->ReprojectedHistoryMoireTexture = ReprojectedHistoryMoireTexture ? GraphBuilder.CreateSRV(FRDGTextureSRVDesc::CreateForSlice(
 			ReprojectedHistoryMoireTexture, /* SliceIndex = */ 0)) : nullptr;
 		if (bCanResurrectHistory)
 		{
 			PassParameters->ResurrectedHistoryGuideTexture = GraphBuilder.CreateSRV(FRDGTextureSRVDesc::CreateForSlice(
-				ReprojectedHistoryGuideTexture, /* SliceIndex = */ 1));
+				ReprojectedHistoryGuideTexture, /* SliceIndex = */ 1 * HistoryColorGuideSliceCountWithoutResurrection + 0));
+
+			if (bSupportsAlpha)
+			{
+				PassParameters->ResurrectedHistoryGuideMetadataTexture = GraphBuilder.CreateSRV(FRDGTextureSRVDesc::CreateForSlice(
+					ReprojectedHistoryGuideTexture, /* SliceIndex = */ 1 * HistoryColorGuideSliceCountWithoutResurrection + 1));
+			}
 		}
 		PassParameters->DecimateMaskTexture = DecimateMaskTexture;
 		PassParameters->IsMovingMaskTexture = IsMovingMaskTexture;
@@ -2017,9 +2036,8 @@ FDefaultTemporalUpscaler::FOutputs AddTemporalSuperResolutionPasses(
 			else
 			{
 				FRDGTextureUAVDesc GuideUAVDesc(History.GuideArray);
-				GuideUAVDesc.FirstArraySlice = CurrentFrameSliceIndex;
-				GuideUAVDesc.NumArraySlices = 1;
-				GuideUAVDesc.DimensionOverride = ETextureDimension::Texture2D;
+				GuideUAVDesc.FirstArraySlice = CurrentFrameSliceIndex * HistoryColorGuideSliceCountWithoutResurrection;
+				GuideUAVDesc.NumArraySlices = HistoryColorGuideSliceCountWithoutResurrection;
 
 				PassParameters->HistoryGuideOutput = GraphBuilder.CreateUAV(GuideUAVDesc);
 			}
