Path Tracer: Implement Rank-1 Lattice Random Sequence in the RandomSequence interface

This is much faster than Sobol and appears to have similar quality. The default is still Sobol for now until we can accurately quantify the convergence differences.

diff --git a/Engine/Shaders/Private/PathTracing/Utilities/PathTracingRandomSequence.ush b/Engine/Shaders/Private/PathTracing/Utilities/PathTracingRandomSequence.ush
--- a/Engine/Shaders/Private/PathTracing/Utilities/PathTracingRandomSequence.ush
+++ b/Engine/Shaders/Private/PathTracing/Utilities/PathTracingRandomSequence.ush
@@ -11,6 +11,7 @@
 #define RANDSEQ_PURERANDOM	0
 #define RANDSEQ_HALTON		1
 #define RANDSEQ_OWENSOBOL	2
+#define RANDSEQ_LATTICE     3
 
 // Select a default if none was specified
 #ifndef RANDSEQ
@@ -22,7 +23,7 @@
 #endif
 
 #ifndef RANDSEQ_ERROR_DIFFUSION
-#define RANDSEQ_ERROR_DIFFUSION 1 // Enabled by default, only works when using RANDSEQ_OWENSOBOL
+#define RANDSEQ_ERROR_DIFFUSION 1 // Enabled by default, only works when using RANDSEQ_OWENSOBOL or RANDSEQ_LATTICE
 #endif
 
 #ifndef RANDSEQ_UNROLL_SOBOL
@@ -268,6 +269,33 @@ uint EvolveSobolSeed(inout uint Seed)
 #endif
 }
 
+float4 LatticeSampler(uint SampleIndex, inout uint Seed)
+{
+#if 1
+	// Same as FastOwenScrambling, but without the final reversebits
+	uint LatticeIndex = SampleIndex + EvolveSobolSeed(Seed);
+	LatticeIndex ^= LatticeIndex * 0x9c117646u;
+	LatticeIndex ^= LatticeIndex * 0xe0705d72u;
+#else
+	// This is a higher quality owen-scrambling hash but doesn't appear to make much difference
+	// https://psychopath.io/post/2021_01_30_building_a_better_lk_hash
+	uint LatticeIndex = SampleIndex, S = EvolveSobolSeed(Seed);
+	LatticeIndex ^= LatticeIndex * 0x3d20adeau;
+	LatticeIndex += S;
+	LatticeIndex *= (S >> 16) | 1;
+	LatticeIndex ^= LatticeIndex * 0x05526c56u;
+	LatticeIndex ^= LatticeIndex * 0x53a22864u;
+#endif
+
+	// Lattice parameters taken from:
+	// Weighted compound integration rules with higher order convergence for all N
+	// Fred J. Hickernell, Peter Kritzer, Frances Y. Kuo, Dirk Nuyens
+	// Numerical Algorithms - February 2012
+	uint4 Result = LatticeIndex * uint4(1, 364981, 245389, 97823);
+
+	return (Result >> 8) * 5.96046447754e-08; // * 2^-24
+}
+
 float4 SobolSampler(uint SampleIndex, inout uint Seed)
 {
 	// first scramble the index to decorelate from other 4-tuples
@@ -314,7 +342,7 @@ void RandomSequence_Initialize(inout RandomSequence RandSequence, uint PositionS
 #elif RANDSEQ == RANDSEQ_HALTON
 	RandSequence.SampleIndex = TimeSeed + StrongIntegerHash(PositionSeed);
 	RandSequence.SampleSeed = 0; // this sampler just counts dimensions
-#elif RANDSEQ == RANDSEQ_OWENSOBOL
+#elif RANDSEQ == RANDSEQ_OWENSOBOL || RANDSEQ == RANDSEQ_LATTICE
 	// pre-compute bit reversal needed for FastOwenScrambling since this index doesn't change
 	RandSequence.SampleIndex = reversebits(TimeSeed);
 	// change seed to get a unique sequence per pixel
@@ -329,7 +357,7 @@ void RandomSequence_Initialize(inout RandomSequence RandSequence, uint PositionS
 // taken if taking one sample per frame like the path tracer).
 void RandomSequence_Initialize(inout RandomSequence RandSequence, uint2 PixelCoord, uint SampleIndex, uint TimeSeed0, uint MaxSamples)
 {
-#if RANDSEQ_ERROR_DIFFUSION == 1 && RANDSEQ == RANDSEQ_OWENSOBOL
+#if RANDSEQ_ERROR_DIFFUSION == 1 && (RANDSEQ == RANDSEQ_OWENSOBOL || RANDSEQ == RANDSEQ_LATTICE)
 	// The core idea is to use a single sobol pattern across the screen using a space filling curve to correlate nearby pixels
 	// This achieves a much better error diffusion than the random sequence assignment of the method above
 	// This algorithm is covered in the following paper:
@@ -398,6 +426,8 @@ float RandomSequence_GenerateSample1D(inout RandomSequence RandSequence)
 	// rely on compiler to optimize out dead code for now
 	// #dxr_todo: benchmark and specialize if needed
 	Result = SobolSampler(RandSequence.SampleIndex, RandSequence.SampleSeed).x;
+#elif RANDSEQ == RANDSEQ_LATTICE
+	Result = LatticeSampler(RandSequence.SampleIndex, RandSequence.SampleSeed).x;
 #else
 #error "Unknown random sequence chosen in path tracer"
 #endif
@@ -418,6 +448,8 @@ float2 RandomSequence_GenerateSample2D(inout RandomSequence RandSequence)
 	// rely on compiler to optimize out dead code for now
 	// #dxr_todo: benchmark and specialize if needed
 	Result = SobolSampler(RandSequence.SampleIndex, RandSequence.SampleSeed).xy;
+#elif RANDSEQ == RANDSEQ_LATTICE
+	Result = LatticeSampler(RandSequence.SampleIndex, RandSequence.SampleSeed).xy;
 #else
 #error "Unknown random sequence chosen in path tracer"
 #endif
@@ -440,6 +472,8 @@ float3 RandomSequence_GenerateSample3D(inout RandomSequence RandSequence)
 	// rely on compiler to optimize out dead code for now
 	// #dxr_todo: benchmark and specialize if needed
 	Result = SobolSampler(RandSequence.SampleIndex, RandSequence.SampleSeed).xyz;
+#elif RANDSEQ == RANDSEQ_LATTICE
+	Result = LatticeSampler(RandSequence.SampleIndex, RandSequence.SampleSeed).xyz;
 #else
 #error "Unknown random sequence chosen in path tracer"
 #endif
@@ -463,6 +497,8 @@ float4 RandomSequence_GenerateSample4D(inout RandomSequence RandSequence)
 	RandSequence.SampleSeed += 4;
 #elif RANDSEQ == RANDSEQ_OWENSOBOL
 	Result = SobolSampler(RandSequence.SampleIndex, RandSequence.SampleSeed);
+#elif RANDSEQ == RANDSEQ_LATTICE
+	Result = LatticeSampler(RandSequence.SampleIndex, RandSequence.SampleSeed);
 #else
 #error "Unknown random sequence chosen in path tracer"
 #endif
