Add unroll hint to shader compiler to help DXC on some platform.

diff --git a/Engine/Shaders/Private/LandscapeVertexFactory.ush b/Engine/Shaders/Private/LandscapeVertexFactory.ush
--- a/Engine/Shaders/Private/LandscapeVertexFactory.ush
+++ b/Engine/Shaders/Private/LandscapeVertexFactory.ush
@@ -105,6 +105,7 @@ float4 GetNeighborSectionLod(uint ComponentIndex)
     NeighborOffset[2] = int2(1, 0);
     NeighborOffset[3] = int2(0, 1);
 	float NeighborLod[4];
+	UNROLL
 	for (int i = 0; i < 4; i++)
 	{
 		int2 NeighborPos = NeighborOffset[i] + int2(LandscapeParameters.ComponentBaseX, LandscapeParameters.ComponentBaseY);
