Subsurface Profile has artifacts if it is behind a translucent material that uses its basecolor

diff --git a/Engine/Shaders/Private/DeferredShadingCommon.ush b/Engine/Shaders/Private/DeferredShadingCommon.ush
--- a/Engine/Shaders/Private/DeferredShadingCommon.ush
+++ b/Engine/Shaders/Private/DeferredShadingCommon.ush
@@ -913,6 +913,12 @@ bool AdjustBaseColorAndSpecularColorForSubsurfaceProfileLighting(inout float3 Ba
 }
 void AdjustBaseColorAndSpecularColorForSubsurfaceProfileLighting(inout float3 BaseColor, inout float3 SpecularColor, inout float Specular, bool bChecker)
 {
+// If this function is called in translucent pass, we should not modify the BaseColor, instead directly read the base color as it is. 
+// The translucent pass is called after the main post process subsurface scattering pass, it will not affect previous logic.
+#if MATERIALBLENDING_TRANSLUCENT
+	return;
+#endif
+
 	const bool bCheckerboardRequired = AdjustBaseColorAndSpecularColorForSubsurfaceProfileLighting(BaseColor, Specular, bChecker);
 	if (bCheckerboardRequired)
 	{
diff --git a/Engine/Shaders/Private/GBufferHelpers.ush b/Engine/Shaders/Private/GBufferHelpers.ush
--- a/Engine/Shaders/Private/GBufferHelpers.ush
+++ b/Engine/Shaders/Private/GBufferHelpers.ush
@@ -348,6 +348,12 @@ void GBufferPreEncode(inout FGBufferData GBuffer, bool bChecker, float Geometric
 // the number of files that have to be touched on the initial checkin. Will remove duplicates once the first checkin is tested and stable.
 void AdjustBaseColorAndSpecularColorForSubsurfaceProfileLightingCopyHack(inout float3 BaseColor, inout float3 SpecularColor, inout float Specular, bool bChecker)
 {
+	// If this function is called in translucent pass, we should not modify the BaseColor, instead directly read the base color as it is. 
+	// The translucent pass is called after the main post process subsurface scattering pass, it will not affect previous logic.
+#if MATERIALBLENDING_TRANSLUCENT
+	return;
+#endif
+
 #if SUBSURFACE_CHANNEL_MODE == 0
 	// If SUBSURFACE_CHANNEL_MODE is 0, we can't support full-resolution lighting, so we 
 	// ignore View.bCheckerboardSubsurfaceProfileRendering
