DebugPrimitives: correct the offset for ColorUV samples

PopulateSceneDepth is used to 'forward' the colors to the target surface.
The UV generation is always generating from the top left, instead of being properly offset for stereo views
or r.Test.ViewRectOffset

Tested with OpenXR, emulatestereo, and r.Test.ViewRectOffset

diff --git a/Engine/Shaders/Private/PostProcessCompositePrimitives.usf b/Engine/Shaders/Private/PostProcessCompositePrimitives.usf
--- a/Engine/Shaders/Private/PostProcessCompositePrimitives.usf
+++ b/Engine/Shaders/Private/PostProcessCompositePrimitives.usf
@@ -258,9 +258,9 @@ void MainPopulateSceneDepthPS(
 	//This extended pass is primarily to compensate for the lack of scene color output when compositing the debug draw primitives with SimpleElementShaders
 	//Draw out the color whilst compensating for the viewport not necessarily starting at 0,0. 
 	//Does not use jitter for Scene Color as it is not upscaled for this pass
-	float2 ColorUV = clamp(ViewportUV * Color_UVViewportSize, 
-						Color_UVViewportBilinearMin - Color_UVViewportMin,
-						Color_UVViewportBilinearMax - Color_UVViewportMin);
+	float2 ColorUV = clamp(Color_UVViewportMin + (ViewportUV * Color_UVViewportSize), 
+						Color_UVViewportBilinearMin,
+						Color_UVViewportBilinearMax);
 	ColorUV = clamp(ColorUV, float2(0, 0), float2(1, 1));
 	OutColor = Texture2DSampleLevel(ColorTexture, ColorSampler, ColorUV, 0).rgba;
 
