Improves the visualization of TSR's resurrection with the rejection mask

diff --git a/Engine/Shaders/Private/TemporalSuperResolution/TSRVisualize.usf b/Engine/Shaders/Private/TemporalSuperResolution/TSRVisualize.usf
--- a/Engine/Shaders/Private/TemporalSuperResolution/TSRVisualize.usf
+++ b/Engine/Shaders/Private/TemporalSuperResolution/TSRVisualize.usf
@@ -113,16 +113,18 @@ tsr_half3 DisplayParallaxDisocclusionMask(FContext Ctx)
 
 tsr_half3 DisplayResurrectionMask(FContext Ctx)
 {
+	tsr_half LowFrequencyRejection = HistoryRejectionTexture.SampleLevel(GlobalPointClampedSampler, Ctx.InputUV, 0).r;
 	bool ResurrectionMask = (tsr_ushort(HistoryRejectionTexture.SampleLevel(GlobalPointClampedSampler, Ctx.InputUV, 0).a * tsr_half(255.0)) & tsr_ushort(0x2)) != tsr_ushort(0x0);
 	if (!bCanResurrectHistory)
 	{
 		return SampleTSROutputGray(Ctx) * kRed;
 	}
-	return SampleTSROutputGray(Ctx) * select(ResurrectionMask, kGreen, kWhite);
+	return SampleTSROutputGray(Ctx) * select(ResurrectionMask, lerp(kRed, kGreen, LowFrequencyRejection), kWhite);
 }
 
 tsr_half3 DisplayResurrectedColor(FContext Ctx)
 {
+	tsr_half LowFrequencyRejection = HistoryRejectionTexture.SampleLevel(GlobalPointClampedSampler, Ctx.InputUV, 0).r;
 	float ClosestDeviceZ = ClosestDepthTexture.SampleLevel(GlobalPointClampedSampler, Ctx.InputUV, 0);
 	bool ResurrectionMask = (tsr_ushort(HistoryRejectionTexture.SampleLevel(GlobalPointClampedSampler, Ctx.InputUV, 0).a * tsr_half(255.0)) & tsr_ushort(0x2)) != tsr_ushort(0x0);
 
@@ -135,7 +137,7 @@ tsr_half3 DisplayResurrectedColor(FContext Ctx)
 
 	tsr_half3 ResurrectedColor = ResurrectedHistoryColorTexture.SampleLevel(GlobalBilinearClampedSampler, KernelHistoryBufferUV, 0) * tsr_half(ResurrectionPreExposureCorrection);
 
-	return GrayScaleColor(ResurrectedColor) * select(ResurrectionMask, kGreen, kWhite);
+	return GrayScaleColor(ResurrectedColor) * select(ResurrectionMask, lerp(kRed, kGreen, LowFrequencyRejection), kWhite);
 }
 
 tsr_half3 DisplaySpatialAntiAliasingMask(FContext Ctx)
