Fixes a regression on TSR's spatial anti-aliaser not outputing even pixels

diff --git a/Engine/Shaders/Private/TemporalSuperResolution/TSRSpatialAntiAliasing.usf b/Engine/Shaders/Private/TemporalSuperResolution/TSRSpatialAntiAliasing.usf
--- a/Engine/Shaders/Private/TemporalSuperResolution/TSRSpatialAntiAliasing.usf
+++ b/Engine/Shaders/Private/TemporalSuperResolution/TSRSpatialAntiAliasing.usf
@@ -520,7 +520,7 @@ void MainCS(
 		AntiAliasingOutput[dpv_lo(OutputPixelPos)] = tsr_ushort2(dpv_lo(EncodedTexelOffset), round(dpv_lo(NoiseFiltering) * tsr_half(255.0)));
 		
 		#if CONFIG_ENABLE_DUAL_PIXEL_VECTORIZATION
-			AntiAliasingOutput[dpv_lo(OutputPixelPos)] = tsr_ushort2(dpv_hi(EncodedTexelOffset), round(dpv_hi(NoiseFiltering) * tsr_half(255.0)));
+			AntiAliasingOutput[dpv_hi(OutputPixelPos)] = tsr_ushort2(dpv_hi(EncodedTexelOffset), round(dpv_hi(NoiseFiltering) * tsr_half(255.0)));
 		#endif
 
 		#if DEBUG_OUTPUT
