Fixes incorrect TSR's velocity edge on SM6

diff --git a/Engine/Shaders/Private/TemporalSuperResolution/TSRDecimateHistory.usf b/Engine/Shaders/Private/TemporalSuperResolution/TSRDecimateHistory.usf
--- a/Engine/Shaders/Private/TemporalSuperResolution/TSRDecimateHistory.usf
+++ b/Engine/Shaders/Private/TemporalSuperResolution/TSRDecimateHistory.usf
@@ -103,6 +103,7 @@ void LoadVelocityNeighborhood(
 			tsr_short2 Offset = tsr_short2(Offsets2x2[i]) * LaneOffsetSign;
 				
 			SampleInputPixelPos = InputPixelPos + Offset;
+			SampleInputPixelPos = ClampPixelOffset(SampleInputPixelPos, InputPixelPosMin, InputPixelPosMax);
 		}
 		#else
 		{
@@ -162,11 +163,11 @@ void ProcessVelocityNeighborhood(
 			const uint ButterflyArray0[] = { 0x0, 0x0, 0x0, 0x1 };
 			const uint ButterflyArray1[] = { 0x1, 0x2, 0x3, 0x2 };
 				
-			uint FetchIndex[2];
 			FWaveBroadcastSettings BroadcastSettings[2];
 			BroadcastSettings[0] = InitWaveXorButterfly(ButterflyArray0[PairIndex]);
 			BroadcastSettings[1] = InitWaveXorButterfly(ButterflyArray1[PairIndex]);
 				
+			uint FetchIndex[2];
 			if (PairIndex == 3)
 			{
 				FetchIndex[0] = 2;
@@ -176,7 +177,7 @@ void ProcessVelocityNeighborhood(
 			else
 			{
 				FetchIndex[0] = PairIndex + 1;
-				FetchIndex[1] = PairIndex + 1;
+				FetchIndex[1] = 0;
 				Offset0 = tsr_short2(Offsets2x2[PairIndex + 1]);
 			}
 
