Finally fixes TSR ghosting due to incorrect parallax disocclusion mask on Nvidia GPUs with D3D12 SM6

diff --git a/Engine/Shaders/Private/TemporalSuperResolution/TSRDecimateHistory.usf b/Engine/Shaders/Private/TemporalSuperResolution/TSRDecimateHistory.usf
--- a/Engine/Shaders/Private/TemporalSuperResolution/TSRDecimateHistory.usf
+++ b/Engine/Shaders/Private/TemporalSuperResolution/TSRDecimateHistory.usf
@@ -236,6 +236,7 @@ void LoadPrevAtomicTexturesSamples(
 	out uint HistoryClosestDeviceZSamples[4])
 {
 	float2 PrevInputBufferUV = (InputInfo_ScreenPosToViewportScale * PrevScreenPos + InputInfo_ScreenPosToViewportBias) * InputInfo_ExtentInverse;
+	PrevInputBufferUV = INVARIANT(PrevInputBufferUV);
 	FBilinearSampleInfos BilinearInter = GetBilinearSampleLevelInfos(PrevInputBufferUV, InputInfo_Extent, InputInfo_ExtentInverse);
 
 	UNROLL_N(4)
@@ -277,6 +278,7 @@ void ProcessPrevAtomicTexturesSamples(
 	float WorldDepth = ConvertFromDeviceZ(DeviceZ);
 	
 	float2 PrevInputBufferUV = (InputInfo_ScreenPosToViewportScale * PrevScreenPos + InputInfo_ScreenPosToViewportBias) * InputInfo_ExtentInverse;
+	PrevInputBufferUV = INVARIANT(PrevInputBufferUV);
 	FBilinearSampleInfos BilinearInter = GetBilinearSampleLevelInfos(PrevInputBufferUV, InputInfo_Extent, InputInfo_ExtentInverse);
 
 	UNROLL_N(4)
