Fix pixelated sun with light shaft bloom enabled on mobile.

diff --git a/Engine/Shaders/Private/PostProcessMobile.usf b/Engine/Shaders/Private/PostProcessMobile.usf
--- a/Engine/Shaders/Private/PostProcessMobile.usf
+++ b/Engine/Shaders/Private/PostProcessMobile.usf
@@ -67,6 +67,7 @@ float2 Circle(float Start, float Points, float Point)
 	return float2(sin(Rad), cos(Rad));
 }
 
+float BloomMaxBrightness;
 float BloomThreshold;
 
 half FocusDistFar()
@@ -163,6 +164,11 @@ void SunMaskPS_Mobile(
 		half2 DepthMaskScaleBias = SunConstDepthMaskScaleBias();
 		half FarAmount = saturate(InDepth * DepthMaskScaleBias.x + DepthMaskScaleBias.y);
 		half3 SunAmount = SceneColor.rgb * SunColorApertureDiv2.rgb;
+
+		float SunLuminance = max(Luminance(SunAmount), 6.10352e-5);
+		float AdjustedLuminance = clamp(SunLuminance - BloomThreshold, 0.0f, BloomMaxBrightness);
+		SunAmount = SunAmount / SunLuminance * AdjustedLuminance * 2.0f;
+
 	
 		half2 Pos = InUVPos.zw * 0.5 + 0.5;
 		half EdgeMask = 1.0f - Pos.x * (1.0f - Pos.x) * Pos.y * (1.0f - Pos.y) * 8.0f;
@@ -369,10 +375,17 @@ void BloomPS_Mobile(
 #endif
 
 #if MOBILE_USESUN
+#if MOBILE_USEDOF
 	half Sun0 = max(0.0, A0-1.0);
 	half Sun1 = max(0.0, A1-1.0);
 	half Sun2 = max(0.0, A2-1.0);
 	half Sun3 = max(0.0, A3-1.0);
+#else
+	half Sun0 = A0;
+	half Sun1 = A1;
+	half Sun2 = A2;
+	half Sun3 = A3;
+#endif
 
 	// Take average of sun intensity and adjust by bloom threshold.
 	Amount *= 0.25;
diff --git a/Engine/Source/Runtime/Renderer/Private/LightShaftRendering.cpp b/Engine/Source/Runtime/Renderer/Private/LightShaftRendering.cpp
--- a/Engine/Source/Runtime/Renderer/Private/LightShaftRendering.cpp
+++ b/Engine/Source/Runtime/Renderer/Private/LightShaftRendering.cpp
@@ -90,6 +90,7 @@ FMobileLightShaftInfo GetMobileLightShaftInfo(const FViewInfo& View, const FLigh
 	LightShaft.ColorApply = LightSceneInfo.BloomTint;
 	LightShaft.ColorMask *= BloomScale;
 	LightShaft.ColorApply *= BloomScale;
+	LightShaft.BloomMaxBrightness = LightSceneInfo.BloomMaxBrightness;
 	return LightShaft;
 }
 
diff --git a/Engine/Source/Runtime/Renderer/Private/LightShaftRendering.h b/Engine/Source/Runtime/Renderer/Private/LightShaftRendering.h
--- a/Engine/Source/Runtime/Renderer/Private/LightShaftRendering.h
+++ b/Engine/Source/Runtime/Renderer/Private/LightShaftRendering.h
@@ -16,6 +16,7 @@ struct FMobileLightShaftInfo
 	FVector2D Center = FVector2D::ZeroVector;
 	FLinearColor ColorMask = FLinearColor::Transparent;
 	FLinearColor ColorApply = FLinearColor::Transparent;
+	float BloomMaxBrightness = 100.0f;
 };
 
 // Returns mobile light shaft info for the light.
diff --git a/Engine/Source/Runtime/Renderer/Private/PostProcess/PostProcessMobile.cpp b/Engine/Source/Runtime/Renderer/Private/PostProcess/PostProcessMobile.cpp
--- a/Engine/Source/Runtime/Renderer/Private/PostProcess/PostProcessMobile.cpp
+++ b/Engine/Source/Runtime/Renderer/Private/PostProcess/PostProcessMobile.cpp
@@ -675,6 +675,8 @@ public:
 		SHADER_PARAMETER_RDG_TEXTURE(Texture2D, SceneColorTexture)
 		SHADER_PARAMETER_SAMPLER(SamplerState, SceneColorSampler)
 		SHADER_PARAMETER(FVector4f, SunColorApertureDiv2)
+		SHADER_PARAMETER(float, BloomMaxBrightness)
+		SHADER_PARAMETER(float, BloomThreshold)
 		RENDER_TARGET_BINDING_SLOTS()
 	END_SHADER_PARAMETER_STRUCT()
 
@@ -775,6 +777,9 @@ FMobileSunMaskOutputs AddMobileSunMaskPass(FRDGBuilder& GraphBuilder, const FVie
 	PassParameters->SunColorApertureDiv2.Z = MobileLightShaft.ColorMask.B;
 	PassParameters->SunColorApertureDiv2.W = GetMobileDepthOfFieldScale(View) * 0.5f;
 
+	PassParameters->BloomMaxBrightness = MobileLightShaft.BloomMaxBrightness;
+	PassParameters->BloomThreshold = View.FinalPostProcessSettings.BloomThreshold;
+
 	FMobileSunMaskPS::FPermutationDomain PermutationVector = FMobileSunMaskPS::BuildPermutationVector(Inputs.bUseSun, Inputs.bUseDof, Inputs.bUseDepthTexture, Inputs.bUseMetalMSAAHDRDecode);
 	TShaderMapRef<FMobileSunMaskPS> PixelShader(View.ShaderMap, PermutationVector);
 
