Path Tracer: Fix distant lights not obeying cast volume shadows

Cleanup various small issues with heterogenous volumes and remove unused arguments (PathRoughness and Bounce were passed around inconsistently but not used in the end)

diff --git a/Engine/Shaders/Private/PathTracing/PathTracingCore.ush b/Engine/Shaders/Private/PathTracing/PathTracingCore.ush
--- a/Engine/Shaders/Private/PathTracing/PathTracingCore.ush
+++ b/Engine/Shaders/Private/PathTracing/PathTracingCore.ush
@@ -614,7 +614,7 @@ FPathTracingPayload TraceTransparentRay(inout FPathState PathState, int Bounce,
 			if (UseAnalyticTransmittance)
 			{
 				// analytically handle the transmittance to the next surface (which will be higher quality than the transmittance implicitly computed above and prevent noise with emissive surfaces in volumes)
-				PathState.PathThroughput = VolumeGetTransmittance(PathState.PathThroughput, PathState.Ray.Origin, PathState.Ray.Direction, Interval, PathState.PathRoughness, PathState.RandSequence);
+				PathState.PathThroughput = VolumeGetTransmittance(PathState.PathThroughput, PathState.Ray.Origin, PathState.Ray.Direction, Interval, PathState.RandSequence);
 			}
 
 			// Update our current volume intersection list to reflect the fact that we have made it to the end of this interval
@@ -817,7 +817,7 @@ float3 TraceTransparentVisibilityRay(FRayDesc Ray, int Bounce, float PathRoughne
 		}
 		if (VolumeIsectList.HitVolume())
 		{
-			Throughput = VolumeGetTransmittance(Throughput, Ray.Origin, Ray.Direction, VolumeIsectList, PathRoughness, RandSequence);
+			Throughput = VolumeGetTransmittance(Throughput, Ray.Origin, Ray.Direction, VolumeIsectList, RandSequence);
 		}
 		// add in the medium extinction
 		Throughput *= exp(-max(PackedPayload.GetTau(), 0.0));
@@ -1467,7 +1467,7 @@ bool PathTracingKernel(inout FPathState PathState, int Bounce)
 			{
 				// Account for the transmittance up to the current point within the current slice of volume
 				VolumeSegment.Interval.VolumeTMax = EquiAngularT;
-				VolumeSegment.Throughput = VolumeGetTransmittance(VolumeSegment.Throughput, Ro, Rd, VolumeSegment.Interval, PathState.PathRoughness, PathState.RandSequence);
+				VolumeSegment.Throughput = VolumeGetTransmittance(VolumeSegment.Throughput, Ro, Rd, VolumeSegment.Interval, PathState.RandSequence);
 
 				float3 Contrib = VolumeSegment.Throughput * SigmaS;
 
@@ -1645,7 +1645,7 @@ bool PathTracingKernel(inout FPathState PathState, int Bounce)
 
 				float AvgRoughness = ApproximateCaustics ? GetAverageRoughness(Payload) : 0.0;
 
-				const bool bCastShadows = CastsShadow(LightId);
+				const bool bCastShadows = bIsVolumeSample ? CastsVolumeShadow(LightId) : CastsShadow(LightId);
 				const uint MissShaderIndex = GetLightMissShaderIndex(LightId);
 				LightSample.RadianceOverPdf *= TraceTransparentVisibilityRay(LightRay, Bounce, AvgRoughness, MissShaderIndex, bCastShadows, PathState.RandSequence);
 
@@ -1812,7 +1812,7 @@ bool PathTracingKernel(inout FPathState PathState, int Bounce)
 			{
 				FRayDesc LightRay = PathState.Ray;
 				LightRay.TMax = LightResult.HitT;
-				const bool bCastShadows = CastsShadow(LightId);
+				const bool bCastShadows = bIsVolumeSample ? CastsVolumeShadow(LightId) : CastsShadow(LightId);
 				const uint MissShaderIndex = GetLightMissShaderIndex(LightId);
 				LightContrib *= TraceTransparentVisibilityRay(LightRay, Bounce, ShadowRayRoughness, MissShaderIndex, bCastShadows, PathState.RandSequence);
 				// the light made some contribution, and there was nothing along the shadow ray
diff --git a/Engine/Shaders/Private/PathTracing/Volume/PathTracingHeterogeneousVolumes.ush b/Engine/Shaders/Private/PathTracing/Volume/PathTracingHeterogeneousVolumes.ush
--- a/Engine/Shaders/Private/PathTracing/Volume/PathTracingHeterogeneousVolumes.ush
+++ b/Engine/Shaders/Private/PathTracing/Volume/PathTracingHeterogeneousVolumes.ush
@@ -44,10 +44,7 @@ FVolumeIntersection HeterogeneousVolumesIntersect(float3 Origin, float3 Directio
 				TranslatedWorldBoundsMax
 			);
 
-			if (RayHitT.y > RayHitT.x)
-			{
-				return CreateVolumeIntersection(RayHitT.x, RayHitT.y);
-			}
+			return CreateVolumeIntersection(RayHitT.x, RayHitT.y);
 		}
 	}
 
@@ -444,7 +441,7 @@ FVolumeTrackingResult HeterogeneousVolumesDeltaTracking(float3 PathThroughput, f
 		float RandValue = RandomSequence_GenerateSample1D(RandSequence);
 #define SAMPLE_ACHROMATIC 0
 #if SAMPLE_ACHROMATIC
-		float DeltaT = -log(RandValue) / Sample.SigmaBar.x;
+		float DeltaT = -log(1.0 - RandValue) / Sample.SigmaBar.x;
 #else
 		float DeltaT = SampleSpectralTransmittance(RandValue, Sample.SigmaBar, PathThroughput);
 		if (DeltaT < 0)
@@ -491,7 +488,7 @@ FVolumeTrackingResult HeterogeneousVolumesRatioTracking(float3 PathThroughput, f
 	{
 		float RandValue = RandomSequence_GenerateSample1D(RandSequence);
 #if SAMPLE_ACHROMATIC
-		float DeltaT = -log(RandValue) / Sample.SigmaBar.x;
+		float DeltaT = -log(1.0 - RandValue) / Sample.SigmaBar.x;
 #else
 		float DeltaT = SampleSpectralTransmittance(RandValue, Sample.SigmaBar, PathThroughput);
 		if (DeltaT < 0)
@@ -641,17 +638,16 @@ FVolumeTrackingResult HeterogeneousVolumesMajorantBasedTracking(float3 PathThrou
 	return Result;
 }
 
-float3 HeterogeneousVolumesGetTransmittance(float3 PathThroughput, float3 Origin, float3 Direction, float TMin, float TMax, float PathRoughness, inout RandomSequence RandSequence)
+float3 HeterogeneousVolumesGetTransmittance(float3 PathThroughput, float3 Origin, float3 Direction, float TMin, float TMax, inout RandomSequence RandSequence)
 {
-	float3 Transmittance = 1.0;
 	if (EnableHeterogeneousVolumes)
 	{
 		bool bTransmittanceOnly = true;
 		FMajorantData EmptyOverlappingMajorant = CreateMajorantData();
-		Transmittance = HeterogeneousVolumesMajorantBasedTracking(PathThroughput, Origin, Direction, TMin, TMax, bTransmittanceOnly, EmptyOverlappingMajorant, RandSequence).Throughput;
+		return HeterogeneousVolumesMajorantBasedTracking(PathThroughput, Origin, Direction, TMin, TMax, bTransmittanceOnly, EmptyOverlappingMajorant, RandSequence).Throughput;
 	}
 
-	return Transmittance;
+	return PathThroughput;
 }
 
 FVolumeTrackingResult HeterogeneousVolumesSampleDistance(float3 PathThroughput, float3 Origin, float3 Direction, float TMin, float TMax, FMajorantData OverlappingMajorant, inout RandomSequence RandSequence)
diff --git a/Engine/Shaders/Private/PathTracing/Volume/PathTracingVolume.ush b/Engine/Shaders/Private/PathTracing/Volume/PathTracingVolume.ush
--- a/Engine/Shaders/Private/PathTracing/Volume/PathTracingVolume.ush
+++ b/Engine/Shaders/Private/PathTracing/Volume/PathTracingVolume.ush
@@ -75,7 +75,7 @@ FVolumeShadedResult VolumeGetDensity(float3 TranslatedWorldPos, FVolumeIntersect
 
 
 // Return the transmittance along a ray segment
-float3 VolumeGetTransmittanceAnalytical(float3 StartThroughput, float3 Origin, float3 Direction, FVolumeIntersectionInterval Interval, float PathRoughness, inout RandomSequence RandSequence)
+float3 VolumeGetTransmittanceAnalytical(float3 StartThroughput, float3 Origin, float3 Direction, FVolumeIntersectionInterval Interval, inout RandomSequence RandSequence)
 {
 	float3 Throughput = StartThroughput;
 	for (int Index = 0; Index < Interval.Num; Index++)
@@ -84,14 +84,14 @@ float3 VolumeGetTransmittanceAnalytical(float3 StartThroughput, float3 Origin, f
 		{
 			case VOLUMEID_ATMOSPHERE:            Throughput *= AtmosphereGetTransmittance(Origin, Direction, Interval.VolumeTMin, Interval.VolumeTMax); break;
 			case VOLUMEID_FOG:                   Throughput *= FogGetTransmittance(Origin, Direction, Interval.VolumeTMin, Interval.VolumeTMax); break;
-			case VOLUMEID_HETEROGENEOUS_VOLUMES: Throughput *= HeterogeneousVolumesGetTransmittance(StartThroughput, Origin, Direction, Interval.VolumeTMin, Interval.VolumeTMax, PathRoughness, RandSequence); break;
+			case VOLUMEID_HETEROGENEOUS_VOLUMES: Throughput = HeterogeneousVolumesGetTransmittance(Throughput, Origin, Direction, Interval.VolumeTMin, Interval.VolumeTMax, RandSequence); break;
 		}
 	}
 	return Throughput;
 }
 
 // Return the transmittance along a whole intersection list
-float3 VolumeGetTransmittanceAnalytical(float3 StartThroughput, float3 Origin, float3 Direction, FVolumeIntersectionList IntersectionList, float PathRoughness, inout RandomSequence RandSequence)
+float3 VolumeGetTransmittanceAnalytical(float3 StartThroughput, float3 Origin, float3 Direction, FVolumeIntersectionList IntersectionList, inout RandomSequence RandSequence)
 {
 	float3 Throughput = StartThroughput;
 	for (int Index = 0; Index < IntersectionList.Num; Index++)
@@ -102,7 +102,7 @@ float3 VolumeGetTransmittanceAnalytical(float3 StartThroughput, float3 Origin, f
 		{
 			case VOLUMEID_ATMOSPHERE:            Throughput *= AtmosphereGetTransmittance(Origin, Direction, TMin, TMax); break;
 			case VOLUMEID_FOG:                   Throughput *= FogGetTransmittance(Origin, Direction, TMin, TMax); break;
-			case VOLUMEID_HETEROGENEOUS_VOLUMES: Throughput *= HeterogeneousVolumesGetTransmittance(StartThroughput, Origin, Direction, TMin, TMax, PathRoughness, RandSequence); break;
+			case VOLUMEID_HETEROGENEOUS_VOLUMES: Throughput = HeterogeneousVolumesGetTransmittance(Throughput, Origin, Direction, TMin, TMax, RandSequence); break;
 		}
 	}
 	return Throughput;
@@ -150,7 +150,7 @@ float3 VolumeGetTransmittanceAnalytical(float3 StartThroughput, float3 Origin, f
 
 
 // The returned value factors in the initial throughput. This is used to avoid sampling too much if throughput reaches 0.
-float3 VolumeGetTransmittanceNullTracking(float3 StartThroughput, float3 Origin, float3 Direction, FVolumeIntersectionInterval Interval, int Bounce, inout RandomSequence RandSequence)
+float3 VolumeGetTransmittanceNullTracking(float3 StartThroughput, float3 Origin, float3 Direction, FVolumeIntersectionInterval Interval, inout RandomSequence RandSequence)
 {
 	float3 Throughput = StartThroughput;
 	for (int Index = 0; Index < Interval.Num; Index++)
@@ -171,8 +171,7 @@ float3 VolumeGetTransmittanceNullTracking(float3 StartThroughput, float3 Origin,
 			}
 			case VOLUMEID_HETEROGENEOUS_VOLUMES:
 			{
-				float PathRoughness = 0;
-				Throughput *= HeterogeneousVolumesGetTransmittance(StartThroughput, Origin, Direction, TMin, TMax, PathRoughness, RandSequence);
+				Throughput = HeterogeneousVolumesGetTransmittance(Throughput, Origin, Direction, TMin, TMax, RandSequence);
 				break;
 			}
 		}
@@ -181,7 +180,7 @@ float3 VolumeGetTransmittanceNullTracking(float3 StartThroughput, float3 Origin,
 
 }
 
-float3 VolumeGetTransmittanceNullTracking(float3 StartThroughput, float3 Origin, float3 Direction, FVolumeIntersectionList IntersectionList, int Bounce, inout RandomSequence RandSequence)
+float3 VolumeGetTransmittanceNullTracking(float3 StartThroughput, float3 Origin, float3 Direction, FVolumeIntersectionList IntersectionList, inout RandomSequence RandSequence)
 {
 	float3 Throughput = StartThroughput;
 	for (int Index = 0; Index < IntersectionList.Num; Index++)
@@ -202,8 +201,7 @@ float3 VolumeGetTransmittanceNullTracking(float3 StartThroughput, float3 Origin,
 			}
 			case VOLUMEID_HETEROGENEOUS_VOLUMES:
 			{
-				float PathRoughness = 0;
-				Throughput *= HeterogeneousVolumesGetTransmittance(StartThroughput, Origin, Direction, TMin, TMax, PathRoughness, RandSequence);
+				Throughput = HeterogeneousVolumesGetTransmittance(Throughput, Origin, Direction, TMin, TMax, RandSequence);
 				break;
 			}
 		}
@@ -213,24 +211,24 @@ float3 VolumeGetTransmittanceNullTracking(float3 StartThroughput, float3 Origin,
 #undef PATH_TRACER_REFERENCE_TRANSMITTANCE_LOOP
 }
 
-float3 VolumeGetTransmittance(float3 StartThroughput, float3 Origin, float3 Direction, FVolumeIntersectionInterval Interval, int Bounce, inout RandomSequence RandSequence)
+float3 VolumeGetTransmittance(float3 StartThroughput, float3 Origin, float3 Direction, FVolumeIntersectionInterval Interval, inout RandomSequence RandSequence)
 {
 	if (UseAnalyticTransmittance)
 	{
-		return VolumeGetTransmittanceAnalytical(StartThroughput, Origin, Direction, Interval, Bounce, RandSequence);
+		return VolumeGetTransmittanceAnalytical(StartThroughput, Origin, Direction, Interval, RandSequence);
 	}
 
-	return VolumeGetTransmittanceNullTracking(StartThroughput, Origin, Direction, Interval, Bounce, RandSequence);
+	return VolumeGetTransmittanceNullTracking(StartThroughput, Origin, Direction, Interval, RandSequence);
 }
 
-float3 VolumeGetTransmittance(float3 StartThroughput, float3 Origin, float3 Direction, FVolumeIntersectionList IntersectionList, int Bounce, inout RandomSequence RandSequence)
+float3 VolumeGetTransmittance(float3 StartThroughput, float3 Origin, float3 Direction, FVolumeIntersectionList IntersectionList, inout RandomSequence RandSequence)
 {
 	if (UseAnalyticTransmittance)
 	{
-		return VolumeGetTransmittanceAnalytical(StartThroughput, Origin, Direction, IntersectionList, Bounce, RandSequence);
+		return VolumeGetTransmittanceAnalytical(StartThroughput, Origin, Direction, IntersectionList, RandSequence);
 	}
 
-	return VolumeGetTransmittanceNullTracking(StartThroughput, Origin, Direction, IntersectionList, Bounce, RandSequence);
+	return VolumeGetTransmittanceNullTracking(StartThroughput, Origin, Direction, IntersectionList, RandSequence);
 }
 
 FVolumeTrackingResult VolumeSampleDistance(float3 PathThroughput, float3 Origin, float3 Direction, FVolumeIntersectionInterval Interval, bool bIsCameraRay, inout RandomSequence RandSequence)
