Adding ShadowBiasFactor to Heterogeneous Volume interface.

diff --git a/Engine/Plugins/FX/Niagara/Source/Niagara/Private/NiagaraRendererVolumes.cpp b/Engine/Plugins/FX/Niagara/Source/Niagara/Private/NiagaraRendererVolumes.cpp
--- a/Engine/Plugins/FX/Niagara/Source/Niagara/Private/NiagaraRendererVolumes.cpp
+++ b/Engine/Plugins/FX/Niagara/Source/Niagara/Private/NiagaraRendererVolumes.cpp
@@ -50,6 +50,7 @@ namespace NiagaraRendererVolumesLocal
 		FVector3f				VolumeWorldSpaceSize = UNiagaraVolumeRendererProperties::GetDefaultVolumeWorldSpaceSize();
 		float					StepFactor = UNiagaraVolumeRendererProperties::GetDefaultStepFactor();
 		float					ShadowStepFactor = UNiagaraVolumeRendererProperties::GetDefaultShadowStepFactor();
+		float					ShadowBiasFactor = UNiagaraVolumeRendererProperties::GetDefaultShadowBiasFactor();
 		float					LightingDownsampleFactor = UNiagaraVolumeRendererProperties::GetDefaultLightingDownsampleFactor();
 	};
 
@@ -224,6 +225,7 @@ FNiagaraDynamicDataBase* FNiagaraRendererVolumes::GenerateDynamicData(const FNia
 	VolumeDynamicData->VolumeWorldSpaceSize = ParameterStore.GetParameterValueOrDefault(RendererProperties->VolumeWorldSpaceSizeBinding.GetParamMapBindableVariable(), DefaultVolumeWorldSpaceSize);
 	VolumeDynamicData->StepFactor = RendererProperties->StepFactor;
 	VolumeDynamicData->ShadowStepFactor = RendererProperties->ShadowStepFactor;
+	VolumeDynamicData->ShadowBiasFactor = RendererProperties->ShadowBiasFactor;
 	VolumeDynamicData->LightingDownsampleFactor = RendererProperties->LightingDownsampleFactor;
 
 	return VolumeDynamicData;
@@ -357,6 +359,7 @@ void FNiagaraRendererVolumes::GetDynamicMeshElements(const TArray<const FSceneVi
 			HeterogeneousVolumeData->VoxelResolution = VolumeResolution;
 			HeterogeneousVolumeData->StepFactor = VolumeDynamicData->StepFactor;
 			HeterogeneousVolumeData->ShadowStepFactor = VolumeDynamicData->ShadowStepFactor;
+			HeterogeneousVolumeData->ShadowBiasFactor = VolumeDynamicData->ShadowBiasFactor;
 			HeterogeneousVolumeData->LightingDownsampleFactor = VolumeDynamicData->LightingDownsampleFactor;
 			BatchElement.UserData = HeterogeneousVolumeData;
 
diff --git a/Engine/Plugins/FX/Niagara/Source/Niagara/Public/NiagaraVolumeRendererProperties.h b/Engine/Plugins/FX/Niagara/Source/Niagara/Public/NiagaraVolumeRendererProperties.h
--- a/Engine/Plugins/FX/Niagara/Source/Niagara/Public/NiagaraVolumeRendererProperties.h
+++ b/Engine/Plugins/FX/Niagara/Source/Niagara/Public/NiagaraVolumeRendererProperties.h
@@ -79,6 +79,7 @@ public:
 	static FVector3f GetDefaultVolumeWorldSpaceSize() { return FVector3f::ZeroVector; }
 	static float GetDefaultStepFactor() { return 1.0f; }
 	static float GetDefaultShadowStepFactor() { return 8.0f; }
+	static float GetDefaultShadowBiasFactor() { return 0.5f; }
 	static float GetDefaultLightingDownsampleFactor() { return 2.0f; }
 
 	/** What material to use for the volume. */
@@ -112,6 +113,9 @@ public:
 	UPROPERTY(EditAnywhere, Category = "Volume Rendering")
 	float ShadowStepFactor = 8.0f;
 
+	UPROPERTY(EditAnywhere, Category = "Volume Rendering")
+	float ShadowBiasFactor = 0.5f;
+
 	/** Position binding for the center of the volume. */
 	//-TODO: Unhind when binding is supported
 	//UPROPERTY(EditAnywhere, Category = "Bindings")
diff --git a/Engine/Shaders/Private/HeterogeneousVolumes/HeterogeneousVolumesHardwareRayTracing.usf b/Engine/Shaders/Private/HeterogeneousVolumes/HeterogeneousVolumesHardwareRayTracing.usf
--- a/Engine/Shaders/Private/HeterogeneousVolumes/HeterogeneousVolumesHardwareRayTracing.usf
+++ b/Engine/Shaders/Private/HeterogeneousVolumes/HeterogeneousVolumesHardwareRayTracing.usf
@@ -470,7 +470,7 @@ RAY_TRACING_ENTRY_RAYGEN(RenderLightingCacheWithPreshadingRGS)
 	RWLightingCacheTexture[VoxelIndex] = Transmittance;
 
 #elif DIM_LIGHTING_CACHE_MODE == LIGHTING_CACHE_INSCATTERING
-	float3 Inscattering = ComputeInscattering(WorldRayOrigin, LightData, LightType, MaxStepCount, bApplyShadowTransmittance);
+	float3 Inscattering = ComputeInscattering(WorldRayOrigin, LightData, LightType, MaxStepCount, CalcShadowBias(), bApplyShadowTransmittance);
 	RWLightingCacheTexture[VoxelIndex] += Inscattering * View.PreExposure;
 #endif // DIM_LIGHTING_CACHE_MODE
 }
diff --git a/Engine/Shaders/Private/HeterogeneousVolumes/HeterogeneousVolumesLiveShadingPipeline.usf b/Engine/Shaders/Private/HeterogeneousVolumes/HeterogeneousVolumesLiveShadingPipeline.usf
--- a/Engine/Shaders/Private/HeterogeneousVolumes/HeterogeneousVolumesLiveShadingPipeline.usf
+++ b/Engine/Shaders/Private/HeterogeneousVolumes/HeterogeneousVolumesLiveShadingPipeline.usf
@@ -266,7 +266,7 @@ void RenderLightingCacheWithLiveShadingCS(
 	RWLightingCacheTexture[VoxelIndex] = Transmittance;
 
 #elif DIM_LIGHTING_CACHE_MODE == LIGHTING_CACHE_INSCATTERING
-	float3 Inscattering = ComputeInscattering(WorldRayOrigin, LightData, LightType, MaxStepCount, bApplyShadowTransmittance);
+	float3 Inscattering = ComputeInscattering(WorldRayOrigin, LightData, LightType, MaxStepCount, CalcShadowBias(), bApplyShadowTransmittance);
 	RWLightingCacheTexture[VoxelIndex] += Inscattering * View.PreExposure;
 #endif // DIM_LIGHTING_CACHE_MODE
 }
diff --git a/Engine/Shaders/Private/HeterogeneousVolumes/HeterogeneousVolumesPreshadingPipeline.usf b/Engine/Shaders/Private/HeterogeneousVolumes/HeterogeneousVolumesPreshadingPipeline.usf
--- a/Engine/Shaders/Private/HeterogeneousVolumes/HeterogeneousVolumesPreshadingPipeline.usf
+++ b/Engine/Shaders/Private/HeterogeneousVolumes/HeterogeneousVolumesPreshadingPipeline.usf
@@ -491,7 +491,7 @@ void RenderLightingCacheWithPreshadingCS(
 	RWLightingCacheTexture[VoxelIndex] = Transmittance;
 
 #elif DIM_LIGHTING_CACHE_MODE == LIGHTING_CACHE_INSCATTERING
-	float3 Inscattering = ComputeInscattering(WorldRayOrigin, LightData, LightType, MaxStepCount, bApplyShadowTransmittance);
+	float3 Inscattering = ComputeInscattering(WorldRayOrigin, LightData, LightType, MaxStepCount, CalcShadowBias(), bApplyShadowTransmittance);
 	RWLightingCacheTexture[VoxelIndex] += Inscattering * View.PreExposure;
 #endif // DIM_LIGHTING_CACHE_MODE
 }
diff --git a/Engine/Shaders/Private/HeterogeneousVolumes/HeterogeneousVolumesRayMarchingUtils.ush b/Engine/Shaders/Private/HeterogeneousVolumes/HeterogeneousVolumesRayMarchingUtils.ush
--- a/Engine/Shaders/Private/HeterogeneousVolumes/HeterogeneousVolumesRayMarchingUtils.ush
+++ b/Engine/Shaders/Private/HeterogeneousVolumes/HeterogeneousVolumesRayMarchingUtils.ush
@@ -120,6 +120,16 @@ float CalcShadowStepSize(float3 LocalRayDirection)
 	return CalcStepSize(LocalRayDirection, LocalStepSize, LocalStepFactor);
 }
 
+float CalcShadowBias()
+{
+	float3 VolumeBounds = 2.0f * GetLocalBoundsExtent();
+	float3 VoxelSize = VolumeBounds / float3(GetLightingCacheResolution());
+	float3 WorldVoxelSize = mul(float4(VoxelSize, 0), GetLocalToWorld()).xyz;
+	float VoxelDiagonal = length(WorldVoxelSize);
+	float WorldShadowBias = VoxelDiagonal * GetLightingCacheVoxelBias();
+	return WorldShadowBias;
+}
+
 float RecalculateStepSize(inout FRayMarchingContext Context)
 {
 	// Recalculate step-size if the bounding-box diagonal is greater than the maximum step-count
@@ -298,6 +308,7 @@ float3 ComputeInscattering(
 	FDeferredLightData LightData,
 	uint LightType,
 	uint MaxStepCount,
+	float WorldShadowBias,
 	bool bApplyShadowTransmittance
 )
 {
@@ -338,7 +349,8 @@ float3 ComputeInscattering(
 	Inscattering = LightData.Color * LightAttenuation * HardSurfaceShadowFactor;
 	if (any(Inscattering > 0.0) && bApplyShadowTransmittance)
 	{
-		Inscattering *= ComputeTransmittance(WorldPosition, ToLight, MaxStepCount);
+		float3 BiasedWorldPosition = WorldPosition + L * WorldShadowBias;
+		Inscattering *= ComputeTransmittance(BiasedWorldPosition, ToLight, MaxStepCount);
 	}
 #endif // DIM_USE_INSCATTERING_VOLUME
 
@@ -378,6 +390,8 @@ void RayMarchSingleScattering(
 	inout float3 Transmittance
 )
 {
+	float WorldShadowBias = CalcShadowBias();
+
 	for (uint StepIndex = 0; StepIndex < StepCount; ++StepIndex)
 	{
 		float LocalHitT = RayMarchingContext.LocalRayTMin + RayMarchingContext.StepSize * (RayMarchingContext.Jitter + StepIndex);
@@ -415,7 +429,7 @@ void RayMarchSingleScattering(
 			float3 Albedo = SampleAlbedo(SampleContext);
 			if (any(Albedo > 0.0))
 			{
-				float3 Inscattering = ComputeInscattering(WorldPosition, LightData, LightType, RayMarchingContext.MaxStepCount, RayMarchingContext.bApplyShadowTransmittance);
+				float3 Inscattering = ComputeInscattering(WorldPosition, LightData, LightType, RayMarchingContext.MaxStepCount, WorldShadowBias, RayMarchingContext.bApplyShadowTransmittance);
 
 				float3 ScatteringCoefficient = Albedo * Extinction;
 				float IsotropicPhase = 1.0 / (4.0 * PI);
diff --git a/Engine/Shaders/Private/HeterogeneousVolumes/HeterogeneousVolumesTransmittanceVolumeUtils.ush b/Engine/Shaders/Private/HeterogeneousVolumes/HeterogeneousVolumesTransmittanceVolumeUtils.ush
--- a/Engine/Shaders/Private/HeterogeneousVolumes/HeterogeneousVolumesTransmittanceVolumeUtils.ush
+++ b/Engine/Shaders/Private/HeterogeneousVolumes/HeterogeneousVolumesTransmittanceVolumeUtils.ush
@@ -8,6 +8,7 @@
 #define LIGHTING_CACHE_INSCATTERING 1
 
 uint3 LightingCacheResolution;
+float LightingCacheVoxelBias;
 Texture3D LightingCacheTexture;
 
 uint3 GetLightingCacheResolution()
@@ -15,6 +16,11 @@ uint3 GetLightingCacheResolution()
 	return LightingCacheResolution;
 }
 
+float GetLightingCacheVoxelBias()
+{
+	return LightingCacheVoxelBias;
+}
+
 float3 SampleLightingCache(float3 UVW, float MipLevel)
 {
 	float3 Transmittance = Texture3DSampleLevel(
diff --git a/Engine/Source/Runtime/Engine/Classes/Components/HeterogeneousVolumeComponent.h b/Engine/Source/Runtime/Engine/Classes/Components/HeterogeneousVolumeComponent.h
--- a/Engine/Source/Runtime/Engine/Classes/Components/HeterogeneousVolumeComponent.h
+++ b/Engine/Source/Runtime/Engine/Classes/Components/HeterogeneousVolumeComponent.h
@@ -54,6 +54,9 @@ class UHeterogeneousVolumeComponent : public UMeshComponent
 	UPROPERTY(EditAnywhere, Category = Lighting)
 	float ShadowStepFactor;
 
+	UPROPERTY(EditAnywhere, Category = Lighting)
+	float ShadowBiasFactor;
+
 	UPROPERTY(EditAnywhere, Category = Lighting)
 	float LightingDownsampleFactor;
 
diff --git a/Engine/Source/Runtime/Engine/Private/Components/HeterogeneousVolumeComponent.cpp b/Engine/Source/Runtime/Engine/Private/Components/HeterogeneousVolumeComponent.cpp
--- a/Engine/Source/Runtime/Engine/Private/Components/HeterogeneousVolumeComponent.cpp
+++ b/Engine/Source/Runtime/Engine/Private/Components/HeterogeneousVolumeComponent.cpp
@@ -83,6 +83,7 @@ FHeterogeneousVolumeSceneProxy::FHeterogeneousVolumeSceneProxy(UHeterogeneousVol
 
 	HeterogeneousVolumeData.StepFactor = InComponent->StepFactor;
 	HeterogeneousVolumeData.ShadowStepFactor = InComponent->ShadowStepFactor;
+	HeterogeneousVolumeData.ShadowBiasFactor = InComponent->ShadowBiasFactor;
 	HeterogeneousVolumeData.LightingDownsampleFactor = InComponent->LightingDownsampleFactor;
 
 	// Initialize vertex buffer data for a quad
@@ -234,6 +235,7 @@ UHeterogeneousVolumeComponent::UHeterogeneousVolumeComponent(const FObjectInitia
 	EndFrame = 0;
 	StepFactor = 1.0f;
 	ShadowStepFactor = 8.0f;
+	ShadowBiasFactor = 0.0f;
 	LightingDownsampleFactor = 2.0f;
 	MipLevel = 0;
 	bIssueBlockingRequests = false;
diff --git a/Engine/Source/Runtime/Engine/Public/HeterogeneousVolumeInterface.h b/Engine/Source/Runtime/Engine/Public/HeterogeneousVolumeInterface.h
--- a/Engine/Source/Runtime/Engine/Public/HeterogeneousVolumeInterface.h
+++ b/Engine/Source/Runtime/Engine/Public/HeterogeneousVolumeInterface.h
@@ -28,6 +28,7 @@ public:
 	// Lighting
 	virtual float GetStepFactor() const = 0;
 	virtual float GetShadowStepFactor() const = 0;
+	virtual float GetShadowBiasFactor() const = 0;
 	virtual float GetLightingDownsampleFactor() const = 0;
 };
 
@@ -43,6 +44,7 @@ public:
 		, MinimumVoxelSize(0.1)
 		, StepFactor(1.0)
 		, ShadowStepFactor(8.0)
+		, ShadowBiasFactor(0.0)
 		, LightingDownsampleFactor(1.0)
 	{}
 	virtual ~FHeterogeneousVolumeData() {}
@@ -63,6 +65,7 @@ public:
 	// Lighting
 	virtual float GetStepFactor() const { return StepFactor; }
 	virtual float GetShadowStepFactor() const { return ShadowStepFactor; }
+	virtual float GetShadowBiasFactor() const { return ShadowBiasFactor; }
 	virtual float GetLightingDownsampleFactor() const { return LightingDownsampleFactor; }
 
 	const FPrimitiveSceneProxy* PrimitiveSceneProxy;
@@ -71,5 +74,6 @@ public:
 	float MinimumVoxelSize;
 	float StepFactor;
 	float ShadowStepFactor;
+	float ShadowBiasFactor;
 	float LightingDownsampleFactor;
 };
diff --git a/Engine/Source/Runtime/Renderer/Private/HeterogeneousVolumes/HeterogeneousVolumes.h b/Engine/Source/Runtime/Renderer/Private/HeterogeneousVolumes/HeterogeneousVolumes.h
--- a/Engine/Source/Runtime/Renderer/Private/HeterogeneousVolumes/HeterogeneousVolumes.h
+++ b/Engine/Source/Runtime/Renderer/Private/HeterogeneousVolumes/HeterogeneousVolumes.h
@@ -119,6 +119,7 @@ END_UNIFORM_BUFFER_STRUCT()
 
 BEGIN_SHADER_PARAMETER_STRUCT(FLightingCacheParameters, )
 	SHADER_PARAMETER(FIntVector, LightingCacheResolution)
+	SHADER_PARAMETER(float, LightingCacheVoxelBias)
 	SHADER_PARAMETER_RDG_TEXTURE(Texture3D, LightingCacheTexture)
 END_SHADER_PARAMETER_STRUCT()
 
diff --git a/Engine/Source/Runtime/Renderer/Private/HeterogeneousVolumes/HeterogeneousVolumesHardwareRayTracing.cpp b/Engine/Source/Runtime/Renderer/Private/HeterogeneousVolumes/HeterogeneousVolumesHardwareRayTracing.cpp
--- a/Engine/Source/Runtime/Renderer/Private/HeterogeneousVolumes/HeterogeneousVolumesHardwareRayTracing.cpp
+++ b/Engine/Source/Runtime/Renderer/Private/HeterogeneousVolumes/HeterogeneousVolumesHardwareRayTracing.cpp
@@ -613,6 +613,7 @@ void RenderLightingCacheWithPreshadingHardwareRayTracing(
 
 		// Transmittance volume
 		PassParameters->LightingCache.LightingCacheResolution = HeterogeneousVolumes::GetLightingCacheResolution(HeterogeneousVolumeInterface);
+		PassParameters->LightingCache.LightingCacheVoxelBias = HeterogeneousVolumeInterface->GetShadowBiasFactor();
 		//PassParameters->LightingCache.LightingCacheTexture = GraphBuilder.CreateSRV(LightingCacheTexture);
 
 		// Ray data
@@ -794,6 +795,7 @@ void RenderSingleScatteringWithPreshadingHardwareRayTracing(
 		if ((HeterogeneousVolumes::UseLightingCacheForTransmittance() && bApplyShadowTransmittance) || HeterogeneousVolumes::UseLightingCacheForInscattering())
 		{
 			PassParameters->LightingCache.LightingCacheResolution = HeterogeneousVolumes::GetLightingCacheResolution(HeterogeneousVolumeInterface);
+			PassParameters->LightingCache.LightingCacheVoxelBias = HeterogeneousVolumeInterface->GetShadowBiasFactor();
 			PassParameters->LightingCache.LightingCacheTexture = LightingCacheTexture;
 		}
 
diff --git a/Engine/Source/Runtime/Renderer/Private/HeterogeneousVolumes/HeterogeneousVolumesLiveShadingPipeline.cpp b/Engine/Source/Runtime/Renderer/Private/HeterogeneousVolumes/HeterogeneousVolumesLiveShadingPipeline.cpp
--- a/Engine/Source/Runtime/Renderer/Private/HeterogeneousVolumes/HeterogeneousVolumesLiveShadingPipeline.cpp
+++ b/Engine/Source/Runtime/Renderer/Private/HeterogeneousVolumes/HeterogeneousVolumesLiveShadingPipeline.cpp
@@ -439,6 +439,7 @@ void RenderLightingCacheWithLiveShading(
 		// Transmittance volume
 		PassParameters->VoxelResolution = HeterogeneousVolumeInterface->GetVoxelResolution();
 		PassParameters->LightingCache.LightingCacheResolution = HeterogeneousVolumes::GetLightingCacheResolution(HeterogeneousVolumeInterface);
+		PassParameters->LightingCache.LightingCacheVoxelBias = HeterogeneousVolumeInterface->GetShadowBiasFactor();
 		//PassParameters->LightingCache.LightingCacheTexture = GraphBuilder.CreateSRV(LightingCacheTexture);
 		PassParameters->LightingCache.LightingCacheTexture = FRDGSystemTextures::Get(GraphBuilder).VolumetricBlack;
 
@@ -626,6 +627,7 @@ void RenderSingleScatteringWithLiveShading(
 		if ((HeterogeneousVolumes::UseLightingCacheForTransmittance() && bApplyShadowTransmittance) || HeterogeneousVolumes::UseLightingCacheForInscattering())
 		{
 			PassParameters->LightingCache.LightingCacheResolution = HeterogeneousVolumes::GetLightingCacheResolution(HeterogeneousVolumeInterface);
+			PassParameters->LightingCache.LightingCacheVoxelBias = HeterogeneousVolumeInterface->GetShadowBiasFactor();
 			PassParameters->LightingCache.LightingCacheTexture = LightingCacheTexture;
 		}
 		else
diff --git a/Engine/Source/Runtime/Renderer/Private/HeterogeneousVolumes/HeterogeneousVolumesPreshadingPipeline.cpp b/Engine/Source/Runtime/Renderer/Private/HeterogeneousVolumes/HeterogeneousVolumesPreshadingPipeline.cpp
--- a/Engine/Source/Runtime/Renderer/Private/HeterogeneousVolumes/HeterogeneousVolumesPreshadingPipeline.cpp
+++ b/Engine/Source/Runtime/Renderer/Private/HeterogeneousVolumes/HeterogeneousVolumesPreshadingPipeline.cpp
@@ -431,6 +431,7 @@ void RenderLightingCacheWithPreshadingCompute(
 
 		// Transmittance volume
 		PassParameters->LightingCache.LightingCacheResolution = HeterogeneousVolumes::GetLightingCacheResolution(HeterogeneousVolumeInterface);
+		PassParameters->LightingCache.LightingCacheVoxelBias = HeterogeneousVolumeInterface->GetShadowBiasFactor();
 		PassParameters->LightingCache.LightingCacheTexture = LightingCacheTexture;
 
 		// Ray data
@@ -593,6 +594,7 @@ void RenderSingleScatteringWithPreshadingCompute(
 		if ((HeterogeneousVolumes::UseLightingCacheForTransmittance() && bApplyShadowTransmittance) || HeterogeneousVolumes::UseLightingCacheForInscattering())
 		{
 			PassParameters->LightingCache.LightingCacheResolution = HeterogeneousVolumes::GetLightingCacheResolution(HeterogeneousVolumeInterface);
+			PassParameters->LightingCache.LightingCacheVoxelBias = HeterogeneousVolumeInterface->GetShadowBiasFactor();
 			PassParameters->LightingCache.LightingCacheTexture = LightingCacheTexture;
 		}
 
