Substrate - Fixed emissive on clear coat from legacy conversion.

diff --git a/Engine/Shaders/Private/Strata/StrataLegacyConversion.ush b/Engine/Shaders/Private/Strata/StrataLegacyConversion.ush
--- a/Engine/Shaders/Private/Strata/StrataLegacyConversion.ush
+++ b/Engine/Shaders/Private/Strata/StrataLegacyConversion.ush
@@ -336,7 +336,7 @@ FStrataData StrataConvertLegacyMaterialDynamic(
 		LegacySlab.SecondRoughnessWeight = ClearCoat;
 		LegacySlab.SecondRoughness = ClearCoatRoughness;
 		LegacySlab.SecondRoughnessAsSimpleClearCoat = ClearCoat > 0.0f; // Haziness as simple clear coat layer.
-		LegacySlab.Emissive = Zeros;
+		LegacySlab.Emissive = Emissive;
 		LegacySlab.Weight = 1.f;
 
 		FinalWeight = 1.0f;
