Fix incorrect BilinearWeights in DistanceFieldShadowingUpsamplePS
- issue is not very noticeable when running DF shadows at half res, however at lower resolutions the shadows were shiting bottom-right.

- also changed DistanceFieldShadowingUpsamplePS to use OneOverDownsampleFactor parameter instead of DOWNSAMPLE_FACTOR define.

diff --git a/Engine/Shaders/Private/DistanceFieldShadowing.usf b/Engine/Shaders/Private/DistanceFieldShadowing.usf
--- a/Engine/Shaders/Private/DistanceFieldShadowing.usf
+++ b/Engine/Shaders/Private/DistanceFieldShadowing.usf
@@ -773,6 +773,8 @@ float InvFadePlaneLength;
 float NearFadePlaneOffset;
 float InvNearFadePlaneLength;
 
+float OneOverDownsampleFactor;
+
 void DistanceFieldShadowingUpsamplePS(
 	in float4 SVPos : SV_POSITION,
 	out float4 OutColor : SV_Target0)
@@ -786,15 +788,25 @@ void DistanceFieldShadowingUpsamplePS(
 
 #define BILATERAL_UPSAMPLE 1
 #if BILATERAL_UPSAMPLE && UPSAMPLE_REQUIRED
-	float2 LowResBufferSize = floor(View.BufferSizeAndInvSize.xy / DOWNSAMPLE_FACTOR);
+	float2 LowResBufferSize = floor(View.BufferSizeAndInvSize.xy * OneOverDownsampleFactor);
 	float2 LowResTexelSize = 1.0f / LowResBufferSize;
-	float2 Corner00UV = floor(DistanceFieldUVs * LowResBufferSize - .5f) / LowResBufferSize + .5f * LowResTexelSize;
-	float2 BilinearWeights = (DistanceFieldUVs - Corner00UV) * LowResBufferSize;
 
-	float2 TextureValues00 = Texture2DSampleLevel(ShadowFactorsTexture, ShadowFactorsSampler, Corner00UV, 0).xy;
-	float2 TextureValues10 = Texture2DSampleLevel(ShadowFactorsTexture, ShadowFactorsSampler, Corner00UV + float2(LowResTexelSize.x, 0), 0).xy;
-	float2 TextureValues01 = Texture2DSampleLevel(ShadowFactorsTexture, ShadowFactorsSampler, Corner00UV + float2(0, LowResTexelSize.y), 0).xy;
-	float2 TextureValues11 = Texture2DSampleLevel(ShadowFactorsTexture, ShadowFactorsSampler, Corner00UV + LowResTexelSize, 0).xy;
+	// since DF shadow rays are traced from the position of the full-res texel at the top left of low-res tile
+	// we need to use different UVs for calculating bilinear weights and sampling the 4 texels
+
+	float2 Corner00UVBase = floor(DistanceFieldUVs * LowResBufferSize) / LowResBufferSize;
+
+	// Calculate bilinear weights using the same UVs that rays were traced from
+	float2 Corner00UV = Corner00UVBase + 0.5f * View.BufferSizeAndInvSize.zw;
+	float2 BilinearWeights = saturate((DistanceFieldUVs - Corner00UV) * LowResBufferSize);
+
+	// Sample low res texture using UVs at center of low-res texels since we are interpolating manually
+	float2 Corner00SampleUV = Corner00UVBase + 0.5f * LowResTexelSize;
+
+	float2 TextureValues00 = Texture2DSampleLevel(ShadowFactorsTexture, ShadowFactorsSampler, Corner00SampleUV, 0).xy;
+	float2 TextureValues10 = Texture2DSampleLevel(ShadowFactorsTexture, ShadowFactorsSampler, Corner00SampleUV + float2(LowResTexelSize.x, 0), 0).xy;
+	float2 TextureValues01 = Texture2DSampleLevel(ShadowFactorsTexture, ShadowFactorsSampler, Corner00SampleUV + float2(0, LowResTexelSize.y), 0).xy;
+	float2 TextureValues11 = Texture2DSampleLevel(ShadowFactorsTexture, ShadowFactorsSampler, Corner00SampleUV + LowResTexelSize, 0).xy;
 
 	float4 CornerWeights = float4(
 		(1 - BilinearWeights.y) * (1 - BilinearWeights.x), 
diff --git a/Engine/Source/Runtime/Renderer/Private/DistanceFieldShadowing.cpp b/Engine/Source/Runtime/Renderer/Private/DistanceFieldShadowing.cpp
--- a/Engine/Source/Runtime/Renderer/Private/DistanceFieldShadowing.cpp
+++ b/Engine/Source/Runtime/Renderer/Private/DistanceFieldShadowing.cpp
@@ -347,6 +347,7 @@ class FDistanceFieldShadowingUpsamplePS : public FGlobalShader
 		SHADER_PARAMETER(float, InvFadePlaneLength)
 		SHADER_PARAMETER(float, NearFadePlaneOffset)
 		SHADER_PARAMETER(float, InvNearFadePlaneLength)
+		SHADER_PARAMETER(float, OneOverDownsampleFactor)
 	END_SHADER_PARAMETER_STRUCT()
 
 	class FUpsample : SHADER_PERMUTATION_BOOL("UPSAMPLE_REQUIRED");
@@ -359,7 +360,6 @@ class FDistanceFieldShadowingUpsamplePS : public FGlobalShader
 
 	static void ModifyCompilationEnvironment(const FGlobalShaderPermutationParameters& Parameters, FShaderCompilerEnvironment& OutEnvironment)
 	{
-		OutEnvironment.SetDefine(TEXT("DOWNSAMPLE_FACTOR"), GAODownsampleFactor);
 		OutEnvironment.SetDefine(TEXT("FORCE_DEPTH_TEXTURE_READS"), 1);
 	}
 };
@@ -1044,6 +1044,7 @@ void FProjectedShadowInfo::RenderRayTracedDistanceFieldProjection(
 		PassParameters->PS.ShadowFactorsTexture = RayTracedShadowsTexture;
 		PassParameters->PS.ShadowFactorsSampler = TStaticSamplerState<SF_Bilinear>::GetRHI();
 		PassParameters->PS.ScissorRectMinAndSize = FIntRect(ScissorRect.Min, ScissorRect.Size());
+		PassParameters->PS.OneOverDownsampleFactor = 1.0f / GetDFShadowDownsampleFactor();
 
 		if (bDirectionalLight && CascadeSettings.FadePlaneLength > 0)
 		{
@@ -1075,7 +1076,7 @@ void FProjectedShadowInfo::RenderRayTracedDistanceFieldProjection(
 		}
 
 		FDistanceFieldShadowingUpsamplePS::FPermutationDomain PermutationVector;
-		PermutationVector.Set< FDistanceFieldShadowingUpsamplePS::FUpsample >(GFullResolutionDFShadowing == 0);
+		PermutationVector.Set< FDistanceFieldShadowingUpsamplePS::FUpsample >(GetDFShadowDownsampleFactor() != 1);
 		auto PixelShader = View.ShaderMap->GetShader< FDistanceFieldShadowingUpsamplePS >(PermutationVector);
 
 		const bool bReverseCulling = View.bReverseCulling;
