Fixes TSR velocity hole-filling on depth velocities

This is of particular importance to not end up velocity hole filling on walking skeletal meshes as well as moving cars

diff --git a/Engine/Shaders/Private/TemporalSuperResolution/TSRDecimateHistory.usf b/Engine/Shaders/Private/TemporalSuperResolution/TSRDecimateHistory.usf
--- a/Engine/Shaders/Private/TemporalSuperResolution/TSRDecimateHistory.usf
+++ b/Engine/Shaders/Private/TemporalSuperResolution/TSRDecimateHistory.usf
@@ -38,7 +38,11 @@
 float4x4 RotationalClipToPrevClip;
 
 Texture2D<float2> DilatedVelocityTexture;
-Texture2D<float> ClosestDepthTexture;
+#if CONFIG_RESURRECTION_REPROJECTION
+	Texture2D<float2> ClosestDepthTexture;
+#else
+	Texture2D<float> ClosestDepthTexture;
+#endif
 Texture2D<uint> PrevUseCountTexture;
 Texture2D<uint> PrevClosestDepthTexture;
 
@@ -265,7 +269,7 @@ void LoadPrevAtomicTexturesSamples(
 void ProcessPrevAtomicTexturesSamples(
 	uint HistoryUseCountSamples[4],
 	uint HistoryClosestDeviceZSamples[4],
-	float DeviceZ,
+	float PrevDeviceZ,
 	float2 PrevScreenPos,
 	out tsr_half ParallaxRejectionMask,
 	out bool bIsValidHoleFillingPixelVelocity,
@@ -275,7 +279,7 @@ void ProcessPrevAtomicTexturesSamples(
 	bIsValidHoleFillingPixelVelocity = false;
 	EncodedHoleFillingVelocity = 0;
 
-	float WorldDepth = ConvertFromDeviceZ(DeviceZ);
+	float WorldDepth = ConvertFromDeviceZ(PrevDeviceZ);
 	
 	float2 PrevInputBufferUV = (InputInfo_ScreenPosToViewportScale * PrevScreenPos + InputInfo_ScreenPosToViewportBias) * InputInfo_ExtentInverse;
 	PrevInputBufferUV = INVARIANT(PrevInputBufferUV);
@@ -385,7 +389,10 @@ void MainCS(
 
 	// Sample current frame informations.
 	float2 EncodedScreenVelocity;
-	float DeviceZ;
+	#if CONFIG_RESURRECTION_REPROJECTION
+		float DeviceZ;
+	#endif
+	float PrevDeviceZ;
 	tsr_half VelocityEdge = tsr_half(1.0);
 	ISOLATE
 	{
@@ -400,7 +407,13 @@ void MainCS(
 		// Issue texture fetechs
 		ISOLATE
 		{
-			DeviceZ = ClosestDepthTexture[InputPixelPos];
+			#if CONFIG_RESURRECTION_REPROJECTION
+				float2 DeviceZAndPrevDeviceZ = ClosestDepthTexture[InputPixelPos];
+				DeviceZ = DeviceZAndPrevDeviceZ.g;
+				PrevDeviceZ = DeviceZAndPrevDeviceZ.r;
+			#else
+				PrevDeviceZ = ClosestDepthTexture[InputPixelPos];
+			#endif
 
 			#if CONFIG_VELOCITY_SAMPLES == 1
 				EncodedScreenVelocity = DilatedVelocityTexture[InputPixelPos];
@@ -419,7 +432,7 @@ void MainCS(
 			#else
 				SharedArray0[GroupThreadIndex] = EncodedVelocityNeighborhood[0];
 			#endif
-			SharedArray1[GroupThreadIndex] = DeviceZ;
+			SharedArray1[GroupThreadIndex] = PrevDeviceZ;
 		}
 		#endif
 		
@@ -540,7 +553,7 @@ void MainCS(
 		#if CONFIG_MANUAL_LDS_SPILL
 		{
 			EncodedScreenVelocity = SharedArray0[GroupThreadIndex];
-			DeviceZ = SharedArray1[GroupThreadIndex];
+			PrevDeviceZ = SharedArray1[GroupThreadIndex];
 		}
 		#endif
 		
@@ -558,7 +571,7 @@ void MainCS(
 		ProcessPrevAtomicTexturesSamples(
 			HistoryUseCountSamples,
 			HistoryClosestDeviceZSamples,
-			DeviceZ,
+			PrevDeviceZ,
 			PrevScreenPos,
 			/* out */ ParallaxRejectionMask,
 			/* out */ bIsValidHoleFillingPixelVelocity,
diff --git a/Engine/Shaders/Private/TemporalSuperResolution/TSRDilateVelocity.usf b/Engine/Shaders/Private/TemporalSuperResolution/TSRDilateVelocity.usf
--- a/Engine/Shaders/Private/TemporalSuperResolution/TSRDilateVelocity.usf
+++ b/Engine/Shaders/Private/TemporalSuperResolution/TSRDilateVelocity.usf
@@ -56,7 +56,7 @@ Texture2D SceneVelocityTexture;
 Texture2D<uint> PrevScatteredSubpixelDepthTexture;
 
 RWTexture2D<float2> DilatedVelocityOutput;
-RWTexture2D<float> ClosestDepthOutput;
+RWTexture2D<float2> ClosestDepthOutput;
 globallycoherent RWTexture2DArray<uint> PrevAtomicOutput;
 RWTexture2DArray<uint> R8Output;
 RWTexture2D<uint> SubpixelDepthOutput;
@@ -368,6 +368,8 @@ void MainCS(
 		float2 VelocityPolar;
 		ReduceVelocityFlattenTile(GroupThreadIndex, FinalScreenVelocity.xy, /* out */ VelocityPolarRange, /* out */ VelocityPolar);
 	#endif
+	
+	float PrevClosestDeviceZ = ClosestDeviceZ - FinalScreenVelocity.z;
 
 	{
 		bool bValidOutputPixel = all(InputPixelPos < InputInfo_ViewportMax);
@@ -381,7 +383,7 @@ void MainCS(
 			float2 PixelVelocity = FinalScreenVelocity.xy * 0.5 * InputInfo_ViewportSize;
 			uint EncodedHoleFillingVelocity = EncodeHoleFillingVelocity(PixelVelocity);
 
-			uint PrevClosestDepth = (f32tof16(ClosestDeviceZ) << VELOCITY_HOLE_FILLING_BITS) | EncodedHoleFillingVelocity;
+			uint PrevClosestDepth = (f32tof16(PrevClosestDeviceZ) << VELOCITY_HOLE_FILLING_BITS) | EncodedHoleFillingVelocity;
 			
 			bool bValidHistoryCoord = bValidOutputPixel && all(and(PrevInputBufferUV > PrevOutputBufferUVMin, PrevInputBufferUV < PrevOutputBufferUVMax));
 
@@ -419,7 +421,7 @@ void MainCS(
 
 		//DilatedVelocityOutput[OutputPixelPos] = EncodedVelocity.xy; //EncodeVelocityToTexture(float3(FinalScreenVelocity.xy, 0.0)).xy;
 		DilatedVelocityOutput[OutputPixelPos] = EncodeVelocityToTexture(float3(FinalScreenVelocity.xy, 0.0)).xy;
-		ClosestDepthOutput[OutputPixelPos] = ClosestDeviceZ;
+		ClosestDepthOutput[OutputPixelPos] = float2(PrevClosestDeviceZ, ClosestDeviceZ);
 		#if CONFIG_SUBPIXEL_DEPTH
 		{
 			SubpixelDepthOutput[OutputPixelPos] = CompressedSubpixelDepth;
diff --git a/Engine/Shaders/Private/TemporalSuperResolution/TSRMergeRejection.usf b/Engine/Shaders/Private/TemporalSuperResolution/TSRMergeRejection.usf
--- a/Engine/Shaders/Private/TemporalSuperResolution/TSRMergeRejection.usf
+++ b/Engine/Shaders/Private/TemporalSuperResolution/TSRMergeRejection.usf
@@ -13,7 +13,7 @@
 
 float4x4 ClipToResurrectionClip;
 
-Texture2D<float> ClosestDepthTexture;
+Texture2D<float2> ClosestDepthTexture;
 Texture2D<tsr_half4> HistoryRejectionTexture;
 
 RWTexture2D<float2> DilatedVelocityOutput;
@@ -44,7 +44,7 @@ void MainCS(
 	{
 		float2 ResurrectedEncodedScreenVelocity;
 		{
-			float DeviceZ = ClosestDepthTexture[InputPixelPos];
+			float DeviceZ = ClosestDepthTexture[InputPixelPos].g;
 
 			float2 ScreenPos = ApplyScreenTransform(float2(InputPixelPos), InputPixelPosToScreenPos);
 
diff --git a/Engine/Shaders/Private/TemporalSuperResolution/TSRVisualize.usf b/Engine/Shaders/Private/TemporalSuperResolution/TSRVisualize.usf
--- a/Engine/Shaders/Private/TemporalSuperResolution/TSRVisualize.usf
+++ b/Engine/Shaders/Private/TemporalSuperResolution/TSRVisualize.usf
@@ -33,7 +33,7 @@ float MaxHistorySampleCount;
 float OutputToHistoryResolutionFractionSquare;
 
 Texture2D<tsr_half3> SceneColorTexture;
-Texture2D<float>     ClosestDepthTexture;
+Texture2D<float2>    ClosestDepthTexture;
 Texture2D<float2>    DilatedVelocityTexture;
 Texture2D<tsr_half4> HistoryRejectionTexture;
 Texture2D<tsr_ushort> AntiAliasMaskTexture;
@@ -125,7 +125,7 @@ tsr_half3 DisplayResurrectionMask(FContext Ctx)
 tsr_half3 DisplayResurrectedColor(FContext Ctx)
 {
 	tsr_half LowFrequencyRejection = HistoryRejectionTexture.SampleLevel(GlobalPointClampedSampler, Ctx.InputUV, 0).r;
-	float ClosestDeviceZ = ClosestDepthTexture.SampleLevel(GlobalPointClampedSampler, Ctx.InputUV, 0);
+	float ClosestDeviceZ = ClosestDepthTexture.SampleLevel(GlobalPointClampedSampler, Ctx.InputUV, 0).g;
 	bool ResurrectionMask = (tsr_ushort(HistoryRejectionTexture.SampleLevel(GlobalPointClampedSampler, Ctx.InputUV, 0).a * tsr_half(255.0)) & tsr_ushort(0x2)) != tsr_ushort(0x0);
 
 	float4 ThisClip = float4(Ctx.ScreenPos, ClosestDeviceZ, 1);
diff --git a/Engine/Source/Runtime/Renderer/Private/PostProcess/TemporalSuperResolution.cpp b/Engine/Source/Runtime/Renderer/Private/PostProcess/TemporalSuperResolution.cpp
--- a/Engine/Source/Runtime/Renderer/Private/PostProcess/TemporalSuperResolution.cpp
+++ b/Engine/Source/Runtime/Renderer/Private/PostProcess/TemporalSuperResolution.cpp
@@ -1695,10 +1695,18 @@ FDefaultTemporalUpscaler::FOutputs AddTemporalSuperResolutionPasses(
 				/* InFlags = */ TexCreate_ShaderResource | TexCreate_UAV);
 
 			DilatedVelocityTexture = GraphBuilder.CreateTexture(Desc, TEXT("TSR.Velocity.Dilated"));
+		}
+
+		{
+			FRDGTextureDesc Desc = FRDGTextureDesc::Create2D(
+				InputExtent,
+				bCanResurrectHistory ? PF_G16R16F : PF_R16F,
+				FClearValueBinding::None,
+				/* InFlags = */ TexCreate_ShaderResource | TexCreate_UAV);
 
-			Desc.Format = PF_R16F;
 			ClosestDepthTexture = GraphBuilder.CreateTexture(Desc, TEXT("TSR.ClosestDepthTexture"));
 		}
+
 		{
 			FRDGTextureDesc Desc = FRDGTextureDesc::Create2DArray(
 				InputExtent,
