Fix a bug that the IndirectIrradiance applied AO twice on mobile deferred, it has applied AO in the base pass.

diff --git a/Engine/Shaders/Private/DeferredShadingCommon.ush b/Engine/Shaders/Private/DeferredShadingCommon.ush
--- a/Engine/Shaders/Private/DeferredShadingCommon.ush
+++ b/Engine/Shaders/Private/DeferredShadingCommon.ush
@@ -533,7 +533,7 @@ void MobileEncodeGBuffer(
 	{
 		OutGBufferA.rg = UnitVectorToOctahedron( normalize(GBuffer.WorldNormal) ) * 0.5f + 0.5f;
 #if ALLOW_STATIC_LIGHTING
-		OutGBufferA.b = EncodeIndirectIrradiance(GBuffer.IndirectIrradiance * GBuffer.GBufferAO);
+		OutGBufferA.b = EncodeIndirectIrradiance(GBuffer.IndirectIrradiance);
 #else
 		OutGBufferA.b = 1;
 #endif
