Fix incorrect logic in GlobalDistanceField BuildGridTilesCS

diff --git a/Engine/Shaders/Private/DistanceField/GlobalDistanceField.usf b/Engine/Shaders/Private/DistanceField/GlobalDistanceField.usf
--- a/Engine/Shaders/Private/DistanceField/GlobalDistanceField.usf
+++ b/Engine/Shaders/Private/DistanceField/GlobalDistanceField.usf
@@ -199,7 +199,7 @@ void BuildGridTilesCS(
 	{
 		float3 CullGridTranslatedWorldCenter = CullGridCoord * CullGridCoordToTranslatedWorldCenterScale + CullGridCoordToTranslatedWorldCenterBias;
 
-		for (uint UpdateBoundsIndex = ThreadIndex; UpdateBoundsIndex < NumUpdateBounds; UpdateBoundsIndex += THREADGROUP_SIZE)
+		for (uint UpdateBoundsIndex = ThreadIndex; UpdateBoundsIndex < NumUpdateBounds; UpdateBoundsIndex += THREADGROUP_SIZE * THREADGROUP_SIZE * THREADGROUP_SIZE)
 		{
 			float3 UpdateBoundsCenter = UpdateBoundsBuffer[UpdateBoundsIndex * 2 + 0].xyz;
 			float3 UpdateBoundsExtent = UpdateBoundsBuffer[UpdateBoundsIndex * 2 + 1].xyz;
