Fix error when compiling functions in GlobalDistanceFieldDebug.usf that don't need THREADGROUP_SIZE to be defined.

diff --git a/Engine/Shaders/Private/DistanceField/GlobalDistanceFieldDebug.usf b/Engine/Shaders/Private/DistanceField/GlobalDistanceFieldDebug.usf
--- a/Engine/Shaders/Private/DistanceField/GlobalDistanceFieldDebug.usf
+++ b/Engine/Shaders/Private/DistanceField/GlobalDistanceFieldDebug.usf
@@ -5,6 +5,10 @@
 #include "GlobalDistanceFieldShared.ush"
 #include "GlobalDistanceFieldUpdate.ush"
 
+#ifndef THREADGROUP_SIZE
+#define THREADGROUP_SIZE 0
+#endif
+
 StructuredBuffer<uint> GlobalDistanceFieldPageFreeListAllocatorBuffer;
 uint GlobalDistanceFieldMaxPageNum;
 
