Fix hair tangent on group's first and last segments.

diff --git a/Engine/Shaders/Private/HairStrands/HairStrandsTangent.usf b/Engine/Shaders/Private/HairStrands/HairStrandsTangent.usf
--- a/Engine/Shaders/Private/HairStrands/HairStrandsTangent.usf
+++ b/Engine/Shaders/Private/HairStrands/HairStrandsTangent.usf
@@ -31,6 +31,10 @@ void MainCS(uint DispatchThreadId : SV_DispatchThreadID)
 	int IndexCurr = DispatchThreadId;
 	int IndexPrev = clamp(IndexCurr-1, 0, MaxPointCount-1);
 	int IndexNext = clamp(IndexCurr+1, 0, MaxPointCount-1);
+
+	bool bPrevValid = (IndexCurr-1) >= 0;
+	bool bNextValid = (IndexCurr+1) < MaxPointCount;
+
 	bool bValid = IndexCurr < MaxPointCount;
 #if PERMUTATION_CULLING
 	if (HairStrandsVF_bCullingEnable)
@@ -48,6 +52,9 @@ void MainCS(uint DispatchThreadId : SV_DispatchThreadID)
 			IndexCurr = HairStrandsVF_CullingIndexBuffer[FetchIndexCurr];
 			IndexNext = HairStrandsVF_CullingIndexBuffer[FetchIndexNext];
 
+			bPrevValid = IndexPrev >= 0;
+			bNextValid = IndexNext < PointCountAfterCulling;
+
 			bValid = true;
 		}
 	}
@@ -60,14 +67,14 @@ void MainCS(uint DispatchThreadId : SV_DispatchThreadID)
 		const FHairControlPoint CPNext = UnpackHairControlPoint(PositionBuffer[IndexNext]);
 
 		float3 T0 = 0;
-		const bool bIsPrevValid = CPPrev.Type == HAIR_CONTROLPOINT_INSIDE || CPPrev.Type == HAIR_CONTROLPOINT_START;
+		const bool bIsPrevValid = bPrevValid && (CPPrev.Type == HAIR_CONTROLPOINT_INSIDE || CPPrev.Type == HAIR_CONTROLPOINT_START);
 		if (bIsPrevValid)
 		{
 			T0 = normalize(CPCurr.Position - CPPrev.Position);
 		}
 
 		float3 T1 = 0;
-		const bool bIsNextValid = CPNext.Type == HAIR_CONTROLPOINT_INSIDE || CPNext.Type == HAIR_CONTROLPOINT_END;
+		const bool bIsNextValid = bNextValid && (CPNext.Type == HAIR_CONTROLPOINT_INSIDE || CPNext.Type == HAIR_CONTROLPOINT_END);
 		if (bIsNextValid)
 		{
 			T1 = normalize(CPNext.Position - CPCurr.Position);
