Substrate - Fixed mfp and subsurface effect dispearing when applying horizontal parameter blending.

diff --git a/Engine/Shaders/Private/Strata/Strata.ush b/Engine/Shaders/Private/Strata/Strata.ush
--- a/Engine/Shaders/Private/Strata/Strata.ush
+++ b/Engine/Shaders/Private/Strata/Strata.ush
@@ -1355,7 +1355,7 @@ FStrataData GetStrataSlabBSDF(
 	bool bIsThin		= bIsAtTheBottomOfTopology && (bMFPPlugged && bIsThinSurface); // Enable thin only if MFP is plugged, as otherwise there will be no back-tramission
 	bool bHasSSSWrap	= bIsAtTheBottomOfTopology && (bMFPPlugged && !UseSSSDiffusion);
 	bool bHasSSSProfile	= bIsAtTheBottomOfTopology && (bSSSProfilePlugged && SSSMFPScale > 0.f);
-	bool bHasSSS		= bIsAtTheBottomOfTopology && (bMFPPlugged || bHasSSSProfile || bHasSSSWrap);
+	bool bHasSSS		= bMFPPlugged || bIsAtTheBottomOfTopology && (bHasSSSProfile || bHasSSSWrap); // When MFP is plugged, that can be diffusion (bottom) or Simple volume (any position)
 
 	FHaziness Haziness = InitialiseHaziness();
 
@@ -2312,7 +2312,7 @@ FStrataData StrataHorizontalMixingParameterBlending(FStrataData Background, FStr
 #endif // STRATA_GLINTS_ENABLED
 #if STRATA_SPECPROFILE_ENABLED
 	BSDF_SETHASSPECPROFILE(ResultBSDF,	BSDF_GETHASSPECPROFILE(ForegroundBSDF)			| BSDF_GETHASSPECPROFILE(BackgroundBSDF));
-	SLAB_SPECPROFILEID(ResultBSDF) = BSDF_GETHASSPECPROFILE(ForegroundBSDF) ? SLAB_SPECPROFILEID(ForegroundBSDF) : SLAB_SPECPROFILEID(BackgroundBSDF);
+	SLAB_SPECPROFILEID(ResultBSDF) =	BSDF_GETHASSPECPROFILE(ForegroundBSDF) ? SLAB_SPECPROFILEID(ForegroundBSDF) : SLAB_SPECPROFILEID(BackgroundBSDF);
 #endif // STRATA_SPECPROFILE_ENABLED
 #endif // STRATA_COMPLEXSPECIALPATH
 	
@@ -2341,7 +2341,8 @@ FStrataData StrataHorizontalMixingParameterBlending(FStrataData Background, FStr
 	}
 	else
 	{
-		SLAB_SSSMFP(ResultBSDF) =		lerp(SLAB_SSSMFP(BackgroundBSDF),				SLAB_SSSMFP(ForegroundBSDF),				ForegroundOtherMixFactor);
+		SLAB_SSSMFP(ResultBSDF) =				lerp(SLAB_SSSMFP(BackgroundBSDF),				SLAB_SSSMFP(ForegroundBSDF),				ForegroundOtherMixFactor);
+		SLAB_SSSPHASEANISOTROPY(ResultBSDF) =	lerp(SLAB_SSSPHASEANISOTROPY(BackgroundBSDF),	SLAB_SSSPHASEANISOTROPY(ForegroundBSDF),	ForegroundOtherMixFactor);
 
 		// We still try to disable SSS if possible: for performance and also to preserve the surface color (especially specular reflection color).
 		if (all(SLAB_SSSMFP(ResultBSDF) == 0.0f))
