Substrate - added phase function support  for the simple volume case.
This is a big approximation but it is made clear in the comment what it is. We wil limprove that model later.

diff --git a/Engine/Shaders/Private/Strata/StrataEvaluation.ush b/Engine/Shaders/Private/Strata/StrataEvaluation.ush
--- a/Engine/Shaders/Private/Strata/StrataEvaluation.ush
+++ b/Engine/Shaders/Private/Strata/StrataEvaluation.ush
@@ -895,10 +895,20 @@ FStrataEvaluateResult StrataEvaluateBSDFCommon(FStrataBSDFContext BSDFContext, F
 				FParticipatingMedia PM = StrataSlabCreateParticipatingMedia(DiffuseColor, SLAB_SSSMFP(BSDFContext.BSDF));
 				const float DiffuseToVolumeBlend = StrataSlabDiffuseToVolumeBlend(PM);
 
+				// !!! WARNING !!!
+				// SlabDirectionalAlbedo represents single+multiple scattering and has been evaluated for an isotropic phase function.
+				// However want to have some direcitonality now simulated for our medium.
+				// To this aim, we compensate for the uniform phase function and apply a HG phase function.
+				// This is wrong because ontly single scattering should be directly affected by a phase function this way. 
+				// We still have decided to go with such a solution to get artist authroble directional lighting (while we are looking for something better)
+				const float CosTheta		= BSDFContext.Context.VoL;
+				const float HGPhaseFunction	= HenyeyGreensteinPhase(SLAB_SSSPHASEANISOTROPY(BSDFContext.BSDF), CosTheta);
+				const float PhaseFunction	= HGPhaseFunction / IsotropicPhase();
+
 			#if USE_NEW_SIMPLEVOLUME
-				const float3 SlabDirectionalAlbedo = IsotropicMediumSlabPunctualDirectionalAlbedoLUT(PM, BSDFContext.Context.NoV, BSDFContext.Context.NoL);
+				const float3 SlabDirectionalAlbedo = PhaseFunction * IsotropicMediumSlabPunctualDirectionalAlbedoLUT(PM, BSDFContext.Context.NoV, BSDFContext.Context.NoL);
 			#else
-				const float3 SlabDirectionalAlbedo = IsotropicMediumSlabPunctualDirectionalAlbedo(PM);
+				const float3 SlabDirectionalAlbedo = PhaseFunction * IsotropicMediumSlabPunctualDirectionalAlbedo(PM);
 			#endif
 
 				// Shading point <-> View throughput
@@ -910,14 +920,12 @@ FStrataEvaluateResult StrataEvaluateBSDFCommon(FStrataBSDFContext BSDFContext, F
 				const float3 SlabTransmittanceN = IsotropicMediumSlabTransmittance(PM, STRATA_SIMPLEVOLUME_THICKNESS_M, 1.f /*NoL with L==N*/);
 				// The specular transmission is ignored towards the direction of the light. (TransmittanceAlongN can only store transmittance)
 
-				const float Phase = IsotropicPhase();
-
 				// Now lerp between the optically thick and optically thin medium models.
 				// The diffuse and throughput account for the GGX interface SpecularTransmission.
-				Sample.DiffusePathValue		= lerp(Sample.DiffusePathValue,		SlabDirectionalAlbedo * SpecularTransmissionV,		DiffuseToVolumeBlend);
-				Sample.DiffusePDF			= lerp(Sample.DiffusePDF,			Phase,												DiffuseToVolumeBlend);
-				Sample.ThroughputV			= lerp(Sample.ThroughputV,			SlabTransmittanceV    * SpecularTransmissionV,		DiffuseToVolumeBlend);
-				Sample.TransmittanceAlongN	= lerp(Sample.TransmittanceAlongN,	SlabTransmittanceN,									DiffuseToVolumeBlend);
+				Sample.DiffusePathValue		= lerp(Sample.DiffusePathValue,		SlabDirectionalAlbedo * SpecularTransmissionV,	DiffuseToVolumeBlend);
+				Sample.DiffusePDF			= lerp(Sample.DiffusePDF,			HGPhaseFunction,								DiffuseToVolumeBlend);
+				Sample.ThroughputV			= lerp(Sample.ThroughputV,			SlabTransmittanceV    * SpecularTransmissionV,	DiffuseToVolumeBlend);
+				Sample.TransmittanceAlongN	= lerp(Sample.TransmittanceAlongN,	SlabTransmittanceN,								DiffuseToVolumeBlend);
 
 				Sample.bSubsurface = false;	// It should already be the case because Enforce in this case because BSDF_GETHASSSS should be false with SSS_TYPE_SIMPLEVOLUME is true
 
