Substrate - we now finally account refraction at the Slab interface.
This is done in a very simple way for the sake of performance: details of the approximation and future work in code comments.

diff --git a/Engine/Shaders/Private/Strata/StrataEvaluation.ush b/Engine/Shaders/Private/Strata/StrataEvaluation.ush
--- a/Engine/Shaders/Private/Strata/StrataEvaluation.ush
+++ b/Engine/Shaders/Private/Strata/StrataEvaluation.ush
@@ -911,8 +911,25 @@ FStrataEvaluateResult StrataEvaluateBSDFCommon(FStrataBSDFContext BSDFContext, F
 				const float3 SlabDirectionalAlbedo = PhaseFunction * IsotropicMediumSlabPunctualDirectionalAlbedo(PM);
 			#endif
 
+			#if 1
+				// For transmittance, we must account the fact that the view ray is going to be refracted at the interface of the slab.
+				// Otherwise, at tangent view direction, the path length traveled within the medium will tend towards infinity. That result in too dark edges and is unrealistic.
+				// To simplifiy things, we refract the view ray as if it was coming from a air interface into a medium of IOR deducted from F0.
+				// If we do this computation once for each slab, it is equivalent to have each layered slab having the same IOR with a single view vector refraction at the top interface.
+				// This is an aproximation we accept for the sake of performance today. 
+				// RefractedTransmittanceNoV is used to build ThroughputV that will be used when processing the substrate tree to get correct throughput to layers below.
+				// Ideas for later (requires an update of the substrate tree evaluation)
+				//  1- we could compute IOR for each layers (gather IOR for each slab horizontally integrated) and refract V according to the from/to IOR a the interface.
+				//  2- we could have V from a layer propagated to the next layer in order to have better specular match. 
+				//  3- we could also store extra data in the closure in order to be able to account for that effect for the Lighting vector too.
+				float3 RefractedV = refract(BSDFContext.V, -BSDFContext.N, 1.0 / DielectricF0ToIor(F0.y));
+				const float RefractedTransmittanceNoV = dot(RefractedV, BSDFContext.N);
+			#else
+				const float RefractedTransmittanceNoV = BSDFContext.Context.NoV;
+			#endif
+
 				// Shading point <-> View throughput
-				const float3 SlabTransmittanceV = IsotropicMediumSlabTransmittance(PM, STRATA_SIMPLEVOLUME_THICKNESS_M, BSDFContext.Context.NoV);
+				const float3 SlabTransmittanceV = IsotropicMediumSlabTransmittance(PM, STRATA_SIMPLEVOLUME_THICKNESS_M, RefractedTransmittanceNoV);
 				const float SpecularTransmissionV = DirectionalAlbedo_SpecularTransmission;
 
 				// Shading point <-> Light throughput
