Fix hair tangent on group's first and last segments with culling pass.

diff --git a/Engine/Shaders/Private/HairStrands/HairStrandsTangent.usf b/Engine/Shaders/Private/HairStrands/HairStrandsTangent.usf
--- a/Engine/Shaders/Private/HairStrands/HairStrandsTangent.usf
+++ b/Engine/Shaders/Private/HairStrands/HairStrandsTangent.usf
@@ -52,8 +52,8 @@ void MainCS(uint DispatchThreadId : SV_DispatchThreadID)
 			IndexCurr = HairStrandsVF_CullingIndexBuffer[FetchIndexCurr];
 			IndexNext = HairStrandsVF_CullingIndexBuffer[FetchIndexNext];
 
-			bPrevValid = IndexPrev >= 0;
-			bNextValid = IndexNext < PointCountAfterCulling;
+			bPrevValid = FetchIndex-1 >= 0;
+			bNextValid = FetchIndex+1 < PointCountAfterCulling;
 
 			bValid = true;
 		}
