Substrate - fixed thin translucent not working with substrate anymore after some STRATA_FASTPATH optimisations.

diff --git a/Engine/Shaders/Private/Strata/Strata.ush b/Engine/Shaders/Private/Strata/Strata.ush
--- a/Engine/Shaders/Private/Strata/Strata.ush
+++ b/Engine/Shaders/Private/Strata/Strata.ush
@@ -83,6 +83,10 @@
 #define STRATA_USES_CONVERSION_FROM_LEGACY 0
 #endif
 
+#ifndef MATERIALBLENDING_ANY_TRANSLUCENT
+#define MATERIALBLENDING_ANY_TRANSLUCENT 0
+#endif
+
 // Strata packed data is read either from a view StrataMaterial resource (Texture2DArray) or a RayTracingPayload/
 #ifndef STRATA_MATERIALCONTAINER_IS_VIEWRESOURCE
 #define STRATA_MATERIALCONTAINER_IS_VIEWRESOURCE (!RAYCALLABLESHADER && !RAYHITGROUPSHADER && !RAYGENSHADER && !RAYMISSSHADER)
@@ -1418,7 +1422,7 @@ FStrataData GetStrataSlabBSDF(
 	BSDF_SETSHAREDLOCALBASISID(StrataData.InlinedBSDF,		SharedLocalBasisIndex);
 	BSDF_SETEMISSIVE		(StrataData.InlinedBSDF,		Emissive);
 
-#if !STRATA_FASTPATH
+#if (STRATA_FASTPATH==0 || MATERIALBLENDING_ANY_TRANSLUCENT) // !Fast path or rendering translucent materials
 	BSDF_SETSSSTYPE			(StrataData.InlinedBSDF,		bHasSSS ? (bHasSSSWrap ? SSS_TYPE_WRAP : (bHasSSSProfile ? SSS_TYPE_DIFFUSION_PROFILE : SSS_TYPE_DIFFUSION)) : SSS_TYPE_INVALID);
 	BSDF_SETHASANISOTROPY	(StrataData.InlinedBSDF,		bHasAnisotropy);
 	BSDF_SETISTOPLAYER		(StrataData.InlinedBSDF,		0);
diff --git a/Engine/Shaders/Private/Strata/StrataEvaluation.ush b/Engine/Shaders/Private/Strata/StrataEvaluation.ush
--- a/Engine/Shaders/Private/Strata/StrataEvaluation.ush
+++ b/Engine/Shaders/Private/Strata/StrataEvaluation.ush
@@ -40,6 +40,10 @@
 #define STRATA_TRANSLUCENT_ENABLED 0
 #endif
 
+#ifndef MATERIALBLENDING_ANY_TRANSLUCENT
+#define MATERIALBLENDING_ANY_TRANSLUCENT 0
+#endif
+
 #ifndef STRATA_SHEEN_QUALITY
 #define STRATA_SHEEN_QUALITY 0
 #endif
@@ -888,7 +892,7 @@ FStrataEvaluateResult StrataEvaluateBSDFCommon(FStrataBSDFContext BSDFContext, F
 			//// This is used for optically thin translucent objects, or non-bottom layer of a material.
 			////
 
-			#if STRATA_FASTPATH==0
+			#if (STRATA_FASTPATH==0 || MATERIALBLENDING_ANY_TRANSLUCENT) // !Fast path or rendering translucent materials
 			BRANCH
 			if (BSDF_GETSSSTYPE(BSDFContext.BSDF) == SSS_TYPE_SIMPLEVOLUME)
 			{
diff --git a/Engine/Shaders/Private/Strata/StrataTree.ush b/Engine/Shaders/Private/Strata/StrataTree.ush
--- a/Engine/Shaders/Private/Strata/StrataTree.ush
+++ b/Engine/Shaders/Private/Strata/StrataTree.ush
@@ -4,9 +4,6 @@
 #include "/Engine/Private/Strata/StrataEvaluation.ush"
 #include "/Engine/Private/ShadingCommon.ush"
 
-//#ifndef STRATA_OPAQUE_MATERIAL
-//#define STRATA_OPAQUE_MATERIAL 0
-//#endif
 #ifndef STRATA_OPAQUE_MATERIAL
 #error STRATA_OPAQUE_MATERIAL not defined
 #endif
@@ -49,7 +46,7 @@ void FStrataTree::UpdateSingleBSDFOperatorCoverageTransmittance(
 	CurrentBSDF.StrataSanitizeBSDF();
 
 	// Distortion and translucency passes accumnulate pass calls for colored transmittance, coverage and TopDownRefractionLobeStat to be setup correctly.
-#if (STRATA_COMPLEXPATH || MATERIALBLENDING_ANY_TRANSLUCENT) && !STRATA_OPTIMIZED_UNLIT
+#if (STRATA_FASTPATH==0 || MATERIALBLENDING_ANY_TRANSLUCENT) && !STRATA_OPTIMIZED_UNLIT // !Fast path or rendering translucent materials
 
 	// Extract the true thickness of the BSDF before it is cast into the normalised Slab thickness we used.
 	// for opaque rough refraction, we are only interested in thickness of non bottom layers.
