Fix Substrate shader compilation warning.

diff --git a/Engine/Shaders/Private/Strata/Strata.ush b/Engine/Shaders/Private/Strata/Strata.ush
--- a/Engine/Shaders/Private/Strata/Strata.ush
+++ b/Engine/Shaders/Private/Strata/Strata.ush
@@ -1359,7 +1359,7 @@ FStrataData GetStrataSlabBSDF(
 	bool bIsThin		= bIsAtTheBottomOfTopology && (bMFPPlugged && bIsThinSurface); // Enable thin only if MFP is plugged, as otherwise there will be no back-tramission
 	bool bHasSSSWrap	= bIsAtTheBottomOfTopology && (bMFPPlugged && !UseSSSDiffusion);
 	bool bHasSSSProfile	= bIsAtTheBottomOfTopology && (bSSSProfilePlugged && SSSMFPScale > 0.f);
-	bool bHasSSS		= bMFPPlugged || bIsAtTheBottomOfTopology && (bHasSSSProfile || bHasSSSWrap); // When MFP is plugged, that can be diffusion (bottom) or Simple volume (any position)
+	bool bHasSSS		= bMFPPlugged || (bIsAtTheBottomOfTopology && (bHasSSSProfile || bHasSSSWrap)); // When MFP is plugged, that can be diffusion (bottom) or Simple volume (any position)
 
 	FHaziness Haziness = InitialiseHaziness();
 
