PR #10228: Fixed wrong variable name

diff --git a/Engine/Shaders/Private/ForwardLightingCommon.ush b/Engine/Shaders/Private/ForwardLightingCommon.ush
--- a/Engine/Shaders/Private/ForwardLightingCommon.ush
+++ b/Engine/Shaders/Private/ForwardLightingCommon.ush
@@ -194,7 +194,7 @@ FDeferredLightingSplit GetForwardDirectLightingSplit(
 			{
 				#if SUPPORT_VSM_FOWARD_QUALITY == 1
 				{
-					const FLightShaderParameters DireciontalLight = ConvertToLightShaderParameters(LightData);
+					const FLightShaderParameters DirectionalLight = ConvertToLightShaderParameters(LightData);
 					FVirtualShadowMapSMRTSettings SMRTSettings = GetDirectionalLightSMRTSettings();
 
 					const bool bBackfaceCull = !IsSubsurfaceModel(GBufferData.ShadingModelID);
@@ -203,7 +203,7 @@ FDeferredLightingSplit GetForwardDirectLightingSplit(
 
 					FVirtualShadowMapSampleResult VirtualShadowMapSample = TraceDirectional(
 						ForwardLightData.DirectionalLightVSM,
-						DireciontalLight,
+						DirectionalLight,
 						PixelPos,
 						GBufferData.Depth,
 						TranslatedWorldPosition,
@@ -546,4 +546,4 @@ half3 GetImageBasedReflectionLighting(FMaterialPixelParameters MaterialParameter
 half3 GetImageBasedReflectionLighting(FMaterialPixelParameters MaterialParameters, half Roughness, half3 SpecularColor, half IndirectIrradiance, uint GridIndex, int SingleCaptureIndex)
 {
 	return GetImageBasedReflectionLighting(MaterialParameters, Roughness, SpecularColor, IndirectIrradiance, GridIndex, SingleCaptureIndex, 0);
-}
\ No newline at end of file
+}
