Fix typo.

diff --git a/Engine/Shaders/Private/Strata/StrataForwardLighting.ush b/Engine/Shaders/Private/Strata/StrataForwardLighting.ush
--- a/Engine/Shaders/Private/Strata/StrataForwardLighting.ush
+++ b/Engine/Shaders/Private/Strata/StrataForwardLighting.ush
@@ -118,7 +118,7 @@ float3 StrataForwardLighting(
 	{
 		#if SUPPORT_VSM_FOWARD_QUALITY == 1
 		{
-			const FLightShaderParameters DireciontalLight = ConvertToLightShaderParameters(DirLightData);
+			const FLightShaderParameters DirectionalLight = ConvertToLightShaderParameters(DirLightData);
 			FVirtualShadowMapSMRTSettings SMRTSettings = GetDirectionalLightSMRTSettings();
 
 			const bool bBackfaceCull = true; // STRATA_TODO: this needs to be fixed but require to extract hair/SSS/diffusebackface prior to BSDF looping
@@ -130,7 +130,7 @@ float3 StrataForwardLighting(
 
 			FVirtualShadowMapSampleResult VirtualShadowMapSample = TraceDirectional(
 				ForwardLightData.DirectionalLightVSM,
-				DireciontalLight,
+				DirectionalLight,
 				PixelPos,
 				SceneDepth,
 				TranslatedWorldPosition,
