Fixes to Majorant grid build

diff --git a/Engine/Shaders/Private/HeterogeneousVolumes/HeterogeneousVolumesVoxelGridPipeline.usf b/Engine/Shaders/Private/HeterogeneousVolumes/HeterogeneousVolumesVoxelGridPipeline.usf
--- a/Engine/Shaders/Private/HeterogeneousVolumes/HeterogeneousVolumesVoxelGridPipeline.usf
+++ b/Engine/Shaders/Private/HeterogeneousVolumes/HeterogeneousVolumesVoxelGridPipeline.usf
@@ -458,7 +458,7 @@ void BuildMajorantVoxelGridCS(
 
 	int LinearIndex = GetLinearIndex(VoxelIndex, TopLevelGridResolution);
 	FTopLevelGridData TopLevelGridData = TopLevelGridBuffer[LinearIndex];
-	if (!IsBottomLevelEmpty(TopLevelGridData))
+	if (IsBottomLevelAllocated(TopLevelGridData))
 	{
 		int3 VoxelResolution = GetBottomLevelVoxelResolution(TopLevelGridData);
 		int BottomLevelVoxelCount = VoxelResolution.x * VoxelResolution.y * VoxelResolution.z;
@@ -467,14 +467,17 @@ void BuildMajorantVoxelGridCS(
 		for (int Index = 0; Index < BottomLevelVoxelCount; ++Index)
 		{
 			FTopLevelGridData IndirectionData = IndirectionGridBuffer[IndirectionGridIndex + Index];
-			uint IndirectionIndex = GetBottomLevelIndex(IndirectionData);
-			int3 IndirectionResolution = GetBottomLevelVoxelResolution(IndirectionData);
-			int IndirectionVoxelCount = IndirectionResolution.x * IndirectionResolution.y * IndirectionResolution.z;
-			for (uint VoxelIndex = 0; VoxelIndex < IndirectionVoxelCount; ++VoxelIndex)
+			if (IsBottomLevelAllocated(IndirectionData))
 			{
-				uint BottomLevelIndex = IndirectionIndex + VoxelIndex;
-				float3 Extinction = GetExtinction(ExtinctionGridBuffer[BottomLevelIndex]);
-				UpdateMajorantData(MajorantData, Extinction);
+				uint IndirectionIndex = GetBottomLevelIndex(IndirectionData);
+				int3 IndirectionResolution = GetBottomLevelVoxelResolution(IndirectionData);
+				int IndirectionVoxelCount = IndirectionResolution.x * IndirectionResolution.y * IndirectionResolution.z;
+				for (uint VoxelIndex = 0; VoxelIndex < IndirectionVoxelCount; ++VoxelIndex)
+				{
+					uint BottomLevelIndex = IndirectionIndex + VoxelIndex;
+					float3 Extinction = GetExtinction(ExtinctionGridBuffer[BottomLevelIndex]);
+					UpdateMajorantData(MajorantData, Extinction);
+				}
 			}
 		}
 #else // DIM_ENABLE_INDIRECTION_GRID
