PathTracing: Correctly support mixing of HeterogeneousVolumes with SkyAtmosphere and HeightFog.

diff --git a/Engine/Shaders/Private/PathTracing/Volume/PathTracingHeterogeneousVolumes.ush b/Engine/Shaders/Private/PathTracing/Volume/PathTracingHeterogeneousVolumes.ush
--- a/Engine/Shaders/Private/PathTracing/Volume/PathTracingHeterogeneousVolumes.ush
+++ b/Engine/Shaders/Private/PathTracing/Volume/PathTracingHeterogeneousVolumes.ush
@@ -436,44 +436,40 @@ FVolumeTrackingResult HeterogeneousVolumesDeltaTracking(float3 PathThroughput, f
 	FVolumeTrackingResult Sample = CreateVolumeTrackingResult(0);
 	Sample.SigmaBar = max(MajorantData.Majorant, 0);
 
-	while ((Sample.Distance < TMax) && !Sample.bIsCollision)
-	{
-		float RandValue = RandomSequence_GenerateSample1D(RandSequence);
+	float RandValue = RandomSequence_GenerateSample1D(RandSequence);
 #define SAMPLE_ACHROMATIC 0
 #if SAMPLE_ACHROMATIC
-		float DeltaT = -log(1.0 - RandValue) / Sample.SigmaBar.x;
+	float DeltaT = -log(1.0 - RandValue) / Sample.SigmaBar.x;
 #else
-		float DeltaT = SampleSpectralTransmittance(RandValue, Sample.SigmaBar, PathThroughput);
-		if (DeltaT < 0)
-		{
-			Sample.Throughput = 0;
-			Sample.Distance = -1;
-			return Sample;
-		}
+	float DeltaT = SampleSpectralTransmittance(RandValue, Sample.SigmaBar, PathThroughput);
+	if (DeltaT < 0)
+	{
+		Sample.Throughput = 0;
+		Sample.Distance = -1;
+		return Sample;
+	}
 #endif
 
-		Sample.Distance = min(Sample.Distance + DeltaT, TMax);
-		float Transmittance = exp(-Sample.SigmaBar.x * Sample.Distance);
-
-		if (Sample.Distance < TMax)
-		{
-			Sample.bIsCollision = true;
-			float CollisionWeight = 1.0;
+	Sample.Distance = min(Sample.Distance + DeltaT, TMax);
+	if (Sample.Distance < TMax)
+	{
+		Sample.bIsCollision = true;
+		float CollisionWeight = 1.0;
 #if SAMPLE_ACHROMATIC
-			Sample.Throughput *= CollisionWeight / Sample.SigmaBar.x;
-			Sample.Pdf = Transmittance * Sample.SigmaBar.x;
+		Sample.Throughput *= CollisionWeight / Sample.SigmaBar.x;
+		float Transmittance = exp(-Sample.SigmaBar.x * Sample.Distance);
+		Sample.Pdf = Transmittance * Sample.SigmaBar.x;
 #else
-			float4 EvalResult = EvaluateSpectralTransmittanceHit(Sample.Distance, Sample.SigmaBar, PathThroughput);
-			Sample.Throughput = CollisionWeight * EvalResult.xyz;
-			Sample.Pdf = EvalResult.w;
-		}
-		else
-		{
-			float4 EvalResult = EvaluateSpectralTransmittanceMiss(Sample.Distance, Sample.SigmaBar, PathThroughput);
-			Sample.Throughput = EvalResult.xyz;
-			Sample.Pdf = EvalResult.w;
+		float4 EvalResult = EvaluateSpectralTransmittanceHit(Sample.Distance, Sample.SigmaBar, PathThroughput);
+		Sample.Throughput = CollisionWeight * EvalResult.xyz;
+		Sample.Pdf = EvalResult.w;
+	}
+	else
+	{
+		float4 EvalResult = EvaluateSpectralTransmittanceMiss(Sample.Distance, Sample.SigmaBar, PathThroughput);
+		Sample.Throughput = EvalResult.xyz;
+		Sample.Pdf = EvalResult.w;
 #endif
-		}
 	}
 
 	return Sample;
@@ -484,13 +480,14 @@ FVolumeTrackingResult HeterogeneousVolumesRatioTracking(float3 PathThroughput, f
 	FVolumeTrackingResult Sample = CreateVolumeTrackingResult(0);
 	Sample.SigmaBar = max(MajorantData.Majorant, 0);
 
-	while ((Sample.Distance < TMax) && !Sample.bIsCollision)
+	while ((Sample.Distance < TMax))
 	{
 		float RandValue = RandomSequence_GenerateSample1D(RandSequence);
 #if SAMPLE_ACHROMATIC
 		float DeltaT = -log(1.0 - RandValue) / Sample.SigmaBar.x;
 #else
-		float DeltaT = SampleSpectralTransmittance(RandValue, Sample.SigmaBar, PathThroughput);
+		float3 CombinedThroughput = PathThroughput * Sample.Throughput;
+		float DeltaT = SampleSpectralTransmittance(RandValue, Sample.SigmaBar, CombinedThroughput);
 		if (DeltaT < 0)
 		{
 			Sample.Throughput = 0;
@@ -499,29 +496,30 @@ FVolumeTrackingResult HeterogeneousVolumesRatioTracking(float3 PathThroughput, f
 		}
 #endif
 
-		Sample.Distance = min(Sample.Distance + DeltaT, TMax);
-		float Transmittance = exp(-Sample.SigmaBar.x * Sample.Distance);
-
-		if (Sample.Distance < TMax)
+		float CollisionDistance = Sample.Distance + DeltaT;
+		if (CollisionDistance < TMax)
 		{
-			float3 SamplePos = Origin + Direction * Sample.Distance;
+			float3 SamplePos = Origin + Direction * CollisionDistance;
 			float3 SigmaT = HeterogeneousVolumesGetDensity(SamplePos).SigmaT;
-			float3 SpectralNullCollisionWeight = max(Sample.SigmaBar.x - SigmaT, 0);
+			float3 SpectralNullCollisionWeight = max(Sample.SigmaBar - SigmaT, 0);
 #if SAMPLE_ACHROMATIC
 			Sample.Throughput *= SpectralNullCollisionWeight / Sample.SigmaBar.x;
+			float Transmittance = exp(-Sample.SigmaBar.x * DeltaT);
 			Sample.Pdf *= Transmittance * Sample.SigmaBar.x;
 #else
-			float4 EvalResult = EvaluateSpectralTransmittanceHit(Sample.Distance, Sample.SigmaBar, PathThroughput);
+			float4 EvalResult = EvaluateSpectralTransmittanceHit(DeltaT, Sample.SigmaBar, CombinedThroughput);
 			Sample.Throughput *= SpectralNullCollisionWeight * EvalResult.xyz;
 			Sample.Pdf *= EvalResult.w;
 		}
 		else
 		{
-			float4 EvalResult = EvaluateSpectralTransmittanceMiss(Sample.Distance, Sample.SigmaBar, PathThroughput);
+			float4 EvalResult = EvaluateSpectralTransmittanceMiss(TMax - Sample.Distance, Sample.SigmaBar, CombinedThroughput);
 			Sample.Throughput *= EvalResult.xyz;
 			Sample.Pdf *= EvalResult.w;
 #endif
 		}
+
+		Sample.Distance = min(CollisionDistance, TMax);
 	}
 
 	return Sample;
@@ -545,6 +543,16 @@ FVolumeTrackingResult HeterogeneousVolumesMajorantBasedTracking(float3 PathThrou
 	float3 VoxelRayDirection = VoxelRayEnd - VoxelRayBegin;
 	float VoxelRayTMin = 0.0;
 	float VoxelRayTMax = length(VoxelRayDirection);
+
+	if (VoxelRayTMin >= VoxelRayTMax && !bTransmittanceOnly)
+	{
+		FVolumeTrackingResult Sample = HeterogeneousVolumesDeltaTracking(Result.Throughput, WorldRayBegin, Direction, TMax - TMin, OverlappingMajorant, RandSequence);
+		if (Sample.Distance > 0.0)
+		{
+			Result = ConcatenateVolumeTrackingResult(Result, Sample);
+		}
+		return Result;
+	}
 	VoxelRayDirection /= VoxelRayTMax;
 
 	float VoxelToWorldScale = WorldRayTMax / VoxelRayTMax;
@@ -640,14 +648,9 @@ FVolumeTrackingResult HeterogeneousVolumesMajorantBasedTracking(float3 PathThrou
 
 float3 HeterogeneousVolumesGetTransmittance(float3 PathThroughput, float3 Origin, float3 Direction, float TMin, float TMax, inout RandomSequence RandSequence)
 {
-	if (EnableHeterogeneousVolumes)
-	{
-		bool bTransmittanceOnly = true;
-		FMajorantData EmptyOverlappingMajorant = CreateMajorantData();
-		return HeterogeneousVolumesMajorantBasedTracking(PathThroughput, Origin, Direction, TMin, TMax, bTransmittanceOnly, EmptyOverlappingMajorant, RandSequence).Throughput;
-	}
-
-	return PathThroughput;
+	bool bTransmittanceOnly = true;
+	FMajorantData EmptyOverlappingMajorant = CreateMajorantData(0.0f, 0.0f);
+	return HeterogeneousVolumesMajorantBasedTracking(PathThroughput, Origin, Direction, TMin, TMax, bTransmittanceOnly, EmptyOverlappingMajorant, RandSequence).Throughput;
 }
 
 FVolumeTrackingResult HeterogeneousVolumesSampleDistance(float3 PathThroughput, float3 Origin, float3 Direction, float TMin, float TMax, FMajorantData OverlappingMajorant, inout RandomSequence RandSequence)
diff --git a/Engine/Shaders/Private/PathTracing/Volume/PathTracingVolume.ush b/Engine/Shaders/Private/PathTracing/Volume/PathTracingVolume.ush
--- a/Engine/Shaders/Private/PathTracing/Volume/PathTracingVolume.ush
+++ b/Engine/Shaders/Private/PathTracing/Volume/PathTracingVolume.ush
@@ -243,10 +243,16 @@ FVolumeTrackingResult VolumeSampleDistance(float3 PathThroughput, float3 Origin,
 		SigmaBar *= bIsCameraRay ? 2 : 1;
 	}
 
+	bool bIntervalHasHeterogeneousVolume = false;
+	for (int Index = 0; Index < Interval.Num; ++Index)
+	{
+		bIntervalHasHeterogeneousVolume |= (Interval.VolumeID[Index] == VOLUMEID_HETEROGENEOUS_VOLUMES);
+	}
+
 	FVolumeTrackingResult TrackingResult;
-	if (EnableHeterogeneousVolumes)
+	if (bIntervalHasHeterogeneousVolume)
 	{
-		FMajorantData OverlappingMajorant = CreateMajorantData(VolumeDensityBounds.SigmaMin, VolumeDensityBounds.SigmaMax);
+		FMajorantData OverlappingMajorant = CreateMajorantData(VolumeDensityBounds.SigmaMin, SigmaBar);
 		TrackingResult = HeterogeneousVolumesSampleDistance(PathThroughput, Origin, Direction, Interval.VolumeTMin, Interval.VolumeTMax, OverlappingMajorant, RandSequence);
 	}
 	else
@@ -275,5 +281,6 @@ FVolumeTrackingResult VolumeSampleDistance(float3 PathThroughput, float3 Origin,
 			TrackingResult.Pdf = Evaluation.w;
 		}
 	}
+
 	return TrackingResult;
 }
