Dead code removal

diff --git a/Engine/Shaders/Private/PathTracing/Volume/PathTracingHeterogeneousVolumes.ush b/Engine/Shaders/Private/PathTracing/Volume/PathTracingHeterogeneousVolumes.ush
--- a/Engine/Shaders/Private/PathTracing/Volume/PathTracingHeterogeneousVolumes.ush
+++ b/Engine/Shaders/Private/PathTracing/Volume/PathTracingHeterogeneousVolumes.ush
@@ -51,207 +51,10 @@ FVolumeIntersection HeterogeneousVolumesIntersect(float3 Origin, float3 Directio
 	return CreateEmptyVolumeIntersection();
 }
 
-void MarchMajorantGrid(
-	float3 VoxelRayOrigin,
-	float3 VoxelRayDirection,
-	float VoxelRayTMin,
-	float VoxelRayTMax,
-	float WorldRayHitSpan,
-	inout FMajorantData MajorantData
-)
-{
-	float3 VoxelRayDirectionInv = 1.0 / VoxelRayDirection;
-
-	// Grid intersection
-	float VoxelRayMarchT = VoxelRayTMin;
-	float3 VoxelRayMarchPos = VoxelRayOrigin;
-
-	float3 VoxelBoundsPos;
-	VoxelBoundsPos.x = sign(VoxelRayDirection.x) > 0 ? floor(VoxelRayMarchPos.x) + 1 : ceil(VoxelRayMarchPos.x) - 1;
-	VoxelBoundsPos.y = sign(VoxelRayDirection.y) > 0 ? floor(VoxelRayMarchPos.y) + 1 : ceil(VoxelRayMarchPos.y) - 1;
-	VoxelBoundsPos.z = sign(VoxelRayDirection.z) > 0 ? floor(VoxelRayMarchPos.z) + 1 : ceil(VoxelRayMarchPos.z) - 1;
-
-	float3 VoxelBoundsHitT = (VoxelBoundsPos - VoxelRayMarchPos) * VoxelRayDirectionInv;
-
-	// Remove floating-point rounding error
-	float Epsilon = 1.0e-4;
-	if (VoxelBoundsHitT.x <= Epsilon) VoxelBoundsHitT.x += abs(VoxelRayDirectionInv.x);
-	if (VoxelBoundsHitT.y <= Epsilon) VoxelBoundsHitT.y += abs(VoxelRayDirectionInv.y);
-	if (VoxelBoundsHitT.z <= Epsilon) VoxelBoundsHitT.z += abs(VoxelRayDirectionInv.z);
-
-	while ((VoxelRayMarchT < VoxelRayTMax))
-	{
-		float VoxelRayMarchDeltaT = 0.0;
-		if (VoxelBoundsHitT.x < VoxelBoundsHitT.y)
-		{
-			if (VoxelBoundsHitT.x < VoxelBoundsHitT.z)
-			{
-				VoxelRayMarchDeltaT = VoxelBoundsHitT.x - VoxelRayMarchT;
-				VoxelBoundsHitT.x += abs(VoxelRayDirectionInv.x);
-			}
-			else
-			{
-				VoxelRayMarchDeltaT = VoxelBoundsHitT.z - VoxelRayMarchT;
-				VoxelBoundsHitT.z += abs(VoxelRayDirectionInv.z);
-			}
-		}
-		else
-		{
-			if (VoxelBoundsHitT.y < VoxelBoundsHitT.z)
-			{
-				VoxelRayMarchDeltaT = VoxelBoundsHitT.y - VoxelRayMarchT;
-				VoxelBoundsHitT.y += abs(VoxelRayDirectionInv.y);
-			}
-			else
-			{
-				VoxelRayMarchDeltaT = VoxelBoundsHitT.z - VoxelRayMarchT;
-				VoxelBoundsHitT.z += abs(VoxelRayDirectionInv.z);
-			}
-		}
-
-		// Clip voxel delta-t by overall voxel ray-length.
-		if (VoxelRayMarchT + VoxelRayMarchDeltaT > VoxelRayTMax)
-		{
-			VoxelRayMarchDeltaT = VoxelRayTMax - VoxelRayMarchT;
-		}
-
-		// Sample at the midpoint of the voxel ray
-		float3 VoxelEndPos = VoxelRayMarchPos + VoxelRayDirection * VoxelRayMarchDeltaT;
-		float3 VoxelSamplePos = (VoxelRayMarchPos + VoxelEndPos) * 0.5;
-		uint VoxelSampleLinearPos = GetLinearIndex(VoxelSamplePos, OrthoGridUniformBuffer.TopLevelGridResolution);
-
-		int3 BottomLevelVoxelResolution = GetBottomLevelVoxelResolution(OrthoGridUniformBuffer.TopLevelGridBuffer[VoxelSampleLinearPos]);
-		if (IsBottomLevelAllocated(OrthoGridUniformBuffer.TopLevelGridBuffer[VoxelSampleLinearPos]))
-		{
-			UpdateMajorantData(MajorantData, GetMajorantData(OrthoGridUniformBuffer.MajorantGridBuffer[VoxelSampleLinearPos]));
-		}
-
-		VoxelRayMarchT += VoxelRayMarchDeltaT;
-		VoxelRayMarchPos = VoxelEndPos;
-	}
-}
-
-FMajorantData TraceMajorantGrid(float3 Origin, float3 Direction, float TMin, float TMax)
-{
-	FMajorantData MajorantData = CreateMajorantData();
-
-	float3 WorldBoundsMin = OrthoGridUniformBuffer.TopLevelGridWorldBoundsMin;
-	float3 WorldBoundsMax = OrthoGridUniformBuffer.TopLevelGridWorldBoundsMax;
-
-	// Intersect world-space AABB to clip ray
-	float3 WorldRayOrigin = LWCHackToFloat(LWCSubtract(Origin, PrimaryView.PreViewTranslation));
-	float2 ClipHitT = IntersectAABB(WorldRayOrigin, Direction, TMin, TMax, WorldBoundsMin, WorldBoundsMax);
-	if (ClipHitT.x > ClipHitT.y)
-	{
-		return MajorantData;
-	}
-
-	float3 WorldRayBegin = WorldRayOrigin + Direction * ClipHitT.x;
-	float3 WorldRayEnd = WorldRayOrigin + Direction * ClipHitT.y;
-	float WorldRayLength = ClipHitT.y - ClipHitT.x;
-
-	// Transform to voxel-space
-	float3 TopLevelGridWorldBoundsExtent = WorldBoundsMax - WorldBoundsMin;
-	float3 VoxelRayBegin = (WorldRayBegin - WorldBoundsMin) / TopLevelGridWorldBoundsExtent * OrthoGridUniformBuffer.TopLevelGridResolution;
-	float3 VoxelRayEnd = (WorldRayEnd - WorldBoundsMin) / TopLevelGridWorldBoundsExtent * OrthoGridUniformBuffer.TopLevelGridResolution;
-	float3 VoxelRayDirection = VoxelRayEnd - VoxelRayBegin;
-	float VoxelRayTMin = 0.0;
-	float VoxelRayTMax = length(VoxelRayDirection);
-	VoxelRayDirection /= VoxelRayTMax;
-
-	MarchMajorantGrid(
-		VoxelRayBegin,
-		VoxelRayDirection,
-		VoxelRayTMin,
-		VoxelRayTMax,
-		WorldRayLength,
-		MajorantData
-	);
-
-	return MajorantData;
-}
-
 FVolumeDensityBounds HeterogeneousVolumesGetDensityBounds(float3 Origin, float3 Direction, float TMin, float TMax)
 {
-	if (EnableHeterogeneousVolumes)
-	{
-		FMajorantData MajorantData = (FMajorantData) 0;
-		if (OrthoGridUniformBuffer.bUseMajorantGrid)
-		{
-			MajorantData = TraceMajorantGrid(Origin, Direction, TMin, TMax);
-		}
-		else
-		{
-			MajorantData.Minorant = 0;
-			MajorantData.Majorant = 0.1;
-		}
-
-		return CreateVolumeDensityBound(MajorantData.Minorant, MajorantData.Majorant);
-	}
-
+	// Density bounds are iteratively discovered when executing the DDA on the majorant grid
 	return CreateVolumeDensityBound(0, 0);
-
-}
-
-void EvalWithLinearInterpolation(uint BottomLevelIndex, float3 VoxelPos,
-	StructuredBuffer<FScalarGridData> ExtinctionGridBuffer,
-	StructuredBuffer<FVectorGridData> ScatteringGridBuffer,
-	inout float3 SigmaT, inout float3 Albedo)
-{
-	uint Voxel000 = BottomLevelIndex + MortonEncode3(VoxelPos + uint3(0, 0, 0));
-	uint Voxel100 = BottomLevelIndex + MortonEncode3(VoxelPos + uint3(1, 0, 0));
-	uint Voxel010 = BottomLevelIndex + MortonEncode3(VoxelPos + uint3(0, 1, 0));
-	uint Voxel110 = BottomLevelIndex + MortonEncode3(VoxelPos + uint3(1, 1, 0));
-	uint Voxel001 = BottomLevelIndex + MortonEncode3(VoxelPos + uint3(0, 0, 1));
-	uint Voxel101 = BottomLevelIndex + MortonEncode3(VoxelPos + uint3(1, 0, 1));
-	uint Voxel011 = BottomLevelIndex + MortonEncode3(VoxelPos + uint3(0, 1, 1));
-	uint Voxel111 = BottomLevelIndex + MortonEncode3(VoxelPos + uint3(1, 1, 1));
-
-	// SigmaT
-	{
-		float3 SigmaT_000 = GetExtinction(ExtinctionGridBuffer[Voxel000]);
-		float3 SigmaT_100 = GetExtinction(ExtinctionGridBuffer[Voxel100]);
-		float3 SigmaT_010 = GetExtinction(ExtinctionGridBuffer[Voxel010]);
-		float3 SigmaT_110 = GetExtinction(ExtinctionGridBuffer[Voxel110]);
-		float3 SigmaT_001 = GetExtinction(ExtinctionGridBuffer[Voxel001]);
-		float3 SigmaT_101 = GetExtinction(ExtinctionGridBuffer[Voxel101]);
-		float3 SigmaT_011 = GetExtinction(ExtinctionGridBuffer[Voxel011]);
-		float3 SigmaT_111 = GetExtinction(ExtinctionGridBuffer[Voxel111]);
-
-		float3 Weight = frac(VoxelPos);
-		float3 Edge00 = (1.0 - Weight.x) * SigmaT_000 + Weight.x * SigmaT_100;
-		float3 Edge10 = (1.0 - Weight.x) * SigmaT_010 + Weight.x * SigmaT_110;
-		float3 Face0 = (1.0 - Weight.y) * Edge00 + Weight.y * Edge10;
-
-		float3 Edge01 = (1.0 - Weight.x) * SigmaT_001 + Weight.x * SigmaT_101;
-		float3 Edge11 = (1.0 - Weight.x) * SigmaT_011 + Weight.x * SigmaT_111;
-		float3 Face1 = (1.0 - Weight.y) * Edge01 + Weight.y * Edge11;
-
-		SigmaT = (1.0 - Weight.z) * Face0 + Weight.z * Face1;
-	}
-
-	// Albedo
-	{
-		float3 Albedo_000 = GetAlbedo(ScatteringGridBuffer[Voxel000]);
-		float3 Albedo_100 = GetAlbedo(ScatteringGridBuffer[Voxel100]);
-		float3 Albedo_010 = GetAlbedo(ScatteringGridBuffer[Voxel010]);
-		float3 Albedo_110 = GetAlbedo(ScatteringGridBuffer[Voxel110]);
-		float3 Albedo_001 = GetAlbedo(ScatteringGridBuffer[Voxel001]);
-		float3 Albedo_101 = GetAlbedo(ScatteringGridBuffer[Voxel101]);
-		float3 Albedo_011 = GetAlbedo(ScatteringGridBuffer[Voxel011]);
-		float3 Albedo_111 = GetAlbedo(ScatteringGridBuffer[Voxel111]);
-
-		float3 Weight = frac(VoxelPos);
-		float3 Edge00 = (1.0 - Weight.x) * Albedo_000 + Weight.x * Albedo_100;
-		float3 Edge10 = (1.0 - Weight.x) * Albedo_010 + Weight.x * Albedo_110;
-		float3 Face0 = (1.0 - Weight.y) * Edge00 + Weight.y * Edge10;
-
-		float3 Edge01 = (1.0 - Weight.x) * Albedo_001 + Weight.x * Albedo_101;
-		float3 Edge11 = (1.0 - Weight.x) * Albedo_011 + Weight.x * Albedo_111;
-		float3 Face1 = (1.0 - Weight.y) * Edge01 + Weight.y * Edge11;
-
-		Albedo = (1.0 - Weight.z) * Face0 + Weight.z * Face1;
-	}
 }
 
 FVolumeShadedResult GetOrthoVoxelGridDensity(float3 TranslatedWorldPos)
