Heterogeneous Volumes: Fixes flickering when rendering without the lighting cache

diff --git a/Engine/Shaders/Private/HeterogeneousVolumes/HeterogeneousVolumesRayMarchingUtils.ush b/Engine/Shaders/Private/HeterogeneousVolumes/HeterogeneousVolumesRayMarchingUtils.ush
--- a/Engine/Shaders/Private/HeterogeneousVolumes/HeterogeneousVolumesRayMarchingUtils.ush
+++ b/Engine/Shaders/Private/HeterogeneousVolumes/HeterogeneousVolumesRayMarchingUtils.ush
@@ -122,11 +122,16 @@ float CalcShadowStepSize(float3 LocalRayDirection)
 
 float CalcShadowBias()
 {
-	float3 VolumeBounds = 2.0f * GetLocalBoundsExtent();
-	float3 VoxelSize = VolumeBounds / float3(GetLightingCacheResolution());
-	float3 WorldVoxelSize = mul(float4(VoxelSize, 0), GetLocalToWorld()).xyz;
-	float VoxelDiagonal = length(WorldVoxelSize);
-	float WorldShadowBias = VoxelDiagonal * GetLightingCacheVoxelBias();
+	float WorldShadowBias = 0.0f;
+	float3 LightingCacheResolutionAsFloat3 = GetLightingCacheResolution();
+	if (all(LightingCacheResolutionAsFloat3 > 0))
+	{
+		float3 VolumeBounds = 2.0f * GetLocalBoundsExtent();
+		float3 VoxelSize = VolumeBounds * rcp(LightingCacheResolutionAsFloat3);
+		float3 WorldVoxelSize = mul(float4(VoxelSize, 0), GetLocalToWorld()).xyz;
+		float VoxelDiagonal = length(WorldVoxelSize);
+		float WorldShadowBias = VoxelDiagonal * GetLightingCacheVoxelBias();
+	}
 	return WorldShadowBias;
 }
 
diff --git a/Engine/Source/Runtime/Renderer/Private/HeterogeneousVolumes/HeterogeneousVolumesHardwareRayTracing.cpp b/Engine/Source/Runtime/Renderer/Private/HeterogeneousVolumes/HeterogeneousVolumesHardwareRayTracing.cpp
--- a/Engine/Source/Runtime/Renderer/Private/HeterogeneousVolumes/HeterogeneousVolumesHardwareRayTracing.cpp
+++ b/Engine/Source/Runtime/Renderer/Private/HeterogeneousVolumes/HeterogeneousVolumesHardwareRayTracing.cpp
@@ -798,6 +798,12 @@ void RenderSingleScatteringWithPreshadingHardwareRayTracing(
 			PassParameters->LightingCache.LightingCacheVoxelBias = HeterogeneousVolumeInterface->GetShadowBiasFactor();
 			PassParameters->LightingCache.LightingCacheTexture = LightingCacheTexture;
 		}
+		else
+		{
+			PassParameters->LightingCache.LightingCacheResolution = FIntVector::ZeroValue;
+			PassParameters->LightingCache.LightingCacheVoxelBias = 0.0f;
+			PassParameters->LightingCache.LightingCacheTexture = FRDGSystemTextures::Get(GraphBuilder).VolumetricBlack;
+		}
 
 		// Ray data
 		PassParameters->MaxTraceDistance = HeterogeneousVolumes::GetMaxTraceDistance();
diff --git a/Engine/Source/Runtime/Renderer/Private/HeterogeneousVolumes/HeterogeneousVolumesLiveShadingPipeline.cpp b/Engine/Source/Runtime/Renderer/Private/HeterogeneousVolumes/HeterogeneousVolumesLiveShadingPipeline.cpp
--- a/Engine/Source/Runtime/Renderer/Private/HeterogeneousVolumes/HeterogeneousVolumesLiveShadingPipeline.cpp
+++ b/Engine/Source/Runtime/Renderer/Private/HeterogeneousVolumes/HeterogeneousVolumesLiveShadingPipeline.cpp
@@ -632,6 +632,8 @@ void RenderSingleScatteringWithLiveShading(
 		}
 		else
 		{
+			PassParameters->LightingCache.LightingCacheResolution = FIntVector::ZeroValue;
+			PassParameters->LightingCache.LightingCacheVoxelBias = 0.0f;
 			PassParameters->LightingCache.LightingCacheTexture = FRDGSystemTextures::Get(GraphBuilder).VolumetricBlack;
 		}
 
diff --git a/Engine/Source/Runtime/Renderer/Private/HeterogeneousVolumes/HeterogeneousVolumesPreshadingPipeline.cpp b/Engine/Source/Runtime/Renderer/Private/HeterogeneousVolumes/HeterogeneousVolumesPreshadingPipeline.cpp
--- a/Engine/Source/Runtime/Renderer/Private/HeterogeneousVolumes/HeterogeneousVolumesPreshadingPipeline.cpp
+++ b/Engine/Source/Runtime/Renderer/Private/HeterogeneousVolumes/HeterogeneousVolumesPreshadingPipeline.cpp
@@ -597,6 +597,12 @@ void RenderSingleScatteringWithPreshadingCompute(
 			PassParameters->LightingCache.LightingCacheVoxelBias = HeterogeneousVolumeInterface->GetShadowBiasFactor();
 			PassParameters->LightingCache.LightingCacheTexture = LightingCacheTexture;
 		}
+		else
+		{
+			PassParameters->LightingCache.LightingCacheResolution = FIntVector::ZeroValue;
+			PassParameters->LightingCache.LightingCacheVoxelBias = 0.0f;
+			PassParameters->LightingCache.LightingCacheTexture = FRDGSystemTextures::Get(GraphBuilder).VolumetricBlack;
+		}
 
 		// Ray data
 		PassParameters->MaxTraceDistance = HeterogeneousVolumes::GetMaxTraceDistance();
