Added local fog volume common file.

diff --git a/Engine/Shaders/Private/LocalFogVolume.usf b/Engine/Shaders/Private/LocalFogVolume.usf
--- a/Engine/Shaders/Private/LocalFogVolume.usf
+++ b/Engine/Shaders/Private/LocalFogVolume.usf
@@ -6,60 +6,9 @@
 #endif
 
 #include "SceneTexturesCommon.ush"
-#include "HeightFogCommon.ush"
-#include "ParticipatingMediaCommon.ush"
+#include "LocalFogVolumeCommon.ush"
 
 
-struct FLocalFogVolumeGPUInstanceData
-{
-	float4x4 Transform;
-	float4x4 InvTransform;
-
-	float4x4 InvTranformNoScale;
-	float4x4 TransformScaleOnly;
-
-	float Density;
-	float HeightFalloff;
-	float HeightOffset;
-	float RadialAttenuation;
-
-	float3 Albedo;
-	float  PhaseG;
-	float3 Emissive;
-	float  FogMode;
-};
-
-// Structured buffers cannot be sampled from the vertex shader on some mobile platforms.
-// That is why simple Buffers are used, because this tech needs to run on mobile platforms.
-Buffer<float4> LocalFogVolumeInstances;
-FLocalFogVolumeGPUInstanceData GetLocalFogVolumeGPUInstanceData(uint Index)
-{
-	FLocalFogVolumeGPUInstanceData Data;
-	uint Offset = Index * 19;
-
-	Data.Transform		= float4x4(LocalFogVolumeInstances[Offset + 0], LocalFogVolumeInstances[Offset + 1], LocalFogVolumeInstances[Offset + 2], LocalFogVolumeInstances[Offset + 3]);
-	Data.InvTransform	= float4x4(LocalFogVolumeInstances[Offset + 4], LocalFogVolumeInstances[Offset + 5], LocalFogVolumeInstances[Offset + 6], LocalFogVolumeInstances[Offset + 7]);
-
-	Data.InvTranformNoScale = float4x4(LocalFogVolumeInstances[Offset + 8], LocalFogVolumeInstances[Offset + 9], LocalFogVolumeInstances[Offset + 10], LocalFogVolumeInstances[Offset + 11]);
-	Data.TransformScaleOnly = float4x4(LocalFogVolumeInstances[Offset + 12], LocalFogVolumeInstances[Offset + 13], LocalFogVolumeInstances[Offset + 14], LocalFogVolumeInstances[Offset + 15]);
-
-	float4 Temp0 = LocalFogVolumeInstances[Offset + 16];
-	Data.Density			= Temp0.x;
-	Data.HeightFalloff		= Temp0.y;
-	Data.HeightOffset		= Temp0.z;
-	Data.RadialAttenuation	= Temp0.w;
-
-	float4 Temp1 = LocalFogVolumeInstances[Offset + 17];
-	Data.Albedo		= Temp1.xyz;
-	Data.PhaseG		= Temp1.w;
-
-	float4 Temp2 = LocalFogVolumeInstances[Offset + 18];
-	Data.Emissive	= Temp2.xyz;
-	Data.FogMode	= Temp2.w;
-
-	return Data;
-}
-
 
 struct FLocalFogVolumeVertexOutput
 {
@@ -87,127 +36,9 @@ void LocalFogVolumeSplatVS(
 
 	OutPosition = mul(TranslatedWorldPosition, PrimaryView.TranslatedWorldToClip);
 	OtherVertexOutput.InstanceId = InstanceIndex;
-} 
-
-// CameraPos can be behind of RayStart if the camera is out of the volume
-struct FFogData
-{
-	float IntegratedLuminanceFactor;
-	float Coverage;
-};
-
-// Height fog is evaluated according to spatially varying density defined as "Density * Exp[-HeightFalloff * z]"
-FFogData EvaluateHeightFogIntegral(FLocalFogVolumeGPUInstanceData FogInstance,  float3 CameraPos, float3 RayStart, float3 RayDir, float RayLength, float RadialAttenuation)
-{
-	FFogData FogData;
-	
-	float3 RayOriU = mul(float4(RayStart, 1.0f),                      FogInstance.InvTransform).xyz;
-	float3 RayIntU = mul(float4(RayStart + RayDir * RayLength, 1.0f), FogInstance.InvTransform).xyz;
-	float3 RayDirU = mul(float4(RayDir  , 0.0f),                      FogInstance.InvTransform).xyz;
-	RayDirU = normalize(RayDirU);
-
-	float RayLengthU = length(RayIntU - RayOriU);
-
-	float StartHeight = RayOriU.z - FogInstance.HeightOffset; // accoutn for the fog offset.
-
-	float OpticalDepth = 0.0f;
-#if 0
-	// Original integral, less artefact
-	OpticalDepth = (FogInstance.Density / FogInstance.HeightFalloff) * exp(-StartHeight * FogInstance.HeightFalloff) * (1.0 - exp(-RayDirU.z * RayLengthU * FogInstance.HeightFalloff)) / RayDirU.z;
-#elif 1
-	// Reworked to avoid large value as input to exp
-	//OpticalDepth = (FogInstance.Density / FogInstance.HeightFalloff) * (exp(-StartHeight * FogInstance.HeightFalloff) - exp(-StartHeight * FogInstance.HeightFalloff - RayDir.z * RayLength * FogInstance.HeightFalloff)) / RayDir.z;
-
-	// But we use an updated formulation to avoid large input to exp leading to inf and nan.
-	float Factor0 = StartHeight * FogInstance.HeightFalloff;
-	float Factor1 = RayDirU.z * RayLengthU * FogInstance.HeightFalloff;
-	Factor0 = max(-80.0f, Factor0); // Simply clamp Factor0 to a valid range in order to not have exp explode.
-
-	OpticalDepth = (FogInstance.Density / (FogInstance.HeightFalloff * RayDirU.z)) * (exp(-Factor0) - exp(-(Factor0 + Factor1)));
-#endif
-
-	// Transmittance is the 
-	float Transmittance = exp(-OpticalDepth);
-	// Assuming extinction is grey scale, we can compute a single coverage value from transmittance
-	FogData.Coverage = (1.0 - Transmittance) * RadialAttenuation;
-	// Assuming extinction==scattering, i.e. albedo=1, we can compute the integrale for each point on a ray while accouting for transmittance back to the view.
-	// It turns out that, under these conditions, the luminance and emissive factor is simply coverage. => see Mathematica details at the bottom of this page.
-	FogData.IntegratedLuminanceFactor = FogData.Coverage;
-
-	return FogData;
-}
-
-
-FFogData EvaluateRadialFogIntegral(FLocalFogVolumeGPUInstanceData FogInstance, float3 CameraPos, float3 RayStart, float3 RayDir, float RayLength, float RadialAttenuation)
-{
-	FFogData FogData = (FFogData)0;
-
-	float3 RayOriU = mul(float4(RayStart, 1.0f),                      FogInstance.InvTransform).xyz;
-	float3 RayIntU = mul(float4(RayStart + RayDir * RayLength, 1.0f), FogInstance.InvTransform).xyz;
-	float3 RayDirU = mul(float4(RayDir  , 0.0f),                      FogInstance.InvTransform).xyz;
-	RayDirU = normalize(RayDirU);
-
-	float RayLengthU = length(RayIntU - RayOriU);
-
-	float3 SphereCenter = 0.0f;
-	float3 VolumeCenterToRayO = (RayOriU - SphereCenter);
-
-	float b = dot(RayDirU,            VolumeCenterToRayO);
-	float c = dot(VolumeCenterToRayO, VolumeCenterToRayO) - 1.0f;
-	float h = b * b - c;
-
-	float OpticalDepth = 0.0f;
-	if (h >= 0.0)
-	{
-		h = sqrt(h);
-		float Length0 = -b - h;
-		float Length1 = -b + h;
-
-		Length0 = max(Length0, 0.0);
-		Length1 = max(Length1, 0.0);
-		Length1 = min(Length1, RayLengthU);
-
-		// Integral of a density that descreases accordingf to square distance from the center of the volume.
-		float Integral0 = -(c * Length0 + b * Length0 * Length0 + Length0 * Length0 * Length0 / 3.0f);
-		float Integral1 = -(c * Length1 + b * Length1 * Length1 + Length1 * Length1 * Length1 / 3.0f);
-		OpticalDepth = max(0.0, FogInstance.Density * (Integral1 - Integral0) * (3.0f / 4.0f));
-	}
-
-	// Transmittance is the 
-	float Transmittance = exp(-OpticalDepth);
-	// Assuming extinction is grey scale, we can compute a single coverage value from transmittance
-	FogData.Coverage = (1.0 - Transmittance) * RadialAttenuation;
-	// Assuming extinction==scattering, i.e. albedo=1, we can compute the integrale for each point on a ray while accouting for transmittance back to the view.
-	// It turns out that, under these conditions, the luminance and emissive factor is simply coverage. => see Mathematica details at the bottom of this page.
-	FogData.IntegratedLuminanceFactor = FogData.Coverage;
-
-	return FogData;
 }
 
 
-float3 ComputeFogInscattering(in FLocalFogVolumeGPUInstanceData FogInstance, in FFogData FogData, in float3 RayDirWorld)
-{
-	half3 InScattering = 0;
-
-#if PROJECT_SUPPORT_SKY_ATMOSPHERE_AFFECTS_HEIGHFOG
-	InScattering +=  Texture2DSampleLevel(View.DistantSkyLightLutTexture, View.DistantSkyLightLutTextureSampler, float2(0.5f, 0.5f), 0.0f).rgb;
-
-	// No need to test View.AtmosphereLightIlluminanceOnGroundPostTransmittance[0].a because InscatteringLightDirection.w above is doing the same test already.
-	InScattering += View.AtmosphereLightIlluminanceOnGroundPostTransmittance[0].rgb * SchlickPhase(-FogInstance.PhaseG, dot(RayDirWorld, View.AtmosphereLightDirection[0].xyz));
-
-	if (View.AtmosphereLightIlluminanceOnGroundPostTransmittance[1].a > 0.0f) // Skip the second light when disabled.
-	{
-		InScattering += View.AtmosphereLightIlluminanceOnGroundPostTransmittance[0].rgb * SchlickPhase(-FogInstance.PhaseG, dot(RayDirWorld, View.AtmosphereLightDirection[1].xyz));
-	}
-#endif
-
-	InScattering *= FogInstance.Albedo * View.SkyAtmosphereHeightFogContribution.xxx;
-
-	// Now account for emissive luminance after albedo and heighfog contribution has been accounted for.
-	InScattering += FogInstance.Emissive;
-
-	return InScattering * FogData.IntegratedLuminanceFactor;
-}
 
 #define LHFV_FOLLOWS_ROTATIONS 1
 
@@ -292,21 +123,23 @@ void LocalFogVolumeSplatPS(
 		{
 			FFogData FogData = (FFogData)0;
 
+			LocalFogVolumeCommonLightingInput CommonLightingInput = LocalFogVolumeEvaluateCommonLightingInput();
+
 			if (FogInstance.FogMode == 0)
 			{
-				FogData = EvaluateHeightFogIntegral(FogInstance, CamRayWorldOrigin, CamRayWorldOrigin + CamRayWorldDir * Length0, CamRayWorldDir, RayTracedLength, RadialAttenuation);
+				FogData = LocalFogVolumeEvaluateHeightIntegral(FogInstance, CamRayWorldOrigin, CamRayWorldOrigin + CamRayWorldDir * Length0, CamRayWorldDir, RayTracedLength, RadialAttenuation);
 			}
 			else // (FogInstance.FogMode == 1)
 			{
 				// TODO
-				// TODO EvaluateRadialFogIntegral needs to be optimised because it is redoing lots of computations. Maybe move HeightFogIntegral to also be in a local normalised space.
+				// TODO LocalFogVolumeEvaluateRadialIntegral needs to be optimised because it is redoing lots of computations. Maybe move HeightFogIntegral to also be in a local normalised space.
 				// TODO It is like that for now while iterating with artists.
 				// TODO
-				FogData = EvaluateRadialFogIntegral(FogInstance, CamRayWorldOrigin, CamRayWorldOrigin + CamRayWorldDir * Length0, CamRayWorldDir, RayTracedLength, RadialAttenuation);
+				FogData = LocalFogVolumeEvaluateRadialIntegral(FogInstance, CamRayWorldOrigin, CamRayWorldOrigin + CamRayWorldDir * Length0, CamRayWorldDir, RayTracedLength, RadialAttenuation);
 			}
 
-			OutColor.xyz	= ComputeFogInscattering(FogInstance, FogData, CamRayWorldDir);	// = Luminance
-			OutColor.a		= 1.0 - FogData.Coverage;										// = Transmittance
+			OutColor.xyz	= LocalFogVolumeEvaluateInScattering(FogInstance, FogData, CommonLightingInput, CamRayWorldDir);	// = Luminance
+			OutColor.a		= 1.0 - FogData.Coverage;																			// = Transmittance
 		}
 		else
 		{
diff --git a/Engine/Shaders/Private/LocalFogVolumeCommon.ush b/Engine/Shaders/Private/LocalFogVolumeCommon.ush
new file mode 100644
--- /dev/null
+++ b/Engine/Shaders/Private/LocalFogVolumeCommon.ush
@@ -0,0 +1,225 @@
+// Copyright Epic Games, Inc. All Rights Reserved.
+#include "Common.ush"
+
+#include "ParticipatingMediaCommon.ush"
+
+
+
+///////////////////////////////////////////////////////////////////////////////
+// Local Fog Volume Data
+
+
+
+struct FLocalFogVolumeGPUInstanceData
+{
+	float4x4 Transform;
+	float4x4 InvTransform;
+
+	float4x4 InvTranformNoScale;
+	float4x4 TransformScaleOnly;
+
+	float Density;
+	float HeightFalloff;
+	float HeightOffset;
+	float RadialAttenuation;
+
+	float3 Albedo;
+	float  PhaseG;
+	float3 Emissive;
+	float  FogMode;
+};
+
+
+
+// Structured buffers cannot be sampled from the vertex shader on some mobile platforms.
+// That is why simple Buffers are used, because this tech needs to run on mobile platforms.
+Buffer<float4> LocalFogVolumeInstances;
+
+
+
+FLocalFogVolumeGPUInstanceData GetLocalFogVolumeGPUInstanceData(uint Index)
+{
+	FLocalFogVolumeGPUInstanceData Data;
+	uint Offset = Index * 19;
+
+	Data.Transform = float4x4(LocalFogVolumeInstances[Offset + 0], LocalFogVolumeInstances[Offset + 1], LocalFogVolumeInstances[Offset + 2], LocalFogVolumeInstances[Offset + 3]);
+	Data.InvTransform = float4x4(LocalFogVolumeInstances[Offset + 4], LocalFogVolumeInstances[Offset + 5], LocalFogVolumeInstances[Offset + 6], LocalFogVolumeInstances[Offset + 7]);
+
+	Data.InvTranformNoScale = float4x4(LocalFogVolumeInstances[Offset + 8], LocalFogVolumeInstances[Offset + 9], LocalFogVolumeInstances[Offset + 10], LocalFogVolumeInstances[Offset + 11]);
+	Data.TransformScaleOnly = float4x4(LocalFogVolumeInstances[Offset + 12], LocalFogVolumeInstances[Offset + 13], LocalFogVolumeInstances[Offset + 14], LocalFogVolumeInstances[Offset + 15]);
+
+	float4 Temp0 = LocalFogVolumeInstances[Offset + 16];
+	Data.Density = Temp0.x;
+	Data.HeightFalloff = Temp0.y;
+	Data.HeightOffset = Temp0.z;
+	Data.RadialAttenuation = Temp0.w;
+
+	float4 Temp1 = LocalFogVolumeInstances[Offset + 17];
+	Data.Albedo = Temp1.xyz;
+	Data.PhaseG = Temp1.w;
+
+	float4 Temp2 = LocalFogVolumeInstances[Offset + 18];
+	Data.Emissive = Temp2.xyz;
+	Data.FogMode = Temp2.w;
+
+	return Data;
+}
+
+
+
+///////////////////////////////////////////////////////////////////////////////
+// Local Fog Volume Integrals
+
+
+
+// CameraPos can be behind of RayStart if the camera is out of the volume
+struct FFogData
+{
+	float IntegratedLuminanceFactor;
+	float Coverage;
+};
+
+// Height fog is evaluated according to spatially varying density defined as "Density * Exp[-HeightFalloff * z]"
+FFogData LocalFogVolumeEvaluateHeightIntegral(FLocalFogVolumeGPUInstanceData FogInstance, float3 CameraPos, float3 RayStart, float3 RayDir, float RayLength, float RadialAttenuation)
+{
+	FFogData FogData;
+
+	float3 RayOriU = mul(float4(RayStart, 1.0f), FogInstance.InvTransform).xyz;
+	float3 RayIntU = mul(float4(RayStart + RayDir * RayLength, 1.0f), FogInstance.InvTransform).xyz;
+	float3 RayDirU = mul(float4(RayDir, 0.0f), FogInstance.InvTransform).xyz;
+	RayDirU = normalize(RayDirU);
+
+	float RayLengthU = length(RayIntU - RayOriU);
+
+	float StartHeight = RayOriU.z - FogInstance.HeightOffset; // accoutn for the fog offset.
+
+	float OpticalDepth = 0.0f;
+#if 0
+	// Original integral, less artefact
+	OpticalDepth = (FogInstance.Density / FogInstance.HeightFalloff) * exp(-StartHeight * FogInstance.HeightFalloff) * (1.0 - exp(-RayDirU.z * RayLengthU * FogInstance.HeightFalloff)) / RayDirU.z;
+#elif 1
+	// Reworked to avoid large value as input to exp
+	//OpticalDepth = (FogInstance.Density / FogInstance.HeightFalloff) * (exp(-StartHeight * FogInstance.HeightFalloff) - exp(-StartHeight * FogInstance.HeightFalloff - RayDir.z * RayLength * FogInstance.HeightFalloff)) / RayDir.z;
+
+	// But we use an updated formulation to avoid large input to exp leading to inf and nan.
+	float Factor0 = StartHeight * FogInstance.HeightFalloff;
+	float Factor1 = RayDirU.z * RayLengthU * FogInstance.HeightFalloff;
+	Factor0 = max(-80.0f, Factor0); // Simply clamp Factor0 to a valid range in order to not have exp explode.
+
+	OpticalDepth = (FogInstance.Density / (FogInstance.HeightFalloff * RayDirU.z)) * (exp(-Factor0) - exp(-(Factor0 + Factor1)));
+#endif
+
+	// Transmittance is the 
+	float Transmittance = exp(-OpticalDepth);
+	// Assuming extinction is grey scale, we can compute a single coverage value from transmittance
+	FogData.Coverage = (1.0 - Transmittance) * RadialAttenuation;
+	// Assuming extinction==scattering, i.e. albedo=1, we can compute the integrale for each point on a ray while accouting for transmittance back to the view.
+	// It turns out that, under these conditions, the luminance and emissive factor is simply coverage. => see Mathematica details at the bottom of this page.
+	FogData.IntegratedLuminanceFactor = FogData.Coverage;
+
+	return FogData;
+}
+
+
+FFogData LocalFogVolumeEvaluateRadialIntegral(FLocalFogVolumeGPUInstanceData FogInstance, float3 CameraPos, float3 RayStart, float3 RayDir, float RayLength, float RadialAttenuation)
+{
+	FFogData FogData = (FFogData)0;
+
+	float3 RayOriU = mul(float4(RayStart, 1.0f), FogInstance.InvTransform).xyz;
+	float3 RayIntU = mul(float4(RayStart + RayDir * RayLength, 1.0f), FogInstance.InvTransform).xyz;
+	float3 RayDirU = mul(float4(RayDir, 0.0f), FogInstance.InvTransform).xyz;
+	RayDirU = normalize(RayDirU);
+
+	float RayLengthU = length(RayIntU - RayOriU);
+
+	float3 SphereCenter = 0.0f;
+	float3 VolumeCenterToRayO = (RayOriU - SphereCenter);
+
+	float b = dot(RayDirU, VolumeCenterToRayO);
+	float c = dot(VolumeCenterToRayO, VolumeCenterToRayO) - 1.0f;
+	float h = b * b - c;
+
+	float OpticalDepth = 0.0f;
+	if (h >= 0.0)
+	{
+		h = sqrt(h);
+		float Length0 = -b - h;
+		float Length1 = -b + h;
+
+		Length0 = max(Length0, 0.0);
+		Length1 = max(Length1, 0.0);
+		Length1 = min(Length1, RayLengthU);
+
+		// Integral of a density that descreases accordingf to square distance from the center of the volume.
+		float Integral0 = -(c * Length0 + b * Length0 * Length0 + Length0 * Length0 * Length0 / 3.0f);
+		float Integral1 = -(c * Length1 + b * Length1 * Length1 + Length1 * Length1 * Length1 / 3.0f);
+		OpticalDepth = max(0.0, FogInstance.Density * (Integral1 - Integral0) * (3.0f / 4.0f));
+	}
+
+	// Transmittance is the 
+	float Transmittance = exp(-OpticalDepth);
+	// Assuming extinction is grey scale, we can compute a single coverage value from transmittance
+	FogData.Coverage = (1.0 - Transmittance) * RadialAttenuation;
+	// Assuming extinction==scattering, i.e. albedo=1, we can compute the integrale for each point on a ray while accouting for transmittance back to the view.
+	// It turns out that, under these conditions, the luminance and emissive factor is simply coverage. => see Mathematica details at the bottom of this page.
+	FogData.IntegratedLuminanceFactor = FogData.Coverage;
+
+	return FogData;
+}
+
+
+
+///////////////////////////////////////////////////////////////////////////////
+// Local Fog Volume Evaluation
+
+struct LocalFogVolumeCommonLightingInput
+{
+	float3 DistantSkyLight;
+};
+
+LocalFogVolumeCommonLightingInput LocalFogVolumeEvaluateCommonLightingInput()
+{
+	LocalFogVolumeCommonLightingInput CommonLightingInput = (LocalFogVolumeCommonLightingInput)0;
+#if PROJECT_SUPPORT_SKY_ATMOSPHERE_AFFECTS_HEIGHFOG
+	CommonLightingInput.DistantSkyLight = Texture2DSampleLevel(View.DistantSkyLightLutTexture, View.DistantSkyLightLutTextureSampler, float2(0.5f, 0.5f), 0.0f).rgb;
+#endif
+	return CommonLightingInput;
+}
+
+float3 LocalFogVolumeEvaluateInScattering(in FLocalFogVolumeGPUInstanceData FogInstance, in FFogData FogData, in LocalFogVolumeCommonLightingInput CommonLightingInput, in float3 RayDirWorld)
+{
+	half3 InScattering = 0;
+
+#if PROJECT_SUPPORT_SKY_ATMOSPHERE_AFFECTS_HEIGHFOG
+
+	// Accumulate distant sky light
+	InScattering += CommonLightingInput.DistantSkyLight;
+
+	// Accumulate directional lights if needed
+	if (View.AtmosphereLightIlluminanceOnGroundPostTransmittance[0].a > 0.0f)
+	{
+		InScattering += View.AtmosphereLightIlluminanceOnGroundPostTransmittance[0].rgb * SchlickPhase(-FogInstance.PhaseG, dot(RayDirWorld, View.AtmosphereLightDirection[0].xyz));
+	}
+	if (View.AtmosphereLightIlluminanceOnGroundPostTransmittance[1].a > 0.0f)
+	{
+		InScattering += View.AtmosphereLightIlluminanceOnGroundPostTransmittance[0].rgb * SchlickPhase(-FogInstance.PhaseG, dot(RayDirWorld, View.AtmosphereLightDirection[1].xyz));
+	}
+
+	// Weight by user control
+	InScattering *= View.SkyAtmosphereHeightFogContribution.xxx;
+
+	// Apply the medium uniform single scattering albedo
+	InScattering *= FogInstance.Albedo;
+#endif
+
+	// Now account for emissive luminance after albedo and heighfog contribution has been accounted for.
+	// This do mean that the emissive color follows the medium density variation.
+	InScattering += FogInstance.Emissive;
+
+	// Then scale the contribution according to the normalised pre intergrated scattering.
+	return InScattering * FogData.IntegratedLuminanceFactor;
+}
+
+
+
+
