Fix for spline mesh rendering bugs on mobile platforms with GPU Scene enabled.

diff --git a/Engine/Shaders/Private/LocalVertexFactory.ush b/Engine/Shaders/Private/LocalVertexFactory.ush
--- a/Engine/Shaders/Private/LocalVertexFactory.ush
+++ b/Engine/Shaders/Private/LocalVertexFactory.ush
@@ -352,14 +352,14 @@ FInstanceSceneData GetInstanceData(FVertexFactoryIntermediates Intermediates)
 
 #if USE_SPLINEDEFORM
 
-#if VF_USE_PRIMITIVE_SCENE_DATA
+#if VF_USE_PRIMITIVE_SCENE_DATA && FEATURE_LEVEL > FEATURE_LEVEL_ES3_1
 
 FSplineMeshShaderParams GetSplineMeshParams(FInstanceSceneData InstanceData)
 {
 	return SplineMeshLoadParamsFromInstancePayload(InstanceData);
 }
 
-#else // !VF_USE_PRIMITIVE_SCENE_DATA
+#else
 
 // These currently still need to be bound loosely when not using GPUScene
 float4 SplineParams[SPLINE_MESH_PARAMS_FLOAT4_SIZE];
@@ -368,7 +368,7 @@ FSplineMeshShaderParams GetSplineMeshParams(FInstanceSceneData InstanceData)
 	return UnpackSplineMeshParams(SplineParams);
 }
 
-#endif // VF_USE_PRIMITIVE_SCENE_DATA
+#endif // VF_USE_PRIMITIVE_SCENE_DATA && FEATURE_LEVEL > FEATURE_LEVEL_ES3_1
 
 #endif // USE_SPLINEDEFORM
 
diff --git a/Engine/Source/Runtime/Engine/Private/Components/SplineMeshComponent.cpp b/Engine/Source/Runtime/Engine/Private/Components/SplineMeshComponent.cpp
--- a/Engine/Source/Runtime/Engine/Private/Components/SplineMeshComponent.cpp
+++ b/Engine/Source/Runtime/Engine/Private/Components/SplineMeshComponent.cpp
@@ -193,8 +193,10 @@ void FSplineMeshVertexFactoryShaderParameters::GetElementShaderBindings(
 	{
 		ShaderBindings.Add(Shader->GetUniformBufferParameter<FLocalVertexFactoryUniformShaderParameters>(), LocalVertexFactory->GetUniformBuffer());
 	}
-
-	if (!bUseGPUScene)
+	
+	// If we can't use GPU Scene instance data, we have to bind the params to the VS loosely
+	// NOTE: Mobile GPU scene can't support loading the spline params from instance data in VS
+	if (!bUseGPUScene || FeatureLevel == ERHIFeatureLevel::ES3_1)
 	{
 		checkSlow(BatchElement.bIsSplineProxy);
 		const FSplineMeshSceneProxy* SplineProxy = BatchElement.SplineMeshSceneProxy;
