Fix multi-channel variance output error for path tracing when r.PathTracing.Denoiser.Prepass.VarianceType is set to 0.

diff --git a/Engine/Shaders/Private/PathTracing/PathTracingSpatialTemporalDenoising.usf b/Engine/Shaders/Private/PathTracing/PathTracingSpatialTemporalDenoising.usf
--- a/Engine/Shaders/Private/PathTracing/PathTracingSpatialTemporalDenoising.usf
+++ b/Engine/Shaders/Private/PathTracing/PathTracingSpatialTemporalDenoising.usf
@@ -38,9 +38,13 @@ PathTracingSpatialTemporalDenoising.usf: Spatial temporal denoising for path tra
 #define TEMPORAL_FEATURE_FUSION			0
 #endif
 
+// Variance Type
+#define RADIANCE_MULTI_CHANNEL					0
+#define RADIANCE_ALBEDO_NORMAL_SINGLECHANNEL	1
+
 #ifndef TEMPORAL_PREPASS
 #define TEMPORAL_PREPASS				0
-#define VARIANCE_TYPE					0
+#define VARIANCE_TYPE					RADIANCE_MULTI_CHANNEL
 #endif
 
 #define PREPASS_PHASE_INIT				0
@@ -788,9 +792,6 @@ Texture2D<float4> NormalTexture;
 RWStructuredBuffer<FPixelMaterialLightingFingerprint> RWVarianceMap;
 int Iteration;
 
-#define RADIANCE_MULTI_CHANNEL					0
-#define RADIANCE_ALBEDO_NORMAL_SINGLECHANNEL	1
-
 // Note that the radiance, albedo and normal are the mean value.
 float4 GetMeanFromTextures(int2 Position)
 {
