Improve base pass shader cost with Substrate with DXC forcing unrolling now shader compilation time is in better shape.

diff --git a/Engine/Shaders/Private/Substrate/Substrate.ush b/Engine/Shaders/Private/Substrate/Substrate.ush
--- a/Engine/Shaders/Private/Substrate/Substrate.ush
+++ b/Engine/Shaders/Private/Substrate/Substrate.ush
@@ -30,13 +30,8 @@
 // we redefine the unroll function and fall back to loop on DXC
 // To remove this, one must test PC + DX12 + RayTracing and make sure raytracing shaders compile.
 #define STRATA_COMPILER (COMPILER_DXC)
-#if STRATA_COMPILER == 1
-#define STRATA_UNROLL LOOP
-#define STRATA_UNROLL_N(X) LOOP
-#else
 #define STRATA_UNROLL UNROLL
 #define STRATA_UNROLL_N(X) UNROLL_N(X)
-#endif
 
 // During the base pass or forward rendering, shared local bases are simply available from registers.
 // If a BSDF is loaded for processing from the strata buffer, normals are unpacked on demand to reduce VGPR pressure and increase occupancy.
@@ -1117,12 +1112,8 @@ void StrataPackOutStrataTree(
 	OutCoord = 0;
 	OutCoordY = 1;
 
-#if STRATA_COMPILER == 1
-	LOOP
-#else
 	StrataTree.BSDFCount = min(StrataTree.BSDFCount, STRATA_CLAMPED_BSDF_COUNT);
 	STRATA_UNROLL_N(STRATA_CLAMPED_BSDF_COUNT)
-#endif
 	for(int i=0; i< StrataTree.BSDFCount; ++i)
 	{
 		WRITE_UINT(StrataTree.BSDFs[i].State);
@@ -1156,12 +1147,8 @@ void StrataPackOutStrataTree(
 	OutCoord = 0;
 	OutCoordY = 2;
 
-#if STRATA_COMPILER == 1
-	LOOP
-#else
 	StrataTree.OperatorCount = min(StrataTree.OperatorCount, STRATA_MAX_OPERATOR_COUNT);
 	STRATA_UNROLL_N(STRATA_MAX_OPERATOR_COUNT)
-#endif
 	for (uint j = 0; j < StrataTree.OperatorCount; ++j)
 	{
 		WRITE_INT(StrataTree.Operators[j].ParentIndex);
