Add initial binding-free support for ShaderPrint on platform supporting diagnostic buffer.

This allows to add shader print in any shader (i.e., material shaders like BasePass, or global shaders) without any ShaderPrintData binding.

This is experimental for now and for enabling it, you need to:
* Set SHADER_PRINT_EXPERIMENTAL_DIAGNOSTIC_BUFFER to 1 in ShaderPrint.ush.
* Set r.D3D12.DiagnosticBufferExtraMemory to a large enough value for ensuring Diagnostic buffer contains enough memory to store the ShaderPrintData. This is temporary until high level rendering code can drive DiagnosticBuffer allocation.

diff --git a/Engine/Shaders/Private/ShaderPrint.ush b/Engine/Shaders/Private/ShaderPrint.ush
--- a/Engine/Shaders/Private/ShaderPrint.ush
+++ b/Engine/Shaders/Private/ShaderPrint.ush
@@ -14,30 +14,99 @@
 #define SHADER_PRINT_USE_GLOBAL_RESOURCE 0
 #endif
 
+// Shader print data can be stored with 3 differents methods:
+// * ShaderPrintData binding: Shader print data are explicity setup on the CPP side, and the ShaderPrintData buffer are used for storing the data
+// * Custom binding         : Shader print data are explicity setup on the CPP side, and the ShaderPrintData buffer are stored into a custom buffer)
+// * Diagnostic buffer      : Shader print data & buffer are stored into the Diagnostic buffer, which is always bound and does not require any explicit binding. This is only supported on certain platforms	
+
+// Experimental feature, disabled by default
+#define SHADER_PRINT_EXPERIMENTAL_DIAGNOSTIC_BUFFER 0
+
+// On platform supporting diagnostic buffer, shader print data are stored into the diagnostic buffer which is always 
+// bound. This allows to access shader print functionality in any shader without explicit shader binding
+#define SHADER_PRINT_USE_DIAGNOSTIC_BUFFER (PLATFORM_SUPPORTS_DIAGNOSTIC_BUFFER && !SHADER_PRINT_USE_GLOBAL_RESOURCE && SHADER_PRINT_EXPERIMENTAL_DIAGNOSTIC_BUFFER)
+// Offset to the ShaderPrint parameters
+#define SHADER_PRINT_DIAGNOSTIC_BUFFER_PARAMETERS_OFFSET 6u
+// Offset to the ShaderPrint entries
+#define SHADER_PRINT_DIAGNOSTIC_BUFFER_ENTRY_OFFSET (SHADER_PRINT_DIAGNOSTIC_BUFFER_PARAMETERS_OFFSET + 12u)
+
 // Global buffers for ShaderPrint collection.
-StructuredBuffer<uint> ShaderPrint_StateBuffer;
-RWStructuredBuffer<uint> ShaderPrint_RWEntryBuffer;
+#if SHADER_PRINT_USE_DIAGNOSTIC_BUFFER
+	#define SHADER_PRINT_RWENTRYBUFFER(Ctx, Idx) UEDiagnosticBuffer[SHADER_PRINT_DIAGNOSTIC_BUFFER_ENTRY_OFFSET + (Idx)]
+	#define SHADER_PRINT_RWSTATEBUFFER(Ctx, Idx) UEDiagnosticBuffer[SHADER_PRINT_DIAGNOSTIC_BUFFER_ENTRY_OFFSET + (Idx)]
+#else
+	StructuredBuffer<uint> ShaderPrint_StateBuffer;
+	RWStructuredBuffer<uint> ShaderPrint_RWEntryBuffer;
 
-#if SHADER_PRINT_USE_GLOBAL_RESOURCE
-// Use global buffers directly instead of using FShaderPrintContext::Buffers.
-#define SHADER_PRINT_RWENTRYBUFFER(Ctx) ShaderPrint_RWEntryBuffer
-#define SHADER_PRINT_RWSTATEBUFFER(Ctx) ShaderPrint_StateBuffer
+	// Use global buffers directly instead of using FShaderPrintContext::Buffers.
+	#if SHADER_PRINT_USE_GLOBAL_RESOURCE
+	#define SHADER_PRINT_RWENTRYBUFFER(Ctx, Idx) ShaderPrint_RWEntryBuffer[Idx]
+	#define SHADER_PRINT_RWSTATEBUFFER(Ctx, Idx) ShaderPrint_StateBuffer[Idx]
+	#endif
 #endif
 
 #include "ShaderPrintCommon.ush"
 #include "MiniFontCommon.ush"
+#include "PackUnpack.ush"
+
+#if SHADER_PRINT_USE_DIAGNOSTIC_BUFFER
+void PackShaderPrintConfig(FShaderPrintConfig In)
+{
+	const uint Offset = SHADER_PRINT_DIAGNOSTIC_BUFFER_PARAMETERS_OFFSET;
+
+	UEDiagnosticBuffer[Offset+0] = PackUInt2ToUInt(In.Resolution.x, In.Resolution.y);
+	UEDiagnosticBuffer[Offset+1] = PackFloat2ToUInt(In.CursorCoord);
+	UEDiagnosticBuffer[Offset+2] = asuint(In.TranslatedWorldOffset.x); 
+	UEDiagnosticBuffer[Offset+3] = asuint(In.TranslatedWorldOffset.y); 
+	UEDiagnosticBuffer[Offset+4] = asuint(In.TranslatedWorldOffset.z);
+	UEDiagnosticBuffer[Offset+5] = PackFloat2ToUInt(In.FontSize);
+	UEDiagnosticBuffer[Offset+6] = PackFloat2ToUInt(In.FontSpacing);
+	UEDiagnosticBuffer[Offset+7] = In.MaxValueCount;
+	UEDiagnosticBuffer[Offset+8] = In.MaxSymbolCount;
+	UEDiagnosticBuffer[Offset+9] = In.MaxStateCount;
+	UEDiagnosticBuffer[Offset+10]= In.MaxLineCount;
+	UEDiagnosticBuffer[Offset+11]= In.MaxTriangleCount;
+}
+
+FShaderPrintConfig UnpackShaderPrintConfig()
+{
+	const uint Offset = SHADER_PRINT_DIAGNOSTIC_BUFFER_PARAMETERS_OFFSET;
+
+	FShaderPrintConfig Out;
+	Out.Resolution 				= UnpackUInt2FromUInt(UEDiagnosticBuffer[Offset+0]);
+	Out.CursorCoord 			= UnpackFloat2FromUInt(UEDiagnosticBuffer[Offset+1]);
+	Out.TranslatedWorldOffset.x = asfloat(UEDiagnosticBuffer[Offset+2]);
+	Out.TranslatedWorldOffset.y = asfloat(UEDiagnosticBuffer[Offset+3]);
+	Out.TranslatedWorldOffset.z = asfloat(UEDiagnosticBuffer[Offset+4]);
+	Out.FontSize 				= UnpackFloat2FromUInt(UEDiagnosticBuffer[Offset+5]);
+	Out.FontSpacing 			= UnpackFloat2FromUInt(UEDiagnosticBuffer[Offset+6]);
+	Out.MaxValueCount			= UEDiagnosticBuffer[Offset+7];
+	Out.MaxSymbolCount			= UEDiagnosticBuffer[Offset+8];
+	Out.MaxStateCount			= UEDiagnosticBuffer[Offset+9];
+	Out.MaxLineCount			= UEDiagnosticBuffer[Offset+10];
+	Out.MaxTriangleCount		= UEDiagnosticBuffer[Offset+11];
+	return Out;
+}
+#endif
 
 FShaderPrintBuffers InitShaderPrintContextBuffers()
 {
 	FShaderPrintBuffers Buffers;
+#if SHADER_PRINT_USE_DIAGNOSTIC_BUFFER
+	Buffers.RWEntryBuffer = UEDiagnosticBuffer;
+#else
 	Buffers.StateBuffer = ShaderPrint_StateBuffer;
 	Buffers.RWEntryBuffer = ShaderPrint_RWEntryBuffer;
+#endif
 	return Buffers;
 }
 
 // Initialise FShaderPrintConfig from the ShaderPrintData uniform buffer.
 FShaderPrintConfig InitShaderPrintContextConfig()
 {
+#if SHADER_PRINT_USE_DIAGNOSTIC_BUFFER
+	return UnpackShaderPrintConfig();
+#else
 	FShaderPrintConfig Config;
 	Config.Resolution = ShaderPrintData.Resolution;
 	Config.CursorCoord = ShaderPrintData.CursorCoord;
@@ -50,25 +119,34 @@ FShaderPrintConfig InitShaderPrintContextConfig()
 	Config.MaxLineCount = ShaderPrintData.MaxLineCount;
 	Config.MaxTriangleCount = ShaderPrintData.MaxTriangleCount;
 	return Config;
+#endif
 }
 
 // Initialise FShaderPrintContext.
 FShaderPrintContext InitShaderPrintContext(bool bActive, float2 InStartPos)
 {
-	FShaderPrintContext Context;
-	Context.bIsActive = bActive;
-	Context.StartPos = InStartPos;
-	Context.Pos = InStartPos;
-	Context.Config = InitShaderPrintContextConfig();
-#if !SHADER_PRINT_USE_GLOBAL_RESOURCE
-	Context.Buffers = InitShaderPrintContextBuffers();
-#endif
-	return Context;
+	FShaderPrintContext Out;
+	Out.bIsActive = bActive;
+	Out.StartPos = InStartPos;
+	Out.Pos = InStartPos;
+	Out.Config = InitShaderPrintContextConfig();
+	#if !SHADER_PRINT_USE_GLOBAL_RESOURCE
+	Out.Buffers = InitShaderPrintContextBuffers();
+	#endif
+	return Out;
 }
 
 FShaderPrintContext InitShaderPrintContext(bool bActive, uint2 InStartCoord) 
 {
-	return InitShaderPrintContext(bActive, float2(InStartCoord) / float2(ShaderPrintData.Resolution));
+	FShaderPrintConfig Config = InitShaderPrintContextConfig();
+	return InitShaderPrintContext(bActive, float2(InStartCoord) / float2(Config.Resolution));
+}
+
+FShaderPrintContext InitShaderPrintContextAtCursor(uint2 InCurrentCoord, uint2 InStartCoord) 
+{
+	FShaderPrintConfig Config = InitShaderPrintContextConfig();
+	const bool bActive = all(InCurrentCoord == Config.CursorCoord);
+	return InitShaderPrintContext(bActive, float2(InStartCoord) / float2(Config.Resolution));
 }
 
 FShaderPrintContext InitShaderPrintContext()
diff --git a/Engine/Shaders/Private/ShaderPrintCommon.ush b/Engine/Shaders/Private/ShaderPrintCommon.ush
--- a/Engine/Shaders/Private/ShaderPrintCommon.ush
+++ b/Engine/Shaders/Private/ShaderPrintCommon.ush
@@ -111,7 +111,7 @@
 #define SHADER_PRINT_COUNTER_OFFSET_FREE     3
 #define SHADER_PRINT_COUNTER_COUNT           4
 
-uint GetSymbolOffset(uint InIndex)				
+uint GetSymbolOffset(uint InIndex)
 {
 	return SHADER_PRINT_COUNTER_COUNT + SHADER_PRINT_UINT_COUNT_SYMBOL * InIndex; 
 }
@@ -148,13 +148,10 @@ void ClearCounters(RWStructuredBuffer<uint> InRWBuffer)
 
 ///////////////////////////////////////////////////////////////////////////////////////////////////
 // Buffers/Resources accessors
-#ifndef SHADER_PRINT_USE_GLOBAL_RESOURCE
-#define SHADER_PRINT_USE_GLOBAL_RESOURCE 0
-#endif
 
-#if !SHADER_PRINT_USE_GLOBAL_RESOURCE
-#define SHADER_PRINT_RWENTRYBUFFER(Ctx) Ctx.Buffers.RWEntryBuffer
-#define SHADER_PRINT_RWSTATEBUFFER(Ctx) Ctx.Buffers.StateBuffer
+#if !defined(SHADER_PRINT_RWENTRYBUFFER)
+#define SHADER_PRINT_RWENTRYBUFFER(Ctx, Idx) Ctx.Buffers.RWEntryBuffer[Idx]
+#define SHADER_PRINT_RWSTATEBUFFER(Ctx, Idx) Ctx.Buffers.StateBuffer[Idx]
 #endif
 
 ///////////////////////////////////////////////////////////////////////////////////////////////////
@@ -264,6 +261,7 @@ FShaderPrintItem UnpackShaderPrintItem(FPackedShaderPrintItem In)
 	return Out;
 }
 
+// /!\ If you update this function, update the WriteSymbol() version without explicit buffer, below in this file
 void WriteSymbol(uint Offset, FShaderPrintItem In, RWStructuredBuffer<uint> InRWBuffer)
 {
 	const uint Offset4 = GetSymbolOffset(Offset);
@@ -325,6 +323,18 @@ struct FShaderPrintContext
 #endif
 };
 
+// Version using SHADER_PRINT_RWENTRYBUFFER in order to provide offset value
+void WriteSymbol(uint Offset, FShaderPrintItem In, FShaderPrintContext Ctx) 
+{
+	const uint Offset4 = GetSymbolOffset(Offset);
+
+	const FPackedShaderPrintItem Packed = PackShaderPrintItem(In);
+	SHADER_PRINT_RWENTRYBUFFER(Ctx, Offset4 + 0) = Packed.ScreenPos16bits;
+	SHADER_PRINT_RWENTRYBUFFER(Ctx, Offset4 + 1) = Packed.Value;
+	SHADER_PRINT_RWENTRYBUFFER(Ctx, Offset4 + 2) = Packed.TypeAndColor;
+	SHADER_PRINT_RWENTRYBUFFER(Ctx, Offset4 + 3) = Packed.Metadata;
+}
+
 ///////////////////////////////////////////////////////////////////////////////////////////////////
 // Internal helpers
 
@@ -338,7 +348,7 @@ void ShaderPrint_Internal(FShaderPrintContext Ctx, in FShaderPrintItem Item)
 
 	// Buffer counter is stored in first element .Value
 	int IndexToStore = 0;
-	InterlockedAdd(SHADER_PRINT_RWENTRYBUFFER(Ctx)[SHADER_PRINT_COUNTER_OFFSET_SYMBOL], 1, IndexToStore);
+	InterlockedAdd(SHADER_PRINT_RWENTRYBUFFER(Ctx,SHADER_PRINT_COUNTER_OFFSET_SYMBOL), 1, IndexToStore);
 
 	// Prevent writing off the buffer
 	// Note that counter still increases so need clamp when reading it in later passes
@@ -347,7 +357,7 @@ void ShaderPrint_Internal(FShaderPrintContext Ctx, in FShaderPrintItem Item)
 		return;
 	}
 
-	WriteSymbol(IndexToStore, Item, SHADER_PRINT_RWENTRYBUFFER(Ctx));
+	WriteSymbol(IndexToStore, Item, Ctx);
 }
 
 void ShaderPrint_Internal(FShaderPrintContext Ctx, in float2 ScreenPos, in int Value, in FFontColor FontColor, in uint Metadata, in int Type)
@@ -422,7 +432,7 @@ float2 ShaderPrintText_Internal(FShaderPrintContext Ctx, bool bIsActive, float2
 		// Prevent writing off the buffer
 		// Note that counter still increases so need clamp when reading it in later passes
 		int IndexToStore = 0;
-		InterlockedAdd(SHADER_PRINT_RWENTRYBUFFER(Ctx)[SHADER_PRINT_COUNTER_OFFSET_SYMBOL], Count, IndexToStore);
+		InterlockedAdd(SHADER_PRINT_RWENTRYBUFFER(Ctx,SHADER_PRINT_COUNTER_OFFSET_SYMBOL), Count, IndexToStore);
 		if (uint(IndexToStore + Count) < Ctx.Config.MaxValueCount)
 		{
 			for (uint i = Begin; i < End; ++i)
@@ -434,7 +444,7 @@ float2 ShaderPrintText_Internal(FShaderPrintContext Ctx, bool bIsActive, float2
 				Item.Color = InColor.Color;
 				Item.Metadata = 0u;
 
-				WriteSymbol(IndexToStore, Item, SHADER_PRINT_RWENTRYBUFFER(Ctx));
+				WriteSymbol(IndexToStore, Item, Ctx);
 
 				++IndexToStore;
 				InPos.x += Ctx.Config.FontSpacing.x;
@@ -551,15 +561,15 @@ bool ShaderPrintGetStateValue(FShaderPrintContext Ctx, in uint InHash, inout uin
 	OutValue = 0;
 
 	// Slow linear search among all widget states
-	const uint MaxCount = min(SHADER_PRINT_RWSTATEBUFFER(Ctx)[SHADER_PRINT_STATE_COUNT_OFFSET], Ctx.Config.MaxStateCount);
+	const uint MaxCount = min(SHADER_PRINT_RWSTATEBUFFER(Ctx, SHADER_PRINT_STATE_COUNT_OFFSET), Ctx.Config.MaxStateCount);
 	for (uint Index = 0; Index < MaxCount; ++Index)
 	{
 		const uint Index3 = Index * SHADER_PRINT_STATE_STRIDE + SHADER_PRINT_STATE_VALUE_OFFSET;
-		const FShaderPrintMetadata Metadata = ShaderPrintUnpackMetadata(SHADER_PRINT_RWSTATEBUFFER(Ctx)[Index3 + SHADER_PRINT_STATE_INDEX_METADATA]);
+		const FShaderPrintMetadata Metadata = ShaderPrintUnpackMetadata(SHADER_PRINT_RWSTATEBUFFER(Ctx, Index3 + SHADER_PRINT_STATE_INDEX_METADATA));
 		if (Metadata.Hash == InHash)
 		{
 			OutIndex = Metadata.Index;
-			OutValue = SHADER_PRINT_RWSTATEBUFFER(Ctx)[Index3 + SHADER_PRINT_STATE_INDEX_VALUE];
+			OutValue = SHADER_PRINT_RWSTATEBUFFER(Ctx, Index3 + SHADER_PRINT_STATE_INDEX_VALUE);
 			return true;
 		}
 	}
@@ -766,7 +776,7 @@ bool AllocateTriangleElement(FShaderPrintContext Ctx, uint Count, inout uint Out
 		return false;
 	}
 
-	InterlockedAdd(SHADER_PRINT_RWENTRYBUFFER(Ctx)[SHADER_PRINT_COUNTER_OFFSET_TRIANGLE], Count, OutIndex);
+	InterlockedAdd(SHADER_PRINT_RWENTRYBUFFER(Ctx,SHADER_PRINT_COUNTER_OFFSET_TRIANGLE), Count, OutIndex);
 	return (OutIndex + Count) < Ctx.Config.MaxTriangleCount;
 }
 
@@ -775,20 +785,20 @@ void AddTriangleElement(FShaderPrintContext Ctx, FTriangleElement In, uint Index
 	const uint Index12 = GetPrimitiveTriangleOffset(Index, Ctx.Config.MaxValueCount, Ctx.Config.MaxLineCount);
 
 	FPackedTriangleElement Out = PackTriangleElement(In);
-	SHADER_PRINT_RWENTRYBUFFER(Ctx)[Index12 + 0] = Out.Packed0.x;
-	SHADER_PRINT_RWENTRYBUFFER(Ctx)[Index12 + 1] = Out.Packed0.y;
-	SHADER_PRINT_RWENTRYBUFFER(Ctx)[Index12 + 2] = Out.Packed0.z;
-	SHADER_PRINT_RWENTRYBUFFER(Ctx)[Index12 + 3] = Out.Packed0.w;
-
-	SHADER_PRINT_RWENTRYBUFFER(Ctx)[Index12 + 4] = Out.Packed1.x;
-	SHADER_PRINT_RWENTRYBUFFER(Ctx)[Index12 + 5] = Out.Packed1.y;
-	SHADER_PRINT_RWENTRYBUFFER(Ctx)[Index12 + 6] = Out.Packed1.z;
-	SHADER_PRINT_RWENTRYBUFFER(Ctx)[Index12 + 7] = Out.Packed1.w;
-
-	SHADER_PRINT_RWENTRYBUFFER(Ctx)[Index12 + 8] = Out.Packed2.x;
-	SHADER_PRINT_RWENTRYBUFFER(Ctx)[Index12 + 9] = Out.Packed2.y;
-	SHADER_PRINT_RWENTRYBUFFER(Ctx)[Index12 + 10] = Out.Packed2.z;
-	SHADER_PRINT_RWENTRYBUFFER(Ctx)[Index12 + 11] = Out.Packed2.w;
+	SHADER_PRINT_RWENTRYBUFFER(Ctx, Index12 + 0) = Out.Packed0.x;
+	SHADER_PRINT_RWENTRYBUFFER(Ctx, Index12 + 1) = Out.Packed0.y;
+	SHADER_PRINT_RWENTRYBUFFER(Ctx, Index12 + 2) = Out.Packed0.z;
+	SHADER_PRINT_RWENTRYBUFFER(Ctx, Index12 + 3) = Out.Packed0.w;
+
+	SHADER_PRINT_RWENTRYBUFFER(Ctx, Index12 + 4) = Out.Packed1.x;
+	SHADER_PRINT_RWENTRYBUFFER(Ctx, Index12 + 5) = Out.Packed1.y;
+	SHADER_PRINT_RWENTRYBUFFER(Ctx, Index12 + 6) = Out.Packed1.z;
+	SHADER_PRINT_RWENTRYBUFFER(Ctx, Index12 + 7) = Out.Packed1.w;
+
+	SHADER_PRINT_RWENTRYBUFFER(Ctx, Index12 + 8) = Out.Packed2.x;
+	SHADER_PRINT_RWENTRYBUFFER(Ctx, Index12 + 9) = Out.Packed2.y;
+	SHADER_PRINT_RWENTRYBUFFER(Ctx, Index12 + 10) = Out.Packed2.z;
+	SHADER_PRINT_RWENTRYBUFFER(Ctx, Index12 + 11) = Out.Packed2.w;
 }
 
 ///////////////////////////////////////////////////////////////////////////////////////////////////
@@ -965,7 +975,7 @@ bool AllocateLineElement(FShaderPrintContext Ctx, uint Count, inout uint OutInde
 		return false;
 	}
 
-	InterlockedAdd(SHADER_PRINT_RWENTRYBUFFER(Ctx)[SHADER_PRINT_COUNTER_OFFSET_LINE], Count, OutIndex);
+	InterlockedAdd(SHADER_PRINT_RWENTRYBUFFER(Ctx,SHADER_PRINT_COUNTER_OFFSET_LINE), Count, OutIndex);
 	return (OutIndex + Count) < Ctx.Config.MaxLineCount;
 }
 
@@ -974,15 +984,15 @@ void AddLineElement(FShaderPrintContext Ctx, FLineElement In, uint Index)
 	const uint Index8 = GetPrimitiveLineOffset(Index, Ctx.Config.MaxValueCount);
 
 	FPackedLineElement Out = PackLineElement(In);
-	SHADER_PRINT_RWENTRYBUFFER(Ctx)[Index8 + 0] = Out.Packed0.x;
-	SHADER_PRINT_RWENTRYBUFFER(Ctx)[Index8 + 1] = Out.Packed0.y;
-	SHADER_PRINT_RWENTRYBUFFER(Ctx)[Index8 + 2] = Out.Packed0.z;
-	SHADER_PRINT_RWENTRYBUFFER(Ctx)[Index8 + 3] = Out.Packed0.w;
-
-	SHADER_PRINT_RWENTRYBUFFER(Ctx)[Index8 + 4] = Out.Packed1.x;
-	SHADER_PRINT_RWENTRYBUFFER(Ctx)[Index8 + 5] = Out.Packed1.y;
-	SHADER_PRINT_RWENTRYBUFFER(Ctx)[Index8 + 6] = Out.Packed1.z;
-	SHADER_PRINT_RWENTRYBUFFER(Ctx)[Index8 + 7] = Out.Packed1.w;
+	SHADER_PRINT_RWENTRYBUFFER(Ctx,Index8 + 0) = Out.Packed0.x;
+	SHADER_PRINT_RWENTRYBUFFER(Ctx,Index8 + 1) = Out.Packed0.y;
+	SHADER_PRINT_RWENTRYBUFFER(Ctx,Index8 + 2) = Out.Packed0.z;
+	SHADER_PRINT_RWENTRYBUFFER(Ctx,Index8 + 3) = Out.Packed0.w;
+
+	SHADER_PRINT_RWENTRYBUFFER(Ctx,Index8 + 4) = Out.Packed1.x;
+	SHADER_PRINT_RWENTRYBUFFER(Ctx,Index8 + 5) = Out.Packed1.y;
+	SHADER_PRINT_RWENTRYBUFFER(Ctx,Index8 + 6) = Out.Packed1.z;
+	SHADER_PRINT_RWENTRYBUFFER(Ctx,Index8 + 7) = Out.Packed1.w;
 }
 
 ///////////////////////////////////////////////////////////////////////////////////////////////////
diff --git a/Engine/Shaders/Private/ShaderPrintDraw.usf b/Engine/Shaders/Private/ShaderPrintDraw.usf
--- a/Engine/Shaders/Private/ShaderPrintDraw.usf
+++ b/Engine/Shaders/Private/ShaderPrintDraw.usf
@@ -6,7 +6,7 @@
 
 #include "Common.ush"	
 #include "MiniFontCommon.ush"
-#include "ShaderPrintCommon.ush"
+#include "ShaderPrint.ush"
 
 uint2 ViewportResolution;
 uint FrameIndex;
@@ -361,7 +361,7 @@ void DrawSymbolsPS(
 
 // --------------------------------------------------------------------------
 // Zoomed output
-
+#if SHADER_ZOOM
 Texture2D<float4> InTexture;
 RWTexture2D<float4> OutTexture;
 
@@ -411,4 +411,124 @@ void DrawZoomCS(int2 DispatchThreadId : SV_DispatchThreadID)
 
 		}
 	}
-}
\ No newline at end of file
+}
+#endif // SHADER_ZOOM
+
+// --------------------------------------------------------------------------
+// Upload parameters
+#if SHADER_UPLOAD
+
+// Copy ShaderPrintData uniform parameter within UEDiagnosticBuffer
+[numthreads(1, 1, 1)]
+void UploadCS(int2 DispatchThreadId : SV_DispatchThreadID)
+{	
+#if SHADER_PRINT_USE_DIAGNOSTIC_BUFFER
+	if (any(DispatchThreadId != 0))
+	{
+		return;
+	}
+
+	FShaderPrintConfig In;
+	In.Resolution = ShaderPrintData.Resolution;
+	In.CursorCoord = ShaderPrintData.CursorCoord;
+	In.TranslatedWorldOffset = ShaderPrintData.TranslatedWorldOffset;
+	In.FontSize = ShaderPrintData.FontSize;
+	In.FontSpacing = ShaderPrintData.FontSpacing;
+	In.MaxValueCount = ShaderPrintData.MaxValueCount;
+	In.MaxSymbolCount = ShaderPrintData.MaxSymbolCount;
+	In.MaxStateCount = ShaderPrintData.MaxStateCount;
+	In.MaxLineCount = ShaderPrintData.MaxLineCount;
+	In.MaxTriangleCount = ShaderPrintData.MaxTriangleCount;
+	PackShaderPrintConfig(In);
+
+	// Clear counters
+	FShaderPrintContext Ctx; // Dummy context
+	SHADER_PRINT_RWENTRYBUFFER(Ctx, SHADER_PRINT_COUNTER_OFFSET_LINE) = 0;
+	SHADER_PRINT_RWENTRYBUFFER(Ctx, SHADER_PRINT_COUNTER_OFFSET_TRIANGLE) = 0;
+	SHADER_PRINT_RWENTRYBUFFER(Ctx, SHADER_PRINT_COUNTER_OFFSET_SYMBOL) = 0;
+	SHADER_PRINT_RWENTRYBUFFER(Ctx, SHADER_PRINT_COUNTER_OFFSET_FREE) = 0;
+
+#endif
+}
+#endif //SHADER_UPLOAD
+
+// --------------------------------------------------------------------------
+// Copy data prior to draw
+#if SHADER_COPY
+
+#define TYPE_LINE 0
+#define TYPE_TRIANGLE 1
+#define TYPE_SYMBOL 2
+
+#if SHADER_PRINT_USE_DIAGNOSTIC_BUFFER
+void CopyDataType(FShaderPrintConfig Config, uint Type)
+{
+	uint CounterOffset = 0;
+	uint DataTypeSizeInUint = 0;
+	uint DataOffset0 = 0;
+
+	if (Type == TYPE_SYMBOL)
+	{
+		CounterOffset = SHADER_PRINT_COUNTER_OFFSET_SYMBOL;
+		DataTypeSizeInUint = SHADER_PRINT_UINT_COUNT_SYMBOL;
+		DataOffset0 = GetSymbolOffset(0);
+	}
+	else if (Type == TYPE_LINE)
+	{
+		CounterOffset = SHADER_PRINT_COUNTER_OFFSET_LINE; 		
+		DataTypeSizeInUint = SHADER_PRINT_UINT_COUNT_LINE;
+		DataOffset0 = GetPrimitiveLineOffset(0, Config.MaxValueCount);
+	}
+	else if (Type == TYPE_TRIANGLE)
+	{
+		CounterOffset = SHADER_PRINT_COUNTER_OFFSET_TRIANGLE;	
+		DataTypeSizeInUint = SHADER_PRINT_UINT_COUNT_TRIANGLE;
+		DataOffset0 = GetPrimitiveTriangleOffset(0, Config.MaxValueCount, Config.MaxLineCount);	
+	}
+	else
+	{
+		return;
+	}
+
+	FShaderPrintContext Ctx; // Dummy context
+	const uint SrcCount = SHADER_PRINT_RWENTRYBUFFER(Ctx, CounterOffset);
+	const uint DstCount = RWValuesBuffer[CounterOffset];
+
+	// Compute the offset of the first item:
+	// * For the src buffer (i.e. diagonistoc buffer), it is the offset to the first element
+	// * For the dst buffer (i.e. shader print rdg buffer), it is the offset to the last element 
+	//   which was happens through explicit binding
+	const uint SrcOffsetInUint = DataOffset0 + SHADER_PRINT_DIAGNOSTIC_BUFFER_ENTRY_OFFSET;
+	const uint DstOffsetInUint = DataOffset0 + DstCount * DataTypeSizeInUint;
+
+	// Copy data + counter
+	// TODO: make it parallel
+	const uint CountInUint = SrcCount * DataTypeSizeInUint;
+	for (uint It = 0; It < CountInUint; ++It)
+	{
+		const uint SrcIndex = SrcOffsetInUint + It;
+		const uint DstIndex = DstOffsetInUint + It;
+		RWValuesBuffer[DstIndex] = UEDiagnosticBuffer[SrcIndex];
+	}
+	RWValuesBuffer[CounterOffset] = SrcCount + DstCount;
+}
+#endif
+
+// Copy/merge data from UEDiagnosticBuffer into ShaderPrint RDG buffer for display
+[numthreads(1, 1, 1)]
+void CopyCS(int2 DispatchThreadId : SV_DispatchThreadID)
+{	
+#if SHADER_PRINT_USE_DIAGNOSTIC_BUFFER
+	if (any(DispatchThreadId != 0))
+	{
+		return;
+	}
+
+	FShaderPrintConfig Config = UnpackShaderPrintConfig();
+
+	CopyDataType(Config, TYPE_LINE);
+	CopyDataType(Config, TYPE_TRIANGLE);
+	CopyDataType(Config, TYPE_SYMBOL);
+#endif
+}
+#endif // SHADER_COPY
diff --git a/Engine/Source/Runtime/Renderer/Private/ShaderPrint.cpp b/Engine/Source/Runtime/Renderer/Private/ShaderPrint.cpp
--- a/Engine/Source/Runtime/Renderer/Private/ShaderPrint.cpp
+++ b/Engine/Source/Runtime/Renderer/Private/ShaderPrint.cpp
@@ -302,6 +302,60 @@ namespace ShaderPrint
 		PRAGMA_ENABLE_DEPRECATION_WARNINGS
 	}
 
+	//////////////////////////////////////////////////////////////////////////////////////////////////
+	// Common Shaders
+
+	// Upload ShaderPrint parmeters into diagnostic buffer
+	class FShaderPrintUploadCS: public FGlobalShader
+	{
+	public:
+		DECLARE_GLOBAL_SHADER(FShaderPrintUploadCS);
+		SHADER_USE_PARAMETER_STRUCT(FShaderPrintUploadCS, FGlobalShader);
+
+		BEGIN_SHADER_PARAMETER_STRUCT(FParameters, )
+			SHADER_PARAMETER_STRUCT_REF(FShaderPrintCommonParameters, Common)
+		END_SHADER_PARAMETER_STRUCT()
+
+
+		static bool ShouldCompilePermutation(FGlobalShaderPermutationParameters const& Parameters)
+		{
+			return IsSupported(Parameters.Platform);
+		}
+		static void ModifyCompilationEnvironment(const FGlobalShaderPermutationParameters& Parameters, FShaderCompilerEnvironment& OutEnvironment)
+		{
+			FGlobalShader::ModifyCompilationEnvironment(Parameters, OutEnvironment);
+			OutEnvironment.SetDefine(TEXT("SHADER_UPLOAD"), 1);
+		}
+	};
+
+	IMPLEMENT_GLOBAL_SHADER(FShaderPrintUploadCS, "/Engine/Private/ShaderPrintDraw.usf", "UploadCS", SF_Compute);
+
+	
+	// Upload ShaderPrint parmeters into diagnostic buffer
+	class FShaderPrintCopyCS: public FGlobalShader
+	{
+	public:
+		DECLARE_GLOBAL_SHADER(FShaderPrintCopyCS);
+		SHADER_USE_PARAMETER_STRUCT(FShaderPrintCopyCS, FGlobalShader);
+
+		BEGIN_SHADER_PARAMETER_STRUCT(FParameters, )
+			SHADER_PARAMETER_STRUCT_REF(FShaderPrintCommonParameters, Common)
+			SHADER_PARAMETER_RDG_BUFFER_UAV(RWStructuredBuffer<uint>, RWValuesBuffer)
+		END_SHADER_PARAMETER_STRUCT()
+
+		static bool ShouldCompilePermutation(FGlobalShaderPermutationParameters const& Parameters)
+		{
+			return IsSupported(Parameters.Platform);
+		}
+		static void ModifyCompilationEnvironment(const FGlobalShaderPermutationParameters& Parameters, FShaderCompilerEnvironment& OutEnvironment)
+		{
+			FGlobalShader::ModifyCompilationEnvironment(Parameters, OutEnvironment);
+			OutEnvironment.SetDefine(TEXT("SHADER_COPY"), 1);
+		}
+	};
+
+	IMPLEMENT_GLOBAL_SHADER(FShaderPrintCopyCS, "/Engine/Private/ShaderPrintDraw.usf", "CopyCS", SF_Compute);
+
 	//////////////////////////////////////////////////////////////////////////////////////////////////
 	// Widget/Characters Shaders
 	 
@@ -586,6 +640,11 @@ namespace ShaderPrint
 		{
 			return IsSupported(Parameters.Platform);
 		}
+		static void ModifyCompilationEnvironment(const FGlobalShaderPermutationParameters& Parameters, FShaderCompilerEnvironment& OutEnvironment)
+		{
+			FGlobalShader::ModifyCompilationEnvironment(Parameters, OutEnvironment);
+			OutEnvironment.SetDefine(TEXT("SHADER_ZOOM"), 1);
+		}
 	};
 
 	IMPLEMENT_GLOBAL_SHADER(FShaderZoomCS, "/Engine/Private/ShaderPrintDraw.usf", "DrawZoomCS", SF_Compute);
@@ -626,6 +685,16 @@ namespace ShaderPrint
 		bIsDrawLocked = View.State ? ((const FSceneViewState*)View.State)->ShaderPrintStateData.bIsLocked : false;
 	}
 
+	static uint32 GetRequestedEntryBufferSizeInUint(const FShaderPrintSetup& In)
+	{
+		const uint32 UintElementCount = 
+			GetCountersUintSize() + 
+			GetPackedSymbolUintSize() * In.MaxValueCount +
+			GetPackedLineUintSize() * In.MaxLineCount + 
+			GetPackedTriangleUintSize() * In.MaxTriangleCount;
+		return UintElementCount;
+	}
+
 	FShaderPrintData CreateShaderPrintData(FRDGBuilder& GraphBuilder, FShaderPrintSetup const& InSetup, FSceneViewState* InViewState)
 	{
 		FShaderPrintData ShaderPrintData;
@@ -649,11 +718,7 @@ namespace ShaderPrint
 			const bool bLockBufferThisFrame = IsDrawLocked() && InViewState != nullptr && !InViewState->ShaderPrintStateData.bIsLocked;
 			ERDGBufferFlags Flags = bLockBufferThisFrame ? ERDGBufferFlags::MultiFrame : ERDGBufferFlags::None;
 
-			const uint32 UintElementCount = 
-				GetCountersUintSize() + 
-				GetPackedSymbolUintSize() * InSetup.MaxValueCount +
-				GetPackedLineUintSize() * InSetup.MaxLineCount + 
-				GetPackedTriangleUintSize() * InSetup.MaxTriangleCount;
+			const uint32 UintElementCount = GetRequestedEntryBufferSizeInUint(InSetup);
 			ShaderPrintData.ShaderPrintEntryBuffer = GraphBuilder.CreateBuffer(FRDGBufferDesc::CreateStructuredDesc(4, UintElementCount), TEXT("ShaderPrint.EntryBuffer"), Flags);
 
 			// State buffer is retrieved from the view state, or created if it does not exist
@@ -748,6 +813,33 @@ namespace ShaderPrint
 	//////////////////////////////////////////////////////////////////////////////////////////////////
 	// Drawing/Rendering API
 
+	static void InternalUploadParameters(FRDGBuilder& GraphBuilder, const FShaderPrintData& ShaderPrintData)
+	{
+		FGlobalShaderMap* GlobalShaderMap = GetGlobalShaderMap(GMaxRHIFeatureLevel);
+		TShaderMapRef<FShaderPrintUploadCS> ComputeShader(GlobalShaderMap);
+		FShaderPrintUploadCS::FParameters* PassParameters = GraphBuilder.AllocParameters<FShaderPrintUploadCS::FParameters>();
+		PassParameters->Common = ShaderPrintData.UniformBuffer;
+		FComputeShaderUtils::AddPass(
+			GraphBuilder,
+			RDG_EVENT_NAME("ShaderPrint::UploadParameters"),
+			ERDGPassFlags::Compute | ERDGPassFlags::NeverCull,
+			ComputeShader, PassParameters, FIntVector(1,1,1));
+	}
+
+	static void InternalCopyParameters(FRDGBuilder& GraphBuilder, const FShaderPrintData& ShaderPrintData)
+	{
+		FGlobalShaderMap* GlobalShaderMap = GetGlobalShaderMap(GMaxRHIFeatureLevel);
+		TShaderMapRef<FShaderPrintCopyCS> ComputeShader(GlobalShaderMap);
+		FShaderPrintCopyCS::FParameters* PassParameters = GraphBuilder.AllocParameters<FShaderPrintCopyCS::FParameters>();
+		PassParameters->Common = ShaderPrintData.UniformBuffer;
+		PassParameters->RWValuesBuffer = GraphBuilder.CreateUAV(ShaderPrintData.ShaderPrintEntryBuffer);;
+		FComputeShaderUtils::AddPass(
+			GraphBuilder,
+			RDG_EVENT_NAME("ShaderPrint::CopyParameters"),
+			ERDGPassFlags::Compute | ERDGPassFlags::NeverCull,
+			ComputeShader, PassParameters, FIntVector(1,1,1));
+	}
+
 	void BeginView(FRDGBuilder& GraphBuilder, FViewInfo& View)
 	{
 		TRACE_CPUPROFILER_EVENT_SCOPE(ShaderPrint::BeginView);
@@ -761,6 +853,9 @@ namespace ShaderPrint
 		// Create the render data and store on the view.
 		FShaderPrintSetup ShaderPrintSetup(View);
 		View.ShaderPrintData = CreateShaderPrintData(GraphBuilder, ShaderPrintSetup, View.ViewState);
+
+		// Upload/Copy ShaderPrint parameters into UEDiagnostic buffer
+		InternalUploadParameters(GraphBuilder, View.ShaderPrintData);
 	}
 
 	void BeginViews(FRDGBuilder& GraphBuilder, TArrayView<FViewInfo> Views)
@@ -1059,6 +1154,11 @@ namespace ShaderPrint
 		const FIntRect OutputViewRect = OutputTexture.ViewRect;
 		const FVector PreViewTranslation = View.ViewMatrices.GetPreViewTranslation() - ShaderPrintData.Setup.PreViewTranslation;
 
+		// Copy/merge data (when using UEDiagnostic buffer as storage)
+		{
+			InternalCopyParameters(GraphBuilder, View.ShaderPrintData);
+		}
+
 		// Lines
 		{
 			FRDGBufferRef DataBuffer = ShaderPrintData.ShaderPrintEntryBuffer;
