Fix for post process issues with world vector node if position is replace or after tonemap stage.

diff --git a/Engine/Shaders/Private/MaterialTemplate.ush b/Engine/Shaders/Private/MaterialTemplate.ush
--- a/Engine/Shaders/Private/MaterialTemplate.ush
+++ b/Engine/Shaders/Private/MaterialTemplate.ush
@@ -3662,7 +3662,11 @@ void CalcMaterialParametersPost(
 	FIsFrontFace bIsFrontFace)
 {
 	float4 ScreenPosition = SvPositionToScreenPosition(SvPosition);
+#if POST_PROCESS_MATERIAL_BEFORE_TONEMAP	// In this case, SVPositionToTranslatedWorld account for the correct dynamic resolution scaling.
 	float3 TranslatedWorldPosition = SvPositionToTranslatedWorld(SvPosition);
+#else										// Otherwise, we must update SVPosition to account for the resolution scaling not happening after at the tonemap stage and after.
+	float3 TranslatedWorldPosition = SvPositionToTranslatedWorld(SvPosition * float4(View.ResolutionFractionAndInv.xx, 1.0f, 1.0f));
+#endif
 
 	CalcMaterialParametersEx(Parameters, PixelMaterialInputs, SvPosition, ScreenPosition, bIsFrontFace, TranslatedWorldPosition, TranslatedWorldPosition);
 }
