Fix incorrect UV calculation causing bilateral filter to generate incorrect results.

diff --git a/Engine/Shaders/Private/ReflectionEnvironmentPixelShader.usf b/Engine/Shaders/Private/ReflectionEnvironmentPixelShader.usf
--- a/Engine/Shaders/Private/ReflectionEnvironmentPixelShader.usf
+++ b/Engine/Shaders/Private/ReflectionEnvironmentPixelShader.usf
@@ -368,7 +368,7 @@ void ReflectionEnvironmentSkyLighting(
 	float3 BentNormal = GBuffer.WorldNormal;
 #if APPLY_SKY_SHADOWING
 	{
-		const float2 BentNormalUV = (PixelPos - View.ViewRectMin.xy) * View.BufferSizeAndInvSize.zw;
+		const float2 BentNormalUV = (SvPosition.xy - View.ViewRectMin.xy) * View.BufferSizeAndInvSize.zw;
 		BentNormal = UpsampleDFAO(BentNormalUV, GBuffer.Depth, GBuffer.WorldNormal);
 	}
 #endif
