Fix fuzz cancelling SSS on material having a single slab.

diff --git a/Engine/Shaders/Private/Substrate/SubstrateExport.ush b/Engine/Shaders/Private/Substrate/SubstrateExport.ush
--- a/Engine/Shaders/Private/Substrate/SubstrateExport.ush
+++ b/Engine/Shaders/Private/Substrate/SubstrateExport.ush
@@ -359,7 +359,7 @@ void PackSubstrateOut(
 				{
 					OptimisedLegacyMode = SINGLE_OPTLEGACYMODE_CLEARCOAT;
 				}
-				else if ((OptimisedModeMask & SingleHasCloth) == SingleHasCloth)
+				else if ((OptimisedModeMask & SingleHasCloth) == SingleHasCloth && BSDF_GETSSSTYPE(BSDF) == SSS_TYPE_INVALID)
 				{
 					OptimisedLegacyMode = SINGLE_OPTLEGACYMODE_CLOTH;
 				}				
