Fix incorrect break, causing complex materials with fuzz to have incorrect transmittance.

diff --git a/Engine/Shaders/Private/Substrate/SubstrateEvaluation.ush b/Engine/Shaders/Private/Substrate/SubstrateEvaluation.ush
--- a/Engine/Shaders/Private/Substrate/SubstrateEvaluation.ush
+++ b/Engine/Shaders/Private/Substrate/SubstrateEvaluation.ush
@@ -772,7 +772,6 @@ FSubstrateEvaluateResult SubstrateEvaluateBSDFCommon(FSubstrateBSDFContext BSDFC
 
 				Sample.ThroughputV *= Cloth_DirectionalAlbedo_SpecularTransmission;
 				// The specular transmission is ignored towards the direction of the light. (TransmittanceAlongN can only store transmittance)
-				break;
 			}
 			#endif
 
