Fixed Substrate forward shading compilation fix.

diff --git a/Engine/Shaders/Private/Substrate/SubstrateForwardLighting.ush b/Engine/Shaders/Private/Substrate/SubstrateForwardLighting.ush
--- a/Engine/Shaders/Private/Substrate/SubstrateForwardLighting.ush
+++ b/Engine/Shaders/Private/Substrate/SubstrateForwardLighting.ush
@@ -275,7 +275,7 @@ float3 SubstrateForwardLighting(
 					#if FORWARD_SHADING && (MATERIALBLENDING_SOLID || MATERIALBLENDING_MASKED)
 						float2 NearestResolvedDepthScreenUV = CalculateNearestResolvedDepthScreenUV(ScreenUV, MaterialParameters.ScreenPosition.w);
 
-						IndirectOcclusion = GetIndirectOcclusion(NearestResolvedDepthScreenUV, HasDynamicIndirectShadowCasterRepresentation(SubstratePixelHeader));
+						IndirectOcclusion = GetIndirectOcclusion(NearestResolvedDepthScreenUV, SubstratePixelHeader.HasDynamicIndirectShadowCasterRepresentation());
 						DiffuseIndirectLighting *= IndirectOcclusion;
 						SubsurfaceIndirectLighting *= IndirectOcclusion;
 						IndirectIrradiance *= IndirectOcclusion;
diff --git a/Engine/Shaders/Private/Substrate/SubstrateMobileForwardLighting.ush b/Engine/Shaders/Private/Substrate/SubstrateMobileForwardLighting.ush
--- a/Engine/Shaders/Private/Substrate/SubstrateMobileForwardLighting.ush
+++ b/Engine/Shaders/Private/Substrate/SubstrateMobileForwardLighting.ush
@@ -158,7 +158,7 @@ float3 SubstrateMobileForwardLighting(
 					#if 0 && FORWARD_SHADING && (MATERIALBLENDING_SOLID || MATERIALBLENDING_MASKED)
 						float2 NearestResolvedDepthScreenUV = CalculateNearestResolvedDepthScreenUV(ScreenUV, MaterialParameters.ScreenPosition.w);
 
-						IndirectOcclusion = GetIndirectOcclusion(NearestResolvedDepthScreenUV, HasDynamicIndirectShadowCasterRepresentation(SubstratePixelHeader));
+						IndirectOcclusion = GetIndirectOcclusion(NearestResolvedDepthScreenUV, SubstratePixelHeader.HasDynamicIndirectShadowCasterRepresentation());
 						DiffuseIndirectLighting *= IndirectOcclusion;
 						SubsurfaceIndirectLighting *= IndirectOcclusion;
 						IndirectIrradiance *= IndirectOcclusion;
