LFV - made the instance evaluation shared.

#prelfight shader

diff --git a/Engine/Shaders/Private/LocalFogVolume.usf b/Engine/Shaders/Private/LocalFogVolume.usf
--- a/Engine/Shaders/Private/LocalFogVolume.usf
+++ b/Engine/Shaders/Private/LocalFogVolume.usf
@@ -6,8 +6,9 @@
 #endif
 
 #include "SceneTexturesCommon.ush"
-#include "LocalFogVolumeCommon.ush"
 
+#define LFV_CAN_CLIP 1
+#include "LocalFogVolumeCommon.ush"
 
 
 struct FLocalFogVolumeVertexOutput
@@ -15,6 +16,7 @@ struct FLocalFogVolumeVertexOutput
 	nointerpolation uint InstanceId : TEXCOORD0;
 };
 
+
 void LocalFogVolumeSplatVS(
 	in float4 InPosition : ATTRIBUTE0,
 	uint InstanceIndex : SV_InstanceID,
@@ -39,9 +41,6 @@ void LocalFogVolumeSplatVS(
 }
 
 
-
-#define LHFV_FOLLOWS_ROTATIONS 1
-
 void LocalFogVolumeSplatPS(
 	in float4 SVPos : SV_POSITION,
 	in FLocalFogVolumeVertexOutput OtherVertexOutput,
@@ -53,105 +52,29 @@ void LocalFogVolumeSplatPS(
 	float3 CamRayTranslatedWorldOrigin	= PrimaryView.TranslatedWorldCameraOrigin;
 	float3 CamRayTranslatedWorldDir = normalize(SvPositionToTranslatedWorld(float4(SVPos.xy, 0.5, 1.0)));
 	float3 CamRayWorldDir = CamRayTranslatedWorldDir;
-
-	FLocalFogVolumeGPUInstanceData FogInstance = GetLocalFogVolumeGPUInstanceData(InstanceIndex);
-
-	OutColor = float4(0, 0, 0, 1);
-
 	float3 CamRayWorldOrigin = CamRayTranslatedWorldOrigin - LWCHackToFloat(PrimaryView.PreViewTranslation);
 
-	// The "U" prefix is for all compuation done in the Unit Sphere space.
-	float3 RayPosU = mul(float4(CamRayWorldOrigin,        1.0f), FogInstance.InvTransform).xyz;
-	float3 RayDirU = mul(float4(CamRayTranslatedWorldDir, 0.0f), FogInstance.InvTransform).xyz; // Try to use InverseTranspose to correctly handle direction transform with scaling.
-	RayDirU.xyz = normalize(RayDirU);
-
-	float2 TsU = RayIntersectSphere(RayPosU, RayDirU, float4(0.0, 0.0, 0.0, 1.0));
-	if (any(TsU > 0.0))
-	{
 #if SHADING_PATH_MOBILE
-		// LookupDeviceZ with uint2 SVPos.xy as input does not work on mobile.
-		float2 UV = SVPos * View.BufferSizeAndInvSize.zw;
-		const float DeviceZ = LookupDeviceZ(UV);
+	// LookupDeviceZ with uint2 SVPos.xy as input does not work on mobile.
+	float2 UV = SVPos * View.BufferSizeAndInvSize.zw;
+	const float DeviceZ = LookupDeviceZ(UV);
 #else
-		const float DeviceZ = LookupDeviceZ(uint2(SVPos.xy));
+	const float DeviceZ = LookupDeviceZ(uint2(SVPos.xy));
 #endif
 
-		float3 DepthBufferTranslatedWorldPos = SvPositionToTranslatedWorld(float4(SVPos.xy, DeviceZ, 1.0));
-		float LengthD = length(DepthBufferTranslatedWorldPos - CamRayTranslatedWorldOrigin);
-
-		float3 P0U = RayPosU + max(0.0, TsU.x) * RayDirU;
-		float3 P1U = RayPosU + max(0.0, TsU.y) * RayDirU;
-#if LHFV_FOLLOWS_ROTATIONS
-		// Computations are all done transformed with Space = local fog volume but ignoring scaling.
-		float3 P0Space = mul(float4(P0U, 1.0), FogInstance.TransformScaleOnly).xyz;
-		float3 P1Space = mul(float4(P1U, 1.0), FogInstance.TransformScaleOnly).xyz;
-		float3 CamRayOriginSpace = mul(float4(CamRayWorldOrigin, 1.0), FogInstance.InvTranformNoScale).xyz;
-		float3 CamRayDirSpace    = mul(float4(CamRayWorldDir,    0.0), FogInstance.InvTranformNoScale).xyz;	// No scale so no need for a special InverseTranspose matrix.
-		CamRayDirSpace = normalize(CamRayDirSpace);
-#else
-		// Computations are all done with Space = world space
-		float3 P0Space = mul(float4(P0U, 1.0), FogInstance.Transform).xyz;
-		float3 P1Space = mul(float4(P1U, 1.0), FogInstance.Transform).xyz;
-		float3 CamRayOriginSpace = CamRayWorldOrigin;
-		float3 CamRayDirSpace    = CamRayWorldDir;
-#endif
+	float3 DepthBufferTranslatedWorldPos = SvPositionToTranslatedWorld(float4(SVPos.xy, DeviceZ, 1.0));
+
+	float4 LuminanceColorTransmittance = GetLFVInstanceContribution(
+		InstanceIndex,
+		CamRayWorldOrigin,
+		CamRayWorldDir,
+		CamRayTranslatedWorldOrigin,
+		CamRayTranslatedWorldDir,
+		DepthBufferTranslatedWorldPos);
 
-		float Length0 = length(P0Space - CamRayOriginSpace);
-		float Length1 = length(P1Space - CamRayOriginSpace);
-
-		// Radial attenuation based on the distance travelled in the unit sphere and ignoring depth buffer intersection. 
-		// This is only to soften sphere edge.
-		float RadialAttenuation = saturate(abs(max(0.0, TsU.y) - max(0.0, TsU.x)) / FogInstance.RadialAttenuation);
-		// The attenuation is linear but it can be made soft using an Hermit interpolation by controling the tangents.
-		{
-			const float t = RadialAttenuation;
-			const float t2 = t * t;
-			const float t3 = t2 * t;
-			const float HP0 = 0.0f;
-			const float HP1 = 1.0f;
-			const float HM0 = 0.0f;	// tangent = 0 for soft transition
-			const float HM1 = 0.0f; // idem
-			RadialAttenuation = saturate((2 * t3 - 3 * t2 + 1) * HP0 + (t3 - 2 * t2 + t) * HM0 + (-2 * t3 + 3 * t2) * HP1 + (t3 - t2) * HM1);
-		}
-
-		// Clamp the traced distance to depth and compute world length.
-		Length0 = min(LengthD, Length0);
-		Length1 = min(LengthD, Length1);
-		float RayTracedLength = max(0.0, abs(Length1 - Length0));
-
-		if (RayTracedLength > 0.0)
-		{
-			FFogData FogData = (FFogData)0;
-
-			LocalFogVolumeCommonLightingInput CommonLightingInput = LocalFogVolumeEvaluateCommonLightingInput();
-
-			if (FogInstance.FogMode == 0)
-			{
-				FogData = LocalFogVolumeEvaluateHeightIntegral(FogInstance, CamRayWorldOrigin, CamRayWorldOrigin + CamRayWorldDir * Length0, CamRayWorldDir, RayTracedLength, RadialAttenuation);
-			}
-			else // (FogInstance.FogMode == 1)
-			{
-				// TODO
-				// TODO LocalFogVolumeEvaluateRadialIntegral needs to be optimised because it is redoing lots of computations. Maybe move HeightFogIntegral to also be in a local normalised space.
-				// TODO It is like that for now while iterating with artists.
-				// TODO
-				FogData = LocalFogVolumeEvaluateRadialIntegral(FogInstance, CamRayWorldOrigin, CamRayWorldOrigin + CamRayWorldDir * Length0, CamRayWorldDir, RayTracedLength, RadialAttenuation);
-			}
-
-			OutColor.xyz	= LocalFogVolumeEvaluateInScattering(FogInstance, FogData, CommonLightingInput, CamRayWorldDir);	// = Luminance
-			OutColor.a		= 1.0 - FogData.Coverage;																			// = Transmittance
-		}
-		else
-		{
-			clip(-1.0);
-		}
-
-		OutColor.xyz *= PrimaryView.PreExposure;
-	}
-	else
-	{
-		clip(-1.0);
-	}
+	OutColor = float4(
+		LuminanceColorTransmittance.rgb * PrimaryView.PreExposure,
+		LuminanceColorTransmittance.a);
 }
 
 
diff --git a/Engine/Shaders/Private/LocalFogVolumeCommon.ush b/Engine/Shaders/Private/LocalFogVolumeCommon.ush
--- a/Engine/Shaders/Private/LocalFogVolumeCommon.ush
+++ b/Engine/Shaders/Private/LocalFogVolumeCommon.ush
@@ -225,4 +225,118 @@ float3 LocalFogVolumeEvaluateInScattering(in FLocalFogVolumeGPUInstanceData FogI
 
 
 
+///////////////////////////////////////////////////////////////////////////////
+// Local Fog Volume Instance Evaluation
+
+#ifndef LFV_CAN_CLIP
+#define LFV_CAN_CLIP 0
+#endif
+
+#define LHFV_FOLLOWS_ROTATIONS 1
+
+// Returns color transmittance
+float4 GetLFVInstanceContribution(
+	uint InstanceIndex,
+	float3 CamRayWorldOrigin,
+	float3 CamRayWorldDir,
+	float3 CamRayTranslatedWorldOrigin,
+	float3 CamRayTranslatedWorldDir,
+	float3 DepthBufferTranslatedWorldPos)
+{
+	FLocalFogVolumeGPUInstanceData FogInstance = GetLocalFogVolumeGPUInstanceData(InstanceIndex);
+
+	float3 LuminanceColor = float3(0.0f, 0.0f, 0.0f);
+	float Transmittance = 1.0f;
+
+	// The "U" prefix is for all compuation done in the Unit Sphere space.
+	float3 RayPosU = mul(float4(CamRayWorldOrigin, 1.0f), FogInstance.InvTransform).xyz;
+	float3 RayDirU = mul(float4(CamRayWorldDir,    0.0f), FogInstance.InvTransform).xyz; // Try to use InverseTranspose to correctly handle direction transform with scaling.
+	RayDirU.xyz = normalize(RayDirU);
+
+	float2 TsU = RayIntersectSphere(RayPosU, RayDirU, float4(0.0, 0.0, 0.0, 1.0));
+	if (any(TsU > 0.0))
+	{
+		float LengthD = length(DepthBufferTranslatedWorldPos - CamRayTranslatedWorldOrigin);
+
+		float3 P0U = RayPosU + max(0.0, TsU.x) * RayDirU;
+		float3 P1U = RayPosU + max(0.0, TsU.y) * RayDirU;
+#if LHFV_FOLLOWS_ROTATIONS
+		// Computations are all done transformed with Space = local fog volume but ignoring scaling.
+		float3 P0Space = mul(float4(P0U, 1.0), FogInstance.TransformScaleOnly).xyz;
+		float3 P1Space = mul(float4(P1U, 1.0), FogInstance.TransformScaleOnly).xyz;
+		float3 CamRayOriginSpace = mul(float4(CamRayWorldOrigin, 1.0), FogInstance.InvTranformNoScale).xyz;
+		float3 CamRayDirSpace = mul(float4(CamRayWorldDir, 0.0), FogInstance.InvTranformNoScale).xyz;	// No scale so no need for a special InverseTranspose matrix.
+		CamRayDirSpace = normalize(CamRayDirSpace);
+#else
+		// Computations are all done with Space = world space
+		float3 P0Space = mul(float4(P0U, 1.0), FogInstance.Transform).xyz;
+		float3 P1Space = mul(float4(P1U, 1.0), FogInstance.Transform).xyz;
+		float3 CamRayOriginSpace = CamRayWorldOrigin;
+		float3 CamRayDirSpace = CamRayWorldDir;
+#endif
+
+		float Length0 = length(P0Space - CamRayOriginSpace);
+		float Length1 = length(P1Space - CamRayOriginSpace);
+
+		// Radial attenuation based on the distance travelled in the unit sphere and ignoring depth buffer intersection. 
+		// This is only to soften sphere edge.
+		float RadialAttenuation = saturate(abs(max(0.0, TsU.y) - max(0.0, TsU.x)) / FogInstance.RadialAttenuation);
+		// The attenuation is linear but it can be made soft using an Hermit interpolation by controling the tangents.
+		{
+			const float t = RadialAttenuation;
+			const float t2 = t * t;
+			const float t3 = t2 * t;
+			const float HP0 = 0.0f;
+			const float HP1 = 1.0f;
+			const float HM0 = 0.0f;	// tangent = 0 for soft transition
+			const float HM1 = 0.0f; // idem
+			RadialAttenuation = saturate((2 * t3 - 3 * t2 + 1) * HP0 + (t3 - 2 * t2 + t) * HM0 + (-2 * t3 + 3 * t2) * HP1 + (t3 - t2) * HM1);
+		}
+
+		// Clamp the traced distance to depth and compute world length.
+		Length0 = min(LengthD, Length0);
+		Length1 = min(LengthD, Length1);
+		float RayTracedLength = max(0.0, abs(Length1 - Length0));
+
+		if (RayTracedLength > 0.0)
+		{
+			FFogData FogData = (FFogData)0;
+
+			LocalFogVolumeCommonLightingInput CommonLightingInput = LocalFogVolumeEvaluateCommonLightingInput();
+
+			if (FogInstance.FogMode == 0)
+			{
+				FogData = LocalFogVolumeEvaluateHeightIntegral(FogInstance, CamRayWorldOrigin, CamRayWorldOrigin + CamRayWorldDir * Length0, CamRayWorldDir, RayTracedLength, RadialAttenuation);
+			}
+			else // (FogInstance.FogMode == 1)
+			{
+				// TODO
+				// TODO LocalFogVolumeEvaluateRadialIntegral needs to be optimised because it is redoing lots of computations. Maybe move HeightFogIntegral to also be in a local normalised space.
+				// TODO It is like that for now while iterating with artists.
+				// TODO
+				FogData = LocalFogVolumeEvaluateRadialIntegral(FogInstance, CamRayWorldOrigin, CamRayWorldOrigin + CamRayWorldDir * Length0, CamRayWorldDir, RayTracedLength, RadialAttenuation);
+			}
+
+			LuminanceColor = LocalFogVolumeEvaluateInScattering(FogInstance, FogData, CommonLightingInput, CamRayWorldDir);
+			Transmittance = 1.0 - FogData.Coverage;
+		}
+		else
+		{
+#if LFV_CAN_CLIP
+			clip(-1.0);
+#endif
+		}
+	}
+	else
+	{
+#if LFV_CAN_CLIP
+		clip(-1.0);
+#endif
+	}
+
+	return float4(LuminanceColor, Transmittance);
+}
+
+
+
 
