Fixed camera vector and world position nodes when using dynamic resolution scaling on a post process material setup as BL_SSRInput.
Follow up to 27582164

diff --git a/Engine/Plugins/Experimental/RemoteSession/Source/RemoteSession/Private/Channels/RemoteSessionARCameraChannel.cpp b/Engine/Plugins/Experimental/RemoteSession/Source/RemoteSession/Private/Channels/RemoteSessionARCameraChannel.cpp
--- a/Engine/Plugins/Experimental/RemoteSession/Source/RemoteSession/Private/Channels/RemoteSessionARCameraChannel.cpp
+++ b/Engine/Plugins/Experimental/RemoteSession/Source/RemoteSession/Private/Channels/RemoteSessionARCameraChannel.cpp
@@ -76,6 +76,7 @@ public:
 		FMaterialShader::ModifyCompilationEnvironment(Parameters, OutEnvironment);
 		OutEnvironment.SetDefine(TEXT("POST_PROCESS_MATERIAL"), 1);
 		OutEnvironment.SetDefine(TEXT("POST_PROCESS_MATERIAL_BEFORE_TONEMAP"), (Parameters.MaterialParameters.BlendableLocation != BL_AfterTonemapping) ? 1 : 0);
+		OutEnvironment.SetDefine(TEXT("POST_PROCESS_MATERIAL_SSRINPUT"), (Parameters.MaterialParameters.BlendableLocation == BL_SSRInput) ? 1 : 0);
 	}
 };
 
diff --git a/Engine/Plugins/Runtime/AR/AppleAR/AppleARKit/Source/AppleARKit/Private/AppleARKitVideoOverlay.cpp b/Engine/Plugins/Runtime/AR/AppleAR/AppleARKit/Source/AppleARKit/Private/AppleARKitVideoOverlay.cpp
--- a/Engine/Plugins/Runtime/AR/AppleAR/AppleARKit/Source/AppleARKit/Private/AppleARKitVideoOverlay.cpp
+++ b/Engine/Plugins/Runtime/AR/AppleAR/AppleARKit/Source/AppleARKit/Private/AppleARKitVideoOverlay.cpp
@@ -154,6 +154,7 @@ public:
 		FMaterialShader::ModifyCompilationEnvironment(Parameters, OutEnvironment);
 		OutEnvironment.SetDefine(TEXT("POST_PROCESS_MATERIAL"), 1);
 		OutEnvironment.SetDefine(TEXT("POST_PROCESS_MATERIAL_BEFORE_TONEMAP"), (Parameters.MaterialParameters.BlendableLocation != BL_AfterTonemapping) ? 1 : 0);
+		OutEnvironment.SetDefine(TEXT("POST_PROCESS_MATERIAL_SSRINPUT"), (Parameters.MaterialParameters.BlendableLocation == BL_SSRInput) ? 1 : 0);
 		OutEnvironment.SetDefine(TEXT("POST_PROCESS_AR_PASSTHROUGH"), 1);
 
 		const bool bMobileForceDepthRead = MobileUsesFullDepthPrepass(Parameters.Platform);
diff --git a/Engine/Plugins/Runtime/AR/Google/GoogleARCore/Source/GoogleARCoreRendering/Private/GoogleARCorePassthroughCameraRenderer.cpp b/Engine/Plugins/Runtime/AR/Google/GoogleARCore/Source/GoogleARCoreRendering/Private/GoogleARCorePassthroughCameraRenderer.cpp
--- a/Engine/Plugins/Runtime/AR/Google/GoogleARCore/Source/GoogleARCoreRendering/Private/GoogleARCorePassthroughCameraRenderer.cpp
+++ b/Engine/Plugins/Runtime/AR/Google/GoogleARCore/Source/GoogleARCoreRendering/Private/GoogleARCorePassthroughCameraRenderer.cpp
@@ -118,6 +118,7 @@ public:
 		FMaterialShader::ModifyCompilationEnvironment(Parameters, OutEnvironment);
 		OutEnvironment.SetDefine(TEXT("POST_PROCESS_MATERIAL"), 1);
 		OutEnvironment.SetDefine(TEXT("POST_PROCESS_MATERIAL_BEFORE_TONEMAP"), (Parameters.MaterialParameters.BlendableLocation != BL_AfterTonemapping) ? 1 : 0);
+		OutEnvironment.SetDefine(TEXT("POST_PROCESS_MATERIAL_SSRINPUT"), (Parameters.MaterialParameters.BlendableLocation == BL_SSRInput) ? 1 : 0);
 	}
 };
 
diff --git a/Engine/Shaders/Private/MaterialTemplate.ush b/Engine/Shaders/Private/MaterialTemplate.ush
--- a/Engine/Shaders/Private/MaterialTemplate.ush
+++ b/Engine/Shaders/Private/MaterialTemplate.ush
@@ -3662,9 +3662,11 @@ void CalcMaterialParametersPost(
 	FIsFrontFace bIsFrontFace)
 {
 	float4 ScreenPosition = SvPositionToScreenPosition(SvPosition);
-#if POST_PROCESS_MATERIAL_BEFORE_TONEMAP	// In this case, SVPositionToTranslatedWorld account for the correct dynamic resolution scaling.
+#if POST_PROCESS_MATERIAL_BEFORE_TONEMAP && !POST_PROCESS_MATERIAL_SSRINPUT	
+	// In this case, SVPositionToTranslatedWorld account for the correct dynamic resolution scaling.
 	float3 TranslatedWorldPosition = SvPositionToTranslatedWorld(SvPosition);
-#else										// Otherwise, we must update SVPosition to account for the resolution scaling not happening after at the tonemap stage and after.
+#else
+	// Otherwise, we must update SVPosition to account for the resolution scaling not happening after at the tonemap stage and after.
 	float3 TranslatedWorldPosition = SvPositionToTranslatedWorld(SvPosition * float4(View.ResolutionFractionAndInv.xx, 1.0f, 1.0f));
 #endif
 
diff --git a/Engine/Source/Runtime/Renderer/Private/PostProcess/PostProcessMaterial.cpp b/Engine/Source/Runtime/Renderer/Private/PostProcess/PostProcessMaterial.cpp
--- a/Engine/Source/Runtime/Renderer/Private/PostProcess/PostProcessMaterial.cpp
+++ b/Engine/Source/Runtime/Renderer/Private/PostProcess/PostProcessMaterial.cpp
@@ -252,6 +252,10 @@ public:
 
 		EBlendableLocation Location = EBlendableLocation(Parameters.MaterialParameters.BlendableLocation);
 		OutEnvironment.SetDefine(TEXT("POST_PROCESS_MATERIAL_BEFORE_TONEMAP"), (Location == BL_AfterTonemapping || Location == BL_ReplacingTonemapper) ? 0 : 1);
+		// Post process SSR is always rendered at native resolution as if it was after tone mapping, so we need to account for the fact that it is independent from DRS.
+		// SSR input should not be affected by exposure so it should be specified separately from POST_PROCESS_MATERIAL_BEFORE_TONEMAP 
+		// in order to be able to make DRS independent CameraVector and WorldPosition nodes.
+		OutEnvironment.SetDefine(TEXT("POST_PROCESS_MATERIAL_SSRINPUT"), (Location == BL_SSRInput) ? 1 : 0);
 
 		if (IsMobilePlatform(Parameters.Platform))
 		{
