* Add ShaderPrint API to set current pixel/thread to create shader print context in any file.
* Fix shader compilation issue when ShaderPrint.h is included into vertex shader

diff --git a/Engine/Shaders/Private/ShaderPrint.ush b/Engine/Shaders/Private/ShaderPrint.ush
--- a/Engine/Shaders/Private/ShaderPrint.ush
+++ b/Engine/Shaders/Private/ShaderPrint.ush
@@ -24,7 +24,7 @@
 
 // On platform supporting diagnostic buffer, shader print data are stored into the diagnostic buffer which is always 
 // bound. This allows to access shader print functionality in any shader without explicit shader binding
-#define SHADER_PRINT_USE_DIAGNOSTIC_BUFFER (PLATFORM_SUPPORTS_DIAGNOSTIC_BUFFER && !SHADER_PRINT_USE_GLOBAL_RESOURCE && SHADER_PRINT_EXPERIMENTAL_DIAGNOSTIC_BUFFER)
+#define SHADER_PRINT_USE_DIAGNOSTIC_BUFFER (PLATFORM_SUPPORTS_DIAGNOSTIC_BUFFER && (PIXELSHADER || COMPUTESHADER || RAYGENSHADER) && !SHADER_PRINT_USE_GLOBAL_RESOURCE && SHADER_PRINT_EXPERIMENTAL_DIAGNOSTIC_BUFFER)
 // Offset to the ShaderPrint parameters
 #define SHADER_PRINT_DIAGNOSTIC_BUFFER_PARAMETERS_OFFSET 6u
 // Offset to the ShaderPrint entries
@@ -149,6 +149,11 @@ FShaderPrintContext InitShaderPrintContextAtCursor(uint2 InCurrentCoord, uint2 I
 	return InitShaderPrintContext(bActive, float2(InStartCoord) / float2(Config.Resolution));
 }
 
+FShaderPrintContext InitShaderPrintContextAtCursor(uint2 InStartCoord) 
+{
+	return InitShaderPrintContextAtCursor(GetShaderPrintCurrentPixelCoord(), InStartCoord);
+}
+
 FShaderPrintContext InitShaderPrintContext()
 {
 	return InitShaderPrintContext(true, float2(0, 0));
diff --git a/Engine/Shaders/Private/ShaderPrintCommon.ush b/Engine/Shaders/Private/ShaderPrintCommon.ush
--- a/Engine/Shaders/Private/ShaderPrintCommon.ush
+++ b/Engine/Shaders/Private/ShaderPrintCommon.ush
@@ -1400,4 +1400,23 @@ uint LenString(uint InEntryID, uint InInfoCount, StructuredBuffer<uint2> InInfoB
 	uint Num##N()\
 	{\
 		return N##_InfoCount;\
-	}
\ No newline at end of file
+	}
+
+///////////////////////////////////////////////////////////////////////////////////////////////////
+// Current pixel
+
+// Store the current pixel coordinates, which allow to retrieve it in any shaders once set.
+// This data is only visible for the current thread/pixel, not available for neighbors
+static uint2 ShaderPrintCurrentPixelCoord;
+
+// Set current pixel coordinates
+void SetShaderPrintCurrentPixelCoord(uint2 In)
+{
+	ShaderPrintCurrentPixelCoord = In;
+}
+
+// Get current pixel coordinates
+uint2 GetShaderPrintCurrentPixelCoord()
+{
+	return ShaderPrintCurrentPixelCoord;
+}
\ No newline at end of file
