Substrate - fix nan appearing when forcing a wrong material mdoe when visiblebsdfcount==0.

diff --git a/Engine/Shaders/Private/Substrate/SubstrateExport.ush b/Engine/Shaders/Private/Substrate/SubstrateExport.ush
--- a/Engine/Shaders/Private/Substrate/SubstrateExport.ush
+++ b/Engine/Shaders/Private/Substrate/SubstrateExport.ush
@@ -304,7 +304,8 @@ void PackSubstrateOut(
 		TopLayerData.BaseColor	= TopLayerTotalWeight > 0.0f	? TopLayerData.BaseColor / TopLayerTotalWeight	: 0.0f;
 
 		// Set storage layout as either: fast(0), single(1), or complex(2)
-		if (bIsOnlySlab)
+		// We check BSDFVisibleCount to make sure we do not set a MATERIALMODE that would enforce a BSDFVisibleCount==1.
+		if (bIsOnlySlab && BSDFVisibleCount > 0)
 		{
 			bIsSimpleMaterial = bIsSimpleMaterial && BSDFVisibleCount == 1;
 			bIsSingleMaterial = bIsSingleMaterial && BSDFVisibleCount == 1;
@@ -385,6 +386,7 @@ void PackSubstrateOut(
 
 		// Now evaluate Specular Occlusion and Indirect Irradiance.
 		#if MATERIAL_SUBSTRATE_OPAQUE_PRECOMPUTED_LIGHTING
+		if(BSDFVisibleCount > 0)
 		{
 			float BSDFCount = 0;
 			float2 SpecMultiBounceAO_IndirectIrradiance_Acc = 0.0f;
