Path Tracer: Alpha channel should not be affected by exposure

diff --git a/Engine/Shaders/Private/PathTracing/PathTracingCompositingPixelShader.usf b/Engine/Shaders/Private/PathTracing/PathTracingCompositingPixelShader.usf
--- a/Engine/Shaders/Private/PathTracing/PathTracingCompositingPixelShader.usf
+++ b/Engine/Shaders/Private/PathTracing/PathTracingCompositingPixelShader.usf
@@ -25,10 +25,10 @@ void CompositeMain(
 	float2 BufferSize = View.BufferSizeAndInvSize.xy;
 	int3 TexCoord = int3(UV * BufferSize - View.ViewRectMin.xy, 0);
 
-	float4 Radiance = RadianceTexture.Load(TexCoord) * View.PreExposure;
+	float4 Radiance = RadianceTexture.Load(TexCoord);
+
+	OutColor = float4(Radiance.rgb * View.PreExposure, Radiance.a);
 
-	OutColor = float4(Radiance.rgb, Radiance.a);
-		
 	if (AdaptiveSamplingVisualize == 1)
 	{
 		// Visualize which pixels are still being worked and color them according to how many more samples we think they need
