Path Tracer: Modify Depth output to take min over entire filter footprint

The current logic was limiting the depth write to only happen for rays passing through the pixel square, but in practice it is more desireable for compositing to have the depth output in sync with the alpha channel. Keep the old logic around temporarily until the new behavior has been confirmed to produce the desired results.

diff --git a/Engine/Shaders/Private/PathTracing/PathTracingCore.ush b/Engine/Shaders/Private/PathTracing/PathTracingCore.ush
--- a/Engine/Shaders/Private/PathTracing/PathTracingCore.ush
+++ b/Engine/Shaders/Private/PathTracing/PathTracingCore.ush
@@ -1378,8 +1378,13 @@ FPathState CreatePathState(int2 PixelIndex, int2 PrimaryRayTextureIndex)
 		float2 ViewportUV = (LaunchIndex + AAJitter) * View.BufferSizeAndInvSize.zw;
 		Output.Ray = CreatePrimaryRay(ViewportUV);
 
+#if 0
 		// only record the depth within a square.
 		Output.DepthResamplingWeight = all(AAJitter >= 0.0) && all(AAJitter <= 1.0) ? 1 : 0;
+#else
+		// record depth for all rays in the filter footprint (TODO: simplify code that uses this weight instead)
+		Output.DepthResamplingWeight = 1.0;
+#endif
 	}
 
 	if (CameraLensRadius.y > 0)
