LFV - move local fog volume shaders into their own folder. Renamed shader to splat shaders before we implement the tile culling.
Also added check.

#prelfight shader

diff --git a/Engine/Shaders/Private/BasePassVertexShader.usf b/Engine/Shaders/Private/BasePassVertexShader.usf
--- a/Engine/Shaders/Private/BasePassVertexShader.usf
+++ b/Engine/Shaders/Private/BasePassVertexShader.usf
@@ -25,7 +25,7 @@
 #endif
 
 #if LOCAL_FOG_VOLUME_ON_TRANSLUCENT
-#include "LocalFogVolumeCommon.ush"
+#include "LocalFogVolumes/LocalFogVolumeCommon.ush"
 #endif
 
 /** Entry point for the base pass vertex shader. */
diff --git a/Engine/Shaders/Private/LocalFogVolume.usf b/Engine/Shaders/Private/LocalFogVolume.usf
deleted file mode 100644
diff --git a/Engine/Shaders/Private/LocalFogVolumeCommon.ush b/Engine/Shaders/Private/LocalFogVolumeCommon.ush
deleted file mode 100644
diff --git a/Engine/Shaders/Private/LocalFogVolumes/LocalFogVolumeCommon.ush b/Engine/Shaders/Private/LocalFogVolumes/LocalFogVolumeCommon.ush
new file mode 100644
--- /dev/null
+++ b/Engine/Shaders/Private/LocalFogVolumes/LocalFogVolumeCommon.ush
@@ -0,0 +1,353 @@
+// Copyright Epic Games, Inc. All Rights Reserved.
+#include "../Common.ush"
+
+#include "../ParticipatingMediaCommon.ush"
+
+
+
+///////////////////////////////////////////////////////////////////////////////
+// Local Fog Volume Data
+
+
+
+struct FLocalFogVolumeGPUInstanceData
+{
+	float4x4 Transform;
+	float4x4 InvTransform;
+
+	float4x4 InvTranformNoScale;
+	float4x4 TransformScaleOnly;
+
+	float Density;
+	float HeightFalloff;
+	float HeightOffset;
+	float RadialAttenuation;
+
+	float3 Albedo;
+	float  PhaseG;
+	float3 Emissive;
+	float  FogMode;
+};
+
+
+
+// Structured buffers cannot be sampled from the vertex shader on some mobile platforms.
+// That is why simple Buffers are used for LocalFogVolumeInstances, because this tech needs to run on mobile platforms.
+
+#if defined(LFVStruct)
+
+#define LFVInsts LFVStruct.LFV.LocalFogVolumeInstances
+#define LFVInstcount LFVStruct.LFV.LocalFogVolumeInstanceCount
+
+#else
+
+uint LFV_LocalFogVolumeInstanceCount;
+Buffer<float4> LFV_LocalFogVolumeInstances;
+
+#define LFVInsts LFV_LocalFogVolumeInstances
+#define LFVInstcount LFV_LocalFogVolumeInstanceCount
+
+#endif
+
+
+
+FLocalFogVolumeGPUInstanceData GetLocalFogVolumeGPUInstanceData(uint Index)
+{
+	FLocalFogVolumeGPUInstanceData Data;
+	uint Offset = Index * 19;
+
+	Data.Transform			= float4x4(LFVInsts[Offset + 0 ], LFVInsts[Offset + 1 ], LFVInsts[Offset + 2 ], LFVInsts[Offset + 3 ]);
+	Data.InvTransform		= float4x4(LFVInsts[Offset + 4 ], LFVInsts[Offset + 5 ], LFVInsts[Offset + 6 ], LFVInsts[Offset + 7 ]);
+
+	Data.InvTranformNoScale = float4x4(LFVInsts[Offset + 8 ], LFVInsts[Offset + 9 ], LFVInsts[Offset + 10], LFVInsts[Offset + 11]);
+	Data.TransformScaleOnly = float4x4(LFVInsts[Offset + 12], LFVInsts[Offset + 13], LFVInsts[Offset + 14], LFVInsts[Offset + 15]);
+
+	float4 Temp0 = LFVInsts[Offset + 16];
+	Data.Density = Temp0.x;
+	Data.HeightFalloff = Temp0.y;
+	Data.HeightOffset = Temp0.z;
+	Data.RadialAttenuation = Temp0.w;
+
+	float4 Temp1 = LFVInsts[Offset + 17];
+	Data.Albedo = Temp1.xyz;
+	Data.PhaseG = Temp1.w;
+
+	float4 Temp2 = LFVInsts[Offset + 18];
+	Data.Emissive = Temp2.xyz;
+	Data.FogMode = Temp2.w;
+
+	return Data;
+}
+
+
+
+///////////////////////////////////////////////////////////////////////////////
+// Local Fog Volume Integrals
+
+
+
+// CameraPos can be behind of RayStart if the camera is out of the volume
+struct FFogData
+{
+	float IntegratedLuminanceFactor;
+	float Coverage;
+};
+
+// Height fog is evaluated according to spatially varying density defined as "Density * Exp[-HeightFalloff * z]"
+FFogData LocalFogVolumeEvaluateHeightIntegral(FLocalFogVolumeGPUInstanceData FogInstance, float3 CameraPos, float3 RayStart, float3 RayDir, float RayLength, float RadialAttenuation)
+{
+	FFogData FogData;
+
+	float3 RayOriU = mul(float4(RayStart, 1.0f), FogInstance.InvTransform).xyz;
+	float3 RayIntU = mul(float4(RayStart + RayDir * RayLength, 1.0f), FogInstance.InvTransform).xyz;
+	float3 RayDirU = mul(float4(RayDir, 0.0f), FogInstance.InvTransform).xyz;
+	RayDirU = normalize(RayDirU);
+
+	float RayLengthU = length(RayIntU - RayOriU);
+
+	float StartHeight = RayOriU.z - FogInstance.HeightOffset; // accoutn for the fog offset.
+
+	float OpticalDepth = 0.0f;
+#if 0
+	// Original integral, less artefact
+	OpticalDepth = (FogInstance.Density / FogInstance.HeightFalloff) * exp(-StartHeight * FogInstance.HeightFalloff) * (1.0 - exp(-RayDirU.z * RayLengthU * FogInstance.HeightFalloff)) / RayDirU.z;
+#elif 1
+	// Reworked to avoid large value as input to exp
+	//OpticalDepth = (FogInstance.Density / FogInstance.HeightFalloff) * (exp(-StartHeight * FogInstance.HeightFalloff) - exp(-StartHeight * FogInstance.HeightFalloff - RayDir.z * RayLength * FogInstance.HeightFalloff)) / RayDir.z;
+
+	// But we use an updated formulation to avoid large input to exp leading to inf and nan.
+	float Factor0 = StartHeight * FogInstance.HeightFalloff;
+	float Factor1 = RayDirU.z * RayLengthU * FogInstance.HeightFalloff;
+	Factor0 = max(-80.0f, Factor0); // Simply clamp Factor0 to a valid range in order to not have exp explode.
+
+	OpticalDepth = (FogInstance.Density / (FogInstance.HeightFalloff * RayDirU.z)) * (exp(-Factor0) - exp(-(Factor0 + Factor1)));
+#endif
+
+	// Transmittance is the 
+	float Transmittance = exp(-OpticalDepth);
+	// Assuming extinction is grey scale, we can compute a single coverage value from transmittance
+	FogData.Coverage = (1.0 - Transmittance) * RadialAttenuation;
+	// Assuming extinction==scattering, i.e. albedo=1, we can compute the integrale for each point on a ray while accouting for transmittance back to the view.
+	// It turns out that, under these conditions, the luminance and emissive factor is simply coverage. => see Mathematica details at the bottom of this page.
+	FogData.IntegratedLuminanceFactor = FogData.Coverage;
+
+	return FogData;
+}
+
+
+FFogData LocalFogVolumeEvaluateRadialIntegral(FLocalFogVolumeGPUInstanceData FogInstance, float3 CameraPos, float3 RayStart, float3 RayDir, float RayLength, float RadialAttenuation)
+{
+	FFogData FogData = (FFogData)0;
+
+	float3 RayOriU = mul(float4(RayStart, 1.0f), FogInstance.InvTransform).xyz;
+	float3 RayIntU = mul(float4(RayStart + RayDir * RayLength, 1.0f), FogInstance.InvTransform).xyz;
+	float3 RayDirU = mul(float4(RayDir, 0.0f), FogInstance.InvTransform).xyz;
+	RayDirU = normalize(RayDirU);
+
+	float RayLengthU = length(RayIntU - RayOriU);
+
+	float3 SphereCenter = 0.0f;
+	float3 VolumeCenterToRayO = (RayOriU - SphereCenter);
+
+	float b = dot(RayDirU, VolumeCenterToRayO);
+	float c = dot(VolumeCenterToRayO, VolumeCenterToRayO) - 1.0f;
+	float h = b * b - c;
+
+	float OpticalDepth = 0.0f;
+	if (h >= 0.0)
+	{
+		h = sqrt(h);
+		float Length0 = -b - h;
+		float Length1 = -b + h;
+
+		Length0 = max(Length0, 0.0);
+		Length1 = max(Length1, 0.0);
+		Length1 = min(Length1, RayLengthU);
+
+		// Integral of a density that descreases accordingf to square distance from the center of the volume.
+		float Integral0 = -(c * Length0 + b * Length0 * Length0 + Length0 * Length0 * Length0 / 3.0f);
+		float Integral1 = -(c * Length1 + b * Length1 * Length1 + Length1 * Length1 * Length1 / 3.0f);
+		OpticalDepth = max(0.0, FogInstance.Density * (Integral1 - Integral0) * (3.0f / 4.0f));
+	}
+
+	// Transmittance is the 
+	float Transmittance = exp(-OpticalDepth);
+	// Assuming extinction is grey scale, we can compute a single coverage value from transmittance
+	FogData.Coverage = (1.0 - Transmittance) * RadialAttenuation;
+	// Assuming extinction==scattering, i.e. albedo=1, we can compute the integrale for each point on a ray while accouting for transmittance back to the view.
+	// It turns out that, under these conditions, the luminance and emissive factor is simply coverage. => see Mathematica details at the bottom of this page.
+	FogData.IntegratedLuminanceFactor = FogData.Coverage;
+
+	return FogData;
+}
+
+
+
+///////////////////////////////////////////////////////////////////////////////
+// Local Fog Volume Evaluation
+
+struct LocalFogVolumeCommonLightingInput
+{
+	float3 DistantSkyLight;
+};
+
+LocalFogVolumeCommonLightingInput LocalFogVolumeEvaluateCommonLightingInput()
+{
+	LocalFogVolumeCommonLightingInput CommonLightingInput = (LocalFogVolumeCommonLightingInput)0;
+#if PROJECT_SUPPORT_SKY_ATMOSPHERE_AFFECTS_HEIGHFOG
+	CommonLightingInput.DistantSkyLight = Texture2DSampleLevel(View.DistantSkyLightLutTexture, View.DistantSkyLightLutTextureSampler, float2(0.5f, 0.5f), 0.0f).rgb;
+#endif
+	return CommonLightingInput;
+}
+
+float3 LocalFogVolumeEvaluateInScattering(in FLocalFogVolumeGPUInstanceData FogInstance, in FFogData FogData, in LocalFogVolumeCommonLightingInput CommonLightingInput, in float3 RayDirWorld)
+{
+	half3 InScattering = 0;
+
+#if PROJECT_SUPPORT_SKY_ATMOSPHERE_AFFECTS_HEIGHFOG
+
+	// Accumulate distant sky light
+	InScattering += CommonLightingInput.DistantSkyLight;
+
+	// Accumulate directional lights if needed
+	if (View.AtmosphereLightIlluminanceOnGroundPostTransmittance[0].a > 0.0f)
+	{
+		InScattering += View.AtmosphereLightIlluminanceOnGroundPostTransmittance[0].rgb * SchlickPhase(-FogInstance.PhaseG, dot(RayDirWorld, View.AtmosphereLightDirection[0].xyz));
+	}
+	if (View.AtmosphereLightIlluminanceOnGroundPostTransmittance[1].a > 0.0f)
+	{
+		InScattering += View.AtmosphereLightIlluminanceOnGroundPostTransmittance[0].rgb * SchlickPhase(-FogInstance.PhaseG, dot(RayDirWorld, View.AtmosphereLightDirection[1].xyz));
+	}
+
+	// Weight by user control
+	InScattering *= View.SkyAtmosphereHeightFogContribution.xxx;
+
+	// Apply the medium uniform single scattering albedo
+	InScattering *= FogInstance.Albedo;
+#endif
+
+	// Now account for emissive luminance after albedo and heighfog contribution has been accounted for.
+	// This do mean that the emissive color follows the medium density variation.
+	InScattering += FogInstance.Emissive;
+
+	// Then scale the contribution according to the normalised pre intergrated scattering.
+	return InScattering * FogData.IntegratedLuminanceFactor;
+}
+
+
+
+///////////////////////////////////////////////////////////////////////////////
+// Local Fog Volume Instance Evaluation
+
+#ifndef LFV_CAN_CLIP
+#define LFV_CAN_CLIP 0
+#endif
+
+#define LHFV_FOLLOWS_ROTATIONS 1
+
+// Returns color transmittance
+float4 GetLFVInstanceContribution(
+	uint InstanceIndex,
+	float3 CamRayWorldOrigin,
+	float3 CamRayWorldDir,
+	float3 CamRayTranslatedWorldOrigin,
+	float3 CamRayTranslatedWorldDir,
+	float3 DepthBufferTranslatedWorldPos)
+{
+	FLocalFogVolumeGPUInstanceData FogInstance = GetLocalFogVolumeGPUInstanceData(InstanceIndex);
+
+	float3 LuminanceColor = float3(0.0f, 0.0f, 0.0f);
+	float Transmittance = 1.0f;
+
+	// The "U" prefix is for all compuation done in the Unit Sphere space.
+	float3 RayPosU = mul(float4(CamRayWorldOrigin, 1.0f), FogInstance.InvTransform).xyz;
+	float3 RayDirU = mul(float4(CamRayWorldDir,    0.0f), FogInstance.InvTransform).xyz; // Try to use InverseTranspose to correctly handle direction transform with scaling.
+	RayDirU.xyz = normalize(RayDirU);
+
+	float2 TsU = RayIntersectSphere(RayPosU, RayDirU, float4(0.0, 0.0, 0.0, 1.0));
+	if (any(TsU > 0.0))
+	{
+		float LengthD = length(DepthBufferTranslatedWorldPos - CamRayTranslatedWorldOrigin);
+
+		float3 P0U = RayPosU + max(0.0, TsU.x) * RayDirU;
+		float3 P1U = RayPosU + max(0.0, TsU.y) * RayDirU;
+#if LHFV_FOLLOWS_ROTATIONS
+		// Computations are all done transformed with Space = local fog volume but ignoring scaling.
+		float3 P0Space = mul(float4(P0U, 1.0), FogInstance.TransformScaleOnly).xyz;
+		float3 P1Space = mul(float4(P1U, 1.0), FogInstance.TransformScaleOnly).xyz;
+		float3 CamRayOriginSpace = mul(float4(CamRayWorldOrigin, 1.0), FogInstance.InvTranformNoScale).xyz;
+		float3 CamRayDirSpace = mul(float4(CamRayWorldDir, 0.0), FogInstance.InvTranformNoScale).xyz;	// No scale so no need for a special InverseTranspose matrix.
+		CamRayDirSpace = normalize(CamRayDirSpace);
+#else
+		// Computations are all done with Space = world space
+		float3 P0Space = mul(float4(P0U, 1.0), FogInstance.Transform).xyz;
+		float3 P1Space = mul(float4(P1U, 1.0), FogInstance.Transform).xyz;
+		float3 CamRayOriginSpace = CamRayWorldOrigin;
+		float3 CamRayDirSpace = CamRayWorldDir;
+#endif
+
+		float Length0 = length(P0Space - CamRayOriginSpace);
+		float Length1 = length(P1Space - CamRayOriginSpace);
+
+		// Radial attenuation based on the distance travelled in the unit sphere and ignoring depth buffer intersection. 
+		// This is only to soften sphere edge.
+		float RadialAttenuation = saturate(abs(max(0.0, TsU.y) - max(0.0, TsU.x)) / FogInstance.RadialAttenuation);
+		// The attenuation is linear but it can be made soft using an Hermit interpolation by controling the tangents.
+		{
+			const float t = RadialAttenuation;
+			const float t2 = t * t;
+			const float t3 = t2 * t;
+			const float HP0 = 0.0f;
+			const float HP1 = 1.0f;
+			const float HM0 = 0.0f;	// tangent = 0 for soft transition
+			const float HM1 = 0.0f; // idem
+			RadialAttenuation = saturate((2 * t3 - 3 * t2 + 1) * HP0 + (t3 - 2 * t2 + t) * HM0 + (-2 * t3 + 3 * t2) * HP1 + (t3 - t2) * HM1);
+		}
+
+		// Clamp the traced distance to depth and compute world length.
+		Length0 = min(LengthD, Length0);
+		Length1 = min(LengthD, Length1);
+		float RayTracedLength = max(0.0, abs(Length1 - Length0));
+
+		if (RayTracedLength > 0.0)
+		{
+			FFogData FogData = (FFogData)0;
+
+			LocalFogVolumeCommonLightingInput CommonLightingInput = LocalFogVolumeEvaluateCommonLightingInput();
+
+			if (FogInstance.FogMode == 0)
+			{
+				FogData = LocalFogVolumeEvaluateHeightIntegral(FogInstance, CamRayWorldOrigin, CamRayWorldOrigin + CamRayWorldDir * Length0, CamRayWorldDir, RayTracedLength, RadialAttenuation);
+			}
+			else // (FogInstance.FogMode == 1)
+			{
+				// TODO
+				// TODO LocalFogVolumeEvaluateRadialIntegral needs to be optimised because it is redoing lots of computations. Maybe move HeightFogIntegral to also be in a local normalised space.
+				// TODO It is like that for now while iterating with artists.
+				// TODO
+				FogData = LocalFogVolumeEvaluateRadialIntegral(FogInstance, CamRayWorldOrigin, CamRayWorldOrigin + CamRayWorldDir * Length0, CamRayWorldDir, RayTracedLength, RadialAttenuation);
+			}
+
+			LuminanceColor = LocalFogVolumeEvaluateInScattering(FogInstance, FogData, CommonLightingInput, CamRayWorldDir);
+			Transmittance = 1.0 - FogData.Coverage;
+		}
+		else
+		{
+#if LFV_CAN_CLIP
+			clip(-1.0);
+#endif
+		}
+	}
+	else
+	{
+#if LFV_CAN_CLIP
+		clip(-1.0);
+#endif
+	}
+
+	return float4(LuminanceColor, Transmittance);
+}
+
+
+
+
diff --git a/Engine/Shaders/Private/LocalFogVolumes/LocalFogVolumeSplat.usf b/Engine/Shaders/Private/LocalFogVolumes/LocalFogVolumeSplat.usf
new file mode 100644
--- /dev/null
+++ b/Engine/Shaders/Private/LocalFogVolumes/LocalFogVolumeSplat.usf
@@ -0,0 +1,111 @@
+// Copyright Epic Games, Inc. All Rights Reserved.
+#include "../Common.ush"
+
+#if SHADING_PATH_MOBILE
+#define MobileSceneTextures MobileBasePass.SceneTextures 
+#endif
+
+#include "../SceneTexturesCommon.ush"
+
+#define LFV_CAN_CLIP 1
+#include "LocalFogVolumeCommon.ush"
+
+
+struct FLocalFogVolumeVertexOutput
+{
+	nointerpolation uint InstanceId : TEXCOORD0;
+};
+
+
+void LocalFogVolumeSplatVS(
+	in float4 InPosition : ATTRIBUTE0,
+	uint InstanceIndex : SV_InstanceID,
+	out float4 OutPosition : SV_POSITION,
+	out FLocalFogVolumeVertexOutput OtherVertexOutput
+)
+{
+	ResolvedView = ResolveView();
+
+	float4 LocalPosition = InPosition;
+
+	FLocalFogVolumeGPUInstanceData FogInstance = GetLocalFogVolumeGPUInstanceData(InstanceIndex);
+
+	float4x4 FogTransform = FogInstance.Transform;
+
+	float4 WorldPosition = mul(float4(LocalPosition.xyz, 1.0f), FogTransform);
+
+	float4 TranslatedWorldPosition = float4(WorldPosition.xyz + LWCHackToFloat(PrimaryView.PreViewTranslation), 1);
+
+	OutPosition = mul(TranslatedWorldPosition, PrimaryView.TranslatedWorldToClip);
+	OtherVertexOutput.InstanceId = InstanceIndex;
+}
+
+
+void LocalFogVolumeSplatPS(
+	in float4 SVPos : SV_POSITION,
+	in FLocalFogVolumeVertexOutput OtherVertexOutput,
+	out float4 OutColor : SV_Target0)
+{
+	ResolvedView = ResolveView();
+	uint InstanceIndex = OtherVertexOutput.InstanceId;
+
+	float3 CamRayTranslatedWorldOrigin	= PrimaryView.TranslatedWorldCameraOrigin;
+	float3 CamRayTranslatedWorldDir = normalize(SvPositionToTranslatedWorld(float4(SVPos.xy, 0.5, 1.0)));
+	float3 CamRayWorldDir = CamRayTranslatedWorldDir;
+	float3 CamRayWorldOrigin = CamRayTranslatedWorldOrigin - LWCHackToFloat(PrimaryView.PreViewTranslation);
+
+#if SHADING_PATH_MOBILE
+	// LookupDeviceZ with uint2 SVPos.xy as input does not work on mobile.
+	float2 UV = SVPos * View.BufferSizeAndInvSize.zw;
+	const float DeviceZ = LookupDeviceZ(UV);
+#else
+	const float DeviceZ = LookupDeviceZ(uint2(SVPos.xy));
+#endif
+
+	float3 DepthBufferTranslatedWorldPos = SvPositionToTranslatedWorld(float4(SVPos.xy, DeviceZ, 1.0));
+
+	float4 LuminanceColorTransmittance = GetLFVInstanceContribution(
+		InstanceIndex,
+		CamRayWorldOrigin,
+		CamRayWorldDir,
+		CamRayTranslatedWorldOrigin,
+		CamRayTranslatedWorldDir,
+		DepthBufferTranslatedWorldPos);
+
+	OutColor = float4(
+		LuminanceColorTransmittance.rgb * PrimaryView.PreExposure,
+		LuminanceColorTransmittance.a);
+}
+
+
+/*
+
+Proof of the height fog integral equations using Mathematica
+
+(*Density of matter according to height Y*)
+Density [y_] := A * Exp[-B * y]
+ 
+(*Ray position*)
+Ray[t_] := Oz + t * Rz
+
+(*Integrate Optical Depth for a distance T*)
+Fog[T_] := Integrate[Density[Ray[t]], {t, 0, T}]
+
+(*Optical Depth integration equation*)
+Fog[T]
+
+(*Transmittance from Optical Depth*)
+Exp[-Integrate[Density[Ray[t]], {t, 0, T}]]
+
+(*Transmittance from a position to the origin of tracing*)
+Trans[T_] := Exp[-Integrate[Density[Ray[t]], {t, 0, T}]]
+
+(*Integrate emissive and scattering from height fog, assuming density \
+is extinction and extinction=scattering, thus albedo=1*)
+Integrate[Density[Ray[t]]*EmSc *Trans[t], {t, 0, T}]
+
+=> this leads to EmSc - Transmittance * EmSc
+				 EmSc * (1 - Transmittance)
+				 EmSc * Coverage
+
+*/
\ No newline at end of file
diff --git a/Engine/Shaders/Private/MobileBasePassVertexShader.usf b/Engine/Shaders/Private/MobileBasePassVertexShader.usf
--- a/Engine/Shaders/Private/MobileBasePassVertexShader.usf
+++ b/Engine/Shaders/Private/MobileBasePassVertexShader.usf
@@ -21,7 +21,7 @@
 #endif
 
 #if LOCAL_FOG_VOLUME_ON_TRANSLUCENT
-#include "LocalFogVolumeCommon.ush"
+#include "LocalFogVolumes/LocalFogVolumeCommon.ush"
 #endif
 
 struct FMobileShadingBasePassVSToPS
diff --git a/Engine/Source/Runtime/Renderer/Private/DeferredShadingRenderer.cpp b/Engine/Source/Runtime/Renderer/Private/DeferredShadingRenderer.cpp
--- a/Engine/Source/Runtime/Renderer/Private/DeferredShadingRenderer.cpp
+++ b/Engine/Source/Runtime/Renderer/Private/DeferredShadingRenderer.cpp
@@ -3431,7 +3431,7 @@ void FDeferredShadingSceneRenderer::Render(FRDGBuilder& GraphBuilder)
 		// Run before RenderSkyAtmosphereLookUpTables for cloud shadows to be valid.
 		InitVolumetricCloudsForViews(GraphBuilder, bShouldRenderVolumetricCloudBase, InstanceCullingManager);
 
-		// Run local fog volume initialisation before basepass for when data is needed in forward
+		// Run local fog volume initialization before base pass for when data is needed in forward
 		InitLocalFogVolumesForViews(Scene, Views, ViewFamily, GraphBuilder);
 
 		if (SkyAtmospherePassLocation == ESkyAtmospherePassLocation::BeforeOcclusion && bShouldRenderSkyAtmosphere)
diff --git a/Engine/Source/Runtime/Renderer/Private/LocalFogVolumeRendering.cpp b/Engine/Source/Runtime/Renderer/Private/LocalFogVolumeRendering.cpp
--- a/Engine/Source/Runtime/Renderer/Private/LocalFogVolumeRendering.cpp
+++ b/Engine/Source/Runtime/Renderer/Private/LocalFogVolumeRendering.cpp
@@ -118,6 +118,8 @@ bool FScene::HasAnyLocalFogVolume() const
 
 void GetLocalFogVolumeSortingData(const FScene* Scene, FRDGBuilder& GraphBuilder, FLocalFogVolumeSortingData& Out)
 {
+	check(Scene->LocalFogVolumes.Num() > 0); // We should not get there if there is not any local fog volume.
+
 	// No culling as of today
 	Out.LocalFogVolumeInstanceCount = Scene->LocalFogVolumes.Num();
 	Out.LocalFogVolumeInstanceCountFinal = 0;
@@ -243,10 +245,10 @@ void InitLocalFogVolumesForViews(
 	Local height fog rendering - non mobile
 =============================================================================*/
 
-class FLocalFogVolumeVS : public FGlobalShader
+class FLocalFogVolumeSplatVS : public FGlobalShader
 {
-	DECLARE_GLOBAL_SHADER(FLocalFogVolumeVS);
-	SHADER_USE_PARAMETER_STRUCT(FLocalFogVolumeVS, FGlobalShader);
+	DECLARE_GLOBAL_SHADER(FLocalFogVolumeSplatVS);
+	SHADER_USE_PARAMETER_STRUCT(FLocalFogVolumeSplatVS, FGlobalShader);
 
 	BEGIN_SHADER_PARAMETER_STRUCT(FParameters, )
 		SHADER_PARAMETER_STRUCT_REF(FViewUniformShaderParameters, View)
@@ -265,12 +267,12 @@ class FLocalFogVolumeVS : public FGlobalShader
 	}
 };
 
-IMPLEMENT_GLOBAL_SHADER(FLocalFogVolumeVS, "/Engine/Private/LocalFogVolume.usf", "LocalFogVolumeSplatVS", SF_Vertex);
+IMPLEMENT_GLOBAL_SHADER(FLocalFogVolumeSplatVS, "/Engine/Private/LocalFogVolumes/LocalFogVolumeSplat.usf", "LocalFogVolumeSplatVS", SF_Vertex);
 
-class FLocalFogVolumePS : public FGlobalShader
+class FLocalFogVolumeSplatPS : public FGlobalShader
 {
-	DECLARE_GLOBAL_SHADER(FLocalFogVolumePS);
-	SHADER_USE_PARAMETER_STRUCT(FLocalFogVolumePS, FGlobalShader);
+	DECLARE_GLOBAL_SHADER(FLocalFogVolumeSplatPS);
+	SHADER_USE_PARAMETER_STRUCT(FLocalFogVolumeSplatPS, FGlobalShader);
 
 	BEGIN_SHADER_PARAMETER_STRUCT(FParameters, )
 		SHADER_PARAMETER_STRUCT_REF(FViewUniformShaderParameters, View)
@@ -289,11 +291,11 @@ class FLocalFogVolumePS : public FGlobalShader
 	}
 };
 
-IMPLEMENT_GLOBAL_SHADER(FLocalFogVolumePS, "/Engine/Private/LocalFogVolume.usf", "LocalFogVolumeSplatPS", SF_Pixel);
+IMPLEMENT_GLOBAL_SHADER(FLocalFogVolumeSplatPS, "/Engine/Private/LocalFogVolumes/LocalFogVolumeSplat.usf", "LocalFogVolumeSplatPS", SF_Pixel);
 
 BEGIN_SHADER_PARAMETER_STRUCT(FLocalFogVolumePassParameters, )
-	SHADER_PARAMETER_STRUCT_INCLUDE(FLocalFogVolumeVS::FParameters, VS)
-	SHADER_PARAMETER_STRUCT_INCLUDE(FLocalFogVolumePS::FParameters, PS)
+	SHADER_PARAMETER_STRUCT_INCLUDE(FLocalFogVolumeSplatVS::FParameters, VS)
+	SHADER_PARAMETER_STRUCT_INCLUDE(FLocalFogVolumeSplatPS::FParameters, PS)
 	SHADER_PARAMETER_RDG_UNIFORM_BUFFER(FSceneTextureUniformParameters, SceneTextures)
 	RENDER_TARGET_BINDING_SLOTS()
 END_SHADER_PARAMETER_STRUCT()
@@ -331,11 +333,11 @@ void RenderLocalFogVolume(
 			PassParameters->SceneTextures = SceneTextures.UniformBuffer;
 			PassParameters->RenderTargets[0] = FRenderTargetBinding(SceneColorTexture, ERenderTargetLoadAction::ENoAction);
 
-			FLocalFogVolumeVS::FPermutationDomain VSPermutationVector;
-			auto VertexShader = View.ShaderMap->GetShader< FLocalFogVolumeVS >(VSPermutationVector);
+			FLocalFogVolumeSplatVS::FPermutationDomain VSPermutationVector;
+			auto VertexShader = View.ShaderMap->GetShader< FLocalFogVolumeSplatVS >(VSPermutationVector);
 
-			FLocalFogVolumePS::FPermutationDomain PsPermutationVector;
-			auto PixelShader = View.ShaderMap->GetShader< FLocalFogVolumePS >(PsPermutationVector);
+			FLocalFogVolumeSplatPS::FPermutationDomain PsPermutationVector;
+			auto PixelShader = View.ShaderMap->GetShader< FLocalFogVolumeSplatPS >(PsPermutationVector);
 
 			const FIntRect ViewRect = View.ViewRect;
 
@@ -389,10 +391,10 @@ void RenderLocalFogVolume(
 =============================================================================*/
 	
 
-class FMobileLocalFogVolumeVS : public FGlobalShader
+class FMobileLocalFogVolumeSplatVS : public FGlobalShader
 {
-	DECLARE_GLOBAL_SHADER(FMobileLocalFogVolumeVS);
-	SHADER_USE_PARAMETER_STRUCT(FMobileLocalFogVolumeVS, FGlobalShader);
+	DECLARE_GLOBAL_SHADER(FMobileLocalFogVolumeSplatVS);
+	SHADER_USE_PARAMETER_STRUCT(FMobileLocalFogVolumeSplatVS, FGlobalShader);
 
 	BEGIN_SHADER_PARAMETER_STRUCT(FParameters, )
 		SHADER_PARAMETER_STRUCT_INCLUDE(FViewShaderParameters, View)
@@ -418,12 +420,12 @@ class FMobileLocalFogVolumeVS : public FGlobalShader
 	}
 };
 
-IMPLEMENT_GLOBAL_SHADER(FMobileLocalFogVolumeVS, "/Engine/Private/LocalFogVolume.usf", "LocalFogVolumeSplatVS", SF_Vertex);
+IMPLEMENT_GLOBAL_SHADER(FMobileLocalFogVolumeSplatVS, "/Engine/Private/LocalFogVolumes/LocalFogVolumeSplat.usf", "LocalFogVolumeSplatVS", SF_Vertex);
 
-class FMobileLocalFogVolumePS : public FGlobalShader
+class FMobileLocalFogVolumeSplatPS : public FGlobalShader
 {
-	DECLARE_GLOBAL_SHADER(FMobileLocalFogVolumePS);
-	SHADER_USE_PARAMETER_STRUCT(FMobileLocalFogVolumePS, FGlobalShader);
+	DECLARE_GLOBAL_SHADER(FMobileLocalFogVolumeSplatPS);
+	SHADER_USE_PARAMETER_STRUCT(FMobileLocalFogVolumeSplatPS, FGlobalShader);
 
 	BEGIN_SHADER_PARAMETER_STRUCT(FParameters, )
 		SHADER_PARAMETER_STRUCT_INCLUDE(FViewShaderParameters, View)
@@ -450,7 +452,7 @@ class FMobileLocalFogVolumePS : public FGlobalShader
 	}
 };
 
-IMPLEMENT_GLOBAL_SHADER(FMobileLocalFogVolumePS, "/Engine/Private/LocalFogVolume.usf", "LocalFogVolumeSplatPS", SF_Pixel);
+IMPLEMENT_GLOBAL_SHADER(FMobileLocalFogVolumeSplatPS, "/Engine/Private/LocalFogVolumes/LocalFogVolumeSplat.usf", "LocalFogVolumeSplatPS", SF_Pixel);
 
 void RenderLocalFogVolumeMobile(
 	FRHICommandList& RHICmdList,
@@ -463,11 +465,11 @@ void RenderLocalFogVolumeMobile(
 
 	SCOPED_DRAW_EVENT(RHICmdList, LocalFogVolumeVolumes);
 
-	FMobileLocalFogVolumeVS::FPermutationDomain VSPermutationVector;
-	auto VertexShader = View.ShaderMap->GetShader< FMobileLocalFogVolumeVS >(VSPermutationVector);
+	FMobileLocalFogVolumeSplatVS::FPermutationDomain VSPermutationVector;
+	auto VertexShader = View.ShaderMap->GetShader< FMobileLocalFogVolumeSplatVS >(VSPermutationVector);
 
-	FMobileLocalFogVolumePS::FPermutationDomain PsPermutationVector;
-	auto PixelShader = View.ShaderMap->GetShader< FMobileLocalFogVolumePS >(PsPermutationVector);
+	FMobileLocalFogVolumeSplatPS::FPermutationDomain PsPermutationVector;
+	auto PixelShader = View.ShaderMap->GetShader< FMobileLocalFogVolumeSplatPS >(PsPermutationVector);
 
 	const FIntRect ViewRect = View.ViewRect;
 
@@ -488,12 +490,12 @@ void RenderLocalFogVolumeMobile(
 
 	SetGraphicsPipelineState(RHICmdList, GraphicsPSOInit, 0);
 
-	FMobileLocalFogVolumeVS::FParameters VSParameters;
+	FMobileLocalFogVolumeSplatVS::FParameters VSParameters;
 	VSParameters.View = View.GetShaderParameters();
 	VSParameters.LFV = View.LocalFogVolumeViewData.UniformParametersStruct;
 	SetShaderParameters(RHICmdList, VertexShader, VertexShader.GetVertexShader(), VSParameters);
 
-	FMobileLocalFogVolumePS::FParameters PSParameters;
+	FMobileLocalFogVolumeSplatPS::FParameters PSParameters;
 	PSParameters.View = View.GetShaderParameters();
 	PSParameters.LFV = View.LocalFogVolumeViewData.UniformParametersStruct;
 	// PSParameters.MobileBasePass filled up by the RDG pass parameters.
diff --git a/Engine/Source/Runtime/Renderer/Private/MobileShadingRenderer.cpp b/Engine/Source/Runtime/Renderer/Private/MobileShadingRenderer.cpp
--- a/Engine/Source/Runtime/Renderer/Private/MobileShadingRenderer.cpp
+++ b/Engine/Source/Runtime/Renderer/Private/MobileShadingRenderer.cpp
@@ -1273,8 +1273,8 @@ void FMobileSceneRenderer::RenderForward(FRDGBuilder& GraphBuilder, FRDGTextureR
 
 	const FRDGSystemTextures& SystemTextures = FRDGSystemTextures::Get(GraphBuilder);
 
-	FLocalFogVolumeSortingData SortingData;
-	GetLocalFogVolumeSortingData(Scene, GraphBuilder, SortingData);
+	// Run local fog volume initialization before base pass for when data is needed in forward
+	InitLocalFogVolumesForViews(Scene, Views, ViewFamily, GraphBuilder);
 
 	FRenderViewContextArray RenderViews;
 	GetRenderViews(Views, RenderViews);
@@ -1301,7 +1301,6 @@ void FMobileSceneRenderer::RenderForward(FRDGBuilder& GraphBuilder, FRDGTextureR
 		View.BeginRenderView();
 
 		UpdateDirectionalLightUniformBuffers(GraphBuilder, View);
-		CreateViewLocalFogVolumeBufferSRV(View, GraphBuilder, SortingData);
 
 		FMobileBasePassTextures MobileBasePassTextures{};
 		MobileBasePassTextures.ScreenSpaceAO = bRequiresAmbientOcclusionPass ? SceneTextures.ScreenSpaceAO : SystemTextures.White;
@@ -1610,8 +1609,8 @@ void FMobileSceneRenderer::RenderDeferred(FRDGBuilder& GraphBuilder, const FSort
 
 	const FRDGSystemTextures& SystemTextures = FRDGSystemTextures::Get(GraphBuilder);
 
-	FLocalFogVolumeSortingData SortingData;
-	GetLocalFogVolumeSortingData(Scene, GraphBuilder, SortingData);
+	// Run local fog volume initialization before base pass for when data is needed in forward
+	InitLocalFogVolumesForViews(Scene, Views, ViewFamily, GraphBuilder);
 
 	FRenderViewContextArray RenderViews;
 	GetRenderViews(Views, RenderViews);
@@ -1638,7 +1637,6 @@ void FMobileSceneRenderer::RenderDeferred(FRDGBuilder& GraphBuilder, const FSort
 		View.BeginRenderView();
 
 		UpdateDirectionalLightUniformBuffers(GraphBuilder, View);
-		CreateViewLocalFogVolumeBufferSRV(View, GraphBuilder, SortingData);
 
 		FMobileBasePassTextures MobileBasePassTextures{};
 		MobileBasePassTextures.ScreenSpaceAO = bRequiresAmbientOcclusionPass ? SceneTextures.ScreenSpaceAO : SystemTextures.White;
