Fixed instance custom data handling in a mobile specific GPUScene case

diff --git a/Engine/Shaders/Private/SceneDataMobileWriter.ush b/Engine/Shaders/Private/SceneDataMobileWriter.ush
--- a/Engine/Shaders/Private/SceneDataMobileWriter.ush
+++ b/Engine/Shaders/Private/SceneDataMobileWriter.ush
@@ -188,8 +188,9 @@ void WriteInstanceDataUBO(RWStructuredBuffer<float4> Output, uint Offset, uint I
 		// limited space for instance custom data
 		InstanceData.CustomDataCount = CustomDataCount;
 		InstanceData.CustomDataOffset = Offsets.CustomData;
+		uint CustomDataElements = (CustomDataCount + 3u) >> 2u;
 		UNROLL
-		for (uint i = 0; i < CustomDataCount; ++i)
+		for (uint i = 0; i < CustomDataElements; ++i)
 		{
 			float4 CustomData = LoadInstanceCustomDataElement(InstanceData, i);
 			WriteDataElementUBO(Output, Offset, DataIdx + i, CustomData);
