Fixed incorrect implementation of parallel prefix-max operation in the instance culling load balancer. New version is also rather more efficient.

diff --git a/Engine/Shaders/Private/InstanceCulling/InstanceCullingLoadBalancer.ush b/Engine/Shaders/Private/InstanceCulling/InstanceCullingLoadBalancer.ush
--- a/Engine/Shaders/Private/InstanceCulling/InstanceCullingLoadBalancer.ush
+++ b/Engine/Shaders/Private/InstanceCulling/InstanceCullingLoadBalancer.ush
@@ -118,11 +118,28 @@ uint PrefixMax(uint ElementIndex)
 	uint CurrentValue = ItemIndex[ElementIndex];
 	GroupMemoryBarrierWithGroupSync();
 
-	UNROLL
-	for (uint StepSize = 1U; StepSize < NUM_THREADS_PER_GROUP / 2U; ++StepSize)
-	{
-		PrefixMaxStep(ElementIndex, 1U, CurrentValue);
-	}
+	// Note: manually unrolled as UNROLL seems to not figure out *= 2, at least on some platforms
+	//for (uint StepSize = 1U; StepSize < NUM_THREADS_PER_GROUP / 2U; StepSize *= 2u)
+	//{
+	//	PrefixMaxStep(ElementIndex, StepSize, CurrentValue);
+	//}
+	PrefixMaxStep(ElementIndex, 1U, CurrentValue);
+	PrefixMaxStep(ElementIndex, 2U, CurrentValue);
+	PrefixMaxStep(ElementIndex, 4U, CurrentValue);
+	PrefixMaxStep(ElementIndex, 8U, CurrentValue);
+#if NUM_THREADS_PER_GROUP > 32u
+	PrefixMaxStep(ElementIndex, 16U, CurrentValue);
+#endif
+#if NUM_THREADS_PER_GROUP > 64u
+	PrefixMaxStep(ElementIndex, 32U, CurrentValue);
+#endif
+#if NUM_THREADS_PER_GROUP > 128u
+	PrefixMaxStep(ElementIndex, 64U, CurrentValue);
+#endif
+#if NUM_THREADS_PER_GROUP > 256u
+	PrefixMaxStep(ElementIndex, 128U, CurrentValue);
+#endif
+
 	return PrefixMaxLastStep(ElementIndex, CurrentValue);
 }
 
