Temporary fix for Vulkan mobile with Substrate.

diff --git a/Engine/Shaders/Private/Substrate/Substrate.ush b/Engine/Shaders/Private/Substrate/Substrate.ush
--- a/Engine/Shaders/Private/Substrate/Substrate.ush
+++ b/Engine/Shaders/Private/Substrate/Substrate.ush
@@ -29,9 +29,14 @@
 // DXC has issue with unrolling certain part of the code. To address temporarly the warning-as-error issues, 
 // we redefine the unroll function and fall back to loop on DXC
 // To remove this, one must test PC + DX12 + RayTracing and make sure raytracing shaders compile.
-#define SUBSTRATE_COMPILER (COMPILER_DXC)
+#define SUBSTRATE_COMPILER (COMPILER_VULKAN)
+#if SUBSTRATE_COMPILER == 1
+#define SUBSTRATE_UNROLL LOOP
+#define SUBSTRATE_UNROLL_N(X) LOOP
+#else
 #define SUBSTRATE_UNROLL UNROLL
 #define SUBSTRATE_UNROLL_N(X) UNROLL_N(X)
+#endif
 
 // During the base pass or forward rendering, shared local bases are simply available from registers.
 // If a BSDF is loaded for processing from the Substrate buffer, normals are unpacked on demand to reduce VGPR pressure and increase occupancy.
diff --git a/Engine/Shaders/Private/Substrate/SubstrateVisualize.usf b/Engine/Shaders/Private/Substrate/SubstrateVisualize.usf
--- a/Engine/Shaders/Private/Substrate/SubstrateVisualize.usf
+++ b/Engine/Shaders/Private/Substrate/SubstrateVisualize.usf
@@ -799,6 +799,8 @@ float2 GetBranchOffset(FSubstrateOperator Op, bool bIsLeft)
 [numthreads(1, 1, 1)]
 void MaterialDebugSubstrateTreeCS(uint3 DispatchThreadId : SV_DispatchThreadID)
 {
+// UE-194712 - Temporary workaround with Vulkan shader compilation
+#if !COMPILER_VULKAN
 	const uint2 PixelPos = (float2(View.CursorPosition) * View.ViewResolutionFraction);
 
 	FSubstrateAddressing SubstrateAddressing = GetSubstratePixelDataByteOffset(PixelPos, uint2(View.BufferSizeAndInvSize.xy), Substrate.MaxBytesPerPixel);
@@ -839,12 +841,8 @@ void MaterialDebugSubstrateTreeCS(uint3 DispatchThreadId : SV_DispatchThreadID)
 
 		Context = InitShaderPrintContext(true, uint2(50, BSDF_ELEMENT_STARTY));
 		{
-		#if SUBSTRATE_COMPILER == 1
-			LOOP
-		#else
 			SubstrateTreeHeader.OperatorCount = min(SubstrateTreeHeader.OperatorCount, 2);
 			SUBSTRATE_UNROLL_N(2)
-		#endif
 			for (int i = 0; i < SubstrateTreeHeader.OperatorCount; i++)
 			{
 				FSubstrateOperator Op = SubstrateUnpackInSubstrateTreeOperator(
@@ -920,12 +918,8 @@ void MaterialDebugSubstrateTreeCS(uint3 DispatchThreadId : SV_DispatchThreadID)
 		}
 
 		{
-		#if SUBSTRATE_COMPILER == 1
-			LOOP
-		#else
 			SubstrateTreeHeader.BSDFCount = min(SubstrateTreeHeader.BSDFCount, 2);
 			SUBSTRATE_UNROLL_N(2)
-		#endif
 			for (int i = 0; i < SubstrateTreeHeader.BSDFCount; i++)
 			{
 				FSubstrateBSDF BSDF = SubstrateUnpackInSubstrateTreeBSDF(
@@ -1108,6 +1102,7 @@ void MaterialDebugSubstrateTreeCS(uint3 DispatchThreadId : SV_DispatchThreadID)
 	#endif // SUBSTRATE_COMPILER
 
 	}
+#endif
 }
 
 #endif // SHADER_DEBUGSUBSTRATETREE_CS
@@ -1439,12 +1434,8 @@ void MaterialDebugSubstrateTreePS(
 
 		// SUBSTRATE_TODO: read BSDF and Operators "on the fly".
 		{
-		#if SUBSTRATE_COMPILER == 1
-			LOOP
-		#else
 			SubstrateTreeHeader.BSDFCount = min(SubstrateTreeHeader.BSDFCount, 2);
 			SUBSTRATE_UNROLL_N(2)
-		#endif
 			for (int i = 0; i < SubstrateTreeHeader.BSDFCount; i++)
 			{
 				FSubstrateBSDF BSDF = SubstrateUnpackInSubstrateTreeBSDF(
@@ -1488,12 +1479,8 @@ void MaterialDebugSubstrateTreePS(
 		
 		// This is a really basic material topology debug. SUBSTRATE_TODO would be to draw the material graph itself with anottation
 		{
-		#if SUBSTRATE_COMPILER == 1
-			LOOP
-		#else
 			SubstrateTreeHeader.BSDFCount = min(SubstrateTreeHeader.BSDFCount, 2);
 			SUBSTRATE_UNROLL_N(2)
-		#endif
 			for (int i = 0; i < SubstrateTreeHeader.BSDFCount; i++)
 			{
 				FSubstrateBSDF BSDF = SubstrateUnpackInSubstrateTreeBSDF(
