Removing incorrect access of a uniform buffer constant.

diff --git a/Engine/Shaders/Private/ForwardLightingCommon.ush b/Engine/Shaders/Private/ForwardLightingCommon.ush
--- a/Engine/Shaders/Private/ForwardLightingCommon.ush
+++ b/Engine/Shaders/Private/ForwardLightingCommon.ush
@@ -405,7 +405,7 @@ half3 GetImageBasedReflectionSpecular(FMaterialPixelParameters MaterialParameter
 	LumenRadianceCacheCoverage.bValid = false;
 
 	// Lumen Radiance Cache for reflections on translucency
-	if (TranslucentBasePass_FinalProbeResolution > 0 && !bUseLumenFrontLayerReflection)
+	if (FinalProbeResolution > 0 && !bUseLumenFrontLayerReflection)
 	{
 		float ClipmapDitherRandom = InterleavedGradientNoise(MaterialParameters.SvPosition.xy, View.StateFrameIndexMod8);
 		LumenRadianceCacheCoverage = GetRadianceCacheCoverage(LWCHackToFloat(MaterialParameters.AbsoluteWorldPosition), RayDirection, ClipmapDitherRandom);
