LFV - added near clip plane support to tiled culling.

diff --git a/Engine/Shaders/Private/LocalFogVolumes/LocalFogVolumeTiledCulling.usf b/Engine/Shaders/Private/LocalFogVolumes/LocalFogVolumeTiledCulling.usf
--- a/Engine/Shaders/Private/LocalFogVolumes/LocalFogVolumeTiledCulling.usf
+++ b/Engine/Shaders/Private/LocalFogVolumes/LocalFogVolumeTiledCulling.usf
@@ -37,7 +37,7 @@ void LocalFogVolumeTiledCullingCS(uint3 ThreadId : SV_DispatchThreadID)
 	TilePlanes[1] = lerp(LeftPlane, RightPlane, MaxTileLerp.xxxx);		// Right tile plane
 	TilePlanes[2] = lerp(-TopPlane, -BottomPlane, MinTileLerp.yyyy);	// Bottom tile plane
 	TilePlanes[3] = lerp(-TopPlane, -BottomPlane, MaxTileLerp.yyyy);	// Top plane
-//	TilePlanes[4] = NearPlane;		LFV_TODO							// Near plane
+	TilePlanes[4] = NearPlane;											// Near plane
 	// Swap right and top plane to go from lerpable planes to planes we can use for culling
 	TilePlanes[1] *= -1.0f;
 	TilePlanes[3] *= -1.0f;
@@ -45,7 +45,7 @@ void LocalFogVolumeTiledCullingCS(uint3 ThreadId : SV_DispatchThreadID)
 #if 0
 	// Optional plane normalise
 	UNROLL
-	for (uint i = 0; i < 4; ++i)
+	for (uint i = 0; i < 5; ++i)
 	{
 		float SquareSum = TilePlanes[i].x * TilePlanes[i].x + TilePlanes[i].y * TilePlanes[i].y + TilePlanes[i].z * TilePlanes[i].z;
 		float Scale = 1.0 / sqrt(SquareSum);
@@ -60,7 +60,7 @@ void LocalFogVolumeTiledCullingCS(uint3 ThreadId : SV_DispatchThreadID)
 
 		bool bCompletelyOutside = false;
 		UNROLL
-		for (uint i = 0; i < 4; ++i)
+		for (uint i = 0; i < 5; ++i)
 		{
 			if (dot(TilePlanes[i].xyz, Sphere.xyz) - TilePlanes[i].w > Sphere.w)
 			{
diff --git a/Engine/Source/Runtime/Renderer/Private/LocalFogVolumeRendering.cpp b/Engine/Source/Runtime/Renderer/Private/LocalFogVolumeRendering.cpp
--- a/Engine/Source/Runtime/Renderer/Private/LocalFogVolumeRendering.cpp
+++ b/Engine/Source/Runtime/Renderer/Private/LocalFogVolumeRendering.cpp
@@ -194,11 +194,11 @@ static void LocalFogVolumeViewTiledCullingPass(FViewInfo& View, FRDGBuilder& Gra
 		}
 	};
 	// Using world space plane for now. LFV_TODO: do computation in view space.
-	ConvertPlanToVector4f(PassParameters->LeftPlane,	View.CullingFrustum.Planes[0], false);
-	ConvertPlanToVector4f(PassParameters->RightPlane,	View.CullingFrustum.Planes[1], true);
-	ConvertPlanToVector4f(PassParameters->TopPlane,		View.CullingFrustum.Planes[2], true);
-	ConvertPlanToVector4f(PassParameters->BottomPlane,	View.CullingFrustum.Planes[3], false);
-//	ConvertPlanToVector4f(PassParameters->NearPlane,	View.CullingFrustum.Planes[4], false); // LFV TODO View.ViewMatrices.GetViewProjectionMatrix().GetNearPlane
+	ConvertPlanToVector4f(PassParameters->LeftPlane,	View.CullingFrustum.Planes[0],	false);
+	ConvertPlanToVector4f(PassParameters->RightPlane,	View.CullingFrustum.Planes[1],	true);
+	ConvertPlanToVector4f(PassParameters->TopPlane,		View.CullingFrustum.Planes[2],	true);
+	ConvertPlanToVector4f(PassParameters->BottomPlane,	View.CullingFrustum.Planes[3],	false);
+	ConvertPlanToVector4f(PassParameters->NearPlane,	View.NearClippingPlane,			false);
 
 	ERDGPassFlags PassFlag = ERDGPassFlags::Compute; // LFV_TODO try ERDGPassFlags::AsyncCompute later
 
