LFV - added clipping to the tile pass.

diff --git a/Engine/Shaders/Private/LocalFogVolumes/LocalFogVolumeSplat.usf b/Engine/Shaders/Private/LocalFogVolumes/LocalFogVolumeSplat.usf
--- a/Engine/Shaders/Private/LocalFogVolumes/LocalFogVolumeSplat.usf
+++ b/Engine/Shaders/Private/LocalFogVolumes/LocalFogVolumeSplat.usf
@@ -162,7 +162,8 @@ void LocalFogVolumeTiledPS(
 
 	uint2 TilePos = uint2(OtherVertexOutput.TileData >> 16, OtherVertexOutput.TileData & 0xFFFF);
 	uint LFVCount = LFVTileDataTex[uint3(TilePos,0)];
-	float4 LuminanceColorTransmittance = float4(0.0, 0.0, 0.0, 1.0);
+	const float4 DefaultColorTransmittance = float4(0.0, 0.0, 0.0, 1.0);
+	float4 LuminanceColorTransmittance = DefaultColorTransmittance;
 	for (int LFVIndex = 0; LFVIndex < LFVCount; ++LFVIndex)
 	{
 		uint InstanceIndex = LFVTileDataTex[uint3(TilePos, 1 + LFVIndex)];
@@ -182,6 +183,11 @@ void LocalFogVolumeTiledPS(
 		LuminanceColorTransmittance.rgb * PrimaryView.PreExposure,
 		LuminanceColorTransmittance.a);
 
+	if (all(DefaultColorTransmittance == OutColor))
+	{
+		clip(-1.0f);
+	}
+
 	// Debug
 	OutColor += float4(float(LFVCount) * 0.1 * PrimaryView.PreExposure, 0.0, 0.0, 0.0);
 }
