Fix rect light filtering border with non power of two texture.

This fix invalid colored at texture border which could happen since the atlas is never cleared.

diff --git a/Engine/Shaders/Private/RectLightAtlas.usf b/Engine/Shaders/Private/RectLightAtlas.usf
--- a/Engine/Shaders/Private/RectLightAtlas.usf
+++ b/Engine/Shaders/Private/RectLightAtlas.usf
@@ -143,12 +143,25 @@ void MainPS(
 		const uint4 SrcRect = SrcSlotBuffer[InId];
 		const uint2 SrcRectOffset = SrcRect.xy;
 		const uint2 SrcRectResolution = SrcRect.zw;
+		const uint2 SrcRectMax = SrcRectOffset + SrcRectResolution;
 
-		const float4 Color0 = SourceAtlasTexture.Load(uint3(SrcRectOffset + SrcCoord + float2(0,0), 0));
-		const float4 Color1 = SourceAtlasTexture.Load(uint3(SrcRectOffset + SrcCoord + float2(1,0), 0));
-		const float4 Color2 = SourceAtlasTexture.Load(uint3(SrcRectOffset + SrcCoord + float2(0,1), 0));
-		const float4 Color3 = SourceAtlasTexture.Load(uint3(SrcRectOffset + SrcCoord + float2(1,1), 0));
-		OutTexture = (Color0 + Color1 + Color2 + Color3) * 0.25f;
+		const float2 Coord00 = SrcRectOffset + SrcCoord + float2(0,0);
+		const float2 Coord10 = SrcRectOffset + SrcCoord + float2(1,0);
+		const float2 Coord01 = SrcRectOffset + SrcCoord + float2(0,1);
+		const float2 Coord11 = SrcRectOffset + SrcCoord + float2(1,1);
+
+		const float W00 = all(Coord00 < SrcRectMax) ? 1.f : 0.f;
+		const float W10 = all(Coord10 < SrcRectMax) ? 1.f : 0.f;
+		const float W01 = all(Coord01 < SrcRectMax) ? 1.f : 0.f;
+		const float W11 = all(Coord11 < SrcRectMax) ? 1.f : 0.f;
+		const float InvW = 1.f / max(1.f, float(W00 + W10 + W01 + W11));
+
+		const float4 Color0 = SourceAtlasTexture.Load(uint3(Coord00, 0));
+		const float4 Color1 = SourceAtlasTexture.Load(uint3(Coord10, 0));
+		const float4 Color2 = SourceAtlasTexture.Load(uint3(Coord01, 0));
+		const float4 Color3 = SourceAtlasTexture.Load(uint3(Coord11, 0));
+
+		OutTexture = (W00*Color0 + W10*Color1 + W01*Color2 + W11*Color3) * InvW;
 	}
 	else
 	{
@@ -163,7 +176,8 @@ void MainPS(
 		const uint2 SrcRectResolution = SrcRect.zw;
 
 		const float2 MinCoord = SrcRectOffset;
-		const float2 MaxCoord = SrcRectOffset + SrcRectResolution;
+		const float2 MaxCoord = SrcRectOffset + SrcRectResolution - float2(1,1);
+
 		const float2 InvAtlasResolution = 1.0f / float2(SrcAtlasResolution);
 
 		float3 AccColor = 0;
@@ -174,7 +188,8 @@ void MainPS(
 		for (int x = -KernelExtent; x <= KernelExtent; ++x)
 		{
 			const float2 Offset = float2(x, y);
-			const float2 Coord = SrcRectOffset + SrcCoord + Offset;// +float2(0.5f, 0.5f);
+			const float2 Coord = SrcRectOffset + SrcCoord + Offset;
+
 			if (all(Coord >= MinCoord) && all(Coord < MaxCoord))
 			{
 				const float2 UV = Coord * InvAtlasResolution;
@@ -186,7 +201,7 @@ void MainPS(
 				AccW += W;
 			}
 		}
-		OutTexture = float4(AccColor / AccW, 1.0f);
+		OutTexture = float4(AccColor / max(0.001f, AccW), 1.0f);
 	}
 }
 
diff --git a/Engine/Source/Runtime/Renderer/Private/RectLightTextureManager.cpp b/Engine/Source/Runtime/Renderer/Private/RectLightTextureManager.cpp
--- a/Engine/Source/Runtime/Renderer/Private/RectLightTextureManager.cpp
+++ b/Engine/Source/Runtime/Renderer/Private/RectLightTextureManager.cpp
@@ -217,9 +217,11 @@ bool CanContain(const FAtlasRect& Outside, const FAtlasRect& Inside)
 
 static FIntPoint ToMIP(const FIntPoint& InMip, uint32 MipIndex)
 {
+	const int32 Div = 1 << MipIndex;
+
 	FIntPoint Out;
-	Out.X = InMip.X >> MipIndex;
-	Out.Y = InMip.Y >> MipIndex;
+	Out.X = FMath::DivideAndRoundUp(InMip.X,Div);
+	Out.Y = FMath::DivideAndRoundUp(InMip.Y,Div);
 	return Out;
 }
 
