LFV - Translucent now sample LFV using the per tile culled volumes.

#prelfight shader

diff --git a/Engine/Shaders/Private/BasePassVertexShader.usf b/Engine/Shaders/Private/BasePassVertexShader.usf
--- a/Engine/Shaders/Private/BasePassVertexShader.usf
+++ b/Engine/Shaders/Private/BasePassVertexShader.usf
@@ -108,25 +108,33 @@ void Main(
 
 	// Apply local fog volume to if enabled. It is always enabled per vertex as of today, that is why it is out ofr the NEEDS_BASEPASS_VERTEX_FOGGING defines.
 #if LOCAL_FOG_VOLUME_ON_TRANSLUCENT
-	float4 LFVContribution = float4(0.0f, 0.0f, 0.0f, 1.0f);
-	for (uint LVFIndex = 0; LVFIndex < LFVInstcount; ++LVFIndex)
+	float4 VertexClipSpacePosition = mul(float4(WorldPosition.xyz, 1), ResolvedView.TranslatedWorldToClip);
+	float2 SvPosition = (VertexClipSpacePosition.xy / VertexClipSpacePosition.w * float2(.5f, -.5f) + .5f) * ResolvedView.ViewSizeAndInvSize.xy;
+	uint2 TilePos = clamp(uint2(SvPosition.xy / float(LFVTilePixelSize)), uint2(0, 0), LFVTileDataResolution - 1);
+
+	uint LFVCount = LFVTileDataTex[uint3(TilePos, 0)];
+	const float4 DefaultColorTransmittance = float4(0.0, 0.0, 0.0, 1.0);
+	float4 LFVContribution = DefaultColorTransmittance;
+	for (int LFVIndex = 0; LFVIndex < LFVCount; ++LFVIndex)
 	{
+		uint InstanceIndex = LFVTileDataTex[uint3(TilePos, 1 + LFVIndex)];
+
 		float3 CamRayTranslatedWorldOrigin	= ResolvedView.TranslatedWorldCameraOrigin;
 		float3 CamRayWorldOrigin			= CamRayTranslatedWorldOrigin - LWCHackToFloat(ResolvedView.PreViewTranslation);
 		float3 CamRayTranslatedWorldDir		= normalize(WorldPosition - CamRayTranslatedWorldOrigin);	// WorldPosition is expressed in TranslatedWorld space
 		float3 CamRayWorldDir				= CamRayTranslatedWorldDir;
-
+		
 		float3 DepthBufferTranslatedWorldPos = WorldPosition.xyz;
 
-		const float4 LFVInstanceContribution = GetLFVInstanceContribution(
-			LVFIndex,
+		float4 LFVLuminanceColorTransmittance = GetLFVInstanceContribution(
+			InstanceIndex,
 			CamRayWorldOrigin,
 			CamRayWorldDir,
 			CamRayTranslatedWorldOrigin,
 			CamRayTranslatedWorldDir,
 			DepthBufferTranslatedWorldPos);
 
-		LFVContribution = float4(LFVContribution.rgb * LFVInstanceContribution.aaa + LFVInstanceContribution.rgb, LFVContribution.a * LFVInstanceContribution.a);
+		LFVContribution = float4(LFVContribution.rgb * LFVLuminanceColorTransmittance.a + LFVLuminanceColorTransmittance.rgb, LFVContribution.a * LFVLuminanceColorTransmittance.a);
 	}
 
 	Output.BasePassInterpolants.VertexFog = float4(LFVContribution.rgb + Output.BasePassInterpolants.VertexFog.rgb * LFVContribution.a, LFVContribution.a * Output.BasePassInterpolants.VertexFog.a);
diff --git a/Engine/Shaders/Private/MobileBasePassVertexShader.usf b/Engine/Shaders/Private/MobileBasePassVertexShader.usf
--- a/Engine/Shaders/Private/MobileBasePassVertexShader.usf
+++ b/Engine/Shaders/Private/MobileBasePassVertexShader.usf
@@ -111,9 +111,17 @@ void Main(
 
 	// Apply local fog volume to if enabled
 #if LOCAL_FOG_VOLUME_ON_TRANSLUCENT
-	float4 LFVContribution = float4(0.0f, 0.0f, 0.0f, 1.0f);
-	for (uint LVFIndex = 0; LVFIndex < LFVInstcount; ++LVFIndex)
+	float4 VertexClipSpacePosition = mul(float4(WorldPosition.xyz, 1), ResolvedView.TranslatedWorldToClip);
+	float2 SvPosition = (VertexClipSpacePosition.xy / VertexClipSpacePosition.w * float2(.5f, -.5f) + .5f) * ResolvedView.ViewSizeAndInvSize.xy;
+	uint2 TilePos = clamp(uint2(SvPosition.xy / float(LFVTilePixelSize)), uint2(0, 0), LFVTileDataResolution - 1);
+
+	uint LFVCount = LFVTileDataTex[uint3(TilePos, 0)];
+	const float4 DefaultColorTransmittance = float4(0.0, 0.0, 0.0, 1.0);
+	float4 LFVContribution = DefaultColorTransmittance;
+	for (int LFVIndex = 0; LFVIndex < LFVCount; ++LFVIndex)
 	{
+		uint InstanceIndex = LFVTileDataTex[uint3(TilePos, 1 + LFVIndex)];
+
 		float3 CamRayTranslatedWorldOrigin = ResolvedView.TranslatedWorldCameraOrigin;
 		float3 CamRayWorldOrigin = CamRayTranslatedWorldOrigin - LWCHackToFloat(ResolvedView.PreViewTranslation);
 		float3 CamRayTranslatedWorldDir = normalize(WorldPosition - CamRayTranslatedWorldOrigin);	// WorldPosition is expressed in TranslatedWorld space
@@ -121,15 +129,15 @@ void Main(
 
 		float3 DepthBufferTranslatedWorldPos = WorldPosition.xyz;
 
-		const float4 LFVInstanceContribution = GetLFVInstanceContribution(
-			LVFIndex,
+		float4 LFVLuminanceColorTransmittance = GetLFVInstanceContribution(
+			InstanceIndex,
 			CamRayWorldOrigin,
 			CamRayWorldDir,
 			CamRayTranslatedWorldOrigin,
 			CamRayTranslatedWorldDir,
 			DepthBufferTranslatedWorldPos);
 
-		LFVContribution = float4(LFVContribution.rgb * LFVInstanceContribution.aaa + LFVInstanceContribution.rgb, LFVContribution.a * LFVInstanceContribution.a);
+		LFVContribution = float4(LFVContribution.rgb * LFVLuminanceColorTransmittance.a + LFVLuminanceColorTransmittance.rgb, LFVContribution.a * LFVLuminanceColorTransmittance.a);
 	}
 
 	VertexFog = float4(LFVContribution.rgb + VertexFog.rgb * LFVContribution.a, LFVContribution.a * VertexFog.a);
