LFV - made contribution code evfaluation common.

#prelfight shader

diff --git a/Engine/Shaders/Private/BasePassVertexShader.usf b/Engine/Shaders/Private/BasePassVertexShader.usf
--- a/Engine/Shaders/Private/BasePassVertexShader.usf
+++ b/Engine/Shaders/Private/BasePassVertexShader.usf
@@ -108,36 +108,11 @@ void Main(
 
 	// Apply local fog volume to if enabled. It is always enabled per vertex as of today, that is why it is out ofr the NEEDS_BASEPASS_VERTEX_FOGGING defines.
 #if LOCAL_FOG_VOLUME_ON_TRANSLUCENT
-	float4 VertexClipSpacePosition = mul(float4(WorldPosition.xyz, 1), ResolvedView.TranslatedWorldToClip);
-	float2 SvPosition = (VertexClipSpacePosition.xy / VertexClipSpacePosition.w * float2(.5f, -.5f) + .5f) * ResolvedView.ViewSizeAndInvSize.xy;
-	uint2 TilePos = clamp(uint2(SvPosition.xy / float(LFVTilePixelSize)), uint2(0, 0), LFVTileDataResolution - 1);
-
-	uint LFVCount = LFVTileDataTex[uint3(TilePos, 0)];
-	const float4 DefaultColorTransmittance = float4(0.0, 0.0, 0.0, 1.0);
-	float4 LFVContribution = DefaultColorTransmittance;
-	for (int LFVIndex = 0; LFVIndex < LFVCount; ++LFVIndex)
-	{
-		uint InstanceIndex = LFVTileDataTex[uint3(TilePos, 1 + LFVIndex)];
-
-		float3 CamRayTranslatedWorldOrigin	= ResolvedView.TranslatedWorldCameraOrigin;
-		float3 CamRayWorldOrigin			= CamRayTranslatedWorldOrigin - LWCHackToFloat(ResolvedView.PreViewTranslation);
-		float3 CamRayTranslatedWorldDir		= normalize(WorldPosition - CamRayTranslatedWorldOrigin);	// WorldPosition is expressed in TranslatedWorld space
-		float3 CamRayWorldDir				= CamRayTranslatedWorldDir;
-		
-		float3 DepthBufferTranslatedWorldPos = WorldPosition.xyz;
-
-		float4 LFVLuminanceColorTransmittance = GetLFVInstanceContribution(
-			InstanceIndex,
-			CamRayWorldOrigin,
-			CamRayWorldDir,
-			CamRayTranslatedWorldOrigin,
-			CamRayTranslatedWorldDir,
-			DepthBufferTranslatedWorldPos);
-
-		LFVContribution = float4(LFVContribution.rgb * LFVLuminanceColorTransmittance.a + LFVLuminanceColorTransmittance.rgb, LFVContribution.a * LFVLuminanceColorTransmittance.a);
-	}
-
-	Output.BasePassInterpolants.VertexFog = float4(LFVContribution.rgb + Output.BasePassInterpolants.VertexFog.rgb * LFVContribution.a, LFVContribution.a * Output.BasePassInterpolants.VertexFog.a);
+	float4 VertexClipSpacePosition			= mul(float4(WorldPosition.xyz, 1), ResolvedView.TranslatedWorldToClip);
+	float2 SvPosition						= (VertexClipSpacePosition.xy / VertexClipSpacePosition.w * float2(.5f, -.5f) + .5f) * ResolvedView.ViewSizeAndInvSize.xy;
+	uint2 TilePos							= clamp(uint2(SvPosition.xy / float(LFVTilePixelSize)), uint2(0, 0), LFVTileDataResolution - 1);
+	float4 LFVContribution					= GetLFVContribution(ResolvedView, TilePos, WorldPosition);
+	Output.BasePassInterpolants.VertexFog	= float4(LFVContribution.rgb + Output.BasePassInterpolants.VertexFog.rgb * LFVContribution.a, LFVContribution.a * Output.BasePassInterpolants.VertexFog.a);
 #endif // LOCAL_FOG_VOLUME_ON_TRANSLUCENT
 
 #if NEEDS_BASEPASS_VERTEX_FOGGING
diff --git a/Engine/Shaders/Private/LocalFogVolumes/LocalFogVolumeCommon.ush b/Engine/Shaders/Private/LocalFogVolumes/LocalFogVolumeCommon.ush
--- a/Engine/Shaders/Private/LocalFogVolumes/LocalFogVolumeCommon.ush
+++ b/Engine/Shaders/Private/LocalFogVolumes/LocalFogVolumeCommon.ush
@@ -375,4 +375,50 @@ float4 GetLFVInstanceContribution(
 
 
 
+// Evaluate the local fog volume contribution for a world position.
+// Returns a float4 with Luminance Color in RGB and Transmittance in A.
+float4 GetLFVContribution(
+	in ViewState ResolvedView,
+	in uint2 TilePos,
+	in float3 TranslatedWorldPosition,
+	inout uint LFVCount)
+{
+	float3 CamRayTranslatedWorldOrigin = ResolvedView.TranslatedWorldCameraOrigin;
+	float3 PreViewTranslation = LWCHackToFloat(ResolvedView.PreViewTranslation);
+
+	LFVCount = LFVTileDataTex[uint3(TilePos, 0)];
+	const float4 DefaultColorTransmittance = float4(0.0, 0.0, 0.0, 1.0);
+	float4 LFVContribution = DefaultColorTransmittance;
+	for (int LFVIndex = 0; LFVIndex < LFVCount; ++LFVIndex)
+	{
+		uint InstanceIndex = LFVTileDataTex[uint3(TilePos, 1 + LFVIndex)];
+
+		float3 CamRayWorldOrigin = CamRayTranslatedWorldOrigin - PreViewTranslation;
+		float3 CamRayTranslatedWorldDir = normalize(TranslatedWorldPosition - CamRayTranslatedWorldOrigin);
+		float3 CamRayWorldDir = CamRayTranslatedWorldDir;
+
+		float3 DepthBufferTranslatedWorldPos = TranslatedWorldPosition.xyz;
+
+		float4 LFVLuminanceColorTransmittance = GetLFVInstanceContribution(
+			InstanceIndex,
+			CamRayWorldOrigin,
+			CamRayWorldDir,
+			CamRayTranslatedWorldOrigin,
+			CamRayTranslatedWorldDir,
+			DepthBufferTranslatedWorldPos);
+
+		LFVContribution = float4(LFVContribution.rgb * LFVLuminanceColorTransmittance.a + LFVLuminanceColorTransmittance.rgb, LFVContribution.a * LFVLuminanceColorTransmittance.a);
+	}
+	return LFVContribution;
+}
+
+float4 GetLFVContribution(
+	in ViewState ResolvedView,
+	in uint2 TilePos,
+	in float3 TranslatedWorldPosition)
+{
+	uint LFVCount = 0;
+	return GetLFVContribution(ResolvedView, TilePos, TranslatedWorldPosition, LFVCount);
+}
+
 
diff --git a/Engine/Shaders/Private/LocalFogVolumes/LocalFogVolumeSplat.usf b/Engine/Shaders/Private/LocalFogVolumes/LocalFogVolumeSplat.usf
--- a/Engine/Shaders/Private/LocalFogVolumes/LocalFogVolumeSplat.usf
+++ b/Engine/Shaders/Private/LocalFogVolumes/LocalFogVolumeSplat.usf
@@ -153,50 +153,34 @@ void LocalFogVolumeTiledPS(
 {
 	ResolvedView = ResolveView();
 
-	float3 CamRayTranslatedWorldOrigin = PrimaryView.TranslatedWorldCameraOrigin;
-	float3 CamRayTranslatedWorldDir = normalize(SvPositionToTranslatedWorld(float4(SVPos.xy, 0.5, 1.0)));
-	float3 CamRayWorldDir = CamRayTranslatedWorldDir;
-	float3 CamRayWorldOrigin = CamRayTranslatedWorldOrigin - LWCHackToFloat(PrimaryView.PreViewTranslation);
-
 #if SHADING_PATH_MOBILE
 	// LookupDeviceZ with uint2 SVPos.xy as input does not work on mobile.
 	float2 UV = SVPos * View.BufferSizeAndInvSize.zw;
-	const float DeviceZ = LookupDeviceZ(UV);
+	float DeviceZ = LookupDeviceZ(UV);
 #else
-	const float DeviceZ = LookupDeviceZ(uint2(SVPos.xy));
+	float DeviceZ = LookupDeviceZ(uint2(SVPos.xy));
+#endif
+#if HAS_INVERTED_Z_BUFFER
+	DeviceZ = max(0.000000000001, DeviceZ);	// Make sure SvPositionToWorld does not return bad values when DeviceZ is far=0 when using inverted z
 #endif
-
 	float3 DepthBufferTranslatedWorldPos = SvPositionToTranslatedWorld(float4(SVPos.xy, DeviceZ, 1.0));
 
 	uint2 TilePos = LFVUnpackTile(OtherVertexOutput.TileData);
-	uint LFVCount = LFVTileDataTex[uint3(TilePos,0)];
-	const float4 DefaultColorTransmittance = float4(0.0, 0.0, 0.0, 1.0);
-	float4 LuminanceColorTransmittance = DefaultColorTransmittance;
-	for (int LFVIndex = 0; LFVIndex < LFVCount; ++LFVIndex)
-	{
-		uint InstanceIndex = LFVTileDataTex[uint3(TilePos, 1 + LFVIndex)];
-
-		float4 LFVLuminanceColorTransmittance = GetLFVInstanceContribution(
-			InstanceIndex,
-			CamRayWorldOrigin,
-			CamRayWorldDir,
-			CamRayTranslatedWorldOrigin,
-			CamRayTranslatedWorldDir,
-			DepthBufferTranslatedWorldPos);
 
-		LuminanceColorTransmittance = float4(LuminanceColorTransmittance.rgb * LFVLuminanceColorTransmittance.a + LFVLuminanceColorTransmittance.rgb, LuminanceColorTransmittance.a * LFVLuminanceColorTransmittance.a);
-	}
+	uint LFVCount = 0;
+	float4 LFVContribution = GetLFVContribution(ResolvedView, TilePos, DepthBufferTranslatedWorldPos, LFVCount);
 	
 	OutColor = float4(
-		LuminanceColorTransmittance.rgb * PrimaryView.PreExposure,
-		LuminanceColorTransmittance.a);
+		LFVContribution.rgb * PrimaryView.PreExposure,
+		LFVContribution.a);
 
 	bool bDoClip = true;
 #if USE_DEVELOPMENT_SHADERS
 	bDoClip = LocalFogVolumeTileDebug == 2;
 #endif
 
-	if (all(DefaultColorTransmittance == OutColor) && bDoClip)
+	const float4 IdentityLuminanceTransmittance = float4(0.0, 0.0, 0.0, 1.0);
+	if (all(LFVContribution == IdentityLuminanceTransmittance) && bDoClip)
 	{
 		clip(-1.0f);
 	}
diff --git a/Engine/Shaders/Private/MobileBasePassVertexShader.usf b/Engine/Shaders/Private/MobileBasePassVertexShader.usf
--- a/Engine/Shaders/Private/MobileBasePassVertexShader.usf
+++ b/Engine/Shaders/Private/MobileBasePassVertexShader.usf
@@ -111,36 +111,11 @@ void Main(
 
 	// Apply local fog volume to if enabled
 #if LOCAL_FOG_VOLUME_ON_TRANSLUCENT
-	float4 VertexClipSpacePosition = mul(float4(WorldPosition.xyz, 1), ResolvedView.TranslatedWorldToClip);
-	float2 SvPosition = (VertexClipSpacePosition.xy / VertexClipSpacePosition.w * float2(.5f, -.5f) + .5f) * ResolvedView.ViewSizeAndInvSize.xy;
-	uint2 TilePos = clamp(uint2(SvPosition.xy / float(LFVTilePixelSize)), uint2(0, 0), LFVTileDataResolution - 1);
-
-	uint LFVCount = LFVTileDataTex[uint3(TilePos, 0)];
-	const float4 DefaultColorTransmittance = float4(0.0, 0.0, 0.0, 1.0);
-	float4 LFVContribution = DefaultColorTransmittance;
-	for (int LFVIndex = 0; LFVIndex < LFVCount; ++LFVIndex)
-	{
-		uint InstanceIndex = LFVTileDataTex[uint3(TilePos, 1 + LFVIndex)];
-
-		float3 CamRayTranslatedWorldOrigin = ResolvedView.TranslatedWorldCameraOrigin;
-		float3 CamRayWorldOrigin = CamRayTranslatedWorldOrigin - LWCHackToFloat(ResolvedView.PreViewTranslation);
-		float3 CamRayTranslatedWorldDir = normalize(WorldPosition - CamRayTranslatedWorldOrigin);	// WorldPosition is expressed in TranslatedWorld space
-		float3 CamRayWorldDir = CamRayTranslatedWorldDir;
-
-		float3 DepthBufferTranslatedWorldPos = WorldPosition.xyz;
-
-		float4 LFVLuminanceColorTransmittance = GetLFVInstanceContribution(
-			InstanceIndex,
-			CamRayWorldOrigin,
-			CamRayWorldDir,
-			CamRayTranslatedWorldOrigin,
-			CamRayTranslatedWorldDir,
-			DepthBufferTranslatedWorldPos);
-
-		LFVContribution = float4(LFVContribution.rgb * LFVLuminanceColorTransmittance.a + LFVLuminanceColorTransmittance.rgb, LFVContribution.a * LFVLuminanceColorTransmittance.a);
-	}
-
-	VertexFog = float4(LFVContribution.rgb + VertexFog.rgb * LFVContribution.a, LFVContribution.a * VertexFog.a);
+	float4 VertexClipSpacePosition	= mul(float4(WorldPosition.xyz, 1), ResolvedView.TranslatedWorldToClip);
+	float2 SvPosition				= (VertexClipSpacePosition.xy / VertexClipSpacePosition.w * float2(.5f, -.5f) + .5f) * ResolvedView.ViewSizeAndInvSize.xy;
+	uint2 TilePos					= clamp(uint2(SvPosition.xy / float(LFVTilePixelSize)), uint2(0, 0), LFVTileDataResolution - 1);
+	float4 LFVContribution			= GetLFVContribution(ResolvedView, TilePos, WorldPosition);
+	VertexFog						= float4(LFVContribution.rgb + VertexFog.rgb * LFVContribution.a, LFVContribution.a * VertexFog.a);
 #endif // LOCAL_FOG_VOLUME_ON_TRANSLUCENT
 
 #if PACK_INTERPOLANTS
