LFV - fixed unwanted per vertex fog volume computations.

#prelfight shader

diff --git a/Engine/Shaders/Private/BasePassCommon.ush b/Engine/Shaders/Private/BasePassCommon.ush
--- a/Engine/Shaders/Private/BasePassCommon.ush
+++ b/Engine/Shaders/Private/BasePassCommon.ush
@@ -66,8 +66,8 @@
 #define NEEDS_BASEPASS_CLOUD_SHADOW_INTERPOLATOR (SUPPORT_CLOUD_SHADOW_ON_FORWARD_LIT_TRANSLUCENT && IS_MATERIAL_TRANSLUCENT_AND_LIT && !TRANSLUCENCY_PERVERTEX_FORWARD_SHADING && !TRANSLUCENCY_PERVERTEX_LIGHTING_VOLUME)
 #endif
 
-// Local fog volumes are always evaluated per vertex for now, even with forward shading.
-#define LOCAL_FOG_VOLUME_ON_TRANSLUCENT (PROJECT_LOCALFOGVOLUME_APPLYONTRANSLUCENT && (NEEDS_BASEPASS_VERTEX_FOGGING || NEEDS_BASEPASS_PIXEL_FOGGING))
+// Local fog volumes are always evaluated per vertex for now, even with forward shading. But never for opaque meshes rendered in forward.
+#define LOCAL_FOG_VOLUME_ON_TRANSLUCENT (PROJECT_LOCALFOGVOLUME_APPLYONTRANSLUCENT && ((NEEDS_BASEPASS_VERTEX_FOGGING && !(OPAQUE_NEEDS_BASEPASS_FOGGING && PROJECT_VERTEX_FOGGING_FOR_OPAQUE)) || NEEDS_BASEPASS_PIXEL_FOGGING) && !MATERIAL_IS_SKY)
 
 struct FSharedBasePassInterpolants
 {
diff --git a/Engine/Shaders/Private/MobileBasePassCommon.ush b/Engine/Shaders/Private/MobileBasePassCommon.ush
--- a/Engine/Shaders/Private/MobileBasePassCommon.ush
+++ b/Engine/Shaders/Private/MobileBasePassCommon.ush
@@ -23,7 +23,7 @@
 #define USE_VERTEX_FOG	(!PROJECT_MOBILE_DISABLE_VERTEX_FOG || MATERIAL_IS_SKY || (MATERIAL_ENABLE_TRANSLUCENCY_FOGGING && (MATERIALBLENDING_ANY_TRANSLUCENT || MATERIAL_SHADINGMODEL_SINGLELAYERWATER)))
 
 // Local fog volumes are always evaluated per vertex for now, on mobile, it simply follow VertexFog.
-#define LOCAL_FOG_VOLUME_ON_TRANSLUCENT (PROJECT_LOCALFOGVOLUME_APPLYONTRANSLUCENT && USE_VERTEX_FOG)
+#define LOCAL_FOG_VOLUME_ON_TRANSLUCENT (PROJECT_LOCALFOGVOLUME_APPLYONTRANSLUCENT && USE_VERTEX_FOG && !MATERIAL_IS_SKY)
 
 #define PACK_INTERPOLANTS (USE_VERTEX_FOG && NUM_VF_PACKED_INTERPOLANTS > 0 && (ES3_1_PROFILE))
 #define LANDSCAPE_BUG_WORKAROUND (IOS && IS_MOBILE_BASEPASS_VERTEX_SHADER && PACK_INTERPOLANTS)
