Fixing incorrect access of View.BufferSizeAndInvSize

diff --git a/Engine/Shaders/Private/ForwardLightingCommon.ush b/Engine/Shaders/Private/ForwardLightingCommon.ush
--- a/Engine/Shaders/Private/ForwardLightingCommon.ush
+++ b/Engine/Shaders/Private/ForwardLightingCommon.ush
@@ -96,7 +96,7 @@ float2 CalculateNearestResolvedDepthScreenUV(float2 ScreenUV, float SceneDepth)
 
 float4 GetForwardDynamicShadowFactors(float2 ScreenUV)
 {
-	int2 IntScreenUV = int2(trunc(ScreenUV * View_BufferSizeAndInvSize.xy));
+	int2 IntScreenUV = int2(trunc(ScreenUV * View.BufferSizeAndInvSize.xy));
 	float4 Value = 1.0f;
 	// Avoid sampling the white texture fallback because Load returns 0 when out-of-bound.
 	BRANCH
