Fixing incorrect access of HairStrands Uniform Buffer members.

diff --git a/Engine/Shaders/Private/HairStrands/HairStrandsAttributeCommon.ush b/Engine/Shaders/Private/HairStrands/HairStrandsAttributeCommon.ush
--- a/Engine/Shaders/Private/HairStrands/HairStrandsAttributeCommon.ush
+++ b/Engine/Shaders/Private/HairStrands/HairStrandsAttributeCommon.ush
@@ -102,6 +102,12 @@ bool IsAttributeValid(uint InOffset)
 #include "/Engine/Private/HairStrands/HairStrandsVertexFactoryCommon.ush"
 
 #if HAIR_STRAND_MESH_FACTORY
-#define HAIR_STRANDS_ATTRIBUTE_ACCESSORS(Name) HairStrandsVF_##Name
+	#if HAIR_STRANDS_PARAMETERS == 1
+		// Global constants access
+		#define HAIR_STRANDS_ATTRIBUTE_ACCESSORS(Name) HairStrandsVF_##Name
+	#else
+		// Uniform Buffer access
+		#define HAIR_STRANDS_ATTRIBUTE_ACCESSORS(Name) HairStrandsVF.##Name
+	#endif
 #include "/Engine/Private/HairStrands/HairStrandsAttributeTemplate.ush"
 #endif
diff --git a/Engine/Shaders/Private/HairStrands/HairStrandsVertexFactoryCommon.ush b/Engine/Shaders/Private/HairStrands/HairStrandsVertexFactoryCommon.ush
--- a/Engine/Shaders/Private/HairStrands/HairStrandsVertexFactoryCommon.ush
+++ b/Engine/Shaders/Private/HairStrands/HairStrandsVertexFactoryCommon.ush
@@ -50,6 +50,6 @@
 	Buffer<float>		Prefix##_CullingRadiusScaleBuffer; \
 	Buffer<uint>		Prefix##_CullingIndirectBufferArgs;
 
-#if defined(HAIR_STRANDS_PARAMETERS) && HAIR_STRANDS_PARAMETERS == 1
+#if HAIR_STRANDS_PARAMETERS == 1
 	HAIR_STRANDS_INSTANCE_PARAMETERS(HairStrandsVF)
 #endif
\ No newline at end of file
