Fix shader compilation warning

diff --git a/Engine/Shaders/Private/HairStrands/HairStrandsDebug.usf b/Engine/Shaders/Private/HairStrands/HairStrandsDebug.usf
--- a/Engine/Shaders/Private/HairStrands/HairStrandsDebug.usf
+++ b/Engine/Shaders/Private/HairStrands/HairStrandsDebug.usf
@@ -159,7 +159,7 @@ void MainPS(
 	{
 		const int LocalSampleIndex = SampleIndex; 
 		const bool bAverageAllSamples = LocalSampleIndex < 0;
-		const uint SelectedIndex = LocalSampleIndex >= 0 ? clamp(LocalSampleIndex, 0, NodeDesc.Count - 1) : 0;
+		const uint SelectedIndex = LocalSampleIndex >= 0 ? clamp(LocalSampleIndex, 0, int(NodeDesc.Count - 1)) : 0;
 		float	Depth = 0;
 		float3	Tangent = 0;
 		float3	BaseColor = 0;
@@ -898,7 +898,7 @@ void DrawInstanceAABB(uint PrimitiveId)
 	// Since we don't have the actual instance ID (they are transient instance vs. persistent instance), 
 	// we find the primitive with the matching PrimitiveID
 	uint InstanceID = ~0;
-	const uint TotalInstanceCount = min(8192u, Scene.GPUScene.NumInstances);
+	const uint TotalInstanceCount = min(8192u, uint(Scene.GPUScene.NumInstances));
 	for (uint LocalInstanceID = 0; LocalInstanceID < TotalInstanceCount; ++LocalInstanceID)
 	{
 		FInstanceSceneData InstanceData = GetInstanceSceneData(LocalInstanceID, Scene.GPUScene.InstanceDataSOAStride);
