LFV - removed soft fading which is not really working.

#prelfight shader

diff --git a/Engine/Shaders/Private/LocalFogVolumes/LocalFogVolumeCommon.ush b/Engine/Shaders/Private/LocalFogVolumes/LocalFogVolumeCommon.ush
--- a/Engine/Shaders/Private/LocalFogVolumes/LocalFogVolumeCommon.ush
+++ b/Engine/Shaders/Private/LocalFogVolumes/LocalFogVolumeCommon.ush
@@ -19,7 +19,6 @@ struct FLocalFogVolumeGPUInstanceData
 	float Density;
 	float HeightFalloff;
 	float HeightOffset;
-	float RadialAttenuation;
 
 	float3 Albedo;
 	float  PhaseG;
@@ -74,7 +73,6 @@ FLocalFogVolumeGPUInstanceData GetLocalFogVolumeGPUInstanceData(uint Index)
 	Data.Density = Temp0.x;
 	Data.HeightFalloff = Temp0.y;
 	Data.HeightOffset = Temp0.z;
-	Data.RadialAttenuation = Temp0.w;
 
 	float4 Temp1 = LFVInsts[Offset + 17];
 	Data.Albedo = Temp1.xyz;
@@ -119,7 +117,7 @@ struct FFogData
 };
 
 // Height fog is evaluated according to spatially varying density defined as "Density * Exp[-HeightFalloff * z]"
-FFogData LocalFogVolumeEvaluateHeightIntegral(FLocalFogVolumeGPUInstanceData FogInstance, float3 CameraPos, float3 RayStart, float3 RayDir, float RayLength, float RadialAttenuation)
+FFogData LocalFogVolumeEvaluateHeightIntegral(FLocalFogVolumeGPUInstanceData FogInstance, float3 CameraPos, float3 RayStart, float3 RayDir, float RayLength)
 {
 	FFogData FogData;
 
@@ -151,7 +149,7 @@ FFogData LocalFogVolumeEvaluateHeightIntegral(FLocalFogVolumeGPUInstanceData Fog
 	// Transmittance is the 
 	float Transmittance = exp(-OpticalDepth);
 	// Assuming extinction is grey scale, we can compute a single coverage value from transmittance
-	FogData.Coverage = (1.0 - Transmittance) * RadialAttenuation;
+	FogData.Coverage = (1.0 - Transmittance);
 	// Assuming extinction==scattering, i.e. albedo=1, we can compute the integrale for each point on a ray while accouting for transmittance back to the view.
 	// It turns out that, under these conditions, the luminance and emissive factor is simply coverage. => see Mathematica details at the bottom of this page.
 	FogData.IntegratedLuminanceFactor = FogData.Coverage;
@@ -160,7 +158,7 @@ FFogData LocalFogVolumeEvaluateHeightIntegral(FLocalFogVolumeGPUInstanceData Fog
 }
 
 
-FFogData LocalFogVolumeEvaluateRadialIntegral(FLocalFogVolumeGPUInstanceData FogInstance, float3 CameraPos, float3 RayStart, float3 RayDir, float RayLength, float RadialAttenuation)
+FFogData LocalFogVolumeEvaluateRadialIntegral(FLocalFogVolumeGPUInstanceData FogInstance, float3 CameraPos, float3 RayStart, float3 RayDir, float RayLength)
 {
 	FFogData FogData = (FFogData)0;
 
@@ -198,7 +196,7 @@ FFogData LocalFogVolumeEvaluateRadialIntegral(FLocalFogVolumeGPUInstanceData Fog
 	// Transmittance is the 
 	float Transmittance = exp(-OpticalDepth);
 	// Assuming extinction is grey scale, we can compute a single coverage value from transmittance
-	FogData.Coverage = (1.0 - Transmittance) * RadialAttenuation;
+	FogData.Coverage = (1.0 - Transmittance);
 	// Assuming extinction==scattering, i.e. albedo=1, we can compute the integrale for each point on a ray while accouting for transmittance back to the view.
 	// It turns out that, under these conditions, the luminance and emissive factor is simply coverage. => see Mathematica details at the bottom of this page.
 	FogData.IntegratedLuminanceFactor = FogData.Coverage;
@@ -314,21 +312,6 @@ float4 GetLFVInstanceContribution(
 		float Length0 = length(P0Space - CamRayOriginSpace);
 		float Length1 = length(P1Space - CamRayOriginSpace);
 
-		// Radial attenuation based on the distance travelled in the unit sphere and ignoring depth buffer intersection. 
-		// This is only to soften sphere edge.
-		float RadialAttenuation = saturate(abs(max(0.0, TsU.y) - max(0.0, TsU.x)) / FogInstance.RadialAttenuation);
-		// The attenuation is linear but it can be made soft using an Hermit interpolation by controling the tangents.
-		{
-			const float t = RadialAttenuation;
-			const float t2 = t * t;
-			const float t3 = t2 * t;
-			const float HP0 = 0.0f;
-			const float HP1 = 1.0f;
-			const float HM0 = 0.0f;	// tangent = 0 for soft transition
-			const float HM1 = 0.0f; // idem
-			RadialAttenuation = saturate((2 * t3 - 3 * t2 + 1) * HP0 + (t3 - 2 * t2 + t) * HM0 + (-2 * t3 + 3 * t2) * HP1 + (t3 - t2) * HM1);
-		}
-
 		// Clamp the traced distance to depth and compute world length.
 		Length0 = min(LengthD, Length0);
 		Length1 = min(LengthD, Length1);
@@ -342,7 +325,7 @@ float4 GetLFVInstanceContribution(
 
 			if (FogInstance.FogMode == 0)
 			{
-				FogData = LocalFogVolumeEvaluateHeightIntegral(FogInstance, CamRayWorldOrigin, CamRayWorldOrigin + CamRayWorldDir * Length0, CamRayWorldDir, RayTracedLength, RadialAttenuation);
+				FogData = LocalFogVolumeEvaluateHeightIntegral(FogInstance, CamRayWorldOrigin, CamRayWorldOrigin + CamRayWorldDir * Length0, CamRayWorldDir, RayTracedLength);
 			}
 			else // (FogInstance.FogMode == 1)
 			{
@@ -350,7 +333,7 @@ float4 GetLFVInstanceContribution(
 				// TODO LocalFogVolumeEvaluateRadialIntegral needs to be optimised because it is redoing lots of computations. Maybe move HeightFogIntegral to also be in a local normalised space.
 				// TODO It is like that for now while iterating with artists.
 				// TODO
-				FogData = LocalFogVolumeEvaluateRadialIntegral(FogInstance, CamRayWorldOrigin, CamRayWorldOrigin + CamRayWorldDir * Length0, CamRayWorldDir, RayTracedLength, RadialAttenuation);
+				FogData = LocalFogVolumeEvaluateRadialIntegral(FogInstance, CamRayWorldOrigin, CamRayWorldOrigin + CamRayWorldDir * Length0, CamRayWorldDir, RayTracedLength);
 			}
 
 			LuminanceColor = LocalFogVolumeEvaluateInScattering(FogInstance, FogData, CommonLightingInput, CamRayWorldDir);
diff --git a/Engine/Source/Runtime/Engine/Classes/Components/LocalFogVolumeComponent.h b/Engine/Source/Runtime/Engine/Classes/Components/LocalFogVolumeComponent.h
--- a/Engine/Source/Runtime/Engine/Classes/Components/LocalFogVolumeComponent.h
+++ b/Engine/Source/Runtime/Engine/Classes/Components/LocalFogVolumeComponent.h
@@ -39,10 +39,6 @@ class ULocalFogVolumeComponent : public USceneComponent
 	UPROPERTY(EditAnywhere, BlueprintReadOnly, interp, Category = "Fog Distribution", meta = (UIMin = "-2.0", UIMax = "2.0"))
 	float FogHeightOffset = 0.0f;
 
-	/** Controls how strong the radial attenuation of this fog volume is. */
-	UPROPERTY(EditAnywhere, BlueprintReadOnly, interp, Category = "Fog Distribution", meta = (UIMin = "0.0", UIMax = "4.0", SliderExponent = 3.0, ClampMin = 0.0))
-	float FogRadialAttenuation = 0.0f;
-
 	/** Controls the phase `G` parameter, describing the directionality of the scattering within this fog volume. */
 	UPROPERTY(EditAnywhere, BlueprintReadOnly, interp, Category = "Fog Shading", meta = (DisplayName = "Scattering Distribution", UIMin = "0.0", UIMax = "0.999", ClampMin = 0.0, ClampMax = 0.999))
 	float FogPhaseG = 0.8f;
diff --git a/Engine/Source/Runtime/Engine/Private/LocalFogVolumeSceneProxy.cpp b/Engine/Source/Runtime/Engine/Private/LocalFogVolumeSceneProxy.cpp
--- a/Engine/Source/Runtime/Engine/Private/LocalFogVolumeSceneProxy.cpp
+++ b/Engine/Source/Runtime/Engine/Private/LocalFogVolumeSceneProxy.cpp
@@ -13,7 +13,6 @@ FLocalFogVolumeSceneProxy::FLocalFogVolumeSceneProxy(const ULocalFogVolumeCompon
 	: FogDensity(InComponent->FogDensity)
 	, FogHeightFalloff(InComponent->FogHeightFalloff)
 	, FogHeightOffset(InComponent->FogHeightOffset)
-	, FogRadialAttenuation(InComponent->FogRadialAttenuation)
 	, FogMode((uint8)InComponent->FogMode)
 	, FogSortPriority(uint8(127 - int8(InComponent->FogSortPriority))) // FogSortPriority on the component is in [-127,127] and needs to be negated to match expected priority behavior.
 	, FogPhaseG(InComponent->FogPhaseG)
diff --git a/Engine/Source/Runtime/Engine/Public/LocalFogVolumeSceneProxy.h b/Engine/Source/Runtime/Engine/Public/LocalFogVolumeSceneProxy.h
--- a/Engine/Source/Runtime/Engine/Public/LocalFogVolumeSceneProxy.h
+++ b/Engine/Source/Runtime/Engine/Public/LocalFogVolumeSceneProxy.h
@@ -23,7 +23,6 @@ public:
 	float FogDensity;
 	float FogHeightFalloff;
 	float FogHeightOffset;
-	float FogRadialAttenuation;
 
 	uint8 FogMode;
 	uint8 FogSortPriority;
diff --git a/Engine/Source/Runtime/Renderer/Private/LocalFogVolumeRendering.cpp b/Engine/Source/Runtime/Renderer/Private/LocalFogVolumeRendering.cpp
--- a/Engine/Source/Runtime/Renderer/Private/LocalFogVolumeRendering.cpp
+++ b/Engine/Source/Runtime/Renderer/Private/LocalFogVolumeRendering.cpp
@@ -266,7 +266,6 @@ void GetLocalFogVolumeSortingData(const FScene* Scene, FRDGBuilder& GraphBuilder
 		LocalFogVolumeGPUInstanceDataIt->Density = LHF->FogDensity;
 		LocalFogVolumeGPUInstanceDataIt->HeightFalloff = LHF->FogHeightFalloff * 0.01f;	// This scale is used to have artist author reasonable range.
 		LocalFogVolumeGPUInstanceDataIt->HeightOffset = LHF->FogHeightOffset;
-		LocalFogVolumeGPUInstanceDataIt->RadialAttenuation = LHF->FogRadialAttenuation;
 
 		LocalFogVolumeGPUInstanceDataIt->FogMode = float(LHF->FogMode);
 
diff --git a/Engine/Source/Runtime/Renderer/Private/LocalFogVolumeRendering.h b/Engine/Source/Runtime/Renderer/Private/LocalFogVolumeRendering.h
--- a/Engine/Source/Runtime/Renderer/Private/LocalFogVolumeRendering.h
+++ b/Engine/Source/Runtime/Renderer/Private/LocalFogVolumeRendering.h
@@ -47,7 +47,7 @@ public:
 	float Density;
 	float HeightFalloff;
 	float HeightOffset;
-	float RadialAttenuation;
+	float Unused;
 
 	FVector3f Albedo;
 	float PhaseG;
