LFV - Removed unused data from lfv instance gpu data. Also change the default phase to match height/volumetric fog using HenyeyGreensteinPhase.

#prelfight shader

diff --git a/Engine/Shaders/Private/LocalFogVolumes/LocalFogVolumeCommon.ush b/Engine/Shaders/Private/LocalFogVolumes/LocalFogVolumeCommon.ush
--- a/Engine/Shaders/Private/LocalFogVolumes/LocalFogVolumeCommon.ush
+++ b/Engine/Shaders/Private/LocalFogVolumes/LocalFogVolumeCommon.ush
@@ -13,9 +13,6 @@ struct FLocalFogVolumeGPUInstanceData
 	float4x4 Transform;
 	float4x4 InvTransform;
 
-	float4x4 InvTranformNoScale;
-	float4x4 TransformScaleOnly;
-
 	float Density;
 	float HeightFalloff;
 	float HeightOffset;
@@ -61,24 +58,21 @@ Texture2DArray<uint>			LFV_LocalFogVolumeTileDataTexture;
 FLocalFogVolumeGPUInstanceData GetLocalFogVolumeGPUInstanceData(uint Index)
 {
 	FLocalFogVolumeGPUInstanceData Data;
-	uint Offset = Index * 19;
+	uint Offset = Index * 11;
 
 	Data.Transform			= float4x4(LFVInsts[Offset + 0 ], LFVInsts[Offset + 1 ], LFVInsts[Offset + 2 ], LFVInsts[Offset + 3 ]);
 	Data.InvTransform		= float4x4(LFVInsts[Offset + 4 ], LFVInsts[Offset + 5 ], LFVInsts[Offset + 6 ], LFVInsts[Offset + 7 ]);
 
-	Data.InvTranformNoScale = float4x4(LFVInsts[Offset + 8 ], LFVInsts[Offset + 9 ], LFVInsts[Offset + 10], LFVInsts[Offset + 11]);
-	Data.TransformScaleOnly = float4x4(LFVInsts[Offset + 12], LFVInsts[Offset + 13], LFVInsts[Offset + 14], LFVInsts[Offset + 15]);
-
-	float4 Temp0 = LFVInsts[Offset + 16];
+	float4 Temp0 = LFVInsts[Offset + 8];
 	Data.Density = Temp0.x;
 	Data.HeightFalloff = Temp0.y;
 	Data.HeightOffset = Temp0.z;
 
-	float4 Temp1 = LFVInsts[Offset + 17];
+	float4 Temp1 = LFVInsts[Offset + 9];
 	Data.Albedo = Temp1.xyz;
 	Data.PhaseG = Temp1.w;
 
-	float4 Temp2 = LFVInsts[Offset + 18];
+	float4 Temp2 = LFVInsts[Offset + 10];
 	Data.Emissive = Temp2.xyz;
 	Data.FogMode = Temp2.w;
 
@@ -235,11 +229,11 @@ float3 LocalFogVolumeEvaluateInScattering(in FLocalFogVolumeGPUInstanceData FogI
 	// Accumulate directional lights if needed
 	if (View.AtmosphereLightIlluminanceOnGroundPostTransmittance[0].a > 0.0f)
 	{
-		InScattering += View.AtmosphereLightIlluminanceOnGroundPostTransmittance[0].rgb * SchlickPhase(-FogInstance.PhaseG, dot(RayDirWorld, View.AtmosphereLightDirection[0].xyz));
+		InScattering += View.AtmosphereLightIlluminanceOnGroundPostTransmittance[0].rgb * HenyeyGreensteinPhase(-FogInstance.PhaseG, dot(RayDirWorld, View.AtmosphereLightDirection[0].xyz));
 	}
 	if (View.AtmosphereLightIlluminanceOnGroundPostTransmittance[1].a > 0.0f)
 	{
-		InScattering += View.AtmosphereLightIlluminanceOnGroundPostTransmittance[0].rgb * SchlickPhase(-FogInstance.PhaseG, dot(RayDirWorld, View.AtmosphereLightDirection[1].xyz));
+		InScattering += View.AtmosphereLightIlluminanceOnGroundPostTransmittance[0].rgb * HenyeyGreensteinPhase(-FogInstance.PhaseG, dot(RayDirWorld, View.AtmosphereLightDirection[1].xyz));
 	}
 
 	// Weight by user control
@@ -266,8 +260,6 @@ float3 LocalFogVolumeEvaluateInScattering(in FLocalFogVolumeGPUInstanceData FogI
 #define LFV_CAN_CLIP 0
 #endif
 
-#define LHFV_FOLLOWS_ROTATIONS 1
-
 // Returns color transmittance
 float4 GetLFVInstanceContribution(
 	uint InstanceIndex,
@@ -294,20 +286,11 @@ float4 GetLFVInstanceContribution(
 
 		float3 P0U = RayPosU + max(0.0, TsU.x) * RayDirU;
 		float3 P1U = RayPosU + max(0.0, TsU.y) * RayDirU;
-#if LHFV_FOLLOWS_ROTATIONS
-		// Computations are all done transformed with Space = local fog volume but ignoring scaling.
-		float3 P0Space = mul(float4(P0U, 1.0), FogInstance.TransformScaleOnly).xyz;
-		float3 P1Space = mul(float4(P1U, 1.0), FogInstance.TransformScaleOnly).xyz;
-		float3 CamRayOriginSpace = mul(float4(CamRayWorldOrigin, 1.0), FogInstance.InvTranformNoScale).xyz;
-		float3 CamRayDirSpace = mul(float4(CamRayWorldDir, 0.0), FogInstance.InvTranformNoScale).xyz;	// No scale so no need for a special InverseTranspose matrix.
-		CamRayDirSpace = normalize(CamRayDirSpace);
-#else
 		// Computations are all done with Space = world space
 		float3 P0Space = mul(float4(P0U, 1.0), FogInstance.Transform).xyz;
 		float3 P1Space = mul(float4(P1U, 1.0), FogInstance.Transform).xyz;
 		float3 CamRayOriginSpace = CamRayWorldOrigin;
 		float3 CamRayDirSpace = CamRayWorldDir;
-#endif
 
 		float Length0 = length(P0Space - CamRayOriginSpace);
 		float Length1 = length(P1Space - CamRayOriginSpace);
diff --git a/Engine/Source/Runtime/Engine/Classes/Components/LocalFogVolumeComponent.h b/Engine/Source/Runtime/Engine/Classes/Components/LocalFogVolumeComponent.h
--- a/Engine/Source/Runtime/Engine/Classes/Components/LocalFogVolumeComponent.h
+++ b/Engine/Source/Runtime/Engine/Classes/Components/LocalFogVolumeComponent.h
@@ -41,7 +41,7 @@ class ULocalFogVolumeComponent : public USceneComponent
 
 	/** Controls the phase `G` parameter, describing the directionality of the scattering within this fog volume. */
 	UPROPERTY(EditAnywhere, BlueprintReadOnly, interp, Category = "Fog Shading", meta = (DisplayName = "Scattering Distribution", UIMin = "0.0", UIMax = "0.999", ClampMin = 0.0, ClampMax = 0.999))
-	float FogPhaseG = 0.8f;
+	float FogPhaseG = 0.2f;
 
 	/** Controls the albedo of this fog volume. */
 	UPROPERTY(EditAnywhere, BlueprintReadOnly, interp, Category = "Fog Shading", meta = (ClampMin = 0.0, ClampMax = 1.0))
diff --git a/Engine/Source/Runtime/Renderer/Private/LocalFogVolumeRendering.cpp b/Engine/Source/Runtime/Renderer/Private/LocalFogVolumeRendering.cpp
--- a/Engine/Source/Runtime/Renderer/Private/LocalFogVolumeRendering.cpp
+++ b/Engine/Source/Runtime/Renderer/Private/LocalFogVolumeRendering.cpp
@@ -260,8 +260,6 @@ void GetLocalFogVolumeSortingData(const FScene* Scene, FRDGBuilder& GraphBuilder
 		FLocalFogVolumeGPUInstanceData* LocalFogVolumeGPUInstanceDataIt = &Out.LocalFogVolumeGPUInstanceData[Out.LocalFogVolumeInstanceCountFinal];
 		LocalFogVolumeGPUInstanceDataIt->Transform = FMatrix44f(LHF->FogTransform.ToMatrixWithScale());
 		LocalFogVolumeGPUInstanceDataIt->InvTransform = LocalFogVolumeGPUInstanceDataIt->Transform.Inverse();
-		LocalFogVolumeGPUInstanceDataIt->InvTranformNoScale = FMatrix44f(LHF->FogTransform.ToMatrixNoScale()).Inverse();
-		LocalFogVolumeGPUInstanceDataIt->TransformScaleOnly = FMatrix44f(TransformScaleOnly.ToMatrixWithScale());
 
 		LocalFogVolumeGPUInstanceDataIt->Density = LHF->FogDensity;
 		LocalFogVolumeGPUInstanceDataIt->HeightFalloff = LHF->FogHeightFalloff * 0.01f;	// This scale is used to have artist author reasonable range.
diff --git a/Engine/Source/Runtime/Renderer/Private/LocalFogVolumeRendering.h b/Engine/Source/Runtime/Renderer/Private/LocalFogVolumeRendering.h
--- a/Engine/Source/Runtime/Renderer/Private/LocalFogVolumeRendering.h
+++ b/Engine/Source/Runtime/Renderer/Private/LocalFogVolumeRendering.h
@@ -41,9 +41,6 @@ public:
 	FMatrix44f Transform;
 	FMatrix44f InvTransform;
 
-	FMatrix44f InvTranformNoScale;
-	FMatrix44f TransformScaleOnly;
-
 	float Density;
 	float HeightFalloff;
 	float HeightOffset;
