Fix Substrate env. lighting on mobile renderer.

diff --git a/Engine/Shaders/Private/Substrate/SubstrateMobileForwardLighting.ush b/Engine/Shaders/Private/Substrate/SubstrateMobileForwardLighting.ush
--- a/Engine/Shaders/Private/Substrate/SubstrateMobileForwardLighting.ush
+++ b/Engine/Shaders/Private/Substrate/SubstrateMobileForwardLighting.ush
@@ -176,8 +176,7 @@ float3 SubstrateMobileForwardLighting(
 				}
 
 				// 2. Specular evaluation (using regular normal)
-				//#if MATERIALBLENDING_MASKED || MATERIALBLENDING_SOLID || TRANSLUCENCY_LIGHTING_SURFACE_FORWARDSHADING || MATERIAL_SHADINGMODEL_SINGLELAYERWATER
-				#if ENABLE_CLUSTERED_REFLECTION && (FORWARD_PER_PIXEL_SHADING || TRANSLUCENCY_LIGHTING_SURFACE_LIGHTINGVOLUME)
+				#if (FORWARD_PER_PIXEL_SHADING || TRANSLUCENCY_LIGHTING_SURFACE_LIGHTINGVOLUME)
 				if (UseBasePassSkylightSpecular > 0)
 				{
 					// Evaluate specular parameters without bent normal
@@ -186,13 +185,13 @@ float3 SubstrateMobileForwardLighting(
 					// Evaluate specular lighting for the main lobe
 					if (any(SubstrateEnvLight.SpecularWeight > 0.0f))
 					{
-						Color += GetImageBasedReflectionSpecular(SubstrateEnvLight.SpecularDirection, SubstrateEnvLight.SpecularSafeRoughness, MaterialParameters.WorldPosition_CamRelative, SubstrateBSDFContext.N, ShadingOcclusion.SpecularOcclusion, IndirectIrradiance, GridIndex) *
+						Color += GetImageBasedReflectionSpecular(SubstrateEnvLight.SpecularDirection, SubstrateEnvLight.SpecularSafeRoughness, MaterialParameters.WorldPosition_CamRelative, SubstrateBSDFContext.N, ShadingOcclusion.SpecOcclusion, IndirectIrradiance, GridIndex) *
 								 SubstrateEnvLight.SpecularWeight * IndirectOcclusion * LuminanceWeight(SubstrateBSDFContext, CurrentBSDF) * AOMultiBounce(SubstrateEnvLight.SpecularColor, ShadingOcclusion.SpecOcclusion);
 					}
 					// And for the second lobe if needed
 					if (any(SubstrateEnvLight.SpecularHazeWeight > 0.0f))
 					{
-						Color += GetImageBasedReflectionSpecular(SubstrateEnvLight.SpecularHazeDirection, SubstrateEnvLight.SpecularHazeSafeRoughness, MaterialParameters.WorldPosition_CamRelative, SubstrateBSDFContext.N, ShadingOcclusion.SpecularOcclusion, IndirectIrradiance, GridIndex) *
+						Color += GetImageBasedReflectionSpecular(SubstrateEnvLight.SpecularHazeDirection, SubstrateEnvLight.SpecularHazeSafeRoughness, MaterialParameters.WorldPosition_CamRelative, SubstrateBSDFContext.N, ShadingOcclusion.SpecOcclusion, IndirectIrradiance, GridIndex) *
 								 SubstrateEnvLight.SpecularHazeWeight * IndirectOcclusion * LuminanceWeight(SubstrateBSDFContext, CurrentBSDF) * AOMultiBounce(SubstrateEnvLight.SpecularColor, ShadingOcclusion.SpecOcclusion);
 					}
 				}
