Speeds up TSR's anti-flickering heuristic with concatenated convolutions

diff --git a/Engine/Shaders/Private/TemporalSuperResolution/TSRConvolutionNetwork.ush b/Engine/Shaders/Private/TemporalSuperResolution/TSRConvolutionNetwork.ush
--- a/Engine/Shaders/Private/TemporalSuperResolution/TSRConvolutionNetwork.ush
+++ b/Engine/Shaders/Private/TemporalSuperResolution/TSRConvolutionNetwork.ush
@@ -294,6 +294,11 @@ void WriteElementToLDS(const uint SharedIndex, tsr_ushort X)
 	WriteElementToLDS(SharedIndex, tsr_ushort4(X, 0.0, 0.0, 0.0));
 }
 
+void WriteElementToLDS(const uint SharedIndex, tsr_ushort2 X)
+{
+	WriteElementToLDS(SharedIndex, tsr_ushort4(X, 0.0, 0.0));
+}
+
 void WriteElementToLDS(const uint SharedIndex, tsr_ushort3 X)
 {
 	WriteElementToLDS(SharedIndex, tsr_ushort4(X, 0.0));
@@ -304,6 +309,11 @@ void WriteElementToLDS(const uint SharedIndex, tsr_half X)
 	WriteElementToLDS(SharedIndex, as_tsr_ushort(X));
 }
 
+void WriteElementToLDS(const uint SharedIndex, tsr_half2 X)
+{
+	WriteElementToLDS(SharedIndex, tsr_ushort2(as_tsr_ushort(X.x), as_tsr_ushort(X.y)));
+}
+
 void WriteElementToLDS(const uint SharedIndex, tsr_half3 X)
 {
 	WriteElementToLDS(SharedIndex, tsr_ushort3(as_tsr_ushort(X.x), as_tsr_ushort(X.y), as_tsr_ushort(X.z)));
@@ -321,6 +331,13 @@ void ReadElementFromLDS(const uint SharedIndex, out tsr_ushort X)
 	X = UX.x;
 }
 
+void ReadElementFromLDS(const uint SharedIndex, out tsr_ushort2 X)
+{
+	tsr_ushort4 UX;
+	ReadElementFromLDS(SharedIndex, /* out */ UX);
+	X = UX.xy;
+}
+
 void ReadElementFromLDS(const uint SharedIndex, out tsr_ushort3 X)
 {
 	tsr_ushort4 UX;
@@ -335,6 +352,14 @@ void ReadElementFromLDS(const uint SharedIndex, out tsr_half X)
 	X = as_tsr_half(UX.x);
 }
 
+void ReadElementFromLDS(const uint SharedIndex, out tsr_half2 X)
+{
+	tsr_ushort4 UX;
+	ReadElementFromLDS(SharedIndex, /* out */ UX);
+	X.x = as_tsr_half(UX.x);
+	X.y = as_tsr_half(UX.y);
+}
+
 void ReadElementFromLDS(const uint SharedIndex, out tsr_half3 X)
 {
 	tsr_ushort4 UX;
diff --git a/Engine/Shaders/Private/TemporalSuperResolution/TSRRejectShading.usf b/Engine/Shaders/Private/TemporalSuperResolution/TSRRejectShading.usf
--- a/Engine/Shaders/Private/TemporalSuperResolution/TSRRejectShading.usf
+++ b/Engine/Shaders/Private/TemporalSuperResolution/TSRRejectShading.usf
@@ -199,8 +199,7 @@ void ComputeMoireError(
 	}
 		
 	// Convolve the input luma to have the exact same moire pattern as FilteredInput below.
-	InputLuma = Blur3x3(InputLuma);
-	InputLumaModification = Blur3x3(InputLumaModification);
+	Deconcatenate(Blur3x3(Concatenate(InputLuma, InputLumaModification)), /* out */ InputLuma, /* out */ InputLumaModification);
 
 	// Unpack history
 	tsr_simd_half PrevFlickeringHistory = PrevMoireHistory[0];
@@ -266,8 +265,9 @@ void ComputeMoireError(
 
 	// Dilates the total variation and contribution to keep stability despite inacuracies in the history reprojection
 	// of these need a nearest filter which isn't precise at all.
-	tsr_simd_half TotalVariationContribution = Max3x3(GradientVariation);
-	tsr_simd_half TotalVariationCountContribution = Max3x3(IsFlicker);
+	tsr_simd_half TotalVariationContribution;
+	tsr_simd_half TotalVariationCountContribution;
+	Deconcatenate(Max3x3(Concatenate(GradientVariation, IsFlicker)), /* out */ TotalVariationContribution, /* out */ TotalVariationCountContribution);
 
 	UNROLL_N(SIMD_SIZE)
 	for (uint SimdIndex = 0; SimdIndex < SIMD_SIZE; SimdIndex++)
