Fix NaN when using light transmission with SSS MFP with MFP=0 for certain channel.

diff --git a/Engine/Shaders/Private/Substrate/SubstrateEvaluation.ush b/Engine/Shaders/Private/Substrate/SubstrateEvaluation.ush
--- a/Engine/Shaders/Private/Substrate/SubstrateEvaluation.ush
+++ b/Engine/Shaders/Private/Substrate/SubstrateEvaluation.ush
@@ -868,7 +868,7 @@ FSubstrateEvaluateResult SubstrateEvaluateBSDFCommon(FSubstrateBSDFContext BSDFC
 
 					const float3 MeanFreePathInCm = SLAB_SSSMFP(BSDFContext.BSDF);
 					const float3 SubsurfaceAlebdo = SLAB_DIFFUSEALBEDO(BSDFContext.BSDF);
-					TransmissionThroughput = GetBurleyTransmission(SubsurfaceAlebdo, MeanFreePathInCm, ThicknessInCm).xyz;
+					TransmissionThroughput = GetBurleyTransmission(SubsurfaceAlebdo, max(SUBSTRATE_EPSILON, MeanFreePathInCm), ThicknessInCm).xyz;
 
 					PhaseFunctionAnisotropy = SLAB_SSSPHASEANISOTROPY(BSDFContext.BSDF);
 					OneOverIOR				= 1.f / DielectricF0ToIor(F0.y);
