VulkanRHI: VK RT WIP: Circumvent crash in DXC where SPIRV emitter can't figure out 3rd parameter type for ReportHit call (HeterogeneousVolumesHardwareRayTracing).

#rnx

diff --git a/Engine/Shaders/Private/HeterogeneousVolumes/HeterogeneousVolumesHardwareRayTracing.usf b/Engine/Shaders/Private/HeterogeneousVolumes/HeterogeneousVolumesHardwareRayTracing.usf
--- a/Engine/Shaders/Private/HeterogeneousVolumes/HeterogeneousVolumesHardwareRayTracing.usf
+++ b/Engine/Shaders/Private/HeterogeneousVolumes/HeterogeneousVolumesHardwareRayTracing.usf
@@ -182,7 +182,8 @@ RAY_TRACING_ENTRY_INTERSECTION(SparseVoxelsIntersectionShader)
 	if (HitT.y > HitT.x)
 	{
 		uint HitKind = 0;
-		ReportHit(HitT.x, HitKind, CreateVolumeIntersectionAttributes(HitT));
+		FSparseVoxelIntersectionAttributes IntersectionAttributes = CreateVolumeIntersectionAttributes(HitT);
+		ReportHit(HitT.x, HitKind, IntersectionAttributes);
 	}
 }
 
