Moved common logic from DistanceFieldShadowingUpsamplePS and CapsuleShadowingUpsamplePS to UpsampleShadowFactors
- CapsuleShadowingUpsamplePS was also missing some bug fixes that improve image quality.

diff --git a/Engine/Shaders/Private/CapsuleShadowShaders.usf b/Engine/Shaders/Private/CapsuleShadowShaders.usf
--- a/Engine/Shaders/Private/CapsuleShadowShaders.usf
+++ b/Engine/Shaders/Private/CapsuleShadowShaders.usf
@@ -13,6 +13,10 @@
 #include "ReflectionEnvironmentShared.ush"
 #include "Substrate/Substrate.ush"
 
+#ifdef SHADOW_FACTORS_UPSAMPLE_REQUIRED
+#	include "ShadowFactorsUpsampleCommon.ush"
+#endif
+
 #ifndef THREADGROUP_SIZEX
 #define THREADGROUP_SIZEX 1
 #endif
@@ -1195,47 +1199,11 @@ void CapsuleShadowingUpsamplePS(
 #endif
 	)
 {
-	// Distance field shadowing was computed at 0,0 regardless of viewrect min
-	float2 DistanceFieldUVs = SvPositionToBufferUV(SVPos) - ScissorRectMinAndSize.xy * View.BufferSizeAndInvSize.zw;
-	float2 ScreenUV = SvPositionToBufferUV(SVPos);
-	float SceneDepth = CalcSceneDepth(ScreenUV);
-
-#define BILATERAL_UPSAMPLE 1
-#if BILATERAL_UPSAMPLE && UPSAMPLE_REQUIRED
-	float2 LowResBufferSize = floor(View.BufferSizeAndInvSize.xy / DOWNSAMPLE_FACTOR);
-	float2 LowResTexelSize = 1.0f / LowResBufferSize;
-	float2 Corner00UV = floor(DistanceFieldUVs * LowResBufferSize - .5f) / LowResBufferSize + .5f * LowResTexelSize;
-	float2 BilinearWeights = (DistanceFieldUVs - Corner00UV) * LowResBufferSize;
-
-	float2 TextureValues00 = Texture2DSampleLevel(ShadowFactorsTexture, ShadowFactorsSampler, Corner00UV, 0).xy;
-	float2 TextureValues10 = Texture2DSampleLevel(ShadowFactorsTexture, ShadowFactorsSampler, Corner00UV + float2(LowResTexelSize.x, 0), 0).xy;
-	float2 TextureValues01 = Texture2DSampleLevel(ShadowFactorsTexture, ShadowFactorsSampler, Corner00UV + float2(0, LowResTexelSize.y), 0).xy;
-	float2 TextureValues11 = Texture2DSampleLevel(ShadowFactorsTexture, ShadowFactorsSampler, Corner00UV + LowResTexelSize, 0).xy;
-
-	float4 CornerWeights = float4(
-		(1 - BilinearWeights.y) * (1 - BilinearWeights.x), 
-		(1 - BilinearWeights.y) * BilinearWeights.x,
-		BilinearWeights.y * (1 - BilinearWeights.x),
-		BilinearWeights.y * BilinearWeights.x);
+	const float2 ScreenUV = SvPositionToBufferUV(SVPos);
 
-	float Epsilon = .0001f;
-
-	float4 CornerDepths = abs(float4(TextureValues00.y, TextureValues10.y, TextureValues01.y, TextureValues11.y));
-	float4 DepthWeights = 1.0f / (abs(CornerDepths - SceneDepth.xxxx) + Epsilon);
-	float4 FinalWeights = CornerWeights * DepthWeights;
-
-	float InterpolatedResult = 
-		(FinalWeights.x * TextureValues00.x 
-		+ FinalWeights.y * TextureValues10.x
-		+ FinalWeights.z * TextureValues01.x 
-		+ FinalWeights.w * TextureValues11.x)
-		/ dot(FinalWeights, 1);
-
-	float Output = InterpolatedResult;
-
-#else
-	float Output = Texture2DSampleLevel(ShadowFactorsTexture, ShadowFactorsSampler, DistanceFieldUVs, 0).x;
-#endif
+	float Output;
+	float SceneDepth;
+	UpsampleShadowFactors(SVPos, ScissorRectMinAndSize, 1.0f / DOWNSAMPLE_FACTOR, ShadowFactorsTexture, ShadowFactorsSampler, Output, SceneDepth);
 
 	if (OutputtingToLightAttenuation > 0)
 	{
diff --git a/Engine/Shaders/Private/DistanceFieldShadowing.usf b/Engine/Shaders/Private/DistanceFieldShadowing.usf
--- a/Engine/Shaders/Private/DistanceFieldShadowing.usf
+++ b/Engine/Shaders/Private/DistanceFieldShadowing.usf
@@ -11,6 +11,10 @@
 #include "DistanceFieldShadowingShared.ush"
 #include "Substrate/Substrate.ush"
 
+#ifdef SHADOW_FACTORS_UPSAMPLE_REQUIRED
+#	include "ShadowFactorsUpsampleCommon.ush"
+#endif
+
 #define SCATTER_TILE_CULLING (CULLING_TYPE == 0)
 #define POINT_LIGHT (CULLING_TYPE == 2)
 
@@ -779,59 +783,9 @@ void DistanceFieldShadowingUpsamplePS(
 	in float4 SVPos : SV_POSITION,
 	out float4 OutColor : SV_Target0)
 {
-	float2 BufferUV = SvPositionToBufferUV(SVPos);
-
-	// Distance field shadowing was computed at 0,0 regardless of viewrect min
-	float2 DistanceFieldUVs = BufferUV - ScissorRectMinAndSize.xy * View.BufferSizeAndInvSize.zw;
-
-	float SceneDepth = CalcSceneDepth(BufferUV);
-
-#define BILATERAL_UPSAMPLE 1
-#if BILATERAL_UPSAMPLE && UPSAMPLE_REQUIRED
-	float2 LowResBufferSize = floor(View.BufferSizeAndInvSize.xy * OneOverDownsampleFactor);
-	float2 LowResTexelSize = 1.0f / LowResBufferSize;
-
-	// since DF shadow rays are traced from the position of the full-res texel at the top left of low-res tile
-	// we need to use different UVs for calculating bilinear weights and sampling the 4 texels
-
-	float2 Corner00UVBase = floor(DistanceFieldUVs * LowResBufferSize) / LowResBufferSize;
-
-	// Calculate bilinear weights using the same UVs that rays were traced from
-	float2 Corner00UV = Corner00UVBase + 0.5f * View.BufferSizeAndInvSize.zw;
-	float2 BilinearWeights = saturate((DistanceFieldUVs - Corner00UV) * LowResBufferSize);
-
-	// Sample low res texture using UVs at center of low-res texels since we are interpolating manually
-	float2 Corner00SampleUV = Corner00UVBase + 0.5f * LowResTexelSize;
-
-	float2 TextureValues00 = Texture2DSampleLevel(ShadowFactorsTexture, ShadowFactorsSampler, Corner00SampleUV, 0).xy;
-	float2 TextureValues10 = Texture2DSampleLevel(ShadowFactorsTexture, ShadowFactorsSampler, Corner00SampleUV + float2(LowResTexelSize.x, 0), 0).xy;
-	float2 TextureValues01 = Texture2DSampleLevel(ShadowFactorsTexture, ShadowFactorsSampler, Corner00SampleUV + float2(0, LowResTexelSize.y), 0).xy;
-	float2 TextureValues11 = Texture2DSampleLevel(ShadowFactorsTexture, ShadowFactorsSampler, Corner00SampleUV + LowResTexelSize, 0).xy;
-
-	float4 CornerWeights = float4(
-		(1 - BilinearWeights.y) * (1 - BilinearWeights.x), 
-		(1 - BilinearWeights.y) * BilinearWeights.x,
-		BilinearWeights.y * (1 - BilinearWeights.x),
-		BilinearWeights.y * BilinearWeights.x);
-
-	float Epsilon = .0001f;
-
-	float4 CornerDepths = abs(float4(TextureValues00.y, TextureValues10.y, TextureValues01.y, TextureValues11.y));
-	float4 DepthWeights = 1.0f / (abs(CornerDepths - SceneDepth.xxxx) + Epsilon);
-	float4 FinalWeights = CornerWeights * DepthWeights;
-
-	float InterpolatedResult = 
-		(FinalWeights.x * TextureValues00.x 
-		+ FinalWeights.y * TextureValues10.x
-		+ FinalWeights.z * TextureValues01.x 
-		+ FinalWeights.w * TextureValues11.x)
-		/ dot(FinalWeights, 1);
-
-	float Output = InterpolatedResult;
-
-#else
-	float Output = Texture2DSampleLevel(ShadowFactorsTexture, ShadowFactorsSampler, DistanceFieldUVs, 0).x;
-#endif
+	float Output;
+	float SceneDepth;
+	UpsampleShadowFactors(SVPos, ScissorRectMinAndSize, OneOverDownsampleFactor, ShadowFactorsTexture, ShadowFactorsSampler, Output, SceneDepth);
 
 	float FarBlendFactor = 1.0f - saturate((SceneDepth - FadePlaneOffset) * InvFadePlaneLength);
 	Output = lerp(1, Output, FarBlendFactor);
diff --git a/Engine/Shaders/Private/ShadowFactorsUpsampleCommon.ush b/Engine/Shaders/Private/ShadowFactorsUpsampleCommon.ush
new file mode 100644
--- /dev/null
+++ b/Engine/Shaders/Private/ShadowFactorsUpsampleCommon.ush
@@ -0,0 +1,68 @@
+// Copyright Epic Games, Inc. All Rights Reserved.
+
+/*=============================================================================
+	ShadowFactorsUpsampleCommon.ush
+=============================================================================*/
+
+#ifndef SHADOW_FACTORS_UPSAMPLE_REQUIRED
+#error "Must define SHADOW_FACTORS_UPSAMPLE_REQUIRED"
+#endif
+
+#include "Common.ush"
+#include "ComputeShaderUtils.ush"
+
+void UpsampleShadowFactors(float4 SVPos, int4 ScissorRectMinAndSize, float OneOverDownsampleFactor, Texture2D ShadowFactorsTexture, SamplerState ShadowFactorsSampler, out float ShadowFactor, out float SceneDepth)
+{
+	float2 BufferUV = SvPositionToBufferUV(SVPos);
+
+	// ShadowFactors is computed at 0,0 regardless of viewrect min
+	float2 ShadowFactorsUV = BufferUV - ScissorRectMinAndSize.xy * View.BufferSizeAndInvSize.zw;
+
+	SceneDepth = CalcSceneDepth(BufferUV);
+
+#define BILATERAL_UPSAMPLE 1
+#if BILATERAL_UPSAMPLE && SHADOW_FACTORS_UPSAMPLE_REQUIRED
+	float2 LowResBufferSize = floor(View.BufferSizeAndInvSize.xy * OneOverDownsampleFactor);
+	float2 LowResTexelSize = 1.0f / LowResBufferSize;
+
+	// since DF shadow rays are traced from the position of the full-res texel at the top left of low-res tile
+	// we need to use different UVs for calculating bilinear weights and sampling the 4 texels
+
+	float2 Corner00UVBase = floor(ShadowFactorsUV * LowResBufferSize) / LowResBufferSize;
+
+	// Calculate bilinear weights using the same UVs that rays were traced from
+	float2 Corner00UV = Corner00UVBase + 0.5f * View.BufferSizeAndInvSize.zw;
+	float2 BilinearWeights = saturate((ShadowFactorsUV - Corner00UV) * LowResBufferSize);
+
+	// Sample low res texture using UVs at center of low-res texels since we are interpolating manually
+	float2 Corner00SampleUV = Corner00UVBase + 0.5f * LowResTexelSize;
+
+	float2 TextureValues00 = Texture2DSampleLevel(ShadowFactorsTexture, ShadowFactorsSampler, Corner00SampleUV, 0).xy;
+	float2 TextureValues10 = Texture2DSampleLevel(ShadowFactorsTexture, ShadowFactorsSampler, Corner00SampleUV + float2(LowResTexelSize.x, 0), 0).xy;
+	float2 TextureValues01 = Texture2DSampleLevel(ShadowFactorsTexture, ShadowFactorsSampler, Corner00SampleUV + float2(0, LowResTexelSize.y), 0).xy;
+	float2 TextureValues11 = Texture2DSampleLevel(ShadowFactorsTexture, ShadowFactorsSampler, Corner00SampleUV + LowResTexelSize, 0).xy;
+
+	float4 CornerWeights = float4(
+		(1 - BilinearWeights.y) * (1 - BilinearWeights.x), 
+		(1 - BilinearWeights.y) * BilinearWeights.x,
+		BilinearWeights.y * (1 - BilinearWeights.x),
+		BilinearWeights.y * BilinearWeights.x);
+
+	float Epsilon = .0001f;
+
+	float4 CornerDepths = abs(float4(TextureValues00.y, TextureValues10.y, TextureValues01.y, TextureValues11.y));
+	float4 DepthWeights = 1.0f / (abs(CornerDepths - SceneDepth.xxxx) + Epsilon);
+	float4 FinalWeights = CornerWeights * DepthWeights;
+
+	float InterpolatedResult = 
+		(FinalWeights.x * TextureValues00.x 
+		+ FinalWeights.y * TextureValues10.x
+		+ FinalWeights.z * TextureValues01.x 
+		+ FinalWeights.w * TextureValues11.x)
+		/ dot(FinalWeights, 1);
+
+	ShadowFactor = InterpolatedResult;
+#else
+	ShadowFactor = Texture2DSampleLevel(ShadowFactorsTexture, ShadowFactorsSampler, ShadowFactorsUV, 0).x;
+#endif
+}
diff --git a/Engine/Source/Runtime/Renderer/Private/CapsuleShadowRendering.cpp b/Engine/Source/Runtime/Renderer/Private/CapsuleShadowRendering.cpp
--- a/Engine/Source/Runtime/Renderer/Private/CapsuleShadowRendering.cpp
+++ b/Engine/Source/Runtime/Renderer/Private/CapsuleShadowRendering.cpp
@@ -333,7 +333,7 @@ public:
 		SHADER_PARAMETER(float, OutputtingToLightAttenuation)
 	END_SHADER_PARAMETER_STRUCT()
 
-	class FUpsampleRequired : SHADER_PERMUTATION_BOOL("UPSAMPLE_REQUIRED");
+	class FUpsampleRequired : SHADER_PERMUTATION_BOOL("SHADOW_FACTORS_UPSAMPLE_REQUIRED");
 	class FApplySSAO : SHADER_PERMUTATION_BOOL("APPLY_TO_SSAO");
 	using FPermutationDomain = TShaderPermutationDomain<FUpsampleRequired, FApplySSAO>;
 
diff --git a/Engine/Source/Runtime/Renderer/Private/DistanceFieldShadowing.cpp b/Engine/Source/Runtime/Renderer/Private/DistanceFieldShadowing.cpp
--- a/Engine/Source/Runtime/Renderer/Private/DistanceFieldShadowing.cpp
+++ b/Engine/Source/Runtime/Renderer/Private/DistanceFieldShadowing.cpp
@@ -350,7 +350,7 @@ class FDistanceFieldShadowingUpsamplePS : public FGlobalShader
 		SHADER_PARAMETER(float, OneOverDownsampleFactor)
 	END_SHADER_PARAMETER_STRUCT()
 
-	class FUpsample : SHADER_PERMUTATION_BOOL("UPSAMPLE_REQUIRED");
+	class FUpsample : SHADER_PERMUTATION_BOOL("SHADOW_FACTORS_UPSAMPLE_REQUIRED");
 	using FPermutationDomain = TShaderPermutationDomain<FUpsample>;
 
 	static bool ShouldCompilePermutation(const FGlobalShaderPermutationParameters& Parameters)
