Fix shader compilation errors related to UpsampleShadowFactors.

diff --git a/Engine/Shaders/Private/CapsuleShadowShaders.usf b/Engine/Shaders/Private/CapsuleShadowShaders.usf
--- a/Engine/Shaders/Private/CapsuleShadowShaders.usf
+++ b/Engine/Shaders/Private/CapsuleShadowShaders.usf
@@ -13,16 +13,16 @@
 #include "ReflectionEnvironmentShared.ush"
 #include "Substrate/Substrate.ush"
 
-#ifdef SHADOW_FACTORS_UPSAMPLE_REQUIRED
+#ifdef UPSAMPLE_PASS
 #	include "ShadowFactorsUpsampleCommon.ush"
 #endif
 
 #ifndef THREADGROUP_SIZEX
-#define THREADGROUP_SIZEX 1
+#	define THREADGROUP_SIZEX 1
 #endif
 
 #ifndef THREADGROUP_SIZEY
-#define THREADGROUP_SIZEY 1
+#	define THREADGROUP_SIZEY 1
 #endif
 
 #ifndef LIGHT_SOURCE_MODE
@@ -1186,6 +1186,8 @@ void CapsuleShadowingUpsampleVS(
 	}
 }
 
+#ifdef UPSAMPLE_PASS
+
 Texture2D ShadowFactorsTexture;
 SamplerState ShadowFactorsSampler;
 
@@ -1222,4 +1224,6 @@ void CapsuleShadowingUpsamplePS(
 #if APPLY_TO_SSAO
 	OutAmbientOcclusion = Output;
 #endif
-}
\ No newline at end of file
+}
+
+#endif // UPSAMPLE_PASS
diff --git a/Engine/Shaders/Private/DistanceFieldShadowing.usf b/Engine/Shaders/Private/DistanceFieldShadowing.usf
--- a/Engine/Shaders/Private/DistanceFieldShadowing.usf
+++ b/Engine/Shaders/Private/DistanceFieldShadowing.usf
@@ -11,7 +11,7 @@
 #include "DistanceFieldShadowingShared.ush"
 #include "Substrate/Substrate.ush"
 
-#ifdef SHADOW_FACTORS_UPSAMPLE_REQUIRED
+#ifdef UPSAMPLE_PASS
 #	include "ShadowFactorsUpsampleCommon.ush"
 #endif
 
@@ -772,6 +772,8 @@ void DistanceFieldShadowingCS(
 	RWShadowFactors[DispatchThreadId.xy] = float2(Result, SceneDepth);
 }
 
+#ifdef UPSAMPLE_PASS
+
 float FadePlaneOffset;
 float InvFadePlaneLength;
 float NearFadePlaneOffset;
@@ -796,6 +798,8 @@ void DistanceFieldShadowingUpsamplePS(
 	OutColor = EncodeLightAttenuation(half4(Output, Output, Output, Output));
 }
 
+#endif // UPSAMPLE_PASS
+
 #ifdef SHADOW_TILE_VS
 
 Buffer<uint> TileListData;
diff --git a/Engine/Source/Runtime/Renderer/Private/CapsuleShadowRendering.cpp b/Engine/Source/Runtime/Renderer/Private/CapsuleShadowRendering.cpp
--- a/Engine/Source/Runtime/Renderer/Private/CapsuleShadowRendering.cpp
+++ b/Engine/Source/Runtime/Renderer/Private/CapsuleShadowRendering.cpp
@@ -344,6 +344,7 @@ public:
 
 	static void ModifyCompilationEnvironment(const FGlobalShaderPermutationParameters& Parameters, FShaderCompilerEnvironment& OutEnvironment)
 	{
+		OutEnvironment.SetDefine(TEXT("UPSAMPLE_PASS"), 1);
 		OutEnvironment.SetDefine(TEXT("DOWNSAMPLE_FACTOR"), 2);
 	}
 };
diff --git a/Engine/Source/Runtime/Renderer/Private/DistanceFieldShadowing.cpp b/Engine/Source/Runtime/Renderer/Private/DistanceFieldShadowing.cpp
--- a/Engine/Source/Runtime/Renderer/Private/DistanceFieldShadowing.cpp
+++ b/Engine/Source/Runtime/Renderer/Private/DistanceFieldShadowing.cpp
@@ -360,6 +360,7 @@ class FDistanceFieldShadowingUpsamplePS : public FGlobalShader
 
 	static void ModifyCompilationEnvironment(const FGlobalShaderPermutationParameters& Parameters, FShaderCompilerEnvironment& OutEnvironment)
 	{
+		OutEnvironment.SetDefine(TEXT("UPSAMPLE_PASS"), 1);
 		OutEnvironment.SetDefine(TEXT("FORCE_DEPTH_TEXTURE_READS"), 1);
 	}
 };
