Remove clamp from BloomSetupCommon since we use pre-exposure now

diff --git a/Engine/Shaders/Private/PostProcessBloom.usf b/Engine/Shaders/Private/PostProcessBloom.usf
--- a/Engine/Shaders/Private/PostProcessBloom.usf
+++ b/Engine/Shaders/Private/PostProcessBloom.usf
@@ -53,9 +53,6 @@ float4 BloomSetupCommon(float2 UV, float2 ViewportUV, float2 ExposureScaleMiddle
 	}
 #endif
 
-	// clamp to avoid artifacts from exceeding fp16 through framebuffer blending of multiple very bright lights
-	SceneColor.rgb = min(float3(256 * 256, 256 * 256, 256 * 256), SceneColor.rgb);
-
 	half3 LinearColor = SceneColor.rgb;
 
 	half TotalLuminance = Luminance(LinearColor) * ExposureScale;
