Slate MSDF text additive blend mode fix, general text alpha composite blend mode fix

#rnx <Release Notes exclude>

diff --git a/Engine/Shaders/Private/SlateElementPixelShader.usf b/Engine/Shaders/Private/SlateElementPixelShader.usf
--- a/Engine/Shaders/Private/SlateElementPixelShader.usf
+++ b/Engine/Shaders/Private/SlateElementPixelShader.usf
@@ -287,7 +287,11 @@ half4 GetSdfTextElementColor( VertexToPixelInterpolants InVertex )
 	const float2 ScreenSpaceTextureSize = float2(1.0, 1.0) / fwidth(AtlasTextureCoordinates);
 	const float ScreenSpaceSpread = max(dot(ShaderParams.xy, ScreenSpaceTextureSize), 1.0);
 	const float ScreenSpaceDistance = (SampledDistance - ShaderParams.z) * ScreenSpaceSpread;
-	OutColor.a *= half(saturate(ScreenSpaceDistance + 0.5));
+	#if USE_MATERIALS && (MATERIALBLENDING_ADDITIVE || MATERIALBLENDING_ALPHACOMPOSITE)
+		OutColor *= half(saturate(ScreenSpaceDistance + 0.5));
+	#else
+		OutColor.a *= half(saturate(ScreenSpaceDistance + 0.5));
+	#endif
 
 	return OutColor;
 }
@@ -305,7 +309,7 @@ half4 GetFontElementColor( VertexToPixelInterpolants InVertex )
 	half4 OutColor = BaseColor;
 
 #if SHADER_TYPE == ST_GrayscaleFont
-	#if USE_MATERIALS && MATERIALBLENDING_ADDITIVE
+	#if USE_MATERIALS && (MATERIALBLENDING_ADDITIVE || MATERIALBLENDING_ALPHACOMPOSITE)
 		OutColor *= Texture2DSample_A8( ElementTexture, ElementTextureSampler, AtlasTextureCoordinates );
 	#else
 		OutColor.a *= Texture2DSample_A8( ElementTexture, ElementTextureSampler, AtlasTextureCoordinates );
