Fixes debug traces failing to draw correctly in offset stereo views, corrects improper UV scaling in MainPopulateSceneDepthPS and reduces special cases during shader setup.

diff --git a/Engine/Shaders/Private/PostProcessCompositePrimitives.usf b/Engine/Shaders/Private/PostProcessCompositePrimitives.usf
--- a/Engine/Shaders/Private/PostProcessCompositePrimitives.usf
+++ b/Engine/Shaders/Private/PostProcessCompositePrimitives.usf
@@ -253,24 +253,17 @@ void MainPopulateSceneDepthPS(
 	float2 DepthUV = clamp(
 		(Depth_ViewportMin + ViewportUV * Depth_ViewportSize - DepthTextureJitter) * Depth_ExtentInverse,
 		Depth_UVViewportBilinearMin, Depth_UVViewportBilinearMax);
+	DepthUV = clamp(DepthUV, float2(0, 0), float2(1, 1));
 
 #if FORCE_DEBUG_DRAW_COLOR
 	//This extended pass is primarily to compensate for the lack of scene color output when compositing the debug draw primitives with SimpleElementShaders
 	//Draw out the color whilst compensating for the viewport not necessarily starting at 0,0. 
 	//Does not use jitter for Scene Color as it is not upscaled for this pass
-	float2 ColorUV = clamp(Color_UVViewportMin + (ViewportUV * Color_UVViewportSize), 
-						Color_UVViewportBilinearMin,
-						Color_UVViewportBilinearMax);
+	float2 ColorUV = clamp(
+		(Color_ViewportMin + ViewportUV * Color_ViewportSize) * Color_ExtentInverse,
+		Color_UVViewportBilinearMin, Color_UVViewportBilinearMax);
 	ColorUV = clamp(ColorUV, float2(0, 0), float2(1, 1));
 	OutColor = Texture2DSampleLevel(ColorTexture, ColorSampler, ColorUV, 0).rgba;
-
-	//As this is also outputting color and positions need to match, compensate the depth for the viewport position not necessarily starting at 0,0
-	DepthUV = DepthUV - Depth_UVViewportMin;
-	//At larger viewport sizes, the temporal depth upscale can increase the depth viewport size significantly, 
-	//which isn't necessarily compensated for when outputting the color, so scale the DepthUV output down to match the ColorUV viewport size
-	//Clamped to the Min/Max temporal upscale values defined in PostProcessCompositePrimitivesCommon.cpp
-	DepthUV *= clamp(Depth_Extent / Color_Extent, MIN_DEPTHTEX_UPSCALE_FACTOR, MAX_DEPTHTEX_UPSCALE_FACTOR);
-	DepthUV = clamp(DepthUV, float2(0, 0), float2(1, 1));
 #else
 	OutColor = 0.0;
 #endif
diff --git a/Engine/Source/Runtime/Renderer/Private/PostProcess/PostProcessCompositePrimitivesCommon.cpp b/Engine/Source/Runtime/Renderer/Private/PostProcess/PostProcessCompositePrimitivesCommon.cpp
--- a/Engine/Source/Runtime/Renderer/Private/PostProcess/PostProcessCompositePrimitivesCommon.cpp
+++ b/Engine/Source/Runtime/Renderer/Private/PostProcess/PostProcessCompositePrimitivesCommon.cpp
@@ -256,15 +256,11 @@ void PopulateDepthPass(FRDGBuilder& GraphBuilder,
 	{
 		PassParameters->ColorTexture = InSceneColor.Texture;
 		PassParameters->ColorSampler = PointClampSampler;
-
-		PassParameters->Color = GetScreenPassTextureViewportParameters(View.StereoPass == EStereoscopicPass::eSSP_FULL ? FScreenPassTextureViewport(InSceneColor.Texture->Desc.Extent, View.ViewRect) : FScreenPassTextureViewport(InSceneColor));
-		PassParameters->Depth = GetScreenPassTextureViewportParameters(FScreenPassTextureViewport(InSceneDepth.Texture->Desc.Extent, View.ViewRect));
-	}
-	else
-	{
-		PassParameters->Color = GetScreenPassTextureViewportParameters(FScreenPassTextureViewport(InSceneColor));
-		PassParameters->Depth = GetScreenPassTextureViewportParameters(FScreenPassTextureViewport(InSceneDepth));
 	}
+
+	PassParameters->Color = GetScreenPassTextureViewportParameters(FScreenPassTextureViewport(InSceneColor));
+	PassParameters->Depth = GetScreenPassTextureViewportParameters(FScreenPassTextureViewport(InSceneDepth));
+
 	PassParameters->DepthTexture = InSceneDepth.Texture;
 	PassParameters->DepthSampler = PointClampSampler;
 	PassParameters->DepthTextureJitter = SceneDepthJitter;
