Moved part of ShadowRayCast(...) to helper function CastScreenSpaceShadowRay(...)
- to be shared by DeferredLighting and standalone Screen Space Shadows pass (could also be used in VirtualShadowMapScreenRayCast).

diff --git a/Engine/Shaders/Private/DeferredLightingCommon.ush b/Engine/Shaders/Private/DeferredLightingCommon.ush
--- a/Engine/Shaders/Private/DeferredLightingCommon.ush
+++ b/Engine/Shaders/Private/DeferredLightingCommon.ush
@@ -11,6 +11,7 @@
 #include "IESLightProfilesCommon.ush"
 #include "CapsuleLightIntegrate.ush"
 #include "RectLightIntegrate.ush"
+#include "ScreenSpaceShadowRayCast.ush"
 #include "Substrate/Substrate.ush"
 #include "LightData.ush"
 
@@ -42,79 +43,30 @@ float DistanceFromCameraFade(float SceneDepth, FDeferredLightData LightData)
 // Sets bOutHitCastDynamicShadow if the hit point is marked as a dynamic shadow caster
 float ShadowRayCast(
 	float3 RayOriginTranslatedWorld, float3 RayDirection, float RayLength,
-	int NumSteps, float StepOffset, bool bHairNoShadowLight, out bool bOutHitCastContactShadow )
+	int NumSteps, float Dither, bool bHairNoShadowLight, out bool bOutHitCastContactShadow)
 {
-	float4 RayStartClip	= mul( float4( RayOriginTranslatedWorld, 1 ), View.TranslatedWorldToClip );
-	float4 RayDirClip	= mul( float4( RayDirection * RayLength, 0 ), View.TranslatedWorldToClip );
-	float4 RayEndClip	= RayStartClip + RayDirClip;
-
-	float3 RayStartScreen = RayStartClip.xyz / RayStartClip.w;
-	float3 RayEndScreen = RayEndClip.xyz / RayEndClip.w;
-	
-	float3 RayStepScreen = RayEndScreen - RayStartScreen;
-
-	float3 RayStartUVz = float3( RayStartScreen.xy * View.ScreenPositionScaleBias.xy + View.ScreenPositionScaleBias.wz, RayStartScreen.z );
-	float3 RayStepUVz = float3( RayStepScreen.xy * View.ScreenPositionScaleBias.xy, RayStepScreen.z );
-
-	float4 RayDepthClip	= RayStartClip + mul( float4( 0, 0, RayLength, 0 ), View.ViewToClip );
-	float3 RayDepthScreen = RayDepthClip.xyz / RayDepthClip.w;
-
-	const float Step = 1.0 / NumSteps;
-
 	// *2 to get less moire pattern in extreme cases, larger values make object appear not grounded in reflections
-	const float CompareTolerance = abs( RayDepthScreen.z - RayStartScreen.z ) * Step * 2;
-
-	float SampleTime = StepOffset * Step + Step;
-
-	float FirstHitTime = -1.0;
-
-	const float StartDepth = SceneTexturesStruct.SceneDepthTexture.SampleLevel( SceneTexturesStruct_SceneDepthTextureSampler, RayStartUVz.xy, 0 ).r;
-
-	UNROLL
-	for( int i = 0; i < NumSteps; i++ )
-	{
-		float3 SampleUVz = RayStartUVz + RayStepUVz * SampleTime;
-		float SampleDepth = SceneTexturesStruct.SceneDepthTexture.SampleLevel( SceneTexturesStruct_SceneDepthTextureSampler, SampleUVz.xy, 0 ).r;
-
-		float DepthDiff = SampleUVz.z - SampleDepth;
-		bool Hit = abs( DepthDiff + CompareTolerance ) < CompareTolerance;
-		// Avoid self-intersection with the start pixel (exact comparison due to point sampling depth buffer)
-		// Exception is made for hair for occluding transmitted light with non-shadow casting light
-		if (!bHairNoShadowLight)
-		{
-			Hit = Hit && (SampleDepth != StartDepth);
-		}
-
-		FirstHitTime = (Hit && FirstHitTime < 0.0) ? SampleTime : FirstHitTime;
+	const float CompareToleranceScale = 2.0f;
 
-		SampleTime += Step;
-	}
+	float2 HitUV;
+	float HitDistance = CastScreenSpaceShadowRay(RayOriginTranslatedWorld, RayDirection, RayLength, NumSteps, Dither, CompareToleranceScale, bHairNoShadowLight, HitUV);
 
-	float HitDistance = -1.0;
 	bOutHitCastContactShadow = false;
-	if ( FirstHitTime > 0.0 )
+	if (HitDistance > 0.0)
 	{
-		// Ignore hits that come from non-shadow-casting pixels
-		float3 SampleUVz = RayStartUVz + RayStepUVz * FirstHitTime;
-	#if SUBSTRATE_ENABLED
-		#if SUBSTRATE_RAYCAST_ENABLED && SUBSTRATE_MATERIALCONTAINER_IS_VIEWRESOURCE
-		uint2 PixelPos = View.ViewRectMin.xy + View.ViewSizeAndInvSize.xy * SampleUVz.xy;
+#if SUBSTRATE_ENABLED
+#if SUBSTRATE_RAYCAST_ENABLED && SUBSTRATE_MATERIALCONTAINER_IS_VIEWRESOURCE
+		uint2 PixelPos = View.ViewRectMin.xy + View.ViewSizeAndInvSize.xy * HitUV;
 		FSubstrateAddressing SubstrateAddressing = GetSubstratePixelDataByteOffset(PixelPos, uint2(View.BufferSizeAndInvSize.xy), SubstrateRayCast.MaxBytesPerPixel);
 		FSubstratePixelHeader SubstratePixelHeader = UnpackSubstrateHeaderIn(SubstrateRayCast.MaterialTextureArray, SubstrateAddressing, SubstrateRayCast.TopLayerTexture);
 		bOutHitCastContactShadow = SubstratePixelHeader.DoesCastContactShadow();
-		#endif
-	#else
-		FGBufferData SampleGBuffer = GetGBufferData( SampleUVz.xy );
+#endif
+#else
+		FGBufferData SampleGBuffer = GetGBufferData(HitUV);
 		bOutHitCastContactShadow = CastContactShadow(SampleGBuffer);
-	#endif
+#endif
+	}
 
-		// Off screen masking
-		float3 HitUVz = RayStartUVz + RayStepUVz * FirstHitTime;
-		bool bValidUV = all(and(0.0 < HitUVz.xy, HitUVz.xy < 1.0));
-		HitDistance = bValidUV ? ( FirstHitTime * RayLength ) : -1.0;
-	}	
-	
-	
 	return HitDistance;
 }
 #endif
@@ -202,10 +154,9 @@ void GetShadowTerms(float SceneDepth, half4 PrecomputedShadowFactors, uint Shadi
 	BRANCH
 	if (ContactShadowLength > 0.0)
 	{
-		float StepOffset = Dither - 0.5;
 		bool bHitCastContactShadow = false;
 		bool bHairNoShadowLight = ShadingModelID == SHADINGMODELID_HAIR && !LightData.ShadowedBits;
-		float HitDistance = ShadowRayCast( TranslatedWorldPosition, L, ContactShadowLength, 8, StepOffset, bHairNoShadowLight, bHitCastContactShadow );
+		float HitDistance = ShadowRayCast( TranslatedWorldPosition, L, ContactShadowLength, 8, Dither, bHairNoShadowLight, bHitCastContactShadow );
 				
 		if ( HitDistance > 0.0 )
 		{
diff --git a/Engine/Shaders/Private/ScreenSpaceShadowRayCast.ush b/Engine/Shaders/Private/ScreenSpaceShadowRayCast.ush
new file mode 100644
--- /dev/null
+++ b/Engine/Shaders/Private/ScreenSpaceShadowRayCast.ush
@@ -0,0 +1,68 @@
+// Copyright Epic Games, Inc. All Rights Reserved.
+
+#pragma once
+
+#include "Common.ush"
+#include "SceneTexturesCommon.ush"
+
+// Returns distance along ray that the first hit occurred, or negative on miss
+float CastScreenSpaceShadowRay(
+	float3 RayOriginTranslatedWorld, float3 RayDirection, float RayLength, int NumSteps,
+	float Dither, float CompareToleranceScale, bool bHairNoShadowLight,
+	out float2 HitUV)
+{
+	float4 RayStartClip = mul(float4(RayOriginTranslatedWorld, 1), View.TranslatedWorldToClip);
+	float4 RayDirClip = mul(float4(RayDirection * RayLength, 0), View.TranslatedWorldToClip);
+	float4 RayEndClip = RayStartClip + RayDirClip;
+
+	float3 RayStartScreen = RayStartClip.xyz / RayStartClip.w;
+	float3 RayEndScreen = RayEndClip.xyz / RayEndClip.w;
+
+	float3 RayStepScreen = RayEndScreen - RayStartScreen;
+
+	float3 RayStartUVz = float3(RayStartScreen.xy * View.ScreenPositionScaleBias.xy + View.ScreenPositionScaleBias.wz, RayStartScreen.z);
+	float3 RayStepUVz = float3(RayStepScreen.xy * View.ScreenPositionScaleBias.xy, RayStepScreen.z);
+
+	float4 RayDepthClip = RayStartClip + mul(float4(0, 0, RayLength, 0), View.ViewToClip);
+	float3 RayDepthScreen = RayDepthClip.xyz / RayDepthClip.w;
+
+	const float StepOffset = Dither - 0.5f;
+	const float Step = 1.0 / NumSteps;
+
+	const float CompareTolerance = abs(RayDepthScreen.z - RayStartScreen.z) * Step * CompareToleranceScale;
+
+	float SampleTime = StepOffset * Step + Step;
+
+	float FirstHitTime = -1.0;
+
+	const float StartDepth = LookupDeviceZ(RayStartUVz.xy);
+
+	UNROLL
+	for (int i = 0; i < NumSteps; i++)
+	{
+		float3 SampleUVz = RayStartUVz + RayStepUVz * SampleTime;
+		float SampleDepth = LookupDeviceZ(SampleUVz.xy);
+
+		// Avoid self-intersection with the start pixel (exact comparison due to point sampling depth buffer)
+		// Exception is made for hair for occluding transmitted light with non-shadow casting light
+		if (SampleDepth != StartDepth || bHairNoShadowLight)
+		{
+			float DepthDiff = SampleUVz.z - SampleDepth;
+			bool Hit = abs(DepthDiff + CompareTolerance) < CompareTolerance;
+
+			if (Hit)
+			{
+				HitUV = SampleUVz.xy;
+
+				// Off screen masking
+				bool bValidUV = all(and (0.0 < SampleUVz.xy, SampleUVz.xy < 1.0));
+
+				return bValidUV ? (RayLength * SampleTime) : -1.0;
+			}
+		}
+
+		SampleTime += Step;
+	}
+
+	return -1;
+}
\ No newline at end of file
