Fixing invalid access to uniform buffer.

diff --git a/Engine/Shaders/Private/InstancedStereo.ush b/Engine/Shaders/Private/InstancedStereo.ush
--- a/Engine/Shaders/Private/InstancedStereo.ush
+++ b/Engine/Shaders/Private/InstancedStereo.ush
@@ -53,7 +53,7 @@ static const float4 EyeClipEdge[2] = { float4(-1.0, 0.0, 0.0, 1.0), float4(1.0,
 bool IsInstancedStereo()
 {
 #if INSTANCED_STEREO
-	return (uint)InstancedView_StereoPassIndex[1] > 0;
+	return (uint)InstancedView.StereoPassIndex[1] > 0;
 #else
 	return false;
 #endif
