Fixed translucent material been cutoff by the skylight in RT reflection

diff --git a/Engine/Shaders/Private/RayTracing/RayTracingReflections.usf b/Engine/Shaders/Private/RayTracing/RayTracingReflections.usf
--- a/Engine/Shaders/Private/RayTracing/RayTracingReflections.usf
+++ b/Engine/Shaders/Private/RayTracing/RayTracingReflections.usf
@@ -609,6 +609,8 @@ RAY_TRACING_ENTRY_RAYGEN(RayTracingReflectionsRGS)
 				// Trace the top layer reflection
 				float3 TopLayerRadiance = float3(0, 0, 0);
 				FPackedMaterialClosestHitPayload PackedPayload = (FPackedMaterialClosestHitPayload)0;
+				// Trace separate ray for translucent objects
+				float Transmission = 1.0;
 
 				if (DIM_HYBRID && DIM_DEFERRED_MATERIAL_MODE == DEFERRED_MATERIAL_MODE_SHADE 
 					&& DeferredMaterialPayload.SortKey == RAY_TRACING_DEFERRED_MATERIAL_KEY_SCREEN_SPACE)
@@ -648,8 +650,6 @@ RAY_TRACING_ENTRY_RAYGEN(RayTracingReflectionsRGS)
 						RayCone,
 						bEnableSkyLightContribution);
 
-					// Trace separate ray for translucent objects
-					float Transmission = 1.0;
 					//#dxr-dxr_todo: do permutation instead of dynamic branching?
 					if (EnableTranslucency)
 					{
@@ -717,7 +717,7 @@ RAY_TRACING_ENTRY_RAYGEN(RayTracingReflectionsRGS)
 				{
 					float SkyAverageBrightness = 1.0f;
 					float3 SkyLighting = GetSkyLightReflection(TopLayerRay.Direction, TopLayerRoughness, SkyAverageBrightness);
-					TopLayerRadiance = SkyLighting;
+					TopLayerRadiance += SkyLighting * Transmission;
 					bSkyWasHit = true;	// To stop the main path recusion early
 
 					bApplyHeightFog = HeightFog > 0;
