Saves VALU in TSR's convolutions arround packing and unpacking LDS

diff --git a/Engine/Shaders/Private/TemporalSuperResolution/TSRConvolutionNetwork.ush b/Engine/Shaders/Private/TemporalSuperResolution/TSRConvolutionNetwork.ush
--- a/Engine/Shaders/Private/TemporalSuperResolution/TSRConvolutionNetwork.ush
+++ b/Engine/Shaders/Private/TemporalSuperResolution/TSRConvolutionNetwork.ush
@@ -197,191 +197,6 @@ TLaneVector2D<bool, 1, LaneStrideX, LaneStrideY> IsValidTilePixel(const uint2 Ti
 }
 
 
-//------------------------------------------------------- GROUP SHARED COMMUNICATION
-
-#if defined(LDS_SIZE)
-
-#if CONFIG_COMPILE_FP16
-
-groupshared uint2 SharedData[LDS_SIZE];
-
-void IncrementAtomic(const uint SharedIndex, uint Value)
-{
-	return InterlockedAdd(/* inout */ SharedData[SharedIndex].x, Value);
-}
-
-void WriteAtomic(const uint SharedIndex, uint X)
-{
-	SharedData[SharedIndex].x = X;
-}
-
-uint ReadAtomic(const uint SharedIndex)
-{
-	return SharedData[SharedIndex].x;
-}
-
-#if COMPILER_SUPPORT_UINT16_BITCAST
-
-void WriteElementToLDS(const uint SharedIndex, tsr_ushort4 X)
-{
-	SharedData[SharedIndex] = uint2(
-		bit_cast_uint(X.xy),
-		bit_cast_uint(X.zw));
-}
-
-void ReadElementFromLDS(const uint SharedIndex, out tsr_ushort4 X)
-{
-	uint2 RX = SharedData[SharedIndex];
-	X.xy = bit_cast_uint16_t2(RX.x);
-	X.zw = bit_cast_uint16_t2(RX.y);
-}
-
-#else
-
-void WriteElementToLDS(const uint SharedIndex, tsr_ushort4 X)
-{
-	SharedData[SharedIndex] = uint2(
-		uint(X.x) | (uint(X.y) << 16),
-		uint(X.z) | (uint(X.w) << 16));
-}
-
-void ReadElementFromLDS(const uint SharedIndex, out tsr_ushort4 X)
-{
-	uint2 RX = SharedData[SharedIndex];
-	X.x = tsr_ushort((RX.x >>  0u) & 0xFFFFu);
-	X.y = tsr_ushort((RX.x >> 16u) & 0xFFFFu);
-	X.z = tsr_ushort((RX.y >>  0u) & 0xFFFFu);
-	X.w = tsr_ushort((RX.y >> 16u) & 0xFFFFu);
-}
-
-#endif
-
-#else
-
-groupshared uint4 SharedData[LDS_SIZE];
-
-void IncrementAtomic(const uint SharedIndex, uint Value)
-{
-	return InterlockedAdd(/* inout */ SharedData[SharedIndex].x, Value);
-}
-
-void WriteAtomic(const uint SharedIndex, uint X)
-{
-	SharedData[SharedIndex].x = X;
-}
-
-uint ReadAtomic(const uint SharedIndex)
-{
-	return SharedData[SharedIndex].x;
-}
-
-void WriteElementToLDS(const uint SharedIndex, tsr_ushort4 X)
-{
-	SharedData[SharedIndex] = uint4(X);
-}
-
-void ReadElementFromLDS(const uint SharedIndex, out tsr_ushort4 X)
-{
-	X = tsr_ushort4(SharedData[SharedIndex].xyzw);
-}
-
-#endif
-
-
-
-void WriteElementToLDS(const uint SharedIndex, tsr_ushort X)
-{
-	WriteElementToLDS(SharedIndex, tsr_ushort4(X, 0.0, 0.0, 0.0));
-}
-
-void WriteElementToLDS(const uint SharedIndex, tsr_ushort2 X)
-{
-	WriteElementToLDS(SharedIndex, tsr_ushort4(X, 0.0, 0.0));
-}
-
-void WriteElementToLDS(const uint SharedIndex, tsr_ushort3 X)
-{
-	WriteElementToLDS(SharedIndex, tsr_ushort4(X, 0.0));
-}
-
-void WriteElementToLDS(const uint SharedIndex, tsr_half X)
-{
-	WriteElementToLDS(SharedIndex, as_tsr_ushort(X));
-}
-
-void WriteElementToLDS(const uint SharedIndex, tsr_half2 X)
-{
-	WriteElementToLDS(SharedIndex, tsr_ushort2(as_tsr_ushort(X.x), as_tsr_ushort(X.y)));
-}
-
-void WriteElementToLDS(const uint SharedIndex, tsr_half3 X)
-{
-	WriteElementToLDS(SharedIndex, tsr_ushort3(as_tsr_ushort(X.x), as_tsr_ushort(X.y), as_tsr_ushort(X.z)));
-}
-
-void WriteElementToLDS(const uint SharedIndex, tsr_half4 X)
-{
-	WriteElementToLDS(SharedIndex, tsr_ushort4(as_tsr_ushort(X.x), as_tsr_ushort(X.y), as_tsr_ushort(X.z), as_tsr_ushort(X.w)));
-}
-
-void ReadElementFromLDS(const uint SharedIndex, out tsr_ushort X)
-{
-	tsr_ushort4 UX;
-	ReadElementFromLDS(SharedIndex, /* out */ UX);
-	X = UX.x;
-}
-
-void ReadElementFromLDS(const uint SharedIndex, out tsr_ushort2 X)
-{
-	tsr_ushort4 UX;
-	ReadElementFromLDS(SharedIndex, /* out */ UX);
-	X = UX.xy;
-}
-
-void ReadElementFromLDS(const uint SharedIndex, out tsr_ushort3 X)
-{
-	tsr_ushort4 UX;
-	ReadElementFromLDS(SharedIndex, /* out */ UX);
-	X = UX.xyz;
-}
-
-void ReadElementFromLDS(const uint SharedIndex, out tsr_half X)
-{
-	tsr_ushort4 UX;
-	ReadElementFromLDS(SharedIndex, /* out */ UX);
-	X = as_tsr_half(UX.x);
-}
-
-void ReadElementFromLDS(const uint SharedIndex, out tsr_half2 X)
-{
-	tsr_ushort4 UX;
-	ReadElementFromLDS(SharedIndex, /* out */ UX);
-	X.x = as_tsr_half(UX.x);
-	X.y = as_tsr_half(UX.y);
-}
-
-void ReadElementFromLDS(const uint SharedIndex, out tsr_half3 X)
-{
-	tsr_ushort4 UX;
-	ReadElementFromLDS(SharedIndex, /* out */ UX);
-	X.x = as_tsr_half(UX.x);
-	X.y = as_tsr_half(UX.y);
-	X.z = as_tsr_half(UX.z);
-}
-
-void ReadElementFromLDS(const uint SharedIndex, out tsr_half4 X)
-{
-	tsr_ushort4 UX;
-	ReadElementFromLDS(SharedIndex, /* out */ UX);
-	X.x = as_tsr_half(UX.x);
-	X.y = as_tsr_half(UX.y);
-	X.z = as_tsr_half(UX.z);
-	X.w = as_tsr_half(UX.w);
-}
-
-#endif // defined(LDS_SIZE)
-
-
 //------------------------------------------------------- VOTE
 
 CALL_SITE_DEBUGLOC
@@ -392,16 +207,16 @@ bool AnyThread(bool X)
 	
 	if (GGroupThreadIndex == 0)
 	{
-		WriteAtomic(/* SharedIndex = */ 0, 0);
+		WriteDwordToLDS(/* SharedIndex = */ 0, 0);
 	}
 
 	GroupMemoryBarrierWithGroupSync();
 	
-	IncrementAtomic(/* SharedIndex = */ 0, uint(X));
+	AtomicIncrementLDSDword(/* SharedIndex = */ 0, uint(X));
 	
 	GroupMemoryBarrierWithGroupSync();
 	
-	return ReadAtomic(/* SharedIndex = */ 0) > 0;
+	return ReadDwordFromLDS(/* SharedIndex = */ 0) > 0;
 }
 #elif !COMPILER_SUPPORTS_WAVE_VOTE && !defined(LDS_SIZE)
 {
@@ -427,14 +242,12 @@ bool AnyThread(bool X)
 
 	if (WaveIsFirstLane())
 	{
-		WriteElementToLDS(GGroupThreadIndex / kLaneCount, tsr_ushort(AnyX));
+		WriteDwordToLDS(GGroupThreadIndex / kLaneCount, uint(AnyX));
 	}
 		
 	GroupMemoryBarrierWithGroupSync();
 		
-	tsr_ushort WaveX;
-	ReadElementFromLDS(GGroupThreadIndex % kWaveCount, /* out */ WaveX);
-	//WaveX = select(WaveGetLaneIndex() < kWaveCount, WaveX, 0);
+	uint WaveX = ReadDwordFromLDS(GGroupThreadIndex % kWaveCount);
 
 	return WaveActiveAnyTrue(WaveX != 0);
 }
@@ -508,7 +321,7 @@ void AccessNeighborTexels3x1(
 		UNROLL
 		for (uint SimdIndex = 0; SimdIndex < LaneStrideY; SimdIndex++)
 		{
-			WriteElementToLDS(GGroupThreadIndex + SimdIndex * TotalNumberOfLaneInGroup, CO.GetElement(SimdIndex));
+			WriteVectorToLDS(GGroupThreadIndex + SimdIndex * TotalNumberOfLaneInGroup, CO.GetElement(SimdIndex));
 		}
 		GroupMemoryBarrierWithGroupSync();
 		
@@ -517,8 +330,8 @@ void AccessNeighborTexels3x1(
 		{
 			vector<ScalarType, VectorSize> FirstElement;
 			vector<ScalarType, VectorSize> LastElement;
-			ReadElementFromLDS(((GGroupThreadIndex - 1) % TotalNumberOfLaneInGroup) + SimdIndex * TotalNumberOfLaneInGroup, /* out */ FirstElement);
-			ReadElementFromLDS(((GGroupThreadIndex + 1) % TotalNumberOfLaneInGroup) + SimdIndex * TotalNumberOfLaneInGroup, /* out */ LastElement);
+			ReadVectorFromLDS(((GGroupThreadIndex - 1) % TotalNumberOfLaneInGroup) + SimdIndex * TotalNumberOfLaneInGroup, /* out */ FirstElement);
+			ReadVectorFromLDS(((GGroupThreadIndex + 1) % TotalNumberOfLaneInGroup) + SimdIndex * TotalNumberOfLaneInGroup, /* out */ LastElement);
 	
 			CP.SetElement(SimdIndex, LastElement);
 			CN.SetElement(SimdIndex, FirstElement);
@@ -558,7 +371,7 @@ void AccessNeighborTexels1x3(
 		UNROLL
 		for (uint SimdIndex = 0; SimdIndex < LaneStrideX; SimdIndex++)
 		{
-			WriteElementToLDS(GGroupThreadIndex + SimdIndex * TotalNumberOfLaneInGroup, CO.GetElement(SimdIndex));
+			WriteVectorToLDS(GGroupThreadIndex + SimdIndex * TotalNumberOfLaneInGroup, CO.GetElement(SimdIndex));
 		}
 		GroupMemoryBarrierWithGroupSync();
 		
@@ -567,8 +380,8 @@ void AccessNeighborTexels1x3(
 		{
 			vector<ScalarType, VectorSize> FirstElement;
 			vector<ScalarType, VectorSize> LastElement;
-			ReadElementFromLDS(((GGroupThreadIndex - LANE_COUNT_X * WAVE_COUNT_X) % TotalNumberOfLaneInGroup) + SimdIndex * TotalNumberOfLaneInGroup, /* out */ FirstElement);
-			ReadElementFromLDS(((GGroupThreadIndex + LANE_COUNT_X * WAVE_COUNT_X) % TotalNumberOfLaneInGroup) + SimdIndex * TotalNumberOfLaneInGroup, /* out */ LastElement);
+			ReadVectorFromLDS(((GGroupThreadIndex - LANE_COUNT_X * WAVE_COUNT_X) % TotalNumberOfLaneInGroup) + SimdIndex * TotalNumberOfLaneInGroup, /* out */ FirstElement);
+			ReadVectorFromLDS(((GGroupThreadIndex + LANE_COUNT_X * WAVE_COUNT_X) % TotalNumberOfLaneInGroup) + SimdIndex * TotalNumberOfLaneInGroup, /* out */ LastElement);
 	
 			CP.SetElement(SimdIndex, LastElement);
 			CN.SetElement(SimdIndex, FirstElement);
@@ -576,20 +389,20 @@ void AccessNeighborTexels1x3(
 	}
 	else
 	{
-		const uint FirstSimdIndex = 0;
-		const uint LastSimdIndex = LaneStrideX * LaneStrideY - LaneStrideX;
+		const uint FirstElementIndex = 0;
+		const uint LastElementIndex = LaneStrideX * LaneStrideY - LaneStrideX;
 
 		GroupMemoryBarrierWithGroupSync();
 		UNROLL
 		for (uint SimdIndex = 0; SimdIndex < LaneStrideX; SimdIndex++)
 		{
-			WriteElementToLDS(GGroupThreadIndex + (0 + SimdIndex) * TotalNumberOfLaneInGroup, CO.GetElement(FirstSimdIndex + SimdIndex));
-			WriteElementToLDS(GGroupThreadIndex + (2 + SimdIndex) * TotalNumberOfLaneInGroup, CO.GetElement(LastSimdIndex + SimdIndex));
+			WriteVectorToLDS(GGroupThreadIndex + (0 + SimdIndex) * TotalNumberOfLaneInGroup, CO.GetElement(FirstElementIndex + SimdIndex));
+			WriteVectorToLDS(GGroupThreadIndex + (LaneStrideX + SimdIndex) * TotalNumberOfLaneInGroup, CO.GetElement(LastElementIndex + SimdIndex));
 		}
 		GroupMemoryBarrierWithGroupSync();
 	
 		UNROLL
-		for (uint SimdIndex = 0; SimdIndex < LastSimdIndex; SimdIndex++)
+		for (uint SimdIndex = 0; SimdIndex < LastElementIndex; SimdIndex++)
 		{
 			CP.SetElement(SimdIndex + 0, CO.GetElement(SimdIndex + LaneStrideX));
 			CN.SetElement(SimdIndex + LaneStrideX, CO.GetElement(SimdIndex + 0));
@@ -600,11 +413,11 @@ void AccessNeighborTexels1x3(
 		{
 			vector<ScalarType, VectorSize> FirstElement;
 			vector<ScalarType, VectorSize> LastElement;
-			ReadElementFromLDS(((GGroupThreadIndex - LANE_COUNT_X * WAVE_COUNT_X) % TotalNumberOfLaneInGroup) + (2 + SimdIndex) * TotalNumberOfLaneInGroup, /* out */ FirstElement);
-			ReadElementFromLDS(((GGroupThreadIndex + LANE_COUNT_X * WAVE_COUNT_X) % TotalNumberOfLaneInGroup) + (0 + SimdIndex) * TotalNumberOfLaneInGroup, /* out */ LastElement);
+			ReadVectorFromLDS(((GGroupThreadIndex - LANE_COUNT_X * WAVE_COUNT_X) % TotalNumberOfLaneInGroup) + (LaneStrideX + SimdIndex) * TotalNumberOfLaneInGroup, /* out */ FirstElement);
+			ReadVectorFromLDS(((GGroupThreadIndex + LANE_COUNT_X * WAVE_COUNT_X) % TotalNumberOfLaneInGroup) + (0 + SimdIndex) * TotalNumberOfLaneInGroup, /* out */ LastElement);
 	
-			CP.SetElement(LastSimdIndex + SimdIndex, LastElement);
-			CN.SetElement(FirstSimdIndex + SimdIndex, FirstElement);
+			CP.SetElement(LastElementIndex + SimdIndex, LastElement);
+			CN.SetElement(FirstElementIndex + SimdIndex, FirstElement);
 		}
 	}
 	
diff --git a/Engine/Shaders/Private/TemporalSuperResolution/TSRRejectShading.usf b/Engine/Shaders/Private/TemporalSuperResolution/TSRRejectShading.usf
--- a/Engine/Shaders/Private/TemporalSuperResolution/TSRRejectShading.usf
+++ b/Engine/Shaders/Private/TemporalSuperResolution/TSRRejectShading.usf
@@ -30,39 +30,29 @@
 
 
 //------------------------------------------------------- DERIVED
-
-/** Pixel coordinate stride between each lane in the GROUP_TILE_SIZE. */
-#if DIM_WAVE_SIZE == 32 && PLATFORM_GPU_ARCH == PLATFORM_GPU_ARCH_AMD_RDNA_2 && 0
-	/** Load an entire 16x16 tile in a single wave32 with 4x2 worth of data per lane and just 100% waveops all communications. */
-	#define GROUP_TILE_SIZE 16
-
-	#define WAVE_COUNT_X 1
-	#define WAVE_COUNT_Y 1
-
-	#define LANE_COUNT_X 4
-	#define LANE_COUNT_Y 8
-
-	#define LANE_STRIDE_X 4
-	#define LANE_STRIDE_Y 2
 	
-#elif DIM_WAVE_SIZE == 64 || DIM_WAVE_SIZE == 32
-	/** Layouts wave in row major so an entire row of pixels fits within a single wave, and waveops horizontally, but LDS vertically.
-	 * Each lane worth of a 2x2 of data per lane to minimise to trade VGPR pressure for both number of wave ops and LDS.
+#if DIM_WAVE_SIZE == 64 || DIM_WAVE_SIZE == 32
+	/** Layouts wave in row major 32x32 tile so an entire row of pixels fits within a single wave, and waveops horizontally, but LDS vertically.
+	 * Each lane worth of a 1x4 of data per lane to minimise to trade VGPR pressure for both number of wave ops and LDS.
 	 */
 	#define GROUP_TILE_SIZE 32
 
+	#define LANE_STRIDE_X 1
+	#define LANE_STRIDE_Y 4
+
 	#define WAVE_COUNT_X 1
-	#define WAVE_COUNT_Y 16
+	#define WAVE_COUNT_Y (GROUP_TILE_SIZE / LANE_STRIDE_Y)
 
-	#define LANE_COUNT_X 16
+	#define LANE_COUNT_X (GROUP_TILE_SIZE / LANE_STRIDE_X)
 	#define LANE_COUNT_Y 1
 
-	#define LANE_STRIDE_X 2
-	#define LANE_STRIDE_Y 2
+	#if LANE_STRIDE_Y > 1
+		// Each wave needs to output top and bottom rows of pixels into LDS;
+		#define LDS_SIZE ((2 * LANE_STRIDE_X) * LANE_COUNT_X * WAVE_COUNT_Y)
+	#else
+		#define LDS_SIZE (GROUP_TILE_SIZE * GROUP_TILE_SIZE)
+	#endif
 
-	// Each wave needs to output top and bottom rows of pixels into LDS;
-	#define LDS_SIZE (2 * (LANE_STRIDE_X * LANE_COUNT_X) * (WAVE_COUNT_X * WAVE_COUNT_Y))
-	
 	#ifdef PLATFORM_SPECIFIC_TENSOR_REGISTER_LAYOUT
 		#define TENSOR_REGISTER_LAYOUT PLATFORM_SPECIFIC_TENSOR_REGISTER_LAYOUT
 	#endif
@@ -80,7 +70,7 @@
 	#define LANE_STRIDE_X 1
 	#define LANE_STRIDE_Y 1
 
-	#define LDS_SIZE (WAVE_COUNT_X * WAVE_COUNT_Y)
+	#define LDS_SIZE (GROUP_TILE_SIZE * GROUP_TILE_SIZE)
 	
 	#define TENSOR_REGISTER_LAYOUT FRegisters_S_OneVector
 
@@ -88,6 +78,12 @@
 	#error Unknown wave size
 #endif
 
+#if CONFIG_COMPILE_FP16
+	#define LDS_DWORD_COMPONENT_COUNT 2
+#else
+	#define LDS_DWORD_COMPONENT_COUNT 4
+#endif
+
 /** Size of the SIMD per lane that also number of input pixel loaded into each individual lanes. */
 #define SIMD_SIZE (LANE_STRIDE_X * LANE_STRIDE_Y)
 
diff --git a/Engine/Shaders/Public/LaneVectorization.ush b/Engine/Shaders/Public/LaneVectorization.ush
--- a/Engine/Shaders/Public/LaneVectorization.ush
+++ b/Engine/Shaders/Public/LaneVectorization.ush
@@ -325,6 +325,309 @@ int16_t4 PackRegistersTightly(int16_t4 v)
 #endif // PLATFORM_SUPPORTS_REAL_TYPES
 
 
+//------------------------------------------------------- GROUP SHARED COMMUNICATION
+
+#if defined(LDS_SIZE)
+
+groupshared uint SharedData[LDS_SIZE * LDS_DWORD_COMPONENT_COUNT];
+
+void WriteDwordToLDS(const uint SharedIndex, uint V)
+{
+	SharedData[SharedIndex] = V;
+}
+
+uint ReadDwordFromLDS(const uint SharedIndex)
+{
+	return SharedData[SharedIndex];
+}
+
+void AtomicIncrementLDSDword(const uint SharedIndex, uint Value)
+{
+	return InterlockedAdd(/* inout */ SharedData[SharedIndex], Value);
+}
+
+// uint
+void WriteVectorToLDS(const uint SharedIndex, uint V)
+{
+	WriteDwordToLDS(SharedIndex, V);
+}
+
+void WriteVectorToLDS(const uint SharedIndex, uint2 V)
+{
+	WriteDwordToLDS(SharedIndex + 0 * LDS_SIZE, V.x);
+	WriteDwordToLDS(SharedIndex + 1 * LDS_SIZE, V.y);
+}
+
+void WriteVectorToLDS(const uint SharedIndex, uint3 V)
+{
+	WriteDwordToLDS(SharedIndex + 0 * LDS_SIZE, V.x);
+	WriteDwordToLDS(SharedIndex + 1 * LDS_SIZE, V.y);
+	WriteDwordToLDS(SharedIndex + 2 * LDS_SIZE, V.z);
+}
+
+void WriteVectorToLDS(const uint SharedIndex, uint4 V)
+{
+	WriteDwordToLDS(SharedIndex + 0 * LDS_SIZE, V.x);
+	WriteDwordToLDS(SharedIndex + 1 * LDS_SIZE, V.y);
+	WriteDwordToLDS(SharedIndex + 2 * LDS_SIZE, V.z);
+	WriteDwordToLDS(SharedIndex + 3 * LDS_SIZE, V.w);
+}
+
+void ReadVectorFromLDS(const uint SharedIndex, out uint V)
+{
+	V = ReadDwordFromLDS(SharedIndex);
+}
+
+void ReadVectorFromLDS(const uint SharedIndex, out uint2 V)
+{
+	V.x = ReadDwordFromLDS(SharedIndex + 0 * LDS_SIZE);
+	V.y = ReadDwordFromLDS(SharedIndex + 1 * LDS_SIZE);
+}
+
+void ReadVectorFromLDS(const uint SharedIndex, out uint3 V)
+{
+	V.x = ReadDwordFromLDS(SharedIndex + 0 * LDS_SIZE);
+	V.y = ReadDwordFromLDS(SharedIndex + 1 * LDS_SIZE);
+	V.z = ReadDwordFromLDS(SharedIndex + 2 * LDS_SIZE);
+}
+
+void ReadVectorFromLDS(const uint SharedIndex, out uint4 V)
+{
+	V.x = ReadDwordFromLDS(SharedIndex + 0 * LDS_SIZE);
+	V.y = ReadDwordFromLDS(SharedIndex + 1 * LDS_SIZE);
+	V.z = ReadDwordFromLDS(SharedIndex + 2 * LDS_SIZE);
+	V.w = ReadDwordFromLDS(SharedIndex + 3 * LDS_SIZE);
+}
+
+// float
+void WriteVectorToLDS(const uint SharedIndex, float V)
+{
+	WriteVectorToLDS(SharedIndex, asuint(V));
+}
+
+void WriteVectorToLDS(const uint SharedIndex, float2 V)
+{
+	WriteVectorToLDS(SharedIndex, uint2(asuint(V.x), asuint(V.y)));
+}
+
+void WriteVectorToLDS(const uint SharedIndex, float3 V)
+{
+	WriteVectorToLDS(SharedIndex, uint3(asuint(V.x), asuint(V.y), asuint(V.z)));
+}
+
+void WriteVectorToLDS(const uint SharedIndex, float4 V)
+{
+	WriteVectorToLDS(SharedIndex, uint4(asuint(V.x), asuint(V.y), asuint(V.z), asuint(V.w)));
+}
+
+void ReadVectorFromLDS(const uint SharedIndex, out float V)
+{
+	uint UV;
+	ReadVectorFromLDS(SharedIndex, /* out */ UV);
+	V = asfloat(UV);
+}
+
+void ReadVectorFromLDS(const uint SharedIndex, out float2 V)
+{
+	uint2 UV;
+	ReadVectorFromLDS(SharedIndex, /* out */ UV);
+	V.x = asfloat(UV.x);
+	V.y = asfloat(UV.y);
+}
+
+void ReadVectorFromLDS(const uint SharedIndex, out float3 V)
+{
+	uint3 UV;
+	ReadVectorFromLDS(SharedIndex, /* out */ UV);
+	V.x = asfloat(UV.x);
+	V.y = asfloat(UV.y);
+	V.z = asfloat(UV.z);
+}
+
+void ReadVectorFromLDS(const uint SharedIndex, out float4 V)
+{
+	uint4 UV;
+	ReadVectorFromLDS(SharedIndex, /* out */ UV);
+	V.x = asfloat(UV.x);
+	V.y = asfloat(UV.y);
+	V.z = asfloat(UV.z);
+	V.w = asfloat(UV.w);
+}
+
+#if PLATFORM_SUPPORTS_REAL_TYPES
+
+#if COMPILER_SUPPORT_UINT16_BITCAST
+
+// uint16_t
+void WriteVectorToLDS(const uint SharedIndex, uint16_t V)
+{
+	WriteVectorToLDS(SharedIndex, uint(V));
+}
+
+void WriteVectorToLDS(const uint SharedIndex, uint16_t2 V)
+{
+	WriteVectorToLDS(SharedIndex, bit_cast_uint(V.xy));
+}
+
+void WriteVectorToLDS(const uint SharedIndex, uint16_t3 V)
+{
+	WriteVectorToLDS(SharedIndex, uint2(bit_cast_uint(V.xy), uint(V.z)));
+}
+
+void WriteVectorToLDS(const uint SharedIndex, uint16_t4 V)
+{
+	WriteVectorToLDS(SharedIndex, uint2(bit_cast_uint(V.xy), bit_cast_uint(V.zw)));
+}
+
+void ReadVectorFromLDS(const uint SharedIndex, out uint16_t V)
+{
+	uint UV;
+	ReadVectorFromLDS(SharedIndex, /* out */ UV);
+	V = tsr_ushort(UV);
+}
+
+void ReadVectorFromLDS(const uint SharedIndex, out uint16_t2 V)
+{
+	uint UV;
+	ReadVectorFromLDS(SharedIndex, /* out */ UV);
+	V.xy = bit_cast_uint16_t2(UV);
+}
+
+void ReadVectorFromLDS(const uint SharedIndex, out uint16_t3 V)
+{
+	uint2 UV;
+	ReadVectorFromLDS(SharedIndex, /* out */ UV);
+	V.xy = bit_cast_uint16_t2(UV.x);
+	V.z = tsr_ushort(UV.y);
+}
+
+void ReadVectorFromLDS(const uint SharedIndex, out uint16_t4 V)
+{
+	uint2 UV;
+	ReadVectorFromLDS(SharedIndex, /* out */ UV);
+	V.xy = bit_cast_uint16_t2(UV.x);
+	V.zw = bit_cast_uint16_t2(UV.y);
+}
+
+#else // !COMPILER_SUPPORT_UINT16_BITCAST
+
+// uint16_t
+void WriteVectorToLDS(const uint SharedIndex, uint16_t V)
+{
+	WriteVectorToLDS(SharedIndex, uint(V));
+}
+
+void WriteVectorToLDS(const uint SharedIndex, uint16_t2 V)
+{
+	WriteVectorToLDS(SharedIndex, uint(V.x) | (uint(V.y) << 16u));
+}
+
+void WriteVectorToLDS(const uint SharedIndex, uint16_t3 V)
+{
+	WriteVectorToLDS(SharedIndex, uint2(uint(V.x) | (uint(V.y) << 16u), uint(V.z)));
+}
+
+void WriteVectorToLDS(const uint SharedIndex, uint16_t4 V)
+{
+	WriteVectorToLDS(SharedIndex, uint2(uint(V.x) | (uint(V.y) << 16u), uint(V.z) | (uint(V.w) << 16u)));
+}
+
+void ReadVectorFromLDS(const uint SharedIndex, out uint16_t V)
+{
+	uint UV;
+	ReadVectorFromLDS(SharedIndex, /* out */ UV);
+	V = tsr_ushort(UV);
+}
+
+void ReadVectorFromLDS(const uint SharedIndex, out uint16_t2 V)
+{
+	uint UV;
+	ReadVectorFromLDS(SharedIndex, /* out */ UV);
+	V.x = tsr_ushort((UV >>  0u) & 0xFFFFu);
+	V.y = tsr_ushort((UV >> 16u) & 0xFFFFu);
+}
+
+void ReadVectorFromLDS(const uint SharedIndex, out uint16_t3 V)
+{
+	uint2 UV;
+	ReadVectorFromLDS(SharedIndex, /* out */ UV);
+	V.x = tsr_ushort((UV.x >>  0u) & 0xFFFFu);
+	V.y = tsr_ushort((UV.x >> 16u) & 0xFFFFu);
+	V.z = tsr_ushort((UV.y >>  0u) & 0xFFFFu);
+}
+
+void ReadVectorFromLDS(const uint SharedIndex, out uint16_t4 V)
+{
+	uint2 UV;
+	ReadVectorFromLDS(SharedIndex, /* out */ UV);
+	V.x = tsr_ushort((UV.x >>  0u) & 0xFFFFu);
+	V.y = tsr_ushort((UV.x >> 16u) & 0xFFFFu);
+	V.z = tsr_ushort((UV.y >>  0u) & 0xFFFFu);
+	V.w = tsr_ushort((UV.y >> 16u) & 0xFFFFu);
+}
+
+#endif // !COMPILER_SUPPORT_UINT16_BITCAST
+
+// half
+void WriteVectorToLDS(const uint SharedIndex, half V)
+{
+	WriteVectorToLDS(SharedIndex, asuint16(V));
+}
+
+void WriteVectorToLDS(const uint SharedIndex, half2 V)
+{
+	WriteVectorToLDS(SharedIndex, uint16_t2(asuint16(V.x), asuint16(V.y)));
+}
+
+void WriteVectorToLDS(const uint SharedIndex, half3 V)
+{
+	WriteVectorToLDS(SharedIndex, uint16_t3(asuint16(V.x), asuint16(V.y), asuint16(V.z)));
+}
+
+void WriteVectorToLDS(const uint SharedIndex, half4 V)
+{
+	WriteVectorToLDS(SharedIndex, uint16_t4(asuint16(V.x), asuint16(V.y), asuint16(V.z), asuint16(V.w)));
+}
+
+void ReadVectorFromLDS(const uint SharedIndex, out half V)
+{
+	uint16_t UV;
+	ReadVectorFromLDS(SharedIndex, /* out */ UV);
+	V = asfloat16(UV);
+}
+
+void ReadVectorFromLDS(const uint SharedIndex, out half2 V)
+{
+	uint16_t2 UV;
+	ReadVectorFromLDS(SharedIndex, /* out */ UV);
+	V.x = asfloat16(UV.x);
+	V.y = asfloat16(UV.y);
+}
+
+void ReadVectorFromLDS(const uint SharedIndex, out half3 V)
+{
+	uint16_t3 UV;
+	ReadVectorFromLDS(SharedIndex, /* out */ UV);
+	V.x = asfloat16(UV.x);
+	V.y = asfloat16(UV.y);
+	V.z = asfloat16(UV.z);
+}
+
+void ReadVectorFromLDS(const uint SharedIndex, out half4 V)
+{
+	uint16_t4 UV;
+	ReadVectorFromLDS(SharedIndex, /* out */ UV);
+	V.x = asfloat16(UV.x);
+	V.y = asfloat16(UV.y);
+	V.z = asfloat16(UV.z);
+	V.w = asfloat16(UV.w);
+}
+
+#endif // PLATFORM_SUPPORTS_REAL_TYPES
+
+#endif // defined(LDS_SIZE)
+
+
 //------------------------------------------------------- REGISTERS LAYOUTS
 
 /** Stores one unique vector<>. InElementCount must be == 1  */
@@ -611,7 +914,7 @@ struct FRegisters_SoA_Matrix
 	}
 }; // struct FRegisters_SoA_Matrix
 
-/** Stores an array of vector<> as a structure of array in registers using a matrix<>. Requires InElementCount a pair >= 2 */
+/** Stores an array of vector<> as a structure of array in registers using a array of vector<X, 2>. Requires InElementCount a pair >= 2 */
 template<typename ScalarType, uint InVectorSize, uint InElementCount>
 struct FRegisters_AoS_PairArray
 {
@@ -718,7 +1021,7 @@ struct FRegisters_AoS_PairArray
 	}
 }; // struct FRegisters_AoS_PairArray
 
-/** Stores an array of vector<> as a structure of array in registers using a matrix<>. Requires InElementCount a pair >= 2 */
+/** Stores an array of vector<> as a structure of array in registers using a array of vector<X, 2>. Requires InElementCount a pair >= 2 */
 template<typename ScalarType, uint InVectorSize, uint InElementCount>
 struct FRegisters_SoA_PairArray
 {
