Fixed RayTracingDebugHitStatsCHS.usf is using the wrong version

diff --git a/Engine/Shaders/Private/RayTracing/RayTracingDebugHitStatsCHS.usf b/Engine/Shaders/Private/RayTracing/RayTracingDebugHitStatsCHS.usf
--- a/Engine/Shaders/Private/RayTracing/RayTracingDebugHitStatsCHS.usf
+++ b/Engine/Shaders/Private/RayTracing/RayTracingDebugHitStatsCHS.usf
@@ -1,21 +1,20 @@
 // Copyright Epic Games, Inc. All Rights Reserved.
 
+#include "/Engine/Shared/RayTracingDebugTypes.h"
 #include "../Common.ush"
 #include "../Hash.ush"
 
 #include "RayTracingCommon.ush"
 #include "RayTracingDebugUtils.ush"
 #include "RayTracingHitGroupCommon.ush"
-#include "/Engine/Shared/RayTracingDebugTypes.h"
 
 RAY_TRACING_ENTRY_ANY_HIT(RayTracingDebugHitStatsAHS,
 	FRayTracingDebugPayload, Payload,
 	FRayTracingIntersectionAttributes, Attributes)
 {
 	uint Index = GetInstanceUserData();
-	uint OldValue;
-	RayTracingDebugHitStatsUniformBuffer.HitStatsOutput[Index].x = Index;
-	InterlockedAdd(RayTracingDebugHitStatsUniformBuffer.HitStatsOutput[Index].y, 1, OldValue);
+	RayTracingDebugHitStatsUniformBuffer.HitStatsOutput[Index].PrimitiveID = Index;
+	InterlockedAdd(RayTracingDebugHitStatsUniformBuffer.HitStatsOutput[Index].Count, 1);
 }
 
 RAY_TRACING_ENTRY_CLOSEST_HIT(RayTracingDebugHitStatsCHS,
