Substrate - improvements to the Chan diffuse model.

We clamp the BRDF lobe value to an arbitrary value of 1 to get some practical benefits at high roughness:
 - This is to avoid too bright edges when using normal map on a mesh and the local bases, L, N and V ends up in an top emisphere setup.
 - This maintains the full proper rough look of a sphere when not using normal maps.
 - This also fixes the furnace test returning too much energy at the edge of a mesh.

diff --git a/Engine/Shaders/Private/BRDF.ush b/Engine/Shaders/Private/BRDF.ush
--- a/Engine/Shaders/Private/BRDF.ush
+++ b/Engine/Shaders/Private/BRDF.ush
@@ -354,8 +354,16 @@ float3 Diffuse_Chan( float3 DiffuseColor, float a2, float NoV, float NoL, float
 	float Fb = ( (34.5 * g - 59 ) * g + 24.5 ) * VoH * exp2( -max( 73.2 * g - 21.2, 8.9 ) * sqrtFast( NoH ) );
 	// It fades out when lights become area lights in order to avoid visual artefacts.
 	Fb *= RetroReflectivityWeight;
-	
-	return DiffuseColor * ( (1 / PI) * ( Fd + Fb ) );
+
+	float Lobe = (1 / PI) * (Fd + Fb);
+
+	// We clamp the BRDF lobe value to an arbitrary value of 1 to get some practical benefits at high roughness:
+	// - This is to avoid too bright edges when using normal map on a mesh and the local bases, L, N and V ends up in an top emisphere setup.
+	// - This maintains the full proper rough look of a sphere when not using normal maps.
+	// - This also fixes the furnace test returning too much energy at the edge of a mesh.
+	Lobe = min(1.0, Lobe);
+
+	return DiffuseColor * Lobe;
 }
 
 // [Blinn 1977, "Models of light reflection for computer synthesized pictures"]
