LFV - added application on volumetric cloud.

#prelfight shader

diff --git a/Engine/Shaders/Private/VolumetricCloud.usf b/Engine/Shaders/Private/VolumetricCloud.usf
--- a/Engine/Shaders/Private/VolumetricCloud.usf
+++ b/Engine/Shaders/Private/VolumetricCloud.usf
@@ -64,6 +64,10 @@
 #define CLOUD_MIN_AND_MAX_DEPTH 0
 #endif
 
+#ifndef PERMUTATION_SUPPORT_LOCAL_FOG_VOLUME
+#define PERMUTATION_SUPPORT_LOCAL_FOG_VOLUME 0
+#endif
+
 // Draw debug sample positions
 #if CLOUD_DEBUG_SAMPLES || EMPTY_SPACE_SKIPPING_DEBUG
 #include "ShaderPrint.ush"
@@ -81,6 +85,10 @@
 #include "ReflectionEnvironmentShared.ush"
 #include "ParticipatingMediaCommon.ush"
 
+#if PERMUTATION_SUPPORT_LOCAL_FOG_VOLUME
+#define LFVStruct RenderVolumetricCloudParameters
+#include "LocalFogVolumes/LocalFogVolumeCommon.ush"
+#endif
 
 #if defined(SHADER_SHADOW_PS) || SHADER_RENDERVIEW_PS || SHADER_RENDERVIEW_CS || SHADER_EMPTY_SPACE_SKIPPING_CS
 
@@ -1451,7 +1459,16 @@ void MainCommon(in FMaterialPixelParameters MaterialParameters, in float4 SvPosi
 	if (RenderVolumetricCloudParameters.EnableHeightFog && tAPWeightsSum > 0.0f)
 	{
 		float4 HeightFogInscatteringAndTransmittance = CalculateHeightFog(FogSampleWorldPositionRelativeToCameraCm);
-		
+
+	#if PERMUTATION_SUPPORT_LOCAL_FOG_VOLUME
+		uint2 TilePos = uint2(SvPosition.xy * RenderVolumetricCloudParameters.TracingCoordToFullResPixelCoordScaleBias.xy - View.ViewRectMin.xy) / LFVTilePixelSize.xx;
+		float4 LFVContribution = GetLFVContribution(PrimaryView, TilePos, FogSampleWorldPositionRelativeToCameraCm);
+		if (LFVRenderInVolumetricFog > 0)
+		{
+			HeightFogInscatteringAndTransmittance = float4(LFVContribution.rgb + HeightFogInscatteringAndTransmittance.rgb * LFVContribution.a, LFVContribution.a * HeightFogInscatteringAndTransmittance.a);
+		}
+	#endif
+
 		if (FogStruct.ApplyVolumetricFog > 0)
 		{
 			float3 VolumeUV = ComputeVolumeUVFromNDC(ClipPos);
@@ -1459,6 +1476,13 @@ void MainCommon(in FMaterialPixelParameters MaterialParameters, in float4 SvPosi
 			HeightFogInscatteringAndTransmittance = CombineVolumetricFog(HeightFogInscatteringAndTransmittance, VolumeUV, EyeIndex, length(FogSampleWorldPositionRelativeToCameraCm));
 		}
 
+	#if PERMUTATION_SUPPORT_LOCAL_FOG_VOLUME
+		if (LFVRenderInVolumetricFog == 0)
+		{
+			HeightFogInscatteringAndTransmittance = float4(LFVContribution.rgb + HeightFogInscatteringAndTransmittance.rgb * LFVContribution.a, LFVContribution.a * HeightFogInscatteringAndTransmittance.a);
+		}
+	#endif
+
 		// Apply aerial perspective OVER the cloud, assiming coverage is from the full transmittance.
 		// skip if MeanTransmittance is 1
 		float MeanCoverage = 1.0 - MeanTransmittance;
@@ -1508,6 +1532,15 @@ void MainCommon(in FMaterialPixelParameters MaterialParameters, in float4 SvPosi
 		{
 			float4 HeightFogInscatteringAndTransmittance = CalculateHeightFog(FogSampleWorldPositionRelativeToCameraCm);
 
+		#if PERMUTATION_SUPPORT_LOCAL_FOG_VOLUME
+			uint2 TilePos = uint2(SvPosition.xy * RenderVolumetricCloudParameters.TracingCoordToFullResPixelCoordScaleBias.xy - View.ViewRectMin.xy) / LFVTilePixelSize.xx;
+			float4 LFVContribution = GetLFVContribution(PrimaryView, TilePos, FogSampleWorldPositionRelativeToCameraCm);
+			if (LFVRenderInVolumetricFog > 0)
+			{
+				HeightFogInscatteringAndTransmittance = float4(LFVContribution.rgb + HeightFogInscatteringAndTransmittance.rgb * LFVContribution.a, LFVContribution.a * HeightFogInscatteringAndTransmittance.a);
+			}
+		#endif
+
 			if (FogStruct.ApplyVolumetricFog > 0)
 			{
 				float3 VolumeUV = ComputeVolumeUVFromNDC(ClipPos);
@@ -1515,6 +1548,13 @@ void MainCommon(in FMaterialPixelParameters MaterialParameters, in float4 SvPosi
 				HeightFogInscatteringAndTransmittance = CombineVolumetricFog(HeightFogInscatteringAndTransmittance, VolumeUV, EyeIndex, length(FogSampleWorldPositionRelativeToCameraCm));
 			}
 
+		#if PERMUTATION_SUPPORT_LOCAL_FOG_VOLUME
+			if (LFVRenderInVolumetricFog == 0)
+			{
+				HeightFogInscatteringAndTransmittance = float4(LFVContribution.rgb + HeightFogInscatteringAndTransmittance.rgb * LFVContribution.a, LFVContribution.a * HeightFogInscatteringAndTransmittance.a);
+			}
+		#endif
+
 			// Apply aerial perspective OVER the cloud, assiming coverage is from the full transmittance.
 			// skip if MeanTransmittance is 1
 			float MeanCoverage = 1.0 - ClosestMeanTransmittance;
diff --git a/Engine/Source/Runtime/Renderer/Private/VolumetricCloudRendering.cpp b/Engine/Source/Runtime/Renderer/Private/VolumetricCloudRendering.cpp
--- a/Engine/Source/Runtime/Renderer/Private/VolumetricCloudRendering.cpp
+++ b/Engine/Source/Runtime/Renderer/Private/VolumetricCloudRendering.cpp
@@ -583,6 +583,7 @@ BEGIN_GLOBAL_SHADER_PARAMETER_STRUCT(FRenderVolumetricCloudGlobalParameters, )
 	SHADER_PARAMETER_RDG_TEXTURE(Texture2D<float3>, CloudShadowTexture1)
 	SHADER_PARAMETER_SAMPLER(SamplerState, CloudBilinearTextureSampler)
 	SHADER_PARAMETER_STRUCT_INCLUDE(FVolumeShadowingShaderParametersGlobal0, Light0Shadow)
+	SHADER_PARAMETER_STRUCT(FLocalFogVolumeUniformParameters, LFV)
 	SHADER_PARAMETER(int32, VirtualShadowMapId0)
 	SHADER_PARAMETER(FUintVector4, TracingCoordToZbufferCoordScaleBias)
 	SHADER_PARAMETER(FUintVector4, TracingCoordToFullResPixelCoordScaleBias)
@@ -806,7 +807,8 @@ class FRenderVolumetricCloudRenderViewCS : public FMeshMaterialShader
 	class FCloudSampleSecondLight : SHADER_PERMUTATION_BOOL("CLOUD_SAMPLE_SECOND_LIGHT");
 	class FCloudSampleLocalLights : SHADER_PERMUTATION_BOOL("CLOUD_SAMPLE_LOCAL_LIGHTS");
 	class FCloudMinAndMaxDepth : SHADER_PERMUTATION_BOOL("CLOUD_MIN_AND_MAX_DEPTH");
-	class FCloudDebugViewMode : SHADER_PERMUTATION_BOOL("CLOUD_DEBUG_VIEW_MODE");	// Only the CS feature this because CS is typically used in editor.
+	class FCloudLocalFogVolume : SHADER_PERMUTATION_BOOL("PERMUTATION_SUPPORT_LOCAL_FOG_VOLUME");	// Only for the CS since platform not supporting the CS to run async likely won't render cloud
+	class FCloudDebugViewMode : SHADER_PERMUTATION_BOOL("CLOUD_DEBUG_VIEW_MODE");					// Only the CS feature this because CS is typically used in editor.
 
 	using FPermutationDomain = TShaderPermutationDomain<
 		FCloudPerSampleAtmosphereTransmittance,
@@ -814,6 +816,7 @@ class FRenderVolumetricCloudRenderViewCS : public FMeshMaterialShader
 		FCloudSampleSecondLight,
 		FCloudSampleLocalLights,
 		FCloudMinAndMaxDepth,
+		FCloudLocalFogVolume,
 		FCloudDebugViewMode>;
 
 	BEGIN_SHADER_PARAMETER_STRUCT(FParameters, )
@@ -2164,6 +2167,8 @@ static TRDGUniformBufferRef<FRenderVolumetricCloudGlobalParameters> CreateCloudP
 	VolumetricCloudParams.AerialPerspectiveMieOnlyStartDistanceKm = CloudInfo.GetVolumetricCloudSceneProxy().AerialPespectiveMieScatteringStartDistance;
 	VolumetricCloudParams.AerialPerspectiveMieOnlyFadeDistanceKmInv = 1.0f / FMath::Max(CloudInfo.GetVolumetricCloudSceneProxy().AerialPespectiveMieScatteringFadeDistance, 0.00001f);// One centimeter minimum
 
+	VolumetricCloudParams.LFV = MainView.LocalFogVolumeViewData.UniformParametersStruct;
+
 	const FRDGTextureDesc& Desc = CloudRC.RenderTargets.Output[0].GetTexture()->Desc;
 	VolumetricCloudParams.OutputSizeInvSize = FVector4f(Desc.Extent.X, Desc.Extent.Y, 1.0f / Desc.Extent.X, 1.0f / Desc.Extent.Y);
 
@@ -2414,6 +2419,7 @@ void FSceneRenderer::RenderVolumetricCloudsInternal(FRDGBuilder& GraphBuilder, F
 		PermutationVector.Set<typename FRenderVolumetricCloudRenderViewCS::FCloudSampleSecondLight>(!bCloudDebugViewModeEnabled && bSecondAtmosphereLightEnabled);
 		PermutationVector.Set<typename FRenderVolumetricCloudRenderViewCS::FCloudSampleLocalLights>(!bCloudDebugViewModeEnabled && bCloudEnableLocalLightSampling && !CloudRC.bIsSkyRealTimeReflectionRendering);
 		PermutationVector.Set<typename FRenderVolumetricCloudRenderViewCS::FCloudMinAndMaxDepth>(ShouldVolumetricCloudTraceWithMinMaxDepth(MainView) && CloudRC.SecondaryCloudTracingDataTexture!=nullptr && !CloudRC.bIsReflectionRendering);// Only for mode 0 in non reflection
+		PermutationVector.Set<typename FRenderVolumetricCloudRenderViewCS::FCloudLocalFogVolume>(MainView.LocalFogVolumeViewData.GPUInstanceCount > 0);
 		PermutationVector.Set<typename FRenderVolumetricCloudRenderViewCS::FCloudDebugViewMode>(bCloudDebugViewModeEnabled);
 
 		TShaderRef<FRenderVolumetricCloudRenderViewCS> ComputeShader = MaterialResource->GetShader<FRenderVolumetricCloudRenderViewCS>(&FLocalVertexFactory::StaticType, PermutationVector, false);
