GPU lights - fixed FLightSceneData data upload and added conversion utilities.

diff --git a/Engine/Shaders/Private/LightData.ush b/Engine/Shaders/Private/LightData.ush
--- a/Engine/Shaders/Private/LightData.ush
+++ b/Engine/Shaders/Private/LightData.ush
@@ -439,6 +439,73 @@ FDeferredLightData ConvertToDeferredLight(
 	return Out;
 }
 
+// Helper function for converting FLightSceneData into FDeferredLightData
+FDeferredLightData ConvertToDeferredLight(FLightSceneData In)
+{
+	FDeferredLightData Out = (FDeferredLightData)0;
+
+	const uint LightType = UnpackLightType(In.LightTypeAndShadowMapChannelMaskPacked);
+
+	Out.TranslatedWorldPosition = LWCToFloat(LWCAdd(In.WorldPosition, PrimaryView.PreViewTranslation));
+	Out.InvRadius				= In.InvRadius;
+	Out.Color					= In.Color.xyz;
+	Out.FalloffExponent			= In.FalloffExponent;
+	Out.Direction				= In.Direction;
+	Out.SpotAngles				= In.SpotAngles;
+	Out.SourceRadius			= In.SourceRadius;
+	Out.SourceLength			= In.SourceLength;
+	Out.Tangent					= In.Tangent;
+	Out.SoftSourceRadius		= In.SoftSourceRadius;
+	Out.bInverseSquared			= In.FalloffExponent == 0;
+	Out.SpecularScale			= In.SpecularScale;
+	Out.bRadialLight			= true;
+	Out.bSpotLight				= LightType == LIGHT_TYPE_SPOT;
+	Out.bRectLight				= LightType == LIGHT_TYPE_RECT;
+	Out.HairTransmittance		= InitHairTransmittanceData();
+
+	Out.RectLightData.BarnCosAngle = In.RectLightBarnCosAngle;
+	Out.RectLightData.BarnLength = In.RectLightBarnLength;
+	Out.RectLightData.AtlasData.AtlasMaxLevel = In.RectLightAtlasMaxLevel;
+	Out.RectLightData.AtlasData.AtlasUVOffset = In.RectLightAtlasUVOffset;
+	Out.RectLightData.AtlasData.AtlasUVScale = In.RectLightAtlasUVScale;
+
+	Out.IESAtlasIndex = 0;
+
+	// Shadow
+	// * Static shadowing uses ShadowMapChannel, 
+	// * Dynamic shadows are packed into light attenuation using PreviewShadowMapChannel
+	// Data Layout on [10:0] with:
+	// * ShadowMapChannelMask	= [3:0]
+	// * PreviewShadowMapChannel= [7:4] 
+	// * LightincChannelMask    = [8:10] 
+	Out.ShadowedBits				= (In.LightTypeAndShadowMapChannelMaskPacked & 0xFF) != 0 ? 1 : 0;
+	Out.ShadowMapChannelMask		= UnpackShadowMapChannelMask(In.LightTypeAndShadowMapChannelMaskPacked);
+	return Out;
+}
+
+// Helper function for converting FLightSceneData into FLightShaderParameters
+FLightShaderParameters ConvertToLightShaderParameters(FLightSceneData In)
+{
+	FLightShaderParameters Out;
+	Out.TranslatedWorldPosition	= LWCToFloat(LWCAdd(In.WorldPosition, PrimaryView.PreViewTranslation));
+	Out.InvRadius				= In.InvRadius;
+	Out.Color					= In.Color.xyz;
+	Out.FalloffExponent			= In.FalloffExponent;
+	Out.Direction				= In.Direction;
+	Out.Tangent					= In.Tangent;
+	Out.SpotAngles				= In.SpotAngles;
+	Out.SpecularScale			= 1.0f;
+	Out.SourceRadius			= In.SourceRadius;
+	Out.SoftSourceRadius		= In.SoftSourceRadius;
+	Out.SourceLength			= In.SourceLength;
+	Out.RectLightBarnCosAngle	= In.RectLightBarnCosAngle;
+	Out.RectLightBarnLength		= In.RectLightBarnLength;
+	Out.RectLightAtlasUVOffset	= In.RectLightAtlasUVOffset;
+	Out.RectLightAtlasUVScale	= In.RectLightAtlasUVScale;
+	Out.RectLightAtlasMaxLevel	= In.RectLightAtlasMaxLevel;
+	return Out;
+}
+
 // Helper function for converting FLocalLightData into FSimpleDeferredLightData
 FSimpleDeferredLightData ConvertToSimpleLight(FLocalLightData In)
 {
diff --git a/Engine/Shaders/Shared/LightSceneData.h b/Engine/Shaders/Shared/LightSceneData.h
--- a/Engine/Shaders/Shared/LightSceneData.h
+++ b/Engine/Shaders/Shared/LightSceneData.h
@@ -24,12 +24,12 @@ struct FLightSceneData
 	// 1 / light's falloff radius from Position.
 	float InvRadius;
 
-	// Color of the light.
-	float4 Color;
-
 	// The exponent for the falloff of the light intensity from the distance.
 	float FalloffExponent;
 
+	// Color of the light.
+	float4 Color;
+
 	// Direction of the light if applies.
 	float3 Direction;
 
@@ -67,6 +67,8 @@ struct FLightSceneData
 
 	// Extra fields
 	uint LightTypeAndShadowMapChannelMaskPacked;
+
+	float3 Padding;
 };
 
 #ifdef __cplusplus
