Don't scale TwoSidedMeshDistanceBiasScale with mip level.
- there was noticeable shadow popping when new distance field mips were streamed in.

diff --git a/Engine/Shaders/Private/DistanceFieldShadowingShared.ush b/Engine/Shaders/Private/DistanceFieldShadowingShared.ush
--- a/Engine/Shaders/Private/DistanceFieldShadowingShared.ush
+++ b/Engine/Shaders/Private/DistanceFieldShadowingShared.ush
@@ -132,7 +132,7 @@ float ShadowRayTraceThroughCulledObjects(
 				FDFAssetData DFAssetData = LoadDFAssetDataHighestResolution(DFObjectData.AssetIndex);
 
 				const uint NumMips = LoadDFAssetData(DFObjectData.AssetIndex, 0).NumMips;
-				const float ExpandSurfaceDistance = DFObjectData.VolumeSurfaceBias * exp2(DF_MAX_NUM_MIPS - NumMips) * (DFObjectData.bMostlyTwoSided ? TwoSidedMeshDistanceBiasScale : 1);
+				const float ExpandSurfaceDistance = DFObjectData.VolumeSurfaceBias * (exp2(DF_MAX_NUM_MIPS - NumMips) + (DFObjectData.bMostlyTwoSided ? TwoSidedMeshDistanceBiasScale - 1 : 0));
 				const float ExpandSurfaceFalloff = 2.0f * ExpandSurfaceDistance;
 				
 				float MaxEncodedDistance = DFAssetData.DistanceFieldToVolumeScaleBias.x + DFAssetData.DistanceFieldToVolumeScaleBias.y;
