VulkanRHI: Vulkan fixes for world normals in path tracing and nanite double-sided materials.

diff --git a/Engine/Shaders/Private/MaterialTemplate.ush b/Engine/Shaders/Private/MaterialTemplate.ush
--- a/Engine/Shaders/Private/MaterialTemplate.ush
+++ b/Engine/Shaders/Private/MaterialTemplate.ush
@@ -3542,13 +3542,14 @@ uint DiscardMaterialWithPixelCoverage(FMaterialPixelParameters MaterialParameter
 #endif // MATERIALBLENDING_MASKED_USING_COVERAGE
 
 
-	#define FrontFaceSemantic SV_IsFrontFace
-	#define FIsFrontFace bool
-	half GetFloatFacingSign(FIsFrontFace bIsFrontFace)
-	{
-#if COMPILER_DXC && (COMPILER_VULKAN || COMPILER_GLSL_ES3_1)
+#define FrontFaceSemantic SV_IsFrontFace
+#define FIsFrontFace bool
+half GetFloatFacingSign(FIsFrontFace bIsFrontFace)
+{
+#if COMPILER_DXC && (COMPILER_VULKAN || COMPILER_GLSL_ES3_1) && (!RAYHITGROUPSHADER) && (!IS_NANITE_PASS)
 		// We need to flip SV_IsFrontFace for Vulkan when compiling with DXC due to different coordinate systems.
 		// HLSLcc did that by flipping SV_IsFrontFace in the high-level GLSL output.
+		// NOTE: Raytracing and Nanite don't need an SV_IsFrontFace workaround.
 		return bIsFrontFace ? -1 : +1;
 #else
 		return bIsFrontFace ? +1 : -1;
