Fix missing GetUV function in LandscapeVF causing shader compilation error with PT/RT.

diff --git a/Engine/Shaders/Private/LandscapeVertexFactory.ush b/Engine/Shaders/Private/LandscapeVertexFactory.ush
--- a/Engine/Shaders/Private/LandscapeVertexFactory.ush
+++ b/Engine/Shaders/Private/LandscapeVertexFactory.ush
@@ -253,6 +253,20 @@ struct FVertexFactoryInterpolantsVSToPS
 #endif
 };
 
+#if NUM_TEX_COORD_INTERPOLATORS
+float2 GetUV(FVertexFactoryInterpolantsVSToPS Interpolants, int UVIndex)
+{
+	float2 Out = 0;
+	if (UVIndex < NUM_TEX_COORD_INTERPOLATORS)
+	{
+		float2 CustomizedUVs[NUM_TEX_COORD_INTERPOLATORS];
+		GetCustomInterpolators(VertexParameters, CustomizedUVs);
+		Out = CustomizedUVs[UVIndex];
+	}
+	return Out;
+}
+#endif
+
 struct FLandscapeTexCoords
 {
 	float2	LayerTexCoord : TEXCOORD0; // xy == texcoord
