Fix shader compilation with landscape vertex factor.

diff --git a/Engine/Shaders/Private/LandscapeVertexFactory.ush b/Engine/Shaders/Private/LandscapeVertexFactory.ush
--- a/Engine/Shaders/Private/LandscapeVertexFactory.ush
+++ b/Engine/Shaders/Private/LandscapeVertexFactory.ush
@@ -256,14 +256,7 @@ struct FVertexFactoryInterpolantsVSToPS
 #if NUM_TEX_COORD_INTERPOLATORS
 float2 GetUV(FVertexFactoryInterpolantsVSToPS Interpolants, int UVIndex)
 {
-	float2 Out = 0;
-	if (UVIndex < NUM_TEX_COORD_INTERPOLATORS)
-	{
-		float2 CustomizedUVs[NUM_TEX_COORD_INTERPOLATORS];
-		GetCustomInterpolators(VertexParameters, CustomizedUVs);
-		Out = CustomizedUVs[UVIndex];
-	}
-	return Out;
+	return 0;
 }
 #endif
 
