Add r.HDR.Aces.GamutCompression for HDR output so that we get a similar behaviour to the BlueCorrection parameter
See https://docs.acescentral.com/specifications/rgc/ for the use cases this fixes.
It's disabled by default to keep the existing behaviour

#tests PIE

diff --git a/Engine/Shaders/Private/CompositeUIPixelShader.usf b/Engine/Shaders/Private/CompositeUIPixelShader.usf
--- a/Engine/Shaders/Private/CompositeUIPixelShader.usf
+++ b/Engine/Shaders/Private/CompositeUIPixelShader.usf
@@ -45,7 +45,7 @@ void GenerateLUTPS(
 
 	float3 LinearColor = ST2084ToLinear(Neutral);
 
-	FACESTonemapParams LUTTsParams = ComputeACESTonemapParams(ACESMinMaxData, ACESMidData, ACESCoefsLow_0, ACESCoefsHigh_0, ACESCoefsLow_4, ACESCoefsHigh_4, 1.0);
+	FACESTonemapParams LUTTsParams = ComputeACESTonemapParams(ACESMinMaxData, ACESMidData, ACESCoefsLow_0, ACESCoefsHigh_0, ACESCoefsLow_4, ACESCoefsHigh_4, 1.0, 1.0);
 
 	if( OutputDevice == TONEMAPPER_OUTPUT_ACES1000nitST2084 || OutputDevice == TONEMAPPER_OUTPUT_ACES1000nitScRGB)
 	{		
diff --git a/Engine/Shaders/Private/PostProcessCombineLUTs.usf b/Engine/Shaders/Private/PostProcessCombineLUTs.usf
--- a/Engine/Shaders/Private/PostProcessCombineLUTs.usf
+++ b/Engine/Shaders/Private/PostProcessCombineLUTs.usf
@@ -74,6 +74,7 @@ float4 ACESCoefsHigh_0;
 float ACESCoefsLow_4;
 float ACESCoefsHigh_4;
 float ACESSceneColorMultiplier;
+float ACESGamutCompression;
 
 float3 ColorCorrect( float3 WorkingColor,
 	float4 ColorSaturation,
@@ -348,7 +349,7 @@ float4 CombineLUTsCommon(float2 InUV, uint InLayerIndex)
 	else if( GetOutputDevice() == TONEMAPPER_OUTPUT_ACES1000nitST2084 || GetOutputDevice() == TONEMAPPER_OUTPUT_ACES1000nitScRGB)
 	{		
 		// 1000 nit ODT
-		FACESTonemapParams AcesParams = ComputeACESTonemapParams(ACESMinMaxData, ACESMidData, ACESCoefsLow_0, ACESCoefsHigh_0, ACESCoefsLow_4, ACESCoefsHigh_4, ACESSceneColorMultiplier);
+		FACESTonemapParams AcesParams = ComputeACESTonemapParams(ACESMinMaxData, ACESMidData, ACESCoefsLow_0, ACESCoefsHigh_0, ACESCoefsLow_4, ACESCoefsHigh_4, ACESSceneColorMultiplier, ACESGamutCompression);
 		float3 ODTColor = ACESOutputTransforms1000( GradedColor, (float3x3)WorkingColorSpace.ToAP0, AcesParams);
 
 		// Convert from AP1 to specified output gamut except for ScRGB as it might contain negative values. 
@@ -366,7 +367,7 @@ float4 CombineLUTsCommon(float2 InUV, uint InLayerIndex)
 	else if( GetOutputDevice() == TONEMAPPER_OUTPUT_ACES2000nitST2084 || GetOutputDevice() == TONEMAPPER_OUTPUT_ACES2000nitScRGB)
 	{		
 		// 2000 nit ODT
-		FACESTonemapParams AcesParams = ComputeACESTonemapParams(ACESMinMaxData, ACESMidData, ACESCoefsLow_0, ACESCoefsHigh_0, ACESCoefsLow_4, ACESCoefsHigh_4, ACESSceneColorMultiplier);
+		FACESTonemapParams AcesParams = ComputeACESTonemapParams(ACESMinMaxData, ACESMidData, ACESCoefsLow_0, ACESCoefsHigh_0, ACESCoefsLow_4, ACESCoefsHigh_4, ACESSceneColorMultiplier, ACESGamutCompression);
 		float3 ODTColor = ACESOutputTransforms2000( GradedColor, (float3x3)WorkingColorSpace.ToAP0, AcesParams);
 
 		// Convert from AP1 to specified output gamut except for ScRGB as it might contain negative values. 
diff --git a/Engine/Shaders/Private/PostProcessDeviceEncodingOnly.usf b/Engine/Shaders/Private/PostProcessDeviceEncodingOnly.usf
--- a/Engine/Shaders/Private/PostProcessDeviceEncodingOnly.usf
+++ b/Engine/Shaders/Private/PostProcessDeviceEncodingOnly.usf
@@ -40,6 +40,7 @@ float4 ACESCoefsHigh_0;
 float ACESCoefsLow_4;
 float ACESCoefsHigh_4;
 float ACESSceneColorMultiplier;
+float ACESGamutCompression;
 
 uint GetOutputDevice()
 {
@@ -100,7 +101,7 @@ float4 DeviceEncodingOnlyCommonPS(
 	else if( GetOutputDevice() == TONEMAPPER_OUTPUT_ACES1000nitST2084 )
 	{		 
 		// 1000 nit ODT
-		FACESTonemapParams AcesParams = ComputeACESTonemapParams(ACESMinMaxData, ACESMidData, ACESCoefsLow_0, ACESCoefsHigh_0, ACESCoefsLow_4, ACESCoefsHigh_4, ACESSceneColorMultiplier);
+		FACESTonemapParams AcesParams = ComputeACESTonemapParams(ACESMinMaxData, ACESMidData, ACESCoefsLow_0, ACESCoefsHigh_0, ACESCoefsLow_4, ACESCoefsHigh_4, ACESSceneColorMultiplier, ACESGamutCompression);
 		float3 ODTColor = ACESOutputTransforms1000( FilmColor, (float3x3)WorkingColorSpace.ToAP0, AcesParams);
 
 		// Convert from AP1 to specified output gamut
@@ -113,7 +114,7 @@ float4 DeviceEncodingOnlyCommonPS(
 	// ACES 2000nit transform with PQ/2084 encoding, user specified gamut 
 	else if( GetOutputDevice() == TONEMAPPER_OUTPUT_ACES2000nitST2084 )
 	{		
-		FACESTonemapParams AcesParams = ComputeACESTonemapParams(ACESMinMaxData, ACESMidData, ACESCoefsLow_0, ACESCoefsHigh_0, ACESCoefsLow_4, ACESCoefsHigh_4, ACESSceneColorMultiplier);
+		FACESTonemapParams AcesParams = ComputeACESTonemapParams(ACESMinMaxData, ACESMidData, ACESCoefsLow_0, ACESCoefsHigh_0, ACESCoefsLow_4, ACESCoefsHigh_4, ACESSceneColorMultiplier, ACESGamutCompression);
 		
 		// 2000 nit ODT
 		float3 ODTColor = ACESOutputTransforms2000( FilmColor, (float3x3)WorkingColorSpace.ToAP0, AcesParams);
@@ -128,7 +129,7 @@ float4 DeviceEncodingOnlyCommonPS(
 	// ACES 1000nit transform to linear ScRGB
 	else if(GetOutputDevice() == TONEMAPPER_OUTPUT_ACES1000nitScRGB)
 	{
-		FACESTonemapParams AcesParams = ComputeACESTonemapParams(ACESMinMaxData, ACESMidData, ACESCoefsLow_0, ACESCoefsHigh_0, ACESCoefsLow_4, ACESCoefsHigh_4, ACESSceneColorMultiplier);
+		FACESTonemapParams AcesParams = ComputeACESTonemapParams(ACESMinMaxData, ACESMidData, ACESCoefsLow_0, ACESCoefsHigh_0, ACESCoefsLow_4, ACESCoefsHigh_4, ACESSceneColorMultiplier, ACESGamutCompression);
 
 		// 1000 nit ODT
 		float3 ODTColor = ACESOutputTransforms1000( FilmColor, (float3x3)WorkingColorSpace.ToAP0, AcesParams);
@@ -140,7 +141,7 @@ float4 DeviceEncodingOnlyCommonPS(
 	// ACES 2000nit transform to linear ScRGB
 	else if(GetOutputDevice() == TONEMAPPER_OUTPUT_ACES2000nitScRGB)
 	{
-		FACESTonemapParams AcesParams = ComputeACESTonemapParams(ACESMinMaxData, ACESMidData, ACESCoefsLow_0, ACESCoefsHigh_0, ACESCoefsLow_4, ACESCoefsHigh_4, ACESSceneColorMultiplier);
+		FACESTonemapParams AcesParams = ComputeACESTonemapParams(ACESMinMaxData, ACESMidData, ACESCoefsLow_0, ACESCoefsHigh_0, ACESCoefsLow_4, ACESCoefsHigh_4, ACESSceneColorMultiplier, ACESGamutCompression);
 
 		// 2000 nit ODT
 		float3 ODTColor = ACESOutputTransforms2000( FilmColor, (float3x3)WorkingColorSpace.ToAP0, AcesParams);
diff --git a/Engine/Shaders/Private/TonemapCommon.ush b/Engine/Shaders/Private/TonemapCommon.ush
--- a/Engine/Shaders/Private/TonemapCommon.ush
+++ b/Engine/Shaders/Private/TonemapCommon.ush
@@ -356,9 +356,10 @@ struct FACESTonemapParams
 {
 	TsParams AcesParams;
 	float SceneColorMultiplier;
+	float GamutCompression;
 };
 
-FACESTonemapParams ComputeACESTonemapParams(float4 InACESMinMaxData, float4 InACESMidData, float4 InACESCoefsLow_0, float4 InACESCoefsHigh_0, float InACESCoefsLow_4, float InACESCoefsHigh_4, float SceneColorMultiplier)
+FACESTonemapParams ComputeACESTonemapParams(float4 InACESMinMaxData, float4 InACESMidData, float4 InACESCoefsLow_0, float4 InACESCoefsHigh_0, float InACESCoefsLow_4, float InACESCoefsHigh_4, float SceneColorMultiplier, float GamutCompression)
 {
 	FACESTonemapParams OutACESTonemapParams;
 	OutACESTonemapParams.AcesParams.Min.x = InACESMinMaxData.x;
@@ -388,6 +389,7 @@ FACESTonemapParams ComputeACESTonemapParams(float4 InACESMinMaxData, float4 InAC
 	OutACESTonemapParams.AcesParams.coefsHigh[5] = InACESCoefsHigh_4;
 
 	OutACESTonemapParams.SceneColorMultiplier = SceneColorMultiplier;
+	OutACESTonemapParams.GamutCompression = GamutCompression;
 
 	return OutACESTonemapParams;
 
@@ -401,6 +403,9 @@ FACESTonemapParams ComputeACESTonemapParams(float4 InACESMinMaxData, float4 InAC
 float3 ACESOutputTransforms1000( float3 SceneReferredLinearColor, const float3x3 WCS_2_AP0, FACESTonemapParams InACESTonemapParams)
 {
 	float3 aces = mul( WCS_2_AP0, SceneReferredLinearColor * InACESTonemapParams.SceneColorMultiplier);
+	float3 AcesCompressed = compressColor(aces);
+	aces = lerp(aces, AcesCompressed, InACESTonemapParams.GamutCompression);
+
 	const FColorSpace UE_ACES_AP1_DISPLAY_PRI = GetColorSpace(TONEMAPPER_GAMUT_ACEScg_D60);
 	const FColorSpace UE_ACES_AP1_LIMITING_PRI = GetColorSpace(TONEMAPPER_GAMUT_ACEScg_D60);
 	float3 OutputReferredLinearAP1Color = outputTransform(aces, InACESTonemapParams.AcesParams, UE_ACES_AP1_DISPLAY_PRI, UE_ACES_AP1_LIMITING_PRI, UE_ACES_EOTF,
@@ -424,6 +429,9 @@ float3 ACESOutputTransforms1000( float3 SceneReferredLinearsRGBColor, FACESTonem
 float3 ACESOutputTransforms2000( float3 SceneReferredLinearColor, const float3x3 WCS_2_AP0, FACESTonemapParams InACESTonemapParams)
 {
 	float3 aces = mul( WCS_2_AP0, SceneReferredLinearColor * InACESTonemapParams.SceneColorMultiplier);
+	float3 AcesCompressed = compressColor(aces);
+	aces = lerp(aces, AcesCompressed, InACESTonemapParams.GamutCompression);
+
 	const FColorSpace UE_ACES_AP1_DISPLAY_PRI = GetColorSpace(TONEMAPPER_GAMUT_ACEScg_D60);
 	const FColorSpace UE_ACES_AP1_LIMITING_PRI = GetColorSpace(TONEMAPPER_GAMUT_ACEScg_D60);
 	float3 OutputReferredLinearAP1Color = outputTransform(aces, InACESTonemapParams.AcesParams, UE_ACES_AP1_DISPLAY_PRI, UE_ACES_AP1_LIMITING_PRI, UE_ACES_EOTF,
diff --git a/Engine/Source/Runtime/Engine/Private/UnrealEngine.cpp b/Engine/Source/Runtime/Engine/Private/UnrealEngine.cpp
--- a/Engine/Source/Runtime/Engine/Private/UnrealEngine.cpp
+++ b/Engine/Source/Runtime/Engine/Private/UnrealEngine.cpp
@@ -910,6 +910,7 @@ void ENGINE_API HDRSettingChangedSinkCallback()
 	static const auto CVarHDRMidLuminance = IConsoleManager::Get().FindTConsoleVariableDataFloat(TEXT("r.HDR.Display.MidLuminance"));
 	static const auto CVarHDRMaxLuminance = IConsoleManager::Get().FindTConsoleVariableDataInt(TEXT("r.HDR.Display.MaxLuminance"));
 	static const auto CVarHDRSceneColorMultiplier = IConsoleManager::Get().FindTConsoleVariableDataFloat(TEXT("r.HDR.Aces.SceneColorMultiplier"));
+	static const auto CVarHDRGamutCompression = IConsoleManager::Get().FindTConsoleVariableDataFloat(TEXT("r.HDR.Aces.GamutCompression"));
 
 	if (GRHIVendorId == 0)
 	{
@@ -920,6 +921,7 @@ void ENGINE_API HDRSettingChangedSinkCallback()
 	static float GHDRMidLuminnance = CVarHDRMidLuminance->GetValueOnAnyThread();
 	static int32 GHDRMaxLuminnance = CVarHDRMaxLuminance->GetValueOnAnyThread();
 	static float GHDRSceneColorMultiplier = CVarHDRSceneColorMultiplier->GetValueOnAnyThread();
+	static float GHDRGamutCompression = CVarHDRGamutCompression->GetValueOnAnyThread();
 	static int32 GDisplayNitLevel = 1000;
 
 	bool bIsHDREnabled = IsHDREnabled();
@@ -934,7 +936,9 @@ void ENGINE_API HDRSettingChangedSinkCallback()
 	if(bIsHDREnabled != GRHIIsHDREnabled || !FMath::IsNearlyEqual(GHDRMinLuminnanceLog10, CVarHDRMinLuminanceLog10->GetValueOnAnyThread())
 		|| !FMath::IsNearlyEqual(GHDRMidLuminnance, CVarHDRMidLuminance->GetValueOnAnyThread())
 		|| (GHDRMaxLuminnance != CVarHDRMaxLuminance->GetValueOnAnyThread()) || (DisplayNitLevel != GDisplayNitLevel)
-		|| !FMath::IsNearlyEqual(GHDRSceneColorMultiplier, CVarHDRSceneColorMultiplier->GetValueOnAnyThread()))
+		|| !FMath::IsNearlyEqual(GHDRSceneColorMultiplier, CVarHDRSceneColorMultiplier->GetValueOnAnyThread())
+		|| !FMath::IsNearlyEqual(GHDRGamutCompression, CVarHDRGamutCompression->GetValueOnAnyThread())
+		)
 	{
 		HDRConfigureCVars(bIsHDREnabled, DisplayNitLevel, false);
 
@@ -944,6 +948,7 @@ void ENGINE_API HDRSettingChangedSinkCallback()
 		GHDRMidLuminnance = CVarHDRMidLuminance->GetValueOnAnyThread();
 		GHDRMaxLuminnance = CVarHDRMaxLuminance->GetValueOnAnyThread();
 		GHDRSceneColorMultiplier = CVarHDRSceneColorMultiplier->GetValueOnAnyThread();
+		GHDRGamutCompression = CVarHDRGamutCompression->GetValueOnAnyThread();
 		GDisplayNitLevel = DisplayNitLevel;
 	}
 }
diff --git a/Engine/Source/Runtime/RenderCore/Private/RenderCore.cpp b/Engine/Source/Runtime/RenderCore/Private/RenderCore.cpp
--- a/Engine/Source/Runtime/RenderCore/Private/RenderCore.cpp
+++ b/Engine/Source/Runtime/RenderCore/Private/RenderCore.cpp
@@ -405,6 +405,13 @@ static TAutoConsoleVariable<float> CVarHDRAcesColorMultiplier(
 	ECVF_RenderThreadSafe
 );
 
+static TAutoConsoleVariable<float> CVarHDRAcesGamutCompression(
+	TEXT("r.HDR.Aces.GamutCompression"),
+	0.0f,
+	TEXT("HDR equivalent of BlueCorrection: Bright blue desaturates instead of going to violet"),
+	ECVF_RenderThreadSafe
+);
+
 TAutoConsoleVariable<int32> CVarHDROutputEnabled(
 	TEXT("r.HDR.EnableHDROutput"),
 	0,
@@ -983,6 +990,7 @@ void ConfigureACESTonemapParams(FACESTonemapParams& OutACESTonemapParams, float
 	OutACESTonemapParams.ACESCoefsHigh_4 = PARAMS.coefsHigh[4];
 
 	OutACESTonemapParams.ACESSceneColorMultiplier = CVarHDRAcesColorMultiplier.GetValueOnAnyThread();
+	OutACESTonemapParams.ACESGamutCompression = CVarHDRAcesGamutCompression.GetValueOnAnyThread();
 
 }
 
diff --git a/Engine/Source/Runtime/RenderCore/Public/HDRHelper.h b/Engine/Source/Runtime/RenderCore/Public/HDRHelper.h
--- a/Engine/Source/Runtime/RenderCore/Public/HDRHelper.h
+++ b/Engine/Source/Runtime/RenderCore/Public/HDRHelper.h
@@ -44,6 +44,7 @@ struct FACESTonemapParams
 	, ACESCoefsLow_4(1.84773231f)
 	, ACESCoefsHigh_4(3.0f)
 	, ACESSceneColorMultiplier(1.5f)
+	, ACESGamutCompression(0.0f)
 	{
 
 	}
@@ -54,6 +55,7 @@ struct FACESTonemapParams
 	float ACESCoefsLow_4;
 	float ACESCoefsHigh_4;
 	float ACESSceneColorMultiplier;
+	float ACESGamutCompression;
 };
 
 UE_DEPRECATED(all, "For internal use only.")
@@ -73,4 +75,5 @@ inline void GetACESTonemapParameters(FACESTonemapParametersType& ACESTonemapPara
 	ACESTonemapParameters.ACESCoefsLow_4 = TonemapperParams.ACESCoefsLow_4;
 	ACESTonemapParameters.ACESCoefsHigh_4 = TonemapperParams.ACESCoefsHigh_4;
 	ACESTonemapParameters.ACESSceneColorMultiplier = TonemapperParams.ACESSceneColorMultiplier;
+	ACESTonemapParameters.ACESGamutCompression = FMath::Clamp(TonemapperParams.ACESGamutCompression, 0.0f, 1.0f);
 }
\ No newline at end of file
diff --git a/Engine/Source/Runtime/Renderer/Private/PostProcess/PostProcessCombineLUTs.cpp b/Engine/Source/Runtime/Renderer/Private/PostProcess/PostProcessCombineLUTs.cpp
--- a/Engine/Source/Runtime/Renderer/Private/PostProcess/PostProcessCombineLUTs.cpp
+++ b/Engine/Source/Runtime/Renderer/Private/PostProcess/PostProcessCombineLUTs.cpp
@@ -113,6 +113,7 @@ BEGIN_SHADER_PARAMETER_STRUCT(FACESTonemapShaderParameters, )
 	SHADER_PARAMETER(float, ACESCoefsLow_4)
 	SHADER_PARAMETER(float, ACESCoefsHigh_4)
 	SHADER_PARAMETER(float, ACESSceneColorMultiplier)
+	SHADER_PARAMETER(float, ACESGamutCompression)
 END_SHADER_PARAMETER_STRUCT()
 
 
@@ -243,6 +244,7 @@ struct FCachedLUTSettings
 		UPDATE_CACHE_SETTINGS(Parameters.ACESTonemapParameters.ACESCoefsLow_4, TonemapperParams.ACESCoefsLow_4, bHasChanged);
 		UPDATE_CACHE_SETTINGS(Parameters.ACESTonemapParameters.ACESCoefsHigh_4, TonemapperParams.ACESCoefsHigh_4, bHasChanged);
 		UPDATE_CACHE_SETTINGS(Parameters.ACESTonemapParameters.ACESSceneColorMultiplier, TonemapperParams.ACESSceneColorMultiplier, bHasChanged);
+		UPDATE_CACHE_SETTINGS(Parameters.ACESTonemapParameters.ACESGamutCompression, TonemapperParams.ACESGamutCompression, bHasChanged);
 
 		UPDATE_CACHE_SETTINGS(Parameters.ColorScale, FVector3f(View.ColorScale), bHasChanged);
 		UPDATE_CACHE_SETTINGS(Parameters.OverlayColor, FVector4f(View.OverlayColor), bHasChanged);
diff --git a/Engine/Source/Runtime/Renderer/Private/PostProcess/PostProcessDeviceEncodingOnly.cpp b/Engine/Source/Runtime/Renderer/Private/PostProcess/PostProcessDeviceEncodingOnly.cpp
--- a/Engine/Source/Runtime/Renderer/Private/PostProcess/PostProcessDeviceEncodingOnly.cpp
+++ b/Engine/Source/Runtime/Renderer/Private/PostProcess/PostProcessDeviceEncodingOnly.cpp
@@ -92,6 +92,7 @@ BEGIN_SHADER_PARAMETER_STRUCT(FDeviceEncodingACESTonemapShaderParameters, )
 	SHADER_PARAMETER(float, ACESCoefsLow_4)
 	SHADER_PARAMETER(float, ACESCoefsHigh_4)
 	SHADER_PARAMETER(float, ACESSceneColorMultiplier)
+	SHADER_PARAMETER(float, ACESGamutCompression)
 END_SHADER_PARAMETER_STRUCT()
 
 BEGIN_SHADER_PARAMETER_STRUCT(FDeviceEncodingOnlyParameters, )
