Fix UnlitViewMode for Mobile Preview

diff --git a/Engine/Shaders/Private/MobileBasePassPixelShader.usf b/Engine/Shaders/Private/MobileBasePassPixelShader.usf
--- a/Engine/Shaders/Private/MobileBasePassPixelShader.usf
+++ b/Engine/Shaders/Private/MobileBasePassPixelShader.usf
@@ -918,7 +918,7 @@ void Main(
 	Color += Emissive;
 
 #if !MATERIAL_SHADINGMODEL_UNLIT && MOBILE_EMULATION
-	Color = lerp(Color, DiffuseColor, ResolvedView.UnlitViewmodeMask);
+	Color = lerp(Color, DiffuseColorForIndirect, ResolvedView.UnlitViewmodeMask);
 #endif
 
 #if MATERIAL_SHADINGMODEL_THIN_TRANSLUCENT && !SUBSTRATE_ENABLED
