LFV - removed now obsolete code.

diff --git a/Engine/Shaders/Private/LocalFogVolumes/LocalFogVolumeCommon.ush b/Engine/Shaders/Private/LocalFogVolumes/LocalFogVolumeCommon.ush
--- a/Engine/Shaders/Private/LocalFogVolumes/LocalFogVolumeCommon.ush
+++ b/Engine/Shaders/Private/LocalFogVolumes/LocalFogVolumeCommon.ush
@@ -222,21 +222,7 @@ FFogData LocalFogVolumeEvaluateRadialIntegral(FLocalFogVolumeGPUInstanceData Fog
 ///////////////////////////////////////////////////////////////////////////////
 // Local Fog Volume Evaluation
 
-struct LocalFogVolumeCommonLightingInput
-{
-	float3 DistantSkyLight;
-};
-
-LocalFogVolumeCommonLightingInput LocalFogVolumeEvaluateCommonLightingInput()
-{
-	LocalFogVolumeCommonLightingInput CommonLightingInput = (LocalFogVolumeCommonLightingInput)0;
-#if PROJECT_SUPPORT_SKY_ATMOSPHERE_AFFECTS_HEIGHFOG
-	CommonLightingInput.DistantSkyLight = Texture2DSampleLevel(View.DistantSkyLightLutTexture, View.DistantSkyLightLutTextureSampler, float2(0.5f, 0.5f), 0.0f).rgb;
-#endif
-	return CommonLightingInput;
-}
-
-float3 LocalFogVolumeEvaluateInScattering(in FLocalFogVolumeGPUInstanceData FogInstance, in FFogData FogData, in LocalFogVolumeCommonLightingInput CommonLightingInput, in float3 RayDirWorld)
+float3 LocalFogVolumeEvaluateInScattering(in FLocalFogVolumeGPUInstanceData FogInstance, in FFogData FogData, in float3 RayDirWorld)
 {
 	half3 InScattering = 0.0f;
 
@@ -250,7 +236,7 @@ float3 LocalFogVolumeEvaluateInScattering(in FLocalFogVolumeGPUInstanceData FogI
 	{
 		float3 SkyLighting = View.SkyLightColor.rgb * GetSkySHDiffuseSimple(RayDirWorld * -FogInstance.PhaseG);
 
-		float SkyVisibility = 1.0f; // Not possible to reach with analytical integration
+		const float SkyVisibility = 1.0f; // Not possible to use AO with analytical integration
 		InScattering += (SkyVisibility * View.SkyLightVolumetricScatteringIntensity) * SkyLighting;
 	}
 	
@@ -328,8 +314,6 @@ float4 GetLFVInstanceContribution(
 		{
 			FFogData FogData = (FFogData)0;
 
-			LocalFogVolumeCommonLightingInput CommonLightingInput = LocalFogVolumeEvaluateCommonLightingInput();
-
 			if (FogInstance.FogMode == 0)
 			{
 				FogData = LocalFogVolumeEvaluateHeightIntegral(FogInstance, CamRayWorldOrigin, CamRayWorldOrigin + CamRayWorldDir * Length0, CamRayWorldDir, RayTracedLength);
@@ -343,7 +327,7 @@ float4 GetLFVInstanceContribution(
 				FogData = LocalFogVolumeEvaluateRadialIntegral(FogInstance, CamRayWorldOrigin, CamRayWorldOrigin + CamRayWorldDir * Length0, CamRayWorldDir, RayTracedLength);
 			}
 
-			LuminanceColor = LocalFogVolumeEvaluateInScattering(FogInstance, FogData, CommonLightingInput, CamRayWorldDir);
+			LuminanceColor = LocalFogVolumeEvaluateInScattering(FogInstance, FogData, CamRayWorldDir);
 			Transmittance = 1.0 - FogData.Coverage;
 		}
 		else
