Add GBuffer data value picking/visualization.

diff --git a/Engine/Shaders/Private/PostProcessGBufferHints.usf b/Engine/Shaders/Private/PostProcessGBufferHints.usf
--- a/Engine/Shaders/Private/PostProcessGBufferHints.usf
+++ b/Engine/Shaders/Private/PostProcessGBufferHints.usf
@@ -103,3 +103,126 @@ void MainPS(noperspective float4 UVAndScreenPos : TEXCOORD0, float4 SvPosition :
 
 	HighlightAdvice(OutColor.rgb, UV, PixelPos);
 }
+
+
+////////////////////////////////////////////////////////////////////////////////////////////////
+// Debug
+#if SHADER_GBUFFER_PICKING
+
+#include "ShaderPrint.ush"
+#include "SceneTextureParameters.ush"
+#include "DeferredShadingCommon.ush"
+
+void PrintShadingMode(inout FShaderPrintContext Ctx, uint In)
+{
+	if (In == SHADINGMODELID_UNLIT) 				{ Print(Ctx, TEXT("UNLIT"), FontRed); return; }
+	if (In == SHADINGMODELID_DEFAULT_LIT) 			{ Print(Ctx, TEXT("DEFAULT_LIT"), FontRed); return; }
+	if (In == SHADINGMODELID_SUBSURFACE) 			{ Print(Ctx, TEXT("SUBSURFACE"), FontRed); return; }
+	if (In == SHADINGMODELID_PREINTEGRATED_SKIN) 	{ Print(Ctx, TEXT("PREINTEGRATED_SKIN"), FontRed); return; }
+	if (In == SHADINGMODELID_CLEAR_COAT) 			{ Print(Ctx, TEXT("CLEAR_COAT"), FontRed); return; }
+	if (In == SHADINGMODELID_SUBSURFACE_PROFILE)	{ Print(Ctx, TEXT("SUBSURFACE_PROFILE"), FontRed); return; }
+	if (In == SHADINGMODELID_TWOSIDED_FOLIAGE) 		{ Print(Ctx, TEXT("TWOSIDED_FOLIAGE"), FontRed); return; }
+	if (In == SHADINGMODELID_HAIR) 					{ Print(Ctx, TEXT("HAIR"), FontRed); return; }
+	if (In == SHADINGMODELID_CLOTH) 				{ Print(Ctx, TEXT("CLOTH"), FontRed); return; }
+	if (In == SHADINGMODELID_EYE) 					{ Print(Ctx, TEXT("EYE"), FontRed); return; }
+	if (In == SHADINGMODELID_SINGLELAYERWATER)		{ Print(Ctx, TEXT("SINGLELAYERWATER"), FontRed); return; }
+	if (In == SHADINGMODELID_THIN_TRANSLUCENT)		{ Print(Ctx, TEXT("THIN_TRANSLUCENT"), FontRed); return; }
+	if (In == SHADINGMODELID_SUBSTRATE)				{ Print(Ctx, TEXT("SUBSTRATE"), FontRed); return; }
+	if (In == SHADINGMODELID_NUM)					{ Print(Ctx, TEXT("NUM"), FontRed); return; }
+	if (In == SHADINGMODELID_MASK) 					{ Print(Ctx, TEXT("MASK"), FontRed); return; }
+}
+
+void Newline(inout FShaderPrintContext Ctx, inout float2 MaxRect)
+{
+	MaxRect.x = max(MaxRect.x, Ctx.Pos.x);
+	MaxRect.y = Ctx.Pos.y;
+	Newline(Ctx);
+}
+
+void PrintBackgroundRect(inout FShaderPrintContext Ctx, float2 StartRect, float2 EndRect)
+{
+	const float4 RectBackgroundColor = float4(1, 1, 1, 0.25f);
+	AddFilledQuadSS(
+		(StartRect - ShaderPrintData.FontSize) * ShaderPrintData.Resolution,
+		(EndRect + ShaderPrintData.FontSize) * ShaderPrintData.Resolution,
+		RectBackgroundColor);
+}
+
+[numthreads(1, 1, 1)]
+void MainCS(uint3 DispatchThreadId : SV_DispatchThreadID)
+{
+#if !SUBSTRATE_ENABLED
+	// Draw stats
+	if (all(DispatchThreadId == 0))
+	{
+		FShaderPrintContext Ctx = InitShaderPrintContext(true, uint2(50, 50));
+		const FScreenSpaceData ScreenSpaceData = GetScreenSpaceDataUint(ShaderPrintData.CursorCoord, true);
+		const FGBufferData Data = ScreenSpaceData.GBuffer;
+
+		float2 RectMin = Ctx.Pos;
+		float2 RectMax = Ctx.StartPos;
+
+		RectMin = Ctx.Pos;
+		Print(Ctx, TEXT("Global"), FontOrange);Newline(Ctx, RectMax);
+		Print(Ctx, TEXT("Diffuse Color     : "), FontWhite); Print(Ctx, Data.DiffuseColor, FontYellow); Newline(Ctx, RectMax);
+		Print(Ctx, TEXT("Specular Color    : "), FontWhite); Print(Ctx, Data.SpecularColor, FontYellow); Newline(Ctx, RectMax);
+		PrintBackgroundRect(Ctx, RectMin, RectMax);
+		Newline(Ctx);
+
+		RectMin = Ctx.Pos;
+		Print(Ctx, TEXT("Tangent basis"), FontOrange);Newline(Ctx, RectMax);
+		Print(Ctx, TEXT("World Normal      : "), FontWhite); Print(Ctx, Data.WorldNormal, FontYellow); Newline(Ctx, RectMax);
+		Print(Ctx, TEXT("World Tangent     : "), FontWhite); Print(Ctx, Data.WorldTangent, FontYellow); Newline(Ctx, RectMax);
+		PrintBackgroundRect(Ctx, RectMin, RectMax);
+		Newline(Ctx);
+
+		RectMin = Ctx.Pos;
+		Print(Ctx, TEXT("Common"), FontOrange);Newline(Ctx, RectMax);
+		Print(Ctx, TEXT("BaseColor         : "), FontWhite); Print(Ctx, Data.BaseColor, FontYellow); Newline(Ctx, RectMax);
+		Print(Ctx, TEXT("Metallic          : "), FontWhite); Print(Ctx, Data.Metallic, FontYellow); Newline(Ctx, RectMax);
+		Print(Ctx, TEXT("Specular          : "), FontWhite); Print(Ctx, Data.Specular, FontYellow); Newline(Ctx, RectMax);
+		Print(Ctx, TEXT("Roughness         : "), FontWhite); Print(Ctx, Data.Roughness, FontYellow); Newline(Ctx, RectMax);
+		Print(Ctx, TEXT("Anisotropy        : "), FontWhite); Print(Ctx, Data.Anisotropy, FontYellow); Newline(Ctx, RectMax);
+		PrintBackgroundRect(Ctx, RectMin, RectMax);
+		Newline(Ctx);
+
+		RectMin = Ctx.Pos;
+		Print(Ctx, TEXT("Occlusion / Ind. Lighting"), FontOrange);Newline(Ctx, RectMax);
+		Print(Ctx, TEXT("Generic AO        : "), FontWhite); Print(Ctx, Data.GenericAO, FontYellow); Newline(Ctx, RectMax);
+		Print(Ctx, TEXT("GBuffer AO        : "), FontWhite); Print(Ctx, Data.GBufferAO, FontYellow); Newline(Ctx, RectMax);
+		Print(Ctx, TEXT("Ind. Sample Occ.  : "), FontWhite); Print(Ctx, Data.DiffuseIndirectSampleOcclusion, FontYellow); Newline(Ctx, RectMax);
+		Print(Ctx, TEXT("Ind. Irradiance   : "), FontWhite); Print(Ctx, Data.IndirectIrradiance, FontYellow); Newline(Ctx, RectMax);
+		Print(Ctx, TEXT("Prec.Shadow Factor: "), FontWhite); Print(Ctx, Data.PrecomputedShadowFactors, FontYellow); Newline(Ctx, RectMax);
+		PrintBackgroundRect(Ctx, RectMin, RectMax);
+		Newline(Ctx);
+
+		RectMin = Ctx.Pos;
+		Print(Ctx, TEXT("Custom"), FontOrange);Newline(Ctx, RectMax);
+		Print(Ctx, TEXT("Shading Model ID  : "), FontWhite); PrintShadingMode(Ctx, Data.ShadingModelID); Newline(Ctx, RectMax);
+		Print(Ctx, TEXT("Custom Data       : "), FontWhite); Print(Ctx, Data.CustomData, FontYellow); Newline(Ctx, RectMax);
+		if (Data.ShadingModelID == SHADINGMODELID_CLEAR_COAT)
+		{
+		Print(Ctx, TEXT("Clear Coat        : "), FontWhite); Print(Ctx, Data.CustomData.x, FontYellow); Newline(Ctx, RectMax);
+		Print(Ctx, TEXT("Clear Coat Roughn.: "), FontWhite); Print(Ctx, Data.CustomData.y, FontYellow); Newline(Ctx, RectMax);
+		}
+		else if (Data.ShadingModelID == SHADINGMODELID_CLOTH)
+		{
+		Print(Ctx, TEXT("Fuzz Color        : "), FontWhite); Print(Ctx, Data.CustomData.xyz, FontYellow); Newline(Ctx, RectMax);
+		Print(Ctx, TEXT("Fuzz Amount       : "), FontWhite); Print(Ctx, Data.CustomData.a, FontYellow); Newline(Ctx, RectMax);
+		}
+		else if (Data.ShadingModelID == SHADINGMODELID_SUBSURFACE || Data.ShadingModelID == SHADINGMODELID_PREINTEGRATED_SKIN || Data.ShadingModelID == SHADINGMODELID_TWOSIDED_FOLIAGE)
+		{
+		Print(Ctx, TEXT("Subsurface Color  : "), FontWhite); Print(Ctx, ExtractSubsurfaceColor(Data), FontYellow); Newline(Ctx, RectMax);
+		Print(Ctx, TEXT("Opacity           : "), FontWhite); Print(Ctx, Data.CustomData.a, FontYellow); Newline(Ctx, RectMax);
+		}
+		else if (Data.ShadingModelID == SHADINGMODELID_SUBSURFACE_PROFILE)
+		{
+		Print(Ctx, TEXT("SSS Profil ID     : "), FontWhite); Print(Ctx, ExtractSubsurfaceProfileInt(Data.CustomData.x)); Newline(Ctx, RectMax);	
+		}
+		PrintBackgroundRect(Ctx, RectMin, RectMax);
+		Newline(Ctx);
+	}
+#endif
+}
+
+#endif // SHADER_GBUFFER_PICKING
\ No newline at end of file
diff --git a/Engine/Source/Runtime/Renderer/Private/PostProcess/PostProcessing.cpp b/Engine/Source/Runtime/Renderer/Private/PostProcess/PostProcessing.cpp
--- a/Engine/Source/Runtime/Renderer/Private/PostProcess/PostProcessing.cpp
+++ b/Engine/Source/Runtime/Renderer/Private/PostProcess/PostProcessing.cpp
@@ -142,8 +142,19 @@ TAutoConsoleVariable<int32> CVarPostProcessingForceAsyncDispatch(
 	TEXT("Only available for testing in non-shipping builds."),
 	ECVF_RenderThreadSafe);
 #endif
+
+#if WITH_EDITOR
+TAutoConsoleVariable<int32> CVarGBufferPicking(
+	TEXT("r.PostProcessing.GBufferPicking"), 0,
+	TEXT("Evaluate GBuffer value for debugging purpose."),
+	ECVF_RenderThreadSafe);
+#endif
 }
 
+#if WITH_EDITOR
+static void AddGBufferPicking(FRDGBuilder& GraphBuilder, const FViewInfo& View, const TRDGUniformBufferRef<FSceneTextureUniformParameters>& SceneTextures);
+#endif 
+
 EDownsampleQuality GetDownsampleQuality(const TAutoConsoleVariable<int32>& CVar)
 {
 	const int32 DownsampleQuality = FMath::Clamp(CVar.GetValueOnRenderThread(), 0, 1);
@@ -1636,6 +1647,13 @@ void AddPostProcessingPasses(
 			VirtualShadowMapArray->AddVisualizePass(GraphBuilder, View, ViewIndex, SceneColor);
 		}
 
+		#if WITH_EDITOR
+		if (CVarGBufferPicking.GetValueOnRenderThread())
+		{
+			AddGBufferPicking(GraphBuilder, View, Inputs.SceneTextures);
+		}
+		#endif
+
 		{
 			RectLightAtlas::AddDebugPass(GraphBuilder, View, SceneColor.Texture);
 			IESAtlas::AddDebugPass(GraphBuilder, View, SceneColor.Texture);
@@ -2884,3 +2902,51 @@ FScreenPassTexture AddFinalPostProcessDebugInfoPasses(FRDGBuilder& GraphBuilder,
 	return MoveTemp(ScreenPassSceneColor);
 }
 #endif
+
+// Shader for visualizing GBuffer values
+class FGBufferPickingCS : public FGlobalShader
+{
+	DECLARE_GLOBAL_SHADER(FGBufferPickingCS);
+	SHADER_USE_PARAMETER_STRUCT(FGBufferPickingCS, FGlobalShader);
+
+	BEGIN_SHADER_PARAMETER_STRUCT(FParameters, )
+		SHADER_PARAMETER_RDG_UNIFORM_BUFFER(FSceneTextureUniformParameters, SceneTexturesStruct)
+		SHADER_PARAMETER_STRUCT_INCLUDE(ShaderPrint::FShaderParameters, ShaderPrintParameters)
+		SHADER_PARAMETER_STRUCT_REF(FViewUniformShaderParameters, ViewUniformBuffer)
+	END_SHADER_PARAMETER_STRUCT()
+
+public:
+	static bool ShouldCompilePermutation(const FGlobalShaderPermutationParameters& Parameters) 
+	{ 
+		return ShaderPrint::IsSupported(Parameters.Platform) && !Substrate::IsSubstrateEnabled() && 
+			EnumHasAllFlags(Parameters.Flags, EShaderPermutationFlags::HasEditorOnlyData);
+	}
+	static void ModifyCompilationEnvironment(const FGlobalShaderPermutationParameters& Parameters, FShaderCompilerEnvironment& OutEnvironment)
+	{
+		FGlobalShader::ModifyCompilationEnvironment(Parameters, OutEnvironment);
+		OutEnvironment.SetDefine(TEXT("SHADER_GBUFFER_PICKING"), 1);
+	}
+};
+
+IMPLEMENT_GLOBAL_SHADER(FGBufferPickingCS, "/Engine/Private/PostProcessGBufferHints.usf", "MainCS", SF_Compute);
+
+#if WITH_EDITOR
+static void AddGBufferPicking(FRDGBuilder& GraphBuilder, const FViewInfo& View, const TRDGUniformBufferRef<FSceneTextureUniformParameters>& SceneTextures)
+{
+	if (CVarGBufferPicking.GetValueOnRenderThread() <= 0 || !ShaderPrint::IsSupported(View.Family->GetShaderPlatform()))
+	{
+		return;
+	}
+
+	// Force ShaderPrint on.
+	ShaderPrint::SetEnabled(true);
+
+	FGBufferPickingCS::FParameters* Parameters = GraphBuilder.AllocParameters<FGBufferPickingCS::FParameters>();
+	Parameters->ViewUniformBuffer = View.ViewUniformBuffer;
+	Parameters->SceneTexturesStruct = SceneTextures;
+	ShaderPrint::SetParameters(GraphBuilder, View.ShaderPrintData, Parameters->ShaderPrintParameters);
+
+	TShaderMapRef<FGBufferPickingCS> ComputeShader(View.ShaderMap);
+	FComputeShaderUtils::AddPass(GraphBuilder, RDG_EVENT_NAME("Debug::GBufferPicking"), ComputeShader, Parameters, FIntVector(1,1,1));
+}
+#endif
\ No newline at end of file
