Fix missing SpecularProfile weighting on secondary specular lobe for env. evaluation.

diff --git a/Engine/Shaders/Private/Substrate/SubstrateEvaluation.ush b/Engine/Shaders/Private/Substrate/SubstrateEvaluation.ush
--- a/Engine/Shaders/Private/Substrate/SubstrateEvaluation.ush
+++ b/Engine/Shaders/Private/Substrate/SubstrateEvaluation.ush
@@ -1573,7 +1573,7 @@ FSubstrateEnvLightResult SubstrateEvaluateForEnvLight(FSubstrateBSDFContext BSDF
 					const float NoR = saturate(dot(SubstrateEnvLightResult.SpecularDirection, BSDFContext.N));
 					SubstrateEnvLightResult.SpecularWeight *= EvaluateSpecularProfile(SLAB_SPECPROFILEID(BSDFContext.BSDF), BSDFContext.SatNoV, NoR, BSDFContext.SatNoV, NoR);
 				}
-			#endif // SUBSTRATE_COMPLEXPATH
+			#endif // SUBSTRATE_COMPLEXSPECIALPATH
 
 				SubstrateEnvLightResult.SSRSpecularWeight = SubstrateEnvLightResult.SpecularWeight;
 
@@ -1628,9 +1628,22 @@ FSubstrateEnvLightResult SubstrateEvaluateForEnvLight(FSubstrateBSDFContext BSDF
 						SubstrateEnvLightResult.SpecularHazeWeight = lerp(0.0f, HazeWeight, HazinessFadeIn) * EvalEnvBRDFHaze;
 						SubstrateEnvLightResult.SpecularHazeSafeRoughness = HazeSafeRoughness;
 					}
+
+					// Specular profile for haze
+				#if SUBSTRATE_COMPLEXSPECIALPATH
+					BRANCH
+					if (BSDF_GETHASSPECPROFILE(BSDFContext.BSDF))
+					{
+						const float NoR = saturate(dot(SubstrateEnvLightResult.SpecularHazeDirection, BSDFContext.N));
+						SubstrateEnvLightResult.SpecularHazeWeight *= EvaluateSpecularProfile(SLAB_SPECPROFILEID(BSDFContext.BSDF), BSDFContext.SatNoV, NoR, BSDFContext.SatNoV, NoR);
+					}
+				#endif // SUBSTRATE_COMPLEXSPECIALPATH
+
 				}
-				#endif
-			}
+				#endif // SUBSTRATE_FASTPATH==0
+
+			} // if (bEnableSpecular)
+
 
 			#if SUBSTRATE_FASTPATH==0
 			BRANCH
