LFV - volume optical depth will now increase when radius increases.

diff --git a/Engine/Shaders/Private/LocalFogVolumes/LocalFogVolumeCommon.ush b/Engine/Shaders/Private/LocalFogVolumes/LocalFogVolumeCommon.ush
--- a/Engine/Shaders/Private/LocalFogVolumes/LocalFogVolumeCommon.ush
+++ b/Engine/Shaders/Private/LocalFogVolumes/LocalFogVolumeCommon.ush
@@ -159,8 +159,8 @@ FFogData LocalFogVolumeEvaluateHeightIntegral(FLocalFogVolumeGPUInstanceData Fog
 	OpticalDepth = (FogInstance.Density / (FogInstance.HeightFalloff * RayDirU.z)) * (exp(-Factor0) - exp(-(Factor0 + Factor1)));
 #endif
 
-	// Transmittance is the 
-	float Transmittance = exp(-OpticalDepth);
+	// Evaluate transmittance, accounting for the volume scale converted to local unit space meter.
+	float Transmittance = exp(-OpticalDepth * FogInstance.UniformScale * CENTIMETER_TO_METER);
 	// Assuming extinction is grey scale, we can compute a single coverage value from transmittance
 	FogData.Coverage = (1.0 - Transmittance);
 	// Assuming extinction==scattering, i.e. albedo=1, we can compute the integrale for each point on a ray while accouting for transmittance back to the view.
@@ -206,8 +206,8 @@ FFogData LocalFogVolumeEvaluateRadialIntegral(FLocalFogVolumeGPUInstanceData Fog
 		OpticalDepth = max(0.0, FogInstance.Density * (Integral1 - Integral0) * (3.0f / 4.0f));
 	}
 
-	// Transmittance is the 
-	float Transmittance = exp(-OpticalDepth);
+	// Evaluate transmittance, accounting for the volume scale converted to local unit space meter.
+	float Transmittance = exp(-OpticalDepth * FogInstance.UniformScale * CENTIMETER_TO_METER);
 	// Assuming extinction is grey scale, we can compute a single coverage value from transmittance
 	FogData.Coverage = (1.0 - Transmittance);
 	// Assuming extinction==scattering, i.e. albedo=1, we can compute the integrale for each point on a ray while accouting for transmittance back to the view.
diff --git a/Engine/Shaders/Private/VolumetricFog.usf b/Engine/Shaders/Private/VolumetricFog.usf
--- a/Engine/Shaders/Private/VolumetricFog.usf
+++ b/Engine/Shaders/Private/VolumetricFog.usf
@@ -170,9 +170,6 @@ void MaterialSetupCS(
 
 	uint LFVCount = LFVTileDataTex[uint3(TilePos, 0)];
 
-	//AdditionalEmissive += float3(TilePos.x * 0.00001, TilePos.y * 0.00001, 0.0);
-	//AdditionalEmissive += LFVCount * 0.00001;
-
 	for (int LFVIndex = 0; LFVIndex < LFVCount; ++LFVIndex)
 	{
 		uint InstanceIndex = LFVTileDataTex[uint3(TilePos, 1 + LFVIndex)];
@@ -191,12 +188,16 @@ void MaterialSetupCS(
 		float LFVExtinction = 0.0;
 		if (FogInstance.FogMode == 0)
 		{
-			LFVExtinction = SphereFade <= 0.0f ? 0.0f : FogInstance.Density * exp(-FogInstance.HeightFalloff * UnitSpacePos.z) / FogInstance.UniformScale;
+			LFVExtinction = SphereFade <= 0.0f ? 0.0f : FogInstance.Density * exp(-FogInstance.HeightFalloff * UnitSpacePos.z);
 		}
 		else // (FogInstance.FogMode == 1)
 		{
-			LFVExtinction = FogInstance.Density * SphereFade / FogInstance.UniformScale;
+			LFVExtinction = FogInstance.Density * SphereFade;
 		}
+
+		// Convert extinction from unit volume 1/m into extinction for volumetric fog as 1/cm
+		LFVExtinction *= (1.0f / METER_TO_CENTIMETER);
+
 		LFVExtinction = min(LFVMaxDensityIntoVolumetricFog, LFVExtinction);
 
 		float3 LFVScattering = LFVExtinction * FogInstance.Albedo;
diff --git a/Engine/Source/Runtime/Engine/Classes/Components/LocalFogVolumeComponent.h b/Engine/Source/Runtime/Engine/Classes/Components/LocalFogVolumeComponent.h
--- a/Engine/Source/Runtime/Engine/Classes/Components/LocalFogVolumeComponent.h
+++ b/Engine/Source/Runtime/Engine/Classes/Components/LocalFogVolumeComponent.h
@@ -25,11 +25,11 @@ class ULocalFogVolumeComponent : public USceneComponent
 
 	/** Controls the softness of the transition region when the volume is fading out. */
 	UPROPERTY(EditAnywhere, BlueprintReadOnly, Category = "Fog Mode")
-	ELocalFogMode FogMode = ELocalFogMode::LocalHeightFog;
+	ELocalFogMode FogMode = ELocalFogMode::LocalSphereFog;
 
 	/** Global density factor for this fog. */
 	UPROPERTY(EditAnywhere, BlueprintReadOnly, interp, Category = "Fog Distribution", meta = (UIMin = "0", UIMax = "10.0", SliderExponent = 2.0, ClampMin = 0.0))
-	float FogDensity = 5.0f;
+	float FogDensity = 1.0f;
 
 	/** Controls how the density decreases as height increases. Smaller values make the visible transition larger. */
 	UPROPERTY(EditAnywhere, BlueprintReadOnly, interp, Category = "Fog Distribution", meta = (UIMin = "0.001", UIMax = "5000", SliderExponent = 2.0, ClampMin = 0.001))
