LFV - radial and height fog can both be combined now.

#prelfight none

diff --git a/Engine/Shaders/Private/LocalFogVolumes/LocalFogVolumeCommon.ush b/Engine/Shaders/Private/LocalFogVolumes/LocalFogVolumeCommon.ush
--- a/Engine/Shaders/Private/LocalFogVolumes/LocalFogVolumeCommon.ush
+++ b/Engine/Shaders/Private/LocalFogVolumes/LocalFogVolumeCommon.ush
@@ -15,15 +15,15 @@ struct FLocalFogVolumeGPUInstanceData
 	float4x4 Transform;
 	float4x4 InvTransform;
 
-	float Density;
-	float HeightFalloff;
-	float HeightOffset;
-	float UniformScale;
+	float RadialFogExtinction;
+	float HeightFogExtinction;
+	float HeightFogFalloff;
+	float HeightFogOffset;
 
 	float3 Albedo;
 	float  PhaseG;
 	float3 Emissive;
-	float  FogMode;
+	float  UniformScale;
 };
 
 
@@ -81,19 +81,19 @@ FLocalFogVolumeGPUInstanceData GetLocalFogVolumeGPUInstanceData(uint Index)
 	Data.Transform			= float4x4(LFVInsts[Offset + 0 ], LFVInsts[Offset + 1 ], LFVInsts[Offset + 2 ], LFVInsts[Offset + 3 ]);
 	Data.InvTransform		= float4x4(LFVInsts[Offset + 4 ], LFVInsts[Offset + 5 ], LFVInsts[Offset + 6 ], LFVInsts[Offset + 7 ]);
 
-	float4 Temp0 = LFVInsts[Offset + 8];
-	Data.Density = Temp0.x;
-	Data.HeightFalloff = Temp0.y;
-	Data.HeightOffset = Temp0.z;
-	Data.UniformScale = Temp0.w;
+	float4 Temp0			= LFVInsts[Offset + 8];
+	Data.RadialFogExtinction= Temp0.x;
+	Data.HeightFogExtinction= Temp0.y;
+	Data.HeightFogFalloff	= Temp0.z;
+	Data.HeightFogOffset	= Temp0.w;
 
-	float4 Temp1 = LFVInsts[Offset + 9];
-	Data.Albedo = Temp1.xyz;
-	Data.PhaseG = Temp1.w;
+	float4 Temp1			= LFVInsts[Offset + 9];
+	Data.Albedo				= Temp1.xyz;
+	Data.PhaseG				= Temp1.w;
 
-	float4 Temp2 = LFVInsts[Offset + 10];
-	Data.Emissive = Temp2.xyz;
-	Data.FogMode = Temp2.w;
+	float4 Temp2			= LFVInsts[Offset + 10];
+	Data.Emissive			= Temp2.xyz;
+	Data.UniformScale		= Temp2.w;
 
 	return Data;
 }
@@ -129,8 +129,8 @@ struct FFogData
 	float Coverage;
 };
 
-// Height fog is evaluated according to spatially varying density defined as "Density * Exp[-HeightFalloff * z]"
-FFogData LocalFogVolumeEvaluateHeightIntegral(FLocalFogVolumeGPUInstanceData FogInstance, float3 CameraPos, float3 RayStart, float3 RayDir, float RayLength)
+// Height fog is evaluated according to spatially varying density defined as "Density * Exp[-HeightFogFalloff * z]", where density is in fact Extinction.
+FFogData LocalFogVolumeEvaluateAnalyticalIntegral(FLocalFogVolumeGPUInstanceData FogInstance, float3 CameraPos, float3 RayStart, float3 RayDir, float RayLength)
 {
 	FFogData FogData;
 
@@ -141,69 +141,54 @@ FFogData LocalFogVolumeEvaluateHeightIntegral(FLocalFogVolumeGPUInstanceData Fog
 
 	float RayLengthU = length(RayIntU - RayOriU);
 
-	float StartHeight = RayOriU.z - FogInstance.HeightOffset; // accoutn for the fog offset.
+	float StartHeight = RayOriU.z - FogInstance.HeightFogOffset; // accoutn for the fog offset.
 
 	float OpticalDepth = 0.0f;
-#if 0
-	// Original integral, less artefact
-	OpticalDepth = (FogInstance.Density / FogInstance.HeightFalloff) * exp(-StartHeight * FogInstance.HeightFalloff) * (1.0 - exp(-RayDirU.z * RayLengthU * FogInstance.HeightFalloff)) / RayDirU.z;
-#elif 1
-	// Reworked to avoid large value as input to exp
-	//OpticalDepth = (FogInstance.Density / FogInstance.HeightFalloff) * (exp(-StartHeight * FogInstance.HeightFalloff) - exp(-StartHeight * FogInstance.HeightFalloff - RayDir.z * RayLength * FogInstance.HeightFalloff)) / RayDir.z;
-
-	// But we use an updated formulation to avoid large input to exp leading to inf and nan.
-	float Factor0 = StartHeight * FogInstance.HeightFalloff;
-	float Factor1 = RayDirU.z * RayLengthU * FogInstance.HeightFalloff;
-	Factor0 = max(-80.0f, Factor0); // Simply clamp Factor0 to a valid range in order to not have exp explode.
-
-	OpticalDepth = (FogInstance.Density / (FogInstance.HeightFalloff * RayDirU.z)) * (exp(-Factor0) - exp(-(Factor0 + Factor1)));
-#endif
-
-	// Evaluate transmittance, accounting for the volume scale converted to local unit space meter.
-	float Transmittance = exp(-OpticalDepth * FogInstance.UniformScale * CENTIMETER_TO_METER);
-	// Assuming extinction is grey scale, we can compute a single coverage value from transmittance
-	FogData.Coverage = (1.0 - Transmittance);
-	// Assuming extinction==scattering, i.e. albedo=1, we can compute the integrale for each point on a ray while accouting for transmittance back to the view.
-	// It turns out that, under these conditions, the luminance and emissive factor is simply coverage. => see Mathematica details at the bottom of this page.
-	FogData.IntegratedLuminanceFactor = FogData.Coverage;
-
-	return FogData;
-}
-
 
-FFogData LocalFogVolumeEvaluateRadialIntegral(FLocalFogVolumeGPUInstanceData FogInstance, float3 CameraPos, float3 RayStart, float3 RayDir, float RayLength)
-{
-	FFogData FogData = (FFogData)0;
-
-	float3 RayOriU = mul(float4(RayStart, 1.0f), FogInstance.InvTransform).xyz;
-	float3 RayIntU = mul(float4(RayStart + RayDir * RayLength, 1.0f), FogInstance.InvTransform).xyz;
-	float3 RayDirU = mul(float4(RayDir, 0.0f), FogInstance.InvTransform).xyz;
-	RayDirU = normalize(RayDirU);
-
-	float RayLengthU = length(RayIntU - RayOriU);
+	if (FogInstance.RadialFogExtinction > 0.0f)
+	{
+		float3 SphereCenter = 0.0f;
+		float3 VolumeCenterToRayO = (RayOriU - SphereCenter);
 
-	float3 SphereCenter = 0.0f;
-	float3 VolumeCenterToRayO = (RayOriU - SphereCenter);
+		float b = dot(RayDirU, VolumeCenterToRayO);
+		float c = dot(VolumeCenterToRayO, VolumeCenterToRayO) - 1.0f;
+		float h = b * b - c;
 
-	float b = dot(RayDirU, VolumeCenterToRayO);
-	float c = dot(VolumeCenterToRayO, VolumeCenterToRayO) - 1.0f;
-	float h = b * b - c;
+		if (h >= 0.0)
+		{
+			h = sqrt(h);
+			float Length0 = -b - h;
+			float Length1 = -b + h;
+
+			Length0 = max(Length0, 0.0);
+			Length1 = max(Length1, 0.0);
+			Length1 = min(Length1, RayLengthU);
+
+			// Integral of a Extinction that descreases accordingf to square distance from the center of the volume.
+			float Integral0 = -(c * Length0 + b * Length0 * Length0 + Length0 * Length0 * Length0 / 3.0f);
+			float Integral1 = -(c * Length1 + b * Length1 * Length1 + Length1 * Length1 * Length1 / 3.0f);
+			OpticalDepth += max(0.0, FogInstance.RadialFogExtinction * (Integral1 - Integral0) * (3.0f / 4.0f));
+		}
+	}
 
-	float OpticalDepth = 0.0f;
-	if (h >= 0.0)
+	if (FogInstance.HeightFogExtinction > 0.0f)
 	{
-		h = sqrt(h);
-		float Length0 = -b - h;
-		float Length1 = -b + h;
-
-		Length0 = max(Length0, 0.0);
-		Length1 = max(Length1, 0.0);
-		Length1 = min(Length1, RayLengthU);
-
-		// Integral of a density that descreases accordingf to square distance from the center of the volume.
-		float Integral0 = -(c * Length0 + b * Length0 * Length0 + Length0 * Length0 * Length0 / 3.0f);
-		float Integral1 = -(c * Length1 + b * Length1 * Length1 + Length1 * Length1 * Length1 / 3.0f);
-		OpticalDepth = max(0.0, FogInstance.Density * (Integral1 - Integral0) * (3.0f / 4.0f));
+		float StartHeight = RayOriU.z - FogInstance.HeightFogOffset; // accoutn for the fog offset.
+
+	#if 0
+		// Original integral, less artefact
+		OpticalDepth += (FogInstance.HeightFogExtinction / FogInstance.HeightFogFalloff) * exp(-StartHeight * FogInstance.HeightFogFalloff) * (1.0 - exp(-RayDirU.z * RayLengthU * FogInstance.HeightFogFalloff)) / RayDirU.z;
+	#elif 1
+		// Reworked to avoid large value as input to exp
+		//OpticalDepth = (FogInstance.HeightFogExtinction / FogInstance.HeightFogFalloff) * (exp(-StartHeight * FogInstance.HeightFogFalloff) - exp(-StartHeight * FogInstance.HeightFogFalloff - RayDir.z * RayLength * FogInstance.HeightFogFalloff)) / RayDir.z;
+
+		// But we use an updated formulation to avoid large input to exp leading to inf and nan.
+		float Factor0 = StartHeight * FogInstance.HeightFogFalloff;
+		float Factor1 = RayDirU.z * RayLengthU * FogInstance.HeightFogFalloff;
+		Factor0 = max(-80.0f, Factor0); // Simply clamp Factor0 to a valid range in order to not have exp explode.
+
+		OpticalDepth += (FogInstance.HeightFogExtinction / (FogInstance.HeightFogFalloff * RayDirU.z)) * (exp(-Factor0) - exp(-(Factor0 + Factor1)));
+	#endif
 	}
 
 	// Evaluate transmittance, accounting for the volume scale converted to local unit space meter.
@@ -312,20 +297,7 @@ float4 GetLFVInstanceContribution(
 
 		if (RayTracedLength > 0.0)
 		{
-			FFogData FogData = (FFogData)0;
-
-			if (FogInstance.FogMode == 0)
-			{
-				FogData = LocalFogVolumeEvaluateHeightIntegral(FogInstance, CamRayWorldOrigin, CamRayWorldOrigin + CamRayWorldDir * Length0, CamRayWorldDir, RayTracedLength);
-			}
-			else // (FogInstance.FogMode == 1)
-			{
-				// TODO
-				// TODO LocalFogVolumeEvaluateRadialIntegral needs to be optimised because it is redoing lots of computations. Maybe move HeightFogIntegral to also be in a local normalised space.
-				// TODO It is like that for now while iterating with artists.
-				// TODO
-				FogData = LocalFogVolumeEvaluateRadialIntegral(FogInstance, CamRayWorldOrigin, CamRayWorldOrigin + CamRayWorldDir * Length0, CamRayWorldDir, RayTracedLength);
-			}
+			FFogData FogData = LocalFogVolumeEvaluateAnalyticalIntegral(FogInstance, CamRayWorldOrigin, CamRayWorldOrigin + CamRayWorldDir * Length0, CamRayWorldDir, RayTracedLength);
 
 			LuminanceColor = LocalFogVolumeEvaluateInScattering(FogInstance, FogData, CamRayWorldDir);
 			Transmittance = 1.0 - FogData.Coverage;
diff --git a/Engine/Shaders/Private/VolumetricFog.usf b/Engine/Shaders/Private/VolumetricFog.usf
--- a/Engine/Shaders/Private/VolumetricFog.usf
+++ b/Engine/Shaders/Private/VolumetricFog.usf
@@ -186,13 +186,13 @@ void MaterialSetupCS(
 		float SphereFade = saturate(1.0 - length(UnitSpacePos - float3(0, 0, 0)));
 
 		float LFVExtinction = 0.0;
-		if (FogInstance.FogMode == 0)
+		if (FogInstance.HeightFogExtinction > 0.0f)
 		{
-			LFVExtinction = SphereFade <= 0.0f ? 0.0f : FogInstance.Density * exp(-FogInstance.HeightFalloff * UnitSpacePos.z);
+			LFVExtinction += SphereFade <= 0.0f ? 0.0f : FogInstance.HeightFogExtinction * exp(-FogInstance.HeightFogFalloff * UnitSpacePos.z);
 		}
-		else // (FogInstance.FogMode == 1)
+		if(FogInstance.RadialFogExtinction > 0.0f)
 		{
-			LFVExtinction = FogInstance.Density * SphereFade;
+			LFVExtinction += FogInstance.RadialFogExtinction * SphereFade;
 		}
 
 		// Convert extinction from unit volume 1/m into extinction for volumetric fog as 1/cm
diff --git a/Engine/Source/Runtime/Engine/Classes/Components/LocalFogVolumeComponent.h b/Engine/Source/Runtime/Engine/Classes/Components/LocalFogVolumeComponent.h
--- a/Engine/Source/Runtime/Engine/Classes/Components/LocalFogVolumeComponent.h
+++ b/Engine/Source/Runtime/Engine/Classes/Components/LocalFogVolumeComponent.h
@@ -9,13 +9,6 @@
 
 class FLocalFogVolumeSceneProxy;
 
-UENUM()
-enum class ELocalFogMode : uint8
-{
-	LocalHeightFog = 0,
-	LocalSphereFog = 1,
-};
-
 UCLASS(ClassGroup = Rendering, collapsecategories, hidecategories = (Object, Mobility, Activation, "Components|Activation"), editinlinenew, meta = (BlueprintSpawnableComponent), MinimalAPI)
 class ULocalFogVolumeComponent : public USceneComponent
 {
@@ -23,32 +16,32 @@ class ULocalFogVolumeComponent : public USceneComponent
 
 	~ULocalFogVolumeComponent();
 
-	/** Controls the softness of the transition region when the volume is fading out. */
-	UPROPERTY(EditAnywhere, BlueprintReadOnly, Category = "Fog Mode")
-	ELocalFogMode FogMode = ELocalFogMode::LocalSphereFog;
+	/** The density of the radial fog representing its extinction coefficient at the center of the sphere. */
+	UPROPERTY(EditAnywhere, BlueprintReadOnly, interp, Category = "Radial Fog Distribution", meta = (DisplayName = "Radial Fog Density", UIMin = "0", UIMax = "2.0", SliderExponent = 2.0, ClampMin = 0.0))
+	float RadialFogExtinction = 1.0f;
 
-	/** Global density factor for this fog. */
-	UPROPERTY(EditAnywhere, BlueprintReadOnly, interp, Category = "Fog Distribution", meta = (UIMin = "0", UIMax = "10.0", SliderExponent = 2.0, ClampMin = 0.0))
-	float FogDensity = 1.0f;
+	/** The density of the radial fog representing its extinction coefficient at height 0 in the unit sphere. */
+	UPROPERTY(EditAnywhere, BlueprintReadOnly, interp, Category = "Height Fog Distribution", meta = (DisplayName = "Height Fog Density", UIMin = "0", UIMax = "2.0", SliderExponent = 2.0, ClampMin = 0.0))
+	float HeightFogExtinction = 0.0f;
 
 	/** Controls how the density decreases as height increases. Smaller values make the visible transition larger. */
-	UPROPERTY(EditAnywhere, BlueprintReadOnly, interp, Category = "Fog Distribution", meta = (UIMin = "0.001", UIMax = "5000", SliderExponent = 2.0, ClampMin = 0.001))
-	float FogHeightFalloff = 1000.0f;
+	UPROPERTY(EditAnywhere, BlueprintReadOnly, interp, Category = "Height Fog Distribution", meta = (UIMin = "0.001", UIMax = "5000", SliderExponent = 2.0, ClampMin = 0.001))
+	float HeightFogFalloff = 1000.0f;
 
 	/** Height offset, relative to the actor Z position. */
-	UPROPERTY(EditAnywhere, BlueprintReadOnly, interp, Category = "Fog Distribution", meta = (UIMin = "-2.0", UIMax = "2.0"))
-	float FogHeightOffset = 0.0f;
+	UPROPERTY(EditAnywhere, BlueprintReadOnly, interp, Category = "Height Fog Distribution", meta = (UIMin = "-2.0", UIMax = "2.0"))
+	float HeightFogOffset = 0.0f;
 
 	/** Controls the phase `G` parameter, describing the directionality of the scattering within this fog volume. */
-	UPROPERTY(EditAnywhere, BlueprintReadOnly, interp, Category = "Fog Shading", meta = (DisplayName = "Scattering Distribution", UIMin = "0.0", UIMax = "0.999", ClampMin = 0.0, ClampMax = 0.999))
+	UPROPERTY(EditAnywhere, BlueprintReadOnly, interp, Category = "Shading", meta = (DisplayName = "Scattering Distribution", UIMin = "0.0", UIMax = "0.999", ClampMin = 0.0, ClampMax = 0.999))
 	float FogPhaseG = 0.2f;
 
 	/** Controls the albedo of this fog volume. */
-	UPROPERTY(EditAnywhere, BlueprintReadOnly, interp, Category = "Fog Shading", meta = (ClampMin = 0.0, ClampMax = 1.0))
+	UPROPERTY(EditAnywhere, BlueprintReadOnly, interp, Category = "Shading", meta = (ClampMin = 0.0, ClampMax = 1.0))
 	FLinearColor FogAlbedo = FLinearColor::White;
 
 	/** Controls the emissive color of this fog volume. */
-	UPROPERTY(EditAnywhere, BlueprintReadOnly, interp, Category = "Fog Shading", meta = (ClampMin = 0.0, ClampMax = 1.0))
+	UPROPERTY(EditAnywhere, BlueprintReadOnly, interp, Category = "Shading", meta = (ClampMin = 0.0, ClampMax = 1.0))
 	FLinearColor FogEmissive = FLinearColor::Black;
 
 	/** The priority can be used as a way to override the sorting by distance. A lower value means the volume will be considered further away, i.e. it will draw behind the one with a higher priority value. */
diff --git a/Engine/Source/Runtime/Engine/Private/Components/LocalFogVolumeComponent.cpp b/Engine/Source/Runtime/Engine/Private/Components/LocalFogVolumeComponent.cpp
--- a/Engine/Source/Runtime/Engine/Private/Components/LocalFogVolumeComponent.cpp
+++ b/Engine/Source/Runtime/Engine/Private/Components/LocalFogVolumeComponent.cpp
@@ -35,10 +35,9 @@ void ULocalFogVolumeComponent::SendRenderTransformCommand()
 	if (LocalFogVolumeSceneProxy)
 	{
 		FTransform ComponentTransform = GetComponentTransform();
-		float HeightOffset = FogHeightOffset;
 		FLocalFogVolumeSceneProxy* SceneProxy = LocalFogVolumeSceneProxy;
 		ENQUEUE_RENDER_COMMAND(FUpdateLocalFogVolumeSceneProxyTransformCommand)(
-			[SceneProxy, ComponentTransform, HeightOffset](FRHICommandList& RHICmdList)
+			[SceneProxy, ComponentTransform](FRHICommandList& RHICmdList)
 			{
 				// Nothing else is needed so that command could actually go.
 				SceneProxy->UpdateComponentTransform(ComponentTransform);
@@ -99,10 +98,10 @@ bool ULocalFogVolumeComponent::CanEditChange(const FProperty* InProperty) const
 	{
 		FString PropertyName = InProperty->GetName();
 
-		if (FogMode != ELocalFogMode::LocalHeightFog)
+		if (HeightFogExtinction <= 0.0f)
 		{
-			if (PropertyName == GET_MEMBER_NAME_STRING_CHECKED(ULocalFogVolumeComponent, FogHeightFalloff)
-				|| PropertyName == GET_MEMBER_NAME_STRING_CHECKED(ULocalFogVolumeComponent, FogHeightOffset))
+			if (PropertyName == GET_MEMBER_NAME_STRING_CHECKED(ULocalFogVolumeComponent, HeightFogFalloff)
+				|| PropertyName == GET_MEMBER_NAME_STRING_CHECKED(ULocalFogVolumeComponent, HeightFogOffset))
 			{
 				return false;
 			}
diff --git a/Engine/Source/Runtime/Engine/Private/LocalFogVolumeSceneProxy.cpp b/Engine/Source/Runtime/Engine/Private/LocalFogVolumeSceneProxy.cpp
--- a/Engine/Source/Runtime/Engine/Private/LocalFogVolumeSceneProxy.cpp
+++ b/Engine/Source/Runtime/Engine/Private/LocalFogVolumeSceneProxy.cpp
@@ -10,11 +10,11 @@
 
 
 FLocalFogVolumeSceneProxy::FLocalFogVolumeSceneProxy(const ULocalFogVolumeComponent* InComponent)
-	: FogDensity(InComponent->FogDensity)
-	, FogHeightFalloff(InComponent->FogHeightFalloff)
-	, FogHeightOffset(InComponent->FogHeightOffset)
+	: RadialFogExtinction(InComponent->RadialFogExtinction)
+	, HeightFogExtinction(InComponent->HeightFogExtinction)
+	, HeightFogFalloff(InComponent->HeightFogFalloff)
+	, HeightFogOffset(InComponent->HeightFogOffset)
 	, FogUniformScale(1.0f)
-	, FogMode((uint8)InComponent->FogMode)
 	, FogSortPriority(uint8(127 - int8(InComponent->FogSortPriority))) // FogSortPriority on the component is in [-127,127] and needs to be negated to match expected priority behavior.
 	, FogPhaseG(InComponent->FogPhaseG)
 	, FogAlbedo(InComponent->FogAlbedo)
diff --git a/Engine/Source/Runtime/Engine/Public/LocalFogVolumeSceneProxy.h b/Engine/Source/Runtime/Engine/Public/LocalFogVolumeSceneProxy.h
--- a/Engine/Source/Runtime/Engine/Public/LocalFogVolumeSceneProxy.h
+++ b/Engine/Source/Runtime/Engine/Public/LocalFogVolumeSceneProxy.h
@@ -20,12 +20,12 @@ public:
 
 	FTransform FogTransform;
 
-	float FogDensity;
-	float FogHeightFalloff;
-	float FogHeightOffset;
+	float RadialFogExtinction;
+	float HeightFogExtinction;
+	float HeightFogFalloff;
+	float HeightFogOffset;
 	float FogUniformScale;
 
-	uint8 FogMode;
 	uint8 FogSortPriority;
 
 	float FogPhaseG;
diff --git a/Engine/Source/Runtime/Renderer/Private/LocalFogVolumeRendering.cpp b/Engine/Source/Runtime/Renderer/Private/LocalFogVolumeRendering.cpp
--- a/Engine/Source/Runtime/Renderer/Private/LocalFogVolumeRendering.cpp
+++ b/Engine/Source/Runtime/Renderer/Private/LocalFogVolumeRendering.cpp
@@ -292,7 +292,7 @@ void GetLocalFogVolumeSortingData(const FScene* Scene, FRDGBuilder& GraphBuilder
 	Out.LocalFogVolumeSortKeys.SetNumUninitialized(Out.LocalFogVolumeInstanceCount);
 	for (FLocalFogVolumeSceneProxy* LHF : Scene->LocalFogVolumes)
 	{
-		if (LHF->FogDensity <= 0.0f)
+		if (LHF->RadialFogExtinction <= 0.0f && LHF->HeightFogExtinction <= 0.0f)
 		{
 			continue; // this volume will never be visible
 		}
@@ -301,12 +301,12 @@ void GetLocalFogVolumeSortingData(const FScene* Scene, FRDGBuilder& GraphBuilder
 		LocalFogVolumeGPUInstanceDataIt->Transform = FMatrix44f(LHF->FogTransform.ToMatrixWithScale());
 		LocalFogVolumeGPUInstanceDataIt->InvTransform = LocalFogVolumeGPUInstanceDataIt->Transform.Inverse();
 
-		LocalFogVolumeGPUInstanceDataIt->Density = LHF->FogDensity;
-		LocalFogVolumeGPUInstanceDataIt->HeightFalloff = LHF->FogHeightFalloff * 0.01f;	// This scale is used to have artist author reasonable range.
-		LocalFogVolumeGPUInstanceDataIt->HeightOffset = LHF->FogHeightOffset;
-		LocalFogVolumeGPUInstanceDataIt->UniformScale = LHF->FogUniformScale;
+		LocalFogVolumeGPUInstanceDataIt->RadialFogExtinction = LHF->RadialFogExtinction;
+		LocalFogVolumeGPUInstanceDataIt->HeightFogExtinction = LHF->HeightFogExtinction;
+		LocalFogVolumeGPUInstanceDataIt->HeightFogFalloff = LHF->HeightFogFalloff * 0.01f;	// This scale is used to have artist author reasonable range.
+		LocalFogVolumeGPUInstanceDataIt->HeightFogOffset = LHF->HeightFogOffset;
 
-		LocalFogVolumeGPUInstanceDataIt->FogMode = float(LHF->FogMode);
+		LocalFogVolumeGPUInstanceDataIt->UniformScale = LHF->FogUniformScale;
 
 		LocalFogVolumeGPUInstanceDataIt->Albedo = FVector3f(LHF->FogAlbedo);
 		LocalFogVolumeGPUInstanceDataIt->PhaseG = LHF->FogPhaseG;
diff --git a/Engine/Source/Runtime/Renderer/Private/LocalFogVolumeRendering.h b/Engine/Source/Runtime/Renderer/Private/LocalFogVolumeRendering.h
--- a/Engine/Source/Runtime/Renderer/Private/LocalFogVolumeRendering.h
+++ b/Engine/Source/Runtime/Renderer/Private/LocalFogVolumeRendering.h
@@ -45,15 +45,15 @@ public:
 	FMatrix44f Transform;
 	FMatrix44f InvTransform;
 
-	float Density;
-	float HeightFalloff;
-	float HeightOffset;
-	float UniformScale;
+	float RadialFogExtinction;
+	float HeightFogExtinction;
+	float HeightFogFalloff;
+	float HeightFogOffset;
 
 	FVector3f Albedo;
 	float PhaseG;
 	FVector3f Emissive;
-	float FogMode;
+	float UniformScale;
 };
 
 class RENDERER_API FLocalFogVolumeSortKey
