LFV - fix for height fog offset not respected by VFog.

#prelfight none

diff --git a/Engine/Shaders/Private/VolumetricFog.usf b/Engine/Shaders/Private/VolumetricFog.usf
--- a/Engine/Shaders/Private/VolumetricFog.usf
+++ b/Engine/Shaders/Private/VolumetricFog.usf
@@ -188,7 +188,7 @@ void MaterialSetupCS(
 		float LFVExtinction = 0.0;
 		if (FogInstance.HeightFogExtinction > 0.0f)
 		{
-			LFVExtinction += SphereFade <= 0.0f ? 0.0f : FogInstance.HeightFogExtinction * exp(-FogInstance.HeightFogFalloff * UnitSpacePos.z);
+			LFVExtinction += SphereFade <= 0.0f ? 0.0f : FogInstance.HeightFogExtinction * exp(-FogInstance.HeightFogFalloff * (UnitSpacePos.z - FogInstance.HeightFogOffset));
 		}
 		if(FogInstance.RadialFogExtinction > 0.0f)
 		{
