LFV - fix for LFV not showing on clouds when fog is not present in the scene.

#prelfight none

diff --git a/Engine/Shaders/Private/VolumetricCloud.usf b/Engine/Shaders/Private/VolumetricCloud.usf
--- a/Engine/Shaders/Private/VolumetricCloud.usf
+++ b/Engine/Shaders/Private/VolumetricCloud.usf
@@ -1456,9 +1456,14 @@ void MainCommon(in FMaterialPixelParameters MaterialParameters, in float4 SvPosi
 	//
 	// Apply height fog if needed
 	//
-	if (RenderVolumetricCloudParameters.EnableHeightFog && tAPWeightsSum > 0.0f)
+	if (tAPWeightsSum > 0.0f)
 	{
-		float4 HeightFogInscatteringAndTransmittance = CalculateHeightFog(FogSampleWorldPositionRelativeToCameraCm);
+		float4 HeightFogInscatteringAndTransmittance = float4(0, 0, 0, 1);
+
+		if (RenderVolumetricCloudParameters.EnableHeightFog)
+		{
+			HeightFogInscatteringAndTransmittance = CalculateHeightFog(FogSampleWorldPositionRelativeToCameraCm);
+		}
 
 	#if PERMUTATION_SUPPORT_LOCAL_FOG_VOLUME
 		uint2 TilePos = uint2(SvPosition.xy * RenderVolumetricCloudParameters.TracingCoordToFullResPixelCoordScaleBias.xy - View.ViewRectMin.xy) / LFVTilePixelSize.xx;
@@ -1528,9 +1533,14 @@ void MainCommon(in FMaterialPixelParameters MaterialParameters, in float4 SvPosi
 		//
 		// Apply height fog if needed
 		//
-		if (RenderVolumetricCloudParameters.EnableHeightFog && ClostestTAPWeightsSum > 0.0f)
+		if (ClostestTAPWeightsSum > 0.0f)
 		{
-			float4 HeightFogInscatteringAndTransmittance = CalculateHeightFog(FogSampleWorldPositionRelativeToCameraCm);
+			float4 HeightFogInscatteringAndTransmittance = float4(0, 0, 0, 1);
+
+			if (RenderVolumetricCloudParameters.EnableHeightFog)
+			{
+				HeightFogInscatteringAndTransmittance = CalculateHeightFog(FogSampleWorldPositionRelativeToCameraCm);
+			}
 
 		#if PERMUTATION_SUPPORT_LOCAL_FOG_VOLUME
 			uint2 TilePos = uint2(SvPosition.xy * RenderVolumetricCloudParameters.TracingCoordToFullResPixelCoordScaleBias.xy - View.ViewRectMin.xy) / LFVTilePixelSize.xx;
