Fix manual skin cache for groom.

diff --git a/Engine/Plugins/Runtime/HairStrands/Source/HairStrandsCore/Private/HairStrandsMeshProjection.cpp b/Engine/Plugins/Runtime/HairStrands/Source/HairStrandsCore/Private/HairStrandsMeshProjection.cpp
--- a/Engine/Plugins/Runtime/HairStrands/Source/HairStrandsCore/Private/HairStrandsMeshProjection.cpp
+++ b/Engine/Plugins/Runtime/HairStrands/Source/HairStrandsCore/Private/HairStrandsMeshProjection.cpp
@@ -84,6 +84,7 @@ class FSkinUpdateCS : public FGlobalShader
 	using FPermutationDomain = TShaderPermutationDomain<FUnlimitedBoneInfluence, FUseExtraInfluence, FIndexUint16, FPrevious>;
 
 	BEGIN_SHADER_PARAMETER_STRUCT(FParameters, )
+		SHADER_PARAMETER(uint32, NumVertexToProcess)
 		SHADER_PARAMETER(uint32, NumTotalVertices)
 		SHADER_PARAMETER(uint32, SectionVertexBaseIndex)
 		SHADER_PARAMETER(uint32, IndexSize)
@@ -127,6 +128,7 @@ void AddSkinUpdatePass(
 	
 	FSkinUpdateCS::FParameters* Parameters = GraphBuilder.AllocParameters<FSkinUpdateCS::FParameters>();
 	Parameters->IndexSize = SkinWeight->GetBoneIndexByteSize();
+	Parameters->NumVertexToProcess = NumVertexToProcess;
 	Parameters->NumTotalVertices = NumTotalVertices;
 	Parameters->SectionVertexBaseIndex = SectionVertexBaseIndex;
 	Parameters->WeightStride = SkinWeight->GetConstantInfluencesVertexStride();
diff --git a/Engine/Shaders/Private/HairStrands/HairStrandsSkinUpdate.usf b/Engine/Shaders/Private/HairStrands/HairStrandsSkinUpdate.usf
--- a/Engine/Shaders/Private/HairStrands/HairStrandsSkinUpdate.usf
+++ b/Engine/Shaders/Private/HairStrands/HairStrandsSkinUpdate.usf
@@ -9,6 +9,7 @@
 uint IndexSize; // bone index size: 1 - 8bit or 2 - 16bit
 uint SectionVertexBaseIndex;
 uint NumTotalVertices;
+uint NumVertexToProcess;
 uint WeightStride;
 uint BonesOffset;
 
@@ -53,7 +54,7 @@ struct FUnpackedWeight
 
 FBlendedMatrix BlendMatrix(FUnpackedWeight UnpackedWeight, uint VertexIndex, Buffer<float4> InBoneMatrices)
 {
-	const uint MatrixOffset = BonesOffset;
+	const uint MatrixOffset = 0;
 #if GPUSKIN_UNLIMITED_BONE_INFLUENCE
 	uint BlendOffsetCount = WeightLookup[VertexIndex];
 	int NumBoneInfluences = BlendOffsetCount & 0xff;
@@ -79,7 +80,8 @@ FBlendedMatrix BlendMatrix(FUnpackedWeight UnpackedWeight, uint VertexIndex, Buf
 	uint IndexOffsetZ = (MatrixOffset + UnpackedWeight.BlendIndices.z) * 3;
 	uint IndexOffsetW = (MatrixOffset + UnpackedWeight.BlendIndices.w) * 3;
 
-	FBlendedMatrix OutBoneMatrix = UnpackedWeight.BlendWeights.x * FBlendedMatrix(InBoneMatrices[IndexOffsetX], InBoneMatrices[IndexOffsetX+1], InBoneMatrices[IndexOffsetX+2]);
+	FBlendedMatrix 
+	OutBoneMatrix  = UnpackedWeight.BlendWeights.x * FBlendedMatrix(InBoneMatrices[IndexOffsetX], InBoneMatrices[IndexOffsetX+1], InBoneMatrices[IndexOffsetX+2]);
 	OutBoneMatrix += UnpackedWeight.BlendWeights.y * FBlendedMatrix(InBoneMatrices[IndexOffsetY], InBoneMatrices[IndexOffsetY+1], InBoneMatrices[IndexOffsetY+2]);
 	OutBoneMatrix += UnpackedWeight.BlendWeights.z * FBlendedMatrix(InBoneMatrices[IndexOffsetZ], InBoneMatrices[IndexOffsetZ+1], InBoneMatrices[IndexOffsetZ+2]);
 	OutBoneMatrix += UnpackedWeight.BlendWeights.w * FBlendedMatrix(InBoneMatrices[IndexOffsetW], InBoneMatrices[IndexOffsetW+1], InBoneMatrices[IndexOffsetW+2]);
@@ -158,7 +160,7 @@ void UpdateSkinPositionCS(uint3 GroupID : SV_GroupID,
  uint3 GroupThreadID : SV_GroupThreadID)
 {
 	const uint VertexIndex = DispatchThreadID.x + SectionVertexBaseIndex;
-	if (VertexIndex < NumTotalVertices)
+	if (DispatchThreadID.x < NumVertexToProcess && VertexIndex < NumTotalVertices)
 	{
 		const uint PositionOffset = VertexIndex * 3;
 		const uint WeightOffset = VertexIndex * WeightStride / 4;
