Fix sharp edges of spotlights when using mobile local lights buffer

diff --git a/Engine/Shaders/Private/MobileLocalLightsBuffer.usf b/Engine/Shaders/Private/MobileLocalLightsBuffer.usf
--- a/Engine/Shaders/Private/MobileLocalLightsBuffer.usf
+++ b/Engine/Shaders/Private/MobileLocalLightsBuffer.usf
@@ -143,6 +143,10 @@ void Main(float4 SvPosition : SV_POSITION
 		}
 
 		float WeightLight = 1 - DistLight;
+		if (LightData.bSpotLight)
+		{
+			WeightLight *= SpotAttenuation(L, -LightData.Direction, LightData.SpotAngles);
+		}
 		
 		if( LightMask > 0)
 		{
