LFV - using 4x3 matrices instead of 4x4 to read less data (reducing instance bytes from 176 to 144 bytes)

#prelfight https://horde.devtools.epicgames.com/job/6513fbbb196f1bf3bc01e683

diff --git a/Engine/Shaders/Private/LocalFogVolumes/LocalFogVolumeCommon.ush b/Engine/Shaders/Private/LocalFogVolumes/LocalFogVolumeCommon.ush
--- a/Engine/Shaders/Private/LocalFogVolumes/LocalFogVolumeCommon.ush
+++ b/Engine/Shaders/Private/LocalFogVolumes/LocalFogVolumeCommon.ush
@@ -12,8 +12,8 @@
 
 struct FLocalFogVolumeGPUInstanceData
 {
-	float4x4 Transform;
-	float4x4 InvTransform;
+	float4x3 Transform;
+	float4x3 InvTransform;
 
 	float RadialFogExtinction;
 	float HeightFogExtinction;
@@ -76,22 +76,22 @@ Texture2DArray<uint>					LFV_LocalFogVolumeTileDataTexture;
 FLocalFogVolumeGPUInstanceData GetLocalFogVolumeGPUInstanceData(uint Index)
 {
 	FLocalFogVolumeGPUInstanceData Data;
-	uint Offset = Index * 11;
+	uint Offset = Index * 9;
 
-	Data.Transform			= float4x4(LFVInsts[Offset + 0 ], LFVInsts[Offset + 1 ], LFVInsts[Offset + 2 ], LFVInsts[Offset + 3 ]);
-	Data.InvTransform		= float4x4(LFVInsts[Offset + 4 ], LFVInsts[Offset + 5 ], LFVInsts[Offset + 6 ], LFVInsts[Offset + 7 ]);
+	Data.Transform			= float4x3(LFVInsts[Offset + 0 ], LFVInsts[Offset + 1 ], LFVInsts[Offset + 2 ]); // reads 4x3 row major pack as 3 float4
+	Data.InvTransform		= float4x3(LFVInsts[Offset + 3 ], LFVInsts[Offset + 4 ], LFVInsts[Offset + 5 ]);
 
-	float4 Temp0			= LFVInsts[Offset + 8];
+	float4 Temp0			= LFVInsts[Offset + 6];
 	Data.RadialFogExtinction= Temp0.x;
 	Data.HeightFogExtinction= Temp0.y;
 	Data.HeightFogFalloff	= Temp0.z;
 	Data.HeightFogOffset	= Temp0.w;
 
-	float4 Temp1			= LFVInsts[Offset + 9];
+	float4 Temp1			= LFVInsts[Offset + 7];
 	Data.Albedo				= Temp1.xyz;
 	Data.PhaseG				= Temp1.w;
 
-	float4 Temp2			= LFVInsts[Offset + 10];
+	float4 Temp2			= LFVInsts[Offset + 8];
 	Data.Emissive			= Temp2.xyz;
 	Data.UniformScale		= Temp2.w;
 
diff --git a/Engine/Shaders/Private/LocalFogVolumes/LocalFogVolumeSplat.usf b/Engine/Shaders/Private/LocalFogVolumes/LocalFogVolumeSplat.usf
--- a/Engine/Shaders/Private/LocalFogVolumes/LocalFogVolumeSplat.usf
+++ b/Engine/Shaders/Private/LocalFogVolumes/LocalFogVolumeSplat.usf
@@ -58,11 +58,9 @@ void LocalFogVolumeSplatVS(
 
 	FLocalFogVolumeGPUInstanceData FogInstance = GetLocalFogVolumeGPUInstanceData(InstanceIndex);
 
-	float4x4 FogTransform = FogInstance.Transform;
+	float3 WorldPosition = mul(float4(LocalPosition.xyz, 1.0f), FogInstance.Transform).xyz;
 
-	float4 WorldPosition = mul(float4(LocalPosition.xyz, 1.0f), FogTransform);
-
-	float4 TranslatedWorldPosition = float4(WorldPosition.xyz + LWCHackToFloat(PrimaryView.PreViewTranslation), 1);
+	float4 TranslatedWorldPosition = float4(WorldPosition + LWCHackToFloat(PrimaryView.PreViewTranslation), 1);
 
 	OutPosition = mul(TranslatedWorldPosition, PrimaryView.TranslatedWorldToClip);
 	OtherVertexOutput.InstanceId = InstanceIndex;
diff --git a/Engine/Source/Runtime/Renderer/Private/LocalFogVolumeRendering.cpp b/Engine/Source/Runtime/Renderer/Private/LocalFogVolumeRendering.cpp
--- a/Engine/Source/Runtime/Renderer/Private/LocalFogVolumeRendering.cpp
+++ b/Engine/Source/Runtime/Renderer/Private/LocalFogVolumeRendering.cpp
@@ -298,8 +298,28 @@ void GetLocalFogVolumeSortingData(const FScene* Scene, FRDGBuilder& GraphBuilder
 		}
 
 		FLocalFogVolumeGPUInstanceData* LocalFogVolumeGPUInstanceDataIt = &Out.LocalFogVolumeGPUInstanceData[Out.LocalFogVolumeInstanceCountFinal];
-		LocalFogVolumeGPUInstanceDataIt->Transform = FMatrix44f(LHF->FogTransform.ToMatrixWithScale());
-		LocalFogVolumeGPUInstanceDataIt->InvTransform = LocalFogVolumeGPUInstanceDataIt->Transform.Inverse();
+
+		auto ConvertFromMatrix44fTo4x3Array = [](FMatrix44f& InMat, float* OutArr)
+			{
+				// Row major order
+				OutArr[0 ] = InMat.M[0][0];
+				OutArr[1 ] = InMat.M[0][1];
+				OutArr[2 ] = InMat.M[0][2];
+				OutArr[3 ] = InMat.M[1][0];
+				OutArr[4 ] = InMat.M[1][1];
+				OutArr[5 ] = InMat.M[1][2];
+				OutArr[6 ] = InMat.M[2][0];
+				OutArr[7 ] = InMat.M[2][1];
+				OutArr[8 ] = InMat.M[2][2];
+				OutArr[9 ] = InMat.M[3][0];
+				OutArr[10] = InMat.M[3][1];
+				OutArr[11] = InMat.M[3][2];
+			};
+
+		FMatrix44f Transform    = FMatrix44f(LHF->FogTransform.ToMatrixWithScale());
+		FMatrix44f InvTransform = Transform.Inverse();
+		ConvertFromMatrix44fTo4x3Array(Transform, LocalFogVolumeGPUInstanceDataIt->Transform);
+		ConvertFromMatrix44fTo4x3Array(InvTransform, LocalFogVolumeGPUInstanceDataIt->InvTransform);
 
 		LocalFogVolumeGPUInstanceDataIt->RadialFogExtinction = LHF->RadialFogExtinction;
 		LocalFogVolumeGPUInstanceDataIt->HeightFogExtinction = LHF->HeightFogExtinction;
diff --git a/Engine/Source/Runtime/Renderer/Private/LocalFogVolumeRendering.h b/Engine/Source/Runtime/Renderer/Private/LocalFogVolumeRendering.h
--- a/Engine/Source/Runtime/Renderer/Private/LocalFogVolumeRendering.h
+++ b/Engine/Source/Runtime/Renderer/Private/LocalFogVolumeRendering.h
@@ -42,8 +42,9 @@ END_SHADER_PARAMETER_STRUCT()
 class FLocalFogVolumeGPUInstanceData
 {
 public:
-	FMatrix44f Transform;
-	FMatrix44f InvTransform;
+
+	float Transform[4 * 3];
+	float InvTransform[4 * 3];
 
 	float RadialFogExtinction;
 	float HeightFogExtinction;
