Added support for directional light in scattering for mobile volumetricfog.

diff --git a/Engine/Shaders/Private/VolumetricFog.usf b/Engine/Shaders/Private/VolumetricFog.usf
--- a/Engine/Shaders/Private/VolumetricFog.usf
+++ b/Engine/Shaders/Private/VolumetricFog.usf
@@ -662,6 +662,12 @@ float2 UseHeightFogColors;	// x=override directional light using height fog insc
 float UseDirectionalLightShadowing;
 float StaticLightingScatteringIntensity;
 
+#if SHADING_PATH_MOBILE
+uint   MobileHasDirectionalLight;
+float3 MobileDirectionalLightColor;
+float3 MobileDirectionalLightDirection;
+#endif
+
 Texture3D<float> RaytracedShadowsVolume;
 
 [numthreads(THREADGROUP_SIZE_X, THREADGROUP_SIZE_Y, THREADGROUP_SIZE_Z)]
@@ -731,9 +737,18 @@ void LightScatteringCS(
 		float3 WorldPosition = TranslatedWorldPosition - LWCHackToFloat(PrimaryView.PreViewTranslation); // LWC_TODO
 		float CameraVectorLength = length(TranslatedWorldPosition - PrimaryView.TranslatedWorldCameraOrigin);
 		float3 CameraVector = (TranslatedWorldPosition - PrimaryView.TranslatedWorldCameraOrigin) / CameraVectorLength;
-
+		
+	#if SHADING_PATH_MOBILE
+		const bool	 bHasDirectionalLight		= MobileHasDirectionalLight > 0;
+		const float3 InDirectionalLightColor	= MobileDirectionalLightColor;	// DirectionalLightVolumetricScatteringIntensity is baked in MobileDirectionalLightColor
+		const float3 InDirectionalLightDirection= MobileDirectionalLightDirection;
+	#else
+		const bool	 bHasDirectionalLight		= ForwardLightData.HasDirectionalLight;
+		const float3 InDirectionalLightColor	= ForwardLightData.DirectionalLightColor * ForwardLightData.DirectionalLightVolumetricScatteringIntensity;
+		const float3 InDirectionalLightDirection= ForwardLightData.DirectionalLightDirection;
+	#endif
 		BRANCH
-		if (ForwardLightData.HasDirectionalLight)
+		if (bHasDirectionalLight)
 		{
 			float ShadowFactor = 1;
 			
@@ -760,18 +775,15 @@ void LightScatteringCS(
 
 			ShadowFactor *= GetDirectionalLightFunction(TranslatedWorldPosition);
 
-			float3 DirectionalLightColor = ForwardLightData.DirectionalLightColor;
+			float3 DirectionalLightColor = InDirectionalLightColor;
 
 			if (UseHeightFogColors.x > 0)
 			{
 				// Attempt to maintain intensity ratio between sky and sun
-				DirectionalLightColor = VolumetricFog.HeightFogDirectionalLightInscatteringColor * Luminance(ForwardLightData.DirectionalLightColor);
+				DirectionalLightColor = VolumetricFog.HeightFogDirectionalLightInscatteringColor * Luminance(InDirectionalLightColor);
 			}
 
-			LightScattering += DirectionalLightColor
-				* (ShadowFactor 
-				* ForwardLightData.DirectionalLightVolumetricScatteringIntensity 
-				* PhaseFunction(PhaseG, dot(ForwardLightData.DirectionalLightDirection, -CameraVector)));
+			LightScattering += DirectionalLightColor * ShadowFactor * PhaseFunction(PhaseG, dot(InDirectionalLightDirection, -CameraVector));
 		}
 
 		FTwoBandSHVector RotatedHGZonalHarmonic;
diff --git a/Engine/Source/Runtime/Renderer/Private/VolumetricFog.cpp b/Engine/Source/Runtime/Renderer/Private/VolumetricFog.cpp
--- a/Engine/Source/Runtime/Renderer/Private/VolumetricFog.cpp
+++ b/Engine/Source/Runtime/Renderer/Private/VolumetricFog.cpp
@@ -988,6 +988,8 @@ class FVolumetricFogLightScatteringCS : public FGlobalShader
 		SHADER_PARAMETER(uint32, SampleSkyLightDiffuseEnvMap)
 		SHADER_PARAMETER(FMatrix44f, DirectionalLightFunctionTranslatedWorldToShadow)
 		SHADER_PARAMETER(FMatrix44f, CloudShadowmapTranslatedWorldToLightClipMatrix)
+		SHADER_PARAMETER(FVector3f, MobileDirectionalLightColor)
+		SHADER_PARAMETER(FVector3f, MobileDirectionalLightDirection)
 		SHADER_PARAMETER(FVector2f, PrevConservativeDepthTextureSize)
 		SHADER_PARAMETER(FVector2f, UseHeightFogColors)
 		SHADER_PARAMETER(FVector2f, LightScatteringHistoryPreExposureAndInv)
@@ -1001,6 +1003,7 @@ class FVolumetricFogLightScatteringCS : public FGlobalShader
 		SHADER_PARAMETER(float, UseDirectionalLightShadowing)
 		SHADER_PARAMETER(uint32, UseConservativeDepthTexture)
 		SHADER_PARAMETER(uint32, UseEmissive)
+		SHADER_PARAMETER(uint32, MobileHasDirectionalLight)
 	END_SHADER_PARAMETER_STRUCT()
 
 	static FIntVector GetGroupSize()
@@ -1571,6 +1574,23 @@ void FSceneRenderer::ComputeVolumetricFog(FRDGBuilder& GraphBuilder,
 				PassParameters->SampleSkyLightDiffuseEnvMap = 0;
 			}
 
+			// Mobile handles directional differently as of today to handle light masking (does not use and fill up the FForwardLightData). 
+			// Volumetric fog does not work with light mask so we simply pick up the first one available. In the long run we might want something more common.
+			PassParameters->MobileDirectionalLightColor		= FVector3f::Zero();
+			PassParameters->MobileDirectionalLightDirection = FVector3f::Zero();
+			PassParameters->MobileHasDirectionalLight		= 0;
+			for (uint32 ChannelIdx = 0; ChannelIdx < UE_ARRAY_COUNT(Scene->MobileDirectionalLights); ChannelIdx++)
+			{
+				FLightSceneInfo* Light = Scene->MobileDirectionalLights[ChannelIdx];
+				if (Light != nullptr)
+				{
+					PassParameters->MobileDirectionalLightColor		= FVector3f(Light->Proxy->GetSunIlluminanceAccountingForSkyAtmospherePerPixelTransmittance() * Light->Proxy->GetVolumetricScatteringIntensity());
+					PassParameters->MobileDirectionalLightDirection = FVector3f(-Light->Proxy->GetDirection());
+					PassParameters->MobileHasDirectionalLight		= 1;
+					break;
+				}
+			}
+
 			float StaticLightingScatteringIntensityValue = 0;
 			if (View.Family->EngineShowFlags.GlobalIllumination && View.Family->EngineShowFlags.VolumetricLightmap)
 			{
