Use correct instance bounds for vertex factories with UBO style GPUScene.
Was using the primitive bounds instead.

diff --git a/Engine/Shaders/Private/SceneData.ush b/Engine/Shaders/Private/SceneData.ush
--- a/Engine/Shaders/Private/SceneData.ush
+++ b/Engine/Shaders/Private/SceneData.ush
@@ -1124,8 +1124,8 @@ FSceneDataIntermediates GetSceneDataIntermediates(uint InstanceIdOffset, uint Dr
 		Intermediates.InstanceData.InvNonUniformScale = Intermediates.Primitive.InvNonUniformScale;
 		Intermediates.InstanceData.DeterminantSign = GetPrimitive_DeterminantSign_FromFlags(Intermediates.InstanceData.Flags);
 	#endif
-	Intermediates.InstanceData.LocalBoundsCenter = (Intermediates.Primitive.LocalObjectBoundsMax + Intermediates.Primitive.LocalObjectBoundsMin) * 0.5f;
-	Intermediates.InstanceData.LocalBoundsExtent = (Intermediates.Primitive.LocalObjectBoundsMax - Intermediates.Primitive.LocalObjectBoundsMin) * 0.5f;
+	Intermediates.InstanceData.LocalBoundsCenter = Intermediates.Primitive.InstanceLocalBoundsCenter;
+	Intermediates.InstanceData.LocalBoundsExtent = Intermediates.Primitive.InstanceLocalBoundsExtent;
 	Intermediates.InstanceData.ValidInstance = true;
 	Intermediates.CullingFlags = INSTANCE_CULLING_FLAGS_DEFAULT;
 	Intermediates.ViewIndex = 0U;
