Substrate - Fixed some dx11 compilation issues.
Also glint evasluation code is now gone when glint is turned off for the project/platform (deffered and forward).

#prelfight none

diff --git a/Engine/Shaders/Private/Substrate/SubstrateEvaluation.ush b/Engine/Shaders/Private/Substrate/SubstrateEvaluation.ush
--- a/Engine/Shaders/Private/Substrate/SubstrateEvaluation.ush
+++ b/Engine/Shaders/Private/Substrate/SubstrateEvaluation.ush
@@ -429,8 +429,8 @@ FSubstrateEvaluateResult SubstrateEvaluateBSDFCommon(FSubstrateBSDFContext BSDFC
 						float  Alpha  = Square(SafeRoughness);
 						float2 AlphaXY = 0;
 						GetAnisotropicRoughness(Alpha, SLAB_ANISOTROPY(BSDFContext.BSDF), AlphaXY.x, AlphaXY.y);
-
-					#if SUBSTRATE_COMPLEXSPECIALPATH && SUBSTRATE_GLINTS_ALLOWED
+				
+					#if SUBSTRATE_COMPLEXSPECIALPATH && SUBSTRATE_GLINTS_ALLOWED && SUBSTRATE_GLINTS_ENABLED
 						// Glint shading only in the complex path
 						BRANCH
 						if (BSDF_GETHASGLINT(BSDFContext.BSDF))
@@ -463,7 +463,7 @@ FSubstrateEvaluateResult SubstrateEvaluateBSDFCommon(FSubstrateBSDFContext BSDFC
 							RectLightSpec = RectGGXApproxLTC(SafeRoughness, F0, BSDFContext.N, BSDFContext.V, AreaLightContext.AreaLight.Rect, AreaLightContext.AreaLight.Texture, MeanLightWorldDirection);
 							
 							// Now combine the rectangular area light with Glints.
-						#if SUBSTRATE_COMPLEXSPECIALPATH && SUBSTRATE_GLINTS_ALLOWED
+						#if SUBSTRATE_COMPLEXSPECIALPATH && SUBSTRATE_GLINTS_ALLOWED && SUBSTRATE_GLINTS_ENABLED
 							BRANCH
 							if (BSDF_GETHASGLINT(BSDFContext.BSDF))
 							{
@@ -489,7 +489,7 @@ FSubstrateEvaluateResult SubstrateEvaluateBSDFCommon(FSubstrateBSDFContext BSDFC
 						{
 							if(IntegrationType == INTEGRATION_PUNCTUAL_LIGHT)
 							{
-							#if SUBSTRATE_COMPLEXSPECIALPATH && SUBSTRATE_GLINTS_ALLOWED
+							#if SUBSTRATE_COMPLEXSPECIALPATH && SUBSTRATE_GLINTS_ALLOWED && SUBSTRATE_GLINTS_ENABLED
 								BRANCH
 								if (BSDF_GETHASGLINT(BSDFContext.BSDF))
 								{
@@ -506,8 +506,8 @@ FSubstrateEvaluateResult SubstrateEvaluateBSDFCommon(FSubstrateBSDFContext BSDFC
 							else
 							{
 								const float Energy = EnergyNormalization(Alpha2Spec, BSDFContext.Context.VoH, AreaLightContext.AreaLight);
-
-							#if SUBSTRATE_COMPLEXSPECIALPATH && SUBSTRATE_GLINTS_ALLOWED
+							
+							#if SUBSTRATE_COMPLEXSPECIALPATH && SUBSTRATE_GLINTS_ALLOWED && SUBSTRATE_GLINTS_ENABLED
 								BRANCH
 								if (BSDF_GETHASGLINT(BSDFContext.BSDF))
 								{
@@ -637,8 +637,9 @@ FSubstrateEvaluateResult SubstrateEvaluateBSDFCommon(FSubstrateBSDFContext BSDFC
 				if (bHazeAsSimpleClearCoat)
 				{
 					const BxDFContext ClearCoatContext = RefractClearCoatContext(BSDFContext.Context);
-					const float3 HazeClearCoatTransmittance = SimpleClearCoatTransmittance(ClearCoatContext.NoL, ClearCoatContext.NoV, SubstrateGetBSDFMetallic(BSDFContext.BSDF), SubstrateGetBSDFBaseColor(BSDFContext.BSDF));
-
+					const float SlabMetallic = F0RGBToMetallic(F0);
+					const float3 HazeClearCoatTransmittance = SimpleClearCoatTransmittance(ClearCoatContext.NoL, ClearCoatContext.NoV, SlabMetallic, lerp(SLAB_DIFFUSEALBEDO(BSDFContext.BSDF), F0, SlabMetallic));
+				
 					const float TopLayerCoverage = HazeWeight;
 					const float TopLayerSpecularTransmittionApprox = saturate(1.0f - HazeFresnelTerm.x); // Simple and do not use the LUT for now. With bHazeAsSimpleClearCoat, Fresnel is achromatic
 					const float3 TopLayerThrouput = lerp(1.0f, HazeClearCoatTransmittance * TopLayerSpecularTransmittionApprox, TopLayerCoverage);
@@ -1546,7 +1547,7 @@ FSubstrateEnvLightResult SubstrateEvaluateForEnvLight(FSubstrateBSDFContext BSDF
 
 				// SUBSTRATE_TODO environment light importance sampling
 				// Glints
-			#if SUBSTRATE_COMPLEXSPECIALPATH && SUBSTRATE_GLINTS_ALLOWED
+			#if SUBSTRATE_COMPLEXSPECIALPATH && SUBSTRATE_GLINTS_ALLOWED && SUBSTRATE_GLINTS_ENABLED
 				BRANCH
 				if (BSDF_GETHASGLINT(BSDFContext.BSDF))
 				{
@@ -1575,7 +1576,7 @@ FSubstrateEnvLightResult SubstrateEvaluateForEnvLight(FSubstrateBSDFContext BSDF
 					SubstrateEnvLightResult.SpecularWeight *= saturate(GlintD);
 					SubstrateEnvLightResult.SpecularSafeRoughness = lerp(SubstrateEnvLightResult.SpecularSafeRoughness, 0.001f, saturate(GlintD));
 				}
-			#endif // SUBSTRATE_COMPLEXSPECIALPATH && SUBSTRATE_GLINTS_ALLOWED
+			#endif // SUBSTRATE_COMPLEXSPECIALPATH && SUBSTRATE_GLINTS_ALLOWED && SUBSTRATE_GLINTS_ENABLED
 
 			#if SUBSTRATE_COMPLEXSPECIALPATH
 				// Specular profile
