Substrate -SpecLUT evasluation code is now gone when it is turned off for the project/platform (deffered and forward).

#prelfight none

diff --git a/Engine/Shaders/Private/Substrate/SubstrateEvaluation.ush b/Engine/Shaders/Private/Substrate/SubstrateEvaluation.ush
--- a/Engine/Shaders/Private/Substrate/SubstrateEvaluation.ush
+++ b/Engine/Shaders/Private/Substrate/SubstrateEvaluation.ush
@@ -664,7 +664,7 @@ FSubstrateEvaluateResult SubstrateEvaluateBSDFCommon(FSubstrateBSDFContext BSDFC
 			#endif
 
 			// Lighting LUT
-			#if SUBSTRATE_FASTPATH==0 && SUBSTRATE_COMPLEXSPECIALPATH
+			#if SUBSTRATE_FASTPATH==0 && SUBSTRATE_COMPLEXSPECIALPATH && SUBSTRATE_SPECPROFILE_ENABLED
 			BRANCH
 			if (BSDF_GETHASSPECPROFILE(BSDFContext.BSDF))
 			{
@@ -1578,7 +1578,7 @@ FSubstrateEnvLightResult SubstrateEvaluateForEnvLight(FSubstrateBSDFContext BSDF
 				}
 			#endif // SUBSTRATE_COMPLEXSPECIALPATH && SUBSTRATE_GLINTS_ALLOWED && SUBSTRATE_GLINTS_ENABLED
 
-			#if SUBSTRATE_COMPLEXSPECIALPATH
+			#if SUBSTRATE_COMPLEXSPECIALPATH && SUBSTRATE_SPECPROFILE_ENABLED
 				// Specular profile
 				BRANCH
 				if (BSDF_GETHASSPECPROFILE(BSDFContext.BSDF))
@@ -1643,7 +1643,7 @@ FSubstrateEnvLightResult SubstrateEvaluateForEnvLight(FSubstrateBSDFContext BSDF
 					}
 
 					// Specular profile for haze
-				#if SUBSTRATE_COMPLEXSPECIALPATH
+				#if SUBSTRATE_COMPLEXSPECIALPATH && SUBSTRATE_SPECPROFILE_ENABLED
 					BRANCH
 					if (BSDF_GETHASSPECPROFILE(BSDFContext.BSDF))
 					{
