Improves TSR's stability without the anti-flickering heuristic

This has many additional other advantages:
0) Better stability on low and medium anti-aliasing scalability
1) Better stability on any movements
2) Better history resurrection on noisy environement
3) Allow to blend the clamped and unclamped history in LDR reducing ghosting chances on fireflier VFX without terrible image stability
4) Allows to reduce the strength of anti-flickering heuristic, making TSR ghost even less

Performance optimisation deferred to followup CL

diff --git a/Engine/Shaders/Private/TemporalSuperResolution/TSRCommon.ush b/Engine/Shaders/Private/TemporalSuperResolution/TSRCommon.ush
--- a/Engine/Shaders/Private/TemporalSuperResolution/TSRCommon.ush
+++ b/Engine/Shaders/Private/TemporalSuperResolution/TSRCommon.ush
@@ -13,7 +13,8 @@
 
 
 //------------------------------------------------------- RECOMPILE HASH
-// 2A42D2EF-728A-4AC9-825F-2DC009BE04B0
+
+#pragma message("UESHADERMETADATA_VERSION 2A42D2EF-728A-4AC9-825F-2DC009BE04B0")
 
 
 //------------------------------------------------------- COMPILER CONFIG
diff --git a/Engine/Shaders/Private/TemporalSuperResolution/TSRConvolutionNetwork.ush b/Engine/Shaders/Private/TemporalSuperResolution/TSRConvolutionNetwork.ush
--- a/Engine/Shaders/Private/TemporalSuperResolution/TSRConvolutionNetwork.ush
+++ b/Engine/Shaders/Private/TemporalSuperResolution/TSRConvolutionNetwork.ush
@@ -617,7 +617,7 @@ TLaneVector2D<ScalarType, VectorSize, LaneStrideX, LaneStrideY> MaxMinusMin3x3(
 }
 
 
-//------------------------------------------------------- CLAMP CONVOLUTIONS
+//------------------------------------------------------- CLAMP 3x3 CONVOLUTIONS
 
 CALL_SITE_DEBUGLOC
 template<typename ScalarType, uint VectorSize, uint LaneStrideX, uint LaneStrideY>
@@ -654,6 +654,40 @@ void AnnihilateMutually3x3(
 }
 
 
+//------------------------------------------------------- MIN & MAX & CLAMP 3x3 PLUS CONVOLUTION
+
+CALL_SITE_DEBUGLOC
+template<typename ScalarType, uint VectorSize, uint LaneStrideX, uint LaneStrideY>
+void MinMaxPlus3x3(
+	TLaneVector2D<ScalarType, VectorSize, LaneStrideX, LaneStrideY> Center,
+	out TLaneVector2D<ScalarType, VectorSize, LaneStrideX, LaneStrideY> OutMin,
+	out TLaneVector2D<ScalarType, VectorSize, LaneStrideX, LaneStrideY> OutMax)
+{
+	TLaneVector2D<ScalarType, VectorSize, LaneStrideX, LaneStrideY> CVP, CVN;
+	AccessNeighborTexels1x3(Center, /* out */ CVP, /* out */ CVN);
+	OutMin = min3(CVP, CVN, Center);
+	OutMax = max3(CVP, CVN, Center);
+	
+	TLaneVector2D<ScalarType, VectorSize, LaneStrideX, LaneStrideY> CHP, CHN;
+	AccessNeighborTexels3x1(Center, /* out */ CHP, /* out */ CHN);
+	OutMin = min3(CHN, CHP, OutMin);
+	OutMax = max3(CHN, CHP, OutMax);
+}
+
+CALL_SITE_DEBUGLOC
+template<typename ScalarType, uint VectorSize, uint LaneStrideX, uint LaneStrideY>
+TLaneVector2D<ScalarType, VectorSize, LaneStrideX, LaneStrideY> ClampPlus3x3(
+	TLaneVector2D<ScalarType, VectorSize, LaneStrideX, LaneStrideY> ToClamp,
+	TLaneVector2D<ScalarType, VectorSize, LaneStrideX, LaneStrideY> BoundaryCenter)
+{
+	TLaneVector2D<ScalarType, VectorSize, LaneStrideX, LaneStrideY> BoundaryMin;
+	TLaneVector2D<ScalarType, VectorSize, LaneStrideX, LaneStrideY> BoundaryMax;
+	MinMaxPlus3x3(BoundaryCenter, /* out */ BoundaryMin, /* out */ BoundaryMax);
+
+	return fastClamp(ToClamp, BoundaryMin, BoundaryMax);
+}
+
+
 //------------------------------------------------------- MEDIAN 3x3 CONVOLUTION
 
 // Operator that output, lowest, median and highest values from 3 input values.
diff --git a/Engine/Shaders/Private/TemporalSuperResolution/TSRRejectShading.usf b/Engine/Shaders/Private/TemporalSuperResolution/TSRRejectShading.usf
--- a/Engine/Shaders/Private/TemporalSuperResolution/TSRRejectShading.usf
+++ b/Engine/Shaders/Private/TemporalSuperResolution/TSRRejectShading.usf
@@ -157,6 +157,8 @@ tsr_simd_halfC ChangePerceptionAdd(tsr_simd_halfC SrcColor, const float SourcePe
 
 //------------------------------------------------------- MAIN ALGORITHMS
 
+static const tsr_half FilteringWeight = rcp(tsr_half(1.0 + 4 * 0.5 + 4 * 0.25));
+
 void ComputeMoireError(
 	tsr_simd_bool  bIsValid,
 	tsr_simd_half  bIsStatic,
@@ -166,12 +168,12 @@ void ComputeMoireError(
 	out tsr_simd_half4 OutMoireHistory,
 	out tsr_simd_half  OutMoireError)
 {
-	const tsr_half BlendFinal = tsr_half(0.05);
+	const tsr_half MinBlendFinal = tsr_half(0.05);
 	const tsr_half MoireEncodingError = tsr_half(rcp(127.0));
 	const tsr_half MaxPrevTotalVariationCount = tsr_half(20.0);
 
 	// Compute the threshold at which the variation count should quick in
-	const tsr_half TotalVariationCountThreshold = tsr_half(1.0 / (1.0 - pow(1.0 - BlendFinal, tsr_half(FlickeringFramePeriod))));
+	const tsr_half TotalVariationCountThreshold = tsr_half(1.0 / (1.0 - pow(1.0 - MinBlendFinal, tsr_half(FlickeringFramePeriod))));
 
 	// Input luminance to monitor flickering from.
 	tsr_simd_half InputLuma = MoireInput;
@@ -193,40 +195,88 @@ void ComputeMoireError(
 
 		InputLumaModification = max(abs(InputLuma - SceneColorLuma) - tsr_simd_half::Const(InputLumaEncodingError), tsr_simd_half::Const(0.0));
 	}
-		
-	// Convolve the input luma to have the exact same moire pattern as FilteredInput below.
-	Deconcatenate(Blur3x3(Concatenate(InputLuma, InputLumaModification)), /* out */ InputLuma, /* out */ InputLumaModification);
 
 	// Unpack history
 	tsr_simd_half PrevFlickeringHistory = PrevMoireHistory[0];
 	tsr_simd_half PrevGradient = PrevMoireHistory[1] * tsr_half(255.0 / 127.0) - tsr_half(1.0);
-	//tsr_simd_half PrevTotalVariation = PrevMoireHistory[2] * tsr_simd_half::Const(1.0 - BlendFinal);
-	tsr_simd_half PrevTotalVariation = PrevMoireHistory[2] * tsr_simd_half::Const(tsr_half(1.0) - BlendFinal);
-	tsr_simd_half PrevTotalVariationCount = PrevMoireHistory[3] * MaxPrevTotalVariationCount * tsr_simd_half::Const(tsr_half(1.0) - BlendFinal);
-		
+	//tsr_simd_half PrevTotalVariation = PrevMoireHistory[2] * tsr_simd_half::Const(1.0 - MinBlendFinal);
+	tsr_simd_half PrevTotalVariation = PrevMoireHistory[2] * tsr_simd_half::Const(tsr_half(1.0) - MinBlendFinal);
+	tsr_simd_half PrevTotalVariationCount = PrevMoireHistory[3] * MaxPrevTotalVariationCount * tsr_simd_half::Const(tsr_half(1.0) - MinBlendFinal);
+	
 	// Discard history on moving objects.
 	PrevTotalVariation = PrevTotalVariation * bIsStatic;
 	PrevTotalVariationCount = PrevTotalVariationCount * bIsStatic;
 
+	// Run the exact same heuristic as MeasureRejection on a single channel.
+	tsr_simd_half FilteredInputLumaModification;
+	tsr_simd_half OutRejectionBlendFinal;
+	tsr_simd_half OutRejectionClampBlend;
+	{
+		tsr_simd_half BackbufferQuantizationErrorVector = tsr_simd_half::Const(MoireEncodingError);
+		
+		tsr_simd_half InputC0 = InputLuma;
 
-	// Compute clamping box of the limance.
-	tsr_simd_half InputLumaMin;
-	tsr_simd_half InputLumaMax;
-	MinMax3x3(InputLuma, /* out */ InputLumaMin, /* out */ InputLumaMax);
+		// Anhilate history
+		tsr_simd_half InputC2 = InputC0;
+		tsr_simd_half HistoryC2 = PrevFlickeringHistory;
+		AnnihilateMutually3x3(InputLuma, PrevFlickeringHistory, /* out */ InputC2, /* out */ HistoryC2);
 
-	// Clamp and contribue the history.		
-	tsr_simd_half ClampedPrevFlickeringHistory = fastClamp(PrevFlickeringHistory, InputLumaMin, InputLumaMax);
-	tsr_simd_half FinalFlickeringHistory = InputLuma * tsr_simd_half::Const(BlendFinal) - ClampedPrevFlickeringHistory * tsr_simd_half::Const(BlendFinal - tsr_half(1.0));
+		// Convolve the input luma to have the exact same moire pattern as FilteredInput below.
+		tsr_simd_half FilteredInput;
+		tsr_simd_half FilteredHistory;
+		Deconcatenate(Blur3x3(Concatenate(InputC2, InputLumaModification, HistoryC2)), /* out */ FilteredInput, /* out */ FilteredInputLumaModification, /* out */ FilteredHistory);
+		
+		tsr_simd_half TotalVarInputDiffC0C2 = abs(TotalVariation3x3(abs(InputC0 - InputC2)));
+		tsr_simd_half TotalVarInputC2 = abs(TotalVariation3x3(InputC2));
+
+		// Compute clamping box of the limance.
+		tsr_simd_half FilteredBoxMin;
+		tsr_simd_half FilteredBoxMax;
+		MinMax3x3(FilteredInput, /* out */ FilteredBoxMin, /* out */ FilteredBoxMax);
+		FilteredBoxMin = FilteredBoxMin - BackbufferQuantizationErrorVector;
+		FilteredBoxMax = FilteredBoxMax + BackbufferQuantizationErrorVector;
+		
+		tsr_simd_half BoxSize = (
+			MaxMinusMin3x3(InputC2) * tsr_simd_half::Const(FilteringWeight) +
+			(BackbufferQuantizationErrorVector * tsr_half(2 * FilteringWeight)));
+
+		tsr_simd_half ClampedFilteredHistory = fastClamp(FilteredHistory, FilteredBoxMin, FilteredBoxMax);
+
+		tsr_simd_half Delta = max(abs(FilteredInput - FilteredHistory), BoxSize);
 	
-	// Measure the flickering gradient of the current frame.
-	//tsr_simd_half CurrentGradient = FinalFlickeringHistory - PrevFlickeringHistory;
-	tsr_simd_half CurrentGradient = ClampedPrevFlickeringHistory - PrevFlickeringHistory;
+		tsr_simd_half ClampError = BackbufferQuantizationErrorVector;
+		ClampError = max(ClampError, BoxSize * tsr_half(0.25)); // original
 		
+		tsr_simd_half RawClampedEnergy = abs(ClampedFilteredHistory - FilteredHistory);
+		tsr_simd_half RawRejection     = saturate(tsr_simd_half::Const(1.0) - max(RawClampedEnergy - ClampError, tsr_simd_half::Const(0.0)) * rcp(Delta));
+
+		OutRejectionClampBlend = Min3x3(Median3x3(RawRejection));
+		tsr_simd_half FilteredClampedEnergy = Max3x3(Median3x3(RawClampedEnergy));
+
+		OutRejectionBlendFinal = saturate(tsr_simd_half::Const(1.0) - max(FilteredClampedEnergy - ClampError, tsr_simd_half::Const(0.0)) * rcp(Delta));
+	}
+	
+	// Simulate what happens in TSRUpdateHistory.usf
+	tsr_simd_half FinalFlickeringHistory;
+	tsr_simd_half CurrentGradient;
+	{
+		tsr_simd_half ClampedPrevFlickeringHistory = ClampPlus3x3(PrevFlickeringHistory, InputLuma);
+		tsr_simd_half FinalPrevFlickeringHistory = lerp(ClampedPrevFlickeringHistory, PrevFlickeringHistory, OutRejectionClampBlend);
+
+		tsr_simd_half FinalBlendFinal = max(tsr_simd_half::Const(1.0) - OutRejectionBlendFinal, tsr_simd_half::Const(MinBlendFinal));
+		
+		FinalFlickeringHistory = InputLuma * FinalBlendFinal - FinalPrevFlickeringHistory * (FinalBlendFinal - tsr_simd_half::Const(1.0));
+		
+		tsr_simd_half GhostingHistory = InputLuma * tsr_simd_half::Const(MinBlendFinal) - PrevFlickeringHistory * tsr_simd_half::Const(MinBlendFinal - (1.0));
+
+		CurrentGradient = FinalFlickeringHistory - GhostingHistory;
+	}
+
 	// Remove how much luma has changed from luma measurement to final scene color from the gradient.
 	for (uint SimdIndex = 0; SimdIndex < SIMD_SIZE; SimdIndex++)
 	{
 		tsr_half LocalCurrentGradient = CurrentGradient.GetElement(SimdIndex);
-		tsr_half LocalInputLumaModification = InputLumaModification.GetElement(SimdIndex);
+		tsr_half LocalInputLumaModification = FilteredInputLumaModification.GetElement(SimdIndex);
 			
 		CurrentGradient.SetElement(SimdIndex, clamp(LocalCurrentGradient * tsr_half(POSITIVE_INFINITY), tsr_half(-1.0), tsr_half(1.0)) * max(abs(LocalCurrentGradient) - LocalInputLumaModification, tsr_half(0.0)));
 	}
@@ -281,10 +331,10 @@ void ComputeMoireError(
 		tsr_half LocalTotalVariationCountContribution = TotalVariationCountContribution.GetElement(SimdIndex);
 			
 		// Blend out previous gradient
-		tsr_half LocalPrevBlendedGradient = LocalPrevGradient * (tsr_half(1.0) - BlendFinal) * (tsr_half(1.0) - LocalIsFlicker);
+		tsr_half LocalPrevBlendedGradient = LocalPrevGradient * (tsr_half(1.0) - MinBlendFinal) * (tsr_half(1.0) - LocalIsFlicker);
 
 		// Compute new gradient.
-		#if 0
+		#if 1
 			tsr_half LocalNewGradient = LocalPrevBlendedGradient + LocalCurrentGradient;
 		#else
 			tsr_half LocalNewGradient = abs(LocalPrevBlendedGradient) > abs(LocalCurrentGradient) ? LocalPrevBlendedGradient : LocalCurrentGradient;
@@ -330,18 +380,18 @@ void ComputeMoireError(
 template<bool ApplyMoireError>
 void MeasureRejection(
 	tsr_simd_bool  bIsValid,
-	tsr_simd_halfC ClampedInput,
-	tsr_simd_halfC ClampedHistory,
+	tsr_simd_halfC InputC0,
+	tsr_simd_halfC InputC2,
+	tsr_simd_halfC HistoryC2,
 	tsr_simd_half  MoireError,
 	out tsr_simd_half OutRejectionBlendFinal,
 	out tsr_simd_half OutRejectionClampBlend)
 {
-	const tsr_half FilteringWeight = rcp(tsr_half(1.0 + 4 * 0.5 + 4 * 0.25));
 	tsr_simd_halfC BackbufferQuantizationErrorVector = tsr_simd_halfC::Const(MeasureBackbufferLDRQuantizationError());
 		
 	// Filter clamped input and history to increase quantization of the rejection detection
-	tsr_simd_halfC FilteredInput   = Blur3x3(ClampedInput);
-	tsr_simd_halfC FilteredHistory = Blur3x3(ClampedHistory);
+	tsr_simd_halfC FilteredInput   = Blur3x3(InputC2);
+	tsr_simd_halfC FilteredHistory = Blur3x3(HistoryC2);
 	
 	// Build the clamping box of the filtered signal
 	tsr_simd_halfC FilteredBoxMin;
@@ -350,6 +400,9 @@ void MeasureRejection(
 	FilteredBoxMin = FilteredBoxMin - BackbufferQuantizationErrorVector;
 	FilteredBoxMax = FilteredBoxMax + BackbufferQuantizationErrorVector;
 		
+	tsr_simd_halfC TotalVarInputDiffC0C2 = abs(TotalVariation3x3(abs(InputC0 - InputC2)));
+	tsr_simd_halfC TotalVarInputC2 = abs(TotalVariation3x3(InputC2));
+
 	// Clamp filtered signal
 	tsr_simd_halfC ClampedFilteredHistory = fastClamp(FilteredHistory, FilteredBoxMin, FilteredBoxMax);
 	if (ApplyMoireError)
@@ -382,7 +435,7 @@ void MeasureRejection(
 	tsr_simd_halfC BoxSize;
 	{
 		BoxSize = (
-			MaxMinusMin3x3(ClampedInput) * tsr_simd_halfC::Const(FilteringWeight) +
+			MaxMinusMin3x3(InputC2) * tsr_simd_halfC::Const(FilteringWeight) +
 			(BackbufferQuantizationErrorVector * tsr_half(2 * FilteringWeight)));
 
 		if (ApplyMoireError)
@@ -401,13 +454,14 @@ void MeasureRejection(
 		}
 	}
 
-	tsr_simd_halfC Delta = abs(FilteredInput - FilteredHistory);
-	tsr_simd_halfC ClampError = max(
-		BoxSize * tsr_half(0.25),
-		BackbufferQuantizationErrorVector);
-			
+	tsr_simd_halfC Delta = max(abs(FilteredInput - FilteredHistory), BoxSize);
+
+	tsr_simd_halfC ClampError = BackbufferQuantizationErrorVector;
+	ClampError = max(ClampError, BoxSize * tsr_half(0.25));
+	ClampError = max(ClampError, Blur3x3(min(TotalVarInputDiffC0C2, TotalVarInputC2)));
+
 	tsr_simd_halfC RawClampedEnergy      = abs(ClampedFilteredHistory - FilteredHistory);
-	tsr_simd_halfC RawFactor      = saturate(tsr_simd_halfC::Const(1.0) - max(RawClampedEnergy - ClampError, tsr_simd_halfC::Const(0.0)) * rcp(max(Delta, BoxSize)));
+	tsr_simd_halfC RawFactor      = saturate(tsr_simd_halfC::Const(1.0) - max(RawClampedEnergy - ClampError, tsr_simd_halfC::Const(0.0)) * rcp(Delta));
 	tsr_simd_half RawRejection      = min3(RawFactor[0], RawFactor[1], RawFactor[2]);
 
 	#if CONFIG_SCENE_COLOR_ALPHA
@@ -418,7 +472,7 @@ void MeasureRejection(
 		Deconcatenate(MaxRGBMinA3x3(Median3x3(Concatenate(RawClampedEnergy, RawRejection))), /* out */ FilteredClampedEnergy, /* out */ OutRejectionClampBlend);
 	#endif
 
-	tsr_simd_halfC FilteredFactor = saturate(tsr_simd_halfC::Const(1.0) - max(FilteredClampedEnergy - ClampError, tsr_simd_halfC::Const(0.0)) * rcp(max(Delta, BoxSize)));
+	tsr_simd_halfC FilteredFactor = saturate(tsr_simd_halfC::Const(1.0) - max(FilteredClampedEnergy - ClampError, tsr_simd_halfC::Const(0.0)) * rcp(Delta));
 	tsr_simd_half FilteredRejection = min3(FilteredFactor[0], FilteredFactor[1], FilteredFactor[2]);
 	OutRejectionBlendFinal = saturate(FilteredRejection);
 } // MeasureRejection()
@@ -695,9 +749,10 @@ void MainCS(
 
 		MeasureRejection</* ApplyMoireError = */ CONFIG_MOIRE_DETECTION != 0>(
 			bIsValid,
-			ClampedInput,
-			ClampedHistory,
-			MoireError,
+			/* InputC0 = */    ExposedInput,
+			/* InputC2 = */    ClampedInput,
+			/* HistoryC2 = */  ClampedHistory,
+			/* MoireError = */ MoireError,
 			/* out */ RejectionBlendFinal,
 			/* out */ RejectionClampBlend);
 	
@@ -730,8 +785,9 @@ void MainCS(
 				tsr_simd_half ResurrectionRejectionClampBlend = tsr_simd_half::Const(1.0);
 				MeasureRejection</* ApplyMoireError = */ false>(
 					bIsValid,
-					ClampedResurrectedInput,
-					ClampedResurrectedHistory,
+					/* InputC0 = */    ExposedInput,
+					/* InputC2 = */    ClampedResurrectedInput,
+					/* HistoryC2 = */  ClampedResurrectedHistory,
 					/* MoireError = */ tsr_simd_half::Const(0.0),
 					/* out */ ResurrectionRejectionBlendFinal,
 					/* out */ ResurrectionRejectionClampBlend);
diff --git a/Engine/Shaders/Private/TemporalSuperResolution/TSRUpdateHistory.usf b/Engine/Shaders/Private/TemporalSuperResolution/TSRUpdateHistory.usf
--- a/Engine/Shaders/Private/TemporalSuperResolution/TSRUpdateHistory.usf
+++ b/Engine/Shaders/Private/TemporalSuperResolution/TSRUpdateHistory.usf
@@ -962,14 +962,12 @@ void MainCS(
 	tsr_half2 FinalHistoryValidity;
 
 	{
-		tsr_halfCx2 ClampedPrevHighFrequencyColor = fastClamp(PrevHighFrequencyColor, InputMinColor, InputMaxColor);
 		tsr_halfCx2 BlendedPrevHighFrequencyColor;
-		BlendedPrevHighFrequencyColor[0] = lerp(ClampedPrevHighFrequencyColor[0], PrevHighFrequencyColor[0], DisableHistoryClamp);
-		BlendedPrevHighFrequencyColor[1] = lerp(ClampedPrevHighFrequencyColor[1], PrevHighFrequencyColor[1], DisableHistoryClamp);
-		BlendedPrevHighFrequencyColor[2] = lerp(ClampedPrevHighFrequencyColor[2], PrevHighFrequencyColor[2], DisableHistoryClamp);
-		#if CONFIG_SCENE_COLOR_ALPHA
-			BlendedPrevHighFrequencyColor[3] = lerp(ClampedPrevHighFrequencyColor[3], PrevHighFrequencyColor[3], DisableHistoryClamp);
-		#endif
+		{
+			tsr_halfCx2 ClampedPrevHighFrequencyColor = fastClamp(PrevHighFrequencyColor, InputMinColor, InputMaxColor);
+			tsr_half2x2 Weight = WeightedLerpFactors(HdrWeight4(ClampedPrevHighFrequencyColor), HdrWeight4(PrevHighFrequencyColor), DisableHistoryClamp);
+			BlendedPrevHighFrequencyColor = dpv_scale(ClampedPrevHighFrequencyColor, Weight[0]) + dpv_scale(PrevHighFrequencyColor, Weight[1]);
+		}
 		
 		// Clamp the validity due to motion to maintain better sharpness in history reprojection under motion.
 		#if 1
diff --git a/Engine/Shaders/Public/LaneVectorization.ush b/Engine/Shaders/Public/LaneVectorization.ush
--- a/Engine/Shaders/Public/LaneVectorization.ush
+++ b/Engine/Shaders/Public/LaneVectorization.ush
@@ -1378,11 +1378,11 @@ struct TLaneVector2D
 
 CALL_SITE_DEBUGLOC
 template<typename ScalarType, uint VectorSizeA, uint VectorSizeB, uint SimdSizeX, uint SimdSizeY>
-void Concatenate(
+TLaneVector2D<ScalarType, VectorSizeA + VectorSizeB, SimdSizeX, SimdSizeY> Concatenate(
 	TLaneVector2D<ScalarType, VectorSizeA, SimdSizeX, SimdSizeY> A,
-	TLaneVector2D<ScalarType, VectorSizeB, SimdSizeX, SimdSizeY> B,
-	out TLaneVector2D<ScalarType, VectorSizeA + VectorSizeB, SimdSizeX, SimdSizeY> R)
+	TLaneVector2D<ScalarType, VectorSizeB, SimdSizeX, SimdSizeY> B)
 {
+	TLaneVector2D<ScalarType, VectorSizeA + VectorSizeB, SimdSizeX, SimdSizeY> R;
 	UNROLL
 	for (uint ElementIndex = 0; ElementIndex < SimdSizeX * SimdSizeY; ElementIndex++)
 	{
@@ -1398,16 +1398,38 @@ void Concatenate(
 			R.Registers.SetElementComponent(ElementIndex, VectorSizeA + ComponentIndex, B.Registers.GetElementComponent(ElementIndex, ComponentIndex));
 		}
 	}
+	return R;
 }
 
 CALL_SITE_DEBUGLOC
-template<typename ScalarType, uint VectorSizeA, uint VectorSizeB, uint SimdSizeX, uint SimdSizeY>
-TLaneVector2D<ScalarType, VectorSizeA + VectorSizeB, SimdSizeX, SimdSizeY> Concatenate(
+template<typename ScalarType, uint VectorSizeA, uint VectorSizeB, uint VectorSizeC, uint SimdSizeX, uint SimdSizeY>
+TLaneVector2D<ScalarType, VectorSizeA + VectorSizeB + VectorSizeC, SimdSizeX, SimdSizeY> Concatenate(
 	TLaneVector2D<ScalarType, VectorSizeA, SimdSizeX, SimdSizeY> A,
-	TLaneVector2D<ScalarType, VectorSizeB, SimdSizeX, SimdSizeY> B)
+	TLaneVector2D<ScalarType, VectorSizeB, SimdSizeX, SimdSizeY> B,
+	TLaneVector2D<ScalarType, VectorSizeC, SimdSizeX, SimdSizeY> C)
 {
-	TLaneVector2D<ScalarType, VectorSizeA + VectorSizeB, SimdSizeX, SimdSizeY> R;
-	Concatenate(A, B, /* out */ R);
+	TLaneVector2D<ScalarType, VectorSizeA + VectorSizeB + VectorSizeC, SimdSizeX, SimdSizeY> R;
+	UNROLL
+	for (uint ElementIndex = 0; ElementIndex < SimdSizeX * SimdSizeY; ElementIndex++)
+	{
+		UNROLL
+		for (uint ComponentIndex = 0; ComponentIndex < VectorSizeA; ComponentIndex++)
+		{
+			R.Registers.SetElementComponent(ElementIndex, ComponentIndex, A.Registers.GetElementComponent(ElementIndex, ComponentIndex));
+		}
+
+		UNROLL
+		for (uint ComponentIndex = 0; ComponentIndex < VectorSizeB; ComponentIndex++)
+		{
+			R.Registers.SetElementComponent(ElementIndex, VectorSizeA + ComponentIndex, B.Registers.GetElementComponent(ElementIndex, ComponentIndex));
+		}
+		
+		UNROLL
+		for (uint ComponentIndex = 0; ComponentIndex < VectorSizeC; ComponentIndex++)
+		{
+			R.Registers.SetElementComponent(ElementIndex, VectorSizeA + VectorSizeB + ComponentIndex, C.Registers.GetElementComponent(ElementIndex, ComponentIndex));
+		}
+	}
 	return R;
 }
 
@@ -1435,6 +1457,37 @@ void Deconcatenate(
 	}
 }
 
+CALL_SITE_DEBUGLOC
+template<typename ScalarType, uint VectorSizeA, uint VectorSizeB, uint VectorSizeC, uint SimdSizeX, uint SimdSizeY>
+void Deconcatenate(
+	TLaneVector2D<ScalarType, VectorSizeA + VectorSizeB + VectorSizeC, SimdSizeX, SimdSizeY> M,
+	out TLaneVector2D<ScalarType, VectorSizeA, SimdSizeX, SimdSizeY> A,
+	out TLaneVector2D<ScalarType, VectorSizeB, SimdSizeX, SimdSizeY> B,
+	out TLaneVector2D<ScalarType, VectorSizeC, SimdSizeX, SimdSizeY> C)
+{
+	UNROLL
+	for (uint ElementIndex = 0; ElementIndex < SimdSizeX * SimdSizeY; ElementIndex++)
+	{
+		UNROLL
+		for (uint ComponentIndex = 0; ComponentIndex < VectorSizeA; ComponentIndex++)
+		{
+			A.Registers.SetElementComponent(ElementIndex, ComponentIndex, M.Registers.GetElementComponent(ElementIndex, ComponentIndex));
+		}
+
+		UNROLL
+		for (uint ComponentIndex = 0; ComponentIndex < VectorSizeB; ComponentIndex++)
+		{
+			B.Registers.SetElementComponent(ElementIndex, ComponentIndex, M.Registers.GetElementComponent(ElementIndex, VectorSizeA + ComponentIndex));
+		}
+
+		UNROLL
+		for (uint ComponentIndex = 0; ComponentIndex < VectorSizeC; ComponentIndex++)
+		{
+			C.Registers.SetElementComponent(ElementIndex, ComponentIndex, M.Registers.GetElementComponent(ElementIndex, VectorSizeA + VectorSizeB + ComponentIndex));
+		}
+	}
+}
+
 
 // ------------- 1 parameter
 
diff --git a/Engine/Source/Runtime/Renderer/Private/PostProcess/TemporalSuperResolution.cpp b/Engine/Source/Runtime/Renderer/Private/PostProcess/TemporalSuperResolution.cpp
--- a/Engine/Source/Runtime/Renderer/Private/PostProcess/TemporalSuperResolution.cpp
+++ b/Engine/Source/Runtime/Renderer/Private/PostProcess/TemporalSuperResolution.cpp
@@ -125,7 +125,7 @@ TAutoConsoleVariable<int32> CVarTSRFlickeringAdjustToFrameRate(
 	ECVF_Scalability | ECVF_RenderThreadSafe);
 
 TAutoConsoleVariable<float> CVarTSRFlickeringPeriod(
-	TEXT("r.TSR.ShadingRejection.Flickering.Period"), 3.0f,
+	TEXT("r.TSR.ShadingRejection.Flickering.Period"), 2.0f,
 	TEXT("Periode in frames in which luma oscilations at equal or greater frequency is considered flickering and should ghost to stabilize the image ")
 	TEXT("Please read r.TSR.ShadingRejection.Flickering's help for further details. (Default to 3 frames)."),
 	ECVF_RenderThreadSafe);
