Improves VisualizeTSR showflag by adding more colors

diff --git a/Engine/Shaders/Private/TemporalSuperResolution/TSRVisualize.usf b/Engine/Shaders/Private/TemporalSuperResolution/TSRVisualize.usf
--- a/Engine/Shaders/Private/TemporalSuperResolution/TSRVisualize.usf
+++ b/Engine/Shaders/Private/TemporalSuperResolution/TSRVisualize.usf
@@ -38,6 +38,7 @@ float FlickeringFramePeriod;
 Texture2D<tsr_half3> SceneColorTexture;
 Texture2D<float2>    ClosestDepthTexture;
 Texture2D<float2>    DilatedVelocityTexture;
+Texture2D<tsr_ushort> IsMovingMaskTexture;
 Texture2D<tsr_half4> HistoryRejectionTexture;
 Texture2D<tsr_half4> MoireHistoryTexture;
 Texture2D<tsr_ushort> AntiAliasMaskTexture;
@@ -100,13 +101,23 @@ tsr_half3 DisplayHistorySampleCount(FContext Ctx)
 tsr_half3 DisplayHistoryRejection(FContext Ctx)
 {
 	tsr_half LowFrequencyRejection = HistoryRejectionTexture.SampleLevel(GlobalPointClampedSampler, Ctx.InputUV, 0).r;
-	return SampleTSROutputGray(Ctx) * lerp(kRed, kWhite, saturate(LowFrequencyRejection));
+
+	tsr_half3 Color = kWhite;
+	Color = lerp(kOrange, Color, saturate((LowFrequencyRejection - 0.5) * 2.0));
+	Color = lerp(kRed,    Color, saturate((LowFrequencyRejection - 0.0) * 2.0));
+
+	return SampleTSROutputGray(Ctx) * Color;
 }
 
 tsr_half3 DisplayHistoryClamp(FContext Ctx)
 {
 	tsr_half LowFrequencyClamp = HistoryRejectionTexture.SampleLevel(GlobalPointClampedSampler, Ctx.InputUV, 0).g;
-	return SampleTSROutputGray(Ctx) * lerp(kRed, kWhite, saturate(LowFrequencyClamp));
+
+	tsr_half3 Color = kWhite;
+	Color = lerp(kOrange, Color, saturate((LowFrequencyClamp - 0.5) * 2.0));
+	Color = lerp(kRed,    Color, saturate((LowFrequencyClamp - 0.0) * 2.0));
+
+	return SampleTSROutputGray(Ctx) * Color;
 }
 
 tsr_half3 DisplayParallaxDisocclusionMask(FContext Ctx)
@@ -117,12 +128,24 @@ tsr_half3 DisplayParallaxDisocclusionMask(FContext Ctx)
 
 tsr_half3 DisplayResurrectionMask(FContext Ctx)
 {
+	if (!bCanResurrectHistory)
+	{
+		return SampleTSROutputGray(Ctx) * kPink;
+	}
+
 	tsr_half LowFrequencyRejection = HistoryRejectionTexture.SampleLevel(GlobalPointClampedSampler, Ctx.InputUV, 0).r;
+	float ClosestDeviceZ = ClosestDepthTexture.SampleLevel(GlobalPointClampedSampler, Ctx.InputUV, 0).g;
 	bool ResurrectionMask = (tsr_ushort(HistoryRejectionTexture.SampleLevel(GlobalPointClampedSampler, Ctx.InputUV, 0).a * tsr_half(255.0)) & tsr_ushort(0x2)) != tsr_ushort(0x0);
-	if (!bCanResurrectHistory)
+
+	float4 ThisClip = float4(Ctx.ScreenPos, ClosestDeviceZ, 1);
+	float4 ResurrectionClip = mul(ThisClip, ClipToResurrectionClip); 
+	float2 ResurrectionScreen = ResurrectionClip.xy / ResurrectionClip.w;
+
+	if (any(ResurrectionScreen != clamp(ResurrectionScreen, -1.0, 1.0)))
 	{
 		return SampleTSROutputGray(Ctx) * kPink;
 	}
+
 	return SampleTSROutputGray(Ctx) * select(ResurrectionMask, lerp(kRed, kGreen, LowFrequencyRejection), kWhite);
 }
 
@@ -132,6 +155,7 @@ tsr_half3 DisplayResurrectedColor(FContext Ctx)
 	{
 		return SampleTSROutputGray(Ctx) * kPink;
 	}
+
 	tsr_half LowFrequencyRejection = HistoryRejectionTexture.SampleLevel(GlobalPointClampedSampler, Ctx.InputUV, 0).r;
 	float ClosestDeviceZ = ClosestDepthTexture.SampleLevel(GlobalPointClampedSampler, Ctx.InputUV, 0).g;
 	bool ResurrectionMask = (tsr_ushort(HistoryRejectionTexture.SampleLevel(GlobalPointClampedSampler, Ctx.InputUV, 0).a * tsr_half(255.0)) & tsr_ushort(0x2)) != tsr_ushort(0x0);
@@ -144,8 +168,12 @@ tsr_half3 DisplayResurrectedColor(FContext Ctx)
 	KernelHistoryBufferUV = fastClamp(KernelHistoryBufferUV, PrevHistoryInfo_UVViewportBilinearMin, PrevHistoryInfo_UVViewportBilinearMax);
 
 	tsr_half3 ResurrectedColor = ResurrectedHistoryColorTexture.SampleLevel(GlobalBilinearClampedSampler, KernelHistoryBufferUV, 0) * tsr_half(ResurrectionPreExposureCorrection);
+	
+	tsr_half3 Color = kWhite;
+	Color = select(any(ResurrectionScreen != clamp(ResurrectionScreen, -1.0, 1.0)), kPink, Color);
+	Color = select(ResurrectionMask, lerp(kRed, kGreen, LowFrequencyRejection), Color);
 
-	return GrayScaleColor(ResurrectedColor) * select(ResurrectionMask, lerp(kRed, kGreen, LowFrequencyRejection), kWhite);
+	return GrayScaleColor(ResurrectedColor) * Color;
 }
 
 tsr_half3 DisplaySpatialAntiAliasingMask(FContext Ctx)
@@ -174,6 +202,7 @@ tsr_half3 DisplayAntiFlickering(FContext Ctx)
 
 	tsr_half4 MoireHistory = MoireHistoryTexture.SampleLevel(GlobalPointClampedSampler, Ctx.InputUV, 0);
 
+	tsr_half DisableAntiFlickering = tsr_half(1.0) - tsr_half(IsMovingMaskTexture[Ctx.InputPixelPos]) * rcp(tsr_half(255.0));
 
 	// Unpack history
 	tsr_half FlickeringHistory   = MoireHistory[0];
@@ -193,6 +222,7 @@ tsr_half3 DisplayAntiFlickering(FContext Ctx)
 	Color = lerp(kGreen, kYellow, saturate(LocalMoireError * 16.0));
 	Color = lerp(Color,  kRed,    saturate(LocalMoireError * 16.0 - 1.0));
 	Color = lerp(kWhite, Color,   LocalCountFadeIn);
+	Color = lerp(kPink, Color,   DisableAntiFlickering);
 
 	return SampleTSROutputGray(Ctx) * Color;
 }
diff --git a/Engine/Source/Runtime/Renderer/Private/PostProcess/TemporalSuperResolution.cpp b/Engine/Source/Runtime/Renderer/Private/PostProcess/TemporalSuperResolution.cpp
--- a/Engine/Source/Runtime/Renderer/Private/PostProcess/TemporalSuperResolution.cpp
+++ b/Engine/Source/Runtime/Renderer/Private/PostProcess/TemporalSuperResolution.cpp
@@ -947,6 +947,7 @@ class FTSRVisualizeCS : public FTSRShader
 		SHADER_PARAMETER_RDG_TEXTURE_SRV(Texture2D, SceneColorTexture)
 		SHADER_PARAMETER_RDG_TEXTURE(Texture2D, ClosestDepthTexture)
 		SHADER_PARAMETER_RDG_TEXTURE(Texture2D, DilatedVelocityTexture)
+		SHADER_PARAMETER_RDG_TEXTURE_SRV(Texture2D, IsMovingMaskTexture)
 		SHADER_PARAMETER_RDG_TEXTURE(Texture2D, HistoryRejectionTexture)
 		SHADER_PARAMETER_RDG_TEXTURE_SRV(Texture2D, MoireHistoryTexture)
 		SHADER_PARAMETER_RDG_TEXTURE(Texture2D, AntiAliasMaskTexture)
@@ -2540,6 +2541,7 @@ FDefaultTemporalUpscaler::FOutputs AddTemporalSuperResolutionPasses(
 			PassParameters->SceneColorTexture = SceneColorOutputTextureSRV;
 			PassParameters->ClosestDepthTexture = ClosestDepthTexture;
 			PassParameters->DilatedVelocityTexture = DilatedVelocityTexture;
+			PassParameters->IsMovingMaskTexture = IsMovingMaskTexture ? IsMovingMaskTexture : GraphBuilder.CreateSRV(FRDGTextureSRVDesc(BlackUintDummy));
 			PassParameters->HistoryRejectionTexture = HistoryRejectionTexture;
 			PassParameters->MoireHistoryTexture = MoireHistoryTexture ? MoireHistoryTexture : GraphBuilder.CreateSRV(FRDGTextureSRVDesc(BlackDummy));
 			PassParameters->AntiAliasMaskTexture = AntiAliasMaskTexture ? AntiAliasMaskTexture : BlackUintDummy;
