Refactor skin weight / indices unpacking and BoneMatrix computation to share code across:
* GpuSkinVertexFactory
* GpuSkinCache
* Hair manual skin cache

This removes a lot of duplicate code to simplify future maintenance/support.

diff --git a/Engine/Plugins/Animation/DeformerGraph/Shaders/Private/DataInterfaceSkeleton.ush b/Engine/Plugins/Animation/DeformerGraph/Shaders/Private/DataInterfaceSkeleton.ush
--- a/Engine/Plugins/Animation/DeformerGraph/Shaders/Private/DataInterfaceSkeleton.ush
+++ b/Engine/Plugins/Animation/DeformerGraph/Shaders/Private/DataInterfaceSkeleton.ush
@@ -21,52 +21,12 @@ uint ReadNumVertices_{DataInterfaceName}()
 
 int4 GetBlendIndices_{DataInterfaceName}(uint StreamOffset, uint ExtraInfluenceIndex)
 {
-	int4 UnpackedBlendIndices = 0;
-
-#if GPUSKIN_BONE_INDEX_UINT16
-	StreamOffset += GPUSKIN_VB_OFFSET_INFLUENCEBONES + 2 * ExtraInfluenceIndex;
-	uint PackedBlendIndices = {DataInterfaceName}_InputWeightStream[StreamOffset];
-	UnpackedBlendIndices.x = PackedBlendIndices & 0xffff;
-	UnpackedBlendIndices.y = PackedBlendIndices >> 16 & 0xffff;
-	PackedBlendIndices = {DataInterfaceName}_InputWeightStream[StreamOffset + 1];
-	UnpackedBlendIndices.z = PackedBlendIndices & 0xffff;
-	UnpackedBlendIndices.w = PackedBlendIndices >> 16 & 0xffff;
-#else // GPUSKIN_BONE_INDEX_UINT16
-	StreamOffset += GPUSKIN_VB_OFFSET_INFLUENCEBONES + ExtraInfluenceIndex;
-	uint PackedBlendIndices = {DataInterfaceName}_InputWeightStream[StreamOffset];
-	UnpackedBlendIndices.x = PackedBlendIndices & 0xff;
-	UnpackedBlendIndices.y = PackedBlendIndices >> 8 & 0xff;
-	UnpackedBlendIndices.z = PackedBlendIndices >> 16 & 0xff;
-	UnpackedBlendIndices.w = PackedBlendIndices >> 24 & 0xff;
-#endif // GPUSKIN_BONE_INDEX_UINT16
-
-	return UnpackedBlendIndices;
+	return GetBlendIndices({DataInterfaceName}_InputWeightStream, StreamOffset, ExtraInfluenceIndex);
 }
 
 float4 GetBlendWeights_{DataInterfaceName}(uint StreamOffset, uint ExtraInfluenceIndex)
 {
-	StreamOffset += (({DataInterfaceName}_NumBoneInfluences + 3) / 4) * (1 + GPUSKIN_BONE_INDEX_UINT16);
-
-	float4 UnpackedBlendWeights = 0;
-#if GPUSKIN_BONE_WEIGHTS_UINT16
-	StreamOffset += ExtraInfluenceIndex * 2;
-	uint PackedBlendWeights = {DataInterfaceName}_InputWeightStream[StreamOffset];
-	UnpackedBlendWeights.x = float(PackedBlendWeights & 0xffff) / 65535.0f;
-	UnpackedBlendWeights.y = float(PackedBlendWeights >> 16 & 0xffff) / 65535.0f;
-	PackedBlendWeights = {DataInterfaceName}_InputWeightStream[StreamOffset + 1];
-	UnpackedBlendWeights.z = float(PackedBlendWeights & 0xffff) / 65535.0f;
-	UnpackedBlendWeights.w = float(PackedBlendWeights >> 16 & 0xffff) / 65535.0f;
-	
-#else
-	StreamOffset += ExtraInfluenceIndex;
-	uint PackedBlendWeights = {DataInterfaceName}_InputWeightStream[StreamOffset];
-	
-	UnpackedBlendWeights.x = float(PackedBlendWeights & 0xff) / 255.0f;
-	UnpackedBlendWeights.y = float(PackedBlendWeights >> 8 & 0xff) / 255.0f;
-	UnpackedBlendWeights.z = float(PackedBlendWeights >> 16 & 0xff) / 255.0f;
-	UnpackedBlendWeights.w = float(PackedBlendWeights >> 24 & 0xff) / 255.0f;
-#endif	
-	return UnpackedBlendWeights;
+	return GetBlendWeights( { DataInterfaceName }_InputWeightStream, StreamOffset, ExtraInfluenceIndex, { DataInterfaceName }_NumBoneInfluences);
 }
 
 #endif // GPUSKIN_UNLIMITED_BONE_INFLUENCE
@@ -160,65 +120,20 @@ float3x4 ReadWeightedBoneMatrix_{DataInterfaceName}(uint VertexIndex)
 
 #elif GPUSKIN_UNLIMITED_BONE_INFLUENCE
 
-	float3x4 BlendMatrix = (float3x4)0;
-
-	uint BufferIndex = {DataInterfaceName}_InputStreamStart + VertexIndex;
-	uint BlendOffsetCount = {DataInterfaceName}_InputWeightLookupStream[BufferIndex];
-	int NumBoneInfluencesLocal = BlendOffsetCount & 0xff;
-	int IndicesOffset = BlendOffsetCount >> 8;
-	int BoneIndexSize = {DataInterfaceName}_InputWeightIndexSize & 0xff;
-	int BoneWeightSize = {DataInterfaceName}_InputWeightIndexSize >> 8;
-	int WeightsOffset = IndicesOffset + (BoneIndexSize * NumBoneInfluencesLocal);
-
-	for (int InfluenceIndex = 0; InfluenceIndex < NumBoneInfluencesLocal; InfluenceIndex++)
-	{
-		int BoneIndexOffset = IndicesOffset + (BoneIndexSize * InfluenceIndex);
-		int BoneIndex = {DataInterfaceName}_InputWeightStream[BoneIndexOffset];
-		if (BoneIndexSize > 1)
-		{
-			BoneIndex = {DataInterfaceName}_InputWeightStream[BoneIndexOffset + 1] << 8 | BoneIndex;
-		}
-		int BoneWeightOffset = WeightsOffset + (BoneWeightSize * InfluenceIndex);
-		float BoneWeight = float({ DataInterfaceName }_InputWeightStream[BoneWeightOffset]) / 255.0f;
-		if (BoneWeightSize > 1)
-		{
-			BoneWeight = float({DataInterfaceName}_InputWeightStream[BoneWeightOffset + 1] << 8 | {DataInterfaceName}_InputWeightStream[BoneWeightOffset]) / 65535.0;
-		}
-		BlendMatrix += BoneWeight * GetBoneMatrix_{DataInterfaceName}(BoneIndex);
-	}
-
-	return BlendMatrix;
+	const uint BufferIndex 		= {DataInterfaceName}_InputStreamStart + VertexIndex;
+	const uint BlendOffsetCount = {DataInterfaceName}_InputWeightLookupStream[BufferIndex];
+	return ComputeBoneMatrixWithUnlimitedInfluences({DataInterfaceName}_BoneMatrices, {DataInterfaceName}_InputWeightStream, {DataInterfaceName}_InputWeightIndexSize, BlendOffsetCount);
 
-#else // !GPUSKIN_UNLIMITED_BONE_INFLUENCE
-
-	float3x4 BlendMatrix = (float3x4)0;
-
-	uint StreamOffset = {DataInterfaceName}_InputWeightStart + (VertexIndex * ({DataInterfaceName}_InputWeightStride / 4));
-
-	if ({DataInterfaceName}_NumBoneInfluences > 0)
-	{
-		int4 BlendIndices = GetBlendIndices_{DataInterfaceName}(StreamOffset, 0);
-		float4 BlendWeights = GetBlendWeights_{DataInterfaceName}(StreamOffset, 0);
-
-		BlendMatrix += BlendWeights.x * GetBoneMatrix_{DataInterfaceName}(BlendIndices.x);
-		BlendMatrix += BlendWeights.y * GetBoneMatrix_{DataInterfaceName}(BlendIndices.y);
-		BlendMatrix += BlendWeights.z * GetBoneMatrix_{DataInterfaceName}(BlendIndices.z);
-		BlendMatrix += BlendWeights.w * GetBoneMatrix_{DataInterfaceName}(BlendIndices.w);
-	}
-
-	// todo[CF]: make this a permutation when permutation support is working again?
-	if ({DataInterfaceName}_NumBoneInfluences > 4)
-	{
-		int4 BlendIndices = GetBlendIndices_{DataInterfaceName}(StreamOffset, 1);
-		float4 BlendWeights = GetBlendWeights_{DataInterfaceName}(StreamOffset, 1);
+#else
 
-		BlendMatrix += BlendWeights.x * GetBoneMatrix_{DataInterfaceName}(BlendIndices.x);
-		BlendMatrix += BlendWeights.y * GetBoneMatrix_{DataInterfaceName}(BlendIndices.y);
-		BlendMatrix += BlendWeights.z * GetBoneMatrix_{DataInterfaceName}(BlendIndices.z);
-		BlendMatrix += BlendWeights.w * GetBoneMatrix_{DataInterfaceName}(BlendIndices.w);
-	}
+	const uint StreamOffset = {DataInterfaceName}_InputWeightStart + (VertexIndex * ({DataInterfaceName}_InputWeightStride / 4));
 
-	return BlendMatrix;
+	FGPUSkinIndexAndWeight IndicesAndWeights = (FGPUSkinIndexAndWeight)0;
+	IndicesAndWeights.BlendIndices  = GetBlendIndices_{DataInterfaceName}(StreamOffset, 0);
+	IndicesAndWeights.BlendIndices2 = GetBlendIndices_{DataInterfaceName}(StreamOffset, 1);
+	IndicesAndWeights.BlendWeights  = GetBlendWeights_{DataInterfaceName}(StreamOffset, 0);
+	IndicesAndWeights.BlendWeights2 = GetBlendWeights_{DataInterfaceName}(StreamOffset, 1);
+	return ComputeBoneMatrixWithLimitedInfluences({DataInterfaceName}_BoneMatrices, IndicesAndWeights);
 
 #endif
 }
diff --git a/Engine/Plugins/Runtime/HairStrands/Source/HairStrandsCore/Private/HairStrandsMeshProjection.cpp b/Engine/Plugins/Runtime/HairStrands/Source/HairStrandsCore/Private/HairStrandsMeshProjection.cpp
--- a/Engine/Plugins/Runtime/HairStrands/Source/HairStrandsCore/Private/HairStrandsMeshProjection.cpp
+++ b/Engine/Plugins/Runtime/HairStrands/Source/HairStrandsCore/Private/HairStrandsMeshProjection.cpp
@@ -80,14 +80,15 @@ class FSkinUpdateCS : public FGlobalShader
 	class FPrevious : SHADER_PERMUTATION_BOOL("PERMUTATION_PREV");
 	class FUnlimitedBoneInfluence : SHADER_PERMUTATION_BOOL("GPUSKIN_UNLIMITED_BONE_INFLUENCE");
 	class FUseExtraInfluence : SHADER_PERMUTATION_BOOL("GPUSKIN_USE_EXTRA_INFLUENCES");
-	class FIndexUint16 : SHADER_PERMUTATION_BOOL("GPUSKIN_BONE_INDEX_UINT16");
-	using FPermutationDomain = TShaderPermutationDomain<FUnlimitedBoneInfluence, FUseExtraInfluence, FIndexUint16, FPrevious>;
+	class FBoneIndexUint16 : SHADER_PERMUTATION_BOOL("GPUSKIN_BONE_INDEX_UINT16");
+	class FBoneWeightUint16 : SHADER_PERMUTATION_BOOL("GPUSKIN_BONE_WEIGHTS_UINT16");	
+	using FPermutationDomain = TShaderPermutationDomain<FUnlimitedBoneInfluence, FUseExtraInfluence, FBoneIndexUint16, FBoneWeightUint16, FPrevious>;
 
 	BEGIN_SHADER_PARAMETER_STRUCT(FParameters, )
 		SHADER_PARAMETER(uint32, NumVertexToProcess)
 		SHADER_PARAMETER(uint32, NumTotalVertices)
 		SHADER_PARAMETER(uint32, SectionVertexBaseIndex)
-		SHADER_PARAMETER(uint32, IndexSize)
+		SHADER_PARAMETER(uint32, WeightIndexSize)
 		SHADER_PARAMETER(uint32, WeightStride)
 		SHADER_PARAMETER(uint32, BonesOffset)
 		SHADER_PARAMETER_SRV(Buffer<uint>, WeightLookup)
@@ -127,7 +128,7 @@ void AddSkinUpdatePass(
 	const bool bPrevPosition = OutPrevDeformedPosition != nullptr && PrevBoneMatrices != nullptr;
 	
 	FSkinUpdateCS::FParameters* Parameters = GraphBuilder.AllocParameters<FSkinUpdateCS::FParameters>();
-	Parameters->IndexSize = SkinWeight->GetBoneIndexByteSize();
+	Parameters->WeightIndexSize = SkinWeight->GetBoneIndexByteSize() | (SkinWeight->GetBoneWeightByteSize() << 8);
 	Parameters->NumVertexToProcess = NumVertexToProcess;
 	Parameters->NumTotalVertices = NumTotalVertices;
 	Parameters->SectionVertexBaseIndex = SectionVertexBaseIndex;
@@ -148,7 +149,8 @@ void AddSkinUpdatePass(
 	FSkinUpdateCS::FPermutationDomain PermutationVector;
 	PermutationVector.Set<FSkinUpdateCS::FUnlimitedBoneInfluence>(SkinWeight->GetBoneInfluenceType() == GPUSkinBoneInfluenceType::UnlimitedBoneInfluence);
 	PermutationVector.Set<FSkinUpdateCS::FUseExtraInfluence>(SkinWeight->GetMaxBoneInfluences() > MAX_INFLUENCES_PER_STREAM);
-	PermutationVector.Set<FSkinUpdateCS::FIndexUint16>(SkinWeight->Use16BitBoneIndex());
+	PermutationVector.Set<FSkinUpdateCS::FBoneIndexUint16 >(SkinWeight->Use16BitBoneIndex());
+	PermutationVector.Set<FSkinUpdateCS::FBoneIndexUint16 >(SkinWeight->Use16BitBoneWeight());
 	PermutationVector.Set<FSkinUpdateCS::FPrevious>(bPrevPosition);
 
 	const FIntVector DispatchGroupCount = FComputeShaderUtils::GetGroupCount(NumVertexToProcess, 64);
diff --git a/Engine/Shaders/Private/GpuSkinCacheComputeShader.usf b/Engine/Shaders/Private/GpuSkinCacheComputeShader.usf
--- a/Engine/Shaders/Private/GpuSkinCacheComputeShader.usf
+++ b/Engine/Shaders/Private/GpuSkinCacheComputeShader.usf
@@ -103,14 +103,7 @@ struct FVertexUnpacked
 	float4	TangentZ;
 
 #if !GPUSKIN_UNLIMITED_BONE_INFLUENCE
-	int4	BlendIndices;
-#if GPUSKIN_USE_EXTRA_INFLUENCES
-	int4	BlendIndices2;
-#endif
-	float4	BlendWeights;
-#if GPUSKIN_USE_EXTRA_INFLUENCES
-	float4	BlendWeights2;
-#endif
+	FGPUSkinIndexAndWeight IndexAndWeights;
 #endif // GPUSKIN_UNLIMITED_BONE_INFLUENCE
 
 #if GPUSKIN_MORPH_BLEND
@@ -152,67 +145,6 @@ Buffer<uint> InputWeightLookupStream;
 #endif
 Buffer<uint> InputWeightStream;
 
-FBoneMatrix CalcBoneMatrix( FVertexUnpacked VertexUnpacked, uint VertexIndex )
-{
-#if GPUSKIN_UNLIMITED_BONE_INFLUENCE
-	uint BlendOffsetCount = InputWeightLookupStream[VertexIndex];
-	int NumBoneInfluences = BlendOffsetCount & 0xff;
-	int StreamOffset = BlendOffsetCount >> 8;
-	int BoneIndexSize = InputWeightIndexSize & 0xff;
-	int BoneWeightSize = InputWeightIndexSize >> 8;
-	int WeightsOffset = StreamOffset + (BoneIndexSize * NumBoneInfluences);
-
-	FBoneMatrix BoneMatrix = FBoneMatrix(float4(0,0,0,0), float4(0,0,0,0), float4(0,0,0,0));
-	for (int InfluenceIdx = 0; InfluenceIdx < NumBoneInfluences; InfluenceIdx++)
-	{
-		int BoneIndexOffset = StreamOffset + (BoneIndexSize * InfluenceIdx);
-		int BoneWeightOffset = WeightsOffset + InfluenceIdx * BoneWeightSize;
-		int BoneIndex = InputWeightStream[BoneIndexOffset];
-		if (BoneIndexSize > 1)
-		{
-			BoneIndex = InputWeightStream[BoneIndexOffset + 1] << 8 | BoneIndex;
-			//@todo-lh: Workaround to fix issue in SPIRVEmitter of DXC; this block must be inside the if branch
-			float BoneWeight;
-			if (BoneWeightSize > 1)
-			{
-				BoneWeight = float(InputWeightStream[BoneWeightOffset + 1] << 8 | InputWeightStream[BoneWeightOffset]) / 65535.0;
-			}
-			else
-			{
-				BoneWeight = float(InputWeightStream[BoneWeightOffset]) / 255.0f;
-			}
-			BoneMatrix += BoneWeight * FBoneMatrix(BoneMatrices[BoneIndex*3], BoneMatrices[BoneIndex*3+1], BoneMatrices[BoneIndex*3+2]);
-		}
-		else
-		{
-			//@todo-lh: Workaround to fix issue in SPIRVEmitter of DXC; this block must be inside the if branch
-			float BoneWeight;
-			if (BoneWeightSize > 1)
-			{
-				BoneWeight = float(InputWeightStream[BoneWeightOffset + 1] << 8 | InputWeightStream[BoneWeightOffset]) / 65535.0;
-			}
-			else
-			{
-				BoneWeight = float(InputWeightStream[BoneWeightOffset]) / 255.0f;
-			}
-			BoneMatrix += BoneWeight * FBoneMatrix(BoneMatrices[BoneIndex*3], BoneMatrices[BoneIndex*3+1], BoneMatrices[BoneIndex*3+2]);
-		}
-	}
-#else // GPUSKIN_UNLIMITED_BONE_INFLUENCE
-	FBoneMatrix BoneMatrix = VertexUnpacked.BlendWeights.x * FBoneMatrix(BoneMatrices[VertexUnpacked.BlendIndices.x*3], BoneMatrices[VertexUnpacked.BlendIndices.x*3+1], BoneMatrices[VertexUnpacked.BlendIndices.x*3+2]);
-	BoneMatrix += VertexUnpacked.BlendWeights.y * FBoneMatrix(BoneMatrices[VertexUnpacked.BlendIndices.y*3], BoneMatrices[VertexUnpacked.BlendIndices.y*3+1], BoneMatrices[VertexUnpacked.BlendIndices.y*3+2]);
-	BoneMatrix += VertexUnpacked.BlendWeights.z * FBoneMatrix(BoneMatrices[VertexUnpacked.BlendIndices.z*3], BoneMatrices[VertexUnpacked.BlendIndices.z*3+1], BoneMatrices[VertexUnpacked.BlendIndices.z*3+2]);
-	BoneMatrix += VertexUnpacked.BlendWeights.w * FBoneMatrix(BoneMatrices[VertexUnpacked.BlendIndices.w*3], BoneMatrices[VertexUnpacked.BlendIndices.w*3+1], BoneMatrices[VertexUnpacked.BlendIndices.w*3+2]);
-#if GPUSKIN_USE_EXTRA_INFLUENCES
-	BoneMatrix += VertexUnpacked.BlendWeights2.x * FBoneMatrix(BoneMatrices[VertexUnpacked.BlendIndices2.x*3], BoneMatrices[VertexUnpacked.BlendIndices2.x*3+1], BoneMatrices[VertexUnpacked.BlendIndices2.x*3+2]);
-	BoneMatrix += VertexUnpacked.BlendWeights2.y * FBoneMatrix(BoneMatrices[VertexUnpacked.BlendIndices2.y*3], BoneMatrices[VertexUnpacked.BlendIndices2.y*3+1], BoneMatrices[VertexUnpacked.BlendIndices2.y*3+2]);
-	BoneMatrix += VertexUnpacked.BlendWeights2.z * FBoneMatrix(BoneMatrices[VertexUnpacked.BlendIndices2.z*3], BoneMatrices[VertexUnpacked.BlendIndices2.z*3+1], BoneMatrices[VertexUnpacked.BlendIndices2.z*3+2]);
-	BoneMatrix += VertexUnpacked.BlendWeights2.w * FBoneMatrix(BoneMatrices[VertexUnpacked.BlendIndices2.w*3], BoneMatrices[VertexUnpacked.BlendIndices2.w*3+1], BoneMatrices[VertexUnpacked.BlendIndices2.w*3+2]);
-#endif
-#endif // GPUSKIN_UNLIMITED_BONE_INFLUENCE
-	return BoneMatrix;
-}
-
 Buffer<SNORM float4> TangentInputBuffer;
 RWBuffer<TANGENT_RWBUFFER_FORMAT> TangentBufferUAV;
 Buffer<float> PositionInputBuffer;
@@ -236,65 +168,8 @@ FVertexUnpacked UnpackedVertex(uint MorphIndex, uint ClothIndex, uint ClothIndex
 	Unpacked.TangentZ = TangentBias_SkinCache(TangentInputBuffer[2 * VertexIndex + GPUSKIN_VB_OFFSET_TANGENT_Z]);
 
 #if !GPUSKIN_UNLIMITED_BONE_INFLUENCE
-	uint BlendIndices = InputWeightStream[WeightOffset + GPUSKIN_VB_OFFSET_INFLUENCEBONES];
-#if GPUSKIN_BONE_INDEX_UINT16
-	Unpacked.BlendIndices.x = BlendIndices & 0xffff;
-	Unpacked.BlendIndices.y = BlendIndices >> 16 & 0xffff;
-	BlendIndices = InputWeightStream[WeightOffset + GPUSKIN_VB_OFFSET_INFLUENCEBONES + 1];
-	Unpacked.BlendIndices.z = BlendIndices & 0xffff;
-	Unpacked.BlendIndices.w = BlendIndices >> 16 & 0xffff;
-#if GPUSKIN_USE_EXTRA_INFLUENCES
-	BlendIndices = InputWeightStream[WeightOffset + GPUSKIN_VB_OFFSET_INFLUENCEBONES + 2];
-	Unpacked.BlendIndices2.x = BlendIndices & 0xffff;
-	Unpacked.BlendIndices2.y = BlendIndices >> 16 & 0xffff;
-	BlendIndices = InputWeightStream[WeightOffset + GPUSKIN_VB_OFFSET_INFLUENCEBONES + 3];
-	Unpacked.BlendIndices2.z = BlendIndices & 0xffff;
-	Unpacked.BlendIndices2.w = BlendIndices >> 16 & 0xffff;
-#endif
-#else // GPUSKIN_BONE_INDEX_UINT16
-	Unpacked.BlendIndices.x = BlendIndices & 0xff;
-	Unpacked.BlendIndices.y = BlendIndices >> 8 & 0xff;
-	Unpacked.BlendIndices.z = BlendIndices >> 16 & 0xff;
-	Unpacked.BlendIndices.w = BlendIndices >> 24 & 0xff;
-#if GPUSKIN_USE_EXTRA_INFLUENCES
-	BlendIndices = InputWeightStream[WeightOffset + GPUSKIN_VB_OFFSET_INFLUENCEBONES + 1];
-	Unpacked.BlendIndices2.x = BlendIndices & 0xff;
-	Unpacked.BlendIndices2.y = BlendIndices >> 8 & 0xff;
-	Unpacked.BlendIndices2.z = BlendIndices >> 16 & 0xff;
-	Unpacked.BlendIndices2.w = BlendIndices >> 24 & 0xff;
+	Unpacked.IndexAndWeights = GetBlendIndicesAndWeights(InputWeightStream, WeightOffset);
 #endif
-#endif // GPUSKIN_BONE_INDEX_UINT16
-#if GPUSKIN_BONE_WEIGHTS_UINT16
-	uint BlendWeights = InputWeightStream[WeightOffset + GPUSKIN_VB_OFFSET_INFLUENCEWEIGHTS];
-	Unpacked.BlendWeights.x = float(BlendWeights & 0xffff) / 65535.0f;
-	Unpacked.BlendWeights.y = float(BlendWeights >> 16 & 0xffff) / 65535.0f;
-	BlendWeights = InputWeightStream[WeightOffset + GPUSKIN_VB_OFFSET_INFLUENCEWEIGHTS + 1];
-	Unpacked.BlendWeights.z = float(BlendWeights & 0xffff) / 65535.0f;
-	Unpacked.BlendWeights.w = float(BlendWeights >> 16 & 0xffff) / 65535.0f;
-#if GPUSKIN_USE_EXTRA_INFLUENCES
-	BlendWeights = InputWeightStream[WeightOffset + GPUSKIN_VB_OFFSET_INFLUENCEWEIGHTS + 2];
-	Unpacked.BlendWeights2.x = float(BlendWeights & 0xffff) / 65535.0f;
-	Unpacked.BlendWeights2.y = float(BlendWeights >> 16 & 0xffff) / 65535.0f;
-	BlendWeights = InputWeightStream[WeightOffset + GPUSKIN_VB_OFFSET_INFLUENCEWEIGHTS + 3];
-	Unpacked.BlendWeights2.z = float(BlendWeights & 0xffff) / 65535.0f;
-	Unpacked.BlendWeights2.w = float(BlendWeights >> 16 & 0xffff) / 65535.0f;
-#endif
-#else
-	uint BlendWeights = InputWeightStream[WeightOffset + GPUSKIN_VB_OFFSET_INFLUENCEWEIGHTS];
-	Unpacked.BlendWeights.x = float(BlendWeights & 0xff) / 255.0f;
-	Unpacked.BlendWeights.y = float(BlendWeights >> 8 & 0xff) / 255.0f;
-	Unpacked.BlendWeights.z = float(BlendWeights >> 16 & 0xff) / 255.0f;
-	Unpacked.BlendWeights.w = float(BlendWeights >> 24 & 0xff) / 255.0f;
-#if GPUSKIN_USE_EXTRA_INFLUENCES
-	BlendWeights = InputWeightStream[WeightOffset + GPUSKIN_VB_OFFSET_INFLUENCEWEIGHTS + 1];
-	Unpacked.BlendWeights2.x = float(BlendWeights & 0xff) / 255.0f;
-	Unpacked.BlendWeights2.y = float(BlendWeights >> 8 & 0xff) / 255.0f;
-	Unpacked.BlendWeights2.z = float(BlendWeights >> 16 & 0xff) / 255.0f;
-	Unpacked.BlendWeights2.w = float(BlendWeights >> 24 & 0xff) / 255.0f;
-#endif
-#endif // GPUSKIN_BONE_WEIGHTS_UINT16
-
-#endif // GPUSKIN_UNLIMITED_BONE_INFLUENCE
 
 #if GPUSKIN_MORPH_BLEND
 	GetMorphBuffer(MorphIndex, Unpacked.DeltaPosition, Unpacked.DeltaTangentZ);
@@ -351,13 +226,15 @@ void SkinCacheUpdateBatchCS(uint3 GroupID : SV_GroupID,
 	// Perform the skinning
 	FVertexFactoryIntermediates	Intermediates = (FVertexFactoryIntermediates)0;
 	
-#if 1
-	Intermediates.BlendMatrix = CalcBoneMatrix(Unpacked, InputOffset);
-#else
+#if 0
 	// Test no blend mtx
 	Intermediates.BlendMatrix[0] = float4(1,0,0,0);
 	Intermediates.BlendMatrix[1] = float4(0,1,0,0);
 	Intermediates.BlendMatrix[2] = float4(0,0,1,0);
+#elif GPUSKIN_UNLIMITED_BONE_INFLUENCE
+	Intermediates.BlendMatrix = ComputeBoneMatrixWithUnlimitedInfluences(BoneMatrices, InputWeightStream, InputWeightIndexSize, InputWeightLookupStream[InputOffset]);
+#else
+	Intermediates.BlendMatrix = ComputeBoneMatrixWithLimitedInfluences(BoneMatrices, Unpacked.IndexAndWeights);
 #endif
 
 	Intermediates.UnpackedPosition.x = PositionInputBuffer[InputOffset * 3 + 0];
diff --git a/Engine/Shaders/Private/GpuSkinCommon.ush b/Engine/Shaders/Private/GpuSkinCommon.ush
--- a/Engine/Shaders/Private/GpuSkinCommon.ush
+++ b/Engine/Shaders/Private/GpuSkinCommon.ush
@@ -51,6 +51,10 @@
 #define GPUSKIN_RWBUFFER_OFFSET_TANGENT_Z	1
 #endif
 
+#ifndef GPUSKIN_LIMIT_2BONE_INFLUENCES
+#define GPUSKIN_LIMIT_2BONE_INFLUENCES 0
+#endif
+
 // @return 0..255
 uint4 UnpackU8x4(uint Packed)
 {
@@ -106,4 +110,160 @@ uint PackS8x3N(float3 Unpacked)
 #define TANGENT_RWBUFFER_FORMAT		SNORM float4
 #define TangentBias_SkinCache(X)	(X)
 #define TangentUnbias_SkinCache(X)	(X)
-#endif // OPENGL_PROFILE
\ No newline at end of file
+#endif // OPENGL_PROFILE
+
+///////////////////////////////////////////////////////////////////////////////////////////////////
+// Helpers for accessing vertex indices and weights
+
+int4 GetBlendIndices(Buffer<uint> InWeightBuffer, uint StreamOffset, uint ExtraInfluenceIndex)
+{
+	int4 UnpackedBlendIndices = 0;
+
+#if GPUSKIN_BONE_INDEX_UINT16
+	StreamOffset += GPUSKIN_VB_OFFSET_INFLUENCEBONES + 2 * ExtraInfluenceIndex;
+	UnpackedBlendIndices.xy = UnpackUInt2FromUInt(InWeightBuffer[StreamOffset   ]);
+	UnpackedBlendIndices.zw = UnpackUInt2FromUInt(InWeightBuffer[StreamOffset + 1]);
+#else // GPUSKIN_BONE_INDEX_UINT16
+	StreamOffset += GPUSKIN_VB_OFFSET_INFLUENCEBONES + ExtraInfluenceIndex;
+	UnpackedBlendIndices = UnpackUintToUint4(InWeightBuffer[StreamOffset]);
+#endif // GPUSKIN_BONE_INDEX_UINT16
+	return UnpackedBlendIndices;
+}
+
+float4 GetBlendWeights(Buffer<uint> InWeightBuffer, uint StreamOffset, uint ExtraInfluenceIndex, uint InNumBoneInfluences)
+{
+	StreamOffset += ((InNumBoneInfluences + 3) / 4) * (1 + GPUSKIN_BONE_INDEX_UINT16);
+
+	float4 UnpackedBlendWeights = 0;
+#if GPUSKIN_BONE_WEIGHTS_UINT16
+	StreamOffset += ExtraInfluenceIndex * 2;
+	UnpackedBlendWeights.xy = UnpackUnorm2x16(InWeightBuffer[StreamOffset   ]);
+	UnpackedBlendWeights.zw = UnpackUnorm2x16(InWeightBuffer[StreamOffset + 1]);	
+#else
+	StreamOffset += ExtraInfluenceIndex;
+	UnpackedBlendWeights = UnpackRGBA8(InWeightBuffer[StreamOffset]);
+#endif	
+	return UnpackedBlendWeights;
+}
+
+///////////////////////////////////////////////////////////////////////////////////////////////////
+// Retrieve vertex indices & weights and compute bone matrix
+
+struct FGPUSkinIndexAndWeight
+{
+	int4 BlendIndices;
+	int4 BlendIndices2;
+	float4 BlendWeights;
+	float4 BlendWeights2;
+};
+
+FGPUSkinIndexAndWeight GetBlendIndicesAndWeights(Buffer<uint> InWeightBuffer, uint InWeightOffset)
+{
+	FGPUSkinIndexAndWeight Out = (FGPUSkinIndexAndWeight)0;
+
+	// Indices
+#if GPUSKIN_BONE_INDEX_UINT16
+	Out.BlendIndices.xy = UnpackUInt2FromUInt(InWeightBuffer[InWeightOffset + GPUSKIN_VB_OFFSET_INFLUENCEBONES    ]);
+	Out.BlendIndices.zw = UnpackUInt2FromUInt(InWeightBuffer[InWeightOffset + GPUSKIN_VB_OFFSET_INFLUENCEBONES + 1]);
+	#if GPUSKIN_USE_EXTRA_INFLUENCES
+	Out.BlendIndices2.xy = UnpackUInt2FromUInt(InWeightBuffer[InWeightOffset + GPUSKIN_VB_OFFSET_INFLUENCEBONES + 2]);
+	Out.BlendIndices2.zw = UnpackUInt2FromUInt(InWeightBuffer[InWeightOffset + GPUSKIN_VB_OFFSET_INFLUENCEBONES + 3]);
+	#endif
+#else // GPUSKIN_BONE_INDEX_UINT16
+	Out.BlendIndices  = UnpackUintToUint4(InWeightBuffer[InWeightOffset + GPUSKIN_VB_OFFSET_INFLUENCEBONES]);
+	#if GPUSKIN_USE_EXTRA_INFLUENCES
+	Out.BlendIndices2 = UnpackUintToUint4(InWeightBuffer[InWeightOffset + GPUSKIN_VB_OFFSET_INFLUENCEBONES + 1]);
+	#endif
+#endif // GPUSKIN_BONE_INDEX_UINT16
+
+	// Weights
+#if GPUSKIN_BONE_WEIGHTS_UINT16
+	Out.BlendWeights.xy = UnpackUnorm2x16(InWeightBuffer[InWeightOffset + GPUSKIN_VB_OFFSET_INFLUENCEWEIGHTS    ]);
+	Out.BlendWeights.zw = UnpackUnorm2x16(InWeightBuffer[InWeightOffset + GPUSKIN_VB_OFFSET_INFLUENCEWEIGHTS + 1]);
+	#if GPUSKIN_USE_EXTRA_INFLUENCES
+	Out.BlendWeights2.xy = UnpackUnorm2x16(InWeightBuffer[InWeightOffset + GPUSKIN_VB_OFFSET_INFLUENCEWEIGHTS + 2]);
+	Out.BlendWeights2.zw = UnpackUnorm2x16(InWeightBuffer[InWeightOffset + GPUSKIN_VB_OFFSET_INFLUENCEWEIGHTS + 3]);
+	#endif
+#else // GPUSKIN_BONE_WEIGHTS_UINT16
+	Out.BlendWeights  = UnpackRGBA8(InWeightBuffer[InWeightOffset + GPUSKIN_VB_OFFSET_INFLUENCEWEIGHTS]);
+	#if GPUSKIN_USE_EXTRA_INFLUENCES
+	Out.BlendWeights2 = UnpackRGBA8(InWeightBuffer[InWeightOffset + GPUSKIN_VB_OFFSET_INFLUENCEWEIGHTS + 1]);
+	#endif
+#endif // GPUSKIN_BONE_WEIGHTS_UINT16
+
+	return Out;
+}
+
+#define FBoneMatrix float3x4
+
+FBoneMatrix InternalGetBoneMatrix(Buffer<float4> InBoneMatrices, uint In) 
+{ 
+	return FBoneMatrix(InBoneMatrices[In*3], InBoneMatrices[In*3 + 1], InBoneMatrices[In*3 + 2]); 
+} 
+
+// Get bone matrix with a bone matrix buffer
+FBoneMatrix ComputeBoneMatrixWithLimitedInfluences(Buffer<float4> InBoneMatrices, FGPUSkinIndexAndWeight In, bool bExtraInfluences = GPUSKIN_USE_EXTRA_INFLUENCES)
+{
+	FBoneMatrix 
+	BoneMatrix  = In.BlendWeights.x * InternalGetBoneMatrix(InBoneMatrices, In.BlendIndices.x);
+	BoneMatrix += In.BlendWeights.y * InternalGetBoneMatrix(InBoneMatrices, In.BlendIndices.y);
+#if !GPUSKIN_LIMIT_2BONE_INFLUENCES
+	BoneMatrix += In.BlendWeights.z * InternalGetBoneMatrix(InBoneMatrices, In.BlendIndices.z);
+	BoneMatrix += In.BlendWeights.w * InternalGetBoneMatrix(InBoneMatrices, In.BlendIndices.w);
+	if (bExtraInfluences)
+	{
+		BoneMatrix += In.BlendWeights2.x * InternalGetBoneMatrix(InBoneMatrices, In.BlendIndices2.x);
+		BoneMatrix += In.BlendWeights2.y * InternalGetBoneMatrix(InBoneMatrices, In.BlendIndices2.y);
+		BoneMatrix += In.BlendWeights2.z * InternalGetBoneMatrix(InBoneMatrices, In.BlendIndices2.z);
+		BoneMatrix += In.BlendWeights2.w * InternalGetBoneMatrix(InBoneMatrices, In.BlendIndices2.w);
+	}
+#endif // GPUSKIN_LIMIT_2BONE_INFLUENCES
+	return BoneMatrix;
+}
+
+FBoneMatrix ComputeBoneMatrixWithUnlimitedInfluences(Buffer<float4> InBoneMatrices, Buffer<uint> InVertexWeights, uint InWeightIndexSize, uint InBlendOffsetCount)
+{
+	int NumBoneInfluences = InBlendOffsetCount & 0xff;
+	int StreamOffset = InBlendOffsetCount >> 8;
+	int BoneIndexSize = InWeightIndexSize & 0xff;
+	int BoneWeightSize = InWeightIndexSize >> 8;
+	int WeightsOffset = StreamOffset + (BoneIndexSize * NumBoneInfluences);
+
+	FBoneMatrix BoneMatrix = FBoneMatrix(float4(0,0,0,0), float4(0,0,0,0), float4(0,0,0,0));
+	for (int InfluenceIdx = 0; InfluenceIdx < NumBoneInfluences; InfluenceIdx++)
+	{
+		int BoneIndexOffset = StreamOffset + (BoneIndexSize * InfluenceIdx);
+		int BoneWeightOffset = WeightsOffset + InfluenceIdx * BoneWeightSize;
+		int BoneIndex = InVertexWeights[BoneIndexOffset];
+		if (BoneIndexSize > 1)
+		{
+			BoneIndex = InVertexWeights[BoneIndexOffset + 1] << 8 | BoneIndex;
+			//@todo-lh: Workaround to fix issue in SPIRVEmitter of DXC; this block must be inside the if branch
+			float BoneWeight;
+			if (BoneWeightSize > 1)
+			{
+				BoneWeight = float(InVertexWeights[BoneWeightOffset + 1] << 8 | InVertexWeights[BoneWeightOffset]) / 65535.0;
+			}
+			else
+			{
+				BoneWeight = float(InVertexWeights[BoneWeightOffset]) / 255.0f;
+			}
+			BoneMatrix += BoneWeight * InternalGetBoneMatrix(InBoneMatrices, BoneIndex);
+		}
+		else
+		{
+			//@todo-lh: Workaround to fix issue in SPIRVEmitter of DXC; this block must be inside the if branch
+			float BoneWeight;
+			if (BoneWeightSize > 1)
+			{
+				BoneWeight = float(InVertexWeights[BoneWeightOffset + 1] << 8 | InVertexWeights[BoneWeightOffset]) / 65535.0;
+			}
+			else
+			{
+				BoneWeight = float(InVertexWeights[BoneWeightOffset]) / 255.0f;
+			}
+			BoneMatrix += BoneWeight * InternalGetBoneMatrix(InBoneMatrices, BoneIndex);
+		}
+	}
+	return BoneMatrix;
+}
\ No newline at end of file
diff --git a/Engine/Shaders/Private/GpuSkinVertexFactory.ush b/Engine/Shaders/Private/GpuSkinVertexFactory.ush
--- a/Engine/Shaders/Private/GpuSkinVertexFactory.ush
+++ b/Engine/Shaders/Private/GpuSkinVertexFactory.ush
@@ -530,66 +530,19 @@ FBoneMatrix GetBoneMatrix(int Index)
 FBoneMatrix CalcBoneMatrix( FVertexFactoryInput Input )
 {
 #if GPUSKIN_UNLIMITED_BONE_INFLUENCE
-	int NumBoneInfluences = Input.BlendOffsetCount & 0xff;
-	int StreamOffset = Input.BlendOffsetCount >> 8;
-	int BoneIndexSize = InputWeightIndexSize & 0xff;
-	int BoneWeightSize = InputWeightIndexSize >> 8;
-	int WeightsOffset = StreamOffset + (BoneIndexSize * NumBoneInfluences);
-
-	FBoneMatrix BoneMatrix = FBoneMatrix(float4(0,0,0,0), float4(0,0,0,0), float4(0,0,0,0));
-	for (int InfluenceIdx = 0; InfluenceIdx < NumBoneInfluences; InfluenceIdx++)
-	{
-		int BoneIndexOffset = StreamOffset + (BoneIndexSize * InfluenceIdx);
-		int BoneWeightOffset = WeightsOffset + InfluenceIdx * BoneWeightSize;
-		int BoneIndex = InputWeightStream[BoneIndexOffset];
-		if (BoneIndexSize > 1)
-		{
-			BoneIndex = InputWeightStream[BoneIndexOffset + 1] << 8 | BoneIndex;
-			//@todo-lh: Workaround to fix issue in SPIRVEmitter of DXC; this block must be inside the if branch
-			// Q: Is this workaround still necessary?
-			float BoneWeight;
-			if (BoneWeightSize > 1)
-			{
-				BoneWeight = float(InputWeightStream[BoneWeightOffset + 1] << 8 | InputWeightStream[BoneWeightOffset]) / 65535.0;
-			}
-			else
-			{
-				BoneWeight = float(InputWeightStream[BoneWeightOffset]) / 255.0f;
-			}
-			BoneMatrix += BoneWeight * GetBoneMatrix(BoneIndex);
-		}
-		else
-		{
-			//@todo-lh: Workaround to fix issue in SPIRVEmitter of DXC; this block must be inside the if branch
-			float BoneWeight;
-			if (BoneWeightSize > 1)
-			{
-				BoneWeight = float(InputWeightStream[BoneWeightOffset + 1] << 8 | InputWeightStream[BoneWeightOffset]) / 65535.0;
-			}
-			else
-			{
-				BoneWeight = float(InputWeightStream[BoneWeightOffset]) / 255.0f;
-			}
-			BoneMatrix += BoneWeight * GetBoneMatrix(BoneIndex);
-		}
-	}
-#else // GPUSKIN_UNLIMITED_BONE_INFLUENCE
-	FBoneMatrix BoneMatrix = Input.BlendWeights.x * GetBoneMatrix(Input.BlendIndices.x);
-	BoneMatrix += Input.BlendWeights.y * GetBoneMatrix(Input.BlendIndices.y);
-
-#if !GPUSKIN_LIMIT_2BONE_INFLUENCES
-	BoneMatrix += Input.BlendWeights.z * GetBoneMatrix(Input.BlendIndices.z);
-	BoneMatrix += Input.BlendWeights.w * GetBoneMatrix(Input.BlendIndices.w);
-	if (NumBoneInfluencesParam > MAX_INFLUENCES_PER_STREAM)
+	return ComputeBoneMatrixWithUnlimitedInfluences(BoneMatrices, InputWeightStream, InputWeightIndexSize, Input.BlendOffsetCount);
+#else
+	const bool bHasExtraInfluences = NumBoneInfluencesParam > MAX_INFLUENCES_PER_STREAM;
+	FGPUSkinIndexAndWeight In = (FGPUSkinIndexAndWeight)0;
+	In.BlendIndices = Input.BlendIndices;
+	In.BlendWeights = Input.BlendWeights;
+	if (bHasExtraInfluences)
 	{
-		BoneMatrix += Input.BlendWeightsExtra.x * GetBoneMatrix(Input.BlendIndicesExtra.x);
-		BoneMatrix += Input.BlendWeightsExtra.y * GetBoneMatrix(Input.BlendIndicesExtra.y);
-		BoneMatrix += Input.BlendWeightsExtra.z * GetBoneMatrix(Input.BlendIndicesExtra.z);
-		BoneMatrix += Input.BlendWeightsExtra.w * GetBoneMatrix(Input.BlendIndicesExtra.w);
+		In.BlendIndices2 = Input.BlendIndicesExtra;
+		In.BlendWeights2 = Input.BlendWeightsExtra;
 	}
-#endif//GPUSKIN_LIMIT_2BONE_INFLUENCES
-#endif // GPUSKIN_UNLIMITED_BONE_INFLUENCE
-	return BoneMatrix;
+	return ComputeBoneMatrixWithLimitedInfluences(BoneMatrices, In, bHasExtraInfluences);
+#endif
 }
 
 FBoneMatrix GetPreviousBoneMatrix(int Index)
@@ -600,68 +553,22 @@ FBoneMatrix GetPreviousBoneMatrix(int Index)
 	return FBoneMatrix(A,B,C);
 }
 
-FBoneMatrix CalcPreviousBoneMatrix( FVertexFactoryInput Input )
+FBoneMatrix CalcPreviousBoneMatrix(FVertexFactoryInput Input)
 {
 #if GPUSKIN_UNLIMITED_BONE_INFLUENCE
-	int NumBoneInfluences = Input.BlendOffsetCount & 0xff;
-	int StreamOffset = Input.BlendOffsetCount >> 8;
-	int BoneIndexSize = InputWeightIndexSize & 0xff;
-	int BoneWeightSize = InputWeightIndexSize >> 8;
-	int WeightsOffset = StreamOffset + (BoneIndexSize * NumBoneInfluences);
-
-	FBoneMatrix BoneMatrix = FBoneMatrix(float4(0,0,0,0), float4(0,0,0,0), float4(0,0,0,0));
-	for (int InfluenceIdx = 0; InfluenceIdx < NumBoneInfluences; InfluenceIdx++)
-	{
-		int BoneIndexOffset = StreamOffset + (BoneIndexSize * InfluenceIdx);
-		int BoneWeightOffset = WeightsOffset + InfluenceIdx * BoneWeightSize;
-		int BoneIndex = InputWeightStream[BoneIndexOffset];
-		if (BoneIndexSize > 1)
-		{
-			BoneIndex = InputWeightStream[BoneIndexOffset + 1] << 8 | BoneIndex;
-			//@todo-lh: Workaround to fix issue in SPIRVEmitter of DXC; this block must be inside the if branch
-			float BoneWeight;
-			if (BoneWeightSize > 1)
-			{
-				BoneWeight = float(InputWeightStream[BoneWeightOffset + 1] << 8 | InputWeightStream[BoneWeightOffset]) / 65535.0;
-			}
-			else
-			{
-				BoneWeight = float(InputWeightStream[BoneWeightOffset]) / 255.0f;
-			}
-			BoneMatrix += BoneWeight * GetPreviousBoneMatrix(BoneIndex);
-		}
-		else
-		{
-			//@todo-lh: Workaround to fix issue in SPIRVEmitter of DXC; this block must be inside the if branch
-			float BoneWeight;
-			if (BoneWeightSize > 1)
-			{
-				BoneWeight = float(InputWeightStream[BoneWeightOffset + 1] << 8 | InputWeightStream[BoneWeightOffset]) / 65535.0;
-			}
-			else
-			{
-				BoneWeight = float(InputWeightStream[BoneWeightOffset]) / 255.0f;
-			}
-			BoneMatrix += BoneWeight * GetPreviousBoneMatrix(BoneIndex);
-		}
-	}
-#else // GPUSKIN_UNLIMITED_BONE_INFLUENCE
-	FBoneMatrix BoneMatrix = Input.BlendWeights.x * GetPreviousBoneMatrix(Input.BlendIndices.x);
-	BoneMatrix += Input.BlendWeights.y * GetPreviousBoneMatrix(Input.BlendIndices.y);
-
-#if !GPUSKIN_LIMIT_2BONE_INFLUENCES
-	BoneMatrix += Input.BlendWeights.z * GetPreviousBoneMatrix(Input.BlendIndices.z);
-	BoneMatrix += Input.BlendWeights.w * GetPreviousBoneMatrix(Input.BlendIndices.w);
-	if (NumBoneInfluencesParam > MAX_INFLUENCES_PER_STREAM)
+	return ComputeBoneMatrixWithUnlimitedInfluences(PreviousBoneMatrices, InputWeightStream, InputWeightIndexSize, Input.BlendOffsetCount);
+#else
+	const bool bHasExtraInfluences = NumBoneInfluencesParam > MAX_INFLUENCES_PER_STREAM;
+	FGPUSkinIndexAndWeight In = (FGPUSkinIndexAndWeight)0;
+	In.BlendIndices = Input.BlendIndices;
+	In.BlendWeights = Input.BlendWeights;
+	if (bHasExtraInfluences)
 	{
-		BoneMatrix += Input.BlendWeightsExtra.x * GetPreviousBoneMatrix(Input.BlendIndicesExtra.x);
-		BoneMatrix += Input.BlendWeightsExtra.y * GetPreviousBoneMatrix(Input.BlendIndicesExtra.y);
-		BoneMatrix += Input.BlendWeightsExtra.z * GetPreviousBoneMatrix(Input.BlendIndicesExtra.z);
-		BoneMatrix += Input.BlendWeightsExtra.w * GetPreviousBoneMatrix(Input.BlendIndicesExtra.w);
+		In.BlendIndices2 = Input.BlendIndicesExtra;
+		In.BlendWeights2 = Input.BlendWeightsExtra;
 	}
-#endif//GPUSKIN_LIMIT_2BONE_INFLUENCES
-#endif // GPUSKIN_UNLIMITED_BONE_INFLUENCE
-	return BoneMatrix;
+	return ComputeBoneMatrixWithLimitedInfluences(PreviousBoneMatrices, In, bHasExtraInfluences);
+#endif
 }
 
 /** transform position by weighted sum of skinning matrices */
diff --git a/Engine/Shaders/Private/HairStrands/HairStrandsSkinUpdate.usf b/Engine/Shaders/Private/HairStrands/HairStrandsSkinUpdate.usf
--- a/Engine/Shaders/Private/HairStrands/HairStrandsSkinUpdate.usf
+++ b/Engine/Shaders/Private/HairStrands/HairStrandsSkinUpdate.usf
@@ -3,10 +3,11 @@
 #pragma once
 
 #include "/Engine/Public/Platform.ush"
+#include "../GpuSkinCommon.ush"
 
 //----------------------------------------------------------------------------------
 
-uint IndexSize; // bone index size: 1 - 8bit or 2 - 16bit
+uint WeightIndexSize; // bone index size: 1 - 8bit or 2 - 16bit
 uint SectionVertexBaseIndex;
 uint NumTotalVertices;
 uint NumVertexToProcess;
@@ -29,131 +30,6 @@ RWBuffer<float> PrevDeformedPositions;
 
 //----------------------------------------------------------------------------------
 
-#define	GPUSKIN_VB_OFFSET_INFLUENCEBONES	0
-#define	GPUSKIN_VB_OFFSET_INFLUENCEWEIGHTS	((1 + GPUSKIN_USE_EXTRA_INFLUENCES) * (1 + GPUSKIN_BONE_INDEX_UINT16))
-
-//----------------------------------------------------------------------------------
-
-#define FBlendedMatrix float3x4
-
-struct FUnpackedWeight
-{
-#if !GPUSKIN_UNLIMITED_BONE_INFLUENCE
-	int4	BlendIndices;
-#if GPUSKIN_USE_EXTRA_INFLUENCES
-	int4	ExtraIndices;
-#endif
-	float4	BlendWeights;
-#if GPUSKIN_USE_EXTRA_INFLUENCES
-	float4	ExtraWeights;
-#endif
-#endif // GPUSKIN_UNLIMITED_BONE_INFLUENCE
-};
-
-//----------------------------------------------------------------------------------
-
-FBlendedMatrix BlendMatrix(FUnpackedWeight UnpackedWeight, uint VertexIndex, Buffer<float4> InBoneMatrices)
-{
-	const uint MatrixOffset = 0;
-#if GPUSKIN_UNLIMITED_BONE_INFLUENCE
-	uint BlendOffsetCount = WeightLookup[VertexIndex];
-	int NumBoneInfluences = BlendOffsetCount & 0xff;
-	int StreamOffset = BlendOffsetCount >> 8;
-	int WeightsOffset = StreamOffset + (IndexSize * NumBoneInfluences);
-
-	FBlendedMatrix OutBoneMatrix = FBlendedMatrix(float4(0,0,0,0), float4(0,0,0,0), float4(0,0,0,0));
-	for (int InfluenceIdx = 0; InfluenceIdx < NumBoneInfluences; InfluenceIdx++)
-	{
-		int BoneIndexOffset = StreamOffset + (IndexSize * InfluenceIdx);
-		int BoneIndex = VertexWeights[BoneIndexOffset];
-		BoneIndex = (IndexSize > 1) ? VertexWeights[BoneIndexOffset + 1] << 8 | BoneIndex : BoneIndex;
-		
-		const uint IndexOffset = (MatrixOffset + BoneIndex) * 3;
-		//@todo-lh: Workaround to fix issue in SPIRVEmitter of DXC; this block must be inside the if branch
-		float BoneWeight = float(VertexWeights[WeightsOffset + InfluenceIdx]) / 255.0f;
-		OutBoneMatrix += BoneWeight * FBlendedMatrix(InBoneMatrices[IndexOffset], InBoneMatrices[IndexOffset+1], InBoneMatrices[IndexOffset+2]);
-	}
-	
-#else // GPUSKIN_UNLIMITED_BONE_INFLUENCE
-	uint IndexOffsetX = (MatrixOffset + UnpackedWeight.BlendIndices.x) * 3;
-	uint IndexOffsetY = (MatrixOffset + UnpackedWeight.BlendIndices.y) * 3;
-	uint IndexOffsetZ = (MatrixOffset + UnpackedWeight.BlendIndices.z) * 3;
-	uint IndexOffsetW = (MatrixOffset + UnpackedWeight.BlendIndices.w) * 3;
-
-	FBlendedMatrix 
-	OutBoneMatrix  = UnpackedWeight.BlendWeights.x * FBlendedMatrix(InBoneMatrices[IndexOffsetX], InBoneMatrices[IndexOffsetX+1], InBoneMatrices[IndexOffsetX+2]);
-	OutBoneMatrix += UnpackedWeight.BlendWeights.y * FBlendedMatrix(InBoneMatrices[IndexOffsetY], InBoneMatrices[IndexOffsetY+1], InBoneMatrices[IndexOffsetY+2]);
-	OutBoneMatrix += UnpackedWeight.BlendWeights.z * FBlendedMatrix(InBoneMatrices[IndexOffsetZ], InBoneMatrices[IndexOffsetZ+1], InBoneMatrices[IndexOffsetZ+2]);
-	OutBoneMatrix += UnpackedWeight.BlendWeights.w * FBlendedMatrix(InBoneMatrices[IndexOffsetW], InBoneMatrices[IndexOffsetW+1], InBoneMatrices[IndexOffsetW+2]);
-  #if GPUSKIN_USE_EXTRA_INFLUENCES
-	IndexOffsetX = (MatrixOffset + UnpackedWeight.ExtraIndices.x) * 3;
-	IndexOffsetY = (MatrixOffset + UnpackedWeight.ExtraIndices.y) * 3;
-	IndexOffsetZ = (MatrixOffset + UnpackedWeight.ExtraIndices.z) * 3;
-	IndexOffsetW = (MatrixOffset + UnpackedWeight.ExtraIndices.w) * 3;
-
-	OutBoneMatrix += UnpackedWeight.ExtraWeights.x * FBlendedMatrix(BoneMatrices[IndexOffsetX], BoneMatrices[IndexOffsetX+1], BoneMatrices[IndexOffsetX+2]);
-	OutBoneMatrix += UnpackedWeight.ExtraWeights.y * FBlendedMatrix(BoneMatrices[IndexOffsetY], BoneMatrices[IndexOffsetY+1], BoneMatrices[IndexOffsetY+2]);
-	OutBoneMatrix += UnpackedWeight.ExtraWeights.z * FBlendedMatrix(BoneMatrices[IndexOffsetZ], BoneMatrices[IndexOffsetZ+1], BoneMatrices[IndexOffsetZ+2]);
-	OutBoneMatrix += UnpackedWeight.ExtraWeights.w * FBlendedMatrix(BoneMatrices[IndexOffsetW], BoneMatrices[IndexOffsetW+1], BoneMatrices[IndexOffsetW+2]);
-  #endif
-#endif // GPUSKIN_UNLIMITED_BONE_INFLUENCE
-	return OutBoneMatrix;
-}
-
-//----------------------------------------------------------------------------------
-
-FUnpackedWeight UnpackWeight(uint WeightOffset)
-{
-	FUnpackedWeight Unpacked;
-
-#if !GPUSKIN_UNLIMITED_BONE_INFLUENCE
-	uint BlendIndices = VertexWeights[WeightOffset + GPUSKIN_VB_OFFSET_INFLUENCEBONES];
-#if GPUSKIN_BONE_INDEX_UINT16
-	Unpacked.BlendIndices.x = BlendIndices & 0xffff;
-	Unpacked.BlendIndices.y = BlendIndices >> 16 & 0xffff;
-	BlendIndices = VertexWeights[WeightOffset + GPUSKIN_VB_OFFSET_INFLUENCEBONES + 1];
-	Unpacked.BlendIndices.z = BlendIndices & 0xffff;
-	Unpacked.BlendIndices.w = BlendIndices >> 16 & 0xffff;
-#if GPUSKIN_USE_EXTRA_INFLUENCES
-	BlendIndices = VertexWeights[WeightOffset + GPUSKIN_VB_OFFSET_INFLUENCEBONES + 2];
-	Unpacked.ExtraIndices.x = BlendIndices & 0xffff;
-	Unpacked.ExtraIndices.y = BlendIndices >> 16 & 0xffff;
-	BlendIndices = VertexWeights[WeightOffset + GPUSKIN_VB_OFFSET_INFLUENCEBONES + 3];
-	Unpacked.ExtraIndices.z = BlendIndices & 0xffff;
-	Unpacked.ExtraIndices.w = BlendIndices >> 16 & 0xffff;
-#endif
-#else // GPUSKIN_BONE_INDEX_UINT16
-	Unpacked.BlendIndices.x = BlendIndices & 0xff;
-	Unpacked.BlendIndices.y = BlendIndices >> 8 & 0xff;
-	Unpacked.BlendIndices.z = BlendIndices >> 16 & 0xff;
-	Unpacked.BlendIndices.w = BlendIndices >> 24 & 0xff;
-#if GPUSKIN_USE_EXTRA_INFLUENCES
-	BlendIndices = VertexWeights[WeightOffset + GPUSKIN_VB_OFFSET_INFLUENCEBONES + 1];
-	Unpacked.ExtraIndices.x = BlendIndices & 0xff;
-	Unpacked.ExtraIndices.y = BlendIndices >> 8 & 0xff;
-	Unpacked.ExtraIndices.z = BlendIndices >> 16 & 0xff;
-	Unpacked.ExtraIndices.w = BlendIndices >> 24 & 0xff;
-#endif
-#endif // GPUSKIN_BONE_INDEX_UINT16
-	uint BlendWeights = VertexWeights[WeightOffset + GPUSKIN_VB_OFFSET_INFLUENCEWEIGHTS];
-	Unpacked.BlendWeights.x = float(BlendWeights & 0xff) / 255.0f;
-	Unpacked.BlendWeights.y = float(BlendWeights >> 8 & 0xff) / 255.0f;
-	Unpacked.BlendWeights.z = float(BlendWeights >> 16 & 0xff) / 255.0f;
-	Unpacked.BlendWeights.w = float(BlendWeights >> 24 & 0xff) / 255.0f;
-#if GPUSKIN_USE_EXTRA_INFLUENCES
-	BlendWeights = VertexWeights[WeightOffset + GPUSKIN_VB_OFFSET_INFLUENCEWEIGHTS + 1];
-	Unpacked.ExtraWeights.x = float(BlendWeights & 0xff) / 255.0f;
-	Unpacked.ExtraWeights.y = float(BlendWeights >> 8 & 0xff) / 255.0f;
-	Unpacked.ExtraWeights.z = float(BlendWeights >> 16 & 0xff) / 255.0f;
-	Unpacked.ExtraWeights.w = float(BlendWeights >> 24 & 0xff) / 255.0f;
-#endif
-#endif // GPUSKIN_UNLIMITED_BONE_INFLUENCE
-
-	return Unpacked;
-}
-
-//----------------------------------------------------------------------------------
-
 [numthreads(64,1,1)]
 void UpdateSkinPositionCS(uint3 GroupID : SV_GroupID,
  uint3 DispatchThreadID : SV_DispatchThreadID,
@@ -165,31 +41,43 @@ void UpdateSkinPositionCS(uint3 GroupID : SV_GroupID,
 		const uint PositionOffset = VertexIndex * 3;
 		const uint WeightOffset = VertexIndex * WeightStride / 4;
 
-		const FUnpackedWeight UnpackedWeight = UnpackWeight(WeightOffset);
+		#if !GPUSKIN_UNLIMITED_BONE_INFLUENCE
+		const FGPUSkinIndexAndWeight IndicesAndWeights = GetBlendIndicesAndWeights(VertexWeights, WeightOffset);
+		#endif
 
 		const float4 InputPosition = float4(
 			RestPositions[PositionOffset],
-			RestPositions[PositionOffset + 1],
-			RestPositions[PositionOffset + 2],
+			RestPositions[PositionOffset+1],
+			RestPositions[PositionOffset+2],
 			1.0);
 
 		// Current position
 		{
-			const FBlendedMatrix BlendedMatrix = BlendMatrix(UnpackedWeight, VertexIndex, BoneMatrices);
+			#if GPUSKIN_UNLIMITED_BONE_INFLUENCE
+			const FBoneMatrix BlendedMatrix = ComputeBoneMatrixWithUnlimitedInfluences(BoneMatrices, VertexWeights, WeightIndexSize, WeightLookup[VertexIndex]);
+			#else
+			const FBoneMatrix BlendedMatrix = ComputeBoneMatrixWithLimitedInfluences(BoneMatrices, IndicesAndWeights);
+			#endif
+
 			const float3 OutputPosition =  mul(BlendedMatrix, InputPosition);
-			DeformedPositions[PositionOffset    ] = OutputPosition.x;
-			DeformedPositions[PositionOffset + 1] = OutputPosition.y;
-			DeformedPositions[PositionOffset + 2] = OutputPosition.z;
+			DeformedPositions[PositionOffset  ] = OutputPosition.x;
+			DeformedPositions[PositionOffset+1] = OutputPosition.y;
+			DeformedPositions[PositionOffset+2] = OutputPosition.z;
 		}
 
 		// Previous position
 	#if PERMUTATION_PREV
 		{
-			const FBlendedMatrix PrevBlendedMatrix = BlendMatrix(UnpackedWeight, VertexIndex, PrevBoneMatrices);
+			#if GPUSKIN_UNLIMITED_BONE_INFLUENCE
+			const FBoneMatrix PrevBlendedMatrix = ComputeBoneMatrixWithUnlimitedInfluences(PrevBoneMatrices, VertexWeights, WeightIndexSize, WeightLookup[VertexIndex]);
+			#else
+			const FBoneMatrix PrevBlendedMatrix = ComputeBoneMatrixWithLimitedInfluences(PrevBoneMatrices, IndicesAndWeights);
+			#endif
+
 			const float3 OutputPosition = mul(PrevBlendedMatrix, InputPosition);
-			PrevDeformedPositions[PositionOffset    ] = OutputPosition.x;
-			PrevDeformedPositions[PositionOffset + 1] = OutputPosition.y;
-			PrevDeformedPositions[PositionOffset + 2] = OutputPosition.z;
+			PrevDeformedPositions[PositionOffset  ] = OutputPosition.x;
+			PrevDeformedPositions[PositionOffset+1] = OutputPosition.y;
+			PrevDeformedPositions[PositionOffset+2] = OutputPosition.z;
 		}
 	#endif
 	}
diff --git a/Engine/Shaders/Private/PackUnpack.ush b/Engine/Shaders/Private/PackUnpack.ush
--- a/Engine/Shaders/Private/PackUnpack.ush
+++ b/Engine/Shaders/Private/PackUnpack.ush
@@ -143,6 +143,27 @@ float4 UnpackRGBA8(uint In)
 #endif
 }
 
+uint PackUint4ToUint(uint4 In)
+{
+	// TODO use BitAlignU32
+	uint r = (In.r & 0xFFu);
+	uint g = (In.g & 0xFFu) << 8u;
+	uint b = (In.b & 0xFFu) << 16u;
+	uint a = (In.a & 0xFFu) << 24u;
+	return r | g | b | a;
+}
+
+uint4 UnpackUintToUint4(uint In)
+{
+	// TODO use BitAlignU32
+	uint4 Out;
+	Out.r = (In >> 0u)  & 0xFFu;
+	Out.g = (In >> 8u)  & 0xFFu;
+	Out.b = (In >> 16u) & 0xFFu;
+	Out.a = (In >> 24u) & 0xFFu;
+	return Out;
+}
+
 //
 //	float <-> f10
 //
