LFV - shader VGPRs/ALU optimisations

#prelfight none

diff --git a/Engine/Shaders/Private/LocalFogVolumes/LocalFogVolumeCommon.ush b/Engine/Shaders/Private/LocalFogVolumes/LocalFogVolumeCommon.ush
--- a/Engine/Shaders/Private/LocalFogVolumes/LocalFogVolumeCommon.ush
+++ b/Engine/Shaders/Private/LocalFogVolumes/LocalFogVolumeCommon.ush
@@ -134,7 +134,6 @@ uint LFVPackTile(uint2 TileCoord)
 ///////////////////////////////////////////////////////////////////////////////
 // Local Fog Volume Integrals
 
-// CameraPos can be behind of RayStart if the camera is out of the volume
 struct FFogData
 {
 	float IntegratedLuminanceFactor;
@@ -142,25 +141,16 @@ struct FFogData
 };
 
 // Height fog is evaluated according to spatially varying density defined as "Density * Exp[-HeightFogFalloff * z]", where density is in fact Extinction.
-FFogData LocalFogVolumeEvaluateAnalyticalIntegral(FLocalFogVolumeGPUInstanceData FogInstance, float3 CameraPos, float3 RayStart, float3 RayDir, float RayLength)
+FFogData LocalFogVolumeEvaluateAnalyticalIntegral(FLocalFogVolumeGPUInstanceData FogInstance, float3 RayStartU, float3 RayDirU, float RayLengthU)
 {
 	FFogData FogData;
 
-	float3 RayOriU = mul(float4(RayStart, 1.0f), FogInstance.InvTransform).xyz;
-	float3 RayIntU = mul(float4(RayStart + RayDir * RayLength, 1.0f), FogInstance.InvTransform).xyz;
-	float3 RayDirU = mul(float4(RayDir, 0.0f), FogInstance.InvTransform).xyz;
-	RayDirU = normalize(RayDirU);
-
-	float RayLengthU = length(RayIntU - RayOriU);
-
-	float StartHeight = RayOriU.z - FogInstance.HeightFogOffset; // accoutn for the fog offset.
-
 	float OpticalDepth = 0.0f;
 
 	if (FogInstance.RadialFogExtinction > 0.0f)
 	{
 		float3 SphereCenter = 0.0f;
-		float3 VolumeCenterToRayO = (RayOriU - SphereCenter);
+		float3 VolumeCenterToRayO = (RayStartU - SphereCenter);
 
 		float b = dot(RayDirU, VolumeCenterToRayO);
 		float c = dot(VolumeCenterToRayO, VolumeCenterToRayO) - 1.0f;
@@ -185,7 +175,7 @@ FFogData LocalFogVolumeEvaluateAnalyticalIntegral(FLocalFogVolumeGPUInstanceData
 
 	if (FogInstance.HeightFogExtinction > 0.0f)
 	{
-		float StartHeight = RayOriU.z - FogInstance.HeightFogOffset; // accoutn for the fog offset.
+		float StartHeight = RayStartU.z - FogInstance.HeightFogOffset; // account for the fog offset.
 
 	#if 0
 		// Original integral, less artefact
@@ -262,9 +252,7 @@ float4 GetLFVInstanceContribution(
 	uint InstanceIndex,
 	float3 CamRayWorldOrigin,
 	float3 CamRayWorldDir,
-	float3 CamRayTranslatedWorldOrigin,
-	float3 CamRayTranslatedWorldDir,
-	float3 DepthBufferTranslatedWorldPos)
+	float3 DepthBufferWorldPos)
 {
 	FLocalFogVolumeGPUInstanceData FogInstance = GetLocalFogVolumeGPUInstanceData(InstanceIndex);
 
@@ -272,44 +260,37 @@ float4 GetLFVInstanceContribution(
 	float Transmittance = 1.0f;
 
 	// The "U" prefix is for all compuation done in the Unit Sphere space.
-	float3 RayPosU = mul(float4(CamRayWorldOrigin, 1.0f), FogInstance.InvTransform).xyz;
-	float3 RayDirU = mul(float4(CamRayWorldDir,    0.0f), FogInstance.InvTransform).xyz; // Try to use InverseTranspose to correctly handle direction transform with scaling.
+	float3 DepthBufferPosU	= mul(float4(DepthBufferWorldPos, 1.0f),	FogInstance.InvTransform).xyz;
+	float3 RayPosU			= mul(float4(CamRayWorldOrigin,   1.0f),	FogInstance.InvTransform).xyz;
+	float3 RayDirU			= mul(float4(CamRayWorldDir,      0.0f),	FogInstance.InvTransform).xyz;
 	RayDirU.xyz = normalize(RayDirU);
 
 	float2 TsU = RayIntersectSphere(RayPosU, RayDirU, float4(0.0, 0.0, 0.0, 1.0));
-	if (any(TsU > 0.0))
-	{
-		float LengthD = length(DepthBufferTranslatedWorldPos - CamRayTranslatedWorldOrigin);
 
-		float3 P0U = RayPosU + max(0.0, TsU.x) * RayDirU;
-		float3 P1U = RayPosU + max(0.0, TsU.y) * RayDirU;
-		// Computations are all done with Space = world space
-		float3 P0Space = mul(float4(P0U, 1.0), FogInstance.Transform).xyz;
-		float3 P1Space = mul(float4(P1U, 1.0), FogInstance.Transform).xyz;
-		float3 CamRayOriginSpace = CamRayWorldOrigin;
-		float3 CamRayDirSpace = CamRayWorldDir;
+	//Clamp to make sure we only consider intersection in front of the camera
+	TsU = max(0.0f.xx, TsU);
 
-		float Length0 = length(P0Space - CamRayOriginSpace);
-		float Length1 = length(P1Space - CamRayOriginSpace);
+	// Clamp the traced distance to the depth buffer and compute world length.
+	float LengthDepthBufferU = length(DepthBufferPosU - RayPosU);
+	TsU = min(TsU, LengthDepthBufferU);
 
+	if (any(TsU > 0.0))
+	{
+		// Make the ray start after the volumetric fog far distance
 		if (LFVRenderInVolumetricFog)
 		{
 			// Volumetric fog covers up to MaxDistance along ViewZ, exclude analytical fog from this range
 			float CosAngle = dot(CamRayWorldDir, View.ViewForward);
 			float InvCosAngle = (CosAngle > 0.0001) ? rcp(CosAngle) : 0;
-			float ExcludeDistance = View.VolumetricFogMaxDistance * InvCosAngle;
-			Length0 = max(Length0, ExcludeDistance);
-			Length1 = max(Length1, ExcludeDistance);
+			float ExcludeDistance = View.VolumetricFogMaxDistance * InvCosAngle / FogInstance.UniformScale;
+			TsU = max(TsU, ExcludeDistance);
 		}
 
-		// Clamp the traced distance to depth and compute world length.
-		Length0 = min(LengthD, Length0);
-		Length1 = min(LengthD, Length1);
-		float RayTracedLength = max(0.0, abs(Length1 - Length0));
+		float RayTracedLengthU = max(0.0, abs(TsU.y - TsU.x));
 
-		if (RayTracedLength > 0.0)
+		if (RayTracedLengthU > 0.0)
 		{
-			FFogData FogData = LocalFogVolumeEvaluateAnalyticalIntegral(FogInstance, CamRayWorldOrigin, CamRayWorldOrigin + CamRayWorldDir * Length0, CamRayWorldDir, RayTracedLength);
+			FFogData FogData = LocalFogVolumeEvaluateAnalyticalIntegral(FogInstance, RayPosU + RayDirU * TsU.x, RayDirU, RayTracedLengthU);
 
 			LuminanceColor = LocalFogVolumeEvaluateInScattering(FogInstance, FogData, CamRayWorldDir);
 			Transmittance = 1.0 - FogData.Coverage;
@@ -338,32 +319,27 @@ float4 GetLFVInstanceContribution(
 float4 GetLFVContribution(
 	in ViewState ResolvedView,
 	in uint2 TilePos,
-	in float3 TranslatedWorldPosition,
-	inout uint LFVCount)
+	in float3 DepthBufferTranslatedWorldPos,
+	inout uint OutLFVCount)
 {
-	float3 CamRayTranslatedWorldOrigin = ResolvedView.TranslatedWorldCameraOrigin;
-	float3 PreViewTranslation = LWCHackToFloat(ResolvedView.PreViewTranslation);
+	const float3 PreViewTranslation			= LWCHackToFloat(ResolvedView.PreViewTranslation);
+	const float3 CamRayTranslatedWorldOrigin= ResolvedView.TranslatedWorldCameraOrigin;
+	const float3 CamRayWorldOrigin			= CamRayTranslatedWorldOrigin - PreViewTranslation;
+	const float3 CamRayWorldDir				= normalize(DepthBufferTranslatedWorldPos - CamRayTranslatedWorldOrigin);
+	const float3 DepthBufferWorldPos		= DepthBufferTranslatedWorldPos - PreViewTranslation;
 
-	LFVCount = LFVTileDataTex[uint3(TilePos, 0)];
+	OutLFVCount = LFVTileDataTex[uint3(TilePos, 0)];
 	const float4 DefaultColorTransmittance = float4(0.0, 0.0, 0.0, 1.0);
 	float4 LFVContribution = DefaultColorTransmittance;
-	for (int LFVIndex = 0; LFVIndex < LFVCount; ++LFVIndex)
+	for (int LFVIndex = 0; LFVIndex < OutLFVCount; ++LFVIndex)
 	{
 		uint InstanceIndex = LFVTileDataTex[uint3(TilePos, 1 + LFVIndex)];
 
-		float3 CamRayWorldOrigin = CamRayTranslatedWorldOrigin - PreViewTranslation;
-		float3 CamRayTranslatedWorldDir = normalize(TranslatedWorldPosition - CamRayTranslatedWorldOrigin);
-		float3 CamRayWorldDir = CamRayTranslatedWorldDir;
-
-		float3 DepthBufferTranslatedWorldPos = TranslatedWorldPosition.xyz;
-
 		float4 LFVLuminanceColorTransmittance = GetLFVInstanceContribution(
 			InstanceIndex,
 			CamRayWorldOrigin,
 			CamRayWorldDir,
-			CamRayTranslatedWorldOrigin,
-			CamRayTranslatedWorldDir,
-			DepthBufferTranslatedWorldPos);
+			DepthBufferWorldPos);
 
 		LFVContribution = float4(LFVContribution.rgb * LFVLuminanceColorTransmittance.a + LFVLuminanceColorTransmittance.rgb, LFVContribution.a * LFVLuminanceColorTransmittance.a);
 	}
@@ -375,8 +351,8 @@ float4 GetLFVContribution(
 	in uint2 TilePos,
 	in float3 TranslatedWorldPosition)
 {
-	uint LFVCount = 0;
-	return GetLFVContribution(ResolvedView, TilePos, TranslatedWorldPosition, LFVCount);
+	uint OutLFVCount = 0;
+	return GetLFVContribution(ResolvedView, TilePos, TranslatedWorldPosition, OutLFVCount);
 }
 
 
diff --git a/Engine/Shaders/Private/LocalFogVolumes/LocalFogVolumeSplat.usf b/Engine/Shaders/Private/LocalFogVolumes/LocalFogVolumeSplat.usf
--- a/Engine/Shaders/Private/LocalFogVolumes/LocalFogVolumeSplat.usf
+++ b/Engine/Shaders/Private/LocalFogVolumes/LocalFogVolumeSplat.usf
@@ -91,14 +91,13 @@ void LocalFogVolumeSplatPS(
 #endif
 
 	float3 DepthBufferTranslatedWorldPos = SvPositionToTranslatedWorld(float4(SVPos.xy, DeviceZ, 1.0));
+	float3 DepthBufferWorldPos = DepthBufferTranslatedWorldPos - LWCHackToFloat(ResolvedView.PreViewTranslation);
 
 	float4 LuminanceColorTransmittance = GetLFVInstanceContribution(
 		InstanceIndex,
 		CamRayWorldOrigin,
 		CamRayWorldDir,
-		CamRayTranslatedWorldOrigin,
-		CamRayTranslatedWorldDir,
-		DepthBufferTranslatedWorldPos);
+		DepthBufferWorldPos);
 
 	OutColor = float4(
 		LuminanceColorTransmittance.rgb * PrimaryView.PreExposure,
@@ -167,14 +166,10 @@ void LocalFogVolumeTiledPS(
 
 	uint LFVCount = 0;
 	float4 LFVContribution = GetLFVContribution(ResolvedView, TilePos, DepthBufferTranslatedWorldPos, LFVCount);
-	
-	OutColor = float4(
-		LFVContribution.rgb * PrimaryView.PreExposure,
-		LFVContribution.a);
 
 	bool bDoClip = true;
 #if USE_DEVELOPMENT_SHADERS
-	bDoClip = LocalFogVolumeTileDebug == 2;
+	bDoClip = LocalFogVolumeTileDebug != 1;
 #endif
 
 	const float4 IdentityLuminanceTransmittance = float4(0.0, 0.0, 0.0, 1.0);
@@ -183,6 +178,10 @@ void LocalFogVolumeTiledPS(
 		clip(-1.0f);
 	}
 
+	OutColor = float4(
+		LFVContribution.rgb * PrimaryView.PreExposure,
+		LFVContribution.a);
+
 #if USE_DEVELOPMENT_SHADERS
 	if (LocalFogVolumeTileDebug > 0)
 	{
