Path Tracing: SvPosition was not set correctly, leading to inconsistent output from the ScreenPosition material node

This change accounts for the viewport dimensions and centralizes the logic between Ray and Path Tracing shaders

diff --git a/Engine/Shaders/Private/PathTracing/PathTracingGPULightmassMaterialHitShader.usf b/Engine/Shaders/Private/PathTracing/PathTracingGPULightmassMaterialHitShader.usf
--- a/Engine/Shaders/Private/PathTracing/PathTracingGPULightmassMaterialHitShader.usf
+++ b/Engine/Shaders/Private/PathTracing/PathTracingGPULightmassMaterialHitShader.usf
@@ -67,7 +67,7 @@ RAY_TRACING_ENTRY_CLOSEST_HIT(GPULightmassMaterialCHS,
 	ResolvedView = ResolveView();
 
 	const float3 TranslatedWorldPosition = TranslatedWorldRayOrigin() + RayTCurrent() * WorldRayDirection();
-	const float4 SvPosition = mul(float4(TranslatedWorldPosition, 1.0f), ResolvedView.TranslatedWorldToClip);
+	const float4 SvPosition = TranslatedWorldPositionToSvPosition(TranslatedWorldPosition);
 
 	CurrentPayloadInputFlags = PackedPayload.GetFlags();
 
@@ -415,7 +415,7 @@ RAY_TRACING_ENTRY_ANY_HIT(GPULightmassMaterialAHS,
 	ResolvedView = ResolveView();
 
 	const float3 TranslatedWorldPosition = TranslatedWorldRayOrigin() + RayTCurrent() * WorldRayDirection();
-	const float4 SvPosition = mul(float4(TranslatedWorldPosition, 1.0f), ResolvedView.TranslatedWorldToClip);
+	const float4 SvPosition = TranslatedWorldPositionToSvPosition(TranslatedWorldPosition);
 
 	CurrentPayloadInputFlags = PackedPayload.GetFlags();
 
diff --git a/Engine/Shaders/Private/PathTracing/PathTracingMaterialHitShader.usf b/Engine/Shaders/Private/PathTracing/PathTracingMaterialHitShader.usf
--- a/Engine/Shaders/Private/PathTracing/PathTracingMaterialHitShader.usf
+++ b/Engine/Shaders/Private/PathTracing/PathTracingMaterialHitShader.usf
@@ -185,7 +185,8 @@ RAY_TRACING_ENTRY_CLOSEST_HIT(PathTracingMaterialCHS,
 	ResolvedView = ResolveView();
 
 	const float3 TranslatedWorldPosition = TranslatedWorldRayOrigin() + RayTCurrent() * WorldRayDirection();
-	const float4 SvPosition = mul(float4(TranslatedWorldPosition, 1.0f), ResolvedView.TranslatedWorldToClip);
+	const float4 SvPosition = TranslatedWorldPositionToSvPosition(TranslatedWorldPosition);
+
 
 	CurrentPayloadInputFlags = PackedPayload.GetFlags();
 
@@ -1342,7 +1343,7 @@ RAY_TRACING_ENTRY_ANY_HIT(PathTracingMaterialAHS,
 	ResolvedView = ResolveView();
 
 	const float3 TranslatedWorldPosition = TranslatedWorldRayOrigin() + RayTCurrent() * WorldRayDirection();
-	const float4 SvPosition = mul(float4(TranslatedWorldPosition, 1.0f), ResolvedView.TranslatedWorldToClip);
+	const float4 SvPosition = TranslatedWorldPositionToSvPosition(TranslatedWorldPosition);
 
 	CurrentPayloadInputFlags = PackedPayload.GetFlags();
 
diff --git a/Engine/Shaders/Private/RayTracing/RayTracingCommon.ush b/Engine/Shaders/Private/RayTracing/RayTracingCommon.ush
--- a/Engine/Shaders/Private/RayTracing/RayTracingCommon.ush
+++ b/Engine/Shaders/Private/RayTracing/RayTracingCommon.ush
@@ -428,6 +428,13 @@ float3 TLASWorldToTranslatedWorld(float3 TLASWorldPosition)
 	return TranslatedWorldPosition;
 }
 
+float4 TranslatedWorldPositionToSvPosition(float3 TranslatedWorldPosition)
+{
+	const float4 ClipSpacePosition = mul(float4(TranslatedWorldPosition, 1.0f), ResolvedView.TranslatedWorldToClip);
+	const float4 SvPosition = float4((ClipSpacePosition.xy / ClipSpacePosition.w * float2(.5f, - .5f) + .5f) * ResolvedView.ViewSizeAndInvSize.xy, ClipSpacePosition.zw);
+	return SvPosition;
+}
+
 // Pack normal vector into an uint32 using Octahedron encoding
 uint PackNormalToUInt32(float3 Normal)
 {
diff --git a/Engine/Shaders/Private/RayTracing/RayTracingDecalMaterialShader.usf b/Engine/Shaders/Private/RayTracing/RayTracingDecalMaterialShader.usf
--- a/Engine/Shaders/Private/RayTracing/RayTracingDecalMaterialShader.usf
+++ b/Engine/Shaders/Private/RayTracing/RayTracingDecalMaterialShader.usf
@@ -129,7 +129,7 @@ RAY_TRACING_ENTRY_CALLABLE(RayTracingDecalMaterialShader, FDecalShaderPayload, P
 	MaterialParameters.TwoSidedSign = 1;
 	MaterialParameters.VertexColor = 1;
 	MaterialParameters.CameraVector = CameraVector;
-	MaterialParameters.SvPosition = mul(float4(TranslatedWorldPosition, 1.0f), ResolvedView.TranslatedWorldToClip);
+	MaterialParameters.SvPosition = TranslatedWorldPositionToSvPosition(TranslatedWorldPosition);
 	MaterialParameters.ScreenPosition = SvPositionToResolvedScreenPosition(MaterialParameters.SvPosition);
 	MaterialParameters.LightVector = SwizzlePos;
 
@@ -185,7 +185,7 @@ RAY_TRACING_ENTRY_CLOSEST_HIT(RayTracingDecalMaterialCHS,
 	ResolvedView = ResolveView();
 
 	const float3 TranslatedWorldPosition = TranslatedWorldRayOrigin() + WorldRayDirection() * RayTCurrent();
-	const float4 SvPosition = mul(float4(TranslatedWorldPosition, 1.0f), ResolvedView.TranslatedWorldToClip);
+	const float4 SvPosition = TranslatedWorldPositionToSvPosition(TranslatedWorldPosition);
 
 #if VF_SUPPORTS_RAYTRACING_PREPARE_MATERIAL_PIXEL_PARAMETERS
 	// this is a newer codepath that is both more flexible and allows for more direct calculation compared to the other codepath
@@ -243,7 +243,7 @@ RAY_TRACING_ENTRY_ANY_HIT(RayTracingDecalMaterialAHS,
 	ResolvedView = ResolveView();
 
 	const float3 TranslatedWorldPosition = TranslatedWorldRayOrigin() + WorldRayDirection() * RayTCurrent();
-	const float4 SvPosition = mul(float4(TranslatedWorldPosition, 1.0f), ResolvedView.TranslatedWorldToClip);
+	const float4 SvPosition = TranslatedWorldPositionToSvPosition(TranslatedWorldPosition);
 
 	FVertexFactoryInterpolantsVSToPS Interpolants;
 
diff --git a/Engine/Shaders/Private/RayTracing/RayTracingMaterialHitShaders.usf b/Engine/Shaders/Private/RayTracing/RayTracingMaterialHitShaders.usf
--- a/Engine/Shaders/Private/RayTracing/RayTracingMaterialHitShaders.usf
+++ b/Engine/Shaders/Private/RayTracing/RayTracingMaterialHitShaders.usf
@@ -480,7 +480,7 @@ RAY_TRACING_ENTRY_CLOSEST_HIT(MaterialCHS,
 
 	FVertexFactoryInterpolantsVSToPS Interpolants;
 	const float3 TranslatedWorldPosition = TranslatedWorldRayOrigin() + RayTCurrent() * WorldRayDirection();
-	float4 SvPosition = mul(float4(TranslatedWorldPosition, 1.0f), ResolvedView.TranslatedWorldToClip);
+	const float4 SvPosition = TranslatedWorldPositionToSvPosition(TranslatedWorldPosition);
 	
 	FRayCone PropagatedCone = PropagateRayCone(PackedPayload.GetRayCone(), 0 /* surface curvature */, RayTCurrent());
 	CalcInterpolants(PropagatedCone, Attributes, Interpolants);
@@ -625,7 +625,7 @@ RAY_TRACING_ENTRY_ANY_HIT(MaterialAHS,
 
 	FVertexFactoryInterpolantsVSToPS Interpolants;
 	const float3 TranslatedWorldPosition = TranslatedWorldRayOrigin() + RayTCurrent() * WorldRayDirection();
-	float4 SvPosition = mul(float4(TranslatedWorldPosition, 1.0f), ResolvedView.TranslatedWorldToClip);
+	const float4 SvPosition = TranslatedWorldPositionToSvPosition(TranslatedWorldPosition);
 
 	CalcInterpolants(PackedPayload.GetRayCone(), Attributes, Interpolants);
 
