Store the SpecularScale in the alpha channel of the MobileLocalLightTextureB of the MobileLocalLight pass
MobileLocalLightTextureB is RGBA8 now from RGB10A2
Direction is stored in RGB8 using oct24 packed as 888

diff --git a/Engine/Shaders/Private/MobileBasePassPixelShader.usf b/Engine/Shaders/Private/MobileBasePassPixelShader.usf
--- a/Engine/Shaders/Private/MobileBasePassPixelShader.usf
+++ b/Engine/Shaders/Private/MobileBasePassPixelShader.usf
@@ -826,7 +826,7 @@ void Main(
 		MergedLightData.SourceLength	= 0;
 		MergedLightData.Tangent			= 0;
 		MergedLightData.SoftSourceRadius= 0;
-		MergedLightData.SpecularScale	= 1.0f;
+		MergedLightData.SpecularScale	= LocalLightB.a;
 		MergedLightData.bSpotLight		= false;
 		MergedLightData.bRectLight		= false;
 		MergedLightData.Color = LocalLightA.rgb;
@@ -835,7 +835,7 @@ void Main(
 		half4 LightAttenuation = half4(1, 1, 1, 1); 
 		float SurfaceShadow = 0;
 
-		MergedLightData.Direction = LocalLightB.rgb * 2 - 1;
+		MergedLightData.Direction = OctahedronToUnitVector(Pack888To1212(LocalLightB.rgb) * 2 - 1);
 
 		FLightAccumulator NewLighting = AccumulateDynamicLighting(MaterialParameters.WorldPosition_CamRelative, CameraVector, GBuffer, 1, GBuffer.ShadingModelID, MergedLightData, LightAttenuation, 0, uint2(0, 0), SurfaceShadow);
 		DirectLighting = LightAccumulator_Add(DirectLighting, NewLighting);
diff --git a/Engine/Shaders/Private/MobileLocalLightsBuffer.usf b/Engine/Shaders/Private/MobileLocalLightsBuffer.usf
--- a/Engine/Shaders/Private/MobileLocalLightsBuffer.usf
+++ b/Engine/Shaders/Private/MobileLocalLightsBuffer.usf
@@ -115,6 +115,7 @@ void Main(float4 SvPosition : SV_POSITION
 	float3 Color = uint3(0.f, 0.f , 0.f);
 	float3 SavedL = float3(0,0,0);
 	float TotalWeight = 0;
+	float SpecularScale = 0.f;
 
 	LOOP
 	for (uint LocalLightListIndex = 0; LocalLightListIndex < NumLocalLights; LocalLightListIndex++)
@@ -130,7 +131,9 @@ void Main(float4 SvPosition : SV_POSITION
 		{
 			continue;
 		}
+
 		FDeferredLightData LightData = ConvertToDeferredLight_Mobile(LocalLight);
+		LightData.SpecularScale	= UnpackLightSpecularScale(LocalLight);
 
 		float3 L = LightData.Direction;	// Already normalized
 		float3 ToLight = L;
@@ -150,6 +153,8 @@ void Main(float4 SvPosition : SV_POSITION
 		
 		if( LightMask > 0)
 		{
+			SpecularScale += LightData.SpecularScale * WeightLight;
+
 			TotalWeight += WeightLight;
 			// For lights with light functions we only only use the compute the direction (OutColorB)
 			// The Color is added in additional draws using MainLightFunction()
@@ -167,12 +172,13 @@ void Main(float4 SvPosition : SV_POSITION
 	{
 		SavedL = (SavedL / TotalWeight);
 		OutColorA = float3(Color.x, Color.y, Color.z);
-		OutColorB = 0.5f * float4(SavedL.x, SavedL.y, SavedL.z, 0)  + float4(0.5f, 0.5f, 0.5f, 0.5f);
+		OutColorB.rgb = Pack1212To888(UnitVectorToOctahedron(SavedL) * 0.5 + 0.5);
+		OutColorB.a = SpecularScale / TotalWeight;
 	}
 	else
 	{
 		OutColorA = float3(0, 0, 0);
-		OutColorB = float4(0.5,0.5,0.5, 0.5);
+		OutColorB = float4(0.5f, 0.5f, 0.5f, 0.f);
 		#if GENERATE_LIGHT_FUNCTION_DEPTH_STENCIL
 			// Clip to avoid writing to stencil
 			clip(-1);
diff --git a/Engine/Source/Runtime/Renderer/Private/SceneTextures.cpp b/Engine/Source/Runtime/Renderer/Private/SceneTextures.cpp
--- a/Engine/Source/Runtime/Renderer/Private/SceneTextures.cpp
+++ b/Engine/Source/Runtime/Renderer/Private/SceneTextures.cpp
@@ -624,7 +624,7 @@ void FSceneTextures::InitializeViewFamily(FRDGBuilder& GraphBuilder, FViewFamily
 		FRDGTextureDesc MobileLocalLightTextureADesc = FRDGTextureDesc::Create2D(Config.Extent, PF_FloatR11G11B10, FClearValueBinding::Transparent, TexCreate_RenderTargetable | TexCreate_ShaderResource);
 		SceneTextures.MobileLocalLightTextureA = GraphBuilder.CreateTexture(MobileLocalLightTextureADesc, TEXT("MobileLocalLightTextureA"));
 
-		FRDGTextureDesc MobileLocalLightTextureBDesc = FRDGTextureDesc::Create2D(Config.Extent, PF_A2B10G10R10, FClearValueBinding::DefaultNormal8Bit, TexCreate_RenderTargetable | TexCreate_ShaderResource);
+		FRDGTextureDesc MobileLocalLightTextureBDesc = FRDGTextureDesc::Create2D(Config.Extent, PF_B8G8R8A8, FClearValueBinding({ 0.5f, 0.5f, 0.5f, 0.f }), TexCreate_RenderTargetable | TexCreate_ShaderResource);
 		SceneTextures.MobileLocalLightTextureB = GraphBuilder.CreateTexture(MobileLocalLightTextureBDesc, TEXT("MobileLocalLightTextureB"));
 	}
 
