Saves register pressure in TSR's reject shading pass

diff --git a/Engine/Shaders/Private/TemporalSuperResolution/TSRDecimateHistory.usf b/Engine/Shaders/Private/TemporalSuperResolution/TSRDecimateHistory.usf
--- a/Engine/Shaders/Private/TemporalSuperResolution/TSRDecimateHistory.usf
+++ b/Engine/Shaders/Private/TemporalSuperResolution/TSRDecimateHistory.usf
@@ -351,7 +351,9 @@ Type SampleTextureCatrom(Texture2DArray<Type> InTexture, float2 TextureUV, float
 // Correct pre exposure of the guide color.
 tsr_halfC CorrectGuideColorExposure(tsr_halfC ReprojectedHistoryGuideColor, const float PreExposureCorrection)
 {
-	tsr_halfC LinearReprojectedHistoryGuideColor = ReprojectedHistoryGuideColor * rcp(tsr_half(1.0) - ReprojectedHistoryGuideColor) * tsr_half(HistoryPreExposureCorrection);
+	tsr_halfC LDRToLinear = min(rcp(saturate(tsr_half(1.0) - ReprojectedHistoryGuideColor)) * tsr_half(HistoryPreExposureCorrection), tsr_half(MaxHalfFloat));
+
+	tsr_halfC LinearReprojectedHistoryGuideColor = ReprojectedHistoryGuideColor * LDRToLinear;
 				
 	ReprojectedHistoryGuideColor = LinearReprojectedHistoryGuideColor * rcp(LinearReprojectedHistoryGuideColor + tsr_half(1.0));
 
diff --git a/Engine/Shaders/Private/TemporalSuperResolution/TSRRejectShading.usf b/Engine/Shaders/Private/TemporalSuperResolution/TSRRejectShading.usf
--- a/Engine/Shaders/Private/TemporalSuperResolution/TSRRejectShading.usf
+++ b/Engine/Shaders/Private/TemporalSuperResolution/TSRRejectShading.usf
@@ -28,6 +28,8 @@
 
 #define CONFIG_MANUAL_LDS_SPILL 0
 
+#define CONFIG_MINIMISE_VGPRS 1
+
 
 //------------------------------------------------------- DERIVED
 	
@@ -146,11 +148,26 @@ static tsr_simd_halfC Debug[DEBUG_ARRAY_SIZE];
 //------------------------------------------------------- GENERAL OPERATIONS
 
 CALL_SITE_DEBUGLOC
-tsr_simd_halfC ChangePerceptionAdd(tsr_simd_halfC SrcColor, const float SourcePerceptionAdd, const float DestPerceptionAdd)
+tsr_simd_halfC RemovePerceptionAdd(tsr_simd_halfC SrcColor, const float SourcePerceptionAdd)
+{
+	tsr_simd_halfC Minus = tsr_simd_halfC::Const(tsr_half(1.0)) - SrcColor;
+	tsr_simd_halfC LDRToLinear = min(tsr_simd_halfC::Const(tsr_half(SourcePerceptionAdd)) * rcp(Minus), tsr_simd_halfC::Const(tsr_half(MaxHalfFloat)));
+
+	return SrcColor * LDRToLinear;
+}
+
+CALL_SITE_DEBUGLOC
+tsr_simd_halfC AddPerceptionAdd(tsr_simd_halfC SrcColor, const float DestPerceptionAdd)
 {
-	SrcColor = SrcColor * rcp(tsr_simd_halfC::Const(tsr_half(SourcePerceptionAdd)) - SrcColor);
 	SrcColor = SrcColor * rcp(SrcColor + tsr_simd_halfC::Const(tsr_half(DestPerceptionAdd)));
+	return SrcColor;
+}
 
+CALL_SITE_DEBUGLOC
+tsr_simd_halfC ChangePerceptionAdd(tsr_simd_halfC SrcColor, const float SourcePerceptionAdd, const float DestPerceptionAdd)
+{
+	SrcColor = RemovePerceptionAdd(SrcColor, SourcePerceptionAdd);
+	SrcColor = AddPerceptionAdd(SrcColor, DestPerceptionAdd);
 	return SrcColor;
 }
 
@@ -722,7 +739,7 @@ void MainCS(
 	tsr_simd_bool bAliasingIsVisible = tsr_simd_bool::Const(false);
 	{
 		// Convert input to LDR
-		tsr_simd_halfC ExposedInput = OriginalInput * rcp(OriginalInput + tsr_simd_halfC::Const(tsr_half(PerceptionAdd)));
+		tsr_simd_halfC ExposedInput = AddPerceptionAdd(OriginalInput, /* DestPerceptionAdd = */ PerceptionAdd);
 	
 		// Change the exposure of history for the shading rejection
 		tsr_simd_halfC ExposedHistory = ChangePerceptionAdd(
@@ -798,9 +815,28 @@ void MainCS(
 				RejectionBlendFinal = select(bResurrectHistory, ResurrectionRejectionBlendFinal, RejectionBlendFinal);
 				RejectionClampBlend = select(bResurrectHistory, ResurrectionRejectionClampBlend, RejectionClampBlend);
 				OriginalHistoryUncertainty = select(bResurrectHistory, OriginalResurrectedHistoryUncertainty, OriginalHistoryUncertainty);
+				
+				#if CONFIG_MINIMISE_VGPRS
+					ExposedHistory = select(tsr_simd_boolC::Vectorize(bResurrectHistory), ExposedResurrectedHistory, ExposedHistory);
+				#else
+					OriginalHistory = select(tsr_simd_boolC::Vectorize(bResurrectHistory), OriginalResurrectedHistory, OriginalHistory);
+				#endif
 			}
 		}
 		#endif
+
+		// Recomputes original off the exposed shading to save register pressure.
+		#if CONFIG_MINIMISE_VGPRS
+		PLATFORM_SPECIFIC_ISOLATE
+		{
+			OriginalInput =  RemovePerceptionAdd(ExposedInput, /* SourcePerceptionAdd = */ PerceptionAdd);
+			
+			OriginalHistory = ChangePerceptionAdd(
+				ExposedHistory,
+				/* SourcePerceptionAdd = */ PerceptionAdd,
+				/* DestPerceptionAdd   = */ 1.0);
+		}
+		#endif
 	}
 	
 	tsr_simd_bool bIsDisoccluded = not(bIsValid);
@@ -809,7 +845,7 @@ void MainCS(
 	tsr_simd_halfC FinalGuide;
 	{
 		// Convert input to LDR
-		tsr_simd_halfC ExposedInput = OriginalInput * rcp(OriginalInput + tsr_simd_halfC::Const(tsr_half(1.0)));
+		tsr_simd_halfC ExposedInput = AddPerceptionAdd(OriginalInput, /* DestPerceptionAdd = */ 1.0);
 	
 
 		tsr_simd_half BlendFinal;
@@ -827,10 +863,6 @@ void MainCS(
 			BlendFinal.SetElement(SimdIndex, max(RejectedBlendFactor, TotalBlend));
 		}
 		
-		#if CONFIG_HISTORY_RESURRECTION
-			OriginalHistory = select(tsr_simd_boolC::Vectorize(bResurrectHistory), OriginalResurrectedHistory, OriginalHistory);
-		#endif
-
 		FinalGuide = ExposedInput * tsr_simd_halfC::Vectorize(BlendFinal) - OriginalHistory * tsr_simd_halfC::Vectorize(BlendFinal - 1.0);
 	}
 	
