LFV - applied LFV on mobile ignoring VolumetricFog because it is not pplied on translucent on there.

diff --git a/Engine/Shaders/Private/LocalFogVolumes/LocalFogVolumeCommon.ush b/Engine/Shaders/Private/LocalFogVolumes/LocalFogVolumeCommon.ush
--- a/Engine/Shaders/Private/LocalFogVolumes/LocalFogVolumeCommon.ush
+++ b/Engine/Shaders/Private/LocalFogVolumes/LocalFogVolumeCommon.ush
@@ -243,10 +243,11 @@ float3 LocalFogVolumeEvaluateInScattering(in FLocalFogVolumeGPUInstanceData FogI
 
 // Returns color transmittance
 float4 GetLFVInstanceContribution(
-	uint InstanceIndex,
-	float3 CamRayWorldOrigin,
-	float3 CamRayWorldDir,
-	float3 DepthBufferWorldPos)
+	in uint InstanceIndex,
+	in float3 CamRayWorldOrigin,
+	in float3 CamRayWorldDir,
+	in float3 DepthBufferWorldPos,
+	in bool bPlatformSupportsVolumetricFogOntranslucent)
 {
 	FLocalFogVolumeGPUInstanceData FogInstance = GetLocalFogVolumeGPUInstanceData(InstanceIndex);
 
@@ -271,7 +272,7 @@ float4 GetLFVInstanceContribution(
 	if (any(TsU > 0.0))
 	{
 		// Make the ray start after the volumetric fog far distance
-		if (LFVRenderInVolumetricFog)
+		if (LFVRenderInVolumetricFog && bPlatformSupportsVolumetricFogOntranslucent)
 		{
 			// Volumetric fog covers up to MaxDistance along ViewZ, exclude analytical fog from this range
 			float CosAngle = dot(CamRayWorldDir, View.ViewForward);
@@ -302,6 +303,7 @@ float4 GetLFVContribution(
 	in ViewState ResolvedView,
 	in uint2 TilePos,
 	in float3 DepthBufferTranslatedWorldPos,
+	in bool bPlatformSupportsVolumetricFogOntranslucent,
 	inout uint OutLFVCount)
 {
 	const float3 PreViewTranslation			= LWCHackToFloat(ResolvedView.PreViewTranslation);
@@ -321,7 +323,8 @@ float4 GetLFVContribution(
 			InstanceIndex,
 			CamRayWorldOrigin,
 			CamRayWorldDir,
-			DepthBufferWorldPos);
+			DepthBufferWorldPos,
+			bPlatformSupportsVolumetricFogOntranslucent);
 
 		LFVContribution = float4(LFVContribution.rgb * LFVLuminanceColorTransmittance.a + LFVLuminanceColorTransmittance.rgb, LFVContribution.a * LFVLuminanceColorTransmittance.a);
 	}
@@ -331,10 +334,11 @@ float4 GetLFVContribution(
 float4 GetLFVContribution(
 	in ViewState ResolvedView,
 	in uint2 TilePos,
-	in float3 TranslatedWorldPosition)
+	in float3 TranslatedWorldPosition,
+	in bool bPlatformSupportsVolumetricFogOntranslucent = true)
 {
 	uint OutLFVCount = 0;
-	return GetLFVContribution(ResolvedView, TilePos, TranslatedWorldPosition, OutLFVCount);
+	return GetLFVContribution(ResolvedView, TilePos, TranslatedWorldPosition, bPlatformSupportsVolumetricFogOntranslucent, OutLFVCount);
 }
 
 
diff --git a/Engine/Shaders/Private/LocalFogVolumes/LocalFogVolumeSplat.usf b/Engine/Shaders/Private/LocalFogVolumes/LocalFogVolumeSplat.usf
--- a/Engine/Shaders/Private/LocalFogVolumes/LocalFogVolumeSplat.usf
+++ b/Engine/Shaders/Private/LocalFogVolumes/LocalFogVolumeSplat.usf
@@ -83,7 +83,8 @@ void LocalFogVolumeTiledPS(
 	uint2 TilePos = LFVUnpackTile(OtherVertexOutput.TileData);
 
 	uint LFVCount = 0;
-	float4 LFVContribution = GetLFVContribution(ResolvedView, TilePos, DepthBufferTranslatedWorldPos, LFVCount);
+	const bool bPlatformSupportsVolumetricFogOntranslucent = false;
+	float4 LFVContribution = GetLFVContribution(ResolvedView, TilePos, DepthBufferTranslatedWorldPos, bPlatformSupportsVolumetricFogOntranslucent, LFVCount);
 
 	bool bDoClip = true;
 #if USE_DEVELOPMENT_SHADERS
diff --git a/Engine/Shaders/Private/MobileBasePassVertexShader.usf b/Engine/Shaders/Private/MobileBasePassVertexShader.usf
--- a/Engine/Shaders/Private/MobileBasePassVertexShader.usf
+++ b/Engine/Shaders/Private/MobileBasePassVertexShader.usf
@@ -109,12 +109,13 @@ void Main(
 	}
 #endif
 
-	// Apply local fog volume to if enabled
+	// Apply local fog volume if enabled. Note: volumetric fog is not applied on mobile so it can result in visual discrepancy.
 #if LOCAL_FOG_VOLUME_ON_TRANSLUCENT
 	float4 VertexClipSpacePosition	= mul(float4(WorldPosition.xyz, 1), ResolvedView.TranslatedWorldToClip);
 	float2 SvPosition				= (VertexClipSpacePosition.xy / VertexClipSpacePosition.w * float2(.5f, -.5f) + .5f) * ResolvedView.ViewSizeAndInvSize.xy;
 	uint2 TilePos					= clamp(uint2(SvPosition.xy / float(LFVTilePixelSize)), uint2(0, 0), LFVTileDataResolution - 1);
-	float4 LFVContribution			= GetLFVContribution(ResolvedView, TilePos, WorldPosition);
+	const bool bPlatformSupportsVolumetricFogOntranslucent = false; // Mobile does not do that so we try to compensate for it by evalauting LFV regularly on translucent.
+	float4 LFVContribution			= GetLFVContribution(ResolvedView, TilePos, WorldPosition, bPlatformSupportsVolumetricFogOntranslucent);
 	VertexFog						= float4(LFVContribution.rgb + VertexFog.rgb * LFVContribution.a, LFVContribution.a * VertexFog.a);
 #endif // LOCAL_FOG_VOLUME_ON_TRANSLUCENT
 
