LFV - fixed disapearing LFV on non-mobile forward opqaue.

diff --git a/Engine/Shaders/Private/BasePassPixelShader.usf b/Engine/Shaders/Private/BasePassPixelShader.usf
--- a/Engine/Shaders/Private/BasePassPixelShader.usf
+++ b/Engine/Shaders/Private/BasePassPixelShader.usf
@@ -181,6 +181,10 @@
 #include "/Engine/Private/Substrate/SubstrateExport.ush"
 #endif
 
+#if (NEEDS_BASEPASS_PIXEL_VOLUMETRIC_FOGGING && OPAQUE_NEEDS_BASEPASS_FOGGING)
+#include "LocalFogVolumes/LocalFogVolumeCommon.ush"
+#endif
+
 #if !FORWARD_SHADING
 
 bool ManualDepthTestEqual(float4 SvPosition, float PSDeviceZWithOffset)
@@ -1306,6 +1310,13 @@ void FPixelShaderInOut_MainPS(
 #if NEEDS_BASEPASS_PIXEL_VOLUMETRIC_FOGGING
 	if (FogStruct.ApplyVolumetricFog > 0) 
 	{
+	#if OPAQUE_NEEDS_BASEPASS_FOGGING
+		float3 TranslatedWorldPosition = SvPositionToResolvedTranslatedWorld(In.SvPosition);
+		uint2 TilePos = clamp(uint2(In.SvPosition.xy / float(LFVTilePixelSize)), uint2(0, 0), LFVTileDataResolution - 1);
+		float4 LocalFogVolumeContrib = GetLFVContribution(ResolvedView, TilePos, TranslatedWorldPosition);
+		HeightFogging = float4(LocalFogVolumeContrib.rgb + HeightFogging.rgb * LocalFogVolumeContrib.a, LocalFogVolumeContrib.a * HeightFogging.a);
+	#endif 
+
 		float3 VolumeUV = ComputeVolumeUV(MaterialParameters.AbsoluteWorldPosition, ResolvedView.WorldToClip);
 		Fogging = CombineVolumetricFog(HeightFogging, VolumeUV, EyeIndex, GBuffer.Depth);
 	}
