LFV - clarified what is happening and fixed density evalution volumetric fog to be quadratic.

diff --git a/Engine/Shaders/Private/LocalFogVolumes/LocalFogVolumeCommon.ush b/Engine/Shaders/Private/LocalFogVolumes/LocalFogVolumeCommon.ush
--- a/Engine/Shaders/Private/LocalFogVolumes/LocalFogVolumeCommon.ush
+++ b/Engine/Shaders/Private/LocalFogVolumes/LocalFogVolumeCommon.ush
@@ -138,7 +138,6 @@ struct FFogData
 	float Coverage;
 };
 
-// Height fog is evaluated according to spatially varying density defined as "Density * Exp[-HeightFogFalloff * z]", where density is in fact Extinction.
 FFogData LocalFogVolumeEvaluateAnalyticalIntegral(FLocalFogVolumeGPUInstanceData FogInstance, float3 RayStartU, float3 RayDirU, float RayLengthU)
 {
 	FFogData FogData;
@@ -147,6 +146,10 @@ FFogData LocalFogVolumeEvaluateAnalyticalIntegral(FLocalFogVolumeGPUInstanceData
 
 	if (FogInstance.RadialFogExtinction > 0.0f)
 	{
+		//////////
+		// Radial fog is evaluated according to spatially varying density defined as "Density * (1 - (d*d)/(r*r))" in the unint sphere space, where density is in fact Extinction.
+		// An explanation of it can be found in https://iquilezles.org/articles/spheredensity/
+
 		float3 SphereCenter = 0.0f;
 		float3 VolumeCenterToRayO = (RayStartU - SphereCenter);
 
@@ -164,15 +167,21 @@ FFogData LocalFogVolumeEvaluateAnalyticalIntegral(FLocalFogVolumeGPUInstanceData
 			Length1 = max(Length1, 0.0);
 			Length1 = min(Length1, RayLengthU);
 
-			// Integral of a Extinction that descreases accordingf to square distance from the center of the volume.
-			float Integral0 = -(c * Length0 + b * Length0 * Length0 + Length0 * Length0 * Length0 / 3.0f);
-			float Integral1 = -(c * Length1 + b * Length1 * Length1 + Length1 * Length1 * Length1 / 3.0f);
+			// Integral of a Extinction that descreases according to square distance from the center of the volume.
+			const float Length0Sqr = Length0 * Length0;
+			const float Length1Sqr = Length1 * Length1;
+			float Integral0 = -(c * Length0 + b * Length0Sqr + Length0Sqr * Length0 / 3.0f);
+			float Integral1 = -(c * Length1 + b * Length1Sqr + Length1Sqr * Length1 / 3.0f);
 			OpticalDepth += max(0.0, FogInstance.RadialFogExtinction * (Integral1 - Integral0) * (3.0f / 4.0f));
 		}
 	}
 
 	if (FogInstance.HeightFogExtinction > 0.0f)
 	{
+		//////////
+		// Height fog is evaluated according to spatially varying density defined as "Density * Exp[-HeightFogFalloff * z]", where density is in fact Extinction.
+		// An explanation of it can be found in https://iquilezles.org/articles/fog/.
+
 		float StartHeight = RayStartU.z - FogInstance.HeightFogOffset; // account for the fog offset.
 
 	#if 0
diff --git a/Engine/Shaders/Private/VolumetricFog.usf b/Engine/Shaders/Private/VolumetricFog.usf
--- a/Engine/Shaders/Private/VolumetricFog.usf
+++ b/Engine/Shaders/Private/VolumetricFog.usf
@@ -184,7 +184,8 @@ void MaterialSetupCS(
 
 		float3 UnitSpacePos = mul(float4(WorldPosition, 1.0), FogInstance.InvTransform).xyz;
 
-		float SphereFade = saturate(1.0 - length(UnitSpacePos - float3(0, 0, 0)));
+		float DistToCenter = length(UnitSpacePos - float3(0.0f, 0.0f, 0.0f));
+		float SphereFade = saturate(1.0f - DistToCenter * DistToCenter);
 
 		float LFVExtinction = 0.0;
 		if (FogInstance.HeightFogExtinction > 0.0f)
