Fix ShaderPrint compilation error when using diagonostic buffer output.

diff --git a/Engine/Shaders/Private/ShaderPrint.ush b/Engine/Shaders/Private/ShaderPrint.ush
--- a/Engine/Shaders/Private/ShaderPrint.ush
+++ b/Engine/Shaders/Private/ShaderPrint.ush
@@ -306,7 +306,7 @@ FShaderPrintContext Assert(bool InCondition, FShaderPrintText InText, uint2 InWr
 	if (!InCondition)
 	{
 		uint LineOffset = 0;
-		InterlockedAdd(ShaderPrint_RWEntryBuffer[3] /* Free counter */, 1, LineOffset);
+		InterlockedAdd(SHADER_PRINT_RWENTRYBUFFER(OutCtx, 3) /* Free counter */, 1, LineOffset);
 		InWritePos.y += LineOffset * 12;
 
 		OutCtx = InitShaderPrintContext(true, InWritePos);
@@ -322,7 +322,7 @@ FShaderPrintContext Assert(bool InCondition, FShaderPrintText InText, FShaderPri
 	if (!InCondition)
 	{
 		uint LineOffset = 0;
-		InterlockedAdd(ShaderPrint_RWEntryBuffer[3] /* Free counter */, 1, LineOffset);
+		InterlockedAdd(SHADER_PRINT_RWENTRYBUFFER(OutCtx, 3) /* Free counter */, 1, LineOffset);
 		InWritePos.y += LineOffset * 12;
 
 		OutCtx = InitShaderPrintContext(true, InWritePos);
