Optimises TSR DecimateHistory pass with history resurrection

diff --git a/Engine/Shaders/Private/TemporalSuperResolution/TSRDecimateHistory.usf b/Engine/Shaders/Private/TemporalSuperResolution/TSRDecimateHistory.usf
--- a/Engine/Shaders/Private/TemporalSuperResolution/TSRDecimateHistory.usf
+++ b/Engine/Shaders/Private/TemporalSuperResolution/TSRDecimateHistory.usf
@@ -327,26 +327,41 @@ void ProcessPrevAtomicTexturesSamples(
 	}
 } // ProcessPrevAtomicTexturesSamples()
 
-// Sample a texture with a catrom kernel
 template<typename Type>
-Type SampleTextureCatrom(Texture2DArray<Type> InTexture, float2 TextureUV, float FrameIndex, float2 MinUV, float2 MaxUV)
+struct FCatmullRomFetches
 {
-	FCatmullRomSamples Samples = GetBicubic2DCatmullRomSamples_Stubbe(TextureUV, PrevGuideInfo_Extent, PrevGuideInfo_ExtentInverse);
-
-	Type Color = tsr_half(0.0);
-	for (uint i = 0; i < Samples.Count; i++)
+	Type Colors[BICUBIC_CATMULL_ROM_SAMPLES];
+	tsr_half Weights[BICUBIC_CATMULL_ROM_SAMPLES];
+	
+	void FetchSamples(Texture2DArray<Type> InTexture, float2 UV, float SliceIndex, float2 Size, float2 InvSize, float2 MinUV, float2 MaxUV)
 	{
-		float2 SampleUV = Samples.UV[i];
-		SampleUV = fastClamp(
-			SampleUV,
-			MinUV,
-			MaxUV);
-
-		Color += InTexture.SampleLevel(GlobalBilinearClampedSampler, float3(SampleUV, FrameIndex), 0) * tsr_half(Samples.Weight[i]);
+		FCatmullRomSamples Samples = GetBicubic2DCatmullRomSamples_Stubbe(UV, Size, InvSize);
+		
+		UNROLL_N(BICUBIC_CATMULL_ROM_SAMPLES)
+		for (uint i = 0; i < BICUBIC_CATMULL_ROM_SAMPLES; i++)
+		{
+			float2 SampleUV = Samples.UV[i];
+			SampleUV = fastClamp(
+				SampleUV,
+				MinUV,
+				MaxUV);
+
+			Colors[i] = InTexture.SampleLevel(GlobalBilinearClampedSampler, float3(SampleUV, SliceIndex), 0);
+			Weights[i] = tsr_half(Samples.Weight[i]);
+		}
 	}
 
-	return Color;
-}
+	Type AccumulateSamples()
+	{
+		Type Color = tsr_half(0.0);
+		UNROLL_N(BICUBIC_CATMULL_ROM_SAMPLES)
+		for (uint i = 0; i < BICUBIC_CATMULL_ROM_SAMPLES; i++)
+		{
+			Color += Colors[i] * Weights[i];
+		}
+		return Color;
+	}
+};
 
 // Correct pre exposure of the guide color.
 tsr_halfC CorrectGuideColorExposure(tsr_halfC ReprojectedHistoryGuideColor, const float PreExposureCorrection)
@@ -452,8 +467,11 @@ void MainCS(
 	uint HistoryUseCountSamples[4];
 	uint HistoryClosestDeviceZSamples[4];
 	
-	tsr_halfC ReprojectedHistoryGuideColor;
-	tsr_half  ReprojectedHistoryGuideUncertainty;
+
+	FCatmullRomFetches<tsr_half4> ReprojectedHistoryGuideSamples;
+	#if CONFIG_SCENE_COLOR_ALPHA
+		FCatmullRomFetches<tsr_half4> ReprojectedHistoryGuideUncertaintySamples;
+	#endif
 	#if CONFIG_MOIRE_REPROJECTION
 		tsr_half4 ReprojectedHistoryMoireColor;
 	#endif
@@ -470,11 +488,13 @@ void MainCS(
 			Map8x8Tile2x2Lane(GroupThreadIndex) +
 			(tsr_ushort2(InputInfo_ViewportMin) + tsr_ushort2(GroupId) * tsr_ushort2(TILE_SIZE, TILE_SIZE)));
 			
-		float2 ScreenPos = ApplyScreenTransform(float2(InputPixelPos), InputPixelPosToScreenPos);
-		float2 PrevScreenPos = ScreenPos - DecodeVelocityFromTexture(float4(EncodedScreenVelocity, 0.0, 0.0)).xy;
-		LoadPrevAtomicTexturesSamples(PrevScreenPos,
-			/* out */ HistoryUseCountSamples,
-			/* out */ HistoryClosestDeviceZSamples);
+		{
+			float2 ScreenPos = ApplyScreenTransform(float2(InputPixelPos), InputPixelPosToScreenPos);
+			float2 PrevScreenPos = ScreenPos - DecodeVelocityFromTexture(float4(EncodedScreenVelocity, 0.0, 0.0)).xy;
+			LoadPrevAtomicTexturesSamples(PrevScreenPos,
+				/* out */ HistoryUseCountSamples,
+				/* out */ HistoryClosestDeviceZSamples);
+		}
 			
 		{
 			float2 OutputScreenPos = ApplyScreenTransform(float2(InputPixelPos), InputPixelPosToReprojectScreenPos);
@@ -488,31 +508,26 @@ void MainCS(
 				
 			float2 KernelHistoryBufferUV = ApplyScreenTransform(PrevOutputScreenPos, ScreenPosToPrevHistoryGuideBufferUV);
 
-			#if CONFIG_SCENE_COLOR_ALPHA
-				ReprojectedHistoryGuideColor = SampleTextureCatrom(
-					PrevHistoryGuide,
-					KernelHistoryBufferUV,
-					PrevFrameIndex * 2.0 + 0.0,
-					PrevGuideInfo_UVViewportBilinearMin,
-					PrevGuideInfo_UVViewportBilinearMax);
+			ReprojectedHistoryGuideSamples.FetchSamples(
+				PrevHistoryGuide,
+				KernelHistoryBufferUV,
+				PrevFrameIndex * 2.0 + 0.0,
+				PrevGuideInfo_Extent,
+				PrevGuideInfo_ExtentInverse,
+				PrevGuideInfo_UVViewportBilinearMin,
+				PrevGuideInfo_UVViewportBilinearMax);
 					
-				ReprojectedHistoryGuideUncertainty = SampleTextureCatrom(
+			#if CONFIG_SCENE_COLOR_ALPHA
+				ReprojectedHistoryGuideUncertaintySamples.FetchSamples(
 					PrevHistoryGuide,
 					KernelHistoryBufferUV,
 					PrevFrameIndex * 2.0 + 1.0,
+					PrevGuideInfo_Extent,
+					PrevGuideInfo_ExtentInverse,
 					PrevGuideInfo_UVViewportBilinearMin,
 					PrevGuideInfo_UVViewportBilinearMax).r;
-			#else
-				tsr_half4 RawGuide = SampleTextureCatrom(
-					PrevHistoryGuide,
-					KernelHistoryBufferUV,
-					PrevFrameIndex,
-					PrevGuideInfo_UVViewportBilinearMin,
-					PrevGuideInfo_UVViewportBilinearMax);
-				ReprojectedHistoryGuideColor = RawGuide.rgb;
-				ReprojectedHistoryGuideUncertainty = RawGuide.a;
 			#endif
-			ReprojectedHistoryGuideColor = CorrectGuideColorExposure(ReprojectedHistoryGuideColor, HistoryPreExposureCorrection);
+
 			#if CONFIG_MOIRE_REPROJECTION
 			{
 				ReprojectedHistoryMoireColor = ReprojectPrevHistoryMoire(InputPixelPos, KernelHistoryBufferUV, PrevFrameIndex);
@@ -539,16 +554,37 @@ void MainCS(
 				ResurrectedHistoryGuideColor = RawGuide.rgb;
 				ResurrectedHistoryGuideUncertainty = RawGuide.a;
 			#endif
-			ResurrectedHistoryGuideColor = CorrectGuideColorExposure(ResurrectedHistoryGuideColor, ResurrectionPreExposureCorrection);
 		}
 		#endif
 	}
 	
 	// Process the history samples
+	tsr_halfC ReprojectedHistoryGuideColor;
+	tsr_half  ReprojectedHistoryGuideUncertainty;
 	tsr_half ParallaxRejectionMask;
-	
 	ISOLATE
 	{
+		#if CONFIG_SCENE_COLOR_ALPHA
+		{
+			ReprojectedHistoryGuideColor = ReprojectedHistoryGuideSamples.AccumulateSamples();
+			ReprojectedHistoryGuideUncertainty = ReprojectedHistoryGuideUncertaintySamples.AccumulateSamples().r;
+		}
+		#else
+		{
+			tsr_half4 RawGuide = ReprojectedHistoryGuideSamples.AccumulateSamples();
+			ReprojectedHistoryGuideColor = RawGuide.rgb;
+			ReprojectedHistoryGuideUncertainty = RawGuide.a;
+		}
+		#endif
+		
+		ReprojectedHistoryGuideColor = CorrectGuideColorExposure(ReprojectedHistoryGuideColor, HistoryPreExposureCorrection);
+
+		#if CONFIG_RESURRECTION_REPROJECTION
+		{
+			ResurrectedHistoryGuideColor = CorrectGuideColorExposure(ResurrectedHistoryGuideColor, ResurrectionPreExposureCorrection);
+		}
+		#endif
+
 		// LDS despill
 		#if CONFIG_MANUAL_LDS_SPILL
 		{
