Fix compilation issue on non-deferred platform.

diff --git a/Engine/Shaders/Private/PostProcessGBufferHints.usf b/Engine/Shaders/Private/PostProcessGBufferHints.usf
--- a/Engine/Shaders/Private/PostProcessGBufferHints.usf
+++ b/Engine/Shaders/Private/PostProcessGBufferHints.usf
@@ -151,7 +151,7 @@ void PrintBackgroundRect(inout FShaderPrintContext Ctx, float2 StartRect, float2
 [numthreads(1, 1, 1)]
 void MainCS(uint3 DispatchThreadId : SV_DispatchThreadID)
 {
-#if !SUBSTRATE_ENABLED
+#if !SUBSTRATE_ENABLED && SHADING_PATH_DEFERRED && FEATURE_LEVEL >= FEATURE_LEVEL_SM5
 	// Draw stats
 	if (all(DispatchThreadId == 0))
 	{
