Allocate VSM shadow mask slots based on pruned light grid rather than relying on sorted light order
This avoids issues with simple lights getting sorted to the front of the list (required for some other passes) and wasting slots, and makes the allocation more robust against light ordering
Update clustered shading path to properly sample the new packed mask
Separate out VSM directional/local light sampling functions so that callers who know what the light type is (ex. anything light grid for local lights) do not need to include the other code paths in the shader

diff --git a/Engine/Shaders/Private/ClusteredDeferredShadingPixelShader.usf b/Engine/Shaders/Private/ClusteredDeferredShadingPixelShader.usf
--- a/Engine/Shaders/Private/ClusteredDeferredShadingPixelShader.usf
+++ b/Engine/Shaders/Private/ClusteredDeferredShadingPixelShader.usf
@@ -40,6 +40,7 @@
 #include "DeferredLightingCommon.ush"
 #include "LightData.ush"
 #include "VirtualShadowMaps/VirtualShadowMapTransmissionCommon.ush"
+#include "VirtualShadowMaps/VirtualShadowMapMaskBitsCommon.ush"
 #include "Substrate/SubstrateEvaluation.ush"
 #include "Substrate/SubstrateDeferredLighting.ush"
 
@@ -220,7 +221,8 @@ float4 GetLightGridLocalLighting(
 	const FCulledLightsGridData InLightGridData, 
 	float3 TranslatedWorldPosition, 
 	float3 CameraVector, 
-	float2 ScreenUV, 
+	float2 ScreenUV,
+	float SceneDepth,
 	uint EyeIndex, 
 	float Dither, 
 	uint FirstNonSimpleLightIndex, 
@@ -235,7 +237,7 @@ float4 GetLightGridLocalLighting(
 	const uint NumLocalLights = min(InLightGridData.NumLocalLights, GetMaxLightsPerCell(EyeIndex));
 
 	int2 PixelPos = int2( ScreenUV.xy * View.BufferSizeAndInvSize.xy );
-	uint4 ShadowMask = ~ShadowMaskBits[ PixelPos ];
+	uint4 ShadowMask = ShadowMaskBits[ PixelPos ];
 	
 	LOOP
 	for (uint LocalLightListIndex = FirstNonSimpleLightIndex; LocalLightListIndex < NumLocalLights; LocalLightListIndex++)
@@ -257,22 +259,18 @@ float4 GetLightGridLocalLighting(
 		FDeferredLightData LightData = ConvertToDeferredLight(LocalLight);
 		LightData.ShadowedBits = 1;
 		float4 LightAttenuation = float4(1, 1, 1, 1);
-		if (LocalLightListIndex < GetPackedShadowMaskMaxLightCount())
-		{
-			LightAttenuation = UnpackShadowMask(ShadowMask, LocalLightListIndex, Dither).xxxx;
-		}
-		else if (LocalLight.VirtualShadowMapId != INDEX_NONE)
+
+		// Find packed VSM projection result if present
+		// NOTE: We have to do the search as in the deferred pass again currently as the shadow masks
+		// are packed according to a prune light list that contains only shadow-casting lights.
+		if (LocalLight.VirtualShadowMapId != INDEX_NONE)
 		{
-			// Single sample VSM lookup
-			// TODO: Normal bias? Screen space ray? Etc? Do we need a path that does more of that stuff
-			// unified but without SMRT? For now we'll just keep this as a simple lookup to avoid potential performance tradeoffs
-			// as this is meant to be the fallback path on overflowing tiles.
-			FVirtualShadowMapSampleResult VSMResult = SampleVirtualShadowMapTranslatedWorld(
+			LightAttenuation = GetVirtualShadowMapMaskForLight(
+				ShadowMask,
+				uint2(PixelPos),
+				SceneDepth,
 				LocalLight.VirtualShadowMapId,
-				TranslatedWorldPosition,
-				30.0f // RayStartDistance TODO... maybe at least make this a proportional screen offset like the contact shadow trace
-			);
-			LightAttenuation = VSMResult.ShadowFactor.xxxx;
+				TranslatedWorldPosition).xxxx;
 		}
 
 		#if USE_HAIR_LIGHTING
@@ -380,7 +378,7 @@ float4 GetLightGridLocalLighting(
 		float3 TranslatedWorldPosition = mul(float4(GetScreenPositionForProjectionType(ScreenPosition, SceneDepth), SceneDepth, 1), PrimaryView.ScreenToTranslatedWorld).xyz;	\
 		float3 CameraVector = normalize(TranslatedWorldPosition - PrimaryView.TranslatedWorldCameraOrigin);	\
 		ScreenSpaceData.GBuffer.ShadingModelID = ShadingModelId;	\
-		CompositedLighting += GetLightGridLocalLighting(ScreenSpaceData, LightGridData, TranslatedWorldPosition, CameraVector, ScreenUV, 0, Dither, FirstNonSimpleLightIndex); \
+		CompositedLighting += GetLightGridLocalLighting(ScreenSpaceData, LightGridData, TranslatedWorldPosition, CameraVector, ScreenUV, SceneDepth, 0, Dither, FirstNonSimpleLightIndex); \
 	}	\
 }
 
@@ -434,7 +432,7 @@ void ClusteredShadingPixelShader(
 
 		// Regular lights
 		CompositedLighting += GetLightGridLocalLighting(SubstrateAddressing, SubstratePixelHeader, OpaqueRoughRefractionSceneColor, SubSurfaceSceneColor,
-			CulledLightGridData, TranslatedWorldPosition, CameraVector, ScreenUV, EyeIndex, Dither, FirstNonSimpleLightIndex);
+			CulledLightGridData, TranslatedWorldPosition, CameraVector, ScreenUV, SceneDepth, EyeIndex, Dither, FirstNonSimpleLightIndex);
 	}
 	
 #if !VISUALIZE_LIGHT_CULLING
@@ -502,7 +500,7 @@ void ClusteredShadingPixelShader(
 	GET_LIGHT_GRID_LOCAL_LIGHTING_SINGLE_SM(SHADINGMODELID_SINGLELAYERWATER,	PixelShadingModelID, CompositedLighting, ScreenUV, CulledLightGridData, Dither, FirstNonSimpleLightIndex);
 	// SHADINGMODELID_THIN_TRANSLUCENT - skipping because it can not be opaque
 #else // !USE_PASS_PER_SHADING_MODEL
-	CompositedLighting += GetLightGridLocalLighting(GetScreenSpaceData(ScreenUV), CulledLightGridData, TranslatedWorldPosition, CameraVector, ScreenUV, 0, Dither, FirstNonSimpleLightIndex);
+	CompositedLighting += GetLightGridLocalLighting(GetScreenSpaceData(ScreenUV), CulledLightGridData, TranslatedWorldPosition, CameraVector, ScreenUV, SceneDepth, 0, Dither, FirstNonSimpleLightIndex);
 #endif // USE_PASS_PER_SHADING_MODEL
 	
 #if !VISUALIZE_LIGHT_CULLING
@@ -559,7 +557,7 @@ void ClusteredShadingPixelShader(
 
 	float4 CompositedLighting = 0;
 	uint FirstNonSimpleLightIndex = 0;
-	CompositedLighting += GetLightGridLocalLighting(ScreenSpaceData, CulledLightGridData, TranslatedWorldPosition, CameraVector, ScreenUV, 0, Dither, FirstNonSimpleLightIndex, SampleIndex, TotalSampleCount, PixelCoord);
+	CompositedLighting += GetLightGridLocalLighting(ScreenSpaceData, CulledLightGridData, TranslatedWorldPosition, CameraVector, ScreenUV, SceneDepth, 0, Dither, FirstNonSimpleLightIndex, SampleIndex, TotalSampleCount, PixelCoord);
 #if !VISUALIZE_LIGHT_CULLING
 	CompositedLighting *= View.PreExposure;
 #endif
diff --git a/Engine/Shaders/Private/DeferredLightPixelShaders.usf b/Engine/Shaders/Private/DeferredLightPixelShaders.usf
--- a/Engine/Shaders/Private/DeferredLightPixelShaders.usf
+++ b/Engine/Shaders/Private/DeferredLightPixelShaders.usf
@@ -113,7 +113,7 @@ Texture2D<uint4> ShadowMaskBits;
 int VirtualShadowMapId;
 #endif
 
-float4 GetLightAttenuationFromShadow(in FInputParams InputParams, float SceneDepth)
+float4 GetLightAttenuationFromShadow(in FInputParams InputParams, float SceneDepth, float3 TranslatedWorldPosition)
 {
 	float4 LightAttenuation = float4(1, 1, 1, 1);
 
@@ -121,13 +121,11 @@ float4 GetLightAttenuationFromShadow(in FInputParams InputParams, float SceneDep
 	if (VirtualShadowMapId != INDEX_NONE)
 	{
 		float ShadowMask = GetVirtualShadowMapMaskForLight(
-			ShadowMaskBits,
+			ShadowMaskBits[InputParams.PixelPos],
 			uint2(InputParams.PixelPos),
 			SceneDepth,
-			0,	// TODO: EyeIndex
-			VirtualShadowMapId);
-
-		// TODO: Subsurface...?
+			VirtualShadowMapId,
+			TranslatedWorldPosition);
 		return ShadowMask.xxxx;
 	}
 	else
@@ -280,7 +278,7 @@ void DeferredLightPixelMain(
 		if (LightMask > 0)
 		{
 			FShadowTerms ShadowTerms = { SubstrateGetAO(SubstratePixelHeader), 1.0, 1.0, InitHairTransmittanceData() };
-			float4 LightAttenuation = GetLightAttenuationFromShadow(InputParams, SceneDepth);
+			float4 LightAttenuation = GetLightAttenuationFromShadow(InputParams, SceneDepth, DerivedParams.TranslatedWorldPosition);
 
 			float Dither = InterleavedGradientNoise(InputParams.PixelPos, View.StateFrameIndexMod8);
 			const uint FakeShadingModelID = 0;
@@ -335,7 +333,7 @@ void DeferredLightPixelMain(
 
 		float SurfaceShadow = 1.0f;
 		
-		float4 LightAttenuation = GetLightAttenuationFromShadow(InputParams, SceneDepth);
+		float4 LightAttenuation = GetLightAttenuationFromShadow(InputParams, SceneDepth, DerivedParams.TranslatedWorldPosition);
 		float4 Radiance = GetDynamicLighting(DerivedParams.TranslatedWorldPosition, DerivedParams.CameraVector, ScreenSpaceData.GBuffer, ScreenSpaceData.AmbientOcclusion, ScreenSpaceData.GBuffer.ShadingModelID, LightData, LightAttenuation, Dither, uint2(InputParams.PixelPos), SurfaceShadow);
 
 		OutColor += Radiance;
diff --git a/Engine/Shaders/Private/ForwardLightingCommon.ush b/Engine/Shaders/Private/ForwardLightingCommon.ush
--- a/Engine/Shaders/Private/ForwardLightingCommon.ush
+++ b/Engine/Shaders/Private/ForwardLightingCommon.ush
@@ -224,7 +224,7 @@ FDeferredLightingSplit GetForwardDirectLightingSplit(
 				}
 				#else
 				{
-					FVirtualShadowMapSampleResult VirtualShadowMapSample = SampleVirtualShadowMapTranslatedWorld( ForwardLightData.DirectionalLightVSM, TranslatedWorldPosition );
+					FVirtualShadowMapSampleResult VirtualShadowMapSample = SampleVirtualShadowMapDirectional( ForwardLightData.DirectionalLightVSM, TranslatedWorldPosition );
 					DynamicShadowFactor *= VirtualShadowMapSample.ShadowFactor;
 				}
 				#endif
diff --git a/Engine/Shaders/Private/HairStrands/HairStrandsDeepTransmittanceMask.usf b/Engine/Shaders/Private/HairStrands/HairStrandsDeepTransmittanceMask.usf
--- a/Engine/Shaders/Private/HairStrands/HairStrandsDeepTransmittanceMask.usf
+++ b/Engine/Shaders/Private/HairStrands/HairStrandsDeepTransmittanceMask.usf
@@ -373,11 +373,13 @@ void MainCS(uint2 DispatchThreadId : SV_DispatchThreadID)
 		TranslatedWorldPosition += VirtualVoxel.TranslatedWorldOffsetStereoCorrection;
 	}
 
-	const uint4 PackedShadowMask = ~ShadowMaskBitsTexture[PixelCoord];
+	// TODO: Need to update the shadow mask encoding/decoding functions here
+	//const uint4 PackedShadowMask = ~ShadowMaskBitsTexture[PixelCoord];
 	const float3 Random = GetHairVoxelJitter(PixelCoord.xy, View.StateFrameIndexMod8, VirtualVoxel.JitterMode);
 
 	const uint2 LocalPixelPos = PixelCoord;
 
+	// TODO: Possibly make use of the new VSM culled light grid
 	const uint EyeIndex = 0; // TODO: Instanced stereo
 	const uint GridLinearIndex = ComputeLightGridCellIndex(LocalPixelPos, SceneDepth, EyeIndex);
 	const FCulledLightsGridData CulledLightGridData = GetCulledLightsGrid(GridLinearIndex, EyeIndex);
diff --git a/Engine/Shaders/Private/LightGridCommon.ush b/Engine/Shaders/Private/LightGridCommon.ush
--- a/Engine/Shaders/Private/LightGridCommon.ush
+++ b/Engine/Shaders/Private/LightGridCommon.ush
@@ -317,7 +317,7 @@ FLocalLightData GetLocalLightData(uint GridIndex, uint EyeIndex)
  */
 uint4 InitializePackedShadowMask()
 {
-	return uint(0xffffffff).xxxx;
+	return uint(0U).xxxx;
 }
 
 uint GetPackedShadowMaskMaxLightCount()
@@ -327,22 +327,20 @@ uint GetPackedShadowMaskMaxLightCount()
 
 void PackShadowMask(inout uint4 InOutShadowMask, float InShadowFactor, uint InLightIndex)
 {
-	uint Value = ~uint(round(InShadowFactor * 15.0f)) & 15U;
+	uint Value = uint(round((InShadowFactor) * 15.0f)) & 15U;
 	uint Dword = InLightIndex / 8;
 
-	InOutShadowMask.x ^= (Dword == 0U) ? (Value << (InLightIndex      ) * 4U) : 0U;
-	InOutShadowMask.y ^= (Dword == 1U) ? (Value << (InLightIndex -  8U) * 4U) : 0U;
-	InOutShadowMask.z ^= (Dword == 2U) ? (Value << (InLightIndex - 16U) * 4U) : 0U;
-	InOutShadowMask.w ^= (Dword == 3U) ? (Value << (InLightIndex - 24U) * 4U) : 0U;
-
-	//OutShadowMask ^= (~uint(round(InShadowFactor * 7.0)) & 7u) << (InLightIndex * 3u);
+	InOutShadowMask.x ^= (Dword == 0U) ? (Value << ((InLightIndex      ) * 4U)) : 0U;
+	InOutShadowMask.y ^= (Dword == 1U) ? (Value << ((InLightIndex -  8U) * 4U)) : 0U;
+	InOutShadowMask.z ^= (Dword == 2U) ? (Value << ((InLightIndex - 16U) * 4U)) : 0U;
+	InOutShadowMask.w ^= (Dword == 3U) ? (Value << ((InLightIndex - 24U) * 4U)) : 0U;
 }
 
 float UnpackShadowMask(uint4 InShadowMask, uint InLightIndex)
-{	
+{
 	uint Dword = InLightIndex / 8;
-	return ((InShadowMask[Dword] >> (InLightIndex - Dword*8U) * 4U) & 15U) / 15.0f;	
-	//return ((InShadowMask >> (InLightIndex * 3u)) & 7u) / 7.0;
+	uint MaskBits = (InShadowMask[Dword] >> (InLightIndex - Dword*8U) * 4U) & 15U;
+	return (float(MaskBits) / 15.0f);
 }
 
 // Unpack with dither to hide some of the quantization
diff --git a/Engine/Shaders/Private/SkyAtmosphere.usf b/Engine/Shaders/Private/SkyAtmosphere.usf
--- a/Engine/Shaders/Private/SkyAtmosphere.usf
+++ b/Engine/Shaders/Private/SkyAtmosphere.usf
@@ -677,7 +677,7 @@ SingleScatteringResult IntegrateSingleScatteredLuminance(
 		#if VIRTUAL_SHADOW_MAP
 			if (VirtualShadowMapId0 != INDEX_NONE)
 			{
-				FVirtualShadowMapSampleResult VirtualShadowMapSample = SampleVirtualShadowMapTranslatedWorld(VirtualShadowMapId0, ShadowSampleWorldPosition0);
+				FVirtualShadowMapSampleResult VirtualShadowMapSample = SampleVirtualShadowMapDirectional(VirtualShadowMapId0, ShadowSampleWorldPosition0);
 				PlanetShadow0 *= VirtualShadowMapSample.ShadowFactor;
 			}
 		#endif // VIRTUALSHADOW_MAP
@@ -713,7 +713,7 @@ SingleScatteringResult IntegrateSingleScatteredLuminance(
 		#if VIRTUAL_SHADOW_MAP
 			if (VirtualShadowMapId1 != INDEX_NONE)
 			{
-				FVirtualShadowMapSampleResult VirtualShadowMapSample = SampleVirtualShadowMapTranslatedWorld(VirtualShadowMapId1, ShadowSampleWorldPosition1);
+				FVirtualShadowMapSampleResult VirtualShadowMapSample = SampleVirtualShadowMapDirectional(VirtualShadowMapId1, ShadowSampleWorldPosition1);
 				PlanetShadow1 *= VirtualShadowMapSample.ShadowFactor;
 			}
 		#endif // VIRTUALSHADOW_MAP
diff --git a/Engine/Shaders/Private/Substrate/SubstrateForwardLighting.ush b/Engine/Shaders/Private/Substrate/SubstrateForwardLighting.ush
--- a/Engine/Shaders/Private/Substrate/SubstrateForwardLighting.ush
+++ b/Engine/Shaders/Private/Substrate/SubstrateForwardLighting.ush
@@ -151,7 +151,7 @@ float3 SubstrateForwardLighting(
 		}
 		#else
 		{
-			FVirtualShadowMapSampleResult VirtualShadowMapSample = SampleVirtualShadowMapTranslatedWorld(ForwardLightData.DirectionalLightVSM, TranslatedWorldPosition);
+			FVirtualShadowMapSampleResult VirtualShadowMapSample = SampleVirtualShadowMapDirectional(ForwardLightData.DirectionalLightVSM, TranslatedWorldPosition);
 			DirLightDynamicShadowFactor *= VirtualShadowMapSample.ShadowFactor;
 		}
 		#endif
diff --git a/Engine/Shaders/Private/VirtualShadowMaps/VirtualShadowMapCacheManagement.usf b/Engine/Shaders/Private/VirtualShadowMaps/VirtualShadowMapCacheManagement.usf
--- a/Engine/Shaders/Private/VirtualShadowMaps/VirtualShadowMapCacheManagement.usf
+++ b/Engine/Shaders/Private/VirtualShadowMaps/VirtualShadowMapCacheManagement.usf
@@ -105,15 +105,16 @@ void VirtualSmInvalidateInstancePagesCS(
 		{
 			for (uint VirtualShadowMapId = FirstVirtualShadowMapId; VirtualShadowMapId < EndVirtualShadowMapId; ++VirtualShadowMapId)
 			{
-				// 1. Load cached projection data
-				FVirtualShadowMapProjectionShaderData ProjectionData = GetVirtualShadowMapProjectionData(VirtualShadowMapId);
-				const bool bDirectionalLight = (ProjectionData.LightType == LIGHT_TYPE_DIRECTIONAL);
-
 				// Stop invalidations if the light is 'distant'
-				if (ProjectionData.bCurrentDistantLight)
+				if (IsSinglePageVirtualShadowMap(VirtualShadowMapId))
 				{
 					continue;
 				}
+
+				// 1. Load cached projection data
+				FVirtualShadowMapProjectionShaderData ProjectionData = GetVirtualShadowMapProjectionData(VirtualShadowMapId);
+				const bool bDirectionalLight = (ProjectionData.LightType == LIGHT_TYPE_DIRECTIONAL);
+
 				// NOTE: This is the *shadow view*'s translated world, not primary view
 				float4x4 LocalToTranslatedWorld = LWCMultiplyTranslation(InstanceSceneData.LocalToWorld, ProjectionData.PreViewTranslation);
 				
diff --git a/Engine/Shaders/Private/VirtualShadowMaps/VirtualShadowMapLightGrid.ush b/Engine/Shaders/Private/VirtualShadowMaps/VirtualShadowMapLightGrid.ush
new file mode 100644
--- /dev/null
+++ b/Engine/Shaders/Private/VirtualShadowMaps/VirtualShadowMapLightGrid.ush
@@ -0,0 +1,30 @@
+// Copyright Epic Games, Inc. All Rights Reserved.
+
+/*=============================================================================
+VirtualShadowMapLightGrid.ush:
+=============================================================================*/
+#pragma once
+
+#include "../Common.ush"
+#include "../LightGridCommon.ush"
+
+FCulledLightsGridData VirtualShadowMapGetLightsGrid(uint2 PixelPos, float SceneDepth)
+{
+	// Has no effect unless INSTANCED_STEREO is enabled. Used in GetLightGridData to select ForwardLightData or ForwardLightDataISR, but then only uses redundant parameters anyway.
+	uint EyeIndex = 0U;
+
+	uint GridLinearIndex = ComputeLightGridCellIndex(PixelPos, SceneDepth, EyeIndex);
+	FCulledLightsGridData CulledLightGridData = GetCulledLightsGrid(GridLinearIndex, EyeIndex);
+
+	// Replace light count with our pruned count
+	CulledLightGridData.NumLocalLights = VirtualShadowMap.NumCulledLightsGrid[GridLinearIndex];
+	return CulledLightGridData;
+}
+
+FLocalLightData VirtualShadowMapGetLocalLightData(FCulledLightsGridData GridData, uint Index)
+{
+	const uint ListIndex = GridData.DataStartIndex + Index;
+	checkStructuredBufferAccessSlow(VirtualShadowMap.LightGridData, ListIndex);
+	const uint LightGridLightIndex = VirtualShadowMap.LightGridData[ListIndex];
+	return GetLocalLightDataNonStereo(LightGridLightIndex);
+}
diff --git a/Engine/Shaders/Private/VirtualShadowMaps/VirtualShadowMapMaskBitsCommon.ush b/Engine/Shaders/Private/VirtualShadowMaps/VirtualShadowMapMaskBitsCommon.ush
--- a/Engine/Shaders/Private/VirtualShadowMaps/VirtualShadowMapMaskBitsCommon.ush
+++ b/Engine/Shaders/Private/VirtualShadowMaps/VirtualShadowMapMaskBitsCommon.ush
@@ -10,44 +10,50 @@ VirtualShadowMapMaskBitsCommon.ush:
 #include "../SceneTexturesCommon.ush"
 #include "../LightShaderParameters.ush"
 #include "../LightGridCommon.ush"
+#include "../VirtualShadowMaps/VirtualShadowMapLightGrid.ush"
+#include "../VirtualShadowMaps/VirtualShadowMapProjectionCommon.ush"
 
 float GetVirtualShadowMapMaskForLight(
-	Texture2D<uint4> ShadowMaskBits,
+	uint4 ShadowMask,
 	uint2 PixelPos,
 	float SceneDepth,
-	uint EyeIndex,
-	int VirtualShadowMapId)
+	int VirtualShadowMapId,
+	float3 TranslatedWorldPosition)
 {
 	// Light grid is relative to the view
 	const uint2 LocalPixelPos = PixelPos.xy - uint2(View.ViewRectMinAndSize.xy);
-
-	uint GridLinearIndex = ComputeLightGridCellIndex(LocalPixelPos, SceneDepth, EyeIndex);
-	const FCulledLightsGridData CulledLightGridData = GetCulledLightsGrid(GridLinearIndex, EyeIndex);
+	const FCulledLightsGridData LightGridData = VirtualShadowMapGetLightsGrid(LocalPixelPos, SceneDepth);	
 
 	// We can only handle so many lights in our output encoding right now, so no purpose in computing more
-	uint LightCount = min(GetPackedShadowMaskMaxLightCount(), CulledLightGridData.NumLocalLights);
+	uint LightCount = min(GetPackedShadowMaskMaxLightCount(), LightGridData.NumLocalLights);
 
 	uint Index = 0;
 	for (; Index < LightCount; ++Index)
 	{
-		const FLocalLightData LightData = GetLocalLightData(CulledLightGridData.DataStartIndex + Index, EyeIndex);
+		const FLocalLightData LightData = VirtualShadowMapGetLocalLightData(LightGridData, Index);
 		if (LightData.VirtualShadowMapId == VirtualShadowMapId)
 		{
 			break;
 		}
 	}
 
-	float UnpackedShadowMask = 1.0f;
+	float LightAttenuation = 1.0f;
 	if (Index < LightCount)
 	{
-		uint4 ShadowMask = ~ShadowMaskBits[PixelPos];
 		// TODO: Dither!
-		UnpackedShadowMask = UnpackShadowMask(ShadowMask, Index);
+		LightAttenuation = UnpackShadowMask(ShadowMask, Index);
 	}
 	else
 	{
-		// Not found!
+		const float StaticBias = 200.0f;
+		// Fall back to a single sample VSM lookup
+		FVirtualShadowMapSampleResult VSMResult = SampleVirtualShadowMapLocal(
+			VirtualShadowMapId,
+			TranslatedWorldPosition,
+			StaticBias
+		);
+		LightAttenuation = VSMResult.ShadowFactor;
 	}
 
-	return UnpackedShadowMask;
+	return LightAttenuation;
 }
diff --git a/Engine/Shaders/Private/VirtualShadowMaps/VirtualShadowMapPageMarking.usf b/Engine/Shaders/Private/VirtualShadowMaps/VirtualShadowMapPageMarking.usf
--- a/Engine/Shaders/Private/VirtualShadowMaps/VirtualShadowMapPageMarking.usf
+++ b/Engine/Shaders/Private/VirtualShadowMaps/VirtualShadowMapPageMarking.usf
@@ -18,6 +18,7 @@
 #include "VirtualShadowMapProjectionSpot.ush"
 #include "VirtualShadowMapProjectionCommon.ush"
 #include "VirtualShadowMapStats.ush"
+#include "VirtualShadowMapLightGrid.ush"
 
 // Type of input data consume by the page allocation (i.e., data read from the source buffer: Gbuffer, HairStrands data, ...)
 #define INPUT_TYPE_GBUFFER 0
@@ -81,7 +82,6 @@ void MarkPage(uint VirtualShadowMapId, uint MipLevel, float3 TranslatedWorldPosi
 	FVirtualShadowMapProjectionShaderData ProjectionData = GetVirtualShadowMapProjectionData(VirtualShadowMapId);
 
 	// MarkPage (or mark pixel pages) should never run for a distant light.
-	checkSlow(!ProjectionData.bCurrentDistantLight);
 	checkSlow(!IsSinglePageVirtualShadowMap(VirtualShadowMapId));
 
 	float3 ShadowTranslatedWorldPosition = TranslatedWorldPosition + 
@@ -147,7 +147,7 @@ RWStructuredBuffer<uint> OutPrunedNumCulledLightsGrid;
 [numthreads(VSM_DEFAULT_CS_GROUP_X, 1, 1)]
 void PruneLightGridCS(uint GridLinearIndex : SV_DispatchThreadID)
 {
-	uint EyeIndex = 0; // ??
+	uint EyeIndex = 0;
 	const FLightGridData GridData = GetLightGridData(EyeIndex);
 
 	if (GridLinearIndex >= GridData.CulledGridSize.x * GridData.CulledGridSize.y * GridData.CulledGridSize.z)
@@ -157,33 +157,39 @@ void PruneLightGridCS(uint GridLinearIndex : SV_DispatchThreadID)
 	const FCulledLightsGridData CulledLightGridData = GetCulledLightsGrid(GridLinearIndex, EyeIndex);
 
 	uint NumRetainedLights = 0U;
+
+	// First pass add the non-distant lights with VSMs
 	LOOP
 	for (uint Index = 0; Index < CulledLightGridData.NumLocalLights; ++Index)
 	{
-		const FLocalLightData LightData = GetLocalLightData(CulledLightGridData.DataStartIndex + Index, EyeIndex);
-		int VirtualShadowMapId = LightData.VirtualShadowMapId;
-		// Discard any light without a VSM
-		if (VirtualShadowMapId != INDEX_NONE)
+		int VirtualShadowMapId = GetLocalLightData(CulledLightGridData.DataStartIndex + Index, EyeIndex).VirtualShadowMapId;
+		if (VirtualShadowMapId != INDEX_NONE && !IsSinglePageVirtualShadowMap(VirtualShadowMapId))
 		{
-			FVirtualShadowMapProjectionShaderData ProjectionData = GetVirtualShadowMapProjectionData(VirtualShadowMapId);
-
-			// Discard any distant lights (they get marked through the coarse page marking)
-			if (!ProjectionData.bCurrentDistantLight)
-			{
-				checkSlow(!IsSinglePageVirtualShadowMap(VirtualShadowMapId));
+			// Copy light grid data index
+			checkStructuredBufferAccessSlow(OutPrunedLightGridData, CulledLightGridData.DataStartIndex + NumRetainedLights);
+			OutPrunedLightGridData[CulledLightGridData.DataStartIndex + NumRetainedLights++] = GetCulledLightDataGrid(CulledLightGridData.DataStartIndex + Index);
+		}
+	}
 
-				// Copy light grid data index
-				checkStructuredBufferAccessSlow(OutPrunedLightGridData, CulledLightGridData.DataStartIndex + NumRetainedLights);
-				OutPrunedLightGridData[CulledLightGridData.DataStartIndex + NumRetainedLights++] = GetCulledLightDataGrid(CulledLightGridData.DataStartIndex + Index);
-			}
+	// Second pass add the distant lights with VSMs after
+	// NOTE: We could add these to the end in the first pass and then re-pack them instead, but this works
+	// well enough for moderate light counts.
+	LOOP
+	for (uint Index = 0; Index < CulledLightGridData.NumLocalLights; ++Index)
+	{
+		int VirtualShadowMapId = GetLocalLightData(CulledLightGridData.DataStartIndex + Index, EyeIndex).VirtualShadowMapId;
+		// Discard any light without a VSM
+		if (VirtualShadowMapId != INDEX_NONE && IsSinglePageVirtualShadowMap(VirtualShadowMapId))
+		{
+			checkStructuredBufferAccessSlow(OutPrunedLightGridData, CulledLightGridData.DataStartIndex + NumRetainedLights);
+			OutPrunedLightGridData[CulledLightGridData.DataStartIndex + NumRetainedLights++] = GetCulledLightDataGrid(CulledLightGridData.DataStartIndex + Index);
 		}
 	}
+
 	checkStructuredBufferAccessSlow(OutPrunedNumCulledLightsGrid, GridLinearIndex);
 	OutPrunedNumCulledLightsGrid[GridLinearIndex] = NumRetainedLights;
 }
 
-StructuredBuffer<uint> PrunedLightGridData;
-StructuredBuffer<uint> PrunedNumCulledLightsGrid;
 uint2 PixelStride;
 
 #if PERMUTATION_WATER_DEPTH
@@ -437,24 +443,20 @@ void GeneratePageFlagsFromPixels(
 	{
 		const bool bUsePageDilation = PageDilationBorderSizeLocal > 0.0f;
 		const float2 PageDilationOffset = PageDilationBorderSizeLocal * PageDilationDither;
-
-		uint EyeIndex = 0U; // Has no effect unless INSTANCED_STEREO is enabled, which is not the case for this CS, (used in GetLightGridData to select ForwardLightData or ForwardLightDataISR, but then only uses redundant parameters anyway)
+				
 		uint2 LocalPosition = PixelPos - uint2(View.ViewRectMin.xy);
-		uint3 GridCoordinate = ComputeLightGridCellCoordinate(LocalPosition, Primary.SceneDepth, EyeIndex);
-
-		uint GridLinearIndex = ComputeLightGridCellIndex(GridCoordinate, EyeIndex);
-		const FCulledLightsGridData CulledLightGridData = GetCulledLightsGrid(GridLinearIndex, EyeIndex);
-
-		checkStructuredBufferAccessSlow(PrunedNumCulledLightsGrid, GridLinearIndex);
-		const uint PrunedNumLocalLights = PrunedNumCulledLightsGrid[GridLinearIndex];
+		const FCulledLightsGridData LightGridData = VirtualShadowMapGetLightsGrid(LocalPosition, Primary.SceneDepth);
 
 		LOOP
-		for (uint Index = 0; Index < PrunedNumLocalLights; ++Index)
+		for (uint Index = 0; Index < LightGridData.NumLocalLights; ++Index)
 		{
-			checkStructuredBufferAccessSlow(PrunedLightGridData, CulledLightGridData.DataStartIndex + Index);
-			const uint LightGridLightIndex = PrunedLightGridData[CulledLightGridData.DataStartIndex + Index];
+			const FLocalLightData LightData = VirtualShadowMapGetLocalLightData(LightGridData, Index);
 
-			const FLocalLightData LightData = GetLocalLightDataNonStereo(LightGridLightIndex);
+			// If we hit a distant light, we're done as we sorted those to the end. See PruneLightGridCS.
+			if (IsSinglePageVirtualShadowMap(LightData.VirtualShadowMapId))
+			{
+				break;
+			}
 
 			// Relative to PrimaryView
 			const float3 LightTranslatedWorldPosition = LightData.LightPositionAndInvRadius.xyz;
@@ -520,7 +522,6 @@ void MarkCoarsePages(uint DispatchThreadId : SV_DispatchThreadID)
 	}
 
 	FVirtualShadowMapProjectionShaderData ProjectionData = GetVirtualShadowMapProjectionData(VirtualShadowMapId);
-	checkSlow(ProjectionData.bCurrentDistantLight == IsSinglePageVirtualShadowMap(VirtualShadowMapId));
 		
 	// NOTE: Coarse pages are very large and tend to get invalidated a lot due to anything in the scene moving
 	// Rendering non-nanite geometry into these pages can be very expensive and thus isn't always desirable.
@@ -556,7 +557,7 @@ void MarkCoarsePages(uint DispatchThreadId : SV_DispatchThreadID)
 		}
 	}
 	// Note: always mark last mip for "distant" light
-	else if (bMarkCoarsePagesLocal != 0U || ProjectionData.bCurrentDistantLight)
+	else if (bMarkCoarsePagesLocal != 0U || IsSinglePageVirtualShadowMap(VirtualShadowMapId))
 	{
 		// Mark last mip
 		uint PageOffset = CalcPageOffset(VirtualShadowMapId, VSM_MAX_MIP_LEVELS - 1, uint2(0, 0));
diff --git a/Engine/Shaders/Private/VirtualShadowMaps/VirtualShadowMapPhysicalPageManagement.usf b/Engine/Shaders/Private/VirtualShadowMaps/VirtualShadowMapPhysicalPageManagement.usf
--- a/Engine/Shaders/Private/VirtualShadowMaps/VirtualShadowMapPhysicalPageManagement.usf
+++ b/Engine/Shaders/Private/VirtualShadowMaps/VirtualShadowMapPhysicalPageManagement.usf
@@ -280,7 +280,7 @@ void UpdatePhysicalPages(uint3 Index : SV_DispatchThreadID)
 					uint NextPageFlags = VSM_ALLOCATED_FLAG;
 
 					// Distant lights block GPU-GPU invalidations (and all others) as they are round-robin invalidated
-					if (bDynamicPageInvalidation && !Projection.bCurrentDistantLight)
+					if (bDynamicPageInvalidation && !IsSinglePageVirtualShadowMap(VirtualShadowMapId))
 					{
 						uint InvalidationFlags = (PrevPhysicalFlags >> VSM_PHYSICAL_PAGE_INVALIDATION_FLAGS_SHIFT);
 						if (VirtualShadowMapShouldCacheStaticSeparately())
diff --git a/Engine/Shaders/Private/VirtualShadowMaps/VirtualShadowMapProjection.usf b/Engine/Shaders/Private/VirtualShadowMaps/VirtualShadowMapProjection.usf
--- a/Engine/Shaders/Private/VirtualShadowMaps/VirtualShadowMapProjection.usf
+++ b/Engine/Shaders/Private/VirtualShadowMaps/VirtualShadowMapProjection.usf
@@ -17,6 +17,7 @@
 #include "VirtualShadowMapProjectionDirectional.ush"
 #include "VirtualShadowMapProjectionSpot.ush"
 #include "VirtualShadowMapTransmissionCommon.ush"
+#include "VirtualShadowMapLightGrid.ush"
 #include "/Engine/Shared/VirtualShadowMapDefinitions.h"
 
 #if HAS_HAIR_STRANDS
@@ -268,9 +269,9 @@ FVirtualShadowMapSampleResult ProjectLight(
 			}
 		}
 
+		#if DIRECTIONAL_LIGHT
 		if (SMRTRayCount > 0)
 		{
-		#if DIRECTIONAL_LIGHT
 			Result = TraceDirectional(
 				VirtualShadowMapId,
 				Light,
@@ -283,11 +284,23 @@ FVirtualShadowMapSampleResult ProjectLight(
 				SMRTRayLengthScale,
 				Noise,
 				bBackfaceCull,
-				half3(WorldNormal),
+				WorldNormal,
 				SMRTTexelDitherScale,
 				SMRTExtrapolateSlope,
 				SMRTAdaptiveRayCount);
-		#else
+		}
+		else
+		{
+			FLWCVector3 WorldPosition = LWCSubtract(TranslatedWorldPosition, PrimaryView.PreViewTranslation);
+			Result = SampleVirtualShadowMapDirectional(
+				VirtualShadowMapId,
+				WorldPosition,
+				SMRTRayOffset,
+				WorldNormal);
+		}
+		#else // !DIRECTIONAL_LIGHT
+		if (SMRTRayCount > 0)
+		{
 			Result = TraceLocalLight(
 				VirtualShadowMapId,
 				Light,
@@ -305,16 +318,17 @@ FVirtualShadowMapSampleResult ProjectLight(
 				SMRTExtrapolateSlope,
 				SMRTMaxSlopeBias,
 				SMRTAdaptiveRayCount);
-		#endif // DIRECTIONAL_LIGHT
 		}
 		else
 		{
-			Result = SampleVirtualShadowMapTranslatedWorld(
+			FLWCVector3 WorldPosition = LWCSubtract(TranslatedWorldPosition, PrimaryView.PreViewTranslation);
+			Result = SampleVirtualShadowMapLocal(
 				VirtualShadowMapId,
-				TranslatedWorldPosition,
+				WorldPosition,
 				SMRTRayOffset,
 				WorldNormal);
 		}
+		#endif // DIRECTIONAL_LIGHT
 		
 		// Explanation:
 		//	- If bDataIsOpacity = 1: Data stores the surface's SSS opacity. We compute a transmitted color when the surface is not fully opaque (opacity == 1).
@@ -624,19 +638,16 @@ void VirtualShadowMapProjection(
 	FVirtualShadowMapSampleResult VisualizeResult = InitVirtualShadowMapSampleResult();
 
 #if ONE_PASS_PROJECTION
-	uint EyeIndex = 0; // TODO: Instanced stereo
-	uint GridLinearIndex = ComputeLightGridCellIndex( LocalPixelPos, SceneDepth, EyeIndex );
-	const FCulledLightsGridData CulledLightGridData = GetCulledLightsGrid( GridLinearIndex, EyeIndex );
-
+	const FCulledLightsGridData LightGridData = VirtualShadowMapGetLightsGrid(LocalPixelPos, SceneDepth);	
 	// We can only handle so many lights in our output encoding right now, so no purpose in computing more
-	uint LightCount = min(GetPackedShadowMaskMaxLightCount(), CulledLightGridData.NumLocalLights);
+	uint LightCount = min(GetPackedShadowMaskMaxLightCount(), LightGridData.NumLocalLights);
 
 	uint4 ShadowMask = InitializePackedShadowMask();
 
 	LOOP
 	for (uint Index = 0; Index < LightCount; ++Index)
 	{
-		const FLocalLightData LightData = GetLocalLightData( CulledLightGridData.DataStartIndex + Index, EyeIndex );
+		const FLocalLightData LightData = VirtualShadowMapGetLocalLightData(LightGridData, Index);
 		int VirtualShadowMapId = LightData.VirtualShadowMapId;
 
 		if (VirtualShadowMapId != INDEX_NONE)
@@ -663,7 +674,7 @@ void VirtualShadowMapProjection(
 		}
 	}
 
-	OutShadowMaskBits[ PixelPos ] = ~ShadowMask;
+	OutShadowMaskBits[PixelPos] = ShadowMask;
 
 #else // !ONE_PASS_PROJECTION
 	{
diff --git a/Engine/Shaders/Private/VirtualShadowMaps/VirtualShadowMapProjectionCommon.ush b/Engine/Shaders/Private/VirtualShadowMaps/VirtualShadowMapProjectionCommon.ush
--- a/Engine/Shaders/Private/VirtualShadowMaps/VirtualShadowMapProjectionCommon.ush
+++ b/Engine/Shaders/Private/VirtualShadowMaps/VirtualShadowMapProjectionCommon.ush
@@ -334,8 +334,9 @@ FVirtualShadowMapSampleResult InitVirtualShadowMapSampleResult()
 }
 
 
-
 /**
+* Sample virtual shadow (clip)map for a directional light
+*
 * VirtualShadowMapId is the ID of the virtual shadow map to sample
 * WorldPosition is the sample position in world space to project into the shadow map
 * RayStartDistance is an optional offset to move the lookup along the shadow ray towards the light
@@ -346,107 +347,64 @@ FVirtualShadowMapSampleResult InitVirtualShadowMapSampleResult()
 *   - If bUseOptimalBias is true, this is used to compute a receiver-plane-based bias for the sample point
 *   - The shading normal can be used if the geometric normal is not available, but divergence from the geometric normal can cause biasing issues.
 */
-FVirtualShadowMapSampleResult SampleVirtualShadowMapInner(int VirtualShadowMapId, FLWCVector3 WorldPosition, float RayStartDistance, bool bUseOptimalBias, float3 EstimatedGeoWorldNormal)
+FVirtualShadowMapSampleResult SampleVirtualShadowMapDirectional(
+	int VirtualShadowMapId,
+	FLWCVector3 WorldPosition,
+	float RayStartDistance,
+	float3 EstimatedGeoWorldNormal,
+	bool bUseOptimalBias = true)
 {
 	RayStartDistance = max(RayStartDistance, 0.0f);
 
 	FVirtualShadowMapProjectionShaderData BaseProjectionData = GetVirtualShadowMapProjectionData(VirtualShadowMapId);
+	checkSlow(BaseProjectionData.LightType == LIGHT_TYPE_DIRECTIONAL);
+	
+	const int ClipmapLevel = int(floor(CalcAbsoluteClipmapLevel(BaseProjectionData, WorldPosition)));
+	int ClipmapIndex = max(0, ClipmapLevel - BaseProjectionData.ClipmapLevel);
 
-	if (BaseProjectionData.LightType == LIGHT_TYPE_DIRECTIONAL)
+	// Check if sample is within the clipmap range (from camera)
+	if (ClipmapIndex < BaseProjectionData.ClipmapLevelCountRemaining)
 	{
-		const int ClipmapLevel = int(floor(CalcAbsoluteClipmapLevel(BaseProjectionData, WorldPosition)));
-		int ClipmapIndex = max(0, ClipmapLevel - BaseProjectionData.ClipmapLevel);
+		int ClipmapLevelVirtualShadowMapId = VirtualShadowMapId + ClipmapIndex;
+		FVirtualShadowMapProjectionShaderData ProjectionData = GetVirtualShadowMapProjectionData(ClipmapLevelVirtualShadowMapId);
 
-		// Check if sample is within the clipmap range (from camera)
-		if (ClipmapIndex < BaseProjectionData.ClipmapLevelCountRemaining)
-		{
-			int ClipmapLevelVirtualShadowMapId = VirtualShadowMapId + ClipmapIndex;
-			FVirtualShadowMapProjectionShaderData ProjectionData = GetVirtualShadowMapProjectionData(ClipmapLevelVirtualShadowMapId);
-
-			// No perspective divided needed for ortho projection
-			float3 ShadowTranslatedWorldPosition = LWCToFloat(LWCAdd(WorldPosition, ProjectionData.PreViewTranslation));
-			float4 ShadowUVz = mul(float4(ShadowTranslatedWorldPosition, 1.0f), ProjectionData.TranslatedWorldToShadowUVMatrix);
-			
-			FVirtualShadowMapSample SmSample;
-			SmSample = SampleVirtualShadowMapClipmap(ClipmapLevelVirtualShadowMapId, ShadowUVz.xy);
-			
-			if (SmSample.bValid)
-			{
-				int SampledClipmapIndex = SmSample.VirtualShadowMapId - VirtualShadowMapId;
-
-				FVirtualShadowMapSampleResult Result = InitVirtualShadowMapSampleResult();
-				Result.bValid = true;
-				Result.ShadowFactor = 1.0f;
-				Result.OccluderDistance = -1.0f;
-				Result.ClipmapIndexOrMipLevel = ClipmapIndex;
-				Result.VirtualTexelAddress = SmSample.VirtualTexelAddress;
-				Result.PhysicalTexelAddress = SmSample.PhysicalTexelAddress;
-				Result.RayCount = 1;
-
-				float OptimalSlopeBias = bUseOptimalBias 
-					? ComputeVirtualShadowMapOptimalSlopeBias(ClipmapLevelVirtualShadowMapId, SmSample, ShadowTranslatedWorldPosition, EstimatedGeoWorldNormal)
-					: 0.0f;
-				
-				float RayStartBias = -RayStartDistance * ProjectionData.ShadowViewToClipMatrix._33;
-				float BiasedDepth = SmSample.Depth - OptimalSlopeBias - RayStartBias;
-
-				if (BiasedDepth > ShadowUVz.z)
-				{
-					Result.ShadowFactor = 0.0f;
-					Result.OccluderDistance = ComputeOccluderDistanceOrtho(
-						ProjectionData.ShadowViewToClipMatrix,
-						SmSample.Depth,
-						ShadowUVz.z);
-				}
-
-				return Result;
-			}
-		}
-	}
-	else
-	{
-		FVirtualShadowMapProjectionShaderData ProjectionData = GetVirtualShadowMapProjectionData(VirtualShadowMapId);
+		// No perspective divided needed for ortho projection
 		float3 ShadowTranslatedWorldPosition = LWCToFloat(LWCAdd(WorldPosition, ProjectionData.PreViewTranslation));
-
-		if (ProjectionData.LightType != LIGHT_TYPE_SPOT)
-		{
-			VirtualShadowMapId += VirtualShadowMapGetCubeFace(ShadowTranslatedWorldPosition);
-			ProjectionData = GetVirtualShadowMapProjectionData(VirtualShadowMapId);
-		}
-
-		// TODO cubemap math directly instead of vector fetching matrix
 		float4 ShadowUVz = mul(float4(ShadowTranslatedWorldPosition, 1.0f), ProjectionData.TranslatedWorldToShadowUVMatrix);
-		ShadowUVz.xyz /= ShadowUVz.w;
 
-		FVirtualShadowMapSample SmSample = SampleVirtualShadowMap(VirtualShadowMapId, ShadowUVz.xy);
+		FVirtualShadowMapSample SmSample;
+		SmSample = SampleVirtualShadowMapClipmap(ClipmapLevelVirtualShadowMapId, ShadowUVz.xy);
+
 		if (SmSample.bValid)
 		{
+			int SampledClipmapIndex = SmSample.VirtualShadowMapId - VirtualShadowMapId;
+
 			FVirtualShadowMapSampleResult Result = InitVirtualShadowMapSampleResult();
 			Result.bValid = true;
 			Result.ShadowFactor = 1.0f;
 			Result.OccluderDistance = -1.0f;
-			Result.ClipmapIndexOrMipLevel = SmSample.MipLevel;
+			Result.ClipmapIndexOrMipLevel = ClipmapIndex;
 			Result.VirtualTexelAddress = SmSample.VirtualTexelAddress;
 			Result.PhysicalTexelAddress = SmSample.PhysicalTexelAddress;
 			Result.RayCount = 1;
 
-			float OptimalSlopeBias = bUseOptimalBias 
-				? ComputeVirtualShadowMapOptimalSlopeBias(VirtualShadowMapId, SmSample, ShadowTranslatedWorldPosition, EstimatedGeoWorldNormal)
-				: 0.0f;
+			float OptimalSlopeBias = 0.0f;
+			[BRANCH]
+			if (bUseOptimalBias)
+			{
+				OptimalSlopeBias = ComputeVirtualShadowMapOptimalSlopeBias(ClipmapLevelVirtualShadowMapId, SmSample, ShadowTranslatedWorldPosition, EstimatedGeoWorldNormal);
+			}
 
-			float RayStartBias = -RayStartDistance * ProjectionData.ShadowViewToClipMatrix._33 / ShadowUVz.w;
+			float RayStartBias = -RayStartDistance * ProjectionData.ShadowViewToClipMatrix._33;
 			float BiasedDepth = SmSample.Depth - OptimalSlopeBias - RayStartBias;
 
 			if (BiasedDepth > ShadowUVz.z)
 			{
 				Result.ShadowFactor = 0.0f;
-				// Shadow view matrix is rotation
-				float ReceiverDistance = length(ShadowTranslatedWorldPosition);
-				Result.OccluderDistance = RayStartDistance + ComputeOccluderDistancePerspective(
+				Result.OccluderDistance = ComputeOccluderDistanceOrtho(
 					ProjectionData.ShadowViewToClipMatrix,
 					SmSample.Depth,
-					ShadowUVz.z,
-					ReceiverDistance);
+					ShadowUVz.z);
 			}
 
 			return Result;
@@ -457,24 +415,104 @@ FVirtualShadowMapSampleResult SampleVirtualShadowMapInner(int VirtualShadowMapId
 	return InitVirtualShadowMapSampleResult();
 }
 
-// Sample in world space
-FVirtualShadowMapSampleResult SampleVirtualShadowMap(int VirtualShadowMapId, FLWCVector3 WorldPosition, float RayStartDistance, float3 EstimatedGeoWorldNormal)
+
+/**
+* Sample virtual shadow map for a local light
+* 
+* Parameter descriptions are as above
+*/
+FVirtualShadowMapSampleResult SampleVirtualShadowMapLocal(
+	int VirtualShadowMapId,
+	FLWCVector3 WorldPosition,
+	float RayStartDistance,
+	float3 EstimatedGeoWorldNormal,
+	bool bUseOptimalBias = true)
 {
-	return SampleVirtualShadowMapInner(VirtualShadowMapId, WorldPosition, RayStartDistance, true, EstimatedGeoWorldNormal);
+	RayStartDistance = max(RayStartDistance, 0.0f);
+
+	FVirtualShadowMapProjectionShaderData ProjectionData = GetVirtualShadowMapProjectionData(VirtualShadowMapId);
+	float3 ShadowTranslatedWorldPosition = LWCToFloat(LWCAdd(WorldPosition, ProjectionData.PreViewTranslation));
+
+	if (ProjectionData.LightType != LIGHT_TYPE_SPOT)
+	{
+		VirtualShadowMapId += VirtualShadowMapGetCubeFace(ShadowTranslatedWorldPosition);
+		ProjectionData = GetVirtualShadowMapProjectionData(VirtualShadowMapId);
+	}
+
+	// TODO cubemap math directly instead of vector fetching matrix
+	float4 ShadowUVz = mul(float4(ShadowTranslatedWorldPosition, 1.0f), ProjectionData.TranslatedWorldToShadowUVMatrix);
+	ShadowUVz.xyz /= ShadowUVz.w;
+
+	FVirtualShadowMapSample SmSample = SampleVirtualShadowMap(VirtualShadowMapId, ShadowUVz.xy);
+	if (SmSample.bValid)
+	{
+		FVirtualShadowMapSampleResult Result = InitVirtualShadowMapSampleResult();
+		Result.bValid = true;
+		Result.ShadowFactor = 1.0f;
+		Result.OccluderDistance = -1.0f;
+		Result.ClipmapIndexOrMipLevel = SmSample.MipLevel;
+		Result.VirtualTexelAddress = SmSample.VirtualTexelAddress;
+		Result.PhysicalTexelAddress = SmSample.PhysicalTexelAddress;
+		Result.RayCount = 1;
+
+		float OptimalSlopeBias = 0.0f;
+		[BRANCH]
+		if (bUseOptimalBias)
+		{
+			OptimalSlopeBias = ComputeVirtualShadowMapOptimalSlopeBias(VirtualShadowMapId, SmSample, ShadowTranslatedWorldPosition, EstimatedGeoWorldNormal);
+		}
+
+		float RayStartBias = -RayStartDistance * ProjectionData.ShadowViewToClipMatrix._33 / ShadowUVz.w;
+		float BiasedDepth = SmSample.Depth - OptimalSlopeBias - RayStartBias;
+
+		if (BiasedDepth > ShadowUVz.z)
+		{
+			Result.ShadowFactor = 0.0f;
+			// Shadow view matrix is rotation
+			float ReceiverDistance = length(ShadowTranslatedWorldPosition);
+			Result.OccluderDistance = RayStartDistance + ComputeOccluderDistancePerspective(
+				ProjectionData.ShadowViewToClipMatrix,
+				SmSample.Depth,
+				ShadowUVz.z,
+				ReceiverDistance);
+		}
+
+		return Result;
+	}
+
+	// Invalid
+	return InitVirtualShadowMapSampleResult();
 }
-FVirtualShadowMapSampleResult SampleVirtualShadowMap(int VirtualShadowMapId, FLWCVector3 WorldPosition, float RayStartDistance = 0.0f)
+
+
+// Some common variants/convenience functions
+
+// Sample in PrimaryView translated world space with no optimal bias
+FVirtualShadowMapSampleResult SampleVirtualShadowMapDirectional(int VirtualShadowMapId, float3 TranslatedWorldPosition, float RayStartDistance = 0.0f)
 {
-	return SampleVirtualShadowMapInner(VirtualShadowMapId, WorldPosition, RayStartDistance, false, float3(0, 0, 0));
+	FLWCVector3 WorldPosition = LWCSubtract(TranslatedWorldPosition, PrimaryView.PreViewTranslation);
+	return SampleVirtualShadowMapDirectional(VirtualShadowMapId, WorldPosition, RayStartDistance, float3(0, 0, 0), false);
 }
-
-// Sample in PrimaryView translated world space
-FVirtualShadowMapSampleResult SampleVirtualShadowMapTranslatedWorld(int VirtualShadowMapId, float3 TranslatedWorldPosition, float RayStartDistance, float3 EstimatedGeoWorldNormal)
+FVirtualShadowMapSampleResult SampleVirtualShadowMapLocal(int VirtualShadowMapId, float3 TranslatedWorldPosition, float RayStartDistance = 0.0f)
 {
 	FLWCVector3 WorldPosition = LWCSubtract(TranslatedWorldPosition, PrimaryView.PreViewTranslation);
-	return SampleVirtualShadowMapInner(VirtualShadowMapId, WorldPosition, RayStartDistance, true, EstimatedGeoWorldNormal);
+	return SampleVirtualShadowMapLocal(VirtualShadowMapId, WorldPosition, RayStartDistance, float3(0, 0, 0), false);
 }
+/**
+* Sample a virtual shadow map for an unknown light type with no optimal bias.
+* Will route to the appropriate function above.
+* Generally the directional/local variants should be used when the light's type is known to cut down on the
+* code/register overhead of including both paths.
+*/
 FVirtualShadowMapSampleResult SampleVirtualShadowMapTranslatedWorld(int VirtualShadowMapId, float3 TranslatedWorldPosition, float RayStartDistance = 0.0f)
 {
-	FLWCVector3 WorldPosition = LWCSubtract(TranslatedWorldPosition, PrimaryView.PreViewTranslation);
-	return SampleVirtualShadowMapInner(VirtualShadowMapId, WorldPosition, RayStartDistance, false, float3(0, 0, 0));
+	FVirtualShadowMapProjectionShaderData BaseProjectionData = GetVirtualShadowMapProjectionData(VirtualShadowMapId);
+	if (BaseProjectionData.LightType == LIGHT_TYPE_DIRECTIONAL)
+	{
+		return SampleVirtualShadowMapDirectional(VirtualShadowMapId, TranslatedWorldPosition, RayStartDistance);
+	}
+	else
+	{
+		return SampleVirtualShadowMapLocal(VirtualShadowMapId, TranslatedWorldPosition, RayStartDistance);
+	}
 }
diff --git a/Engine/Shaders/Private/VirtualShadowMaps/VirtualShadowMapProjectionComposite.usf b/Engine/Shaders/Private/VirtualShadowMaps/VirtualShadowMapProjectionComposite.usf
--- a/Engine/Shaders/Private/VirtualShadowMaps/VirtualShadowMapProjectionComposite.usf
+++ b/Engine/Shaders/Private/VirtualShadowMaps/VirtualShadowMapProjectionComposite.usf
@@ -50,11 +50,11 @@ void VirtualShadowMapCompositeFromMaskBitsPS(
 	const float DeviceZ = SceneTexturesStruct.SceneDepthTexture.Load(int3(PixelPos, 0)).r;
 	const float SceneDepth = ConvertFromDeviceZ(DeviceZ);
 
-	uint EyeIndex = 0; // TODO: Instanced stereo
+	// Possibly needed to do fallback single sample VSM lookup
+	const float3 TranslatedWorldPosition = SvPositionToTranslatedWorld(SvPosition);
+
 	float ShadowMask = GetVirtualShadowMapMaskForLight(
-		ShadowMaskBits, PixelPos, SceneDepth, EyeIndex, CompositeVirtualShadowMapId);;
-	
-	// TODO: Subsurface...?
-		
+		ShadowMaskBits[PixelPos], PixelPos, SceneDepth, CompositeVirtualShadowMapId, TranslatedWorldPosition);
+			
 	OutShadowMask = EncodeLightAttenuation(ShadowMask.xxxx);
 }
diff --git a/Engine/Shaders/Private/VirtualShadowMaps/VirtualShadowMapProjectionDirectional.ush b/Engine/Shaders/Private/VirtualShadowMaps/VirtualShadowMapProjectionDirectional.ush
--- a/Engine/Shaders/Private/VirtualShadowMaps/VirtualShadowMapProjectionDirectional.ush
+++ b/Engine/Shaders/Private/VirtualShadowMaps/VirtualShadowMapProjectionDirectional.ush
@@ -213,7 +213,7 @@ FVirtualShadowMapSampleResult TraceDirectional(
 	float RayLengthScale,
 	const float Noise,
 	bool bCullBackface,
-	half3 WorldNormal,
+	float3 WorldNormal,
 	float TexelDitherScale,
 	float ExtrapolateSlope,
 	uint SMRTAdaptiveRayCount)
diff --git a/Engine/Shaders/Private/VirtualShadowMaps/VirtualShadowMapProjectionStructs.ush b/Engine/Shaders/Private/VirtualShadowMaps/VirtualShadowMapProjectionStructs.ush
--- a/Engine/Shaders/Private/VirtualShadowMaps/VirtualShadowMapProjectionStructs.ush
+++ b/Engine/Shaders/Private/VirtualShadowMaps/VirtualShadowMapProjectionStructs.ush
@@ -29,8 +29,7 @@ struct FVirtualShadowMapProjectionShaderData
 	float LightRadius;
 
 	// Derived data for convenience when passing the structure around
-	int VirtualShadowMapId;	
-	bool bCurrentDistantLight; // Low detail == single page only
+	int VirtualShadowMapId;
 	bool bUnCached; // see VSM_PROJ_FLAG_UNCACHED
 	bool bUnreferenced; // see VSM_PROJ_FLAG_UNREFERENCED
 };
@@ -83,7 +82,6 @@ FVirtualShadowMapProjectionShaderData DecodeVirtualShadowMapProjectionData(ByteA
 	Result.PreViewTranslation				= MakeLWCVector3(PreViewTranslationTile, PreViewTranslationOffset);
 	Result.ClipmapWorldOrigin				= LWCNegate(MakeLWCVector3(PreViewTranslationTile, NegativeClipmapWorldOriginOffset));
 
-	Result.bCurrentDistantLight = (Result.Flags & VSM_PROJ_FLAG_CURRENT_DISTANT_LIGHT) != 0U;
 	Result.bUnCached = (Result.Flags & VSM_PROJ_FLAG_UNCACHED) != 0U;
 	Result.bUnreferenced = (Result.Flags & VSM_PROJ_FLAG_UNREFERENCED) != 0U;
 
diff --git a/Engine/Shaders/Private/VolumetricCloud.usf b/Engine/Shaders/Private/VolumetricCloud.usf
--- a/Engine/Shaders/Private/VolumetricCloud.usf
+++ b/Engine/Shaders/Private/VolumetricCloud.usf
@@ -979,7 +979,7 @@ void MainCommon(in FMaterialPixelParameters MaterialParameters, in float4 SvPosi
 #if VIRTUAL_SHADOW_MAP 
 				if (RenderVolumetricCloudParameters.VirtualShadowMapId0 != INDEX_NONE)
 				{
-					FVirtualShadowMapSampleResult VirtualShadowMapSample = SampleVirtualShadowMapTranslatedWorld(RenderVolumetricCloudParameters.VirtualShadowMapId0, SampleTranslatedWorldPosition);
+					FVirtualShadowMapSampleResult VirtualShadowMapSample = SampleVirtualShadowMapDirectional(RenderVolumetricCloudParameters.VirtualShadowMapId0, SampleTranslatedWorldPosition);
 					OpaqueShadow *= VirtualShadowMapSample.ShadowFactor;
 				}
 #endif // VIRTUALSHADOW_MAP
diff --git a/Engine/Source/Runtime/Renderer/Private/LightRendering.cpp b/Engine/Source/Runtime/Renderer/Private/LightRendering.cpp
--- a/Engine/Source/Runtime/Renderer/Private/LightRendering.cpp
+++ b/Engine/Source/Runtime/Renderer/Private/LightRendering.cpp
@@ -1177,10 +1177,6 @@ void FSceneRenderer::GatherAndSortLights(FSortedLightSetSceneInfo& OutSortedLigh
 						&& LightSceneInfoCompact.LightType != LightType_Rect
 						&& !bSupportedByLumenDirectLighting;
 
-					// One pass projection is supported for lights with only virtual shadow maps
-					// TODO: Exclude lights that also have non-virtual shadow maps
-					bool bHasVirtualShadowMap = VisibleLightInfos[LightSceneInfo->Id].GetVirtualShadowMapId(&Views[ViewIndex]) != INDEX_NONE;
-					SortedLightInfo->SortKey.Fields.bDoesNotWriteIntoPackedShadowMask = LightSceneInfoCompact.LightType == LightType_Directional || !bHasVirtualShadowMap;
 					SortedLightInfo->SortKey.Fields.bClusteredDeferredNotSupported = !bClusteredDeferredSupported;
 					SortedLightInfo->SortKey.Fields.bHandledByLumen = bSupportedByLumenDirectLighting;
 					break;
@@ -1528,10 +1524,11 @@ void FDeferredShadingSceneRenderer::RenderLights(
 					bElideScreenShadowMask = 
 						CVarOnePassProjectionSkipScreenShadowMask.GetValueOnRenderThread() != 0 &&
 						ShadowSceneRenderer->UsePackedShadowMaskBits() &&
-						!SortedLightInfo.SortKey.Fields.bDoesNotWriteIntoPackedShadowMask &&
 						OcclusionType == FLightOcclusionType::Shadowmap &&
 						!(bDirectLighting && bDrawLightFunction) &&
 						!bDrawPreviewIndicator &&
+						SortedLightInfo.SortKey.Fields.LightType != LightType_Directional &&
+						VisibleLightInfo.VirtualShadowMapId != INDEX_NONE &&	// Not a directional light, so no per-view clipmaps
 						VisibleLightInfo.ContainsOnlyVirtualShadowMaps();
 
 					if (!SharedScreenShadowMaskTexture || !SharedScreenShadowMaskSubPixelTexture)
diff --git a/Engine/Source/Runtime/Renderer/Private/LightSceneInfo.h b/Engine/Source/Runtime/Renderer/Private/LightSceneInfo.h
--- a/Engine/Source/Runtime/Renderer/Private/LightSceneInfo.h
+++ b/Engine/Source/Runtime/Renderer/Private/LightSceneInfo.h
@@ -94,8 +94,6 @@ struct FSortedLightSceneInfo
 			uint32 bShadowed : 1;
 			/** Whether the light is NOT a simple light - they always support tiled/clustered but may want to be selected separately. */
 			uint32 bIsNotSimpleLight : 1;
-			/* We want to sort the lights that write into the packed shadow mask (when enabled) to the front of the list so we don't waste slots in the packed shadow mask. */
-			uint32 bDoesNotWriteIntoPackedShadowMask : 1;
 			/** 
 			 * True if the light doesn't support clustered deferred, logic is inverted so that lights that DO support clustered deferred will sort first in list 
 			 * Super-set of lights supporting tiled, so the tiled lights will end up in the first part of this range.
diff --git a/Engine/Source/Runtime/Renderer/Private/Shadows/ShadowSceneRenderer.cpp b/Engine/Source/Runtime/Renderer/Private/Shadows/ShadowSceneRenderer.cpp
--- a/Engine/Source/Runtime/Renderer/Private/Shadows/ShadowSceneRenderer.cpp
+++ b/Engine/Source/Runtime/Renderer/Private/Shadows/ShadowSceneRenderer.cpp
@@ -187,8 +187,7 @@ FVirtualShadowMapProjectionShaderData FShadowSceneRenderer::GetLocalLightProject
 	bool bIsSinglePageSM = FVirtualShadowMapArray::IsSinglePage(VirtualShadowMapId);
 	check(VirtualShadowMapId != INDEX_NONE && CacheEntry->bIsDistantLight == bIsSinglePageSM);
 
-	uint32 Flags = bIsSinglePageSM ? VSM_PROJ_FLAG_CURRENT_DISTANT_LIGHT : 0U;
-	Flags |= CacheEntry->IsUncached() ? VSM_PROJ_FLAG_UNCACHED : 0U;
+	uint32 Flags = CacheEntry->IsUncached() ? VSM_PROJ_FLAG_UNCACHED : 0U;
 
 	const FViewMatrices ViewMatrices = ProjectedShadowInfo->GetShadowDepthRenderingViewMatrices(MapIndex, true);
 	const FLargeWorldRenderPosition PreViewTranslation(ProjectedShadowInfo->PreShadowTranslation);
diff --git a/Engine/Source/Runtime/Renderer/Private/VirtualShadowMaps/VirtualShadowMapArray.cpp b/Engine/Source/Runtime/Renderer/Private/VirtualShadowMaps/VirtualShadowMapArray.cpp
--- a/Engine/Source/Runtime/Renderer/Private/VirtualShadowMaps/VirtualShadowMapArray.cpp
+++ b/Engine/Source/Runtime/Renderer/Private/VirtualShadowMaps/VirtualShadowMapArray.cpp
@@ -421,6 +421,8 @@ void FVirtualShadowMapArray::Initialize(
 	UniformParameters.ProjectionData = GraphBuilder.CreateSRV(GSystemTextures.GetDefaultByteAddressBuffer(GraphBuilder, sizeof(FVirtualShadowMapProjectionShaderData)));
 	UniformParameters.PageFlags = GraphBuilder.CreateSRV(GSystemTextures.GetDefaultStructuredBuffer(GraphBuilder, sizeof(uint32)));
 	UniformParameters.PageRectBounds = GraphBuilder.CreateSRV(GSystemTextures.GetDefaultStructuredBuffer(GraphBuilder, sizeof(FIntVector4)));
+	UniformParameters.LightGridData = GraphBuilder.CreateSRV(GSystemTextures.GetDefaultStructuredBuffer(GraphBuilder, sizeof(uint32)));
+	UniformParameters.NumCulledLightsGrid = GraphBuilder.CreateSRV(GSystemTextures.GetDefaultStructuredBuffer(GraphBuilder, sizeof(uint32)));
 
 	if (bEnabled)
 	{
@@ -705,8 +707,6 @@ class FGeneratePageFlagsFromPixelsCS : public FVirtualPageManagementShader
 		SHADER_PARAMETER_RDG_TEXTURE(Texture2D<uint2>, VisBuffer64)
 		SHADER_PARAMETER_RDG_BUFFER_UAV(RWStructuredBuffer<uint>, OutPageRequestFlags)
 		SHADER_PARAMETER_RDG_BUFFER_SRV(StructuredBuffer< uint >, DirectionalLightIds)
-		SHADER_PARAMETER_RDG_BUFFER_SRV(StructuredBuffer< uint >, PrunedLightGridData)
-		SHADER_PARAMETER_RDG_BUFFER_SRV(StructuredBuffer< uint >, PrunedNumCulledLightsGrid)
 		// PERMUTATION_WATER_DEPTH
 		SHADER_PARAMETER_RDG_TEXTURE(Texture2D<float>, SingleLayerWaterDepthTexture)
 		SHADER_PARAMETER_RDG_BUFFER_SRV(StructuredBuffer< uint >, SingleLayerWaterTileMask)
@@ -1449,19 +1449,15 @@ void FVirtualShadowMapArray::BuildPageAllocations(
 				);
 			}
 
-			// Mark pages based on projected depth buffer pixels
-			if (CVarMarkPixelPages.GetValueOnRenderThread() != 0)
+			// Prune light grid to remove lights without VSMs
 			{
-				// It's currently safe to overlap these passes that all write to same page request flags
-				FRDGBufferUAVRef PageRequestFlagsUAV = GraphBuilder.CreateUAV(PageRequestFlagsRDG, ERDGUnorderedAccessViewFlags::SkipBarrier);
-
 				const bool bLightGridUses16BitBuffers = LightGridUses16BitBuffers(View.GetShaderPlatform());
 				FRDGBufferSRVRef CulledLightDataGrid = bLightGridUses16BitBuffers ? View.ForwardLightingResources.ForwardLightData->CulledLightDataGrid16Bit : View.ForwardLightingResources.ForwardLightData->CulledLightDataGrid32Bit;
 				uint32 LightGridIndexBufferSize = CulledLightDataGrid->Desc.Buffer->Desc.GetSize();
-				FRDGBufferRef PrunedLightGridDataRDG = GraphBuilder.CreateBuffer(FRDGBufferDesc::CreateStructuredDesc(sizeof(uint32), LightGridIndexBufferSize), TEXT("Shadow.Virtual.PrunedLightGridData"));
-				
+				FRDGBufferRef PrunedLightGridDataRDG = GraphBuilder.CreateBuffer(FRDGBufferDesc::CreateStructuredDesc(sizeof(uint32), LightGridIndexBufferSize), TEXT("Shadow.Virtual.LightGridData"));
+
 				const uint32 NumLightGridCells = View.ForwardLightingResources.ForwardLightData->NumGridCells;
-				FRDGBufferRef PrunedNumCulledLightsGridRDG = GraphBuilder.CreateBuffer(FRDGBufferDesc::CreateStructuredDesc(sizeof(uint32), NumLightGridCells), TEXT("Shadow.Virtual.PrunedNumCulledLightsGrid"));
+				FRDGBufferRef PrunedNumCulledLightsGridRDG = GraphBuilder.CreateBuffer(FRDGBufferDesc::CreateStructuredDesc(sizeof(uint32), NumLightGridCells), TEXT("Shadow.Virtual.NumCulledLightsGrid"));
 
 				{
 					FPruneLightGridCS::FParameters* PassParameters = GraphBuilder.AllocParameters< FPruneLightGridCS::FParameters >();
@@ -1475,6 +1471,16 @@ void FVirtualShadowMapArray::BuildPageAllocations(
 					FComputeShaderUtils::AddPass(GraphBuilder, RDG_EVENT_NAME("PruneLightGrid"), ComputeShader, PassParameters, FComputeShaderUtils::GetGroupCount(NumLightGridCells, FPruneLightGridCS::DefaultCSGroupX));
 				};
 
+				UniformParameters.LightGridData = GraphBuilder.CreateSRV(PrunedLightGridDataRDG);
+				UniformParameters.NumCulledLightsGrid = GraphBuilder.CreateSRV(PrunedNumCulledLightsGridRDG);
+			}
+
+			// Mark pages based on projected depth buffer pixels
+			if (CVarMarkPixelPages.GetValueOnRenderThread() != 0)
+			{
+				// It's currently safe to overlap these passes that all write to same page request flags
+				FRDGBufferUAVRef PageRequestFlagsUAV = GraphBuilder.CreateUAV(PageRequestFlagsRDG, ERDGUnorderedAccessViewFlags::SkipBarrier);
+
 				auto GeneratePageFlags = [&](const EVirtualShadowMapProjectionInputType InputType)
 				{
 					FGeneratePageFlagsFromPixelsCS::FParameters* PassParameters = GraphBuilder.AllocParameters< FGeneratePageFlagsFromPixelsCS::FParameters >();
@@ -1521,8 +1527,6 @@ void FVirtualShadowMapArray::BuildPageAllocations(
 					PassParameters->OutPageRequestFlags = PageRequestFlagsUAV;
 					PassParameters->ForwardLightData = View.ForwardLightingResources.ForwardLightUniformBuffer;
 					PassParameters->DirectionalLightIds = GraphBuilder.CreateSRV(DirectionalLightIdsRDG);
-					PassParameters->PrunedLightGridData = GraphBuilder.CreateSRV(PrunedLightGridDataRDG);
-					PassParameters->PrunedNumCulledLightsGrid = GraphBuilder.CreateSRV(PrunedNumCulledLightsGridRDG);
 					bool bWaterEnabled = false;
 					if (SingleLayerWaterPrePassResult && InputType == EVirtualShadowMapProjectionInputType::GBuffer)
 					{
diff --git a/Engine/Source/Runtime/Renderer/Private/VirtualShadowMaps/VirtualShadowMapArray.h b/Engine/Source/Runtime/Renderer/Private/VirtualShadowMaps/VirtualShadowMapArray.h
--- a/Engine/Source/Runtime/Renderer/Private/VirtualShadowMaps/VirtualShadowMapArray.h
+++ b/Engine/Source/Runtime/Renderer/Private/VirtualShadowMaps/VirtualShadowMapArray.h
@@ -152,6 +152,11 @@ BEGIN_GLOBAL_SHADER_PARAMETER_STRUCT(FVirtualShadowMapUniformParameters, )
 	SHADER_PARAMETER_RDG_BUFFER_SRV(StructuredBuffer<uint>, PageFlags)
 	SHADER_PARAMETER_RDG_BUFFER_SRV(StructuredBuffer<uint4>, PageRectBounds)
 	SHADER_PARAMETER_RDG_TEXTURE(Texture2DArray<uint>, PhysicalPagePool)
+
+	// Light grid with only the lights that have VSMs present
+	// Still references the original indices from the global light grid
+	SHADER_PARAMETER_RDG_BUFFER_SRV(StructuredBuffer< uint >, LightGridData)
+	SHADER_PARAMETER_RDG_BUFFER_SRV(StructuredBuffer< uint >, NumCulledLightsGrid)
 END_GLOBAL_SHADER_PARAMETER_STRUCT()
 
 BEGIN_SHADER_PARAMETER_STRUCT(FVirtualShadowMapSamplingParameters, )
diff --git a/Engine/Source/Runtime/Renderer/Private/VirtualShadowMaps/VirtualShadowMapCacheManager.cpp b/Engine/Source/Runtime/Renderer/Private/VirtualShadowMaps/VirtualShadowMapCacheManager.cpp
--- a/Engine/Source/Runtime/Renderer/Private/VirtualShadowMaps/VirtualShadowMapCacheManager.cpp
+++ b/Engine/Source/Runtime/Renderer/Private/VirtualShadowMaps/VirtualShadowMapCacheManager.cpp
@@ -986,12 +986,14 @@ void FVirtualShadowMapArrayCacheManager::ExtractFrameData(
 			// Enqueue readback
 			StaticGPUInvalidationsFeedback.SubmitFeedbackBuffer(GraphBuilder, VirtualShadowMapArray.StaticInvalidatingPrimitivesRDG);
 			
-			// Store but drop any temp references embedded in the uniform parameters this frame.
-			// We'll reestablish them when we reimport the extracted resources next frame
+			// Store but drop any temp references embedded in the uniform parameters this frame
 			PrevUniformParameters = VirtualShadowMapArray.UniformParameters;
 			PrevUniformParameters.ProjectionData = nullptr;
 			PrevUniformParameters.PageTable = nullptr;
-			PrevUniformParameters.PhysicalPagePool = nullptr;
+			PrevUniformParameters.PageRectBounds = nullptr;
+			PrevUniformParameters.PageFlags = nullptr;
+			PrevUniformParameters.LightGridData = nullptr;
+			PrevUniformParameters.NumCulledLightsGrid = nullptr;
 		}
 
 		// propagate current-frame primitive state to cache entry
@@ -1301,8 +1303,10 @@ void FVirtualShadowMapArrayCacheManager::ProcessInvalidations(
 			UniformParameters->PageTable = RegExtCreateSrv(PrevBuffers.PageTable, TEXT("Shadow.Virtual.PrevPageTable"));
 			UniformParameters->PageFlags = RegExtCreateSrv(PrevBuffers.PageFlags, TEXT("Shadow.Virtual.PrevPageFlags"));
 			UniformParameters->PageRectBounds = RegExtCreateSrv(PrevBuffers.PageRectBounds, TEXT("Shadow.Virtual.PrevPageRectBounds"));
-			// Unused in this path
+			// Unused in this path... may be a better way to handle this
 			UniformParameters->PhysicalPagePool = GSystemTextures.GetZeroUIntArrayDummy(GraphBuilder);
+			UniformParameters->LightGridData = GraphBuilder.CreateSRV(GSystemTextures.GetDefaultStructuredBuffer(GraphBuilder, sizeof(uint32)));
+			UniformParameters->NumCulledLightsGrid = GraphBuilder.CreateSRV(GSystemTextures.GetDefaultStructuredBuffer(GraphBuilder, sizeof(uint32)));
 		}
 		PassParameters->VirtualShadowMap = GraphBuilder.CreateUniformBuffer(UniformParameters);
 
