Rename XXUsesComplexSpecialRenderPath -> IsComplexSpecialMaterial.

diff --git a/Engine/Shaders/Private/Substrate/Substrate.ush b/Engine/Shaders/Private/Substrate/Substrate.ush
--- a/Engine/Shaders/Private/Substrate/Substrate.ush
+++ b/Engine/Shaders/Private/Substrate/Substrate.ush
@@ -2070,6 +2070,7 @@ struct FSubstratePixelHeader
 	bool IsSimpleMaterial()										{ return (this.State & HEADER_MASK_MATERIALMODE) == HEADER_MATERIALMODE_SLAB_SIMPLE || SUBSTRATE_FASTPATH; }
 	bool IsSingleMaterial()										{ return (this.State & HEADER_MASK_MATERIALMODE) == HEADER_MATERIALMODE_SLAB_SINGLE || SUBSTRATE_SINGLEPATH; }
 	bool IsComplexMaterial()									{ return (this.State & HEADER_MASK_MATERIALMODE) == HEADER_MATERIALMODE_SLAB_COMPLEX && !SUBSTRATE_FASTPATH && !SUBSTRATE_SINGLEPATH; }
+	bool IsComplexSpecialMaterial()								{ return (this.State & HEADER_COMPLEXSPECIALPATH_MASK) != 0; }
 	bool IsSingleLayerWater()									{ return (this.State & HEADER_MASK_MATERIALMODE) == HEADER_MATERIALMODE_SLWATER; }
 	bool IsHair()												{ return (this.State & HEADER_MASK_MATERIALMODE) == HEADER_MATERIALMODE_HAIR; }
 	bool IsEye()												{ return (this.State & HEADER_MASK_MATERIALMODE) == HEADER_MATERIALMODE_EYE; }
@@ -2099,11 +2100,10 @@ struct FSubstratePixelHeader
 	void SetHasSubsurface(bool bIn)								{ this.State |= (bIn ? HEADER_MASK_HASSUBSURFACE : 0u); }
 	void SetDynamicIndirectShadowCasterRepresentation(bool bIn)	{ this.State |= (bIn ? HEADER_MASK_HASDYNINDIRECTSHADOWCASTER : 0u); }
 
-	bool IsSubstrateMaterial()										{ return this.GetMaterialMode() != HEADER_MATERIALMODE_NONE; }
+	bool IsSubstrateMaterial()									{ return this.GetMaterialMode() != HEADER_MATERIALMODE_NONE; }
 	void SetMaterialMode(uint MaterialMode)						{ this.State = (this.State & ~HEADER_MASK_MATERIALMODE) | (MaterialMode & HEADER_MASK_MATERIALMODE); }
 
-	bool GetUsesComplexSpecialRenderPath()						{ return (this.State & HEADER_COMPLEXSPECIALPATH_MASK) != 0; }
-	void SetUsesComplexSpecialRenderPath(bool bIn)				{ this.State |= (bIn ? HEADER_COMPLEXSPECIALPATH_MASK : 0u); }
+	void SetIsComplexSpecialMaterial(bool bIn)					{ this.State |= (bIn ? HEADER_COMPLEXSPECIALPATH_MASK : 0u); }
 };
 
 #if SUBSTRATE_DEFERRED_SHADING && SUBSTRATE_MATERIALCONTAINER_IS_VIEWRESOURCE
@@ -3818,7 +3818,7 @@ FSubstratePixelHeader UnpackSubstrateHeaderIn(uint PackedHeader, inout FSubstrat
 		uint SharedLocalBasesCount = HEADER_GETSHAREDLOCALBASESCOUNT(PackedHeader); // ?
 		uint SharedLocalBasesTypes = HEADER_GETSHAREDLOCALBASESTYPE(PackedHeader);
 
-		Out.SetUsesComplexSpecialRenderPath((PackedHeader & HEADER_COMPLEXSPECIALPATH_MASK) > 0);
+		Out.SetIsComplexSpecialMaterial((PackedHeader & HEADER_COMPLEXSPECIALPATH_MASK) > 0);
 
 	#if SUBSTRATE_DEFERRED_SHADING
 		Out.SharedLocalBasesTypes_PackedHeader = SharedLocalBasesTypes;
diff --git a/Engine/Shaders/Private/Substrate/SubstrateExport.ush b/Engine/Shaders/Private/Substrate/SubstrateExport.ush
--- a/Engine/Shaders/Private/Substrate/SubstrateExport.ush
+++ b/Engine/Shaders/Private/Substrate/SubstrateExport.ush
@@ -227,7 +227,7 @@ void PackSubstrateOut(
 						const bool bSSSMask = BSDF.bIsBottom;	// SSS and Thin can only be on the bottom layer, so ignoring bIsSimpleVolume here.
 						SubstratePixelHeader.SetHasSubsurface(bSSSMask && bHasSSS && !bIsSimpleVolume);
 
-						SubstratePixelHeader.SetUsesComplexSpecialRenderPath(bForceComplexSpecialPath);
+						SubstratePixelHeader.SetIsComplexSpecialMaterial(bForceComplexSpecialPath);
 						break;
 					}
 
diff --git a/Engine/Shaders/Private/Substrate/SubstrateMaterialClassification.usf b/Engine/Shaders/Private/Substrate/SubstrateMaterialClassification.usf
--- a/Engine/Shaders/Private/Substrate/SubstrateMaterialClassification.usf
+++ b/Engine/Shaders/Private/Substrate/SubstrateMaterialClassification.usf
@@ -129,7 +129,7 @@ void TileMainCS(uint2 DispatchThreadId : SV_DispatchThreadID, uint LinearIndex :
 
 		const bool bIsSimple =  SubstratePixelHeader.IsSimpleMaterial() || SubstratePixelHeader.BSDFCount == 0; // BSDFCount == 0 ensures that non-Substrate pixel, like sky pixels, won't make a simple tile flagged as complex
 		const bool bIsSingle = !SubstratePixelHeader.IsSimpleMaterial() && SubstratePixelHeader.IsSingleMaterial();
-		const bool bIsComplexSpecial = SubstratePixelHeader.GetUsesComplexSpecialRenderPath();
+		const bool bIsComplexSpecial = SubstratePixelHeader.IsComplexSpecialMaterial();
 		bContainsSubstrateMaterial  = SubstratePixelHeader.BSDFCount > 0;
 		bContainsSimpleMaterial  = bIsSimple;
 		bContainsSingleMaterial  = bIsSingle;
diff --git a/Engine/Shaders/Private/Substrate/SubstrateVisualize.usf b/Engine/Shaders/Private/Substrate/SubstrateVisualize.usf
--- a/Engine/Shaders/Private/Substrate/SubstrateVisualize.usf
+++ b/Engine/Shaders/Private/Substrate/SubstrateVisualize.usf
@@ -94,7 +94,7 @@ void PrintPixelType(inout FShaderPrintContext Ctx, in FSubstratePixelHeader Head
 {
 	if (Header.IsSimpleMaterial())						Print(Ctx, TEXT("Simple  "), InColor);
 	else if (Header.IsSingleMaterial())					Print(Ctx, TEXT("Single  "), InColor);
-	else if (Header.GetUsesComplexSpecialRenderPath())	Print(Ctx, TEXT("Complex Special "), InColor);
+	else if (Header.IsComplexSpecialMaterial())			Print(Ctx, TEXT("Complex Special "), InColor);
 	else if (Header.IsComplexMaterial())				Print(Ctx, TEXT("Complex "), InColor);
 	else if (Header.IsSingleLayerWater())				Print(Ctx, TEXT("SLWater "), InColor);
 	else if (Header.IsHair())							Print(Ctx, TEXT("Hair    "), InColor);
