Fixed substrate material previw for subsurface

diff --git a/Engine/Shaders/Private/Substrate/Substrate.ush b/Engine/Shaders/Private/Substrate/Substrate.ush
--- a/Engine/Shaders/Private/Substrate/Substrate.ush
+++ b/Engine/Shaders/Private/Substrate/Substrate.ush
@@ -3063,10 +3063,10 @@ float3 SubstrateGetBSDFBaseColor(in FSubstrateBSDF BSDF)
 {
 	switch (BSDF_GETTYPE(BSDF))
 	{
-	case SUBSTRATE_BSDF_TYPE_SLAB: 			return lerp(SLAB_DIFFUSEALBEDO(BSDF), SLAB_F0(BSDF), F0RGBToMetallic(SLAB_F0(BSDF)));  // This BaseColor recovery based on Metallic recovers is not accurate, and only used for legacy conversion purpose (e.g., PostProcessMaterial texture access)
-	case SUBSTRATE_BSDF_TYPE_HAIR: 			return HAIR_BASECOLOR(BSDF);
+	case SUBSTRATE_BSDF_TYPE_SLAB: 				return lerp(SLAB_DIFFUSEALBEDO(BSDF), SLAB_F0(BSDF), F0RGBToMetallic(SLAB_F0(BSDF)));  // This BaseColor recovery based on Metallic recovers is not accurate, and only used for legacy conversion purpose (e.g., PostProcessMaterial texture access)
+	case SUBSTRATE_BSDF_TYPE_HAIR: 				return HAIR_BASECOLOR(BSDF);
 	case SUBSTRATE_BSDF_TYPE_EYE: 				return EYE_DIFFUSEALBEDO(BSDF);
-	case SUBSTRATE_BSDF_TYPE_SINGLELAYERWATER: return SLW_BASECOLOR(BSDF);
+	case SUBSTRATE_BSDF_TYPE_SINGLELAYERWATER:	return SLW_BASECOLOR(BSDF);
 	}
 	return 0.f;
 }
@@ -3075,10 +3075,10 @@ float3 SubstrateGetBSDFDiffuseColor(in FSubstrateBSDF BSDF)
 {
 	switch (BSDF_GETTYPE(BSDF))
 	{
-	case SUBSTRATE_BSDF_TYPE_SLAB: 			return SLAB_DIFFUSEALBEDO(BSDF);
-	case SUBSTRATE_BSDF_TYPE_HAIR: 			return HAIR_BASECOLOR(BSDF);
+	case SUBSTRATE_BSDF_TYPE_SLAB: 				return SLAB_DIFFUSEALBEDO(BSDF);
+	case SUBSTRATE_BSDF_TYPE_HAIR: 				return HAIR_BASECOLOR(BSDF);
 	case SUBSTRATE_BSDF_TYPE_EYE: 				return EYE_DIFFUSEALBEDO(BSDF);
-	case SUBSTRATE_BSDF_TYPE_SINGLELAYERWATER: return SLW_BASECOLOR(BSDF);
+	case SUBSTRATE_BSDF_TYPE_SINGLELAYERWATER:	return SLW_BASECOLOR(BSDF);
 	}
 	return 0.f;
 }
@@ -3087,8 +3087,8 @@ void SubstrateSetBSDFDiffuseColor(inout FSubstrateBSDF BSDF, float3 InColor)
 {
 	switch (BSDF_GETTYPE(BSDF))
 	{
-	case SUBSTRATE_BSDF_TYPE_SLAB: 			SLAB_DIFFUSEALBEDO(BSDF) = InColor; break;
-	case SUBSTRATE_BSDF_TYPE_HAIR: 			HAIR_BASECOLOR(BSDF) = InColor; break;
+	case SUBSTRATE_BSDF_TYPE_SLAB: 				SLAB_DIFFUSEALBEDO(BSDF) = InColor; break;
+	case SUBSTRATE_BSDF_TYPE_HAIR: 				HAIR_BASECOLOR(BSDF) = InColor; break;
 	case SUBSTRATE_BSDF_TYPE_EYE: 				EYE_DIFFUSEALBEDO(BSDF) = InColor; break;
 	case SUBSTRATE_BSDF_TYPE_SINGLELAYERWATER: SLW_BASECOLOR(BSDF) = InColor; break;
 	}
@@ -3098,10 +3098,10 @@ float3 SubstrateGetBSDFSpecularF0(in FSubstrateBSDF BSDF)
 {
 	switch (BSDF_GETTYPE(BSDF))
 	{
-	case SUBSTRATE_BSDF_TYPE_SLAB: 			return SLAB_F0(BSDF);
-	case SUBSTRATE_BSDF_TYPE_HAIR: 			return DielectricSpecularToF0(HAIR_SPECULAR(BSDF));
+	case SUBSTRATE_BSDF_TYPE_SLAB: 				return SLAB_F0(BSDF);
+	case SUBSTRATE_BSDF_TYPE_HAIR: 				return DielectricSpecularToF0(HAIR_SPECULAR(BSDF));
 	case SUBSTRATE_BSDF_TYPE_EYE: 				return EYE_F0(BSDF);
-	case SUBSTRATE_BSDF_TYPE_SINGLELAYERWATER: return DielectricSpecularToF0(SLW_SPECULAR(BSDF));
+	case SUBSTRATE_BSDF_TYPE_SINGLELAYERWATER:	return DielectricSpecularToF0(SLW_SPECULAR(BSDF));
 	}
 	return 0.f;
 }
@@ -3111,10 +3111,10 @@ float3 SubstrateGetBSDFSpecularF90(in FSubstrateBSDF BSDF)
 {
 	switch (BSDF_GETTYPE(BSDF))
 	{
-	case SUBSTRATE_BSDF_TYPE_SLAB: 			return SLAB_F90(BSDF);
-	case SUBSTRATE_BSDF_TYPE_HAIR: 			return 1.0f;
+	case SUBSTRATE_BSDF_TYPE_SLAB: 				return SLAB_F90(BSDF);
+	case SUBSTRATE_BSDF_TYPE_HAIR: 				return 1.0f;
 	case SUBSTRATE_BSDF_TYPE_EYE: 				return EYE_F90(BSDF);
-	case SUBSTRATE_BSDF_TYPE_SINGLELAYERWATER: return 1.0f;
+	case SUBSTRATE_BSDF_TYPE_SINGLELAYERWATER:	return 1.0f;
 	}
 	return 0.f;
 }
@@ -3123,10 +3123,10 @@ void SubstrateSetBSDFSpecularF0(in FSubstrateBSDF BSDF, float3 F0)
 {
 	switch (BSDF_GETTYPE(BSDF))
 	{
-	case SUBSTRATE_BSDF_TYPE_SLAB: 			SLAB_F0(BSDF) = F0; break;
-	case SUBSTRATE_BSDF_TYPE_HAIR: 			HAIR_SPECULAR(BSDF) = F0ToDielectricSpecular(max3(F0.x, F0.y, F0.z)); break;
+	case SUBSTRATE_BSDF_TYPE_SLAB: 				SLAB_F0(BSDF) = F0; break;
+	case SUBSTRATE_BSDF_TYPE_HAIR: 				HAIR_SPECULAR(BSDF) = F0ToDielectricSpecular(max3(F0.x, F0.y, F0.z)); break;
 	case SUBSTRATE_BSDF_TYPE_EYE: 				/* EYE_F0 us hardcoded */ break;
-	case SUBSTRATE_BSDF_TYPE_SINGLELAYERWATER: SLW_SPECULAR(BSDF) = F0ToDielectricSpecular(max3(F0.x, F0.y, F0.z)); break;
+	case SUBSTRATE_BSDF_TYPE_SINGLELAYERWATER:	SLW_SPECULAR(BSDF) = F0ToDielectricSpecular(max3(F0.x, F0.y, F0.z)); break;
 	}
 }
 
@@ -3134,10 +3134,10 @@ void SubstrateSetBSDFSpecularF90(in FSubstrateBSDF BSDF, float3 F90)
 {
 	switch (BSDF_GETTYPE(BSDF))
 	{
-	case SUBSTRATE_BSDF_TYPE_SLAB: 			SLAB_F90(BSDF) = F90; break;
-	case SUBSTRATE_BSDF_TYPE_HAIR: 			break;
+	case SUBSTRATE_BSDF_TYPE_SLAB: 				SLAB_F90(BSDF) = F90; break;
+	case SUBSTRATE_BSDF_TYPE_HAIR: 				break;
 	case SUBSTRATE_BSDF_TYPE_EYE: 				break;
-	case SUBSTRATE_BSDF_TYPE_SINGLELAYERWATER: break;
+	case SUBSTRATE_BSDF_TYPE_SINGLELAYERWATER:	break;
 	}
 }
 
@@ -3145,10 +3145,18 @@ float3 SubstrateGetBSDFSubSurfaceColor(in FSubstrateBSDF BSDF)
 {
 	switch (BSDF_GETTYPE(BSDF))
 	{
-	case SUBSTRATE_BSDF_TYPE_SLAB: 			return BSDF_GETSSSTYPE(BSDF) != SSS_TYPE_INVALID ? SLAB_DIFFUSEALBEDO(BSDF) : 0; // A more accurate version would be IsotropicMediumSlabPunctualDirectionalAlbedo(SubstrateSlabCreateParticipatingMedia(SLAB_DIFFUSEALBEDO(BSDF), SLAB_SSSMFP(BSDFContext.BSDF)));
-	case SUBSTRATE_BSDF_TYPE_HAIR: 			return 0;
-	case SUBSTRATE_BSDF_TYPE_EYE: 				return EYE_DIFFUSEALBEDO(BSDF);
-	case SUBSTRATE_BSDF_TYPE_SINGLELAYERWATER: return 0;
+	case SUBSTRATE_BSDF_TYPE_SLAB:
+	{
+		if (BSDF_GETSSSTYPE(BSDF) != SSS_TYPE_INVALID)
+		{
+			const float3 Extinction = 1.0f / max(SUBSTRATE_EPSILON, SLAB_SSSMFP(BSDF) * CENTIMETER_TO_METER);
+			return ExtinctionToTransmittance(Extinction, SUBSTRATE_SIMPLEVOLUME_THICKNESS_CM * CENTIMETER_TO_METER);
+		}
+		return 0;
+	}
+	case SUBSTRATE_BSDF_TYPE_HAIR: 				return 0;
+	case SUBSTRATE_BSDF_TYPE_EYE: 				return EYE_DIFFUSEALBEDO(BSDF); // Eye only has SSS via SubsurfaceProfile, so simply give the diffuse.
+	case SUBSTRATE_BSDF_TYPE_SINGLELAYERWATER:	return 0;
 	}
 	return 0.f;
 }
@@ -3157,10 +3165,10 @@ float SubstrateGetBSDFSpecular(in FSubstrateBSDF BSDF)
 {
 	switch (BSDF_GETTYPE(BSDF))
 	{
-	case SUBSTRATE_BSDF_TYPE_SLAB: 			return F0ToDielectricSpecular(SLAB_F0(BSDF).x);
-	case SUBSTRATE_BSDF_TYPE_HAIR: 			return HAIR_SPECULAR(BSDF);
+	case SUBSTRATE_BSDF_TYPE_SLAB: 				return F0ToDielectricSpecular(SLAB_F0(BSDF).x);
+	case SUBSTRATE_BSDF_TYPE_HAIR: 				return HAIR_SPECULAR(BSDF);
 	case SUBSTRATE_BSDF_TYPE_EYE: 				return F0ToDielectricSpecular(EYE_F0(BSDF));
-	case SUBSTRATE_BSDF_TYPE_SINGLELAYERWATER: return SLW_SPECULAR(BSDF);
+	case SUBSTRATE_BSDF_TYPE_SINGLELAYERWATER:	return SLW_SPECULAR(BSDF);
 	}
 	return 0.f;
 }
