Fixing BRANCH statements.

diff --git a/Engine/Shaders/Private/VirtualShadowMaps/VirtualShadowMapProjectionCommon.ush b/Engine/Shaders/Private/VirtualShadowMaps/VirtualShadowMapProjectionCommon.ush
--- a/Engine/Shaders/Private/VirtualShadowMaps/VirtualShadowMapProjectionCommon.ush
+++ b/Engine/Shaders/Private/VirtualShadowMaps/VirtualShadowMapProjectionCommon.ush
@@ -389,7 +389,7 @@ FVirtualShadowMapSampleResult SampleVirtualShadowMapDirectional(
 			Result.RayCount = 1;
 
 			float OptimalSlopeBias = 0.0f;
-			[BRANCH]
+			BRANCH
 			if (bUseOptimalBias)
 			{
 				OptimalSlopeBias = ComputeVirtualShadowMapOptimalSlopeBias(ClipmapLevelVirtualShadowMapId, SmSample, ShadowTranslatedWorldPosition, EstimatedGeoWorldNormal);
@@ -456,7 +456,7 @@ FVirtualShadowMapSampleResult SampleVirtualShadowMapLocal(
 		Result.RayCount = 1;
 
 		float OptimalSlopeBias = 0.0f;
-		[BRANCH]
+		BRANCH
 		if (bUseOptimalBias)
 		{
 			OptimalSlopeBias = ComputeVirtualShadowMapOptimalSlopeBias(VirtualShadowMapId, SmSample, ShadowTranslatedWorldPosition, EstimatedGeoWorldNormal);
