Fixes TSR shader compilation failure with r.PostProcessing.PropagateAlpha=1

diff --git a/Engine/Shaders/Private/TemporalSuperResolution/TSRDecimateHistory.usf b/Engine/Shaders/Private/TemporalSuperResolution/TSRDecimateHistory.usf
--- a/Engine/Shaders/Private/TemporalSuperResolution/TSRDecimateHistory.usf
+++ b/Engine/Shaders/Private/TemporalSuperResolution/TSRDecimateHistory.usf
@@ -525,7 +525,7 @@ void MainCS(
 					PrevGuideInfo_Extent,
 					PrevGuideInfo_ExtentInverse,
 					PrevGuideInfo_UVViewportBilinearMin,
-					PrevGuideInfo_UVViewportBilinearMax).r;
+					PrevGuideInfo_UVViewportBilinearMax);
 			#endif
 
 			#if CONFIG_MOIRE_REPROJECTION
