Substrate - fixes SSS transmission being different as compared to legacy.
Substrate  must use ApproximateHG to match legacy.

diff --git a/Engine/Shaders/Private/Substrate/SubstrateEvaluation.ush b/Engine/Shaders/Private/Substrate/SubstrateEvaluation.ush
--- a/Engine/Shaders/Private/Substrate/SubstrateEvaluation.ush
+++ b/Engine/Shaders/Private/Substrate/SubstrateEvaluation.ush
@@ -853,9 +853,9 @@ FSubstrateEvaluateResult SubstrateEvaluateBSDFCommon(FSubstrateBSDFContext BSDFC
 			{
 				float ThicknessInCm = DecodeThickness(ShadowTerms.TransmissionThickness) * SSSS_MAX_TRANSMISSION_PROFILE_DISTANCE;
 
-				float OneOverIOR = 1.0f;
-				float PhaseFunctionAnisotropy = 0.0f;
 				float3 TransmissionThroughput = 1.0f;
+				float PhaseFunction = 0.0f;
+				BRANCH
 				if (BSDF_GETSSSTYPE(BSDFContext.BSDF) == SSS_TYPE_DIFFUSION_PROFILE)
 				{
 					ThicknessInCm = BSDF_GETISTHIN(BSDFContext.BSDF) ? SLAB_SSSPROFILETHICKNESSCM(BSDFContext.BSDF) : ThicknessInCm;
@@ -863,8 +863,13 @@ FSubstrateEvaluateResult SubstrateEvaluateBSDFCommon(FSubstrateBSDFContext BSDFC
 					const uint ProfileId = SubstrateSubsurfaceProfileIdTo8bits(SLAB_SSSPROFILEID(BSDFContext.BSDF)); // TODO move this into the PackSubstrateOut( function, to avoid this decode here
 					const FTransmissionProfileParams TransmissionParams = GetTransmissionProfileParams(ProfileId);
 					TransmissionThroughput	= GetTransmissionProfile(ProfileId, ThicknessInCm).rgb;
-					PhaseFunctionAnisotropy = TransmissionParams.ScatteringDistribution;
-					OneOverIOR				= TransmissionParams.OneOverIOR;
+
+					const float PhaseFunctionAnisotropy = TransmissionParams.ScatteringDistribution;
+					const float OneOverIOR				= TransmissionParams.OneOverIOR;
+					const float3 RefracV				= refract(BSDFContext.V, -BSDFContext.N, OneOverIOR);
+					// We must use ApproximateHG to match the legacy for the moment.
+					// Later, we should look at using the same phase function as SSS_TYPE_DIFFUSION (Keep in mind that phase function should not be multiple with multiple scattering, only single scattering)
+					PhaseFunction						= ApproximateHG(dot(-BSDFContext.L, RefracV), TransmissionParams.ScatteringDistribution);
 				}
 				else
 				{
@@ -875,13 +880,12 @@ FSubstrateEvaluateResult SubstrateEvaluateBSDFCommon(FSubstrateBSDFContext BSDFC
 					const float3 SubsurfaceAlebdo = SLAB_DIFFUSEALBEDO(BSDFContext.BSDF);
 					TransmissionThroughput = GetBurleyTransmission(SubsurfaceAlebdo, max(SUBSTRATE_EPSILON, MeanFreePathInCm), ThicknessInCm).xyz;
 
-					PhaseFunctionAnisotropy = SLAB_SSSPHASEANISOTROPY(BSDFContext.BSDF);
-					OneOverIOR				= 1.f / DielectricF0ToIor(F0.y);
+					const float PhaseFunctionAnisotropy = SLAB_SSSPHASEANISOTROPY(BSDFContext.BSDF);
+					const float OneOverIOR				= 1.f / DielectricF0ToIor(F0.y);
+					const float3 RefracV				= refract(BSDFContext.V, -BSDFContext.N, OneOverIOR);
+					PhaseFunction						= HenyeyGreensteinPhase(PhaseFunctionAnisotropy, dot(BSDFContext.L, RefracV));
 				}
 
-				const float3 RefracV = refract(BSDFContext.V, -BSDFContext.N, OneOverIOR);
-				const float PhaseFunction = HenyeyGreensteinPhase(PhaseFunctionAnisotropy, dot(BSDFContext.L, RefracV));
-
 				Sample.ThroughputV				= OpaqueBSDFThroughput; // SSS is not translucent as of today
 				Sample.TransmittanceAlongN		= OpaqueBSDFThroughput; // idem
 				Sample.TransmissionPathValue	= TransmissionThroughput * PhaseFunction;
