Substrate - Fixed env lighting for legacy subsurface (wrap model with isthin=false).
Substrate must add the subsurface color to the diffuse albedo in order to match legacy (adding those two contribution in a non energy conservative way).

diff --git a/Engine/Shaders/Private/Substrate/SubstrateEvaluation.ush b/Engine/Shaders/Private/Substrate/SubstrateEvaluation.ush
--- a/Engine/Shaders/Private/Substrate/SubstrateEvaluation.ush
+++ b/Engine/Shaders/Private/Substrate/SubstrateEvaluation.ush
@@ -1666,22 +1666,37 @@ FSubstrateEnvLightResult SubstrateEvaluateForEnvLight(FSubstrateBSDFContext BSDF
 			BRANCH
 			if (BSDF_GETSSSTYPE(BSDFContext.BSDF) == SSS_TYPE_WRAP)
 			{
-				FParticipatingMedia PM = SubstrateSlabCreateParticipatingMedia(DiffuseColor, SLAB_SSSMFP(BSDFContext.BSDF));
-
+				const bool bIsThin = BSDF_GETISTHIN(BSDFContext.BSDF);
+				const float3 SlabDiffuseColor = bIsThin ? DiffuseColor : float3(1, 1, 1);
+				FParticipatingMedia PM = SubstrateSlabCreateParticipatingMedia(SlabDiffuseColor, SLAB_SSSMFP(BSDFContext.BSDF));
+			
 				// We consider single scattering and perpendicular to the surface light transmittance.
 				// Using BackFaceNoL instead of 1 for TransmittanceNoL would be more correct but it gives some high frequency details looking weird without multiple scattering.
 				const float TransmittanceNoL = 1.0f;
 				const float3 SlabTransmittance = IsotropicMediumSlabTransmittance(PM, SUBSTRATE_SIMPLEVOLUME_THICKNESS_M, TransmittanceNoL);
+			
+				if (bIsThin)
+				{
+					// We recover the transmitted luminance from as the back face Lambert affected by the transmittance.
+					// The MFP is specified on the slab node and then transmittance recoved from it.
+					// We assume the back face diffuse albedo is white=1 and no diffusion happens at this stage.
+					const float3 BackFaceWhiteDiffuseColor = float3(1.0f, 1.0f, 1.0f);
+			
+					// The diffuse is attenuated by the slab transmittance and the GGX interface SpecularTransmission.
+					SubstrateEnvLightResult.DiffuseBackFaceWeight = Diffuse_Lambert(BackFaceWhiteDiffuseColor) * SlabTransmittance * DirectionalAlbedo_SpecularTransmission;
+				}
+				else
+				{
+					// The wrap lighting is a non-physically based shading which intends to match legacy behavior.
+					// We recover SubsurfaceColor that has been encoded as the Slab transmittance color.
+					const float3 DirectionalSubSurfaceColor = SlabTransmittance * DirectionalAlbedo_SpecularTransmission;
 
-				// We recover the transmitted luminance from as the back face Lambert affected by the transmittance.
-				// The MFP is specified on the slab node and then transmittance recoved from it.
-				// We assume the back face diffuse albedo is white=1 and no diffusion happens at this stage.
-				const float3 BackFaceWhiteDiffuseColor = float3(1.0f, 1.0f, 1.0f);
-
-				// The diffuse is attenuated by the slab transmittance and the GGX interface SpecularTransmission.
-				SubstrateEnvLightResult.DiffuseBackFaceWeight = Diffuse_Lambert(BackFaceWhiteDiffuseColor) * SlabTransmittance * DirectionalAlbedo_SpecularTransmission;
-
-				// We do not want thin (a.k.a. two sided lighting) contribution to go the SSS post process path.
+					// As of today, for the diffuse to match the legacy behavior, we have to add the susurface color contribution.
+					SubstrateEnvLightResult.DiffuseColor  += DirectionalSubSurfaceColor;
+					SubstrateEnvLightResult.DiffuseWeight += DirectionalSubSurfaceColor;
+				}
+			
+				// We do not want thin (a.k.a. two sided lighting) or the non thin path (a.k.a. subsurface) contribution to go the SSS post process path.
 				SubstrateEnvLightResult.bSubsurface = false;
 			}
 			#endif
