Fixes CorrectGuideColorExposure() on TSR history resurrection

diff --git a/Engine/Shaders/Private/TemporalSuperResolution/TSRDecimateHistory.usf b/Engine/Shaders/Private/TemporalSuperResolution/TSRDecimateHistory.usf
--- a/Engine/Shaders/Private/TemporalSuperResolution/TSRDecimateHistory.usf
+++ b/Engine/Shaders/Private/TemporalSuperResolution/TSRDecimateHistory.usf
@@ -366,7 +366,7 @@ struct FCatmullRomFetches
 // Correct pre exposure of the guide color.
 tsr_halfC CorrectGuideColorExposure(tsr_halfC ReprojectedHistoryGuideColor, const float PreExposureCorrection)
 {
-	tsr_halfC LDRToLinear = min(rcp(saturate(tsr_half(1.0) - ReprojectedHistoryGuideColor)) * tsr_half(HistoryPreExposureCorrection), tsr_half(MaxHalfFloat));
+	tsr_halfC LDRToLinear = min(rcp(saturate(tsr_half(1.0) - ReprojectedHistoryGuideColor)) * tsr_half(PreExposureCorrection), tsr_half(MaxHalfFloat));
 
 	tsr_halfC LinearReprojectedHistoryGuideColor = ReprojectedHistoryGuideColor * LDRToLinear;
 				
@@ -508,16 +508,16 @@ void MainCS(
 				
 			float2 KernelHistoryBufferUV = ApplyScreenTransform(PrevOutputScreenPos, ScreenPosToPrevHistoryGuideBufferUV);
 
-			ReprojectedHistoryGuideSamples.FetchSamples(
-				PrevHistoryGuide,
-				KernelHistoryBufferUV,
-				PrevFrameIndex * 2.0 + 0.0,
-				PrevGuideInfo_Extent,
-				PrevGuideInfo_ExtentInverse,
-				PrevGuideInfo_UVViewportBilinearMin,
-				PrevGuideInfo_UVViewportBilinearMax);
-					
 			#if CONFIG_SCENE_COLOR_ALPHA
+				ReprojectedHistoryGuideSamples.FetchSamples(
+					PrevHistoryGuide,
+					KernelHistoryBufferUV,
+					PrevFrameIndex * 2.0 + 0.0,
+					PrevGuideInfo_Extent,
+					PrevGuideInfo_ExtentInverse,
+					PrevGuideInfo_UVViewportBilinearMin,
+					PrevGuideInfo_UVViewportBilinearMax);
+					
 				ReprojectedHistoryGuideUncertaintySamples.FetchSamples(
 					PrevHistoryGuide,
 					KernelHistoryBufferUV,
@@ -526,6 +526,15 @@ void MainCS(
 					PrevGuideInfo_ExtentInverse,
 					PrevGuideInfo_UVViewportBilinearMin,
 					PrevGuideInfo_UVViewportBilinearMax);
+			#else
+				ReprojectedHistoryGuideSamples.FetchSamples(
+					PrevHistoryGuide,
+					KernelHistoryBufferUV,
+					PrevFrameIndex,
+					PrevGuideInfo_Extent,
+					PrevGuideInfo_ExtentInverse,
+					PrevGuideInfo_UVViewportBilinearMin,
+					PrevGuideInfo_UVViewportBilinearMax);
 			#endif
 
 			#if CONFIG_MOIRE_REPROJECTION
