LFV - improve look by allowing a mode when the throughput of both radial and height fogs are combine in order to get softer edges.
This also work when LFV are voxelised in volumetric fog using some math trick.

diff --git a/Engine/Shaders/Private/LocalFogVolumes/LocalFogVolumeCommon.ush b/Engine/Shaders/Private/LocalFogVolumes/LocalFogVolumeCommon.ush
--- a/Engine/Shaders/Private/LocalFogVolumes/LocalFogVolumeCommon.ush
+++ b/Engine/Shaders/Private/LocalFogVolumes/LocalFogVolumeCommon.ush
@@ -142,7 +142,8 @@ FFogData LocalFogVolumeEvaluateAnalyticalIntegral(FLocalFogVolumeGPUInstanceData
 {
 	FFogData FogData;
 
-	float OpticalDepth = 0.0f;
+	float RadialOpticalDepth = 0.0f;
+	float HeightOpticalDepth = 0.0f;
 
 	if (FogInstance.RadialFogExtinction > 0.0f)
 	{
@@ -172,7 +173,7 @@ FFogData LocalFogVolumeEvaluateAnalyticalIntegral(FLocalFogVolumeGPUInstanceData
 			const float Length1Sqr = Length1 * Length1;
 			float Integral0 = -(c * Length0 + b * Length0Sqr + Length0Sqr * Length0 / 3.0f);
 			float Integral1 = -(c * Length1 + b * Length1Sqr + Length1Sqr * Length1 / 3.0f);
-			OpticalDepth += max(0.0, FogInstance.RadialFogExtinction * (Integral1 - Integral0) * (3.0f / 4.0f));
+			RadialOpticalDepth = max(0.0, FogInstance.RadialFogExtinction * (Integral1 - Integral0) * (3.0f / 4.0f));
 		}
 	}
 
@@ -186,7 +187,7 @@ FFogData LocalFogVolumeEvaluateAnalyticalIntegral(FLocalFogVolumeGPUInstanceData
 
 	#if 0
 		// Original integral, less artefact
-		OpticalDepth += (FogInstance.HeightFogExtinction / FogInstance.HeightFogFalloff) * exp(-StartHeight * FogInstance.HeightFogFalloff) * (1.0 - exp(-RayDirU.z * RayLengthU * FogInstance.HeightFogFalloff)) / RayDirU.z;
+		HeightOpticalDepth = (FogInstance.HeightFogExtinction / FogInstance.HeightFogFalloff) * exp(-StartHeight * FogInstance.HeightFogFalloff) * (1.0 - exp(-RayDirU.z * RayLengthU * FogInstance.HeightFogFalloff)) / RayDirU.z;
 	#elif 1
 		// Reworked to avoid large value as input to exp
 		//OpticalDepth = (FogInstance.HeightFogExtinction / FogInstance.HeightFogFalloff) * (exp(-StartHeight * FogInstance.HeightFogFalloff) - exp(-StartHeight * FogInstance.HeightFogFalloff - RayDir.z * RayLength * FogInstance.HeightFogFalloff)) / RayDir.z;
@@ -196,12 +197,51 @@ FFogData LocalFogVolumeEvaluateAnalyticalIntegral(FLocalFogVolumeGPUInstanceData
 		float Factor1 = RayDirU.z * RayLengthU * FogInstance.HeightFogFalloff;
 		Factor0 = max(-80.0f, Factor0); // Simply clamp Factor0 to a valid range in order to not have exp explode.
 
-		OpticalDepth += (FogInstance.HeightFogExtinction / (FogInstance.HeightFogFalloff * RayDirU.z)) * (exp(-Factor0) - exp(-(Factor0 + Factor1)));
+		HeightOpticalDepth += (FogInstance.HeightFogExtinction / (FogInstance.HeightFogFalloff * RayDirU.z)) * (exp(-Factor0) - exp(-(Factor0 + Factor1)));
 	#endif
+
+		/*
+		Proof of the height fog integral equations using Mathematica
+		This show that IntegratedLuminanceFactor is really equivalent to coverage in this simple case.
+
+		(*Density of matter according to height Y*)
+		Density [y_] := A * Exp[-B * y]
+
+		(*Ray position*)
+		Ray[t_] := Oz + t * Rz
+
+		(*Integrate Optical Depth for a distance T*)
+		Fog[T_] := Integrate[Density[Ray[t]], {t, 0, T}]
+
+		(*Optical Depth integration equation*)
+		Fog[T]
+
+		(*Transmittance from Optical Depth*)
+		Exp[-Integrate[Density[Ray[t]], {t, 0, T}]]
+
+		(*Transmittance from a position to the origin of tracing*)
+		Trans[T_] := Exp[-Integrate[Density[Ray[t]], {t, 0, T}]]
+
+		(*Integrate emissive and scattering from height fog, assuming density \
+		is extinction and extinction=scattering, thus albedo=1*)
+		Integrate[Density[Ray[t]]*EmSc *Trans[t], {t, 0, T}]
+
+		=> this leads to EmSc - Transmittance * EmSc
+						 EmSc * (1 - Transmittance)
+						 EmSc * Coverage
+		*/
 	}
 
+	// We now combine the two optical depth for the "opacity"=1-transmittance to be combined together.
+	// This helps to have soft edges when using height fog. We also account for that when voxelising in the volumetric fog.
+	float TR = exp(-RadialOpticalDepth);
+	float TH = exp(-HeightOpticalDepth);
+	float T = 1 - (1-TR) * (1-TH);
+	float OpticalDepth = -log(T);
+
 	// Evaluate transmittance, accounting for the volume scale converted to local unit space meter.
 	float Transmittance = exp(-OpticalDepth * FogInstance.UniformScale * CENTIMETER_TO_METER);
+
 	// Assuming extinction is grey scale, we can compute a single coverage value from transmittance
 	FogData.Coverage = (1.0 - Transmittance);
 	// Assuming extinction==scattering, i.e. albedo=1, we can compute the integrale for each point on a ray while accouting for transmittance back to the view.
diff --git a/Engine/Shaders/Private/VolumetricFog.usf b/Engine/Shaders/Private/VolumetricFog.usf
--- a/Engine/Shaders/Private/VolumetricFog.usf
+++ b/Engine/Shaders/Private/VolumetricFog.usf
@@ -187,15 +187,41 @@ void MaterialSetupCS(
 		float DistToCenter = length(UnitSpacePos - float3(0.0f, 0.0f, 0.0f));
 		float SphereFade = saturate(1.0f - DistToCenter * DistToCenter);
 
-		float LFVExtinction = 0.0;
-		if (FogInstance.HeightFogExtinction > 0.0f)
-		{
-			LFVExtinction += SphereFade <= 0.0f ? 0.0f : FogInstance.HeightFogExtinction * exp(-FogInstance.HeightFogFalloff * (UnitSpacePos.z - FogInstance.HeightFogOffset));
-		}
-		if(FogInstance.RadialFogExtinction > 0.0f)
-		{
-			LFVExtinction += FogInstance.RadialFogExtinction * SphereFade;
-		}
+		// Evaluate both radial and height local fog contribution
+		float LFVExtinctionRadial = SphereFade <= 0.0f ? 0.0f : FogInstance.HeightFogExtinction * exp(-FogInstance.HeightFogFalloff * (UnitSpacePos.z - FogInstance.HeightFogOffset));
+		float LFVExtinctionHeight = FogInstance.RadialFogExtinction * SphereFade;
+
+		// As for the local fog volume analitical evaluation, we want to combine both two optical depth for the "opacity"=1-transmittance to be combined together.
+		// This helps to have soft edges when using height fog. 
+		/*
+		In this case, we have to find math so that the combined extinction resulting in a transmittance representing the multiplication of coverage = 1 - transmittance factors.
+
+		Reminder
+			Transmittance as T
+			Optical Depth as Tau
+			T   =  exp(-Tau)
+			Tau = -log(T)		(*)
+
+		We want to combine the both mode transmittance Tr (radial) and Th (height) as done in LocalFogVolumeEvaluateAnalyticalIntegral to have
+			 T		= 1 - (1-Tr)(1-Th)											// Where (1-Tr) is virtual coverage of the medium expressed from transmittance for the radial fog.
+			 T		= 1 - (1 - exp(-Tau_r)) (1 - exp(-Tau_h)))					// Where Tr = exp(-Tau_r)
+			-Tau	= log(1 - (1 - exp(-Tau_r)) (1 - exp(-Tau_h)))				// Using (*)
+			 Tau	=-log(exp(-Tau_r) - exp((-Tau_r-Tau_h) + exp(-Tau_h)))
+
+		Optical Depth Tau = Sigma_t * dt so it is linear with respect to dt and extinction Sigma_t. So we can also write
+			Sigma_T = -log(exp(-Sigma_T_r) - exp((-Sigma_T_r-Sigma_T_h) + exp(-Sigma_T_h)))
+
+		Factoring some values a bit we can get
+			A		= exp(-Sigma_T_r)
+			B		= exp(-Sigma_T_h)
+			Sigma_T	= -log(A - A*B + B)
+		*/
+
+		// Optical Depth is linear with Extinction so there is no need to account for the step length.
+		float FactorA = exp(-LFVExtinctionHeight);
+		float FactorB = exp(-LFVExtinctionRadial);
+		float LFVExtinction = -log(FactorA - FactorA* FactorB + FactorB);
+		LFVExtinction = max(0.0f, LFVExtinction);
 
 		// Convert extinction from unit volume 1/m into extinction for volumetric fog as 1/cm
 		LFVExtinction *= (1.0f / METER_TO_CENTIMETER);
