Reorders some convolutions in TSR RejectShading to saves more VGPRs

diff --git a/Engine/Shaders/Private/TemporalSuperResolution/TSRRejectShading.usf b/Engine/Shaders/Private/TemporalSuperResolution/TSRRejectShading.usf
--- a/Engine/Shaders/Private/TemporalSuperResolution/TSRRejectShading.usf
+++ b/Engine/Shaders/Private/TemporalSuperResolution/TSRRejectShading.usf
@@ -40,7 +40,11 @@
 	#define GROUP_TILE_SIZE 32
 
 	#define LANE_STRIDE_X 1
-	#define LANE_STRIDE_Y 4
+	#if CONFIG_COMPILE_FP16
+		#define LANE_STRIDE_Y 4
+	#else
+		#define LANE_STRIDE_Y 2
+	#endif
 
 	#define WAVE_COUNT_X 1
 	#define WAVE_COUNT_Y (GROUP_TILE_SIZE / LANE_STRIDE_Y)
@@ -120,6 +124,8 @@
 //------------------------------------------------------- TYPEDEFS
 
 #define tsr_simd_half  TLaneVector2D<tsr_half, 1, LANE_STRIDE_X, LANE_STRIDE_Y>
+#define tsr_simd_half2 TLaneVector2D<tsr_half, 2, LANE_STRIDE_X, LANE_STRIDE_Y>
+#define tsr_simd_half3 TLaneVector2D<tsr_half, 3, LANE_STRIDE_X, LANE_STRIDE_Y>
 #define tsr_simd_half4 TLaneVector2D<tsr_half, 4, LANE_STRIDE_X, LANE_STRIDE_Y>
 #define tsr_simd_halfC TLaneVector2D<tsr_half, CONFIG_CHANNEL_COUNT, LANE_STRIDE_X, LANE_STRIDE_Y>
 #define tsr_simd_ushort  TLaneVector2D<tsr_ushort, 1, LANE_STRIDE_X, LANE_STRIDE_Y>
@@ -429,20 +435,27 @@ void MeasureRejection(
 {
 	tsr_simd_halfC BackbufferQuantizationErrorVector = tsr_simd_halfC::Const(MeasureBackbufferLDRQuantizationError());
 		
-	// Filter clamped input and history to increase quantization of the rejection detection
-	tsr_simd_halfC FilteredInput   = Blur3x3(InputC2);
 	tsr_simd_halfC FilteredHistory = Blur3x3(HistoryC2);
 	
+	tsr_simd_halfC TotalVarInputDiffC0C2 = abs(TotalVariation3x3(abs(InputC0 - InputC2)));
+	tsr_simd_halfC TotalVarInputC2 = abs(TotalVariation3x3(InputC2));
+
+	tsr_simd_halfC ClampError = BackbufferQuantizationErrorVector;
+	ClampError = max(ClampError, Blur3x3(min(TotalVarInputDiffC0C2, TotalVarInputC2)));
+	
+	tsr_simd_halfC InputC2BoxSize = MaxMinusMin3x3(InputC2);
+	ClampError = max(ClampError, InputC2BoxSize * tsr_half(FilteringWeight * 0.25));
+	ClampError = ClampError + BackbufferQuantizationErrorVector;
+
+	tsr_simd_halfC FilteredInput   = Blur3x3(InputC2);
+
 	// Build the clamping box of the filtered signal
 	tsr_simd_halfC FilteredBoxMin;
 	tsr_simd_halfC FilteredBoxMax;
 	MinMax3x3(FilteredInput, /* out */ FilteredBoxMin, /* out */ FilteredBoxMax);
-	FilteredBoxMin = FilteredBoxMin - BackbufferQuantizationErrorVector;
-	FilteredBoxMax = FilteredBoxMax + BackbufferQuantizationErrorVector;
+	FilteredBoxMin = FilteredBoxMin - ClampError;
+	FilteredBoxMax = FilteredBoxMax + ClampError;
 		
-	tsr_simd_halfC TotalVarInputDiffC0C2 = abs(TotalVariation3x3(abs(InputC0 - InputC2)));
-	tsr_simd_halfC TotalVarInputC2 = abs(TotalVariation3x3(InputC2));
-
 	// Clamp filtered signal
 	tsr_simd_halfC ClampedFilteredHistory = fastClamp(FilteredHistory, FilteredBoxMin, FilteredBoxMax);
 	if (ApplyMoireError)
@@ -475,7 +488,7 @@ void MeasureRejection(
 	tsr_simd_halfC BoxSize;
 	{
 		BoxSize = (
-			MaxMinusMin3x3(InputC2) * tsr_simd_halfC::Const(FilteringWeight) +
+			InputC2BoxSize * tsr_simd_halfC::Const(FilteringWeight) +
 			(BackbufferQuantizationErrorVector * tsr_half(2 * FilteringWeight)));
 
 		if (ApplyMoireError)
@@ -496,12 +509,9 @@ void MeasureRejection(
 
 	tsr_simd_halfC Delta = max(abs(FilteredInput - FilteredHistory), BoxSize);
 
-	tsr_simd_halfC ClampError = BackbufferQuantizationErrorVector;
-	ClampError = max(ClampError, BoxSize * tsr_half(0.25));
-	ClampError = max(ClampError, Blur3x3(min(TotalVarInputDiffC0C2, TotalVarInputC2)));
-
 	tsr_simd_halfC RawClampedEnergy      = abs(ClampedFilteredHistory - FilteredHistory);
-	tsr_simd_halfC RawFactor      = saturate(tsr_simd_halfC::Const(1.0) - max(RawClampedEnergy - ClampError, tsr_simd_halfC::Const(0.0)) * rcp(Delta));
+
+	tsr_simd_halfC RawFactor        = saturate(tsr_simd_halfC::Const(1.0) - RawClampedEnergy * rcp(Delta));
 	tsr_simd_half RawRejection      = min3(RawFactor[0], RawFactor[1], RawFactor[2]);
 
 	#if CONFIG_SCENE_COLOR_ALPHA
@@ -512,9 +522,9 @@ void MeasureRejection(
 		Deconcatenate(MaxRGBMinA3x3(Median3x3(Concatenate(RawClampedEnergy, RawRejection))), /* out */ FilteredClampedEnergy, /* out */ OutRejectionClampBlend);
 	#endif
 
-	tsr_simd_halfC FilteredFactor = saturate(tsr_simd_halfC::Const(1.0) - max(FilteredClampedEnergy - ClampError, tsr_simd_halfC::Const(0.0)) * rcp(Delta));
+	tsr_simd_halfC FilteredFactor   = saturate(tsr_simd_halfC::Const(1.0) - FilteredClampedEnergy * rcp(Delta));
 	tsr_simd_half FilteredRejection = min3(FilteredFactor[0], FilteredFactor[1], FilteredFactor[2]);
-	OutRejectionBlendFinal = saturate(FilteredRejection);
+	OutRejectionBlendFinal          = FilteredRejection;
 } // MeasureRejection()
 
 
