Add support for DBuffer decals to the mobile forward renderer.

The DBuffer layout and most code is the same as for the desktop renderer.
Added new CVar r.Mobile.DBuffer so that is disabled by default. DBuffer for the SM5 renderer on mobile is enabled by default.

diff --git a/Engine/Config/BaseEngine.ini b/Engine/Config/BaseEngine.ini
--- a/Engine/Config/BaseEngine.ini
+++ b/Engine/Config/BaseEngine.ini
@@ -3014,6 +3014,7 @@ DefaultASTCQualityBySize=3
 +VersionedIntRValues=r.VelocityOutputPass
 +VersionedIntRValues=r.SelectiveBasePassOutputs
 +VersionedIntRValues=r.DBuffer
++VersionedIntRValues=r.Mobile.DBuffer
 +VersionedIntRValues=r.Shaders.Symbols
 +VersionedIntRValues=r.Shaders.GenerateSymbols
 +VersionedIntRValues=r.Shaders.WriteSymbols
diff --git a/Engine/Shaders/Private/DBufferDecalShared.ush b/Engine/Shaders/Private/DBufferDecalShared.ush
--- a/Engine/Shaders/Private/DBufferDecalShared.ush
+++ b/Engine/Shaders/Private/DBufferDecalShared.ush
@@ -11,19 +11,26 @@
 #endif
 
 #if DBUFFER_BASEPASS_RESOURCE
-#define OpaqueBasePassDBufferATexture OpaqueBasePass.DBufferATexture
-#define OpaqueBasePassDBufferBTexture OpaqueBasePass.DBufferBTexture
-#define OpaqueBasePassDBufferCTexture OpaqueBasePass.DBufferCTexture
-#define OpaqueBasePassDBufferRenderMask OpaqueBasePass.DBufferRenderMask
+
+#if SHADING_PATH_MOBILE
+#define DBufferBasePass MobileBasePass
+#else
+#define DBufferBasePass OpaqueBasePass
+#endif
+
+#define OpaqueBasePassDBufferATexture DBufferBasePass.DBufferATexture
+#define OpaqueBasePassDBufferBTexture DBufferBasePass.DBufferBTexture
+#define OpaqueBasePassDBufferCTexture DBufferBasePass.DBufferCTexture
+#define OpaqueBasePassDBufferRenderMask DBufferBasePass.DBufferRenderMask
 
 #if SUPPORTS_INDEPENDENT_SAMPLERS
-#define OpaqueBasePassDBufferATextureSampler OpaqueBasePass.DBufferATextureSampler
-#define OpaqueBasePassDBufferBTextureSampler OpaqueBasePass.DBufferATextureSampler
-#define OpaqueBasePassDBufferCTextureSampler OpaqueBasePass.DBufferATextureSampler
+#define OpaqueBasePassDBufferATextureSampler DBufferBasePass.DBufferATextureSampler
+#define OpaqueBasePassDBufferBTextureSampler DBufferBasePass.DBufferATextureSampler
+#define OpaqueBasePassDBufferCTextureSampler DBufferBasePass.DBufferATextureSampler
 #else
-#define OpaqueBasePassDBufferATextureSampler OpaqueBasePass.DBufferATextureSampler
-#define OpaqueBasePassDBufferBTextureSampler OpaqueBasePass.DBufferBTextureSampler
-#define OpaqueBasePassDBufferCTextureSampler OpaqueBasePass.DBufferCTextureSampler
+#define OpaqueBasePassDBufferATextureSampler DBufferBasePass.DBufferATextureSampler
+#define OpaqueBasePassDBufferBTextureSampler DBufferBasePass.DBufferBTextureSampler
+#define OpaqueBasePassDBufferCTextureSampler DBufferBasePass.DBufferCTextureSampler
 #endif
 
 #endif // DBUFFER_BASEPASS_RESOURCE
@@ -413,19 +420,19 @@ FDBufferData GetDBufferData(float2 BufferUV, uint RTMaskBit)
 	BRANCH
 	if (RTMaskBit & 0x1)
 	{
-		DBufferA = Texture2DSampleLevel(OpaqueBasePass.DBufferATexture, OpaqueBasePassDBufferATextureSampler, BufferUV, 0);
+		DBufferA = Texture2DSampleLevel(OpaqueBasePassDBufferATexture, OpaqueBasePassDBufferATextureSampler, BufferUV, 0);
 	}
 
 	BRANCH
 	if (RTMaskBit & 0x2)
 	{
-		DBufferB = Texture2DSampleLevel(OpaqueBasePass.DBufferBTexture, OpaqueBasePassDBufferBTextureSampler, BufferUV, 0);
+		DBufferB = Texture2DSampleLevel(OpaqueBasePassDBufferBTexture, OpaqueBasePassDBufferBTextureSampler, BufferUV, 0);
 	}
 
 	BRANCH
 	if (RTMaskBit & 0x4)
 	{
-		DBufferC = Texture2DSampleLevel(OpaqueBasePass.DBufferCTexture, OpaqueBasePassDBufferCTextureSampler, BufferUV, 0);
+		DBufferC = Texture2DSampleLevel(OpaqueBasePassDBufferCTexture, OpaqueBasePassDBufferCTextureSampler, BufferUV, 0);
 	}
 #elif USE_DBUFFER && RAYHITGROUPSHADER
 	BRANCH
diff --git a/Engine/Shaders/Private/DecalCommon.ush b/Engine/Shaders/Private/DecalCommon.ush
--- a/Engine/Shaders/Private/DecalCommon.ush
+++ b/Engine/Shaders/Private/DecalCommon.ush
@@ -57,11 +57,14 @@ float DecalCommonOutput(
 	// If we output to DBuffer, write to all the 3 RTs.
 	SubstratePackDBuffer(DBuffer, Out.MRT[0], Out.MRT[1], Out.MRT[2]);
 
+	#if DECAL_RENDERTARGET_COUNT == 4
+		Out.MRT[3] = 1.0f; // DBufferMask
+	#endif
+
 	// This is to set the coverage over the background to an arbitrary value in case premultiplied alpha is used.
 	Out.MRT[0].a = CoverageOverBackground;
 	Out.MRT[1].a = CoverageOverBackground;
 	Out.MRT[2].a = CoverageOverBackground;
-	Out.MRT[3].a = CoverageOverBackground;
 #endif
 
 	// Allow shader compiler to cull outputs that we won't use.
diff --git a/Engine/Shaders/Private/DeferredDecal.usf b/Engine/Shaders/Private/DeferredDecal.usf
--- a/Engine/Shaders/Private/DeferredDecal.usf
+++ b/Engine/Shaders/Private/DeferredDecal.usf
@@ -7,9 +7,14 @@
 #include "Common.ush"
 
 #if SHADING_PATH_MOBILE
-// Reroute MobileSceneTextures uniform buffer references to the base pass uniform buffer 
-#define MobileSceneTextures MobileBasePass.SceneTextures
-#define EyeAdaptationStruct MobileBasePass
+#if DECAL_RENDERTARGETMODE == DECAL_RENDERTARGETMODE_DBUFFER
+	#define MobileSceneTextures DecalPass.MobileSceneTextures
+	#define EyeAdaptationStruct DecalPass
+#else
+	// Reroute MobileSceneTextures uniform buffer references to the base pass uniform buffer
+	#define MobileSceneTextures MobileBasePass.SceneTextures
+	#define EyeAdaptationStruct MobileBasePass
+#endif
 #else
 #define SceneTexturesStruct DecalPass.SceneTextures
 #define EyeAdaptationStruct DecalPass
diff --git a/Engine/Shaders/Private/MobileBasePassPixelShader.usf b/Engine/Shaders/Private/MobileBasePassPixelShader.usf
--- a/Engine/Shaders/Private/MobileBasePassPixelShader.usf
+++ b/Engine/Shaders/Private/MobileBasePassPixelShader.usf
@@ -437,6 +437,34 @@ void Main(
 	float SubsurfaceProfile = 0;
 #endif
 
+	float DBufferOpacity = 1.0f;
+#if USE_DBUFFER && MATERIALDECALRESPONSEMASK && !MATERIALBLENDING_ANY_TRANSLUCENT && !MATERIAL_SHADINGMODEL_SINGLELAYERWATER
+	// apply decals from the DBuffer
+	BRANCH
+	if ((GetPrimitiveData(MaterialParameters).Flags & PRIMITIVE_SCENE_DATA_FLAG_DECAL_RECEIVER) != 0 && View.ShowDecalsMask > 0)
+	{
+		uint ValidDBufferTargetMask = GetDBufferTargetMask(uint2(SvPosition.xy)) & MATERIALDECALRESPONSEMASK;
+
+		BRANCH
+		if (ValidDBufferTargetMask)
+		{
+			float2 BufferUV = SvPositionToBufferUV(SvPosition);
+		#if SUBSTRATE_ENABLED 
+			#if SUBSTRATE_INLINE_SHADING && !SUBSTRATE_OPTIMIZED_UNLIT
+			const FSubstrateDBuffer SubstrateBufferData = GetDBufferData(BufferUV, ValidDBufferTargetMask);
+			ApplyDBufferData(SubstrateBufferData, SubstratePixelHeader, SubstrateData);
+			DBufferOpacity = SubstrateBufferData.Coverage;
+			#endif
+		#else
+			FDBufferData DBufferData = GetDBufferData(BufferUV, ValidDBufferTargetMask);
+			ApplyDBufferData(DBufferData, MaterialParameters.WorldNormal, SubsurfaceColor, Roughness, BaseColor, Metallic, Specular);
+			DBufferOpacity = (DBufferData.ColorOpacity + DBufferData.NormalOpacity + DBufferData.RoughnessOpacity) * (1.0f / 3.0f);
+		#endif // SUBSTRATE_ENABLED
+		}
+	}
+#endif
+
+
 #if !SUBSTRATE_ENABLED || SUBSTRATE_INLINE_SINGLELAYERWATER
 	const half BaseMaterialCoverageOverWater = Opacity;
 	const half WaterVisibility = 1.0 - BaseMaterialCoverageOverWater;
@@ -1016,7 +1044,7 @@ void Main(
 	FinalizeVirtualTextureFeedback(
 		MaterialParameters.VirtualTextureFeedback,
 		MaterialParameters.SvPosition,
-		Opacity,
+		Opacity * DBufferOpacity,
 		View.FrameNumber,
 		View.VTFeedbackBuffer
 	);
diff --git a/Engine/Source/Developer/TargetPlatform/Private/TargetPlatformBase.cpp b/Engine/Source/Developer/TargetPlatform/Private/TargetPlatformBase.cpp
--- a/Engine/Source/Developer/TargetPlatform/Private/TargetPlatformBase.cpp
+++ b/Engine/Source/Developer/TargetPlatform/Private/TargetPlatformBase.cpp
@@ -89,6 +89,12 @@ bool FTargetPlatformBase::UsesMobileAmbientOcclusion() const
 	return CVar ? (CVar->GetInt() != 0) : false;
 }
 
+bool FTargetPlatformBase::UsesMobileDBuffer() const
+{
+	static IConsoleVariable* CVar = IConsoleManager::Get().FindConsoleVariable(TEXT("r.Mobile.DBuffer"));
+	return CVar ? (CVar->GetInt() != 0) : false;
+}
+
 int32 GASTCHDRProfile = 0;
 static FAutoConsoleVariableRef CVarAllowASTCHDRProfile(
 	TEXT("cook.AllowASTCHDRProfile"),
diff --git a/Engine/Source/Developer/TargetPlatform/Public/Common/TargetPlatformBase.h b/Engine/Source/Developer/TargetPlatform/Public/Common/TargetPlatformBase.h
--- a/Engine/Source/Developer/TargetPlatform/Public/Common/TargetPlatformBase.h
+++ b/Engine/Source/Developer/TargetPlatform/Public/Common/TargetPlatformBase.h
@@ -102,6 +102,8 @@ public:
 
 	TARGETPLATFORM_API virtual bool UsesMobileAmbientOcclusion() const override;
 
+	TARGETPLATFORM_API virtual bool UsesMobileDBuffer() const override;
+
 	TARGETPLATFORM_API virtual bool UsesASTCHDR() const override;
 
 	TARGETPLATFORM_API virtual void GetRayTracingShaderFormats(TArray<FName>& OutFormats) const override;
diff --git a/Engine/Source/Developer/TargetPlatform/Public/Interfaces/ITargetPlatform.h b/Engine/Source/Developer/TargetPlatform/Public/Interfaces/ITargetPlatform.h
--- a/Engine/Source/Developer/TargetPlatform/Public/Interfaces/ITargetPlatform.h
+++ b/Engine/Source/Developer/TargetPlatform/Public/Interfaces/ITargetPlatform.h
@@ -538,6 +538,11 @@ public:
 	 * Gets whether the platform uses Mobile AO
 	 */
 	virtual bool UsesMobileAmbientOcclusion() const = 0;
+
+	/**
+	 * Gets whether the platform should use DBuffer for decals when using the mobile renderer.
+	 */
+	virtual bool UsesMobileDBuffer() const = 0;
 	
 	/**
 	 * Gets whether the platform uses ASTC HDR
diff --git a/Engine/Source/Runtime/Core/Private/HAL/ConsoleManager.cpp b/Engine/Source/Runtime/Core/Private/HAL/ConsoleManager.cpp
--- a/Engine/Source/Runtime/Core/Private/HAL/ConsoleManager.cpp
+++ b/Engine/Source/Runtime/Core/Private/HAL/ConsoleManager.cpp
@@ -3336,6 +3336,16 @@ static TAutoConsoleVariable<int32> CVarDBuffer(
 		 " 1: on (default)"),
 	ECVF_RenderThreadSafe | ECVF_ReadOnly);
 
+static TAutoConsoleVariable<int32> CVarMobileDBuffer(
+	TEXT("r.Mobile.DBuffer"),
+	0,
+	TEXT("Enables DBuffer decal material blend modes when using the mobile forward renderer.\n"
+		"DBuffer decals are rendered before the base pass, allowing them to affect static lighting and skylighting correctly. \n"
+		"When enabled, a full prepass will be forced which adds CPU / GPU cost.  Several texture lookups will be done in the base pass to fetch the decal properties, which adds pixel work.\n"
+		" 0: off (default)\n"
+		" 1: on"),
+	ECVF_RenderThreadSafe | ECVF_ReadOnly);
+
 static TAutoConsoleVariable<float> CVarSkeletalMeshLODRadiusScale(
 	TEXT("r.SkeletalMeshLODRadiusScale"),
 	1.0f,
diff --git a/Engine/Source/Runtime/Engine/Classes/Engine/RendererSettings.h b/Engine/Source/Runtime/Engine/Classes/Engine/RendererSettings.h
--- a/Engine/Source/Runtime/Engine/Classes/Engine/RendererSettings.h
+++ b/Engine/Source/Runtime/Engine/Classes/Engine/RendererSettings.h
@@ -1146,6 +1146,12 @@ class URendererSettings : public UDeveloperSettings
 		ConfigRestartRequired = true))
 		uint32 bMobileAmbientOcclusion : 1;
 
+	UPROPERTY(config, EditAnywhere, Category = Mobile, meta = (
+		ConsoleVariable = "r.Mobile.DBuffer", DisplayName = "Mobile DBuffer Decals",
+		ToolTip = "Whether to accumulate decal properties to a buffer before the base pass with mobile rendering. DBuffer enabled forces a full prepass. Changing this setting requires restarting the editor.",
+		ConfigRestartRequired = true))
+		uint32 bMobileDBuffer : 1;
+
 	UPROPERTY(config, EditAnywhere, Category = Skinning, meta = (
 		ConsoleVariable = "r.GPUSkin.UnlimitedBoneInfluences", DisplayName = "Use Unlimited Bone Influences",
 		ToolTip = "If enabled, a new mesh imported will use unlimited bone buffer instead of fixed MaxBoneInfluences for rendering.",
diff --git a/Engine/Source/Runtime/Engine/Private/Materials/Material.cpp b/Engine/Source/Runtime/Engine/Private/Materials/Material.cpp
--- a/Engine/Source/Runtime/Engine/Private/Materials/Material.cpp
+++ b/Engine/Source/Runtime/Engine/Private/Materials/Material.cpp
@@ -4395,8 +4395,9 @@ bool UMaterial::CanEditChange(const FProperty* InProperty) const
 		if (PropertyName == GET_MEMBER_NAME_STRING_CHECKED(UMaterial, MaterialDecalResponse))
 		{
 			static auto* CVar = IConsoleManager::Get().FindTConsoleVariableDataInt(TEXT("r.DBuffer"));
+			static auto* CVarMobile = IConsoleManager::Get().FindTConsoleVariableDataInt(TEXT("r.Mobile.DBuffer"));
 
-			return MaterialDomain == MD_Surface && CVar->GetValueOnGameThread() > 0;
+			return MaterialDomain == MD_Surface && (CVar->GetValueOnGameThread() > 0 || CVarMobile->GetValueOnGameThread() > 0);
 		}		
 
 		if(MaterialDomain == MD_PostProcess)
diff --git a/Engine/Source/Runtime/RenderCore/Private/RenderUtils.cpp b/Engine/Source/Runtime/RenderCore/Private/RenderUtils.cpp
--- a/Engine/Source/Runtime/RenderCore/Private/RenderUtils.cpp
+++ b/Engine/Source/Runtime/RenderCore/Private/RenderUtils.cpp
@@ -644,7 +644,7 @@ RENDERCORE_API bool MobileBasePassAlwaysUsesCSM(const FStaticShaderPlatform Plat
 RENDERCORE_API bool MobileUsesFullDepthPrepass(const FStaticShaderPlatform Platform)
 {
 	static FShaderPlatformCachedIniValue<int32> CVarMobileEarlyZPass(TEXT("r.Mobile.EarlyZPass"));
-	return MobileUsesShadowMaskTexture(Platform) || IsMobileAmbientOcclusionEnabled(Platform) || (CVarMobileEarlyZPass.Get(Platform) == 1);
+	return MobileUsesShadowMaskTexture(Platform) || IsMobileAmbientOcclusionEnabled(Platform) || IsUsingDBuffers(Platform) || (CVarMobileEarlyZPass.Get(Platform) == 1);
 }
 
 RENDERCORE_API bool ShouldForceFullDepthPass(const FStaticShaderPlatform Platform)
@@ -820,7 +820,7 @@ RENDERCORE_API void RenderUtilsInit()
 
 					GForwardShadingPlatformMask[ShaderPlatformIndex] = TargetPlatform->UsesForwardShading();
 
-					GDBufferPlatformMask[ShaderPlatformIndex] = TargetPlatform->UsesDBuffer() && !IsMobilePlatform(ShaderPlatformToEdit);
+					GDBufferPlatformMask[ShaderPlatformIndex] = IsMobilePlatform(ShaderPlatformToEdit) ? (TargetPlatform->UsesMobileDBuffer() && !IsMobileDeferredShadingEnabled(ShaderPlatformToEdit)) : TargetPlatform->UsesDBuffer();
 
 					GSelectiveBasePassOutputsPlatformMask[ShaderPlatformIndex] = TargetPlatform->UsesSelectiveBasePassOutputs();
 
@@ -859,10 +859,10 @@ RENDERCORE_API void RenderUtilsInit()
 	}
 
 #else // WITH_EDITOR
-
 	if (IsMobilePlatform(GMaxRHIShaderPlatform))
 	{
-		GDBufferPlatformMask.Init(false, EShaderPlatform::SP_NumPlatforms);
+		static IConsoleVariable* MobileDBufferCVar = IConsoleManager::Get().FindConsoleVariable(TEXT("r.Mobile.DBuffer"));
+		GDBufferPlatformMask.Init(MobileDBufferCVar && MobileDBufferCVar->GetInt() && !IsMobileDeferredShadingEnabled(GMaxRHIShaderPlatform), EShaderPlatform::SP_NumPlatforms);
 	}
 
 	if (RayTracingCVar && RayTracingCVar->GetInt() && GRHISupportsRayTracing)
@@ -1420,7 +1420,7 @@ bool IsSingleLayerWaterDepthPrepassEnabled(const FStaticShaderPlatform Platform,
 
 RENDERCORE_API bool IsUsingDBuffers(const FStaticShaderPlatform Platform)
 {
-	return (GDBufferPlatformMask[(int)Platform]);
+	return GDBufferPlatformMask[(int)Platform];
 }
 
 RENDERCORE_API bool AreSkinCacheShadersEnabled(EShaderPlatform Platform)
diff --git a/Engine/Source/Runtime/Renderer/Private/CompositionLighting/PostProcessDeferredDecals.cpp b/Engine/Source/Runtime/Renderer/Private/CompositionLighting/PostProcessDeferredDecals.cpp
--- a/Engine/Source/Runtime/Renderer/Private/CompositionLighting/PostProcessDeferredDecals.cpp
+++ b/Engine/Source/Runtime/Renderer/Private/CompositionLighting/PostProcessDeferredDecals.cpp
@@ -69,17 +69,36 @@ FDeferredDecalPassTextures GetDeferredDecalPassTextures(
 	FDeferredDecalPassTextures PassTextures;
 
 	auto* Parameters = GraphBuilder.AllocParameters<FDecalPassUniformParameters>();
-	const ESceneTextureSetupMode TextureReadAccess = ESceneTextureSetupMode::GBufferA | ESceneTextureSetupMode::SceneDepth | ESceneTextureSetupMode::CustomDepth;
+	
+	const bool bIsMobile = (View.GetFeatureLevel() == ERHIFeatureLevel::ES3_1);
+	ESceneTextureSetupMode TextureReadAccess = ESceneTextureSetupMode::None;
+	EMobileSceneTextureSetupMode MobileTextureReadAccess = EMobileSceneTextureSetupMode::None;
+	if (bIsMobile)
+	{
+		MobileTextureReadAccess = EMobileSceneTextureSetupMode::SceneDepth | EMobileSceneTextureSetupMode::CustomDepth;
+	}
+	else
+	{
+		TextureReadAccess = ESceneTextureSetupMode::GBufferA | ESceneTextureSetupMode::SceneDepth | ESceneTextureSetupMode::CustomDepth;
+	}
+
 	SetupSceneTextureUniformParameters(GraphBuilder, &SceneTextures, View.FeatureLevel, TextureReadAccess, Parameters->SceneTextures);
+	SetupMobileSceneTextureUniformParameters(GraphBuilder, &SceneTextures, MobileTextureReadAccess, Parameters->MobileSceneTextures);
 	Parameters->EyeAdaptationBuffer = GraphBuilder.CreateSRV(GetEyeAdaptationBuffer(GraphBuilder, View));
 	PassTextures.DecalPassUniformBuffer = GraphBuilder.CreateUniformBuffer(Parameters);
 
 	PassTextures.Depth = SceneTextures.Depth;
 	PassTextures.Color = SceneTextures.Color.Target;
-	PassTextures.GBufferA = (*SceneTextures.UniformBuffer)->GBufferATexture;
-	PassTextures.GBufferB = (*SceneTextures.UniformBuffer)->GBufferBTexture;
-	PassTextures.GBufferC = (*SceneTextures.UniformBuffer)->GBufferCTexture;
-	PassTextures.GBufferE = (*SceneTextures.UniformBuffer)->GBufferETexture;
+
+	// Mobile deferred renderer does not use dbuffer 
+	if (!bIsMobile)
+	{
+		PassTextures.GBufferA = (*SceneTextures.UniformBuffer)->GBufferATexture;
+		PassTextures.GBufferB = (*SceneTextures.UniformBuffer)->GBufferBTexture;
+		PassTextures.GBufferC = (*SceneTextures.UniformBuffer)->GBufferCTexture;
+		PassTextures.GBufferE = (*SceneTextures.UniformBuffer)->GBufferETexture;
+	}
+
 	PassTextures.DBufferTextures = DBufferTextures;
 
 	return PassTextures;
diff --git a/Engine/Source/Runtime/Renderer/Private/CompositionLighting/PostProcessDeferredDecals.h b/Engine/Source/Runtime/Renderer/Private/CompositionLighting/PostProcessDeferredDecals.h
--- a/Engine/Source/Runtime/Renderer/Private/CompositionLighting/PostProcessDeferredDecals.h
+++ b/Engine/Source/Runtime/Renderer/Private/CompositionLighting/PostProcessDeferredDecals.h
@@ -19,6 +19,7 @@ bool IsDBufferEnabled(const FSceneViewFamily& ViewFamily, EShaderPlatform Shader
 
 BEGIN_GLOBAL_SHADER_PARAMETER_STRUCT(FDecalPassUniformParameters, )
 	SHADER_PARAMETER_STRUCT(FSceneTextureUniformParameters, SceneTextures)
+	SHADER_PARAMETER_STRUCT(FMobileSceneTextureUniformParameters, MobileSceneTextures)
 	SHADER_PARAMETER_RDG_BUFFER_SRV(StructuredBuffer<float4>, EyeAdaptationBuffer)
 END_GLOBAL_SHADER_PARAMETER_STRUCT()
 
diff --git a/Engine/Source/Runtime/Renderer/Private/DBufferTextures.cpp b/Engine/Source/Runtime/Renderer/Private/DBufferTextures.cpp
--- a/Engine/Source/Runtime/Renderer/Private/DBufferTextures.cpp
+++ b/Engine/Source/Runtime/Renderer/Private/DBufferTextures.cpp
@@ -62,6 +62,11 @@ FDBufferTexturesDesc GetDBufferTexturesDesc(FIntPoint Extent, EShaderPlatform Sh
 			// This significantly reduces bandwidth for clearing, writing and reading on some GPUs.
 			// While a smaller format, such as R8_UINT, will use less video memory, it will result in slower clears and higher bandwidth requirements.
 			check(Desc.Format == PF_B8G8R8A8);
+			// On mobile platforms using PF_B8G8R8A8 has no benefits over R8.
+			if (IsMobilePlatform(ShaderPlatform))
+			{
+				Desc.Format = PF_R8;
+			}
 			Desc.Flags = TexCreate_ShaderResource | TexCreate_RenderTargetable;
 			Desc.ClearValue = FClearValueBinding::Transparent;
 			DBufferTexturesDesc.DBufferMaskDesc = Desc;
diff --git a/Engine/Source/Runtime/Renderer/Private/DecalRenderingCommon.cpp b/Engine/Source/Runtime/Renderer/Private/DecalRenderingCommon.cpp
--- a/Engine/Source/Runtime/Renderer/Private/DecalRenderingCommon.cpp
+++ b/Engine/Source/Runtime/Renderer/Private/DecalRenderingCommon.cpp
@@ -23,7 +23,7 @@ namespace DecalRendering
 		const bool bIsSubstrateEnabled = Substrate::IsSubstrateEnabled();
 		const bool bIsMobilePlatform = IsMobilePlatform(Platform);
 		const bool bIsMobileDeferredPlatform = bIsMobilePlatform && IsMobileDeferredShadingEnabled(Platform);
-		const bool bIsDBufferPlatform = !bIsMobilePlatform && IsUsingDBuffers(Platform);
+		const bool bIsDBufferPlatform = IsUsingDBuffers(Platform);
 		const bool bIsDBufferMaskPlatform = bIsDBufferPlatform && FDataDrivenShaderPlatformInfo::GetSupportsPerPixelDBufferMask(Platform);
 
 		Desc.bWriteDBufferMask = bIsDBufferMaskPlatform;
@@ -39,7 +39,7 @@ namespace DecalRendering
 		}
 
 		// Enforce platform output limitations.
-		if (bIsMobilePlatform && !bIsMobileDeferredPlatform)
+		if (bIsMobilePlatform && !bIsMobileDeferredPlatform && !bIsDBufferPlatform)
 		{
 			Desc.bWriteNormal = false;
 			Desc.bWriteRoughnessSpecularMetallic = false;
@@ -60,7 +60,7 @@ namespace DecalRendering
 		{
 			Desc.RenderStageMask |= 1 << (uint32)EDecalRenderStage::MobileBeforeLighting;
 		}
-		else if (bIsMobilePlatform && (Desc.bWriteEmissive || Desc.bWriteBaseColor))
+		else if ((bIsMobilePlatform && !bIsDBufferPlatform) && (Desc.bWriteEmissive || Desc.bWriteBaseColor))
 		{
 			Desc.RenderStageMask |= 1 << (uint32)EDecalRenderStage::Mobile;
 		}
diff --git a/Engine/Source/Runtime/Renderer/Private/MobileBasePassRendering.cpp b/Engine/Source/Runtime/Renderer/Private/MobileBasePassRendering.cpp
--- a/Engine/Source/Runtime/Renderer/Private/MobileBasePassRendering.cpp
+++ b/Engine/Source/Runtime/Renderer/Private/MobileBasePassRendering.cpp
@@ -217,6 +217,8 @@ void SetupMobileBasePassUniformParameters(
 	{
 		AmbientOcclusionTexture = MobileBasePassTextures.ScreenSpaceAO;
 	}
+	
+	BasePassParameters.DBuffer = GetDBufferParameters(GraphBuilder, MobileBasePassTextures.DBufferTextures, View.GetShaderPlatform());
 
 	BasePassParameters.AmbientOcclusionTexture = AmbientOcclusionTexture;
 	BasePassParameters.AmbientOcclusionSampler = TStaticSamplerState<SF_Bilinear, AM_Clamp, AM_Clamp, AM_Clamp>::GetRHI();
diff --git a/Engine/Source/Runtime/Renderer/Private/MobileBasePassRendering.h b/Engine/Source/Runtime/Renderer/Private/MobileBasePassRendering.h
--- a/Engine/Source/Runtime/Renderer/Private/MobileBasePassRendering.h
+++ b/Engine/Source/Runtime/Renderer/Private/MobileBasePassRendering.h
@@ -26,6 +26,8 @@
 #include "RenderUtils.h"
 #include "DebugViewModeRendering.h"
 #include "LocalFogVolumeRendering.h"
+#include "DBufferTextures.h"
+#include "CompositionLighting/PostProcessDeferredDecals.h"
 
 bool MobileLocalLightsBufferEnabled(const FStaticShaderPlatform Platform);
 bool MobileLocalLightsBufferPrepassEnabled(const FStaticShaderPlatform Platform);
@@ -34,6 +36,7 @@ bool MobileLocalLightsBufferPostprocessEnabled(const FStaticShaderPlatform Platf
 struct FMobileBasePassTextures
 {
 	FRDGTextureRef ScreenSpaceAO = nullptr;
+	FDBufferTextures DBufferTextures = {};
 };
 
 BEGIN_GLOBAL_SHADER_PARAMETER_STRUCT(FMobileBasePassUniformParameters, )
@@ -56,6 +59,7 @@ BEGIN_GLOBAL_SHADER_PARAMETER_STRUCT(FMobileBasePassUniformParameters, )
 	SHADER_PARAMETER_SAMPLER(SamplerState, AmbientOcclusionSampler)
 	SHADER_PARAMETER_RDG_TEXTURE(Texture2D, ScreenSpaceShadowMaskTexture)
 	SHADER_PARAMETER_SAMPLER(SamplerState, ScreenSpaceShadowMaskSampler)
+	SHADER_PARAMETER_STRUCT_INCLUDE(FDBufferParameters, DBuffer)
 END_GLOBAL_SHADER_PARAMETER_STRUCT()
 
 enum class EMobileBasePass
@@ -408,6 +412,7 @@ public:
 		OutEnvironment.SetDefine(TEXT("PREPASS_LOCAL_LIGHTS_MOBILE"), (LocalLightSetting == EMobileLocalLightSetting::LOCAL_LIGHTS_BUFFER) ? 1u: 0u);
 		OutEnvironment.SetDefine(TEXT("ENABLE_CLUSTERED_REFLECTION"), bEnableClusteredReflections ? 1u : 0u);
 		OutEnvironment.SetDefine(TEXT("USE_SHADOWMASKTEXTURE"), bMobileUsesShadowMaskTexture && !bTranslucentMaterial ? 1u : 0u);
+		OutEnvironment.SetDefine(TEXT("ENABLE_DBUFFER_TEXTURES"), Parameters.MaterialParameters.MaterialDomain == MD_Surface ? 1u : 0u);
 	}
 	
 	/** Initialization constructor. */
diff --git a/Engine/Source/Runtime/Renderer/Private/MobileDecalRendering.cpp b/Engine/Source/Runtime/Renderer/Private/MobileDecalRendering.cpp
--- a/Engine/Source/Runtime/Renderer/Private/MobileDecalRendering.cpp
+++ b/Engine/Source/Runtime/Renderer/Private/MobileDecalRendering.cpp
@@ -17,6 +17,8 @@
 #include "DecalRenderingShared.h"
 #include "RenderCore.h"
 #include "DataDrivenShaderPlatformInfo.h"
+#include "CompositionLighting/PostProcessDeferredDecals.h"
+#include "DBufferTextures.h"
 
 void RenderMeshDecalsMobile(FRHICommandList& RHICmdList, const FViewInfo& View, EDecalRenderStage DecalRenderStage, EDecalRenderTargetMode RenderTargetMode);
 extern void RenderDeferredDecalsMobile(FRHICommandList& RHICmdList, const FScene& Scene, const FViewInfo& View, EDecalRenderStage DecalRenderStage, EDecalRenderTargetMode RenderTargetMode);
@@ -53,7 +55,7 @@ void FMobileSceneRenderer::RenderDecals(FRHICommandList& RHICmdList, const FView
 	SCOPE_CYCLE_COUNTER(STAT_DecalsDrawTime);
 
 	const bool bIsMobileDeferred = IsMobileDeferredShadingEnabled(View.GetShaderPlatform());
-	const EDecalRenderStage DecalRenderStage = bIsMobileDeferred ? EDecalRenderStage::MobileBeforeLighting : EDecalRenderStage::Mobile;
+	const EDecalRenderStage DecalRenderStage = bIsMobileDeferred ? EDecalRenderStage::MobileBeforeLighting : bRequiresDBufferDecals ? EDecalRenderStage::Emissive : EDecalRenderStage::Mobile;
 	const EDecalRenderTargetMode RenderTargetMode = bIsMobileDeferred ? EDecalRenderTargetMode::SceneColorAndGBuffer : EDecalRenderTargetMode::SceneColor;
 
 	// Deferred decals
@@ -137,3 +139,24 @@ void RenderDeferredDecalsMobile(FRHICommandList& RHICmdList, const FScene& Scene
 		}
 	}
 }
+
+void FMobileSceneRenderer::RenderDBuffer(FRDGBuilder& GraphBuilder, FSceneTextures& SceneTextures, FDBufferTextures& DBufferTextures)
+{
+	RDG_EVENT_SCOPE(GraphBuilder, "RenderDBuffer");
+	QUICK_SCOPE_CYCLE_COUNTER(STAT_RenderDBuffer);
+
+	const EShaderPlatform Platform = GetViewFamilyInfo(Views).GetShaderPlatform();
+
+	for (int32 ViewIndex = 0; ViewIndex < Views.Num(); ++ViewIndex)
+	{
+		FViewInfo& View = Views[ViewIndex];
+
+		if (!View.ShouldRenderView())
+		{
+			continue;
+		}
+
+		FDeferredDecalPassTextures DecalPassTextures = GetDeferredDecalPassTextures(GraphBuilder, View, SceneTextures, &DBufferTextures);
+		AddDeferredDecalPass(GraphBuilder, View, DecalPassTextures, EDecalRenderStage::BeforeBasePass);
+	}
+}
diff --git a/Engine/Source/Runtime/Renderer/Private/MobileShadingRenderer.cpp b/Engine/Source/Runtime/Renderer/Private/MobileShadingRenderer.cpp
--- a/Engine/Source/Runtime/Renderer/Private/MobileShadingRenderer.cpp
+++ b/Engine/Source/Runtime/Renderer/Private/MobileShadingRenderer.cpp
@@ -268,6 +268,7 @@ FMobileSceneRenderer::FMobileSceneRenderer(const FSceneViewFamily* InViewFamily,
 	: FSceneRenderer(InViewFamily, HitProxyConsumer)
 	, bGammaSpace(!IsMobileHDR())
 	, bDeferredShading(IsMobileDeferredShadingEnabled(ShaderPlatform))
+	, bRequiresDBufferDecals(bDeferredShading ? false : IsUsingDBuffers(ShaderPlatform))
 	, bUseVirtualTexturing(UseVirtualTexturing(ShaderPlatform) && GetRendererOutput() == FSceneRenderer::ERendererOutput::FinalSceneColor)
 {
 	bRenderToSceneColor = false;
@@ -1059,6 +1060,7 @@ void FMobileSceneRenderer::Render(FRDGBuilder& GraphBuilder)
 		CopySceneCaptureComponentToTarget(GraphBuilder, SceneTextures, nullptr, ViewFamily, PassInfo.Views);
 	}
 
+	FDBufferTextures DBufferTextures{};
 	if (bIsFullDepthPrepassEnabled)
 	{
 		RenderFullDepthPrepass(GraphBuilder, Views, SceneTextures);
@@ -1097,6 +1099,15 @@ void FMobileSceneRenderer::Render(FRDGBuilder& GraphBuilder)
 		{
 			RenderMobileLocalLightsBuffer(GraphBuilder, SceneTextures, true, SortedLightSet);
 		}
+
+		if (bRendererOutputFinalSceneColor)
+		{
+			if (bRequiresDBufferDecals)
+			{
+				DBufferTextures = CreateDBufferTextures(GraphBuilder, SceneTextures.Config.Extent, ShaderPlatform);
+				RenderDBuffer(GraphBuilder, SceneTextures, DBufferTextures);
+			}
+		}
 	}
 
 	if (bRendererOutputFinalSceneColor)
@@ -1107,7 +1118,7 @@ void FMobileSceneRenderer::Render(FRDGBuilder& GraphBuilder)
 		}
 		else
 		{
-			RenderForward(GraphBuilder, ViewFamilyTexture, SceneTextures);
+			RenderForward(GraphBuilder, ViewFamilyTexture, SceneTextures, DBufferTextures);
 		}
 
 		EndOcclusionScope(GraphBuilder, Views);
@@ -1229,7 +1240,7 @@ void FMobileSceneRenderer::BuildInstanceCullingDrawParams(FRDGBuilder& GraphBuil
 	}
 }
 
-void FMobileSceneRenderer::RenderForward(FRDGBuilder& GraphBuilder, FRDGTextureRef ViewFamilyTexture, FSceneTextures& SceneTextures)
+void FMobileSceneRenderer::RenderForward(FRDGBuilder& GraphBuilder, FRDGTextureRef ViewFamilyTexture, FSceneTextures& SceneTextures, FDBufferTextures& DBufferTextures)
 {
 	const FViewInfo& MainView = Views[0];
 
@@ -1314,6 +1325,7 @@ void FMobileSceneRenderer::RenderForward(FRDGBuilder& GraphBuilder, FRDGTextureR
 
 		FMobileBasePassTextures MobileBasePassTextures{};
 		MobileBasePassTextures.ScreenSpaceAO = bRequiresAmbientOcclusionPass ? SceneTextures.ScreenSpaceAO : SystemTextures.White;
+		MobileBasePassTextures.DBufferTextures = DBufferTextures;
 
 		EMobileSceneTextureSetupMode SetupMode = (bIsFullDepthPrepassEnabled ? EMobileSceneTextureSetupMode::SceneDepth : EMobileSceneTextureSetupMode::None) | EMobileSceneTextureSetupMode::CustomDepth;
 		FMobileRenderPassParameters* PassParameters = GraphBuilder.AllocParameters<FMobileRenderPassParameters>();
diff --git a/Engine/Source/Runtime/Renderer/Private/SceneRendering.h b/Engine/Source/Runtime/Renderer/Private/SceneRendering.h
--- a/Engine/Source/Runtime/Renderer/Private/SceneRendering.h
+++ b/Engine/Source/Runtime/Renderer/Private/SceneRendering.h
@@ -89,6 +89,7 @@ struct FNaniteMaterialPassCommand;
 struct FScreenMessageWriter;
 struct FVolumetricFogIntegrationParameterData;
 class FLumenHardwareRayTracingUniformBufferParameters;
+struct FDBufferTextures;
 
 namespace Nanite
 {
@@ -2687,6 +2688,7 @@ protected:
 
 	/** Renders decals. */
 	void RenderDecals(FRHICommandList& RHICmdList, const FViewInfo& View);
+	void RenderDBuffer(FRDGBuilder& GraphBuilder, FSceneTextures& SceneTextures, FDBufferTextures& DBufferTextures);
 
 	/** Renders the atmospheric and height fog */
 	void RenderFog(FRHICommandList& RHICmdList, const FViewInfo& View);
@@ -2704,7 +2706,7 @@ protected:
 	
 	void BuildInstanceCullingDrawParams(FRDGBuilder& GraphBuilder, FViewInfo& View, class FMobileRenderPassParameters* PassParameters);
 
-	void RenderForward(FRDGBuilder& GraphBuilder, FRDGTextureRef ViewFamilyTexture, FSceneTextures& SceneTextures);
+	void RenderForward(FRDGBuilder& GraphBuilder, FRDGTextureRef ViewFamilyTexture, FSceneTextures& SceneTextures, FDBufferTextures& DBufferTextures);
 	void RenderForwardSinglePass(FRDGBuilder& GraphBuilder, class FMobileRenderPassParameters* PassParameters, struct FRenderViewContext& ViewContext, FSceneTextures& SceneTextures);
 	void RenderForwardMultiPass(FRDGBuilder& GraphBuilder, class FMobileRenderPassParameters* PassParameters, struct FRenderViewContext& ViewContext, FSceneTextures& SceneTextures);
 	
@@ -2720,6 +2722,7 @@ protected:
 private:
 	const bool bGammaSpace;
 	const bool bDeferredShading;
+	const bool bRequiresDBufferDecals;
 	const bool bUseVirtualTexturing;
 	int32 NumMSAASamples;
 	bool bRenderToSceneColor;
