Substrate - fix for legacy unlit translucent to follow Substrate tree required behavior with unlit.
(This got changed recently to account for the tree processign when using Coverage node)

This fixes substrate affecting emissive by coverage twice when using Translucent blend modes.
Composite also should not specify any coverage affecting luminance since it is used from the root node.

diff --git a/Engine/Shaders/Private/Substrate/SubstrateLegacyConversion.ush b/Engine/Shaders/Private/Substrate/SubstrateLegacyConversion.ush
--- a/Engine/Shaders/Private/Substrate/SubstrateLegacyConversion.ush
+++ b/Engine/Shaders/Private/Substrate/SubstrateLegacyConversion.ush
@@ -479,24 +479,25 @@ FSubstrateData SubstrateConvertLegacyMaterialStatic(
 
 	if (ShadingModel == SHADINGMODELID_UNLIT)
 	{
+		// Now we convert the legacy behavior to Substrate Coverage/Transmittance controls.
+		// Note: Emissive will always be weighted later by the Unlit BSDF Coverage (see SubstrateUpdateTreeUnlit).
+
 		const float GreyTransmittance = saturate(1.0f - Opacity);
 	#if MATERIALBLENDING_SOLID || MATERIALBLENDING_MASKED
 		// Opaque materials only write emissive.
 		Out = GetSubstrateUnlitBSDF(Emissive, 0.0f, RawNormal);
 	#elif MATERIALBLENDING_ALPHACOMPOSITE
-		// Uses (Luminance, black transmittance) for premultiplied alpha blending. Opacity is controled using coverage
+		// Uses (Luminance, black transmittance) for premultiplied alpha blending. 
+		// Coverage over the background is specified through the root node input "Coverage Override".
 		Out = GetSubstrateUnlitBSDF(Emissive, 0.0f, RawNormal);
-		#if SUBSTRATE_INLINE_SHADING
-			Out.InlinedBSDF.Coverage = Opacity;
-		#endif
 	#elif MATERIALBLENDING_ALPHAHOLDOUT
-		// The MATERIALBLENDING_ALPHAHOLDOUT case needs to be before the MATERIALBLENDING_TRANSLUCENT case, as alphaholdout material have both cases defined
+		// The MATERIALBLENDING_ALPHAHOLDOUT case needs to be before the MATERIALBLENDING_TRANSLUCENT case, as alphaholdout material have both cases defined.
 		Out = GetSubstrateUnlitBSDF(0.0f, GreyTransmittance, RawNormal);
 	#elif MATERIALBLENDING_TRANSLUCENT
-		// Uses (Luminance, black transmittance) for translucency blending. Opacity is controled using coverage
-		Out = GetSubstrateUnlitBSDF(Emissive * Opacity, 0.0f, RawNormal);
+		// Uses (Luminance, black transmittance) for translucency blending. Opacity is controled using coverage.
+		Out = GetSubstrateUnlitBSDF(Emissive, 0.0f, RawNormal);
 		#if SUBSTRATE_INLINE_SHADING
-			Out.InlinedBSDF.Coverage = Opacity;
+			Out.InlinedBSDF.Coverage = Opacity; // Applied to emissive via the Substrate Tree, (see SubstrateUpdateTreeUnlit).
 		#endif
 	#elif MATERIALBLENDING_ADDITIVE
 		// Emissive is weighted by opacity in the legacy material.
@@ -505,7 +506,7 @@ FSubstrateData SubstrateConvertLegacyMaterialStatic(
 		// Setting up emissive as the transmittance color. It is not clamped in case brightening is required.
 		Out = GetSubstrateUnlitBSDF(0.0f, Emissive, RawNormal);
 	#else
-		// Required default for some materials such as Editor UI
+		// Required default for some materials such as Editor UI.
 		Out = GetSubstrateUnlitBSDF(Emissive, 0.0f, RawNormal);
 	#endif
 	}
