Integrated Bend Screen Space Shadows.
- only supported on mobile renderer at the moment since it requires running as standalone pass.
- doesn't respect Directional Light Contact Shadow properties at the moment.
- calculated at full resolution instead of respecting GetScreenSpaceShadowDownsampleFactor().

diff --git a/Engine/Shaders/Private/ScreenSpaceShadows.usf b/Engine/Shaders/Private/ScreenSpaceShadows.usf
--- a/Engine/Shaders/Private/ScreenSpaceShadows.usf
+++ b/Engine/Shaders/Private/ScreenSpaceShadows.usf
@@ -4,6 +4,43 @@
 	ScreenSpaceShadows.usf
 =============================================================================*/
 
+#include "/Engine/Public/Platform.ush"
+
+#define WAVE_SIZE 64
+#define SAMPLE_COUNT 60
+#define HARD_SHADOW_SAMPLES 4
+#define FADE_OUT_SAMPLES 8
+
+#include "bend_sss_gpu.ush"
+
+float2 InvDepthTextureSize;
+float4 LightCoordinate;
+int3 WaveOffset;
+
+Texture2D<float> DepthTexture;
+RWTexture2D<float> OutputTexture;
+SamplerState PointBorderSampler;
+
+DispatchParameters SetupBendDispatchParameters(float NearDepthVal, float FarDepthVal)
+{
+	DispatchParameters Parameters;
+	Parameters.SetDefaults();
+
+	Parameters.LightCoordinate = LightCoordinate;
+	Parameters.WaveOffset = WaveOffset.xy;
+
+	Parameters.FarDepthValue = FarDepthVal;
+	Parameters.NearDepthValue = NearDepthVal;
+
+	Parameters.InvDepthTextureSize = InvDepthTextureSize;
+
+	Parameters.DepthTexture = DepthTexture;
+	Parameters.OutputTexture = OutputTexture;
+	Parameters.PointBorderSampler = PointBorderSampler;
+
+	return Parameters;
+}
+
 #include "ScreenSpaceShadowRayCast.ush"
 
 #include "Common.ush"
@@ -63,6 +100,17 @@ void ScreenSpaceShadowsCS(
 	RWShadowFactors[DispatchThreadId.xy] = float2(Result, SceneDepth);
 }
 
+[numthreads(THREADGROUP_SIZEX, THREADGROUP_SIZEY, 1)]
+void ScreenSpaceShadowsBendCS(
+	uint3 GroupId : SV_GroupID,
+	uint3 DispatchThreadId : SV_DispatchThreadID,
+	uint3 GroupThreadId : SV_GroupThreadID)
+{
+	DispatchParameters Parameters = SetupBendDispatchParameters(NearDepthValue, FarDepthValue);
+
+	WriteScreenSpaceShadow(Parameters, GroupId, GroupThreadId.x);
+}
+
 #ifdef UPSAMPLE_PASS
 
 Texture2D ShadowFactorsTexture;
diff --git a/Engine/Shaders/Private/bend_sss_gpu.ush b/Engine/Shaders/Private/bend_sss_gpu.ush
new file mode 100644
--- /dev/null
+++ b/Engine/Shaders/Private/bend_sss_gpu.ush
@@ -0,0 +1,487 @@
+#pragma once
+
+// Copyright 2023 Sony Interactive Entertainment.
+//
+// Licensed under the Apache License, Version 2.0 (the "License");
+// you may not use this file except in compliance with the License.
+// You may obtain a copy of the License at
+//
+//     http://www.apache.org/licenses/LICENSE-2.0
+//
+// Unless required by applicable law or agreed to in writing, software
+// distributed under the License is distributed on an "AS IS" BASIS,
+// WITHOUT WARRANTIES OR CONDITIONS OF ANY KIND, either express or implied.
+// See the License for the specific language governing permissions and
+// limitations under the License.
+
+// If you have feedback, or found this code useful, we'd love to hear from you.
+// https://www.bendstudio.com
+// https://www.twitter.com/bendstudio
+// 
+// We are *always* looking for talented graphics and technical programmers!
+// https://www.bendstudio.com/careers
+
+
+// Common screen space shadow projection code (GPU):
+//--------------------------------------------------------------
+
+// The main shadow generation function is WriteScreenSpaceShadow(), it will read a depth texture, and write to a shadow texture
+// This code is setup to target DX12 DXC shader compiler, but has also been tested on PS5 with appropriate API remapping. 
+// It can compile to DX11, but requires some modifications (e.g., early-out's use of wave intrinsics is not supported in DX11).
+// Note; you can customize the 'EarlyOutPixel' function to perform custom early-out logic to optimize this shader.
+
+// The following Macros must be defined in the compute shader file before including this header:
+//		
+//		
+//		#define WAVE_SIZE 64						// Wavefront size of the compute shader running this code. 
+//													// numthreads[WAVE_SIZE, 1, 1]
+//													// Only tested with 64.
+//		
+//		#define SAMPLE_COUNT 60						// Number of shadow samples per-pixel.
+//													// Determines overall cost, as this value controls the length of the shadow (in pixels).
+//													// The number of texture-reads performed per-thread will be (SAMPLE_COUNT / WAVE_SIZE + 2) * 2.
+//													// Recommended starting value is 60 (This would be 4 reads per thread if WAVE_SIZE is 64). A value of 64 would require 6 reads.
+//		
+//		// Not all shadow samples are treated the same:
+//		//	The bulk of samples will average together in to groups of 4, to produce a slightly smoothed result (so one sample cannot fully show the pixel)
+//		//	However, the samples very close to the start pixel can optionally be forced to disable this averaging, so a single sample can fully shadow the pixel (HardShadowSamples)
+//		//	Plus, a number of the last (most distant) samples can (for a small cost) apply a fade-out effect to soften a hash shadow cutoff (FadeOutSamples)
+//		
+//		#define HARD_SHADOW_SAMPLES 4				// Number of initial shadow samples that will produce a hard shadow, and not perform sample-averaging.
+//													// This trades aliasing for grounding pixels very close to the shadow caster.
+//													// Recommended starting value: 4
+//		
+//		#define FADE_OUT_SAMPLES 8					// Number of samples that will fade out at the end of the shadow (for a minor cost).
+//													// Recommended starting value: 8
+
+
+#if defined(__HLSL_VERSION) || defined(__hlsl_dx_compiler) || 1
+	
+	#define USE_HALF_PIXEL_OFFSET 1		// Apply a 0.5 texel offset when sampling a texture. Toggle this macro if the output shadow has odd, regular grid-like artefacts.
+
+	// HLSL enforces that a pixel offset in a Sample() call must be a compile time constant, which isn't always required - and in some cases can give a small perf boost if used.
+	#define USE_UV_PIXEL_BIAS 1			// Use Sample(uv + bias) instead of Sample(uv, bias)
+
+#endif
+
+
+// Forward declare:
+// Generate the shadow
+//	Call this function from a compute shader with thread dimensions: numthreads[WAVE_SIZE, 1, 1]
+// 
+//	(int3)	inGroupID:			Compute shader group id register (SV_GroupID)
+//	(int)	inGroupThreadId:	Compute shader group thread id register (SV_GroupThreadID)
+void WriteScreenSpaceShadow(struct DispatchParameters inParameters, int3 inGroupID, int inGroupThreadID);
+
+
+
+// This is the list of runtime properties to pass to the shader
+// Wherever possible, it is highly recommended to have these values be compile-time constants
+struct DispatchParameters
+{
+	// Visual configuration:
+	// These values will require manual tuning.
+	// All shadow computation is performed in non-linear depth space (not in world space), so tuned value choices will depend on scene depth distribution (as determined by the Projection Matrix setup).
+
+	float SurfaceThickness;				// This is the assumed thickness of each pixel for shadow-casting, measured as a percentage of the difference in non-linear depth between the sample and FarDepthValue.
+										// Recommended starting value: 0.005 (0.5%)
+
+	float BilinearThreshold;			// Percentage threshold for determining if the difference between two depth values represents an edge, and should not perform interpolation.
+										// To tune this value, set 'DebugOutputEdgeMask' to true to visualize where edges are being detected.
+										// Recommended starting value: 0.02 (2%)
+
+	float ShadowContrast;				// A contrast boost is applied to the transition in/out of shadow.
+										// Recommended starting value: 2 or 4. Values >= 1 are valid.
+
+	bool IgnoreEdgePixels;				// If an edge is detected, the edge pixel will not contribute to the shadow.
+										// If a very flat surface is being lit and rendered at an grazing angles, the edge detect may incorrectly detect multiple 'edge' pixels along that flat surface.
+										// In these cases, the grazing angle of the light may subsequently produce aliasing artefacts in the shadow where these incorrect edges were detected.
+										// Setting this value to true would mean that those pixels would not cast a shadow, however it can also thin out otherwise valid shadows, especially on foliage edges.
+										// Recommended starting value: false, unless typical scenes have numerous large flat surfaces, in which case true.
+
+	bool UsePrecisionOffset;			// A small offset is applied to account for an imprecise depth buffer (recommend off)
+
+
+	bool BilinearSamplingOffsetMode;	// There are two modes to compute bilinear samples for shadow depth:
+										// true = sampling points for pixels are offset to the wavefront shared ray, shadow depths and starting depths are the same. Can project more jagged/aliased shadow lines in some cases.
+										// false = sampling points for pixels are not offset and start from pixel centers. Shadow depths are biased based on depth gradient across the current pixel bilinear sample. Has more issues in back-face / grazing areas.
+										// Both modes have subtle visual differences, which may / may not exaggerate depth buffer aliasing that gets projected in to the shadow.
+										// Evaluating the visual difference between each mode is recommended, then hard-coding the mode used to optimize the shader.
+										// Recommended starting value: false
+
+	// Debug views
+	bool DebugOutputEdgeMask;			// Use this to visualize edges, for tuning the 'BilinearThreshold' value.
+	bool DebugOutputThreadIndex;		// Debug output to visualize layout of compute threads
+	bool DebugOutputWaveIndex;			// Debug output to visualize layout of compute wavefronts, useful to sanity check the Light Coordinate is being computed correctly.
+
+	// Culling / Early out:
+	float2 DepthBounds;					// Depth Bounds (min, max) for the on-screen volume of the light. Typically (0,1) for directional lights. Only used when 'UseEarlyOut' is true.
+
+	bool UseEarlyOut;					// Set to true to early-out when depth values are not within [DepthBounds] - otherwise DepthBounds is unused
+										// [Optionally customize the 'EarlyOutPixel()' function to perform your own early-out logic, e.g. skipping pixels that a shadow map indicates are already fully occluded]
+										// This can dramatically reduce cost when only a small portion of the pixels need a shadow term (e.g., cull out sky pixels), however it does have some overhead (~15%) in worst-case where nothing early-outs
+										// Note; Early-out is most efficient when WAVE_SIZE matches the hardware wavefront size - otherwise cross wave communication is required.
+
+
+	// Set sensible starting tuning values
+	void SetDefaults()
+	{
+		SurfaceThickness			= 0.005;
+		BilinearThreshold			= 0.02;
+		ShadowContrast				= 4;
+		IgnoreEdgePixels			= false;
+		UsePrecisionOffset			= false;
+		BilinearSamplingOffsetMode	= false;
+		DebugOutputEdgeMask			= false;
+		DebugOutputThreadIndex		= false;
+		DebugOutputWaveIndex		= false;
+		DepthBounds					= float2(0,1);
+		UseEarlyOut					= false;
+	}
+
+	// Runtime data returned from BuildDispatchList():
+	float4 LightCoordinate;				// Values stored in DispatchList::LightCoordinate_Shader by BuildDispatchList()
+	int2 WaveOffset;					// Values stored in DispatchData::WaveOffset_Shader by BuildDispatchList()
+
+	// Renderer Specific Values:
+	float FarDepthValue;				// Set to the Depth Buffer Value for the far clip plane, as determined by renderer projection matrix setup (typically 0).
+	float NearDepthValue;				// Set to the Depth Buffer Value for the near clip plane, as determined by renderer projection matrix setup (typically 1).
+
+	// Sampling data:
+	float2 InvDepthTextureSize;			// Inverse of the texture dimensions for 'DepthTexture' (used to convert from pixel coordinates to UVs)
+										// If 'PointBorderSampler' is an Unnormalized sampler, then this value can be hard-coded to 1.
+										// The 'USE_HALF_PIXEL_OFFSET' macro might need to be defined if sampling at exact pixel coordinates isn't precise (e.g., if odd patterns appear in the shadow).
+
+	Texture2D<float> DepthTexture;		// Depth Buffer Texture (rasterized non-linear depth)
+	RWTexture2D<float> OutputTexture;	// Output screen-space shadow buffer (typically single-channel, 8bit)
+
+	SamplerState PointBorderSampler;	// A point sampler, with Wrap Mode set to Clamp-To-Border-Color (D3D12_TEXTURE_ADDRESS_MODE_BORDER), and Border Color set to "FarDepthValue" (typically zero), or some other far-depth value out of DepthBounds.
+										// If you have issues where invalid shadows are appearing from off-screen, it is likely that this sampler is not correctly setup
+};
+
+
+#if !defined(WAVE_SIZE) || !defined(SAMPLE_COUNT) || !defined(HARD_SHADOW_SAMPLES) || !defined(FADE_OUT_SAMPLES)
+	#error Before including bend_sss_gpu.h, four macros must be defined to configure the shader compile: WAVE_SIZE, SAMPLE_COUNT, HARD_SHADOW_SAMPLES, and FADE_OUT_SAMPLES. See the top of this file for details.
+#else
+
+	static bool EarlyOutPixel(struct DispatchParameters inParameters, int2 pixel_xy, float depth)
+	{
+		//OPTIONAL TODO; customize this function to return true if the pixel should early-out for custom reasons. E.g., A shadow map pass already found the pixel was in shadow / backfaced, etc.
+		// Recommended to keep this code very simple!
+
+		// Example:
+		// return inParameters.CustomShadowMapTerm[pixel_xy] == 0;
+
+#if 0 // disabled to workaround compilation issues on some platforms
+		(void)pixel_xy;	//unused by this implementation, avoid potential compiler warning.
+#endif
+
+		// The compiled code will be more optimal if the 'depth' value is not referenced.
+		return depth >= inParameters.DepthBounds.y || depth <= inParameters.DepthBounds.x;
+	}
+
+	// Gets the start pixel coordinates for the pixels in the wavefront
+	// Also returns the delta to get to the next pixel after WAVE_COUNT pixels along the ray
+	static void ComputeWavefrontExtents(DispatchParameters inParameters, int3 inGroupID, uint inGroupThreadID, out float2 outDeltaXY, out float2 outPixelXY, out float outPixelDistance, out bool outMajorAxisX)
+	{
+		int2 xy = inGroupID.yz * WAVE_SIZE + inParameters.WaveOffset.xy;
+
+		//integer light position / fractional component
+		float2 light_xy = floor(inParameters.LightCoordinate.xy) + 0.5;
+		float2 light_xy_fraction = inParameters.LightCoordinate.xy - light_xy;
+		bool reverse_direction = inParameters.LightCoordinate.w > 0.0f;
+
+		int2 sign_xy = sign(xy);
+		bool horizontal = abs(xy.x + sign_xy.y) < abs(xy.y - sign_xy.x);
+
+		int2 axis;
+		axis.x = horizontal ? (+sign_xy.y) : (0);
+		axis.y = horizontal ? (0) : (-sign_xy.x);
+
+		// Apply wave offset
+		xy = axis * (int)inGroupID.x + xy;
+		float2 xy_f = (float2)xy;
+
+		// For interpolation to the light center, we only really care about the larger of the two axis
+		bool x_axis_major = abs(xy_f.x) > abs(xy_f.y);
+		float major_axis = x_axis_major ? xy_f.x : xy_f.y;
+
+		float major_axis_start = abs(major_axis);
+		float major_axis_end = abs(major_axis) - (float)WAVE_SIZE;
+
+		float ma_light_frac = x_axis_major ? light_xy_fraction.x : light_xy_fraction.y;
+		ma_light_frac = major_axis > 0 ? -ma_light_frac : ma_light_frac;
+
+		// back in to screen direction
+		float2 start_xy = xy_f + light_xy;
+
+		// For the very inner most ring, we need to interpolate to a pixel centered UV, so the UV->pixel rounding doesn't skip output pixels
+		float2 end_xy = lerp(inParameters.LightCoordinate.xy, start_xy, (major_axis_end + ma_light_frac) / (major_axis_start + ma_light_frac));
+
+		// The major axis should be a round number
+		float2 xy_delta = (start_xy - end_xy);
+
+		// Inverse the read order when reverse direction is true
+		float thread_step = (float)(inGroupThreadID ^ (reverse_direction ? 0 : (WAVE_SIZE - 1)));
+
+		float2 pixel_xy = lerp(start_xy, end_xy, thread_step / (float)WAVE_SIZE);
+		float pixel_distance = major_axis_start - thread_step + ma_light_frac;
+
+		outPixelXY = pixel_xy;
+		outPixelDistance = pixel_distance;
+		outDeltaXY = xy_delta;
+		outMajorAxisX = x_axis_major;
+	}
+
+	// Number of bilinear sample reads performed per-thread
+	#define READ_COUNT (SAMPLE_COUNT / WAVE_SIZE + 2)
+
+	// Common shared data
+	groupshared float DepthData[READ_COUNT * WAVE_SIZE];
+	groupshared bool LdsEarlyOut;
+
+	// Generate the shadow
+	//	Call this function from a compute shader with thread dimensions: numthreads[WAVE_SIZE, 1, 1]
+	// 
+	//	(int3)	inGroupID:			Compute shader group id register (SV_GroupID)
+	//	(int)	inGroupThreadId:	Compute shader group thread id register (SV_GroupThreadID)
+	void WriteScreenSpaceShadow(DispatchParameters inParameters, int3 inGroupID, int inGroupThreadID)
+	{
+		float2 xy_delta;
+		float2 pixel_xy;
+		float pixel_distance;
+		bool x_axis_major;	// major axis is x axis? abs(xy_delta.x) > abs(xy_delta.y).
+
+		ComputeWavefrontExtents(inParameters, (int3)inGroupID, inGroupThreadID.x, xy_delta, pixel_xy, pixel_distance, x_axis_major);
+
+		// Read in the depth values
+		float sampling_depth[READ_COUNT];
+		float shadowing_depth[READ_COUNT];
+		float depth_thickness_scale[READ_COUNT];
+		float sample_distance[READ_COUNT];
+
+		const float direction = -inParameters.LightCoordinate.w;
+		const float z_sign = inParameters.NearDepthValue > inParameters.FarDepthValue ? -1 : +1;
+
+		int i;
+		bool is_edge = false;
+		bool skip_pixel = false;
+		float2 write_xy = floor(pixel_xy);
+
+		[unroll] for (i = 0; i < READ_COUNT; i++)
+		{
+			// We sample depth twice per pixel per sample, and interpolate with an edge detect filter
+			// Interpolation should only occur on the minor axis of the ray - major axis coordinates should be at pixel centers
+			float2 read_xy = floor(pixel_xy);
+			float minor_axis = x_axis_major ? pixel_xy.y : pixel_xy.x;
+
+			// If a pixel has been detected as an edge, then optionally (inParameters.IgnoreEdgePixels) don't include it in the shadow
+			const float edge_skip = 1e20;	// if edge skipping is enabled, apply an extreme value/blend on edge samples to push the value out of range
+
+			float2 depths;
+			float bilinear = frac(minor_axis) - 0.5;
+
+		#if USE_HALF_PIXEL_OFFSET
+			read_xy += 0.5;
+		#endif
+
+		#if USE_UV_PIXEL_BIAS
+			float bias = bilinear > 0 ? 1 : -1;
+			float2 offset_xy = float2(x_axis_major ? 0 : bias, x_axis_major ? bias : 0);
+
+			// HLSL enforces that a pixel offset is a compile-time constant, which isn't strictly required (and can sometimes be a bit faster)
+			// So this fallback will use a manual uv offset instead
+			depths.x = inParameters.DepthTexture.SampleLevel(inParameters.PointBorderSampler, read_xy * inParameters.InvDepthTextureSize, 0);
+			depths.y = inParameters.DepthTexture.SampleLevel(inParameters.PointBorderSampler, (read_xy + offset_xy) * inParameters.InvDepthTextureSize, 0);
+		#else
+			int bias = bilinear > 0 ? 1 : -1;
+			int2 offset_xy = int2(x_axis_major ? 0 : bias, x_axis_major ? bias : 0);
+
+			depths.x = inParameters.DepthTexture.SampleLevel(inParameters.PointBorderSampler, read_xy * inParameters.InvDepthTextureSize, 0);
+			depths.y = inParameters.DepthTexture.SampleLevel(inParameters.PointBorderSampler, read_xy * inParameters.InvDepthTextureSize, 0, offset_xy);
+		#endif
+
+			// Depth thresholds (bilinear/shadow thickness) are based on a fractional ratio of the difference between sampled depth and the far clip depth
+			depth_thickness_scale[i] = abs(inParameters.FarDepthValue - depths.x);
+
+			// If depth variance is more than a specific threshold, then just use point filtering
+			bool use_point_filter = abs(depths.x - depths.y) > depth_thickness_scale[i] * inParameters.BilinearThreshold;
+
+			// Store for debug output when inParameters.DebugOutputEdgeMask is true
+			if (i == 0) is_edge = use_point_filter;
+
+			if (inParameters.BilinearSamplingOffsetMode)
+			{
+				bilinear = use_point_filter ? 0 : bilinear;
+				//both shadow depth and starting depth are the same in this mode, unless shadow skipping edges
+				sampling_depth[i] = lerp(depths.x, depths.y, abs(bilinear));
+				shadowing_depth[i] = (inParameters.IgnoreEdgePixels && use_point_filter) ? edge_skip : sampling_depth[i];
+			}
+			else
+			{
+				// The pixel starts sampling at this depth
+				sampling_depth[i] = depths.x;
+
+				float edge_depth = inParameters.IgnoreEdgePixels ? edge_skip : depths.x;
+				// Any sample in this wavefront is possibly interpolated towards the bilinear sample
+				// So use should use a shadowing depth that is further away, based on the difference between the two samples
+				float shadow_depth = depths.x + abs(depths.x - depths.y) * z_sign;
+
+				// Shadows cast from this depth
+				shadowing_depth[i] = use_point_filter ? edge_depth : shadow_depth;
+			}
+
+			// Store for later
+			sample_distance[i] = pixel_distance + (WAVE_SIZE * i) * direction;
+
+			// Iterate to the next pixel along the ray. This will be WAVE_SIZE pixels along the ray...
+			pixel_xy += xy_delta * direction;
+		}
+
+		// Using early out, and no debug mode is enabled?
+		if (inParameters.UseEarlyOut && (inParameters.DebugOutputWaveIndex == false && inParameters.DebugOutputThreadIndex == false && inParameters.DebugOutputEdgeMask == false))
+		{
+			// read the depth of the pixel we are shadowing, and early-out
+			// The compiler will typically rearrange this code to put it directly after the first depth read
+			skip_pixel = EarlyOutPixel(inParameters, (int2)write_xy, sampling_depth[0]);
+
+			// are all threads in this wave out of bounds?
+			bool early_out = WaveActiveAnyTrue(!skip_pixel) == false;
+
+			// WaveGetLaneCount returns the hardware wave size
+			if (WaveGetLaneCount() == WAVE_SIZE)
+			{
+				// Optimal case:
+				// If each wavefront is just a single wave, then we can trivially early-out.
+				if (early_out == true)
+					return;
+			}
+			else
+			{
+				// This wavefront is made up of multiple small waves, so we need to coordinate them for all to early-out together.
+				// Doing this can make the worst case (all pixels drawn) a bit more expensive (~15%), but the best-case (all early-out) is typically 2-3x better.
+				LdsEarlyOut = true;
+
+				GroupMemoryBarrierWithGroupSync();
+
+				[branch] if (early_out == false)
+					LdsEarlyOut = false;
+
+				GroupMemoryBarrierWithGroupSync();
+
+				[branch] if (LdsEarlyOut)
+					return;
+			}
+		}
+
+		// Write the shadow depths to LDS
+		[unroll] for (i = 0; i < READ_COUNT; i++)
+		{
+			// Perspective correct the shadowing depth, in this space, all light rays are parallel
+			float stored_depth = (shadowing_depth[i] - inParameters.LightCoordinate.z) / sample_distance[i];
+
+			if (i != 0)
+			{
+				// For pixels within sampling distance of the light, it is possible that sampling will
+				// overshoot the light coordinate for extended reads. We want to ignore these samples
+				stored_depth = sample_distance[i] > 0 ? stored_depth : 1e10;
+			}
+
+			// Store the depth values in groupshared
+			int idx = (i * WAVE_SIZE) + inGroupThreadID.x;
+			DepthData[idx] = stored_depth;
+		}
+
+		// Sync wavefronts now groupshared DepthData is written
+		GroupMemoryBarrierWithGroupSync();
+
+		// If the starting depth isn't in depth bounds, then we don't need a shadow
+		if (skip_pixel)
+			return;
+
+		float start_depth = sampling_depth[0];
+
+		// lerp away from far depth by a tiny fraction?
+		if (inParameters.UsePrecisionOffset)
+			start_depth = lerp(start_depth, inParameters.FarDepthValue, -1.0 / 0xFFFF);
+
+		// perspective correct the depth
+		start_depth = (start_depth - inParameters.LightCoordinate.z) / sample_distance[0];
+
+		// Start by reading the next value
+		int sample_index = inGroupThreadID.x + 1;
+
+		float4 shadow_value = 1;
+		float hard_shadow = 1;
+
+		// This is the inverse of how large the shadowing window is for the projected sample data. 
+		// All values in the LDS sample list are scaled by 1.0 / sample_distance, such that all light directions become parallel.
+		// The multiply by sample_distance[0] here is to compensate for the projection divide in the data.
+		// The 1.0 / inParameters.SurfaceThickness is to adjust user selected thickness. So a 0.5% thickness will scale depth values from [0,1] to [0,200]. The shadow window is always 1 wide.
+		// 1.0 / depth_thickness_scale[0] is because SurfaceThickness is percentage of remaining depth between the sample and the far clip - not a percentage of the full depth range.
+		// The min() function is to make sure the window is a minimum width when very close to the light. The +direction term will bias the result so the pixel at the very center of the light is either fully lit or shadowed
+		float depth_scale = min(sample_distance[0] + direction, 1.0 / inParameters.SurfaceThickness) * sample_distance[0] / depth_thickness_scale[0];
+
+		start_depth = start_depth * depth_scale - z_sign;
+
+		// The first number of hard shadow samples, a single pixel can produce a full shadow
+		[unroll] for (i = 0; i < HARD_SHADOW_SAMPLES; i++)
+		{
+			float depth_delta = abs(start_depth - DepthData[sample_index + i] * depth_scale);
+
+			// We want to find the distance of the sample that is closest to the reference depth
+			hard_shadow = min(hard_shadow, depth_delta);
+		}
+
+		// Brute force go!
+		// The main shadow samples, averaged in to a set of 4 shadow values
+		[unroll] for (i = HARD_SHADOW_SAMPLES; i < SAMPLE_COUNT - FADE_OUT_SAMPLES; i++)
+		{
+			float depth_delta = abs(start_depth - DepthData[sample_index + i] * depth_scale);
+
+			// Do the same as the hard_shadow code above, but this will accumulate to 4 separate values.
+			// By using 4 values, the average shadow can be taken, which can help soften single-pixel shadows.
+			shadow_value[i & 3] = min(shadow_value[i & 3], depth_delta);
+		}
+
+		// Final fade out samples
+		[unroll] for (i = SAMPLE_COUNT - FADE_OUT_SAMPLES; i < SAMPLE_COUNT; i++)
+		{
+			float depth_delta = abs(start_depth - DepthData[sample_index + i] * depth_scale);
+
+			// Add the fade value to these samples
+			const float fade_out = (float)(i + 1 - (SAMPLE_COUNT - FADE_OUT_SAMPLES)) / (float)(FADE_OUT_SAMPLES + 1) * 0.75;
+
+			shadow_value[i & 3] = min(shadow_value[i & 3], depth_delta + fade_out);
+		}
+
+		// Apply the contrast value.
+		// A value of 0 indicates a sample was exactly matched to the reference depth (and the result is fully shadowed)
+		// We want some boost to this range, so samples don't have to exactly match to produce a full shadow. 
+		shadow_value = saturate(shadow_value * (inParameters.ShadowContrast) + (1 - inParameters.ShadowContrast));
+		hard_shadow = saturate(hard_shadow * (inParameters.ShadowContrast) + (1 - inParameters.ShadowContrast));
+
+		float result = 0;
+
+		// Take the average of 4 samples, this is useful to reduces aliasing noise in the source depth, especially with long shadows.
+		result = dot(shadow_value, 0.25);
+
+		// If the first samples are always producing a hard shadow, then compute this value separately.
+		result = min(hard_shadow, result);
+
+		//write the result
+		{
+			if (inParameters.DebugOutputEdgeMask)
+				result = is_edge ? 1 : 0;
+			if (inParameters.DebugOutputThreadIndex)
+				result = (inGroupThreadID / (float)WAVE_SIZE);
+			if (inParameters.DebugOutputWaveIndex)			
+				result = frac(inGroupID.x / (float)WAVE_SIZE);
+			
+			// Asking the GPU to write scattered single-byte pixels isn't great,
+			// But thankfully the latency is hidden by all the work we're doing...
+			inParameters.OutputTexture[(int2)write_xy] = result;
+		}
+	}
+
+#endif // macro check
\ No newline at end of file
diff --git a/Engine/Source/Runtime/Renderer/Private/Shadows/ScreenSpaceShadows.cpp b/Engine/Source/Runtime/Renderer/Private/Shadows/ScreenSpaceShadows.cpp
--- a/Engine/Source/Runtime/Renderer/Private/Shadows/ScreenSpaceShadows.cpp
+++ b/Engine/Source/Runtime/Renderer/Private/Shadows/ScreenSpaceShadows.cpp
@@ -15,6 +15,24 @@
 #include "RenderGraphUtils.h"
 #include "PixelShaderUtils.h"
 
+#include "bend_sss_cpu.h"
+
+enum class ContactShadowsMethod
+{
+	StochasticJittering = 0,
+	BendSSS = 1,
+};
+
+static int32 GContactShadowsMethod = 0;
+static FAutoConsoleVariableRef CVarContactShadowsMethod(
+	TEXT("r.ContactShadows.Standalone.Method"),
+	GContactShadowsMethod,
+	TEXT("Technique to use to calculate Contact (Screen Space) Shadows:\n")
+	TEXT("0 - Stochastic Jittering.\n")
+	TEXT("1 - Bend Screen Space Shadows."),
+	ECVF_Scalability | ECVF_RenderThreadSafe
+);
+
 extern void GetLightContactShadowParameters(const FLightSceneProxy* Proxy, float& OutLength, bool& bOutLengthInWS, float& OutCastingIntensity, float& OutNonCastingIntensity);
 
 const int32 GScreenSpaceShadowsTileSizeX = 8;
@@ -63,6 +81,45 @@ class FScreenSpaceShadowsCS : public FGlobalShader
 
 IMPLEMENT_GLOBAL_SHADER(FScreenSpaceShadowsCS, "/Engine/Private/ScreenSpaceShadows.usf", "ScreenSpaceShadowsCS", SF_Compute);
 
+class FScreenSpaceShadowsBendCS : public FGlobalShader
+{
+	DECLARE_GLOBAL_SHADER(FScreenSpaceShadowsBendCS);
+	SHADER_USE_PARAMETER_STRUCT(FScreenSpaceShadowsBendCS, FGlobalShader);
+
+	BEGIN_SHADER_PARAMETER_STRUCT(FParameters, )
+		SHADER_PARAMETER_RDG_TEXTURE_UAV(RWTexture2D<float>, OutputTexture)
+		SHADER_PARAMETER_RDG_TEXTURE(Texture2D<float>, DepthTexture)
+		SHADER_PARAMETER_SAMPLER(SamplerState, PointBorderSampler)
+		SHADER_PARAMETER_STRUCT_REF(FViewUniformShaderParameters, View)
+		SHADER_PARAMETER(FVector3f, LightDirection)
+		SHADER_PARAMETER(float, ContactShadowLength)
+		SHADER_PARAMETER(uint32, bContactShadowLengthInWS)
+		SHADER_PARAMETER(float, ContactShadowCastingIntensity)
+		SHADER_PARAMETER(FIntRect, ScissorRectMinAndSize)
+		SHADER_PARAMETER(uint32, DownsampleFactor)
+		SHADER_PARAMETER(FVector2f, InvDepthTextureSize)
+		SHADER_PARAMETER(FVector4f, LightCoordinate)
+		SHADER_PARAMETER(FIntVector, WaveOffset)
+	END_SHADER_PARAMETER_STRUCT()
+
+	static bool ShouldCompilePermutation(const FGlobalShaderPermutationParameters& Parameters)
+	{
+		return true;
+	}
+
+	static void ModifyCompilationEnvironment(const FGlobalShaderPermutationParameters& Parameters, FShaderCompilerEnvironment& OutEnvironment)
+	{
+		OutEnvironment.SetDefine(TEXT("THREADGROUP_SIZEX"), 64);
+		OutEnvironment.SetDefine(TEXT("THREADGROUP_SIZEY"), 1);
+		OutEnvironment.SetDefine(TEXT("FORCE_DEPTH_TEXTURE_READS"), 1);
+		OutEnvironment.SetDefine(TEXT("PLATFORM_SUPPORTS_TYPED_UAV_LOAD"), (int32)RHISupports4ComponentUAVReadWrite(Parameters.Platform));
+
+		OutEnvironment.CompilerFlags.Add(CFLAG_ForceDXC);
+	}
+};
+
+IMPLEMENT_GLOBAL_SHADER(FScreenSpaceShadowsBendCS, "/Engine/Private/ScreenSpaceShadows.usf", "ScreenSpaceShadowsBendCS", SF_Compute);
+
 class FScreenSpaceShadowsUpsamplePS : public FGlobalShader
 {
 	DECLARE_GLOBAL_SHADER(FScreenSpaceShadowsUpsamplePS);
@@ -99,6 +156,70 @@ BEGIN_SHADER_PARAMETER_STRUCT(FScreenSpaceShadowsUpsample, )
 	RENDER_TARGET_BINDING_SLOTS()
 END_SHADER_PARAMETER_STRUCT()
 
+void UpsampleScreenSpaceShadows(
+	FRDGBuilder& GraphBuilder,
+	const FMinimalSceneTextures& SceneTextures,
+	const FViewInfo& View,
+	FIntRect ScissorRect,
+	bool bProjectingForForwardShading,
+	const FLightSceneInfo* LightSceneInfo,
+	FRDGTextureRef ShadowsTexture,
+	int32 DownsampleFactor,
+	FRDGTextureRef ScreenShadowMaskTexture
+)
+{
+	FScreenSpaceShadowsUpsample* PassParameters = GraphBuilder.AllocParameters<FScreenSpaceShadowsUpsample>();
+	PassParameters->RenderTargets[0] = FRenderTargetBinding(ScreenShadowMaskTexture, ERenderTargetLoadAction::ELoad);
+	PassParameters->RenderTargets.DepthStencil = FDepthStencilBinding(SceneTextures.Depth.Target, ERenderTargetLoadAction::ELoad, ERenderTargetLoadAction::ELoad, FExclusiveDepthStencil::DepthRead_StencilRead);
+
+	PassParameters->PS.View = GetShaderBinding(View.ViewUniformBuffer);
+	PassParameters->PS.SceneTextures = SceneTextures.GetSceneTextureShaderParameters(View.GetFeatureLevel());
+	PassParameters->PS.ShadowFactorsTexture = ShadowsTexture;
+	PassParameters->PS.ShadowFactorsSampler = TStaticSamplerState<SF_Bilinear>::GetRHI();
+	PassParameters->PS.ScissorRectMinAndSize = FIntRect(ScissorRect.Min, ScissorRect.Size());
+	PassParameters->PS.OneOverDownsampleFactor = 1.0f / DownsampleFactor;
+
+	FScreenSpaceShadowsUpsamplePS::FPermutationDomain PermutationVector;
+	PermutationVector.Set<FScreenSpaceShadowsUpsamplePS::FUpsample>(DownsampleFactor != 1);
+	auto PixelShader = View.ShaderMap->GetShader<FScreenSpaceShadowsUpsamplePS>(PermutationVector);
+
+	// blend separately from CSM / DF Shadows since those interact with static lighting
+	// this matches behavior of GetShadowTerms(...) 
+	const bool bIsWholeSceneDirectionalShadow = false;
+
+	FRHIBlendState* BlendState = FProjectedShadowInfo::GetBlendStateForProjection(
+		LightSceneInfo->GetDynamicShadowMapChannel(),
+		bIsWholeSceneDirectionalShadow,
+		false,
+		bProjectingForForwardShading,
+		false);
+
+	ClearUnusedGraphResources(PixelShader, &PassParameters->PS);
+
+	GraphBuilder.AddPass(
+		RDG_EVENT_NAME("Upsample"),
+		PassParameters,
+		ERDGPassFlags::Raster,
+		[PassParameters, &View, PixelShader, BlendState, ScissorRect](FRHICommandList& RHICmdList)
+		{
+			RHICmdList.SetViewport(ScissorRect.Min.X, ScissorRect.Min.Y, 0.0f, ScissorRect.Max.X, ScissorRect.Max.Y, 1.0f);
+			RHICmdList.SetScissorRect(true, ScissorRect.Min.X, ScissorRect.Min.Y, ScissorRect.Max.X, ScissorRect.Max.Y);
+
+			FGraphicsPipelineStateInitializer GraphicsPSOInit;
+			FPixelShaderUtils::InitFullscreenPipelineState(RHICmdList, View.ShaderMap, PixelShader, GraphicsPSOInit);
+
+			GraphicsPSOInit.BoundShaderState.PixelShaderRHI = PixelShader.GetPixelShader();
+			GraphicsPSOInit.BlendState = BlendState;
+
+			SetGraphicsPipelineState(RHICmdList, GraphicsPSOInit, 0);
+
+			SetShaderParameters(RHICmdList, PixelShader, PixelShader.GetPixelShader(), PassParameters->PS);
+
+			FPixelShaderUtils::DrawFullscreenTriangle(RHICmdList);
+			RHICmdList.SetScissorRect(false, 0, 0, 0, 0);
+		});
+}
+
 void RenderScreenSpaceShadows(
 	FRDGBuilder& GraphBuilder,
 	bool bAsyncCompute,
@@ -155,58 +276,110 @@ void RenderScreenSpaceShadows(
 			FIntVector(GroupSizeX, GroupSizeY, 1));
 	}
 
+	UpsampleScreenSpaceShadows(
+		GraphBuilder,
+		SceneTextures,
+		View,
+		ScissorRect,
+		bProjectingForForwardShading,
+		LightSceneInfo,
+		ShadowsTexture,
+		GetScreenSpaceShadowDownsampleFactor(),
+		ScreenShadowMaskTexture);
+}
+
+void RenderScreenSpaceShadowsBend(
+	FRDGBuilder& GraphBuilder,
+	bool bAsyncCompute,
+	const FMinimalSceneTextures& SceneTextures,
+	const FViewInfo& View,
+	FIntRect ScissorRect,
+	bool bProjectingForForwardShading,
+	const FLightSceneInfo* LightSceneInfo,
+	FRDGTextureRef ScreenShadowMaskTexture)
+{
+	check(ScissorRect.Area() > 0);
+
+	const int32 DownsampleFactor = 1;
+
+	const FIntPoint BufferSize = View.GetSceneTexturesConfig().Extent;
+	FRDGTextureDesc Desc(FRDGTextureDesc::Create2D(BufferSize, PF_R16F, FClearValueBinding::None, TexCreate_UAV | TexCreate_ShaderResource));
+	FRDGTextureRef ShadowsTexture = GraphBuilder.CreateTexture(Desc, TEXT("ScreenSpaceShadows"));
+
+	const FRDGTextureDesc& DepthDesc = SceneTextures.Depth.Resolve->Desc;
+
 	{
-		FScreenSpaceShadowsUpsample* PassParameters = GraphBuilder.AllocParameters<FScreenSpaceShadowsUpsample>();
-		PassParameters->RenderTargets[0] = FRenderTargetBinding(ScreenShadowMaskTexture, ERenderTargetLoadAction::ELoad);
-		PassParameters->RenderTargets.DepthStencil = FDepthStencilBinding(SceneTextures.Depth.Target, ERenderTargetLoadAction::ELoad, ERenderTargetLoadAction::ELoad, FExclusiveDepthStencil::DepthRead_StencilRead);
-
-		PassParameters->PS.View = GetShaderBinding(View.ViewUniformBuffer);
-		PassParameters->PS.SceneTextures = SceneTextures.GetSceneTextureShaderParameters(View.GetFeatureLevel());
-		PassParameters->PS.ShadowFactorsTexture = ShadowsTexture;
-		PassParameters->PS.ShadowFactorsSampler = TStaticSamplerState<SF_Bilinear>::GetRHI();
-		PassParameters->PS.ScissorRectMinAndSize = FIntRect(ScissorRect.Min, ScissorRect.Size());
-		PassParameters->PS.OneOverDownsampleFactor = 1.0f / GetScreenSpaceShadowDownsampleFactor();
-
-		FScreenSpaceShadowsUpsamplePS::FPermutationDomain PermutationVector;
-		PermutationVector.Set<FScreenSpaceShadowsUpsamplePS::FUpsample>(GetScreenSpaceShadowDownsampleFactor() != 1);
-		auto PixelShader = View.ShaderMap->GetShader<FScreenSpaceShadowsUpsamplePS>(PermutationVector);
-
-		// blend separately from CSM / DF Shadows since those interact with static lighting
-		// this matches behavior of GetShadowTerms(...) 
-		const bool bIsWholeSceneDirectionalShadow = false;
-
-		FRHIBlendState* BlendState = FProjectedShadowInfo::GetBlendStateForProjection(
-			LightSceneInfo->GetDynamicShadowMapChannel(),
-			bIsWholeSceneDirectionalShadow,
-			false,
-			bProjectingForForwardShading,
-			false);
-
-		ClearUnusedGraphResources(PixelShader, &PassParameters->PS);
-
-		GraphBuilder.AddPass(
-			RDG_EVENT_NAME("Upsample"),
-			PassParameters,
-			ERDGPassFlags::Raster,
-			[PassParameters, &View, PixelShader, BlendState, ScissorRect](FRHICommandList& RHICmdList)
-			{
-				RHICmdList.SetViewport(ScissorRect.Min.X, ScissorRect.Min.Y, 0.0f, ScissorRect.Max.X, ScissorRect.Max.Y, 1.0f);
-				RHICmdList.SetScissorRect(true, ScissorRect.Min.X, ScissorRect.Min.Y, ScissorRect.Max.X, ScissorRect.Max.Y);
+		const FLightSceneProxy* LightProxy = LightSceneInfo->Proxy;
+
+		const FMatrix ViewProjection = View.ViewMatrices.GetViewProjectionMatrix();
+		const FVector4 LightDirection4 = FVector4(-LightProxy->GetDirection(), 0.0f);
+		const FVector4 LightDirection4Clip = ViewProjection.TransformFVector4(LightDirection4);
+
+		FVector4f LightProjection = (FVector4f)LightDirection4Clip;
+		FIntPoint MinRenderBounds = ScissorRect.Min;
+		FIntPoint MaxRenderBounds = ScissorRect.Max;
+		Bend::DispatchList DispatchList = Bend::BuildDispatchList((float*)&LightProjection, (int32*)&BufferSize, (int32*)&MinRenderBounds, (int32*)&MaxRenderBounds);
+
+		for (int32 DispatchIndex = 0; DispatchIndex < DispatchList.DispatchCount; ++DispatchIndex)
+		{
+			const Bend::DispatchData& Dispatch = DispatchList.Dispatch[DispatchIndex];
+
+			FScreenSpaceShadowsBendCS::FParameters* PassParameters = GraphBuilder.AllocParameters<FScreenSpaceShadowsBendCS::FParameters>();
+
+			PassParameters->OutputTexture = GraphBuilder.CreateUAV(ShadowsTexture);
+
+			PassParameters->DepthTexture = SceneTextures.Depth.Resolve;
+			PassParameters->InvDepthTextureSize = FVector2f(1.0f / DepthDesc.Extent.X, 1.0f / DepthDesc.Extent.Y);
+
+			// A point sampler, with Wrap Mode set to Clamp-To-Border-Color (D3D12_TEXTURE_ADDRESS_MODE_BORDER), and Border Color set to "FarDepthValue" (typically zero), or some other far-depth value out of DepthBounds.
+			// If you have issues where invalid shadows are appearing from off-screen, it is likely that this sampler is not correctly setup
+			PassParameters->PointBorderSampler = TStaticSamplerState<SF_Point, AM_Border, AM_Border, AM_Border, 0, 1, 0>::GetRHI();
 
-				FGraphicsPipelineStateInitializer GraphicsPSOInit;
-				FPixelShaderUtils::InitFullscreenPipelineState(RHICmdList, View.ShaderMap, PixelShader, GraphicsPSOInit);
+			PassParameters->View = View.ViewUniformBuffer;
 
-				GraphicsPSOInit.BoundShaderState.PixelShaderRHI = PixelShader.GetPixelShader();
-				GraphicsPSOInit.BlendState = BlendState;
+			PassParameters->ScissorRectMinAndSize = FIntRect(ScissorRect.Min, ScissorRect.Size());
+			PassParameters->DownsampleFactor = DownsampleFactor;
 
-				SetGraphicsPipelineState(RHICmdList, GraphicsPSOInit, 0);
+			FLightRenderParameters LightParameters;
+			LightProxy->GetLightShaderParameters(LightParameters);
 
-				SetShaderParameters(RHICmdList, PixelShader, PixelShader.GetPixelShader(), PassParameters->PS);
+			PassParameters->LightDirection = LightParameters.Direction;
 
-				FPixelShaderUtils::DrawFullscreenTriangle(RHICmdList);
-				RHICmdList.SetScissorRect(false, 0, 0, 0, 0);
-			});
+			float ContactShadowLength;
+			bool bContactShadowLengthInWS;
+			float ContactShadowCastingIntensity;
+			float ContactShadowNonCastingIntensity;
+			GetLightContactShadowParameters(LightProxy, ContactShadowLength, bContactShadowLengthInWS, ContactShadowCastingIntensity, ContactShadowNonCastingIntensity);
+
+			PassParameters->ContactShadowLength = ContactShadowLength;
+			PassParameters->bContactShadowLengthInWS = bContactShadowLengthInWS;
+			PassParameters->ContactShadowCastingIntensity = ContactShadowCastingIntensity;
+
+			PassParameters->LightCoordinate = FVector4f(DispatchList.LightCoordinate_Shader[0], DispatchList.LightCoordinate_Shader[1], DispatchList.LightCoordinate_Shader[2], DispatchList.LightCoordinate_Shader[3]);
+			PassParameters->WaveOffset = FIntVector(Dispatch.WaveOffset_Shader[0], Dispatch.WaveOffset_Shader[1], 0);
+
+			auto ComputeShader = View.ShaderMap->GetShader<FScreenSpaceShadowsBendCS>();
+
+			FComputeShaderUtils::AddPass(
+				GraphBuilder,
+				RDG_EVENT_NAME("ScreenSpaceShadowing (Bend)"),
+				bAsyncCompute ? ERDGPassFlags::AsyncCompute : ERDGPassFlags::Compute,
+				ComputeShader,
+				PassParameters,
+				FIntVector(Dispatch.WaveCount[0], Dispatch.WaveCount[1], Dispatch.WaveCount[2]));
+		}
 	}
+
+	UpsampleScreenSpaceShadows(
+		GraphBuilder,
+		SceneTextures,
+		View,
+		ScissorRect,
+		bProjectingForForwardShading,
+		LightSceneInfo,
+		ShadowsTexture,
+		DownsampleFactor,
+		ScreenShadowMaskTexture);
 }
 
 void RenderScreenSpaceShadows(
@@ -247,15 +420,30 @@ void RenderScreenSpaceShadows(
 
 		if (ScissorRect.Area() > 0)
 		{
-			RenderScreenSpaceShadows(
-				GraphBuilder,
-				false,
-				SceneTextures,
-				View,
-				ScissorRect,
-				bProjectingForForwardShading,
-				LightSceneInfo,
-				ScreenShadowMaskTexture);
+			if (GContactShadowsMethod == (uint32)ContactShadowsMethod::BendSSS)
+			{
+				RenderScreenSpaceShadowsBend(
+					GraphBuilder,
+					false,
+					SceneTextures,
+					View,
+					ScissorRect,
+					bProjectingForForwardShading,
+					LightSceneInfo,
+					ScreenShadowMaskTexture);
+			}
+			else
+			{
+				RenderScreenSpaceShadows(
+					GraphBuilder,
+					false,
+					SceneTextures,
+					View,
+					ScissorRect,
+					bProjectingForForwardShading,
+					LightSceneInfo,
+					ScreenShadowMaskTexture);
+			}
 		}
 	}
 }
\ No newline at end of file
diff --git a/Engine/Source/Runtime/Renderer/Private/Shadows/bend_sss_cpu.h b/Engine/Source/Runtime/Renderer/Private/Shadows/bend_sss_cpu.h
new file mode 100644
--- /dev/null
+++ b/Engine/Source/Runtime/Renderer/Private/Shadows/bend_sss_cpu.h
@@ -0,0 +1,256 @@
+#pragma once
+
+// Copyright 2023 Sony Interactive Entertainment.
+//
+// Licensed under the Apache License, Version 2.0 (the "License");
+// you may not use this file except in compliance with the License.
+// You may obtain a copy of the License at
+//
+//     http://www.apache.org/licenses/LICENSE-2.0
+//
+// Unless required by applicable law or agreed to in writing, software
+// distributed under the License is distributed on an "AS IS" BASIS,
+// WITHOUT WARRANTIES OR CONDITIONS OF ANY KIND, either express or implied.
+// See the License for the specific language governing permissions and
+// limitations under the License.
+
+// If you have feedback, or found this code useful, we'd love to hear from you.
+// https://www.bendstudio.com
+// https://www.twitter.com/bendstudio
+// 
+// We are *always* looking for talented graphics and technical programmers!
+// https://www.bendstudio.com/careers
+
+// Common screen space shadow projection code (CPU):
+//--------------------------------------------------------------
+
+namespace Bend
+{
+	// Generating a screen-space-shadow requires a number of Compute Shader dispatches
+	// The compute shader reads from a depth buffer, and writes a single-channel texture of the same dimensions
+	// Each dispatch is of the same compute shader, (see bend_sss_gpu.h).
+	// The number of dispatches required varies based on the on-screen location of the light.
+	// Typically there will be just one or two dispatches when the light is off-screen, and 4 to 6 when the light is on-screen.
+	// Syncing the GPU between individual dispatches is not required
+
+	// These structures and function are used to generate the number of dispatches, the wave count of each dispatch (X/Y/Z) and shader parameters for each dispatch
+
+	struct DispatchData
+	{
+		int WaveCount[3];					// Compute Shader Dispatch(X,Y,Z) wave counts X/Y/Z
+		int WaveOffset_Shader[2];			// This value is passed in to shader. It will be different for each dispatch
+	};
+
+	struct DispatchList
+	{
+		float LightCoordinate_Shader[4];	// This value is passed in to shader, this will be the same value for all dispatches for this light
+
+		DispatchData Dispatch[8];			// List of dispatches (max count is 8)
+		int DispatchCount;					// Number of compute dispatches written to the list
+	};
+
+
+	// Helper functions
+	inline int bend_min(const int a, const int b) { return a > b ? b : a; }
+	inline int bend_max(const int a, const int b) { return a > b ? a : b; }
+
+
+	// Call this function on the CPU to get a list of Compute Shader dispatches required to generate a screen-space shadow for a given light
+	// Syncing the GPU between individual dispatches is not required
+	// 
+	// inLightProjection:		Homogeneous coordinate of the light, result of {light} * {ViewProjectionMatrix}, (without W divide)
+	//							For infinite directional lights, use {light} = float4(normalized light direction, 0) and for point/spot lights use {light} = float4(light world position, 1)
+	// 
+	// inViewportSize:			width/height of the render target
+	// 
+	// inRenderBounds:			2D Screen Bounds of the light within the viewport, inclusive. [0,0], [width,height] for full-screen.
+	//							Note; the shader will still read/write outside of these bounds (by a maximum of 2 * WAVE_SIZE pixels), due to how the wavefront projection works.
+	// 
+	// inExpandedDepthRange:	Set to true if the rendering API expects z/w coordinate output from a vertex shader to be a [-1,+1] expanded range, and becomes [0,1] range in the depth buffer. Typically this is false.
+	// 
+	// inWaveSize:				Wavefront size of the compiled compute shader (currently only tested with 64)
+	//
+	DispatchList BuildDispatchList(float inLightProjection[4], int inViewportSize[2], int inMinRenderBounds[2], int inMaxRenderBounds[2], bool inExpandedZRange = false, int inWaveSize = 64)
+	{
+		DispatchList result = {};
+
+		// Floating point division in the shader has a practical limit for precision when the light is *very* far off screen (~1m pixels+)
+		// So when computing the light XY coordinate, use an adjusted w value to handle these extreme values
+		float xy_light_w = inLightProjection[3];
+		float FP_limit = 0.000002f * (float)inWaveSize;
+
+		if (xy_light_w >= 0 && xy_light_w < FP_limit) xy_light_w = FP_limit;
+		else if (xy_light_w < 0 && xy_light_w > -FP_limit) xy_light_w = -FP_limit;
+
+		// Need precise XY pixel coordinates of the light
+		result.LightCoordinate_Shader[0] = ((inLightProjection[0] / xy_light_w) * +0.5f + 0.5f) * (float)inViewportSize[0];
+		result.LightCoordinate_Shader[1] = ((inLightProjection[1] / xy_light_w) * -0.5f + 0.5f) * (float)inViewportSize[1];
+		result.LightCoordinate_Shader[2] = inLightProjection[3] == 0 ? 0 : (inLightProjection[2] / inLightProjection[3]);
+		result.LightCoordinate_Shader[3] = inLightProjection[3] > 0 ? 1 : -1;
+
+		if (inExpandedZRange)
+		{
+			result.LightCoordinate_Shader[2] = result.LightCoordinate_Shader[2] * 0.5f + 0.5f;
+		}
+
+		int light_xy[2] = { (int)(result.LightCoordinate_Shader[0] + 0.5f), (int)(result.LightCoordinate_Shader[1] + 0.5f) };
+
+		// Make the bounds inclusive, relative to the light
+		const int biased_bounds[4] =
+		{
+			inMinRenderBounds[0] - light_xy[0],
+			-(inMaxRenderBounds[1] - light_xy[1]),
+			inMaxRenderBounds[0] - light_xy[0],
+			-(inMinRenderBounds[1] - light_xy[1]),
+		};
+
+		// Process 4 quadrants around the light center,
+		// They each form a rectangle with one corner on the light XY coordinate
+		// If the rectangle isn't square, it will need breaking in two on the larger axis
+		// 0 = bottom left, 1 = bottom right, 2 = top left, 2 = top right
+		for (int q = 0; q < 4; q++)
+		{
+			// Quads 0 and 3 needs to be +1 vertically, 1 and 2 need to be +1 horizontally
+			bool vertical = q == 0 || q == 3;
+
+			// Bounds relative to the quadrant
+			const int bounds[4] =
+			{
+				bend_max(0, ((q & 1) ? biased_bounds[0] : -biased_bounds[2])) / inWaveSize,
+				bend_max(0, ((q & 2) ? biased_bounds[1] : -biased_bounds[3])) / inWaveSize,
+				bend_max(0, (((q & 1) ? biased_bounds[2] : -biased_bounds[0]) + inWaveSize * (vertical ? 1 : 2) - 1)) / inWaveSize,
+				bend_max(0, (((q & 2) ? biased_bounds[3] : -biased_bounds[1]) + inWaveSize * (vertical ? 2 : 1) - 1)) / inWaveSize,
+			};
+
+			if ((bounds[2] - bounds[0]) > 0 && (bounds[3] - bounds[1]) > 0)
+			{
+				int bias_x = (q == 2 || q == 3) ? 1 : 0;
+				int bias_y = (q == 1 || q == 3) ? 1 : 0;
+
+				DispatchData& disp = result.Dispatch[result.DispatchCount++];
+
+				disp.WaveCount[0] = inWaveSize;
+				disp.WaveCount[1] = bounds[2] - bounds[0];
+				disp.WaveCount[2] = bounds[3] - bounds[1];
+				disp.WaveOffset_Shader[0] = ((q & 1) ? bounds[0] : -bounds[2]) + bias_x;
+				disp.WaveOffset_Shader[1] = ((q & 2) ? -bounds[3] : bounds[1]) + bias_y;
+
+				// We want the far corner of this quadrant relative to the light,
+				// as we need to know where the diagonal light ray intersects with the edge of the bounds
+				int axis_delta = +biased_bounds[0] - biased_bounds[1];
+				if (q == 1) axis_delta = +biased_bounds[2] + biased_bounds[1];
+				if (q == 2) axis_delta = -biased_bounds[0] - biased_bounds[3];
+				if (q == 3) axis_delta = -biased_bounds[2] + biased_bounds[3];
+
+				axis_delta = (axis_delta + inWaveSize - 1) / inWaveSize;
+
+				if (axis_delta > 0)
+				{
+					DispatchData& disp2 = result.Dispatch[result.DispatchCount++];
+
+					// Take copy of current volume
+					disp2 = disp;
+
+					if (q == 0)
+					{
+						// Split on Y, split becomes -1 larger on x
+						disp2.WaveCount[2] = bend_min(disp.WaveCount[2], axis_delta);
+						disp.WaveCount[2] -= disp2.WaveCount[2];
+						disp2.WaveOffset_Shader[1] = disp.WaveOffset_Shader[1] + disp.WaveCount[2];
+						disp2.WaveOffset_Shader[0]--;
+						disp2.WaveCount[1] ++;
+					}
+					if (q == 1)
+					{
+						// Split on X, split becomes +1 larger on y
+						disp2.WaveCount[1] = bend_min(disp.WaveCount[1], axis_delta);
+						disp.WaveCount[1] -= disp2.WaveCount[1];
+						disp2.WaveOffset_Shader[0] = disp.WaveOffset_Shader[0] + disp.WaveCount[1];
+						disp2.WaveCount[2] ++;
+					}
+					if (q == 2)
+					{
+						// Split on X, split becomes -1 larger on y
+						disp2.WaveCount[1] = bend_min(disp.WaveCount[1], axis_delta);
+						disp.WaveCount[1] -= disp2.WaveCount[1];
+						disp.WaveOffset_Shader[0] += disp2.WaveCount[1];
+						disp2.WaveCount[2] ++;
+						disp2.WaveOffset_Shader[1]--;
+					}
+					if (q == 3)
+					{
+						// Split on Y, split becomes +1 larger on x
+						disp2.WaveCount[2] = bend_min(disp.WaveCount[2], axis_delta);
+						disp.WaveCount[2] -= disp2.WaveCount[2];
+						disp.WaveOffset_Shader[1] += disp2.WaveCount[2];
+						disp2.WaveCount[1] ++;
+					}
+
+					// Remove if too small
+					if (disp2.WaveCount[1] <= 0 || disp2.WaveCount[2] <= 0)
+					{
+						disp2 = result.Dispatch[--result.DispatchCount];
+					}
+					if (disp.WaveCount[1] <= 0 || disp.WaveCount[2] <= 0)
+					{
+						disp = result.Dispatch[--result.DispatchCount];
+					}
+				}
+			}
+		}
+
+		// Scale the shader values by the wave count, the shader expects this
+		for (int i = 0; i < result.DispatchCount; i++)
+		{
+			result.Dispatch[i].WaveOffset_Shader[0] *= inWaveSize;
+			result.Dispatch[i].WaveOffset_Shader[1] *= inWaveSize;
+		}
+
+		return result;
+	}
+}
+
+
+/*
+* Common Problems, and tips to solve them:
+*
+* The shader doesn't compile?
+*  -	The shader is only tested with HLSL (DXC compiler) and PS5 using a HLSL->PS5 warpper
+*		It should be possible to compile for DX11 with FXC, with features removed (early-out's use of wave intrinsics is not supported in DX11)
+*		Other shader languages (e.g., glsl) are unsupported and will require manual conversion
+*
+* I have it compiled and running, but I'm seeing complete nonsense?
+*  -	Start by enabling 'DebugOutputWaveIndex' to visualize the wavefront layout.
+*		You should see wavefronts aligned and projected towards the light position / direction. If not, then 'inLightProjection' is probably wrong.
+*
+* Struggling to get 'inLightProjection' right?
+*  -	Think of this as similar to what a vertex-shader would output for position export; that is, a 4-component transformed position (e.g., SV_POSITION output from the VS)
+*		This usually means:
+*		inLightProjection = float4(position,1) * ViewProjectionMatrix		- positional light
+*		or:
+*		inLightProjection = float4(direction,0) * ViewProjectionMatrix		- directional light
+*
+* Almost everything is in shadow?
+* Is the background around an object casting a shadow on to it?
+* Does everything look like a weird paper cut-out?
+* There is a big shadow halo around the light source?
+*  -	FarDepthValue / NearDepthValue may not be set correctly
+*		These are the values you'd see in the depth buffer for objects at the near and far clip plane (typically far = 0, near = 1)
+*		In very rare cases, the vertex shader may be expected to output a [-1,+1] z-range, which becomes [0,1] in the depth buffer. If this is the case, enable 'inExpandedZRange'.
+*
+* There are small glitchy lines all over the output?
+*  -	Try with/without 'USE_HALF_PIXEL_OFFSET' defined in the shader. This is enabled by default for HLSL.
+*
+* Invalid shadows occasionally appear from offscreen / at the edge of the screen?
+*  -	The 'PointBorderSampler' may not be setup correctly. The shader will intentionally read from offscreen, so the sampler must return an invalid value (e.g., FarDepthValue) in these cases by using clamp-to-border.
+*
+* Light is always coming from the same direction no matter how I rotate the camera?
+*  -	You may be using just the ProjectionMatrix, not the ViewProjectionMatrix when computing 'inLightProjection'.
+*
+* Shadow is extremely thick, or very faded?
+*  -	Try scaling 'SurfaceThickness' up and down, start with 0.005, and scale up/down in multiples of 2. Make sure to scale BilinearThreshold in a similar way.
+*
+* I see lots of striated patterns on flat surfaces?
+*  -	'BilinearThreshold' may not be set to an ideal value (enable 'DebugOutputEdgeMask' to debug it), or try enabling 'IgnoreEdgePixels'
+*  -	These issues can be more common in otherwise occluded areas when BilinearSamplingOffsetMode is false, and may be prevalent if visualizing the shadow - but not noticeable in a lit scene.
+*/
\ No newline at end of file
