Support NonCastingIntensity in ScreenSpaceShadowsCS.
- sample GBuffer at hit position to get whether surface casts contact shadows or not (similar to DeferredLightingCommon.ush).
- only supported by deferred shading paths at the moment.

diff --git a/Engine/Shaders/Private/ScreenSpaceShadows.usf b/Engine/Shaders/Private/ScreenSpaceShadows.usf
--- a/Engine/Shaders/Private/ScreenSpaceShadows.usf
+++ b/Engine/Shaders/Private/ScreenSpaceShadows.usf
@@ -47,6 +47,8 @@ DispatchParameters SetupBendDispatchParameters(float NearDepthVal, float FarDept
 #include "SceneTexturesCommon.ush"
 #include "ComputeShaderUtils.ush"
 
+#include "DeferredShadingCommon.ush"
+
 #ifdef UPSAMPLE_PASS
 #	include "ShadowFactorsUpsampleCommon.ush"
 #endif
@@ -63,12 +65,43 @@ float3 LightDirection;
 float ContactShadowLength;
 uint bContactShadowLengthInWS;
 float ContactShadowCastingIntensity;
+float ContactShadowNonCastingIntensity;
 
 uint DownsampleFactor;
 int4 ScissorRectMinAndSize;
 
 RWTexture2D<float2> RWShadowFactors;
 
+float GetContactShadowOcclusion(float2 HitUV)
+{
+	const uint2 HitPixelCoord = uint2(HitUV * View.BufferSizeAndInvSize.xy);
+
+#if SHADING_PATH_MOBILE
+#	if MOBILE_DEFERRED_SHADING
+		FGBufferData SampleGBuffer = MobileFetchAndDecodeGBuffer(HitUV, HitPixelCoord);
+		const bool bHitCastContactShadow = CastContactShadow(SampleGBuffer);
+#	else 
+		const bool bHitCastContactShadow = true; // TODO: support mobile forward path
+#	endif
+#else
+#	if SUBSTRATE_ENABLED
+#		if SUBSTRATE_RAYCAST_ENABLED && SUBSTRATE_MATERIALCONTAINER_IS_VIEWRESOURCE
+			uint2 PixelPos = View.ViewRectMin.xy + View.ViewSizeAndInvSize.xy * HitUV;
+			FSubstrateAddressing SubstrateAddressing = GetSubstratePixelDataByteOffset(PixelPos, uint2(View.BufferSizeAndInvSize.xy), SubstrateRayCast.MaxBytesPerPixel);
+			FSubstratePixelHeader SubstratePixelHeader = UnpackSubstrateHeaderIn(SubstrateRayCast.MaterialTextureArray, SubstrateAddressing, SubstrateRayCast.TopLayerTexture);
+			const bool bHitCastContactShadow = SubstratePixelHeader.DoesCastContactShadow();
+#		else
+			const bool bHitCastContactShadow = false;
+#		endif
+#	else
+		FGBufferData SampleGBuffer = GetGBufferData(HitUV);
+		const bool bHitCastContactShadow = CastContactShadow(SampleGBuffer);
+#	endif
+#endif
+
+	return bHitCastContactShadow ? ContactShadowCastingIntensity : ContactShadowNonCastingIntensity;
+}
+
 [numthreads(THREADGROUP_SIZEX, THREADGROUP_SIZEY, 1)]
 void ScreenSpaceShadowsCS(
 	uint3 GroupId : SV_GroupID,
@@ -77,8 +110,8 @@ void ScreenSpaceShadowsCS(
 {
 	uint ThreadIndex = GroupThreadId.y * THREADGROUP_SIZEX + GroupThreadId.x;
 
-	float2 ScreenUV = float2((DispatchThreadId.xy * DownsampleFactor + ScissorRectMinAndSize.xy + .5f) * View.BufferSizeAndInvSize.zw);
-	float2 ScreenPosition = (ScreenUV.xy - View.ScreenPositionScaleBias.wz) / View.ScreenPositionScaleBias.xy;
+	float2 ScreenUV = (DispatchThreadId.xy * DownsampleFactor + ScissorRectMinAndSize.xy + .5f) * View.BufferSizeAndInvSize.zw;
+	float2 ScreenPosition = (ScreenUV - View.ScreenPositionScaleBias.wz) / View.ScreenPositionScaleBias.xy;
 
 	float SceneDepth = CalcSceneDepth(ScreenUV);
 	float3 OpaqueTranslatedWorldPosition = mul(float4(GetScreenPositionForProjectionType(ScreenPosition, SceneDepth), SceneDepth, 1), PrimaryView.ScreenToTranslatedWorld).xyz;
@@ -91,7 +124,7 @@ void ScreenSpaceShadowsCS(
 	float2 HitUV;
 	float HitDistance = CastScreenSpaceShadowRay(OpaqueTranslatedWorldPosition, LightDirection, ActualContactShadowLength, 8, Dither, 2.0f, false, HitUV);
 
-	float Result = HitDistance > 0.0 ? (1.0f - ContactShadowCastingIntensity) : 1.0f;
+	float Result = HitDistance > 0.0 ? (1.0f - GetContactShadowOcclusion(HitUV)) : 1.0f;
 
 	#if METAL_PROFILE 
 		// clamp max depth to avoid #inf
diff --git a/Engine/Source/Runtime/Renderer/Private/Shadows/ScreenSpaceShadows.cpp b/Engine/Source/Runtime/Renderer/Private/Shadows/ScreenSpaceShadows.cpp
--- a/Engine/Source/Runtime/Renderer/Private/Shadows/ScreenSpaceShadows.cpp
+++ b/Engine/Source/Runtime/Renderer/Private/Shadows/ScreenSpaceShadows.cpp
@@ -61,6 +61,7 @@ class FScreenSpaceShadowsCS : public FGlobalShader
 		SHADER_PARAMETER(float, ContactShadowLength)
 		SHADER_PARAMETER(uint32, bContactShadowLengthInWS)
 		SHADER_PARAMETER(float, ContactShadowCastingIntensity)
+		SHADER_PARAMETER(float, ContactShadowNonCastingIntensity)
 		SHADER_PARAMETER(FIntRect, ScissorRectMinAndSize)
 		SHADER_PARAMETER(uint32, DownsampleFactor)
 	END_SHADER_PARAMETER_STRUCT()
@@ -95,6 +96,7 @@ class FScreenSpaceShadowsBendCS : public FGlobalShader
 		SHADER_PARAMETER(float, ContactShadowLength)
 		SHADER_PARAMETER(uint32, bContactShadowLengthInWS)
 		SHADER_PARAMETER(float, ContactShadowCastingIntensity)
+		SHADER_PARAMETER(float, ContactShadowNonCastingIntensity)
 		SHADER_PARAMETER(FIntRect, ScissorRectMinAndSize)
 		SHADER_PARAMETER(uint32, DownsampleFactor)
 		SHADER_PARAMETER(FVector2f, InvDepthTextureSize)
@@ -261,6 +263,7 @@ void RenderScreenSpaceShadows(
 		PassParameters->ContactShadowLength = ContactShadowLength;
 		PassParameters->bContactShadowLengthInWS = bContactShadowLengthInWS;
 		PassParameters->ContactShadowCastingIntensity = ContactShadowCastingIntensity;
+		PassParameters->ContactShadowNonCastingIntensity = ContactShadowNonCastingIntensity;
 
 		auto ComputeShader = View.ShaderMap->GetShader<FScreenSpaceShadowsCS>();
 
@@ -354,6 +357,7 @@ void RenderScreenSpaceShadowsBend(
 			PassParameters->ContactShadowLength = ContactShadowLength;
 			PassParameters->bContactShadowLengthInWS = bContactShadowLengthInWS;
 			PassParameters->ContactShadowCastingIntensity = ContactShadowCastingIntensity;
+			PassParameters->ContactShadowNonCastingIntensity = ContactShadowNonCastingIntensity;
 
 			PassParameters->LightCoordinate = FVector4f(DispatchList.LightCoordinate_Shader[0], DispatchList.LightCoordinate_Shader[1], DispatchList.LightCoordinate_Shader[2], DispatchList.LightCoordinate_Shader[3]);
 			PassParameters->WaveOffset = FIntVector(Dispatch.WaveOffset_Shader[0], Dispatch.WaveOffset_Shader[1], 0);
