Simplifies the implementation of TSR's AnyThread() on desktop to try work arround GPU hang

#tests win64, ps5, xboxsx

diff --git a/Engine/Shaders/Private/TemporalSuperResolution/TSRConvolutionNetwork.ush b/Engine/Shaders/Private/TemporalSuperResolution/TSRConvolutionNetwork.ush
--- a/Engine/Shaders/Private/TemporalSuperResolution/TSRConvolutionNetwork.ush
+++ b/Engine/Shaders/Private/TemporalSuperResolution/TSRConvolutionNetwork.ush
@@ -194,38 +194,15 @@ TLaneVector2D<bool, 1, LaneStrideX, LaneStrideY> IsValidTilePixel(const uint2 Ti
 
 //------------------------------------------------------- VOTE
 
+groupshared uint SharedAtomic;
+
 CALL_SITE_DEBUGLOC
 bool AnyThread(bool X)
-#if !COMPILER_SUPPORTS_WAVE_VOTE && defined(LDS_SIZE)
-{
-	GroupMemoryBarrierWithGroupSync();
-	
-	if (GGroupThreadIndex == 0)
-	{
-		WriteDwordToLDS(/* SharedIndex = */ 0, 0);
-	}
-
-	GroupMemoryBarrierWithGroupSync();
-	
-	AtomicIncrementLDSDword(/* SharedIndex = */ 0, uint(X));
-	
-	GroupMemoryBarrierWithGroupSync();
-	
-	return ReadDwordFromLDS(/* SharedIndex = */ 0) > 0;
-}
-#elif !COMPILER_SUPPORTS_WAVE_VOTE && !defined(LDS_SIZE)
-{
-	// Metal SM6 have COMPILER_SUPPORTS_WAVE_VOTE=0 due to shader compiler bug
-	#if !COMPILER_METAL
-		#error Requires to specify LDS_SIZE
-	#endif
-	return true;
-}
-#elif WAVE_COUNT_X == 1 && WAVE_COUNT_Y == 1
+#if WAVE_COUNT_X == 1 && WAVE_COUNT_Y == 1 && COMPILER_SUPPORTS_WAVE_VOTE
 {
 	return WaveActiveAnyTrue(X);
 }
-#else
+#elif COMPILER_SUPPORTS_WAVE_VOTE && PLATFORM_GPU_ARCH == PLATFORM_GPU_ARCH_AMD_RDNA_2
 {
 	const uint kLaneCount = WaveGetLaneCount();
 	const uint kThreadCount = LANE_COUNT_X * LANE_COUNT_Y * WAVE_COUNT_X * WAVE_COUNT_Y;
@@ -237,7 +214,7 @@ bool AnyThread(bool X)
 
 	if (WaveIsFirstLane())
 	{
-		WriteDwordToLDS(GGroupThreadIndex / kLaneCount, uint(AnyX));
+		WriteDwordToLDS(GGroupThreadIndex / kLaneCount, select(AnyX, 1u, 0u));
 	}
 		
 	GroupMemoryBarrierWithGroupSync();
@@ -246,6 +223,42 @@ bool AnyThread(bool X)
 
 	return WaveActiveAnyTrue(WaveX != 0);
 }
+#elif defined(LDS_SIZE)
+{
+	const uint SharedIndex = 0;
+
+	GroupMemoryBarrierWithGroupSync();
+	
+	if (GGroupThreadIndex == 0)
+	{
+		WriteDwordToLDS(SharedIndex, 0);
+	}
+
+	GroupMemoryBarrierWithGroupSync();
+	
+	AtomicIncrementLDSDword(SharedIndex, select(X, 1u, 0u));
+	
+	GroupMemoryBarrierWithGroupSync();
+	
+	return ReadDwordFromLDS(SharedIndex) > 0;
+}
+#else
+{
+	GroupMemoryBarrierWithGroupSync();
+	
+	if (GGroupThreadIndex == 0)
+	{
+		SharedAtomic = 0;
+	}
+
+	GroupMemoryBarrierWithGroupSync();
+	
+	InterlockedAdd(/* inout */ SharedAtomic, select(X, 1u, 0u));
+	
+	GroupMemoryBarrierWithGroupSync();
+	
+	return SharedAtomic > 0;
+}
 #endif
 
 
