Add coverage/tangent fallback on card groom when no texture are provided.

diff --git a/Engine/Plugins/Runtime/HairStrands/Source/HairStrandsCore/Private/GroomAsset.cpp b/Engine/Plugins/Runtime/HairStrands/Source/HairStrandsCore/Private/GroomAsset.cpp
--- a/Engine/Plugins/Runtime/HairStrands/Source/HairStrandsCore/Private/GroomAsset.cpp
+++ b/Engine/Plugins/Runtime/HairStrands/Source/HairStrandsCore/Private/GroomAsset.cpp
@@ -412,7 +412,9 @@ template<typename ResourceType>
 static void InitAtlasTexture(ResourceType* InResource, UTexture2D* InTexture, EHairAtlasTextureType InType)
 {
 	if (InTexture == nullptr || InResource == nullptr)
+	{
 		return;
+	}
 
 	InTexture->ConditionalPostLoad();
 
diff --git a/Engine/Plugins/Runtime/HairStrands/Source/HairStrandsCore/Private/HairCardsVertexFactory.cpp b/Engine/Plugins/Runtime/HairStrands/Source/HairStrandsCore/Private/HairCardsVertexFactory.cpp
--- a/Engine/Plugins/Runtime/HairStrands/Source/HairStrandsCore/Private/HairCardsVertexFactory.cpp
+++ b/Engine/Plugins/Runtime/HairStrands/Source/HairStrandsCore/Private/HairCardsVertexFactory.cpp
@@ -27,12 +27,15 @@ template<typename T> inline void VFC_BindParam(FMeshDrawSingleShaderBindings& Sh
 
 enum class EHairCardsFactoryFlags : uint32
 {
-	InvertedUV = 1,
-	TextureRootUV = 2,
-	TextureGroupIndex = 4,
-	TextureBaseColor = 8,
-	TextureRoughness = 10,
-	TextureAttribute = 12
+	InvertedUV        = 0x1,
+	TextureRootUV     = 0x2,
+	TextureGroupIndex = 0x4,
+	TextureBaseColor  = 0x8,
+	TextureRoughness  = 0x10,
+	TextureAttribute  = 0x20,
+	TextureCoverage   = 0x40,
+	TextureDepth      = 0x80,
+	TextureTangent    = 0x100
 };
 
 ////////////////////////////////////////////////////////////////////////////////////////////////////////////////////////
@@ -53,7 +56,16 @@ FHairCardsUniformBuffer CreateHairCardsVFUniformBuffer(
 
 		// Cards atlas UV are inverted so fetching needs to be inverted on the y-axis
 		UniformParameters.Flags = 0;
-		if (LOD.RestResource->bInvertUV) { UniformParameters.Flags |= uint32(EHairCardsFactoryFlags::InvertedUV);  }
+		UniformParameters.Flags |= LOD.RestResource->bInvertUV 					? uint32(EHairCardsFactoryFlags::InvertedUV) : 0u;
+		UniformParameters.Flags |= LOD.RestResource->CoverageTexture != nullptr ? uint32(EHairCardsFactoryFlags::TextureCoverage) : 0;
+		UniformParameters.Flags |= LOD.RestResource->DepthTexture != nullptr 	? uint32(EHairCardsFactoryFlags::TextureDepth) : 0;
+		UniformParameters.Flags |= LOD.RestResource->TangentTexture != nullptr 	? uint32(EHairCardsFactoryFlags::TextureTangent) : 0;
+		// These attributes are stored on vertex data rather than on textures
+		// EHairCardsFactoryFlags::TextureRootUV
+		// EHairCardsFactoryFlags::TextureGroupIndex
+		// EHairCardsFactoryFlags::TextureBaseColor
+		// EHairCardsFactoryFlags::TextureRoughness
+
 		UniformParameters.MaxVertexCount = LOD.RestResource->GetVertexCount();
 		UniformParameters.CoverageBias = FMath::Clamp(GHairCardCoverageBias, -16.f, 16.f);
 
@@ -110,10 +122,13 @@ FHairCardsUniformBuffer CreateHairCardsVFUniformBuffer(
 		// Meshes UV are not inverted so no need to invert the y-axis
 		UniformParameters.Flags = 0;
 		UniformParameters.Flags |= uint32(EHairCardsFactoryFlags::TextureRootUV);
-		UniformParameters.Flags |= uint32(EHairCardsFactoryFlags::TextureAttribute);
 		UniformParameters.Flags |= uint32(EHairCardsFactoryFlags::TextureGroupIndex);
 		UniformParameters.Flags |= uint32(EHairCardsFactoryFlags::TextureBaseColor);
 		UniformParameters.Flags |= uint32(EHairCardsFactoryFlags::TextureRoughness);
+		UniformParameters.Flags |= uint32(EHairCardsFactoryFlags::TextureAttribute);
+		UniformParameters.Flags |= LOD.RestResource->CoverageTexture != nullptr ? uint32(EHairCardsFactoryFlags::TextureCoverage) : 0;
+		UniformParameters.Flags |= LOD.RestResource->DepthTexture != nullptr ? uint32(EHairCardsFactoryFlags::TextureDepth) : 0;
+		UniformParameters.Flags |= LOD.RestResource->TangentTexture != nullptr ? uint32(EHairCardsFactoryFlags::TextureTangent) : 0;
 
 		UniformParameters.CoverageBias = 0;
 		UniformParameters.MaxVertexCount = LOD.RestResource->GetVertexCount();
diff --git a/Engine/Shaders/Private/HairStrands/HairCardsAttributeCommon.ush b/Engine/Shaders/Private/HairStrands/HairCardsAttributeCommon.ush
--- a/Engine/Shaders/Private/HairStrands/HairCardsAttributeCommon.ush
+++ b/Engine/Shaders/Private/HairStrands/HairCardsAttributeCommon.ush
@@ -8,12 +8,15 @@
 
 ////////////////////////////////////////////////////////////////////////////////
 
-#define EHairCardsFVFlags_InvertedUV 1
-#define EHairCardsFVFlags_TextureRootUV 2
-#define EHairCardsFVFlags_TextureGroupIndex 4
-#define EHairCardsFVFlags_TextureBaseColor 8
-#define EHairCardsFVFlags_TextureRoughness 10
-#define EHairCardsFVFlags_TextureAttribute 12
+#define EHairCardsFVFlags_InvertedUV        0x1
+#define EHairCardsFVFlags_TextureRootUV     0x2
+#define EHairCardsFVFlags_TextureGroupIndex 0x4
+#define EHairCardsFVFlags_TextureBaseColor  0x8
+#define EHairCardsFVFlags_TextureRoughness  0x10
+#define EHairCardsFVFlags_TextureAttribute  0x20
+#define EHairCardsFVFlags_TextureCoverage   0x40
+#define EHairCardsFVFlags_TextureDepth      0x80
+#define EHairCardsFVFlags_TextureTangent    0x100
 
 struct FHairVertexAttributes
 {
@@ -28,9 +31,13 @@ struct FHairVertexAttributes
 
 float GetHairStrandsDepth(float2 InAtlasUV, float InDeviceZ)
 {
-	const bool bInvertUV = HairCardsVF.Flags & EHairCardsFVFlags_InvertedUV;
-	InAtlasUV.y = bInvertUV ? 1 - InAtlasUV.y : InAtlasUV.y;
-	const float SceneDepthOffset = Texture2DSample(HairCardsVF.DepthTexture, HairCardsVF.DepthSampler, InAtlasUV).x;
+	float SceneDepthOffset = 0;
+	if (HairCardsVF.Flags & EHairCardsFVFlags_TextureDepth)
+	{
+		const bool bInvertUV = HairCardsVF.Flags & EHairCardsFVFlags_InvertedUV;
+		InAtlasUV.y = bInvertUV ? 1 - InAtlasUV.y : InAtlasUV.y;
+		SceneDepthOffset = Texture2DSample(HairCardsVF.DepthTexture, HairCardsVF.DepthSampler, InAtlasUV).x;
+	}
 	const float SceneDepth = ConvertFromDeviceZ(InDeviceZ);
 
 	return ConvertToDeviceZ(SceneDepth + SceneDepthOffset);
@@ -38,18 +45,27 @@ float GetHairStrandsDepth(float2 InAtlasUV, float InDeviceZ)
 
 float3 GetHairStrandsTangent(float2 InAtlasUV, half3x3 TangentToWorld, bool bUseTangentSpace)
 {
-	const bool bInvertUV = HairCardsVF.Flags & EHairCardsFVFlags_InvertedUV;
-	InAtlasUV.y = bInvertUV ? 1 - InAtlasUV.y : InAtlasUV.y;
-	const float3 LocalTangent = Texture2DSample(HairCardsVF.TangentTexture, HairCardsVF.TangentSampler, InAtlasUV).xyz * 2 - 1;
+	float3 LocalTangent = float3(0, 1, 0);
+	if (HairCardsVF.Flags & EHairCardsFVFlags_TextureTangent)
+	{
+		const bool bInvertUV = HairCardsVF.Flags & EHairCardsFVFlags_InvertedUV;
+		InAtlasUV.y = bInvertUV ? 1 - InAtlasUV.y : InAtlasUV.y;
+		LocalTangent = Texture2DSample(HairCardsVF.TangentTexture, HairCardsVF.TangentSampler, InAtlasUV).xyz * 2 - 1;
+	}
 	const float3 WorldTangent = mul(LocalTangent, TangentToWorld);
 	return bUseTangentSpace ? LocalTangent : WorldTangent;
 }
 
 float GetHairStrandsCoverage(float2 InAtlasUV)
 {
-	const bool bInvertUV = HairCardsVF.Flags & EHairCardsFVFlags_InvertedUV;
-	InAtlasUV.y = bInvertUV ? 1 - InAtlasUV.y : InAtlasUV.y;
-	return Texture2DSampleBias(HairCardsVF.CoverageTexture, HairCardsVF.CoverageSampler, InAtlasUV, HairCardsVF.CoverageBias).x;
+	float Coverage = 1.f;
+	if (HairCardsVF.Flags & EHairCardsFVFlags_TextureCoverage)
+	{
+		const bool bInvertUV = HairCardsVF.Flags & EHairCardsFVFlags_InvertedUV;
+		InAtlasUV.y = bInvertUV ? 1 - InAtlasUV.y : InAtlasUV.y;
+		Coverage = Texture2DSampleBias(HairCardsVF.CoverageTexture, HairCardsVF.CoverageSampler, InAtlasUV, HairCardsVF.CoverageBias).x;
+	}
+	return Coverage;
 }
 
 FHairVertexAttributes GetHairStrandsAttributes(float2 InAtlasUV, float2 InRootUV = float2(0, 0), float InCardLength = 0.f, float InCardGroupIndex = 0.f, float3 InCardBaseColor = 0.f, float InCardRoughness = 0.f)
