Substrate - fixed compilation bug with compute mateiral path.

diff --git a/Engine/Shaders/Private/ComputeShaderOutputCommon.ush b/Engine/Shaders/Private/ComputeShaderOutputCommon.ush
--- a/Engine/Shaders/Private/ComputeShaderOutputCommon.ush
+++ b/Engine/Shaders/Private/ComputeShaderOutputCommon.ush
@@ -165,7 +165,7 @@ void ConditionalExport(bool Forced, RWTexture2D<SUBSTRATE_TOP_LAYER_TYPE> OutTar
 {
 #if CONDITIONAL_EXPORT
 	BRANCH
-	if (Forced || TopLayerData != 0)
+	if (Forced || any(TopLayerData != 0))
 #endif
 	{
 		OutTarget[PixelPos] = TopLayerData;
