Substrate - fixed compilation issue.

diff --git a/Engine/Shaders/Private/ShadingFurnaceTest.usf b/Engine/Shaders/Private/ShadingFurnaceTest.usf
--- a/Engine/Shaders/Private/ShadingFurnaceTest.usf
+++ b/Engine/Shaders/Private/ShadingFurnaceTest.usf
@@ -49,7 +49,7 @@ void MainPS(
 		// Fetch the first layer of roughness & normal only. This fine for simple material but will not work for clear coat.
 		FSubstrateAddressing SubstrateAddressing = GetSubstratePixelDataByteOffset(PixelPos, uint2(View.BufferSizeAndInvSize.xy), Substrate.MaxBytesPerPixel);
 		FSubstratePixelHeader SubstratePixelHeader = UnpackSubstrateHeaderIn(Substrate.MaterialTextureArray, SubstrateAddressing, Substrate.TopLayerTexture);
-		bIsValid = SubstratePixelHeader.BSDFCount > 0;
+		bIsValid = SubstratePixelHeader.ClosureCount > 0;
 		if (bIsValid)
 		{
 			const float3 DummyV = float3(0, 0, 1);
