Substrate - Fixed differences bettwen enabled/disable for Subsurface shading model.
There was one error, a wrong NoL and specular transmittance that can be ignored fora 100% match.

diff --git a/Engine/Shaders/Private/Substrate/SubstrateEvaluation.ush b/Engine/Shaders/Private/Substrate/SubstrateEvaluation.ush
--- a/Engine/Shaders/Private/Substrate/SubstrateEvaluation.ush
+++ b/Engine/Shaders/Private/Substrate/SubstrateEvaluation.ush
@@ -822,7 +822,7 @@ FSubstrateEvaluateResult SubstrateEvaluateBSDFCommon(FSubstrateBSDFContext BSDFC
 					// * Simplified version (with w=.5,  n=1.5): 
 					//     half WrappedDiffuse = 2 * pow(saturate((dot(N, L) + w) / (1.0f + w)), n) * (n + 1) / (2 * (1 + w));
 					//     NormalContribution = WrappedDiffuse * Opacity + 1 - Opacity;
-					const half WrappedDiffuse = pow(saturate(AreaLightContext.NoL * (1.f / 1.5f) + (0.5f / 1.5f)), 1.5f) * (2.5f / 1.5f);
+					const half WrappedDiffuse = pow(saturate(BSDFContext.Context.NoL * (1.f / 1.5f) + (0.5f / 1.5f)), 1.5f) * (2.5f / 1.5f);
 					const half NormalContribution = lerp(1.f, WrappedDiffuse, Opacity);
 					const half BackScatter = /* GBuffer.GBufferAO */ NormalContribution / (PI * 2);
 
@@ -836,11 +836,11 @@ FSubstrateEvaluateResult SubstrateEvaluateBSDFCommon(FSubstrateBSDFContext BSDFC
 					const half3 TransmittedColor = HSV_2_LinearRGB(half3(LinearRGB_2_HSV(RawTransmittedColor).xy, LinearRGB_2_HSV(SubSurfaceColor).z));
 
 					// Lerp to never exceed 1 (energy conserving)
-					TransmissionThroughput = lerp(BackScatter, 1, InScatter) * lerp(TransmittedColor, SubSurfaceColor, ShadowTerms.TransmissionShadow);
+					TransmissionThroughput = lerp(BackScatter, 1, InScatter) * lerp(TransmittedColor, SubSurfaceColor, ShadowTerms.TransmissionThickness);
 				}
 
 				Sample.TransmissionPDF = 1.0f / (4.0f * PI); // SUBSTRATE_TODO this currently match the uniform sphere sampling from SubstrateImportanceSampleBSDF
-				Sample.TransmissionPathValue = TransmissionThroughput * DirectionalAlbedo_SpecularTransmission;
+				Sample.TransmissionPathValue = TransmissionThroughput; // We are not applying DirectionalAlbedo_SpecularTransmission here to match legacy shading model for now.
 				Sample.IntegratedDiffuseValue += (ShadowTerms.TransmissionShadow * AreaLightContext.Falloff) * AreaLightContext.AreaLight.FalloffColor * Sample.TransmissionPathValue;
 			}
 			#endif
