Stochastic Shadows - fixed light trace compation shader compilation error and fixed reprojection on screen border.

diff --git a/Engine/Shaders/Private/StochasticShadows/StochasticShadows.usf b/Engine/Shaders/Private/StochasticShadows/StochasticShadows.usf
--- a/Engine/Shaders/Private/StochasticShadows/StochasticShadows.usf
+++ b/Engine/Shaders/Private/StochasticShadows/StochasticShadows.usf
@@ -139,7 +139,7 @@ void CompositeShadowMaskTracesCS(
 			int2 ScreenCoord00 = ShadowMaskHistoryScreenCoord00[ShadowMaskScreenCoord];
 			float4 Weights = ShadowMaskHistoryWeights[ShadowMaskScreenCoord];
 
-			if (all(ScreenCoord00 > 0))
+			if (any(Weights > 0.0f))
 			{
 				FShadowMaskAccumulator Accumulator = InitShadowMaskAccumulator();
 
diff --git a/Engine/Shaders/Private/StochasticShadows/StochasticShadowsTracing.ush b/Engine/Shaders/Private/StochasticShadows/StochasticShadowsTracing.ush
--- a/Engine/Shaders/Private/StochasticShadows/StochasticShadowsTracing.ush
+++ b/Engine/Shaders/Private/StochasticShadows/StochasticShadowsTracing.ush
@@ -6,6 +6,12 @@
 #include "../RayTracing/RayTracingSphereLight.ush"
 #include "../RayTracing/RayTracingRectLight.ush"
 
+#ifndef NUM_SAMPLES_PER_PIXEL
+	#define NUM_SAMPLES_PER_PIXEL 4
+	#define NUM_SAMPLES_PER_PIXEL_Y 2
+	#define NUM_SAMPLES_PER_PIXEL_X 2
+#endif
+
 uint2 SampleCoordToDownsampledScreenCoord(uint2 SampleCoord)
 {
 	return SampleCoord / uint2(NUM_SAMPLES_PER_PIXEL_X, NUM_SAMPLES_PER_PIXEL_Y);
