Stochastic Shadows - fixed shader compilation with r.Shaders.RemoveDeadCode=False

diff --git a/Engine/Shaders/Private/StochasticShadows/StochasticShadows.ush b/Engine/Shaders/Private/StochasticShadows/StochasticShadows.ush
--- a/Engine/Shaders/Private/StochasticShadows/StochasticShadows.ush
+++ b/Engine/Shaders/Private/StochasticShadows/StochasticShadows.ush
@@ -29,6 +29,12 @@
 #define DOWNSAMPLE_FACTOR 2
 #define SHADOW_MASK_DOWNSAMPLE_FACTOR 2
 
+#ifndef NUM_SAMPLES_PER_PIXEL
+	#define NUM_SAMPLES_PER_PIXEL 4
+	#define NUM_SAMPLES_PER_PIXEL_Y 2
+	#define NUM_SAMPLES_PER_PIXEL_X 2
+#endif
+
 float3 GetTranslatedWorldPositionFromScreenUV(float2 ScreenUV, float SceneDepth)
 {
 	float2 ScreenPosition = (ScreenUV - View.ScreenPositionScaleBias.wz) / View.ScreenPositionScaleBias.xy;
diff --git a/Engine/Shaders/Private/StochasticShadows/StochasticShadowsTracing.ush b/Engine/Shaders/Private/StochasticShadows/StochasticShadowsTracing.ush
--- a/Engine/Shaders/Private/StochasticShadows/StochasticShadowsTracing.ush
+++ b/Engine/Shaders/Private/StochasticShadows/StochasticShadowsTracing.ush
@@ -6,12 +6,6 @@
 #include "../RayTracing/RayTracingSphereLight.ush"
 #include "../RayTracing/RayTracingRectLight.ush"
 
-#ifndef NUM_SAMPLES_PER_PIXEL
-	#define NUM_SAMPLES_PER_PIXEL 4
-	#define NUM_SAMPLES_PER_PIXEL_Y 2
-	#define NUM_SAMPLES_PER_PIXEL_X 2
-#endif
-
 uint2 SampleCoordToDownsampledScreenCoord(uint2 SampleCoord)
 {
 	return SampleCoord / uint2(NUM_SAMPLES_PER_PIXEL_X, NUM_SAMPLES_PER_PIXEL_Y);
