Adding new Landscape LOD parameter: LOD Blend Range

This controls how much area of each landscape section we allow to blend towards the LOD of neighboring landscape sections.  At the default 1.0, it blends across the entire section (matching previous behavior).  Lower numbers reduce the blend region, limiting it to an area around the section boundary.

diff --git a/Engine/Plugins/Experimental/GPULightmass/Source/GPULightmass/Private/Scene/Landscape.h b/Engine/Plugins/Experimental/GPULightmass/Source/GPULightmass/Private/Scene/Landscape.h
--- a/Engine/Plugins/Experimental/GPULightmass/Source/GPULightmass/Private/Scene/Landscape.h
+++ b/Engine/Plugins/Experimental/GPULightmass/Source/GPULightmass/Private/Scene/Landscape.h
@@ -25,6 +25,7 @@ public:
 		int32       SubsectionSizeVerts;
 		float       StaticLightingResolution;
 		int32       StaticLightingLOD;
+		float       InvLODBlendRange;
 		FIntPoint   ComponentBase;
 		FIntPoint   SectionBase;
 		UTexture2D* HeightmapTexture;
diff --git a/Engine/Plugins/Experimental/GPULightmass/Source/GPULightmass/Private/Scene/Scene.cpp b/Engine/Plugins/Experimental/GPULightmass/Source/GPULightmass/Private/Scene/Scene.cpp
--- a/Engine/Plugins/Experimental/GPULightmass/Source/GPULightmass/Private/Scene/Scene.cpp
+++ b/Engine/Plugins/Experimental/GPULightmass/Source/GPULightmass/Private/Scene/Scene.cpp
@@ -1246,6 +1246,7 @@ void FScene::AddGeometryInstanceFromComponent(ULandscapeComponent* InComponent)
 	Initializer.HeightmapScaleBias         = (FVector4f)InComponent->HeightmapScaleBias;
 	Initializer.WeightmapScaleBias         = (FVector4f)InComponent->WeightmapScaleBias;
 	Initializer.WeightmapSubsectionOffset  = InComponent->WeightmapSubsectionOffset;
+	Initializer.InvLODBlendRange           = 1.0f / FMath::Max(InComponent->GetLandscapeProxy()->LODBlendRange, 0.01f);
 
 	TArray<int32> RelevantPointLightsToAddOnRenderThread = AddAllPossiblyRelevantLightsToGeometry(LightScene.PointLights, Instance);
 	TArray<int32> RelevantSpotLightsToAddOnRenderThread = AddAllPossiblyRelevantLightsToGeometry(LightScene.SpotLights, Instance);
@@ -1325,6 +1326,7 @@ void FScene::AddGeometryInstanceFromComponent(ULandscapeComponent* InComponent)
 			LandscapeParams.HeightmapUVScaleBias = Initializer.HeightmapScaleBias;
 			LandscapeParams.WeightmapUVScaleBias = Initializer.WeightmapScaleBias;
 			LandscapeParams.LocalToWorldNoScaling = FMatrix44f(InstanceRenderState.LocalToWorldNoScaling);				// LWC_TODO: Precision loss
+			LandscapeParams.InvLODBlendRange = Initializer.InvLODBlendRange;
 
 			LandscapeParams.LandscapeLightmapScaleBias = FVector4f(
 				LightmapScaleX,
diff --git a/Engine/Shaders/Private/LandscapeVertexFactory.ush b/Engine/Shaders/Private/LandscapeVertexFactory.ush
--- a/Engine/Shaders/Private/LandscapeVertexFactory.ush
+++ b/Engine/Shaders/Private/LandscapeVertexFactory.ush
@@ -83,47 +83,36 @@ float2 GetLodBias(uint ComponentIndex)
 #endif
 }
 
-/** 
- * Get current Lods for neighbour section instances. 
- *       -Y
- *    - - 0 - -
- *    |       |
- * -X 1   P   2 +X
- *    |       |
- *    - - 3 - -
- *       +Y
- */
-float4 GetNeighborSectionLod(uint ComponentIndex)
+
+float GetNeighborSectionLodFromOffset(int2 NeighborOffset, float CenterLOD)
 {
 #if FIXED_GRID
 	// Fixed grid has fixed Lod
-	return LandscapeFixedGrid.LodValues.xxxx;
-#else
-	int2 NeighborOffset[4];
-    NeighborOffset[0] = int2(0, -1);
-    NeighborOffset[1] = int2(-1, 0);
-    NeighborOffset[2] = int2(1, 0);
-    NeighborOffset[3] = int2(0, 1);
-	float NeighborLod[4];
-	UNROLL
-	for (int i = 0; i < 4; i++)
-	{
-		int2 NeighborPos = NeighborOffset[i] + int2(LandscapeParameters.ComponentBaseX, LandscapeParameters.ComponentBaseY);
-		NeighborPos = max(LandscapeContinuousLODParameters.Min                                                     , NeighborPos);
-		NeighborPos = min(LandscapeContinuousLODParameters.Min + LandscapeContinuousLODParameters.Size - int2(1, 1), NeighborPos);
+	return LandscapeFixedGrid.LodValues.x;
+#else // !FIXED_GRID
+	// Calculate neighbor position, and clamp to valid range
+	int2 NeighborPos = NeighborOffset + int2(LandscapeParameters.ComponentBaseX, LandscapeParameters.ComponentBaseY);
+	NeighborPos = max(LandscapeContinuousLODParameters.Min                                                     , NeighborPos);
+	NeighborPos = min(LandscapeContinuousLODParameters.Min + LandscapeContinuousLODParameters.Size - int2(1, 1), NeighborPos);
 
-		uint NeighborComponentIndex = (NeighborPos.y - LandscapeContinuousLODParameters.Min.y) * LandscapeContinuousLODParameters.Size.x + (NeighborPos.x - LandscapeContinuousLODParameters.Min.x);
-		float SectionLOD = GetViewLODData(NeighborComponentIndex);
+	// lookup neighbor LOD
+	uint NeighborComponentIndex = (NeighborPos.y - LandscapeContinuousLODParameters.Min.y) * LandscapeContinuousLODParameters.Size.x + (NeighborPos.x - LandscapeContinuousLODParameters.Min.x);
+	float NeighborLOD = GetViewLODData(NeighborComponentIndex);
 
-		NeighborLod[i] = max(SectionLOD, GetSectionLod(ComponentIndex));
-		NeighborLod[i] = NeighborLod[i] + View.FarShadowStaticMeshLODBias;
-		NeighborLod[i] = min(NeighborLod[i], LandscapeParameters.LastLOD);
-	}
-	return float4(NeighborLod[0], NeighborLod[1], NeighborLod[2], NeighborLod[3]);
-#endif
+	// the Edge LOD is the max of Our LOD and the Neighbor LOD (both sides will agree on EdgeLOD this way)
+	float EdgeLOD = max(CenterLOD, NeighborLOD);
+
+	// apply override and shadow bias
+	EdgeLOD = EdgeLOD + View.FarShadowStaticMeshLODBias;
+
+	// clamp to last lod
+	EdgeLOD = min(EdgeLOD, LandscapeParameters.LastLOD);
+
+	return EdgeLOD;
+#endif // !FIXED_GRID
 }
 
-/** Calculate fractional render LOD for a subsection. */
+/** Calculate fractional render LOD for a vertex within a subsection. */
 float CalcLOD(uint ComponentIndex, float2 xyLocalToSubsection, float2 Subsection)
 {
 #if FIXED_GRID
@@ -131,20 +120,25 @@ float CalcLOD(uint ComponentIndex, float2 xyLocalToSubsection, float2 Subsection
 	return LandscapeFixedGrid.LodValues.x;
 #else
 	float2 xy = (xyLocalToSubsection + Subsection) / LandscapeParameters.NumSubsections;
-	// LOD value based on Barycentric coordinate
-	float4 L0 = float4(xy.y, xy.x, (1 - xy.x), (1 - xy.y)) * 2;
-
-	float4 LODCalculated4 = L0 * GetSectionLod(ComponentIndex) + (float4(1, 1, 1, 1) - L0) * GetNeighborSectionLod(ComponentIndex);
-
-	// Find which quadrant of the subsection we're in - top, left, right or bottom.
-	// 0 --------- 1 X  1 ----- 2  1          1 
-	//   | \ 0 / |        \   /    | \      / | 
-	//   | 1 + 2 |          0      |  0    0  |      0       
-	//   | / 3 \ |                 | /      \ |    /   \     
-	// 2 --------- 3               2          2  1 ----- 2   
-	float LODCalculated = ((xy.x + xy.y) > 1) ?
-		(xy.x < xy.y ? LODCalculated4.w : LODCalculated4.z) :
-		(xy.x < xy.y ? LODCalculated4.y : LODCalculated4.x);
+	
+	// calculate blend value (normalized distance from section border: 1 at the border, 0 at the center)
+	float2 Delta = xy * 2 - 1;
+	float2 AbsDelta = abs(Delta);
+	float LB = max(AbsDelta.x, AbsDelta.y);
+
+	// shift the blend towards 0 to reduce the lod blend range
+	float k = LandscapeParameters.InvLODBlendRange;
+	LB = saturate(1 - (1-LB) * k);
+
+	// determine which neighbor is closest
+	int2 NeighborOffset = (AbsDelta.x > AbsDelta.y) ? int2(sign(Delta.x), 0) : int2(0, sign(Delta.y));
+
+	// calculate lod for this section, and the closest neighbor
+	float CenterLod = GetSectionLod(ComponentIndex);
+	float NeighborLod = GetNeighborSectionLodFromOffset(NeighborOffset, CenterLod);
+
+	// blend from center lod to the neighbor lod, based on distance from the border
+	float LODCalculated = (1-LB) * CenterLod + LB * NeighborLod;
 
 	return LODCalculated;
 #endif
diff --git a/Engine/Source/Runtime/Landscape/Classes/LandscapeProxy.h b/Engine/Source/Runtime/Landscape/Classes/LandscapeProxy.h
--- a/Engine/Source/Runtime/Landscape/Classes/LandscapeProxy.h
+++ b/Engine/Source/Runtime/Landscape/Classes/LandscapeProxy.h
@@ -480,6 +480,10 @@ public:
 	UPROPERTY(EditAnywhere, Category = "LOD Distribution", meta = (DisplayName = "Other LODs", ClampMin = "1.0", ClampMax = "10.0", UIMin = "1.0", UIMax = "10.0", LandscapeInherited))
 	float LODDistributionSetting;
 
+	/** This controls the area that blends LOD between neighboring sections. At 1.0 it blends across the entire section, and lower numbers reduce the blend region to be closer to the boundary. */
+	UPROPERTY(EditAnywhere, Category = "LOD Distribution", meta = (DisplayName = "Blend Range", ClampMin = "0.01", ClampMax = "1.0", UIMin = "0.01", UIMax = "1.0", LandscapeInherited))
+	float LODBlendRange = 1.0f;
+
 #if WITH_EDITORONLY_DATA
 	/** LOD level to use when exporting the landscape to obj or FBX */
 	UPROPERTY(EditAnywhere, Category=LOD, AdvancedDisplay)
diff --git a/Engine/Source/Runtime/Landscape/Private/LandscapeRender.cpp b/Engine/Source/Runtime/Landscape/Private/LandscapeRender.cpp
--- a/Engine/Source/Runtime/Landscape/Private/LandscapeRender.cpp
+++ b/Engine/Source/Runtime/Landscape/Private/LandscapeRender.cpp
@@ -99,6 +99,17 @@ FAutoConsoleVariableRef CVarLandscapeLODDistributionOverride(
 	FConsoleVariableDelegate::CreateStatic(&OnLODDistributionScaleChanged),
 	ECVF_Cheat
 );
+
+float GLandscapeLODBlendRangeOverride = -1.f;
+FAutoConsoleVariableRef CVarLandscapeLODBlendRangeOverride(
+	TEXT("r.Landscape.Override.LODBlendRange"),
+	GLandscapeLODBlendRangeOverride,
+	TEXT("When > 0, force the LODBlendRange property on all landscapes"),
+	FConsoleVariableDelegate::CreateStatic(&OnLODDistributionScaleChanged),
+	ECVF_Cheat
+);
+#else
+constexpr float GLandscapeLODBlendRangeOverride = -1.f;
 #endif // !UE_BUILD_SHIPPING
 
 float GLandscapeLOD0DistributionScale = 1.f;
@@ -1112,6 +1123,7 @@ FLandscapeComponentSceneProxy::FLandscapeComponentSceneProxy(ULandscapeComponent
 	, LastLOD(MaxLOD)
 	, ComponentMaxExtend(0.0f)
 	, ComponentSquaredScreenSizeToUseSubSections(FMath::Square(InComponent->GetLandscapeProxy()->ComponentScreenSizeToUseSubSections))
+	, InvLODBlendRange(1.0f / FMath::Max(0.01f, GLandscapeLODBlendRangeOverride > 0.0f ? GLandscapeLODBlendRangeOverride : InComponent->GetLandscapeProxy()->LODBlendRange))
 	, NumSubsections(InComponent->NumSubsections)
 	, SubsectionSizeQuads(InComponent->SubsectionSizeQuads)
 	, SubsectionSizeVerts(InComponent->SubsectionSizeQuads + 1)
@@ -2014,6 +2026,7 @@ void FLandscapeComponentSceneProxy::OnTransformChanged(FRHICommandListBase& RHIC
 	LandscapeParams.HeightmapUVScaleBias = HeightmapScaleBias;
 	LandscapeParams.WeightmapUVScaleBias = WeightmapScaleBias;
 	LandscapeParams.LocalToWorldNoScaling = FMatrix44f(LocalToWorldNoScaling);			// LWC_TODO: Precision loss
+	LandscapeParams.InvLODBlendRange = InvLODBlendRange;
 
 	LandscapeParams.LandscapeLightmapScaleBias = FVector4f(
 		LightmapScaleX,
diff --git a/Engine/Source/Runtime/Landscape/Public/LandscapeRender.h b/Engine/Source/Runtime/Landscape/Public/LandscapeRender.h
--- a/Engine/Source/Runtime/Landscape/Public/LandscapeRender.h
+++ b/Engine/Source/Runtime/Landscape/Public/LandscapeRender.h
@@ -107,6 +107,7 @@ BEGIN_GLOBAL_SHADER_PARAMETER_STRUCT(FLandscapeUniformShaderParameters, LANDSCAP
 	SHADER_PARAMETER(int32, NumSubsections)
 	SHADER_PARAMETER(int32, LastLOD)
 	SHADER_PARAMETER(uint32, VirtualTexturePerPixelHeight)
+	SHADER_PARAMETER(float, InvLODBlendRange)
 	SHADER_PARAMETER(FVector4f, HeightmapTextureSize)
 	SHADER_PARAMETER(FVector4f, HeightmapUVScaleBias)
     SHADER_PARAMETER(FVector4f, WeightmapUVScaleBias)
@@ -744,6 +745,7 @@ protected:
 	int32						LastVirtualTextureLOD;
 	float						ComponentMaxExtend; 		// The max extend value in any axis
 	float						ComponentSquaredScreenSizeToUseSubSections; // Size at which we start to draw in sub lod if LOD are different per sub section
+	float						InvLODBlendRange;			// 1.0 / LODBlendRange;
 
 	FLandscapeRenderSystem::LODSettingsComponent LODSettings;
 
