Fixed fwidth with Nanite CS materials

diff --git a/Engine/Shaders/Private/MaterialTemplate.ush b/Engine/Shaders/Private/MaterialTemplate.ush
--- a/Engine/Shaders/Private/MaterialTemplate.ush
+++ b/Engine/Shaders/Private/MaterialTemplate.ush
@@ -137,7 +137,7 @@
 	#define clip(x)
 	#if COMPUTE_SHADED
 		// Threads are remapped into 2x2 quad layout, so quad read x/y/d work for ddx/ddy
-		#define fwidth(x) 0
+		// https://microsoft.github.io/DirectX-Specs/d3d/HLSL_SM_6_6_Derivatives.html
 	#elif COMPUTESHADER
 		#define ddx(x) 0
 		#define ddy(x) 0
