Add dynamic ray traced translucent shadows.
1. Change the default to not trace translucent materials for shadow to align well with Lumen where translucent material does not cast shadow.
2. Use 'r.RayTracing.Shadows.Translucency 1' to turn on the dynamic ray traced translucent shadow. The shadow term is the accumulated visibility based on opacity.

diff --git a/Engine/Shaders/Private/RayTracing/RayTracingCommon.ush b/Engine/Shaders/Private/RayTracing/RayTracingCommon.ush
--- a/Engine/Shaders/Private/RayTracing/RayTracingCommon.ush
+++ b/Engine/Shaders/Private/RayTracing/RayTracingCommon.ush
@@ -601,6 +601,8 @@ struct FPackedMaterialClosestHitPayload : FMinimalPayload
 	void SetIgnoreTranslucentMaterials()	{ Flags |= RAY_TRACING_PAYLOAD_INPUT_FLAG_IGNORE_TRANSLUCENT; }
 	bool IsIgnoreTranslucentMaterials()		{ return (Flags & RAY_TRACING_PAYLOAD_INPUT_FLAG_IGNORE_TRANSLUCENT) != 0; }
 
+	bool IsShadowRay()						{ return (Flags & RAY_TRACING_PAYLOAD_INPUT_FLAG_SHADOW_RAY) != 0; }
+
 	FRayCone GetRayCone()					{ return UnpackRayCone(PackedRayCone); }
 	void SetRayCone(FRayCone NewRayCone)	{ PackedRayCone = PackRayCone(NewRayCone); }
 
@@ -772,6 +774,8 @@ struct FPackedMaterialClosestHitPayload : FMinimalPayload
 	void SetIgnoreTranslucentMaterials() { IorAndShadingModelIDAndBlendingModeAndFlagsAndPrimitiveLightingChannelMask |= RAY_TRACING_PAYLOAD_INPUT_FLAG_IGNORE_TRANSLUCENT << 23; }
 	bool IsIgnoreTranslucentMaterials() { return ((IorAndShadingModelIDAndBlendingModeAndFlagsAndPrimitiveLightingChannelMask >> 23) & RAY_TRACING_PAYLOAD_INPUT_FLAG_IGNORE_TRANSLUCENT) != 0; }
 
+	bool IsShadowRay() { return ((IorAndShadingModelIDAndBlendingModeAndFlagsAndPrimitiveLightingChannelMask >> 23) & RAY_TRACING_PAYLOAD_INPUT_FLAG_SHADOW_RAY) != 0; }
+
 	FRayCone GetRayCone() { return UnpackRayCone(PackedRayCone); }
 	void SetRayCone(FRayCone NewRayCone) { PackedRayCone = PackRayCone(NewRayCone); }
 
@@ -853,6 +857,7 @@ struct FPackedMaterialClosestHitPayload : FMinimalPayload
 	}
 
 	float GetOpacity()       { return f16tof32(BaseColorAndOpacity[1] >> 16); }
+	void SetOpacity(float Opacity) { BaseColorAndOpacity[1] = f32tof16(Opacity) << 16 | (BaseColorAndOpacity[1] & 0xFF); }
 	float GetMetallic()      { return float(MetallicAndSpecularAndRoughness & 0xFF) / 255.0f; }
 	float GetSpecular()      { return float((MetallicAndSpecularAndRoughness >> 8) & 0xFF) / 255.0f; }
 	float GetRoughness()     { return f16tof32(MetallicAndSpecularAndRoughness >> 16); }
@@ -1094,6 +1099,9 @@ FMinimalPayload TraceVisibilityRay(
 	
 	PackedPayload.SetFlags(RAY_TRACING_PAYLOAD_INPUT_FLAG_SHADOW_RAY);
 
+	//Have shadow ray behaves the same to Lumen, not tracing translucent materials.
+	PackedPayload.SetIgnoreTranslucentMaterials();
+
 	TraceVisibilityRayPacked(PackedPayload, TLAS, RayFlags, InstanceInclusionMask, Ray);
 
 	// Unpack the payload
@@ -1110,6 +1118,65 @@ FMinimalPayload TraceVisibilityRay(
 
 #if !COMPUTESHADER
 
+///////////////////////////////////////////////////////////////////////////////////////////////////
+// Tracing translucent visibility ray functions
+
+struct FTranslucentMinimalPayload : FMinimalPayload
+{
+	float Opacity; // The opacity for the translucent material.
+};
+
+void TraceTranslucentVisibilityRayPacked(
+	inout FPackedMaterialClosestHitPayload PackedPayload,
+	in RaytracingAccelerationStructure TLAS,
+	in uint RayFlags,
+	in uint InstanceInclusionMask,
+	in FRayDesc Ray)
+{
+	const uint RayContributionToHitGroupIndex = RAY_TRACING_SHADER_SLOT_SHADOW;
+	const uint MultiplierForGeometryContributionToShaderIndex = RAY_TRACING_NUM_SHADER_SLOTS;
+	const uint MissShaderIndex = 0;
+
+	// By enabling minimal payload mode all other payload information is ignored, meaning these functions need no payload inputs
+	PackedPayload.SetMinimalPayloadMode();
+	PackedPayload.HitT = 0;
+
+	// @todo: add mask to select translucent shadow only.
+	RayFlags |= RAY_TRACING_MASK_TRANSLUCENT;
+
+	// Trace the ray
+	TraceRay(
+		TLAS,
+		RayFlags,
+		InstanceInclusionMask,
+		RayContributionToHitGroupIndex,
+		MultiplierForGeometryContributionToShaderIndex,
+		MissShaderIndex,
+		Ray.GetNativeDesc(),
+		PackedPayload);
+}
+
+FTranslucentMinimalPayload TraceTranslucentVisibilityRay(
+	in RaytracingAccelerationStructure TLAS,
+	in uint RayFlags,
+	in uint InstanceInclusionMask,
+	in FRayDesc Ray)
+{
+	FPackedMaterialClosestHitPayload PackedPayload = (FPackedMaterialClosestHitPayload)0;
+
+	PackedPayload.SetFlags(RAY_TRACING_PAYLOAD_INPUT_FLAG_SHADOW_RAY);
+
+	TraceTranslucentVisibilityRayPacked(PackedPayload, TLAS, RayFlags, InstanceInclusionMask, Ray);
+
+	// Unpack the payload
+	FTranslucentMinimalPayload MinimalPayload = (FTranslucentMinimalPayload)0;
+
+	MinimalPayload.HitT = PackedPayload.HitT;
+	MinimalPayload.Opacity = PackedPayload.GetOpacity();
+
+	return MinimalPayload;
+}
+
 void TraceMaterialRayPacked(
 	inout FPackedMaterialClosestHitPayload Payload,
 	in RaytracingAccelerationStructure TLAS,
diff --git a/Engine/Shaders/Private/RayTracing/RayTracingMaterialHitShaders.usf b/Engine/Shaders/Private/RayTracing/RayTracingMaterialHitShaders.usf
--- a/Engine/Shaders/Private/RayTracing/RayTracingMaterialHitShaders.usf
+++ b/Engine/Shaders/Private/RayTracing/RayTracingMaterialHitShaders.usf
@@ -508,10 +508,17 @@ RAY_TRACING_ENTRY_CLOSEST_HIT(MaterialCHS,
 
 	if (PackedPayload.IsMinimalPayloadMode())
 	{
-		// Minimal payload mode only fills FMinimalPayload::HitT, skipping actual material evaluation.
-		// This mode is used when tracing shadow rays against masked geometry.
-		// Dynamic branch is used to avoid compiling an extra shader permutation.
-		return;
+#if MATERIALBLENDING_TRANSLUCENT
+		//For shadow ray that does not ignore translucent materials, the material will be evaluated
+		// even with minimal payload mode to query the opacity for translucent shadow.
+		if (PackedPayload.IsIgnoreTranslucentMaterials() || !PackedPayload.IsShadowRay())
+#endif
+		{
+			// Minimal payload mode only fills FMinimalPayload::HitT, skipping actual material evaluation.
+			// This mode is used when tracing shadow rays against masked geometry.
+			// Dynamic branch is used to avoid compiling an extra shader permutation.
+			return;
+		}
 	}
 
 	CurrentPayloadInputFlags = PackedPayload.GetFlags();
diff --git a/Engine/Shaders/Private/RayTracing/RayTracingOcclusionRGS.usf b/Engine/Shaders/Private/RayTracing/RayTracingOcclusionRGS.usf
--- a/Engine/Shaders/Private/RayTracing/RayTracingOcclusionRGS.usf
+++ b/Engine/Shaders/Private/RayTracing/RayTracingOcclusionRGS.usf
@@ -52,6 +52,8 @@ uint TransmissionMeanFreePathType;
 uint RejectionSamplingTrials;
 uint bAcceptFirstHit;
 uint bTwoSidedGeometry;
+uint TranslucentShadow;
+uint MaxTranslucencyHitCount;
 
 #define TransmissionSampling_ConstantTrackingInfiniteDomain 0
 #define TransmissionSampling_ConstantTrackingFiniteDomain 1
@@ -465,6 +467,7 @@ FOcclusionResult ComputeOcclusion(
 		uint RayFlagsForOpaque = bAcceptFirstHit != 0 ? RAY_FLAG_ACCEPT_FIRST_HIT_AND_END_SEARCH : 0;
 
 		FMinimalPayload MinimalPayload;
+		const bool NeedToTraceTranslucentShadow = TranslucentShadow == 1 && MaxTranslucencyHitCount != 0;
 
 #if AVOID_SELF_INTERSECTION_TRACE
 		{
@@ -557,9 +560,69 @@ FOcclusionResult ComputeOcclusion(
 		}
 		else
 		{
-			Out.ClosestRayDistance = (Out.ClosestRayDistance == DENOISER_INVALID_HIT_DISTANCE) ? DENOISER_MISS_HIT_DISTANCE : Out.ClosestRayDistance;
-			Out.TransmissionDistance += 1.0;
-			Out.Visibility += 1.0;
+
+#if !COMPUTESHADER // Ignore translucent shadow path when using inline ray tracing
+
+			bool bHasClosestEffectiveTranslucentHit = false;
+
+			if (NeedToTraceTranslucentShadow)
+			{
+				// When there is no intersection, test for translucent material hit if translucent shadow
+				// is supported. If there is a hit, record the accumulated opacity.
+				FTranslucentMinimalPayload TranslucentMinimalPayload = (FTranslucentMinimalPayload)0;
+				float Visibility = 1.0f;
+				float HitT = 0;
+				int TranslucentRayTriangleHitCount = 0;
+				FRayDesc TranslucentRay = Ray;
+				
+				while (MaxTranslucencyHitCount != 0 &&
+					TranslucentRay.TMin < TranslucentRay.TMax &&
+					TranslucentRayTriangleHitCount != MaxTranslucencyHitCount &&
+					Visibility > 0.0f)
+				{
+					TranslucentMinimalPayload = TraceTranslucentVisibilityRay(
+						TLAS,
+						RayFlags, // Accept first hit will cause artifact.
+						RaytracingMask,
+						TranslucentRay);
+
+					if (!TranslucentMinimalPayload.IsHit())
+					{
+						break;
+					}
+
+					TranslucentRay.TMin = asfloat(asuint(TranslucentMinimalPayload.HitT) + 1);
+					TranslucentRayTriangleHitCount++;
+
+					// Record first hit with opacity > 0
+					if (!bHasClosestEffectiveTranslucentHit && 
+						TranslucentMinimalPayload.Opacity > 0)
+					{
+						HitT = TranslucentMinimalPayload.HitT;
+						bHasClosestEffectiveTranslucentHit = true;
+					}
+
+					Visibility *= 1 - TranslucentMinimalPayload.Opacity;
+				}
+
+				if (bHasClosestEffectiveTranslucentHit)
+				{
+					Out.ClosestRayDistance =
+						(Out.ClosestRayDistance == DENOISER_INVALID_HIT_DISTANCE) ||
+						(HitT < Out.ClosestRayDistance) ? HitT : Out.ClosestRayDistance;
+					Out.SumRayDistance += HitT;
+					Out.HitCount += 1.0f;
+					Out.Visibility += Visibility;
+				}
+			}
+			
+			if (!bHasClosestEffectiveTranslucentHit)
+#endif
+			{
+				Out.ClosestRayDistance = (Out.ClosestRayDistance == DENOISER_INVALID_HIT_DISTANCE) ? DENOISER_MISS_HIT_DISTANCE : Out.ClosestRayDistance;
+				Out.TransmissionDistance += 1.0;
+				Out.Visibility += 1.0;
+			}
 		}
 	}
 #if !ENABLE_MULTIPLE_SAMPLES_PER_PIXEL
diff --git a/Engine/Source/Runtime/Renderer/Private/RayTracing/RayTracingShadows.cpp b/Engine/Source/Runtime/Renderer/Private/RayTracing/RayTracingShadows.cpp
--- a/Engine/Source/Runtime/Renderer/Private/RayTracing/RayTracingShadows.cpp
+++ b/Engine/Source/Runtime/Renderer/Private/RayTracing/RayTracingShadows.cpp
@@ -109,6 +109,26 @@ static TAutoConsoleVariable<int32> CVarRayTracingShadowsAcceptFirstHit(
 	ECVF_RenderThreadSafe
 );
 
+static TAutoConsoleVariable<int32> CVarRayTracingShadowsTranslucency(
+	TEXT("r.RayTracing.Shadows.Translucency"),
+	0,
+	TEXT("0: Translucent material will not cast shadow (by default).")
+	TEXT("1: Translucent material cast approximate translucent shadows based on opacity (Very expensive)."),
+	ECVF_RenderThreadSafe
+);
+static TAutoConsoleVariable<int32> CVarRayTracingShadowsMaxTranslucencyHitCount(
+	TEXT("r.RayTracing.Shadows.MaxTranslucencyHitCount"),
+	-1,
+	TEXT("-1: Evaluate all intersections (default).")
+	TEXT(" 0: Disable translucent shadow testing.")
+	TEXT(">0: Limit the number of intersections."),
+	ECVF_RenderThreadSafe
+);
+
+int32 GetRayTracingShadowsMaxTranslucencyHitCount()
+{
+	return CVarRayTracingShadowsMaxTranslucencyHitCount.GetValueOnRenderThread();
+}
 
 bool EnableRayTracingShadowTwoSidedGeometry()
 {
@@ -189,6 +209,8 @@ class FOcclusionRGS : public FGlobalShader
 		SHADER_PARAMETER(uint32, RejectionSamplingTrials)
 		SHADER_PARAMETER(uint32, bAcceptFirstHit)
 		SHADER_PARAMETER(uint32, bTwoSidedGeometry)
+		SHADER_PARAMETER(uint32, TranslucentShadow)
+		SHADER_PARAMETER(uint32, MaxTranslucencyHitCount)
 
 		SHADER_PARAMETER_STRUCT(FLightShaderParameters, Light)
 		SHADER_PARAMETER_STRUCT_INCLUDE(FSceneTextureParameters, SceneTextures)
@@ -395,6 +417,8 @@ void FDeferredShadingSceneRenderer::RenderRayTracingShadows(
 		CommonPassParameters->AvoidSelfIntersectionTraceDistance = GRayTracingShadowsAvoidSelfIntersectionTraceDistance;
 		CommonPassParameters->bAcceptFirstHit = CVarRayTracingShadowsAcceptFirstHit.GetValueOnRenderThread();
 		CommonPassParameters->bTwoSidedGeometry = EnableRayTracingShadowTwoSidedGeometry() ? 1 : 0;
+		CommonPassParameters->TranslucentShadow = CVarRayTracingShadowsTranslucency.GetValueOnRenderThread();
+		CommonPassParameters->MaxTranslucencyHitCount = GetRayTracingShadowsMaxTranslucencyHitCount();
 		CommonPassParameters->TLAS = View.GetRayTracingSceneLayerViewChecked(ERayTracingSceneLayer::Base);
 		CommonPassParameters->ViewUniformBuffer = View.ViewUniformBuffer;
 		CommonPassParameters->SceneTextures = SceneTextures;
