Add groom instance bound legend.

diff --git a/Engine/Plugins/Runtime/HairStrands/Source/HairStrandsCore/Private/GroomDebug.cpp b/Engine/Plugins/Runtime/HairStrands/Source/HairStrandsCore/Private/GroomDebug.cpp
--- a/Engine/Plugins/Runtime/HairStrands/Source/HairStrandsCore/Private/GroomDebug.cpp
+++ b/Engine/Plugins/Runtime/HairStrands/Source/HairStrandsCore/Private/GroomDebug.cpp
@@ -846,6 +846,7 @@ public:
 		// Skip optimization for avoiding long compilation time due to large UAV writes
 		FGlobalShader::ModifyCompilationEnvironment(Parameters, OutEnvironment);
 		OutEnvironment.CompilerFlags.Add(CFLAG_Debug);
+		OutEnvironment.CompilerFlags.Add(CFLAG_ForceDXC);
 		OutEnvironment.SetDefine(TEXT("SHADER_PRINT_INSTANCE"), 1);
 		OutEnvironment.CompilerFlags.Add(CFLAG_ForceDXC);
 	}
@@ -1120,6 +1121,7 @@ public:
 		// Skip optimization for avoiding long compilation time due to large UAV writes
 		FGlobalShader::ModifyCompilationEnvironment(Parameters, OutEnvironment);
 		OutEnvironment.CompilerFlags.Add(CFLAG_Debug);
+		OutEnvironment.CompilerFlags.Add(CFLAG_ForceDXC);
 		OutEnvironment.SetDefine(TEXT("SHADER_PRINT_MEMORY"), 1);
 		OutEnvironment.CompilerFlags.Add(CFLAG_ForceDXC);
 	}
diff --git a/Engine/Shaders/Private/HairStrands/HairStrandsDebug.usf b/Engine/Shaders/Private/HairStrands/HairStrandsDebug.usf
--- a/Engine/Shaders/Private/HairStrands/HairStrandsDebug.usf
+++ b/Engine/Shaders/Private/HairStrands/HairStrandsDebug.usf
@@ -4,12 +4,20 @@
 #include "../HairShadingCommon.ush"
 #include "../ColorMap.ush"
 
-#if SHADER_PRINT_INSTANCE || SHADER_PRINT || SHADER_PRINT_MEMORY || SHADER_CLUSTERAABB
+#if SHADER_PRINT_INSTANCE || SHADER_PRINT || SHADER_PRINT_MEMORY || SHADER_CLUSTERAABB || SHADER_DOMINFO
 /////////////////////////////////////////////////////////////////////////////////////////
 // Utils function
 
 #include "../ShaderPrint.ush"
 
+// Shared color
+float4 GetGPUSceneBoundColor()						{ return ColorDarkGreen; }
+float4 GetInstanceGroupScreenBoundColor()			{ return ColorGreen; }
+
+float4 GetGPUAABBInstanceGroupBoundColor()			{ return ColorDarkOrange; }
+float4 GetGPUAABBMacroGroupBoundColor()				{ return ColorRed; }
+float4 GetGPUAABBMacroGroupBoundVoxelAlignedColor()	{ return ColorLightRed; }
+
 // Print float number in 0..1000 aligned on the right
 void PrintNumberAlignedRight(inout FShaderPrintContext Context, float Value, FFontColor InColor)
 {
@@ -80,7 +88,7 @@ void PrintPercentage(inout FShaderPrintContext Context, float Value, FFontColor
 	Print(Context, TEXT("% "), InColor);
 }
 
-#endif // SHADER_PRINT_INSTANCE || SHADER_PRINT
+#endif // SHADER_PRINT_INSTANCE || SHADER_PRINT || SHADER_PRINT_MEMORY || SHADER_CLUSTERAABB || SHADER_DOMINFO
 
 /////////////////////////////////////////////////////////////////////////////////////////
 // Visualize hair strands data as full screen view mode (BaseColor, Tangent, Roughness, ...)
@@ -777,10 +785,11 @@ void MainCS(uint3 DispatchThreadId : SV_DispatchThreadID)
 	for (uint MacroGroupIt = 0; MacroGroupIt < MacroGroupCount; ++MacroGroupIt)
 	{
 		const FHairAABB Bound = ReadHairAABB(MacroGroupIt, MacroGroupAABBBuffer);
-		AddAABBTWS(Bound.Min, Bound.Max, float4(1, 0.5f, 0, 1));
+		AddAABBTWS(Bound.Min, Bound.Max, GetGPUAABBMacroGroupBoundVoxelAlignedColor());
 	}
 
 	const float4 FrustumColor = float4(0.5f, 1, 0, 1);
+	LOOP
 	for (uint SlotIt = 0; SlotIt < AllocatedSlotCount; ++SlotIt)
 	{
 		const float4x4 TranslatedWorldToClip = ShadowViewInfoBuffer[SlotIt].TranslatedWorldToClip;
@@ -912,8 +921,7 @@ void DrawInstanceAABB(uint PrimitiveId)
 	if (InstanceID != ~0)
 	{
 		FInstanceSceneData InstanceData = GetInstanceSceneData(InstanceID, Scene.GPUScene.InstanceDataSOAStride);
-		const float4 InstanceBoundColor = ColorLightGreen;
-		AddOBBWS(InstanceData.LocalBoundsCenter - InstanceData.LocalBoundsExtent, InstanceData.LocalBoundsCenter + InstanceData.LocalBoundsExtent, InstanceBoundColor, LWCHackToFloat(InstanceData.LocalToWorld));
+		AddOBBWS(InstanceData.LocalBoundsCenter - InstanceData.LocalBoundsExtent, InstanceData.LocalBoundsCenter + InstanceData.LocalBoundsExtent, GetGPUSceneBoundColor(), LWCHackToFloat(InstanceData.LocalToWorld));
 		AddReferentialWS(LWCHackToFloat(InstanceData.LocalToWorld), 50.f);
 	}
 }
@@ -928,9 +936,8 @@ void MainCS(uint InGroupIndex : SV_GroupIndex, uint3 DispatchThreadId : SV_Dispa
 		{
 			for (uint GroupIndex =0;GroupIndex< HairMacroGroupCount; GroupIndex++)
 			{
-				const float4 InstanceBoundColor = float4(0.25f, 1.f, 0.25f, 1.0f); 
 				FHairAABB Bound = ReadHairAABB(GroupIndex, HairMacroGroupAABBBuffer);
-				AddAABBTWS(Bound.Min, Bound.Max, InstanceBoundColor);
+				AddAABBTWS(Bound.Min, Bound.Max, GetGPUAABBMacroGroupBoundColor());
 			}
 		}
 
@@ -938,9 +945,8 @@ void MainCS(uint InGroupIndex : SV_GroupIndex, uint3 DispatchThreadId : SV_Dispa
 		{
 			for (uint GroupIndex = 0; GroupIndex < HairMacroGroupCount; GroupIndex++)
 			{
-				const float4 InstanceBoundColor = float4(1.f, 0.5f, 0.5f, 1.0f);
 				FHairAABB Bound = ReadHairAABB(GroupIndex, HairMacroGroupVoxelAlignedAABBBuffer);
-				AddAABBTWS(Bound.Min, Bound.Max, InstanceBoundColor);
+				AddAABBTWS(Bound.Min, Bound.Max, GetGPUAABBMacroGroupBoundVoxelAlignedColor());
 			}
 		}
 
@@ -1359,6 +1365,7 @@ void MainCS(uint GroupIndex : SV_GroupIndex, uint3 DispatchThreadId : SV_Dispatc
 	Print(Context, TEXT("   Group  LOD     Geometry Binding Cache Sim RBF Clip Screen CovSc. Strands            Points             Bind.Tri/Sect Vis. Shad. SwLOD RT  Name"));
 	Newline(Context);
 
+	LOOP
 	for (uint InstanceIndex = 0; InstanceIndex < InstanceCount; ++InstanceIndex)
 	{
 		const FInstanceInfo Info = ReadInstanceInfo(InstanceIndex);
@@ -1479,6 +1486,7 @@ void MainCS(uint GroupIndex : SV_GroupIndex, uint3 DispatchThreadId : SV_Dispatc
 		Print(Context, TEXT("Root  "), FontWhite);
 		Print(Context, TEXT("Tip   "), FontWhite);
 		const uint AttributeNameInfoCount = NumAttributeNames();
+		LOOP
 		for (uint AttributeIt=0; AttributeIt < AttributeNameInfoCount; ++AttributeIt)
 		{
 			PrintAttributeNames(Context, AttributeIt, FontWhite);
@@ -1486,6 +1494,7 @@ void MainCS(uint GroupIndex : SV_GroupIndex, uint3 DispatchThreadId : SV_Dispatc
 		}
 		Newline(Context);
 
+		LOOP
 		for (uint StrandInstanceIndex = 0; StrandInstanceIndex < InstanceCount; ++StrandInstanceIndex)
 		{
 			const FInstanceInfo Info = ReadInstanceInfo(StrandInstanceIndex);
@@ -1530,6 +1539,16 @@ void MainCS(uint GroupIndex : SV_GroupIndex, uint3 DispatchThreadId : SV_Dispatc
 			Newline(Context);
 		}
 	}
+
+	// Bounds legends
+	{
+		FShaderPrintContext CtxL = InitShaderPrintContext(true, uint2(900, 100));
+		Print(CtxL, TEXT("Bound - GPU Scene"), 							FFontColor(GetGPUSceneBoundColor().xyz)); Newline(CtxL);
+		Print(CtxL, TEXT("Bound - Instance Group Screen"), 				FFontColor(GetInstanceGroupScreenBoundColor().xyz)); Newline(CtxL);
+		Print(CtxL, TEXT("Bound - AABB Instance Group"), 				FFontColor(GetGPUAABBInstanceGroupBoundColor().xyz)); Newline(CtxL);
+		Print(CtxL, TEXT("Bound - AABB Macro Group"), 					FFontColor(GetGPUAABBMacroGroupBoundColor().xyz)); Newline(CtxL);
+		Print(CtxL, TEXT("Bound - AABB Macro Group (voxel-aligned)"),	FFontColor(GetGPUAABBMacroGroupBoundVoxelAlignedColor().xyz)); Newline(CtxL);
+	}
 }
 #else  // PERMUTATION_OUTPUT_TYPE
 {
@@ -1537,7 +1556,7 @@ void MainCS(uint GroupIndex : SV_GroupIndex, uint3 DispatchThreadId : SV_Dispatc
 	{
 		// Bound
 		const FHairAABB AABB = ReadHairAABB(0, InstanceAABB);
-		AddAABBTWS(AABB.Min, AABB.Max, ColorDarkGreen);
+		AddAABBTWS(AABB.Min, AABB.Max, GetGPUAABBInstanceGroupBoundColor());
 
 		// Screen size
 		FShaderPrintContext Context = InitShaderPrintContext(true, uint2(50, 50));
@@ -2723,6 +2742,7 @@ void MainCS(uint GroupIndex : SV_GroupIndex, uint3 DispatchThreadId : SV_Dispatc
 		Print(Ctx, TEXT("Components"), FontYellow); Newline(Ctx);
 		{
 			Print(Ctx, TEXT("Grp. Common Guides Strands Cards Meshes Total  Curv Name"), FontSilver); Newline(Ctx);
+			LOOP
 			for (uint It = 0; It < ComponentCount; ++It)
 			{
 				const FComponent Info = GetComponent(It);
@@ -2747,7 +2767,6 @@ void MainCS(uint GroupIndex : SV_GroupIndex, uint3 DispatchThreadId : SV_Dispatc
 		}
 		Newline(Ctx);
 	}
-	// Total
 
 	// Grooms
 	if (bGroom)
@@ -2756,6 +2775,7 @@ void MainCS(uint GroupIndex : SV_GroupIndex, uint3 DispatchThreadId : SV_Dispatc
 		if (!bStrandsOnly)
 		{
 			Print(Ctx, TEXT("Grp. Guides Strands Cards Meshes Total Curv Name"), FontSilver); Newline(Ctx);
+			LOOP
 			for (uint It = 0; It < GroomCount; ++It)
 			{
 				const FGroom Info = GetGroom(It);
@@ -2776,6 +2796,7 @@ void MainCS(uint GroupIndex : SV_GroupIndex, uint3 DispatchThreadId : SV_Dispatc
 			FFontColor FontCurve = FontSilver;
 
 			Print(Ctx, TEXT("Grp. Rest.  Intrp. Cluster  RT.  Total  Curv Name"), FontSilver); Newline(Ctx);
+			LOOP
 			for (uint It = 0; It < GroomCount; ++It)
 			{
 				const FGroom Info = GetGroom(It);
@@ -2809,13 +2830,13 @@ void MainCS(uint GroupIndex : SV_GroupIndex, uint3 DispatchThreadId : SV_Dispatc
 		}
 		Newline(Ctx);
 	}
-	// Total
 
 	// Bindings
 	if (bBinding)
 	{
 		Print(Ctx, TEXT("Bindings"), FontYellow); Newline(Ctx);
 		Print(Ctx, TEXT("Grp. Guides Strands Cards Meshes Total Curv Name"), FontSilver); Newline(Ctx);
+		LOOP
 		for (uint It = 0; It < BindingCount; ++It)
 		{
 			const FBinding Info = GetBinding(It);
@@ -2838,7 +2859,6 @@ void MainCS(uint GroupIndex : SV_GroupIndex, uint3 DispatchThreadId : SV_Dispatc
 		}
 		Newline(Ctx);
 	}
-	// Total
 }
 
 #endif // SHADER_PRINT_MEMORY
diff --git a/Engine/Shaders/Private/ShaderPrintCommon.ush b/Engine/Shaders/Private/ShaderPrintCommon.ush
--- a/Engine/Shaders/Private/ShaderPrintCommon.ush
+++ b/Engine/Shaders/Private/ShaderPrintCommon.ush
@@ -211,6 +211,7 @@ FFontColor GetDefaultFontColor()							{ FFontColor Out; Out.Color = float3(1,1,
 #define ColorDarkRed	float4(0.50f, 0.25f, 0.25f, 1.0f)
 #define ColorDarkGreen	float4(0.25f, 0.50f, 0.25f, 1.0f)
 #define ColorDarkBlue	float4(0.25f, 0.25f, 0.50f, 1.0f)
+#define ColorDarkOrange	float4(ColorOrange.xyz * 0.5f,1.0f)
 
 ///////////////////////////////////////////////////////////////////////////////////////////////////
 // Util pack/unpack functions
diff --git a/Engine/Source/Runtime/Renderer/Private/HairStrands/HairStrandsDebug.cpp b/Engine/Source/Runtime/Renderer/Private/HairStrands/HairStrandsDebug.cpp
--- a/Engine/Source/Runtime/Renderer/Private/HairStrands/HairStrandsDebug.cpp
+++ b/Engine/Source/Runtime/Renderer/Private/HairStrands/HairStrandsDebug.cpp
@@ -685,6 +685,8 @@ public:
 	static void ModifyCompilationEnvironment(const FGlobalShaderPermutationParameters& Parameters, FShaderCompilerEnvironment& OutEnvironment)
 	{
 		FGlobalShader::ModifyCompilationEnvironment(Parameters, OutEnvironment);
+		OutEnvironment.CompilerFlags.Add(CFLAG_Debug);
+		OutEnvironment.CompilerFlags.Add(CFLAG_ForceDXC);
 		OutEnvironment.SetDefine(TEXT("SHADER_DOMINFO"), 1);
 	}
 };
