Temporarily disable shader code that causes compiling stalls for further investigation.

diff --git a/Engine/Shaders/Private/RayTracing/RayTracingOcclusionRGS.usf b/Engine/Shaders/Private/RayTracing/RayTracingOcclusionRGS.usf
--- a/Engine/Shaders/Private/RayTracing/RayTracingOcclusionRGS.usf
+++ b/Engine/Shaders/Private/RayTracing/RayTracingOcclusionRGS.usf
@@ -561,8 +561,8 @@ FOcclusionResult ComputeOcclusion(
 		else
 		{
 
-#if !COMPUTESHADER // Ignore translucent shadow path when using inline ray tracing
-
+#if !COMPUTESHADER && 0 // Ignore translucent shadow path when using inline ray tracing
+			// Temporarily disable the code to investigate a compiling issue.
 			bool bHasClosestEffectiveTranslucentHit = false;
 
 			if (NeedToTraceTranslucentShadow)
