Substrate - fixed surfaces identified as using post process subsurface when they should not.
Also fixed DiffuseIndirectComposite to add on the UAV as this is possible here.

diff --git a/Engine/Shaders/Private/BasePassPixelShader.usf b/Engine/Shaders/Private/BasePassPixelShader.usf
--- a/Engine/Shaders/Private/BasePassPixelShader.usf
+++ b/Engine/Shaders/Private/BasePassPixelShader.usf
@@ -1669,17 +1669,6 @@ void FPixelShaderInOut_MainPS(
 		Out.SubstrateTopLayerData = SubstratePackTopLayerData(TopLayerData);
 	#if SUBSTRATE_OUTPUT_ROUGH_REFRACTION
 		OpaqueRoughRefractionData.OpaqueRoughRefractionEnabled = true; // As long as we output to rough refraction, we want it to be enabled for the tiling process.
-	#if SUBSTRATE_USES_CONVERSION_FROM_LEGACY
-		// When running legacy support, we want to only enable rough refraction for multi layered materials (eye, clear coat) or SSS.
-		// This is because the compiler will always see a vertical layer. But we do not want to enable opaque rough refraction for all the legacy materials.
-		if (SubstratePixelHeader.SubstrateTree.BSDFCount > 1)
-		{
-			// We only enable opaque rough refraction if the material is coated or has SSS.
-			// BSDFs and Operators indices must match SubstrateLegacyConversion.ush.
-			OpaqueRoughRefractionData.OpaqueRoughRefractionEnabled = (SubstratePixelHeader.SubstrateTree.Operators[2].Weight > 0.0f && SubstratePixelHeader.SubstrateTree.Operators[3].Weight > 0.0f) 
-				|| BSDF_GETSSSTYPE(SubstratePixelHeader.SubstrateTree.BSDFs[0]) != SSS_TYPE_INVALID || BSDF_GETSSSTYPE(SubstratePixelHeader.SubstrateTree.BSDFs[1]) != SSS_TYPE_INVALID;
-		}
-	#endif
 		SubstrateStruct.OpaqueRoughRefractionTextureUAV[PixelPos] = SubstratePackOpaqueRoughRefractionData(OpaqueRoughRefractionData);
 	#endif
 
diff --git a/Engine/Shaders/Private/DiffuseIndirectComposite.usf b/Engine/Shaders/Private/DiffuseIndirectComposite.usf
--- a/Engine/Shaders/Private/DiffuseIndirectComposite.usf
+++ b/Engine/Shaders/Private/DiffuseIndirectComposite.usf
@@ -299,11 +299,11 @@ FSubstrateDeferredLighting SubstrateIndirectLighting(
 		#endif 
 
 		// SSS Checkerboard
-		OutBSDF.Specular *= GetSSSCheckerboadSpecularScale(PixelPos, SubstrateEnvLight.bSubsurface);
+		OutBSDF.Specular *= GetSSSCheckerboadSpecularScale(PixelPos, SubstrateEnvLight.bPostProcessSubsurface);
 
 		FLightAccumulator Accumulator = (FLightAccumulator)0;
-		LightAccumulator_AddSplit(Accumulator, OutBSDF.Diffuse, OutBSDF.Specular, OutBSDF.Diffuse, 1.f /*CommonMultiplier*/, SubstrateEnvLight.bSubsurface);
-		AccumulateSubstrateDeferredLighting(Out, Accumulator, SubstrateEnvLight.bSubsurface, BSDF_GETISTOPLAYER(BSDF));
+		LightAccumulator_AddSplit(Accumulator, OutBSDF.Diffuse, OutBSDF.Specular, OutBSDF.Diffuse, 1.f /*CommonMultiplier*/, SubstrateEnvLight.bPostProcessSubsurface);
+		AccumulateSubstrateDeferredLighting(Out, Accumulator, SubstrateEnvLight.bPostProcessSubsurface, BSDF_GETISTOPLAYER(BSDF));
 	}
 
 	return Out;
@@ -475,8 +475,8 @@ void MainPS(
 					OutAddColor = IndirectLighting_Substrate.SceneColor;
 				#if SUBSTRATE_OPAQUE_ROUGH_REFRACTION_ENABLED
 					const uint2 OutCoord = PixelPos;
-					OutOpaqueRoughRefractionSceneColor[OutCoord] = IndirectLighting_Substrate.OpaqueRoughRefractionSceneColor;
-					OutSubSurfaceSceneColor[OutCoord] = IndirectLighting_Substrate.SubSurfaceSceneColor;
+					OutOpaqueRoughRefractionSceneColor[OutCoord] += IndirectLighting_Substrate.OpaqueRoughRefractionSceneColor;
+					OutSubSurfaceSceneColor[OutCoord] += IndirectLighting_Substrate.SubSurfaceSceneColor;
 				#endif
 				}
 			}
diff --git a/Engine/Shaders/Private/Substrate/SubstrateDeferredLighting.ush b/Engine/Shaders/Private/Substrate/SubstrateDeferredLighting.ush
--- a/Engine/Shaders/Private/Substrate/SubstrateDeferredLighting.ush
+++ b/Engine/Shaders/Private/Substrate/SubstrateDeferredLighting.ush
@@ -94,8 +94,8 @@ FSubstrateDeferredLighting SubstrateDeferredLighting(
 		const float3 CommonMultiplier = LightData.Color * LightMask * LuminanceWeight(BSDFContext, BSDF);
 
 		FLightAccumulator Out = (FLightAccumulator)0;
-		LightAccumulator_AddSplit(Out, DiffuseLuminance, SpecularLuminance, DiffuseLuminance, CommonMultiplier, BSDFEvaluate.bSubsurface);
-		AccumulateSubstrateDeferredLighting(SubstrateLighting, Out, BSDFEvaluate.bSubsurface, BSDF_GETISTOPLAYER(BSDF));
+		LightAccumulator_AddSplit(Out, DiffuseLuminance, SpecularLuminance, DiffuseLuminance, CommonMultiplier, BSDFEvaluate.bPostProcessSubsurface);
+		AccumulateSubstrateDeferredLighting(SubstrateLighting, Out, BSDFEvaluate.bPostProcessSubsurface, BSDF_GETISTOPLAYER(BSDF));
 	}
 
 	return SubstrateLighting;
diff --git a/Engine/Shaders/Private/Substrate/SubstrateEnvironmentLighting.ush b/Engine/Shaders/Private/Substrate/SubstrateEnvironmentLighting.ush
--- a/Engine/Shaders/Private/Substrate/SubstrateEnvironmentLighting.ush
+++ b/Engine/Shaders/Private/Substrate/SubstrateEnvironmentLighting.ush
@@ -106,8 +106,8 @@ FSubstrateDeferredLighting SubstrateEnvironmentLighting(
 		// SUBSTRATE_TODO ReflectionEnvironment(...) for static reflection volumes.
 
 		FLightAccumulator Out = (FLightAccumulator)0;
-		LightAccumulator_AddSplit(Out, DiffuseLighting, SpecularLighting, DiffuseLighting, View.PreExposure, SubstrateEnvLight.bSubsurface);
-		AccumulateSubstrateDeferredLighting(SubstrateLighting, Out, SubstrateEnvLight.bSubsurface, BSDF_GETISTOPLAYER(BSDF));
+		LightAccumulator_AddSplit(Out, DiffuseLighting, SpecularLighting, DiffuseLighting, View.PreExposure, SubstrateEnvLight.bPostProcessSubsurface);
+		AccumulateSubstrateDeferredLighting(SubstrateLighting, Out, SubstrateEnvLight.bPostProcessSubsurface, BSDF_GETISTOPLAYER(BSDF));
 	}
 
 	SSRReductionFactor = saturate(SSRReductionFactor);
diff --git a/Engine/Shaders/Private/Substrate/SubstrateEvaluation.ush b/Engine/Shaders/Private/Substrate/SubstrateEvaluation.ush
--- a/Engine/Shaders/Private/Substrate/SubstrateEvaluation.ush
+++ b/Engine/Shaders/Private/Substrate/SubstrateEvaluation.ush
@@ -304,7 +304,7 @@ struct FSubstrateEvaluateResult
 
 	float3 ThroughputV;					// Throughput to the next layer (from the view to the shading point). Contains many things such as transmittance or cloth energy conservation/preservation factors.
 	float3 TransmittanceAlongN;			// Transmittance to the next layer along the surface normal. This is used later to compute the transmittance towards the light.
-	bool   bSubsurface;					// True if we need to separate the subsurface light contribution for the screen space diffusion process.
+	bool   bPostProcessSubsurface;		// True if we need to separate the subsurface light contribution for the screen space diffusion process.
 	bool   bApplyProjectedSolidAngle;	// True if the Saturate(NoL) factor should be applied or not over the value of the BSDF
 };
 
@@ -389,7 +389,7 @@ FSubstrateEvaluateResult SubstrateEvaluateBSDFCommon(FSubstrateBSDFContext BSDFC
 			Sample.IntegratedDiffuseValue += (ShadowTerms.SurfaceShadow * AreaLightContext.NoL * AreaLightContext.Falloff) * Sample.DiffusePathValue * AreaLightContext.AreaLight.FalloffColor;
 			Sample.DiffuseColor		= DiffuseColor;
 			Sample.DiffusePDF		= BSDFContext.SatNoL / PI;
-			Sample.bSubsurface		= BSDF_GETSSSTYPE(BSDFContext.BSDF) == SSS_TYPE_DIFFUSION || BSDF_GETSSSTYPE(BSDFContext.BSDF) == SSS_TYPE_DIFFUSION_PROFILE;
+			Sample.bPostProcessSubsurface = false;
 
 
 
@@ -848,15 +848,19 @@ FSubstrateEvaluateResult SubstrateEvaluateBSDFCommon(FSubstrateBSDFContext BSDFC
 			////
 			//// Evaluate transmitted light through a mesh due to sub surface scattering.
 			////
-			#if SUBSTRATE_FASTPATH==0 && SUBSTRATE_SSS_TRANSMISSION
-			if (BSDF_GETSSSTYPE(BSDFContext.BSDF) == SSS_TYPE_DIFFUSION_PROFILE || BSDF_GETSSSTYPE(BSDFContext.BSDF) == SSS_TYPE_DIFFUSION)
+			#if SUBSTRATE_FASTPATH==0
+			const bool bIsSubsurfaceDiffusionProfile = BSDF_GETSSSTYPE(BSDFContext.BSDF) == SSS_TYPE_DIFFUSION_PROFILE;
+			const bool bIsSubsurfaceDiffusionMFP = BSDF_GETSSSTYPE(BSDFContext.BSDF) == SSS_TYPE_DIFFUSION;
+			Sample.bPostProcessSubsurface = bIsSubsurfaceDiffusionProfile || bIsSubsurfaceDiffusionMFP;
+			#if SUBSTRATE_SSS_TRANSMISSION
+			if (bIsSubsurfaceDiffusionProfile || bIsSubsurfaceDiffusionMFP)
 			{
 				float ThicknessInCm = DecodeThickness(ShadowTerms.TransmissionThickness) * SSSS_MAX_TRANSMISSION_PROFILE_DISTANCE;
 
 				float3 TransmissionThroughput = 1.0f;
 				float PhaseFunction = 0.0f;
 				BRANCH
-				if (BSDF_GETSSSTYPE(BSDFContext.BSDF) == SSS_TYPE_DIFFUSION_PROFILE)
+				if (bIsSubsurfaceDiffusionProfile)
 				{
 					ThicknessInCm = BSDF_GETISTHIN(BSDFContext.BSDF) ? SLAB_SSSPROFILETHICKNESSCM(BSDFContext.BSDF) : ThicknessInCm;
 
@@ -874,7 +878,8 @@ FSubstrateEvaluateResult SubstrateEvaluateBSDFCommon(FSubstrateBSDFContext BSDFC
 				else
 				{
 					// When the surface 'IsThin' MFP is rescaled to SUBSTRATE_SIMPLEVOLUME_THICKNESS_CM
-					ThicknessInCm = BSDF_GETISTHIN(BSDFContext.BSDF) ? SUBSTRATE_SIMPLEVOLUME_THICKNESS_CM : ThicknessInCm;
+					const bool bIsThin = BSDF_GETISTHIN(BSDFContext.BSDF);
+					ThicknessInCm = bIsThin ? SUBSTRATE_SIMPLEVOLUME_THICKNESS_CM : ThicknessInCm;
 
 					const float3 MeanFreePathInCm = SLAB_SSSMFP(BSDFContext.BSDF);
 					const float3 SubsurfaceAlebdo = SLAB_DIFFUSEALBEDO(BSDFContext.BSDF);
@@ -894,6 +899,7 @@ FSubstrateEvaluateResult SubstrateEvaluateBSDFCommon(FSubstrateBSDFContext BSDFC
 				Sample.IntegratedDiffuseValue += (ShadowTerms.TransmissionShadow * AreaLightContext.Falloff) * AreaLightContext.AreaLight.FalloffColor * Sample.TransmissionPathValue;
 			}
 			#endif
+			#endif
 
 			////
 			//// Evaluate a layer of participating media: scattering and transmittance.
@@ -956,8 +962,6 @@ FSubstrateEvaluateResult SubstrateEvaluateBSDFCommon(FSubstrateBSDFContext BSDFC
 				Sample.ThroughputV			= lerp(Sample.ThroughputV,			SlabTransmittanceV    * SpecularTransmissionV,	DiffuseToVolumeBlend);
 				Sample.TransmittanceAlongN	= lerp(Sample.TransmittanceAlongN,	SlabTransmittanceN,								DiffuseToVolumeBlend);
 
-				Sample.bSubsurface = false;	// It should already be the case because Enforce in this case because BSDF_GETHASSSS should be false with SSS_TYPE_SIMPLEVOLUME is true
-
 				// Use Context.NoL which is not saturate, and allows to handle backface lighting, necessary for medium support
 				#if SUBSTRATE_TRANSLUCENT_ENABLED
 				const float MediumNoL = abs(BSDFContext.Context.NoL);
@@ -1053,7 +1057,7 @@ FSubstrateEvaluateResult SubstrateEvaluateBSDFCommon(FSubstrateBSDFContext BSDFC
 
 			Sample.DiffuseColor		= DiffuseColor;
 			Sample.DiffusePDF		= BSDFContext.SatNoL / PI;
-			Sample.bSubsurface		= BSDF_GETSSSTYPE(BSDFContext.BSDF) == SSS_TYPE_DIFFUSION || BSDF_GETSSSTYPE(BSDFContext.BSDF) == SSS_TYPE_DIFFUSION_PROFILE;
+			Sample.bPostProcessSubsurface = BSDF_GETSSSTYPE(BSDFContext.BSDF) == SSS_TYPE_DIFFUSION || BSDF_GETSSSTYPE(BSDFContext.BSDF) == SSS_TYPE_DIFFUSION_PROFILE;
 
 			//
 			// Apply energy conservation on the diffuse component
@@ -1436,7 +1440,7 @@ struct FSubstrateEnvLightResult
 	// This can vary depending if the simple clear coat is used.
 	float3 SSRSpecularWeight;
 
-	bool   bSubsurface; // True if we need to separate the subsurface light contribution for the screen space diffusion process.
+	bool   bPostProcessSubsurface; // True if we need to separate the subsurface light contribution for the screen space diffusion process.
 };
 
 FSubstrateEnvLightResult SubstrateEvaluateForEnvLight(FSubstrateBSDFContext BSDFContext, bool bEnableSpecular, FSubstrateIntegrationSettings Settings)
@@ -1456,7 +1460,7 @@ FSubstrateEnvLightResult SubstrateEvaluateForEnvLight(FSubstrateBSDFContext BSDF
 			F90 *= F0RGBToMicroOcclusion(F0);
 
 			SubstrateEnvLightResult.DiffuseNormal = BSDFContext.N;
-			SubstrateEnvLightResult.bSubsurface = BSDF_GETSSSTYPE(BSDFContext.BSDF) != SSS_TYPE_INVALID;
+			SubstrateEnvLightResult.bPostProcessSubsurface = false;
 
 			float SafeRoughness = MakeRoughnessSafe(SLAB_ROUGHNESS(BSDFContext.BSDF));
 
@@ -1484,6 +1488,11 @@ FSubstrateEnvLightResult SubstrateEvaluateForEnvLight(FSubstrateBSDFContext BSDF
 			}
 
 			#if SUBSTRATE_FASTPATH==0
+			const bool bIsSubsurfaceDiffusionProfile = BSDF_GETSSSTYPE(BSDFContext.BSDF) == SSS_TYPE_DIFFUSION_PROFILE;
+			const bool bIsSubsurfaceDiffusionMFP = BSDF_GETSSSTYPE(BSDFContext.BSDF) == SSS_TYPE_DIFFUSION;
+			const bool bIsThin = BSDF_GETISTHIN(BSDFContext.BSDF);
+			SubstrateEnvLightResult.bPostProcessSubsurface = bIsSubsurfaceDiffusionProfile || bIsSubsurfaceDiffusionMFP;
+
 			BRANCH
 			if (BSDF_GETSSSTYPE(BSDFContext.BSDF) == SSS_TYPE_SIMPLEVOLUME)
 			{
@@ -1521,7 +1530,6 @@ FSubstrateEnvLightResult SubstrateEvaluateForEnvLight(FSubstrateBSDFContext BSDF
 
 				SubstrateEnvLightResult.DiffuseWeight = lerp(SubstrateEnvLightResult.DiffuseWeight, SlabDirectionalAlbedo, DiffuseToVolumeBlend);
 				// We keep the same normal
-				SubstrateEnvLightResult.bSubsurface = false;
 			}
 			#endif
 			
@@ -1695,9 +1703,6 @@ FSubstrateEnvLightResult SubstrateEvaluateForEnvLight(FSubstrateBSDFContext BSDF
 					SubstrateEnvLightResult.DiffuseColor  += DirectionalSubSurfaceColor;
 					SubstrateEnvLightResult.DiffuseWeight += DirectionalSubSurfaceColor;
 				}
-			
-				// We do not want thin (a.k.a. two sided lighting) or the non thin path (a.k.a. subsurface) contribution to go the SSS post process path.
-				SubstrateEnvLightResult.bSubsurface = false;
 			}
 			#endif
 
@@ -1781,8 +1786,8 @@ FSubstrateEnvLightResult SubstrateEvaluateForEnvLight(FSubstrateBSDFContext BSDF
 				EnvBRDFApproxFullyRough(DiffuseColor, F0, F90);
 			}
 
-			SubstrateEnvLightResult.DiffuseNormal = BSDFContext.N;
-			SubstrateEnvLightResult.bSubsurface = BSDF_GETSSSTYPE(BSDFContext.BSDF) != SSS_TYPE_INVALID; // TODO profile
+			SubstrateEnvLightResult.DiffuseNormal			= BSDFContext.N;
+			SubstrateEnvLightResult.bPostProcessSubsurface	= BSDF_GETSSSTYPE(BSDFContext.BSDF) == SSS_TYPE_DIFFUSION || BSDF_GETSSSTYPE(BSDFContext.BSDF) == SSS_TYPE_DIFFUSION_PROFILE;
 
 			float SafeRoughness = MakeRoughnessSafe(SLAB_ROUGHNESS(BSDFContext.BSDF));
 
@@ -1793,17 +1798,17 @@ FSubstrateEnvLightResult SubstrateEvaluateForEnvLight(FSubstrateBSDFContext BSDF
 				const float DiffuseDirectionalAlbedo = 1.0f;
 				EvalEnvBRDF									= SpecularEnergyTerms.E;
 				DirectionalAlbedo_SpecularTransmission		= ComputeEnergyPreservation(SpecularEnergyTerms);
-				SubstrateEnvLightResult.DiffuseWeight			= DiffuseColor * DiffuseDirectionalAlbedo * DirectionalAlbedo_SpecularTransmission;
-				SubstrateEnvLightResult.DiffuseColor			= DiffuseColor;
+				SubstrateEnvLightResult.DiffuseWeight		= DiffuseColor * DiffuseDirectionalAlbedo * DirectionalAlbedo_SpecularTransmission;
+				SubstrateEnvLightResult.DiffuseColor		= DiffuseColor;
 			}
 
 			BRANCH
 			if (bEnableSpecular)
 			{
-				SubstrateEnvLightResult.SpecularDirection = SubstrateGetOffSpecularPeakReflectionDir(BSDFContext.N, BSDFContext.R, SafeRoughness);
-				SubstrateEnvLightResult.SpecularWeight = EvalEnvBRDF;
-				SubstrateEnvLightResult.SpecularSafeRoughness = SafeRoughness;
-				SubstrateEnvLightResult.SpecularColor = F0;
+				SubstrateEnvLightResult.SpecularDirection		= SubstrateGetOffSpecularPeakReflectionDir(BSDFContext.N, BSDFContext.R, SafeRoughness);
+				SubstrateEnvLightResult.SpecularWeight			= EvalEnvBRDF;
+				SubstrateEnvLightResult.SpecularSafeRoughness	= SafeRoughness;
+				SubstrateEnvLightResult.SpecularColor			= F0;
 			}
 
 			break;
