Substrate - fixed translucent contact shadow causing non matching visual with legacy.
ShadowTerm is now evaluated in SubstrateDeferredLighting. With shadow sampled if requested at source (AO is always used to initialize the Shadow term).
The light ShadowTerm is evaluated first for all the closures for the common parts. Then we evaluate and combine that the shadow parts that are dependent on the CLosure: contact shadow.

diff --git a/Engine/Shaders/Private/ClusteredDeferredShadingPixelShader.usf b/Engine/Shaders/Private/ClusteredDeferredShadingPixelShader.usf
--- a/Engine/Shaders/Private/ClusteredDeferredShadingPixelShader.usf
+++ b/Engine/Shaders/Private/ClusteredDeferredShadingPixelShader.usf
@@ -157,7 +157,9 @@ float4 GetSimpleLightLighting(
 	#if SUBSTRATE_ENABLED && USE_HAIR_LIGHTING == 0
 		FDeferredLightData LightData = ConvertToDeferredLight(LocalLight);
 		LightData.ShadowedBits = 1;
-		const FShadowTerms ShadowTerms = { SubstrateGetAO(SubstratePixelHeader), 1.0, 1.0, InitHairTransmittanceData() };
+
+		FSubstrateShadowTermInputParameters SubstrateShadowTermInputParameters	= GetInitialisedSubstrateShadowTermInputParameters();
+		SubstrateShadowTermInputParameters.bEvaluateShadowTerm					= false;
 
 		const float3 ToLight = LightData.TranslatedWorldPosition - TranslatedWorldPosition;
 		const float3 L = normalize(ToLight);
@@ -170,7 +172,7 @@ float4 GetSimpleLightLighting(
 			L,
 			ToLight,
 			LightMask,
-			ShadowTerms,
+			SubstrateShadowTermInputParameters,
 			Substrate.MaterialTextureArray,
 			SubstrateAddressing,
 			SubstratePixelHeader).SceneColor;
@@ -308,16 +310,18 @@ float4 GetLightGridLocalLighting(
 		#endif
 
 	#if SUBSTRATE_ENABLED && USE_HAIR_LIGHTING == 0
-		FShadowTerms ShadowTerms = { SubstrateGetAO(SubstratePixelHeader), 1.0, 1.0, InitHairTransmittanceData() };
 
-		const uint FakeShadingModelID = 0;
-		const float FakeContactShadowOpacity = 1.0f;
-		const float SceneDepth = CalcSceneDepth(ScreenUV);
 		const float3 ToLight = LightData.TranslatedWorldPosition - TranslatedWorldPosition;
 		const float3 L = normalize(ToLight);
 		const float3 V = -CameraVector;
-		float4 PrecomputedShadowFactors = SubstrateReadPrecomputedShadowFactors(SubstratePixelHeader, PixelPos, SceneTexturesStruct.GBufferETexture);
-		GetShadowTerms(SceneDepth, PrecomputedShadowFactors, FakeShadingModelID, FakeContactShadowOpacity, LightData, TranslatedWorldPosition, L, LightAttenuation, Dither, ShadowTerms);
+
+		FSubstrateShadowTermInputParameters SubstrateShadowTermInputParameters	= GetInitialisedSubstrateShadowTermInputParameters();
+		SubstrateShadowTermInputParameters.bEvaluateShadowTerm					= true;
+		SubstrateShadowTermInputParameters.SceneDepth							= CalcSceneDepth(ScreenUV);
+		SubstrateShadowTermInputParameters.PrecomputedShadowFactors				= SubstrateReadPrecomputedShadowFactors(SubstratePixelHeader, PixelPos, SceneTexturesStruct.GBufferETexture);
+		SubstrateShadowTermInputParameters.TranslatedWorldPosition				= TranslatedWorldPosition;
+		SubstrateShadowTermInputParameters.LightAttenuation						= LightAttenuation;
+		SubstrateShadowTermInputParameters.Dither								= Dither;
 
 		const float LightMask = 1;
 
@@ -329,7 +333,7 @@ float4 GetLightGridLocalLighting(
 			L,
 			ToLight,
 			LightMask,
-			ShadowTerms,
+			SubstrateShadowTermInputParameters,
 			Substrate.MaterialTextureArray,
 			SubstrateAddressing,
 			SubstratePixelHeader);
diff --git a/Engine/Shaders/Private/DeferredLightPixelShaders.usf b/Engine/Shaders/Private/DeferredLightPixelShaders.usf
--- a/Engine/Shaders/Private/DeferredLightPixelShaders.usf
+++ b/Engine/Shaders/Private/DeferredLightPixelShaders.usf
@@ -277,15 +277,13 @@ void DeferredLightPixelMain(
 
 		if (LightMask > 0)
 		{
-			FShadowTerms ShadowTerms = { SubstrateGetAO(SubstratePixelHeader), 1.0, 1.0, InitHairTransmittanceData() };
-			float4 LightAttenuation = GetLightAttenuationFromShadow(InputParams, SceneDepth, DerivedParams.TranslatedWorldPosition);
-
-			float Dither = InterleavedGradientNoise(InputParams.PixelPos, View.StateFrameIndexMod8);
-			const uint FakeShadingModelID = 0;
-			const float FakeContactShadowOpacity = 1.0f;
-			float4 PrecomputedShadowFactors = SubstrateReadPrecomputedShadowFactors(SubstratePixelHeader, PixelPos, SceneTexturesStruct.GBufferETexture);
-			GetShadowTerms(SceneDepth, PrecomputedShadowFactors, FakeShadingModelID, FakeContactShadowOpacity,
-				LightData, DerivedParams.TranslatedWorldPosition, L, LightAttenuation, Dither, ShadowTerms);
+			FSubstrateShadowTermInputParameters SubstrateShadowTermInputParameters	= GetInitialisedSubstrateShadowTermInputParameters();
+			SubstrateShadowTermInputParameters.bEvaluateShadowTerm					= true;
+			SubstrateShadowTermInputParameters.SceneDepth							= SceneDepth;
+			SubstrateShadowTermInputParameters.PrecomputedShadowFactors				= SubstrateReadPrecomputedShadowFactors(SubstratePixelHeader, PixelPos, SceneTexturesStruct.GBufferETexture);
+			SubstrateShadowTermInputParameters.TranslatedWorldPosition				= DerivedParams.TranslatedWorldPosition;
+			SubstrateShadowTermInputParameters.LightAttenuation						= GetLightAttenuationFromShadow(InputParams, SceneDepth, DerivedParams.TranslatedWorldPosition);
+			SubstrateShadowTermInputParameters.Dither								= InterleavedGradientNoise(InputParams.PixelPos, View.StateFrameIndexMod8);
 
 			FSubstrateDeferredLighting SubstrateLighting = SubstrateDeferredLighting(
 				LightData,
@@ -293,7 +291,7 @@ void DeferredLightPixelMain(
 				L,
 				ToLight,
 				LightMask,
-				ShadowTerms,
+				SubstrateShadowTermInputParameters,
 				Substrate.MaterialTextureArray,
 				SubstrateAddressing,
 				SubstratePixelHeader);
diff --git a/Engine/Shaders/Private/DeferredLightingCommon.ush b/Engine/Shaders/Private/DeferredLightingCommon.ush
--- a/Engine/Shaders/Private/DeferredLightingCommon.ush
+++ b/Engine/Shaders/Private/DeferredLightingCommon.ush
@@ -75,11 +75,13 @@ float ShadowRayCast(
 #error "Must set SUPPORT_CONTACT_SHADOWS"
 #endif
 
-void GetShadowTerms(float SceneDepth, half4 PrecomputedShadowFactors, uint ShadingModelID, float ContactShadowOpacity, FDeferredLightData LightData, float3 TranslatedWorldPosition, half3 L, half4 LightAttenuation, float Dither, inout FShadowTerms Shadow)
+void GetShadowTermsBase(
+	float SceneDepth, 
+	half4 PrecomputedShadowFactors, 
+	FDeferredLightData LightData, 
+	half4 LightAttenuation, 
+	inout FShadowTerms OutShadow)
 {
-	float ContactShadowLength = 0.0f;
-	const float ContactShadowLengthScreenScale = GetTanHalfFieldOfView().y * SceneDepth;
-
 	BRANCH
 	if (LightData.ShadowedBits)
 	{
@@ -99,12 +101,12 @@ void GetShadowTerms(float SceneDepth, half4 PrecomputedShadowFactors, uint Shadi
 		{
 			// Remapping the light attenuation buffer (see ShadowRendering.cpp)
 
-			Shadow.SurfaceShadow = LightAttenuation.z * StaticShadowing;
+			OutShadow.SurfaceShadow = LightAttenuation.z * StaticShadowing;
 			// SSS uses a separate shadowing term that allows light to penetrate the surface
 			//@todo - how to do static shadowing of SSS correctly?
-			Shadow.TransmissionShadow = LightAttenuation.w * StaticShadowing;
+			OutShadow.TransmissionShadow = LightAttenuation.w * StaticShadowing;
 
-			Shadow.TransmissionThickness = LightAttenuation.w;
+			OutShadow.TransmissionThickness = LightAttenuation.w;
 		}
 		else
 		{
@@ -114,27 +116,43 @@ void GetShadowTerms(float SceneDepth, half4 PrecomputedShadowFactors, uint Shadi
 
 			float DynamicShadowFraction = DistanceFromCameraFade(SceneDepth, LightData);
 			// For a directional light, fade between static shadowing and the whole scene dynamic shadowing based on distance + per object shadows
-			Shadow.SurfaceShadow = lerp(LightAttenuation.x, StaticShadowing, DynamicShadowFraction);
+			OutShadow.SurfaceShadow = lerp(LightAttenuation.x, StaticShadowing, DynamicShadowFraction);
 			// Fade between SSS dynamic shadowing and static shadowing based on distance
-			Shadow.TransmissionShadow = min(lerp(LightAttenuation.y, StaticShadowing, DynamicShadowFraction), LightAttenuation.w);
+			OutShadow.TransmissionShadow = min(lerp(LightAttenuation.y, StaticShadowing, DynamicShadowFraction), LightAttenuation.w);
 
-			Shadow.SurfaceShadow *= LightAttenuation.z;
-			Shadow.TransmissionShadow *= LightAttenuation.z;
+			OutShadow.SurfaceShadow *= LightAttenuation.z;
+			OutShadow.TransmissionShadow *= LightAttenuation.z;
 
 			// Need this min or backscattering will leak when in shadow which cast by non perobject shadow(Only for directional light)
-			Shadow.TransmissionThickness = min(LightAttenuation.y, LightAttenuation.w);
-		}
-
-		FLATTEN
-		if (LightData.ShadowedBits > 1 && LightData.ContactShadowLength > 0)
-		{
-			ContactShadowLength = LightData.ContactShadowLength * (LightData.ContactShadowLengthInWS ? 1.0f : ContactShadowLengthScreenScale);
+			OutShadow.TransmissionThickness = min(LightAttenuation.y, LightAttenuation.w);
 		}
 	}
 
+	OutShadow.HairTransmittance = LightData.HairTransmittance;
+	OutShadow.HairTransmittance.OpaqueVisibility = OutShadow.SurfaceShadow;
+}
+
+void ApplyContactShadowWithShadowTerms(
+	float SceneDepth, 
+	uint ShadingModelID, 
+	float ContactShadowOpacity, 
+	FDeferredLightData LightData, 
+	float3 TranslatedWorldPosition, 
+	half3 L, 
+	float Dither, 
+	inout FShadowTerms OutShadow)
+{
 #if SUPPORT_CONTACT_SHADOWS
 
-#if SUBSTRATE_ENABLED == 0
+	float ContactShadowLength = 0.0f;
+	const float ContactShadowLengthScreenScale = GetTanHalfFieldOfView().y * SceneDepth;
+
+	FLATTEN
+	if (LightData.ShadowedBits > 1 && LightData.ContactShadowLength > 0)
+	{
+		ContactShadowLength = LightData.ContactShadowLength * (LightData.ContactShadowLengthInWS ? 1.0f : ContactShadowLengthScreenScale);
+	}
+
 	if (LightData.ShadowedBits < 2 && (ShadingModelID == SHADINGMODELID_HAIR))
 	{
 		ContactShadowLength = 0.2 * ContactShadowLengthScreenScale;
@@ -145,11 +163,10 @@ void GetShadowTerms(float SceneDepth, half4 PrecomputedShadowFactors, uint Shadi
 		ContactShadowLength = 0.5;
 		
 	}
-#endif
 
-	#if MATERIAL_CONTACT_SHADOWS
-		ContactShadowLength = 0.2 * ContactShadowLengthScreenScale;
-	#endif
+#if MATERIAL_CONTACT_SHADOWS
+	ContactShadowLength = 0.2 * ContactShadowLengthScreenScale;
+#endif
 
 	BRANCH
 	if (ContactShadowLength > 0.0)
@@ -162,7 +179,6 @@ void GetShadowTerms(float SceneDepth, half4 PrecomputedShadowFactors, uint Shadi
 		{
 			float ContactShadowOcclusion = bHitCastContactShadow ? LightData.ContactShadowCastingIntensity : LightData.ContactShadowNonCastingIntensity;
 
-#if SUBSTRATE_ENABLED == 0
 			// Exponential attenuation is not applied on hair/eye/SSS-profile here, as the hit distance (shading-point to blocker) is different from the estimated 
 			// thickness (closest-point-from-light to shading-point), and this creates light leaks. Instead we consider first hit as a blocker (old behavior)
 			BRANCH
@@ -178,19 +194,37 @@ void GetShadowTerms(float SceneDepth, half4 PrecomputedShadowFactors, uint Shadi
 				float Density = SubsurfaceDensityFromOpacity(ContactShadowOpacity);
 				ContactShadowOcclusion *= 1.0 - saturate( exp( -Density * HitDistance ) );
 			}
-#endif
 
 			float ContactShadow = 1.0 - ContactShadowOcclusion;
 
-			Shadow.SurfaceShadow *= ContactShadow;
-			Shadow.TransmissionShadow *= ContactShadow;
+			OutShadow.SurfaceShadow *= ContactShadow;
+			OutShadow.TransmissionShadow *= ContactShadow;
 		}
-		
 	}
-#endif
 
-	Shadow.HairTransmittance = LightData.HairTransmittance;
-	Shadow.HairTransmittance.OpaqueVisibility = Shadow.SurfaceShadow;
+	OutShadow.HairTransmittance = LightData.HairTransmittance;
+	OutShadow.HairTransmittance.OpaqueVisibility = OutShadow.SurfaceShadow;
+
+#endif // SUPPORT_CONTACT_SHADOWS
+}
+
+void GetShadowTerms(
+	float SceneDepth,
+	half4 PrecomputedShadowFactors,
+	uint ShadingModelID,
+	float ContactShadowOpacity,
+	FDeferredLightData LightData,
+	float3 TranslatedWorldPosition,
+	half3 L,
+	half4 LightAttenuation,
+	float Dither,
+	inout FShadowTerms OutShadow)
+{
+	// Get the basic shadow terms
+	GetShadowTermsBase(SceneDepth, PrecomputedShadowFactors, LightData, LightAttenuation, OutShadow);
+
+	// Now combined with screens space contact shadow if necessary.
+	ApplyContactShadowWithShadowTerms (SceneDepth, ShadingModelID, ContactShadowOpacity, LightData, TranslatedWorldPosition, L, Dither, OutShadow);
 }
 
 float GetLocalLightAttenuation(
diff --git a/Engine/Shaders/Private/RayTracing/RayTracingLightingMS.usf b/Engine/Shaders/Private/RayTracing/RayTracingLightingMS.usf
--- a/Engine/Shaders/Private/RayTracing/RayTracingLightingMS.usf
+++ b/Engine/Shaders/Private/RayTracing/RayTracingLightingMS.usf
@@ -89,13 +89,14 @@ RAY_TRACING_ENTRY_MISS(RayTracingLightingMS, FPackedMaterialClosestHitPayload, P
 	{
 		FSubstrateAddressing SubstrateAddressing = GetSubstratePixelDataByteOffset(0,0,0);
 		FSubstratePixelHeader SubstratePixelHeader = UnpackSubstrateHeaderIn(PackedPayload.SubstrateData, SubstrateAddressing, PackedPayload.SubstrateData);
-		FShadowTerms ShadowTerms = { SubstrateGetAO(SubstratePixelHeader), 1.0, 1.0, InitHairTransmittanceData() };
 
-		const float FakeSceneDepth = 1.f; // Do not use depth
-		const uint FakeShadingModelID = 0;
-		const float FakeContactShadowOpacity = 1.0f;
-		float4 FakePrecomputedShadowFactors = 1.f;
-		GetShadowTerms(FakeSceneDepth, FakePrecomputedShadowFactors, FakeShadingModelID, FakeContactShadowOpacity, LightData, WorldPosition, L, LightAttenuation, Dither, ShadowTerms);
+		FSubstrateShadowTermInputParameters SubstrateShadowTermInputParameters	= GetInitialisedSubstrateShadowTermInputParameters();
+		SubstrateShadowTermInputParameters.bEvaluateShadowTerm					= true;
+		SubstrateShadowTermInputParameters.SceneDepth							= 1.0f; // Do not use depth
+		SubstrateShadowTermInputParameters.PrecomputedShadowFactors				= 1.0f; // Do not use precomputed factors
+		SubstrateShadowTermInputParameters.TranslatedWorldPosition				= WorldPosition;
+		SubstrateShadowTermInputParameters.LightAttenuation						= LightAttenuation;
+		SubstrateShadowTermInputParameters.Dither								= Dither;
 
 		// RayTracing + Lighting + Substrate is too expenssive (as in uses too many registers) to be supported on some platofrms
 		FSubstrateDeferredLighting Lighting = SubstrateDeferredLighting(
@@ -104,7 +105,7 @@ RAY_TRACING_ENTRY_MISS(RayTracingLightingMS, FPackedMaterialClosestHitPayload, P
 			L,
 			ToLight,
 			LightMask,
-			ShadowTerms,
+			SubstrateShadowTermInputParameters,
 			PackedPayload.SubstrateData,
 			SubstrateAddressing,
 			SubstratePixelHeader);
diff --git a/Engine/Shaders/Private/ShadingFurnaceTest.usf b/Engine/Shaders/Private/ShadingFurnaceTest.usf
--- a/Engine/Shaders/Private/ShadingFurnaceTest.usf
+++ b/Engine/Shaders/Private/ShadingFurnaceTest.usf
@@ -228,14 +228,17 @@ void MainPS(
 					BRANCH
 					if (SubstratePixelHeader.BSDFCount > 0)
 					{
-						FShadowTerms ShadowTerms = { 1.0, 1.0, 1.0, InitHairTransmittanceData() };
+						HEADER_SETIRRADIANCE_AO(SubstratePixelHeader.State, SubstratePackIrradianceAndOcclusion(1.0, 1.0));
+
+						FSubstrateShadowTermInputParameters SubstrateShadowTermInputParameters	= GetInitialisedSubstrateShadowTermInputParameters();
+						SubstrateShadowTermInputParameters.bEvaluateShadowTerm					= false;
 						SampleRadiance = SubstrateDeferredLighting(
 							LightData,
 							V,
 							L,
 							L,
 							LightMask,
-							ShadowTerms,
+							SubstrateShadowTermInputParameters,
 							RectTexture,
 							Substrate.MaterialTextureArray,
 							SubstrateAddressing,
diff --git a/Engine/Shaders/Private/Substrate/Substrate.ush b/Engine/Shaders/Private/Substrate/Substrate.ush
--- a/Engine/Shaders/Private/Substrate/Substrate.ush
+++ b/Engine/Shaders/Private/Substrate/Substrate.ush
@@ -3253,10 +3253,16 @@ float SubstrateGetLegacyShadingModels(in FSubstrateBSDF BSDF)
 	{
 		// Base on features, we can only partially recovers the original shading models. Shading models like foliage, eye, clear coat are impossible/harder to infer.
 		if (BSDF_GETSSSTYPE(BSDF) == SSS_TYPE_DIFFUSION_PROFILE){ return SHADINGMODELID_SUBSURFACE_PROFILE; }
-		else if (BSDF_GETSSSTYPE(BSDF) != SSS_TYPE_INVALID)		{ return SHADINGMODELID_SUBSURFACE; }
+		else if (BSDF_GETSSSTYPE(BSDF) == SSS_TYPE_WRAP)		
+		{ 
+			if (BSDF_GETISTHIN(BSDF))
+			{
+				return SHADINGMODELID_TWOSIDED_FOLIAGE;
+			}
+			return SHADINGMODELID_SUBSURFACE; 
+		}
 		else if (BSDF_GETHASFUZZ(BSDF))							{ return SHADINGMODELID_CLOTH; }
-		else if (BSDF_GETHASFUZZ(BSDF))							{ return SHADINGMODELID_TWOSIDED_FOLIAGE; }
-		else													{ return SHADINGMODELID_DEFAULT_LIT; }
+		else													{ return SHADINGMODELID_DEFAULT_LIT; } // Also include SSS_TYPE_DIFFUSION and SSS_TYPE_SIMPLE_VOLUME
 	}
 	case SUBSTRATE_BSDF_TYPE_HAIR: 								{ return SHADINGMODELID_HAIR; }
 	case SUBSTRATE_BSDF_TYPE_EYE: 									{ return SHADINGMODELID_EYE; }
diff --git a/Engine/Shaders/Private/Substrate/SubstrateDeferredLighting.ush b/Engine/Shaders/Private/Substrate/SubstrateDeferredLighting.ush
--- a/Engine/Shaders/Private/Substrate/SubstrateDeferredLighting.ush
+++ b/Engine/Shaders/Private/Substrate/SubstrateDeferredLighting.ush
@@ -26,6 +26,22 @@ float4 SubstrateReadPrecomputedShadowFactors(FSubstratePixelHeader SubstratePixe
 	return SubstratePixelHeader.HasZeroPrecShadowMask() ? 0.0f : 1.0f;
 }
 
+struct FSubstrateShadowTermInputParameters
+{
+	bool bEvaluateShadowTerm;
+	float SceneDepth;
+	float4 PrecomputedShadowFactors;
+	float3 TranslatedWorldPosition;
+	float4 LightAttenuation;
+	float Dither;
+};
+
+FSubstrateShadowTermInputParameters GetInitialisedSubstrateShadowTermInputParameters()
+{
+	FSubstrateShadowTermInputParameters Params = (FSubstrateShadowTermInputParameters)0;
+	return Params;
+}
+
 // Analytical lighting evaluation for Substrate material.
 // Unpack BSDF on-the-fly
 FSubstrateDeferredLighting SubstrateDeferredLighting(
@@ -34,7 +50,7 @@ FSubstrateDeferredLighting SubstrateDeferredLighting(
 	float3 L,
 	float3 ToLight,
 	float LightMask,
-	FShadowTerms ShadowTerms,
+	FSubstrateShadowTermInputParameters SubstrateShadowTermInputParameters,
 	FSubstrateMaterialContainer MaterialBuffer,
 	FSubstrateAddressing SubstrateAddressing,
 	FSubstratePixelHeader SubstratePixelHeader)
@@ -48,11 +64,42 @@ FSubstrateDeferredLighting SubstrateDeferredLighting(
 
 	FRectTexture RectTexture = ConvertToRectTexture(LightData);
 
+	// Get the basic shadow terms
+	FShadowTerms ShadowTerms = { SubstrateGetAO(SubstratePixelHeader), 1.0, 1.0, InitHairTransmittanceData() };
+	if (SubstrateShadowTermInputParameters.bEvaluateShadowTerm)
+	{
+		// Get the shadow term that is independent from Substrate closures.
+		GetShadowTermsBase(
+			SubstrateShadowTermInputParameters.SceneDepth,
+			SubstrateShadowTermInputParameters.PrecomputedShadowFactors,
+			LightData,
+			SubstrateShadowTermInputParameters.LightAttenuation,
+			ShadowTerms);
+	}
+
 	Substrate_for (uint ClosureIndex = 0, ClosureIndex < SubstratePixelHeader.ClosureCount, ++ClosureIndex)
 	{
 		// Unpack BSDF data
 		FSubstrateBSDF BSDF = UnpackSubstrateBSDF(MaterialBuffer, SubstrateAddressing, SubstratePixelHeader);
 
+		FShadowTerms BSDFShadowTerms = ShadowTerms;
+		if (SubstrateShadowTermInputParameters.bEvaluateShadowTerm)	// This is going to be a evaluation done at compile time
+		{
+			const float BSDFContactShadowOpacity= 1.f - abs(SLAB_SSSPHASEANISOTROPY(BSDF));
+			const uint  BSDFShadingModelID		= SubstrateGetLegacyShadingModels(BSDF);
+
+			// Now combine with shadow term that are dependent on the Substrate closures (screens space contact shadow) if necessary.
+			ApplyContactShadowWithShadowTerms(
+				SubstrateShadowTermInputParameters.SceneDepth,
+				BSDFShadingModelID,
+				BSDFContactShadowOpacity,
+				LightData,
+				SubstrateShadowTermInputParameters.TranslatedWorldPosition,
+				L,
+				SubstrateShadowTermInputParameters.Dither,
+				BSDFShadowTerms);
+		}
+
 		FSubstrateBSDFContext BSDFContext = SubstrateCreateBSDFContext(SubstratePixelHeader, BSDF, SubstrateAddressing, V, L);
 
 		float Roughness = SubstrateGetBSDFRoughness(BSDFContext.BSDF);
@@ -70,7 +117,7 @@ FSubstrateDeferredLighting SubstrateDeferredLighting(
 
 			// We must have the SubstrateIntegrateBSDF inside the if due to the rectlight texture: it must be non ambiguous which texture is going to be used.
 			// After the compilation, a local resource must map to a unique global resource (the default or the actual rect light texture).
-			BSDFEvaluate = SubstrateEvaluateBSDFCommon(BSDFContext, ShadowTerms, AreaLightContext, Settings, INTEGRATION_AREA_LIGHT_RECT);
+			BSDFEvaluate = SubstrateEvaluateBSDFCommon(BSDFContext, BSDFShadowTerms, AreaLightContext, Settings, INTEGRATION_AREA_LIGHT_RECT);
 		}
 		else
 		{
@@ -80,11 +127,11 @@ FSubstrateDeferredLighting SubstrateDeferredLighting(
 			BRANCH
 			if(IsAreaLight(AreaLightContext.AreaLight))
 			{
-				BSDFEvaluate = SubstrateEvaluateBSDFCommon(BSDFContext, ShadowTerms, AreaLightContext, Settings, INTEGRATION_AREA_LIGHT_CAPSULE);
+				BSDFEvaluate = SubstrateEvaluateBSDFCommon(BSDFContext, BSDFShadowTerms, AreaLightContext, Settings, INTEGRATION_AREA_LIGHT_CAPSULE);
 			}
 			else
 			{
-				BSDFEvaluate = SubstrateEvaluateBSDFCommon(BSDFContext, ShadowTerms, AreaLightContext, Settings, INTEGRATION_PUNCTUAL_LIGHT);
+				BSDFEvaluate = SubstrateEvaluateBSDFCommon(BSDFContext, BSDFShadowTerms, AreaLightContext, Settings, INTEGRATION_PUNCTUAL_LIGHT);
 			}
 		}
 
