Bump material version to fix issues that were caused by changing read-only CVars that affect material shader compilation but did not invalidate cache.

diff --git a/Engine/Shaders/Private/MaterialTemplate.ush b/Engine/Shaders/Private/MaterialTemplate.ush
--- a/Engine/Shaders/Private/MaterialTemplate.ush
+++ b/Engine/Shaders/Private/MaterialTemplate.ush
@@ -36,7 +36,7 @@
 
 // Update this GUID to force all material shaders to recompile
 // Merge conflicts on this line should be resolved by generating a new GUID
-#pragma message("UESHADERMETADATA_VERSION 54F3173D-DFF4-4BB7-B389-F3212042C3F4")
+#pragma message("UESHADERMETADATA_VERSION 2475CED6-EBA9-4332-AFD0-4EA70825BABD")
 
 #if USES_SPEEDTREE
 	#include "/Engine/Private/SpeedTreeCommon.ush"
