Fixes regression in TSR's parallax disocclusion introduced by 27385033

#perflight 6530433be11a53cd00f27ec8

diff --git a/Engine/Shaders/Private/TemporalSuperResolution/TSRDecimateHistory.usf b/Engine/Shaders/Private/TemporalSuperResolution/TSRDecimateHistory.usf
--- a/Engine/Shaders/Private/TemporalSuperResolution/TSRDecimateHistory.usf
+++ b/Engine/Shaders/Private/TemporalSuperResolution/TSRDecimateHistory.usf
@@ -312,14 +312,14 @@ void ProcessPrevAtomicTexturesSamples(
 		{
 			const float PixelDepthError = 3.0;
 			const float RadiusToDiameter = 2.0;
-			float WorldDepthEpsilon = GetDepthPixelRadiusForProjectionType(HistoryClosestDeviceZ) * PixelDepthError * RadiusToDiameter;
+			float WorldDepthEpsilon = GetDepthPixelRadiusForProjectionType(HistoryClosestWorldDepth) * PixelDepthError * RadiusToDiameter;
 			float DeltaDepth = abs(HistoryClosestWorldDepth - WorldDepth);
 
 			DepthRejection = saturate(tsr_half(2.0) - tsr_half(DeltaDepth / WorldDepthEpsilon));
 		}
 
 		tsr_half HistoryGhostingRejection = saturate(tsr_half(1.0) - tsr_half(3.0) * abs(tsr_half(HistoryUseCount) * rcp(tsr_half(PREV_USE_COUNT_QUANTIZATION)) - tsr_half(1.0)));
-			
+		
 		bIsValidHoleFillingPixelVelocity = bIsValidHoleFillingPixelVelocity || bValidPixelPos;
 
 		EncodedHoleFillingVelocity = select(bValidPixelPos, max(EncodedHoleFillingVelocity, SampleHistoryClosestDeviceZ), EncodedHoleFillingVelocity);
