Modified Bend SSS to return result instead of writing directly to UAV.
- this way we can modify the result based on light parameters.

diff --git a/Engine/Shaders/Private/ScreenSpaceShadows.usf b/Engine/Shaders/Private/ScreenSpaceShadows.usf
--- a/Engine/Shaders/Private/ScreenSpaceShadows.usf
+++ b/Engine/Shaders/Private/ScreenSpaceShadows.usf
@@ -35,7 +35,6 @@ DispatchParameters SetupBendDispatchParameters(float NearDepthVal, float FarDept
 	Parameters.InvDepthTextureSize = InvDepthTextureSize;
 
 	Parameters.DepthTexture = DepthTexture;
-	Parameters.OutputTexture = OutputTexture;
 	Parameters.PointBorderSampler = PointBorderSampler;
 
 	return Parameters;
@@ -147,7 +146,11 @@ void ScreenSpaceShadowsBendCS(
 {
 	DispatchParameters Parameters = SetupBendDispatchParameters(NearDepthValue, FarDepthValue);
 
-	WriteScreenSpaceShadow(Parameters, GroupId, GroupThreadId.x);
+	float Result;
+	int2 ResultXY;
+	WriteScreenSpaceShadow(Parameters, GroupId, GroupThreadId.x, Result, ResultXY);
+
+	OutputTexture[ResultXY] = 1 - (1 - Result) * ContactShadowCastingIntensity;
 }
 
 #ifdef UPSAMPLE_PASS
diff --git a/Engine/Shaders/Private/bend_sss_gpu.ush b/Engine/Shaders/Private/bend_sss_gpu.ush
--- a/Engine/Shaders/Private/bend_sss_gpu.ush
+++ b/Engine/Shaders/Private/bend_sss_gpu.ush
@@ -71,7 +71,7 @@
 // 
 //	(int3)	inGroupID:			Compute shader group id register (SV_GroupID)
 //	(int)	inGroupThreadId:	Compute shader group thread id register (SV_GroupThreadID)
-void WriteScreenSpaceShadow(struct DispatchParameters inParameters, int3 inGroupID, int inGroupThreadID);
+void WriteScreenSpaceShadow(struct DispatchParameters inParameters, int3 inGroupID, int inGroupThreadID, out float outResult, out int2 outResultXY);
 
 
 
@@ -153,7 +153,6 @@ struct DispatchParameters
 										// The 'USE_HALF_PIXEL_OFFSET' macro might need to be defined if sampling at exact pixel coordinates isn't precise (e.g., if odd patterns appear in the shadow).
 
 	Texture2D<float> DepthTexture;		// Depth Buffer Texture (rasterized non-linear depth)
-	RWTexture2D<float> OutputTexture;	// Output screen-space shadow buffer (typically single-channel, 8bit)
 
 	SamplerState PointBorderSampler;	// A point sampler, with Wrap Mode set to Clamp-To-Border-Color (D3D12_TEXTURE_ADDRESS_MODE_BORDER), and Border Color set to "FarDepthValue" (typically zero), or some other far-depth value out of DepthBounds.
 										// If you have issues where invalid shadows are appearing from off-screen, it is likely that this sampler is not correctly setup
@@ -245,7 +244,7 @@ struct DispatchParameters
 	// 
 	//	(int3)	inGroupID:			Compute shader group id register (SV_GroupID)
 	//	(int)	inGroupThreadId:	Compute shader group thread id register (SV_GroupThreadID)
-	void WriteScreenSpaceShadow(DispatchParameters inParameters, int3 inGroupID, int inGroupThreadID)
+	void WriteScreenSpaceShadow(DispatchParameters inParameters, int3 inGroupID, int inGroupThreadID, out float outResult, out int2 outResultXY)
 	{
 		float2 xy_delta;
 		float2 pixel_xy;
@@ -477,10 +476,9 @@ struct DispatchParameters
 				result = (inGroupThreadID / (float)WAVE_SIZE);
 			if (inParameters.DebugOutputWaveIndex)			
 				result = frac(inGroupID.x / (float)WAVE_SIZE);
-			
-			// Asking the GPU to write scattered single-byte pixels isn't great,
-			// But thankfully the latency is hidden by all the work we're doing...
-			inParameters.OutputTexture[(int2)write_xy] = result;
+
+			outResult = result;
+			outResultXY = (int2)write_xy;
 		}
 	}
 
