Fill out instance bounds from local bounds in the non-instanced path.
This allows materials to use instance bounds for all content, instanced and non-instanced.

diff --git a/Engine/Shaders/Private/SceneDataMobileLoader.ush b/Engine/Shaders/Private/SceneDataMobileLoader.ush
--- a/Engine/Shaders/Private/SceneDataMobileLoader.ush
+++ b/Engine/Shaders/Private/SceneDataMobileLoader.ush
@@ -103,6 +103,12 @@ FPrimitiveSceneData LoadPrimitiveDataUBO(uint DrawInstanceId)
 		PrimitiveData.PreSkinnedLocalBoundsMax = PrimitiveData.LocalObjectBoundsMax;
 	}
 
+	// InstanceLocalBounds
+	{
+		PrimitiveData.InstanceLocalBoundsCenter = (PrimitiveData.LocalObjectBoundsMin + PrimitiveData.LocalObjectBoundsMax) * 0.5f;
+		PrimitiveData.InstanceLocalBoundsExtent = PrimitiveData.LocalObjectBoundsMax - PrimitiveData.InstanceLocalBoundsCenter;
+	}
+
 	// WorldToLocal, TODO compute from LocalToWorld?
 	{
 		float4x4 RelativeWorldToLocal = transpose(float4x4(
