Fix Substrate with Lumen Reflection only (i.e. without ScreenProbeGather.

This path goes through ReflectionEnvironmentPixelShader and was not implemented for Lumen reflection.

diff --git a/Engine/Shaders/Private/AmbientCubemapComposite.usf b/Engine/Shaders/Private/AmbientCubemapComposite.usf
--- a/Engine/Shaders/Private/AmbientCubemapComposite.usf
+++ b/Engine/Shaders/Private/AmbientCubemapComposite.usf
@@ -50,7 +50,7 @@ float3 GetEnvSpecularLighting(
 #include "ReflectionEnvironmentShared.ush"
 #include "Substrate/Substrate.ush"
 #include "Substrate/SubstrateEvaluation.ush"
-#include "Substrate/SubstrateEnvironmentLighting.ush"
+#include "Substrate/SubstrateLightingCommon.ush"
 
 void MainPS(
 	float4 SvPosition : SV_POSITION, 
@@ -87,42 +87,45 @@ void MainPS(
 		// Sample the ambient occlusion that is dynamically generated every frame.
 		const float AmbientOcclusion = AmbientOcclusionTexture.SampleLevel(AmbientOcclusionSampler, BufferUV, 0).r;
 
-		const float4 SSR = 0.f; //Texture2DSample(ScreenSpaceReflectionsTexture, ScreenSpaceReflectionsSampler, BufferUV);
-
-		// SSR is computed for the top level by averaging BSDF components.
-		// This is the top level specular contribution factor according to areas.
-		const float TopLayerSpecularContributionFactor = 1.0f - SSR.a;
+		const FSubstrateIntegrationSettings Settings = InitSubstrateIntegrationSettings(false /*bForceFullyRough*/, Substrate.bRoughDiffuse, Substrate.PeelLayersAboveDepth, Substrate.bRoughnessTracking);
+		FSubstrateDeferredLighting SubstrateLighting = GetInitialisedSubstrateDeferredLighting();
 
-		// A reduction of SSR contribution is needed when there is haziness because we only consider SSR for the sharpest lobe.
-		float SSRReductionFactor = 1.0f;
-		float3 SSRTopLayerEnvBRDF = 0.0f;
+		const float CombinedScreenAndMaterialAO = SubstrateGetAO(SubstratePixelHeader) * AmbientOcclusion;
 
-		const FSubstrateIntegrationSettings Settings = InitSubstrateIntegrationSettings(false /*bForceFullyRough*/, Substrate.bRoughDiffuse, Substrate.PeelLayersAboveDepth, Substrate.bRoughnessTracking);
+		Substrate_for (uint ClosureIndex = 0, ClosureIndex < SubstratePixelHeader.ClosureCount, ++ClosureIndex)
+		{
+			FSubstrateBSDF BSDF = UnpackSubstrateBSDF(Substrate.MaterialTextureArray, SubstrateAddressing, SubstratePixelHeader);
+			FSubstrateBSDFContext SubstrateBSDFContext = SubstrateCreateBSDFContext(SubstratePixelHeader, BSDF, SubstrateAddressing, V);
+			const float3 BSDFThroughput = LuminanceWeight(SubstrateBSDFContext, BSDF); // Use the reflected direction
+
+			// Evaluate environment lighting
+			FSubstrateEnvLightResult SubstrateEnvLight = SubstrateEvaluateForEnvLight(SubstrateBSDFContext, true /*bEnableSpecular*/, Settings);
+
+			float3 DiffuseLighting = 0;
+			float3 SpecularLighting = 0;
+			SubstrateEnvLightingCommon(
+				SubstrateEnvLight,
+				SubstratePixelHeader,
+				SubstrateBSDFContext,
+				BSDF,
+				SubstrateEnvLight.DiffuseNormal,
+				BSDFThroughput,
+				0 /*CaptureDataStartIndex*/,
+				0 /*NumCulledReflectionCaptures*/,
+				AmbientOcclusion,
+				1.0f /*CloudVolumetricAOShadow*/,
+				1.0f /*TopLayerSpecularContributionFactor*/,
+				TranslatedWorldPosition,
+				CombinedScreenAndMaterialAO,
+				DiffuseLighting,
+				SpecularLighting);
+
+			FLightAccumulator Out = (FLightAccumulator)0;
+			LightAccumulator_AddSplit(Out, DiffuseLighting, SpecularLighting, DiffuseLighting, View.PreExposure, SubstrateEnvLight.bPostProcessSubsurface);
+			AccumulateSubstrateDeferredLighting(SubstrateLighting, Out, SubstrateEnvLight.bPostProcessSubsurface, BSDF_GETISTOPLAYER(BSDF));
+		}
 
-		FSubstrateDeferredLighting SubstrateLighting = SubstrateEnvironmentLighting(
-			SvPosition,
-			BufferUV,
-			SceneDepth,
-			WorldPosition,
-			TranslatedWorldPosition,
-			CameraToPixel,
-			V,
-			AmbientOcclusion,
-			TopLayerSpecularContributionFactor,
-			true /*bEnableSpecular*/,
-			Settings,
-			Substrate.MaterialTextureArray,
-			SubstrateAddressing,
-			SubstratePixelHeader,
-			1.0f /*CloudVolumetricAOShadow*/,
-			SSRReductionFactor,
-			SSRTopLayerEnvBRDF);
-
-		// Screen space reflections added on top of environment lighting.
-		// Alpha is 0 to not affect the subsurface diffuse luma information.
-		const float4 SSRColor = float4(SSR.rgb * SSRTopLayerEnvBRDF * SSRReductionFactor, 0.0);
-
-		OutColor += SubstrateLighting.SceneColor + SSRColor;
+		OutColor += SubstrateLighting.SceneColor;
 #if SUBSTRATE_OPAQUE_ROUGH_REFRACTION_ENABLED
 		OutOpaqueRoughRefractionSceneColor += SubstrateLighting.OpaqueRoughRefractionSceneColor;
 		OutSubSurfaceSceneColor += SubstrateLighting.SubSurfaceSceneColor;
diff --git a/Engine/Shaders/Private/ReflectionEnvironmentPixelShader.usf b/Engine/Shaders/Private/ReflectionEnvironmentPixelShader.usf
--- a/Engine/Shaders/Private/ReflectionEnvironmentPixelShader.usf
+++ b/Engine/Shaders/Private/ReflectionEnvironmentPixelShader.usf
@@ -33,12 +33,17 @@
 #include "ClearCoatCommon.ush"
 #define REFLECTION_COMPOSITE_USE_BLENDED_REFLECTION_CAPTURES 1
 #define REFLECTION_COMPOSITE_SUPPORT_SKYLIGHT_BLEND 1
+#define USE_SUBSTRATE_ENV_LIGHTING_COMMON 1
 #include "ReflectionEnvironmentComposite.ush"
 #include "/Engine/Private/Substrate/SubstrateEvaluation.ush"
-#include "/Engine/Private/Substrate/SubstrateEnvironmentLighting.ush"
+#include "/Engine/Private/Substrate/SubstrateLightingCommon.ush"
 #include "Lumen/LumenReflectionsCombine.ush"
 
+#if SUBSTRATE_ENABLED && LUMEN_STANDALONE_REFLECTIONS
+Texture2DArray ReflectionTexture;
+#else
 Texture2D ReflectionTexture;
+#endif
 SamplerState ReflectionTextureSampler;
 
 #if ENABLE_DYNAMIC_SKY_LIGHT
@@ -255,6 +260,71 @@ float GetCloudVolumetricAOShadow(in float3 TranslatedWorldPosition)
 	return 1.0f;
 }
 
+#if SUBSTRATE_ENABLED
+struct FSSRContribution
+{
+	float3 SSRTopLayerEnvBRDF;
+	float  SSRReductionFactor;
+};
+
+FSSRContribution InitSSRContribution()
+{
+	FSSRContribution Out;
+	Out.SSRReductionFactor = 1.0f;
+	Out.SSRTopLayerEnvBRDF = 0.0f;
+	return Out;
+}
+
+void GetSSSRContribution(
+	in FSubstrateEnvLightResult SubstrateEnvLight,
+	in FSubstrateBSDF BSDF,
+	in float3 BSDFThroughput,
+	inout FSSRContribution Out)
+{
+	// The specular path weight applied on SSR. It must account for throughput even for top surface because it also contains coverage.
+	Out.SSRTopLayerEnvBRDF += BSDFThroughput * SubstrateEnvLight.SSRSpecularWeight;
+
+#if SUBSTRATE_FASTPATH==0
+	if (BSDF_GETISTOPLAYER(BSDF) && BSDF_GETHASHAZINESS(BSDF) && any((SubstrateEnvLight.SpecularWeight + SubstrateEnvLight.SpecularHazeWeight) > 0.0f))
+	{
+		// SSR is traced for the sharpest lob. The smoothest one does not rely on SSR so we need to lower energy coming from SSR according to the lobe blend weight.
+		// And we also try to make the transition smooth using Haziness
+		Out.SSRReductionFactor -= dot(BSDFThroughput, float3(1.0f / 3.0f, 1.0f / 3.0f, 1.0f / 3.0f)) * SubstrateEnvLight.SSRReduction;
+	}
+#endif 
+}
+
+float4 ApplySSRContribution(FSSRContribution In, float4 InReflectionInput)
+{
+	// Screen space reflections added on top of environment lighting.
+	// Alpha is 0 to not affect the subsurface diffuse luma information.
+	return float4(InReflectionInput.rgb * In.SSRTopLayerEnvBRDF * saturate(In.SSRReductionFactor), 0.0);
+}
+
+float4 GetLumenReflectionOnly(
+	in FSubstrateEnvLightResult SubstrateEnvLight,
+	in FSubstrateBSDF BSDF,
+	in float3 BSDFThroughput,
+	in float4 InReflectionInput)
+{
+	float4 OutLighting = float4(0, 0, 0, 1);
+#if LUMEN_STANDALONE_REFLECTIONS
+	const float Roughness = SubstrateGetBSDFRoughness(BSDF);
+	const bool  bHasBackfaceDiffuse = SubstrateGetBSDFType(BSDF) == SUBSTRATE_BSDF_TYPE_SLAB && BSDF.HasBackScattering();
+	const float FadeAlpha = LumenCombineReflectionsAlpha(Roughness, bHasBackfaceDiffuse);
+	// Must branch as Lumen Reflections can be uninitialized where not needed and contain NaN
+	if (FadeAlpha > 0.0f)
+	{
+		OutLighting.rgb = InReflectionInput.xyz * FadeAlpha;
+		OutLighting.a = 1 - FadeAlpha;
+	}
+#endif
+	return OutLighting;
+}
+
+#endif // SUBSTRATE_ENABLED
+
+
 void ReflectionEnvironmentSkyLighting(
 	in float4 SvPosition : SV_Position,
 	out float4 OutColor : SV_Target0
@@ -294,49 +364,105 @@ void ReflectionEnvironmentSkyLighting(
 		// Sample the ambient occlusion that is dynamically generated every frame.
 		float AmbientOcclusion = AmbientOcclusionTexture.SampleLevel(AmbientOcclusionSampler, BufferUV, 0).r;
 
-		const float4 ReflectionInput = Texture2DSample(ReflectionTexture, ReflectionTextureSampler, BufferUV);
-
+	#if LUMEN_STANDALONE_REFLECTIONS
+		const float TopLayerSpecularContributionFactor = 1.0f;
+	#else
+		const float4 SSRReflection = Texture2DSample(ReflectionTexture, ReflectionTextureSampler, BufferUV);
 		// SSR is computed for the top level by averaging BSDF components.
 		// This is the top level specular contribution factor according to areas.
-		const float TopLayerSpecularContributionFactor = 1.0f - ReflectionInput.a;
+		const float TopLayerSpecularContributionFactor = 1.0f - SSRReflection.a;
 
 		// A reduction of SSR contribution is needed when there is haziness because we only consider SSR for the sharpest lobe.
-		float SSRReductionFactor = 1.0f;
-		float3 SSRTopLayerEnvBRDF = 0.0f;
+		FSSRContribution SSRContribution = InitSSRContribution();
+	#endif
 
 		const float CloudVolumetricAOShadow = GetCloudVolumetricAOShadow(TranslatedWorldPosition);
 		const FSubstrateIntegrationSettings Settings = InitSubstrateIntegrationSettings(false /*bForceFullyRough*/, Substrate.bRoughDiffuse, Substrate.PeelLayersAboveDepth, Substrate.bRoughnessTracking);
 
 		const bool bEnableSpecular = ReflectionStruct.SkyLightParameters.y > 0.0f;
+		FSubstrateDeferredLighting SubstrateLighting = GetInitialisedSubstrateDeferredLighting();
+	
+	#if APPLY_SKY_SHADOWING
+		// Sample DFAO only once
+		FUpsampleDFAOOutput UpsampleDFAOOutput = UpsampleDFAO(BufferUV, SceneDepth);
+	#endif
+	
+		float CombinedScreenAndMaterialAO = SubstrateGetAO(SubstratePixelHeader) * AmbientOcclusion;
+	
+		
+	#if USE_DEFAULT_ENV_LIGHTING_INPUT && FEATURE_LEVEL >= FEATURE_LEVEL_SM5
+		float2 LocalPosition = SvPosition.xy - View.ViewRectMin.xy;
+	
+		uint GridIndex = ComputeLightGridCellIndex(uint2(LocalPosition.x, LocalPosition.y), SceneDepth);
+		uint NumCulledEntryIndex = (ForwardLightData.NumGridCells + GridIndex) * NUM_CULLED_LIGHTS_GRID_STRIDE;
+		uint NumCulledReflectionCaptures = min(ForwardLightData.NumCulledLightsGrid[NumCulledEntryIndex + 0], ForwardLightData.NumReflectionCaptures);
+		uint CaptureDataStartIndex = ForwardLightData.NumCulledLightsGrid[NumCulledEntryIndex + 1];
+	#else
+		uint CaptureDataStartIndex = 0;
+		uint NumCulledReflectionCaptures = 0;
+	#endif
+	
+		Substrate_for (uint ClosureIndex = 0, ClosureIndex < SubstratePixelHeader.ClosureCount, ++ClosureIndex)
+		{
+			FSubstrateBSDF BSDF = UnpackSubstrateBSDF(Substrate.MaterialTextureArray, SubstrateAddressing, SubstratePixelHeader);
+			FSubstrateBSDFContext SubstrateBSDFContext = SubstrateCreateBSDFContext(SubstratePixelHeader, BSDF, SubstrateAddressing, V);
+			const float3 BSDFThroughput = LuminanceWeight(SubstrateBSDFContext, BSDF); // Use the reflected direction
+	
+			// Evaluate environment lighting
+			FSubstrateEnvLightResult SubstrateEnvLight = SubstrateEvaluateForEnvLight(SubstrateBSDFContext, bEnableSpecular, Settings);
+	
+			float3 BentNormal = SubstrateEnvLight.DiffuseNormal;
+		#if APPLY_SKY_SHADOWING
+			// Set DiffuseNormal as the bent normal for all diffuse computations.
+			BentNormal = ApplyDFAO(UpsampleDFAOOutput, SubstrateEnvLight.DiffuseNormal);
+		#endif
+	
+			float3 DiffuseLighting = 0;
+			float3 SpecularLighting = 0;
+			SubstrateEnvLightingCommon(
+				SubstrateEnvLight,
+				SubstratePixelHeader,
+				SubstrateBSDFContext,
+				BSDF,
+				BentNormal,
+				BSDFThroughput,
+				CaptureDataStartIndex,
+				NumCulledReflectionCaptures,
+				AmbientOcclusion,
+				CloudVolumetricAOShadow,
+				TopLayerSpecularContributionFactor,
+				TranslatedWorldPosition,
+				CombinedScreenAndMaterialAO,
+				DiffuseLighting,
+				SpecularLighting);
+	
+		#if LUMEN_STANDALONE_REFLECTIONS
+			{
+				const float4 LumenReflections = GetLumenReflectionOnly(SubstrateEnvLight, BSDF, BSDFThroughput, Texture2DArraySample(ReflectionTexture, ReflectionTextureSampler, float3(BufferUV, ClosureIndex)));
+				SpecularLighting *= LumenReflections.a;
+				SpecularLighting += LumenReflections.xyz * BSDFThroughput * SubstrateEnvLight.SpecularWeight;
+			}
+		#else
+			if (BSDF_GETISTOPLAYER(BSDF))
+			{
+				GetSSSRContribution(SubstrateEnvLight, BSDF, BSDFThroughput, SSRContribution);
+			}
+		#endif
+			
+			FLightAccumulator Out = (FLightAccumulator)0;
+			LightAccumulator_AddSplit(Out, DiffuseLighting, SpecularLighting, DiffuseLighting, 1.f /*CommonMultiplier*/, SubstrateEnvLight.bPostProcessSubsurface);
+			AccumulateSubstrateDeferredLighting(SubstrateLighting, Out, SubstrateEnvLight.bPostProcessSubsurface, BSDF_GETISTOPLAYER(BSDF));
+		}
+
+		OutColor = SubstrateLighting.SceneColor;
+		#if LUMEN_STANDALONE_REFLECTIONS == 0
+		OutColor += ApplySSRContribution(SSRContribution, SSRReflection);
+		#endif
 
-		FSubstrateDeferredLighting SubstrateLighting = SubstrateEnvironmentLighting(
-			SvPosition,
-			BufferUV,
-			SceneDepth,
-			WorldPosition,
-			TranslatedWorldPosition,
-			CameraToPixel,
-			V,
-			AmbientOcclusion,
-			TopLayerSpecularContributionFactor,
-			bEnableSpecular,
-			Settings,
-			Substrate.MaterialTextureArray,
-			SubstrateAddressing,
-			SubstratePixelHeader,
-			CloudVolumetricAOShadow,
-			SSRReductionFactor,
-			SSRTopLayerEnvBRDF);
-
-		// Screen space reflections added on top of environment lighting.
-		// Alpha is 0 to not affect the subsurface diffuse luma information.
-		const float4 SSRColor = float4(ReflectionInput.rgb * SSRTopLayerEnvBRDF * SSRReductionFactor, 0.0);
-
-		OutColor += SubstrateLighting.SceneColor + SSRColor;
-#if SUBSTRATE_OPAQUE_ROUGH_REFRACTION_ENABLED
+		#if SUBSTRATE_OPAQUE_ROUGH_REFRACTION_ENABLED
 		OutOpaqueRoughRefractionSceneColor += SubstrateLighting.OpaqueRoughRefractionSceneColor;
 		OutSubSurfaceSceneColor += SubstrateLighting.SubSurfaceSceneColor;
-#endif
+		#endif
 	}
 
 #else // SUBSTRATE_ENABLED
diff --git a/Engine/Shaders/Private/Substrate/SubstrateEnvironmentLighting.ush b/Engine/Shaders/Private/Substrate/SubstrateEnvironmentLighting.ush
deleted file mode 100644
diff --git a/Engine/Shaders/Private/Substrate/SubstrateLightingCommon.ush b/Engine/Shaders/Private/Substrate/SubstrateLightingCommon.ush
--- a/Engine/Shaders/Private/Substrate/SubstrateLightingCommon.ush
+++ b/Engine/Shaders/Private/Substrate/SubstrateLightingCommon.ush
@@ -200,7 +200,6 @@ void SubstrateEnvLightingCommon(
 	in float TopLayerSpecularContributionFactor,
 	in float3 TranslatedWorldPosition,
 	in float CombinedScreenAndMaterialAO,
-	inout float SSRReductionFactor,
 	inout float3 DiffuseLighting,
 	inout float3 SpecularLighting)
 {
@@ -298,15 +297,9 @@ void SubstrateEnvLightingCommon(
 					ExtraIndirectSpecular,
 					NumCulledReflectionCaptures,
 					CaptureDataStartIndex);
-
-			// SSR is traced for the sharpest lob. The smoothest one does not rely on SSR so we need to lower energy coming from SSR according to the lobe blend weight.
-			// And we also try to make the transition smooth using Haziness
-			SSRReductionFactor -= bIsTopLayer ? dot(BSDFThroughput, float3(1.0f / 3.0f, 1.0f / 3.0f, 1.0f / 3.0f)) * SubstrateEnvLight.SSRReduction : 0.0f;
 		}
 	#endif 
 	}
 }
 
-#endif // SUBSTRATE_ENV_LIGHTING_COMMON
-
-
+#endif // SUBSTRATE_ENV_LIGHTING_COMMON
\ No newline at end of file
diff --git a/Engine/Shaders/Private/Substrate/SubstrateVisualize.usf b/Engine/Shaders/Private/Substrate/SubstrateVisualize.usf
--- a/Engine/Shaders/Private/Substrate/SubstrateVisualize.usf
+++ b/Engine/Shaders/Private/Substrate/SubstrateVisualize.usf
@@ -1254,7 +1254,6 @@ void EvaluateLighting(
 		const bool bEnableSpecular = ReflectionStruct.SkyLightParameters.y > 0.0f;
 		FSubstrateEnvLightResult SubstrateEnvLight = SubstrateEvaluateForEnvLight(SubstrateBSDFContext, bEnableSpecular, Settings);
 
-		float SSRReductionFactor = 0.0;
 		float3 DiffuseLighting = 0.0;
 		float3 SpecularLighting = 0.0;
 		SubstrateEnvLightingCommon(
@@ -1271,7 +1270,6 @@ void EvaluateLighting(
 			1.0,//in float TopLayerSpecularContributionFactor,
 			TranslatedWorldPosition,
 			1.0,//in float CombinedScreenAndMaterialAO,
-			SSRReductionFactor,
 			DiffuseLighting,
 			SpecularLighting);
 
