Fix shader compilation.

diff --git a/Engine/Shaders/Private/Substrate/SubstrateLightingCommon.ush b/Engine/Shaders/Private/Substrate/SubstrateLightingCommon.ush
--- a/Engine/Shaders/Private/Substrate/SubstrateLightingCommon.ush
+++ b/Engine/Shaders/Private/Substrate/SubstrateLightingCommon.ush
@@ -2,6 +2,12 @@
 
 #pragma once
 
+#ifndef SUBSTRATE_ENABLED
+	#define SUBSTRATE_ENABLED 1
+	#error SUBSTRATE_ENABLED needs to be defined
+#endif
+
+#if SUBSTRATE_ENABLED
 #ifdef USE_SUBSTRATE_FORWARD_LIGHTING_COMMON
 
 // Used by non-mobile path
@@ -302,4 +308,5 @@ void SubstrateEnvLightingCommon(
 	}
 }
 
-#endif // SUBSTRATE_ENV_LIGHTING_COMMON
\ No newline at end of file
+#endif // SUBSTRATE_ENV_LIGHTING_COMMON
+#endif // SUBSTRATE_ENABLED
\ No newline at end of file
