Fixed volumetric fog extinction being affect by exposure when it should not.
This does converges with time so it should not affct visual much. But could impact our engine screenshot tests.

diff --git a/Engine/Shaders/Private/VolumetricFog.usf b/Engine/Shaders/Private/VolumetricFog.usf
--- a/Engine/Shaders/Private/VolumetricFog.usf
+++ b/Engine/Shaders/Private/VolumetricFog.usf
@@ -939,7 +939,7 @@ void LightScatteringCS(
 	if (HistoryAlpha > 0)
 	{
 		float4 PreExposedHistoryScatteringAndExtinction = Texture3DSampleLevel(LightScatteringHistory, LightScatteringHistorySampler, HistoryUV, 0);
-		PreExposedHistoryScatteringAndExtinction *= LightScatteringHistoryPreExposureAndInv.y * View.PreExposure; // Bring previous frame pre exposed history luminance as current frame pre exposed luminance.
+		PreExposedHistoryScatteringAndExtinction.rgb *= LightScatteringHistoryPreExposureAndInv.y * View.PreExposure; // Bring previous frame pre exposed history luminance as current frame pre exposed luminance. Leave extinction untouched! 
 		PreExposedScatteringAndExtinction = lerp(PreExposedScatteringAndExtinction, PreExposedHistoryScatteringAndExtinction, HistoryAlpha);
 	}
 	
