Fix raytraced translucent shadow with Substrate

Two seperate issues:
* The Flags/Opacity/Blend/... data were not correctly writtent in Substrate PackedPayload data
* The coverage value wasn't computed for translucent surfaces

diff --git a/Engine/Shaders/Private/RayTracing/RayTracingCommon.ush b/Engine/Shaders/Private/RayTracing/RayTracingCommon.ush
--- a/Engine/Shaders/Private/RayTracing/RayTracingCommon.ush
+++ b/Engine/Shaders/Private/RayTracing/RayTracingCommon.ush
@@ -607,7 +607,7 @@ struct FPackedMaterialClosestHitPayload : FMinimalPayload
 	void SetRayCone(FRayCone NewRayCone)	{ PackedRayCone = PackRayCone(NewRayCone); }
 
 	float GetMipBias()						{ return UnpackR8(Flags>>8u); }
-	void SetMipBias(float NewMipBias)		{ Flags = PackR8(NewMipBias) << 8u; }
+	void SetMipBias(float NewMipBias)		{ Flags = (Flags & 0xFFFF00FFF) | (PackR8(NewMipBias) << 8u); }
 
 	void SetRadiance(in float3 In)			{ Radiance = PackR11G11B10F(In); }
 	float3 GetRadiance()					{ return UnpackR11G11B10F(Radiance); }
@@ -637,9 +637,9 @@ struct FPackedMaterialClosestHitPayload : FMinimalPayload
 	float GetOpacity()						{ return UnpackR8(Flags >> 24u); }
 
 	void SetFlags(uint In)							{ Flags |= ((In & 0xFF)); }
-	void SetBlendingMode(uint In)					{ Flags |= ((In & 0x7)<<16u); }
-	void SetPrimitiveLightingChannelMask(uint In)	{ Flags |= ((In & 0x7)<<19u); }
-	void SetOpacity(float In)						{ Flags |= PackR8(In) << 24u; }
+	void SetBlendingMode(uint In)					{ Flags  = (Flags & 0xFFF8FFFF) | ((In & 0x7)<<16u); }
+	void SetPrimitiveLightingChannelMask(uint In)	{ Flags  = (Flags & 0xFFC7FFFF) | ((In & 0x7)<<19u); }
+	void SetOpacity(float In)						{ Flags  = (Flags & 0x00FFFFFF) | (PackR8(In) << 24u); }
 
 	void SetRayDirection(float3 In)			{ PackedRayCone = SubstratePackNormal(In); }
 	float3 GetRayDirection()				{ return SubstrateUnpackNormal(PackedRayCone); }
diff --git a/Engine/Shaders/Private/RayTracing/RayTracingMaterialHitShaders.usf b/Engine/Shaders/Private/RayTracing/RayTracingMaterialHitShaders.usf
--- a/Engine/Shaders/Private/RayTracing/RayTracingMaterialHitShaders.usf
+++ b/Engine/Shaders/Private/RayTracing/RayTracingMaterialHitShaders.usf
@@ -271,7 +271,21 @@ void GetMaterialPayload(
 	// Generate the Substrate material data to write out
 	FSubstrateAddressing SubstrateAddressing = GetSubstratePixelDataByteOffset(0,0,0);
 
-	const float3 CamVector = MaterialParameters.CameraVector;	
+	// Compute TotalCoverage for translucent material
+#if MATERIALBLENDING_ANY_TRANSLUCENT
+	if (SubstratePixelHeader.SubstrateTree.BSDFCount > 0)
+	{
+		float OutTransmittancePreCoverage = 0.f;
+		SubstratePixelHeader.SubstrateUpdateTree(
+			SubstrateData, 
+			MaterialParameters.CameraVector, 
+			Settings, 
+			TotalCoverage, 
+			OutTransmittancePreCoverage);
+	}
+	#endif
+
+	const float3 CamVector = MaterialParameters.CameraVector;
 	Payload.SubstrateData = InitialiseRWSubstrateMaterialContainer();
 	PackSubstrateOut(
 		Payload.SubstrateData,
