Add transmittance support for translucent shadow.

diff --git a/Engine/Shaders/Private/RayTracing/RayTracingMaterialHitShaders.usf b/Engine/Shaders/Private/RayTracing/RayTracingMaterialHitShaders.usf
--- a/Engine/Shaders/Private/RayTracing/RayTracingMaterialHitShaders.usf
+++ b/Engine/Shaders/Private/RayTracing/RayTracingMaterialHitShaders.usf
@@ -272,16 +272,17 @@ void GetMaterialPayload(
 	FSubstrateAddressing SubstrateAddressing = GetSubstratePixelDataByteOffset(0,0,0);
 
 	// Compute TotalCoverage for translucent material
-#if MATERIALBLENDING_ANY_TRANSLUCENT
+	#if MATERIALBLENDING_ANY_TRANSLUCENT
 	if (SubstratePixelHeader.SubstrateTree.BSDFCount > 0)
 	{
-		float OutTransmittancePreCoverage = 0.f;
+		float TransmittancePreCoverage = 0.f;
 		SubstratePixelHeader.SubstrateUpdateTree(
 			SubstrateData, 
 			MaterialParameters.CameraVector, 
 			Settings, 
 			TotalCoverage, 
-			OutTransmittancePreCoverage);
+			TransmittancePreCoverage);
+		TotalCoverage *= (1.f-dot(TransmittancePreCoverage, float(1.f/3.f).xxx));
 	}
 	#endif
 
