LFV - fixed shader compilation warning.

diff --git a/Engine/Shaders/Private/BasePassVertexShader.usf b/Engine/Shaders/Private/BasePassVertexShader.usf
--- a/Engine/Shaders/Private/BasePassVertexShader.usf
+++ b/Engine/Shaders/Private/BasePassVertexShader.usf
@@ -111,7 +111,7 @@ void Main(
 	float4 VertexClipSpacePosition			= mul(float4(WorldPosition.xyz, 1), ResolvedView.TranslatedWorldToClip);
 	float2 SvPosition						= (VertexClipSpacePosition.xy / VertexClipSpacePosition.w * float2(.5f, -.5f) + .5f) * ResolvedView.ViewSizeAndInvSize.xy;
 	uint2 TilePos							= clamp(uint2(SvPosition.xy / float(LFVTilePixelSize)), uint2(0, 0), LFVTileDataResolution - 1);
-	float4 LFVContribution					= GetLFVContribution(ResolvedView, TilePos, WorldPosition);
+	float4 LFVContribution					= GetLFVContribution(ResolvedView, TilePos, WorldPosition.xyz);
 	Output.BasePassInterpolants.VertexFog	= float4(LFVContribution.rgb + Output.BasePassInterpolants.VertexFog.rgb * LFVContribution.a, LFVContribution.a * Output.BasePassInterpolants.VertexFog.a);
 #endif // LOCAL_FOG_VOLUME_ON_TRANSLUCENT
 
diff --git a/Engine/Shaders/Private/MobileBasePassVertexShader.usf b/Engine/Shaders/Private/MobileBasePassVertexShader.usf
--- a/Engine/Shaders/Private/MobileBasePassVertexShader.usf
+++ b/Engine/Shaders/Private/MobileBasePassVertexShader.usf
@@ -115,7 +115,7 @@ void Main(
 	float2 SvPosition				= (VertexClipSpacePosition.xy / VertexClipSpacePosition.w * float2(.5f, -.5f) + .5f) * ResolvedView.ViewSizeAndInvSize.xy;
 	uint2 TilePos					= clamp(uint2(SvPosition.xy / float(LFVTilePixelSize)), uint2(0, 0), LFVTileDataResolution - 1);
 	const bool bPlatformSupportsVolumetricFogOntranslucent = false; // Mobile does not do that so we try to compensate for it by evalauting LFV regularly on translucent.
-	float4 LFVContribution			= GetLFVContribution(ResolvedView, TilePos, WorldPosition, bPlatformSupportsVolumetricFogOntranslucent);
+	float4 LFVContribution			= GetLFVContribution(ResolvedView, TilePos, WorldPosition.xyz, bPlatformSupportsVolumetricFogOntranslucent);
 	VertexFog						= float4(LFVContribution.rgb + VertexFog.rgb * LFVContribution.a, LFVContribution.a * VertexFog.a);
 #endif // LOCAL_FOG_VOLUME_ON_TRANSLUCENT
 
