Fix out of bound reading/writing in hair culling pass.

Due to dispatch count rounding, more curve can be dispatched causing out of bound reading/writing.

diff --git a/Engine/Shaders/Private/HairStrands/HairStrandsClusterCulling.usf b/Engine/Shaders/Private/HairStrands/HairStrandsClusterCulling.usf
--- a/Engine/Shaders/Private/HairStrands/HairStrandsClusterCulling.usf
+++ b/Engine/Shaders/Private/HairStrands/HairStrandsClusterCulling.usf
@@ -180,20 +180,30 @@ void Main(uint3 DispatchThreadId : SV_DispatchThreadID, uint ThreadIndex : SV_Gr
 	{
 		const uint LocalCurveIndex  = ThreadIndex;
 		const uint GlobalCurveIndex = GetGlobalCurveIndex(ThreadIndex, GroupIndex);
-
-		// 1.1 Cluster ID
-		const uint ClusterId = CurveToClusterIdBuffer[GlobalCurveIndex];
+		const bool bValidCurveIndex = GlobalCurveIndex < CurveCount;
 	
+		// 1.1 Cluster ID
 		// 1.3 Get cluster info (i.e. map all info per LOD for a given cluster ID) and derive effective LOD data (radius scale)
-		const FHairClusterLOD ClusterLODInfo = GetHairClusterLOD(ClusterInfoBuffer[ClusterId], ClusterInfoParameters, LODIndex);
+		FHairClusterLOD ClusterLODInfo = (FHairClusterLOD)0;
+		if (bValidCurveIndex)
+		{
+			const uint ClusterId = CurveToClusterIdBuffer[GlobalCurveIndex];
+			ClusterLODInfo = GetHairClusterLOD(ClusterInfoBuffer[ClusterId], ClusterInfoParameters, LODIndex);
+		}
 		ClusterLOD[LocalCurveIndex] 		 = ClusterLODInfo.LOD;
 		ClusterRadiusScale[LocalCurveIndex]  = ClusterLODInfo.RadiusScale; // Curve are no longer scaled individually by ClusterLODInfo.RadiusScale, but scale globally based on CPU value. This is for legacy path, using discrete LOD with manual ThicknessScale
 		LODCount[LocalCurveIndex]			 = ClusterLODInfo.LODCount;
 
 		// 1.4 Curve's point count/offset
-		const FHairCurve Curve = UnpackHairCurve(CurveBuffer[GlobalCurveIndex]);
-		CurvePointIndex[LocalCurveIndex] = Curve.PointIndex;
-		CurvePointCount[LocalCurveIndex] = Curve.PointCount;
+		CurvePointIndex[LocalCurveIndex] = 0;
+		CurvePointCount[LocalCurveIndex] = 0;
+		if (bValidCurveIndex)
+		{
+			const FHairCurve Curve = UnpackHairCurve(CurveBuffer[GlobalCurveIndex]);
+			CurvePointIndex[LocalCurveIndex] = Curve.PointIndex;
+			CurvePointCount[LocalCurveIndex] = Curve.PointCount;
+		}
+
 
 	}
 	if (ThreadIndex == 0)
@@ -215,7 +225,7 @@ void Main(uint3 DispatchThreadId : SV_DispatchThreadID, uint ThreadIndex : SV_Gr
 		{	
 			if (CurrentPointIndex < CurvePointCount[CurveIndex])
 			{
-				const uint PointIndex = CurrentPointIndex + CurvePointIndex[CurveIndex];				
+				const uint PointIndex = CurrentPointIndex + CurvePointIndex[CurveIndex];
 				const uint MinLOD = GetHairControlPointMinLOD(PointIndex, PointLODBuffer);
 
 				// When a curve reach its last LOD, fade its radius progressively
@@ -271,7 +281,7 @@ void Main(uint3 DispatchThreadId : SV_DispatchThreadID, uint ThreadIndex : SV_Gr
 	if (ThreadIndex < CURVE_PER_GROUP)
 	{
 		const uint LocalCurveIndex  = ThreadIndex;
-		const uint GlobalCurveIndex = GetGlobalCurveIndex(ThreadIndex, GroupIndex);;
+		const uint GlobalCurveIndex = GetGlobalCurveIndex(ThreadIndex, GroupIndex);
 		if (GlobalCurveIndex < CurveCount)
 		{
 			FHairCurve CulledCurve;
@@ -309,22 +319,35 @@ void Main(uint3 DispatchThreadId : SV_DispatchThreadID, uint ThreadIndex : SV_Gr
 {	
 	// 1. Read the Curve data
 	const uint CurveIndex = GetGlobalCurveIndex(ThreadIndex, GroupIndex);
+	const bool bIsCurveValid = CurveIndex < CurveCount;
+	if (!bIsCurveValid)
+	{
+		return;
+	}
 	
 	if (ThreadIndex == 0)
 	{
 		// 1.1 Cluster ID
-		const uint ClusterId = CurveToClusterIdBuffer[CurveIndex];
-	
 		// 1.2 Get cluster info (i.e. map all info per LOD for a given cluster ID) and derive effective LOD data (radius scale)
-		const FHairClusterLOD ClusterLODInfo = GetHairClusterLOD(ClusterInfoBuffer[ClusterId], ClusterInfoParameters, LODIndex);
+		FHairClusterLOD ClusterLODInfo = (FHairClusterLOD)0;
+		if (bIsCurveValid)
+		{
+			const uint ClusterId = CurveToClusterIdBuffer[CurveIndex];
+			ClusterLODInfo = GetHairClusterLOD(ClusterInfoBuffer[ClusterId], ClusterInfoParameters, LODIndex);
+		}
 		ClusterLOD = ClusterLODInfo.LOD;
 		ClusterRadiusScale = ClusterLODInfo.RadiusScale; // Curve are no longer scaled individually by ClusterLODInfo.RadiusScale, but scale globally based on CPU value. This is for legacy path, using discrete LOD with manual ThicknessScale
 		LODCount = ClusterLODInfo.LODCount;
 
 		// 1.4 Curve's point count/offset
-		const FHairCurve Curve = UnpackHairCurve(CurveBuffer[CurveIndex]);
-		CurvePointIndex = Curve.PointIndex;
-		CurvePointCount = Curve.PointCount;
+		CurvePointIndex = 0;
+		CurvePointCount = 0;
+		if (bIsCurveValid)
+		{
+			const FHairCurve Curve = UnpackHairCurve(CurveBuffer[CurveIndex]);
+			CurvePointIndex = Curve.PointIndex;
+			CurvePointCount = Curve.PointCount;
+		}
 
 		AllocationCount = 0;
 		AllocationOffset = 0;
@@ -406,7 +429,7 @@ void Main(uint3 DispatchThreadId : SV_DispatchThreadID, uint ThreadIndex : SV_Gr
 	}
 
 	// 6. Write out curve data with visible points for subsequent passes
-	if (ThreadIndex == 0)
+	if (ThreadIndex == 0 && bIsCurveValid)
 	{	
 		FHairCurve CulledCurve;
 		CulledCurve.PointCount = AllocationCount;
