Fix errors compiling shaders in ScreenSpaceShadows.usf.
- only include bend_sss_gpu.ush when compiling FScreenSpaceShadowsBendCS since other shaders don't force DXC.

diff --git a/Engine/Shaders/Private/ScreenSpaceShadows.usf b/Engine/Shaders/Private/ScreenSpaceShadows.usf
--- a/Engine/Shaders/Private/ScreenSpaceShadows.usf
+++ b/Engine/Shaders/Private/ScreenSpaceShadows.usf
@@ -4,6 +4,8 @@
 	ScreenSpaceShadows.usf
 =============================================================================*/
 
+#if BEND_SSS
+
 #include "/Engine/Public/Platform.ush"
 
 #define WAVE_SIZE 64
@@ -42,6 +44,8 @@ DispatchParameters SetupBendDispatchParameters(float NearDepthVal, float FarDept
 	return Parameters;
 }
 
+#endif // BEND_SSS
+
 #include "ScreenSpaceShadowRayCast.ush"
 
 #include "Common.ush"
@@ -141,6 +145,8 @@ void ScreenSpaceShadowsCS(
 	RWShadowFactors[DispatchThreadId.xy] = float2(Result, SceneDepth);
 }
 
+#if BEND_SSS
+
 [numthreads(THREADGROUP_SIZEX, THREADGROUP_SIZEY, 1)]
 void ScreenSpaceShadowsBendCS(
 	uint3 GroupId : SV_GroupID,
@@ -156,6 +162,8 @@ void ScreenSpaceShadowsBendCS(
 	OutputTexture[ResultXY] = 1 - (1 - Result) * ContactShadowCastingIntensity;
 }
 
+#endif // BEND_SSS
+
 #ifdef UPSAMPLE_PASS
 
 Texture2D ShadowFactorsTexture;
diff --git a/Engine/Source/Runtime/Renderer/Private/Shadows/ScreenSpaceShadows.cpp b/Engine/Source/Runtime/Renderer/Private/Shadows/ScreenSpaceShadows.cpp
--- a/Engine/Source/Runtime/Renderer/Private/Shadows/ScreenSpaceShadows.cpp
+++ b/Engine/Source/Runtime/Renderer/Private/Shadows/ScreenSpaceShadows.cpp
@@ -126,6 +126,7 @@ class FScreenSpaceShadowsBendCS : public FGlobalShader
 		OutEnvironment.SetDefine(TEXT("THREADGROUP_SIZEY"), 1);
 		OutEnvironment.SetDefine(TEXT("FORCE_DEPTH_TEXTURE_READS"), 1);
 		OutEnvironment.SetDefine(TEXT("PLATFORM_SUPPORTS_TYPED_UAV_LOAD"), (int32)RHISupports4ComponentUAVReadWrite(Parameters.Platform));
+		OutEnvironment.SetDefine(TEXT("BEND_SSS"), 1);
 
 		OutEnvironment.CompilerFlags.Add(CFLAG_ForceDXC);
 	}
