LFV - fixed compilation issue on some material

diff --git a/Engine/Shaders/Private/BasePassVertexShader.usf b/Engine/Shaders/Private/BasePassVertexShader.usf
--- a/Engine/Shaders/Private/BasePassVertexShader.usf
+++ b/Engine/Shaders/Private/BasePassVertexShader.usf
@@ -112,7 +112,11 @@ void Main(
 	float2 SvPosition						= (VertexClipSpacePosition.xy / VertexClipSpacePosition.w * float2(.5f, -.5f) + .5f) * ResolvedView.ViewSizeAndInvSize.xy;
 	uint2 TilePos							= clamp(uint2(SvPosition.xy / float(LFVTilePixelSize)), uint2(0, 0), LFVTileDataResolution - 1);
 	float4 LFVContribution					= GetLFVContribution(ResolvedView, TilePos, WorldPosition.xyz);
+	#if NEEDS_BASEPASS_VERTEX_FOGGING
 	Output.BasePassInterpolants.VertexFog	= float4(LFVContribution.rgb + Output.BasePassInterpolants.VertexFog.rgb * LFVContribution.a, LFVContribution.a * Output.BasePassInterpolants.VertexFog.a);
+	#else
+	Output.BasePassInterpolants.VertexFog	= LFVContribution; // path when Output.BasePassInterpolants.VertexFog has not been initialised
+	#endif
 #endif // LOCAL_FOG_VOLUME_ON_TRANSLUCENT
 
 #if NEEDS_BASEPASS_VERTEX_FOGGING
