GGX Sampling: Improve ImportanceSampleVisibleGGX with new bounds on the reflected direction

Based on the paper: "Bounded VNDF Sampling for Smith-GGX Reflections" from Siggraph Asia 2023

This reduces the number of rays samples below the hemisphere, resulting in more valid samples being traced. The new sampling is enabled by default, but can be disabled by either passing false to ImportanceSampleVisibleGGX or defining GGX_BOUNDED_VNDF_SAMPLING to 0.

diff --git a/Engine/Shaders/Private/MonteCarlo.ush b/Engine/Shaders/Private/MonteCarlo.ush
--- a/Engine/Shaders/Private/MonteCarlo.ush
+++ b/Engine/Shaders/Private/MonteCarlo.ush
@@ -362,7 +362,16 @@ float4 ImportanceSampleGGX( float2 E, float a2 )
 	return float4( H, PDF );
 }
 
-float VisibleGGXPDF(float3 V, float3 H, float a2)
+#ifndef GGX_BOUNDED_VNDF_SAMPLING
+// "Bounded VNDF Sampling for Smith–GGX Reflections"
+// Kenta Eto & Yusuke Tokuyoshi - Siggraph Asia 2023
+// This paper further improves the method of Dupuy and Benyoub in the case where the microfacet normal is used for reflection only.
+// It allows many fewer rays to be rejected at high roughness which improves quality. This is particularly noticeable
+// in multi-bounce reflections of rough metals because longer paths can be followed without early termination.
+#define GGX_BOUNDED_VNDF_SAMPLING 1
+#endif
+
+float VisibleGGXPDF(float3 V, float3 H, float a2, bool bLimitVDNFToReflection = true)
 {
 	float NoV = V.z;
 	float NoH = H.z;
@@ -370,12 +379,20 @@ float VisibleGGXPDF(float3 V, float3 H, float a2)
 
 	float d = (NoH * a2 - NoH) * NoH + 1;
 	float D = a2 / (PI*d*d);
-
-	float PDF = 2 * VoH * D / (NoV + sqrt(NoV * (NoV - NoV * a2) + a2));
+	float k = 1.0;
+#if GGX_BOUNDED_VNDF_SAMPLING
+	if (bLimitVDNFToReflection)
+	{
+		float s = 1.0f + length(V.xy);
+		float s2 = s * s;
+		k = (s2 - a2 * s2) / (s2 + a2 * V.z * V.z); // Eq. 5
+	}
+#endif
+	float PDF = 2 * VoH * D / (k * NoV + sqrt(NoV * (NoV - NoV * a2) + a2));
 	return PDF;
 }
 
-float VisibleGGXPDF_aniso(float3 V, float3 H, float2 Alpha)
+float VisibleGGXPDF_aniso(float3 V, float3 H, float2 Alpha, bool bLimitVDNFToReflection = true)
 {
 	float NoV = V.z;
 	float NoH = H.z;
@@ -385,42 +402,48 @@ float VisibleGGXPDF_aniso(float3 V, float3 H, float2 Alpha)
 	float S = dot(Hs, Hs);
 	float D = (1.0f / PI) * a2 * Square(a2 / S);
 	float LenV = length(float3(V.x * Alpha.x, V.y * Alpha.y, NoV));
-	float Pdf = (2 * D * VoH) / (NoV + LenV);
+	float k = 1.0;
+#if GGX_BOUNDED_VNDF_SAMPLING
+	if (bLimitVDNFToReflection)
+	{
+		float a = saturate(min(Alpha.x, Alpha.y));
+		float s = 1.0f + length(V.xy);
+		float ka2 = a * a, s2 = s * s;
+		k = (s2 - ka2 * s2) / (s2 + ka2 * V.z * V.z); // Eq. 5
+	}
+#endif
+	float Pdf = (2 * D * VoH) / (k * NoV + LenV);
 	return Pdf;
 }
 
 // PDF = G_SmithV * VoH * D / NoV / (4 * VoH)
 // PDF = G_SmithV * D / (4 * NoV)
-float4 ImportanceSampleVisibleGGX(float2 E, float2 Alpha, float3 V)
+float4 ImportanceSampleVisibleGGX(float2 E, float2 Alpha, float3 V, bool bLimitVDNFToReflection = true)
 {
 	// stretch
 	float3 Vh = normalize(float3(Alpha * V.xy, V.z));
 
-#if 0
-	// [ Heitz 2018, "Sampling the GGX Distribution of Visible Normals" ]
-	// http://jcgt.org/published/0007/04/01/
-
-	// Stable tangent basis based on V
-	float LenSq = Vh.x * Vh.x + Vh.y * Vh.y;
-	float3 Tx = LenSq > 0 ? float3(-Vh.y, Vh.x, 0) * rsqrt(LenSq) : float3(1, 0, 0);
-	float3 Ty = cross(Vh, Tx);
-
-	// #dxr_todo: Evaluate UniformSampleDiskConcentric
-	float2 p = UniformSampleDisk(E);
-	float s = 0.5 + 0.5 * Vh.z;
-	p.y = lerp(sqrt(1 - p.x * p.x), p.y, s);
-
-	float3 H = p.x * Tx + p.y * Ty + sqrt(saturate(1 - dot(p, p))) * Vh;
-#else
 	// "Sampling Visible GGX Normals with Spherical Caps"
 	// Jonathan Dupuy & Anis Benyoub - High Performance Graphics 2023
 	float Phi = (2 * PI) * E.x;
-	float Z = lerp(-Vh.z, 1.0, E.y);
+	float k = 1.0;
+#if GGX_BOUNDED_VNDF_SAMPLING
+	if (bLimitVDNFToReflection)
+	{
+		// If we know we will be reflecting the view vector around the sampled micronormal, we can
+		// tweak the range a bit more to eliminate some of the vectors that will point below the horizon
+		float a = saturate(min(Alpha.x, Alpha.y));
+		float s = 1.0 + length(V.xy);
+		float a2 = a * a, s2 = s * s;
+		k = (s2 - a2 * s2) / (s2 + a2 * V.z * V.z);
+	}
+#endif
+	float Z = lerp(1.0, -k * Vh.z, E.y);
 	float SinTheta = sqrt(saturate(1 - Z * Z));
 	float X = SinTheta * cos(Phi);
 	float Y = SinTheta * sin(Phi);
 	float3 H = float3(X, Y, Z) + Vh;
-#endif
+
 	// unstretch
 	H = normalize(float3(Alpha * H.xy, max(0.0, H.z)));
 
diff --git a/Engine/Shaders/Private/PathTracing/Material/PathTracingGlossy.ush b/Engine/Shaders/Private/PathTracing/Material/PathTracingGlossy.ush
--- a/Engine/Shaders/Private/PathTracing/Material/PathTracingGlossy.ush
+++ b/Engine/Shaders/Private/PathTracing/Material/PathTracingGlossy.ush
@@ -64,7 +64,7 @@ float3x3 GetGGXBasis(float Roughness, float3 WorldNormal, out float2 Alpha)
 }
 
 // Compute GGX lobe Weight and Pdf (without Fresnel term) given a set of vectors in local space (Z up)
-float2 GGXEvalReflection(float3 L, float3 V, float3 H, float2 Alpha)
+float2 GGXEvalReflection(float3 L, float3 V, float3 H, float2 Alpha, bool bLimitVDNFToReflection = true)
 {
 	const float NoL = saturate(L.z);
 	const float NoV = saturate(V.z);
@@ -74,11 +74,21 @@ float2 GGXEvalReflection(float3 L, float3 V, float3 H, float2 Alpha)
 		const float D = D_GGXaniso(Alpha.x, Alpha.y, H.z, H.x, H.y);
 		// See implementation in Vis_SmithJointAniso for G2/(4*NoV*NoL)
 		// We can simplify a bit further since we need both the weight G2/G1 and the pdf
-		const float LenL = length(float3(L.x * Alpha.x, L.y * Alpha.y, NoL));
-		const float LenV = length(float3(V.x * Alpha.x, V.y * Alpha.y, NoV));
-		const float Weight = NoL * (LenV + NoV) / (NoV * LenL + NoL * LenV);
-		const float Pdf = 0.5 * D * rcp(LenV + NoV);
-
+		const float LenL = length(float3(L.xy * Alpha, NoL));
+		const float LenV = length(float3(V.xy * Alpha, NoV));
+		float k = 1.0;
+#if GGX_BOUNDED_VNDF_SAMPLING
+		if (bLimitVDNFToReflection)
+		{
+			float a = saturate(min(Alpha.x, Alpha.y));
+			float s = 1.0f + length(V.xy);
+			float a2 = a * a, s2 = s * s;
+			k = (s2 - a2 * s2) / (s2 + a2 * NoV * NoV); // Eq. 5
+		}
+#endif
+		const float Weight = NoL * (LenV + k * NoV) / (NoV * LenL + NoL * LenV);
+		const float Pdf = 0.5 * D * rcp(LenV + k * NoV);
+		
 		return float2(Weight, Pdf);
 	}
 	return 0;
@@ -99,8 +109,8 @@ float2 GGXEvalRefraction(float3 L, float3 V, float3 H, float2 Alpha, float Eta)
 		const float LoH = saturate(-dot(L, H));
 
 		// Same as logic above, only the pdf math differs slightly
-		const float LenL = length(float3(L.x * Alpha.x, L.y * Alpha.y, NoL));
-		const float LenV = length(float3(V.x * Alpha.x, V.y * Alpha.y, NoV));
+		const float LenL = length(float3(L.xy * Alpha, NoL));
+		const float LenV = length(float3(V.xy * Alpha, NoV));
 		const float Weight = NoL * (LenV + NoV) / (NoV * LenL + NoL * LenV);
 		const float Pdf = (LoH * VoH * Eta * Eta * 2.0 * D) / (Pow2(VoH - Eta * LoH) * (LenV + NoV));
 
diff --git a/Engine/Shaders/Private/PathTracing/Material/PathTracingSolidGlass.ush b/Engine/Shaders/Private/PathTracing/Material/PathTracingSolidGlass.ush
--- a/Engine/Shaders/Private/PathTracing/Material/PathTracingSolidGlass.ush
+++ b/Engine/Shaders/Private/PathTracing/Material/PathTracingSolidGlass.ush
@@ -111,7 +111,7 @@ FMaterialSample RoughGlass_SampleMaterial(
 		}
 
 
-		const float3 H = ImportanceSampleVisibleGGX(RandSample.xy, Data.Alpha, V).xyz;
+		const float3 H = ImportanceSampleVisibleGGX(RandSample.xy, Data.Alpha, V, bDoSpec).xyz;
 
 		if (bDoSpec)
 		{
@@ -130,7 +130,7 @@ FMaterialSample RoughGlass_SampleMaterial(
 			const float NoH = saturate(H.z);
 
 			// Specular lobe
-			const float2 GGXResult = GGXEvalReflection(L, V, H, Data.Alpha);
+			const float2 GGXResult = GGXEvalReflection(L, V, H, Data.Alpha, true);
 			const float3 F = F_Schlick(Payload.SpecularColor, VoH);
 			const float3 SpecWeight = (Payload.BSDFOpacity / Data.LobePdf.y) * F * GGXResult.x * Data.Spec.W;
 			const float SpecPdf = GGXResult.y * Data.LobePdf.y;
@@ -147,7 +147,7 @@ FMaterialSample RoughGlass_SampleMaterial(
 			// transform to world space
 			const float3 L_World = normalize(mul(L, Data.Basis));
 
-			const float2 GGXResult = bRefract ? GGXEvalRefraction(L, V, H, Data.Alpha, Data.Eta) : GGXEvalReflection(L, V, H, Data.Alpha);
+			const float2 GGXResult = bRefract ? GGXEvalRefraction(L, V, H, Data.Alpha, Data.Eta) : GGXEvalReflection(L, V, H, Data.Alpha, false);
 
 			return CreateMaterialSample(L_World, (1.0 - Payload.BSDFOpacity) * GGXResult.x / (Data.EGlass * Data.LobePdf.z), F * GGXResult.y * Data.LobePdf.z, bRefract ? -1.0 : 1.0, Payload.Roughness, bRefract ? PATHTRACER_SCATTER_REFRACT : PATHTRACER_SCATTER_SPECULAR);
 		}
@@ -190,17 +190,22 @@ FMaterialEval RoughGlass_EvalMaterial(
 		Result.AddLobeWithMIS(Payload.BSDFOpacity * Diffuse * ShadowTerminatorTerm(L_World, Payload.WorldNormal, Payload.WorldSmoothNormal) * DiffuseSpecularScale.x, NoL / PI, Data.LobePdf.x);
 
 		// Specular lobe
-		const float2 GGXResult = GGXEvalReflection(L, V, H, Data.Alpha);
-		const float3 Fs = F_Schlick(Payload.SpecularColor, VoH);
-		const float3 SpecWeight = Payload.BSDFOpacity * GGXResult.x * Data.Spec.W * Fs;
-		const float SpecPdf = GGXResult.y;
-		Result.AddLobeWithMIS(SpecWeight * DiffuseSpecularScale.y, SpecPdf, Data.LobePdf.y);
+		{
+			const float2 GGXResult = GGXEvalReflection(L, V, H, Data.Alpha, true);
+			const float3 Fs = F_Schlick(Payload.SpecularColor, VoH);
+			const float3 SpecWeight = Payload.BSDFOpacity * GGXResult.x * Data.Spec.W * Fs;
+			const float SpecPdf = GGXResult.y;
+			Result.AddLobeWithMIS(SpecWeight * DiffuseSpecularScale.y, SpecPdf, Data.LobePdf.y);
+		}
 
 		// Glass lobe
-		const float Fg = FresnelReflectance(VoH, Data.Eta, Data.F0);
-		const float3 GlassWeight = GGXResult.x * (1 - Payload.BSDFOpacity) / Data.EGlass;
-		const float GlassPdf = GGXResult.y * Fg;
-		Result.AddLobeWithMIS(GlassWeight * DiffuseSpecularScale.y, GlassPdf, Data.LobePdf.z);
+		{
+			const float2 GGXResult = GGXEvalReflection(L, V, H, Data.Alpha, false);
+			const float Fg = FresnelReflectance(VoH, Data.Eta, Data.F0);
+			const float3 GlassWeight = GGXResult.x * (1 - Payload.BSDFOpacity) / Data.EGlass;
+			const float GlassPdf = GGXResult.y * Fg;
+			Result.AddLobeWithMIS(GlassWeight * DiffuseSpecularScale.y, GlassPdf, Data.LobePdf.z);
+		}
 	}
 	else if (L.z < 0 && Data.LobePdf.z > 0)
 	{
diff --git a/Engine/Shaders/Private/PathTracing/Material/PathTracingSubstrateSolidGlass.ush b/Engine/Shaders/Private/PathTracing/Material/PathTracingSubstrateSolidGlass.ush
--- a/Engine/Shaders/Private/PathTracing/Material/PathTracingSubstrateSolidGlass.ush
+++ b/Engine/Shaders/Private/PathTracing/Material/PathTracingSubstrateSolidGlass.ush
@@ -98,7 +98,7 @@ FMaterialSample SubstrateSolidGlass_SampleMaterial(
 			OutRoughness = Payload.RoughnessData.y;
 		}
 
-		H = ImportanceSampleVisibleGGX(RandSample.xy, bUseSpec0 ? Data.Alpha0 : Data.Alpha1, V).xyz;
+		H = ImportanceSampleVisibleGGX(RandSample.xy, bUseSpec0 ? Data.Alpha0 : Data.Alpha1, V, false).xyz;
 
 		// Glass lobe (reflection and refraction)
 		float3 L = 0;
@@ -110,8 +110,8 @@ FMaterialSample SubstrateSolidGlass_SampleMaterial(
 
 		FMaterialSample Result = CreateMaterialSample(L_World, 0.0, 0.0, bRefract ? -1.0 : 1.0, OutRoughness, bRefract ? PATHTRACER_SCATTER_REFRACT : PATHTRACER_SCATTER_SPECULAR);
 
-		const float2 GGXResult0 = bRefract ? GGXEvalRefraction(L, V, H, Data.Alpha0, Data.Eta) : GGXEvalReflection(L, V, H, Data.Alpha0);
-		const float2 GGXResult1 = bRefract ? GGXEvalRefraction(L, V, H, Data.Alpha1, Data.Eta) : GGXEvalReflection(L, V, H, Data.Alpha1);
+		const float2 GGXResult0 = bRefract ? GGXEvalRefraction(L, V, H, Data.Alpha0, Data.Eta) : GGXEvalReflection(L, V, H, Data.Alpha0, false);
+		const float2 GGXResult1 = bRefract ? GGXEvalRefraction(L, V, H, Data.Alpha1, Data.Eta) : GGXEvalReflection(L, V, H, Data.Alpha1, false);
 
 		Result.AddLobeWithMIS((1.0 - Payload.RoughnessData.z) * GGXResult0.x / Data.EGlass0, F * GGXResult0.y, Data.LobePdf.x);
 		Result.AddLobeWithMIS((      Payload.RoughnessData.z) * GGXResult1.x / Data.EGlass1, F * GGXResult1.y, Data.LobePdf.y);
@@ -210,8 +210,8 @@ FMaterialEval SubstrateSolidGlass_EvalMaterial(
 
 		{
 			// Specular Lobes
-			const float2 GGXResult0 = GGXEvalReflection(L, V, H, Data.Alpha0);
-			const float2 GGXResult1 = GGXEvalReflection(L, V, H, Data.Alpha1);
+			const float2 GGXResult0 = GGXEvalReflection(L, V, H, Data.Alpha0, false);
+			const float2 GGXResult1 = GGXEvalReflection(L, V, H, Data.Alpha1, false);
 
 			const float Fg = FresnelReflectance(VoH, Data.Eta, Data.F0);
 			const float3 FTint = SubstrateSpecularTint(Payload, NoV, NoL, VoH, NoH);
diff --git a/Engine/Shaders/Private/ShadingEnergyConservationTable.usf b/Engine/Shaders/Private/ShadingEnergyConservationTable.usf
--- a/Engine/Shaders/Private/ShadingEnergyConservationTable.usf
+++ b/Engine/Shaders/Private/ShadingEnergyConservationTable.usf
@@ -67,7 +67,7 @@ void BuildEnergyTableCS(uint3 DispatchThreadId : SV_DispatchThreadID)
 			Weight = float2(GGXResult.x, GGXResult.x * Pow5(1 - dot(V, H)));
 		}
 #elif BUILD_ENERGY_TABLE == 1
-		const float3 H = ImportanceSampleVisibleGGX(RandSample.xy, Alpha, V).xyz;
+		const float3 H = ImportanceSampleVisibleGGX(RandSample.xy, Alpha, V, false).xyz;
 		{
 			// compute weight for outside -> in case
 			float3 L = 0;
@@ -78,7 +78,7 @@ void BuildEnergyTableCS(uint3 DispatchThreadId : SV_DispatchThreadID)
 			}
 			else
 			{
-				Weight.x = GGXEvalReflection(L, V, H, Alpha).x;
+				Weight.x = GGXEvalReflection(L, V, H, Alpha, false).x;
 			}
 		}
 		{
@@ -91,7 +91,7 @@ void BuildEnergyTableCS(uint3 DispatchThreadId : SV_DispatchThreadID)
 			}
 			else
 			{
-				Weight.y = GGXEvalReflection(L, V, H, Alpha).x;
+				Weight.y = GGXEvalReflection(L, V, H, Alpha, false).x;
 			}
 		}
 #elif BUILD_ENERGY_TABLE == 2
diff --git a/Engine/Shaders/Private/Substrate/SubstrateRoughRefraction.usf b/Engine/Shaders/Private/Substrate/SubstrateRoughRefraction.usf
--- a/Engine/Shaders/Private/Substrate/SubstrateRoughRefraction.usf
+++ b/Engine/Shaders/Private/Substrate/SubstrateRoughRefraction.usf
@@ -186,7 +186,7 @@ bool GetGGXDirection(float2 RandomUV, float Roughness, float3 V, inout float3 L,
 	float2 Alpha; 
 	GetGGXBasis(Roughness, WOLRD_NORMAL, Alpha);
 
-	float4 H = ImportanceSampleVisibleGGX(RandomUV, Alpha, V);
+	float4 H = ImportanceSampleVisibleGGX(RandomUV, Alpha, V, false);
 
 	Value = 0;
 	L = 0;
@@ -201,7 +201,7 @@ bool GetGGXDirection(float2 RandomUV, float Roughness, float3 V, inout float3 L,
 	}
 
 	// Reflection
-	Value = GGXEvalReflection(L, V, H.xyz, Alpha).x;
+	Value = GGXEvalReflection(L, V, H.xyz, Alpha, false).x;
 	return MODE_REFLECTION == 1;
 }
 
