* Rename ShaderPrint maxvalue into maxcharacter for better clarity.
* Improve RequestXXX API comment for better clarity

diff --git a/Engine/Plugins/Animation/DeformerGraph/Shaders/Private/DataInterfaceDebugDraw.ush b/Engine/Plugins/Animation/DeformerGraph/Shaders/Private/DataInterfaceDebugDraw.ush
--- a/Engine/Plugins/Animation/DeformerGraph/Shaders/Private/DataInterfaceDebugDraw.ush
+++ b/Engine/Plugins/Animation/DeformerGraph/Shaders/Private/DataInterfaceDebugDraw.ush
@@ -48,7 +48,7 @@ float4x4 {DataInterfaceName}_LocalToWorld;
 int2 {DataInterfaceName}_Resolution;
 float2 {DataInterfaceName}_FontSize;
 float2 {DataInterfaceName}_FontSpacing;
-uint {DataInterfaceName}_MaxValueCount;
+uint {DataInterfaceName}_MaxCharacterCount;
 uint {DataInterfaceName}_MaxSymbolCount;
 uint {DataInterfaceName}_MaxStateCount;
 uint {DataInterfaceName}_MaxLineCount;
@@ -67,7 +67,7 @@ FDebugDraw ReadDebugDraw_{DataInterfaceName}()
 	DebugDraw.Context.Config.TranslatedWorldOffset = 0;
 	DebugDraw.Context.Config.FontSize = {DataInterfaceName}_FontSize;
 	DebugDraw.Context.Config.FontSpacing = {DataInterfaceName}_FontSpacing;
-	DebugDraw.Context.Config.MaxValueCount = {DataInterfaceName}_MaxValueCount;
+	DebugDraw.Context.Config.MaxCharacterCount = {DataInterfaceName}_MaxCharacterCount;
 	DebugDraw.Context.Config.MaxSymbolCount = {DataInterfaceName}_MaxSymbolCount;
 	DebugDraw.Context.Config.MaxStateCount = {DataInterfaceName}_MaxStateCount;
 	DebugDraw.Context.Config.MaxLineCount = {DataInterfaceName}_MaxLineCount;
diff --git a/Engine/Plugins/Animation/DeformerGraph/Source/OptimusCore/Private/DataInterfaces/OptimusDataInterfaceDebugDraw.cpp b/Engine/Plugins/Animation/DeformerGraph/Source/OptimusCore/Private/DataInterfaces/OptimusDataInterfaceDebugDraw.cpp
--- a/Engine/Plugins/Animation/DeformerGraph/Source/OptimusCore/Private/DataInterfaces/OptimusDataInterfaceDebugDraw.cpp
+++ b/Engine/Plugins/Animation/DeformerGraph/Source/OptimusCore/Private/DataInterfaces/OptimusDataInterfaceDebugDraw.cpp
@@ -62,7 +62,7 @@ BEGIN_SHADER_PARAMETER_STRUCT(FDebugDrawDataInterfaceParameters, )
 	SHADER_PARAMETER(FIntPoint, Resolution)
 	SHADER_PARAMETER(FVector2f, FontSize)
 	SHADER_PARAMETER(FVector2f, FontSpacing)
-	SHADER_PARAMETER(uint32, MaxValueCount)
+	SHADER_PARAMETER(uint32, MaxCharacterCount)
 	SHADER_PARAMETER(uint32, MaxSymbolCount)
 	SHADER_PARAMETER(uint32, MaxStateCount)
 	SHADER_PARAMETER(uint32, MaxLineCount)
@@ -128,7 +128,7 @@ FOptimusDebugDrawDataProviderProxy::FOptimusDebugDrawDataProviderProxy(UPrimitiv
 	Setup.FontSize = InDebugDrawParameters.FontSize;
 	Setup.MaxLineCount = Setup.bEnabled ? InDebugDrawParameters.MaxLineCount : 0;
 	Setup.MaxTriangleCount = Setup.bEnabled ? InDebugDrawParameters.MaxTriangleCount : 0;
-	Setup.MaxValueCount = Setup.bEnabled ? InDebugDrawParameters.MaxCharacterCount : 0;
+	Setup.MaxCharacterCount = Setup.bEnabled ? InDebugDrawParameters.MaxCharacterCount : 0;
 	Setup.PreViewTranslation = PreViewTranslation;
 
 	ShaderPrint::GetParameters(Setup, ConfigParameters);
@@ -183,7 +183,7 @@ void FOptimusDebugDrawDataProviderProxy::GatherDispatchData(FDispatchData const&
 		Parameters.Resolution = ConfigParameters.Resolution;
 		Parameters.FontSize = ConfigParameters.FontSize;
 		Parameters.FontSpacing = ConfigParameters.FontSpacing;
-		Parameters.MaxValueCount = ConfigParameters.MaxValueCount;
+		Parameters.MaxCharacterCount = ConfigParameters.MaxCharacterCount;
 		Parameters.MaxSymbolCount = ConfigParameters.MaxSymbolCount;
 		Parameters.MaxStateCount = ConfigParameters.MaxStateCount;
 		Parameters.MaxLineCount = ConfigParameters.MaxLineCount;
diff --git a/Engine/Shaders/Private/ShaderPrint.ush b/Engine/Shaders/Private/ShaderPrint.ush
--- a/Engine/Shaders/Private/ShaderPrint.ush
+++ b/Engine/Shaders/Private/ShaderPrint.ush
@@ -61,7 +61,7 @@ void PackShaderPrintConfig(FShaderPrintConfig In)
 	UEDiagnosticBuffer[Offset+4] = asuint(In.TranslatedWorldOffset.z);
 	UEDiagnosticBuffer[Offset+5] = PackFloat2ToUInt(In.FontSize);
 	UEDiagnosticBuffer[Offset+6] = PackFloat2ToUInt(In.FontSpacing);
-	UEDiagnosticBuffer[Offset+7] = In.MaxValueCount;
+	UEDiagnosticBuffer[Offset+7] = In.MaxCharacterCount;
 	UEDiagnosticBuffer[Offset+8] = In.MaxSymbolCount;
 	UEDiagnosticBuffer[Offset+9] = In.MaxStateCount;
 	UEDiagnosticBuffer[Offset+10]= In.MaxLineCount;
@@ -80,7 +80,7 @@ FShaderPrintConfig UnpackShaderPrintConfig()
 	Out.TranslatedWorldOffset.z = asfloat(UEDiagnosticBuffer[Offset+4]);
 	Out.FontSize 				= UnpackFloat2FromUInt(UEDiagnosticBuffer[Offset+5]);
 	Out.FontSpacing 			= UnpackFloat2FromUInt(UEDiagnosticBuffer[Offset+6]);
-	Out.MaxValueCount			= UEDiagnosticBuffer[Offset+7];
+	Out.MaxCharacterCount		= UEDiagnosticBuffer[Offset+7];
 	Out.MaxSymbolCount			= UEDiagnosticBuffer[Offset+8];
 	Out.MaxStateCount			= UEDiagnosticBuffer[Offset+9];
 	Out.MaxLineCount			= UEDiagnosticBuffer[Offset+10];
@@ -113,7 +113,7 @@ FShaderPrintConfig InitShaderPrintContextConfig()
 	Config.TranslatedWorldOffset = ShaderPrintData.TranslatedWorldOffset;
 	Config.FontSize = ShaderPrintData.FontSize;
 	Config.FontSpacing = ShaderPrintData.FontSpacing;
-	Config.MaxValueCount = ShaderPrintData.MaxValueCount;
+	Config.MaxCharacterCount = ShaderPrintData.MaxCharacterCount;
 	Config.MaxSymbolCount = ShaderPrintData.MaxSymbolCount;
 	Config.MaxStateCount = ShaderPrintData.MaxStateCount;
 	Config.MaxLineCount = ShaderPrintData.MaxLineCount;
diff --git a/Engine/Shaders/Private/ShaderPrintCommon.ush b/Engine/Shaders/Private/ShaderPrintCommon.ush
--- a/Engine/Shaders/Private/ShaderPrintCommon.ush
+++ b/Engine/Shaders/Private/ShaderPrintCommon.ush
@@ -115,13 +115,13 @@ uint GetSymbolOffset(uint InIndex)
 {
 	return SHADER_PRINT_COUNTER_COUNT + SHADER_PRINT_UINT_COUNT_SYMBOL * InIndex; 
 }
-uint GetPrimitiveLineOffset(uint InIndex, uint MaxValueCount)
+uint GetPrimitiveLineOffset(uint InIndex, uint MaxCharacterCount)
 {
-	return SHADER_PRINT_COUNTER_COUNT + SHADER_PRINT_UINT_COUNT_SYMBOL * MaxValueCount + SHADER_PRINT_UINT_COUNT_LINE * InIndex; 
+	return SHADER_PRINT_COUNTER_COUNT + SHADER_PRINT_UINT_COUNT_SYMBOL * MaxCharacterCount + SHADER_PRINT_UINT_COUNT_LINE * InIndex; 
 }
-uint GetPrimitiveTriangleOffset(uint InIndex, uint MaxValueCount, uint MaxLineCount)
+uint GetPrimitiveTriangleOffset(uint InIndex, uint MaxCharacterCount, uint MaxLineCount)
 {
-	return SHADER_PRINT_COUNTER_COUNT + SHADER_PRINT_UINT_COUNT_SYMBOL * MaxValueCount + SHADER_PRINT_UINT_COUNT_LINE * MaxLineCount + SHADER_PRINT_UINT_COUNT_TRIANGLE * InIndex; 
+	return SHADER_PRINT_COUNTER_COUNT + SHADER_PRINT_UINT_COUNT_SYMBOL * MaxCharacterCount + SHADER_PRINT_UINT_COUNT_LINE * MaxLineCount + SHADER_PRINT_UINT_COUNT_TRIANGLE * InIndex; 
 }
 
 void ClearCounters(RWStructuredBuffer<uint> InRWBuffer)
@@ -298,7 +298,7 @@ struct FShaderPrintConfig
 	float3 TranslatedWorldOffset;
 	float2 FontSize;
 	float2 FontSpacing;
-	uint MaxValueCount;
+	uint MaxCharacterCount;
 	uint MaxSymbolCount;
 	uint MaxStateCount;
 	uint MaxLineCount;
@@ -341,8 +341,8 @@ void WriteSymbol(uint Offset, FShaderPrintItem In, FShaderPrintContext Ctx)
 
 void ShaderPrint_Internal(FShaderPrintContext Ctx, in FShaderPrintItem Item)
 {
-	// If MaxValueCount is 0 then we don't reset the buffer counter so early out here
-	if (!Ctx.bIsActive || Ctx.Config.MaxValueCount == 0)
+	// If MaxCharacterCount is 0 then we don't reset the buffer counter so early out here
+	if (!Ctx.bIsActive || Ctx.Config.MaxCharacterCount == 0)
 	{
 		return;
 	}
@@ -353,7 +353,7 @@ void ShaderPrint_Internal(FShaderPrintContext Ctx, in FShaderPrintItem Item)
 
 	// Prevent writing off the buffer
 	// Note that counter still increases so need clamp when reading it in later passes
-	if (uint(IndexToStore) >= Ctx.Config.MaxValueCount)
+	if (uint(IndexToStore) >= Ctx.Config.MaxCharacterCount)
 	{
 		return;
 	}
@@ -422,8 +422,8 @@ uint ShaderPrintGetHash(FShaderPrintText InTextEntry);
 // Function for reading global TEXT string
 float2 ShaderPrintText_Internal(FShaderPrintContext Ctx, bool bIsActive, float2 InPos, FShaderPrintText InTextEntry, FFontColor InColor)
 {
-	// If MaxValueCount is 0 then we don't reset the buffer counter so early out here
-	if (bIsActive && Ctx.Config.MaxValueCount > 0)
+	// If MaxCharacterCount is 0 then we don't reset the buffer counter so early out here
+	if (bIsActive && Ctx.Config.MaxCharacterCount > 0)
 	{
 		const uint Begin = ShaderPrintGetOffset(InTextEntry);
 		const uint End   = ShaderPrintGetOffset(InitShaderPrintText(InTextEntry.Index + 1));
@@ -434,7 +434,7 @@ float2 ShaderPrintText_Internal(FShaderPrintContext Ctx, bool bIsActive, float2
 		// Note that counter still increases so need clamp when reading it in later passes
 		int IndexToStore = 0;
 		InterlockedAdd(SHADER_PRINT_RWENTRYBUFFER(Ctx,SHADER_PRINT_COUNTER_OFFSET_SYMBOL), Count, IndexToStore);
-		if (uint(IndexToStore + Count) < Ctx.Config.MaxValueCount)
+		if (uint(IndexToStore + Count) < Ctx.Config.MaxCharacterCount)
 		{
 			for (uint i = Begin; i < End; ++i)
 			{
@@ -747,9 +747,9 @@ FPackedTriangleElement PackTriangleElement(FTriangleElement In)
 	return Out;
 }
 
-FTriangleElement UnpackTriangleElement(StructuredBuffer<uint> InPrimitiveBuffer, uint InIndex, uint MaxValueCount, uint MaxLineCount)
+FTriangleElement UnpackTriangleElement(StructuredBuffer<uint> InPrimitiveBuffer, uint InIndex, uint MaxCharacterCount, uint MaxLineCount)
 {
-	const uint Index12 = GetPrimitiveTriangleOffset(InIndex, MaxValueCount, MaxLineCount);
+	const uint Index12 = GetPrimitiveTriangleOffset(InIndex, MaxCharacterCount, MaxLineCount);
 
 	FPackedTriangleElement Out = (FPackedTriangleElement)0;
 	Out.Packed0.x = InPrimitiveBuffer[Index12 + 0];
@@ -783,7 +783,7 @@ bool AllocateTriangleElement(FShaderPrintContext Ctx, uint Count, inout uint Out
 
 void AddTriangleElement(FShaderPrintContext Ctx, FTriangleElement In, uint Index)
 {
-	const uint Index12 = GetPrimitiveTriangleOffset(Index, Ctx.Config.MaxValueCount, Ctx.Config.MaxLineCount);
+	const uint Index12 = GetPrimitiveTriangleOffset(Index, Ctx.Config.MaxCharacterCount, Ctx.Config.MaxLineCount);
 
 	FPackedTriangleElement Out = PackTriangleElement(In);
 	SHADER_PRINT_RWENTRYBUFFER(Ctx, Index12 + 0) = Out.Packed0.x;
@@ -951,9 +951,9 @@ FPackedLineElement PackLineElement(FLineElement In)
 	return Out;
 }
 
-FLineElement UnpackLineElement(StructuredBuffer<uint> InPrimitiveBuffer, uint InIndex, uint MaxValueCount)
+FLineElement UnpackLineElement(StructuredBuffer<uint> InPrimitiveBuffer, uint InIndex, uint MaxCharacterCount)
 {
-	const uint Index8 = GetPrimitiveLineOffset(InIndex, MaxValueCount);
+	const uint Index8 = GetPrimitiveLineOffset(InIndex, MaxCharacterCount);
 
 	FPackedLineElement Out = (FPackedLineElement)0;
 	Out.Packed0.x = InPrimitiveBuffer[Index8 + 0];
@@ -982,7 +982,7 @@ bool AllocateLineElement(FShaderPrintContext Ctx, uint Count, inout uint OutInde
 
 void AddLineElement(FShaderPrintContext Ctx, FLineElement In, uint Index)
 {
-	const uint Index8 = GetPrimitiveLineOffset(Index, Ctx.Config.MaxValueCount);
+	const uint Index8 = GetPrimitiveLineOffset(Index, Ctx.Config.MaxCharacterCount);
 
 	FPackedLineElement Out = PackLineElement(In);
 	SHADER_PRINT_RWENTRYBUFFER(Ctx,Index8 + 0) = Out.Packed0.x;
diff --git a/Engine/Shaders/Private/ShaderPrintDraw.usf b/Engine/Shaders/Private/ShaderPrintDraw.usf
--- a/Engine/Shaders/Private/ShaderPrintDraw.usf
+++ b/Engine/Shaders/Private/ShaderPrintDraw.usf
@@ -42,7 +42,7 @@ void ClearCounterCS()
 [numthreads(1, 1, 1)]
 void BuildIndirectDispatchArgsCS()
 {
-	uint ValueCount = min(ShaderPrintData.MaxValueCount, ValuesBuffer[SHADER_PRINT_COUNTER_OFFSET_SYMBOL]);
+	uint ValueCount = min(ShaderPrintData.MaxCharacterCount, ValuesBuffer[SHADER_PRINT_COUNTER_OFFSET_SYMBOL]);
 	WriteDispatchIndirectArgs(RWIndirectDispatchArgsBuffer, 0, (ValueCount + 63)/64, 1, 1);
 	ClearCounters(RWSymbolsBuffer);
 }
@@ -200,7 +200,7 @@ void AddHexSymbols(inout float2 ScreenPos, uint Value, in float3 Color)
 [numthreads(64, 1, 1)]
 void BuildSymbolBufferCS(uint3 DispatchThreadId : SV_DispatchThreadID)
 {
-	uint ValueCount = min(ShaderPrintData.MaxValueCount, ValuesBuffer[SHADER_PRINT_COUNTER_OFFSET_SYMBOL]);
+	uint ValueCount = min(ShaderPrintData.MaxCharacterCount, ValuesBuffer[SHADER_PRINT_COUNTER_OFFSET_SYMBOL]);
 	if (DispatchThreadId.x >= ValueCount)
 	{
 		return;
@@ -434,7 +434,7 @@ void UploadCS(int2 DispatchThreadId : SV_DispatchThreadID)
 	In.TranslatedWorldOffset = ShaderPrintData.TranslatedWorldOffset;
 	In.FontSize = ShaderPrintData.FontSize;
 	In.FontSpacing = ShaderPrintData.FontSpacing;
-	In.MaxValueCount = ShaderPrintData.MaxValueCount;
+	In.MaxCharacterCount = ShaderPrintData.MaxCharacterCount;
 	In.MaxSymbolCount = ShaderPrintData.MaxSymbolCount;
 	In.MaxStateCount = ShaderPrintData.MaxStateCount;
 	In.MaxLineCount = ShaderPrintData.MaxLineCount;
@@ -477,13 +477,13 @@ void CopyDataType(FShaderPrintConfig Config, uint Type)
 	{
 		CounterOffset = SHADER_PRINT_COUNTER_OFFSET_LINE; 		
 		DataTypeSizeInUint = SHADER_PRINT_UINT_COUNT_LINE;
-		DataOffset0 = GetPrimitiveLineOffset(0, Config.MaxValueCount);
+		DataOffset0 = GetPrimitiveLineOffset(0, Config.MaxCharacterCount);
 	}
 	else if (Type == TYPE_TRIANGLE)
 	{
 		CounterOffset = SHADER_PRINT_COUNTER_OFFSET_TRIANGLE;	
 		DataTypeSizeInUint = SHADER_PRINT_UINT_COUNT_TRIANGLE;
-		DataOffset0 = GetPrimitiveTriangleOffset(0, Config.MaxValueCount, Config.MaxLineCount);	
+		DataOffset0 = GetPrimitiveTriangleOffset(0, Config.MaxCharacterCount, Config.MaxLineCount);	
 	}
 	else
 	{
diff --git a/Engine/Shaders/Private/ShaderPrintDrawPrimitive.usf b/Engine/Shaders/Private/ShaderPrintDrawPrimitive.usf
--- a/Engine/Shaders/Private/ShaderPrintDrawPrimitive.usf
+++ b/Engine/Shaders/Private/ShaderPrintDrawPrimitive.usf
@@ -27,7 +27,7 @@ void ShaderDrawDebugVS(
 	out float4 Position : SV_POSITION,
 	out float4 ColorAlpha : TEXCOORD0)
 {
-	const FLineElement Item = UnpackLineElement(ShaderDrawDebugPrimitive, InstanceId, ShaderPrintData.MaxValueCount);
+	const FLineElement Item = UnpackLineElement(ShaderDrawDebugPrimitive, InstanceId, ShaderPrintData.MaxCharacterCount);
 
 	float4 VertexPositionTWS = float4(VertexId == 0 ? Item.Pos0 : Item.Pos1, 1.0f);
 	float4 VertexColor = VertexId == 0 ? Item.Color0 : Item.Color1;
@@ -61,7 +61,7 @@ void ShaderDrawDebugVS(
 	out float4 Position : SV_POSITION,
 	out float4 ColorAlpha : TEXCOORD0)
 {
-	const FTriangleElement Item = UnpackTriangleElement(ShaderDrawDebugPrimitive, InstanceId, ShaderPrintData.MaxValueCount, ShaderPrintData.MaxLineCount);
+	const FTriangleElement Item = UnpackTriangleElement(ShaderDrawDebugPrimitive, InstanceId, ShaderPrintData.MaxCharacterCount, ShaderPrintData.MaxLineCount);
 
 	float4 VertexPositionTWS = float4(VertexId == 0 ? Item.Pos0 : (VertexId == 1 ? Item.Pos1 : Item.Pos2), 1.0f);
 	float4 VertexColor = Item.Color;
diff --git a/Engine/Shaders/Public/ShaderVersion.ush b/Engine/Shaders/Public/ShaderVersion.ush
--- a/Engine/Shaders/Public/ShaderVersion.ush
+++ b/Engine/Shaders/Public/ShaderVersion.ush
@@ -2,4 +2,4 @@
 // Each time the console command is executed it generates a new GUID. As this file is included
 // in Platform.ush (which should be included in any shader) it allows to invalidate the shader DDC.
 // 
-#pragma message("UESHADERMETADATA_VERSION 9A9CDD21-4719-4331-8470-F3EA166C1395")
+#pragma message("UESHADERMETADATA_VERSION 15A33BF2-B4B6-4932-BF5A-6BD3D2689F97")
diff --git a/Engine/Source/Runtime/Renderer/Private/ShaderPrint.cpp b/Engine/Source/Runtime/Renderer/Private/ShaderPrint.cpp
--- a/Engine/Source/Runtime/Renderer/Private/ShaderPrint.cpp
+++ b/Engine/Source/Runtime/Renderer/Private/ShaderPrint.cpp
@@ -143,7 +143,7 @@ namespace ShaderPrint
 	//////////////////////////////////////////////////////////////////////////////////////////////////
 	// Struct & Functions
 
-	static uint32 GetMaxValueCount()
+	static uint32 GetMaxCharacterCount()
 	{
 		return FMath::Max(CVarMaxCharacterCount.GetValueOnAnyThread() + int32(GCachedShaderPrintMaxRequest.CharacterCount), 0);
 	}
@@ -170,10 +170,10 @@ namespace ShaderPrint
 	static uint32 GetPackedSymbolUintSize()   { return 4; }
 
 	// Get symbol buffer size
-	// This is some multiple of the value buffer size to allow for maximum value->symbol expansion
-	static uint32 GetMaxSymbolCountFromValueCount(uint32 MaxValueCount)
+	// This is some multiple of the character buffer size to allow for maximum character->symbol expansion
+	static uint32 GetMaxSymbolCountFromCharacterCount(uint32 InMaxCharacterCount)
 	{
-		return MaxValueCount * 12u;
+		return InMaxCharacterCount * 12u;
 	}
 
 	static bool IsDrawLocked()
@@ -245,8 +245,8 @@ namespace ShaderPrint
 		OutParameters.FontSpacing = FVector2f(FontWidth + SpaceWidth, FontHeight + SpaceHeight);
 		OutParameters.Resolution = InSetup.ViewRect.Size();
 		OutParameters.CursorCoord = InSetup.CursorCoord;
-		OutParameters.MaxValueCount = InSetup.MaxValueCount;
-		OutParameters.MaxSymbolCount = GetMaxSymbolCountFromValueCount(InSetup.MaxValueCount);
+		OutParameters.MaxCharacterCount = InSetup.MaxCharacterCount;
+		OutParameters.MaxSymbolCount = GetMaxSymbolCountFromCharacterCount(InSetup.MaxCharacterCount);
 		OutParameters.MaxStateCount = InSetup.MaxStateCount;
 		OutParameters.MaxLineCount = InSetup.MaxLineCount;
 		OutParameters.MaxTriangleCount = InSetup.MaxTriangleCount;
@@ -711,7 +711,7 @@ namespace ShaderPrint
 
 		FontSize = FIntPoint(FMath::Max(CVarFontSize.GetValueOnAnyThread(), 1), FMath::Max(CVarFontSize.GetValueOnAnyThread(), 1));
 		FontSpacing = FIntPoint(FMath::Max(CVarFontSpacingX.GetValueOnAnyThread(), 1), FMath::Max(CVarFontSpacingY.GetValueOnAnyThread(), 1));
-		MaxValueCount = bEnabled ? GetMaxValueCount() : 0;
+		MaxCharacterCount = bEnabled ? GetMaxCharacterCount() : 0;
 		MaxStateCount = bEnabled ? GetMaxWidgetCount() : 0;
 		MaxLineCount = bEnabled ? GetMaxLineCount() : 0;
 		MaxTriangleCount = bEnabled ? GetMaxTriangleCount(): 0;
@@ -729,7 +729,7 @@ namespace ShaderPrint
 
 		FontSize = FIntPoint(FMath::Max(CVarFontSize.GetValueOnAnyThread(), 1), FMath::Max(CVarFontSize.GetValueOnAnyThread(), 1));
 		FontSpacing = FIntPoint(FMath::Max(CVarFontSpacingX.GetValueOnAnyThread(), 1), FMath::Max(CVarFontSpacingY.GetValueOnAnyThread(), 1));
-		MaxValueCount = bEnabled ? GetMaxValueCount() : 0;
+		MaxCharacterCount = bEnabled ? GetMaxCharacterCount() : 0;
 		MaxStateCount = bEnabled ? GetMaxWidgetCount() : 0;
 		MaxLineCount = bEnabled ? GetMaxLineCount() : 0;
 		MaxTriangleCount = bEnabled ? GetMaxTriangleCount() : 0;
@@ -740,7 +740,7 @@ namespace ShaderPrint
 	{
 		const uint32 UintElementCount = 
 			GetCountersUintSize() + 
-			GetPackedSymbolUintSize() * In.MaxValueCount +
+			GetPackedSymbolUintSize() * In.MaxCharacterCount +
 			GetPackedLineUintSize() * In.MaxLineCount + 
 			GetPackedTriangleUintSize() * In.MaxTriangleCount;
 		return UintElementCount;
@@ -938,7 +938,7 @@ namespace ShaderPrint
 		FScreenPassTexture OutputTexture)
 	{
 		// Initialize graph managed resources
-		const uint32 UintElementCount = GetCountersUintSize() + GetPackedSymbolUintSize() * GetMaxSymbolCountFromValueCount(ShaderPrintData.Setup.MaxValueCount);
+		const uint32 UintElementCount = GetCountersUintSize() + GetPackedSymbolUintSize() * GetMaxSymbolCountFromCharacterCount(ShaderPrintData.Setup.MaxCharacterCount);
 		FRDGBufferRef SymbolBuffer = GraphBuilder.CreateBuffer(FRDGBufferDesc::CreateStructuredDesc(4, UintElementCount), TEXT("ShaderPrint.SymbolBuffer"));
 		FRDGBufferRef IndirectDispatchArgsBuffer = GraphBuilder.CreateBuffer(FRDGBufferDesc::CreateIndirectDesc<FRHIDispatchIndirectParameters>(1), TEXT("ShaderPrint.IndirectDispatchArgs"));
 		FRDGBufferRef IndirectDrawArgsBuffer = GraphBuilder.CreateBuffer(FRDGBufferDesc::CreateIndirectDesc(5), TEXT("ShaderPrint.IndirectDrawArgs"));
diff --git a/Engine/Source/Runtime/Renderer/Public/ShaderPrintParameters.h b/Engine/Source/Runtime/Renderer/Public/ShaderPrintParameters.h
--- a/Engine/Source/Runtime/Renderer/Public/ShaderPrintParameters.h
+++ b/Engine/Source/Runtime/Renderer/Public/ShaderPrintParameters.h
@@ -30,7 +30,7 @@ namespace ShaderPrint
 		SHADER_PARAMETER(FVector3f, TranslatedWorldOffset)
 		SHADER_PARAMETER(FVector2f, FontSize)
 		SHADER_PARAMETER(FVector2f, FontSpacing)
-		SHADER_PARAMETER(uint32, MaxValueCount)
+		SHADER_PARAMETER(uint32, MaxCharacterCount)
 		SHADER_PARAMETER(uint32, MaxSymbolCount)
 		SHADER_PARAMETER(uint32, MaxStateCount)
 		SHADER_PARAMETER(uint32, MaxLineCount)
@@ -75,13 +75,27 @@ namespace ShaderPrint
 	RENDERER_API bool IsDefaultViewEnabled();
 
 	/**
-	 * Call to ensure enough space for some number of characters/lines, is added cumulatively each frame, to make 
+	 * Call to ensure enough space for some number of characters, is added cumulatively each frame, to make 
 	 * it possible for several systems to request a certain number independently.
 	 * Is used to grow the max element count for subsequent frames (as the allocation happens early in the frame).
-	 * @param The number of elements requested, an element corresponds to a line, so a cube, for example, needs 12 elements.
+	 * @param The number of characters requested.
 	 */
 	RENDERER_API void RequestSpaceForCharacters(uint32 MaxElementCount);
+
+	/**
+	* Call to ensure enough space for some number of lines segments, is added cumulatively each frame, to make 
+	* it possible for several systems to request a certain number independently.
+	* Is used to grow the max element count for subsequent frames (as the allocation happens early in the frame).
+	* @param The number of line segments requested.
+	*/
 	RENDERER_API void RequestSpaceForLines(uint32 MaxElementCount);
+
+	/**
+	* Call to ensure enough space for some number of solid triangles, is added cumulatively each frame, to make 
+	* it possible for several systems to request a certain number independently.
+	* Is used to grow the max element count for subsequent frames (as the allocation happens early in the frame).
+	* @param The number of triangles requested.
+	*/
 	RENDERER_API void RequestSpaceForTriangles(uint32 MaxElementCount);
 
 	/** Structure containing setup for shader print capturing. */
@@ -109,7 +123,7 @@ namespace ShaderPrint
 		/** Font spacing in pixels (not including font size). */
 		FIntPoint FontSpacing = 1;
 		/** Initial size of character buffer. Will also be increased by RequestSpaceForCharacters(). */
-		uint32 MaxValueCount = 0;
+		uint32 MaxCharacterCount = 0;
 		/** Initial size of widget buffer. */
 		uint32 MaxStateCount = 0;
 		/** Initial size of line buffer. Will also be increased by RequestSpaceForLines(). */
