Support both ContactShadowCastingIntensity and ContactShadowNonCastingIntensity in Bend SSS.
- extended WriteScreenSpaceShadow(...) to support outputting HitUV
    - adds significant overhead on some platforms so needs to be manually enabled using preprocessor flag
- implemented 2 permutations of FScreenSpaceShadowsBendCS:
    - sample GBuffer/Stencil at HitUV to determine whether CastingIntensity or NonCastingIntensity should be used
    - fade between CastingIntensity or NonCastingIntensity based on pixel depth (not correct but can be useful to avoid HitUV overhead on some platforms)

diff --git a/Engine/Shaders/Private/ScreenSpaceShadows.usf b/Engine/Shaders/Private/ScreenSpaceShadows.usf
--- a/Engine/Shaders/Private/ScreenSpaceShadows.usf
+++ b/Engine/Shaders/Private/ScreenSpaceShadows.usf
@@ -8,10 +8,15 @@
 
 #include "/Engine/Public/Platform.ush"
 
+#define INTENSITY_MODE_PRIMITIVE_FLAG 0
+#define INTENSITY_MODE_DEPTH_BASED_APPROXIMATION 1
+#define INTENSITY_MODE_FORCE_CASTING_INTENSITY 2
+
 #define WAVE_SIZE 64
 #define SAMPLE_COUNT 60
 #define HARD_SHADOW_SAMPLES 4
 #define FADE_OUT_SAMPLES 8
+#define TRACK_HIT_UV (DIM_INTENSITY_MODE == INTENSITY_MODE_PRIMITIVE_FLAG)
 
 #include "bend_sss_gpu.ush"
 
@@ -19,11 +24,10 @@ float2 InvDepthTextureSize;
 float4 LightCoordinate;
 int3 WaveOffset;
 
-Texture2D<float> DepthTexture;
 RWTexture2D<float> OutputTexture;
 SamplerState PointBorderSampler;
 
-DispatchParameters SetupBendDispatchParameters(float NearDepthVal, float FarDepthVal, float SurfaceThickness)
+DispatchParameters SetupBendDispatchParameters(Texture2D DepthTexture, float NearDepthVal, float FarDepthVal, float SurfaceThickness)
 {
 	DispatchParameters Parameters;
 	Parameters.SetDefaults();
@@ -71,6 +75,8 @@ float ContactShadowLength;
 uint bContactShadowLengthInWS;
 float ContactShadowCastingIntensity;
 float ContactShadowNonCastingIntensity;
+float ContactShadowIntensityFadeStart;
+float ContactShadowIntensityFadeOneOverLength;
 float SurfaceThickness;
 
 uint DownsampleFactor;
@@ -153,13 +159,32 @@ void ScreenSpaceShadowsBendCS(
 	uint3 DispatchThreadId : SV_DispatchThreadID,
 	uint3 GroupThreadId : SV_GroupThreadID)
 {
-	DispatchParameters Parameters = SetupBendDispatchParameters(NearDepthValue, FarDepthValue, SurfaceThickness);
+#if SHADING_PATH_MOBILE
+	Texture2D DepthTexture = MobileSceneTextures.SceneDepthTexture;
+#else
+	Texture2D DepthTexture = SceneTexturesStruct.SceneDepthTexture;
+#endif
+	DispatchParameters Parameters = SetupBendDispatchParameters(DepthTexture, NearDepthValue, FarDepthValue, SurfaceThickness);
 
 	float Result;
 	int2 ResultXY;
-	WriteScreenSpaceShadow(Parameters, GroupId, GroupThreadId.x, Result, ResultXY);
+	float2 HitUV;
+	WriteScreenSpaceShadow(Parameters, GroupId, GroupThreadId.x, Result, ResultXY, HitUV);
+
+#if DIM_INTENSITY_MODE == INTENSITY_MODE_PRIMITIVE_FLAG
+	const float ShadowOcclusion = GetContactShadowOcclusion(HitUV);
+#elif DIM_INTENSITY_MODE == INTENSITY_MODE_DEPTH_BASED_APPROXIMATION
+	const float DeviceZ = DepthTexture[ResultXY].r;
+	const float SceneDepth = ConvertFromDeviceZ(DeviceZ);
+	const float T = saturate((SceneDepth - ContactShadowIntensityFadeStart) * ContactShadowIntensityFadeOneOverLength);
+	const float ShadowOcclusion = lerp(ContactShadowNonCastingIntensity, ContactShadowCastingIntensity, T);
+#elif DIM_INTENSITY_MODE == INTENSITY_MODE_FORCE_CASTING_INTENSITY
+	const float ShadowOcclusion = ContactShadowCastingIntensity;
+#else
+#	error "Unsupported DIM_INTENSITY_MODE"
+#endif
 
-	OutputTexture[ResultXY] = 1 - (1 - Result) * ContactShadowCastingIntensity;
+	OutputTexture[ResultXY] = 1 - (1 - Result) * ShadowOcclusion;
 }
 
 #endif // BEND_SSS
diff --git a/Engine/Shaders/Private/bend_sss_gpu.ush b/Engine/Shaders/Private/bend_sss_gpu.ush
--- a/Engine/Shaders/Private/bend_sss_gpu.ush
+++ b/Engine/Shaders/Private/bend_sss_gpu.ush
@@ -53,6 +53,9 @@
 //		
 //		#define FADE_OUT_SAMPLES 8					// Number of samples that will fade out at the end of the shadow (for a minor cost).
 //													// Recommended starting value: 8
+//		
+//		#define TRACK_HIT_UV 1						// Whether to track Hit UV while tracing (causes significant overhead on some platforms).
+//													// If disabled outHitUV will be (0,0)
 
 
 #if defined(__HLSL_VERSION) || defined(__hlsl_dx_compiler) || 1
@@ -71,7 +74,7 @@
 // 
 //	(int3)	inGroupID:			Compute shader group id register (SV_GroupID)
 //	(int)	inGroupThreadId:	Compute shader group thread id register (SV_GroupThreadID)
-void WriteScreenSpaceShadow(struct DispatchParameters inParameters, int3 inGroupID, int inGroupThreadID, out float outResult, out int2 outResultXY);
+void WriteScreenSpaceShadow(struct DispatchParameters inParameters, int3 inGroupID, int inGroupThreadID, out float outResult, out int2 outResultXY, out float2 outHitUV);
 
 
 
@@ -152,15 +155,15 @@ struct DispatchParameters
 										// If 'PointBorderSampler' is an Unnormalized sampler, then this value can be hard-coded to 1.
 										// The 'USE_HALF_PIXEL_OFFSET' macro might need to be defined if sampling at exact pixel coordinates isn't precise (e.g., if odd patterns appear in the shadow).
 
-	Texture2D<float> DepthTexture;		// Depth Buffer Texture (rasterized non-linear depth)
+	Texture2D DepthTexture;				// Depth Buffer Texture (rasterized non-linear depth)
 
 	SamplerState PointBorderSampler;	// A point sampler, with Wrap Mode set to Clamp-To-Border-Color (D3D12_TEXTURE_ADDRESS_MODE_BORDER), and Border Color set to "FarDepthValue" (typically zero), or some other far-depth value out of DepthBounds.
 										// If you have issues where invalid shadows are appearing from off-screen, it is likely that this sampler is not correctly setup
 };
 
 
-#if !defined(WAVE_SIZE) || !defined(SAMPLE_COUNT) || !defined(HARD_SHADOW_SAMPLES) || !defined(FADE_OUT_SAMPLES)
-	#error Before including bend_sss_gpu.h, four macros must be defined to configure the shader compile: WAVE_SIZE, SAMPLE_COUNT, HARD_SHADOW_SAMPLES, and FADE_OUT_SAMPLES. See the top of this file for details.
+#if !defined(WAVE_SIZE) || !defined(SAMPLE_COUNT) || !defined(HARD_SHADOW_SAMPLES) || !defined(FADE_OUT_SAMPLES) || !defined(TRACK_HIT_UV)
+	#error Before including bend_sss_gpu.h, four macros must be defined to configure the shader compile: WAVE_SIZE, SAMPLE_COUNT, HARD_SHADOW_SAMPLES, FADE_OUT_SAMPLES and TRACK_HIT_UV. See the top of this file for details.
 #else
 
 	static bool EarlyOutPixel(struct DispatchParameters inParameters, int2 pixel_xy, float depth)
@@ -239,12 +242,19 @@ struct DispatchParameters
 	groupshared float DepthData[READ_COUNT * WAVE_SIZE];
 	groupshared bool LdsEarlyOut;
 
+	#define DEBUG_HIT_UV (0)
+	#define PACKED_HIT_INDEX (1 && !DEBUG_HIT_UV)
+
+#if DEBUG_HIT_UV
+	groupshared float2 UVData[READ_COUNT * WAVE_SIZE];
+#endif
+
 	// Generate the shadow
 	//	Call this function from a compute shader with thread dimensions: numthreads[WAVE_SIZE, 1, 1]
 	// 
 	//	(int3)	inGroupID:			Compute shader group id register (SV_GroupID)
 	//	(int)	inGroupThreadId:	Compute shader group thread id register (SV_GroupThreadID)
-	void WriteScreenSpaceShadow(DispatchParameters inParameters, int3 inGroupID, int inGroupThreadID, out float outResult, out int2 outResultXY)
+	void WriteScreenSpaceShadow(DispatchParameters inParameters, int3 inGroupID, int inGroupThreadID, out float outResult, out int2 outResultXY, out float2 outHitUV)
 	{
 		float2 xy_delta;
 		float2 pixel_xy;
@@ -258,9 +268,14 @@ struct DispatchParameters
 		float shadowing_depth[READ_COUNT];
 		float depth_thickness_scale[READ_COUNT];
 		float sample_distance[READ_COUNT];
+	#if DEBUG_HIT_UV
+		float2 sample_uv[READ_COUNT];
+	#endif
 
 		const float direction = -inParameters.LightCoordinate.w;
 		const float z_sign = inParameters.NearDepthValue > inParameters.FarDepthValue ? -1 : +1;
+		
+		const float2 start_xy = pixel_xy;
 
 		int i;
 		bool is_edge = false;
@@ -290,14 +305,14 @@ struct DispatchParameters
 
 			// HLSL enforces that a pixel offset is a compile-time constant, which isn't strictly required (and can sometimes be a bit faster)
 			// So this fallback will use a manual uv offset instead
-			depths.x = inParameters.DepthTexture.SampleLevel(inParameters.PointBorderSampler, read_xy * inParameters.InvDepthTextureSize, 0);
-			depths.y = inParameters.DepthTexture.SampleLevel(inParameters.PointBorderSampler, (read_xy + offset_xy) * inParameters.InvDepthTextureSize, 0);
+			depths.x = inParameters.DepthTexture.SampleLevel(inParameters.PointBorderSampler, read_xy * inParameters.InvDepthTextureSize, 0).r;
+			depths.y = inParameters.DepthTexture.SampleLevel(inParameters.PointBorderSampler, (read_xy + offset_xy) * inParameters.InvDepthTextureSize, 0).r;
 		#else
 			int bias = bilinear > 0 ? 1 : -1;
 			int2 offset_xy = int2(x_axis_major ? 0 : bias, x_axis_major ? bias : 0);
 
-			depths.x = inParameters.DepthTexture.SampleLevel(inParameters.PointBorderSampler, read_xy * inParameters.InvDepthTextureSize, 0);
-			depths.y = inParameters.DepthTexture.SampleLevel(inParameters.PointBorderSampler, read_xy * inParameters.InvDepthTextureSize, 0, offset_xy);
+			depths.x = inParameters.DepthTexture.SampleLevel(inParameters.PointBorderSampler, read_xy * inParameters.InvDepthTextureSize, 0).r;
+			depths.y = inParameters.DepthTexture.SampleLevel(inParameters.PointBorderSampler, read_xy * inParameters.InvDepthTextureSize, 0, offset_xy).r;
 		#endif
 
 			// Depth thresholds (bilinear/shadow thickness) are based on a fractional ratio of the difference between sampled depth and the far clip depth
@@ -332,6 +347,9 @@ struct DispatchParameters
 
 			// Store for later
 			sample_distance[i] = pixel_distance + (WAVE_SIZE * i) * direction;
+		#if DEBUG_HIT_UV
+			sample_uv[i] = read_xy * inParameters.InvDepthTextureSize;
+		#endif
 
 			// Iterate to the next pixel along the ray. This will be WAVE_SIZE pixels along the ray...
 			pixel_xy += xy_delta * direction;
@@ -389,6 +407,9 @@ struct DispatchParameters
 			// Store the depth values in groupshared
 			int idx = (i * WAVE_SIZE) + inGroupThreadID.x;
 			DepthData[idx] = stored_depth;
+		#if DEBUG_HIT_UV
+			UVData[idx] = sample_uv[i];
+		#endif
 		}
 
 		// Sync wavefronts now groupshared DepthData is written
@@ -410,8 +431,13 @@ struct DispatchParameters
 		// Start by reading the next value
 		int sample_index = inGroupThreadID.x + 1;
 
+#if TRACK_HIT_UV && PACKED_HIT_INDEX
+		uint4 shadow_value_packed = asuint(1.0f);
+		uint hard_shadow_packed = asuint(1.0f);
+#else
 		float4 shadow_value = 1;
 		float hard_shadow = 1;
+#endif
 
 		// This is the inverse of how large the shadowing window is for the projected sample data. 
 		// All values in the LDS sample list are scaled by 1.0 / sample_distance, such that all light directions become parallel.
@@ -422,14 +448,38 @@ struct DispatchParameters
 		float depth_scale = min(sample_distance[0] + direction, 1.0 / inParameters.SurfaceThickness) * sample_distance[0] / depth_thickness_scale[0];
 
 		start_depth = start_depth * depth_scale - z_sign;
+		
+	#if TRACK_HIT_UV && !PACKED_HIT_INDEX
+		float hit_shadow = 1.0f;
+		int hit_index = 0;
+		#if DEBUG_HIT_UV
+			float2 hit_uv = 0.0f;
+		#endif
+	#endif
 
 		// The first number of hard shadow samples, a single pixel can produce a full shadow
 		[unroll] for (i = 0; i < HARD_SHADOW_SAMPLES; i++)
 		{
 			float depth_delta = abs(start_depth - DepthData[sample_index + i] * depth_scale);
 
+		#if TRACK_HIT_UV && PACKED_HIT_INDEX
+			uint depth_delta_packed = asuint(depth_delta) & 0xffffff80 | (i + 1);
+			hard_shadow_packed = min(hard_shadow_packed, depth_delta_packed);
+		#else
 			// We want to find the distance of the sample that is closest to the reference depth
 			hard_shadow = min(hard_shadow, depth_delta);
+		#endif
+
+		#if TRACK_HIT_UV && !PACKED_HIT_INDEX
+			if (depth_delta < hit_shadow)
+			{
+				hit_shadow = depth_delta;
+				hit_index = i + 1;
+			#if DEBUG_HIT_UV
+				hit_uv = UVData[sample_index + i];
+			#endif
+			}
+		#endif
 		}
 
 		// Brute force go!
@@ -438,22 +488,66 @@ struct DispatchParameters
 		{
 			float depth_delta = abs(start_depth - DepthData[sample_index + i] * depth_scale);
 
+		#if TRACK_HIT_UV && PACKED_HIT_INDEX
+			uint depth_delta_packed = asuint(depth_delta) & 0xffffff80 | (i + 1);
+			shadow_value_packed[i & 3] = min(shadow_value_packed[i & 3], depth_delta_packed);
+		#else
 			// Do the same as the hard_shadow code above, but this will accumulate to 4 separate values.
 			// By using 4 values, the average shadow can be taken, which can help soften single-pixel shadows.
 			shadow_value[i & 3] = min(shadow_value[i & 3], depth_delta);
+		#endif
+
+		#if TRACK_HIT_UV && !PACKED_HIT_INDEX
+			if (depth_delta < hit_shadow)
+			{
+				hit_shadow = depth_delta;
+				hit_index = i + 1;
+			#if DEBUG_HIT_UV
+				hit_uv = UVData[sample_index + i];
+			#endif
+			}
+		#endif
 		}
 
 		// Final fade out samples
 		[unroll] for (i = SAMPLE_COUNT - FADE_OUT_SAMPLES; i < SAMPLE_COUNT; i++)
 		{
-			float depth_delta = abs(start_depth - DepthData[sample_index + i] * depth_scale);
-
 			// Add the fade value to these samples
 			const float fade_out = (float)(i + 1 - (SAMPLE_COUNT - FADE_OUT_SAMPLES)) / (float)(FADE_OUT_SAMPLES + 1) * 0.75;
 
-			shadow_value[i & 3] = min(shadow_value[i & 3], depth_delta + fade_out);
+			float depth_delta = abs(start_depth - DepthData[sample_index + i] * depth_scale) + fade_out;
+
+		#if TRACK_HIT_UV && PACKED_HIT_INDEX
+			uint depth_delta_packed = asuint(depth_delta) & 0xffffff80 | (i + 1);
+			shadow_value_packed[i & 3] = min(shadow_value_packed[i & 3], depth_delta_packed);
+		#else
+			shadow_value[i & 3] = min(shadow_value[i & 3], depth_delta);
+		#endif
+			
+		#if TRACK_HIT_UV && !PACKED_HIT_INDEX
+			if (depth_delta < hit_shadow)
+			{
+				hit_shadow = depth_delta;
+				hit_index = i + 1;
+			#if DEBUG_HIT_UV
+				hit_uv = UVData[sample_index + i];
+			#endif
+			}
+		#endif
 		}
 
+	#if TRACK_HIT_UV && PACKED_HIT_INDEX
+		int hit_index = hard_shadow_packed;
+		hit_index = min(hit_index, shadow_value_packed.x);
+		hit_index = min(hit_index, shadow_value_packed.y);
+		hit_index = min(hit_index, shadow_value_packed.z);
+		hit_index = min(hit_index, shadow_value_packed.w);
+		hit_index &= 0x7f;
+
+		float4 shadow_value = asfloat(shadow_value_packed);
+		float hard_shadow = asfloat(hard_shadow_packed);
+	#endif
+
 		// Apply the contrast value.
 		// A value of 0 indicates a sample was exactly matched to the reference depth (and the result is fully shadowed)
 		// We want some boost to this range, so samples don't have to exactly match to produce a full shadow. 
@@ -479,6 +573,20 @@ struct DispatchParameters
 
 			outResult = result;
 			outResultXY = (int2)write_xy;
+
+		#if TRACK_HIT_UV
+			float2 hit_xy = floor(start_xy + hit_index * (xy_delta * direction) / WAVE_SIZE);
+			#if USE_HALF_PIXEL_OFFSET
+				hit_xy += 0.5;
+			#endif
+				outHitUV = hit_xy * inParameters.InvDepthTextureSize;
+
+			#if DEBUG_HIT_UV
+				outHitUV = hit_uv;
+			#endif
+		#else
+			outHitUV = 0.0f;
+		#endif
 		}
 	}
 
diff --git a/Engine/Source/Runtime/Renderer/Private/Shadows/ScreenSpaceShadows.cpp b/Engine/Source/Runtime/Renderer/Private/Shadows/ScreenSpaceShadows.cpp
--- a/Engine/Source/Runtime/Renderer/Private/Shadows/ScreenSpaceShadows.cpp
+++ b/Engine/Source/Runtime/Renderer/Private/Shadows/ScreenSpaceShadows.cpp
@@ -38,9 +38,47 @@ static FAutoConsoleVariableRef CVarBendShadowsOverrideSurfaceThickness(
 	TEXT("r.ContactShadows.Bend.OverrideSurfaceThickness"),
 	GBendShadowsOverrideSurfaceThickness,
 	TEXT("How thick the surface represented by a pixel is assumed to be when determining whether a ray intersects it."),
+	ECVF_RenderThreadSafe
+);
+
+enum class EContactShadowsIntensityMode
+{
+	PrimitiveFlag,
+	DepthBasedApproximation,
+	ForceCastingIntensity,
+
+	MAX
+};
+
+static int32 GContactShadowsIntensityMode = 0;
+static FAutoConsoleVariableRef CVarContactShadowsIntensityMode(
+	TEXT("r.ContactShadows.Intensity.Mode"),
+	GContactShadowsIntensityMode,
+	TEXT("Control how contact shadow intensity is calculated:\n")
+	TEXT("0 - Respect bCastContactShadow flag on Primitive Component.\n")
+	TEXT("1 - Depth based approximation.\n")
+	TEXT("2 - Use Casting Intensity."),
 	ECVF_Scalability | ECVF_RenderThreadSafe
 );
 
+static float GContactShadowsIntensityFadeStart = 1600; // 16m
+static FAutoConsoleVariableRef CVarContactShadowsIntensityFadeStart(
+	TEXT("r.ContactShadows.Intensity.FadeStart"),
+	GContactShadowsIntensityFadeStart,
+	TEXT("Depth value at which contact shadows starts fading from NonCastingIntensity to CastingIntensity.\n")
+	TEXT("Only used when r.ContactShadows.Intensity.Mode=1"),
+	ECVF_RenderThreadSafe
+);
+
+static float GContactShadowsIntensityFadeLength = 800; // 8m
+static FAutoConsoleVariableRef CVarContactShadowsIntensityFadeLength(
+	TEXT("r.ContactShadows.Intensity.FadeLength"),
+	GContactShadowsIntensityFadeLength,
+	TEXT("Length of the fading interval from NonCastingIntensity to CastingIntensity.\n")
+	TEXT("Only used when r.ContactShadows.Intensity.Mode=1"),
+	ECVF_RenderThreadSafe
+);
+
 extern void GetLightContactShadowParameters(const FLightSceneProxy* Proxy, float& OutLength, bool& bOutLengthInWS, float& OutCastingIntensity, float& OutNonCastingIntensity);
 
 const int32 GScreenSpaceShadowsTileSizeX = 8;
@@ -98,7 +136,7 @@ class FScreenSpaceShadowsBendCS : public FGlobalShader
 
 	BEGIN_SHADER_PARAMETER_STRUCT(FParameters, )
 		SHADER_PARAMETER_RDG_TEXTURE_UAV(RWTexture2D<float>, OutputTexture)
-		SHADER_PARAMETER_RDG_TEXTURE(Texture2D<float>, DepthTexture)
+		SHADER_PARAMETER_STRUCT_INCLUDE(FSceneTextureShaderParameters, SceneTextures)
 		SHADER_PARAMETER_RDG_TEXTURE_SRV(Texture2D<uint2>, StencilTexture)
 		SHADER_PARAMETER_SAMPLER(SamplerState, PointBorderSampler)
 		SHADER_PARAMETER_STRUCT_REF(FViewUniformShaderParameters, View)
@@ -107,6 +145,8 @@ class FScreenSpaceShadowsBendCS : public FGlobalShader
 		SHADER_PARAMETER(uint32, bContactShadowLengthInWS)
 		SHADER_PARAMETER(float, ContactShadowCastingIntensity)
 		SHADER_PARAMETER(float, ContactShadowNonCastingIntensity)
+		SHADER_PARAMETER(float, ContactShadowIntensityFadeStart)
+		SHADER_PARAMETER(float, ContactShadowIntensityFadeOneOverLength)
 		SHADER_PARAMETER(float, SurfaceThickness)
 		SHADER_PARAMETER(FIntRect, ScissorRectMinAndSize)
 		SHADER_PARAMETER(uint32, DownsampleFactor)
@@ -115,6 +155,9 @@ class FScreenSpaceShadowsBendCS : public FGlobalShader
 		SHADER_PARAMETER(FIntVector, WaveOffset)
 	END_SHADER_PARAMETER_STRUCT()
 
+	class FIntensityModeDim : SHADER_PERMUTATION_ENUM_CLASS("DIM_INTENSITY_MODE", EContactShadowsIntensityMode);
+	using FPermutationDomain = TShaderPermutationDomain<FIntensityModeDim>;
+
 	static bool ShouldCompilePermutation(const FGlobalShaderPermutationParameters& Parameters)
 	{
 		return true;
@@ -344,7 +387,7 @@ void RenderScreenSpaceShadowsBend(
 
 			PassParameters->OutputTexture = GraphBuilder.CreateUAV(ShadowsTexture);
 
-			PassParameters->DepthTexture = SceneTextures.Depth.Resolve;
+			PassParameters->SceneTextures = SceneTextures.GetSceneTextureShaderParameters(View.GetFeatureLevel());
 			PassParameters->InvDepthTextureSize = FVector2f(1.0f / DepthDesc.Extent.X, 1.0f / DepthDesc.Extent.Y);
 
 			PassParameters->StencilTexture = SceneTextures.Stencil;
@@ -373,12 +416,17 @@ void RenderScreenSpaceShadowsBend(
 			PassParameters->bContactShadowLengthInWS = bContactShadowLengthInWS;
 			PassParameters->ContactShadowCastingIntensity = ContactShadowCastingIntensity;
 			PassParameters->ContactShadowNonCastingIntensity = ContactShadowNonCastingIntensity;
+			PassParameters->ContactShadowIntensityFadeStart = GContactShadowsIntensityFadeStart;
+			PassParameters->ContactShadowIntensityFadeOneOverLength = 1.0f / GContactShadowsIntensityFadeStart;
 			PassParameters->SurfaceThickness = GBendShadowsOverrideSurfaceThickness;
 
 			PassParameters->LightCoordinate = FVector4f(DispatchList.LightCoordinate_Shader[0], DispatchList.LightCoordinate_Shader[1], DispatchList.LightCoordinate_Shader[2], DispatchList.LightCoordinate_Shader[3]);
 			PassParameters->WaveOffset = FIntVector(Dispatch.WaveOffset_Shader[0], Dispatch.WaveOffset_Shader[1], 0);
 
-			auto ComputeShader = View.ShaderMap->GetShader<FScreenSpaceShadowsBendCS>();
+			FScreenSpaceShadowsBendCS::FPermutationDomain PermutationVector;
+			PermutationVector.Set<FScreenSpaceShadowsBendCS::FIntensityModeDim>((EContactShadowsIntensityMode)GContactShadowsIntensityMode);
+
+			auto ComputeShader = View.ShaderMap->GetShader<FScreenSpaceShadowsBendCS>(PermutationVector);
 
 			FComputeShaderUtils::AddPass(
 				GraphBuilder,
