Fix GBuffer picking shader compilation issue on platform not supporting deferred path.

diff --git a/Engine/Shaders/Private/PostProcessGBufferHints.usf b/Engine/Shaders/Private/PostProcessGBufferHints.usf
--- a/Engine/Shaders/Private/PostProcessGBufferHints.usf
+++ b/Engine/Shaders/Private/PostProcessGBufferHints.usf
@@ -1,8 +1,11 @@
 // Copyright Epic Games, Inc. All Rights Reserved.
 
+#include "Common.ush"
 #include "DeferredShadingCommon.ush"
 #include "ScreenPass.ush"
 
+#define SUPPORT_GBUFFER_FETCH (SHADING_PATH_DEFERRED && FEATURE_LEVEL >= FEATURE_LEVEL_SM5)
+
 float2 Input_ExtentInverse;
 
 Texture2D SceneColorTexture;
@@ -11,6 +14,8 @@ SamplerState SceneColorSampler;
 Texture2D OriginalSceneColorTexture;
 SamplerState OriginalSceneColorSampler;
 
+#if SUPPORT_GBUFFER_FETCH
+
 // 0: good, 1: heavy performance cost, 2:warning (very uncommon material), 3:error (impossible material attributes)
 uint ComputeAdvice(FGBufferData BufferData, float3 OriginalSceneColor)
 {
@@ -92,6 +97,7 @@ void HighlightAdvice(inout float3 OutColor, float2 UV, uint2 PixelPos)
 		 }
 	}
 }
+#endif
 
 void MainPS(noperspective float4 UVAndScreenPos : TEXCOORD0, float4 SvPosition : SV_POSITION, out float4 OutColor : SV_Target0)
 {
@@ -101,7 +107,9 @@ void MainPS(noperspective float4 UVAndScreenPos : TEXCOORD0, float4 SvPosition :
 	// background is the scene color
 	OutColor = Texture2DSample(SceneColorTexture, SceneColorSampler, UV);
 
+#if SUPPORT_GBUFFER_FETCH
 	HighlightAdvice(OutColor.rgb, UV, PixelPos);
+#endif
 }
 
 
@@ -151,7 +159,7 @@ void PrintBackgroundRect(inout FShaderPrintContext Ctx, float2 StartRect, float2
 [numthreads(1, 1, 1)]
 void MainCS(uint3 DispatchThreadId : SV_DispatchThreadID)
 {
-#if !SUBSTRATE_ENABLED && SHADING_PATH_DEFERRED && FEATURE_LEVEL >= FEATURE_LEVEL_SM5
+#if !SUBSTRATE_ENABLED && SUPPORT_GBUFFER_FETCH
 	// Draw stats
 	if (all(DispatchThreadId == 0))
 	{
diff --git a/Engine/Source/Runtime/Renderer/Private/PostProcess/PostProcessing.cpp b/Engine/Source/Runtime/Renderer/Private/PostProcess/PostProcessing.cpp
--- a/Engine/Source/Runtime/Renderer/Private/PostProcess/PostProcessing.cpp
+++ b/Engine/Source/Runtime/Renderer/Private/PostProcess/PostProcessing.cpp
@@ -2925,7 +2925,11 @@ class FGBufferPickingCS : public FGlobalShader
 public:
 	static bool ShouldCompilePermutation(const FGlobalShaderPermutationParameters& Parameters) 
 	{ 
-		return ShaderPrint::IsSupported(Parameters.Platform) && !IsMobilePlatform(Parameters.Platform) && !Substrate::IsSubstrateEnabled() &&
+		return 
+			IsFeatureLevelSupported(Parameters.Platform, ERHIFeatureLevel::SM5) &&
+			ShaderPrint::IsSupported(Parameters.Platform) && 
+			!IsMobilePlatform(Parameters.Platform) && 
+			!Substrate::IsSubstrateEnabled() &&
 			EnumHasAllFlags(Parameters.Flags, EShaderPermutationFlags::HasEditorOnlyData);
 	}
 	static void ModifyCompilationEnvironment(const FGlobalShaderPermutationParameters& Parameters, FShaderCompilerEnvironment& OutEnvironment)
