Fix error compiling FScreenSpaceShadowsBendCS on some platforms.

diff --git a/Engine/Shaders/Private/bend_sss_gpu.ush b/Engine/Shaders/Private/bend_sss_gpu.ush
--- a/Engine/Shaders/Private/bend_sss_gpu.ush
+++ b/Engine/Shaders/Private/bend_sss_gpu.ush
@@ -537,7 +537,7 @@ struct DispatchParameters
 		}
 
 	#if TRACK_HIT_UV && PACKED_HIT_INDEX
-		int hit_index = hard_shadow_packed;
+		uint hit_index = hard_shadow_packed;
 		hit_index = min(hit_index, shadow_value_packed.x);
 		hit_index = min(hit_index, shadow_value_packed.y);
 		hit_index = min(hit_index, shadow_value_packed.z);
