Explicitely disable rect light in cluster shading to optimized out rect light code.

Cluster shading does not support rectlight for now in any case (they are not injected into the sorted lights). Later on we could build a permutation depending if there is rect light or not injected into the structure.

diff --git a/Engine/Shaders/Private/ClusteredDeferredShadingPixelShader.usf b/Engine/Shaders/Private/ClusteredDeferredShadingPixelShader.usf
--- a/Engine/Shaders/Private/ClusteredDeferredShadingPixelShader.usf
+++ b/Engine/Shaders/Private/ClusteredDeferredShadingPixelShader.usf
@@ -157,6 +157,7 @@ float4 GetSimpleLightLighting(
 	#if SUBSTRATE_ENABLED && USE_HAIR_LIGHTING == 0
 		FDeferredLightData LightData = ConvertToDeferredLight(LocalLight);
 		LightData.ShadowedBits = 1;
+		LightData.bRectLight = false;
 
 		FSubstrateShadowTermInputParameters SubstrateShadowTermInputParameters	= GetInitialisedSubstrateShadowTermInputParameters();
 		SubstrateShadowTermInputParameters.bEvaluateShadowTerm					= false;
@@ -260,6 +261,7 @@ float4 GetLightGridLocalLighting(
 
 		FDeferredLightData LightData = ConvertToDeferredLight(LocalLight);
 		LightData.ShadowedBits = 1;
+		LightData.bRectLight = false;
 		float4 LightAttenuation = float4(1, 1, 1, 1);
 
 		// Find packed VSM projection result if present
