Stochastic Shadows
* Fixed crash when with indirect args
* Replaced atomics with groupshared reduction in tile classifcation (0.1ms->0.8ms)

diff --git a/Engine/Shaders/Private/StochasticShadows/StochasticShadows.usf b/Engine/Shaders/Private/StochasticShadows/StochasticShadows.usf
--- a/Engine/Shaders/Private/StochasticShadows/StochasticShadows.usf
+++ b/Engine/Shaders/Private/StochasticShadows/StochasticShadows.usf
@@ -5,7 +5,7 @@
 #include "../ShaderPrint.ush"
 #include "StochasticShadows.ush"
 
-groupshared uint SharedTileMode;
+groupshared uint SharedTileMode[THREADGROUP_SIZE * THREADGROUP_SIZE];
 
 RWStructuredBuffer<uint> RWTileAllocator;
 RWStructuredBuffer<uint> RWTileData;
@@ -22,10 +22,9 @@ void TileClassificationCS(
 	uint3 GroupThreadId : SV_GroupThreadID,
 	uint3 DispatchThreadId : SV_DispatchThreadID)
 {
-	if (all(GroupThreadId == 0))
-	{
-		SharedTileMode = 0;
-	}
+	uint LinearThreadIndex = GroupThreadId.x + THREADGROUP_SIZE * GroupThreadId.y;
+
+	SharedTileMode[LinearThreadIndex] = 0;
 	
 	GroupMemoryBarrierWithGroupSync();
 
@@ -50,27 +49,68 @@ void TileClassificationCS(
 		}
 		else if (GBuffer.ShadingModelID == SHADINGMODELID_DEFAULT_LIT)
 		{
-			InterlockedOr(SharedTileMode, 0x1);
+			SharedTileMode[LinearThreadIndex] = 0x1;
 		}
 		else
 		{
-			InterlockedOr(SharedTileMode, 0x2);
+			SharedTileMode[LinearThreadIndex] = 0x2;
 		}
 	}
 
 	GroupMemoryBarrierWithGroupSync();
 
-	if (all(GroupThreadId == 0))
+	// GroupShared reduction
+	if (LinearThreadIndex < 32)
+	{
+		SharedTileMode[LinearThreadIndex] = SharedTileMode[LinearThreadIndex] || SharedTileMode[LinearThreadIndex + 32];
+	}
+
+	GroupMemoryBarrierWithGroupSync();
+
+	if (LinearThreadIndex < 16)
+	{
+		SharedTileMode[LinearThreadIndex] = SharedTileMode[LinearThreadIndex] || SharedTileMode[LinearThreadIndex + 16];
+	}
+
+	GroupMemoryBarrierWithGroupSync();
+
+	if (LinearThreadIndex < 8)
+	{
+		SharedTileMode[LinearThreadIndex] = SharedTileMode[LinearThreadIndex] || SharedTileMode[LinearThreadIndex + 8];
+	}
+
+	GroupMemoryBarrierWithGroupSync();
+
+	if (LinearThreadIndex < 4)
+	{
+		SharedTileMode[LinearThreadIndex] = SharedTileMode[LinearThreadIndex] || SharedTileMode[LinearThreadIndex + 4];
+	}
+
+	GroupMemoryBarrierWithGroupSync();
+
+	if (LinearThreadIndex < 2)
+	{
+		SharedTileMode[LinearThreadIndex] = SharedTileMode[LinearThreadIndex] || SharedTileMode[LinearThreadIndex + 2];
+	}
+
+	GroupMemoryBarrierWithGroupSync();
+
+	if (LinearThreadIndex < 1)
+	{
+		SharedTileMode[LinearThreadIndex] = SharedTileMode[LinearThreadIndex] || SharedTileMode[LinearThreadIndex + 1];
+	}
+
+	if (LinearThreadIndex == 0)
 	{
 		const uint DataStride = DOWNSAMPLED_CLASSIFICATION ? DownsampledTileDataStride : TileDataStride;
 
-		if (SharedTileMode & 0x2)
+		if (SharedTileMode[0] & 0x2)
 		{
 			uint StochasticShadowsTileIndex = 0;
 			InterlockedAdd(RWTileAllocator[TILE_MODE_COMPLEX_SHADING], 1, StochasticShadowsTileIndex);
 			RWTileData[DataStride * TILE_MODE_COMPLEX_SHADING + StochasticShadowsTileIndex] = PackTile(ScreenCoord / 8);
 		}
-		else if (SharedTileMode & 0x1)
+		else if (SharedTileMode[0] & 0x1)
 		{
 			uint StochasticShadowsTileIndex = 0;
 			InterlockedAdd(RWTileAllocator[TILE_MODE_SIMPLE_SHADING], 1, StochasticShadowsTileIndex);
diff --git a/Engine/Shaders/Private/StochasticShadows/StochasticShadows.ush b/Engine/Shaders/Private/StochasticShadows/StochasticShadows.ush
--- a/Engine/Shaders/Private/StochasticShadows/StochasticShadows.ush
+++ b/Engine/Shaders/Private/StochasticShadows/StochasticShadows.ush
@@ -2,6 +2,9 @@
 
 #pragma once
 
+// Change this to force recompilation of all Stochastic Shadows shaders
+#pragma message("UESHADERMETADATA_VERSION F06AC34F-0C35-4FDF-9EA6-82395F79FB2B")
+
 #define SUPPORT_CONTACT_SHADOWS 0
 #define USE_IES_PROFILE 0
 #define NON_DIRECTIONAL_DIRECT_LIGHTING 0
diff --git a/Engine/Source/Runtime/Renderer/Private/StochasticShadows/StochasticShadows.cpp b/Engine/Source/Runtime/Renderer/Private/StochasticShadows/StochasticShadows.cpp
--- a/Engine/Source/Runtime/Renderer/Private/StochasticShadows/StochasticShadows.cpp
+++ b/Engine/Source/Runtime/Renderer/Private/StochasticShadows/StochasticShadows.cpp
@@ -466,6 +466,7 @@ class FGenerateSamplesCS : public FGlobalShader
 	SHADER_USE_PARAMETER_STRUCT(FGenerateSamplesCS, FGlobalShader)
 
 	BEGIN_SHADER_PARAMETER_STRUCT(FParameters, )
+		RDG_BUFFER_ACCESS(IndirectArgs, ERHIAccess::IndirectArgs)
 		SHADER_PARAMETER_STRUCT_INCLUDE(FStochasticShadowsParameters, StochasticShadowsParameters)
 		SHADER_PARAMETER_RDG_TEXTURE_UAV(RWTexture2D<float>, RWDownsampledSceneDepth)
 		SHADER_PARAMETER_RDG_TEXTURE_UAV(RWTexture2D<uint32>, RWShadowMaskHistoryScreenCoord00)
@@ -857,6 +858,7 @@ class FShadeLightSamplesCS : public FGlobalShader
 	SHADER_USE_PARAMETER_STRUCT(FShadeLightSamplesCS, FGlobalShader)
 
 	BEGIN_SHADER_PARAMETER_STRUCT(FParameters, )
+		RDG_BUFFER_ACCESS(IndirectArgs, ERHIAccess::IndirectArgs)
 		SHADER_PARAMETER_STRUCT_INCLUDE(FStochasticShadowsParameters, StochasticShadowsParameters)
 		SHADER_PARAMETER_RDG_TEXTURE_UAV(RWTexture2D<float4>, RWSceneColor)
 		SHADER_PARAMETER_RDG_TEXTURE(Texture2D<uint>, ShadowMaskPageTable)
@@ -1271,6 +1273,7 @@ void FDeferredShadingSceneRenderer::RenderStochasticShadows(FRDGBuilder& GraphBu
 		for (int32 TileType = 0; TileType < (int32)StochasticShadows::ETileType::MAX; ++TileType)
 		{
 			FGenerateSamplesCS::FParameters* PassParameters = GraphBuilder.AllocParameters<FGenerateSamplesCS::FParameters>();
+			PassParameters->IndirectArgs = DownsampledTileIndirectArgs;
 			PassParameters->StochasticShadowsParameters = StochasticShadowsParameters;
 			PassParameters->RWDownsampledSceneDepth = DownsampledSceneDepthUAV;
 			PassParameters->RWShadowMaskHistoryScreenCoord00 = ShadowMaskHistoryScreenCoord00UAV;
@@ -1590,6 +1593,7 @@ void FDeferredShadingSceneRenderer::RenderStochasticShadows(FRDGBuilder& GraphBu
 		for (int32 TileType = 0; TileType < (int32)StochasticShadows::ETileType::MAX; ++TileType)
 		{
 			FShadeLightSamplesCS::FParameters* PassParameters = GraphBuilder.AllocParameters<FShadeLightSamplesCS::FParameters>();
+			PassParameters->IndirectArgs = TileIndirectArgs;
 			PassParameters->StochasticShadowsParameters = StochasticShadowsParameters;
 			PassParameters->ShadowMaskPageTable = ShadowMaskPageTable;
 			PassParameters->ShadowMaskTileAtlas = ShadowMaskTileAtlas;
