Fix a bug where spline meshes could potentially sample outside their LUT boundaries.

diff --git a/Engine/Shaders/Private/SplineMeshCommon.ush b/Engine/Shaders/Private/SplineMeshCommon.ush
--- a/Engine/Shaders/Private/SplineMeshCommon.ush
+++ b/Engine/Shaders/Private/SplineMeshCommon.ush
@@ -113,7 +113,7 @@ FSplineMeshSlice SplineMeshCalcSlice(FSplineMeshShaderParams Params, half Spline
 {
 	// Sample the slice position and rotation from scene textures
 	const half SplineLastTexel = (SPLINE_MESH_TEXEL_WIDTH - 1);
-	const half SplineDistTexels = min(SplineDist * Params.SplineDistToTexelScale + Params.SplineDistToTexelOffset, SplineLastTexel);
+	const half SplineDistTexels = clamp(SplineDist * Params.SplineDistToTexelScale + Params.SplineDistToTexelOffset, 0, SplineLastTexel);
 	const float2 TexelSizeUV = Scene.SplineMesh.SplineTextureInvExtent;
 	const float2 UVOffset = (Params.TextureCoord + (float2)0.5f) * TexelSizeUV;
 	const float2 UV = UVOffset + float2(SplineDistTexels * TexelSizeUV.x, 0);
