Fix for StochasicShadow with dead code removal

diff --git a/Engine/Shaders/Private/StochasticShadows/StochasticShadows.usf b/Engine/Shaders/Private/StochasticShadows/StochasticShadows.usf
--- a/Engine/Shaders/Private/StochasticShadows/StochasticShadows.usf
+++ b/Engine/Shaders/Private/StochasticShadows/StochasticShadows.usf
@@ -102,7 +102,11 @@ void TileClassificationCS(
 
 	if (LinearThreadIndex == 0)
 	{
-		const uint DataStride = DOWNSAMPLED_CLASSIFICATION ? DownsampledTileDataStride : TileDataStride;
+#if DOWNSAMPLED_CLASSIFICATION
+		const uint DataStride = DownsampledTileDataStride;
+#else
+		const uint DataStride = TileDataStride;
+#endif
 
 		if (SharedTileMode[0] & 0x2)
 		{
