Lumen - fix for translucent casting opaque shadows in  hit lighting reflections.
We will refine this later with proper translucency information in order to mesh composed of opaque and translucent parts.

diff --git a/Engine/Shaders/Private/RayTracing/RayTracingLightingCommon.ush b/Engine/Shaders/Private/RayTracing/RayTracingLightingCommon.ush
--- a/Engine/Shaders/Private/RayTracing/RayTracingLightingCommon.ush
+++ b/Engine/Shaders/Private/RayTracing/RayTracingLightingCommon.ush
@@ -401,7 +401,9 @@ float3 ComputeDirectLighting(
 		// Always bias towards the light
 		ApplyPositionBias(ShadowRay, bBackFace ? -Payload.GetWorldNormal() : Payload.GetWorldNormal(), ShadowMaxNormalBias);
 
-		uint RayFlags = RAY_FLAG_ACCEPT_FIRST_HIT_AND_END_SEARCH | RAY_FLAG_SKIP_CLOSEST_HIT_SHADER;
+		uint RayFlags = RAY_FLAG_ACCEPT_FIRST_HIT_AND_END_SEARCH	// Terminate immediatly after detecting a hit
+						| RAY_FLAG_SKIP_CLOSEST_HIT_SHADER			// Skip Closest Hit shader, only run Any Hit shader
+						| RAY_FLAG_CULL_NON_OPAQUE;					// Do not include any translucent from shadow as of today since we do not handle coverage & transmittance.
 		const uint InstanceInclusionMask = (ApplyShadow) ? RAY_TRACING_MASK_SHADOW : 0;
 
 #if !ENABLE_TWO_SIDED_GEOMETRY
