Fix shader compilation with Subtrate.

diff --git a/Engine/Shaders/Private/RayTracing/RayTracingCommon.ush b/Engine/Shaders/Private/RayTracing/RayTracingCommon.ush
--- a/Engine/Shaders/Private/RayTracing/RayTracingCommon.ush
+++ b/Engine/Shaders/Private/RayTracing/RayTracingCommon.ush
@@ -610,7 +610,7 @@ struct FPackedMaterialClosestHitPayload : FMinimalPayload
 	void SetRayCone(FRayCone NewRayCone)	{ PackedRayCone = PackRayCone(NewRayCone); }
 
 	float GetMipBias()						{ return UnpackR8(Flags>>8u); }
-	void SetMipBias(float NewMipBias)		{ Flags = (Flags & 0xFFFF00FFF) | (PackR8(NewMipBias) << 8u); }
+	void SetMipBias(float NewMipBias)		{ Flags = (Flags & 0xFFFF00FFFu) | (PackR8(NewMipBias) << 8u); }
 
 	void SetRadiance(in float3 In)			{ Radiance = PackR11G11B10F(In); }
 	float3 GetRadiance()					{ return UnpackR11G11B10F(Radiance); }
diff --git a/Engine/Shaders/Private/ReflectionEnvironmentPixelShader.usf b/Engine/Shaders/Private/ReflectionEnvironmentPixelShader.usf
--- a/Engine/Shaders/Private/ReflectionEnvironmentPixelShader.usf
+++ b/Engine/Shaders/Private/ReflectionEnvironmentPixelShader.usf
@@ -39,12 +39,13 @@
 #include "/Engine/Private/Substrate/SubstrateLightingCommon.ush"
 #include "Lumen/LumenReflectionsCombine.ush"
 
-#if SUBSTRATE_ENABLED && LUMEN_STANDALONE_REFLECTIONS
-Texture2DArray ReflectionTexture;
+#if LUMEN_STANDALONE_REFLECTIONS
+Texture2DArray LumenReflectionTexture;
+SamplerState LumenReflectionTextureSampler;
 #else
 Texture2D ReflectionTexture;
-#endif
 SamplerState ReflectionTextureSampler;
+#endif
 
 #if ENABLE_DYNAMIC_SKY_LIGHT
 #include "VolumetricCloudCommon.ush"
@@ -144,7 +145,11 @@ float3 ReflectionEnvironment(FGBufferData GBuffer, float AmbientOcclusion, float
 	R = GetOffSpecularPeakReflectionDir(N, R, GBuffer.Roughness);
 
 	// Sample input reflection texture that may contain SSR, planar reflections, RT reflections or Lumen Reflections
+#if LUMEN_STANDALONE_REFLECTIONS
+	float4 ReflectionInput = Texture2DArraySample(LumenReflectionTexture, LumenReflectionTextureSampler, float3(BufferUV, 0));
+#else
 	float4 ReflectionInput = Texture2DSample(ReflectionTexture, ReflectionTextureSampler, BufferUV);
+#endif
 
 	Color = CompositeReflections(ReflectionInput, BufferUV, GBuffer.Roughness, ShadingModelID);
 
@@ -438,7 +443,7 @@ void ReflectionEnvironmentSkyLighting(
 	
 		#if LUMEN_STANDALONE_REFLECTIONS
 			{
-				const float4 LumenReflections = GetLumenReflectionOnly(SubstrateEnvLight, BSDF, BSDFThroughput, Texture2DArraySample(ReflectionTexture, ReflectionTextureSampler, float3(BufferUV, ClosureIndex)));
+				const float4 LumenReflections = GetLumenReflectionOnly(SubstrateEnvLight, BSDF, BSDFThroughput, Texture2DArraySample(LumenReflectionTexture, LumenReflectionTextureSampler, float3(BufferUV, ClosureIndex)));
 				SpecularLighting *= LumenReflections.a;
 				SpecularLighting += LumenReflections.xyz * BSDFThroughput * SubstrateEnvLight.SpecularWeight;
 			}
diff --git a/Engine/Source/Runtime/Renderer/Private/IndirectLightRendering.cpp b/Engine/Source/Runtime/Renderer/Private/IndirectLightRendering.cpp
--- a/Engine/Source/Runtime/Renderer/Private/IndirectLightRendering.cpp
+++ b/Engine/Source/Runtime/Renderer/Private/IndirectLightRendering.cpp
@@ -380,6 +380,9 @@ class FReflectionEnvironmentSkyLightingPS : public FGlobalShader
 		SHADER_PARAMETER_RDG_TEXTURE(Texture2D, ReflectionTexture)
 		SHADER_PARAMETER_SAMPLER(SamplerState, ReflectionTextureSampler)
 
+		SHADER_PARAMETER_RDG_TEXTURE(Texture2D, LumenReflectionTexture)
+		SHADER_PARAMETER_SAMPLER(SamplerState, LumenReflectionTextureSampler)
+
 		SHADER_PARAMETER_TEXTURE(Texture2D, PreIntegratedGF)
 		SHADER_PARAMETER_SAMPLER(SamplerState, PreIntegratedGFSampler)
 
@@ -1667,8 +1670,16 @@ static void AddSkyReflectionPass(
 		PassParameters->PS.AmbientOcclusionTexture = AmbientOcclusionTexture;
 		PassParameters->PS.AmbientOcclusionSampler = TStaticSamplerState<SF_Point>::GetRHI();
 
-		PassParameters->PS.ReflectionTexture = ReflectionsColor ? ReflectionsColor : SystemTextures.Black;
-		PassParameters->PS.ReflectionTextureSampler = TStaticSamplerState<SF_Point>::GetRHI();
+		if (bLumenStandaloneReflections)
+		{
+			PassParameters->PS.LumenReflectionTexture = ReflectionsColor ? ReflectionsColor : SystemTextures.BlackArray;
+			PassParameters->PS.LumenReflectionTextureSampler = TStaticSamplerState<SF_Point>::GetRHI();
+		}
+		else
+		{
+			PassParameters->PS.ReflectionTexture = ReflectionsColor ? ReflectionsColor : SystemTextures.Black;
+			PassParameters->PS.ReflectionTextureSampler = TStaticSamplerState<SF_Point>::GetRHI();
+		}
 
 		if (Scene->HasVolumetricCloud())
 		{
