Lumen - fix translucent shadow casting shadow for all cases

diff --git a/Engine/Shaders/Private/RayTracing/RayTracingLightingCommon.ush b/Engine/Shaders/Private/RayTracing/RayTracingLightingCommon.ush
--- a/Engine/Shaders/Private/RayTracing/RayTracingLightingCommon.ush
+++ b/Engine/Shaders/Private/RayTracing/RayTracingLightingCommon.ush
@@ -402,10 +402,12 @@ float3 ComputeDirectLighting(
 		ApplyPositionBias(ShadowRay, bBackFace ? -Payload.GetWorldNormal() : Payload.GetWorldNormal(), ShadowMaxNormalBias);
 
 		uint RayFlags = RAY_FLAG_ACCEPT_FIRST_HIT_AND_END_SEARCH	// Terminate immediatly after detecting a hit
-						| RAY_FLAG_SKIP_CLOSEST_HIT_SHADER			// Skip Closest Hit shader, only run Any Hit shader
-						| RAY_FLAG_CULL_NON_OPAQUE;					// Do not include any translucent from shadow as of today since we do not handle coverage & transmittance.
+						| RAY_FLAG_SKIP_CLOSEST_HIT_SHADER;			// Skip Closest Hit shader, only run Any Hit shader
 		const uint InstanceInclusionMask = (ApplyShadow) ? RAY_TRACING_MASK_SHADOW : 0;
 
+		// Do not include any translucent when tracing for shadow, since we do not handle coverage & transmittance yet.
+		Payload.SetIgnoreTranslucentMaterials();
+
 #if !ENABLE_TWO_SIDED_GEOMETRY
 		RayFlags |= RAY_FLAG_CULL_BACK_FACING_TRIANGLES;
 #endif // !ENABLE_TWO_SIDED_GEOMETRY
