Explicitly remove InstanceIdOffset attrib from vertex shaders that use UBO path for GPUScene

diff --git a/Engine/Shaders/Private/VertexFactoryCommon.ush b/Engine/Shaders/Private/VertexFactoryCommon.ush
--- a/Engine/Shaders/Private/VertexFactoryCommon.ush
+++ b/Engine/Shaders/Private/VertexFactoryCommon.ush
@@ -22,15 +22,21 @@
 // Required VF inputs to support dynamic instancing 
 //
 #if VF_USE_PRIMITIVE_SCENE_DATA
-	#define VF_GPUSCENE_DECLARE_INPUT_BLOCK(Attr0) \
-		uint InstanceIdOffset : ATTRIBUTE##Attr0; \
-		uint DrawInstanceId : SV_InstanceID;
-	
-	#define VF_GPUSCENE_GET_INTERMEDIATES(VFInput) \
-		GetSceneDataIntermediates( \
-			VFInput.InstanceIdOffset, \
-			VFInput.DrawInstanceId)
+	#if SHADER_USES_PRIMITIVE_UBO
+		#define VF_GPUSCENE_DECLARE_INPUT_BLOCK(Attr0) \
+			uint DrawInstanceId : SV_InstanceID;
+
+		#define VF_GPUSCENE_GET_INTERMEDIATES(VFInput) \
+			GetSceneDataIntermediates(0u, VFInput.DrawInstanceId)
+	#else
+		#define VF_GPUSCENE_DECLARE_INPUT_BLOCK(Attr0) \
+			uint InstanceIdOffset : ATTRIBUTE##Attr0; \
+			uint DrawInstanceId : SV_InstanceID;
 
+		#define VF_GPUSCENE_GET_INTERMEDIATES(VFInput) \
+			GetSceneDataIntermediates(VFInput.InstanceIdOffset, VFInput.DrawInstanceId)
+	#endif //SHADER_USES_PRIMITIVE_UBO
+	
 	#define VF_GPUSCENE_GET_LIGHTMAP_UV_SCALE_BIAS(VFInput, LightmapDataIndex) \
 		GetLightmapData(LightmapDataIndex, VFInput.DrawInstanceId).LightMapCoordinateScaleBias
 	
