Works arrounds bug in VelocityEncodeDepth() that prevents console preview in editor

diff --git a/Engine/Shaders/Private/TemporalSuperResolution/TSRDilateVelocity.usf b/Engine/Shaders/Private/TemporalSuperResolution/TSRDilateVelocity.usf
--- a/Engine/Shaders/Private/TemporalSuperResolution/TSRDilateVelocity.usf
+++ b/Engine/Shaders/Private/TemporalSuperResolution/TSRDilateVelocity.usf
@@ -260,7 +260,8 @@ void MainCS(
 	
 	bool HasPixelAnimation = false;
 	#if !VELOCITY_ENCODE_HAS_PIXEL_ANIMATION
-		#error TSR requires the bHasPixelAnimation encoded in the velocity buffer.
+		// UE-198416: there is high mess arround VelocityEncodeDepth() to clean up
+		//#error TSR requires the bHasPixelAnimation encoded in the velocity buffer.
 	#endif
 
 	// Compute final screen vecocity.
