Fix for underwater normal being bad for two sided SLW material with normal specified in world space.

diff --git a/Engine/Shaders/Private/MaterialTemplate.ush b/Engine/Shaders/Private/MaterialTemplate.ush
--- a/Engine/Shaders/Private/MaterialTemplate.ush
+++ b/Engine/Shaders/Private/MaterialTemplate.ush
@@ -3292,6 +3292,9 @@ float GetRoughnessFromNormalCurvature(FMaterialPixelParameters InMaterialParamet
 // EvaluateMaterialAttributes
 %{evaluate_material_attributes}
 
+// We must to flip the normal of two sided SLW material with normal specified in world space.
+#define TWO_SIDED_WORLD_SPACE_SINGLELAYERWATER_NORMAL (MATERIAL_SHADINGMODEL_SINGLELAYERWATER && !MATERIAL_TANGENTSPACENORMAL && MATERIAL_TWOSIDED)
+
 // Be very, very careful about changing CalcPixelMaterialInputs() or CalcPixelMaterialInputsAnalyticDerivatives(). They apply the same basic calculation,
 // but CalcPixelMaterialInputs() uses finite difference derivates from hardware whereas CalcPixelMaterialInputsAnalyticDerivatives() uses analytic
 // derivatives. If you change anything, make sure that both functions stay identical. 
@@ -3332,7 +3335,7 @@ void CalcPixelMaterialInputs(in out FMaterialPixelParameters Parameters, in out
 
 #endif //MATERIAL_TANGENTSPACENORMAL
 
-#if MATERIAL_TANGENTSPACENORMAL
+#if MATERIAL_TANGENTSPACENORMAL || TWO_SIDED_WORLD_SPACE_SINGLELAYERWATER_NORMAL
 	// flip the normal for backfaces being rendered with a two-sided material
 	Parameters.WorldNormal *= Parameters.TwoSidedSign;
 #endif
@@ -3394,7 +3397,7 @@ void CalcPixelMaterialInputsAnalyticDerivatives(in out FMaterialPixelParameters
 
 #endif //MATERIAL_TANGENTSPACENORMAL
 
-#if MATERIAL_TANGENTSPACENORMAL
+#if MATERIAL_TANGENTSPACENORMAL || TWO_SIDED_WORLD_SPACE_SINGLELAYERWATER_NORMAL
 	// flip the normal for backfaces being rendered with a two-sided material
 	Parameters.WorldNormal *= Parameters.TwoSidedSign;
 #endif
