Fix substrate shader compilation.

diff --git a/Engine/Shaders/Private/DefaultSSRTiles.usf b/Engine/Shaders/Private/DefaultSSRTiles.usf
--- a/Engine/Shaders/Private/DefaultSSRTiles.usf
+++ b/Engine/Shaders/Private/DefaultSSRTiles.usf
@@ -47,7 +47,8 @@ bool DoesPixelContainSSR(uint2 PixelPos)
 	const FSubstrateTopLayerData TopLayerData = SubstrateUnpackTopLayerData(SubstratePixelHeader.PackedTopLayerData);
 	bool bContainValidPixels =
 		TopLayerData.Roughness <= MaxRoughness
-		&& TopLayerData.Specular >= MinSpecular
+		// TODO_SUBSTRATE: need F0 value in order to have tge sane optimization as legacy
+		//&& TopLayerData.Specular >= MinSpecular
 		&& !SubstratePixelHeader.IsHair()
 		&& !SubstratePixelHeader.IsSingleLayerWater();
 	
