Prepare hair cards/meshes to have custom attributes packing.

diff --git a/Engine/Plugins/Runtime/HairStrands/Source/HairStrandsCore/Private/HairCardsVertexFactory.cpp b/Engine/Plugins/Runtime/HairStrands/Source/HairStrandsCore/Private/HairCardsVertexFactory.cpp
--- a/Engine/Plugins/Runtime/HairStrands/Source/HairStrandsCore/Private/HairCardsVertexFactory.cpp
+++ b/Engine/Plugins/Runtime/HairStrands/Source/HairStrandsCore/Private/HairCardsVertexFactory.cpp
@@ -35,7 +35,8 @@ enum class EHairCardsFactoryFlags : uint32
 	TextureAttribute  = 0x20,
 	TextureCoverage   = 0x40,
 	TextureDepth      = 0x80,
-	TextureTangent    = 0x100
+	TextureTangent    = 0x100,
+	TextureAuxilary   = 0x200
 };
 
 ////////////////////////////////////////////////////////////////////////////////////////////////////////////////////////
@@ -56,16 +57,19 @@ FHairCardsUniformBuffer CreateHairCardsVFUniformBuffer(
 
 		// Cards atlas UV are inverted so fetching needs to be inverted on the y-axis
 		UniformParameters.Flags = 0;
-		UniformParameters.Flags |= LOD.RestResource->bInvertUV 					? uint32(EHairCardsFactoryFlags::InvertedUV) : 0u;
-		UniformParameters.Flags |= LOD.RestResource->CoverageTexture != nullptr ? uint32(EHairCardsFactoryFlags::TextureCoverage) : 0;
-		UniformParameters.Flags |= LOD.RestResource->DepthTexture != nullptr 	? uint32(EHairCardsFactoryFlags::TextureDepth) : 0;
-		UniformParameters.Flags |= LOD.RestResource->TangentTexture != nullptr 	? uint32(EHairCardsFactoryFlags::TextureTangent) : 0;
+		UniformParameters.Flags |= LOD.RestResource->bInvertUV 						? uint32(EHairCardsFactoryFlags::InvertedUV) : 0u;
+		UniformParameters.Flags |= LOD.RestResource->CoverageTexture != nullptr 	? uint32(EHairCardsFactoryFlags::TextureCoverage) : 0;
+		UniformParameters.Flags |= LOD.RestResource->DepthTexture != nullptr 		? uint32(EHairCardsFactoryFlags::TextureDepth) : 0;
+		UniformParameters.Flags |= LOD.RestResource->TangentTexture != nullptr 		? uint32(EHairCardsFactoryFlags::TextureTangent) : 0;
+		UniformParameters.Flags |= LOD.RestResource->AuxilaryDataTexture != nullptr ? uint32(EHairCardsFactoryFlags::TextureAuxilary) : 0;
 		// These attributes are stored on vertex data rather than on textures
 		// EHairCardsFactoryFlags::TextureRootUV
 		// EHairCardsFactoryFlags::TextureGroupIndex
 		// EHairCardsFactoryFlags::TextureBaseColor
 		// EHairCardsFactoryFlags::TextureRoughness
 
+		UniformParameters.AttributeTextureIndex = 0; // TODO pack attribute texture/channel description. Currently use hardcoded layout in shader attribute code.
+		UniformParameters.AttributeChannelIndex = 0; // TODO pack attribute texture/channel description. Currently use hardcoded layout in shader attribute code.
 		UniformParameters.MaxVertexCount = LOD.RestResource->GetVertexCount();
 		UniformParameters.CoverageBias = FMath::Clamp(GHairCardCoverageBias, -16.f, 16.f);
 
@@ -88,18 +92,19 @@ FHairCardsUniformBuffer CreateHairCardsVFUniformBuffer(
 		UniformParameters.UVsBuffer = LOD.RestResource->UVsBuffer.ShaderResourceViewRHI.GetReference();
 		UniformParameters.MaterialsBuffer = LOD.RestResource->MaterialsBuffer.ShaderResourceViewRHI.GetReference();
 
-		UniformParameters.DepthTexture = LOD.RestResource->DepthTexture;
-		UniformParameters.DepthSampler = LOD.RestResource->DepthSampler;
-		UniformParameters.TangentTexture = LOD.RestResource->TangentTexture;
-		UniformParameters.TangentSampler = LOD.RestResource->TangentSampler;
-		UniformParameters.CoverageTexture = LOD.RestResource->CoverageTexture;
-		UniformParameters.CoverageSampler = LOD.RestResource->CoverageSampler;
-		UniformParameters.AttributeTexture = LOD.RestResource->AttributeTexture;
-		UniformParameters.AttributeSampler = LOD.RestResource->AttributeSampler;
-		UniformParameters.AuxilaryDataTexture = LOD.RestResource->AuxilaryDataTexture;
-		UniformParameters.AuxilaryDataSampler = LOD.RestResource->AuxilaryDataSampler;
-		UniformParameters.MaterialTexture = nullptr; // Material properties & Group index are stored on vertices for cards
-		UniformParameters.MaterialSampler = nullptr; // Material properties & Group index are stored on vertices for cards
+
+		UniformParameters.Texture0Texture = LOD.RestResource->DepthTexture;
+		UniformParameters.Texture0Sampler = LOD.RestResource->DepthSampler;
+		UniformParameters.Texture1Texture = LOD.RestResource->TangentTexture;
+		UniformParameters.Texture1Sampler = LOD.RestResource->TangentSampler;
+		UniformParameters.Texture2Texture = LOD.RestResource->CoverageTexture;
+		UniformParameters.Texture2Sampler = LOD.RestResource->CoverageSampler;
+		UniformParameters.Texture3Texture = LOD.RestResource->AttributeTexture;
+		UniformParameters.Texture3Sampler = LOD.RestResource->AttributeSampler;
+		UniformParameters.Texture4Texture = nullptr; // Material properties & Group index are stored on vertices for cards
+		UniformParameters.Texture4Sampler = nullptr; // Material properties & Group index are stored on vertices for cards
+		UniformParameters.Texture5Texture = LOD.RestResource->AuxilaryDataTexture;
+		UniformParameters.Texture5Sampler = LOD.RestResource->AuxilaryDataSampler;
 	}
 	else if (GeometryType == EHairGeometryType::Meshes)
 	{
@@ -129,52 +134,56 @@ FHairCardsUniformBuffer CreateHairCardsVFUniformBuffer(
 		UniformParameters.Flags |= LOD.RestResource->CoverageTexture != nullptr ? uint32(EHairCardsFactoryFlags::TextureCoverage) : 0;
 		UniformParameters.Flags |= LOD.RestResource->DepthTexture != nullptr ? uint32(EHairCardsFactoryFlags::TextureDepth) : 0;
 		UniformParameters.Flags |= LOD.RestResource->TangentTexture != nullptr ? uint32(EHairCardsFactoryFlags::TextureTangent) : 0;
+		UniformParameters.Flags |= LOD.RestResource->AuxilaryDataTexture != nullptr ? uint32(EHairCardsFactoryFlags::TextureAuxilary) : 0;
 
-		UniformParameters.CoverageBias = 0;
+		UniformParameters.AttributeTextureIndex = 0; // TODO pack attribute texture/channel description. Currently use hardcoded layout in shader attribute code.
+		UniformParameters.AttributeChannelIndex = 0; // TODO pack attribute texture/channel description. Currently use hardcoded layout in shader attribute code.
 		UniformParameters.MaxVertexCount = LOD.RestResource->GetVertexCount();
+		UniformParameters.CoverageBias = 0;
+
 		UniformParameters.NormalsBuffer = LOD.RestResource->NormalsBuffer.ShaderResourceViewRHI.GetReference();
 		UniformParameters.UVsBuffer = LOD.RestResource->UVsBuffer.ShaderResourceViewRHI.GetReference();
 		UniformParameters.MaterialsBuffer = UniformParameters.NormalsBuffer; // Reuse normal buffer as a dummy input for material buffer, since it not used for meshes (material data is fetch through textures)
 
-		UniformParameters.DepthTexture = LOD.RestResource->DepthTexture;
-		UniformParameters.DepthSampler = LOD.RestResource->DepthSampler;
-		UniformParameters.TangentTexture = LOD.RestResource->TangentTexture;
-		UniformParameters.TangentSampler = LOD.RestResource->TangentSampler;
-		UniformParameters.CoverageTexture = LOD.RestResource->CoverageTexture;
-		UniformParameters.CoverageSampler = LOD.RestResource->CoverageSampler;
-		UniformParameters.AttributeTexture = LOD.RestResource->AttributeTexture;
-		UniformParameters.AttributeSampler = LOD.RestResource->AttributeSampler;
-		UniformParameters.AuxilaryDataTexture = LOD.RestResource->AuxilaryDataTexture;
-		UniformParameters.AuxilaryDataSampler = LOD.RestResource->AuxilaryDataSampler;
-		UniformParameters.MaterialTexture = LOD.RestResource->MaterialTexture;
-		UniformParameters.MaterialSampler = LOD.RestResource->MaterialSampler;
+		UniformParameters.Texture0Texture = LOD.RestResource->DepthTexture;
+		UniformParameters.Texture0Sampler = LOD.RestResource->DepthSampler;
+		UniformParameters.Texture1Texture = LOD.RestResource->TangentTexture;
+		UniformParameters.Texture1Sampler = LOD.RestResource->TangentSampler;
+		UniformParameters.Texture2Texture = LOD.RestResource->CoverageTexture;
+		UniformParameters.Texture2Sampler = LOD.RestResource->CoverageSampler;
+		UniformParameters.Texture3Texture = LOD.RestResource->AttributeTexture;
+		UniformParameters.Texture3Sampler = LOD.RestResource->AttributeSampler;
+		UniformParameters.Texture5Texture = LOD.RestResource->AuxilaryDataTexture;
+		UniformParameters.Texture4Texture = LOD.RestResource->MaterialTexture;
+		UniformParameters.Texture4Sampler = LOD.RestResource->MaterialSampler;
+		UniformParameters.Texture5Sampler = LOD.RestResource->AuxilaryDataSampler;
 	}
 
 	if (!bSupportsManualVertexFetch)
 	{
-		UniformParameters.PositionBuffer = GNullVertexBuffer.VertexBufferSRV;
+		UniformParameters.PositionBuffer         = GNullVertexBuffer.VertexBufferSRV;
 		UniformParameters.PreviousPositionBuffer = GNullVertexBuffer.VertexBufferSRV;
-		UniformParameters.NormalsBuffer = GNullVertexBuffer.VertexBufferSRV;
-		UniformParameters.UVsBuffer = GNullVertexBuffer.VertexBufferSRV;
-		UniformParameters.MaterialsBuffer = GNullVertexBuffer.VertexBufferSRV;
+		UniformParameters.NormalsBuffer          = GNullVertexBuffer.VertexBufferSRV;
+		UniformParameters.UVsBuffer              = GNullVertexBuffer.VertexBufferSRV;
+		UniformParameters.MaterialsBuffer        = GNullVertexBuffer.VertexBufferSRV;
 	}
 
 	FRHITexture* DefaultTexture = GBlackTexture->TextureRHI;
 	FSamplerStateRHIRef DefaultSampler = TStaticSamplerState<SF_Bilinear, AM_Clamp, AM_Clamp, AM_Clamp>::GetRHI();
 
-	if (!UniformParameters.DepthTexture)		{ UniformParameters.DepthTexture = DefaultTexture;	  }
-	if (!UniformParameters.TangentTexture)		{ UniformParameters.TangentTexture = DefaultTexture;  }
-	if (!UniformParameters.CoverageTexture)		{ UniformParameters.CoverageTexture = DefaultTexture; }
-	if (!UniformParameters.AttributeTexture)	{ UniformParameters.AttributeTexture = DefaultTexture;}
-	if (!UniformParameters.AuxilaryDataTexture)	{ UniformParameters.AuxilaryDataTexture = DefaultTexture; }
-	if (!UniformParameters.MaterialTexture)		{ UniformParameters.MaterialTexture = DefaultTexture; }
-
-	if (!UniformParameters.DepthSampler)		{ UniformParameters.DepthSampler = DefaultSampler;	  }
-	if (!UniformParameters.TangentSampler)		{ UniformParameters.TangentSampler = DefaultSampler;  }
-	if (!UniformParameters.CoverageSampler)		{ UniformParameters.CoverageSampler = DefaultSampler; }
-	if (!UniformParameters.AttributeSampler)	{ UniformParameters.AttributeSampler = DefaultSampler;}
-	if (!UniformParameters.AuxilaryDataSampler) { UniformParameters.AuxilaryDataSampler = DefaultSampler; }
-	if (!UniformParameters.MaterialSampler)		{ UniformParameters.MaterialSampler = DefaultSampler;}
+	if (!UniformParameters.Texture0Texture)	{ UniformParameters.Texture0Texture = DefaultTexture; }
+	if (!UniformParameters.Texture1Texture)	{ UniformParameters.Texture1Texture = DefaultTexture; }
+	if (!UniformParameters.Texture2Texture)	{ UniformParameters.Texture2Texture = DefaultTexture; }
+	if (!UniformParameters.Texture3Texture)	{ UniformParameters.Texture3Texture = DefaultTexture; }
+	if (!UniformParameters.Texture4Texture)	{ UniformParameters.Texture4Texture = DefaultTexture; }
+	if (!UniformParameters.Texture5Texture)	{ UniformParameters.Texture5Texture = DefaultTexture; }
+
+	if (!UniformParameters.Texture0Sampler) { UniformParameters.Texture0Sampler = DefaultSampler; }
+	if (!UniformParameters.Texture1Sampler) { UniformParameters.Texture1Sampler = DefaultSampler; }
+	if (!UniformParameters.Texture2Sampler) { UniformParameters.Texture2Sampler = DefaultSampler; }
+	if (!UniformParameters.Texture3Sampler) { UniformParameters.Texture3Sampler = DefaultSampler; }
+	if (!UniformParameters.Texture4Sampler) { UniformParameters.Texture4Sampler = DefaultSampler; }
+	if (!UniformParameters.Texture5Sampler) { UniformParameters.Texture5Sampler = DefaultSampler; }
 
 	return TUniformBufferRef<FHairCardsVertexFactoryUniformShaderParameters>::CreateUniformBufferImmediate(UniformParameters, UniformBuffer_MultiFrame);
 }
diff --git a/Engine/Plugins/Runtime/HairStrands/Source/HairStrandsCore/Public/GroomInstance.h b/Engine/Plugins/Runtime/HairStrands/Source/HairStrandsCore/Public/GroomInstance.h
--- a/Engine/Plugins/Runtime/HairStrands/Source/HairStrandsCore/Public/GroomInstance.h
+++ b/Engine/Plugins/Runtime/HairStrands/Source/HairStrandsCore/Public/GroomInstance.h
@@ -24,6 +24,8 @@ enum class EGroomViewMode : uint8;
 //  * 12 bits for cards count -> 4000 cards for a hair group
 BEGIN_GLOBAL_SHADER_PARAMETER_STRUCT(FHairCardsVertexFactoryUniformShaderParameters, HAIRSTRANDSCORE_API)
 	SHADER_PARAMETER(uint32, Flags)
+	SHADER_PARAMETER(uint32, AttributeTextureIndex)
+	SHADER_PARAMETER(uint32, AttributeChannelIndex)
 	SHADER_PARAMETER(uint32, MaxVertexCount)
 	SHADER_PARAMETER(float, CoverageBias)
 	SHADER_PARAMETER_SRV(Buffer<float4>, PositionBuffer)
@@ -31,19 +33,18 @@ BEGIN_GLOBAL_SHADER_PARAMETER_STRUCT(FHairCardsVertexFactoryUniformShaderParamet
 	SHADER_PARAMETER_SRV(Buffer<float4>, NormalsBuffer)
 	SHADER_PARAMETER_SRV(Buffer<float4>, UVsBuffer)
 	SHADER_PARAMETER_SRV(Buffer<float4>, MaterialsBuffer)
-
-	SHADER_PARAMETER_TEXTURE(Texture2D<float4>, DepthTexture)
-	SHADER_PARAMETER_SAMPLER(SamplerState, DepthSampler)
-	SHADER_PARAMETER_TEXTURE(Texture2D<float4>, TangentTexture)
-	SHADER_PARAMETER_SAMPLER(SamplerState, TangentSampler)
-	SHADER_PARAMETER_TEXTURE(Texture2D<float4>, CoverageTexture)
-	SHADER_PARAMETER_SAMPLER(SamplerState, CoverageSampler)
-	SHADER_PARAMETER_TEXTURE(Texture2D<float4>, AttributeTexture)
-	SHADER_PARAMETER_SAMPLER(SamplerState, AttributeSampler)
-	SHADER_PARAMETER_TEXTURE(Texture2D<float4>, AuxilaryDataTexture)
-	SHADER_PARAMETER_SAMPLER(SamplerState, AuxilaryDataSampler)
-	SHADER_PARAMETER_TEXTURE(Texture2D<float4>, MaterialTexture)
-	SHADER_PARAMETER_SAMPLER(SamplerState, MaterialSampler)
+	SHADER_PARAMETER_TEXTURE(Texture2D<float4>, Texture0Texture) 
+	SHADER_PARAMETER_SAMPLER(SamplerState,      Texture0Sampler)
+	SHADER_PARAMETER_TEXTURE(Texture2D<float4>, Texture1Texture)
+	SHADER_PARAMETER_SAMPLER(SamplerState,      Texture1Sampler)
+	SHADER_PARAMETER_TEXTURE(Texture2D<float4>, Texture2Texture)
+	SHADER_PARAMETER_SAMPLER(SamplerState,      Texture2Sampler)
+	SHADER_PARAMETER_TEXTURE(Texture2D<float4>, Texture3Texture)
+	SHADER_PARAMETER_SAMPLER(SamplerState,      Texture3Sampler)
+	SHADER_PARAMETER_TEXTURE(Texture2D<float4>, Texture4Texture)
+	SHADER_PARAMETER_SAMPLER(SamplerState,      Texture4Sampler)
+	SHADER_PARAMETER_TEXTURE(Texture2D<float4>, Texture5Texture)
+	SHADER_PARAMETER_SAMPLER(SamplerState,      Texture5Sampler)
 END_GLOBAL_SHADER_PARAMETER_STRUCT()
 
 typedef TUniformBufferRef<FHairCardsVertexFactoryUniformShaderParameters> FHairCardsUniformBuffer;
diff --git a/Engine/Shaders/Private/HairStrands/HairCardsAttributeCommon.ush b/Engine/Shaders/Private/HairStrands/HairCardsAttributeCommon.ush
--- a/Engine/Shaders/Private/HairStrands/HairCardsAttributeCommon.ush
+++ b/Engine/Shaders/Private/HairStrands/HairCardsAttributeCommon.ush
@@ -6,162 +6,197 @@
 
 #include "/Engine/Private/HairStrands/HairStrandsVertexFactoryCommon.ush"
 
-////////////////////////////////////////////////////////////////////////////////
+///////////////////////////////////////////////////////////////////////////////////////////////////
+
+// Vertex factory Flags
+#define EHairCardsVFlags_InvertedUV        0x1
+#define EHairCardsVFlags_TextureRootUV     0x2
+#define EHairCardsVFlags_TextureGroupIndex 0x4
+#define EHairCardsVFlags_TextureColor      0x8
+#define EHairCardsVFlags_TextureRoughness  0x10
+#define EHairCardsVFlags_TextureAttribute  0x20
+#define EHairCardsVFlags_TextureCoverage   0x40
+#define EHairCardsVFlags_TextureDepth      0x80
+#define EHairCardsVFlags_TextureTangent    0x100
+#define EHairCardsVFlags_TextureAuxilary   0x200
+
+// Attribute index
+#define EHairCardsVFAttribute_Depth        0
+#define EHairCardsVFAttribute_Seed         1
+#define EHairCardsVFAttribute_CoordU       2
+#define EHairCardsVFAttribute_GroupIndex   3
+#define EHairCardsVFAttribute_Roughness    4
+#define EHairCardsVFAttribute_RootUV       5
+#define EHairCardsVFAttribute_Color        6
+#define EHairCardsVFAttribute_Tangent      7
+#define EHairCardsVFAttribute_Coverage     8
+#define EHairCardsVFAttribute_Auxilary     9
+
+///////////////////////////////////////////////////////////////////////////////////////////////////
+// Helper functions for fetching attributes
+
+#define HAIRCARDSVF_USE_LEGACY_LAYOUT 1
+#if HAIRCARDSVF_USE_LEGACY_LAYOUT
+uint GetHairCardTextureIndex(uint In) 
+{ 
+	switch (In)
+	{
+		case EHairCardsVFAttribute_Depth      : return 0;
+		case EHairCardsVFAttribute_Seed       : return 3;
+		case EHairCardsVFAttribute_CoordU     : return 3;
+		case EHairCardsVFAttribute_GroupIndex : return 4; // Only used for mesh, not for cards
+		case EHairCardsVFAttribute_Roughness  : return 4;
+		case EHairCardsVFAttribute_RootUV     : return 3;
+		case EHairCardsVFAttribute_Color      : return 4;
+		case EHairCardsVFAttribute_Tangent    : return 1;
+		case EHairCardsVFAttribute_Coverage   : return 2;
+		case EHairCardsVFAttribute_Auxilary   : return 5; 
+	}
+	return 0;
+}
+
+uint GetHairCardChannelIndex(uint In) 
+{ 
+	switch (In)
+	{
+		case EHairCardsVFAttribute_Depth      : return 0;
+		case EHairCardsVFAttribute_Seed       : return 3;
+		case EHairCardsVFAttribute_CoordU     : return 2;
+		case EHairCardsVFAttribute_GroupIndex : return 3;
+		case EHairCardsVFAttribute_Roughness  : return 3;
+		case EHairCardsVFAttribute_RootUV     : return 0;
+		case EHairCardsVFAttribute_Color      : return 0;
+		case EHairCardsVFAttribute_Tangent    : return 0;
+		case EHairCardsVFAttribute_Coverage   : return 0;
+		case EHairCardsVFAttribute_Auxilary   : return 0; 
+	}
+	return 0;
+}
+#else
+uint GetHairCardTextureIndex(uint In)	{ return ((HairCards.AttributeTextureIndex) >> (In*3u) & 0x7u); }
+uint GetHairCardChannelIndex(uint In)	{ return ((HairCards.AttributeChannelIndex) >> (In*2u) & 0x3u); }
+#endif
+float2 GetAtlasUV(float2 In) 			{ return float2(In.x, (HairCardsVF.Flags & EHairCardsVFlags_InvertedUV) ? 1 - In.y : In.y); }
+
+// Generic attribute texture fetches
+float4 FetchAttributeTexture(uint InTextureIndex, float2 InAtlasUV, float InBias)
+{
+	float4 Out = 0;
+	     if (InTextureIndex == 0) { Out = Texture2DSampleBias(HairCardsVF.Texture0Texture, HairCardsVF.Texture0Sampler, InAtlasUV, InBias); }
+	else if (InTextureIndex == 1) { Out = Texture2DSampleBias(HairCardsVF.Texture1Texture, HairCardsVF.Texture1Sampler, InAtlasUV, InBias); }
+	else if (InTextureIndex == 2) { Out = Texture2DSampleBias(HairCardsVF.Texture2Texture, HairCardsVF.Texture2Sampler, InAtlasUV, InBias); }
+	else if (InTextureIndex == 3) { Out = Texture2DSampleBias(HairCardsVF.Texture3Texture, HairCardsVF.Texture3Sampler, InAtlasUV, InBias); }
+	else if (InTextureIndex == 4) { Out = Texture2DSampleBias(HairCardsVF.Texture4Texture, HairCardsVF.Texture4Sampler, InAtlasUV, InBias); }
+	else if (InTextureIndex == 5) { Out = Texture2DSampleBias(HairCardsVF.Texture5Texture, HairCardsVF.Texture5Sampler, InAtlasUV, InBias); }
+	return Out;
+}
 
-#define EHairCardsFVFlags_InvertedUV        0x1
-#define EHairCardsFVFlags_TextureRootUV     0x2
-#define EHairCardsFVFlags_TextureGroupIndex 0x4
-#define EHairCardsFVFlags_TextureBaseColor  0x8
-#define EHairCardsFVFlags_TextureRoughness  0x10
-#define EHairCardsFVFlags_TextureAttribute  0x20
-#define EHairCardsFVFlags_TextureCoverage   0x40
-#define EHairCardsFVFlags_TextureDepth      0x80
-#define EHairCardsFVFlags_TextureTangent    0x100
-
-struct FHairVertexAttributes
+float4 FetchAttributeTexture(uint InTextureIndex, float2 InAtlasUV)
 {
-	float2 HairUV;
-	float2 HairRootUV;
-	float2 HairDimensions;
-	float3 HairBaseColor;
-	float  HairSeed;
-	float  HairRoughness;
-	float  HairGroupIndex;
-};
+	float4 Out = 0;
+		 if (InTextureIndex == 0) { Out = Texture2DSample(HairCardsVF.Texture0Texture, HairCardsVF.Texture0Sampler, InAtlasUV); }
+	else if (InTextureIndex == 1) { Out = Texture2DSample(HairCardsVF.Texture1Texture, HairCardsVF.Texture1Sampler, InAtlasUV); }
+	else if (InTextureIndex == 2) { Out = Texture2DSample(HairCardsVF.Texture2Texture, HairCardsVF.Texture2Sampler, InAtlasUV); }
+	else if (InTextureIndex == 3) { Out = Texture2DSample(HairCardsVF.Texture3Texture, HairCardsVF.Texture3Sampler, InAtlasUV); }
+	else if (InTextureIndex == 4) { Out = Texture2DSample(HairCardsVF.Texture4Texture, HairCardsVF.Texture4Sampler, InAtlasUV); }
+	else if (InTextureIndex == 5) { Out = Texture2DSample(HairCardsVF.Texture5Texture, HairCardsVF.Texture5Sampler, InAtlasUV); }
+	return Out;
+}
+
+// Generic attribute fetch functions
+float  FetchAttribute1(uint In, float2 InAtlasUV, float InBias) { const float4 Value = FetchAttributeTexture(GetHairCardTextureIndex(In), InAtlasUV, InBias); return Value[GetHairCardChannelIndex(In)];                     }
+float  FetchAttribute1(uint In, float2 InAtlasUV)               { const float4 Value = FetchAttributeTexture(GetHairCardTextureIndex(In), InAtlasUV);         return Value[GetHairCardChannelIndex(In)];                     }
+float2 FetchAttribute2(uint In, float2 InAtlasUV)               { const float4 Value = FetchAttributeTexture(GetHairCardTextureIndex(In), InAtlasUV);         return GetHairCardChannelIndex(In) == 0 ? Value.xy : Value.zw; } // float2 alignemnt is forced onto .xy or .zw
+float3 FetchAttribute3(uint In, float2 InAtlasUV)               { const float4 Value = FetchAttributeTexture(GetHairCardTextureIndex(In), InAtlasUV);         return Value.xyz;                                              } // float3 alignemnt forced onto .xyz
+float4 FetchAttribute4(uint In, float2 InAtlasUV)               { const float4 Value = FetchAttributeTexture(GetHairCardTextureIndex(In), InAtlasUV);         return Value;                                                  }
+
+// Special attribute fetch functions
+float  GetDepthOffset (float2 InAtlasUV)                        { float  Out = 0;         if (HairCardsVF.Flags & EHairCardsVFlags_TextureDepth)     { Out = FetchAttribute1(EHairCardsVFAttribute_Depth,     InAtlasUV);                          } return Out; }
+float  GetSeed        (float2 InAtlasUV)                        { float  Out = 0;         if (HairCardsVF.Flags & EHairCardsVFlags_TextureDepth)     { Out = FetchAttribute1(EHairCardsVFAttribute_Seed,      InAtlasUV);                          } return Out; }
+float  GetCoordU      (float2 InAtlasUV)                        { float  Out = 0;         if (HairCardsVF.Flags & EHairCardsVFlags_TextureDepth)     { Out = FetchAttribute1(EHairCardsVFAttribute_CoordU,    InAtlasUV);                          } return Out; }
+float  GetGroupIndex  (float2 InAtlasUV, float  InDefault)      { float  Out = InDefault; if (HairCardsVF.Flags & EHairCardsVFlags_TextureGroupIndex){ Out = FetchAttribute1(EHairCardsVFAttribute_GroupIndex,InAtlasUV); Out *= 255.f;            } return Out; }
+float  GetRoughness   (float2 InAtlasUV, float  InDefault)      { float  Out = InDefault; if (HairCardsVF.Flags & EHairCardsVFlags_TextureRoughness) { Out = FetchAttribute1(EHairCardsVFAttribute_Roughness, InAtlasUV);                          } return Out; }
+float  GetCoverage    (float2 InAtlasUV, float  InDefault)      { float  Out = InDefault; if (HairCardsVF.Flags & EHairCardsVFlags_TextureCoverage)  { Out = FetchAttribute1(EHairCardsVFAttribute_Coverage,  InAtlasUV, HairCardsVF.CoverageBias);} return Out; }
+float2 GetRootUV      (float2 InAtlasUV, float2 InDefault)      { float2 Out = InDefault; if (HairCardsVF.Flags & EHairCardsVFlags_TextureRootUV)    { Out = FetchAttribute2(EHairCardsVFAttribute_RootUV,    InAtlasUV);                          } return Out; }
+float3 GetColor       (float2 InAtlasUV, float3 InDefault)      { float3 Out = InDefault; if (HairCardsVF.Flags & EHairCardsVFlags_TextureColor)     { Out = FetchAttribute3(EHairCardsVFAttribute_Color,     InAtlasUV); Out = Pow2(Out);         } return Out; } // Cheap sRGB -> Linear encoding
+float3 GetTangent     (float2 InAtlasUV, float3 InDefault)      { float3 Out = InDefault; if (HairCardsVF.Flags & EHairCardsVFlags_TextureTangent)   { Out = FetchAttribute3(EHairCardsVFAttribute_Tangent,   InAtlasUV); Out = Out*2-1.f;         } return Out; }
+float4 GetAuxilary    (float2 InAtlasUV, float4 InDefault)      { float4 Out = InDefault; if (HairCardsVF.Flags & EHairCardsVFlags_TextureAuxilary)  { Out = FetchAttribute4(EHairCardsVFAttribute_Auxilary,  InAtlasUV);                          } return Out; }
+
+///////////////////////////////////////////////////////////////////////////////////////////////////
+// MaterialTemplate interface
 
 float GetHairStrandsDepth(float2 InAtlasUV, float InDeviceZ)
 {
-	float SceneDepthOffset = 0;
-	if (HairCardsVF.Flags & EHairCardsFVFlags_TextureDepth)
-	{
-		const bool bInvertUV = HairCardsVF.Flags & EHairCardsFVFlags_InvertedUV;
-		InAtlasUV.y = bInvertUV ? 1 - InAtlasUV.y : InAtlasUV.y;
-		SceneDepthOffset = Texture2DSample(HairCardsVF.DepthTexture, HairCardsVF.DepthSampler, InAtlasUV).x;
-	}
+	const float SceneDepthOffset = GetDepthOffset(InAtlasUV);
 	const float SceneDepth = ConvertFromDeviceZ(InDeviceZ);
-
 	return ConvertToDeviceZ(SceneDepth + SceneDepthOffset);
 }
 
 float3 GetHairStrandsTangent(float2 InAtlasUV, half3x3 TangentToWorld, bool bUseTangentSpace)
 {
-	float3 LocalTangent = float3(0, 1, 0);
-	if (HairCardsVF.Flags & EHairCardsFVFlags_TextureTangent)
-	{
-		const bool bInvertUV = HairCardsVF.Flags & EHairCardsFVFlags_InvertedUV;
-		InAtlasUV.y = bInvertUV ? 1 - InAtlasUV.y : InAtlasUV.y;
-		LocalTangent = Texture2DSample(HairCardsVF.TangentTexture, HairCardsVF.TangentSampler, InAtlasUV).xyz * 2 - 1;
-	}
+	const float3 LocalTangent = GetTangent(InAtlasUV, float3(0, 1, 0));
 	const float3 WorldTangent = mul(LocalTangent, TangentToWorld);
 	return bUseTangentSpace ? LocalTangent : WorldTangent;
 }
 
 float GetHairStrandsCoverage(float2 InAtlasUV)
 {
-	float Coverage = 1.f;
-	if (HairCardsVF.Flags & EHairCardsFVFlags_TextureCoverage)
-	{
-		const bool bInvertUV = HairCardsVF.Flags & EHairCardsFVFlags_InvertedUV;
-		InAtlasUV.y = bInvertUV ? 1 - InAtlasUV.y : InAtlasUV.y;
-		Coverage = Texture2DSampleBias(HairCardsVF.CoverageTexture, HairCardsVF.CoverageSampler, InAtlasUV, HairCardsVF.CoverageBias).x;
-	}
-	return Coverage;
-}
-
-FHairVertexAttributes GetHairStrandsAttributes(float2 InAtlasUV, float2 InRootUV = float2(0, 0), float InCardLength = 0.f, float InCardGroupIndex = 0.f, float3 InCardBaseColor = 0.f, float InCardRoughness = 0.f)
-{
-	const bool bInvertUV = HairCardsVF.Flags & EHairCardsFVFlags_InvertedUV;
-	InAtlasUV.y = bInvertUV ? 1 - InAtlasUV.y : InAtlasUV.y;
-	const float4 Value = Texture2DSample(HairCardsVF.AttributeTexture, HairCardsVF.AttributeSampler, InAtlasUV);
-	const float CoordU = Value.z;
-	const float Seed = Value.w;
-	const float WorldRadius = 0.01f; // Hardcoded for cards
-	const float WorldLength = InCardLength;
-	
-	// Invert V to compensate image origin flip
-	// Similar to DecodeHairAttribute in HairStrandVertexFactoryCommon.ush
-	InRootUV.y = 1.0 - InRootUV.y;
-
-	FHairVertexAttributes Out;
-	Out.HairUV = float2(CoordU, 0.5f);
-	Out.HairRootUV = (HairCardsVF.Flags & EHairCardsFVFlags_TextureRootUV) ? Value.xy : InRootUV;
-	Out.HairDimensions = float2(WorldLength, WorldRadius);
-	Out.HairBaseColor = InCardBaseColor;
-	Out.HairSeed = Seed;
-	Out.HairRoughness = InCardRoughness;
-	Out.HairGroupIndex = InCardGroupIndex;
-	return Out;
+	return GetCoverage(InAtlasUV, 1.f);
 }
 
-float4 GetHairStrandsAuxilaryData(float2 HairPrimitiveUV)
+float4 GetHairStrandsAuxilaryData(float2 InAtlasUV)
 {
-	return Texture2DSample(HairCardsVF.AuxilaryDataTexture, HairCardsVF.AuxilaryDataSampler, HairPrimitiveUV);
+	return GetAuxilary(InAtlasUV, 0.f);
 }
 
-float2 GetHairStrandsUV(float2 HairPrimitiveUV)
+float2 GetHairStrandsUV(float2 InAtlasUV)
 {
-	return GetHairStrandsAttributes(HairPrimitiveUV).HairUV;
+	return float2(GetCoordU(InAtlasUV), 0.5f);
 }
 
-float2 GetHairStrandsDimensions(float2 HairPrimitiveUV, float HairPrimitiveLength)
+float2 GetHairStrandsDimensions(float2 InAtlasUV, float HairPrimitiveLength)
 {
-	return GetHairStrandsAttributes(HairPrimitiveUV, float2(0, 0), HairPrimitiveLength).HairDimensions;
+	return float2(HairPrimitiveLength, 0.01f/*WorldRadius*/);
 }
 
-float2 GetHairStrandsRootUV(float2 HairPrimitiveUV, float2 InRootUV)
+float2 GetHairStrandsRootUV(float2 InAtlasUV, float2 InRootUV)
 {
 	// Invert V to compensate image origin flip
 	// Similar to DecodeHairAttribute in HairStrandVertexFactoryCommon.ush
-	InRootUV.y = 1.0 - InRootUV.y; 
-
-	// If VF use texture data (for helmet/meshes), fetch root UV from the attribute texture
-	if (HairCardsVF.Flags & EHairCardsFVFlags_TextureRootUV)
-	{
-		InRootUV = Texture2DSample(HairCardsVF.AttributeTexture, HairCardsVF.AttributeSampler, HairPrimitiveUV).xy;
-	}
-	return InRootUV;
+	return GetRootUV(InAtlasUV, float2(InRootUV.x, 1.0 - InRootUV.y));
 }
 
-float  GetHairStrandsSeed(float2 HairPrimitiveUV)
+float GetHairStrandsSeed(float2 InAtlasUV)
 {
-	return GetHairStrandsAttributes(HairPrimitiveUV).HairSeed;
+	return GetSeed(InAtlasUV);
 }
 
-uint3 GetHairStrandsClumpID(float2 HairPrimitiveUV)
+uint3 GetHairStrandsClumpID(float2 InAtlasUV)
 {
 	return 0;
 }
 
-float3 GetHairStrandsBaseColor(float2 InAtlasUV, float3 InBaseColor)
+float3 GetHairStrandsColor(float2 InAtlasUV, float3 InColor)
 {
-	if (HairCardsVF.Flags & EHairCardsFVFlags_TextureBaseColor)
-	{
-		const float3 BaseColorInSRGB = Texture2DSample(HairCardsVF.MaterialTexture, HairCardsVF.MaterialSampler, InAtlasUV).xyz;
-		InBaseColor = BaseColorInSRGB*BaseColorInSRGB; // Cheap sRGB -> Linear encoding
-	}
-	return InBaseColor;
+	return GetColor(InAtlasUV, InColor);
 }
 
-float GetHairStrandsRoughness(float2 HairPrimitiveUV, float InRoughness)
+float GetHairStrandsRoughness(float2 InAtlasUV, float InCardRoughness)
 {
-	// For meshes, we don't support roughness yet
-	return (HairCardsVF.Flags & EHairCardsFVFlags_TextureRoughness) ? 0 : InRoughness;
+	//return GetRoughness(InAtlasUV, InCardRoughness);
+	return (HairCardsVF.Flags & EHairCardsVFlags_TextureRoughness) ? 0 : InCardRoughness;
 }
 
-float GetHairStrandsAO(float2 HairPrimitiveUV, float InAO)
+float GetHairStrandsAO(float2 InAtlasUV, float InAO)
 {
 	return 1.0f;
 }
 
-float GetHairStrandsGroupIndex(float2 InAtlasUV, float InHairPrimitiveGroupIndex)
+float GetHairStrandsGroupIndex(float2 InAtlasUV, float InCardGroupIndex)
 {
 	// * For hair cards,  group index is stored onto vertices
 	// * For hair meshes, group index is stored within a texture
-	if (HairCardsVF.Flags & EHairCardsFVFlags_TextureGroupIndex)
-	{
-		InHairPrimitiveGroupIndex = Texture2DSample(HairCardsVF.MaterialTexture, HairCardsVF.MaterialSampler, InAtlasUV).w * 255.f;
-	}
-	return InHairPrimitiveGroupIndex;
+	return GetGroupIndex(InAtlasUV, InCardGroupIndex);
 }
 ////////////////////////////////////////////////////////////////////////////////
 
diff --git a/Engine/Shaders/Private/MaterialTemplate.ush b/Engine/Shaders/Private/MaterialTemplate.ush
--- a/Engine/Shaders/Private/MaterialTemplate.ush
+++ b/Engine/Shaders/Private/MaterialTemplate.ush
@@ -1613,7 +1613,7 @@ float3 MaterialExpressionGetHairBaseColor(FMaterialPixelParameters Parameters)
 #if HAIR_STRAND_MESH_FACTORY
 	return GetHairStrandsColor(Parameters.HairControlPointId, Parameters.HairPrimitiveUV.x);
 #elif HAIR_CARD_MESH_FACTORY
-	return GetHairStrandsBaseColor(Parameters.HairPrimitiveUV, Parameters.HairPrimitiveMaterial.xyz);
+	return GetHairStrandsColor(Parameters.HairPrimitiveUV, Parameters.HairPrimitiveMaterial.xyz);
 #else
 	return float3(0,0,0);
 #endif
