Change Orthographic projection boolean logic so shaders to rely on the shader compiler's interpretation rather than manually attempting to optimise via casting.
Also change material shaders to rely on the same logic.

diff --git a/Engine/Shaders/Private/Common.ush b/Engine/Shaders/Private/Common.ush
--- a/Engine/Shaders/Private/Common.ush
+++ b/Engine/Shaders/Private/Common.ush
@@ -1184,10 +1184,9 @@ float ConvertFromDeviceZ(float DeviceZ)
 }
 
 // Returns whether or not the given projection matrix is orthographic
-// Utilises asuint call to reduce compare to 1ALU cost
 bool IsOrthoProjection(float4x4 ViewToClip)
 {
-	return asuint(ViewToClip._44);
+	return ViewToClip._44 >= 1.0f;
 }
 
 // Returns whether or not the default view is orthographic
@@ -1196,6 +1195,12 @@ bool IsOrthoProjection()
 	return IsOrthoProjection(View.ViewToClip);
 }
 
+// Returns whether or not the given projection matrix is orthographic via float
+float IsOrthoProjectionFloat()
+{
+	return select(IsOrthoProjection(), 1.0f, 0.0f);
+}
+
 // inverse operation of ConvertFromDeviceZ()
 // @param SceneDepth (linear in world units, W)
 // @return DeviceZ (Z/W)
diff --git a/Engine/Source/Runtime/Engine/Private/Materials/HLSLMaterialTranslator.cpp b/Engine/Source/Runtime/Engine/Private/Materials/HLSLMaterialTranslator.cpp
--- a/Engine/Source/Runtime/Engine/Private/Materials/HLSLMaterialTranslator.cpp
+++ b/Engine/Source/Runtime/Engine/Private/Materials/HLSLMaterialTranslator.cpp
@@ -5014,7 +5014,7 @@ int32 FHLSLMaterialTranslator::ViewProperty(EMaterialExposedViewProperty Propert
 
 int32 FHLSLMaterialTranslator::IsOrthographic()
 {
-	return AddInlinedCodeChunkZeroDeriv(MCT_Float, TEXT("((View.ViewToClip[3][3] < 1.0f) ? 0.0f : 1.0f)"));
+	return AddInlinedCodeChunkZeroDeriv(MCT_Float, TEXT("IsOrthoProjectionFloat()"));
 }
 
 
