StochasticShadows - added neighbourhood clamp to reduce ghosting.

diff --git a/Engine/Shaders/Private/StochasticShadows/StochasticShadows.ush b/Engine/Shaders/Private/StochasticShadows/StochasticShadows.ush
--- a/Engine/Shaders/Private/StochasticShadows/StochasticShadows.ush
+++ b/Engine/Shaders/Private/StochasticShadows/StochasticShadows.ush
@@ -3,7 +3,7 @@
 #pragma once
 
 // Change this to force recompilation of all Stochastic Shadows shaders
-#pragma message("UESHADERMETADATA_VERSION A16ACC4F-1C35-1CD1-1AAF-82395A79FB2B")
+#pragma message("UESHADERMETADATA_VERSION CA6ACC4F-2C35-2CD1-1CAF-82395A79FB2B")
 
 #define SUPPORT_CONTACT_SHADOWS 0
 #define USE_IES_PROFILE 0
@@ -24,7 +24,7 @@
 #define SHADOW_MASK_TILE_SIZE 8
 #define INVALID_SHADOW_MASK_TILE_INDEX 0xFFFFFFFF
 #define MAX_SHADOW_MASK_AGE 7
-#define MAX_HISTORY_LENGTH 7
+#define MAX_SHADOW_MASK_HISTORY_LENGTH 7
 
 // Keep in sync with StochasticShadows.cpp
 #define MAX_LIGHT_SCENE_ID_XY 16
@@ -127,8 +127,12 @@ struct FShadowMask
 {
 	float Factor;
 	float FactorSq;
-	uint Age;			// [0;7] When was last traced
-	uint HistoryLength; // [0;7] How many history frames were accumulated
+
+	// Number of frames from when this pixel was traced. Values in [0;7].
+	uint Age;			
+
+	// Number of frames accumulated in history. Values in [0;7].
+	uint HistoryLength;
 
 	bool IsValid()
 	{
@@ -153,7 +157,7 @@ uint PackShadowMask(FShadowMask ShadowMask)
 	PackedShadowMask |= clamp(uint(ShadowMask.Factor * 0x3FF + 0.5f), 0u, 0x3FFu); // 10 bits
 	PackedShadowMask |= clamp(uint(ShadowMask.FactorSq * 0x3FF + 0.5f), 0, 0x3FFu) << 10; // 10 bits
 	PackedShadowMask |= min(ShadowMask.Age, MAX_SHADOW_MASK_AGE) << 20;
-	PackedShadowMask |= min(ShadowMask.HistoryLength, MAX_HISTORY_LENGTH) << 23;
+	PackedShadowMask |= min(ShadowMask.HistoryLength, MAX_SHADOW_MASK_HISTORY_LENGTH) << 23;
 	return PackedShadowMask;
 }
 
@@ -351,8 +355,8 @@ void AccumulateHistoryShadowMask(uint LightSceneId, uint2 ShadowMaskScreenCoord,
 			FShadowMask HistoryShadowMask = UnpackShadowMask(ShadowMaskAtlasHistory[HistoryShadowMaskAtlasCoord]);
 
 			// Increase age as we are sampling previous frame
-			HistoryShadowMask.Age += 1;
-			HistoryShadowMask.HistoryLength += 1;
+			HistoryShadowMask.Age = min(HistoryShadowMask.Age + 1, MAX_SHADOW_MASK_AGE);
+			HistoryShadowMask.HistoryLength = min(HistoryShadowMask.HistoryLength + 1, MAX_SHADOW_MASK_HISTORY_LENGTH);
 
 			if (HistoryShadowMask.IsValid())
 			{
diff --git a/Engine/Shaders/Private/StochasticShadows/StochasticShadowsComposite.usf b/Engine/Shaders/Private/StochasticShadows/StochasticShadowsComposite.usf
--- a/Engine/Shaders/Private/StochasticShadows/StochasticShadowsComposite.usf
+++ b/Engine/Shaders/Private/StochasticShadows/StochasticShadowsComposite.usf
@@ -5,11 +5,37 @@
 #include "../ShaderPrint.ush"
 #include "StochasticShadows.ush"
 
-RWTexture2D<uint> RWShadowMaskTileAtlas;
 StructuredBuffer<uint2> ShadowMaskTileData;
 Texture2D<uint2> ShadowMaskHistoryScreenCoord00;
 Texture2D<float4> ShadowMaskHistoryWeights;
 Texture2D<uint> LightSamples;
+uint TemporalMaxFramesAccumulated;
+float TemporalStdDevOffset;
+
+RWTexture2D<uint> RWShadowMaskTileAtlas;
+
+#define SHARED_TILE_BORDER 1
+#define SHARED_TILE_SIZE (8 + 2 * SHARED_TILE_BORDER)
+groupshared uint SharedSampleAccumulator[SHARED_TILE_SIZE][SHARED_TILE_SIZE];
+
+struct FSampleAccumulator
+{
+	uint SampleSum;
+	uint NumSamples;
+};
+
+uint PackSampleAccumulator(FSampleAccumulator SampleAccumulator)
+{
+	return SampleAccumulator.SampleSum | (SampleAccumulator.NumSamples << 16);
+}
+
+FSampleAccumulator UnpackSampleAccumulator(uint Packed)
+{
+	FSampleAccumulator SampleAccumulator;
+	SampleAccumulator.SampleSum = Packed & 0xFFFF;
+	SampleAccumulator.NumSamples = Packed >> 16;
+	return SampleAccumulator;
+}
 
 /**
  * Run one thread per shadow mask pixel and composite current frame's traces.
@@ -22,14 +48,48 @@ void CompositeShadowMaskTracesCS(
 	uint3 DispatchThreadId : SV_DispatchThreadID)
 {
 	uint ShadowMaskTileIndex = GroupId.x;
-	uint2 LocalShadowMaskTileCoord = GroupThreadId.xy;
 
 	uint LightSceneId = 0;
-	uint2 ShadowMaskScreenCoord = 0;	
-	UnpackShadowMaskTile(ShadowMaskTileData[ShadowMaskTileIndex], LightSceneId, ShadowMaskScreenCoord);
-	ShadowMaskScreenCoord += LocalShadowMaskTileCoord;
+	uint2 BaseShadowMaskScreenCoord = 0;	
+	UnpackShadowMaskTile(ShadowMaskTileData[ShadowMaskTileIndex], LightSceneId, BaseShadowMaskScreenCoord);
+
+	GroupMemoryBarrierWithGroupSync();
+	
+	// Composite traces into SharedShadowMask
+	for (uint SharedShadowMaskY = GroupThreadId.y; SharedShadowMaskY < SHARED_TILE_SIZE; SharedShadowMaskY += THREADGROUP_SIZE)
+	{
+		for (uint SharedShadowMaskX = GroupThreadId.x; SharedShadowMaskX < SHARED_TILE_SIZE; SharedShadowMaskX += THREADGROUP_SIZE)
+		{
+			uint2 ShadowMaskScreenCoord = clamp(int2(BaseShadowMaskScreenCoord) + int2(SharedShadowMaskX - SHARED_TILE_BORDER, SharedShadowMaskY - SHARED_TILE_BORDER), 0, ShadowMaskViewSize - 1);
+			
+			FSampleAccumulator SampleAccumulator;
+			SampleAccumulator.SampleSum = 0;
+			SampleAccumulator.NumSamples = 0;
+
+			for (uint LocalSampleY = 0; LocalSampleY < NUM_SAMPLES_PER_PIXEL_Y; ++LocalSampleY)
+			{
+				for (uint LocalSampleX = 0; LocalSampleX < NUM_SAMPLES_PER_PIXEL_X; ++LocalSampleX)
+				{
+					uint2 SampleCoord = ShadowMaskScreenCoord * uint2(NUM_SAMPLES_PER_PIXEL_X, NUM_SAMPLES_PER_PIXEL_Y) + uint2(LocalSampleX, LocalSampleY);
+					FLightSample LightSample = UnpackLightSample(LightSamples[SampleCoord]);
+
+					if (LightSample.LightSceneId == LightSceneId)
+					{
+						SampleAccumulator.SampleSum += LightSample.bVisible ? 1 : 0;
+						SampleAccumulator.NumSamples += 1;
+					}
+				}
+			}
+
+			SharedSampleAccumulator[SharedShadowMaskX][SharedShadowMaskY] = PackSampleAccumulator(SampleAccumulator);			
+		}
+	}
+	
+	GroupMemoryBarrierWithGroupSync();
+
+	uint2 LocalShadowMaskTileCoord = GroupThreadId.xy;
+	uint2 ShadowMaskScreenCoord = BaseShadowMaskScreenCoord + LocalShadowMaskTileCoord;
 
-#define DEBUG_MODE 0
 #if DEBUG_MODE
 	int2 DebugScreenCoord = View.CursorPosition.x >= 0 ? View.CursorPosition * View.ViewResolutionFraction : -1;
 	bool bDebug = all(ShadowMaskScreenCoord == DebugScreenCoord / SHADOW_MASK_DOWNSAMPLE_FACTOR);
@@ -37,44 +97,53 @@ void CompositeShadowMaskTracesCS(
 
 	if (bDebug)
 	{
-		Print(Context, TEXT("CompositeShadowMaskTracesCS"));
+		Print(Context, TEXT("CompositeShadowMaskTraces"));
 	}
 #endif
 
+	// Temporal reproject and write out new shadow masks
 	if (all(ShadowMaskScreenCoord < ShadowMaskViewSize))
 	{
-		// Accumulate current frame traces
+		FSampleAccumulator CenterSampleAccumulator = UnpackSampleAccumulator(SharedSampleAccumulator[LocalShadowMaskTileCoord.x + SHARED_TILE_BORDER][LocalShadowMaskTileCoord.y + SHARED_TILE_BORDER]);
+
 		FShadowMask ShadowMask = InitShadowMask();
+		if (CenterSampleAccumulator.NumSamples > 0)
 		{
-			FShadowMaskAccumulator Accumulator = InitShadowMaskAccumulator();
-			for (uint LocalSampleY = 0; LocalSampleY < NUM_SAMPLES_PER_PIXEL_Y; ++LocalSampleY)
+			ShadowMask.Factor = CenterSampleAccumulator.SampleSum / (float)CenterSampleAccumulator.NumSamples;
+			ShadowMask.FactorSq = CenterSampleAccumulator.SampleSum / (float)CenterSampleAccumulator.NumSamples;
+			ShadowMask.Age = 0;
+		}
+
+		FSampleAccumulator SampleAccumulator = CenterSampleAccumulator;
+		const int FilterSize = 1;
+		for (int SampleOffsetY = -FilterSize; SampleOffsetY <= FilterSize; ++SampleOffsetY)
+		{
+			for (int SampleOffsetX = -FilterSize; SampleOffsetX <= FilterSize; ++SampleOffsetX)
 			{
-				for (uint LocalSampleX = 0; LocalSampleX < NUM_SAMPLES_PER_PIXEL_X; ++LocalSampleX)
+				if (!(SampleOffsetX == 0 && SampleOffsetY == 0))
 				{
-					uint2 SampleCoord = ShadowMaskScreenCoord * uint2(NUM_SAMPLES_PER_PIXEL_X, NUM_SAMPLES_PER_PIXEL_Y) + uint2(LocalSampleX, LocalSampleY);
-					FLightSample LightSample = UnpackLightSample(LightSamples[SampleCoord]);
-
-					if (LightSample.LightSceneId == LightSceneId)
-					{
-						float Factor = LightSample.bVisible ? 1.0f : 0.0f;
-
-						float Weight = 1.0f;
-						Accumulator.FactorSum += Factor * Weight;
-						Accumulator.FactorSqSum += Pow2(Factor) * Weight;
-						Accumulator.WeightSum += Weight;
-					}
+					FSampleAccumulator NeighborSampleAccumulator = UnpackSampleAccumulator(SharedSampleAccumulator[LocalShadowMaskTileCoord.x + SHARED_TILE_BORDER + SampleOffsetX][LocalShadowMaskTileCoord.y + SHARED_TILE_BORDER + SampleOffsetY]);
+					SampleAccumulator.SampleSum += NeighborSampleAccumulator.SampleSum;
+					SampleAccumulator.NumSamples += NeighborSampleAccumulator.NumSamples;
 				}
 			}
+		}
 
-			if (Accumulator.WeightSum > 0.0f)
-			{
-				ShadowMask.Age = 0;
-				ShadowMask.Factor = Accumulator.FactorSum / Accumulator.WeightSum;
-				ShadowMask.FactorSq = Accumulator.FactorSqSum / Accumulator.WeightSum;
-			}
+		#if DEBUG_MODE
+		if (bDebug && ShadowMask.IsValid())
+		{
+			Newline(Context);
+			Print(Context, TEXT("NumSamples: "));
+			Print(Context, SampleAccumulator.NumSamples);
+			Newline(Context);
+			Print(Context, TEXT("SampleSum: "));
+			Print(Context, SampleAccumulator.SampleSum);
+			Newline(Context);
+			Print(Context, TEXT("ShadowMask: "));
+			Print(Context, ShadowMask.Factor);
 		}
+		#endif
 
-		// History reprojection
 		#if TEMPORAL_ACCUMULATION
 		{
 			int2 ScreenCoord00 = ShadowMaskHistoryScreenCoord00[ShadowMaskScreenCoord];
@@ -94,21 +163,52 @@ void CompositeShadowMaskTracesCS(
 					float HistoryShadowFactor = Accumulator.FactorSum / Accumulator.WeightSum;
 					float HistoryShadowFactorSq = Accumulator.FactorSqSum / Accumulator.WeightSum;
 
+					// Neighborhood clamp
+					if (SampleAccumulator.NumSamples > CenterSampleAccumulator.NumSamples)
+					{
+						float Mean = SampleAccumulator.SampleSum / float(SampleAccumulator.NumSamples);
+						float Variance = max(Mean - Mean * Mean, 0.0f);
+						float StdDev = sqrt(Variance) + TemporalStdDevOffset;
+						float ClampMin = Mean - 0.5f * StdDev;
+						float ClampMax = Mean + 0.5f * StdDev;
+		
+						#if DEBUG_MODE
+						if (bDebug && ShadowMask.IsValid())
+						{
+							Newline(Context);
+							Print(Context, TEXT("Clamp: "));
+							Print(Context, ClampMin);
+							Print(Context, ClampMax);
+							Newline(Context);
+							Print(Context, TEXT("HistoryShadowFactor: "));
+							Print(Context, HistoryShadowFactor);
+						}
+						#endif
+
+						float ClampedHistoryShadowFactor = clamp(HistoryShadowFactor, ClampMin, ClampMax);
+						float DistanceToClamped = abs(ClampedHistoryShadowFactor - HistoryShadowFactor);
+
+						// Speedup accumulation if new samples don't fit history
+						Accumulator.MaxHistoryLength *= saturate(1.0f - 2.0f * DistanceToClamped);
+						HistoryShadowFactor = ClampedHistoryShadowFactor;
+					}
+
 					// Merge with reprojected history data
 					if (ShadowMask.IsValid())
 					{
-						float NewNumFramesAccumulated = max(Accumulator.MaxHistoryLength, 8.0f);
+						float NewNumFramesAccumulated = min(Accumulator.MaxHistoryLength, min(TemporalMaxFramesAccumulated, MAX_SHADOW_MASK_HISTORY_LENGTH));
 						float Alpha = 1.0f / (1.0f + NewNumFramesAccumulated);
 						ShadowMask.Factor = lerp(HistoryShadowFactor, ShadowMask.Factor, Alpha);
 						ShadowMask.FactorSq = lerp(HistoryShadowFactorSq, ShadowMask.FactorSq, Alpha);
 						ShadowMask.Age = 0;
+						ShadowMask.HistoryLength = Accumulator.MaxHistoryLength;
 					}
 					else
 					{
 						ShadowMask.Factor = HistoryShadowFactor;
 						ShadowMask.FactorSq = HistoryShadowFactorSq;
 						ShadowMask.Age = Accumulator.MinAge;
-						ShadowMask.HistoryLength = 0;
+						ShadowMask.HistoryLength = Accumulator.MaxHistoryLength;
 					}
 				}
 			}
diff --git a/Engine/Shaders/Private/StochasticShadows/StochasticShadowsHardwareRayTracing.usf b/Engine/Shaders/Private/StochasticShadows/StochasticShadowsHardwareRayTracing.usf
--- a/Engine/Shaders/Private/StochasticShadows/StochasticShadowsHardwareRayTracing.usf
+++ b/Engine/Shaders/Private/StochasticShadows/StochasticShadowsHardwareRayTracing.usf
@@ -34,17 +34,6 @@ StructuredBuffer<uint> CompactedTraceTexelData;
 StructuredBuffer<uint> CompactedTraceTexelAllocator;
 Texture2D<float> DownsampledSceneDepth;
 
-//#sdl_todo: Shared with Lumen?
-float3 CalcPrevTranslatedWorldPositionFromGPUSceneInstanceIndex(float3 TranslatedWorldPosition, uint GPUSceneInstanceIndex)
-{
-	FInstanceSceneData InstanceSceneData = GetInstanceSceneData(GPUSceneInstanceIndex, Scene.GPUScene.InstanceDataSOAStride);
-
-	FLWCVector3 WorldPosition = LWCSubtract(LWCPromote(TranslatedWorldPosition), PrimaryView.PreViewTranslation);
-	float3 LocalPosition = LWCMultiply(WorldPosition, InstanceSceneData.WorldToLocal);
-	FLWCVector3 PrevWorldPosition = LWCMultiply(LocalPosition, InstanceSceneData.PrevLocalToWorld);
-	return LWCToFloat(LWCAdd(PrevWorldPosition, PrimaryView.PreViewTranslation));
-}
-
 RAY_TRACING_ENTRY_RAYGEN_OR_INLINE(HardwareRayTraceLightSamples)
 {
 	uint ThreadIndex = DispatchThreadIndex.x;
@@ -139,18 +128,9 @@ RAY_TRACING_ENTRY_RAYGEN_OR_INLINE(HardwareRayTraceLightSamples)
 
 		if (Payload.IsHit())
 		{
-			uint SceneInstanceIndex = Payload.GetGPUSceneInstanceIndex();
-			float3 TranslatedWorldPosition = Ray.Origin + Ray.Direction * Payload.HitT;
-			float3 WorldVelocity = TranslatedWorldPosition - CalcPrevTranslatedWorldPositionFromGPUSceneInstanceIndex(TranslatedWorldPosition, SceneInstanceIndex);
-
 			LightSample.bVisible = false;
 			RWLightSamples[SampleCoord] = PackLightSample(LightSample);
 		}
-		else
-		{
-			LightSample.bVisible = true;
-			RWLightSamples[SampleCoord] = PackLightSample(LightSample);
-		}
 
 		#if DEBUG_MODE
 		if (bDebug && DebugMode == DEBUG_MODE_VISUALIZE_TRACING)
diff --git a/Engine/Source/Runtime/Renderer/Private/StochasticShadows/StochasticShadows.cpp b/Engine/Source/Runtime/Renderer/Private/StochasticShadows/StochasticShadows.cpp
--- a/Engine/Source/Runtime/Renderer/Private/StochasticShadows/StochasticShadows.cpp
+++ b/Engine/Source/Runtime/Renderer/Private/StochasticShadows/StochasticShadows.cpp
@@ -39,8 +39,22 @@ static TAutoConsoleVariable<int32> CVarStochasticShadowsTemporal(
 	ECVF_Scalability | ECVF_RenderThreadSafe
 );
 
-static TAutoConsoleVariable<int32> CVarStochasticShadowsGuiding(
-	TEXT("r.StochasticShadows.Guiding"),
+static TAutoConsoleVariable<int32> CVarStochasticShadowsTemporalMaxFramesAccumulated(
+	TEXT("r.StochasticShadows.Temporal.MaxFramesAccumulated"),
+	8,
+	TEXT("Max history length when accumulating frames. Lower values have less ghosting, but more noise."),
+	ECVF_Scalability | ECVF_RenderThreadSafe
+);
+
+static TAutoConsoleVariable<float> CVarStochasticShadowsTemporalStdDevOffset(
+	TEXT("r.StochasticShadows.Temporal.StdDevOffset"),
+	0.1f,
+	TEXT("Increases standard deviation in neighborhood clamp. Higher values cause more ghosting, but allow smoother temporal accumulation."),
+	ECVF_Scalability | ECVF_RenderThreadSafe
+);
+
+static TAutoConsoleVariable<int32> CVarStochasticSamplingShadowEstimate(
+	TEXT("r.StochasticShadows.Sampling.ShadowEstimate"),
 	1,
 	TEXT("Whether to use shadow mask history for sample guiding."),
 	ECVF_Scalability | ECVF_RenderThreadSafe
@@ -421,7 +435,8 @@ class FCompositeShadowMaskTracesCS : public FGlobalShader
 	class FNumSamplesPerPixel : SHADER_PERMUTATION_SPARSE_INT("NUM_SAMPLES_PER_PIXEL", 1, 2, 4);
 	class FTemporalAccumulation : SHADER_PERMUTATION_BOOL("TEMPORAL_ACCUMULATION");
 	class FHashTable : SHADER_PERMUTATION_BOOL("HASH_TABLE");
-	using FPermutationDomain = TShaderPermutationDomain<FNumSamplesPerPixel, FTemporalAccumulation, FHashTable>;
+	class FDebugMode : SHADER_PERMUTATION_BOOL("DEBUG_MODE");
+	using FPermutationDomain = TShaderPermutationDomain<FNumSamplesPerPixel, FTemporalAccumulation, FHashTable, FDebugMode>;
 
 	static bool ShouldCompilePermutation(const FGlobalShaderPermutationParameters& Parameters)
 	{
@@ -708,6 +723,8 @@ void FDeferredShadingSceneRenderer::RenderStochasticShadows(FRDGBuilder& GraphBu
 		StochasticShadowsParameters.MinLightSampleWeight = CVarStochasticShadowsMinSampleWeight.GetValueOnRenderThread();
 		StochasticShadowsParameters.TileDataStride = TileDataStride;
 		StochasticShadowsParameters.DownsampledTileDataStride = DownsampledTileDataStride;
+		StochasticShadowsParameters.TemporalMaxFramesAccumulated = CVarStochasticShadowsTemporalMaxFramesAccumulated.GetValueOnRenderThread();
+		StochasticShadowsParameters.TemporalStdDevOffset = CVarStochasticShadowsTemporalStdDevOffset.GetValueOnRenderThread();
 		StochasticShadowsParameters.DebugMode = CVarStochasticShadowsDebug.GetValueOnRenderThread();
 
 		if (bDebug)
@@ -842,7 +859,7 @@ void FDeferredShadingSceneRenderer::RenderStochasticShadows(FRDGBuilder& GraphBu
 			PermutationVector.Set<FGenerateSamplesCS::FTileType>(TileType);
 			PermutationVector.Set<FGenerateSamplesCS::FNumSamplesPerPixel>(NumSamplesPerPixel1d);
 			PermutationVector.Set<FGenerateSamplesCS::FShadowMaskReprojectionWeights>(bTemporal);
-			PermutationVector.Set<FGenerateSamplesCS::FShadowFactorEstimate>(bTemporal && CVarStochasticShadowsGuiding.GetValueOnRenderThread() != 0);
+			PermutationVector.Set<FGenerateSamplesCS::FShadowFactorEstimate>(bTemporal && CVarStochasticSamplingShadowEstimate.GetValueOnRenderThread() != 0);
 			PermutationVector.Set<FGenerateSamplesCS::FDebugMode>(bDebug);
 			PermutationVector.Set<FGenerateSamplesCS::FHashTable>(bUseHashTable);
 			PermutationVector.Set<FGenerateSamplesCS::FCandidateLightMask>(ShadowMaskAtlasHistory && CVarStochasticShadowsCandidateLightMask.GetValueOnRenderThread() != 0);
@@ -951,6 +968,7 @@ void FDeferredShadingSceneRenderer::RenderStochasticShadows(FRDGBuilder& GraphBu
 		PermutationVector.Set<FCompositeShadowMaskTracesCS::FNumSamplesPerPixel>(NumSamplesPerPixel1d);
 		PermutationVector.Set<FCompositeShadowMaskTracesCS::FTemporalAccumulation>(bTemporal);
 		PermutationVector.Set<FCompositeShadowMaskTracesCS::FHashTable>(bUseHashTable);
+		PermutationVector.Set<FCompositeShadowMaskTracesCS::FDebugMode>(bDebug);
 		auto ComputeShader = View.ShaderMap->GetShader<FCompositeShadowMaskTracesCS>(PermutationVector);
 
 		FComputeShaderUtils::AddPass(
diff --git a/Engine/Source/Runtime/Renderer/Private/StochasticShadows/StochasticShadowsInternal.h b/Engine/Source/Runtime/Renderer/Private/StochasticShadows/StochasticShadowsInternal.h
--- a/Engine/Source/Runtime/Renderer/Private/StochasticShadows/StochasticShadowsInternal.h
+++ b/Engine/Source/Runtime/Renderer/Private/StochasticShadows/StochasticShadowsInternal.h
@@ -28,6 +28,8 @@ BEGIN_SHADER_PARAMETER_STRUCT(FStochasticShadowsParameters, )
 	SHADER_PARAMETER(float, MinLightSampleWeight)
 	SHADER_PARAMETER(int32, TileDataStride)
 	SHADER_PARAMETER(int32, DownsampledTileDataStride)
+	SHADER_PARAMETER(int32, TemporalMaxFramesAccumulated)
+	SHADER_PARAMETER(float, TemporalStdDevOffset)
 	SHADER_PARAMETER(int32, DebugMode)
 	SHADER_PARAMETER_RDG_TEXTURE(Texture2D<float>, DownsampledSceneDepth)
 END_SHADER_PARAMETER_STRUCT()
