Fix issue in DF shadows when using r.HeightfieldShadowing
- tiles without intersecting objects wouldn't respect HAS_PREVIOUS_OUTPUT which caused result from previous pass to be discarded.
- also extended HAS_PREVIOUS_OUTPUT to work with all DISTANCEFIELD_PRIMITIVE_TYPE (it was only supported by DFPT_HeightField).

diff --git a/Engine/Shaders/Private/DistanceFieldShadowing.usf b/Engine/Shaders/Private/DistanceFieldShadowing.usf
--- a/Engine/Shaders/Private/DistanceFieldShadowing.usf
+++ b/Engine/Shaders/Private/DistanceFieldShadowing.usf
@@ -653,15 +653,6 @@ void DistanceFieldShadowingCS(
 		// Mitigate self-shadow on steep terrain surfces
 		float SelfShadowViewBias = SELF_SHADOW_VIEW_BIAS;
 
-#if HAS_PREVIOUS_OUTPUT
-	#if PLATFORM_SUPPORTS_TYPED_UAV_LOAD
-		float PrevResult = RWShadowFactors[DispatchThreadId.xy].x;
-	#else
-		float2 PrevResultUV = (DispatchThreadId.xy + 0.5) * InvOutputBufferSize;
-		float PrevResult = Texture2DSampleLevel(ShadowFactorsTexture, ShadowFactorsSampler, PrevResultUV, 0).x;
-	#endif
-#endif
-
 		OpaqueTranslatedWorldPosition.z += SelfShadowVerticalBias;
 		OpaqueTranslatedWorldPosition -= SelfShadowViewBias * View.ViewForward;
 
@@ -714,10 +705,6 @@ void DistanceFieldShadowingCS(
 				NumIntersectingObjects,
 				bUseCulling,
 				bUseScatterTileCulling);
-
-#if HAS_PREVIOUS_OUTPUT
-			Result = min(Result, PrevResult);
-#endif
 		}
 		else
 		{
@@ -760,6 +747,21 @@ void DistanceFieldShadowingCS(
 				/*bExpandSurface*/ false);
 		}
 	}
+
+#if HAS_PREVIOUS_OUTPUT
+	if (bShouldComputeShadowing)
+	{
+#		if PLATFORM_SUPPORTS_TYPED_UAV_LOAD
+			float PrevResult = RWShadowFactors[DispatchThreadId.xy].x;
+#		else
+			float2 PrevResultUV = (DispatchThreadId.xy + 0.5) * InvOutputBufferSize;
+			float PrevResult = Texture2DSampleLevel(ShadowFactorsTexture, ShadowFactorsSampler, PrevResultUV, 0).x;
+#		endif
+
+		Result = min(Result, PrevResult);
+	}
+#endif
+
 #else
 	//Result = bShouldComputeShadowing;
 	Result = bShouldComputeShadowing ? NumIntersectingObjects / 256.0f : 0.0f;
