Add DoubleFloat math library. This is a floating point format for performing double-precision math on GPU, that offers more precision than the tile-offset system.

diff --git a/Engine/Shaders/Private/DoubleFloat.ush b/Engine/Shaders/Private/DoubleFloat.ush
new file mode 100644
--- /dev/null
+++ b/Engine/Shaders/Private/DoubleFloat.ush
@@ -0,0 +1,503 @@
+// Copyright Epic Games, Inc. All Rights Reserved.
+
+#pragma once
+
+// A high-precision floating point type, consisting of two 32-bit floats (High & Low).
+// 'High' stores the core value, 'Low' stores the residual error that couldn't fit in 'High'.
+// This combination has 2x23=46 significant bits, providing twice the precision a float offers.
+// Operations are slower than floats, but faster than doubles on consumer GPUs (with potentially greater support)
+// Platforms that don't support fused multiply-add and INVARIANT may see reduced precision.
+// 
+// Based on:
+//  [0] Thall, A. (2006), Extended-precision floating-point numbers for GPU computation.
+//  [1] Mioara Maria Joldes, Jean-Michel Muller, Valentina Popescu. (2017), Tight and rigourous error bounds for basic building blocks of double-word arithmetic.
+//  [2] Vincent Lefevre, Nicolas Louvet, Jean-Michel Muller, Joris Picot, and Laurence Rideau. (2022), Accurate Calculation of Euclidean Norms using Double-Word Arithmetic
+//  [3] Jean-Michel Muller and Laurence Rideau. (2022), Formalization of Double-Word Arithmetic, and Comments on "Tight and Rigorous Error Bounds for Basic Building Blocks of Double-Word Arithmetic"
+
+#ifndef UE_DF_FORCE_FP32_OPS
+#define UE_DF_FORCE_FP32_OPS 0
+#endif
+
+#ifndef UE_DF_NO_FAST_MATH
+#define UE_DF_NO_FAST_MATH 0
+#endif
+
+/* FDFScalar */
+
+#define FDFType FDFScalar
+#define FDFTypeDeriv FDFScalarDeriv
+#define FFloatType float
+#define FBoolType bool
+#define GNumComponents 1
+#define DFConstructor MakeDFScalar
+#if COMPILER_SUPPORTS_HLSL2021
+struct FDFType;
+FDFType DFAdd(FDFType Lhs, FDFType Rhs);
+FDFType DFAdd(FDFType Lhs, FFloatType Rhs);
+FDFType DFSubtract(FDFType Lhs, FDFType Rhs);
+FDFType DFSubtract(FDFType Lhs, FFloatType Rhs);
+FDFType DFMultiply(FDFType Lhs, FDFType Rhs);
+FDFType DFMultiply(FDFType Lhs, FFloatType Rhs);
+FDFType DFDivide(FDFType Lhs, FDFType Rhs);
+FDFType DFDivide(FDFType Lhs, FFloatType Rhs);
+FDFType DFNegate(FDFType V);
+FBoolType DFGreater(FDFType Lhs, FFloatType Rhs);
+FBoolType DFLess(FDFType Lhs, FFloatType Rhs);
+FFloatType DFDemote(FDFType V);
+#endif
+
+struct FDFScalar
+{
+	float High;
+	float Low;
+
+#if COMPILER_SUPPORTS_HLSL2021
+	/*FDFType operator+(FDFType Rhs) { return DFAdd(this, Rhs); }
+	FDFType operator+(FFloatType Rhs) { return DFAdd(this, Rhs); }
+	FDFType operator-(FDFType Rhs) { return DFSubtract(this, Rhs); }
+	FDFType operator-(FFloatType Rhs) { return DFSubtract(this, Rhs); }
+	FDFType operator*(FDFType Rhs) { return DFMultiply(this, Rhs); }
+	FDFType operator*(FFloatType Rhs) { return DFMultiply(this, Rhs); }
+	FDFType operator/(FDFType Rhs) { return DFDivide(this, Rhs); }
+	FDFType operator/(FFloatType Rhs) { return DFDivide(this, Rhs); }*/
+	FDFType operator-() { return DFNegate(this); }
+	FBoolType operator>(FFloatType Rhs) { return DFGreater(this, Rhs); }
+	FBoolType operator<(FFloatType Rhs) { return DFLess(this, Rhs); }
+
+	FFloatType Demote() { return DFDemote(this); }
+#endif
+};
+
+#include "DoubleFloatOperations.ush"
+#undef FDFType
+#undef FDFTypeDeriv
+#undef FFloatType
+#undef FBoolType
+#undef GNumComponents
+#undef DFConstructor
+
+FDFScalar MakeDFVector(FDFScalar X) { return X; }
+
+FDFScalar DFGetComponent(float V, int C) { return MakeDFScalar(V, 0); }
+FDFScalar DFGetComponent(FDFScalar V, int C) { return V; }
+
+float DFLength(float V) { return abs(V); }
+float DFNormalize(float V) { return (float)sign(V); }
+FDFScalar DFLength(FDFScalar V) { return DFAbs(V); }
+FDFScalar DFNormalize(FDFScalar V) { return DFPromote(V); }
+float DFNormalizeDemote(FDFScalar V) { return DFSign(V); }
+
+/* FDFVector2 */
+
+#define FDFType FDFVector2
+#define FDFTypeDeriv FDFVector2Deriv
+#define FFloatType float2
+#define FBoolType bool2
+#define GNumComponents 2
+#define DFConstructor MakeDFVector2
+#define DFBroadcastFunc DFBroadcast2
+#include "DoubleFloatVectorDefinition.ush"
+#undef FDFType
+#undef FDFTypeDeriv
+#undef FFloatType
+#undef FBoolType
+#undef GNumComponents
+#undef DFConstructor
+#undef DFBroadcastFunc
+
+FDFVector2 MakeDFVector(FDFScalar X, FDFScalar Y) { return MakeDFVector2(float2(X.High, Y.High), float2(X.Low, Y.Low)); }
+
+/* FDFVector3 */
+
+#define FDFType FDFVector3
+#define FDFTypeDeriv FDFVector3Deriv
+#define FFloatType float3
+#define FBoolType bool3
+#define GNumComponents 3
+#define DFConstructor MakeDFVector3
+#define DFBroadcastFunc DFBroadcast3
+#include "DoubleFloatVectorDefinition.ush"
+#undef FDFType
+#undef FDFTypeDeriv
+#undef FFloatType
+#undef FBoolType
+#undef GNumComponents
+#undef DFConstructor
+#undef DFBroadcastFunc
+
+FDFVector3 MakeDFVector(FDFScalar   X, FDFScalar   Y, FDFScalar Z) { return MakeDFVector3(float3( X.High,  Y.High, Z.High), float3( X.Low,  Y.Low, Z.Low)); }
+FDFVector3 MakeDFVector(FDFScalar   X, FDFVector2 YZ)              { return MakeDFVector3(float3( X.High, YZ.High),         float3( X.Low, YZ.Low)); }
+FDFVector3 MakeDFVector(FDFVector2 XY, FDFScalar   Z)              { return MakeDFVector3(float3(XY.High,  Z.High),         float3(XY.Low,  Z.Low)); }
+
+/* FDFVector4 */
+
+#define FDFType FDFVector4
+#define FDFTypeDeriv FDFVector4Deriv
+#define FFloatType float4
+#define FBoolType bool4
+#define GNumComponents 4
+#define DFConstructor MakeDFVector4
+#define DFBroadcastFunc DFBroadcast4
+#include "DoubleFloatVectorDefinition.ush"
+#undef FDFType
+#undef FDFTypeDeriv
+#undef FFloatType
+#undef FBoolType
+#undef GNumComponents
+#undef DFConstructor
+#undef DFBroadcastFunc
+
+FDFVector4 MakeDFVector(FDFScalar    X, FDFScalar    Y, FDFScalar   Z, FDFScalar W) { return MakeDFVector4(float4(  X.High,   Y.High,  Z.High, W.High), float4(  X.Low,   Y.Low,  Z.Low, W.Low)); }
+FDFVector4 MakeDFVector(FDFScalar    X, FDFScalar    Y, FDFVector2 ZW)              { return MakeDFVector4(float4(  X.High,   Y.High, ZW.High),         float4(  X.Low,   Y.Low, ZW.Low)); }
+FDFVector4 MakeDFVector(FDFScalar    X, FDFVector2  YZ, FDFScalar   W)              { return MakeDFVector4(float4(  X.High,  YZ.High,  W.High),         float4(  X.Low,  YZ.Low,  W.Low)); }
+FDFVector4 MakeDFVector(FDFVector2  XY, FDFScalar    Z, FDFScalar   W)              { return MakeDFVector4(float4( XY.High,   Z.High,  W.High),         float4( XY.Low,   Z.Low,  W.Low)); }
+FDFVector4 MakeDFVector(FDFVector2  XY, FDFVector2  ZW)                             { return MakeDFVector4(float4( XY.High,  ZW.High),                  float4( XY.Low,  ZW.Low)); }
+FDFVector4 MakeDFVector(FDFScalar    X, FDFVector3 YZW)                             { return MakeDFVector4(float4(  X.High, YZW.High),                  float4(  X.Low, YZW.Low)); }
+FDFVector4 MakeDFVector(FDFVector3 XYZ, FDFScalar    W)                             { return MakeDFVector4(float4(XYZ.High,   W.High),                  float4(XYZ.Low,   W.Low)); }
+
+FDFVector4 MakeDFVector(FDFVector3 In, float W) { return MakeDFVector4(float4(In.High, W), float4(In.Low, 0)); }
+
+
+#define DFGetX(V) DFGetComponent(V, 0)
+#define DFGetY(V) DFGetComponent(V, 1)
+#define DFGetZ(V) DFGetComponent(V, 2)
+#define DFGetW(V) DFGetComponent(V, 3)
+
+FDFVector3 DFGetXYZ(FDFVector4 V)
+{
+	return MakeDFVector3(V.High.xyz, V.Low.xyz);
+}
+
+// Put the components of V at the specified indices in a new DF vector
+
+FDFScalar DFSwizzle(FDFScalar V, int C0)  { return V; }
+FDFScalar DFSwizzle(FDFVector2 V, int C0) { return DFGetComponent(V, C0); }
+FDFScalar DFSwizzle(FDFVector3 V, int C0) { return DFGetComponent(V, C0); }
+FDFScalar DFSwizzle(FDFVector4 V, int C0) { return DFGetComponent(V, C0); }
+
+FDFVector2 DFSwizzle(FDFScalar V, int C0, int C1)  { return MakeDFVector(V, V); }
+FDFVector2 DFSwizzle(FDFVector2 V, int C0, int C1) { return MakeDFVector(DFGetComponent(V, C0), DFGetComponent(V, C1)); }
+FDFVector2 DFSwizzle(FDFVector3 V, int C0, int C1) { return MakeDFVector(DFGetComponent(V, C0), DFGetComponent(V, C1)); }
+FDFVector2 DFSwizzle(FDFVector4 V, int C0, int C1) { return MakeDFVector(DFGetComponent(V, C0), DFGetComponent(V, C1)); }
+
+FDFVector3 DFSwizzle(FDFScalar V, int C0, int C1, int C2)  { return MakeDFVector(V, V, V); }
+FDFVector3 DFSwizzle(FDFVector2 V, int C0, int C1, int C2) { return MakeDFVector(DFGetComponent(V, C0), DFGetComponent(V, C1), DFGetComponent(V, C2)); }
+FDFVector3 DFSwizzle(FDFVector3 V, int C0, int C1, int C2) { return MakeDFVector(DFGetComponent(V, C0), DFGetComponent(V, C1), DFGetComponent(V, C2)); }
+FDFVector3 DFSwizzle(FDFVector4 V, int C0, int C1, int C2) { return MakeDFVector(DFGetComponent(V, C0), DFGetComponent(V, C1), DFGetComponent(V, C2)); }
+
+FDFVector4 DFSwizzle(FDFScalar V, int C0, int C1, int C2, int C3)  { return MakeDFVector(V, V, V, V); }
+FDFVector4 DFSwizzle(FDFVector2 V, int C0, int C1, int C2, int C3) { return MakeDFVector(DFGetComponent(V, C0), DFGetComponent(V, C1), DFGetComponent(V, C2), DFGetComponent(V, C3)); }
+FDFVector4 DFSwizzle(FDFVector3 V, int C0, int C1, int C2, int C3) { return MakeDFVector(DFGetComponent(V, C0), DFGetComponent(V, C1), DFGetComponent(V, C2), DFGetComponent(V, C3)); }
+FDFVector4 DFSwizzle(FDFVector4 V, int C0, int C1, int C2, int C3) { return MakeDFVector(DFGetComponent(V, C0), DFGetComponent(V, C1), DFGetComponent(V, C2), DFGetComponent(V, C3)); }
+
+
+#define DFHackToFloat DFDemote
+
+// Transforms *to* absolute world space
+struct FDFMatrix
+{
+	float4x4 M;
+	float3 PostTranslation; // Added to result, *after* multiplying 'M'
+};
+
+float4x4 Make4x3Matrix(float4x4 M)
+{
+	// Explicitly ignore the 4th column of the input
+	float4x4 Result;
+	Result[0] = float4(M[0].xyz, 0.0f);
+	Result[1] = float4(M[1].xyz, 0.0f);
+	Result[2] = float4(M[2].xyz, 0.0f);
+	Result[3] = float4(M[3].xyz, 1.0f);
+	return Result;
+}
+
+float4x4 MakeTranslationMatrix(float3 Offset)
+{
+	float4x4 Result;
+	Result[0] = float4(1.0f, 0.0f, 0.0f, 0.0f);
+	Result[1] = float4(0.0f, 1.0f, 0.0f, 0.0f);
+	Result[2] = float4(0.0f, 0.0f, 1.0f, 0.0f);
+	Result[3] = float4(Offset, 1.0f);
+	return Result;
+}
+
+FDFMatrix MakeDFMatrix4x3(float3 PostTranslation, float4x4 InMatrix)
+{
+	FDFMatrix Result;
+	Result.PostTranslation = PostTranslation;
+	Result.M = Make4x3Matrix(InMatrix);
+	return Result;
+}
+
+FDFMatrix MakeDFMatrix(float3 PostTranslation, float4x4 InMatrix)
+{
+	FDFMatrix Result;
+	Result.PostTranslation = PostTranslation;
+	Result.M = InMatrix;
+	return Result;
+}
+
+FDFVector3 DFGetOrigin(FDFMatrix InMatrix)
+{
+#if UE_DF_FORCE_FP32_OPS
+	return DFPromote(InMatrix.PostTranslation + InMatrix.M[3].xyz);
+#elif UE_DF_NO_FAST_MATH
+	return DFTwoSum(InMatrix.PostTranslation, InMatrix.M[3].xyz);
+#else
+	return DFFastTwoSum(InMatrix.PostTranslation, InMatrix.M[3].xyz);
+#endif
+}
+
+// Transforms *from* absolute world space
+struct FDFInverseMatrix
+{
+	float4x4 M;
+	float3 PreTranslation; // Subtracted from input position *before* multiplying 'M'
+#if !COMPILER_SUPPORTS_HLSL2021
+	int Dummy; // avoid problem with HLSL overloading/implicit cast (should be DCE'd)
+#endif
+};
+
+FDFInverseMatrix MakeDFInverseMatrix(float3 PreTranslation, float4x4 InMatrix)
+{
+	FDFInverseMatrix Result;
+	Result.PreTranslation = PreTranslation;
+	Result.M = InMatrix;
+#if !COMPILER_SUPPORTS_HLSL2021
+	Result.Dummy = 0;
+#endif
+	return Result;
+}
+
+FDFInverseMatrix MakeDFInverseMatrix4x3(float3 PreTranslation, float4x4 InMatrix)
+{
+	FDFInverseMatrix Result;
+	Result.PreTranslation = PreTranslation;
+	Result.M = Make4x3Matrix(InMatrix);
+#if !COMPILER_SUPPORTS_HLSL2021
+	Result.Dummy = 0;
+#endif
+	return Result;
+}
+
+FDFVector3 DFMultiply(float3 Vector, FDFMatrix InMatrix)
+{
+	float4 Result = mul(float4(Vector, 1), InMatrix.M);
+	//Result /= Result.w;
+	FDFVector3 V = DFTwoSum(InMatrix.PostTranslation, Result.xyz);
+	return V;
+}
+
+FDFVector4 DFMultiply(float4 Vector, FDFMatrix InMatrix)
+{
+	float4 Result = mul(Vector, InMatrix.M);
+	FDFVector4 V = DFTwoSum(float4(InMatrix.PostTranslation, 0), Result);
+	return V;
+}
+
+float3 DFMultiplyDemote(FDFVector3 Vector, FDFInverseMatrix InMatrix)
+{
+	float3 V = DFFastSubtractDemote(Vector, InMatrix.PreTranslation);
+	float4 Result = mul(float4(V, 1), InMatrix.M);
+	//Result /= Result.w;
+	return Result.xyz;
+}
+
+float4 DFMultiplyDemote(FDFVector4 Vector, FDFInverseMatrix InMatrix)
+{
+	float4 V = DFFastSubtractDemote(Vector, float4(InMatrix.PreTranslation, 0));
+	return mul(V, InMatrix.M);
+}
+
+FDFInverseMatrix DFMultiply(FDFInverseMatrix Lhs, float4x4 Rhs)
+{
+	float4x4 M = mul(Lhs.M, Rhs);
+	float3 PreTranslation = Lhs.PreTranslation;
+	return MakeDFInverseMatrix(PreTranslation, M);
+}
+FDFMatrix DFMultiply(float4x4 Lhs, FDFMatrix Rhs)
+{
+	float4x4 M = mul(Lhs, Rhs.M);
+	float3 PostTranslation = Rhs.PostTranslation;
+	return MakeDFMatrix(PostTranslation, M);
+}
+
+float4x4 MultiplyTranslation(float4x4 M, float3 Translation)
+{
+	float4x4 Result = M;
+	// M * translate(PostTranslation)
+	Result[3].xyz += Translation;
+	return Result;
+}
+
+float4x4 MultiplyTranslation(float3 Translation, float4x4 M)
+{
+	return mul(MakeTranslationMatrix(Translation), M);
+}
+
+// Perform Lhs * T, where T is a translation matrix with the specified vector
+float4x4 DFMultiplyTranslationDemote(FDFMatrix Lhs, FDFVector3 Rhs)
+{
+	float3 Translation = DFFastAddDemote(Lhs.PostTranslation, Rhs);
+	float4x4 Result = MultiplyTranslation(Lhs.M, Translation);
+	return Result;
+}
+// Same as DFMultiplyTranslationDemote, but assumes the translation vector
+// roughly cancels out the translation in the matrix.
+// More formally, Lhs.PostTranslation and -Rhs should meet the requirements of DFFastLocalSubtractDemote
+float4x4 DFFastMultiplyTranslationDemote(FDFMatrix Lhs, FDFVector3 Rhs)
+{
+#if UE_DF_NO_FAST_MATH
+	return DFMultiplyTranslationDemote(Lhs, Rhs);
+#else
+	float3 Translation = DFFastLocalSubtractDemote(Lhs.PostTranslation, DFNegate(Rhs));
+	float4x4 Result = MultiplyTranslation(Lhs.M, Translation);
+	return Result;
+#endif
+}
+
+float4x4 DFMultiplyTranslationDemote(FDFMatrix Lhs, float3 Rhs)
+{
+	float3 Translation = Lhs.PostTranslation + Rhs; 
+	float4x4 Result = MultiplyTranslation(Lhs.M, Translation);
+	return Result;
+}
+FDFMatrix DFMultiplyTranslation(FDFMatrix Lhs, float3 Rhs)
+{
+	FDFVector3 Translation = DFTwoSum(Lhs.PostTranslation, Rhs);
+
+	float3 PostTranslation = Translation.High;
+#if UE_DF_FORCE_FP32_OPS
+	float4x4 M = Lhs.M;
+#else
+	float4x4 M = MultiplyTranslation(Lhs.M, Translation.Low);
+#endif
+	return MakeDFMatrix(PostTranslation, M);
+}
+FDFMatrix DFMultiplyTranslation(float4x4 Lhs, FDFVector3 Rhs)
+{
+	FDFVector3 Translation = Rhs;
+
+	float3 PostTranslation = Translation.High;
+#if UE_DF_FORCE_FP32_OPS
+	float4x4 M = Lhs;
+#else
+	float4x4 M = MultiplyTranslation(Lhs, Translation.Low);
+#endif
+	return MakeDFMatrix(PostTranslation, M);
+}
+
+// Perform Lhs * T, where T is a translation matrix with the specified vector
+FDFInverseMatrix DFMultiplyTranslation(FDFVector3 Lhs, FDFInverseMatrix Rhs)
+{
+	FDFVector3 Translation = DFFastAdd(Lhs, -Rhs.PreTranslation); 
+
+	float3 PreTranslation = -Translation.High;
+#if UE_DF_FORCE_FP32_OPS
+	float4x4 M = Rhs.M;
+#else
+	float4x4 M = MultiplyTranslation(Translation.Low, Rhs.M);
+#endif
+	return MakeDFInverseMatrix(PreTranslation, M);
+}
+
+float4x4 DFMultiplyTranslationDemote(FDFVector3 Lhs, FDFInverseMatrix Rhs)
+{
+	float3 Translation = DFFastAddDemote(Lhs, -Rhs.PreTranslation);
+	float4x4 Result = MultiplyTranslation(Translation, Rhs.M);
+	return Result;
+}
+
+// Same as DFMultiplyTranslationDemote, but assumes the translation vector
+// roughly cancels out the translation in the matrix.
+// More formally, Lhs and Rhs.PreTranslation should meet the requirements of DFFastLocalSubtractDemote
+float4x4 DFFastMultiplyTranslationDemote(FDFVector3 Lhs, FDFInverseMatrix Rhs)
+{
+#if UE_DF_NO_FAST_MATH
+	return DFMultiplyTranslationDemote(Lhs, Rhs);
+#else
+	float3 Translation = DFFastLocalSubtractDemote(Lhs, Rhs.PreTranslation);
+	float4x4 Result = MultiplyTranslation(Translation, Rhs.M);
+	return Result;
+#endif
+}
+
+FDFInverseMatrix DFMultiplyTranslation(float3 Lhs, FDFInverseMatrix Rhs)
+{
+	FDFVector3 Translation = DFTwoSum(Lhs, -Rhs.PreTranslation);
+
+	float3 PreTranslation = -Translation.High;
+#if UE_DF_FORCE_FP32_OPS
+	float4x4 M = Rhs.M;
+#else
+	float4x4 M = MultiplyTranslation(Translation.Low, Rhs.M);
+#endif
+	return MakeDFInverseMatrix(PreTranslation, M);
+}
+
+FDFInverseMatrix DFMultiplyTranslation(FDFVector3 Lhs, float4x4 Rhs)
+{
+	float3 PreTranslation = -Lhs.High;
+#if UE_DF_FORCE_FP32_OPS
+	float4x4 M = Rhs;
+#else
+	float4x4 M = MultiplyTranslation(Lhs.Low, Rhs);
+#endif
+	return MakeDFInverseMatrix(PreTranslation, M);
+}
+
+float3x3 DFToFloat3x3(FDFMatrix Value)
+{
+	return (float3x3)Value.M;
+}
+
+float3x3 DFToFloat3x3(FDFInverseMatrix Value)
+{
+	return (float3x3)Value.M;
+}
+
+float3 DFMultiplyVector(float3 Vector, FDFMatrix InMatrix)
+{
+	return mul(Vector, (float3x3)InMatrix.M);
+}
+
+float3 DFMultiplyVector(float3 Vector, FDFInverseMatrix InMatrix)
+{
+	return mul(Vector, (float3x3)InMatrix.M);
+}
+
+FDFMatrix DFPromote(FDFMatrix Value) { return Value; }
+FDFInverseMatrix DFPromoteInverse(FDFInverseMatrix Value) { return Value; }
+FDFMatrix  DFPromote(float4x4 Value) { return MakeDFMatrix((float3)0, Value); }
+FDFInverseMatrix DFPromoteInverse(float4x4 Value) { return MakeDFInverseMatrix((float3)0, Value); }
+
+float4x4 DFDemote(FDFMatrix V)
+{
+	return MultiplyTranslation(V.M, V.PostTranslation);
+}
+
+float4x4 DFDemote(FDFInverseMatrix V)
+{
+	return MultiplyTranslation(-V.PreTranslation, V.M);
+}
+
+float4 DFTransformLocalToTranslatedWorld(float3 LocalPosition, FDFMatrix LocalToWorld, FDFVector3 PreViewTranslation)
+{
+	// Multiply local position with LocalToWorld 3x3 first for high precision rotation and scale, then move the world position part to translated world before adding it  
+	float3 RotatedScaledPosition = mad(LocalPosition.xxx, LocalToWorld.M[0].xyz, mad(LocalPosition.yyy, LocalToWorld.M[1].xyz, LocalPosition.zzz * LocalToWorld.M[2].xyz));
+
+#if UE_DF_NO_FAST_MATH
+	float3 TranslatedWorldPositionOrigin = DFAddDemote(DFGetOrigin(LocalToWorld), ResolvedView.PreViewTranslation);
+#else
+	float3 TranslatedWorldPositionOrigin = DFFastLocalSubtractDemote(LocalToWorld.PostTranslation, DFNegate(PreViewTranslation)) + LocalToWorld.M[3].xyz;
+#endif
+	
+	return float4(RotatedScaledPosition + TranslatedWorldPositionOrigin, 1.0f);
+}
+
+FDFVector3 DFMultiplyLHSAndFastTwoSum(float3 Lhs, float Factor, float3 Rhs)
+{
+	const float3 S = mad(Lhs, Factor, Rhs);
+	const float3 E = INVARIANT(Rhs - mad(-Lhs, Factor, S));
+	return MakeDFVector3(S, E);
+}
diff --git a/Engine/Shaders/Private/DoubleFloatOperations.ush b/Engine/Shaders/Private/DoubleFloatOperations.ush
new file mode 100644
--- /dev/null
+++ b/Engine/Shaders/Private/DoubleFloatOperations.ush
@@ -0,0 +1,965 @@
+// Copyright Epic Games, Inc. All Rights Reserved.
+
+// This constructs a FDFScalar, but does not 'rebalance' the input floats.
+// Use DFTwoSum instead if the input is not already in the correct high/low format.
+FDFType DFConstructor(FFloatType High, FFloatType Low)
+{
+	FDFType Result;
+	Result.High = High;
+	Result.Low = Low;
+	return Result;
+}
+
+FDFType DFPromote(FDFType Value) { return Value; }
+
+// Convert to double float
+FDFType DFPromote(FFloatType Value) { return DFConstructor(Value, (FFloatType)0); }
+
+// Truncate to single precision float
+// (This could just return Value.High if you're sure the input is in the correct DF format)
+FFloatType DFDemote(FDFType Value) { return Value.High + Value.Low; }
+
+FFloatType DFDemote(FFloatType Value) { return Value; }
+
+FDFType DFNegate(FDFType Value)
+{
+	return DFConstructor(-Value.High, -Value.Low);
+}
+
+FFloatType DFSign(FDFType Value)
+{
+	return (FFloatType)sign(Value.High);
+}
+
+FDFType DFAbs(FDFType Value)
+{
+	const FFloatType Sign = DFSign(Value);
+	return DFConstructor(Sign * Value.High, Sign * Value.Low);
+}
+
+/**************/
+/** ADDITION **/
+/**************/
+
+// Sum up two single precision vectors to one double precision vector.
+// [1] Algorithm 2
+FDFType DFTwoSum(FFloatType Lhs, FFloatType Rhs)
+{
+#if UE_DF_FORCE_FP32_OPS
+	return DFPromote(Lhs + Rhs);
+#else
+	const FFloatType S = Lhs + Rhs;
+	const FFloatType V = INVARIANT(S - Lhs);
+	const FFloatType Q = INVARIANT(S - V);
+	const FFloatType R = INVARIANT(Lhs - Q);
+	const FFloatType T = INVARIANT(Rhs - V);
+	return DFConstructor(S, R + T);
+#endif
+}
+
+// An optimized version of DFTwoSum, under the assumption that
+// a = 0 or b = 0, or e1 >= e2, where e1 and e2 are the exponents of a and b respectively.
+// [1] Algorithm 1
+// The worst case precision if this assumption is violated, is just regular fp32 precision
+FDFType DFFastTwoSum(FFloatType Lhs, FFloatType Rhs)
+{
+#if UE_DF_FORCE_FP32_OPS
+	return DFPromote(Lhs + Rhs);
+#else
+	const FFloatType S = Lhs + Rhs;
+	const FFloatType E = INVARIANT(Rhs - (S - Lhs));
+	return DFConstructor(S, E);
+#endif
+}
+
+// Alias for DFTwoSum
+FDFType DFAdd(FFloatType Lhs, FFloatType Rhs) { return DFTwoSum(Lhs, Rhs); }
+FDFType DFFastAdd(FFloatType Lhs, FFloatType Rhs) { return DFTwoSum(Lhs, Rhs); }
+
+// [1] Algorithm 6
+FDFType DFAdd(FDFType Lhs, FDFType Rhs)
+{
+#if UE_DF_FORCE_FP32_OPS
+	return DFPromote(Lhs.High + Rhs.High);
+#else
+	// lossless add
+	FDFType S = DFTwoSum(Lhs.High, Rhs.High);
+	const FDFType T = DFTwoSum(Lhs.Low, Rhs.Low);
+	// merge and rebalance
+	S.Low += T.High;
+	S = DFFastTwoSum(S.High, S.Low);
+	S.Low += T.Low;
+	S = DFFastTwoSum(S.High, S.Low);
+	return S;
+#endif
+}
+
+// [1] Algorithm 4
+FDFType DFAdd(FDFType Lhs, FFloatType Rhs)
+{
+#if UE_DF_FORCE_FP32_OPS
+	return DFPromote(Lhs.High + Rhs);
+#else
+	const FDFType S = DFTwoSum(Lhs.High, Rhs);
+	const FFloatType T = Lhs.Low + S.Low;
+	return DFFastTwoSum(S.High, T);
+#endif
+}
+
+FDFType DFAdd(FFloatType Lhs, FDFType Rhs) { return DFAdd(Rhs, Lhs); }
+
+FFloatType DFAddDemote(FDFType Lhs, FDFType Rhs)
+{
+#if UE_DF_FORCE_FP32_OPS
+	return Lhs.High + Rhs.High;
+#else
+	// lossless add
+	FDFType S = DFTwoSum(Lhs.High, Rhs.High);
+	const FDFType T = DFTwoSum(Lhs.Low, Rhs.Low);
+	// merge and rebalance
+	S.Low += T.High;
+	S = DFFastTwoSum(S.High, S.Low);
+	S.Low += T.Low;
+	return S.High + S.Low;
+#endif
+}
+
+FFloatType DFAddDemote(FDFType Lhs, FFloatType Rhs)
+{
+#if UE_DF_FORCE_FP32_OPS
+	return Lhs.High + Rhs;
+#else
+	const FDFType S = DFTwoSum(Lhs.High, Rhs);
+	const FFloatType T = Lhs.Low + S.Low;
+	return S.High + T;
+#endif
+}
+
+FFloatType DFAddDemote(FFloatType Lhs, FDFType Rhs) { return DFAddDemote(Rhs, Lhs); }
+
+FDFType DFFastAdd(FDFType Lhs, FDFType Rhs)
+{
+#if UE_DF_FORCE_FP32_OPS
+	return DFPromote(Lhs.High + Rhs.High);
+#else
+	FDFType S = DFTwoSum(Lhs.High, Rhs.High);
+	S.Low += Lhs.Low + Rhs.Low;
+	S = DFFastTwoSum(S.High, S.Low);
+	return S;
+#endif
+}
+
+FDFType DFFastAdd(FDFType Lhs, FFloatType Rhs)
+{
+#if UE_DF_FORCE_FP32_OPS
+	return DFPromote(Lhs.High + Rhs);
+#else
+	FDFType S = DFTwoSum(Lhs.High, Rhs);
+	S.Low += Lhs.Low;
+	S = DFFastTwoSum(S.High, S.Low);
+	return S;
+#endif
+}
+
+FDFType DFFastAdd(FFloatType Lhs, FDFType Rhs) { return DFFastAdd(Rhs, Lhs); }
+
+FFloatType DFFastAddDemote(FDFType Lhs, FDFType Rhs)
+{
+#if UE_DF_FORCE_FP32_OPS
+	return Lhs.High + Rhs.High;
+#else
+	FDFType S = DFTwoSum(Lhs.High, Rhs.High);
+	S.Low += Lhs.Low + Rhs.Low;
+	return S.High + S.Low;
+#endif
+}
+
+FFloatType DFFastAddDemote(FDFType Lhs, FFloatType Rhs) { return DFAddDemote(Lhs, Rhs); }
+FFloatType DFFastAddDemote(FFloatType Lhs, FDFType Rhs) { return DFFastAddDemote(Rhs, Lhs); }
+
+/*****************/
+/** SUBTRACTION **/
+/*****************/
+
+FDFType DFSubtract(FFloatType Lhs, FFloatType Rhs) { return DFAdd(Lhs, -Rhs); }
+FDFType DFSubtract(FDFType Lhs, FDFType Rhs) { return DFAdd(Lhs, DFNegate(Rhs)); }
+FDFType DFSubtract(FDFType Lhs, FFloatType Rhs) { return DFAdd(Lhs, -Rhs); }
+FDFType DFSubtract(FFloatType Lhs, FDFType Rhs) { return DFAdd(Lhs, DFNegate(Rhs)); }
+
+FFloatType DFSubtractDemote(FDFType Lhs, FDFType Rhs) { return DFAddDemote(Lhs, DFNegate(Rhs)); }
+FFloatType DFSubtractDemote(FDFType Lhs, FFloatType Rhs) { return DFAddDemote(Lhs, -Rhs); }
+FFloatType DFSubtractDemote(FFloatType Lhs, FDFType Rhs) { return DFAddDemote(Lhs, DFNegate(Rhs)); }
+
+FDFType DFFastSubtract(FDFType Lhs, FDFType Rhs) { return DFFastAdd(Lhs, DFNegate(Rhs)); }
+FDFType DFFastSubtract(FDFType Lhs, FFloatType Rhs) { return DFFastAdd(Lhs, -Rhs); }
+FDFType DFFastSubtract(FFloatType Lhs, FDFType Rhs) { return DFFastAdd(Lhs, DFNegate(Rhs)); }
+
+FFloatType DFFastSubtractDemote(FDFType Lhs, FDFType Rhs) { return DFFastAddDemote(Lhs, DFNegate(Rhs)); }
+FFloatType DFFastSubtractDemote(FDFType Lhs, FFloatType Rhs) { return DFFastAddDemote(Lhs, -Rhs); }
+FFloatType DFFastSubtractDemote(FFloatType Lhs, FDFType Rhs) { return DFFastAddDemote(Lhs, DFNegate(Rhs)); }
+
+// Subtract using only 2 ops instead of ~20, and convert the result to fp32 (1 op)
+// CAREFUL, this has strict input preconditions to get output that has any precision at all.
+// Assuming you require a precision of 2^-4, the precondition for Lhs and Rhs is
+//   (abs(Lhs) < 2^20 && abs(Rhs) < 2^20) ||
+//   ( abs(Lhs) < 2^43 && abs(Rhs) < 2^43 && Sign(Lhs) == Sign(Rhs) && abs(trunc(log2(abs(Lhs))) - trunc(log2(abs(Rhs)))) < 1 )
+// As the preconditions imply, this does not work for addition due to FP rounding.
+FFloatType DFFastLocalSubtractDemote(FDFType Lhs, FDFType Rhs)
+{
+#if UE_DF_FORCE_FP32_OPS
+	return Lhs.High - Rhs.High;
+#elif UE_DF_NO_FAST_MATH
+	return DFSubtractDemote(Lhs, Rhs);
+#else
+	const FFloatType High = INVARIANT(Lhs.High - Rhs.High);
+	const FFloatType Low = INVARIANT(Lhs.Low - Rhs.Low);
+	return INVARIANT(High + Low);
+#endif
+}
+
+FFloatType DFFastLocalSubtractDemote(FDFType Lhs, FFloatType Rhs)
+{
+#if UE_DF_FORCE_FP32_OPS
+	return Lhs.High - Rhs;
+#elif UE_DF_NO_FAST_MATH
+	return DFSubtractDemote(Lhs, Rhs);
+#else
+	const FFloatType High = INVARIANT(Lhs.High - Rhs);
+	return INVARIANT(High + Lhs.Low);
+#endif
+}
+
+FFloatType DFFastLocalSubtractDemote(FFloatType Lhs, FDFType Rhs)
+{
+	return DFFastLocalSubtractDemote(DFPromote(Lhs), Rhs);
+}
+
+/********************/
+/** MULTIPLICATION **/
+/********************/
+
+// Multiply up two single precision scalars to one double precision scalar. (AKA Fast2Mult)
+// [1] Algorithm 3
+FDFType DFTwoProduct(FFloatType Lhs, FFloatType Rhs)
+{
+#if UE_DF_FORCE_FP32_OPS
+	return DFPromote(Lhs * Rhs);
+#else
+	const FFloatType P = Lhs * Rhs;
+	const FFloatType E = INVARIANT(mad(Lhs, Rhs, -P));
+	return DFConstructor(P, E);
+#endif
+}
+
+// Alias for DFTwoProduct
+FDFType DFMultiply(FFloatType Lhs, FFloatType Rhs) { return DFTwoProduct(Lhs, Rhs); }
+
+// [1] Algorithm 11
+FDFType DFMultiply(FDFType Lhs, FDFType Rhs)
+{
+#if UE_DF_FORCE_FP32_OPS
+	return DFPromote(Lhs.High * Rhs.High);
+#else
+	const FDFType P = DFTwoProduct(Lhs.High, Rhs.High);
+	FFloatType T = P.Low;
+	T = INVARIANT(mad(Lhs.High, Rhs.Low, T));
+	T = INVARIANT(mad(Lhs.Low, Rhs.High, T));
+
+	return DFFastTwoSum(P.High, T);
+#endif
+}
+
+// [1] Algorithm 9
+FDFType DFMultiply(FDFType Lhs, FFloatType Rhs)
+{
+#if UE_DF_FORCE_FP32_OPS
+	return DFPromote(Lhs.High * Rhs);
+#else
+	FDFType P = DFTwoProduct(Lhs.High, Rhs);
+	P.Low = mad(Lhs.Low, Rhs, P.Low);
+	return DFFastTwoSum(P.High, P.Low);
+#endif
+}
+
+FDFType DFMultiply(FFloatType Lhs, FDFType Rhs) { return DFMultiply(Rhs, Lhs); }
+
+FFloatType DFMultiplyDemote(FDFType Lhs, FDFType Rhs)
+{
+#if UE_DF_FORCE_FP32_OPS
+	return Lhs.High * Rhs.High;
+#else
+	const FDFType P = DFTwoProduct(Lhs.High, Rhs.High);
+	FFloatType T = P.Low;
+	T = INVARIANT(mad(Lhs.High, Rhs.Low, T));
+	T = INVARIANT(mad(Lhs.Low, Rhs.High, T));
+
+	return P.High + T;
+#endif
+}
+
+FFloatType DFMultiplyDemote(FDFType Lhs, FFloatType Rhs)
+{
+#if UE_DF_FORCE_FP32_OPS
+	return Lhs.High * Rhs;
+#else
+	FDFType P = DFTwoProduct(Lhs.High, Rhs);
+	P.Low = mad(Lhs.Low, Rhs, P.Low);
+	return P.High + P.Low;
+#endif
+}
+
+FFloatType DFMultiplyDemote(FFloatType Lhs, FDFType Rhs) { return DFMultiplyDemote(Rhs, Lhs); }
+
+// Fast multiplication that assumes the factor is a power of two.
+// If this assumption is valid, no precision is lost.
+FDFType DFMultiplyByPow2(FDFType Lhs, FFloatType Rhs)
+{
+#if UE_DF_FORCE_FP32_OPS
+	return DFPromote(Lhs.High * Rhs);
+#else
+	return DFConstructor(Lhs.High * Rhs, Lhs.Low * Rhs);
+#endif
+}
+
+FDFType DFSqr(FDFType V) { return DFMultiply(V, V); }
+FDFType DFSqr(FFloatType V) { return DFTwoProduct(V, V); }
+
+/**************/
+/** DIVISION **/
+/**************/
+
+// [1] Algorithm 18
+FDFType DFDivide(FDFType Lhs, FDFType Rhs)
+{
+#if UE_DF_FORCE_FP32_OPS
+	return DFPromote(Lhs.High / Rhs.High);
+#else
+	const FFloatType Th = INVARIANT(1.0 / Rhs.High);
+	const FFloatType Rh = INVARIANT(mad(-Rhs.High, Th, 1));
+	const FFloatType Rl = INVARIANT(-Rhs.Low * Th);
+	const FDFType E = DFFastTwoSum(Rh, Rl);
+	const FDFType D = DFMultiply(E, Th);
+	const FDFType M = DFAdd(D, Th);
+	return DFMultiply(Lhs, M);
+#endif
+}
+
+// [1] Algorithm 15
+FDFType DFDivide(FDFType Lhs, FFloatType Rhs)
+{
+#if UE_DF_FORCE_FP32_OPS
+	return DFPromote(Lhs.High / Rhs);
+#else
+	const FFloatType Th = INVARIANT(Lhs.High / Rhs);
+	const FDFType P = DFTwoProduct(Th, Rhs);
+	const FFloatType Dh = INVARIANT(Lhs.High - P.High);
+	const FFloatType Dt = INVARIANT(Dh - P.Low);
+	const FFloatType D = INVARIANT(Dt + Lhs.Low);
+	const FFloatType Tl = INVARIANT(D / Rhs);
+	return DFFastTwoSum(Th, Tl);
+#endif
+}
+
+FDFType DFDivide(FFloatType Lhs, FFloatType Rhs)
+{
+#if UE_DF_FORCE_FP32_OPS
+	return DFPromote(Lhs / Rhs);
+#else
+	return DFDivide(DFPromote(Lhs), Rhs);
+#endif
+}
+
+FDFType DFDivide(FFloatType Lhs, FDFType Rhs)
+{
+#if UE_DF_FORCE_FP32_OPS
+	return DFPromote(Lhs / Rhs.High);
+#else
+	return DFDivide(DFPromote(Lhs), Rhs);
+#endif
+}
+
+// Less precise than 18, but almost half the ops
+// [1] Algorithm 17
+FDFType DFFastDivide(FDFType Lhs, FDFType Rhs)
+{
+#if UE_DF_FORCE_FP32_OPS
+	return DFPromote(Lhs.High / Rhs.High);
+#elif UE_DF_NO_FAST_MATH
+	return DFDivide(Lhs, Rhs);
+#else
+	const FFloatType Th = INVARIANT(Lhs.High / Rhs.High);
+	const FDFType R = DFMultiply(Rhs, Th);
+	const FFloatType Ph = INVARIANT(Lhs.High - R.High);
+	const FFloatType Dl = INVARIANT(Lhs.Low - R.Low);
+	const FFloatType D = INVARIANT(Ph + Dl);
+	const FFloatType Tl = INVARIANT(D / Rhs.High);
+	return DFFastTwoSum(Th, Tl);
+#endif
+}
+
+FDFType DFFastDivide(FDFType Lhs, FFloatType Rhs)
+{
+#if UE_DF_FORCE_FP32_OPS
+	return DFPromote(Lhs.High / Rhs);
+#elif UE_DF_NO_FAST_MATH
+	return DFDivide(Lhs, Rhs);
+#else
+	const FFloatType Th = INVARIANT(Lhs.High / Rhs);
+	const FDFType R = DFTwoProduct(Rhs, Th);
+	const FFloatType Ph = INVARIANT(Lhs.High - R.High);
+	const FFloatType Dl = INVARIANT(Lhs.Low - R.Low);
+	const FFloatType D = INVARIANT(Ph + Dl);
+	const FFloatType Tl = INVARIANT(D / Rhs);
+	return DFFastTwoSum(Th, Tl);
+#endif
+}
+
+FDFType DFFastDivide(FFloatType Lhs, FDFType Rhs) { return DFFastDivide(DFPromote(Lhs), Rhs); }
+
+FFloatType DFFastDivideDemote(FDFType Lhs, FDFType Rhs)
+{
+#if UE_DF_FORCE_FP32_OPS
+	return Lhs.High / Rhs.High;
+#elif UE_DF_NO_FAST_MATH
+	return DFDemote(DFDivide(Lhs, Rhs));
+#else
+	const FFloatType Th = INVARIANT(Lhs.High / Rhs.High);
+	const FDFType R = DFMultiply(Rhs, Th);
+	const FFloatType Ph = INVARIANT(Lhs.High - R.High);
+	const FFloatType Dl = INVARIANT(Lhs.Low - R.Low);
+	const FFloatType D = INVARIANT(Ph + Dl);
+	const FFloatType Tl = INVARIANT(D / Rhs.High);
+	return Th + Tl;
+#endif
+}
+
+FFloatType DFFastDivideDemote(FDFType Lhs, FFloatType Rhs)
+{
+#if UE_DF_FORCE_FP32_OPS
+	return Lhs.High / Rhs;
+#elif UE_DF_NO_FAST_MATH
+	return DFDemote(DFDivide(Lhs, Rhs));
+#else
+	const FFloatType Th = INVARIANT(Lhs.High / Rhs);
+	const FDFType R = DFTwoProduct(Rhs, Th);
+	const FFloatType Ph = INVARIANT(Lhs.High - R.High);
+	const FFloatType Dl = INVARIANT(Lhs.Low - R.Low);
+	const FFloatType D = INVARIANT(Ph + Dl);
+	const FFloatType Tl = INVARIANT(D / Rhs);
+	return Th + Tl;
+#endif
+}
+
+FFloatType DFFastDivideDemote(FFloatType Lhs, FDFType Rhs) { return DFFastDivideDemote(DFPromote(Lhs), Rhs); }
+
+// Fast division that assumes the divisor is a power of two.
+// If this assumption is valid, no precision is lost.
+FDFType DFDivideByPow2(FDFType Lhs, FFloatType Rhs)
+{
+#if UE_DF_FORCE_FP32_OPS
+	return DFPromote(Lhs.High / Rhs);
+#else
+	return DFConstructor(Lhs.High / Rhs, Lhs.Low / Rhs);
+#endif
+}
+
+FDFType DFRcp(FDFType V) { return DFFastDivide((FFloatType)1.0, V); }
+FFloatType DFRcpDemote(FDFType V) { return rcp(V.High); }
+
+/****************/
+/** COMPARISON **/
+/****************/
+
+FBoolType DFEqualsApprox(FDFType Lhs, FDFType Rhs, float Threshold)
+{
+	return DFSubtractDemote(Lhs, Rhs) < Threshold;
+}
+FBoolType DFEqualsApprox(FDFType Lhs, FFloatType Rhs, float Threshold)
+{
+	return DFSubtractDemote(Lhs, Rhs) < Threshold;
+}
+FBoolType DFEqualsApprox(FFloatType Lhs, FDFType Rhs, float Threshold)
+{
+	return DFSubtractDemote(Lhs, Rhs) < Threshold;
+}
+
+FBoolType DFEquals(FDFType Lhs, FDFType Rhs)
+{
+#if UE_DF_FORCE_FP32_OPS
+	return Lhs.High == Rhs.High;
+#else
+#if COMPILER_SUPPORTS_HLSL2021
+	return and(Lhs.High == Rhs.High, Lhs.Low == Rhs.Low);
+#else
+	return Lhs.High == Rhs.High && Lhs.Low == Rhs.Low;
+#endif
+#endif
+}
+FBoolType DFEquals(FDFType Lhs, FFloatType Rhs)
+{
+	return DFEquals(Lhs, DFPromote(Rhs));
+}
+FBoolType DFEquals(FFloatType Lhs, FDFType Rhs)
+{
+	return DFEquals(DFPromote(Lhs), Rhs);
+}
+
+FDFType DFSelect(FBoolType S, FDFType Lhs, FDFType Rhs) { return DFConstructor(select(S, Lhs.High, Rhs.High), select(S, Lhs.Low, Rhs.Low)); }
+FDFType DFSelect(FBoolType S, FDFType Lhs, FFloatType Rhs) { return DFConstructor(select(S, Lhs.High, Rhs), select(S, Lhs.Low, (FFloatType)0)); }
+FDFType DFSelect(FBoolType S, FFloatType Lhs, FDFType Rhs) { return DFConstructor(select(S, Lhs, Rhs.High), select(S, (FFloatType)0, Rhs.Low)); }
+
+FBoolType DFGreater(FDFType Lhs, FDFType Rhs)
+{
+#if UE_DF_FORCE_FP32_OPS
+	return Lhs.High > Rhs.High;
+#else
+#if COMPILER_SUPPORTS_HLSL2021
+		return or(Lhs.High > Rhs.High, and(Lhs.High == Rhs.High, Lhs.Low > Rhs.Low));
+#else
+		return Lhs.High > Rhs.High || (Lhs.High == Rhs.High && Lhs.Low > Rhs.Low);
+#endif
+#endif
+}
+
+FBoolType DFLess(FDFType Lhs, FDFType Rhs)
+{
+#if UE_DF_FORCE_FP32_OPS
+	return Lhs.High < Rhs.High;
+#else
+#if COMPILER_SUPPORTS_HLSL2021
+	return or(Lhs.High < Rhs.High, and(Lhs.High == Rhs.High, Lhs.Low < Rhs.Low));
+#else
+	return Lhs.High < Rhs.High || (Lhs.High == Rhs.High && Lhs.Low < Rhs.Low);
+#endif
+#endif
+}
+
+FBoolType DFGreater(FDFType Lhs, FFloatType Rhs)
+{
+#if UE_DF_FORCE_FP32_OPS
+	return Lhs.High > Rhs;
+#else
+#if COMPILER_SUPPORTS_HLSL2021
+	return or(Lhs.High > Rhs, and(Lhs.High == Rhs, Lhs.Low > 0.0));
+#else
+	return Lhs.High > Rhs || (Lhs.High == Rhs && Lhs.Low > 0.0);
+#endif
+#endif
+}
+
+FBoolType DFLess(FDFType Lhs, FFloatType Rhs)
+{
+#if UE_DF_FORCE_FP32_OPS
+	return Lhs.High < Rhs;
+#else
+#if COMPILER_SUPPORTS_HLSL2021
+	return or(Lhs.High < Rhs, and(Lhs.High == Rhs, Lhs.Low < 0.0));
+#else
+	return Lhs.High < Rhs || (Lhs.High == Rhs && Lhs.Low < 0.0);
+#endif
+#endif
+}
+
+FBoolType DFGreater(FFloatType Lhs, FDFType Rhs) { return DFLess(Rhs, Lhs); }
+FBoolType DFLess(FFloatType Lhs, FDFType Rhs) { return DFGreater(Rhs, Lhs); }
+
+FBoolType DFGreaterEqual(FDFType Lhs, FDFType Rhs) { return !DFLess(Lhs, Rhs); }
+FBoolType DFGreaterEqual(FDFType Lhs, FFloatType Rhs) { return !DFLess(Lhs, Rhs); }
+FBoolType DFGreaterEqual(FFloatType Lhs, FDFType Rhs) { return !DFLess(Lhs, Rhs); }
+
+FBoolType DFLessEqual(FDFType Lhs, FFloatType Rhs) { return !DFGreater(Lhs, Rhs); }
+FBoolType DFLessEqual(FDFType Lhs, FDFType Rhs) { return !DFGreater(Lhs, Rhs); }
+FBoolType DFLessEqual(FFloatType Lhs, FDFType Rhs) { return !DFGreater(Lhs, Rhs); }
+
+FDFType DFMin(FDFType Lhs, FDFType Rhs)
+{
+	return DFSelect(DFLess(Lhs, Rhs), Lhs, Rhs);
+}
+
+FDFType DFMin(FDFType Lhs, FFloatType Rhs)
+{
+	return DFSelect(DFLess(Lhs, Rhs), Lhs, DFPromote(Rhs));
+}
+
+FDFType DFMin(FFloatType Lhs, FDFType Rhs)
+{
+	return DFSelect(DFLess(Lhs, Rhs), DFPromote(Lhs), Rhs);
+}
+
+FDFType DFMax(FDFType Lhs, FDFType Rhs)
+{
+	return DFSelect(DFLess(Lhs, Rhs), Rhs, Lhs);
+}
+
+FDFType DFMax(FDFType Lhs, FFloatType Rhs)
+{
+	return DFSelect(DFLess(Lhs, Rhs), DFPromote(Rhs), Lhs);
+}
+
+FDFType DFMax(FFloatType Lhs, FDFType Rhs)
+{
+	return DFSelect(DFLess(Lhs, Rhs), Rhs, DFPromote(Lhs));
+}
+
+/**********/
+/** MISC **/
+/**********/
+
+struct FDFTypeDeriv
+{
+	FDFType Value;
+	FFloatType Ddx;
+	FFloatType Ddy;
+};
+
+// [2] Algorithm 8
+FDFType DFSqrt(FDFType V)
+{
+#if UE_DF_FORCE_FP32_OPS
+	return DFPromote(sqrt(V.High));
+#else
+	const FFloatType Sh = sqrt(V.High);
+	const FFloatType P1 = INVARIANT(mad(-Sh, Sh, V.High));
+	const FFloatType P2 = INVARIANT(V.Low + P1);
+
+	const FFloatType Sl = INVARIANT(P2 / (2.0 * Sh));
+	return DFFastTwoSum(Sh, Sl);
+
+	// equivalent to (remark 3.5)
+	//precise FFloatType Tl = P2 / Sh;
+	//precise FFloatType Zh = mad(0.5, Tl, Sh);
+	//precise FFloatType D = Zh - Sh;
+	//precise FFloatType Zl = mad(0.5, Tl, -D);
+	//return DFConstructor(Zh, Zl);
+#endif
+}
+
+// [2] Algorithm 9
+FFloatType DFSqrtDemote(FDFType V)
+{
+#if UE_DF_FORCE_FP32_OPS
+	return sqrt(V.High);
+#else
+	const FFloatType Sh = sqrt(V.High);
+	const FFloatType P1 = INVARIANT(mad(-Sh, Sh, V.High));
+	const FFloatType P2 = INVARIANT(V.Low + P1);
+
+	const FFloatType Sl = INVARIANT(P2 / (2.0 * Sh));
+	return Sh + Sl;
+#endif
+}
+
+FDFType DFRsqrt(FDFType V)
+{
+#if UE_DF_FORCE_FP32_OPS
+	return DFPromote(rsqrt(V.High));
+#else
+	return DFDivide((FFloatType)1.0, DFSqrt(V));
+#endif
+}
+FFloatType DFRsqrtDemote(FDFType V) { return rsqrt(V.High); }
+
+FFloatType DFSin(FDFType V)
+{
+#if UE_DF_FORCE_FP32_OPS
+	return sin(V.High);
+#else
+	FFloatType Sh, Ch;
+	sincos(V.High, Sh, Ch);
+	FFloatType Sl, Cl;
+	sincos(V.Low, Sl, Cl);
+	return Sh*Cl + Ch*Sl; // trigonometric identity for sin(h+l)
+#endif
+}
+
+FFloatType DFCos(FDFType V)
+{
+#if UE_DF_FORCE_FP32_OPS
+	return cos(V.High);
+#else
+	FFloatType Sh, Ch;
+	sincos(V.High, Sh, Ch);
+	FFloatType Sl, Cl;
+	sincos(V.Low, Sl, Cl);
+	return Ch*Cl - Sh*Sl; // trigonometric identity for cos(h+l)
+#endif
+}
+
+void DFSinCos(FDFType V, out FFloatType Sin, out FFloatType Cos)
+{
+#if UE_DF_FORCE_FP32_OPS
+	sincos(V.High, Sin, Cos);
+#else
+	FFloatType Sh, Ch;
+	sincos(V.High, Sh, Ch);
+	FFloatType Sl, Cl;
+	sincos(V.Low, Sl, Cl);
+	Sin = mad(Sh, Cl,  Ch*Sl); // trigonometric identity for sin(h+l)
+	Cos = mad(Ch, Cl, -Sh*Sl); // trigonometric identity for cos(h+l)
+#endif
+}
+
+FFloatType DFTan(FDFType V)
+{
+#if UE_DF_FORCE_FP32_OPS
+	return tan(V.High);
+#else
+	FFloatType Th = tan(V.High);
+	FFloatType Tl = tan(V.Low);
+	return (Th + Tl) / (1.0 - Th * Tl); // trigonometric identity for tan(h+l)
+#endif
+}
+
+FFloatType DFAsin(FDFType V)
+{
+	return asin(DFDemote(V));
+}
+
+FFloatType DFAcos(FDFType V)
+{
+	return acos(DFDemote(V));
+}
+
+FFloatType DFAtan(FDFType V)
+{
+	return atan(DFDemote(V));
+}
+
+FFloatType DFModf(FDFType V, out FDFType Integer)
+{
+#if UE_DF_FORCE_FP32_OPS
+	Integer.Low = 0;
+	return modf(V.High, Integer.High);
+#else
+	FFloatType IntHigh;
+	FFloatType FracHigh = modf(V.High, IntHigh);
+	FFloatType IntLow;
+	FFloatType FracLow = modf(V.Low, IntLow);
+	FFloatType IntSum;
+	FFloatType FracSum = modf(FracHigh + FracLow, IntSum);
+	Integer = DFTwoSum(IntHigh, IntLow + IntSum);
+	return FracSum;
+#endif
+}
+
+FFloatType DFModfDemote(FDFType V, out FFloatType Integer)
+{
+#if UE_DF_FORCE_FP32_OPS
+	return modf(V.High, Integer);
+#else
+	FFloatType IntHigh;
+	FFloatType FracHigh = modf(V.High, IntHigh);
+	FFloatType IntLow;
+	FFloatType FracLow = modf(V.Low, IntLow);
+	FFloatType IntSum;
+	FFloatType FracSum = modf(FracHigh + FracLow, IntSum);
+	Integer = IntHigh + IntLow + IntSum;
+	return FracSum;
+#endif
+}
+
+FDFType DFCeil(FDFType V)
+{
+#if UE_DF_FORCE_FP32_OPS
+	return DFPromote(ceil(V.High));
+#else
+	FFloatType IntHigh;
+	FFloatType FracHigh = modf(V.High, IntHigh);
+	FFloatType IntLow;
+	FFloatType FracLow = modf(V.Low, IntLow);
+	IntLow += ceil(FracHigh + FracLow);
+	return DFTwoSum(IntHigh, IntLow);
+#endif
+}
+
+FDFType DFFloor(FDFType V)
+{
+#if UE_DF_FORCE_FP32_OPS
+	return DFPromote(floor(V.High));
+#else
+	FFloatType IntHigh;
+	FFloatType FracHigh = modf(V.High, IntHigh);
+	FFloatType IntLow;
+	FFloatType FracLow = modf(V.Low, IntLow);
+	IntLow += floor(FracHigh + FracLow);
+	return DFTwoSum(IntHigh, IntLow);
+#endif
+}
+
+FDFType DFRound(FDFType V)
+{
+#if UE_DF_FORCE_FP32_OPS
+	return DFPromote(round(V.High));
+#else
+	FFloatType Sign = DFSign(V);
+	FDFType VAbs = DFConstructor(Sign*V.High, Sign*V.Low);
+
+	FFloatType IntHigh;
+	FFloatType FracHigh = modf(VAbs.High, IntHigh);
+	FFloatType IntLow;
+	FFloatType FracLow = modf(VAbs.Low, IntLow);
+
+	IntLow += floor(FracHigh + FracLow + 0.5);
+	return DFTwoSum(Sign * IntHigh, Sign * IntLow);
+#endif
+}
+
+FDFType DFTrunc(FDFType V)
+{
+#if UE_DF_FORCE_FP32_OPS
+	return DFPromote(trunc(V.High));
+#else
+	FFloatType Sign = DFSign(V);
+	FDFType VAbs = DFConstructor(Sign*V.High, Sign*V.Low);
+
+	FFloatType IntHigh;
+	FFloatType FracHigh = modf(VAbs.High, IntHigh);
+	FFloatType IntLow;
+	FFloatType FracLow = modf(VAbs.Low, IntLow);
+
+	IntLow += floor(FracHigh + FracLow);
+	return DFTwoSum(Sign * IntHigh, Sign * IntLow);
+#endif
+}
+
+FDFType DFFrac(FDFType V)
+{
+#if UE_DF_FORCE_FP32_OPS
+	return DFPromote(frac(V.High));
+#else
+	return DFSubtract(V, DFFloor(V));
+#endif
+}
+
+// Technically, this function has very low precision due to discontinuity+rounding
+// (e.g. 0.99.. ~= 1.0 => frac(0.99..) = frac(1.0) = 0.0)
+// but this is irrelevant in practical cases
+// TODO: test negative values
+FFloatType DFFracDemote(FDFType V)
+{
+#if UE_DF_FORCE_FP32_OPS
+	return frac(V.High);
+#else
+	return frac(frac(V.High) + frac(V.Low));
+#endif
+}
+
+// Similar to HLSL fmod, this is equivalent to `lhs - rhs * trunc(lhs / rhs)`
+// so fmod(-0.1, 1.0) returns -0.1, not 0.9 (as mod does in GLSL, which is the true mathematical modulo operator)
+FDFType DFFmod(FDFType Lhs, FFloatType Rhs)
+{
+#if UE_DF_FORCE_FP32_OPS
+	return DFPromote(fmod(Lhs.High, Rhs));
+#else
+	return DFSubtract(Lhs, DFMultiply(DFTrunc(DFDivide(Lhs, Rhs)), Rhs));
+#endif
+}
+
+FFloatType DFFmodDemote(FDFType Lhs, FFloatType Rhs)
+{
+#if UE_DF_FORCE_FP32_OPS
+	return fmod(Lhs.High, Rhs);
+#else
+	//DF_TODO: optimize. Is equal to frac(Lhs/Rhs)*Rhs?
+	//not equal to fmod(fmod(high)+fmod(low)) due to how fmod handles negative values
+	return DFSubtractDemote(Lhs, DFMultiply(DFTrunc(DFDivide(Lhs, Rhs)), Rhs)); 
+#endif
+}
+
+FFloatType DFFmodByPow2Demote(FDFType Lhs, FFloatType Rhs)
+{
+#if UE_DF_FORCE_FP32_OPS
+	return fmod(Lhs.High, Rhs);
+#else
+	return DFFracDemote(DFDivideByPow2(Lhs, Rhs)) * Rhs;
+#endif
+}
+
+FDFType DFLerp(FDFType Lhs, FDFType Rhs, FFloatType S)
+{
+#if UE_DF_FORCE_FP32_OPS
+	return DFPromote(lerp(Lhs.High, Rhs.High, S));
+#else
+	return DFTwoSum(lerp(Lhs.High, Rhs.High, S), lerp(Lhs.Low, Rhs.Low, S)); //DF_TODO: precision issues
+#endif
+}
+
+FDFType DFSaturate(FDFType V)
+{
+#if UE_DF_FORCE_FP32_OPS
+	return DFPromote(saturate(V.High));
+#else
+	V = DFSelect(DFLess(V, 1.0), V, 1.0);
+	V = DFSelect(DFLess(V, 0.0), 0.0, V);
+	return V;
+#endif
+}
+FFloatType DFSaturateDemote(FDFType V) { return saturate(DFDemote(V)); }
+
+FDFType DFSmoothStep(FDFType Lhs, FDFType Rhs, FDFType S)
+{
+#if UE_DF_FORCE_FP32_OPS
+	return DFPromote(smoothstep(Lhs.High, Rhs.High, S.High));
+#else
+	FDFType T = DFSaturate(DFDivide(DFSubtract(S, Lhs), DFSubtract(Rhs, Lhs)));
+	return DFMultiply(DFSqr(T), DFSubtract(3.0f, DFMultiplyByPow2(T, 2.0f)));
+#endif
+}
+FFloatType DFSmoothStepDemote(FDFType Lhs, FDFType Rhs, FDFType S)
+{
+#if UE_DF_FORCE_FP32_OPS
+	return smoothstep(Lhs.High, Rhs.High, S.High);
+#else
+	FFloatType T = DFSaturateDemote(DFDivide(DFSubtract(S, Lhs), DFSubtract(Rhs, Lhs)));
+	return T*T*(3.0f - (2.0f*T));
+#endif
+}
+
+FFloatType DFStep(FDFType Lhs, FDFType Rhs)    { return select(DFGreaterEqual(Rhs, Lhs), (FFloatType)1.0f, (FFloatType)0.0f); }
+FFloatType DFStep(FDFType Lhs, FFloatType Rhs) { return select(DFGreaterEqual(Rhs, Lhs), (FFloatType)1.0f, (FFloatType)0.0f); }
+FFloatType DFStep(FFloatType Lhs, FDFType Rhs) { return select(DFGreaterEqual(Rhs, Lhs), (FFloatType)1.0f, (FFloatType)0.0f); }
+
+// No ddxy inside ray tracing shaders
+#if RAYHITGROUPSHADER || RAYMISSHADER || RAYCALLABLESHADER || USE_FORCE_TEXTURE_MIP
+FDFType DFDdx(FDFType V) { return (FDFType)0; }
+FDFType DFDdy(FDFType V) { return (FDFType)0; }
+FFloatType DFDdxDemote(FDFType V) { return (FFloatType)0; }
+FFloatType DFDdyDemote(FDFType V) { return (FFloatType)0; }
+#else
+
+FDFType DFDdx(FDFType V)
+{
+#if UE_DF_FORCE_FP32_OPS
+	return DFPromote(ddx(V.High));
+#else
+	return DFTwoSum(ddx(V.High), ddx(V.Low));
+#endif
+}
+
+FDFType DFDdy(FDFType V)
+{
+#if UE_DF_FORCE_FP32_OPS
+	return DFPromote(ddy(V.High));
+#else
+	return DFTwoSum(ddy(V.High), ddy(V.Low));
+#endif
+}
+
+FFloatType DFDdxDemote(FDFType V)
+{
+#if UE_DF_FORCE_FP32_OPS
+	return ddx(V.High);
+#else
+	return ddx(V.High) + ddx(V.Low);
+#endif
+}
+
+FFloatType DFDdyDemote(FDFType V)
+{
+#if UE_DF_FORCE_FP32_OPS
+	return ddy(V.High);
+#else
+	return ddy(V.High) + ddy(V.Low);
+#endif
+}
+#endif
diff --git a/Engine/Shaders/Private/DoubleFloatVectorDefinition.ush b/Engine/Shaders/Private/DoubleFloatVectorDefinition.ush
new file mode 100644
--- /dev/null
+++ b/Engine/Shaders/Private/DoubleFloatVectorDefinition.ush
@@ -0,0 +1,246 @@
+// Copyright Epic Games, Inc. All Rights Reserved.
+
+#if COMPILER_SUPPORTS_HLSL2021
+
+#define DECLARE_VEC_VEC_OP(DFOperation) \
+	FDFType DFOperation(FDFType Lhs, FDFType Rhs); \
+	FDFType DFOperation(FDFType Lhs, FFloatType Rhs); \
+	FDFType DFOperation(FDFType Lhs, FDFScalar Rhs); \
+	FDFType DFOperation(FDFType Lhs, float Rhs);
+
+struct FDFType;
+FDFType DFBroadcastFunc(FDFScalar V);
+FDFType DFBroadcastFunc(float V);
+FDFScalar DFGetComponent(FDFType V, int C);
+FDFType DFNegate(FDFType In);
+DECLARE_VEC_VEC_OP(DFAdd);
+DECLARE_VEC_VEC_OP(DFSubtract);
+DECLARE_VEC_VEC_OP(DFMultiply);
+DECLARE_VEC_VEC_OP(DFDivide);
+FBoolType DFGreater(FDFType Lhs, FFloatType Rhs);
+FBoolType DFLess(FDFType Lhs, FFloatType Rhs);
+FFloatType DFDemote(FDFType In);
+
+#undef DECLARE_VEC_VEC_OP
+
+#endif
+
+struct FDFType
+{
+	FFloatType High;
+	FFloatType Low;
+
+#if COMPILER_SUPPORTS_HLSL2021
+
+#define DECLARE_VEC_VEC_OP(Operator, DFOperation) \
+	FDFType Operator(FDFType Rhs) { return DFOperation(this, Rhs); } \
+	FDFType Operator(FFloatType Rhs) { return DFOperation(this, Rhs); } \
+	FDFType Operator(FDFScalar Rhs) { return DFOperation(this, Rhs); } \
+	FDFType Operator(float Rhs) { return DFOperation(this, Rhs); }
+
+	FDFScalar operator[](int idx) { return DFGetComponent(this, idx); }
+	FDFType operator-() { return DFNegate(this); }
+	//DECLARE_VEC_VEC_OP(operator+, DFAdd)
+	//DECLARE_VEC_VEC_OP(operator-, DFSubtract)
+	//DECLARE_VEC_VEC_OP(operator*, DFMultiply)
+	//DECLARE_VEC_VEC_OP(operator/, DFDivide)
+	FBoolType operator>(FFloatType Rhs) { return DFGreater(this, Rhs); }
+	FBoolType operator<(FFloatType Rhs) { return DFLess(this, Rhs); }
+
+	FFloatType Demote() { return DFDemote(this); }
+
+#undef DECLARE_VEC_VEC_OP
+
+#endif
+};
+
+FDFScalar DFGetComponent(FDFType V, int C) { return MakeDFScalar(V.High[C], V.Low[C]); }
+
+#include "DoubleFloatOperations.ush"
+
+FDFType DFBroadcastFunc(FDFScalar V)
+{
+	const FFloatType High = V.High;
+	const FFloatType Low = V.Low;
+	return DFConstructor(High, Low);
+}
+
+FDFType DFBroadcastFunc(float V)
+{
+	const FFloatType High = V;
+	return DFPromote(High);
+}
+
+#define DEFINE_BROADCASTED_OP(DFOperation) \
+FDFType DFOperation(FDFType Lhs, FDFScalar Rhs) { return DFOperation(Lhs, DFBroadcastFunc(Rhs)); } \
+FDFType DFOperation(FDFType Lhs, float Rhs)     { return DFOperation(Lhs, DFBroadcastFunc(Rhs)); }
+
+DEFINE_BROADCASTED_OP(DFAdd)
+DEFINE_BROADCASTED_OP(DFFastAdd)
+DEFINE_BROADCASTED_OP(DFSubtract)
+DEFINE_BROADCASTED_OP(DFFastSubtract)
+DEFINE_BROADCASTED_OP(DFMultiply)
+DEFINE_BROADCASTED_OP(DFDivide)
+DEFINE_BROADCASTED_OP(DFFastDivide)
+
+#undef DEFINE_BROADCASTED_OP
+
+FDFType DFDivideByPow2(FDFType Lhs, float Rhs)
+{
+	FFloatType RhsVect = Rhs;
+	return DFDivideByPow2(Lhs, RhsVect);
+}
+
+FDFScalar DFVectorSum(FDFType V)
+{
+#if UE_DF_FORCE_FP32_OPS
+	float Result = V.High[0] + V.High[1];
+	[unroll]
+	for (int i = 2; i < GNumComponents; ++i)
+	{
+		Result = Result + V.High[i];
+	}
+	return DFPromote(Result);
+#else
+	FDFScalar Result = DFAdd(DFGetComponent(V, 0), DFGetComponent(V, 1));
+	[unroll]
+	for (int i = 2; i < GNumComponents; ++i)
+	{
+		Result = DFAdd(Result, DFGetComponent(V, i));
+	}
+	return Result;
+#endif
+}
+
+FDFScalar DFFastVectorSum(FDFType V)
+{
+#if UE_DF_FORCE_FP32_OPS
+	return DFVectorSum(V);
+#else
+	FDFScalar Result = DFFastAdd(DFGetComponent(V, 0), DFGetComponent(V, 1));
+	[unroll]
+	for (int i = 2; i < GNumComponents; ++i)
+	{
+		Result = DFFastAdd(Result, DFGetComponent(V, i));
+	}
+	return Result;
+#endif
+}
+
+FDFScalar DFDot(FDFType Lhs, FFloatType Rhs)
+{
+#if UE_DF_FORCE_FP32_OPS
+	return DFPromote(dot(Lhs.High, Rhs));
+#else
+	return DFFastVectorSum(DFMultiply(Lhs, Rhs));
+#endif
+}
+FDFScalar DFDot(FFloatType Lhs, FDFType Rhs) { return DFDot(Rhs, Lhs); }
+FDFScalar DFDot(FDFType Lhs, FDFType Rhs)
+{
+#if UE_DF_FORCE_FP32_OPS
+	return DFPromote(dot(Lhs.High, Rhs.High));
+#else
+	return DFFastVectorSum(DFMultiply(Lhs, Rhs));
+#endif
+}
+
+FFloatType DFDotDemote(FDFType Lhs, FFloatType Rhs)
+{
+#if UE_DF_FORCE_FP32_OPS
+	return dot(Lhs.High, Rhs);
+#else
+	return DFDemote(DFFastVectorSum(DFMultiply(Lhs, Rhs)));
+#endif
+}
+FFloatType DFDotDemote(FFloatType Lhs, FDFType Rhs) { return DFDotDemote(Rhs, Lhs); }
+FFloatType DFDotDemote(FDFType Lhs, FDFType Rhs)
+{
+#if UE_DF_FORCE_FP32_OPS
+	return dot(Lhs.High, Rhs.High);
+#else
+	return DFDemote(DFFastVectorSum(DFMultiply(Lhs, Rhs)));
+#endif
+}
+
+FDFScalar DFLengthSqr(FFloatType V)
+{
+#if UE_DF_FORCE_FP32_OPS
+	return DFPromote(dot(V, V));
+#else
+	return DFFastVectorSum(DFSqr(V));
+#endif
+}
+
+FDFScalar DFLength(FFloatType V)
+{
+#if UE_DF_FORCE_FP32_OPS
+	return DFPromote(length(V));
+#else
+	return DFSqrt(DFLengthSqr(V));
+#endif
+}
+
+// DF_TODO: test potential precision issues
+
+FDFScalar DFLengthSqr(FDFType V)
+{
+#if UE_DF_FORCE_FP32_OPS
+	return DFPromote(dot(V.High, V.High));
+#else
+	return DFFastVectorSum(DFSqr(V));
+#endif
+}
+FDFScalar DFLength(FDFType V)
+{
+#if UE_DF_FORCE_FP32_OPS
+	return DFPromote(length(V.High));
+#else
+	return DFSqrt(DFLengthSqr(V));
+#endif
+}
+
+float DFLengthDemote(FDFType V)
+{
+#if UE_DF_FORCE_FP32_OPS
+	return length(V.High);
+#else
+	return DFSqrtDemote(DFLengthSqr(V));
+#endif
+}
+
+float DFRcpLengthDemote(FDFType V)
+{
+#if UE_DF_FORCE_FP32_OPS
+	return rcp(length(V.High));
+#else
+	return DFRsqrtDemote(DFLengthSqr(V));
+#endif
+}
+
+FDFType DFNormalize(FFloatType V)
+{
+#if UE_DF_FORCE_FP32_OPS
+	return DFPromote(normalize(V));
+#else
+	return DFDivide(V, DFBroadcastFunc(DFLength(V)));
+#endif
+}
+
+FDFType DFNormalize(FDFType V)
+{
+#if UE_DF_FORCE_FP32_OPS
+	return DFPromote(normalize(V.High));
+#else
+	return DFFastDivide(V, DFLength(V));
+#endif
+}
+
+FFloatType DFNormalizeDemote(FDFType V)
+{
+#if UE_DF_FORCE_FP32_OPS
+	return normalize(V.High);
+#else
+	return V.High / length(V.High);
+#endif
+}
diff --git a/Engine/Shaders/Private/LargeWorldCoordinates.ush b/Engine/Shaders/Private/LargeWorldCoordinates.ush
--- a/Engine/Shaders/Private/LargeWorldCoordinates.ush
+++ b/Engine/Shaders/Private/LargeWorldCoordinates.ush
@@ -2,28 +2,42 @@
 
 #pragma once
 
+#include "DoubleFloat.ush"
+
 struct FLWCScalar
 {
 	float Tile;
 	float Offset;
+#if !COMPILER_SUPPORTS_HLSL2021
+	int Dummy; // avoid problem with HLSL overloading/implicit cast (should be DCE'd)
+#endif
 };
 
 struct FLWCVector2
 {
 	float2 Tile;
 	float2 Offset;
+#if !COMPILER_SUPPORTS_HLSL2021
+	int Dummy; // avoid problem with HLSL overloading/implicit cast (should be DCE'd)
+#endif
 };
 
 struct FLWCVector3
 {
 	float3 Tile;
 	float3 Offset;
+#if !COMPILER_SUPPORTS_HLSL2021
+	int Dummy; // avoid problem with HLSL overloading/implicit cast (should be DCE'd)
+#endif
 };
 
 struct FLWCVector4
 {
 	float4 Tile;
 	float4 Offset;
+#if !COMPILER_SUPPORTS_HLSL2021
+	int Dummy; // avoid problem with HLSL overloading/implicit cast (should be DCE'd)
+#endif
 };
 
 struct FLWCScalarDeriv
@@ -59,6 +73,9 @@ struct FLWCMatrix
 {
 	float4x4 M;
 	float3 Tile; // Added to result, *after* multiplying 'M'
+#if !COMPILER_SUPPORTS_HLSL2021
+	int2 Dummy; // avoid problem with HLSL overloading/implicit cast (should be DCE'd)
+#endif
 };
 
 // Transforms *from* absolute world space
@@ -66,7 +83,9 @@ struct FLWCInverseMatrix
 {
 	float4x4 M;
 	float3 Tile; // Added to input position *before* multiplying 'M'
-	int Dummy; // avoid problem with stupid HLSL overloading
+#if !COMPILER_SUPPORTS_HLSL2021
+	int3 Dummy; // avoid problem with HLSL overloading/implicit cast (should be DCE'd)
+#endif
 };
 
 #define LWCSetTile(V, InTile) (V).Tile = (InTile)
@@ -79,32 +98,14 @@ float4 LWCGetTileOffset(FLWCVector4 V) { return LWCGetTile(V) * UE_LWC_RENDER_TI
 float3 LWCGetTileOffset(FLWCMatrix V) { return LWCGetTile(V) * UE_LWC_RENDER_TILE_SIZE; }
 float3 LWCGetTileOffset(FLWCInverseMatrix V) { return LWCGetTile(V) * UE_LWC_RENDER_TILE_SIZE; }
 
-float4x4 Make4x3Matrix(float4x4 M)
-{
-	// Explicitly ignore the 4th column of the input
-	float4x4 Result;
-	Result[0] = float4(M[0].xyz, 0.0f);
-	Result[1] = float4(M[1].xyz, 0.0f);
-	Result[2] = float4(M[2].xyz, 0.0f);
-	Result[3] = float4(M[3].xyz, 1.0f);
-	return Result;
-}
-
-float4x4 MakeTranslationMatrix(float3 Offset)
-{
-	float4x4 Result;
-	Result[0] = float4(1.0f, 0.0f, 0.0f, 0.0f);
-	Result[1] = float4(0.0f, 1.0f, 0.0f, 0.0f);
-	Result[2] = float4(0.0f, 0.0f, 1.0f, 0.0f);
-	Result[3] = float4(Offset, 1.0f);
-	return Result;
-}
-
 FLWCScalar MakeLWCScalar(float Tile, float Offset)
 {
 	FLWCScalar Result;
 	LWCSetTile(Result, Tile);
 	Result.Offset = Offset;
+#if !COMPILER_SUPPORTS_HLSL2021
+	Result.Dummy = 0;
+#endif
 	return Result;
 }
 
@@ -113,6 +114,9 @@ FLWCVector2 MakeLWCVector2(float2 Tile, float2 Offset)
 	FLWCVector2 Result;
 	LWCSetTile(Result, Tile);
 	Result.Offset = Offset;
+#if !COMPILER_SUPPORTS_HLSL2021
+	Result.Dummy = 0;
+#endif
 	return Result;
 }
 
@@ -121,6 +125,9 @@ FLWCVector3 MakeLWCVector3(float3 Tile, float3 Offset)
 	FLWCVector3 Result;
 	LWCSetTile(Result, Tile);
 	Result.Offset = Offset;
+#if !COMPILER_SUPPORTS_HLSL2021
+	Result.Dummy = 0;
+#endif
 	return Result;
 }
 
@@ -129,6 +136,9 @@ FLWCVector4 MakeLWCVector4(float4 Tile, float4 Offset)
 	FLWCVector4 Result;
 	LWCSetTile(Result, Tile);
 	Result.Offset = Offset;
+#if !COMPILER_SUPPORTS_HLSL2021
+	Result.Dummy = 0;
+#endif
 	return Result;
 }
 
@@ -163,6 +173,9 @@ FLWCMatrix MakeLWCMatrix(float3 Tile, float4x4 InMatrix)
 	FLWCMatrix Result;
 	LWCSetTile(Result, Tile);
 	Result.M = InMatrix;
+#if !COMPILER_SUPPORTS_HLSL2021
+	Result.Dummy = 0;
+#endif
 	return Result;
 }
 
@@ -171,6 +184,9 @@ FLWCMatrix MakeLWCMatrix4x3(float3 Tile, float4x4 InMatrix)
 	FLWCMatrix Result;
 	LWCSetTile(Result, Tile);
 	Result.M = Make4x3Matrix(InMatrix);
+#if !COMPILER_SUPPORTS_HLSL2021
+	Result.Dummy = 0;
+#endif
 	return Result;
 }
 
@@ -179,7 +195,9 @@ FLWCInverseMatrix MakeLWCInverseMatrix(float3 Tile, float4x4 InMatrix)
 	FLWCInverseMatrix Result;
 	LWCSetTile(Result, -Tile);
 	Result.M = InMatrix;
+#if !COMPILER_SUPPORTS_HLSL2021
 	Result.Dummy = 0;
+#endif
 	return Result;
 }
 
@@ -188,7 +206,9 @@ FLWCInverseMatrix MakeLWCInverseMatrix4x3(float3 Tile, float4x4 InMatrix)
 	FLWCInverseMatrix Result;
 	LWCSetTile(Result, -Tile);
 	Result.M = Make4x3Matrix(InMatrix);
+#if !COMPILER_SUPPORTS_HLSL2021
 	Result.Dummy = 0;
+#endif
 	return Result;
 }
 
@@ -492,3 +512,112 @@ float4x4 LWCHackToFloat(FLWCInverseMatrix v) { return (float4x4)0.0f; }
 #else
 #define LWCHackToFloat(V) LWCToFloat(V)
 #endif
+
+#define DEFINE_DF_TO_TILEOFFSET_OPERATOR(FDFType, FLWCType, LWCConstructor, FFloatType) \
+	FLWCType DFToTileOffset(FDFType In)																							\
+	{																															\
+		/* Used rounding, removes fractional data stored in Tile, so the resulting coordinate is 'normalized' */				\
+		/* T+C-C benchmarked slightly faster than round(), can also shift C right to retain some fractional bits */ 			\
+		const float C = 8388608; /* = 2^23 (23 mantisse bits for int part, 0 for fraction) */ 									\
+		FFloatType Tile = mad(In.High, 1 / UE_LWC_RENDER_TILE_SIZE, C) - C;														\
+		FFloatType Offset = mad(Tile, -UE_LWC_RENDER_TILE_SIZE, In.High) + In.Low;												\
+																																\
+		return LWCConstructor(Tile, Offset);																					\
+	}																															\
+	FLWCType DFFastToTileOffset(FDFType In)																						\
+	{																															\
+		/* Produces a non-normalized tile-offset coordinate. */																	\
+		/* If we're not rounding, Tile has the full precision of In.High because UE_LWC_RENDER_TILE_SIZE is pow2 */				\
+		/* Therefore, we don't need to calculate the error term and can just use In.Low */										\
+		 FFloatType Tile = In.High * (1 / UE_LWC_RENDER_TILE_SIZE); 															\
+		 FFloatType Offset = In.Low; 																							\
+																																\
+		return LWCConstructor(Tile, Offset);																					\
+	}
+
+DEFINE_DF_TO_TILEOFFSET_OPERATOR(FDFScalar, FLWCScalar, MakeLWCScalar, float)
+DEFINE_DF_TO_TILEOFFSET_OPERATOR(FDFVector2, FLWCVector2, MakeLWCVector2, float2)
+DEFINE_DF_TO_TILEOFFSET_OPERATOR(FDFVector3, FLWCVector3, MakeLWCVector3, float3)
+DEFINE_DF_TO_TILEOFFSET_OPERATOR(FDFVector4, FLWCVector4, MakeLWCVector4, float4)
+#undef DEFINE_DF_TO_TILEOFFSET_OPERATOR
+
+FLWCMatrix DFToTileOffset(FDFMatrix In)
+{
+	FLWCVector3 PosHigh = DFToTileOffset(MakeDFVector3(In.PostTranslation, 0));
+
+	float4x4 M = MultiplyTranslation(In.M, PosHigh.Offset);
+	float3 Tile = PosHigh.Tile;
+	return MakeLWCMatrix(Tile, M);
+}
+
+FLWCMatrix DFFastToTileOffset(FDFMatrix In)
+{
+	float4x4 M = In.M;
+	float3 Tile = In.PostTranslation / UE_LWC_RENDER_TILE_SIZE;
+	return MakeLWCMatrix(Tile, M);
+}
+
+FLWCInverseMatrix DFToTileOffset(FDFInverseMatrix In)
+{
+	FLWCVector3 PosHigh = DFToTileOffset(MakeDFVector3(In.PreTranslation, 0));
+	float4x4 M = MultiplyTranslation(-PosHigh.Offset, In.M);
+	float3 Tile = PosHigh.Tile;
+	return MakeLWCInverseMatrix(Tile, M);
+}
+
+FLWCInverseMatrix DFFastToTileOffset(FDFInverseMatrix In)
+{
+	float4x4 M = In.M;
+	float3 Tile = In.PreTranslation / UE_LWC_RENDER_TILE_SIZE;
+	return MakeLWCInverseMatrix(Tile, M);
+}
+
+#define DFToTileOffset_Hack DFToTileOffset
+
+FDFScalar DFFromTileOffset(FLWCScalar In)
+{
+	float TileOffset = In.Tile * UE_LWC_RENDER_TILE_SIZE;
+	return DFFastTwoSum(TileOffset, In.Offset);
+}
+
+FDFVector2 DFFromTileOffset(FLWCVector2 In)
+{
+	float2 TileOffset = In.Tile * UE_LWC_RENDER_TILE_SIZE;
+	return DFFastTwoSum(TileOffset, In.Offset);
+}
+
+FDFVector3 DFFromTileOffset(FLWCVector3 In)
+{
+	float3 TileOffset = In.Tile * UE_LWC_RENDER_TILE_SIZE;
+	return DFFastTwoSum(TileOffset, In.Offset);
+}
+
+FDFVector4 DFFromTileOffset(FLWCVector4 In)
+{
+	float4 TileOffset = In.Tile * UE_LWC_RENDER_TILE_SIZE;
+	return DFFastTwoSum(TileOffset, In.Offset);
+}
+
+FDFMatrix DFFromTileOffset(FLWCMatrix In)
+{
+	float4x4 M = In.M;
+	float3 PostTranslation = In.Tile * UE_LWC_RENDER_TILE_SIZE;
+	return MakeDFMatrix(PostTranslation, M);
+}
+
+FDFInverseMatrix DFFromTileOffset(FLWCInverseMatrix In)
+{
+	float4x4 M = In.M;
+	float3 PreTranslation = -In.Tile * UE_LWC_RENDER_TILE_SIZE;
+	return MakeDFInverseMatrix(PreTranslation, M);
+}
+
+#define DFFromTileOffset_Hack DFFromTileOffset
+
+// Hacky LWCFunc(DFParam) shims for custom hlsl nodes in materials
+float LWCToFloat(FDFScalar Value)   { return DFDemote(Value); }
+float2 LWCToFloat(FDFVector2 Value) { return DFDemote(Value); }
+float3 LWCToFloat(FDFVector3 Value) { return DFDemote(Value); }
+float4 LWCToFloat(FDFVector4 Value) { return DFDemote(Value); }
+float4x4 LWCToFloat(FDFMatrix Value) { return DFDemote(Value); }
+float4x4 LWCToFloat(FDFInverseMatrix Value) { return DFDemote(Value); }
