Remove DF operator definitions as they are currently unsupported by HlslParser (cook error).

diff --git a/Engine/Shaders/Private/DoubleFloat.ush b/Engine/Shaders/Private/DoubleFloat.ush
--- a/Engine/Shaders/Private/DoubleFloat.ush
+++ b/Engine/Shaders/Private/DoubleFloat.ush
@@ -60,9 +60,10 @@ struct FDFScalar
 	FDFType operator*(FFloatType Rhs) { return DFMultiply(this, Rhs); }
 	FDFType operator/(FDFType Rhs) { return DFDivide(this, Rhs); }
 	FDFType operator/(FFloatType Rhs) { return DFDivide(this, Rhs); }*/
-	FDFType operator-() { return DFNegate(this); }
+	//Commented out until supported by HlslParser
+	/*FDFType operator-() { return DFNegate(this); }
 	FBoolType operator>(FFloatType Rhs) { return DFGreater(this, Rhs); }
-	FBoolType operator<(FFloatType Rhs) { return DFLess(this, Rhs); }
+	FBoolType operator<(FFloatType Rhs) { return DFLess(this, Rhs); }*/
 
 	FFloatType Demote() { return DFDemote(this); }
 #endif
diff --git a/Engine/Shaders/Private/DoubleFloatVectorDefinition.ush b/Engine/Shaders/Private/DoubleFloatVectorDefinition.ush
--- a/Engine/Shaders/Private/DoubleFloatVectorDefinition.ush
+++ b/Engine/Shaders/Private/DoubleFloatVectorDefinition.ush
@@ -38,14 +38,15 @@ struct FDFType
 	FDFType Operator(FDFScalar Rhs) { return DFOperation(this, Rhs); } \
 	FDFType Operator(float Rhs) { return DFOperation(this, Rhs); }
 
-	FDFScalar operator[](int idx) { return DFGetComponent(this, idx); }
-	FDFType operator-() { return DFNegate(this); }
-	//DECLARE_VEC_VEC_OP(operator+, DFAdd)
-	//DECLARE_VEC_VEC_OP(operator-, DFSubtract)
-	//DECLARE_VEC_VEC_OP(operator*, DFMultiply)
-	//DECLARE_VEC_VEC_OP(operator/, DFDivide)
-	FBoolType operator>(FFloatType Rhs) { return DFGreater(this, Rhs); }
-	FBoolType operator<(FFloatType Rhs) { return DFLess(this, Rhs); }
+	//Commented out until supported by HlslParser
+	//FDFScalar operator[](int idx) { return DFGetComponent(this, idx); }
+	//FDFType operator-() { return DFNegate(this); }
+	////DECLARE_VEC_VEC_OP(operator+, DFAdd)
+	////DECLARE_VEC_VEC_OP(operator-, DFSubtract)
+	////DECLARE_VEC_VEC_OP(operator*, DFMultiply)
+	////DECLARE_VEC_VEC_OP(operator/, DFDivide)
+	//FBoolType operator>(FFloatType Rhs) { return DFGreater(this, Rhs); }
+	//FBoolType operator<(FFloatType Rhs) { return DFLess(this, Rhs); }
 
 	FFloatType Demote() { return DFDemote(this); }
 
