Add missing function parameter permutations for DoubleFloat vector math operations

diff --git a/Engine/Shaders/Private/DoubleFloatOperations.ush b/Engine/Shaders/Private/DoubleFloatOperations.ush
--- a/Engine/Shaders/Private/DoubleFloatOperations.ush
+++ b/Engine/Shaders/Private/DoubleFloatOperations.ush
@@ -136,6 +136,7 @@ FFloatType DFAddDemote(FDFType Lhs, FFloatType Rhs)
 }
 
 FFloatType DFAddDemote(FFloatType Lhs, FDFType Rhs) { return DFAddDemote(Rhs, Lhs); }
+FFloatType DFAddDemote(FFloatType Lhs, FFloatType Rhs) { return Lhs + Rhs; }
 
 FDFType DFFastAdd(FDFType Lhs, FDFType Rhs)
 {
@@ -176,6 +177,7 @@ FFloatType DFFastAddDemote(FDFType Lhs, FDFType Rhs)
 
 FFloatType DFFastAddDemote(FDFType Lhs, FFloatType Rhs) { return DFAddDemote(Lhs, Rhs); }
 FFloatType DFFastAddDemote(FFloatType Lhs, FDFType Rhs) { return DFFastAddDemote(Rhs, Lhs); }
+FFloatType DFFastAddDemote(FFloatType Lhs, FFloatType Rhs) { return Rhs + Lhs; }
 
 /*****************/
 /** SUBTRACTION **/
@@ -186,14 +188,17 @@ FDFType DFSubtract(FDFType Lhs, FDFType Rhs) { return DFAdd(Lhs, DFNegate(Rhs));
 FDFType DFSubtract(FDFType Lhs, FFloatType Rhs) { return DFAdd(Lhs, -Rhs); }
 FDFType DFSubtract(FFloatType Lhs, FDFType Rhs) { return DFAdd(Lhs, DFNegate(Rhs)); }
 
+FFloatType DFSubtractDemote(FFloatType Lhs, FFloatType Rhs) { return Lhs - Rhs; }
 FFloatType DFSubtractDemote(FDFType Lhs, FDFType Rhs) { return DFAddDemote(Lhs, DFNegate(Rhs)); }
 FFloatType DFSubtractDemote(FDFType Lhs, FFloatType Rhs) { return DFAddDemote(Lhs, -Rhs); }
 FFloatType DFSubtractDemote(FFloatType Lhs, FDFType Rhs) { return DFAddDemote(Lhs, DFNegate(Rhs)); }
 
+FDFType DFFastSubtract(FFloatType Lhs, FFloatType Rhs) { return DFFastAdd(Lhs, -Rhs); }
 FDFType DFFastSubtract(FDFType Lhs, FDFType Rhs) { return DFFastAdd(Lhs, DFNegate(Rhs)); }
 FDFType DFFastSubtract(FDFType Lhs, FFloatType Rhs) { return DFFastAdd(Lhs, -Rhs); }
 FDFType DFFastSubtract(FFloatType Lhs, FDFType Rhs) { return DFFastAdd(Lhs, DFNegate(Rhs)); }
 
+FFloatType DFFastSubtractDemote(FFloatType Lhs, FFloatType Rhs) { return Lhs - Rhs; }
 FFloatType DFFastSubtractDemote(FDFType Lhs, FDFType Rhs) { return DFFastAddDemote(Lhs, DFNegate(Rhs)); }
 FFloatType DFFastSubtractDemote(FDFType Lhs, FFloatType Rhs) { return DFFastAddDemote(Lhs, -Rhs); }
 FFloatType DFFastSubtractDemote(FFloatType Lhs, FDFType Rhs) { return DFFastAddDemote(Lhs, DFNegate(Rhs)); }
@@ -309,6 +314,7 @@ FFloatType DFMultiplyDemote(FDFType Lhs, FFloatType Rhs)
 }
 
 FFloatType DFMultiplyDemote(FFloatType Lhs, FDFType Rhs) { return DFMultiplyDemote(Rhs, Lhs); }
+FFloatType DFMultiplyDemote(FFloatType Lhs, FFloatType Rhs) { return Lhs * Rhs; }
 
 // Fast multiplication that assumes the factor is a power of two.
 // If this assumption is valid, no precision is lost.
@@ -415,6 +421,7 @@ FDFType DFFastDivide(FDFType Lhs, FFloatType Rhs)
 }
 
 FDFType DFFastDivide(FFloatType Lhs, FDFType Rhs) { return DFFastDivide(DFPromote(Lhs), Rhs); }
+FDFType DFFastDivide(FFloatType Lhs, FFloatType Rhs) { return DFFastDivide(DFPromote(Lhs), Rhs); }
 
 FFloatType DFFastDivideDemote(FDFType Lhs, FDFType Rhs)
 {
@@ -451,6 +458,7 @@ FFloatType DFFastDivideDemote(FDFType Lhs, FFloatType Rhs)
 }
 
 FFloatType DFFastDivideDemote(FFloatType Lhs, FDFType Rhs) { return DFFastDivideDemote(DFPromote(Lhs), Rhs); }
+FFloatType DFFastDivideDemote(FFloatType Lhs, FFloatType Rhs) { return Lhs / Rhs; }
 
 // Fast division that assumes the divisor is a power of two.
 // If this assumption is valid, no precision is lost.
diff --git a/Engine/Shaders/Private/DoubleFloatVectorDefinition.ush b/Engine/Shaders/Private/DoubleFloatVectorDefinition.ush
--- a/Engine/Shaders/Private/DoubleFloatVectorDefinition.ush
+++ b/Engine/Shaders/Private/DoubleFloatVectorDefinition.ush
@@ -10,7 +10,7 @@
 
 struct FDFType;
 FDFType DFBroadcastFunc(FDFScalar V);
-FDFType DFBroadcastFunc(float V);
+FFloatType DFBroadcastFunc(float V);
 FDFScalar DFGetComponent(FDFType V, int C);
 FDFType DFNegate(FDFType In);
 DECLARE_VEC_VEC_OP(DFAdd);
@@ -66,23 +66,36 @@ FDFType DFBroadcastFunc(FDFScalar V)
 	return DFConstructor(High, Low);
 }
 
-FDFType DFBroadcastFunc(float V)
+FFloatType DFBroadcastFunc(float V)
 {
-	const FFloatType High = V;
-	return DFPromote(High);
+	const FFloatType Single = V;
+	return Single;
 }
 
-#define DEFINE_BROADCASTED_OP(DFOperation) \
-FDFType DFOperation(FDFType Lhs, FDFScalar Rhs) { return DFOperation(Lhs, DFBroadcastFunc(Rhs)); } \
-FDFType DFOperation(FDFType Lhs, float Rhs)     { return DFOperation(Lhs, DFBroadcastFunc(Rhs)); }
+#define DEFINE_BROADCASTED_OP(FReturnType, DFOperation) \
+	FReturnType DFOperation(FDFScalar Lhs, FDFType Rhs) { return DFOperation(DFBroadcastFunc(Lhs), Rhs); } \
+	FReturnType DFOperation(float Lhs, FDFType Rhs) { return DFOperation(DFBroadcastFunc(Lhs), Rhs); } \
+	FReturnType DFOperation(FDFScalar Lhs, FFloatType Rhs) { return DFOperation(DFBroadcastFunc(Lhs), Rhs); } \
+	FReturnType DFOperation(float Lhs, FFloatType Rhs) { return DFOperation(DFBroadcastFunc(Lhs), Rhs); } \
+	FReturnType DFOperation(FDFType Lhs, FDFScalar Rhs) { return DFOperation(Lhs, DFBroadcastFunc(Rhs)); } \
+	FReturnType DFOperation(FDFType Lhs, float Rhs)     { return DFOperation(Lhs, DFBroadcastFunc(Rhs)); } \
+	FReturnType DFOperation(FFloatType Lhs, FDFScalar Rhs) { return DFOperation(Lhs, DFBroadcastFunc(Rhs)); } \
+	FReturnType DFOperation(FFloatType Lhs, float Rhs) { return DFOperation(Lhs, DFBroadcastFunc(Rhs)); }
+
+DEFINE_BROADCASTED_OP(FDFType	, DFAdd)
+DEFINE_BROADCASTED_OP(FFloatType, DFAddDemote)
+DEFINE_BROADCASTED_OP(FDFType	, DFFastAdd)
+DEFINE_BROADCASTED_OP(FFloatType, DFFastAddDemote)
+DEFINE_BROADCASTED_OP(FDFType	, DFSubtract)
+DEFINE_BROADCASTED_OP(FFloatType, DFSubtractDemote)
+DEFINE_BROADCASTED_OP(FDFType	, DFFastSubtract)
+DEFINE_BROADCASTED_OP(FFloatType, DFFastSubtractDemote)
+DEFINE_BROADCASTED_OP(FDFType	, DFMultiply)
+DEFINE_BROADCASTED_OP(FFloatType, DFMultiplyDemote)
+DEFINE_BROADCASTED_OP(FDFType	, DFDivide)
+DEFINE_BROADCASTED_OP(FDFType	, DFFastDivide)
+DEFINE_BROADCASTED_OP(FFloatType, DFFastDivideDemote)
 
-DEFINE_BROADCASTED_OP(DFAdd)
-DEFINE_BROADCASTED_OP(DFFastAdd)
-DEFINE_BROADCASTED_OP(DFSubtract)
-DEFINE_BROADCASTED_OP(DFFastSubtract)
-DEFINE_BROADCASTED_OP(DFMultiply)
-DEFINE_BROADCASTED_OP(DFDivide)
-DEFINE_BROADCASTED_OP(DFFastDivide)
 
 #undef DEFINE_BROADCASTED_OP
 
