Path Tracing: Use new DoubleFloat math library for intersecting sky atmosphere planet sphere

Remove the old header as it is now redundant with the new library

diff --git a/Engine/Shaders/Private/DoubleWordMath.ush b/Engine/Shaders/Private/DoubleWordMath.ush
deleted file mode 100644
diff --git a/Engine/Shaders/Private/PathTracing/Volume/PathTracingAtmosphere.ush b/Engine/Shaders/Private/PathTracing/Volume/PathTracingAtmosphere.ush
--- a/Engine/Shaders/Private/PathTracing/Volume/PathTracingAtmosphere.ush
+++ b/Engine/Shaders/Private/PathTracing/Volume/PathTracingAtmosphere.ush
@@ -3,7 +3,7 @@
 #pragma once
 
 #include "./PathTracingAtmosphereCommon.ush"
-#include "../../DoubleWordMath.ush"
+#include "../../DoubleFloat.ush"
 #include "../../SkyAtmosphereCommon.ush"
 
 
@@ -25,55 +25,42 @@ float3 RayPlanetIntersectPrecise(float3 Ro, float3 Rd)
 {
 	float3 Result = -1.0;
 
-	DWFloat CenterX = DWFloatCreate(PlanetCenterTranslatedWorldHi.x, PlanetCenterTranslatedWorldLo.x);
-	DWFloat CenterY = DWFloatCreate(PlanetCenterTranslatedWorldHi.y, PlanetCenterTranslatedWorldLo.y);
-	DWFloat CenterZ = DWFloatCreate(PlanetCenterTranslatedWorldHi.z, PlanetCenterTranslatedWorldLo.z);
-	DWFloat PlanetRadius = DWMul_ff(View.SkyAtmosphereBottomRadiusKm, SKY_UNIT_TO_CM);
-	DWFloat AtmoRadius   = DWMul_ff(View.SkyAtmosphereTopRadiusKm   , SKY_UNIT_TO_CM);
+	const FDFVector3 Center = MakeDFVector3(PlanetCenterTranslatedWorldHi, PlanetCenterTranslatedWorldLo);
+	const FDFScalar PlanetRadius = DFMultiply(View.SkyAtmosphereBottomRadiusKm, SKY_UNIT_TO_CM);
+	const FDFScalar AtmoRadius   = DFMultiply(View.SkyAtmosphereTopRadiusKm   , SKY_UNIT_TO_CM);
 
-	const DWFloat CoX = DWSub_df(CenterX, Ro.x);
-	const DWFloat CoY = DWSub_df(CenterY, Ro.y);
-	const DWFloat CoZ = DWSub_df(CenterZ, Ro.z);
-	const DWFloat b = DWAdd_dd(DWAdd_dd(
-		DWMul_df(CoX, Rd.x),
-		DWMul_df(CoY, Rd.y)),
-		DWMul_df(CoZ, Rd.z));
-	const DWFloat Hx = DWSub_dd(CoX, DWMul_df(b, Rd.x));
-	const DWFloat Hy = DWSub_dd(CoY, DWMul_df(b, Rd.y));
-	const DWFloat Hz = DWSub_dd(CoZ, DWMul_df(b, Rd.z));
-	const DWFloat hd = DWLengthSquared_ddd(Hx, Hy, Hz);
-	const DWFloat AR2 = DWSqr_d(AtmoRadius);
-	DWFloat ha = DWSub_dd(AR2, hd);
-	if (DWGreaterThan(ha, 0.0))
+	const FDFVector3 Co = DFSubtract(Center, Ro);
+	const FDFScalar b = DFDot(Co, Rd);
+	const FDFVector3 H = DFSubtract(Co, DFMultiply(Rd, b));
+	const FDFScalar hd = DFLengthSqr(H);
+	const FDFScalar AR2 = DFSqr(AtmoRadius);
+	FDFScalar ha = DFSubtract(AR2, hd);
+	if (DFGreater(ha, 0.0))
 	{
-		ha = DWSqrt_d(ha);
-		DWFloat q;
-		if (DWGreaterThan(b, 0.0))
-			q = DWAdd_dd(b, ha);
+		ha = DFSqrt(ha);
+		FDFScalar q;
+		if (DFGreater(b, 0.0))
+			q = DFAdd(b, ha);
 		else
-			q = DWSub_dd(b, ha);
-		const DWFloat ta = q;
-		const DWFloat tb = DWDiv_dd(DWSub_dd(DWLengthSquared_ddd(CoX, CoY, CoZ), AR2), q);
-		const float t0 = ta.GetFloat();
-		const float t1 = tb.GetFloat();
+			q = DFSubtract(b, ha);
+		const float t0 = DFDemote(q);
+		const float t1 = DFFastDivideDemote(DFSubtract(DFLengthSqr(Co), AR2), q);
 		Result.y = min(t0, t1);
 		Result.z = max(t0, t1);
 	}
-	const DWFloat PR2 = DWSqr_d(PlanetRadius);
-	DWFloat hp = DWSub_dd(PR2, hd);
-	if (DWGreaterThan(hp, 0.0))
+	const FDFScalar PR2 = DFSqr(PlanetRadius);
+	FDFScalar hp = DFSubtract(PR2, hd);
+	if (DFGreater(hp, 0.0))
 	{
-		hp = DWSqrt_d(hp);
-		DWFloat q;
-		if (DWGreaterThan(b, 0.0))
-			q = DWAdd_dd(b, hp);
+		hp = DFSqrt(hp);
+		FDFScalar q;
+		if (DFGreater(b, 0.0))
+			q = DFAdd(b, hp);
 		else
-			q = DWSub_dd(b, hp);
-		const DWFloat ta = q;
-		const DWFloat tb = DWDiv_dd(DWSub_dd(DWLengthSquared_ddd(CoX, CoY, CoZ), PR2), q);
-		float t0 = ta.GetFloat();
-		float t1 = tb.GetFloat();
-		Result.x = min(t0, t1);
+			q = DFSubtract(b, hp);
+		const float t0 = DFDemote(q);
+		const float t1 = DFFastDivideDemote(DFSubtract(DFLengthSqr(Co), PR2), q);
+			Result.x = min(t0, t1);
 		// clip volume layer against the back of the planet in case we started inside
 		if (Result.x < 0.0)
 			Result.y = max(t0, t1);
@@ -183,21 +170,13 @@ FSphericalPlanetParameters GetSphericalPlanetParameters(float3 Origin, float3 Di
 	FSphericalPlanetParameters Result = (FSphericalPlanetParameters)0;
 #if PLANET_ISECT_USES_DOUBLE_WORD_MATH
 	// Compute the origin's height/angle with high precision since it involves distances to the planet center
-	const DWFloat CenterX = DWFloatCreate(PlanetCenterTranslatedWorldHi.x, PlanetCenterTranslatedWorldLo.x);
-	const DWFloat CenterY = DWFloatCreate(PlanetCenterTranslatedWorldHi.y, PlanetCenterTranslatedWorldLo.y);
-	const DWFloat CenterZ = DWFloatCreate(PlanetCenterTranslatedWorldHi.z, PlanetCenterTranslatedWorldLo.z);
-	const DWFloat OCx = DWMul_df(DWSub_fd(Origin.x, CenterX), CM_TO_SKY_UNIT);
-	const DWFloat OCy = DWMul_df(DWSub_fd(Origin.y, CenterY), CM_TO_SKY_UNIT);
-	const DWFloat OCz = DWMul_df(DWSub_fd(Origin.z, CenterZ), CM_TO_SKY_UNIT);
-	const DWFloat DWViewHeight = DWSqrt_d(DWLengthSquared_ddd(OCx, OCy, OCz));
-	const DWFloat DWViewZenithCosAngle = DWDiv_dd(DWAdd_dd(DWAdd_dd(
-		DWMul_df(OCx, Direction.x),
-		DWMul_df(OCy, Direction.y)),
-		DWMul_df(OCz, Direction.z)),
-		DWViewHeight);
+	const FDFVector3 Center = MakeDFVector3(PlanetCenterTranslatedWorldHi, PlanetCenterTranslatedWorldLo);
 
-	Result.HeightKm = DWViewHeight.GetFloat();
-	Result.ZenithCosAngle = DWViewZenithCosAngle.GetFloat();
+	const FDFVector3 OC = DFMultiply(DFSubtract(Origin, Center), CM_TO_SKY_UNIT);
+
+	const FDFScalar HeightKm = DFLength(OC);
+	Result.HeightKm = DFDemote(HeightKm);
+	Result.ZenithCosAngle = DFFastDivideDemote(DFDot(OC, Direction), HeightKm);
 #else
 	// Get the ViewHeight / CosAngle of the origin
 	const float3 PosFromCenter = Origin - (PlanetCenterTranslatedWorldHi + PlanetCenterTranslatedWorldLo);
@@ -267,7 +246,7 @@ FVolumeDensityBounds AtmosphereGetDensityBounds(float3 Origin, float3 Direction,
 FVolumeShadedResult AtmosphereGetDensity(float3 TranslatedWorldPos)
 {
 #if 0 //PLANET_ISECT_USES_DOUBLE_WORD_MATH
-	// Extra precision here doesn't seem to matter too much -- probably because density function is very smooth, so the exact position doesn't matter too much
+	// Extra precision here doesn't seem to matter too much -- probably because density function is very smooth
 	const float H = GetSphericalPlanetParameters(TranslatedWorldPos, float3(0, 0, 1)).HeightKm;
 	MediumSampleRGB Medium = SampleAtmosphereMediumRGB(float3(0.0, 0.0, H));
 #else
