RayTracing: Fix alpha channel in RT translucency and PrimaryRay debug view and add Holdout support

The alpha channel is now tracked like in other passes (it was not working consistently before). This applies both to RT translucency and PrimaryRay debug view.

A holdout flag is now propagated from the shader to the raygen shader so that holdouts can be tracked properly. This involved also flagging camera rays (as the holdout logic should only apply in these cases).

Also fix alpha channel of RT debug views for consistency

diff --git a/Engine/Shaders/Private/RayTracing/RayTracingBarycentrics.usf b/Engine/Shaders/Private/RayTracing/RayTracingBarycentrics.usf
--- a/Engine/Shaders/Private/RayTracing/RayTracingBarycentrics.usf
+++ b/Engine/Shaders/Private/RayTracing/RayTracingBarycentrics.usf
@@ -81,7 +81,7 @@ void RayTracingBarycentricsCommon(uint2 DispatchThreadIndex)
 	float4 Result = float4(0, 0, 0, 1);
 	if (Payload.IsHit())
 	{
-		Result = float4(Payload.Barycentrics, 1.0f - Payload.Barycentrics.x - Payload.Barycentrics.y, 1.0f);
+		Result = float4(Payload.Barycentrics, 1.0f - Payload.Barycentrics.x - Payload.Barycentrics.y, 0.0f);
 	}
 	Output[PixelCoord] = Result;
 }
diff --git a/Engine/Shaders/Private/RayTracing/RayTracingCommon.ush b/Engine/Shaders/Private/RayTracing/RayTracingCommon.ush
--- a/Engine/Shaders/Private/RayTracing/RayTracingCommon.ush
+++ b/Engine/Shaders/Private/RayTracing/RayTracingCommon.ush
@@ -173,7 +173,7 @@ void name##_INTERNAL(uint3 DispatchThreadIndex)
 #define RAY_TRACING_PAYLOAD_OUTPUT_FLAG_NONE                               0
 #define RAY_TRACING_PAYLOAD_OUTPUT_FLAG_FRONT_FACE                   (1 << 0) // Indicates that ray has hit front face of a primitive. Set by closest hit shader.
 #define RAY_TRACING_PAYLOAD_OUTPUT_FLAG_TWO_SIDED                    (1 << 1) // Ray hit a two-sided material
-#define RAY_TRACING_PAYLOAD_OUTPUT_FLAG_UNUSED2                      (1 << 2) // Free for future use
+#define RAY_TRACING_PAYLOAD_OUTPUT_FLAG_IS_HOLDOUT                   (1 << 2) // We hit a primitive flagged as a holdout
 #define RAY_TRACING_PAYLOAD_OUTPUT_FLAG_UNUSED3                      (1 << 3) // Free for future use
 #define RAY_TRACING_PAYLOAD_OUTPUT_FLAG_UNUSED4                      (1 << 4) // Free for future use
 #define RAY_TRACING_PAYLOAD_OUTPUT_FLAG_UNUSED5                      (1 << 5) // Free for future use
@@ -184,7 +184,7 @@ void name##_INTERNAL(uint3 DispatchThreadIndex)
 #define RAY_TRACING_PAYLOAD_INPUT_FLAG_LUMEN_PAYLOAD		         (1 << 2) // Indicates that closest hit shader should evaluate geometry normal and instanceid.
 #define RAY_TRACING_PAYLOAD_INPUT_FLAG_IGNORE_TRANSLUCENT            (1 << 3) // Indicates that translucent materials should be ignored by IgnoreHit() in anyhit shaders.
 #define RAY_TRACING_PAYLOAD_INPUT_FLAG_SHADOW_RAY                    (1 << 4) // Indicates this ray is being traced as part of a raytraced shadow calculation (see TraceVisibilityRay).
-#define RAY_TRACING_PAYLOAD_INPUT_FLAG_UNUSED5                       (1 << 5) // Free for future use
+#define RAY_TRACING_PAYLOAD_INPUT_FLAG_CAMERA_RAY                    (1 << 5) // Indicates this ray is being traced directly from the camera
 
 ///////////////////////////////////////////////////////////////////////////////////////////////////
 // Instance descriptor
@@ -604,8 +604,12 @@ struct FPackedMaterialClosestHitPayload : FMinimalPayload
 	void SetIgnoreTranslucentMaterials()	{ Flags |= RAY_TRACING_PAYLOAD_INPUT_FLAG_IGNORE_TRANSLUCENT; }
 	bool IsIgnoreTranslucentMaterials()		{ return (Flags & RAY_TRACING_PAYLOAD_INPUT_FLAG_IGNORE_TRANSLUCENT) != 0; }
 
+	void SetShadowRay()                     { Flags |= RAY_TRACING_PAYLOAD_INPUT_FLAG_SHADOW_RAY; }
 	bool IsShadowRay()						{ return (Flags & RAY_TRACING_PAYLOAD_INPUT_FLAG_SHADOW_RAY) != 0; }
 
+	void SetCameraRay()                     { Flags |= RAY_TRACING_PAYLOAD_INPUT_FLAG_CAMERA_RAY; }
+	bool IsCameraRay()						{ return (Flags & RAY_TRACING_PAYLOAD_INPUT_FLAG_CAMERA_RAY) != 0; }
+
 	FRayCone GetRayCone()					{ return UnpackRayCone(PackedRayCone); }
 	void SetRayCone(FRayCone NewRayCone)	{ PackedRayCone = PackRayCone(NewRayCone); }
 
@@ -739,6 +743,9 @@ struct FMaterialClosestHitPayload : FMinimalPayload
 
 	void SetTwoSided() { Flags |= RAY_TRACING_PAYLOAD_OUTPUT_FLAG_TWO_SIDED; }
 	bool IsTwoSided() { return (Flags & RAY_TRACING_PAYLOAD_OUTPUT_FLAG_TWO_SIDED) != 0; }
+
+	void SetHoldout() { Flags |= RAY_TRACING_PAYLOAD_OUTPUT_FLAG_IS_HOLDOUT; }
+	bool IsHoldout() { return (Flags & RAY_TRACING_PAYLOAD_OUTPUT_FLAG_IS_HOLDOUT) != 0; }
 	
 	FRayCone GetRayCone() { return RayCone; }
 	void SetRayCone(FRayCone NewRayCone) { RayCone = NewRayCone; }
@@ -777,7 +784,11 @@ struct FPackedMaterialClosestHitPayload : FMinimalPayload
 	void SetIgnoreTranslucentMaterials() { IorAndShadingModelIDAndBlendingModeAndFlagsAndPrimitiveLightingChannelMask |= RAY_TRACING_PAYLOAD_INPUT_FLAG_IGNORE_TRANSLUCENT << 23; }
 	bool IsIgnoreTranslucentMaterials() { return ((IorAndShadingModelIDAndBlendingModeAndFlagsAndPrimitiveLightingChannelMask >> 23) & RAY_TRACING_PAYLOAD_INPUT_FLAG_IGNORE_TRANSLUCENT) != 0; }
 
-	bool IsShadowRay() { return ((IorAndShadingModelIDAndBlendingModeAndFlagsAndPrimitiveLightingChannelMask >> 23) & RAY_TRACING_PAYLOAD_INPUT_FLAG_SHADOW_RAY) != 0; }
+	void SetShadowRay()                     { IorAndShadingModelIDAndBlendingModeAndFlagsAndPrimitiveLightingChannelMask |= RAY_TRACING_PAYLOAD_INPUT_FLAG_SHADOW_RAY << 23; }
+	bool IsShadowRay()						{ return ((IorAndShadingModelIDAndBlendingModeAndFlagsAndPrimitiveLightingChannelMask >> 23) & RAY_TRACING_PAYLOAD_INPUT_FLAG_SHADOW_RAY) != 0; }
+
+	void SetCameraRay()                     { IorAndShadingModelIDAndBlendingModeAndFlagsAndPrimitiveLightingChannelMask |= RAY_TRACING_PAYLOAD_INPUT_FLAG_CAMERA_RAY << 23; }
+	bool IsCameraRay()						{ return ((IorAndShadingModelIDAndBlendingModeAndFlagsAndPrimitiveLightingChannelMask >> 23) & RAY_TRACING_PAYLOAD_INPUT_FLAG_CAMERA_RAY) != 0; }
 
 	FRayCone GetRayCone() { return UnpackRayCone(PackedRayCone); }
 	void SetRayCone(FRayCone NewRayCone) { PackedRayCone = PackRayCone(NewRayCone); }
@@ -1188,7 +1199,8 @@ void TraceMaterialRayPacked(
 	in FRayDesc Ray,
 	// Payload Inputs
 	inout FRayCone RayCone,
-	in bool bEnableSkyLightContribution)
+	in bool bEnableSkyLightContribution,
+	in bool bIsCameraRay)
 {
 	const uint RayContributionToHitGroupIndex = RAY_TRACING_SHADER_SLOT_MATERIAL;
 	const uint MultiplierForGeometryContributionToShaderIndex = RAY_TRACING_NUM_SHADER_SLOTS;
@@ -1204,6 +1216,11 @@ void TraceMaterialRayPacked(
 		Payload.SetEnableSkyLightContribution();
 	}
 
+	if (bIsCameraRay)
+	{
+		Payload.SetCameraRay();
+	}
+
 
 	// Trace the ray
 
diff --git a/Engine/Shaders/Private/RayTracing/RayTracingDebug.usf b/Engine/Shaders/Private/RayTracing/RayTracingDebug.usf
--- a/Engine/Shaders/Private/RayTracing/RayTracingDebug.usf
+++ b/Engine/Shaders/Private/RayTracing/RayTracingDebug.usf
@@ -46,18 +46,16 @@ float4 AnisoDebugColor(float InAniso)
 	// Follow same mapping than raster buffer visualization (BP located in Engine/Contents/Anisotropy)
 	float AniG = saturate(InAniso);
 	float AniB = saturate(-1.0 * InAniso);
-	return float4(0.0, pow(AniG, 2.2), pow(AniB, 2.2), 1.0f);
+	return float4(0.0, pow(AniG, 2.2), pow(AniB, 2.2), 0.0f);
 }
 
 float4 RoughnessDebugColor(float InRoughness)
 {
-	return float4(pow(InRoughness, 2.2f).xxx, 1.0f);
+	return float4(pow(InRoughness, 2.2f).xxx, 0.0f);
 }
 
 float4 DebugRayTracingMaterial(FRayDesc Ray, bool bIsFarField)
 {
-	uint2 PixelCoord = DispatchRaysIndex().xy + View.ViewRectMin.xy;
-
 	FRayCone RayCone = (FRayCone)0;
 	RayCone.SpreadAngle = View.EyeToPixelSpreadAngle;
 
@@ -99,16 +97,16 @@ float4 DebugRayTracingMaterial(FRayDesc Ray, bool bIsFarField)
 		// 1. Non-material dependent debug view
 		switch (VisualizationMode)
 		{
-		case RAY_TRACING_DEBUG_VIZ_RADIANCE:				return float4(Payload.Radiance, 1.0f);
-		case RAY_TRACING_DEBUG_VIZ_OPACITY:					return float4((1.0f - Payload.Opacity).xxx, 1.0f);
-		case RAY_TRACING_DEBUG_VIZ_BLENDING_MODE:			return float4(Payload.BlendingMode.xxx, 1.0f);
-		case RAY_TRACING_DEBUG_VIZ_LIGHTING_CHANNEL_MASK:	return float4(Payload.PrimitiveLightingChannelMask & 0x1, Payload.PrimitiveLightingChannelMask & 0x2, Payload.PrimitiveLightingChannelMask & 0x4, 1.0f);
-		case RAY_TRACING_DEBUG_VIZ_INDIRECT_IRRADIANCE:		return float4(Payload.IndirectIrradiance, 1.0f);
-		case RAY_TRACING_DEBUG_VIZ_WORLD_POSITION:			return float4(Payload.TranslatedWorldPos, 1.0f);
-		case RAY_TRACING_DEBUG_VIZ_HITKIND:					return Payload.IsFrontFace() ? float4(0.0, 1.0, 0.0, 1.0f) : float4(1.0, 0.0, 0.0, 1.0f);
-		case RAY_TRACING_DEBUG_VIZ_LIGHT_GRID_COUNT:	    return float4(FCulledLightList::Create(Payload.TranslatedWorldPos).Visualize(1) * abs(dot(Payload.WorldNormal, Ray.Direction)), 1.0f);
+		case RAY_TRACING_DEBUG_VIZ_RADIANCE:				return float4(Payload.Radiance, 0.0f);
+		case RAY_TRACING_DEBUG_VIZ_OPACITY:					return float4((1.0f - Payload.Opacity).xxx, 0.0f);
+		case RAY_TRACING_DEBUG_VIZ_BLENDING_MODE:			return float4(Payload.BlendingMode.xxx, 0.0f);
+		case RAY_TRACING_DEBUG_VIZ_LIGHTING_CHANNEL_MASK:	return float4(Payload.PrimitiveLightingChannelMask & 0x1, Payload.PrimitiveLightingChannelMask & 0x2, Payload.PrimitiveLightingChannelMask & 0x4, 0.0f);
+		case RAY_TRACING_DEBUG_VIZ_INDIRECT_IRRADIANCE:		return float4(Payload.IndirectIrradiance, 0.0f);
+		case RAY_TRACING_DEBUG_VIZ_WORLD_POSITION:			return float4(Payload.TranslatedWorldPos, 0.0f);
+		case RAY_TRACING_DEBUG_VIZ_HITKIND:					return Payload.IsFrontFace() ? float4(0.0, 1.0, 0.0, 0.0f) : float4(1.0, 0.0, 0.0, 0.0f);
+		case RAY_TRACING_DEBUG_VIZ_LIGHT_GRID_COUNT:	    return float4(FCulledLightList::Create(Payload.TranslatedWorldPos).Visualize(1) * abs(dot(Payload.WorldNormal, Ray.Direction)), 0.0f);
 		#if PLATFORM_SUPPORTS_SHADER_TIMESTAMP
-		case RAY_TRACING_DEBUG_VIZ_PERFORMANCE:				return float4(ColorMapInferno(float(TimeEnd - TimeBegin) * TimingScale), 1);
+		case RAY_TRACING_DEBUG_VIZ_PERFORMANCE:				return float4(ColorMapInferno(float(TimeEnd - TimeBegin) * TimingScale), 0.0f);
 		#endif // PLATFORM_SUPPORTS_SHADER_TIMESTAMP
 		}
 
@@ -125,51 +123,51 @@ float4 DebugRayTracingMaterial(FRayDesc Ray, bool bIsFarField)
 			switch (VisualizationMode)
 			{				
 			case RAY_TRACING_DEBUG_VIZ_FAR_FIELD:		return float4(SubstrateGetBSDFBaseColor(BSDF), Payload.HitT);
-			case RAY_TRACING_DEBUG_VIZ_GBUFFER_AO:		return float4(SubstrateGetIrradianceAndAO(SubstratePixelHeader).MaterialAO.xxx, 1.0f);
-			case RAY_TRACING_DEBUG_VIZ_SHADING_MODEL:	return float4(GetShadingModelColor(SubstrateGetLegacyShadingModels(BSDF)), 1.0f);
-			case RAY_TRACING_DEBUG_VIZ_BASE_COLOR:		return float4(SubstrateGetBSDFBaseColor(BSDF), 1.0f);
-			case RAY_TRACING_DEBUG_VIZ_DIFFUSE_COLOR:	return float4(SubstrateGetBSDFDiffuseColor(BSDF), 1.0f);
-			case RAY_TRACING_DEBUG_VIZ_SPECULAR_COLOR:	return float4(SubstrateGetBSDFSpecularColor(BSDF).xxx, 1.0f);
-			case RAY_TRACING_DEBUG_VIZ_METALLIC:		return float4(SubstrateGetBSDFMetallic(BSDF).xxx, 1.0f);
-			case RAY_TRACING_DEBUG_VIZ_SPECULAR:		return float4(SubstrateGetBSDFSpecular(BSDF).xxx, 1.0f);
+			case RAY_TRACING_DEBUG_VIZ_GBUFFER_AO:		return float4(SubstrateGetIrradianceAndAO(SubstratePixelHeader).MaterialAO.xxx, 0.0f);
+			case RAY_TRACING_DEBUG_VIZ_SHADING_MODEL:	return float4(GetShadingModelColor(SubstrateGetLegacyShadingModels(BSDF)), 0.0f);
+			case RAY_TRACING_DEBUG_VIZ_BASE_COLOR:		return float4(SubstrateGetBSDFBaseColor(BSDF), 0.0f);
+			case RAY_TRACING_DEBUG_VIZ_DIFFUSE_COLOR:	return float4(SubstrateGetBSDFDiffuseColor(BSDF), 0.0f);
+			case RAY_TRACING_DEBUG_VIZ_SPECULAR_COLOR:	return float4(SubstrateGetBSDFSpecularColor(BSDF).xxx, 0.0f);
+			case RAY_TRACING_DEBUG_VIZ_METALLIC:		return float4(SubstrateGetBSDFMetallic(BSDF).xxx, 0.0f);
+			case RAY_TRACING_DEBUG_VIZ_SPECULAR:		return float4(SubstrateGetBSDFSpecular(BSDF).xxx, 0.0f);
 			case RAY_TRACING_DEBUG_VIZ_ROUGHNESS:		return RoughnessDebugColor(SubstrateGetBSDFRoughness(BSDF));
-			case RAY_TRACING_DEBUG_VIZ_IOR:				return float4(1, 1, 1, 1.0f);
-			case RAY_TRACING_DEBUG_VIZ_CUSTOM_DATA:		return float4(0, 0, 0, 0);
-			case RAY_TRACING_DEBUG_VIZ_WORLD_NORMAL:	return float4(TangentBasis[2] * 0.5 + 0.5, 1.0f);
-			case RAY_TRACING_DEBUG_VIZ_WORLD_TANGENT:	return float4(TangentBasis[0] * 0.5 + 0.5, 1.0f);
+			case RAY_TRACING_DEBUG_VIZ_IOR:				return float4(1, 1, 1, 0.0f);
+			case RAY_TRACING_DEBUG_VIZ_CUSTOM_DATA:		return float4(0, 0, 0, 0.0f);
+			case RAY_TRACING_DEBUG_VIZ_WORLD_NORMAL:	return float4(TangentBasis[2] * 0.5 + 0.5, 0.0f);
+			case RAY_TRACING_DEBUG_VIZ_WORLD_TANGENT:	return float4(TangentBasis[0] * 0.5 + 0.5, 0.0f);
 			case RAY_TRACING_DEBUG_VIZ_ANISOTROPY:		return AnisoDebugColor(SubstrateGetBSDFAnisotropy(BSDF));
 			}
 		}
 	#else // SUBSTRATE_ENABLED
 		switch (VisualizationMode)
 		{
-			case RAY_TRACING_DEBUG_VIZ_WORLD_NORMAL:	return float4(Payload.WorldNormal * 0.5 + 0.5, 1.0f);
+			case RAY_TRACING_DEBUG_VIZ_WORLD_NORMAL:	return float4(Payload.WorldNormal * 0.5 + 0.5, 0.0f);
 			case RAY_TRACING_DEBUG_VIZ_FAR_FIELD:		return float4(Payload.BaseColor.rgb, Payload.HitT);
-			case RAY_TRACING_DEBUG_VIZ_GBUFFER_AO:		return float4(Payload.GBufferAO, Payload.GBufferAO, Payload.GBufferAO, 1.0f);
-			case RAY_TRACING_DEBUG_VIZ_SHADING_MODEL:	return float4(GetShadingModelColor(Payload.ShadingModelID), 1.0f);
-			case RAY_TRACING_DEBUG_VIZ_BASE_COLOR:		return float4(Payload.BaseColor.rgb, 1.0f);
-			case RAY_TRACING_DEBUG_VIZ_DIFFUSE_COLOR:	return float4(Payload.DiffuseColor.rgb, 1.0f);
-			case RAY_TRACING_DEBUG_VIZ_SPECULAR_COLOR:	return float4(Payload.SpecularColor.rgb, 1.0f);
-			case RAY_TRACING_DEBUG_VIZ_METALLIC:		return float4(Payload.Metallic, Payload.Metallic, Payload.Metallic, 1.0f);
-			case RAY_TRACING_DEBUG_VIZ_SPECULAR:		return float4(Payload.Specular, Payload.Specular, Payload.Specular, 1.0f);
+			case RAY_TRACING_DEBUG_VIZ_GBUFFER_AO:		return float4(Payload.GBufferAO, Payload.GBufferAO, Payload.GBufferAO, 0.0f);
+			case RAY_TRACING_DEBUG_VIZ_SHADING_MODEL:	return float4(GetShadingModelColor(Payload.ShadingModelID), 0.0f);
+			case RAY_TRACING_DEBUG_VIZ_BASE_COLOR:		return float4(Payload.BaseColor.rgb, 0.0f);
+			case RAY_TRACING_DEBUG_VIZ_DIFFUSE_COLOR:	return float4(Payload.DiffuseColor.rgb, 0.0f);
+			case RAY_TRACING_DEBUG_VIZ_SPECULAR_COLOR:	return float4(Payload.SpecularColor.rgb, 0.0f);
+			case RAY_TRACING_DEBUG_VIZ_METALLIC:		return float4(Payload.Metallic, Payload.Metallic, Payload.Metallic, 0.0f);
+			case RAY_TRACING_DEBUG_VIZ_SPECULAR:		return float4(Payload.Specular, Payload.Specular, Payload.Specular, 0.0f);
 			case RAY_TRACING_DEBUG_VIZ_ROUGHNESS:		return RoughnessDebugColor(Payload.Roughness);
-			case RAY_TRACING_DEBUG_VIZ_IOR:				return float4(Payload.Ior, Payload.Ior, Payload.Ior, 1.0f);
-			case RAY_TRACING_DEBUG_VIZ_CUSTOM_DATA:		return float4(Payload.CustomData);
-			case RAY_TRACING_DEBUG_VIZ_WORLD_TANGENT:	return float4(Payload.WorldTangent * 0.5 + 0.5, 1.0f);
+			case RAY_TRACING_DEBUG_VIZ_IOR:				return float4(Payload.Ior, Payload.Ior, Payload.Ior, 0.0f);
+			case RAY_TRACING_DEBUG_VIZ_CUSTOM_DATA:		return float4(Payload.CustomData.xyz, 0.0f);
+			case RAY_TRACING_DEBUG_VIZ_WORLD_TANGENT:	return float4(Payload.WorldTangent * 0.5 + 0.5, 0.0f);
 			case RAY_TRACING_DEBUG_VIZ_ANISOTROPY:		return AnisoDebugColor(Payload.Anisotropy);
 		}
 	#endif // SUBSTRATE_ENABLED
 	}
 
 	// HitT is packed into the alpha for far field compositing
-	return (RAY_TRACING_DEBUG_VIZ_FAR_FIELD) ? float4(0, 0, 0, -1) : float4(0, 0, 0, 1);
+	return (VisualizationMode == RAY_TRACING_DEBUG_VIZ_FAR_FIELD) ? float4(0, 0, 0, -1.0) : float4(0, 0, 0, 1);
 }
 
 #else // USE_DEBUG_CHS
 
 // Add functions that use the debug payload only
 
-float4 DebugRayTracingInstance(FRayDesc Ray, uint2 PixelCoord, out float OutDepth)
+float4 DebugRayTracingInstance(FRayDesc Ray, out float OutDepth)
 {
 	float4 Result = float4(0, 0, 0, 1);
 
@@ -265,6 +263,7 @@ float4 DebugRayTracingInstance(FRayDesc Ray, uint2 PixelCoord, out float OutDept
 		float3 TranslatedWorldPosition = Ray.Origin + Payload.HitT * Ray.Direction;
 		float4 ClipPosition = mul(float4(TranslatedWorldPosition, 1.0), View.TranslatedWorldToClip);
 		OutDepth = ClipPosition.z / ClipPosition.w;
+		Result.a = 0.0;
 	}
 
 	return Result;
@@ -290,14 +289,14 @@ RAY_TRACING_ENTRY_RAYGEN(RayTracingDebugMainRGS)
 
 #if USE_DEBUG_CHS
 	{
-		Result = DebugRayTracingInstance(Ray, PixelCoord, OutDepth);
+		Result = DebugRayTracingInstance(Ray, OutDepth);
 	}
 #else // USE_DEBUG_CHS
 	{
 		bool bIsFarField = false;
 		Result = DebugRayTracingMaterial(Ray, bIsFarField);
 		
-		const bool bNearFieldMiss = Result.w <= 0.0;
+		const bool bNearFieldMiss = Result.w < 0.0;
 		if (VisualizationMode == RAY_TRACING_DEBUG_VIZ_FAR_FIELD && bNearFieldMiss)
 		{
 			bIsFarField = true;
@@ -310,7 +309,7 @@ RAY_TRACING_ENTRY_RAYGEN(RayTracingDebugMainRGS)
 	if (ShouldUsePreExposure)
 	{
 		// Only  when the output is tonemapped
-		Result *= View.PreExposure;
+		Result.rgb *= View.PreExposure;
 	}
 
 	if (!IsFinite(Result.r) || !IsFinite(Result.g) || !IsFinite(Result.b))
@@ -321,6 +320,8 @@ RAY_TRACING_ENTRY_RAYGEN(RayTracingDebugMainRGS)
 			Result.rgb = float3(1,0,1);
 		}
 	}
+
+	Result.a = Result.a;
 	
 	Output[PixelCoord] = Result;
 
diff --git a/Engine/Shaders/Private/RayTracing/RayTracingLightingCommon.ush b/Engine/Shaders/Private/RayTracing/RayTracingLightingCommon.ush
--- a/Engine/Shaders/Private/RayTracing/RayTracingLightingCommon.ush
+++ b/Engine/Shaders/Private/RayTracing/RayTracingLightingCommon.ush
@@ -513,6 +513,7 @@ FMaterialClosestHitPayload TraceRayAndAccumulateResults(
 	in uint ShouldDoEmissiveAndIndirectLighting,
 	in bool bRayTraceSkyLightContribution,
 	in bool bDecoupleSampleGeneration,
+	in bool bIsCameraRay,
 	inout FRayCone RayCone,
 	in bool bEnableSkyLightContribution,
 	inout float3 Radiance)
@@ -531,7 +532,8 @@ FMaterialClosestHitPayload TraceRayAndAccumulateResults(
 		InstanceInclusionMask,
 		Ray,
 		RayCone,
-		bEnableSkyLightContribution);
+		bEnableSkyLightContribution,
+		bIsCameraRay);
 
 	float3 TranslatedWorldPosition = Ray.Origin + Ray.Direction * Payload.HitT;
 	float3 ViewDirection = Ray.Direction;
diff --git a/Engine/Shaders/Private/RayTracing/RayTracingMaterialHitShaders.usf b/Engine/Shaders/Private/RayTracing/RayTracingMaterialHitShaders.usf
--- a/Engine/Shaders/Private/RayTracing/RayTracingMaterialHitShaders.usf
+++ b/Engine/Shaders/Private/RayTracing/RayTracingMaterialHitShaders.usf
@@ -575,8 +575,10 @@ RAY_TRACING_ENTRY_CLOSEST_HIT(MaterialCHS,
 		uint MaterialBlendingMode = RAY_TRACING_BLEND_MODE_OPAQUE;
 	#endif
 
+	const uint PrimitiveFlags = GetPrimitiveData(MaterialParameters.PrimitiveId).Flags;
+
 	Payload.BlendingMode = MaterialBlendingMode;
-	Payload.PrimitiveLightingChannelMask = GetPrimitive_LightingChannelMask(MaterialParameters.PrimitiveId);
+	Payload.PrimitiveLightingChannelMask = GetPrimitive_LightingChannelMask_FromFlags(PrimitiveFlags);
 
 	Payload.HitT = RayTCurrent();
 	if (HitKind() == HIT_KIND_TRIANGLE_FRONT_FACE)
@@ -591,6 +593,22 @@ RAY_TRACING_ENTRY_CLOSEST_HIT(MaterialCHS,
 	// Fill payload with material data
 	GetMaterialPayload(PixelMaterialInputs, MaterialParameters, Interpolants, PackedPayload.IsEnableSkyLightContribution(), Payload); 
 
+	// After we have gathered the material info, check to see if we are a holdout to see if we need to overwrite the shaded result
+#if MATERIALBLENDING_ALPHAHOLDOUT
+	const bool bIsHoldout = true;
+#else
+	const bool bIsHoldout = (PrimitiveFlags & PRIMITIVE_SCENE_DATA_FLAG_HOLDOUT) != 0;
+#endif
+	const bool bIsCameraRay = PackedPayload.IsCameraRay();
+	if (bIsHoldout && bIsCameraRay)
+	{
+		Payload.ShadingModelID = SHADINGMODELID_UNLIT;
+		Payload.Radiance = 0.0;
+		Payload.BaseColor = 0.0;
+		Payload.Specular = 0.0;
+		Payload.SetHoldout();
+	}
+
 	PackedPayload = PackRayTracingPayload(Payload, PropagatedCone);
 
 	// Override packed payload with custom Lumen data
diff --git a/Engine/Shaders/Private/RayTracing/RayTracingPrimaryRays.usf b/Engine/Shaders/Private/RayTracing/RayTracingPrimaryRays.usf
--- a/Engine/Shaders/Private/RayTracing/RayTracingPrimaryRays.usf
+++ b/Engine/Shaders/Private/RayTracing/RayTracingPrimaryRays.usf
@@ -128,8 +128,8 @@ RAY_TRACING_ENTRY_RAYGEN(RayTracingPrimaryRaysRGS)
 	const bool bSkyLightAffectReflection = ShouldSkyLightAffectReflection();
 	
 	const bool bPrimaryView = (ERayTracingPrimaryRaysFlag_PrimaryView & PrimaryRayFlags) != 0;
-	bool bHasScattered = bPrimaryView;
-	float AccumulatedOpacity = 0.0;
+	bool bHasScattered = false;
+	float BackgroundVisibility = 0.0;
 
 	// Integrated data by path tracing
 	float3 PathRadiance = 0.0;
@@ -139,8 +139,9 @@ RAY_TRACING_ENTRY_RAYGEN(RayTracingPrimaryRaysRGS)
 
 	for (uint RefractionRayIndex = 0; RefractionRayIndex < MaxRefractionRays; ++RefractionRayIndex)
 	{
-		const uint RefractionRayFlags = RefractionRayIndex == 0 ? RAY_FLAG_CULL_BACK_FACING_TRIANGLES : 0;
-		const uint RefractionInstanceInclusionMask = (!bPrimaryView && RefractionRayIndex == 0) || !TranslucencyRefraction ? RAY_TRACING_MASK_TRANSLUCENT : RAY_TRACING_MASK_ALL;
+		const bool bIsCameraRay = !bHasScattered;
+		const uint RefractionRayFlags = bIsCameraRay ? RAY_FLAG_CULL_BACK_FACING_TRIANGLES : 0;
+		const uint RefractionInstanceInclusionMask = (!bPrimaryView && bIsCameraRay) || !TranslucencyRefraction ? RAY_TRACING_MASK_TRANSLUCENT : RAY_TRACING_MASK_ALL;
 		const bool bRefractionRayTraceSkyLightContribution = false;
 		const bool bRefractionDecoupleSampleGeneration = true;
 #if PROJECT_SUPPORTS_LUMEN
@@ -167,6 +168,7 @@ RAY_TRACING_ENTRY_RAYGEN(RayTracingPrimaryRaysRGS)
 			ShouldDoEmissiveAndIndirectLighting,
 			bRefractionRayTraceSkyLightContribution,
 			bRefractionDecoupleSampleGeneration,
+			bIsCameraRay,
 			RayCone,
 			bRefractionEnableSkyLightContribution,
 			PathVertexRadiance);
@@ -191,7 +193,7 @@ RAY_TRACING_ENTRY_RAYGEN(RayTracingPrimaryRaysRGS)
 		// Handle surface lighting
 		//
 
-		float vertexRadianceWeight = Payload.Opacity;	// Opacity as coverage. This works for RAY_TRACING_BLEND_MODE_OPAQUE and RAY_TRACING_BLEND_MODE_TRANSLUCENT.
+		const float vertexRadianceWeight = Payload.Opacity;	// Opacity as coverage. This works for RAY_TRACING_BLEND_MODE_OPAQUE and RAY_TRACING_BLEND_MODE_TRANSLUCENT.
 		// It is also needed for RAY_TRACING_BLEND_MODE_ADDITIVE and  RAY_TRACING_BLEND_MODE_ALPHA_COMPOSITE: radiance contribution is alway weighted by coverage.
 		// #dxr_todo: I have not been able to setup a material using RAY_TRACING_BLEND_MODE_MODULATE.
 
@@ -221,7 +223,12 @@ RAY_TRACING_ENTRY_RAYGEN(RayTracingPrimaryRaysRGS)
 
 		PathRadiance += PathThroughput * vertexRadianceWeight * (PathVertexRadiance * ViewFogInScatteringAndTransmittance.a + ViewFogInScatteringAndTransmittance.rgb);
 
-		AccumulatedOpacity += vertexRadianceWeight;
+		const bool bIsHoldout = Payload.IsHoldout();
+		if (bIsHoldout)
+		{
+			// A holdout object acts as if it was a background pixel so that it influences the alpha channel
+			BackgroundVisibility += Luminance(PathThroughput) * vertexRadianceWeight;
+		}
 
 		//
 		// Handle reflection tracing with a ray per vertex of the refraction path
@@ -265,6 +272,7 @@ RAY_TRACING_ENTRY_RAYGEN(RayTracingPrimaryRaysRGS)
 				ShouldDoEmissiveAndIndirectLighting,
 				bReflectionRayTraceSkyLightContribution,
 				bReflectionDecoupleSampleGeneration,
+				false,
 				RayCone,
 				bReflectionEnableSkyLightContribution,
 				ReflectionRadiance);
@@ -337,6 +345,7 @@ RAY_TRACING_ENTRY_RAYGEN(RayTracingPrimaryRaysRGS)
 			}
 			case RAY_TRACING_BLEND_MODE_ALPHA_COMPOSITE:
 			case RAY_TRACING_BLEND_MODE_TRANSLUCENT:
+			case RAY_TRACING_BLEND_MODE_ALPHA_HOLDOUT:
 			{
 				PathThroughput *= 1.0 - Payload.Opacity;
 				break;
@@ -370,9 +379,26 @@ RAY_TRACING_ENTRY_RAYGEN(RayTracingPrimaryRaysRGS)
 			const float3 Transmittance = Payload.CustomData.xyz;
 			const float PathLength = rcp(saturate(abs(NoV) + 1e-5));
 			PathThroughput *= pow(Transmittance, PathLength) * Pow2(1 - F);
+			bHasScattered = true; // treat the thin translucency as refraction for consistency with solid glass, even though the ray is not bent
+			if (bIsHoldout)
+			{
+				// In the refraction case, the non-opaque part should also contribute to the holdout
+				// The net result is that the entire object becomes opaque for the alpha channel.
+				BackgroundVisibility += Luminance(PathThroughput) * (1.0 - Payload.Opacity);
+				PathThroughput = 0.0;
+				break;
+			}
 		}
 		else if (TranslucencyRefraction && Payload.BlendingMode == RAY_TRACING_BLEND_MODE_TRANSLUCENT && Payload.Ior > 0.0)
 		{
+			if (bIsHoldout)
+			{
+				// In the refraction case, the non-opaque part should also contribute to the holdout
+				// The net result is that the entire object becomes opaque for the alpha channel.
+				BackgroundVisibility += Luminance(PathThroughput) * (1.0 - Payload.Opacity);
+				PathThroughput = 0.0;
+				break;
+			}
 			float Ior = Payload.Ior;
 			bHasScattered |= Ior > 1.0 ? true : false;
 
@@ -418,22 +444,24 @@ RAY_TRACING_ENTRY_RAYGEN(RayTracingPrimaryRaysRGS)
 
 	if (!bHasScattered)
 	{
-		// Use the scene radiance for ray that has not been scattered/refracted (no surface or IORin=IORout). Still apply the throughput in case we have traversed surfaces with opacity>0.
-		PathRadiance += PathThroughput * SceneColorTexture.SampleLevel(GlobalPointClampedSampler, UV, 0).xyz / View.PreExposure;
-	}
-
-	float FinalAlpha = 0.0f;
-	if (AccumulatedOpacity > 0.0f)
-	{
-		FinalAlpha = saturate(1.0 - AccumulatedOpacity);
-	}
-	else
-	{
-		FinalAlpha = SceneColorTexture.SampleLevel(GlobalPointClampedSampler, UV, 0).w;
+		// If the ray did not bend, composite over the SceneColor texture and accumulate alpha
+		if (bPrimaryView)
+		{
+			// our path reached the background, accumulate its contribution for the alpha channel
+			BackgroundVisibility += Luminance(PathThroughput);
+		}
+		else
+		{
+			float4 SceneColorContrib = SceneColorTexture.SampleLevel(GlobalPointClampedSampler, UV, 0);
+			PathRadiance += PathThroughput * SceneColorContrib.xyz / View.PreExposure;
+			BackgroundVisibility += Luminance(PathThroughput) * SceneColorContrib.w;
+		}
 	}
 
-	PathRadiance.rgb *= View.PreExposure;
+	// NOTE: UE by convention tracks background visibility which is (1.0-AccumulatedOpacity) (the complement of traditional alpha)
+	const float FinalAlpha = saturate(BackgroundVisibility);
 
+	PathRadiance *= View.PreExposure;
 	PathRadiance = ClampToHalfFloatRange(PathRadiance);
 	ColorOutput[DispatchThreadId] = float4(PathRadiance, FinalAlpha);
 	RayHitDistanceOutput[DispatchThreadId] = HitDistance;
