Hotfix for cooking issue on certain platforms without 'mad' support in DoubleFloat.

diff --git a/Engine/Shaders/Private/DoubleFloat.ush b/Engine/Shaders/Private/DoubleFloat.ush
--- a/Engine/Shaders/Private/DoubleFloat.ush
+++ b/Engine/Shaders/Private/DoubleFloat.ush
@@ -22,6 +22,11 @@
 #define UE_DF_NO_FAST_MATH 0
 #endif
 
+// Temporary hack until mad support is added to Switch toolchain
+#if SWITCH_PROFILE || SWITCH_PROFILE_FORWARD
+#define mad(a, b, c) ((a*b)+c)
+#endif
+
 /* FDFScalar */
 
 #define FDFType FDFScalar
@@ -502,3 +507,8 @@ FDFVector3 DFMultiplyLHSAndFastTwoSum(float3 Lhs, float Factor, float3 Rhs)
 	const float3 E = INVARIANT(Rhs - mad(-Lhs, Factor, S));
 	return MakeDFVector3(S, E);
 }
+
+// Temporary hack until mad support is added to Switch toolchain
+#if SWITCH_PROFILE || SWITCH_PROFILE_FORWARD
+#undef mad
+#endif
\ No newline at end of file
