RayTracing: Fix RectLight occlusion ray logic accessing uninitilized fields when using soft shadows from rect lights

diff --git a/Engine/Shaders/Private/RayTracing/RayTracingLightingCommon.ush b/Engine/Shaders/Private/RayTracing/RayTracingLightingCommon.ush
--- a/Engine/Shaders/Private/RayTracing/RayTracingLightingCommon.ush
+++ b/Engine/Shaders/Private/RayTracing/RayTracingLightingCommon.ush
@@ -181,7 +181,6 @@ float4 SampleAreaLightDirection(
 	else if (LightType == LIGHT_TYPE_POINT || LightType == LIGHT_TYPE_SPOT)
 	{
 		// NOTE: Spot cone is being checked before calling this function, so we know we are inside the cone
-		float RayPdf;
 		if (LightParameters.SourceLength > 0.0)
 		{
 			GenerateCapsuleLightOcclusionRayWithSolidAngleSampling(
@@ -209,6 +208,9 @@ float4 SampleAreaLightDirection(
 	}
 	else if (LightType == LIGHT_TYPE_RECT)
 	{
+		// fill in extra parameters used only for rect lights
+		LightParameters.RectLightBarnLength = LightData.RectLightData.BarnLength;
+		LightParameters.RectLightBarnCosAngle = LightData.RectLightData.BarnCosAngle;
 		GenerateRectLightOcclusionRay(
 			LightParameters,
 			TranslatedWorldPosition, WorldNormal,
