Water: Added some comments to clarify why the water info texture merge shader _must_ use samplers rather than load.

diff --git a/Engine/Plugins/Experimental/Water/Shaders/Private/WaterInfoMerge.usf b/Engine/Plugins/Experimental/Water/Shaders/Private/WaterInfoMerge.usf
--- a/Engine/Plugins/Experimental/Water/Shaders/Private/WaterInfoMerge.usf
+++ b/Engine/Plugins/Experimental/Water/Shaders/Private/WaterInfoMerge.usf
@@ -54,6 +54,7 @@ void Main(
 	out float4 OutColor : SV_Target0
 	)
 {
+	// All input textures MUST be sampled with a point sampler to avoid a driver issue on older Android devices which occurs when using Load instead.
 	float4 ColorSample = ColorTexture.SampleLevel(ColorTextureSampler, Input.UV, 0);
 	float4 DepthSample = DepthTexture.SampleLevel(DepthTextureSampler, Input.UV, 0);
 	float4 DilationSample = DilationTexture.SampleLevel(DilationTextureSampler, Input.UV, 0);
diff --git a/Engine/Shaders/Private/WaterInfoTextureMerge.usf b/Engine/Shaders/Private/WaterInfoTextureMerge.usf
--- a/Engine/Shaders/Private/WaterInfoTextureMerge.usf
+++ b/Engine/Shaders/Private/WaterInfoTextureMerge.usf
@@ -37,6 +37,7 @@ void Main(
 	const float2 ScreenPosition = SvPositionToScreenPosition(Position).xy;
 	float2 ScreenUV = ScreenPosition * View.ScreenPositionScaleBias.xy + View.ScreenPositionScaleBias.wz;
 
+	// All input textures MUST be sampled with a point sampler to avoid a driver issue on older Android devices which occurs when using Load instead.
 	const float2 WaterBodyVelocity = WaterBodyTexture.SampleLevel(WaterBodyTextureSampler, ScreenUV, 0);
 	const float WaterBodyDeviceDepth =  WaterBodyDepthTexture.SampleLevel(WaterBodyDepthTextureSampler, ScreenUV, 0);
 	const float DilationDeviceDepth = DilatedWaterBodyDepthTexture.SampleLevel(DilatedWaterBodyDepthTextureSampler, ScreenUV, 0);
