Stochastic Shadows:
* Restored injection of lights with stochastic shadows into translucency lighting volume
* Tweaked neighborhood clamp to accept more samples and don't clamp when we don't have enough samples available. This improves quality of temporaly accumulated shadows.

diff --git a/Engine/Shaders/Private/StochasticShadows/StochasticShadows.ush b/Engine/Shaders/Private/StochasticShadows/StochasticShadows.ush
--- a/Engine/Shaders/Private/StochasticShadows/StochasticShadows.ush
+++ b/Engine/Shaders/Private/StochasticShadows/StochasticShadows.ush
@@ -356,7 +356,6 @@ void AccumulateHistoryShadowMask(uint LightSceneId, uint2 ShadowMaskScreenCoord,
 
 			// Increase age as we are sampling previous frame
 			HistoryShadowMask.Age = min(HistoryShadowMask.Age + 1, MAX_SHADOW_MASK_AGE);
-			HistoryShadowMask.HistoryLength = min(HistoryShadowMask.HistoryLength + 1, MAX_SHADOW_MASK_HISTORY_LENGTH);
 
 			if (HistoryShadowMask.IsValid())
 			{
diff --git a/Engine/Shaders/Private/StochasticShadows/StochasticShadowsComposite.usf b/Engine/Shaders/Private/StochasticShadows/StochasticShadowsComposite.usf
--- a/Engine/Shaders/Private/StochasticShadows/StochasticShadowsComposite.usf
+++ b/Engine/Shaders/Private/StochasticShadows/StochasticShadowsComposite.usf
@@ -14,7 +14,7 @@ float TemporalStdDevOffset;
 
 RWTexture2D<uint> RWShadowMaskTileAtlas;
 
-#define SHARED_TILE_BORDER 1
+#define SHARED_TILE_BORDER 2
 #define SHARED_TILE_SIZE (8 + 2 * SHARED_TILE_BORDER)
 groupshared uint SharedSampleAccumulator[SHARED_TILE_SIZE][SHARED_TILE_SIZE];
 
@@ -92,7 +92,7 @@ void CompositeShadowMaskTracesCS(
 
 #if DEBUG_MODE
 	int2 DebugScreenCoord = View.CursorPosition.x >= 0 ? View.CursorPosition * View.ViewResolutionFraction : -1;
-	bool bDebug = all(ShadowMaskScreenCoord == DebugScreenCoord / SHADOW_MASK_DOWNSAMPLE_FACTOR);
+	bool bDebug = all(ShadowMaskScreenCoord == DebugScreenCoord / SHADOW_MASK_DOWNSAMPLE_FACTOR) && LightSceneId == 3;
 	FShaderPrintContext Context = InitShaderPrintContext(true, float2(0.1, 0.3));
 
 	if (bDebug)
@@ -115,10 +115,10 @@ void CompositeShadowMaskTracesCS(
 		}
 
 		FSampleAccumulator SampleAccumulator = CenterSampleAccumulator;
-		const int FilterSize = 1;
-		for (int SampleOffsetY = -FilterSize; SampleOffsetY <= FilterSize; ++SampleOffsetY)
+		int KernelSize = 2;
+		for (int SampleOffsetY = -KernelSize; SampleOffsetY <= KernelSize; ++SampleOffsetY)
 		{
-			for (int SampleOffsetX = -FilterSize; SampleOffsetX <= FilterSize; ++SampleOffsetX)
+			for (int SampleOffsetX = -KernelSize; SampleOffsetX <= KernelSize; ++SampleOffsetX)
 			{
 				if (!(SampleOffsetX == 0 && SampleOffsetY == 0))
 				{
@@ -164,7 +164,7 @@ void CompositeShadowMaskTracesCS(
 					float HistoryShadowFactorSq = Accumulator.FactorSqSum / Accumulator.WeightSum;
 
 					// Neighborhood clamp
-					if (SampleAccumulator.NumSamples > CenterSampleAccumulator.NumSamples)
+					if (SampleAccumulator.NumSamples > 4)
 					{
 						float Mean = SampleAccumulator.SampleSum / float(SampleAccumulator.NumSamples);
 						float Variance = max(Mean - Mean * Mean, 0.0f);
@@ -201,7 +201,7 @@ void CompositeShadowMaskTracesCS(
 						ShadowMask.Factor = lerp(HistoryShadowFactor, ShadowMask.Factor, Alpha);
 						ShadowMask.FactorSq = lerp(HistoryShadowFactorSq, ShadowMask.FactorSq, Alpha);
 						ShadowMask.Age = 0;
-						ShadowMask.HistoryLength = Accumulator.MaxHistoryLength;
+						ShadowMask.HistoryLength = min(Accumulator.MaxHistoryLength + 1, MAX_SHADOW_MASK_HISTORY_LENGTH);
 					}
 					else
 					{
diff --git a/Engine/Shaders/Private/StochasticShadows/StochasticShadowsSpatial.usf b/Engine/Shaders/Private/StochasticShadows/StochasticShadowsSpatial.usf
--- a/Engine/Shaders/Private/StochasticShadows/StochasticShadowsSpatial.usf
+++ b/Engine/Shaders/Private/StochasticShadows/StochasticShadowsSpatial.usf
@@ -152,7 +152,7 @@ void ShadowMaskSpatialPassCS(
 		float4 ScenePlane = float4(GBuffer.WorldNormal, dot(TranslatedWorldPosition, GBuffer.WorldNormal));
 
 		const int FilterSize = 1;
-		const float FilterWeights[2] = { 1.0f, 1.0f / 2.0f };
+		const float FilterWeights[3] = { 1.0f, 1.0f / 2.0f, 1 / 3.0f };
 		for (int SampleOffsetY = -FilterSize; SampleOffsetY <= FilterSize; ++SampleOffsetY)
 		{
 			for (int SampleOffsetX = -FilterSize; SampleOffsetX <= FilterSize; ++SampleOffsetX)
diff --git a/Engine/Source/Runtime/Renderer/Private/LightRendering.cpp b/Engine/Source/Runtime/Renderer/Private/LightRendering.cpp
--- a/Engine/Source/Runtime/Renderer/Private/LightRendering.cpp
+++ b/Engine/Source/Runtime/Renderer/Private/LightRendering.cpp
@@ -1362,7 +1362,7 @@ void FDeferredShadingSceneRenderer::RenderLights(
 				FTranslucentLightInjectionCollector Collector(GraphBuilder, Views);
 
 				// Collect all the light injection data
-				for (int32 LightIndex = SimpleLightsEnd; LightIndex < StochasticShadowsLightStart; LightIndex++)
+				for (int32 LightIndex = SimpleLightsEnd; LightIndex < SortedLights.Num(); LightIndex++)
 				{
 					const FSortedLightSceneInfo& SortedLightInfo = SortedLights[LightIndex];
 					const FLightSceneInfo& LightSceneInfo = *SortedLightInfo.LightSceneInfo;
