RayTracing: Fix for holdout feature when Substrate is enabled

diff --git a/Engine/Shaders/Private/RayTracing/RayTracingCommon.ush b/Engine/Shaders/Private/RayTracing/RayTracingCommon.ush
--- a/Engine/Shaders/Private/RayTracing/RayTracingCommon.ush
+++ b/Engine/Shaders/Private/RayTracing/RayTracingCommon.ush
@@ -571,6 +571,9 @@ struct FMaterialClosestHitPayload : FMinimalPayload
 	void SetTwoSided()						{ Flags |= RAY_TRACING_PAYLOAD_OUTPUT_FLAG_TWO_SIDED; }
 	bool IsTwoSided()						{ return (Flags & RAY_TRACING_PAYLOAD_OUTPUT_FLAG_TWO_SIDED) != 0; }
 
+	void SetHoldout()                       { Flags |= RAY_TRACING_PAYLOAD_OUTPUT_FLAG_IS_HOLDOUT; }
+	bool IsHoldout()                        { return (Flags & RAY_TRACING_PAYLOAD_OUTPUT_FLAG_IS_HOLDOUT) != 0; }
+
 	FRayCone GetRayCone()					{ return RayCone; }
 	void SetRayCone(FRayCone In)			{ RayCone = In; }
 	float3 GetRayDirection()				{ return RayDirection; }
diff --git a/Engine/Shaders/Private/RayTracing/RayTracingMaterialHitShaders.usf b/Engine/Shaders/Private/RayTracing/RayTracingMaterialHitShaders.usf
--- a/Engine/Shaders/Private/RayTracing/RayTracingMaterialHitShaders.usf
+++ b/Engine/Shaders/Private/RayTracing/RayTracingMaterialHitShaders.usf
@@ -602,10 +602,14 @@ RAY_TRACING_ENTRY_CLOSEST_HIT(MaterialCHS,
 	const bool bIsCameraRay = PackedPayload.IsCameraRay();
 	if (bIsHoldout && bIsCameraRay)
 	{
+#if SUBSTRATE_ENABLED
+		// SUBSTRATE_TODO: Need to clear the material data and set a black unlit material in this case
+#else
 		Payload.ShadingModelID = SHADINGMODELID_UNLIT;
 		Payload.Radiance = 0.0;
 		Payload.BaseColor = 0.0;
 		Payload.Specular = 0.0;
+#endif
 		Payload.SetHoldout();
 	}
 
