Stochastic Shadows
* Proper weighted filter and edge stopping function for spatial moments used for neighborhood clamp. This fixed various flickering artifacts on borders and leaking on object edges
* Reuse already calculated neighborhood for filling missing samples during stochastic sampling in order to reduce noise before history is filled with good samples
* Cache GBuffer.WorldNormal in DownsampledSceneWorldNormal

diff --git a/Engine/Shaders/Private/StochasticShadows/StochasticShadowsComposite.usf b/Engine/Shaders/Private/StochasticShadows/StochasticShadowsComposite.usf
--- a/Engine/Shaders/Private/StochasticShadows/StochasticShadowsComposite.usf
+++ b/Engine/Shaders/Private/StochasticShadows/StochasticShadowsComposite.usf
@@ -9,6 +9,8 @@ StructuredBuffer<uint2> ShadowMaskTileData;
 Texture2D<uint2> ShadowMaskHistoryScreenCoord00;
 Texture2D<float4> ShadowMaskHistoryWeights;
 Texture2D<uint> LightSamples;
+Texture2D<float> DownsampledSceneDepth;
+Texture2D<UNORM float3> DownsampledSceneWorldNormal;
 uint TemporalMaxFramesAccumulated;
 float TemporalStdDevOffset;
 
@@ -17,23 +19,24 @@ RWTexture2D<uint> RWShadowMaskTileAtlas;
 #define SHARED_TILE_BORDER 2
 #define SHARED_TILE_SIZE (8 + 2 * SHARED_TILE_BORDER)
 groupshared uint SharedSampleAccumulator[SHARED_TILE_SIZE][SHARED_TILE_SIZE];
+groupshared float SharedSceneDepth[SHARED_TILE_SIZE][SHARED_TILE_SIZE];
 
 struct FSampleAccumulator
 {
-	uint SampleSum;
-	uint NumSamples;
+	float SampleSum;
+	float SampleWeightSum;
 };
 
 uint PackSampleAccumulator(FSampleAccumulator SampleAccumulator)
 {
-	return SampleAccumulator.SampleSum | (SampleAccumulator.NumSamples << 16);
+	return f32tof16(SampleAccumulator.SampleSum) | (f32tof16(SampleAccumulator.SampleWeightSum) << 16);
 }
 
 FSampleAccumulator UnpackSampleAccumulator(uint Packed)
 {
 	FSampleAccumulator SampleAccumulator;
-	SampleAccumulator.SampleSum = Packed & 0xFFFF;
-	SampleAccumulator.NumSamples = Packed >> 16;
+	SampleAccumulator.SampleSum = f16tof32(Packed & 0xFFFF);
+	SampleAccumulator.SampleWeightSum = f16tof32(Packed >> 16);
 	return SampleAccumulator;
 }
 
@@ -63,8 +66,8 @@ void CompositeShadowMaskTracesCS(
 			uint2 ShadowMaskScreenCoord = clamp(int2(BaseShadowMaskScreenCoord) + int2(SharedShadowMaskX - SHARED_TILE_BORDER, SharedShadowMaskY - SHARED_TILE_BORDER), 0, ShadowMaskViewSize - 1);
 			
 			FSampleAccumulator SampleAccumulator;
-			SampleAccumulator.SampleSum = 0;
-			SampleAccumulator.NumSamples = 0;
+			SampleAccumulator.SampleSum = 0.0f;
+			SampleAccumulator.SampleWeightSum = 0.0f;
 
 			for (uint LocalSampleY = 0; LocalSampleY < NUM_SAMPLES_PER_PIXEL_Y; ++LocalSampleY)
 			{
@@ -76,12 +79,13 @@ void CompositeShadowMaskTracesCS(
 					if (LightSample.LightSceneId == LightSceneId)
 					{
 						SampleAccumulator.SampleSum += LightSample.bVisible ? 1 : 0;
-						SampleAccumulator.NumSamples += 1;
+						SampleAccumulator.SampleWeightSum += 1.0f;
 					}
 				}
 			}
 
 			SharedSampleAccumulator[SharedShadowMaskX][SharedShadowMaskY] = PackSampleAccumulator(SampleAccumulator);			
+			SharedSceneDepth[SharedShadowMaskX][SharedShadowMaskY] = DownsampledSceneDepth[ShadowMaskScreenCoord];
 		}
 	}
 	
@@ -92,7 +96,7 @@ void CompositeShadowMaskTracesCS(
 
 #if DEBUG_MODE
 	int2 DebugScreenCoord = View.CursorPosition.x >= 0 ? View.CursorPosition * View.ViewResolutionFraction : -1;
-	bool bDebug = all(ShadowMaskScreenCoord == DebugScreenCoord / SHADOW_MASK_DOWNSAMPLE_FACTOR) && LightSceneId == 3;
+	bool bDebug = all(ShadowMaskScreenCoord == DebugScreenCoord / SHADOW_MASK_DOWNSAMPLE_FACTOR) && LightSceneId == 0;
 	FShaderPrintContext Context = InitShaderPrintContext(true, float2(0.1, 0.3));
 
 	if (bDebug)
@@ -104,27 +108,48 @@ void CompositeShadowMaskTracesCS(
 	// Temporal reproject and write out new shadow masks
 	if (all(ShadowMaskScreenCoord < ShadowMaskViewSize))
 	{
-		FSampleAccumulator CenterSampleAccumulator = UnpackSampleAccumulator(SharedSampleAccumulator[LocalShadowMaskTileCoord.x + SHARED_TILE_BORDER][LocalShadowMaskTileCoord.y + SHARED_TILE_BORDER]);
+		uint2 CenterSampleSharedTileCoord = LocalShadowMaskTileCoord + SHARED_TILE_BORDER;
+
+		FSampleAccumulator CenterSampleAccumulator = UnpackSampleAccumulator(SharedSampleAccumulator[CenterSampleSharedTileCoord.x][CenterSampleSharedTileCoord.y]);
 
 		FShadowMask ShadowMask = InitShadowMask();
-		if (CenterSampleAccumulator.NumSamples > 0)
+		if (CenterSampleAccumulator.SampleWeightSum > 0.0f)
 		{
-			ShadowMask.Factor = CenterSampleAccumulator.SampleSum / (float)CenterSampleAccumulator.NumSamples;
-			ShadowMask.FactorSq = CenterSampleAccumulator.SampleSum / (float)CenterSampleAccumulator.NumSamples;
+			ShadowMask.Factor = CenterSampleAccumulator.SampleSum / CenterSampleAccumulator.SampleWeightSum;
+			ShadowMask.FactorSq = CenterSampleAccumulator.SampleSum / CenterSampleAccumulator.SampleWeightSum;
 			ShadowMask.Age = 0;
 		}
 
+		const float2 ScreenUV = (DownsampledScreenCoordToScreenCoord(ShadowMaskScreenCoord) + 0.5f) * View.BufferSizeAndInvSize.zw;
+		float CenterSampleDepth = SharedSceneDepth[CenterSampleSharedTileCoord.x][CenterSampleSharedTileCoord.y];
+		float3 CenterSampleNormal = normalize(DecodeNormal(DownsampledSceneWorldNormal[ShadowMaskScreenCoord]));
+		float3 TranslatedWorldPosition = GetTranslatedWorldPositionFromScreenUV(ScreenUV, CenterSampleDepth);
+		float4 ScenePlane = float4(CenterSampleNormal, dot(TranslatedWorldPosition, CenterSampleNormal));
+
 		FSampleAccumulator SampleAccumulator = CenterSampleAccumulator;
 		int KernelSize = 2;
+		const float KernelWeights[3] = { 1.0f, 1.0f / 2.0f, 1.0f / 3.0f };
 		for (int SampleOffsetY = -KernelSize; SampleOffsetY <= KernelSize; ++SampleOffsetY)
 		{
 			for (int SampleOffsetX = -KernelSize; SampleOffsetX <= KernelSize; ++SampleOffsetX)
 			{
 				if (!(SampleOffsetX == 0 && SampleOffsetY == 0))
 				{
-					FSampleAccumulator NeighborSampleAccumulator = UnpackSampleAccumulator(SharedSampleAccumulator[LocalShadowMaskTileCoord.x + SHARED_TILE_BORDER + SampleOffsetX][LocalShadowMaskTileCoord.y + SHARED_TILE_BORDER + SampleOffsetY]);
-					SampleAccumulator.SampleSum += NeighborSampleAccumulator.SampleSum;
-					SampleAccumulator.NumSamples += NeighborSampleAccumulator.NumSamples;
+					uint2 SharedTileCoord = LocalShadowMaskTileCoord + SHARED_TILE_BORDER + int2(SampleOffsetX, SampleOffsetY);
+					float SampleWeight = KernelWeights[abs(SampleOffsetX)] * KernelWeights[abs(SampleOffsetY)];
+
+					uint2 SampleShadowMaskScreenCoord = clamp((int2)BaseShadowMaskScreenCoord + (int2)SharedTileCoord - SHARED_TILE_BORDER, 0, ShadowMaskViewSize - 1);
+					float SampleDepth = SharedSceneDepth[SharedTileCoord.x][SharedTileCoord.y];
+					float3 SampleTranslatedWorldPosition = GetTranslatedWorldPositionFromScreenUV(DownsampledScreenCoordToScreenUV(SampleShadowMaskScreenCoord), SampleDepth);
+
+					float BilateralFilterDepthWeightScale = 10000.0f;
+					float PlaneDistance = abs(dot(float4(SampleTranslatedWorldPosition, -1), ScenePlane));
+					float RelativeDepthDifference = PlaneDistance / SampleDepth;
+					SampleWeight *= exp2(-BilateralFilterDepthWeightScale * (RelativeDepthDifference * RelativeDepthDifference));
+
+					FSampleAccumulator NeighborSampleAccumulator = UnpackSampleAccumulator(SharedSampleAccumulator[SharedTileCoord.x][SharedTileCoord.y]);
+					SampleAccumulator.SampleSum += SampleWeight * NeighborSampleAccumulator.SampleSum;
+					SampleAccumulator.SampleWeightSum += SampleWeight * NeighborSampleAccumulator.SampleWeightSum;
 				}
 			}
 		}
@@ -132,18 +157,26 @@ void CompositeShadowMaskTracesCS(
 		#if DEBUG_MODE
 		if (bDebug && ShadowMask.IsValid())
 		{
-			Newline(Context);
-			Print(Context, TEXT("NumSamples: "));
-			Print(Context, SampleAccumulator.NumSamples);
 			Newline(Context);
 			Print(Context, TEXT("SampleSum: "));
 			Print(Context, SampleAccumulator.SampleSum);
 			Newline(Context);
+			Print(Context, TEXT("SampleWeightSum: "));
+			Print(Context, SampleAccumulator.SampleWeightSum);
+			Newline(Context);
 			Print(Context, TEXT("ShadowMask: "));
 			Print(Context, ShadowMask.Factor);
 		}
 		#endif
 
+		// Interpolate from the neighbourghood if we didn't trace this pixel
+		if (SampleAccumulator.SampleWeightSum > 0 && CenterSampleAccumulator.SampleWeightSum == 0.0f)
+		{
+			ShadowMask.Factor = SampleAccumulator.SampleSum / SampleAccumulator.SampleWeightSum;
+			ShadowMask.FactorSq = SampleAccumulator.SampleSum / SampleAccumulator.SampleWeightSum;
+			ShadowMask.Age = 0;
+		}
+
 		#if TEMPORAL_ACCUMULATION
 		{
 			int2 ScreenCoord00 = ShadowMaskHistoryScreenCoord00[ShadowMaskScreenCoord];
@@ -163,17 +196,18 @@ void CompositeShadowMaskTracesCS(
 					float HistoryShadowFactor = Accumulator.FactorSum / Accumulator.WeightSum;
 					float HistoryShadowFactorSq = Accumulator.FactorSqSum / Accumulator.WeightSum;
 
+					// #sdl_todo: with more samples clamp box should be smaller, instead of using a constant threshold
 					// Neighborhood clamp
-					if (SampleAccumulator.NumSamples > 4)
+					if (SampleAccumulator.SampleWeightSum > 1.0f)
 					{
-						float Mean = SampleAccumulator.SampleSum / float(SampleAccumulator.NumSamples);
+						float Mean = SampleAccumulator.SampleSum / SampleAccumulator.SampleWeightSum;
 						float Variance = max(Mean - Mean * Mean, 0.0f);
 						float StdDev = sqrt(Variance) + TemporalStdDevOffset;
 						float ClampMin = Mean - 0.5f * StdDev;
 						float ClampMax = Mean + 0.5f * StdDev;
 		
 						#if DEBUG_MODE
-						if (bDebug && ShadowMask.IsValid())
+						if (bDebug)
 						{
 							Newline(Context);
 							Print(Context, TEXT("Clamp: "));
@@ -181,7 +215,10 @@ void CompositeShadowMaskTracesCS(
 							Print(Context, ClampMax);
 							Newline(Context);
 							Print(Context, TEXT("HistoryShadowFactor: "));
+							Newline(Context);
 							Print(Context, HistoryShadowFactor);
+							Newline(Context);
+							Print(Context, clamp(HistoryShadowFactor, ClampMin, ClampMax));
 						}
 						#endif
 
@@ -215,6 +252,15 @@ void CompositeShadowMaskTracesCS(
 		}
 		#endif
 
+		#if DEBUG_MODE
+		if (bDebug && ShadowMask.IsValid())
+		{
+			Newline(Context);
+			Print(Context, TEXT("Final: "));
+			Print(Context, ShadowMask.Factor);
+		}
+		#endif
+
 		uint2 ShadowMaskTileAtlasCoord = GetShadowMaskTileAtlasCoord(ShadowMaskTileIndex, LocalShadowMaskTileCoord);
 		RWShadowMaskTileAtlas[ShadowMaskTileAtlasCoord] = PackShadowMask(ShadowMask);
 	}
diff --git a/Engine/Shaders/Private/StochasticShadows/StochasticShadowsSampling.usf b/Engine/Shaders/Private/StochasticShadows/StochasticShadowsSampling.usf
--- a/Engine/Shaders/Private/StochasticShadows/StochasticShadowsSampling.usf
+++ b/Engine/Shaders/Private/StochasticShadows/StochasticShadowsSampling.usf
@@ -90,6 +90,7 @@ int DebugMode;
 RWTexture2D<uint2> RWShadowMaskHistoryScreenCoord00;
 RWTexture2D<float4> RWShadowMaskHistoryWeights;
 RWTexture2D<float> RWDownsampledSceneDepth;
+RWTexture2D<UNORM float3> RWDownsampledSceneWorldNormal;
 RWTexture2D<uint> RWShadowMaskPageTable;
 RWStructuredBuffer<uint> RWShadowMaskHashTable;
 RWStructuredBuffer<uint> RWShadowMaskTileAllocator;
@@ -619,7 +620,9 @@ void GenerateSamplesCS(
 				RWShadowMaskHistoryWeights[DownsampledScreenCoord] = ShadowMaskHistoryWeights;
 			}
 			#endif
+
 			RWDownsampledSceneDepth[DownsampledScreenCoord] = SceneDepth;
+			RWDownsampledSceneWorldNormal[DownsampledScreenCoord] = EncodeNormal(GBuffer.WorldNormal);
 
 			for (uint LightSampleIndex = 0; LightSampleIndex < NUM_SAMPLES_PER_PIXEL; ++LightSampleIndex)
 			{
diff --git a/Engine/Shaders/Private/StochasticShadows/StochasticShadowsSpatial.usf b/Engine/Shaders/Private/StochasticShadows/StochasticShadowsSpatial.usf
--- a/Engine/Shaders/Private/StochasticShadows/StochasticShadowsSpatial.usf
+++ b/Engine/Shaders/Private/StochasticShadows/StochasticShadowsSpatial.usf
@@ -10,6 +10,7 @@ Texture2D<uint> ShadowMaskPageTable;
 StructuredBuffer<uint> ShadowMaskHashTable;
 StructuredBuffer<uint2> ShadowMaskTileData;
 Texture2D<float> DownsampledSceneDepth;
+Texture2D<UNORM float3> DownsampledSceneWorldNormal;
 uint MaxShadingTiles;
 
 Texture2D<float4> ShadowMaskUpsampleWeights;
@@ -145,24 +146,25 @@ void ShadowMaskSpatialPassCS(
 		uint2 ShadowMaskTileAtlasCoord = GetShadowMaskTileAtlasCoord(ShadowMaskTileIndex, LocalShadowMaskTileCoord);
 		FShadowMaskAccumulator ShadowMaskAccumulator = InitShadowMaskAccumulator();
 
-		// #sdl_todo: load downsampled depth instead
+		uint2 CenterSampleSharedTileCoord = LocalShadowMaskTileCoord + SHARED_TILE_BORDER;
 		const float2 ScreenUV = (DownsampledScreenCoordToScreenCoord(ShadowMaskScreenCoord) + 0.5f) * View.BufferSizeAndInvSize.zw;
-		const FGBufferData GBuffer = GetGBufferData(ScreenUV);
-		float3 TranslatedWorldPosition = GetTranslatedWorldPositionFromScreenUV(ScreenUV, GBuffer.Depth);
-		float4 ScenePlane = float4(GBuffer.WorldNormal, dot(TranslatedWorldPosition, GBuffer.WorldNormal));
-
-		const int FilterSize = 1;
-		const float FilterWeights[3] = { 1.0f, 1.0f / 2.0f, 1 / 3.0f };
-		for (int SampleOffsetY = -FilterSize; SampleOffsetY <= FilterSize; ++SampleOffsetY)
+		float CenterSampleDepth = SharedSceneDepth[CenterSampleSharedTileCoord.x][CenterSampleSharedTileCoord.y];
+		float3 CenterSampleNormal = normalize(DecodeNormal(DownsampledSceneWorldNormal[ShadowMaskScreenCoord]));
+		float3 TranslatedWorldPosition = GetTranslatedWorldPositionFromScreenUV(ScreenUV, CenterSampleDepth);
+		float4 ScenePlane = float4(CenterSampleNormal, dot(TranslatedWorldPosition, CenterSampleNormal));
+
+		const int KernelSize = 1;
+		const float KernelWeights[3] = { 1.0f, 1.0f / 2.0f, 1.0f / 3.0f };
+		for (int SampleOffsetY = -KernelSize; SampleOffsetY <= KernelSize; ++SampleOffsetY)
 		{
-			for (int SampleOffsetX = -FilterSize; SampleOffsetX <= FilterSize; ++SampleOffsetX)
+			for (int SampleOffsetX = -KernelSize; SampleOffsetX <= KernelSize; ++SampleOffsetX)
 			{
 				uint2 SharedShadowMaskTileCoord = (GroupThreadId.xy + SHARED_TILE_BORDER) - int2(SampleOffsetX, SampleOffsetY);
 				FShadowMask SampleShadowMask = UnpackShadowMask(SharedShadowMask[SharedShadowMaskTileCoord.x][SharedShadowMaskTileCoord.y]);
 
 				if (SampleShadowMask.IsValid())
 				{
-					float SampleWeight = FilterWeights[abs(SampleOffsetX)] * FilterWeights[abs(SampleOffsetY)];
+					float SampleWeight = KernelWeights[abs(SampleOffsetX)] * KernelWeights[abs(SampleOffsetY)];
 
 					uint2 SampleShadowMaskScreenCoord = clamp((int2)BaseShadowMaskScreenCoord + (int2)SharedShadowMaskTileCoord - SHARED_TILE_BORDER, 0, ShadowMaskViewSize - 1);
 					float SampleDepth = SharedSceneDepth[SharedShadowMaskTileCoord.x][SharedShadowMaskTileCoord.y];
diff --git a/Engine/Shaders/Private/StochasticShadows/StochasticShadowsTracing.usf b/Engine/Shaders/Private/StochasticShadows/StochasticShadowsTracing.usf
--- a/Engine/Shaders/Private/StochasticShadows/StochasticShadowsTracing.usf
+++ b/Engine/Shaders/Private/StochasticShadows/StochasticShadowsTracing.usf
@@ -180,6 +180,7 @@ void CompactLightSampleTracesCS(
 
 StructuredBuffer<uint> CompactedTraceTexelData;
 Texture2D<float> DownsampledSceneDepth;
+Texture2D<UNORM float3> DownsampledSceneWorldNormal;
 
 [numthreads(THREADGROUP_SIZE, 1, 1)]
 void ScreenSpaceRayTraceLightSamplesCS(
@@ -208,6 +209,7 @@ void ScreenSpaceRayTraceLightSamplesCS(
 
 		float2 ScreenUV = (ScreenCoord + 0.5f) * View.BufferSizeAndInvSize.zw;
 		float SceneDepth = DownsampledSceneDepth[DownsampledScreenCoord];
+		float3 SceneWorldNormal = normalize(DecodeNormal(DownsampledSceneWorldNormal[DownsampledScreenCoord]));
 		float3 TranslatedWorldPosition = GetTranslatedWorldPositionFromScreenUV(ScreenUV, SceneDepth);
 
 		FLightSample LightSample = UnpackLightSample(RWLightSamples[SampleCoord]);
@@ -225,12 +227,11 @@ void ScreenSpaceRayTraceLightSamplesCS(
 			float2 CornerScreenUV = ScreenUV + .5f * View.BufferSizeAndInvSize.zw;
 			float2 CornerScreenPosition = (CornerScreenUV - View.ScreenPositionScaleBias.wz) / View.ScreenPositionScaleBias.xy;
 			float3 CornerTranslatedWorldPosition = mul(float4(GetScreenPositionForProjectionType(CornerScreenPosition, SceneDepth), SceneDepth, 1), View.ScreenToTranslatedWorld).xyz;
-			float3 WorldNormal = GetGBufferDataFromSceneTextures(ScreenUV).WorldNormal; // #sdl_todo: get normal from a DownsampledNormal buffer
 
 			// #sdl_todo: fixup bias and incorporate it into Lumen
 			// Project texel corner onto normal
-			float NormalBias = abs(dot(CornerTranslatedWorldPosition - TranslatedWorldPosition, WorldNormal)) * 2.0f;
-			TranslatedWorldPosition += NormalBias * WorldNormal + 0.05f * WorldNormal;
+			float NormalBias = abs(dot(CornerTranslatedWorldPosition - TranslatedWorldPosition, SceneWorldNormal)) * 2.0f;
+			TranslatedWorldPosition += NormalBias * SceneWorldNormal + 0.05f * SceneWorldNormal;
 
 			float4 RayStartClip = mul(float4(TranslatedWorldPosition, 1.0f), View.TranslatedWorldToClip);
 
diff --git a/Engine/Source/Runtime/Renderer/Private/StochasticShadows/StochasticShadows.cpp b/Engine/Source/Runtime/Renderer/Private/StochasticShadows/StochasticShadows.cpp
--- a/Engine/Source/Runtime/Renderer/Private/StochasticShadows/StochasticShadows.cpp
+++ b/Engine/Source/Runtime/Renderer/Private/StochasticShadows/StochasticShadows.cpp
@@ -55,7 +55,7 @@ static TAutoConsoleVariable<float> CVarStochasticShadowsTemporalStdDevOffset(
 
 static TAutoConsoleVariable<int32> CVarStochasticSamplingShadowEstimate(
 	TEXT("r.StochasticShadows.Sampling.ShadowEstimate"),
-	1,
+	0,
 	TEXT("Whether to use shadow mask history for sample guiding."),
 	ECVF_Scalability | ECVF_RenderThreadSafe
 );
@@ -306,6 +306,7 @@ class FGenerateSamplesCS : public FGlobalShader
 		RDG_BUFFER_ACCESS(IndirectArgs, ERHIAccess::IndirectArgs)
 		SHADER_PARAMETER_STRUCT_INCLUDE(FStochasticShadowsParameters, StochasticShadowsParameters)
 		SHADER_PARAMETER_RDG_TEXTURE_UAV(RWTexture2D<float>, RWDownsampledSceneDepth)
+		SHADER_PARAMETER_RDG_TEXTURE_UAV(RWTexture2D<UNORM float3>, RWDownsampledSceneWorldNormal)
 		SHADER_PARAMETER_RDG_TEXTURE_UAV(RWTexture2D<uint32>, RWShadowMaskHistoryScreenCoord00)
 		SHADER_PARAMETER_RDG_TEXTURE_UAV(RWTexture2D<float4>, RWShadowMaskHistoryWeights)
 		SHADER_PARAMETER_RDG_BUFFER_UAV(RWStructuredBuffer<uint>, RWShadowMaskTileAllocator)
@@ -605,6 +606,10 @@ void FDeferredShadingSceneRenderer::RenderStochasticShadows(FRDGBuilder& GraphBu
 		FRDGTextureDesc::Create2D(DownsampledBufferSize, PF_R32_FLOAT, FClearValueBinding::Black, TexCreate_ShaderResource | TexCreate_UAV),
 		TEXT("StochasticShadows.DownsampledSceneDepth"));
 
+	FRDGTextureRef DownsampledSceneWorldNormal = GraphBuilder.CreateTexture(
+		FRDGTextureDesc::Create2D(DownsampledBufferSize, PF_A2B10G10R10, FClearValueBinding::Black, TexCreate_ShaderResource | TexCreate_UAV),
+		TEXT("StochasticShadows.DownsampledSceneWorldNormal"));
+
 	FRDGTextureRef LightSamples = GraphBuilder.CreateTexture(
 		FRDGTextureDesc::Create2D(DonwnsampledSampleBufferSize, PF_R16_UINT, FClearValueBinding::Black, TexCreate_ShaderResource | TexCreate_UAV),
 		TEXT("StochasticShadows.LightSamples"));
@@ -713,6 +718,7 @@ void FDeferredShadingSceneRenderer::RenderStochasticShadows(FRDGBuilder& GraphBu
 		StochasticShadowsParameters.ShadowMaskViewSize = ShadowMaskViewSize;
 		StochasticShadowsParameters.StochasticShadowsStateFrameIndex = StochasticShadows::GetStateFrameIndex(View.ViewState);
 		StochasticShadowsParameters.DownsampledSceneDepth = DownsampledSceneDepth;
+		StochasticShadowsParameters.DownsampledSceneWorldNormal = DownsampledSceneWorldNormal;
 		StochasticShadowsParameters.MaxShadowMaskTiles = MaxShadowMaskTiles;
 		StochasticShadowsParameters.MaxShadingTiles = (ShadingTileAtlasSize.X * ShadingTileAtlasSize.Y) / (StochasticShadows::ShadowMaskTileSize * StochasticShadows::ShadowMaskTileSize);
 		StochasticShadowsParameters.MaxShadingTilesPerGridCell = StochasticShadows::MaxShadingTilesPerGridCell;
@@ -823,6 +829,7 @@ void FDeferredShadingSceneRenderer::RenderStochasticShadows(FRDGBuilder& GraphBu
 	// Generate new candidate light samples
 	{
 		FRDGTextureUAVRef DownsampledSceneDepthUAV = GraphBuilder.CreateUAV(DownsampledSceneDepth, ERDGUnorderedAccessViewFlags::SkipBarrier);
+		FRDGTextureUAVRef DownsampledSceneWorldNormalUAV = GraphBuilder.CreateUAV(DownsampledSceneWorldNormal, ERDGUnorderedAccessViewFlags::SkipBarrier);
 		FRDGTextureUAVRef ShadowMaskHistoryScreenCoord00UAV = GraphBuilder.CreateUAV(ShadowMaskHistoryScreenCoord00, ERDGUnorderedAccessViewFlags::SkipBarrier);
 		FRDGTextureUAVRef ShadowMaskHistoryWeightsUAV = GraphBuilder.CreateUAV(ShadowMaskHistoryWeights, ERDGUnorderedAccessViewFlags::SkipBarrier);
 		FRDGBufferUAVRef ShadowMaskTileAllocatorUAV = GraphBuilder.CreateUAV(ShadowMaskTileAllocator, ERDGUnorderedAccessViewFlags::SkipBarrier);
@@ -837,6 +844,7 @@ void FDeferredShadingSceneRenderer::RenderStochasticShadows(FRDGBuilder& GraphBu
 			PassParameters->IndirectArgs = DownsampledTileIndirectArgs;
 			PassParameters->StochasticShadowsParameters = StochasticShadowsParameters;
 			PassParameters->RWDownsampledSceneDepth = DownsampledSceneDepthUAV;
+			PassParameters->RWDownsampledSceneWorldNormal = DownsampledSceneWorldNormalUAV;
 			PassParameters->RWShadowMaskHistoryScreenCoord00 = ShadowMaskHistoryScreenCoord00UAV;
 			PassParameters->RWShadowMaskHistoryWeights = ShadowMaskHistoryWeightsUAV;
 			PassParameters->RWShadowMaskTileAllocator = ShadowMaskTileAllocatorUAV;
diff --git a/Engine/Source/Runtime/Renderer/Private/StochasticShadows/StochasticShadowsInternal.h b/Engine/Source/Runtime/Renderer/Private/StochasticShadows/StochasticShadowsInternal.h
--- a/Engine/Source/Runtime/Renderer/Private/StochasticShadows/StochasticShadowsInternal.h
+++ b/Engine/Source/Runtime/Renderer/Private/StochasticShadows/StochasticShadowsInternal.h
@@ -32,6 +32,7 @@ BEGIN_SHADER_PARAMETER_STRUCT(FStochasticShadowsParameters, )
 	SHADER_PARAMETER(float, TemporalStdDevOffset)
 	SHADER_PARAMETER(int32, DebugMode)
 	SHADER_PARAMETER_RDG_TEXTURE(Texture2D<float>, DownsampledSceneDepth)
+	SHADER_PARAMETER_RDG_TEXTURE(Texture2D<UNORM float3>, DownsampledSceneWorldNormal)
 END_SHADER_PARAMETER_STRUCT()
 
 // Internal functions, don't use outside of the StochasticShadows
