Fix invalid unpacking of hair attribute with multiple clumpIDs.

diff --git a/Engine/Shaders/Private/HairStrands/HairStrandsPack.ush b/Engine/Shaders/Private/HairStrands/HairStrandsPack.ush
--- a/Engine/Shaders/Private/HairStrands/HairStrandsPack.ush
+++ b/Engine/Shaders/Private/HairStrands/HairStrandsPack.ush
@@ -127,7 +127,7 @@ uint3 UnpackHairClumpID(uint In)
 
 uint3 UnpackHairClumpID(uint2 In)
 {
-	return uint3(In.x & 0xFFFF, (In.x>>16) & 0xFFFF, (In.x>>32) & 0xFFFF);
+	return uint3(In.x & 0xFFFF, (In.x>>16) & 0xFFFF, In.y & 0xFFFF);
 }
 
 uint PackHairColor(float3 In)
