Hotfix for cooking issue on certain platforms without 'mad' support in DoubleFloat/LWC.

diff --git a/Engine/Shaders/Private/DoubleFloat.ush b/Engine/Shaders/Private/DoubleFloat.ush
--- a/Engine/Shaders/Private/DoubleFloat.ush
+++ b/Engine/Shaders/Private/DoubleFloat.ush
@@ -24,7 +24,7 @@
 
 // Temporary hack until mad support is added to Switch toolchain
 #if SWITCH_PROFILE || SWITCH_PROFILE_FORWARD
-#define mad(a, b, c) ((a*b)+c)
+#define mad(a, b, c) (((a)*(b))+(c))
 #endif
 
 /* FDFScalar */
diff --git a/Engine/Shaders/Private/LargeWorldCoordinates.ush b/Engine/Shaders/Private/LargeWorldCoordinates.ush
--- a/Engine/Shaders/Private/LargeWorldCoordinates.ush
+++ b/Engine/Shaders/Private/LargeWorldCoordinates.ush
@@ -513,6 +513,11 @@ float4x4 LWCHackToFloat(FLWCInverseMatrix v) { return (float4x4)0.0f; }
 #define LWCHackToFloat(V) LWCToFloat(V)
 #endif
 
+// Temporary hack until mad support is added to Switch toolchain
+#if SWITCH_PROFILE || SWITCH_PROFILE_FORWARD
+#define mad(a, b, c) (((a)*(b))+(c))
+#endif
+
 #define DEFINE_DF_TO_TILEOFFSET_OPERATOR(FDFType, FLWCType, LWCConstructor, FFloatType) \
 	FLWCType DFToTileOffset(FDFType In)																							\
 	{																															\
@@ -541,6 +546,11 @@ DEFINE_DF_TO_TILEOFFSET_OPERATOR(FDFVector3, FLWCVector3, MakeLWCVector3, float3
 DEFINE_DF_TO_TILEOFFSET_OPERATOR(FDFVector4, FLWCVector4, MakeLWCVector4, float4)
 #undef DEFINE_DF_TO_TILEOFFSET_OPERATOR
 
+// Temporary hack until mad support is added to Switch toolchain
+#if SWITCH_PROFILE || SWITCH_PROFILE_FORWARD
+#undef mad
+#endif
+
 FLWCMatrix DFToTileOffset(FDFMatrix In)
 {
 	FLWCVector3 PosHigh = DFToTileOffset(MakeDFVector3(In.PostTranslation, 0));
