Fix shader warning with Substrate with TopLayerData

diff --git a/Engine/Shaders/Private/Substrate/Substrate.ush b/Engine/Shaders/Private/Substrate/Substrate.ush
--- a/Engine/Shaders/Private/Substrate/Substrate.ush
+++ b/Engine/Shaders/Private/Substrate/Substrate.ush
@@ -751,6 +751,17 @@ SUBSTRATE_TOP_LAYER_TYPE SubstratePackTopLayerData(FSubstrateTopLayerData In)
 #endif
 }
 
+bool SubstrateIsTopLayerMaterial(SUBSTRATE_TOP_LAYER_TYPE In)
+{
+#if SUBSTRATE_NORMAL_QUALITY==0
+	return In != 0;
+#elif SUBSTRATE_NORMAL_QUALITY==1
+	return In.y != 0;
+#else
+#error Unkown SUBSTRATE_NORMAL_QUALITY
+#endif
+}
+
 FSubstrateTopLayerData SubstrateUnpackTopLayerData(uint In)
 {
 	FSubstrateTopLayerData Out = (FSubstrateTopLayerData)0;
diff --git a/Engine/Shaders/Private/Substrate/SubstrateMaterialClassification.usf b/Engine/Shaders/Private/Substrate/SubstrateMaterialClassification.usf
--- a/Engine/Shaders/Private/Substrate/SubstrateMaterialClassification.usf
+++ b/Engine/Shaders/Private/Substrate/SubstrateMaterialClassification.usf
@@ -193,7 +193,7 @@ void TileMainCS(uint2 DispatchThreadId : SV_DispatchThreadID, uint LinearIndex :
 		BRANCH
 		if (CMask > 0u && CMask < 0xF)
 		{
-			bClearHeader = TopLayerTexture.Load(uint3(PixelCoord,0)).x == 0;
+			bClearHeader = !SubstrateIsTopLayerMaterial(TopLayerTexture.Load(uint3(PixelCoord, 0)));
 		}
 
 		// Header clear
