* Change position offset buffer to a structured buffer to scalaraize load.
* Abstract position offset read to move it as a global resource & batch update later.

diff --git a/Engine/Plugins/Runtime/HairStrands/Shaders/Private/DeformerDataInterfaceGroom.ush b/Engine/Plugins/Runtime/HairStrands/Shaders/Private/DeformerDataInterfaceGroom.ush
--- a/Engine/Plugins/Runtime/HairStrands/Shaders/Private/DeformerDataInterfaceGroom.ush
+++ b/Engine/Plugins/Runtime/HairStrands/Shaders/Private/DeformerDataInterfaceGroom.ush
@@ -28,7 +28,7 @@ FHairControlPoint ReadControlPointData(uint ControlPointIndex)
 	// Hair deformer work in local space so no need for position offset
 	return UnpackHairControlPoint(
 		{DataInterfaceName}_PositionBuffer[ControlPointIndex], 
-		{DataInterfaceName}_PositionOffsetBuffer[0].xyz,
+		{DataInterfaceName}_GetHairInstancePositionOffset(),
 		{DataInterfaceName}_Radius, 
 		{DataInterfaceName}_RootScale, 
 		{DataInterfaceName}_TipScale);
diff --git a/Engine/Plugins/Runtime/HairStrands/Shaders/Private/DeformerDataInterfaceGroomGuide.ush b/Engine/Plugins/Runtime/HairStrands/Shaders/Private/DeformerDataInterfaceGroomGuide.ush
--- a/Engine/Plugins/Runtime/HairStrands/Shaders/Private/DeformerDataInterfaceGroomGuide.ush
+++ b/Engine/Plugins/Runtime/HairStrands/Shaders/Private/DeformerDataInterfaceGroomGuide.ush
@@ -23,7 +23,7 @@ FHairControlPoint ReadControlPointData_{DataInterfaceName}(uint ControlPointInde
 	// Hair deformer work in local space so no need for position offset
 	return UnpackHairControlPoint(
 		{DataInterfaceName}_PositionBuffer[ControlPointIndex], 
-		{DataInterfaceName}_PositionOffsetBuffer[0].xyz,
+		{DataInterfaceName}_GetHairInstancePositionOffset(),
 		{DataInterfaceName}_Radius, 
 		{DataInterfaceName}_RootScale, 
 		{DataInterfaceName}_TipScale);
diff --git a/Engine/Plugins/Runtime/HairStrands/Shaders/Private/DeformerDataInterfaceGroomWrite.ush b/Engine/Plugins/Runtime/HairStrands/Shaders/Private/DeformerDataInterfaceGroomWrite.ush
--- a/Engine/Plugins/Runtime/HairStrands/Shaders/Private/DeformerDataInterfaceGroomWrite.ush
+++ b/Engine/Plugins/Runtime/HairStrands/Shaders/Private/DeformerDataInterfaceGroomWrite.ush
@@ -5,12 +5,19 @@
 
 HAIR_STRANDS_INSTANCE_PARAMETERS({DataInterfaceName})
 
+// Note: resources are explicitely bound here (instance of using HAIR_STRANDS_INSTANCE_PARAMETERS's resources) 
+// beause we write out data which would be bound as SRV otherwise (position, attributes, ...)
 uint 				{DataInterfaceName}_OutputStreamStart;
-Buffer<float4>		{DataInterfaceName}_PositionOffsetBufferSRV;
-Buffer<uint4>   	{DataInterfaceName}_PositionBufferSRV;
-RWBuffer<uint4> 	{DataInterfaceName}_PositionBufferUAV;
-RWByteAddressBuffer {DataInterfaceName}_PointAttributeBufferUAV;
-RWByteAddressBuffer {DataInterfaceName}_CurveAttributeBufferUAV;
+StructuredBuffer<float4> {DataInterfaceName}_PositionOffsetBufferSRV;
+Buffer<uint4>   		{DataInterfaceName}_PositionBufferSRV;
+RWBuffer<uint4> 		{DataInterfaceName}_PositionBufferUAV;
+RWByteAddressBuffer 	{DataInterfaceName}_PointAttributeBufferUAV;
+RWByteAddressBuffer 	{DataInterfaceName}_CurveAttributeBufferUAV;
+
+float3 {DataInterfaceName}_GetPositionOffsetBufferSRV() 
+{ 
+	return ReadRenPositionOffset({DataInterfaceName}_PositionOffsetBufferSRV,{DataInterfaceName}_RegisteredIndex); 
+}
 
 uint ReadNumControlPoints_{DataInterfaceName}()
 {
@@ -28,7 +35,7 @@ void WritePosition_{DataInterfaceName}(uint ControlPointIndex, float3 Position)
 	if (WriteIndex < {DataInterfaceName}_PointCount)
 	{
 		const uint4 PackedCP = {DataInterfaceName}_PositionBufferSRV[WriteIndex];
-		const uint4 NewPackedCP = PackHairControlPointPosition(PackedCP, Position, {DataInterfaceName}_PositionOffsetBufferSRV[0].xyz);
+		const uint4 NewPackedCP = PackHairControlPointPosition(PackedCP, Position, {DataInterfaceName}_GetPositionOffsetBufferSRV());
 		{DataInterfaceName}_PositionBufferUAV[WriteIndex] = NewPackedCP;
 	}
 }
@@ -51,7 +58,7 @@ void WritePositionAndRadius_{DataInterfaceName}(uint ControlPointIndex, float4 P
 	{
 		const uint4 PackedCP = {DataInterfaceName}_PositionBufferSRV[WriteIndex];
 		uint4 NewPackedCP = 0;
-		NewPackedCP = PackHairControlPointPosition(PackedCP, PositionAndRadius.xyz, {DataInterfaceName}_PositionOffsetBufferSRV[0].xyz);
+		NewPackedCP = PackHairControlPointPosition(PackedCP, PositionAndRadius.xyz, {DataInterfaceName}_GetPositionOffsetBufferSRV());
 		NewPackedCP = PackHairControlPointRadius(NewPackedCP, PositionAndRadius.w, {DataInterfaceName}_Radius);
 		{DataInterfaceName}_PositionBufferUAV[WriteIndex] = NewPackedCP;
 	}
diff --git a/Engine/Plugins/Runtime/HairStrands/Source/HairStrandsCore/Private/Deformer/DeformerDataInterfaceGroomWrite.cpp b/Engine/Plugins/Runtime/HairStrands/Source/HairStrandsCore/Private/Deformer/DeformerDataInterfaceGroomWrite.cpp
--- a/Engine/Plugins/Runtime/HairStrands/Source/HairStrandsCore/Private/Deformer/DeformerDataInterfaceGroomWrite.cpp
+++ b/Engine/Plugins/Runtime/HairStrands/Source/HairStrandsCore/Private/Deformer/DeformerDataInterfaceGroomWrite.cpp
@@ -112,7 +112,7 @@ void UOptimusGroomWriteDataInterface::GetSupportedOutputs(TArray<FShaderFunction
 BEGIN_SHADER_PARAMETER_STRUCT(FGroomWriteDataInterfaceParameters, )
 	SHADER_PARAMETER_STRUCT_INCLUDE(FHairStrandsInstanceCommonParameters, Common)
 	SHADER_PARAMETER(uint32, OutputStreamStart)
-	SHADER_PARAMETER_RDG_BUFFER_SRV(Buffer, PositionOffsetBufferSRV)
+	SHADER_PARAMETER_RDG_BUFFER_SRV(StructuredBuffer, PositionOffsetBufferSRV)
 	SHADER_PARAMETER_RDG_BUFFER_SRV(Buffer<uint4>, PositionBufferSRV)
 	SHADER_PARAMETER_RDG_BUFFER_UAV(RWBuffer<uint4>, PositionBufferUAV)
 	SHADER_PARAMETER_RDG_BUFFER_UAV(RWByteAddressBuffer, CurveAttributeBufferUAV)
diff --git a/Engine/Plugins/Runtime/HairStrands/Source/HairStrandsCore/Private/GroomDebug.cpp b/Engine/Plugins/Runtime/HairStrands/Source/HairStrandsCore/Private/GroomDebug.cpp
--- a/Engine/Plugins/Runtime/HairStrands/Source/HairStrandsCore/Private/GroomDebug.cpp
+++ b/Engine/Plugins/Runtime/HairStrands/Source/HairStrandsCore/Private/GroomDebug.cpp
@@ -640,16 +640,18 @@ class FDrawDebugCardGuidesCS : public FGlobalShader
 
 	BEGIN_SHADER_PARAMETER_STRUCT(FParameters, )
 		SHADER_PARAMETER_STRUCT_REF(FViewUniformShaderParameters, ViewUniformBuffer)
-		SHADER_PARAMETER(uint32,  DebugMode)
+		SHADER_PARAMETER(uint32, InstanceRegisteredIndex)
+		SHADER_PARAMETER(uint32, CardLODIndex)
+		SHADER_PARAMETER(uint32, DebugMode)
 		SHADER_PARAMETER(FMatrix44f, LocalToWorld)
 		
 		SHADER_PARAMETER(uint32,  RenVertexCount)
 		SHADER_PARAMETER(FVector3f, RenRestOffset)
-		SHADER_PARAMETER_RDG_BUFFER_SRV(Buffer, RenDeformedOffset)
+		SHADER_PARAMETER_RDG_BUFFER_SRV(StructuredBuffer, RenDeformedOffset)
 		
 		SHADER_PARAMETER(uint32,  SimVertexCount)
 		SHADER_PARAMETER(FVector3f, SimRestOffset)
-		SHADER_PARAMETER_RDG_BUFFER_SRV(Buffer, SimDeformedOffset)
+		SHADER_PARAMETER_RDG_BUFFER_SRV(StructuredBuffer, SimDeformedOffset)
 
 		SHADER_PARAMETER_RDG_BUFFER_SRV(Buffer, RenRestPosition)
 		SHADER_PARAMETER_RDG_BUFFER_SRV(Buffer, RenDeformedPosition)
@@ -732,6 +734,8 @@ static void AddDrawDebugCardsGuidesPass(
 	FDrawDebugCardGuidesCS::FParameters* Parameters = GraphBuilder.AllocParameters<FDrawDebugCardGuidesCS::FParameters>();
 	Parameters->ViewUniformBuffer = View.ViewUniformBuffer;
 
+	Parameters->InstanceRegisteredIndex = Instance->RegisteredIndex;
+	Parameters->CardLODIndex = HairLODIndex;
 	Parameters->RenVertexCount = 0;
 	Parameters->RenRestOffset = FVector3f::ZeroVector;
 	Parameters->RenRestPosition = DefaultBuffer;
diff --git a/Engine/Plugins/Runtime/HairStrands/Source/HairStrandsCore/Private/GroomManager.cpp b/Engine/Plugins/Runtime/HairStrands/Source/HairStrandsCore/Private/GroomManager.cpp
--- a/Engine/Plugins/Runtime/HairStrands/Source/HairStrandsCore/Private/GroomManager.cpp
+++ b/Engine/Plugins/Runtime/HairStrands/Source/HairStrandsCore/Private/GroomManager.cpp
@@ -300,6 +300,9 @@ static void RunInternalHairInterpolation(
 						AddHairStrandUpdatePositionOffsetPass(
 							GraphBuilder,
 							ShaderMap,
+							EHairPositionUpdateType::Strands,
+							Instance->RegisteredIndex,
+							HairLODIndex,
 							MeshLODIndex,
 							Instance->Strands.DeformedRootResource,
 							Instance->Strands.DeformedResource);
@@ -311,6 +314,9 @@ static void RunInternalHairInterpolation(
 						AddHairStrandUpdatePositionOffsetPass(
 							GraphBuilder,
 							ShaderMap,
+							EHairPositionUpdateType::Strands,
+							Instance->RegisteredIndex,
+							HairLODIndex,
 							MeshLODIndex,
 							nullptr,
 							Instance->Strands.DeformedResource);
@@ -341,6 +347,9 @@ static void RunInternalHairInterpolation(
 							AddHairStrandUpdatePositionOffsetPass(
 								GraphBuilder,
 								ShaderMap,
+								EHairPositionUpdateType::Cards,
+								Instance->RegisteredIndex,
+								HairLODIndex,
 								MeshLODIndex,
 								CardsInstance.Guides.DeformedRootResource,
 								CardsInstance.Guides.DeformedResource);
@@ -351,6 +360,9 @@ static void RunInternalHairInterpolation(
 							AddHairStrandUpdatePositionOffsetPass(
 								GraphBuilder,
 								ShaderMap,
+								EHairPositionUpdateType::Cards,
+								Instance->RegisteredIndex,
+								HairLODIndex,
 								MeshLODIndex,
 								nullptr,
 								CardsInstance.Guides.DeformedResource);
@@ -522,6 +534,9 @@ static void RunHairStrandsInterpolation_Guide(
 			AddHairStrandUpdatePositionOffsetPass(
 				GraphBuilder,
 				ShaderMap,
+				EHairPositionUpdateType::Guides,
+				InstanceData.Instance->RegisteredIndex,
+				InstanceData.HairLODIndex,
 				InstanceData.Instance->Debug.MeshLODIndex,
 				InstanceData.Instance->Guides.DeformedRootResource,
 				InstanceData.Instance->Guides.DeformedResource);
@@ -548,6 +563,9 @@ static void RunHairStrandsInterpolation_Guide(
 			AddHairStrandUpdatePositionOffsetPass(
 				GraphBuilder,
 				ShaderMap,
+				EHairPositionUpdateType::Guides,
+				InstanceData.Instance->RegisteredIndex,
+				InstanceData.HairLODIndex,
 				InstanceData.Instance->Debug.MeshLODIndex,
 				nullptr,
 				InstanceData.Instance->Guides.DeformedResource);
@@ -638,6 +656,9 @@ static void RunHairStrandsInterpolation_Guide(
 			AddHairStrandUpdatePositionOffsetPass(
 				GraphBuilder,
 				ShaderMap,
+				EHairPositionUpdateType::Guides,
+				InstanceData.Instance->RegisteredIndex,
+				InstanceData.HairLODIndex,
 				InstanceData.MeshLODIndex,
 				InstanceData.Instance->Guides.DeformedRootResource,
 				InstanceData.Instance->Guides.DeformedResource);
diff --git a/Engine/Plugins/Runtime/HairStrands/Source/HairStrandsCore/Private/HairStrandsInterpolation.cpp b/Engine/Plugins/Runtime/HairStrands/Source/HairStrandsCore/Private/HairStrandsInterpolation.cpp
--- a/Engine/Plugins/Runtime/HairStrands/Source/HairStrandsCore/Private/HairStrandsInterpolation.cpp
+++ b/Engine/Plugins/Runtime/HairStrands/Source/HairStrandsCore/Private/HairStrandsInterpolation.cpp
@@ -915,13 +915,14 @@ class FHairCardsDeformationCS : public FGlobalShader
 	using FPermutationDomain = TShaderPermutationDomain<FDynamicGeometry>;
 
 	BEGIN_SHADER_PARAMETER_STRUCT(FParameters, )
+		SHADER_PARAMETER(uint32, InstanceRegisteredIndex)
 		SHADER_PARAMETER(FMatrix44f, LocalToWorld)
 		SHADER_PARAMETER(uint32, CardsVertexCount)
 		SHADER_PARAMETER(uint32, GuideVertexCount)
 		SHADER_PARAMETER(FVector3f, GuideRestPositionOffset)
 		SHADER_PARAMETER_RDG_BUFFER_SRV(Buffer, GuideRestPositionBuffer)
 		SHADER_PARAMETER_RDG_BUFFER_SRV(Buffer, GuideDeformedPositionBuffer)
-		SHADER_PARAMETER_RDG_BUFFER_SRV(Buffer, GuideDeformedPositionOffsetBuffer)
+		SHADER_PARAMETER_RDG_BUFFER_SRV(StructuredBuffer, GuideDeformedPositionOffsetBuffer)
 		SHADER_PARAMETER_SRV(Buffer, CardsRestPositionBuffer)
 		SHADER_PARAMETER_SRV(Buffer, CardsRestTangentBuffer)
 		SHADER_PARAMETER_RDG_BUFFER_SRV(ByteAddressBuffer, CardsInterpolationBuffer)
@@ -983,6 +984,7 @@ static void AddHairCardsDeformationPass(
 	}
 	
 	FHairCardsDeformationCS::FParameters* Parameters = GraphBuilder.AllocParameters<FHairCardsDeformationCS::FParameters>();
+	Parameters->InstanceRegisteredIndex		= Instance->RegisteredIndex;
 	Parameters->LocalToWorld				= FMatrix44f(Instance->GetCurrentLocalToWorld().ToMatrixWithScale());		// LWC_TODO: Precision loss
 	Parameters->GuideVertexCount			= LOD.Guides.RestResource->GetPointCount();
 	Parameters->GuideRestPositionOffset		= (FVector3f)LOD.Guides.RestResource->GetPositionOffset();
@@ -1146,6 +1148,7 @@ class FHairRaytracingGeometryCS : public FGlobalShader
 
 	BEGIN_SHADER_PARAMETER_STRUCT(FParameters, )
 		SHADER_PARAMETER(uint32, PointCount)
+		SHADER_PARAMETER(uint32, HairStrandsVF_RegisteredIndex)
 		SHADER_PARAMETER(float, HairStrandsVF_HairRadius)
 		SHADER_PARAMETER(float, HairStrandsVF_HairRootScale)
 		SHADER_PARAMETER(float, HairStrandsVF_HairTipScale)
@@ -1156,7 +1159,7 @@ class FHairRaytracingGeometryCS : public FGlobalShader
 		SHADER_PARAMETER_RDG_BUFFER_SRV(Buffer,	HairStrandsVF_CullingRadiusScaleBuffer)
 		RDG_BUFFER_ACCESS(HairStrandsVF_CullingIndirectBufferArgs, ERHIAccess::IndirectArgs)
 
-		SHADER_PARAMETER_RDG_BUFFER_SRV(Buffer, PositionOffsetBuffer)
+		SHADER_PARAMETER_RDG_BUFFER_SRV(StructuredBuffer, PositionOffsetBuffer)
 		SHADER_PARAMETER_RDG_BUFFER_SRV(Buffer, PositionBuffer)
 		SHADER_PARAMETER_RDG_BUFFER_SRV(Buffer, TangentBuffer)
 		SHADER_PARAMETER_RDG_BUFFER_UAV(RWBuffer, OutputPositionBuffer)
@@ -1191,6 +1194,7 @@ static void AddGenerateRaytracingGeometryPass(
 	FRDGBuilder& GraphBuilder,
 	FGlobalShaderMap* ShaderMap,
 	const FShaderPrintData* ShaderPrintData,
+	uint32 InstanceRegisteredIndex,
 	uint32 PointCount,
 	bool bProceduralPrimitive,
 	int ProceduralSplits,
@@ -1207,6 +1211,7 @@ static void AddGenerateRaytracingGeometryPass(
 	FHairRaytracingGeometryCS::FParameters* Parameters = GraphBuilder.AllocParameters<FHairRaytracingGeometryCS::FParameters>();
 	Parameters->PointCount = PointCount;
 	Parameters->PositionOffsetBuffer = HairWorldOffsetBuffer;
+	Parameters->HairStrandsVF_RegisteredIndex = InstanceRegisteredIndex;
 	Parameters->HairStrandsVF_HairRadius = HairRadius;
 	Parameters->HairStrandsVF_HairRootScale = RootScale;
 	Parameters->HairStrandsVF_HairTipScale = TipScale;
@@ -1751,6 +1756,9 @@ void ComputeHairStrandsInterpolation(
 					AddHairStrandUpdatePositionOffsetPass(
 						GraphBuilder,
 						ShaderMap,
+						EHairPositionUpdateType::Strands,
+						Instance->RegisteredIndex,
+						0,
 						MeshLODIndex,
 						Instance->Strands.DeformedRootResource,
 						Instance->Strands.DeformedResource);
@@ -1866,6 +1874,9 @@ void ComputeHairStrandsInterpolation(
 					AddHairStrandUpdatePositionOffsetPass(
 						GraphBuilder,
 						ShaderMap,
+						EHairPositionUpdateType::Strands,
+						Instance->RegisteredIndex,
+						0,
 						MeshLODIndex,
 						Instance->Strands.DeformedRootResource,
 						Instance->Strands.DeformedResource);
@@ -2001,6 +2012,7 @@ void ComputeHairStrandsInterpolation(
 						GraphBuilder,
 						ShaderMap,
 						ShaderPrintData,
+						Instance->RegisteredIndex,
 						ActivePointCount,
 						bProceduralPrimitive,
 						ProceduralSplits,
diff --git a/Engine/Plugins/Runtime/HairStrands/Source/HairStrandsCore/Private/HairStrandsMeshProjection.cpp b/Engine/Plugins/Runtime/HairStrands/Source/HairStrandsCore/Private/HairStrandsMeshProjection.cpp
--- a/Engine/Plugins/Runtime/HairStrands/Source/HairStrandsCore/Private/HairStrandsMeshProjection.cpp
+++ b/Engine/Plugins/Runtime/HairStrands/Source/HairStrandsCore/Private/HairStrandsMeshProjection.cpp
@@ -1297,12 +1297,15 @@ private:
 	using FPermutationDomain = TShaderPermutationDomain<FUseGPUOffset, FPrevious>;
 
 	BEGIN_SHADER_PARAMETER_STRUCT(FParameters, )
+		SHADER_PARAMETER(uint32, UpdateType)
+		SHADER_PARAMETER(uint32, CardLODIndex)
+		SHADER_PARAMETER(uint32, InstanceRegisteredIndex)
 		SHADER_PARAMETER(FVector3f, CPUCurrPositionOffset)
 		SHADER_PARAMETER(FVector3f, CPUPrevPositionOffset)
 		SHADER_PARAMETER_RDG_BUFFER_SRV(StructuredBuffer, RootTriangleCurrPositionBuffer)
 		SHADER_PARAMETER_RDG_BUFFER_SRV(StructuredBuffer, RootTrianglePrevPositionBuffer)
-		SHADER_PARAMETER_RDG_BUFFER_UAV(RWBuffer, OutCurrOffsetBuffer)
-		SHADER_PARAMETER_RDG_BUFFER_UAV(RWBuffer, OutPrevOffsetBuffer)
+		SHADER_PARAMETER_RDG_BUFFER_UAV(RWStructuredBuffer, OutCurrOffsetBuffer)
+		SHADER_PARAMETER_RDG_BUFFER_UAV(RWStructuredBuffer, OutPrevOffsetBuffer)
 	END_SHADER_PARAMETER_STRUCT()
 
 public:
@@ -1322,11 +1325,14 @@ IMPLEMENT_GLOBAL_SHADER(FHairUpdatePositionOffsetCS, "/Engine/Private/HairStrand
 void AddHairStrandUpdatePositionOffsetPass(
 	FRDGBuilder& GraphBuilder,
 	FGlobalShaderMap* ShaderMap,
-	const int32 LODIndex,
+	EHairPositionUpdateType UpdateType,
+	const int32 InstanceRegisteredIndex,
+	const int32 HairLODIndex,
+	const int32 MeshLODIndex,
 	FHairStrandsDeformedRootResource* DeformedRootResources,
 	FHairStrandsDeformedResource* DeformedResources)
 {
-	if ((DeformedRootResources && LODIndex < 0) || DeformedResources == nullptr)
+	if ((DeformedRootResources && MeshLODIndex < 0) || DeformedResources == nullptr)
 	{
 		return;
 	}
@@ -1340,8 +1346,8 @@ void AddHairStrandUpdatePositionOffsetPass(
 	FRDGImportedBuffer RootTrianglePrevPositionBuffer;
 	if (DeformedRootResources)
 	{
-		RootTriangleCurrPositionBuffer = Register(GraphBuilder, DeformedRootResources->LODs[LODIndex].GetDeformedUniqueTrianglePositionBuffer(FHairStrandsDeformedRootResource::FLOD::Current), ERDGImportedBufferFlags::CreateSRV);
-		RootTrianglePrevPositionBuffer = Register(GraphBuilder, DeformedRootResources->LODs[LODIndex].GetDeformedUniqueTrianglePositionBuffer(FHairStrandsDeformedRootResource::FLOD::Previous), ERDGImportedBufferFlags::CreateSRV);
+		RootTriangleCurrPositionBuffer = Register(GraphBuilder, DeformedRootResources->LODs[MeshLODIndex].GetDeformedUniqueTrianglePositionBuffer(FHairStrandsDeformedRootResource::FLOD::Current), ERDGImportedBufferFlags::CreateSRV);
+		RootTrianglePrevPositionBuffer = Register(GraphBuilder, DeformedRootResources->LODs[MeshLODIndex].GetDeformedUniqueTrianglePositionBuffer(FHairStrandsDeformedRootResource::FLOD::Previous), ERDGImportedBufferFlags::CreateSRV);
 	}
 
 	const bool bUseGPUOffset = DeformedRootResources != nullptr && DeformedRootResources->IsValid() && DeformedRootResources->IsInitialized() && GHairStrandsUseGPUPositionOffset > 0;
@@ -1349,6 +1355,10 @@ void AddHairStrandUpdatePositionOffsetPass(
 	const uint32 PrevOffsetIndex = DeformedResources->GetIndex(FHairStrandsDeformedResource::Previous);
 
 	FHairUpdatePositionOffsetCS::FParameters* PassParameters = GraphBuilder.AllocParameters<FHairUpdatePositionOffsetCS::FParameters>();
+	PassParameters->UpdateType 				= uint32(UpdateType);
+	PassParameters->CardLODIndex			= HairLODIndex;
+	PassParameters->InstanceRegisteredIndex = InstanceRegisteredIndex;
+
 	PassParameters->CPUCurrPositionOffset			= (FVector3f)DeformedResources->PositionOffset[CurrOffsetIndex];
 	PassParameters->CPUPrevPositionOffset			= (FVector3f)DeformedResources->PositionOffset[PrevOffsetIndex];
 	PassParameters->RootTriangleCurrPositionBuffer	= RootTriangleCurrPositionBuffer.SRV;
diff --git a/Engine/Plugins/Runtime/HairStrands/Source/HairStrandsCore/Public/HairStrandsMeshProjection.h b/Engine/Plugins/Runtime/HairStrands/Source/HairStrandsCore/Public/HairStrandsMeshProjection.h
--- a/Engine/Plugins/Runtime/HairStrands/Source/HairStrandsCore/Public/HairStrandsMeshProjection.h
+++ b/Engine/Plugins/Runtime/HairStrands/Source/HairStrandsCore/Public/HairStrandsMeshProjection.h
@@ -136,10 +136,20 @@ void AddHairCardsRBFInterpolationPass(
 	FHairStrandsRestRootResource* RestRootResources,
 	FHairStrandsDeformedRootResource* DeformedRootResources);
 
+enum class EHairPositionUpdateType : uint8
+{
+	Guides  = 0,
+	Strands = 1,
+	Cards   = 2,
+};
+
 void AddHairStrandUpdatePositionOffsetPass(
 	FRDGBuilder& GraphBuilder,
 	FGlobalShaderMap* ShaderMap,
-	const int32 LODIndex,
+	EHairPositionUpdateType UpdateType,
+	const int32 InstanceRegisteredIndex,
+	const int32 HairLODIndex,
+	const int32 MeshLODIndex,
 	FHairStrandsDeformedRootResource* DeformedRootResources,
 	FHairStrandsDeformedResource* DeformedResources);
 
diff --git a/Engine/Shaders/Private/HairStrands/HairCardsDeformation.usf b/Engine/Shaders/Private/HairStrands/HairCardsDeformation.usf
--- a/Engine/Shaders/Private/HairStrands/HairCardsDeformation.usf
+++ b/Engine/Shaders/Private/HairStrands/HairCardsDeformation.usf
@@ -14,10 +14,12 @@
 #if SHADER_DEFORM_CARDS
 float4x4 LocalToWorld;
 
+uint CardLODIndex;
+uint InstanceRegisteredIndex;
 uint CardsVertexCount;
 uint GuideVertexCount;
 float3 GuideRestPositionOffset;
-Buffer<float4> GuideDeformedPositionOffsetBuffer;
+StructuredBuffer<float4> GuideDeformedPositionOffsetBuffer;
 
 Buffer<uint4>  GuideRestPositionBuffer;
 Buffer<uint4>  GuideDeformedPositionBuffer;
@@ -84,7 +86,7 @@ void MainCS(uint DispatchThreadId : SV_DispatchThreadID)
 	const float GuideRadius = 1;
 	const float GuideRootScale = 1;
 	const float GuideTipScale = 1;
-	const float3 GuideDeformedPositionOffset = GuideDeformedPositionOffsetBuffer[0].xyz;
+	const float3 GuideDeformedPositionOffset = ReadCardPositionOffset(GuideDeformedPositionOffsetBuffer, InstanceRegisteredIndex, CardLODIndex);
 
 #if PERMUTATION_DYNAMIC_GEOMETRY == 0
 	// Simple deformation based on world space offset
diff --git a/Engine/Shaders/Private/HairStrands/HairStrandsDebug.usf b/Engine/Shaders/Private/HairStrands/HairStrandsDebug.usf
--- a/Engine/Shaders/Private/HairStrands/HairStrandsDebug.usf
+++ b/Engine/Shaders/Private/HairStrands/HairStrandsDebug.usf
@@ -268,7 +268,7 @@ void MainCS(uint3 DispatchThreadId : SV_DispatchThreadID)
 
 	const float3 P_Local = UnpackHairControlPoint(
 		HairStrandsVF.PositionBuffer[VertexIndex],
-		HairStrandsVF.PositionOffsetBuffer[0].xyz,
+		ReadRenPositionOffset(HairStrandsVF.PositionOffsetBuffer, HairStrandsVF.RegisteredIndex),
 		HairStrandsVF.Radius * ScaleR,
 		HairStrandsVF.RootScale,
 		HairStrandsVF.TipScale).Position;
@@ -1844,6 +1844,9 @@ void MainCS(uint3 DispatchThreadId : SV_DispatchThreadID)
 #include "HairStrandsVertexFactoryCommon.ush"
 #include "../ShaderPrint.ush"
 
+uint InstanceRegisteredIndex;
+uint CardLODIndex;
+
 uint RenVertexCount;
 float3 RenRestOffset;
 
@@ -1854,11 +1857,11 @@ float4x4 LocalToWorld;
 
 Buffer<uint4> RenRestPosition;
 Buffer<uint4> RenDeformedPosition;
-Buffer<float4> RenDeformedOffset;
+StructuredBuffer<float4> RenDeformedOffset;
 
 Buffer<uint4> SimRestPosition;
 Buffer<uint4> SimDeformedPosition;
-Buffer<float4> SimDeformedOffset;
+StructuredBuffer<float4> SimDeformedOffset;
 
 float3 RestOffset;
 float3 DeformedOffset;
@@ -1898,7 +1901,7 @@ void MainCS(uint3 DispatchThreadId : SV_DispatchThreadID)
 	{
 		PackedP0 = RenDeformedPosition[VertexId];
 		PackedP1 = RenDeformedPosition[VertexId + 1];
-		Offset   = RenDeformedOffset[0].xyz;
+		Offset   = ReadCardPositionOffset(RenDeformedOffset, InstanceRegisteredIndex, CardLODIndex);
 		ColorRoot= ColorYellow;
 		ColorTip = ColorPurple;
 	}
@@ -1914,7 +1917,7 @@ void MainCS(uint3 DispatchThreadId : SV_DispatchThreadID)
 	{
 		PackedP0 = SimDeformedPosition[VertexId];
 		PackedP1 = SimDeformedPosition[VertexId + 1];
-		Offset   = SimDeformedOffset[0].xyz;
+		Offset   = ReadSimPositionOffset(SimDeformedOffset, InstanceRegisteredIndex);
 		ColorRoot= ColorCyan;
 		ColorTip = ColorRed;
 	}
@@ -2481,7 +2484,7 @@ void CSMain(
 		{
 			FHairControlPoint Out = UnpackHairControlPoint(
 				HairInstance_PositionBuffer[VisNode.ControlPointId],
-				HairInstance_PositionOffsetBuffer[0].xyz,
+				HairInstance_GetHairInstancePositionOffset(),
 				HairInstance_Radius,
 				HairInstance_RootScale,
 				HairInstance_TipScale);
diff --git a/Engine/Shaders/Private/HairStrands/HairStrandsForwardCulling.usf b/Engine/Shaders/Private/HairStrands/HairStrandsForwardCulling.usf
--- a/Engine/Shaders/Private/HairStrands/HairStrandsForwardCulling.usf
+++ b/Engine/Shaders/Private/HairStrands/HairStrandsForwardCulling.usf
@@ -100,7 +100,7 @@ void CSMain(uint DispatchThreadID : SV_DispatchThreadID, uint GroupThreadID : SV
 
 	const uint CurveID = DispatchThreadID;
 	const bool bValidCurve = CurveID < CurveCount;
-	const float3 PositionOffset = HairStrandsVF_PositionOffsetBuffer[0].xyz;
+	const float3 PositionOffset = HairStrandsVF_GetHairInstancePositionOffset();
 
 	bool bIsCurveVisibile = false;
 	uint PrimWriteOffset = 0;
@@ -260,7 +260,7 @@ void CSMain(uint DispatchThreadID : SV_DispatchThreadID, uint GroupThreadID : SV
 	bool bIsPrimValid = ControlPointId < ControlPointCount;
 	if (bIsPrimValid)
 	{
-		const float3 PositionOffset = HairStrandsVF_PositionOffsetBuffer[0].xyz;
+		const float3 PositionOffset = HairStrandsVF_GetHairInstancePositionOffset();
 
 		// Fetch both control points
 		FHairControlPoint CP0 = UnpackHairControlPoint(
diff --git a/Engine/Shaders/Private/HairStrands/HairStrandsInterpolation.usf b/Engine/Shaders/Private/HairStrands/HairStrandsInterpolation.usf
--- a/Engine/Shaders/Private/HairStrands/HairStrandsInterpolation.usf
+++ b/Engine/Shaders/Private/HairStrands/HairStrandsInterpolation.usf
@@ -118,9 +118,12 @@ Buffer<uint>	SimRootBarycentricBuffer;
 Buffer<uint>	SimRootToUniqueTriangleIndexBuffer;
 #endif
 
-Buffer<float4> OutSimHairPositionOffsetBuffer;
-Buffer<float4> OutRenHairPositionOffsetBuffer;
+uint InstanceRegisteredIndex;
+StructuredBuffer<float4> OutSimHairPositionOffsetBuffer;
+StructuredBuffer<float4> OutRenHairPositionOffsetBuffer;
 
+float3 GetSimPositionOffset() { return ReadSimPositionOffset(OutSimHairPositionOffsetBuffer, InstanceRegisteredIndex); }
+float3 GetRenPositionOffset() { return ReadRenPositionOffset(OutRenHairPositionOffsetBuffer, InstanceRegisteredIndex); }
 
 #ifndef GROUP_SIZE
 	#error Unknown group size	
@@ -436,7 +439,7 @@ void MainCS(uint2 DispatchThreadId : SV_DispatchThreadID)
 	// 2. Appply deformation
 	if (VertexIndex < VertexCount)
 	{
-		const FGuideDataWithOffset GuideData = GetGuideData(VertexIndex, PERMUTATION_USE_SINGLE_GUIDE, OutSimHairPositionOffsetBuffer[0].xyz);
+		const FGuideDataWithOffset GuideData = GetGuideData(VertexIndex, PERMUTATION_USE_SINGLE_GUIDE, GetSimPositionOffset());
 		const FRootTriangleData RenTri = GetRootTriangleData(VertexIndex, false);
 
 		// Manual decoding of the rest position
@@ -495,7 +498,7 @@ void MainCS(uint2 DispatchThreadId : SV_DispatchThreadID)
 				OutPosition = INFINITE_FLOAT;
 			}
 		}
-		OutRenDeformedPositionBuffer[VertexIndex] = PackHairControlPointPosition(PackedRestPosition, OutPosition, OutRenHairPositionOffsetBuffer[0].xyz);
+		OutRenDeformedPositionBuffer[VertexIndex] = PackHairControlPointPosition(PackedRestPosition, OutPosition, GetRenPositionOffset());
 	}
 }
 #endif // SHADER_HAIRINTERPOLATION
diff --git a/Engine/Shaders/Private/HairStrands/HairStrandsMesh.usf b/Engine/Shaders/Private/HairStrands/HairStrandsMesh.usf
--- a/Engine/Shaders/Private/HairStrands/HairStrandsMesh.usf
+++ b/Engine/Shaders/Private/HairStrands/HairStrandsMesh.usf
@@ -2,6 +2,7 @@
 
 #include "../Common.ush"
 #include "HairStrandsBindingCommon.ush"
+#include "HairstrandsVertexFactoryCommon.ush"
 
 ////////////////////////////////////////////////////////////////////////
 // Deform hair mesh basd on RBF weights
@@ -453,20 +454,24 @@ void MainCS(uint3 DispatchThreadId : SV_DispatchThreadID)
 	float3 GetOffsetPrevPosition()	{ return CPUPrevPositionOffset; }
 #endif
 
-RWBuffer<float4> OutCurrOffsetBuffer;
+RWStructuredBuffer<float4> OutCurrOffsetBuffer;
 #if PERMUTATION_PREVIOUS
-RWBuffer<float4> OutPrevOffsetBuffer;
+RWStructuredBuffer<float4> OutPrevOffsetBuffer;
 #endif
 
+uint UpdateType;
+uint CardLODIndex;
+uint InstanceRegisteredIndex;
+
 [numthreads(1, 1, 1)]
 void MainCS(uint3 DispatchThreadId : SV_DispatchThreadID)
 {
 	if (DispatchThreadId.x > 0)
 		return;
 
-	OutCurrOffsetBuffer[0] = float4(GetOffsetCurrPosition(), 1);
+	WritePositionOffset(OutCurrOffsetBuffer, InstanceRegisteredIndex, UpdateType, HAIR_CURR, CardLODIndex, GetOffsetCurrPosition());
 #if PERMUTATION_PREVIOUS
-	OutPrevOffsetBuffer[0] = float4(GetOffsetPrevPosition(), 1);
+	WritePositionOffset(OutPrevOffsetBuffer, InstanceRegisteredIndex, UpdateType, HAIR_PREV, CardLODIndex, GetOffsetPrevPosition());
 #endif
 }
 
diff --git a/Engine/Shaders/Private/HairStrands/HairStrandsRaytracingGeometry.usf b/Engine/Shaders/Private/HairStrands/HairStrandsRaytracingGeometry.usf
--- a/Engine/Shaders/Private/HairStrands/HairStrandsRaytracingGeometry.usf
+++ b/Engine/Shaders/Private/HairStrands/HairStrandsRaytracingGeometry.usf
@@ -14,10 +14,11 @@ float HairStrandsVF_HairRadius;
 float HairStrandsVF_HairRootScale;
 float HairStrandsVF_HairTipScale;
 
-Buffer<float4>	 PositionOffsetBuffer;
+uint HairStrandsVF_RegisteredIndex;
+StructuredBuffer<float4> PositionOffsetBuffer;
 Buffer<uint4> 	 PositionBuffer;
 Buffer<float4> 	 TangentBuffer;
-RWBuffer<float4> OutputPositionBuffer;
+RWBuffer<float4> OutputPositionBuffer;	
 #if PERMUTATION_PROCEDURAL_PRIMITIVE 
 uint             RaytracingProceduralSplits;
 #else
@@ -40,7 +41,7 @@ FHairControlPoint GetControlPoint(uint ControlPointId)
 	{
 		Out = UnpackHairControlPoint(
 			PositionBuffer[ControlPointId],
-			PositionOffsetBuffer[0].xyz,
+			ReadRenPositionOffset(PositionOffsetBuffer, HairStrandsVF_RegisteredIndex),
 			HairStrandsVF_HairRadius,
 			HairStrandsVF_HairRootScale,
 			HairStrandsVF_HairTipScale);
diff --git a/Engine/Shaders/Private/HairStrands/HairStrandsVertexFactory.ush b/Engine/Shaders/Private/HairStrands/HairStrandsVertexFactory.ush
--- a/Engine/Shaders/Private/HairStrands/HairStrandsVertexFactory.ush
+++ b/Engine/Shaders/Private/HairStrands/HairStrandsVertexFactory.ush
@@ -189,6 +189,16 @@ struct FVertexFactoryIntermediates
 	FSceneDataIntermediates SceneData;
 };
 
+float3 GetPositionOffset()
+{
+	return ReadRenPositionOffset(HairStrandsVF.PositionOffsetBuffer, HairStrandsVF.RegisteredIndex);
+}
+
+float3 GetPrevPositionOffset()
+{
+	return ReadRenPrevPositionOffset(HairStrandsVF.PreviousPositionOffsetBuffer, HairStrandsVF.RegisteredIndex);
+}
+
 uint GetHairControlPointId(FVertexFactoryInput Input)
 {
 	return GetVertexInfo(Input).HairControlPointId;
@@ -196,7 +206,7 @@ uint GetHairControlPointId(FVertexFactoryInput Input)
 
 FHairControlPoint GetVertexPosition(const FVertexInfo VertexInfo, bool bInvalidJointVertex=true)
 {
-	const float3 PositionOffset = HairStrandsVF.PositionOffsetBuffer[0].xyz;
+	const float3 PositionOffset = GetPositionOffset();
 
 	const uint4 PackedData = HairStrandsVF.PositionBuffer[VertexInfo.VertexIndex];
 	FHairControlPoint Out = UnpackHairControlPoint(
@@ -237,7 +247,7 @@ float GetWorldStrandRadius(FVertexFactoryInput Input)
 
 float3 GetVertexPreviousPosition(FVertexFactoryInput Input, bool bInvalidJointVertex = true)
 {
-	const float3 PreviousPositionOffset = HairStrandsVF.PreviousPositionOffsetBuffer[0].xyz;
+	const float3 PreviousPositionOffset = GetPrevPositionOffset();
 
 	FVertexInfo VertexInfo = GetVertexInfo(Input);
 	const uint4 PackedData = HairStrandsVF.PreviousPositionBuffer[VertexInfo.VertexIndex];
@@ -407,7 +417,7 @@ RAY_TRACING_ENTRY_INTERSECTION(MaterialIS)
 	}
 
 	const float EffectiveRadius = HairStrandsVF.Radius * HairStrandsVF.RaytracingRadiusScale;
-	const float3 PositionOffset = HairStrandsVF.PositionOffsetBuffer[0].xyz;
+	const float3 PositionOffset = GetPositionOffset();
 	FHairControlPoint P0 = UnpackHairControlPoint(
 		HairStrandsVF.PositionBuffer[VertexIndex0],
 		PositionOffset,
@@ -491,7 +501,7 @@ FMaterialPixelParameters GetMaterialPixelParameters(float3 RayOrigin, float3 Ray
 	const float EffectiveRadius = HairStrandsVF.Radius * HairStrandsVF.RaytracingRadiusScale;
 	float UCoord = Barycentrics.x;
 	float VCoord = Barycentrics.y;
-	const float3 PositionOffset = HairStrandsVF.PositionOffsetBuffer[0].xyz;
+	const float3 PositionOffset = GetPositionOffset();
 	FHairControlPoint P0 = UnpackHairControlPoint(
 		HairStrandsVF.PositionBuffer[VertexIndex0],
 		PositionOffset,
@@ -543,7 +553,7 @@ FMaterialPixelParameters GetMaterialPixelParameters(float3 RayOrigin, float3 Ray
 	float VCoord = 0;
 	{
 		// find central point of the hair strand
-		const float3 PositionOffset = HairStrandsVF.PositionOffsetBuffer[0].xyz;
+		const float3 PositionOffset = GetPositionOffset();
 		const float EffectiveRadius = HairStrandsVF.Radius * HairStrandsVF.RaytracingRadiusScale;
 		FHairControlPoint P0 = UnpackHairControlPoint(
 			HairStrandsVF.PositionBuffer[VertexIndex0],
diff --git a/Engine/Shaders/Private/HairStrands/HairStrandsVertexFactoryCommon.ush b/Engine/Shaders/Private/HairStrands/HairStrandsVertexFactoryCommon.ush
--- a/Engine/Shaders/Private/HairStrands/HairStrandsVertexFactoryCommon.ush
+++ b/Engine/Shaders/Private/HairStrands/HairStrandsVertexFactoryCommon.ush
@@ -9,46 +9,87 @@
 #define HAIR_STRANDS_PARAMETERS 0
 #endif
 
+#define HAIR_SIM  0
+#define HAIR_REN  1
+#define HAIR_CAR  2
+#define HAIR_MAX_CARDS_LOD 8
+
+#define HAIR_CURR 0
+#define HAIR_PREV 1
+
+float3 ReadPositionOffset(StructuredBuffer<float4> In, uint InIndex, uint InType, uint InTime, uint InLOD) 
+{ 
+	return In[0].xyz; 
+}
+
+void WritePositionOffset(RWStructuredBuffer<float4> Out, uint InIndex, uint InType, uint InTime, uint InLOD, float3 InPositionOffset) 
+{ 
+	Out[0] = float4(InPositionOffset, 1);
+}
+
+float3 ReadRenPositionOffset(StructuredBuffer<float4> In, uint InIndex) 
+{ 
+	return ReadPositionOffset(In, InIndex, HAIR_REN, HAIR_CURR, 0);
+}
+
+float3 ReadRenPrevPositionOffset(StructuredBuffer<float4> In, uint InIndex) 
+{ 
+	return ReadPositionOffset(In, InIndex, HAIR_REN, HAIR_PREV, 0);
+}
+
+float3 ReadSimPositionOffset(StructuredBuffer<float4> In, uint InIndex) 
+{ 
+	return ReadPositionOffset(In, InIndex, HAIR_SIM, HAIR_CURR, 0);
+}
+
+float3 ReadCardPositionOffset(StructuredBuffer<float4> In, uint InIndex, uint InLOD) 
+{ 
+	return ReadPositionOffset(In, InIndex, HAIR_CAR, HAIR_CURR, InLOD);
+}
+
+
 ///////////////////////////////////////////////////////////////////////////////////////////////////
 // Hair strands instance parameters
 
 #define HAIR_STRANDS_INSTANCE_PARAMETERS(Prefix) \
-	uint				Prefix##_bCullingEnable; \
-	uint				Prefix##_bRaytracingGeometry; \
-	uint				Prefix##_RaytracingProceduralSplits; \
-	uint				Prefix##_GroupIndex; \
-	uint				Prefix##_GroupCount; \
-	float				Prefix##_Radius; \
-	float				Prefix##_RootScale; \
-	float				Prefix##_TipScale; \
-	float				Prefix##_Length; \
-	float				Prefix##_Density; \
-	uint				Prefix##_bStableRasterization; \
-	float3				Prefix##_PositionOffset; \
-	uint				Prefix##_PointCount; \
-	uint				Prefix##_CurveCount; \
-	uint				Prefix##_bSimulation; \
-	uint				Prefix##_bSingleGuide; \
-	float4x4			Prefix##_LocalToWorldPrimitiveTransform; \
-	float4x4			Prefix##_LocalToTranslatedWorldPrimitiveTransform; \
-	uint4				Prefix##_CurveAttributeOffsets[HAIR_CURVE_ATTRIBUTE_OFFSET_COUNT]; \
-	uint4				Prefix##_PointAttributeOffsets[HAIR_POINT_ATTRIBUTE_OFFSET_COUNT]; \
-	uint 				Prefix##_CurveAttributeIndexToChunkDivAsShift; \
-	uint				Prefix##_CurveAttributeChunkElementCount; \
-	uint				Prefix##_CurveAttributeChunkStrideInBytes; \
-	uint				Prefix##_PointAttributeIndexToChunkDivAsShift; \
-	uint				Prefix##_PointAttributeChunkElementCount; \
-	uint				Prefix##_PointAttributeChunkStrideInBytes; \
-	Buffer<uint4>		Prefix##_PositionBuffer; \
-	Buffer<float4>		Prefix##_PositionOffsetBuffer; \
-	Buffer<uint>		Prefix##_CurveBuffer; \
-	Buffer<uint>		Prefix##_PointToCurveBuffer; \
-	ByteAddressBuffer 	Prefix##_CurveAttributeBuffer; \
-	ByteAddressBuffer 	Prefix##_PointAttributeBuffer; \
-	Buffer<uint>		Prefix##_CullingIndirectBuffer; \
-	Buffer<uint>		Prefix##_CullingIndexBuffer; \
-	Buffer<float>		Prefix##_CullingRadiusScaleBuffer; \
-	Buffer<uint>		Prefix##_CullingIndirectBufferArgs;
+	uint						Prefix##_RegisteredIndex; \
+	uint						Prefix##_bCullingEnable; \
+	uint						Prefix##_bRaytracingGeometry; \
+	uint						Prefix##_RaytracingProceduralSplits; \
+	uint						Prefix##_GroupIndex; \
+	uint						Prefix##_GroupCount; \
+	float						Prefix##_Radius; \
+	float						Prefix##_RootScale; \
+	float						Prefix##_TipScale; \
+	float						Prefix##_Length; \
+	float						Prefix##_Density; \
+	uint						Prefix##_bStableRasterization; \
+	float3						Prefix##_PositionOffset; \
+	uint						Prefix##_PointCount; \
+	uint						Prefix##_CurveCount; \
+	uint						Prefix##_bSimulation; \
+	uint						Prefix##_bSingleGuide; \
+	float4x4					Prefix##_LocalToWorldPrimitiveTransform; \
+	float4x4					Prefix##_LocalToTranslatedWorldPrimitiveTransform; \
+	uint4						Prefix##_CurveAttributeOffsets[HAIR_CURVE_ATTRIBUTE_OFFSET_COUNT]; \
+	uint4						Prefix##_PointAttributeOffsets[HAIR_POINT_ATTRIBUTE_OFFSET_COUNT]; \
+	uint 						Prefix##_CurveAttributeIndexToChunkDivAsShift; \
+	uint						Prefix##_CurveAttributeChunkElementCount; \
+	uint						Prefix##_CurveAttributeChunkStrideInBytes; \
+	uint						Prefix##_PointAttributeIndexToChunkDivAsShift; \
+	uint						Prefix##_PointAttributeChunkElementCount; \
+	uint						Prefix##_PointAttributeChunkStrideInBytes; \
+	Buffer<uint4>				Prefix##_PositionBuffer; \
+	StructuredBuffer<float4>	Prefix##_PositionOffsetBuffer; \
+	Buffer<uint>				Prefix##_CurveBuffer; \
+	Buffer<uint>				Prefix##_PointToCurveBuffer; \
+	ByteAddressBuffer 			Prefix##_CurveAttributeBuffer; \
+	ByteAddressBuffer 			Prefix##_PointAttributeBuffer; \
+	Buffer<uint>				Prefix##_CullingIndirectBuffer; \
+	Buffer<uint>				Prefix##_CullingIndexBuffer; \
+	Buffer<float>				Prefix##_CullingRadiusScaleBuffer; \
+	Buffer<uint> 				Prefix##_CullingIndirectBufferArgs; \
+	float3 Prefix##_GetHairInstancePositionOffset() { return ReadRenPositionOffset(Prefix##_PositionOffsetBuffer, Prefix##_RegisteredIndex); }
 
 #if HAIR_STRANDS_PARAMETERS == 1
 	HAIR_STRANDS_INSTANCE_PARAMETERS(HairStrandsVF)
diff --git a/Engine/Shaders/Private/HairStrands/HairStrandsVisibilityRasterClassification.usf b/Engine/Shaders/Private/HairStrands/HairStrandsVisibilityRasterClassification.usf
--- a/Engine/Shaders/Private/HairStrands/HairStrandsVisibilityRasterClassification.usf
+++ b/Engine/Shaders/Private/HairStrands/HairStrandsVisibilityRasterClassification.usf
@@ -71,7 +71,7 @@ void CSMain(uint DispatchThreadID : SV_DispatchThreadID, uint GroupThreadID : SV
 	const uint NumBatches = (ControlPointCount + 1023) / 1024;
 #endif // PERMUTATION_CULLING
 	const uint NumLoops = (NumBatches + (NumWorkGroups - 1)) / NumWorkGroups;
-	const float3 PositionOffset = HairStrandsVF_PositionOffsetBuffer[0].xyz;
+	const float3 PositionOffset = HairStrandsVF_GetHairInstancePositionOffset();
 
 	LOOP for (uint LoopIndex = 0; LoopIndex < NumLoops; ++LoopIndex)
 	{
diff --git a/Engine/Shaders/Private/HairStrands/HairStrandsVisibilityRasterCompute.usf b/Engine/Shaders/Private/HairStrands/HairStrandsVisibilityRasterCompute.usf
--- a/Engine/Shaders/Private/HairStrands/HairStrandsVisibilityRasterCompute.usf
+++ b/Engine/Shaders/Private/HairStrands/HairStrandsVisibilityRasterCompute.usf
@@ -467,17 +467,18 @@ void BinningCS(uint DispatchThreadID : SV_DispatchThreadID, uint GroupThreadID :
 		// Project segment end points and clip them to the screen
 		if (bSegValid)
 		{
+			const float3 InstancePositionOffset = HairStrandsVF_GetHairInstancePositionOffset();
 			float4 H0 = 0.0f;
 			float4 H1 = 0.0f;
 			uint Type = -1;
-			CalcHomogenousPos(HairStrandsVF_PositionBuffer[PrimID], HairStrandsVF_PositionOffsetBuffer[0].xyz, H0, Type);
+			CalcHomogenousPos(HairStrandsVF_PositionBuffer[PrimID], InstancePositionOffset, H0, Type);
 
 			bool bIsEndCV = (Type == HAIR_CONTROLPOINT_END);
 			bSegValid = !bIsEndCV;
 
 			if (bSegValid)
 			{
-				CalcHomogenousPos(HairStrandsVF_PositionBuffer[PrimID + 1], HairStrandsVF_PositionOffsetBuffer[0].xyz, H1, Type);
+				CalcHomogenousPos(HairStrandsVF_PositionBuffer[PrimID + 1], InstancePositionOffset, H1, Type);
 				
 				// Do clipping in homogenous coordinates
 				bSegValid = BlinnLineClipping(H0, H1);
@@ -884,7 +885,7 @@ void RasterCS(uint DispatchThreadID : SV_DispatchThreadID, uint GroupThreadID :
 		group_TileNum = VisTileArgs[0];
 		group_LoopNum = (float(group_TileNum) + float(NumRasterizers - 1)) * RcpNumRasterizers;
 
-		group_PositionOffset = HairStrandsVF_PositionOffsetBuffer[0].xyz;
+		group_PositionOffset = HairStrandsVF_GetHairInstancePositionOffset();
 
 		/* no longer in use - keep for ref? Moving these values to group shared memory did seem to reduce VGPRs - more experimentation needed
 		group_RadScale = (((HairStrandsVF_TipScale - HairStrandsVF_RootScale) * HairStrandsVF_Radius * OutputResolutionf.x) / 63.0) / 255.0;
@@ -1297,7 +1298,7 @@ void RasterMultiSampleCS(uint DispatchThreadID : SV_DispatchThreadID, uint Group
 		group_TileNum = VisTileArgs[0];
 		group_LoopNum = (float(group_TileNum) + float(NumRasterizers - 1)) * RcpNumRasterizers;
 
-		group_PositionOffset = HairStrandsVF_PositionOffsetBuffer[0].xyz;
+		group_PositionOffset = HairStrandsVF_GetHairInstancePositionOffset();
 	}
 
 	GroupMemoryBarrierWithGroupSync();
diff --git a/Engine/Shaders/Private/HairStrands/HairStrandsVisibilityRasterComputeNaive.usf b/Engine/Shaders/Private/HairStrands/HairStrandsVisibilityRasterComputeNaive.usf
--- a/Engine/Shaders/Private/HairStrands/HairStrandsVisibilityRasterComputeNaive.usf
+++ b/Engine/Shaders/Private/HairStrands/HairStrandsVisibilityRasterComputeNaive.usf
@@ -131,7 +131,7 @@ void CSMain(uint DispatchThreadID : SV_DispatchThreadID, uint GroupThreadID : SV
 
 	const uint NumBatches = (NumControlPoints + 63) / 64;
 	const uint NumLoops = (NumBatches + (NumWorkGroups - 1)) / NumWorkGroups;
-	const float3 PositionOffset = HairStrandsVF_PositionOffsetBuffer[0].xyz;
+	const float3 PositionOffset = HairStrandsVF_GetHairInstancePositionOffset();
 
 	LOOP for (uint LoopIndex = 0; LoopIndex < NumLoops; ++LoopIndex)
 	{
diff --git a/Engine/Shaders/Private/HairStrands/HairStrandsVisibilityRasterHW.usf b/Engine/Shaders/Private/HairStrands/HairStrandsVisibilityRasterHW.usf
--- a/Engine/Shaders/Private/HairStrands/HairStrandsVisibilityRasterHW.usf
+++ b/Engine/Shaders/Private/HairStrands/HairStrandsVisibilityRasterHW.usf
@@ -82,7 +82,7 @@ VSOutput VSMain(VSInput Input)
 	#endif // PERMUTATION_CULLING
 #endif // PERMUTATION_DRAW_INDIRECT
 
-	const float3 PositionOffset = HairStrandsVF_PositionOffsetBuffer[0].xyz;
+	const float3 PositionOffset = HairStrandsVF_GetHairInstancePositionOffset();
 
 	// Fetch both control points
 	const FHairControlPoint CP0 = UnpackHairControlPoint(
diff --git a/Engine/Shaders/Private/HairStrands/HairStrandsVoxelRasterCompute.usf b/Engine/Shaders/Private/HairStrands/HairStrandsVoxelRasterCompute.usf
--- a/Engine/Shaders/Private/HairStrands/HairStrandsVoxelRasterCompute.usf
+++ b/Engine/Shaders/Private/HairStrands/HairStrandsVoxelRasterCompute.usf
@@ -715,7 +715,7 @@ void MainCS(uint2 DispatchThreadID : SV_DispatchThreadID)
 	}
 #endif
 
-	const float3 PositionOffset = HairStrandsVF_PositionOffsetBuffer[0].xyz;
+	const float3 PositionOffset = HairStrandsVF_GetHairInstancePositionOffset();
 
 	const FHairControlPoint CP0 = UnpackHairControlPoint(
 		HairStrandsVF_PositionBuffer[VertexIndex0],
diff --git a/Engine/Source/Runtime/Renderer/Public/HairStrandsInterface.h b/Engine/Source/Runtime/Renderer/Public/HairStrandsInterface.h
--- a/Engine/Source/Runtime/Renderer/Public/HairStrandsInterface.h
+++ b/Engine/Source/Runtime/Renderer/Public/HairStrandsInterface.h
@@ -45,6 +45,7 @@ BEGIN_SHADER_PARAMETER_STRUCT(FHairStrandsInstanceCommonParameters, RENDERER_API
 	SHADER_PARAMETER(float,  Length)
 	SHADER_PARAMETER(float,  LengthScale)
 	SHADER_PARAMETER(float,  RaytracingRadiusScale)
+	SHADER_PARAMETER(uint32, RegisteredIndex)
 	SHADER_PARAMETER(uint32, GroupIndex)
 	SHADER_PARAMETER(uint32, GroupCount)
 	SHADER_PARAMETER(uint32, PointCount)
@@ -65,7 +66,7 @@ END_SHADER_PARAMETER_STRUCT()
 // Instance resources (RDG)
 BEGIN_SHADER_PARAMETER_STRUCT(FHairStrandsInstanceResourceParameters, RENDERER_API)
 	SHADER_PARAMETER_RDG_BUFFER_SRV(Buffer, PositionBuffer)
-	SHADER_PARAMETER_RDG_BUFFER_SRV(Buffer, PositionOffsetBuffer)
+	SHADER_PARAMETER_RDG_BUFFER_SRV(StructuredBuffer, PositionOffsetBuffer)
 	SHADER_PARAMETER_RDG_BUFFER_SRV(Buffer, CurveBuffer)
 	SHADER_PARAMETER_RDG_BUFFER_SRV(Buffer, PointToCurveBuffer)
 	SHADER_PARAMETER_RDG_BUFFER_SRV(ByteAddressBuffer, CurveAttributeBuffer)
@@ -75,7 +76,7 @@ END_SHADER_PARAMETER_STRUCT()
 // Instance prev. resources (RDG)
 BEGIN_SHADER_PARAMETER_STRUCT(FHairStrandsInstancePrevResourceParameters, RENDERER_API)
 	SHADER_PARAMETER_RDG_BUFFER_SRV(Buffer, PrevPositionBuffer)
-	SHADER_PARAMETER_RDG_BUFFER_SRV(Buffer, PrevPositionOffsetBuffer)
+	SHADER_PARAMETER_RDG_BUFFER_SRV(StructuredBuffer, PrevPositionOffsetBuffer)
 END_SHADER_PARAMETER_STRUCT()
 
 // Instance culling resources (RDG)
@@ -95,7 +96,7 @@ END_SHADER_PARAMETER_STRUCT()
 // Instance resources (Raw)
 BEGIN_SHADER_PARAMETER_STRUCT(FHairStrandsInstanceResourceRawParameters, RENDERER_API)
 	SHADER_PARAMETER_SRV(Buffer<uint4>, PositionBuffer)
-	SHADER_PARAMETER_SRV(Buffer<float4>, PositionOffsetBuffer)
+	SHADER_PARAMETER_SRV(StructuredBuffer<float4>, PositionOffsetBuffer)
 	SHADER_PARAMETER_SRV(Buffer<uint>, CurveBuffer)
 	SHADER_PARAMETER_SRV(Buffer<uint>, PointToCurveBuffer)
 	SHADER_PARAMETER_SRV(Buffer<float4>, TangentBuffer)
@@ -106,7 +107,7 @@ END_SHADER_PARAMETER_STRUCT()
 // Instance prev. resources (Raw)
 BEGIN_SHADER_PARAMETER_STRUCT(FHairStrandsInstancePrevResourceRawParameters, RENDERER_API)
 	SHADER_PARAMETER_SRV(Buffer<uint4>, PreviousPositionBuffer)
-	SHADER_PARAMETER_SRV(Buffer<float4>, PreviousPositionOffsetBuffer)
+	SHADER_PARAMETER_SRV(StructuredBuffer<float4>, PreviousPositionOffsetBuffer)
 END_SHADER_PARAMETER_STRUCT()
 
 // Instance culling resources (Raw)
