Fix for various minor issues noticed by enabling warnings as errors on global shaders

diff --git a/Engine/Plugins/FX/Niagara/Shaders/Private/NiagaraRayTracingShaders.usf b/Engine/Plugins/FX/Niagara/Shaders/Private/NiagaraRayTracingShaders.usf
--- a/Engine/Plugins/FX/Niagara/Shaders/Private/NiagaraRayTracingShaders.usf
+++ b/Engine/Plugins/FX/Niagara/Shaders/Private/NiagaraRayTracingShaders.usf
@@ -25,7 +25,7 @@ uint MaxRetraces;
 #if NIAGARA_SUPPORTS_COLLISION_GROUPS
 bool FindCollisionGroup(int GPUSceneIndexId, out int CollisionGroup)
 {
-	int Index;
+	uint Index = 0;
 	if (HashTableFind(GPUSceneIndexId, Index))
 	{
 		CollisionGroup = HashToCollisionGroups[Index];
diff --git a/Engine/Shaders/Private/HairStrands/HairStrandsForwardRaster.usf b/Engine/Shaders/Private/HairStrands/HairStrandsForwardRaster.usf
--- a/Engine/Shaders/Private/HairStrands/HairStrandsForwardRaster.usf
+++ b/Engine/Shaders/Private/HairStrands/HairStrandsForwardRaster.usf
@@ -866,7 +866,7 @@ void CompactionCS(uint DispatchThreadID : SV_DispatchThreadID, uint GroupThreadI
 
 			FDepthRange DepthRange;
 			DepthRange.MinZ = group_MinZ;
-			DepthRange.MinZ = group_MaxZ;
+			DepthRange.MaxZ = group_MaxZ;
 			StoreOutVisTileData(OutData, CompactedTile, VT_MinMaxDepth, PackDepthRange(DepthRange));
 		}
 	}
diff --git a/Engine/Shaders/Private/HairStrands/HairStrandsVisibilityRasterHW.usf b/Engine/Shaders/Private/HairStrands/HairStrandsVisibilityRasterHW.usf
--- a/Engine/Shaders/Private/HairStrands/HairStrandsVisibilityRasterHW.usf
+++ b/Engine/Shaders/Private/HairStrands/HairStrandsVisibilityRasterHW.usf
@@ -141,7 +141,7 @@ RWTexture2D<uint>	OutDepthCovTexture;
 RWTexture2D<uint>	OutPrimMatTexture;
 float				CoverageScale;
 
-[earlydepthstencil]
+EARLYDEPTHSTENCIL
 void PSMain(
 	float4 Position : SV_Position,
 	float Radius : RADIUS,
