Substrate: Improve type safety for Substrate's TopLayerTexture.

The type held in this texture is either uint or uint2 depending on the requested normal quality. By defining the type in an enum, we can actually enforce a bit more type safety and avoid shader warnings without casts.

diff --git a/Engine/Shaders/Private/Definitions.usf b/Engine/Shaders/Private/Definitions.usf
--- a/Engine/Shaders/Private/Definitions.usf
+++ b/Engine/Shaders/Private/Definitions.usf
@@ -90,14 +90,7 @@
 
 #ifndef SUBSTRATE_NORMAL_QUALITY
 #define SUBSTRATE_NORMAL_QUALITY							0
-#endif
-
-#if SUBSTRATE_NORMAL_QUALITY==0
-#define SUBSTRATE_TOP_LAYER_TYPE uint
-#elif SUBSTRATE_NORMAL_QUALITY==1
-#define SUBSTRATE_TOP_LAYER_TYPE uint2
-#else
-#error Uknown SUBSTRATE_NORMAL_QUALITY
+#define SUBSTRATE_TOP_LAYER_TYPE                            uint
 #endif
 
 #ifndef MATERIAL_SHADINGMODEL_DEFAULT_LIT
diff --git a/Engine/Shaders/Private/ScreenSpaceDenoise/SSDCommon.ush b/Engine/Shaders/Private/ScreenSpaceDenoise/SSDCommon.ush
--- a/Engine/Shaders/Private/ScreenSpaceDenoise/SSDCommon.ush
+++ b/Engine/Shaders/Private/ScreenSpaceDenoise/SSDCommon.ush
@@ -256,7 +256,7 @@ FSSDCompressedSceneInfos SampleCompressedSceneMetadata(
 	}
 
 	// SUBSTRATE_TODO: Manage bPrevFrame
-	const FSubstrateTopLayerData TopLayerData = SubstrateUnpackTopLayerData((SUBSTRATE_TOP_LAYER_TYPE)Substrate.TopLayerTexture.Load(uint3(PixelCoord, 0)));
+	const FSubstrateTopLayerData TopLayerData = SubstrateUnpackTopLayerData(Substrate.TopLayerTexture.Load(uint3(PixelCoord, 0)));
 	return MaterialToCompressedSceneMetadata(SceneDepth, TopLayerData.WorldNormal, TopLayerData.Roughness, IsSubstrateMaterial(TopLayerData) ? SHADINGMODELID_SUBSTRATE : SHADINGMODELID_UNLIT);
   #else // SUBSTRATE_ENABLED
 	FGBufferData GBufferData;
diff --git a/Engine/Shaders/Private/Substrate/Substrate.ush b/Engine/Shaders/Private/Substrate/Substrate.ush
--- a/Engine/Shaders/Private/Substrate/Substrate.ush
+++ b/Engine/Shaders/Private/Substrate/Substrate.ush
@@ -762,19 +762,13 @@ bool SubstrateIsTopLayerMaterial(SUBSTRATE_TOP_LAYER_TYPE In)
 #endif
 }
 
-FSubstrateTopLayerData SubstrateUnpackTopLayerData(uint In)
+FSubstrateTopLayerData SubstrateUnpackTopLayerData(SUBSTRATE_TOP_LAYER_TYPE In)
 {
 	FSubstrateTopLayerData Out = (FSubstrateTopLayerData)0;
+#if SUBSTRATE_NORMAL_QUALITY==0
 	Out.Material = (In & 0x3);
 	Out.WorldNormal = SubstrateUnpackNormal22(In >> 10);
 	Out.Roughness = UnpackR8(In >> 2);
-	return Out;
-}
-FSubstrateTopLayerData SubstrateUnpackTopLayerData(uint2 In)
-{
-	FSubstrateTopLayerData Out = (FSubstrateTopLayerData)0;
-#if SUBSTRATE_NORMAL_QUALITY==0
-	Out = SubstrateUnpackTopLayerData(In.x);
 #elif SUBSTRATE_NORMAL_QUALITY==1
 	Out.Material = (In.y & 0x3);
 	Out.WorldNormal = SubstrateUnpackNormal(In.x);
@@ -1894,22 +1888,20 @@ uint SubstrateLoadUint1AtIndexOffset(in FSubstrateMaterialContainer SubstrateBuf
 }
 
 // FSubstrateTopLayerDataContainer can be a uint or a uint2 resource, following SUBSTRATE_TOP_LAYER_TYPE according to r.GBufferFormat.
-// However, we do not change the type of the resource:
-//    - Leverage the hardware fetch unit filtering what can be loaded/stored according to the actual resource format.
-//    - We also filter the read data from SubstrateLoadTopLayerData 
+// The type of the texture is set dynamically to avoid compiler warnings and maximize type safety. Note that the SUBSTRATE_TOP_LAYER_TYPE
+// define is set at the C++ level so that it is always available, even inside global uniform buffer declarations.
 #if SUBSTRATE_MATERIALCONTAINER_IS_WRITABLE
-#define FSubstrateTopLayerDataContainer RWTexture2D<uint2>
+#define FSubstrateTopLayerDataContainer RWTexture2D<SUBSTRATE_TOP_LAYER_TYPE>
 #else
-#define FSubstrateTopLayerDataContainer Texture2D<uint2>
+#define FSubstrateTopLayerDataContainer Texture2D<SUBSTRATE_TOP_LAYER_TYPE>
 #endif
 
 SUBSTRATE_TOP_LAYER_TYPE SubstrateLoadTopLayerData(FSubstrateTopLayerDataContainer InTopLayerContainer, inout FSubstrateAddressing SubstrateAddressing)
 {
-	// Top layer container is always uint2 so we cast to XXX to avoid warnings.
 	#if SUBSTRATE_MATERIALCONTAINER_IS_WRITABLE
-	return (SUBSTRATE_TOP_LAYER_TYPE)InTopLayerContainer[SubstrateAddressing.PixelCoords];
+	return InTopLayerContainer[SubstrateAddressing.PixelCoords];
 	#else
-	return (SUBSTRATE_TOP_LAYER_TYPE)InTopLayerContainer.Load(uint3(SubstrateAddressing.PixelCoords, 0));
+	return InTopLayerContainer.Load(uint3(SubstrateAddressing.PixelCoords, 0));
 	#endif
 }
 
diff --git a/Engine/Shaders/Private/Substrate/SubstrateDBuffer.usf b/Engine/Shaders/Private/Substrate/SubstrateDBuffer.usf
--- a/Engine/Shaders/Private/Substrate/SubstrateDBuffer.usf
+++ b/Engine/Shaders/Private/Substrate/SubstrateDBuffer.usf
@@ -76,7 +76,7 @@ SamplerState DBufferCTextureSampler;
 int2 ViewResolution;
 uint MaxBytesPerPixel;
 uint FirstSliceStoringSubstrateSSSData;
-RWTexture2D<uint2> TopLayerTexture;
+RWTexture2D<SUBSTRATE_TOP_LAYER_TYPE> TopLayerTexture;
 RWTexture2DArray<uint> MaterialTextureArrayUAV;
 
 uint TileListBufferOffset;
diff --git a/Engine/Shaders/Private/Substrate/SubstrateMaterialClassification.usf b/Engine/Shaders/Private/Substrate/SubstrateMaterialClassification.usf
--- a/Engine/Shaders/Private/Substrate/SubstrateMaterialClassification.usf
+++ b/Engine/Shaders/Private/Substrate/SubstrateMaterialClassification.usf
@@ -17,7 +17,7 @@ int bRectPrimitive;
 int2 ViewResolution;
 uint MaxBytesPerPixel;
 int FirstSliceStoringSubstrateSSSData;
-Texture2D<uint2> TopLayerTexture;
+Texture2D<SUBSTRATE_TOP_LAYER_TYPE> TopLayerTexture;
 #if PERMUTATION_CMASK
 Texture2D<uint> TopLayerCmaskTexture;
 #endif
@@ -411,7 +411,7 @@ uint TileSizeLog2;
 
 int2 TileCount_Primary;
 
-Texture2D<uint2> TopLayerTexture;
+Texture2D<SUBSTRATE_TOP_LAYER_TYPE> TopLayerTexture;
 Texture2DArray<uint> MaterialTextureArray;
 
 Buffer<uint> TileListBuffer;
diff --git a/Engine/Shaders/Private/Substrate/SubstrateVisualize.usf b/Engine/Shaders/Private/Substrate/SubstrateVisualize.usf
--- a/Engine/Shaders/Private/Substrate/SubstrateVisualize.usf
+++ b/Engine/Shaders/Private/Substrate/SubstrateVisualize.usf
@@ -503,7 +503,7 @@ void PrintPixel(uint2 InCoord, float3 InWorldPosition, float3 V)
 	// Header
 	if (ClosureIndex < Header.ClosureCount && ClosureIndex == 0)
 	{
-		const FSubstrateTopLayerData TopLayerData = SubstrateUnpackTopLayerData((SUBSTRATE_TOP_LAYER_TYPE)Substrate.TopLayerTexture.Load(uint3(InCoord, 0)));
+		const FSubstrateTopLayerData TopLayerData = SubstrateUnpackTopLayerData(Substrate.TopLayerTexture.Load(uint3(InCoord, 0)));
 
 		const bool bSubstrateMaterial = Header.ClosureCount > 0;
 		const bool bIsSimpleMaterial = Header.IsSimpleMaterial() || Header.ClosureCount == 0;
diff --git a/Engine/Source/Runtime/Engine/Private/ShaderCompiler/ShaderCompiler.cpp b/Engine/Source/Runtime/Engine/Private/ShaderCompiler/ShaderCompiler.cpp
--- a/Engine/Source/Runtime/Engine/Private/ShaderCompiler/ShaderCompiler.cpp
+++ b/Engine/Source/Runtime/Engine/Private/ShaderCompiler/ShaderCompiler.cpp
@@ -8126,6 +8126,14 @@ void GlobalBeginCompileShader(
 
 			const uint32 SubstrateNormalQuality = Substrate::GetNormalQuality();
 			SET_SHADER_DEFINE(Input.Environment, SUBSTRATE_NORMAL_QUALITY, SubstrateNormalQuality);
+			if (SubstrateNormalQuality == 0)
+			{
+				SET_SHADER_DEFINE(Input.Environment, SUBSTRATE_TOP_LAYER_TYPE, TEXT("uint"));
+			}
+			else
+			{
+				SET_SHADER_DEFINE(Input.Environment, SUBSTRATE_TOP_LAYER_TYPE, TEXT("uint2"));
+			}
 
 			const uint32 SubstrateUintPerPixel = Substrate::GetBytePerPixel(Target.GetPlatform()) / 4u;
 			SET_SHADER_DEFINE(Input.Environment, SUBSTRATE_MATERIAL_NUM_UINTS, SubstrateUintPerPixel);
@@ -8145,6 +8153,11 @@ void GlobalBeginCompileShader(
 			const bool bSpecularProfileEnabled = Substrate::IsSpecularProfileEnabled(Target.GetPlatform());
 			SET_SHADER_DEFINE(Input.Environment, PLATFORM_ENABLES_SUBSTRATE_SPECULAR_PROFILE, bSpecularProfileEnabled ? 1 : 0);
 		}
+		else
+		{
+			// Some global uniform buffers reference this type -- so we need to have it defined in all cases
+			SET_SHADER_DEFINE(Input.Environment, SUBSTRATE_TOP_LAYER_TYPE, TEXT("uint"));
+		}
 
 		const bool bSubstrateBackCompatibility = bSubstrate && Substrate::IsBackCompatibilityEnabled();
 		SET_SHADER_DEFINE(Input.Environment, PROJECT_SUBSTRATE_BACKCOMPATIBILITY, bSubstrateBackCompatibility ? 1 : 0);
diff --git a/Engine/Source/Runtime/Renderer/Private/Substrate/Substrate.cpp b/Engine/Source/Runtime/Renderer/Private/Substrate/Substrate.cpp
--- a/Engine/Source/Runtime/Renderer/Private/Substrate/Substrate.cpp
+++ b/Engine/Source/Runtime/Renderer/Private/Substrate/Substrate.cpp
@@ -676,7 +676,7 @@ class FSubstrateClosureTilePassCS : public FGlobalShader
 		SHADER_PARAMETER(FIntPoint, TileCount_Primary)
 		SHADER_PARAMETER(FIntPoint, ViewResolution)
 		SHADER_PARAMETER(uint32, MaxBytesPerPixel)
-		SHADER_PARAMETER_RDG_TEXTURE(Texture2D, TopLayerTexture)
+		SHADER_PARAMETER_RDG_TEXTURE(Texture2D<SUBSTRATE_TOP_LAYER_TYPE>, TopLayerTexture)
 		SHADER_PARAMETER_RDG_TEXTURE_SRV(Texture2DArray<uint>, MaterialTextureArray)
 
 		SHADER_PARAMETER_RDG_TEXTURE_UAV(RWTexture2D<uint>, RWClosureOffsetTexture)
@@ -732,7 +732,7 @@ class FSubstrateMaterialTileClassificationPassCS : public FGlobalShader
 		SHADER_PARAMETER(FIntPoint, ViewResolution)
 		SHADER_PARAMETER(uint32, MaxBytesPerPixel)
 		SHADER_PARAMETER_ARRAY(FUintVector4, TileListBufferOffsets, [SUBSTRATE_TILE_TYPE_COUNT])
-		SHADER_PARAMETER_RDG_TEXTURE(Texture2D, TopLayerTexture)
+		SHADER_PARAMETER_RDG_TEXTURE(Texture2D<SUBSTRATE_TOP_LAYER_TYPE>, TopLayerTexture)
 		SHADER_PARAMETER_RDG_TEXTURE(Texture2D, TopLayerCmaskTexture)
 		SHADER_PARAMETER_RDG_BUFFER_UAV(RWBuffer, TileDrawIndirectDataBufferUAV)
 		SHADER_PARAMETER_RDG_BUFFER_UAV(RWBuffer, TileListBufferUAV)
@@ -787,7 +787,7 @@ class FSubstrateDBufferPassCS : public FGlobalShader
 		SHADER_PARAMETER(uint32, FirstSliceStoringSubstrateSSSData)
 		SHADER_PARAMETER_STRUCT_REF(FViewUniformShaderParameters, ViewUniformBuffer)
 		SHADER_PARAMETER_STRUCT_INCLUDE(FDBufferParameters, DBuffer)
-		SHADER_PARAMETER_RDG_TEXTURE_UAV(RWTexture2D<uint>, TopLayerTexture)
+		SHADER_PARAMETER_RDG_TEXTURE_UAV(RWTexture2D<SUBSTRATE_TOP_LAYER_TYPE>, TopLayerTexture)
 		SHADER_PARAMETER_RDG_TEXTURE_UAV(RWTexture2DArray<uint>, MaterialTextureArrayUAV)
 		SHADER_PARAMETER_RDG_BUFFER_SRV(Buffer<uint>, TileListBuffer)
 		SHADER_PARAMETER(uint32, TileListBufferOffset)
diff --git a/Engine/Source/Runtime/Renderer/Private/Substrate/Substrate.h b/Engine/Source/Runtime/Renderer/Private/Substrate/Substrate.h
--- a/Engine/Source/Runtime/Renderer/Private/Substrate/Substrate.h
+++ b/Engine/Source/Runtime/Renderer/Private/Substrate/Substrate.h
@@ -41,7 +41,7 @@ BEGIN_SHADER_PARAMETER_STRUCT(FSubstrateForwardPassUniformParameters, )
 	SHADER_PARAMETER_STRUCT_INCLUDE(FSubstrateCommonParameters, Common)
 	SHADER_PARAMETER(int32, FirstSliceStoringSubstrateSSSData)
 	SHADER_PARAMETER_RDG_TEXTURE(Texture2DArray<uint>, MaterialTextureArray)
-	SHADER_PARAMETER_RDG_TEXTURE(Texture2D<uint2>, TopLayerTexture)
+	SHADER_PARAMETER_RDG_TEXTURE(Texture2D<SUBSTRATE_TOP_LAYER_TYPE>, TopLayerTexture)
 END_SHADER_PARAMETER_STRUCT()
 
 BEGIN_SHADER_PARAMETER_STRUCT(FSubstrateMobileForwardPassUniformParameters, )
@@ -64,7 +64,7 @@ BEGIN_GLOBAL_SHADER_PARAMETER_STRUCT(FSubstrateGlobalUniformParameters, RENDERER
 	SHADER_PARAMETER(uint32, TileSizeLog2)
 	SHADER_PARAMETER(FIntPoint, TileCount)
 	SHADER_PARAMETER_RDG_TEXTURE(Texture2DArray<uint>, MaterialTextureArray)
-	SHADER_PARAMETER_RDG_TEXTURE(Texture2D<uint2>, TopLayerTexture)
+	SHADER_PARAMETER_RDG_TEXTURE(Texture2D<SUBSTRATE_TOP_LAYER_TYPE>, TopLayerTexture)
 	SHADER_PARAMETER_RDG_TEXTURE(Texture2D<float3>, OpaqueRoughRefractionTexture)
 	SHADER_PARAMETER_RDG_TEXTURE(Texture2D<uint>, ClosureOffsetTexture)
 	SHADER_PARAMETER_RDG_BUFFER_SRV(Buffer<uint>, ClosureTileBuffer)
@@ -72,7 +72,7 @@ BEGIN_GLOBAL_SHADER_PARAMETER_STRUCT(FSubstrateGlobalUniformParameters, RENDERER
 END_GLOBAL_SHADER_PARAMETER_STRUCT()
 
 BEGIN_GLOBAL_SHADER_PARAMETER_STRUCT(FSubstratePublicGlobalUniformParameters, RENDERER_API)
-	SHADER_PARAMETER_RDG_TEXTURE(Texture2D<uint2>, TopLayerTexture)
+	SHADER_PARAMETER_RDG_TEXTURE(Texture2D<SUBSTRATE_TOP_LAYER_TYPE>, TopLayerTexture)
 END_GLOBAL_SHADER_PARAMETER_STRUCT()
 
 // This must map to the SUBSTRATE_TILE_TYPE defines.
