LIGHT_SCENE_ID_BITS and LIGHT_SCENE_ID_MASK instead of hardcoding masks.

diff --git a/Engine/Shaders/Private/StochasticShadows/StochasticShadows.ush b/Engine/Shaders/Private/StochasticShadows/StochasticShadows.ush
--- a/Engine/Shaders/Private/StochasticShadows/StochasticShadows.ush
+++ b/Engine/Shaders/Private/StochasticShadows/StochasticShadows.ush
@@ -17,24 +17,27 @@
 #define DEBUG_MODE_VISUALIZE_TRACING  2
 
 // Limited by PackLightSample
-#define INVALID_LIGHT_SCENE_ID 0x3FFF
-#define LIGHT_SCENE_ID_MASK_SIZE (MAX_LIGHT_SCENE_ID / 32)
+#define LIGHT_SCENE_ID_BITS			(14)
+#define LIGHT_SCENE_ID_MASK			((1 << LIGHT_SCENE_ID_BITS) - 1)
+#define INVALID_LIGHT_SCENE_ID		(LIGHT_SCENE_ID_MASK)
 
-#define TILE_SIZE 8
-#define SHADOW_MASK_TILE_SIZE 8
-#define INVALID_SHADOW_MASK_TILE_INDEX 0xFFFFFFFF
-#define MAX_SHADOW_MASK_AGE 7
-#define MAX_SHADOW_MASK_HISTORY_LENGTH 7
+#define TILE_SIZE						8
+#define SHADOW_MASK_TILE_SIZE			8
+#define INVALID_SHADOW_MASK_TILE_INDEX	0xFFFFFFFF
+#define MAX_SHADOW_MASK_AGE				7
+#define MAX_SHADOW_MASK_HISTORY_LENGTH	7
 
 // Keep in sync with StochasticShadows.cpp
-#define MAX_LIGHT_SCENE_ID_XY 16
-#define MAX_LIGHT_SCENE_ID (MAX_LIGHT_SCENE_ID_XY * MAX_LIGHT_SCENE_ID_XY)
+#define MAX_LIGHT_SCENE_ID_XY		16
+#define MAX_LIGHT_SCENE_ID			(MAX_LIGHT_SCENE_ID_XY * MAX_LIGHT_SCENE_ID_XY)
+#define LIGHT_SCENE_ID_MASK_SIZE	(MAX_LIGHT_SCENE_ID / 32)
+
 #define SHADOW_MASK_ATLAS_SIZE_IN_TILES 512
-#define DOWNSAMPLE_FACTOR 2
-#define SHADOW_MASK_DOWNSAMPLE_FACTOR 2
+#define DOWNSAMPLE_FACTOR				2
+#define SHADOW_MASK_DOWNSAMPLE_FACTOR	2
 
 #ifndef NUM_SAMPLES_PER_PIXEL
-	#define NUM_SAMPLES_PER_PIXEL 4
+	#define NUM_SAMPLES_PER_PIXEL	4
 	#define NUM_SAMPLES_PER_PIXEL_Y 2
 	#define NUM_SAMPLES_PER_PIXEL_X 2
 #endif
@@ -210,7 +213,7 @@ struct FLightSample
 
 uint PackLightSample(FLightSample LightSample)
 {
-	uint PackedLightSample = LightSample.LightSceneId & 0x3FFF;
+	uint PackedLightSample = LightSample.LightSceneId & LIGHT_SCENE_ID_MASK;
 	PackedLightSample |= LightSample.bVisible ? 0x8000 : 0;
 	PackedLightSample |= LightSample.bCompleted ? 0x4000 : 0;
 	return PackedLightSample;
@@ -219,7 +222,7 @@ uint PackLightSample(FLightSample LightSample)
 FLightSample UnpackLightSample(uint PackedLightSample)
 {
 	FLightSample LightSample;
-	LightSample.LightSceneId = PackedLightSample & 0x3FFF;
+	LightSample.LightSceneId = PackedLightSample & LIGHT_SCENE_ID_MASK;
 	LightSample.bVisible = PackedLightSample & 0x8000 ? true : false;
 	LightSample.bCompleted = PackedLightSample & 0x4000 ? true : false;
 	return LightSample;
