Fix warnings about vector truncation in WaterInfoTextureMerge shader

Adjust the texture type to match what is being looked up. The type of WaterBodyTexture is left as float4 even though only the first two components are used since the texture format is actually RGBA at the moment.

diff --git a/Engine/Shaders/Private/WaterInfoTextureMerge.usf b/Engine/Shaders/Private/WaterInfoTextureMerge.usf
--- a/Engine/Shaders/Private/WaterInfoTextureMerge.usf
+++ b/Engine/Shaders/Private/WaterInfoTextureMerge.usf
@@ -2,13 +2,13 @@
 
 #include "Common.ush"
 
-Texture2D<float4> GroundDepthTexture;
+Texture2D<float> GroundDepthTexture;
 SamplerState GroundDepthTextureSampler;
 Texture2D<float4> WaterBodyTexture;
 SamplerState WaterBodyTextureSampler;
-Texture2D<float4> WaterBodyDepthTexture;
+Texture2D<float> WaterBodyDepthTexture;
 SamplerState WaterBodyDepthTextureSampler;
-Texture2D<float4> DilatedWaterBodyDepthTexture;
+Texture2D<float> DilatedWaterBodyDepthTexture;
 SamplerState DilatedWaterBodyDepthTextureSampler;
 
 float2 WaterHeightExtents;
@@ -38,7 +38,7 @@ void Main(
 	float2 ScreenUV = ScreenPosition * View.ScreenPositionScaleBias.xy + View.ScreenPositionScaleBias.wz;
 
 	// All input textures MUST be sampled with a point sampler to avoid a driver issue on older Android devices which occurs when using Load instead.
-	const float2 WaterBodyVelocity = WaterBodyTexture.SampleLevel(WaterBodyTextureSampler, ScreenUV, 0);
+	const float2 WaterBodyVelocity = WaterBodyTexture.SampleLevel(WaterBodyTextureSampler, ScreenUV, 0).xy;
 	const float WaterBodyDeviceDepth =  WaterBodyDepthTexture.SampleLevel(WaterBodyDepthTextureSampler, ScreenUV, 0);
 	const float DilationDeviceDepth = DilatedWaterBodyDepthTexture.SampleLevel(DilatedWaterBodyDepthTextureSampler, ScreenUV, 0);
 	const float GroundDeviceDepth =  GroundDepthTexture.SampleLevel(GroundDepthTextureSampler, ScreenUV, 0);
diff --git a/Engine/Source/Runtime/Renderer/Private/WaterInfoTextureRendering.cpp b/Engine/Source/Runtime/Renderer/Private/WaterInfoTextureRendering.cpp
--- a/Engine/Source/Runtime/Renderer/Private/WaterInfoTextureRendering.cpp
+++ b/Engine/Source/Runtime/Renderer/Private/WaterInfoTextureRendering.cpp
@@ -53,13 +53,13 @@ public:
 	BEGIN_SHADER_PARAMETER_STRUCT(FParameters, )
 		SHADER_PARAMETER_RDG_UNIFORM_BUFFER(FViewUniformShaderParameters, View)
 		SHADER_PARAMETER_STRUCT_INCLUDE(FSceneTextureShaderParameters, SceneTextures)
-		SHADER_PARAMETER_RDG_TEXTURE(Texture2D<float4>, GroundDepthTexture)
+		SHADER_PARAMETER_RDG_TEXTURE(Texture2D<float>, GroundDepthTexture)
 		SHADER_PARAMETER_SAMPLER(SamplerState, GroundDepthTextureSampler)
 		SHADER_PARAMETER_RDG_TEXTURE(Texture2D<float4>, WaterBodyTexture)
 		SHADER_PARAMETER_SAMPLER(SamplerState, WaterBodyTextureSampler)
-		SHADER_PARAMETER_RDG_TEXTURE(Texture2D<float4>, WaterBodyDepthTexture)
+		SHADER_PARAMETER_RDG_TEXTURE(Texture2D<float>, WaterBodyDepthTexture)
 		SHADER_PARAMETER_SAMPLER(SamplerState, WaterBodyDepthTextureSampler)
-		SHADER_PARAMETER_RDG_TEXTURE(Texture2D<float4>, DilatedWaterBodyDepthTexture)
+		SHADER_PARAMETER_RDG_TEXTURE(Texture2D<float>, DilatedWaterBodyDepthTexture)
 		SHADER_PARAMETER_SAMPLER(SamplerState, DilatedWaterBodyDepthTextureSampler)
 		SHADER_PARAMETER(FVector2f, WaterHeightExtents)
 		SHADER_PARAMETER(float, GroundZMin)
