Fix out of bound hair group AABB when there are more than 32 instances.

diff --git a/Engine/Shaders/Private/HairStrands/HairStrandsAABB.usf b/Engine/Shaders/Private/HairStrands/HairStrandsAABB.usf
--- a/Engine/Shaders/Private/HairStrands/HairStrandsAABB.usf
+++ b/Engine/Shaders/Private/HairStrands/HairStrandsAABB.usf
@@ -11,6 +11,7 @@ uint MacroGroupCount;
 float PixelSizeAtDepth1;
 float NumPixelPerVoxel;
 uint  VoxelPageResolution;
+uint  PassCount;
 
 // Group data packed as {uint3 AABBMin, uint3 AABBMax}
 Buffer<uint> RegisteredIndexBuffer;
@@ -26,15 +27,33 @@ void UpdateValue(inout FHairAABB B, in uint Index, in Buffer<int> InBuffer)
 groupshared float3 g_MinBounds[GROUP_SIZE];
 groupshared float3 g_MaxBounds[GROUP_SIZE];
 [numthreads(GROUP_SIZE, 1, 1)]
-void Main(uint2 DispatchThreadId : SV_DispatchThreadID)
+void Main(uint2 DispatchThreadId : SV_DispatchThreadID, uint LinearIndex : SV_GroupIndex)
 {
-	const uint ThreadId = DispatchThreadId.x;
-	if (ThreadId < MacroGroupCount)
+	const uint ThreadId = LinearIndex;
+	for (uint PassIt = 0; PassIt < PassCount; ++PassIt)
 	{
-		const uint RegisteredIndex = RegisteredIndexBuffer[ThreadId];
-		FHairAABB Bound = ReadHairAABB(RegisteredIndex, InGroupAABBBuffer);
-		g_MinBounds[ThreadId] = Bound.Min;
-		g_MaxBounds[ThreadId] = Bound.Max;
+		const uint PassOffset = PassIt * GROUP_SIZE;
+		const uint BoundIndex = PassOffset + ThreadId;
+		if (BoundIndex < MacroGroupCount)
+		{
+			const uint RegisteredIndex = RegisteredIndexBuffer[BoundIndex];
+			const FHairAABB Bound = ReadHairAABB(RegisteredIndex, InGroupAABBBuffer);
+			if (PassIt == 0)
+			{
+				g_MinBounds[LinearIndex] = Bound.Min;
+				g_MaxBounds[LinearIndex] = Bound.Max;
+			}
+			else
+			{
+				g_MinBounds[LinearIndex].x = min(g_MinBounds[LinearIndex].x, Bound.Min.x);
+				g_MinBounds[LinearIndex].y = min(g_MinBounds[LinearIndex].y, Bound.Min.y);
+				g_MinBounds[LinearIndex].z = min(g_MinBounds[LinearIndex].z, Bound.Min.z);
+
+				g_MaxBounds[LinearIndex].x = max(g_MaxBounds[LinearIndex].x, Bound.Max.x);
+				g_MaxBounds[LinearIndex].y = max(g_MaxBounds[LinearIndex].y, Bound.Max.y);
+				g_MaxBounds[LinearIndex].z = max(g_MaxBounds[LinearIndex].z, Bound.Max.z);
+			}
+		}
 	}
 	GroupMemoryBarrierWithGroupSync();
 	
diff --git a/Engine/Source/Runtime/Renderer/Private/HairStrands/HairStrandsCluster.cpp b/Engine/Source/Runtime/Renderer/Private/HairStrands/HairStrandsCluster.cpp
--- a/Engine/Source/Runtime/Renderer/Private/HairStrands/HairStrandsCluster.cpp
+++ b/Engine/Source/Runtime/Renderer/Private/HairStrands/HairStrandsCluster.cpp
@@ -22,6 +22,7 @@ class FHairMacroGroupAABBCS : public FGlobalShader
 	SHADER_USE_PARAMETER_STRUCT(FHairMacroGroupAABBCS, FGlobalShader);
 
 	BEGIN_SHADER_PARAMETER_STRUCT(FParameters, )
+		SHADER_PARAMETER(uint32, PassCount)
 		SHADER_PARAMETER(uint32, MacroGroupId)
 		SHADER_PARAMETER(uint32, MacroGroupCount)
 
@@ -83,9 +84,6 @@ static void AddHairMacroGroupAABBPass(
 
 	const float NumPixelPerVoxel = FMath::Clamp(GHairVirtualVoxel_NumPixelPerVoxel, 1.f, 50.f);
 
-	// Can only aggregate 32 GroupAABBs 
-	check(uint32(RegisteredIndices.Num()) < FHairMacroGroupAABBCS::GetGroupSize());
-
 	FHairMacroGroupAABBCS::FParameters* Parameters = GraphBuilder.AllocParameters<FHairMacroGroupAABBCS::FParameters>();
 	Parameters->MacroGroupId 					= MacroGroup.MacroGroupId;
 	Parameters->RegisteredIndexBuffer			= GraphBuilder.CreateSRV(RegisteredIndexBuffer, PF_R32_UINT);
@@ -97,6 +95,7 @@ static void AddHairMacroGroupAABBPass(
 	Parameters->VoxelPageResolution 			= VoxelPageResolution;
 	Parameters->View			    			= View.ViewUniformBuffer;
 	Parameters->MacroGroupCount					= MacroGroupCount;
+	Parameters->PassCount						= FMath::DivideAndRoundUp(MacroGroupCount, FHairMacroGroupAABBCS::GetGroupSize());
 	
 	TShaderMapRef<FHairMacroGroupAABBCS> ComputeShader(View.ShaderMap);
 	FComputeShaderUtils::AddPass(
