Stochastic Shadows
- Use LocalLightIndex to access ShadowMaskPageTable.
- Makes it possible to have up to 256 (MAX_LIGHT_SCENE_ID) in view without having to use Hash Table.

- Storing PrevLocalLightIndex (required to access HistoryShadowMask) in FLocalLightData.VirtualShadowMapId to avoid increasing FLocalLightData size.
- Keep track of previous frame LocalLightIndex using a map in FSceneViewState.

diff --git a/Engine/Shaders/Private/LightData.ush b/Engine/Shaders/Private/LightData.ush
--- a/Engine/Shaders/Private/LightData.ush
+++ b/Engine/Shaders/Private/LightData.ush
@@ -74,7 +74,7 @@ struct FLocalLightData
 	float4 SpotAnglesAndSourceRadiusPacked;
 	float4 LightDirectionAndShadowMask;
 	float4 LightTangentAndIESDataAndSpecularScale;
-	float4 RectDataAndVirtualShadowMapId;
+	float3 RectData;
 
 	// Derived terms after unpacking
 	/** Flag is true if clustered deferred is supported for this light. They are always first / together.*/
@@ -83,8 +83,10 @@ struct FLocalLightData
 	bool bStochasticShadowsSupported;
 	/** Flag is true if it is a simple light. They are always first / together.*/
 	bool bIsSimpleLight;
-	/** Virtual shadow map ID or INDEX_NONE if not present. */
+	/** Virtual shadow map ID or INDEX_NONE if not present. Only valid if STOCHASTIC_SHADOWS is NOT defined. */
 	int  VirtualShadowMapId;
+	/** Previous frame LocalLightIndex or INDEX_NONE if not present. Only valid if STOCHASTIC_SHADOWS is defined. */
+	int  PrevLocalLightIndex;
 	/** Index into GPUScene lights */
 	int LightSceneId;
 };
@@ -306,9 +308,9 @@ FDeferredLightData ConvertToDeferredLight(
 	Out.bRectLight				= LightType == LIGHT_TYPE_RECT;
 	Out.HairTransmittance		= InitHairTransmittanceData();
 	Out.RectLightData			= UnpackRectLightData(
-									asuint(In.RectDataAndVirtualShadowMapId.x),
-									asuint(In.RectDataAndVirtualShadowMapId.y),
-									asuint(In.RectDataAndVirtualShadowMapId.w));
+									asuint(In.RectData.x),
+									asuint(In.RectData.y),
+									asuint(In.RectData.z));
 	Out.IESAtlasIndex			= UnpackLigthIESAtlasIndex(In);
 
 	// Shadow
@@ -390,9 +392,9 @@ FLightShaderParameters ConvertFromLocal(const FLocalLightData In)
 	Out.SourceLength			= UnpackLightSourceLength(In);
 
 	const FRectLightData RectLightData = UnpackRectLightData(
-									asuint(In.RectDataAndVirtualShadowMapId.x),
-									asuint(In.RectDataAndVirtualShadowMapId.y),
-									asuint(In.RectDataAndVirtualShadowMapId.w));
+									asuint(In.RectData.x),
+									asuint(In.RectData.y),
+									asuint(In.RectData.z));
 	Out.RectLightBarnCosAngle	= RectLightData.BarnCosAngle;
 	Out.RectLightBarnLength		= RectLightData.BarnLength;
 	Out.RectLightAtlasUVOffset	= RectLightData.AtlasData.AtlasUVOffset;
diff --git a/Engine/Shaders/Private/LightGridCommon.ush b/Engine/Shaders/Private/LightGridCommon.ush
--- a/Engine/Shaders/Private/LightGridCommon.ush
+++ b/Engine/Shaders/Private/LightGridCommon.ush
@@ -239,8 +239,14 @@ FLocalLightData GetLocalLightData_Internal(
 	Out.LightDirectionAndShadowMask				= InData2;
 	Out.SpotAnglesAndSourceRadiusPacked			= InData3;
 	Out.LightTangentAndIESDataAndSpecularScale	= InData4;
-	Out.RectDataAndVirtualShadowMapId			= InData5;
-	Out.VirtualShadowMapId						= int(Out.RectDataAndVirtualShadowMapId.z);
+	Out.RectData								= InData5.xyz;
+#if STOCHASTIC_SHADOWS
+	Out.VirtualShadowMapId						= -1;
+	Out.PrevLocalLightIndex						= int(InData5.w);
+#else
+	Out.VirtualShadowMapId						= int(InData5.w);
+	Out.PrevLocalLightIndex						= -1;
+#endif
 	Out.LightSceneId							= int(Out.LightColorAndIdAndFalloffExponent.z);
 
 	return Out;
diff --git a/Engine/Shaders/Private/StochasticShadows/StochasticShadows.usf b/Engine/Shaders/Private/StochasticShadows/StochasticShadows.usf
--- a/Engine/Shaders/Private/StochasticShadows/StochasticShadows.usf
+++ b/Engine/Shaders/Private/StochasticShadows/StochasticShadows.usf
@@ -1,5 +1,7 @@
 // Copyright Epic Games, Inc. All Rights Reserved.
 
+#define STOCHASTIC_SHADOWS 1
+
 #include "../Common.ush"
 #include "../LightGridCommon.ush"
 #include "../ShaderPrint.ush"
diff --git a/Engine/Shaders/Private/StochasticShadows/StochasticShadows.ush b/Engine/Shaders/Private/StochasticShadows/StochasticShadows.ush
--- a/Engine/Shaders/Private/StochasticShadows/StochasticShadows.ush
+++ b/Engine/Shaders/Private/StochasticShadows/StochasticShadows.ush
@@ -47,6 +47,8 @@
 #define TILE_MODE_COMPLEX_SHADING	1
 #define TILE_MODE_MAX				2
 
+#define SHADOW_MASK_USING_LOCAL_LIGHT_INDEX 1
+
 uint PackTile(uint2 TileCoord)
 {
 	return TileCoord.x | (TileCoord.y << 16);
@@ -258,19 +260,23 @@ float2 DownsampledScreenCoordToScreenUV(uint2 DownsampledScreenCoord)
 	return ScreenUV;
 }
 
-uint2 PackShadowMaskTile(uint LightIndex, uint2 ShadowMaskScreenCoord)
+uint2 PackShadowMaskTile(uint LightSceneId, uint LocalLightIndex, uint2 ShadowMaskScreenCoord)
 {
 	ShadowMaskScreenCoord /= SHADOW_MASK_TILE_SIZE;
 
+	checkSlow(LightSceneId <= LIGHT_SCENE_ID_MASK);
+	checkSlow(LocalLightIndex <= (0xFFFFFFFF >> MAX_LIGHT_SCENE_ID));
+
 	uint2 PackedTile;
-	PackedTile.x = LightIndex;
+	PackedTile.x = (LightSceneId & LIGHT_SCENE_ID_MASK) | (LocalLightIndex << LIGHT_SCENE_ID_BITS);
 	PackedTile.y = (ShadowMaskScreenCoord.x & 0xFFFF) | ((ShadowMaskScreenCoord.y & 0xFFFF) << 16);
 	return PackedTile;
 }
 
-void UnpackShadowMaskTile(uint2 PackedTile, inout uint LightIndex, inout uint2 ShadowMaskScreenCoord)
+void UnpackShadowMaskTile(uint2 PackedTile, out uint LightSceneId, out uint LocalLightIndex, out uint2 ShadowMaskScreenCoord)
 {
-	LightIndex = PackedTile.x;
+	LightSceneId = PackedTile.x & LIGHT_SCENE_ID_MASK;
+	LocalLightIndex = PackedTile.x >> LIGHT_SCENE_ID_BITS;
 	ShadowMaskScreenCoord.x = PackedTile.y & 0xFFFF;
 	ShadowMaskScreenCoord.y = PackedTile.y >> 16;
 
@@ -320,11 +326,15 @@ Texture2D<uint> ShadowMaskAtlasHistory;
 int2 ShadowMaskViewSize;
 uint ShadowMaskHashTableIndexWrapMask;
 
-void AccumulateHistoryShadowMask(uint LightSceneId, uint2 ShadowMaskScreenCoord, float Weight, inout FShadowMaskAccumulator Accumulator)
+void AccumulateHistoryShadowMask(uint LightSceneId, uint PrevLocalLightIndex, uint2 ShadowMaskScreenCoord, float Weight, inout FShadowMaskAccumulator Accumulator)
 {
 	if (all(ShadowMaskScreenCoord < ShadowMaskViewSize) && Weight > 0.01f)
 	{
+#if SHADOW_MASK_USING_LOCAL_LIGHT_INDEX
+		uint2 HistoryShadowMaskPageTableCoord = GetShadowMaskPageTableCoord(PrevLocalLightIndex, ShadowMaskScreenCoord);
+#else
 		uint2 HistoryShadowMaskPageTableCoord = GetShadowMaskPageTableCoord(LightSceneId, ShadowMaskScreenCoord);
+#endif
 
 #if HASH_TABLE
 		uint HistoryShadowMaskTileIndex = INVALID_SHADOW_MASK_TILE_INDEX;
diff --git a/Engine/Shaders/Private/StochasticShadows/StochasticShadowsComposite.usf b/Engine/Shaders/Private/StochasticShadows/StochasticShadowsComposite.usf
--- a/Engine/Shaders/Private/StochasticShadows/StochasticShadowsComposite.usf
+++ b/Engine/Shaders/Private/StochasticShadows/StochasticShadowsComposite.usf
@@ -1,5 +1,7 @@
 // Copyright Epic Games, Inc. All Rights Reserved.
 
+#define STOCHASTIC_SHADOWS 1
+
 #include "../Common.ush"
 #include "../LightGridCommon.ush"
 #include "../ShaderPrint.ush"
@@ -52,9 +54,10 @@ void CompositeShadowMaskTracesCS(
 {
 	uint ShadowMaskTileIndex = GroupId.x;
 
-	uint LightSceneId = 0;
-	uint2 BaseShadowMaskScreenCoord = 0;	
-	UnpackShadowMaskTile(ShadowMaskTileData[ShadowMaskTileIndex], LightSceneId, BaseShadowMaskScreenCoord);
+	uint LightSceneId;
+	uint LocalLightIndex;
+	uint2 BaseShadowMaskScreenCoord;	
+	UnpackShadowMaskTile(ShadowMaskTileData[ShadowMaskTileIndex], LightSceneId, LocalLightIndex, BaseShadowMaskScreenCoord);
 
 	GroupMemoryBarrierWithGroupSync();
 	
@@ -186,10 +189,13 @@ void CompositeShadowMaskTracesCS(
 			{
 				FShadowMaskAccumulator Accumulator = InitShadowMaskAccumulator();
 
-				AccumulateHistoryShadowMask(LightSceneId, ScreenCoord00 + uint2(0, 0), Weights.x, Accumulator);
-				AccumulateHistoryShadowMask(LightSceneId, ScreenCoord00 + uint2(1, 0), Weights.y, Accumulator);
-				AccumulateHistoryShadowMask(LightSceneId, ScreenCoord00 + uint2(0, 1), Weights.z, Accumulator);
-				AccumulateHistoryShadowMask(LightSceneId, ScreenCoord00 + uint2(1, 1), Weights.w, Accumulator);
+				const FLocalLightData LocalLightData = GetLocalLightDataNonStereo(LocalLightIndex);
+				const uint PrevLocalLightIndex = LocalLightData.PrevLocalLightIndex;
+
+				AccumulateHistoryShadowMask(LightSceneId, PrevLocalLightIndex, ScreenCoord00 + uint2(0, 0), Weights.x, Accumulator);
+				AccumulateHistoryShadowMask(LightSceneId, PrevLocalLightIndex, ScreenCoord00 + uint2(1, 0), Weights.y, Accumulator);
+				AccumulateHistoryShadowMask(LightSceneId, PrevLocalLightIndex, ScreenCoord00 + uint2(0, 1), Weights.z, Accumulator);
+				AccumulateHistoryShadowMask(LightSceneId, PrevLocalLightIndex, ScreenCoord00 + uint2(1, 1), Weights.w, Accumulator);
 
 				if (Accumulator.WeightSum > 0.0f)
 				{
diff --git a/Engine/Shaders/Private/StochasticShadows/StochasticShadowsHardwareRayTracing.usf b/Engine/Shaders/Private/StochasticShadows/StochasticShadowsHardwareRayTracing.usf
--- a/Engine/Shaders/Private/StochasticShadows/StochasticShadowsHardwareRayTracing.usf
+++ b/Engine/Shaders/Private/StochasticShadows/StochasticShadowsHardwareRayTracing.usf
@@ -1,5 +1,7 @@
 // Copyright Epic Games, Inc. All Rights Reserved.
 
+#define STOCHASTIC_SHADOWS 1
+
 #include "../Common.ush"
 #include "../BlueNoise.ush"
 #include "../RayTracing/RayTracingCommon.ush"
diff --git a/Engine/Shaders/Private/StochasticShadows/StochasticShadowsSampling.usf b/Engine/Shaders/Private/StochasticShadows/StochasticShadowsSampling.usf
--- a/Engine/Shaders/Private/StochasticShadows/StochasticShadowsSampling.usf
+++ b/Engine/Shaders/Private/StochasticShadows/StochasticShadowsSampling.usf
@@ -1,5 +1,7 @@
 // Copyright Epic Games, Inc. All Rights Reserved.
 
+#define STOCHASTIC_SHADOWS 1
+
 #define NON_DIRECTIONAL_DIRECT_LIGHTING 1
 #define SUPPORT_CONTACT_SHADOWS 0
 
@@ -108,6 +110,8 @@ Texture2D<float> ShadowMaskSceneDepthHistory;
 	groupshared uint SharedLightMask[LIGHT_SCENE_ID_MASK_SIZE];
 #endif
 
+#define CANDIDATE_LIGHT_MASK_USING_LOCAL_LIGHT_INDEX 1
+
 #if CANDIDATE_LIGHT_MASK
 	groupshared uint SharedCandidateLightMask[LIGHT_SCENE_ID_MASK_SIZE];
 #endif
@@ -160,11 +164,15 @@ void MarkLightSceneId(uint LightSceneId)
 #endif
 }
 
-void WriteShadowMaskTileToPageTable(uint LightSceneId, uint2 ShadowMaskScreenCoord, uint ShadowMaskTileIndex)
+void WriteShadowMaskTileToPageTable(uint LightSceneId, uint LocalLightIndex, uint2 ShadowMaskScreenCoord, uint ShadowMaskTileIndex)
 {
+#if SHADOW_MASK_USING_LOCAL_LIGHT_INDEX
+	uint2 ShadowMaskPageTableCoord = GetShadowMaskPageTableCoord(LocalLightIndex, ShadowMaskScreenCoord);
+#else
 	uint2 ShadowMaskPageTableCoord = GetShadowMaskPageTableCoord(LightSceneId, ShadowMaskScreenCoord);
+#endif
 
-	RWShadowMaskTileData[ShadowMaskTileIndex] = PackShadowMaskTile(LightSceneId, ShadowMaskScreenCoord);
+	RWShadowMaskTileData[ShadowMaskTileIndex] = PackShadowMaskTile(LightSceneId, LocalLightIndex, ShadowMaskScreenCoord);
 
 	#if HASH_TABLE
 	{
@@ -269,13 +277,26 @@ void GenerateSamplesCS(
 
 			for (uint GridLightIndex = 0; GridLightIndex < NumLightsInGridCell; ++GridLightIndex)
 			{
+#if CANDIDATE_LIGHT_MASK_USING_LOCAL_LIGHT_INDEX
+				const uint LocalLightIndex = GetCulledLightDataGrid(CulledLightGridData.DataStartIndex + GridLightIndex);
+				const FLocalLightData LocalLightData = GetLocalLightDataNonStereo(LocalLightIndex);
+
+				if (LocalLightData.bStochasticShadowsSupported && LocalLightIndex < MAX_LIGHT_SCENE_ID)
+				{
+					uint DWORDIndex = LocalLightIndex / 32;
+					uint BitMask = 1u << (LocalLightIndex % 32);
+					InterlockedOr(SharedCandidateLightMask[DWORDIndex], BitMask);
+				}
+#else
 				const FLocalLightData LocalLightData = GetLocalLightData(CulledLightGridData.DataStartIndex + GridLightIndex, EyeIndex);
+
 				if (LocalLightData.bStochasticShadowsSupported && LocalLightData.LightSceneId < MAX_LIGHT_SCENE_ID)
 				{
 					uint DWORDIndex = LocalLightData.LightSceneId / 32;
 					uint BitMask = 1u << (LocalLightData.LightSceneId % 32);
 					InterlockedOr(SharedCandidateLightMask[DWORDIndex], BitMask);
 				}
+#endif
 			}
 		}
 		#endif
@@ -365,6 +386,7 @@ void GenerateSamplesCS(
 				const uint GridIndex = ComputeLightGridCellIndex(ScreenCoord, SceneDepth, EyeIndex);
 				const FCulledLightsGridData CulledLightGridData = GetCulledLightsGrid(GridIndex, EyeIndex);
 				const uint NumLightsInGridCell = min(CulledLightGridData.NumLights, GetMaxLightsPerCell(EyeIndex));
+				const uint NumLocalLights = GetNumLocalLights(EyeIndex);
 
 				float LightIndexRandom[NUM_SAMPLES_PER_PIXEL];
 				for (uint LightSampleIndex = 0; LightSampleIndex < NUM_SAMPLES_PER_PIXEL; ++LightSampleIndex)
@@ -390,6 +412,15 @@ void GenerateSamplesCS(
 					Newline(Context);
 					Print(Context, TEXT("NumLightsInGridCell: "));
 					Print(Context, NumLightsInGridCell);
+					Print(Context, TEXT("NumLocalLights: "));
+					if (NumLocalLights <= MAX_LIGHT_SCENE_ID)
+					{
+						Print(Context, NumLocalLights);
+					}
+					else
+					{
+						Print(Context, NumLocalLights, FontRed);
+					}
 					AddCrossTWS(Context, TranslatedWorldPosition, 5.0f, float4(1, 1, 0, 1));
 					Newline(Context);
 					Print(Context, TEXT("LightId  | SampleWeight | Weight   | Estimate | Shadow | ShadowValid | ShadowAge | HistoryLength"));
@@ -411,10 +442,17 @@ void GenerateSamplesCS(
 						const uint NextBitMask = 1u << NextBitIndex;
 						CandidateLightMask ^= NextBitMask;
 
-						uint LightSceneId = MaskIndex * 32 + NextBitIndex;
-						FLightSceneData LightSceneData = GetLightSceneData(LightSceneId);
+#if CANDIDATE_LIGHT_MASK_USING_LOCAL_LIGHT_INDEX
+						const uint LocalLightIndex = MaskIndex * 32 + NextBitIndex;
+						const FLocalLightData LocalLightData = GetLocalLightDataNonStereo(LocalLightIndex);
+						const uint LightSceneId = LocalLightData.LightSceneId;
+						FDeferredLightData LightData = ConvertToDeferredLight(LocalLightData);
+#else
+						const uint LightSceneId = MaskIndex * 32 + NextBitIndex;
+						const FLightSceneData LightSceneData = GetLightSceneData(LightSceneId);
 						FDeferredLightData LightData = ConvertToDeferredLight(LightSceneData);
-
+#endif
+					
 						float SampleWeight = GetLocalLightTargetPDF(LightData, TranslatedWorldPosition, GBuffer, ScreenCoord);
 
 						// Simulate tonemapping
@@ -430,10 +468,16 @@ void GenerateSamplesCS(
 							{
 								FShadowMaskAccumulator Accumulator = InitShadowMaskAccumulator();
 
-								AccumulateHistoryShadowMask(LightSceneId, ShadowMaskHistoryScreenCoord00 + uint2(0, 0), ShadowMaskHistoryWeights.x, Accumulator);
-								AccumulateHistoryShadowMask(LightSceneId, ShadowMaskHistoryScreenCoord00 + uint2(1, 0), ShadowMaskHistoryWeights.y, Accumulator);
-								AccumulateHistoryShadowMask(LightSceneId, ShadowMaskHistoryScreenCoord00 + uint2(0, 1), ShadowMaskHistoryWeights.z, Accumulator);
-								AccumulateHistoryShadowMask(LightSceneId, ShadowMaskHistoryScreenCoord00 + uint2(1, 1), ShadowMaskHistoryWeights.w, Accumulator);
+#if SHADOW_MASK_USING_LOCAL_LIGHT_INDEX
+								uint PrevLocalLightIndex = LocalLightData.PrevLocalLightIndex;
+#else
+								uint PrevLocalLightIndex = 0; // unused
+#endif
+
+								AccumulateHistoryShadowMask(LightSceneId, PrevLocalLightIndex, ShadowMaskHistoryScreenCoord00 + uint2(0, 0), ShadowMaskHistoryWeights.x, Accumulator);
+								AccumulateHistoryShadowMask(LightSceneId, PrevLocalLightIndex, ShadowMaskHistoryScreenCoord00 + uint2(1, 0), ShadowMaskHistoryWeights.y, Accumulator);
+								AccumulateHistoryShadowMask(LightSceneId, PrevLocalLightIndex, ShadowMaskHistoryScreenCoord00 + uint2(0, 1), ShadowMaskHistoryWeights.z, Accumulator);
+								AccumulateHistoryShadowMask(LightSceneId, PrevLocalLightIndex, ShadowMaskHistoryScreenCoord00 + uint2(1, 1), ShadowMaskHistoryWeights.w, Accumulator);
 
 								if (Accumulator.WeightSum > 0.0f)
 								{
@@ -464,7 +508,11 @@ void GenerateSamplesCS(
 						if (SampleWeight > SamplingMinWeight)
 						{
 							// Allocate required shadow mask tiles for this frame
+#if SHADOW_MASK_USING_LOCAL_LIGHT_INDEX
+							MarkLightSceneId(LocalLightIndex);
+#else
 							MarkLightSceneId(LightSceneId);
+#endif
 
 							// Apply shadowing estimate from the last frame
 							if (HistoryShadowMask.IsValid())
@@ -500,7 +548,12 @@ void GenerateSamplesCS(
 	#else
 				for (uint GridLightIndex = 0; GridLightIndex < NumLightsInGridCell; ++GridLightIndex)
 				{
+#if SHADOW_MASK_USING_LOCAL_LIGHT_INDEX
+					const uint LocalLightIndex = GetCulledLightDataGrid(CulledLightGridData.DataStartIndex + GridLightIndex);
+					const FLocalLightData LocalLightData = GetLocalLightDataNonStereo(LocalLightIndex);
+#else
 					const FLocalLightData LocalLightData = GetLocalLightData(CulledLightGridData.DataStartIndex + GridLightIndex, EyeIndex);
+#endif
 					if (LocalLightData.bStochasticShadowsSupported)
 					{
 						FDeferredLightData LightData = ConvertToDeferredLight(LocalLightData);
@@ -519,10 +572,12 @@ void GenerateSamplesCS(
 							{
 								FShadowMaskAccumulator Accumulator = InitShadowMaskAccumulator();
 
-								AccumulateHistoryShadowMask(LocalLightData.LightSceneId, ShadowMaskHistoryScreenCoord00 + uint2(0, 0), ShadowMaskHistoryWeights.x, Accumulator);
-								AccumulateHistoryShadowMask(LocalLightData.LightSceneId, ShadowMaskHistoryScreenCoord00 + uint2(1, 0), ShadowMaskHistoryWeights.y, Accumulator);
-								AccumulateHistoryShadowMask(LocalLightData.LightSceneId, ShadowMaskHistoryScreenCoord00 + uint2(0, 1), ShadowMaskHistoryWeights.z, Accumulator);
-								AccumulateHistoryShadowMask(LocalLightData.LightSceneId, ShadowMaskHistoryScreenCoord00 + uint2(1, 1), ShadowMaskHistoryWeights.w, Accumulator);
+								uint PrevLocalLightIndex = LocalLightData.PrevLocalLightIndex;
+
+								AccumulateHistoryShadowMask(LocalLightData.LightSceneId, PrevLocalLightIndex, ShadowMaskHistoryScreenCoord00 + uint2(0, 0), ShadowMaskHistoryWeights.x, Accumulator);
+								AccumulateHistoryShadowMask(LocalLightData.LightSceneId, PrevLocalLightIndex, ShadowMaskHistoryScreenCoord00 + uint2(1, 0), ShadowMaskHistoryWeights.y, Accumulator);
+								AccumulateHistoryShadowMask(LocalLightData.LightSceneId, PrevLocalLightIndex, ShadowMaskHistoryScreenCoord00 + uint2(0, 1), ShadowMaskHistoryWeights.z, Accumulator);
+								AccumulateHistoryShadowMask(LocalLightData.LightSceneId, PrevLocalLightIndex, ShadowMaskHistoryScreenCoord00 + uint2(1, 1), ShadowMaskHistoryWeights.w, Accumulator);
 
 								if (Accumulator.WeightSum > 0.0f)
 								{
@@ -553,7 +608,11 @@ void GenerateSamplesCS(
 						if (SampleWeight > SamplingMinWeight)
 						{
 							// Allocate required shadow mask tiles for this frame
+#if SHADOW_MASK_USING_LOCAL_LIGHT_INDEX
+							MarkLightSceneId(LocalLightIndex);
+#else
 							MarkLightSceneId(LocalLightData.LightSceneId);
+#endif
 
 							// Apply shadowing estimate from the last frame
 							if (HistoryShadowMask.IsValid())
@@ -638,17 +697,26 @@ void GenerateSamplesCS(
 		{
 			uint2 ShadowMaskScreenCoord = DownsampledTileCoord * SHADOW_MASK_TILE_SIZE;
 
-			uint LightIndex = SharedHashTable[LinearThreadIndex];
+			const uint StoredValue = SharedHashTable[LinearThreadIndex];
 
-			if (LightIndex != INVALID_LIGHT_SCENE_ID)
+			if (StoredValue != INVALID_LIGHT_SCENE_ID)
 			{
+#if SHADOW_MASK_USING_LOCAL_LIGHT_INDEX
+				const uint LocalLightIndex = StoredValue;
+				const FLocalLightData LocalLightData = GetLocalLightDataNonStereo(LocalLightIndex);
+				const uint LightIndex = LocalLightData.LightSceneId;
+#else
+				const uint LocalLightIndex = 0; // unused
+				const uint LightIndex = StoredValue;
+#endif
+
 				// #sdl_todo: merge allocations through groupshared and do only one atomic per group
 				uint ShadowMaskTileIndex = 0;
 				InterlockedAdd(RWShadowMaskTileAllocator[0], 1, ShadowMaskTileIndex);
 
 				if (ShadowMaskTileIndex < MaxShadowMaskTiles)
 				{
-					WriteShadowMaskTileToPageTable(LightIndex, ShadowMaskScreenCoord, ShadowMaskTileIndex);
+					WriteShadowMaskTileToPageTable(LightIndex, LocalLightIndex, ShadowMaskScreenCoord, ShadowMaskTileIndex);
 				}
 			}
 		}
@@ -668,7 +736,14 @@ void GenerateSamplesCS(
 				const uint NextBitMask = 1u << NextBitIndex;
 				LightMask ^= NextBitMask;
 
-				uint LightIndex = LinearThreadIndex + NextBitIndex * LIGHT_SCENE_ID_MASK_SIZE;
+#if SHADOW_MASK_USING_LOCAL_LIGHT_INDEX
+				const uint LocalLightIndex = LinearThreadIndex + NextBitIndex * LIGHT_SCENE_ID_MASK_SIZE;
+				const FLocalLightData LocalLightData = GetLocalLightDataNonStereo(LocalLightIndex);
+				const uint LightIndex = LocalLightData.LightSceneId;
+#else
+				const uint LocalLightIndex = 0; // unused
+				const uint LightIndex = LinearThreadIndex + NextBitIndex * LIGHT_SCENE_ID_MASK_SIZE;
+#endif
 
 				// #sdl_todo: merge allocations through groupshared and do only one atomic per group
 				uint ShadowMaskTileIndex = 0;
@@ -676,7 +751,7 @@ void GenerateSamplesCS(
 
 				if (ShadowMaskTileIndex < MaxShadowMaskTiles)
 				{
-					WriteShadowMaskTileToPageTable(LightIndex, ShadowMaskScreenCoord, ShadowMaskTileIndex);
+					WriteShadowMaskTileToPageTable(LightIndex, LocalLightIndex, ShadowMaskScreenCoord, ShadowMaskTileIndex);
 				}
 			}
 		}
diff --git a/Engine/Shaders/Private/StochasticShadows/StochasticShadowsShading.usf b/Engine/Shaders/Private/StochasticShadows/StochasticShadowsShading.usf
--- a/Engine/Shaders/Private/StochasticShadows/StochasticShadowsShading.usf
+++ b/Engine/Shaders/Private/StochasticShadows/StochasticShadowsShading.usf
@@ -1,5 +1,7 @@
 // Copyright Epic Games, Inc. All Rights Reserved.
 
+#define STOCHASTIC_SHADOWS 1
+
 #define NON_DIRECTIONAL_DIRECT_LIGHTING 1
 #define SUPPORT_CONTACT_SHADOWS 0
 
diff --git a/Engine/Shaders/Private/StochasticShadows/StochasticShadowsSpatial.usf b/Engine/Shaders/Private/StochasticShadows/StochasticShadowsSpatial.usf
--- a/Engine/Shaders/Private/StochasticShadows/StochasticShadowsSpatial.usf
+++ b/Engine/Shaders/Private/StochasticShadows/StochasticShadowsSpatial.usf
@@ -1,5 +1,7 @@
 // Copyright Epic Games, Inc. All Rights Reserved.
 
+#define STOCHASTIC_SHADOWS 1
+
 #include "../Common.ush"
 #include "../LightGridCommon.ush"
 #include "../ShaderPrint.ush"
@@ -57,9 +59,10 @@ void ShadowMaskSpatialPassCS(
 	uint ShadowMaskTileIndex = GroupId.x;
 	uint2 LocalShadowMaskTileCoord = GroupThreadId.xy;
 
-	uint LightSceneId = 0;
-	uint2 BaseShadowMaskScreenCoord = 0;	
-	UnpackShadowMaskTile(ShadowMaskTileData[ShadowMaskTileIndex], LightSceneId, BaseShadowMaskScreenCoord);
+	uint LightSceneId;
+	uint LocalLightIndex;
+	uint2 BaseShadowMaskScreenCoord;	
+	UnpackShadowMaskTile(ShadowMaskTileData[ShadowMaskTileIndex], LightSceneId, LocalLightIndex, BaseShadowMaskScreenCoord);
 
 	if (all(GroupThreadId == 0))
 	{
@@ -79,8 +82,12 @@ void ShadowMaskSpatialPassCS(
 		for (uint SharedShadowMaskX = GroupThreadId.x; SharedShadowMaskX < SHARED_TILE_SIZE; SharedShadowMaskX += THREADGROUP_SIZE)
 		{
 			uint2 SampleShadowMaskScreenCoord = clamp(int2(BaseShadowMaskScreenCoord) + int2(SharedShadowMaskX - SHARED_TILE_BORDER, SharedShadowMaskY - SHARED_TILE_BORDER), 0, ShadowMaskViewSize - 1);
-			uint2 SampleShadowMaskPageTableCoord = GetShadowMaskPageTableCoord(LightSceneId, SampleShadowMaskScreenCoord);
 
+#if SHADOW_MASK_USING_LOCAL_LIGHT_INDEX
+			uint2 SampleShadowMaskPageTableCoord = GetShadowMaskPageTableCoord(LocalLightIndex, SampleShadowMaskScreenCoord);
+#else
+			uint2 SampleShadowMaskPageTableCoord = GetShadowMaskPageTableCoord(LightSceneId, SampleShadowMaskScreenCoord);
+#endif
 			uint SampleShadowMaskTileIndex = INVALID_SHADOW_MASK_TILE_INDEX;
 			#if HASH_TABLE
 			{
diff --git a/Engine/Shaders/Private/StochasticShadows/StochasticShadowsTracing.usf b/Engine/Shaders/Private/StochasticShadows/StochasticShadowsTracing.usf
--- a/Engine/Shaders/Private/StochasticShadows/StochasticShadowsTracing.usf
+++ b/Engine/Shaders/Private/StochasticShadows/StochasticShadowsTracing.usf
@@ -1,5 +1,7 @@
 // Copyright Epic Games, Inc. All Rights Reserved.
 
+#define STOCHASTIC_SHADOWS 1
+
 #include "../Common.ush"
 #define DISTANCE_FIELD_IN_VIEW_UB 1
 #include "../BlueNoise.ush"
diff --git a/Engine/Source/Runtime/Renderer/Private/BasePassRendering.h b/Engine/Source/Runtime/Renderer/Private/BasePassRendering.h
--- a/Engine/Source/Runtime/Renderer/Private/BasePassRendering.h
+++ b/Engine/Source/Runtime/Renderer/Private/BasePassRendering.h
@@ -55,7 +55,7 @@ public:
 	FVector4f LightDirectionAndShadowMapChannelMask;
 	FVector4f SpotAnglesAndSourceRadiusPacked;
 	FVector4f LightTangentAndIESDataAndSpecularScale;
-	FVector4f RectDataAndVirtualShadowMapId;
+	FVector4f RectDataAndVirtualShadowMapIdOrPrevLocalLightIndex;
 };
 
 struct FForwardBasePassTextures
diff --git a/Engine/Source/Runtime/Renderer/Private/LightGridInjection.cpp b/Engine/Source/Runtime/Renderer/Private/LightGridInjection.cpp
--- a/Engine/Source/Runtime/Renderer/Private/LightGridInjection.cpp
+++ b/Engine/Source/Runtime/Renderer/Private/LightGridInjection.cpp
@@ -354,7 +354,7 @@ static void PackLocalLightData(
 	// Pack both rect light data (barn door length is initialized to -2 
 	const uint32 RectPackedX = 0;
 	const uint32 RectPackedY = 0;
-	const uint32 RectPackedW = FFloat16(-2.f).Encoded;
+	const uint32 RectPackedZ = FFloat16(-2.f).Encoded;
 
 	// Pack specular scale and IES profile index
 	const float SpecularScale = 1.f;
@@ -364,12 +364,12 @@ static void PackLocalLightData(
 	const FVector3f LightColor = (FVector3f)SimpleLight.Color * FLightRenderParameters::GetLightExposureScale(View.GetLastEyeAdaptationExposure(), SimpleLight.InverseExposureBlend);
 	const FVector2f LightColorPacked = PackLightColor(LightColor);
 
-	Out.LightPositionAndInvRadius				= FVector4f(LightTranslatedWorldPosition, 1.0f / FMath::Max(SimpleLight.Radius, KINDA_SMALL_NUMBER));
-	Out.LightColorAndIdAndFalloffExponent		= FVector4f(LightColorPacked.X, LightColorPacked.Y, INDEX_NONE, SimpleLight.Exponent);
-	Out.LightDirectionAndShadowMapChannelMask	= FVector4f(FVector3f(1, 0, 0), FMath::AsFloat(ShadowMapChannelMask));
-	Out.SpotAnglesAndSourceRadiusPacked			= FVector4f(-2, 1, FMath::AsFloat(PackedZ), FMath::AsFloat(PackedW));
-	Out.LightTangentAndIESDataAndSpecularScale	= FVector4f(1.0f, 0.0f, 0.0f, FMath::AsFloat(SpecularScaleAndIESData));
-	Out.RectDataAndVirtualShadowMapId			= FVector4f(FMath::AsFloat(RectPackedX), FMath::AsFloat(RectPackedY), -1, FMath::AsFloat(RectPackedW));
+	Out.LightPositionAndInvRadius							= FVector4f(LightTranslatedWorldPosition, 1.0f / FMath::Max(SimpleLight.Radius, KINDA_SMALL_NUMBER));
+	Out.LightColorAndIdAndFalloffExponent					= FVector4f(LightColorPacked.X, LightColorPacked.Y, INDEX_NONE, SimpleLight.Exponent);
+	Out.LightDirectionAndShadowMapChannelMask				= FVector4f(FVector3f(1, 0, 0), FMath::AsFloat(ShadowMapChannelMask));
+	Out.SpotAnglesAndSourceRadiusPacked						= FVector4f(-2, 1, FMath::AsFloat(PackedZ), FMath::AsFloat(PackedW));
+	Out.LightTangentAndIESDataAndSpecularScale				= FVector4f(1.0f, 0.0f, 0.0f, FMath::AsFloat(SpecularScaleAndIESData));
+	Out.RectDataAndVirtualShadowMapIdOrPrevLocalLightIndex	= FVector4f(FMath::AsFloat(RectPackedX), FMath::AsFloat(RectPackedY), FMath::AsFloat(RectPackedZ), -1);
 }
 
 static void PackLocalLightData(
@@ -379,6 +379,8 @@ static void PackLocalLightData(
 	const uint32 LightTypeAndShadowMapChannelMaskPacked,
 	const int32 LightSceneId,
 	const int32 VirtualShadowMapId,
+	const int32 PrevLocalLightIndex,
+	const bool bHandledByStochasticShadows,
 	const float VolumetricScatteringIntensity)
 {
 	const FVector3f LightTranslatedWorldPosition(View.ViewMatrices.GetPreViewTranslation() + LightParameters.WorldPosition);
@@ -392,27 +394,33 @@ static void PackLocalLightData(
 	// Pack rect light data
 	uint32 RectPackedX = PackRG16(LightParameters.RectLightAtlasUVOffset.X, LightParameters.RectLightAtlasUVOffset.Y);
 	uint32 RectPackedY = PackRG16(LightParameters.RectLightAtlasUVScale.X, LightParameters.RectLightAtlasUVScale.Y);
-	uint32 RectPackedW = 0;
-	RectPackedW |= FFloat16(LightParameters.RectLightBarnLength).Encoded;									// 16 bits
-	RectPackedW |= uint32(FMath::Clamp(LightParameters.RectLightBarnCosAngle,  0.f, 1.0f) * 0x3FF) << 16;	// 10 bits
-	RectPackedW |= uint32(FMath::Clamp(LightParameters.RectLightAtlasMaxLevel, 0.f, 63.f)) << 26;			//  6 bits
+	uint32 RectPackedZ = 0;
+	RectPackedZ |= FFloat16(LightParameters.RectLightBarnLength).Encoded;									// 16 bits
+	RectPackedZ |= uint32(FMath::Clamp(LightParameters.RectLightBarnCosAngle,  0.f, 1.0f) * 0x3FF) << 16;	// 10 bits
+	RectPackedZ |= uint32(FMath::Clamp(LightParameters.RectLightAtlasMaxLevel, 0.f, 63.f)) << 26;			//  6 bits
 
 	// Pack specular scale and IES profile index
 	const uint32 SpecularScaleAndIESData = PackRG16(LightParameters.SpecularScale, LightParameters.IESAtlasIndex);
 
 	const FVector2f LightColorPacked = PackLightColor(FVector3f(LightParameters.Color));
 
-	// NOTE: This cast of VirtualShadowMapId to float is not ideal, but bitcast has issues here with INDEX_NONE -> NaN
+	// Since lights don't use VSM and Stochastic Shadows simultaneously and
+	// currently PrevLocalLightIndex is only accessed by Stochastic Shadows shaders
+	// we can store only one of the values to avoid increasing the size of the struct
+	// TODO: Improve packing to avoid this so that PrevLocalLightIndex can be accessed in other shaders as well
+	const int32 VirtualShadowMapIdOrPrevLocalLightIndex = bHandledByStochasticShadows ? PrevLocalLightIndex : VirtualShadowMapId;
+
+	// NOTE: This cast of VirtualShadowMapIdOrPrevLocalLightIndex to float is not ideal, but bitcast has issues here with INDEX_NONE -> NaN
 	// and 32-bit floats have enough mantissa to cover all reasonable numbers here for now.
 	// NOTE: SpotAngles needs full-precision for VSM one pass projection
-	Out.LightPositionAndInvRadius				= FVector4f(LightTranslatedWorldPosition, LightParameters.InvRadius);
-	Out.LightColorAndIdAndFalloffExponent		= FVector4f(LightColorPacked.X, LightColorPacked.Y, LightSceneId, LightParameters.FalloffExponent);
-	Out.LightDirectionAndShadowMapChannelMask	= FVector4f(LightParameters.Direction, FMath::AsFloat(LightTypeAndShadowMapChannelMaskPacked));
-	Out.SpotAnglesAndSourceRadiusPacked			= FVector4f(LightParameters.SpotAngles.X, LightParameters.SpotAngles.Y, FMath::AsFloat(PackedZ), FMath::AsFloat(PackedW));
-	Out.LightTangentAndIESDataAndSpecularScale	= FVector4f(LightParameters.Tangent, FMath::AsFloat(SpecularScaleAndIESData));
-	Out.RectDataAndVirtualShadowMapId			= FVector4f(FMath::AsFloat(RectPackedX), FMath::AsFloat(RectPackedY), float(VirtualShadowMapId), FMath::AsFloat(RectPackedW));
-
-	checkSlow(int32(Out.RectDataAndVirtualShadowMapId.Z) == VirtualShadowMapId);
+	Out.LightPositionAndInvRadius							= FVector4f(LightTranslatedWorldPosition, LightParameters.InvRadius);
+	Out.LightColorAndIdAndFalloffExponent					= FVector4f(LightColorPacked.X, LightColorPacked.Y, LightSceneId, LightParameters.FalloffExponent);
+	Out.LightDirectionAndShadowMapChannelMask				= FVector4f(LightParameters.Direction, FMath::AsFloat(LightTypeAndShadowMapChannelMaskPacked));
+	Out.SpotAnglesAndSourceRadiusPacked						= FVector4f(LightParameters.SpotAngles.X, LightParameters.SpotAngles.Y, FMath::AsFloat(PackedZ), FMath::AsFloat(PackedW));
+	Out.LightTangentAndIESDataAndSpecularScale				= FVector4f(LightParameters.Tangent, FMath::AsFloat(SpecularScaleAndIESData));
+	Out.RectDataAndVirtualShadowMapIdOrPrevLocalLightIndex	= FVector4f(FMath::AsFloat(RectPackedX), FMath::AsFloat(RectPackedY), FMath::AsFloat(RectPackedZ), float(VirtualShadowMapIdOrPrevLocalLightIndex));
+
+	checkSlow(int32(Out.RectDataAndVirtualShadowMapIdOrPrevLocalLightIndex.W) == VirtualShadowMapIdOrPrevLocalLightIndex);
 }
 
 FComputeLightGridOutput FSceneRenderer::ComputeLightGrid(FRDGBuilder& GraphBuilder, bool bCullLightsToGrid, FSortedLightSetSceneInfo& SortedLightSet)
@@ -549,6 +557,13 @@ FComputeLightGridOutput FSceneRenderer::ComputeLightGrid(FRDGBuilder& GraphBuild
 						(SortedLightInfo.SortKey.Fields.LightType == LightType_Spot && ViewFamily.EngineShowFlags.SpotLights) ||
 						(SortedLightInfo.SortKey.Fields.LightType == LightType_Rect && ViewFamily.EngineShowFlags.RectLights))
 					{
+						int32 PrevLocalLightIndex = INDEX_NONE;
+						if (View.ViewState)
+						{
+							PrevLocalLightIndex = View.ViewState->LightSceneIdToLocalLightIndex.FindOrAdd(LightSceneInfo->Id, INDEX_NONE);
+							View.ViewState->LightSceneIdToLocalLightIndex[LightSceneInfo->Id] = ForwardLocalLightData.Num();
+						}
+
 						ForwardLocalLightData.AddUninitialized(1);
 						FForwardLocalLightData& LightData = ForwardLocalLightData.Last();
 						LocalLightVisibleLightInfosIndex.Add(LightSceneInfo->Id);
@@ -575,7 +590,16 @@ FComputeLightGridOutput FSceneRenderer::ComputeLightGrid(FRDGBuilder& GraphBuild
 							VolumetricScatteringIntensity = 0;
 						}
 
-						PackLocalLightData(LightData, View, LightParameters, LightTypeAndShadowMapChannelMaskPacked, LightSceneInfo->Id, VirtualShadowMapId, VolumetricScatteringIntensity);
+						PackLocalLightData(
+							LightData,
+							View,
+							LightParameters,
+							LightTypeAndShadowMapChannelMaskPacked,
+							LightSceneInfo->Id,
+							VirtualShadowMapId,
+							PrevLocalLightIndex,
+							SortedLightInfo.SortKey.Fields.bHandledByStochasticShadows,
+							VolumetricScatteringIntensity);
 
 						const FSphere BoundingSphere = LightProxy->GetBoundingSphere();
 						const float Distance = View.ViewMatrices.GetViewMatrix().TransformPosition(BoundingSphere.Center).Z + BoundingSphere.W;
diff --git a/Engine/Source/Runtime/Renderer/Private/ScenePrivate.h b/Engine/Source/Runtime/Renderer/Private/ScenePrivate.h
--- a/Engine/Source/Runtime/Renderer/Private/ScenePrivate.h
+++ b/Engine/Source/Runtime/Renderer/Private/ScenePrivate.h
@@ -1019,6 +1019,10 @@ public:
 	FLumenViewState Lumen;
 	FStochasticShadowsViewState StochasticShadows;
 
+	// Map from Light ID in GPU Scene to index in the View's ForwardLightData array
+	// This is stored in ViewState so we can access previous frame mapping
+	TMap<int32, int32> LightSceneIdToLocalLightIndex;
+
 	// Pre-computed filter in spectral (i.e. FFT) domain along with data to determine if we need to up date it
 	struct {
 		/// @cond DOXYGEN_WARNINGS
diff --git a/Engine/Source/Runtime/Renderer/Private/StochasticShadows/StochasticShadows.cpp b/Engine/Source/Runtime/Renderer/Private/StochasticShadows/StochasticShadows.cpp
--- a/Engine/Source/Runtime/Renderer/Private/StochasticShadows/StochasticShadows.cpp
+++ b/Engine/Source/Runtime/Renderer/Private/StochasticShadows/StochasticShadows.cpp
@@ -123,6 +123,7 @@ static TAutoConsoleVariable<int> CVarStochasticShadowsCandidateLightMask(
 
 namespace StochasticShadows
 {
+	// must match values in StochasticShadows.ush
 	constexpr int32 TileSize = 8;
 	constexpr int32 ShadowMaskTileSize = 8;	// Stored downsampled
 	constexpr int32 MaxLightSceneIdXY = 16; // 16 * 16 = 256
@@ -426,6 +427,7 @@ class FCompositeShadowMaskTracesCS : public FGlobalShader
 		SHADER_PARAMETER_RDG_BUFFER_SRV(StructuredBuffer<uint>, ShadowMaskHashTableHistory)
 		SHADER_PARAMETER_RDG_TEXTURE(Texture2D<uint>, ShadowMaskAtlasHistory)
 		SHADER_PARAMETER_RDG_TEXTURE(Texture2D<uint>, LightSamples)
+		SHADER_PARAMETER_RDG_UNIFORM_BUFFER(FForwardLightData, ForwardLightData)
 	END_SHADER_PARAMETER_STRUCT()
 
 	static int32 GetGroupSize()
@@ -972,6 +974,7 @@ void FDeferredShadingSceneRenderer::RenderStochasticShadows(FRDGBuilder& GraphBu
 		PassParameters->ShadowMaskHashTableHistory = ShadowMaskHashTableHistory ? GraphBuilder.CreateSRV(ShadowMaskHashTableHistory) : nullptr;
 		PassParameters->ShadowMaskAtlasHistory = ShadowMaskAtlasHistory;
 		PassParameters->LightSamples = LightSamples;
+		PassParameters->ForwardLightData = View.ForwardLightingResources.ForwardLightUniformBuffer;
 
 		FCompositeShadowMaskTracesCS::FPermutationDomain PermutationVector;
 		PermutationVector.Set<FCompositeShadowMaskTracesCS::FNumSamplesPerPixel>(NumSamplesPerPixel1d);
