Stochastic Shadows
- Remove !SHADOW_MASK_USING_LOCAL_LIGHT_INDEX and !CANDIDATE_LIGHT_MASK_USING_LOCAL_LIGHT_INDEX code paths

diff --git a/Engine/Shaders/Private/StochasticShadows/StochasticShadows.ush b/Engine/Shaders/Private/StochasticShadows/StochasticShadows.ush
--- a/Engine/Shaders/Private/StochasticShadows/StochasticShadows.ush
+++ b/Engine/Shaders/Private/StochasticShadows/StochasticShadows.ush
@@ -47,8 +47,6 @@
 #define TILE_MODE_COMPLEX_SHADING	1
 #define TILE_MODE_MAX				2
 
-#define SHADOW_MASK_USING_LOCAL_LIGHT_INDEX 1
-
 uint PackTile(uint2 TileCoord)
 {
 	return TileCoord.x | (TileCoord.y << 16);
@@ -330,11 +328,7 @@ void AccumulateHistoryShadowMask(uint LightSceneId, uint PrevLocalLightIndex, ui
 {
 	if (all(ShadowMaskScreenCoord < ShadowMaskViewSize) && Weight > 0.01f)
 	{
-#if SHADOW_MASK_USING_LOCAL_LIGHT_INDEX
 		uint2 HistoryShadowMaskPageTableCoord = GetShadowMaskPageTableCoord(PrevLocalLightIndex, ShadowMaskScreenCoord);
-#else
-		uint2 HistoryShadowMaskPageTableCoord = GetShadowMaskPageTableCoord(LightSceneId, ShadowMaskScreenCoord);
-#endif
 
 #if HASH_TABLE
 		uint HistoryShadowMaskTileIndex = INVALID_SHADOW_MASK_TILE_INDEX;
diff --git a/Engine/Shaders/Private/StochasticShadows/StochasticShadowsSampling.usf b/Engine/Shaders/Private/StochasticShadows/StochasticShadowsSampling.usf
--- a/Engine/Shaders/Private/StochasticShadows/StochasticShadowsSampling.usf
+++ b/Engine/Shaders/Private/StochasticShadows/StochasticShadowsSampling.usf
@@ -110,8 +110,6 @@ Texture2D<float> ShadowMaskSceneDepthHistory;
 	groupshared uint SharedLightMask[LIGHT_SCENE_ID_MASK_SIZE];
 #endif
 
-#define CANDIDATE_LIGHT_MASK_USING_LOCAL_LIGHT_INDEX 1
-
 #if CANDIDATE_LIGHT_MASK
 	groupshared uint SharedCandidateLightMask[LIGHT_SCENE_ID_MASK_SIZE];
 #endif
@@ -166,11 +164,7 @@ void MarkLightSceneId(uint LightSceneId)
 
 void WriteShadowMaskTileToPageTable(uint LightSceneId, uint LocalLightIndex, uint2 ShadowMaskScreenCoord, uint ShadowMaskTileIndex)
 {
-#if SHADOW_MASK_USING_LOCAL_LIGHT_INDEX
 	uint2 ShadowMaskPageTableCoord = GetShadowMaskPageTableCoord(LocalLightIndex, ShadowMaskScreenCoord);
-#else
-	uint2 ShadowMaskPageTableCoord = GetShadowMaskPageTableCoord(LightSceneId, ShadowMaskScreenCoord);
-#endif
 
 	RWShadowMaskTileData[ShadowMaskTileIndex] = PackShadowMaskTile(LightSceneId, LocalLightIndex, ShadowMaskScreenCoord);
 
@@ -277,7 +271,6 @@ void GenerateSamplesCS(
 
 			for (uint GridLightIndex = 0; GridLightIndex < NumLightsInGridCell; ++GridLightIndex)
 			{
-#if CANDIDATE_LIGHT_MASK_USING_LOCAL_LIGHT_INDEX
 				const uint LocalLightIndex = GetCulledLightDataGrid(CulledLightGridData.DataStartIndex + GridLightIndex);
 				const FLocalLightData LocalLightData = GetLocalLightDataNonStereo(LocalLightIndex);
 
@@ -287,16 +280,6 @@ void GenerateSamplesCS(
 					uint BitMask = 1u << (LocalLightIndex % 32);
 					InterlockedOr(SharedCandidateLightMask[DWORDIndex], BitMask);
 				}
-#else
-				const FLocalLightData LocalLightData = GetLocalLightData(CulledLightGridData.DataStartIndex + GridLightIndex, EyeIndex);
-
-				if (LocalLightData.bStochasticShadowsSupported && LocalLightData.LightSceneId < MAX_LIGHT_SCENE_ID)
-				{
-					uint DWORDIndex = LocalLightData.LightSceneId / 32;
-					uint BitMask = 1u << (LocalLightData.LightSceneId % 32);
-					InterlockedOr(SharedCandidateLightMask[DWORDIndex], BitMask);
-				}
-#endif
 			}
 		}
 		#endif
@@ -442,16 +425,10 @@ void GenerateSamplesCS(
 						const uint NextBitMask = 1u << NextBitIndex;
 						CandidateLightMask ^= NextBitMask;
 
-#if CANDIDATE_LIGHT_MASK_USING_LOCAL_LIGHT_INDEX
 						const uint LocalLightIndex = MaskIndex * 32 + NextBitIndex;
 						const FLocalLightData LocalLightData = GetLocalLightDataNonStereo(LocalLightIndex);
 						const uint LightSceneId = LocalLightData.LightSceneId;
 						FDeferredLightData LightData = ConvertToDeferredLight(LocalLightData);
-#else
-						const uint LightSceneId = MaskIndex * 32 + NextBitIndex;
-						const FLightSceneData LightSceneData = GetLightSceneData(LightSceneId);
-						FDeferredLightData LightData = ConvertToDeferredLight(LightSceneData);
-#endif
 					
 						float SampleWeight = GetLocalLightTargetPDF(LightData, TranslatedWorldPosition, GBuffer, ScreenCoord);
 
@@ -468,11 +445,7 @@ void GenerateSamplesCS(
 							{
 								FShadowMaskAccumulator Accumulator = InitShadowMaskAccumulator();
 
-#if SHADOW_MASK_USING_LOCAL_LIGHT_INDEX
 								uint PrevLocalLightIndex = LocalLightData.PrevLocalLightIndex;
-#else
-								uint PrevLocalLightIndex = 0; // unused
-#endif
 
 								AccumulateHistoryShadowMask(LightSceneId, PrevLocalLightIndex, ShadowMaskHistoryScreenCoord00 + uint2(0, 0), ShadowMaskHistoryWeights.x, Accumulator);
 								AccumulateHistoryShadowMask(LightSceneId, PrevLocalLightIndex, ShadowMaskHistoryScreenCoord00 + uint2(1, 0), ShadowMaskHistoryWeights.y, Accumulator);
@@ -508,11 +481,7 @@ void GenerateSamplesCS(
 						if (SampleWeight > SamplingMinWeight)
 						{
 							// Allocate required shadow mask tiles for this frame
-#if SHADOW_MASK_USING_LOCAL_LIGHT_INDEX
 							MarkLightSceneId(LocalLightIndex);
-#else
-							MarkLightSceneId(LightSceneId);
-#endif
 
 							// Apply shadowing estimate from the last frame
 							if (HistoryShadowMask.IsValid())
@@ -548,12 +517,9 @@ void GenerateSamplesCS(
 	#else
 				for (uint GridLightIndex = 0; GridLightIndex < NumLightsInGridCell; ++GridLightIndex)
 				{
-#if SHADOW_MASK_USING_LOCAL_LIGHT_INDEX
 					const uint LocalLightIndex = GetCulledLightDataGrid(CulledLightGridData.DataStartIndex + GridLightIndex);
 					const FLocalLightData LocalLightData = GetLocalLightDataNonStereo(LocalLightIndex);
-#else
-					const FLocalLightData LocalLightData = GetLocalLightData(CulledLightGridData.DataStartIndex + GridLightIndex, EyeIndex);
-#endif
+
 					if (LocalLightData.bStochasticShadowsSupported)
 					{
 						FDeferredLightData LightData = ConvertToDeferredLight(LocalLightData);
@@ -608,11 +574,7 @@ void GenerateSamplesCS(
 						if (SampleWeight > SamplingMinWeight)
 						{
 							// Allocate required shadow mask tiles for this frame
-#if SHADOW_MASK_USING_LOCAL_LIGHT_INDEX
 							MarkLightSceneId(LocalLightIndex);
-#else
-							MarkLightSceneId(LocalLightData.LightSceneId);
-#endif
 
 							// Apply shadowing estimate from the last frame
 							if (HistoryShadowMask.IsValid())
@@ -701,14 +663,9 @@ void GenerateSamplesCS(
 
 			if (StoredValue != INVALID_LIGHT_SCENE_ID)
 			{
-#if SHADOW_MASK_USING_LOCAL_LIGHT_INDEX
 				const uint LocalLightIndex = StoredValue;
 				const FLocalLightData LocalLightData = GetLocalLightDataNonStereo(LocalLightIndex);
 				const uint LightIndex = LocalLightData.LightSceneId;
-#else
-				const uint LocalLightIndex = 0; // unused
-				const uint LightIndex = StoredValue;
-#endif
 
 				// #sdl_todo: merge allocations through groupshared and do only one atomic per group
 				uint ShadowMaskTileIndex = 0;
@@ -736,14 +693,9 @@ void GenerateSamplesCS(
 				const uint NextBitMask = 1u << NextBitIndex;
 				LightMask ^= NextBitMask;
 
-#if SHADOW_MASK_USING_LOCAL_LIGHT_INDEX
 				const uint LocalLightIndex = LinearThreadIndex + NextBitIndex * LIGHT_SCENE_ID_MASK_SIZE;
 				const FLocalLightData LocalLightData = GetLocalLightDataNonStereo(LocalLightIndex);
 				const uint LightIndex = LocalLightData.LightSceneId;
-#else
-				const uint LocalLightIndex = 0; // unused
-				const uint LightIndex = LinearThreadIndex + NextBitIndex * LIGHT_SCENE_ID_MASK_SIZE;
-#endif
 
 				// #sdl_todo: merge allocations through groupshared and do only one atomic per group
 				uint ShadowMaskTileIndex = 0;
diff --git a/Engine/Shaders/Private/StochasticShadows/StochasticShadowsSpatial.usf b/Engine/Shaders/Private/StochasticShadows/StochasticShadowsSpatial.usf
--- a/Engine/Shaders/Private/StochasticShadows/StochasticShadowsSpatial.usf
+++ b/Engine/Shaders/Private/StochasticShadows/StochasticShadowsSpatial.usf
@@ -83,11 +83,8 @@ void ShadowMaskSpatialPassCS(
 		{
 			uint2 SampleShadowMaskScreenCoord = clamp(int2(BaseShadowMaskScreenCoord) + int2(SharedShadowMaskX - SHARED_TILE_BORDER, SharedShadowMaskY - SHARED_TILE_BORDER), 0, ShadowMaskViewSize - 1);
 
-#if SHADOW_MASK_USING_LOCAL_LIGHT_INDEX
 			uint2 SampleShadowMaskPageTableCoord = GetShadowMaskPageTableCoord(LocalLightIndex, SampleShadowMaskScreenCoord);
-#else
-			uint2 SampleShadowMaskPageTableCoord = GetShadowMaskPageTableCoord(LightSceneId, SampleShadowMaskScreenCoord);
-#endif
+
 			uint SampleShadowMaskTileIndex = INVALID_SHADOW_MASK_TILE_INDEX;
 			#if HASH_TABLE
 			{
