Stochastic Shadows - fixed temporal accumulation on shadow edges
* Add dither to shadow mask in order to fix prevision issues during temporal accumulation
* Remove accumulation speedup when sample falls outside of the neighborhood as it was causing lack of antialiasing on edges

diff --git a/Engine/Shaders/Private/StochasticShadows/StochasticShadows.ush b/Engine/Shaders/Private/StochasticShadows/StochasticShadows.ush
--- a/Engine/Shaders/Private/StochasticShadows/StochasticShadows.ush
+++ b/Engine/Shaders/Private/StochasticShadows/StochasticShadows.ush
@@ -153,12 +153,12 @@ FShadowMask InitShadowMask()
 	return ShadowMask;
 }
 
-uint PackShadowMask(FShadowMask ShadowMask)
+// Dither assumed to be random variable in [0;1]
+uint PackShadowMask(FShadowMask ShadowMask, float Dither)
 {
-	// #sdl_todo: dither factor and variance?
 	uint PackedShadowMask = 0;
-	PackedShadowMask |= clamp(uint(ShadowMask.Factor * 0x3FF + 0.5f), 0u, 0x3FFu); // 10 bits
-	PackedShadowMask |= clamp(uint(ShadowMask.FactorSq * 0x3FF + 0.5f), 0, 0x3FFu) << 10; // 10 bits
+	PackedShadowMask |= clamp(uint(ShadowMask.Factor * 0x3FF + Dither), 0u, 0x3FFu); // 10 bits
+	PackedShadowMask |= clamp(uint(ShadowMask.FactorSq * 0x3FF + Dither), 0, 0x3FFu) << 10; // 10 bits
 	PackedShadowMask |= min(ShadowMask.Age, MAX_SHADOW_MASK_AGE) << 20;
 	PackedShadowMask |= min(ShadowMask.HistoryLength, MAX_SHADOW_MASK_HISTORY_LENGTH) << 23;
 	return PackedShadowMask;
diff --git a/Engine/Shaders/Private/StochasticShadows/StochasticShadowsComposite.usf b/Engine/Shaders/Private/StochasticShadows/StochasticShadowsComposite.usf
--- a/Engine/Shaders/Private/StochasticShadows/StochasticShadowsComposite.usf
+++ b/Engine/Shaders/Private/StochasticShadows/StochasticShadowsComposite.usf
@@ -99,7 +99,7 @@ void CompositeShadowMaskTracesCS(
 
 #if DEBUG_MODE
 	int2 DebugScreenCoord = View.CursorPosition.x >= 0 ? View.CursorPosition * View.ViewResolutionFraction : -1;
-	bool bDebug = all(ShadowMaskScreenCoord == DebugScreenCoord / SHADOW_MASK_DOWNSAMPLE_FACTOR) && LightSceneId == 0;
+	bool bDebug = all(ShadowMaskScreenCoord == DebugScreenCoord / SHADOW_MASK_DOWNSAMPLE_FACTOR) && LightSceneId == 1;
 	FShaderPrintContext Context = InitShaderPrintContext(true, float2(0.1, 0.3));
 
 	if (bDebug)
@@ -172,7 +172,7 @@ void CompositeShadowMaskTracesCS(
 		}
 		#endif
 
-		// Interpolate from the neighbourghood if we didn't trace this pixel
+		// Interpolate from the neighborhood if we didn't trace this pixel
 		if (SampleAccumulator.SampleWeightSum > 0 && CenterSampleAccumulator.SampleWeightSum == 0.0f)
 		{
 			ShadowMask.Factor = SampleAccumulator.SampleSum / SampleAccumulator.SampleWeightSum;
@@ -220,19 +220,15 @@ void CompositeShadowMaskTracesCS(
 							Print(Context, ClampMin);
 							Print(Context, ClampMax);
 							Newline(Context);
-							Print(Context, TEXT("HistoryShadowFactor: "));
-							Newline(Context);
+							Print(Context, TEXT("HistoryShadowFactor:        "));
 							Print(Context, HistoryShadowFactor);
 							Newline(Context);
+							Print(Context, TEXT("HistoryShadowFactorClamped: "));
 							Print(Context, clamp(HistoryShadowFactor, ClampMin, ClampMax));
 						}
 						#endif
 
 						float ClampedHistoryShadowFactor = clamp(HistoryShadowFactor, ClampMin, ClampMax);
-						float DistanceToClamped = abs(ClampedHistoryShadowFactor - HistoryShadowFactor);
-
-						// Speedup accumulation if new samples don't fit history
-						Accumulator.MaxHistoryLength *= saturate(1.0f - 2.0f * DistanceToClamped);
 						HistoryShadowFactor = ClampedHistoryShadowFactor;
 					}
 
@@ -267,7 +263,8 @@ void CompositeShadowMaskTracesCS(
 		}
 		#endif
 
+		float Dither = InterleavedGradientNoise(ShadowMaskScreenCoord, StochasticShadowsStateFrameIndex);
 		uint2 ShadowMaskTileAtlasCoord = GetShadowMaskTileAtlasCoord(ShadowMaskTileIndex, LocalShadowMaskTileCoord);
-		RWShadowMaskTileAtlas[ShadowMaskTileAtlasCoord] = PackShadowMask(ShadowMask);
+		RWShadowMaskTileAtlas[ShadowMaskTileAtlasCoord] = PackShadowMask(ShadowMask, Dither);
 	}
 }
\ No newline at end of file
diff --git a/Engine/Shaders/Private/StochasticShadows/StochasticShadowsSpatial.usf b/Engine/Shaders/Private/StochasticShadows/StochasticShadowsSpatial.usf
--- a/Engine/Shaders/Private/StochasticShadows/StochasticShadowsSpatial.usf
+++ b/Engine/Shaders/Private/StochasticShadows/StochasticShadowsSpatial.usf
@@ -94,7 +94,7 @@ void ShadowMaskSpatialPassCS(
 				ShadowMask = UnpackShadowMask(ShadowMaskTileAtlas[SampleShadowMaskTileAtlasCoord]);
 			}
 
-			SharedShadowMask[SharedShadowMaskX][SharedShadowMaskY] = PackShadowMask(ShadowMask);
+			SharedShadowMask[SharedShadowMaskX][SharedShadowMaskY] = PackShadowMask(ShadowMask, /*Dither*/ 0.5f);
 			SharedSceneDepth[SharedShadowMaskX][SharedShadowMaskY] = DownsampledSceneDepth[SampleShadowMaskScreenCoord];
 		}
 	}
@@ -195,7 +195,7 @@ void ShadowMaskSpatialPassCS(
 		GroupMemoryBarrierWithGroupSync();
 
 		// Update shared shadow mask tile with filtered results
-		SharedShadowMask[GroupThreadId.x + SHARED_TILE_BORDER][GroupThreadId.y + SHARED_TILE_BORDER] = PackShadowMask(ShadowMask);
+		SharedShadowMask[GroupThreadId.x + SHARED_TILE_BORDER][GroupThreadId.y + SHARED_TILE_BORDER] = PackShadowMask(ShadowMask, /*Dither*/ 0.5f);
 
 		if (ShadowMask.IsValid() && ShadowMask.Factor > 0.0f)
 		{
