Stochastic Shadows - fixed PrevLocalLightIndex in Composite shader

diff --git a/Engine/Shaders/Private/StochasticShadows/StochasticShadows.ush b/Engine/Shaders/Private/StochasticShadows/StochasticShadows.ush
--- a/Engine/Shaders/Private/StochasticShadows/StochasticShadows.ush
+++ b/Engine/Shaders/Private/StochasticShadows/StochasticShadows.ush
@@ -3,7 +3,7 @@
 #pragma once
 
 // Change this to force recompilation of all Stochastic Shadows shaders
-#pragma message("UESHADERMETADATA_VERSION CA6ACC4F-2C35-2CD1-1CAF-82395A79FB2B")
+#pragma message("UESHADERMETADATA_VERSION 1A63AC4B-2335-2AD1-2CAB-82295C72AFFC")
 
 #define SUPPORT_CONTACT_SHADOWS 0
 #define USE_IES_PROFILE 0
diff --git a/Engine/Shaders/Private/StochasticShadows/StochasticShadowsHardwareRayTracing.usf b/Engine/Shaders/Private/StochasticShadows/StochasticShadowsHardwareRayTracing.usf
--- a/Engine/Shaders/Private/StochasticShadows/StochasticShadowsHardwareRayTracing.usf
+++ b/Engine/Shaders/Private/StochasticShadows/StochasticShadowsHardwareRayTracing.usf
@@ -43,7 +43,7 @@ RAY_TRACING_ENTRY_RAYGEN_OR_INLINE(HardwareRayTraceLightSamples)
 	uint TraceTexelIndex = GroupIndex * 64 + ThreadIndex;
 
 #if DEBUG_MODE
-	FShaderPrintContext Context = InitShaderPrintContext(true, TraceTexelIndex == 0 ? float2(0.55, 0.3) : float2(0.55, 0.35));
+	FShaderPrintContext Context = InitShaderPrintContext(true, TraceTexelIndex == 0 ? float2(0.55, 0.4) : float2(0.55, 0.45));
 	if (TraceTexelIndex == 0)
 	{
 		Print(Context, TEXT("HardwareRayTraceLightSamples"));
diff --git a/Engine/Shaders/Private/StochasticShadows/StochasticShadowsShading.usf b/Engine/Shaders/Private/StochasticShadows/StochasticShadowsShading.usf
--- a/Engine/Shaders/Private/StochasticShadows/StochasticShadowsShading.usf
+++ b/Engine/Shaders/Private/StochasticShadows/StochasticShadowsShading.usf
@@ -77,13 +77,14 @@ void ShadeLightSamplesCS(
 					Print(Context, NumShadingTilesInCell);
 					Newline(Context);
 					Print(Context, TEXT("ShadowMaskTiles: "));
-					Print(Context, MaxShadowMaskTiles);
 					Print(Context, ShadowMaskTileAllocator[0]);
+					Print(Context, MaxShadowMaskTiles);
+					Print(Context, ShadowMaskTileAllocator[0] / float(MaxShadowMaskTiles));
 					Print(Context, float(ShadowMaskTileAllocator[0] * 8 * 8) / (View.ViewSizeAndInvSize.x * View.ViewSizeAndInvSize.y));
 					Newline(Context);
-					Print(Context, TEXT("ShadingTiles: "));
-					Print(Context, MaxShadingTiles);
+					Print(Context, TEXT("ShadingTiles:    "));
 					Print(Context, ShadingTileAllocator[0]);
+					Print(Context, MaxShadingTiles);
 					Print(Context, ShadingTileAllocator[0] / float(MaxShadingTiles));
 					Newline(Context);
 					Print(Context, TEXT("LightId   | Shadow    | Irradiance"));
diff --git a/Engine/Source/Runtime/Renderer/Private/StochasticShadows/StochasticShadows.cpp b/Engine/Source/Runtime/Renderer/Private/StochasticShadows/StochasticShadows.cpp
--- a/Engine/Source/Runtime/Renderer/Private/StochasticShadows/StochasticShadows.cpp
+++ b/Engine/Source/Runtime/Renderer/Private/StochasticShadows/StochasticShadows.cpp
@@ -436,6 +436,7 @@ class FCompositeShadowMaskTracesCS : public FGlobalShader
 	static void ModifyCompilationEnvironment(const FGlobalShaderPermutationParameters& Parameters, FShaderCompilerEnvironment& OutEnvironment)
 	{
 		FGlobalShader::ModifyCompilationEnvironment(Parameters, OutEnvironment);
+		FForwardLightingParameters::ModifyCompilationEnvironment(Parameters.Platform, OutEnvironment);
 		OutEnvironment.SetDefine(TEXT("THREADGROUP_SIZE"), GetGroupSize());
 
 		FPermutationDomain PermutationVector(Parameters.PermutationId);
