* Merge last hair voxel MIPs level computation into a single pass
* Merge page index clearing with 0 density with the last mip computation pass.

This reduces voxelization cost from 0.23ms -> 0.2ms

diff --git a/Engine/Shaders/Private/HairStrands/HairStrandsVoxelDebug.usf b/Engine/Shaders/Private/HairStrands/HairStrandsVoxelDebug.usf
--- a/Engine/Shaders/Private/HairStrands/HairStrandsVoxelDebug.usf
+++ b/Engine/Shaders/Private/HairStrands/HairStrandsVoxelDebug.usf
@@ -262,7 +262,12 @@ void MainCS(uint3 DispatchThreadId : SV_DispatchThreadID)
 		bool  bForceMipLevel = AddCheckbox(WidgetContext, TEXT("Force MIP"), false, FontWhite); Newline(WidgetContext);
 		if (bForceMipLevel)
 		{
-			ForcedMipLevel = 8u * AddSlider(WidgetContext, TEXT("MIP"), 0, FontWhite); Newline(WidgetContext);
+			ForcedMipLevel = 8u * AddSlider(WidgetContext, TEXT("MIP"), 0, FontWhite); 
+			ForcedMipLevel = clamp(ForcedMipLevel, 0, VirtualVoxel.PageResolutionLog2);
+			Print(WidgetContext, ForcedMipLevel, FontWhite, 1, 1);
+			Print(WidgetContext, TEXT("/"), FontWhite);
+			Print(WidgetContext, VirtualVoxel.PageResolutionLog2, FontWhite, 1, 1);
+			Newline(WidgetContext);
 		}
 	}
 
diff --git a/Engine/Shaders/Private/HairStrands/HairStrandsVoxelRasterCompute.usf b/Engine/Shaders/Private/HairStrands/HairStrandsVoxelRasterCompute.usf
--- a/Engine/Shaders/Private/HairStrands/HairStrandsVoxelRasterCompute.usf
+++ b/Engine/Shaders/Private/HairStrands/HairStrandsVoxelRasterCompute.usf
@@ -105,8 +105,7 @@ void AllocatePageIndex(uint2 DispatchThreadId : SV_DispatchThreadID)
 			{
 				PageIndexResolution = ceil(MacroGroupSize / PageWorldSize);
 			}
-			MacroGroupSize = PageIndexResolution * PageWorldSize;
-			Bound.Max = MacroGroupSize + Bound.Min;
+			Bound.Max = (PageIndexResolution * PageWorldSize) + Bound.Min; // Snap Bound's Max to page size
 		}
 
 		const uint TotalPageIndex = PageIndexResolution.x * PageIndexResolution.y * PageIndexResolution.z;
@@ -1101,7 +1100,7 @@ void MainPS(
 ///////////////////////////////////////////////////////////////////////////
 // Common function for mipmapping voxels
 
-#if SHADER_MIP_VIRTUALVOXEL || SHADER_MIP_INDIRECTARGS || SHADER_UPDATE_PAGEINDEX
+#if SHADER_MIP_VIRTUALVOXEL || SHADER_MIP_INDIRECTARGS
 uint ComputeMipDensity(
 	const uint RawDensity0,
 	const uint RawDensity1,
@@ -1151,11 +1150,14 @@ uint ComputeMipDensity(
 	// could arise due to numerical precision).
 	return min(uint(VOXEL_HAIR_MASK), bHasData ? max(TotalHair, 1u) : 0u) | OutTotalOpaque | (bHasShadowCaster ? 0u : VOXEL_CAST_NO_SHADOW_MASK);
 }
-#endif // SHADER_MIP_VIRTUALVOXEL || SHADER_MIP_INDIRECTARGS || SHADER_UPDATE_PAGEINDEX
+#endif // SHADER_MIP_VIRTUALVOXEL || SHADER_MIP_INDIRECTARGS
 
 ///////////////////////////////////////////////////////////////////////////
 
 #if SHADER_MIP_VIRTUALVOXEL
+#include "../MortonCode.ush"
+
+uint bPatchEmptyPage;
 int3 PageCountResolution;
 uint PageResolution;
 uint SourceMip;
@@ -1165,8 +1167,30 @@ Buffer<uint> 	  AllocatedPageCountBuffer;
 Texture3D<uint>   InDensityTexture;
 RWTexture3D<uint> OutDensityTexture;
 
+uint MortonEncode3(uint3 Voxel)
+{
+	return MortonCode3(Voxel.x) | MortonCode3(Voxel.y) << 1 | MortonCode3(Voxel.z) << 2;
+}
+
+uint3 MortonDecode3(uint Morton)
+{
+	uint3 Voxel = uint3(ReverseMortonCode3(Morton), ReverseMortonCode3(Morton >> 1), ReverseMortonCode3(Morton >> 2));
+	return Voxel;
+}
+
+#if PERMUTATION_MIP_AGGREGATE
+Buffer<uint>	  PageToPageIndexBuffer;
+RWBuffer<uint>	  OutPageIndexBuffer;
+
+RWTexture3D<uint> OutDensityTexture2;
+RWTexture3D<uint> OutDensityTexture1;
+
+groupshared uint g_Density4[64][GROUP_SIZE_Z];
+groupshared uint g_Density2[8][GROUP_SIZE_Z];
+#endif
+
 [numthreads(GROUP_SIZE_X, 1, GROUP_SIZE_Z)]
-void MainCS(uint3 DispatchThreadId : SV_DispatchThreadID)
+void MainCS(uint3 DispatchThreadId : SV_DispatchThreadID, uint3 GroupThreadId : SV_GroupThreadID)
 {
 	const uint  PageIndex = DispatchThreadId.z;
 	if (PageIndex >= AllocatedPageCountBuffer[0])
@@ -1176,48 +1200,115 @@ void MainCS(uint3 DispatchThreadId : SV_DispatchThreadID)
 
 	const uint3 PageCoord = IndexToCoord(PageIndex, PageCountResolution);
 
-	const uint  InPageResolution = PageResolution >> SourceMip;
-	const uint OutPageResolution = PageResolution >> TargetMip;
-	const uint TotalVoxelPerOutPageResolution = OutPageResolution * OutPageResolution * OutPageResolution;
-	const uint VoxelIndex = DispatchThreadId.x;
-	if (VoxelIndex >= TotalVoxelPerOutPageResolution)
-		return;
+	uint Total = 0;
+	{
+		const uint  InPageResolution = PageResolution >> SourceMip;
+		const uint OutPageResolution = PageResolution >> TargetMip;
+		const uint TotalVoxelPerOutPageResolution = OutPageResolution * OutPageResolution * OutPageResolution;	
+		if (DispatchThreadId.x < TotalVoxelPerOutPageResolution)
+		{
+			const uint VoxelIndex = DispatchThreadId.x;
+	
+			const uint3 OutVoxelCoordOffset = MortonDecode3(VoxelIndex);
+			const uint3  InVoxelCoordOffset = OutVoxelCoordOffset << 1;
+		
+			const uint3 OutVoxelCoord = PageCoord * OutPageResolution + OutVoxelCoordOffset;
+			const uint3  InVoxelCoord = PageCoord * InPageResolution + InVoxelCoordOffset;
+		
+			const uint3 InVoxelCoord0 = InVoxelCoord;
+			const uint3 InVoxelCoord1 = InVoxelCoord0 + uint3(1, 0, 0);
+			const uint3 InVoxelCoord2 = InVoxelCoord0 + uint3(0, 1, 0);
+			const uint3 InVoxelCoord3 = InVoxelCoord0 + uint3(1, 1, 0);
+			const uint3 InVoxelCoord4 = InVoxelCoord0 + uint3(0, 0, 1);
+			const uint3 InVoxelCoord5 = InVoxelCoord0 + uint3(1, 0, 1);
+			const uint3 InVoxelCoord6 = InVoxelCoord0 + uint3(0, 1, 1);
+			const uint3 InVoxelCoord7 = InVoxelCoord0 + uint3(1, 1, 1);
+		
+			const uint RawDensity0 = InDensityTexture[InVoxelCoord0];
+			const uint RawDensity1 = InDensityTexture[InVoxelCoord1];
+			const uint RawDensity2 = InDensityTexture[InVoxelCoord2];
+			const uint RawDensity3 = InDensityTexture[InVoxelCoord3];
+			const uint RawDensity4 = InDensityTexture[InVoxelCoord4];
+			const uint RawDensity5 = InDensityTexture[InVoxelCoord5];
+			const uint RawDensity6 = InDensityTexture[InVoxelCoord6];
+			const uint RawDensity7 = InDensityTexture[InVoxelCoord7];
+		
+			Total = ComputeMipDensity(
+				RawDensity0,
+				RawDensity1,
+				RawDensity2,
+				RawDensity3,
+				RawDensity4,
+				RawDensity5,
+				RawDensity6,
+				RawDensity7);
+		
+			OutDensityTexture[OutVoxelCoord] = Total;
+	
+			#if PERMUTATION_MIP_AGGREGATE
+			// Store 4x4x4 values
+			if (GroupThreadId.x < 64)
+			{
+				uint StoreIndex = MortonEncode3(OutVoxelCoordOffset);
+				g_Density4[StoreIndex][GroupThreadId.z] = Total;
+			}
+			#endif
+		}
+	}
+	
+	#if PERMUTATION_MIP_AGGREGATE
+	GroupMemoryBarrierWithGroupSync();
+
+	// Target Page Res: 2x2x2 
+	if (GroupThreadId.x < 8)
+	{
+		const uint Total4 = ComputeMipDensity(
+			g_Density4[GroupThreadId.x*8+0][GroupThreadId.z],
+			g_Density4[GroupThreadId.x*8+1][GroupThreadId.z],
+			g_Density4[GroupThreadId.x*8+2][GroupThreadId.z],
+			g_Density4[GroupThreadId.x*8+3][GroupThreadId.z],
+			g_Density4[GroupThreadId.x*8+4][GroupThreadId.z],
+			g_Density4[GroupThreadId.x*8+5][GroupThreadId.z],
+			g_Density4[GroupThreadId.x*8+6][GroupThreadId.z],
+			g_Density4[GroupThreadId.x*8+7][GroupThreadId.z]);
+
+
+		const uint3 StoreVoxel = MortonDecode3(GroupThreadId.x*8)>>1;
+		const uint StoreIndex  = MortonEncode3(StoreVoxel); // GroupThreadId.x >> 3;
+		g_Density2[StoreIndex][GroupThreadId.z] = Total4;
+
+		const uint OutPageResolution = 2;
+		const uint3 OutVoxelCoord = PageCoord * OutPageResolution + StoreVoxel;
+		OutDensityTexture2[OutVoxelCoord] = Total4;
+	}
+	GroupMemoryBarrierWithGroupSync();
+
+	// Target Page Res: 1x1x1 
+	if (GroupThreadId.x < 1)
+	{
+		const uint Total2 = ComputeMipDensity(
+			g_Density2[0][GroupThreadId.z],
+			g_Density2[1][GroupThreadId.z],
+			g_Density2[2][GroupThreadId.z],
+			g_Density2[3][GroupThreadId.z],
+			g_Density2[4][GroupThreadId.z],
+			g_Density2[5][GroupThreadId.z],
+			g_Density2[6][GroupThreadId.z],
+			g_Density2[7][GroupThreadId.z]);
+
+		const uint OutPageResolution = 1;
+		const uint3 OutVoxelCoord = PageCoord * OutPageResolution + 0;
+		OutDensityTexture1[OutVoxelCoord] = Total2;
 
-	const uint3 OutVoxelCoordOffset = IndexToCoord(VoxelIndex, OutPageResolution);
-	const uint3  InVoxelCoordOffset = OutVoxelCoordOffset << 1;
-
-	const uint3 OutVoxelCoord = PageCoord * OutPageResolution + OutVoxelCoordOffset;
-	const uint3  InVoxelCoord = PageCoord * InPageResolution + InVoxelCoordOffset;
-
-	const uint3 InVoxelCoord0 = InVoxelCoord;
-	const uint3 InVoxelCoord1 = InVoxelCoord0 + uint3(1, 0, 0);
-	const uint3 InVoxelCoord2 = InVoxelCoord0 + uint3(0, 1, 0);
-	const uint3 InVoxelCoord3 = InVoxelCoord0 + uint3(1, 1, 0);
-	const uint3 InVoxelCoord4 = InVoxelCoord0 + uint3(0, 0, 1);
-	const uint3 InVoxelCoord5 = InVoxelCoord0 + uint3(1, 0, 1);
-	const uint3 InVoxelCoord6 = InVoxelCoord0 + uint3(0, 1, 1);
-	const uint3 InVoxelCoord7 = InVoxelCoord0 + uint3(1, 1, 1);
-
-	const uint RawDensity0 = InDensityTexture[InVoxelCoord0];
-	const uint RawDensity1 = InDensityTexture[InVoxelCoord1];
-	const uint RawDensity2 = InDensityTexture[InVoxelCoord2];
-	const uint RawDensity3 = InDensityTexture[InVoxelCoord3];
-	const uint RawDensity4 = InDensityTexture[InVoxelCoord4];
-	const uint RawDensity5 = InDensityTexture[InVoxelCoord5];
-	const uint RawDensity6 = InDensityTexture[InVoxelCoord6];
-	const uint RawDensity7 = InDensityTexture[InVoxelCoord7];
-
-	const uint Total = ComputeMipDensity(
-		RawDensity0,
-		RawDensity1,
-		RawDensity2,
-		RawDensity3,
-		RawDensity4,
-		RawDensity5,
-		RawDensity6,
-		RawDensity7);
-
-	OutDensityTexture[OutVoxelCoord] = Total;
+		// Update the page index with invalid page index if the voxel does not contain any data.
+		// This allow to save tracing cost when evaluating the transmittance.
+		if (bPatchEmptyPage > 0 && Total2 == 0)
+		{
+			const uint PageIndexOffset = PageToPageIndexBuffer[PageIndex];
+			OutPageIndexBuffer[PageIndexOffset] = INVALID_VOXEL_PAGE_INDEX;
+		}
+	}
+	#endif
 }
 #endif
 
@@ -1249,73 +1340,3 @@ void MainCS(uint3 DispatchThreadId : SV_DispatchThreadID)
 
 ///////////////////////////////////////////////////////////////////////////
 
-#if SHADER_UPDATE_PAGEINDEX
-int3 PageCountResolution;
-uint PageResolution;
-uint MipIt;
-
-Buffer<uint>	  PageIndexGlobalCounter; // Allocated page counter
-Texture3D<uint>	  DensityTexture;
-Buffer<uint>	  PageToPageIndexBuffer;
-RWBuffer<uint>	  OutPageIndexBuffer;
-RWBuffer<uint2>	  OutPageIndexOccupancyBuffer;
-
-[numthreads(GROUP_SIZE_X, 1, GROUP_SIZE_Z)]
-void MainCS(uint3 DispatchThreadId : SV_DispatchThreadID)
-{
-	const uint  PageIndex = DispatchThreadId.z;
-	const uint3 PageCoord = IndexToCoord(PageIndex, PageCountResolution);
-
-	const uint AllocatedPageCount = PageIndexGlobalCounter[0];
-	if (PageIndex >= AllocatedPageCount)
-		return;
-
-	const uint OutPageResolution = PageResolution >> MipIt;
-	const uint TotalVoxelPerOutPageResolution = OutPageResolution * OutPageResolution * OutPageResolution;
-	const uint VoxelIndex = DispatchThreadId.x;
-	if (VoxelIndex >= TotalVoxelPerOutPageResolution)
-		return;
-
-	{
-		const uint3 OutVoxelCoordOffset = IndexToCoord(VoxelIndex, OutPageResolution);
-		const uint3 OutVoxelCoord = PageCoord * OutPageResolution + OutVoxelCoordOffset;
-
-		const uint Density = DensityTexture.Load(uint4(OutVoxelCoord, MipIt));
-		// Update the page index with invalid page index if the voxel does not contain any data.
-		// This allow to save tracing cost when evaluating the transmittance.
-		if (Density == 0)
-		{
-			const uint PageIndexOffset = PageToPageIndexBuffer[PageIndex];
-			OutPageIndexBuffer[PageIndexOffset] = INVALID_VOXEL_PAGE_INDEX;
-		}
-	}
-
-	// Fill in the occupancy table
-	{
-		const uint OccMipIt = 3; // A voxel at this mip has a size of 8x8x8
-		const uint OccPageResolution = PageResolution >> OccMipIt;
-		uint2 Occupancy = 0;
-		uint BitIndex = 0;
-		for (uint z = 0; z < 4; ++z)
-			for (uint y = 0; y < 4; ++y)
-				for (uint x = 0; x < 4; ++x)
-				{
-					const uint3 OutVoxelCoord = PageCoord * OccPageResolution + uint3(x, y, z);
-					const uint Density = DensityTexture.Load(uint4(OutVoxelCoord, OccMipIt));
-					if (Density > 0)
-					{
-						if (BitIndex < 32)
-							Occupancy.x = Occupancy.x | (1u << BitIndex);
-						else
-							Occupancy.y = Occupancy.y | (1u << (BitIndex - 32));
-					}
-					++BitIndex;
-				}
-		const uint PageIndexOffset = PageToPageIndexBuffer[PageIndex];
-		OutPageIndexOccupancyBuffer[PageIndexOffset] = Occupancy;
-	}
-}
-#endif
-
-///////////////////////////////////////////////////////////////////////////
-
diff --git a/Engine/Source/Runtime/Renderer/Private/HairStrands/HairStrandsVoxelization.cpp b/Engine/Source/Runtime/Renderer/Private/HairStrands/HairStrandsVoxelization.cpp
--- a/Engine/Source/Runtime/Renderer/Private/HairStrands/HairStrandsVoxelization.cpp
+++ b/Engine/Source/Runtime/Renderer/Private/HairStrands/HairStrandsVoxelization.cpp
@@ -1199,6 +1199,9 @@ class FVirtualVoxelGenerateMipCS : public FGlobalShader
 {
 	DECLARE_GLOBAL_SHADER(FVirtualVoxelGenerateMipCS);
 	SHADER_USE_PARAMETER_STRUCT(FVirtualVoxelGenerateMipCS, FGlobalShader);
+	
+	class FAggregate : SHADER_PERMUTATION_BOOL("PERMUTATION_MIP_AGGREGATE");
+	using FPermutationDomain = TShaderPermutationDomain<FAggregate>;	
 
 	BEGIN_SHADER_PARAMETER_STRUCT(FParameters, )
 		SHADER_PARAMETER_STRUCT_INCLUDE(FSceneTextureParameters, SceneTextures)
@@ -1207,10 +1210,16 @@ class FVirtualVoxelGenerateMipCS : public FGlobalShader
 		SHADER_PARAMETER(uint32, PageResolution)
 		SHADER_PARAMETER(uint32, SourceMip)
 		SHADER_PARAMETER(uint32, TargetMip)
+		SHADER_PARAMETER(uint32, bPatchEmptyPage)
+
+		SHADER_PARAMETER_RDG_BUFFER_SRV(Buffer<uint>, PageToPageIndexBuffer)
+		SHADER_PARAMETER_RDG_BUFFER_UAV(RWBuffer<uint>, OutPageIndexBuffer)
 
 		RDG_BUFFER_ACCESS(IndirectDispatchArgs, ERHIAccess::IndirectArgs)
 		SHADER_PARAMETER_RDG_TEXTURE_SRV(Texture3D, InDensityTexture)
 		SHADER_PARAMETER_RDG_TEXTURE_UAV(RWTexture3D, OutDensityTexture)
+		SHADER_PARAMETER_RDG_TEXTURE_UAV(RWTexture3D, OutDensityTexture2)
+		SHADER_PARAMETER_RDG_TEXTURE_UAV(RWTexture3D, OutDensityTexture1)
 
 	END_SHADER_PARAMETER_STRUCT()
 
@@ -1249,43 +1258,8 @@ public:
 	}
 };
 
-class FVirtualVoxelPatchPageIndexWithMipDataCS : public FGlobalShader
-{
-	DECLARE_GLOBAL_SHADER(FVirtualVoxelPatchPageIndexWithMipDataCS);
-	SHADER_USE_PARAMETER_STRUCT(FVirtualVoxelPatchPageIndexWithMipDataCS, FGlobalShader);
-
-	BEGIN_SHADER_PARAMETER_STRUCT(FParameters, )
-		SHADER_PARAMETER(FIntVector, PageCountResolution)
-		SHADER_PARAMETER(uint32, PageResolution)
-		SHADER_PARAMETER(uint32, bUpdatePageIndex)
-		SHADER_PARAMETER(uint32, MipIt)
-
-		SHADER_PARAMETER_RDG_TEXTURE(Texture3D, DensityTexture)
-		RDG_BUFFER_ACCESS(IndirectDispatchArgs, ERHIAccess::IndirectArgs)
-		SHADER_PARAMETER_RDG_BUFFER_SRV(Buffer, PageIndexGlobalCounter)
-		SHADER_PARAMETER_RDG_BUFFER_SRV(Buffer<uint>, PageToPageIndexBuffer)
-		SHADER_PARAMETER_RDG_BUFFER_UAV(RWBuffer<uint>, OutPageIndexBuffer)
-		SHADER_PARAMETER_RDG_BUFFER_UAV(RWBuffer<uint2>, OutPageIndexOccupancyBuffer)
-
-	END_SHADER_PARAMETER_STRUCT()
-
-public:
-	static FIntVector GetGroupSize() { return GHairVoxel_GroupSize_NumVoxelPerPage_1_NumAllocatedPage; }
-	static bool ShouldCompilePermutation(const FGlobalShaderPermutationParameters& Parameters) { return IsHairStrandsSupported(EHairStrandsShaderType::Strands, Parameters.Platform); }
-	static void ModifyCompilationEnvironment(const FGlobalShaderPermutationParameters& Parameters, FShaderCompilerEnvironment& OutEnvironment)
-	{
-		FGlobalShader::ModifyCompilationEnvironment(Parameters, OutEnvironment);
-		OutEnvironment.SetDefine(TEXT("SHADER_UPDATE_PAGEINDEX"), 1);
-		OutEnvironment.SetDefine(TEXT("GROUP_SIZE_X"), GetGroupSize().X);
-		OutEnvironment.SetDefine(TEXT("GROUP_SIZE_Y"), GetGroupSize().Y);
-		OutEnvironment.SetDefine(TEXT("GROUP_SIZE_Z"), GetGroupSize().Z);
-	}
-};
-
-
 IMPLEMENT_GLOBAL_SHADER(FVirtualVoxelGenerateMipCS, "/Engine/Private/HairStrands/HairStrandsVoxelRasterCompute.usf", "MainCS", SF_Compute);
 IMPLEMENT_GLOBAL_SHADER(FVirtualVoxelIndirectArgMipCS, "/Engine/Private/HairStrands/HairStrandsVoxelRasterCompute.usf", "MainCS", SF_Compute);
-IMPLEMENT_GLOBAL_SHADER(FVirtualVoxelPatchPageIndexWithMipDataCS, "/Engine/Private/HairStrands/HairStrandsVoxelRasterCompute.usf", "MainCS", SF_Compute);
 
 
 static void AddVirtualVoxelGenerateMipPass(
@@ -1323,8 +1297,15 @@ static void AddVirtualVoxelGenerateMipPass(
 		FComputeShaderUtils::AddPass(GraphBuilder, RDG_EVENT_NAME("HairStrands::BuildVoxelMipIndirectArgs"), ComputeShader, Parameters, FIntVector(1, 1, 1));
 	}
 
+	// Patch the page index buffer with page whose voxels are empty after the voxelization is done
+	const bool bAMDPC = IsPCPlatform(View.GetShaderPlatform()) && IsRHIDeviceAMD();
+	const bool bPatchEmptyPage = GHairVirtualVoxelInvalidEmptyPageIndex > 0 && !bAMDPC;
+	FRDGBufferSRVRef PageToPageIndexBufferSRV = GraphBuilder.CreateSRV(InPageToPageIndexBuffer, PF_R32_UINT);
+	FRDGBufferUAVRef PageIndexBufferUAV = GraphBuilder.CreateUAV(VoxelResources.PageIndexBuffer, PF_R32_UINT);
+
 	// Generate MIP level (in one go for all allocated pages)
-	for (uint32 MipIt = 0; MipIt < MipCount - 1; ++MipIt)
+	const uint32 LastMipMinus3 = MipCount - 3;
+	for (uint32 MipIt = 0; MipIt < LastMipMinus3; ++MipIt)
 	{
 		const uint32 SourceMipIndex = MipIt;
 		const uint32 TargetMipIndex = MipIt + 1;
@@ -1338,45 +1319,25 @@ static void AddVirtualVoxelGenerateMipPass(
 		Parameters->TargetMip = TargetMipIndex;
 		Parameters->AllocatedPageCountBuffer = VoxelResources.Parameters.Common.AllocatedPageCountBuffer;
 		Parameters->IndirectDispatchArgs = MipIndirectArgsBuffers[MipIt];
+		Parameters->bPatchEmptyPage = bPatchEmptyPage;
 
-		TShaderMapRef<FVirtualVoxelGenerateMipCS> ComputeShader(View.ShaderMap);
-		ClearUnusedGraphResources(ComputeShader, Parameters);
-		GraphBuilder.AddPass(
-			RDG_EVENT_NAME("HairStrands::ComputeVoxelMip"),
-			Parameters,
-			ERDGPassFlags::Compute,
-			[Parameters, ComputeShader](FRHICommandList& RHICmdList)
+		const uint32 TargetPageResolution = Parameters->PageResolution >> TargetMipIndex;
+		const bool bAggregate = TargetPageResolution <= 4;
+		if (bAggregate)
 		{
-			FComputeShaderUtils::DispatchIndirect(RHICmdList, ComputeShader, *Parameters, Parameters->IndirectDispatchArgs->GetIndirectRHICallBuffer(), 0);
-		});
-	}
+			Parameters->PageToPageIndexBuffer = PageToPageIndexBufferSRV;
+			Parameters->OutPageIndexBuffer = PageIndexBufferUAV;
 
-	// Patch the page index buffer with page whose voxels are empty after the voxelization is done
-	FRDGBufferSRVRef PageToPageIndexBufferSRV = GraphBuilder.CreateSRV(InPageToPageIndexBuffer, PF_R32_UINT);
-	FRDGBufferUAVRef PageIndexBufferUAV = GraphBuilder.CreateUAV(VoxelResources.PageIndexBuffer, PF_R32_UINT);
-	FRDGBufferUAVRef PageIndexOccupancyBufferUAV = GraphBuilder.CreateUAV(VoxelResources.PageIndexOccupancyBuffer, PF_R32G32_UINT);
-	
-	// Note: Do not clear empty page on AMD hardware as there are some issue precision or dispatch issue (to be refined)
-	const bool bAMDPC = IsPCPlatform(View.GetShaderPlatform()) && IsRHIDeviceAMD();
-	const bool bPatchEmptyPage = GHairVirtualVoxelInvalidEmptyPageIndex > 0 && !bAMDPC;
-	if (bPatchEmptyPage)
-	{
-		const uint32 LastMipIt = MipCount - 1;
-		FVirtualVoxelPatchPageIndexWithMipDataCS::FParameters* Parameters = GraphBuilder.AllocParameters<FVirtualVoxelPatchPageIndexWithMipDataCS::FParameters>();
-		Parameters->MipIt = LastMipIt;
-		Parameters->PageIndexGlobalCounter = GraphBuilder.CreateSRV(VoxelResources.PageIndexGlobalCounter, PF_R32_UINT);
-		Parameters->PageResolution = VoxelResources.Parameters.Common.PageResolution;
-		Parameters->PageCountResolution = VoxelResources.Parameters.Common.PageCountResolution;
-		Parameters->DensityTexture = VoxelResources.PageTexture;
-		Parameters->PageToPageIndexBuffer = PageToPageIndexBufferSRV;
-		Parameters->OutPageIndexBuffer = PageIndexBufferUAV;
-		Parameters->OutPageIndexOccupancyBuffer = PageIndexOccupancyBufferUAV;
-		Parameters->IndirectDispatchArgs = MipIndirectArgsBuffers[LastMipIt-1];
+			Parameters->OutDensityTexture2 = GraphBuilder.CreateUAV(FRDGTextureUAVDesc(VoxelResources.PageTexture, MipCount-2));
+			Parameters->OutDensityTexture1 = GraphBuilder.CreateUAV(FRDGTextureUAVDesc(VoxelResources.PageTexture, MipCount-1));
+		}
 
-		TShaderMapRef<FVirtualVoxelPatchPageIndexWithMipDataCS> ComputeShader(View.ShaderMap);
+		FVirtualVoxelGenerateMipCS::FPermutationDomain PermutationVector;
+		PermutationVector.Set<FVirtualVoxelGenerateMipCS::FAggregate>(bAggregate);
+		TShaderMapRef<FVirtualVoxelGenerateMipCS> ComputeShader(View.ShaderMap, PermutationVector);
 		ClearUnusedGraphResources(ComputeShader, Parameters);
 		GraphBuilder.AddPass(
-			RDG_EVENT_NAME("HairStrands::PatchPageIndexWithMip"),
+			RDG_EVENT_NAME("HairStrands::ComputeVoxelMip"),
 			Parameters,
 			ERDGPassFlags::Compute,
 			[Parameters, ComputeShader](FRHICommandList& RHICmdList)
