Fix shader compilation issue on SM5.

diff --git a/Engine/Shaders/Private/HairStrands/HairStrandsVoxelRasterCompute.usf b/Engine/Shaders/Private/HairStrands/HairStrandsVoxelRasterCompute.usf
--- a/Engine/Shaders/Private/HairStrands/HairStrandsVoxelRasterCompute.usf
+++ b/Engine/Shaders/Private/HairStrands/HairStrandsVoxelRasterCompute.usf
@@ -1193,14 +1193,11 @@ groupshared uint g_Density2[8][GROUP_SIZE_Z];
 void MainCS(uint3 DispatchThreadId : SV_DispatchThreadID, uint3 GroupThreadId : SV_GroupThreadID)
 {
 	const uint  PageIndex = DispatchThreadId.z;
-	if (PageIndex >= AllocatedPageCountBuffer[0])
-	{
-		return;
-	}
-
 	const uint3 PageCoord = IndexToCoord(PageIndex, PageCountResolution);
-
+	const bool bValid = PageIndex < AllocatedPageCountBuffer[0];
+	
 	uint Total = 0;
+	if (bValid)
 	{
 		const uint  InPageResolution = PageResolution >> SourceMip;
 		const uint OutPageResolution = PageResolution >> TargetMip;
@@ -1260,7 +1257,7 @@ void MainCS(uint3 DispatchThreadId : SV_DispatchThreadID, uint3 GroupThreadId :
 	GroupMemoryBarrierWithGroupSync();
 
 	// Target Page Res: 2x2x2 
-	if (GroupThreadId.x < 8)
+	if (bValid && GroupThreadId.x < 8)
 	{
 		const uint Total4 = ComputeMipDensity(
 			g_Density4[GroupThreadId.x*8+0][GroupThreadId.z],
@@ -1284,7 +1281,7 @@ void MainCS(uint3 DispatchThreadId : SV_DispatchThreadID, uint3 GroupThreadId :
 	GroupMemoryBarrierWithGroupSync();
 
 	// Target Page Res: 1x1x1 
-	if (GroupThreadId.x < 1)
+	if (bValid && GroupThreadId.x < 1)
 	{
 		const uint Total2 = ComputeMipDensity(
 			g_Density2[0][GroupThreadId.z],
