Fix shader warning on certain platform.

diff --git a/Engine/Shaders/Private/HairStrands/HairStrandsVoxelDebug.usf b/Engine/Shaders/Private/HairStrands/HairStrandsVoxelDebug.usf
--- a/Engine/Shaders/Private/HairStrands/HairStrandsVoxelDebug.usf
+++ b/Engine/Shaders/Private/HairStrands/HairStrandsVoxelDebug.usf
@@ -263,7 +263,7 @@ void MainCS(uint3 DispatchThreadId : SV_DispatchThreadID)
 		if (bForceMipLevel)
 		{
 			ForcedMipLevel = 8u * AddSlider(WidgetContext, TEXT("MIP"), 0, FontWhite); 
-			ForcedMipLevel = clamp(ForcedMipLevel, 0, VirtualVoxel.PageResolutionLog2);
+			ForcedMipLevel = clamp(ForcedMipLevel, 0, int(VirtualVoxel.PageResolutionLog2));
 			Print(WidgetContext, ForcedMipLevel, FontWhite, 1, 1);
 			Print(WidgetContext, TEXT("/"), FontWhite);
 			Print(WidgetContext, VirtualVoxel.PageResolutionLog2, FontWhite, 1, 1);
