Stochastic Shadows - added proper constant and normal bias for hardware ray tracing. Also fixed r.StochasticShadows.Debug 2 ray visualization and marked where screen space traces end.

diff --git a/Engine/Shaders/Private/StochasticShadows/StochasticShadowsComposite.usf b/Engine/Shaders/Private/StochasticShadows/StochasticShadowsComposite.usf
--- a/Engine/Shaders/Private/StochasticShadows/StochasticShadowsComposite.usf
+++ b/Engine/Shaders/Private/StochasticShadows/StochasticShadowsComposite.usf
@@ -100,7 +100,7 @@ void CompositeShadowMaskTracesCS(
 #if DEBUG_MODE
 	int2 DebugScreenCoord = View.CursorPosition.x >= 0 ? View.CursorPosition * View.ViewResolutionFraction : -1;
 	bool bDebug = all(ShadowMaskScreenCoord == DebugScreenCoord / SHADOW_MASK_DOWNSAMPLE_FACTOR) && LightSceneId == 1;
-	FShaderPrintContext Context = InitShaderPrintContext(true, float2(0.1, 0.3));
+	FShaderPrintContext Context = InitShaderPrintContext(true, float2(0.05, 0.3));
 
 	if (bDebug)
 	{
diff --git a/Engine/Shaders/Private/StochasticShadows/StochasticShadowsHardwareRayTracing.usf b/Engine/Shaders/Private/StochasticShadows/StochasticShadowsHardwareRayTracing.usf
--- a/Engine/Shaders/Private/StochasticShadows/StochasticShadowsHardwareRayTracing.usf
+++ b/Engine/Shaders/Private/StochasticShadows/StochasticShadowsHardwareRayTracing.usf
@@ -13,6 +13,8 @@
 	#include "../ShaderPrint.ush"
 #endif
 
+float RayTracingBias;
+float RayTracingNormalBias;
 uint StochasticShadowsStateFrameIndes;
 uint2 SampleViewSize;
 uint DebugMode;
@@ -35,6 +37,7 @@ Texture2D<float> LightSampleRayDistance;
 StructuredBuffer<uint> CompactedTraceTexelData;
 StructuredBuffer<uint> CompactedTraceTexelAllocator;
 Texture2D<float> DownsampledSceneDepth;
+Texture2D<UNORM float3> DownsampledSceneWorldNormal;
 
 RAY_TRACING_ENTRY_RAYGEN_OR_INLINE(HardwareRayTraceLightSamples)
 {
@@ -42,43 +45,35 @@ RAY_TRACING_ENTRY_RAYGEN_OR_INLINE(HardwareRayTraceLightSamples)
 	uint GroupIndex = DispatchThreadIndex.y;
 	uint TraceTexelIndex = GroupIndex * 64 + ThreadIndex;
 
-#if DEBUG_MODE
-	FShaderPrintContext Context = InitShaderPrintContext(true, TraceTexelIndex == 0 ? float2(0.55, 0.4) : float2(0.55, 0.45));
-	if (TraceTexelIndex == 0)
-	{
-		Print(Context, TEXT("HardwareRayTraceLightSamples"));
-		Newline(Context);
-		Print(Context, TEXT("CompactedTraceTexelAllocator: "));
-		Print(Context, CompactedTraceTexelAllocator[0]);
-	}
-#endif
-
 	if (TraceTexelIndex < CompactedTraceTexelAllocator[0])
 	{
 		uint2 SampleCoord = UnpackTraceTexel(CompactedTraceTexelData[TraceTexelIndex]);
 		uint2 ScreenCoord = SampleCoordToScreenCoord(SampleCoord);
 		uint2 DownsampledScreenCoord = SampleCoordToDownsampledScreenCoord(SampleCoord);
 
-		#if DEBUG_MODE
-			int2 DebugScreenCoord = View.CursorPosition.x >= 0 ? View.CursorPosition * View.ViewResolutionFraction : -1;
-			bool bDebug = all(DebugScreenCoord / DOWNSAMPLE_FACTOR == DownsampledScreenCoord);
-		#endif
+#if DEBUG_MODE
+		FShaderPrintContext Context = InitShaderPrintContext(true, float2(0.55, 0.45));
+		int2 DebugScreenCoord = View.CursorPosition.x >= 0 ? View.CursorPosition * View.ViewResolutionFraction : -1;
+		bool bDebug = all(DebugScreenCoord / DOWNSAMPLE_FACTOR == DownsampledScreenCoord);
+#endif
 
 		float2 ScreenUV = (ScreenCoord + 0.5f) * View.BufferSizeAndInvSize.zw;
 		float SceneDepth = DownsampledSceneDepth[DownsampledScreenCoord];
+		float3 SceneWorldNormal = normalize(DecodeNormal(DownsampledSceneWorldNormal[DownsampledScreenCoord]));
 		float3 TranslatedWorldPosition = GetTranslatedWorldPositionFromScreenUV(ScreenUV, SceneDepth);
 
 		FLightSample LightSample = UnpackLightSample(RWLightSamples[SampleCoord]);
 		FLightSampleTrace LightSampleTrace = GetLightSampleTrace(TranslatedWorldPosition, LightSample.LightSceneId, SampleCoord);
 		const float LightSampleRayMinT = LightSampleRayDistance[SampleCoord];
 
-		// #sdl_todo: proper bias
 		FRayDesc Ray = (FRayDesc)0;
 		Ray.Origin = TranslatedWorldPosition;
 		Ray.Direction = LightSampleTrace.Direction;
-		Ray.TMin = max(LightSampleRayMinT, 1.0f);
+		Ray.TMin = max(LightSampleRayMinT, RayTracingBias);
 		Ray.TMax = LightSampleTrace.Distance;
 
+		ApplyPositionBias(Ray.Origin, Ray.Direction, SceneWorldNormal, RayTracingNormalBias);
+
 		uint RayFlags = RAY_FLAG_ACCEPT_FIRST_HIT_AND_END_SEARCH;
 		uint RayMask = RAY_TRACING_MASK_OPAQUE;
 
@@ -138,7 +133,12 @@ RAY_TRACING_ENTRY_RAYGEN_OR_INLINE(HardwareRayTraceLightSamples)
 		if (bDebug && DebugMode == DEBUG_MODE_VISUALIZE_TRACING)
 		{
 			FLightSceneData LightSceneData = GetLightSceneData(LightSample.LightSceneId);
-			AddLineTWS(Context, TranslatedWorldPosition + LightSampleTrace.Direction * Ray.TMin, TranslatedWorldPosition + LightSampleTrace.Direction * (Payload.IsHit() ? Payload.HitT : Ray.TMax), float4(LightSceneData.Color.xyz / Luminance(LightSceneData.Color.xyz), 1.0f));
+
+			float3 RayStart = Ray.Origin + Ray.Direction * Ray.TMin;
+			float3 RayEnd = Ray.Origin + Ray.Direction * (Payload.IsHit() ? Payload.HitT : Ray.TMax);
+			float4 RayColor = float4(LightSceneData.Color.xyz / Luminance(LightSceneData.Color.xyz), 1.0f);
+
+			AddLineTWS(Context, RayStart, RayEnd, RayColor);
 		}
 		#endif
 	}
diff --git a/Engine/Shaders/Private/StochasticShadows/StochasticShadowsTracing.usf b/Engine/Shaders/Private/StochasticShadows/StochasticShadowsTracing.usf
--- a/Engine/Shaders/Private/StochasticShadows/StochasticShadowsTracing.usf
+++ b/Engine/Shaders/Private/StochasticShadows/StochasticShadowsTracing.usf
@@ -286,7 +286,13 @@ void ScreenSpaceRayTraceLightSamplesCS(
 		if (bDebug && DebugMode == DEBUG_MODE_VISUALIZE_TRACING)
 		{
 			FLightSceneData LightSceneData = GetLightSceneData(LightSample.LightSceneId);
-			AddLineTWS(Context, TranslatedWorldPosition, TranslatedWorldPosition + LightSampleTrace.Direction * min(LightSampleTrace.Distance, RayT), float4(LightSceneData.Color.xyz / Luminance(LightSceneData.Color.xyz), 1.0f));
+
+			float3 RayStart = TranslatedWorldPosition;
+			float3 RayEnd = TranslatedWorldPosition + LightSampleTrace.Direction * min(LightSampleTrace.Distance, RayT);
+			float4 RayColor = float4(LightSceneData.Color.xyz / Luminance(LightSceneData.Color.xyz), 1.0f);
+
+			AddLineTWS(Context, RayStart, RayEnd, RayColor, 1.0f);
+			AddCrossTWS(Context, RayEnd, 2.0f, RayColor);
 		}
 		#endif
 	}
diff --git a/Engine/Source/Runtime/Renderer/Private/StochasticShadows/StochasticShadowsRayTracing.cpp b/Engine/Source/Runtime/Renderer/Private/StochasticShadows/StochasticShadowsRayTracing.cpp
--- a/Engine/Source/Runtime/Renderer/Private/StochasticShadows/StochasticShadowsRayTracing.cpp
+++ b/Engine/Source/Runtime/Renderer/Private/StochasticShadows/StochasticShadowsRayTracing.cpp
@@ -54,6 +54,20 @@ static TAutoConsoleVariable<int32> CVarStochasticShadowsHardwareRayTracingInline
 	ECVF_RenderThreadSafe | ECVF_Scalability
 );
 
+static TAutoConsoleVariable<int32> CVarStochasticShadowsHardwareRayTracingBias(
+	TEXT("r.StochasticShadows.HardwareRayTracing.Bias"),
+	1.0f,
+	TEXT("Constant bias for hardware ray traced shadow rays."),
+	ECVF_Scalability | ECVF_RenderThreadSafe
+);
+
+static TAutoConsoleVariable<int32> CVarStochasticShadowsHardwareRayTracingNormalBias(
+	TEXT("r.StochasticShadows.HardwareRayTracing.NormalBias"),
+	0.1f,
+	TEXT("Normal bias for hardware ray traced shadow rays."),
+	ECVF_Scalability | ECVF_RenderThreadSafe
+);
+
 static TAutoConsoleVariable<int32> CVarStochasticShadowsHardwareRayTracingMaxIterations(
 	TEXT("r.StochasticShadows.HardwareRayTracing.MaxIterations"),
 	8192,
@@ -205,6 +219,8 @@ class FHardwareRayTraceLightSamples : public FLumenHardwareRayTracingShaderBase
 		SHADER_PARAMETER_STRUCT_INCLUDE(FStochasticShadowsParameters, StochasticShadowsParameters)
 		SHADER_PARAMETER_RDG_TEXTURE_UAV(RWTexture2D<uint>, RWLightSamples)
 		SHADER_PARAMETER_RDG_TEXTURE(Texture2D<float>, LightSampleRayDistance)
+		SHADER_PARAMETER(float, RayTracingBias)
+		SHADER_PARAMETER(float, RayTracingNormalBias)
 		// Ray Tracing
 		SHADER_PARAMETER(uint32, MaxTraversalIterations)
 		SHADER_PARAMETER_RDG_BUFFER_SRV(RaytracingAccelerationStructure, TLAS)
@@ -375,6 +391,8 @@ void StochasticShadows::SetHardwareRayTracingPassParameters(
 	PassParameters->StochasticShadowsParameters = StochasticShadowsParameters;
 	PassParameters->RWLightSamples = GraphBuilder.CreateUAV(LightSamples);
 	PassParameters->LightSampleRayDistance = LightSampleRayDistance;
+	PassParameters->RayTracingBias = CVarStochasticShadowsHardwareRayTracingBias.GetValueOnRenderThread();
+	PassParameters->RayTracingNormalBias = CVarStochasticShadowsHardwareRayTracingNormalBias.GetValueOnRenderThread();
 
 	checkf(View.HasRayTracingScene(), TEXT("TLAS does not exist. Verify that the current pass is represented in Lumen::AnyLumenHardwareRayTracingPassEnabled()."));
 	PassParameters->TLAS = View.GetRayTracingSceneLayerViewChecked(ERayTracingSceneLayer::Base);
