Fix shader warnings

diff --git a/Engine/Shaders/Private/HairStrands/HairStrandsDebug.usf b/Engine/Shaders/Private/HairStrands/HairStrandsDebug.usf
--- a/Engine/Shaders/Private/HairStrands/HairStrandsDebug.usf
+++ b/Engine/Shaders/Private/HairStrands/HairStrandsDebug.usf
@@ -129,7 +129,7 @@ void MainPS(
 	if (!bIsValid)
 		discard;
 
-	const bool bIsDepthValid = HairStrands.HairOnlyDepthTexture.Load(uint3(PixelPos, 0)) > 0;
+	const bool bIsDepthValid = HairStrands.HairOnlyDepthTexture.Load(uint3(PixelPos, 0)).x > 0;
 	if (!bIsDepthValid)
 	{
 		discard;
diff --git a/Engine/Shaders/Private/HairStrands/HairStrandsDeepShadowMask.usf b/Engine/Shaders/Private/HairStrands/HairStrandsDeepShadowMask.usf
--- a/Engine/Shaders/Private/HairStrands/HairStrandsDeepShadowMask.usf
+++ b/Engine/Shaders/Private/HairStrands/HairStrandsDeepShadowMask.usf
@@ -338,7 +338,7 @@ void MainPS(
 	const float2 UV = Input.UV;
 	const uint2 PixelCoord = floor(Input.Position.xy);
 	const float SceneDepth = ConvertFromDeviceZ(SceneDepthTexture.Load(uint3(PixelCoord, 0)).x);
-	const bool bIsHairPixel = HairStrands.HairOnlyDepthTexture.Load(uint3(PixelCoord, 0)) > 0;
+	const bool bIsHairPixel = HairStrands.HairOnlyDepthTexture.Load(uint3(PixelCoord, 0)).x > 0;
 
 	const float3 TranslatedWorldPosition = GetTranslatedWorldPosition(UV, SceneDepth);
 	float4 AtlasScaleBias;
@@ -445,7 +445,7 @@ void MainPS(
 
 	const float DistanceThreshold = 100000.0f;
 	const float CoverageThreshold = 0.995f; // When Coverage is high, we do not trace shadow on opaque since hair/fur is covering the background.
-	const bool  bIsHairPixel = HairStrands.HairOnlyDepthTexture.Load(uint3(PixelCoord, 0)) > 0;
+	const bool  bIsHairPixel = HairStrands.HairOnlyDepthTexture.Load(uint3(PixelCoord, 0)).x > 0;
 	const float HairPixelCoverage = HairStrands.HairCoverageTexture.Load(uint3(PixelCoord, 0));
 	const bool bIsValidOpaquePixel = bIsHairPixel && HairPixelCoverage < CoverageThreshold;
 	
diff --git a/Engine/Shaders/Private/LocalVertexFactory.ush b/Engine/Shaders/Private/LocalVertexFactory.ush
--- a/Engine/Shaders/Private/LocalVertexFactory.ush
+++ b/Engine/Shaders/Private/LocalVertexFactory.ush
@@ -1193,7 +1193,7 @@ float4 VertexFactoryGetPreviousWorldPosition(FVertexFactoryInput Input, FVertexF
 		PrevLocalPosition = mul(Input.Position, InstanceTransform);
 #elif USE_SPLINEDEFORM && !RAYHITGROUPSHADER
 		// Just like CalcWorldPosition, transform into mesh space
-		PrevLocalPosition = float4(SplineMeshDeformLocalPos(Intermediates.Common.SplineMeshParams, Input.Position), Input.Position.w);
+		PrevLocalPosition = float4(SplineMeshDeformLocalPos(Intermediates.Common.SplineMeshParams, Input.Position.xyz), Input.Position.w);
 #endif
 	}
 
diff --git a/Engine/Shaders/Private/ParticleGPUSpriteVertexFactory.ush b/Engine/Shaders/Private/ParticleGPUSpriteVertexFactory.ush
--- a/Engine/Shaders/Private/ParticleGPUSpriteVertexFactory.ush
+++ b/Engine/Shaders/Private/ParticleGPUSpriteVertexFactory.ush
@@ -719,7 +719,7 @@ FVertexFactoryInterpolantsVSToPS VertexFactoryGetInterpolantsVSToPS( FVertexFact
 float4 VertexFactoryGetPreviousWorldPosition( FVertexFactoryInput Input, FVertexFactoryIntermediates Intermediates )
 {
 	float3 Position = Intermediates.VertexWorldPosition;
-	Position.xyz -= Intermediates.ParticleVelocity * EmitterUniforms.UseVelocityForMotionBlur;
+	Position.xyz -= Intermediates.ParticleVelocity.xyz * EmitterUniforms.UseVelocityForMotionBlur;
 	return float4(Position, 1);
 }
 
