LFV - fix split screen LFV view direction due to badSVPos code (DepthBufferTranslatedWorldPos do not need ViewRectMin offset).

diff --git a/Engine/Shaders/Private/LocalFogVolumes/LocalFogVolumeSplat.usf b/Engine/Shaders/Private/LocalFogVolumes/LocalFogVolumeSplat.usf
--- a/Engine/Shaders/Private/LocalFogVolumes/LocalFogVolumeSplat.usf
+++ b/Engine/Shaders/Private/LocalFogVolumes/LocalFogVolumeSplat.usf
@@ -82,7 +82,8 @@ void LocalFogVolumeTiledPS(
 #else
 	const float HalfResScale = 1.0f;
 #endif
-	const float2 SourcePixelPos = SVPos.xy * HalfResScale + ResolvedView.ViewRectMin.xy;
+	const float2 ViewRectSVPos = SVPos.xy * HalfResScale;
+	const float2 SourcePixelPos = ViewRectSVPos + ResolvedView.ViewRectMin.xy;
 
 #if SHADING_PATH_MOBILE
 	float2 UVDepth = SourcePixelPos * ResolvedView.BufferSizeAndInvSize.zw;
@@ -94,10 +95,10 @@ void LocalFogVolumeTiledPS(
 #if HAS_INVERTED_Z_BUFFER
 	DeviceZ = max(0.000000000001, DeviceZ);	// Make sure SvPositionToWorld does not return bad values when DeviceZ is far=0 when using inverted z
 #endif
-	float3 DepthBufferTranslatedWorldPos = SvPositionToTranslatedWorld(float4(SourcePixelPos, DeviceZ, 1.0));
+	float3 DepthBufferTranslatedWorldPos = SvPositionToResolvedTranslatedWorld(float4(ViewRectSVPos, DeviceZ, 1.0));
 
 #if LFV_HALFRES_PS
-	uint2 TilePos = (SourcePixelPos - ResolvedView.ViewRectMin.xy) / LFVTilePixelSize;
+	uint2 TilePos = ViewRectSVPos / LFVTilePixelSize;
 #else
 	uint2 TilePos = LFVUnpackTile(OtherVertexOutput.TileData);
 #endif
