Fix incorrect hair shadow mask.

diff --git a/Engine/Shaders/Private/HairStrands/HairStrandsDeepShadowMask.usf b/Engine/Shaders/Private/HairStrands/HairStrandsDeepShadowMask.usf
--- a/Engine/Shaders/Private/HairStrands/HairStrandsDeepShadowMask.usf
+++ b/Engine/Shaders/Private/HairStrands/HairStrandsDeepShadowMask.usf
@@ -338,7 +338,8 @@ void MainPS(
 	const float2 UV = Input.UV;
 	const uint2 PixelCoord = floor(Input.Position.xy);
 	const float SceneDepth = ConvertFromDeviceZ(SceneDepthTexture.Load(uint3(PixelCoord, 0)).x);
-	const bool bIsHairPixel = HairStrands.HairOnlyDepthTexture.Load(uint3(PixelCoord, 0)).x > 0;
+	const bool bIsValidHairPixel = HairStrands.HairOnlyDepthTexture.Load(uint3(PixelCoord, 0)).x > 0;
+	const bool bIsHairPixel = bIsValidHairPixel && HairStrands.HairCoverageTexture.Load(uint3(PixelCoord, 0)) >= 1;
 
 	const float3 TranslatedWorldPosition = GetTranslatedWorldPosition(UV, SceneDepth);
 	float4 AtlasScaleBias;
@@ -445,9 +446,9 @@ void MainPS(
 
 	const float DistanceThreshold = 100000.0f;
 	const float CoverageThreshold = 0.995f; // When Coverage is high, we do not trace shadow on opaque since hair/fur is covering the background.
-	const bool  bIsHairPixel = HairStrands.HairOnlyDepthTexture.Load(uint3(PixelCoord, 0)).x > 0;
+	const bool bIsHairPixel = HairStrands.HairOnlyDepthTexture.Load(uint3(PixelCoord, 0)).x > 0;
 	const float HairPixelCoverage = HairStrands.HairCoverageTexture.Load(uint3(PixelCoord, 0));
-	const bool bIsValidOpaquePixel = bIsHairPixel && HairPixelCoverage < CoverageThreshold;
+	const bool bIsValidOpaquePixel = !bIsHairPixel || (bIsHairPixel && HairPixelCoverage < CoverageThreshold);
 	
 	float3 TranslatedWorldPosition = GetTranslatedWorldPosition(UV, SceneDepth);
 	if (View.StereoPassIndex == 1)
