LFV - A better match for LFV between volumetric fog 0/1.

diff --git a/Engine/Shaders/Private/VolumetricFog.usf b/Engine/Shaders/Private/VolumetricFog.usf
--- a/Engine/Shaders/Private/VolumetricFog.usf
+++ b/Engine/Shaders/Private/VolumetricFog.usf
@@ -188,11 +188,11 @@ void MaterialSetupCS(
 			float3 UnitSpacePos = mul(float4(WorldPosition, 1.0), FogInstance.InvTransform).xyz;
 
 			float DistToCenter = length(UnitSpacePos - float3(0.0f, 0.0f, 0.0f));
-			float SphereFade = saturate(1.0f - DistToCenter * DistToCenter);
+			float SphereFade = saturate(1.0f - DistToCenter);
 
 			// Evaluate both radial and height local fog contribution
-			float LFVExtinctionRadial = SphereFade <= 0.0f ? 0.0f : FogInstance.HeightFogExtinction * exp(-FogInstance.HeightFogFalloff * (UnitSpacePos.z - FogInstance.HeightFogOffset));
-			float LFVExtinctionHeight = FogInstance.RadialFogExtinction * SphereFade;
+			float LFVExtinctionHeight = SphereFade <= 0.0f ? 0.0f : FogInstance.HeightFogExtinction * exp(-FogInstance.HeightFogFalloff * (UnitSpacePos.z - FogInstance.HeightFogOffset));
+			float LFVExtinctionRadial = SphereFade <= 0.0f ? 0.0f : FogInstance.RadialFogExtinction * pow(SphereFade, 0.82); // 0.82 is a correction factor needed to get a better match. It is unclear why it is needed...
 
 			// As for the local fog volume analitical evaluation, we want to combine both two optical depth for the "opacity"=1-transmittance to be combined together.
 			// This helps to have soft edges when using height fog. 
