Fix issue with non empty voxel page

diff --git a/Engine/Shaders/Private/HairStrands/HairStrandsVoxelDebug.usf b/Engine/Shaders/Private/HairStrands/HairStrandsVoxelDebug.usf
--- a/Engine/Shaders/Private/HairStrands/HairStrandsVoxelDebug.usf
+++ b/Engine/Shaders/Private/HairStrands/HairStrandsVoxelDebug.usf
@@ -307,7 +307,7 @@ void MainCS(uint3 DispatchThreadId : SV_DispatchThreadID)
 
 					AccDensity = 1;
 					bIsValid = true;
-					HitColor = PageCoord / 12.f;
+					HitColor = PageIndexCoord / 12.f;
 					break;
 				}
 			}
diff --git a/Engine/Shaders/Private/HairStrands/HairStrandsVoxelRasterCompute.usf b/Engine/Shaders/Private/HairStrands/HairStrandsVoxelRasterCompute.usf
--- a/Engine/Shaders/Private/HairStrands/HairStrandsVoxelRasterCompute.usf
+++ b/Engine/Shaders/Private/HairStrands/HairStrandsVoxelRasterCompute.usf
@@ -1297,7 +1297,9 @@ void MainCS(uint3 DispatchThreadId : SV_DispatchThreadID, uint3 GroupThreadId :
 
 		// Update the page index with invalid page index if the voxel does not contain any data.
 		// This allow to save tracing cost when evaluating the transmittance.
-		if (bPatchEmptyPage > 0 && Total2 == 0)
+
+		const bool bIsEmpty =(Total2 & VOXEL_HAIR_MASK) == 0;
+		if (bPatchEmptyPage > 0 && bIsEmpty)
 		{
 			const uint PageIndexOffset = PageToPageIndexBuffer[PageIndex];
 			OutPageIndexBuffer[PageIndexOffset] = INVALID_VOXEL_PAGE_INDEX;
