Fixes a bug where the TSR velocity edge mask was biased towards one direction of movement

These notably affected characters's eye lid opening up...

diff --git a/Engine/Shaders/Private/TemporalSuperResolution/TSRDecimateHistory.usf b/Engine/Shaders/Private/TemporalSuperResolution/TSRDecimateHistory.usf
--- a/Engine/Shaders/Private/TemporalSuperResolution/TSRDecimateHistory.usf
+++ b/Engine/Shaders/Private/TemporalSuperResolution/TSRDecimateHistory.usf
@@ -214,8 +214,8 @@ void ProcessVelocityNeighborhood(
 			
 		// Reject high frequency on velocity edge discontinuity.
 		{
-			tsr_half IsOverlapingEdge0 = tsr_half(1.0) - dot(saturate(NeighborDeltaPixelVelocity0 * Offset0), tsr_half(1.0).xx); // * saturate(ScreenPixelVelocityLength - NeighborPixelVelocityLength0);
-			tsr_half IsOverlapingEdge1 = tsr_half(1.0) - dot(saturate(NeighborDeltaPixelVelocity1 * Offset1), tsr_half(1.0).xx); // * saturate(ScreenPixelVelocityLength - NeighborPixelVelocityLength1);
+			tsr_half IsOverlapingEdge0 = tsr_half(1.0) - dot(saturate(abs(NeighborDeltaPixelVelocity0 * Offset0)), tsr_half(1.0).xx); // * saturate(ScreenPixelVelocityLength - NeighborPixelVelocityLength0);
+			tsr_half IsOverlapingEdge1 = tsr_half(1.0) - dot(saturate(abs(NeighborDeltaPixelVelocity1 * Offset1)), tsr_half(1.0).xx); // * saturate(ScreenPixelVelocityLength - NeighborPixelVelocityLength1);
 				
 			tsr_half IsOverlapingEdge = min(IsOverlapingEdge0, IsOverlapingEdge1);
 
