Fix incorrect resolution scale when printing hair instance data.

diff --git a/Engine/Shaders/Private/HairStrands/HairStrandsDebug.usf b/Engine/Shaders/Private/HairStrands/HairStrandsDebug.usf
--- a/Engine/Shaders/Private/HairStrands/HairStrandsDebug.usf
+++ b/Engine/Shaders/Private/HairStrands/HairStrandsDebug.usf
@@ -879,7 +879,7 @@ void MainCS(uint3 DispatchThreadId : SV_DispatchThreadID)
 #include "HairStrandsVisibilityCommon.ush"
 #include "HairStrandsAABBCommon.ush" 
 
-int2 PixelCoord;
+float ResolutionScale;
 int2 MaxResolution;
 uint FastResolveMask;
 uint HairMacroGroupCount;
@@ -997,19 +997,18 @@ void MainCS(uint InGroupIndex : SV_GroupIndex, uint3 DispatchThreadId : SV_Dispa
 		Newline(Context);
 	}
 
-	const float2 PixelPos = PixelCoord;
-	if (PixelCoord.x < 0 || PixelCoord.y < 0 || PixelCoord.x >= MaxResolution.x || PixelCoord.y >= MaxResolution.y)
+	FShaderPrintContext Context = InitShaderPrintContext(true, uint2(50, 500));	
+	const int2 PixelCoord = int2(ResolutionScale * GetCursorPos(Context));
+	if (any(PixelCoord < 0) || any(PixelCoord >= MaxResolution))
 		return;
 
-	const float HairPixelCoverage = HairStrands.HairCoverageTexture.Load(uint3(PixelPos, 0));
-	const bool bIsValid = HairPixelCoverage > 0;
+	const bool bIsValid = HairStrands.HairOnlyDepthTexture.Load(uint3(PixelCoord, 0)).x > 0;
 	if (!bIsValid)
 		return;
 
-	const float ViewHairCount = HairCountTexture.Load(uint3(PixelPos, 0));
-	const float ViewHairCountUint = HairCountUintTexture.Load(uint3(PixelPos, 0)) /  1000.f;
-
-	FShaderPrintContext Context = InitShaderPrintContext(true, uint2(50, 500));
+	const float ViewHairCount = HairCountTexture.Load(uint3(PixelCoord, 0));
+	const float ViewHairCountUint = HairCountUintTexture.Load(uint3(PixelCoord, 0)) /  1000.f;
+	const float HairPixelCoverage = HairStrands.HairCoverageTexture.Load(uint3(PixelCoord, 0));
 
 	uint Total8BitCoverage = 0;
 	const FNodeDesc NodeDesc = DecodeNodeDesc(HairStrands.HairSampleOffset.Load(uint3(PixelCoord, 0)));
@@ -1057,7 +1056,7 @@ void MainCS(uint InGroupIndex : SV_GroupIndex, uint3 DispatchThreadId : SV_Dispa
 	Newline(Context);
 
 	// Resolve type
-	const uint Stencil = StencilTexture.Load(uint3(PixelPos,0)).y;
+	const uint Stencil = StencilTexture.Load(uint3(PixelCoord,0)).y;
 	const bool bIsFastResolve = (Stencil & FastResolveMask) != 0;
 	Print(Context, TEXT("Resolve         "));
 
diff --git a/Engine/Source/Runtime/Renderer/Private/HairStrands/HairStrandsDebug.cpp b/Engine/Source/Runtime/Renderer/Private/HairStrands/HairStrandsDebug.cpp
--- a/Engine/Source/Runtime/Renderer/Private/HairStrands/HairStrandsDebug.cpp
+++ b/Engine/Source/Runtime/Renderer/Private/HairStrands/HairStrandsDebug.cpp
@@ -193,13 +193,13 @@ class FHairDebugPrintCS : public FGlobalShader
 
 	BEGIN_SHADER_PARAMETER_STRUCT(FParameters, )
 		SHADER_PARAMETER(FIntPoint, GroupSize)
-		SHADER_PARAMETER(FIntPoint, PixelCoord)
 		SHADER_PARAMETER(FIntPoint, MaxResolution)
 		SHADER_PARAMETER(uint32, FastResolveMask)
 		SHADER_PARAMETER(uint32, HairMacroGroupCount)
 		SHADER_PARAMETER(uint32, HairVisibilityNodeGroupSize)
 		SHADER_PARAMETER(uint32, AllocatedSampleCount)
 		SHADER_PARAMETER(uint32, HairInstanceCount)
+		SHADER_PARAMETER(float, ResolutionScale)
 		SHADER_PARAMETER_RDG_BUFFER_SRV(Buffer<uint>, HairInstanceIDs)
 		SHADER_PARAMETER_RDG_TEXTURE(Texture2D, HairCountTexture)
 		SHADER_PARAMETER_RDG_TEXTURE(Texture2D, HairCountUintTexture)
@@ -284,13 +284,13 @@ static void AddDebugHairPrintPass(
 	const FIntPoint Resolution(Viewport.Width(), Viewport.Height());
 
 	FHairDebugPrintCS::FParameters* Parameters = GraphBuilder.AllocParameters<FHairDebugPrintCS::FParameters>();
+	Parameters->ResolutionScale = float(View->ViewRect.Width()) / float(View->UnscaledViewRect.Width());
 	Parameters->Scene = View->GetSceneUniforms().GetBuffer(GraphBuilder);
 	Parameters->HairInstanceCount = InstanceIDs.Num();
 	Parameters->HairInstanceIDs = GraphBuilder.CreateSRV(InstancesIDBuffer, PF_R32_UINT);
 	Parameters->GroupSize = GetVendorOptimalGroupSize2D();
 	Parameters->ViewUniformBuffer = View->ViewUniformBuffer;
 	Parameters->MaxResolution = VisibilityData.CoverageTexture ? VisibilityData.CoverageTexture->Desc.Extent : FIntPoint(0,0);
-	Parameters->PixelCoord = View->CursorPos;
 	Parameters->AllocatedSampleCount = VisibilityData.MaxNodeCount;
 	Parameters->FastResolveMask = STENCIL_TEMPORAL_RESPONSIVE_AA_MASK;
 	Parameters->HairCountTexture = ViewHairCountTexture;
