Adds a min() to force a couple 16bit mul to be done before adds to ensure no overflow happens

This notably allows to not use the "precise" HLSL qualifier which would apply on all parent instruction, removing complete flexibility for compiler to reshufle the order of adds.

This change takes the oportunity to add a couple of missing min(Max10BitsFloat) on linear color UAV write as additional safety measure.

diff --git a/Engine/Shaders/Private/TemporalSuperResolution/TSRCommon.ush b/Engine/Shaders/Private/TemporalSuperResolution/TSRCommon.ush
--- a/Engine/Shaders/Private/TemporalSuperResolution/TSRCommon.ush
+++ b/Engine/Shaders/Private/TemporalSuperResolution/TSRCommon.ush
@@ -51,6 +51,11 @@
 
 #define CONFIG_ENABLE_STOCASTIC_QUANTIZATION (!CONFIG_SCENE_COLOR_ALPHA)
 
+// DXC allows changing order between multiply and adds [REDACTED_PROFILE_URL]/DirectXShaderCompiler/blob/main/docs/DXIL.rst#precise-qualifier
+// which on 16bit, if linear color spaces ends up close to half(65504.0), this can go +inf on if the compiler transform a * 0.5 + b * 0.5 into (a + b) * 0.5.
+#define CONFIG_FP16_PRECISE_MULTIPLY_ORDER 1
+
+
 //------------------------------------------------------- CONSTANTS
 
 /* Maximum number of sample. */
diff --git a/Engine/Shaders/Private/TemporalSuperResolution/TSRRejectShading.usf b/Engine/Shaders/Private/TemporalSuperResolution/TSRRejectShading.usf
--- a/Engine/Shaders/Private/TemporalSuperResolution/TSRRejectShading.usf
+++ b/Engine/Shaders/Private/TemporalSuperResolution/TSRRejectShading.usf
@@ -695,6 +695,8 @@ void MainCS(
 			CenterLumaLDR.SetElement(SimdIndex, CenterLuma / (SpatialAAExposure + CenterLuma));
 		}
 
+		CenterFinalSceneColor = min(CenterFinalSceneColor, tsr_simd_halfC::Const(tsr_half(Max10BitsFloat)));
+
 		// Output some data early.
 		PLATFORM_SPECIFIC_ISOLATE
 		{
diff --git a/Engine/Shaders/Private/TemporalSuperResolution/TSRResolveHistory.usf b/Engine/Shaders/Private/TemporalSuperResolution/TSRResolveHistory.usf
--- a/Engine/Shaders/Private/TemporalSuperResolution/TSRResolveHistory.usf
+++ b/Engine/Shaders/Private/TemporalSuperResolution/TSRResolveHistory.usf
@@ -246,6 +246,8 @@ void MainCS(
 			tsr_short2(OutputViewRectMin) +
 			tsr_short2(GroupId) * tsr_short(TILE_SIZE).xx);
 		
+		Output = min(Output, tsr_output_halfC::Const(tsr_half(Max10BitsFloat)));
+
 		UNROLL_N(OUTPUT_SIMD_SIZE)
 		for (uint SimdIndex = 0; SimdIndex < OUTPUT_SIMD_SIZE; SimdIndex++)
 		{
@@ -259,6 +261,11 @@ void MainCS(
 		
 		{
 			tsr_output_halfC HalfResOutput = Output * tsr_half(0.25);
+			// Forces the * tsr_half(0.25) to be applied before to avoid turning bright pixels to +inf in the adds below.
+			#if CONFIG_FP16_PRECISE_MULTIPLY_ORDER
+				HalfResOutput = min(HalfResOutput, tsr_output_halfC::Const(tsr_half(Max10BitsFloat * 0.25)));
+			#endif
+
 			HalfResOutput = HalfResOutput + WaveAccessNeighborTexel(HalfResOutput, tsr_short2(1, 0));
 			HalfResOutput = HalfResOutput + WaveAccessNeighborTexel(HalfResOutput, tsr_short2(0, 1));
 		
@@ -353,6 +360,11 @@ void MainCS(
 		uint LocalGroupThreadIndex = GroupThreadIndex;
 
 		tsr_halfC HalfResOutput = FinalOutputColor * tsr_half(0.25);
+		// Forces the * tsr_half(0.25) to be applied before to avoid turning bright pixels to +inf in the adds below.
+		#if CONFIG_FP16_PRECISE_MULTIPLY_ORDER
+			HalfResOutput = min(HalfResOutput, tsr_half(Max10BitsFloat * 0.25));
+		#endif
+
 		tsr_short2 HalfResOutputPixelPos;
 		HalfResOutputPixelPos[0] = (OutputPixelPosition[0] >> tsr_short(1)) | (((OutputPixelPosition[0] & tsr_short(0x1)) | tsr_short(!bGenerateOutputMip1)) * tsr_short(~0));
 		HalfResOutputPixelPos[1] = (OutputPixelPosition[1] >> tsr_short(1)) | (((OutputPixelPosition[1] & tsr_short(0x1)) | tsr_short(!bGenerateOutputMip1)) * tsr_short(~0));
diff --git a/Engine/Shaders/Private/TemporalSuperResolution/TSRUpdateHistory.usf b/Engine/Shaders/Private/TemporalSuperResolution/TSRUpdateHistory.usf
--- a/Engine/Shaders/Private/TemporalSuperResolution/TSRUpdateHistory.usf
+++ b/Engine/Shaders/Private/TemporalSuperResolution/TSRUpdateHistory.usf
@@ -325,6 +325,15 @@ T BilinearSampleColorHistory(Texture2D<T> Texture, float2 SampleUV)
 tsr_halfCx2 DownsampleSceneColor(tsr_halfCx2 SceneColor, uint LocalGroupThreadIndex, const uint XorButterFly)
 {
 	SceneColor = SceneColor * tsr_half(0.25);
+	// Forces the * tsr_half(0.25) to be applied before to avoid turning bright pixels to +inf in the adds below.
+	#if CONFIG_FP16_PRECISE_MULTIPLY_ORDER
+		#if CONFIG_SCENE_COLOR_ALPHA
+			const tsr_halfC ColorMax10BitsFloat = tsr_half(Max10BitsFloat * 0.25).xxxx;
+		#else
+			const tsr_halfC ColorMax10BitsFloat = tsr_half(Max10BitsFloat * 0.25).xxx;
+		#endif
+		SceneColor = dpv_min(SceneColor, ColorMax10BitsFloat);
+	#endif
 
 	#if PLATFORM_SUPPORTS_WAVE_BROADCAST && CONFIG_ENABLE_DUAL_PIXEL_VECTORIZATION
 	if (uint(XorButterFly * 4) <= WaveGetLaneCount())
diff --git a/Engine/Source/Runtime/Renderer/Private/PostProcess/TemporalSuperResolution.cpp b/Engine/Source/Runtime/Renderer/Private/PostProcess/TemporalSuperResolution.cpp
--- a/Engine/Source/Runtime/Renderer/Private/PostProcess/TemporalSuperResolution.cpp
+++ b/Engine/Source/Runtime/Renderer/Private/PostProcess/TemporalSuperResolution.cpp
@@ -92,7 +92,7 @@ TAutoConsoleVariable<int32> CVarTSR16BitVALUOnAMD(
 	ECVF_RenderThreadSafe);
 
 TAutoConsoleVariable<int32> CVarTSR16BitVALUOnIntel(
-	TEXT("r.TSR.16BitVALU.Intel"), 0,
+	TEXT("r.TSR.16BitVALU.Intel"), 1,
 	TEXT("Overrides whether to use 16bit VALU on Intel desktop GPUs"),
 	ECVF_RenderThreadSafe);
 
