Added missing max(0, ) to MakeFromColorTemperature() and MaterialExpressionBlackBody() and also added a missing clamp on the temperature input to MaterialExpressionBlackBody.

#rnx

diff --git a/Engine/Shaders/Private/MaterialTemplate.ush b/Engine/Shaders/Private/MaterialTemplate.ush
--- a/Engine/Shaders/Private/MaterialTemplate.ush
+++ b/Engine/Shaders/Private/MaterialTemplate.ush
@@ -1532,6 +1532,8 @@ float MaterialExpressionDepthOfFieldFunction(float SceneDepth, int FunctionValue
 // TODO convert to LUT
 float3 MaterialExpressionBlackBody( float Temp )
 {
+	Temp = clamp(Temp, 1000.0f, 15000.0f);
+
 	float u = ( 0.860117757f + 1.54118254e-4f * Temp + 1.28641212e-7f * Temp*Temp ) / ( 1.0f + 8.42420235e-4f * Temp + 7.08145163e-7f * Temp*Temp );
 	float v = ( 0.317398726f + 4.22806245e-5f * Temp + 4.20481691e-8f * Temp*Temp ) / ( 1.0f - 2.89741816e-5f * Temp + 1.61456053e-7f * Temp*Temp );
 
@@ -1550,7 +1552,8 @@ float3 MaterialExpressionBlackBody( float Temp )
 		 0.0556434, -0.2040259,  1.0572252,
 	};
 
-	return mul( XYZtoRGB, float3( X, Y, Z ) ) * pow( 0.0004 * Temp, 4 );
+	// The XYZtoRGB transform can result in negative values, so we need to clamp here.
+	return max( 0.0f, mul( XYZtoRGB, float3( X, Y, Z ) ) ) * pow( 0.0004 * Temp, 4 );
 }
 
 float2 MaterialExpressionGetHairRootUV(FMaterialPixelParameters Parameters)
diff --git a/Engine/Source/Runtime/ColorManagement/Private/ColorSpace.cpp b/Engine/Source/Runtime/ColorManagement/Private/ColorSpace.cpp
--- a/Engine/Source/Runtime/ColorManagement/Private/ColorSpace.cpp
+++ b/Engine/Source/Runtime/ColorManagement/Private/ColorSpace.cpp
@@ -239,7 +239,8 @@ FLinearColor FColorSpace::MakeFromColorTemperature(float Temp) const
 	FVector3d XYZ = FVector3d(1.0 / y * x, 1.0, 1.0 / y * z);
 	FVector4d RGB = XYZToRgb.TransformVector(XYZ);
 
-	return FLinearColor((float)RGB.X, (float)RGB.Y, (float)RGB.Z);
+	// The XYZ to RGB transform can result in negative values, so we need to clamp here.
+	return FLinearColor(FMath::Max(0.0f, (float)RGB.X), FMath::Max(0.0f, (float)RGB.Y), FMath::Max(0.0f, (float)RGB.Z));
 }
 
 float FColorSpace::GetLuminance(const FLinearColor& Color) const
diff --git a/Engine/Source/Runtime/Core/Private/Math/Color.cpp b/Engine/Source/Runtime/Core/Private/Math/Color.cpp
--- a/Engine/Source/Runtime/Core/Private/Math/Color.cpp
+++ b/Engine/Source/Runtime/Core/Private/Math/Color.cpp
@@ -511,7 +511,8 @@ FLinearColor FLinearColor::MakeFromColorTemperature( float Temp )
 	float G = -0.9692660f * X +  1.8760108f * Y +  0.0415560f * Z;
 	float B =  0.0556434f * X + -0.2040259f * Y +  1.0572252f * Z;
 
-	return FLinearColor(R,G,B);
+	// The XYZ to RGB transform can result in negative values, so we need to clamp here.
+	return FLinearColor(FMath::Max(0.0f, R), FMath::Max(0.0f, G), FMath::Max(0.0f, B));
 }
 
 FLinearColor FLinearColor::MakeRandomSeededColor(int32 Seed)
