Use GPUSceneInstanceIndex instead of HitCulledObjectIndex to color instances in DF visualization.
- makes the color stable when instances are modified.

diff --git a/Engine/Shaders/Private/DistanceFieldVisualization.usf b/Engine/Shaders/Private/DistanceFieldVisualization.usf
--- a/Engine/Shaders/Private/DistanceFieldVisualization.usf
+++ b/Engine/Shaders/Private/DistanceFieldVisualization.usf
@@ -250,7 +250,10 @@ void VisualizeMeshDistanceFieldCS(
 		Lighting += GetSkySHDiffuse(WorldHitNormal) * View.SkyLightColor.rgb;
 		float3 Albedo = .05f;
 #if !USE_GLOBAL_DISTANCE_FIELD
-		Albedo = 0.1f * IntToColor(HitCulledObjectIndex + 1);
+		{
+			FDFObjectData DFObjectData = LoadDFObjectData(HitCulledObjectIndex);
+			Albedo = 0.1f * IntToColor(DFObjectData.GPUSceneInstanceIndex);
+		}
 #endif
 		Lighting *= Albedo;
 	}
