Substrate - Fix for clustered shading to feature proper light attenuation.

diff --git a/Engine/Shaders/Private/ClusteredDeferredShadingPixelShader.usf b/Engine/Shaders/Private/ClusteredDeferredShadingPixelShader.usf
--- a/Engine/Shaders/Private/ClusteredDeferredShadingPixelShader.usf
+++ b/Engine/Shaders/Private/ClusteredDeferredShadingPixelShader.usf
@@ -162,9 +162,8 @@ float4 GetSimpleLightLighting(
 		FSubstrateShadowTermInputParameters SubstrateShadowTermInputParameters	= GetInitialisedSubstrateShadowTermInputParameters();
 		SubstrateShadowTermInputParameters.bEvaluateShadowTerm					= false;
 
-		const float3 ToLight = LightData.TranslatedWorldPosition - TranslatedWorldPosition;
-		const float3 L = normalize(ToLight);
 		const float3 V = -CameraVector;
+		float3 ToLight, L;
 		const float LightMask = GetLocalLightAttenuation(TranslatedWorldPosition, LightData, ToLight, L);
 
 		FSubstrateDeferredLighting SubstrateLighting = SubstrateDeferredLighting(
@@ -313,10 +312,6 @@ float4 GetLightGridLocalLighting(
 
 	#if SUBSTRATE_ENABLED && USE_HAIR_LIGHTING == 0
 
-		const float3 ToLight = LightData.TranslatedWorldPosition - TranslatedWorldPosition;
-		const float3 L = normalize(ToLight);
-		const float3 V = -CameraVector;
-
 		FSubstrateShadowTermInputParameters SubstrateShadowTermInputParameters	= GetInitialisedSubstrateShadowTermInputParameters();
 		SubstrateShadowTermInputParameters.bEvaluateShadowTerm					= true;
 		SubstrateShadowTermInputParameters.SceneDepth							= CalcSceneDepth(ScreenUV);
@@ -325,7 +320,9 @@ float4 GetLightGridLocalLighting(
 		SubstrateShadowTermInputParameters.LightAttenuation						= LightAttenuation;
 		SubstrateShadowTermInputParameters.Dither								= Dither;
 
-		const float LightMask = 1;
+		const float3 V = -CameraVector;
+		float3 ToLight, L;
+		const float LightMask = GetLocalLightAttenuation(TranslatedWorldPosition, LightData, ToLight, L);
 
 		FSubstrateDeferredLighting SubstrateLighting = (FSubstrateDeferredLighting)0;
 	#if SUPPORT_SUBSTRATE_LIGHTING
@@ -419,6 +416,13 @@ void ClusteredShadingPixelShader(
 
 	FSubstrateAddressing SubstrateAddressing = GetSubstratePixelDataByteOffset(PixelPos, uint2(View.BufferSizeAndInvSize.xy), Substrate.MaxBytesPerPixel);
 	FSubstratePixelHeader SubstratePixelHeader = UnpackSubstrateHeaderIn(Substrate.MaterialTextureArray, SubstrateAddressing, Substrate.TopLayerTexture);
+
+	// NOTE: this early out helps with the case where there are no lights in the grid cell.
+	if (CulledLightGridData.NumLights == 0 || !SubstratePixelHeader.IsSubstrateMaterial())
+	{
+		OutColor = 0;
+		return;
+	}
 	
 	float4 CompositedLighting = 0;
 	float3 OpaqueRoughRefractionSceneColor = 0;
