Fix CIS shader compilation issue.

diff --git a/Engine/Shaders/Private/HairStrands/HairCardsAttributeCommon.ush b/Engine/Shaders/Private/HairStrands/HairCardsAttributeCommon.ush
--- a/Engine/Shaders/Private/HairStrands/HairCardsAttributeCommon.ush
+++ b/Engine/Shaders/Private/HairStrands/HairCardsAttributeCommon.ush
@@ -2,6 +2,8 @@
 
 #pragma once
 
+#include "/Engine/Shared/HairStrandsDefinitions.h"
+
 ///////////////////////////////////////////////////////////////////////////////////////////////////
 
 // Vertex factory Flags
