Shows TSR Visualizers in pink on camera cuts

diff --git a/Engine/Shaders/Private/TemporalSuperResolution/TSRVisualize.usf b/Engine/Shaders/Private/TemporalSuperResolution/TSRVisualize.usf
--- a/Engine/Shaders/Private/TemporalSuperResolution/TSRVisualize.usf
+++ b/Engine/Shaders/Private/TemporalSuperResolution/TSRVisualize.usf
@@ -105,6 +105,7 @@ tsr_half3 DisplayHistoryRejection(FContext Ctx)
 	tsr_half3 Color = kWhite;
 	Color = lerp(kOrange, Color, saturate((LowFrequencyRejection - 0.5) * 2.0));
 	Color = lerp(kRed,    Color, saturate((LowFrequencyRejection - 0.0) * 2.0));
+	Color = select(bCameraCut == 0, Color, kPink);
 
 	return SampleTSROutputGray(Ctx) * Color;
 }
@@ -116,6 +117,7 @@ tsr_half3 DisplayHistoryClamp(FContext Ctx)
 	tsr_half3 Color = kWhite;
 	Color = lerp(kOrange, Color, saturate((LowFrequencyClamp - 0.5) * 2.0));
 	Color = lerp(kRed,    Color, saturate((LowFrequencyClamp - 0.0) * 2.0));
+	Color = select(bCameraCut == 0, Color, kPink);
 
 	return SampleTSROutputGray(Ctx) * Color;
 }
@@ -123,7 +125,11 @@ tsr_half3 DisplayHistoryClamp(FContext Ctx)
 tsr_half3 DisplayParallaxDisocclusionMask(FContext Ctx)
 {
 	bool ParallaxDisocclusionMask = (tsr_ushort(HistoryRejectionTexture.SampleLevel(GlobalPointClampedSampler, Ctx.InputUV, 0).a * tsr_half(255.0)) & tsr_ushort(0x1)) != tsr_ushort(0x0);
-	return SampleTSROutputGray(Ctx) * select(ParallaxDisocclusionMask, kWhite, kRed);
+	
+	tsr_half3 Color = kWhite;
+	Color = select(bCameraCut == 0, Color, kPink);
+	Color = select(ParallaxDisocclusionMask, Color, kRed);
+	return SampleTSROutputGray(Ctx) * Color;
 }
 
 tsr_half3 DisplayResurrectionMask(FContext Ctx)
@@ -223,6 +229,7 @@ tsr_half3 DisplayAntiFlickering(FContext Ctx)
 	Color = lerp(Color,  kRed,    saturate(LocalMoireError * 16.0 - 1.0));
 	Color = lerp(kWhite, Color,   LocalCountFadeIn);
 	Color = lerp(kPink, Color,   DisableAntiFlickering);
+	Color = select(bCameraCut == 0, Color, kPink);
 
 	return SampleTSROutputGray(Ctx) * Color;
 }
