Stochastic Shadows:
- Replaced all usage of LightSceneId with LocalLightIndex (except in debug info).

diff --git a/Engine/Shaders/Private/StochasticShadows/StochasticShadows.ush b/Engine/Shaders/Private/StochasticShadows/StochasticShadows.ush
--- a/Engine/Shaders/Private/StochasticShadows/StochasticShadows.ush
+++ b/Engine/Shaders/Private/StochasticShadows/StochasticShadows.ush
@@ -16,9 +16,9 @@
 #define DEBUG_MODE_VISUALIZE_TRACING  2
 
 // Limited by PackLightSample
-#define LIGHT_SCENE_ID_BITS			(14)
-#define LIGHT_SCENE_ID_MASK			((1 << LIGHT_SCENE_ID_BITS) - 1)
-#define INVALID_LIGHT_SCENE_ID		(LIGHT_SCENE_ID_MASK)
+#define LOCAL_LIGHT_INDEX_BITS		(14)
+#define LOCAL_LIGHT_INDEX_MASK		((1 << LOCAL_LIGHT_INDEX_BITS) - 1)
+#define INVALID_LOCAL_LIGHT_INDEX	(LOCAL_LIGHT_INDEX_MASK)
 
 #define TILE_SIZE						8
 #define SHADOW_MASK_TILE_SIZE			8
@@ -27,9 +27,10 @@
 #define MAX_SHADOW_MASK_HISTORY_LENGTH	7
 
 // Keep in sync with StochasticShadows.cpp
-#define MAX_LIGHT_SCENE_ID_XY		16
-#define MAX_LIGHT_SCENE_ID			(MAX_LIGHT_SCENE_ID_XY * MAX_LIGHT_SCENE_ID_XY)
-#define LIGHT_SCENE_ID_MASK_SIZE	(MAX_LIGHT_SCENE_ID / 32)
+#define MAX_LOCAL_LIGHT_INDEX_XY	16
+#define MAX_LOCAL_LIGHT_INDEX		(MAX_LOCAL_LIGHT_INDEX_XY * MAX_LOCAL_LIGHT_INDEX_XY)
+
+#define SHARED_LIGHT_MASK_SIZE		(MAX_LOCAL_LIGHT_INDEX / 32)
 
 #define SHADOW_MASK_ATLAS_SIZE_IN_TILES 512
 #define DOWNSAMPLE_FACTOR				2
@@ -88,8 +89,8 @@ uint2 ShadowMaskPageTablePerLightSize;
 uint2 GetShadowMaskPageTableCoord(uint LightIndex, uint2 ShadowMaskScreenCoord)
 {
 	uint2 LightOffset;
-	LightOffset.x = (LightIndex % MAX_LIGHT_SCENE_ID_XY) * ShadowMaskPageTablePerLightSize.x;
-	LightOffset.y = (LightIndex / MAX_LIGHT_SCENE_ID_XY) * ShadowMaskPageTablePerLightSize.y;
+	LightOffset.x = (LightIndex % MAX_LOCAL_LIGHT_INDEX_XY) * ShadowMaskPageTablePerLightSize.x;
+	LightOffset.y = (LightIndex / MAX_LOCAL_LIGHT_INDEX_XY) * ShadowMaskPageTablePerLightSize.y;
 
 	uint2 TileCoord = ShadowMaskScreenCoord / SHADOW_MASK_TILE_SIZE;
 
@@ -199,14 +200,14 @@ FShadowMaskAccumulator InitShadowMaskAccumulator()
 
 struct FLightSample
 {
-	uint LightSceneId;
+	uint LocalLightIndex;
 	bool bVisible;
 	bool bCompleted; // Whether tracing was completed
 };
 
 uint PackLightSample(FLightSample LightSample)
 {
-	uint PackedLightSample = LightSample.LightSceneId & LIGHT_SCENE_ID_MASK;
+	uint PackedLightSample = LightSample.LocalLightIndex & LOCAL_LIGHT_INDEX_MASK;
 	PackedLightSample |= LightSample.bVisible ? 0x8000 : 0;
 	PackedLightSample |= LightSample.bCompleted ? 0x4000 : 0;
 	return PackedLightSample;
@@ -215,7 +216,7 @@ uint PackLightSample(FLightSample LightSample)
 FLightSample UnpackLightSample(uint PackedLightSample)
 {
 	FLightSample LightSample;
-	LightSample.LightSceneId = PackedLightSample & LIGHT_SCENE_ID_MASK;
+	LightSample.LocalLightIndex = PackedLightSample & LOCAL_LIGHT_INDEX_MASK;
 	LightSample.bVisible = PackedLightSample & 0x8000 ? true : false;
 	LightSample.bCompleted = PackedLightSample & 0x4000 ? true : false;
 	return LightSample;
@@ -251,43 +252,30 @@ float2 DownsampledScreenCoordToScreenUV(uint2 DownsampledScreenCoord)
 	return ScreenUV;
 }
 
-uint2 PackShadowMaskTile(uint LightSceneId, uint LocalLightIndex, uint2 ShadowMaskScreenCoord)
+uint2 PackShadowMaskTile(uint LocalLightIndex, uint2 ShadowMaskScreenCoord)
 {
 	ShadowMaskScreenCoord /= SHADOW_MASK_TILE_SIZE;
 
-	checkSlow(LightSceneId <= LIGHT_SCENE_ID_MASK);
-	checkSlow(LocalLightIndex <= (0xFFFFFFFF >> MAX_LIGHT_SCENE_ID));
-
 	uint2 PackedTile;
-	PackedTile.x = (LightSceneId & LIGHT_SCENE_ID_MASK) | (LocalLightIndex << LIGHT_SCENE_ID_BITS);
+	PackedTile.x = LocalLightIndex;
 	PackedTile.y = (ShadowMaskScreenCoord.x & 0xFFFF) | ((ShadowMaskScreenCoord.y & 0xFFFF) << 16);
 	return PackedTile;
 }
 
-void UnpackShadowMaskTile(uint2 PackedTile, out uint LightSceneId, out uint LocalLightIndex, out uint2 ShadowMaskScreenCoord)
+void UnpackShadowMaskTile(uint2 PackedTile, out uint LocalLightIndex, out uint2 ShadowMaskScreenCoord)
 {
-	LightSceneId = PackedTile.x & LIGHT_SCENE_ID_MASK;
-	LocalLightIndex = PackedTile.x >> LIGHT_SCENE_ID_BITS;
+	LocalLightIndex = PackedTile.x;
 	ShadowMaskScreenCoord.x = PackedTile.y & 0xFFFF;
 	ShadowMaskScreenCoord.y = PackedTile.y >> 16;
 
 	ShadowMaskScreenCoord *= SHADOW_MASK_TILE_SIZE;
 }
 
-/**
- * Map shadow mask tile to a page table hash map entry.
- */
-uint GetShadowMaskPageTableKey(uint LightSceneId, uint2 PageTableCoord)
-{
-	uint HashedKey = MurmurMix(MurmurAdd(LightSceneId, MurmurAdd(PageTableCoord.x, PageTableCoord.y)));
-	return HashedKey;
-}
-
 Texture2D<uint> ShadowMaskPageTableHistory;
 Texture2D<uint> ShadowMaskAtlasHistory;
 int2 ShadowMaskViewSize;
 
-void AccumulateHistoryShadowMask(uint LightSceneId, uint PrevLocalLightIndex, uint2 ShadowMaskScreenCoord, float Weight, inout FShadowMaskAccumulator Accumulator)
+void AccumulateHistoryShadowMask(uint PrevLocalLightIndex, uint2 ShadowMaskScreenCoord, float Weight, inout FShadowMaskAccumulator Accumulator)
 {
 	if (all(ShadowMaskScreenCoord < ShadowMaskViewSize) && Weight > 0.01f)
 	{
@@ -317,13 +305,13 @@ void AccumulateHistoryShadowMask(uint LightSceneId, uint PrevLocalLightIndex, ui
 
 struct FShadingTile
 {
-	uint LightSceneId;
+	uint LocalLightIndex;
 	uint ShadingTileIndex;
 };
 
-uint PackShadingTile(uint LightSceneId, uint ShadingTileIndex)
+uint PackShadingTile(uint LocalLightIndex, uint ShadingTileIndex)
 {
-	uint PackedShadingTile = LightSceneId & 0xFFF;
+	uint PackedShadingTile = LocalLightIndex & 0xFFF;
 	PackedShadingTile |= ShadingTileIndex << 12;
 	return PackedShadingTile;
 }
@@ -331,7 +319,7 @@ uint PackShadingTile(uint LightSceneId, uint ShadingTileIndex)
 FShadingTile UnpackShadingTile(uint PackedShadingTile)
 {
 	FShadingTile ShadingTile;
-	ShadingTile.LightSceneId = PackedShadingTile & 0xFFF;
+	ShadingTile.LocalLightIndex = PackedShadingTile & 0xFFF;
 	ShadingTile.ShadingTileIndex = PackedShadingTile >> 12;
 	return ShadingTile;
 }
diff --git a/Engine/Shaders/Private/StochasticShadows/StochasticShadowsComposite.usf b/Engine/Shaders/Private/StochasticShadows/StochasticShadowsComposite.usf
--- a/Engine/Shaders/Private/StochasticShadows/StochasticShadowsComposite.usf
+++ b/Engine/Shaders/Private/StochasticShadows/StochasticShadowsComposite.usf
@@ -56,12 +56,13 @@ void CompositeShadowMaskTracesCS(
 {
 	uint ShadowMaskTileIndex = GroupId.x;
 
-	uint LightSceneId;
 	uint LocalLightIndex;
 	uint2 BaseShadowMaskScreenCoord;	
-	UnpackShadowMaskTile(ShadowMaskTileData[ShadowMaskTileIndex], LightSceneId, LocalLightIndex, BaseShadowMaskScreenCoord);
+	UnpackShadowMaskTile(ShadowMaskTileData[ShadowMaskTileIndex], LocalLightIndex, BaseShadowMaskScreenCoord);
 
 	GroupMemoryBarrierWithGroupSync();
+
+	const FLocalLightData LocalLightData = GetLocalLightDataNonStereo(LocalLightIndex);
 	
 	// Composite traces into SharedShadowMask
 	for (uint SharedShadowMaskY = GroupThreadId.y; SharedShadowMaskY < SHARED_TILE_SIZE; SharedShadowMaskY += THREADGROUP_SIZE)
@@ -81,7 +82,7 @@ void CompositeShadowMaskTracesCS(
 					uint2 SampleCoord = ShadowMaskScreenCoord * NumSamplesPerPixel + uint2(LocalSampleX, LocalSampleY);
 					FLightSample LightSample = UnpackLightSample(LightSamples[SampleCoord]);
 
-					if (LightSample.LightSceneId == LightSceneId)
+					if (LightSample.LocalLightIndex == LocalLightIndex)
 					{
 						SampleAccumulator.SampleSum += LightSample.bVisible ? 1 : 0;
 						SampleAccumulator.SampleWeightSum += 1.0f;
@@ -101,7 +102,7 @@ void CompositeShadowMaskTracesCS(
 
 #if DEBUG_MODE
 	int2 DebugScreenCoord = View.CursorPosition.x >= 0 ? View.CursorPosition * View.ViewResolutionFraction : -1;
-	bool bDebug = all(ShadowMaskScreenCoord == DebugScreenCoord / SHADOW_MASK_DOWNSAMPLE_FACTOR) && LightSceneId == DebugLightId;
+	bool bDebug = all(ShadowMaskScreenCoord == DebugScreenCoord / SHADOW_MASK_DOWNSAMPLE_FACTOR) && LocalLightData.LightSceneId == DebugLightId;
 	FShaderPrintContext Context = InitShaderPrintContext(true, float2(0.05, 0.3));
 
 	if (bDebug)
@@ -194,10 +195,10 @@ void CompositeShadowMaskTracesCS(
 				const FLocalLightData LocalLightData = GetLocalLightDataNonStereo(LocalLightIndex);
 				const uint PrevLocalLightIndex = LocalLightData.PrevLocalLightIndex;
 
-				AccumulateHistoryShadowMask(LightSceneId, PrevLocalLightIndex, ScreenCoord00 + uint2(0, 0), Weights.x, Accumulator);
-				AccumulateHistoryShadowMask(LightSceneId, PrevLocalLightIndex, ScreenCoord00 + uint2(1, 0), Weights.y, Accumulator);
-				AccumulateHistoryShadowMask(LightSceneId, PrevLocalLightIndex, ScreenCoord00 + uint2(0, 1), Weights.z, Accumulator);
-				AccumulateHistoryShadowMask(LightSceneId, PrevLocalLightIndex, ScreenCoord00 + uint2(1, 1), Weights.w, Accumulator);
+				AccumulateHistoryShadowMask(PrevLocalLightIndex, ScreenCoord00 + uint2(0, 0), Weights.x, Accumulator);
+				AccumulateHistoryShadowMask(PrevLocalLightIndex, ScreenCoord00 + uint2(1, 0), Weights.y, Accumulator);
+				AccumulateHistoryShadowMask(PrevLocalLightIndex, ScreenCoord00 + uint2(0, 1), Weights.z, Accumulator);
+				AccumulateHistoryShadowMask(PrevLocalLightIndex, ScreenCoord00 + uint2(1, 1), Weights.w, Accumulator);
 
 				if (Accumulator.WeightSum > 0.0f)
 				{
diff --git a/Engine/Shaders/Private/StochasticShadows/StochasticShadowsHardwareRayTracing.usf b/Engine/Shaders/Private/StochasticShadows/StochasticShadowsHardwareRayTracing.usf
--- a/Engine/Shaders/Private/StochasticShadows/StochasticShadowsHardwareRayTracing.usf
+++ b/Engine/Shaders/Private/StochasticShadows/StochasticShadowsHardwareRayTracing.usf
@@ -63,7 +63,7 @@ RAY_TRACING_ENTRY_RAYGEN_OR_INLINE(HardwareRayTraceLightSamples)
 		float3 TranslatedWorldPosition = GetTranslatedWorldPositionFromScreenUV(ScreenUV, SceneDepth);
 
 		FLightSample LightSample = UnpackLightSample(RWLightSamples[SampleCoord]);
-		FLightSampleTrace LightSampleTrace = GetLightSampleTrace(TranslatedWorldPosition, LightSample.LightSceneId, SampleCoord);
+		FLightSampleTrace LightSampleTrace = GetLightSampleTrace(TranslatedWorldPosition, LightSample.LocalLightIndex, SampleCoord);
 		const float LightSampleRayMinT = LightSampleRayDistance[SampleCoord];
 
 		FRayDesc Ray = (FRayDesc)0;
@@ -132,11 +132,13 @@ RAY_TRACING_ENTRY_RAYGEN_OR_INLINE(HardwareRayTraceLightSamples)
 		#if DEBUG_MODE
 		if (bDebug && DebugMode == DEBUG_MODE_VISUALIZE_TRACING)
 		{
-			FLightSceneData LightSceneData = GetLightSceneData(LightSample.LightSceneId);
+			const FLocalLightData LocalLightData = GetLocalLightDataNonStereo(LightSample.LocalLightIndex);
+			const FLightSceneData LightSceneData = GetLightSceneData(LocalLightData.LightSceneId);
+			const FDeferredLightData LightData = ConvertToDeferredLight(LightSceneData);
 
 			float3 RayStart = Ray.Origin + Ray.Direction * Ray.TMin;
 			float3 RayEnd = Ray.Origin + Ray.Direction * (Payload.IsHit() ? Payload.HitT : Ray.TMax);
-			float4 RayColor = float4(LightSceneData.Color.xyz / Luminance(LightSceneData.Color.xyz), 1.0f);
+			float4 RayColor = float4(LightData.Color.xyz / Luminance(LightData.Color.xyz), 1.0f);
 
 			AddLineTWS(Context, RayStart, RayEnd, RayColor);
 		}
diff --git a/Engine/Shaders/Private/StochasticShadows/StochasticShadowsSampling.usf b/Engine/Shaders/Private/StochasticShadows/StochasticShadowsSampling.usf
--- a/Engine/Shaders/Private/StochasticShadows/StochasticShadowsSampling.usf
+++ b/Engine/Shaders/Private/StochasticShadows/StochasticShadowsSampling.usf
@@ -104,28 +104,25 @@ float4 HistoryScreenPositionScaleBias;
 float4 HistoryUVMinMax;
 Texture2D<float> ShadowMaskSceneDepthHistory;
 
-groupshared uint SharedLightMask[LIGHT_SCENE_ID_MASK_SIZE];
+groupshared uint SharedLightMask[SHARED_LIGHT_MASK_SIZE];
 
 #if CANDIDATE_LIGHT_MASK
-	groupshared uint SharedCandidateLightMask[LIGHT_SCENE_ID_MASK_SIZE];
+	groupshared uint SharedCandidateLightMask[SHARED_LIGHT_MASK_SIZE];
 #endif
 
-/** 
- * Adds LightSceneId to the linear probing hash table in groupshared memory in order build a compacted list of unique light ids.
- */
-void MarkLightSceneId(uint LightSceneId)
+void MarkLocalLightIndex(uint LocalLightIndex)
 {
 	// Re-distribute nearby ids to different bytes in order to balance gather workload
-	uint DWORDIndex = LightSceneId % LIGHT_SCENE_ID_MASK_SIZE;
-	uint BitMask = 1u << (LightSceneId / LIGHT_SCENE_ID_MASK_SIZE);
+	uint DWORDIndex = LocalLightIndex % SHARED_LIGHT_MASK_SIZE;
+	uint BitMask = 1u << (LocalLightIndex / SHARED_LIGHT_MASK_SIZE);
 	InterlockedOr(SharedLightMask[DWORDIndex], BitMask);
 }
 
-void WriteShadowMaskTileToPageTable(uint LightSceneId, uint LocalLightIndex, uint2 ShadowMaskScreenCoord, uint ShadowMaskTileIndex)
+void WriteShadowMaskTileToPageTable(uint LocalLightIndex, uint2 ShadowMaskScreenCoord, uint ShadowMaskTileIndex)
 {
 	uint2 ShadowMaskPageTableCoord = GetShadowMaskPageTableCoord(LocalLightIndex, ShadowMaskScreenCoord);
 
-	RWShadowMaskTileData[ShadowMaskTileIndex] = PackShadowMaskTile(LightSceneId, LocalLightIndex, ShadowMaskScreenCoord);
+	RWShadowMaskTileData[ShadowMaskTileIndex] = PackShadowMaskTile(LocalLightIndex, ShadowMaskScreenCoord);
 
 	RWShadowMaskPageTable[ShadowMaskPageTableCoord] = ShadowMaskTileIndex;
 }
@@ -146,13 +143,13 @@ void GenerateSamplesCS(
 {
 	uint LinearThreadIndex = GroupThreadId.y * THREADGROUP_SIZE + GroupThreadId.x;
 
-	if (LinearThreadIndex < LIGHT_SCENE_ID_MASK_SIZE)
+	if (LinearThreadIndex < SHARED_LIGHT_MASK_SIZE)
 	{
 		SharedLightMask[LinearThreadIndex] = 0;
 	}
 
 #if CANDIDATE_LIGHT_MASK
-	if (LinearThreadIndex < LIGHT_SCENE_ID_MASK_SIZE)
+	if (LinearThreadIndex < SHARED_LIGHT_MASK_SIZE)
 	{
 		SharedCandidateLightMask[LinearThreadIndex] = 0;
 	}
@@ -186,7 +183,7 @@ void GenerateSamplesCS(
 				const uint LocalLightIndex = GetCulledLightDataGrid(CulledLightGridData.DataStartIndex + GridLightIndex);
 				const FLocalLightData LocalLightData = GetLocalLightDataNonStereo(LocalLightIndex);
 
-				if (LocalLightData.bStochasticShadowsSupported && LocalLightIndex < MAX_LIGHT_SCENE_ID)
+				if (LocalLightData.bStochasticShadowsSupported && LocalLightIndex < MAX_LOCAL_LIGHT_INDEX)
 				{
 					uint DWORDIndex = LocalLightIndex / 32;
 					uint BitMask = 1u << (LocalLightIndex % 32);
@@ -211,7 +208,7 @@ void GenerateSamplesCS(
 			uint SelectedLights[NUM_SAMPLES_PER_PIXEL_1D];
 			for (uint LightSampleIndex = 0; LightSampleIndex < NUM_SAMPLES_PER_PIXEL_1D; ++LightSampleIndex)
 			{
-				SelectedLights[LightSampleIndex] = INVALID_LIGHT_SCENE_ID;
+				SelectedLights[LightSampleIndex] = INVALID_LOCAL_LIGHT_INDEX;
 			}
 
 			const uint2 LightSampleOffsets[4] = 
@@ -308,7 +305,7 @@ void GenerateSamplesCS(
 					Print(Context, TEXT("NumLightsInGridCell: "));
 					Print(Context, NumLightsInGridCell);
 					Print(Context, TEXT("NumLocalLights: "));
-					if (NumLocalLights <= MAX_LIGHT_SCENE_ID)
+					if (NumLocalLights <= MAX_LOCAL_LIGHT_INDEX)
 					{
 						Print(Context, NumLocalLights);
 					}
@@ -328,7 +325,7 @@ void GenerateSamplesCS(
 				float Xi[NUM_SAMPLES_PER_PIXEL_1D] = LightIndexRandom;
 
 	#if CANDIDATE_LIGHT_MASK
-				for (uint MaskIndex = 0; MaskIndex < LIGHT_SCENE_ID_MASK_SIZE; ++MaskIndex)
+				for (uint MaskIndex = 0; MaskIndex < SHARED_LIGHT_MASK_SIZE; ++MaskIndex)
 				{
 					uint CandidateLightMask = SharedCandidateLightMask[MaskIndex];
 					while (CandidateLightMask != 0)
@@ -339,8 +336,8 @@ void GenerateSamplesCS(
 
 						const uint LocalLightIndex = MaskIndex * 32 + NextBitIndex;
 						const FLocalLightData LocalLightData = GetLocalLightDataNonStereo(LocalLightIndex);
-						const uint LightSceneId = LocalLightData.LightSceneId;
-						FDeferredLightData LightData = ConvertToDeferredLight(LocalLightData);
+						const FLightSceneData LightSceneData = GetLightSceneData(LocalLightData.LightSceneId);
+						FDeferredLightData LightData = ConvertToDeferredLight(LightSceneData);
 					
 						float SampleWeight = GetLocalLightTargetPDF(LightData, TranslatedWorldPosition, GBuffer, ScreenCoord);
 
@@ -359,10 +356,10 @@ void GenerateSamplesCS(
 
 								uint PrevLocalLightIndex = LocalLightData.PrevLocalLightIndex;
 
-								AccumulateHistoryShadowMask(LightSceneId, PrevLocalLightIndex, ShadowMaskHistoryScreenCoord00 + uint2(0, 0), ShadowMaskHistoryWeights.x, Accumulator);
-								AccumulateHistoryShadowMask(LightSceneId, PrevLocalLightIndex, ShadowMaskHistoryScreenCoord00 + uint2(1, 0), ShadowMaskHistoryWeights.y, Accumulator);
-								AccumulateHistoryShadowMask(LightSceneId, PrevLocalLightIndex, ShadowMaskHistoryScreenCoord00 + uint2(0, 1), ShadowMaskHistoryWeights.z, Accumulator);
-								AccumulateHistoryShadowMask(LightSceneId, PrevLocalLightIndex, ShadowMaskHistoryScreenCoord00 + uint2(1, 1), ShadowMaskHistoryWeights.w, Accumulator);
+								AccumulateHistoryShadowMask(PrevLocalLightIndex, ShadowMaskHistoryScreenCoord00 + uint2(0, 0), ShadowMaskHistoryWeights.x, Accumulator);
+								AccumulateHistoryShadowMask(PrevLocalLightIndex, ShadowMaskHistoryScreenCoord00 + uint2(1, 0), ShadowMaskHistoryWeights.y, Accumulator);
+								AccumulateHistoryShadowMask(PrevLocalLightIndex, ShadowMaskHistoryScreenCoord00 + uint2(0, 1), ShadowMaskHistoryWeights.z, Accumulator);
+								AccumulateHistoryShadowMask(PrevLocalLightIndex, ShadowMaskHistoryScreenCoord00 + uint2(1, 1), ShadowMaskHistoryWeights.w, Accumulator);
 
 								if (Accumulator.WeightSum > 0.0f)
 								{
@@ -379,7 +376,7 @@ void GenerateSamplesCS(
 						if (bDebug)
 						{
 							Newline(Context);
-							Print(Context, LightSceneId);
+							Print(Context, LocalLightData.LightSceneId);
 							Print(Context, SampleWeight);
 							Print(Context, SampleWeight * ShadowFactorEstimate);
 							Print(Context, ShadowFactorEstimate);
@@ -393,7 +390,7 @@ void GenerateSamplesCS(
 						if (SampleWeight > SamplingMinWeight)
 						{
 							// Allocate required shadow mask tiles for this frame
-							MarkLightSceneId(LocalLightIndex);
+							MarkLocalLightIndex(LocalLightIndex);
 
 							// Apply shadowing estimate from the last frame
 							if (HistoryShadowMask.IsValid())
@@ -416,7 +413,7 @@ void GenerateSamplesCS(
 									{
 										// Select this sample
 										Xi[LightSampleIndex] = (Xi[LightSampleIndex] - Tau) / (1.0f - Tau);
-										SelectedLights[LightSampleIndex] = LightSceneId;
+										SelectedLights[LightSampleIndex] = LocalLightIndex;
 									}
 
 									Xi[LightSampleIndex] = clamp(Xi[LightSampleIndex], 0, 0.9999f);
@@ -434,7 +431,8 @@ void GenerateSamplesCS(
 
 					if (LocalLightData.bStochasticShadowsSupported)
 					{
-						FDeferredLightData LightData = ConvertToDeferredLight(LocalLightData);
+						const FLightSceneData LightSceneData = GetLightSceneData(LocalLightData.LightSceneId);
+						FDeferredLightData LightData = ConvertToDeferredLight(LightSceneData);
 						float SampleWeight = GetLocalLightTargetPDF(LightData, TranslatedWorldPosition, GBuffer, ScreenCoord);
 
 						// Simulate tonemapping
@@ -452,10 +450,10 @@ void GenerateSamplesCS(
 
 								uint PrevLocalLightIndex = LocalLightData.PrevLocalLightIndex;
 
-								AccumulateHistoryShadowMask(LocalLightData.LightSceneId, PrevLocalLightIndex, ShadowMaskHistoryScreenCoord00 + uint2(0, 0), ShadowMaskHistoryWeights.x, Accumulator);
-								AccumulateHistoryShadowMask(LocalLightData.LightSceneId, PrevLocalLightIndex, ShadowMaskHistoryScreenCoord00 + uint2(1, 0), ShadowMaskHistoryWeights.y, Accumulator);
-								AccumulateHistoryShadowMask(LocalLightData.LightSceneId, PrevLocalLightIndex, ShadowMaskHistoryScreenCoord00 + uint2(0, 1), ShadowMaskHistoryWeights.z, Accumulator);
-								AccumulateHistoryShadowMask(LocalLightData.LightSceneId, PrevLocalLightIndex, ShadowMaskHistoryScreenCoord00 + uint2(1, 1), ShadowMaskHistoryWeights.w, Accumulator);
+								AccumulateHistoryShadowMask(PrevLocalLightIndex, ShadowMaskHistoryScreenCoord00 + uint2(0, 0), ShadowMaskHistoryWeights.x, Accumulator);
+								AccumulateHistoryShadowMask(PrevLocalLightIndex, ShadowMaskHistoryScreenCoord00 + uint2(1, 0), ShadowMaskHistoryWeights.y, Accumulator);
+								AccumulateHistoryShadowMask(PrevLocalLightIndex, ShadowMaskHistoryScreenCoord00 + uint2(0, 1), ShadowMaskHistoryWeights.z, Accumulator);
+								AccumulateHistoryShadowMask(PrevLocalLightIndex, ShadowMaskHistoryScreenCoord00 + uint2(1, 1), ShadowMaskHistoryWeights.w, Accumulator);
 
 								if (Accumulator.WeightSum > 0.0f)
 								{
@@ -486,7 +484,7 @@ void GenerateSamplesCS(
 						if (SampleWeight > SamplingMinWeight)
 						{
 							// Allocate required shadow mask tiles for this frame
-							MarkLightSceneId(LocalLightIndex);
+							MarkLocalLightIndex(LocalLightIndex);
 
 							// Apply shadowing estimate from the last frame
 							if (HistoryShadowMask.IsValid())
@@ -509,7 +507,7 @@ void GenerateSamplesCS(
 									{
 										// Select this sample
 										Xi[LightSampleIndex] = (Xi[LightSampleIndex] - Tau) / (1.0f - Tau);
-										SelectedLights[LightSampleIndex] = LocalLightData.LightSceneId;
+										SelectedLights[LightSampleIndex] = LocalLightIndex;
 									}
 
 									Xi[LightSampleIndex] = clamp(Xi[LightSampleIndex], 0, 0.9999f);
@@ -556,8 +554,8 @@ void GenerateSamplesCS(
 			for (uint LightSampleIndex = 0; LightSampleIndex < NUM_SAMPLES_PER_PIXEL_1D; ++LightSampleIndex)
 			{
 				FLightSample LightSample;
-				LightSample.LightSceneId = SelectedLights[LightSampleIndex];
-				LightSample.bCompleted = LightSample.LightSceneId == INVALID_LIGHT_SCENE_ID ? true : false;
+				LightSample.LocalLightIndex = SelectedLights[LightSampleIndex];
+				LightSample.bCompleted = LightSample.LocalLightIndex == INVALID_LOCAL_LIGHT_INDEX ? true : false;
 				LightSample.bVisible = true;
 				RWLightSamples[DownsampledScreenCoord * NumSamplesPerPixel + LightSampleOffsets[LightSampleIndex]] = PackLightSample(LightSample);
 			}
@@ -567,7 +565,7 @@ void GenerateSamplesCS(
 
 		// Write out allocated shadow mask tiles
 		// #sdl_todo: process part of the list per tread
-		if (LinearThreadIndex < LIGHT_SCENE_ID_MASK_SIZE)
+		if (LinearThreadIndex < SHARED_LIGHT_MASK_SIZE)
 		{
 			const uint LightIndex = LinearThreadIndex;
 			const uint BitMask = 1u << LightIndex;
@@ -581,9 +579,7 @@ void GenerateSamplesCS(
 				const uint NextBitMask = 1u << NextBitIndex;
 				LightMask ^= NextBitMask;
 
-				const uint LocalLightIndex = LinearThreadIndex + NextBitIndex * LIGHT_SCENE_ID_MASK_SIZE;
-				const FLocalLightData LocalLightData = GetLocalLightDataNonStereo(LocalLightIndex);
-				const uint LightIndex = LocalLightData.LightSceneId;
+				const uint LocalLightIndex = LinearThreadIndex + NextBitIndex * SHARED_LIGHT_MASK_SIZE;
 
 				// #sdl_todo: merge allocations through groupshared and do only one atomic per group
 				uint ShadowMaskTileIndex = 0;
@@ -591,7 +587,7 @@ void GenerateSamplesCS(
 
 				if (ShadowMaskTileIndex < MaxShadowMaskTiles)
 				{
-					WriteShadowMaskTileToPageTable(LightIndex, LocalLightIndex, ShadowMaskScreenCoord, ShadowMaskTileIndex);
+					WriteShadowMaskTileToPageTable(LocalLightIndex, ShadowMaskScreenCoord, ShadowMaskTileIndex);
 				}
 			}
 		}
diff --git a/Engine/Shaders/Private/StochasticShadows/StochasticShadowsShading.usf b/Engine/Shaders/Private/StochasticShadows/StochasticShadowsShading.usf
--- a/Engine/Shaders/Private/StochasticShadows/StochasticShadowsShading.usf
+++ b/Engine/Shaders/Private/StochasticShadows/StochasticShadowsShading.usf
@@ -92,7 +92,9 @@ void ShadeLightSamplesCS(
 				{
 					const uint PackedShadingTile = ShadingTileGrid[GetShadingTileGridOffset(ShadingTileCoord, IndexInCell)];
 					const FShadingTile ShadingTile = UnpackShadingTile(PackedShadingTile);
-					FDeferredLightData LightData = ConvertToDeferredLight(GetLightSceneData(ShadingTile.LightSceneId));
+					const FLocalLightData LocalLightData = GetLocalLightDataNonStereo(ShadingTile.LocalLightIndex);
+					const FLightSceneData LightSceneData = GetLightSceneData(LocalLightData.LightSceneId);
+					FDeferredLightData LightData = ConvertToDeferredLight(LightSceneData);
 
 					float ShadowFactor = 1.0f; // Initialize as unshadowed as tiles with index >= MaxShadingTiles are treated as fully unshadowed.
 
@@ -107,7 +109,7 @@ void ShadeLightSamplesCS(
 					if (bDebug)
 					{
 						Newline(Context);
-						Print(Context, ShadingTile.LightSceneId);
+						Print(Context, LocalLightData.LightSceneId);
 						Print(Context, ShadowFactor);
 					}
 					#endif
diff --git a/Engine/Shaders/Private/StochasticShadows/StochasticShadowsSpatial.usf b/Engine/Shaders/Private/StochasticShadows/StochasticShadowsSpatial.usf
--- a/Engine/Shaders/Private/StochasticShadows/StochasticShadowsSpatial.usf
+++ b/Engine/Shaders/Private/StochasticShadows/StochasticShadowsSpatial.usf
@@ -33,13 +33,13 @@ groupshared uint SharedShadingTileIndex[2][2];
 uint GetLitBitMask(uint2 ShadingTileCoord)
 {
 	const uint BitOffset = (ShadingTileCoord.x + ShadingTileCoord.y * 2);
-	return 0x1 << BitOffset;
+	return 0x1u << BitOffset;
 }
 
 uint GetShadowedBitMask(uint2 ShadingTileCoord)
 {
 	const uint BitOffset = (ShadingTileCoord.x + ShadingTileCoord.y * 2);
-	return 0x10 << BitOffset;
+	return 0x10u << BitOffset;
 }
 
 void SetSharedShadingTileMask(uint2 ShadingTileCoord, FShadowMask ShadowMask)
@@ -90,10 +90,9 @@ void ShadowMaskSpatialPassCS(
 	uint ShadowMaskTileIndex = GroupId.x;
 	uint2 LocalShadowMaskTileCoord = GroupThreadId.xy;
 
-	uint LightSceneId;
 	uint LocalLightIndex;
 	uint2 BaseShadowMaskScreenCoord;	
-	UnpackShadowMaskTile(ShadowMaskTileData[ShadowMaskTileIndex], LightSceneId, LocalLightIndex, BaseShadowMaskScreenCoord);
+	UnpackShadowMaskTile(ShadowMaskTileData[ShadowMaskTileIndex], LocalLightIndex, BaseShadowMaskScreenCoord);
 
 	if (all(GroupThreadId == 0))
 	{
@@ -268,7 +267,7 @@ void ShadowMaskSpatialPassCS(
 
 					if (ShadingTileIndexInCell < MaxShadingTilesPerGridCell)
 					{
-						RWShadingTileGrid[GetShadingTileGridOffset(ShadingTileCoord, ShadingTileIndexInCell)] = PackShadingTile(LightSceneId, ShadingTileIndex);
+						RWShadingTileGrid[GetShadingTileGridOffset(ShadingTileCoord, ShadingTileIndexInCell)] = PackShadingTile(LocalLightIndex, ShadingTileIndex);
 						SharedShadingTileIndex[ShadingTileOffset.x][ShadingTileOffset.y] = ShadingTileIndex;
 					}
 				}
diff --git a/Engine/Shaders/Private/StochasticShadows/StochasticShadowsTracing.usf b/Engine/Shaders/Private/StochasticShadows/StochasticShadowsTracing.usf
--- a/Engine/Shaders/Private/StochasticShadows/StochasticShadowsTracing.usf
+++ b/Engine/Shaders/Private/StochasticShadows/StochasticShadowsTracing.usf
@@ -115,7 +115,7 @@ void CompactLightSampleTracesCS(
 		// #sdl_todo: try to preserve z-order here for improved tracing coherency
 		FLightSample LightSample = UnpackLightSample(LightSamples[SampleCoord]);
 
-		if (LightSample.LightSceneId != INVALID_LIGHT_SCENE_ID && !LightSample.bCompleted)
+		if (LightSample.LocalLightIndex != INVALID_LOCAL_LIGHT_INDEX && !LightSample.bCompleted)
 		{
 			#if WAVE_OPS
 			{
@@ -228,7 +228,7 @@ void ScreenSpaceRayTraceLightSamplesCS(
 		float3 TranslatedWorldPosition = GetTranslatedWorldPositionFromScreenUV(ScreenUV, SceneDepth);
 
 		FLightSample LightSample = UnpackLightSample(RWLightSamples[SampleCoord]);
-		FLightSampleTrace LightSampleTrace = GetLightSampleTrace(TranslatedWorldPosition, LightSample.LightSceneId, SampleCoord);
+		FLightSampleTrace LightSampleTrace = GetLightSampleTrace(TranslatedWorldPosition, LightSample.LocalLightIndex, SampleCoord);
 
 		#if DEBUG_MODE
 			int2 DebugScreenCoord = View.CursorPosition.x >= 0 ? View.CursorPosition * View.ViewResolutionFraction : -1;
@@ -298,11 +298,13 @@ void ScreenSpaceRayTraceLightSamplesCS(
 		#if DEBUG_MODE
 		if (bDebug && DebugMode == DEBUG_MODE_VISUALIZE_TRACING)
 		{
-			FLightSceneData LightSceneData = GetLightSceneData(LightSample.LightSceneId);
+			const FLocalLightData LocalLightData = GetLocalLightDataNonStereo(LightSample.LocalLightIndex);
+			const FLightSceneData LightSceneData = GetLightSceneData(LocalLightData.LightSceneId);
+			const FDeferredLightData LightData = ConvertToDeferredLight(LightSceneData);
 
 			float3 RayStart = TranslatedWorldPosition;
 			float3 RayEnd = TranslatedWorldPosition + LightSampleTrace.Direction * min(LightSampleTrace.Distance, RayT);
-			float4 RayColor = float4(LightSceneData.Color.xyz / Luminance(LightSceneData.Color.xyz), 1.0f);
+			float4 RayColor = float4(LightData.Color.xyz / Luminance(LightData.Color.xyz), 1.0f);
 
 			AddLineTWS(Context, RayStart, RayEnd, RayColor, 1.0f);
 			AddCrossTWS(Context, RayEnd, 2.0f, RayColor);
@@ -333,7 +335,7 @@ void SoftwareRayTraceLightSamplesCS(
 
 		FLightSample LightSample = UnpackLightSample(RWLightSamples[SampleCoord]);
 		const float LightSampleRayMinT = LightSampleRayDistance[SampleCoord];
-		FLightSampleTrace LightSampleTrace = GetLightSampleTrace(TranslatedWorldPosition, LightSample.LightSceneId, SampleCoord);
+		FLightSampleTrace LightSampleTrace = GetLightSampleTrace(TranslatedWorldPosition, LightSample.LocalLightIndex, SampleCoord);
 
 		FGlobalSDFTraceInput TraceInput = SetupGlobalSDFTraceInput(TranslatedWorldPosition, LightSampleTrace.Direction, LightSampleRayMinT, LightSampleTrace.Distance, 1.0f, 1.0f);
 
diff --git a/Engine/Shaders/Private/StochasticShadows/StochasticShadowsTracing.ush b/Engine/Shaders/Private/StochasticShadows/StochasticShadowsTracing.ush
--- a/Engine/Shaders/Private/StochasticShadows/StochasticShadowsTracing.ush
+++ b/Engine/Shaders/Private/StochasticShadows/StochasticShadowsTracing.ush
@@ -5,6 +5,7 @@
 #include "../RayTracing/RayTracingCapsuleLight.ush"
 #include "../RayTracing/RayTracingSphereLight.ush"
 #include "../RayTracing/RayTracingRectLight.ush"
+#include "../LightGridCommon.ush"
 
 uint2 NumSamplesPerPixelDivideShift;
 
@@ -81,7 +82,7 @@ struct FLightSampleTrace
 /**
  * Compute shadow ray direction and distance for a given reservoir.
  */
-FLightSampleTrace GetLightSampleTrace(float3 TranslatedWorldPosition, uint LightSceneId, uint2 SampleCoord)
+FLightSampleTrace GetLightSampleTrace(float3 TranslatedWorldPosition, uint LocalLightIndex, uint2 SampleCoord)
 {
 	const float2 LightSampleRandom = BlueNoiseVec2(SampleCoord, StochasticShadowsStateFrameIndex);
 
@@ -89,11 +90,12 @@ FLightSampleTrace GetLightSampleTrace(float3 TranslatedWorldPosition, uint Light
 	LightSampleTrace.Direction = 0.0f;
 	LightSampleTrace.Distance = 0.0f;
 
-	if (LightSceneId != INVALID_LIGHT_SCENE_ID)
+	if (LocalLightIndex != INVALID_LOCAL_LIGHT_INDEX)
 	{
-		FLightSceneData LightSceneData = GetLightSceneData(LightSceneId);
-		FDeferredLightData LightData = ConvertToDeferredLight(LightSceneData);
-		FLightShaderParameters LightParameters = ConvertToLightShaderParameters(LightSceneData);
+		const FLocalLightData LocalLightData = GetLocalLightDataNonStereo(LocalLightIndex);
+		const FLightSceneData LightSceneData = GetLightSceneData(LocalLightData.LightSceneId);
+		const FDeferredLightData LightData = ConvertToDeferredLight(LightSceneData);
+		const FLightShaderParameters LightParameters = ConvertToLightShaderParameters(LightSceneData);
 
 		float3 Unused;
 		float Unused2;
diff --git a/Engine/Source/Runtime/Renderer/Private/StochasticShadows/StochasticShadows.cpp b/Engine/Source/Runtime/Renderer/Private/StochasticShadows/StochasticShadows.cpp
--- a/Engine/Source/Runtime/Renderer/Private/StochasticShadows/StochasticShadows.cpp
+++ b/Engine/Source/Runtime/Renderer/Private/StochasticShadows/StochasticShadows.cpp
@@ -140,7 +140,7 @@ namespace StochasticShadows
 	// must match values in StochasticShadows.ush
 	constexpr int32 TileSize = 8;
 	constexpr int32 ShadowMaskTileSize = 8;	// Stored downsampled
-	constexpr int32 MaxLightSceneIdXY = 16; // 16 * 16 = 256
+	constexpr int32 MaxLocalLightIndexXY = 16; // 16 * 16 = 256
 	constexpr int32 ShadowMaskAtlasSizeInTiles = 512;
 	constexpr uint32 InvalidShadowMaskTileIndex = 0xFFFFFFFF;
 	constexpr uint32 ShadingTileIndexUnshadowed = 0xFFFFF; // limited by PackShadingTile()
@@ -202,6 +202,11 @@ namespace StochasticShadows
 			&& RHISupportsWaveOperations(ShaderPlatform);
 	}
 
+	void ModifyCompilationEnvironment(EShaderPlatform Platform, FShaderCompilerEnvironment& OutEnvironment)
+	{
+		FForwardLightingParameters::ModifyCompilationEnvironment(Platform, OutEnvironment);
+	}
+
 	// Keep in sync with TILE_TYPE_* in shaders
 	enum class ETileType : uint8
 	{
@@ -323,7 +328,6 @@ class FGenerateSamplesCS : public FGlobalShader
 		SHADER_PARAMETER_RDG_BUFFER_UAV(RWStructuredBuffer<uint2>, RWShadowMaskTileData)
 		SHADER_PARAMETER_RDG_TEXTURE_UAV(RWTexture2D<uint>, RWShadowMaskPageTable)
 		SHADER_PARAMETER_RDG_TEXTURE_UAV(RWTexture2D<uint>, RWLightSamples)
-		SHADER_PARAMETER_RDG_UNIFORM_BUFFER(FForwardLightData, ForwardLightData)
 		SHADER_PARAMETER_RDG_BUFFER_SRV(StructuredBuffer<uint>, DownsampledTileAllocator)
 		SHADER_PARAMETER_RDG_BUFFER_SRV(StructuredBuffer<uint>, DownsampledTileData)
 		SHADER_PARAMETER_RDG_TEXTURE(Texture2D<uint>, ShadowMaskPageTableHistory)
@@ -373,7 +377,7 @@ class FGenerateSamplesCS : public FGlobalShader
 	static void ModifyCompilationEnvironment(const FGlobalShaderPermutationParameters& Parameters, FShaderCompilerEnvironment& OutEnvironment)
 	{
 		FGlobalShader::ModifyCompilationEnvironment(Parameters, OutEnvironment);
-		FForwardLightingParameters::ModifyCompilationEnvironment(Parameters.Platform, OutEnvironment);
+		StochasticShadows::ModifyCompilationEnvironment(Parameters.Platform, OutEnvironment);
 		OutEnvironment.SetDefine(TEXT("THREADGROUP_SIZE"), GetGroupSize());
 	}
 };
@@ -428,7 +432,6 @@ class FCompositeShadowMaskTracesCS : public FGlobalShader
 		SHADER_PARAMETER_RDG_TEXTURE(Texture2D<uint>, ShadowMaskPageTableHistory)
 		SHADER_PARAMETER_RDG_TEXTURE(Texture2D<uint>, ShadowMaskAtlasHistory)
 		SHADER_PARAMETER_RDG_TEXTURE(Texture2D<uint>, LightSamples)
-		SHADER_PARAMETER_RDG_UNIFORM_BUFFER(FForwardLightData, ForwardLightData)
 	END_SHADER_PARAMETER_STRUCT()
 
 	static int32 GetGroupSize()
@@ -448,7 +451,7 @@ class FCompositeShadowMaskTracesCS : public FGlobalShader
 	static void ModifyCompilationEnvironment(const FGlobalShaderPermutationParameters& Parameters, FShaderCompilerEnvironment& OutEnvironment)
 	{
 		FGlobalShader::ModifyCompilationEnvironment(Parameters, OutEnvironment);
-		FForwardLightingParameters::ModifyCompilationEnvironment(Parameters.Platform, OutEnvironment);
+		StochasticShadows::ModifyCompilationEnvironment(Parameters.Platform, OutEnvironment);
 		OutEnvironment.SetDefine(TEXT("THREADGROUP_SIZE"), GetGroupSize());
 	}
 };
@@ -534,6 +537,7 @@ class FShadeLightSamplesCS : public FGlobalShader
 	static void ModifyCompilationEnvironment(const FGlobalShaderPermutationParameters& Parameters, FShaderCompilerEnvironment& OutEnvironment)
 	{
 		FGlobalShader::ModifyCompilationEnvironment(Parameters, OutEnvironment);
+		StochasticShadows::ModifyCompilationEnvironment(Parameters.Platform, OutEnvironment);
 		OutEnvironment.SetDefine(TEXT("THREADGROUP_SIZE"), GetGroupSize());
 	}
 };
@@ -590,7 +594,7 @@ void FDeferredShadingSceneRenderer::RenderStochasticShadows(FRDGBuilder& GraphBu
 	const FIntPoint ShadowMaskBufferSize = FIntPoint::DivideAndRoundUp(SceneTextures.Config.Extent, ShadowMaskDownsampleFactor);
 	const FIntPoint ShadowMaskViewSize = FIntPoint::DivideAndRoundUp(View.ViewRect.Size(), ShadowMaskDownsampleFactor);
 	const FIntPoint ShadowMaskPageTablePerLightSize = FIntPoint::DivideAndRoundUp(ShadowMaskBufferSize, StochasticShadows::ShadowMaskTileSize);
-	const FIntPoint ShadowMaskPageTableSize = ShadowMaskPageTablePerLightSize * StochasticShadows::MaxLightSceneIdXY;
+	const FIntPoint ShadowMaskPageTableSize = ShadowMaskPageTablePerLightSize * StochasticShadows::MaxLocalLightIndexXY;
 	const FIntPoint ShadowMaskTileAtlasSize = StochasticShadows::ShadowMaskTileSize * StochasticShadows::ShadowMaskAtlasSizeInTiles;
 	const int32 MaxShadowMaskTiles = StochasticShadows::ShadowMaskAtlasSizeInTiles * StochasticShadows::ShadowMaskAtlasSizeInTiles;
 	check(MaxShadowMaskTiles == (ShadowMaskTileAtlasSize.X * ShadowMaskTileAtlasSize.Y) / (StochasticShadows::ShadowMaskTileSize * StochasticShadows::ShadowMaskTileSize));
@@ -683,6 +687,7 @@ void FDeferredShadingSceneRenderer::RenderStochasticShadows(FRDGBuilder& GraphBu
 		StochasticShadowsParameters.SceneTextures = GetSceneTextureParameters(GraphBuilder, SceneTextures.UniformBuffer);
 		StochasticShadowsParameters.SceneTexturesStruct = SceneTextures.UniformBuffer;
 		StochasticShadowsParameters.Substrate = Substrate::BindSubstrateGlobalUniformParameters(View);
+		StochasticShadowsParameters.ForwardLightData = View.ForwardLightingResources.ForwardLightUniformBuffer;
 		StochasticShadowsParameters.BlueNoise = BlueNoiseUniformBuffer;
 		StochasticShadowsParameters.DownsampledViewSize = DownsampledViewSize;
 		StochasticShadowsParameters.SampleViewSize = SampleViewSize;
@@ -847,7 +852,6 @@ void FDeferredShadingSceneRenderer::RenderStochasticShadows(FRDGBuilder& GraphBu
 			PassParameters->RWLightSamples = LightSamplesUAV;
 			PassParameters->DownsampledTileAllocator = GraphBuilder.CreateSRV(DownsampledTileAllocator);
 			PassParameters->DownsampledTileData = GraphBuilder.CreateSRV(DownsampledTileData);
-			PassParameters->ForwardLightData = View.ForwardLightingResources.ForwardLightUniformBuffer;
 			PassParameters->ShadowMaskPageTableHistory = ShadowMaskPageTableHistory;
 			PassParameters->ShadowMaskAtlasHistory = ShadowMaskAtlasHistory;
 			PassParameters->ShadowMaskTileAtlas = ShadowMaskTileAtlas;
@@ -962,7 +966,6 @@ void FDeferredShadingSceneRenderer::RenderStochasticShadows(FRDGBuilder& GraphBu
 		PassParameters->ShadowMaskPageTableHistory = ShadowMaskPageTableHistory;
 		PassParameters->ShadowMaskAtlasHistory = ShadowMaskAtlasHistory;
 		PassParameters->LightSamples = LightSamples;
-		PassParameters->ForwardLightData = View.ForwardLightingResources.ForwardLightUniformBuffer;
 
 		FCompositeShadowMaskTracesCS::FPermutationDomain PermutationVector;
 		PermutationVector.Set<FCompositeShadowMaskTracesCS::FTemporalAccumulation>(bTemporal);
diff --git a/Engine/Source/Runtime/Renderer/Private/StochasticShadows/StochasticShadowsInternal.h b/Engine/Source/Runtime/Renderer/Private/StochasticShadows/StochasticShadowsInternal.h
--- a/Engine/Source/Runtime/Renderer/Private/StochasticShadows/StochasticShadowsInternal.h
+++ b/Engine/Source/Runtime/Renderer/Private/StochasticShadows/StochasticShadowsInternal.h
@@ -13,6 +13,7 @@ BEGIN_SHADER_PARAMETER_STRUCT(FStochasticShadowsParameters, )
 	SHADER_PARAMETER_RDG_UNIFORM_BUFFER(FSceneUniformParameters, Scene)
 	SHADER_PARAMETER_RDG_UNIFORM_BUFFER(FSceneTextureUniformParameters, SceneTexturesStruct)
 	SHADER_PARAMETER_RDG_UNIFORM_BUFFER(FSubstrateGlobalUniformParameters, Substrate)
+	SHADER_PARAMETER_RDG_UNIFORM_BUFFER(FForwardLightData, ForwardLightData)
 	SHADER_PARAMETER_STRUCT_REF(FBlueNoise, BlueNoise)
 	SHADER_PARAMETER(FIntPoint, SampleViewSize)
 	SHADER_PARAMETER(FIntPoint, DownsampledViewSize)
@@ -54,4 +55,6 @@ namespace StochasticShadows
 	bool ShouldCompileShaders(const FGlobalShaderPermutationParameters& Parameters);
 	bool UseWaveOps(EShaderPlatform ShaderPlatform);
 	int32 GetDebugMode();
+
+	void ModifyCompilationEnvironment(EShaderPlatform Platform, FShaderCompilerEnvironment& OutEnvironment);
 };
\ No newline at end of file
diff --git a/Engine/Source/Runtime/Renderer/Private/StochasticShadows/StochasticShadowsRayTracing.cpp b/Engine/Source/Runtime/Renderer/Private/StochasticShadows/StochasticShadowsRayTracing.cpp
--- a/Engine/Source/Runtime/Renderer/Private/StochasticShadows/StochasticShadowsRayTracing.cpp
+++ b/Engine/Source/Runtime/Renderer/Private/StochasticShadows/StochasticShadowsRayTracing.cpp
@@ -4,6 +4,7 @@
 #include "StochasticShadowsInternal.h"
 #include "Lumen/LumenTracingUtils.h"
 #include "Lumen/LumenHardwareRayTracingCommon.h"
+#include "BasePassRendering.h"
 
 static TAutoConsoleVariable<int32> CVarStochasticShadowsScreenTraces(
 	TEXT("r.StochasticShadows.ScreenTraces"),
@@ -242,6 +243,7 @@ class FHardwareRayTraceLightSamples : public FLumenHardwareRayTracingShaderBase
 	static void ModifyCompilationEnvironment(const FGlobalShaderPermutationParameters& Parameters, Lumen::ERayTracingShaderDispatchType ShaderDispatchType, FShaderCompilerEnvironment& OutEnvironment)
 	{
 		FLumenHardwareRayTracingShaderBase::ModifyCompilationEnvironment(Parameters, ShaderDispatchType, Lumen::ESurfaceCacheSampling::AlwaysResidentPagesWithoutFeedback, OutEnvironment);
+		StochasticShadows::ModifyCompilationEnvironment(Parameters.Platform, OutEnvironment);
 	}
 
 	static ERayTracingPayloadType GetRayTracingPayloadType(const int32 PermutationId)
@@ -285,6 +287,7 @@ class FSoftwareRayTraceLightSamplesCS : public FGlobalShader
 	FORCENOINLINE static void ModifyCompilationEnvironment(const FGlobalShaderPermutationParameters& Parameters, FShaderCompilerEnvironment& OutEnvironment)
 	{
 		FGlobalShader::ModifyCompilationEnvironment(Parameters, OutEnvironment);
+		StochasticShadows::ModifyCompilationEnvironment(Parameters.Platform, OutEnvironment);
 		OutEnvironment.CompilerFlags.Add(CFLAG_Wave32);
 		OutEnvironment.SetDefine(TEXT("THREADGROUP_SIZE"), GetGroupSize());
 
@@ -328,6 +331,7 @@ class FScreenSpaceRayTraceLightSamplesCS : public FGlobalShader
 	FORCENOINLINE static void ModifyCompilationEnvironment(const FGlobalShaderPermutationParameters& Parameters, FShaderCompilerEnvironment& OutEnvironment)
 	{
 		FGlobalShader::ModifyCompilationEnvironment(Parameters, OutEnvironment);
+		StochasticShadows::ModifyCompilationEnvironment(Parameters.Platform, OutEnvironment);
 		OutEnvironment.CompilerFlags.Add(CFLAG_Wave32);
 		OutEnvironment.SetDefine(TEXT("THREADGROUP_SIZE"), GetGroupSize());
 	}
