Heterogeneous Volumes: Fix multi-volume compositing when the LightingCache is disabled.

diff --git a/Engine/Shaders/Private/HeterogeneousVolumes/HeterogeneousVolumesLiveShadingPipeline.usf b/Engine/Shaders/Private/HeterogeneousVolumes/HeterogeneousVolumesLiveShadingPipeline.usf
--- a/Engine/Shaders/Private/HeterogeneousVolumes/HeterogeneousVolumesLiveShadingPipeline.usf
+++ b/Engine/Shaders/Private/HeterogeneousVolumes/HeterogeneousVolumesLiveShadingPipeline.usf
@@ -356,9 +356,12 @@ void RenderSingleScatteringWithLiveShadingCS(
 		);
 
 		// Output..
-		float3 Opacity = bApplyEmissionAndTransmittance ? 1.0 - Transmittance : 0.0;
 		RWLightingTexture[PixelCoord].rgb += Radiance * View.PreExposure;
-		RWLightingTexture[PixelCoord].a = Luminance(Opacity);
+		if (bApplyEmissionAndTransmittance)
+		{
+			float3 Opacity = 1.0 - Transmittance;
+			RWLightingTexture[PixelCoord].a = Luminance(Opacity);
+		}
 
 		// Debug output
 		//RWDebugOutputBuffer[LinearIndex] = DebugOutput;
diff --git a/Engine/Source/Runtime/Renderer/Private/HeterogeneousVolumes/HeterogeneousVolumesLiveShadingPipeline.cpp b/Engine/Source/Runtime/Renderer/Private/HeterogeneousVolumes/HeterogeneousVolumesLiveShadingPipeline.cpp
--- a/Engine/Source/Runtime/Renderer/Private/HeterogeneousVolumes/HeterogeneousVolumesLiveShadingPipeline.cpp
+++ b/Engine/Source/Runtime/Renderer/Private/HeterogeneousVolumes/HeterogeneousVolumesLiveShadingPipeline.cpp
@@ -702,7 +702,7 @@ static void RenderWithTransmittanceVolumePipeline(
 	int32 NumPasses = FMath::Max(LightSceneInfoCompact.Num(), 1);
 	for (int32 PassIndex = 0; PassIndex < NumPasses; ++PassIndex)
 	{
-		bool bApplyEmissionAndTransmittance = PassIndex == 0;
+		bool bApplyEmissionAndTransmittance = (PassIndex == (NumPasses - 1));
 		bool bApplyDirectLighting = !LightSceneInfoCompact.IsEmpty();
 		bool bApplyShadowTransmittance = false;
 
