LFAtlas - using sincos in shader.

diff --git a/Engine/Shaders/Private/LightFunctionAtlas/LightFunctionAtlasRender.usf b/Engine/Shaders/Private/LightFunctionAtlas/LightFunctionAtlasRender.usf
--- a/Engine/Shaders/Private/LightFunctionAtlas/LightFunctionAtlasRender.usf
+++ b/Engine/Shaders/Private/LightFunctionAtlas/LightFunctionAtlasRender.usf
@@ -34,9 +34,13 @@ void Main(
 	// Generate a local space light vector that is at least correct for point lights
 	float Theta = (UVs.x * (2 * PI)) - PI;
 	float Phi = UVs.y * PI;
-	float X = cos(Theta) * sin(Phi);
-	float Y = sin(Theta) * sin(Phi);
-	float Z = cos(Phi);
+	float SinTheta, CosTheta;
+	float SinPhi, CosPhi;
+	sincos(Theta, SinTheta, CosTheta);
+	sincos(Phi, SinPhi, CosPhi);
+	float X = CosTheta * SinPhi;
+	float Y = SinTheta * SinPhi;
+	float Z = CosPhi;
 	LightVector = normalize(float3(X, Y, Z)); 
 
 	LightFunction = GetLightFunctionColor(LightVector, TranslatedWorldPosition.xyz, UVs.xy);
