Added missing copyright boilerplate to shaders.

diff --git a/Engine/Shaders/Private/LightFunctionAtlas/LightFunctionAtlasCommon.usf b/Engine/Shaders/Private/LightFunctionAtlas/LightFunctionAtlasCommon.usf
--- a/Engine/Shaders/Private/LightFunctionAtlas/LightFunctionAtlasCommon.usf
+++ b/Engine/Shaders/Private/LightFunctionAtlas/LightFunctionAtlasCommon.usf
@@ -1,3 +1,5 @@
+// Copyright Epic Games, Inc. All Rights Reserved.
+
 /*=============================================================================
 	LightFunctionAtlasCommon.usf
 =============================================================================*/
diff --git a/Engine/Shaders/Private/LightFunctionAtlas/LightFunctionAtlasRender.usf b/Engine/Shaders/Private/LightFunctionAtlas/LightFunctionAtlasRender.usf
--- a/Engine/Shaders/Private/LightFunctionAtlas/LightFunctionAtlasRender.usf
+++ b/Engine/Shaders/Private/LightFunctionAtlas/LightFunctionAtlasRender.usf
@@ -1,3 +1,5 @@
+// Copyright Epic Games, Inc. All Rights Reserved.
+
 /*=============================================================================
 	RenderLightFunctionAtlas.usf.usf
 =============================================================================*/
