Fix substrate shader compilation issue with forward lighting.

diff --git a/Engine/Shaders/Private/Substrate/SubstrateForwardLighting.ush b/Engine/Shaders/Private/Substrate/SubstrateForwardLighting.ush
--- a/Engine/Shaders/Private/Substrate/SubstrateForwardLighting.ush
+++ b/Engine/Shaders/Private/Substrate/SubstrateForwardLighting.ush
@@ -474,7 +474,7 @@ float3 SubstrateForwardLighting(
 				LOOP
 				for (uint GridLightListIndex = 0; GridLightListIndex < NumLightsInGridCell; GridLightListIndex++)
 				{
-					const FLocalLightData LocalLight = GetLocalLightData(CulledLightsGrid.DataStartIndex + LocalLightListIndex, EyeIndex);
+					const FLocalLightData LocalLight = GetLocalLightData(CulledLightsGrid.DataStartIndex + GridLightListIndex, EyeIndex);
 
 					half4 PreviewShadowMapChannelMask = 1;
 					uint  LocalLightChannelMask = LIGHTING_CHANNEL_MASK;
