Fix Substrate shader compilation issue with hair strands.

diff --git a/Engine/Shaders/Private/HairStrands/HairStrandsEnvironmentAO.usf b/Engine/Shaders/Private/HairStrands/HairStrandsEnvironmentAO.usf
--- a/Engine/Shaders/Private/HairStrands/HairStrandsEnvironmentAO.usf
+++ b/Engine/Shaders/Private/HairStrands/HairStrandsEnvironmentAO.usf
@@ -2,17 +2,10 @@
   
 #include "../Common.ush"
 #include "../MonteCarlo.ush"
-#include "../DeferredShadingCommon.ush"  
-#include "../SceneTextureParameters.ush"
+#include "../DeferredShadingCommon.ush"
 #include "../PositionReconstructionCommon.ush"
-#include "../HairShadingCommon.ush"
-#include "../HairBsdf.ush"
-#include "HairStrandsCommon.ush"
-#include "HairStrandsVisibilityCommon.ush"
-#include "HairStrandsVoxelPageCommon.ush"
 #include "HairStrandsEnvironmentLightingCommon.ush"
-
-#include "../ShaderPrint.ush"
+#include "../Substrate/Substrate.ush"
 
 #define VOXEL_TRAVERSAL_TYPE VOXEL_TRAVERSAL_LINEAR_MIPMAP
 #define VOXEL_TRAVERSAL_DEBUG 0 
@@ -117,8 +110,7 @@ float TraceAO(float3 TranslatedWorldPosition, float3 WorldNormal, float2 SvPosit
 
 void MainPS(
 	in float4 SvPosition : SV_Position,
-	out float4 OutAO : SV_Target0
-)
+	out float4 OutAO : SV_Target0)
 {
 	const bool bHalfRes = true; // Output_bHalfRes > 0;
 	float2 BufferUV = SvPositionToBufferUV(SvPosition);
@@ -127,20 +119,30 @@ void MainPS(
 		BufferUV = SvPosition.xy * Output_InvResolution;
 	}
 
-	const FGBufferData GBuffer = GetGBufferDataFromSceneTextures(BufferUV);
-
-	OutAO = 1;
+	const uint2 PixelCoord = floor(SvPosition.xy);
+#if SUBSTRATE_ENABLED
+	FSubstrateAddressing SubstrateAddressing = GetSubstratePixelDataByteOffset(PixelCoord, uint2(View.BufferSizeAndInvSize.xy), Substrate.MaxBytesPerPixel);
+	FSubstratePixelHeader SubstratePixelHeader = UnpackSubstrateHeaderIn(Substrate.MaterialTextureArray, SubstrateAddressing, Substrate.TopLayerTexture);
+	const FSubstrateTopLayerData TopLayerData = SubstrateUnpackTopLayerData(Substrate.TopLayerTexture.Load(uint3(PixelCoord, 0)));
+	const bool bIsValid = SubstratePixelHeader.ClosureCount > 0;
+	const float3 WorldNormal = TopLayerData.WorldNormal;
 
+#else
 	// Trace AO for all pixel which are partially covered or not covered at all.
 	// Fully covered pixel are marked as unlit
-	const uint2 PixelCoord = floor(SvPosition.xy);
+	const FGBufferData GBuffer = GetGBufferDataFromSceneTextures(BufferUV);
 	const bool bIsValid = GBuffer.ShadingModelID != SHADINGMODELID_UNLIT;
+	const float3 WorldNormal = GBuffer.WorldNormal;
+#endif
 
+	OutAO = 1;
 	if (bIsValid)
 	{
-		const float3 TranslatedWorldPosition = ReconstructTranslatedWorldPositionFromDepth(BufferUV, GBuffer.Depth);
-		const float PixelRadius = GBuffer.Depth * VirtualVoxel.HairCoveragePixelRadiusAtDepth1;
-		OutAO = TraceAO(TranslatedWorldPosition, GBuffer.WorldNormal, SvPosition.xy, PixelRadius);
+		const float DeviceZ = SampleDeviceZFromSceneTextures(BufferUV);
+		const float SceneDepth = ConvertFromDeviceZ(DeviceZ);
+		const float3 TranslatedWorldPosition = ReconstructTranslatedWorldPositionFromDepth(BufferUV, SceneDepth);
+		const float PixelRadius = SceneDepth * VirtualVoxel.HairCoveragePixelRadiusAtDepth1;
+		OutAO = TraceAO(TranslatedWorldPosition, WorldNormal, SvPosition.xy, PixelRadius);
 	}
 	else
 	{
diff --git a/Engine/Source/Runtime/Renderer/Private/HairStrands/HairStrandsEnvironment.cpp b/Engine/Source/Runtime/Renderer/Private/HairStrands/HairStrandsEnvironment.cpp
--- a/Engine/Source/Runtime/Renderer/Private/HairStrands/HairStrandsEnvironment.cpp
+++ b/Engine/Source/Runtime/Renderer/Private/HairStrands/HairStrandsEnvironment.cpp
@@ -122,6 +122,7 @@ class FHairEnvironmentAO : public FGlobalShader
 		SHADER_PARAMETER_STRUCT_REF(FViewUniformShaderParameters, ViewUniformBuffer)
 		SHADER_PARAMETER_RDG_UNIFORM_BUFFER(FVirtualVoxelParameters, VirtualVoxel)
 		SHADER_PARAMETER_RDG_UNIFORM_BUFFER(FHairStrandsViewUniformParameters, HairStrands)
+		SHADER_PARAMETER_RDG_UNIFORM_BUFFER(FSubstrateGlobalUniformParameters, Substrate)
 		SHADER_PARAMETER_STRUCT_INCLUDE(ShaderPrint::FShaderParameters, ShaderPrintParameters)
 
 		RENDER_TARGET_BINDING_SLOTS()
@@ -152,6 +153,7 @@ static void AddHairStrandsEnvironmentAOPass(
 
 	PassParameters->ViewUniformBuffer = View.ViewUniformBuffer;
 	PassParameters->HairStrands = HairStrands::BindHairStrandsViewUniformParameters(View);
+	PassParameters->Substrate = Substrate::BindSubstrateGlobalUniformParameters(View);
 	const FFinalPostProcessSettings& Settings = View.FinalPostProcessSettings;
 	PassParameters->AO_Power = Settings.AmbientOcclusionPower;
 	PassParameters->AO_Intensity = Settings.AmbientOcclusionIntensity;
