Fixed UnpackRectLightData() to handle BarnCosAngle = 1.

diff --git a/Engine/Shaders/Private/LightData.ush b/Engine/Shaders/Private/LightData.ush
--- a/Engine/Shaders/Private/LightData.ush
+++ b/Engine/Shaders/Private/LightData.ush
@@ -44,7 +44,7 @@ FRectLightData UnpackRectLightData(uint In0, uint In1, uint In2)
 	Out.AtlasData.AtlasUVOffset	= half2(f16tof32(In0 & 0xFFFF), f16tof32(In0 >> 16));
 	Out.AtlasData.AtlasUVScale	= half2(f16tof32(In1 & 0xFFFF), f16tof32(In1 >> 16));
 	Out.BarnLength				= f16tof32(In2 & 0xFFFF);			// 16 bits
-	Out.BarnCosAngle			= ((In2 >> 16) & 0x3FF) / 1024.f;	// 10 bits
+	Out.BarnCosAngle			= UnpackUnorm10(In2 >> 16);			// 10 bits
 	Out.AtlasData.AtlasMaxLevel = half(In2 >> 26);					//  6 bits
 	return Out;
 }
diff --git a/Engine/Shaders/Private/PackUnpack.ush b/Engine/Shaders/Private/PackUnpack.ush
--- a/Engine/Shaders/Private/PackUnpack.ush
+++ b/Engine/Shaders/Private/PackUnpack.ush
@@ -164,6 +164,20 @@ uint4 UnpackUintToUint4(uint In)
 	return Out;
 }
 
+//
+//	float <-> unorm
+//
+
+uint PackUnorm10(float Value)
+{
+	return uint(saturate(Value) * 1023.0f);
+}
+
+float UnpackUnorm10(uint In)
+{
+	return float(In & 0x3FF) * (1.0f / 1023.0f);
+}
+
 //
 //	float <-> f10
 //
