Fix missing include in HairStrandsEnvironmentAO.usf

#prelight skip

diff --git a/Engine/Shaders/Private/HairStrands/HairStrandsEnvironmentAO.usf b/Engine/Shaders/Private/HairStrands/HairStrandsEnvironmentAO.usf
--- a/Engine/Shaders/Private/HairStrands/HairStrandsEnvironmentAO.usf
+++ b/Engine/Shaders/Private/HairStrands/HairStrandsEnvironmentAO.usf
@@ -4,6 +4,8 @@
 #include "../MonteCarlo.ush"
 #include "../DeferredShadingCommon.ush"
 #include "../PositionReconstructionCommon.ush"
+#include "../HairShadingCommon.ush"
+#include "../HairBsdf.ush"
 #include "HairStrandsEnvironmentLightingCommon.ush"
 #include "../Substrate/Substrate.ush"
 
