LFAtlas - added color atlas mode. As a project setting requiring shader compilation.

diff --git a/Engine/Shaders/Private/DeferredLightPixelShaders.usf b/Engine/Shaders/Private/DeferredLightPixelShaders.usf
--- a/Engine/Shaders/Private/DeferredLightPixelShaders.usf
+++ b/Engine/Shaders/Private/DeferredLightPixelShaders.usf
@@ -55,9 +55,7 @@
 #include "VirtualShadowMaps/VirtualShadowMapMaskBitsCommon.ush"
 #endif
 
-#if USE_LIGHT_FUNCTION_ATLAS
 #include "LightFunctionAtlas/LightFunctionAtlasCommon.usf"
-#endif
 
 // Remap light shape permutation into light type
 #if LIGHT_SOURCE_SHAPE == 0
@@ -342,9 +340,7 @@ void DeferredLightPixelMain(
 #if USE_LIGHT_FUNCTION_ATLAS
 		if (bUseLightFunctionAtlas)
 		{
-			Radiance *= GetLocalLightFunctionCommon(
-				DerivedParams.TranslatedWorldPosition,
-				LightData.LightFunctionAtlasLightIndex);
+			Radiance.rgb *= GetLocalLightFunctionCommon(DerivedParams.TranslatedWorldPosition, LightData.LightFunctionAtlasLightIndex);
 		}
 #endif
 
@@ -454,16 +450,15 @@ void DeferredLightPixelMain(
 
 	const float LocalCoverage = From8bitCoverage(Sample.Coverage8bit);
 
-#if USE_LIGHT_FUNCTION_ATLAS
+	FLightFunctionColor LightFunctionColor = 1.0f;
+#if	USE_LIGHT_FUNCTION_ATLAS
 	if (bUseLightFunctionAtlas)
 	{
-		Attenuation *= GetLocalLightFunctionCommon(
-			TranslatedWorldPosition,
-			LightData.LightFunctionAtlasLightIndex);
+		LightFunctionColor = GetLocalLightFunctionCommon(TranslatedWorldPosition, LightData.LightFunctionAtlasLightIndex);
 	}
 #endif
 
-	OutColor.rgb = AtmosphericTransmittance * Radiance.xyz * Attenuation * LocalCoverage * GetExposure();
+	OutColor.rgb = AtmosphericTransmittance * Radiance.xyz * LightFunctionColor * Attenuation * LocalCoverage * GetExposure();
 	OutColor.a = LocalCoverage;
 }
 #endif
diff --git a/Engine/Shaders/Private/LightFunctionAtlas/LightFunctionAtlasCommon.usf b/Engine/Shaders/Private/LightFunctionAtlas/LightFunctionAtlasCommon.usf
--- a/Engine/Shaders/Private/LightFunctionAtlas/LightFunctionAtlasCommon.usf
+++ b/Engine/Shaders/Private/LightFunctionAtlas/LightFunctionAtlasCommon.usf
@@ -11,7 +11,13 @@
 
 // In order to use function from that file, one must use FLightFunctionAtlasGlobalParameters
 
-float GetLocalLightFunctionCommon(float3 TranslatedWorldPosition, uint LightFunctionAtlasLightIndex)
+#if COLORED_LIGHT_FUNCTION_ATLAS
+#define FLightFunctionColor float3
+#else
+#define FLightFunctionColor float
+#endif
+
+FLightFunctionColor GetLocalLightFunctionCommon(float3 TranslatedWorldPosition, uint LightFunctionAtlasLightIndex)
 {
 	if (LightFunctionAtlasLightIndex == 0)
 	{
@@ -62,17 +68,21 @@ float GetLocalLightFunctionCommon(float3 TranslatedWorldPosition, uint LightFunc
 	// Slot min uv and size already account for half texel offsets
 	float2 AtlasUV = Slot_MinUV + saturate(LightFunctionUV) * LightFunctionAtlas.Slot_UVSize.xx;
 
-	float LightFunctionValue = Texture2DSampleLevel(LightFunctionAtlas.LightFunctionAtlasTexture, LightFunctionAtlas.LightFunctionAtlasSampler, AtlasUV, 0).x;
-	LightFunctionValue = LightFunctionValue * LightFunctionValue; // Simple Gamma decompression
+#ifdef COLORED_LIGHT_FUNCTION_ATLAS
+	FLightFunctionColor LightFunctionColor = Texture2DSampleLevel(LightFunctionAtlas.LightFunctionAtlasTexture, LightFunctionAtlas.LightFunctionAtlasSampler, AtlasUV, 0).rgb;
+#else
+	FLightFunctionColor LightFunctionColor = Texture2DSampleLevel(LightFunctionAtlas.LightFunctionAtlasTexture, LightFunctionAtlas.LightFunctionAtlasSampler, AtlasUV, 0).x;
+#endif
+	LightFunctionColor = LightFunctionColor * LightFunctionColor; // Simple Gamma decompression
 
 	// Calculate radial view distance for stable fading
 	float ViewDistance = length(PrimaryView.TranslatedWorldCameraOrigin - TranslatedWorldPosition);
 	float DistanceFadeAlpha = saturate((FadeDistance - ViewDistance) / (FadeDistance * .2f));
 	// Fade to disabled based on FadeDistance
-	LightFunctionValue = lerp(DisabledBrightness, LightFunctionValue, DistanceFadeAlpha);
+	LightFunctionColor = lerp(DisabledBrightness, LightFunctionColor, DistanceFadeAlpha);
 	// Fade to disabled based on ShadowFadeFraction
 	const float ShadowFadeFraction = 1.0f; // Not setup
-	LightFunctionValue = LightFunctionValue; // lerp(DisabledBrightness, LightFunctionValue, ShadowFadeFraction);
+	LightFunctionColor = LightFunctionColor; // lerp(DisabledBrightness, LightFunctionColor, ShadowFadeFraction);
 
-	return LightFunctionValue;
+	return LightFunctionColor;
 }
\ No newline at end of file
diff --git a/Engine/Shaders/Private/LightFunctionAtlas/LightFunctionAtlasRender.usf b/Engine/Shaders/Private/LightFunctionAtlas/LightFunctionAtlasRender.usf
--- a/Engine/Shaders/Private/LightFunctionAtlas/LightFunctionAtlasRender.usf
+++ b/Engine/Shaders/Private/LightFunctionAtlas/LightFunctionAtlasRender.usf
@@ -68,10 +68,16 @@ void Main(
 
 #endif
 
+#if COLORED_LIGHT_FUNCTION_ATLAS
 
-	float GreyScale = dot(LightFunction, .3333f).x;
+	LightFunction = sqrt(max(0.0f, LightFunction));	// Simple Gamma compression
+	OutColor = float4(LightFunction, 0.0f);
 
-	GreyScale = sqrt(max(0.0f, GreyScale)); // Simple Gamma compression
+#else
 
-	OutColor = GreyScale;
+	float GreyScale = dot(LightFunction, .3333f).x;
+	GreyScale = sqrt(max(0.0f, GreyScale));						// Simple Gamma compression
+	OutColor = float4(GreyScale, 0.0f, 0.0f, 0.0f);
+
+#endif
 }
diff --git a/Engine/Shaders/Private/VolumetricFog.usf b/Engine/Shaders/Private/VolumetricFog.usf
--- a/Engine/Shaders/Private/VolumetricFog.usf
+++ b/Engine/Shaders/Private/VolumetricFog.usf
@@ -326,27 +326,7 @@ int VirtualShadowMapId;
 #define HISTORY_MISS_SUPER_SAMPLE_COUNT 1
 #endif
 
-#if USE_LIGHT_FUNCTION_ATLAS || USE_LIGHT_FUNCTION
 #include "LightFunctionAtlas/LightFunctionAtlasCommon.usf"
-#endif
-
-#if	USE_LIGHT_FUNCTION // Used when injecting light with with shadow, unbatched
-
-/** Matrix can be either WorldToShadow or WorldToLight depending on light type. */
-float4x4 LocalLightFunctionMatrix;
-float4 LightFunctionAtlasTileMinMaxUvBound;
-
-Texture2D LightFunctionAtlasTexture;
-SamplerState LightFunctionAtlasSampler;
-
-float GetLocalLightFunction(float3 TranslatedWorldPosition, FDeferredLightData LightData)
-{
-	return GetLocalLightFunctionCommon(
-		TranslatedWorldPosition, 
-		LightData.LightFunctionAtlasLightIndex);
-}
-
-#endif // USE_LIGHT_FUNCTION
 
 Texture2D<float> ConservativeDepthTexture;
 float2 PrevConservativeDepthTextureSize;
@@ -480,11 +460,12 @@ float4 InjectShadowedLocalLightCommon(uint3 GridCoordinate)
 			}
 #endif
 	
+			FLightFunctionColor LightFunctionColor = 1.0f;
 #if	USE_LIGHT_FUNCTION
-			ShadowFactor *= GetLocalLightFunction(TranslatedWorldPosition, LightData);
+			LightFunctionColor = GetLocalLightFunctionCommon(TranslatedWorldPosition, LightData.LightFunctionAtlasLightIndex);
 #endif
 
-			OutScattering.rgb += LightColor * (PhaseFunction(PhaseG, dot(L, -CameraVector)) * CombinedAttenuation * ShadowFactor * VolumetricScatteringIntensity);
+			OutScattering.rgb += LightColor * LightFunctionColor * (PhaseFunction(PhaseG, dot(L, -CameraVector)) * CombinedAttenuation * ShadowFactor * VolumetricScatteringIntensity);
 				
 			// To debug culling
 			//OutScattering.rgb += DeferredLightUniforms.Color * .0000001f;
@@ -670,7 +651,7 @@ SamplerState DirectionalLightLightFunctionSampler;
 uint DirectionalApplyLightFunctionFromAtlas;
 uint DirectionalLightFunctionAtlasLightIndex;
 
-float GetDirectionalLightFunction(float3 TranslatedWorldPosition) // directional light
+FLightFunctionColor GetDirectionalLightFunction(float3 TranslatedWorldPosition) // directional light
 {
 #if USE_LIGHT_FUNCTION_ATLAS  // Used when injecting lights from the light grid
 	if (DirectionalApplyLightFunctionFromAtlas > 0)
@@ -802,7 +783,7 @@ void LightScatteringCS(
 #endif
 			}
 
-			ShadowFactor *= GetDirectionalLightFunction(TranslatedWorldPosition);
+			FLightFunctionColor LightFunctionColor = GetDirectionalLightFunction(TranslatedWorldPosition);
 
 			float3 DirectionalLightColor = InDirectionalLightColor;
 
@@ -812,7 +793,7 @@ void LightScatteringCS(
 				DirectionalLightColor = VolumetricFog.HeightFogDirectionalLightInscatteringColor * Luminance(InDirectionalLightColor);
 			}
 
-			LightScattering += DirectionalLightColor * ShadowFactor * PhaseFunction(PhaseG, dot(InDirectionalLightDirection, -CameraVector));
+			LightScattering += DirectionalLightColor * LightFunctionColor * ShadowFactor * PhaseFunction(PhaseG, dot(InDirectionalLightDirection, -CameraVector));
 		}
 
 		FTwoBandSHVector RotatedHGZonalHarmonic;
@@ -912,13 +893,12 @@ void LightScatteringCS(
 
 				float CombinedAttenuation = Lighting * LightMask;
 
-			#if USE_LIGHT_FUNCTION_ATLAS
-				CombinedAttenuation *= GetLocalLightFunctionCommon(
-					TranslatedWorldPosition,
-					LightData.LightFunctionAtlasLightIndex);
+				FLightFunctionColor LightFunctionColor = 1.0f;
+			#if	USE_LIGHT_FUNCTION_ATLAS
+				LightFunctionColor = GetLocalLightFunctionCommon(TranslatedWorldPosition, LightData.LightFunctionAtlasLightIndex);
 			#endif
 
-				LightScattering += LightData.Color * (PhaseFunction(PhaseG, dot(L, -CameraVector)) * CombinedAttenuation * VolumetricScatteringIntensity);
+				LightScattering += LightData.Color * LightFunctionColor * (PhaseFunction(PhaseG, dot(L, -CameraVector)) * CombinedAttenuation * VolumetricScatteringIntensity);
 
 				// To debug culling
 				//LightScattering += LocalLight.LightColorAndFalloffExponent.xyz * .0000001f;
diff --git a/Engine/Source/Runtime/Engine/Classes/Engine/RendererSettings.h b/Engine/Source/Runtime/Engine/Classes/Engine/RendererSettings.h
--- a/Engine/Source/Runtime/Engine/Classes/Engine/RendererSettings.h
+++ b/Engine/Source/Runtime/Engine/Classes/Engine/RendererSettings.h
@@ -232,6 +232,16 @@ namespace EDefaultBackBufferPixelFormat
 	ENGINE_API EDefaultBackBufferPixelFormat::Type FromInt(int32 InDefaultBackBufferPixelFormat);
 }
 
+UENUM()
+namespace ELightFunctionAtlasPixelFormat
+{
+	enum Type : int
+	{
+		LFAPF_R8 = 0				UMETA(DisplayName = "8 bits Gray Scale"),
+		LFAPF_R8G8B8 = 1			UMETA(DisplayName = "8 bits RGB  Color")
+	};
+}
+
 /**
  * Enumerates supported shader compression formats.
  */
@@ -874,6 +884,15 @@ class URendererSettings : public UDeveloperSettings
 		ConfigRestartRequired = true))
 		uint32 bSupportCloudShadowOnForwardLitTranslucent : 1;
 
+	/**
+	"Select the format of the light function atlas texture."
+	*/
+	UPROPERTY(config, EditAnywhere, Category = Optimizations, meta = (
+		ConsoleVariable = "r.LightFunctionAtlas.Format", DisplayName = "Light Function Atlas Format",
+		ToolTip = "Select the format of the light function atlas texture.",
+		ConfigRestartRequired = true))
+		TEnumAsByte<ELightFunctionAtlasPixelFormat::Type> LightFunctionAtlasPixelFormat;
+
 	/**
 	"Enable rect light evaluation on translucent materials when using the Forward Shading mode."
 	*/
diff --git a/Engine/Source/Runtime/Engine/Private/ShaderCompiler/ShaderCompiler.cpp b/Engine/Source/Runtime/Engine/Private/ShaderCompiler/ShaderCompiler.cpp
--- a/Engine/Source/Runtime/Engine/Private/ShaderCompiler/ShaderCompiler.cpp
+++ b/Engine/Source/Runtime/Engine/Private/ShaderCompiler/ShaderCompiler.cpp
@@ -8302,6 +8302,8 @@ void GlobalBeginCompileShader(
 		SET_SHADER_DEFINE(Input.Environment, LA_NORMALMAPS, 1);
 	}
 
+	SET_SHADER_DEFINE(Input.Environment, COLORED_LIGHT_FUNCTION_ATLAS, GetLightFunctionAtlasFormat() > 0 ? 1 : 0);
+
 	SET_SHADER_DEFINE(Input.Environment, PLATFORM_SUPPORTS_SHADER_ROOT_CONSTANTS, RHISupportsShaderRootConstants(EShaderPlatform(Target.Platform)) ? 1 : 0);
 	SET_SHADER_DEFINE(Input.Environment, PLATFORM_SUPPORTS_SHADER_BUNDLE_DISPATCH, RHISupportsShaderBundleDispatch(EShaderPlatform(Target.Platform)) ? 1 : 0);
 	SET_SHADER_DEFINE(Input.Environment, PLATFORM_SUPPORTS_RENDERTARGET_WRITE_MASK, RHISupportsRenderTargetWriteMask(EShaderPlatform(Target.Platform)) ? 1 : 0);
diff --git a/Engine/Source/Runtime/RenderCore/Private/RenderUtils.cpp b/Engine/Source/Runtime/RenderCore/Private/RenderUtils.cpp
--- a/Engine/Source/Runtime/RenderCore/Private/RenderUtils.cpp
+++ b/Engine/Source/Runtime/RenderCore/Private/RenderUtils.cpp
@@ -1663,6 +1663,17 @@ bool UseSplineMeshSceneResources(const FStaticShaderPlatform Platform)
 	return false;
 }
 
+static TAutoConsoleVariable<int32> CVarLightFunctionAtlasFormat(
+	TEXT("r.LightFunctionAtlas.Format"),
+	0,
+	TEXT("0: grey scale in [0,1]. 1: colored in [0,1]"),
+	ECVF_ReadOnly | ECVF_RenderThreadSafe);
+
+int32 GetLightFunctionAtlasFormat()
+{
+	return CVarLightFunctionAtlasFormat.GetValueOnAnyThread();
+}
+
 ///////////////////////////////////////////////////////////////////////////////////////////////////
 // Substrate settings interface
 
diff --git a/Engine/Source/Runtime/RenderCore/Private/Shader.cpp b/Engine/Source/Runtime/RenderCore/Private/Shader.cpp
--- a/Engine/Source/Runtime/RenderCore/Private/Shader.cpp
+++ b/Engine/Source/Runtime/RenderCore/Private/Shader.cpp
@@ -2073,6 +2073,14 @@ void ShaderMapAppendKeyString(EShaderPlatform Platform, FString& KeyString)
 		}
 	}
 
+	{
+		int32 LightFunctionAtlasFormat = GetLightFunctionAtlasFormat();
+		if (LightFunctionAtlasFormat > 0)
+		{
+			KeyString += FString::Printf(TEXT("_LFAC%u"), LightFunctionAtlasFormat);
+		}
+	}
+
 	{
 		static const auto CVar = IConsoleManager::Get().FindTConsoleVariableDataInt(TEXT("r.Material.EnergyConservation"));
 		if (CVar && CVar->GetValueOnAnyThread() > 0)
diff --git a/Engine/Source/Runtime/RenderCore/Public/RenderUtils.h b/Engine/Source/Runtime/RenderCore/Public/RenderUtils.h
--- a/Engine/Source/Runtime/RenderCore/Public/RenderUtils.h
+++ b/Engine/Source/Runtime/RenderCore/Public/RenderUtils.h
@@ -549,6 +549,8 @@ namespace Substrate
 	RENDERCORE_API uint32 GetShadingQuality(EShaderPlatform InPlatform);
 }
 
+RENDERCORE_API int32 GetLightFunctionAtlasFormat();
+
 // LuminanceMax is the amount of light that will cause the sensor to saturate at EV100.
 //  See also https://en.wikipedia.org/wiki/Film_speed and https://en.wikipedia.org/wiki/Exposure_value for more info.
 FORCEINLINE float EV100ToLuminance(float LuminanceMax, float EV100)
diff --git a/Engine/Source/Runtime/Renderer/Private/LightFunctionAtlas.cpp b/Engine/Source/Runtime/Renderer/Private/LightFunctionAtlas.cpp
--- a/Engine/Source/Runtime/Renderer/Private/LightFunctionAtlas.cpp
+++ b/Engine/Source/Runtime/Renderer/Private/LightFunctionAtlas.cpp
@@ -17,6 +17,7 @@
 #include "CanvasItem.h"
 #include "CanvasTypes.h"
 #include "Containers/HashTable.h"
+#include "RenderUtils.h"
 
 DECLARE_GPU_STAT(LightFunctionAtlasGeneration);
 
@@ -581,9 +582,11 @@ void FLightFunctionAtlas::AllocateTexture2DAtlas(FRDGBuilder& GraphBuilder)
 	uint32 AtlasResolution = AtlasSlotResolution * AtlasEdgeSize;
 	const uint32 MipCount = 1;
 
+	int32 LightFunctionAtlasFormat = GetLightFunctionAtlasFormat();
+
 	RDGAtlasTexture2D = GraphBuilder.CreateTexture(FRDGTextureDesc::Create2D(
 		FIntPoint(AtlasResolution, AtlasResolution),
-		PF_R8,
+		LightFunctionAtlasFormat == 0 ? PF_R8 : PF_R8G8B8A8,
 		FClearValueBinding::Black,
 		ETextureCreateFlags::UAV | ETextureCreateFlags::ShaderResource | ETextureCreateFlags::RenderTargetable,
 		MipCount),
diff --git a/Engine/Source/Runtime/Renderer/Private/LightFunctionAtlas.h b/Engine/Source/Runtime/Renderer/Private/LightFunctionAtlas.h
--- a/Engine/Source/Runtime/Renderer/Private/LightFunctionAtlas.h
+++ b/Engine/Source/Runtime/Renderer/Private/LightFunctionAtlas.h
@@ -100,7 +100,7 @@ struct FLightFunctionSlot
 // Using a constant buffer also workaround the fact that we would otherwise need another SRV in forward shaders. 
 // The light atlas texture itself already use 1 extra SRV. We could an extra SRV and have LightInfoDataXXX be in a buffer that scale with amount of light in the scene.
 BEGIN_GLOBAL_SHADER_PARAMETER_STRUCT(FLightFunctionAtlasGlobalParameters, )
-	SHADER_PARAMETER_RDG_TEXTURE(Texture2D, LightFunctionAtlasTexture)
+	SHADER_PARAMETER_RDG_TEXTURE(Texture2D<float4>, LightFunctionAtlasTexture)
 	SHADER_PARAMETER_SAMPLER(SamplerState, LightFunctionAtlasSampler)
 	SHADER_PARAMETER(float, Slot_UVSize)
 	SHADER_PARAMETER_ARRAY(FMatrix44f, LightInfoDataMatrix, [LIGHT_FUNCTION_ATLAS_MAX_LIGHT_COUNT])		// Light data
