Removed clamp in MaterialExpressionBlackBody. Materials require the BlackBody node to give somewhat reasonable (even if not correct) values below 1000 Kelvin too.

#rnx

diff --git a/Engine/Shaders/Private/MaterialTemplate.ush b/Engine/Shaders/Private/MaterialTemplate.ush
--- a/Engine/Shaders/Private/MaterialTemplate.ush
+++ b/Engine/Shaders/Private/MaterialTemplate.ush
@@ -1519,7 +1519,8 @@ float MaterialExpressionDepthOfFieldFunction(float SceneDepth, int FunctionValue
 // TODO convert to LUT
 float3 MaterialExpressionBlackBody( float Temp )
 {
-	Temp = clamp(Temp, 1000.0f, 15000.0f);
+	// Technically this approximation to the Planckian locus is only valid within the 1000-15000 temperature range, 
+	// but we need to support temperatures outside this range too; especially the range between 0 and 1000.
 
 	float u = ( 0.860117757f + 1.54118254e-4f * Temp + 1.28641212e-7f * Temp*Temp ) / ( 1.0f + 8.42420235e-4f * Temp + 7.08145163e-7f * Temp*Temp );
 	float v = ( 0.317398726f + 4.22806245e-5f * Temp + 4.20481691e-8f * Temp*Temp ) / ( 1.0f - 2.89741816e-5f * Temp + 1.61456053e-7f * Temp*Temp );
