LFAtlas - fixed missing water separated luminance not influence by light function.

diff --git a/Engine/Shaders/Private/BasePassPixelShader.usf b/Engine/Shaders/Private/BasePassPixelShader.usf
--- a/Engine/Shaders/Private/BasePassPixelShader.usf
+++ b/Engine/Shaders/Private/BasePassPixelShader.usf
@@ -1483,11 +1483,11 @@ void FPixelShaderInOut_MainPS(
 			Specular, SingleLayerWater.DistortionParams,
 			SunIlluminance, WaterDiffuseIndirectIlluminance, EnvBrdf,
 			CameraIsUnderWater, WaterVisibility, EyeIndex,
-			bSeparateWaterMainDirLightLuminance, SeparatedWaterMainDirLightLuminance).Luminance
-	#if USE_LIGHT_FUNCTION_ATLAS
-			* GetLocalLightFunctionCommon(SvPositionToResolvedTranslatedWorld(In.SvPosition), GetDirectionalLightData(EyeIndex).LightFunctionAtlasLightIndex)
-	#endif
-		;
+			bSeparateWaterMainDirLightLuminance, SeparatedWaterMainDirLightLuminance
+		#if USE_LIGHT_FUNCTION_ATLAS
+			, GetLocalLightFunctionCommon(SvPositionToResolvedTranslatedWorld(In.SvPosition), GetDirectionalLightData(EyeIndex).LightFunctionAtlasLightIndex)
+		#endif
+			).Luminance;
 #endif
 	}
 #endif // MATERIAL_SHADINGMODEL_SINGLELAYERWATER
diff --git a/Engine/Shaders/Private/SingleLayerWaterShading.ush b/Engine/Shaders/Private/SingleLayerWaterShading.ush
--- a/Engine/Shaders/Private/SingleLayerWaterShading.ush
+++ b/Engine/Shaders/Private/SingleLayerWaterShading.ush
@@ -45,7 +45,8 @@ WaterVolumeLightingOutput EvaluateWaterVolumeLighting(
 	float Specular, const float4 DistortionParams,
 	float3 SunIlluminance, float3 AmbiantIlluminance, float3 EnvBrdf,
 	bool CameraIsUnderWater, float WaterVisibility, uint EyeIndex,
-	bool bSeparateMainDirLight, inout float3 SeparatedWaterMainDirLightScatteredLuminance)
+	bool bSeparateMainDirLight, inout float3 SeparatedWaterMainDirLightScatteredLuminance,
+	in float3 LightFunctionColor = 1.0f)
 {
 	WaterVolumeLightingOutput Output;
 	Output.Luminance = 0.0f;
@@ -192,13 +193,13 @@ WaterVolumeLightingOutput EvaluateWaterVolumeLighting(
 
 		const float3 OpticalDepth = ExtinctionCoeff * BehindWaterDeltaDepth;
 		float3 Transmittance = exp(-OpticalDepth);
-		float3 IncomingLuminance = (AmbScattLuminance + SunScattLuminance * DirectionalLightShadow * MainDirLightFactor);
+		float3 IncomingLuminance = (AmbScattLuminance + SunScattLuminance * LightFunctionColor * DirectionalLightShadow * MainDirLightFactor);
 		// The SafeScatteringAmount should be in 0~1, add saturate to fix an artifact caused by half precision on mobile.
 		float3 SafeScatteringAmount = saturate(ScatteringCoeff * (1.0f - Transmittance) / ExtinctionCoeffSafe);
 		float3 ScatteredLuminance = IncomingLuminance * SafeScatteringAmount;
 
 		// Compute separated direcitonal light scattering if requested
-		float3 SeparatedIncomingLuminance = SunScattLuminance * DirectionalLightShadow * SeparatedMainDirLightFactor;
+		float3 SeparatedIncomingLuminance = SunScattLuminance * LightFunctionColor * DirectionalLightShadow * SeparatedMainDirLightFactor;
 		// Apply Fresnel and water visibility to separated main dir light luminance.
 		SeparatedWaterMainDirLightScatteredLuminance += SeparatedIncomingLuminance * SafeScatteringAmount * (CameraIsUnderWater ? 1.0 : (1.0 - EnvBrdf)) * WaterVisibility;
 
