Fixed sky light illuminance meter t owork correctly with exposure.

diff --git a/Engine/Shaders/Private/ReflectionEnvironmentShaders.usf b/Engine/Shaders/Private/ReflectionEnvironmentShaders.usf
--- a/Engine/Shaders/Private/ReflectionEnvironmentShaders.usf
+++ b/Engine/Shaders/Private/ReflectionEnvironmentShaders.usf
@@ -1054,26 +1054,26 @@ void VisualizeIlluminanceMeterPS(
 		float2 UV = (SVPos.xy - float2(IlluMapPreviewOrigin)) / float2(IlluMapPreviewSize);
 		float4 DebugVis = SampleCubeAsLatLon(SourceCubemapTexture, SourceCubemapSampler, UV);
 		DebugVis.rgb *= 1.0 / DebugVis.a;	// This is a way to remove the texel dArea factor to show what luminance is integrated as a unction of NoL.
-		OutColor = float4(DebugVis.rgb, 0.0);
+		OutColor = float4(DebugVis.rgb * ResolvedView.PreExposure, 0.0);
 	}
 	else if (all(uint2(SVPos.xy) >= SkyMapPreviewOrigin) && all(uint2(SVPos.xy) <= (SkyMapPreviewOrigin + SkyMapPreviewSize)))
 	{
 		float2 UV = (SVPos.xy - float2(SkyMapPreviewOrigin)) / float2(SkyMapPreviewSize);
 		float4 DebugVis = SampleCubeAsLatLon(SkyLightCubeTexture, SkyLightCubeSampler, UV);
-		OutColor = float4(DebugVis.rgb, 0.0);
+		OutColor = float4(DebugVis.rgb * ResolvedView.PreExposure, 0.0);
 	}
 	else if (all(uint2(SVPos.xy) >= SphereIlluMapPreviewOrigin) && all(uint2(SVPos.xy) <= (SphereIlluMapPreviewOrigin + SphereIlluMapPreviewSize)))
 	{
 		float2 UV = (SVPos.xy - float2(SphereIlluMapPreviewOrigin)) / float2(SphereIlluMapPreviewSize);
 		float4 DebugVis = SampleCubeAsSphere(SourceCubemapTexture, SourceCubemapSampler, UV);
 		DebugVis.rgb *= 1.0 / DebugVis.a;	// This is a way to remove the texel dArea factor to show what luminance is integrated as a unction of NoL.
-		OutColor = float4(DebugVis.rgb, 0.0);
+		OutColor = float4(DebugVis.rgb * ResolvedView.PreExposure, 0.0);
 	}
 	else if (all(uint2(SVPos.xy) >= SphereSkyMapPreviewOrigin) && all(uint2(SVPos.xy) <= (SphereSkyMapPreviewOrigin + SphereSkyMapPreviewSize)))
 	{
 		float2 UV = (SVPos.xy - float2(SphereSkyMapPreviewOrigin)) / float2(SphereSkyMapPreviewSize);
 		float4 DebugVis = SampleCubeAsSphere(SkyLightCubeTexture, SkyLightCubeSampler, UV);
-		OutColor = float4(DebugVis.rgb, 0.0);
+		OutColor = float4(DebugVis.rgb * ResolvedView.PreExposure, 0.0);
 	}
 	else
 	{
diff --git a/Engine/Source/Runtime/Renderer/Private/IlluminanceMeter.cpp b/Engine/Source/Runtime/Renderer/Private/IlluminanceMeter.cpp
--- a/Engine/Source/Runtime/Renderer/Private/IlluminanceMeter.cpp
+++ b/Engine/Source/Runtime/Renderer/Private/IlluminanceMeter.cpp
@@ -110,7 +110,7 @@ void FScene::ProcessAndRenderIlluminanceMeter(FRDGBuilder& GraphBuilder, TArrayV
 
 	// 0- Get source cubemap and allocate cube map transient resources
 
-	FTextureRHIRef SkyLightTextureRHI = GBlackTextureCube->TextureRHI;;
+	FTextureRHIRef SkyLightTextureRHI = GBlackTextureCube->TextureRHI;
 	if (SkyLight->bRealTimeCaptureEnabled && ConvolvedSkyRenderTargetReadyIndex >= 0)
 	{
 		// Cannot blend with this capture mode as of today.
