Add support for material compiler auto generated function Texture2DArraySample_Decal()

diff --git a/Engine/Shaders/Private/Common.ush b/Engine/Shaders/Private/Common.ush
--- a/Engine/Shaders/Private/Common.ush
+++ b/Engine/Shaders/Private/Common.ush
@@ -434,6 +434,14 @@ MaterialFloat4 Texture3DSample_Decal(Texture3D Tex, SamplerState Sampler, float3
 	return Texture3DSample(Tex, Sampler, UV);
 #endif
 }
+MaterialFloat4 Texture2DArraySample_Decal(Texture2DArray Tex, SamplerState Sampler, float3 UV)
+{
+#if METAL_PROFILE || COMPILER_GLSL_ES3_1
+	return Texture2DArraySampleLevel(Tex, Sampler, UV, 0);
+#else
+	return Texture2DArraySample(Tex, Sampler, UV);
+#endif
+}
 MaterialFloat4 TextureCubeSample_Decal(TextureCube Tex, SamplerState Sampler, float3 UV)
 {
 #if METAL_PROFILE || COMPILER_GLSL_ES3_1
