LightFunctions: Fix shader warning with the colored light function atlas code

The type definition was using if and the usage was using ifdef which led to a (harmless) shader warning for the float case.

diff --git a/Engine/Shaders/Private/LightFunctionAtlas/LightFunctionAtlasCommon.usf b/Engine/Shaders/Private/LightFunctionAtlas/LightFunctionAtlasCommon.usf
--- a/Engine/Shaders/Private/LightFunctionAtlas/LightFunctionAtlasCommon.usf
+++ b/Engine/Shaders/Private/LightFunctionAtlas/LightFunctionAtlasCommon.usf
@@ -74,7 +74,7 @@ FLightFunctionColor GetLocalLightFunctionCommon(float3 TranslatedWorldPosition,
 	// Slot min uv and size already account for half texel offsets
 	float2 AtlasUV = Slot_MinUV + saturate(LightFunctionUV) * LightFunctionAtlas.Slot_UVSize.xx;
 
-#ifdef COLORED_LIGHT_FUNCTION_ATLAS
+#if COLORED_LIGHT_FUNCTION_ATLAS
 	FLightFunctionColor LightFunctionColor = Texture2DSampleLevel(LightFunctionAtlas.LightFunctionAtlasTexture, LightFunctionAtlas.LightFunctionAtlasSampler, AtlasUV, 0).rgb;
 #else
 	FLightFunctionColor LightFunctionColor = Texture2DSampleLevel(LightFunctionAtlas.LightFunctionAtlasTexture, LightFunctionAtlas.LightFunctionAtlasSampler, AtlasUV, 0).x;
