PathTracing: Tweak the default ray bias to avoid issues in larger scenes

diff --git a/Engine/Shaders/Private/PathTracing/PathTracingCore.ush b/Engine/Shaders/Private/PathTracing/PathTracingCore.ush
--- a/Engine/Shaders/Private/PathTracing/PathTracingCore.ush
+++ b/Engine/Shaders/Private/PathTracing/PathTracingCore.ush
@@ -925,7 +925,7 @@ void ApplyRayBias(inout FRayDesc Ray, float HitT, float3 Normal)
 	const float RefValue = max(max(abs(Ray.Origin.x), abs(Ray.Origin.y)), max(abs(Ray.Origin.z), HitT));
 	const uint UlpRadius = 16; // number of floating point ulps to skip around an intersection
 	const float RelativeOffset = asfloat(asuint(RefValue) + UlpRadius) - RefValue;
-	const float BaseOffset = 0.001; // 0.01mm (avoid precision issues very close to the origin where ulps become tiny)
+	const float BaseOffset = 0.00390625; // 2^-8cm (roughly 39micro meters) (This avoids precision issues very close to the origin where the ulps become too small)
 	Ray.Origin += max(BaseOffset, RelativeOffset) * Normal;
 }
 
