Unpack LightFunctionAtlasLightIndex from FLightSceneData::LightTypeAndShadowMapChannelMaskPacked

diff --git a/Engine/Shaders/Private/LightData.ush b/Engine/Shaders/Private/LightData.ush
--- a/Engine/Shaders/Private/LightData.ush
+++ b/Engine/Shaders/Private/LightData.ush
@@ -225,16 +225,26 @@ bool UnpackCastShadow(uint InShadowMapChannelMaskPacked)
 	return ((InShadowMapChannelMaskPacked >> 18) & 0x1) != 0;
 }
 
+bool UnpackHasLightFunction(uint InShadowMapChannelMaskPacked)
+{
+	return ((InShadowMapChannelMaskPacked >> 19) & 0x1) != 0;
+}
+
 bool UnpackHasLightFunction(FLocalLightData In)
 {
 	const uint LightTypeAndPackedShadowMapChannelMask = asuint(In.LightDirectionAndShadowMask.w);
-	return ((LightTypeAndPackedShadowMapChannelMask >> 19) & 0x1) != 0;
+	return UnpackHasLightFunction(LightTypeAndPackedShadowMapChannelMask);
+}
+
+uint UnpackLightFunctionAtlasIndex(uint InShadowMapChannelMaskPacked)
+{
+	return ((InShadowMapChannelMaskPacked >> 20) & 0xFF);
 }
 
 uint UnpackLightFunctionAtlasIndex(FLocalLightData In)
 {
 	const uint LightTypeAndPackedShadowMapChannelMask = asuint(In.LightDirectionAndShadowMask.w);
-	return ((LightTypeAndPackedShadowMapChannelMask >> 20) & 0xFF);
+	return UnpackLightFunctionAtlasIndex(LightTypeAndPackedShadowMapChannelMask);
 }
 
 float UnpackVolumetricScatteringIntensity(FLocalLightData In)
@@ -482,7 +492,7 @@ FDeferredLightData ConvertToDeferredLight(FLightSceneData In)
 	Out.RectLightData.AtlasData.AtlasUVScale = In.RectLightAtlasUVScale;
 
 	Out.IESAtlasIndex					= In.IESAtlasIndex;
-	Out.LightFunctionAtlasLightIndex	= 0;
+	Out.LightFunctionAtlasLightIndex	= UnpackLightFunctionAtlasIndex(In.LightTypeAndShadowMapChannelMaskPacked);
 
 	// Shadow
 	// * Static shadowing uses ShadowMapChannel, 
@@ -490,7 +500,7 @@ FDeferredLightData ConvertToDeferredLight(FLightSceneData In)
 	// Data Layout on [10:0] with:
 	// * ShadowMapChannelMask	= [3:0]
 	// * PreviewShadowMapChannel= [7:4] 
-	// * LightincChannelMask    = [8:10] 
+	// * LightingChannelMask    = [8:10] 
 	Out.ShadowedBits				= (In.LightTypeAndShadowMapChannelMaskPacked & 0xFF) != 0 ? 1 : 0;
 	Out.ShadowMapChannelMask		= UnpackShadowMapChannelMask(In.LightTypeAndShadowMapChannelMaskPacked);
 	return Out;
