Fixes a source of ghosting on large pre-exposure changes with r.TSR.ShadingRejection.Flickering

diff --git a/Engine/Shaders/Private/TemporalSuperResolution/TSRDecimateHistory.usf b/Engine/Shaders/Private/TemporalSuperResolution/TSRDecimateHistory.usf
--- a/Engine/Shaders/Private/TemporalSuperResolution/TSRDecimateHistory.usf
+++ b/Engine/Shaders/Private/TemporalSuperResolution/TSRDecimateHistory.usf
@@ -63,6 +63,7 @@ FScreenTransform ScreenPosToResurrectionGuideBufferUV;
 float2 ResurrectionGuideUVViewportBilinearMin;
 float2 ResurrectionGuideUVViewportBilinearMax;
 float3 HistoryGuideQuantizationError;
+float PerceptionAdd;
 float ResurrectionFrameIndex;
 float PrevFrameIndex;
 float4x4 ClipToResurrectionClip;
@@ -378,8 +379,8 @@ tsr_halfC CorrectGuideColorExposure(tsr_halfC ReprojectedHistoryGuideColor, cons
 	return ReprojectedHistoryGuideColor;
 }
 
-// Sample previous moire history
 #if CONFIG_MOIRE_REPROJECTION
+// Sample previous moire history
 tsr_half4 ReprojectPrevHistoryMoire(tsr_ushort2 InputPixelPos, float2 TextureUV, float FrameIndex)
 {
 	uint2 Random = Rand3DPCG16(int3(InputPixelPos - InputInfo_ViewportMin, View.StateFrameIndexMod8)).xy;
@@ -387,6 +388,23 @@ tsr_half4 ReprojectPrevHistoryMoire(tsr_ushort2 InputPixelPos, float2 TextureUV,
 
 	return PrevHistoryMoire.SampleLevel(GlobalPointClampedSampler, float3(TextureUV + (E - 0.5) * PrevGuideInfo_ExtentInverse, FrameIndex), 0);
 }
+
+// Correct pre exposure of the moire history.
+tsr_half4 CorrectMoireExposure(tsr_half4 ReprojectedHistoryMoireColor, const float PreExposureCorrection)
+{
+	tsr_half PrevFlickeringHistory = ReprojectedHistoryMoireColor[0];
+
+	tsr_half LDRToLinear = min(rcp(saturate(tsr_half(1.0) - PrevFlickeringHistory)) * tsr_half(PerceptionAdd) * tsr_half(PreExposureCorrection), tsr_half(MaxHalfFloat));
+	
+	tsr_half LinearPrevFlickeringHistory = PrevFlickeringHistory * LDRToLinear;
+		
+	PrevFlickeringHistory = LinearPrevFlickeringHistory * rcp(LinearPrevFlickeringHistory + tsr_half(PerceptionAdd));
+
+	// Removes NaN
+	PrevFlickeringHistory = -min(-PrevFlickeringHistory, tsr_half(0.0));
+	
+	return tsr_half4(PrevFlickeringHistory, ReprojectedHistoryMoireColor.yzw);
+}
 #endif
 
 
@@ -588,6 +606,12 @@ void MainCS(
 		
 		ReprojectedHistoryGuideColor = CorrectGuideColorExposure(ReprojectedHistoryGuideColor, HistoryPreExposureCorrection);
 
+		#if CONFIG_MOIRE_REPROJECTION
+		{
+			ReprojectedHistoryMoireColor = CorrectMoireExposure(ReprojectedHistoryMoireColor, HistoryPreExposureCorrection);
+		}
+		#endif
+
 		#if CONFIG_RESURRECTION_REPROJECTION
 		{
 			ResurrectedHistoryGuideColor = CorrectGuideColorExposure(ResurrectedHistoryGuideColor, ResurrectionPreExposureCorrection);
diff --git a/Engine/Source/Runtime/Renderer/Private/PostProcess/TemporalSuperResolution.cpp b/Engine/Source/Runtime/Renderer/Private/PostProcess/TemporalSuperResolution.cpp
--- a/Engine/Source/Runtime/Renderer/Private/PostProcess/TemporalSuperResolution.cpp
+++ b/Engine/Source/Runtime/Renderer/Private/PostProcess/TemporalSuperResolution.cpp
@@ -627,6 +627,7 @@ class FTSRDecimateHistoryCS : public FTSRShader
 		SHADER_PARAMETER(FVector2f, ResurrectionGuideUVViewportBilinearMin)
 		SHADER_PARAMETER(FVector2f, ResurrectionGuideUVViewportBilinearMax)
 		SHADER_PARAMETER(FVector3f, HistoryGuideQuantizationError)
+		SHADER_PARAMETER(float, PerceptionAdd)
 		SHADER_PARAMETER(float, ResurrectionFrameIndex)
 		SHADER_PARAMETER(float, PrevFrameIndex)
 		SHADER_PARAMETER(FMatrix44f, ClipToResurrectionClip)
@@ -2054,6 +2055,7 @@ FDefaultTemporalUpscaler::FOutputs AddTemporalSuperResolutionPasses(
 			PassParameters->ResurrectionGuideUVViewportBilinearMin = GetScreenPassTextureViewportParameters(ResurrectionGuideViewport).UVViewportBilinearMin;
 			PassParameters->ResurrectionGuideUVViewportBilinearMax = GetScreenPassTextureViewportParameters(ResurrectionGuideViewport).UVViewportBilinearMax;
 			PassParameters->HistoryGuideQuantizationError = ComputePixelFormatQuantizationError(PrevHistory.GuideArray->Desc.Format);
+			PassParameters->PerceptionAdd = FMath::Pow(0.5f, CVarTSRShadingExposureOffset.GetValueOnRenderThread());
 		}
 
 		PassParameters->ResurrectionFrameIndex = ResurrectionFrameSliceIndex;
