Fix incorrect shading on mirrored Nanite meshes when a two-sided material is used.

The shader was trying to skip the TwoSidedSign flag computation for raytracing, but it did it by checking if we're in a pixel shader, which incorrectly skips the math when we're using Nanite compute materials.

diff --git a/Engine/Shaders/Private/MaterialTemplate.ush b/Engine/Shaders/Private/MaterialTemplate.ush
--- a/Engine/Shaders/Private/MaterialTemplate.ush
+++ b/Engine/Shaders/Private/MaterialTemplate.ush
@@ -3666,7 +3666,7 @@ void CalcMaterialParametersEx(
 
 #if MATERIAL_TWOSIDED && HAS_PRIMITIVE_UNIFORM_BUFFER
 	// #dxr: DirectX Raytracing's HitKind() intrinsic already accounts for negative scaling
-	#if PIXELSHADER
+	#if !RAYHITGROUPSHADER
 		Parameters.TwoSidedSign *= ResolvedView.CullingSign * GetPrimitive_DeterminantSign(Parameters.PrimitiveId);
 	#endif
 #endif
