LFAtlas - made it work with substrate deferred.

diff --git a/Engine/Shaders/Private/Substrate/SubstrateDeferredLighting.ush b/Engine/Shaders/Private/Substrate/SubstrateDeferredLighting.ush
--- a/Engine/Shaders/Private/Substrate/SubstrateDeferredLighting.ush
+++ b/Engine/Shaders/Private/Substrate/SubstrateDeferredLighting.ush
@@ -4,6 +4,10 @@
 
 #include "../AreaLightCommon.ush"
 
+#if USE_LIGHT_FUNCTION_ATLAS
+#include "../LightFunctionAtlas/LightFunctionAtlasCommon.usf"
+#endif
+
 // Sanity guard. 
 #ifndef SUBSTRATE_ENABLED
 	#define SUBSTRATE_ENABLED 1
@@ -76,6 +80,10 @@ FSubstrateDeferredLighting SubstrateDeferredLighting(
 			SubstrateShadowTermInputParameters.LightAttenuation,
 			ShadowTerms);
 	}
+	
+#if USE_LIGHT_FUNCTION_ATLAS
+	const FLightFunctionColor LightFunctionColor = GetLocalLightFunctionCommon(SubstrateShadowTermInputParameters.TranslatedWorldPosition, LightData.LightFunctionAtlasLightIndex);
+#endif
 
 	Substrate_for (uint ClosureIndex = 0, ClosureIndex < SubstratePixelHeader.ClosureCount, ++ClosureIndex)
 	{
@@ -138,7 +146,11 @@ FSubstrateDeferredLighting SubstrateDeferredLighting(
 		float3 DiffuseLuminance = BSDFEvaluate.IntegratedDiffuseValue;
 		float3 SpecularLuminance = BSDFEvaluate.IntegratedSpecularValue * LightData.SpecularScale;
 
-		const float3 CommonMultiplier = LightData.Color * LightMask * LuminanceWeight(BSDFContext, BSDF);
+		const float3 CommonMultiplier = LightData.Color * LightMask * LuminanceWeight(BSDFContext, BSDF)
+	#if USE_LIGHT_FUNCTION_ATLAS
+			* LightFunctionColor
+	#endif
+			;
 
 		FLightAccumulator Out = (FLightAccumulator)0;
 		LightAccumulator_AddSplit(Out, DiffuseLuminance, SpecularLuminance, DiffuseLuminance, CommonMultiplier, BSDFEvaluate.bPostProcessSubsurface);
