Implement support for virtual texture feedback in RayTracingMaterialHitShaders

diff --git a/Engine/Shaders/Private/RayTracing/RayTracingCommon.ush b/Engine/Shaders/Private/RayTracing/RayTracingCommon.ush
--- a/Engine/Shaders/Private/RayTracing/RayTracingCommon.ush
+++ b/Engine/Shaders/Private/RayTracing/RayTracingCommon.ush
@@ -1248,7 +1248,8 @@ FMaterialClosestHitPayload TraceMaterialRay(
 	// Payload Inputs
 	inout FRayCone RayCone,
 	in bool bEnableSkyLightContribution,
-	in bool bIgnoreTranslucentMaterials)
+	in bool bIgnoreTranslucentMaterials,
+	in bool bIsCameraRay)
 {
 	const uint RayContributionToHitGroupIndex = RAY_TRACING_SHADER_SLOT_MATERIAL;
 	const uint MultiplierForGeometryContributionToShaderIndex = RAY_TRACING_NUM_SHADER_SLOTS;
@@ -1271,6 +1272,11 @@ FMaterialClosestHitPayload TraceMaterialRay(
 		PackedPayload.SetIgnoreTranslucentMaterials();
 	}
 
+	if (bIsCameraRay)
+	{
+		PackedPayload.SetCameraRay();
+	}
+
 	// Trace the ray
 
 	TraceRay(
diff --git a/Engine/Shaders/Private/RayTracing/RayTracingDebug.usf b/Engine/Shaders/Private/RayTracing/RayTracingDebug.usf
--- a/Engine/Shaders/Private/RayTracing/RayTracingDebug.usf
+++ b/Engine/Shaders/Private/RayTracing/RayTracingDebug.usf
@@ -63,6 +63,7 @@ float4 DebugRayTracingMaterial(FRayDesc Ray, bool bIsFarField)
 	const uint InstanceInclusionMask = (bIsFarField) ? RAY_TRACING_MASK_FAR_FIELD : (OpaqueOnly ? (RAY_TRACING_MASK_OPAQUE | RAY_TRACING_MASK_HAIR_STRANDS) : RAY_TRACING_MASK_ALL);
 	const bool bEnableSkyLightContribution = true;
 	const bool bIgnoreTranslucentMaterials = false;
+	const bool bCameraRay = true;
 
 #if PLATFORM_SUPPORTS_SHADER_TIMESTAMP
 	uint64_t TimeBegin = GetShaderTimestamp();
@@ -86,7 +87,8 @@ float4 DebugRayTracingMaterial(FRayDesc Ray, bool bIsFarField)
 		Ray,
 		RayCone,
 		bEnableSkyLightContribution,
-		bIgnoreTranslucentMaterials);
+		bIgnoreTranslucentMaterials,
+		bCameraRay);
 
 #if PLATFORM_SUPPORTS_SHADER_TIMESTAMP
 	uint64_t TimeEnd = GetShaderTimestamp();
diff --git a/Engine/Shaders/Private/RayTracing/RayTracingLightingCommon.ush b/Engine/Shaders/Private/RayTracing/RayTracingLightingCommon.ush
--- a/Engine/Shaders/Private/RayTracing/RayTracingLightingCommon.ush
+++ b/Engine/Shaders/Private/RayTracing/RayTracingLightingCommon.ush
@@ -574,6 +574,7 @@ FMaterialClosestHitPayload TraceRayAndAccumulateBottomLayerResults(
 	in bool bDecoupleSampleGeneration,
 	inout FRayCone RayCone,
 	in bool bEnableSkyLightContribution,
+	in bool bIsCameraRay,
 	inout float3 Radiance)
 {
 	if (bRayTraceSkyLightContribution)
@@ -591,7 +592,8 @@ FMaterialClosestHitPayload TraceRayAndAccumulateBottomLayerResults(
 		Ray,
 		RayCone,
 		bEnableSkyLightContribution,
-		bIgnoreTranslucentMaterials);
+		bIgnoreTranslucentMaterials,
+		bIsCameraRay);
 
 	ComputeBottomLayerMaterialProperties(Ray, BottomLayerPayload);
 
diff --git a/Engine/Shaders/Private/RayTracing/RayTracingMaterialHitShaders.usf b/Engine/Shaders/Private/RayTracing/RayTracingMaterialHitShaders.usf
--- a/Engine/Shaders/Private/RayTracing/RayTracingMaterialHitShaders.usf
+++ b/Engine/Shaders/Private/RayTracing/RayTracingMaterialHitShaders.usf
@@ -624,6 +624,20 @@ RAY_TRACING_ENTRY_CLOSEST_HIT(MaterialCHS,
 		PackedPayload.SetGeometryNormal(GeoNormal);
 	}
 
+#if MATERIAL_VIRTUALTEXTURE_FEEDBACK
+	if (bIsCameraRay)
+	{
+		// Virtual texturing feedback logic (camera rays only for now)
+		FinalizeVirtualTextureFeedback(
+			MaterialParameters.VirtualTextureFeedback,
+			MaterialParameters.SvPosition,
+			Payload.Opacity,
+			View.FrameNumber,
+			View.VTFeedbackBuffer
+		);
+	}
+#endif
+
 #endif // USE_MATERIAL_CLOSEST_HIT_SHADER
 }
 
