SampledDirectLighting - fixed tile artifacts between shading tiles due to bilinear footprint. Also fixed blue noise correlation issues between samples.

diff --git a/Engine/Shaders/Private/SampledDirectLighting/SampledDirectLightingComposite.usf b/Engine/Shaders/Private/SampledDirectLighting/SampledDirectLightingComposite.usf
--- a/Engine/Shaders/Private/SampledDirectLighting/SampledDirectLightingComposite.usf
+++ b/Engine/Shaders/Private/SampledDirectLighting/SampledDirectLightingComposite.usf
@@ -22,18 +22,11 @@ uint DebugLightId;
 #define SHARED_TILE_BORDER 1
 #define SHARED_TILE_SIZE (8 + 2 * SHARED_TILE_BORDER)
 groupshared float SharedSampleWeights[SHARED_TILE_SIZE][SHARED_TILE_SIZE];
-groupshared float SharedSceneDepth[SHARED_TILE_SIZE][SHARED_TILE_SIZE];
 
 // 2x2 tile mask. 1 means that we need to spawn a shading tile for that region.
-groupshared uint SharedShadingTileMask;
+groupshared uint SharedShadingTileMask[2][2];
 groupshared uint SharedShadingTileIndex[2][2];
 
-uint GetLitBitMask(uint2 ShadingTileCoord)
-{
-	const uint BitOffset = (ShadingTileCoord.x + ShadingTileCoord.y * 2);
-	return 0x1u << BitOffset;
-}
-
 void AccumulateLightSample(uint2 SharedTileCoord, float UpsampleWeight, inout float SampleWeightSum)
 {
 	float SampleWeight = SharedSampleWeights[SharedTileCoord.x][SharedTileCoord.y];
@@ -56,13 +49,9 @@ void CompositeLightSamplesCS(
 	uint2 BaseDownsampledScreenCoord;
 	UnpackCompositeTile(CompositeTileData[CompositeTileIndex], LocalLightIndex, BaseDownsampledScreenCoord);
 
-	if (all(GroupThreadId == 0))
-	{
-		SharedShadingTileMask = 0;
-	}
-
 	if (all(GroupThreadId < 2))
 	{
+		SharedShadingTileMask[GroupThreadId.x][GroupThreadId.y] = 0;
 		SharedShadingTileIndex[GroupThreadId.x][GroupThreadId.y] = 0xFFFFFFFF;
 	}
 
@@ -75,7 +64,8 @@ void CompositeLightSamplesCS(
 	{
 		for (uint SharedCoordX = GroupThreadId.x; SharedCoordX < SHARED_TILE_SIZE; SharedCoordX += THREADGROUP_SIZE)
 		{
-			uint2 DownsampledScreenCoord = clamp(int2(BaseDownsampledScreenCoord) + int2(SharedCoordX - SHARED_TILE_BORDER, SharedCoordY - SHARED_TILE_BORDER), 0, DownsampledViewSize - 1);
+			const uint2 SharedCoord = uint2(SharedCoordX, SharedCoordY);
+			const uint2 DownsampledScreenCoord = clamp(int2(BaseDownsampledScreenCoord) + SharedCoord - SHARED_TILE_BORDER, 0, DownsampledViewSize - 1);
 			
 			float WeightSum = 0.0f;
 
@@ -93,49 +83,65 @@ void CompositeLightSamplesCS(
 				}
 			}
 
+			// Mark sub-tiles requiring shading tiles, accounting for a bilinear interpolation footprint
+			if (WeightSum > 0.001f)
+			{
+				uint LitBitMask = 0;
+
+				for (uint ShadingTileOffsetY = 0; ShadingTileOffsetY < 2; ++ShadingTileOffsetY)
+				{
+					for (uint ShadingTileOffsetX = 0; ShadingTileOffsetX < 2; ++ShadingTileOffsetX)
+					{
+						const uint2 ShadingTileOffset = uint2(ShadingTileOffsetX, ShadingTileOffsetY);
+
+						// Bilinear footprint
+						uint2 ShadingTileMin = ShadingTileOffset * 4 + (SHARED_TILE_BORDER - 1);
+						uint2 ShadingTileMax = ShadingTileMin + 5;
+
+						if (and(all(SharedCoord >= ShadingTileMin), all(SharedCoord <= ShadingTileMax)))
+						{
+							// #sdl_todo: don't write to the spot from so many threads
+							SharedShadingTileMask[ShadingTileOffset.x][ShadingTileOffset.y] = 1;
+						}
+					}
+				}
+			}
+
 			SharedSampleWeights[SharedCoordX][SharedCoordY] = WeightSum;
-			SharedSceneDepth[SharedCoordX][SharedCoordY] = DownsampledSceneDepth[DownsampledScreenCoord];
 		}
 	}
 
 	GroupMemoryBarrierWithGroupSync();
 
-	uint2 LocalDownsampledTileCoord = GroupThreadId.xy;
-	uint2 DownsampledScreenCoord = BaseDownsampledScreenCoord + LocalDownsampledTileCoord;
-
-	GroupMemoryBarrierWithGroupSync();
-
-	// Mark sub-tiles requiring shading tiles
-	if (SharedSampleWeights[GroupThreadId.x + SHARED_TILE_BORDER][GroupThreadId.y + SHARED_TILE_BORDER] > 0.0f)
+	if (View.GeneralPurposeTweak <= 0.5)
 	{
-		// #sdl_todo: replace this atomic with waveops + groupshared
-		uint2 ShadingTileOffset = GroupThreadId.xy / 4;
-		InterlockedOr(SharedShadingTileMask, GetLitBitMask(ShadingTileOffset));
+		// Mark sub-tiles requiring shading tiles
+		if (SharedSampleWeights[GroupThreadId.x + SHARED_TILE_BORDER][GroupThreadId.y + SHARED_TILE_BORDER] > 0.0f)
+		{
+			uint2 ShadingTileOffset = GroupThreadId.xy / 4;
+			SharedShadingTileMask[ShadingTileOffset.x][ShadingTileOffset.y] = 1;
+		}
 	}
 
 	GroupMemoryBarrierWithGroupSync();
 
+	uint2 LocalDownsampledTileCoord = GroupThreadId.xy;
+	uint2 DownsampledScreenCoord = BaseDownsampledScreenCoord + LocalDownsampledTileCoord;
 	uint2 ShadingTileOffset = GroupThreadId.xy;
 
 	// Allocate shading tiles
 	if (all(ShadingTileOffset < 2))
 	{
-		if (SharedShadingTileMask & GetLitBitMask(ShadingTileOffset))
+		if (SharedShadingTileMask[ShadingTileOffset.x][ShadingTileOffset.y] != 0)
 		{	
 			const uint2 ShadingTileCoord = (BaseDownsampledScreenCoord * DOWNSAMPLE_FACTOR) / TILE_SIZE + ShadingTileOffset;
 
 			if (all(ShadingTileCoord < ShadingTileGridSize))
 			{
-				const uint UnshadowedTileIndex = 0xFFFFFFFF;
-
-				uint ShadingTileIndex = UnshadowedTileIndex;
-
-				if (SharedShadingTileMask & GetLitBitMask(ShadingTileOffset))
-				{
-					InterlockedAdd(RWShadingTileAllocator[0], 1u, ShadingTileIndex);
-				}
+				uint ShadingTileIndex = MaxShadingTiles;
+				InterlockedAdd(RWShadingTileAllocator[0], 1u, ShadingTileIndex);
 
-				if (ShadingTileIndex < MaxShadingTiles || ShadingTileIndex == UnshadowedTileIndex)
+				if (ShadingTileIndex < MaxShadingTiles)
 				{
 					uint ShadingTileIndexInCell = 0;
 					InterlockedAdd(RWShadingTileGridAllocator[ShadingTileCoord], 1u, ShadingTileIndexInCell);
@@ -166,7 +172,7 @@ void CompositeLightSamplesCS(
 				uint2 ScreenCoord = ShadingTileCoord * TILE_SIZE + CoordInShadingTile;
 				float4 DownsampledScreenWeights = CompositeUpsampleWeights[ScreenCoord];
 
-				const uint2 SharedTileCoord00 = (uint2(ShadingTileOffsetX, ShadingTileOffsetY) * TILE_SIZE + CoordInShadingTile + 1) / DOWNSAMPLE_FACTOR;
+				const uint2 SharedTileCoord00 = (uint2(ShadingTileOffsetX, ShadingTileOffsetY) * TILE_SIZE + CoordInShadingTile + 1) / DOWNSAMPLE_FACTOR + (SHARED_TILE_BORDER - 1);
 
 				float SampleWeightSum = 0.0f;
 				AccumulateLightSample(SharedTileCoord00 + uint2(0, 0), DownsampledScreenWeights.x, SampleWeightSum);
diff --git a/Engine/Shaders/Private/SampledDirectLighting/SampledDirectLightingSampling.usf b/Engine/Shaders/Private/SampledDirectLighting/SampledDirectLightingSampling.usf
--- a/Engine/Shaders/Private/SampledDirectLighting/SampledDirectLightingSampling.usf
+++ b/Engine/Shaders/Private/SampledDirectLighting/SampledDirectLightingSampling.usf
@@ -219,10 +219,16 @@ void GenerateSamplesCS(
 				const uint NumLightsInGridCell = min(CulledLightGridData.NumLights, GetMaxLightsPerCell(EyeIndex));
 				const uint NumLocalLights = GetNumLocalLights(EyeIndex);
 
+				// Initialize random variables using spatiotemporal Blue Noise
 				float LightIndexRandom[NUM_SAMPLES_PER_PIXEL_1D];
-				for (uint LightSampleIndex = 0; LightSampleIndex < NUM_SAMPLES_PER_PIXEL_1D; ++LightSampleIndex)
 				{
-					LightIndexRandom[LightSampleIndex] = BlueNoiseScalar(DownsampledScreenCoord * DOWNSAMPLE_FACTOR + LightSampleOffsets[LightSampleIndex], SampledDirectLightingStateFrameIndex);
+					float GoldenRatioConjugate = 0.6180339887f;
+					float RandomScalar = BlueNoiseScalar(DownsampledScreenCoord, SampledDirectLightingStateFrameIndex);
+
+					for (uint LightSampleIndex = 0; LightSampleIndex < NUM_SAMPLES_PER_PIXEL_1D; ++LightSampleIndex)
+					{
+						LightIndexRandom[LightSampleIndex] = frac(RandomScalar + GoldenRatioConjugate * LightSampleIndex);
+					}
 				}
 
 				#if DEBUG_MODE
@@ -258,9 +264,6 @@ void GenerateSamplesCS(
 				}
 				#endif
 
-				// Weighted blue noise sampling
-				float Xi[NUM_SAMPLES_PER_PIXEL_1D] = LightIndexRandom;
-
 	#if CANDIDATE_LIGHT_MASK
 				for (uint MaskIndex = 0; MaskIndex < SHARED_LIGHT_MASK_SIZE; ++MaskIndex)
 				{
@@ -294,19 +297,19 @@ void GenerateSamplesCS(
 
 							for (uint LightSampleIndex = 0; LightSampleIndex < NUM_SAMPLES_PER_PIXEL_1D; ++LightSampleIndex)
 							{
-								if (Xi[LightSampleIndex] < Tau)
+								if (LightIndexRandom[LightSampleIndex] < Tau)
 								{
-									Xi[LightSampleIndex] /= Tau;
+									LightIndexRandom[LightSampleIndex] /= Tau;
 								}
 								else
 								{
 									// Select this sample
-									Xi[LightSampleIndex] = (Xi[LightSampleIndex] - Tau) / (1.0f - Tau);
+									LightIndexRandom[LightSampleIndex] = (LightIndexRandom[LightSampleIndex] - Tau) / (1.0f - Tau);
 									LightSamples[LightSampleIndex].LocalLightIndex = LocalLightIndex;
 									LightSamples[LightSampleIndex].Weight = SampleWeight;
 								}
 
-								Xi[LightSampleIndex] = clamp(Xi[LightSampleIndex], 0, 0.9999f);
+								LightIndexRandom[LightSampleIndex] = clamp(LightIndexRandom[LightSampleIndex], 0, 0.9999f);
 							}
 						}
 					}
@@ -340,19 +343,19 @@ void GenerateSamplesCS(
 
 							for (uint LightSampleIndex = 0; LightSampleIndex < NUM_SAMPLES_PER_PIXEL_1D; ++LightSampleIndex)
 							{
-								if (Xi[LightSampleIndex] < Tau)
+								if (LightIndexRandom[LightSampleIndex] < Tau)
 								{
-									Xi[LightSampleIndex] /= Tau;
+									LightIndexRandom[LightSampleIndex] /= Tau;
 								}
 								else
 								{
 									// Select this sample
-									Xi[LightSampleIndex] = (Xi[LightSampleIndex] - Tau) / (1.0f - Tau);
+									LightIndexRandom[LightSampleIndex] = (LightIndexRandom[LightSampleIndex] - Tau) / (1.0f - Tau);
 									LightSamples[LightSampleIndex].LocalLightIndex = LocalLightIndex;
 									LightSamples[LightSampleIndex].Weight = SampleWeight;
 								}
 
-								Xi[LightSampleIndex] = clamp(Xi[LightSampleIndex], 0, 0.9999f);
+								LightIndexRandom[LightSampleIndex] = clamp(LightIndexRandom[LightSampleIndex], 0, 0.9999f);
 							}
 						}
 					}
