SampledDirectLighting - fixed shader compilation error on console

diff --git a/Engine/Shaders/Private/SampledDirectLighting/SampledDirectLightingComposite.usf b/Engine/Shaders/Private/SampledDirectLighting/SampledDirectLightingComposite.usf
--- a/Engine/Shaders/Private/SampledDirectLighting/SampledDirectLightingComposite.usf
+++ b/Engine/Shaders/Private/SampledDirectLighting/SampledDirectLightingComposite.usf
@@ -65,7 +65,7 @@ void CompositeLightSamplesCS(
 		for (uint SharedCoordX = GroupThreadId.x; SharedCoordX < SHARED_TILE_SIZE; SharedCoordX += THREADGROUP_SIZE)
 		{
 			const uint2 SharedCoord = uint2(SharedCoordX, SharedCoordY);
-			const uint2 DownsampledScreenCoord = clamp(int2(BaseDownsampledScreenCoord) + SharedCoord - SHARED_TILE_BORDER, 0, DownsampledViewSize - 1);
+			const uint2 DownsampledScreenCoord = clamp(int2(BaseDownsampledScreenCoord + SharedCoord) - SHARED_TILE_BORDER, 0, DownsampledViewSize - 1);
 			
 			float WeightSum = 0.0f;
 
