PathTracer: Simplify Depth AOV logic

Only need to track Z value across hits, not the hit point or a counter. This reduces the live state by 20 bytes (432 down to 412) and speeds up the path tracer by roughly 2%.

diff --git a/Engine/Shaders/Private/PathTracing/Material/PathTracingMaterialSampling.ush b/Engine/Shaders/Private/PathTracing/Material/PathTracingMaterialSampling.ush
--- a/Engine/Shaders/Private/PathTracing/Material/PathTracingMaterialSampling.ush
+++ b/Engine/Shaders/Private/PathTracing/Material/PathTracingMaterialSampling.ush
@@ -271,49 +271,6 @@ void AccumulateAlbedo(float3 SigmaS, float3 PathThroughput, float PathRoughness,
 
 #endif // SIMPLIFIED_MATERIAL_SHADER
 
-
-#define DEBUG_DEPTH_SAMPLING 0
-
-struct FDepthContext
-{
-	float3 TranslatedWorldPos;
-	uint HitCount;
-
-#if DEBUG_DEPTH_SAMPLING
-	// visualize the depth status for the ith ray hit.
-	int DebugHitCount;
-#endif
-
-	void AddHit(FPathTracingPayload Payload)
-	{
-#if DEBUG_DEPTH_SAMPLING
-		if (DebugHitCount == 0)
-		{
-			return;
-		}
-#endif
-		// 2. Matching depth output to the rasterizer.
-		if (Payload.ShouldOutputDepth() && HitCount == 0)
-		{
-			TranslatedWorldPos = Payload.TranslatedWorldPos;
-			HitCount++;
-		}
-
-#if DEBUG_DEPTH_SAMPLING
-		DebugHitCount--;
-#endif
-	}
-};
-
-FDepthContext InitDepthContext()
-{
-	FDepthContext Context = (FDepthContext)0;
-#if DEBUG_DEPTH_SAMPLING
-	Context.DebugHitCount = View.GeneralPurposeTweak;
-#endif
-	return Context;
-}
-
 // returns an average roughness for the whole material (for the path roughness of shadow rays from this material)
 float GetAverageRoughness(FPathTracingPayload Payload)
 {
@@ -326,7 +283,6 @@ float GetAverageRoughness(FPathTracingPayload Payload)
 		case SHADINGMODELID_SUBSTRATE:
 		case SHADINGMODELID_THIN_TRANSLUCENT:
 		{
-
 			AvgRoughness = max(AvgRoughness, Payload.RoughnessData.y);
 			// Only non-glass slabs contain diffuse
 			float3 DiffuseColor = Payload.HasRefraction() ? 0.0 : Payload.SubsurfaceColor + Payload.DiffuseColor;
diff --git a/Engine/Shaders/Private/PathTracing/PathTracing.usf b/Engine/Shaders/Private/PathTracing/PathTracing.usf
--- a/Engine/Shaders/Private/PathTracing/PathTracing.usf
+++ b/Engine/Shaders/Private/PathTracing/PathTracing.usf
@@ -141,6 +141,13 @@ RAY_TRACING_ENTRY_RAYGEN(PathTracingMainRG)
 	}
 
 	const bool KeepGoing = PathTracingKernel(PathState, Bounce);
+	const uint2 TextureIndex = ComputeTextureIndex(uint2(LinearPixelIndex % DispatchDim.x,
+	                                                     LinearPixelIndex / DispatchDim.x));
+	if (Bounce == 0)
+	{
+		// Camera ray must write depth right away because we don't want to have to carry DepthZ in the PackedPathState
+		PathState.WriteDepth(TextureIndex);
+	}
 	if (KeepGoing)
 	{
 		// NOTE: using wave instructions to reduce contention seems to run slightly slower? (tested on RTX-3080)
@@ -153,8 +160,7 @@ RAY_TRACING_ENTRY_RAYGEN(PathTracingMainRG)
 	{
 		// nothing left to do
 		// Accumulate radiance and update pixel variance
-		PathState.WritePixel(ComputeTextureIndex(uint2(LinearPixelIndex % DispatchDim.x,
-			                                           LinearPixelIndex / DispatchDim.x)));
+		PathState.WritePixel(TextureIndex);
 	}
 }
 
@@ -193,7 +199,9 @@ RAY_TRACING_ENTRY_RAYGEN(PathTracingMainRG)
 	}
 
 	// Accumulate radiance and update pixel variance
-	PathState.WritePixel();
+	int2 TextureIndex = ComputeTextureIndex(DispatchIdx);
+	PathState.WritePixel(TextureIndex);
+	PathState.WriteDepth(TextureIndex);
 }
 
 #endif
\ No newline at end of file
diff --git a/Engine/Shaders/Private/PathTracing/PathTracingCore.ush b/Engine/Shaders/Private/PathTracing/PathTracingCore.ush
--- a/Engine/Shaders/Private/PathTracing/PathTracingCore.ush
+++ b/Engine/Shaders/Private/PathTracing/PathTracingCore.ush
@@ -92,8 +92,7 @@ struct FPathState
 									// 80 bytes total (see FPathTracingPackedPathState)
 
 	// Temporary parameters not stored				
-	bool DepthResamplingWeight;		
-	int2 PrimaryRayTextureIndex;		// texture index cache to write to the frame buffer for primary ray operation like depth.	
+	float DepthZ;                   // Camera ray depth
 
 	bool HasMadeContributionScatter()
 	{
@@ -196,24 +195,10 @@ struct FPathState
 		NormalTexture[TextureIndex].xyz = lerp(OldNormal.xyz,Normal.xyz, Blend);
 	}
 
-	void WritePixel()
+	void WriteDepth(uint2 TextureIndex)
 	{
-		WritePixel(PrimaryRayTextureIndex);
-	}
-
-	void WriteDepth(FDepthContext DepthContext)
-	{
-		float OldDepth = Iteration > 0 ? NormalTexture[PrimaryRayTextureIndex].w : 0;
-		
-		float Depth = 0;
-		
-		if (DepthContext.HitCount > 0)
-		{
-			const float3 ViewZ = View.ViewToTranslatedWorld[2].xyz;
-			Depth = dot(DepthContext.TranslatedWorldPos, ViewZ);
-		}
-
-		NormalTexture[PrimaryRayTextureIndex].w = lerp(OldDepth, max(OldDepth, ConvertToDeviceZ(Depth)), DepthResamplingWeight);
+		float OldDepth = Iteration > 0 ? NormalTexture[TextureIndex].w : 0;
+		NormalTexture[TextureIndex].w = max(OldDepth, ConvertToDeviceZ(DepthZ));
 	}
 };
 
@@ -413,7 +398,6 @@ FPathTracingPayload TraceTransparentRay(inout FPathState PathState, int Bounce,
 	float2 RISRandSample = RandomSequence_GenerateSample2D(PathState.RandSequence);
 	FRISContext HitSample = InitRISContext(RISRandSample.x);
 	FRISContext VolSample = InitRISContext(RISRandSample.y);
-	FDepthContext DepthContext = InitDepthContext();
 
 	float3 PayloadThroughput = PathState.PathThroughput;
 	FPathTracingPayload Payload;
@@ -696,9 +680,10 @@ FPathTracingPayload TraceTransparentRay(inout FPathState PathState, int Bounce,
 			break;
 		}
 
-		if (bIsCameraRay)
+		if (bIsCameraRay && HitPayload.ShouldOutputDepth() && PathState.DepthZ == 0.0)
 		{
-			DepthContext.AddHit(HitPayload);
+			const float3 ViewZ = View.ViewToTranslatedWorld[2].xyz;
+			PathState.DepthZ = dot(HitPayload.TranslatedWorldPos, ViewZ);
 		}
 
 		if (!bLastBounce)
@@ -744,11 +729,6 @@ FPathTracingPayload TraceTransparentRay(inout FPathState PathState, int Bounce,
 		Payload.SetMiss();
 	}
 
-	if (bIsCameraRay)
-	{
-		PathState.WriteDepth(DepthContext);
-	}
-
 	return Payload;
 }
 
@@ -1341,7 +1321,7 @@ bool ProcessSubsurfaceRandomWalk(inout FPathTracingPayload Payload, inout float3
 	return false;
 }
 
-FPathState CreatePathState(int2 PixelIndex, int2 PrimaryRayTextureIndex)
+FPathState CreatePathState(int2 PixelIndex, int2 TextureIndex)
 {
 	FPathState Output = (FPathState)0;
 
@@ -1350,7 +1330,7 @@ FPathState CreatePathState(int2 PixelIndex, int2 PrimaryRayTextureIndex)
 #if PATH_TRACER_USE_ADAPTIVE_SAMPLING
 	// initialize progressive sobol sequence based on current index for this pixel
 	// NOTE: Error diffusion sampler cannot be used in this mode
-	int SampleIndex = int(VarianceTexture[PrimaryRayTextureIndex].z);
+	int SampleIndex = int(VarianceTexture[TextureIndex].z);
 	RandomSequence_Initialize(Output.RandSequence, LaunchIndex.x + LaunchIndex.y * 65536, TemporalSeed - Iteration + SampleIndex);
 #else
 	// Initialize random sequence 
@@ -1377,14 +1357,6 @@ FPathState CreatePathState(int2 PixelIndex, int2 PrimaryRayTextureIndex)
 
 		float2 ViewportUV = (LaunchIndex + AAJitter) * View.BufferSizeAndInvSize.zw;
 		Output.Ray = CreatePrimaryRay(ViewportUV);
-
-#if 0
-		// only record the depth within a square.
-		Output.DepthResamplingWeight = all(AAJitter >= 0.0) && all(AAJitter <= 1.0) ? 1 : 0;
-#else
-		// record depth for all rays in the filter footprint (TODO: simplify code that uses this weight instead)
-		Output.DepthResamplingWeight = 1.0;
-#endif
 	}
 
 	if (CameraLensRadius.y > 0)
@@ -1414,7 +1386,6 @@ FPathState CreatePathState(int2 PixelIndex, int2 PrimaryRayTextureIndex)
 	Output.PathRoughness = 0;
 	Output.SigmaT = float3(StartingExtinctionCoefficient[0], StartingExtinctionCoefficient[1], StartingExtinctionCoefficient[2]);
 	Output.FirstScatterType = PATHTRACER_SCATTER_CAMERA;
-	Output.PrimaryRayTextureIndex = PrimaryRayTextureIndex;
 
 	return Output;
 }
