Path Tracer: Further optimization to live state size

Shuffling some constants and changing the order of trace/eval to reduce how much of the payload needs to be kept around in the direct lighting loop. This reduces the main shadow ray trace call from 344 bytes of live state down to 260 bytes. This gives a roughly 3-5% speedup overall.

diff --git a/Engine/Shaders/Private/PathTracing/PathTracingCore.ush b/Engine/Shaders/Private/PathTracing/PathTracingCore.ush
--- a/Engine/Shaders/Private/PathTracing/PathTracingCore.ush
+++ b/Engine/Shaders/Private/PathTracing/PathTracingCore.ush
@@ -384,8 +384,10 @@ void EvaluateDecals(FPathState PathState, uint InstanceInclusionMask, FPackedPat
 #endif
 }
 
-FPathTracingPayload TraceTransparentRay(inout FPathState PathState, int Bounce, bool bIsCameraRay, bool bLastBounce, bool bIncludeEmission, uint NumVisibleLights, inout FVolumeSegment VolumeSegment)
+FPathTracingPayload TraceTransparentRay(inout FPathState PathState, int Bounce, inout FVolumeSegment VolumeSegment)
 {
+	const bool bIsCameraRay = Bounce == 0;
+	const bool bLastBounce = Bounce == MaxBounces;
 	const uint RayFlags = bIsCameraRay && EnableCameraBackfaceCulling ? RAY_FLAG_CULL_BACK_FACING_TRIANGLES : 0;
 	const bool bUseCameraRayType = !PathState.HasMadeContributionScatter();
 	uint InstanceInclusionMask = bUseCameraRayType ? PATHTRACER_MASK_CAMERA : PATHTRACER_MASK_INDIRECT;
@@ -393,7 +395,8 @@ FPathTracingPayload TraceTransparentRay(inout FPathState PathState, int Bounce,
 	{
 		InstanceInclusionMask |= bUseCameraRayType ? PATHTRACER_MASK_HAIR_CAMERA : PATHTRACER_MASK_HAIR_INDIRECT;
 	}
-	
+
+
 	const uint MissShaderIndex = 0;
 	float2 RISRandSample = RandomSequence_GenerateSample2D(PathState.RandSequence);
 	FRISContext HitSample = InitRISContext(RISRandSample.x);
@@ -401,7 +404,7 @@ FPathTracingPayload TraceTransparentRay(inout FPathState PathState, int Bounce,
 
 	float3 PayloadThroughput = PathState.PathThroughput;
 	FPathTracingPayload Payload;
-	if (!bIncludeEmission && bLastBounce)
+	if (bLastBounce && !PathState.ShouldAccumulateEmissive())
 	{
 		Payload.SetMiss();
 		PathState.PathThroughput = 0;
@@ -485,7 +488,10 @@ FPathTracingPayload TraceTransparentRay(inout FPathState PathState, int Bounce,
 
 		// Loop over lights to capture their contribution
 		// #dxr_todo: if we have lots of lights, having some hierarchical structure would be better ....
-		for (uint LightId = 0; LightId < NumVisibleLights; ++LightId)
+		// number of directly visible lights for the first bounce
+		// after the first bounce, we don't need to include lights in the trace call
+		// because NEE handled it for us
+		for (uint LightId = 0, NumVisibleLights = bIsCameraRay ? SceneVisibleLightCount : 0; LightId < NumVisibleLights; ++LightId)
 		{
 			FRayDesc LightRay = PathState.Ray;
 			LightRay.TMax = PackedPayload.IsMiss() ? PathState.Ray.TMax : PackedPayload.HitT;
@@ -527,8 +533,8 @@ FPathTracingPayload TraceTransparentRay(inout FPathState PathState, int Bounce,
 				Interval.VolumeTMax = min(PackedPayload.HitT, Interval.VolumeTMax);
 			}
 
-			// if this isn't our last bounce, potentially keep this segment for later so we can integrate local lighting
-			if (!bLastBounce)
+			// if this isn't our last bounce and we are including volumetric contributions, potentially keep this segment for later so we can integrate local lighting
+			if (!bLastBounce && PathState.ShouldAccumulateVolume())
 			{
 				// TODO: can we improve the selection probability for each ray segment somehow?
 				// NOTE: this only matters if we are tracing through several transparent hits or if there are multiple overlapping volumes
@@ -669,7 +675,7 @@ FPathTracingPayload TraceTransparentRay(inout FPathState PathState, int Bounce,
 		}
 
 		// add in surface emission
-		if (bIncludeEmission)
+		if (PathState.ShouldAccumulateEmissive())
 		{
 			AccumulateRadiance(PathState.Radiance, PathState.PathThroughput * HitPayload.Radiance, bIsCameraRay);
 		}
@@ -1402,26 +1408,15 @@ bool PathTracingKernel(inout FPathState PathState, int Bounce)
 #endif
 
 	const bool bIsCameraRay = Bounce == 0;
-	const bool bIsLastBounce = Bounce == MaxBounces;
-	const bool bIncludeEmissive = PathState.ShouldAccumulateEmissive();
-
-	// number of directly visible lights for the first bounce
-	// after the first bounce, we don't need to include lights in the trace call
-	// because NEE handled it for us
-	uint NumVisibleLights = bIsCameraRay ? SceneVisibleLightCount : 0;
 
 	FVolumeSegment VolumeSegment = CreateEmptyVolumeSegment();
 	FPathTracingPayload Payload = TraceTransparentRay(
 		PathState,
 		Bounce,
-		bIsCameraRay,
-		bIsLastBounce,
-		bIncludeEmissive,
-		NumVisibleLights,
 		VolumeSegment);
 
 	// process the returned volume segment if we got one
-	if (VolumeSegment.IsValid() && VolumeMISMode != 0 && PathState.ShouldAccumulateVolume())
+	if (VolumeSegment.IsValid() && VolumeMISMode != 0)
 	{
 		const float3 Ro = PathState.Ray.Origin;
 		const float3 Rd = PathState.Ray.Direction;
@@ -1492,8 +1487,6 @@ bool PathTracingKernel(inout FPathState PathState, int Bounce)
 						LightRay.TMin = 0;
 						LightRay.Direction = LightSample.Direction;
 						LightRay.TMax = LightSample.Distance;
-						LightSample.RadianceOverPdf *= TraceTransparentVisibilityRay(LightRay, Bounce, 1.0, MissShaderIndex, bCastShadows, PathState.RandSequence);
-						// #dxr_todo: Is it cheaper to fire the ray first? Or eval the material first?
 						if (any(LightSample.RadianceOverPdf > 0))
 						{
 							// Evaluate material
@@ -1505,6 +1498,8 @@ bool PathTracingKernel(inout FPathState PathState, int Bounce)
 							{
 								LightContrib *= MISWeightRobust(LightSample.Pdf, MaterialEval.Pdf);
 							}
+							LightContrib *= TraceTransparentVisibilityRay(LightRay, Bounce, 1.0, MissShaderIndex, bCastShadows, PathState.RandSequence);
+
 							AccumulateRadiance(PathState.Radiance, LightContrib, bIsCameraRay);
 						}
 					}
@@ -1597,6 +1592,9 @@ bool PathTracingKernel(inout FPathState PathState, int Bounce)
 
 	float LightPickingCdfSum = 0;
 
+	// Sample material
+	FMaterialSample MaterialSample = SampleMaterial(-PathState.Ray.Direction, Payload, RandSample.xyz);
+
 	// If we are using Light sampling and the material can use it ...
 	if (MISMode != 0 && SceneLightCount > 0 && bDoLightLoop)
 	{
@@ -1646,10 +1644,6 @@ bool PathTracingKernel(inout FPathState PathState, int Bounce)
 
 				float AvgRoughness = ApproximateCaustics ? GetAverageRoughness(Payload) : 0.0;
 
-				const bool bCastShadows = bIsVolumeSample ? CastsVolumeShadow(LightId) : CastsShadow(LightId);
-				const uint MissShaderIndex = GetLightMissShaderIndex(LightId);
-				LightSample.RadianceOverPdf *= TraceTransparentVisibilityRay(LightRay, Bounce, AvgRoughness, MissShaderIndex, bCastShadows, PathState.RandSequence);
-
 				if (bIsVolumeSample)
 				{
 					LightSample.RadianceOverPdf *= GetVolumetricScatteringIntensity(LightId);
@@ -1669,15 +1663,16 @@ bool PathTracingKernel(inout FPathState PathState, int Bounce)
 					{
 						LightContrib *= MISWeightRobust(LightSample.Pdf, MaterialEval.Pdf);
 					}
+					const bool bCastShadows = bIsVolumeSample ? CastsVolumeShadow(LightId) : CastsShadow(LightId);
+					const uint MissShaderIndex = GetLightMissShaderIndex(LightId);
+					LightContrib *= TraceTransparentVisibilityRay(LightRay, Bounce, AvgRoughness, MissShaderIndex, bCastShadows, PathState.RandSequence);
+
 					AccumulateRadiance(PathState.Radiance, LightContrib, bIsCameraRay);
 				}
 			}
 		}
 	}
 
-	// Sample material
-	FMaterialSample MaterialSample = SampleMaterial(-PathState.Ray.Direction, Payload, RandSample.xyz);
-
 	if (MaterialSample.Pdf < 0 || asuint(MaterialSample.Pdf) > 0x7F800000)
 	{
 		// Pdf became invalid (either negative or NaN)
@@ -1760,9 +1755,13 @@ bool PathTracingKernel(inout FPathState PathState, int Bounce)
 		// of shadow ray specific logic for transparent shadows or various per light tricks like shadow casting
 		const bool bUseMIS = MISMode == 2 && LightPickingCdfSum > 0;
 		const float ShadowRayRoughness = ApproximateCaustics ? PathState.PathRoughness : 0.0;
+		float PreviousCdfValue = 0.0;
 		for (uint Index = 0, Num = LightLoopCount.NumMISLights; Index < Num; ++Index)
 		{
-			uint LightId = GetLightId(Index, LightLoopCount);
+			const uint LightId = GetLightId(Index, LightLoopCount);
+			const float CdfValue = LightPickingCdf[Index];
+			const float LightPickPdf = (CdfValue - PreviousCdfValue) / LightPickingCdfSum;
+			PreviousCdfValue = CdfValue;
 			if ((Payload.PrimitiveLightingChannelMask & GetLightingChannelMask(LightId)) == 0)
 			{
 				// light does not affect the current ray
@@ -1799,13 +1798,6 @@ bool PathTracingKernel(inout FPathState PathState, int Bounce)
 
 			if (bUseMIS)
 			{
-				float PreviousCdfValue = 0.0;
-				BRANCH if (Index > 0)
-				{
-					PreviousCdfValue = LightPickingCdf[Index - 1];
-				}
-				float LightPickPdf = (LightPickingCdf[Index] - PreviousCdfValue) / LightPickingCdfSum;
-
 				LightContrib *= MISWeightRobust(MaterialSample.Pdf, LightResult.Pdf * LightPickPdf);
 			}
 
