Modularizes TSR's RejectShading shader for better register pressure experiments

diff --git a/Engine/Shaders/Private/TemporalSuperResolution/TSRRejectShading.usf b/Engine/Shaders/Private/TemporalSuperResolution/TSRRejectShading.usf
--- a/Engine/Shaders/Private/TemporalSuperResolution/TSRRejectShading.usf
+++ b/Engine/Shaders/Private/TemporalSuperResolution/TSRRejectShading.usf
@@ -18,8 +18,7 @@
 
 #define DEBUG_DISABLE_PADDING 0
 
-/** Whether to compose translucency to the input. */
-#define CONFIG_COMPOSE_TRANSLUCENCY 1
+#define DEBUG_DIGESTABLE 0
 
 /** Whether to detect flickering. */
 #define CONFIG_MOIRE_DETECTION (DIM_FLICKERING_DETECTION)
@@ -33,7 +32,24 @@
 
 //------------------------------------------------------- DERIVED
 	
-#if DIM_WAVE_SIZE == 64 || DIM_WAVE_SIZE == 32 || DIM_WAVE_SIZE == 16
+#if DIM_WAVE_SIZE == 0 || DEBUG_DIGESTABLE
+	/** Load an entire 32x32 with 1x1 worth of data by lane, only using LDS. */
+	#define GROUP_TILE_SIZE 32
+
+	#define WAVE_COUNT_X GROUP_TILE_SIZE
+	#define WAVE_COUNT_Y GROUP_TILE_SIZE
+
+	#define LANE_COUNT_X 1
+	#define LANE_COUNT_Y 1
+
+	#define LANE_STRIDE_X 1
+	#define LANE_STRIDE_Y 1
+
+	#define LDS_SIZE (GROUP_TILE_SIZE * GROUP_TILE_SIZE)
+
+	#define TENSOR_REGISTER_LAYOUT FRegisters_S_OneVector
+
+#elif DIM_WAVE_SIZE == 64 || DIM_WAVE_SIZE == 32 || DIM_WAVE_SIZE == 16
 	/** Layouts wave in row major 32x32 tile so an entire row of pixels fits within a single wave, and waveops horizontally, but LDS vertically.
 	 * Each lane worth of a 1x4 of data per lane to minimise to trade VGPR pressure for both number of wave ops and LDS.
 	 */
@@ -76,23 +92,6 @@
 		#define CONFIG_LDS_DESPILL 1
 	#endif
 
-#elif DIM_WAVE_SIZE == 0
-	/** Load an entire 32x32 with 1x1 worth of data by lane, only using LDS. */
-	#define GROUP_TILE_SIZE 32
-
-	#define WAVE_COUNT_X GROUP_TILE_SIZE
-	#define WAVE_COUNT_Y GROUP_TILE_SIZE
-
-	#define LANE_COUNT_X 1
-	#define LANE_COUNT_Y 1
-
-	#define LANE_STRIDE_X 1
-	#define LANE_STRIDE_Y 1
-
-	#define LDS_SIZE (GROUP_TILE_SIZE * GROUP_TILE_SIZE)
-	
-	#define TENSOR_REGISTER_LAYOUT FRegisters_S_OneVector
-
 #else
 	#error Unknown wave size
 #endif
@@ -146,6 +145,11 @@
 #endif
 
 
+//------------------------------------------------------- CONSTANTS
+
+static const tsr_half FilteringWeight = rcp(tsr_half(1.0 + 4 * 0.5 + 4 * 0.25));
+
+
 //------------------------------------------------------- PARAMETERS
 
 FScreenTransform InputPixelPosToTranslucencyTextureUV;
@@ -182,7 +186,22 @@ RWTexture2D<tsr_ushort> AntiAliasMaskOutput;
 static tsr_simd_halfC Debug[DEBUG_ARRAY_SIZE];
 
 
-//------------------------------------------------------- GENERAL OPERATIONS
+//------------------------------------------------------- COMPOSE SEPARATE TRANSLUCENCY
+
+/** Compose translucency in linear color space. */
+CALL_SITE_DEBUGLOC
+tsr_simd_halfC ComposeTranslucency(tsr_simd_halfC Color, tsr_simd_half4 Translucency)
+{
+	tsr_simd_halfC ComposedColor;
+	ComposedColor = Color * tsr_simd_halfC::Vectorize(Translucency[3]) + ResizeChannels<CONFIG_CHANNEL_COUNT>(Translucency);
+	#if CONFIG_CHANNEL_COUNT == 4
+		ComposedColor.SetComponent(3, Color[3] * Translucency[3]);
+	#endif
+	return ComposedColor;
+}
+
+
+//------------------------------------------------------- CHANGE LDR EXPOSURE
 
 CALL_SITE_DEBUGLOC
 tsr_simd_halfC RemovePerceptionAdd(tsr_simd_halfC SrcColor, const float SourcePerceptionAdd)
@@ -209,9 +228,52 @@ tsr_simd_halfC ChangePerceptionAdd(tsr_simd_halfC SrcColor, const float SourcePe
 }
 
 
-//------------------------------------------------------- MAIN ALGORITHMS
+//------------------------------------------------------- SPATIAL ANTI-ALIASER
+
+// Compute the LDR luminance the spatial anti-aliaser should anti-alias.
+CALL_SITE_DEBUGLOC
+tsr_simd_half ComputeSpatialAntiAliaserLumaLDR(tsr_simd_halfC SceneColor)
+{
+	const tsr_half SpatialAAExposure = tsr_half(0.5);
+		
+	tsr_simd_half AALumaLDR;
+
+	UNROLL_N(SIMD_SIZE)
+	for (uint ElementIndex = 0; ElementIndex < SIMD_SIZE; ElementIndex++)
+	{
+		tsr_half PixelLuma = dot(SceneColor.GetElement(ElementIndex).rgb, tsr_half3(0.299f, 0.587f, 0.114f));
+		AALumaLDR.SetElement(ElementIndex, PixelLuma / (SpatialAAExposure + PixelLuma));
+	}
+
+	return AALumaLDR;
+}
+
+// Returns whether the aliasing is visible at all, to avoid paying spatial-antialiaser on low contrasts areas.
+CALL_SITE_DEBUGLOC
+tsr_simd_bool IsAliasingVisible(tsr_simd_halfC ExposedInputBoxSize)
+{
+	#if CONFIG_SCENE_COLOR_ALPHA
+		tsr_simd_half ExposedInputLuminanceBox = dot(ExposedInputBoxSize, tsr_simd_halfC::Const(tsr_halfC(0.299f, 0.587f, 0.114f, 0.0)));
+	#else
+		tsr_simd_half ExposedInputLuminanceBox = dot(ExposedInputBoxSize, tsr_simd_halfC::Const(tsr_halfC(0.299f, 0.587f, 0.114f)));
+	#endif
+
+	tsr_simd_bool bAliasingIsVisible = ExposedInputLuminanceBox > tsr_simd_half::Const(SPATIAL_ANTI_ALIASER_MIN_LUMIMANCE);
+
+	return bAliasingIsVisible;
+}
+
+// Whether to run spatial anti-aliaser on the frame. Must match SpatialAntiAliasingLerp.
+CALL_SITE_DEBUGLOC
+tsr_simd_bool ShouldSpatialAntiAlias(tsr_simd_bool bIsDisoccluded, tsr_simd_bool bResurrectHistory, tsr_simd_bool bAliasingIsVisible, tsr_simd_half RejectionBlendFinal)
+{
+	tsr_simd_bool bSpatialAntiAlias = and(bAliasingIsVisible, or(RejectionBlendFinal < tsr_simd_half::Const(tsr_half(0.25)), and(bIsDisoccluded, not(bResurrectHistory))));
+	
+	return bSpatialAntiAlias;
+}
 
-static const tsr_half FilteringWeight = rcp(tsr_half(1.0 + 4 * 0.5 + 4 * 0.25));
+
+//------------------------------------------------------- FLICKERING TEMPORAL ANALYSIS
 
 void ComputeMoireError(
 	tsr_simd_bool  bIsValid,
@@ -431,6 +493,9 @@ void ComputeMoireError(
 	OutMoireError = OutMoireError * bIsStatic;
 } // ComputeMoireError()
 
+
+//------------------------------------------------------- SHADING REJECTION
+
 template<bool ApplyMoireError>
 void MeasureRejection(
 	tsr_simd_bool  bIsValid,
@@ -542,6 +607,73 @@ void MeasureRejection(
 } // MeasureRejection()
 
 
+//------------------------------------------------------- RESURRECTION
+
+#if CONFIG_HISTORY_RESURRECTION
+
+// Returns whether resurrected frame is closer than previous frame on per pixel basis.
+CALL_SITE_DEBUGLOC
+tsr_simd_bool IsResurrectedFrameCloserThanPreviousFrame(tsr_simd_halfC ExposedInput, tsr_simd_halfC ExposedHistory, tsr_simd_halfC ExposedResurrectedHistory)
+{
+	tsr_simd_halfC PrevMatch         = abs(ExposedInput - ExposedHistory);
+	tsr_simd_halfC ResurrectionMatch = abs(ExposedInput - ExposedResurrectedHistory);
+
+	tsr_simd_halfC MatchDiff = PrevMatch - ResurrectionMatch;
+	#if CONFIG_SCENE_COLOR_ALPHA
+		tsr_simd_half TotalMatchDiff = BlurEven3x3(MatchDiff[0] + MatchDiff[1] + MatchDiff[2] + MatchDiff[3]);
+	#else
+		tsr_simd_half TotalMatchDiff = BlurEven3x3(MatchDiff[0] + MatchDiff[1] + MatchDiff[2]);
+	#endif
+
+	tsr_simd_bool bResurrectionIsCloser = MedianBool3x3(TotalMatchDiff > tsr_simd_half::Const(tsr_half(0.05 * 3.0)));
+
+	return bResurrectionIsCloser;
+}
+
+// Whether it is worth for the entire tile to MeasureRejection() on the resurrected frame.
+CALL_SITE_DEBUGLOC
+bool IsWorthMeasuringRejectionOfResurrectedFrame(tsr_simd_half PrevRejectionBlendFinal, tsr_simd_bool bResurrectionIsCloser)
+{
+	tsr_simd_bool bPrevIsNotGoodEnough = PrevRejectionBlendFinal < tsr_simd_half::Const(0.5);
+	tsr_simd_bool bIsValidTilePixel = IsValidTilePixel<LANE_STRIDE_X, LANE_STRIDE_Y>(TileOverscan);
+	tsr_simd_bool bShouldComputeResurrection = and(and(bResurrectionIsCloser, bPrevIsNotGoodEnough), bIsValidTilePixel);
+
+	bool bComputeResurrection = AnyThread(AnyComponent(bShouldComputeResurrection));
+
+	return bComputeResurrection;
+}
+
+// Final logic whether the resurrection must actually happen on per pixel basis.
+CALL_SITE_DEBUGLOC
+tsr_simd_bool ShouldResurrectHistory(tsr_simd_half PrevRejectionBlendFinal, tsr_simd_half ResurrectionRejectionBlendFinal, tsr_simd_bool bResurrectionIsCloser)
+{
+	tsr_simd_bool bResurrectionIsBetter = (ResurrectionRejectionBlendFinal - PrevRejectionBlendFinal) > tsr_simd_half::Const(0.1);
+	tsr_simd_bool bResurrectHistory = and(bResurrectionIsCloser, bResurrectionIsBetter);
+
+	return bResurrectHistory;
+}
+
+#endif // CONFIG_HISTORY_RESURRECTION
+
+
+//------------------------------------------------------- UPDATE GUDIE
+
+// Update the guide for next frame.
+CALL_SITE_DEBUGLOC
+tsr_simd_halfC UpdateGuide(tsr_simd_halfC ExposedInput, tsr_simd_halfC ExposedHistory, tsr_simd_bool bIsDisoccluded, tsr_simd_half RejectionBlendFinal)
+{
+	tsr_simd_bool bIsOffScreen = tsr_simd_bool::Const(false);
+
+	tsr_simd_half TotalBlend = select(or(bIsOffScreen, bIsDisoccluded), tsr_simd_half::Const(1.0), saturate(tsr_simd_half::Const(1.0) - RejectionBlendFinal * tsr_simd_half::Const(4.0)));
+	tsr_simd_half RejectedBlendFactor = max(tsr_simd_half::Const(tsr_half(TheoricBlendFactor)), tsr_simd_half::Const(1.0) - RejectionBlendFinal);
+
+	tsr_simd_half BlendFinal = max(RejectedBlendFactor, TotalBlend);
+
+	tsr_simd_halfC FinalGuide = ExposedInput * tsr_simd_halfC::Vectorize(BlendFinal) - ExposedHistory * tsr_simd_halfC::Vectorize(BlendFinal - 1.0);
+	return FinalGuide;
+}
+
+
 //------------------------------------------------------- ENTRY POINT
 
 #if COMPILER_SUPPORTS_WAVE_SIZE && DIM_WAVE_SIZE > 0
@@ -568,9 +700,7 @@ void MainCS(
 		
 	// Load
 	tsr_simd_halfC OriginalOpaqueInput;
-	#if CONFIG_COMPOSE_TRANSLUCENCY
-		tsr_simd_half4 OriginalTranslucencyInput;
-	#endif
+	tsr_simd_half4 OriginalTranslucencyInput;
 	tsr_simd_halfC OriginalHistory;
 	tsr_simd_half OriginalHistoryUncertainty;
 	tsr_simd_half VelocityEdge;
@@ -614,9 +744,7 @@ void MainCS(
 			#endif
 
 			OriginalOpaqueInput.SetElement(SimdIndex, InputTexture[ClampedFetchPixelPos]);
-			#if CONFIG_COMPOSE_TRANSLUCENCY
-				OriginalTranslucencyInput.SetElement(SimdIndex, InputSceneTranslucencyTexture.SampleLevel(GlobalBilinearClampedSampler, TranslucencyTextureUV, 0));
-			#endif
+			OriginalTranslucencyInput.SetElement(SimdIndex, InputSceneTranslucencyTexture.SampleLevel(GlobalBilinearClampedSampler, TranslucencyTextureUV, 0));
 			#if CONFIG_SCENE_COLOR_ALPHA
 				OriginalHistory.SetElement(SimdIndex, ReprojectedGuide);
 				OriginalHistoryUncertainty.SetElement(SimdIndex, ReprojectedGuideMetadata);
@@ -641,67 +769,19 @@ void MainCS(
 				#endif
 			#endif
 		}
-		
-		OriginalOpaqueInput.TightenRegisters();
-		#if CONFIG_COMPOSE_TRANSLUCENCY
-			OriginalTranslucencyInput.TightenRegisters();
-		#endif
-		OriginalHistory.TightenRegisters();
-		OriginalHistoryUncertainty.TightenRegisters();
-		VelocityEdge.TightenRegisters();
-		bIsValid.TightenRegisters();
-		#if CONFIG_MOIRE_DETECTION
-			MoireInput.TightenRegisters();
-			PrevMoireHistory.TightenRegisters();
-			bIsStatic.TightenRegisters();
-		#endif
-		#if CONFIG_HISTORY_RESURRECTION
-			OriginalResurrectedHistory.TightenRegisters();
-			OriginalResurrectedHistoryUncertainty.TightenRegisters();
-		#endif
 	}
 	
-	// Compute AA Luma
+	// Compute AA Luma and compose separate translucency into scene color for optimised HistoryUpdate.
 	{
-		const tsr_half SpatialAAExposure = tsr_half(0.5);
-		
-		tsr_simd_halfC CenterFinalSceneColor;
-		tsr_simd_half CenterLumaLDR;
-
-		UNROLL_N(SIMD_SIZE)
-		for (uint SimdIndex = 0; SimdIndex < SIMD_SIZE; SimdIndex++)
-		{
-			tsr_halfC Color = OriginalOpaqueInput.GetElement(SimdIndex);
-			
-			#if CONFIG_COMPOSE_TRANSLUCENCY
-				tsr_half4 Translucency = OriginalTranslucencyInput.GetElement(SimdIndex);
-				Color.rgb = Color.rgb * Translucency.a + Translucency.rgb;
-				#if CONFIG_SCENE_COLOR_ALPHA
-					Color.a = Color.a * Translucency.a;
-				#endif
-			#endif
-			
-			#if CONFIG_SCENE_COLOR_OVERFLOW
-				#if CONFIG_SCENE_COLOR_ALPHA
-					Color = min(Color, LargestSceneColorRGBA);
-				#else
-					Color = min(Color, LargestSceneColorRGB);
-				#endif
-			#endif
-
-			CenterFinalSceneColor.SetElement(SimdIndex, Color);
-
-			tsr_half CenterLuma = dot(Color.rgb, tsr_half3(0.299f, 0.587f, 0.114f));
-			CenterLumaLDR.SetElement(SimdIndex, CenterLuma / (SpatialAAExposure + CenterLuma));
-		}
-
+		tsr_simd_halfC CenterFinalSceneColor = ComposeTranslucency(OriginalOpaqueInput, OriginalTranslucencyInput);
 		CenterFinalSceneColor = min(CenterFinalSceneColor, tsr_simd_halfC::Const(tsr_half(Max10BitsFloat)));
+		
+		tsr_simd_half CenterLumaLDR = ComputeSpatialAntiAliaserLumaLDR(CenterFinalSceneColor);
 
-		// Output some data early.
 		PLATFORM_SPECIFIC_ISOLATE
 		{
 			UNROLL_N(SIMD_SIZE)
-			for (uint SimdIndex = 0; SimdIndex < SIMD_SIZE; SimdIndex++)
+			for (uint ElementIndex = 0; ElementIndex < SIMD_SIZE; ElementIndex++)
 			{
 				const tsr_short2 OutputPixelPos = ComputeElementOutputPixelPos(
 					/* LaneStride = */ uint2(LANE_STRIDE_X, LANE_STRIDE_Y),
@@ -710,36 +790,15 @@ void MainCS(
 					InputInfo_ViewportMax,
 					GroupId,
 					GroupThreadIndex,
-					SimdIndex);
+					ElementIndex);
 		
-				InputSceneColorOutput[OutputPixelPos] = CenterFinalSceneColor.GetElement(SimdIndex);
-				InputSceneColorLdrLumaOutput[OutputPixelPos] = CenterLumaLDR.GetElement(SimdIndex);
+				InputSceneColorOutput[OutputPixelPos] = CenterFinalSceneColor.GetElement(ElementIndex);
+				InputSceneColorLdrLumaOutput[OutputPixelPos] = CenterLumaLDR.GetElement(ElementIndex);
 			}
 		}
 	}
 	
-	tsr_simd_halfC OriginalInput = OriginalOpaqueInput;
-
-	#if CONFIG_COMPOSE_TRANSLUCENCY
-	{
-		tsr_simd_half4 FilteredTranslucencyInput = Blur3x3(OriginalTranslucencyInput);
-
-		UNROLL_N(SIMD_SIZE)
-		for (uint SimdIndex = 0; SimdIndex < SIMD_SIZE; SimdIndex++)
-		{
-			tsr_halfC Color = OriginalOpaqueInput.GetElement(SimdIndex);
-			tsr_half4 Translucency = FilteredTranslucencyInput.GetElement(SimdIndex);
-
-			Color.rgb = Color.rgb * Translucency.a + Translucency.rgb;
-			#if CONFIG_SCENE_COLOR_ALPHA
-				Color.a *= Translucency.a;
-			#endif 
-
-			OriginalInput.SetElement(SimdIndex, Color);
-		}
-	}
-	#endif
-	
+	tsr_simd_halfC OriginalInput = ComposeTranslucency(OriginalOpaqueInput, Blur3x3(OriginalTranslucencyInput));
 	#if CONFIG_SCENE_COLOR_ALPHA
 		OriginalInput = min(OriginalInput, tsr_simd_halfC::Const(LargestSceneColorRGBA));
 	#else
@@ -749,37 +808,36 @@ void MainCS(
 	// Analyses the moire pattern overtime.
 	tsr_simd_half  MoireError = tsr_simd_half::Const(0.0);
 	#if CONFIG_MOIRE_DETECTION
-	tsr_simd_half4 MoireHistory;
-	ComputeMoireError(
-		bIsValid,
-		bIsStatic,
-		MoireInput,
-		OriginalOpaqueInput,
-		PrevMoireHistory,
-		/* out */ MoireHistory,
-		/* out */ MoireError);
-	#endif
-	
-	// Output some data early.
-	PLATFORM_SPECIFIC_ISOLATE
 	{
-		UNROLL_N(SIMD_SIZE)
-		for (uint SimdIndex = 0; SimdIndex < SIMD_SIZE; SimdIndex++)
+		tsr_simd_half4 MoireHistory;
+		ComputeMoireError(
+			bIsValid,
+			bIsStatic,
+			MoireInput,
+			OriginalOpaqueInput,
+			PrevMoireHistory,
+			/* out */ MoireHistory,
+			/* out */ MoireError);
+
+		PLATFORM_SPECIFIC_ISOLATE
 		{
-			const tsr_short2 OutputPixelPos = ComputeElementOutputPixelPos(
-				/* LaneStride = */ uint2(LANE_STRIDE_X, LANE_STRIDE_Y),
-				TileOverscan,
-				InputInfo_ViewportMin,
-				InputInfo_ViewportMax,
-				GroupId,
-				GroupThreadIndex,
-				SimdIndex);
+			UNROLL_N(SIMD_SIZE)
+			for (uint ElementIndex = 0; ElementIndex < SIMD_SIZE; ElementIndex++)
+			{
+				const tsr_short2 OutputPixelPos = ComputeElementOutputPixelPos(
+					/* LaneStride = */ uint2(LANE_STRIDE_X, LANE_STRIDE_Y),
+					TileOverscan,
+					InputInfo_ViewportMin,
+					InputInfo_ViewportMax,
+					GroupId,
+					GroupThreadIndex,
+					ElementIndex);
 
-			#if CONFIG_MOIRE_DETECTION
-				HistoryMoireOutput[tsr_short3(OutputPixelPos, 0)] = MoireHistory.GetElement(SimdIndex);
-			#endif
+				HistoryMoireOutput[tsr_short3(OutputPixelPos, 0)] = MoireHistory.GetElement(ElementIndex);
+			}
 		}
 	}
+	#endif // CONFIG_MOIRE_DETECTION
 
 	// Measure rejection
 	tsr_simd_half RejectionBlendFinal = tsr_simd_half::Const(1.0);
@@ -822,16 +880,8 @@ void MainCS(
 		tsr_simd_halfC ExposedInputMin;
 		tsr_simd_halfC ExposedInputMax;
 		MinMax3x3(ExposedInput, /* out */ ExposedInputMin, /* out */ ExposedInputMax);
-		
-		{
-			#if CONFIG_SCENE_COLOR_ALPHA
-				tsr_simd_half ExposedInputLuminanceBox = dot(ExposedInputMax - ExposedInputMin, tsr_simd_halfC::Const(tsr_halfC(0.299f, 0.587f, 0.114f, 0.0)));
-			#else
-				tsr_simd_half ExposedInputLuminanceBox = dot(ExposedInputMax - ExposedInputMin, tsr_simd_halfC::Const(tsr_halfC(0.299f, 0.587f, 0.114f)));
-			#endif
 
-			bAliasingIsVisible = ExposedInputLuminanceBox > tsr_simd_half::Const(SPATIAL_ANTI_ALIASER_MIN_LUMIMANCE);
-		}
+		bAliasingIsVisible = IsAliasingVisible(ExposedInputMax - ExposedInputMin);
 
 		tsr_simd_halfC ClampedInput = Clamp3x3(ExposedInput, fastClamp(ExposedHistory, ExposedInputMin, ExposedInputMax));
 		tsr_simd_halfC ClampedHistory = AnnihilateToGuide3x3(ExposedHistory, /* Guide = */ ExposedInput);
@@ -860,22 +910,9 @@ void MainCS(
 			}
 			#endif // CONFIG_LDS_DESPILL
 
-			tsr_simd_halfC PrevMatch         = abs(ExposedInput - ExposedHistory);
-			tsr_simd_halfC ResurrectionMatch = abs(ExposedInput - ExposedResurrectedHistory);
+			tsr_simd_bool bResurrectionIsCloser = IsResurrectedFrameCloserThanPreviousFrame(ExposedInput, ExposedHistory, ExposedResurrectedHistory);
+			bool bComputeResurrection = IsWorthMeasuringRejectionOfResurrectedFrame(RejectionBlendFinal, bResurrectionIsCloser);
 
-			tsr_simd_halfC MatchDiff = PrevMatch - ResurrectionMatch;
-			#if CONFIG_SCENE_COLOR_ALPHA
-				tsr_simd_half TotalMatchDiff = BlurEven3x3(MatchDiff[0] + MatchDiff[1] + MatchDiff[2] + MatchDiff[3]);
-			#else
-				tsr_simd_half TotalMatchDiff = BlurEven3x3(MatchDiff[0] + MatchDiff[1] + MatchDiff[2]);
-			#endif
-
-			tsr_simd_bool bResurrectionIsCloser = MedianBool3x3(TotalMatchDiff > tsr_simd_half::Const(tsr_half(0.05 * 3.0)));
-			tsr_simd_bool bPrevIsNotGoodEnough = RejectionBlendFinal < tsr_simd_half::Const(0.5);
-			tsr_simd_bool bIsValidTilePixel = IsValidTilePixel<LANE_STRIDE_X, LANE_STRIDE_Y>(TileOverscan);
-			tsr_simd_bool bShouldComputeResurrection = and(and(bResurrectionIsCloser, bPrevIsNotGoodEnough), bIsValidTilePixel);
-
-			bool bComputeResurrection = AnyThread(AnyComponent(bShouldComputeResurrection));
 			BRANCH
 			if (bComputeResurrection)
 			{
@@ -893,9 +930,8 @@ void MainCS(
 					/* out */ ResurrectionRejectionBlendFinal,
 					/* out */ ResurrectionRejectionClampBlend);
 			
-				tsr_simd_bool bResurrectionIsBetter = (ResurrectionRejectionBlendFinal - RejectionBlendFinal) > tsr_simd_half::Const(0.1);
-				bResurrectHistory = and(bResurrectionIsCloser, bResurrectionIsBetter);
-				
+				bResurrectHistory = ShouldResurrectHistory(RejectionBlendFinal, ResurrectionRejectionBlendFinal, bResurrectionIsCloser);
+
 				#if CONFIG_LDS_DESPILL
 				PLATFORM_SPECIFIC_ISOLATE
 				{
@@ -944,23 +980,7 @@ void MainCS(
 		// Convert input to LDR
 		tsr_simd_halfC ExposedInput = AddPerceptionAdd(OriginalInput, /* DestPerceptionAdd = */ 1.0);
 	
-
-		tsr_simd_half BlendFinal;
-		// Completely discard the history
-		UNROLL_N(SIMD_SIZE)
-		for (uint SimdIndex = 0; SimdIndex < SIMD_SIZE; SimdIndex++)
-		{
-			bool bIsOffScreen = false;
-			bool bIsDisoccludedLocal = bIsDisoccluded.GetElement(SimdIndex);
-			tsr_half LocalRejectionBlendFinal = RejectionBlendFinal.GetElement(SimdIndex);
-
-			tsr_half TotalBlend = select(or(bIsOffScreen, bIsDisoccludedLocal), tsr_half(1.0), saturate(tsr_half(1.0) - LocalRejectionBlendFinal * tsr_half(4.0)));
-			tsr_half RejectedBlendFactor = max(tsr_half(TheoricBlendFactor), tsr_half(1.0) - LocalRejectionBlendFinal);
-
-			BlendFinal.SetElement(SimdIndex, max(RejectedBlendFactor, TotalBlend));
-		}
-		
-		FinalGuide = ExposedInput * tsr_simd_halfC::Vectorize(BlendFinal) - OriginalHistory * tsr_simd_halfC::Vectorize(BlendFinal - 1.0);
+		FinalGuide = UpdateGuide(ExposedInput, OriginalHistory, bIsDisoccluded, RejectionBlendFinal);
 	}
 	
 	tsr_simd_half InputUncertainty = VelocityEdge;
@@ -972,29 +992,17 @@ void MainCS(
 		DisableHistoryClamp        = select(and(bIsDisoccluded, not(bResurrectHistory)), tsr_simd_half::Const(0.0), min(RejectionClampBlend, min(InputUncertainty, Min3x3(OriginalHistoryUncertainty))));
 		IncreaseValidityMultiplier = select(and(bIsDisoccluded, not(bResurrectHistory)), tsr_simd_half::Const(0.0), RejectionClampBlend);
 	}
-
-	// Mark the top left corner of the tile for debugging purposes.
-	#if DEBUG_OUTPUT && DEBUG_DISABLE_PADDING == 2
-	if (GGroupThreadIndex == 0)
-	{
-		UNROLL_N(DEBUG_ARRAY_SIZE)
-		for (uint DebugId = 0; DebugId < DEBUG_ARRAY_SIZE; DebugId++)
-		{
-			Debug[DebugId].SetElement(/* SimdIndex = */ 0, 1.0);
-		}
-	}
-	#endif
-
+	
 	// Whether to run spatial anti-aliaser on the frame. Must match SpatialAntiAliasingLerp.
-	tsr_simd_bool bSpatialAntiAlias = and(bAliasingIsVisible, or(RejectionBlendFinal < tsr_simd_half::Const(tsr_half(0.25)), and(bIsDisoccluded, not(bResurrectHistory))));
+	tsr_simd_bool bSpatialAntiAlias = ShouldSpatialAntiAlias(bIsDisoccluded, bResurrectHistory, bAliasingIsVisible, RejectionBlendFinal);
 
 	// Output data.
 	PLATFORM_SPECIFIC_ISOLATE
 	{
 		UNROLL_N(SIMD_SIZE)
-		for (uint SimdIndex = 0; SimdIndex < SIMD_SIZE; SimdIndex++)
+		for (uint ElementIndex = 0; ElementIndex < SIMD_SIZE; ElementIndex++)
 		{
-			const tsr_short2 LaneSimdPixelOffset = GetElementPixelOffsetInTile(LANE_STRIDE_X, LANE_STRIDE_Y, LaneIndex, SimdIndex);
+			const tsr_short2 LaneSimdPixelOffset = GetElementPixelOffsetInTile(LANE_STRIDE_X, LANE_STRIDE_Y, LaneIndex, ElementIndex);
 			const tsr_short2 OutputPixelPos = ComputeElementOutputPixelPos(
 				/* LaneStride = */ uint2(LANE_STRIDE_X, LANE_STRIDE_Y),
 				TileOverscan,
@@ -1002,13 +1010,13 @@ void MainCS(
 				InputInfo_ViewportMax,
 				GroupId,
 				GroupThreadIndex,
-				SimdIndex);
+				ElementIndex);
 
-			tsr_halfC FinalGuideColor = FinalGuide.GetElement(SimdIndex);
+			tsr_halfC FinalGuideColor = FinalGuide.GetElement(ElementIndex);
 			
 			tsr_half BitMask = (
-				select(bIsValid.GetElement(SimdIndex),          tsr_half(0x1) * rcp(tsr_half(255.0)), tsr_half(0x0)) + 
-				select(bResurrectHistory.GetElement(SimdIndex), tsr_half(0x2) * rcp(tsr_half(255.0)), tsr_half(0x0)));
+				select(bIsValid.GetElement(ElementIndex),          tsr_half(0x1) * rcp(tsr_half(255.0)), tsr_half(0x0)) + 
+				select(bResurrectHistory.GetElement(ElementIndex), tsr_half(0x2) * rcp(tsr_half(255.0)), tsr_half(0x0)));
 
 			#if CONFIG_ENABLE_STOCASTIC_QUANTIZATION
 			{
@@ -1021,21 +1029,35 @@ void MainCS(
 
 			#if CONFIG_SCENE_COLOR_ALPHA
 				HistoryGuideOutput[tsr_short3(OutputPixelPos, 0)] = FinalGuideColor;
-				HistoryGuideOutput[tsr_short3(OutputPixelPos, 1)] = tsr_half4(InputUncertainty.GetElement(SimdIndex), 0, 0, 0);
+				HistoryGuideOutput[tsr_short3(OutputPixelPos, 1)] = tsr_half4(InputUncertainty.GetElement(ElementIndex), 0, 0, 0);
 			#else
-				HistoryGuideOutput[tsr_short3(OutputPixelPos, 0)] = tsr_half4(FinalGuideColor.xyz, InputUncertainty.GetElement(SimdIndex));
+				HistoryGuideOutput[tsr_short3(OutputPixelPos, 0)] = tsr_half4(FinalGuideColor.xyz, InputUncertainty.GetElement(ElementIndex));
 			#endif
 			HistoryRejectionOutput[OutputPixelPos] = tsr_half4(
-				RejectionBlendFinal.GetElement(SimdIndex) - tsr_half(0.5 * rcp(256.0)),
-				DisableHistoryClamp.GetElement(SimdIndex) - tsr_half(0.5 * rcp(256.0)),
-				tsr_half(255.5 * rcp(256.0)) - IncreaseValidityMultiplier.GetElement(SimdIndex),
+				RejectionBlendFinal.GetElement(ElementIndex) - tsr_half(0.5 * rcp(256.0)),
+				DisableHistoryClamp.GetElement(ElementIndex) - tsr_half(0.5 * rcp(256.0)),
+				tsr_half(255.5 * rcp(256.0)) - IncreaseValidityMultiplier.GetElement(ElementIndex),
 				BitMask);
-			AntiAliasMaskOutput[OutputPixelPos] = select(bSpatialAntiAlias.GetElement(SimdIndex), tsr_ushort(1), tsr_ushort(0));
+			AntiAliasMaskOutput[OutputPixelPos] = select(bSpatialAntiAlias.GetElement(ElementIndex), tsr_ushort(1), tsr_ushort(0));
 		}
-		
-		#if DEBUG_OUTPUT
+	}
+	
+	#if DEBUG_OUTPUT
+	{
+		// Mark the top left corner of the tile for debugging purposes.
+		#if DEBUG_DISABLE_PADDING == 2
+		if (GGroupThreadIndex == 0)
+		{
+			UNROLL_N(DEBUG_ARRAY_SIZE)
+			for (uint DebugId = 0; DebugId < DEBUG_ARRAY_SIZE; DebugId++)
+			{
+				Debug[DebugId].SetElement(/* ElementIndex = */ 0, 1.0);
+			}
+		}
+		#endif // DEBUG_DISABLE_PADDING == 2
+	
 		UNROLL_N(SIMD_SIZE)
-		for (uint SimdIndex = 0; SimdIndex < SIMD_SIZE; SimdIndex++)
+		for (uint ElementIndex = 0; ElementIndex < SIMD_SIZE; ElementIndex++)
 		{
 			const tsr_short2 OutputPixelPos = ComputeElementOutputPixelPos(
 				/* LaneStride = */ uint2(LANE_STRIDE_X, LANE_STRIDE_Y),
@@ -1044,7 +1066,7 @@ void MainCS(
 				InputInfo_ViewportMax,
 				GroupId,
 				GroupThreadIndex,
-				SimdIndex,
+				ElementIndex,
 				/* OutputDataOffset = */ int2(0, 0),
 				/* OutputResDivisor = */ 1,
 				/* bDebugDisablePadding = */ DEBUG_DISABLE_PADDING != 0);
@@ -1052,9 +1074,9 @@ void MainCS(
 			UNROLL_N(DEBUG_ARRAY_SIZE)
 			for (uint DebugId = 0; DebugId < DEBUG_ARRAY_SIZE; DebugId++)
 			{
-				DebugOutput[tsr_short3(OutputPixelPos, DebugId)] = float4(Debug[DebugId].GetElement(SimdIndex), 0.0);
+				DebugOutput[tsr_short3(OutputPixelPos, DebugId)] = float4(ResizeChannels<4>(Debug[DebugId]).GetElement(ElementIndex));
 			}
 		}
-		#endif
 	}
+	#endif // DEBUG_OUTPUT
 }
diff --git a/Engine/Shaders/Public/LaneVectorization.ush b/Engine/Shaders/Public/LaneVectorization.ush
--- a/Engine/Shaders/Public/LaneVectorization.ush
+++ b/Engine/Shaders/Public/LaneVectorization.ush
@@ -1465,6 +1465,25 @@ struct TLaneVector2D
 }; // TLaneVector2D
 
 
+CALL_SITE_DEBUGLOC
+template<uint DestVectorSize, typename ScalarType, uint VectorSize, uint SimdSizeX, uint SimdSizeY>
+TLaneVector2D<ScalarType, DestVectorSize, SimdSizeX, SimdSizeY> ResizeChannels(
+	TLaneVector2D<ScalarType, VectorSize, SimdSizeX, SimdSizeY> A)
+{
+	TLaneVector2D<ScalarType, DestVectorSize, SimdSizeX, SimdSizeY> R;
+	R.SetAllElements(ScalarType(0));
+	UNROLL
+	for (uint ElementIndex = 0; ElementIndex < SimdSizeX * SimdSizeY; ElementIndex++)
+	{
+		UNROLL
+		for (uint ComponentIndex = 0; ComponentIndex < VectorSize; ComponentIndex++)
+		{
+			R.Registers.SetElementComponent(ElementIndex, ComponentIndex, A.Registers.GetElementComponent(ElementIndex, ComponentIndex));
+		}
+	}
+	return R;
+}
+
 CALL_SITE_DEBUGLOC
 template<typename ScalarType, uint VectorSizeA, uint VectorSizeB, uint SimdSizeX, uint SimdSizeY>
 TLaneVector2D<ScalarType, VectorSizeA + VectorSizeB, SimdSizeX, SimdSizeY> Concatenate(
