Fix a bug of black screen on Mac mobile preview with mobile deferred.

diff --git a/Engine/Shaders/Private/DeferredShadingCommon.ush b/Engine/Shaders/Private/DeferredShadingCommon.ush
--- a/Engine/Shaders/Private/DeferredShadingCommon.ush
+++ b/Engine/Shaders/Private/DeferredShadingCommon.ush
@@ -766,7 +766,7 @@ void MobileFetchGBuffer(in float2 UV, out half4 GBufferA, out half4 GBufferB, ou
 	GBufferD = VulkanSubpassFetch4();
 #endif
 	SceneDepth = ConvertFromDeviceZ(VulkanSubpassDepthFetch());
-#elif METAL_PROFILE
+#elif (METAL_PROFILE && !MAC)
 	GBufferA = SubpassFetchRGBA_1(); 
 	GBufferB = SubpassFetchRGBA_2(); 
 	GBufferC = SubpassFetchRGBA_3(); 
