Merge and rename hair composition shaders.

diff --git a/Engine/Shaders/Private/HairStrands/HairStrandsComposition.usf b/Engine/Shaders/Private/HairStrands/HairStrandsComposition.usf
new file mode 100644
--- /dev/null
+++ b/Engine/Shaders/Private/HairStrands/HairStrandsComposition.usf
@@ -0,0 +1,378 @@
+// Copyright Epic Games, Inc. All Rights Reserved.
+
+#include "../Common.ush"
+#include "../SceneTexturesCommon.ush"
+#include "../DeferredShadingCommon.ush"
+#include "HairStrandsVisibilityCommon.ush"   
+
+////////////////////////////////////////////////////////////////////////////////////////////////
+
+#if SHADER_FASTRESOLVE_MASK
+
+Texture2D<uint> ResolveMaskTexture;
+
+void FastResolvePS(in FScreenVertexOutput Input)
+{
+	const uint2 PixelCoord = floor(Input.Position.xy);
+	const bool bNeedFastResolve = ResolveMaskTexture.Load(uint3(PixelCoord, 0)) > 0;
+
+	if (!bNeedFastResolve)
+		discard;
+}
+
+#endif // SHADER_FASTRESOLVE_MASK
+
+////////////////////////////////////////////////////////////////////////////////////////////////
+
+#if SHADER_COMPOSE_SAMPLE
+
+#include "../HeightFogCommon.ush"
+#include "../SkyAtmosphereCommon.ush"
+
+int2 OutputResolution;
+uint bComposeDofDepth;
+uint bHasHoldout;
+
+Texture2D<float>	HairDOFDepthTexture;
+Texture2D<float4>	HairLightingSampleBuffer;
+
+float4 EvaluateVolumetric(float3 WorldPosition, float SceneDepth)
+{
+	const float3 CameraRelative_WorldPosition = WorldPosition - LWCHackToFloat(PrimaryView.WorldCameraOrigin);
+	float4 HeightFogging = CalculateHeightFog(CameraRelative_WorldPosition);
+	float4 Fogging = HeightFogging;
+
+	if (FogStruct.ApplyVolumetricFog > 0)
+	{
+		const uint EyeIndex = 0;
+		float3 VolumeUV = ComputeVolumeUV_DEPRECATED(WorldPosition, LWCHackToFloat(PrimaryView.WorldToClip));
+		Fogging = CombineVolumetricFog(HeightFogging, VolumeUV, EyeIndex, SceneDepth);
+	}
+
+	Fogging.rgb *= View.PreExposure;
+
+	if (View.SkyAtmosphereApplyCameraAerialPerspectiveVolume > 0.0f)
+	{
+		float4 NDCPosition = mul(float4(WorldPosition.xyz, 1), LWCHackToFloat(PrimaryView.WorldToClip));
+
+		// Sample the aerial perspective (AP).
+		Fogging = GetAerialPerspectiveLuminanceTransmittanceWithFogOver(
+			View.RealTimeReflectionCapture, 
+			View.SkyAtmosphereCameraAerialPerspectiveVolumeSizeAndInvSize,
+			NDCPosition, 
+			CameraRelative_WorldPosition * CM_TO_SKY_UNIT,
+			View.CameraAerialPerspectiveVolume, 
+			View.CameraAerialPerspectiveVolumeSampler,
+			View.SkyAtmosphereCameraAerialPerspectiveVolumeDepthResolutionInv,
+			View.SkyAtmosphereCameraAerialPerspectiveVolumeDepthResolution,
+			View.SkyAtmosphereAerialPerspectiveStartDepthKm,
+			View.SkyAtmosphereCameraAerialPerspectiveVolumeDepthSliceLengthKm,
+			View.SkyAtmosphereCameraAerialPerspectiveVolumeDepthSliceLengthKmInv,
+			View.OneOverPreExposure, 
+			Fogging);
+	}
+
+	return Fogging;
+}
+
+void ComposeSamplePS(
+	in FScreenVertexOutput Input,
+	out float4 OutColor : SV_Target0,
+	out float  OutDepth : SV_DEPTH)
+#if PERMUTATION_DEBUG == 0
+{
+	OutColor = 0.0f;
+	const uint3 PixelCoord = uint3(floor(Input.Position.xy), 0);
+	
+	const float HairDeviceZ = HairStrands.HairOnlyDepthTexture.Load(PixelCoord).x;
+
+	const FNodeDesc NodeDesc = DecodeNodeDesc(HairStrands.HairSampleOffset.Load(PixelCoord));
+	if (NodeDesc.Count == 0)
+	{
+		discard;
+	}
+
+	const float PixelCoverage = min(HairStrands.HairCoverageTexture.Load(PixelCoord), 1);
+	if (PixelCoverage == 0)
+	{
+		discard;
+	}
+
+	const uint TotalNodeCount = HairStrands.HairSampleCount.Load(uint3(0, 0, 0));
+	const uint Resolution = uint(ceil(sqrt(float(TotalNodeCount))));
+
+	const float3 ClosetPointWorldPosition = LWCHackToFloat(SvPositionToWorld(float4(Input.Position.xy, HairDeviceZ, 1.0)));
+	const float4 Volumetric = EvaluateVolumetric(ClosetPointWorldPosition, ConvertFromDeviceZ(HairDeviceZ));
+
+	bool bHoldout = false;
+	float3 LocalAccColor = 0;
+	LOOP
+	for (uint SampleIt = 0; SampleIt < NodeDesc.Count; SampleIt++)
+	{
+		#if POST_PROCESS_ALPHA
+		if (bHasHoldout)
+		{
+			const FHairSample Sample = UnpackHairSample(HairStrands.HairSampleData[NodeDesc.Offset + SampleIt]);
+			bHoldout = bHoldout || HasHairFlags(Sample.Flags, HAIR_FLAGS_HOLDOUT);
+		}
+		#endif
+
+		const uint LocalOffset = NodeDesc.Offset + SampleIt;
+		const uint2 Coord = uint2(LocalOffset % Resolution, LocalOffset / Resolution);
+		LocalAccColor += HairLightingSampleBuffer.Load(uint3(Coord, 0)).xyz;
+	}
+	OutColor.rgb = LocalAccColor * Volumetric.a + Volumetric.rgb;
+
+	// If any sample support holdout, then the lighting contribution is set to 0, to mask out the current hair pixel
+	OutColor.rgb *= bHoldout ? 0 : PixelCoverage;
+	OutColor.a = PixelCoverage;
+
+	OutDepth = HairDeviceZ;
+	if (bComposeDofDepth)
+	{
+		OutDepth = HairDOFDepthTexture.Load(PixelCoord);
+	}
+}
+#else // PERMUTATION_DEBUG == 1
+{
+	OutColor = 0.0f;
+	const uint3 PixelCoord = uint3(floor(Input.Position.xy), 0);
+
+	const FNodeDesc NodeDesc = DecodeNodeDesc(HairStrands.HairSampleOffset.Load(PixelCoord));
+	if (NodeDesc.Count == 0)
+	{
+		discard;
+	}
+
+	const float PixelCoverage = min(HairStrands.HairCoverageTexture.Load(PixelCoord), 1);
+	if (PixelCoverage == 0)
+		discard;
+
+	const float HairDeviceZ = HairStrands.HairOnlyDepthTexture.Load(PixelCoord).x;
+
+	float3 LocalAccColor = 0;
+	LOOP
+	for (uint SampleIt = 0; SampleIt < NodeDesc.Count; SampleIt++)
+	{
+		const FPackedHairSample PackedSample = HairStrands.HairSampleData[NodeDesc.Offset + SampleIt];
+		const FHairSample Sample = UnpackHairSample(PackedSample);
+		const float LocalCoverage = From8bitCoverage(Sample.Coverage8bit);
+		LocalAccColor += LocalCoverage * Sample.BaseColor;
+	}
+	OutColor.rgb = LocalAccColor;
+
+	OutColor.rgb *= PixelCoverage;
+	OutColor.a = PixelCoverage;
+	OutDepth = HairDeviceZ;
+}
+#endif
+
+#endif // SHADER_COMPOSE_SAMPLE
+
+////////////////////////////////////////////////////////////////////////////////////////////////
+
+#if SHADER_HOLDOUT
+
+uint bComposeDofDepth;
+Texture2D<float> HairDOFDepthTexture;
+
+void HoldoutPS(
+	in FScreenVertexOutput Input,
+	out float4 OutColor : SV_Target0,
+	out float  OutDepth : SV_DEPTH)
+{
+#if POST_PROCESS_ALPHA
+	OutColor = 0.0f;
+	const uint3 PixelCoord = uint3(floor(Input.Position.xy), 0);
+	
+	const float HairDeviceZ = HairStrands.HairOnlyDepthTexture.Load(PixelCoord).x;
+
+	const FNodeDesc NodeDesc = DecodeNodeDesc(HairStrands.HairSampleOffset.Load(PixelCoord));
+	if (NodeDesc.Count == 0)
+	{
+		discard;
+	}
+
+	const float PixelCoverage = min(HairStrands.HairCoverageTexture.Load(PixelCoord), 1);
+	if (PixelCoverage == 0)
+	{
+		discard;
+	}
+
+	bool bHoldout = false;
+
+	LOOP
+	for (uint SampleIt = 0; SampleIt < NodeDesc.Count; SampleIt++)
+	{
+		const FHairSample Sample = UnpackHairSample(HairStrands.HairSampleData[NodeDesc.Offset + SampleIt]);
+		bHoldout = bHoldout || HasHairFlags(Sample.Flags, HAIR_FLAGS_HOLDOUT);
+	}
+
+	if (!bHoldout)
+	{
+		discard;
+	}
+
+	OutColor.rgb = 0;
+	OutColor.a = 1;  //PixelCoverage;
+	OutDepth     = HairDeviceZ;
+	if (bComposeDofDepth)
+	{
+		OutDepth = HairDOFDepthTexture.Load(PixelCoord);
+	}
+#else
+	OutColor = 0;
+	OutDepth = 0;
+	discard;
+#endif
+}
+
+#endif // SHADER_HOLDOUT
+
+////////////////////////////////////////////////////////////////////////////////////////////////
+
+#if SHADER_DOFDEPTH
+#include "../CircleDOFCommon.ush"
+
+// Included from DeferredLightingCommon now
+//Texture2D<float>	SceneDepthTexture;
+
+void DOFDepthPS(
+	in FScreenVertexOutput Input,
+	out float OutTexture : SV_Target0)
+{
+	OutTexture = 0.0f;
+	const uint3 PixelCoord = uint3(floor(Input.Position.xy), 0);
+
+	const float PixelCoverage = min(HairStrands.HairCoverageTexture.Load(PixelCoord), 1);
+	if (PixelCoverage == 0)
+	{
+		return;
+	}
+	   
+	const float HalfResToFullRes = 2;
+
+	float SceneDeviceZ = 0;
+	float SceneDepth = 0;
+	float SceneCoC = 0; // Radius in pixel
+	{
+		SceneDeviceZ = SceneDepthTexture.Load(PixelCoord).x;
+		SceneDepth = ConvertFromDeviceZ(SceneDeviceZ);
+		SceneDeviceZ = 1 - SceneDeviceZ;
+		SceneCoC = DepthToCoc(SceneDepth) * HalfResToFullRes;
+	}
+
+	float HairLuminance = 0;
+	float HairDeviceZ = 0;
+	float HairDepth = 0;
+	float HairCoC = 0; // Radius in pixel
+	{
+		HairDeviceZ = HairStrands.HairOnlyDepthTexture.Load(PixelCoord).x;
+		HairDepth = ConvertFromDeviceZ(HairDeviceZ);
+		HairDeviceZ = 1 - HairDeviceZ;
+		HairCoC = DepthToCoc(HairDepth) * HalfResToFullRes;
+	}
+
+	const float OutDeviceZ = lerp(SceneDeviceZ, HairDeviceZ, PixelCoverage);
+	OutTexture = 1 - OutDeviceZ;
+}
+
+#endif // SHADER_DOFDEPTH
+
+////////////////////////////////////////////////////////////////////////////////////////////////
+
+#if SHADER_WRITE_GBUFFER
+
+#include "HairStrandsGBufferCommon.ush"
+
+uint bWriteDummyData;
+
+void MainPS(
+	in float2 UV : TEXCOORD
+	, in float4 SvPosition: SV_Position
+	, out float4 OutGBufferA : SV_Target0
+	, out float4 OutGBufferB : SV_Target1
+#if PERMUTATION_OUTPUT_TYPE == 1
+	, out float4 OutGBufferC : SV_Target2
+	, out float4 OutGBufferD : SV_Target3
+	, out float4 OutGBufferE : SV_Target4
+	, out float  OutDepth    : SV_DEPTH
+#endif
+)
+{
+	ResolvedView = ResolveView();
+	OutGBufferA = 0;
+	OutGBufferB = 0;
+#if PERMUTATION_OUTPUT_TYPE == 0
+	float4 OutGBufferC = 0;
+	float4 OutGBufferD = 0;
+	float4 OutGBufferE = 0;
+	float  OutDepth = 0;
+#endif
+
+	const uint3 PixelCoord = uint3(floor(SvPosition.xy), 0);
+	const FNodeDesc NodeDesc = DecodeNodeDesc(HairStrands.HairSampleOffset.Load(PixelCoord));
+	if (NodeDesc.Count == 0)
+	{
+		discard;
+	}
+
+	if (bWriteDummyData > 0)
+	{
+		const float3 Tangent = float3(0, 0, 1);
+		OutGBufferA = float4(Tangent, 0);
+		OutGBufferB = float4(0, 0, 0, EncodeShadingModelIdAndSelectiveOutputMask(SHADINGMODELID_HAIR, 0));
+	}
+	else
+	{
+		float3 Samples_Tangent = 0;
+		float  Samples_Specular = 0;
+		float3 Samples_BaseColor = 0;
+		float  Samples_Roughnesss = 0;
+		uint   Samples_LightChannelMask = 0;
+		float  Samples_Backlit = 0;
+		float  Samples_Depth = 0;
+		LOOP
+		for (uint NodeIt = 0; NodeIt < NodeDesc.Count; ++NodeIt)
+		{
+			const uint NodeOffset = NodeDesc.Offset + NodeIt;
+			const FPackedHairSample Data = HairStrands.HairSampleData[NodeOffset];
+			const FHairSample Sample = UnpackHairSample(Data);
+
+			Samples_Tangent		+= Sample.Tangent;
+			Samples_Specular	+= Sample.Specular;
+			Samples_BaseColor	+= Sample.BaseColor;
+			Samples_Roughnesss	+= Sample.Roughness;
+			Samples_Backlit		+= Sample.Backlit;
+			if (NodeIt == 0)
+			{
+				Samples_LightChannelMask = Sample.LightChannelMask;
+			}
+		}
+		Samples_Depth = HairStrands.HairOnlyDepthTexture.Load(PixelCoord).x;
+
+		const float InvSampleCount = 1.f / max(1u, NodeDesc.Count);
+		Samples_Tangent		 = normalize(Samples_Tangent);
+		Samples_Specular	*= InvSampleCount;
+		Samples_BaseColor	*= InvSampleCount;
+		Samples_Roughnesss	*= InvSampleCount;
+		Samples_Backlit		*= InvSampleCount;
+
+		WriteGBuffer(
+			Samples_Tangent,
+			Samples_Specular,
+			Samples_BaseColor,
+			Samples_Roughnesss,
+			Samples_LightChannelMask,
+			Samples_Backlit,
+			Samples_Depth,
+			OutGBufferA,
+			OutGBufferB,
+			OutGBufferC,
+			OutGBufferD,
+			OutGBufferE,
+			OutDepth);
+	}
+}
+
+#endif // SHADER_WRITE_GBUFFER
\ No newline at end of file
diff --git a/Engine/Shaders/Private/HairStrands/HairStrandsGBufferWrite.usf b/Engine/Shaders/Private/HairStrands/HairStrandsGBufferWrite.usf
deleted file mode 100644
diff --git a/Engine/Shaders/Private/HairStrands/HairStrandsVisibilityComposeSubPixelPS.usf b/Engine/Shaders/Private/HairStrands/HairStrandsVisibilityComposeSubPixelPS.usf
deleted file mode 100644
diff --git a/Engine/Source/Runtime/Renderer/Private/HairStrands/HairStrandsComposition.cpp b/Engine/Source/Runtime/Renderer/Private/HairStrands/HairStrandsComposition.cpp
--- a/Engine/Source/Runtime/Renderer/Private/HairStrands/HairStrandsComposition.cpp
+++ b/Engine/Source/Runtime/Renderer/Private/HairStrands/HairStrandsComposition.cpp
@@ -179,7 +179,7 @@ public:
 	}
 };
 
-IMPLEMENT_GLOBAL_SHADER(FHairHoldoutPS, "/Engine/Private/HairStrands/HairStrandsVisibilityComposeSubPixelPS.usf", "HoldoutPS", SF_Pixel);
+IMPLEMENT_GLOBAL_SHADER(FHairHoldoutPS, "/Engine/Private/HairStrands/HairStrandsComposition.usf", "HoldoutPS", SF_Pixel);
 
 static void AddHairHoldoutPass(
 	FRDGBuilder& GraphBuilder,
@@ -247,7 +247,7 @@ public:
 	}
 };
 
-IMPLEMENT_GLOBAL_SHADER(FHairVisibilityComposeSamplePS, "/Engine/Private/HairStrands/HairStrandsVisibilityComposeSubPixelPS.usf", "ComposeSamplePS", SF_Pixel);
+IMPLEMENT_GLOBAL_SHADER(FHairVisibilityComposeSamplePS, "/Engine/Private/HairStrands/HairStrandsComposition.usf", "ComposeSamplePS", SF_Pixel);
 
 static void AddHairVisibilityComposeSamplePass(
 	FRDGBuilder& GraphBuilder,
@@ -319,7 +319,7 @@ public:
 	}
 };
 
-IMPLEMENT_GLOBAL_SHADER(FHairDOFDepthPS, "/Engine/Private/HairStrands/HairStrandsVisibilityComposeSubPixelPS.usf", "DOFDepthPS", SF_Pixel);
+IMPLEMENT_GLOBAL_SHADER(FHairDOFDepthPS, "/Engine/Private/HairStrands/HairStrandsComposition.usf", "DOFDepthPS", SF_Pixel);
 
 static FRDGTextureRef AddHairDOFDepthPass(
 	FRDGBuilder& GraphBuilder,
@@ -384,7 +384,7 @@ public:
 	}
 };
 
-IMPLEMENT_GLOBAL_SHADER(FHairVisibilityFastResolveMaskPS, "/Engine/Private/HairStrands/HairStrandsVisibilityComposeSubPixelPS.usf", "FastResolvePS", SF_Pixel);
+IMPLEMENT_GLOBAL_SHADER(FHairVisibilityFastResolveMaskPS, "/Engine/Private/HairStrands/HairStrandsComposition.usf", "FastResolvePS", SF_Pixel);
 
 static void AddHairVisibilityFastResolveMaskPass(
 	FRDGBuilder& GraphBuilder,
@@ -457,10 +457,11 @@ public:
 		FGlobalShader::ModifyCompilationEnvironment(Parameters, OutEnvironment);
 		OutEnvironment.SetRenderTargetOutputFormat(0, PF_B8G8R8A8);
 		OutEnvironment.SetRenderTargetOutputFormat(1, PF_FloatRGBA);
+		OutEnvironment.SetDefine(TEXT("SHADER_WRITE_GBUFFER"), 1);
 	}
 };
 
-IMPLEMENT_GLOBAL_SHADER(FHairVisibilityGBufferWritePS, "/Engine/Private/HairStrands/HairStrandsGBufferWrite.usf", "MainPS", SF_Pixel);
+IMPLEMENT_GLOBAL_SHADER(FHairVisibilityGBufferWritePS, "/Engine/Private/HairStrands/HairStrandsComposition.usf", "MainPS", SF_Pixel);
 
 static void AddHairVisibilityGBufferWritePass(
 	FRDGBuilder& GraphBuilder,
