Skip unnecessary page marking and SMRT work for shadow receiver pixels behind rect lights

From Pari at The Coalition

diff --git a/Engine/Shaders/Private/VirtualShadowMaps/VirtualShadowMapPageMarking.usf b/Engine/Shaders/Private/VirtualShadowMaps/VirtualShadowMapPageMarking.usf
--- a/Engine/Shaders/Private/VirtualShadowMaps/VirtualShadowMapPageMarking.usf
+++ b/Engine/Shaders/Private/VirtualShadowMaps/VirtualShadowMapPageMarking.usf
@@ -475,6 +475,19 @@ void GeneratePageFlagsFromPixels(
 				continue;
 			}
 
+			const uint LightTypeAndPackedShadowMapChannelMask = asuint(LightData.LightDirectionAndShadowMask.w);
+			const uint LightType = UnpackLightType(LightTypeAndPackedShadowMapChannelMask);
+			
+			const bool bIsRectLight = LightType == LIGHT_TYPE_RECT;
+			if (bIsRectLight)
+			{
+				const bool bIsBehindRectLight = dot(ToLight, LightData.LightDirectionAndShadowMask.xyz) < 0;
+				if (bIsBehindRectLight)
+				{
+					continue;
+				}
+			}
+
 			// TODO: Precise radius test necessary?
 
 		#if PERMUTATION_INPUT_TYPE != INPUT_TYPE_HAIRSTRANDS
diff --git a/Engine/Shaders/Private/VirtualShadowMaps/VirtualShadowMapProjection.usf b/Engine/Shaders/Private/VirtualShadowMaps/VirtualShadowMapProjection.usf
--- a/Engine/Shaders/Private/VirtualShadowMaps/VirtualShadowMapProjection.usf
+++ b/Engine/Shaders/Private/VirtualShadowMaps/VirtualShadowMapProjection.usf
@@ -232,6 +232,16 @@ FVirtualShadowMapSampleResult ProjectLight(
 	bool bInLightCone = SpotAttenuationMask(L, -Light.Direction, Light.SpotAngles) > 0.0f;
 
 	bool bInLightRegion = bInLightRadius && bInLightCone;
+
+	const bool bIsRectLight = Light.RectLightBarnLength > -2.0f;
+	if (bIsRectLight)
+	{
+		const bool bIsBehindRectLight = dot(ToLight, Light.Direction) < 0;
+		if (bIsBehindRectLight)
+		{
+			bInLightRegion = false;
+		}
+	}
 #endif // DIRECTIONAL_LIGHT
 
 	FVirtualShadowMapSampleResult Result = InitVirtualShadowMapSampleResult();
