SampledDirectLighting - fixed subsurface shading models

diff --git a/Engine/Shaders/Private/SampledDirectLighting/SampledDirectLightingSpatial.usf b/Engine/Shaders/Private/SampledDirectLighting/SampledDirectLightingSpatial.usf
--- a/Engine/Shaders/Private/SampledDirectLighting/SampledDirectLightingSpatial.usf
+++ b/Engine/Shaders/Private/SampledDirectLighting/SampledDirectLightingSpatial.usf
@@ -113,22 +113,27 @@ void SampledDirectLightingSpatialCS(
 			const float3 V = normalize(View.TranslatedWorldCameraOrigin - TranslatedWorldPosition);
 			const float NoV = saturate(dot(N, V));
 
-			// #sdl_todo: handle alpha channel
 			// Final pass outputs composites irradiance and outputs it to scene color
-			float3 Irradiance = GBuffer.DiffuseColor * (DiffuseLightingSum / DiffuseLightingWeightSum);
-			Irradiance += (SpecularLightingSum / SpecularLightingWeightSum) * EnvBRDF(GBuffer.SpecularColor, GBuffer.Roughness, NoV);
+			float3 DiffuseLighting = GBuffer.DiffuseColor * (DiffuseLightingSum / DiffuseLightingWeightSum);
+			float3 SpecularLighting = (SpecularLightingSum / SpecularLightingWeightSum) * EnvBRDF(GBuffer.SpecularColor, GBuffer.Roughness, NoV);
+
+			// Composite diffuse and specular into scene color
+			const bool bNeedsSeparateSubsurfaceLightAccumulation = UseSubsurfaceProfile(GBuffer.ShadingModelID);
+			FLightAccumulator LightAccumulator = (FLightAccumulator)0;
+			LightAccumulator_AddSplit(LightAccumulator, DiffuseLighting, SpecularLighting, /*ScatterableLight*/ DiffuseLighting, /*CommonMultiplier*/ 1.0f, bNeedsSeparateSubsurfaceLightAccumulation);
+			float4 AccumulatedSceneColor = LightAccumulator_GetResult(LightAccumulator);
 
 			#if DEBUG_MODE
 			if (bDebug)
 			{
 				Print(Context, TEXT("Spatial pass"));
 				Newline(Context);
-				Print(Context, TEXT("Irradiance: "));
-				Print(Context, Irradiance);
+				Print(Context, TEXT("AccumulatedSceneColor: "));
+				Print(Context, AccumulatedSceneColor);
 			}
 			#endif
 
-			RWSceneColor[ScreenCoord] = RWSceneColor[ScreenCoord] + float4(Irradiance, 0.0f);
+			RWSceneColor[ScreenCoord] = RWSceneColor[ScreenCoord] + AccumulatedSceneColor;
 		}
 	}
 }
\ No newline at end of file
