SampledDirectLighting
* Implemented support for unshadowed lights
* Fixed demodulate/modulate for low radiance

diff --git a/Engine/Shaders/Private/SampledDirectLighting/SampledDirectLighting.ush b/Engine/Shaders/Private/SampledDirectLighting/SampledDirectLighting.ush
--- a/Engine/Shaders/Private/SampledDirectLighting/SampledDirectLighting.ush
+++ b/Engine/Shaders/Private/SampledDirectLighting/SampledDirectLighting.ush
@@ -205,4 +205,16 @@ uint MaxShadingTilesPerGridCell;
 uint GetShadingTileGridOffset(uint2 ShadingTileCoord, uint IndexInCell)
 {
 	return (ShadingTileCoord.x + ShadingTileCoord.y * ShadingTileGridSize.x) * MaxShadingTilesPerGridCell + IndexInCell;
+}
+
+// Remove texture detail before denoising
+float3 SafeDemodulate(float3 Lighting, float3 Texture)
+{
+	return Lighting / max(Texture, 0.001f);
+}
+
+// Restore texture detail after denoising
+float3 SafeModulate(float3 Lighting, float3 Texture)
+{
+	return Lighting * max(Texture, 0.001f);
 }
\ No newline at end of file
diff --git a/Engine/Shaders/Private/SampledDirectLighting/SampledDirectLightingSampling.usf b/Engine/Shaders/Private/SampledDirectLighting/SampledDirectLightingSampling.usf
--- a/Engine/Shaders/Private/SampledDirectLighting/SampledDirectLightingSampling.usf
+++ b/Engine/Shaders/Private/SampledDirectLighting/SampledDirectLightingSampling.usf
@@ -408,8 +408,12 @@ void GenerateSamplesCS(
 
 				if (LightSample.LocalLightIndex != INVALID_LOCAL_LIGHT_INDEX)
 				{
+					const FLocalLightData LocalLightData = GetLocalLightDataNonStereo(LightSample.LocalLightIndex);
+					const FLightSceneData LightSceneData = GetLightSceneData(LocalLightData.LightSceneId);
+					const bool bCastShadows = UnpackCastShadow(LightSceneData.LightTypeAndShadowMapChannelMaskPacked);
+
 					LightSample.bVisible = true;
-					LightSample.bCompleted = false;
+					LightSample.bCompleted = bCastShadows ? false : true;
 					LightSample.Weight = WeightSum / (NUM_SAMPLES_PER_PIXEL_1D * LightSample.Weight);
 				}
 
diff --git a/Engine/Shaders/Private/SampledDirectLighting/SampledDirectLightingShading.usf b/Engine/Shaders/Private/SampledDirectLighting/SampledDirectLightingShading.usf
--- a/Engine/Shaders/Private/SampledDirectLighting/SampledDirectLightingShading.usf
+++ b/Engine/Shaders/Private/SampledDirectLighting/SampledDirectLightingShading.usf
@@ -19,14 +19,6 @@ StructuredBuffer<uint> TileAllocator;
 StructuredBuffer<uint> TileData;
 uint TileDataStride;
 
-float3 SafeDemodulate(float3 L, float3 C)
-{
-	L.x = C.x >= 0.01f ? L.x / C.x : L.x;
-	L.y = C.y >= 0.01f ? L.y / C.y : L.y;
-	L.z = C.z >= 0.01f ? L.z / C.z : L.z;
-	return L;
-}
-
 /**
  * Apply all lights per pixel to affected tiles using sampled weights from upscaled shading tiles
  */
@@ -53,6 +45,9 @@ void ShadeLightSamplesCS(
 			float2 ScreenUV = (ScreenCoord + 0.5f) * View.BufferSizeAndInvSize.zw;
 			FGBufferData GBuffer = GetGBufferData(ScreenUV);
 
+			// #sdl_todo: handle this path when shadows are disabled
+			float ScreenSpaceAO = 1.0f;//Texture2DSampleLevel(SceneTexturesStruct.ScreenSpaceAOTexture, SceneTexturesStruct_ScreenSpaceAOTextureSampler, ScreenUV, 0).x;
+
 			float3 DiffuseLighting = 0.0f;
 			float3 SpecularLighting = 0.0f;
 
@@ -125,7 +120,7 @@ void ShadeLightSamplesCS(
 						float4 LightAttenuation = 1.0f;
 						float Dither = 0.5f;
 						float SurfaceShadow = 1;
-						float AmbientOcclusion = 1.0f;
+						float AmbientOcclusion = ScreenSpaceAO;
 
 						// #sdl_todo: Move more shading models to TILE_MODE_SIMPLE_SHADING
 						#if TILE_TYPE == TILE_MODE_SIMPLE_SHADING
@@ -133,7 +128,6 @@ void ShadeLightSamplesCS(
 							GBuffer.ShadingModelID = SHADINGMODELID_DEFAULT_LIT;
 						}
 						#endif
-						GBuffer.ShadingModelID = SHADINGMODELID_DEFAULT_LIT;
 						LightData.ShadowedBits = 0;
 
 						FDeferredLightingSplit SplitLighting = GetDynamicLightingSplit(
@@ -164,8 +158,8 @@ void ShadeLightSamplesCS(
 				const float3 SpecularEnv = EnvBRDF(GBuffer.SpecularColor, GBuffer.Roughness, NoV);
 
 				// #sdl_todo: demodulate earlier during BRDF evaluation
-				DiffuseLighting.xyz = SafeDemodulate(DiffuseLighting.xyz, GBuffer.DiffuseColor);
-				SpecularLighting.xyz = SafeDemodulate(SpecularLighting.xyz, SpecularEnv);
+				DiffuseLighting = SafeDemodulate(DiffuseLighting, GBuffer.DiffuseColor);
+				SpecularLighting = SafeDemodulate(SpecularLighting, SpecularEnv);
 			}
 
 			#if DEBUG_MODE
diff --git a/Engine/Shaders/Private/SampledDirectLighting/SampledDirectLightingSpatial.usf b/Engine/Shaders/Private/SampledDirectLighting/SampledDirectLightingSpatial.usf
--- a/Engine/Shaders/Private/SampledDirectLighting/SampledDirectLightingSpatial.usf
+++ b/Engine/Shaders/Private/SampledDirectLighting/SampledDirectLightingSpatial.usf
@@ -112,10 +112,11 @@ void SampledDirectLightingSpatialCS(
 			const float3 N = GBuffer.WorldNormal;
 			const float3 V = normalize(View.TranslatedWorldCameraOrigin - TranslatedWorldPosition);
 			const float NoV = saturate(dot(N, V));
+			const float3 SpecularEnv = EnvBRDF(GBuffer.SpecularColor, GBuffer.Roughness, NoV);
 
 			// Final pass outputs composites irradiance and outputs it to scene color
-			float3 DiffuseLighting = GBuffer.DiffuseColor * (DiffuseLightingSum / DiffuseLightingWeightSum);
-			float3 SpecularLighting = (SpecularLightingSum / SpecularLightingWeightSum) * EnvBRDF(GBuffer.SpecularColor, GBuffer.Roughness, NoV);
+			float3 DiffuseLighting = SafeModulate(DiffuseLightingSum / DiffuseLightingWeightSum, GBuffer.DiffuseColor);
+			float3 SpecularLighting = SafeModulate(SpecularLightingSum / SpecularLightingWeightSum, SpecularEnv);
 
 			// Composite diffuse and specular into scene color
 			const bool bNeedsSeparateSubsurfaceLightAccumulation = UseSubsurfaceProfile(GBuffer.ShadingModelID);
diff --git a/Engine/Shaders/Private/SampledDirectLighting/SampledDirectLightingTemporal.usf b/Engine/Shaders/Private/SampledDirectLighting/SampledDirectLightingTemporal.usf
--- a/Engine/Shaders/Private/SampledDirectLighting/SampledDirectLightingTemporal.usf
+++ b/Engine/Shaders/Private/SampledDirectLighting/SampledDirectLightingTemporal.usf
@@ -203,28 +203,16 @@ void SampledDirectLightingTemporalAccumulationCS(
 			DiffuseLighting = MakeFinite(DiffuseLighting);
 			SpecularLighting = MakeFinite(SpecularLighting);
 
-			float3 TranslatedWorldPosition = GetTranslatedWorldPositionFromScreenUV(ScreenUV, GBuffer.Depth);
-			const float3 N = GBuffer.WorldNormal;
-			const float3 V = normalize(View.TranslatedWorldCameraOrigin - TranslatedWorldPosition);
-			const float NoV = saturate(dot(N, V));
-
-			// #sdl_todo: handle alpha channel
-			// Reverse demodulation and combine lighting components
-			float3 Irradiance = GBuffer.DiffuseColor * DiffuseLighting + SpecularLighting * EnvBRDF(GBuffer.SpecularColor, GBuffer.Roughness, NoV);
-
-			float DiffuseVariance = abs(DiffuseSecondMoment - Pow2(Luminance(DiffuseLighting)));
-			float SpecularVariance = max(SpecularSecondMoment - Pow2(Luminance(SpecularLighting)), 0.0f);
-
 			#if DEBUG_MODE
 			if (bDebug)
 			{
+				float DiffuseVariance = max(DiffuseSecondMoment - Pow2(Luminance(DiffuseLighting)), 0.0f);
+				float SpecularVariance = max(SpecularSecondMoment - Pow2(Luminance(SpecularLighting)), 0.0f);
+
 				Newline(Context);
 				Print(Context, TEXT("NumFramesAccumulated: "));
 				Print(Context, NumFramesAccumulated);
 				Newline(Context);
-				Print(Context, TEXT("Irradiance: "));
-				Print(Context, Irradiance);
-				Newline(Context);
 				Print(Context, TEXT("DiffuseVariance:  "));
 				Print(Context, DiffuseVariance);
 				Newline(Context);
