Refactors TSR RejectShading to save register pressure on convolutions to be able to do even more in interest of quality

diff --git a/Engine/Shaders/Private/TemporalSuperResolution/TSRRejectShading.usf b/Engine/Shaders/Private/TemporalSuperResolution/TSRRejectShading.usf
--- a/Engine/Shaders/Private/TemporalSuperResolution/TSRRejectShading.usf
+++ b/Engine/Shaders/Private/TemporalSuperResolution/TSRRejectShading.usf
@@ -25,10 +25,6 @@
 
 #define CONFIG_HISTORY_RESURRECTION (DIM_HISTORY_RESURRECTION)
 
-#define CONFIG_MANUAL_LDS_SPILL 0
-
-#define CONFIG_MINIMISE_VGPRS 1
-
 
 //------------------------------------------------------- DERIVED
 	
@@ -87,9 +83,10 @@
 		#define TENSOR_REGISTER_LAYOUT PLATFORM_SPECIFIC_TENSOR_REGISTER_LAYOUT
 	#endif
 	
-	/** Manually despill to LDS on RDNA platforms with history resurrection to save more register pressure. */
-	#if CONFIG_MINIMISE_VGPRS && CONFIG_COMPILE_FP16 && !CONFIG_SCENE_COLOR_ALPHA
+	/** Manually despill to LDS to save more register pressure. */
+	#if CONFIG_COMPILE_FP16 && !CONFIG_SCENE_COLOR_ALPHA
 		#define CONFIG_LDS_DESPILL 1
+		#define CONFIG_CACHE_REFETCH 1
 	#endif
 
 #else
@@ -106,16 +103,17 @@
 	#define CONFIG_LDS_DESPILL 0
 #endif
 
+/** Whether should minimise register pressure by reloading some data from cache using memory fetches. */
+#ifndef CONFIG_CACHE_REFETCH
+	#define CONFIG_CACHE_REFETCH 0
+#endif
+
 /** Manually despill to LDS on RDNA platforms to save more register pressure. */
 #if CONFIG_LDS_DESPILL
 	#define LDS_DESPILL_THREAD_COUNT (WAVE_COUNT_X * WAVE_COUNT_Y) * (LANE_COUNT_X * LANE_COUNT_Y)
 	#define LDS_DESPILL_OFFSET       (LANE_STRIDE_X * LANE_STRIDE_Y) * GetVectorDwordSize(tsr_half, 4)
 
-	#if CONFIG_HISTORY_RESURRECTION
-		#define LDS_DESPILL_OFFSET_COUNT 2
-	#else
-		#define LDS_DESPILL_OFFSET_COUNT 1
-	#endif
+	#define LDS_DESPILL_OFFSET_COUNT 2
 	
 	#define LDS_DESPILL_DWORD_COUNT  LDS_DESPILL_THREAD_COUNT * LDS_DESPILL_OFFSET * LDS_DESPILL_OFFSET_COUNT
 #endif
@@ -136,7 +134,8 @@
 #define tsr_simd_half4 TLaneVector2D<tsr_half, 4, LANE_STRIDE_X, LANE_STRIDE_Y>
 #define tsr_simd_halfC TLaneVector2D<tsr_half, CONFIG_CHANNEL_COUNT, LANE_STRIDE_X, LANE_STRIDE_Y>
 #define tsr_simd_ushort  TLaneVector2D<tsr_ushort, 1, LANE_STRIDE_X, LANE_STRIDE_Y>
-#define tsr_simd_bool  TLaneVector2D<bool, 1, LANE_STRIDE_X, LANE_STRIDE_Y>
+#define tsr_simd_short2  TLaneVector2D<tsr_ushort, 2, LANE_STRIDE_X, LANE_STRIDE_Y>
+#define tsr_simd_bool    TLaneVector2D<bool, 1, LANE_STRIDE_X, LANE_STRIDE_Y>
 #define tsr_simd_bool4   TLaneVector2D<bool, 4, LANE_STRIDE_X, LANE_STRIDE_Y>
 #define tsr_simd_boolC   TLaneVector2D<bool, CONFIG_CHANNEL_COUNT, LANE_STRIDE_X, LANE_STRIDE_Y>
 
@@ -147,8 +146,15 @@
 
 //------------------------------------------------------- CONSTANTS
 
+static const float kHistoryGuidePreceptionAdd = 1.0;
+
 static const tsr_half FilteringWeight = rcp(tsr_half(1.0 + 4 * 0.5 + 4 * 0.25));
 
+#if CONFIG_LDS_DESPILL
+	static const uint kDespillInput = LDS_DESPILL_THREAD_COUNT * LDS_DESPILL_OFFSET * 0;
+	static const uint kDespillHistory = LDS_DESPILL_THREAD_COUNT * LDS_DESPILL_OFFSET * 1;
+#endif
+
 
 //------------------------------------------------------- PARAMETERS
 
@@ -168,8 +174,10 @@ Texture2D<tsr_half4> InputSceneTranslucencyTexture;
 Texture2D<tsr_half4> ReprojectedHistoryGuideTexture;
 Texture2D<tsr_half>  ReprojectedHistoryGuideMetadataTexture;
 Texture2D<tsr_half4> ReprojectedHistoryMoireTexture;
-Texture2D<tsr_half4> ResurrectedHistoryGuideTexture;
-Texture2D<tsr_half>  ResurrectedHistoryGuideMetadataTexture;
+#if CONFIG_HISTORY_RESURRECTION
+	Texture2D<tsr_half4> ResurrectedHistoryGuideTexture;
+	Texture2D<tsr_half>  ResurrectedHistoryGuideMetadataTexture;
+#endif
 Texture2D<tsr_half2> DecimateMaskTexture;
 Texture2D<tsr_ushort> IsMovingMaskTexture;
 
@@ -181,11 +189,133 @@ RWTexture2D<tsr_half> InputSceneColorLdrLumaOutput;
 RWTexture2D<tsr_ushort> AntiAliasMaskOutput;
 
 
-//------------------------------------------------------- DEBUG
+//------------------------------------------------------- GLOBALS
+
+static uint2 GGroupId = 0;
 
 static tsr_simd_halfC Debug[DEBUG_ARRAY_SIZE];
 
 
+//------------------------------------------------------- INPUT FETCHES
+
+tsr_simd_short2 ComputePixelPos()
+{
+	tsr_simd_short2 PixelPos;
+
+	UNROLL_N(SIMD_SIZE)
+	for (uint ElementIndex = 0; ElementIndex < SIMD_SIZE; ElementIndex++)
+	{
+		tsr_short2 ClampedFetchPixelPos = ComputeElementInputPixelPos(
+			/* LaneStride = */ uint2(LANE_STRIDE_X, LANE_STRIDE_Y),
+			TileOverscan,
+			InputInfo_ViewportMin,
+			InputInfo_ViewportMax,
+			GGroupId,
+			GGroupThreadIndex,
+			ElementIndex);
+
+		PixelPos.SetElement(ElementIndex, ClampedFetchPixelPos);
+	}
+
+	return PixelPos;
+}
+
+tsr_short2 ComputeElementOutputPixelPos(const uint ElementIndex)
+{
+	return ComputeElementOutputPixelPos(
+		/* LaneStride = */ uint2(LANE_STRIDE_X, LANE_STRIDE_Y),
+		TileOverscan,
+		InputInfo_ViewportMin,
+		InputInfo_ViewportMax,
+		GGroupId,
+		GGroupThreadIndex,
+		ElementIndex);
+}
+
+void FetchDecimateMask(
+	tsr_simd_short2 PixelPos,
+	out tsr_simd_half VelocityEdge,
+	out tsr_simd_bool bIsDisoccluded)
+{
+	UNROLL_N(SIMD_SIZE)
+	for (uint ElementIndex = 0; ElementIndex < SIMD_SIZE; ElementIndex++)
+	{
+		tsr_short2 ClampedFetchPixelPos = PixelPos.GetElement(ElementIndex);
+		
+		tsr_half2 DecimateMask = DecimateMaskTexture[ClampedFetchPixelPos];
+
+		VelocityEdge.SetElement(ElementIndex, DecimateMask.g);
+		bIsDisoccluded.SetElement(ElementIndex, !(DecimateMask.r > tsr_half(PARALLAX_REJECTION_MASK_THRESHOLD)));
+	}
+} // FetchDecimateMask()
+
+void FetchSceneColorAndTranslucency(
+	tsr_simd_short2 PixelPos,
+	out tsr_simd_halfC OriginalOpaqueInput,
+	out tsr_simd_half4 OriginalTranslucencyInput)
+{
+	UNROLL_N(SIMD_SIZE)
+	for (uint ElementIndex = 0; ElementIndex < SIMD_SIZE; ElementIndex++)
+	{
+		tsr_short2 ClampedFetchPixelPos = PixelPos.GetElement(ElementIndex);
+		
+		float2 TranslucencyTextureUV = ApplyScreenTransform(float2(ClampedFetchPixelPos), InputPixelPosToTranslucencyTextureUV);
+		TranslucencyTextureUV = fastClamp(TranslucencyTextureUV, TranslucencyTextureUVMin, TranslucencyTextureUVMax);
+		
+		OriginalOpaqueInput.SetElement(ElementIndex, InputTexture[ClampedFetchPixelPos]);
+		OriginalTranslucencyInput.SetElement(ElementIndex, InputSceneTranslucencyTexture.SampleLevel(GlobalBilinearClampedSampler, TranslucencyTextureUV, 0));
+	}
+} // FetchSceneColorAndTranslucency()
+
+void FetchHistoryGuide(
+	Texture2D<tsr_half4> HistoryGuideTexture,
+	Texture2D<tsr_half>  HistoryGuideMetadataTexture,
+	tsr_simd_short2 PixelPos,
+	out tsr_simd_halfC History,
+	out tsr_simd_half  HistoryUncertainty)
+{
+	UNROLL_N(SIMD_SIZE)
+	for (uint ElementIndex = 0; ElementIndex < SIMD_SIZE; ElementIndex++)
+	#if CONFIG_SCENE_COLOR_ALPHA
+	{
+		tsr_short2 ClampedFetchPixelPos = PixelPos.GetElement(ElementIndex);
+
+		History.SetElement(ElementIndex, HistoryGuideTexture[ClampedFetchPixelPos]);
+		HistoryUncertainty.SetElement(ElementIndex, HistoryGuideMetadataTexture[ClampedFetchPixelPos]);
+	}
+	#else
+	{
+		tsr_short2 ClampedFetchPixelPos = PixelPos.GetElement(ElementIndex);
+
+		tsr_half4 ReprojectedGuide = HistoryGuideTexture[ClampedFetchPixelPos];
+		History.SetElement(ElementIndex, ReprojectedGuide.rgb);
+		HistoryUncertainty.SetElement(ElementIndex, ReprojectedGuide.a);
+	}
+	#endif
+} // FetchHistoryGuide()
+
+#if CONFIG_MOIRE_DETECTION
+
+void FetchMoireInput(
+	tsr_simd_short2 PixelPos,
+	out tsr_simd_half MoireInput,
+	out tsr_simd_half4 PrevMoireHistory,
+	out tsr_simd_half bIsStatic)
+{
+	UNROLL_N(SIMD_SIZE)
+	for (uint ElementIndex = 0; ElementIndex < SIMD_SIZE; ElementIndex++)
+	{
+		tsr_short2 ClampedFetchPixelPos = PixelPos.GetElement(ElementIndex);
+		
+		MoireInput.SetElement(ElementIndex, InputMoireLumaTexture[ClampedFetchPixelPos]);
+		PrevMoireHistory.SetElement(ElementIndex, ReprojectedHistoryMoireTexture[ClampedFetchPixelPos]);
+		bIsStatic.SetElement(ElementIndex, tsr_half(1.0) - tsr_half(IsMovingMaskTexture[ClampedFetchPixelPos]) * rcp(tsr_half(255.0)));
+	}
+} // FetchMoireInput()
+
+#endif // CONFIG_MOIRE_DETECTION
+
+
 //------------------------------------------------------- COMPOSE SEPARATE TRANSLUCENCY
 
 /** Compose translucency in linear color space. */
@@ -197,6 +327,11 @@ tsr_simd_halfC ComposeTranslucency(tsr_simd_halfC Color, tsr_simd_half4 Transluc
 	#if CONFIG_SCENE_COLOR_ALPHA
 		ComposedColor.SetComponent(3, Color[3] * Translucency[3]);
 	#endif
+	#if CONFIG_SCENE_COLOR_ALPHA
+		ComposedColor = min(ComposedColor, tsr_simd_halfC::Const(LargestSceneColorRGBA));
+	#else
+		ComposedColor = min(ComposedColor, tsr_simd_halfC::Const(LargestSceneColorRGB));
+	#endif
 	return ComposedColor;
 }
 
@@ -276,7 +411,7 @@ tsr_simd_bool ShouldSpatialAntiAlias(tsr_simd_bool bIsDisoccluded, tsr_simd_bool
 //------------------------------------------------------- FLICKERING TEMPORAL ANALYSIS
 
 void ComputeMoireError(
-	tsr_simd_bool  bIsValid,
+	tsr_simd_bool  bIsDisoccluded,
 	tsr_simd_half  bIsStatic,
 	tsr_simd_half  MoireInput,
 	tsr_simd_halfC OriginalOpaqueInput,
@@ -439,7 +574,7 @@ void ComputeMoireError(
 		tsr_half LocalPrevGradient = PrevGradient.GetElement(SimdIndex);
 		tsr_half LocalPrevTotalVariation = PrevTotalVariation.GetElement(SimdIndex);
 		tsr_half LocalPrevTotalVariationCount = PrevTotalVariationCount.GetElement(SimdIndex);
-		bool LocalIsValid = bIsValid.GetElement(SimdIndex);
+		bool LocalIsDisoccluded = bIsDisoccluded.GetElement(SimdIndex);
 
 		tsr_half LocalIsFlicker = IsFlicker.GetElement(SimdIndex);
 
@@ -461,9 +596,9 @@ void ComputeMoireError(
 		tsr_half LocalNewTotalVariationCount = LocalPrevTotalVariationCount + LocalTotalVariationCountContribution;
 
 		// Discard all moire history on parallax disocclusion.
-		LocalNewGradient = select(LocalIsValid, LocalNewGradient, tsr_half(0.0));
-		LocalNewTotalVariation = select(LocalIsValid, LocalNewTotalVariation, tsr_half(0.0));
-		LocalNewTotalVariationCount = select(LocalIsValid, LocalNewTotalVariationCount, tsr_half(0.0));
+		LocalNewGradient            = select(LocalIsDisoccluded, tsr_half(0.0), LocalNewGradient);
+		LocalNewTotalVariation      = select(LocalIsDisoccluded, tsr_half(0.0), LocalNewTotalVariation);
+		LocalNewTotalVariationCount = select(LocalIsDisoccluded, tsr_half(0.0), LocalNewTotalVariationCount);
 
 		// Quantise total variation and count to history bit depth to variation drift in the history over time.
 		tsr_half LocalQuantizedNewTotalVariationCount = floor(LocalNewTotalVariationCount * (tsr_half(255.0) / MaxPrevTotalVariationCount)) * (MaxPrevTotalVariationCount / tsr_half(255.0));
@@ -498,7 +633,6 @@ void ComputeMoireError(
 
 template<bool ApplyMoireError>
 void MeasureRejection(
-	tsr_simd_bool  bIsValid,
 	tsr_simd_halfC InputC0,
 	tsr_simd_halfC InputC2,
 	tsr_simd_halfC HistoryC2,
@@ -620,12 +754,16 @@ tsr_simd_bool IsResurrectedFrameCloserThanPreviousFrame(tsr_simd_halfC ExposedIn
 
 	tsr_simd_halfC MatchDiff = PrevMatch - ResurrectionMatch;
 	#if CONFIG_SCENE_COLOR_ALPHA
-		tsr_simd_half TotalMatchDiff = BlurEven3x3(MatchDiff[0] + MatchDiff[1] + MatchDiff[2] + MatchDiff[3]);
+		tsr_simd_half TotalMatchDiff = Sum3x3(MatchDiff[0] + MatchDiff[1] + MatchDiff[2] + MatchDiff[3]);
 	#else
-		tsr_simd_half TotalMatchDiff = BlurEven3x3(MatchDiff[0] + MatchDiff[1] + MatchDiff[2]);
+		tsr_simd_half TotalMatchDiff = Sum3x3(MatchDiff[0] + MatchDiff[1] + MatchDiff[2]);
 	#endif
 
-	tsr_simd_bool bResurrectionIsCloser = MedianBool3x3(TotalMatchDiff > tsr_simd_half::Const(tsr_half(0.05 * 3.0)));
+	tsr_simd_half bResurrectedPixelIsCloser = saturate((TotalMatchDiff - tsr_simd_half::Const(tsr_half(0.05 * 3.0 * 9.0))) * tsr_simd_half::Const(tsr_half(POSITIVE_INFINITY)));
+
+	tsr_simd_bool bResurrectionIsCloser = Sum3x3(bResurrectedPixelIsCloser) > tsr_simd_half::Const(4.0);
+
+	//tsr_simd_bool bResurrectionIsCloser = MedianBool3x3((TotalMatchDiff - tsr_simd_half::Const(tsr_half(0.05 * 3.0 * 9.0))));
 
 	return bResurrectionIsCloser;
 }
@@ -674,6 +812,52 @@ tsr_simd_halfC UpdateGuide(tsr_simd_halfC ExposedInput, tsr_simd_halfC ExposedHi
 }
 
 
+//------------------------------------------------------- REGISTER PRESSURE HELPERS
+
+#if CONFIG_LDS_DESPILL || CONFIG_CACHE_REFETCH
+
+void RefetchInputAndHistoryFromLDSOrCache(
+	out tsr_simd_halfC Input,
+	out tsr_simd_half  InputUncertainty,
+	out tsr_simd_halfC History,
+	out tsr_simd_half  HistoryUncertainty)
+#if CONFIG_LDS_DESPILL
+{
+	Deconcatenate(tsr_simd_half4::LoadLDSSpill(kDespillInput), /* out */ Input, /* out */ InputUncertainty);
+	Deconcatenate(tsr_simd_half4::LoadLDSSpill(kDespillHistory), /* out */ History, /* out */ HistoryUncertainty);
+}
+#else
+{
+	tsr_simd_short2 PixelPos = ComputePixelPos();
+
+	tsr_simd_halfC OriginalOpaqueInput;
+	tsr_simd_half4 OriginalTranslucencyInput;
+	tsr_simd_halfC OriginalHistory;
+	tsr_simd_half  OriginalHistoryUncertainty;
+	tsr_simd_half  VelocityEdge;
+	tsr_simd_bool  bIsDisoccluded;
+
+	FetchDecimateMask(PixelPos, /* out */ VelocityEdge, /* out */ bIsDisoccluded);
+	FetchSceneColorAndTranslucency(PixelPos, /* out */ OriginalOpaqueInput, /* out */ OriginalTranslucencyInput);
+	FetchHistoryGuide(
+		ReprojectedHistoryGuideTexture,
+		ReprojectedHistoryGuideMetadataTexture,
+		PixelPos,
+		/* out */ OriginalHistory,
+		/* out */ OriginalHistoryUncertainty);
+
+	tsr_simd_halfC OriginalInput = ComposeTranslucency(OriginalOpaqueInput, Blur3x3(OriginalTranslucencyInput));
+
+	Input            = OriginalInput;
+	InputUncertainty = VelocityEdge;
+	History            = OriginalHistory;
+	HistoryUncertainty = OriginalHistoryUncertainty;
+}
+#endif
+
+#endif // CONFIG_LDS_DESPILL || CONFIG_CACHE_REFETCH
+
+
 //------------------------------------------------------- ENTRY POINT
 
 #if COMPILER_SUPPORTS_WAVE_SIZE && DIM_WAVE_SIZE > 0
@@ -684,6 +868,7 @@ void MainCS(
 	uint2 GroupId : SV_GroupID,
 	uint GroupThreadIndex : SV_GroupIndex)
 {
+	GGroupId = GroupId;
 	GGroupThreadIndex = GroupThreadIndex;
 
 	UNROLL_N(DEBUG_ARRAY_SIZE)
@@ -691,85 +876,87 @@ void MainCS(
 	{
 		Debug[DebugId].SetAllElements(0.0);
 	}
-
-	const uint LaneIndex = GroupThreadIndex;
-	const tsr_short2 GroupPixelOffset = (
-		tsr_short2(InputPixelPosMin) +
-		tsr_short2(GroupId) * tsr_short(GROUP_TILE_SIZE - 2 * TileOverscan).xx +
-		tsr_short(-TileOverscan).xx);
 		
-	// Load
+	// Analyses the moire pattern overtime.
 	tsr_simd_halfC OriginalOpaqueInput;
 	tsr_simd_half4 OriginalTranslucencyInput;
-	tsr_simd_halfC OriginalHistory;
-	tsr_simd_half OriginalHistoryUncertainty;
-	tsr_simd_half VelocityEdge;
-	tsr_simd_bool bIsValid;
+	tsr_simd_half  VelocityEdge;
+	tsr_simd_bool  bIsDisoccluded;
+	tsr_simd_half  MoireError = tsr_simd_half::Const(0.0);
 	#if CONFIG_MOIRE_DETECTION
-		tsr_simd_half MoireInput;
-		tsr_simd_half4 PrevMoireHistory;
-		tsr_simd_half bIsStatic;
-	#endif
-	#if CONFIG_HISTORY_RESURRECTION
-		tsr_simd_halfC OriginalResurrectedHistory;
-		tsr_simd_half OriginalResurrectedHistoryUncertainty;
-	#endif
 	PLATFORM_SPECIFIC_ISOLATE
 	{
-		UNROLL_N(SIMD_SIZE)
-		for (uint SimdIndex = 0; SimdIndex < SIMD_SIZE; SimdIndex++)
+		// Load its own data
+		#if CONFIG_MOIRE_DETECTION
+			tsr_simd_half  MoireInput;
+			tsr_simd_half4 PrevMoireHistory;
+			tsr_simd_half  bIsStatic;
+		#endif
+		PLATFORM_SPECIFIC_ISOLATE
 		{
-			tsr_short2 ClampedFetchPixelPos = ComputeElementInputPixelPos(
-				/* LaneStride = */ uint2(LANE_STRIDE_X, LANE_STRIDE_Y),
-				TileOverscan,
-				InputInfo_ViewportMin,
-				InputInfo_ViewportMax,
-				GroupId,
-				GroupThreadIndex,
-				SimdIndex);
-
-			float2 TranslucencyTextureUV = ApplyScreenTransform(float2(ClampedFetchPixelPos), InputPixelPosToTranslucencyTextureUV);
-			TranslucencyTextureUV = fastClamp(TranslucencyTextureUV, TranslucencyTextureUVMin, TranslucencyTextureUVMax);
+			tsr_simd_short2 PixelPos = ComputePixelPos();
+		
+			FetchDecimateMask(PixelPos, /* out */ VelocityEdge, /* out */ bIsDisoccluded);
+			FetchSceneColorAndTranslucency(PixelPos, /* out */ OriginalOpaqueInput, /* out */ OriginalTranslucencyInput);
+			FetchMoireInput(
+				PixelPos,
+				/* out */ MoireInput,
+				/* out */ PrevMoireHistory,
+				/* out */ bIsStatic);
+		}
+	
+		tsr_simd_half4 MoireHistory;
+		ComputeMoireError(
+			bIsDisoccluded,
+			bIsStatic,
+			MoireInput,
+			OriginalOpaqueInput,
+			PrevMoireHistory,
+			/* out */ MoireHistory,
+			/* out */ MoireError);
 
-			tsr_half2 DecimateMask = DecimateMaskTexture[ClampedFetchPixelPos];
-			tsr_half4 ReprojectedGuide = ReprojectedHistoryGuideTexture[ClampedFetchPixelPos];
-			#if CONFIG_SCENE_COLOR_ALPHA
-				tsr_half ReprojectedGuideMetadata = ReprojectedHistoryGuideMetadataTexture[ClampedFetchPixelPos];
-			#endif
-			#if CONFIG_HISTORY_RESURRECTION
-				tsr_half4 ResurrectedGuide = ResurrectedHistoryGuideTexture[ClampedFetchPixelPos];
-				#if CONFIG_SCENE_COLOR_ALPHA
-					tsr_half ResurrectedGuideMetadata = ReprojectedHistoryGuideMetadataTexture[ClampedFetchPixelPos];
-				#endif
-			#endif
+		// Stores moire output to save VGPR for what is next.
+		UNROLL_N(SIMD_SIZE)
+		for (uint ElementIndex = 0; ElementIndex < SIMD_SIZE; ElementIndex++)
+		{
+			const tsr_short2 OutputPixelPos = ComputeElementOutputPixelPos(ElementIndex);
 
-			OriginalOpaqueInput.SetElement(SimdIndex, InputTexture[ClampedFetchPixelPos]);
-			OriginalTranslucencyInput.SetElement(SimdIndex, InputSceneTranslucencyTexture.SampleLevel(GlobalBilinearClampedSampler, TranslucencyTextureUV, 0));
-			#if CONFIG_SCENE_COLOR_ALPHA
-				OriginalHistory.SetElement(SimdIndex, ReprojectedGuide);
-				OriginalHistoryUncertainty.SetElement(SimdIndex, ReprojectedGuideMetadata);
-			#else
-				OriginalHistory.SetElement(SimdIndex, ReprojectedGuide.rgb);
-				OriginalHistoryUncertainty.SetElement(SimdIndex, ReprojectedGuide.a);
-			#endif
-			VelocityEdge.SetElement(SimdIndex, DecimateMask.g);
-			bIsValid.SetElement(SimdIndex, DecimateMask.r > tsr_half(PARALLAX_REJECTION_MASK_THRESHOLD));
-			#if CONFIG_MOIRE_DETECTION
-				MoireInput.SetElement(SimdIndex, InputMoireLumaTexture[ClampedFetchPixelPos]);
-				PrevMoireHistory.SetElement(SimdIndex, ReprojectedHistoryMoireTexture[ClampedFetchPixelPos]);
-				bIsStatic.SetElement(SimdIndex, tsr_half(1.0) - tsr_half(IsMovingMaskTexture[ClampedFetchPixelPos]) * rcp(tsr_half(255.0)));
-			#endif
-			#if CONFIG_HISTORY_RESURRECTION
-				#if CONFIG_SCENE_COLOR_ALPHA
-					OriginalResurrectedHistory.SetElement(SimdIndex, ResurrectedGuide);
-					OriginalResurrectedHistoryUncertainty.SetElement(SimdIndex, ResurrectedGuideMetadata);
-				#else
-					OriginalResurrectedHistory.SetElement(SimdIndex, ResurrectedGuide.rgb);
-					OriginalResurrectedHistoryUncertainty.SetElement(SimdIndex, ResurrectedGuide.a);
-				#endif
-			#endif
+			HistoryMoireOutput[tsr_short3(OutputPixelPos, 0)] = MoireHistory.GetElement(ElementIndex);
 		}
 	}
+	#endif // CONFIG_MOIRE_DETECTION
+	
+	// Load all the history data.
+	tsr_simd_halfC OriginalHistory;
+	tsr_simd_half  OriginalHistoryUncertainty;
+	#if CONFIG_HISTORY_RESURRECTION
+		tsr_simd_halfC OriginalResurrectedHistory;
+		tsr_simd_half  OriginalResurrectedHistoryUncertainty;
+	#endif
+	PLATFORM_SPECIFIC_ISOLATE
+	{
+		tsr_simd_short2 PixelPos = ComputePixelPos();
+
+		#if !CONFIG_MOIRE_DETECTION || 1
+			FetchDecimateMask(PixelPos, /* out */ VelocityEdge, /* out */ bIsDisoccluded);
+			FetchSceneColorAndTranslucency(PixelPos, /* out */ OriginalOpaqueInput, /* out */ OriginalTranslucencyInput);
+		#endif
+		FetchHistoryGuide(
+			ReprojectedHistoryGuideTexture,
+			ReprojectedHistoryGuideMetadataTexture,
+			PixelPos,
+			/* out */ OriginalHistory,
+			/* out */ OriginalHistoryUncertainty);
+			
+		#if CONFIG_HISTORY_RESURRECTION
+			FetchHistoryGuide(
+				ResurrectedHistoryGuideTexture,
+				ReprojectedHistoryGuideMetadataTexture,
+				PixelPos,
+				/* out */ OriginalResurrectedHistory,
+				/* out */ OriginalResurrectedHistoryUncertainty);
+		#endif
+	}
 	
 	// Compute AA Luma and compose separate translucency into scene color for optimised HistoryUpdate.
 	{
@@ -783,14 +970,7 @@ void MainCS(
 			UNROLL_N(SIMD_SIZE)
 			for (uint ElementIndex = 0; ElementIndex < SIMD_SIZE; ElementIndex++)
 			{
-				const tsr_short2 OutputPixelPos = ComputeElementOutputPixelPos(
-					/* LaneStride = */ uint2(LANE_STRIDE_X, LANE_STRIDE_Y),
-					TileOverscan,
-					InputInfo_ViewportMin,
-					InputInfo_ViewportMax,
-					GroupId,
-					GroupThreadIndex,
-					ElementIndex);
+				const tsr_short2 OutputPixelPos = ComputeElementOutputPixelPos(ElementIndex);
 		
 				InputSceneColorOutput[OutputPixelPos] = CenterFinalSceneColor.GetElement(ElementIndex);
 				InputSceneColorLdrLumaOutput[OutputPixelPos] = CenterLumaLDR.GetElement(ElementIndex);
@@ -799,84 +979,47 @@ void MainCS(
 	}
 	
 	tsr_simd_halfC OriginalInput = ComposeTranslucency(OriginalOpaqueInput, Blur3x3(OriginalTranslucencyInput));
-	#if CONFIG_SCENE_COLOR_ALPHA
-		OriginalInput = min(OriginalInput, tsr_simd_halfC::Const(LargestSceneColorRGBA));
-	#else
-		OriginalInput = min(OriginalInput, tsr_simd_halfC::Const(LargestSceneColorRGB));
-	#endif
 
-	// Analyses the moire pattern overtime.
-	tsr_simd_half  MoireError = tsr_simd_half::Const(0.0);
-	#if CONFIG_MOIRE_DETECTION
+	// Despill to LDS to make room in VGPRs for the convolutions of ComputeMoireError()
+	#if CONFIG_LDS_DESPILL
+	PLATFORM_SPECIFIC_ISOLATE
 	{
-		tsr_simd_half4 MoireHistory;
-		ComputeMoireError(
-			bIsValid,
-			bIsStatic,
-			MoireInput,
-			OriginalOpaqueInput,
-			PrevMoireHistory,
-			/* out */ MoireHistory,
-			/* out */ MoireError);
-
-		PLATFORM_SPECIFIC_ISOLATE
-		{
-			UNROLL_N(SIMD_SIZE)
-			for (uint ElementIndex = 0; ElementIndex < SIMD_SIZE; ElementIndex++)
-			{
-				const tsr_short2 OutputPixelPos = ComputeElementOutputPixelPos(
-					/* LaneStride = */ uint2(LANE_STRIDE_X, LANE_STRIDE_Y),
-					TileOverscan,
-					InputInfo_ViewportMin,
-					InputInfo_ViewportMax,
-					GroupId,
-					GroupThreadIndex,
-					ElementIndex);
-
-				HistoryMoireOutput[tsr_short3(OutputPixelPos, 0)] = MoireHistory.GetElement(ElementIndex);
-			}
-		}
+		Concatenate(OriginalInput, VelocityEdge).StoreLDSSpill(kDespillInput);
+		Concatenate(OriginalHistory, OriginalHistoryUncertainty).StoreLDSSpill(kDespillHistory);
 	}
-	#endif // CONFIG_MOIRE_DETECTION
+	#endif // CONFIG_LDS_DESPILL
 
 	// Measure rejection
 	tsr_simd_half RejectionBlendFinal = tsr_simd_half::Const(1.0);
 	tsr_simd_half RejectionClampBlend = tsr_simd_half::Const(1.0);
-	tsr_simd_bool bResurrectHistory = tsr_simd_bool::Const(false);
 	tsr_simd_bool bAliasingIsVisible = tsr_simd_bool::Const(false);
+	#if CONFIG_HISTORY_RESURRECTION
+		tsr_simd_bool bResurrectionIsCloser = tsr_simd_bool::Const(false);
+	#endif
+	PLATFORM_SPECIFIC_ISOLATE
 	{
-		// Convert input to LDR
-		tsr_simd_halfC ExposedInput = AddPerceptionAdd(OriginalInput, /* DestPerceptionAdd = */ PerceptionAdd);
-	
-		// Change the exposure of history for the shading rejection
+		// Reload all the data from LDS or cache with moire to avoid register pressure
+		tsr_simd_halfC Input   = OriginalInput;
+		tsr_simd_halfC History = OriginalHistory;
+
+		// Convert input to shading LDR
+		tsr_simd_halfC ExposedInput = AddPerceptionAdd(Input, /* DestPerceptionAdd = */ PerceptionAdd);
 		tsr_simd_halfC ExposedHistory = ChangePerceptionAdd(
-			OriginalHistory,
-			/* SourcePerceptionAdd = */ 1.0,
+			History,
+			/* SourcePerceptionAdd = */ kHistoryGuidePreceptionAdd,
 			/* DestPerceptionAdd   = */ PerceptionAdd);
 			
-		// Despill to LDS to make room in VGPRs for the convolutions of MeasureRejection()
-		#if CONFIG_LDS_DESPILL
-		const uint kDespillHistory = LDS_DESPILL_THREAD_COUNT * LDS_DESPILL_OFFSET * 0;
+		// Check whether the history resurrection is closer to input, to discard and no pay the register pressure of OriginalResurrectedHistory
 		#if CONFIG_HISTORY_RESURRECTION
-			const uint kDespillResurrectedHistory = LDS_DESPILL_THREAD_COUNT * LDS_DESPILL_OFFSET * 1;
-		#endif
-		PLATFORM_SPECIFIC_ISOLATE
-		{
-			Concatenate(ExposedHistory, OriginalHistoryUncertainty).StoreLDSSpill(kDespillHistory);
-			
-			#if CONFIG_HISTORY_RESURRECTION
-			{
-				tsr_simd_halfC ExposedResurrectedHistory = ChangePerceptionAdd(
-					OriginalResurrectedHistory,
-					/* SourcePerceptionAdd = */ 1.0,
-					/* DestPerceptionAdd   = */ PerceptionAdd);
+			tsr_simd_halfC ExposedResurrectedHistory = ChangePerceptionAdd(
+				OriginalResurrectedHistory,
+				/* SourcePerceptionAdd = */ kHistoryGuidePreceptionAdd,
+				/* DestPerceptionAdd   = */ PerceptionAdd);
 				
-				Concatenate(ExposedResurrectedHistory, OriginalResurrectedHistoryUncertainty).StoreLDSSpill(kDespillResurrectedHistory);
-			}
-			#endif
-		}
-		#endif // CONFIG_LDS_DESPILL
-
+			bResurrectionIsCloser = IsResurrectedFrameCloserThanPreviousFrame(ExposedInput, ExposedHistory, ExposedResurrectedHistory);
+		#endif
+		
+		// Mutual Annihilation
 		tsr_simd_halfC ExposedInputMin;
 		tsr_simd_halfC ExposedInputMax;
 		MinMax3x3(ExposedInput, /* out */ ExposedInputMin, /* out */ ExposedInputMax);
@@ -887,109 +1030,142 @@ void MainCS(
 		tsr_simd_halfC ClampedHistory = AnnihilateToGuide3x3(ExposedHistory, /* Guide = */ ExposedInput);
 
 		MeasureRejection</* ApplyMoireError = */ CONFIG_MOIRE_DETECTION != 0>(
-			bIsValid,
 			/* InputC0 = */    ExposedInput,
 			/* InputC2 = */    ClampedInput,
 			/* HistoryC2 = */  ClampedHistory,
 			/* MoireError = */ MoireError,
 			/* out */ RejectionBlendFinal,
 			/* out */ RejectionClampBlend);
+	}
 	
-		#if CONFIG_HISTORY_RESURRECTION
+	// Try resurrecting the history if necessary
+	tsr_simd_bool bResurrectHistory = tsr_simd_bool::Const(false);
+	#if CONFIG_HISTORY_RESURRECTION
+	bool bComputeResurrection = IsWorthMeasuringRejectionOfResurrectedFrame(RejectionBlendFinal, bResurrectionIsCloser);
+	BRANCH
+	if (bComputeResurrection)
+	PLATFORM_SPECIFIC_ISOLATE
+	{
+		// Reload the resurrected history from cache.
+		tsr_simd_halfC Input;
+		tsr_simd_halfC ResurrectedHistory;
+		tsr_simd_half  ResurrectedHistoryUncertainty;
+		#if CONFIG_CACHE_REFETCH || CONFIG_LDS_DESPILL
 		{
-			tsr_simd_halfC ExposedResurrectedHistory = ChangePerceptionAdd(
-				OriginalResurrectedHistory,
-				/* SourcePerceptionAdd = */ 1.0,
-				/* DestPerceptionAdd   = */ PerceptionAdd);
+			tsr_simd_short2 PixelPos = ComputePixelPos();
+		
+			FetchHistoryGuide(
+				ResurrectedHistoryGuideTexture,
+				ReprojectedHistoryGuideMetadataTexture,
+				PixelPos,
+				/* out */ ResurrectedHistory,
+				/* out */ ResurrectedHistoryUncertainty);
 				
 			#if CONFIG_LDS_DESPILL
-			PLATFORM_SPECIFIC_ISOLATE
-			{
-				Deconcatenate(tsr_simd_half4::LoadLDSSpill(kDespillHistory), /* out */ ExposedHistory, /* out */ OriginalHistoryUncertainty);
-				Deconcatenate(tsr_simd_half4::LoadLDSSpill(kDespillResurrectedHistory), /* out */ ExposedResurrectedHistory, /* out */ OriginalResurrectedHistoryUncertainty);
-			}
-			#endif // CONFIG_LDS_DESPILL
-
-			tsr_simd_bool bResurrectionIsCloser = IsResurrectedFrameCloserThanPreviousFrame(ExposedInput, ExposedHistory, ExposedResurrectedHistory);
-			bool bComputeResurrection = IsWorthMeasuringRejectionOfResurrectedFrame(RejectionBlendFinal, bResurrectionIsCloser);
+				Input = ResizeChannels<3>(tsr_simd_half4::LoadLDSSpill(kDespillInput));
+			#else
+				tsr_simd_halfC OriginalOpaqueInput;
+				tsr_simd_half4 OriginalTranslucencyInput;
+				FetchSceneColorAndTranslucency(PixelPos, /* out */ OriginalOpaqueInput, /* out */ OriginalTranslucencyInput);
 
-			BRANCH
-			if (bComputeResurrection)
-			{
-				tsr_simd_halfC ClampedResurrectedInput, ClampedResurrectedHistory;
-				AnnihilateMutually3x3(ExposedInput, ExposedResurrectedHistory, /* out */ ClampedResurrectedInput, /* out */ ClampedResurrectedHistory);
-
-				tsr_simd_half ResurrectionRejectionBlendFinal = tsr_simd_half::Const(1.0);
-				tsr_simd_half ResurrectionRejectionClampBlend = tsr_simd_half::Const(1.0);
-				MeasureRejection</* ApplyMoireError = */ false>(
-					bIsValid,
-					/* InputC0 = */    ExposedInput,
-					/* InputC2 = */    ClampedResurrectedInput,
-					/* HistoryC2 = */  ClampedResurrectedHistory,
-					/* MoireError = */ tsr_simd_half::Const(0.0),
-					/* out */ ResurrectionRejectionBlendFinal,
-					/* out */ ResurrectionRejectionClampBlend);
-			
-				bResurrectHistory = ShouldResurrectHistory(RejectionBlendFinal, ResurrectionRejectionBlendFinal, bResurrectionIsCloser);
-
-				#if CONFIG_LDS_DESPILL
-				PLATFORM_SPECIFIC_ISOLATE
-				{
-					Deconcatenate(tsr_simd_half4::LoadLDSSpill(kDespillHistory), /* out */ ExposedHistory, /* out */ OriginalHistoryUncertainty);
-					Deconcatenate(tsr_simd_half4::LoadLDSSpill(kDespillResurrectedHistory), /* out */ ExposedResurrectedHistory, /* out */ OriginalResurrectedHistoryUncertainty);
-				}
-				#endif // CONFIG_LDS_DESPILL
-
-				RejectionBlendFinal = select(bResurrectHistory, ResurrectionRejectionBlendFinal, RejectionBlendFinal);
-				RejectionClampBlend = select(bResurrectHistory, ResurrectionRejectionClampBlend, RejectionClampBlend);
-				OriginalHistoryUncertainty = select(bResurrectHistory, OriginalResurrectedHistoryUncertainty, OriginalHistoryUncertainty);
-				
-				#if CONFIG_MINIMISE_VGPRS
-					ExposedHistory = select(tsr_simd_boolC::Vectorize(bResurrectHistory), ExposedResurrectedHistory, ExposedHistory);
-				#else
-					OriginalHistory = select(tsr_simd_boolC::Vectorize(bResurrectHistory), OriginalResurrectedHistory, OriginalHistory);
-				#endif
-			}
-		}
-		#elif !CONFIG_HISTORY_RESURRECTION && CONFIG_LDS_DESPILL
-		PLATFORM_SPECIFIC_ISOLATE
-		{
-			Deconcatenate(tsr_simd_half4::LoadLDSSpill(kDespillHistory), /* out */ ExposedHistory, /* out */ OriginalHistoryUncertainty);
+				Input = ComposeTranslucency(OriginalOpaqueInput, Blur3x3(OriginalTranslucencyInput));
+			#endif
 		}
-		#endif //!CONFIG_HISTORY_RESURRECTION
+		#else
+		ResurrectedHistory = OriginalResurrectedHistory;
+		Input = OriginalInput;
+		#endif
+		
+		// Convert input to shading LDR
+		tsr_simd_halfC ExposedInput = AddPerceptionAdd(Input, /* DestPerceptionAdd = */ PerceptionAdd);
+		tsr_simd_halfC ExposedResurrectedHistory = ChangePerceptionAdd(
+			ResurrectedHistory,
+			/* SourcePerceptionAdd = */ kHistoryGuidePreceptionAdd,
+			/* DestPerceptionAdd   = */ PerceptionAdd);
+			
+		// Mutual Annihilation
+		tsr_simd_halfC ClampedResurrectedInput, ClampedResurrectedHistory;
+		AnnihilateMutually3x3(ExposedInput, ExposedResurrectedHistory, /* out */ ClampedResurrectedInput, /* out */ ClampedResurrectedHistory);
 
-		// Recomputes original off the exposed shading to save register pressure.
-		#if CONFIG_MINIMISE_VGPRS
-		PLATFORM_SPECIFIC_ISOLATE
-		{
-			OriginalInput =  RemovePerceptionAdd(ExposedInput, /* SourcePerceptionAdd = */ PerceptionAdd);
+		tsr_simd_half ResurrectionRejectionBlendFinal = tsr_simd_half::Const(1.0);
+		tsr_simd_half ResurrectionRejectionClampBlend = tsr_simd_half::Const(1.0);
+		MeasureRejection</* ApplyMoireError = */ false>(
+			/* InputC0 = */    ExposedInput,
+			/* InputC2 = */    ClampedResurrectedInput,
+			/* HistoryC2 = */  ClampedResurrectedHistory,
+			/* MoireError = */ tsr_simd_half::Const(0.0),
+			/* out */ ResurrectionRejectionBlendFinal,
+			/* out */ ResurrectionRejectionClampBlend);
 			
-			OriginalHistory = ChangePerceptionAdd(
-				ExposedHistory,
-				/* SourcePerceptionAdd = */ PerceptionAdd,
-				/* DestPerceptionAdd   = */ 1.0);
-		}
-		#endif
-	}
-	
-	tsr_simd_bool bIsDisoccluded = not(bIsValid);
+		bResurrectHistory = ShouldResurrectHistory(RejectionBlendFinal, ResurrectionRejectionBlendFinal, bResurrectionIsCloser);
 
+		RejectionBlendFinal = select(bResurrectHistory, ResurrectionRejectionBlendFinal, RejectionBlendFinal);
+		RejectionClampBlend = select(bResurrectHistory, ResurrectionRejectionClampBlend, RejectionClampBlend);
+	} // if (bComputeResurrection)
+	#endif // CONFIG_HISTORY_RESURRECTION
+	
 	// Update the guide for next frame.
 	tsr_simd_halfC FinalGuide;
+	tsr_simd_half  InputUncertainty;
+	tsr_simd_half  HistoryUncertainty;
+	PLATFORM_SPECIFIC_ISOLATE
 	{
+		// Reload all the data from LDS or cache
+		tsr_simd_halfC Input;
+		tsr_simd_halfC History;
+		#if CONFIG_LDS_DESPILL || CONFIG_CACHE_REFETCH
+			RefetchInputAndHistoryFromLDSOrCache(
+				/* out */ Input,
+				/* out */ InputUncertainty,
+				/* out */ History,
+				/* out */ HistoryUncertainty);
+		#else
+			Input            = OriginalInput;
+			InputUncertainty = VelocityEdge;
+			History            = OriginalHistory;
+			HistoryUncertainty = OriginalHistoryUncertainty;
+		#endif
+	
+		// Apply resurrection
+		#if CONFIG_HISTORY_RESURRECTION
+		BRANCH
+		if (AnyComponent(bResurrectHistory))
+		PLATFORM_SPECIFIC_ISOLATE
+		{
+			// Reload the resurrected history from cache.
+			tsr_simd_halfC ResurrectedHistory;
+			tsr_simd_half  ResurrectedHistoryUncertainty;
+			#if CONFIG_CACHE_REFETCH
+			{
+				FetchHistoryGuide(
+					ResurrectedHistoryGuideTexture,
+					ReprojectedHistoryGuideMetadataTexture,
+					ComputePixelPos(),
+					/* out */ ResurrectedHistory,
+					/* out */ ResurrectedHistoryUncertainty);
+			}
+			#else
+				ResurrectedHistory            = OriginalResurrectedHistory;
+				ResurrectedHistoryUncertainty = OriginalResurrectedHistoryUncertainty;
+			#endif
+
+			History = select(tsr_simd_boolC::Vectorize(bResurrectHistory), ResurrectedHistory, History);
+			HistoryUncertainty = select(bResurrectHistory, ResurrectedHistoryUncertainty, HistoryUncertainty);
+		}
+		#endif // CONFIG_HISTORY_RESURRECTION
+
 		// Convert input to LDR
-		tsr_simd_halfC ExposedInput = AddPerceptionAdd(OriginalInput, /* DestPerceptionAdd = */ 1.0);
+		tsr_simd_halfC ExposedInput = AddPerceptionAdd(Input, /* DestPerceptionAdd = */ kHistoryGuidePreceptionAdd);
 	
-		FinalGuide = UpdateGuide(ExposedInput, OriginalHistory, bIsDisoccluded, RejectionBlendFinal);
+		FinalGuide = UpdateGuide(ExposedInput, History, bIsDisoccluded, RejectionBlendFinal);
 	}
-	
-	tsr_simd_half InputUncertainty = VelocityEdge;
-	
+
 	// Controls how the history and its validity is clamped.
-	tsr_simd_half DisableHistoryClamp;
-	tsr_simd_half IncreaseValidityMultiplier;
+	tsr_simd_half  DisableHistoryClamp;
+	tsr_simd_half  IncreaseValidityMultiplier;
+	PLATFORM_SPECIFIC_ISOLATE
 	{
-		DisableHistoryClamp        = select(and(bIsDisoccluded, not(bResurrectHistory)), tsr_simd_half::Const(0.0), min(RejectionClampBlend, min(InputUncertainty, Min3x3(OriginalHistoryUncertainty))));
+		DisableHistoryClamp        = select(and(bIsDisoccluded, not(bResurrectHistory)), tsr_simd_half::Const(0.0), min(RejectionClampBlend, min(InputUncertainty, Min3x3(HistoryUncertainty))));
 		IncreaseValidityMultiplier = select(and(bIsDisoccluded, not(bResurrectHistory)), tsr_simd_half::Const(0.0), RejectionClampBlend);
 	}
 	
@@ -999,23 +1175,18 @@ void MainCS(
 	// Output data.
 	PLATFORM_SPECIFIC_ISOLATE
 	{
+		const uint LaneIndex = GGroupThreadIndex;
+
 		UNROLL_N(SIMD_SIZE)
 		for (uint ElementIndex = 0; ElementIndex < SIMD_SIZE; ElementIndex++)
 		{
-			const tsr_short2 LaneSimdPixelOffset = GetElementPixelOffsetInTile(LANE_STRIDE_X, LANE_STRIDE_Y, LaneIndex, ElementIndex);
-			const tsr_short2 OutputPixelPos = ComputeElementOutputPixelPos(
-				/* LaneStride = */ uint2(LANE_STRIDE_X, LANE_STRIDE_Y),
-				TileOverscan,
-				InputInfo_ViewportMin,
-				InputInfo_ViewportMax,
-				GroupId,
-				GroupThreadIndex,
-				ElementIndex);
+			const tsr_short2 LaneSimdPixelOffset = GetElementPixelOffsetInTile(LANE_STRIDE_X, LANE_STRIDE_Y, GGroupThreadIndex, ElementIndex);
+			const tsr_short2 OutputPixelPos = ComputeElementOutputPixelPos(ElementIndex);
 
 			tsr_halfC FinalGuideColor = FinalGuide.GetElement(ElementIndex);
 			
 			tsr_half BitMask = (
-				select(bIsValid.GetElement(ElementIndex),          tsr_half(0x1) * rcp(tsr_half(255.0)), tsr_half(0x0)) + 
+				select(bIsDisoccluded.GetElement(ElementIndex),    tsr_half(0x0), tsr_half(0x1) * rcp(tsr_half(255.0))) + 
 				select(bResurrectHistory.GetElement(ElementIndex), tsr_half(0x2) * rcp(tsr_half(255.0)), tsr_half(0x0)));
 
 			#if CONFIG_ENABLE_STOCASTIC_QUANTIZATION
@@ -1064,8 +1235,8 @@ void MainCS(
 				TileOverscan,
 				InputInfo_ViewportMin,
 				InputInfo_ViewportMax,
-				GroupId,
-				GroupThreadIndex,
+				GGroupId,
+				GGroupThreadIndex,
 				ElementIndex,
 				/* OutputDataOffset = */ int2(0, 0),
 				/* OutputResDivisor = */ 1,
