Speeds up r.TSR.ShadingRejection.Flickering by concatenating more 3x3 convolutions together

diff --git a/Engine/Shaders/Private/TemporalSuperResolution/TSRConvolutionNetwork.ush b/Engine/Shaders/Private/TemporalSuperResolution/TSRConvolutionNetwork.ush
--- a/Engine/Shaders/Private/TemporalSuperResolution/TSRConvolutionNetwork.ush
+++ b/Engine/Shaders/Private/TemporalSuperResolution/TSRConvolutionNetwork.ush
@@ -580,6 +580,28 @@ TLaneVector2D<ScalarType, 4, LaneStrideX, LaneStrideY> MaxRGBMinA3x3(TLaneVector
 	return R;
 }
 
+CALL_SITE_DEBUGLOC
+template<typename ScalarType, uint LaneStrideX, uint LaneStrideY>
+TLaneVector2D<ScalarType, 2, LaneStrideX, LaneStrideY> MaxRMinG3x3(TLaneVector2D<ScalarType, 2, LaneStrideX, LaneStrideY> Center)
+{
+	TLaneVector2D<ScalarType, 2, LaneStrideX, LaneStrideY> CO = Center;
+	TLaneVector2D<ScalarType, 2, LaneStrideX, LaneStrideY> CP;
+	TLaneVector2D<ScalarType, 2, LaneStrideX, LaneStrideY> CN;
+	AccessNeighborTexels1x3(CO, /* out */ CP, /* out */ CN);
+	
+	TLaneVector2D<ScalarType, 2, LaneStrideX, LaneStrideY> MaxO = max3(CN, CP, CO);
+	MaxO.SetComponent(1, min3(CN[1], CP[1], CO[1]));
+
+	TLaneVector2D<ScalarType, 2, LaneStrideX, LaneStrideY> MaxP;
+	TLaneVector2D<ScalarType, 2, LaneStrideX, LaneStrideY> MaxN;
+	AccessNeighborTexels3x1(MaxO, /* out */ MaxP, /* out */ MaxN);
+	
+	TLaneVector2D<ScalarType, 2, LaneStrideX, LaneStrideY> R = max3(MaxN, MaxP, MaxO);
+	R.SetComponent(1, min3(MaxN[1], MaxP[1], MaxO[1]));
+
+	return R;
+}
+
 CALL_SITE_DEBUGLOC
 template<typename ScalarType, uint VectorSize, uint LaneStrideX, uint LaneStrideY>
 void MinMax3x3(
@@ -661,6 +683,17 @@ void AnnihilateMutually3x3(
 	AnnihilatedHistory = AnnihilateToGuide3x3(History, /* Guide = */ Input);
 }
 
+CALL_SITE_DEBUGLOC
+template<typename ScalarType, uint LaneStrideX, uint LaneStrideY>
+void AnnihilateMutuallySingleChannel3x3(
+	TLaneVector2D<ScalarType, 1, LaneStrideX, LaneStrideY> Input,
+	TLaneVector2D<ScalarType, 1, LaneStrideX, LaneStrideY> History,
+	out TLaneVector2D<ScalarType, 1, LaneStrideX, LaneStrideY> AnnihilatedInput,
+	out TLaneVector2D<ScalarType, 1, LaneStrideX, LaneStrideY> AnnihilatedHistory)
+{
+	Deconcatenate(AnnihilateToGuide3x3(/* ToClamp = */ Concatenate(Input, History), /* Guide = */ Concatenate(History, Input)), /* out */ AnnihilatedInput, /* out */ AnnihilatedHistory);
+}
+
 
 //------------------------------------------------------- MIN & MAX & CLAMP 3x3 PLUS CONVOLUTION
 
diff --git a/Engine/Shaders/Private/TemporalSuperResolution/TSRRejectShading.usf b/Engine/Shaders/Private/TemporalSuperResolution/TSRRejectShading.usf
--- a/Engine/Shaders/Private/TemporalSuperResolution/TSRRejectShading.usf
+++ b/Engine/Shaders/Private/TemporalSuperResolution/TSRRejectShading.usf
@@ -470,26 +470,28 @@ void ComputeMoireError(
 		// Anhilate history
 		tsr_simd_half InputC2 = InputC0;
 		tsr_simd_half HistoryC2 = PrevFlickeringHistory;
-		AnnihilateMutually3x3(InputLuma, PrevFlickeringHistory, /* out */ InputC2, /* out */ HistoryC2);
+		AnnihilateMutuallySingleChannel3x3(InputLuma, PrevFlickeringHistory, /* out */ InputC2, /* out */ HistoryC2);
 
 		// Convolve the input luma to have the exact same moire pattern as FilteredInput below.
 		tsr_simd_half FilteredInput;
 		tsr_simd_half FilteredHistory;
 		Deconcatenate(Blur3x3(Concatenate(InputC2, InputLumaModification, HistoryC2)), /* out */ FilteredInput, /* out */ FilteredInputLumaModification, /* out */ FilteredHistory);
 		
-		tsr_simd_half TotalVarInputDiffC0C2 = abs(TotalVariation3x3(abs(InputC0 - InputC2)));
-		tsr_simd_half TotalVarInputC2 = abs(TotalVariation3x3(InputC2));
-
 		// Compute clamping box of the limance.
 		tsr_simd_half FilteredBoxMin;
 		tsr_simd_half FilteredBoxMax;
-		MinMax3x3(FilteredInput, /* out */ FilteredBoxMin, /* out */ FilteredBoxMax);
-		FilteredBoxMin = FilteredBoxMin - BackbufferQuantizationErrorVector;
-		FilteredBoxMax = FilteredBoxMax + BackbufferQuantizationErrorVector;
+		tsr_simd_half BoxSize;
+		{
+			tsr_simd_half2 Min;
+			tsr_simd_half2 Max;
+			MinMax3x3(Concatenate(FilteredInput, InputC2), /* out */ Min, /* out */ Max);
+			FilteredBoxMin = Min[0] - BackbufferQuantizationErrorVector;
+			FilteredBoxMax = Max[0] + BackbufferQuantizationErrorVector;
 		
-		tsr_simd_half BoxSize = (
-			MaxMinusMin3x3(InputC2) * tsr_simd_half::Const(FilteringWeight) +
-			(BackbufferQuantizationErrorVector * tsr_half(2 * FilteringWeight)));
+			BoxSize = (
+				(Max[1] - Min[1]) * tsr_simd_half::Const(FilteringWeight) +
+				(BackbufferQuantizationErrorVector * tsr_half(2 * FilteringWeight)));
+		}
 
 		tsr_simd_half ClampedFilteredHistory = fastClamp(FilteredHistory, FilteredBoxMin, FilteredBoxMax);
 
@@ -501,8 +503,8 @@ void ComputeMoireError(
 		tsr_simd_half RawClampedEnergy = abs(ClampedFilteredHistory - FilteredHistory);
 		tsr_simd_half RawRejection     = saturate(tsr_simd_half::Const(1.0) - max(RawClampedEnergy - ClampError, tsr_simd_half::Const(0.0)) * rcp(Delta));
 
-		OutRejectionClampBlend = Min3x3(Median3x3(RawRejection));
-		tsr_simd_half FilteredClampedEnergy = Max3x3(Median3x3(RawClampedEnergy));
+		tsr_simd_half FilteredClampedEnergy;
+		Deconcatenate(MaxRMinG3x3(Median3x3(Concatenate(RawClampedEnergy, RawRejection))), /* out */ FilteredClampedEnergy, /* out */ OutRejectionClampBlend);
 
 		OutRejectionBlendFinal = saturate(tsr_simd_half::Const(1.0) - max(FilteredClampedEnergy - ClampError, tsr_simd_half::Const(0.0)) * rcp(Delta));
 	}
