Lumen - Enabled multiple scattering for diffuse composition compositing rough specular.
If not enabled by the project, it will compile to the old EnvBRDF.

diff --git a/Engine/Shaders/Private/DiffuseIndirectComposite.usf b/Engine/Shaders/Private/DiffuseIndirectComposite.usf
--- a/Engine/Shaders/Private/DiffuseIndirectComposite.usf
+++ b/Engine/Shaders/Private/DiffuseIndirectComposite.usf
@@ -207,7 +207,7 @@ float3 MixSpecularAndRoughReflections(float InRoughness, bool bHasBackfaceDiffus
 	return AddContrastAndSpecularScale(Lighting);
 }
 
-float3 CombineRoughSpecular(FGBufferData GBuffer, bool bHasBackfaceDiffuse, float NoV, float4 RayTracedReflections, float3 RoughReflections, float3 SpecularColor)
+float3 CombineRoughSpecular(FGBufferData GBuffer, bool bHasBackfaceDiffuse, float NoV, float4 RayTracedReflections, float3 RoughReflections, float3 SpecularColor, float3 EnvBRDF)
 {
 	float3 Lighting;
 	
@@ -217,7 +217,7 @@ float3 CombineRoughSpecular(FGBufferData GBuffer, bool bHasBackfaceDiffuse, floa
 	}
 	else
 	{
-		Lighting = MixSpecularAndRoughReflections(GBuffer.Roughness, bHasBackfaceDiffuse, RayTracedReflections, RoughReflections) * EnvBRDF(SpecularColor, GBuffer.Roughness, NoV);
+		Lighting = MixSpecularAndRoughReflections(GBuffer.Roughness, bHasBackfaceDiffuse, RayTracedReflections, RoughReflections) * EnvBRDF;
 	}
 	return Lighting;
 }
@@ -568,14 +568,23 @@ void MainPS(
 					DiffuseColor += ClothFuzz * GBuffer.CustomData.a;
 				}
 
-				IndirectLighting.Diffuse = (DiffuseIndirectLighting * DiffuseColor + BackfaceDiffuseIndirectLighting) * Occlusion.DiffuseOcclusion;
+			#if USE_ENERGY_CONSERVATION
+				FBxDFEnergyTerms SpecularEnergyTerms = ComputeGGXSpecEnergyTerms(GBuffer.Roughness, NoV, SpecularColor);
+				float3 EnvBRDFValue = SpecularEnergyTerms.E; // EnvBRDF accounting for multiple scattering when enabled
+				float EnergyPreservationFactor = ComputeEnergyPreservation(SpecularEnergyTerms);
+			#else
+				float EnvBRDFValue = EnvBRDF(SpecularColor, GBuffer.Roughness, NoV);
+				float EnergyPreservationFactor = 1.0;
+			#endif
+
+				IndirectLighting.Diffuse = (DiffuseIndirectLighting * DiffuseColor + BackfaceDiffuseIndirectLighting) * Occlusion.DiffuseOcclusion * EnergyPreservationFactor;
 				IndirectLighting.Transmission = 0;
 
 				#if DIM_APPLY_DIFFUSE_INDIRECT == DIM_APPLY_DIFFUSE_INDIRECT_SCREEN_PROBE_GATHER
 					RoughSpecularIndirectLighting *= Occlusion.SpecularOcclusion;
-					IndirectLighting.Specular = CombineRoughSpecular(GBuffer, HasBackfaceDiffuse(Material), NoV, SpecularIndirectLighting, RoughSpecularIndirectLighting, SpecularColor);
+					IndirectLighting.Specular = CombineRoughSpecular(GBuffer, HasBackfaceDiffuse(Material), NoV, SpecularIndirectLighting, RoughSpecularIndirectLighting, SpecularColor, EnvBRDFValue);
 				#else
-					IndirectLighting.Specular = AddContrastAndSpecularScale(SpecularIndirectLighting.xyz) * EnvBRDF(SpecularColor, GBuffer.Roughness, NoV);
+					IndirectLighting.Specular = AddContrastAndSpecularScale(SpecularIndirectLighting.xyz) * EnvBRDFValue;
 				#endif
 			}
 			}
