Implemented alpha hold out for the sky, cloud and fog.

diff --git a/Engine/Shaders/Private/HeightFogCommon.ush b/Engine/Shaders/Private/HeightFogCommon.ush
--- a/Engine/Shaders/Private/HeightFogCommon.ush
+++ b/Engine/Shaders/Private/HeightFogCommon.ush
@@ -1,7 +1,7 @@
 // Copyright Epic Games, Inc. All Rights Reserved.
 
 /*=============================================================================
-	HeightFogCommon.usf: 
+	HeightFogCommon.usf:  
 =============================================================================*/
 
 #pragma once
@@ -205,6 +205,13 @@ half4 GetExponentialHeightFog(float3 WorldPositionRelativeToCamera, float Exclud
 		half3 FogColor = (InscatteringColor) * (1 - ExpFogFactor) + DirectionalInscattering;
 	#endif
 
+#if POST_PROCESS_ALPHA
+	if (FogStruct.bHoldout && View.RenderingReflectionCaptureMask == 0.0f)
+	{
+		FogColor.rgb *= 0;
+	}
+#endif
+
 	return half4(FogColor, ExpFogFactor);
 }
 
@@ -265,6 +272,14 @@ float4 CombineVolumetricFog(float4 GlobalFog, float3 VolumeUV, uint EyeIndex, fl
 	// Visualize depth distribution
 	//VolumetricFogLookup.rgb += .1f * frac(min(ZSlice, 1.0f) / View.VolumetricFogInvGridSize.z);
 
+#if POST_PROCESS_ALPHA
+	if (FogStruct.bHoldout && View.RenderingReflectionCaptureMask == 0.0f)
+	{
+		VolumetricFogLookup.rgb *= 0;
+		GlobalFog.rgb *= 0;
+	}
+#endif
+
 	return float4(VolumetricFogLookup.rgb + GlobalFog.rgb * VolumetricFogLookup.a, VolumetricFogLookup.a * GlobalFog.a);
 }
 
diff --git a/Engine/Shaders/Private/SkyAtmosphere.usf b/Engine/Shaders/Private/SkyAtmosphere.usf
--- a/Engine/Shaders/Private/SkyAtmosphere.usf
+++ b/Engine/Shaders/Private/SkyAtmosphere.usf
@@ -1,7 +1,7 @@
 // Copyright Epic Games, Inc. All Rights Reserved.
 
 /*=============================================================================
-	SkyAtmosphere.usf: Sky and atmosphere rendering functions. 
+	SkyAtmosphere.usf: Sky and atmosphere rendering functions.  
 =============================================================================*/
 
 
@@ -824,6 +824,13 @@ float4 PrepareOutput(float3 PreExposedLuminance, float3 Transmittance = float3(1
 	// We also half that range to also make sure we have room for other additive elements such as bloom, clouds or particle visual effects.
 	float3 SafePreExposedLuminance = min(PreExposedLuminance, Max10BitsFloat.xxx * 0.5f);
 
+#if POST_PROCESS_ALPHA
+	if (Atmosphere.bHoldout && View.RenderingReflectionCaptureMask == 0.0f)
+	{
+		SafePreExposedLuminance *= 0;
+	}
+#endif
+
 	const float GreyScaleTransmittance = dot(Transmittance, float3(1.0f / 3.0f, 1.0f / 3.0f, 1.0f / 3.0f));
 	return float4(SafePreExposedLuminance, GreyScaleTransmittance);
 }
@@ -1570,6 +1577,15 @@ void RenderCameraAerialPerspectiveVolumeCS(uint3 ThreadId : SV_DispatchThreadID)
 		AerialPespectiveViewDistanceScale,
 		tMaxMax);
 
+#if POST_PROCESS_ALPHA
+	if (Atmosphere.bHoldout && !RealTimeReflection360Mode)
+	{
+		ss.L *= 0;
+		ss.LMieOnly *= 0;
+		ss.LRayOnly *= 0;
+	}
+#endif
+
 	const float Transmittance = dot(ss.Transmittance, float3(1.0f / 3.0f, 1.0f / 3.0f, 1.0f / 3.0f));
 	CameraAerialPerspectiveVolumeUAV[ThreadId] = float4(ss.L, Transmittance);
 #if SEPARATE_MIE_RAYLEIGH_SCATTERING
diff --git a/Engine/Shaders/Private/VolumetricCloud.usf b/Engine/Shaders/Private/VolumetricCloud.usf
--- a/Engine/Shaders/Private/VolumetricCloud.usf
+++ b/Engine/Shaders/Private/VolumetricCloud.usf
@@ -1,7 +1,7 @@
 // Copyright Epic Games, Inc. All Rights Reserved.
 
 /**
- * VolumetricCloud.usf: Render volumetric cloud on screen.   
+ * VolumetricCloud.usf: Render volumetric cloud on screen.    
  */  
 
 
@@ -1607,6 +1607,16 @@ void MainCommon(in FMaterialPixelParameters MaterialParameters, in float4 SvPosi
 		OutDepth.x = ((GrayScaleTransmittance > 0.99) ? NoCloudDepth : tAP) * CENTIMETER_TO_KILOMETER; // using a small threshold on transmittance
 	}
 
+#if POST_PROCESS_ALPHA
+	if (RenderVolumetricCloudParameters.bHoldout2 && View.RenderingReflectionCaptureMask == 0.0f)
+	{
+		OutColor0.rgb *= 0;
+	#if CLOUD_MIN_AND_MAX_DEPTH
+		OutColor1.rgb *= 0;
+	#endif
+	}
+#endif
+
 	//OutColor0.g = 0.3f; // Debug revealing reveal pixels that have been traced.
 }
 
diff --git a/Engine/Source/Runtime/Engine/Classes/Components/ExponentialHeightFogComponent.h b/Engine/Source/Runtime/Engine/Classes/Components/ExponentialHeightFogComponent.h
--- a/Engine/Source/Runtime/Engine/Classes/Components/ExponentialHeightFogComponent.h
+++ b/Engine/Source/Runtime/Engine/Classes/Components/ExponentialHeightFogComponent.h
@@ -200,6 +200,10 @@ class UExponentialHeightFogComponent : public USceneComponent
 	UPROPERTY(EditAnywhere, BlueprintReadOnly, Category = VolumetricFog, AdvancedDisplay)
 	bool bOverrideLightColorsWithFogInscatteringColors;
 
+	/** If this is True, this primitive will render black with an alpha of 0, but all secondary effects (shadows, reflections, indirect lighting) remain. This feature required the project setting "Enable alpha channel support in post processing". */
+	UPROPERTY(EditAnywhere, AdvancedDisplay, BlueprintReadOnly, Category = Rendering, Interp)
+	uint8 bHoldout : 1;
+
 public:
 	UFUNCTION(BlueprintCallable, Category="Rendering|Components|ExponentialHeightFog")
 	ENGINE_API void SetFogDensity(float Value);
@@ -273,6 +277,9 @@ public:
 	UFUNCTION(BlueprintCallable, Category="Rendering|VolumetricFog")
 	ENGINE_API void SetSecondFogData(FExponentialHeightFogData NewValue);
 
+	UFUNCTION(BlueprintCallable, Category = "Rendering")
+	ENGINE_API void SetHoldout(bool bNewHoldout);
+
 protected:
 	//~ Begin UActorComponent Interface.
 	ENGINE_API virtual void CreateRenderState_Concurrent(FRegisterComponentContext* Context) override;
diff --git a/Engine/Source/Runtime/Engine/Classes/Components/SkyAtmosphereComponent.h b/Engine/Source/Runtime/Engine/Classes/Components/SkyAtmosphereComponent.h
--- a/Engine/Source/Runtime/Engine/Classes/Components/SkyAtmosphereComponent.h
+++ b/Engine/Source/Runtime/Engine/Classes/Components/SkyAtmosphereComponent.h
@@ -160,6 +160,10 @@ class USkyAtmosphereComponent : public USceneComponent
 	UPROPERTY(EditAnywhere, BlueprintReadOnly, interp, Category = "Art Direction", meta = (UIMin = 0.001f, UIMax = 10.0f, ClampMin = 0.001f))
 	float AerialPerspectiveStartDepth;
 
+	/** If this is True, this primitive will render black with an alpha of 0, but all secondary effects (shadows, reflections, indirect lighting) remain. This feature required the project setting "Enable alpha channel support in post processing". */
+	UPROPERTY(EditAnywhere, AdvancedDisplay, BlueprintReadOnly, Category = Rendering, Interp)
+	uint8 bHoldout : 1;
+
 
 	UFUNCTION(BlueprintCallable, Category = "Rendering")
 	ENGINE_API void OverrideAtmosphereLightDirection(int32 AtmosphereLightIndex, const FVector& LightDirection);
@@ -212,6 +216,9 @@ class USkyAtmosphereComponent : public USceneComponent
 	UFUNCTION(BlueprintCallable, Category = "Rendering")
 	ENGINE_API void SetHeightFogContribution(float NewValue);
 
+	UFUNCTION(BlueprintCallable, Category = "Rendering")
+	ENGINE_API void SetHoldout(bool bNewHoldout);
+
 	UFUNCTION(BlueprintCallable, Category = "Utilities", meta = (DisplayName = "Get Atmosphere Transmitance On Ground At Planet Top"))
 	ENGINE_API FLinearColor GetAtmosphereTransmitanceOnGroundAtPlanetTop(UDirectionalLightComponent* DirectionalLight);
 
diff --git a/Engine/Source/Runtime/Engine/Classes/Components/VolumetricCloudComponent.h b/Engine/Source/Runtime/Engine/Classes/Components/VolumetricCloudComponent.h
--- a/Engine/Source/Runtime/Engine/Classes/Components/VolumetricCloudComponent.h
+++ b/Engine/Source/Runtime/Engine/Classes/Components/VolumetricCloudComponent.h
@@ -139,6 +139,10 @@ class UVolumetricCloudComponent : public USceneComponent
 	UPROPERTY(EditAnywhere, BlueprintReadOnly, interp, Category = "Art Direction", meta = (UIMin = 0.0, UIMax = 100.0, ClampMin = 0.0, SliderExponent = 2.0))
 	float AerialPespectiveMieScatteringFadeDistance;
 
+	/** If this is True, this primitive will render black with an alpha of 0, but all secondary effects (shadows, reflections, indirect lighting) remain. This feature required the project setting "Enable alpha channel support in post processing". */
+	UPROPERTY(EditAnywhere, AdvancedDisplay, BlueprintReadOnly, Category = Rendering, Interp)
+	uint8 bHoldout : 1;
+
 
 	UFUNCTION(BlueprintCallable, Category = "Rendering")
 	ENGINE_API void SetLayerBottomAltitude(float NewValue);
@@ -181,6 +185,9 @@ class UVolumetricCloudComponent : public USceneComponent
 	UFUNCTION(BlueprintCallable, Category = "Rendering", meta = (DeprecatedFunction, DeprecationMessage = "This function has been replaced by SetShadowReflectionViewSampleCountScale."))
 	ENGINE_API void SetShadowReflectionSampleCountScale(float NewValue);
 
+	UFUNCTION(BlueprintCallable, Category = "Rendering")
+	ENGINE_API void SetHoldout(bool bNewHoldout);
+
 
 protected:
 	//~ Begin UActorComponent Interface.
diff --git a/Engine/Source/Runtime/Engine/Private/Components/HeightFogComponent.cpp b/Engine/Source/Runtime/Engine/Private/Components/HeightFogComponent.cpp
--- a/Engine/Source/Runtime/Engine/Private/Components/HeightFogComponent.cpp
+++ b/Engine/Source/Runtime/Engine/Private/Components/HeightFogComponent.cpp
@@ -349,6 +349,15 @@ void UExponentialHeightFogComponent::SetSecondFogData(FExponentialHeightFogData
 	}
 }
 
+void UExponentialHeightFogComponent::SetHoldout(bool bNewHoldout)
+{
+	if (bHoldout != bNewHoldout)
+	{
+		bHoldout = bNewHoldout;
+		MarkRenderStateDirty();
+	}
+}
+
 void UExponentialHeightFogComponent::Serialize(FArchive& Ar)
 {
 	Super::Serialize(Ar);
diff --git a/Engine/Source/Runtime/Engine/Private/Components/SkyAtmosphereComponent.cpp b/Engine/Source/Runtime/Engine/Private/Components/SkyAtmosphereComponent.cpp
--- a/Engine/Source/Runtime/Engine/Private/Components/SkyAtmosphereComponent.cpp
+++ b/Engine/Source/Runtime/Engine/Private/Components/SkyAtmosphereComponent.cpp
@@ -387,6 +387,15 @@ SKY_DECLARE_BLUEPRINT_SETFUNCTION_LINEARCOEFFICIENT(SkyLuminanceFactor);
 SKY_DECLARE_BLUEPRINT_SETFUNCTION(float, AerialPespectiveViewDistanceScale);
 SKY_DECLARE_BLUEPRINT_SETFUNCTION(float, HeightFogContribution);
 
+void USkyAtmosphereComponent::SetHoldout(bool bNewHoldout)
+{
+	if (bHoldout != bNewHoldout)
+	{
+		bHoldout = bNewHoldout;
+		MarkRenderStateDirty();
+	}
+}
+
 FLinearColor USkyAtmosphereComponent::GetAtmosphereTransmitanceOnGroundAtPlanetTop(UDirectionalLightComponent* DirectionalLight)
 {
 	if(DirectionalLight != nullptr)
diff --git a/Engine/Source/Runtime/Engine/Private/Components/VolumetricCloudComponent.cpp b/Engine/Source/Runtime/Engine/Private/Components/VolumetricCloudComponent.cpp
--- a/Engine/Source/Runtime/Engine/Private/Components/VolumetricCloudComponent.cpp
+++ b/Engine/Source/Runtime/Engine/Private/Components/VolumetricCloudComponent.cpp
@@ -166,6 +166,15 @@ CLOUD_DECLARE_BLUEPRINT_SETFUNCTION(float, ShadowTracingDistance);
 CLOUD_DECLARE_BLUEPRINT_SETFUNCTION(float, StopTracingTransmittanceThreshold);
 CLOUD_DECLARE_BLUEPRINT_SETFUNCTION(UMaterialInterface*, Material);
 
+void UVolumetricCloudComponent::SetHoldout(bool bNewHoldout)
+{
+	if (bHoldout != bNewHoldout)
+	{
+		bHoldout = bNewHoldout;
+		MarkRenderStateDirty();
+	}
+}
+
 void UVolumetricCloudComponent::SetReflectionViewSampleCountScale(float NewValue)
 {
 	if (AreDynamicDataChangesAllowed() && ReflectionViewSampleCountScaleValue != NewValue)
diff --git a/Engine/Source/Runtime/Engine/Private/VolumetricCloudProxy.cpp b/Engine/Source/Runtime/Engine/Private/VolumetricCloudProxy.cpp
--- a/Engine/Source/Runtime/Engine/Private/VolumetricCloudProxy.cpp
+++ b/Engine/Source/Runtime/Engine/Private/VolumetricCloudProxy.cpp
@@ -19,6 +19,7 @@ FVolumetricCloudSceneProxy::FVolumetricCloudSceneProxy(const UVolumetricCloudCom
 	, PlanetRadiusKm(InComponent->PlanetRadius)
 	, GroundAlbedo(InComponent->GroundAlbedo)
 	, bUsePerSampleAtmosphericLightTransmittance(InComponent->bUsePerSampleAtmosphericLightTransmittance)
+	, bHoldout(InComponent->bHoldout)
 	, SkyLightCloudBottomOcclusion(InComponent->SkyLightCloudBottomOcclusion)
 	, ViewSampleCountScale(InComponent->ViewSampleCountScale)
 	, ReflectionViewSampleCountScale(InComponent->ReflectionViewSampleCountScaleValue)
diff --git a/Engine/Source/Runtime/Engine/Public/Rendering/SkyAtmosphereCommonData.cpp b/Engine/Source/Runtime/Engine/Public/Rendering/SkyAtmosphereCommonData.cpp
--- a/Engine/Source/Runtime/Engine/Public/Rendering/SkyAtmosphereCommonData.cpp
+++ b/Engine/Source/Runtime/Engine/Public/Rendering/SkyAtmosphereCommonData.cpp
@@ -104,6 +104,7 @@ FAtmosphereSetup::FAtmosphereSetup(const USkyAtmosphereComponent& SkyAtmosphereC
 	}
 
 	TransmittanceMinLightElevationAngle = SkyAtmosphereComponent.TransmittanceMinLightElevationAngle;
+	bHoldout = SkyAtmosphereComponent.bHoldout ? 1 : 0;
 
 	UpdateTransform(SkyAtmosphereComponent.GetComponentTransform(), uint8(SkyAtmosphereComponent.TransformMode));
 }
diff --git a/Engine/Source/Runtime/Engine/Public/Rendering/SkyAtmosphereCommonData.h b/Engine/Source/Runtime/Engine/Public/Rendering/SkyAtmosphereCommonData.h
--- a/Engine/Source/Runtime/Engine/Public/Rendering/SkyAtmosphereCommonData.h
+++ b/Engine/Source/Runtime/Engine/Public/Rendering/SkyAtmosphereCommonData.h
@@ -43,6 +43,8 @@ struct FAtmosphereSetup
 
 	FLinearColor GroundAlbedo;
 
+	uint32 bHoldout;
+
 	float TransmittanceMinLightElevationAngle;
 
 	ENGINE_API FAtmosphereSetup(const USkyAtmosphereComponent& SkyAtmosphereComponent);
diff --git a/Engine/Source/Runtime/Engine/Public/VolumetricCloudProxy.h b/Engine/Source/Runtime/Engine/Public/VolumetricCloudProxy.h
--- a/Engine/Source/Runtime/Engine/Public/VolumetricCloudProxy.h
+++ b/Engine/Source/Runtime/Engine/Public/VolumetricCloudProxy.h
@@ -40,6 +40,7 @@ public:
 	float PlanetRadiusKm;
 	FColor GroundAlbedo;
 	bool bUsePerSampleAtmosphericLightTransmittance;
+	bool bHoldout;
 	float SkyLightCloudBottomOcclusion;
 	
 	float ViewSampleCountScale;
diff --git a/Engine/Source/Runtime/Renderer/Private/FogRendering.cpp b/Engine/Source/Runtime/Renderer/Private/FogRendering.cpp
--- a/Engine/Source/Runtime/Renderer/Private/FogRendering.cpp
+++ b/Engine/Source/Runtime/Renderer/Private/FogRendering.cpp
@@ -84,6 +84,7 @@ void SetupFogUniformParameters(FRDGBuilder& GraphBuilder, const FViewInfo& View,
 		OutParameters.IntegratedLightScatteringSampler = TStaticSamplerState<SF_Bilinear, AM_Clamp, AM_Clamp, AM_Clamp>::GetRHI();
 		OutParameters.VolumetricFogStartDistance = View.VolumetricFogStartDistance;
 		OutParameters.VolumetricFogNearFadeInDistanceInv = View.VolumetricFogNearFadeInDistanceInv;
+		OutParameters.bHoldout = View.bExponentialFogHoldout;
 	}
 }
 
@@ -246,6 +247,7 @@ void FSceneRenderer::InitFogConstants()
 				View.bUseDirectionalInscattering = SunLight != nullptr;
 				View.bEnableVolumetricFog = FogInfo.bEnableVolumetricFog;
 				View.VolumetricFogStartDistance = FogInfo.VolumetricFogStartDistance;
+				View.bExponentialFogHoldout = FogInfo.bHoldout;
 				View.VolumetricFogNearFadeInDistanceInv = FogInfo.VolumetricFogNearFadeInDistance > 0.0f ? (1.0f / FogInfo.VolumetricFogNearFadeInDistance) : 100000000.0f;
 			}
 		}
diff --git a/Engine/Source/Runtime/Renderer/Private/FogRendering.h b/Engine/Source/Runtime/Renderer/Private/FogRendering.h
--- a/Engine/Source/Runtime/Renderer/Private/FogRendering.h
+++ b/Engine/Source/Runtime/Renderer/Private/FogRendering.h
@@ -26,6 +26,7 @@ BEGIN_GLOBAL_SHADER_PARAMETER_STRUCT(FFogUniformParameters,)
 	SHADER_PARAMETER(float, ApplyVolumetricFog)
 	SHADER_PARAMETER(float, VolumetricFogStartDistance)
 	SHADER_PARAMETER(float, VolumetricFogNearFadeInDistanceInv)
+	SHADER_PARAMETER(int32, bHoldout)
 	SHADER_PARAMETER_TEXTURE(TextureCube, FogInscatteringColorCubemap)
 	SHADER_PARAMETER_SAMPLER(SamplerState, FogInscatteringColorSampler)
 	SHADER_PARAMETER_RDG_TEXTURE(Texture3D, IntegratedLightScattering)
diff --git a/Engine/Source/Runtime/Renderer/Private/SceneCore.cpp b/Engine/Source/Runtime/Renderer/Private/SceneCore.cpp
--- a/Engine/Source/Runtime/Renderer/Private/SceneCore.cpp
+++ b/Engine/Source/Runtime/Renderer/Private/SceneCore.cpp
@@ -418,6 +418,7 @@ FExponentialHeightFogSceneInfo::FExponentialHeightFogSceneInfo(const UExponentia
 	VolumetricFogDistance = FMath::Max(InComponent->VolumetricFogStartDistance + InComponent->VolumetricFogDistance, 0.0f);
 	VolumetricFogStaticLightingScatteringIntensity = FMath::Max(InComponent->VolumetricFogStaticLightingScatteringIntensity, 0.0f);
 	bOverrideLightColorsWithFogInscatteringColors = InComponent->bOverrideLightColorsWithFogInscatteringColors;
+	bHoldout = InComponent->bHoldout;
 
 	VolumetricFogStartDistance = InComponent->VolumetricFogStartDistance;
 	VolumetricFogNearFadeInDistance = InComponent->VolumetricFogNearFadeInDistance;
diff --git a/Engine/Source/Runtime/Renderer/Private/SceneCore.h b/Engine/Source/Runtime/Renderer/Private/SceneCore.h
--- a/Engine/Source/Runtime/Renderer/Private/SceneCore.h
+++ b/Engine/Source/Runtime/Renderer/Private/SceneCore.h
@@ -167,6 +167,7 @@ public:
 	float VolumetricFogDistance;
 	float VolumetricFogStaticLightingScatteringIntensity;
 	bool bOverrideLightColorsWithFogInscatteringColors;
+	bool bHoldout;
 	float VolumetricFogStartDistance;
 	float VolumetricFogNearFadeInDistance;
 
diff --git a/Engine/Source/Runtime/Renderer/Private/SceneRendering.h b/Engine/Source/Runtime/Renderer/Private/SceneRendering.h
--- a/Engine/Source/Runtime/Renderer/Private/SceneRendering.h
+++ b/Engine/Source/Runtime/Renderer/Private/SceneRendering.h
@@ -1378,6 +1378,7 @@ public:
 	FVector4f ExponentialFogParameters3;
 	FVector4f SkyAtmosphereAmbientContributionColorScale;
 	bool bEnableVolumetricFog;
+	bool bExponentialFogHoldout;
 	float VolumetricFogStartDistance;
 	float VolumetricFogNearFadeInDistanceInv;
 	FVector2f SinCosInscatteringColorCubemapRotation;
diff --git a/Engine/Source/Runtime/Renderer/Private/SkyAtmosphereRendering.cpp b/Engine/Source/Runtime/Renderer/Private/SkyAtmosphereRendering.cpp
--- a/Engine/Source/Runtime/Renderer/Private/SkyAtmosphereRendering.cpp
+++ b/Engine/Source/Runtime/Renderer/Private/SkyAtmosphereRendering.cpp
@@ -431,6 +431,7 @@ static void CopyAtmosphereSetupToUniformShaderParameters(FAtmosphereUniformShade
 	COPYMACRO(AbsorptionDensity1LinearTerm);
 	COPYMACRO(AbsorptionExtinction);
 	COPYMACRO(GroundAlbedo);
+	COPYMACRO(bHoldout);
 #undef COPYMACRO
 }
 
diff --git a/Engine/Source/Runtime/Renderer/Private/SkyAtmosphereRendering.h b/Engine/Source/Runtime/Renderer/Private/SkyAtmosphereRendering.h
--- a/Engine/Source/Runtime/Renderer/Private/SkyAtmosphereRendering.h
+++ b/Engine/Source/Runtime/Renderer/Private/SkyAtmosphereRendering.h
@@ -53,6 +53,7 @@ BEGIN_GLOBAL_SHADER_PARAMETER_STRUCT(FAtmosphereUniformShaderParameters, )
 	SHADER_PARAMETER(float, AbsorptionDensity1LinearTerm)
 	SHADER_PARAMETER(FLinearColor, AbsorptionExtinction)
 	SHADER_PARAMETER(FLinearColor, GroundAlbedo)
+	SHADER_PARAMETER(int32, bHoldout)
 END_GLOBAL_SHADER_PARAMETER_STRUCT()
 
 // These parameters are shared on the view global uniform buffer and are dynamically changed with cvars.
diff --git a/Engine/Source/Runtime/Renderer/Private/VolumetricCloudRendering.cpp b/Engine/Source/Runtime/Renderer/Private/VolumetricCloudRendering.cpp
--- a/Engine/Source/Runtime/Renderer/Private/VolumetricCloudRendering.cpp
+++ b/Engine/Source/Runtime/Renderer/Private/VolumetricCloudRendering.cpp
@@ -1611,6 +1611,7 @@ void FSceneRenderer::InitVolumetricCloudsForViews(FRDGBuilder& GraphBuilder, boo
 			CloudGlobalShaderParams.ShadowTracingMaxDistance			= KilometersToCentimeters * FMath::Max(0.1f, CloudProxy.ShadowTracingDistance);
 			CloudGlobalShaderParams.InvDistanceToSampleCountMax			= 1.0f / FMath::Max(1.0f, KilometersToCentimeters * CVarVolumetricCloudDistanceToSampleMaxCount.GetValueOnAnyThread());
 			CloudGlobalShaderParams.StopTracingTransmittanceThreshold	= FMath::Clamp(CloudProxy.StopTracingTransmittanceThreshold, 0.0f, 1.0f);
+			CloudGlobalShaderParams.bHoldout2							= CloudProxy.bHoldout ? 1 : 0;
 
 			auto PrepareCloudShadowMapLightData = [&](FLightSceneProxy* AtmosphericLight, int LightIndex)
 			{
diff --git a/Engine/Source/Runtime/Renderer/Private/VolumetricCloudRendering.h b/Engine/Source/Runtime/Renderer/Private/VolumetricCloudRendering.h
--- a/Engine/Source/Runtime/Renderer/Private/VolumetricCloudRendering.h
+++ b/Engine/Source/Runtime/Renderer/Private/VolumetricCloudRendering.h
@@ -38,6 +38,7 @@ BEGIN_SHADER_PARAMETER_STRUCT(FVolumetricCloudCommonShaderParameters, )
 	SHADER_PARAMETER(int32, SampleCountMax)
 	SHADER_PARAMETER(float, InvDistanceToSampleCountMax)
 	SHADER_PARAMETER(int32, ShadowSampleCountMax)
+	SHADER_PARAMETER(int32, bHoldout2)
 	SHADER_PARAMETER(float, ShadowTracingMaxDistance)
 	SHADER_PARAMETER(float, StopTracingTransmittanceThreshold)
 	SHADER_PARAMETER(float, SkyLightCloudBottomVisibility)
