Substrate - fixed compilation issue due to some compiler being more strick with constness.

diff --git a/Engine/Shaders/Private/MaterialTemplate.ush b/Engine/Shaders/Private/MaterialTemplate.ush
--- a/Engine/Shaders/Private/MaterialTemplate.ush
+++ b/Engine/Shaders/Private/MaterialTemplate.ush
@@ -2567,7 +2567,7 @@ float4 SceneTextureLookup(float2 UV, int SceneTextureIndex, bool bFiltered)
 			const FSubstrateSubsurfaceHeader SSSHeader = SubstrateLoadSubsurfaceHeader(SubstrateSceneTexture.MaterialTextureArray, SubstrateSceneTexture.FirstSliceStoringSubstrateSSSData, SubstrateAddressing.PixelCoords);
 			const uint SSSType = SubstrateSubSurfaceHeaderGetSSSType(SSSHeader);
 
-			const FSubstrateBSDF BSDF = UnpackSubstrateBSDFIn(SubstrateSceneTexture.MaterialTextureArray, SubstrateAddressing, SubstratePixelHeader);
+			FSubstrateBSDF BSDF = UnpackSubstrateBSDFIn(SubstrateSceneTexture.MaterialTextureArray, SubstrateAddressing, SubstratePixelHeader);
 			switch (SceneTextureIndex)
 			{
 				// order needs to match to ESceneTextureId
