Fixes regression where TSR history resurrection could happen on tile's hidden pixels

diff --git a/Engine/Shaders/Private/TemporalSuperResolution/TSRRejectShading.usf b/Engine/Shaders/Private/TemporalSuperResolution/TSRRejectShading.usf
--- a/Engine/Shaders/Private/TemporalSuperResolution/TSRRejectShading.usf
+++ b/Engine/Shaders/Private/TemporalSuperResolution/TSRRejectShading.usf
@@ -787,8 +787,9 @@ bool IsWorthMeasuringRejectionOfResurrectedFrame(tsr_simd_half PrevRejectionBlen
 CALL_SITE_DEBUGLOC
 tsr_simd_bool ShouldResurrectHistory(tsr_simd_half PrevRejectionBlendFinal, tsr_simd_half ResurrectionRejectionBlendFinal, tsr_simd_bool bResurrectionIsCloser)
 {
+	tsr_simd_bool bIsValidTilePixel = IsValidTilePixel<LANE_STRIDE_X, LANE_STRIDE_Y>(TileOverscan);
 	tsr_simd_bool bResurrectionIsBetter = (ResurrectionRejectionBlendFinal - PrevRejectionBlendFinal) > tsr_simd_half::Const(0.1);
-	tsr_simd_bool bResurrectHistory = and(bResurrectionIsCloser, bResurrectionIsBetter);
+	tsr_simd_bool bResurrectHistory = and(and(bResurrectionIsCloser, bResurrectionIsBetter), bIsValidTilePixel);
 
 	return bResurrectHistory;
 }
