Add support for IES profile in cluster shading pass.

diff --git a/Engine/Shaders/Private/ClusteredDeferredShadingPixelShader.usf b/Engine/Shaders/Private/ClusteredDeferredShadingPixelShader.usf
--- a/Engine/Shaders/Private/ClusteredDeferredShadingPixelShader.usf
+++ b/Engine/Shaders/Private/ClusteredDeferredShadingPixelShader.usf
@@ -7,7 +7,7 @@
 
 // Following the lead of Tiled Deferred 
 #define SUPPORT_CONTACT_SHADOWS 0
-#define USE_IES_PROFILE 0
+#define USE_IES_PROFILE 1
 #define NON_DIRECTIONAL_DIRECT_LIGHTING 0
 #define SUBSTRATE_SSS_TRANSMISSION USE_TRANSMISSION
 
@@ -261,6 +261,9 @@ float4 GetLightGridLocalLighting(
 		FDeferredLightData LightData = ConvertToDeferredLight(LocalLight);
 		LightData.ShadowedBits = 1;
 		LightData.bRectLight = false;
+		#if USE_IES_PROFILE
+		LightData.Color *= ComputeLightProfileMultiplier(TranslatedWorldPosition, LightData.TranslatedWorldPosition, -LightData.Direction, LightData.Tangent, LightData.IESAtlasIndex);
+		#endif
 		float4 LightAttenuation = float4(1, 1, 1, 1);
 
 		// Find packed VSM projection result if present
diff --git a/Engine/Source/Runtime/Renderer/Private/LightRendering.cpp b/Engine/Source/Runtime/Renderer/Private/LightRendering.cpp
--- a/Engine/Source/Runtime/Renderer/Private/LightRendering.cpp
+++ b/Engine/Source/Runtime/Renderer/Private/LightRendering.cpp
@@ -1195,7 +1195,6 @@ void FSceneRenderer::GatherAndSortLights(FSortedLightSetSceneInfo& OutSortedLigh
 					// Usually it'll have shadows and stuff anyway.
 					// Rect lights are not supported as the performance impact is significant even if not used, for now, left for trad. deferred.
 					const bool bClusteredDeferredSupported =
-						!SortedLightInfo->SortKey.Fields.bTextureProfile &&
 						(!SortedLightInfo->SortKey.Fields.bShadowed || bShadowedLightsInClustered) &&
 						(!SortedLightInfo->SortKey.Fields.bLightFunction || bUseLightFunctionAtlas) &&
 						!SortedLightInfo->SortKey.Fields.bUsesLightingChannels
