Helps a shader compiler

diff --git a/Engine/Shaders/Private/TemporalSuperResolution/TSRRejectShading.usf b/Engine/Shaders/Private/TemporalSuperResolution/TSRRejectShading.usf
--- a/Engine/Shaders/Private/TemporalSuperResolution/TSRRejectShading.usf
+++ b/Engine/Shaders/Private/TemporalSuperResolution/TSRRejectShading.usf
@@ -1208,17 +1208,20 @@ void MainCS(
 			History = select(tsr_simd_boolC::Vectorize(bResurrectHistory), ResurrectedHistory, History);
 			HistoryUncertainty = select(bResurrectHistory, ResurrectedHistoryUncertainty, HistoryUncertainty);
 			
-			// Compute the motion vector for the resurrection
-			tsr_simd_float2 ResurrectionMotionVector = ComputeResurrectionMotionVector(DeviceZ);
-			
-			// Amends the motion vector for resurrection areas.
-			UNROLL_N(SIMD_SIZE)
-			for (uint ElementIndex = 0; ElementIndex < SIMD_SIZE; ElementIndex++)
+			PLATFORM_SPECIFIC_ISOLATE
 			{
-				tsr_short2 OutputPixelPos = ComputeElementOutputPixelPos(ElementIndex);
-				OutputPixelPos.x = select(bResurrectHistory.GetElement(ElementIndex), OutputPixelPos.x, tsr_short(-1));
-
-				DilatedVelocityOutput[OutputPixelPos] = ResurrectionMotionVector.GetElement(ElementIndex);
+				// Compute the motion vector for the resurrection
+				tsr_simd_float2 ResurrectionMotionVector = ComputeResurrectionMotionVector(DeviceZ);
+			
+				// Amends the motion vector for resurrection areas.
+				UNROLL_N(SIMD_SIZE)
+				for (uint ElementIndex = 0; ElementIndex < SIMD_SIZE; ElementIndex++)
+				{
+					tsr_short2 OutputPixelPos = ComputeElementOutputPixelPos(ElementIndex);
+					OutputPixelPos.x = select(bResurrectHistory.GetElement(ElementIndex), OutputPixelPos.x, tsr_short(-1));
+
+					DilatedVelocityOutput[OutputPixelPos] = ResurrectionMotionVector.GetElement(ElementIndex);
+				}
 			}
 		}
 		#endif // CONFIG_HISTORY_RESURRECTION
