Path Tracing: Simplify logic in TraceTransparentRay to reduce live state, simplify decal math

The main change here is making the decal logic accumulate inside the decal shader itself so we can reduce the amount of logic in the path tracer and also reduce the amount of live state required for the decal calls to TraceRay (532 bytes down to 456 bytes) and CallShader (544 bytes down to 468 bytes). This appears to improve performance on decal heavy scenes.

Performance is slightly reduced in other scenes which do not use decals. This can probably be compensated for by further working to reduce the live state by applying similar optimizations to the other payloads.

diff --git a/Engine/Shaders/Private/PathTracing/PathTracingCommon.ush b/Engine/Shaders/Private/PathTracing/PathTracingCommon.ush
--- a/Engine/Shaders/Private/PathTracing/PathTracingCommon.ush
+++ b/Engine/Shaders/Private/PathTracing/PathTracingCommon.ush
@@ -775,6 +775,5 @@ FPathTracingPayload UnpackPathTracingPayload(FPackedPathTracingPayload Input, FR
 	Output.SpecularColor = ComputeF0(Output.Specular, Output.BaseColor, Output.Metallic);
 #endif // PATHTRACING_SUBSTRATE_PAYLOAD
 #endif // IS_PAYLOAD_ENABLED(RT_PAYLOAD_TYPE_GPULIGHTMASS)
-	Output.TranslatedWorldPos = Ray.Origin + Output.HitT * Ray.Direction;
 	return Output;
 }
diff --git a/Engine/Shaders/Private/PathTracing/PathTracingCore.ush b/Engine/Shaders/Private/PathTracing/PathTracingCore.ush
--- a/Engine/Shaders/Private/PathTracing/PathTracingCore.ush
+++ b/Engine/Shaders/Private/PathTracing/PathTracingCore.ush
@@ -217,128 +217,94 @@ void AccumulateRadiance(inout float3 TotalRadiance, float3 PathRadiance, bool bI
 	TotalRadiance += PathRadiance;
 }
 
-struct FDecalData
+bool HasDecalData(FDecalShaderPayload DecalPayload)
 {
-	float4 BaseColor;
-	float4 WorldNormal;
-	float4 MetallicSpecularRoughness;
-	float3 Emissive;
-
-	bool HasData()
-	{
-		return min3(BaseColor.a, WorldNormal.a, MetallicSpecularRoughness.a) < 1.0;
-	}
+	return min3(DecalPayload.GetBaseColor().a, DecalPayload.GetWorldNormal().a, DecalPayload.GetMetallicSpecularRoughness().a) < 1.0;
+}
 
-	bool ApplyDecal(inout FPathTracingPayload HitPayload)
-	{
+bool ApplyDecal(FDecalShaderPayload DecalPayload, inout FPathTracingPayload HitPayload)
+{
 #if PATHTRACING_SUBSTRATE_PAYLOAD
-		float Roughness = HitPayload.RoughnessData.x;
-		float Metallic = F0RGBToMetallic(HitPayload.SpecularColor);
-		float Specular = F0RGBToDielectricSpecular(HitPayload.SpecularColor);
-		float3 BaseColor = lerp(HitPayload.DiffuseColor, HitPayload.SpecularColor, Metallic);
+	float Roughness = HitPayload.RoughnessData.x;
+	float Metallic = F0RGBToMetallic(HitPayload.SpecularColor);
+	float Specular = F0RGBToDielectricSpecular(HitPayload.SpecularColor);
+	float3 BaseColor = lerp(HitPayload.DiffuseColor, HitPayload.SpecularColor, Metallic);
 #else
-		float3 BaseColor = HitPayload.BaseColor;
-		float Metallic = HitPayload.Metallic;
-		float Specular = HitPayload.Specular;
-		float Roughness = HitPayload.Roughness;
+	float3 BaseColor = HitPayload.BaseColor;
+	float Metallic = HitPayload.Metallic;
+	float Specular = HitPayload.Specular;
+	float Roughness = HitPayload.Roughness;
 #endif // PATHTRACING_SUBSTRATE_PAYLOAD
 
-		bool bHasColorResponse = HitPayload.HasDecalResponseColor() && this.BaseColor.a < 1.0;
-		bool bHasNormalResponse = HitPayload.HasDecalResponseNormal() && this.WorldNormal.a < 1.0;
-		bool bHasRoughnessResponse = HitPayload.HasDecalResponseRoughness() && this.MetallicSpecularRoughness.a < 1.0;
-
-		if (bHasColorResponse)
-		{
-			BaseColor = BaseColor * this.BaseColor.a + this.BaseColor.rgb;
-		}
-
-		if (bHasNormalResponse)
-		{
-			// We normalize the normal to get smoother visual result (it helps to avoid having D_GGX explodes toward infinity, and matches ApplyDBufferData(...))
-			HitPayload.WorldNormal = normalize(HitPayload.WorldNormal * this.WorldNormal.a + this.WorldNormal.rgb);
-		}
-
-		if (bHasRoughnessResponse)
-		{
-			Metallic = Metallic * this.MetallicSpecularRoughness.a + this.MetallicSpecularRoughness.x;
-			Specular = Specular * this.MetallicSpecularRoughness.a + this.MetallicSpecularRoughness.y;
-			Roughness = Roughness * this.MetallicSpecularRoughness.a + this.MetallicSpecularRoughness.z;
-		}
+	float4 DecalBaseColor = DecalPayload.GetBaseColor();
+	float4 DecalWorldNormal = DecalPayload.GetWorldNormal();
+	float4 DecalMetallicSpecularRoughness = DecalPayload.GetMetallicSpecularRoughness();
 
-		HitPayload.Radiance += this.Emissive;
+	bool bHasColorResponse = HitPayload.HasDecalResponseColor() && DecalBaseColor.a < 1.0;
+	bool bHasNormalResponse = HitPayload.HasDecalResponseNormal() && DecalWorldNormal.a < 1.0;
+	bool bHasRoughnessResponse = HitPayload.HasDecalResponseRoughness() && DecalMetallicSpecularRoughness.a < 1.0;
 
-#if PATHTRACING_SUBSTRATE_PAYLOAD
-		if (bHasColorResponse || bHasRoughnessResponse)
-		{
-			// the encode-decode is potentially lossy, so only do it if we actually need to
-			HitPayload.DiffuseColor = BaseColor - BaseColor * Metallic;
-			HitPayload.SpecularColor = lerp(DielectricSpecularToF0(Specular), BaseColor, Metallic);
-			HitPayload.RoughnessData.x = Roughness;
-		}
-#else
-		// no need for an encode-decode here, so safe to overwrite even if nothing changes
-		HitPayload.BaseColor = BaseColor;
-		HitPayload.Metallic = Metallic;
-		HitPayload.Specular = Specular;
-		HitPayload.Roughness = Roughness;
-#endif
-		return HasData();
-	}
-};
-
-void CombineDecal(FDecalShaderPayload DecalPayload, inout FDecalData DecalData)
-{
-	const float DecalTransparency = DecalPayload.GetTransparency();
-	const uint DecalFlags = DecalPayload.GetFlags();
-
-	if (DecalFlags & DECAL_WRITE_BASE_COLOR_FLAG)
+	if (bHasColorResponse)
 	{
-		DecalData.BaseColor.rgb = DecalData.BaseColor.rgb * DecalTransparency + DecalPayload.GetBaseColor();
-		DecalData.BaseColor.a *= DecalTransparency;
+		BaseColor = BaseColor * DecalBaseColor.a + DecalBaseColor.rgb;
 	}
 
-	if (DecalFlags & DECAL_WRITE_NORMAL_FLAG)
+	if (bHasNormalResponse)
 	{
-		DecalData.WorldNormal.rgb = DecalData.WorldNormal.rgb * DecalTransparency + DecalPayload.GetWorldNormal();
-		DecalData.WorldNormal.a *= DecalTransparency;
+		// We normalize the normal to get smoother visual result (it helps to avoid having D_GGX explodes toward infinity, and matches ApplyDBufferData(...))
+		HitPayload.WorldNormal = normalize(HitPayload.WorldNormal * DecalWorldNormal.a + DecalWorldNormal.rgb);
 	}
 
-	if (DecalFlags & DECAL_WRITE_ROUGHNESS_SPECULAR_METALLIC_FLAG)
+	if (bHasRoughnessResponse)
 	{
-		DecalData.MetallicSpecularRoughness.rgb = DecalData.MetallicSpecularRoughness.rgb * DecalTransparency + DecalPayload.GetMetallicSpecularRoughness();
-		DecalData.MetallicSpecularRoughness.a *= DecalTransparency;
+		Metallic  = Metallic  * DecalMetallicSpecularRoughness.a + DecalMetallicSpecularRoughness.x;
+		Specular  = Specular  * DecalMetallicSpecularRoughness.a + DecalMetallicSpecularRoughness.y;
+		Roughness = Roughness * DecalMetallicSpecularRoughness.a + DecalMetallicSpecularRoughness.z;
 	}
 
-	if (DecalFlags & DECAL_WRITE_EMISSIVE_FLAG)
+	HitPayload.Radiance += DecalPayload.GetEmissive();
+
+#if PATHTRACING_SUBSTRATE_PAYLOAD
+	if (bHasColorResponse || bHasRoughnessResponse)
 	{
-		DecalData.Emissive += DecalPayload.GetEmissive();
+		// the encode-decode is potentially lossy, so only do it if we actually need to
+		HitPayload.DiffuseColor = BaseColor - BaseColor * Metallic;
+		HitPayload.SpecularColor = lerp(DielectricSpecularToF0(Specular), BaseColor, Metallic);
+		HitPayload.RoughnessData.x = Roughness;
 	}
+#else
+	// no need for an encode-decode here, so safe to overwrite even if nothing changes
+	HitPayload.BaseColor = BaseColor;
+	HitPayload.Metallic = Metallic;
+	HitPayload.Specular = Specular;
+	HitPayload.Roughness = Roughness;
+#endif
+	return bHasColorResponse || bHasNormalResponse || bHasRoughnessResponse;
 }
 
-void EvaluateDecals(FPathState PathState, uint InstanceInclusionMask, FPackedPathTracingPayload PackedPayload, inout FDecalData DecalData)
+void EvaluateDecals(FPathState PathState, uint InstanceInclusionMask, float HitT, inout FDecalShaderPayload DecalPayload)
 {
-	const float3 HitTranslatedWorldPos = PathState.Ray.Origin + PackedPayload.HitT * PathState.Ray.Direction;
-
 	if (PathState.PathRoughness < MeshDecalRoughnessCutoff)
 	{
 		FRayDesc DecalRay = (FRayDesc)0;
 		// go a bit further than the current hit, to handle decals partially inside the model
 		DecalRay.Direction = -PathState.Ray.Direction;
-		DecalRay.Origin = HitTranslatedWorldPos - DecalRay.Direction * MeshDecalBias;
+		DecalRay.Origin = DecalPayload.TranslatedWorldPos - DecalRay.Direction * MeshDecalBias;
 		DecalRay.TMin = 0.0f;
-		DecalRay.TMax = PackedPayload.HitT + MeshDecalBias;
+		DecalRay.TMax = HitT + MeshDecalBias;
 
-		const bool bReverseCulling = true;
 		const uint MissShaderIndex = 0; // TODO
 
 		for (;;)
 		{
-			FDecalShaderPayload DecalPayload = (FDecalShaderPayload)0;
-			DecalPayload.SetInputTranslatedWorldPosition(HitTranslatedWorldPos);
 			DecalPayload.SetMiss();
 			TraceRay(
 				DecalTLAS,
-				bReverseCulling ? RAY_FLAG_CULL_FRONT_FACING_TRIANGLES : RAY_FLAG_CULL_BACK_FACING_TRIANGLES,
+#if DECAL_USE_REVERSE_CULLING
+				RAY_FLAG_CULL_FRONT_FACING_TRIANGLES,
+#else
+				RAY_FLAG_CULL_BACK_FACING_TRIANGLES,
+#endif
 				InstanceInclusionMask,
 				RAY_TRACING_SHADER_SLOT_MATERIAL,
 				RAY_TRACING_NUM_SHADER_SLOTS,
@@ -351,14 +317,6 @@ void EvaluateDecals(FPathState PathState, uint InstanceInclusionMask, FPackedPat
 				break;
 			}
 
-			if (bReverseCulling)
-			{
-				// also reverse normal
-				DecalPayload.SetWorldNormal(-DecalPayload.GetWorldNormal());
-			}
-
-			CombineDecal(DecalPayload, DecalData);
-
 			// prepare next step around the loop
 			// retrace the exact same ray with TMin one ulp past the hit we just found
 			DecalRay.TMin = asfloat(asuint(DecalPayload.HitT) + 1);
@@ -368,17 +326,11 @@ void EvaluateDecals(FPathState PathState, uint InstanceInclusionMask, FPackedPat
 #if PLATFORM_SUPPORTS_CALLABLE_SHADERS
 	if (PathState.PathRoughness < DecalRoughnessCutoff)
 	{
-		FDecalLoopCount DecalLoopCount = DecalGridLookup(HitTranslatedWorldPos);
+		FDecalLoopCount DecalLoopCount = DecalGridLookup(DecalPayload.TranslatedWorldPos);
 		for (uint Index = 0, Num = DecalLoopCount.NumDecals; Index < Num; ++Index)
 		{
 			uint DecalId = GetDecalId(Index, DecalLoopCount);
-
-			FDecalShaderPayload DecalPayload = (FDecalShaderPayload)0;
-			DecalPayload.SetInputTranslatedWorldPosition(HitTranslatedWorldPos);
-
 			CallShader(DecalId, DecalPayload);
-
-			CombineDecal(DecalPayload, DecalData);
 		}
 	}
 #endif
@@ -410,13 +362,6 @@ FPathTracingPayload TraceTransparentRay(inout FPathState PathState, int Bounce,
 		PathState.PathThroughput = 0;
 		return Payload;
 	}
-	FVolumeIntersectionList VolumeIsectList = VolumeIntersect(PathState.Ray.Origin, PathState.Ray.Direction, PathState.Ray.TMin, PathState.Ray.TMax);
-	if (VolumeIsectList.HitBlocker())
-	{
-		// We will hit the volume's blocker, 
-		// Clip the ray just before the planet hit so that we can intersect all other geometry but this one
-		PathState.Ray.TMax = asfloat(asuint(VolumeIsectList.BlockerHitT) - 1);
-	}
 	for (;;)
 	{
 		FPackedPathTracingPayload PackedPayload = InitPathTracingPayload(PathState.FirstScatterType, ApproximateCaustics ? PathState.PathRoughness : 0.0);
@@ -437,49 +382,8 @@ FPathTracingPayload TraceTransparentRay(inout FPathState PathState, int Bounce,
 			PathState.Ray.GetNativeDesc(),
 			PackedPayload);
 
-		FDecalData DecalData;
-		DecalData.BaseColor = float4(0, 0, 0, 1);
-		DecalData.WorldNormal = float4(0, 0, 0, 1);
-		DecalData.MetallicSpecularRoughness = float4(0, 0, 0, 1);
-		DecalData.Emissive = 0;
-
-		if (!PackedPayload.IsMiss() && PackedPayload.IsDecalReceiver())
-		{
-			EvaluateDecals(PathState, InstanceInclusionMask, PackedPayload, DecalData);
-
-#if USE_DBUFFER
-			if (PackedPayload.UsesDBufferLookup() && DecalData.HasData())
-			{
-				if (EnableDBuffer)
-				{
-					// Retrace ray with DBuffer data
-					PackedPayload = InitPathTracingPayload(PathState.FirstScatterType, ApproximateCaustics ? PathState.PathRoughness : 0.0);
-
-					PackedPayload.SetDBufferA(DecalData.BaseColor);
-					PackedPayload.SetDBufferB(DecalData.WorldNormal);
-					PackedPayload.SetDBufferC(DecalData.MetallicSpecularRoughness);
-					TraceRay(
-						TLAS,
-						RayFlags,
-						InstanceInclusionMask,
-						RAY_TRACING_SHADER_SLOT_MATERIAL,
-						RAY_TRACING_NUM_SHADER_SLOTS,
-						MissShaderIndex,
-						PathState.Ray.GetNativeDesc(),
-						PackedPayload);
-				}
-				else
-				{
-					// The user asked for a dbuffer lookup, but we have disabled the double-trace.
-					// Set all the decal response flags so we at least pickup some kind of response
-					// even if it won't be exactly what the user had in mind.
-					PackedPayload.SetFlag(PATH_TRACING_PAYLOAD_OUTPUT_FLAG_DECAL_RESPONSE_COLOR | PATH_TRACING_PAYLOAD_OUTPUT_FLAG_DECAL_RESPONSE_NORMAL | PATH_TRACING_PAYLOAD_OUTPUT_FLAG_DECAL_RESPONSE_ROUGHNESS);
-				}
-			}
-#endif
-		}
-
-		if (PackedPayload.IsMiss() && VolumeIsectList.HitBlocker())
+		FVolumeIntersectionList VolumeIsectList = VolumeIntersect(PathState.Ray.Origin, PathState.Ray.Direction, PathState.Ray.TMin, PathState.Ray.TMax);
+		if (VolumeIsectList.HitBlocker() && (PackedPayload.IsMiss() || VolumeIsectList.BlockerHitT < PackedPayload.HitT))
 		{
 			// we didn't hit any real geometry, but we did hit the volume's blocker geometry
 			// create a virtual hit with it here
@@ -509,7 +413,6 @@ FPathTracingPayload TraceTransparentRay(inout FPathState PathState, int Bounce,
 				// light hits are fully transparent to match how the light loop works (the light surface does not cast a shadow)
 				LightHitPayload.BSDFOpacity = 1.0;
 				LightHitPayload.TransparencyColor = 1.0;
-				LightHitPayload.PrimitiveLightingChannelMask = 7;
 				LightHitPayload.SetFrontFace();
 				PackedPayload = PackPathTracingPayload(LightHitPayload);
 			}
@@ -643,17 +546,53 @@ FPathTracingPayload TraceTransparentRay(inout FPathState PathState, int Bounce,
 			break;
 		}
 
+		FDecalShaderPayload DecalPayload = InitDecalShaderPayload(PathState.Ray.Origin + PackedPayload.HitT * PathState.Ray.Direction);
+		if (PackedPayload.IsDecalReceiver())
+		{
+			EvaluateDecals(PathState, InstanceInclusionMask, PackedPayload.HitT, DecalPayload);
+#if USE_DBUFFER
+			if (PackedPayload.UsesDBufferLookup() && HasDecalData(DecalPayload))
+			{
+				if (EnableDBuffer)
+				{
+					// Retrace ray with DBuffer data
+					PackedPayload = InitPathTracingPayload(PathState.FirstScatterType, ApproximateCaustics ? PathState.PathRoughness : 0.0);
+
+					PackedPayload.SetDBufferA(DecalPayload.GetBaseColor());
+					PackedPayload.SetDBufferB(DecalPayload.GetWorldNormal());
+					PackedPayload.SetDBufferC(DecalPayload.GetMetallicSpecularRoughness());
+					TraceRay(
+						TLAS,
+						RayFlags,
+						InstanceInclusionMask,
+						RAY_TRACING_SHADER_SLOT_MATERIAL,
+						RAY_TRACING_NUM_SHADER_SLOTS,
+						MissShaderIndex,
+						PathState.Ray.GetNativeDesc(),
+						PackedPayload);
+				}
+				else
+				{
+					// The user asked for a dbuffer lookup, but we have disabled the double-trace.
+					// Set all the decal response flags so we at least pickup some kind of response
+					// even if it won't be exactly what the user had in mind.
+					PackedPayload.SetFlag(PATH_TRACING_PAYLOAD_OUTPUT_FLAG_DECAL_RESPONSE_COLOR | PATH_TRACING_PAYLOAD_OUTPUT_FLAG_DECAL_RESPONSE_NORMAL | PATH_TRACING_PAYLOAD_OUTPUT_FLAG_DECAL_RESPONSE_ROUGHNESS);
+				}
+			}
+#endif
+		}
+
 		// Unpack the payload
 		FPathTracingPayload HitPayload = UnpackPathTracingPayload(PackedPayload, PathState.Ray);
-		if (DecalData.ApplyDecal(HitPayload))
+		if (ApplyDecal(DecalPayload, HitPayload))
 		{
 			// if we ran any decal logic, update the payload
 			if (ApproximateCaustics)
 			{
 #if PATHTRACING_SUBSTRATE_PAYLOAD
-				Payload.RoughnessData.x = max(Payload.RoughnessData.x, PathState.PathRoughness);
+				HitPayload.RoughnessData.x = max(HitPayload.RoughnessData.x, PathState.PathRoughness);
 #else
-				Payload.Roughness = max(Payload.Roughness, PathState.PathRoughness);
+				HitPayload.Roughness = max(HitPayload.Roughness, PathState.PathRoughness);
 #endif
 			}
 			HitPayload.WorldNormal = AdjustShadingNormal(HitPayload.WorldNormal, HitPayload.WorldGeoNormal, PathState.Ray.Direction);
diff --git a/Engine/Shaders/Private/RayTracing/RayTracingCommon.ush b/Engine/Shaders/Private/RayTracing/RayTracingCommon.ush
--- a/Engine/Shaders/Private/RayTracing/RayTracingCommon.ush
+++ b/Engine/Shaders/Private/RayTracing/RayTracingCommon.ush
@@ -988,66 +988,80 @@ CHECK_RT_PAYLOAD_SIZE(FPackedMaterialClosestHitPayload)
 ///////////////////////////////////////////////////////////////////////////////////////////////////
 // Decal parameters
 
+// Must match EDecalsWriteFlags
+// TODO: Move to a shared header?
 #define DECAL_WRITE_BASE_COLOR_FLAG						(1 << 0)
 #define DECAL_WRITE_NORMAL_FLAG							(1 << 1)
 #define DECAL_WRITE_ROUGHNESS_SPECULAR_METALLIC_FLAG	(1 << 2)
 #define DECAL_WRITE_EMISSIVE_FLAG						(1 << 3)
 
+#define DECAL_USE_REVERSE_CULLING 						1
+
 #if IS_PAYLOAD_ENABLED(RT_PAYLOAD_TYPE_DECALS)
 struct FDecalShaderPayload : FMinimalPayload
 {
-	uint Data[14];
+	// float HitT                          //  4 bytes
+	float3 TranslatedWorldPos;             // 12 bytes
+	uint2 PackedBaseColor;                 //  8 bytes half4
+	uint2 PackedMetallicSpecularRoughness; //  8 bytes - half4
+	uint2 PackedWorldNormal;               //  8 bytes - half4
+	uint2 PackedEmissive;                  //  8 bytes - half3 (alpha is not used)
+	                                       // 48 bytes total
+
+	float4 GetBaseColor() {
+		return float4(UnpackFloat2FromUInt(PackedBaseColor.x),
+		              UnpackFloat2FromUInt(PackedBaseColor.y));
+	}
 
-	void SetInputTranslatedWorldPosition(float3 In)
-	{
-		Data[0] = asuint(In.x);
-		Data[1] = asuint(In.y);
-		Data[2] = asuint(In.z);
+	void SetBaseColor(float4 BaseColor) {
+		PackedBaseColor.x = PackFloat2ToUInt(BaseColor.rg);
+		PackedBaseColor.y = PackFloat2ToUInt(BaseColor.ba);
 	}
-	float3 GetInputTranslatedWorldPosition() { return float3(asfloat(Data[0]), asfloat(Data[1]), asfloat(Data[2])); }
 
-	void SetMetallicSpecularRoughness(float3 In)
-	{
-		Data[3] = f32tof16(In.r);
-		Data[3] |= f32tof16(In.g) << 16;
-		Data[4] = f32tof16(In.b);
+	float4 GetMetallicSpecularRoughness() {
+		return float4(UnpackFloat2FromUInt(PackedMetallicSpecularRoughness.x),
+		              UnpackFloat2FromUInt(PackedMetallicSpecularRoughness.y));
 	}
-	float3 GetMetallicSpecularRoughness() { return float3(f16tof32(Data[3]), f16tof32(Data[3] >> 16), f16tof32(Data[4])); }
 
-	void SetBaseColor(float3 In)
-	{
-		Data[5] = f32tof16(In.r);
-		Data[5] |= f32tof16(In.g) << 16;
-		Data[6] = f32tof16(In.b);
+	void SetMetallicSpecularRoughness(float4 MetallicSpecularRoughness) {
+		PackedMetallicSpecularRoughness.x = PackFloat2ToUInt(MetallicSpecularRoughness.rg);
+		PackedMetallicSpecularRoughness.y = PackFloat2ToUInt(MetallicSpecularRoughness.ba);
 	}
-	float3 GetBaseColor() { return float3(f16tof32(Data[5]), f16tof32(Data[5] >> 16), f16tof32(Data[6])); }
 
-	void SetWorldNormal(float3 In)
-	{
-		Data[7] = asuint(In.r);
-		Data[8] = asuint(In.g);
-		Data[9] = asuint(In.b);
+	float4 GetWorldNormal() {
+		return float4(UnpackFloat2FromUInt(PackedWorldNormal.x),
+		              UnpackFloat2FromUInt(PackedWorldNormal.y));
 	}
-	float3 GetWorldNormal() { return float3(asfloat(Data[7]), asfloat(Data[8]), asfloat(Data[9])); }
 
-	void SetEmissive(float3 In)
-	{
-		Data[10] = f32tof16(In.r);
-		Data[10] |= f32tof16(In.g) << 16;
-		Data[11] = f32tof16(In.b);
+	void SetWorldNormal(float4 WorldNormal) {
+		PackedWorldNormal.x = PackFloat2ToUInt(WorldNormal.xy);
+		PackedWorldNormal.y = PackFloat2ToUInt(WorldNormal.zw);
 	}
-	float3 GetEmissive() { return float3(f16tof32(Data[10]), f16tof32(Data[10] >> 16), f16tof32(Data[11])); }
 
-	void SetTransparency(float In)
-	{
-		Data[12] = asuint(In);
+	float3 GetEmissive() {
+		return float3(UnpackFloat2FromUInt(PackedEmissive.x),
+		              f16tof32(PackedEmissive.y));
 	}
-	float GetTransparency() { return asfloat(Data[12]); }
 
-	void SetFlags(uint In) { Data[13] = In; }
-	uint GetFlags() { return Data[13]; }
+	void SetEmissive(float3 Emissive) {
+		PackedEmissive.x = PackFloat2ToUInt(Emissive.xy);
+		PackedEmissive.y = f32tof16(Emissive.z);
+	}
 };
 CHECK_RT_PAYLOAD_SIZE(FDecalShaderPayload)
+
+
+FDecalShaderPayload InitDecalShaderPayload(float3 TranslatedWorldPos)
+{
+	FDecalShaderPayload Payload = (FDecalShaderPayload)0;
+	Payload.TranslatedWorldPos = TranslatedWorldPos;
+	// set alpha=1 on all fields
+	Payload.PackedBaseColor.y = f32tof16(1.0) << 16;
+	Payload.PackedMetallicSpecularRoughness.y = f32tof16(1.0) << 16;
+	Payload.PackedWorldNormal.y = f32tof16(1.0) << 16;
+	return Payload;
+}
+
 #endif
 
 
diff --git a/Engine/Shaders/Private/RayTracing/RayTracingDecalMaterialShader.usf b/Engine/Shaders/Private/RayTracing/RayTracingDecalMaterialShader.usf
--- a/Engine/Shaders/Private/RayTracing/RayTracingDecalMaterialShader.usf
+++ b/Engine/Shaders/Private/RayTracing/RayTracingDecalMaterialShader.usf
@@ -32,28 +32,36 @@ struct FDecalData
 	float3 WorldNormal;
 	float3 MetallicSpecularRoughness;
 	float3 Emissive;
-
-	void WriteToPayload(inout FDecalShaderPayload Payload, uint WriteFlags, float OpacityMult)
+	
+	void CombineWithPayload(inout FDecalShaderPayload Payload, uint WriteFlags, float OpacityMult)
 	{
 		Opacity *= OpacityMult;
+		float Transparency = 1.0 - Opacity;
 		if (WriteFlags & DECAL_WRITE_BASE_COLOR_FLAG)
 		{
-			Payload.SetBaseColor(BaseColor * Opacity);
+			float4 DestBaseColor = Payload.GetBaseColor();
+			DestBaseColor.rgb = BaseColor * Opacity + DestBaseColor.rgb * Transparency;
+			DestBaseColor.a *= Transparency;
+			Payload.SetBaseColor(DestBaseColor);
 		}
 		if (WriteFlags & DECAL_WRITE_NORMAL_FLAG)
 		{
-			Payload.SetWorldNormal(WorldNormal * Opacity);
+			float4 DestWorldNormal = Payload.GetWorldNormal();
+			DestWorldNormal.xyz = WorldNormal * Opacity + DestWorldNormal.xyz * Transparency;
+			DestWorldNormal.w *= Transparency;
+			Payload.SetWorldNormal(DestWorldNormal);
 		}
 		if (WriteFlags & DECAL_WRITE_ROUGHNESS_SPECULAR_METALLIC_FLAG)
 		{
-			Payload.SetMetallicSpecularRoughness(MetallicSpecularRoughness * Opacity);
+			float4 DestMetallicSpecularRoughness = Payload.GetMetallicSpecularRoughness();
+			DestMetallicSpecularRoughness.rgb = MetallicSpecularRoughness * Opacity + DestMetallicSpecularRoughness.rgb * Transparency;
+			DestMetallicSpecularRoughness.a *= Transparency;
+			Payload.SetMetallicSpecularRoughness(DestMetallicSpecularRoughness);
 		}
 		if (WriteFlags & DECAL_WRITE_EMISSIVE_FLAG)
 		{
-			Payload.SetEmissive(Emissive * Opacity);
+			Payload.SetEmissive(Payload.GetEmissive() + Emissive * Opacity);
 		}
-		Payload.SetTransparency(1.0f - Opacity);
-		Payload.SetFlags(WriteFlags);
 	}
 };
 
@@ -97,14 +105,13 @@ RAY_TRACING_ENTRY_CALLABLE(RayTracingDecalMaterialShader, FDecalShaderPayload, P
 {
 	ResolvedView = ResolveView();
 
-	const float3 TranslatedWorldPosition = Payload.GetInputTranslatedWorldPosition();
+	const float3 TranslatedWorldPosition = Payload.TranslatedWorldPos;
 
 	float4 DecalPos = mul(float4(TranslatedWorldPosition.xyz, 1), WorldToDecal);
 
 	if (any(abs(DecalPos.xyz) > 1))
 	{
-		// clip content outside the decal
-		Payload.SetTransparency(1.0f);
+		// out of bounds, nothing to do
 		return;
 	}
 
@@ -157,7 +164,7 @@ RAY_TRACING_ENTRY_CALLABLE(RayTracingDecalMaterialShader, FDecalShaderPayload, P
 	}
 #endif
 	
-	Data.WriteToPayload(Payload, DecalWriteFlags, DecalFading);
+	Data.CombineWithPayload(Payload, DecalWriteFlags, DecalFading);
 
 }
 #endif // RAYCALLABLESHADER
@@ -231,7 +238,11 @@ RAY_TRACING_ENTRY_CLOSEST_HIT(RayTracingDecalMaterialCHS,
 	}
 #endif
 
-	Data.WriteToPayload(Payload, DECAL_PAYLOAD_FLAGS, 1.0);
+#if DECAL_USE_REVERSE_CULLING
+	Data.WorldNormal = -Data.WorldNormal; // flip normal (see matching logic in the PathTracer)
+#endif
+
+	Data.CombineWithPayload(Payload, DECAL_PAYLOAD_FLAGS, 1.0);
 }
 
 #if USE_MATERIAL_ANY_HIT_SHADER
diff --git a/Engine/Shaders/Private/RayTracing/RayTracingDefaultDecalHitShader.usf b/Engine/Shaders/Private/RayTracing/RayTracingDefaultDecalHitShader.usf
--- a/Engine/Shaders/Private/RayTracing/RayTracingDefaultDecalHitShader.usf
+++ b/Engine/Shaders/Private/RayTracing/RayTracingDefaultDecalHitShader.usf
@@ -7,16 +7,15 @@ RAY_TRACING_ENTRY_CLOSEST_HIT(RayTracingDefaultOpaqueDecalMaterialCHS,
 	FRayTracingIntersectionAttributes, Attributes)
 {
 	Payload.HitT = RayTCurrent();
-
-	Payload.SetTransparency(0.0f);
-	Payload.SetBaseColor(float3(0, 0, 0));
-	Payload.SetFlags(DECAL_WRITE_BASE_COLOR_FLAG);
+	// Equivalent to a pure black BaseColor only decal
+	Payload.SetBaseColor(float4(0, 0, 0, 0));
 }
 
 RAY_TRACING_ENTRY_CLOSEST_HIT(RayTracingDefaultHiddenDecalMaterialCHS,
 	FDecalShaderPayload, Payload,
 	FRayTracingIntersectionAttributes, Attributes)
 {
+	// nothing to do here since AHS will ignore the hit
 }
 
 RAY_TRACING_ENTRY_ANY_HIT(RayTracingDefaultHiddenDecalMaterialAHS,
diff --git a/Engine/Source/Runtime/Renderer/Private/RayTracing/RayTracingDecals.cpp b/Engine/Source/Runtime/Renderer/Private/RayTracing/RayTracingDecals.cpp
--- a/Engine/Source/Runtime/Renderer/Private/RayTracing/RayTracingDecals.cpp
+++ b/Engine/Source/Runtime/Renderer/Private/RayTracing/RayTracingDecals.cpp
@@ -226,7 +226,7 @@ private:
 	LAYOUT_FIELD(FShaderUniformBufferParameter, DecalParameter);
 };
 
-IMPLEMENT_RT_PAYLOAD_TYPE(ERayTracingPayloadType::Decals, 60);
+IMPLEMENT_RT_PAYLOAD_TYPE(ERayTracingPayloadType::Decals, 48);
 
 IMPLEMENT_SHADER_TYPE(, FRayTracingDecalMaterialShader, TEXT("/Engine/Private/RayTracing/RayTracingDecalMaterialShader.usf"), TEXT("RayTracingDecalMaterialShader"), SF_RayCallable);
 
