Holdouts: Fix holdouts on subsurface materials which use Opacity channel for shading instead of for translucency

diff --git a/Engine/Shaders/Private/BasePassPixelShader.usf b/Engine/Shaders/Private/BasePassPixelShader.usf
--- a/Engine/Shaders/Private/BasePassPixelShader.usf
+++ b/Engine/Shaders/Private/BasePassPixelShader.usf
@@ -2141,8 +2141,13 @@ void FPixelShaderInOut_MainPS(
 	{
 		// Discard this pixel so that nothing gets written to the G-Buffers and alpha is kept as its default value (background)
 
+#if MATERIALBLENDING_SOLID
+		// This catches the common case as well as handles cases where Opacity has a different meaning, such as for Subsurface materials
+		Out.MRT[0] = half4(0.0, 0.0, 0.0, 1.0);
+#else
 		// TODO: The translucent case doesn't seem to get the right alpha (compared to the path traced reference)
 		Out.MRT[0] = half4(0.0, 0.0, 0.0, Opacity);
+#endif
 
 #if USES_GBUFFER
 		Out.MRT[1] = 0.0;
