Optimises register pressure of TSR's Min/Max3x3 convolutions using a convolution sink

diff --git a/Engine/Shaders/Private/TemporalSuperResolution/TSRConvolutionNetwork.ush b/Engine/Shaders/Private/TemporalSuperResolution/TSRConvolutionNetwork.ush
--- a/Engine/Shaders/Private/TemporalSuperResolution/TSRConvolutionNetwork.ush
+++ b/Engine/Shaders/Private/TemporalSuperResolution/TSRConvolutionNetwork.ush
@@ -262,57 +262,80 @@ bool AnyThread(bool X)
 #endif
 
 
-//------------------------------------------------------- ACCESS NEIGHBOR
+//------------------------------------------------------- ACCESS NEIGHBOR USING WAVE OPS
+
+CALL_SITE_DEBUGLOC
+template<uint LaneStrideX, uint LaneStrideY>
+uint GetNeighborElementIndex(const uint ElementIndex, const tsr_short2 Offset)
+{
+	const uint2 ElementPixelCoord = uint2(GetSimdIndexPixelCoordinateInLane<LaneStrideX, LaneStrideY>(ElementIndex));
+
+	const uint2 NeigborElementPixelCoord = (ElementPixelCoord + uint2(Offset)) % uint2(LaneStrideX, LaneStrideY);
+	const uint NeigborElementIndex = dot(NeigborElementPixelCoord, uint2(1, LaneStrideX));
+
+	return NeigborElementIndex;
+}
 
 #if WAVE_COUNT_X == 1 || WAVE_COUNT_Y == 1
 
+CALL_SITE_DEBUGLOC
+template<uint LaneStrideX, uint LaneStrideY>
+FWaveBroadcastSettings GetWaveBroadcastSettingsForNeighbor(const uint ElementIndex, const tsr_short2 Offset)
+{
+	const uint2 ElementPixelCoord = uint2(GetSimdIndexPixelCoordinateInLane<LaneStrideX, LaneStrideY>(ElementIndex));
+
+	const bool2 bNeedsLaneRotation = (ElementPixelCoord + uint2(Offset)) >= uint2(LaneStrideX, LaneStrideY);
+	const int LaneRotation = dot(select(bNeedsLaneRotation, int2(Offset), 0), int2(1, LANE_COUNT_X));
+
+	return InitWaveRotateLaneGroup(/* LaneGroupSize = */ LANE_COUNT, LaneRotation);
+}
+
+CALL_SITE_DEBUGLOC
+template<typename ScalarType, uint VectorSize, uint LaneStrideX, uint LaneStrideY>
+vector<ScalarType, VectorSize> WaveAccessNeighborElement(
+	TLaneVector2D<ScalarType, VectorSize, LaneStrideX, LaneStrideY> Center,
+	const uint ElementIndex,
+	const tsr_short2 Offset)
+{
+	const uint NeigborElementIndex = GetNeighborElementIndex<LaneStrideX, LaneStrideY>(ElementIndex, Offset);
+	const FWaveBroadcastSettings BroadcastSettings = GetWaveBroadcastSettingsForNeighbor<LaneStrideX, LaneStrideY>(ElementIndex, Offset);
+		
+	// Access the lement.
+	vector<ScalarType, VectorSize> ReturnSimdElement;
+	if (BroadcastSettings.Rotate != 0)
+	{
+		ReturnSimdElement = Center.Registers.WaveBroadcastElement(BroadcastSettings, NeigborElementIndex);
+	}
+	else
+	{
+		ReturnSimdElement = Center.GetElement(NeigborElementIndex);
+	}
+	return ReturnSimdElement;
+}
+
 CALL_SITE_DEBUGLOC
 template<typename ScalarType, uint VectorSize, uint LaneStrideX, uint LaneStrideY>
 TLaneVector2D<ScalarType, VectorSize, LaneStrideX, LaneStrideY> WaveAccessNeighborTexel(
-	TLaneVector2D<ScalarType, VectorSize, LaneStrideX, LaneStrideY> Input,
+	TLaneVector2D<ScalarType, VectorSize, LaneStrideX, LaneStrideY> Center,
 	const tsr_short2 Offset)
 {
 	TLaneVector2D<ScalarType, VectorSize, LaneStrideX, LaneStrideY> Return;
 	
 	UNROLL
-	for (uint SimdIndex = 0; SimdIndex < LaneStrideX * LaneStrideY; SimdIndex++)
+	for (uint ElementIndex = 0; ElementIndex < LaneStrideX * LaneStrideY; ElementIndex++)
 	{
-		const uint2 SimdPixelCoord = uint2(GetSimdIndexPixelCoordinateInLane<LaneStrideX, LaneStrideY>(SimdIndex));
-
-		// Compute the simd index of the neighbor.
-		const uint2 NeigborSimdPixelCoord = (SimdPixelCoord + uint2(Offset)) % uint2(LaneStrideX, LaneStrideY);
-		const uint NeigborSimdIndex = dot(NeigborSimdPixelCoord, uint2(1, LaneStrideX));
-
-		// Compute the lane rotaton of the neighbor.
-		const bool2 bNeedsLaneRotation = (SimdPixelCoord + uint2(Offset)) >= uint2(LaneStrideX, LaneStrideY);
-		const int LaneRotation = dot(select(bNeedsLaneRotation, int2(Offset), 0), int2(1, LANE_COUNT_X));
-		
-		// Access the lement.
-		vector<ScalarType, VectorSize> ReturnSimdElement;
-		if (LaneRotation != 0)
-		{
-			// Only uses wave intrinsics on platforms that provides RDNA's ds_swizzle because we know the wave size of this platforms. 
-			#if !PLATFORM_SUPPORTS_WAVE_ROTATE
-				//#error InitWaveRotateLaneGroup() isn't supported.
-			#endif
-			const FWaveBroadcastSettings BroadcastSettings = InitWaveRotateLaneGroup(/* LaneGroupSize = */ LANE_COUNT, LaneRotation);
-			ReturnSimdElement = Input.Registers.WaveBroadcastElement(BroadcastSettings, NeigborSimdIndex);
-		}
-		else
-		{
-			ReturnSimdElement = Input.GetElement(NeigborSimdIndex);
-		}
-
-		Return.SetElement(SimdIndex, ReturnSimdElement);
+		Return.SetElement(ElementIndex, WaveAccessNeighborElement(Center, ElementIndex, Offset));
 	}
 
 	Return.TightenRegisters();
-
 	return Return;
 }
 
 #endif // WAVE_COUNT_X == 1 || WAVE_COUNT_Y == 1
 
+
+//------------------------------------------------------- 3x1 NEIGHBOR ACCESS
+
 CALL_SITE_DEBUGLOC
 template<typename ScalarType, uint VectorSize, uint LaneStrideX, uint LaneStrideY>
 void AccessNeighborTexels3x1(
@@ -362,37 +385,24 @@ void AccessNeighborTexels3x1(
 #endif
 
 
+//------------------------------------------------------- 1x3 NEIGHBOR ACCESS
+
 CALL_SITE_DEBUGLOC
 ISOLATE
 template<typename ScalarType, uint VectorSize, uint LaneStrideX, uint LaneStrideY>
-void AccessNeighborTexels1x3(
-	TLaneVector2D<ScalarType, VectorSize, LaneStrideX, LaneStrideY> CO,
-	out TLaneVector2D<ScalarType, VectorSize, LaneStrideX, LaneStrideY> CP,
-	out TLaneVector2D<ScalarType, VectorSize, LaneStrideX, LaneStrideY> CN)
+void Write1x3CenterToLDS(TLaneVector2D<ScalarType, VectorSize, LaneStrideX, LaneStrideY> CO)
 #if WAVE_COUNT_Y > 1
 {
 	const uint TotalNumberOfLaneInGroup = LANE_COUNT_X * WAVE_COUNT;
 
+	GroupMemoryBarrierWithGroupSync();
+
 	if (LaneStrideY == 1)
 	{
-		GroupMemoryBarrierWithGroupSync();
-		UNROLL
-		for (uint SimdIndex = 0; SimdIndex < LaneStrideX; SimdIndex++)
-		{
-			WriteVectorToLDS<ScalarType, VectorSize>(GGroupThreadIndex + SimdIndex * TotalNumberOfLaneInGroup, CO.GetElement(SimdIndex));
-		}
-		GroupMemoryBarrierWithGroupSync();
-		
 		UNROLL
-		for (uint SimdIndex = 0; SimdIndex < LaneStrideX; SimdIndex++)
+		for (uint ElementIndex = 0; ElementIndex < LaneStrideX; ElementIndex++)
 		{
-			vector<ScalarType, VectorSize> FirstElement;
-			vector<ScalarType, VectorSize> LastElement;
-			ReadVectorFromLDS<ScalarType, VectorSize>(((GGroupThreadIndex - LANE_COUNT_X * WAVE_COUNT_X) % TotalNumberOfLaneInGroup) + SimdIndex * TotalNumberOfLaneInGroup, /* out */ FirstElement);
-			ReadVectorFromLDS<ScalarType, VectorSize>(((GGroupThreadIndex + LANE_COUNT_X * WAVE_COUNT_X) % TotalNumberOfLaneInGroup) + SimdIndex * TotalNumberOfLaneInGroup, /* out */ LastElement);
-	
-			CP.SetElement(SimdIndex, LastElement);
-			CN.SetElement(SimdIndex, FirstElement);
+			WriteVectorToLDS<ScalarType, VectorSize>(GGroupThreadIndex + ElementIndex * TotalNumberOfLaneInGroup, CO.GetElement(ElementIndex));
 		}
 	}
 	else
@@ -400,45 +410,176 @@ void AccessNeighborTexels1x3(
 		const uint FirstElementIndex = 0;
 		const uint LastElementIndex = LaneStrideX * LaneStrideY - LaneStrideX;
 
-		GroupMemoryBarrierWithGroupSync();
 		UNROLL
-		for (uint SimdIndex = 0; SimdIndex < LaneStrideX; SimdIndex++)
+		for (uint ElementIndex = 0; ElementIndex < LaneStrideX; ElementIndex++)
 		{
-			WriteVectorToLDS<ScalarType, VectorSize>(GGroupThreadIndex + (0 + SimdIndex) * TotalNumberOfLaneInGroup, CO.GetElement(FirstElementIndex + SimdIndex));
-			WriteVectorToLDS<ScalarType, VectorSize>(GGroupThreadIndex + (LaneStrideX + SimdIndex) * TotalNumberOfLaneInGroup, CO.GetElement(LastElementIndex + SimdIndex));
+			WriteVectorToLDS<ScalarType, VectorSize>(GGroupThreadIndex + (0 + ElementIndex) * TotalNumberOfLaneInGroup, CO.GetElement(FirstElementIndex + ElementIndex));
+			WriteVectorToLDS<ScalarType, VectorSize>(GGroupThreadIndex + (LaneStrideX + ElementIndex) * TotalNumberOfLaneInGroup, CO.GetElement(LastElementIndex + ElementIndex));
 		}
-		GroupMemoryBarrierWithGroupSync();
-	
-		UNROLL
-		for (uint SimdIndex = 0; SimdIndex < LastElementIndex; SimdIndex++)
+	}
+
+	GroupMemoryBarrierWithGroupSync();
+}
+#else
+{
+	// NOP
+}
+#endif
+
+CALL_SITE_DEBUGLOC
+ISOLATE
+template<typename ScalarType, uint VectorSize, uint LaneStrideX, uint LaneStrideY>
+void Read1x3NeighborElementsFromLDS(
+	TLaneVector2D<ScalarType, VectorSize, LaneStrideX, LaneStrideY> CO,
+	const uint ElementIndex,
+	out vector<ScalarType, VectorSize> CP,
+	out vector<ScalarType, VectorSize> CN)
+#if WAVE_COUNT_Y > 1
+{
+	const uint TotalNumberOfLaneInGroup = LANE_COUNT_X * WAVE_COUNT;
+
+	if (LaneStrideY == 1)
+	{
+		ReadVectorFromLDS<ScalarType, VectorSize>(((GGroupThreadIndex - LANE_COUNT_X * WAVE_COUNT_X) % TotalNumberOfLaneInGroup) + ElementIndex * TotalNumberOfLaneInGroup, /* out */ CN);
+		ReadVectorFromLDS<ScalarType, VectorSize>(((GGroupThreadIndex + LANE_COUNT_X * WAVE_COUNT_X) % TotalNumberOfLaneInGroup) + ElementIndex * TotalNumberOfLaneInGroup, /* out */ CP);
+	}
+	else // if (LaneStrideY > 1)
+	{
+		if (ElementIndex < LaneStrideX)
 		{
-			CP.SetElement(SimdIndex + 0, CO.GetElement(SimdIndex + LaneStrideX));
-			CN.SetElement(SimdIndex + LaneStrideX, CO.GetElement(SimdIndex + 0));
+			ReadVectorFromLDS<ScalarType, VectorSize>(((GGroupThreadIndex - LANE_COUNT_X * WAVE_COUNT_X) % TotalNumberOfLaneInGroup) + (LaneStrideX + ElementIndex) * TotalNumberOfLaneInGroup, /* out */ CN);
 		}
-	
-		UNROLL
-		for (uint SimdIndex = 0; SimdIndex < LaneStrideX; SimdIndex++)
+		else
 		{
-			vector<ScalarType, VectorSize> FirstElement;
-			vector<ScalarType, VectorSize> LastElement;
-			ReadVectorFromLDS<ScalarType, VectorSize>(((GGroupThreadIndex - LANE_COUNT_X * WAVE_COUNT_X) % TotalNumberOfLaneInGroup) + (LaneStrideX + SimdIndex) * TotalNumberOfLaneInGroup, /* out */ FirstElement);
-			ReadVectorFromLDS<ScalarType, VectorSize>(((GGroupThreadIndex + LANE_COUNT_X * WAVE_COUNT_X) % TotalNumberOfLaneInGroup) + (0 + SimdIndex) * TotalNumberOfLaneInGroup, /* out */ LastElement);
-	
-			CP.SetElement(LastElementIndex + SimdIndex, LastElement);
-			CN.SetElement(FirstElementIndex + SimdIndex, FirstElement);
+			CN = CO.GetElement(ElementIndex - LaneStrideX);
+		}
+		
+		const uint LastElementIndex = LaneStrideX * LaneStrideY - LaneStrideX;
+		if (ElementIndex >= LastElementIndex)
+		{
+			ReadVectorFromLDS<ScalarType, VectorSize>(((GGroupThreadIndex + LANE_COUNT_X * WAVE_COUNT_X) % TotalNumberOfLaneInGroup) + (0 + ElementIndex - LastElementIndex) * TotalNumberOfLaneInGroup, /* out */ CP);
+		}
+		else
+		{
+			CP = CO.GetElement(ElementIndex + LaneStrideX);
 		}
 	}
+}
+#else
+{
+	CP = WaveAccessNeighborElement(CO, ElementIndex, tsr_short2(0, 1));
+	CN = WaveAccessNeighborElement(CO, ElementIndex, tsr_short2(0, -1));
+}
+#endif
+
+CALL_SITE_DEBUGLOC
+ISOLATE
+template<typename ScalarType, uint VectorSize, uint LaneStrideX, uint LaneStrideY>
+void Read1x3CenterFromLDS(
+	TLaneVector2D<ScalarType, VectorSize, LaneStrideX, LaneStrideY> CO,
+	out TLaneVector2D<ScalarType, VectorSize, LaneStrideX, LaneStrideY> CP,
+	out TLaneVector2D<ScalarType, VectorSize, LaneStrideX, LaneStrideY> CN)
+{
+	UNROLL
+	for (uint ElementIndex = 0; ElementIndex < LaneStrideX * LaneStrideY; ElementIndex++)
+	{
+		vector<ScalarType, VectorSize> CPElement;
+		vector<ScalarType, VectorSize> CNElement;
+		Read1x3NeighborElementsFromLDS(CO, ElementIndex, /* out */ CPElement, /* out */ CNElement);
+
+		CP.SetElement(ElementIndex, CPElement);
+		CN.SetElement(ElementIndex, CNElement);
+	}
 	
 	CP.TightenRegisters();
 	CN.TightenRegisters();
 }
+
+CALL_SITE_DEBUGLOC
+ISOLATE
+template<typename ScalarType, uint VectorSize, uint LaneStrideX, uint LaneStrideY>
+void AccessNeighborTexels1x3(
+	TLaneVector2D<ScalarType, VectorSize, LaneStrideX, LaneStrideY> CO,
+	out TLaneVector2D<ScalarType, VectorSize, LaneStrideX, LaneStrideY> CP,
+	out TLaneVector2D<ScalarType, VectorSize, LaneStrideX, LaneStrideY> CN)
+{
+	Write1x3CenterToLDS(CO);
+	Read1x3CenterFromLDS(CO, /* out */ CP, /* out */ CN);
+}
+
+
+//------------------------------------------------------- CONVOLUTION WITH SINK FOR LOWER REGISTER PRESSURE
+
+CALL_SITE_DEBUGLOC
+template<typename ScalarType, uint VectorSize, uint LaneStrideX, uint LaneStrideY, typename ConvolutionSinkType>
+TLaneVector2D<ScalarType, VectorSize, LaneStrideX, LaneStrideY> Convolve1x3Then3x1FromLDS(TLaneVector2D<ScalarType, VectorSize, LaneStrideX, LaneStrideY> C, ConvolutionSinkType ConvolutionSink)
+#if WAVE_COUNT_X == 1 && 0
+{
+	// Only if LaneStrideX==1 and WAVE_COUNT_X==1
+	TLaneVector2D<ScalarType, VectorSize, LaneStrideX, LaneStrideY> R;
+	
+	UNROLL
+	for (uint ElementIndex = 0; ElementIndex < LaneStrideX * LaneStrideY; ElementIndex++)
+	{
+		vector<ScalarType, VectorSize> CPElement;
+		vector<ScalarType, VectorSize> CNElement;
+		Read1x3NeighborElementsFromLDS(C, ElementIndex, /* out */ CPElement, /* out */ CNElement);
+
+		vector<ScalarType, VectorSize> HCElement = ConvolutionSink.ConvolveElements1x3(CNElement, C.GetElement(ElementIndex), CPElement);
+
+		const FWaveBroadcastSettings HPSettings = GetWaveBroadcastSettingsForNeighbor<LaneStrideX, LaneStrideY>(ElementIndex, tsr_short2(1, 0));
+		const FWaveBroadcastSettings HNSettings = GetWaveBroadcastSettingsForNeighbor<LaneStrideX, LaneStrideY>(ElementIndex, tsr_short2(-1, 0));
+		
+		vector<ScalarType, VectorSize> HPElement = WaveBroadcast(HPSettings, HCElement);
+		vector<ScalarType, VectorSize> HNElement = WaveBroadcast(HNSettings, HCElement);
+
+		R.SetElement(ElementIndex, ConvolutionSink.ConvolveElements3x1(HNElement, HCElement, HPElement));
+	}
+	
+	return R;
+}
 #else
 {
-	CP = WaveAccessNeighborTexel(CO, tsr_short2(0, 1));
-	CN = WaveAccessNeighborTexel(CO, tsr_short2(0, -1));
+	TLaneVector2D<ScalarType, VectorSize, LaneStrideX, LaneStrideY> VO;
+	{
+		TLaneVector2D<ScalarType, VectorSize, LaneStrideX, LaneStrideY> CO = C;
+		TLaneVector2D<ScalarType, VectorSize, LaneStrideX, LaneStrideY> CP;
+		TLaneVector2D<ScalarType, VectorSize, LaneStrideX, LaneStrideY> CN;
+		Read1x3CenterFromLDS(CO, /* out */ CP, /* out */ CN);
+
+		UNROLL
+		for (uint RegisterRowIndex = 0; RegisterRowIndex < TLaneVector2D<ScalarType, VectorSize, LaneStrideX, LaneStrideY>::RegisterRowCount; RegisterRowIndex++)
+		{
+			VO.Registers.SetRegisterRow(RegisterRowIndex, ConvolutionSink.ConvolveRegisterRows1x3(CN.Registers.GetRegisterRow(RegisterRowIndex), CO.Registers.GetRegisterRow(RegisterRowIndex), CP.Registers.GetRegisterRow(RegisterRowIndex)));
+		}
+	}
+	
+	TLaneVector2D<ScalarType, VectorSize, LaneStrideX, LaneStrideY> R;
+	{
+		TLaneVector2D<ScalarType, VectorSize, LaneStrideX, LaneStrideY> CO = VO;
+		TLaneVector2D<ScalarType, VectorSize, LaneStrideX, LaneStrideY> CP;
+		TLaneVector2D<ScalarType, VectorSize, LaneStrideX, LaneStrideY> CN;
+		AccessNeighborTexels3x1(CO, /* out */ CP, /* out */ CN);
+
+		UNROLL
+		for (uint RegisterRowIndex = 0; RegisterRowIndex < TLaneVector2D<ScalarType, VectorSize, LaneStrideX, LaneStrideY>::RegisterRowCount; RegisterRowIndex++)
+		{
+			R.Registers.SetRegisterRow(RegisterRowIndex, ConvolutionSink.ConvolveRegisterRows3x1(CN.Registers.GetRegisterRow(RegisterRowIndex), CO.Registers.GetRegisterRow(RegisterRowIndex), CP.Registers.GetRegisterRow(RegisterRowIndex)));
+		}
+	}
+
+	return R;
 }
 #endif
 
+CALL_SITE_DEBUGLOC
+template<typename ScalarType, uint VectorSize, uint LaneStrideX, uint LaneStrideY, typename ConvolutionSinkType>
+TLaneVector2D<ScalarType, VectorSize, LaneStrideX, LaneStrideY> Convolve1x3Then3x1(TLaneVector2D<ScalarType, VectorSize, LaneStrideX, LaneStrideY> Center, ConvolutionSinkType ConvolutionSink)
+{
+	Write1x3CenterToLDS(Center);
+	return Convolve1x3Then3x1FromLDS(Center, ConvolutionSink);
+}
+
 
 //------------------------------------------------------- BLUR 3x3 CONVOLUTION
 
@@ -524,9 +665,71 @@ TLaneVector2D<ScalarType, VectorSize, LaneStrideX, LaneStrideY> TotalVariation3x
 
 //------------------------------------------------------- MIN & MAX 3x3 CONVOLUTION
 
+template<typename ScalarType, uint VectorSize, uint InSimdSizeX, uint InSimdSizeY>
+struct TConvolutionSinkMin
+{
+	static const uint RegisterRowSize = TLaneVector2D<ScalarType, VectorSize, InSimdSizeX, InSimdSizeY>::RegisterRowSize;
+
+	uint Unused;
+
+	vector<ScalarType, VectorSize> ConvolveElements3x1(vector<ScalarType, VectorSize> N, vector<ScalarType, VectorSize> C, vector<ScalarType, VectorSize> P)
+	{
+		return min3(N, C, P);
+	}
+
+	vector<ScalarType, VectorSize> ConvolveElements1x3(vector<ScalarType, VectorSize> N, vector<ScalarType, VectorSize> C, vector<ScalarType, VectorSize> P)
+	{
+		return min3(N, C, P);
+	}
+	
+	vector<ScalarType, RegisterRowSize> ConvolveRegisterRows3x1(vector<ScalarType, RegisterRowSize> N, vector<ScalarType, RegisterRowSize> C, vector<ScalarType, RegisterRowSize> P)
+	{
+		return min3(N, C, P);
+	}
+
+	vector<ScalarType, RegisterRowSize> ConvolveRegisterRows1x3(vector<ScalarType, RegisterRowSize> N, vector<ScalarType, RegisterRowSize> C, vector<ScalarType, RegisterRowSize> P)
+	{
+		return min3(N, C, P);
+	}
+};
+
+template<typename ScalarType, uint VectorSize, uint InSimdSizeX, uint InSimdSizeY>
+struct TConvolutionSinkMax
+{
+	static const uint RegisterRowSize = TLaneVector2D<ScalarType, VectorSize, InSimdSizeX, InSimdSizeY>::RegisterRowSize;
+
+	uint Unused;
+
+	vector<ScalarType, VectorSize> ConvolveElements3x1(vector<ScalarType, VectorSize> N, vector<ScalarType, VectorSize> C, vector<ScalarType, VectorSize> P)
+	{
+		return max3(N, C, P);
+	}
+
+	vector<ScalarType, VectorSize> ConvolveElements1x3(vector<ScalarType, VectorSize> N, vector<ScalarType, VectorSize> C, vector<ScalarType, VectorSize> P)
+	{
+		return max3(N, C, P);
+	}
+	
+	vector<ScalarType, RegisterRowSize> ConvolveRegisterRows3x1(vector<ScalarType, RegisterRowSize> N, vector<ScalarType, RegisterRowSize> C, vector<ScalarType, RegisterRowSize> P)
+	{
+		return max3(N, C, P);
+	}
+
+	vector<ScalarType, RegisterRowSize> ConvolveRegisterRows1x3(vector<ScalarType, RegisterRowSize> N, vector<ScalarType, RegisterRowSize> C, vector<ScalarType, RegisterRowSize> P)
+	{
+		return max3(N, C, P);
+	}
+};
+
 CALL_SITE_DEBUGLOC
 template<typename ScalarType, uint VectorSize, uint LaneStrideX, uint LaneStrideY>
 TLaneVector2D<ScalarType, VectorSize, LaneStrideX, LaneStrideY> Min3x3(TLaneVector2D<ScalarType, VectorSize, LaneStrideX, LaneStrideY> Center)
+#if 1
+{
+	TConvolutionSinkMin<ScalarType, VectorSize, LaneStrideX, LaneStrideY> MinSink = {0};
+	return Convolve1x3Then3x1(Center, MinSink);
+}
+#else
 {
 	TLaneVector2D<ScalarType, VectorSize, LaneStrideX, LaneStrideY> CO = Center;
 	TLaneVector2D<ScalarType, VectorSize, LaneStrideX, LaneStrideY> CP;
@@ -540,10 +743,17 @@ TLaneVector2D<ScalarType, VectorSize, LaneStrideX, LaneStrideY> Min3x3(TLaneVect
 
 	return min3(MinN, MinP, MinO);
 }
+#endif
 
 CALL_SITE_DEBUGLOC
 template<typename ScalarType, uint VectorSize, uint LaneStrideX, uint LaneStrideY>
 TLaneVector2D<ScalarType, VectorSize, LaneStrideX, LaneStrideY> Max3x3(TLaneVector2D<ScalarType, VectorSize, LaneStrideX, LaneStrideY> Center)
+#if 1
+{
+	TConvolutionSinkMax<ScalarType, VectorSize, LaneStrideX, LaneStrideY> MaxSink = {0};
+	return Convolve1x3Then3x1(Center, MaxSink);
+}
+#else
 {
 	TLaneVector2D<ScalarType, VectorSize, LaneStrideX, LaneStrideY> CO = Center;
 	TLaneVector2D<ScalarType, VectorSize, LaneStrideX, LaneStrideY> CP;
@@ -557,6 +767,7 @@ TLaneVector2D<ScalarType, VectorSize, LaneStrideX, LaneStrideY> Max3x3(TLaneVect
 
 	return max3(MaxN, MaxP, MaxO);
 }
+#endif
 
 CALL_SITE_DEBUGLOC
 template<typename ScalarType, uint LaneStrideX, uint LaneStrideY>
@@ -608,7 +819,16 @@ void MinMax3x3(
 	TLaneVector2D<ScalarType, VectorSize, LaneStrideX, LaneStrideY> Center,
 	out TLaneVector2D<ScalarType, VectorSize, LaneStrideX, LaneStrideY> OutMin,
 	out TLaneVector2D<ScalarType, VectorSize, LaneStrideX, LaneStrideY> OutMax)
-#if 1
+#if WAVE_COUNT_X == 1 && 1
+{
+	TConvolutionSinkMin<ScalarType, VectorSize, LaneStrideX, LaneStrideY> MinSink = {0};
+	TConvolutionSinkMax<ScalarType, VectorSize, LaneStrideX, LaneStrideY> MaxSink = {0};
+	
+	Write1x3CenterToLDS(Center);
+	OutMin = Convolve1x3Then3x1FromLDS(Center, MinSink);
+	OutMax = Convolve1x3Then3x1FromLDS(Center, MaxSink);
+}
+#elif 1
 {
 	TLaneVector2D<ScalarType, VectorSize, LaneStrideX, LaneStrideY> CO = Center;
 	TLaneVector2D<ScalarType, VectorSize, LaneStrideX, LaneStrideY> CP;
