Fixed skt dome mesh to output alpha = transmittance = 0.

diff --git a/Engine/Shaders/Private/BasePassPixelShader.usf b/Engine/Shaders/Private/BasePassPixelShader.usf
--- a/Engine/Shaders/Private/BasePassPixelShader.usf
+++ b/Engine/Shaders/Private/BasePassPixelShader.usf
@@ -2046,14 +2046,6 @@ void FPixelShaderInOut_MainPS(
 	#else
 		Out.MRT[0].rgba *= ViewPreExposure;
 	#endif
-
-	#if MATERIAL_IS_SKY && POST_PROCESS_ALPHA
-		// In this case, alpha will represent transmitance toward the background.
-		// But the sky material is the background. And we have no way torecover transmittance to outer space for now so we simply write transmittance=1.
-		// This matches the path tracer behavior. Note: when the SkyAtmosphere is standalone (no sky dome mesh material), we can write grey scale trnasmitance to outer space.
-		Out.MRT[0].a = 1.0f;
-	#endif
-
 #endif
 
 	// If support OIT, then remove blending contribution and insert OIT sample
