Fixed alpha holdout on lcoud totally removing fog and atmosphere contribution.

diff --git a/Engine/Shaders/Private/VolumetricCloud.usf b/Engine/Shaders/Private/VolumetricCloud.usf
--- a/Engine/Shaders/Private/VolumetricCloud.usf
+++ b/Engine/Shaders/Private/VolumetricCloud.usf
@@ -1408,7 +1408,15 @@ void MainCommon(in FMaterialPixelParameters MaterialParameters, in float4 SvPosi
 	// When rendering a real time reflection capture (sky envmap) whe use a different exposure
 	const float OutputPreExposure = (ResolvedView.RealTimeReflectionCapture ? ResolvedView.RealTimeReflectionCapturePreExposure : ResolvedView.PreExposure);
 
-
+#if POST_PROCESS_ALPHA
+	if (RenderVolumetricCloudParameters.bHoldout2 && View.RenderingReflectionCaptureMask == 0.0f)
+	{
+		Luminance *= 0;
+	#if CLOUD_MIN_AND_MAX_DEPTH
+		ClosestLuminance *= 0;
+	#endif
+	}
+#endif
 
 	//
 	// Apply aerial perspective if needed
@@ -1607,16 +1615,6 @@ void MainCommon(in FMaterialPixelParameters MaterialParameters, in float4 SvPosi
 		OutDepth.x = ((GrayScaleTransmittance > 0.99) ? NoCloudDepth : tAP) * CENTIMETER_TO_KILOMETER; // using a small threshold on transmittance
 	}
 
-#if POST_PROCESS_ALPHA
-	if (RenderVolumetricCloudParameters.bHoldout2 && View.RenderingReflectionCaptureMask == 0.0f)
-	{
-		OutColor0.rgb *= 0;
-	#if CLOUD_MIN_AND_MAX_DEPTH
-		OutColor1.rgb *= 0;
-	#endif
-	}
-#endif
-
 	//OutColor0.g = 0.3f; // Debug revealing reveal pixels that have been traced.
 }
 
