Improves TSR's temporal flickering analysis by also including spatial-anti-flickering

This allows to have it ghosts less where there spatial anti-flickering do good enough job.

diff --git a/Engine/Shaders/Private/TemporalSuperResolution/TSRRejectShading.usf b/Engine/Shaders/Private/TemporalSuperResolution/TSRRejectShading.usf
--- a/Engine/Shaders/Private/TemporalSuperResolution/TSRRejectShading.usf
+++ b/Engine/Shaders/Private/TemporalSuperResolution/TSRRejectShading.usf
@@ -478,11 +478,12 @@ void ComputeMoireError(
 		tsr_simd_half BackbufferQuantizationErrorVector = tsr_simd_half::Const(MoireEncodingError);
 		
 		tsr_simd_half InputC0 = InputLuma;
+		tsr_simd_half HistoryC0 = PrevFlickeringHistory;
 
 		// Anhilate history
 		tsr_simd_half InputC2 = InputC0;
-		tsr_simd_half HistoryC2 = PrevFlickeringHistory;
-		AnnihilateMutuallySingleChannel3x3(InputLuma, PrevFlickeringHistory, /* out */ InputC2, /* out */ HistoryC2);
+		tsr_simd_half HistoryC2 = HistoryC0;
+		AnnihilateMutuallySingleChannel3x3(InputLuma, HistoryC0, /* out */ InputC2, /* out */ HistoryC2);
 
 		// Convolve the input luma to have the exact same moire pattern as FilteredInput below.
 		tsr_simd_half FilteredInput;
@@ -492,33 +493,49 @@ void ComputeMoireError(
 		// Compute clamping box of the limance.
 		tsr_simd_half FilteredBoxMin;
 		tsr_simd_half FilteredBoxMax;
-		tsr_simd_half BoxSize;
+		tsr_simd_half InputC2BoxSize;
 		{
 			tsr_simd_half2 Min;
 			tsr_simd_half2 Max;
 			MinMax3x3(Concatenate(FilteredInput, InputC2), /* out */ Min, /* out */ Max);
-			FilteredBoxMin = Min[0] - BackbufferQuantizationErrorVector;
-			FilteredBoxMax = Max[0] + BackbufferQuantizationErrorVector;
+			FilteredBoxMin = Min[0];
+			FilteredBoxMax = Max[0];
 		
-			BoxSize = (
-				(Max[1] - Min[1]) * tsr_simd_half::Const(FilteringWeight) +
-				(BackbufferQuantizationErrorVector * tsr_half(2 * FilteringWeight)));
+			InputC2BoxSize = Max[1] - Min[1];
 		}
 
+		// Compute the clamp error
+		tsr_simd_half ClampError;
+		{
+			tsr_simd_half TotalVarInputDiffC0C2;
+			tsr_simd_half TotalVarInputC2;
+			Deconcatenate(abs(TotalVariation3x3(Concatenate(abs(InputC0 - InputC2), InputC2))), /* out */ TotalVarInputDiffC0C2, /* out */ TotalVarInputC2);
+		
+			ClampError = BackbufferQuantizationErrorVector;
+			ClampError = max(ClampError, Blur3x3(min(TotalVarInputDiffC0C2, TotalVarInputC2)));
+			ClampError = max(ClampError, InputC2BoxSize * tsr_half(FilteringWeight * 0.25));
+			ClampError = ClampError + BackbufferQuantizationErrorVector;
+		}
+
+		// Apply clamp error
+		FilteredBoxMin = FilteredBoxMin - ClampError;
+		FilteredBoxMax = FilteredBoxMax + ClampError;
+
+		tsr_simd_half BoxSize = (
+			InputC2BoxSize * tsr_simd_half::Const(FilteringWeight) +
+			(BackbufferQuantizationErrorVector * tsr_half(2 * FilteringWeight)));
+
 		tsr_simd_half ClampedFilteredHistory = fastClamp(FilteredHistory, FilteredBoxMin, FilteredBoxMax);
 
 		tsr_simd_half Delta = max(abs(FilteredInput - FilteredHistory), BoxSize);
 	
-		tsr_simd_half ClampError = BackbufferQuantizationErrorVector;
-		ClampError = max(ClampError, BoxSize * tsr_half(0.25)); // original
-		
 		tsr_simd_half RawClampedEnergy = abs(ClampedFilteredHistory - FilteredHistory);
-		tsr_simd_half RawRejection     = saturate(tsr_simd_half::Const(1.0) - max(RawClampedEnergy - ClampError, tsr_simd_half::Const(0.0)) * rcp(Delta));
+		tsr_simd_half RawRejection     = saturate(tsr_simd_half::Const(1.0) - max(RawClampedEnergy, tsr_simd_half::Const(0.0)) * rcp(Delta));
 
 		tsr_simd_half FilteredClampedEnergy;
 		Deconcatenate(MaxRMinG3x3(Median3x3(Concatenate(RawClampedEnergy, RawRejection))), /* out */ FilteredClampedEnergy, /* out */ OutRejectionClampBlend);
 
-		OutRejectionBlendFinal = saturate(tsr_simd_half::Const(1.0) - max(FilteredClampedEnergy - ClampError, tsr_simd_half::Const(0.0)) * rcp(Delta));
+		OutRejectionBlendFinal = saturate(tsr_simd_half::Const(1.0) - max(FilteredClampedEnergy, tsr_simd_half::Const(0.0)) * rcp(Delta));
 	}
 	
 	// Simulate what happens in TSRUpdateHistory.usf
@@ -578,7 +595,7 @@ void ComputeMoireError(
 	// of these need a nearest filter which isn't precise at all.
 	tsr_simd_half TotalVariationContribution;
 	tsr_simd_half TotalVariationCountContribution;
-	Deconcatenate(Max3x3(Concatenate(GradientVariation, IsFlicker)), /* out */ TotalVariationContribution, /* out */ TotalVariationCountContribution);
+	Deconcatenate(Max3x3(Max3x3(Concatenate(GradientVariation, IsFlicker))), /* out */ TotalVariationContribution, /* out */ TotalVariationCountContribution);
 
 	UNROLL_N(SIMD_SIZE)
 	for (uint SimdIndex = 0; SimdIndex < SIMD_SIZE; SimdIndex++)
