Fix missing VSM shadow when camera is too close due to too short rays.

diff --git a/Engine/Shaders/Private/VirtualShadowMaps/VirtualShadowMapProjection.usf b/Engine/Shaders/Private/VirtualShadowMaps/VirtualShadowMapProjection.usf
--- a/Engine/Shaders/Private/VirtualShadowMaps/VirtualShadowMapProjection.usf
+++ b/Engine/Shaders/Private/VirtualShadowMaps/VirtualShadowMapProjection.usf
@@ -374,7 +374,7 @@ FVirtualShadowMapSampleResult ProjectLight(
 					Light,
 					PixelPos,
 					TranslatedWorldPosition,
-					SMRTRayOffset,		// TODO: Is this actually what we want here?
+					SMRTRayLengthScale * 100,
 					RayIndex,
 					SMRTRayCount,
 					RayStart,
@@ -425,8 +425,7 @@ FVirtualShadowMapSampleResult ProjectLight(
 				const uint VoxelDescCount = VirtualVoxel.NodeDescCount;
 				for (uint VoxelDescIt=0; VoxelDescIt<VoxelDescCount; ++VoxelDescIt)
 				{
-					const FPackedVirtualVoxelNodeDesc PackedNode = VirtualVoxel.NodeDescBuffer[VoxelDescIt];
-					const FVirtualVoxelNodeDesc NodeDesc = UnpackVoxelNode(PackedNode, VirtualVoxel.PageResolution);
+					const FVirtualVoxelNodeDesc NodeDesc = UnpackVoxelNode(VirtualVoxel.NodeDescBuffer[VoxelDescIt], VirtualVoxel.PageResolution);
 
 					FHairTraversalResult HairResult = InitHairTraversalResult();
 					HairResult = ComputeHairCountVirtualVoxel(
