Stochastic Shadows - spatial filter driven by variance

diff --git a/Engine/Shaders/Private/StochasticShadows/StochasticShadows.ush b/Engine/Shaders/Private/StochasticShadows/StochasticShadows.ush
--- a/Engine/Shaders/Private/StochasticShadows/StochasticShadows.ush
+++ b/Engine/Shaders/Private/StochasticShadows/StochasticShadows.ush
@@ -276,8 +276,8 @@ void AccumulateHistoryShadowMask(uint PrevLocalLightIndex, uint2 ShadowMaskScree
 
 			if (HistoryShadowMask.IsValid())
 			{
-				Accumulator.FactorSum += Weight * HistoryShadowMask.Factor;
-				Accumulator.FactorSqSum += Pow2(HistoryShadowMask.Factor) * Weight;
+				Accumulator.FactorSum += HistoryShadowMask.Factor * Weight;
+				Accumulator.FactorSqSum += HistoryShadowMask.FactorSq * Weight;
 				Accumulator.WeightSum += Weight;
 				Accumulator.MinAge = min(Accumulator.MinAge, HistoryShadowMask.Age);
 				Accumulator.MaxHistoryLength = max(Accumulator.MaxHistoryLength, HistoryShadowMask.HistoryLength);
diff --git a/Engine/Shaders/Private/StochasticShadows/StochasticShadowsComposite.usf b/Engine/Shaders/Private/StochasticShadows/StochasticShadowsComposite.usf
--- a/Engine/Shaders/Private/StochasticShadows/StochasticShadowsComposite.usf
+++ b/Engine/Shaders/Private/StochasticShadows/StochasticShadowsComposite.usf
@@ -180,10 +180,12 @@ void CompositeShadowMaskTracesCS(
 		if (SampleAccumulator.SampleWeightSum > 0 && CenterSampleAccumulator.SampleWeightSum == 0.0f)
 		{
 			ShadowMask.Factor = SampleAccumulator.SampleSum / SampleAccumulator.SampleWeightSum;
-			ShadowMask.FactorSq = SampleAccumulator.SampleSum / SampleAccumulator.SampleWeightSum;
 			ShadowMask.Age = 0;
 		}
 
+		// Tracking only temporal variance
+		ShadowMask.FactorSq = Pow2(ShadowMask.Factor);
+
 		#if TEMPORAL_ACCUMULATION
 		{
 			int2 ScreenCoord00 = ShadowMaskHistoryScreenCoord00[ShadowMaskScreenCoord];
diff --git a/Engine/Shaders/Private/StochasticShadows/StochasticShadowsSpatial.usf b/Engine/Shaders/Private/StochasticShadows/StochasticShadowsSpatial.usf
--- a/Engine/Shaders/Private/StochasticShadows/StochasticShadowsSpatial.usf
+++ b/Engine/Shaders/Private/StochasticShadows/StochasticShadowsSpatial.usf
@@ -130,7 +130,6 @@ void ShadowMaskSpatialPassCS(
 
 	uint2 ShadowMaskScreenCoord = BaseShadowMaskScreenCoord + GroupThreadId.xy;
 
-#define DEBUG_MODE 0
 #if DEBUG_MODE
 	int2 DebugScreenCoord = View.CursorPosition.x >= 0 ? View.CursorPosition * View.ViewResolutionFraction : -1;
 	bool bDebug = all(ShadowMaskScreenCoord == DebugScreenCoord / SHADOW_MASK_DOWNSAMPLE_FACTOR);
@@ -149,6 +148,30 @@ void ShadowMaskSpatialPassCS(
 		FShadowMaskAccumulator ShadowMaskAccumulator = InitShadowMaskAccumulator();
 
 		uint2 CenterSampleSharedTileCoord = LocalShadowMaskTileCoord + SHARED_TILE_BORDER;
+		FShadowMask CenterShadowMask = UnpackShadowMask(SharedShadowMask[CenterSampleSharedTileCoord.x][CenterSampleSharedTileCoord.y]);
+		
+		float CenterSampleStdDev = 0.0f;
+
+		if (CenterShadowMask.IsValid())
+		{
+			ShadowMaskAccumulator.FactorSum += CenterShadowMask.Factor;
+			ShadowMaskAccumulator.FactorSqSum += Pow2(CenterShadowMask.Factor);
+			ShadowMaskAccumulator.MinAge = CenterShadowMask.Age;
+			ShadowMaskAccumulator.MaxHistoryLength = CenterShadowMask.HistoryLength;
+			ShadowMaskAccumulator.WeightSum += 1.0f;
+
+			CenterSampleStdDev = sqrt(max(CenterShadowMask.FactorSq - Pow2(CenterShadowMask.Factor), 0.0f));
+
+			#if DEBUG_MODE
+			if (bDebug)
+			{
+				Newline(Context);
+				Print(Context, TEXT("StdDev: "));
+				Print(Context, CenterSampleStdDev);
+			}
+			#endif
+		}
+
 		const float2 ScreenUV = (DownsampledScreenCoordToScreenCoord(ShadowMaskScreenCoord) + 0.5f) * View.BufferSizeAndInvSize.zw;
 		float CenterSampleDepth = SharedSceneDepth[CenterSampleSharedTileCoord.x][CenterSampleSharedTileCoord.y];
 		float3 CenterSampleNormal = normalize(DecodeNormal(DownsampledSceneWorldNormal[ShadowMaskScreenCoord]));
@@ -156,7 +179,7 @@ void ShadowMaskSpatialPassCS(
 		float4 ScenePlane = float4(CenterSampleNormal, dot(TranslatedWorldPosition, CenterSampleNormal));
 
 		// #ss_todo: expose as CVars
-		float KernelRadiusInPixels = 8.0f;
+		float KernelRadiusInPixels = 8.0f * saturate(4.0f * CenterSampleStdDev);
 		uint NumSamples = 8;
 
 		float GuassianNormalize = 2.0f / (KernelRadiusInPixels * KernelRadiusInPixels);
@@ -209,51 +232,6 @@ void ShadowMaskSpatialPassCS(
 			}
 		}
 
-/*
-		const int KernelSize = 2;
-		const float KernelWeights[3] = { 1.0f, 1.0f / 2.0f, 1.0f / 3.0f };
-		for (int SampleOffsetY = -KernelSize; SampleOffsetY <= KernelSize; ++SampleOffsetY)
-		{
-			for (int SampleOffsetX = -KernelSize; SampleOffsetX <= KernelSize; ++SampleOffsetX)
-			{
-				uint2 SharedShadowMaskTileCoord = (GroupThreadId.xy + SHARED_TILE_BORDER) - int2(SampleOffsetX, SampleOffsetY);
-				FShadowMask SampleShadowMask = UnpackShadowMask(SharedShadowMask[SharedShadowMaskTileCoord.x][SharedShadowMaskTileCoord.y]);
-
-				if (SampleShadowMask.IsValid())
-				{
-					float SampleWeight = KernelWeights[abs(SampleOffsetX)] * KernelWeights[abs(SampleOffsetY)];
-
-					uint2 SampleShadowMaskScreenCoord = clamp((int2)BaseShadowMaskScreenCoord + (int2)SharedShadowMaskTileCoord - SHARED_TILE_BORDER, 0, ShadowMaskViewSize - 1);
-					float SampleDepth = SharedSceneDepth[SharedShadowMaskTileCoord.x][SharedShadowMaskTileCoord.y];
-					float3 SampleTranslatedWorldPosition = GetTranslatedWorldPositionFromScreenUV(DownsampledScreenCoordToScreenUV(SampleShadowMaskScreenCoord), SampleDepth);
-
-					float BilateralFilterDepthWeightScale = 10000.0f;
-					float PlaneDistance = abs(dot(float4(SampleTranslatedWorldPosition, -1), ScenePlane));
-					float RelativeDepthDifference = PlaneDistance / SampleDepth;
-					SampleWeight *= exp2(-BilateralFilterDepthWeightScale * (RelativeDepthDifference * RelativeDepthDifference));
-
-					ShadowMaskAccumulator.FactorSum += SampleShadowMask.Factor * SampleWeight;
-					ShadowMaskAccumulator.FactorSqSum += Pow2(SampleShadowMask.Factor) * SampleWeight;
-					ShadowMaskAccumulator.WeightSum += SampleWeight;
-
-					if (SampleWeight > 0.0f)
-					{
-						ShadowMaskAccumulator.MinAge = min(ShadowMaskAccumulator.MinAge, SampleShadowMask.Age);
-					}
-
-					#if DEBUG_MODE
-					if (bDebug)
-					{
-						Newline(Context);
-						Print(Context, TEXT("SampleShadowMask: "));
-						Print(Context, SampleShadowMask.Factor);
-						Print(Context, SampleDepth);
-					}
-					#endif
-				}
-			}
-		}*/
-
 		FShadowMask ShadowMask = InitShadowMask();
 		if (ShadowMaskAccumulator.WeightSum > 0.0f)
 		{
@@ -273,21 +251,18 @@ void ShadowMaskSpatialPassCS(
 		}
 		#endif
 
-		GroupMemoryBarrierWithGroupSync();
-
 		// Update shared shadow mask tile with filtered results
-		SharedShadowMask[GroupThreadId.x + SHARED_TILE_BORDER][GroupThreadId.y + SHARED_TILE_BORDER] = PackShadowMask(ShadowMask, /*Dither*/ 0.5f);
-
-		SetSharedShadingTileMask(LocalShadowMaskTileCoord / 4, ShadowMask);
+		SharedShadowMask[CenterSampleSharedTileCoord.x][CenterSampleSharedTileCoord.y] = PackShadowMask(ShadowMask, /*Dither*/ 0.5f);
 	}
-#else
+
+	GroupMemoryBarrierWithGroupSync();
+#endif
+	
 	// Only mark sub-tiles requiring shading tiles
 	{
 		FShadowMask ShadowMask = UnpackShadowMask(SharedShadowMask[GroupThreadId.x + SHARED_TILE_BORDER][GroupThreadId.y + SHARED_TILE_BORDER]);
-
 		SetSharedShadingTileMask(LocalShadowMaskTileCoord / 4, ShadowMask);
 	}
-#endif
 
 	GroupMemoryBarrierWithGroupSync();
 
diff --git a/Engine/Source/Runtime/Renderer/Private/StochasticShadows/StochasticShadows.cpp b/Engine/Source/Runtime/Renderer/Private/StochasticShadows/StochasticShadows.cpp
--- a/Engine/Source/Runtime/Renderer/Private/StochasticShadows/StochasticShadows.cpp
+++ b/Engine/Source/Runtime/Renderer/Private/StochasticShadows/StochasticShadows.cpp
@@ -543,7 +543,8 @@ class FShadowMaskSpatialPassCS : public FGlobalShader
 	}
 
 	class FSpatialPass : SHADER_PERMUTATION_BOOL("SPATIAL_PASS");
-	using FPermutationDomain = TShaderPermutationDomain<FSpatialPass>;
+	class FDebugMode : SHADER_PERMUTATION_BOOL("DEBUG_MODE");
+	using FPermutationDomain = TShaderPermutationDomain<FSpatialPass, FDebugMode>;
 
 	static bool ShouldCompilePermutation(const FGlobalShaderPermutationParameters& Parameters)
 	{
@@ -1083,6 +1084,7 @@ void FDeferredShadingSceneRenderer::RenderStochasticShadows(FRDGBuilder& GraphBu
 
 		FShadowMaskSpatialPassCS::FPermutationDomain PermutationVector;
 		PermutationVector.Set<FShadowMaskSpatialPassCS::FSpatialPass>(CVarStochasticShadowsSpatial.GetValueOnRenderThread() != 0);
+		PermutationVector.Set<FShadowMaskSpatialPassCS::FDebugMode>(bDebug);
 		auto ComputeShader = View.ShaderMap->GetShader<FShadowMaskSpatialPassCS>(PermutationVector);
 
 		FComputeShaderUtils::AddPass(
