Fixes TSR's Flickering Temporal Analysis behind AfterDOF translucencies

diff --git a/Engine/Shaders/Private/TemporalSuperResolution/TSRRejectShading.usf b/Engine/Shaders/Private/TemporalSuperResolution/TSRRejectShading.usf
--- a/Engine/Shaders/Private/TemporalSuperResolution/TSRRejectShading.usf
+++ b/Engine/Shaders/Private/TemporalSuperResolution/TSRRejectShading.usf
@@ -427,6 +427,7 @@ void ComputeMoireError(
 	tsr_simd_half  bIsStatic,
 	tsr_simd_half  MoireInput,
 	tsr_simd_halfC OriginalOpaqueInput,
+	tsr_simd_half4 OriginalTranslucencyInput,
 	tsr_simd_half4 PrevMoireHistory,
 	out tsr_simd_half4 OutMoireHistory,
 	out tsr_simd_half  OutMoireError)
@@ -446,17 +447,24 @@ void ComputeMoireError(
 	{
 		const tsr_half InputLumaEncodingError = tsr_half(0.5 / 255.0);
 			
-		OriginalOpaqueInput = OriginalOpaqueInput * rcp(OriginalOpaqueInput + tsr_simd_halfC::Const(tsr_half(PerceptionAdd)));
+		tsr_simd_halfC OriginalInput = ComposeTranslucency(OriginalOpaqueInput, OriginalTranslucencyInput);
 
+		OriginalInput       = AddPerceptionAdd(OriginalInput, PerceptionAdd);
+		OriginalOpaqueInput = AddPerceptionAdd(OriginalOpaqueInput, PerceptionAdd);
 
 		tsr_simd_half SceneColorLuma;
+		tsr_simd_half TranslucencyChange;
 		UNROLL_N(SIMD_SIZE)
 		for (uint SimdIndex = 0; SimdIndex < SIMD_SIZE; SimdIndex++)
 		{
 			SceneColorLuma.SetElement(SimdIndex, dot(OriginalOpaqueInput.GetElement(SimdIndex).rgb, kMoireLumaWeights));
+			TranslucencyChange.SetElement(SimdIndex, dot(abs(OriginalOpaqueInput.GetElement(SimdIndex).rgb - OriginalInput.GetElement(SimdIndex).rgb), kMoireLumaWeights));
 		}
 
-		InputLumaModification = max(abs(InputLuma - SceneColorLuma) - tsr_simd_half::Const(InputLumaEncodingError), tsr_simd_half::Const(0.0));
+		InputLumaModification = tsr_simd_half::Const(0.0);
+		InputLumaModification = max(InputLumaModification, abs(InputLuma - SceneColorLuma) - tsr_simd_half::Const(InputLumaEncodingError));
+		InputLumaModification = max(InputLumaModification, TranslucencyChange);
+		InputLumaModification = max(InputLumaModification, tsr_simd_half::Const(1.0) - OriginalTranslucencyInput[3]);
 	}
 
 	// Unpack history
@@ -656,7 +664,7 @@ void ComputeMoireError(
 		
 	// Moving object in the scene have very low chance of creating moire pattern with TAA jitter, so discard the moire error to still
 	// have character animation ghost free evem as environement shadow project different things on character.
-	OutMoireError = OutMoireError * bIsStatic;
+	OutMoireError = saturate(OutMoireError * bIsStatic - InputLumaModification);
 } // ComputeMoireError()
 
 
@@ -979,6 +987,7 @@ void MainCS(
 			bIsStatic,
 			MoireInput,
 			OriginalOpaqueInput,
+			OriginalTranslucencyInput,
 			PrevMoireHistory,
 			/* out */ MoireHistory,
 			/* out */ MoireError);
