Displays where TSR's flickering temporal analysis is disabled due to scene color not matching flickering luminance

diff --git a/Engine/Shaders/Private/TemporalSuperResolution/TSRVisualize.usf b/Engine/Shaders/Private/TemporalSuperResolution/TSRVisualize.usf
--- a/Engine/Shaders/Private/TemporalSuperResolution/TSRVisualize.usf
+++ b/Engine/Shaders/Private/TemporalSuperResolution/TSRVisualize.usf
@@ -34,7 +34,11 @@ uint bCanSpatialAntiAlias;
 float MaxHistorySampleCount;
 float OutputToHistoryResolutionFractionSquare;
 float FlickeringFramePeriod;
+float PerceptionAdd;
 
+Texture2D<tsr_half3> InputTexture;
+Texture2D<tsr_half>  InputMoireLumaTexture;
+Texture2D<tsr_half4> InputSceneTranslucencyTexture;
 Texture2D<tsr_half3> SceneColorTexture;
 Texture2D<float2>    ClosestDepthTexture;
 Texture2D<float2>    DilatedVelocityTexture;
@@ -206,8 +210,11 @@ tsr_half3 DisplayAntiFlickering(FContext Ctx)
 	// Compute the threshold at which the variation count should quick in
 	const tsr_half TotalVariationCountThreshold = tsr_half(1.0 / (1.0 - pow(1.0 - BlendFinal, tsr_half(FlickeringFramePeriod))));
 
+	// Texture fetches
+	tsr_half3 OriginalOpaqueInput = InputTexture.SampleLevel(GlobalPointClampedSampler, Ctx.InputUV, 0);
+	tsr_half  MoireInput = InputMoireLumaTexture.SampleLevel(GlobalPointClampedSampler, Ctx.InputUV, 0);
+	tsr_half4 OriginalTranslucencyInput = InputSceneTranslucencyTexture.SampleLevel(GlobalPointClampedSampler, Ctx.InputUV, 0);
 	tsr_half4 MoireHistory = MoireHistoryTexture.SampleLevel(GlobalPointClampedSampler, Ctx.InputUV, 0);
-
 	tsr_half DisableAntiFlickering = tsr_half(1.0) - tsr_half(IsMovingMaskTexture[Ctx.InputPixelPos]) * rcp(tsr_half(255.0));
 
 	// Unpack history
@@ -216,6 +223,26 @@ tsr_half3 DisplayAntiFlickering(FContext Ctx)
 	tsr_half TotalVariation      = MoireHistory[2] * (tsr_half(1.0) - BlendFinal);
 	tsr_half TotalVariationCount = MoireHistory[3] * MaxPrevTotalVariationCount * (tsr_half(1.0) - BlendFinal);
 	
+	tsr_half InputLumaModification;
+	{
+		const tsr_half InputLumaEncodingError = tsr_half(0.5 / 255.0);
+		
+		tsr_half3 OriginalInput = OriginalOpaqueInput * OriginalTranslucencyInput.a + OriginalTranslucencyInput.rgb;
+		
+		OriginalInput = OriginalInput * rcp(OriginalInput + tsr_half(PerceptionAdd));
+		OriginalOpaqueInput = OriginalOpaqueInput * rcp(OriginalOpaqueInput + tsr_half(PerceptionAdd));
+
+		tsr_half SceneColorLuma = dot(OriginalOpaqueInput.rgb, kMoireLumaWeights);
+		tsr_half TranslucencyChange = dot(abs(OriginalOpaqueInput.rgb - OriginalInput.rgb), kMoireLumaWeights);
+		
+		InputLumaModification = tsr_half(0.0);
+		InputLumaModification = max(InputLumaModification, abs(MoireInput - SceneColorLuma) - InputLumaEncodingError);
+		InputLumaModification = max(InputLumaModification, TranslucencyChange);
+		InputLumaModification = max(InputLumaModification, tsr_half(1.0) - OriginalTranslucencyInput[3]);
+	}
+
+	tsr_half InputLumaModificationPercent = saturate(InputLumaModification * tsr_half(16.0));
+
 	// Compute the final luminance 
 	#if 1
 		tsr_half LocalCountFadeIn = saturate(TotalVariationCount / TotalVariationCountThreshold - tsr_half(0.5));
@@ -228,7 +255,7 @@ tsr_half3 DisplayAntiFlickering(FContext Ctx)
 	Color = lerp(kGreen, kYellow, saturate(LocalMoireError * 16.0));
 	Color = lerp(Color,  kRed,    saturate(LocalMoireError * 16.0 - 1.0));
 	Color = lerp(kWhite, Color,   LocalCountFadeIn);
-	Color = lerp(kPink, Color,   DisableAntiFlickering);
+	Color = lerp(kPink,  Color,   min(DisableAntiFlickering, tsr_half(1.0) - InputLumaModificationPercent));
 	Color = select(bCameraCut == 0, Color, kPink);
 
 	return SampleTSROutputGray(Ctx) * Color;
diff --git a/Engine/Source/Runtime/Renderer/Private/PostProcess/TemporalSuperResolution.cpp b/Engine/Source/Runtime/Renderer/Private/PostProcess/TemporalSuperResolution.cpp
--- a/Engine/Source/Runtime/Renderer/Private/PostProcess/TemporalSuperResolution.cpp
+++ b/Engine/Source/Runtime/Renderer/Private/PostProcess/TemporalSuperResolution.cpp
@@ -1037,7 +1037,11 @@ class FTSRVisualizeCS : public FTSRShader
 		SHADER_PARAMETER(float, MaxHistorySampleCount)
 		SHADER_PARAMETER(float, OutputToHistoryResolutionFractionSquare)
 		SHADER_PARAMETER(float, FlickeringFramePeriod)
+		SHADER_PARAMETER(float, PerceptionAdd)
 
+		SHADER_PARAMETER_RDG_TEXTURE(Texture2D, InputTexture)
+		SHADER_PARAMETER_RDG_TEXTURE(Texture2D, InputMoireLumaTexture)
+		SHADER_PARAMETER_RDG_TEXTURE(Texture2D, InputSceneTranslucencyTexture)
 		SHADER_PARAMETER_RDG_TEXTURE_SRV(Texture2D, SceneColorTexture)
 		SHADER_PARAMETER_RDG_TEXTURE(Texture2D, ClosestDepthTexture)
 		SHADER_PARAMETER_RDG_TEXTURE(Texture2D, DilatedVelocityTexture)
@@ -2686,7 +2690,19 @@ FDefaultTemporalUpscaler::FOutputs AddTemporalSuperResolutionPasses(
 			PassParameters->MaxHistorySampleCount = MaxHistorySampleCount;
 			PassParameters->OutputToHistoryResolutionFractionSquare = OutputToHistoryResolutionFractionSquare;
 			PassParameters->FlickeringFramePeriod = FlickeringFramePeriod;
+			PassParameters->PerceptionAdd = FMath::Pow(0.5f, CVarTSRShadingExposureOffset.GetValueOnRenderThread());
 
+			PassParameters->InputTexture = PassInputs.SceneColor.Texture;
+			if (PassInputs.FlickeringInputTexture.IsValid())
+			{
+				ensure(InputRect == PassInputs.FlickeringInputTexture.ViewRect);
+				PassParameters->InputMoireLumaTexture = PassInputs.FlickeringInputTexture.Texture;
+			}
+			else
+			{
+				PassParameters->InputMoireLumaTexture = BlackDummy;
+			}
+			PassParameters->InputSceneTranslucencyTexture = SeparateTranslucencyTexture;
 			PassParameters->SceneColorTexture = SceneColorOutputTextureSRV;
 			PassParameters->ClosestDepthTexture = ClosestDepthTexture;
 			PassParameters->DilatedVelocityTexture = DilatedVelocityTexture;
