Fix hair RBF sample lagging one frame due to incorrect skel. mesh buffer sampling.

diff --git a/Engine/Shaders/Private/HairStrands/HairStrandsMesh.usf b/Engine/Shaders/Private/HairStrands/HairStrandsMesh.usf
--- a/Engine/Shaders/Private/HairStrands/HairStrandsMesh.usf
+++ b/Engine/Shaders/Private/HairStrands/HairStrandsMesh.usf
@@ -67,10 +67,10 @@ void MainCS(uint3 DispatchThreadId : SV_DispatchThreadID)
 	if (VertexIndex >= MaxVertexCount || SectionIndex >= MaxSectionCount) 
 		return;
 	
-	const bool bIsSwapped = MeshSectionBuffer[SectionIndex] > 0;
+	const bool bBuffer0 = MeshSectionBuffer[SectionIndex] == 0;
 
 	float3 SamplePosition = float3(0,0,0);
-	if (bIsSwapped)
+	if (bBuffer0)
 	{
 		SamplePosition.x = MeshPositionBuffer0.Load(VertexIndex * 3 + 0);
 		SamplePosition.y = MeshPositionBuffer0.Load(VertexIndex * 3 + 1);
