Path Tracing: Merge logic for point and spot lights to avoid code duplication

Improve sharing of logic between point and capsule cases as well and make minor simplifications to the code. This appears to have a small performance advantage (1-2%) in scenes with many spot and point lights.

diff --git a/Engine/Shaders/Private/PathTracing/Light/PathTracingCapsuleLight.ush b/Engine/Shaders/Private/PathTracing/Light/PathTracingCapsuleLight.ush
--- a/Engine/Shaders/Private/PathTracing/Light/PathTracingCapsuleLight.ush
+++ b/Engine/Shaders/Private/PathTracing/Light/PathTracingCapsuleLight.ush
@@ -70,7 +70,7 @@ FLightHit CapsuleLight_TraceLight(FRayDesc Ray, int LightId)
 	{
 		float CapsArea = LightRadius2;
 		float BodyArea = 0.5 * LightRadius * SourceLength;
-		float3 LightPower = GetColor(LightId) * ComputeIESAttenuation(LightId, Ray.Origin);
+		float3 LightPower = GetColor(LightId);
 		float3 LightRadiance = LightPower / (PI * (CapsArea + BodyArea));
 		float t = CapsuleIntersect(Ray.Direction, TranslatedLightPosition - Ray.Origin, GetdPdv(LightId), LightRadius2, SourceLength);
 		if (t > Ray.TMin && t < Ray.TMax)
@@ -84,7 +84,7 @@ FLightHit CapsuleLight_TraceLight(FRayDesc Ray, int LightId)
 	{
 		float CapsArea = LightRadius2;
 		float BodyArea = 0.5 * LightRadius * SourceLength;
-		float3 LightPower = GetColor(LightId) * ComputeIESAttenuation(LightId, Ray.Origin);
+		float3 LightPower = GetColor(LightId);
 		float3 LightRadiance = LightPower / (PI * (CapsArea + BodyArea));
 		{
 			// exact
@@ -103,11 +103,10 @@ FLightHit CapsuleLight_TraceLight(FRayDesc Ray, int LightId)
 	{
 		float CapsArea = LightRadius2;
 		float BodyArea = 0.5 * LightRadius * SourceLength;
-		float3 LightPower = GetColor(LightId) * ComputeIESAttenuation(LightId, Ray.Origin);
+		float3 LightPower = GetColor(LightId);
 		float3 LightRadiance = LightPower / (PI * (CapsArea + BodyArea));
 		float3 LightDirection = TranslatedLightPosition - Ray.Origin;
 		float LightDistanceSquared = dot(LightDirection, LightDirection);
-		LightRadiance *= ComputeAttenuationFalloff(LightDistanceSquared, LightId);
 #if CAPSULE_SOLIDANGLE_SAMPLING
 		float Pdf = 1.0 / GetCapsuleBoundsSolidAngle(CapsuleGetSphericalBounds(LightDirection, Axis, LightRadius, SourceLength));
 #else
diff --git a/Engine/Shaders/Private/PathTracing/Light/PathTracingLightSampling.ush b/Engine/Shaders/Private/PathTracing/Light/PathTracingLightSampling.ush
--- a/Engine/Shaders/Private/PathTracing/Light/PathTracingLightSampling.ush
+++ b/Engine/Shaders/Private/PathTracing/Light/PathTracingLightSampling.ush
@@ -17,7 +17,6 @@ LightSampling.usf: Light sampling for next-event estimation
 #include "PathTracingPointLight.ush"
 #include "PathTracingRectLight.ush"
 #include "PathTracingSkyLight.ush"
-#include "PathTracingSpotLight.ush"
 
 // Returns float4 packed with Radiance and Pdf
 FLightHit TraceLight(FRayDesc Ray, int LightId)
@@ -25,10 +24,10 @@ FLightHit TraceLight(FRayDesc Ray, int LightId)
 	switch (SceneLights[LightId].Flags & PATHTRACER_FLAG_TYPE_MASK)
 	{
 		case PATHTRACING_LIGHT_SKY:			return SkyLight_TraceLight(Ray, LightId);
-		case PATHTRACING_LIGHT_POINT:		return PointLight_TraceLight(Ray, LightId);
 		case PATHTRACING_LIGHT_DIRECTIONAL: return DirectionalLight_TraceLight(Ray, LightId);
+		case PATHTRACING_LIGHT_POINT:
+		case PATHTRACING_LIGHT_SPOT:		return PointLight_TraceLight(Ray, LightId);
 		case PATHTRACING_LIGHT_RECT:		return RectLight_TraceLight(Ray, LightId);
-		case PATHTRACING_LIGHT_SPOT:		return SpotLight_TraceLight(Ray, LightId);
 		default:							return NullLightHit(); // unknown light type?
 	}
 }
@@ -44,10 +43,10 @@ FLightSample SampleLight(
 	switch (SceneLights[LightId].Flags & PATHTRACER_FLAG_TYPE_MASK)
 	{
 		case PATHTRACING_LIGHT_SKY:			return SkyLight_SampleLight(LightId, RandSample, WorldNormal);
-		case PATHTRACING_LIGHT_POINT:		return PointLight_SampleLight(LightId, RandSample, TranslatedWorldPos, WorldNormal);
 		case PATHTRACING_LIGHT_DIRECTIONAL: return DirectionalLight_SampleLight(LightId, RandSample, WorldNormal);
+		case PATHTRACING_LIGHT_POINT:
+		case PATHTRACING_LIGHT_SPOT:		return PointLight_SampleLight(LightId, RandSample, TranslatedWorldPos, WorldNormal);
 		case PATHTRACING_LIGHT_RECT:		return RectLight_SampleLight(LightId, RandSample, TranslatedWorldPos, WorldNormal);
-		case PATHTRACING_LIGHT_SPOT:		return SpotLight_SampleLight(LightId, RandSample, TranslatedWorldPos, WorldNormal);
 		default:							return NullLightSample(); // unknown light type?
 	}
 }
@@ -68,10 +67,10 @@ float EstimateLight(
 	switch (SceneLights[LightId].Flags & PATHTRACER_FLAG_TYPE_MASK)
 	{
 		case PATHTRACING_LIGHT_SKY:			return SkyLight_EstimateLight(LightId, WorldNormal, IsTransmissiveMaterial);
-		case PATHTRACING_LIGHT_POINT:		return PointLight_EstimateLight(LightId, TranslatedWorldPos, WorldNormal, IsTransmissiveMaterial);
 		case PATHTRACING_LIGHT_DIRECTIONAL:	return DirectionalLight_EstimateLight(LightId, WorldNormal, IsTransmissiveMaterial);
+		case PATHTRACING_LIGHT_POINT:
+		case PATHTRACING_LIGHT_SPOT:		return PointLight_EstimateLight(LightId, TranslatedWorldPos, WorldNormal, IsTransmissiveMaterial);
 		case PATHTRACING_LIGHT_RECT:		return RectLight_EstimateLight(LightId, TranslatedWorldPos, WorldNormal, IsTransmissiveMaterial);
-		case PATHTRACING_LIGHT_SPOT:		return SpotLight_EstimateLight(LightId, TranslatedWorldPos, WorldNormal, IsTransmissiveMaterial);
 		default:							return 0.0;
 	}
 }
@@ -88,10 +87,10 @@ FVolumeLightSampleSetup PrepareLightVolumeSample(
 	switch (SceneLights[LightId].Flags & PATHTRACER_FLAG_TYPE_MASK)
 	{
 		//case PATHTRACING_LIGHT_SKY:         return SkyLight_PrepareLightVolumeSample(LightId, RayOrigin, RayDirection, TMin, TMax);
-		case PATHTRACING_LIGHT_POINT:       return PointLight_PrepareLightVolumeSample(LightId, RayOrigin, RayDirection, TMin, TMax);
 		//case PATHTRACING_LIGHT_DIRECTIONAL: return DirectionalLight_PrepareLightVolumeSample(LightId, RayOrigin, RayDirection, TMin, TMax);
+		case PATHTRACING_LIGHT_POINT:
+		case PATHTRACING_LIGHT_SPOT:        return PointLight_PrepareLightVolumeSample(LightId, RayOrigin, RayDirection, TMin, TMax);
 		case PATHTRACING_LIGHT_RECT:        return RectLight_PrepareLightVolumeSample(LightId, RayOrigin, RayDirection, TMin, TMax);
-		case PATHTRACING_LIGHT_SPOT:        return SpotLight_PrepareLightVolumeSample(LightId, RayOrigin, RayDirection, TMin, TMax);
 		default:                            return NullVolumeLightSetup();
 	}
 }
diff --git a/Engine/Shaders/Private/PathTracing/Light/PathTracingPointLight.ush b/Engine/Shaders/Private/PathTracing/Light/PathTracingPointLight.ush
--- a/Engine/Shaders/Private/PathTracing/Light/PathTracingPointLight.ush
+++ b/Engine/Shaders/Private/PathTracing/Light/PathTracingPointLight.ush
@@ -10,44 +10,74 @@ PathTracingPointLight.ush: Light sampling functions for Point lights
 #include "PathTracingCapsuleLight.ush"
 
 
+
+// #dxr_todo: These are stolen from DynamicLightingCommon.ush but I was having trouble getting that header to include cleanly
+/** 
+* Calculates attenuation for a spot light.
+* L normalize vector to light. 
+* SpotDirection is the direction of the spot light.
+* SpotAngles.x is CosOuterCone, SpotAngles.y is InvCosConeDifference. 
+*/
+float SpotAttenuation(float3 L, float3 SpotDirection, float2 SpotAngles)
+{
+	float ConeAngleFalloff = Square(saturate((dot(L, -SpotDirection) - SpotAngles.x) * SpotAngles.y));
+	return ConeAngleFalloff;
+}
+
+
 FLightHit PointLight_TraceLight(FRayDesc Ray, int LightId)
 {
+	FLightHit Result = NullLightHit();
 	if (GetSourceLength(LightId) > 0)
 	{
-		return CapsuleLight_TraceLight(Ray, LightId);
+		Result = CapsuleLight_TraceLight(Ray, LightId);
 	}
-
-	float3 TranslatedLightPosition = GetTranslatedPosition(LightId);
-	float LightRadius = GetRadius(LightId);
-	float LightRadius2 = Pow2(LightRadius);
-	float3 oc = Ray.Origin - TranslatedLightPosition;
-	float b = dot(oc, Ray.Direction);
-	// "Precision Improvements for Ray / Sphere Intersection" - Ray Tracing Gems (2019)
-	// https://link.springer.com/content/pdf/10.1007%2F978-1-4842-4427-2_7.pdf
-	// NOTE: we assume the incoming ray direction is normalized
-	float h = LightRadius2 - length2(oc - b * Ray.Direction);
-	if (h > 0)
+	else
 	{
-		float t = -b - sqrt(h);
-		if (t > Ray.TMin && t < Ray.TMax)
+		float3 TranslatedLightPosition = GetTranslatedPosition(LightId);
+		float LightRadius = GetRadius(LightId);
+		float LightRadius2 = Pow2(LightRadius);
+		float3 oc = Ray.Origin - TranslatedLightPosition;
+		float b = dot(oc, Ray.Direction);
+		// "Precision Improvements for Ray / Sphere Intersection" - Ray Tracing Gems (2019)
+		// https://link.springer.com/content/pdf/10.1007%2F978-1-4842-4427-2_7.pdf
+		// NOTE: we assume the incoming ray direction is normalized
+		float h = LightRadius2 - length2(oc - b * Ray.Direction);
+		if (h > 0)
 		{
-			float LightDistanceSquared = dot(oc, oc);
+			float t = -b - sqrt(h);
+			if (t > Ray.TMin && t < Ray.TMax)
+			{
+				float LightDistanceSquared = dot(oc, oc);
 
-			// #dxr_todo: sphere area is 4*pi*r^2 -- but the factor of 4 is missing for some reason?
-			float3 LightPower = GetColor(LightId) * ComputeIESAttenuation(LightId, Ray.Origin);
-			float3 LightRadiance = LightPower / (PI * LightRadius2);
-			LightRadiance *= ComputeAttenuationFalloff(LightDistanceSquared, LightId);
+				// #dxr_todo: sphere area is 4*pi*r^2 -- but the factor of 4 is missing for some reason?
+				float3 LightPower = GetColor(LightId);
+				float3 LightRadiance = LightPower / (PI * LightRadius2);
 
-			float SinThetaMax2 = saturate(LightRadius2 / LightDistanceSquared);
-			float OneMinusCosThetaMax = SinThetaMax2 < 0.01 ? SinThetaMax2 * (0.5 + 0.125 * SinThetaMax2) : 1 - sqrt(1 - SinThetaMax2);
+				float SinThetaMax2 = saturate(LightRadius2 / LightDistanceSquared);
+				float OneMinusCosThetaMax = SinThetaMax2 < 0.01 ? SinThetaMax2 * (0.5 + 0.125 * SinThetaMax2) : 1 - sqrt(1 - SinThetaMax2);
 
-			float SolidAngle = 2 * PI * OneMinusCosThetaMax;
-			return CreateLightHit(LightRadiance, 1.0 / SolidAngle, t);
+				float SolidAngle = 2 * PI * OneMinusCosThetaMax;
+				Result = CreateLightHit(LightRadiance, 1.0 / SolidAngle, t);
+			}
+			// #dxr_todo: process inside hit here ...
+			// How should we define radiance on the inside of the sphere?
+		}
+	}
+	if (Result.IsHit())
+	{
+		float3 LightDirection = GetTranslatedPosition(LightId) - Ray.Origin;
+		float LightDistanceSquared = dot(LightDirection, LightDirection);
+		Result.Radiance *= ComputeIESAttenuation(LightId, Ray.Origin);
+		Result.Radiance *= ComputeAttenuationFalloff(LightDistanceSquared, LightId);
+		if (IsSpotLight(LightId))
+		{
+			float3 LightNormal = GetNormal(LightId);
+			float2 CosConeAngles = GetCosConeAngles(LightId);
+			Result.Radiance *= SpotAttenuation(LightDirection * rsqrt(LightDistanceSquared), LightNormal, CosConeAngles);
 		}
-		// #dxr_todo: process inside hit here ...
-		// How should we define radiance on the inside of the sphere?
 	}
-	return NullLightHit();
+	return Result;
 }
 
 FLightSample PointLight_SampleLight(
@@ -98,6 +128,13 @@ FLightSample PointLight_SampleLight(
 	// NOTE: uses distance to center to keep fadeoff mask consistent for all samples
 	Result.RadianceOverPdf *= ComputeAttenuationFalloff(LightDistanceSquared, LightId);
 	Result.RadianceOverPdf *= ComputeIESAttenuation(LightId, TranslatedWorldPos);
+	if (IsSpotLight(LightId))
+	{
+		float3 LightNormal = GetNormal(LightId);
+		float2 CosConeAngles = GetCosConeAngles(LightId);
+		Result.RadianceOverPdf *= SpotAttenuation(LightDirection * rsqrt(LightDistanceSquared), LightNormal, CosConeAngles);
+	}
+
 	return Result;
 }
 
@@ -117,9 +154,67 @@ float PointLight_EstimateLight(
 
 	float LightPower = Luminance(GetColor(LightId));
 	float Falloff = ComputeAttenuationFalloff(LightDistanceSquared, LightId);
-	return LightPower * Falloff * NoL / LightDistanceSquared;
+	float OutIrradiance = LightPower * Falloff * NoL / LightDistanceSquared;
+	if (IsSpotLight(LightId))
+	{
+		float3 LightNormal = GetNormal(LightId);
+		float2 CosConeAngles = GetCosConeAngles(LightId);
+		OutIrradiance *= SpotAttenuation(LightDirection * rsqrt(LightDistanceSquared), LightNormal, CosConeAngles);
+	}
+	return OutIrradiance;
+}
+
+
+float2 IntersectInfiniteConeVolume(
+	float3  ro, float3  rd,
+	float3  apex, float3  normal,
+	float radius)
+{
+	// Adapted from: https://www.iquilezles.org/www/articles/intersectors/intersectors.htm
+	const float3 ba = normal;
+	const float3 oa = ro - apex;
+
+	const float m0 = dot(ba, ba);
+	const float m1 = dot(oa, ba);
+	const float m3 = dot(rd, ba);
+
+	const float m4 = dot(rd, oa);
+	const float m5 = dot(oa, oa);
+	const float hy = m0 + radius * radius;
+
+	const float k2 = m0 * m0 - m3 * m3 * hy;
+	const float k1 = m0 * m0 * m4 - m1 * m3 * hy;
+	const float k0 = m0 * m0 * m5 - m1 * m1 * hy;
+
+	const float h = k1 * k1 - k2 * k0;
+	if (h > 0)
+	{
+		const float s = sign(k2);
+		const float2 t = (-k1 + float2(-s, s) * sqrt(h)) / k2;
+		const float2 y = m1 + t * m3;
+
+		// check which hits are valid:
+		const int c = ((y.x > 0) ? 1 : 0) + (y.y > 0 ? 2 : 0);
+		if (c == 1)
+		{
+			// only smaller root is valid, far root is invalid
+			return float2(0.0, t.x);
+		}
+		if (c == 2)
+		{
+			// only larger root is valid, must have started outside the cone
+			return float2(t.y, POSITIVE_INFINITY);
+		}
+		if (c == 3)
+		{
+			// both roots are valid
+			return t;
+		}
+	}
+	return -1.0;
 }
 
+
 FVolumeLightSampleSetup PointLight_PrepareLightVolumeSample(
 	int LightId,
 	float3 RayOrigin,
@@ -139,6 +234,25 @@ FVolumeLightSampleSetup PointLight_PrepareLightVolumeSample(
 		float2 t = -b + float2(-1, +1) * sqrt(h);
 		TMin = max(t.x, TMin);
 		TMax = min(t.y, TMax);
+		if (IsSpotLight(LightId))
+		{
+			float3 N = GetNormal(LightId);
+			// now clip against the cone of the spot
+			const float CosOuterConeAngle = GetCosConeAngles(LightId).x;
+			const float SinOuterConeAngle = sqrt(1 - CosOuterConeAngle * CosOuterConeAngle);
+
+			float2 TCone = IntersectInfiniteConeVolume(RayOrigin, RayDirection, Center, N * CosOuterConeAngle, SinOuterConeAngle);
+			if (TCone.x < TCone.y)
+			{
+				TMin = max(TCone.x, TMin);
+				TMax = min(TCone.y, TMax);
+			}
+			else
+			{
+				// hit the sphere but missed the cone
+				return NullVolumeLightSetup();
+			}
+		}
 		float3 C = GetColor(LightId);
 		float LightImportance = max3(C.x, C.y, C.z);
 		return CreateEquiAngularSampler(LightImportance, Center, RayOrigin, RayDirection, TMin, TMax);
diff --git a/Engine/Shaders/Private/PathTracing/Light/PathTracingSpotLight.ush b/Engine/Shaders/Private/PathTracing/Light/PathTracingSpotLight.ush
deleted file mode 100644
diff --git a/Engine/Shaders/Private/PathTracing/PathTracingCore.ush b/Engine/Shaders/Private/PathTracing/PathTracingCore.ush
--- a/Engine/Shaders/Private/PathTracing/PathTracingCore.ush
+++ b/Engine/Shaders/Private/PathTracing/PathTracingCore.ush
@@ -1542,7 +1542,7 @@ bool PathTracingKernel(inout FPathState PathState, int Bounce)
 	FMaterialSample MaterialSample = SampleMaterial(-PathState.Ray.Direction, Payload, RandSample.xyz);
 
 	// If we are using Light sampling and the material can use it ...
-	if (MISMode != 0 && SceneLightCount > 0 && bDoLightLoop)
+	if (MISMode != 0 && bDoLightLoop)
 	{
 		// Choose a light and sample it
 		float3 TranslatedWorldPos = Payload.TranslatedWorldPos;
