Fix missing screen shadow casting on non-shadowed light for hair strands samples.

diff --git a/Engine/Shaders/Private/DeferredLightingCommon.ush b/Engine/Shaders/Private/DeferredLightingCommon.ush
--- a/Engine/Shaders/Private/DeferredLightingCommon.ush
+++ b/Engine/Shaders/Private/DeferredLightingCommon.ush
@@ -63,11 +63,11 @@ float ShadowRayCast(
 		uint2 PixelPos = View.ViewRectMin.xy + View.ViewSizeAndInvSize.xy * HitUV;
 		FSubstrateAddressing SubstrateAddressing = GetSubstratePixelDataByteOffset(PixelPos, uint2(View.BufferSizeAndInvSize.xy), SubstrateRayCast.MaxBytesPerPixel);
 		FSubstratePixelHeader SubstratePixelHeader = UnpackSubstrateHeaderIn(SubstrateRayCast.MaterialTextureArray, SubstrateAddressing, SubstrateRayCast.TopLayerTexture);
-		bOutHitCastContactShadow = SubstratePixelHeader.DoesCastContactShadow();
+		bOutHitCastContactShadow = bHairNoShadowLight || SubstratePixelHeader.DoesCastContactShadow();
 #endif
 #else
 		FGBufferData SampleGBuffer = GetGBufferData(HitUV);
-		bOutHitCastContactShadow = CastContactShadow(SampleGBuffer);
+		bOutHitCastContactShadow = bHairNoShadowLight || CastContactShadow(SampleGBuffer);
 #endif
 	}
 
