Use FMA macro in DoubleFloat headers. Add code path for platforms without FMA. Tweak DFMultiply for improved precision. Add DFFastToTranslatedWorld. Add a WaveReadLaneAt overload for FDFMatrix and FDFInverseMatrix.

diff --git a/Engine/Shaders/Private/DoubleFloat.ush b/Engine/Shaders/Private/DoubleFloat.ush
--- a/Engine/Shaders/Private/DoubleFloat.ush
+++ b/Engine/Shaders/Private/DoubleFloat.ush
@@ -6,13 +6,14 @@
 // 'High' stores the core value, 'Low' stores the residual error that couldn't fit in 'High'.
 // This combination has 2x23=46 significant bits, providing twice the precision a float offers.
 // Operations are slower than floats, but faster than doubles on consumer GPUs (with potentially greater support)
-// Platforms that don't support fused multiply-add and INVARIANT may see reduced precision.
+// Platforms that don't support fused multiply-add and INVARIANT may see decreased performance or reduced precision.
 // 
 // Based on:
 //  [0] Thall, A. (2006), Extended-precision floating-point numbers for GPU computation.
 //  [1] Mioara Maria Joldes, Jean-Michel Muller, Valentina Popescu. (2017), Tight and rigourous error bounds for basic building blocks of double-word arithmetic.
 //  [2] Vincent Lefevre, Nicolas Louvet, Jean-Michel Muller, Joris Picot, and Laurence Rideau. (2022), Accurate Calculation of Euclidean Norms using Double-Word Arithmetic
 //  [3] Jean-Michel Muller and Laurence Rideau. (2022), Formalization of Double-Word Arithmetic, and Comments on "Tight and Rigorous Error Bounds for Basic Building Blocks of Double-Word Arithmetic"
+//  [4] T. J. Dekker. (1971), A floating-point technique for extending the available precision.
 
 #ifndef UE_DF_FORCE_FP32_OPS
 #define UE_DF_FORCE_FP32_OPS 0
@@ -22,11 +23,6 @@
 #define UE_DF_NO_FAST_MATH 0
 #endif
 
-// Temporary hack until mad support is added to Switch toolchain
-#if SWITCH_PROFILE || SWITCH_PROFILE_FORWARD
-#define mad(a, b, c) (((a)*(b))+(c))
-#endif
-
 /* FDFScalar */
 
 #define FDFType FDFScalar
@@ -196,6 +192,15 @@ FDFVector4 DFSwizzle(FDFVector2 V, int C0, int C1, int C2, int C3) { return Make
 FDFVector4 DFSwizzle(FDFVector3 V, int C0, int C1, int C2, int C3) { return MakeDFVector(DFGetComponent(V, C0), DFGetComponent(V, C1), DFGetComponent(V, C2), DFGetComponent(V, C3)); }
 FDFVector4 DFSwizzle(FDFVector4 V, int C0, int C1, int C2, int C3) { return MakeDFVector(DFGetComponent(V, C0), DFGetComponent(V, C1), DFGetComponent(V, C2), DFGetComponent(V, C3)); }
 
+float3 DFFastToTranslatedWorld(FDFVector3 WorldPosition, FDFVector3 PreViewTranslation)
+{
+	return DFFastLocalSubtractDemote(WorldPosition, DFNegate(PreViewTranslation));
+}
+
+float3 DFFastToTranslatedWorld(float3 WorldPosition, FDFVector3 PreViewTranslation)
+{
+	return DFFastLocalSubtractDemote(WorldPosition, DFNegate(PreViewTranslation));
+}
 
 #define DFHackToFloat DFDemote
 
@@ -482,6 +487,11 @@ float4x4 DFDemote(FDFMatrix V)
 	return MultiplyTranslation(V.M, V.PostTranslation);
 }
 
+float4x4 DFDemote(float4x4 V)
+{
+	return V;
+}
+
 float4x4 DFDemote(FDFInverseMatrix V)
 {
 	return MultiplyTranslation(-V.PreTranslation, V.M);
@@ -490,7 +500,7 @@ float4x4 DFDemote(FDFInverseMatrix V)
 float4 DFTransformLocalToTranslatedWorld(float3 LocalPosition, FDFMatrix LocalToWorld, FDFVector3 PreViewTranslation)
 {
 	// Multiply local position with LocalToWorld 3x3 first for high precision rotation and scale, then move the world position part to translated world before adding it  
-	float3 RotatedScaledPosition = mad(LocalPosition.xxx, LocalToWorld.M[0].xyz, mad(LocalPosition.yyy, LocalToWorld.M[1].xyz, LocalPosition.zzz * LocalToWorld.M[2].xyz));
+	float3 RotatedScaledPosition = FMA(LocalPosition.xxx, LocalToWorld.M[0].xyz, FMA(LocalPosition.yyy, LocalToWorld.M[1].xyz, LocalPosition.zzz * LocalToWorld.M[2].xyz));
 
 #if UE_DF_NO_FAST_MATH
 	float3 TranslatedWorldPositionOrigin = DFAddDemote(DFGetOrigin(LocalToWorld), ResolvedView.PreViewTranslation);
@@ -503,12 +513,31 @@ float4 DFTransformLocalToTranslatedWorld(float3 LocalPosition, FDFMatrix LocalTo
 
 FDFVector3 DFMultiplyLHSAndFastTwoSum(float3 Lhs, float Factor, float3 Rhs)
 {
-	const float3 S = mad(Lhs, Factor, Rhs);
-	const float3 E = INVARIANT(Rhs - mad(-Lhs, Factor, S));
+	const float3 S = INVARIANT(FMA(Lhs, Factor, Rhs));
+	const float3 E = INVARIANT(Rhs - FMA(-Lhs, Factor, S));
 	return MakeDFVector3(S, E);
 }
 
-// Temporary hack until mad support is added to Switch toolchain
-#if SWITCH_PROFILE || SWITCH_PROFILE_FORWARD
-#undef mad
-#endif
\ No newline at end of file
+#if FEATURE_LEVEL >= FEATURE_LEVEL_SM6 || PLATFORM_SUPPORTS_SM6_0_WAVE_OPERATIONS
+FDFMatrix WaveReadLaneAt(FDFMatrix In, uint SrcIndex)
+{
+	FDFMatrix Result;
+	Result.M[0] = WaveReadLaneAt(In.M[0], SrcIndex);
+	Result.M[1] = WaveReadLaneAt(In.M[1], SrcIndex);
+	Result.M[2] = WaveReadLaneAt(In.M[2], SrcIndex);
+	Result.M[3] = WaveReadLaneAt(In.M[3], SrcIndex);
+	Result.PostTranslation = WaveReadLaneAt(In.PostTranslation, SrcIndex);
+	return Result;
+}
+
+FDFInverseMatrix WaveReadLaneAt(FDFInverseMatrix In, uint SrcIndex)
+{
+	FDFInverseMatrix Result;
+	Result.M[0] = WaveReadLaneAt(In.M[0], SrcIndex);
+	Result.M[1] = WaveReadLaneAt(In.M[1], SrcIndex);
+	Result.M[2] = WaveReadLaneAt(In.M[2], SrcIndex);
+	Result.M[3] = WaveReadLaneAt(In.M[3], SrcIndex);
+	Result.PreTranslation = WaveReadLaneAt(In.PreTranslation, SrcIndex);
+	return Result;
+}
+#endif
diff --git a/Engine/Shaders/Private/DoubleFloatOperations.ush b/Engine/Shaders/Private/DoubleFloatOperations.ush
--- a/Engine/Shaders/Private/DoubleFloatOperations.ush
+++ b/Engine/Shaders/Private/DoubleFloatOperations.ush
@@ -218,7 +218,8 @@ FFloatType DFFastLocalSubtractDemote(FDFType Lhs, FDFType Rhs)
 #else
 	const FFloatType High = INVARIANT(Lhs.High - Rhs.High);
 	const FFloatType Low = INVARIANT(Lhs.Low - Rhs.Low);
-	return INVARIANT(High + Low);
+	const FFloatType Sum = INVARIANT(High + Low);
+	return Sum;
 #endif
 }
 
@@ -230,7 +231,8 @@ FFloatType DFFastLocalSubtractDemote(FDFType Lhs, FFloatType Rhs)
 	return DFSubtractDemote(Lhs, Rhs);
 #else
 	const FFloatType High = INVARIANT(Lhs.High - Rhs);
-	return INVARIANT(High + Lhs.Low);
+	const FFloatType Sum = INVARIANT(High + Lhs.Low);
+	return Sum;
 #endif
 }
 
@@ -243,15 +245,34 @@ FFloatType DFFastLocalSubtractDemote(FFloatType Lhs, FDFType Rhs)
 /** MULTIPLICATION **/
 /********************/
 
+// [4] Algorithm 'mul12' (Page 241)
+FDFType DFTwoProductNoFMA(FFloatType Lhs, FFloatType Rhs)
+{
+	const FFloatType Constant = 4097; // = 2^(t - t/2) + 1, where t is the number of mantissa bits
+	const FFloatType Px = INVARIANT(Lhs * Constant);
+	const FFloatType Hx = INVARIANT(Lhs - Px + Px);
+	const FFloatType Tx = INVARIANT(Lhs - Hx);
+	const FFloatType Py = INVARIANT(Rhs * Constant);
+	const FFloatType Hy = INVARIANT(Rhs - Py + Py);
+	const FFloatType Ty = INVARIANT(Rhs - Hy);
+	const FFloatType P = INVARIANT(Hx * Hy);
+	const FFloatType Q = INVARIANT(Hx * Ty + Tx * Hy);
+	const FFloatType Z = INVARIANT(P + Q);
+	const FFloatType E = INVARIANT(P - Z + Q + Tx * Ty);
+	return DFFastTwoSum(Z, E);
+}
+
 // Multiply up two single precision scalars to one double precision scalar. (AKA Fast2Mult)
 // [1] Algorithm 3
 FDFType DFTwoProduct(FFloatType Lhs, FFloatType Rhs)
 {
 #if UE_DF_FORCE_FP32_OPS
 	return DFPromote(Lhs * Rhs);
+#elif !PLATFORM_SUPPORTS_FMA
+	return DFTwoProductNoFMA(Lhs, Rhs);
 #else
 	const FFloatType P = Lhs * Rhs;
-	const FFloatType E = INVARIANT(mad(Lhs, Rhs, -P));
+	const FFloatType E = INVARIANT(FMA(Lhs, Rhs, -P));
 	return DFConstructor(P, E);
 #endif
 }
@@ -259,30 +280,32 @@ FDFType DFTwoProduct(FFloatType Lhs, FFloatType Rhs)
 // Alias for DFTwoProduct
 FDFType DFMultiply(FFloatType Lhs, FFloatType Rhs) { return DFTwoProduct(Lhs, Rhs); }
 
-// [1] Algorithm 11
+// [1] Algorithm 11.
+// If no FMA is available, mad() expansion makes this equivalent to [1] Algorithm 10
 FDFType DFMultiply(FDFType Lhs, FDFType Rhs)
 {
 #if UE_DF_FORCE_FP32_OPS
 	return DFPromote(Lhs.High * Rhs.High);
 #else
 	const FDFType P = DFTwoProduct(Lhs.High, Rhs.High);
-	FFloatType T = P.Low;
-	T = INVARIANT(mad(Lhs.High, Rhs.Low, T));
-	T = INVARIANT(mad(Lhs.Low, Rhs.High, T));
+	FFloatType T = INVARIANT(Lhs.High * Rhs.Low);
+	T = INVARIANT(FMA(Lhs.Low, Rhs.High, T));
+	T = INVARIANT(T + P.Low);
 
 	return DFFastTwoSum(P.High, T);
 #endif
 }
 
 // [1] Algorithm 9
+// If no FMA is available, mad() expansion makes this equivalent to [1] Algorithm 8
 FDFType DFMultiply(FDFType Lhs, FFloatType Rhs)
 {
 #if UE_DF_FORCE_FP32_OPS
 	return DFPromote(Lhs.High * Rhs);
 #else
 	FDFType P = DFTwoProduct(Lhs.High, Rhs);
-	P.Low = mad(Lhs.Low, Rhs, P.Low);
-	return DFFastTwoSum(P.High, P.Low);
+	FFloatType T = INVARIANT(FMA(Lhs.Low, Rhs, P.Low));
+	return DFFastTwoSum(P.High, T);
 #endif
 }
 
@@ -294,9 +317,9 @@ FFloatType DFMultiplyDemote(FDFType Lhs, FDFType Rhs)
 	return Lhs.High * Rhs.High;
 #else
 	const FDFType P = DFTwoProduct(Lhs.High, Rhs.High);
-	FFloatType T = P.Low;
-	T = INVARIANT(mad(Lhs.High, Rhs.Low, T));
-	T = INVARIANT(mad(Lhs.Low, Rhs.High, T));
+	FFloatType T = INVARIANT(Lhs.High * Rhs.Low);
+	T = INVARIANT(FMA(Lhs.Low, Rhs.High, T));
+	T = INVARIANT(T + P.Low);
 
 	return P.High + T;
 #endif
@@ -308,7 +331,7 @@ FFloatType DFMultiplyDemote(FDFType Lhs, FFloatType Rhs)
 	return Lhs.High * Rhs;
 #else
 	FDFType P = DFTwoProduct(Lhs.High, Rhs);
-	P.Low = mad(Lhs.Low, Rhs, P.Low);
+	P.Low = FMA(Lhs.Low, Rhs, P.Low);
 	return P.High + P.Low;
 #endif
 }
@@ -334,14 +357,18 @@ FDFType DFSqr(FFloatType V) { return DFTwoProduct(V, V); }
 /** DIVISION **/
 /**************/
 
+FDFType DFFastDivide(FDFType Lhs, FDFType Rhs);
+
 // [1] Algorithm 18
 FDFType DFDivide(FDFType Lhs, FDFType Rhs)
 {
 #if UE_DF_FORCE_FP32_OPS
 	return DFPromote(Lhs.High / Rhs.High);
+#elif !PLATFORM_SUPPORTS_FMA
+	return DFFastDivide(Lhs, Rhs);
 #else
 	const FFloatType Th = INVARIANT(1.0 / Rhs.High);
-	const FFloatType Rh = INVARIANT(mad(-Rhs.High, Th, 1));
+	const FFloatType Rh = INVARIANT(FMA(-Rhs.High, Th, 1));
 	const FFloatType Rl = INVARIANT(-Rhs.Low * Th);
 	const FDFType E = DFFastTwoSum(Rh, Rl);
 	const FDFType D = DFMultiply(E, Th);
@@ -407,7 +434,7 @@ FDFType DFFastDivide(FDFType Lhs, FFloatType Rhs)
 {
 #if UE_DF_FORCE_FP32_OPS
 	return DFPromote(Lhs.High / Rhs);
-#elif UE_DF_NO_FAST_MATH
+#elif UE_DF_NO_FAST_MATH && PLATFORM_SUPPORTS_FMA
 	return DFDivide(Lhs, Rhs);
 #else
 	const FFloatType Th = INVARIANT(Lhs.High / Rhs);
@@ -627,7 +654,7 @@ FDFType DFSqrt(FDFType V)
 	return DFPromote(sqrt(V.High));
 #else
 	const FFloatType Sh = sqrt(V.High);
-	const FFloatType P1 = INVARIANT(mad(-Sh, Sh, V.High));
+	const FFloatType P1 = INVARIANT(FMA(-Sh, Sh, V.High)); // Note: precision loss here if FMA is not supported
 	const FFloatType P2 = INVARIANT(V.Low + P1);
 
 	const FFloatType Sl = INVARIANT(P2 / (2.0 * Sh));
@@ -635,9 +662,9 @@ FDFType DFSqrt(FDFType V)
 
 	// equivalent to (remark 3.5)
 	//precise FFloatType Tl = P2 / Sh;
-	//precise FFloatType Zh = mad(0.5, Tl, Sh);
+	//precise FFloatType Zh = FMA(0.5, Tl, Sh);
 	//precise FFloatType D = Zh - Sh;
-	//precise FFloatType Zl = mad(0.5, Tl, -D);
+	//precise FFloatType Zl = FMA(0.5, Tl, -D);
 	//return DFConstructor(Zh, Zl);
 #endif
 }
@@ -649,7 +676,7 @@ FFloatType DFSqrtDemote(FDFType V)
 	return sqrt(V.High);
 #else
 	const FFloatType Sh = sqrt(V.High);
-	const FFloatType P1 = INVARIANT(mad(-Sh, Sh, V.High));
+	const FFloatType P1 = INVARIANT(FMA(-Sh, Sh, V.High)); // Note: precision loss here if FMA is not supported
 	const FFloatType P2 = INVARIANT(V.Low + P1);
 
 	const FFloatType Sl = INVARIANT(P2 / (2.0 * Sh));
@@ -702,8 +729,8 @@ void DFSinCos(FDFType V, out FFloatType Sin, out FFloatType Cos)
 	sincos(V.High, Sh, Ch);
 	FFloatType Sl, Cl;
 	sincos(V.Low, Sl, Cl);
-	Sin = mad(Sh, Cl,  Ch*Sl); // trigonometric identity for sin(h+l)
-	Cos = mad(Ch, Cl, -Sh*Sl); // trigonometric identity for cos(h+l)
+	Sin = FMA(Sh, Cl,  Ch*Sl); // trigonometric identity for sin(h+l)
+	Cos = FMA(Ch, Cl, -Sh*Sl); // trigonometric identity for cos(h+l)
 #endif
 }
 
diff --git a/Engine/Shaders/Private/LargeWorldCoordinates.ush b/Engine/Shaders/Private/LargeWorldCoordinates.ush
--- a/Engine/Shaders/Private/LargeWorldCoordinates.ush
+++ b/Engine/Shaders/Private/LargeWorldCoordinates.ush
@@ -513,11 +513,6 @@ float4x4 LWCHackToFloat(FLWCInverseMatrix v) { return (float4x4)0.0f; }
 #define LWCHackToFloat(V) LWCToFloat(V)
 #endif
 
-// Temporary hack until mad support is added to Switch toolchain
-#if SWITCH_PROFILE || SWITCH_PROFILE_FORWARD
-#define mad(a, b, c) (((a)*(b))+(c))
-#endif
-
 #define DEFINE_DF_TO_TILEOFFSET_OPERATOR(FDFType, FLWCType, LWCConstructor, FFloatType) \
 	FLWCType DFToTileOffset(FDFType In)																							\
 	{																															\
@@ -546,11 +541,6 @@ DEFINE_DF_TO_TILEOFFSET_OPERATOR(FDFVector3, FLWCVector3, MakeLWCVector3, float3
 DEFINE_DF_TO_TILEOFFSET_OPERATOR(FDFVector4, FLWCVector4, MakeLWCVector4, float4)
 #undef DEFINE_DF_TO_TILEOFFSET_OPERATOR
 
-// Temporary hack until mad support is added to Switch toolchain
-#if SWITCH_PROFILE || SWITCH_PROFILE_FORWARD
-#undef mad
-#endif
-
 FLWCMatrix DFToTileOffset(FDFMatrix In)
 {
 	FLWCVector3 PosHigh = DFToTileOffset(MakeDFVector3(In.PostTranslation, 0));
