Stochastic Direct Lighting
* Fixed history neighborhood clamp
* Fixed double light attenuation during sampling
* Removed non candidate light mask light sample generation path
* Clamp history based on the sample luminance in order to reduce ghosting

diff --git a/Engine/Shaders/Private/StochasticDirectLighting/StochasticDirectLighting.ush b/Engine/Shaders/Private/StochasticDirectLighting/StochasticDirectLighting.ush
--- a/Engine/Shaders/Private/StochasticDirectLighting/StochasticDirectLighting.ush
+++ b/Engine/Shaders/Private/StochasticDirectLighting/StochasticDirectLighting.ush
@@ -177,13 +177,27 @@ uint GetShadingTileGridOffset(uint2 ShadingTileCoord, uint IndexInCell)
 }
 
 // Remove texture detail before denoising
-float3 SafeDemodulate(float3 Lighting, float3 Texture)
+void DemodulateLighting(FGBufferData GBuffer, float3 TranslatedWorldPosition, inout float3 DiffuseLighting, inout float3 SpecularLighting)
 {
-	return Lighting / max(Texture, 0.001f);
+	const float3 N = GBuffer.WorldNormal;
+	const float3 V = normalize(View.TranslatedWorldCameraOrigin - TranslatedWorldPosition);
+	const float NoV = saturate(dot(N, V));
+	const float3 SpecularEnv = EnvBRDF(GBuffer.SpecularColor, GBuffer.Roughness, NoV);
+
+	// #sdl_todo: demodulate earlier during BRDF evaluation
+	DiffuseLighting = DiffuseLighting / max(GBuffer.DiffuseColor, 0.001f);
+	SpecularLighting = SpecularLighting / max(SpecularEnv, 0.001f);	
 }
 
 // Restore texture detail after denoising
-float3 SafeModulate(float3 Lighting, float3 Texture)
+void ModulateLighting(FGBufferData GBuffer, float3 TranslatedWorldPosition, inout float3 DiffuseLighting, inout float3 SpecularLighting)
 {
-	return Lighting * max(Texture, 0.001f);
+	const float3 N = GBuffer.WorldNormal;
+	const float3 V = normalize(View.TranslatedWorldCameraOrigin - TranslatedWorldPosition);
+	const float NoV = saturate(dot(N, V));
+	const float3 SpecularEnv = EnvBRDF(GBuffer.SpecularColor, GBuffer.Roughness, NoV);
+
+	// Final pass outputs composites irradiance and outputs it to scene color
+	DiffuseLighting = DiffuseLighting * max(GBuffer.DiffuseColor, 0.001f);
+	SpecularLighting = SpecularLighting * max(SpecularEnv, 0.001f);
 }
\ No newline at end of file
diff --git a/Engine/Shaders/Private/StochasticDirectLighting/StochasticDirectLightingSampling.usf b/Engine/Shaders/Private/StochasticDirectLighting/StochasticDirectLightingSampling.usf
--- a/Engine/Shaders/Private/StochasticDirectLighting/StochasticDirectLightingSampling.usf
+++ b/Engine/Shaders/Private/StochasticDirectLighting/StochasticDirectLightingSampling.usf
@@ -13,78 +13,62 @@
 #include "StochasticDirectLighting.ush"
 #include "StochasticDirectLightingTracing.ush"
 
-float GetLocalLightTargetPDF(FDeferredLightData LightData, float3 TranslatedWorldPosition, FGBufferData GBuffer, uint2 ScreenCoord)
+struct FLightTargetPDF
 {
-	float SampleWeight = 0.0f;
+	float3 Diffuse;
+	float3 Specular;
+	float Weight;
+};
 
-	float3 L = LightData.Direction;
-	float3 ToLight = L;
+FLightTargetPDF InitLightTargetPDF()
+{
+	FLightTargetPDF LightTargetPDF;
+	LightTargetPDF.Diffuse = 0.0f;
+	LightTargetPDF.Specular = 0.0f;
+	LightTargetPDF.Weight = 0.0f;
+	return LightTargetPDF;
+}
 
-	float LightMask = GetLocalLightAttenuation(TranslatedWorldPosition, LightData, ToLight, L);
+FLightTargetPDF GetLocalLightTargetPDF(FDeferredLightData LightData, float3 TranslatedWorldPosition, FGBufferData GBuffer, uint2 ScreenCoord)
+{
+	FLightTargetPDF LightTargetPDF = InitLightTargetPDF();
 
-	if (LightMask > 0.0f)
-	{
-		float Attenuation = 0.0f;
+	float3 CameraVector = normalize(TranslatedWorldPosition - View.TranslatedWorldCameraOrigin);
+	GBuffer.Roughness = max(GBuffer.Roughness, View.MinRoughness);
 
-		if (LightData.bRectLight)
-		{
-			FRect Rect = GetRect(ToLight, LightData);
-			Attenuation = IntegrateLight(Rect);
-		}
-		else
-		{
-			FCapsuleLight Capsule = GetCapsule(ToLight, LightData);
-			Attenuation = IntegrateLight(Capsule, LightData.bInverseSquared);
-		}
+	float4 LightAttenuation = 1.0f;
+	float Dither = 0.5f;
+	float SurfaceShadow = 1.0f;
+	float AmbientOcclusion = 1.0f;
 
-		// #sdl_todo: figure out cheap sampling
-		SampleWeight = Attenuation * LightMask * View.PreExposure * Luminance(LightData.Color);
+	// #sdl_todo: Move more shading models to TILE_MODE_SIMPLE_SHADING
+	#if TILE_TYPE == TILE_MODE_SIMPLE_SHADING
+	{
+		GBuffer.ShadingModelID = SHADINGMODELID_DEFAULT_LIT;
+	}
+	#endif
+	LightData.ShadowedBits = 0;
 
-		#define REFERENCE_SAMPLING 1
-		#if REFERENCE_SAMPLING
-		{
-			float3 CameraVector = normalize(TranslatedWorldPosition - View.TranslatedWorldCameraOrigin);
-			GBuffer.Roughness = max(GBuffer.Roughness, View.MinRoughness);
+	FDeferredLightingSplit SplitLighting = GetDynamicLightingSplit(
+		TranslatedWorldPosition, CameraVector, GBuffer, AmbientOcclusion, GBuffer.ShadingModelID, 
+		LightData, LightAttenuation, Dither, ScreenCoord, 
+		SurfaceShadow);
 
-			float4 LightAttenuation = 1.0f;
-			float Dither = 0.5f;
-			float SurfaceShadow = 1.0f;
-			float AmbientOcclusion = 1.0f;
+	LightTargetPDF.Diffuse = SplitLighting.DiffuseLighting.xyz * View.PreExposure;
+	LightTargetPDF.Specular = SplitLighting.SpecularLighting.xyz * View.PreExposure;
 
-			// #sdl_todo: Move more shading models to TILE_MODE_SIMPLE_SHADING
-			#if TILE_TYPE == TILE_MODE_SIMPLE_SHADING
-			{
-				GBuffer.ShadingModelID = SHADINGMODELID_DEFAULT_LIT;
-			}
-			#endif
-			LightData.ShadowedBits = 0;
-
-			float4 LightRadiance = GetDynamicLighting(
-				TranslatedWorldPosition,
-				CameraVector,
-				GBuffer, 
-				AmbientOcclusion,
-				GBuffer.ShadingModelID,
-				LightData,
-				LightAttenuation,
-				Dither,
-				ScreenCoord, 
-				SurfaceShadow);
-
-			SampleWeight = Luminance(LightRadiance.xyz * View.PreExposure);
-
-			if (SampleWeight > 0.01f && LightData.IESAtlasIndex >= 0)
-			{
-				SampleWeight *= ComputeLightProfileMultiplier(TranslatedWorldPosition, LightData.TranslatedWorldPosition, -LightData.Direction, LightData.Tangent, LightData.IESAtlasIndex);
-			}
-		}
-		#endif
+	// #sdl_todo: move light profile to GetDynamicLightingSplit for better handling
+	if (LightData.IESAtlasIndex >= 0 && Luminance(LightTargetPDF.Diffuse + LightTargetPDF.Specular) > 0.01f)
+	{
+		const float LightProfileMult = ComputeLightProfileMultiplier(TranslatedWorldPosition, LightData.TranslatedWorldPosition, -LightData.Direction, LightData.Tangent, LightData.IESAtlasIndex);
+		LightTargetPDF.Diffuse *= LightProfileMult;
+		LightTargetPDF.Specular *= LightProfileMult;
 	}
 
 	// Simulate tonemapping
-	SampleWeight = log2(max(SampleWeight + 1.0f, 1.0f));
+	LightTargetPDF.Weight = log2(Luminance(LightTargetPDF.Diffuse + LightTargetPDF.Specular) + 1.0f);
 
-	return SampleWeight;
+	return LightTargetPDF;
 }
 
 uint2 DownsampledViewSize;
@@ -95,6 +79,7 @@ int DebugMode;
 
 RWTexture2D<float> RWDownsampledSceneDepth;
 RWTexture2D<UNORM float3> RWDownsampledSceneWorldNormal;
+RWTexture2D<float2> RWSampleLuminanceSum;
 RWStructuredBuffer<uint> RWCompositeTileAllocator;
 RWStructuredBuffer<uint2> RWCompositeTileData;
 RWTexture2D<uint> RWLightSamples;
@@ -103,10 +88,7 @@ float4 HistoryScreenPositionScaleBias;
 float4 HistoryUVMinMax;
 
 groupshared uint SharedLightMask[SHARED_LIGHT_MASK_SIZE];
-
-#if CANDIDATE_LIGHT_MASK
-	groupshared uint SharedCandidateLightMask[SHARED_LIGHT_MASK_SIZE];
-#endif
+groupshared uint SharedCandidateLightMask[SHARED_LIGHT_MASK_SIZE];
 
 void MarkLocalLightIndex(uint LocalLightIndex)
 {
@@ -125,7 +107,7 @@ float2 DownsampledBufferInvSize;
  * Run one thread per sample and generate new light samples for tracing
  */
 [numthreads(THREADGROUP_SIZE, THREADGROUP_SIZE, 1)]
-void GenerateSamplesCS(
+void GenerateLightSamplesCS(
 	uint3 GroupId : SV_GroupID,
 	uint3 GroupThreadId : SV_GroupThreadID,
 	uint3 DispatchThreadId : SV_DispatchThreadID)
@@ -137,12 +119,10 @@ void GenerateSamplesCS(
 		SharedLightMask[LinearThreadIndex] = 0;
 	}
 
-#if CANDIDATE_LIGHT_MASK
 	if (LinearThreadIndex < SHARED_LIGHT_MASK_SIZE)
 	{
 		SharedCandidateLightMask[LinearThreadIndex] = 0;
 	}
-#endif
 
 	GroupMemoryBarrierWithGroupSync();
 
@@ -152,7 +132,6 @@ void GenerateSamplesCS(
 		uint2 DownsampledTileCoord = UnpackTile(DownsampledTileData[DownsampledTileIndex + TILE_TYPE * DownsampledTileDataStride]);
 		uint2 DownsampledScreenCoord = DownsampledTileCoord * TILE_SIZE + GroupThreadId.xy;
 
-		#if CANDIDATE_LIGHT_MASK
 		if (all(DownsampledScreenCoord < DownsampledViewSize))
 		{
 			uint2 ScreenCoord = DownsampledScreenCoordToScreenCoord(DownsampledScreenCoord);	
@@ -180,7 +159,6 @@ void GenerateSamplesCS(
 				}
 			}
 		}
-		#endif
 	
 		GroupMemoryBarrierWithGroupSync();
 
@@ -211,6 +189,8 @@ void GenerateSamplesCS(
 			const float2 ScreenUV = (ScreenCoord + 0.5f) * View.BufferSizeAndInvSize.zw;
 			const FGBufferData GBuffer = GetGBufferData(ScreenUV);
 			const float SceneDepth = GBuffer.Depth;
+			float3 DiffuseSum = 0.0f;
+			float3 SpecularSum = 0.0f;
 			float WeightSum = 0.0f;
 
 			if (SceneDepth > 0)
@@ -268,7 +248,6 @@ void GenerateSamplesCS(
 				}
 				#endif
 
-	#if CANDIDATE_LIGHT_MASK
 				for (uint MaskIndex = 0; MaskIndex < SHARED_LIGHT_MASK_SIZE; ++MaskIndex)
 				{
 					uint CandidateLightMask = SharedCandidateLightMask[MaskIndex];
@@ -283,67 +262,24 @@ void GenerateSamplesCS(
 						const FLightSceneData LightSceneData = GetLightSceneData(LocalLightData.LightSceneId);
 						FDeferredLightData LightData = ConvertToDeferredLight(LightSceneData);
 					
-						float SampleWeight = GetLocalLightTargetPDF(LightData, TranslatedWorldPosition, GBuffer, ScreenCoord);
+						FLightTargetPDF LightTargetPDF = GetLocalLightTargetPDF(LightData, TranslatedWorldPosition, GBuffer, ScreenCoord);
 
 						#if DEBUG_MODE
 						if (bDebug)
 						{
 							Newline(Context);
 							Print(Context, LocalLightData.LightSceneId);
-							Print(Context, SampleWeight);
+							Print(Context, LightTargetPDF.Weight);
 						}
 						#endif
 
-						if (SampleWeight > SamplingMinWeight)
-						{
-							float Tau = WeightSum / (WeightSum + SampleWeight);
-							WeightSum += SampleWeight;
-
-							for (uint LightSampleIndex = 0; LightSampleIndex < NUM_SAMPLES_PER_PIXEL_1D; ++LightSampleIndex)
-							{
-								if (LightIndexRandom[LightSampleIndex] < Tau)
-								{
-									LightIndexRandom[LightSampleIndex] /= Tau;
-								}
-								else
-								{
-									// Select this sample
-									LightIndexRandom[LightSampleIndex] = (LightIndexRandom[LightSampleIndex] - Tau) / (1.0f - Tau);
-									LightSamples[LightSampleIndex].LocalLightIndex = LocalLightIndex;
-									LightSamples[LightSampleIndex].Weight = SampleWeight;
-								}
-
-								LightIndexRandom[LightSampleIndex] = clamp(LightIndexRandom[LightSampleIndex], 0, 0.9999f);
-							}
-						}
-					}
-				}
-
-	#else
-				for (uint GridLightIndex = 0; GridLightIndex < NumLightsInGridCell; ++GridLightIndex)
-				{
-					const uint LocalLightIndex = GetCulledLightDataGrid(CulledLightGridData.DataStartIndex + GridLightIndex);
-					const FLocalLightData LocalLightData = GetLocalLightDataNonStereo(LocalLightIndex);
-
-					if (LocalLightData.bStochasticDirectLightingSupported)
-					{
-						const FLightSceneData LightSceneData = GetLightSceneData(LocalLightData.LightSceneId);
-						FDeferredLightData LightData = ConvertToDeferredLight(LightSceneData);
-						float SampleWeight = GetLocalLightTargetPDF(LightData, TranslatedWorldPosition, GBuffer, ScreenCoord);
-
-						#if DEBUG_MODE
-						if (bDebug)
+						if (LightTargetPDF.Weight > SamplingMinWeight)
 						{
-							Newline(Context);
-							Print(Context, LocalLightData.LightSceneId);
-							Print(Context, SampleWeight);
-						}
-						#endif
+							float Tau = WeightSum / (WeightSum + LightTargetPDF.Weight);
+							WeightSum += LightTargetPDF.Weight;
 
-						if (SampleWeight > SamplingMinWeight)
-						{
-							float Tau = WeightSum / (WeightSum + SampleWeight);
-							WeightSum += SampleWeight;
+							DiffuseSum += LightTargetPDF.Diffuse;
+							SpecularSum += LightTargetPDF.Specular;
 
 							for (uint LightSampleIndex = 0; LightSampleIndex < NUM_SAMPLES_PER_PIXEL_1D; ++LightSampleIndex)
 							{
@@ -356,7 +292,7 @@ void GenerateSamplesCS(
 									// Select this sample
 									LightIndexRandom[LightSampleIndex] = (LightIndexRandom[LightSampleIndex] - Tau) / (1.0f - Tau);
 									LightSamples[LightSampleIndex].LocalLightIndex = LocalLightIndex;
-									LightSamples[LightSampleIndex].Weight = SampleWeight;
+									LightSamples[LightSampleIndex].Weight = LightTargetPDF.Weight;
 								}
 
 								LightIndexRandom[LightSampleIndex] = clamp(LightIndexRandom[LightSampleIndex], 0, 0.9999f);
@@ -364,7 +300,8 @@ void GenerateSamplesCS(
 						}
 					}
 				}
-	#endif
+
+				DemodulateLighting(GBuffer, TranslatedWorldPosition, DiffuseSum, SpecularSum);
 
 				#if DEBUG_MODE
 				if (bDebug)
@@ -395,12 +332,20 @@ void GenerateSamplesCS(
 							AddLineTWS(Context, TranslatedWorldPosition, TranslatedWorldPosition + LightSampleTrace.Direction * LightSampleTrace.Distance, float4(LightSceneData.Color.xyz / Luminance(LightSceneData.Color.xyz), 1.0f));
 						}
 					}
+
+					Newline(Context);
+					Print(Context, TEXT("Diffuse sum:  "));
+					Print(Context, DiffuseSum);
+					Newline(Context);
+					Print(Context, TEXT("Specular sum: "));
+					Print(Context, SpecularSum);
 				}
 				#endif
 			}
 
 			RWDownsampledSceneDepth[DownsampledScreenCoord] = SceneDepth;
 			RWDownsampledSceneWorldNormal[DownsampledScreenCoord] = EncodeNormal(GBuffer.WorldNormal);
+			RWSampleLuminanceSum[DownsampledScreenCoord] = float2(Luminance(DiffuseSum), Luminance(SpecularSum));
 
 			for (uint LightSampleIndex = 0; LightSampleIndex < NUM_SAMPLES_PER_PIXEL_1D; ++LightSampleIndex)
 			{
diff --git a/Engine/Shaders/Private/StochasticDirectLighting/StochasticDirectLightingShading.usf b/Engine/Shaders/Private/StochasticDirectLighting/StochasticDirectLightingShading.usf
--- a/Engine/Shaders/Private/StochasticDirectLighting/StochasticDirectLightingShading.usf
+++ b/Engine/Shaders/Private/StochasticDirectLighting/StochasticDirectLightingShading.usf
@@ -152,21 +152,14 @@ void ShadeLightSamplesCS(
 				DiffuseLighting *= View.PreExposure;
 				SpecularLighting *= View.PreExposure;
 
-				const float3 N = GBuffer.WorldNormal;
-				const float3 V = normalize(View.TranslatedWorldCameraOrigin - TranslatedWorldPosition);
-				const float NoV = saturate(dot(N, V));
-				const float3 SpecularEnv = EnvBRDF(GBuffer.SpecularColor, GBuffer.Roughness, NoV);
-
-				// #sdl_todo: demodulate earlier during BRDF evaluation
-				DiffuseLighting = SafeDemodulate(DiffuseLighting, GBuffer.DiffuseColor);
-				SpecularLighting = SafeDemodulate(SpecularLighting, SpecularEnv);
+				DemodulateLighting(GBuffer, TranslatedWorldPosition, DiffuseLighting, SpecularLighting);
 			}
 
 			#if DEBUG_MODE
 			if (bDebug)
 			{
 				Newline(Context);
-				Print(Context, TEXT("Diffuse: ")); 
+				Print(Context, TEXT("Diffuse:  ")); 
 				Print(Context, DiffuseLighting);
 				Newline(Context);
 				Print(Context, TEXT("Specular: ")); 
diff --git a/Engine/Shaders/Private/StochasticDirectLighting/StochasticDirectLightingSpatial.usf b/Engine/Shaders/Private/StochasticDirectLighting/StochasticDirectLightingSpatial.usf
--- a/Engine/Shaders/Private/StochasticDirectLighting/StochasticDirectLightingSpatial.usf
+++ b/Engine/Shaders/Private/StochasticDirectLighting/StochasticDirectLightingSpatial.usf
@@ -23,7 +23,7 @@ void SDLSpatialFilterCS(
 #if DEBUG_MODE
 	int2 DebugScreenCoord = View.CursorPosition.x >= 0 ? View.CursorPosition * View.ViewResolutionFraction : -1;
 	bool bDebug = all(ScreenCoord == DebugScreenCoord);
-	FShaderPrintContext Context = InitShaderPrintContext(true, float2(0.55, 0.45));
+	FShaderPrintContext Context = InitShaderPrintContext(true, float2(0.55, 0.7));
 #endif
 
 	if (all(ScreenCoord < View.ViewRectMinAndSize.zw))
@@ -109,14 +109,10 @@ void SDLSpatialFilterCS(
 			}
 			#endif
 
-			const float3 N = GBuffer.WorldNormal;
-			const float3 V = normalize(View.TranslatedWorldCameraOrigin - TranslatedWorldPosition);
-			const float NoV = saturate(dot(N, V));
-			const float3 SpecularEnv = EnvBRDF(GBuffer.SpecularColor, GBuffer.Roughness, NoV);
-
 			// Final pass outputs composites irradiance and outputs it to scene color
-			float3 DiffuseLighting = SafeModulate(DiffuseLightingSum / DiffuseLightingWeightSum, GBuffer.DiffuseColor);
-			float3 SpecularLighting = SafeModulate(SpecularLightingSum / SpecularLightingWeightSum, SpecularEnv);
+			float3 DiffuseLighting = DiffuseLightingSum / DiffuseLightingWeightSum;
+			float3 SpecularLighting = SpecularLightingSum / SpecularLightingWeightSum;
+			ModulateLighting(GBuffer, TranslatedWorldPosition, DiffuseLighting, SpecularLighting);
 
 			// Composite diffuse and specular into scene color
 			const bool bNeedsSeparateSubsurfaceLightAccumulation = UseSubsurfaceProfile(GBuffer.ShadingModelID);
diff --git a/Engine/Shaders/Private/StochasticDirectLighting/StochasticDirectLightingTemporal.usf b/Engine/Shaders/Private/StochasticDirectLighting/StochasticDirectLightingTemporal.usf
--- a/Engine/Shaders/Private/StochasticDirectLighting/StochasticDirectLightingTemporal.usf
+++ b/Engine/Shaders/Private/StochasticDirectLighting/StochasticDirectLightingTemporal.usf
@@ -3,6 +3,8 @@
 #include "../Common.ush"
 #include "StochasticDirectLighting.ush"
 
+Texture2D<float4> CompositeUpsampleWeights;
+Texture2D<float2> SampleLuminanceSumTexture;
 Texture2D ResolvedDiffuseLighting;
 Texture2D ResolvedSpecularLighting;
 Texture2D StochasticDirectLightingDepthHistory;
@@ -15,6 +17,7 @@ float4 HistoryUVMinMax;
 float PrevSceneColorPreExposureCorrection;
 float TemporalMaxFramesAccumulated;
 float TemporalNeighborhoodClampScale;
+uint2 DownsampledViewSize;
 
 RWTexture2D<float4> RWDiffuseLightingAndSecondMoment;
 RWTexture2D<float4> RWSpecularLightingAndSecondMoment;
@@ -22,13 +25,14 @@ RWTexture2D<UNORM float> RWNumFramesAccumulated;
 RWTexture2D<float> RWSceneDepth;
 RWTexture2D<float4> RWSceneColor;
 
-struct FNeighborhoodClamp
+struct FNeighborhood
 {
-	float3 Min;
-	float3 Max;
+	float3 Mean;
+	float3 ClampMin;
+	float3 ClampMax;
 };
 
-FNeighborhoodClamp GetNeighborhoodClamp(
+FNeighborhood GetNeighborhood(
 	uint2 ScreenCoord, 
 	Texture2D SampleTexture,
 	float3 CenterSample)
@@ -46,7 +50,7 @@ FNeighborhoodClamp GetNeighborhoodClamp(
 			{
 				uint2 NeigborScreenCoord = clamp(ScreenCoord + int2(NeigborOffsetX, NeigborOffsetY), 0u, View.ViewRectMinAndSize.zw - 1);
 				float3 NeigborSample = SampleTexture[NeigborScreenCoord].xyz;
-				
+	
 				SampleSum += NeigborSample;
 				SampleSqSum += Pow2(NeigborSample);
 			}
@@ -58,11 +62,30 @@ FNeighborhoodClamp GetNeighborhoodClamp(
 	float3 M2 = SampleSqSum / NumSamples;
 	float3 Variance = max(M2 - Pow2(M1), 0.0f);
 	float3 StdDev = sqrt(Variance);
+	
+	FNeighborhood Neighborhood = (FNeighborhood)0;
+	Neighborhood.Mean = M1;
+	Neighborhood.ClampMin = M1 - TemporalNeighborhoodClampScale * StdDev;
+	Neighborhood.ClampMax = M1 + TemporalNeighborhoodClampScale * StdDev;
+	return Neighborhood;
+}
+
+void AccumulateSampleLuminanceSum(uint2 DownsampledScreenCoord, float UpsampleWeight, inout float2 SampleLuminanceSum)
+{
+	if (all(DownsampledScreenCoord < DownsampledViewSize))
+	{
+		SampleLuminanceSum += SampleLuminanceSumTexture[DownsampledScreenCoord] * UpsampleWeight;
+	}
+}
+
+float3 ClampLuminance(float3 Lighting, float LuminanceClamp)
+{
+	if (Luminance(Lighting) > LuminanceClamp)
+	{
+		Lighting *= LuminanceClamp / Luminance(Lighting);
+	}
 
-	FNeighborhoodClamp NeighborhoodClamp = (FNeighborhoodClamp)0;
-	NeighborhoodClamp.Min = M1 - TemporalNeighborhoodClampScale * StdDev;
-	NeighborhoodClamp.Max = M1 + TemporalNeighborhoodClampScale * StdDev;
-	return NeighborhoodClamp;
+	return Lighting;
 }
 
 [numthreads(THREADGROUP_SIZE, THREADGROUP_SIZE, 1)]
@@ -76,7 +99,7 @@ void SDLTemporalAccumulationCS(
 #if DEBUG_MODE
 	int2 DebugScreenCoord = View.CursorPosition.x >= 0 ? View.CursorPosition * View.ViewResolutionFraction : -1;
 	bool bDebug = all(ScreenCoord == DebugScreenCoord);
-	FShaderPrintContext Context = InitShaderPrintContext(true, float2(0.55, 0.3));
+	FShaderPrintContext Context = InitShaderPrintContext(true, float2(0.55, 0.4));
 #endif
 
 	if (all(ScreenCoord < View.ViewRectMinAndSize.zw))
@@ -95,6 +118,16 @@ void SDLTemporalAccumulationCS(
 
 		if (GBuffer.Depth > 0.0f && GBuffer.ShadingModelID != SHADINGMODELID_UNLIT)
 		{
+			float2 SampleLuminanceSum = 0.0f;
+			{
+				uint2 DownsampledScreenCoord = floor(ScreenUV * View.BufferSizeAndInvSize.xy / DOWNSAMPLE_FACTOR - 0.5f);
+				float4 DownsampledScreenWeights = CompositeUpsampleWeights[ScreenCoord];
+				AccumulateSampleLuminanceSum(DownsampledScreenCoord + uint2(0, 0), DownsampledScreenWeights.x, SampleLuminanceSum);
+				AccumulateSampleLuminanceSum(DownsampledScreenCoord + uint2(1, 0), DownsampledScreenWeights.y, SampleLuminanceSum);
+				AccumulateSampleLuminanceSum(DownsampledScreenCoord + uint2(0, 1), DownsampledScreenWeights.z, SampleLuminanceSum);
+				AccumulateSampleLuminanceSum(DownsampledScreenCoord + uint2(1, 1), DownsampledScreenWeights.w, SampleLuminanceSum);
+			}
+
 			DiffuseLighting = ResolvedDiffuseLighting[ScreenCoord].xyz;
 			SpecularLighting = ResolvedSpecularLighting[ScreenCoord].xyz;
 
@@ -165,11 +198,10 @@ void SDLTemporalAccumulationCS(
 						SpecularLightingAndSecondMomentHistory = SpecularLightingAndSecondMomentHistory / dot(HistoryWeights, 1.0f);
 
 						// Correct history for current frame exposure
-						DiffuseLightingAndSecondMomentHistory.xyz *= PrevSceneColorPreExposureCorrection;
-						SpecularLightingAndSecondMomentHistory.xyz *= PrevSceneColorPreExposureCorrection;
-
-						DiffuseLightingAndSecondMomentHistory.w *= Pow2(PrevSceneColorPreExposureCorrection);
-						SpecularLightingAndSecondMomentHistory.w *= Pow2(PrevSceneColorPreExposureCorrection);
+						float3 HistoryDiffuseLighting = DiffuseLightingAndSecondMomentHistory.xyz * PrevSceneColorPreExposureCorrection;
+						float3 HistorySpecularLighting = SpecularLightingAndSecondMomentHistory.xyz * PrevSceneColorPreExposureCorrection;
+						float HistoryDiffuseSecondMoment = DiffuseLightingAndSecondMomentHistory.w * Pow2(PrevSceneColorPreExposureCorrection);
+						float HistorySpecularSecondMoment = SpecularLightingAndSecondMomentHistory.w * Pow2(PrevSceneColorPreExposureCorrection);
 
 						// Reproject and rescale normalized NumFramesAccumulated
 						float NumFramesAccumulatedHistory = 0.0f;
@@ -177,24 +209,24 @@ void SDLTemporalAccumulationCS(
 						NumFramesAccumulatedHistory = NumFramesAccumulatedHistory / dot(HistoryWeights, 1.0f);
 						NumFramesAccumulated = min(NumFramesAccumulatedHistory + 1.0f, TemporalMaxFramesAccumulated);
 
-						float Alpha = 1.0f / (1.0f + NumFramesAccumulated);
-						DiffuseLighting = lerp(DiffuseLightingAndSecondMomentHistory.xyz, DiffuseLighting, Alpha);
-						SpecularLighting = lerp(SpecularLightingAndSecondMomentHistory.xyz, SpecularLighting, Alpha);
-						DiffuseSecondMoment = lerp(DiffuseLightingAndSecondMomentHistory.w, DiffuseSecondMoment, Alpha);
-						SpecularSecondMoment = lerp(SpecularLightingAndSecondMomentHistory.w, SpecularSecondMoment, Alpha);
-
-						FNeighborhoodClamp DiffuseNeighborhoodClamp = GetNeighborhoodClamp(ScreenCoord, ResolvedDiffuseLighting, DiffuseLighting);
-						FNeighborhoodClamp SpecularNeighborhoodClamp = GetNeighborhoodClamp(ScreenCoord, ResolvedSpecularLighting, SpecularLighting);
-
-						float3 ClampedDiffuseLighting = clamp(DiffuseLighting, DiffuseNeighborhoodClamp.Min, DiffuseNeighborhoodClamp.Max);
-						float3 ClampedSpecularLighting = clamp(SpecularLighting, SpecularNeighborhoodClamp.Min, SpecularNeighborhoodClamp.Max);
+						// Clamp history to current neighborhood
+						FNeighborhood DiffuseNeighborhood = GetNeighborhood(ScreenCoord, ResolvedDiffuseLighting, float4(DiffuseLighting, DiffuseSecondMoment));
+						FNeighborhood SpecularNeighborhood = GetNeighborhood(ScreenCoord, ResolvedSpecularLighting, float4(SpecularLighting, SpecularSecondMoment));
+						float3 ClampedHistoryDiffuseLighting = clamp(HistoryDiffuseLighting, DiffuseNeighborhood.ClampMin, DiffuseNeighborhood.ClampMax);
+						float3 ClampedHistorySpecularLighting = clamp(HistorySpecularLighting, SpecularNeighborhood.ClampMin, SpecularNeighborhood.ClampMax);
 
-						// Approximate clamp for second moments in order to reduce ghosting from the spatial pass while still allowing accumulated secondary moments to deviate from accumulated lighting
-						DiffuseSecondMoment += Pow2(Luminance(ClampedDiffuseLighting)) - Pow2(Luminance(DiffuseLighting));
-						SpecularSecondMoment += Pow2(Luminance(ClampedSpecularLighting)) - Pow2(Luminance(SpecularLighting));
+						// Clamp history to max sample luminance
+						ClampedHistoryDiffuseLighting = ClampLuminance(ClampedHistoryDiffuseLighting, SampleLuminanceSum.x);
+						ClampedHistorySpecularLighting = ClampLuminance(ClampedHistorySpecularLighting, SampleLuminanceSum.y);
+						float ClampedHistoryDiffuseSecondMoment = clamp(ClampedHistoryDiffuseSecondMoment, 0, Pow2(SampleLuminanceSum.x));
+						float ClampedHistorySpecularSecondMoment = clamp(ClampedHistorySpecularSecondMoment, 0, Pow2(SampleLuminanceSum.y));
 
-						DiffuseLighting = ClampedDiffuseLighting;
-						SpecularLighting = ClampedSpecularLighting;
+						// Blend history with new samples
+						float Alpha = 1.0f / (1.0f + NumFramesAccumulated);
+						DiffuseLighting = lerp(ClampedHistoryDiffuseLighting, DiffuseLighting, Alpha);
+						SpecularLighting = lerp(ClampedHistorySpecularLighting, SpecularLighting, Alpha);
+						DiffuseSecondMoment = lerp(ClampedHistoryDiffuseSecondMoment, DiffuseSecondMoment, Alpha);
+						SpecularSecondMoment = lerp(ClampedHistorySpecularSecondMoment, SpecularSecondMoment, Alpha);
 					}
 				}
 			}
@@ -218,6 +250,13 @@ void SDLTemporalAccumulationCS(
 				Newline(Context);
 				Print(Context, TEXT("SpecularVariance: "));
 				Print(Context, SpecularVariance);
+				Newline(Context);
+				Print(Context, TEXT("SampleLuminanceSum: "));
+				Print(Context, SampleLuminanceSum);
+				Newline(Context);
+				Print(Context, TEXT("Luminance:          "));
+				Print(Context, Luminance(DiffuseLighting));
+				Print(Context, Luminance(SpecularLighting));
 			}
 			#endif
 		}
diff --git a/Engine/Source/Runtime/Renderer/Private/StochasticDirectLighting/StochasticDirectLighting.cpp b/Engine/Source/Runtime/Renderer/Private/StochasticDirectLighting/StochasticDirectLighting.cpp
--- a/Engine/Source/Runtime/Renderer/Private/StochasticDirectLighting/StochasticDirectLighting.cpp
+++ b/Engine/Source/Runtime/Renderer/Private/StochasticDirectLighting/StochasticDirectLighting.cpp
@@ -121,13 +121,6 @@ static TAutoConsoleVariable<int> CVarStochasticDirectLightingFixedStateFrameInde
 	ECVF_RenderThreadSafe
 );
 
-static TAutoConsoleVariable<int> CVarStochasticDirectLightingCandidateLightMask(
-	TEXT("r.StochasticDirectLighting.CandidateLightMask"),
-	1,
-	TEXT("#sdl_todo: finish and pick one shader path."),
-	ECVF_Scalability | ECVF_RenderThreadSafe
-);
-
 static TAutoConsoleVariable<int> CVarStochasticDirectLightingTexturedRectLights(
 	TEXT("r.StochasticDirectLighting.TexturedRectLights"),
 	0,
@@ -331,14 +324,15 @@ class FInitCompositeIndirectArgsCS : public FGlobalShader
 
 IMPLEMENT_GLOBAL_SHADER(FInitCompositeIndirectArgsCS, "/Engine/Private/StochasticDirectLighting/StochasticDirectLighting.usf", "InitCompositeIndirectArgsCS", SF_Compute);
 
-class FGenerateSamplesCS : public FGlobalShader
+class FGenerateLightSamplesCS : public FGlobalShader
 {
-	DECLARE_GLOBAL_SHADER(FGenerateSamplesCS)
-	SHADER_USE_PARAMETER_STRUCT(FGenerateSamplesCS, FGlobalShader)
+	DECLARE_GLOBAL_SHADER(FGenerateLightSamplesCS)
+	SHADER_USE_PARAMETER_STRUCT(FGenerateLightSamplesCS, FGlobalShader)
 
 	BEGIN_SHADER_PARAMETER_STRUCT(FParameters, )
 		RDG_BUFFER_ACCESS(IndirectArgs, ERHIAccess::IndirectArgs)
 		SHADER_PARAMETER_STRUCT_INCLUDE(FStochasticDirectLightingParameters, StochasticDirectLightingParameters)
+		SHADER_PARAMETER_RDG_TEXTURE_UAV(RWTexture2D<float2>, RWSampleLuminanceSum)
 		SHADER_PARAMETER_RDG_TEXTURE_UAV(RWTexture2D<float>, RWDownsampledSceneDepth)
 		SHADER_PARAMETER_RDG_TEXTURE_UAV(RWTexture2D<UNORM float3>, RWDownsampledSceneWorldNormal)
 		SHADER_PARAMETER_RDG_BUFFER_UAV(RWStructuredBuffer<uint>, RWCompositeTileAllocator)
@@ -355,9 +349,9 @@ class FGenerateSamplesCS : public FGlobalShader
 	class FLightFunctionAtlas : SHADER_PERMUTATION_BOOL("USE_LIGHT_FUNCTION_ATLAS");
 	class FTexturedRectLights : SHADER_PERMUTATION_BOOL("USE_SOURCE_TEXTURE");
 	class FNumSamplesPerPixel1d : SHADER_PERMUTATION_SPARSE_INT("NUM_SAMPLES_PER_PIXEL_1D", 1, 2, 4);
-	class FCandidateLightMask : SHADER_PERMUTATION_BOOL("CANDIDATE_LIGHT_MASK");
+
 	class FDebugMode : SHADER_PERMUTATION_BOOL("DEBUG_MODE");
-	using FPermutationDomain = TShaderPermutationDomain<FTileType, FIESProfile, FLightFunctionAtlas, FTexturedRectLights, FNumSamplesPerPixel1d, FCandidateLightMask, FDebugMode>;
+	using FPermutationDomain = TShaderPermutationDomain<FTileType, FIESProfile, FLightFunctionAtlas, FTexturedRectLights, FNumSamplesPerPixel1d, FDebugMode>;
 
 	static int32 GetGroupSize()
 	{	
@@ -377,7 +371,7 @@ class FGenerateSamplesCS : public FGlobalShader
 	}
 };
 
-IMPLEMENT_GLOBAL_SHADER(FGenerateSamplesCS, "/Engine/Private/StochasticDirectLighting/StochasticDirectLightingSampling.usf", "GenerateSamplesCS", SF_Compute);
+IMPLEMENT_GLOBAL_SHADER(FGenerateLightSamplesCS, "/Engine/Private/StochasticDirectLighting/StochasticDirectLightingSampling.usf", "GenerateLightSamplesCS", SF_Compute);
 
 class FClearLightSamplesCS : public FGlobalShader
 {
@@ -540,6 +534,8 @@ class FSDLTemporalAccumulationCS : public FGlobalShader
 
 	BEGIN_SHADER_PARAMETER_STRUCT(FParameters, )
 		SHADER_PARAMETER_STRUCT_INCLUDE(FStochasticDirectLightingParameters, StochasticDirectLightingParameters)
+		SHADER_PARAMETER_RDG_TEXTURE(Texture2D<float4>, CompositeUpsampleWeights)
+		SHADER_PARAMETER_RDG_TEXTURE(Texture2D<float2>, SampleLuminanceSumTexture)
 		SHADER_PARAMETER_RDG_TEXTURE(Texture2D<float4>, ResolvedDiffuseLighting)
 		SHADER_PARAMETER_RDG_TEXTURE(Texture2D<float4>, ResolvedSpecularLighting)
 		SHADER_PARAMETER_RDG_TEXTURE(Texture2D<float4>, DiffuseLightingAndSecondMomentHistoryTexture)
@@ -870,9 +866,15 @@ void FDeferredShadingSceneRenderer::RenderStochasticDirectLighting(FRDGBuilder&
 	FRDGBufferRef CompositeTileAllocator = GraphBuilder.CreateBuffer(FRDGBufferDesc::CreateStructuredDesc(sizeof(uint32), 1), TEXT("StochasticDirectLighting.CompositeTileAllocator"));
 	FRDGBufferRef CompositeTileData = GraphBuilder.CreateBuffer(FRDGBufferDesc::CreateStructuredDesc(2 * sizeof(uint32), MaxCompositeTiles), TEXT("StochasticDirectLighting.CompositeTileData"));
 	AddClearUAVPass(GraphBuilder, GraphBuilder.CreateUAV(CompositeTileAllocator), 0);
+
+
+	FRDGTextureRef SampleLuminanceSum = GraphBuilder.CreateTexture(
+		FRDGTextureDesc::Create2D(DownsampledBufferSize, PF_G16R16F, FClearValueBinding::Black, TexCreate_ShaderResource | TexCreate_UAV),
+		TEXT("StochasticDirectLighting.SampleLuminanceSum"));
 	
 	// Generate new candidate light samples
 	{
+		FRDGTextureUAVRef SampleLuminanceSumUAV = GraphBuilder.CreateUAV(SampleLuminanceSum, ERDGUnorderedAccessViewFlags::SkipBarrier);
 		FRDGTextureUAVRef DownsampledSceneDepthUAV = GraphBuilder.CreateUAV(DownsampledSceneDepth, ERDGUnorderedAccessViewFlags::SkipBarrier);
 		FRDGTextureUAVRef DownsampledSceneWorldNormalUAV = GraphBuilder.CreateUAV(DownsampledSceneWorldNormal, ERDGUnorderedAccessViewFlags::SkipBarrier);
 		FRDGBufferUAVRef CompositeTileAllocatorUAV = GraphBuilder.CreateUAV(CompositeTileAllocator, ERDGUnorderedAccessViewFlags::SkipBarrier);
@@ -905,9 +907,10 @@ void FDeferredShadingSceneRenderer::RenderStochasticDirectLighting(FRDGBuilder&
 
 		for (int32 TileType = 0; TileType < (int32)StochasticDirectLighting::ETileType::SHADING_MAX; ++TileType)
 		{
-			FGenerateSamplesCS::FParameters* PassParameters = GraphBuilder.AllocParameters<FGenerateSamplesCS::FParameters>();
+			FGenerateLightSamplesCS::FParameters* PassParameters = GraphBuilder.AllocParameters<FGenerateLightSamplesCS::FParameters>();
 			PassParameters->IndirectArgs = DownsampledTileIndirectArgs;
 			PassParameters->StochasticDirectLightingParameters = StochasticDirectLightingParameters;
+			PassParameters->RWSampleLuminanceSum = SampleLuminanceSumUAV;
 			PassParameters->RWDownsampledSceneDepth = DownsampledSceneDepthUAV;
 			PassParameters->RWDownsampledSceneWorldNormal = DownsampledSceneWorldNormalUAV;
 			PassParameters->RWLightSamples = LightSamplesUAV;
@@ -918,15 +921,14 @@ void FDeferredShadingSceneRenderer::RenderStochasticDirectLighting(FRDGBuilder&
 			PassParameters->HistoryScreenPositionScaleBias = HistoryScreenPositionScaleBias;
 			PassParameters->HistoryUVMinMax = HistoryUVMinMax;
 
-			FGenerateSamplesCS::FPermutationDomain PermutationVector;
-			PermutationVector.Set<FGenerateSamplesCS::FTileType>(TileType);
-			PermutationVector.Set<FGenerateSamplesCS::FIESProfile>(CVarStochasticDirectLightingIESProfiles.GetValueOnRenderThread() != 0);
-			PermutationVector.Set<FGenerateSamplesCS::FLightFunctionAtlas>(bUseLightFunctionAtlas);
-			PermutationVector.Set<FGenerateSamplesCS::FTexturedRectLights>(CVarStochasticDirectLightingTexturedRectLights.GetValueOnRenderThread() != 0);
-			PermutationVector.Set<FGenerateSamplesCS::FNumSamplesPerPixel1d>(NumSamplesPerPixel2d.X * NumSamplesPerPixel2d.Y);
-			PermutationVector.Set<FGenerateSamplesCS::FDebugMode>(bDebug);
-			PermutationVector.Set<FGenerateSamplesCS::FCandidateLightMask>(CVarStochasticDirectLightingCandidateLightMask.GetValueOnRenderThread() != 0);
-			auto ComputeShader = View.ShaderMap->GetShader<FGenerateSamplesCS>(PermutationVector);
+			FGenerateLightSamplesCS::FPermutationDomain PermutationVector;
+			PermutationVector.Set<FGenerateLightSamplesCS::FTileType>(TileType);
+			PermutationVector.Set<FGenerateLightSamplesCS::FIESProfile>(CVarStochasticDirectLightingIESProfiles.GetValueOnRenderThread() != 0);
+			PermutationVector.Set<FGenerateLightSamplesCS::FLightFunctionAtlas>(bUseLightFunctionAtlas);
+			PermutationVector.Set<FGenerateLightSamplesCS::FTexturedRectLights>(CVarStochasticDirectLightingTexturedRectLights.GetValueOnRenderThread() != 0);
+			PermutationVector.Set<FGenerateLightSamplesCS::FNumSamplesPerPixel1d>(NumSamplesPerPixel2d.X * NumSamplesPerPixel2d.Y);
+			PermutationVector.Set<FGenerateLightSamplesCS::FDebugMode>(bDebug);
+			auto ComputeShader = View.ShaderMap->GetShader<FGenerateLightSamplesCS>(PermutationVector);
 
 			FComputeShaderUtils::AddPass(
 				GraphBuilder,
@@ -1102,6 +1104,8 @@ void FDeferredShadingSceneRenderer::RenderStochasticDirectLighting(FRDGBuilder&
 	{
 		FSDLTemporalAccumulationCS::FParameters* PassParameters = GraphBuilder.AllocParameters<FSDLTemporalAccumulationCS::FParameters>();
 		PassParameters->StochasticDirectLightingParameters = StochasticDirectLightingParameters;
+		PassParameters->CompositeUpsampleWeights = CompositeUpsampleWeights;
+		PassParameters->SampleLuminanceSumTexture = SampleLuminanceSum;
 		PassParameters->ResolvedDiffuseLighting = ResolvedDiffuseLighting;
 		PassParameters->ResolvedSpecularLighting = ResolvedSpecularLighting;
 		PassParameters->DiffuseLightingAndSecondMomentHistoryTexture = DiffuseLightingAndSecondMomentHistory;
