Add missing overloads for vectors to FMA

diff --git a/Engine/Shaders/Private/DoubleFloatOperations.ush b/Engine/Shaders/Private/DoubleFloatOperations.ush
--- a/Engine/Shaders/Private/DoubleFloatOperations.ush
+++ b/Engine/Shaders/Private/DoubleFloatOperations.ush
@@ -368,7 +368,7 @@ FDFType DFDivide(FDFType Lhs, FDFType Rhs)
 	return DFFastDivide(Lhs, Rhs);
 #else
 	const FFloatType Th = INVARIANT(1.0 / Rhs.High);
-	const FFloatType Rh = INVARIANT(FMA(-Rhs.High, Th, 1));
+	const FFloatType Rh = INVARIANT(FMA(-Rhs.High, Th, (FFloatType)1));
 	const FFloatType Rl = INVARIANT(-Rhs.Low * Th);
 	const FDFType E = DFFastTwoSum(Rh, Rl);
 	const FDFType D = DFMultiply(E, Th);
@@ -662,9 +662,9 @@ FDFType DFSqrt(FDFType V)
 
 	// equivalent to (remark 3.5)
 	//precise FFloatType Tl = P2 / Sh;
-	//precise FFloatType Zh = FMA(0.5, Tl, Sh);
+	//precise FFloatType Zh = FMA((FFloatType)0.5, Tl, Sh);
 	//precise FFloatType D = Zh - Sh;
-	//precise FFloatType Zl = FMA(0.5, Tl, -D);
+	//precise FFloatType Zl = FMA((FFloatType)0.5, Tl, -D);
 	//return DFConstructor(Zh, Zl);
 #endif
 }
diff --git a/Engine/Shaders/Public/ShaderVersion.ush b/Engine/Shaders/Public/ShaderVersion.ush
--- a/Engine/Shaders/Public/ShaderVersion.ush
+++ b/Engine/Shaders/Public/ShaderVersion.ush
@@ -3,4 +3,4 @@
 // Each time the console command is executed it generates a new GUID. As this file is included
 // in Platform.ush (which should be included in any shader) it allows to invalidate the shader DDC.
 // 
-#pragma message("UESHADERMETADATA_VERSION 3c752777-4213-49d1-9357-e1b47a284227")
+#pragma message("UESHADERMETADATA_VERSION F632B8CA-8FAA-4BBC-B489-1618212BA5F0")
