Fix minor shader warning

diff --git a/Engine/Shaders/Private/StochasticDirectLighting/StochasticDirectLightingTemporal.usf b/Engine/Shaders/Private/StochasticDirectLighting/StochasticDirectLightingTemporal.usf
--- a/Engine/Shaders/Private/StochasticDirectLighting/StochasticDirectLightingTemporal.usf
+++ b/Engine/Shaders/Private/StochasticDirectLighting/StochasticDirectLightingTemporal.usf
@@ -210,8 +210,8 @@ void SDLTemporalAccumulationCS(
 						NumFramesAccumulated = min(NumFramesAccumulatedHistory + 1.0f, TemporalMaxFramesAccumulated);
 
 						// Clamp history to current neighborhood
-						FNeighborhood DiffuseNeighborhood = GetNeighborhood(ScreenCoord, ResolvedDiffuseLighting, float4(DiffuseLighting, DiffuseSecondMoment));
-						FNeighborhood SpecularNeighborhood = GetNeighborhood(ScreenCoord, ResolvedSpecularLighting, float4(SpecularLighting, SpecularSecondMoment));
+						FNeighborhood DiffuseNeighborhood = GetNeighborhood(ScreenCoord, ResolvedDiffuseLighting, float4(DiffuseLighting, DiffuseSecondMoment).xyz);
+						FNeighborhood SpecularNeighborhood = GetNeighborhood(ScreenCoord, ResolvedSpecularLighting, float4(SpecularLighting, SpecularSecondMoment).xyz);
 						float3 ClampedHistoryDiffuseLighting = clamp(HistoryDiffuseLighting, DiffuseNeighborhood.ClampMin, DiffuseNeighborhood.ClampMax);
 						float3 ClampedHistorySpecularLighting = clamp(HistorySpecularLighting, SpecularNeighborhood.ClampMin, SpecularNeighborhood.ClampMax);
 
