Path Tracing: Refactor how shadows are traced

Go back to using AHS instead of CHS for colored shadows in the path tracer. To avoid artifacts from the BVH reporting multiple hits against the same primitive, simply use the payload to check if the previous hit is at exactly the same distance.

This assumes that redundant shader invocations happen sequentially, which seems like a reasonable assumption and does indeed seem to work.

The same change has been applied to GPULightmass as well.

diff --git a/Engine/Plugins/Experimental/GPULightmass/Shaders/Private/LightmapGBuffer.usf b/Engine/Plugins/Experimental/GPULightmass/Shaders/Private/LightmapGBuffer.usf
--- a/Engine/Plugins/Experimental/GPULightmass/Shaders/Private/LightmapGBuffer.usf
+++ b/Engine/Plugins/Experimental/GPULightmass/Shaders/Private/LightmapGBuffer.usf
@@ -78,7 +78,7 @@ void LightmapGBufferPS(
 	FPixelMaterialInputs PixelMaterialInputs;
 	// Support WPO in the far future
 	// CalcMaterialParametersEx(MaterialParameters, PixelMaterialInputs, float4(0, 0, 0, 0), float4(0, 0, 0, 0), bIsFrontFace, Input.SavedWorldPositionWithShaderOffsets, Input.SavedWorldPosition);
-	CalcMaterialParametersEx(MaterialParameters, PixelMaterialInputs, float4(0, 0, 0, 0), float4(0, 0, 0, 0), bIsFrontFace, Input.SavedWorldPosition, Input.SavedWorldPosition);
+	CalcMaterialParametersEx(MaterialParameters, PixelMaterialInputs, float4(0, 0, 0, 0), float4(0, 0, 0, 0), bIsFrontFace, Input.SavedWorldPosition.xyz, Input.SavedWorldPosition.xyz);
 	
 	float4 OutWorldPosition = float4(LWCHackToFloat(MaterialParameters.AbsoluteWorldPosition) * WorldPositionScalar, 1.0f);
 	// Faceted triangle normal
diff --git a/Engine/Plugins/Experimental/GPULightmass/Shaders/Private/LightmapPathTracing.usf b/Engine/Plugins/Experimental/GPULightmass/Shaders/Private/LightmapPathTracing.usf
--- a/Engine/Plugins/Experimental/GPULightmass/Shaders/Private/LightmapPathTracing.usf
+++ b/Engine/Plugins/Experimental/GPULightmass/Shaders/Private/LightmapPathTracing.usf
@@ -293,40 +293,27 @@ FPathTracingPayload TraceTransparentRay(FRayDesc Ray, bool IsCameraRay, bool Las
 
 float3 TraceTransparentVisibilityRay(FRayDesc Ray, float PathRoughness)
 {
-	const uint RayFlags = 0;
+	const uint RayFlags = RAY_FLAG_ACCEPT_FIRST_HIT_AND_END_SEARCH | RAY_FLAG_SKIP_CLOSEST_HIT_SHADER;
 	const uint InstanceInclusionMask = PATHTRACER_MASK_SHADOW;
 	const uint RayContributionToHitGroupIndex = RAY_TRACING_SHADER_SLOT_SHADOW;
 	const uint MultiplierForGeometryContributionToShaderIndex = RAY_TRACING_NUM_SHADER_SLOTS;
 	const uint MissShaderIndex = 0;
 
 	FPackedPathTracingPayload PackedPayload = InitPathTracingVisibilityPayload(PathRoughness);
-
-	for (;;)
+	TraceRay(
+		TLAS,
+		RayFlags,
+		InstanceInclusionMask,
+		RayContributionToHitGroupIndex,
+		MultiplierForGeometryContributionToShaderIndex,
+		MissShaderIndex,
+		Ray.GetNativeDesc(),
+		PackedPayload);
+	if (PackedPayload.IsHit())
 	{
-		// Trace the ray
-		PackedPayload.SetMiss();
-		TraceRay(
-			TLAS,
-			RayFlags,
-			InstanceInclusionMask,
-			RayContributionToHitGroupIndex,
-			MultiplierForGeometryContributionToShaderIndex,
-			MissShaderIndex,
-			Ray.GetNativeDesc(),
-			PackedPayload);
-		if (PackedPayload.IsMiss())
-		{
-			break;
-		}
-		if (!any(PackedPayload.GetRayThroughput() > 0))
-		{
-			// accumulated to fully opacity already
-			break;
-		}
-		// one past the current hit to get the next transparent hit
-		Ray.TMin = asfloat(asuint(PackedPayload.HitT) + 1);
+		// We didn't run the miss shader, therefore we must have hit something opaque (or reached full opacity)
+		return 0.0;
 	}
-
 	return PackedPayload.GetRayThroughput();
 }
 
diff --git a/Engine/Plugins/Experimental/GPULightmass/Shaders/Private/VolumetricLightmapVoxelization.usf b/Engine/Plugins/Experimental/GPULightmass/Shaders/Private/VolumetricLightmapVoxelization.usf
--- a/Engine/Plugins/Experimental/GPULightmass/Shaders/Private/VolumetricLightmapVoxelization.usf
+++ b/Engine/Plugins/Experimental/GPULightmass/Shaders/Private/VolumetricLightmapVoxelization.usf
@@ -28,7 +28,7 @@ void VLMVoxelizationVS(
 
 	FVertexFactoryIntermediates VFIntermediates = GetVertexFactoryIntermediates(Input);
 	Output.WorldPosition = VertexFactoryGetWorldPosition(Input, VFIntermediates) - float4(LWCHackToFloat(ResolvedView.PreViewTranslation), 0.0f);
-	Output.WorldPosition = float4((Output.WorldPosition.xyz - VLMVoxelizationParams.VolumeCenter) / max(VLMVoxelizationParams.VolumeExtent, float4(0.0001, 0.0001, 0.0001, 0)), 1.0f);
+	Output.WorldPosition = float4((Output.WorldPosition.xyz - VLMVoxelizationParams.VolumeCenter.xyz) / max(VLMVoxelizationParams.VolumeExtent.xyz, float3(0.0001, 0.0001, 0.0001)), 1.0f);
 }
 
 struct FVLMVoxelizationVertex
@@ -123,12 +123,14 @@ void VLMVoxelizationGS(triangle FVLMVoxelizationVSToGS Inputs[3], inout Triangle
 	
 	float MaxProj[3] = {0, 0, 0};
 	float3 EdgeNormal[3];
-	for (int i = 0; i < 3; i++)
 	{
-		EdgeNormal[i] = cross(Edge[i], float3(0,0, TriangleNormal[DominantAxis] > 0 ? 1 : -1)); // Unnormalized, doesn't matter
-		for (int Dir = 0; Dir < 4; Dir++)
+		for (int i = 0; i < 3; i++)
 		{
-			MaxProj[i] = max(dot(SemiDiagonalVectors[Dir], EdgeNormal[i]), MaxProj[i]);
+			EdgeNormal[i] = cross(Edge[i], float3(0, 0, TriangleNormal[DominantAxis] > 0 ? 1 : -1)); // Unnormalized, doesn't matter
+			for (int Dir = 0; Dir < 4; Dir++)
+			{
+				MaxProj[i] = max(dot(SemiDiagonalVectors[Dir], EdgeNormal[i]), MaxProj[i]);
+			}
 		}
 	}
 
@@ -137,12 +139,14 @@ void VLMVoxelizationGS(triangle FVLMVoxelizationVSToGS Inputs[3], inout Triangle
 	Vertices[2].ScreenPosition.xyz += ( MaxProj[2] * Edge[1]/dot(Edge[1].xy,EdgeNormal[2].xy) + MaxProj[1] * Edge[2]/dot(Edge[2].xy,EdgeNormal[1].xy) );
 #endif
 
-	for (int i = 0; i < 3; i++)
 	{
-		Vertices[i].ScreenPosition.z += 1;
-		Vertices[i].ScreenPosition.z /= 2;
+		for (int i = 0; i < 3; i++)
+		{
+			Vertices[i].ScreenPosition.z += 1;
+			Vertices[i].ScreenPosition.z /= 2;
 	
-		OutStream.Append(Vertices[i]);
+			OutStream.Append(Vertices[i]);
+		}
 	}
 }
 
diff --git a/Engine/Shaders/Private/PathTracing/PathTracingCommon.ush b/Engine/Shaders/Private/PathTracing/PathTracingCommon.ush
--- a/Engine/Shaders/Private/PathTracing/PathTracingCommon.ush
+++ b/Engine/Shaders/Private/PathTracing/PathTracingCommon.ush
@@ -598,7 +598,11 @@ FPackedPathTracingPayload InitPathTracingVisibilityPayload(float PathRoughness)
 	Output.SetVisibilityRay();
 	Output.SetPathRoughness(PathRoughness);
 	Output.SetRayThroughput(1.0);
-	Output.SetMiss();
+	// NOTE: We start with HitT = 0 which counts as "IsHit()". This will be changed by the miss shader if we don't register any hits (or the AHS writes don't accumulate all the way to 0)
+	// This is because we want to be able to skip the CHS on shadow rays. Therefore the miss shader is the one that will be responsible for telling us a ray did not hit anything opaque.
+	// If after tracing the ray we still have IsHit() returning true, this means we must have hit something opaque (without having run either CHS or AHS).
+	// This also has the benefit of preparing the "mailbox" used for detecting duplicate invocations of AHS with an invalid value.
+	Output.HitT = 0.0;
 	return Output;
 }
 
diff --git a/Engine/Shaders/Private/PathTracing/PathTracingCore.ush b/Engine/Shaders/Private/PathTracing/PathTracingCore.ush
--- a/Engine/Shaders/Private/PathTracing/PathTracingCore.ush
+++ b/Engine/Shaders/Private/PathTracing/PathTracingCore.ush
@@ -684,9 +684,7 @@ FPathTracingPayload TraceTransparentRay(inout FPathState PathState, int Bounce,
 	return Payload;
 }
 
-#define PATH_TRACING_DEBUG_TRANSPARENCY_RAY 0
-
-float3 TraceTransparentVisibilityRay(FRayDesc Ray, int Bounce, float PathRoughness, uint MissShaderIndex, bool bCastShadows, inout RandomSequence RandSequence)
+float3 TraceTransparentVisibilityRay(FRayDesc Ray, float PathRoughness, uint MissShaderIndex, bool bCastShadows, inout RandomSequence RandSequence)
 {
 	if (!bCastShadows && MissShaderIndex == 0)
 	{
@@ -694,7 +692,7 @@ float3 TraceTransparentVisibilityRay(FRayDesc Ray, int Bounce, float PathRoughne
 		return 1.0;
 	}
 
-	const uint RayFlags = 0;
+	const uint RayFlags = RAY_FLAG_ACCEPT_FIRST_HIT_AND_END_SEARCH | RAY_FLAG_SKIP_CLOSEST_HIT_SHADER;
 	const uint InstanceInclusionMask = PATHTRACER_MASK_SHADOW | (length2(Ray.Origin) < Pow2(PATH_TRACING_SKIP_HAIR_DISTANCE) ? PATHTRACER_MASK_HAIR_SHADOW : 0);
 	const uint RayContributionToHitGroupIndex = RAY_TRACING_SHADER_SLOT_SHADOW;
 	const uint MultiplierForGeometryContributionToShaderIndex = RAY_TRACING_NUM_SHADER_SLOTS;
@@ -707,60 +705,26 @@ float3 TraceTransparentVisibilityRay(FRayDesc Ray, int Bounce, float PathRoughne
 		Ray.TMin = POSITIVE_INFINITY;
 		Ray.TMax = POSITIVE_INFINITY;
 	}
-
-	const float OrigTMin = Ray.TMin;
-
-	// Trace the ray - progressing in steps through transparent geometry
-#if PATH_TRACING_DEBUG_TRANSPARENCY_RAY
-	const int MaxSteps = 256; // avoid TDRs during debugging by capping how many times we go around
-	int Steps = 0;
-	for (int Counter = 0; Counter < MaxSteps; Counter++)
-#else
-	for (;;)
-#endif
+	TraceRay(
+		TLAS,
+		RayFlags,
+		InstanceInclusionMask,
+		RayContributionToHitGroupIndex,
+		MultiplierForGeometryContributionToShaderIndex,
+		MissShaderIndex,
+		Ray.GetNativeDesc(),
+		PackedPayload);
+
+	if (PackedPayload.IsHit())
 	{
-		// Reset payload to a miss for the rare case of a mesh without a valid CHS assigned
-		// Note that in this case we will get incorrect shadowing, but this is preferable to an infinite loop
-		PackedPayload.SetMiss();
-		TraceRay(
-			TLAS,
-			RayFlags,
-			InstanceInclusionMask,
-			RayContributionToHitGroupIndex,
-			MultiplierForGeometryContributionToShaderIndex,
-			MissShaderIndex,
-			Ray.GetNativeDesc(),
-			PackedPayload);
-
-		if (PackedPayload.IsMiss())
-		{
-			// we didn't hit anything, we are done
-			break;
-		}
-		if (!any(PackedPayload.GetRayThroughput() > 0))
-		{
-			// we've accumulated to full opacity, we are done
-			break;
-		}
-		// we hit some geometry, move one ulp past the hit and trace the same ray again
-		Ray.TMin = asfloat(asuint(PackedPayload.HitT) + 1);
-#if PATH_TRACING_DEBUG_TRANSPARENCY_RAY
-		Steps++;
-#endif
+		// We didn't run the miss shader, therefore we must have hit something opaque (or reached full opacity)
+		return 0.0;
 	}
 
-#if PATH_TRACING_DEBUG_TRANSPARENCY_RAY
-	if (Steps == 0) return 1.0;
-	if (Steps == 1) return float3(0, 0, 1);
-	if (Steps == 2) return float3(0, 1, 0);
-	if (Steps == MaxSteps) return float3(1, 0, 0);
-	const float t = float(Steps) / float(MaxSteps);
-	return t;
-#else
 	float3 Throughput = PackedPayload.GetRayThroughput();
 	if (bCastShadows && any(Throughput > 0))
 	{
-		FVolumeIntersectionList VolumeIsectList = VolumeIntersect(Ray.Origin, Ray.Direction, OrigTMin, Ray.TMax);
+		FVolumeIntersectionList VolumeIsectList = VolumeIntersect(Ray.Origin, Ray.Direction, Ray.TMin, Ray.TMax);
 		if (VolumeIsectList.HitBlocker())
 		{
 			// blockers are opaque -- done!
@@ -774,7 +738,6 @@ float3 TraceTransparentVisibilityRay(FRayDesc Ray, int Bounce, float PathRoughne
 		Throughput *= exp(-max(PackedPayload.GetTau(), 0.0));
 	}
 	return Throughput;
-#endif
 }
 
 struct FProbeResult
@@ -1347,9 +1310,9 @@ bool PathTracingKernel(inout FPathState PathState, int Bounce)
 	// This array will hold a CDF for light picking
 	float LightPickingCdf[RAY_TRACING_LIGHT_COUNT_MAXIMUM];
 
-#if PATH_TRACING_DEBUG_TRANSPARENCY_RAY
+#if 0
 	// visualize the shadow ray handling directly
-	PathState.Radiance = TraceTransparentVisibilityRay(PathState.Ray, Bounce, 0.0, 0, true, PathState.RandSequence);
+	PathState.Radiance = TraceTransparentVisibilityRay(PathState.Ray, 0.0, 0, true, PathState.RandSequence);
 	return false;
 #endif
 
@@ -1444,7 +1407,7 @@ bool PathTracingKernel(inout FPathState PathState, int Bounce)
 							{
 								LightContrib *= MISWeightRobust(LightSample.Pdf, MaterialEval.Pdf);
 							}
-							LightContrib *= TraceTransparentVisibilityRay(LightRay, Bounce, 1.0, MissShaderIndex, bCastShadows, PathState.RandSequence);
+							LightContrib *= TraceTransparentVisibilityRay(LightRay, 1.0, MissShaderIndex, bCastShadows, PathState.RandSequence);
 
 							AccumulateRadiance(PathState.Radiance, LightContrib, bIsCameraRay);
 						}
@@ -1470,7 +1433,7 @@ bool PathTracingKernel(inout FPathState PathState, int Bounce)
 								LightContrib *= MISWeightRobust(MaterialSample.Pdf, LightResult.Pdf);
 							}
 							MaterialRay.TMax = LightResult.HitT;
-							LightContrib *= TraceTransparentVisibilityRay(MaterialRay, Bounce, MaterialSample.Roughness, MissShaderIndex, bCastShadows, PathState.RandSequence);
+							LightContrib *= TraceTransparentVisibilityRay(MaterialRay, MaterialSample.Roughness, MissShaderIndex, bCastShadows, PathState.RandSequence);
 							AccumulateRadiance(PathState.Radiance, LightContrib, bIsCameraRay);
 						}
 					}
@@ -1611,7 +1574,7 @@ bool PathTracingKernel(inout FPathState PathState, int Bounce)
 					}
 					const bool bCastShadows = bIsVolumeSample ? CastsVolumeShadow(LightId) : CastsShadow(LightId);
 					const uint MissShaderIndex = GetLightMissShaderIndex(LightId);
-					LightContrib *= TraceTransparentVisibilityRay(LightRay, Bounce, AvgRoughness, MissShaderIndex, bCastShadows, PathState.RandSequence);
+					LightContrib *= TraceTransparentVisibilityRay(LightRay, AvgRoughness, MissShaderIndex, bCastShadows, PathState.RandSequence);
 
 					AccumulateRadiance(PathState.Radiance, LightContrib, bIsCameraRay);
 				}
@@ -1753,7 +1716,7 @@ bool PathTracingKernel(inout FPathState PathState, int Bounce)
 				LightRay.TMax = LightResult.HitT;
 				const bool bCastShadows = bIsVolumeSample ? CastsVolumeShadow(LightId) : CastsShadow(LightId);
 				const uint MissShaderIndex = GetLightMissShaderIndex(LightId);
-				LightContrib *= TraceTransparentVisibilityRay(LightRay, Bounce, ShadowRayRoughness, MissShaderIndex, bCastShadows, PathState.RandSequence);
+				LightContrib *= TraceTransparentVisibilityRay(LightRay, ShadowRayRoughness, MissShaderIndex, bCastShadows, PathState.RandSequence);
 				// the light made some contribution, and there was nothing along the shadow ray
 				AccumulateRadiance(PathState.Radiance, LightContrib, bIsCameraRay);
 			}
diff --git a/Engine/Shaders/Private/PathTracing/PathTracingDefaultHitShader.usf b/Engine/Shaders/Private/PathTracing/PathTracingDefaultHitShader.usf
--- a/Engine/Shaders/Private/PathTracing/PathTracingDefaultHitShader.usf
+++ b/Engine/Shaders/Private/PathTracing/PathTracingDefaultHitShader.usf
@@ -9,19 +9,11 @@ RAY_TRACING_ENTRY_CLOSEST_HIT(PathTracingDefaultOpaqueCHS,
 	FPackedPathTracingPayload, PackedPayload,
 	FRayTracingIntersectionAttributes, Attributes)
 {
-	if (PackedPayload.IsVisibilityRay())
-	{
-		// mark shadow fully opaque
-		PackedPayload.SetRayThroughput(0.0);
-	}
-	else
-	{
-		// mark hit as solid black
-		FPathTracingPayload Payload = (FPathTracingPayload)0;
-		Payload.ShadingModelID = SHADINGMODELID_UNLIT;
-		Payload.BSDFOpacity = 1.0;
-		PackedPayload = PackPathTracingPayload(Payload);
-	}
+	// mark hit as solid black
+	FPathTracingPayload Payload = (FPathTracingPayload)0;
+	Payload.ShadingModelID = SHADINGMODELID_UNLIT;
+	Payload.BSDFOpacity = 1.0;
+	PackedPayload = PackPathTracingPayload(Payload);
 	PackedPayload.HitT = RayTCurrent();
 }
 
@@ -29,19 +21,11 @@ RAY_TRACING_ENTRY_CLOSEST_HIT(PathTracingDefaultHiddenCHS,
 	FPackedPathTracingPayload, PackedPayload,
 	FRayTracingIntersectionAttributes, Attributes)
 {
-	// This is not reachable in theory since AHS calls IgnoreHit(), but could be if the mesh gets flagged with ForceOpaque (or the ray ignores AHS)
-	if (PackedPayload.IsVisibilityRay())
-	{
-		// don't change the throughput (object should be invisible)
-	}
-	else
-	{
-		// mark hit as fully transparent
-		FPathTracingPayload Payload = (FPathTracingPayload)0;
-		Payload.ShadingModelID = SHADINGMODELID_UNLIT;
-		Payload.TransparencyColor = 1.0;
-		PackedPayload = PackPathTracingPayload(Payload);
-	}
+	// mark hit as fully transparent
+	FPathTracingPayload Payload = (FPathTracingPayload)0;
+	Payload.ShadingModelID = SHADINGMODELID_UNLIT;
+	Payload.TransparencyColor = 1.0;
+	PackedPayload = PackPathTracingPayload(Payload);
 	PackedPayload.HitT = RayTCurrent();
 }
 
@@ -51,4 +35,3 @@ RAY_TRACING_ENTRY_ANY_HIT(PathTracingDefaultHiddenAHS,
 {
 	IgnoreHit();
 }
-
diff --git a/Engine/Shaders/Private/PathTracing/PathTracingGPULightmassMaterialHitShader.usf b/Engine/Shaders/Private/PathTracing/PathTracingGPULightmassMaterialHitShader.usf
--- a/Engine/Shaders/Private/PathTracing/PathTracingGPULightmassMaterialHitShader.usf
+++ b/Engine/Shaders/Private/PathTracing/PathTracingGPULightmassMaterialHitShader.usf
@@ -47,23 +47,6 @@ RAY_TRACING_ENTRY_CLOSEST_HIT(GPULightmassMaterialCHS,
 	FRayTracingIntersectionAttributes, Attributes)
 {
 	PackedPayload.HitT = RayTCurrent();
-#if MATERIALBLENDING_SOLID || MATERIALBLENDING_ALPHAHOLDOUT || (MATERIALBLENDING_MASKED && !MATERIAL_DITHER_OPACITY_MASK)
-	if (PackedPayload.IsVisibilityRay())
-	{
-		// Shadow ray case -- all these blend modes get treated as fully opaque, so exit early.
-		// In the case of Masked materials without dithered opacity, the masking logic was already handled by AHS.
-		PackedPayload.SetRayThroughput(0.0);
-		return;
-	}
-#endif
-#if MATERIALBLENDING_ADDITIVE
-	if (PackedPayload.IsVisibilityRay())
-	{
-		// should be fully invisible to shadow rays
-		// this case should be caught by making the mesh invisible to shadows, but it doesn't cost too much to have this extra check here just in case that fails
-		return;
-	}
-#endif
 	ResolvedView = ResolveView();
 
 	const float3 TranslatedWorldPosition = TranslatedWorldRayOrigin() + RayTCurrent() * WorldRayDirection();
@@ -123,14 +106,6 @@ RAY_TRACING_ENTRY_CLOSEST_HIT(GPULightmassMaterialCHS,
 
 
 #if MATERIALBLENDING_ANY_TRANSLUCENT
-	if (PackedPayload.IsVisibilityRay())
-	{
-		// Update the ray throughput (it is packed simply into the payload since we don't need to carry any other information across hits)
-		float3 RayThroughput = PackedPayload.GetRayThroughput();
-		RayThroughput *= saturate((1.0f - Coverage) + Coverage * TransmittancePreCoverage);
-		PackedPayload.SetRayThroughput(RayThroughput);
-		return;
-	}
 	// SUBSTRATE_TODO: Make sure this is correct for all blend modes/glass cases
 	float3 TransparencyColor = saturate((1.0f - Coverage) + Coverage * TransmittancePreCoverage);
 #else
@@ -192,60 +167,7 @@ RAY_TRACING_ENTRY_CLOSEST_HIT(GPULightmassMaterialCHS,
 
 
 
-	// the following blend modes need to process shadows after having executed the material
-#if MATERIALBLENDING_MODULATE || (MATERIALBLENDING_MASKED && MATERIAL_DITHER_OPACITY_MASK) || MATERIALBLENDING_ALPHACOMPOSITE || MATERIALBLENDING_TRANSLUCENT
-	if (PackedPayload.IsVisibilityRay())
-	{
-#if MATERIALBLENDING_MODULATE
-		const float3 Transparency = GetMaterialEmissive(PixelMaterialInputs);
-#elif MATERIALBLENDING_ALPHACOMPOSITE
-		const float Opacity = GetMaterialOpacity(PixelMaterialInputs);
-		const float Transparency = 1 - Opacity;
-#elif MATERIALBLENDING_MASKED && MATERIAL_DITHER_OPACITY_MASK
-		// See MATERIAL_DITHER_OPACITY_MASK comment below
-		const float Opacity = saturate(GetMaterialMaskInputRaw(PixelMaterialInputs));
-		const float Transparency = 1 - Opacity;
-#elif MATERIALBLENDING_TRANSLUCENT
-		const float Opacity = GetMaterialOpacity(PixelMaterialInputs);
-		// GPULM has backwards compatibility with some cpu lightmass tricks
-		float3 Transparency = 1 - Opacity;
-		uint ShadingModelID = GetMaterialShadingModel(PixelMaterialInputs);
 
-#if MATERIAL_SHADINGMODEL_THIN_TRANSLUCENT
-		if (ShadingModelID == SHADINGMODELID_THIN_TRANSLUCENT)
-		{
-			float3 Transmission = GetThinTranslucentMaterialOutput0(MaterialParameters);
-			float Ior = 0.0;
-			float3 V = WorldRayDirection();
-			float3 N = normalize(MaterialParameters.WorldNormal);
-			float VoN = abs(dot(V, N));
-			// simplified logic with no bending at the interface and no fresnel
-			Transparency *= pow(Transmission, 1.0 / VoN);
-		}
-		else
-#endif // MATERIAL_SHADINGMODEL_THIN_TRANSLUCENT
-#if MATERIAL_SHADINGMODEL_DEFAULT_LIT
-		if (ShadingModelID == SHADINGMODELID_DEFAULT_LIT)
-		{
-			// CPU Lightmass translucency behavior emulation
-			Transparency = lerp(float3(1, 1, 1), GetMaterialBaseColor(PixelMaterialInputs), Opacity);
-		}
-		else
-#endif// MATERIAL_SHADINGMODEL_DEFAULT_LIT
-		{
-			// base case for shadingmodel if/else
-		}
-#else // MATERIALBLEINDG_*
-#error Unhandled blending mode!
-#endif
-
-		// Update the ray throughput (it is packed simply into the payload since we don't need to carry any other information across hits)
-		float3 RayThroughput = PackedPayload.GetRayThroughput();
-		RayThroughput *= Transparency;
-		PackedPayload.SetRayThroughput(RayThroughput);
-		return;
-	}
-#endif
 
 	FPathTracingPayload Payload = (FPathTracingPayload)0;
 
@@ -410,8 +332,20 @@ RAY_TRACING_ENTRY_ANY_HIT(GPULightmassMaterialAHS,
 	FPackedPathTracingPayload, PackedPayload,
 	FRayTracingIntersectionAttributes, Attributes)
 {
-#if MATERIALBLENDING_MASKED && !MATERIAL_DITHER_OPACITY_MASK
-	// This is the only case which actually needs to run the full material
+#if MATERIALBLENDING_MODULATE || (MATERIALBLENDING_MASKED && MATERIAL_DITHER_OPACITY_MASK) || MATERIALBLENDING_ALPHACOMPOSITE || MATERIALBLENDING_TRANSLUCENT
+	if (!PackedPayload.IsVisibilityRay())
+	{
+		// not a shadow ray -- don't need to run the AHS logic
+		return;
+	}
+	if (PackedPayload.HitT == RayTCurrent())
+	{
+		// We just processed this exact hit, don't double-count it
+		IgnoreHit();
+		return;
+	}
+#endif
+
 	ResolvedView = ResolveView();
 
 	const float3 TranslatedWorldPosition = TranslatedWorldRayOrigin() + RayTCurrent() * WorldRayDirection();
@@ -444,16 +378,103 @@ RAY_TRACING_ENTRY_ANY_HIT(GPULightmassMaterialAHS,
 		CalcMaterialParametersEx(MaterialParameters, PixelMaterialInputs, SvPosition, ScreenPosition, bIsFrontFace, TranslatedWorldPosition, TranslatedWorldPosition);
 	}
 
+#if MATERIALBLENDING_MASKED && !MATERIAL_DITHER_OPACITY_MASK
 	// Regardless of payload flags -- we always apply this 
 	if (GetMaterialMask(PixelMaterialInputs) < 0)
 	{
 		IgnoreHit();
 	}
+#endif
+
+#if SUBSTRATE_ENABLED
+
+	float  Coverage = 0;
+	float3 TransmittancePreCoverage = 0;
+	float3 BaseColorPostCoverage = 0;
+	float3 WorldNormal = 0;
+	float3 EmissiveLuminance = 0;
+
+#if MATERIAL_IS_SUBSTRATE
+	FSubstrateData SubstrateData = PixelMaterialInputs.GetFrontSubstrateData();
+	FSubstratePixelHeader SubstratePixelHeader = MaterialParameters.GetFrontSubstrateHeader();
+
+
+	FSubstrateIntegrationSettings Settings = InitSubstrateIntegrationSettings(true /*bForceFullyRough*/, false /*SubstrateStruct.bRoughDiffuse*/, -1/*SubstrateStruct.PeelLayersAboveDepth*/, false/*SubstrateStruct.bRoughnessTracking*/);
+
+	const float3 SurfaceWorldNormal = GeoNormal;
+	SubstrateMaterialExportOut(
+		Settings,
+		SubstratePixelHeader,
+		SubstrateData,
+		SurfaceWorldNormal,
+		Coverage,
+		TransmittancePreCoverage,
+		BaseColorPostCoverage,
+		WorldNormal,
+		EmissiveLuminance);
+
+	float3 Transparency = saturate((1.0f - Coverage) + Coverage * TransmittancePreCoverage);
 #else
-	// All other blending modes need to run CHS to decide what to do, or are handled via the default opaque/translucent CHS/AHS materials.
-	// in these cases, the AHS should be disabled, so emit an error if we get here somehow
-#error Unexpected blend mode encountered in AHS!
+#error "Shouldn't need to compile AHS for non-substrate material"
 #endif
+
+#else // SUBSTRATE_ENABLED
+
+	// the following blend modes need to process shadows after having executed the material
+#if MATERIALBLENDING_MODULATE
+	const float3 Transparency = GetMaterialEmissive(PixelMaterialInputs);
+#elif MATERIALBLENDING_ALPHACOMPOSITE
+	const float Opacity = GetMaterialOpacity(PixelMaterialInputs);
+	const float Transparency = 1 - Opacity;
+#elif MATERIALBLENDING_MASKED && MATERIAL_DITHER_OPACITY_MASK
+	// See MATERIAL_DITHER_OPACITY_MASK comment below
+	const float Opacity = saturate(GetMaterialMaskInputRaw(PixelMaterialInputs));
+	const float Transparency = 1 - Opacity;
+#elif MATERIALBLENDING_TRANSLUCENT
+	const float Opacity = GetMaterialOpacity(PixelMaterialInputs);
+	// GPULM has backwards compatibility with some cpu lightmass tricks
+	float3 Transparency = 1 - Opacity;
+	uint ShadingModelID = GetMaterialShadingModel(PixelMaterialInputs);
+
+#if MATERIAL_SHADINGMODEL_THIN_TRANSLUCENT
+	if (ShadingModelID == SHADINGMODELID_THIN_TRANSLUCENT)
+	{
+		float3 Transmission = GetThinTranslucentMaterialOutput0(MaterialParameters);
+		float Ior = 0.0;
+		float3 V = WorldRayDirection();
+		float3 N = normalize(MaterialParameters.WorldNormal);
+		float VoN = abs(dot(V, N));
+		// simplified logic with no bending at the interface and no fresnel
+		Transparency *= pow(Transmission, 1.0 / VoN);
+	}
+	else
+#endif // MATERIAL_SHADINGMODEL_THIN_TRANSLUCENT
+#if MATERIAL_SHADINGMODEL_DEFAULT_LIT
+	if (ShadingModelID == SHADINGMODELID_DEFAULT_LIT)
+	{
+		// CPU Lightmass translucency behavior emulation
+		Transparency = lerp(float3(1, 1, 1), GetMaterialBaseColor(PixelMaterialInputs), Opacity);
+	}
+	else
+#endif// MATERIAL_SHADINGMODEL_DEFAULT_LIT
+	{
+		// base case for shadingmodel if/else
+	}
+#else // MATERIALBLENDING_*
+#error Unhandled blending mode!
+#endif
+
+#endif // !SUBSTRATE_ENABLED
+
+	// Update the ray throughput (it is packed simply into the payload since we don't need to carry any other information across hits)
+	float3 RayThroughput = PackedPayload.GetRayThroughput();
+	RayThroughput *= Transparency;
+	PackedPayload.SetRayThroughput(RayThroughput);
+	if (any(RayThroughput > 0))
+	{
+		// if there is energy left in the ray, 
+		IgnoreHit();
+	}
 }
 
 #endif // USE_MATERIAL_ANY_HIT_SHADER
diff --git a/Engine/Shaders/Private/PathTracing/PathTracingMaterialHitShader.usf b/Engine/Shaders/Private/PathTracing/PathTracingMaterialHitShader.usf
--- a/Engine/Shaders/Private/PathTracing/PathTracingMaterialHitShader.usf
+++ b/Engine/Shaders/Private/PathTracing/PathTracingMaterialHitShader.usf
@@ -165,23 +165,6 @@
 	FRayTracingIntersectionAttributes, Attributes)
 {
 	PackedPayload.HitT = RayTCurrent();
-#if MATERIALBLENDING_SOLID || MATERIALBLENDING_ALPHAHOLDOUT || (MATERIALBLENDING_MASKED && !MATERIAL_DITHER_OPACITY_MASK) || SUBSTRATE_OPAQUE_MATERIAL
-	if (PackedPayload.IsVisibilityRay())
-	{
-		// Shadow ray case -- all these blend modes get treated as fully opaque, so exit early.
-		// In the case of Masked materials without dithered opacity, the masking logic was already handled by AHS.
-		PackedPayload.SetRayThroughput(0.0);
-		return;
-	}
-#endif
-#if MATERIALBLENDING_ADDITIVE
-	if (PackedPayload.IsVisibilityRay())
-	{
-		// should be fully invisible to shadow rays
-		// this case should be caught by making the mesh invisible to shadows, but it doesn't cost too much to have this extra check here just in case that fails
-		return;
-	}
-#endif
 	ResolvedView = ResolveView();
 
 	const float3 TranslatedWorldPosition = TranslatedWorldRayOrigin() + RayTCurrent() * WorldRayDirection();
@@ -276,203 +259,6 @@
 
 #endif // SUBSTRATE_ENABLED
 
-	// the following blend modes need to process shadows after having executed the material
-#if MATERIALBLENDING_MODULATE || (MATERIALBLENDING_MASKED && MATERIAL_DITHER_OPACITY_MASK) || MATERIALBLENDING_ALPHACOMPOSITE || MATERIALBLENDING_TRANSLUCENT
-	if (PackedPayload.IsVisibilityRay())
-	{
-
-#if SUBSTRATE_ENABLED
-
-#if MATERIAL_IS_SUBSTRATE && !SUBSTRATE_OPTIMIZED_UNLIT
-		// We already figured out how much Transparency there is above
-		// So only do extra logic for refraction and thin refraction ases
-#if MATERIALBLENDING_TRANSLUCENT && (REFRACTION_USE_INDEX_OF_REFRACTION || MATERIAL_ISTHINSURFACE)
-
-		// SUBSTRATE_TODO: refraction handling
-		float PathRoughness = PackedPayload.GetPathRoughness();
-
-		// SUBSTRATE_TODO: This does not match the logic in DistortAccumulatePS.usf
-		//  The idea here is to take F0 from the bottom refraction layers since this matches most closely what the refraction logic below does (only bottom slabs can refract).
-		//  The distortion approximation on the raster side instead averages F0 across layers which may produce different results.
-		//  In both cases - we get a potentially strange answer when there is more than one bottom slab that produces refraction
-
-		// SUBSTRATE_TODO: How should we handle MFP when there are multiple bottom slabs?
-		float F0 = 0.0f;
-		float3 RoughnessData = 0.0;
-		float NoV = 0.0f;
-		float3 LocalSigmaT = 0.0f;
-		SUBSTRATE_UNROLL_N(SUBSTRATE_CLAMPED_CLOSURE_COUNT)
-		for (int BSDFIdx = 0; BSDFIdx < SubstratePixelHeader.SubstrateTree.BSDFCount; ++BSDFIdx)
-		{
-#define BSDF SubstratePixelHeader.SubstrateTree.BSDFs[BSDFIdx]
-			if (BSDF.bIsBottom)
-			{
-				// SUBSTRATE_TODO: Also consider the MFP to see how much this slab is contributing to refraction ....
-				F0 = F0RGBToF0(SubstrateGetBSDFSpecularF0(BSDF));
-				RoughnessData.x = SLAB_ROUGHNESS(BSDF);
-				if (BSDF_GETHASHAZINESS(BSDF))
-				{
-					// NOTE: It should be safe to assume that the clearcoat case cannot happen here as it is only triggered by legacy clearcoat materials which are unlikely to have refraction
-					const FHaziness Haziness = UnpackHaziness(SLAB_HAZINESS(BSDF));
-					RoughnessData.y = Haziness.Roughness;
-					RoughnessData.z = Haziness.Weight;
-				}
-
-				float3 N = SubstratePixelHeader.SharedLocalBases.Normals[BSDF_GETSHAREDLOCALBASISID(BSDF)];
-				LocalSigmaT = rcp(max(SLAB_SSSMFP(BSDF), 0.0001));
-				NoV = dot(N, V_World);
-				break;
-			}
-#undef BSDF
-		}
-
-
-#if REFRACTION_USE_INDEX_OF_REFRACTION
-		float Ior = GetRefractionIor(PixelMaterialInputs, F0);
-#else
-		float Ior = 0.0;
-#endif
-
-#if MATERIAL_ISTHINSURFACE
-		// Thin Glass Fake Caustics
-
-		// SUBSTRATE_TODO: Revisit this -- TotalTransmittancePreCoverage should probably account for this in the thin case
-		float3 Transmission = exp(-LocalSigmaT * SUBSTRATE_SIMPLEVOLUME_THICKNESS_CM);
-		float3 Tr = ComputeThinSlabWeights(Transmission, NoV, Ior, F0).Transmitted;
-#else
-		HLSL_STATIC_ASSERT(REFRACTION_USE_INDEX_OF_REFRACTION == 1, "Expect refraction to be true here");
-		// Solid Glass Fake Caustics
-		float Tr = FakeCaustics(F0, Ior, NoV);
-		// Account for transmission 
-		PackedPayload.SetTau(PackedPayload.GetTau() + ComputeOpticalDepthForHit(LocalSigmaT, bIsFrontFace));
-#endif
-
-		// If material is doing refraction, apply an extra factor to account for ApproximateCaustics
-#if REFRACTION_USE_INDEX_OF_REFRACTION
-		if (Ior > 0.0)
-		{
-			Tr *= lerp(FakeCausticsPathFactor(RoughnessData.x, PathRoughness),
-					   FakeCausticsPathFactor(RoughnessData.y, PathRoughness), RoughnessData.z);
-		}
-#endif
-		Transparency += TotalCoverage * Tr;
-
-		// SUBSTRATE_TODO: Is this sanity check necessary?
-		Transparency = saturate(Transparency);
-#endif // MATERIALBLENDING_TRANSLUCENT && (REFRACTION_USE_INDEX_OF_REFRACTION || MATERIAL_ISTHINSURFACE)
-
-#endif // MATERIAL_IS_SUBSTRATE && !SUBSTRATE_OPTIMIZED_UNLIT
-
-#else // SUBSTRATE_ENABLED
-
-#if MATERIALBLENDING_MODULATE
-		const float3 Transparency = GetMaterialEmissive(PixelMaterialInputs);
-#elif MATERIALBLENDING_ALPHACOMPOSITE
-		const float Opacity = GetMaterialOpacity(PixelMaterialInputs);
-		const float Transparency = 1 - Opacity;
-#elif MATERIALBLENDING_MASKED && MATERIAL_DITHER_OPACITY_MASK
-		// See MATERIAL_DITHER_OPACITY_MASK comment below
-		const float Opacity = saturate(GetMaterialMaskInputRaw(PixelMaterialInputs));
-		const float Transparency = 1 - Opacity;
-#elif MATERIALBLENDING_TRANSLUCENT
-		const float Opacity = GetMaterialOpacity(PixelMaterialInputs);
-#if MATERIAL_SHADINGMODEL_THIN_TRANSLUCENT
-		// We can only get colored shadows for thin translucent (solid translucent case happens via absorption only)
-		float3 Transparency = 1 - Opacity;
-#else
-		float Transparency = 1 - Opacity;
-#endif
-		uint ShadingModelID = GetMaterialShadingModel(PixelMaterialInputs);
-
-#if MATERIAL_SHADINGMODEL_THIN_TRANSLUCENT
-		if (ShadingModelID == SHADINGMODELID_THIN_TRANSLUCENT)
-		{
-			float3 Transmission = GetThinTranslucentMaterialOutput0(MaterialParameters);
-#if REFRACTION_USE_INDEX_OF_REFRACTION
-			float Ior = GetRefractionIor(PixelMaterialInputs);
-#else
-			float Ior = 0.0;
-#endif
-			float3 V = WorldRayDirection();
-			float3 N = normalize(MaterialParameters.WorldNormal);
-			float VoN = abs(dot(V, N));
-			if (Opacity < 1.0)
-			{
-				float PathRoughness = PackedPayload.GetPathRoughness();
-				float Roughness = GetMaterialRoughness(PixelMaterialInputs);
-				if (Ior > 0.0 && PathRoughness <= Roughness)
-				{
-					// not using fast caustic approximation - treat as opaque
-					PackedPayload.SetRayThroughput(0.0);
-					return;
-				}
-				// compute transmission through the slab (fresnel + absorption)
-				float3 BaseColor = GetMaterialBaseColor(PixelMaterialInputs);
-				float  Metallic = GetMaterialMetallic(PixelMaterialInputs);
-				float  Specular = GetMaterialSpecular(PixelMaterialInputs);
-				float  F0 = F0RGBToF0(ComputeF0(Specular, BaseColor, Metallic));
-				Transparency *= ComputeThinSlabWeights(Transmission, VoN, Ior, F0).Transmitted;
-
-				// fake caustic approximation (see comments below)
-#if REFRACTION_USE_INDEX_OF_REFRACTION
-				if (Ior > 0.0)
-				{
-					Transparency *= FakeCausticsPathFactor(Roughness, PathRoughness);
-				}
-#endif
-			}
-		}
-		else
-#endif // MATERIAL_SHADINGMODEL_THIN_TRANSLUCENT
-#if MATERIAL_SHADINGMODEL_DEFAULT_LIT && REFRACTION_USE_INDEX_OF_REFRACTION
-		if (ShadingModelID == SHADINGMODELID_DEFAULT_LIT)
-		{
-			// Is refraction enabled?
-			float Ior = GetRefractionIor(PixelMaterialInputs);
-			if (Transparency > 0 && Ior > 0.0)
-			{
-				// current hit has some refraction
-				float PathRoughness = PackedPayload.GetPathRoughness();
-				float Roughness = GetMaterialRoughness(PixelMaterialInputs);
-
-				float3 BaseColor = GetMaterialBaseColor(PixelMaterialInputs);
-				float  Metallic = GetMaterialMetallic(PixelMaterialInputs);
-				float  Specular = GetMaterialSpecular(PixelMaterialInputs);
-				float  F0 = F0RGBToF0(ComputeF0(Specular, BaseColor, Metallic));
-
-				float3 N = normalize(MaterialParameters.WorldNormal);
-				float NoV = dot(WorldRayDirection(), N);
-
-				Transparency *= FakeCaustics(F0, Ior, NoV) * FakeCausticsPathFactor(Roughness, PathRoughness);
-
-#if HAVE_GetAbsorptionMediumMaterialOutput0
-				// Does the material have any kind of volumetric absorption to apply?
-				float3 TransmittanceColor = GetAbsorptionMediumMaterialOutput0(MaterialParameters);
-				float3 LocalSigmaT = PathTracingGlassTransmittanceToExtinction(TransmittanceColor);
-				
-				PackedPayload.SetTau(PackedPayload.GetTau() + ComputeOpticalDepthForHit(LocalSigmaT, bIsFrontFace));
-#endif // HAVE_GetAbsorptionMediumMaterialOutput0
-			}
-		}
-		else
-#endif // MATERIAL_SHADINGMODEL_DEFAULT_LIT
-		{
-			// base case for shadingmodel if/else
-		}
-#else // MATERIALBLENDING_*
-#error Unhandled blending mode!
-#endif
-
-#endif // SUBSTRATE_ENABLED
-
-		// Update the ray throughput (it is packed simply into the payload since we don't need to carry any other information across hits)
-		float3 RayThroughput = PackedPayload.GetRayThroughput();
-		RayThroughput *= Transparency;
-		PackedPayload.SetRayThroughput(RayThroughput);
-		return;
-	}
-#endif
-
 	FPathTracingPayload Payload = (FPathTracingPayload)0;
 
 #if !SUBSTRATE_ENABLED
@@ -1338,7 +1124,19 @@
 	FPackedPathTracingPayload, PackedPayload,
 	FRayTracingIntersectionAttributes, Attributes)
 {
-#if MATERIALBLENDING_MASKED && !MATERIAL_DITHER_OPACITY_MASK
+#if MATERIALBLENDING_MODULATE || (MATERIALBLENDING_MASKED && MATERIAL_DITHER_OPACITY_MASK) || MATERIALBLENDING_ALPHACOMPOSITE || MATERIALBLENDING_TRANSLUCENT
+	if (!PackedPayload.IsVisibilityRay())
+	{
+		// not a shadow ray -- don't need to run the AHS logic
+		return;
+	}
+	if (PackedPayload.HitT == RayTCurrent())
+	{
+		// We just processed this exact hit, don't double-count it
+		IgnoreHit();
+		return;
+	}
+#endif
 	// This is the only case which actually needs to run the full material
 	ResolvedView = ResolveView();
 
@@ -1371,16 +1169,214 @@
 
 		CalcMaterialParametersEx(MaterialParameters, PixelMaterialInputs, SvPosition, ScreenPosition, bIsFrontFace, TranslatedWorldPosition, TranslatedWorldPosition);
 	}
-
+#if MATERIALBLENDING_MASKED && !MATERIAL_DITHER_OPACITY_MASK
 	// Regardless of payload flags -- we always apply this 
 	if (GetMaterialMask(PixelMaterialInputs) < 0)
 	{
 		IgnoreHit();
 	}
+#endif
+
+	// the following blend modes need to process shadows after having executed the material
+#if MATERIALBLENDING_MODULATE || (MATERIALBLENDING_MASKED && MATERIAL_DITHER_OPACITY_MASK) || MATERIALBLENDING_ALPHACOMPOSITE || MATERIALBLENDING_TRANSLUCENT
+	PackedPayload.HitT = RayTCurrent(); // remember the distance at which we ran the shader, in case we get a double-hit
+
+	{ // TODO: Cleanup this indentation level
+
+#if SUBSTRATE_ENABLED
+
+#if MATERIAL_IS_SUBSTRATE && !SUBSTRATE_OPTIMIZED_UNLIT
+		// We already figured out how much Transparency there is above
+		// So only do extra logic for refraction and thin refraction ases
+#if MATERIALBLENDING_TRANSLUCENT && (REFRACTION_USE_INDEX_OF_REFRACTION || MATERIAL_ISTHINSURFACE)
+
+		// SUBSTRATE_TODO: refraction handling
+		float PathRoughness = PackedPayload.GetPathRoughness();
+
+		// SUBSTRATE_TODO: This does not match the logic in DistortAccumulatePS.usf
+		//  The idea here is to take F0 from the bottom refraction layers since this matches most closely what the refraction logic below does (only bottom slabs can refract).
+		//  The distortion approximation on the raster side instead averages F0 across layers which may produce different results.
+		//  In both cases - we get a potentially strange answer when there is more than one bottom slab that produces refraction
+
+		// SUBSTRATE_TODO: How should we handle MFP when there are multiple bottom slabs?
+		float F0 = 0.0f;
+		float3 RoughnessData = 0.0;
+		float NoV = 0.0f;
+		float3 LocalSigmaT = 0.0f;
+		SUBSTRATE_UNROLL_N(SUBSTRATE_CLAMPED_CLOSURE_COUNT)
+		for (int BSDFIdx = 0; BSDFIdx < SubstratePixelHeader.SubstrateTree.BSDFCount; ++BSDFIdx)
+		{
+#define BSDF SubstratePixelHeader.SubstrateTree.BSDFs[BSDFIdx]
+			if (BSDF.bIsBottom)
+			{
+				// SUBSTRATE_TODO: Also consider the MFP to see how much this slab is contributing to refraction ....
+				F0 = F0RGBToF0(SubstrateGetBSDFSpecularF0(BSDF));
+				RoughnessData.x = SLAB_ROUGHNESS(BSDF);
+				if (BSDF_GETHASHAZINESS(BSDF))
+				{
+					// NOTE: It should be safe to assume that the clearcoat case cannot happen here as it is only triggered by legacy clearcoat materials which are unlikely to have refraction
+					const FHaziness Haziness = UnpackHaziness(SLAB_HAZINESS(BSDF));
+					RoughnessData.y = Haziness.Roughness;
+					RoughnessData.z = Haziness.Weight;
+				}
+
+				float3 N = SubstratePixelHeader.SharedLocalBases.Normals[BSDF_GETSHAREDLOCALBASISID(BSDF)];
+				LocalSigmaT = rcp(max(SLAB_SSSMFP(BSDF), 0.0001));
+				NoV = dot(N, V_World);
+				break;
+			}
+#undef BSDF
+		}
+
+
+#if REFRACTION_USE_INDEX_OF_REFRACTION
+		float Ior = GetRefractionIor(PixelMaterialInputs, F0);
 #else
-	// All other blending modes need to run CHS to decide what to do, or are handled via the default opaque/translucent CHS/AHS materials.
-	// in these cases, the AHS should be disabled, so emit an error if we get here somehow
-	#error Unexpected blend mode encountered in AHS!
+		float Ior = 0.0;
 #endif
+
+#if MATERIAL_ISTHINSURFACE
+		// Thin Glass Fake Caustics
+
+		// SUBSTRATE_TODO: Revisit this -- TotalTransmittancePreCoverage should probably account for this in the thin case
+		float3 Transmission = exp(-LocalSigmaT * SUBSTRATE_SIMPLEVOLUME_THICKNESS_CM);
+		float3 Tr = ComputeThinSlabWeights(Transmission, NoV, Ior, F0).Transmitted;
+#else
+		HLSL_STATIC_ASSERT(REFRACTION_USE_INDEX_OF_REFRACTION == 1, "Expect refraction to be true here");
+		// Solid Glass Fake Caustics
+		float Tr = FakeCaustics(F0, Ior, NoV);
+		// Account for transmission 
+		PackedPayload.SetTau(PackedPayload.GetTau() + ComputeOpticalDepthForHit(LocalSigmaT, bIsFrontFace));
+#endif
+
+		// If material is doing refraction, apply an extra factor to account for ApproximateCaustics
+#if REFRACTION_USE_INDEX_OF_REFRACTION
+		if (Ior > 0.0)
+		{
+			Tr *= lerp(FakeCausticsPathFactor(RoughnessData.x, PathRoughness),
+			           FakeCausticsPathFactor(RoughnessData.y, PathRoughness), RoughnessData.z);
+		}
+#endif
+		Transparency += TotalCoverage * Tr;
+
+		// SUBSTRATE_TODO: Is this sanity check necessary?
+		Transparency = saturate(Transparency);
+#endif // MATERIALBLENDING_TRANSLUCENT && (REFRACTION_USE_INDEX_OF_REFRACTION || MATERIAL_ISTHINSURFACE)
+
+#endif // MATERIAL_IS_SUBSTRATE && !SUBSTRATE_OPTIMIZED_UNLIT
+
+#else // SUBSTRATE_ENABLED
+
+#if MATERIALBLENDING_MODULATE
+		const float3 Transparency = GetMaterialEmissive(PixelMaterialInputs);
+#elif MATERIALBLENDING_ALPHACOMPOSITE
+		const float Opacity = GetMaterialOpacity(PixelMaterialInputs);
+		const float Transparency = 1 - Opacity;
+#elif MATERIALBLENDING_MASKED && MATERIAL_DITHER_OPACITY_MASK
+		// See MATERIAL_DITHER_OPACITY_MASK comment below
+		const float Opacity = saturate(GetMaterialMaskInputRaw(PixelMaterialInputs));
+		const float Transparency = 1 - Opacity;
+#elif MATERIALBLENDING_TRANSLUCENT
+		const float Opacity = GetMaterialOpacity(PixelMaterialInputs);
+#if MATERIAL_SHADINGMODEL_THIN_TRANSLUCENT
+		// We can only get colored shadows for thin translucent (solid translucent case happens via absorption only)
+		float3 Transparency = 1 - Opacity;
+#else
+		float Transparency = 1 - Opacity;
+#endif
+		uint ShadingModelID = GetMaterialShadingModel(PixelMaterialInputs);
+
+#if MATERIAL_SHADINGMODEL_THIN_TRANSLUCENT
+		if (ShadingModelID == SHADINGMODELID_THIN_TRANSLUCENT)
+		{
+			float3 Transmission = GetThinTranslucentMaterialOutput0(MaterialParameters);
+#if REFRACTION_USE_INDEX_OF_REFRACTION
+			float Ior = GetRefractionIor(PixelMaterialInputs);
+#else
+			float Ior = 0.0;
+#endif
+			float3 V = WorldRayDirection();
+			float3 N = normalize(MaterialParameters.WorldNormal);
+			float VoN = abs(dot(V, N));
+			if (Opacity < 1.0)
+			{
+				float PathRoughness = PackedPayload.GetPathRoughness();
+				float Roughness = GetMaterialRoughness(PixelMaterialInputs);
+				if (Ior > 0.0 && PathRoughness <= Roughness)
+				{
+					// not using fast caustic approximation - treat as opaque
+					PackedPayload.SetRayThroughput(0.0);
+					return;
+				}
+				// compute transmission through the slab (fresnel + absorption)
+				float3 BaseColor = GetMaterialBaseColor(PixelMaterialInputs);
+				float  Metallic = GetMaterialMetallic(PixelMaterialInputs);
+				float  Specular = GetMaterialSpecular(PixelMaterialInputs);
+				float  F0 = F0RGBToF0(ComputeF0(Specular, BaseColor, Metallic));
+				Transparency *= ComputeThinSlabWeights(Transmission, VoN, Ior, F0).Transmitted;
+
+				// fake caustic approximation (see comments below)
+#if REFRACTION_USE_INDEX_OF_REFRACTION
+				if (Ior > 0.0)
+				{
+					Transparency *= FakeCausticsPathFactor(Roughness, PathRoughness);
+				}
+#endif
+			}
+		}
+		else
+#endif // MATERIAL_SHADINGMODEL_THIN_TRANSLUCENT
+#if MATERIAL_SHADINGMODEL_DEFAULT_LIT && REFRACTION_USE_INDEX_OF_REFRACTION
+		if (ShadingModelID == SHADINGMODELID_DEFAULT_LIT)
+		{
+			// Is refraction enabled?
+			float Ior = GetRefractionIor(PixelMaterialInputs);
+			if (Transparency > 0 && Ior > 0.0)
+			{
+				// current hit has some refraction
+				float PathRoughness = PackedPayload.GetPathRoughness();
+				float Roughness = GetMaterialRoughness(PixelMaterialInputs);
+
+				float3 BaseColor = GetMaterialBaseColor(PixelMaterialInputs);
+				float  Metallic = GetMaterialMetallic(PixelMaterialInputs);
+				float  Specular = GetMaterialSpecular(PixelMaterialInputs);
+				float  F0 = F0RGBToF0(ComputeF0(Specular, BaseColor, Metallic));
+
+				float3 N = normalize(MaterialParameters.WorldNormal);
+				float NoV = dot(WorldRayDirection(), N);
+
+				Transparency *= FakeCaustics(F0, Ior, NoV) * FakeCausticsPathFactor(Roughness, PathRoughness);
+
+#if HAVE_GetAbsorptionMediumMaterialOutput0
+				// Does the material have any kind of volumetric absorption to apply?
+				float3 TransmittanceColor = GetAbsorptionMediumMaterialOutput0(MaterialParameters);
+				float3 LocalSigmaT = PathTracingGlassTransmittanceToExtinction(TransmittanceColor);
+				
+				PackedPayload.SetTau(PackedPayload.GetTau() + ComputeOpticalDepthForHit(LocalSigmaT, bIsFrontFace));
+#endif // HAVE_GetAbsorptionMediumMaterialOutput0
+			}
+		}
+		else
+#endif // MATERIAL_SHADINGMODEL_DEFAULT_LIT
+		{
+			// base case for shadingmodel if/else
+		}
+#else // MATERIALBLENDING_*
+#error Unhandled blending mode!
+#endif
+
+#endif // SUBSTRATE_ENABLED
+
+		// Update the ray throughput (it is packed simply into the payload since we don't need to carry any other information across hits)
+		float3 RayThroughput = PackedPayload.GetRayThroughput();
+		RayThroughput *= Transparency;
+		PackedPayload.SetRayThroughput(RayThroughput);
+		if (any(RayThroughput > 0))
+		{
+			// keep tracing if we still have some energy left
+			IgnoreHit();
+		}
+	}
+#endif
 }
 
diff --git a/Engine/Source/Runtime/Renderer/Private/PathTracing.cpp b/Engine/Source/Runtime/Renderer/Private/PathTracing.cpp
--- a/Engine/Source/Runtime/Renderer/Private/PathTracing.cpp
+++ b/Engine/Source/Runtime/Renderer/Private/PathTracing.cpp
@@ -1397,16 +1397,26 @@ FRayTracingLightFunctionMap GatherLightFunctionLightsPathTracing(FScene* Scene,
 	return RayTracingLightFunctionMap;
 }
 
-static bool NeedsAnyHitShader(bool bIsMasked, bool bIsDitherMasked)
+static bool NeedsAnyHitShader(EBlendMode BlendMode)
 {
-	// Masked materials use AHS to quickly reject masked out portions
-	// However, dither masking gets treated as ordinary translucency
-	return (bIsMasked && !bIsDitherMasked);
+	switch (BlendMode)
+	{
+		case BLEND_Opaque: 							return false; // always hit
+		case BLEND_Masked: 							return true;  // runs shader (NOTE: dithered masking gets turned into translucent for the path tracer)
+		case BLEND_Translucent: 					return true;  // casts transparent (colored) shadows depending on the shading model setup (fake caustics or transparent shadows)
+		case BLEND_Additive: 						return false; // never hit for shadows, goes through the default shader instead, so no need to use AHS for primary rays
+		case BLEND_Modulate: 						return true;  // casts colored shadows
+		case BLEND_AlphaComposite: 					return true;
+		case BLEND_AlphaHoldout: 					return false; // treat as opaque for shadows
+		case BLEND_TranslucentColoredTransmittance: return true;  // NOTE: Substrate only
+		default: checkf(false, TEXT("Unhandled blend mode %d"), int(BlendMode)); return false;
+	}
+
 }
 
 static bool NeedsAnyHitShader(const FMaterial& RESTRICT MaterialResource)
 {
-	return NeedsAnyHitShader(MaterialResource.GetBlendMode() == BLEND_Masked, MaterialResource.IsDitherMasked());
+	return NeedsAnyHitShader(MaterialResource.GetBlendMode());
 }
 
 template<bool UseAnyHitShader, bool UseIntersectionShader, bool IsGPULightmass, bool SimplifySubstrate>
@@ -1437,7 +1447,7 @@ public:
 			// This material is only for surfaces at the moment
 			return false;
 		}
-		if (NeedsAnyHitShader(Parameters.MaterialParameters.BlendMode == BLEND_Masked, Parameters.MaterialParameters.bIsDitherMasked) != UseAnyHitShader)
+		if (NeedsAnyHitShader(Parameters.MaterialParameters.BlendMode) != UseAnyHitShader)
 		{
 			return false;
 		}
