Fix cook error w/CL 30512983

Custom data define was not checked after previous submit rollback.

diff --git a/Engine/Shaders/Private/SceneData.ush b/Engine/Shaders/Private/SceneData.ush
--- a/Engine/Shaders/Private/SceneData.ush
+++ b/Engine/Shaders/Private/SceneData.ush
@@ -276,8 +276,10 @@ FInstanceSceneData WaveReadLaneAt(FInstanceSceneData In, uint SrcIndex)
 	Result.NaniteRuntimeResourceID		= WaveReadLaneAt(In.NaniteRuntimeResourceID, SrcIndex);
 	Result.NaniteHierarchyOffset		= WaveReadLaneAt(In.NaniteHierarchyOffset, SrcIndex);
 	Result.RandomID						= WaveReadLaneAt(In.RandomID, SrcIndex);
+#if ENABLE_PER_INSTANCE_CUSTOM_DATA
 	Result.CustomDataOffset				= WaveReadLaneAt(In.CustomDataOffset, SrcIndex);
 	Result.CustomDataCount				= WaveReadLaneAt(In.CustomDataCount, SrcIndex);
+#endif
 	Result.LightMapAndShadowMapUVBias	= WaveReadLaneAt(In.LightMapAndShadowMapUVBias, SrcIndex);
 	Result.ValidInstance				= WaveReadLaneAt(In.ValidInstance, SrcIndex);
 	Result.Flags						= WaveReadLaneAt(In.Flags, SrcIndex);
