Add toggle of hair bounds when visualizing hair instance debug data.

diff --git a/Engine/Plugins/Runtime/HairStrands/Source/HairStrandsCore/Private/GroomDebug.cpp b/Engine/Plugins/Runtime/HairStrands/Source/HairStrandsCore/Private/GroomDebug.cpp
--- a/Engine/Plugins/Runtime/HairStrands/Source/HairStrandsCore/Private/GroomDebug.cpp
+++ b/Engine/Plugins/Runtime/HairStrands/Source/HairStrandsCore/Private/GroomDebug.cpp
@@ -705,8 +705,7 @@ class FHairDebugPrintInstanceCS : public FGlobalShader
 	DECLARE_GLOBAL_SHADER(FHairDebugPrintInstanceCS);
 	SHADER_USE_PARAMETER_STRUCT(FHairDebugPrintInstanceCS, FGlobalShader);
 
-	class FOutputType : SHADER_PERMUTATION_INT("PERMUTATION_OUTPUT_TYPE", 2);
-	using FPermutationDomain = TShaderPermutationDomain<FOutputType>;
+	using FPermutationDomain = TShaderPermutationDomain<>;
 
 	BEGIN_SHADER_PARAMETER_STRUCT(FParameters, )
 		SHADER_PARAMETER(uint32, InstanceCount)
@@ -933,7 +932,6 @@ static void AddHairDebugPrintInstancePass(
 		Parameters->Infos = GraphBuilder.CreateSRV(InfoBuffer, PF_R32_UINT);
 		ShaderPrint::SetParameters(GraphBuilder, *ShaderPrintData, Parameters->ShaderPrintUniformBuffer);
 		FHairDebugPrintInstanceCS::FPermutationDomain PermutationVector;
-		PermutationVector.Set<FHairDebugPrintInstanceCS::FOutputType>(0);
 		TShaderMapRef<FHairDebugPrintInstanceCS> ComputeShader(ShaderMap, PermutationVector);
 
 		ClearUnusedGraphResources(ComputeShader, Parameters);
@@ -945,36 +943,6 @@ static void AddHairDebugPrintInstancePass(
 			Parameters,
 			FIntVector(1, 1, 1));
 	}
-
-	// Draw instances bound (one pass for each instance, due to separate AABB resources)
-	FHairDebugPrintInstanceCS::FPermutationDomain PermutationVector;
-	PermutationVector.Set<FHairDebugPrintInstanceCS::FOutputType>(1);
-	TShaderMapRef<FHairDebugPrintInstanceCS> ComputeShader(ShaderMap, PermutationVector);
-	for (uint32 InstanceIndex = 0; InstanceIndex < InstanceCount; ++InstanceIndex)
-	{
-		const FHairStrandsInstance* AbstractInstance = Instances[InstanceIndex];
-		const FHairGroupInstance* Instance = static_cast<const FHairGroupInstance*>(AbstractInstance);
-
-		if (Instance->GeometryType == EHairGeometryType::Strands)
-		{
-			const float MaxRectSizeInPixels = FMath::Min(View.UnscaledViewRect.Height(), View.UnscaledViewRect.Width());
-			const float ContinousLODRadius = Instance->HairGroupPublicData->ContinuousLODScreenSize * MaxRectSizeInPixels * 0.5f; // Diameter->Radius
-
-			FHairDebugPrintInstanceCS::FParameters* Parameters = GraphBuilder.AllocParameters<FHairDebugPrintInstanceCS::FParameters>();
-			Parameters->InstanceRegisteredIndex = Instance->RegisteredIndex;
-			Parameters->InstanceAABB = TransientResources.GroupAABBSRV;
-			Parameters->InstanceScreenSphereBound = FVector4f(Instance->HairGroupPublicData->ContinuousLODScreenPos.X, Instance->HairGroupPublicData->ContinuousLODScreenPos.Y, 0.f, ContinousLODRadius);
-			ShaderPrint::SetParameters(GraphBuilder, *ShaderPrintData, Parameters->ShaderPrintUniformBuffer);
-			ClearUnusedGraphResources(ComputeShader, Parameters);
-
-			FComputeShaderUtils::AddPass(
-				GraphBuilder,
-				RDG_EVENT_NAME("HairStrands::DebugPrintInstance(Bound)", InstanceCount),
-				ComputeShader,
-				Parameters,
-				FIntVector(1, 1, 1));
-		}
-	}
 }
 
 ///////////////////////////////////////////////////////////////////////////////////////////////////
diff --git a/Engine/Shaders/Private/HairStrands/HairStrandsDebug.usf b/Engine/Shaders/Private/HairStrands/HairStrandsDebug.usf
--- a/Engine/Shaders/Private/HairStrands/HairStrandsDebug.usf
+++ b/Engine/Shaders/Private/HairStrands/HairStrandsDebug.usf
@@ -960,7 +960,9 @@ uint HairVisibilityNodeGroupSize;
 uint AllocatedSampleCount;
 uint HairInstanceCount;
 
-Buffer<uint> HairInstanceIDs;
+Buffer<uint> HairInstanceDataBuffer;
+Buffer<int> InstanceAABBBuffer;
+
 Texture2D<float> HairCountTexture;
 Texture2D<uint> HairCountUintTexture;
 
@@ -1004,37 +1006,99 @@ void DrawInstanceAABB(uint PrimitiveId)
 	}
 }
 
+struct FData
+{
+	uint InstanceID;
+	uint RegisteredIndex;
+	uint GeometryType;
+	float4 InstanceScreenSphereBound;
+};
+
+FData LoadInstanceData(uint InIndex)
+{
+	const uint Index8 = InIndex*8;
+	FData Out;
+	Out.InstanceID					= HairInstanceDataBuffer[Index8+0];
+	Out.RegisteredIndex				= HairInstanceDataBuffer[Index8+1];
+	Out.GeometryType				= HairInstanceDataBuffer[Index8+2];
+
+	Out.InstanceScreenSphereBound.x	= asfloat(HairInstanceDataBuffer[Index8+4]);
+	Out.InstanceScreenSphereBound.y	= asfloat(HairInstanceDataBuffer[Index8+5]);
+	Out.InstanceScreenSphereBound.z	= asfloat(HairInstanceDataBuffer[Index8+6]);
+	Out.InstanceScreenSphereBound.w	= asfloat(HairInstanceDataBuffer[Index8+7]);
+	return Out;
+}
 [numthreads(1, 1, 1)]
 void MainCS(uint InGroupIndex : SV_GroupIndex, uint3 DispatchThreadId : SV_DispatchThreadID)
 {	
 	// Draw AABB of the instance groups
 	if (all(DispatchThreadId == 0))
 	{
-		// Macro groups bounds
+		// Draw bound & legends
 		{
-			for (uint GroupIndex =0;GroupIndex< HairMacroGroupCount; GroupIndex++)
+			FShaderPrintContext CtxL = InitShaderPrintContext(true, uint2(900, 100));
+			bool bDrawGPUScene 					= false;
+			bool bDrawInstanceGroupScreen		= false;
+			bool bDrawHairInstances 			= false;
+			bool bDrawMacroBounds 				= false;
+			bool bDrawVoxelAlignedMacroBounds	= false;
+
+			// Bounds legends
 			{
-				FHairAABB Bound = ReadHairAABB(GroupIndex, HairMacroGroupAABBBuffer);
-				AddAABBTWS(Bound.Min, Bound.Max, GetGPUAABBMacroGroupBoundColor());
+				bDrawGPUScene 				= AddCheckbox(CtxL, TEXT("Bound - GPU Scene"), 							true, FFontColor(GetGPUSceneBoundColor().xyz)); 						Newline(CtxL);
+				bDrawInstanceGroupScreen	= AddCheckbox(CtxL, TEXT("Bound - Instance Group Screen"), 				true, FFontColor(GetInstanceGroupScreenBoundColor().xyz)); 				Newline(CtxL);
+				bDrawHairInstances 			= AddCheckbox(CtxL, TEXT("Bound - AABB Instance Group"), 				true, FFontColor(GetGPUAABBInstanceGroupBoundColor().xyz)); 			Newline(CtxL);
+				bDrawMacroBounds 			= AddCheckbox(CtxL, TEXT("Bound - AABB Macro Group"), 					true, FFontColor(GetGPUAABBMacroGroupBoundColor().xyz)); 				Newline(CtxL);
+				bDrawVoxelAlignedMacroBounds= AddCheckbox(CtxL, TEXT("Bound - AABB Macro Group (voxel-aligned)"),	true, FFontColor(GetGPUAABBMacroGroupBoundVoxelAlignedColor().xyz)); 	Newline(CtxL);
 			}
-		}
 
-		// Voxel aligned macro groups bounds
-		{
-			for (uint GroupIndex = 0; GroupIndex < HairMacroGroupCount; GroupIndex++)
+			// Macro groups bounds
+			if (bDrawMacroBounds)
 			{
-				FHairAABB Bound = ReadHairAABB(GroupIndex, HairMacroGroupVoxelAlignedAABBBuffer);
-				AddAABBTWS(Bound.Min, Bound.Max, GetGPUAABBMacroGroupBoundVoxelAlignedColor());
+				for (uint GroupIndex =0;GroupIndex< HairMacroGroupCount; GroupIndex++)
+				{
+					FHairAABB Bound = ReadHairAABB(GroupIndex, HairMacroGroupAABBBuffer);
+					AddAABBTWS(Bound.Min, Bound.Max, GetGPUAABBMacroGroupBoundColor());
+				}
 			}
-		}
-
-		// Hair instance bounds
-		for (uint HairInstanceIndex = 0; HairInstanceIndex < HairInstanceCount; HairInstanceIndex++)
-		{
-			const uint InstanceID = HairInstanceIDs[HairInstanceIndex];
-			if (InstanceID != ~0)
+	
+			// Voxel aligned macro groups bounds
+			if (bDrawVoxelAlignedMacroBounds)
+			{
+				for (uint GroupIndex = 0; GroupIndex < HairMacroGroupCount; GroupIndex++)
+				{
+					FHairAABB Bound = ReadHairAABB(GroupIndex, HairMacroGroupVoxelAlignedAABBBuffer);
+					AddAABBTWS(Bound.Min, Bound.Max, GetGPUAABBMacroGroupBoundVoxelAlignedColor());
+				}
+			}
+	
+			// Hair instance
+			for (uint HairInstanceIndex = 0; HairInstanceIndex < HairInstanceCount; HairInstanceIndex++)
 			{
-				DrawInstanceAABB(InstanceID);
+				const FData InstanceData = LoadInstanceData(HairInstanceIndex);
+
+				// GPU scene bounds
+				if (bDrawGPUScene)
+				{
+					if (InstanceData.InstanceID != ~0)
+					{
+						DrawInstanceAABB(InstanceData.InstanceID);
+					}
+				}
+
+				// Instance AABB
+				if (bDrawHairInstances && InstanceData.GeometryType == 0/*Strands*/)
+				{
+					const FHairAABB AABB = ReadHairAABB(InstanceData.RegisteredIndex, InstanceAABBBuffer);
+					AddAABBTWS(AABB.Min, AABB.Max, GetGPUAABBInstanceGroupBoundColor());
+				}
+
+				// Screen AABB/size
+				if (bDrawInstanceGroupScreen && InstanceData.GeometryType == 0/*Strands*/)
+				{
+					FShaderPrintContext Context = InitShaderPrintContext(true, uint2(50, 50));
+					AddCircleSS(Context, InstanceData.InstanceScreenSphereBound.xy, InstanceData.InstanceScreenSphereBound.w, ColorGreen);
+				}
 			}
 		}
 
@@ -1381,7 +1445,6 @@ bool HasAttribute(FInstanceInfo In, uint InAttributeIt)
 
 [numthreads(1, 1, 1)]
 void MainCS(uint GroupIndex : SV_GroupIndex, uint3 DispatchThreadId : SV_DispatchThreadID)
-#if PERMUTATION_OUTPUT_TYPE == 0
 {
 	// Draw AABB of the instance groups
 	if (any(DispatchThreadId != 0))
@@ -1618,31 +1681,7 @@ void MainCS(uint GroupIndex : SV_GroupIndex, uint3 DispatchThreadId : SV_Dispatc
 			Newline(Context);
 		}
 	}
-
-	// Bounds legends
-	{
-		FShaderPrintContext CtxL = InitShaderPrintContext(true, uint2(900, 100));
-		Print(CtxL, TEXT("Bound - GPU Scene"), 							FFontColor(GetGPUSceneBoundColor().xyz)); Newline(CtxL);
-		Print(CtxL, TEXT("Bound - Instance Group Screen"), 				FFontColor(GetInstanceGroupScreenBoundColor().xyz)); Newline(CtxL);
-		Print(CtxL, TEXT("Bound - AABB Instance Group"), 				FFontColor(GetGPUAABBInstanceGroupBoundColor().xyz)); Newline(CtxL);
-		Print(CtxL, TEXT("Bound - AABB Macro Group"), 					FFontColor(GetGPUAABBMacroGroupBoundColor().xyz)); Newline(CtxL);
-		Print(CtxL, TEXT("Bound - AABB Macro Group (voxel-aligned)"),	FFontColor(GetGPUAABBMacroGroupBoundVoxelAlignedColor().xyz)); Newline(CtxL);
-	}
-}
-#else  // PERMUTATION_OUTPUT_TYPE
-{
-	if (all(DispatchThreadId == 0))
-	{
-		// Bound
-		const FHairAABB AABB = ReadHairAABB(InstanceRegisteredIndex, InstanceAABB);
-		AddAABBTWS(AABB.Min, AABB.Max, GetGPUAABBInstanceGroupBoundColor());
-
-		// Screen size
-		FShaderPrintContext Context = InitShaderPrintContext(true, uint2(50, 50));
-		AddCircleSS(Context, InstanceScreenSphereBound.xy, InstanceScreenSphereBound.w, ColorGreen);
-	}
 }
-#endif // PERMUTATION_OUTPUT_TYPE
 
 #endif // SHADER_PRINT_INSTANCE
 
diff --git a/Engine/Source/Runtime/Renderer/Private/HairStrands/HairStrandsDebug.cpp b/Engine/Source/Runtime/Renderer/Private/HairStrands/HairStrandsDebug.cpp
--- a/Engine/Source/Runtime/Renderer/Private/HairStrands/HairStrandsDebug.cpp
+++ b/Engine/Source/Runtime/Renderer/Private/HairStrands/HairStrandsDebug.cpp
@@ -200,7 +200,8 @@ class FHairDebugPrintCS : public FGlobalShader
 		SHADER_PARAMETER(uint32, AllocatedSampleCount)
 		SHADER_PARAMETER(uint32, HairInstanceCount)
 		SHADER_PARAMETER(float, ResolutionScale)
-		SHADER_PARAMETER_RDG_BUFFER_SRV(Buffer<uint>, HairInstanceIDs)
+		SHADER_PARAMETER_RDG_BUFFER_SRV(Buffer<uint>, HairInstanceDataBuffer)
+		SHADER_PARAMETER_RDG_BUFFER_SRV(Buffer<int>, InstanceAABBBuffer)
 		SHADER_PARAMETER_RDG_TEXTURE(Texture2D, HairCountTexture)
 		SHADER_PARAMETER_RDG_TEXTURE(Texture2D, HairCountUintTexture)
 		SHADER_PARAMETER_RDG_BUFFER_SRV(Buffer<uint>, HairVisibilityIndirectArgsBuffer)
@@ -245,15 +246,24 @@ static void AddDebugHairPrintPass(
 		return;
 	}
 
-	if (!TryEnableShaderDrawAndShaderPrint(*View, MacroGroupResources.MacroGroupCount * 32u, 2000u))
+	if (!TryEnableShaderDrawAndShaderPrint(*View, MacroGroupResources.MacroGroupCount * 32u, 8192u))
 	{
 		return;
 	}
 
+	struct FData
+	{
+		uint32 PrimitiveID = ~0;
+		uint32 RegisteredIndex = ~0;
+		uint32 GeometryType = 0;
+		uint32 Pad0 = 0;
+		FVector4f InstanceScreenSphereBound = FVector4f::Zero();
+	};
+
 	// Build mapping Instance -> PrimitiveID to fetch primitive data (i.e., transform & co)
 	const uint32 MacroGroupCount = MacroGroupDatas.Num();
-	TArray<uint32> InstanceIDs;
-	InstanceIDs.Reserve(MacroGroupCount * 4u);
+	TArray<FData> InstanceDatas;
+	InstanceDatas.Reserve(MacroGroupCount * 4u);
 	for (uint32 MacroGroupIndex = 0; MacroGroupIndex < MacroGroupCount; ++MacroGroupIndex)
 	{
 		const FHairStrandsMacroGroupData& MacroGroup = MacroGroupDatas[MacroGroupIndex];
@@ -261,21 +271,30 @@ static void AddDebugHairPrintPass(
 		{
 			if (PrimitiveInfo.PrimitiveSceneProxy)
 			{
-				uint32 PrimitiveID = ~0u;
+				FData& InstanceData = InstanceDatas.AddDefaulted_GetRef();
+				InstanceData.PrimitiveID = ~0u;
 				if (const FPrimitiveSceneInfo* SceneInfo = PrimitiveInfo.PrimitiveSceneProxy->GetPrimitiveSceneInfo())
 				{
-					PrimitiveID = SceneInfo->GetIndex();
+					InstanceData.PrimitiveID = SceneInfo->GetIndex();
 				}
-				InstanceIDs.Add(PrimitiveID);
+
+				FHairStrandsInstance* Instance = PrimitiveInfo.PublicDataPtr->Instance;
+				check(Instance);
+
+				const float MaxRectSizeInPixels = FMath::Min(View->UnscaledViewRect.Height(), View->UnscaledViewRect.Width());
+				const float ContinousLODRadius = PrimitiveInfo.PublicDataPtr->ContinuousLODScreenSize * MaxRectSizeInPixels * 0.5f; // Diameter->Radius
+
+				InstanceData.GeometryType 				= Instance->GetHairGeometry();
+				InstanceData.RegisteredIndex			= Instance->RegisteredIndex;
+				InstanceData.InstanceScreenSphereBound 	= FVector4f(PrimitiveInfo.PublicDataPtr->ContinuousLODScreenPos.X, PrimitiveInfo.PublicDataPtr->ContinuousLODScreenPos.Y, 0.f, ContinousLODRadius);
 			}
 			else
 			{
-				InstanceIDs.Add(0u);
+				InstanceDatas.AddDefaulted();
 			}
 		}
 	}
-	FRDGBufferRef InstancesIDBuffer = CreateVertexBuffer(GraphBuilder, TEXT("Hair.Debug.InstanceIDs"), FRDGBufferDesc::CreateBufferDesc(4, InstanceIDs.Num()), InstanceIDs.GetData(), 4u * InstanceIDs.Num());
-
+	FRDGBufferRef InstanceDataBuffer = CreateVertexBuffer(GraphBuilder, TEXT("Hair.Debug.InstanceDatas"), FRDGBufferDesc::CreateBufferDesc(sizeof(FData), InstanceDatas.Num()), InstanceDatas.GetData(), sizeof(FData) * InstanceDatas.Num());
 
 	FRDGTextureRef ViewHairCountTexture = VisibilityData.ViewHairCountTexture ? VisibilityData.ViewHairCountTexture : GSystemTextures.GetBlackDummy(GraphBuilder);
 	FRDGTextureRef ViewHairCountUintTexture = VisibilityData.ViewHairCountUintTexture ? VisibilityData.ViewHairCountUintTexture : GSystemTextures.GetBlackDummy(GraphBuilder);
@@ -286,8 +305,9 @@ static void AddDebugHairPrintPass(
 	FHairDebugPrintCS::FParameters* Parameters = GraphBuilder.AllocParameters<FHairDebugPrintCS::FParameters>();
 	Parameters->ResolutionScale = float(View->ViewRect.Width()) / float(View->UnscaledViewRect.Width());
 	Parameters->Scene = View->GetSceneUniforms().GetBuffer(GraphBuilder);
-	Parameters->HairInstanceCount = InstanceIDs.Num();
-	Parameters->HairInstanceIDs = GraphBuilder.CreateSRV(InstancesIDBuffer, PF_R32_UINT);
+	Parameters->HairInstanceCount = InstanceDatas.Num();
+	Parameters->HairInstanceDataBuffer = GraphBuilder.CreateSRV(InstanceDataBuffer, PF_R32_UINT);
+	Parameters->InstanceAABBBuffer = Scene->HairStrandsSceneData.TransientResources->GroupAABBSRV;
 	Parameters->GroupSize = GetVendorOptimalGroupSize2D();
 	Parameters->ViewUniformBuffer = View->ViewUniformBuffer;
 	Parameters->MaxResolution = VisibilityData.CoverageTexture ? VisibilityData.CoverageTexture->Desc.Extent : FIntPoint(0,0);
