Fixes TSR performance regression on sg.AntiAliasingQuality=0 with 16bits ops

diff --git a/Engine/Shaders/Private/TemporalSuperResolution/TSRRejectShading.usf b/Engine/Shaders/Private/TemporalSuperResolution/TSRRejectShading.usf
--- a/Engine/Shaders/Private/TemporalSuperResolution/TSRRejectShading.usf
+++ b/Engine/Shaders/Private/TemporalSuperResolution/TSRRejectShading.usf
@@ -306,6 +306,20 @@ void FetchHistoryGuide(
 	#endif
 } // FetchHistoryGuide()
 
+// Fetch the previous history history.
+void FetchReprojectedHistoryGuide(
+	tsr_simd_short2 PixelPos,
+	out tsr_simd_halfC History,
+	out tsr_simd_half  HistoryUncertainty)
+{
+	FetchHistoryGuide(
+		ReprojectedHistoryGuideTexture,
+		ReprojectedHistoryGuideMetadataTexture,
+		PixelPos,
+		/* out */ History,
+		/* out */ HistoryUncertainty);
+}
+
 #if CONFIG_MOIRE_DETECTION
 
 void FetchMoireInput(
@@ -704,7 +718,8 @@ void MeasureRejection(
 		
 	// Clamp filtered signal
 	tsr_simd_halfC ClampedFilteredHistory = fastClamp(FilteredHistory, FilteredBoxMin, FilteredBoxMax);
-	if (ApplyMoireError)
+	BRANCH
+	if (ApplyMoireError && bEnableFlickeringHeuristic)
 	{
 		tsr_simd_halfC MoireErrorSize;
 			
@@ -737,7 +752,8 @@ void MeasureRejection(
 			InputC2BoxSize * tsr_simd_halfC::Const(FilteringWeight) +
 			(BackbufferQuantizationErrorVector * tsr_half(2 * FilteringWeight)));
 
-		if (ApplyMoireError)
+		BRANCH
+		if (ApplyMoireError && bEnableFlickeringHeuristic)
 		{
 			tsr_simd_halfC MoireErrorSize;
 			UNROLL_N(SIMD_SIZE)
@@ -784,6 +800,20 @@ void MeasureRejection(
 
 #if CONFIG_HISTORY_RESURRECTION
 
+// Fetch the resurrected history.
+void FetchResurrectedHistoryGuide(
+	tsr_simd_short2 PixelPos,
+	out tsr_simd_halfC History,
+	out tsr_simd_half  HistoryUncertainty)
+{
+	FetchHistoryGuide(
+		ResurrectedHistoryGuideTexture,
+		ResurrectedHistoryGuideMetadataTexture,
+		PixelPos,
+		/* out */ History,
+		/* out */ HistoryUncertainty);
+}
+
 // Returns whether resurrected frame is closer than previous frame on per pixel basis.
 CALL_SITE_DEBUGLOC
 tsr_simd_bool IsResurrectedFrameCloserThanPreviousFrame(tsr_simd_halfC ExposedInput, tsr_simd_halfC ExposedHistory, tsr_simd_halfC ExposedResurrectedHistory)
@@ -811,11 +841,20 @@ tsr_simd_bool IsResurrectedFrameCloserThanPreviousFrame(tsr_simd_halfC ExposedIn
 CALL_SITE_DEBUGLOC
 bool IsWorthMeasuringRejectionOfResurrectedFrame(tsr_simd_half PrevRejectionBlendFinal, tsr_simd_bool bResurrectionIsCloser)
 {
-	tsr_simd_bool bPrevIsNotGoodEnough = PrevRejectionBlendFinal < tsr_simd_half::Const(0.5);
-	tsr_simd_bool bIsValidTilePixel = IsValidTilePixel<LANE_STRIDE_X, LANE_STRIDE_Y>(TileOverscan);
-	tsr_simd_bool bShouldComputeResurrection = and(and(bResurrectionIsCloser, bPrevIsNotGoodEnough), bIsValidTilePixel);
+	bool bComputeResurrection;
+	BRANCH
+	if (bEnableResurrection)
+	{
+		tsr_simd_bool bPrevIsNotGoodEnough = PrevRejectionBlendFinal < tsr_simd_half::Const(0.5);
+		tsr_simd_bool bIsValidTilePixel = IsValidTilePixel<LANE_STRIDE_X, LANE_STRIDE_Y>(TileOverscan);
+		tsr_simd_bool bShouldComputeResurrection = and(and(bResurrectionIsCloser, bPrevIsNotGoodEnough), bIsValidTilePixel);
 
-	bool bComputeResurrection = AnyThread(AnyComponent(bShouldComputeResurrection));
+		bComputeResurrection = AnyThread(AnyComponent(bShouldComputeResurrection));
+	}
+	else
+	{
+		bComputeResurrection = false;
+	}
 
 	return bComputeResurrection;
 }
@@ -915,9 +954,7 @@ void RefetchInputAndHistoryFromLDSOrCache(
 
 	FetchDecimateMask(PixelPos, /* out */ VelocityEdge, /* out */ bIsDisoccluded);
 	FetchSceneColorAndTranslucency(PixelPos, /* out */ OriginalOpaqueInput, /* out */ OriginalTranslucencyInput);
-	FetchHistoryGuide(
-		ReprojectedHistoryGuideTexture,
-		ReprojectedHistoryGuideMetadataTexture,
+	FetchReprojectedHistoryGuide(
 		PixelPos,
 		/* out */ OriginalHistory,
 		/* out */ OriginalHistoryUncertainty);
@@ -1018,27 +1055,42 @@ void MainCS(
 		tsr_simd_halfC OriginalResurrectedHistory;
 		tsr_simd_half  OriginalResurrectedHistoryUncertainty;
 	#endif
-	PLATFORM_SPECIFIC_ISOLATE
 	{
-		tsr_simd_short2 PixelPos = ComputePixelPos();
-
-		FetchDecimateMask(PixelPos, /* out */ VelocityEdge, /* out */ bIsDisoccluded);
-		FetchSceneColorAndTranslucency(PixelPos, /* out */ OriginalOpaqueInput, /* out */ OriginalTranslucencyInput);
-		FetchHistoryGuide(
-			ReprojectedHistoryGuideTexture,
-			ReprojectedHistoryGuideMetadataTexture,
-			PixelPos,
-			/* out */ OriginalHistory,
-			/* out */ OriginalHistoryUncertainty);
-			
 		#if CONFIG_HISTORY_RESURRECTION
-			FetchHistoryGuide(
-				ResurrectedHistoryGuideTexture,
-				ReprojectedHistoryGuideMetadataTexture,
+		BRANCH
+		if (bEnableResurrection)
+		PLATFORM_SPECIFIC_ISOLATE
+		{
+			tsr_simd_short2 PixelPos = ComputePixelPos();
+
+			FetchDecimateMask(PixelPos, /* out */ VelocityEdge, /* out */ bIsDisoccluded);
+			FetchSceneColorAndTranslucency(PixelPos, /* out */ OriginalOpaqueInput, /* out */ OriginalTranslucencyInput);
+			FetchReprojectedHistoryGuide(
+				PixelPos,
+				/* out */ OriginalHistory,
+				/* out */ OriginalHistoryUncertainty);
+			FetchResurrectedHistoryGuide(
 				PixelPos,
 				/* out */ OriginalResurrectedHistory,
 				/* out */ OriginalResurrectedHistoryUncertainty);
-		#endif
+		}
+		else // if (!bEnableResurrection)
+		#endif // CONFIG_HISTORY_RESURRECTION
+		PLATFORM_SPECIFIC_ISOLATE
+		{
+			tsr_simd_short2 PixelPos = ComputePixelPos();
+
+			FetchDecimateMask(PixelPos, /* out */ VelocityEdge, /* out */ bIsDisoccluded);
+			FetchSceneColorAndTranslucency(PixelPos, /* out */ OriginalOpaqueInput, /* out */ OriginalTranslucencyInput);
+			FetchReprojectedHistoryGuide(
+				PixelPos,
+				/* out */ OriginalHistory,
+				/* out */ OriginalHistoryUncertainty);
+			#if CONFIG_HISTORY_RESURRECTION
+				OriginalResurrectedHistory            = tsr_simd_halfC::Const(0.0);
+				OriginalResurrectedHistoryUncertainty = tsr_simd_half::Const(0.0);
+			#endif
+		}
 	}
 	
 	// Compute AA Luma and compose separate translucency into scene color for optimised HistoryUpdate.
@@ -1145,9 +1197,7 @@ void MainCS(
 		{
 			tsr_simd_short2 PixelPos = ComputePixelPos();
 		
-			FetchHistoryGuide(
-				ResurrectedHistoryGuideTexture,
-				ReprojectedHistoryGuideMetadataTexture,
+			FetchResurrectedHistoryGuide(
 				PixelPos,
 				/* out */ ResurrectedHistory,
 				/* out */ ResurrectedHistoryUncertainty);
@@ -1232,9 +1282,7 @@ void MainCS(
 				tsr_simd_short2 PixePos = ComputePixelPos();
 				
 				#if CONFIG_CACHE_REFETCH
-					FetchHistoryGuide(
-						ResurrectedHistoryGuideTexture,
-						ReprojectedHistoryGuideMetadataTexture,
+					FetchResurrectedHistoryGuide(
 						PixePos,
 						/* out */ ResurrectedHistory,
 						/* out */ ResurrectedHistoryUncertainty);
