Fix a bug that the bloom luminance is not scaled by Exposure.

diff --git a/Engine/Shaders/Private/PostProcessMobile.usf b/Engine/Shaders/Private/PostProcessMobile.usf
--- a/Engine/Shaders/Private/PostProcessMobile.usf
+++ b/Engine/Shaders/Private/PostProcessMobile.usf
@@ -324,7 +324,7 @@ void BloomPS_Mobile(
 	#endif
 
 	// Try to kill negatives and NaNs here
-	AverageColor.rgb = max(AverageColor.rgb, 0) * View.OneOverPreExposure;
+	AverageColor.rgb = max(AverageColor.rgb, 0);
 
 	
 	#if MOBILE_USEBLOOM || MOBILE_USESUN
@@ -338,8 +338,6 @@ void BloomPS_Mobile(
 		OutColor.rgb = AverageColor;
 		OutColor.rgb *= Amount;
 		OutColor.a = 0;
-
-		OutColor = OutColor * View.PreExposure;
 	#endif
 #endif
 
@@ -393,7 +391,7 @@ void BloomPS_Mobile(
 #endif
 
 #if MOBILE_USEEYEADAPTATION
-	half Intensity = dot(AverageColor, half3(1.0f, 1.0f, 1.0f) / 3.0f);
+	half Intensity = dot(AverageColor * View.OneOverPreExposure, half3(1.0f, 1.0f, 1.0f) / 3.0f);
 	if (Intensity == 0.0)
 	{
 		// Override intensity for black pixels to MidGray. This behaviour is controlled by the cvar "r.EyeAdaptation.BlackHistogramBucketInfluence"
