Substrate - fixed compilation issue.

diff --git a/Engine/Shaders/Private/HairStrands/HairStrandsVisibilityDepthPS.usf b/Engine/Shaders/Private/HairStrands/HairStrandsVisibilityDepthPS.usf
--- a/Engine/Shaders/Private/HairStrands/HairStrandsVisibilityDepthPS.usf
+++ b/Engine/Shaders/Private/HairStrands/HairStrandsVisibilityDepthPS.usf
@@ -80,11 +80,8 @@ void MainPS(
 	const FHairSampleData SampleData = ExtractSampleData(PixelCoord);
 	SvDepth = SampleData.DeviceZ - DepthBias;
 #if SUBSTRATE_ENABLED
-	Output0 = 0; // MaterialData
+	Output0 = 0; // MaterialData SUBSTRATE_TODO Force PerObjectData contact shadow to respect hair settings, in order to have correct screen cast later
 	Output1 = 0; // TopLayerData
-
-	// Force PerObjectData contact shadow to respect hair settings, in order to have correct screen cast later
-	Output0 = SetCastContactShadow(SampleData.CastContactShadow);
 #else
 	Output0 = float4(0, 0, 0, 0);
 	Output1 = 0;
