Integrate NdoL in the MergeLocalLight when done inline in the basepass

diff --git a/Engine/Shaders/Private/MobileBasePassPixelShader.usf b/Engine/Shaders/Private/MobileBasePassPixelShader.usf
--- a/Engine/Shaders/Private/MobileBasePassPixelShader.usf
+++ b/Engine/Shaders/Private/MobileBasePassPixelShader.usf
@@ -106,6 +106,7 @@
 #include "MobileBasePassCommon.ush"
 #include "/Engine/Generated/VertexFactory.ush"
 #include "LightmapCommon.ush"  
+#define ENABLE_NDOTL_INTEGRATION 1
 #include "MobileLightingCommon.ush"
 #include "ShadingModelsMaterial.ush"
 #include "ThinTranslucentCommon.ush"
@@ -902,7 +903,7 @@ void Main(
 		float3 MergedLightL = float3(0,0,0);
 		float MergedSpecularScale = 0.f;
 		float MergedTotalWeight = 0;
-		MergeLocalLights(CulledLightGridData, MaterialParameters.WorldPosition_CamRelative, EyeIndex, false, MergedLightColor, MergedLightL, MergedSpecularScale, MergedTotalWeight);
+		MergeLocalLights(CulledLightGridData, MaterialParameters.WorldPosition_CamRelative, EyeIndex, false, MergedLightColor, MergedLightL, MergedSpecularScale, MergedTotalWeight, GBuffer.WorldNormal);
 
 		if(MergedTotalWeight > 0.f)
 		{
diff --git a/Engine/Shaders/Private/MobileLightingCommon.ush b/Engine/Shaders/Private/MobileLightingCommon.ush
--- a/Engine/Shaders/Private/MobileLightingCommon.ush
+++ b/Engine/Shaders/Private/MobileLightingCommon.ush
@@ -611,7 +611,11 @@ void MergeLocalLights(const FCulledLightsGridData InLightGridData
 	, inout float3 MergedLightColor
 	, inout float3 MergedLightL
 	, inout float MergedSpecularScale
-	, inout float MergedTotalWeight)
+	, inout float MergedTotalWeight
+#if ENABLE_NDOTL_INTEGRATION
+	, inout float3 N
+#endif
+)
 {
 	MergedLightColor = uint3(0.f, 0.f , 0.f);
 	MergedLightL = float3(0,0,0);
@@ -653,7 +657,13 @@ void MergeLocalLights(const FCulledLightsGridData InLightGridData
 		{
 			WeightLight *= SpotAttenuation(L, -LightData.Direction, LightData.SpotAngles);
 		}
-		
+
+		#if ENABLE_NDOTL_INTEGRATION
+			half NoL = max(0, dot(N, L));
+			WeightLight = WeightLight * NoL;
+			MaskedLightColor = MaskedLightColor * NoL;
+		#endif
+
 		if( LightMask > 0)
 		{
 			MergedSpecularScale += LightData.SpecularScale * WeightLight;
