Stochastic Direct Lighting - fixed second moment clamping and optimized spatial filter

diff --git a/Engine/Shaders/Private/StochasticDirectLighting/StochasticDirectLightingSpatial.usf b/Engine/Shaders/Private/StochasticDirectLighting/StochasticDirectLightingSpatial.usf
--- a/Engine/Shaders/Private/StochasticDirectLighting/StochasticDirectLightingSpatial.usf
+++ b/Engine/Shaders/Private/StochasticDirectLighting/StochasticDirectLightingSpatial.usf
@@ -48,8 +48,8 @@ void SDLSpatialFilterCS(
 			float SpecularLightingWeightSum = 1.0f;
 
 			// #sdl_todo: expose as CVars
-			float KernelRadiusInPixels = 16.0f;
-			uint NumSamples = 8;
+			float KernelRadiusInPixels = 8.0f;
+			uint NumSamples = 4;
 
 			float GuassianNormalize = 2.0f / (KernelRadiusInPixels * KernelRadiusInPixels);
 			uint2 RandomSeed = Rand3DPCG16(int3(ScreenCoord, StochasticDirectLightingStateFrameIndex)).xy;
@@ -57,8 +57,11 @@ void SDLSpatialFilterCS(
 			float3 TranslatedWorldPosition = GetTranslatedWorldPositionFromScreenUV(ScreenUV, GBuffer.Depth);
 			float4 ScenePlane = float4(GBuffer.WorldNormal, dot(TranslatedWorldPosition, GBuffer.WorldNormal));
 
+			// #sdl_todo: expose as CVars
+			float MinStdDev = 0.1f;
+
 			#if SPATIAL_FILTER
-			if (KernelRadiusInPixels > 0.0f)
+			if (KernelRadiusInPixels > 0.0f && (CenterDiffuseStdDev > MinStdDev || CenterSpecularStdDev > MinStdDev))
 			{
 				for (uint NeighborIndex = 0; NeighborIndex < NumSamples; ++NeighborIndex)
 				{
@@ -82,11 +85,12 @@ void SDLSpatialFilterCS(
 						float NormalWeight = 1.0f - saturate(AngleBetweenNormals);
 
 						// Diffuse
+						if (CenterDiffuseStdDev > MinStdDev)
 						{
 							float4 NeighborDiffuseLightingAndSecondMoment = DiffuseLightingAndSecondMomentTexture[NeighborCoord];
 
 							float LuminanceDelta = abs(Luminance(CenterDiffuseLightingAndSecondMoment.xyz) - Luminance(NeighborDiffuseLightingAndSecondMoment.xyz));
-							float LuminanceWeight = exp2(-LuminanceDelta / max(1.0f * CenterDiffuseStdDev, 0.00001f));
+							float LuminanceWeight = exp2(-LuminanceDelta / max(CenterDiffuseStdDev, 0.001f));
 
 							const float DiffuseNeighborWeight = DepthWeight * NormalWeight * LuminanceWeight;
 							DiffuseLightingSum += NeighborDiffuseLightingAndSecondMoment.xyz * DiffuseNeighborWeight;
@@ -94,11 +98,12 @@ void SDLSpatialFilterCS(
 						}
 
 						// Specular
+						if (CenterSpecularStdDev > MinStdDev)
 						{
 							float4 NeighborSpecularLightingAndSecondMoment = SpecularLightingAndSecondMomentTexture[NeighborCoord];
 
 							float LuminanceDelta = abs(Luminance(CenterSpecularLightingAndSecondMoment.xyz) - Luminance(NeighborSpecularLightingAndSecondMoment.xyz));
-							float LuminanceWeight = exp2(-LuminanceDelta / max(1.0f * CenterSpecularStdDev, 0.00001f));
+							float LuminanceWeight = exp2(-LuminanceDelta / max(CenterSpecularStdDev, 0.001f));
 
 							const float SpecularNeighborWeight = DepthWeight * NormalWeight * LuminanceWeight;
 							SpecularLightingSum += NeighborSpecularLightingAndSecondMoment.xyz * SpecularNeighborWeight;
@@ -127,6 +132,14 @@ void SDLSpatialFilterCS(
 				Newline(Context);
 				Print(Context, TEXT("AccumulatedSceneColor: "));
 				Print(Context, AccumulatedSceneColor);
+				Newline(Context);
+				Print(Context, CenterSpecularLightingAndSecondMoment);
+				Newline(Context);
+				Print(Context, TEXT("CenterDiffuseStdDev:  "));
+				Print(Context, CenterDiffuseStdDev);
+				Newline(Context);
+				Print(Context, TEXT("CenterSpecularStdDev: "));
+				Print(Context, CenterSpecularStdDev);
 			}
 			#endif
 
diff --git a/Engine/Shaders/Private/StochasticDirectLighting/StochasticDirectLightingTemporal.usf b/Engine/Shaders/Private/StochasticDirectLighting/StochasticDirectLightingTemporal.usf
--- a/Engine/Shaders/Private/StochasticDirectLighting/StochasticDirectLightingTemporal.usf
+++ b/Engine/Shaders/Private/StochasticDirectLighting/StochasticDirectLightingTemporal.usf
@@ -210,16 +210,16 @@ void SDLTemporalAccumulationCS(
 						NumFramesAccumulated = min(NumFramesAccumulatedHistory + 1.0f, TemporalMaxFramesAccumulated);
 
 						// Clamp history to current neighborhood
-						FNeighborhood DiffuseNeighborhood = GetNeighborhood(ScreenCoord, ResolvedDiffuseLighting, float4(DiffuseLighting, DiffuseSecondMoment).xyz);
-						FNeighborhood SpecularNeighborhood = GetNeighborhood(ScreenCoord, ResolvedSpecularLighting, float4(SpecularLighting, SpecularSecondMoment).xyz);
+						FNeighborhood DiffuseNeighborhood = GetNeighborhood(ScreenCoord, ResolvedDiffuseLighting, DiffuseLighting);
+						FNeighborhood SpecularNeighborhood = GetNeighborhood(ScreenCoord, ResolvedSpecularLighting, SpecularLighting);
 						float3 ClampedHistoryDiffuseLighting = clamp(HistoryDiffuseLighting, DiffuseNeighborhood.ClampMin, DiffuseNeighborhood.ClampMax);
 						float3 ClampedHistorySpecularLighting = clamp(HistorySpecularLighting, SpecularNeighborhood.ClampMin, SpecularNeighborhood.ClampMax);
 
 						// Clamp history to max sample luminance
 						ClampedHistoryDiffuseLighting = ClampLuminance(ClampedHistoryDiffuseLighting, SampleLuminanceSum.x);
 						ClampedHistorySpecularLighting = ClampLuminance(ClampedHistorySpecularLighting, SampleLuminanceSum.y);
-						float ClampedHistoryDiffuseSecondMoment = clamp(ClampedHistoryDiffuseSecondMoment, 0, Pow2(SampleLuminanceSum.x));
-						float ClampedHistorySpecularSecondMoment = clamp(ClampedHistorySpecularSecondMoment, 0, Pow2(SampleLuminanceSum.y));
+						float ClampedHistoryDiffuseSecondMoment = clamp(HistoryDiffuseSecondMoment, 0, Pow2(SampleLuminanceSum.x));
+						float ClampedHistorySpecularSecondMoment = clamp(HistorySpecularSecondMoment, 0, Pow2(SampleLuminanceSum.y));
 
 						// Blend history with new samples
 						float Alpha = 1.0f / (1.0f + NumFramesAccumulated);
@@ -251,6 +251,12 @@ void SDLTemporalAccumulationCS(
 				Print(Context, TEXT("SpecularVariance: "));
 				Print(Context, SpecularVariance);
 				Newline(Context);
+				Print(Context, TEXT("Specular Second Moment: "));
+				Print(Context, SpecularSecondMoment);
+				Newline(Context);
+				Print(Context, TEXT("Specular Lighting: "));
+				Print(Context, SpecularLighting);
+				Newline(Context);
 				Print(Context, TEXT("SampleLuminanceSum: "));
 				Print(Context, SampleLuminanceSum);
 				Newline(Context);
