Fix missing contact shadow casting on hair with Substrate.

diff --git a/Engine/Shaders/Private/HairStrands/HairStrandsVisibilityDepthPS.usf b/Engine/Shaders/Private/HairStrands/HairStrandsVisibilityDepthPS.usf
--- a/Engine/Shaders/Private/HairStrands/HairStrandsVisibilityDepthPS.usf
+++ b/Engine/Shaders/Private/HairStrands/HairStrandsVisibilityDepthPS.usf
@@ -3,6 +3,7 @@
 #include "../Common.ush"
 #include "../SceneTexturesCommon.ush"
 #include "../DeferredShadingCommon.ush"
+#include "../Substrate/Substrate.ush"
 #include "HairStrandsVisibilityCommon.ush"
 
 #define OUTPUT_DEPTH 0
@@ -80,8 +81,14 @@ void MainPS(
 	const FHairSampleData SampleData = ExtractSampleData(PixelCoord);
 	SvDepth = SampleData.DeviceZ - DepthBias;
 #if SUBSTRATE_ENABLED
-	Output0 = 0; // MaterialData SUBSTRATE_TODO Force PerObjectData contact shadow to respect hair settings, in order to have correct screen cast later
+	Output0 = 0; // MaterialData
 	Output1 = 0; // TopLayerData
+
+	// Force PerObjectData contact shadow to respect hair settings, in order to have correct screen cast later
+	FSubstratePixelHeader DummyHeader = InitialiseSubstratePixelHeader();
+	DummyHeader.SetCastContactShadow(SampleData.CastContactShadow ? 1 : 0);
+	Output0 = DummyHeader.State;
+
 #else
 	Output0 = float4(0, 0, 0, 0);
 	Output1 = 0;
diff --git a/Engine/Source/Runtime/Renderer/Private/HairStrands/HairStrandsVisibility.cpp b/Engine/Source/Runtime/Renderer/Private/HairStrands/HairStrandsVisibility.cpp
--- a/Engine/Source/Runtime/Renderer/Private/HairStrands/HairStrandsVisibility.cpp
+++ b/Engine/Source/Runtime/Renderer/Private/HairStrands/HairStrandsVisibility.cpp
@@ -2778,7 +2778,10 @@ static void AddHairAuxilaryPass(
 		Parameters->RenderTargets[0] = FRenderTargetBinding(Output0, ERenderTargetLoadAction::ELoad, 0 /*Mip-0*/, 0/*First slice, which contains state/counter*/);
 		Parameters->RenderTargets[1] = FRenderTargetBinding(Output1, ERenderTargetLoadAction::ELoad);
 		Parameters->RenderTargets[2] = FRenderTargetBinding(OutColorTexture, ERenderTargetLoadAction::ELoad);
-		Parameters->RenderTargets[3] = FRenderTargetBinding(Output2, ERenderTargetLoadAction::ELoad);
+		if (!bSubstrateEnabled)
+		{
+			Parameters->RenderTargets[3] = FRenderTargetBinding(Output2, ERenderTargetLoadAction::ELoad);
+		}
 	}
 
 	if (PassType == EHairAuxilaryPassType::MaterialData_LightChannelMask || PassType == EHairAuxilaryPassType::LightChannelMask)
