Fixes a bunch of missing alpha channel quantization 0 and 1 across DOF, TSR and MB

diff --git a/Engine/Shaders/Private/DiaphragmDOF/DOFRecombine.usf b/Engine/Shaders/Private/DiaphragmDOF/DOFRecombine.usf
--- a/Engine/Shaders/Private/DiaphragmDOF/DOFRecombine.usf
+++ b/Engine/Shaders/Private/DiaphragmDOF/DOFRecombine.usf
@@ -1136,10 +1136,10 @@ void RecombineMainCS(
 
 	// Ensure that alpha values that are expected to be opaque (but are only close to opaque) are forced to be opaque.
 	// (0.005 chosen to accomodate handling of 1/255)
-	if (OutputFinalSceneColor.a < 0.005)
-	{
-		OutputFinalSceneColor.a = 0;
-	}
+	#if CONFIG_DOF_ALPHA
+		OutputFinalSceneColor.a = select(OutputFinalSceneColor.a < 0.005, 0.0, OutputFinalSceneColor.a);
+		OutputFinalSceneColor.a = select(OutputFinalSceneColor.a > 0.995, 1.0, OutputFinalSceneColor.a);
+	#endif
 	
 	// Debug optimisation colors.
 	#if 0
diff --git a/Engine/Shaders/Private/MotionBlur/MotionBlurApply.usf b/Engine/Shaders/Private/MotionBlur/MotionBlurApply.usf
--- a/Engine/Shaders/Private/MotionBlur/MotionBlurApply.usf
+++ b/Engine/Shaders/Private/MotionBlur/MotionBlurApply.usf
@@ -884,8 +884,22 @@ void MainCS(
 			OutputColor3.rgb = OutputColor3.rgb * PostMotionBlurTranslucency3.a + PostMotionBlurTranslucency3.rgb;
 		}
 		#endif
-
-		#if !POST_PROCESS_ALPHA
+		
+		// Ensure that alpha values that are expected to be opaque (but are only close to opaque) are forced to be opaque.
+		// (0.995 chosen to accommodate handling of 254/255)
+		#if POST_PROCESS_ALPHA
+		{
+			OutputColor0[3] = select(OutputColor0[3] > mb_half(0.995), tsr_half(1.0), OutputColor0[3]);
+			OutputColor1[3] = select(OutputColor1[3] > mb_half(0.995), tsr_half(1.0), OutputColor1[3]);
+			OutputColor2[3] = select(OutputColor2[3] > mb_half(0.995), tsr_half(1.0), OutputColor2[3]);
+			OutputColor3[3] = select(OutputColor3[3] > mb_half(0.995), tsr_half(1.0), OutputColor3[3]);
+
+			OutputColor0[3] = select(OutputColor0[3] < mb_half(0.005), tsr_half(0.0), OutputColor0[3]);
+			OutputColor1[3] = select(OutputColor1[3] < mb_half(0.005), tsr_half(0.0), OutputColor1[3]);
+			OutputColor2[3] = select(OutputColor2[3] < mb_half(0.005), tsr_half(0.0), OutputColor2[3]);
+			OutputColor3[3] = select(OutputColor3[3] < mb_half(0.005), tsr_half(0.0), OutputColor3[3]);
+		}
+		#else
 		{
 			OutputColor0.a = 0;
 			OutputColor1.a = 0;
@@ -995,7 +1009,14 @@ void MainCS(
 		}
 		#endif
 
-		#if !POST_PROCESS_ALPHA
+		// Ensure that alpha values that are expected to be opaque (but are only close to opaque) are forced to be opaque.
+		// (0.995 chosen to accommodate handling of 254/255)
+		#if POST_PROCESS_ALPHA
+		{
+			OutputColor[3] = select(OutputColor[3] > mb_half(0.995), tsr_half(1.0), OutputColor[3]);
+			OutputColor[3] = select(OutputColor[3] < mb_half(0.005), tsr_half(0.0), OutputColor[3]);
+		}
+		#else
 		{
 			OutputColor.a = 0;
 		}
diff --git a/Engine/Shaders/Private/TemporalSuperResolution/TSRResolveHistory.usf b/Engine/Shaders/Private/TemporalSuperResolution/TSRResolveHistory.usf
--- a/Engine/Shaders/Private/TemporalSuperResolution/TSRResolveHistory.usf
+++ b/Engine/Shaders/Private/TemporalSuperResolution/TSRResolveHistory.usf
@@ -36,6 +36,7 @@
 
 	#define tsr_input_halfC  TLaneVector2D<tsr_half, CONFIG_CHANNEL_COUNT, LANE_INPUT_STRIDE_X, LANE_INPUT_STRIDE_Y>
 	#define tsr_output_halfC TLaneVector2D<tsr_half, CONFIG_CHANNEL_COUNT, LANE_OUTPUT_STRIDE_X, LANE_OUTPUT_STRIDE_Y>
+	#define tsr_output_half  TLaneVector2D<tsr_half, 1, LANE_OUTPUT_STRIDE_X, LANE_OUTPUT_STRIDE_Y>
 
 #else
 	#define TILE_SIZE 8
@@ -246,7 +247,16 @@ void MainCS(
 			tsr_short2(OutputViewRectMin) +
 			tsr_short2(GroupId) * tsr_short(TILE_SIZE).xx);
 		
-		Output = min(Output, tsr_output_halfC::Const(tsr_half(Max10BitsFloat)));
+		Output = clamp(Output, tsr_output_halfC::Const(tsr_half(0.0)), tsr_output_halfC::Const(tsr_half(Max10BitsFloat)));
+
+		// Ensure that alpha values that are expected to be opaque (but are only close to opaque) are forced to be opaque.
+		// (0.995 chosen to accommodate handling of 254/255)
+		#if CONFIG_SCENE_COLOR_ALPHA
+		{
+			Output.SetComponent(3, select(Output[3] > tsr_output_half::Const(0.995), tsr_output_half::Const(1.0), Output[3]));
+			Output.SetComponent(3, select(Output[3] < tsr_output_half::Const(0.005), tsr_output_half::Const(0.0), Output[3]));
+		}
+		#endif
 
 		UNROLL_N(OUTPUT_SIMD_SIZE)
 		for (uint SimdIndex = 0; SimdIndex < OUTPUT_SIMD_SIZE; SimdIndex++)
@@ -347,6 +357,17 @@ void MainCS(
 	
 	// Compute final output
 	tsr_halfC FinalOutputColor = Color;
+	
+	FinalOutputColor = clamp(FinalOutputColor, tsr_half(0.0), tsr_half(Max10BitsFloat));
+
+	// Ensure that alpha values that are expected to be opaque (but are only close to opaque) are forced to be opaque.
+	// (0.995 chosen to accommodate handling of 254/255)
+	#if CONFIG_SCENE_COLOR_ALPHA
+	{
+		FinalOutputColor[3] = select(FinalOutputColor[3] > tsr_half(0.995), tsr_half(1.0), FinalOutputColor[3]);
+		FinalOutputColor[3] = select(FinalOutputColor[3] < tsr_half(0.005), tsr_half(0.0), FinalOutputColor[3]);
+	}
+	#endif
 
 	tsr_short2 OutputPixelPosition = InvalidateOutputPixelPos(tsr_short2(OutputViewRectMin + DispatchThreadId), OutputViewRectMax);
 	
diff --git a/Engine/Shaders/Private/TemporalSuperResolution/TSRUpdateHistory.usf b/Engine/Shaders/Private/TemporalSuperResolution/TSRUpdateHistory.usf
--- a/Engine/Shaders/Private/TemporalSuperResolution/TSRUpdateHistory.usf
+++ b/Engine/Shaders/Private/TemporalSuperResolution/TSRUpdateHistory.usf
@@ -63,12 +63,6 @@
 	#define CONFIG_INTERPOLATE_VELOCITY 0
 #endif
 
-#if 0
-	#define CONFIG_OUTPUT_MERGED_SCENE_COLOR 1
-#else
-	#define CONFIG_OUTPUT_MERGED_SCENE_COLOR 0
-#endif
-
 
 //------------------------------------------------------- CONSTANTS
 
@@ -1017,32 +1011,6 @@ void MainCS(
 		FinalHistoryValidity = ceil(tsr_half(255.0) * OutputValidity) * rcp(tsr_half(255.0));
 	}
 
-	// Compute final output
-	tsr_halfCx2 FinalOutputColor;
-	{
-		FinalOutputColor = FinalHighFrequencyColor;
-		
-		// Ensure that alpha values that are expected to be opaque (but are only close to opaque) are forced to be opaque.
-		// (0.995 chosen to accommodate handling of 254/255)
-		#if CONFIG_SCENE_COLOR_ALPHA
-		{
-			FinalOutputColor[3] = select(FinalOutputColor[3] > 0.995, tsr_half(1.0), FinalOutputColor[3]);
-			FinalOutputColor[3] = select(FinalOutputColor[3] < 0.005, tsr_half(0.0), FinalOutputColor[3]);
-		}
-		#endif
-		
-		#if 0
-		{
-			//tsr_half2 DebugValue = tsr_half(1.0) - tsr_half2(bIsDisoccluded);
-			tsr_half2 DebugValue = LowFrequencyRejection;
-			FinalOutputColor *= 
-				dpv_scale(dpv_interleave_mono_registers(tsr_half3(1, 0.5, 0.5)), tsr_half(1.0) - DebugValue) +
-				dpv_scale(dpv_interleave_mono_registers(tsr_half3(0.5, 1, 0.5)), DebugValue);
-		}
-		#endif
-	}
-
-
 	PLATFORM_SPECIFIC_ISOLATE
 	{
 		#if CONFIG_SCENE_COLOR_ALPHA
@@ -1079,7 +1047,6 @@ void MainCS(
 			uint2 Random = Rand3DPCG16(int3(dpv_lo(LocalHistoryPixelPos), View.StateFrameIndexMod8)).xy;
 			tsr_half E = tsr_half(Hammersley16(0, 1, Random).x);
 			
-			FinalOutputColor = QuantizeForFloatRenderTarget(FinalOutputColor, E, HistoryQuantizationError);
 			FinalHighFrequencyColor = QuantizeForFloatRenderTarget(FinalHighFrequencyColor, E, HistoryQuantizationError);
 		}
 		#endif
@@ -1090,6 +1057,15 @@ void MainCS(
 			FinalHighFrequencyColor = dpv_min(FinalHighFrequencyColor, ColorMax10BitsFloat);
 		}
 
+		// Ensure that alpha values that are expected to be opaque (but are only close to opaque) are forced to be opaque.
+		// (0.995 chosen to accommodate handling of 254/255)
+		#if CONFIG_SCENE_COLOR_ALPHA
+		{
+			FinalHighFrequencyColor[3] = select(FinalHighFrequencyColor[3] > tsr_half(0.995), tsr_half(1.0), FinalHighFrequencyColor[3]);
+			FinalHighFrequencyColor[3] = select(FinalHighFrequencyColor[3] < tsr_half(0.005), tsr_half(0.0), FinalHighFrequencyColor[3]);
+		}
+		#endif
+
 		// Output full res history
 		{
 			// Output final history lo pixel.
@@ -1107,26 +1083,9 @@ void MainCS(
 			#endif // !CONFIG_ENABLE_DUAL_PIXEL_VECTORIZATION
 		}
 
-		// Output final scene color Mip0
-		tsr_halfCx2 OutputColor;
-		{
-			OutputColor = FinalOutputColor;
-
-			OutputColor = -dpv_min(-OutputColor, ColorNull);
-			OutputColor = dpv_min(OutputColor, ColorMax10BitsFloat);
-
-			#if CONFIG_OUTPUT_MERGED_SCENE_COLOR
-				SceneColorOutputMip0[dpv_lo(LocalHistoryPixelPos)] = dpv_lo(OutputColor);
-			
-				#if CONFIG_ENABLE_DUAL_PIXEL_VECTORIZATION
-					SceneColorOutputMip0[dpv_hi(LocalHistoryPixelPos)] = dpv_hi(OutputColor);
-				#endif
-			#endif
-		}
-		
 		// Output final scene color Mip1
 		{
-			tsr_halfCx2 HalfResOutput = OutputColor;
+			tsr_halfCx2 HalfResOutput = FinalHighFrequencyColor;
 			tsr_short2x2 HalfResOutputPixelPos = dpv_interleave_mono_registers(tsr_short(-1).xx);
 
 			BRANCH
