GPULightmass: Fix missing preprocessor case leading to blending errors on masked materials

diff --git a/Engine/Shaders/Private/PathTracing/PathTracingGPULightmassMaterialHitShader.usf b/Engine/Shaders/Private/PathTracing/PathTracingGPULightmassMaterialHitShader.usf
--- a/Engine/Shaders/Private/PathTracing/PathTracingGPULightmassMaterialHitShader.usf
+++ b/Engine/Shaders/Private/PathTracing/PathTracingGPULightmassMaterialHitShader.usf
@@ -386,6 +386,8 @@ RAY_TRACING_ENTRY_ANY_HIT(GPULightmassMaterialAHS,
 	}
 #endif
 
+#if MATERIALBLENDING_MODULATE || (MATERIALBLENDING_MASKED && MATERIAL_DITHER_OPACITY_MASK) || MATERIALBLENDING_ALPHACOMPOSITE || MATERIALBLENDING_TRANSLUCENT
+
 #if SUBSTRATE_ENABLED
 
 	float  Coverage = 0;
@@ -475,6 +477,7 @@ RAY_TRACING_ENTRY_ANY_HIT(GPULightmassMaterialAHS,
 		// if there is energy left in the ray, 
 		IgnoreHit();
 	}
+#endif // MATERIALBLENDING_MODULATE || (MATERIALBLENDING_MASKED && MATERIAL_DITHER_OPACITY_MASK) || MATERIALBLENDING_ALPHACOMPOSITE || MATERIALBLENDING_TRANSLUCENT
 }
 
 #endif // USE_MATERIAL_ANY_HIT_SHADER
