Path Tracer: Refactor the math for tracking the alpha channel

Instead of tracking background visibility, make sure we track the alpha channel directly. This changes how holdouts are handled, in particular for the volumetric cases.

To maintain compatibility with the rest of the engine, the 1-Alpha convention happens in the display shader.

Also, fix holdout on skyatmosphere when not using the reference atmosphere mode (by simply hiding the skylight in this case)

diff --git a/Engine/Shaders/Private/PathTracing/PathTracing.usf b/Engine/Shaders/Private/PathTracing/PathTracing.usf
--- a/Engine/Shaders/Private/PathTracing/PathTracing.usf
+++ b/Engine/Shaders/Private/PathTracing/PathTracing.usf
@@ -52,7 +52,7 @@ FPathState LoadPathStateData(uint Index)
 	Output.RandSequence.SampleIndex = Data.RandSeqSampleIndex;
 	Output.RandSequence.SampleSeed  = Data.RandSeqSampleSeed;
 	Output.Radiance = Data.Radiance;
-	Output.BackgroundVisibility = Data.BackgroundVisibility;
+	Output.Alpha = Data.Alpha;
 	Output.Albedo.x = f16tof32(Data.PackedAlbedoNormal.x);
 	Output.Albedo.y = f16tof32(Data.PackedAlbedoNormal.x >> 16);
 	Output.Albedo.z = f16tof32(Data.PackedAlbedoNormal.y);
@@ -84,7 +84,7 @@ void StorePathStateData(FPathState PathState, uint Index)
 	Data.RandSeqSampleIndex = PathState.RandSequence.SampleIndex;
 	Data.RandSeqSampleSeed  = PathState.RandSequence.SampleSeed;
 	Data.Radiance = PathState.Radiance;
-	Data.BackgroundVisibility = PathState.BackgroundVisibility;
+	Data.Alpha = PathState.Alpha;
 	Data.PackedAlbedoNormal.x = f32tof16(PathState.Albedo.x) | (f32tof16(PathState.Albedo.y) << 16);
 	Data.PackedAlbedoNormal.y = f32tof16(PathState.Albedo.z) | (f32tof16(PathState.Normal.x) << 16);
 	Data.PackedAlbedoNormal.z = f32tof16(PathState.Normal.y) | (f32tof16(PathState.Normal.z) << 16);
diff --git a/Engine/Shaders/Private/PathTracing/PathTracingCompositingPixelShader.usf b/Engine/Shaders/Private/PathTracing/PathTracingCompositingPixelShader.usf
--- a/Engine/Shaders/Private/PathTracing/PathTracingCompositingPixelShader.usf
+++ b/Engine/Shaders/Private/PathTracing/PathTracingCompositingPixelShader.usf
@@ -27,7 +27,8 @@ void CompositeMain(
 
 	float4 Radiance = RadianceTexture.Load(TexCoord);
 
-	OutColor = float4(Radiance.rgb * View.PreExposure, Radiance.a);
+	// NOTE: UE has the convention of using 1.0 - Alpha
+	OutColor = float4(Radiance.rgb * View.PreExposure, 1.0 - Radiance.a);
 
 	if (AdaptiveSamplingVisualize == 1)
 	{
diff --git a/Engine/Shaders/Private/PathTracing/PathTracingCore.ush b/Engine/Shaders/Private/PathTracing/PathTracingCore.ush
--- a/Engine/Shaders/Private/PathTracing/PathTracingCore.ush
+++ b/Engine/Shaders/Private/PathTracing/PathTracingCore.ush
@@ -81,7 +81,7 @@ struct FPathState
 {                                   // packed size
 	RandomSequence RandSequence;    //  8 bytes
 	float3 Radiance;                // 12 bytes
-	float  BackgroundVisibility;    //  4 bytes
+	float  Alpha;                   //  4 bytes
 	float3 Albedo;                  //  6 bytes // half precision
 	float3 Normal;                  //  6 bytes // half precision
 	FRayDesc Ray;                   // 12+12 = 24 bytes (TMin/TMax do not need to be stored)
@@ -188,7 +188,7 @@ struct FPathState
 		float4 OldAlbedo = Iteration > 0 ? AlbedoTexture[TextureIndex]   : 0;
 		float4 OldNormal = Iteration > 0 ? NormalTexture[TextureIndex]   : 0;
 
-		const float4 PixelValue = float4(Radiance, saturate(BackgroundVisibility));
+		const float4 PixelValue = float4(Radiance, saturate(Alpha));
 
 		RadianceTexture[TextureIndex] = lerp(OldPixel, PixelValue, Blend);
 		AlbedoTexture[TextureIndex]   = lerp(OldAlbedo, float4(Albedo, 0), Blend);
@@ -452,16 +452,14 @@ FPathTracingPayload TraceTransparentRay(inout FPathState PathState, int Bounce,
 			}
 
 			// NOTE: when using analytic raymarching, we can skip all raymarching on the last bounce because we know we don't need a scattering point
-			const bool bUseAnalyticTransmittance = (VolumeFlags & PATH_TRACER_VOLUME_USE_ANALYTIC_TRANSMITTANCE) != 0;
-			const bool bUseNullTrackingTransmittance = !bUseAnalyticTransmittance;
-			if (bUseNullTrackingTransmittance || (bUseAnalyticTransmittance && !bLastBounce))
 			{
 				// remember the start of our interval before ray marching messes with it
 				float StartVolumeTMin = Interval.VolumeTMin;
 				// Ray marching loop
 				float3 PathThroughput = PathState.PathThroughput;
 				float3 VolumeRadiance = 0;
-				// Limit number of steps to prevent timeouts // FIXME: This biases the result!
+				// Limit number of steps to prevent timeouts
+				// TODO: This biases the result! Is there a better way to limit the number of steps?
 				for (int Step = 0; Step < MaxRaymarchSteps; Step++)
 				{
 					FVolumeTrackingResult TrackingResult = VolumeSampleDistance(PathThroughput, PathState.Ray.Origin, PathState.Ray.Direction, Interval, bIsCameraRay, PathState.RandSequence);
@@ -487,10 +485,11 @@ FPathTracingPayload TraceTransparentRay(inout FPathState PathState, int Bounce,
 
 						if (bIsCameraRay)
 						{
-							PathState.BackgroundVisibility += Luminance(PathThroughput * Result.SigmaH);
+							PathState.Alpha += Luminance(PathThroughput * (Result.SigmaT - Result.SigmaH));
 						}
 
-						if (bUseAnalyticTransmittance || (bUseNullTrackingTransmittance && !bLastBounce)) // check now if this loop is also computing transmittance
+						// if this is not the last bounce, consider the possibility of scattering as giving up medium hits
+						if (!bLastBounce)
 						{
 							float3 SigmaS = min(Result.SigmaSRayleigh + Result.SigmaSHG, SigmaT);
 							// accumulate a signal for the denoiser
@@ -522,18 +521,19 @@ FPathTracingPayload TraceTransparentRay(inout FPathState PathState, int Bounce,
 
 				// restore interval start now that we are finished with ray marching
 				Interval.VolumeTMin = StartVolumeTMin;
-				if (bUseNullTrackingTransmittance)
+
+				const bool bUseAnalyticTransmittance = (VolumeFlags & PATH_TRACER_VOLUME_USE_ANALYTIC_TRANSMITTANCE) != 0;
+				if (bUseAnalyticTransmittance)
+				{
+					// analytically handle the transmittance to the next surface (which will be higher quality than the transmittance implicitly computed above and prevent noise with emissive surfaces in volumes)
+					PathState.PathThroughput = VolumeGetTransmittance(PathState.PathThroughput, PathState.Ray.Origin, PathState.Ray.Direction, Interval, PathState.RandSequence);
+				}
+				else
 				{
 					// record the change in transmittance since we already computed it
 					PathState.PathThroughput = PathThroughput;
 				}
 			}
-			
-			if (bUseAnalyticTransmittance)
-			{
-				// analytically handle the transmittance to the next surface (which will be higher quality than the transmittance implicitly computed above and prevent noise with emissive surfaces in volumes)
-				PathState.PathThroughput = VolumeGetTransmittance(PathState.PathThroughput, PathState.Ray.Origin, PathState.Ray.Direction, Interval, PathState.RandSequence);
-			}
 
 			// Update our current volume intersection list to reflect the fact that we have made it to the end of this interval
 			// This will either clip the finished portions of active segments, or remove the segments we have fully marched through
@@ -544,11 +544,6 @@ FPathTracingPayload TraceTransparentRay(inout FPathState PathState, int Bounce,
 		if (PackedPayload.IsMiss())
 		{
 			// Ray didn't hit any real geometry
-			// Accumulate a hit against the sky (for camera rays) and exit
-			if (bIsCameraRay)
-			{
-				PathState.BackgroundVisibility += Luminance(PathState.PathThroughput);
-			}
 			break;
 		}
 
@@ -612,11 +607,12 @@ FPathTracingPayload TraceTransparentRay(inout FPathState PathState, int Bounce,
 
 		// account for Beer's law through the currently active medium
 		// TODO: merge this with the volume raymarching?
+		// NOTE: select statement is to avoid corner case of 0.0*inf that can occur with directly visible skydomes
 		PathState.PathThroughput *= select(PathState.SigmaT > 0.0, exp(-PathState.SigmaT * (HitPayload.HitT - PathState.Ray.TMin)), 1.0);
 
-		if (HitPayload.IsHoldout() && bIsCameraRay)
+		if (!HitPayload.IsHoldout() && bIsCameraRay)
 		{
-			PathState.BackgroundVisibility += Luminance(PathState.PathThroughput * HitPayload.BSDFOpacity);
+			PathState.Alpha += Luminance(PathState.PathThroughput * (1.0 - HitPayload.TransparencyColor));
 		}
 
 		// add in surface emission
@@ -627,8 +623,7 @@ FPathTracingPayload TraceTransparentRay(inout FPathState PathState, int Bounce,
 
 		if (HitPayload.HitT == POSITIVE_INFINITY)
 		{
-			// if our hit was against an infinite light, treat it as having background alpha, and exit
-			PathState.BackgroundVisibility += Luminance(PathState.PathThroughput);
+			// if our hit was against an infinite light, exit now
 			break;
 		}
 
diff --git a/Engine/Shaders/Shared/RayTracingTypes.h b/Engine/Shaders/Shared/RayTracingTypes.h
--- a/Engine/Shaders/Shared/RayTracingTypes.h
+++ b/Engine/Shaders/Shared/RayTracingTypes.h
@@ -66,7 +66,7 @@ struct FPathTracingPackedPathState {
 	uint      RandSeqSampleIndex;
 	uint      RandSeqSampleSeed;
 	float3    Radiance;
-	float     BackgroundVisibility;
+	float     Alpha;
 	uint3     PackedAlbedoNormal;
 	float3    RayOrigin;
 	float3    RayDirection;
diff --git a/Engine/Source/Runtime/Renderer/Private/PathTracing.cpp b/Engine/Source/Runtime/Renderer/Private/PathTracing.cpp
--- a/Engine/Source/Runtime/Renderer/Private/PathTracing.cpp
+++ b/Engine/Source/Runtime/Renderer/Private/PathTracing.cpp
@@ -2079,10 +2079,12 @@ void SetLightParameters(FRDGBuilder& GraphBuilder, FPathTracingRG::FParameters*
 		DestLight.VolumetricScatteringIntensity = Scene->SkyLight->VolumetricScatteringIntensity;
 		DestLight.IESAtlasIndex = INDEX_NONE;
 		DestLight.MissShaderIndex = 0;
-		if (Scene->SkyLight->bRealTimeCaptureEnabled || CVarPathTracingVisibleLights.GetValueOnRenderThread() == 2)
+		if ((Scene->SkyLight->bRealTimeCaptureEnabled && (View.SkyAtmosphereUniformShaderParameters == nullptr ||  !View.SkyAtmosphereUniformShaderParameters->bHoldout)) || CVarPathTracingVisibleLights.GetValueOnRenderThread() == 2)
 		{
 			// When using the realtime capture system, always make the skylight visible
 			// because this is our only way of "seeing" the atmo/clouds at the moment
+			// The one exception to this case is if the sky atmo has been marked as holdout.
+
 			// Also allow seeing just the sky via a cvar for debugging purposes
 			PassParameters->SceneVisibleLightCount = 1;
 		}
