Fix shader compilation error when Propagate Alpha is enabled: use mb_half instead of tsr_half

diff --git a/Engine/Shaders/Private/MotionBlur/MotionBlurApply.usf b/Engine/Shaders/Private/MotionBlur/MotionBlurApply.usf
--- a/Engine/Shaders/Private/MotionBlur/MotionBlurApply.usf
+++ b/Engine/Shaders/Private/MotionBlur/MotionBlurApply.usf
@@ -889,15 +889,15 @@ void MainCS(
 		// (0.995 chosen to accommodate handling of 254/255)
 		#if POST_PROCESS_ALPHA
 		{
-			OutputColor0[3] = select(OutputColor0[3] > mb_half(0.995), tsr_half(1.0), OutputColor0[3]);
-			OutputColor1[3] = select(OutputColor1[3] > mb_half(0.995), tsr_half(1.0), OutputColor1[3]);
-			OutputColor2[3] = select(OutputColor2[3] > mb_half(0.995), tsr_half(1.0), OutputColor2[3]);
-			OutputColor3[3] = select(OutputColor3[3] > mb_half(0.995), tsr_half(1.0), OutputColor3[3]);
-
-			OutputColor0[3] = select(OutputColor0[3] < mb_half(0.005), tsr_half(0.0), OutputColor0[3]);
-			OutputColor1[3] = select(OutputColor1[3] < mb_half(0.005), tsr_half(0.0), OutputColor1[3]);
-			OutputColor2[3] = select(OutputColor2[3] < mb_half(0.005), tsr_half(0.0), OutputColor2[3]);
-			OutputColor3[3] = select(OutputColor3[3] < mb_half(0.005), tsr_half(0.0), OutputColor3[3]);
+			OutputColor0[3] = select(OutputColor0[3] > mb_half(0.995), mb_half(1.0), OutputColor0[3]);
+			OutputColor1[3] = select(OutputColor1[3] > mb_half(0.995), mb_half(1.0), OutputColor1[3]);
+			OutputColor2[3] = select(OutputColor2[3] > mb_half(0.995), mb_half(1.0), OutputColor2[3]);
+			OutputColor3[3] = select(OutputColor3[3] > mb_half(0.995), mb_half(1.0), OutputColor3[3]);
+
+			OutputColor0[3] = select(OutputColor0[3] < mb_half(0.005), mb_half(0.0), OutputColor0[3]);
+			OutputColor1[3] = select(OutputColor1[3] < mb_half(0.005), mb_half(0.0), OutputColor1[3]);
+			OutputColor2[3] = select(OutputColor2[3] < mb_half(0.005), mb_half(0.0), OutputColor2[3]);
+			OutputColor3[3] = select(OutputColor3[3] < mb_half(0.005), mb_half(0.0), OutputColor3[3]);
 		}
 		#else
 		{
@@ -1013,8 +1013,8 @@ void MainCS(
 		// (0.995 chosen to accommodate handling of 254/255)
 		#if POST_PROCESS_ALPHA
 		{
-			OutputColor[3] = select(OutputColor[3] > mb_half(0.995), tsr_half(1.0), OutputColor[3]);
-			OutputColor[3] = select(OutputColor[3] < mb_half(0.005), tsr_half(0.0), OutputColor[3]);
+			OutputColor[3] = select(OutputColor[3] > mb_half(0.995), mb_half(1.0), OutputColor[3]);
+			OutputColor[3] = select(OutputColor[3] < mb_half(0.005), mb_half(0.0), OutputColor[3]);
 		}
 		#else
 		{
