Fix (debug) render pass being culled when their only output is through ShaderPrintBuffer, when shader print use the diagnostic buffer as final output.

diff --git a/Engine/Shaders/Private/ShaderPrint.ush b/Engine/Shaders/Private/ShaderPrint.ush
--- a/Engine/Shaders/Private/ShaderPrint.ush
+++ b/Engine/Shaders/Private/ShaderPrint.ush
@@ -16,6 +16,10 @@
 #define SHADER_PRINT_USE_GLOBAL_RESOURCE 0
 #endif
 
+#ifndef SHADER_PRINT_EXPLICIT_BINDING
+#define SHADER_PRINT_EXPLICIT_BINDING 0
+#endif
+
 // Shader print data can be stored with 3 differents methods:
 // * ShaderPrintData binding: Shader print data are explicity setup on the CPP side, and the ShaderPrintData buffer are used for storing the data
 // * Custom binding         : Shader print data are explicity setup on the CPP side, and the ShaderPrintData buffer are stored into a custom buffer)
@@ -34,6 +38,13 @@
 
 // Global buffers for ShaderPrint collection.
 #if SHADER_PRINT_USE_DIAGNOSTIC_BUFFER
+	// When a shader pass use explicit ShaderPrint parameter binding, it is possible that the only output of that pass is the shader print entry buffer. 
+	// When using Diagnostic buffer as output, the regular shader entry output won't be used, and the pass will be culled. 
+	// To prevent this we add a dummy write, so that the pass does get culled
+	#if SHADER_PRINT_EXPLICIT_BINDING
+	RWStructuredBuffer<uint> ShaderPrint_RWEntryBuffer;
+	#endif
+
 	#define SHADER_PRINT_RWENTRYBUFFER(Ctx, Idx) UEDiagnosticBuffer[SHADER_PRINT_DIAGNOSTIC_BUFFER_ENTRY_OFFSET + (Idx)]
 	#define SHADER_PRINT_RWSTATEBUFFER(Ctx, Idx) UEDiagnosticBuffer[SHADER_PRINT_DIAGNOSTIC_BUFFER_ENTRY_OFFSET + (Idx)]
 #else
@@ -87,6 +98,15 @@ FShaderPrintConfig UnpackShaderPrintConfig()
 	Out.MaxStateCount			= UEDiagnosticBuffer[Offset+9];
 	Out.MaxLineCount			= UEDiagnosticBuffer[Offset+10];
 	Out.MaxTriangleCount		= UEDiagnosticBuffer[Offset+11];
+
+	// Dummy write to avoid the pass to be culled when using explicit ShaderPrint parameter bindings
+	#if SHADER_PRINT_EXPLICIT_BINDING
+	if (Out.Resolution.x > 15999)
+	{
+		ShaderPrint_RWEntryBuffer[0] = 0;
+	}
+	#endif
+
 	return Out;
 }
 #endif
diff --git a/Engine/Source/Runtime/Renderer/Private/ShaderPrint.cpp b/Engine/Source/Runtime/Renderer/Private/ShaderPrint.cpp
--- a/Engine/Source/Runtime/Renderer/Private/ShaderPrint.cpp
+++ b/Engine/Source/Runtime/Renderer/Private/ShaderPrint.cpp
@@ -305,6 +305,8 @@ namespace ShaderPrint
 		{
 			OutEnvironment.CompilerFlags.Add(CFLAG_ForceDXC);
 		}
+
+		OutEnvironment.SetDefine(TEXT("SHADER_PRINT_EXPLICIT_BINDING"), 1);
 	}
 
 	bool IsEnabled()
