Fix shader compilation issue on Vulkan with Substrate.

diff --git a/Engine/Shaders/Private/Substrate/SubstrateVisualize.usf b/Engine/Shaders/Private/Substrate/SubstrateVisualize.usf
--- a/Engine/Shaders/Private/Substrate/SubstrateVisualize.usf
+++ b/Engine/Shaders/Private/Substrate/SubstrateVisualize.usf
@@ -1422,6 +1422,9 @@ void MaterialDebugSubstrateTreePS(
 {
 	OutColor = float4(0.0f, 0.0f, 0.0f, 1.0f);
 
+	// Workaround on Vulkan where this shader does not compile correctly
+#if !VULKAN_PROFILE & !VULKAN_PROFILE_SM5 && !VULKAN_PROFILE_SM6
+
 	FSubstrateAddressing SubstrateAddressing = GetSubstratePixelDataByteOffset((float2(GetCursorPosition()) * View.ViewResolutionFraction), uint2(View.BufferSizeAndInvSize.xy), Substrate.MaxBytesPerPixel);
 	FSubstratePixelHeader Header = UnpackSubstrateHeaderIn(Substrate.MaterialTextureArray, SubstrateAddressing, Substrate.TopLayerTexture);
 	if (Header.ClosureCount > 0)
@@ -1533,6 +1536,7 @@ void MaterialDebugSubstrateTreePS(
 	}
 
 	OutColor = float4(pow(OutColor.rgb, 1.0 / 2.2), OutColor.a);
+#endif
 }
 
 #endif // SHADER_DEBUGSUBSTRATETREE_PS
