Stochastic Direct Lighting - bumped max local light index to 1024. Sample generation now uses a hierarchical bit mask in order not to regress performance.

Console, 1080p, ShooterGame, Highrise: 0.55ms -> 0.54ms

diff --git a/Engine/Shaders/Private/StochasticDirectLighting/StochasticDirectLighting.ush b/Engine/Shaders/Private/StochasticDirectLighting/StochasticDirectLighting.ush
--- a/Engine/Shaders/Private/StochasticDirectLighting/StochasticDirectLighting.ush
+++ b/Engine/Shaders/Private/StochasticDirectLighting/StochasticDirectLighting.ush
@@ -23,18 +23,12 @@
 #define DEBUG_MODE_VISUALIZE_SAMPLING 1
 #define DEBUG_MODE_VISUALIZE_TRACING  2
 
-// Limited by PackLightSample
-#define LOCAL_LIGHT_INDEX_BITS		(14)
-#define LOCAL_LIGHT_INDEX_MASK		((1 << LOCAL_LIGHT_INDEX_BITS) - 1)
-#define INVALID_LOCAL_LIGHT_INDEX	(LOCAL_LIGHT_INDEX_MASK)
-
 // Keep in sync with StochasticDirectLighting.cpp
 #define SHADING_ATLAS_SIZE_IN_TILES 512
 #define DOWNSAMPLE_FACTOR			2
 #define TILE_SIZE					8
-#define MAX_LOCAL_LIGHT_INDEX_XY	16
-#define MAX_LOCAL_LIGHT_INDEX		(MAX_LOCAL_LIGHT_INDEX_XY * MAX_LOCAL_LIGHT_INDEX_XY)
-
+// Limited by PackLightSample, PackShadingTile and SharedCandidateLightHiMask
+#define MAX_LOCAL_LIGHT_INDEX		1024
 #define SHARED_LIGHT_MASK_SIZE		(MAX_LOCAL_LIGHT_INDEX / 32)
 
 // Keep in sync with StochasticDirectLighting::ETileType
@@ -72,7 +66,7 @@ struct FLightSample
 FLightSample InitLightSample()
 {
 	FLightSample LightSample;
-	LightSample.LocalLightIndex = INVALID_LOCAL_LIGHT_INDEX;
+	LightSample.LocalLightIndex = MAX_LOCAL_LIGHT_INDEX;
 	LightSample.Weight = 0.0f;
 	LightSample.bVisible = false;
 	LightSample.bCompleted = true;
@@ -81,7 +75,7 @@ FLightSample InitLightSample()
 
 uint PackLightSample(FLightSample LightSample)
 {
-	uint PackedLightSample = LightSample.LocalLightIndex & LOCAL_LIGHT_INDEX_MASK;
+	uint PackedLightSample = LightSample.LocalLightIndex & 0x3FFF;
 	PackedLightSample |= LightSample.bVisible ? 0x8000 : 0;
 	PackedLightSample |= LightSample.bCompleted ? 0x4000 : 0;
 	PackedLightSample |= f32tof16(LightSample.Weight) << 16;
diff --git a/Engine/Shaders/Private/StochasticDirectLighting/StochasticDirectLightingSampling.usf b/Engine/Shaders/Private/StochasticDirectLighting/StochasticDirectLightingSampling.usf
--- a/Engine/Shaders/Private/StochasticDirectLighting/StochasticDirectLightingSampling.usf
+++ b/Engine/Shaders/Private/StochasticDirectLighting/StochasticDirectLightingSampling.usf
@@ -88,6 +88,8 @@ float4 HistoryScreenPositionScaleBias;
 float4 HistoryUVMinMax;
 
 groupshared uint SharedLightMask[SHARED_LIGHT_MASK_SIZE];
+
+groupshared uint SharedCandidateLightHiMask;
 groupshared uint SharedCandidateLightMask[SHARED_LIGHT_MASK_SIZE];
 
 void MarkLocalLightIndex(uint LocalLightIndex)
@@ -114,13 +116,14 @@ void GenerateLightSamplesCS(
 {
 	uint LinearThreadIndex = GroupThreadId.y * THREADGROUP_SIZE + GroupThreadId.x;
 
-	if (LinearThreadIndex < SHARED_LIGHT_MASK_SIZE)
+	if (LinearThreadIndex == 0)
 	{
-		SharedLightMask[LinearThreadIndex] = 0;
+		SharedCandidateLightHiMask = 0;
 	}
 
 	if (LinearThreadIndex < SHARED_LIGHT_MASK_SIZE)
 	{
+		SharedLightMask[LinearThreadIndex] = 0;
 		SharedCandidateLightMask[LinearThreadIndex] = 0;
 	}
 
@@ -146,6 +149,8 @@ void GenerateLightSamplesCS(
 			const FCulledLightsGridData CulledLightGridData = GetCulledLightsGrid(GridIndex, EyeIndex);
 			const uint NumLightsInGridCell = min(CulledLightGridData.NumLights, GetMaxLightsPerCell(EyeIndex));
 
+			uint LocalCandidateLightHiMask = 0;
+
 			for (uint GridLightIndex = 0; GridLightIndex < NumLightsInGridCell; ++GridLightIndex)
 			{
 				const uint LocalLightIndex = GetCulledLightDataGrid(CulledLightGridData.DataStartIndex + GridLightIndex);
@@ -156,8 +161,17 @@ void GenerateLightSamplesCS(
 					uint DWORDIndex = LocalLightIndex / 32;
 					uint BitMask = 1u << (LocalLightIndex % 32);
 					InterlockedOr(SharedCandidateLightMask[DWORDIndex], BitMask);
+
+					uint HiBitMask = 1u << DWORDIndex;
+					LocalCandidateLightHiMask |= HiBitMask;
 				}
 			}
+
+			if (WaveIsFirstLane())
+			{
+				LocalCandidateLightHiMask = WaveActiveBitOr(LocalCandidateLightHiMask);
+				InterlockedOr(SharedCandidateLightHiMask, LocalCandidateLightHiMask);
+			}
 		}
 	
 		GroupMemoryBarrierWithGroupSync();
@@ -248,9 +262,16 @@ void GenerateLightSamplesCS(
 				}
 				#endif
 
-				for (uint MaskIndex = 0; MaskIndex < SHARED_LIGHT_MASK_SIZE; ++MaskIndex)
+				uint CandidateLightHiMask = SharedCandidateLightHiMask;
+				while (CandidateLightHiMask != 0)
 				{
+					const uint NextHiBitIndex = firstbitlow(CandidateLightHiMask);
+					const uint NextHiBitMask = 1u << NextHiBitIndex;
+					CandidateLightHiMask ^= NextHiBitMask;
+
+					const uint MaskIndex = NextHiBitIndex;
 					uint CandidateLightMask = SharedCandidateLightMask[MaskIndex];
+
 					while (CandidateLightMask != 0)
 					{
 						const uint NextBitIndex = firstbitlow(CandidateLightMask);
@@ -351,7 +372,7 @@ void GenerateLightSamplesCS(
 			{
 				FLightSample LightSample = LightSamples[LightSampleIndex];
 
-				if (LightSample.LocalLightIndex != INVALID_LOCAL_LIGHT_INDEX)
+				if (LightSample.LocalLightIndex != MAX_LOCAL_LIGHT_INDEX)
 				{
 					const FLocalLightData LocalLightData = GetLocalLightDataNonStereo(LightSample.LocalLightIndex);
 					const FLightSceneData LightSceneData = GetLightSceneData(LocalLightData.LightSceneId);
@@ -365,7 +386,7 @@ void GenerateLightSamplesCS(
 				RWLightSamples[DownsampledScreenCoord * NumSamplesPerPixel + LightSampleOffsets[LightSampleIndex]] = PackLightSample(LightSample);
 
 				// Mark tiles which need to be processed
-				if (LightSample.LocalLightIndex != INVALID_LOCAL_LIGHT_INDEX)
+				if (LightSample.LocalLightIndex != MAX_LOCAL_LIGHT_INDEX)
 				{
 					MarkLocalLightIndex(LightSample.LocalLightIndex);
 				}
diff --git a/Engine/Shaders/Private/StochasticDirectLighting/StochasticDirectLightingTracing.usf b/Engine/Shaders/Private/StochasticDirectLighting/StochasticDirectLightingTracing.usf
--- a/Engine/Shaders/Private/StochasticDirectLighting/StochasticDirectLightingTracing.usf
+++ b/Engine/Shaders/Private/StochasticDirectLighting/StochasticDirectLightingTracing.usf
@@ -99,7 +99,7 @@ void CompactLightSampleTracesCS(
 		// #sdl_todo: try to preserve z-order here for improved tracing coherency
 		FLightSample LightSample = UnpackLightSample(LightSamples[SampleCoord]);
 
-		if (LightSample.LocalLightIndex != INVALID_LOCAL_LIGHT_INDEX && !LightSample.bCompleted)
+		if (LightSample.LocalLightIndex != MAX_LOCAL_LIGHT_INDEX && !LightSample.bCompleted)
 		{
 			#if WAVE_OPS
 			{
diff --git a/Engine/Shaders/Private/StochasticDirectLighting/StochasticDirectLightingTracing.ush b/Engine/Shaders/Private/StochasticDirectLighting/StochasticDirectLightingTracing.ush
--- a/Engine/Shaders/Private/StochasticDirectLighting/StochasticDirectLightingTracing.ush
+++ b/Engine/Shaders/Private/StochasticDirectLighting/StochasticDirectLightingTracing.ush
@@ -89,7 +89,7 @@ FLightSampleTrace GetLightSampleTrace(float3 TranslatedWorldPosition, uint Local
 	LightSampleTrace.Direction = 0.0f;
 	LightSampleTrace.Distance = 0.0f;
 
-	if (LocalLightIndex != INVALID_LOCAL_LIGHT_INDEX)
+	if (LocalLightIndex != MAX_LOCAL_LIGHT_INDEX)
 	{
 		const FLocalLightData LocalLightData = GetLocalLightDataNonStereo(LocalLightIndex);
 		const FLightSceneData LightSceneData = GetLightSceneData(LocalLightData.LightSceneId);
diff --git a/Engine/Source/Runtime/Renderer/Private/StochasticDirectLighting/StochasticDirectLighting.cpp b/Engine/Source/Runtime/Renderer/Private/StochasticDirectLighting/StochasticDirectLighting.cpp
--- a/Engine/Source/Runtime/Renderer/Private/StochasticDirectLighting/StochasticDirectLighting.cpp
+++ b/Engine/Source/Runtime/Renderer/Private/StochasticDirectLighting/StochasticDirectLighting.cpp
@@ -179,7 +179,7 @@ namespace StochasticDirectLighting
 		}
 
 		// SM6 because it uses typed loads to accumulate lights
-		return IsFeatureLevelSupported(Parameters.Platform, ERHIFeatureLevel::SM6);
+		return IsFeatureLevelSupported(Parameters.Platform, ERHIFeatureLevel::SM6) && RHISupportsWaveOperations(Parameters.Platform);
 	}
 
 	uint32 GetStateFrameIndex(FSceneViewState* ViewState)
@@ -369,6 +369,7 @@ class FGenerateLightSamplesCS : public FGlobalShader
 		FGlobalShader::ModifyCompilationEnvironment(Parameters, OutEnvironment);
 		StochasticDirectLighting::ModifyCompilationEnvironment(Parameters.Platform, OutEnvironment);
 		OutEnvironment.SetDefine(TEXT("THREADGROUP_SIZE"), GetGroupSize());
+		OutEnvironment.CompilerFlags.Add(CFLAG_WaveOperations);
 	}
 };
 
