Remove point to root index indirection as now interpolation already contains this data.

This avoids an indirect fetch. This also fixed what looks like a long standing issue where the curve index was reinterpreted as the point root index.

diff --git a/Engine/Plugins/Runtime/HairStrands/Source/HairStrandsCore/Private/GroomManager.cpp b/Engine/Plugins/Runtime/HairStrands/Source/HairStrandsCore/Private/GroomManager.cpp
--- a/Engine/Plugins/Runtime/HairStrands/Source/HairStrandsCore/Private/GroomManager.cpp
+++ b/Engine/Plugins/Runtime/HairStrands/Source/HairStrandsCore/Private/GroomManager.cpp
@@ -1396,7 +1396,6 @@ static void RunHairStrandsInterpolation_Strands(
 					bHasSkinning ? InstanceData.Instance->Strands.DeformedRootResource : nullptr,
 					bHasSkinning && bValidGuide ? InstanceData.Instance->Guides.DeformedRootResource : nullptr,
 					RegisterAsSRV(GraphBuilder, InstanceData.Instance->Strands.RestResource->PositionBuffer),
-					RegisterAsSRV(GraphBuilder, InstanceData.Instance->Strands.RestResource->PointToCurveBuffer),
 					RegisterAsSRV(GraphBuilder, InstanceData.Instance->Strands.RestResource->CurveBuffer),
 					bValidGuide ? RegisterAsSRV(GraphBuilder, InstanceData.Instance->Guides.RestResource->CurveBuffer) : nullptr,
 					bUseSingleGuide,
@@ -1404,7 +1403,6 @@ static void RunHairStrandsInterpolation_Strands(
 					bValidGuide ? RegisterAsSRV(GraphBuilder, InstanceData.Instance->Strands.InterpolationResource->PointInterpolationBuffer) : nullptr,
 					bValidGuide ? RegisterAsSRV(GraphBuilder, InstanceData.Instance->Guides.RestResource->PositionBuffer) : nullptr,
 					bValidGuide ? RegisterAsSRV(GraphBuilder, InstanceData.Instance->Guides.DeformedResource->GetBuffer(FHairStrandsDeformedResource::Current)) : nullptr,
-					bValidGuide ? RegisterAsSRV(GraphBuilder, InstanceData.Instance->Guides.RestResource->PointToCurveBuffer) : nullptr,
 					InstanceData.RDGResources.DeformerPositionSRV,
 					InstanceData.RDGResources.PositionUAV,
 					FHairStrandsDeformedRootResource::FLOD::Current);
@@ -1956,7 +1954,6 @@ static void RunHairStrandsInterpolation_Cards(
 					InstanceData.bHasSkinning ? LOD.Guides.DeformedRootResource : nullptr,
 					InstanceData.bHasSkinning && InstanceData.bValidGuide ? InstanceData.Instance->Guides.DeformedRootResource : nullptr,
 					RegisterAsSRV(GraphBuilder, LOD.Guides.RestResource->PositionBuffer),
-					RegisterAsSRV(GraphBuilder, LOD.Guides.RestResource->PointToCurveBuffer),
 					RegisterAsSRV(GraphBuilder, LOD.Guides.RestResource->CurveBuffer),
 					InstanceData.bValidGuide ? RegisterAsSRV(GraphBuilder, InstanceData.Instance->Guides.RestResource->CurveBuffer) : nullptr,
 					bUseSingleGuide,
@@ -1964,7 +1961,6 @@ static void RunHairStrandsInterpolation_Cards(
 					InstanceData.bValidGuide ? RegisterAsSRV(GraphBuilder, LOD.Guides.InterpolationResource->PointInterpolationBuffer) : nullptr,
 					InstanceData.bValidGuide ? RegisterAsSRV(GraphBuilder, InstanceData.Instance->Guides.RestResource->PositionBuffer) : nullptr,
 					InstanceData.bValidGuide ? RegisterAsSRV(GraphBuilder, InstanceData.Instance->Guides.DeformedResource->GetBuffer(FHairStrandsDeformedResource::Current)) : nullptr,
-					InstanceData.bValidGuide ? RegisterAsSRV(GraphBuilder, InstanceData.Instance->Guides.RestResource->PointToCurveBuffer) : nullptr,
 					nullptr,
 					Guides_DeformedPositionUAV,
 					FHairStrandsDeformedRootResource::FLOD::Current); // <- this should be optional
diff --git a/Engine/Plugins/Runtime/HairStrands/Source/HairStrandsCore/Private/HairStrandsInterpolation.cpp b/Engine/Plugins/Runtime/HairStrands/Source/HairStrandsCore/Private/HairStrandsInterpolation.cpp
--- a/Engine/Plugins/Runtime/HairStrands/Source/HairStrandsCore/Private/HairStrandsInterpolation.cpp
+++ b/Engine/Plugins/Runtime/HairStrands/Source/HairStrandsCore/Private/HairStrandsInterpolation.cpp
@@ -495,8 +495,6 @@ class FHairInterpolationCS : public FGlobalShader
 		SHADER_PARAMETER_RDG_BUFFER_SRV(Buffer<uint>, SimRootBarycentricBuffer)
 		SHADER_PARAMETER_RDG_BUFFER_SRV(Buffer<uint>, SimRootToUniqueTriangleIndexBuffer)
 
-		SHADER_PARAMETER_RDG_BUFFER_SRV(Buffer<uint>, SimRootPointIndexBuffer)
-
 		SHADER_PARAMETER_RDG_BUFFER_SRV(Buffer<float4>, OutSimHairPositionOffsetBuffer)
 		SHADER_PARAMETER_RDG_BUFFER_SRV(Buffer<float4>, OutRenHairPositionOffsetBuffer)
 
@@ -552,7 +550,6 @@ void AddHairStrandsInterpolationPass(
 	const FHairStrandsDeformedRootResource* RenDeformedRootResources,
 	const FHairStrandsDeformedRootResource* SimDeformedRootResources,
 	const FRDGBufferSRVRef& RenRestPosePositionBuffer,
-	const FRDGBufferSRVRef& RenPointToCurveBuffer,
 	const FRDGBufferSRVRef& RenCurveBuffer,
 	const FRDGBufferSRVRef& SimCurveBuffer,
 	const bool bUseSingleGuide,
@@ -560,7 +557,6 @@ void AddHairStrandsInterpolationPass(
 	const FRDGBufferSRVRef& PointInterpolationBuffer,
 	const FRDGBufferSRVRef& SimRestPosePositionBuffer,
 	const FRDGBufferSRVRef& SimDeformedPositionBuffer,
-	const FRDGBufferSRVRef& SimRootPointIndexBuffer,
 	const FRDGBufferSRVRef& RenDeformerPositionBuffer,
 	FRDGBufferUAVRef& OutRenPositionBuffer,
 	const FHairStrandsDeformedRootResource::FLOD::EFrameType DeformedFrame)
@@ -580,9 +576,6 @@ void AddHairStrandsInterpolationPass(
 
 	Parameters->OutSimHairPositionOffsetBuffer = OutSimHairPositionOffsetBuffer;
 	Parameters->OutRenHairPositionOffsetBuffer = OutRenHairPositionOffsetBuffer;
-
-	Parameters->SimRootPointIndexBuffer = SimRootPointIndexBuffer;
-	Parameters->RenPointToCurveBuffer = RenPointToCurveBuffer;
 	
 	Parameters->LocalToWorldMatrix = FMatrix44f(Instance->LocalToWorld.ToMatrixWithScale());		// LWC_TODO: Precision loss
 	Parameters->HairLengthScale = HairLengthScale;
diff --git a/Engine/Plugins/Runtime/HairStrands/Source/HairStrandsCore/Public/HairStrandsInterpolation.h b/Engine/Plugins/Runtime/HairStrands/Source/HairStrandsCore/Public/HairStrandsInterpolation.h
--- a/Engine/Plugins/Runtime/HairStrands/Source/HairStrandsCore/Public/HairStrandsInterpolation.h
+++ b/Engine/Plugins/Runtime/HairStrands/Source/HairStrandsCore/Public/HairStrandsInterpolation.h
@@ -135,7 +135,6 @@ void AddHairStrandsInterpolationPass(
 	const FHairStrandsDeformedRootResource* RenDeformedRootResources,
 	const FHairStrandsDeformedRootResource* SimDeformedRootResources,
 	const FRDGBufferSRVRef& RenRestPosePositionBuffer,
-	const FRDGBufferSRVRef& RenPointToCurveBuffer,
 	const FRDGBufferSRVRef& RenCurveBuffer,
 	const FRDGBufferSRVRef& SimCurveBuffer,
 	const bool bUseSingleGuide,
@@ -143,7 +142,7 @@ void AddHairStrandsInterpolationPass(
 	const FRDGBufferSRVRef& PoinInterpolationBuffer,
 	const FRDGBufferSRVRef& SimRestPosePositionBuffer,
 	const FRDGBufferSRVRef& SimDeformedPositionBuffer,
-	const FRDGBufferSRVRef& SimRootPointIndexBuffer,
+
 	const FRDGBufferSRVRef& RenDeformerPositionBuffer,
 	FRDGBufferUAVRef& OutRenPositionBuffer,
 	const FHairStrandsDeformedRootResource::FLOD::EFrameType DeformedFrame);
diff --git a/Engine/Shaders/Private/HairStrands/HairStrandsInterpolation.usf b/Engine/Shaders/Private/HairStrands/HairStrandsInterpolation.usf
--- a/Engine/Shaders/Private/HairStrands/HairStrandsInterpolation.usf
+++ b/Engine/Shaders/Private/HairStrands/HairStrandsInterpolation.usf
@@ -35,9 +35,10 @@ ByteAddressBuffer PointInterpolationBuffer;
 
 struct FGuideDataWithOffset
 {
-	uint2  VertexIndices;
-	float2 VertexLerps;
-	float2 VertexWeights;
+	uint2  RootPointIndices;
+	uint2  PointIndices;
+	float2 PointLerps;
+	float2 PointWeights;
 	float3 OutPositionOffset;
 };
 
@@ -45,10 +46,11 @@ FGuideDataWithOffset GetGuideData(FGuideCurveData InCurveData, FGuidePointData I
 {
 #if PERMUTATION_SIMULATION > 0
 	FGuideDataWithOffset Out;
-	Out.VertexIndices     = InCurveData.RootPointIndices + InPointData.PointIndices;
-	Out.VertexLerps       = InPointData.PointLerps;
-	Out.VertexWeights     = InCurveData.CurveWeights;
-	Out.OutPositionOffset = OutPositionOffset;
+	Out.RootPointIndices = InCurveData.RootPointIndices;
+	Out.PointIndices     = InCurveData.RootPointIndices + InPointData.PointIndices;
+	Out.PointLerps       = InPointData.PointLerps;
+	Out.PointWeights     = InCurveData.CurveWeights;
+	Out.OutPositionOffset= OutPositionOffset;
 	return Out;
 #else
 	return (FGuideDataWithOffset)0;
@@ -81,7 +83,6 @@ float4x4 LocalToWorldMatrix;
 Buffer<uint4> RenRestPosePositionBuffer;
 Buffer<uint4> SimRestPosePositionBuffer;
 Buffer<uint4> SimDeformedPositionBuffer;
-Buffer<uint>  SimRootPointIndexBuffer;
 
 uint HairStrandsVFTODO_bCullingEnable;
 #if PERMUTATION_CULLING == 1
@@ -90,11 +91,6 @@ Buffer<uint>	HairStrandsVFTODO_CullingIndexBuffer;
 Buffer<float>	HairStrandsVFTODO_CullingRadiusScaleBuffer;
 #endif
 
-#if PERMUTATION_DYNAMIC_GEOMETRY >= 1 || PERMUTATION_DEBUG == 1 
-Buffer<uint>   RenPointToCurveBuffer;
-Buffer<uint>   SimPointToCurveBuffer;
-#endif
-
 #if PERMUTATION_DEFORMER
 Buffer<uint4>  RenDeformerPositionBuffer;
 #endif
@@ -124,73 +120,82 @@ float3 GetRenPositionOffset() { return ReadRenPositionOffset(OutRenHairPositionO
 ///////////////////////////////////////////////////////////////////////////////////////////////////
 // Deformation offset
 
+struct FGuideIndex
+{
+	uint PointIndex;
+	uint RootPointIndex;
+};
+
+FGuideIndex ToGuideIndex(FGuideDataWithOffset In, uint InK, uint InOffset)
+{
+	FGuideIndex Out;
+	Out.PointIndex		= In.PointIndices[InK] + InOffset;
+	Out.RootPointIndex	= In.RootPointIndices[InK];
+	return Out;
+};
+
+
 float3 ComputeStaticGeometryOffset(
-	uint GuideIndex, 
+	FGuideIndex GuideIndex,
 	float GuideVertexWeight, 
 	float3 OutSimHairPositionOffset)
 {
-	const float3 RestGuidePoint = UnpackHairControlPointPosition(SimRestPosePositionBuffer[GuideIndex], InSimHairPositionOffset);
+	const float3 RestGuidePoint = UnpackHairControlPointPosition(SimRestPosePositionBuffer[GuideIndex.PointIndex], InSimHairPositionOffset);
 
-	const float3 DeformedGuidePoint = UnpackHairControlPointPosition(SimDeformedPositionBuffer[GuideIndex], OutSimHairPositionOffset);
+	const float3 DeformedGuidePoint = UnpackHairControlPointPosition(SimDeformedPositionBuffer[GuideIndex.PointIndex], OutSimHairPositionOffset);
 
 	return (DeformedGuidePoint - RestGuidePoint) * GuideVertexWeight;
 }
 
 float3 ComputeLocalGeometryOffset(
-	uint GuideIndex, 
+	FGuideIndex GuideIndex, 
 	float GuideVertexWeight)
 {
-	const uint RootIndex = SimRootPointIndexBuffer[GuideIndex];
-
-	const float3 RestGuidePoint = UnpackHairControlPointPosition(SimRestPosePositionBuffer[GuideIndex]) - UnpackHairControlPointPosition(SimRestPosePositionBuffer[RootIndex]);
+	const float3 RestGuidePoint = UnpackHairControlPointPosition(SimRestPosePositionBuffer[GuideIndex.PointIndex]) - UnpackHairControlPointPosition(SimRestPosePositionBuffer[GuideIndex.RootPointIndex]);
 
-	const float3 DeformedGuidePoint = UnpackHairControlPointPosition(SimDeformedPositionBuffer[GuideIndex]) - UnpackHairControlPointPosition(SimDeformedPositionBuffer[RootIndex]);
+	const float3 DeformedGuidePoint = UnpackHairControlPointPosition(SimDeformedPositionBuffer[GuideIndex.PointIndex]) - UnpackHairControlPointPosition(SimDeformedPositionBuffer[GuideIndex.RootPointIndex]);
 
 	return (DeformedGuidePoint - RestGuidePoint) * GuideVertexWeight;
 }
 
 float3 ComputeDynamicGeometryOffset(
-	uint GuideIndex, 
+	FGuideIndex GuideIndex, 
 	float GuideVertexWeight, 
 	FHairMeshTriangle RestTriangle, 
 	FHairMeshTriangle DeformedTriangle,
 	float3 RootBarycentric,
 	float3 OutSimHairPositionOffset)
 {
-	const float3 RestGuidePoint = UnpackHairControlPointPosition(SimRestPosePositionBuffer[GuideIndex], InSimHairPositionOffset);
+	const float3 RestGuidePoint = UnpackHairControlPointPosition(SimRestPosePositionBuffer[GuideIndex.PointIndex], InSimHairPositionOffset);
 	const float3 LocalRestGuidePoint = ToTriangle(RestGuidePoint, RootBarycentric, RestTriangle);
 
-	const float3 DeformedGuidePoint = UnpackHairControlPointPosition(SimDeformedPositionBuffer[GuideIndex], OutSimHairPositionOffset);
+	const float3 DeformedGuidePoint = UnpackHairControlPointPosition(SimDeformedPositionBuffer[GuideIndex.PointIndex], OutSimHairPositionOffset);
 	const float3 LocalDeformedGuidePoint = ToTriangle(DeformedGuidePoint, RootBarycentric, DeformedTriangle);
 
 	return (LocalDeformedGuidePoint - LocalRestGuidePoint) * GuideVertexWeight;
 }
 
 float3 ComputeTranslateGeometryOffset(
-	uint GuideIndex, 
+	FGuideIndex GuideIndex, 
 	float GuideVertexWeight)
 {
-	const uint RootIndex = SimRootPointIndexBuffer[GuideIndex];
-
-	const float3 RestGuidePoint = UnpackHairControlPointPosition(SimRestPosePositionBuffer[GuideIndex]) - UnpackHairControlPointPosition(SimRestPosePositionBuffer[RootIndex]) ;
+	const float3 RestGuidePoint = UnpackHairControlPointPosition(SimRestPosePositionBuffer[GuideIndex.PointIndex]) - UnpackHairControlPointPosition(SimRestPosePositionBuffer[GuideIndex.RootPointIndex]);
 
-	const float3 DeformedGuidePoint = UnpackHairControlPointPosition(SimDeformedPositionBuffer[GuideIndex]) - UnpackHairControlPointPosition(SimDeformedPositionBuffer[RootIndex]) ;
+	const float3 DeformedGuidePoint = UnpackHairControlPointPosition(SimDeformedPositionBuffer[GuideIndex.PointIndex]) - UnpackHairControlPointPosition(SimDeformedPositionBuffer[GuideIndex.RootPointIndex]);
 
 	return (DeformedGuidePoint - RestGuidePoint) * GuideVertexWeight;
 }
 
 float3 ComputeRotateGeometryOffset(
-	uint GuideIndex, 
+	FGuideIndex GuideIndex, 
 	float GuideVertexWeight,
 	FHairMeshTriangle RestTriangle, 
 	FHairMeshTriangle DeformedTriangle)
 {	
-	const uint RootIndex = SimRootPointIndexBuffer[GuideIndex];
-
-	const float3 RestGuidePoint = UnpackHairControlPointPosition(SimRestPosePositionBuffer[GuideIndex]) - UnpackHairControlPointPosition(SimRestPosePositionBuffer[RootIndex]) ;
+	const float3 RestGuidePoint = UnpackHairControlPointPosition(SimRestPosePositionBuffer[GuideIndex.PointIndex]) - UnpackHairControlPointPosition(SimRestPosePositionBuffer[GuideIndex.RootPointIndex]);
 	const float3 LocalRestGuidePoint = VectorToTriangle(RestGuidePoint, RestTriangle);
 
-	const float3 DeformedGuidePoint = UnpackHairControlPointPosition(SimDeformedPositionBuffer[GuideIndex]) - UnpackHairControlPointPosition(SimDeformedPositionBuffer[RootIndex]) ;
+	const float3 DeformedGuidePoint = UnpackHairControlPointPosition(SimDeformedPositionBuffer[GuideIndex.PointIndex]) - UnpackHairControlPointPosition(SimDeformedPositionBuffer[GuideIndex.RootPointIndex]);
 	const float3 LocalDeformedGuidePoint = VectorToTriangle(DeformedGuidePoint, DeformedTriangle);
 
 	return (LocalDeformedGuidePoint - LocalRestGuidePoint) * GuideVertexWeight;
@@ -241,13 +246,10 @@ float3 ApplyRigidInterpolation(FGuideDataWithOffset GuideData, float3 InPosition
 	for (uint K = 0; K < HAIR_INTERPOLATION_MAX_GUIDE_COUNT; ++K)
 	#endif
 	{
-		const uint GuideIndex0 = GuideData.VertexIndices[K];
-		const uint GuideIndex1 = GuideIndex0 + 1;
-
-		const float3 Offset0 = ComputeStaticGeometryOffset(GuideIndex0, GuideData.VertexWeights[K], GuideData.OutPositionOffset);
-		const float3 Offset1 = ComputeStaticGeometryOffset(GuideIndex1, GuideData.VertexWeights[K], GuideData.OutPositionOffset);
+		const float3 Offset0 = ComputeStaticGeometryOffset(ToGuideIndex(GuideData, K, 0), GuideData.PointWeights[K], GuideData.OutPositionOffset);
+		const float3 Offset1 = ComputeStaticGeometryOffset(ToGuideIndex(GuideData, K, 1), GuideData.PointWeights[K], GuideData.OutPositionOffset);
 
-		CurrOffset += lerp(Offset0, Offset1, GuideData.VertexLerps[K]);
+		CurrOffset += lerp(Offset0, Offset1, GuideData.PointLerps[K]);
 	}
 #endif
 
@@ -269,15 +271,12 @@ float3 ApplySkinningWithOffset(const FGuideDataWithOffset GuideData, const FRoot
 	for (uint K = 0; K < HAIR_INTERPOLATION_MAX_GUIDE_COUNT; ++K)
 	#endif
 	{
-		const uint GuideIndex0 = GuideData.VertexIndices[K];
-		const uint GuideIndex1 = GuideIndex0 + 1;
-
 		const FRootTriangleData SimTri = SimTris[K];
 
-		const float3 Offset0 = ComputeDynamicGeometryOffset(GuideIndex0, GuideData.VertexWeights[K], SimTri.RestTriangle, SimTri.DeformedTriangle, SimTri.RootBarycentric, GuideData.OutPositionOffset);
-		const float3 Offset1 = ComputeDynamicGeometryOffset(GuideIndex1, GuideData.VertexWeights[K], SimTri.RestTriangle, SimTri.DeformedTriangle, SimTri.RootBarycentric, GuideData.OutPositionOffset);
+		const float3 Offset0 = ComputeDynamicGeometryOffset(ToGuideIndex(GuideData, K, 0), GuideData.PointWeights[K], SimTri.RestTriangle, SimTri.DeformedTriangle, SimTri.RootBarycentric, GuideData.OutPositionOffset);
+		const float3 Offset1 = ComputeDynamicGeometryOffset(ToGuideIndex(GuideData, K, 1), GuideData.PointWeights[K], SimTri.RestTriangle, SimTri.DeformedTriangle, SimTri.RootBarycentric, GuideData.OutPositionOffset);
 
-		CurrOffset += VectorToWorld(lerp(Offset0, Offset1, GuideData.VertexLerps[K]), SimTri.DeformedTriangle);
+		CurrOffset += VectorToWorld(lerp(Offset0, Offset1, GuideData.PointLerps[K]), SimTri.DeformedTriangle);
 	}
 #endif		
 
@@ -303,13 +302,10 @@ float3 ApplySkinningWithTranslation(const FGuideDataWithOffset GuideData, const
 	for (uint K = 0; K < HAIR_INTERPOLATION_MAX_GUIDE_COUNT; ++K)
 	#endif
 	{
-		const uint GuideIndex0 = GuideData.VertexIndices[K];
-		const uint GuideIndex1 = GuideIndex0 + 1;
-
-		const float3 Offset0 = ComputeTranslateGeometryOffset(GuideIndex0, 1.0);
-		const float3 Offset1 = ComputeTranslateGeometryOffset(GuideIndex1, 1.0);
+		const float3 Offset0 = ComputeTranslateGeometryOffset(ToGuideIndex(GuideData, K, 0), 1.0);
+		const float3 Offset1 = ComputeTranslateGeometryOffset(ToGuideIndex(GuideData, K, 1), 1.0);
 
-		CurrOffset += GuideData.VertexWeights[K] * lerp(Offset0, Offset1, GuideData.VertexLerps[K]);
+		CurrOffset += GuideData.PointWeights[K] * lerp(Offset0, Offset1, GuideData.PointLerps[K]);
 	}
 #endif
 
@@ -333,17 +329,14 @@ float3 ApplySkinningWithRotation(const FGuideDataWithOffset GuideData, const FRo
 	for (uint K = 0; K < HAIR_INTERPOLATION_MAX_GUIDE_COUNT; ++K)
 	#endif
 	{
-		const uint GuideIndex0 = GuideData.VertexIndices[K];
-		const uint GuideIndex1 = GuideIndex0 + 1;
-
 		const FRootTriangleData SimTri = SimTris[K];
 
 		const float3 LocalPoint = VectorToTriangle(InPosition, SimTri.RestTriangle);
 
-		const float3 Offset0 = ComputeRotateGeometryOffset(GuideIndex0, 1.0, SimTri.RestTriangle, SimTri.DeformedTriangle/*, SimHairPositionOffset*/);
-		const float3 Offset1 = ComputeRotateGeometryOffset(GuideIndex1, 1.0, SimTri.RestTriangle, SimTri.DeformedTriangle/*, SimHairPositionOffset*/);
+		const float3 Offset0 = ComputeRotateGeometryOffset(ToGuideIndex(GuideData, K, 0), 1.0, SimTri.RestTriangle, SimTri.DeformedTriangle/*, SimHairPositionOffset*/);
+		const float3 Offset1 = ComputeRotateGeometryOffset(ToGuideIndex(GuideData, K, 1), 1.0, SimTri.RestTriangle, SimTri.DeformedTriangle/*, SimHairPositionOffset*/);
 
-		CurrOffset += GuideData.VertexWeights[K] * VectorToWorld(lerp(Offset0, Offset1, GuideData.VertexLerps[K]) + LocalPoint, SimTri.DeformedTriangle);
+		CurrOffset += GuideData.PointWeights[K] * VectorToWorld(lerp(Offset0, Offset1, GuideData.PointLerps[K]) + LocalPoint, SimTri.DeformedTriangle);
 	}
 #endif
 
@@ -365,17 +358,14 @@ float3 ApplySkinningWithTranslationAndRotation(const FGuideDataWithOffset GuideD
 	for (uint K = 0; K < HAIR_INTERPOLATION_MAX_GUIDE_COUNT; ++K)
 	#endif
 	{
-		const uint GuideIndex0 = GuideData.VertexIndices[K];
-		const uint GuideIndex1 = GuideIndex0 + 1;
-
 		const FRootTriangleData SimTri = SimTris[K];  //GetRootTriangleData(GuideIndex0, true);
 
 		const float3 LocalPoint = ToTriangle(InPosition, SimTri.RootBarycentric, SimTri.RestTriangle);
 
-		const float3 Offset0 = ComputeDynamicGeometryOffset(GuideIndex0, 1.0, SimTri.RestTriangle, SimTri.DeformedTriangle, SimTri.RootBarycentric, GuideData.OutPositionOffset);
-		const float3 Offset1 = ComputeDynamicGeometryOffset(GuideIndex1, 1.0, SimTri.RestTriangle, SimTri.DeformedTriangle, SimTri.RootBarycentric, GuideData.OutPositionOffset);
+		const float3 Offset0 = ComputeDynamicGeometryOffset(ToGuideIndex(GuideData, K, 0), 1.0, SimTri.RestTriangle, SimTri.DeformedTriangle, SimTri.RootBarycentric, GuideData.OutPositionOffset);
+		const float3 Offset1 = ComputeDynamicGeometryOffset(ToGuideIndex(GuideData, K, 1), 1.0, SimTri.RestTriangle, SimTri.DeformedTriangle, SimTri.RootBarycentric, GuideData.OutPositionOffset);
 
-		CurrOffset += GuideData.VertexWeights[K] * ToWorld(lerp(Offset0, Offset1, GuideData.VertexLerps[K]) + LocalPoint, SimTri.RootBarycentric, SimTri.DeformedTriangle);
+		CurrOffset += GuideData.PointWeights[K] * ToWorld(lerp(Offset0, Offset1, GuideData.PointLerps[K]) + LocalPoint, SimTri.RootBarycentric, SimTri.DeformedTriangle);
 	}
 #endif
 
@@ -797,7 +787,7 @@ void MainCS(uint2 DispatchThreadId : SV_DispatchThreadID, uint GroupThread1D : S
 			const FGuidePointData SimPointData = UnpackGuidePointData(PointInterpolationBuffer, VertexIndex0, bUseSingleGuide); 
 			const FGuideDataWithOffset GuideData = GetGuideData(SimCurveData, SimPointData, 0 /*OutPositionOffset*/);
 
-			const uint GuideIndex		= GuideData.VertexIndices[0]; // Take the closest guide
+			const uint GuideIndex		= GuideData.PointIndices[0]; // Take the closest guide
 			SeedValues[GroupThread1D]	= ToSeedHash(GuideIndex);
 		}
 		#endif // PERMUTATION_SIMULATION
