Change Depth and CustomDepth in a FGBufferData to a full float, half precision is not enough for a depth values

diff --git a/Engine/Shaders/Private/DeferredShadingCommon.ush b/Engine/Shaders/Private/DeferredShadingCommon.ush
--- a/Engine/Shaders/Private/DeferredShadingCommon.ush
+++ b/Engine/Shaders/Private/DeferredShadingCommon.ush
@@ -366,12 +366,12 @@ struct FGBufferData
 	// 0..1, 2 bits, use CastContactShadow(GBuffer) or HasDynamicIndirectShadowCasterRepresentation(GBuffer) to extract
 	half PerObjectGBufferData;
 	// in world units
-	half CustomDepth;
+	float CustomDepth;
 	// Custom depth stencil value
 	uint CustomStencil;
 	// in unreal units (linear), can be used to reconstruct world position,
 	// only valid when decoding the GBuffer as the value gets reconstructed from the Z buffer
-	half Depth;
+	float Depth;
 	// Velocity for motion blur (only used when WRITES_VELOCITY_TO_GBUFFER is enabled)
 	half4 Velocity;
 
