Reorganises TSR's reject shading functions in headers

diff --git a/Engine/Shaders/Private/TemporalSuperResolution/TSRConvolutionNetworkPass.ush b/Engine/Shaders/Private/TemporalSuperResolution/TSRConvolutionNetworkPass.ush
new file mode 100644
--- /dev/null
+++ b/Engine/Shaders/Private/TemporalSuperResolution/TSRConvolutionNetworkPass.ush
@@ -0,0 +1,126 @@
+// Copyright Epic Games, Inc. All Rights Reserved.
+
+#pragma once
+
+#include "TSRCommon.ush"
+
+
+//------------------------------------------------------- DERIVED
+	
+#if DIM_WAVE_SIZE == 0 || DEBUG_DIGESTABLE
+	/** Load an entire 32x32 with 1x1 worth of data by lane, only using LDS. */
+	#define GROUP_TILE_SIZE 32
+
+	#define WAVE_COUNT_X GROUP_TILE_SIZE
+	#define WAVE_COUNT_Y GROUP_TILE_SIZE
+
+	#define LANE_COUNT_X 1
+	#define LANE_COUNT_Y 1
+
+	#define LANE_STRIDE_X 1
+	#define LANE_STRIDE_Y 1
+
+	#define LDS_SIZE (GROUP_TILE_SIZE * GROUP_TILE_SIZE)
+
+	#define TENSOR_REGISTER_LAYOUT FRegisters_S_OneVector
+
+#elif DIM_WAVE_SIZE == 64 || DIM_WAVE_SIZE == 32 || DIM_WAVE_SIZE == 16
+	/** Layouts wave in row major 32x32 tile so an entire row of pixels fits within a single wave, and waveops horizontally, but LDS vertically.
+	 * Each lane worth of a 1x4 of data per lane to minimise to trade VGPR pressure for both number of wave ops and LDS.
+	 */
+	#define GROUP_TILE_SIZE 32
+
+	#if DIM_WAVE_SIZE == 16
+		// We need a LANE_STRIDE_X of at least 2 to have the entire row of pixels fits within a single wave.
+		#define LANE_STRIDE_X 2
+
+		// Only process 2 pixels per lane to reduce register pressure on Intel Arc Alchemist GPUs
+		#define LANE_STRIDE_Y 1
+	#else
+		#define LANE_STRIDE_X 1
+		#if CONFIG_COMPILE_FP16
+			#define LANE_STRIDE_Y 4
+		#else
+			#define LANE_STRIDE_Y 2
+		#endif
+	#endif
+
+	#define WAVE_COUNT_X 1
+	#define WAVE_COUNT_Y (GROUP_TILE_SIZE / LANE_STRIDE_Y)
+
+	#define LANE_COUNT_X (GROUP_TILE_SIZE / LANE_STRIDE_X)
+	#define LANE_COUNT_Y 1
+
+	#if LANE_STRIDE_Y > 1
+		// Each wave needs to output top and bottom rows of pixels into LDS;
+		#define LDS_SIZE ((2 * LANE_STRIDE_X) * LANE_COUNT_X * WAVE_COUNT_Y)
+	#else
+		#define LDS_SIZE (GROUP_TILE_SIZE * GROUP_TILE_SIZE)
+	#endif
+
+	#ifdef PLATFORM_SPECIFIC_TENSOR_REGISTER_LAYOUT
+		#define TENSOR_REGISTER_LAYOUT PLATFORM_SPECIFIC_TENSOR_REGISTER_LAYOUT
+	#endif
+	
+	/** Manually despill to LDS to save more register pressure. */
+	#if CONFIG_COMPILE_FP16 && !CONFIG_SCENE_COLOR_ALPHA
+		#define CONFIG_LDS_DESPILL 1
+		#define CONFIG_CACHE_REFETCH 1
+	#endif
+
+#else
+	#error Unknown wave size
+#endif
+
+#if CONFIG_COMPILE_FP16
+	#define LDS_DWORD_COMPONENT_COUNT 2
+#else
+	#define LDS_DWORD_COMPONENT_COUNT 4
+#endif
+
+#ifndef CONFIG_LDS_DESPILL
+	#define CONFIG_LDS_DESPILL 0
+#endif
+
+/** Whether should minimise register pressure by reloading some data from cache using memory fetches. */
+#ifndef CONFIG_CACHE_REFETCH
+	#define CONFIG_CACHE_REFETCH 0
+#endif
+
+/** Manually despill to LDS on RDNA platforms to save more register pressure. */
+#if CONFIG_LDS_DESPILL
+	#define LDS_DESPILL_THREAD_COUNT (WAVE_COUNT_X * WAVE_COUNT_Y) * (LANE_COUNT_X * LANE_COUNT_Y)
+	#define LDS_DESPILL_OFFSET       (LANE_STRIDE_X * LANE_STRIDE_Y) * GetVectorDwordSize(tsr_half, 4)
+
+	#define LDS_DESPILL_OFFSET_COUNT 2
+	
+	#define LDS_DESPILL_DWORD_COUNT  LDS_DESPILL_THREAD_COUNT * LDS_DESPILL_OFFSET * LDS_DESPILL_OFFSET_COUNT
+#endif
+
+/** Size of the SIMD per lane that also number of input pixel loaded into each individual lanes. */
+#define SIMD_SIZE (LANE_STRIDE_X * LANE_STRIDE_Y)
+
+
+//------------------------------------------------------- INCLUDES
+
+#include "TSRConvolutionNetwork.ush"
+
+
+//------------------------------------------------------- TYPEDEFS
+
+#define tsr_tensor_float  TLaneVector2D<float, 1, LANE_STRIDE_X, LANE_STRIDE_Y>
+#define tsr_tensor_float2 TLaneVector2D<float, 2, LANE_STRIDE_X, LANE_STRIDE_Y>
+#define tsr_tensor_half  TLaneVector2D<tsr_half, 1, LANE_STRIDE_X, LANE_STRIDE_Y>
+#define tsr_tensor_half2 TLaneVector2D<tsr_half, 2, LANE_STRIDE_X, LANE_STRIDE_Y>
+#define tsr_tensor_half3 TLaneVector2D<tsr_half, 3, LANE_STRIDE_X, LANE_STRIDE_Y>
+#define tsr_tensor_half4 TLaneVector2D<tsr_half, 4, LANE_STRIDE_X, LANE_STRIDE_Y>
+#define tsr_tensor_halfC TLaneVector2D<tsr_half, CONFIG_CHANNEL_COUNT, LANE_STRIDE_X, LANE_STRIDE_Y>
+#define tsr_tensor_ushort  TLaneVector2D<tsr_ushort, 1, LANE_STRIDE_X, LANE_STRIDE_Y>
+#define tsr_tensor_short2  TLaneVector2D<tsr_ushort, 2, LANE_STRIDE_X, LANE_STRIDE_Y>
+#define tsr_tensor_bool    TLaneVector2D<bool, 1, LANE_STRIDE_X, LANE_STRIDE_Y>
+#define tsr_tensor_bool4   TLaneVector2D<bool, 4, LANE_STRIDE_X, LANE_STRIDE_Y>
+#define tsr_tensor_boolC   TLaneVector2D<bool, CONFIG_CHANNEL_COUNT, LANE_STRIDE_X, LANE_STRIDE_Y>
+
+#ifndef PLATFORM_SPECIFIC_ISOLATE
+	#define PLATFORM_SPECIFIC_ISOLATE ISOLATE
+#endif
diff --git a/Engine/Shaders/Private/TemporalSuperResolution/TSRRejectShading.usf b/Engine/Shaders/Private/TemporalSuperResolution/TSRRejectShading.usf
--- a/Engine/Shaders/Private/TemporalSuperResolution/TSRRejectShading.usf
+++ b/Engine/Shaders/Private/TemporalSuperResolution/TSRRejectShading.usf
@@ -1,15 +1,9 @@
 // Copyright Epic Games, Inc. All Rights Reserved.
 
-//------------------------------------------------------- CONFIG WAVE SIZE
-
-// needed if CONFIG_OPTIMIZED_LDS_ACCESS == 0 for ConvertWaveBroadcastToLDS()
-//#if DIM_WAVE_SIZE == 0 // fallback to LDS
-//	#define WAVE_BROADCAST_GROUPSIZE (16 * 16)
-//#endif
-
 //------------------------------------------------------- INCLUDES
 
-#include "TSRCommon.ush"
+#include "TSRConvolutionNetworkPass.ush"
+#include "TSRShadingAnalysis.ush"
 
 
 //------------------------------------------------------- CONFIG
@@ -26,132 +20,8 @@
 #define CONFIG_HISTORY_RESURRECTION (DIM_HISTORY_RESURRECTION)
 
 
-//------------------------------------------------------- DERIVED
-	
-#if DIM_WAVE_SIZE == 0 || DEBUG_DIGESTABLE
-	/** Load an entire 32x32 with 1x1 worth of data by lane, only using LDS. */
-	#define GROUP_TILE_SIZE 32
-
-	#define WAVE_COUNT_X GROUP_TILE_SIZE
-	#define WAVE_COUNT_Y GROUP_TILE_SIZE
-
-	#define LANE_COUNT_X 1
-	#define LANE_COUNT_Y 1
-
-	#define LANE_STRIDE_X 1
-	#define LANE_STRIDE_Y 1
-
-	#define LDS_SIZE (GROUP_TILE_SIZE * GROUP_TILE_SIZE)
-
-	#define TENSOR_REGISTER_LAYOUT FRegisters_S_OneVector
-
-#elif DIM_WAVE_SIZE == 64 || DIM_WAVE_SIZE == 32 || DIM_WAVE_SIZE == 16
-	/** Layouts wave in row major 32x32 tile so an entire row of pixels fits within a single wave, and waveops horizontally, but LDS vertically.
-	 * Each lane worth of a 1x4 of data per lane to minimise to trade VGPR pressure for both number of wave ops and LDS.
-	 */
-	#define GROUP_TILE_SIZE 32
-
-	#if DIM_WAVE_SIZE == 16
-		// We need a LANE_STRIDE_X of at least 2 to have the entire row of pixels fits within a single wave.
-		#define LANE_STRIDE_X 2
-
-		// Only process 2 pixels per lane to reduce register pressure on Intel Arc Alchemist GPUs
-		#define LANE_STRIDE_Y 1
-	#else
-		#define LANE_STRIDE_X 1
-		#if CONFIG_COMPILE_FP16
-			#define LANE_STRIDE_Y 4
-		#else
-			#define LANE_STRIDE_Y 2
-		#endif
-	#endif
-
-	#define WAVE_COUNT_X 1
-	#define WAVE_COUNT_Y (GROUP_TILE_SIZE / LANE_STRIDE_Y)
-
-	#define LANE_COUNT_X (GROUP_TILE_SIZE / LANE_STRIDE_X)
-	#define LANE_COUNT_Y 1
-
-	#if LANE_STRIDE_Y > 1
-		// Each wave needs to output top and bottom rows of pixels into LDS;
-		#define LDS_SIZE ((2 * LANE_STRIDE_X) * LANE_COUNT_X * WAVE_COUNT_Y)
-	#else
-		#define LDS_SIZE (GROUP_TILE_SIZE * GROUP_TILE_SIZE)
-	#endif
-
-	#ifdef PLATFORM_SPECIFIC_TENSOR_REGISTER_LAYOUT
-		#define TENSOR_REGISTER_LAYOUT PLATFORM_SPECIFIC_TENSOR_REGISTER_LAYOUT
-	#endif
-	
-	/** Manually despill to LDS to save more register pressure. */
-	#if CONFIG_COMPILE_FP16 && !CONFIG_SCENE_COLOR_ALPHA
-		#define CONFIG_LDS_DESPILL 1
-		#define CONFIG_CACHE_REFETCH 1
-	#endif
-
-#else
-	#error Unknown wave size
-#endif
-
-#if CONFIG_COMPILE_FP16
-	#define LDS_DWORD_COMPONENT_COUNT 2
-#else
-	#define LDS_DWORD_COMPONENT_COUNT 4
-#endif
-
-#ifndef CONFIG_LDS_DESPILL
-	#define CONFIG_LDS_DESPILL 0
-#endif
-
-/** Whether should minimise register pressure by reloading some data from cache using memory fetches. */
-#ifndef CONFIG_CACHE_REFETCH
-	#define CONFIG_CACHE_REFETCH 0
-#endif
-
-/** Manually despill to LDS on RDNA platforms to save more register pressure. */
-#if CONFIG_LDS_DESPILL
-	#define LDS_DESPILL_THREAD_COUNT (WAVE_COUNT_X * WAVE_COUNT_Y) * (LANE_COUNT_X * LANE_COUNT_Y)
-	#define LDS_DESPILL_OFFSET       (LANE_STRIDE_X * LANE_STRIDE_Y) * GetVectorDwordSize(tsr_half, 4)
-
-	#define LDS_DESPILL_OFFSET_COUNT 2
-	
-	#define LDS_DESPILL_DWORD_COUNT  LDS_DESPILL_THREAD_COUNT * LDS_DESPILL_OFFSET * LDS_DESPILL_OFFSET_COUNT
-#endif
-
-/** Size of the SIMD per lane that also number of input pixel loaded into each individual lanes. */
-#define SIMD_SIZE (LANE_STRIDE_X * LANE_STRIDE_Y)
-
-//------------------------------------------------------- INCLUDE CONVOLUTION NETWORK
-
-#include "TSRConvolutionNetwork.ush"
-
-
-//------------------------------------------------------- TYPEDEFS
-
-#define tsr_simd_float  TLaneVector2D<float, 1, LANE_STRIDE_X, LANE_STRIDE_Y>
-#define tsr_simd_float2 TLaneVector2D<float, 2, LANE_STRIDE_X, LANE_STRIDE_Y>
-#define tsr_simd_half  TLaneVector2D<tsr_half, 1, LANE_STRIDE_X, LANE_STRIDE_Y>
-#define tsr_simd_half2 TLaneVector2D<tsr_half, 2, LANE_STRIDE_X, LANE_STRIDE_Y>
-#define tsr_simd_half3 TLaneVector2D<tsr_half, 3, LANE_STRIDE_X, LANE_STRIDE_Y>
-#define tsr_simd_half4 TLaneVector2D<tsr_half, 4, LANE_STRIDE_X, LANE_STRIDE_Y>
-#define tsr_simd_halfC TLaneVector2D<tsr_half, CONFIG_CHANNEL_COUNT, LANE_STRIDE_X, LANE_STRIDE_Y>
-#define tsr_simd_ushort  TLaneVector2D<tsr_ushort, 1, LANE_STRIDE_X, LANE_STRIDE_Y>
-#define tsr_simd_short2  TLaneVector2D<tsr_ushort, 2, LANE_STRIDE_X, LANE_STRIDE_Y>
-#define tsr_simd_bool    TLaneVector2D<bool, 1, LANE_STRIDE_X, LANE_STRIDE_Y>
-#define tsr_simd_bool4   TLaneVector2D<bool, 4, LANE_STRIDE_X, LANE_STRIDE_Y>
-#define tsr_simd_boolC   TLaneVector2D<bool, CONFIG_CHANNEL_COUNT, LANE_STRIDE_X, LANE_STRIDE_Y>
-
-#ifndef PLATFORM_SPECIFIC_ISOLATE
-	#define PLATFORM_SPECIFIC_ISOLATE ISOLATE
-#endif
-
-
 //------------------------------------------------------- CONSTANTS
 
-static const float kHistoryGuidePreceptionAdd = 1.0;
-
-static const tsr_half FilteringWeight = rcp(tsr_half(1.0 + 4 * 0.5 + 4 * 0.25));
-
 #if CONFIG_LDS_DESPILL
 	static const uint kDespillInput = LDS_DESPILL_THREAD_COUNT * LDS_DESPILL_OFFSET * 0;
 	static const uint kDespillHistory = LDS_DESPILL_THREAD_COUNT * LDS_DESPILL_OFFSET * 1;
@@ -200,7 +70,7 @@ RWTexture2D<tsr_ushort> AntiAliasMaskOutput;
 
 static uint2 GGroupId = 0;
 
-static tsr_simd_halfC Debug[DEBUG_ARRAY_SIZE];
+static tsr_tensor_halfC Debug[DEBUG_ARRAY_SIZE];
 
 
 //------------------------------------------------------- INPUT FETCHES
@@ -219,9 +89,9 @@ tsr_short2 ComputeElementPixelPos(const uint ElementIndex)
 	return ClampedFetchPixelPos;
 }
 
-tsr_simd_short2 ComputePixelPos()
+tsr_tensor_short2 ComputePixelPos()
 {
-	tsr_simd_short2 PixelPos;
+	tsr_tensor_short2 PixelPos;
 
 	UNROLL_N(SIMD_SIZE)
 	for (uint ElementIndex = 0; ElementIndex < SIMD_SIZE; ElementIndex++)
@@ -245,9 +115,9 @@ tsr_short2 ComputeElementOutputPixelPos(const uint ElementIndex)
 }
 
 void FetchDecimateMask(
-	tsr_simd_short2 PixelPos,
-	out tsr_simd_half VelocityEdge,
-	out tsr_simd_bool bIsDisoccluded)
+	tsr_tensor_short2 PixelPos,
+	out tsr_tensor_half VelocityEdge,
+	out tsr_tensor_bool bIsDisoccluded)
 {
 	UNROLL_N(SIMD_SIZE)
 	for (uint ElementIndex = 0; ElementIndex < SIMD_SIZE; ElementIndex++)
@@ -262,9 +132,9 @@ void FetchDecimateMask(
 } // FetchDecimateMask()
 
 void FetchSceneColorAndTranslucency(
-	tsr_simd_short2 PixelPos,
-	out tsr_simd_halfC OriginalOpaqueInput,
-	out tsr_simd_half4 OriginalTranslucencyInput)
+	tsr_tensor_short2 PixelPos,
+	out tsr_tensor_halfC OriginalOpaqueInput,
+	out tsr_tensor_half4 OriginalTranslucencyInput)
 {
 	UNROLL_N(SIMD_SIZE)
 	for (uint ElementIndex = 0; ElementIndex < SIMD_SIZE; ElementIndex++)
@@ -282,9 +152,9 @@ void FetchSceneColorAndTranslucency(
 void FetchHistoryGuide(
 	Texture2D<tsr_half4> HistoryGuideTexture,
 	Texture2D<tsr_half>  HistoryGuideMetadataTexture,
-	tsr_simd_short2 PixelPos,
-	out tsr_simd_halfC History,
-	out tsr_simd_half  HistoryUncertainty)
+	tsr_tensor_short2 PixelPos,
+	out tsr_tensor_halfC History,
+	out tsr_tensor_half  HistoryUncertainty)
 {
 	UNROLL_N(SIMD_SIZE)
 	for (uint ElementIndex = 0; ElementIndex < SIMD_SIZE; ElementIndex++)
@@ -308,9 +178,9 @@ void FetchHistoryGuide(
 
 // Fetch the previous history history.
 void FetchReprojectedHistoryGuide(
-	tsr_simd_short2 PixelPos,
-	out tsr_simd_halfC History,
-	out tsr_simd_half  HistoryUncertainty)
+	tsr_tensor_short2 PixelPos,
+	out tsr_tensor_halfC History,
+	out tsr_tensor_half  HistoryUncertainty)
 {
 	FetchHistoryGuide(
 		ReprojectedHistoryGuideTexture,
@@ -323,10 +193,10 @@ void FetchReprojectedHistoryGuide(
 #if CONFIG_MOIRE_DETECTION
 
 void FetchMoireInput(
-	tsr_simd_short2 PixelPos,
-	out tsr_simd_half MoireInput,
-	out tsr_simd_half4 PrevMoireHistory,
-	out tsr_simd_half bIsStatic)
+	tsr_tensor_short2 PixelPos,
+	out tsr_tensor_half MoireInput,
+	out tsr_tensor_half4 PrevMoireHistory,
+	out tsr_tensor_half bIsStatic)
 {
 	UNROLL_N(SIMD_SIZE)
 	for (uint ElementIndex = 0; ElementIndex < SIMD_SIZE; ElementIndex++)
@@ -341,470 +211,13 @@ void FetchMoireInput(
 
 #endif // CONFIG_MOIRE_DETECTION
 
-
-//------------------------------------------------------- COMPOSE SEPARATE TRANSLUCENCY
-
-/** Compose translucency in linear color space. */
-CALL_SITE_DEBUGLOC
-tsr_simd_halfC ComposeTranslucency(tsr_simd_halfC Color, tsr_simd_half4 Translucency)
-{
-	tsr_simd_halfC ComposedColor;
-	ComposedColor = Color * tsr_simd_halfC::Vectorize(Translucency[3]) + ResizeChannels<CONFIG_CHANNEL_COUNT>(Translucency);
-	#if CONFIG_SCENE_COLOR_ALPHA
-		ComposedColor.SetComponent(3, Color[3] * Translucency[3]);
-	#endif
-	#if CONFIG_SCENE_COLOR_ALPHA
-		ComposedColor = min(ComposedColor, tsr_simd_halfC::Const(LargestSceneColorRGBA));
-	#else
-		ComposedColor = min(ComposedColor, tsr_simd_halfC::Const(LargestSceneColorRGB));
-	#endif
-	return ComposedColor;
-}
-
-
-//------------------------------------------------------- CHANGE LDR EXPOSURE
-
-CALL_SITE_DEBUGLOC
-tsr_simd_halfC RemovePerceptionAdd(tsr_simd_halfC SrcColor, const float SourcePerceptionAdd)
-{
-	tsr_simd_halfC Minus = tsr_simd_halfC::Const(tsr_half(1.0)) - SrcColor;
-	tsr_simd_halfC LDRToLinear = min(tsr_simd_halfC::Const(tsr_half(SourcePerceptionAdd)) * rcp(Minus), tsr_simd_halfC::Const(tsr_half(MaxHalfFloat)));
-
-	return SrcColor * LDRToLinear;
-}
-
-CALL_SITE_DEBUGLOC
-tsr_simd_halfC AddPerceptionAdd(tsr_simd_halfC SrcColor, const float DestPerceptionAdd)
-{
-	SrcColor = SrcColor * rcp(SrcColor + tsr_simd_halfC::Const(tsr_half(DestPerceptionAdd)));
-	return SrcColor;
-}
-
-CALL_SITE_DEBUGLOC
-tsr_simd_halfC ChangePerceptionAdd(tsr_simd_halfC SrcColor, const float SourcePerceptionAdd, const float DestPerceptionAdd)
-{
-	SrcColor = RemovePerceptionAdd(SrcColor, SourcePerceptionAdd);
-	SrcColor = AddPerceptionAdd(SrcColor, DestPerceptionAdd);
-	return SrcColor;
-}
-
-
-//------------------------------------------------------- SPATIAL ANTI-ALIASER
-
-// Compute the LDR luminance the spatial anti-aliaser should anti-alias.
-CALL_SITE_DEBUGLOC
-tsr_simd_half ComputeSpatialAntiAliaserLumaLDR(tsr_simd_halfC SceneColor)
-{
-	const tsr_half SpatialAAExposure = tsr_half(0.5);
-		
-	tsr_simd_half AALumaLDR;
-
-	UNROLL_N(SIMD_SIZE)
-	for (uint ElementIndex = 0; ElementIndex < SIMD_SIZE; ElementIndex++)
-	{
-		tsr_half PixelLuma = dot(SceneColor.GetElement(ElementIndex).rgb, tsr_half3(0.299f, 0.587f, 0.114f));
-		AALumaLDR.SetElement(ElementIndex, PixelLuma / (SpatialAAExposure + PixelLuma));
-	}
-
-	return AALumaLDR;
-}
-
-// Returns whether the aliasing is visible at all, to avoid paying spatial-antialiaser on low contrasts areas.
-CALL_SITE_DEBUGLOC
-tsr_simd_bool IsAliasingVisible(tsr_simd_halfC ExposedInputBoxSize)
-{
-	#if CONFIG_SCENE_COLOR_ALPHA
-		tsr_simd_half ExposedInputLuminanceBox = dot(ExposedInputBoxSize, tsr_simd_halfC::Const(tsr_halfC(0.299f, 0.587f, 0.114f, 0.0)));
-	#else
-		tsr_simd_half ExposedInputLuminanceBox = dot(ExposedInputBoxSize, tsr_simd_halfC::Const(tsr_halfC(0.299f, 0.587f, 0.114f)));
-	#endif
-
-	tsr_simd_bool bAliasingIsVisible = ExposedInputLuminanceBox > tsr_simd_half::Const(SPATIAL_ANTI_ALIASER_MIN_LUMIMANCE);
-
-	return bAliasingIsVisible;
-}
-
-// Whether to run spatial anti-aliaser on the frame. Must match SpatialAntiAliasingLerp.
-CALL_SITE_DEBUGLOC
-tsr_simd_bool ShouldSpatialAntiAlias(tsr_simd_bool bIsDisoccluded, tsr_simd_bool bResurrectHistory, tsr_simd_bool bAliasingIsVisible, tsr_simd_half RejectionBlendFinal)
-{
-	tsr_simd_bool bSpatialAntiAlias = and(bAliasingIsVisible, or(RejectionBlendFinal < tsr_simd_half::Const(tsr_half(0.25)), and(bIsDisoccluded, not(bResurrectHistory))));
-	
-	return bSpatialAntiAlias;
-}
-
-
-//------------------------------------------------------- FLICKERING TEMPORAL ANALYSIS
-
-void ComputeMoireError(
-	tsr_simd_bool  bIsDisoccluded,
-	tsr_simd_half  bIsStatic,
-	tsr_simd_half  MoireInput,
-	tsr_simd_halfC OriginalOpaqueInput,
-	tsr_simd_half4 OriginalTranslucencyInput,
-	tsr_simd_half4 PrevMoireHistory,
-	out tsr_simd_half4 OutMoireHistory,
-	out tsr_simd_half  OutMoireError)
-{
-	const tsr_half MinBlendFinal = tsr_half(0.05);
-	const tsr_half MoireEncodingError = tsr_half(rcp(127.0));
-	const tsr_half MaxPrevTotalVariationCount = tsr_half(20.0);
-
-	// Compute the threshold at which the variation count should quick in
-	const tsr_half TotalVariationCountThreshold = tsr_half(1.0 / (1.0 - pow(1.0 - MinBlendFinal, tsr_half(FlickeringFramePeriod))));
-
-	// Input luminance to monitor flickering from.
-	tsr_simd_half InputLuma = MoireInput;
-
-	// How much input luminance from the TSRComputeMoireLuma.usf measurement has been modified since for instance with fog, water.
-	tsr_simd_half InputLumaModification;
-	{
-		const tsr_half InputLumaEncodingError = tsr_half(0.5 / 255.0);
-			
-		tsr_simd_halfC OriginalInput = ComposeTranslucency(OriginalOpaqueInput, OriginalTranslucencyInput);
-
-		OriginalInput       = AddPerceptionAdd(OriginalInput, PerceptionAdd);
-		OriginalOpaqueInput = AddPerceptionAdd(OriginalOpaqueInput, PerceptionAdd);
-
-		tsr_simd_half SceneColorLuma;
-		tsr_simd_half TranslucencyChange;
-		UNROLL_N(SIMD_SIZE)
-		for (uint SimdIndex = 0; SimdIndex < SIMD_SIZE; SimdIndex++)
-		{
-			SceneColorLuma.SetElement(SimdIndex, dot(OriginalOpaqueInput.GetElement(SimdIndex).rgb, kMoireLumaWeights));
-			TranslucencyChange.SetElement(SimdIndex, dot(abs(OriginalOpaqueInput.GetElement(SimdIndex).rgb - OriginalInput.GetElement(SimdIndex).rgb), kMoireLumaWeights));
-		}
-
-		InputLumaModification = tsr_simd_half::Const(0.0);
-		InputLumaModification = max(InputLumaModification, abs(InputLuma - SceneColorLuma) - tsr_simd_half::Const(InputLumaEncodingError));
-		InputLumaModification = max(InputLumaModification, TranslucencyChange);
-		InputLumaModification = max(InputLumaModification, tsr_simd_half::Const(1.0) - OriginalTranslucencyInput[3]);
-	}
-
-	// Unpack history
-	tsr_simd_half PrevFlickeringHistory = PrevMoireHistory[0];
-	tsr_simd_half PrevGradient = PrevMoireHistory[1] * tsr_half(255.0 / 127.0) - tsr_half(1.0);
-	//tsr_simd_half PrevTotalVariation = PrevMoireHistory[2] * tsr_simd_half::Const(1.0 - MinBlendFinal);
-	tsr_simd_half PrevTotalVariation = PrevMoireHistory[2] * tsr_simd_half::Const(tsr_half(1.0) - MinBlendFinal);
-	tsr_simd_half PrevTotalVariationCount = PrevMoireHistory[3] * MaxPrevTotalVariationCount * tsr_simd_half::Const(tsr_half(1.0) - MinBlendFinal);
-	
-	// Discard history on moving objects.
-	PrevTotalVariation = PrevTotalVariation * bIsStatic;
-	PrevTotalVariationCount = PrevTotalVariationCount * bIsStatic;
-
-	// Run the exact same heuristic as MeasureRejection on a single channel.
-	tsr_simd_half FilteredInputLumaModification;
-	tsr_simd_half OutRejectionBlendFinal;
-	tsr_simd_half OutRejectionClampBlend;
-	{
-		tsr_simd_half BackbufferQuantizationErrorVector = tsr_simd_half::Const(MoireEncodingError);
-		
-		tsr_simd_half InputC0 = InputLuma;
-		tsr_simd_half HistoryC0 = PrevFlickeringHistory;
-
-		// Anhilate history
-		tsr_simd_half InputC2 = InputC0;
-		tsr_simd_half HistoryC2 = HistoryC0;
-		AnnihilateMutuallySingleChannel3x3(InputLuma, HistoryC0, /* out */ InputC2, /* out */ HistoryC2);
-
-		// Convolve the input luma to have the exact same moire pattern as FilteredInput below.
-		tsr_simd_half FilteredInput;
-		tsr_simd_half FilteredHistory;
-		Deconcatenate(Blur3x3(Concatenate(InputC2, InputLumaModification, HistoryC2)), /* out */ FilteredInput, /* out */ FilteredInputLumaModification, /* out */ FilteredHistory);
-		
-		// Compute clamping box of the limance.
-		tsr_simd_half FilteredBoxMin;
-		tsr_simd_half FilteredBoxMax;
-		tsr_simd_half InputC2BoxSize;
-		{
-			tsr_simd_half2 Min;
-			tsr_simd_half2 Max;
-			MinMax3x3(Concatenate(FilteredInput, InputC2), /* out */ Min, /* out */ Max);
-			FilteredBoxMin = Min[0];
-			FilteredBoxMax = Max[0];
-		
-			InputC2BoxSize = Max[1] - Min[1];
-		}
-
-		// Compute the clamp error
-		tsr_simd_half ClampError;
-		{
-			tsr_simd_half TotalVarInputDiffC0C2;
-			tsr_simd_half TotalVarInputC2;
-			Deconcatenate(abs(TotalVariation3x3(Concatenate(abs(InputC0 - InputC2), InputC2))), /* out */ TotalVarInputDiffC0C2, /* out */ TotalVarInputC2);
-		
-			ClampError = BackbufferQuantizationErrorVector;
-			ClampError = max(ClampError, Blur3x3(min(TotalVarInputDiffC0C2, TotalVarInputC2)));
-			ClampError = max(ClampError, InputC2BoxSize * tsr_half(FilteringWeight * 0.25));
-			ClampError = ClampError + BackbufferQuantizationErrorVector;
-		}
-
-		// Apply clamp error
-		FilteredBoxMin = FilteredBoxMin - ClampError;
-		FilteredBoxMax = FilteredBoxMax + ClampError;
-
-		tsr_simd_half BoxSize = (
-			InputC2BoxSize * tsr_simd_half::Const(FilteringWeight) +
-			(BackbufferQuantizationErrorVector * tsr_half(2 * FilteringWeight)));
-
-		tsr_simd_half ClampedFilteredHistory = fastClamp(FilteredHistory, FilteredBoxMin, FilteredBoxMax);
-
-		tsr_simd_half Delta = max(abs(FilteredInput - FilteredHistory), BoxSize);
-	
-		tsr_simd_half RawClampedEnergy = abs(ClampedFilteredHistory - FilteredHistory);
-		tsr_simd_half RawRejection     = saturate(tsr_simd_half::Const(1.0) - max(RawClampedEnergy, tsr_simd_half::Const(0.0)) * rcp(Delta));
-
-		tsr_simd_half FilteredClampedEnergy;
-		Deconcatenate(MaxRMinG3x3(Median3x3(Concatenate(RawClampedEnergy, RawRejection))), /* out */ FilteredClampedEnergy, /* out */ OutRejectionClampBlend);
-
-		OutRejectionBlendFinal = saturate(tsr_simd_half::Const(1.0) - max(FilteredClampedEnergy, tsr_simd_half::Const(0.0)) * rcp(Delta));
-	}
-	
-	// Simulate what happens in TSRUpdateHistory.usf
-	tsr_simd_half FinalFlickeringHistory;
-	tsr_simd_half CurrentGradient;
-	{
-		tsr_simd_half ClampedPrevFlickeringHistory = ClampPlus3x3(PrevFlickeringHistory, InputLuma);
-		tsr_simd_half FinalPrevFlickeringHistory = lerp(ClampedPrevFlickeringHistory, PrevFlickeringHistory, OutRejectionClampBlend);
-
-		tsr_simd_half FinalBlendFinal = max(tsr_simd_half::Const(1.0) - OutRejectionBlendFinal, tsr_simd_half::Const(MinBlendFinal));
-		
-		FinalFlickeringHistory = InputLuma * FinalBlendFinal - FinalPrevFlickeringHistory * (FinalBlendFinal - tsr_simd_half::Const(1.0));
-		
-		tsr_simd_half GhostingHistory = InputLuma * tsr_simd_half::Const(MinBlendFinal) - PrevFlickeringHistory * tsr_simd_half::Const(MinBlendFinal - (1.0));
-
-		CurrentGradient = FinalFlickeringHistory - GhostingHistory;
-	}
-
-	// Remove how much luma has changed from luma measurement to final scene color from the gradient.
-	for (uint SimdIndex = 0; SimdIndex < SIMD_SIZE; SimdIndex++)
-	{
-		tsr_half LocalCurrentGradient = CurrentGradient.GetElement(SimdIndex);
-		tsr_half LocalInputLumaModification = FilteredInputLumaModification.GetElement(SimdIndex);
-			
-		CurrentGradient.SetElement(SimdIndex, clamp(LocalCurrentGradient * tsr_half(POSITIVE_INFINITY), tsr_half(-1.0), tsr_half(1.0)) * max(abs(LocalCurrentGradient) - LocalInputLumaModification, tsr_half(0.0)));
-	}
-
-	#if 0
-		tsr_simd_half InputBoxSize = InputLumaMax - InputLumaMin;
-		CurrentGradient = sign(CurrentGradient) * max(abs(CurrentGradient) - tsr_simd_half::Const(0.5) * InputBoxSize, tsr_simd_half::Const(0.0));
-	#endif
-
-	tsr_simd_half IsFlicker;
-	UNROLL_N(SIMD_SIZE)
-	for (uint SimdIndex = 0; SimdIndex < SIMD_SIZE; SimdIndex++)
-	{
-		tsr_half LocalCurrentGradient = CurrentGradient.GetElement(SimdIndex);
-		tsr_half LocalPrevGradient = PrevGradient.GetElement(SimdIndex);
-
-		bool bGradientsAreSameSign = LocalCurrentGradient * LocalPrevGradient > tsr_half(0.0);
-		bool bGradientsAreWithinError = abs(LocalCurrentGradient) < MoireEncodingError || abs(LocalPrevGradient) < MoireEncodingError;
-		//bool bCurrentGradientSubstentialEnough = abs(LocalCurrentGradient) > View.GeneralPurposeTweak * abs(LocalPrevGradient); // - 2.0 * MoireEncodingError;
-
-		tsr_half LocalIsFlicker = select((
-			bGradientsAreSameSign ||
-			bGradientsAreWithinError ||
-			//!bCurrentGradientSubstentialEnough ||
-			bCameraCut), tsr_half(0.0), tsr_half(1.0));
-
-		IsFlicker.SetElement(SimdIndex, LocalIsFlicker);
-	}
-
-	// Compute the total variation of the gradient over time.
-	tsr_simd_half GradientVariation = min(abs(PrevGradient), abs(CurrentGradient)) * IsFlicker;
-
-	// Dilates the total variation and contribution to keep stability despite inacuracies in the history reprojection
-	// of these need a nearest filter which isn't precise at all.
-	tsr_simd_half TotalVariationContribution;
-	tsr_simd_half TotalVariationCountContribution;
-	Deconcatenate(Max3x3(Max3x3(Concatenate(GradientVariation, IsFlicker))), /* out */ TotalVariationContribution, /* out */ TotalVariationCountContribution);
-
-	UNROLL_N(SIMD_SIZE)
-	for (uint SimdIndex = 0; SimdIndex < SIMD_SIZE; SimdIndex++)
-	{
-		tsr_half LocalCurrentGradient = CurrentGradient.GetElement(SimdIndex);
-
-		tsr_half LocalPrevGradient = PrevGradient.GetElement(SimdIndex);
-		tsr_half LocalPrevTotalVariation = PrevTotalVariation.GetElement(SimdIndex);
-		tsr_half LocalPrevTotalVariationCount = PrevTotalVariationCount.GetElement(SimdIndex);
-		bool LocalIsDisoccluded = bIsDisoccluded.GetElement(SimdIndex);
-
-		tsr_half LocalIsFlicker = IsFlicker.GetElement(SimdIndex);
-
-		tsr_half LocalTotalVariationContribution = TotalVariationContribution.GetElement(SimdIndex);
-		tsr_half LocalTotalVariationCountContribution = TotalVariationCountContribution.GetElement(SimdIndex);
-			
-		// Blend out previous gradient
-		tsr_half LocalPrevBlendedGradient = LocalPrevGradient * (tsr_half(1.0) - MinBlendFinal) * (tsr_half(1.0) - LocalIsFlicker);
-
-		// Compute new gradient.
-		#if 1
-			tsr_half LocalNewGradient = LocalPrevBlendedGradient + LocalCurrentGradient;
-		#else
-			tsr_half LocalNewGradient = abs(LocalPrevBlendedGradient) > abs(LocalCurrentGradient) ? LocalPrevBlendedGradient : LocalCurrentGradient;
-		#endif
-			
-		// Accumulate current frame variation.
-		tsr_half LocalNewTotalVariation = LocalPrevTotalVariation + LocalTotalVariationContribution;
-		tsr_half LocalNewTotalVariationCount = LocalPrevTotalVariationCount + LocalTotalVariationCountContribution;
-
-		// Discard all moire history on parallax disocclusion.
-		LocalNewGradient            = select(LocalIsDisoccluded, tsr_half(0.0), LocalNewGradient);
-		LocalNewTotalVariation      = select(LocalIsDisoccluded, tsr_half(0.0), LocalNewTotalVariation);
-		LocalNewTotalVariationCount = select(LocalIsDisoccluded, tsr_half(0.0), LocalNewTotalVariationCount);
-
-		// Quantise total variation and count to history bit depth to variation drift in the history over time.
-		tsr_half LocalQuantizedNewTotalVariationCount = floor(LocalNewTotalVariationCount * (tsr_half(255.0) / MaxPrevTotalVariationCount)) * (MaxPrevTotalVariationCount / tsr_half(255.0));
-		LocalNewTotalVariation *= LocalQuantizedNewTotalVariationCount * SafeRcp(LocalNewTotalVariationCount);
-
-		// Compute the final luminance 
-		#if 1
-			tsr_half LocalCountFadeIn = saturate(LocalNewTotalVariationCount / TotalVariationCountThreshold - tsr_half(0.5));
-		#else
-			tsr_half LocalCountFadeIn = saturate(LocalNewTotalVariationCount - TotalVariationCountThreshold);
-		#endif
-		tsr_half LocalMoireError = (abs(LocalNewTotalVariation * SafeRcp(LocalNewTotalVariationCount)) + LocalNewTotalVariationCount * MoireEncodingError) * LocalCountFadeIn;
-
-		tsr_half4 PackedMoireHistory = tsr_half4(
-			FinalFlickeringHistory.GetElement(SimdIndex),
-			LocalNewGradient * tsr_half(127.0 / 255.0) + tsr_half(127.0 / 255.0),
-			LocalNewTotalVariation,
-			LocalQuantizedNewTotalVariationCount * rcp(MaxPrevTotalVariationCount));
-
-		OutMoireHistory.SetElement(SimdIndex, PackedMoireHistory);
-
-		OutMoireError.SetElement(SimdIndex, LocalMoireError);
-	}
-		
-	// Moving object in the scene have very low chance of creating moire pattern with TAA jitter, so discard the moire error to still
-	// have character animation ghost free evem as environement shadow project different things on character.
-	OutMoireError = saturate(OutMoireError * bIsStatic - InputLumaModification);
-} // ComputeMoireError()
-
-
-//------------------------------------------------------- SHADING REJECTION
-
-template<bool ApplyMoireError>
-void MeasureRejection(
-	tsr_simd_halfC InputC0,
-	tsr_simd_halfC InputC2,
-	tsr_simd_halfC HistoryC2,
-	tsr_simd_half  MoireError,
-	out tsr_simd_half OutRejectionBlendFinal,
-	out tsr_simd_half OutRejectionClampBlend)
-{
-	tsr_simd_halfC BackbufferQuantizationErrorVector = tsr_simd_halfC::Const(MeasureBackbufferLDRQuantizationError());
-		
-	tsr_simd_halfC FilteredHistory = Blur3x3(HistoryC2);
-	
-	tsr_simd_halfC TotalVarInputDiffC0C2 = abs(TotalVariation3x3(abs(InputC0 - InputC2)));
-	tsr_simd_halfC TotalVarInputC2 = abs(TotalVariation3x3(InputC2));
-
-	tsr_simd_halfC ClampError = BackbufferQuantizationErrorVector;
-	ClampError = max(ClampError, Blur3x3(min(TotalVarInputDiffC0C2, TotalVarInputC2)));
-	
-	tsr_simd_halfC InputC2BoxSize = MaxMinusMin3x3(InputC2);
-	ClampError = max(ClampError, InputC2BoxSize * tsr_half(FilteringWeight * 0.25));
-	ClampError = ClampError + BackbufferQuantizationErrorVector;
-
-	tsr_simd_halfC FilteredInput   = Blur3x3(InputC2);
-
-	// Build the clamping box of the filtered signal
-	tsr_simd_halfC FilteredBoxMin;
-	tsr_simd_halfC FilteredBoxMax;
-	MinMax3x3(FilteredInput, /* out */ FilteredBoxMin, /* out */ FilteredBoxMax);
-	FilteredBoxMin = FilteredBoxMin - ClampError;
-	FilteredBoxMax = FilteredBoxMax + ClampError;
-		
-	// Clamp filtered signal
-	tsr_simd_halfC ClampedFilteredHistory = fastClamp(FilteredHistory, FilteredBoxMin, FilteredBoxMax);
-	BRANCH
-	if (ApplyMoireError && bEnableFlickeringHeuristic)
-	{
-		tsr_simd_halfC MoireErrorSize;
-			
-		UNROLL_N(SIMD_SIZE)
-		for (uint SimdIndex = 0; SimdIndex < SIMD_SIZE; SimdIndex++)
-		#if 0
-		{
-			tsr_halfC ClampedEnergy = abs(FlickeringClampedFilteredHistory.GetElement(SimdIndex) - FilteredHistory.GetElement(SimdIndex));
-			tsr_halfC ClampedEnergyAmount = ClampedEnergy * SafeRcp(ClampedEnergy[0] + ClampedEnergy[1] + ClampedEnergy[2]);
-
-			// Apply the moire error on each individual channel based on how much they need to clamp.
-			MoireErrorSize.SetElement(SimdIndex, ClampedEnergyAmount * (kMoireLumaToChannel * MoireError.GetElement(SimdIndex)));
-		}
-		#else
-		{
-			MoireErrorSize.SetElement(SimdIndex, kMoireLumaToChannel * MoireError.GetElement(SimdIndex));
-		}
-		#endif
-
-		tsr_simd_halfC StableFilteredBoxMin = min(FilteredBoxMin, FilteredBoxMax - MoireErrorSize);
-		tsr_simd_halfC StableFilteredBoxMax = max(FilteredBoxMax, FilteredBoxMin + MoireErrorSize);
-			
-		ClampedFilteredHistory = fastClamp(FilteredHistory, StableFilteredBoxMin, StableFilteredBoxMax);
-	}
-		
-	// Compute the box size.
-	tsr_simd_halfC BoxSize;
-	{
-		BoxSize = (
-			InputC2BoxSize * tsr_simd_halfC::Const(FilteringWeight) +
-			(BackbufferQuantizationErrorVector * tsr_half(2 * FilteringWeight)));
-
-		BRANCH
-		if (ApplyMoireError && bEnableFlickeringHeuristic)
-		{
-			tsr_simd_halfC MoireErrorSize;
-			UNROLL_N(SIMD_SIZE)
-			for (uint SimdIndex = 0; SimdIndex < SIMD_SIZE; SimdIndex++)
-			{
-				tsr_halfC ClampedEnergy = abs(ClampedFilteredHistory.GetElement(SimdIndex) - FilteredHistory.GetElement(SimdIndex));
-				tsr_halfC ClampedEnergyAmount = ClampedEnergy * SafeRcp(ClampedEnergy[0] + ClampedEnergy[1] + ClampedEnergy[2]);
-
-				// Apply the moire error on each individual channel based on how much they need to clamp.
-				MoireErrorSize.SetElement(SimdIndex, ClampedEnergyAmount * ((FilteringWeight * tsr_half(3.0)) * MoireError.GetElement(SimdIndex)));
-			}
-			BoxSize = max(BoxSize, MoireErrorSize);
-		}
-	}
-
-	tsr_simd_halfC Delta = max(abs(FilteredInput - FilteredHistory), BoxSize);
-
-	tsr_simd_halfC RawClampedEnergy      = abs(ClampedFilteredHistory - FilteredHistory);
-
-	tsr_simd_halfC RawFactor        = saturate(tsr_simd_halfC::Const(1.0) - RawClampedEnergy * rcp(Delta));
-	tsr_simd_half RawRejection      = min3(RawFactor[0], RawFactor[1], RawFactor[2]);
-	#if CONFIG_SCENE_COLOR_ALPHA
-		RawRejection = min(RawRejection, RawFactor[3]);
-	#endif
-
-	#if CONFIG_SCENE_COLOR_ALPHA
-		OutRejectionClampBlend = Min3x3(Median3x3(RawRejection));
-		tsr_simd_halfC FilteredClampedEnergy = Max3x3(Median3x3(RawClampedEnergy));
-	#else
-		tsr_simd_halfC FilteredClampedEnergy;
-		Deconcatenate(MaxRGBMinA3x3(Median3x3(Concatenate(RawClampedEnergy, RawRejection))), /* out */ FilteredClampedEnergy, /* out */ OutRejectionClampBlend);
-	#endif
-
-	tsr_simd_halfC FilteredFactor   = saturate(tsr_simd_halfC::Const(1.0) - FilteredClampedEnergy * rcp(Delta));
-	tsr_simd_half FilteredRejection = min3(FilteredFactor[0], FilteredFactor[1], FilteredFactor[2]);
-	#if CONFIG_SCENE_COLOR_ALPHA
-		FilteredRejection = min(FilteredRejection, FilteredFactor[3]);
-	#endif
-	OutRejectionBlendFinal          = FilteredRejection;
-} // MeasureRejection()
-
-
-//------------------------------------------------------- RESURRECTION
-
 #if CONFIG_HISTORY_RESURRECTION
 
 // Fetch the resurrected history.
 void FetchResurrectedHistoryGuide(
-	tsr_simd_short2 PixelPos,
-	out tsr_simd_halfC History,
-	out tsr_simd_half  HistoryUncertainty)
+	tsr_tensor_short2 PixelPos,
+	out tsr_tensor_halfC History,
+	out tsr_tensor_half  HistoryUncertainty)
 {
 	FetchHistoryGuide(
 		ResurrectedHistoryGuideTexture,
@@ -814,67 +227,11 @@ void FetchResurrectedHistoryGuide(
 		/* out */ HistoryUncertainty);
 }
 
-// Returns whether resurrected frame is closer than previous frame on per pixel basis.
-CALL_SITE_DEBUGLOC
-tsr_simd_bool IsResurrectedFrameCloserThanPreviousFrame(tsr_simd_halfC ExposedInput, tsr_simd_halfC ExposedHistory, tsr_simd_halfC ExposedResurrectedHistory)
-{
-	tsr_simd_halfC PrevMatch         = abs(ExposedInput - ExposedHistory);
-	tsr_simd_halfC ResurrectionMatch = abs(ExposedInput - ExposedResurrectedHistory);
-
-	tsr_simd_halfC MatchDiff = PrevMatch - ResurrectionMatch;
-	#if CONFIG_SCENE_COLOR_ALPHA
-		tsr_simd_half TotalMatchDiff = Sum3x3(MatchDiff[0] + MatchDiff[1] + MatchDiff[2] + MatchDiff[3]);
-	#else
-		tsr_simd_half TotalMatchDiff = Sum3x3(MatchDiff[0] + MatchDiff[1] + MatchDiff[2]);
-	#endif
-
-	tsr_simd_half bResurrectedPixelIsCloser = saturate((TotalMatchDiff - tsr_simd_half::Const(tsr_half(0.05 * 3.0 * 9.0))) * tsr_simd_half::Const(tsr_half(POSITIVE_INFINITY)));
-
-	tsr_simd_bool bResurrectionIsCloser = Sum3x3(bResurrectedPixelIsCloser) > tsr_simd_half::Const(4.0);
-
-	//tsr_simd_bool bResurrectionIsCloser = MedianBool3x3((TotalMatchDiff - tsr_simd_half::Const(tsr_half(0.05 * 3.0 * 9.0))));
-
-	return bResurrectionIsCloser;
-}
-
-// Whether it is worth for the entire tile to MeasureRejection() on the resurrected frame.
-CALL_SITE_DEBUGLOC
-bool IsWorthMeasuringRejectionOfResurrectedFrame(tsr_simd_half PrevRejectionBlendFinal, tsr_simd_bool bResurrectionIsCloser)
-{
-	bool bComputeResurrection;
-	BRANCH
-	if (bEnableResurrection)
-	{
-		tsr_simd_bool bPrevIsNotGoodEnough = PrevRejectionBlendFinal < tsr_simd_half::Const(0.5);
-		tsr_simd_bool bIsValidTilePixel = IsValidTilePixel<LANE_STRIDE_X, LANE_STRIDE_Y>(TileOverscan);
-		tsr_simd_bool bShouldComputeResurrection = and(and(bResurrectionIsCloser, bPrevIsNotGoodEnough), bIsValidTilePixel);
-
-		bComputeResurrection = AnyThread(AnyComponent(bShouldComputeResurrection));
-	}
-	else
-	{
-		bComputeResurrection = false;
-	}
-
-	return bComputeResurrection;
-}
-
-// Final logic whether the resurrection must actually happen on per pixel basis.
-CALL_SITE_DEBUGLOC
-tsr_simd_bool ShouldResurrectHistory(tsr_simd_half PrevRejectionBlendFinal, tsr_simd_half ResurrectionRejectionBlendFinal, tsr_simd_bool bResurrectionIsCloser)
-{
-	tsr_simd_bool bIsValidTilePixel = IsValidTilePixel<LANE_STRIDE_X, LANE_STRIDE_Y>(TileOverscan);
-	tsr_simd_bool bResurrectionIsBetter = (ResurrectionRejectionBlendFinal - PrevRejectionBlendFinal) > tsr_simd_half::Const(0.1);
-	tsr_simd_bool bResurrectHistory = and(and(bResurrectionIsCloser, bResurrectionIsBetter), bIsValidTilePixel);
-
-	return bResurrectHistory;
-}
-
 // Load the depth for ComputeResurrectionVelocity()
 CALL_SITE_DEBUGLOC
 void FetchDeviceZ(
-	tsr_simd_short2 PixelPos,
-	out tsr_simd_float DeviceZ)
+	tsr_tensor_short2 PixelPos,
+	out tsr_tensor_float DeviceZ)
 {
 	UNROLL_N(SIMD_SIZE)
 	for (uint ElementIndex = 0; ElementIndex < SIMD_SIZE; ElementIndex++)
@@ -887,9 +244,9 @@ void FetchDeviceZ(
 
 // Compute the velocity used by the reprojection of the resurrection.
 CALL_SITE_DEBUGLOC
-tsr_simd_float2 ComputeResurrectionMotionVector(tsr_simd_float DeviceZ)
+tsr_tensor_float2 ComputeResurrectionMotionVector(tsr_tensor_float DeviceZ)
 {
-	tsr_simd_float2 ResurrectedEncodedScreenVelocity;
+	tsr_tensor_float2 ResurrectedEncodedScreenVelocity;
 	
 	UNROLL_N(SIMD_SIZE)
 	for (uint ElementIndex = 0; ElementIndex < SIMD_SIZE; ElementIndex++)
@@ -913,16 +270,16 @@ tsr_simd_float2 ComputeResurrectionMotionVector(tsr_simd_float DeviceZ)
 
 // Update the guide for next frame.
 CALL_SITE_DEBUGLOC
-tsr_simd_halfC UpdateGuide(tsr_simd_halfC ExposedInput, tsr_simd_halfC ExposedHistory, tsr_simd_bool bIsDisoccluded, tsr_simd_half RejectionBlendFinal)
+tsr_tensor_halfC UpdateGuide(tsr_tensor_halfC ExposedInput, tsr_tensor_halfC ExposedHistory, tsr_tensor_bool bIsDisoccluded, tsr_tensor_half RejectionBlendFinal)
 {
-	tsr_simd_bool bIsOffScreen = tsr_simd_bool::Const(false);
+	tsr_tensor_bool bIsOffScreen = tsr_tensor_bool::Const(false);
 
-	tsr_simd_half TotalBlend = select(or(bIsOffScreen, bIsDisoccluded), tsr_simd_half::Const(1.0), saturate(tsr_simd_half::Const(1.0) - RejectionBlendFinal * tsr_simd_half::Const(4.0)));
-	tsr_simd_half RejectedBlendFactor = max(tsr_simd_half::Const(tsr_half(TheoricBlendFactor)), tsr_simd_half::Const(1.0) - RejectionBlendFinal);
+	tsr_tensor_half TotalBlend = select(or(bIsOffScreen, bIsDisoccluded), tsr_tensor_half::Const(1.0), saturate(tsr_tensor_half::Const(1.0) - RejectionBlendFinal * tsr_tensor_half::Const(4.0)));
+	tsr_tensor_half RejectedBlendFactor = max(tsr_tensor_half::Const(tsr_half(TheoricBlendFactor)), tsr_tensor_half::Const(1.0) - RejectionBlendFinal);
 
-	tsr_simd_half BlendFinal = max(RejectedBlendFactor, TotalBlend);
+	tsr_tensor_half BlendFinal = max(RejectedBlendFactor, TotalBlend);
 
-	tsr_simd_halfC FinalGuide = ExposedInput * tsr_simd_halfC::Vectorize(BlendFinal) - ExposedHistory * tsr_simd_halfC::Vectorize(BlendFinal - 1.0);
+	tsr_tensor_halfC FinalGuide = ExposedInput * tsr_tensor_halfC::Vectorize(BlendFinal) - ExposedHistory * tsr_tensor_halfC::Vectorize(BlendFinal - 1.0);
 	return FinalGuide;
 }
 
@@ -932,25 +289,25 @@ tsr_simd_halfC UpdateGuide(tsr_simd_halfC ExposedInput, tsr_simd_halfC ExposedHi
 #if CONFIG_LDS_DESPILL || CONFIG_CACHE_REFETCH
 
 void RefetchInputAndHistoryFromLDSOrCache(
-	out tsr_simd_halfC Input,
-	out tsr_simd_half  InputUncertainty,
-	out tsr_simd_halfC History,
-	out tsr_simd_half  HistoryUncertainty)
+	out tsr_tensor_halfC Input,
+	out tsr_tensor_half  InputUncertainty,
+	out tsr_tensor_halfC History,
+	out tsr_tensor_half  HistoryUncertainty)
 #if CONFIG_LDS_DESPILL
 {
-	Deconcatenate(tsr_simd_half4::LoadLDSSpill(kDespillInput), /* out */ Input, /* out */ InputUncertainty);
-	Deconcatenate(tsr_simd_half4::LoadLDSSpill(kDespillHistory), /* out */ History, /* out */ HistoryUncertainty);
+	Deconcatenate(tsr_tensor_half4::LoadLDSSpill(kDespillInput), /* out */ Input, /* out */ InputUncertainty);
+	Deconcatenate(tsr_tensor_half4::LoadLDSSpill(kDespillHistory), /* out */ History, /* out */ HistoryUncertainty);
 }
 #else
 {
-	tsr_simd_short2 PixelPos = ComputePixelPos();
+	tsr_tensor_short2 PixelPos = ComputePixelPos();
 
-	tsr_simd_halfC OriginalOpaqueInput;
-	tsr_simd_half4 OriginalTranslucencyInput;
-	tsr_simd_halfC OriginalHistory;
-	tsr_simd_half  OriginalHistoryUncertainty;
-	tsr_simd_half  VelocityEdge;
-	tsr_simd_bool  bIsDisoccluded;
+	tsr_tensor_halfC OriginalOpaqueInput;
+	tsr_tensor_half4 OriginalTranslucencyInput;
+	tsr_tensor_halfC OriginalHistory;
+	tsr_tensor_half  OriginalHistoryUncertainty;
+	tsr_tensor_half  VelocityEdge;
+	tsr_tensor_bool  bIsDisoccluded;
 
 	FetchDecimateMask(PixelPos, /* out */ VelocityEdge, /* out */ bIsDisoccluded);
 	FetchSceneColorAndTranslucency(PixelPos, /* out */ OriginalOpaqueInput, /* out */ OriginalTranslucencyInput);
@@ -959,7 +316,7 @@ void RefetchInputAndHistoryFromLDSOrCache(
 		/* out */ OriginalHistory,
 		/* out */ OriginalHistoryUncertainty);
 
-	tsr_simd_halfC OriginalInput = ComposeTranslucency(OriginalOpaqueInput, Blur3x3(OriginalTranslucencyInput));
+	tsr_tensor_halfC OriginalInput = ComposeTranslucency(OriginalOpaqueInput, Blur3x3(OriginalTranslucencyInput));
 
 	Input            = OriginalInput;
 	InputUncertainty = VelocityEdge;
@@ -991,23 +348,23 @@ void MainCS(
 	}
 		
 	// Analyses the moire pattern overtime.
-	tsr_simd_half  MoireError;
+	tsr_tensor_half  MoireError;
 	#if CONFIG_MOIRE_DETECTION
 	BRANCH
 	if (bEnableFlickeringHeuristic)
 	PLATFORM_SPECIFIC_ISOLATE
 	{
 		// Load its own data
-		tsr_simd_half  VelocityEdge;
-		tsr_simd_bool  bIsDisoccluded;
-		tsr_simd_halfC OriginalOpaqueInput;
-		tsr_simd_half4 OriginalTranslucencyInput;
-		tsr_simd_half  MoireInput;
-		tsr_simd_half4 PrevMoireHistory;
-		tsr_simd_half  bIsStatic;
+		tsr_tensor_half  VelocityEdge;
+		tsr_tensor_bool  bIsDisoccluded;
+		tsr_tensor_halfC OriginalOpaqueInput;
+		tsr_tensor_half4 OriginalTranslucencyInput;
+		tsr_tensor_half  MoireInput;
+		tsr_tensor_half4 PrevMoireHistory;
+		tsr_tensor_half  bIsStatic;
 		PLATFORM_SPECIFIC_ISOLATE
 		{
-			tsr_simd_short2 PixelPos = ComputePixelPos();
+			tsr_tensor_short2 PixelPos = ComputePixelPos();
 		
 			FetchDecimateMask(PixelPos, /* out */ VelocityEdge, /* out */ bIsDisoccluded);
 			FetchSceneColorAndTranslucency(PixelPos, /* out */ OriginalOpaqueInput, /* out */ OriginalTranslucencyInput);
@@ -1018,7 +375,7 @@ void MainCS(
 				/* out */ bIsStatic);
 		}
 	
-		tsr_simd_half4 MoireHistory;
+		tsr_tensor_half4 MoireHistory;
 		ComputeMoireError(
 			bIsDisoccluded,
 			bIsStatic,
@@ -1041,19 +398,19 @@ void MainCS(
 	else
 	#endif // CONFIG_MOIRE_DETECTION
 	{
-		MoireError = tsr_simd_half::Const(0.0);
+		MoireError = tsr_tensor_half::Const(0.0);
 	}
 	
 	// Load all the history data.
-	tsr_simd_half  VelocityEdge;
-	tsr_simd_bool  bIsDisoccluded;
-	tsr_simd_halfC OriginalOpaqueInput;
-	tsr_simd_half4 OriginalTranslucencyInput;
-	tsr_simd_halfC OriginalHistory;
-	tsr_simd_half  OriginalHistoryUncertainty;
+	tsr_tensor_half  VelocityEdge;
+	tsr_tensor_bool  bIsDisoccluded;
+	tsr_tensor_halfC OriginalOpaqueInput;
+	tsr_tensor_half4 OriginalTranslucencyInput;
+	tsr_tensor_halfC OriginalHistory;
+	tsr_tensor_half  OriginalHistoryUncertainty;
 	#if CONFIG_HISTORY_RESURRECTION
-		tsr_simd_halfC OriginalResurrectedHistory;
-		tsr_simd_half  OriginalResurrectedHistoryUncertainty;
+		tsr_tensor_halfC OriginalResurrectedHistory;
+		tsr_tensor_half  OriginalResurrectedHistoryUncertainty;
 	#endif
 	{
 		#if CONFIG_HISTORY_RESURRECTION
@@ -1061,7 +418,7 @@ void MainCS(
 		if (bEnableResurrection)
 		PLATFORM_SPECIFIC_ISOLATE
 		{
-			tsr_simd_short2 PixelPos = ComputePixelPos();
+			tsr_tensor_short2 PixelPos = ComputePixelPos();
 
 			FetchDecimateMask(PixelPos, /* out */ VelocityEdge, /* out */ bIsDisoccluded);
 			FetchSceneColorAndTranslucency(PixelPos, /* out */ OriginalOpaqueInput, /* out */ OriginalTranslucencyInput);
@@ -1078,7 +435,7 @@ void MainCS(
 		#endif // CONFIG_HISTORY_RESURRECTION
 		PLATFORM_SPECIFIC_ISOLATE
 		{
-			tsr_simd_short2 PixelPos = ComputePixelPos();
+			tsr_tensor_short2 PixelPos = ComputePixelPos();
 
 			FetchDecimateMask(PixelPos, /* out */ VelocityEdge, /* out */ bIsDisoccluded);
 			FetchSceneColorAndTranslucency(PixelPos, /* out */ OriginalOpaqueInput, /* out */ OriginalTranslucencyInput);
@@ -1087,18 +444,18 @@ void MainCS(
 				/* out */ OriginalHistory,
 				/* out */ OriginalHistoryUncertainty);
 			#if CONFIG_HISTORY_RESURRECTION
-				OriginalResurrectedHistory            = tsr_simd_halfC::Const(0.0);
-				OriginalResurrectedHistoryUncertainty = tsr_simd_half::Const(0.0);
+				OriginalResurrectedHistory            = tsr_tensor_halfC::Const(0.0);
+				OriginalResurrectedHistoryUncertainty = tsr_tensor_half::Const(0.0);
 			#endif
 		}
 	}
 	
 	// Compute AA Luma and compose separate translucency into scene color for optimised HistoryUpdate.
 	{
-		tsr_simd_halfC CenterFinalSceneColor = ComposeTranslucency(OriginalOpaqueInput, OriginalTranslucencyInput);
-		CenterFinalSceneColor = min(CenterFinalSceneColor, tsr_simd_halfC::Const(tsr_half(Max10BitsFloat)));
+		tsr_tensor_halfC CenterFinalSceneColor = ComposeTranslucency(OriginalOpaqueInput, OriginalTranslucencyInput);
+		CenterFinalSceneColor = min(CenterFinalSceneColor, tsr_tensor_halfC::Const(tsr_half(Max10BitsFloat)));
 		
-		tsr_simd_half CenterLumaLDR = ComputeSpatialAntiAliaserLumaLDR(CenterFinalSceneColor);
+		tsr_tensor_half CenterLumaLDR = ComputeSpatialAntiAliaserLumaLDR(CenterFinalSceneColor);
 
 		PLATFORM_SPECIFIC_ISOLATE
 		{
@@ -1113,7 +470,7 @@ void MainCS(
 		}
 	}
 	
-	tsr_simd_halfC OriginalInput = ComposeTranslucency(OriginalOpaqueInput, Blur3x3(OriginalTranslucencyInput));
+	tsr_tensor_halfC OriginalInput = ComposeTranslucency(OriginalOpaqueInput, Blur3x3(OriginalTranslucencyInput));
 
 	// Despill to LDS to make room in VGPRs for the convolutions of ComputeMoireError()
 	#if CONFIG_LDS_DESPILL
@@ -1125,21 +482,21 @@ void MainCS(
 	#endif // CONFIG_LDS_DESPILL
 
 	// Measure rejection
-	tsr_simd_half RejectionBlendFinal = tsr_simd_half::Const(1.0);
-	tsr_simd_half RejectionClampBlend = tsr_simd_half::Const(1.0);
-	tsr_simd_bool bAliasingIsVisible = tsr_simd_bool::Const(false);
+	tsr_tensor_half RejectionBlendFinal = tsr_tensor_half::Const(1.0);
+	tsr_tensor_half RejectionClampBlend = tsr_tensor_half::Const(1.0);
+	tsr_tensor_bool bAliasingIsVisible = tsr_tensor_bool::Const(false);
 	#if CONFIG_HISTORY_RESURRECTION
-		tsr_simd_bool bResurrectionIsCloser;
+		tsr_tensor_bool bResurrectionIsCloser;
 	#endif
 	PLATFORM_SPECIFIC_ISOLATE
 	{
 		// Reload all the data from LDS or cache with moire to avoid register pressure
-		tsr_simd_halfC Input   = OriginalInput;
-		tsr_simd_halfC History = OriginalHistory;
+		tsr_tensor_halfC Input   = OriginalInput;
+		tsr_tensor_halfC History = OriginalHistory;
 
 		// Convert input to shading LDR
-		tsr_simd_halfC ExposedInput = AddPerceptionAdd(Input, /* DestPerceptionAdd = */ PerceptionAdd);
-		tsr_simd_halfC ExposedHistory = ChangePerceptionAdd(
+		tsr_tensor_halfC ExposedInput = AddPerceptionAdd(Input, /* DestPerceptionAdd = */ PerceptionAdd);
+		tsr_tensor_halfC ExposedHistory = ChangePerceptionAdd(
 			History,
 			/* SourcePerceptionAdd = */ kHistoryGuidePreceptionAdd,
 			/* DestPerceptionAdd   = */ PerceptionAdd);
@@ -1149,7 +506,7 @@ void MainCS(
 		BRANCH
 		if (bEnableResurrection)
 		{
-			tsr_simd_halfC ExposedResurrectedHistory = ChangePerceptionAdd(
+			tsr_tensor_halfC ExposedResurrectedHistory = ChangePerceptionAdd(
 				OriginalResurrectedHistory,
 				/* SourcePerceptionAdd = */ kHistoryGuidePreceptionAdd,
 				/* DestPerceptionAdd   = */ PerceptionAdd);
@@ -1158,19 +515,19 @@ void MainCS(
 		}
 		else
 		{
-			bResurrectionIsCloser = tsr_simd_bool::Const(false);
+			bResurrectionIsCloser = tsr_tensor_bool::Const(false);
 		}
 		#endif
 		
 		// Mutual Annihilation
-		tsr_simd_halfC ExposedInputMin;
-		tsr_simd_halfC ExposedInputMax;
+		tsr_tensor_halfC ExposedInputMin;
+		tsr_tensor_halfC ExposedInputMax;
 		MinMax3x3(ExposedInput, /* out */ ExposedInputMin, /* out */ ExposedInputMax);
 
 		bAliasingIsVisible = IsAliasingVisible(ExposedInputMax - ExposedInputMin);
 
-		tsr_simd_halfC ClampedInput = Clamp3x3(ExposedInput, fastClamp(ExposedHistory, ExposedInputMin, ExposedInputMax));
-		tsr_simd_halfC ClampedHistory = AnnihilateToGuide3x3(ExposedHistory, /* Guide = */ ExposedInput);
+		tsr_tensor_halfC ClampedInput = Clamp3x3(ExposedInput, fastClamp(ExposedHistory, ExposedInputMin, ExposedInputMax));
+		tsr_tensor_halfC ClampedHistory = AnnihilateToGuide3x3(ExposedHistory, /* Guide = */ ExposedInput);
 
 		MeasureRejection</* ApplyMoireError = */ CONFIG_MOIRE_DETECTION != 0>(
 			/* InputC0 = */    ExposedInput,
@@ -1182,7 +539,7 @@ void MainCS(
 	}
 	
 	// Try resurrecting the history if necessary
-	tsr_simd_bool bResurrectHistory = tsr_simd_bool::Const(false);
+	tsr_tensor_bool bResurrectHistory = tsr_tensor_bool::Const(false);
 	#if CONFIG_HISTORY_RESURRECTION
 	bool bComputeResurrection = IsWorthMeasuringRejectionOfResurrectedFrame(RejectionBlendFinal, bResurrectionIsCloser);
 	BRANCH
@@ -1190,12 +547,12 @@ void MainCS(
 	PLATFORM_SPECIFIC_ISOLATE
 	{
 		// Reload the resurrected history from cache.
-		tsr_simd_halfC Input;
-		tsr_simd_halfC ResurrectedHistory;
-		tsr_simd_half  ResurrectedHistoryUncertainty;
+		tsr_tensor_halfC Input;
+		tsr_tensor_halfC ResurrectedHistory;
+		tsr_tensor_half  ResurrectedHistoryUncertainty;
 		#if CONFIG_CACHE_REFETCH || CONFIG_LDS_DESPILL
 		{
-			tsr_simd_short2 PixelPos = ComputePixelPos();
+			tsr_tensor_short2 PixelPos = ComputePixelPos();
 		
 			FetchResurrectedHistoryGuide(
 				PixelPos,
@@ -1203,10 +560,10 @@ void MainCS(
 				/* out */ ResurrectedHistoryUncertainty);
 				
 			#if CONFIG_LDS_DESPILL
-				Input = ResizeChannels<3>(tsr_simd_half4::LoadLDSSpill(kDespillInput));
+				Input = ResizeChannels<3>(tsr_tensor_half4::LoadLDSSpill(kDespillInput));
 			#else
-				tsr_simd_halfC OriginalOpaqueInput;
-				tsr_simd_half4 OriginalTranslucencyInput;
+				tsr_tensor_halfC OriginalOpaqueInput;
+				tsr_tensor_half4 OriginalTranslucencyInput;
 				FetchSceneColorAndTranslucency(PixelPos, /* out */ OriginalOpaqueInput, /* out */ OriginalTranslucencyInput);
 
 				Input = ComposeTranslucency(OriginalOpaqueInput, Blur3x3(OriginalTranslucencyInput));
@@ -1218,23 +575,23 @@ void MainCS(
 		#endif
 		
 		// Convert input to shading LDR
-		tsr_simd_halfC ExposedInput = AddPerceptionAdd(Input, /* DestPerceptionAdd = */ PerceptionAdd);
-		tsr_simd_halfC ExposedResurrectedHistory = ChangePerceptionAdd(
+		tsr_tensor_halfC ExposedInput = AddPerceptionAdd(Input, /* DestPerceptionAdd = */ PerceptionAdd);
+		tsr_tensor_halfC ExposedResurrectedHistory = ChangePerceptionAdd(
 			ResurrectedHistory,
 			/* SourcePerceptionAdd = */ kHistoryGuidePreceptionAdd,
 			/* DestPerceptionAdd   = */ PerceptionAdd);
 			
 		// Mutual Annihilation
-		tsr_simd_halfC ClampedResurrectedInput, ClampedResurrectedHistory;
+		tsr_tensor_halfC ClampedResurrectedInput, ClampedResurrectedHistory;
 		AnnihilateMutually3x3(ExposedInput, ExposedResurrectedHistory, /* out */ ClampedResurrectedInput, /* out */ ClampedResurrectedHistory);
 
-		tsr_simd_half ResurrectionRejectionBlendFinal = tsr_simd_half::Const(1.0);
-		tsr_simd_half ResurrectionRejectionClampBlend = tsr_simd_half::Const(1.0);
+		tsr_tensor_half ResurrectionRejectionBlendFinal = tsr_tensor_half::Const(1.0);
+		tsr_tensor_half ResurrectionRejectionClampBlend = tsr_tensor_half::Const(1.0);
 		MeasureRejection</* ApplyMoireError = */ false>(
 			/* InputC0 = */    ExposedInput,
 			/* InputC2 = */    ClampedResurrectedInput,
 			/* HistoryC2 = */  ClampedResurrectedHistory,
-			/* MoireError = */ tsr_simd_half::Const(0.0),
+			/* MoireError = */ tsr_tensor_half::Const(0.0),
 			/* out */ ResurrectionRejectionBlendFinal,
 			/* out */ ResurrectionRejectionClampBlend);
 			
@@ -1246,14 +603,14 @@ void MainCS(
 	#endif // CONFIG_HISTORY_RESURRECTION
 	
 	// Update the guide for next frame.
-	tsr_simd_halfC FinalGuide;
-	tsr_simd_half  InputUncertainty;
-	tsr_simd_half  HistoryUncertainty;
+	tsr_tensor_halfC FinalGuide;
+	tsr_tensor_half  InputUncertainty;
+	tsr_tensor_half  HistoryUncertainty;
 	PLATFORM_SPECIFIC_ISOLATE
 	{
 		// Reload all the data from LDS or cache
-		tsr_simd_halfC Input;
-		tsr_simd_halfC History;
+		tsr_tensor_halfC Input;
+		tsr_tensor_halfC History;
 		#if CONFIG_LDS_DESPILL || CONFIG_CACHE_REFETCH
 			RefetchInputAndHistoryFromLDSOrCache(
 				/* out */ Input,
@@ -1274,12 +631,12 @@ void MainCS(
 		PLATFORM_SPECIFIC_ISOLATE
 		{
 			// Reload the resurrected history from cache.
-			tsr_simd_halfC ResurrectedHistory;
-			tsr_simd_half  ResurrectedHistoryUncertainty;
-			tsr_simd_float DeviceZ;
+			tsr_tensor_halfC ResurrectedHistory;
+			tsr_tensor_half  ResurrectedHistoryUncertainty;
+			tsr_tensor_float DeviceZ;
 			PLATFORM_SPECIFIC_ISOLATE
 			{
-				tsr_simd_short2 PixePos = ComputePixelPos();
+				tsr_tensor_short2 PixePos = ComputePixelPos();
 				
 				#if CONFIG_CACHE_REFETCH
 					FetchResurrectedHistoryGuide(
@@ -1295,13 +652,13 @@ void MainCS(
 			}
 
 			// Apply the resurrection on the guide immediatly to free up registers.
-			History = select(tsr_simd_boolC::Vectorize(bResurrectHistory), ResurrectedHistory, History);
+			History = select(tsr_tensor_boolC::Vectorize(bResurrectHistory), ResurrectedHistory, History);
 			HistoryUncertainty = select(bResurrectHistory, ResurrectedHistoryUncertainty, HistoryUncertainty);
 			
 			PLATFORM_SPECIFIC_ISOLATE
 			{
 				// Compute the motion vector for the resurrection
-				tsr_simd_float2 ResurrectionMotionVector = ComputeResurrectionMotionVector(DeviceZ);
+				tsr_tensor_float2 ResurrectionMotionVector = ComputeResurrectionMotionVector(DeviceZ);
 			
 				// Amends the motion vector for resurrection areas.
 				UNROLL_N(SIMD_SIZE)
@@ -1317,22 +674,22 @@ void MainCS(
 		#endif // CONFIG_HISTORY_RESURRECTION
 
 		// Convert input to LDR
-		tsr_simd_halfC ExposedInput = AddPerceptionAdd(Input, /* DestPerceptionAdd = */ kHistoryGuidePreceptionAdd);
+		tsr_tensor_halfC ExposedInput = AddPerceptionAdd(Input, /* DestPerceptionAdd = */ kHistoryGuidePreceptionAdd);
 	
 		FinalGuide = UpdateGuide(ExposedInput, History, bIsDisoccluded, RejectionBlendFinal);
 	}
 
 	// Controls how the history and its validity is clamped.
-	tsr_simd_half  DisableHistoryClamp;
-	tsr_simd_half  IncreaseValidityMultiplier;
+	tsr_tensor_half  DisableHistoryClamp;
+	tsr_tensor_half  IncreaseValidityMultiplier;
 	PLATFORM_SPECIFIC_ISOLATE
 	{
-		DisableHistoryClamp        = select(and(bIsDisoccluded, not(bResurrectHistory)), tsr_simd_half::Const(0.0), min(RejectionClampBlend, min(InputUncertainty, Min3x3(HistoryUncertainty))));
-		IncreaseValidityMultiplier = select(and(bIsDisoccluded, not(bResurrectHistory)), tsr_simd_half::Const(0.0), RejectionClampBlend);
+		DisableHistoryClamp        = select(and(bIsDisoccluded, not(bResurrectHistory)), tsr_tensor_half::Const(0.0), min(RejectionClampBlend, min(InputUncertainty, Min3x3(HistoryUncertainty))));
+		IncreaseValidityMultiplier = select(and(bIsDisoccluded, not(bResurrectHistory)), tsr_tensor_half::Const(0.0), RejectionClampBlend);
 	}
 	
 	// Whether to run spatial anti-aliaser on the frame. Must match SpatialAntiAliasingLerp.
-	tsr_simd_bool bSpatialAntiAlias = ShouldSpatialAntiAlias(bIsDisoccluded, bResurrectHistory, bAliasingIsVisible, RejectionBlendFinal);
+	tsr_tensor_bool bSpatialAntiAlias = ShouldSpatialAntiAlias(bIsDisoccluded, bResurrectHistory, bAliasingIsVisible, RejectionBlendFinal);
 
 	// Output data.
 	PLATFORM_SPECIFIC_ISOLATE
diff --git a/Engine/Shaders/Private/TemporalSuperResolution/TSRShadingAnalysis.ush b/Engine/Shaders/Private/TemporalSuperResolution/TSRShadingAnalysis.ush
new file mode 100644
--- /dev/null
+++ b/Engine/Shaders/Private/TemporalSuperResolution/TSRShadingAnalysis.ush
@@ -0,0 +1,525 @@
+// Copyright Epic Games, Inc. All Rights Reserved.
+
+#pragma once
+
+#include "TSRConvolutionNetwork.ush"
+
+
+//------------------------------------------------------- CONSTANTS
+
+static const float kHistoryGuidePreceptionAdd = 1.0;
+
+static const tsr_half FilteringWeight = rcp(tsr_half(1.0 + 4 * 0.5 + 4 * 0.25));
+
+
+//------------------------------------------------------- COMPOSE SEPARATE TRANSLUCENCY
+
+/** Compose translucency in linear color space. */
+CALL_SITE_DEBUGLOC
+tsr_tensor_halfC ComposeTranslucency(tsr_tensor_halfC Color, tsr_tensor_half4 Translucency)
+{
+	tsr_tensor_halfC ComposedColor;
+	ComposedColor = Color * tsr_tensor_halfC::Vectorize(Translucency[3]) + ResizeChannels<CONFIG_CHANNEL_COUNT>(Translucency);
+	#if CONFIG_SCENE_COLOR_ALPHA
+		ComposedColor.SetComponent(3, Color[3] * Translucency[3]);
+	#endif
+	#if CONFIG_SCENE_COLOR_ALPHA
+		ComposedColor = min(ComposedColor, tsr_tensor_halfC::Const(LargestSceneColorRGBA));
+	#else
+		ComposedColor = min(ComposedColor, tsr_tensor_halfC::Const(LargestSceneColorRGB));
+	#endif
+	return ComposedColor;
+}
+
+
+//------------------------------------------------------- CHANGE LDR EXPOSURE
+
+CALL_SITE_DEBUGLOC
+tsr_tensor_halfC RemovePerceptionAdd(tsr_tensor_halfC SrcColor, const float SourcePerceptionAdd)
+{
+	tsr_tensor_halfC Minus = tsr_tensor_halfC::Const(tsr_half(1.0)) - SrcColor;
+	tsr_tensor_halfC LDRToLinear = min(tsr_tensor_halfC::Const(tsr_half(SourcePerceptionAdd)) * rcp(Minus), tsr_tensor_halfC::Const(tsr_half(MaxHalfFloat)));
+
+	return SrcColor * LDRToLinear;
+}
+
+CALL_SITE_DEBUGLOC
+tsr_tensor_halfC AddPerceptionAdd(tsr_tensor_halfC SrcColor, const float DestPerceptionAdd)
+{
+	SrcColor = SrcColor * rcp(SrcColor + tsr_tensor_halfC::Const(tsr_half(DestPerceptionAdd)));
+	return SrcColor;
+}
+
+CALL_SITE_DEBUGLOC
+tsr_tensor_halfC ChangePerceptionAdd(tsr_tensor_halfC SrcColor, const float SourcePerceptionAdd, const float DestPerceptionAdd)
+{
+	SrcColor = RemovePerceptionAdd(SrcColor, SourcePerceptionAdd);
+	SrcColor = AddPerceptionAdd(SrcColor, DestPerceptionAdd);
+	return SrcColor;
+}
+
+
+//------------------------------------------------------- SPATIAL ANTI-ALIASER
+
+// Compute the LDR luminance the spatial anti-aliaser should anti-alias.
+CALL_SITE_DEBUGLOC
+tsr_tensor_half ComputeSpatialAntiAliaserLumaLDR(tsr_tensor_halfC SceneColor)
+{
+	const tsr_half SpatialAAExposure = tsr_half(0.5);
+		
+	tsr_tensor_half AALumaLDR;
+
+	UNROLL_N(SIMD_SIZE)
+	for (uint ElementIndex = 0; ElementIndex < SIMD_SIZE; ElementIndex++)
+	{
+		tsr_half PixelLuma = dot(SceneColor.GetElement(ElementIndex).rgb, tsr_half3(0.299f, 0.587f, 0.114f));
+		AALumaLDR.SetElement(ElementIndex, PixelLuma / (SpatialAAExposure + PixelLuma));
+	}
+
+	return AALumaLDR;
+}
+
+// Returns whether the aliasing is visible at all, to avoid paying spatial-antialiaser on low contrasts areas.
+CALL_SITE_DEBUGLOC
+tsr_tensor_bool IsAliasingVisible(tsr_tensor_halfC ExposedInputBoxSize)
+{
+	#if CONFIG_SCENE_COLOR_ALPHA
+		tsr_tensor_half ExposedInputLuminanceBox = dot(ExposedInputBoxSize, tsr_tensor_halfC::Const(tsr_halfC(0.299f, 0.587f, 0.114f, 0.0)));
+	#else
+		tsr_tensor_half ExposedInputLuminanceBox = dot(ExposedInputBoxSize, tsr_tensor_halfC::Const(tsr_halfC(0.299f, 0.587f, 0.114f)));
+	#endif
+
+	tsr_tensor_bool bAliasingIsVisible = ExposedInputLuminanceBox > tsr_tensor_half::Const(SPATIAL_ANTI_ALIASER_MIN_LUMIMANCE);
+
+	return bAliasingIsVisible;
+}
+
+// Whether to run spatial anti-aliaser on the frame. Must match SpatialAntiAliasingLerp.
+CALL_SITE_DEBUGLOC
+tsr_tensor_bool ShouldSpatialAntiAlias(tsr_tensor_bool bIsDisoccluded, tsr_tensor_bool bResurrectHistory, tsr_tensor_bool bAliasingIsVisible, tsr_tensor_half RejectionBlendFinal)
+{
+	tsr_tensor_bool bSpatialAntiAlias = and(bAliasingIsVisible, or(RejectionBlendFinal < tsr_tensor_half::Const(tsr_half(0.25)), and(bIsDisoccluded, not(bResurrectHistory))));
+	
+	return bSpatialAntiAlias;
+}
+
+
+//------------------------------------------------------- RESURRECTION
+
+// Returns whether resurrected frame is closer than previous frame on per pixel basis.
+CALL_SITE_DEBUGLOC
+tsr_tensor_bool IsResurrectedFrameCloserThanPreviousFrame(tsr_tensor_halfC ExposedInput, tsr_tensor_halfC ExposedHistory, tsr_tensor_halfC ExposedResurrectedHistory)
+{
+	tsr_tensor_halfC PrevMatch         = abs(ExposedInput - ExposedHistory);
+	tsr_tensor_halfC ResurrectionMatch = abs(ExposedInput - ExposedResurrectedHistory);
+
+	tsr_tensor_halfC MatchDiff = PrevMatch - ResurrectionMatch;
+	#if CONFIG_SCENE_COLOR_ALPHA
+		tsr_tensor_half TotalMatchDiff = Sum3x3(MatchDiff[0] + MatchDiff[1] + MatchDiff[2] + MatchDiff[3]);
+	#else
+		tsr_tensor_half TotalMatchDiff = Sum3x3(MatchDiff[0] + MatchDiff[1] + MatchDiff[2]);
+	#endif
+
+	tsr_tensor_half bResurrectedPixelIsCloser = saturate((TotalMatchDiff - tsr_tensor_half::Const(tsr_half(0.05 * 3.0 * 9.0))) * tsr_tensor_half::Const(tsr_half(POSITIVE_INFINITY)));
+
+	tsr_tensor_bool bResurrectionIsCloser = Sum3x3(bResurrectedPixelIsCloser) > tsr_tensor_half::Const(4.0);
+
+	//tsr_tensor_bool bResurrectionIsCloser = MedianBool3x3((TotalMatchDiff - tsr_tensor_half::Const(tsr_half(0.05 * 3.0 * 9.0))));
+
+	return bResurrectionIsCloser;
+}
+
+// Whether it is worth for the entire tile to MeasureRejection() on the resurrected frame.
+CALL_SITE_DEBUGLOC
+bool IsWorthMeasuringRejectionOfResurrectedFrame(tsr_tensor_half PrevRejectionBlendFinal, tsr_tensor_bool bResurrectionIsCloser)
+{
+	bool bComputeResurrection;
+	BRANCH
+	if (bEnableResurrection)
+	{
+		tsr_tensor_bool bPrevIsNotGoodEnough = PrevRejectionBlendFinal < tsr_tensor_half::Const(0.5);
+		tsr_tensor_bool bIsValidTilePixel = IsValidTilePixel<LANE_STRIDE_X, LANE_STRIDE_Y>(TileOverscan);
+		tsr_tensor_bool bShouldComputeResurrection = and(and(bResurrectionIsCloser, bPrevIsNotGoodEnough), bIsValidTilePixel);
+
+		bComputeResurrection = AnyThread(AnyComponent(bShouldComputeResurrection));
+	}
+	else
+	{
+		bComputeResurrection = false;
+	}
+
+	return bComputeResurrection;
+}
+
+// Final logic whether the resurrection must actually happen on per pixel basis.
+CALL_SITE_DEBUGLOC
+tsr_tensor_bool ShouldResurrectHistory(tsr_tensor_half PrevRejectionBlendFinal, tsr_tensor_half ResurrectionRejectionBlendFinal, tsr_tensor_bool bResurrectionIsCloser)
+{
+	tsr_tensor_bool bIsValidTilePixel = IsValidTilePixel<LANE_STRIDE_X, LANE_STRIDE_Y>(TileOverscan);
+	tsr_tensor_bool bResurrectionIsBetter = (ResurrectionRejectionBlendFinal - PrevRejectionBlendFinal) > tsr_tensor_half::Const(0.1);
+	tsr_tensor_bool bResurrectHistory = and(and(bResurrectionIsCloser, bResurrectionIsBetter), bIsValidTilePixel);
+
+	return bResurrectHistory;
+}
+
+
+//------------------------------------------------------- FLICKERING TEMPORAL ANALYSIS
+
+void ComputeMoireError(
+	tsr_tensor_bool  bIsDisoccluded,
+	tsr_tensor_half  bIsStatic,
+	tsr_tensor_half  MoireInput,
+	tsr_tensor_halfC OriginalOpaqueInput,
+	tsr_tensor_half4 OriginalTranslucencyInput,
+	tsr_tensor_half4 PrevMoireHistory,
+	out tsr_tensor_half4 OutMoireHistory,
+	out tsr_tensor_half  OutMoireError)
+{
+	const tsr_half MinBlendFinal = tsr_half(0.05);
+	const tsr_half MoireEncodingError = tsr_half(rcp(127.0));
+	const tsr_half MaxPrevTotalVariationCount = tsr_half(20.0);
+
+	// Compute the threshold at which the variation count should quick in
+	const tsr_half TotalVariationCountThreshold = tsr_half(1.0 / (1.0 - pow(1.0 - MinBlendFinal, tsr_half(FlickeringFramePeriod))));
+
+	// Input luminance to monitor flickering from.
+	tsr_tensor_half InputLuma = MoireInput;
+
+	// How much input luminance from the TSRComputeMoireLuma.usf measurement has been modified since for instance with fog, water.
+	tsr_tensor_half InputLumaModification;
+	{
+		const tsr_half InputLumaEncodingError = tsr_half(0.5 / 255.0);
+			
+		tsr_tensor_halfC OriginalInput = ComposeTranslucency(OriginalOpaqueInput, OriginalTranslucencyInput);
+
+		OriginalInput       = AddPerceptionAdd(OriginalInput, PerceptionAdd);
+		OriginalOpaqueInput = AddPerceptionAdd(OriginalOpaqueInput, PerceptionAdd);
+
+		tsr_tensor_half SceneColorLuma;
+		tsr_tensor_half TranslucencyChange;
+		UNROLL_N(SIMD_SIZE)
+		for (uint ElementIndex = 0; ElementIndex < SIMD_SIZE; ElementIndex++)
+		{
+			SceneColorLuma.SetElement(ElementIndex, dot(OriginalOpaqueInput.GetElement(ElementIndex).rgb, kMoireLumaWeights));
+			TranslucencyChange.SetElement(ElementIndex, dot(abs(OriginalOpaqueInput.GetElement(ElementIndex).rgb - OriginalInput.GetElement(ElementIndex).rgb), kMoireLumaWeights));
+		}
+
+		InputLumaModification = tsr_tensor_half::Const(0.0);
+		InputLumaModification = max(InputLumaModification, abs(InputLuma - SceneColorLuma) - tsr_tensor_half::Const(InputLumaEncodingError));
+		InputLumaModification = max(InputLumaModification, TranslucencyChange);
+		InputLumaModification = max(InputLumaModification, tsr_tensor_half::Const(1.0) - OriginalTranslucencyInput[3]);
+	}
+
+	// Unpack history
+	tsr_tensor_half PrevFlickeringHistory = PrevMoireHistory[0];
+	tsr_tensor_half PrevGradient = PrevMoireHistory[1] * tsr_half(255.0 / 127.0) - tsr_half(1.0);
+	//tsr_tensor_half PrevTotalVariation = PrevMoireHistory[2] * tsr_tensor_half::Const(1.0 - MinBlendFinal);
+	tsr_tensor_half PrevTotalVariation = PrevMoireHistory[2] * tsr_tensor_half::Const(tsr_half(1.0) - MinBlendFinal);
+	tsr_tensor_half PrevTotalVariationCount = PrevMoireHistory[3] * MaxPrevTotalVariationCount * tsr_tensor_half::Const(tsr_half(1.0) - MinBlendFinal);
+	
+	// Discard history on moving objects.
+	PrevTotalVariation = PrevTotalVariation * bIsStatic;
+	PrevTotalVariationCount = PrevTotalVariationCount * bIsStatic;
+
+	// Run the exact same heuristic as MeasureRejection on a single channel.
+	tsr_tensor_half FilteredInputLumaModification;
+	tsr_tensor_half OutRejectionBlendFinal;
+	tsr_tensor_half OutRejectionClampBlend;
+	{
+		tsr_tensor_half BackbufferQuantizationErrorVector = tsr_tensor_half::Const(MoireEncodingError);
+		
+		tsr_tensor_half InputC0 = InputLuma;
+		tsr_tensor_half HistoryC0 = PrevFlickeringHistory;
+
+		// Anhilate history
+		tsr_tensor_half InputC2 = InputC0;
+		tsr_tensor_half HistoryC2 = HistoryC0;
+		AnnihilateMutuallySingleChannel3x3(InputLuma, HistoryC0, /* out */ InputC2, /* out */ HistoryC2);
+
+		// Convolve the input luma to have the exact same moire pattern as FilteredInput below.
+		tsr_tensor_half FilteredInput;
+		tsr_tensor_half FilteredHistory;
+		Deconcatenate(Blur3x3(Concatenate(InputC2, InputLumaModification, HistoryC2)), /* out */ FilteredInput, /* out */ FilteredInputLumaModification, /* out */ FilteredHistory);
+		
+		// Compute clamping box of the limance.
+		tsr_tensor_half FilteredBoxMin;
+		tsr_tensor_half FilteredBoxMax;
+		tsr_tensor_half InputC2BoxSize;
+		{
+			tsr_tensor_half2 Min;
+			tsr_tensor_half2 Max;
+			MinMax3x3(Concatenate(FilteredInput, InputC2), /* out */ Min, /* out */ Max);
+			FilteredBoxMin = Min[0];
+			FilteredBoxMax = Max[0];
+		
+			InputC2BoxSize = Max[1] - Min[1];
+		}
+
+		// Compute the clamp error
+		tsr_tensor_half ClampError;
+		{
+			tsr_tensor_half TotalVarInputDiffC0C2;
+			tsr_tensor_half TotalVarInputC2;
+			Deconcatenate(abs(TotalVariation3x3(Concatenate(abs(InputC0 - InputC2), InputC2))), /* out */ TotalVarInputDiffC0C2, /* out */ TotalVarInputC2);
+		
+			ClampError = BackbufferQuantizationErrorVector;
+			ClampError = max(ClampError, Blur3x3(min(TotalVarInputDiffC0C2, TotalVarInputC2)));
+			ClampError = max(ClampError, InputC2BoxSize * tsr_half(FilteringWeight * 0.25));
+			ClampError = ClampError + BackbufferQuantizationErrorVector;
+		}
+
+		// Apply clamp error
+		FilteredBoxMin = FilteredBoxMin - ClampError;
+		FilteredBoxMax = FilteredBoxMax + ClampError;
+
+		tsr_tensor_half BoxSize = (
+			InputC2BoxSize * tsr_tensor_half::Const(FilteringWeight) +
+			(BackbufferQuantizationErrorVector * tsr_half(2 * FilteringWeight)));
+
+		tsr_tensor_half ClampedFilteredHistory = fastClamp(FilteredHistory, FilteredBoxMin, FilteredBoxMax);
+
+		tsr_tensor_half Delta = max(abs(FilteredInput - FilteredHistory), BoxSize);
+	
+		tsr_tensor_half RawClampedEnergy = abs(ClampedFilteredHistory - FilteredHistory);
+		tsr_tensor_half RawRejection     = saturate(tsr_tensor_half::Const(1.0) - max(RawClampedEnergy, tsr_tensor_half::Const(0.0)) * rcp(Delta));
+
+		tsr_tensor_half FilteredClampedEnergy;
+		Deconcatenate(MaxRMinG3x3(Median3x3(Concatenate(RawClampedEnergy, RawRejection))), /* out */ FilteredClampedEnergy, /* out */ OutRejectionClampBlend);
+
+		OutRejectionBlendFinal = saturate(tsr_tensor_half::Const(1.0) - max(FilteredClampedEnergy, tsr_tensor_half::Const(0.0)) * rcp(Delta));
+	}
+	
+	// Simulate what happens in TSRUpdateHistory.usf
+	tsr_tensor_half FinalFlickeringHistory;
+	tsr_tensor_half CurrentGradient;
+	{
+		tsr_tensor_half ClampedPrevFlickeringHistory = ClampPlus3x3(PrevFlickeringHistory, InputLuma);
+		tsr_tensor_half FinalPrevFlickeringHistory = lerp(ClampedPrevFlickeringHistory, PrevFlickeringHistory, OutRejectionClampBlend);
+
+		tsr_tensor_half FinalBlendFinal = max(tsr_tensor_half::Const(1.0) - OutRejectionBlendFinal, tsr_tensor_half::Const(MinBlendFinal));
+		
+		FinalFlickeringHistory = InputLuma * FinalBlendFinal - FinalPrevFlickeringHistory * (FinalBlendFinal - tsr_tensor_half::Const(1.0));
+		
+		tsr_tensor_half GhostingHistory = InputLuma * tsr_tensor_half::Const(MinBlendFinal) - PrevFlickeringHistory * tsr_tensor_half::Const(MinBlendFinal - (1.0));
+
+		CurrentGradient = FinalFlickeringHistory - GhostingHistory;
+	}
+
+	// Remove how much luma has changed from luma measurement to final scene color from the gradient.
+	for (uint ElementIndex = 0; ElementIndex < SIMD_SIZE; ElementIndex++)
+	{
+		tsr_half LocalCurrentGradient = CurrentGradient.GetElement(ElementIndex);
+		tsr_half LocalInputLumaModification = FilteredInputLumaModification.GetElement(ElementIndex);
+			
+		CurrentGradient.SetElement(ElementIndex, clamp(LocalCurrentGradient * tsr_half(POSITIVE_INFINITY), tsr_half(-1.0), tsr_half(1.0)) * max(abs(LocalCurrentGradient) - LocalInputLumaModification, tsr_half(0.0)));
+	}
+
+	#if 0
+		tsr_tensor_half InputBoxSize = InputLumaMax - InputLumaMin;
+		CurrentGradient = sign(CurrentGradient) * max(abs(CurrentGradient) - tsr_tensor_half::Const(0.5) * InputBoxSize, tsr_tensor_half::Const(0.0));
+	#endif
+
+	tsr_tensor_half IsFlicker;
+	UNROLL_N(SIMD_SIZE)
+	for (uint ElementIndex = 0; ElementIndex < SIMD_SIZE; ElementIndex++)
+	{
+		tsr_half LocalCurrentGradient = CurrentGradient.GetElement(ElementIndex);
+		tsr_half LocalPrevGradient = PrevGradient.GetElement(ElementIndex);
+
+		bool bGradientsAreSameSign = LocalCurrentGradient * LocalPrevGradient > tsr_half(0.0);
+		bool bGradientsAreWithinError = abs(LocalCurrentGradient) < MoireEncodingError || abs(LocalPrevGradient) < MoireEncodingError;
+		//bool bCurrentGradientSubstentialEnough = abs(LocalCurrentGradient) > View.GeneralPurposeTweak * abs(LocalPrevGradient); // - 2.0 * MoireEncodingError;
+
+		tsr_half LocalIsFlicker = select((
+			bGradientsAreSameSign ||
+			bGradientsAreWithinError ||
+			//!bCurrentGradientSubstentialEnough ||
+			bCameraCut), tsr_half(0.0), tsr_half(1.0));
+
+		IsFlicker.SetElement(ElementIndex, LocalIsFlicker);
+	}
+
+	// Compute the total variation of the gradient over time.
+	tsr_tensor_half GradientVariation = min(abs(PrevGradient), abs(CurrentGradient)) * IsFlicker;
+
+	// Dilates the total variation and contribution to keep stability despite inacuracies in the history reprojection
+	// of these need a nearest filter which isn't precise at all.
+	tsr_tensor_half TotalVariationContribution;
+	tsr_tensor_half TotalVariationCountContribution;
+	Deconcatenate(Max3x3(Max3x3(Concatenate(GradientVariation, IsFlicker))), /* out */ TotalVariationContribution, /* out */ TotalVariationCountContribution);
+
+	UNROLL_N(SIMD_SIZE)
+	for (uint ElementIndex = 0; ElementIndex < SIMD_SIZE; ElementIndex++)
+	{
+		tsr_half LocalCurrentGradient = CurrentGradient.GetElement(ElementIndex);
+
+		tsr_half LocalPrevGradient = PrevGradient.GetElement(ElementIndex);
+		tsr_half LocalPrevTotalVariation = PrevTotalVariation.GetElement(ElementIndex);
+		tsr_half LocalPrevTotalVariationCount = PrevTotalVariationCount.GetElement(ElementIndex);
+		bool LocalIsDisoccluded = bIsDisoccluded.GetElement(ElementIndex);
+
+		tsr_half LocalIsFlicker = IsFlicker.GetElement(ElementIndex);
+
+		tsr_half LocalTotalVariationContribution = TotalVariationContribution.GetElement(ElementIndex);
+		tsr_half LocalTotalVariationCountContribution = TotalVariationCountContribution.GetElement(ElementIndex);
+			
+		// Blend out previous gradient
+		tsr_half LocalPrevBlendedGradient = LocalPrevGradient * (tsr_half(1.0) - MinBlendFinal) * (tsr_half(1.0) - LocalIsFlicker);
+
+		// Compute new gradient.
+		#if 1
+			tsr_half LocalNewGradient = LocalPrevBlendedGradient + LocalCurrentGradient;
+		#else
+			tsr_half LocalNewGradient = abs(LocalPrevBlendedGradient) > abs(LocalCurrentGradient) ? LocalPrevBlendedGradient : LocalCurrentGradient;
+		#endif
+			
+		// Accumulate current frame variation.
+		tsr_half LocalNewTotalVariation = LocalPrevTotalVariation + LocalTotalVariationContribution;
+		tsr_half LocalNewTotalVariationCount = LocalPrevTotalVariationCount + LocalTotalVariationCountContribution;
+
+		// Discard all moire history on parallax disocclusion.
+		LocalNewGradient            = select(LocalIsDisoccluded, tsr_half(0.0), LocalNewGradient);
+		LocalNewTotalVariation      = select(LocalIsDisoccluded, tsr_half(0.0), LocalNewTotalVariation);
+		LocalNewTotalVariationCount = select(LocalIsDisoccluded, tsr_half(0.0), LocalNewTotalVariationCount);
+
+		// Quantise total variation and count to history bit depth to variation drift in the history over time.
+		tsr_half LocalQuantizedNewTotalVariationCount = floor(LocalNewTotalVariationCount * (tsr_half(255.0) / MaxPrevTotalVariationCount)) * (MaxPrevTotalVariationCount / tsr_half(255.0));
+		LocalNewTotalVariation *= LocalQuantizedNewTotalVariationCount * SafeRcp(LocalNewTotalVariationCount);
+
+		// Compute the final luminance 
+		#if 1
+			tsr_half LocalCountFadeIn = saturate(LocalNewTotalVariationCount / TotalVariationCountThreshold - tsr_half(0.5));
+		#else
+			tsr_half LocalCountFadeIn = saturate(LocalNewTotalVariationCount - TotalVariationCountThreshold);
+		#endif
+		tsr_half LocalMoireError = (abs(LocalNewTotalVariation * SafeRcp(LocalNewTotalVariationCount)) + LocalNewTotalVariationCount * MoireEncodingError) * LocalCountFadeIn;
+
+		tsr_half4 PackedMoireHistory = tsr_half4(
+			FinalFlickeringHistory.GetElement(ElementIndex),
+			LocalNewGradient * tsr_half(127.0 / 255.0) + tsr_half(127.0 / 255.0),
+			LocalNewTotalVariation,
+			LocalQuantizedNewTotalVariationCount * rcp(MaxPrevTotalVariationCount));
+
+		OutMoireHistory.SetElement(ElementIndex, PackedMoireHistory);
+
+		OutMoireError.SetElement(ElementIndex, LocalMoireError);
+	}
+		
+	// Moving object in the scene have very low chance of creating moire pattern with TAA jitter, so discard the moire error to still
+	// have character animation ghost free evem as environement shadow project different things on character.
+	OutMoireError = saturate(OutMoireError * bIsStatic - InputLumaModification);
+} // ComputeMoireError()
+
+
+//------------------------------------------------------- SHADING REJECTION
+
+template<bool ApplyMoireError>
+void MeasureRejection(
+	tsr_tensor_halfC InputC0,
+	tsr_tensor_halfC InputC2,
+	tsr_tensor_halfC HistoryC2,
+	tsr_tensor_half  MoireError,
+	out tsr_tensor_half OutRejectionBlendFinal,
+	out tsr_tensor_half OutRejectionClampBlend)
+{
+	tsr_tensor_halfC BackbufferQuantizationErrorVector = tsr_tensor_halfC::Const(MeasureBackbufferLDRQuantizationError());
+		
+	tsr_tensor_halfC FilteredHistory = Blur3x3(HistoryC2);
+	
+	tsr_tensor_halfC TotalVarInputDiffC0C2 = abs(TotalVariation3x3(abs(InputC0 - InputC2)));
+	tsr_tensor_halfC TotalVarInputC2 = abs(TotalVariation3x3(InputC2));
+
+	tsr_tensor_halfC ClampError = BackbufferQuantizationErrorVector;
+	ClampError = max(ClampError, Blur3x3(min(TotalVarInputDiffC0C2, TotalVarInputC2)));
+	
+	tsr_tensor_halfC InputC2BoxSize = MaxMinusMin3x3(InputC2);
+	ClampError = max(ClampError, InputC2BoxSize * tsr_half(FilteringWeight * 0.25));
+	ClampError = ClampError + BackbufferQuantizationErrorVector;
+
+	tsr_tensor_halfC FilteredInput   = Blur3x3(InputC2);
+
+	// Build the clamping box of the filtered signal
+	tsr_tensor_halfC FilteredBoxMin;
+	tsr_tensor_halfC FilteredBoxMax;
+	MinMax3x3(FilteredInput, /* out */ FilteredBoxMin, /* out */ FilteredBoxMax);
+	FilteredBoxMin = FilteredBoxMin - ClampError;
+	FilteredBoxMax = FilteredBoxMax + ClampError;
+		
+	// Clamp filtered signal
+	tsr_tensor_halfC ClampedFilteredHistory = fastClamp(FilteredHistory, FilteredBoxMin, FilteredBoxMax);
+	BRANCH
+	if (ApplyMoireError && bEnableFlickeringHeuristic)
+	{
+		tsr_tensor_halfC MoireErrorSize;
+			
+		UNROLL_N(SIMD_SIZE)
+		for (uint ElementIndex = 0; ElementIndex < SIMD_SIZE; ElementIndex++)
+		#if 0
+		{
+			tsr_halfC ClampedEnergy = abs(FlickeringClampedFilteredHistory.GetElement(ElementIndex) - FilteredHistory.GetElement(ElementIndex));
+			tsr_halfC ClampedEnergyAmount = ClampedEnergy * SafeRcp(ClampedEnergy[0] + ClampedEnergy[1] + ClampedEnergy[2]);
+
+			// Apply the moire error on each individual channel based on how much they need to clamp.
+			MoireErrorSize.SetElement(ElementIndex, ClampedEnergyAmount * (kMoireLumaToChannel * MoireError.GetElement(ElementIndex)));
+		}
+		#else
+		{
+			MoireErrorSize.SetElement(ElementIndex, kMoireLumaToChannel * MoireError.GetElement(ElementIndex));
+		}
+		#endif
+
+		tsr_tensor_halfC StableFilteredBoxMin = min(FilteredBoxMin, FilteredBoxMax - MoireErrorSize);
+		tsr_tensor_halfC StableFilteredBoxMax = max(FilteredBoxMax, FilteredBoxMin + MoireErrorSize);
+			
+		ClampedFilteredHistory = fastClamp(FilteredHistory, StableFilteredBoxMin, StableFilteredBoxMax);
+	}
+		
+	// Compute the box size.
+	tsr_tensor_halfC BoxSize;
+	{
+		BoxSize = (
+			InputC2BoxSize * tsr_tensor_halfC::Const(FilteringWeight) +
+			(BackbufferQuantizationErrorVector * tsr_half(2 * FilteringWeight)));
+
+		BRANCH
+		if (ApplyMoireError && bEnableFlickeringHeuristic)
+		{
+			tsr_tensor_halfC MoireErrorSize;
+			UNROLL_N(SIMD_SIZE)
+			for (uint ElementIndex = 0; ElementIndex < SIMD_SIZE; ElementIndex++)
+			{
+				tsr_halfC ClampedEnergy = abs(ClampedFilteredHistory.GetElement(ElementIndex) - FilteredHistory.GetElement(ElementIndex));
+				tsr_halfC ClampedEnergyAmount = ClampedEnergy * SafeRcp(ClampedEnergy[0] + ClampedEnergy[1] + ClampedEnergy[2]);
+
+				// Apply the moire error on each individual channel based on how much they need to clamp.
+				MoireErrorSize.SetElement(ElementIndex, ClampedEnergyAmount * ((FilteringWeight * tsr_half(3.0)) * MoireError.GetElement(ElementIndex)));
+			}
+			BoxSize = max(BoxSize, MoireErrorSize);
+		}
+	}
+
+	tsr_tensor_halfC Delta = max(abs(FilteredInput - FilteredHistory), BoxSize);
+
+	tsr_tensor_halfC RawClampedEnergy      = abs(ClampedFilteredHistory - FilteredHistory);
+
+	tsr_tensor_halfC RawFactor        = saturate(tsr_tensor_halfC::Const(1.0) - RawClampedEnergy * rcp(Delta));
+	tsr_tensor_half RawRejection      = min3(RawFactor[0], RawFactor[1], RawFactor[2]);
+	#if CONFIG_SCENE_COLOR_ALPHA
+		RawRejection = min(RawRejection, RawFactor[3]);
+	#endif
+
+	#if CONFIG_SCENE_COLOR_ALPHA
+		OutRejectionClampBlend = Min3x3(Median3x3(RawRejection));
+		tsr_tensor_halfC FilteredClampedEnergy = Max3x3(Median3x3(RawClampedEnergy));
+	#else
+		tsr_tensor_halfC FilteredClampedEnergy;
+		Deconcatenate(MaxRGBMinA3x3(Median3x3(Concatenate(RawClampedEnergy, RawRejection))), /* out */ FilteredClampedEnergy, /* out */ OutRejectionClampBlend);
+	#endif
+
+	tsr_tensor_halfC FilteredFactor   = saturate(tsr_tensor_halfC::Const(1.0) - FilteredClampedEnergy * rcp(Delta));
+	tsr_tensor_half FilteredRejection = min3(FilteredFactor[0], FilteredFactor[1], FilteredFactor[2]);
+	#if CONFIG_SCENE_COLOR_ALPHA
+		FilteredRejection = min(FilteredRejection, FilteredFactor[3]);
+	#endif
+	OutRejectionBlendFinal          = FilteredRejection;
+} // MeasureRejection()
