Stochastic Direct Lighting - fixed bug in waveops building the hierarchical light candidate mask.

diff --git a/Engine/Shaders/Private/StochasticDirectLighting/StochasticDirectLightingSampling.usf b/Engine/Shaders/Private/StochasticDirectLighting/StochasticDirectLightingSampling.usf
--- a/Engine/Shaders/Private/StochasticDirectLighting/StochasticDirectLightingSampling.usf
+++ b/Engine/Shaders/Private/StochasticDirectLighting/StochasticDirectLightingSampling.usf
@@ -132,6 +132,7 @@ void GenerateLightSamplesCS(
 	uint DownsampledTileIndex = GroupId.x;
 	if (DownsampledTileIndex < DownsampledTileAllocator[TILE_TYPE])
 	{
+		uint LocalCandidateLightHiMask = 0;
 		uint2 DownsampledTileCoord = UnpackTile(DownsampledTileData[DownsampledTileIndex + TILE_TYPE * DownsampledTileDataStride]);
 		uint2 DownsampledScreenCoord = DownsampledTileCoord * TILE_SIZE + GroupThreadId.xy;
 
@@ -149,8 +150,6 @@ void GenerateLightSamplesCS(
 			const FCulledLightsGridData CulledLightGridData = GetCulledLightsGrid(GridIndex, EyeIndex);
 			const uint NumLightsInGridCell = min(CulledLightGridData.NumLights, GetMaxLightsPerCell(EyeIndex));
 
-			uint LocalCandidateLightHiMask = 0;
-
 			for (uint GridLightIndex = 0; GridLightIndex < NumLightsInGridCell; ++GridLightIndex)
 			{
 				const uint LocalLightIndex = GetCulledLightDataGrid(CulledLightGridData.DataStartIndex + GridLightIndex);
@@ -166,12 +165,12 @@ void GenerateLightSamplesCS(
 					LocalCandidateLightHiMask |= HiBitMask;
 				}
 			}
+		}
 
-			if (WaveIsFirstLane())
-			{
-				LocalCandidateLightHiMask = WaveActiveBitOr(LocalCandidateLightHiMask);
-				InterlockedOr(SharedCandidateLightHiMask, LocalCandidateLightHiMask);
-			}
+		uint WaveHiMask = WaveActiveBitOr(LocalCandidateLightHiMask);
+		if (WaveIsFirstLane())
+		{
+			InterlockedOr(SharedCandidateLightHiMask, WaveHiMask);
 		}
 	
 		GroupMemoryBarrierWithGroupSync();
