Fix uninitialized hair  data fetching.

diff --git a/Engine/Shaders/Private/HairStrands/HairStrandsComposition.usf b/Engine/Shaders/Private/HairStrands/HairStrandsComposition.usf
--- a/Engine/Shaders/Private/HairStrands/HairStrandsComposition.usf
+++ b/Engine/Shaders/Private/HairStrands/HairStrandsComposition.usf
@@ -311,6 +311,13 @@ void MainPS(
 #endif
 
 	const uint3 PixelCoord = uint3(floor(SvPosition.xy), 0);
+
+	const bool bIsValid = HairStrands.HairOnlyDepthTexture.Load(PixelCoord).x > 0;
+	if (!bIsValid)
+	{
+		discard;
+	}
+
 	const FNodeDesc NodeDesc = DecodeNodeDesc(HairStrands.HairSampleOffset.Load(PixelCoord));
 	if (NodeDesc.Count == 0)
 	{
