Prepare hair position buffer to use byteaddressbuffer.

diff --git a/Engine/Plugins/Runtime/HairStrands/Shaders/Private/DeformerDataInterfaceGroom.ush b/Engine/Plugins/Runtime/HairStrands/Shaders/Private/DeformerDataInterfaceGroom.ush
--- a/Engine/Plugins/Runtime/HairStrands/Shaders/Private/DeformerDataInterfaceGroom.ush
+++ b/Engine/Plugins/Runtime/HairStrands/Shaders/Private/DeformerDataInterfaceGroom.ush
@@ -26,8 +26,9 @@ FHairControlPoint ReadControlPointData(uint ControlPointIndex)
 	}
 
 	// Hair deformer work in local space so no need for position offset
-	return UnpackHairControlPoint(
-		{DataInterfaceName}_PositionBuffer[ControlPointIndex], 
+	return ReadHairControlPoint(
+		{DataInterfaceName}_PositionBuffer,
+		ControlPointIndex, 
 		{DataInterfaceName}_GetHairInstancePositionOffset(),
 		{DataInterfaceName}_Radius, 
 		{DataInterfaceName}_RootScale, 
diff --git a/Engine/Plugins/Runtime/HairStrands/Shaders/Private/DeformerDataInterfaceGroomGuide.ush b/Engine/Plugins/Runtime/HairStrands/Shaders/Private/DeformerDataInterfaceGroomGuide.ush
--- a/Engine/Plugins/Runtime/HairStrands/Shaders/Private/DeformerDataInterfaceGroomGuide.ush
+++ b/Engine/Plugins/Runtime/HairStrands/Shaders/Private/DeformerDataInterfaceGroomGuide.ush
@@ -21,8 +21,9 @@ FHairControlPoint ReadControlPointData_{DataInterfaceName}(uint ControlPointInde
 	}
 
 	// Hair deformer work in local space so no need for position offset
-	return UnpackHairControlPoint(
-		{DataInterfaceName}_PositionBuffer[ControlPointIndex], 
+	return ReadHairControlPoint(
+		{DataInterfaceName}_PositionBuffer,
+		ControlPointIndex, 
 		{DataInterfaceName}_GetHairInstancePositionOffset(),
 		{DataInterfaceName}_Radius, 
 		{DataInterfaceName}_RootScale, 
diff --git a/Engine/Plugins/Runtime/HairStrands/Shaders/Private/NiagaraDataInterfaceHairStrandsTemplate.ush b/Engine/Plugins/Runtime/HairStrands/Shaders/Private/NiagaraDataInterfaceHairStrandsTemplate.ush
--- a/Engine/Plugins/Runtime/HairStrands/Shaders/Private/NiagaraDataInterfaceHairStrandsTemplate.ush
+++ b/Engine/Plugins/Runtime/HairStrands/Shaders/Private/NiagaraDataInterfaceHairStrandsTemplate.ush
@@ -144,13 +144,13 @@ float ComputeStrandLength_{ParameterName}(in int StrandIndex)
 		int PointOffset = Curve.PointIndex;
 		const int EdgeCount = Curve.PointCount-1;
 
-		float3 PointNext = UnpackHairControlPointPosition({ParameterName}_RestPositionBuffer[PointOffset++]);
+		float3 PointNext = ReadHairControlPointPosition({ParameterName}_RestPositionBuffer, PointOffset++);
 		float3 PointPrev = PointNext;
 
 		for (int EdgeIndex = 0; EdgeIndex < EdgeCount; ++EdgeIndex)
 		{
 			PointPrev = PointNext;
-			PointNext = UnpackHairControlPointPosition({ParameterName}_RestPositionBuffer[PointOffset++]);
+			PointNext = ReadHairControlPointPosition({ParameterName}_RestPositionBuffer, PointOffset++);
 			StrandLength += length(PointNext-PointPrev);
 		}
 	}
@@ -211,19 +211,19 @@ void SetNodePosition_{ParameterName}(in bool RootNode, const float NodePoint, co
 		const int FirstNext = PointPrev+1;
 
 		const float FirstAlpha = FirstPoint - (float)FirstPrev;
-		const float3 FirstPosition =  UnpackHairControlPointPosition({ParameterName}_RestPositionBuffer[PointOffset+FirstPrev]) * (1.0-FirstAlpha) + 
-						 			  UnpackHairControlPointPosition({ParameterName}_RestPositionBuffer[PointOffset+FirstNext]) * FirstAlpha;
-		const float3 DiffPosition = FirstPosition - UnpackHairControlPointPosition({ParameterName}_RestPositionBuffer[PointOffset]);
+		const float3 FirstPosition =  ReadHairControlPointPosition({ParameterName}_RestPositionBuffer, PointOffset+FirstPrev) * (1.0-FirstAlpha) + 
+						 			  ReadHairControlPointPosition({ParameterName}_RestPositionBuffer, PointOffset+FirstNext) * FirstAlpha;
+		const float3 DiffPosition = FirstPosition - ReadHairControlPointPosition({ParameterName}_RestPositionBuffer, PointOffset);
 		const float3 EdgeDirection =  normalize(DiffPosition);
 
 		const float EdgeLength = ComputeStrandLength_{ParameterName}(StrandIndex) * CoordScale;
-		OutNodePosition = 	UnpackHairControlPointPosition({ParameterName}_RestPositionBuffer[PointOffset+PointPrev]) - EdgeLength * normalize(EdgeDirection);					 
+		OutNodePosition = 	ReadHairControlPointPosition({ParameterName}_RestPositionBuffer, PointOffset+PointPrev) - EdgeLength * normalize(EdgeDirection);
 	}
 	else
 	{
 		const float NodeAlpha = NodePoint - (float)PointPrev;
-		OutNodePosition =  UnpackHairControlPointPosition({ParameterName}_RestPositionBuffer[PointOffset+PointPrev]) * (1.0-NodeAlpha) + 
-						   UnpackHairControlPointPosition({ParameterName}_RestPositionBuffer[PointOffset+PointNext]) * NodeAlpha;
+		OutNodePosition =  ReadHairControlPointPosition({ParameterName}_RestPositionBuffer, PointOffset+PointPrev) * (1.0-NodeAlpha) + 
+						   ReadHairControlPointPosition({ParameterName}_RestPositionBuffer, PointOffset+PointNext) * NodeAlpha;
 	}
 	OutNodePosition = mul(float4( OutNodePosition+{ParameterName}_RestPositionOffset, 1.0), {ParameterName}_WorldTransform).xyz;
 }
@@ -666,19 +666,24 @@ void UpdateDeformedPositions_{ParameterName}(in float3 LocalDisplace, in int Pos
 {
 	if( {ParameterName}_InterpolationMode >= 1)
 	{
-		// TODO: convert this code back to use UnpackHairControlPointPosition & PackHairControlPointPosition
+		const FPackedHairPosition Packed = ReadPackedHairPosition({ParameterName}_RestPositionBuffer, PositionIndex);
+		const FHairControlPoint CP = UnpackHairControlPoint(Packed);
+
 		const float3 LocalTrianglePosition = TriangleLocalPosition_{ParameterName}(RestTrianglePosition, RestTriangleOrientation, 
-			f16tof32({ParameterName}_RestPositionBuffer[PositionIndex].xyz) + {ParameterName}_RestPositionOffset) + LocalDisplace;
+			CP.Position.xyz + {ParameterName}_RestPositionOffset) + LocalDisplace;
 			
 		const float3 LocalComponentPosition = TriangleWorldPosition_{ParameterName}(DeformedTrianglePosition, DeformedTriangleOrientation, LocalTrianglePosition) - {ParameterName}_DeformedPositionOffset[0].xyz;
-		{ParameterName}_DeformedPositionBuffer[PositionIndex] = uint4(f32tof16(LocalComponentPosition), {ParameterName}_RestPositionBuffer[PositionIndex].w);
 		
+		WritePackedHairControlPointPosition({ParameterName}_DeformedPositionBuffer, PositionIndex, Packed, LocalComponentPosition);
 	}
 	else
 	{
-		// TODO: convert this code back to use UnpackHairControlPointPosition & PackHairControlPointPosition
-		const float3 LocalComponentPosition = f16tof32({ParameterName}_RestPositionBuffer[PositionIndex].xyz) + LocalDisplace + {ParameterName}_RestPositionOffset - {ParameterName}_DeformedPositionOffset[0].xyz;
-		{ParameterName}_DeformedPositionBuffer[PositionIndex] = uint4(f32tof16(LocalComponentPosition), {ParameterName}_RestPositionBuffer[PositionIndex].w);	
+		const FPackedHairPosition Packed = ReadPackedHairPosition({ParameterName}_RestPositionBuffer, PositionIndex);
+		const FHairControlPoint CP = UnpackHairControlPoint(Packed);
+
+		const float3 LocalComponentPosition = CP.Position.xyz + LocalDisplace + {ParameterName}_RestPositionOffset - {ParameterName}_DeformedPositionOffset[0].xyz;
+
+		WritePackedHairControlPointPosition({ParameterName}_DeformedPositionBuffer, PositionIndex, Packed, LocalComponentPosition);
 	}
 }
 
@@ -725,7 +730,7 @@ void UpdatePointPosition_{ParameterName}(in float3 NodePosition, in float3 RestP
 // Get the point position
 void GetPointPosition_{ParameterName}(in int PointIndex, out float3 OutPointPosition)
 {
-	const float3 Position = UnpackHairControlPointPosition({ParameterName}_DeformedPositionBuffer[PointIndex], {ParameterName}_DeformedPositionOffset[0].xyz);
+	const float3 Position = ReadHairControlPointPosition({ParameterName}_DeformedPositionBuffer, PointIndex, {ParameterName}_DeformedPositionOffset[0].xyz);
 	OutPointPosition = mul(float4(Position, 1.0), {ParameterName}_WorldTransform).xyz;
 }
 
diff --git a/Engine/Plugins/Runtime/HairStrands/Shaders/Private/NiagaraInterpolateGroomGuides.usf b/Engine/Plugins/Runtime/HairStrands/Shaders/Private/NiagaraInterpolateGroomGuides.usf
--- a/Engine/Plugins/Runtime/HairStrands/Shaders/Private/NiagaraInterpolateGroomGuides.usf
+++ b/Engine/Plugins/Runtime/HairStrands/Shaders/Private/NiagaraInterpolateGroomGuides.usf
@@ -130,12 +130,12 @@ void MainCS(uint3 DispatchThreadId : SV_DispatchThreadID)
 
 		const float3 LocalPointDisplace = (TrianglePrevNodePosition - TrianglePrevRestPosition)  * (1.0 - PointAlpha) + (TriangleNextNodePosition - TriangleNextRestPosition) * PointAlpha;
 		
-		const uint4 Packed = RestPositionBuffer[DispatchThreadId.x];
+		const FPackedHairPosition Packed = ReadPackedHairPosition(RestPositionBuffer, DispatchThreadId.x);
 		const float3 LocalTrianglePosition = DIHairStrands_TriangleLocalPosition(RestTrianglePosition, RestTriangleOrientation, 
 					UnpackHairControlPointPosition(Packed, RestPositionOffset)) + LocalPointDisplace;
 		
 		const float3 LocalComponentPosition = DIHairStrands_TriangleWorldPosition(DeformedTrianglePosition, DeformedTriangleOrientation, LocalTrianglePosition) - DeformedPositionOffset[0].xyz;
-		DeformedPositionBuffer[DispatchThreadId.x] = PackHairControlPointPosition(Packed, LocalComponentPosition);
+		WritePackedHairControlPointPosition(DeformedPositionBuffer, DispatchThreadId.x, Packed, LocalComponentPosition);
 
 #endif
 		
@@ -143,9 +143,9 @@ void MainCS(uint3 DispatchThreadId : SV_DispatchThreadID)
 		
 		const float3 LocalPointDisplace = (LocalPrevNodePosition - LocalPrevRestPosition)  * (1.0 - PointAlpha) + (LocalNextNodePosition - LocalNextRestPosition) * PointAlpha;
 
-		const uint4 Packed = RestPositionBuffer[DispatchThreadId.x];
-		const float3 LocalComponentPosition = UnpackHairControlPointPosition(RestPositionBuffer[DispatchThreadId.x], RestPositionOffset) + LocalPointDisplace;
-		DeformedPositionBuffer[DispatchThreadId.x] = PackHairControlPointPosition(Packed, LocalComponentPosition, DeformedPositionOffset[0].xyz);
+		const FPackedHairPosition Packed = ReadPackedHairPosition(RestPositionBuffer, DispatchThreadId.x);
+		const float3 LocalComponentPosition = ReadHairControlPointPosition(RestPositionBuffer, DispatchThreadId.x, RestPositionOffset) + LocalPointDisplace;
+		WritePackedHairControlPointPosition(DeformedPositionBuffer, DispatchThreadId.x, Packed, LocalComponentPosition, DeformedPositionOffset[0].xyz);
 		
 #endif
 	}
diff --git a/Engine/Shaders/Private/HairStrands/HairCardsDeformation.usf b/Engine/Shaders/Private/HairStrands/HairCardsDeformation.usf
--- a/Engine/Shaders/Private/HairStrands/HairCardsDeformation.usf
+++ b/Engine/Shaders/Private/HairStrands/HairCardsDeformation.usf
@@ -49,10 +49,10 @@ float3 ComputeDynamicGeometryOffset(
 	float3 RootBarycentric,
 	float3 GuideDeformedPositionOffset)
 {
-	const float3 RestGuidePoint = UnpackHairControlPointPosition(GuideRestPositionBuffer[GuideIndex], GuideRestPositionOffset);
+	const float3 RestGuidePoint = ReadHairControlPointPosition(GuideRestPositionBuffer, GuideIndex, GuideRestPositionOffset);
 	const float3 LocalRestGuidePoint = ToTriangle(RestGuidePoint, RootBarycentric, RestTriangle);
 
-	const float3 DeformedGuidePoint = UnpackHairControlPointPosition(GuideDeformedPositionBuffer[GuideIndex], GuideDeformedPositionOffset);
+	const float3 DeformedGuidePoint = ReadHairControlPointPosition(GuideDeformedPositionBuffer, GuideIndex, GuideDeformedPositionOffset);
 	const float3 LocalDeformedGuidePoint = ToTriangle(DeformedGuidePoint, RootBarycentric, DeformedTriangle);
 
 	return (LocalDeformedGuidePoint - LocalRestGuidePoint) * GuideVertexWeight;
@@ -95,29 +95,33 @@ void MainCS(uint DispatchThreadId : SV_DispatchThreadID)
 	const uint GuideIndex1 = min(GuideIndex0+1, GuideVertexCount-1);
 	const float GuideLerp = GuideData.VertexLerps.x;
 
-	const FHairControlPoint RestGuideP0 = UnpackHairControlPoint(
-		GuideRestPositionBuffer[GuideIndex0],
+	const FHairControlPoint RestGuideP0 = ReadHairControlPoint(
+		GuideRestPositionBuffer, 
+		GuideIndex0,
 		GuideRestPositionOffset,
 		GuideRadius,
 		GuideRootScale,
 		GuideTipScale);
 
-	const FHairControlPoint RestGuideP1 = UnpackHairControlPoint(
-		GuideRestPositionBuffer[GuideIndex1],
+	const FHairControlPoint RestGuideP1 = ReadHairControlPoint(
+		GuideRestPositionBuffer,
+		GuideIndex1,
 		GuideRestPositionOffset,
 		GuideRadius,
 		GuideRootScale,
 		GuideTipScale);
 
-	const FHairControlPoint DeformedGuideP0 = UnpackHairControlPoint(
-		GuideDeformedPositionBuffer[GuideIndex0],
+	const FHairControlPoint DeformedGuideP0 = ReadHairControlPoint(
+		GuideDeformedPositionBuffer,
+		GuideIndex0,
 		GuideDeformedPositionOffset,
 		GuideRadius,
 		GuideRootScale,
 		GuideTipScale);
 
-	const FHairControlPoint DeformedGuideP1 = UnpackHairControlPoint(
-		GuideDeformedPositionBuffer[GuideIndex1],
+	const FHairControlPoint DeformedGuideP1 = ReadHairControlPoint(
+		GuideDeformedPositionBuffer,
+		GuideIndex1,
 		GuideDeformedPositionOffset,
 		GuideRadius,
 		GuideRootScale,
diff --git a/Engine/Shaders/Private/HairStrands/HairStrandsAttributeCommon.ush b/Engine/Shaders/Private/HairStrands/HairStrandsAttributeCommon.ush
--- a/Engine/Shaders/Private/HairStrands/HairStrandsAttributeCommon.ush
+++ b/Engine/Shaders/Private/HairStrands/HairStrandsAttributeCommon.ush
@@ -81,14 +81,14 @@ uint InternalGetHairStrandsAttributesPerVertex(uint HairControlPointId, uint Att
 	return InternalGetHairStrandsAttributes(HairControlPointId, AttributeOffsetInBytes, AttributeStrideInBytes, InAttributeBuffer, InIndexToChunkDivAsShift, InChunkElementCount, InChunkStrideInBytes);
 }
 
-float InternalGetHairStrandsCoordU(Buffer<uint4> InPositionBuffer, uint VertexId, float InRadius, float InRootScale, float InTipScale)
+float InternalGetHairStrandsCoordU(Buffer<uint4> InPositionBuffer, uint InPointIndex, float InRadius, float InRootScale, float InTipScale)
 {
-	return UnpackHairControlPoint(InPositionBuffer[VertexId], float3(0, 0, 0), InRadius, InRootScale, InTipScale).UCoord;
+	return ReadHairControlPoint(InPositionBuffer, InPointIndex, float3(0, 0, 0), InRadius, InRootScale, InTipScale).UCoord;
 }
 
-float InternalGetHairStrandsWorldRadius(Buffer<uint4> InPositionBuffer, uint VertexId, float InRadius, float InRootScale, float InTipScale)
+float InternalGetHairStrandsWorldRadius(Buffer<uint4> InPositionBuffer, uint InPointIndex, float InRadius, float InRootScale, float InTipScale)
 {
-	return UnpackHairControlPoint(InPositionBuffer[VertexId], float3(0, 0, 0), InRadius, InRootScale, InTipScale).WorldRadius;
+	return ReadHairControlPoint(InPositionBuffer, InPointIndex, float3(0, 0, 0), InRadius, InRootScale, InTipScale).WorldRadius;
 }
 
 bool IsAttributeValid(uint InOffset)
diff --git a/Engine/Shaders/Private/HairStrands/HairStrandsClusterCulling.usf b/Engine/Shaders/Private/HairStrands/HairStrandsClusterCulling.usf
--- a/Engine/Shaders/Private/HairStrands/HairStrandsClusterCulling.usf
+++ b/Engine/Shaders/Private/HairStrands/HairStrandsClusterCulling.usf
@@ -174,7 +174,7 @@ void Main(uint3 DispatchThreadId : SV_DispatchThreadID, uint ThreadIndex : SV_Gr
 		// 2.2 Mark point having their LOD visibility smaller than the requested LOD level
 		float RadiusFade = 1.0f;
 		ValidPoints[CurveIndex][CurrentPointIndex] = 0; 
-		if (CurveIndex < CurveCount)
+		if (CurveIndex < CurveCount) // incorrect
 		{	
 			if (CurrentPointIndex < CurvePointCount[CurveIndex])
 			{
@@ -504,7 +504,7 @@ groupshared float3 OutHairPositionOffset;
 		{
 			const uint GlobalPointIndex = CurvePointIndex + LocalPointIndex;
 			const bool bIsActive = IsHairControlPointActive(GlobalPointIndex, RenPointLODBuffer, LODIndex);
-			const FHairControlPoint CP = UnpackHairControlPoint(RenderDeformedPositionBuffer[GlobalPointIndex], OutHairPositionOffset, 1, 1, 1);
+			const FHairControlPoint CP = ReadHairControlPoint(RenderDeformedPositionBuffer, GlobalPointIndex, OutHairPositionOffset, 1, 1, 1);
 			if (bIsActive && all(IsFinite(CP.Position)))
 			{
 				const int3 WorldSpaceCentimeters = int3(mul(float4(CP.Position, 1.0f), LocalToTranslatedWorldMatrix).xyz);
diff --git a/Engine/Shaders/Private/HairStrands/HairStrandsDebug.usf b/Engine/Shaders/Private/HairStrands/HairStrandsDebug.usf
--- a/Engine/Shaders/Private/HairStrands/HairStrandsDebug.usf
+++ b/Engine/Shaders/Private/HairStrands/HairStrandsDebug.usf
@@ -273,8 +273,9 @@ void MainCS(uint3 DispatchThreadId : SV_DispatchThreadID)
 		VertexIndex = HairStrandsVF.CullingIndexBuffer[VertexIndex];
 	}
 
-	const float3 P_Local = UnpackHairControlPoint(
-		HairStrandsVF.PositionBuffer[VertexIndex],
+	const float3 P_Local = ReadHairControlPoint(
+		HairStrandsVF.PositionBuffer,
+		VertexIndex,
 		ReadRenPositionOffset(HairStrandsVF.PositionOffsetBuffer, HairStrandsVF.RegisteredIndex),
 		HairStrandsVF.Radius * ScaleR,
 		HairStrandsVF.RootScale,
@@ -2001,48 +2002,45 @@ void MainCS(uint3 DispatchThreadId : SV_DispatchThreadID)
 	if (VertexId+1 >= VertexCount)
 		return;
 
-	uint4 PackedP0 = 0;
-	uint4 PackedP1 = 0;
+	FHairControlPoint P0 = (FHairControlPoint)0;
+	FHairControlPoint P1 = (FHairControlPoint)0;
 	float3 Offset = 0;
 	float4 ColorRoot = ColorWhite;
 	float4 ColorTip  = ColorWhite;
 
 	if (DebugMode == DEBUGMODE_REN_REST)
 	{
-		PackedP0 = RenRestPosition[VertexId];
-		PackedP1 = RenRestPosition[VertexId + 1];
+		P0       = ReadHairControlPoint(RenRestPosition, VertexId, 1, 1, 1);
+		P1       = ReadHairControlPoint(RenRestPosition, VertexId + 1, 1, 1, 1);
 		Offset   = RenRestOffset;
 		ColorRoot= ColorYellow;
 		ColorTip = ColorPurple;
 	}
 	else if (DebugMode == DEBUGMODE_REN_DEFORMED)
 	{
-		PackedP0 = RenDeformedPosition[VertexId];
-		PackedP1 = RenDeformedPosition[VertexId + 1];
+		P0       = ReadHairControlPoint(RenDeformedPosition, VertexId, 1, 1, 1);
+		P1       = ReadHairControlPoint(RenDeformedPosition, VertexId + 1, 1, 1, 1);
 		Offset   = ReadCardPositionOffset(RenDeformedOffset, InstanceRegisteredIndex, CardLODIndex);
 		ColorRoot= ColorYellow;
 		ColorTip = ColorPurple;
 	}
 	else if (DebugMode == DEBUGMODE_SIM_REST)
 	{
-		PackedP0 = SimRestPosition[VertexId];
-		PackedP1 = SimRestPosition[VertexId + 1];
+		P0       = ReadHairControlPoint(SimRestPosition, VertexId, 1, 1, 1);
+		P1       = ReadHairControlPoint(SimRestPosition, VertexId + 1, 1, 1, 1);
 		Offset   = SimRestOffset;
 		ColorRoot= ColorCyan;
 		ColorTip = ColorRed;
 	}
 	else if (DebugMode == DEBUGMODE_SIM_DEFORMED)
 	{
-		PackedP0 = SimDeformedPosition[VertexId];
-		PackedP1 = SimDeformedPosition[VertexId + 1];
+		P0       = ReadHairControlPoint(SimDeformedPosition, VertexId, 1, 1, 1);
+		P1       = ReadHairControlPoint(SimDeformedPosition, VertexId + 1, 1, 1, 1);
 		Offset   = ReadSimPositionOffset(SimDeformedOffset, InstanceRegisteredIndex);
 		ColorRoot= ColorCyan;
 		ColorTip = ColorRed;
 	}
 
-	const FHairControlPoint P0 = UnpackHairControlPoint(PackedP0, Offset, 1, 1, 1);
-	const FHairControlPoint P1 = UnpackHairControlPoint(PackedP1, Offset, 1, 1, 1);
-
 	if (P1.Type != HAIR_CONTROLPOINT_START)
 	{
 		const float4 Color0 = lerp(ColorRoot, ColorTip, P0.UCoord);
@@ -2600,8 +2598,9 @@ void CSMain(
 		const FHairVis VisNode = UnpackHairVis(NodeVisBuffer[SampleIndex]);
 		if (VisNode.MaterialId == MaterialId)
 		{
-			FHairControlPoint Out = UnpackHairControlPoint(
-				HairInstance_PositionBuffer[VisNode.ControlPointId],
+			FHairControlPoint Out = ReadHairControlPoint(
+				HairInstance_PositionBuffer,
+				VisNode.ControlPointId,
 				HairInstance_GetHairInstancePositionOffset(),
 				HairInstance_Radius,
 				HairInstance_RootScale,
diff --git a/Engine/Shaders/Private/HairStrands/HairStrandsForwardCulling.usf b/Engine/Shaders/Private/HairStrands/HairStrandsForwardCulling.usf
--- a/Engine/Shaders/Private/HairStrands/HairStrandsForwardCulling.usf
+++ b/Engine/Shaders/Private/HairStrands/HairStrandsForwardCulling.usf
@@ -118,15 +118,17 @@ void CSMain(uint DispatchThreadID : SV_DispatchThreadID, uint GroupThreadID : SV
 			if (bValidPoint)
 			{
 				// Fetch both control points
-				const FHairControlPoint CP0 = UnpackHairControlPoint(
-					HairStrandsVF_PositionBuffer[PrimID],
+				const FHairControlPoint CP0 = ReadHairControlPoint(
+					HairStrandsVF_PositionBuffer,
+					PrimID,
 					PositionOffset,
 					HairStrandsVF_Radius,
 					HairStrandsVF_RootScale,
 					HairStrandsVF_TipScale);
 
-				const FHairControlPoint CP1 = UnpackHairControlPoint(
-					HairStrandsVF_PositionBuffer[PrimID + 1],
+				const FHairControlPoint CP1 = ReadHairControlPoint(
+					HairStrandsVF_PositionBuffer,
+					PrimID + 1,
 					PositionOffset,
 					HairStrandsVF_Radius,
 					HairStrandsVF_RootScale,
@@ -174,15 +176,17 @@ void CSMain(uint DispatchThreadID : SV_DispatchThreadID, uint GroupThreadID : SV
 			if (bValidPoint)
 			{
 				// Fetch both control points
-				const FHairControlPoint CP0 = UnpackHairControlPoint(
-					HairStrandsVF_PositionBuffer[PrimID],
+				const FHairControlPoint CP0 = ReadHairControlPoint(
+					HairStrandsVF_PositionBuffer,
+					PrimID,
 					PositionOffset,
 					HairStrandsVF_Radius,
 					HairStrandsVF_RootScale,
 					HairStrandsVF_TipScale);
 
-				const FHairControlPoint CP1 = UnpackHairControlPoint(
-					HairStrandsVF_PositionBuffer[PrimID + 1],
+				const FHairControlPoint CP1 = ReadHairControlPoint(
+					HairStrandsVF_PositionBuffer,
+					PrimID + 1,
 					PositionOffset,
 					HairStrandsVF_Radius,
 					HairStrandsVF_RootScale,
@@ -263,8 +267,9 @@ void CSMain(uint DispatchThreadID : SV_DispatchThreadID, uint GroupThreadID : SV
 		const float3 PositionOffset = HairStrandsVF_GetHairInstancePositionOffset();
 
 		// Fetch both control points
-		FHairControlPoint CP0 = UnpackHairControlPoint(
-			HairStrandsVF_PositionBuffer[ControlPointId],
+		FHairControlPoint CP0 = ReadHairControlPoint(
+			HairStrandsVF_PositionBuffer,
+			ControlPointId,
 			PositionOffset,
 			HairStrandsVF_Radius,
 			HairStrandsVF_RootScale,
@@ -273,8 +278,9 @@ void CSMain(uint DispatchThreadID : SV_DispatchThreadID, uint GroupThreadID : SV
 		FHairControlPoint CP1 = CP0;
 		if (CP0.Type != HAIR_CONTROLPOINT_END)
 		{
-			CP1 = UnpackHairControlPoint(
-				HairStrandsVF_PositionBuffer[ControlPointId + 1],
+			CP1 = ReadHairControlPoint(
+				HairStrandsVF_PositionBuffer,
+				ControlPointId + 1,
 				PositionOffset,
 				HairStrandsVF_Radius,
 				HairStrandsVF_RootScale,
diff --git a/Engine/Shaders/Private/HairStrands/HairStrandsGuideDeform.usf b/Engine/Shaders/Private/HairStrands/HairStrandsGuideDeform.usf
--- a/Engine/Shaders/Private/HairStrands/HairStrandsGuideDeform.usf
+++ b/Engine/Shaders/Private/HairStrands/HairStrandsGuideDeform.usf
@@ -37,7 +37,7 @@ float3 DeformGuideByTriangle(
 	FHairMeshTriangle DeformedTriangle,
 	float3 RootBarycentric)
 {
-	const float3 WorldRestGuidePoint = UnpackHairControlPointPosition(SimRestPosePositionBuffer[GuideIndex], SimRestOffset);
+	const float3 WorldRestGuidePoint = ReadHairControlPointPosition(SimRestPosePositionBuffer, GuideIndex, SimRestOffset);
 	const float3 LocalRestGuidePoint = ToTriangle(WorldRestGuidePoint, RootBarycentric, RestTriangle);
 	const float3 WorldDeformedGuidePoint = ToWorld(LocalRestGuidePoint, RootBarycentric, DeformedTriangle);
 
@@ -105,10 +105,11 @@ void MainCS(uint DispatchThreadId : SV_DispatchThreadID)
 	
 	if (VertexIndex < VertexCount)
 	{
-		const uint4  PackedCP			= SimRestPosePositionBuffer[VertexIndex];
+		const FPackedHairPosition PackedCP = ReadPackedHairPosition(SimRestPosePositionBuffer, VertexIndex);
 		const float3 CPPosition			= UnpackHairControlPointPosition(PackedCP);
-		const float3 DisplacedPosition	= DisplacePosition(VertexIndex, CPPosition);		
-		OutSimDeformedPositionBuffer[VertexIndex] = PackHairControlPointPosition(PackedCP, DisplacedPosition);
+		const float3 DisplacedPosition	= DisplacePosition(VertexIndex, CPPosition);
+		WritePackedHairControlPointPosition(OutSimDeformedPositionBuffer, VertexIndex, PackedCP, DisplacedPosition);
+//		OutSimDeformedPositionBuffer[VertexIndex] = PackHairControlPointPosition(PackedCP, DisplacedPosition);
 	}
 }
 
diff --git a/Engine/Shaders/Private/HairStrands/HairStrandsInterpolation.usf b/Engine/Shaders/Private/HairStrands/HairStrandsInterpolation.usf
--- a/Engine/Shaders/Private/HairStrands/HairStrandsInterpolation.usf
+++ b/Engine/Shaders/Private/HairStrands/HairStrandsInterpolation.usf
@@ -140,9 +140,9 @@ float3 ComputeStaticGeometryOffset(
 	float GuideVertexWeight, 
 	float3 OutSimHairPositionOffset)
 {
-	const float3 RestGuidePoint = UnpackHairControlPointPosition(SimRestPosePositionBuffer[GuideIndex.PointIndex], InSimHairPositionOffset);
+	const float3 RestGuidePoint = ReadHairControlPointPosition(SimRestPosePositionBuffer, GuideIndex.PointIndex, InSimHairPositionOffset);
 
-	const float3 DeformedGuidePoint = UnpackHairControlPointPosition(SimDeformedPositionBuffer[GuideIndex.PointIndex], OutSimHairPositionOffset);
+	const float3 DeformedGuidePoint = ReadHairControlPointPosition(SimDeformedPositionBuffer, GuideIndex.PointIndex, OutSimHairPositionOffset);
 
 	return (DeformedGuidePoint - RestGuidePoint) * GuideVertexWeight;
 }
@@ -151,9 +151,9 @@ float3 ComputeLocalGeometryOffset(
 	FGuideIndex GuideIndex, 
 	float GuideVertexWeight)
 {
-	const float3 RestGuidePoint = UnpackHairControlPointPosition(SimRestPosePositionBuffer[GuideIndex.PointIndex]) - UnpackHairControlPointPosition(SimRestPosePositionBuffer[GuideIndex.RootPointIndex]);
+	const float3 RestGuidePoint = ReadHairControlPointPosition(SimRestPosePositionBuffer, GuideIndex.PointIndex) - ReadHairControlPointPosition(SimRestPosePositionBuffer, GuideIndex.RootPointIndex);
 
-	const float3 DeformedGuidePoint = UnpackHairControlPointPosition(SimDeformedPositionBuffer[GuideIndex.PointIndex]) - UnpackHairControlPointPosition(SimDeformedPositionBuffer[GuideIndex.RootPointIndex]);
+	const float3 DeformedGuidePoint = ReadHairControlPointPosition(SimDeformedPositionBuffer, GuideIndex.PointIndex) - ReadHairControlPointPosition(SimDeformedPositionBuffer, GuideIndex.RootPointIndex);
 
 	return (DeformedGuidePoint - RestGuidePoint) * GuideVertexWeight;
 }
@@ -166,10 +166,10 @@ float3 ComputeDynamicGeometryOffset(
 	float3 RootBarycentric,
 	float3 OutSimHairPositionOffset)
 {
-	const float3 RestGuidePoint = UnpackHairControlPointPosition(SimRestPosePositionBuffer[GuideIndex.PointIndex], InSimHairPositionOffset);
+	const float3 RestGuidePoint = ReadHairControlPointPosition(SimRestPosePositionBuffer, GuideIndex.PointIndex, InSimHairPositionOffset);
 	const float3 LocalRestGuidePoint = ToTriangle(RestGuidePoint, RootBarycentric, RestTriangle);
 
-	const float3 DeformedGuidePoint = UnpackHairControlPointPosition(SimDeformedPositionBuffer[GuideIndex.PointIndex], OutSimHairPositionOffset);
+	const float3 DeformedGuidePoint = ReadHairControlPointPosition(SimDeformedPositionBuffer, GuideIndex.PointIndex, OutSimHairPositionOffset);
 	const float3 LocalDeformedGuidePoint = ToTriangle(DeformedGuidePoint, RootBarycentric, DeformedTriangle);
 
 	return (LocalDeformedGuidePoint - LocalRestGuidePoint) * GuideVertexWeight;
@@ -179,9 +179,9 @@ float3 ComputeTranslateGeometryOffset(
 	FGuideIndex GuideIndex, 
 	float GuideVertexWeight)
 {
-	const float3 RestGuidePoint = UnpackHairControlPointPosition(SimRestPosePositionBuffer[GuideIndex.PointIndex]) - UnpackHairControlPointPosition(SimRestPosePositionBuffer[GuideIndex.RootPointIndex]);
+	const float3 RestGuidePoint = ReadHairControlPointPosition(SimRestPosePositionBuffer, GuideIndex.PointIndex) - ReadHairControlPointPosition(SimRestPosePositionBuffer, GuideIndex.RootPointIndex);
 
-	const float3 DeformedGuidePoint = UnpackHairControlPointPosition(SimDeformedPositionBuffer[GuideIndex.PointIndex]) - UnpackHairControlPointPosition(SimDeformedPositionBuffer[GuideIndex.RootPointIndex]);
+	const float3 DeformedGuidePoint = ReadHairControlPointPosition(SimDeformedPositionBuffer, GuideIndex.PointIndex) - ReadHairControlPointPosition(SimDeformedPositionBuffer, GuideIndex.RootPointIndex);
 
 	return (DeformedGuidePoint - RestGuidePoint) * GuideVertexWeight;
 }
@@ -192,10 +192,10 @@ float3 ComputeRotateGeometryOffset(
 	FHairMeshTriangle RestTriangle, 
 	FHairMeshTriangle DeformedTriangle)
 {	
-	const float3 RestGuidePoint = UnpackHairControlPointPosition(SimRestPosePositionBuffer[GuideIndex.PointIndex]) - UnpackHairControlPointPosition(SimRestPosePositionBuffer[GuideIndex.RootPointIndex]);
+	const float3 RestGuidePoint = ReadHairControlPointPosition(SimRestPosePositionBuffer, GuideIndex.PointIndex) - ReadHairControlPointPosition(SimRestPosePositionBuffer, GuideIndex.RootPointIndex);
 	const float3 LocalRestGuidePoint = VectorToTriangle(RestGuidePoint, RestTriangle);
 
-	const float3 DeformedGuidePoint = UnpackHairControlPointPosition(SimDeformedPositionBuffer[GuideIndex.PointIndex]) - UnpackHairControlPointPosition(SimDeformedPositionBuffer[GuideIndex.RootPointIndex]);
+	const float3 DeformedGuidePoint = ReadHairControlPointPosition(SimDeformedPositionBuffer, GuideIndex.PointIndex) - ReadHairControlPointPosition(SimDeformedPositionBuffer, GuideIndex.RootPointIndex);
 	const float3 LocalDeformedGuidePoint = VectorToTriangle(DeformedGuidePoint, DeformedTriangle);
 
 	return (LocalDeformedGuidePoint - LocalRestGuidePoint) * GuideVertexWeight;
@@ -514,7 +514,7 @@ void MainCS(uint3 DispatchThreadId : SV_DispatchThreadID, uint ThreadIndex : SV_
 				#endif
 
 				// Manual decoding of the rest position
-				uint4 PackedRestPosition = RenRestPosePositionBuffer[VertexIndex];
+				FPackedHairPosition PackedRestPosition = ReadPackedHairPosition(RenRestPosePositionBuffer, VertexIndex);
 				const float3 RestPosition = UnpackHairControlPointPosition(PackedRestPosition, InRenHairPositionOffset);
 				
 				float3 OutPosition = 0.0f;
@@ -552,10 +552,10 @@ void MainCS(uint3 DispatchThreadId : SV_DispatchThreadID, uint ThreadIndex : SV_
 				// 3. Apply deformer
 				#if PERMUTATION_DEFORMER
 				{
-					const uint4 PackedDeformerPosition = RenDeformerPositionBuffer[VertexIndex];
+					const FPackedHairPosition PackedDeformerPosition = ReadPackedHairPosition(RenDeformerPositionBuffer, VertexIndex);
 					const float3 DeformerPosition = UnpackHairControlPointPosition(PackedDeformerPosition, InRenHairPositionOffset);
 					OutPosition = ApplyDeformer(RenTri, RestPosition, DeformerPosition, OutPosition);
-					PackedRestPosition.w = PackedDeformerPosition.w;
+					CopyPackedAttribute(PackedRestPosition, PackedDeformerPosition);
 				}
 				#endif
 
@@ -570,7 +570,7 @@ void MainCS(uint3 DispatchThreadId : SV_DispatchThreadID, uint ThreadIndex : SV_
 					}
 				}
 
-				OutRenDeformedPositionBuffer[VertexIndex] = PackHairControlPointPosition(PackedRestPosition, OutPosition, RenPositionOffset);
+				WritePackedHairControlPointPosition(OutRenDeformedPositionBuffer, VertexIndex, PackedRestPosition, OutPosition, RenPositionOffset);
 			}
 		}
 	}
@@ -661,7 +661,7 @@ void MainCS(uint3 DispatchThreadId : SV_DispatchThreadID, uint ThreadIndex : SV_
 				#endif
 
 				// Manual decoding of the rest position
-				uint4 PackedRestPosition = RenRestPosePositionBuffer[VertexIndex];
+				FPackedHairPosition PackedRestPosition = ReadPackedHairPosition(RenRestPosePositionBuffer, VertexIndex);
 				const float3 RestPosition = UnpackHairControlPointPosition(PackedRestPosition, InRenHairPositionOffset);
 
 				float3 OutPosition = 0.0f;
@@ -699,10 +699,10 @@ void MainCS(uint3 DispatchThreadId : SV_DispatchThreadID, uint ThreadIndex : SV_
 				// 3. Apply deformer
 				#if PERMUTATION_DEFORMER
 				{
-					const uint4 PackedDeformerPosition = RenDeformerPositionBuffer[VertexIndex];
+					const FPackedHairPosition PackedDeformerPosition = ReadPackedHairPosition(RenDeformerPositionBuffer, VertexIndex);
 					const float3 DeformerPosition = UnpackHairControlPointPosition(PackedDeformerPosition, InRenHairPositionOffset);
 					OutPosition = ApplyDeformer(RenTri, RestPosition, DeformerPosition, OutPosition);
-					PackedRestPosition.w = PackedDeformerPosition.w;
+					CopyPackedAttribute(PackedRestPosition, PackedDeformerPosition);
 				}
 				#endif
 
@@ -717,7 +717,7 @@ void MainCS(uint3 DispatchThreadId : SV_DispatchThreadID, uint ThreadIndex : SV_
 					}
 				}
 
-				OutRenDeformedPositionBuffer[VertexIndex] = PackHairControlPointPosition(PackedRestPosition, OutPosition, RenPositionOffset);
+				WritePackedHairControlPointPosition(OutRenDeformedPositionBuffer, VertexIndex, PackedRestPosition, OutPosition, RenPositionOffset);
 			}
 		}
 	}
@@ -878,7 +878,7 @@ void MainCS(uint DispatchThreadId : SV_DispatchThreadID)
 		InPosition1Buffer[Index3+2]);
 	
 	const float3 WorldVertexPosition = lerp(InPosition0, InPosition1, InterpolationFactor);
-	FHairControlPoint CP = UnpackHairControlPoint(InRestPoseBuffer[Index], float3(0, 0, 0), 1, 1, 1);
+	FHairControlPoint CP = ReadHairControlPoint(InRestPoseBuffer, Index, float3(0, 0, 0), 1, 1, 1);
 	CP.Position = WorldVertexPosition;
 
 	if (bHasRadiusData == 1)
diff --git a/Engine/Shaders/Private/HairStrands/HairStrandsPack.ush b/Engine/Shaders/Private/HairStrands/HairStrandsPack.ush
--- a/Engine/Shaders/Private/HairStrands/HairStrandsPack.ush
+++ b/Engine/Shaders/Private/HairStrands/HairStrandsPack.ush
@@ -9,6 +9,18 @@
 ////////////////////////////////////////////////////////////////////////////////////////////////////////////////////////////////
 // Hair control points
 
+#define FPackedHairPosition uint4
+
+FPackedHairPosition ReadPackedHairPosition(Buffer<FPackedHairPosition> InBuffer, uint InPointIndex)
+{
+	return InBuffer[InPointIndex];
+}
+
+void CopyPackedAttribute(inout FPackedHairPosition Out, in FPackedHairPosition In)
+{
+	Out.w = Out.w;
+}
+
 struct FHairControlPoint
 {
 	float3 Position;
@@ -17,19 +29,19 @@ struct FHairControlPoint
 	uint   Type;
 };
 
-uint4 PackHairControlPoint(
+FPackedHairPosition PackHairControlPoint(
 	FHairControlPoint CP,
 	float3 InVF_PositionOffset,
-	float InVF_Radius)
+	float InVF_Radius) 
 {
 	uint PackedAttribute =  (uint(clamp(CP.UCoord * 0xFFu, 0u, 0xFFu)) << 8) |
 							(uint(clamp((CP.WorldRadius / InVF_Radius) * 63u, 0u, 63u)) << 2u) |
 							(CP.Type & 0x3);
-	return uint4(f32tof16(CP.Position - InVF_PositionOffset), PackedAttribute);
+	return FPackedHairPosition(f32tof16(CP.Position - InVF_PositionOffset), PackedAttribute);
 }
 
 FHairControlPoint UnpackHairControlPoint(
-	uint4 InPackedData,
+	FPackedHairPosition InPackedData,
 	float3 InVF_PositionOffset=0,
 	float InVF_Radius=1,
 	float InVF_RootScale=1,
@@ -48,26 +60,52 @@ FHairControlPoint UnpackHairControlPoint(
 	return Out;	
 }
 
-float UnpackHairControlPointCoordU(uint4 InPackedData)
+FHairControlPoint ReadHairControlPoint(
+	Buffer<FPackedHairPosition> InBuffer,
+	uint InIndex,
+	float3 InVF_PositionOffset=0,
+	float InVF_Radius=1,
+	float InVF_RootScale=1,
+	float InVF_TipScale=1)
+{
+	return UnpackHairControlPoint(
+		InBuffer[InIndex],
+		InVF_PositionOffset,
+		InVF_Radius,
+		InVF_RootScale,
+		InVF_TipScale);
+}
+
+float UnpackHairControlPointCoordU(FPackedHairPosition InPackedData)
 {
 	const uint PackedAlpha = asuint(InPackedData.w);
 	return ((PackedAlpha & 0xFF00) >> 8) / 255.f;
 }
 
-float3 UnpackHairControlPointPosition(uint4 InPackedData, float3 InVF_PositionOffset=0)
+float3 UnpackHairControlPointPosition(FPackedHairPosition InPackedData, float3 InVF_PositionOffset=0)
 {
 	return f16tof32(InPackedData.xyz) + InVF_PositionOffset;
 }
 
-uint4 PackHairControlPointPosition(uint4 InPackedData, float3 InNewPosition, float3 InPositionOffset=0)
+float3 ReadHairControlPointPosition(Buffer<FPackedHairPosition> InBuffer, uint InPointIndex, float3 InVF_PositionOffset=0)
+{
+	return UnpackHairControlPointPosition(InBuffer[InPointIndex], InVF_PositionOffset);
+}
+
+FPackedHairPosition PackHairControlPointPosition(FPackedHairPosition InPackedData, float3 InNewPosition, float3 InPositionOffset=0)
+{
+	return FPackedHairPosition(f32tof16(InNewPosition.xyz - InPositionOffset), InPackedData.w);
+}
+
+void WritePackedHairControlPointPosition(RWBuffer<FPackedHairPosition> OutBuffer, uint InIndex, FPackedHairPosition InPackedData, float3 InNewPosition, float3 InPositionOffset=0)
 {
-	return uint4(f32tof16(InNewPosition.xyz - InPositionOffset), InPackedData.w);
+	OutBuffer[InIndex] = PackHairControlPointPosition(InPackedData, InNewPosition, InPositionOffset);
 }
 
-uint4 PackHairControlPointRadius(uint4 InPackedData, float InNewWorldRadius, float InVFRadius)
+FPackedHairPosition PackHairControlPointRadius(FPackedHairPosition InPackedData, float InNewWorldRadius, float InVFRadius)
 {
 	const uint PackedNewRadius = uint(clamp(saturate(InNewWorldRadius / InVFRadius) * 63u, 0u, 63u)) << 2u;	
-	return uint4(InPackedData.xyz, (InPackedData.w & (~0xFC)) | PackedNewRadius);
+	return FPackedHairPosition(InPackedData.xyz, (InPackedData.w & (~0xFC)) | PackedNewRadius);
 }
 
 bool IsValidHairStrandsSegment(const FHairControlPoint CP0, const FHairControlPoint CP1)
diff --git a/Engine/Shaders/Private/HairStrands/HairStrandsRaytracingGeometry.usf b/Engine/Shaders/Private/HairStrands/HairStrandsRaytracingGeometry.usf
--- a/Engine/Shaders/Private/HairStrands/HairStrandsRaytracingGeometry.usf
+++ b/Engine/Shaders/Private/HairStrands/HairStrandsRaytracingGeometry.usf
@@ -39,8 +39,9 @@ FHairControlPoint GetControlPoint(uint ControlPointId)
 
 	if (ControlPointId < PointCount)
 	{
-		Out = UnpackHairControlPoint(
-			PositionBuffer[ControlPointId],
+		Out = ReadHairControlPoint(
+			PositionBuffer, 
+			ControlPointId,
 			ReadRenPositionOffset(PositionOffsetBuffer, HairStrandsVF_RegisteredIndex),
 			HairStrandsVF_HairRadius,
 			HairStrandsVF_HairRootScale,
@@ -340,8 +341,8 @@ void MainCS(uint2 DispatchThreadId : SV_DispatchThreadID, uint LinearIndex : SV_
 		// Position comparison
 		if (bIsValid)
 		{
-			const FHairControlPoint Curr = UnpackHairControlPoint(CurrPositionBuffer[VertexIndex0], float3(0, 0, 0), 1, 1, 1);
-			const FHairControlPoint Prev = UnpackHairControlPoint(PrevPositionBuffer[VertexIndex0], float3(0, 0, 0), 1, 1, 1);
+			const FHairControlPoint Curr = ReadHairControlPoint(CurrPositionBuffer, VertexIndex0, float3(0, 0, 0), 1, 1, 1);
+			const FHairControlPoint Prev = ReadHairControlPoint(PrevPositionBuffer, VertexIndex0, float3(0, 0, 0), 1, 1, 1);
 			const float3 Diff = Curr.Position - Prev.Position;
 			const bool bIsInvalid = dot(Diff, Diff) > PositionThreshold2;
 
diff --git a/Engine/Shaders/Private/HairStrands/HairStrandsTangent.usf b/Engine/Shaders/Private/HairStrands/HairStrandsTangent.usf
--- a/Engine/Shaders/Private/HairStrands/HairStrandsTangent.usf
+++ b/Engine/Shaders/Private/HairStrands/HairStrandsTangent.usf
@@ -62,9 +62,9 @@ void MainCS(uint DispatchThreadId : SV_DispatchThreadID)
 
 	if (IndexCurr < MaxPointCount)
 	{
-		const FHairControlPoint CPCurr = UnpackHairControlPoint(PositionBuffer[IndexCurr]);
-		const FHairControlPoint CPPrev = UnpackHairControlPoint(PositionBuffer[IndexPrev]);
-		const FHairControlPoint CPNext = UnpackHairControlPoint(PositionBuffer[IndexNext]);
+		const FHairControlPoint CPCurr = ReadHairControlPoint(PositionBuffer,IndexCurr);
+		const FHairControlPoint CPPrev = ReadHairControlPoint(PositionBuffer,IndexPrev);
+		const FHairControlPoint CPNext = ReadHairControlPoint(PositionBuffer,IndexNext);
 
 		float3 T0 = 0;
 		const bool bIsPrevValid = bPrevValid && (CPPrev.Type == HAIR_CONTROLPOINT_INSIDE || CPPrev.Type == HAIR_CONTROLPOINT_START);
diff --git a/Engine/Shaders/Private/HairStrands/HairStrandsTexturesGeneration.usf b/Engine/Shaders/Private/HairStrands/HairStrandsTexturesGeneration.usf
--- a/Engine/Shaders/Private/HairStrands/HairStrandsTexturesGeneration.usf
+++ b/Engine/Shaders/Private/HairStrands/HairStrandsTexturesGeneration.usf
@@ -322,11 +322,8 @@ bool Trace(
 						const uint P0 = P0_P1.x;
 						const uint P1 = P0_P1.y;
 
-						const uint4 PackedCP_P0 = InVF_PositionBuffer[P0];
-						const uint4 PackedCP_P1 = InVF_PositionBuffer[P1];
-
-						const FHairControlPoint CP_Curr = UnpackHairControlPoint(PackedCP_P0, InVF_PositionOffset, InVF_Radius, InVF_RootScale, InVF_TipScale);
-						const FHairControlPoint CP_Next = UnpackHairControlPoint(PackedCP_P1, InVF_PositionOffset, InVF_Radius, InVF_RootScale, InVF_TipScale);
+						const FHairControlPoint CP_Curr = ReadHairControlPoint(InVF_PositionBuffer, P0, InVF_PositionOffset, InVF_Radius, InVF_RootScale, InVF_TipScale);
+						const FHairControlPoint CP_Next = ReadHairControlPoint(InVF_PositionBuffer, P1, InVF_PositionOffset, InVF_Radius, InVF_RootScale, InVF_TipScale);
 
 						const float3 Tangent = ComputeTangent(CP_Curr, CP_Next);
 
@@ -664,9 +661,8 @@ void MainCS(int2 DispatchThreadId : SV_DispatchThreadID)
 	uint VertexIndex = DispatchThreadId.x;
 	if (VertexIndex < MaxVertexCount)
 	{
-		const uint4 PackedPosition = InPositions[VertexIndex];
 		const float3 PositionOffset = InPositionOffset[0].xyz;
-		const FHairControlPoint CP = UnpackHairControlPoint(PackedPosition, PositionOffset, 1.0f /*InVF_Radius*/, 1.0f /*InVF_RootScale*/, 1.0f /*InVF_TipScale*/);
+		const FHairControlPoint CP = ReadHairControlPoint(InPositions, VertexIndex, PositionOffset, 1.0f /*InVF_Radius*/, 1.0f /*InVF_RootScale*/, 1.0f /*InVF_TipScale*/);
 		
 		OutPositions[VertexIndex] = float4(CP.Position, CP.Type);
 	}
diff --git a/Engine/Shaders/Private/HairStrands/HairStrandsVertexFactory.ush b/Engine/Shaders/Private/HairStrands/HairStrandsVertexFactory.ush
--- a/Engine/Shaders/Private/HairStrands/HairStrandsVertexFactory.ush
+++ b/Engine/Shaders/Private/HairStrands/HairStrandsVertexFactory.ush
@@ -208,9 +208,9 @@ FHairControlPoint GetVertexPosition(const FVertexInfo VertexInfo, bool bInvalidJ
 {
 	const float3 PositionOffset = GetPositionOffset();
 
-	const uint4 PackedData = HairStrandsVF.PositionBuffer[VertexInfo.VertexIndex];
-	FHairControlPoint Out = UnpackHairControlPoint(
-		PackedData, 
+	FHairControlPoint Out = ReadHairControlPoint(
+		HairStrandsVF.PositionBuffer,
+		VertexInfo.VertexIndex, 
 		PositionOffset,
 #if RAYHITGROUPSHADER
 		HairStrandsVF.Radius,
@@ -250,9 +250,9 @@ float3 GetVertexPreviousPosition(FVertexFactoryInput Input, bool bInvalidJointVe
 	const float3 PreviousPositionOffset = GetPrevPositionOffset();
 
 	FVertexInfo VertexInfo = GetVertexInfo(Input);
-	const uint4 PackedData = HairStrandsVF.PreviousPositionBuffer[VertexInfo.VertexIndex];
-	FHairControlPoint Out = UnpackHairControlPoint(
-		PackedData,
+	FHairControlPoint Out = ReadHairControlPoint(
+		HairStrandsVF.PreviousPositionBuffer,
+		VertexInfo.VertexIndex,
 		PreviousPositionOffset,
 #if RAYHITGROUPSHADER
 		HairStrandsVF.Radius,
@@ -418,14 +418,16 @@ RAY_TRACING_ENTRY_INTERSECTION(MaterialIS)
 
 	const float EffectiveRadius = HairStrandsVF.Radius * HairStrandsVF.RaytracingRadiusScale;
 	const float3 PositionOffset = GetPositionOffset();
-	FHairControlPoint P0 = UnpackHairControlPoint(
-		HairStrandsVF.PositionBuffer[VertexIndex0],
+	FHairControlPoint P0 = ReadHairControlPoint(
+		HairStrandsVF.PositionBuffer,
+		VertexIndex0,
 		PositionOffset,
 		EffectiveRadius * ScaleR0,
 		HairStrandsVF.RootScale,
 		HairStrandsVF.TipScale);
-	FHairControlPoint P1 = UnpackHairControlPoint(
-		HairStrandsVF.PositionBuffer[VertexIndex1],
+	FHairControlPoint P1 = ReadHairControlPoint(
+		HairStrandsVF.PositionBuffer,
+		VertexIndex1,
 		PositionOffset,
 		EffectiveRadius * ScaleR1,
 		HairStrandsVF.RootScale,
@@ -502,14 +504,16 @@ FMaterialPixelParameters GetMaterialPixelParameters(float3 RayOrigin, float3 Ray
 	float UCoord = Barycentrics.x;
 	float VCoord = Barycentrics.y;
 	const float3 PositionOffset = GetPositionOffset();
-	FHairControlPoint P0 = UnpackHairControlPoint(
-		HairStrandsVF.PositionBuffer[VertexIndex0],
+	FHairControlPoint P0 = ReadHairControlPoint(
+		HairStrandsVF.PositionBuffer,
+		VertexIndex0,
 		PositionOffset,
 		EffectiveRadius,
 		HairStrandsVF.RootScale,
 		HairStrandsVF.TipScale);
-	FHairControlPoint P1 = UnpackHairControlPoint(
-		HairStrandsVF.PositionBuffer[VertexIndex1],
+	FHairControlPoint P1 = ReadHairControlPoint(
+		HairStrandsVF.PositionBuffer,
+		VertexIndex1,
 		PositionOffset,
 		EffectiveRadius,
 		HairStrandsVF.RootScale,
@@ -555,14 +559,16 @@ FMaterialPixelParameters GetMaterialPixelParameters(float3 RayOrigin, float3 Ray
 		// find central point of the hair strand
 		const float3 PositionOffset = GetPositionOffset();
 		const float EffectiveRadius = HairStrandsVF.Radius * HairStrandsVF.RaytracingRadiusScale;
-		FHairControlPoint P0 = UnpackHairControlPoint(
-			HairStrandsVF.PositionBuffer[VertexIndex0],
+		FHairControlPoint P0 = ReadHairControlPoint(
+			HairStrandsVF.PositionBuffer,
+			VertexIndex0,
 			PositionOffset,
 			EffectiveRadius,
 			HairStrandsVF.RootScale,
 			HairStrandsVF.TipScale);
-		FHairControlPoint P1 = UnpackHairControlPoint(
-			HairStrandsVF.PositionBuffer[VertexIndex1],
+		FHairControlPoint P1 = ReadHairControlPoint(
+			HairStrandsVF.PositionBuffer,
+			VertexIndex1,
 			PositionOffset,
 			EffectiveRadius,
 			HairStrandsVF.RootScale,
diff --git a/Engine/Shaders/Private/HairStrands/HairStrandsVisibilityRasterClassification.usf b/Engine/Shaders/Private/HairStrands/HairStrandsVisibilityRasterClassification.usf
--- a/Engine/Shaders/Private/HairStrands/HairStrandsVisibilityRasterClassification.usf
+++ b/Engine/Shaders/Private/HairStrands/HairStrandsVisibilityRasterClassification.usf
@@ -19,8 +19,9 @@ uint				NumWorkGroups;
 
 bool ShouldSoftwareRasterize(float3 PositionOffset, uint PrimID, out bool bSegValid)
 {
-	FHairControlPoint CP0 = UnpackHairControlPoint(
-		HairStrandsVF_PositionBuffer[PrimID],
+	FHairControlPoint CP0 = ReadHairControlPoint(
+		HairStrandsVF_PositionBuffer,
+		PrimID,
 		PositionOffset,
 		HairStrandsVF_Radius,
 		HairStrandsVF_RootScale,
@@ -33,8 +34,9 @@ bool ShouldSoftwareRasterize(float3 PositionOffset, uint PrimID, out bool bSegVa
 		return false;
 	}
 
-	FHairControlPoint CP1 = UnpackHairControlPoint(
-		HairStrandsVF_PositionBuffer[PrimID + 1],
+	FHairControlPoint CP1 = ReadHairControlPoint(
+		HairStrandsVF_PositionBuffer,
+		PrimID + 1,
 		PositionOffset,
 		HairStrandsVF_Radius,
 		HairStrandsVF_RootScale,
diff --git a/Engine/Shaders/Private/HairStrands/HairStrandsVisibilityRasterCompute.usf b/Engine/Shaders/Private/HairStrands/HairStrandsVisibilityRasterCompute.usf
--- a/Engine/Shaders/Private/HairStrands/HairStrandsVisibilityRasterCompute.usf
+++ b/Engine/Shaders/Private/HairStrands/HairStrandsVisibilityRasterCompute.usf
@@ -82,10 +82,11 @@ float3 NDCToPixelCoord(float4 InDC)
 	return float3(UV * OutputResolution, NDC.z);
 }
 
-void CalcHomogenousPos(in uint4 PB, in float3 PBO, out float4 HP, out uint Type)
+void CalcHomogenousPos(in uint InPointIndex, in float3 PBO, out float4 HP, out uint Type)
 {
-	const FHairControlPoint CP = UnpackHairControlPoint(
-		PB,
+	const FHairControlPoint CP = ReadHairControlPoint(
+		HairStrandsVF_PositionBuffer,
+		InPointIndex,
 		PBO,
 		HairStrandsVF_Radius,
 		HairStrandsVF_RootScale,
@@ -96,10 +97,11 @@ void CalcHomogenousPos(in uint4 PB, in float3 PBO, out float4 HP, out uint Type)
 	Type = CP.Type;
 }
 
-void CalcHomogenousPosAndRad(in uint4 PB, in float3 PBO, out float4 HP, out float Rad, out uint Type)
+void CalcHomogenousPosAndRad(in uint InPointIndex, in float3 PBO, out float4 HP, out float Rad, out uint Type)
 {
-	const FHairControlPoint CP = UnpackHairControlPoint(
-		PB,
+	const FHairControlPoint CP = ReadHairControlPoint(
+		HairStrandsVF_PositionBuffer,
+		InPointIndex,
 		PBO,
 		HairStrandsVF_Radius,
 		HairStrandsVF_RootScale,
@@ -471,14 +473,14 @@ void BinningCS(uint DispatchThreadID : SV_DispatchThreadID, uint GroupThreadID :
 			float4 H0 = 0.0f;
 			float4 H1 = 0.0f;
 			uint Type = -1;
-			CalcHomogenousPos(HairStrandsVF_PositionBuffer[PrimID], InstancePositionOffset, H0, Type);
+			CalcHomogenousPos(PrimID, InstancePositionOffset, H0, Type);
 
 			bool bIsEndCV = (Type == HAIR_CONTROLPOINT_END);
 			bSegValid = !bIsEndCV;
 
 			if (bSegValid)
 			{
-				CalcHomogenousPos(HairStrandsVF_PositionBuffer[PrimID + 1], InstancePositionOffset, H1, Type);
+				CalcHomogenousPos(PrimID + 1, InstancePositionOffset, H1, Type);
 				
 				// Do clipping in homogenous coordinates
 				bSegValid = BlinnLineClipping(H0, H1);
diff --git a/Engine/Shaders/Private/HairStrands/HairStrandsVisibilityRasterComputeNaive.usf b/Engine/Shaders/Private/HairStrands/HairStrandsVisibilityRasterComputeNaive.usf
--- a/Engine/Shaders/Private/HairStrands/HairStrandsVisibilityRasterComputeNaive.usf
+++ b/Engine/Shaders/Private/HairStrands/HairStrandsVisibilityRasterComputeNaive.usf
@@ -182,8 +182,9 @@ void CSMain(uint DispatchThreadID : SV_DispatchThreadID, uint GroupThreadID : SV
 		}
 
 		// Fetch both control points
-		FHairControlPoint CP0 = UnpackHairControlPoint(
-			HairStrandsVF_PositionBuffer[PrimID],
+		FHairControlPoint CP0 = ReadHairControlPoint(
+			HairStrandsVF_PositionBuffer,
+			PrimID,
 			PositionOffset,
 			HairStrandsVF_Radius,
 			HairStrandsVF_RootScale,
@@ -194,8 +195,9 @@ void CSMain(uint DispatchThreadID : SV_DispatchThreadID, uint GroupThreadID : SV
 			continue;
 		}
 
-		const FHairControlPoint CP1 = UnpackHairControlPoint(
-			HairStrandsVF_PositionBuffer[PrimID + 1],
+		const FHairControlPoint CP1 = ReadHairControlPoint(
+			HairStrandsVF_PositionBuffer,
+			PrimID + 1,
 			PositionOffset,
 			HairStrandsVF_Radius,
 			HairStrandsVF_RootScale,
diff --git a/Engine/Shaders/Private/HairStrands/HairStrandsVisibilityRasterHW.usf b/Engine/Shaders/Private/HairStrands/HairStrandsVisibilityRasterHW.usf
--- a/Engine/Shaders/Private/HairStrands/HairStrandsVisibilityRasterHW.usf
+++ b/Engine/Shaders/Private/HairStrands/HairStrandsVisibilityRasterHW.usf
@@ -85,8 +85,9 @@ VSOutput VSMain(VSInput Input)
 	const float3 PositionOffset = HairStrandsVF_GetHairInstancePositionOffset();
 
 	// Fetch both control points
-	const FHairControlPoint CP0 = UnpackHairControlPoint(
-		HairStrandsVF_PositionBuffer[PrimID],
+	const FHairControlPoint CP0 = ReadHairControlPoint(
+		HairStrandsVF_PositionBuffer,
+		PrimID,
 		PositionOffset,
 		HairStrandsVF_Radius,
 		HairStrandsVF_RootScale,
@@ -98,8 +99,9 @@ VSOutput VSMain(VSInput Input)
 		return GetNullOutput();
 	}
 
-	const FHairControlPoint CP1 = UnpackHairControlPoint(
-		HairStrandsVF_PositionBuffer[PrimID + 1],
+	const FHairControlPoint CP1 = ReadHairControlPoint(
+		HairStrandsVF_PositionBuffer,
+		PrimID + 1,
 		PositionOffset,
 		HairStrandsVF_Radius,
 		HairStrandsVF_RootScale,
diff --git a/Engine/Shaders/Private/HairStrands/HairStrandsVoxelRasterCompute.usf b/Engine/Shaders/Private/HairStrands/HairStrandsVoxelRasterCompute.usf
--- a/Engine/Shaders/Private/HairStrands/HairStrandsVoxelRasterCompute.usf
+++ b/Engine/Shaders/Private/HairStrands/HairStrandsVoxelRasterCompute.usf
@@ -708,8 +708,9 @@ void MainCS(uint2 DispatchThreadID : SV_DispatchThreadID)
 
 	const float3 PositionOffset = HairStrandsVF_GetHairInstancePositionOffset();
 
-	const FHairControlPoint CP0 = UnpackHairControlPoint(
-		HairStrandsVF_PositionBuffer[VertexIndex0],
+	const FHairControlPoint CP0 = ReadHairControlPoint(
+		HairStrandsVF_PositionBuffer,
+		VertexIndex0,
 		PositionOffset,
 		HairStrandsVF_Radius * LodRadiusScale0,
 		HairStrandsVF_RootScale,
@@ -718,8 +719,9 @@ void MainCS(uint2 DispatchThreadID : SV_DispatchThreadID)
 	if (CP0.Type == HAIR_CONTROLPOINT_END)
 		return;
 
-	const FHairControlPoint CP1 = UnpackHairControlPoint(
-		HairStrandsVF_PositionBuffer[VertexIndex1],
+	const FHairControlPoint CP1 = ReadHairControlPoint(
+		HairStrandsVF_PositionBuffer,
+		VertexIndex1,
 		PositionOffset,
 		HairStrandsVF_Radius * LodRadiusScale1,
 		HairStrandsVF_RootScale,
