Fix shader compilation warning.

diff --git a/Engine/Shaders/Private/HairStrands/HairStrandsPack.ush b/Engine/Shaders/Private/HairStrands/HairStrandsPack.ush
--- a/Engine/Shaders/Private/HairStrands/HairStrandsPack.ush
+++ b/Engine/Shaders/Private/HairStrands/HairStrandsPack.ush
@@ -92,6 +92,11 @@ float3 ReadHairControlPointPosition(Buffer<FPackedHairPosition> InBuffer, uint I
 	return UnpackHairControlPointPosition(InBuffer[InPointIndex], InVF_PositionOffset);
 }
 
+float3 ReadHairControlPointPosition(RWBuffer<FPackedHairPosition> InBuffer, uint InPointIndex, float3 InVF_PositionOffset=0)
+{
+	return UnpackHairControlPointPosition(InBuffer[InPointIndex], InVF_PositionOffset);
+}
+
 FPackedHairPosition PackHairControlPointPosition(FPackedHairPosition InPackedData, float3 InNewPosition, float3 InPositionOffset=0)
 {
 	return FPackedHairPosition(f32tof16(InNewPosition.xyz - InPositionOffset), InPackedData.w);
