Fix hair interpolation being incorrect on the last curves of a groom due to incorrect culling.

diff --git a/Engine/Shaders/Private/HairStrands/HairStrandsInterpolation.usf b/Engine/Shaders/Private/HairStrands/HairStrandsInterpolation.usf
--- a/Engine/Shaders/Private/HairStrands/HairStrandsInterpolation.usf
+++ b/Engine/Shaders/Private/HairStrands/HairStrandsInterpolation.usf
@@ -396,7 +396,6 @@ float3 ApplyDeformer(const FRootTriangleData RenTri, float3 InRestPostion, float
 // Interpolation
 
 uint CurveCount;
-uint CurvePerGroup;
 
 Buffer<uint> SimCurveBuffer;
 Buffer<uint> RenCurveBuffer;
@@ -487,7 +486,10 @@ void MainCS(uint3 DispatchThreadId : SV_DispatchThreadID, uint ThreadIndex : SV_
 		const uint LocalCurveIndex   = ThreadIndex / uint(MAX_POINT);
 		const uint CurrentPointIndex = ThreadIndex % uint(MAX_POINT);
 
-		const uint GlobalCurveIndex = GetGlobalCurveIndex(ThreadIndex, GroupIndex);
+		// Use LocalCurveIndex to compute the global curve index.
+		// Here 1thread == 1control point, while in the previous block it was 1thread == 1curve. 
+		// This is to remaps correctly the curve index.
+		const uint GlobalCurveIndex = GetGlobalCurveIndex(LocalCurveIndex, GroupIndex); 
 		if (GlobalCurveIndex < CurveCount)
 		{		
 			if (CurrentPointIndex < CurvePointCount[LocalCurveIndex])
