Change groom position buffer to a byteaddressbuffer.

diff --git a/Engine/Plugins/Runtime/HairStrands/Shaders/Private/DeformerDataInterfaceGroomWrite.ush b/Engine/Plugins/Runtime/HairStrands/Shaders/Private/DeformerDataInterfaceGroomWrite.ush
--- a/Engine/Plugins/Runtime/HairStrands/Shaders/Private/DeformerDataInterfaceGroomWrite.ush
+++ b/Engine/Plugins/Runtime/HairStrands/Shaders/Private/DeformerDataInterfaceGroomWrite.ush
@@ -9,8 +9,8 @@ HAIR_STRANDS_INSTANCE_PARAMETERS({DataInterfaceName})
 // beause we write out data which would be bound as SRV otherwise (position, attributes, ...)
 uint 				{DataInterfaceName}_OutputStreamStart;
 StructuredBuffer<float4> {DataInterfaceName}_PositionOffsetBufferSRV;
-Buffer<uint4>   		{DataInterfaceName}_PositionBufferSRV;
-RWBuffer<uint4> 		{DataInterfaceName}_PositionBufferUAV;
+ByteAddressBuffer		{DataInterfaceName}_PositionBufferSRV;
+RWByteAddressBuffer		{DataInterfaceName}_PositionBufferUAV;
 RWByteAddressBuffer 	{DataInterfaceName}_PointAttributeBufferUAV;
 RWByteAddressBuffer 	{DataInterfaceName}_CurveAttributeBufferUAV;
 
@@ -34,9 +34,8 @@ void WritePosition_{DataInterfaceName}(uint ControlPointIndex, float3 Position)
 	const uint WriteIndex = {DataInterfaceName}_OutputStreamStart + ControlPointIndex;
 	if (WriteIndex < {DataInterfaceName}_PointCount)
 	{
-		const uint4 PackedCP = {DataInterfaceName}_PositionBufferSRV[WriteIndex];
-		const uint4 NewPackedCP = PackHairControlPointPosition(PackedCP, Position, {DataInterfaceName}_GetPositionOffsetBufferSRV());
-		{DataInterfaceName}_PositionBufferUAV[WriteIndex] = NewPackedCP;
+		const FPackedHairPosition PackedCP = ReadPackedHairPosition({DataInterfaceName}_PositionBufferSRV, WriteIndex);
+		WritePackedHairControlPointPosition({DataInterfaceName}_PositionBufferUAV, WriteIndex, PackedCP, Position, {DataInterfaceName}_GetPositionOffsetBufferSRV());
 	}
 }
 
@@ -45,9 +44,9 @@ void WriteRadius_{DataInterfaceName}(uint ControlPointIndex, float Radius)
 	const uint WriteIndex = {DataInterfaceName}_OutputStreamStart + ControlPointIndex;
 	if (WriteIndex < {DataInterfaceName}_PointCount)
 	{
-		const uint4 PackedCP = {DataInterfaceName}_PositionBufferSRV[WriteIndex];
-		const uint4 NewPackedCP = PackHairControlPointRadius(PackedCP, Radius, {DataInterfaceName}_Radius);
-		{DataInterfaceName}_PositionBufferUAV[WriteIndex] = NewPackedCP;
+		const FPackedHairPosition PackedCP = ReadPackedHairPosition({DataInterfaceName}_PositionBufferSRV, WriteIndex);
+		const FPackedHairPosition NewPackedCP = PackHairControlPointRadius(PackedCP, Radius, {DataInterfaceName}_Radius);
+		WritePackedHairPosition({DataInterfaceName}_PositionBufferUAV, WriteIndex, NewPackedCP);
 	}
 }
 
@@ -56,11 +55,11 @@ void WritePositionAndRadius_{DataInterfaceName}(uint ControlPointIndex, float4 P
 	const uint WriteIndex = {DataInterfaceName}_OutputStreamStart + ControlPointIndex;
 	if (WriteIndex < {DataInterfaceName}_PointCount)
 	{
-		const uint4 PackedCP = {DataInterfaceName}_PositionBufferSRV[WriteIndex];
-		uint4 NewPackedCP = 0;
+		const FPackedHairPosition PackedCP = ReadPackedHairPosition({DataInterfaceName}_PositionBufferSRV, WriteIndex);
+		FPackedHairPosition NewPackedCP = 0;
 		NewPackedCP = PackHairControlPointPosition(PackedCP, PositionAndRadius.xyz, {DataInterfaceName}_GetPositionOffsetBufferSRV());
 		NewPackedCP = PackHairControlPointRadius(NewPackedCP, PositionAndRadius.w, {DataInterfaceName}_Radius);
-		{DataInterfaceName}_PositionBufferUAV[WriteIndex] = NewPackedCP;
+		WritePackedHairPosition({DataInterfaceName}_PositionBufferUAV, WriteIndex, NewPackedCP);
 	}
 }
 
diff --git a/Engine/Plugins/Runtime/HairStrands/Shaders/Private/NiagaraDataInterfaceHairStrandsTemplate.ush b/Engine/Plugins/Runtime/HairStrands/Shaders/Private/NiagaraDataInterfaceHairStrandsTemplate.ush
--- a/Engine/Plugins/Runtime/HairStrands/Shaders/Private/NiagaraDataInterfaceHairStrandsTemplate.ush
+++ b/Engine/Plugins/Runtime/HairStrands/Shaders/Private/NiagaraDataInterfaceHairStrandsTemplate.ush
@@ -5,8 +5,8 @@ NiagaraDataInterfaceHairStrands.ush
 =============================================================================*/
 #include "/Engine/Private/HairStrands/HairStrandsBindingCommon.ush"
 
-RWBuffer<uint4>		{ParameterName}_DeformedPositionBuffer;
-Buffer<uint4>		{ParameterName}_RestPositionBuffer;
+RWByteAddressBuffer	{ParameterName}_DeformedPositionBuffer;
+ByteAddressBuffer	{ParameterName}_RestPositionBuffer;
 Buffer<uint>		{ParameterName}_CurvesOffsetsBuffer;
 Buffer<float4>		{ParameterName}_RestTrianglePositionBuffer;
 Buffer<float4>		{ParameterName}_DeformedTrianglePositionBuffer;
diff --git a/Engine/Plugins/Runtime/HairStrands/Shaders/Private/NiagaraInterpolateGroomGuides.usf b/Engine/Plugins/Runtime/HairStrands/Shaders/Private/NiagaraInterpolateGroomGuides.usf
--- a/Engine/Plugins/Runtime/HairStrands/Shaders/Private/NiagaraInterpolateGroomGuides.usf
+++ b/Engine/Plugins/Runtime/HairStrands/Shaders/Private/NiagaraInterpolateGroomGuides.usf
@@ -24,10 +24,10 @@ Buffer<uint>		RootToUniqueTriangleIndexBuffer;
 #endif
 
 Buffer<uint>		PointToCurveIndexBuffer;
-RWBuffer<uint4>		DeformedPositionBuffer;
-Buffer<uint4>		RestPositionBuffer;
+RWByteAddressBuffer	DeformedPositionBuffer;
+ByteAddressBuffer	RestPositionBuffer;
 Buffer<uint>		CurvesOffsetsBuffer;
-Buffer<float4>		DeformedPositionOffset;
+StructuredBuffer<float4> DeformedPositionOffset;
 float3				RestPositionOffset;
 float4x4			WorldToLocal;
 
diff --git a/Engine/Plugins/Runtime/HairStrands/Source/HairStrandsCore/Private/Deformer/DeformerDataInterfaceGroom.cpp b/Engine/Plugins/Runtime/HairStrands/Source/HairStrandsCore/Private/Deformer/DeformerDataInterfaceGroom.cpp
--- a/Engine/Plugins/Runtime/HairStrands/Source/HairStrandsCore/Private/Deformer/DeformerDataInterfaceGroom.cpp
+++ b/Engine/Plugins/Runtime/HairStrands/Source/HairStrandsCore/Private/Deformer/DeformerDataInterfaceGroom.cpp
@@ -232,7 +232,7 @@ void FOptimusGroomDataProviderProxy::AllocateResources(FRDGBuilder& GraphBuilder
 
 			if (!FallbackSRV)
 			{
-				FallbackSRV = GraphBuilder.CreateSRV(GraphBuilder.RegisterExternalBuffer(GWhiteVertexBufferWithRDG->Buffer), PF_R16G16B16A16_UINT);
+				FallbackSRV = GraphBuilder.CreateSRV(GraphBuilder.RegisterExternalBuffer(GWhiteVertexBufferWithRDG->Buffer), PF_R16G16B16A16_UINT); // This should be a ByteAddressBuffer
 			}
 		}
 	}
diff --git a/Engine/Plugins/Runtime/HairStrands/Source/HairStrandsCore/Private/Deformer/DeformerDataInterfaceGroomGuide.cpp b/Engine/Plugins/Runtime/HairStrands/Source/HairStrandsCore/Private/Deformer/DeformerDataInterfaceGroomGuide.cpp
--- a/Engine/Plugins/Runtime/HairStrands/Source/HairStrandsCore/Private/Deformer/DeformerDataInterfaceGroomGuide.cpp
+++ b/Engine/Plugins/Runtime/HairStrands/Source/HairStrandsCore/Private/Deformer/DeformerDataInterfaceGroomGuide.cpp
@@ -164,7 +164,7 @@ void FOptimusGroomGuideDataProviderProxy::AllocateResources(FRDGBuilder& GraphBu
 
 			if (!FallbackSRV)
 			{
-				FallbackSRV = GraphBuilder.CreateSRV(GraphBuilder.RegisterExternalBuffer(GWhiteVertexBufferWithRDG->Buffer), PF_R16G16B16A16_UINT);
+				FallbackSRV = GraphBuilder.CreateSRV(GraphBuilder.RegisterExternalBuffer(GWhiteVertexBufferWithRDG->Buffer), PF_R16G16B16A16_UINT); // This should be a byte address buffer
 			}
 		}
 	}
diff --git a/Engine/Plugins/Runtime/HairStrands/Source/HairStrandsCore/Private/Deformer/DeformerDataInterfaceGroomWrite.cpp b/Engine/Plugins/Runtime/HairStrands/Source/HairStrandsCore/Private/Deformer/DeformerDataInterfaceGroomWrite.cpp
--- a/Engine/Plugins/Runtime/HairStrands/Source/HairStrandsCore/Private/Deformer/DeformerDataInterfaceGroomWrite.cpp
+++ b/Engine/Plugins/Runtime/HairStrands/Source/HairStrandsCore/Private/Deformer/DeformerDataInterfaceGroomWrite.cpp
@@ -114,8 +114,8 @@ BEGIN_SHADER_PARAMETER_STRUCT(FGroomWriteDataInterfaceParameters, )
 	SHADER_PARAMETER_STRUCT_INCLUDE(FHairStrandsInstanceCommonParameters, Common)
 	SHADER_PARAMETER(uint32, OutputStreamStart)
 	SHADER_PARAMETER_RDG_BUFFER_SRV(StructuredBuffer, PositionOffsetBufferSRV)
-	SHADER_PARAMETER_RDG_BUFFER_SRV(Buffer<uint4>, PositionBufferSRV)
-	SHADER_PARAMETER_RDG_BUFFER_UAV(RWBuffer<uint4>, PositionBufferUAV)
+	SHADER_PARAMETER_RDG_BUFFER_SRV(ByteAddressBuffer, PositionBufferSRV)
+	SHADER_PARAMETER_RDG_BUFFER_UAV(RWByteAddressBuffer, PositionBufferUAV)
 	SHADER_PARAMETER_RDG_BUFFER_UAV(RWByteAddressBuffer, CurveAttributeBufferUAV)
 	SHADER_PARAMETER_RDG_BUFFER_UAV(RWByteAddressBuffer, PointAttributeBufferUAV)
 END_SHADER_PARAMETER_STRUCT()
@@ -201,8 +201,8 @@ void FOptimusGroomWriteDataProviderProxy::AllocateResources(FRDGBuilder& GraphBu
 			FResources& R = Resources.AddDefaulted_GetRef();
 
 			FRDGBufferRef PositionBuffer_fallback = GraphBuilder.CreateBuffer(FRDGBufferDesc::CreateBufferDesc(8, 1), TEXT("Groom.DeformedPositionBuffer"), ERDGBufferFlags::None);
-			R.PositionBufferSRV_fallback = GraphBuilder.CreateSRV(GraphBuilder.RegisterExternalBuffer(GWhiteVertexBufferWithRDG->Buffer), PF_R16G16B16A16_UINT);
-			R.PositionBufferUAV_fallback = GraphBuilder.CreateUAV(PositionBuffer_fallback, PF_R16G16B16A16_UINT, ERDGUnorderedAccessViewFlags::SkipBarrier);
+			R.PositionBufferSRV_fallback = GraphBuilder.CreateSRV(GraphBuilder.RegisterExternalBuffer(GWhiteVertexBufferWithRDG->Buffer));
+			R.PositionBufferUAV_fallback = GraphBuilder.CreateUAV(PositionBuffer_fallback, ERDGUnorderedAccessViewFlags::SkipBarrier);
 
 			FRDGBufferRef AttributeBuffer_fallback = GraphBuilder.CreateBuffer(FRDGBufferDesc::CreateByteAddressDesc(16), TEXT("Groom.DeformedAttributeBuffer"), ERDGBufferFlags::None);
 			R.AttributeBufferUAV_fallback = GraphBuilder.CreateUAV(AttributeBuffer_fallback, PF_Unknown, ERDGUnorderedAccessViewFlags::SkipBarrier);
@@ -219,8 +219,8 @@ void FOptimusGroomWriteDataProviderProxy::AllocateResources(FRDGBuilder& GraphBu
 				FRDGBufferRef DummyUAVBuffer = GraphBuilder.CreateBuffer(FRDGBufferDesc::CreateBufferDesc(8u, 1), TEXT("Hair.Deformer.DummyBuffer"));
 				
 				R.PositionOffsetBufferSRV 	= GraphBuilder.CreateSRV(GSystemTextures.GetDefaultStructuredBuffer(GraphBuilder, 16u), PF_A32B32G32R32F); //GraphBuilder.CreateSRV(GraphBuilder.RegisterExternalBuffer(GWhiteVertexBufferWithRDG->Buffer), PF_A32B32G32R32F);
-				R.PositionBufferSRV 		= GraphBuilder.CreateSRV(GraphBuilder.RegisterExternalBuffer(GWhiteVertexBufferWithRDG->Buffer), PF_R16G16B16A16_UINT);
-				R.PositionBufferUAV 		= GraphBuilder.CreateUAV(DummyUAVBuffer, PF_R16G16B16A16_UINT);
+				R.PositionBufferSRV 		= GraphBuilder.CreateSRV(GraphBuilder.RegisterExternalBuffer(GWhiteVertexBufferWithRDG->Buffer));
+				R.PositionBufferUAV 		= GraphBuilder.CreateUAV(DummyUAVBuffer);
 			}
 			
 			// Curve Attributes
diff --git a/Engine/Plugins/Runtime/HairStrands/Source/HairStrandsCore/Private/GroomBuilder.cpp b/Engine/Plugins/Runtime/HairStrands/Source/HairStrandsCore/Private/GroomBuilder.cpp
--- a/Engine/Plugins/Runtime/HairStrands/Source/HairStrandsCore/Private/GroomBuilder.cpp
+++ b/Engine/Plugins/Runtime/HairStrands/Source/HairStrandsCore/Private/GroomBuilder.cpp
@@ -40,7 +40,7 @@ static FAutoConsoleVariableRef CVarHairGroupIndexBuilder_MaxVoxelResolution(TEXT
 
 FString FGroomBuilder::GetVersion()
 {
-	return TEXT("v11b");
+	return TEXT("v12c");
 }
 
 // For debug purpose
@@ -167,15 +167,17 @@ namespace HairStrandsBuilder
 		CopyToBulkData<TFormatType>(Out.Data, Data);
 	}
 
-	/** Build the packed datas for gpu rendering/simulation */
-	void BuildRenderData(const FHairStrandsDatas& HairStrands, const TArray<uint8>& RandomSeeds, FHairStrandsBulkData& OutBulkData)
+	/** Build the bulk/packed datas for gpu rendering/simulation */
+	void BuildBulkData(const FHairStrandsDatas& HairStrands, const TArray<uint8>& RandomSeeds, FHairStrandsBulkData& OutBulkData)
 	{
-		TRACE_CPUPROFILER_EVENT_SCOPE(HairStrandsBuilder::BuildRenderData);
+		TRACE_CPUPROFILER_EVENT_SCOPE(HairStrandsBuilder::BuildBulkData);
 
 		const uint32 NumCurves = HairStrands.GetNumCurves();
 		const uint32 NumPoints = HairStrands.GetNumPoints();
 		if (!(NumCurves > 0 && NumPoints > 0))
+		{
 			return;
+		}
 
 		TArray<FHairStrandsPositionFormat::Type> OutPackedPositions;
 		TArray<FHairStrandsCurveFormat::Type> OutPackedCurves;
@@ -226,24 +228,13 @@ namespace HairStrandsBuilder
 		const float MaxLength = GetHairStrandsMaxLength(HairStrands);
 		const float MaxRadius = GetHairStrandsMaxRadius(HairStrands);
 
-		struct FPackedRadiusAndType
-		{
-			union
-			{
-				struct
-				{
-					uint8 ControlPointType : 2;
-					uint8 NormalizedRadius : 6;
-				} Data;
-				uint8 Packed;
-			};
-		};
+		static_assert(sizeof(FPackedHairVertex) == sizeof(FPackedHairVertex::BulkType));
+		static_assert(sizeof(FPackedHairVertex) == 8u);
 
 		uint32 MinPointPerCurve = HAIR_MAX_NUM_POINT_PER_CURVE;
 		uint32 MaxPointPerCurve = 0;
 		uint32 AccPointPerCurve = 0;
 
-		static_assert(sizeof(FPackedRadiusAndType) == sizeof(uint8));
 		const bool bSeedValid = RandomSeeds.Num() > 0;
 		for (uint32 CurveIndex = 0; CurveIndex < NumCurves; ++CurveIndex)
 		{
@@ -260,18 +251,16 @@ namespace HairStrandsBuilder
 				const uint32 NextIndex = FMath::Min(PointCount + 1, PointCount - 1);
 				const FVector3f& PointPosition = Points.PointsPosition[PointIndex + IndexOffset];
 
-				// Packed Position|CoordU|NormalizedRadius
+				// Packed Position|CoordU|NormalizedRadius|Type
 				{
 					const float CoordU = Points.PointsCoordU[PointIndex + IndexOffset];
 					const float NormalizedRadius = Points.PointsRadius[PointIndex + IndexOffset] / MaxRadius;
 
 					FHairStrandsPositionFormat::Type& PackedPosition = OutPackedPositions[PointIndex + IndexOffset];
 					CopyVectorToPosition(PointPosition - HairBoxCenter, PackedPosition);
-					FPackedRadiusAndType PackedRadiusAndType;
-					PackedRadiusAndType.Data.ControlPointType = (PointIndex == 0) ? HAIR_CONTROLPOINT_START : (PointIndex == (PointCount - 1) ? HAIR_CONTROLPOINT_END : HAIR_CONTROLPOINT_INSIDE);
-					PackedRadiusAndType.Data.NormalizedRadius = uint8(FMath::Clamp(NormalizedRadius * 63.f, 0.f, 63.f));
-					PackedPosition.PackedRadiusAndType = PackedRadiusAndType.Packed; //-V614
 					PackedPosition.UCoord = uint8(FMath::Clamp(CoordU * 255.f, 0.f, 255.f));
+					PackedPosition.Radius = uint8(FMath::Clamp(NormalizedRadius * 63.f, 0.f, 63.f));
+					PackedPosition.Type = (PointIndex == 0) ? HAIR_CONTROLPOINT_START : (PointIndex == (PointCount - 1) ? HAIR_CONTROLPOINT_END : HAIR_CONTROLPOINT_INSIDE);
 				}
 
 				// Point to Curve
@@ -617,10 +606,10 @@ namespace HairStrandsBuilder
 		OutBulkData.Header.Strides.PointAttributeChunkElementCount = PointAttributeChunkElementCount;
 	}
 
-	void BuildRenderData(const FHairStrandsDatas& HairStrands, FHairStrandsBulkData& OutBulkData)
+	void BuildBulkData(const FHairStrandsDatas& HairStrands, FHairStrandsBulkData& OutBulkData)
 	{
 		TArray<uint8> RandomSeeds;
-		BuildRenderData(HairStrands, RandomSeeds, OutBulkData);
+		BuildBulkData(HairStrands, RandomSeeds, OutBulkData);
 	}
 
 } // namespace HairStrandsBuilder
@@ -2230,7 +2219,7 @@ void FGroomBuilder::BuildBulkData(
 	{
 		CurveSeeds[Index] = Random.RandHelper(255);
 	}
-	HairStrandsBuilder::BuildRenderData(InData, CurveSeeds, OutBulkData);
+	HairStrandsBuilder::BuildBulkData(InData, CurveSeeds, OutBulkData);
 }
 
 void FGroomBuilder::BuildInterplationData(
diff --git a/Engine/Plugins/Runtime/HairStrands/Source/HairStrandsCore/Private/GroomDebug.cpp b/Engine/Plugins/Runtime/HairStrands/Source/HairStrandsCore/Private/GroomDebug.cpp
--- a/Engine/Plugins/Runtime/HairStrands/Source/HairStrandsCore/Private/GroomDebug.cpp
+++ b/Engine/Plugins/Runtime/HairStrands/Source/HairStrandsCore/Private/GroomDebug.cpp
@@ -523,11 +523,11 @@ class FDrawDebugCardGuidesCS : public FGlobalShader
 		SHADER_PARAMETER(FVector3f, SimRestOffset)
 		SHADER_PARAMETER_RDG_BUFFER_SRV(StructuredBuffer, SimDeformedOffset)
 
-		SHADER_PARAMETER_RDG_BUFFER_SRV(Buffer, RenRestPosition)
-		SHADER_PARAMETER_RDG_BUFFER_SRV(Buffer, RenDeformedPosition)
+		SHADER_PARAMETER_RDG_BUFFER_SRV(ByteAddressBuffer, RenRestPosition)
+		SHADER_PARAMETER_RDG_BUFFER_SRV(ByteAddressBuffer, RenDeformedPosition)
 
-		SHADER_PARAMETER_RDG_BUFFER_SRV(Buffer, SimRestPosition)
-		SHADER_PARAMETER_RDG_BUFFER_SRV(Buffer, SimDeformedPosition)
+		SHADER_PARAMETER_RDG_BUFFER_SRV(ByteAddressBuffer, SimRestPosition)
+		SHADER_PARAMETER_RDG_BUFFER_SRV(ByteAddressBuffer, SimDeformedPosition)
 
 		SHADER_PARAMETER_STRUCT_INCLUDE(ShaderPrint::FShaderParameters, ShaderPrintParameters)
 	END_SHADER_PARAMETER_STRUCT()
@@ -599,6 +599,7 @@ static void AddDrawDebugCardsGuidesPass(
 	if (!bRen && !bGuideValid)						{ return; }
 	if (!bRen && bDeformed && !bGuideDeformValid)	{ return; }
 
+	FRDGBufferSRVRef DefaultByteAddreeBuffer = GraphBuilder.CreateSRV(GSystemTextures.GetDefaultByteAddressBuffer(GraphBuilder, 8u));
 	FRDGBufferSRVRef DefaultBuffer = GraphBuilder.CreateSRV(GSystemTextures.GetDefaultBuffer(GraphBuilder, 8, 0u), PF_R16G16B16A16_UINT);
 
 	FDrawDebugCardGuidesCS::FParameters* Parameters = GraphBuilder.AllocParameters<FDrawDebugCardGuidesCS::FParameters>();
@@ -608,15 +609,15 @@ static void AddDrawDebugCardsGuidesPass(
 	Parameters->CardLODIndex = HairLODIndex;
 	Parameters->RenVertexCount = 0;
 	Parameters->RenRestOffset = FVector3f::ZeroVector;
-	Parameters->RenRestPosition = DefaultBuffer;
+	Parameters->RenRestPosition = DefaultByteAddreeBuffer;
 	Parameters->RenDeformedOffset = DefaultBuffer;
-	Parameters->RenDeformedPosition = DefaultBuffer;
+	Parameters->RenDeformedPosition = DefaultByteAddreeBuffer;
 
 	Parameters->SimVertexCount = 0;
 	Parameters->SimRestOffset = FVector3f::ZeroVector;
-	Parameters->SimRestPosition = DefaultBuffer;
+	Parameters->SimRestPosition = DefaultByteAddreeBuffer;
 	Parameters->SimDeformedOffset = DefaultBuffer;
-	Parameters->SimDeformedPosition = DefaultBuffer;
+	Parameters->SimDeformedPosition = DefaultByteAddreeBuffer;
 
 	if (bRen)
 	{
diff --git a/Engine/Plugins/Runtime/HairStrands/Source/HairStrandsCore/Private/GroomResources.cpp b/Engine/Plugins/Runtime/HairStrands/Source/HairStrandsCore/Private/GroomResources.cpp
--- a/Engine/Plugins/Runtime/HairStrands/Source/HairStrandsCore/Private/GroomResources.cpp
+++ b/Engine/Plugins/Runtime/HairStrands/Source/HairStrandsCore/Private/GroomResources.cpp
@@ -1040,7 +1040,7 @@ void FHairStrandsRestResource::InternalAllocate(FRDGBuilder& GraphBuilder)
 
 	// 1. Lock data, which force the loading data from files (on non-editor build/cooked data). These data are then uploaded to the GPU
 	// 2. A local copy is done by the buffer uploader. This copy is discarded once the uploading is done.
-	InternalCreateVertexBufferRDG_FromHairBulkData<FHairStrandsPositionFormat>(GraphBuilder, BulkData.Data.Positions, PointCount, PositionBuffer, ToHairResourceDebugName(HAIRSTRANDS_RESOUCE_NAME(CurveType, Hair.StrandsRest_PositionBuffer), ResourceName), OwnerName, EHairResourceUsageType::Static);
+	InternalCreateByteAddressBufferRDG_FromHairBulkData(GraphBuilder, BulkData.Data.Positions, PositionBuffer, ToHairResourceDebugName(HAIRSTRANDS_RESOUCE_NAME(CurveType, Hair.StrandsRest_PositionBuffer), ResourceName), OwnerName, EHairResourceUsageType::Static);
 	InternalCreateByteAddressBufferRDG_FromHairBulkData(GraphBuilder, BulkData.Data.CurveAttributes, CurveAttributeBuffer, ToHairResourceDebugName(HAIRSTRANDS_RESOUCE_NAME(CurveType, Hair.StrandsRest_CurveAttributeBuffer), ResourceName), OwnerName, EHairResourceUsageType::Static);
 	if (!!(BulkData.Header.Flags & FHairStrandsBulkData::DataFlags_HasPointAttribute))
 	{
@@ -1149,8 +1149,8 @@ void FHairStrandsDeformedResource::InternalAllocate(FRDGBuilder& GraphBuilder, u
 	FRDGBufferRef Prev0 = Register(GraphBuilder, DeformedPositionBuffer[0], ERDGImportedBufferFlags::None).Buffer;
 	FRDGBufferRef Prev1 = Register(GraphBuilder, DeformedPositionBuffer[1], ERDGImportedBufferFlags::None).Buffer;
 
-	FRDGBufferRef Deformed0 = InternalCreateVertexBufferRDG<FHairStrandsPositionFormat>(GraphBuilder, PointCount, DeformedPositionBuffer[0], ToHairResourceDebugName(HAIRSTRANDS_RESOUCE_NAME(CurveType, Hair.StrandsDeformed_DeformedPositionBuffer0), ResourceName), OwnerName, EHairResourceUsageType::Dynamic);
-	FRDGBufferRef Deformed1 = InternalCreateVertexBufferRDG<FHairStrandsPositionFormat>(GraphBuilder, PointCount, DeformedPositionBuffer[1], ToHairResourceDebugName(HAIRSTRANDS_RESOUCE_NAME(CurveType, Hair.StrandsDeformed_DeformedPositionBuffer1), ResourceName), OwnerName, EHairResourceUsageType::Dynamic);
+	FRDGBufferRef Deformed0 = InternalCreateByteAddressBufferRDG(GraphBuilder, PointCount * BulkData.Header.Strides.PositionStride, DeformedPositionBuffer[0], ToHairResourceDebugName(HAIRSTRANDS_RESOUCE_NAME(CurveType, Hair.StrandsDeformed_DeformedPositionBuffer0), ResourceName), OwnerName, EHairResourceUsageType::Dynamic);
+	FRDGBufferRef Deformed1 = InternalCreateByteAddressBufferRDG(GraphBuilder, PointCount * BulkData.Header.Strides.PositionStride, DeformedPositionBuffer[1], ToHairResourceDebugName(HAIRSTRANDS_RESOUCE_NAME(CurveType, Hair.StrandsDeformed_DeformedPositionBuffer1), ResourceName), OwnerName, EHairResourceUsageType::Dynamic);
 	InternalCreateVertexBufferRDG<FHairStrandsTangentFormat>(GraphBuilder,  PointCount * FHairStrandsTangentFormat::ComponentCount, TangentBuffer, ToHairResourceDebugName(HAIRSTRANDS_RESOUCE_NAME(CurveType, Hair.StrandsDeformed_TangentBuffer), ResourceName), OwnerName, EHairResourceUsageType::Dynamic);
 
 	if (Prev0 != nullptr)
@@ -1175,7 +1175,7 @@ FRDGExternalBuffer& FHairStrandsDeformedResource::GetDeformerBuffer(FRDGBuilder&
 	// Lazy allocation and update
 	if (DeformerBuffer.Buffer == nullptr)
 	{
-		InternalCreateVertexBufferRDG<FHairStrandsPositionFormat>(GraphBuilder, BulkData.GetNumPoints(), DeformerBuffer, ToHairResourceDebugName(HAIRSTRANDS_RESOUCE_NAME(CurveType, Hair.StrandsDeformed_DeformerBuffer), ResourceName), OwnerName, EHairResourceUsageType::Dynamic);
+		InternalCreateByteAddressBufferRDG(GraphBuilder, BulkData.Header.PointCount * BulkData.Header.Strides.PositionStride, DeformerBuffer, ToHairResourceDebugName(HAIRSTRANDS_RESOUCE_NAME(CurveType, Hair.StrandsDeformed_DeformerBuffer), ResourceName), OwnerName, EHairResourceUsageType::Dynamic);
 	}
 	return DeformerBuffer;
 }
diff --git a/Engine/Plugins/Runtime/HairStrands/Source/HairStrandsCore/Private/HairStrandsInterpolation.cpp b/Engine/Plugins/Runtime/HairStrands/Source/HairStrandsCore/Private/HairStrandsInterpolation.cpp
--- a/Engine/Plugins/Runtime/HairStrands/Source/HairStrandsCore/Private/HairStrandsInterpolation.cpp
+++ b/Engine/Plugins/Runtime/HairStrands/Source/HairStrandsCore/Private/HairStrandsInterpolation.cpp
@@ -472,12 +472,12 @@ class FHairInterpolationCS : public FGlobalShader
 		SHADER_PARAMETER_RDG_BUFFER_SRV(Buffer, RenCurveBuffer)
 		SHADER_PARAMETER_RDG_BUFFER_SRV(Buffer, SimCurveBuffer)
 
-		SHADER_PARAMETER_RDG_BUFFER_SRV(Buffer, RenRestPosePositionBuffer)
-		SHADER_PARAMETER_RDG_BUFFER_SRV(Buffer, RenDeformerPositionBuffer)
-		SHADER_PARAMETER_RDG_BUFFER_UAV(RWBuffer, OutRenDeformedPositionBuffer)
+		SHADER_PARAMETER_RDG_BUFFER_SRV(ByteAddressBuffer, RenRestPosePositionBuffer)
+		SHADER_PARAMETER_RDG_BUFFER_SRV(ByteAddressBuffer, RenDeformerPositionBuffer)
+		SHADER_PARAMETER_RDG_BUFFER_UAV(RWByteAddressBuffer, OutRenDeformedPositionBuffer)
 
-		SHADER_PARAMETER_RDG_BUFFER_SRV(Buffer, SimRestPosePositionBuffer)
-		SHADER_PARAMETER_RDG_BUFFER_SRV(Buffer, SimDeformedPositionBuffer)
+		SHADER_PARAMETER_RDG_BUFFER_SRV(ByteAddressBuffer, SimRestPosePositionBuffer)
+		SHADER_PARAMETER_RDG_BUFFER_SRV(ByteAddressBuffer, SimDeformedPositionBuffer)
 
 		SHADER_PARAMETER_RDG_BUFFER_SRV(ByteAddressBuffer, CurveInterpolationBuffer)
 		SHADER_PARAMETER_RDG_BUFFER_SRV(ByteAddressBuffer, PointInterpolationBuffer)
diff --git a/Engine/Plugins/Runtime/HairStrands/Source/HairStrandsCore/Private/Niagara/NiagaraDataInterfaceHairStrands.cpp b/Engine/Plugins/Runtime/HairStrands/Source/HairStrandsCore/Private/Niagara/NiagaraDataInterfaceHairStrands.cpp
--- a/Engine/Plugins/Runtime/HairStrands/Source/HairStrandsCore/Private/Niagara/NiagaraDataInterfaceHairStrands.cpp
+++ b/Engine/Plugins/Runtime/HairStrands/Source/HairStrandsCore/Private/Niagara/NiagaraDataInterfaceHairStrands.cpp
@@ -31,8 +31,8 @@
 namespace NDIHairStrandsLocal
 {
 BEGIN_SHADER_PARAMETER_STRUCT(FShaderParameters,)
-	SHADER_PARAMETER_RDG_BUFFER_UAV(RWBuffer<uint4>,	DeformedPositionBuffer)
-	SHADER_PARAMETER_RDG_BUFFER_SRV(Buffer<uint4>,		RestPositionBuffer)
+	SHADER_PARAMETER_RDG_BUFFER_UAV(RWByteAddressBuffer,DeformedPositionBuffer)
+	SHADER_PARAMETER_RDG_BUFFER_SRV(ByteAddressBuffer,	RestPositionBuffer)
 	SHADER_PARAMETER_RDG_BUFFER_SRV(Buffer<uint>,		CurvesOffsetsBuffer)
 	SHADER_PARAMETER_RDG_BUFFER_SRV(Buffer<float4>,		RestTrianglePositionBuffer)
 	SHADER_PARAMETER_RDG_BUFFER_SRV(Buffer<float4>,		DeformedTrianglePositionBuffer)
@@ -824,8 +824,8 @@ class FInterpolateGroomGuidesCS : public FGlobalShader
 		SHADER_PARAMETER_RDG_BUFFER_SRV(Buffer<uint>,	RootToUniqueTriangleIndexBuffer)
 
 		SHADER_PARAMETER_RDG_BUFFER_SRV(Buffer<uint>,		PointToCurveIndexBuffer)
-		SHADER_PARAMETER_RDG_BUFFER_UAV(RWBuffer<uint4>,	DeformedPositionBuffer)
-		SHADER_PARAMETER_RDG_BUFFER_SRV(Buffer<uint4>,		RestPositionBuffer)
+		SHADER_PARAMETER_RDG_BUFFER_UAV(RWByteAddressBuffer,DeformedPositionBuffer)
+		SHADER_PARAMETER_RDG_BUFFER_SRV(ByteAddressBuffer,	RestPositionBuffer)
 		SHADER_PARAMETER_RDG_BUFFER_SRV(Buffer<uint>,		CurvesOffsetsBuffer)
 		SHADER_PARAMETER_RDG_BUFFER_SRV(Buffer<float4>,	DeformedPositionOffset)
 		SHADER_PARAMETER(FVector3f,				RestPositionOffset)
@@ -3375,7 +3375,7 @@ void UNiagaraDataInterfaceHairStrands::SetShaderParameters(const FNiagaraDataInt
 		if (bIsDeformedValid)
 		{
 			FRDGPooledBuffer* DeformedPositionBuffer = HairStrandsBuffer->SourceDeformedResources->DeformedPositionBuffer[HairStrandsBuffer->SourceDeformedResources->CurrentIndex].Buffer;
-			DeformedPositionBufferUAV = GraphBuilder.CreateUAV(GraphBuilder.RegisterExternalBuffer(DeformedPositionBuffer), FHairStrandsPositionFormat::Format);	//-OPT: If constant across all dispatches can create in PreStage once
+			DeformedPositionBufferUAV = GraphBuilder.CreateUAV(GraphBuilder.RegisterExternalBuffer(DeformedPositionBuffer));	//-OPT: If constant across all dispatches can create in PreStage once
 			DeformedPositionOffsetSRV = RegisterAsSRV(GraphBuilder, HairStrandsBuffer->SourceDeformedResources->DeformedOffsetBuffer[HairStrandsBuffer->SourceDeformedResources->CurrentIndex]);
 
 			// Keep the readback commented out for future debug
@@ -3391,7 +3391,7 @@ void UNiagaraDataInterfaceHairStrands::SetShaderParameters(const FNiagaraDataInt
 		}
 		else
 		{
-			DeformedPositionBufferUAV = GraphBuilder.CreateUAV(GraphBuilder.RegisterExternalBuffer(HairStrandsBuffer->DeformedPositionBuffer), FHairStrandsPositionFormat::Format);
+			DeformedPositionBufferUAV = GraphBuilder.CreateUAV(GraphBuilder.RegisterExternalBuffer(HairStrandsBuffer->DeformedPositionBuffer));
 			DeformedPositionOffsetSRV = Context.GetComputeDispatchInterface().GetEmptyBufferSRV(GraphBuilder, FHairStrandsPositionOffsetFormat::Format);
 		}
 		const int32 MeshLODIndex = bIsRootValid ? HairStrandsBuffer->SourceDeformedRootResources->MeshLODIndex : -1;
@@ -3516,7 +3516,7 @@ void UNiagaraDataInterfaceHairStrands::SetShaderParameters(const FNiagaraDataInt
 		ShaderParameters->SampleCount = 0;
 
 		// Set Shader UAV
-		ShaderParameters->DeformedPositionBuffer = Context.GetComputeDispatchInterface().GetEmptyBufferUAV(GraphBuilder, FHairStrandsPositionFormat::Format);
+		ShaderParameters->DeformedPositionBuffer = Context.GetComputeDispatchInterface().GetEmptyBufferUAV(GraphBuilder, PF_R32_UINT);
 		ShaderParameters->BoundingBoxBuffer = Context.GetComputeDispatchInterface().GetEmptyBufferUAV(GraphBuilder, PF_R32_UINT);
 
 		// Set Shader SRV
diff --git a/Engine/Plugins/Runtime/HairStrands/Source/HairStrandsCore/Public/HairStrandsDatas.h b/Engine/Plugins/Runtime/HairStrands/Source/HairStrandsCore/Public/HairStrandsDatas.h
--- a/Engine/Plugins/Runtime/HairStrands/Source/HairStrandsCore/Public/HairStrandsDatas.h
+++ b/Engine/Plugins/Runtime/HairStrands/Source/HairStrandsCore/Public/HairStrandsDatas.h
@@ -40,8 +40,9 @@ struct FPackedHairVertex
 	typedef uint64 BulkType;
 
 	FFloat16 X, Y, Z;
-	uint8 PackedRadiusAndType;
 	uint8 UCoord;
+	uint8 Radius : 6;
+	uint8 Type : 2;
 };
 
 struct FPackedHairAttribute0Vertex
@@ -76,7 +77,7 @@ struct FHairStrandsPositionFormat
 	static const uint32 ComponentCount = 1;
 	static const uint32 SizeInByte = sizeof(Type);
 	static const EVertexElementType VertexElementType = VET_UShort4;
-	static const EPixelFormat Format = PF_R16G16B16A16_UINT;
+	static const EPixelFormat Format = PF_R32G32_UINT;
 };
 
 struct FHairStrandsPositionOffsetFormat
diff --git a/Engine/Shaders/Private/HairStrands/HairCardsDeformation.usf b/Engine/Shaders/Private/HairStrands/HairCardsDeformation.usf
--- a/Engine/Shaders/Private/HairStrands/HairCardsDeformation.usf
+++ b/Engine/Shaders/Private/HairStrands/HairCardsDeformation.usf
@@ -8,7 +8,7 @@
 #define DEBUG_ENABLE 0
 
 #if DEBUG_ENABLE
-#include "../ShaderPrint.ush"?
+#include "../ShaderPrint.ush"
 #endif
 
 #if SHADER_DEFORM_CARDS
@@ -21,8 +21,8 @@ uint GuideVertexCount;
 float3 GuideRestPositionOffset;
 StructuredBuffer<float4> GuideDeformedPositionOffsetBuffer;
 
-Buffer<uint4>  GuideRestPositionBuffer;
-Buffer<uint4>  GuideDeformedPositionBuffer;
+ByteAddressBuffer GuideRestPositionBuffer;
+ByteAddressBuffer GuideDeformedPositionBuffer;
 Buffer<float4> CardsRestPositionBuffer;
 ByteAddressBuffer CardsInterpolationBuffer;
 
diff --git a/Engine/Shaders/Private/HairStrands/HairStrandsAttributeCommon.ush b/Engine/Shaders/Private/HairStrands/HairStrandsAttributeCommon.ush
--- a/Engine/Shaders/Private/HairStrands/HairStrandsAttributeCommon.ush
+++ b/Engine/Shaders/Private/HairStrands/HairStrandsAttributeCommon.ush
@@ -81,12 +81,12 @@ uint InternalGetHairStrandsAttributesPerVertex(uint HairControlPointId, uint Att
 	return InternalGetHairStrandsAttributes(HairControlPointId, AttributeOffsetInBytes, AttributeStrideInBytes, InAttributeBuffer, InIndexToChunkDivAsShift, InChunkElementCount, InChunkStrideInBytes);
 }
 
-float InternalGetHairStrandsCoordU(Buffer<uint4> InPositionBuffer, uint InPointIndex, float InRadius, float InRootScale, float InTipScale)
+float InternalGetHairStrandsCoordU(ByteAddressBuffer InPositionBuffer, uint InPointIndex, float InRadius, float InRootScale, float InTipScale)
 {
 	return ReadHairControlPoint(InPositionBuffer, InPointIndex, float3(0, 0, 0), InRadius, InRootScale, InTipScale).UCoord;
 }
 
-float InternalGetHairStrandsWorldRadius(Buffer<uint4> InPositionBuffer, uint InPointIndex, float InRadius, float InRootScale, float InTipScale)
+float InternalGetHairStrandsWorldRadius(ByteAddressBuffer InPositionBuffer, uint InPointIndex, float InRadius, float InRootScale, float InTipScale)
 {
 	return ReadHairControlPoint(InPositionBuffer, InPointIndex, float3(0, 0, 0), InRadius, InRootScale, InTipScale).WorldRadius;
 }
diff --git a/Engine/Shaders/Private/HairStrands/HairStrandsClusterCulling.usf b/Engine/Shaders/Private/HairStrands/HairStrandsClusterCulling.usf
--- a/Engine/Shaders/Private/HairStrands/HairStrandsClusterCulling.usf
+++ b/Engine/Shaders/Private/HairStrands/HairStrandsClusterCulling.usf
@@ -448,7 +448,7 @@ float4x4 LocalToTranslatedWorldMatrix;
 
 Buffer<uint> RenCurveBuffer;
 Buffer<uint> RenPointLODBuffer;
-Buffer<uint4> RenderDeformedPositionBuffer;
+ByteAddressBuffer RenderDeformedPositionBuffer;
 StructuredBuffer<float4> RenderDeformedOffsetBuffer;
 
 RWBuffer<int> OutClusterAABBBuffer; 	// Cluster data packed as {uint3 AABBMin, uint3 AABBMax}
diff --git a/Engine/Shaders/Private/HairStrands/HairStrandsDebug.usf b/Engine/Shaders/Private/HairStrands/HairStrandsDebug.usf
--- a/Engine/Shaders/Private/HairStrands/HairStrandsDebug.usf
+++ b/Engine/Shaders/Private/HairStrands/HairStrandsDebug.usf
@@ -1974,19 +1974,14 @@ float3 SimRestOffset;
 
 float4x4 LocalToWorld;
 
-Buffer<uint4> RenRestPosition;
-Buffer<uint4> RenDeformedPosition;
+ByteAddressBuffer RenRestPosition;
+ByteAddressBuffer RenDeformedPosition;
 StructuredBuffer<float4> RenDeformedOffset;
 
-Buffer<uint4> SimRestPosition;
-Buffer<uint4> SimDeformedPosition;
+ByteAddressBuffer SimRestPosition;
+ByteAddressBuffer SimDeformedPosition;
 StructuredBuffer<float4> SimDeformedOffset;
 
-float3 RestOffset;
-float3 DeformedOffset;
-Buffer<uint4> RestPosition;
-Buffer<uint4> DeformedPosition; 
-
 int DebugMode;
 
 #define DEBUGMODE_REN_REST 1
@@ -2004,39 +1999,39 @@ void MainCS(uint3 DispatchThreadId : SV_DispatchThreadID)
 
 	FHairControlPoint P0 = (FHairControlPoint)0;
 	FHairControlPoint P1 = (FHairControlPoint)0;
-	float3 Offset = 0;
+
 	float4 ColorRoot = ColorWhite;
 	float4 ColorTip  = ColorWhite;
 
 	if (DebugMode == DEBUGMODE_REN_REST)
 	{
-		P0       = ReadHairControlPoint(RenRestPosition, VertexId, 1, 1, 1);
-		P1       = ReadHairControlPoint(RenRestPosition, VertexId + 1, 1, 1, 1);
-		Offset   = RenRestOffset;
+		P0       = ReadHairControlPoint(RenRestPosition, VertexId,     RenRestOffset, 1, 1, 1);
+		P1       = ReadHairControlPoint(RenRestPosition, VertexId + 1, RenRestOffset, 1, 1, 1);
 		ColorRoot= ColorYellow;
 		ColorTip = ColorPurple;
 	}
 	else if (DebugMode == DEBUGMODE_REN_DEFORMED)
 	{
-		P0       = ReadHairControlPoint(RenDeformedPosition, VertexId, 1, 1, 1);
-		P1       = ReadHairControlPoint(RenDeformedPosition, VertexId + 1, 1, 1, 1);
-		Offset   = ReadCardPositionOffset(RenDeformedOffset, InstanceRegisteredIndex, CardLODIndex);
+		const float3 Offset = ReadCardPositionOffset(RenDeformedOffset, InstanceRegisteredIndex, CardLODIndex);
+
+		P0       = ReadHairControlPoint(RenDeformedPosition, VertexId,     Offset, 1, 1, 1);
+		P1       = ReadHairControlPoint(RenDeformedPosition, VertexId + 1, Offset, 1, 1, 1);
 		ColorRoot= ColorYellow;
 		ColorTip = ColorPurple;
 	}
 	else if (DebugMode == DEBUGMODE_SIM_REST)
 	{
-		P0       = ReadHairControlPoint(SimRestPosition, VertexId, 1, 1, 1);
-		P1       = ReadHairControlPoint(SimRestPosition, VertexId + 1, 1, 1, 1);
-		Offset   = SimRestOffset;
+		P0       = ReadHairControlPoint(SimRestPosition, VertexId,     SimRestOffset, 1, 1, 1);
+		P1       = ReadHairControlPoint(SimRestPosition, VertexId + 1, SimRestOffset, 1, 1, 1);
 		ColorRoot= ColorCyan;
 		ColorTip = ColorRed;
 	}
 	else if (DebugMode == DEBUGMODE_SIM_DEFORMED)
 	{
-		P0       = ReadHairControlPoint(SimDeformedPosition, VertexId, 1, 1, 1);
-		P1       = ReadHairControlPoint(SimDeformedPosition, VertexId + 1, 1, 1, 1);
-		Offset   = ReadSimPositionOffset(SimDeformedOffset, InstanceRegisteredIndex);
+		const float3 Offset = ReadSimPositionOffset(SimDeformedOffset, InstanceRegisteredIndex);
+
+		P0       = ReadHairControlPoint(SimDeformedPosition, VertexId,     Offset, 1, 1, 1);
+		P1       = ReadHairControlPoint(SimDeformedPosition, VertexId + 1, Offset, 1, 1, 1);
 		ColorRoot= ColorCyan;
 		ColorTip = ColorRed;
 	}
diff --git a/Engine/Shaders/Private/HairStrands/HairStrandsGuideDeform.usf b/Engine/Shaders/Private/HairStrands/HairStrandsGuideDeform.usf
--- a/Engine/Shaders/Private/HairStrands/HairStrandsGuideDeform.usf
+++ b/Engine/Shaders/Private/HairStrands/HairStrandsGuideDeform.usf
@@ -12,8 +12,8 @@ uint VertexCount;
 uint InstanceRegisteredIndex;
 
 StructuredBuffer<float4> SimDeformedOffsetBuffer;
-Buffer<uint4> SimRestPosePositionBuffer;
-RWBuffer<uint4> OutSimDeformedPositionBuffer;
+ByteAddressBuffer SimRestPosePositionBuffer;
+RWByteAddressBuffer OutSimDeformedPositionBuffer;
 
 #define DEFORMATION_BYPASS 0
 #define DEFORMATION_OFFSET_GUIDE 1
diff --git a/Engine/Shaders/Private/HairStrands/HairStrandsInterpolation.usf b/Engine/Shaders/Private/HairStrands/HairStrandsInterpolation.usf
--- a/Engine/Shaders/Private/HairStrands/HairStrandsInterpolation.usf
+++ b/Engine/Shaders/Private/HairStrands/HairStrandsInterpolation.usf
@@ -80,9 +80,9 @@ float3 InRenHairPositionOffset;
 float3 InSimHairPositionOffset;
 float4x4 LocalToWorldMatrix;
 
-Buffer<uint4> RenRestPosePositionBuffer;
-Buffer<uint4> SimRestPosePositionBuffer;
-Buffer<uint4> SimDeformedPositionBuffer;
+ByteAddressBuffer RenRestPosePositionBuffer;
+ByteAddressBuffer SimRestPosePositionBuffer;
+ByteAddressBuffer SimDeformedPositionBuffer;
 
 uint HairStrandsVFTODO_bCullingEnable;
 #if PERMUTATION_CULLING == 1
@@ -92,9 +92,9 @@ Buffer<float>	HairStrandsVFTODO_CullingRadiusScaleBuffer;
 #endif
 
 #if PERMUTATION_DEFORMER
-Buffer<uint4>  RenDeformerPositionBuffer;
+ByteAddressBuffer  RenDeformerPositionBuffer;
 #endif
-RWBuffer<uint4> OutRenDeformedPositionBuffer; 
+RWByteAddressBuffer OutRenDeformedPositionBuffer; 
 
 #if PERMUTATION_DYNAMIC_GEOMETRY >= 1 
 
@@ -824,8 +824,8 @@ void MainCS(uint2 DispatchThreadId : SV_DispatchThreadID, uint GroupThread1D : S
 #if SHADER_HAIRTRANSFER_PREV_POSITION
 
 uint ElementCount;
-Buffer<uint4> InBuffer;
-RWBuffer<uint4> OutBuffer;
+ByteAddressBuffer InBuffer;
+RWByteAddressBuffer OutBuffer;
 
 [numthreads(GROUP_SIZE, 1, 1)]
 void MainCS(uint DispatchThreadId : SV_DispatchThreadID)
@@ -835,7 +835,7 @@ void MainCS(uint DispatchThreadId : SV_DispatchThreadID)
 	{
 		return;
 	}
-	OutBuffer[Index] = InBuffer[Index];
+	CopyPackedHairControlPoint(OutBuffer, InBuffer, Index);
 }
 
 #endif // SHADER_HAIRTRANSFER_PREV_POSITION
@@ -855,9 +855,9 @@ Buffer<float> InPosition1Buffer;
 Buffer<float> InRadius0Buffer;
 Buffer<float> InRadius1Buffer;
 
-Buffer<uint4>  InRestPoseBuffer;
+ByteAddressBuffer  InRestPoseBuffer;
 StructuredBuffer<float4> InDeformedOffsetBuffer;
-RWBuffer<uint4> OutDeformedBuffer;
+RWByteAddressBuffer OutDeformedBuffer;
 
 [numthreads(GROUP_SIZE, 1, 1)]
 void MainCS(uint DispatchThreadId : SV_DispatchThreadID)
@@ -880,7 +880,7 @@ void MainCS(uint DispatchThreadId : SV_DispatchThreadID)
 		InPosition1Buffer[Index3+2]);
 	
 	const float3 WorldVertexPosition = lerp(InPosition0, InPosition1, InterpolationFactor);
-	FHairControlPoint CP = ReadHairControlPoint(InRestPoseBuffer, Index, float3(0, 0, 0), 1, 1, 1);
+	FHairControlPoint CP = ReadHairControlPoint(InRestPoseBuffer, Index, float3(0, 0, 0)/*Offset*/, 1/*Radius*/, 1/*RootScale*/, 1/*TipScale*/);
 	CP.Position = WorldVertexPosition;
 
 	if (bHasRadiusData == 1)
@@ -889,7 +889,8 @@ void MainCS(uint DispatchThreadId : SV_DispatchThreadID)
 	}
 
 	const float3 DeformedPositionOffset = ReadRenPositionOffset(InDeformedOffsetBuffer, InstanceRegisteredIndex);
-	OutDeformedBuffer[Index] = PackHairControlPoint(CP, DeformedPositionOffset, 1);
+
+	WriteHairControlPointPosition(OutDeformedBuffer, Index, CP, DeformedPositionOffset, 1 /*Radius*/);
 }
 
 #endif // SHADER_GROOMCACHE_UPDATE
diff --git a/Engine/Shaders/Private/HairStrands/HairStrandsMesh.usf b/Engine/Shaders/Private/HairStrands/HairStrandsMesh.usf
--- a/Engine/Shaders/Private/HairStrands/HairStrandsMesh.usf
+++ b/Engine/Shaders/Private/HairStrands/HairStrandsMesh.usf
@@ -321,6 +321,7 @@ TEMPLATE_BUFFERS(StructuredBuffer<float4> Structured, )
 TEMPLATE_BUFFERS(Buffer<uint> VertexUInt, )
 TEMPLATE_BUFFERS(Buffer<uint4> VertexUInt4, )
 TEMPLATE_BUFFERS(Buffer<float4> VertexFloat4, )
+TEMPLATE_BUFFERS(ByteAddressBuffer ByteAddress, )
 
 RWStructuredBuffer<uint> DummyOutput;
 uint DummyValue;
@@ -340,8 +341,12 @@ void MainCS(uint3 DispatchThreadId : SV_DispatchThreadID)
 		TEMPLATE_BUFFERS(Acc += VertexUInt, [0])
 	#elif PERMUTATION_BUFFER_TYPE == 2
 		TEMPLATE_BUFFERS(Acc += VertexUInt4, [0].x)
-	#else
+	#elif PERMUTATION_BUFFER_TYPE == 3
 		TEMPLATE_BUFFERS(Acc += VertexFloat4, [0].x)
+	#elif PERMUTATION_BUFFER_TYPE == 4
+		TEMPLATE_BUFFERS(Acc += ByteAddress, .Load(0))
+	#else
+		#error Not implemented
 	#endif
 		DummyOutput[0] = Acc;
 	}
diff --git a/Engine/Shaders/Private/HairStrands/HairStrandsPack.ush b/Engine/Shaders/Private/HairStrands/HairStrandsPack.ush
--- a/Engine/Shaders/Private/HairStrands/HairStrandsPack.ush
+++ b/Engine/Shaders/Private/HairStrands/HairStrandsPack.ush
@@ -9,16 +9,36 @@
 ////////////////////////////////////////////////////////////////////////////////////////////////////////////////////////////////
 // Hair control points
 
-#define FPackedHairPosition uint4
+// Hair control point layout
+// [                  uint2                   ]
+// [      x     ][              y             ]
+// [   32bits   ][           32bits           ]
+// [  16 ][  16 ][ 16  ][   8  ][   6  ][  2  ]
+// [ 0-15][16-32][ 0-15][ 16-23][ 24-29][30-31]
+// [Pos.x][Pos.y][Pos.z][CoordU][Radius][ Type]
+#define FPackedHairPosition uint2
+#define FPackedHairPositionStrideInBytes 8u
 
-FPackedHairPosition ReadPackedHairPosition(Buffer<FPackedHairPosition> InBuffer, uint InPointIndex)
+FPackedHairPosition ReadPackedHairPosition(ByteAddressBuffer InBuffer, uint InPointIndex)
 {
-	return InBuffer[InPointIndex];
+	return InBuffer.Load2(InPointIndex * FPackedHairPositionStrideInBytes);
+}
+
+void WritePackedHairPosition(ByteAddressBuffer InBuffer, uint InPointIndex, FPackedHairPosition In)
+{
+	InBuffer.Store2(InPointIndex * FPackedHairPositionStrideInBytes, In);
 }
 
 void CopyPackedAttribute(inout FPackedHairPosition Out, in FPackedHairPosition In)
 {
-	Out.w = Out.w;
+	// Copy the last 16bits which contains the CoordU|Radius|Type
+	Out.y = (Out.y & 0x0000FFFF) | (In.y & 0xFFFF0000);
+}
+
+void CopyPackedHairControlPoint(RWByteAddressBuffer OutBuffer, ByteAddressBuffer InBuffer, uint Index)
+{
+	const uint AddressInBytes = Index * FPackedHairPositionStrideInBytes;
+	OutBuffer.Store2(AddressInBytes, InBuffer.Load2(AddressInBytes));
 }
 
 struct FHairControlPoint
@@ -34,34 +54,36 @@ FPackedHairPosition PackHairControlPoint(
 	float3 InVF_PositionOffset,
 	float InVF_Radius) 
 {
-	uint PackedAttribute =  (uint(clamp(CP.UCoord * 0xFFu, 0u, 0xFFu)) << 8) |
-							(uint(clamp((CP.WorldRadius / InVF_Radius) * 63u, 0u, 63u)) << 2u) |
-							(CP.Type & 0x3);
-	return FPackedHairPosition(f32tof16(CP.Position - InVF_PositionOffset), PackedAttribute);
+	const uint UCoord8bits = PackR8(CP.UCoord);
+	const uint Radius6bits = PackR6(CP.WorldRadius / InVF_Radius);
+	const uint Type2bits   = CP.Type & 0x3;
+	CP.Position -= InVF_PositionOffset;
+
+	FPackedHairPosition Out = 0;
+	Out.x = PackFloat2ToUInt(CP.Position.x, CP.Position.y);
+	Out.y = f32tof16(CP.Position.z) | (UCoord8bits << 16u) | (Radius6bits << 24u) | (Type2bits << 30u);
+	return Out;
 }
 
 FHairControlPoint UnpackHairControlPoint(
-	FPackedHairPosition InPackedData,
+	FPackedHairPosition In,
 	float3 InVF_PositionOffset=0,
 	float InVF_Radius=1,
 	float InVF_RootScale=1,
 	float InVF_TipScale=1)
 {
-	const uint PackedAlpha	= asuint(InPackedData.a);
-
 	FHairControlPoint Out = (FHairControlPoint)0;
-	Out.Position	= f16tof32(InPackedData.xyz) + InVF_PositionOffset;
-	Out.UCoord		= ((PackedAlpha & 0xFF00) >> 8) / 255.f;
-	Out.WorldRadius	= ((PackedAlpha & 0x00FC) >> 2) / 63.f;
-	Out.Type		= PackedAlpha & 0x3;
+	Out.Position	= float3(UnpackFloat2FromUInt(In.x), f16tof32(BitFieldExtractU32(In.y, 16, 0))) + InVF_PositionOffset;
+	Out.UCoord		= UnpackR8(BitFieldExtractU32(In.y, 8, 16));
+	Out.WorldRadius	= UnpackR6(BitFieldExtractU32(In.y, 6, 24));
+	Out.Type		= BitFieldExtractU32(In.y, 2, 30);
 
 	Out.WorldRadius *= InVF_Radius * lerp(InVF_RootScale, InVF_TipScale, Out.UCoord);
-
-	return Out;	
+	return Out;
 }
 
 FHairControlPoint ReadHairControlPoint(
-	Buffer<FPackedHairPosition> InBuffer,
+	ByteAddressBuffer InBuffer,
 	uint InIndex,
 	float3 InVF_PositionOffset=0,
 	float InVF_Radius=1,
@@ -69,48 +91,65 @@ FHairControlPoint ReadHairControlPoint(
 	float InVF_TipScale=1)
 {
 	return UnpackHairControlPoint(
-		InBuffer[InIndex],
+		InBuffer.Load2(InIndex * FPackedHairPositionStrideInBytes),
 		InVF_PositionOffset,
 		InVF_Radius,
 		InVF_RootScale,
 		InVF_TipScale);
 }
 
-float UnpackHairControlPointCoordU(FPackedHairPosition InPackedData)
+void WriteHairControlPointPosition(
+	RWByteAddressBuffer OutBuffer, 
+	uint InIndex, 
+	FHairControlPoint InCP,
+	float3 InVF_PositionOffset=0,
+	float InVF_Radius=1)
+{
+	OutBuffer.Store2(InIndex * FPackedHairPositionStrideInBytes, PackHairControlPoint(InCP, InVF_PositionOffset, InVF_Radius));
+}
+
+float UnpackHairControlPointCoordU(FPackedHairPosition In)
 {
-	const uint PackedAlpha = asuint(InPackedData.w);
-	return ((PackedAlpha & 0xFF00) >> 8) / 255.f;
+	return UnpackR8(BitFieldExtractU32(In.y, 8, 16));
 }
 
-float3 UnpackHairControlPointPosition(FPackedHairPosition InPackedData, float3 InVF_PositionOffset=0)
+float3 UnpackHairControlPointPosition(FPackedHairPosition In, float3 InVF_PositionOffset=0)
 {
-	return f16tof32(InPackedData.xyz) + InVF_PositionOffset;
+	return float3(UnpackFloat2FromUInt(In.x), f16tof32(BitFieldExtractU32(In.y, 16, 0))) + InVF_PositionOffset;
 }
 
-float3 ReadHairControlPointPosition(Buffer<FPackedHairPosition> InBuffer, uint InPointIndex, float3 InVF_PositionOffset=0)
+float3 ReadHairControlPointPosition(ByteAddressBuffer InBuffer, uint InPointIndex, float3 InVF_PositionOffset=0)
 {
-	return UnpackHairControlPointPosition(InBuffer[InPointIndex], InVF_PositionOffset);
+	return UnpackHairControlPointPosition(InBuffer.Load2(InPointIndex * FPackedHairPositionStrideInBytes), InVF_PositionOffset);
 }
 
-float3 ReadHairControlPointPosition(RWBuffer<FPackedHairPosition> InBuffer, uint InPointIndex, float3 InVF_PositionOffset=0)
+float3 ReadHairControlPointPosition(RWByteAddressBuffer InBuffer, uint InPointIndex, float3 InVF_PositionOffset=0)
 {
-	return UnpackHairControlPointPosition(InBuffer[InPointIndex], InVF_PositionOffset);
+	return UnpackHairControlPointPosition(InBuffer.Load2(InPointIndex * FPackedHairPositionStrideInBytes), InVF_PositionOffset);
 }
 
 FPackedHairPosition PackHairControlPointPosition(FPackedHairPosition InPackedData, float3 InNewPosition, float3 InPositionOffset=0)
 {
-	return FPackedHairPosition(f32tof16(InNewPosition.xyz - InPositionOffset), InPackedData.w);
+	InNewPosition -= InPositionOffset;
+
+	FPackedHairPosition Out = 0;
+	Out.x = PackFloat2ToUInt(InNewPosition.x, InNewPosition.y);
+	Out.y = f32tof16(InNewPosition.z) | (InPackedData.y & 0xFFFF0000);
+	return Out;
 }
 
-void WritePackedHairControlPointPosition(RWBuffer<FPackedHairPosition> OutBuffer, uint InIndex, FPackedHairPosition InPackedData, float3 InNewPosition, float3 InPositionOffset=0)
+void WritePackedHairControlPointPosition(RWByteAddressBuffer OutBuffer, uint InIndex, FPackedHairPosition InPackedData, float3 InNewPosition, float3 InPositionOffset=0)
 {
-	OutBuffer[InIndex] = PackHairControlPointPosition(InPackedData, InNewPosition, InPositionOffset);
+	OutBuffer.Store2(InIndex * FPackedHairPositionStrideInBytes, PackHairControlPointPosition(InPackedData, InNewPosition, InPositionOffset));
 }
 
 FPackedHairPosition PackHairControlPointRadius(FPackedHairPosition InPackedData, float InNewWorldRadius, float InVFRadius)
 {
-	const uint PackedNewRadius = uint(clamp(saturate(InNewWorldRadius / InVFRadius) * 63u, 0u, 63u)) << 2u;	
-	return FPackedHairPosition(InPackedData.xyz, (InPackedData.w & (~0xFC)) | PackedNewRadius);
+	const uint Radius6bits = PackR6(InNewWorldRadius / InVFRadius);	
+
+	FPackedHairPosition Out = InPackedData;
+	Out.y = (Out.y & 0xC0FFFFFF) | (Radius6bits << 24);
+	return Out;
 }
 
 bool IsValidHairStrandsSegment(const FHairControlPoint CP0, const FHairControlPoint CP1)
diff --git a/Engine/Shaders/Private/HairStrands/HairStrandsRaytracingGeometry.usf b/Engine/Shaders/Private/HairStrands/HairStrandsRaytracingGeometry.usf
--- a/Engine/Shaders/Private/HairStrands/HairStrandsRaytracingGeometry.usf
+++ b/Engine/Shaders/Private/HairStrands/HairStrandsRaytracingGeometry.usf
@@ -16,7 +16,7 @@ float HairStrandsVF_HairTipScale;
 
 uint HairStrandsVF_RegisteredIndex;
 StructuredBuffer<float4> PositionOffsetBuffer;
-Buffer<uint4> 	 PositionBuffer;
+ByteAddressBuffer PositionBuffer;
 Buffer<float4> 	 TangentBuffer;
 RWBuffer<float4> OutputPositionBuffer;	
 #if PERMUTATION_PROCEDURAL_PRIMITIVE 
@@ -307,8 +307,8 @@ float PositionThreshold2;
 uint HairStrandsVF_bCullingEnable;
 Buffer<uint> 	 HairStrandsVF_CullingIndirectBuffer;
 Buffer<uint> 	 HairStrandsVF_CullingIndexBuffer;
-Buffer<uint4> 	 CurrPositionBuffer;
-Buffer<uint4> 	 PrevPositionBuffer;
+ByteAddressBuffer CurrPositionBuffer;
+ByteAddressBuffer PrevPositionBuffer;
 Buffer<int> 	 GroupAABBBuffer;
 RWBuffer<uint> 	 InvalidationBuffer;
 RWBuffer<uint> 	 InvalidationPrintCounter;
diff --git a/Engine/Shaders/Private/HairStrands/HairStrandsTangent.usf b/Engine/Shaders/Private/HairStrands/HairStrandsTangent.usf
--- a/Engine/Shaders/Private/HairStrands/HairStrandsTangent.usf
+++ b/Engine/Shaders/Private/HairStrands/HairStrandsTangent.usf
@@ -8,7 +8,7 @@
 #if SHADER_TANGENT
 uint PointCount;
 
-Buffer<uint4>  	 PositionBuffer;
+ByteAddressBuffer PositionBuffer;
 RWBuffer<float4> OutputTangentBuffer;
 
 uint			HairStrandsVF_bCullingEnable;
diff --git a/Engine/Shaders/Private/HairStrands/HairStrandsTexturesGeneration.usf b/Engine/Shaders/Private/HairStrands/HairStrandsTexturesGeneration.usf
--- a/Engine/Shaders/Private/HairStrands/HairStrandsTexturesGeneration.usf
+++ b/Engine/Shaders/Private/HairStrands/HairStrandsTexturesGeneration.usf
@@ -652,7 +652,7 @@ void MainCS(int2 DispatchThreadId : SV_DispatchThreadID)
 #include "HairStrandsVertexFactoryCommon.ush"
 
 uint MaxVertexCount;
-Buffer<uint4> InPositions;
+ByteAddressBuffer InPositions;
 Buffer<float4> InPositionOffset;
 RWBuffer<float4> OutPositions;
 [numthreads(128, 1, 1)]
diff --git a/Engine/Shaders/Private/HairStrands/HairStrandsVertexFactoryCommon.ush b/Engine/Shaders/Private/HairStrands/HairStrandsVertexFactoryCommon.ush
--- a/Engine/Shaders/Private/HairStrands/HairStrandsVertexFactoryCommon.ush
+++ b/Engine/Shaders/Private/HairStrands/HairStrandsVertexFactoryCommon.ush
@@ -78,7 +78,7 @@ float3 ReadCardPositionOffset(StructuredBuffer<float4> In, uint InIndex, uint In
 	uint						Prefix##_PointAttributeIndexToChunkDivAsShift; \
 	uint						Prefix##_PointAttributeChunkElementCount; \
 	uint						Prefix##_PointAttributeChunkStrideInBytes; \
-	Buffer<uint4>				Prefix##_PositionBuffer; \
+	ByteAddressBuffer			Prefix##_PositionBuffer; \
 	StructuredBuffer<float4>	Prefix##_PositionOffsetBuffer; \
 	Buffer<uint>				Prefix##_CurveBuffer; \
 	Buffer<uint>				Prefix##_PointToCurveBuffer; \
diff --git a/Engine/Shaders/Private/HairStrands/HairStrandsVisibilityRasterCompute.usf b/Engine/Shaders/Private/HairStrands/HairStrandsVisibilityRasterCompute.usf
--- a/Engine/Shaders/Private/HairStrands/HairStrandsVisibilityRasterCompute.usf
+++ b/Engine/Shaders/Private/HairStrands/HairStrandsVisibilityRasterCompute.usf
@@ -962,8 +962,8 @@ void RasterCS(uint DispatchThreadID : SV_DispatchThreadID, uint GroupThreadID :
 			group_PrimMatID[Prim] = PackHairVisControlPointMaterialId(PrimID, HairMaterialId);
 
 			uint TypeDummy;
-			CalcHomogenousPosAndRad(HairStrandsVF_PositionBuffer[PrimID], group_PositionOffset, group_SP0[Prim], group_Rad0[Prim], TypeDummy);
-			CalcHomogenousPosAndRad(HairStrandsVF_PositionBuffer[PrimID+1], group_PositionOffset, group_SP1[Prim], group_Rad1[Prim], TypeDummy);
+			CalcHomogenousPosAndRad(PrimID, group_PositionOffset, group_SP0[Prim], group_Rad0[Prim], TypeDummy);
+			CalcHomogenousPosAndRad(PrimID+1, group_PositionOffset, group_SP1[Prim], group_Rad1[Prim], TypeDummy);
 		}
 
 		if (bWaveThreadValid)
@@ -1011,8 +1011,8 @@ void RasterCS(uint DispatchThreadID : SV_DispatchThreadID, uint GroupThreadID :
 				PrimMatID = PackHairVisControlPointMaterialId(PrimID, HairMaterialId);
 
 				uint TypeDummy;
-				CalcHomogenousPosAndRad(HairStrandsVF_PositionBuffer[PrimID], group_PositionOffset, SP0, Rad0, TypeDummy);
-				CalcHomogenousPosAndRad(HairStrandsVF_PositionBuffer[PrimID + 1], group_PositionOffset, SP1, Rad1, TypeDummy);
+				CalcHomogenousPosAndRad(PrimID, group_PositionOffset, SP0, Rad0, TypeDummy);
+				CalcHomogenousPosAndRad(PrimID + 1, group_PositionOffset, SP1, Rad1, TypeDummy);
 			}
 
 			// Clipping
@@ -1360,8 +1360,8 @@ void RasterMultiSampleCS(uint DispatchThreadID : SV_DispatchThreadID, uint Group
 			PrimMatID = PackHairVisControlPointMaterialId(PrimID, HairMaterialId);
 
 			uint Type;
-			CalcHomogenousPosAndRad(HairStrandsVF_PositionBuffer[PrimID], group_PositionOffset, SP0, Rad0, Type);
-			CalcHomogenousPosAndRad(HairStrandsVF_PositionBuffer[PrimID + 1], group_PositionOffset, SP1, Rad1, Type);
+			CalcHomogenousPosAndRad(PrimID, group_PositionOffset, SP0, Rad0, Type);
+			CalcHomogenousPosAndRad(PrimID + 1, group_PositionOffset, SP1, Rad1, Type);
 			bIsEndPoint = (Type == HAIR_CONTROLPOINT_END);
 
 			SP0 = float4(NDCToPixelCoord(SP0), 1.0f / SP0.w);
diff --git a/Engine/Shaders/Private/PackUnpack.ush b/Engine/Shaders/Private/PackUnpack.ush
--- a/Engine/Shaders/Private/PackUnpack.ush
+++ b/Engine/Shaders/Private/PackUnpack.ush
@@ -108,6 +108,20 @@ float UnpackR8(uint In)
 #endif
 }
 
+//
+//	float <-> u6
+//
+
+uint PackR6(float Value)
+{
+	return uint(saturate(Value) * 63.0f);
+}
+
+float UnpackR6(uint In)
+{
+	return float(In & 0x3F) * (1.0f / 63.0f);
+}
+
 //
 //	float4 <-> u32 as u8.u8.u8.u8
 //
diff --git a/Engine/Shaders/Private/Substrate/Substrate.ush b/Engine/Shaders/Private/Substrate/Substrate.ush
--- a/Engine/Shaders/Private/Substrate/Substrate.ush
+++ b/Engine/Shaders/Private/Substrate/Substrate.ush
@@ -3396,11 +3396,6 @@ uint PackR6(float Value, float Dither)
 	return uint(DitheredValue * 63.0f);
 }
 
-float UnpackR6(uint Value)
-{
-	return float(Value & 0x3F) * (1.0f / 63.0f);
-}
-
 // Simple Octahedral from "A Survey of Efficient Representations for Independent Unit Vectors" http://jcgt.org/published/0003/02/01/
 float SignNotZero(in float k)
 {
diff --git a/Engine/Source/Runtime/Renderer/Private/HairStrands/HairStrandsUtils.cpp b/Engine/Source/Runtime/Renderer/Private/HairStrands/HairStrandsUtils.cpp
--- a/Engine/Source/Runtime/Renderer/Private/HairStrands/HairStrandsUtils.cpp
+++ b/Engine/Source/Runtime/Renderer/Private/HairStrands/HairStrandsUtils.cpp
@@ -359,7 +359,7 @@ private:
 	DECLARE_GLOBAL_SHADER(FHairResourceTransitionPass);
 	SHADER_USE_PARAMETER_STRUCT(FHairResourceTransitionPass, FGlobalShader);
 
-	class FBufferType : SHADER_PERMUTATION_INT("PERMUTATION_BUFFER_TYPE", 4);
+	class FBufferType : SHADER_PERMUTATION_INT("PERMUTATION_BUFFER_TYPE", 5);
 	using FPermutationDomain = TShaderPermutationDomain<FBufferType>;
 
 	BEGIN_SHADER_PARAMETER_STRUCT(FParameters, )
@@ -368,6 +368,7 @@ private:
 		SHADER_PARAMETER_RDG_BUFFER_SRV_ARRAY(Buffer<uint4>, VertexUInt4Buffers, [MaxBufferCount])
 		SHADER_PARAMETER_RDG_BUFFER_SRV_ARRAY(Buffer<float4>, VertexFloat4Buffers, [MaxBufferCount])
 		SHADER_PARAMETER_RDG_BUFFER_SRV_ARRAY(StructuredBuffer<float4>, StructuredBuffers, [MaxBufferCount])
+		SHADER_PARAMETER_RDG_BUFFER_SRV_ARRAY(ByteAddressBuffer, ByteAddressBuffers, [MaxBufferCount])
 		SHADER_PARAMETER_RDG_BUFFER_UAV(RWStructuredBuffer<uint>, DummyOutput)
 	END_SHADER_PARAMETER_STRUCT()
 
@@ -376,10 +377,15 @@ public:
 	// 1 : Buffer<uint>
 	// 2 : Buffer<uint4>
 	// 3 : Buffer<float4>
-	static const int32 PermutationCount = 4;
-	static int32 GetPermutationIndex(bool bStructured, bool bInteger, int32 NumComponents)
+	// 4 : ByteAdressBuffer
+	static const int32 PermutationCount = 5;
+	static int32 GetPermutationIndex(bool bStructured, bool bByteAddressBuffer, bool bInteger, int32 NumComponents)
 	{
-		if (bStructured)
+		if (bByteAddressBuffer)
+		{
+			return 4;
+		}
+		else if (bStructured)
 		{
 			return 0;
 		}
@@ -432,12 +438,14 @@ void AddTransitionPass(
 	DummyInputs[1] = GraphBuilder.CreateSRV(GSystemTextures.GetDefaultBuffer(GraphBuilder, 4u, 1u), PF_R32_UINT);
 	DummyInputs[2] = GraphBuilder.CreateSRV(GSystemTextures.GetDefaultBuffer(GraphBuilder, 16u, 1u), PF_R32G32B32A32_UINT);
 	DummyInputs[3] = GraphBuilder.CreateSRV(GSystemTextures.GetDefaultBuffer(GraphBuilder, 16u, 1u), PF_A32B32G32R32F);
+	DummyInputs[4] = GraphBuilder.CreateSRV(GSystemTextures.GetDefaultByteAddressBuffer(GraphBuilder, 4u));
 
 	int32 ArrayCounts[FHairResourceTransitionPass::PermutationCount];
 	ArrayCounts[0] = 0;
 	ArrayCounts[1] = 0;
 	ArrayCounts[2] = 0;
 	ArrayCounts[3] = 0;
+	ArrayCounts[4] = 0;
 
 	auto FlushArray = [&](int32 PermutationIndex, int32 TransitionCount)
 	{
@@ -452,6 +460,7 @@ void AddTransitionPass(
 		case 1: ParamArray = &PassParameters->VertexUIntBuffers; break;
 		case 2: ParamArray = &PassParameters->VertexUInt4Buffers; break;
 		case 3: ParamArray = &PassParameters->VertexFloat4Buffers; break;
+		case 4: ParamArray = &PassParameters->ByteAddressBuffers; break;
 		default: checkNoEntry();
 		};
 
@@ -484,11 +493,12 @@ void AddTransitionPass(
 	for (int32 TransitionIndex = 0; TransitionIndex < Transitions.Num(); ++TransitionIndex)
 	{
 		const bool bStructuredBuffer = EnumHasAnyFlags(Transitions[TransitionIndex]->Desc.Buffer->Desc.Usage, EBufferUsageFlags::StructuredBuffer);
+		const bool bByteAddressBuffer = EnumHasAnyFlags(Transitions[TransitionIndex]->Desc.Buffer->Desc.Usage, EBufferUsageFlags::ByteAddressBuffer);
 		const EPixelFormat PixelFormat = Transitions[TransitionIndex]->Desc.Format;
 		const bool bIsIntegerFormat = IsInteger(PixelFormat);
 		const int32 NumComponents = GPixelFormats[PixelFormat].NumComponents;
 
-		const int32 PermutationIndex = FHairResourceTransitionPass::GetPermutationIndex(bStructuredBuffer, bIsIntegerFormat, NumComponents);
+		const int32 PermutationIndex = FHairResourceTransitionPass::GetPermutationIndex(bStructuredBuffer, bByteAddressBuffer, bIsIntegerFormat, NumComponents);
 
 		SortedTransitions[PermutationIndex][ArrayCounts[PermutationIndex]++] = Transitions[TransitionIndex];
 
diff --git a/Engine/Source/Runtime/Renderer/Private/HairStrands/HairStrandsVisibility.cpp b/Engine/Source/Runtime/Renderer/Private/HairStrands/HairStrandsVisibility.cpp
--- a/Engine/Source/Runtime/Renderer/Private/HairStrands/HairStrandsVisibility.cpp
+++ b/Engine/Source/Runtime/Renderer/Private/HairStrands/HairStrandsVisibility.cpp
@@ -4294,8 +4294,8 @@ class FHairStrandsPositionChangedCS : public FGlobalShader
 		SHADER_PARAMETER(uint32, bDrawInvalidElement)
 		SHADER_PARAMETER_RDG_BUFFER_SRV(Buffer<uint>, HairStrandsVF_CullingIndexBuffer)
 		SHADER_PARAMETER_RDG_BUFFER_SRV(Buffer<uint>, HairStrandsVF_CullingIndirectBuffer)
-		SHADER_PARAMETER_RDG_BUFFER_SRV(Buffer<uint4>, CurrPositionBuffer)
-		SHADER_PARAMETER_RDG_BUFFER_SRV(Buffer<uint4>, PrevPositionBuffer)
+		SHADER_PARAMETER_RDG_BUFFER_SRV(ByteAddressBuffer, CurrPositionBuffer)
+		SHADER_PARAMETER_RDG_BUFFER_SRV(ByteAddressBuffer, PrevPositionBuffer)
 		SHADER_PARAMETER_RDG_BUFFER_SRV(Buffer<int>, GroupAABBBuffer)
 		SHADER_PARAMETER_RDG_BUFFER_UAV(RWBuffer<uint>, InvalidationBuffer)
 		SHADER_PARAMETER_RDG_BUFFER_UAV(RWBuffer<uint>, InvalidationPrintCounter)
diff --git a/Engine/Source/Runtime/Renderer/Public/HairStrandsInterface.h b/Engine/Source/Runtime/Renderer/Public/HairStrandsInterface.h
--- a/Engine/Source/Runtime/Renderer/Public/HairStrandsInterface.h
+++ b/Engine/Source/Runtime/Renderer/Public/HairStrandsInterface.h
@@ -96,7 +96,7 @@ END_SHADER_PARAMETER_STRUCT()
 
 // Instance resources (Raw)
 BEGIN_SHADER_PARAMETER_STRUCT(FHairStrandsInstanceResourceRawParameters, RENDERER_API)
-	SHADER_PARAMETER_SRV(Buffer<uint4>, PositionBuffer)
+	SHADER_PARAMETER_SRV(ByteAddressBuffer, PositionBuffer)
 	SHADER_PARAMETER_SRV(StructuredBuffer<float4>, PositionOffsetBuffer)
 	SHADER_PARAMETER_SRV(Buffer<uint>, CurveBuffer)
 	SHADER_PARAMETER_SRV(Buffer<uint>, PointToCurveBuffer)
@@ -107,7 +107,7 @@ END_SHADER_PARAMETER_STRUCT()
 
 // Instance prev. resources (Raw)
 BEGIN_SHADER_PARAMETER_STRUCT(FHairStrandsInstancePrevResourceRawParameters, RENDERER_API)
-	SHADER_PARAMETER_SRV(Buffer<uint4>, PreviousPositionBuffer)
+	SHADER_PARAMETER_SRV(ByteAddressBuffer, PreviousPositionBuffer)
 	SHADER_PARAMETER_SRV(StructuredBuffer<float4>, PreviousPositionOffsetBuffer)
 END_SHADER_PARAMETER_STRUCT()
 
