Path Tracer: Fix shader errors from shadow tracing refactor when Substrate is enabled

A number of lines of code got removed during the AHS shadow refactor.

diff --git a/Engine/Shaders/Private/PathTracing/PathTracingMaterialHitShader.usf b/Engine/Shaders/Private/PathTracing/PathTracingMaterialHitShader.usf
--- a/Engine/Shaders/Private/PathTracing/PathTracingMaterialHitShader.usf
+++ b/Engine/Shaders/Private/PathTracing/PathTracingMaterialHitShader.usf
@@ -1185,10 +1185,28 @@ RAY_TRACING_ENTRY_ANY_HIT(PathTracingMaterialAHS,
 
 #if SUBSTRATE_ENABLED
 
+		float3 Transparency = 0.0;
+
 #if MATERIAL_IS_SUBSTRATE && !SUBSTRATE_OPTIMIZED_UNLIT
 		// We already figured out how much Transparency there is above
 		// So only do extra logic for refraction and thin refraction ases
+
+
 #if MATERIALBLENDING_TRANSLUCENT && (REFRACTION_USE_INDEX_OF_REFRACTION || MATERIAL_ISTHINSURFACE)
+		FSubstratePixelHeader SubstratePixelHeader = MaterialParameters.GetFrontSubstrateHeader();
+		FSubstrateData SubstrateData = PixelMaterialInputs.GetFrontSubstrateData();
+
+		FSubstrateIntegrationSettings IntegrationSettings = InitSubstrateIntegrationSettings(
+			false /* bForceFullyRough */,
+			true  /* SubstrateStruct.bRoughDiffuse */,
+			-1    /* SubstrateStruct.PeelLayersAboveDepth */,
+			true  /* SubstrateStruct.bRoughnessTracking */
+		);
+
+		float TotalCoverage = 1;
+		float3 TotalTransmittancePreCoverage = 0;
+		const float3 V_World = MaterialParameters.CameraVector;
+		SubstratePixelHeader.SubstrateUpdateTree(SubstrateData, V_World, IntegrationSettings, TotalCoverage, TotalTransmittancePreCoverage);
 
 		// SUBSTRATE_TODO: refraction handling
 		float PathRoughness = PackedPayload.GetPathRoughness();
