HairStrands: Fix minor shader warning

diff --git a/Engine/Shaders/Private/HairStrands/HairStrandsPack.ush b/Engine/Shaders/Private/HairStrands/HairStrandsPack.ush
--- a/Engine/Shaders/Private/HairStrands/HairStrandsPack.ush
+++ b/Engine/Shaders/Private/HairStrands/HairStrandsPack.ush
@@ -24,7 +24,7 @@ FPackedHairPosition ReadPackedHairPosition(ByteAddressBuffer InBuffer, uint InPo
 	return InBuffer.Load2(InPointIndex * FPackedHairPositionStrideInBytes);
 }
 
-void WritePackedHairPosition(ByteAddressBuffer InBuffer, uint InPointIndex, FPackedHairPosition In)
+void WritePackedHairPosition(RWByteAddressBuffer InBuffer, uint InPointIndex, FPackedHairPosition In)
 {
 	InBuffer.Store2(InPointIndex * FPackedHairPositionStrideInBytes, In);
 }
