Fixed potential float overflow in a mobile bloom shader

diff --git a/Engine/Shaders/Private/PostProcessMobile.usf b/Engine/Shaders/Private/PostProcessMobile.usf
--- a/Engine/Shaders/Private/PostProcessMobile.usf
+++ b/Engine/Shaders/Private/PostProcessMobile.usf
@@ -391,7 +391,7 @@ void BloomPS_Mobile(
 #endif
 
 #if MOBILE_USEEYEADAPTATION
-	half Intensity = dot(AverageColor * View.OneOverPreExposure, half3(1.0f, 1.0f, 1.0f) / 3.0f);
+	float Intensity = dot(AverageColor * View.OneOverPreExposure, float3(1.0f, 1.0f, 1.0f) / 3.0f);
 	if (Intensity == 0.0)
 	{
 		// Override intensity for black pixels to MidGray. This behaviour is controlled by the cvar "r.EyeAdaptation.BlackHistogramBucketInfluence"
