Path Tracer: Avoid pack/unpack of payload in the main trace loop

In the loop over transparent hits, unpacking the payload immediately increases the live state substantially. After this change, the live state for the main TraceRay calls reduces from 404 bytes down to 282 bytes. Additionally, some of the logic around decals can be simplified slightly.

The observed speedup is around 5-10% depending on the scene with simpler scenes seeing a larger gain.

diff --git a/Engine/Shaders/Private/PathTracing/PathTracingCommon.ush b/Engine/Shaders/Private/PathTracing/PathTracingCommon.ush
--- a/Engine/Shaders/Private/PathTracing/PathTracingCommon.ush
+++ b/Engine/Shaders/Private/PathTracing/PathTracingCommon.ush
@@ -145,16 +145,18 @@ struct FPackedPathTracingPayload : FMinimalPayload
 	// float FMinimalPayload::HitT                       // 4  bytes
 #if IS_PAYLOAD_ENABLED(RT_PAYLOAD_TYPE_GPULIGHTMASS)
 	uint PackedData[7];                                  // 28 bytes (encoded data, depends on shading model and ray type)
-														 // 32 bytes total
+	// 32 bytes total
 #elif PATHTRACING_SUBSTRATE_PAYLOAD
 	// Need a bit more space for Substrate slab data
 	uint PackedData[18];                                 // 72 bytes
-	                                                     // 76 bytes total
+	// 76 bytes total
 #else
 	uint PackedData[15];								 // 60 bytes (encoded data, depends on shading model and ray type)
-														 // 64 bytes total
+	// 64 bytes total
 #endif
 
+	float UnpackBSDFOpacity() { return f16tof32(PackedData[0]); }
+	uint GetShadingModelID() { return (PackedData[0] >> 16) & 0xF; }
 	uint GetFlags() { return (PackedData[0] >> 20) & 0x1FF; }
 	void SetFlag(uint Flag) { PackedData[0] |= Flag << 20; }
 	void RemoveFlag(uint Flag) { PackedData[0] &= ~(Flag << 20); }
@@ -178,11 +180,16 @@ struct FPackedPathTracingPayload : FMinimalPayload
 			PATH_TRACING_PAYLOAD_OUTPUT_FLAG_USE_DBUFFER_LOOKUP;
 		return (GetFlags() & Mask) != 0;
 	}
-	bool UsesDBufferLookup() { return (GetFlags() & PATH_TRACING_PAYLOAD_OUTPUT_FLAG_USE_DBUFFER_LOOKUP) != 0; }
+	bool HasDecalResponseColor()     { return (GetFlags() & PATH_TRACING_PAYLOAD_OUTPUT_FLAG_DECAL_RESPONSE_COLOR    ) != 0; }
+	bool HasDecalResponseNormal()    { return (GetFlags() & PATH_TRACING_PAYLOAD_OUTPUT_FLAG_DECAL_RESPONSE_NORMAL   ) != 0; }
+	bool HasDecalResponseRoughness() { return (GetFlags() & PATH_TRACING_PAYLOAD_OUTPUT_FLAG_DECAL_RESPONSE_ROUGHNESS) != 0; }
+	bool UsesDBufferLookup()         { return (GetFlags() & PATH_TRACING_PAYLOAD_OUTPUT_FLAG_USE_DBUFFER_LOOKUP      ) != 0; }
 
 	void SetOutputDepth() { SetFlag(PATH_TRACING_PAYLOAD_OUTPUT_FLAG_OUTOUT_DEPTH); }
 	bool ShouldOutputDepth() { return (GetFlags() & PATH_TRACING_PAYLOAD_OUTPUT_FLAG_OUTOUT_DEPTH) != 0; }
 
+	bool IsHoldout() { return (GetFlags() & PATH_TRACING_PAYLOAD_OUTPUT_FLAG_IS_HOLDOUT) != 0; }
+
 	// These method are meant to be used only when IsVisibilityRay() is true
 	float3 GetRayThroughput()
 	{
@@ -211,8 +218,74 @@ struct FPackedPathTracingPayload : FMinimalPayload
 	}
 
 #if IS_PAYLOAD_ENABLED(RT_PAYLOAD_TYPE_GPULIGHTMASS)
-	// skip some unused methods in GPULM case
+	float3 UnpackWorldGeoNormal()    { return PayloadDecodeUnitVector(PackedData[1]); }
+	float3 UnpackWorldNormal()       { return PayloadDecodeUnitVector(PackedData[2]); }
+	float3 UnpackRadiance()          { return PayloadDecodeHDRColor(PackedData[3]); }
+	float3 UnpackBaseColor()         { return PayloadDecodeLDRColor(PackedData[4]); }
+	float3 UnpackTransparencyColor() { return PayloadDecodeLDRColor(PackedData[5]); }
+	float3 UnpackSubsurfaceColor()   { return PayloadDecodeLDRColor(PackedData[6]); }
+
+
+#else
+
+	float3 UnpackWorldGeoNormal() 	 { return PayloadDecodeUnitVector(PackedData[1]); }
+	float3 UnpackWorldSmoothNormal() { return PayloadDecodeUnitVector(PackedData[2]); }
+	float3 UnpackWorldNormal()       { return PayloadDecodeUnitVector(PackedData[3]); }
+	float3 UnpackRadiance()          { return PayloadDecodeHDRColor(PackedData[4]); }
+
+	void PackWorldNormal(float3 N)   { PackedData[3] = PayloadEncodeUnitVector(N); }
+	void PackRadiance(float3 R)      { PackedData[4] = PayloadEncodeHDRColor(R); }
+
+	
+#if PATHTRACING_SUBSTRATE_PAYLOAD
+
+	float3 UnpackDiffuseColor()       { return PayloadDecodeLDRColor(PackedData[5]); }
+	float3 UnpackSpecularColor()      { return PayloadDecodeLDRColor(PackedData[6]); }
+	float3 UnpackSpecularEdgeColor()  { return PayloadDecodeLDRColor(PackedData[7]); }
+	float4 UnpackRoughnessAniso()     { return PayloadDecodeRoughnessAniso(PackedData[8]); }
+	float3 UnpackWorldTangent()       { return PayloadDecodeUnitVector(PackedData[9]); }
+	float3 UnpackTransparencyColor()  { return PayloadDecodeLDRColor(PackedData[10]); }
+	float3 UnpackMeanFreePath()       { return float3(f16tof32(PackedData[11]),
+	                                                  f16tof32(PackedData[11] >> 16),
+	                                                  f16tof32(PackedData[12])); }
+	float UnpackPhaseG()              { return f16tof32(PackedData[12] >> 16); }
+	float3 UnpackFuzzColor()          { return PayloadDecodeLDRColor(PackedData[13]); }
+	float3 UnpackFuzzData()           { return PayloadDecodeLDRColor(PackedData[14]); } // x: UNUSED, y: FuzzAmount, z: FuzzRoughness
+	float3 UnpackWeightV()            { return PayloadDecodeLDRColor(PackedData[15]); }
+	float3 UnpackTransmittanceN()     { return PayloadDecodeLDRColor(PackedData[16]); }
+	float UnpackIor()                 { return Unpack10F(PackedData[17]); }
+	float UnpackCoverageAboveAlongN() { return Unpack10F(PackedData[17] >> 10); }
+	float UnpackSpecularProfileId()   { return Unpack10F(PackedData[17] >> 20); }
+
+	void PackDiffuseColor(float3 C)   { PackedData[5] = PayloadEncodeLDRColor(C); }
+	void PackSpecularColor(float3 C)  { PackedData[6] = PayloadEncodeLDRColor(C); }
+	void PackRoughnessAniso(float4 D) { PackedData[8] = PayloadEncodeRoughnessAniso(D.xyz, D.w); }
+
 #else
+	float3 UnpackBaseColor()         { return PayloadDecodeLDRColor(PackedData[5]); }
+	float3 UnpackTransparencyColor() { return PayloadDecodeLDRColor(PackedData[6]); }
+	float UnpackMetallic()           { return f16tof32(PackedData[7]); }
+	float UnpackSpecular()           { return f16tof32(PackedData[7] >> 16); }
+	float UnpackRoughness()          { return f16tof32(PackedData[8]); }
+	float UnpackAnisotropy()         { return f16tof32(PackedData[8] >> 16); }
+	float3 UnpackWorldTangent()      { return PayloadDecodeUnitVector(PackedData[9]); }
+	float4 UnpackCustomData0()       { return float4(f16tof32(PackedData[10]),
+	                                                 f16tof32(PackedData[10] >> 16),
+	                                                 f16tof32(PackedData[11]),
+	                                                 f16tof32(PackedData[11] >> 16)); }
+	float4 UnpackCustomData1()       { return float4(f16tof32(PackedData[12]),
+	                                                 f16tof32(PackedData[12] >> 16),
+	                                                 f16tof32(PackedData[13]),
+	                                                 f16tof32(PackedData[13] >> 16)); }
+	float UnpackIor()                { return f16tof32(PackedData[14]); }
+
+
+	void PackBaseColor(float3 C) { PackedData[5] = PayloadEncodeLDRColor(C); }
+	void PackMetallicSpecular(float M, float S) { PackedData[7] = f32tof16(M) | (f32tof16(S) << 16); }
+	void PackRoughnessAniso(float R, float A)   { PackedData[8] = f32tof16(R) | (f32tof16(A) << 16); }
+
+#endif
+	
 	// tau: optical depth along shadow ray, this is the integral of the extinction coefficient
 	float3 GetTau()
 	{
@@ -375,10 +448,6 @@ struct FPathTracingPayload : FMinimalPayload
 	bool IsFrontFace() { return (Flags & PATH_TRACING_PAYLOAD_OUTPUT_FLAG_FRONT_FACE) != 0; }
 
 	void SetDecalReceiver(uint DecalResponseMask) { Flags |= (DecalResponseMask & 0x07) << PATH_TRACING_PAYLOAD_OUTPUT_FLAG_DECAL_RESPONSE_MASK_SHIFT; }
-	bool IsDecalReceiver() { return (Flags & (PATH_TRACING_PAYLOAD_OUTPUT_FLAG_DECAL_RESPONSE_COLOR | PATH_TRACING_PAYLOAD_OUTPUT_FLAG_DECAL_RESPONSE_NORMAL | PATH_TRACING_PAYLOAD_OUTPUT_FLAG_DECAL_RESPONSE_ROUGHNESS)) != 0; }
-	bool HasDecalResponseColor() { return (Flags & PATH_TRACING_PAYLOAD_OUTPUT_FLAG_DECAL_RESPONSE_COLOR) != 0; }
-	bool HasDecalResponseNormal() { return (Flags & PATH_TRACING_PAYLOAD_OUTPUT_FLAG_DECAL_RESPONSE_NORMAL) != 0; }
-	bool HasDecalResponseRoughness() { return (Flags & PATH_TRACING_PAYLOAD_OUTPUT_FLAG_DECAL_RESPONSE_ROUGHNESS) != 0; }
 
 	void SetUseDBufferLookup() { Flags |= PATH_TRACING_PAYLOAD_OUTPUT_FLAG_USE_DBUFFER_LOOKUP; }
 	bool UsesDBufferLookup() { return (Flags & PATH_TRACING_PAYLOAD_OUTPUT_FLAG_USE_DBUFFER_LOOKUP) != 0; }
@@ -695,46 +764,46 @@ FPathTracingPayload UnpackPathTracingPayload(FPackedPathTracingPayload Input, FR
 
 	Output.HitT = Input.HitT;
 	Output.TranslatedWorldPos = Ray.Origin + Output.HitT * Ray.Direction;
-	Output.BSDFOpacity = f16tof32(Input.PackedData[0]);
-	Output.ShadingModelID = (Input.PackedData[0] >> 16) & 0xF;
-	Output.Flags = Input.GetFlags();
+
+	Output.BSDFOpacity                  = Input.UnpackBSDFOpacity();
+	Output.ShadingModelID               = Input.GetShadingModelID();
+	Output.Flags                        = Input.GetFlags();
 	Output.PrimitiveLightingChannelMask = Input.GetPrimitiveLightingChannelMask();
 #if IS_PAYLOAD_ENABLED(RT_PAYLOAD_TYPE_GPULIGHTMASS)
-	Output.WorldGeoNormal = PayloadDecodeUnitVector(Input.PackedData[1]);
-	Output.WorldNormal = PayloadDecodeUnitVector(Input.PackedData[2]);
-	Output.Radiance = PayloadDecodeHDRColor(Input.PackedData[3]);
-	Output.BaseColor = PayloadDecodeLDRColor(Input.PackedData[4]);
-	Output.TransparencyColor = PayloadDecodeLDRColor(Input.PackedData[5]);
-	Output.SubsurfaceColor = PayloadDecodeLDRColor(Input.PackedData[6]);
+	Output.WorldGeoNormal    = Input.UnpackWorldGeoNormal();
+	Output.WorldNormal       = Input.UnpackWorldNormal();
+	Output.Radiance          = Input.UnpackRadiance();
+	Output.BaseColor         = Input.UnpackBaseColor();
+	Output.TransparencyColor = Input.UnpackTransparencyColor();
+	Output.SubsurfaceColor   = Input.UnpackSubsurfaceColor();
 #else
-	Output.WorldGeoNormal    = PayloadDecodeUnitVector(Input.PackedData[1]);
-	Output.WorldSmoothNormal = PayloadDecodeUnitVector(Input.PackedData[2]);
-	Output.WorldNormal       = PayloadDecodeUnitVector(Input.PackedData[3]);
-	Output.Radiance          = PayloadDecodeHDRColor(Input.PackedData[4]);
+	Output.WorldGeoNormal    = Input.UnpackWorldGeoNormal();
+	Output.WorldSmoothNormal = Input.UnpackWorldSmoothNormal();
+	Output.WorldNormal       = Input.UnpackWorldNormal();
+	Output.Radiance          = Input.UnpackRadiance();
 
 #if PATHTRACING_SUBSTRATE_PAYLOAD
-	Output.DiffuseColor        = PayloadDecodeLDRColor(Input.PackedData[5]);
-	Output.SpecularColor       = PayloadDecodeLDRColor(Input.PackedData[6]);
-	Output.SpecularEdgeColor   = PayloadDecodeLDRColor(Input.PackedData[7]);
-	float4 RoughnessAniso      = PayloadDecodeRoughnessAniso(Input.PackedData[8]); 
+	Output.DiffuseColor        = Input.UnpackDiffuseColor();
+	Output.SpecularColor       = Input.UnpackSpecularColor();
+	Output.SpecularEdgeColor   = Input.UnpackSpecularEdgeColor();
+	float4 RoughnessAniso      = Input.UnpackRoughnessAniso();
 	Output.RoughnessData       = RoughnessAniso.xyz;
 	Output.Anisotropy          = RoughnessAniso.w;
-	Output.WorldTangent        = PayloadDecodeUnitVector(Input.PackedData[9]);
-	Output.TransparencyColor   = PayloadDecodeLDRColor(Input.PackedData[10]);
-	Output.MeanFreePath.x      = f16tof32(Input.PackedData[11]);
-	Output.MeanFreePath.y      = f16tof32(Input.PackedData[11] >> 16);
-	Output.MeanFreePath.z      = f16tof32(Input.PackedData[12]);
-	Output.PhaseG              = f16tof32(Input.PackedData[12] >> 16);
-	Output.FuzzColor           = PayloadDecodeLDRColor(Input.PackedData[13]);
-	const float3 Data14        = PayloadDecodeLDRColor(Input.PackedData[14]);
+	Output.WorldTangent        = Input.UnpackWorldTangent();
+	Output.TransparencyColor   = Input.UnpackTransparencyColor();
+	Output.MeanFreePath        = Input.UnpackMeanFreePath();
+	Output.PhaseG              = Input.UnpackPhaseG();
+	Output.FuzzColor           = Input.UnpackFuzzColor();
+	const float3 Data14        = Input.UnpackFuzzData();
 	Output.FuzzAmount          = Data14.y;
 	Output.FuzzRoughness       = Data14.z;
-	Output.WeightV             = PayloadDecodeLDRColor(Input.PackedData[15]);
-	Output.TransmittanceN      = PayloadDecodeLDRColor(Input.PackedData[16]);
-	Output.Ior                 = Unpack10F(Input.PackedData[17]);
-	Output.CoverageAboveAlongN = Unpack10F(Input.PackedData[17] >> 10);
-	Output.SpecularProfileId   = Unpack10F(Input.PackedData[17] >> 20);
+	Output.WeightV             = Input.UnpackWeightV();
+	Output.TransmittanceN      = Input.UnpackTransmittanceN();
+	Output.Ior                 = Input.UnpackIor();
+	Output.CoverageAboveAlongN = Input.UnpackCoverageAboveAlongN();
+	Output.SpecularProfileId   = Input.UnpackSpecularProfileId();
 	
+	// TODO: figure out how to make glint unpacking closer to functions above
 	// In order to not increase the data payload size, we alias SSS/Fuzz parameters when glint are enabled
 	const bool bHasGlints = ((Input.PackedData[17] >> 30u) & 0x1u) > 0u;
 	if (bHasGlints)
@@ -753,29 +822,19 @@ FPathTracingPayload UnpackPathTracingPayload(FPackedPathTracingPayload Input, FR
 	}
 
 #else // PATHTRACING_SUBSTRATE_PAYLOAD
-	Output.BaseColor = PayloadDecodeLDRColor(Input.PackedData[5]);
-	Output.TransparencyColor = PayloadDecodeLDRColor(Input.PackedData[6]);
-	
-	Output.Metallic = f16tof32(Input.PackedData[7]);
-	Output.Specular = f16tof32(Input.PackedData[7] >> 16);
-	Output.Roughness = f16tof32(Input.PackedData[8]);
-	Output.Anisotropy = f16tof32(Input.PackedData[8] >> 16);
-
-	Output.WorldTangent = PayloadDecodeUnitVector(Input.PackedData[9]);
-
-	Output.CustomData0.x = f16tof32(Input.PackedData[10]);
-	Output.CustomData0.y = f16tof32(Input.PackedData[10] >> 16);
-	Output.CustomData0.z = f16tof32(Input.PackedData[11]);
-	Output.CustomData0.w = f16tof32(Input.PackedData[11] >> 16);
-
-	Output.CustomData1.x = f16tof32(Input.PackedData[12]);
-	Output.CustomData1.y = f16tof32(Input.PackedData[12] >> 16);
-	Output.CustomData1.z = f16tof32(Input.PackedData[13]);
-	Output.CustomData1.w = f16tof32(Input.PackedData[13] >> 16);
-
-	Output.Ior = f16tof32(Input.PackedData[14]);
-
-	Output.DiffuseColor = Output.BaseColor - Output.BaseColor * Output.Metallic;
+	Output.BaseColor         = Input.UnpackBaseColor();
+	Output.TransparencyColor = Input.UnpackTransparencyColor();
+	Output.Metallic          = Input.UnpackMetallic();
+	Output.Specular          = Input.UnpackSpecular();
+	Output.Roughness         = Input.UnpackRoughness();
+	Output.Anisotropy        = Input.UnpackAnisotropy();
+	Output.WorldTangent      = Input.UnpackWorldTangent();
+	Output.CustomData0       = Input.UnpackCustomData0();
+	Output.CustomData1       = Input.UnpackCustomData1();
+	Output.Ior               = Input.UnpackIor();
+
+	// 
+	Output.DiffuseColor  = Output.BaseColor - Output.BaseColor * Output.Metallic;
 	Output.SpecularColor = ComputeF0(Output.Specular, Output.BaseColor, Output.Metallic);
 #endif // PATHTRACING_SUBSTRATE_PAYLOAD
 #endif // IS_PAYLOAD_ENABLED(RT_PAYLOAD_TYPE_GPULIGHTMASS)
diff --git a/Engine/Shaders/Private/PathTracing/PathTracingCore.ush b/Engine/Shaders/Private/PathTracing/PathTracingCore.ush
--- a/Engine/Shaders/Private/PathTracing/PathTracingCore.ush
+++ b/Engine/Shaders/Private/PathTracing/PathTracingCore.ush
@@ -222,27 +222,42 @@ bool HasDecalData(FDecalShaderPayload DecalPayload)
 	return min3(DecalPayload.GetBaseColor().a, DecalPayload.GetWorldNormal().a, DecalPayload.GetMetallicSpecularRoughness().a) < 1.0;
 }
 
-bool ApplyDecal(FDecalShaderPayload DecalPayload, inout FPathTracingPayload HitPayload)
+void ApplyDecal(FDecalShaderPayload DecalPayload, inout FPackedPathTracingPayload PackedPayload, float3 RayDirection, float PathRoughness)
 {
-#if PATHTRACING_SUBSTRATE_PAYLOAD
-	float Roughness = HitPayload.RoughnessData.x;
-	float Metallic = F0RGBToMetallic(HitPayload.SpecularColor);
-	float Specular = F0RGBToDielectricSpecular(HitPayload.SpecularColor);
-	float3 BaseColor = lerp(HitPayload.DiffuseColor, HitPayload.SpecularColor, Metallic);
-#else
-	float3 BaseColor = HitPayload.BaseColor;
-	float Metallic = HitPayload.Metallic;
-	float Specular = HitPayload.Specular;
-	float Roughness = HitPayload.Roughness;
-#endif // PATHTRACING_SUBSTRATE_PAYLOAD
+	{
+		float3 Radiance = PackedPayload.UnpackRadiance();
+		Radiance += DecalPayload.GetEmissive();
+		PackedPayload.PackRadiance(Radiance);
+	}
 
 	float4 DecalBaseColor = DecalPayload.GetBaseColor();
 	float4 DecalWorldNormal = DecalPayload.GetWorldNormal();
 	float4 DecalMetallicSpecularRoughness = DecalPayload.GetMetallicSpecularRoughness();
 
-	bool bHasColorResponse = HitPayload.HasDecalResponseColor() && DecalBaseColor.a < 1.0;
-	bool bHasNormalResponse = HitPayload.HasDecalResponseNormal() && DecalWorldNormal.a < 1.0;
-	bool bHasRoughnessResponse = HitPayload.HasDecalResponseRoughness() && DecalMetallicSpecularRoughness.a < 1.0;
+	bool bHasColorResponse = PackedPayload.HasDecalResponseColor() && DecalBaseColor.a < 1.0;
+	bool bHasNormalResponse = PackedPayload.HasDecalResponseNormal() && DecalWorldNormal.a < 1.0;
+	bool bHasRoughnessResponse = PackedPayload.HasDecalResponseRoughness() && DecalMetallicSpecularRoughness.a < 1.0;
+
+	if (!bHasColorResponse && !bHasNormalResponse && !bHasRoughnessResponse)
+	{
+		// nothing to do
+		return;
+	}
+
+#if PATHTRACING_SUBSTRATE_PAYLOAD
+	float4 RoughnessData = PackedPayload.UnpackRoughnessAniso();
+	float Roughness = RoughnessData.x;
+	float3 DiffuseColor = PackedPayload.UnpackDiffuseColor();
+	float3 SpecularColor = PackedPayload.UnpackSpecularColor();
+	float Metallic = F0RGBToMetallic(SpecularColor);
+	float Specular = F0RGBToDielectricSpecular(SpecularColor);
+	float3 BaseColor = lerp(DiffuseColor, SpecularColor, Metallic);
+#else
+	float3 BaseColor = PackedPayload.UnpackBaseColor();
+	float Metallic = PackedPayload.UnpackMetallic();
+	float Specular = PackedPayload.UnpackSpecular();
+	float Roughness = PackedPayload.UnpackRoughness();
+#endif // PATHTRACING_SUBSTRATE_PAYLOAD
 
 	if (bHasColorResponse)
 	{
@@ -252,7 +267,12 @@ bool ApplyDecal(FDecalShaderPayload DecalPayload, inout FPathTracingPayload HitP
 	if (bHasNormalResponse)
 	{
 		// We normalize the normal to get smoother visual result (it helps to avoid having D_GGX explodes toward infinity, and matches ApplyDBufferData(...))
-		HitPayload.WorldNormal = normalize(HitPayload.WorldNormal * DecalWorldNormal.a + DecalWorldNormal.rgb);
+		float3 WorldNormal = PackedPayload.UnpackWorldNormal();
+		WorldNormal = normalize(WorldNormal * DecalWorldNormal.a + DecalWorldNormal.rgb);
+		// re-apply this (normally happens as part of shading, but since decal changed the normal, apply it again)
+		WorldNormal = AdjustShadingNormal(WorldNormal, PackedPayload.UnpackWorldGeoNormal(), RayDirection);
+
+		PackedPayload.PackWorldNormal(WorldNormal);
 	}
 
 	if (bHasRoughnessResponse)
@@ -260,26 +280,32 @@ bool ApplyDecal(FDecalShaderPayload DecalPayload, inout FPathTracingPayload HitP
 		Metallic  = Metallic  * DecalMetallicSpecularRoughness.a + DecalMetallicSpecularRoughness.x;
 		Specular  = Specular  * DecalMetallicSpecularRoughness.a + DecalMetallicSpecularRoughness.y;
 		Roughness = Roughness * DecalMetallicSpecularRoughness.a + DecalMetallicSpecularRoughness.z;
-	}
 
-	HitPayload.Radiance += DecalPayload.GetEmissive();
+		if (ApproximateCaustics)
+		{
+			Roughness = max(Roughness, PathRoughness);
+		}
+	}
 
 #if PATHTRACING_SUBSTRATE_PAYLOAD
 	if (bHasColorResponse || bHasRoughnessResponse)
 	{
 		// the encode-decode is potentially lossy, so only do it if we actually need to
-		HitPayload.DiffuseColor = BaseColor - BaseColor * Metallic;
-		HitPayload.SpecularColor = lerp(DielectricSpecularToF0(Specular), BaseColor, Metallic);
-		HitPayload.RoughnessData.x = Roughness;
+		PackedPayload.PackDiffuseColor(BaseColor - BaseColor * Metallic);
+		PackedPayload.PackSpecularColor(lerp(DielectricSpecularToF0(Specular), BaseColor, Metallic));
+		PackedPayload.PackRoughnessAniso(float4(Roughness, RoughnessData.yzw));
 	}
 #else
-	// no need for an encode-decode here, so safe to overwrite even if nothing changes
-	HitPayload.BaseColor = BaseColor;
-	HitPayload.Metallic = Metallic;
-	HitPayload.Specular = Specular;
-	HitPayload.Roughness = Roughness;
+	if (bHasColorResponse)
+	{
+		PackedPayload.PackBaseColor(BaseColor);
+	}
+	if (bHasRoughnessResponse)
+	{
+		PackedPayload.PackMetallicSpecular(Metallic, Specular);
+		PackedPayload.PackRoughnessAniso(Roughness, PackedPayload.UnpackAnisotropy());
+	}
 #endif
-	return bHasColorResponse || bHasNormalResponse || bHasRoughnessResponse;
 }
 
 void EvaluateDecals(FPathState PathState, uint InstanceInclusionMask, float HitT, inout FDecalShaderPayload DecalPayload)
@@ -336,7 +362,7 @@ void EvaluateDecals(FPathState PathState, uint InstanceInclusionMask, float HitT
 #endif
 }
 
-FPathTracingPayload TraceTransparentRay(inout FPathState PathState, int Bounce, inout FVolumeSegment VolumeSegment)
+FPackedPathTracingPayload TraceTransparentRay(inout FPathState PathState, int Bounce, inout FVolumeSegment VolumeSegment)
 {
 	const bool bIsCameraRay = Bounce == 0;
 	const bool bLastBounce = Bounce == MaxBounces;
@@ -355,7 +381,7 @@ FPathTracingPayload TraceTransparentRay(inout FPathState PathState, int Bounce,
 	FRISContext VolSample = InitRISContext(RISRandSample.y);
 
 	float3 PayloadThroughput = PathState.PathThroughput;
-	FPathTracingPayload Payload;
+	FPackedPathTracingPayload Payload;
 	if (bLastBounce && !PathState.ShouldAccumulateEmissive())
 	{
 		Payload.SetMiss();
@@ -499,7 +525,7 @@ FPathTracingPayload TraceTransparentRay(inout FPathState PathState, int Bounce,
 							if (HitSample.Accept(SelectionWeight))
 							{
 								// stash this hit for next time
-								Payload = CreateMediumHitPayload(Interval.VolumeTMin, TranslatedWorldPos, Result);
+								Payload = PackPathTracingPayload(CreateMediumHitPayload(Interval.VolumeTMin, TranslatedWorldPos, Result));
 								PayloadThroughput = Contrib / SelectionWeight;
 							}
 						}
@@ -583,58 +609,43 @@ FPathTracingPayload TraceTransparentRay(inout FPathState PathState, int Bounce,
 #endif
 		}
 
-		// Unpack the payload
-		FPathTracingPayload HitPayload = UnpackPathTracingPayload(PackedPayload, PathState.Ray);
-		if (ApplyDecal(DecalPayload, HitPayload))
-		{
-			// if we ran any decal logic, update the payload
-			if (ApproximateCaustics)
-			{
-#if PATHTRACING_SUBSTRATE_PAYLOAD
-				HitPayload.RoughnessData.x = max(HitPayload.RoughnessData.x, PathState.PathRoughness);
-#else
-				HitPayload.Roughness = max(HitPayload.Roughness, PathState.PathRoughness);
-#endif
-			}
-			HitPayload.WorldNormal = AdjustShadingNormal(HitPayload.WorldNormal, HitPayload.WorldGeoNormal, PathState.Ray.Direction);
-#if !PATHTRACING_SUBSTRATE_PAYLOAD
-			HitPayload.DiffuseColor = HitPayload.BaseColor - HitPayload.BaseColor * HitPayload.Metallic;
-			HitPayload.SpecularColor = ComputeF0(HitPayload.Specular, HitPayload.BaseColor, HitPayload.Metallic);
-#endif
-		}
-		AdjustPayloadAfterUnpack(HitPayload, bIsCameraRay && ApplyDiffuseSpecularOverrides);
-
+		// Now that we know the real decal data, apply it
+		ApplyDecal(DecalPayload, PackedPayload, PathState.Ray.Direction, PathState.PathRoughness);
 
 		// account for Beer's law through the currently active medium
 		// TODO: merge this with the volume raymarching?
 		// NOTE: select statement is to avoid corner case of 0.0*inf that can occur with directly visible skydomes
-		PathState.PathThroughput *= select(PathState.SigmaT > 0.0, exp(-PathState.SigmaT * (HitPayload.HitT - PathState.Ray.TMin)), 1.0);
+		PathState.PathThroughput *= select(PathState.SigmaT > 0.0, exp(-PathState.SigmaT * (PackedPayload.HitT - PathState.Ray.TMin)), 1.0);
 
-		if (!HitPayload.IsHoldout() && bIsCameraRay)
+		if (!PackedPayload.IsHoldout() && bIsCameraRay)
 		{
-			PathState.Alpha += Luminance(PathState.PathThroughput * (1.0 - HitPayload.TransparencyColor));
+			PathState.Alpha += Luminance(PathState.PathThroughput * (1.0 - PackedPayload.UnpackTransparencyColor()));
 		}
 
 		// add in surface emission
 		if (PathState.ShouldAccumulateEmissive())
 		{
-			AccumulateRadiance(PathState.Radiance, PathState.PathThroughput * HitPayload.Radiance, bIsCameraRay);
+			AccumulateRadiance(PathState.Radiance, PathState.PathThroughput * PackedPayload.UnpackRadiance(), bIsCameraRay);
 		}
 
-		if (HitPayload.HitT == POSITIVE_INFINITY)
+		if (PackedPayload.HitT == POSITIVE_INFINITY)
 		{
 			// if our hit was against an infinite light, exit now
 			break;
 		}
 
-		if (bIsCameraRay && HitPayload.ShouldOutputDepth() && PathState.DepthZ == 0.0)
+		if (bIsCameraRay && PackedPayload.ShouldOutputDepth() && PathState.DepthZ == 0.0)
 		{
 			const float3 ViewZ = View.ViewToTranslatedWorld[2].xyz;
-			PathState.DepthZ = dot(HitPayload.TranslatedWorldPos, ViewZ);
+			const float3 TranslatedWorldPos = PathState.Ray.Origin + PackedPayload.HitT * PathState.Ray.Direction;
+			PathState.DepthZ = dot(TranslatedWorldPos, ViewZ);
 		}
 
 		if (!bLastBounce)
 		{
+			FPathTracingPayload HitPayload = UnpackPathTracingPayload(PackedPayload, PathState.Ray);
+			AdjustPayloadAfterUnpack(HitPayload, bIsCameraRay && ApplyDiffuseSpecularOverrides);
+
 			float3 Contrib = PathState.PathThroughput * EstimateMaterialAlbedo(HitPayload);
 
 			// accumulate what the denoiser wants into albedo/normal (as long as the current path is rough enough)
@@ -644,17 +655,17 @@ FPathTracingPayload TraceTransparentRay(inout FPathState PathState, int Bounce,
 			if (HitSample.Accept(SelectionWeight))
 			{
 				// stash this hit for next time
-				Payload = HitPayload;
+				Payload = PackedPayload;
 				PayloadThroughput = PathState.PathThroughput / SelectionWeight;
 			}
 		}
 
 		// account for local transparency change
-		PathState.PathThroughput *= HitPayload.TransparencyColor;
+		PathState.PathThroughput *= PackedPayload.UnpackTransparencyColor();
 
 		// prepare next step around the loop
 		// retrace the exact same ray with TMin one ulp past the hit we just found
-		PathState.Ray.TMin = asfloat(asuint(HitPayload.HitT) + 1);
+		PathState.Ray.TMin = asfloat(asuint(PackedPayload.HitT) + 1);
 
 		if (all(PathState.PathThroughput == 0))
 		{
@@ -1314,7 +1325,7 @@ bool PathTracingKernel(inout FPathState PathState, int Bounce)
 	const bool bIsCameraRay = Bounce == 0;
 
 	FVolumeSegment VolumeSegment = CreateEmptyVolumeSegment();
-	FPathTracingPayload Payload = TraceTransparentRay(
+	FPackedPathTracingPayload PackedPayload = TraceTransparentRay(
 		PathState,
 		Bounce,
 		VolumeSegment);
@@ -1437,12 +1448,15 @@ bool PathTracingKernel(inout FPathState PathState, int Bounce)
 		}
 	}
 
-	if (Payload.IsMiss())
+	if (PackedPayload.IsMiss())
 	{
 		// we didn't hit anything selectable for further shading, we are done
 		return false;
 	}
 
+	FPathTracingPayload Payload = UnpackPathTracingPayload(PackedPayload, PathState.Ray);
+	AdjustPayloadAfterUnpack(Payload, bIsCameraRay && ApplyDiffuseSpecularOverrides);
+
 #if 0
 	PathState.Radiance = 0.5 * Payload.WorldGeoNormal + 0.5;
 	return false;
