Path Tracing: Optimize storage of volume intervals during tracing

Use bitmasks to track enabled volumes instead of an array of uint IDs. This further reduces the live state of the main TraceRay call. It also simplifies the logic for volume tracing, giving an overall speedup around 5-10% (depending on the scene).

diff --git a/Engine/Shaders/Private/PathTracing/PathTracingCore.ush b/Engine/Shaders/Private/PathTracing/PathTracingCore.ush
--- a/Engine/Shaders/Private/PathTracing/PathTracingCore.ush
+++ b/Engine/Shaders/Private/PathTracing/PathTracingCore.ush
@@ -413,7 +413,7 @@ FPackedPathTracingPayload TraceTransparentRay(inout FPathState PathState, int Bo
 		{
 			// we didn't hit any real geometry, but we did hit the volume's blocker geometry
 			// create a virtual hit with it here
-			PackedPayload = VolumeGetBlockerHit(VolumeIsectList.BlockerID, PathState.Ray.Origin, PathState.Ray.Direction, VolumeIsectList.BlockerHitT, bIsCameraRay);
+			PackedPayload = VolumeGetBlockerHit(PathState.Ray.Origin, PathState.Ray.Direction, VolumeIsectList.BlockerHitT, bIsCameraRay);
 		}
 
 		// Loop over lights to capture their contribution
diff --git a/Engine/Shaders/Private/PathTracing/Volume/PathTracingVolume.ush b/Engine/Shaders/Private/PathTracing/Volume/PathTracingVolume.ush
--- a/Engine/Shaders/Private/PathTracing/Volume/PathTracingVolume.ush
+++ b/Engine/Shaders/Private/PathTracing/Volume/PathTracingVolume.ush
@@ -32,9 +32,10 @@ FVolumeIntersectionList VolumeIntersect(float3 Origin, float3 Direction, float T
 	return Result;
 }
 
-FPackedPathTracingPayload VolumeGetBlockerHit(uint VolumeID, float3 Origin, float3 Direction, float HitT, bool bIsCameraRay)
+FPackedPathTracingPayload VolumeGetBlockerHit(float3 Origin, float3 Direction, float HitT, bool bIsCameraRay)
 {
 	// NOTE: only atmosphere supports blockers, so if this gets called it is safe to assume its for the atmosphere volume
+	// If we ever add other volume types that imply blockers, we would have to add the ID of the volume type as input
 	return AtmosphereGetBlockerHit(Origin, Direction, HitT, bIsCameraRay);
 }
 
@@ -42,15 +43,21 @@ FPackedPathTracingPayload VolumeGetBlockerHit(uint VolumeID, float3 Origin, floa
 FVolumeDensityBounds VolumeGetDensityBounds(float3 Origin, float3 Direction, FVolumeIntersectionInterval Interval)
 {
 	FVolumeDensityBounds Result = CreateVolumeDensityBound(0, 0);
-	for (int Index = 0; Index < Interval.Num; Index++)
+	if (Interval.VolumeMask & PATH_TRACER_VOLUME_ENABLE_ATMOSPHERE)
 	{
-		switch (Interval.VolumeID[Index])
-		{
-			case VOLUMEID_ATMOSPHERE:            MergeVolumeDensityBounds(Result, AtmosphereGetDensityBounds(Origin, Direction, Interval.VolumeTMin, Interval.VolumeTMax)); break;
-			case VOLUMEID_FOG:                   MergeVolumeDensityBounds(Result, FogGetDensityBounds(Origin, Direction, Interval.VolumeTMin, Interval.VolumeTMax)); break;
-			//case VOLUMEID_HETEROGENEOUS_VOLUMES: MergeVolumeDensityBounds(Result, HeterogeneousVolumesGetDensityBounds(Origin, Direction, Interval.VolumeTMin, Interval.VolumeTMax)); break;
-		}
+		MergeVolumeDensityBounds(Result, AtmosphereGetDensityBounds(Origin, Direction, Interval.VolumeTMin, Interval.VolumeTMax));
+	}
+	if (Interval.VolumeMask & PATH_TRACER_VOLUME_ENABLE_FOG)
+	{
+		MergeVolumeDensityBounds(Result, FogGetDensityBounds(Origin, Direction, Interval.VolumeTMin, Interval.VolumeTMax));
+	}
+#if 0
+	// NOTE: this case is not possible because we don't use density bounds directly when heterogeneous volumes are present
+	if (Interval.VolumeMask & PATH_TRACER_VOLUME_ENABLE_HETEROGENEOUS_VOLUMES)
+	{
+		MergeVolumeDensityBounds(Result, HeterogeneousVolumesGetDensityBounds(Origin, Direction, Interval.VolumeTMin, Interval.VolumeTMax));
 	}
+#endif
 	return Result;
 }
 
@@ -58,62 +65,21 @@ FVolumeDensityBounds VolumeGetDensityBounds(float3 Origin, float3 Direction, FVo
 FVolumeShadedResult VolumeGetDensity(float3 TranslatedWorldPos, FVolumeIntersectionInterval Interval, bool bIsCameraRay)
 {
 	FVolumeShadedResult Result = (FVolumeShadedResult)0;
-	for (int Index = 0; Index < Interval.Num; Index++)
+	if (Interval.VolumeMask & PATH_TRACER_VOLUME_ENABLE_ATMOSPHERE)
 	{
-		FVolumeShadedResult LocalResult = (FVolumeShadedResult)0;
-		switch (Interval.VolumeID[Index])
-		{
-			case VOLUMEID_ATMOSPHERE:            LocalResult = AtmosphereGetDensity(TranslatedWorldPos); break;
-			case VOLUMEID_FOG:                   LocalResult = FogGetDensity(TranslatedWorldPos); break;
-			case VOLUMEID_HETEROGENEOUS_VOLUMES: LocalResult = HeterogeneousVolumesGetDensity(TranslatedWorldPos); break;
-		}
-		if (bIsCameraRay && (VolumeFlags & PATH_TRACER_VOLUME_HOLDOUT_BIT << Interval.VolumeID[Index]) != 0)
-		{
-			// remove scattering and set holdout density
-			LocalResult.SigmaH = LocalResult.SigmaT;
-			LocalResult.SigmaSHG = LocalResult.SigmaSRayleigh = 0;
-		}
-		MergeVolumeShadedResult(Result, LocalResult);
+		MergeVolumeShadedResult(Result, AtmosphereGetDensity(TranslatedWorldPos), bIsCameraRay && (VolumeFlags & PATH_TRACER_VOLUME_HOLDOUT_ATMOSPHERE) != 0);
 	}
-	return Result;
-}
-
-
-// Return the transmittance along a ray segment
-float3 VolumeGetTransmittanceAnalytical(float3 StartThroughput, float3 Origin, float3 Direction, FVolumeIntersectionInterval Interval, inout RandomSequence RandSequence)
-{
-	float3 Throughput = StartThroughput;
-	for (int Index = 0; Index < Interval.Num; Index++)
+	if (Interval.VolumeMask & PATH_TRACER_VOLUME_ENABLE_FOG)
 	{
-		switch (Interval.VolumeID[Index])
-		{
-			case VOLUMEID_ATMOSPHERE:            Throughput *= AtmosphereGetTransmittance(Origin, Direction, Interval.VolumeTMin, Interval.VolumeTMax); break;
-			case VOLUMEID_FOG:                   Throughput *= FogGetTransmittance(Origin, Direction, Interval.VolumeTMin, Interval.VolumeTMax); break;
-			case VOLUMEID_HETEROGENEOUS_VOLUMES: Throughput = HeterogeneousVolumesGetTransmittance(Throughput, Origin, Direction, Interval.VolumeTMin, Interval.VolumeTMax, RandSequence); break;
-		}
+		MergeVolumeShadedResult(Result, FogGetDensity(TranslatedWorldPos), bIsCameraRay && (VolumeFlags & PATH_TRACER_VOLUME_HOLDOUT_FOG) != 0);
 	}
-	return Throughput;
-}
-
-// Return the transmittance along a whole intersection list
-float3 VolumeGetTransmittanceAnalytical(float3 StartThroughput, float3 Origin, float3 Direction, FVolumeIntersectionList IntersectionList, inout RandomSequence RandSequence)
-{
-	float3 Throughput = StartThroughput;
-	for (int Index = 0; Index < IntersectionList.Num; Index++)
+	if (Interval.VolumeMask & PATH_TRACER_VOLUME_ENABLE_HETEROGENEOUS_VOLUMES)
 	{
-		float TMin = IntersectionList.VolumeTMin[Index];
-		float TMax = IntersectionList.VolumeTMax[Index];
-		switch (IntersectionList.VolumeID[Index])
-		{
-			case VOLUMEID_ATMOSPHERE:            Throughput *= AtmosphereGetTransmittance(Origin, Direction, TMin, TMax); break;
-			case VOLUMEID_FOG:                   Throughput *= FogGetTransmittance(Origin, Direction, TMin, TMax); break;
-			case VOLUMEID_HETEROGENEOUS_VOLUMES: Throughput = HeterogeneousVolumesGetTransmittance(Throughput, Origin, Direction, TMin, TMax, RandSequence); break;
-		}
+		MergeVolumeShadedResult(Result, HeterogeneousVolumesGetDensity(TranslatedWorldPos), bIsCameraRay && (VolumeFlags & PATH_TRACER_VOLUME_HOLDOUT_HETEROGENEOUS_VOLUMES) != 0);
 	}
-	return Throughput;
+	return Result;
 }
 
-
 // Reference implementation of transmittance that uses only GetDensity and GetDensityBounds
 // express the inner loop here as a macro so we can stamp down a separate copy per volume ID and avoid a nested loop during ray marching
 #define PATH_TRACER_REFERENCE_TRANSMITTANCE_LOOP(GetDensityBoundsFunc, GetDensityFunc)				\
@@ -155,85 +121,84 @@ float3 VolumeGetTransmittanceAnalytical(float3 StartThroughput, float3 Origin, f
 
 
 // The returned value factors in the initial throughput. This is used to avoid sampling too much if throughput reaches 0.
-float3 VolumeGetTransmittanceNullTracking(float3 StartThroughput, float3 Origin, float3 Direction, FVolumeIntersectionInterval Interval, inout RandomSequence RandSequence)
+// Return the transmittance along a ray segment
+float3 VolumeGetTransmittance(float3 StartThroughput, float3 Origin, float3 Direction, FVolumeIntersectionInterval Interval, inout RandomSequence RandSequence)
 {
 	float3 Throughput = StartThroughput;
-	for (int Index = 0; Index < Interval.Num; Index++)
+	if (Interval.VolumeMask & PATH_TRACER_VOLUME_ENABLE_ATMOSPHERE)
 	{
 		float TMin = Interval.VolumeTMin;
 		float TMax = Interval.VolumeTMax;
-		switch (Interval.VolumeID[Index])
+		if (VolumeFlags & PATH_TRACER_VOLUME_USE_ANALYTIC_TRANSMITTANCE)
 		{
-			case VOLUMEID_ATMOSPHERE:
-			{
-				PATH_TRACER_REFERENCE_TRANSMITTANCE_LOOP(AtmosphereGetDensityBounds, AtmosphereGetDensity);
-				break;
-			}
-			case VOLUMEID_FOG:
-			{
-				PATH_TRACER_REFERENCE_TRANSMITTANCE_LOOP(FogGetDensityBounds, FogGetDensity);
-				break;
-			}
-			case VOLUMEID_HETEROGENEOUS_VOLUMES:
-			{
-				Throughput = HeterogeneousVolumesGetTransmittance(Throughput, Origin, Direction, TMin, TMax, RandSequence);
-				break;
-			}
+			Throughput *= AtmosphereGetTransmittance(Origin, Direction, TMin, TMax);
+		}
+		else
+		{
+			PATH_TRACER_REFERENCE_TRANSMITTANCE_LOOP(AtmosphereGetDensityBounds, AtmosphereGetDensity);
 		}
 	}
+	if (Interval.VolumeMask & PATH_TRACER_VOLUME_ENABLE_FOG)
+	{
+		float TMin = Interval.VolumeTMin;
+		float TMax = Interval.VolumeTMax;
+		if (VolumeFlags & PATH_TRACER_VOLUME_USE_ANALYTIC_TRANSMITTANCE)
+		{
+			Throughput *= FogGetTransmittance(Origin, Direction, TMin, TMax);
+		}
+		else
+		{
+			PATH_TRACER_REFERENCE_TRANSMITTANCE_LOOP(FogGetDensityBounds, FogGetDensity);
+		}
+	}
+	if (Interval.VolumeMask & PATH_TRACER_VOLUME_ENABLE_HETEROGENEOUS_VOLUMES)
+	{
+		float TMin = Interval.VolumeTMin;
+		float TMax = Interval.VolumeTMax;
+		// NOTE: there is analytic mode in this case as the presence of heterogeneous volumes disables analytic transmittance
+		Throughput = HeterogeneousVolumesGetTransmittance(Throughput, Origin, Direction, TMin, TMax, RandSequence);
+	}
 	return Throughput;
-
 }
 
-float3 VolumeGetTransmittanceNullTracking(float3 StartThroughput, float3 Origin, float3 Direction, FVolumeIntersectionList IntersectionList, inout RandomSequence RandSequence)
+// Return the transmittance along a whole intersection list
+float3 VolumeGetTransmittance(float3 StartThroughput, float3 Origin, float3 Direction, FVolumeIntersectionList IntersectionList, inout RandomSequence RandSequence)
 {
 	float3 Throughput = StartThroughput;
-	for (int Index = 0; Index < IntersectionList.Num; Index++)
+	if (IntersectionList.VolumeMask & PATH_TRACER_VOLUME_ENABLE_ATMOSPHERE)
 	{
-		float TMin = IntersectionList.VolumeTMin[Index];
-		float TMax = IntersectionList.VolumeTMax[Index];
-		switch (IntersectionList.VolumeID[Index])
+		float TMin = IntersectionList.VolumeTMin[VOLUMEID_ATMOSPHERE];
+		float TMax = IntersectionList.VolumeTMax[VOLUMEID_ATMOSPHERE];
+		if (VolumeFlags & PATH_TRACER_VOLUME_USE_ANALYTIC_TRANSMITTANCE)
 		{
-			case VOLUMEID_ATMOSPHERE:
-			{
-				PATH_TRACER_REFERENCE_TRANSMITTANCE_LOOP(AtmosphereGetDensityBounds, AtmosphereGetDensity);
-				break;
-			}
-			case VOLUMEID_FOG:
-			{
-				PATH_TRACER_REFERENCE_TRANSMITTANCE_LOOP(FogGetDensityBounds, FogGetDensity);
-				break;
-			}
-			case VOLUMEID_HETEROGENEOUS_VOLUMES:
-			{
-				Throughput = HeterogeneousVolumesGetTransmittance(Throughput, Origin, Direction, TMin, TMax, RandSequence);
-				break;
-			}
+			Throughput *= AtmosphereGetTransmittance(Origin, Direction, TMin, TMax);
+		}
+		else
+		{
+			PATH_TRACER_REFERENCE_TRANSMITTANCE_LOOP(AtmosphereGetDensityBounds, AtmosphereGetDensity);
 		}
 	}
-	return Throughput;
-
-#undef PATH_TRACER_REFERENCE_TRANSMITTANCE_LOOP
-}
-
-float3 VolumeGetTransmittance(float3 StartThroughput, float3 Origin, float3 Direction, FVolumeIntersectionInterval Interval, inout RandomSequence RandSequence)
-{
-	if (VolumeFlags & PATH_TRACER_VOLUME_USE_ANALYTIC_TRANSMITTANCE)
+	if (IntersectionList.VolumeMask & PATH_TRACER_VOLUME_ENABLE_FOG)
 	{
-		return VolumeGetTransmittanceAnalytical(StartThroughput, Origin, Direction, Interval, RandSequence);
+		float TMin = IntersectionList.VolumeTMin[VOLUMEID_FOG];
+		float TMax = IntersectionList.VolumeTMax[VOLUMEID_FOG];
+		if (VolumeFlags & PATH_TRACER_VOLUME_USE_ANALYTIC_TRANSMITTANCE)
+		{
+			Throughput *= FogGetTransmittance(Origin, Direction, TMin, TMax);
+		}
+		else
+		{
+			PATH_TRACER_REFERENCE_TRANSMITTANCE_LOOP(FogGetDensityBounds, FogGetDensity);
+		}
 	}
-
-	return VolumeGetTransmittanceNullTracking(StartThroughput, Origin, Direction, Interval, RandSequence);
-}
-
-float3 VolumeGetTransmittance(float3 StartThroughput, float3 Origin, float3 Direction, FVolumeIntersectionList IntersectionList, inout RandomSequence RandSequence)
-{
-	if (VolumeFlags & PATH_TRACER_VOLUME_USE_ANALYTIC_TRANSMITTANCE)
+	if (IntersectionList.VolumeMask & PATH_TRACER_VOLUME_ENABLE_HETEROGENEOUS_VOLUMES)
 	{
-		return VolumeGetTransmittanceAnalytical(StartThroughput, Origin, Direction, IntersectionList, RandSequence);
+		float TMin = IntersectionList.VolumeTMin[VOLUMEID_HETEROGENEOUS_VOLUMES];
+		float TMax = IntersectionList.VolumeTMax[VOLUMEID_HETEROGENEOUS_VOLUMES];
+		// NOTE: there is analytic mode in this case as the presence of heterogeneous volumes disables analytic transmittance
+		Throughput = HeterogeneousVolumesGetTransmittance(Throughput, Origin, Direction, TMin, TMax, RandSequence);
 	}
-
-	return VolumeGetTransmittanceNullTracking(StartThroughput, Origin, Direction, IntersectionList, RandSequence);
+	return Throughput;
 }
 
 FVolumeTrackingResult VolumeSampleDistance(float3 PathThroughput, float3 Origin, float3 Direction, FVolumeIntersectionInterval Interval, bool bIsCameraRay, inout RandomSequence RandSequence)
@@ -248,20 +213,16 @@ FVolumeTrackingResult VolumeSampleDistance(float3 PathThroughput, float3 Origin,
 		SigmaBar *= bIsCameraRay ? 2 : 1;
 	}
 
-	bool bIntervalHasHeterogeneousVolume = false;
-	for (int Index = 0; Index < Interval.Num; ++Index)
-	{
-		bIntervalHasHeterogeneousVolume |= (Interval.VolumeID[Index] == VOLUMEID_HETEROGENEOUS_VOLUMES);
-	}
-
 	FVolumeTrackingResult TrackingResult;
-	if (bIntervalHasHeterogeneousVolume)
+	if (Interval.VolumeMask & PATH_TRACER_VOLUME_ENABLE_HETEROGENEOUS_VOLUMES)
 	{
+		// Interval is crossing a heterogeneous volume
 		FMajorantData OverlappingMajorant = CreateMajorantData(VolumeDensityBounds.SigmaMin, SigmaBar);
 		TrackingResult = HeterogeneousVolumesSampleDistance(PathThroughput, Origin, Direction, Interval.VolumeTMin, Interval.VolumeTMax, OverlappingMajorant, RandSequence);
 	}
 	else
 	{
+		// Interval is only crossing atmo/fog
 		float RandValue = RandomSequence_GenerateSample1D(RandSequence);
 		TrackingResult.Distance = SampleSpectralTransmittance(RandValue, SigmaBar, PathThroughput);
 		if (TrackingResult.Distance < 0.0)
diff --git a/Engine/Shaders/Private/PathTracing/Volume/PathTracingVolumeCommon.ush b/Engine/Shaders/Private/PathTracing/Volume/PathTracingVolumeCommon.ush
--- a/Engine/Shaders/Private/PathTracing/Volume/PathTracingVolumeCommon.ush
+++ b/Engine/Shaders/Private/PathTracing/Volume/PathTracingVolumeCommon.ush
@@ -39,24 +39,21 @@ struct FVolumeIntersectionInterval
 {
 	float VolumeTMin;
 	float VolumeTMax;
-	uint VolumeID[PATH_TRACER_MAX_VOLUMES];
-	int Num;
+	uint VolumeMask;
 };
 
 struct FVolumeIntersectionList
 {
-	float VolumeTMin[PATH_TRACER_MAX_VOLUMES];
-	float VolumeTMax[PATH_TRACER_MAX_VOLUMES];
-	uint  VolumeID[PATH_TRACER_MAX_VOLUMES];
-	int Num;
+	float3 VolumeTMin;
+	float3 VolumeTMax;
+	uint VolumeMask;
 
 	// only need to track one blocker (assume they are opaque)
 	float BlockerHitT;
-	uint  BlockerID;
 
 	bool HitVolume()
 	{
-		return Num > 0;
+		return VolumeMask != 0;
 	}
 
 	bool HitBlocker()
@@ -69,17 +66,15 @@ struct FVolumeIntersectionList
 		// merge the provided hit into the list
 		if (VolumeIntersection.HitVolume())
 		{
-			VolumeTMin[Num] = VolumeIntersection.VolumeTMin;
-			VolumeTMax[Num] = VolumeIntersection.VolumeTMax;
-			VolumeID[Num] = ID;
-			Num++;
+			VolumeTMin[ID] = VolumeIntersection.VolumeTMin;
+			VolumeTMax[ID] = VolumeIntersection.VolumeTMax;
+			VolumeMask |= PATH_TRACER_VOLUME_ENABLE_BIT << ID;
 		}
 		if (VolumeIntersection.HitBlocker())
 		{
 			if (!HitBlocker() || BlockerHitT > VolumeIntersection.BlockerHitT)
 			{
 				BlockerHitT = VolumeIntersection.BlockerHitT;
-				BlockerID = ID;
 			}
 		}
 	}
@@ -88,41 +83,58 @@ struct FVolumeIntersectionList
 	FVolumeIntersectionInterval GetCurrentInterval()
 	{
 		FVolumeIntersectionInterval Result = (FVolumeIntersectionInterval)0;
-		Result.Num = 0;
-		Result.VolumeTMin = VolumeTMin[0];
-		Result.VolumeTMax = VolumeTMax[0];
-		for (int Index = 1; Index < Num; Index++)
+		for (int Index = 0; Index < PATH_TRACER_MAX_VOLUMES; Index++)
 		{
-			// NOTE: we only enter the loop body if Result was a valid interval
-			float TMin = VolumeTMin[Index];
-			float TMax = VolumeTMax[Index];
-			if (Result.VolumeTMin == TMin)
+			uint Mask = PATH_TRACER_VOLUME_ENABLE_BIT << Index;
+			if ((VolumeMask & Mask) == 0)
 			{
-				// both segments start together, figure out which ends first
-				Result.VolumeTMax = min(Result.VolumeTMax, TMax);
+				continue; // not valid
 			}
-			else if (Result.VolumeTMin < TMin)
+
+			float TMin = VolumeTMin[Index];
+			float TMax = VolumeTMax[Index];
+			if (Result.VolumeMask == 0)
 			{
-				// what we have so far _starts_ before the current segment
-				// so either the two segments are disjoint (Result.y is closer)
-				// or the two overlap, and we need to stop at the next transition line
-				Result.VolumeTMax = min(Result.VolumeTMax, TMin);
+				// we haven't stored the start interval yet
+				Result.VolumeTMin = TMin;
+				Result.VolumeTMax = TMax;
+				Result.VolumeMask = Mask;
 			}
 			else
 			{
-				// segments could be disjoint (TMax is closer) or overlapping (TMax is closer)
-				Result.VolumeTMax = min(TMax, Result.VolumeTMin);
-				Result.VolumeTMin = TMin; // we just established this one is closer
+				if (Result.VolumeTMin == TMin)
+				{
+					// both segments start together, figure out which ends first
+					Result.VolumeTMax = min(Result.VolumeTMax, TMax);
+				}
+				else if (Result.VolumeTMin < TMin)
+				{
+					// what we have so far _starts_ before the current segment
+					// so either the two segments are disjoint (Result.y is closer)
+					// or the two overlap, and we need to stop at the next transition line
+					Result.VolumeTMax = min(Result.VolumeTMax, TMin);
+				}
+				else
+				{
+					// segments could be disjoint (TMax is closer) or overlapping (TMax is closer)
+					Result.VolumeTMax = min(TMax, Result.VolumeTMin);
+					Result.VolumeTMin = TMin; // we just established this one is closer
+				}
 			}
 		}
 		// now that we established the distance, figure out which IDs overlap with the resolved segment
-		for (int Index2 = 0; Index2 < Num; Index2++)
+		Result.VolumeMask = 0;
+		for (int Index2 = 0; Index2 < PATH_TRACER_MAX_VOLUMES; Index2++)
 		{
+			uint Mask = PATH_TRACER_VOLUME_ENABLE_BIT << Index2;
+			if ((VolumeMask & Mask) == 0)
+			{
+				continue; // not overlapping
+			}
 			if (Result.VolumeTMax > VolumeTMin[Index2] &&
 				Result.VolumeTMin < VolumeTMax[Index2])
 			{
-				Result.VolumeID[Result.Num] = VolumeID[Index2];
-				Result.Num++;
+				Result.VolumeMask |= Mask;
 			}
 		}
 		return Result;
@@ -132,22 +144,25 @@ struct FVolumeIntersectionList
 	FVolumeIntersectionList Update(float T)
 	{
 		FVolumeIntersectionList Result = (FVolumeIntersectionList)0;
-		Result.Num = 0;
 		// clip the current intervals, knowning that we are "T" away
-		for (int Index = 0; Index < Num; Index++)
+		for (int Index = 0; Index < PATH_TRACER_MAX_VOLUMES; Index++)
 		{
+			uint Mask = PATH_TRACER_VOLUME_ENABLE_BIT << Index;
+			if ((VolumeMask & Mask) == 0)
+			{
+				continue; // not valid
+			}
+
 			float TMax = VolumeTMax[Index];
 			if (T < TMax)
 			{
 				// still valid? then add it back in
-				Result.VolumeTMin[Result.Num] = max(T, VolumeTMin[Index]);
-				Result.VolumeTMax[Result.Num] = TMax;
-				Result.VolumeID[Result.Num] = VolumeID[Index];
-				Result.Num++;
+				Result.VolumeTMin[Index] = max(T, VolumeTMin[Index]);
+				Result.VolumeTMax[Index] = TMax;
+				Result.VolumeMask |= Mask;
 			}
 		}
 		Result.BlockerHitT = BlockerHitT;
-		Result.BlockerID = BlockerID;
 		return Result;
 	}
 };
@@ -209,17 +224,24 @@ FPathTracingPayload CreateMediumHitPayload(float HitT, float3 TranslatedWorldPos
 	return Result;
 }
 
-void MergeVolumeShadedResult(inout FVolumeShadedResult Merged, FVolumeShadedResult Other)
+void MergeVolumeShadedResult(inout FVolumeShadedResult Merged, FVolumeShadedResult Other, bool bAsHoldout)
 {
-	Merged.SigmaT += Other.SigmaT;
-	Merged.SigmaH += Other.SigmaH;
-	Merged.SigmaSRayleigh += Other.SigmaSRayleigh;
-	Merged.SigmaSHG += Other.SigmaSHG;
-	Merged.Emission += Other.Emission;
-	// blend phase function anisotropy based on amount of each medium
-	float Qa = Other.SigmaSHG.x + Other.SigmaSHG.y + Other.SigmaSHG.z;
-	float Qs = Merged.SigmaSHG.x + Merged.SigmaSHG.y + Merged.SigmaSHG.z;
-	Merged.PhaseG = lerp(Merged.PhaseG, Other.PhaseG, Qs > 0 ? Qa / Qs : 0.0);
+	if (bAsHoldout)
+	{
+		Merged.SigmaT += Other.SigmaT;
+		Merged.SigmaH += Other.SigmaT;
+	}
+	else
+	{
+		Merged.SigmaT += Other.SigmaT;
+		Merged.SigmaSRayleigh += Other.SigmaSRayleigh;
+		Merged.SigmaSHG += Other.SigmaSHG;
+		Merged.Emission += Other.Emission;
+		// blend phase function anisotropy based on amount of each medium
+		float Qa = Other.SigmaSHG.x + Other.SigmaSHG.y + Other.SigmaSHG.z;
+		float Qs = Merged.SigmaSHG.x + Merged.SigmaSHG.y + Merged.SigmaSHG.z;
+		Merged.PhaseG = lerp(Merged.PhaseG, Other.PhaseG, Qs > 0 ? Qa / Qs : 0.0);
+	}
 }
 
 // represent a stochastically chosen volume element that is guaranteed to lie between 
diff --git a/Engine/Shaders/Shared/PathTracingDefinitions.h b/Engine/Shaders/Shared/PathTracingDefinitions.h
--- a/Engine/Shaders/Shared/PathTracingDefinitions.h
+++ b/Engine/Shaders/Shared/PathTracingDefinitions.h
@@ -53,16 +53,16 @@
 #define VOLUMEID_ATMOSPHERE				0
 #define VOLUMEID_FOG					1
 #define VOLUMEID_HETEROGENEOUS_VOLUMES	2
-#define PATH_TRACER_MAX_VOLUMES			3 // NOTE: Adjust _BIT constants below if this ever grows
+#define PATH_TRACER_MAX_VOLUMES			3
 
 #define PATH_TRACER_VOLUME_ENABLE_BIT						(1u)
 #define PATH_TRACER_VOLUME_ENABLE_ATMOSPHERE  				(PATH_TRACER_VOLUME_ENABLE_BIT << VOLUMEID_ATMOSPHERE)
 #define PATH_TRACER_VOLUME_ENABLE_FOG               		(PATH_TRACER_VOLUME_ENABLE_BIT << VOLUMEID_FOG)
 #define PATH_TRACER_VOLUME_ENABLE_HETEROGENEOUS_VOLUMES 	(PATH_TRACER_VOLUME_ENABLE_BIT << VOLUMEID_HETEROGENEOUS_VOLUMES)
 
-#define PATH_TRACER_VOLUME_HOLDOUT_BIT						(8u)
+#define PATH_TRACER_VOLUME_HOLDOUT_BIT						(PATH_TRACER_VOLUME_ENABLE_BIT << PATH_TRACER_MAX_VOLUMES)
 #define PATH_TRACER_VOLUME_HOLDOUT_ATMOSPHERE  				(PATH_TRACER_VOLUME_HOLDOUT_BIT << VOLUMEID_ATMOSPHERE)
 #define PATH_TRACER_VOLUME_HOLDOUT_FOG               		(PATH_TRACER_VOLUME_HOLDOUT_BIT << VOLUMEID_FOG)
 #define PATH_TRACER_VOLUME_HOLDOUT_HETEROGENEOUS_VOLUMES	(PATH_TRACER_VOLUME_HOLDOUT_BIT << VOLUMEID_HETEROGENEOUS_VOLUMES)
 
-#define PATH_TRACER_VOLUME_USE_ANALYTIC_TRANSMITTANCE       (64u)
+#define PATH_TRACER_VOLUME_USE_ANALYTIC_TRANSMITTANCE       (PATH_TRACER_VOLUME_HOLDOUT_BIT << PATH_TRACER_MAX_VOLUMES)
