Fix missing LWCData initialization in decals and light functions

diff --git a/Engine/Shaders/Private/DeferredDecal.usf b/Engine/Shaders/Private/DeferredDecal.usf
--- a/Engine/Shaders/Private/DeferredDecal.usf
+++ b/Engine/Shaders/Private/DeferredDecal.usf
@@ -139,8 +139,8 @@ void FPixelShaderInOut_MainPS(inout FPixelShaderIn In, inout FPixelShaderOut Out
 	MaterialParameters.LightVector = SwizzlePos;
 
 	MaterialParameters.AbsoluteWorldPosition = MaterialParameters.WorldPosition_NoOffsets = DFFastSubtract(TranslatedWorldPosition, PrimaryView.PreViewTranslation);
-	MaterialParameters.LWCData.AbsoluteWorldPosition = MaterialParameters.LWCData.WorldPosition_NoOffsets = DFToWS(MaterialParameters.AbsoluteWorldPosition);
 	MaterialParameters.WorldPosition_CamRelative = MaterialParameters.WorldPosition_NoOffsets_CamRelative = TranslatedWorldPosition;
+	MaterialParameters.LWCData = MakeMaterialLWCData(MaterialParameters);
 
 #if MATERIAL_VIRTUALTEXTURE_FEEDBACK
 	InitializeVirtualTextureFeedback(MaterialParameters.VirtualTextureFeedback, (uint2)In.SvPosition.xy, View.FrameNumber);
diff --git a/Engine/Shaders/Private/LightFunctionCommon.ush b/Engine/Shaders/Private/LightFunctionCommon.ush
--- a/Engine/Shaders/Private/LightFunctionCommon.ush
+++ b/Engine/Shaders/Private/LightFunctionCommon.ush
@@ -70,8 +70,8 @@ float3 GetLightFunctionColor(float3 LightVector, float3 TranslatedWorldSpace
 	MaterialParameters.ReflectionVector = LightVector.xyz;
 	MaterialParameters.LightVector = LightVector.xyz;
 	MaterialParameters.AbsoluteWorldPosition = DFFastSubtract(TranslatedWorldSpace, PrimaryView.PreViewTranslation);
-	MaterialParameters.LWCData.AbsoluteWorldPosition = DFToWS(MaterialParameters.AbsoluteWorldPosition);
 	MaterialParameters.ScreenPosition = mul(float4(TranslatedWorldSpace, 1.0f), PrimaryView.TranslatedWorldToClip);
+	MaterialParameters.LWCData = MakeMaterialLWCData(MaterialParameters);
 
 	FPixelMaterialInputs PixelMaterialInputs;
 	CalcPixelMaterialInputs(MaterialParameters, PixelMaterialInputs);
diff --git a/Engine/Shaders/Private/RayTracing/RayTracingDecalMaterialShader.usf b/Engine/Shaders/Private/RayTracing/RayTracingDecalMaterialShader.usf
--- a/Engine/Shaders/Private/RayTracing/RayTracingDecalMaterialShader.usf
+++ b/Engine/Shaders/Private/RayTracing/RayTracingDecalMaterialShader.usf
@@ -141,8 +141,8 @@ RAY_TRACING_ENTRY_CALLABLE(RayTracingDecalMaterialShader, FDecalShaderPayload, P
 	MaterialParameters.LightVector = SwizzlePos;
 
 	MaterialParameters.AbsoluteWorldPosition = MaterialParameters.WorldPosition_NoOffsets = DFFastSubtract(TranslatedWorldPosition.xyz, PrimaryView.PreViewTranslation);
-	MaterialParameters.LWCData.AbsoluteWorldPosition = MaterialParameters.LWCData.WorldPosition_NoOffsets = DFToWS(MaterialParameters.AbsoluteWorldPosition);
 	MaterialParameters.WorldPosition_CamRelative = MaterialParameters.WorldPosition_NoOffsets_CamRelative = TranslatedWorldPosition.xyz;
+	MaterialParameters.LWCData = MakeMaterialLWCData(MaterialParameters);
 
 	FPixelMaterialInputs PixelMaterialInputs;
 	CalcPixelMaterialInputs(MaterialParameters, PixelMaterialInputs);
