Fix hair texture generation with compact layout

diff --git a/Engine/Plugins/Runtime/HairStrands/Source/HairStrandsCore/Private/GroomTextureBuilder.cpp b/Engine/Plugins/Runtime/HairStrands/Source/HairStrandsCore/Private/GroomTextureBuilder.cpp
--- a/Engine/Plugins/Runtime/HairStrands/Source/HairStrandsCore/Private/GroomTextureBuilder.cpp
+++ b/Engine/Plugins/Runtime/HairStrands/Source/HairStrandsCore/Private/GroomTextureBuilder.cpp
@@ -25,6 +25,11 @@
 #include "DerivedDataRequestOwner.h"
 #endif
 
+// Enable this to force looping texture generation pass and automatically 
+// triggering RenderDoc capture when generating groom textures
+#define DEBUG_TEXTURE_GENERATION 0
+#define DEBUG_TEXTURE_RENDERDOC 0
+
 static int32 GHairStrandsTextureDilationPassCount = 8;
 static FAutoConsoleVariableRef CVarHairStrandsTextureDilationPassCount(TEXT("r.HairStrands.Textures.DilationCount"), GHairStrandsTextureDilationPassCount, TEXT("Number of dilation pass run onto the generated hair strands textures (Default:8)."));
 
@@ -1156,12 +1161,15 @@ static bool TraceTextures(
 	}
 
 	// Increment processed tile
+	// When debugging, never complete texture generation to force pass execution every frame
+	#if !DEBUG_TEXTURE_GENERATION
 	++Output.TileCoord.X;
 	if (Output.TileCoord.X >= Output.TileCount.X)
 	{
 		Output.TileCoord.X = 0;
 		++Output.TileCoord.Y;
 	}
+	#endif
 
 	return true;
 }
@@ -1312,13 +1320,15 @@ void RunHairStrandsTexturesQueries(FRDGBuilder& GraphBuilder, FGlobalShaderMap*
 		{
 			if (R)
 			{
-
+			#if DEBUG_TEXTURE_GENERATION && DEBUG_TEXTURE_RENDERDOC
 				RenderCaptureInterface::FScopedCapture RenderCapture(true, GraphBuilder, TEXT("PassX"));
-
+			#endif
 				FHairStrandsRDGTextures Textures = RegisterTextures(GraphBuilder, R->InProgress);
 
 				TraceTextures(GraphBuilder, ShaderMap, R->Info, DebugShaderData, Textures);
 
+			// When debugging, never complete texture generation to force pass execution every frame
+			#if !DEBUG_TEXTURE_GENERATION
 				if (Textures.IsCompleted())
 				{
 					Textures = DilateTextures(GraphBuilder, ShaderMap, Textures);
@@ -1326,6 +1336,7 @@ void RunHairStrandsTexturesQueries(FRDGBuilder& GraphBuilder, FGlobalShaderMap*
 					GStrandsTexturesReadbacks.Enqueue(R);
 				}
 				else
+			#endif
 				{
 					ExportTextures(GraphBuilder, Textures, R->InProgress);
 					NotCompleted.Add(R);
@@ -1369,7 +1380,9 @@ void RunHairStrandsTexturesQueries(FRDGBuilder& GraphBuilder, FGlobalShaderMap*
 			R->Output = Q.Output;
 			R->Info = Q.Info;
 
+		#if DEBUG_TEXTURE_GENERATION && DEBUG_TEXTURE_RENDERDOC
 			RenderCaptureInterface::FScopedCapture RenderCapture(true, GraphBuilder, TEXT("Pass0"));
+		#endif
 
 			FHairStrandsRDGTextures Textures = CreateTextures(GraphBuilder, FIntPoint(FMath::Clamp(Q.Info.Resolution, 512u, 16384u), FMath::Clamp(Q.Info.Resolution, 512u, 16384u)), Q.Info.Layout);
 			AddTextureClearPass(GraphBuilder, ShaderMap, Q.Info.Resolution, Q.Info.Layout, Textures.Texture);
diff --git a/Engine/Shaders/Private/HairStrands/HairCardsAttributeCommon.ush b/Engine/Shaders/Private/HairStrands/HairCardsAttributeCommon.ush
--- a/Engine/Shaders/Private/HairStrands/HairCardsAttributeCommon.ush
+++ b/Engine/Shaders/Private/HairStrands/HairCardsAttributeCommon.ush
@@ -63,8 +63,8 @@ uint GetHairCardChannelIndex(uint InLayoutIndex, uint InAttribute)
 	{
 		case EHairCardsVFAttribute_Depth      : return LAYOUT_VALUES(InLayoutIndex, 0, 0, 1, 1);
 		case EHairCardsVFAttribute_Seed       : return LAYOUT_VALUES(InLayoutIndex, 3, 3, 2, 2);
-		case EHairCardsVFAttribute_CoordU     : return LAYOUT_VALUES(InLayoutIndex, 2, 2, 3, 1);
-		case EHairCardsVFAttribute_GroupIndex : return LAYOUT_VALUES(InLayoutIndex, 3, 3, 0, 2);
+		case EHairCardsVFAttribute_CoordU     : return LAYOUT_VALUES(InLayoutIndex, 2, 2, 3, 2);
+		case EHairCardsVFAttribute_GroupIndex : return LAYOUT_VALUES(InLayoutIndex, 3, 3, 0, 3);
 		case EHairCardsVFAttribute_Roughness  : return LAYOUT_VALUES(InLayoutIndex, 3, 0, 3, 3);
 		case EHairCardsVFAttribute_RootUV     : return LAYOUT_VALUES(InLayoutIndex, 0, 0, 0, 0);
 		case EHairCardsVFAttribute_Color      : return LAYOUT_VALUES(InLayoutIndex, 0, 0, 0, 0);
@@ -142,74 +142,74 @@ float4 GetAuxilary    (float2 InAtlasUV, float4 InDefault)      { float4 Out = I
 ///////////////////////////////////////////////////////////////////////////////////////////////////
 // MaterialTemplate interface
 
-float GetHairStrandsDepth(float2 InAtlasUV, float InDeviceZ)
+float GetHairCardsDepth(float2 InAtlasUV, float InDeviceZ)
 {
 	const float SceneDepthOffset = GetDepthOffset(InAtlasUV);
 	const float SceneDepth = ConvertFromDeviceZ(InDeviceZ);
 	return ConvertToDeviceZ(SceneDepth + SceneDepthOffset);
 }
 
-float3 GetHairStrandsTangent(float2 InAtlasUV, half3x3 TangentToWorld, bool bUseTangentSpace)
+float3 GetHairCardsTangent(float2 InAtlasUV, half3x3 TangentToWorld, bool bUseTangentSpace)
 {
 	const float3 LocalTangent = GetTangent(InAtlasUV, float3(0, 1, 0));
 	const float3 WorldTangent = mul(LocalTangent, TangentToWorld);
 	return bUseTangentSpace ? LocalTangent : WorldTangent;
 }
 
-float GetHairStrandsCoverage(float2 InAtlasUV)
+float GetHairCardsCoverage(float2 InAtlasUV)
 {
 	return GetCoverage(InAtlasUV, 1.f);
 }
 
-float4 GetHairStrandsAuxilaryData(float2 InAtlasUV)
+float4 GetHairCardsAuxilaryData(float2 InAtlasUV)
 {
 	return GetAuxilary(InAtlasUV, 0.f);
 }
 
-float2 GetHairStrandsUV(float2 InAtlasUV)
+float2 GetHairCardsUV(float2 InAtlasUV)
 {
 	return float2(GetCoordU(InAtlasUV), 0.5f);
 }
 
-float2 GetHairStrandsDimensions(float2 InAtlasUV, float HairPrimitiveLength)
+float2 GetHairCardsDimensions(float2 InAtlasUV, float HairPrimitiveLength)
 {
 	return float2(HairPrimitiveLength, 0.01f/*WorldRadius*/);
 }
 
-float2 GetHairStrandsRootUV(float2 InAtlasUV, float2 InRootUV)
+float2 GetHairCardsRootUV(float2 InAtlasUV, float2 InRootUV)
 {
 	// Invert V to compensate image origin flip
 	// Similar to DecodeHairAttribute in HairStrandVertexFactoryCommon.ush
 	return GetRootUV(InAtlasUV, float2(InRootUV.x, 1.0 - InRootUV.y));
 }
 
-float GetHairStrandsSeed(float2 InAtlasUV)
+float GetHairCardsSeed(float2 InAtlasUV)
 {
 	return GetSeed(InAtlasUV);
 }
 
-uint3 GetHairStrandsClumpID(float2 InAtlasUV)
+uint3 GetHairCardsClumpID(float2 InAtlasUV)
 {
 	return 0;
 }
 
-float3 GetHairStrandsColor(float2 InAtlasUV, float3 InColor)
+float3 GetHairCardsColor(float2 InAtlasUV, float3 InColor)
 {
 	return GetColor(InAtlasUV, InColor);
 }
 
-float GetHairStrandsRoughness(float2 InAtlasUV, float InCardRoughness)
+float GetHairCardsRoughness(float2 InAtlasUV, float InCardRoughness)
 {
 	//return GetRoughness(InAtlasUV, InCardRoughness);
 	return (HairCardsVF.Flags & EHairCardsVFAttribute_Roughness) ? 0 : InCardRoughness;
 }
 
-float GetHairStrandsAO(float2 InAtlasUV, float InAO)
+float GetHairCardsAO(float2 InAtlasUV, float InAO)
 {
 	return 1.0f;
 }
 
-float GetHairStrandsGroupIndex(float2 InAtlasUV, float InCardGroupIndex)
+float GetHairCardsGroupIndex(float2 InAtlasUV, float InCardGroupIndex)
 {
 	// * For hair cards,  group index is stored onto vertices
 	// * For hair meshes, group index is stored within a texture
diff --git a/Engine/Shaders/Private/MaterialTemplate.ush b/Engine/Shaders/Private/MaterialTemplate.ush
--- a/Engine/Shaders/Private/MaterialTemplate.ush
+++ b/Engine/Shaders/Private/MaterialTemplate.ush
@@ -1923,7 +1923,7 @@ float2 MaterialExpressionGetHairRootUV(FMaterialPixelParameters Parameters)
 #if HAIR_STRAND_MESH_FACTORY
 	return GetHairStrandsRootUV(Parameters.HairControlPointId);
 #elif HAIR_CARD_MESH_FACTORY
-	return GetHairStrandsRootUV(Parameters.HairPrimitiveUV, Parameters.HairPrimitiveRootUV);
+	return GetHairCardsRootUV(Parameters.HairPrimitiveUV, Parameters.HairPrimitiveRootUV);
 #else
 	return float2(0, 0);
 #endif
@@ -1934,7 +1934,7 @@ float2 MaterialExpressionGetHairUV(FMaterialPixelParameters Parameters)
 #if HAIR_STRAND_MESH_FACTORY
 	return GetHairStrandsUV(Parameters.HairControlPointId, Parameters.HairPrimitiveUV);
 #elif HAIR_CARD_MESH_FACTORY
-	return GetHairStrandsUV(Parameters.HairPrimitiveUV);
+	return GetHairCardsUV(Parameters.HairPrimitiveUV);
 #else
 	return float2(0,0);
 #endif
@@ -1945,7 +1945,7 @@ float2 MaterialExpressionGetHairDimensions(FMaterialPixelParameters Parameters)
 #if HAIR_STRAND_MESH_FACTORY
 	return GetHairStrandsDimensions(Parameters.HairControlPointId, Parameters.HairPrimitiveUV.x);
 #elif HAIR_CARD_MESH_FACTORY
-	return GetHairStrandsDimensions(Parameters.HairPrimitiveUV, Parameters.HairPrimitiveLength);
+	return GetHairCardsDimensions(Parameters.HairPrimitiveUV, Parameters.HairPrimitiveLength);
 #else
 	return float2(0, 0);
 #endif
@@ -1956,7 +1956,7 @@ float MaterialExpressionGetHairSeed(FMaterialPixelParameters Parameters)
 #if HAIR_STRAND_MESH_FACTORY
 	return GetHairStrandsSeed(Parameters.HairControlPointId);
 #elif HAIR_CARD_MESH_FACTORY
-	return GetHairStrandsSeed(Parameters.HairPrimitiveUV);
+	return GetHairCardsSeed(Parameters.HairPrimitiveUV);
 #else
 	return 0;
 #endif
@@ -1967,7 +1967,7 @@ float3 MaterialExpressionGetHairClumpID(FMaterialPixelParameters Parameters)
 #if HAIR_STRAND_MESH_FACTORY
 	return GetHairStrandsClumpID(Parameters.HairControlPointId);
 #elif HAIR_CARD_MESH_FACTORY
-	return GetHairStrandsClumpID(Parameters.HairPrimitiveUV);
+	return GetHairCardsClumpID(Parameters.HairPrimitiveUV);
 #else
 	return 0;
 #endif
@@ -1978,7 +1978,7 @@ float3 MaterialExpressionGetHairBaseColor(FMaterialPixelParameters Parameters)
 #if HAIR_STRAND_MESH_FACTORY
 	return GetHairStrandsColor(Parameters.HairControlPointId, Parameters.HairPrimitiveUV.x);
 #elif HAIR_CARD_MESH_FACTORY
-	return GetHairStrandsColor(Parameters.HairPrimitiveUV, Parameters.HairPrimitiveMaterial.xyz);
+	return GetHairCardsColor(Parameters.HairPrimitiveUV, Parameters.HairPrimitiveMaterial.xyz);
 #else
 	return float3(0,0,0);
 #endif
@@ -1989,7 +1989,7 @@ float MaterialExpressionGetHairRoughness(FMaterialPixelParameters Parameters)
 #if HAIR_STRAND_MESH_FACTORY
 	return GetHairStrandsRoughness(Parameters.HairControlPointId, Parameters.HairPrimitiveUV.x);
 #elif HAIR_CARD_MESH_FACTORY
-	return GetHairStrandsRoughness(Parameters.HairPrimitiveUV, Parameters.HairPrimitiveMaterial.w);
+	return GetHairCardsRoughness(Parameters.HairPrimitiveUV, Parameters.HairPrimitiveMaterial.w);
 #else
 	return 0;
 #endif
@@ -2000,7 +2000,7 @@ float MaterialExpressionGetHairAO(FMaterialPixelParameters Parameters)
 #if HAIR_STRAND_MESH_FACTORY
 	return GetHairStrandsAO(Parameters.HairControlPointId, Parameters.HairPrimitiveUV.x);
 #elif HAIR_CARD_MESH_FACTORY
-	return GetHairStrandsAO(Parameters.HairPrimitiveUV, 1.0f);
+	return GetHairCardsAO(Parameters.HairPrimitiveUV, 1.0f);
 #else
 	return 0;
 #endif
@@ -2016,7 +2016,7 @@ float MaterialExpressionGetHairDepth(FMaterialPixelParameters Parameters)
 #if HAIR_STRAND_MESH_FACTORY
 	return GetHairStrandsDepth();
 #elif HAIR_CARD_MESH_FACTORY
-	return GetHairStrandsDepth(Parameters.HairPrimitiveUV, Parameters.SvPosition.z);
+	return GetHairCardsDepth(Parameters.HairPrimitiveUV, Parameters.SvPosition.z);
 #else
 	return 0; 
 #endif
@@ -2027,7 +2027,7 @@ float MaterialExpressionGetHairCoverage(FMaterialPixelParameters Parameters)
 #if HAIR_STRAND_MESH_FACTORY
 	return GetHairStrandsCoverage();
 #elif HAIR_CARD_MESH_FACTORY
-	return GetHairStrandsCoverage(Parameters.HairPrimitiveUV);
+	return GetHairCardsCoverage(Parameters.HairPrimitiveUV);
 #else
 	return 0;
 #endif
@@ -2038,7 +2038,7 @@ float3 MaterialExpressionGetHairTangent(FMaterialPixelParameters Parameters, boo
 #if HAIR_STRAND_MESH_FACTORY
 	return bUseTangentSpace ? float3(0,1,0) : Parameters.TangentToWorld[2];
 #elif HAIR_CARD_MESH_FACTORY
-	return GetHairStrandsTangent(Parameters.HairPrimitiveUV, Parameters.TangentToWorld, bUseTangentSpace);
+	return GetHairCardsTangent(Parameters.HairPrimitiveUV, Parameters.TangentToWorld, bUseTangentSpace);
 #else
 	return 0;
 #endif	
@@ -2059,7 +2059,7 @@ float2 MaterialExpressionGetAtlasUVs(FMaterialPixelParameters Parameters)
 float4 MaterialExpressionGetHairAuxilaryData(FMaterialPixelParameters Parameters)
 {
 #if HAIR_CARD_MESH_FACTORY
-	return GetHairStrandsAuxilaryData(Parameters.HairPrimitiveUV);
+	return GetHairCardsAuxilaryData(Parameters.HairPrimitiveUV);
 #else
 	return 0;
 #endif
@@ -2070,7 +2070,7 @@ float MaterialExpressionGetHairGroupIndex(FMaterialPixelParameters Parameters)
 #if HAIR_STRAND_MESH_FACTORY
 	return GetHairStrandsGroupIndex();
 #elif HAIR_CARD_MESH_FACTORY
-	return GetHairStrandsGroupIndex(Parameters.HairPrimitiveUV, Parameters.HairPrimitiveGroupIndex);
+	return GetHairCardsGroupIndex(Parameters.HairPrimitiveUV, Parameters.HairPrimitiveGroupIndex);
 #else
 	return 0;
 #endif
