Fix CoordU channel not being correctly generated during hair strands texture generation.

diff --git a/Engine/Shaders/Private/HairStrands/HairStrandsTexturesGeneration.usf b/Engine/Shaders/Private/HairStrands/HairStrandsTexturesGeneration.usf
--- a/Engine/Shaders/Private/HairStrands/HairStrandsTexturesGeneration.usf
+++ b/Engine/Shaders/Private/HairStrands/HairStrandsTexturesGeneration.usf
@@ -350,7 +350,7 @@ bool Trace(
 							OutDepth 	= Distance;
 							OutSeed		= GetHairStrandsSeed(P0);
 							OutRootUV	= GetHairStrandsRootUV(P0);
-							OutU		= GetHairStrandsDimensions(P0, VertexU).x;
+							OutU		= GetHairStrandsUV(P0, VertexU).x; 
 							OutMaterial	= 0;
 
 							if (HasHairStrandsColor())
