Path Tracing: Maintain PathState compression inside the kernel

This saves around 20 bytes of live state space for all trace calls.

diff --git a/Engine/Shaders/Private/PathTracing/PathTracing.usf b/Engine/Shaders/Private/PathTracing/PathTracing.usf
--- a/Engine/Shaders/Private/PathTracing/PathTracing.usf
+++ b/Engine/Shaders/Private/PathTracing/PathTracing.usf
@@ -53,28 +53,13 @@ FPathState LoadPathStateData(uint Index)
 	Output.RandSequence.SampleSeed  = Data.RandSeqSampleSeed;
 	Output.Radiance = Data.Radiance;
 	Output.Alpha = Data.Alpha;
-	Output.Albedo.x = f16tof32(Data.PackedAlbedoNormal.x);
-	Output.Albedo.y = f16tof32(Data.PackedAlbedoNormal.x >> 16);
-	Output.Albedo.z = f16tof32(Data.PackedAlbedoNormal.y);
-	Output.Normal.x = f16tof32(Data.PackedAlbedoNormal.y >> 16);
-	Output.Normal.y = f16tof32(Data.PackedAlbedoNormal.z);
-	Output.Normal.z = f16tof32(Data.PackedAlbedoNormal.z >> 16);
+	Output.PackedAlbedoNormal = Data.PackedAlbedoNormal;
 	Output.Ray.Origin = Data.RayOrigin;
 	Output.Ray.Direction = Data.RayDirection;
 	Output.Ray.TMin = 0;
 	Output.Ray.TMax = POSITIVE_INFINITY;
 	Output.PathThroughput = Data.PathThroughput;
-	Output.PathRoughness = f16tof32(Data.PackedRoughnessSigma.x);
-	Output.SigmaT.x = f16tof32(Data.PackedRoughnessSigma.x >> 16);
-	Output.SigmaT.y = f16tof32(Data.PackedRoughnessSigma.y);
-	Output.SigmaT.z = f16tof32(Data.PackedRoughnessSigma.y >> 16);
-	// extract scatter type from sign bits of SigmaT
-	Output.FirstScatterType = 0;
-	Output.FirstScatterType |= (asuint(Output.SigmaT.x) >> 31) << 0;
-	Output.FirstScatterType |= (asuint(Output.SigmaT.y) >> 31) << 1;
-	Output.FirstScatterType |= (asuint(Output.SigmaT.z) >> 31) << 2;
-	// remove sign bits from SigmaT (must be a positive value)
-	Output.SigmaT = abs(Output.SigmaT);
+	Output.PackedRoughnessSigma = Data.PackedRoughnessSigma;
 	return Output;
 }
 
@@ -85,19 +70,11 @@ void StorePathStateData(FPathState PathState, uint Index)
 	Data.RandSeqSampleSeed  = PathState.RandSequence.SampleSeed;
 	Data.Radiance = PathState.Radiance;
 	Data.Alpha = PathState.Alpha;
-	Data.PackedAlbedoNormal.x = f32tof16(PathState.Albedo.x) | (f32tof16(PathState.Albedo.y) << 16);
-	Data.PackedAlbedoNormal.y = f32tof16(PathState.Albedo.z) | (f32tof16(PathState.Normal.x) << 16);
-	Data.PackedAlbedoNormal.z = f32tof16(PathState.Normal.y) | (f32tof16(PathState.Normal.z) << 16);
+	Data.PackedAlbedoNormal = PathState.PackedAlbedoNormal;
 	Data.RayOrigin = PathState.Ray.Origin;
 	Data.RayDirection = PathState.Ray.Direction;
 	Data.PathThroughput = PathState.PathThroughput;
-
-	float3 ScatterAndSigmaT = abs(PathState.SigmaT);
-	ScatterAndSigmaT.x = asfloat(((PathState.FirstScatterType & 1) << 31) | asuint(ScatterAndSigmaT.x));
-	ScatterAndSigmaT.y = asfloat(((PathState.FirstScatterType & 2) << 30) | asuint(ScatterAndSigmaT.y));
-	ScatterAndSigmaT.z = asfloat(((PathState.FirstScatterType & 4) << 29) | asuint(ScatterAndSigmaT.z));
-	Data.PackedRoughnessSigma.x = f32tof16(PathState.PathRoughness) | (f32tof16(ScatterAndSigmaT.x) << 16);
-	Data.PackedRoughnessSigma.y = f32tof16(ScatterAndSigmaT.y) | (f32tof16(ScatterAndSigmaT.z) << 16);
+	Data.PackedRoughnessSigma = PathState.PackedRoughnessSigma;
 	PathStateData[Index] = Data;
 }
 
diff --git a/Engine/Shaders/Private/PathTracing/PathTracingCore.ush b/Engine/Shaders/Private/PathTracing/PathTracingCore.ush
--- a/Engine/Shaders/Private/PathTracing/PathTracingCore.ush
+++ b/Engine/Shaders/Private/PathTracing/PathTracingCore.ush
@@ -82,22 +82,97 @@ struct FPathState
 	RandomSequence RandSequence;    //  8 bytes
 	float3 Radiance;                // 12 bytes
 	float  Alpha;                   //  4 bytes
-	float3 Albedo;                  //  6 bytes // half precision
-	float3 Normal;                  //  6 bytes // half precision
+	uint3 PackedAlbedoNormal;       //  6 bytes // half precision (Albedo)
+                                    //  6 bytes // half precision (Normal)
 	FRayDesc Ray;                   // 12+12 = 24 bytes (TMin/TMax do not need to be stored)
 	float3 PathThroughput;          // 12 bytes
-	float PathRoughness;            //  2 bytes // half precision
-	float3 SigmaT; // Extinction    //  6 bytes // half precision 
-	uint FirstScatterType;          //  3 bits - packed in sign bits of SigmaT
+	uint2 PackedRoughnessSigma;     // 2 bytes // half precision (PathRoughness)
+	                                // 6 bytes // half precision  (SigmaT)
+	                                // 3 bits - packed in sign bits of SigmaT (FirstScatterType)
 									// 80 bytes total (see FPathTracingPackedPathState)
 
 	// Temporary parameters not stored				
 	float DepthZ;                   // Camera ray depth
 
+	float3 GetAlbedo()
+	{
+		return float3(f16tof32(PackedAlbedoNormal.x),
+		              f16tof32(PackedAlbedoNormal.x >> 16),
+		              f16tof32(PackedAlbedoNormal.y));
+	}
+
+	float3 GetNormal()
+	{
+		return float3(f16tof32(PackedAlbedoNormal.y >> 16),
+		              f16tof32(PackedAlbedoNormal.z),
+		              f16tof32(PackedAlbedoNormal.z >> 16));
+	}
+
+	void SetAlbedoNormal(float3 Albedo, float3 Normal)
+	{
+		PackedAlbedoNormal.x = f32tof16(Albedo.x) | (f32tof16(Albedo.y) << 16);
+		PackedAlbedoNormal.y = f32tof16(Albedo.z) | (f32tof16(Normal.x) << 16);
+		PackedAlbedoNormal.z = f32tof16(Normal.y) | (f32tof16(Normal.z) << 16);
+	}
+
+	uint GetFirstScatterType()
+	{
+		// Extract SigmaT sign bits
+		uint a = (PackedRoughnessSigma.x >> 29) & 4u;
+		uint b = (PackedRoughnessSigma.y >> 30) & 2u;
+		uint c = (PackedRoughnessSigma.y >> 15) & 1u;
+		return a | b | c;
+	}
+
+	void SetFirstScatterType(uint ScatterType)
+	{
+		uint a = ScatterType & 4u;
+		uint b = ScatterType & 2u;
+		uint c = ScatterType & 1u;
+		// erase old value
+		PackedRoughnessSigma.x &= 0x7fffffffu;
+		PackedRoughnessSigma.y &= 0x7fff7fffu;
+		// set new value
+		PackedRoughnessSigma.x |= a << 29;
+		PackedRoughnessSigma.y |= b << 30;
+		PackedRoughnessSigma.y |= c << 15;
+	}
+
+	float GetPathRoughness()
+	{
+		return f16tof32(PackedRoughnessSigma.x);
+	}
+
+	void SetPathRoughness(float PathRoughness)
+	{
+		PackedRoughnessSigma.x &= 0xffff0000u;
+		PackedRoughnessSigma.x |= f32tof16(PathRoughness);
+	}
+	
+	float3 GetSigmaT()
+	{
+		return abs(float3(f16tof32(PackedRoughnessSigma.x >> 16),
+		                  f16tof32(PackedRoughnessSigma.y),
+		                  f16tof32(PackedRoughnessSigma.y >> 16)));
+	}
+
+	void SetSigmaT(float3 SigmaT)
+	{
+		// erase old value (keeping sign bits)
+		PackedRoughnessSigma.x &= 0x8000ffffu;
+		PackedRoughnessSigma.y &= 0x80008000u;
+		SigmaT = abs(SigmaT); // make sure there are no sign bits
+		// write new value
+		PackedRoughnessSigma.x |= f32tof16(SigmaT.x) << 16;
+		PackedRoughnessSigma.y |= f32tof16(SigmaT.y);
+		PackedRoughnessSigma.y |= f32tof16(SigmaT.z) << 16;
+	}
+
 	bool HasMadeContributionScatter()
 	{
 		// The first scattering type that matters for the PATHTRACER_CONTRIBUTION_* flags is the first non-refractive one
 		// We choose to track camera and pure refracted paths together so that objects behind glass can be treated the same as directly visible objects.
+		uint FirstScatterType = GetFirstScatterType();
 		return !(FirstScatterType == PATHTRACER_SCATTER_CAMERA ||
 			     FirstScatterType == PATHTRACER_SCATTER_REFRACT);
 	}
@@ -111,9 +186,9 @@ struct FPathState
 		// and refracted diffuse to be treated the same.
 		if (!HasMadeContributionScatter())
 		{
-			FirstScatterType = ScatterType;
+			SetFirstScatterType(ScatterType);
 		}
-		switch (FirstScatterType)
+		switch (GetFirstScatterType())
 		{
 			// In these cases, only keep tracing if the corresponding lighting component is enabled
 			case PATHTRACER_SCATTER_DIFFUSE:  return (EnabledIndirectLightingContributions & PATHTRACER_CONTRIBUTION_DIFFUSE ) != 0;
@@ -190,6 +265,9 @@ struct FPathState
 
 		const float4 PixelValue = float4(Radiance, saturate(Alpha));
 
+		const float3 Albedo = GetAlbedo();
+		const float3 Normal = GetNormal();
+
 		RadianceTexture[TextureIndex] = lerp(OldPixel, PixelValue, Blend);
 		AlbedoTexture[TextureIndex]   = lerp(OldAlbedo, float4(Albedo, 0), Blend);
 		NormalTexture[TextureIndex].xyz = lerp(OldNormal.xyz,Normal.xyz, Blend);
@@ -310,7 +388,7 @@ void ApplyDecal(FDecalShaderPayload DecalPayload, inout FPackedPathTracingPayloa
 
 void EvaluateDecals(FPathState PathState, uint InstanceInclusionMask, float HitT, inout FDecalShaderPayload DecalPayload)
 {
-	if (PathState.PathRoughness < MeshDecalRoughnessCutoff)
+	if (PathState.GetPathRoughness() < MeshDecalRoughnessCutoff)
 	{
 		FRayDesc DecalRay = (FRayDesc)0;
 		// go a bit further than the current hit, to handle decals partially inside the model
@@ -350,7 +428,7 @@ void EvaluateDecals(FPathState PathState, uint InstanceInclusionMask, float HitT
 	}
 
 #if PLATFORM_SUPPORTS_CALLABLE_SHADERS
-	if (PathState.PathRoughness < DecalRoughnessCutoff)
+	if (PathState.GetPathRoughness() < DecalRoughnessCutoff)
 	{
 		FDecalLoopCount DecalLoopCount = DecalGridLookup(DecalPayload.TranslatedWorldPos);
 		for (uint Index = 0, Num = DecalLoopCount.NumDecals; Index < Num; ++Index)
@@ -390,7 +468,7 @@ FPackedPathTracingPayload TraceTransparentRay(inout FPathState PathState, int Bo
 	}
 	for (;;)
 	{
-		FPackedPathTracingPayload PackedPayload = InitPathTracingPayload(PathState.FirstScatterType, ApproximateCaustics ? PathState.PathRoughness : 0.0);
+		FPackedPathTracingPayload PackedPayload = InitPathTracingPayload(PathState.GetFirstScatterType(), ApproximateCaustics ? PathState.GetPathRoughness() : 0.0);
 		PackedPayload.SetDBufferA(float4(0, 0, 0, 1));
 		PackedPayload.SetDBufferB(float4(0, 0, 0, 1));
 		PackedPayload.SetDBufferC(float4(0, 0, 0, 1));
@@ -503,7 +581,7 @@ FPackedPathTracingPayload TraceTransparentRay(inout FPathState PathState, int Bo
 						float3 Ro = PathState.Ray.Origin;
 						float3 Rd = PathState.Ray.Direction;
 						float3 TranslatedWorldPos = Ro + Interval.VolumeTMin * Rd;
-						FVolumeShadedResult Result = VolumeGetDensity(TranslatedWorldPos, Interval, bIsCameraRay);
+						FVolumeShadedResult Result = VolumeGetDensity(TranslatedWorldPos, Interval.VolumeMask, bIsCameraRay);
 
 						// clamp to make sure we never exceed the majorant (should not be the case, but need to avoid any possible numerical issues)
 						float3 SigmaT = min(Result.SigmaT, TrackingResult.SigmaBar);
@@ -519,7 +597,9 @@ FPackedPathTracingPayload TraceTransparentRay(inout FPathState PathState, int Bo
 						{
 							float3 SigmaS = min(Result.SigmaSRayleigh + Result.SigmaSHG, SigmaT);
 							// accumulate a signal for the denoiser
-							AccumulateAlbedo(SigmaS, PathThroughput, PathState.PathRoughness, PathState.Albedo);
+							float3 Albedo = PathState.GetAlbedo();
+							AccumulateAlbedo(SigmaS, PathThroughput, PathState.GetPathRoughness(), Albedo);
+							PathState.SetAlbedoNormal(Albedo, PathState.GetNormal());
 							float3 Contrib = PathThroughput * SigmaS;
 							float SelectionWeight = max3(Contrib.x, Contrib.y, Contrib.z);
 							if (HitSample.Accept(SelectionWeight))
@@ -583,7 +663,7 @@ FPackedPathTracingPayload TraceTransparentRay(inout FPathState PathState, int Bo
 				if (EnableDBuffer)
 				{
 					// Retrace ray with DBuffer data
-					PackedPayload = InitPathTracingPayload(PathState.FirstScatterType, ApproximateCaustics ? PathState.PathRoughness : 0.0);
+					PackedPayload = InitPathTracingPayload(PathState.GetFirstScatterType(), ApproximateCaustics ? PathState.GetPathRoughness() : 0.0);
 
 					PackedPayload.SetDBufferA(DecalPayload.GetBaseColor());
 					PackedPayload.SetDBufferB(DecalPayload.GetWorldNormal());
@@ -608,13 +688,16 @@ FPackedPathTracingPayload TraceTransparentRay(inout FPathState PathState, int Bo
 			}
 #endif
 			// Now that we know the real decal data, apply it
-			ApplyDecal(DecalPayload, PackedPayload, PathState.Ray.Direction, PathState.PathRoughness);
+			ApplyDecal(DecalPayload, PackedPayload, PathState.Ray.Direction, PathState.GetPathRoughness());
 		}
 
-		// account for Beer's law through the currently active medium
-		// TODO: merge this with the volume raymarching?
-		// NOTE: select statement is to avoid corner case of 0.0*inf that can occur with directly visible skydomes
-		PathState.PathThroughput *= select(PathState.SigmaT > 0.0, exp(-PathState.SigmaT * (PackedPayload.HitT - PathState.Ray.TMin)), 1.0);
+		{
+			// account for Beer's law through the currently active medium
+			// TODO: merge this with the volume raymarching?
+			// NOTE: select statement is to avoid corner case of 0.0*inf that can occur with directly visible skydomes
+			float3 SigmaT = PathState.GetSigmaT();
+			PathState.PathThroughput *= select(SigmaT > 0.0, exp(-SigmaT * (PackedPayload.HitT - PathState.Ray.TMin)), 1.0);
+		}
 
 		if (!PackedPayload.IsHoldout() && bIsCameraRay)
 		{
@@ -648,7 +731,10 @@ FPackedPathTracingPayload TraceTransparentRay(inout FPathState PathState, int Bo
 			float3 Contrib = PathState.PathThroughput * EstimateMaterialAlbedo(HitPayload);
 
 			// accumulate what the denoiser wants into albedo/normal (as long as the current path is rough enough)
-			AccumulateAlbedoNormal(HitPayload, PathState.PathThroughput, PathState.PathRoughness, PathState.Albedo, PathState.Normal);
+			float3 Albedo = PathState.GetAlbedo();
+			float3 Normal = PathState.GetNormal();
+			AccumulateAlbedoNormal(HitPayload, PathState.PathThroughput, PathState.GetPathRoughness(), Albedo, Normal);
+			PathState.SetAlbedoNormal(Albedo, Normal);
 
 			float SelectionWeight = max3(Contrib.x, Contrib.y, Contrib.z);
 			if (HitSample.Accept(SelectionWeight))
@@ -1303,9 +1389,9 @@ FPathState CreatePathState(int2 PixelIndex, int2 TextureIndex)
 
 	// path state variables (these cary information between bounces)
 	Output.PathThroughput = 1.0;
-	Output.PathRoughness = 0;
-	Output.SigmaT = float3(StartingExtinctionCoefficient[0], StartingExtinctionCoefficient[1], StartingExtinctionCoefficient[2]);
-	Output.FirstScatterType = PATHTRACER_SCATTER_CAMERA;
+	Output.SetPathRoughness(0.0);
+	Output.SetSigmaT(float3(StartingExtinctionCoefficient[0], StartingExtinctionCoefficient[1], StartingExtinctionCoefficient[2]));
+	Output.SetFirstScatterType(PATHTRACER_SCATTER_CAMERA);
 
 	return Output;
 }
@@ -1369,7 +1455,7 @@ bool PathTracingKernel(inout FPathState PathState, int Bounce)
 
 			// find out how much volume exists at the current point
 			const float3 TranslatedWorldPos = Ro + EquiAngularT * Rd;
-			const FVolumeShadedResult Result = VolumeGetDensity(TranslatedWorldPos, VolumeSegment.Interval, bIsCameraRay);
+			const FVolumeShadedResult Result = VolumeGetDensity(TranslatedWorldPos, VolumeSegment.Interval.VolumeMask, bIsCameraRay);
 			const float3 SigmaT = Result.SigmaT;
 			const float3 SigmaS = min(Result.SigmaSRayleigh + Result.SigmaSHG, SigmaT);
 
@@ -1462,7 +1548,7 @@ bool PathTracingKernel(inout FPathState PathState, int Bounce)
 #endif
 
 	// As soon as the path is blurry enough, we can get away with diffuse sampling only
-	const bool bSimplifySSS = PathState.PathRoughness >= 0.15;
+	const bool bSimplifySSS = PathState.GetPathRoughness() >= 0.15;
 
 	// Processing the random walk will (stochastically) move the shading point elsewhere on the surface of the object
 	if (!ProcessSubsurfaceRandomWalk(Payload, PathState.PathThroughput, PathState.RandSequence, PathState.Ray.Direction, bSimplifySSS))
@@ -1643,7 +1729,7 @@ bool PathTracingKernel(inout FPathState PathState, int Bounce)
 	ApplyRayBias(PathState.Ray, Payload.HitT, MaterialSample.PositionBiasSign * Payload.WorldGeoNormal);
 
 	// enlarge roughness based on the chosen lobe roughness
-	PathState.PathRoughness = max(PathState.PathRoughness, MaterialSample.Roughness);
+	PathState.SetPathRoughness(max(PathState.GetPathRoughness(), MaterialSample.Roughness));
 
 	// update the current extinction if we are crossing a boundary on glass or water
 	// summing the local extinction gives a rudimentary way of dealing with overlapping regions
@@ -1651,17 +1737,19 @@ bool PathTracingKernel(inout FPathState PathState, int Bounce)
 	if (MaterialSample.PositionBiasSign < 0 && Payload.IsMaterialSolidGlass())
 	{
 		const float3 LocalSigmaT = Payload.GetExtinction();
+		float3 SigmaT = PathState.GetSigmaT();
 		if (Payload.IsFrontFace())
 		{
 			// entering
-			PathState.SigmaT += LocalSigmaT;
+			SigmaT += LocalSigmaT;
 		}
 		else
 		{
 			// exiting
-			PathState.SigmaT -= LocalSigmaT;
-			PathState.SigmaT = max(PathState.SigmaT, 0);
+			SigmaT -= LocalSigmaT;
+			SigmaT = max(SigmaT, 0);
 		}
+		PathState.SetSigmaT(SigmaT);
 	}
 
 	// If we are using Material sampling for lights
@@ -1671,7 +1759,6 @@ bool PathTracingKernel(inout FPathState PathState, int Bounce)
 		// While it would be possible to just loop around and use the indirect ray to do this, it would prevent the application
 		// of shadow ray specific logic for transparent shadows or various per light tricks like shadow casting
 		const bool bUseMIS = MISMode == 2 && LightPickingCdfSum > 0;
-		const float ShadowRayRoughness = ApproximateCaustics ? PathState.PathRoughness : 0.0;
 		float PreviousCdfValue = 0.0;
 		for (uint Index = 0, Num = LightLoopCount.NumMISLights; Index < Num; ++Index)
 		{
@@ -1724,8 +1811,8 @@ bool PathTracingKernel(inout FPathState PathState, int Bounce)
 				LightRay.TMax = LightResult.HitT;
 				const bool bCastShadows = bIsVolumeSample ? CastsVolumeShadow(LightId) : CastsShadow(LightId);
 				const uint MissShaderIndex = GetLightMissShaderIndex(LightId);
+				const float ShadowRayRoughness = ApproximateCaustics ? PathState.GetPathRoughness() : 0.0;
 				LightContrib *= TraceTransparentVisibilityRay(LightRay, ShadowRayRoughness, MissShaderIndex, bCastShadows, PathState.RandSequence);
-				// the light made some contribution, and there was nothing along the shadow ray
 				AccumulateRadiance(PathState.Radiance, LightContrib, bIsCameraRay);
 			}
 		}
diff --git a/Engine/Shaders/Private/PathTracing/Volume/PathTracingVolume.ush b/Engine/Shaders/Private/PathTracing/Volume/PathTracingVolume.ush
--- a/Engine/Shaders/Private/PathTracing/Volume/PathTracingVolume.ush
+++ b/Engine/Shaders/Private/PathTracing/Volume/PathTracingVolume.ush
@@ -62,18 +62,18 @@ FVolumeDensityBounds VolumeGetDensityBounds(float3 Origin, float3 Direction, FVo
 }
 
 // Given a point in world space, return the amount of volume and its scattering properties
-FVolumeShadedResult VolumeGetDensity(float3 TranslatedWorldPos, FVolumeIntersectionInterval Interval, bool bIsCameraRay)
+FVolumeShadedResult VolumeGetDensity(float3 TranslatedWorldPos, uint VolumeMask, bool bIsCameraRay)
 {
 	FVolumeShadedResult Result = (FVolumeShadedResult)0;
-	if (Interval.VolumeMask & PATH_TRACER_VOLUME_ENABLE_ATMOSPHERE)
+	if (VolumeMask & PATH_TRACER_VOLUME_ENABLE_ATMOSPHERE)
 	{
 		MergeVolumeShadedResult(Result, AtmosphereGetDensity(TranslatedWorldPos), bIsCameraRay && (VolumeFlags & PATH_TRACER_VOLUME_HOLDOUT_ATMOSPHERE) != 0);
 	}
-	if (Interval.VolumeMask & PATH_TRACER_VOLUME_ENABLE_FOG)
+	if (VolumeMask & PATH_TRACER_VOLUME_ENABLE_FOG)
 	{
 		MergeVolumeShadedResult(Result, FogGetDensity(TranslatedWorldPos), bIsCameraRay && (VolumeFlags & PATH_TRACER_VOLUME_HOLDOUT_FOG) != 0);
 	}
-	if (Interval.VolumeMask & PATH_TRACER_VOLUME_ENABLE_HETEROGENEOUS_VOLUMES)
+	if (VolumeMask & PATH_TRACER_VOLUME_ENABLE_HETEROGENEOUS_VOLUMES)
 	{
 		MergeVolumeShadedResult(Result, HeterogeneousVolumesGetDensity(TranslatedWorldPos), bIsCameraRay && (VolumeFlags & PATH_TRACER_VOLUME_HOLDOUT_HETEROGENEOUS_VOLUMES) != 0);
 	}
