* Fix incorrect hair cards deformation.
* Fix crash when visualizing mesh deformation debug data with no visible instances.

diff --git a/Engine/Plugins/Runtime/HairStrands/Source/HairStrandsCore/Private/GroomManager.cpp b/Engine/Plugins/Runtime/HairStrands/Source/HairStrandsCore/Private/GroomManager.cpp
--- a/Engine/Plugins/Runtime/HairStrands/Source/HairStrandsCore/Private/GroomManager.cpp
+++ b/Engine/Plugins/Runtime/HairStrands/Source/HairStrandsCore/Private/GroomManager.cpp
@@ -454,6 +454,13 @@ static void AddHairSkinCacheDebugPass(
 		D.Data1.W = 0;
 	}
 
+	// Add dummy values for creating non-empty buffer
+	if (UniqueMeshCount == 0)
+	{
+		UniqueMeshNames.Add(FString(TEXT("DummyName")), 0u);
+		Infos.Add(FUniqueMeshInfo());
+	}
+
 	const uint32 InfoInBytes  = sizeof(FUniqueMeshInfo);
 	const uint32 InfoInUints = sizeof(FUniqueMeshInfo) / sizeof(uint32);
 	FRDGBufferRef InfoBuffer = CreateVertexBuffer(GraphBuilder, TEXT("Hair.Debug.UniqueMeshNames"), FRDGBufferDesc::CreateBufferDesc(4, InfoInUints * Infos.Num()), Infos.GetData(), InfoInBytes * Infos.Num());
diff --git a/Engine/Shaders/Private/HairStrands/HairCardsDeformation.usf b/Engine/Shaders/Private/HairStrands/HairCardsDeformation.usf
--- a/Engine/Shaders/Private/HairStrands/HairCardsDeformation.usf
+++ b/Engine/Shaders/Private/HairStrands/HairCardsDeformation.usf
@@ -91,13 +91,13 @@ void MainCS(uint DispatchThreadId : SV_DispatchThreadID)
 #if PERMUTATION_DYNAMIC_GEOMETRY == 0
 
 	const FCardsGuideData GuideData = UnpackCardGuideData(CardsInterpolationBuffer, VertexId);
-	const uint GuideIndex0 = GuideData.VertexIndices.x;
-	const uint GuideIndex1 = min(GuideIndex0+1, GuideVertexCount-1);
+	const uint GuidePointIndex0 = GuideData.VertexIndices.x;
+	const uint GuidePointIndex1 = min(GuidePointIndex0+1, GuideVertexCount-1);
 	const float GuideLerp = GuideData.VertexLerps.x;
 
 	const FHairControlPoint RestGuideP0 = ReadHairControlPoint(
 		GuideRestPositionBuffer, 
-		GuideIndex0,
+		GuidePointIndex0,
 		GuideRestPositionOffset,
 		GuideRadius,
 		GuideRootScale,
@@ -105,7 +105,7 @@ void MainCS(uint DispatchThreadId : SV_DispatchThreadID)
 
 	const FHairControlPoint RestGuideP1 = ReadHairControlPoint(
 		GuideRestPositionBuffer,
-		GuideIndex1,
+		GuidePointIndex1,
 		GuideRestPositionOffset,
 		GuideRadius,
 		GuideRootScale,
@@ -113,7 +113,7 @@ void MainCS(uint DispatchThreadId : SV_DispatchThreadID)
 
 	const FHairControlPoint DeformedGuideP0 = ReadHairControlPoint(
 		GuideDeformedPositionBuffer,
-		GuideIndex0,
+		GuidePointIndex0,
 		GuideDeformedPositionOffset,
 		GuideRadius,
 		GuideRootScale,
@@ -121,7 +121,7 @@ void MainCS(uint DispatchThreadId : SV_DispatchThreadID)
 
 	const FHairControlPoint DeformedGuideP1 = ReadHairControlPoint(
 		GuideDeformedPositionBuffer,
-		GuideIndex1,
+		GuidePointIndex1,
 		GuideDeformedPositionOffset,
 		GuideRadius,
 		GuideRootScale,
@@ -142,20 +142,20 @@ void MainCS(uint DispatchThreadId : SV_DispatchThreadID)
 	// Compute the deformation offset in triangle local space
 
 	const FCardsGuideData GuideData = UnpackCardGuideData(CardsInterpolationBuffer, VertexId);
-	const uint GuideIndex0 = GuideData.VertexIndices.x;
-	const uint GuideIndex1 = min(GuideIndex0 + 1, GuideVertexCount - 1);
+	const uint GuidePointIndex0 = GuideData.VertexIndices.x;
+	const uint GuidePointIndex1 = min(GuidePointIndex0 + 1, GuideVertexCount - 1);
 	const float GuideLerp = GuideData.VertexLerps.x;
 
-	const uint RootIndex = GuideVertexToRootIndexBuffer[GuideIndex0];
-	const uint TriangleIndex = GuideRootToUniqueTriangleIndexBuffer[RootIndex];
+	const uint GuideCurveIndex = UnpackHairPointToCurveIndex(GuideVertexToRootIndexBuffer, GuidePointIndex0);
+	const uint TriangleIndex = GuideRootToUniqueTriangleIndexBuffer[GuideCurveIndex];
 
-	const float3 RootBarycentric = UnpackBarycentrics(GuideRootBarycentricBuffer[RootIndex]);
+	const float3 RootBarycentric = UnpackBarycentrics(GuideRootBarycentricBuffer[GuideCurveIndex]);
 
 	const FHairMeshTriangle RestTriangle = GetTriangleTransformation(TriangleIndex, TriangleRestPositionBuffer);
 	const FHairMeshTriangle DeformedTriangle = GetTriangleTransformation(TriangleIndex, TriangleDeformedPositionBuffer);
 
-	const float3 Offset0 = ComputeDynamicGeometryOffset(GuideIndex0, 1, RestTriangle, DeformedTriangle, RootBarycentric, GuideDeformedPositionOffset);
-	const float3 Offset1 = ComputeDynamicGeometryOffset(GuideIndex1, 1, RestTriangle, DeformedTriangle, RootBarycentric, GuideDeformedPositionOffset);
+	const float3 Offset0 = ComputeDynamicGeometryOffset(GuidePointIndex0, 1, RestTriangle, DeformedTriangle, RootBarycentric, GuideDeformedPositionOffset);
+	const float3 Offset1 = ComputeDynamicGeometryOffset(GuidePointIndex1, 1, RestTriangle, DeformedTriangle, RootBarycentric, GuideDeformedPositionOffset);
 	const float3 CurrOffset = VectorToWorld(lerp(Offset0, Offset1, GuideLerp), DeformedTriangle);
 
 	// Transform hair from rest post to deformed pose, based triangle deformation + simulation offset (optional depending of the simulation is actual running or not
