Simplify cards interpolation shader struct since it uses only a single guide.

diff --git a/Engine/Shaders/Private/HairStrands/HairCardsDeformation.usf b/Engine/Shaders/Private/HairStrands/HairCardsDeformation.usf
--- a/Engine/Shaders/Private/HairStrands/HairCardsDeformation.usf
+++ b/Engine/Shaders/Private/HairStrands/HairCardsDeformation.usf
@@ -90,10 +90,9 @@ void MainCS(uint DispatchThreadId : SV_DispatchThreadID)
 
 #if PERMUTATION_DYNAMIC_GEOMETRY == 0
 
-	const FCardsGuideData GuideData = UnpackCardGuideData(CardsInterpolationBuffer, VertexId);
-	const uint GuidePointIndex0 = GuideData.VertexIndices.x;
+	const FCardsGuideData GuideData = ReadCardGuideData(CardsInterpolationBuffer, VertexId);
+	const uint GuidePointIndex0 = GuideData.VertexIndex;
 	const uint GuidePointIndex1 = min(GuidePointIndex0+1, GuideVertexCount-1);
-	const float GuideLerp = GuideData.VertexLerps.x;
 
 	const FHairControlPoint RestGuideP0 = ReadHairControlPoint(
 		GuideRestPositionBuffer, 
@@ -127,8 +126,8 @@ void MainCS(uint DispatchThreadId : SV_DispatchThreadID)
 		GuideRootScale,
 		GuideTipScale);
 
-	const float3 RestGuidePosition		= lerp(RestGuideP0.Position, RestGuideP1.Position, GuideLerp);
-	const float3 DeformedGuidePosition	= lerp(DeformedGuideP0.Position, DeformedGuideP1.Position, GuideLerp);
+	const float3 RestGuidePosition		= lerp(RestGuideP0.Position, RestGuideP1.Position, GuideData.VertexLerp);
+	const float3 DeformedGuidePosition	= lerp(DeformedGuideP0.Position, DeformedGuideP1.Position, GuideData.VertexLerp);
 
 	const float3 DeformationOffset = DeformedGuidePosition - RestGuidePosition;
 	const float4 RestPosition = CardsRestPositionBuffer[VertexId];
@@ -141,10 +140,9 @@ void MainCS(uint DispatchThreadId : SV_DispatchThreadID)
 #elif PERMUTATION_DYNAMIC_GEOMETRY == 1
 	// Compute the deformation offset in triangle local space
 
-	const FCardsGuideData GuideData = UnpackCardGuideData(CardsInterpolationBuffer, VertexId);
-	const uint GuidePointIndex0 = GuideData.VertexIndices.x;
+	const FCardsGuideData GuideData = ReadCardGuideData(CardsInterpolationBuffer, VertexId);
+	const uint GuidePointIndex0 = GuideData.VertexIndex;
 	const uint GuidePointIndex1 = min(GuidePointIndex0 + 1, GuideVertexCount - 1);
-	const float GuideLerp = GuideData.VertexLerps.x;
 
 	const uint GuideCurveIndex = UnpackHairPointToCurveIndex(GuideVertexToRootIndexBuffer, GuidePointIndex0);
 	const uint TriangleIndex = GuideRootToUniqueTriangleIndexBuffer[GuideCurveIndex];
@@ -156,7 +154,7 @@ void MainCS(uint DispatchThreadId : SV_DispatchThreadID)
 
 	const float3 Offset0 = ComputeDynamicGeometryOffset(GuidePointIndex0, 1, RestTriangle, DeformedTriangle, RootBarycentric, GuideDeformedPositionOffset);
 	const float3 Offset1 = ComputeDynamicGeometryOffset(GuidePointIndex1, 1, RestTriangle, DeformedTriangle, RootBarycentric, GuideDeformedPositionOffset);
-	const float3 CurrOffset = VectorToWorld(lerp(Offset0, Offset1, GuideLerp), DeformedTriangle);
+	const float3 CurrOffset = VectorToWorld(lerp(Offset0, Offset1, GuideData.VertexLerp), DeformedTriangle);
 
 	// Transform hair from rest post to deformed pose, based triangle deformation + simulation offset (optional depending of the simulation is actual running or not
 	const float4 RestControlPoint = CardsRestPositionBuffer[VertexId];
diff --git a/Engine/Shaders/Private/HairStrands/HairStrandsPack.ush b/Engine/Shaders/Private/HairStrands/HairStrandsPack.ush
--- a/Engine/Shaders/Private/HairStrands/HairStrandsPack.ush
+++ b/Engine/Shaders/Private/HairStrands/HairStrandsPack.ush
@@ -327,21 +327,17 @@ FGuidePointData UnpackGuidePointData(ByteAddressBuffer InBuffer, uint InPointInd
 
 struct FCardsGuideData
 {
-	uint3  VertexIndices;
-	float3 VertexLerps;
-	float3 VertexWeights;
+	uint  VertexIndex;
+	float VertexLerp;
 };
 
-FCardsGuideData UnpackCardGuideData(ByteAddressBuffer InBuffer, uint InVertexId)
+FCardsGuideData ReadCardGuideData(ByteAddressBuffer InBuffer, uint InVertexId)
 {
 	const uint Data = InBuffer.Load(InVertexId * HAIR_INTERPOLATION_CARDS_GUIDE_STRIDE);
 
 	FCardsGuideData Out = (FCardsGuideData)0;
-	Out.VertexWeights	= float3(1, 0, 0);
-	Out.VertexIndices	= 0;
-	Out.VertexLerps		= 0;
-	Out.VertexIndices.x	= Data & 0xFFFFFF;
-	Out.VertexLerps.x	= UnpackR8(Data >> 24);
+	Out.VertexIndex	= Data & 0xFFFFFF;
+	Out.VertexLerp	= UnpackR8(Data >> 24);
 	return Out;
 }
 
