Fix missing translated world space position in light function shaders

diff --git a/Engine/Shaders/Private/LightFunctionCommon.ush b/Engine/Shaders/Private/LightFunctionCommon.ush
--- a/Engine/Shaders/Private/LightFunctionCommon.ush
+++ b/Engine/Shaders/Private/LightFunctionCommon.ush
@@ -70,6 +70,7 @@ float3 GetLightFunctionColor(float3 LightVector, float3 TranslatedWorldSpace
 	MaterialParameters.ReflectionVector = LightVector.xyz;
 	MaterialParameters.LightVector = LightVector.xyz;
 	MaterialParameters.AbsoluteWorldPosition = DFFastSubtract(TranslatedWorldSpace, PrimaryView.PreViewTranslation);
+	MaterialParameters.WorldPosition_CamRelative = MaterialParameters.WorldPosition_NoOffsets_CamRelative = TranslatedWorldSpace;
 	MaterialParameters.ScreenPosition = mul(float4(TranslatedWorldSpace, 1.0f), PrimaryView.TranslatedWorldToClip);
 	MaterialParameters.LWCData = MakeMaterialLWCData(MaterialParameters);
 
