Update Compact Mesh layout.

diff --git a/Engine/Plugins/Runtime/HairStrands/Source/HairStrandsCore/Private/GroomAssetCards.cpp b/Engine/Plugins/Runtime/HairStrands/Source/HairStrandsCore/Private/GroomAssetCards.cpp
--- a/Engine/Plugins/Runtime/HairStrands/Source/HairStrandsCore/Private/GroomAssetCards.cpp
+++ b/Engine/Plugins/Runtime/HairStrands/Source/HairStrandsCore/Private/GroomAssetCards.cpp
@@ -94,7 +94,7 @@ uint32 GetHairTextureLayoutTextureCount(EHairTextureLayout In)
 		case EHairTextureLayout::Layout0: OutCount = 6u; break;
 		case EHairTextureLayout::Layout1: OutCount = 6u; break;
 		case EHairTextureLayout::Layout2: OutCount = 3u; break;
-		case EHairTextureLayout::Layout3: OutCount = 4u; break;
+		case EHairTextureLayout::Layout3: OutCount = 3u; break;
 	}
 	check(OutCount <= HAIR_CARDS_MAX_TEXTURE_COUNT);
 	return OutCount;
@@ -142,13 +142,12 @@ const TCHAR* GetHairTextureLayoutTextureName(EHairTextureLayout InLayout, uint32
 		} break;
 		case EHairTextureLayout::Layout3:
 		{
-			check(4 == GetHairTextureLayoutTextureCount(EHairTextureLayout::Layout3));
+			check(3 == GetHairTextureLayoutTextureCount(EHairTextureLayout::Layout3));
 			switch (InIndex)
 			{
-				case 0: return bDetail ? TEXT("Tangent\n RGB8")								: TEXT("Tangent");
-				case 1: return bDetail ? TEXT("RooUV | CoordU | GroupID\n RGBA8")			: TEXT("RooUVCoordUGroupID");
-				case 2: return bDetail ? TEXT("Coverage | Depth | Seed\n RGB8")				: TEXT("CoverageDepthSeed");
-				case 3: return bDetail ? TEXT("Color | Roughness\n RGBA8")					: TEXT("ColorRoughness");
+				case 0: return bDetail ? TEXT("Tangent | CoordU\n RGB8")					: TEXT("TangentCoordU");
+				case 1: return bDetail ? TEXT("RooUV | Depth | GroupID\n RGBA8")			: TEXT("RooUVDepthGroupID");
+				case 2: return bDetail ? TEXT("ColorXY | Seed | Coverage \n RGBA8")			: TEXT("ColorXYSeedCoverage");
 			}
 		} break;
 	}
diff --git a/Engine/Shaders/Private/HairStrands/HairCardsAttributeCommon.ush b/Engine/Shaders/Private/HairStrands/HairCardsAttributeCommon.ush
--- a/Engine/Shaders/Private/HairStrands/HairCardsAttributeCommon.ush
+++ b/Engine/Shaders/Private/HairStrands/HairCardsAttributeCommon.ush
@@ -37,13 +37,13 @@ uint GetHairCardTextureIndex(uint InLayoutIndex, uint InAttribute)
 	// See GroomAssetCards.cpp, EHairTextureLayout::Layout0/1/2/3 for mapping details. (9 = not present/not available)
 	switch (InAttribute)
 	{
-		case EHairCardsVFAttribute_Depth      : return LAYOUT_VALUES(InLayoutIndex, 0, 0, 1, 2);
+		case EHairCardsVFAttribute_Depth      : return LAYOUT_VALUES(InLayoutIndex, 0, 0, 1, 1);
 		case EHairCardsVFAttribute_Seed       : return LAYOUT_VALUES(InLayoutIndex, 3, 3, 1, 2);
-		case EHairCardsVFAttribute_CoordU     : return LAYOUT_VALUES(InLayoutIndex, 3, 3, 0, 1);
+		case EHairCardsVFAttribute_CoordU     : return LAYOUT_VALUES(InLayoutIndex, 3, 3, 0, 0);
 		case EHairCardsVFAttribute_GroupIndex : return LAYOUT_VALUES(InLayoutIndex, 9, 4, 9, 1); // Only used for mesh, not for cards
-		case EHairCardsVFAttribute_Roughness  : return LAYOUT_VALUES(InLayoutIndex, 9, 9, 9, 3);
+		case EHairCardsVFAttribute_Roughness  : return LAYOUT_VALUES(InLayoutIndex, 9, 9, 9, 9);
 		case EHairCardsVFAttribute_RootUV     : return LAYOUT_VALUES(InLayoutIndex, 9, 3, 9, 1);
-		case EHairCardsVFAttribute_Color      : return LAYOUT_VALUES(InLayoutIndex, 9, 4, 9, 3);
+		case EHairCardsVFAttribute_Color      : return LAYOUT_VALUES(InLayoutIndex, 9, 4, 9, 2);
 		case EHairCardsVFAttribute_Tangent    : return LAYOUT_VALUES(InLayoutIndex, 2, 2, 0, 0);
 		case EHairCardsVFAttribute_Coverage   : return LAYOUT_VALUES(InLayoutIndex, 1, 1, 1, 2);
 		case EHairCardsVFAttribute_Auxilary   : return LAYOUT_VALUES(InLayoutIndex, 5, 5, 9, 9); 
@@ -61,15 +61,15 @@ uint GetHairCardChannelIndex(uint InLayoutIndex, uint InAttribute)
 	// See GroomAssetCards.cpp, EHairTextureLayout::Layout0/1/2/3 for mapping details
 	switch (InAttribute)
 	{
-		case EHairCardsVFAttribute_Depth      : return LAYOUT_VALUES(InLayoutIndex, 0, 0, 1, 1);
+		case EHairCardsVFAttribute_Depth      : return LAYOUT_VALUES(InLayoutIndex, 0, 0, 1, 2);
 		case EHairCardsVFAttribute_Seed       : return LAYOUT_VALUES(InLayoutIndex, 3, 3, 2, 2);
-		case EHairCardsVFAttribute_CoordU     : return LAYOUT_VALUES(InLayoutIndex, 2, 2, 3, 2);
+		case EHairCardsVFAttribute_CoordU     : return LAYOUT_VALUES(InLayoutIndex, 2, 2, 3, 3);
 		case EHairCardsVFAttribute_GroupIndex : return LAYOUT_VALUES(InLayoutIndex, 3, 3, 0, 3);
 		case EHairCardsVFAttribute_Roughness  : return LAYOUT_VALUES(InLayoutIndex, 3, 0, 3, 3);
 		case EHairCardsVFAttribute_RootUV     : return LAYOUT_VALUES(InLayoutIndex, 0, 0, 0, 0);
-		case EHairCardsVFAttribute_Color      : return LAYOUT_VALUES(InLayoutIndex, 0, 0, 0, 0);
+		case EHairCardsVFAttribute_Color      : return LAYOUT_VALUES(InLayoutIndex, 0, 0, 0, 0); // For layout 3, Color is only Color.XY The Z channel is overriden by Seed. See GetHairCardsColor()
 		case EHairCardsVFAttribute_Tangent    : return LAYOUT_VALUES(InLayoutIndex, 0, 0, 0, 0);
-		case EHairCardsVFAttribute_Coverage   : return LAYOUT_VALUES(InLayoutIndex, 0, 0, 0, 0);
+		case EHairCardsVFAttribute_Coverage   : return LAYOUT_VALUES(InLayoutIndex, 0, 0, 0, 3);
 		case EHairCardsVFAttribute_Auxilary   : return LAYOUT_VALUES(InLayoutIndex, 0, 0, 0, 0); 
 	}
 	return 0;
@@ -195,7 +195,8 @@ uint3 GetHairCardsClumpID(float2 InAtlasUV)
 
 float3 GetHairCardsColor(float2 InAtlasUV, float3 InColor)
 {
-	return GetColor(InAtlasUV, InColor);
+	float3 OutColor = GetColor(InAtlasUV, InColor);
+	return HairCardsVF.LayoutIndex == 3 ? float3(OutColor.xy, 0) : OutColor;
 }
 
 float GetHairCardsRoughness(float2 InAtlasUV, float InCardRoughness)
diff --git a/Engine/Shaders/Private/HairStrands/HairStrandsTexturesGeneration.usf b/Engine/Shaders/Private/HairStrands/HairStrandsTexturesGeneration.usf
--- a/Engine/Shaders/Private/HairStrands/HairStrandsTexturesGeneration.usf
+++ b/Engine/Shaders/Private/HairStrands/HairStrandsTexturesGeneration.usf
@@ -520,13 +520,13 @@ void MainPS(
 	OutDepthTest = OutDepth;
 	OutTriangleMask = 1;
 
+	WriteValue3(LayoutIndex, EHairCardsVFAttribute_Color,      OutColor);
 	WriteValue1(LayoutIndex, EHairCardsVFAttribute_Depth,      OutDepth);
 	WriteValue1(LayoutIndex, EHairCardsVFAttribute_Seed,       OutSeed);
 	WriteValue1(LayoutIndex, EHairCardsVFAttribute_CoordU,     OutCoordU);
 	WriteValue1(LayoutIndex, EHairCardsVFAttribute_GroupIndex, OutGroupIndex);
 	WriteValue1(LayoutIndex, EHairCardsVFAttribute_Roughness,  OutRoughness);
 	WriteValue2(LayoutIndex, EHairCardsVFAttribute_RootUV,     OutRootUV);
-	WriteValue3(LayoutIndex, EHairCardsVFAttribute_Color,      OutColor);
 	WriteValue3(LayoutIndex, EHairCardsVFAttribute_Tangent,    OutTangent);
 	WriteValue1(LayoutIndex, EHairCardsVFAttribute_Coverage,   OutCoverage);
 
@@ -634,7 +634,7 @@ void MainCS(int2 DispatchThreadId : SV_DispatchThreadID)
 						case 0 : WriteOut.x = 0.f; break;
 						case 1 : WriteOut.y = 0.f; break;
 						case 2 : WriteOut.z = 0.f; break;
-						case 3 : WriteOut.z = 0.f; break;
+						case 3 : WriteOut.w = 0.f; break;
 					}
 				}
 				TargetTexture[uint3(PixelCoord, TextureIt)] = WriteOut;
@@ -691,13 +691,13 @@ void MainCS(int2 DispatchThreadId : SV_DispatchThreadID)
 		float4 OutTexture4 = 0;
 		float4 OutTexture5 = 0;
 
+		WriteValue3(LayoutIndex, EHairCardsVFAttribute_Color,      0.f);
 		WriteValue1(LayoutIndex, EHairCardsVFAttribute_Depth,      1.f);
 		WriteValue1(LayoutIndex, EHairCardsVFAttribute_Seed,       0.f);
 		WriteValue1(LayoutIndex, EHairCardsVFAttribute_CoordU,     0.f);
 		WriteValue1(LayoutIndex, EHairCardsVFAttribute_GroupIndex, 0.f);
 		WriteValue1(LayoutIndex, EHairCardsVFAttribute_Roughness,  0.f);
 		WriteValue2(LayoutIndex, EHairCardsVFAttribute_RootUV,     0.f);
-		WriteValue3(LayoutIndex, EHairCardsVFAttribute_Color,      0.f);
 		WriteValue3(LayoutIndex, EHairCardsVFAttribute_Tangent,    0.f);
 		WriteValue1(LayoutIndex, EHairCardsVFAttribute_Coverage,   0.f);
 
