Add Vertex Factory controlled VSM constant bias that applies to Non-Nanite geometry

diff --git a/Engine/Shaders/Private/ShadowDepthVertexShader.usf b/Engine/Shaders/Private/ShadowDepthVertexShader.usf
--- a/Engine/Shaders/Private/ShadowDepthVertexShader.usf
+++ b/Engine/Shaders/Private/ShadowDepthVertexShader.usf
@@ -36,6 +36,10 @@
 #include "/Engine/Generated/VertexFactory.ush"
 #include "ShadowDepthCommon.ush"
 
+// 
+#ifndef VF_USE_VSM_CONSTANT_BIAS
+#define VF_USE_VSM_CONSTANT_BIAS								0
+#endif
 
 #if VIRTUAL_SM_ENABLED
 #include "Nanite/NaniteDataDecode.ush"
@@ -152,6 +156,13 @@ void TransformToVirtualSmPage(inout float4 PointClip, inout float4 ClipPlanesInO
 	}
 
 	ScaleBiasClipToPhysicalSmPage(NaniteView, PointClip, ClipPlanesInOut, PageInfo);
+
+#if VF_USE_VSM_CONSTANT_BIAS
+	{
+		const float ConstantDepthBias = VertexFactoryGetNonNaniteVirtualShadowMapConstantDepthBias();
+		PointClip.z = clamp(PointClip.z + ConstantDepthBias * NaniteView.ViewToClip[2][2], -PointClip.w, PointClip.w);
+	}
+#endif
 }
 
 #endif // VIRTUAL_SM_ENABLED
