Path Tracing: Minor optimization for decals to avoid writes into decal payload if decals are not used

Also simplify which instance mask is used for mesh decals by deriving it inside EvaluateDecals instead of passing it from outside the call.

This saves ~0.2ms on one benchmark scene without decals, and does not measurably affect scenes with decals.

diff --git a/Engine/Shaders/Private/PathTracing/PathTracingCore.ush b/Engine/Shaders/Private/PathTracing/PathTracingCore.ush
--- a/Engine/Shaders/Private/PathTracing/PathTracingCore.ush
+++ b/Engine/Shaders/Private/PathTracing/PathTracingCore.ush
@@ -386,7 +386,7 @@ void ApplyDecal(FDecalShaderPayload DecalPayload, inout FPackedPathTracingPayloa
 #endif
 }
 
-void EvaluateDecals(FPathState PathState, uint InstanceInclusionMask, float HitT, inout FDecalShaderPayload DecalPayload)
+void EvaluateDecals(FPathState PathState, float HitT, inout FDecalShaderPayload DecalPayload)
 {
 	if (PathState.GetPathRoughness() < MeshDecalRoughnessCutoff)
 	{
@@ -397,6 +397,7 @@ void EvaluateDecals(FPathState PathState, uint InstanceInclusionMask, float HitT
 		DecalRay.TMin = 0.0f;
 		DecalRay.TMax = HitT + MeshDecalBias;
 
+		const uint InstanceInclusionMask = PathState.GetFirstScatterType() == PATHTRACER_SCATTER_CAMERA ? PATHTRACER_MASK_CAMERA : PATHTRACER_MASK_INDIRECT;
 		const uint MissShaderIndex = 0; // TODO
 
 		for (;;)
@@ -653,10 +654,10 @@ FPackedPathTracingPayload TraceTransparentRay(inout FPathState PathState, int Bo
 			break;
 		}
 
-		if (PackedPayload.IsDecalReceiver())
+		if (PackedPayload.IsDecalReceiver() && PathState.GetPathRoughness() < max(MeshDecalRoughnessCutoff, DecalRoughnessCutoff))
 		{
 			FDecalShaderPayload DecalPayload = InitDecalShaderPayload(PathState.Ray.Origin + PackedPayload.HitT * PathState.Ray.Direction);
-			EvaluateDecals(PathState, InstanceInclusionMask, PackedPayload.HitT, DecalPayload);
+			EvaluateDecals(PathState, PackedPayload.HitT, DecalPayload);
 #if USE_DBUFFER
 			if (PackedPayload.UsesDBufferLookup() && HasDecalData(DecalPayload))
 			{
