Fix inaccurate luminance/illuminance meter values in Visualize HDR mode.

We now use accurate luminance coefficients provided by the working color space. Furthermore, the maximum nit value is read from the calculated luminance, not the RGB channels directly.

diff --git a/Engine/Shaders/Private/ColorSpace.ush b/Engine/Shaders/Private/ColorSpace.ush
--- a/Engine/Shaders/Private/ColorSpace.ush
+++ b/Engine/Shaders/Private/ColorSpace.ush
@@ -75,11 +75,17 @@ static const float3 WorkingColorSpaceLuminanceFactors = Linear_sRGB_2_XYZ_MAT._m
 static const float3 WorkingColorSpaceLuminanceFactors = WORKING_COLOR_SPACE_RGB_TO_XYZ_MAT._m10_m11_m12;
 #endif
 
-float3 Luminance( float3 LinearColor, float3 LuminanceFactors )
+float Luminance( float3 LinearColor, float3 LuminanceFactors )
 {
 	return dot( LinearColor, LuminanceFactors );
 }
 
+// Working color space-accurate luminance function for scene-referred linear colors.
+float SceneLuminance( float3 LinearColor )
+{
+	return Luminance( LinearColor, WorkingColorSpaceLuminanceFactors );
+}
+
 //------------------------------------------------------------------
 // CIE LAB
 // https://en.wikipedia.org/wiki/Lab_color_space
diff --git a/Engine/Shaders/Private/PostProcessVisualizeHDR.usf b/Engine/Shaders/Private/PostProcessVisualizeHDR.usf
--- a/Engine/Shaders/Private/PostProcessVisualizeHDR.usf
+++ b/Engine/Shaders/Private/PostProcessVisualizeHDR.usf
@@ -7,6 +7,7 @@
 #include "DeferredShadingCommon.ush"
 #include "TonemapCommon.ush"
 #include "MiniFontCommon.ush"
+#include "ColorSpace.ush"
 
 // the prior frame's eye adaptation settings
 StructuredBuffer<float4> EyeAdaptationBuffer;
@@ -24,11 +25,6 @@ SCREEN_PASS_TEXTURE_VIEWPORT(Output)
 
 float IlluminanceMeterEnabled;
 
-float ComputeRgbLuminance(float3 Color)
-{
-	return dot(Color, float3(0.299, 0.587, 0.114));
-}
-
 // only needed for nice visualization
 float ComputeHistogramMax(Texture2D HistogramTexture)
 {
@@ -210,7 +206,7 @@ void MainPS(noperspective float4 UVAndScreenPos : TEXCOORD0, float4 SvPosition :
 		const float PixelDy = abs(PixelPos.y - OutputViewportCenter.y);
 		if (PixelDx < 180 && PixelDy < 42)
 		{
-			float3 LuminanceAvg = 0.0f;
+			float3 HDRSceneColorAvg = 0.0f;
 			float3 SceneColorAvg = 0.0f;
 			float SampleCount = 0.0f;
 			float MeterSize = 10.0f;
@@ -219,14 +215,14 @@ void MainPS(noperspective float4 UVAndScreenPos : TEXCOORD0, float4 SvPosition :
 				for (float OffsetY = -MeterSize; OffsetY <= MeterSize; OffsetY++)
 				{
 					float2 UV = UVViewportCenter + float2(OffsetX, OffsetY) * Input_ExtentInverse;
-					LuminanceAvg  += Texture2DSample(HDRSceneColorTexture, HDRSceneColorSampler, UV).rgb;
+					HDRSceneColorAvg  += Texture2DSample(HDRSceneColorTexture, HDRSceneColorSampler, UV).rgb;
 					SceneColorAvg += Texture2DSample(SceneColorTexture, SceneColorSampler, UV).rgb;
 					SampleCount++;
 				}
 			}
-			LuminanceAvg  = LuminanceAvg / SampleCount;
+			HDRSceneColorAvg  = HDRSceneColorAvg / SampleCount;
 			SceneColorAvg = SceneColorAvg / SampleCount;
-			LuminanceAvg *= View.OneOverPreExposure;
+			HDRSceneColorAvg *= View.OneOverPreExposure;
 
 			if (PixelDx < MeterSize && PixelDy < MeterSize)
 			{
@@ -244,6 +240,7 @@ void MainPS(noperspective float4 UVAndScreenPos : TEXCOORD0, float4 SvPosition :
 				OutColor = Texture2DSample(SceneColorTexture, SceneColorSampler, UV); // Influenced area default scene color
 			}
 
+			float3 LuminanceAvg = SceneLuminance(HDRSceneColorAvg);
 			float MaxLuminance = max(LuminanceAvg.r, max(LuminanceAvg.g, LuminanceAvg.b));
 
 			// Luminance
@@ -270,7 +267,7 @@ void MainPS(noperspective float4 UVAndScreenPos : TEXCOORD0, float4 SvPosition :
 			// EV100
 			TopLeft = OutputViewportCenter + int2(11, 22);
 
-			const float EyeAdaptationEV100 = log2((ComputeRgbLuminance(LuminanceAvg)/.18)/EyeAdaptation_LuminanceMax);
+			const float EyeAdaptationEV100 = log2((LuminanceAvg/.18)/EyeAdaptation_LuminanceMax);
 			PrintMediumFloatWithSign(PixelPos, OutColor.xyz, NumberColor, TopLeft, EyeAdaptationEV100);
 
 			TopLeft.x += 60;
@@ -301,7 +298,7 @@ void MainPS(noperspective float4 UVAndScreenPos : TEXCOORD0, float4 SvPosition :
 		const float AbsPixelDy = abs(PixelDy);
 		if (PixelDx > -20 && PixelDx < 300 && PixelDy < 40 && PixelDy > -20)
 		{
-			float3 LuminanceAvg = 0.0f;
+			float3 HDRSceneColorAvg = 0.0f;
 			float SampleCount = 0.0f;
 			LOOP
 			for (float OffsetX = (-IlluminanceMeterHalfSize + 2); OffsetX < (IlluminanceMeterHalfSize-2); OffsetX++)
@@ -309,12 +306,12 @@ void MainPS(noperspective float4 UVAndScreenPos : TEXCOORD0, float4 SvPosition :
 				for (float OffsetY = (-IlluminanceMeterHalfSize + 2); OffsetY < (IlluminanceMeterHalfSize-2); OffsetY++)
 				{
 					float2 UV = IlluminanceMeterCenterUV + float2(OffsetX, OffsetY) * Input_ExtentInverse;
-					LuminanceAvg += Texture2DSampleLevel(HDRSceneColorTexture, HDRSceneColorSampler, UV, 0).rgb;
+					HDRSceneColorAvg += Texture2DSampleLevel(HDRSceneColorTexture, HDRSceneColorSampler, UV, 0).rgb;
 					SampleCount++;
 				}
 			}
-			LuminanceAvg = LuminanceAvg / SampleCount;
-			LuminanceAvg *= View.OneOverPreExposure;
+			HDRSceneColorAvg = HDRSceneColorAvg / SampleCount;
+			HDRSceneColorAvg *= View.OneOverPreExposure;
 
 			OutColor = Texture2DSample(SceneColorTexture, SceneColorSampler, UV); // Influenced area default scene color
 			const int IlluminanceMeterHalfSizeInt = int(IlluminanceMeterHalfSize);
@@ -328,8 +325,8 @@ void MainPS(noperspective float4 UVAndScreenPos : TEXCOORD0, float4 SvPosition :
 				}
 			}
 
-			float3 Illuminance = LuminanceAvg * PI;
-			float MeanIlluminance = Luminance(Illuminance);
+			float3 Illuminance = HDRSceneColorAvg * PI;
+			float MeanIlluminance = SceneLuminance(Illuminance);
 
 			int2 TopLeft = IlluminanceMeterCenter + int2(11, 11);
 			// Mean illuminance
