Works arround shader compiler issue with bindless in TSR

diff --git a/Engine/Shaders/Private/TemporalSuperResolution/TSRConvolutionNetworkPass.ush b/Engine/Shaders/Private/TemporalSuperResolution/TSRConvolutionNetworkPass.ush
--- a/Engine/Shaders/Private/TemporalSuperResolution/TSRConvolutionNetworkPass.ush
+++ b/Engine/Shaders/Private/TemporalSuperResolution/TSRConvolutionNetworkPass.ush
@@ -122,5 +122,5 @@
 #define tsr_tensor_boolC   TLaneVector2D<bool, CONFIG_CHANNEL_COUNT, LANE_STRIDE_X, LANE_STRIDE_Y>
 
 #ifndef PLATFORM_SPECIFIC_ISOLATE
-	#define PLATFORM_SPECIFIC_ISOLATE ISOLATE
+	#define PLATFORM_SPECIFIC_ISOLATE
 #endif
