Do not discard coverage dilation.

diff --git a/Engine/Plugins/Runtime/HairStrands/Source/HairStrandsCore/Private/GroomTextureBuilder.cpp b/Engine/Plugins/Runtime/HairStrands/Source/HairStrandsCore/Private/GroomTextureBuilder.cpp
--- a/Engine/Plugins/Runtime/HairStrands/Source/HairStrandsCore/Private/GroomTextureBuilder.cpp
+++ b/Engine/Plugins/Runtime/HairStrands/Source/HairStrandsCore/Private/GroomTextureBuilder.cpp
@@ -579,6 +579,7 @@ class FHairStrandTextureDilationCS : public FGlobalShader
 		SHADER_PARAMETER_RDG_TEXTURE(Texture2D, DepthTestTexture)
 		SHADER_PARAMETER_RDG_TEXTURE(Texture2D, SourceTexture)
 		SHADER_PARAMETER_RDG_TEXTURE_UAV(RWTexture2D, TargetTexture)
+		SHADER_PARAMETER_RDG_TEXTURE(Texture2D, OriginalTriangleMaskTexture)
 		SHADER_PARAMETER_RDG_TEXTURE(Texture2D, SourceTriangleMaskTexture)
 		SHADER_PARAMETER_RDG_TEXTURE_UAV(RWTexture2D, TargetTriangleMaskTexture)
 	END_SHADER_PARAMETER_STRUCT()
@@ -601,6 +602,7 @@ static void AddTextureDilationPass(
 	const EHairTextureLayout Layout,
 	const bool bClearPass,
 	FRDGTextureRef& DepthTestTexture,
+	FRDGTextureRef& OriginalTriangleMaskTexture,
 	FRDGTextureRef& SourceTriangleMaskTexture,
 	FRDGTextureRef& TargetTriangleMaskTexture,
 	FRDGTextureRef& SourceTexture,
@@ -612,8 +614,9 @@ static void AddTextureDilationPass(
 	Parameters->bClearPass				= bClearPass ? 1u : 0u;
 	Parameters->Resolution				= Resolution;
 	Parameters->DepthTestTexture		= DepthTestTexture;
-	Parameters->SourceTriangleMaskTexture = SourceTriangleMaskTexture;
-	Parameters->TargetTriangleMaskTexture = GraphBuilder.CreateUAV(TargetTriangleMaskTexture);
+	Parameters->OriginalTriangleMaskTexture= OriginalTriangleMaskTexture;
+	Parameters->SourceTriangleMaskTexture  = SourceTriangleMaskTexture;
+	Parameters->TargetTriangleMaskTexture  = GraphBuilder.CreateUAV(TargetTriangleMaskTexture);
 	Parameters->SourceTexture 			= SourceTexture;
 	Parameters->TargetTexture 			= GraphBuilder.CreateUAV(TargetTexture);
 
@@ -1200,6 +1203,7 @@ static FHairStrandsRDGTextures DilateTextures(FRDGBuilder& GraphBuilder, FGlobal
 			In.Layout,
 			DilationIt == 0,
 			In.DepthTestTexture,
+			In.TriangleMaskTexture,
 			Out[SourceIndex].TriangleMaskTexture,
 			Out[TargetIndex].TriangleMaskTexture,
 			Out[SourceIndex].Texture,
@@ -1320,7 +1324,7 @@ void RunHairStrandsTexturesQueries(FRDGBuilder& GraphBuilder, FGlobalShaderMap*
 		{
 			if (R)
 			{
-			#if DEBUG_TEXTURE_GENERATION && DEBUG_TEXTURE_RENDERDOC
+			#if DEBUG_TEXTURE_RENDERDOC
 				RenderCaptureInterface::FScopedCapture RenderCapture(true, GraphBuilder, TEXT("PassX"));
 			#endif
 				FHairStrandsRDGTextures Textures = RegisterTextures(GraphBuilder, R->InProgress);
@@ -1380,7 +1384,7 @@ void RunHairStrandsTexturesQueries(FRDGBuilder& GraphBuilder, FGlobalShaderMap*
 			R->Output = Q.Output;
 			R->Info = Q.Info;
 
-		#if DEBUG_TEXTURE_GENERATION && DEBUG_TEXTURE_RENDERDOC
+		#if DEBUG_TEXTURE_RENDERDOC
 			RenderCaptureInterface::FScopedCapture RenderCapture(true, GraphBuilder, TEXT("Pass0"));
 		#endif
 
diff --git a/Engine/Shaders/Private/HairStrands/HairStrandsTexturesGeneration.usf b/Engine/Shaders/Private/HairStrands/HairStrandsTexturesGeneration.usf
--- a/Engine/Shaders/Private/HairStrands/HairStrandsTexturesGeneration.usf
+++ b/Engine/Shaders/Private/HairStrands/HairStrandsTexturesGeneration.usf
@@ -520,22 +520,25 @@ void MainPS(
 	OutDepthTest = OutDepth;
 	OutTriangleMask = 1;
 
-	WriteValue3(LayoutIndex, EHairCardsVFAttribute_Color,      OutColor);
-	WriteValue1(LayoutIndex, EHairCardsVFAttribute_Depth,      OutDepth);
-	WriteValue1(LayoutIndex, EHairCardsVFAttribute_Seed,       OutSeed);
-	WriteValue1(LayoutIndex, EHairCardsVFAttribute_CoordU,     OutCoordU);
-	WriteValue1(LayoutIndex, EHairCardsVFAttribute_GroupIndex, OutGroupIndex);
-	WriteValue1(LayoutIndex, EHairCardsVFAttribute_Roughness,  OutRoughness);
-	WriteValue2(LayoutIndex, EHairCardsVFAttribute_RootUV,     OutRootUV);
-	WriteValue3(LayoutIndex, EHairCardsVFAttribute_Tangent,    OutTangent);
-	WriteValue1(LayoutIndex, EHairCardsVFAttribute_Coverage,   OutCoverage);
-
-	OutTexture[uint3(PixelCoord, 0)] = OutTexture0;
-	OutTexture[uint3(PixelCoord, 1)] = OutTexture1;
-	OutTexture[uint3(PixelCoord, 2)] = OutTexture2;
-	OutTexture[uint3(PixelCoord, 3)] = OutTexture3;
-	OutTexture[uint3(PixelCoord, 4)] = OutTexture4;
-	OutTexture[uint3(PixelCoord, 5)] = OutTexture5;
+	if (bValidCount > 0)
+	{
+		WriteValue3(LayoutIndex, EHairCardsVFAttribute_Color,      OutColor);
+		WriteValue1(LayoutIndex, EHairCardsVFAttribute_Depth,      OutDepth);
+		WriteValue1(LayoutIndex, EHairCardsVFAttribute_Seed,       OutSeed);
+		WriteValue1(LayoutIndex, EHairCardsVFAttribute_CoordU,     OutCoordU);
+		WriteValue1(LayoutIndex, EHairCardsVFAttribute_GroupIndex, OutGroupIndex);
+		WriteValue1(LayoutIndex, EHairCardsVFAttribute_Roughness,  OutRoughness);
+		WriteValue2(LayoutIndex, EHairCardsVFAttribute_RootUV,     OutRootUV);
+		WriteValue3(LayoutIndex, EHairCardsVFAttribute_Tangent,    OutTangent);
+		WriteValue1(LayoutIndex, EHairCardsVFAttribute_Coverage,   OutCoverage);
+
+		OutTexture[uint3(PixelCoord, 0)] = OutTexture0;
+		OutTexture[uint3(PixelCoord, 1)] = OutTexture1;
+		OutTexture[uint3(PixelCoord, 2)] = OutTexture2;
+		OutTexture[uint3(PixelCoord, 3)] = OutTexture3;
+		OutTexture[uint3(PixelCoord, 4)] = OutTexture4;
+		OutTexture[uint3(PixelCoord, 5)] = OutTexture5;
+	}
 }
 
 #endif // SHADER_PIXEL
@@ -552,6 +555,7 @@ uint TextureCount;
 int2 Resolution;
 int ClearDistance;
 Texture2D<float> DepthTestTexture;
+Texture2D<uint> OriginalTriangleMaskTexture;
 Texture2D<uint> SourceTriangleMaskTexture;
 RWTexture2D<uint> TargetTriangleMaskTexture;
 Texture2DArray<float4> SourceTexture;
@@ -620,6 +624,9 @@ void MainCS(int2 DispatchThreadId : SV_DispatchThreadID)
 			}
 		}
 
+		// For coverage dialation, we only allow dilation outside of the triangles
+		const bool bIsInsideTriangle = OriginalTriangleMaskTexture.Load(uint3(PixelCoord,0)) > 0;
+
 		// If a valid neighbord has been bound, update pixel value
 		if (bIsDilatedPixelValid)
 		{
@@ -627,7 +634,7 @@ void MainCS(int2 DispatchThreadId : SV_DispatchThreadID)
 			{
 				// Write all attributes but coverage, as coverage should remain untouched.
 				float4 WriteOut = Values[TextureIt];
-				if (TextureIt == CoverageTextureIndex)
+				if (bIsInsideTriangle && TextureIt == CoverageTextureIndex)
 				{
 					switch (CoverageChannelIndex)
 					{
