Heterogeneous Volumes: Code cleanup. Removing unused types and includes.

diff --git a/Engine/Shaders/Private/BasePassPixelShader.usf b/Engine/Shaders/Private/BasePassPixelShader.usf
--- a/Engine/Shaders/Private/BasePassPixelShader.usf
+++ b/Engine/Shaders/Private/BasePassPixelShader.usf
@@ -8,7 +8,6 @@
 #define UE_DF_FORCE_FP32_OPS 1
 #endif
 
-#include "HeterogeneousVolumes/HeterogeneousVolumesAdaptiveVolumetricShadowMapTypes.ush" 
 #include "Common.ush"
 
 #ifndef ADAPTIVE_VOLUMETRIC_SHADOW_MAP
diff --git a/Engine/Shaders/Private/DeferredLightPixelShaders.usf b/Engine/Shaders/Private/DeferredLightPixelShaders.usf
--- a/Engine/Shaders/Private/DeferredLightPixelShaders.usf
+++ b/Engine/Shaders/Private/DeferredLightPixelShaders.usf
@@ -23,7 +23,6 @@
 #endif
 
 
-#include "HeterogeneousVolumes/HeterogeneousVolumesAdaptiveVolumetricShadowMapTypes.ush" 
 //#include "HeterogeneousVolumes/HeterogeneousVolumesVoxelGridTypes.ush"
 #include "HairStrands/HairStrandsVisibilityCommonStruct.ush" 
 #include "Common.ush"
diff --git a/Engine/Shaders/Private/HeterogeneousVolumes/HeterogeneousVolumesAdaptiveVolumetricShadowMapTypes.ush b/Engine/Shaders/Private/HeterogeneousVolumes/HeterogeneousVolumesAdaptiveVolumetricShadowMapTypes.ush
deleted file mode 100644
diff --git a/Engine/Shaders/Private/HeterogeneousVolumes/HeterogeneousVolumesAdaptiveVolumetricShadowMapUtils.ush b/Engine/Shaders/Private/HeterogeneousVolumes/HeterogeneousVolumesAdaptiveVolumetricShadowMapUtils.ush
--- a/Engine/Shaders/Private/HeterogeneousVolumes/HeterogeneousVolumesAdaptiveVolumetricShadowMapUtils.ush
+++ b/Engine/Shaders/Private/HeterogeneousVolumes/HeterogeneousVolumesAdaptiveVolumetricShadowMapUtils.ush
@@ -2,7 +2,7 @@
 
 #pragma once
 
-#include "HeterogeneousVolumesAdaptiveVolumetricShadowMapTypes.ush"
+#define AVSM_NULL_PTR 0xFF
 
 struct FAVSMSampleData
 {
diff --git a/Engine/Shaders/Private/HeterogeneousVolumes/HeterogeneousVolumesLiveShadingPipeline.usf b/Engine/Shaders/Private/HeterogeneousVolumes/HeterogeneousVolumesLiveShadingPipeline.usf
--- a/Engine/Shaders/Private/HeterogeneousVolumes/HeterogeneousVolumesLiveShadingPipeline.usf
+++ b/Engine/Shaders/Private/HeterogeneousVolumes/HeterogeneousVolumesLiveShadingPipeline.usf
@@ -1,6 +1,5 @@
 // Copyright Epic Games, Inc. All Rights Reserved.
 
-#include "HeterogeneousVolumesAdaptiveVolumetricShadowMapTypes.ush"
 #include "HeterogeneousVolumesSparseVoxelUtils.ush"
 #include "../Common.ush"
 
diff --git a/Engine/Shaders/Private/HeterogeneousVolumes/HeterogeneousVolumesVoxelGridRendering.usf b/Engine/Shaders/Private/HeterogeneousVolumes/HeterogeneousVolumesVoxelGridRendering.usf
--- a/Engine/Shaders/Private/HeterogeneousVolumes/HeterogeneousVolumesVoxelGridRendering.usf
+++ b/Engine/Shaders/Private/HeterogeneousVolumes/HeterogeneousVolumesVoxelGridRendering.usf
@@ -1,7 +1,6 @@
 // Copyright Epic Games, Inc. All Rights Reserved.
 
 #include "HeterogeneousVolumesVoxelGridTypes.ush"
-#include "HeterogeneousVolumesAdaptiveVolumetricShadowMapTypes.ush"
 
 #include "../Common.ush"
 #include "../DeferredShadingCommon.ush"
diff --git a/Engine/Shaders/Private/HeterogeneousVolumes/HeterogeneousVolumesVoxelGridShadows.usf b/Engine/Shaders/Private/HeterogeneousVolumes/HeterogeneousVolumesVoxelGridShadows.usf
--- a/Engine/Shaders/Private/HeterogeneousVolumes/HeterogeneousVolumesVoxelGridShadows.usf
+++ b/Engine/Shaders/Private/HeterogeneousVolumes/HeterogeneousVolumesVoxelGridShadows.usf
@@ -1,7 +1,6 @@
 // Copyright Epic Games, Inc. All Rights Reserved.
 
 #include "HeterogeneousVolumesVoxelGridTypes.ush"
-#include "HeterogeneousVolumesAdaptiveVolumetricShadowMapTypes.ush"
 
 #include "../Common.ush"
 #include "../DeferredShadingCommon.ush"
@@ -190,7 +189,7 @@ float AbsoluteErrorThreshold;
 float RelativeErrorThreshold;
 
 RWBuffer<int> RWVolumetricShadowLinkedListAllocatorBuffer;
-RWStructuredBuffer<int2> RWVolumetricShadowLinkedListBuffer;
+RWStructuredBuffer<uint2> RWVolumetricShadowLinkedListBuffer;
 
 // Debug
 struct FVolumetricShadowMapDebugData
@@ -449,7 +448,7 @@ void RenderVolumetricShadowMapForLightWithVoxelGridCS(
 	}
 }
 
-StructuredBuffer<int2> VolumetricShadowLinkedListBuffer;
+StructuredBuffer<uint2> VolumetricShadowLinkedListBuffer;
 
 RWBuffer<int> RWVolumetricShadowIndirectionAllocatorBuffer;
 RWStructuredBuffer<uint2> RWVolumetricShadowIndirectionBuffer;
diff --git a/Engine/Shaders/Private/TranslucentLightInjectionShaders.usf b/Engine/Shaders/Private/TranslucentLightInjectionShaders.usf
--- a/Engine/Shaders/Private/TranslucentLightInjectionShaders.usf
+++ b/Engine/Shaders/Private/TranslucentLightInjectionShaders.usf
@@ -4,7 +4,6 @@
  * TranslucentLightInjectionShaders.usf: Shaders for calculating lighting in a volume to use on translucency
  */
 
-#include "HeterogeneousVolumes/HeterogeneousVolumesAdaptiveVolumetricShadowMapTypes.ush"
 #include "Common.ush"
 #include "SHCommon.ush"
 
diff --git a/Engine/Source/Runtime/Renderer/Private/HeterogeneousVolumes/HeterogeneousVolumes.h b/Engine/Source/Runtime/Renderer/Private/HeterogeneousVolumes/HeterogeneousVolumes.h
--- a/Engine/Source/Runtime/Renderer/Private/HeterogeneousVolumes/HeterogeneousVolumes.h
+++ b/Engine/Source/Runtime/Renderer/Private/HeterogeneousVolumes/HeterogeneousVolumes.h
@@ -300,7 +300,7 @@ void RegisterExternalOrthoVoxelGridUniformBuffer(
 
 struct FAVSMLinkedListPackedData
 {
-	int32 Data[2];
+	uint32 Data[2];
 };
 
 struct FAVSMIndirectionPackedData
@@ -324,7 +324,7 @@ BEGIN_UNIFORM_BUFFER_STRUCT(FAdaptiveVolumetricShadowMapUniformBufferParameters,
 	SHADER_PARAMETER(int32, bIsEmpty)
 	SHADER_PARAMETER(int32, bIsDirectionalLight)
 
-	SHADER_PARAMETER_RDG_BUFFER_SRV(StructuredBuffer<int2>, LinkedListBuffer)
+	SHADER_PARAMETER_RDG_BUFFER_SRV(StructuredBuffer<uint2>, LinkedListBuffer)
 	SHADER_PARAMETER_RDG_BUFFER_SRV(StructuredBuffer<uint2>, IndirectionBuffer)
 	SHADER_PARAMETER_RDG_BUFFER_SRV(StructuredBuffer<uint>, SampleBuffer)
 END_UNIFORM_BUFFER_STRUCT()
diff --git a/Engine/Source/Runtime/Renderer/Private/HeterogeneousVolumes/HeterogeneousVolumesVoxelGridPipeline.cpp b/Engine/Source/Runtime/Renderer/Private/HeterogeneousVolumes/HeterogeneousVolumesVoxelGridPipeline.cpp
--- a/Engine/Source/Runtime/Renderer/Private/HeterogeneousVolumes/HeterogeneousVolumesVoxelGridPipeline.cpp
+++ b/Engine/Source/Runtime/Renderer/Private/HeterogeneousVolumes/HeterogeneousVolumesVoxelGridPipeline.cpp
@@ -2908,7 +2908,7 @@ class FCompressVolumetricShadowMapCS : public FGlobalShader
 		// Input
 		SHADER_PARAMETER(FIntPoint, ShadowResolution)
 		SHADER_PARAMETER(int, MaxSampleCount)
-		SHADER_PARAMETER_RDG_BUFFER_SRV(StructuredBuffer<int2>, VolumetricShadowLinkedListBuffer)
+		SHADER_PARAMETER_RDG_BUFFER_SRV(StructuredBuffer<uint2>, VolumetricShadowLinkedListBuffer)
 
 		// Output
 		SHADER_PARAMETER_RDG_BUFFER_UAV(RWBuffer<int>, RWVolumetricShadowIndirectionAllocatorBuffer)
