Fixed lumen reflection discrepancy on post motion blur translucents.

diff --git a/Engine/Shaders/Private/ForwardLightingCommon.ush b/Engine/Shaders/Private/ForwardLightingCommon.ush
--- a/Engine/Shaders/Private/ForwardLightingCommon.ush
+++ b/Engine/Shaders/Private/ForwardLightingCommon.ush
@@ -402,7 +402,10 @@ half3 GetImageBasedReflectionSpecular(FMaterialPixelParameters MaterialParameter
 
 #if (TRANSLUCENCY_LIGHTING_SURFACE_FORWARDSHADING || TRANSLUCENCY_LIGHTING_SURFACE_LIGHTINGVOLUME) && MATERIALBLENDING_ANY_TRANSLUCENT && PROJECT_SUPPORTS_LUMEN && !FORWARD_SHADING
 	
-	bUseLumenFrontLayerReflection = UseFrontLayerReflection(MaterialParameters.ViewBufferUV, MaterialParameters.ScreenPosition.w);
+	// For translucent rendered after motion blur, we cannot use Lumen front layer reflection because depth test is disabled in this case 
+	// (depth is jittered and post motion blur composition can happen at a different resolution with dynamic resolution scaling).
+	// Such Lumen reflections would look different on translucent meshes when occluded by opaque meshes.
+	bUseLumenFrontLayerReflection = UseFrontLayerReflection(MaterialParameters.ViewBufferUV, MaterialParameters.ScreenPosition.w) && !MATERIAL_TRANSLUCENT_PASS_AFTERMOTIONBLUR;
 
 	FRadianceCacheCoverage LumenRadianceCacheCoverage = InitRadianceCacheCoverage();
 	LumenRadianceCacheCoverage.bValid = false;
diff --git a/Engine/Source/Runtime/Engine/Classes/Materials/Material.h b/Engine/Source/Runtime/Engine/Classes/Materials/Material.h
--- a/Engine/Source/Runtime/Engine/Classes/Materials/Material.h
+++ b/Engine/Source/Runtime/Engine/Classes/Materials/Material.h
@@ -151,7 +151,7 @@ enum EMaterialTranslucencyPass : int
 	MTP_BeforeDOF UMETA(DisplayName="Before DOF"),
 	/** Render after depth of field. */
 	MTP_AfterDOF UMETA(DisplayName="After DOF"),
-	/** Render after motion blur. */
+	/** Render after motion blur. Disables depth test (the reconstruction post MB and TSR would otherwise flicker due to the TSR camera jittering making the depth buffer unstable). Because of that Lumen high quality reflections is also disabled to avoid visual discrepancy at depth intersection. */
 	MTP_AfterMotionBlur UMETA(DisplayName="After Motion Blur"),
 	MTP_MAX
 };
diff --git a/Engine/Source/Runtime/Engine/Private/Materials/MaterialShared.cpp b/Engine/Source/Runtime/Engine/Private/Materials/MaterialShared.cpp
--- a/Engine/Source/Runtime/Engine/Private/Materials/MaterialShared.cpp
+++ b/Engine/Source/Runtime/Engine/Private/Materials/MaterialShared.cpp
@@ -2729,6 +2729,7 @@ void FMaterial::SetupMaterialEnvironment(
 	SET_SHADER_DEFINE(OutEnvironment, TRANSLUCENT_WRITING_FRONT_LAYER_TRANSPARENCY, IsTranslucencyWritingFrontLayerTransparency());
 	SET_SHADER_DEFINE(OutEnvironment, MATERIAL_USE_ALPHA_TO_COVERAGE, IsUsingAlphaToCoverage());
 	SET_SHADER_DEFINE(OutEnvironment, MOBILE_HIGH_QUALITY_BRDF, IsMobileHighQualityBRDFEnabled());
+	SET_SHADER_DEFINE(OutEnvironment, MATERIAL_TRANSLUCENT_PASS_AFTERMOTIONBLUR, IsTranslucencyAfterMotionBlurEnabled() ? 1 : 0);
 
 	bool bFullPrecisionInMaterial = false;
 	bool bFullPrecisionInPS = false;
