Substrate - fixed crashing when accessing scene textures normal in decals.
Only normals should be accessible using the deferred decals (reconstructed from depth buffer).

diff --git a/Engine/Shaders/Private/MaterialTemplate.ush b/Engine/Shaders/Private/MaterialTemplate.ush
--- a/Engine/Shaders/Private/MaterialTemplate.ush
+++ b/Engine/Shaders/Private/MaterialTemplate.ush
@@ -2977,7 +2977,30 @@ float4 SceneTextureLookup(float2 UV, int SceneTextureIndex, bool bFiltered)
 
 	// 2. Post-process inputs which are dependent on material data
 	//    Note: For PostProcess material we always return the first layer dataCalcSceneCustomStencil
-	#if SUBSTRATE_ENABLED
+	#if SUBSTRATE_ENABLED && IS_DBUFFER_DECAL
+	{
+		// A special decal path to avoid requiring the binding of substrate data in the decal passes. Decals can only access depth, stencil and depth-from-normals or reprojected.
+		switch (SceneTextureIndex)
+		{
+		case PPI_DiffuseColor:			return 0.0f;
+		case PPI_SpecularColor:			return 0.0f;
+		case PPI_SubsurfaceColor:		return 0.0f;
+		case PPI_BaseColor:				return 0.0f;
+		case PPI_Specular:				return 0.0f;
+		case PPI_Metallic:				return 0.0f;
+		case PPI_WorldNormal:			return GetDBufferReprojectedWorldNormal(UV); // DBuffer uses depth buffer gradients or reprojection of last frame normal.
+		case PPI_Opacity:				return 1.0f;
+		case PPI_Roughness:				return 0.0f;
+		case PPI_MaterialAO:			return 0.0f;
+		case PPI_ShadingModelColor:		return 0.0f;
+		case PPI_ShadingModelID:		return 0.0f;
+		case PPI_StoredBaseColor:		return 0.0f;
+		case PPI_StoredSpecular:		return 0.0f;
+		case PPI_WorldTangent:			return 0.0f;
+		case PPI_Anisotropy:			return 0.0f;
+		}
+	}
+	#elif SUBSTRATE_ENABLED
 	{
 		#ifdef SubstrateStruct
 			#define SubstrateSceneTexture SubstrateStruct
