Remove a loop for CSM cascade selection in mobile shading to avoid a driver issue on Adreno 7xx

diff --git a/Engine/Shaders/Private/MobileLightingCommon.ush b/Engine/Shaders/Private/MobileLightingCommon.ush
--- a/Engine/Shaders/Private/MobileLightingCommon.ush
+++ b/Engine/Shaders/Private/MobileLightingCommon.ush
@@ -143,40 +143,39 @@ half MobileDirectionalLightCSM(float2 ScreenPosition, float SceneDepth, inout fl
 	Settings.DensityMulConstant = 0;
 	Settings.ProjectionDepthBiasParameters = 0;
 
-	float4 ShadowPosition = float4(0, 0, 0, 0);
-	for (uint i = 0; i < MAX_MOBILE_SHADOWCASCADES; i++)
+	float4 Count = float4(SceneDepth.xxxx >= MobileDirectionalLight.DirectionalLightShadowDistances);
+	uint CascadeIndex = uint(Count.x + Count.y + Count.z + Count.w);
+	
+	if (CascadeIndex < MobileDirectionalLight.DirectionalLightNumCascades)
 	{
-		if (SceneDepth < MobileDirectionalLight.DirectionalLightShadowDistances[i])
-		{
-		#if MOBILE_MULTI_VIEW
-			ShadowPosition = mul(float4(ScreenPosition.x, ScreenPosition.y, SceneDepth, 1), ResolvedView.MobileMultiviewShadowTransform);
-			ShadowPosition = mul(ShadowPosition, MobileDirectionalLight.DirectionalLightScreenToShadow[i]);
-		#else
-			ShadowPosition = mul(float4(ScreenPosition.x, ScreenPosition.y, SceneDepth, 1), MobileDirectionalLight.DirectionalLightScreenToShadow[i]);
-		#endif
-			ShadowPositionZ = ShadowPosition.z;
-			break; // position found.
-		}
-	}
+		float4 ShadowPosition = float4(0, 0, 0, 0);
+	#if MOBILE_MULTI_VIEW
+		ShadowPosition = mul(float4(ScreenPosition.x, ScreenPosition.y, SceneDepth, 1), ResolvedView.MobileMultiviewShadowTransform);
+		ShadowPosition = mul(ShadowPosition, MobileDirectionalLight.DirectionalLightScreenToShadow[CascadeIndex]);
+	#else
+		ShadowPosition = mul(float4(ScreenPosition.x, ScreenPosition.y, SceneDepth, 1), MobileDirectionalLight.DirectionalLightScreenToShadow[CascadeIndex]);
+	#endif
+		ShadowPositionZ = ShadowPosition.z;
 
-	// Process CSM only when ShadowPosition is valid.
-	if (ShadowPosition.z > 0)
-	{
-		// Clamp pixel depth in light space for shadowing opaque, because areas of the shadow depth buffer that weren't rendered to will have been cleared to 1
-		// We want to force the shadow comparison to result in 'unshadowed' in that case, regardless of whether the pixel being shaded is in front or behind that plane
+		// Process CSM only when ShadowPosition is valid.
+		if (ShadowPosition.z > 0)
+		{
+			// Clamp pixel depth in light space for shadowing opaque, because areas of the shadow depth buffer that weren't rendered to will have been cleared to 1
+			// We want to force the shadow comparison to result in 'unshadowed' in that case, regardless of whether the pixel being shaded is in front or behind that plane
 
-		// Invert ShadowZ as the shadow space has been changed (but not yet the filtering code)
-		float ShadowZ = 1.0f - ShadowPosition.z;
-		float LightSpacePixelDepthForOpaque = min(ShadowZ, 0.99999f);
-		Settings.SceneDepth = LightSpacePixelDepthForOpaque;
+			// Invert ShadowZ as the shadow space has been changed (but not yet the filtering code)
+			float ShadowZ = 1.0f - ShadowPosition.z;
+			float LightSpacePixelDepthForOpaque = min(ShadowZ, 0.99999f);
+			Settings.SceneDepth = LightSpacePixelDepthForOpaque;
 
-		ShadowMap = MobileShadowPCF(ShadowPosition.xy, Settings);
+			ShadowMap = MobileShadowPCF(ShadowPosition.xy, Settings);
 
-	#if FADE_CSM
-		float Fade = saturate(SceneDepth * MobileDirectionalLight.DirectionalLightDistanceFadeMADAndSpecularScale.x + MobileDirectionalLight.DirectionalLightDistanceFadeMADAndSpecularScale.y);
-		// lerp out shadow based on fade params.
-		ShadowMap = lerp(ShadowMap, 1.0, Fade * Fade);
-	#endif
+		#if FADE_CSM
+			float Fade = saturate(SceneDepth * MobileDirectionalLight.DirectionalLightDistanceFadeMADAndSpecularScale.x + MobileDirectionalLight.DirectionalLightDistanceFadeMADAndSpecularScale.y);
+			// lerp out shadow based on fade params.
+			ShadowMap = lerp(ShadowMap, 1.0, Fade * Fade);
+		#endif
+		}
 	}
 #endif //ENABLE_MOBILE_CSM
 	return ShadowMap;
diff --git a/Engine/Source/Runtime/Engine/Private/SceneManagement.cpp b/Engine/Source/Runtime/Engine/Private/SceneManagement.cpp
--- a/Engine/Source/Runtime/Engine/Private/SceneManagement.cpp
+++ b/Engine/Source/Runtime/Engine/Private/SceneManagement.cpp
@@ -1072,10 +1072,11 @@ FMobileDirectionalLightShaderParameters::FMobileDirectionalLightShaderParameters
 	DirectionalLightShadowSampler = TStaticSamplerState<SF_Point, AM_Clamp, AM_Clamp, AM_Clamp>::GetRHI();
 	DirectionalLightShadowSize = FVector4f(EForceInit::ForceInitToZero);
 	DirectionalLightDistanceFadeMADAndSpecularScale = FVector4f(EForceInit::ForceInitToZero);
+	DirectionalLightNumCascades = 0;
 	for (int32 i = 0; i < MAX_MOBILE_SHADOWCASCADES; ++i)
 	{
 		DirectionalLightScreenToShadow[i].SetIdentity();
-		DirectionalLightShadowDistances[i] = 0.0f;
+		DirectionalLightShadowDistances[i] = FLT_MAX; // Unused cascades should compare > all scene depths
 	}
 }
 
diff --git a/Engine/Source/Runtime/Engine/Public/SceneView.h b/Engine/Source/Runtime/Engine/Public/SceneView.h
--- a/Engine/Source/Runtime/Engine/Public/SceneView.h
+++ b/Engine/Source/Runtime/Engine/Public/SceneView.h
@@ -720,6 +720,7 @@ BEGIN_GLOBAL_SHADER_PARAMETER_STRUCT_WITH_CONSTRUCTOR(FMobileDirectionalLightSha
 	SHADER_PARAMETER_EX(FVector4f, DirectionalLightDistanceFadeMADAndSpecularScale, EShaderPrecisionModifier::Half) // .z is used for SpecularScale, .w is not used atm
 	SHADER_PARAMETER_EX(FVector4f, DirectionalLightShadowDistances, EShaderPrecisionModifier::Half)
 	SHADER_PARAMETER_ARRAY(FMatrix44f, DirectionalLightScreenToShadow, [MAX_MOBILE_SHADOWCASCADES])
+	SHADER_PARAMETER(uint32, DirectionalLightNumCascades)
 	SHADER_PARAMETER(uint32, DirectionalLightShadowMapChannelMask)
 	SHADER_PARAMETER_TEXTURE(Texture2D, DirectionalLightShadowTexture)
 	SHADER_PARAMETER_SAMPLER(SamplerState, DirectionalLightShadowSampler)
diff --git a/Engine/Source/Runtime/Renderer/Private/MobileBasePassRendering.cpp b/Engine/Source/Runtime/Renderer/Private/MobileBasePassRendering.cpp
--- a/Engine/Source/Runtime/Renderer/Private/MobileBasePassRendering.cpp
+++ b/Engine/Source/Runtime/Renderer/Private/MobileBasePassRendering.cpp
@@ -427,6 +427,7 @@ void SetupMobileDirectionalLightUniformParameters(
 					}
 					Params.DirectionalLightScreenToShadow[OutShadowIndex] = FMatrix44f(ShadowInfo->GetScreenToShadowMatrix(SceneView));		// LWC_TODO: Precision loss?
 					Params.DirectionalLightShadowDistances[OutShadowIndex] = ShadowInfo->CascadeSettings.SplitFar;
+					Params.DirectionalLightNumCascades++;
 					OutShadowIndex++;
 				}
 			}
