Substrate - fixed missing static reflection captures when skylight is not present.

diff --git a/Engine/Shaders/Private/ReflectionEnvironmentPixelShader.usf b/Engine/Shaders/Private/ReflectionEnvironmentPixelShader.usf
--- a/Engine/Shaders/Private/ReflectionEnvironmentPixelShader.usf
+++ b/Engine/Shaders/Private/ReflectionEnvironmentPixelShader.usf
@@ -412,7 +412,6 @@ void ReflectionEnvironmentSkyLighting(
 		const float CloudVolumetricAOShadow = GetCloudVolumetricAOShadow(TranslatedWorldPosition);
 		const FSubstrateIntegrationSettings Settings = InitSubstrateIntegrationSettings(false /*bForceFullyRough*/, Substrate.bRoughDiffuse, Substrate.PeelLayersAboveDepth, Substrate.bRoughnessTracking);
 
-		const bool bEnableSpecular = ReflectionStruct.SkyLightParameters.y > 0.0f;
 		FSubstrateDeferredLighting SubstrateLighting = GetInitialisedSubstrateDeferredLighting();
 	
 	#if APPLY_SKY_SHADOWING
@@ -434,6 +433,8 @@ void ReflectionEnvironmentSkyLighting(
 		uint CaptureDataStartIndex = 0;
 		uint NumCulledReflectionCaptures = 0;
 	#endif
+
+		const bool bEnableSpecular = ReflectionStruct.SkyLightParameters.y > 0.0f || NumCulledReflectionCaptures > 0;
 	
 		Substrate_for (uint ClosureIndex = 0, ClosureIndex < SubstratePixelHeader.ClosureCount, ++ClosureIndex)
 		{
