Changes UMaterial::bHasPixelAnimation to also use the opaque geometry through the translucency rejection code path

diff --git a/Engine/Shaders/Private/TemporalSuperResolution/TSRCommon.ush b/Engine/Shaders/Private/TemporalSuperResolution/TSRCommon.ush
--- a/Engine/Shaders/Private/TemporalSuperResolution/TSRCommon.ush
+++ b/Engine/Shaders/Private/TemporalSuperResolution/TSRCommon.ush
@@ -14,7 +14,7 @@
 
 //------------------------------------------------------- RECOMPILE HASH
 
-#pragma message("UESHADERMETADATA_VERSION D9BD9048-6528-48DC-922C-82BAF2992B38")
+#pragma message("UESHADERMETADATA_VERSION F3C5FF82-E809-4903-98CD-5CC7E2254C77")
 
 
 //------------------------------------------------------- COMPILER CONFIG
diff --git a/Engine/Shaders/Private/TemporalSuperResolution/TSRConvolutionNetworkPass.ush b/Engine/Shaders/Private/TemporalSuperResolution/TSRConvolutionNetworkPass.ush
--- a/Engine/Shaders/Private/TemporalSuperResolution/TSRConvolutionNetworkPass.ush
+++ b/Engine/Shaders/Private/TemporalSuperResolution/TSRConvolutionNetworkPass.ush
@@ -122,5 +122,9 @@
 #define tsr_tensor_boolC   TLaneVector2D<bool, CONFIG_CHANNEL_COUNT, LANE_STRIDE_X, LANE_STRIDE_Y>
 
 #ifndef PLATFORM_SPECIFIC_ISOLATE
-	#define PLATFORM_SPECIFIC_ISOLATE
+	#if ENABLE_BINDLESS_RESOURCES
+		#define PLATFORM_SPECIFIC_ISOLATE
+	#else
+		#define PLATFORM_SPECIFIC_ISOLATE ISOLATE
+	#endif
 #endif
diff --git a/Engine/Shaders/Private/TemporalSuperResolution/TSRDecimateHistory.usf b/Engine/Shaders/Private/TemporalSuperResolution/TSRDecimateHistory.usf
--- a/Engine/Shaders/Private/TemporalSuperResolution/TSRDecimateHistory.usf
+++ b/Engine/Shaders/Private/TemporalSuperResolution/TSRDecimateHistory.usf
@@ -427,6 +427,7 @@ void MainCS(
 	#endif
 	float PrevDeviceZ;
 	tsr_half VelocityEdge = tsr_half(1.0);
+	bool bHasPixelAnimation;
 	ISOLATE
 	{
 		tsr_ushort2 InputPixelPos = (
@@ -443,10 +444,12 @@ void MainCS(
 			#if CONFIG_RESURRECTION_REPROJECTION
 				float2 DeviceZAndPrevDeviceZ = ClosestDepthTexture[InputPixelPos];
 				DeviceZ = DeviceZAndPrevDeviceZ.g;
-				PrevDeviceZ = DeviceZAndPrevDeviceZ.r;
 			#else
-				PrevDeviceZ = ClosestDepthTexture[InputPixelPos];
+				float DeviceZAndPrevDeviceZ = ClosestDepthTexture[InputPixelPos];
 			#endif
+			
+			PrevDeviceZ = abs(DeviceZAndPrevDeviceZ.r);
+			bHasPixelAnimation = DeviceZAndPrevDeviceZ.r < 0.0;
 
 			#if CONFIG_VELOCITY_SAMPLES == 1
 				EncodedScreenVelocity = DilatedVelocityTexture[InputPixelPos];
@@ -676,6 +679,14 @@ void MainCS(
 		}
 	}
 	
+	// Builds the decimate bitmask
+	tsr_ushort DecimateBitMask;
+	{
+		bool bIsDisoccluded = !(ParallaxRejectionMask > tsr_half(PARALLAX_REJECTION_MASK_THRESHOLD));
+
+		DecimateBitMask = select(bIsDisoccluded, tsr_ushort(0x1), tsr_ushort(0x0)) | select(bHasPixelAnimation, tsr_ushort(0x2), tsr_ushort(0x0));
+	}
+
 	// Output full res
 	ISOLATE
 	{
@@ -718,7 +729,7 @@ void MainCS(
 		#endif
 
 		HoleFilledVelocityOutput[OutputPixelPos] = EncodedScreenVelocity;
-		DecimateMaskOutput[InputPixelPos] = tsr_half2(ParallaxRejectionMask, VelocityEdge);
+		DecimateMaskOutput[InputPixelPos] = tsr_half2(tsr_half(DecimateBitMask) * rcp(tsr_half(255.0)), VelocityEdge);
 
 		#if DEBUG_OUTPUT
 		for (uint DebugId = 0; DebugId < DEBUG_ARRAY_SIZE; DebugId++)
diff --git a/Engine/Shaders/Private/TemporalSuperResolution/TSRDilateVelocity.usf b/Engine/Shaders/Private/TemporalSuperResolution/TSRDilateVelocity.usf
--- a/Engine/Shaders/Private/TemporalSuperResolution/TSRDilateVelocity.usf
+++ b/Engine/Shaders/Private/TemporalSuperResolution/TSRDilateVelocity.usf
@@ -206,7 +206,7 @@ void MainCS(
 	float4 CenterEncodedVelocity = SceneVelocityTexture[int2(InputBufferUV * InputInfo_Extent)];
 	bool bIsMovingPixel = CenterEncodedVelocity.x > 0.0;
 	
-	bool HasPixelAnimation = false;
+	bool bHasPixelAnimation = false;
 	#if !VELOCITY_ENCODE_HAS_PIXEL_ANIMATION
 		// UE-198416: there is high mess arround VelocityEncodeDepth() to clean up
 		//#error TSR requires the bHasPixelAnimation encoded in the velocity buffer.
@@ -232,7 +232,7 @@ void MainCS(
 		float3 ScreenVelocity2 = ComputeStaticVelocity(ScreenPos, ClosestDeviceZ2);
 		
 		#if VELOCITY_ENCODE_HAS_PIXEL_ANIMATION
-			HasPixelAnimation = DecodeHasPixelAnimationFromVelocityTexture(EncodedVelocity1);
+			bHasPixelAnimation = DecodeHasPixelAnimationFromVelocityTexture(EncodedVelocity1);
 		#endif
 
 		#if CONFIG_USE_SCENE_VELOCITY
@@ -261,7 +261,7 @@ void MainCS(
 		float4 EncodedVelocity = SceneVelocityTexture[uint2(InputBufferUV * InputInfo_Extent) + VelocityPixelOffset];
 
 		#if VELOCITY_ENCODE_HAS_PIXEL_ANIMATION
-			HasPixelAnimation = DecodeHasPixelAnimationFromVelocityTexture(EncodedVelocity);
+			bHasPixelAnimation = DecodeHasPixelAnimationFromVelocityTexture(EncodedVelocity);
 		#endif
 
 		#if CONFIG_USE_SCENE_VELOCITY
@@ -304,7 +304,7 @@ void MainCS(
 			tsr_half(saturate(length(TranslatedWorldPosition - PrevTranslatedWorldPosition) * rcp(StaticWorldRadius) - 1.0)),
 			saturate(PixelParallaxVelocity * tsr_half(InvFlickeringMaxParralaxVelocity) - tsr_half(0.5)));
 
-		IsMovingMask = select(HasPixelAnimation, tsr_half(1.0), IsMovingMask);
+		IsMovingMask = select(bHasPixelAnimation, tsr_half(1.0), IsMovingMask);
 	}
 
 	// Do motion blur velocity flatten.
@@ -366,7 +366,7 @@ void MainCS(
 
 		//DilatedVelocityOutput[OutputPixelPos] = EncodedVelocity.xy; //EncodeVelocityToTexture(float3(FinalScreenVelocity.xy, 0.0)).xy;
 		DilatedVelocityOutput[OutputPixelPos] = EncodeVelocityToTexture(float3(FinalScreenVelocity.xy, 0.0)).xy;
-		ClosestDepthOutput[OutputPixelPos] = float2(PrevClosestDeviceZ, ClosestDeviceZ);
+		ClosestDepthOutput[OutputPixelPos] = float2(select(bHasPixelAnimation, -PrevClosestDeviceZ, PrevClosestDeviceZ), ClosestDeviceZ);
 
 		BRANCH
 		if (bOutputIsMovingTexture)
diff --git a/Engine/Shaders/Private/TemporalSuperResolution/TSRRejectShading.usf b/Engine/Shaders/Private/TemporalSuperResolution/TSRRejectShading.usf
--- a/Engine/Shaders/Private/TemporalSuperResolution/TSRRejectShading.usf
+++ b/Engine/Shaders/Private/TemporalSuperResolution/TSRRejectShading.usf
@@ -68,7 +68,7 @@ RWTexture2D<tsr_ushort> AntiAliasMaskOutput;
 
 static uint2 GGroupId = 0;
 
-static tsr_tensor_halfC Debug[DEBUG_ARRAY_SIZE];
+static tsr_tensor_half4 Debug[DEBUG_ARRAY_SIZE];
 
 
 //------------------------------------------------------- INPUT FETCHES
@@ -115,7 +115,8 @@ tsr_short2 ComputeElementOutputPixelPos(const uint ElementIndex)
 void FetchDecimateMask(
 	tsr_tensor_short2 PixelPos,
 	out tsr_tensor_half VelocityEdge,
-	out tsr_tensor_bool bIsDisoccluded)
+	out tsr_tensor_bool bIsDisoccluded,
+	out tsr_tensor_bool bHasPixelAnimation)
 {
 	UNROLL_N(SIMD_SIZE)
 	for (uint ElementIndex = 0; ElementIndex < SIMD_SIZE; ElementIndex++)
@@ -124,8 +125,11 @@ void FetchDecimateMask(
 		
 		tsr_half2 DecimateMask = DecimateMaskTexture[ClampedFetchPixelPos];
 
+		tsr_ushort BitMask = tsr_ushort(DecimateMask.r * tsr_half(255.0));
+
 		VelocityEdge.SetElement(ElementIndex, DecimateMask.g);
-		bIsDisoccluded.SetElement(ElementIndex, !(DecimateMask.r > tsr_half(PARALLAX_REJECTION_MASK_THRESHOLD)));
+		bIsDisoccluded.SetElement(ElementIndex,     (BitMask & tsr_ushort(0x1)) != tsr_ushort(0));
+		bHasPixelAnimation.SetElement(ElementIndex, (BitMask & tsr_ushort(0x2)) != tsr_ushort(0));
 	}
 } // FetchDecimateMask()
 
@@ -362,8 +366,9 @@ void RefetchInputAndHistoryFromLDSOrCache(
 	tsr_tensor_half  OriginalHistoryUncertainty;
 	tsr_tensor_half  VelocityEdge;
 	tsr_tensor_bool  bIsDisoccluded;
+	tsr_tensor_bool  bHasPixelAnimation;
 
-	FetchDecimateMask(PixelPos, /* out */ VelocityEdge, /* out */ bIsDisoccluded);
+	FetchDecimateMask(PixelPos, /* out */ VelocityEdge, /* out */ bIsDisoccluded, /* out */ bHasPixelAnimation);
 	FetchSceneColorAndTranslucency(PixelPos, /* out */ OriginalOpaqueInput, /* out */ OriginalTranslucencyInput);
 	FetchReprojectedHistoryGuide(
 		PixelPos,
@@ -411,6 +416,7 @@ void MainCS(
 		// Load its own data
 		tsr_tensor_half  VelocityEdge;
 		tsr_tensor_bool  bIsDisoccluded;
+		tsr_tensor_bool  bHasPixelAnimation;
 		tsr_tensor_halfC OriginalOpaqueInput;
 		tsr_tensor_half4 OriginalTranslucencyInput;
 		tsr_tensor_half  MoireInput;
@@ -420,7 +426,7 @@ void MainCS(
 		{
 			tsr_tensor_short2 PixelPos = ComputePixelPos();
 		
-			FetchDecimateMask(PixelPos, /* out */ VelocityEdge, /* out */ bIsDisoccluded);
+			FetchDecimateMask(PixelPos, /* out */ VelocityEdge, /* out */ bIsDisoccluded, /* out */ bHasPixelAnimation);
 			FetchSceneColorAndTranslucency(PixelPos, /* out */ OriginalOpaqueInput, /* out */ OriginalTranslucencyInput);
 			FetchMoireInput(
 				PixelPos,
@@ -458,6 +464,7 @@ void MainCS(
 	// Load all the history data.
 	tsr_tensor_half  VelocityEdge;
 	tsr_tensor_bool  bIsDisoccluded;
+	tsr_tensor_bool  bHasPixelAnimation;
 	tsr_tensor_halfC OriginalOpaqueInput;
 	tsr_tensor_half4 OriginalTranslucencyInput;
 	tsr_tensor_halfC OriginalHistory;
@@ -474,7 +481,7 @@ void MainCS(
 		{
 			tsr_tensor_short2 PixelPos = ComputePixelPos();
 
-			FetchDecimateMask(PixelPos, /* out */ VelocityEdge, /* out */ bIsDisoccluded);
+			FetchDecimateMask(PixelPos, /* out */ VelocityEdge, /* out */ bIsDisoccluded, /* out */ bHasPixelAnimation);
 			FetchSceneColorAndTranslucency(PixelPos, /* out */ OriginalOpaqueInput, /* out */ OriginalTranslucencyInput);
 			FetchReprojectedHistoryGuide(
 				PixelPos,
@@ -491,7 +498,7 @@ void MainCS(
 		{
 			tsr_tensor_short2 PixelPos = ComputePixelPos();
 
-			FetchDecimateMask(PixelPos, /* out */ VelocityEdge, /* out */ bIsDisoccluded);
+			FetchDecimateMask(PixelPos, /* out */ VelocityEdge, /* out */ bIsDisoccluded, /* out */ bHasPixelAnimation);
 			FetchSceneColorAndTranslucency(PixelPos, /* out */ OriginalOpaqueInput, /* out */ OriginalTranslucencyInput);
 			FetchReprojectedHistoryGuide(
 				PixelPos,
@@ -523,8 +530,25 @@ void MainCS(
 			}
 		}
 	}
-	
-	tsr_tensor_halfC OriginalInput = ComposeTranslucency(OriginalOpaqueInput, Blur3x3(OriginalTranslucencyInput));
+
+	tsr_tensor_halfC OriginalInput;
+	{
+		tsr_tensor_halfC SharpInput = OriginalOpaqueInput;
+		tsr_tensor_half4 BlurInput = OriginalTranslucencyInput;
+
+		//BRANCH
+		//if (AnyComponent(bHasPixelAnimation))
+		{
+			tsr_tensor_halfC CenterFinalSceneColor = ComposeTranslucency(OriginalOpaqueInput, OriginalTranslucencyInput);
+			
+			tsr_tensor_half4 HasAnimationBlurInput = ResizeChannels<4>(CenterFinalSceneColor);
+			
+			BlurInput  = select(tsr_tensor_bool4::Vectorize(bHasPixelAnimation), HasAnimationBlurInput, BlurInput);
+			SharpInput = select(tsr_tensor_boolC::Vectorize(bHasPixelAnimation), tsr_tensor_halfC::Const(0.0), SharpInput);
+		}
+		
+		OriginalInput = ComposeTranslucency(SharpInput, Blur3x3(BlurInput));
+	}
 
 	// Despill to LDS to make room in VGPRs for the convolutions of ComputeMoireError()
 	#if CONFIG_LDS_DESPILL
@@ -820,7 +844,7 @@ void MainCS(
 			UNROLL_N(DEBUG_ARRAY_SIZE)
 			for (uint DebugId = 0; DebugId < DEBUG_ARRAY_SIZE; DebugId++)
 			{
-				DebugOutput[tsr_short3(OutputPixelPos, DebugId)] = float4(ResizeChannels<4>(Debug[DebugId]).GetElement(ElementIndex));
+				DebugOutput[tsr_short3(OutputPixelPos, DebugId)] = float4(Debug[DebugId].GetElement(ElementIndex));
 			}
 		}
 	}
