Added support for bAffectsTranslucent for all material, not only translucent light volume.
Substrate deffered and mobile tested.

diff --git a/Engine/Shaders/Private/ForwardLightingCommon.ush b/Engine/Shaders/Private/ForwardLightingCommon.ush
--- a/Engine/Shaders/Private/ForwardLightingCommon.ush
+++ b/Engine/Shaders/Private/ForwardLightingCommon.ush
@@ -152,7 +152,11 @@ FDeferredLightingSplit GetForwardDirectLightingSplit(
 	const FDirectionalLightData DirectionalLightData = GetDirectionalLightData(EyeIndex);
 
 	BRANCH
-	if (DirectionalLightData.HasDirectionalLight)
+	if (DirectionalLightData.HasDirectionalLight
+#if MATERIALBLENDING_ANY_TRANSLUCENT
+		&& DirectionalLightData.bAffectsTranslucentLighting > 0
+#endif
+		)
 	{
 		half4 PreviewShadowMapChannelMask = 1;
 		uint  DirLightingChannelMask = LIGHTING_CHANNEL_MASK;
@@ -275,7 +279,16 @@ FDeferredLightingSplit GetForwardDirectLightingSplit(
 		half4 PreviewShadowMapChannelMask = 1;
 		uint  LocalLightingChannelMask = LIGHTING_CHANNEL_MASK;
 		const FLocalLightData LocalLight = GetLocalLightData(CulledLightsGrid.DataStartIndex + GridLightListIndex, EyeIndex);
+		
+	#if MATERIALBLENDING_ANY_TRANSLUCENT
+		if(UnpackAffectsTranslucentLighting(LocalLight) == 0)
+		{
+			continue;
+		}
+	#endif
+
 		FDeferredLightData LightData = ConvertToDeferredLight(LocalLight, SpecularScale, PreviewShadowMapChannelMask, LocalLightingChannelMask);
+
 		// Rect light optionally supported in forward.
 		LightData.bRectLight = LightData.bRectLight && SUPPORT_RECTLIGHT_ON_FORWARD_LIT_TRANSLUCENT;
 		#if USE_HAIR_COMPLEX_TRANSMITTANCE
diff --git a/Engine/Shaders/Private/LightData.ush b/Engine/Shaders/Private/LightData.ush
--- a/Engine/Shaders/Private/LightData.ush
+++ b/Engine/Shaders/Private/LightData.ush
@@ -63,6 +63,7 @@ struct FDirectionalLightData
 	float  DirectionalLightSoftSourceRadius;
 	float  DirectionalLightSpecularScale;
 	uint  LightFunctionAtlasLightIndex;
+	uint  bAffectsTranslucentLighting;
 };
 
 ///////////////////////////////////////////////////////////////////////////////////////////////////
@@ -130,6 +131,7 @@ struct FLightShaderParameters
 	float  RectLightAtlasMaxLevel;
 	float  IESAtlasIndex;
 	uint  LightFunctionAtlasLightIndex;
+	uint  bAffectsTranslucentLighting;
 };
 
 ///////////////////////////////////////////////////////////////////////////////////////////////////
@@ -181,6 +183,8 @@ struct FDeferredLightData
 	uint  LightFunctionAtlasLightIndex;
 	/** Hair transmittance data. */
 	FHairTransmittanceData HairTransmittance;
+	// Wether the light affect translucent materials or not 
+	uint  bAffectsTranslucentLighting;
 };
 
 ///////////////////////////////////////////////////////////////////////////////////////////////////
@@ -247,6 +251,17 @@ uint UnpackLightFunctionAtlasIndex(FLocalLightData In)
 	return UnpackLightFunctionAtlasIndex(LightTypeAndPackedShadowMapChannelMask);
 }
 
+uint UnpackAffectsTranslucentLighting(uint InShadowMapChannelMaskPacked)
+{
+	return ((InShadowMapChannelMaskPacked >> 28) & 0x1);
+}
+
+uint UnpackAffectsTranslucentLighting(FLocalLightData In)
+{
+	const uint LightTypeAndPackedShadowMapChannelMask = asuint(In.LightDirectionAndShadowMask.w);
+	return UnpackAffectsTranslucentLighting(LightTypeAndPackedShadowMapChannelMask);
+}
+
 float UnpackVolumetricScatteringIntensity(FLocalLightData In)
 {
 	return f16tof32(asuint(In.SpotAnglesAndSourceRadiusPacked.w) >> 16);
@@ -334,6 +349,8 @@ FDeferredLightData ConvertToDeferredLight(
 	Out.IESAtlasIndex			= UnpackLigthIESAtlasIndex(In);
 	Out.LightFunctionAtlasLightIndex = UnpackLightFunctionAtlasIndex(In);
 
+	Out.bAffectsTranslucentLighting = UnpackAffectsTranslucentLighting(In);
+
 	// Shadow
 	// * Static shadowing uses ShadowMapChannel, 
 	// * Dynamic shadows are packed into light attenuation using PreviewShadowMapChannel
@@ -414,6 +431,8 @@ FLightShaderParameters ConvertFromLocal(const FLocalLightData In)
 
 	Out.LightFunctionAtlasLightIndex = UnpackLightFunctionAtlasIndex(In);
 
+	Out.bAffectsTranslucentLighting = UnpackAffectsTranslucentLighting(In);
+
 	const FRectLightData RectLightData = UnpackRectLightData(
 									asuint(In.RectData.x),
 									asuint(In.RectData.y),
@@ -446,6 +465,7 @@ FDeferredLightData ConvertToDeferredLight(
 	Out.SourceRadius		= In.DirectionalLightSourceRadius;
 	Out.SoftSourceRadius	= In.DirectionalLightSoftSourceRadius;
 	Out.LightFunctionAtlasLightIndex = In.LightFunctionAtlasLightIndex;
+	Out.bAffectsTranslucentLighting = In.bAffectsTranslucentLighting;
 
 	// Shadow
 	// * Static shadowing uses ShadowMapChannel, 
@@ -494,6 +514,8 @@ FDeferredLightData ConvertToDeferredLight(FLightSceneData In)
 	Out.IESAtlasIndex					= In.IESAtlasIndex;
 	Out.LightFunctionAtlasLightIndex	= UnpackLightFunctionAtlasIndex(In.LightTypeAndShadowMapChannelMaskPacked);
 
+	Out.bAffectsTranslucentLighting = UnpackAffectsTranslucentLighting(In.LightTypeAndShadowMapChannelMaskPacked);
+
 	// Shadow
 	// * Static shadowing uses ShadowMapChannel, 
 	// * Dynamic shadows are packed into light attenuation using PreviewShadowMapChannel
diff --git a/Engine/Shaders/Private/LightDataUniforms.ush b/Engine/Shaders/Private/LightDataUniforms.ush
--- a/Engine/Shaders/Private/LightDataUniforms.ush
+++ b/Engine/Shaders/Private/LightDataUniforms.ush
@@ -42,6 +42,8 @@ FDeferredLightData InitDeferredLightFromUniforms(uint InLightType)
 
 	Out.LightFunctionAtlasLightIndex = DeferredLightUniforms.LightFunctionAtlasLightIndex;
 
+	Out.bAffectsTranslucentLighting = DeferredLightUniforms.bAffectsTranslucentLighting;
+
 	Out.HairTransmittance		= InitHairTransmittanceData();
 	return Out;
 }
diff --git a/Engine/Shaders/Private/LightGridCommon.ush b/Engine/Shaders/Private/LightGridCommon.ush
--- a/Engine/Shaders/Private/LightGridCommon.ush
+++ b/Engine/Shaders/Private/LightGridCommon.ush
@@ -204,6 +204,7 @@ FDirectionalLightData GetDirectionalLightData(uint EyeIndex)
 		Result.DirectionalLightSoftSourceRadius		= ForwardLightData.DirectionalLightSoftSourceRadius;
 		Result.DirectionalLightSpecularScale		= ForwardLightData.DirectionalLightSpecularScale;
 		Result.LightFunctionAtlasLightIndex			= ForwardLightData.LightFunctionAtlasLightIndex;
+		Result.bAffectsTranslucentLighting			= ForwardLightData.bAffectsTranslucentLighting;
 
 #if (INSTANCED_STEREO || MOBILE_MULTI_VIEW)
 	}
@@ -218,6 +219,7 @@ FDirectionalLightData GetDirectionalLightData(uint EyeIndex)
 		Result.DirectionalLightSoftSourceRadius		= ForwardLightDataISR.DirectionalLightSoftSourceRadius;
 		Result.DirectionalLightSpecularScale		= ForwardLightDataISR.DirectionalLightSpecularScale;
 		Result.LightFunctionAtlasLightIndex			= ForwardLightDataISR.LightFunctionAtlasLightIndex;
+		Result.bAffectsTranslucentLighting			= ForwardLightDataISR.bAffectsTranslucentLighting;
 	}
 #endif
 
diff --git a/Engine/Shaders/Private/Substrate/SubstrateForwardLighting.ush b/Engine/Shaders/Private/Substrate/SubstrateForwardLighting.ush
--- a/Engine/Shaders/Private/Substrate/SubstrateForwardLighting.ush
+++ b/Engine/Shaders/Private/Substrate/SubstrateForwardLighting.ush
@@ -422,7 +422,11 @@ float3 SubstrateForwardLighting(
 				//// Evaluate the single directional light selected for forward lighting.
 				////
 				BRANCH
-				if (DirectionalLightData.HasDirectionalLight)
+				if (DirectionalLightData.HasDirectionalLight
+			#if SUBSTRATE_TRANSLUCENT_MATERIAL
+					&& DirectionalLightData.bAffectsTranslucentLighting > 0
+			#endif
+						)
 				{
 					float3 DirLightL = DirLightData.Direction;	// Already normalized
 					float3 ToDirLight = DirLightL;
@@ -475,10 +479,18 @@ float3 SubstrateForwardLighting(
 				for (uint GridLightListIndex = 0; GridLightListIndex < NumLightsInGridCell; GridLightListIndex++)
 				{
 					const FLocalLightData LocalLight = GetLocalLightData(CulledLightsGrid.DataStartIndex + GridLightListIndex, EyeIndex);
+		
+				#if MATERIALBLENDING_ANY_TRANSLUCENT
+					if(UnpackAffectsTranslucentLighting(LocalLight) == 0)
+					{
+						continue;
+					}
+				#endif
 
 					half4 PreviewShadowMapChannelMask = 1;
 					uint  LocalLightChannelMask = LIGHTING_CHANNEL_MASK;
 					FDeferredLightData LightData = ConvertToDeferredLight(LocalLight, 1.0f /*SpecularScale*/, PreviewShadowMapChannelMask, LocalLightChannelMask);
+
 					#if USE_IES_PROFILE
 					LightData.Color *= ComputeLightProfileMultiplier(TranslatedWorldPosition, LightData.TranslatedWorldPosition, -LightData.Direction, LightData.Tangent, LightData.IESAtlasIndex);
 					#endif
diff --git a/Engine/Shaders/Private/Substrate/SubstrateMobileForwardLighting.ush b/Engine/Shaders/Private/Substrate/SubstrateMobileForwardLighting.ush
--- a/Engine/Shaders/Private/Substrate/SubstrateMobileForwardLighting.ush
+++ b/Engine/Shaders/Private/Substrate/SubstrateMobileForwardLighting.ush
@@ -229,39 +229,44 @@ float3 SubstrateMobileForwardLighting(
 					uint DirLightingChannelMask = LIGHTING_CHANNEL_MASK;
 					float3 ToDirLight = DirLightData.Direction;	// Already normalized
 					float DirLightMask = 1;
-	
-					float DirectionalLightCloudShadow = 1.0f;
-				#if NEEDS_BASEPASS_CLOUD_SHADOW_INTERPOLATOR
-					DirectionalLightCloudShadow = BasePassInterpolants.VertexCloudShadow;
-				#endif
 
-					FTranslucencyLightingVector LightingVector = GetTranslucencyLightingVector(DirLightData, TranslatedWorldPosition);
-				#if SUBSTRATE_TRANSLUCENCY_LIGHTING_VOLUMETRIC_DIRECTIONAL
-					GetVolumeLightingDirectional(LightingVector.AmbientLightingVector, LightingVector.DirectionalLightingVector, SubstrateBSDFContext.N, BSDFEvaluate.DiffuseColor, GetMaterialTranslucencyDirectionalLightingIntensity(), InterpolatedLighting, VolumeLighting);
-				#elif SUBSTRATE_TRANSLUCENCY_LIGHTING_VOLUMETRIC_NONDIRECTIONAL
-					GetVolumeLightingNonDirectional(LightingVector.AmbientLightingVector, BSDFEvaluate.DiffuseColor, InterpolatedLighting, VolumeLighting);
-				#else
-					SubstrateBSDFContext.SubstrateUpdateBSDFContext(ToDirLight);
-
-					FSubstrateEvaluateResult BSDFEvaluate = (FSubstrateEvaluateResult)0;
-					Color += SubstrateForwardLightingCommon(
-						Dither,
-						Settings,
-						DirLightData,
-						ToDirLight,
-						DirLightMask,
-						DirLightAttenuation,
-						InitRectTexture(),
-						DirLightingChannelMask,
-						PrimitiveLightingChannelMask,
-						PrecomputedShadowFactors,
-						TranslatedWorldPosition,
-						SceneDepth,
-						LuminanceWeight(SubstrateBSDFContext, CurrentBSDF),
-						SubstratePixelHeader,
-						SubstrateBSDFContext,
-						BSDFEvaluate) * DirectionalLightCloudShadow;
+				#if SUBSTRATE_TRANSLUCENT_MATERIAL
+					if (DirLightData.bAffectsTranslucentLighting > 0)
 				#endif
+					{
+						float DirectionalLightCloudShadow = 1.0f;
+					#if NEEDS_BASEPASS_CLOUD_SHADOW_INTERPOLATOR
+						DirectionalLightCloudShadow = BasePassInterpolants.VertexCloudShadow;
+					#endif
+
+						FTranslucencyLightingVector LightingVector = GetTranslucencyLightingVector(DirLightData, TranslatedWorldPosition);
+					#if SUBSTRATE_TRANSLUCENCY_LIGHTING_VOLUMETRIC_DIRECTIONAL
+						GetVolumeLightingDirectional(LightingVector.AmbientLightingVector, LightingVector.DirectionalLightingVector, SubstrateBSDFContext.N, BSDFEvaluate.DiffuseColor, GetMaterialTranslucencyDirectionalLightingIntensity(), InterpolatedLighting, VolumeLighting);
+					#elif SUBSTRATE_TRANSLUCENCY_LIGHTING_VOLUMETRIC_NONDIRECTIONAL
+						GetVolumeLightingNonDirectional(LightingVector.AmbientLightingVector, BSDFEvaluate.DiffuseColor, InterpolatedLighting, VolumeLighting);
+					#else
+						SubstrateBSDFContext.SubstrateUpdateBSDFContext(ToDirLight);
+
+						FSubstrateEvaluateResult BSDFEvaluate = (FSubstrateEvaluateResult)0;
+						Color += SubstrateForwardLightingCommon(
+							Dither,
+							Settings,
+							DirLightData,
+							ToDirLight,
+							DirLightMask,
+							DirLightAttenuation,
+							InitRectTexture(),
+							DirLightingChannelMask,
+							PrimitiveLightingChannelMask,
+							PrecomputedShadowFactors,
+							TranslatedWorldPosition,
+							SceneDepth,
+							LuminanceWeight(SubstrateBSDFContext, CurrentBSDF),
+							SubstratePixelHeader,
+							SubstrateBSDFContext,
+							BSDFEvaluate) * DirectionalLightCloudShadow;
+					#endif
+					}
 				}
 
 				////
@@ -287,6 +292,13 @@ float3 SubstrateMobileForwardLighting(
 				#endif
 					{
 						const FLocalLightData LocalLight = GetLocalLightData(CulledLightsGrid.DataStartIndex + GridLightListIndex, EyeIndex);
+		
+					#if MATERIALBLENDING_ANY_TRANSLUCENT && !SUBSTRATE_FORCE_SINGLE_LOCAL_LIGHT
+						if(UnpackAffectsTranslucentLighting(LocalLight) == 0)
+						{
+							LOCAL_LIGHT_EARLY_OUT;
+						}
+					#endif
 
 						half4 PreviewShadowMapChannelMask = 1;
 						half4 DynamicShadowFactors = 1;
diff --git a/Engine/Source/Runtime/Engine/Private/Components/DirectionalLightComponent.cpp b/Engine/Source/Runtime/Engine/Private/Components/DirectionalLightComponent.cpp
--- a/Engine/Source/Runtime/Engine/Private/Components/DirectionalLightComponent.cpp
+++ b/Engine/Source/Runtime/Engine/Private/Components/DirectionalLightComponent.cpp
@@ -406,6 +406,7 @@ public:
 		LightParameters.RectLightAtlasMaxLevel = FLightRenderParameters::GetRectLightAtlasInvalidMIPLevel();
 		LightParameters.IESAtlasIndex = INDEX_NONE;
 		LightParameters.LightFunctionAtlasLightIndex = GetLightFunctionAtlasLightIndex();
+		LightParameters.bAffectsTranslucentLighting = AffectsTranslucentLighting() ? 1 : 0;
 		LightParameters.InverseExposureBlend = 0.0f;
 	}
 
diff --git a/Engine/Source/Runtime/Engine/Private/Components/LightComponent.cpp b/Engine/Source/Runtime/Engine/Private/Components/LightComponent.cpp
--- a/Engine/Source/Runtime/Engine/Private/Components/LightComponent.cpp
+++ b/Engine/Source/Runtime/Engine/Private/Components/LightComponent.cpp
@@ -351,6 +351,7 @@ void FLightRenderParameters::MakeShaderParameters(const FViewMatrices& ViewMatri
 	OutShaderParameters.RectLightAtlasMaxLevel = RectLightAtlasMaxLevel;
 	OutShaderParameters.IESAtlasIndex = IESAtlasIndex;
 	OutShaderParameters.LightFunctionAtlasLightIndex = LightFunctionAtlasLightIndex;
+	OutShaderParameters.bAffectsTranslucentLighting = bAffectsTranslucentLighting;
 }
 
 // match logic in InverseExposureLerp(...)
diff --git a/Engine/Source/Runtime/Engine/Private/Components/PointLightComponent.cpp b/Engine/Source/Runtime/Engine/Private/Components/PointLightComponent.cpp
--- a/Engine/Source/Runtime/Engine/Private/Components/PointLightComponent.cpp
+++ b/Engine/Source/Runtime/Engine/Private/Components/PointLightComponent.cpp
@@ -56,6 +56,7 @@ void FPointLightSceneProxy::GetLightShaderParameters(FLightRenderParameters& Lig
 	LightParameters.IESAtlasIndex = INDEX_NONE;
 	LightParameters.InverseExposureBlend = InverseExposureBlend;
 	LightParameters.LightFunctionAtlasLightIndex = GetLightFunctionAtlasLightIndex();
+	LightParameters.bAffectsTranslucentLighting = AffectsTranslucentLighting() ? 1 : 0;
 
 	if (IESAtlasId != ~0)
 	{
diff --git a/Engine/Source/Runtime/Engine/Private/Components/RectLightComponent.cpp b/Engine/Source/Runtime/Engine/Private/Components/RectLightComponent.cpp
--- a/Engine/Source/Runtime/Engine/Private/Components/RectLightComponent.cpp
+++ b/Engine/Source/Runtime/Engine/Private/Components/RectLightComponent.cpp
@@ -250,6 +250,7 @@ void FRectLightSceneProxy::GetLightShaderParameters(FLightRenderParameters& Ligh
 	LightParameters.IESAtlasIndex = INDEX_NONE;
 	LightParameters.InverseExposureBlend = InverseExposureBlend;
 	LightParameters.LightFunctionAtlasLightIndex = GetLightFunctionAtlasLightIndex();
+	LightParameters.bAffectsTranslucentLighting = AffectsTranslucentLighting() ? 1 : 0;
 
 	if (IESAtlasId != ~0)
 	{
diff --git a/Engine/Source/Runtime/Engine/Private/Components/SpotLightComponent.cpp b/Engine/Source/Runtime/Engine/Private/Components/SpotLightComponent.cpp
--- a/Engine/Source/Runtime/Engine/Private/Components/SpotLightComponent.cpp
+++ b/Engine/Source/Runtime/Engine/Private/Components/SpotLightComponent.cpp
@@ -76,6 +76,7 @@ public:
 		LightParameters.RectLightAtlasMaxLevel = FLightRenderParameters::GetRectLightAtlasInvalidMIPLevel();
 		LightParameters.IESAtlasIndex = INDEX_NONE;
 		LightParameters.LightFunctionAtlasLightIndex = GetLightFunctionAtlasLightIndex();
+		LightParameters.bAffectsTranslucentLighting = AffectsTranslucentLighting() ? 1 : 0;
 		LightParameters.InverseExposureBlend = InverseExposureBlend;
 
 		if (IESAtlasId != uint32(INDEX_NONE))
diff --git a/Engine/Source/Runtime/Engine/Public/SceneManagement.h b/Engine/Source/Runtime/Engine/Public/SceneManagement.h
--- a/Engine/Source/Runtime/Engine/Public/SceneManagement.h
+++ b/Engine/Source/Runtime/Engine/Public/SceneManagement.h
@@ -1272,6 +1272,9 @@ BEGIN_SHADER_PARAMETER_STRUCT(FLightShaderParameters, ENGINE_API)
 	// Index of the light function in the atlas
 	SHADER_PARAMETER(uint32, LightFunctionAtlasLightIndex)
 
+	// Wether the light affect translucent material or not
+	SHADER_PARAMETER(uint32, bAffectsTranslucentLighting)
+
 END_SHADER_PARAMETER_STRUCT()
 
 
@@ -1347,6 +1350,9 @@ struct FLightRenderParameters
 	// Index of the light in the Light function atlas data
 	uint32 LightFunctionAtlasLightIndex;
 
+	// Wether this lights affect translucent materials or not
+	uint32 bAffectsTranslucentLighting;
+
 	float InverseExposureBlend;
 
 	// Return Invalid rect light atlas MIP level
diff --git a/Engine/Source/Runtime/Renderer/Private/LightGridInjection.cpp b/Engine/Source/Runtime/Renderer/Private/LightGridInjection.cpp
--- a/Engine/Source/Runtime/Renderer/Private/LightGridInjection.cpp
+++ b/Engine/Source/Runtime/Renderer/Private/LightGridInjection.cpp
@@ -141,6 +141,8 @@ void SetupDummyForwardLightUniformParameters(FRDGBuilder& GraphBuilder, FForward
 	ForwardLightData.CulledLightDataGrid16Bit = CulledLightDataGridSRV;
 
 	ForwardLightData.LightFunctionAtlasLightIndex = 0;
+
+	ForwardLightData.bAffectsTranslucentLighting = 0;
 }
 
 TRDGUniformBufferRef<FForwardLightData> CreateDummyForwardLightUniformBuffer(FRDGBuilder& GraphBuilder, EShaderPlatform ShaderPlatform)
@@ -656,6 +658,7 @@ FComputeLightGridOutput FSceneRenderer::ComputeLightGrid(FRDGBuilder& GraphBuild
 							ForwardLightData->DirectionalLightVSM = INDEX_NONE;
 							ForwardLightData->DirectionalLightSMRTSettings = GetVirtualShadowMapSMRTSettings(true);
 							ForwardLightData->LightFunctionAtlasLightIndex = LightParameters.LightFunctionAtlasLightIndex;
+							ForwardLightData->bAffectsTranslucentLighting = LightParameters.bAffectsTranslucentLighting;
 
 							const FVector2D FadeParams = LightProxy->GetDirectionalLightDistanceFadeParameters(View.GetFeatureLevel(), LightSceneInfo->IsPrecomputedLightingValid(), View.MaxShadowCascades);
 
diff --git a/Engine/Source/Runtime/Renderer/Private/LightRendering.cpp b/Engine/Source/Runtime/Renderer/Private/LightRendering.cpp
--- a/Engine/Source/Runtime/Renderer/Private/LightRendering.cpp
+++ b/Engine/Source/Runtime/Renderer/Private/LightRendering.cpp
@@ -671,6 +671,7 @@ FDeferredLightUniformStruct GetSimpleDeferredLightParameters(
 	Out.LightParameters.RectLightAtlasMaxLevel = FLightRenderParameters::GetRectLightAtlasInvalidMIPLevel();
 	Out.LightParameters.IESAtlasIndex = INDEX_NONE;
 	Out.LightParameters.LightFunctionAtlasLightIndex = 0;
+	Out.LightParameters.bAffectsTranslucentLighting = 0;
 	return Out;
 }
 FDeferredLightUniformStruct GetSimpleDeferredLightParameters(
diff --git a/Engine/Source/Runtime/Renderer/Private/LightSceneInfo.cpp b/Engine/Source/Runtime/Renderer/Private/LightSceneInfo.cpp
--- a/Engine/Source/Runtime/Renderer/Private/LightSceneInfo.cpp
+++ b/Engine/Source/Runtime/Renderer/Private/LightSceneInfo.cpp
@@ -390,7 +390,10 @@ uint32 FLightSceneInfo::PackLightTypeAndShadowMapChannelMask(bool bAllowStaticLi
 
 	Result |= Proxy->LightFunctionAtlasLightIndex << (BitOffset);		BitOffset += 8;
 
-	// 28 bits used, 4 bits free
+	uint32 AffectsTranslucentLighting = Proxy->AffectsTranslucentLighting() ? 1 : 0;
+	Result |= AffectsTranslucentLighting << (BitOffset);				BitOffset += 1;
+
+	// 29 bits used, 3 bits free
 
 	return Result;
 }
diff --git a/Engine/Source/Runtime/Renderer/Private/SceneRendering.h b/Engine/Source/Runtime/Renderer/Private/SceneRendering.h
--- a/Engine/Source/Runtime/Renderer/Private/SceneRendering.h
+++ b/Engine/Source/Runtime/Renderer/Private/SceneRendering.h
@@ -761,6 +761,7 @@ BEGIN_GLOBAL_SHADER_PARAMETER_STRUCT_WITH_CONSTRUCTOR(FForwardLightData, )
 	SHADER_PARAMETER(FMatrix44f, DirectionalLightTranslatedWorldToStaticShadow)
 	SHADER_PARAMETER(uint32, DirectLightingShowFlag)
 	SHADER_PARAMETER(uint32, LightFunctionAtlasLightIndex)
+	SHADER_PARAMETER(uint32, bAffectsTranslucentLighting)
 	SHADER_PARAMETER_STRUCT(FVirtualShadowMapSMRTSettings, DirectionalLightSMRTSettings)
 	SHADER_PARAMETER_RDG_TEXTURE(Texture2D, DirectionalLightShadowmapAtlas)
 	SHADER_PARAMETER_SAMPLER(SamplerState, ShadowmapSampler)
