Fixes HasPixelAnimation with the TSR resurrection with memory reload

diff --git a/Engine/Shaders/Private/TemporalSuperResolution/TSRRejectShading.usf b/Engine/Shaders/Private/TemporalSuperResolution/TSRRejectShading.usf
--- a/Engine/Shaders/Private/TemporalSuperResolution/TSRRejectShading.usf
+++ b/Engine/Shaders/Private/TemporalSuperResolution/TSRRejectShading.usf
@@ -2,7 +2,6 @@
 
 //------------------------------------------------------- INCLUDES
 
-#include "TSRConvolutionNetworkPass.ush"
 #include "TSRShadingAnalysis.ush"
 
 
@@ -375,7 +374,7 @@ void RefetchInputAndHistoryFromLDSOrCache(
 		/* out */ OriginalHistory,
 		/* out */ OriginalHistoryUncertainty);
 
-	tsr_tensor_halfC OriginalInput = ComposeTranslucency(OriginalOpaqueInput, Blur3x3(OriginalTranslucencyInput));
+	tsr_tensor_halfC OriginalInput = ComposeTranslucencyForRejection(OriginalOpaqueInput, OriginalTranslucencyInput, bHasPixelAnimation);
 
 	Input            = OriginalInput;
 	InputUncertainty = VelocityEdge;
@@ -532,25 +531,8 @@ void MainCS(
 		}
 	}
 
-	tsr_tensor_halfC OriginalInput;
-	{
-		tsr_tensor_halfC SharpInput = OriginalOpaqueInput;
-		tsr_tensor_half4 BlurInput = OriginalTranslucencyInput;
-
-		//BRANCH
-		//if (AnyComponent(bHasPixelAnimation))
-		{
-			tsr_tensor_halfC CenterFinalSceneColor = ComposeTranslucency(OriginalOpaqueInput, OriginalTranslucencyInput);
-			
-			tsr_tensor_half4 HasAnimationBlurInput = ResizeChannels<4>(CenterFinalSceneColor);
-			
-			BlurInput  = select(tsr_tensor_bool4::Vectorize(bHasPixelAnimation), HasAnimationBlurInput, BlurInput);
-			SharpInput = select(tsr_tensor_boolC::Vectorize(bHasPixelAnimation), tsr_tensor_halfC::Const(0.0), SharpInput);
-		}
-		
-		OriginalInput = ComposeTranslucency(SharpInput, Blur3x3(BlurInput));
-	}
-
+	tsr_tensor_halfC OriginalInput = ComposeTranslucencyForRejection(OriginalOpaqueInput, OriginalTranslucencyInput, bHasPixelAnimation);
+	
 	// Despill to LDS to make room in VGPRs for the convolutions of ComputeMoireError()
 	#if CONFIG_LDS_DESPILL
 	PLATFORM_SPECIFIC_ISOLATE
@@ -646,7 +628,7 @@ void MainCS(
 				tsr_tensor_half4 OriginalTranslucencyInput;
 				FetchSceneColorAndTranslucency(PixelPos, /* out */ OriginalOpaqueInput, /* out */ OriginalTranslucencyInput);
 
-				Input = ComposeTranslucency(OriginalOpaqueInput, Blur3x3(OriginalTranslucencyInput));
+				Input = ComposeTranslucencyForRejection(OriginalOpaqueInput, OriginalTranslucencyInput, bHasPixelAnimation);
 			#endif
 		}
 		#else
diff --git a/Engine/Shaders/Private/TemporalSuperResolution/TSRShadingAnalysis.ush b/Engine/Shaders/Private/TemporalSuperResolution/TSRShadingAnalysis.ush
--- a/Engine/Shaders/Private/TemporalSuperResolution/TSRShadingAnalysis.ush
+++ b/Engine/Shaders/Private/TemporalSuperResolution/TSRShadingAnalysis.ush
@@ -2,7 +2,10 @@
 
 #pragma once
 
-#include "TSRConvolutionNetwork.ush"
+
+//------------------------------------------------------- INCLUDES
+
+#include "TSRConvolutionNetworkPass.ush"
 
 
 //------------------------------------------------------- CONSTANTS
@@ -37,6 +40,24 @@ tsr_tensor_halfC ComposeTranslucency(tsr_tensor_halfC Color, tsr_tensor_half4 Tr
 	return ComposedColor;
 }
 
+/** Compose blury translucency for tighther rejection in linear color space. */
+tsr_tensor_halfC ComposeTranslucencyForRejection(tsr_tensor_halfC OriginalOpaqueInput, tsr_tensor_half4 OriginalTranslucencyInput, tsr_tensor_bool bHasPixelAnimation)
+{
+	tsr_tensor_halfC SharpInput = OriginalOpaqueInput;
+	tsr_tensor_half4 BlurInput = OriginalTranslucencyInput;
+
+	{
+		tsr_tensor_halfC CenterFinalSceneColor = ComposeTranslucency(OriginalOpaqueInput, OriginalTranslucencyInput);
+			
+		tsr_tensor_half4 HasAnimationBlurInput = ResizeChannels<4>(CenterFinalSceneColor);
+			
+		BlurInput  = select(tsr_tensor_bool4::Vectorize(bHasPixelAnimation), HasAnimationBlurInput, BlurInput);
+		SharpInput = select(tsr_tensor_boolC::Vectorize(bHasPixelAnimation), tsr_tensor_halfC::Const(0.0), SharpInput);
+	}
+		
+	return ComposeTranslucency(SharpInput, Blur3x3(BlurInput));
+}
+
 
 //------------------------------------------------------- CHANGE LDR EXPOSURE
 
