Enable separate translucency depth upsampling on SM5 platforms instead of only d3d

diff --git a/Engine/Shaders/Private/SeparateTranslucency.ush b/Engine/Shaders/Private/SeparateTranslucency.ush
--- a/Engine/Shaders/Private/SeparateTranslucency.ush
+++ b/Engine/Shaders/Private/SeparateTranslucency.ush
@@ -32,7 +32,7 @@ NearestDepthNeighborUpsamplingResult NearestDepthNeighborUpsampling(
 	NearestDepthNeighborUpsamplingResult UpsampleResult;
 
 //@todo - support for all platforms, just skip the GatherRed optimization where not supported
-#if ((SM6_PROFILE || SM5_PROFILE) && !(COMPILER_PSSL))
+#if FEATURE_LEVEL >= FEATURE_LEVEL_SM5
 
 	// The relative depth comparison breaks down at larger distances and incorrectly causes point sampling on the skybox pixels
 	float MaxOperationDepth = 2000000.0f;
