Add hair strands position reading from RWByteAddressBuffer.

diff --git a/Engine/Shaders/Private/HairStrands/HairStrandsPack.ush b/Engine/Shaders/Private/HairStrands/HairStrandsPack.ush
--- a/Engine/Shaders/Private/HairStrands/HairStrandsPack.ush
+++ b/Engine/Shaders/Private/HairStrands/HairStrandsPack.ush
@@ -88,6 +88,22 @@ FHairControlPoint ReadHairControlPoint(
 		InVF_TipScale);
 }
 
+FHairControlPoint ReadHairControlPoint(
+	RWByteAddressBuffer InBuffer,
+	uint InIndex,
+	float3 InVF_PositionOffset=0,
+	float InVF_Radius=1,
+	float InVF_RootScale=1,
+	float InVF_TipScale=1)
+{
+	return UnpackHairControlPoint(
+		InBuffer.Load2(InIndex * FPackedHairPositionStrideInBytes),
+		InVF_PositionOffset,
+		InVF_Radius,
+		InVF_RootScale,
+		InVF_TipScale);
+}
+
 void WriteHairControlPointPosition(
 	RWByteAddressBuffer OutBuffer, 
 	uint InIndex, 
