Missing shader file change from "Applying HeightFog and VolumetricFog..."

diff --git a/Engine/Shaders/Private/HeterogeneousVolumes/HeterogeneousVolumesLiveShadingPipeline.usf b/Engine/Shaders/Private/HeterogeneousVolumes/HeterogeneousVolumesLiveShadingPipeline.usf
--- a/Engine/Shaders/Private/HeterogeneousVolumes/HeterogeneousVolumesLiveShadingPipeline.usf
+++ b/Engine/Shaders/Private/HeterogeneousVolumes/HeterogeneousVolumesLiveShadingPipeline.usf
@@ -55,6 +55,10 @@ uint VirtualShadowMapId;
 float ShadowStepSize;
 float ShadowStepFactor;
 
+// Atmosphere
+int bApplyHeightFog;
+int bApplyVolumetricFog;
+
 // Ray data
 float MaxTraceDistance;
 float MaxShadowTraceDistance;
@@ -163,6 +167,16 @@ float GetMaxShadowTraceDistance()
 	return MaxShadowTraceDistance;
 }
 
+int ShouldApplyHeightFog()
+{
+	return bApplyHeightFog;
+}
+
+int ShouldApplyVolumetricFog()
+{
+	return bApplyVolumetricFog;
+}
+
 struct FVolumeSampleContext
 {
 	FMaterialPixelParameters MaterialParameters;
