Many Lights - fixed InverseExposureBlend
* Load all light data from a local grid in order to fix InverseExposureBlend. With a pointer to prev local light there's no need to rely on global indices
* Enabled screen space traces by default. They had some artifacts, but in most cases improve shadow quality by fixing BVH/SDF representation mismatches
* Renamed composite to resolve as it better describes that pass

diff --git a/Engine/Shaders/Private/ManyLights/ManyLightsComposite.usf b/Engine/Shaders/Private/ManyLights/ManyLightsComposite.usf
deleted file mode 100644
diff --git a/Engine/Shaders/Private/ManyLights/ManyLightsHardwareRayTracing.usf b/Engine/Shaders/Private/ManyLights/ManyLightsHardwareRayTracing.usf
--- a/Engine/Shaders/Private/ManyLights/ManyLightsHardwareRayTracing.usf
+++ b/Engine/Shaders/Private/ManyLights/ManyLightsHardwareRayTracing.usf
@@ -134,8 +134,7 @@ RAY_TRACING_ENTRY_RAYGEN_OR_INLINE(HardwareRayTraceLightSamples)
 		if (bDebug && DebugMode == DEBUG_MODE_VISUALIZE_TRACING)
 		{
 			const FLocalLightData LocalLightData = GetLocalLightDataNonStereo(LightSample.LocalLightIndex);
-			const FLightSceneData LightSceneData = GetLightSceneData(LocalLightData.LightSceneId);
-			const FDeferredLightData LightData = ConvertToDeferredLight(LightSceneData);
+			const FDeferredLightData LightData = ConvertToDeferredLight(LocalLightData);
 
 			float3 RayStart = Ray.Origin + Ray.Direction * Ray.TMin;
 			float3 RayEnd = Ray.Origin + Ray.Direction * (Payload.IsHit() ? Payload.HitT : Ray.TMax);
diff --git a/Engine/Shaders/Private/ManyLights/ManyLightsRayTracing.usf b/Engine/Shaders/Private/ManyLights/ManyLightsRayTracing.usf
--- a/Engine/Shaders/Private/ManyLights/ManyLightsRayTracing.usf
+++ b/Engine/Shaders/Private/ManyLights/ManyLightsRayTracing.usf
@@ -286,8 +286,7 @@ void ScreenSpaceRayTraceLightSamplesCS(
 		if (bDebug && DebugMode == DEBUG_MODE_VISUALIZE_TRACING)
 		{
 			const FLocalLightData LocalLightData = GetLocalLightDataNonStereo(LightSample.LocalLightIndex);
-			const FLightSceneData LightSceneData = GetLightSceneData(LocalLightData.LightSceneId);
-			const FDeferredLightData LightData = ConvertToDeferredLight(LightSceneData);
+			const FDeferredLightData LightData = ConvertToDeferredLight(LocalLightData);
 
 			float3 RayStart = TranslatedWorldPosition;
 			float3 RayEnd = TranslatedWorldPosition + LightSampleTrace.Direction * min(LightSampleTrace.Distance, RayT);
@@ -360,8 +359,7 @@ void SoftwareRayTraceLightSamplesCS(
 		if (bDebug && DebugMode == DEBUG_MODE_VISUALIZE_TRACING)
 		{
 			const FLocalLightData LocalLightData = GetLocalLightDataNonStereo(LightSample.LocalLightIndex);
-			const FLightSceneData LightSceneData = GetLightSceneData(LocalLightData.LightSceneId);
-			const FDeferredLightData LightData = ConvertToDeferredLight(LightSceneData);
+			const FDeferredLightData LightData = ConvertToDeferredLight(LocalLightData);
 
 			float3 RayStart = TraceInput.TranslatedWorldRayStart + TraceInput.WorldRayDirection * TraceInput.MinTraceDistance;
 			float3 RayEnd = TraceInput.TranslatedWorldRayStart + TraceInput.WorldRayDirection * (bHit ? SDFResult.HitTime : TraceInput.MaxTraceDistance);
diff --git a/Engine/Shaders/Private/ManyLights/ManyLightsRayTracing.ush b/Engine/Shaders/Private/ManyLights/ManyLightsRayTracing.ush
--- a/Engine/Shaders/Private/ManyLights/ManyLightsRayTracing.ush
+++ b/Engine/Shaders/Private/ManyLights/ManyLightsRayTracing.ush
@@ -92,9 +92,8 @@ FLightSampleTrace GetLightSampleTrace(float3 TranslatedWorldPosition, uint Local
 	if (LocalLightIndex != MAX_LOCAL_LIGHT_INDEX)
 	{
 		const FLocalLightData LocalLightData = GetLocalLightDataNonStereo(LocalLightIndex);
-		const FLightSceneData LightSceneData = GetLightSceneData(LocalLightData.LightSceneId);
-		const FDeferredLightData LightData = ConvertToDeferredLight(LightSceneData);
-		const FLightShaderParameters LightParameters = ConvertToLightShaderParameters(LightSceneData);
+		const FDeferredLightData LightData = ConvertToDeferredLight(LocalLightData);
+		const FLightShaderParameters LightParameters = ConvertToLightShaderParameters(LightData);
 
 		float3 Unused;
 		float Unused2;
diff --git a/Engine/Shaders/Private/ManyLights/ManyLightsResolve.usf b/Engine/Shaders/Private/ManyLights/ManyLightsResolve.usf
new file mode 100644
--- /dev/null
+++ b/Engine/Shaders/Private/ManyLights/ManyLightsResolve.usf
@@ -0,0 +1,189 @@
+// Copyright Epic Games, Inc. All Rights Reserved.
+
+#include "ManyLights.ush"
+
+RWStructuredBuffer<uint> RWShadingTiles;
+RWStructuredBuffer<uint> RWShadingTileAllocator;
+RWTexture2D<uint> RWShadingTileGridAllocator;
+RWStructuredBuffer<uint> RWShadingTileGrid;
+RWTexture2D<UNORM float> RWShadingTileAtlas;
+
+StructuredBuffer<uint2> CompositeTileData;
+Texture2D<uint> LightSamples;
+Texture2D<float> DownsampledSceneDepth;
+Texture2D<UNORM float3> DownsampledSceneWorldNormal;
+Texture2D<float4> CompositeUpsampleWeights;
+
+uint2 NumSamplesPerPixel;
+int2 DownsampledViewSize;
+uint MaxShadingTiles;
+uint DebugLightId;
+
+#define SHARED_TILE_BORDER 1
+#define SHARED_TILE_SIZE (8 + 2 * SHARED_TILE_BORDER)
+groupshared float SharedSampleWeights[SHARED_TILE_SIZE][SHARED_TILE_SIZE];
+
+// 2x2 tile mask. 1 means that we need to spawn a shading tile for that region.
+groupshared uint SharedShadingTileMask[2][2];
+groupshared uint SharedShadingTileIndex[2][2];
+
+void AccumulateLightSample(uint2 SharedTileCoord, float UpsampleWeight, inout float SampleWeightSum)
+{
+	float SampleWeight = SharedSampleWeights[SharedTileCoord.x][SharedTileCoord.y];
+	SampleWeightSum += SampleWeight * UpsampleWeight;
+}
+
+/**
+ * Run one thread per local region and compute fullscreen per pixel BRDF weights for shading.
+ */
+[numthreads(THREADGROUP_SIZE, THREADGROUP_SIZE, 1)]
+void ResolveLightSamplesCS(
+	uint3 GroupId : SV_GroupID,
+	uint3 GroupThreadId : SV_GroupThreadID,
+	uint3 DispatchThreadId : SV_DispatchThreadID)
+{
+	uint CompositeTileIndex = GroupId.x;
+
+	uint LocalLightIndex;
+	uint2 BaseDownsampledScreenCoord;
+	UnpackCompositeTile(CompositeTileData[CompositeTileIndex], LocalLightIndex, BaseDownsampledScreenCoord);
+
+	if (all(GroupThreadId < 2))
+	{
+		SharedShadingTileMask[GroupThreadId.x][GroupThreadId.y] = 0;
+		SharedShadingTileIndex[GroupThreadId.x][GroupThreadId.y] = 0xFFFFFFFF;
+	}
+
+	GroupMemoryBarrierWithGroupSync();
+
+	const FLocalLightData LocalLightData = GetLocalLightDataNonStereo(LocalLightIndex);
+	
+	// Composite traces into SharedSampleWeights
+	for (uint SharedCoordY = GroupThreadId.y; SharedCoordY < SHARED_TILE_SIZE; SharedCoordY += THREADGROUP_SIZE)
+	{
+		for (uint SharedCoordX = GroupThreadId.x; SharedCoordX < SHARED_TILE_SIZE; SharedCoordX += THREADGROUP_SIZE)
+		{
+			const uint2 SharedCoord = uint2(SharedCoordX, SharedCoordY);
+			const uint2 DownsampledScreenCoord = clamp(int2(BaseDownsampledScreenCoord + SharedCoord) - SHARED_TILE_BORDER, 0, DownsampledViewSize - 1);
+			
+			float WeightSum = 0.0f;
+
+			for (uint LocalSampleY = 0; LocalSampleY < NumSamplesPerPixel.y; ++LocalSampleY)
+			{
+				for (uint LocalSampleX = 0; LocalSampleX < NumSamplesPerPixel.x; ++LocalSampleX)
+				{
+					uint2 SampleCoord = DownsampledScreenCoord * NumSamplesPerPixel + uint2(LocalSampleX, LocalSampleY);
+					FLightSample LightSample = UnpackLightSample(LightSamples[SampleCoord]);
+
+					if (LightSample.LocalLightIndex == LocalLightIndex && LightSample.bVisible)
+					{
+						WeightSum += LightSample.Weight;
+					}
+				}
+			}
+
+			// Mark sub-tiles requiring shading tiles, accounting for a bilinear interpolation footprint
+			if (WeightSum > 0.001f)
+			{
+				uint LitBitMask = 0;
+
+				for (uint ShadingTileOffsetY = 0; ShadingTileOffsetY < 2; ++ShadingTileOffsetY)
+				{
+					for (uint ShadingTileOffsetX = 0; ShadingTileOffsetX < 2; ++ShadingTileOffsetX)
+					{
+						const uint2 ShadingTileOffset = uint2(ShadingTileOffsetX, ShadingTileOffsetY);
+
+						// Bilinear footprint
+						uint2 ShadingTileMin = ShadingTileOffset * 4 + (SHARED_TILE_BORDER - 1);
+						uint2 ShadingTileMax = ShadingTileMin + 5;
+
+						if (and(all(SharedCoord >= ShadingTileMin), all(SharedCoord <= ShadingTileMax)))
+						{
+							// #ml_todo: don't write to the spot from so many threads
+							SharedShadingTileMask[ShadingTileOffset.x][ShadingTileOffset.y] = 1;
+						}
+					}
+				}
+			}
+
+			SharedSampleWeights[SharedCoordX][SharedCoordY] = WeightSum;
+		}
+	}
+
+	GroupMemoryBarrierWithGroupSync();
+
+	uint2 LocalDownsampledTileCoord = GroupThreadId.xy;
+	uint2 DownsampledScreenCoord = BaseDownsampledScreenCoord + LocalDownsampledTileCoord;
+	uint2 ShadingTileOffset = GroupThreadId.xy;
+
+	// Allocate shading tiles
+	if (all(ShadingTileOffset < 2))
+	{
+		if (SharedShadingTileMask[ShadingTileOffset.x][ShadingTileOffset.y] != 0)
+		{	
+			const uint2 ShadingTileCoord = (BaseDownsampledScreenCoord * DOWNSAMPLE_FACTOR) / TILE_SIZE + ShadingTileOffset;
+
+			if (all(ShadingTileCoord < ShadingTileGridSize))
+			{
+				uint ShadingTileIndex = MaxShadingTiles;
+				InterlockedAdd(RWShadingTileAllocator[0], 1u, ShadingTileIndex);
+
+				if (ShadingTileIndex < MaxShadingTiles)
+				{
+					uint ShadingTileIndexInCell = 0;
+					InterlockedAdd(RWShadingTileGridAllocator[ShadingTileCoord], 1u, ShadingTileIndexInCell);
+
+					if (ShadingTileIndexInCell < MaxShadingTilesPerGridCell)
+					{
+						RWShadingTileGrid[GetShadingTileGridOffset(ShadingTileCoord, ShadingTileIndexInCell)] = PackShadingTile(LocalLightIndex, ShadingTileIndex);
+						SharedShadingTileIndex[ShadingTileOffset.x][ShadingTileOffset.y] = ShadingTileIndex;
+					}
+				}
+			}
+		}
+	}
+
+	GroupMemoryBarrierWithGroupSync();
+
+	// Composite shading tiles
+	for (uint ShadingTileOffsetY = 0; ShadingTileOffsetY < 2; ++ShadingTileOffsetY)
+	{
+		for (uint ShadingTileOffsetX = 0; ShadingTileOffsetX < 2; ++ShadingTileOffsetX)
+		{
+			uint ShadingTileIndex = SharedShadingTileIndex[ShadingTileOffsetX][ShadingTileOffsetY];
+
+			if (ShadingTileIndex < MaxShadingTiles)
+			{
+				const uint2 CoordInShadingTile = GroupThreadId.xy;
+				const uint2 ShadingTileCoord = (BaseDownsampledScreenCoord * DOWNSAMPLE_FACTOR) / TILE_SIZE + uint2(ShadingTileOffsetX, ShadingTileOffsetY);
+				uint2 ScreenCoord = ShadingTileCoord * TILE_SIZE + CoordInShadingTile;
+				float4 DownsampledScreenWeights = CompositeUpsampleWeights[ScreenCoord];
+
+				const uint2 SharedTileCoord00 = (uint2(ShadingTileOffsetX, ShadingTileOffsetY) * TILE_SIZE + CoordInShadingTile + 1) / DOWNSAMPLE_FACTOR + (SHARED_TILE_BORDER - 1);
+
+				float SampleWeightSum = 0.0f;
+				AccumulateLightSample(SharedTileCoord00 + uint2(0, 0), DownsampledScreenWeights.x, SampleWeightSum);
+				AccumulateLightSample(SharedTileCoord00 + uint2(1, 0), DownsampledScreenWeights.y, SampleWeightSum);
+				AccumulateLightSample(SharedTileCoord00 + uint2(0, 1), DownsampledScreenWeights.z, SampleWeightSum);
+				AccumulateLightSample(SharedTileCoord00 + uint2(1, 1), DownsampledScreenWeights.w, SampleWeightSum);
+
+				#if DEBUG_MODE
+					int2 DebugScreenCoord = View.CursorPosition.x >= 0 ? View.CursorPosition * View.ViewResolutionFraction : -1;
+					bool bDebug = all(ScreenCoord == DebugScreenCoord) && LocalLightData.LightSceneId == DebugLightId;
+					FShaderPrintContext Context = InitShaderPrintContext(true, float2(0.05, 0.3));
+
+					if (bDebug)
+					{
+						Print(Context, TEXT("CompositeLightSamples"));
+						Newline(Context);
+						Print(Context, TEXT("SampleWeight: "));
+						Print(Context, SampleWeightSum);
+					}
+				#endif
+
+				uint2 ShadingTileAtlasCoord = GetShadingTileAtlasCoord(ShadingTileIndex, CoordInShadingTile);
+				RWShadingTileAtlas[ShadingTileAtlasCoord] = SampleWeightSum;
+			}
+		}
+	}
+}
diff --git a/Engine/Shaders/Private/ManyLights/ManyLightsSampling.usf b/Engine/Shaders/Private/ManyLights/ManyLightsSampling.usf
--- a/Engine/Shaders/Private/ManyLights/ManyLightsSampling.usf
+++ b/Engine/Shaders/Private/ManyLights/ManyLightsSampling.usf
@@ -279,8 +279,7 @@ void GenerateLightSamplesCS(
 
 						const uint LocalLightIndex = MaskIndex * 32 + NextBitIndex;
 						const FLocalLightData LocalLightData = GetLocalLightDataNonStereo(LocalLightIndex);
-						const FLightSceneData LightSceneData = GetLightSceneData(LocalLightData.LightSceneId);
-						FDeferredLightData LightData = ConvertToDeferredLight(LightSceneData);
+						FDeferredLightData LightData = ConvertToDeferredLight(LocalLightData);
 					
 						FLightTargetPDF LightTargetPDF = GetLocalLightTargetPDF(LightData, TranslatedWorldPosition, GBuffer, ScreenCoord);
 
@@ -346,10 +345,9 @@ void GenerateLightSamplesCS(
 							const uint2 SampleCoord = DownsampledScreenCoord * DOWNSAMPLE_FACTOR + LightSampleOffsets[LightSampleIndex];
 
 							const FLocalLightData LocalLightData = GetLocalLightDataNonStereo(LightSample.LocalLightIndex);
-							const FLightSceneData LightSceneData = GetLightSceneData(LocalLightData.LightSceneId);
 							const FLightSampleTrace LightSampleTrace = GetLightSampleTrace(TranslatedWorldPosition, LightSample.LocalLightIndex, SampleCoord);
 
-							AddLineTWS(Context, TranslatedWorldPosition, TranslatedWorldPosition + LightSampleTrace.Direction * LightSampleTrace.Distance, float4(LightSceneData.Color.xyz / Luminance(LightSceneData.Color.xyz), 1.0f));
+							AddLineTWS(Context, TranslatedWorldPosition, TranslatedWorldPosition + LightSampleTrace.Direction * LightSampleTrace.Distance, float4(LocalLightData.LightColorAndIdAndFalloffExponent.xyz / Luminance(LocalLightData.LightColorAndIdAndFalloffExponent.xyz), 1.0f));
 						}
 					}
 
@@ -374,8 +372,7 @@ void GenerateLightSamplesCS(
 				if (LightSample.LocalLightIndex != MAX_LOCAL_LIGHT_INDEX)
 				{
 					const FLocalLightData LocalLightData = GetLocalLightDataNonStereo(LightSample.LocalLightIndex);
-					const FLightSceneData LightSceneData = GetLightSceneData(LocalLightData.LightSceneId);
-					const bool bCastShadows = UnpackCastShadow(LightSceneData.LightTypeAndShadowMapChannelMaskPacked);
+					const bool bCastShadows = UnpackCastShadow(asuint(LocalLightData.LightDirectionAndShadowMask.w));
 
 					LightSample.bVisible = true;
 					LightSample.bCompleted = bCastShadows ? false : true;
diff --git a/Engine/Shaders/Private/ManyLights/ManyLightsShading.usf b/Engine/Shaders/Private/ManyLights/ManyLightsShading.usf
--- a/Engine/Shaders/Private/ManyLights/ManyLightsShading.usf
+++ b/Engine/Shaders/Private/ManyLights/ManyLightsShading.usf
@@ -87,8 +87,7 @@ void ShadeLightSamplesCS(
 					const uint PackedShadingTile = ShadingTileGrid[GetShadingTileGridOffset(ShadingTileCoord, IndexInCell)];
 					const FShadingTile ShadingTile = UnpackShadingTile(PackedShadingTile);
 					const FLocalLightData LocalLightData = GetLocalLightDataNonStereo(ShadingTile.LocalLightIndex);
-					const FLightSceneData LightSceneData = GetLightSceneData(LocalLightData.LightSceneId);
-					FDeferredLightData LightData = ConvertToDeferredLight(LightSceneData);
+					FDeferredLightData LightData = ConvertToDeferredLight(LocalLightData);
 
 					float SampleWeight = 0.0f;
 
diff --git a/Engine/Source/Runtime/Renderer/Private/ManyLights/ManyLights.cpp b/Engine/Source/Runtime/Renderer/Private/ManyLights/ManyLights.cpp
--- a/Engine/Source/Runtime/Renderer/Private/ManyLights/ManyLights.cpp
+++ b/Engine/Source/Runtime/Renderer/Private/ManyLights/ManyLights.cpp
@@ -442,10 +442,10 @@ class FInitCompositeUpsampleWeightsCS : public FGlobalShader
 
 IMPLEMENT_GLOBAL_SHADER(FInitCompositeUpsampleWeightsCS, "/Engine/Private/ManyLights/ManyLights.usf", "InitCompositeUpsampleWeightsCS", SF_Compute);
 
-class FCompositeLightSamplesCS : public FGlobalShader
+class FResolveLightSamplesCS : public FGlobalShader
 {
-	DECLARE_GLOBAL_SHADER(FCompositeLightSamplesCS)
-	SHADER_USE_PARAMETER_STRUCT(FCompositeLightSamplesCS, FGlobalShader)
+	DECLARE_GLOBAL_SHADER(FResolveLightSamplesCS)
+	SHADER_USE_PARAMETER_STRUCT(FResolveLightSamplesCS, FGlobalShader)
 
 	BEGIN_SHADER_PARAMETER_STRUCT(FParameters, )
 		SHADER_PARAMETER_STRUCT_INCLUDE(FManyLightsParameters, ManyLightsParameters)
@@ -480,7 +480,7 @@ class FCompositeLightSamplesCS : public FGlobalShader
 	}
 };
 
-IMPLEMENT_GLOBAL_SHADER(FCompositeLightSamplesCS, "/Engine/Private/ManyLights/ManyLightsComposite.usf", "CompositeLightSamplesCS", SF_Compute);
+IMPLEMENT_GLOBAL_SHADER(FResolveLightSamplesCS, "/Engine/Private/ManyLights/ManyLightsResolve.usf", "ResolveLightSamplesCS", SF_Compute);
 
 class FShadeLightSamplesCS : public FGlobalShader
 {
@@ -1016,7 +1016,7 @@ void FDeferredShadingSceneRenderer::RenderManyLights(FRDGBuilder& GraphBuilder,
 
 	// Composite shadow masks traces
 	{
-		FCompositeLightSamplesCS::FParameters* PassParameters = GraphBuilder.AllocParameters<FCompositeLightSamplesCS::FParameters>();
+		FResolveLightSamplesCS::FParameters* PassParameters = GraphBuilder.AllocParameters<FResolveLightSamplesCS::FParameters>();
 		PassParameters->ManyLightsParameters = ManyLightsParameters;
 		PassParameters->IndirectArgs = CompositeTileIndirectArgs;
 		PassParameters->RWShadingTileAllocator = GraphBuilder.CreateUAV(ShadingTileAllocator);
@@ -1027,9 +1027,9 @@ void FDeferredShadingSceneRenderer::RenderManyLights(FRDGBuilder& GraphBuilder,
 		PassParameters->CompositeUpsampleWeights = CompositeUpsampleWeights;
 		PassParameters->LightSamples = LightSamples;
 
-		FCompositeLightSamplesCS::FPermutationDomain PermutationVector;
-		PermutationVector.Set<FCompositeLightSamplesCS::FDebugMode>(bDebug);
-		auto ComputeShader = View.ShaderMap->GetShader<FCompositeLightSamplesCS>(PermutationVector);
+		FResolveLightSamplesCS::FPermutationDomain PermutationVector;
+		PermutationVector.Set<FResolveLightSamplesCS::FDebugMode>(bDebug);
+		auto ComputeShader = View.ShaderMap->GetShader<FResolveLightSamplesCS>(PermutationVector);
 
 		FComputeShaderUtils::AddPass(
 			GraphBuilder,
diff --git a/Engine/Source/Runtime/Renderer/Private/ManyLights/ManyLightsRayTracing.cpp b/Engine/Source/Runtime/Renderer/Private/ManyLights/ManyLightsRayTracing.cpp
--- a/Engine/Source/Runtime/Renderer/Private/ManyLights/ManyLightsRayTracing.cpp
+++ b/Engine/Source/Runtime/Renderer/Private/ManyLights/ManyLightsRayTracing.cpp
@@ -8,7 +8,7 @@
 
 static TAutoConsoleVariable<int32> CVarManyLightsScreenTraces(
 	TEXT("r.ManyLights.ScreenTraces"),
-	0,
+	1,
 	TEXT("Whether to use screen space tracing for shadow rays."),
 	ECVF_Scalability | ECVF_RenderThreadSafe
 );
