fix for HLSL 2021 not supporting global operator overloads. We can put it back if/when https://microsoft.github.io/hlsl-specs/proposals/0008-non-member-operator-overloading.html is approved and implemented

diff --git a/Engine/Shaders/Private/DualNumber.ush b/Engine/Shaders/Private/DualNumber.ush
--- a/Engine/Shaders/Private/DualNumber.ush
+++ b/Engine/Shaders/Private/DualNumber.ush
@@ -8,99 +8,69 @@ struct TDual
 	T Value;
 	T Value_dx;
 	T Value_dy;
-};
-
-template< typename T >
-TDual<T> operator+( TDual<T> A, T B )
-{
-	TDual<T> R;
-	R.Value		= A.Value + B;
-	R.Value_dx	= A.Value_dx;
-	R.Value_dy	= A.Value_dy;
-	return R;
-}
-
-template< typename T >
-TDual<T> operator+( T A, TDual<T> B )
-{
-	return B + A;
-}
-
-template< typename T >
-TDual<T> operator-( TDual<T> A, T B )
-{
-	TDual<T> R;
-	R.Value		= A.Value - B;
-	R.Value_dx	= A.Value_dx;
-	R.Value_dy	= A.Value_dy;
-	return R;
-}
-
-template< typename T >
-TDual<T> operator-( T A, TDual<T> B )
-{
-	TDual<T> R;
-	R.Value		= A - B.Value;
-	R.Value_dx	= B.Value_dx;
-	R.Value_dy	= B.Value_dy;
-	return R;
-}
+	TDual<T> operator+(T B)
+	{
+		TDual<T> R;
+		R.Value = Value + B;
+		R.Value_dx = Value_dx;
+		R.Value_dy = Value_dy;
+		return R;
+	}
 
-template< typename T >
-TDual<T> operator*( TDual<T> A, T B )
-{
-	TDual<T> R;
-	R.Value		= A.Value    * B;
-	R.Value_dx	= A.Value_dx * B;
-	R.Value_dy	= A.Value_dy * B;
-	return R;
-}
+	TDual<T> operator-(T B)
+	{
+		TDual<T> R;
+		R.Value = Value - B;
+		R.Value_dx = Value_dx;
+		R.Value_dy = Value_dy;
+		return R;
+	}
 
-template< typename T >
-TDual<T> operator*( T A, TDual<T> B )
-{
-	return B * A;
-}
+	TDual<T> operator*(T B)
+	{
+		TDual<T> R;
+		R.Value = Value * B;
+		R.Value_dx = Value_dx * B;
+		R.Value_dy = Value_dy * B;
+		return R;
+	}
 
-template< typename T >
-TDual<T> operator+( TDual<T> A, TDual<T> B )
-{
-	TDual<T> R;
-	R.Value		= A.Value    + B.Value;
-	R.Value_dx	= A.Value_dx + B.Value_dx;
-	R.Value_dy	= A.Value_dy + B.Value_dy;
-	return R;
-}
+	TDual<T> operator+(TDual<T> B)
+	{
+		TDual<T> R;
+		R.Value = Value + B.Value;
+		R.Value_dx = Value_dx + B.Value_dx;
+		R.Value_dy = Value_dy + B.Value_dy;
+		return R;
+	}
 
-template< typename T >
-TDual<T> operator-( TDual<T> A, TDual<T> B )
-{
-	TDual<T> R;
-	R.Value		= A.Value    - B.Value;
-	R.Value_dx	= A.Value_dx - B.Value_dx;
-	R.Value_dy	= A.Value_dy - B.Value_dy;
-	return R;
-}
+	TDual<T> operator-(TDual<T> B)
+	{
+		TDual<T> R;
+		R.Value = Value - B.Value;
+		R.Value_dx = Value_dx - B.Value_dx;
+		R.Value_dy = Value_dy - B.Value_dy;
+		return R;
+	}
 
-template< typename T >
-TDual<T> operator*( TDual<T> A, TDual<T> B )
-{
-	TDual<T> R;
-	R.Value		= A.Value * B.Value;
-	R.Value_dx	= A.Value_dx * B.Value + A.Value * B.Value_dx;
-	R.Value_dy	= A.Value_dy * B.Value + A.Value * B.Value_dy;
-	return R;
-}
+	TDual<T> operator*(TDual<T> B)
+	{
+		TDual<T> R;
+		R.Value = Value * B.Value;
+		R.Value_dx = Value_dx * B.Value + Value * B.Value_dx;
+		R.Value_dy = Value_dy * B.Value + Value * B.Value_dy;
+		return R;
+	}
 
-template< typename T >
-TDual<T> operator/( TDual<T> A, TDual<T> B )
-{
-	TDual<T> R;
-	R.Value = A.Value / B.Value;
-	T Denom = rcp( B.Value * B.Value );
-	T dFdA =  B.Value * Denom;
-	T dFdB = -A.Value * Denom;
-	R.Value_dx = dFdA * A.Value_dx + dFdB * B.Value_dx;
-	R.Value_dy = dFdA * A.Value_dy + dFdB * B.Value_dy;
-	return R;
-}
+	TDual<T> operator/(TDual<T> B)
+	{
+		TDual<T> R;
+		R.Value = Value / B.Value;
+		T Denom = rcp(B.Value * B.Value);
+		T dFdA = B.Value * Denom;
+		T dFdB = -Value * Denom;
+		R.Value_dx = dFdA * Value_dx + dFdB * B.Value_dx;
+		R.Value_dy = dFdA * Value_dy + dFdB * B.Value_dy;
+		return R;
+	}
+};
