Path Tracing: Fix blurry textures

The derivative calculations in the path tracing material code doesn't appear to be correct in all cases. We have seen some reports of blurry textures that go away when this is disabled. Since the derivatives are mainly used by Substrate glints for now, only enable them in this case.

Also do not hardcode SUBSTRATE_COMPLEXSPECIALPATH since it should only be defined for materials that need it.

#changelist validated
#virtualized

diff --git a/Engine/Shaders/Private/PathTracing/PathTracingMaterialHitShader.usf b/Engine/Shaders/Private/PathTracing/PathTracingMaterialHitShader.usf
--- a/Engine/Shaders/Private/PathTracing/PathTracingMaterialHitShader.usf
+++ b/Engine/Shaders/Private/PathTracing/PathTracingMaterialHitShader.usf
@@ -22,8 +22,9 @@
 #define SUBSTRATE_SSS_MATERIAL_OVERRIDE 0
 
 // Ensure specular profile and glints are supported
-#define SUBSTRATE_COMPLEXSPECIALPATH 1
-#define USE_ANALYTIC_DERIVATIVES 1
+// TODO: the derivatives don't appear to be correct in all cases, causing incorrectly blurry lookups - so disable them for now unless using substrate
+//  We don't have an easy way to detect if the material is using glints or not, so use the complex special path which includes glints for now
+#define USE_ANALYTIC_DERIVATIVES (SUBSTRATE_ENABLED && SUBSTRATE_COMPLEXSPECIALPATH && PLATFORM_ENABLES_SUBSTRATE_GLINTS)
 
 #include "/Engine/Private/Common.ush"
 #include "/Engine/Private/RayTracing/RayTracingCommon.ush"
@@ -603,18 +604,21 @@ RAY_TRACING_ENTRY_CLOSEST_HIT(PathTracingMaterialCHS,
 						Payload.PhaseG       = 0;
 					}
 					Payload.SpecularProfileId = 0;
+					Payload.GlintValue = 0;
+
+#if SUBSTRATE_COMPLEXSPECIALPATH
 					if (BSDF_GETHASSPECPROFILE(BSDF))
 					{
 						Payload.SpecularProfileId = SLAB_SPECPROFILEID(BSDF);
 					}
-					Payload.GlintValue = 0;
 					if (BSDF_GETHASGLINT(BSDF))
 					{
 						Payload.GlintValue = SLAB_GLINT_VALUE(BSDF);
-						Payload.GlintUV = SLAB_GLINT_UV(BSDF);
-						Payload.GlintUVdx = SLAB_GLINT_UVDDX(BSDF);
-						Payload.GlintUVdy = SLAB_GLINT_UVDDY(BSDF);
+						Payload.GlintUV    = SLAB_GLINT_UV(BSDF);
+						Payload.GlintUVdx  = SLAB_GLINT_UVDDX(BSDF);
+						Payload.GlintUVdy  = SLAB_GLINT_UVDDY(BSDF);
 					}
+#endif
 					Payload.WeightV             = BSDF.LuminanceWeightV;
 					Payload.TransmittanceN      = BSDF.TransmittanceAboveAlongN;
 					Payload.CoverageAboveAlongN = BSDF.CoverageAboveAlongN;
