Fixed DFToTileOffset returning occasional incorrect results on platforms where mad() is not a fused multiply-add operation.
Replaced unnecessary DFToTileOffset conversion.

#changelist validated
#virtualized

diff --git a/Engine/Shaders/Private/DistanceField/GlobalDistanceFieldShared.ush b/Engine/Shaders/Private/DistanceField/GlobalDistanceFieldShared.ush
--- a/Engine/Shaders/Private/DistanceField/GlobalDistanceFieldShared.ush
+++ b/Engine/Shaders/Private/DistanceField/GlobalDistanceFieldShared.ush
@@ -443,7 +443,7 @@ float CalculateDistanceFieldApproxAO(FDFVector3 WorldPosition, float3 WorldNorma
 
 float CalculateDistanceFieldApproxAO(FLWCVector3 WorldPosition, float3 WorldNormal, uint NumSteps, float StepDistance, float StepScale, float DistanceBias, float MaxDistance)
 {
-	const float3 TranslatedWorldPosition = LWCToFloat(LWCAdd(WorldPosition, DFToTileOffset(ResolvedView.PreViewTranslation)));
+	const float3 TranslatedWorldPosition = LWCToFloat(LWCAdd(WorldPosition, ResolvedView.TileOffset.PreViewTranslation));
 	return CalculateDistanceFieldApproxAO(TranslatedWorldPosition, WorldNormal, NumSteps, StepDistance, StepScale, DistanceBias, MaxDistance);
 }
 
diff --git a/Engine/Shaders/Private/LargeWorldCoordinates.ush b/Engine/Shaders/Private/LargeWorldCoordinates.ush
--- a/Engine/Shaders/Private/LargeWorldCoordinates.ush
+++ b/Engine/Shaders/Private/LargeWorldCoordinates.ush
@@ -519,8 +519,8 @@ float4x4 LWCHackToFloat(FLWCInverseMatrix v) { return (float4x4)0.0f; }
 		/* Used rounding, removes fractional data stored in Tile, so the resulting coordinate is 'normalized' */				\
 		/* T+C-C benchmarked slightly faster than round(), can also shift C right to retain some fractional bits */ 			\
 		const float C = 8388608; /* = 2^23 (23 mantisse bits for int part, 0 for fraction) */ 									\
-		FFloatType Tile = mad(In.High, 1 / UE_LWC_RENDER_TILE_SIZE, C) - C;														\
-		FFloatType Offset = mad(Tile, -UE_LWC_RENDER_TILE_SIZE, In.High) + In.Low;												\
+		FFloatType Tile = INVARIANT(FMA(In.High, 1 / UE_LWC_RENDER_TILE_SIZE, C) - C);											\
+		FFloatType Offset = INVARIANT(FMA(Tile, -UE_LWC_RENDER_TILE_SIZE, In.High) + In.Low);									\
 																																\
 		return LWCConstructor(Tile, Offset);																					\
 	}																															\
diff --git a/Engine/Shaders/Public/ShaderVersion.ush b/Engine/Shaders/Public/ShaderVersion.ush
--- a/Engine/Shaders/Public/ShaderVersion.ush
+++ b/Engine/Shaders/Public/ShaderVersion.ush
@@ -3,4 +3,4 @@
 // Each time the console command is executed it generates a new GUID. As this file is included
 // in Platform.ush (which should be included in any shader) it allows to invalidate the shader DDC.
 // 
-#pragma message("UESHADERMETADATA_VERSION 8C9DC429-DD7E-4E7C-838B-923E4DDD3DC5")
+#pragma message("UESHADERMETADATA_VERSION DC1DB8B4-F092-46B2-8B29-0A59A08B8D93")
