Fix shader compilation due to missing type definition for global uniform buffer.

#changelist validated
#virtualized

diff --git a/Engine/Shaders/Private/GeneratedUniformBufferTypes.ush b/Engine/Shaders/Private/GeneratedUniformBufferTypes.ush
--- a/Engine/Shaders/Private/GeneratedUniformBufferTypes.ush
+++ b/Engine/Shaders/Private/GeneratedUniformBufferTypes.ush
@@ -12,4 +12,5 @@
 #include "HairStrands/HairStrandsVoxelPageCommonStruct.ush"
 #include "HairStrands/HairStrandsDeepShadowCommonStruct.ush"
 
-#include "/Engine/Shared/LightSceneData.h"
\ No newline at end of file
+#include "/Engine/Shared/LightSceneData.h"
+#include "/Engine/Shared/RayTracingTypes.h"
\ No newline at end of file
