ManyLights - sharper and more stable shadow reconstruction
* Replaced uniform sampling with 4-rooks pattern
* Replaced bilinear from sample center with triangle weights from 4 samples

diff --git a/Engine/Shaders/Private/ManyLights/ManyLights.usf b/Engine/Shaders/Private/ManyLights/ManyLights.usf
--- a/Engine/Shaders/Private/ManyLights/ManyLights.usf
+++ b/Engine/Shaders/Private/ManyLights/ManyLights.usf
@@ -182,17 +182,23 @@ void InitCompositeUpsampleWeightsCS(
 		FGBufferData GBuffer = GetGBufferData(ScreenUV);
 		float3 TranslatedWorldPosition = GetTranslatedWorldPositionFromScreenUV(ScreenUV, GBuffer.Depth);		
 
-		uint2 DownsampledScreenCoord = floor(ScreenUV * View.BufferSizeAndInvSize.xy / DOWNSAMPLE_FACTOR - 0.5f);
-		float2 DownsampledBilinearWeights = frac(ScreenUV * View.BufferSizeAndInvSize.xy / DOWNSAMPLE_FACTOR - 0.5f);
-		float2 DownsampledGatherUV = (DownsampledScreenCoord + 1.0f) * DownsampledBufferInvSize;
+		uint2 DownsampledScreenCoord00 = floor(ScreenUV * View.BufferSizeAndInvSize.xy / DOWNSAMPLE_FACTOR - 0.5f);
+		float2 DownsampledGatherUV = (DownsampledScreenCoord00 + 1.0f) * DownsampledBufferInvSize;
 
 		float4 CornerDepths = DownsampledSceneDepth.GatherRed(GlobalPointClampedSampler, DownsampledGatherUV).wzxy;
 
-		float4 InterpolationWeights = float4(
-			(1 - DownsampledBilinearWeights.y) * (1 - DownsampledBilinearWeights.x),
-			(1 - DownsampledBilinearWeights.y) * DownsampledBilinearWeights.x,
-			DownsampledBilinearWeights.y * (1 - DownsampledBilinearWeights.x),
-			DownsampledBilinearWeights.y * DownsampledBilinearWeights.x);
+		int2 ScreenCoordOffset = ScreenCoord - DownsampledScreenCoord00 * 2;
+		int2 SampleOffset00 = GetSampleScreenCoordJitter(DownsampledScreenCoord00 + uint2(0, 0)) + uint2(0, 0) * 2 - ScreenCoordOffset;
+		int2 SampleOffset10 = GetSampleScreenCoordJitter(DownsampledScreenCoord00 + uint2(1, 0)) + uint2(1, 0) * 2 - ScreenCoordOffset;
+		int2 SampleOffset01 = GetSampleScreenCoordJitter(DownsampledScreenCoord00 + uint2(0, 1)) + uint2(0, 1) * 2 - ScreenCoordOffset;
+		int2 SampleOffset11 = GetSampleScreenCoordJitter(DownsampledScreenCoord00 + uint2(1, 1)) + uint2(1, 1) * 2 - ScreenCoordOffset;
+
+		// Triangle filter weights between pixel and 4 samples
+		float4 InterpolationWeights;
+		InterpolationWeights.x = (2.0f - abs(SampleOffset00.x)) * (2.0f - abs(SampleOffset00.y));
+		InterpolationWeights.y = (2.0f - abs(SampleOffset10.x)) * (2.0f - abs(SampleOffset10.y));
+		InterpolationWeights.z = (2.0f - abs(SampleOffset01.x)) * (2.0f - abs(SampleOffset01.y));
+		InterpolationWeights.w = (2.0f - abs(SampleOffset11.x)) * (2.0f - abs(SampleOffset11.y));
 
 		float4 DepthWeights = 0.0f;	
 
@@ -201,10 +207,10 @@ void InitCompositeUpsampleWeightsCS(
 		{
 			float4 ScenePlane = float4(GBuffer.WorldNormal, dot(TranslatedWorldPosition, GBuffer.WorldNormal));
 
-			float3 Position00 = GetTranslatedWorldPositionFromScreenUV(DownsampledScreenCoordToScreenUV(DownsampledScreenCoord + uint2(0, 0)), CornerDepths.x);
-			float3 Position10 = GetTranslatedWorldPositionFromScreenUV(DownsampledScreenCoordToScreenUV(DownsampledScreenCoord + uint2(1, 0)), CornerDepths.y);
-			float3 Position01 = GetTranslatedWorldPositionFromScreenUV(DownsampledScreenCoordToScreenUV(DownsampledScreenCoord + uint2(0, 1)), CornerDepths.z);
-			float3 Position11 = GetTranslatedWorldPositionFromScreenUV(DownsampledScreenCoordToScreenUV(DownsampledScreenCoord + uint2(1, 1)), CornerDepths.w);
+			float3 Position00 = GetTranslatedWorldPositionFromScreenUV(DownsampledScreenCoordToScreenUV(DownsampledScreenCoord00 + uint2(0, 0)), CornerDepths.x);
+			float3 Position10 = GetTranslatedWorldPositionFromScreenUV(DownsampledScreenCoordToScreenUV(DownsampledScreenCoord00 + uint2(1, 0)), CornerDepths.y);
+			float3 Position01 = GetTranslatedWorldPositionFromScreenUV(DownsampledScreenCoordToScreenUV(DownsampledScreenCoord00 + uint2(0, 1)), CornerDepths.z);
+			float3 Position11 = GetTranslatedWorldPositionFromScreenUV(DownsampledScreenCoordToScreenUV(DownsampledScreenCoord00 + uint2(1, 1)), CornerDepths.w);
 
 			float4 PlaneDistances;
 			PlaneDistances.x = abs(dot(float4(Position00, -1), ScenePlane));
diff --git a/Engine/Shaders/Private/ManyLights/ManyLights.ush b/Engine/Shaders/Private/ManyLights/ManyLights.ush
--- a/Engine/Shaders/Private/ManyLights/ManyLights.ush
+++ b/Engine/Shaders/Private/ManyLights/ManyLights.ush
@@ -97,22 +97,22 @@ uint ManyLightsStateFrameIndex;
 /** 
  * Returns sample jitter offset in the range [0, DOWNSAMPLE_FACTOR - 1]
  */
-uint2 GetSampleScreenCoordJitter()
+uint2 GetSampleScreenCoordJitter(uint2 DownsampledScreenCoord)
 {
-	return Hammersley16(ManyLightsStateFrameIndex, 8, 0) * DOWNSAMPLE_FACTOR;
-}
+	uint2 CellIndex = DownsampledScreenCoord % 2;
+	uint LinearIndex = CellIndex.x + CellIndex.y * 2;
+	LinearIndex = (LinearIndex + ManyLightsStateFrameIndex) % 4;
 
-/** 
- * Returns sample jitter offset in the range [0, DOWNSAMPLE_FACTOR - 1] from the last frame
- */
-uint2 GetHistorySampleScreenCoordJitter()
-{
-	return Hammersley16(ManyLightsStateFrameIndex + 7, 8, 0) * DOWNSAMPLE_FACTOR;
+	// 4 rooks sampling pattern
+	uint2 Jitter;
+	Jitter.x = LinearIndex & 0x02 ? 1 : 0;
+	Jitter.y = LinearIndex & 0x01 ? 0 : 1;
+	return Jitter;
 }
 
 uint2 DownsampledScreenCoordToScreenCoord(uint2 DownsampledScreenCoord)
 {
-	return DownsampledScreenCoord * DOWNSAMPLE_FACTOR + GetSampleScreenCoordJitter();
+	return DownsampledScreenCoord * DOWNSAMPLE_FACTOR + GetSampleScreenCoordJitter(DownsampledScreenCoord);
 }
 
 float2 DownsampledScreenCoordToScreenUV(uint2 DownsampledScreenCoord)
diff --git a/Engine/Shaders/Private/ManyLights/ManyLightsRayTracing.ush b/Engine/Shaders/Private/ManyLights/ManyLightsRayTracing.ush
--- a/Engine/Shaders/Private/ManyLights/ManyLightsRayTracing.ush
+++ b/Engine/Shaders/Private/ManyLights/ManyLightsRayTracing.ush
@@ -17,7 +17,7 @@ uint2 SampleCoordToDownsampledScreenCoord(uint2 SampleCoord)
 uint2 SampleCoordToScreenCoord(uint2 SampleCoord)
 {
 	uint2 DownsampledScreenCoord = SampleCoordToDownsampledScreenCoord(SampleCoord);
-	return DownsampledScreenCoord * DOWNSAMPLE_FACTOR + GetSampleScreenCoordJitter();
+	return DownsampledScreenCoord * DOWNSAMPLE_FACTOR + GetSampleScreenCoordJitter(DownsampledScreenCoord);
 }
 
 bool GenerateShadowRay(
