Fix the screen space denoiser causing ray traced transmission leakage.
* Add another bilaterial filtering weights specific to transmission for Screen space ray traced shadow denoising. So shadows on subsurface profile will behave the same to other material, and the transmission will be filtered separately. The spatial filtering weight will be set to 0 if the reference and the neighbor pixels are not all subsurface profile.
* Suggest: Increase r.Shadow.Denoiser.PreConvolution to 2 to remove the border line neglectable transmission artifact due to temporal accumulation.

diff --git a/Engine/Shaders/Private/ScreenSpaceDenoise/SSDSignalAccumulator.ush b/Engine/Shaders/Private/ScreenSpaceDenoise/SSDSignalAccumulator.ush
--- a/Engine/Shaders/Private/ScreenSpaceDenoise/SSDSignalAccumulator.ush
+++ b/Engine/Shaders/Private/ScreenSpaceDenoise/SSDSignalAccumulator.ush
@@ -150,7 +150,9 @@ struct FSSDSampleAccumulationInfos
 	FSSDSignalFrequency Frequency;
 
 	// Weight of the sample in the kernel.
-	float FinalWeight;
+	// FinalWeight.x: regular weight of the sample
+	// FinalWeight.y: weight for transmission distance. It should not blend across subsurface profile shading model
+	float2 FinalWeights;
 
 	// 1 / SignalFrequency
 	float InvFrequency;
@@ -178,14 +180,14 @@ float AmendWorldBluringRadiusCausedByPixelSize(float WorldBluringDistance)
 /** Accumulate a sample withing the accumulator. */
 void AccumulateSampleInOneBin(inout FSSDSignalAccumulator Accumulator, FSSDSampleAccumulationInfos SampleInfos)
 {
-	float SampleWeight = SampleInfos.FinalWeight;
+	float2 SampleWeights = SampleInfos.FinalWeights;
 	
 	#if COMPILE_MOMENT1_ACCUMULATOR
-		Accumulator.Moment1 = AddSignal(Accumulator.Moment1, MulSignal(SampleInfos.Sample, SampleWeight));
+		Accumulator.Moment1 = AddSignal(Accumulator.Moment1, MulSignals(SampleInfos.Sample, SampleWeights));
 	#endif
 	
 	#if COMPILE_MOMENT2_ACCUMULATOR
-		Accumulator.Moment2 = AddSignal(Accumulator.Moment2, MulSignal(PowerSignal(SampleInfos.Sample, 2), SampleWeight));
+		Accumulator.Moment2 = AddSignal(Accumulator.Moment2, MulSignals(PowerSignal(SampleInfos.Sample, 2), SampleWeights));
 	#endif
 }
 
@@ -196,7 +198,7 @@ void AccumulateSampleInOneBin(inout FSSDSignalAccumulator Accumulator, FSSDSampl
 void AccumulateSampleInDomainBoundaries(inout FSSDSignalAccumulator Accumulator, FSSDSampleAccumulationInfos SampleInfos)
 {
 	FLATTEN
-	if (SampleInfos.FinalWeight > 0)
+	if (SampleInfos.FinalWeights.x > 0)
 	{
 		#if COMPILE_MINMAX_ACCUMULATOR
 		{
@@ -262,9 +264,9 @@ void AccumulateSampleInDRB(inout FSSDSignalAccumulator Accumulator, FSSDSampleAc
 	
 	// Accumulate current bucket.
 	{
-		float CurrentWeight = A.FinalWeight * bBelongsToCurrent;
-		Accumulator.Current = AddSignal(Accumulator.Current, MulSignal(A.Sample, CurrentWeight));
-		Accumulator.CurrentInvFrequency += A.Sample.SampleCount * ClampedInvFrequency * CurrentWeight;
+		float2 CurrentWeights = A.FinalWeights * bBelongsToCurrent;
+		Accumulator.Current = AddSignal(Accumulator.Current, MulSignals(A.Sample, CurrentWeights));
+		Accumulator.CurrentInvFrequency += A.Sample.SampleCount * ClampedInvFrequency * CurrentWeights.x;
 	}
 
 	// Accumulate current bucket translucency.
@@ -283,9 +285,9 @@ void AccumulateSampleInDRB(inout FSSDSignalAccumulator Accumulator, FSSDSampleAc
 
 	// Accumulate previous bucket.
 	{
-		float PreviousWeight = A.FinalWeight * bBelongsToPrevious;
-		Accumulator.Previous = AddSignal(Accumulator.Previous, MulSignal(A.Sample, PreviousWeight));
-		Accumulator.PreviousInvFrequency += A.Sample.SampleCount * ClampedInvFrequency * PreviousWeight;
+		float2 PreviousWeights = A.FinalWeights * bBelongsToPrevious;
+		Accumulator.Previous = AddSignal(Accumulator.Previous, MulSignals(A.Sample, PreviousWeights));
+		Accumulator.PreviousInvFrequency += A.Sample.SampleCount * ClampedInvFrequency * PreviousWeights.x;
 	}
 }
 
diff --git a/Engine/Shaders/Private/ScreenSpaceDenoise/SSDSignalCore.ush b/Engine/Shaders/Private/ScreenSpaceDenoise/SSDSignalCore.ush
--- a/Engine/Shaders/Private/ScreenSpaceDenoise/SSDSignalCore.ush
+++ b/Engine/Shaders/Private/ScreenSpaceDenoise/SSDSignalCore.ush
@@ -178,6 +178,28 @@ FSSDSignalSample MulSignal(FSSDSignalSample Sample, float Scalar)
 	return OutSample;
 }
 
+FSSDSignalSample MulSignals(FSSDSignalSample Sample, float2 Scalars)
+{
+	FSSDSignalSample OutSample;
+#if COMPILE_SIGNAL_COLOR
+	OutSample.SceneColor = Sample.SceneColor * Scalars.x;
+#endif
+#if COMPILE_SIGNAL_COLOR_ARRAY
+	{
+		UNROLL_N(COMPILE_SIGNAL_COLOR_ARRAY)
+			for (uint ColorId = 0; ColorId < COMPILE_SIGNAL_COLOR_ARRAY; ColorId++)
+				OutSample.ColorArray[ColorId] = Sample.ColorArray[ColorId] * Scalars.x;
+	}
+#endif
+#if COMPILE_SIGNAL_COLOR_SH
+	OutSample.ColorSH = MulSH(Sample.ColorSH, Scalars.x);
+#endif
+	OutSample.SampleCount = Sample.SampleCount * Scalars.x;
+	OutSample.MissCount = Sample.MissCount * Scalars.x;
+	OutSample.TransmissionDistance = Sample.TransmissionDistance * Scalars.y;
+	return OutSample;
+}
+
 FSSDSignalSample AddSignal(FSSDSignalSample SampleA, FSSDSignalSample SampleB)
 {
 	#if COMPILE_SIGNAL_COLOR
diff --git a/Engine/Shaders/Private/ScreenSpaceDenoise/SSDSignalFramework.ush b/Engine/Shaders/Private/ScreenSpaceDenoise/SSDSignalFramework.ush
--- a/Engine/Shaders/Private/ScreenSpaceDenoise/SSDSignalFramework.ush
+++ b/Engine/Shaders/Private/ScreenSpaceDenoise/SSDSignalFramework.ush
@@ -428,7 +428,7 @@ float ComputeLightSourceRadius(FSSDSignalDomainKnowledge DomainKnowledge)
 #endif
 
 /** COmpute the weight of the bilateral filter to use between a reference sample and a neighbor. */
-float ComputeBilateralWeight(
+float2 ComputeBilateralWeight(
 	const uint BilateralSettings,
 	float MaxWorldBluringDistance,
 	FSSDSignalDomainKnowledge DomainKnowledge,
@@ -453,6 +453,8 @@ float ComputeBilateralWeight(
 	float InvMaxWorldBluringDistance = rcp(MaxWorldBluringDistance);
 
 	float PositionBilateralWeight = 1;
+	float TransmissionPositionBilaterialWeight = 0;
+
 	if ((BilateralSettings & 0xF) == 1)
 	{
 		const float WorldRadius = MaxWorldBluringDistance;
@@ -506,14 +508,7 @@ float ComputeBilateralWeight(
 		
 		float PlaneMisAlignementBilateral;
 		{
-			// Use source radius as encoding error to add more penalty to depth to avoid SSS backface transmission rim for spotlight with source radius > 0.
-			// It makes the scaling of the distance by (DistanceFromLight - HitT)/HitT instead of (DistanceFromLight - HitT)/HitT*SourceRadius
-			// When light is behind close to neighbor in view direction, the light tangent bilaterial will be close to 1. Little penalty
-			// is applied in depth. Adding the effect of SourceRadius on NeighborToRefError could add more penalty to depth to mitigate the
-			// bleeding error.
-			const bool bApplyDepthPenaltyWeightForSSS = RefSceneMetadata.ShadingModelID == SHADINGMODELID_SUBSURFACE_PROFILE && ComputeLightSourceRadius(DomainKnowledge) > 0.0f;
-			const float	NormalEncodingError = bApplyDepthPenaltyWeightForSSS ? DomainKnowledge.Light.SourceRadius : -0.25;
-
+			const float	NormalEncodingError = -0.25;
 			
 			float DistFromRefPlane = dot(RefNormal, -NeighborToRefVector);
 
@@ -524,6 +519,18 @@ float ComputeBilateralWeight(
 		}
 
 		PositionBilateralWeight = LightTangentBilateral * PlaneMisAlignementBilateral;
+
+		{
+			// Transmission of subsurface scattering should only be denoised between subsurface profile pixels to avoid bleeding.
+			const bool bApplyFilteringForSSS = RefSceneMetadata.ShadingModelID == SHADINGMODELID_SUBSURFACE_PROFILE
+				&& NeighborSceneMetadata.ShadingModelID == SHADINGMODELID_SUBSURFACE_PROFILE
+				&& ComputeLightSourceRadius(DomainKnowledge) > 0.0f;
+
+			if (bApplyFilteringForSSS)
+			{
+				TransmissionPositionBilaterialWeight = PositionBilateralWeight;
+			}
+		}
 	}
 	else if ((BilateralSettings & 0xF) == 6)
 	{
@@ -587,7 +594,7 @@ float ComputeBilateralWeight(
 			return 0.0;
 	}
 
-	return PositionBilateralWeight * NormalBilateralWeight * LobeSimilarity * AxesSimilarity;
+	return float2(PositionBilateralWeight, TransmissionPositionBilaterialWeight) * NormalBilateralWeight * LobeSimilarity * AxesSimilarity;
 }
 
 float GetSignalWorldBluringRadius(FSSDSignalFrequency SampleFrequency, FSSDSampleSceneInfos SceneMetadata, FSSDSignalDomainKnowledge DomainKnowledge)
diff --git a/Engine/Shaders/Private/ScreenSpaceDenoise/SSDSpatialAccumulation.usf b/Engine/Shaders/Private/ScreenSpaceDenoise/SSDSpatialAccumulation.usf
--- a/Engine/Shaders/Private/ScreenSpaceDenoise/SSDSpatialAccumulation.usf
+++ b/Engine/Shaders/Private/ScreenSpaceDenoise/SSDSpatialAccumulation.usf
@@ -895,7 +895,7 @@ void MainCS(
 				FSSDSampleAccumulationInfos SampleInfos;
 				SampleInfos.Sample = CenterSample;
 				SampleInfos.Frequency = RefFrequencies.Array[SignalMultiplexId];
-				SampleInfos.FinalWeight = 1.0;
+				SampleInfos.FinalWeights = 1.0;
 				SampleInfos.InvFrequency = GetSignalWorldBluringRadius(SampleInfos.Frequency, RefSceneMetadata, DomainKnowledge);
 				
 				if (KernelConfig.BilateralDistanceComputation == SIGNAL_WORLD_FREQUENCY_PRECOMPUTED_BLURING_RADIUS)
diff --git a/Engine/Shaders/Private/ScreenSpaceDenoise/SSDSpatialKernel.ush b/Engine/Shaders/Private/ScreenSpaceDenoise/SSDSpatialKernel.ush
--- a/Engine/Shaders/Private/ScreenSpaceDenoise/SSDSpatialKernel.ush
+++ b/Engine/Shaders/Private/ScreenSpaceDenoise/SSDSpatialKernel.ush
@@ -734,7 +734,7 @@ void AccumulateSampledMultiplexedSignals(
 		}
 	
 		// Compute the weight to be applied to do bilateral rejection.
-		float BilateralWeight = ComputeBilateralWeight(
+		float2 BilateralWeights = ComputeBilateralWeight(
 			KernelConfig.BilateralSettings[SignalMultiplexId],
 			FinalWorldBluringDistance,
 			DomainKnowledge,
@@ -745,12 +745,12 @@ void AccumulateSampledMultiplexedSignals(
 		FSSDSampleAccumulationInfos SampleInfos;
 		SampleInfos.Sample = Sample;
 		SampleInfos.Frequency = SampleFrequency;
-		SampleInfos.FinalWeight = KernelSampleWeight * BilateralWeight;
+		SampleInfos.FinalWeights = KernelSampleWeight * BilateralWeights;
 		SampleInfos.InvFrequency = SignalConvolutionBluringRadius;
 
 		if (bForceSample || KernelConfig.bForceAllAccumulation)
 		{
-			SampleInfos.FinalWeight = 1;
+			SampleInfos.FinalWeights = 1;
 		}
 
 		// TODO(Denoiser): bIsOutsideFrustum could afect number of samples for DRB.
