Change UnpackTileCoordXXbits to use BitFieldExtractU32 intrinsics.

diff --git a/Engine/Shaders/Private/PackUnpack.ush b/Engine/Shaders/Private/PackUnpack.ush
--- a/Engine/Shaders/Private/PackUnpack.ush
+++ b/Engine/Shaders/Private/PackUnpack.ush
@@ -279,7 +279,7 @@ float UnpackR24F(uint In)
 //
 
 #define PackTileCoordXXbits(In, InMask, InBits)  ((In.x & InMask) | ((In.y & InMask) << InBits))
-#define UnpackTileCoordXXbits(In,InMask, InBits) uint2(In & InMask, (In >> InBits) & InMask)
+#define UnpackTileCoordXXbits(In,InMask, InBits) uint2(BitFieldExtractU32(In, InBits, 0), BitFieldExtractU32(In, InBits, InBits))
 
 uint  PackTileCoord8bits(uint2 In)  { return PackTileCoordXXbits  (In, 0xFF,  8u); }
 uint2 UnpackTileCoord8bits(uint In) { return UnpackTileCoordXXbits(In, 0xFF,  8u); }
