Fixes ortographic view in TSR

diff --git a/Engine/Shaders/Private/TemporalSuperResolution/TSRClearPrevTextures.usf b/Engine/Shaders/Private/TemporalSuperResolution/TSRClearPrevTextures.usf
--- a/Engine/Shaders/Private/TemporalSuperResolution/TSRClearPrevTextures.usf
+++ b/Engine/Shaders/Private/TemporalSuperResolution/TSRClearPrevTextures.usf
@@ -12,8 +12,7 @@
 
 //------------------------------------------------------- PARAMETERS
 
-RWTexture2D<uint> PrevUseCountOutput;
-RWTexture2D<uint> PrevClosestDepthOutput;
+RWTexture2DArray<uint> PrevAtomicOutput;
 
 
 //------------------------------------------------------- ENTRY POINT
@@ -45,8 +44,12 @@ void MainCS(
 			#endif
 			PrevUseCountPixelPos.x = select(all(PrevUseCountPixelPos < PrevUseCountMaxPixelPos), PrevUseCountPixelPos.x, uint(~0));
 
-			PrevUseCountOutput[PrevUseCountPixelPos] = 0;
-			PrevClosestDepthOutput[InputPixelPos] = 0;
+			PrevAtomicOutput[uint3(PrevUseCountPixelPos, 0)] = 0;
+			PrevAtomicOutput[uint3(InputPixelPos, 1)] = 0;
+			
+			#if CONFIG_SUPPORT_ORTHO_VIEW
+				PrevAtomicOutput[uint3(InputPixelPos, 2)] = 0;
+			#endif
 		}
 	}
 }
diff --git a/Engine/Shaders/Private/TemporalSuperResolution/TSRCommon.ush b/Engine/Shaders/Private/TemporalSuperResolution/TSRCommon.ush
--- a/Engine/Shaders/Private/TemporalSuperResolution/TSRCommon.ush
+++ b/Engine/Shaders/Private/TemporalSuperResolution/TSRCommon.ush
@@ -55,6 +55,8 @@
 // which on 16bit, if linear color spaces ends up close to half(65504.0), this can go +inf on if the compiler transform a * 0.5 + b * 0.5 into (a + b) * 0.5.
 #define CONFIG_FP16_PRECISE_MULTIPLY_ORDER 1
 
+#define CONFIG_SUPPORT_ORTHO_VIEW 1
+
 
 //------------------------------------------------------- CONSTANTS
 
diff --git a/Engine/Shaders/Private/TemporalSuperResolution/TSRDecimateHistory.usf b/Engine/Shaders/Private/TemporalSuperResolution/TSRDecimateHistory.usf
--- a/Engine/Shaders/Private/TemporalSuperResolution/TSRDecimateHistory.usf
+++ b/Engine/Shaders/Private/TemporalSuperResolution/TSRDecimateHistory.usf
@@ -43,8 +43,7 @@ Texture2D<float2> DilatedVelocityTexture;
 #else
 	Texture2D<float> ClosestDepthTexture;
 #endif
-Texture2D<uint> PrevUseCountTexture;
-Texture2D<uint> PrevClosestDepthTexture;
+Texture2DArray<uint> PrevAtomicTextureArray;
 
 Texture2DArray<tsr_half4> PrevHistoryGuide;
 Texture2DArray<tsr_half4> PrevHistoryMoire;
@@ -238,7 +237,8 @@ void ProcessVelocityNeighborhood(
 void LoadPrevAtomicTexturesSamples(
 	float2 PrevScreenPos,
 	out uint HistoryUseCountSamples[4],
-	out uint HistoryClosestDeviceZSamples[4])
+	out uint HistoryClosestDeviceZSamples0[4],
+	out uint HistoryClosestDeviceZSamples1[4])
 {
 	float2 PrevInputBufferUV = (InputInfo_ScreenPosToViewportScale * PrevScreenPos + InputInfo_ScreenPosToViewportBias) * InputInfo_ExtentInverse;
 	PrevInputBufferUV = INVARIANT(PrevInputBufferUV);
@@ -257,19 +257,30 @@ void LoadPrevAtomicTexturesSamples(
 			uint Shift;
 			ComputeCompressedUseCountPixelCoordinates(PixelPos, /* out */ CompressedUseCountPixelPos, /* out */ Shift);
 
-			uint HistoryUseCount = PrevUseCountTexture[CompressedUseCountPixelPos];
+			uint HistoryUseCount = PrevAtomicTextureArray[uint3(CompressedUseCountPixelPos, 0)];
 		#else
-			uint HistoryUseCount = PrevUseCountTexture[PixelPos];
+			uint HistoryUseCount = PrevAtomicTextureArray[uint3(PixelPos, 0)];
 		#endif
 			
 		HistoryUseCountSamples[i] = HistoryUseCount;
-		HistoryClosestDeviceZSamples[i] = PrevClosestDepthTexture[PixelPos];
+		HistoryClosestDeviceZSamples0[i] = PrevAtomicTextureArray[uint3(PixelPos, 1)];
+
+		#if CONFIG_SUPPORT_ORTHO_VIEW
+		{
+			HistoryClosestDeviceZSamples1[i] = PrevAtomicTextureArray[uint3(PixelPos, 2)];
+		}
+		#else
+		{
+			HistoryClosestDeviceZSamples1[i] = 0;
+		}
+		#endif
 	}
 } // LoadPrevAtomicTexturesSamples()
 
 void ProcessPrevAtomicTexturesSamples(
 	uint HistoryUseCountSamples[4],
-	uint HistoryClosestDeviceZSamples[4],
+	uint HistoryClosestDeviceZSamples0[4],
+	uint HistoryClosestDeviceZSamples1[4],
 	float PrevDeviceZ,
 	float2 PrevScreenPos,
 	out tsr_half ParallaxRejectionMask,
@@ -286,45 +297,105 @@ void ProcessPrevAtomicTexturesSamples(
 	PrevInputBufferUV = INVARIANT(PrevInputBufferUV);
 	FBilinearSampleInfos BilinearInter = GetBilinearSampleLevelInfos(PrevInputBufferUV, InputInfo_Extent, InputInfo_ExtentInverse);
 
-	UNROLL_N(4)
-	for (uint i = 0; i < 4; i++)
+	
+	#if CONFIG_SUPPORT_ORTHO_VIEW
+	BRANCH
+	if (IsOrthoProjection())
 	{
-		float BilinearWeight = GetSampleWeight(BilinearInter, i);
-		uint2 PixelPos = GetSamplePixelCoord(BilinearInter, i);
+		float ClosestDeviceZ = 0.0f;
 
-		bool bValidPixelPos = all(and(PixelPos >= InputInfo_ViewportMin, PixelPos < InputInfo_ViewportMax));
+		UNROLL_N(4)
+		for (uint i = 0; i < 4; i++)
+		{
+			float BilinearWeight = GetSampleWeight(BilinearInter, i);
+			uint2 PixelPos = GetSamplePixelCoord(BilinearInter, i);
+
+			bool bValidPixelPos = all(and(PixelPos >= InputInfo_ViewportMin, PixelPos < InputInfo_ViewportMax));
 				
-		#if COMPRESS_PREV_USE_COUNT
-			uint2 CompressedUseCountPixelPos;
-			uint Shift;
-			ComputeCompressedUseCountPixelCoordinates(PixelPos, /* out */ CompressedUseCountPixelPos, /* out */ Shift);
+			#if COMPRESS_PREV_USE_COUNT
+				uint2 CompressedUseCountPixelPos;
+				uint Shift;
+				ComputeCompressedUseCountPixelCoordinates(PixelPos, /* out */ CompressedUseCountPixelPos, /* out */ Shift);
 
-			uint HistoryUseCount = (HistoryUseCountSamples[i] >> (8 * Shift)) & 0xFF;
-		#else
-			uint HistoryUseCount = HistoryUseCountSamples[i];
-		#endif
+				uint HistoryUseCount = (HistoryUseCountSamples[i] >> (8 * Shift)) & 0xFF;
+			#else
+				uint HistoryUseCount = HistoryUseCountSamples[i];
+			#endif
 			
-		uint SampleHistoryClosestDeviceZ = HistoryClosestDeviceZSamples[i];
+			uint SampleHistoryClosestDeviceZ0 = HistoryClosestDeviceZSamples0[i];
+			uint SampleHistoryClosestDeviceZ1 = HistoryClosestDeviceZSamples1[i];
 		
-		float HistoryClosestDeviceZ = f16tof32(SampleHistoryClosestDeviceZ >> VELOCITY_HOLE_FILLING_BITS);
-		float HistoryClosestWorldDepth = ConvertFromDeviceZ(HistoryClosestDeviceZ);
+			uint SampleEncodedHoleFillingVelocity = ((SampleHistoryClosestDeviceZ0 & 0x1FFu) << 9) | (SampleHistoryClosestDeviceZ1 & 0x1FFu);
 
-		tsr_half DepthRejection;
-		{
-			const float PixelDepthError = 3.0;
-			const float RadiusToDiameter = 2.0;
-			float WorldDepthEpsilon = GetDepthPixelRadiusForProjectionType(HistoryClosestWorldDepth) * PixelDepthError * RadiusToDiameter;
-			float DeltaDepth = abs(HistoryClosestWorldDepth - WorldDepth);
+			float HistoryClosestDeviceZ = float((SampleHistoryClosestDeviceZ0 & 0xFFFFFE00u) >> 9u) * rcp(8388607.0);
+			float HistoryClosestWorldDepth = ConvertFromDeviceZ(HistoryClosestDeviceZ);
+
+			tsr_half DepthRejection;
+			{
+				const float PixelDepthError = 3.0;
+				const float RadiusToDiameter = 2.0;
+				float WorldDepthEpsilon = GetDepthPixelRadiusForProjectionType(HistoryClosestWorldDepth) * PixelDepthError * RadiusToDiameter;
+				float DeltaDepth = abs(HistoryClosestWorldDepth - WorldDepth);
+
+				DepthRejection = saturate(tsr_half(2.0) - tsr_half(DeltaDepth / WorldDepthEpsilon));
+			}
+
+			tsr_half HistoryGhostingRejection = saturate(tsr_half(1.0) - tsr_half(3.0) * abs(tsr_half(HistoryUseCount) * rcp(tsr_half(PREV_USE_COUNT_QUANTIZATION)) - tsr_half(1.0)));
+		
+			bIsValidHoleFillingPixelVelocity = bIsValidHoleFillingPixelVelocity || bValidPixelPos;
 
-			DepthRejection = saturate(tsr_half(2.0) - tsr_half(DeltaDepth / WorldDepthEpsilon));
+			FLATTEN
+			if (bValidPixelPos && HistoryClosestDeviceZ > ClosestDeviceZ)
+			{
+				EncodedHoleFillingVelocity = SampleEncodedHoleFillingVelocity;
+			}
+
+			ParallaxRejectionMask = select(bValidPixelPos, ParallaxRejectionMask + tsr_half(BilinearWeight) * max(DepthRejection, HistoryGhostingRejection), ParallaxRejectionMask);
 		}
+	}
+	else
+	#endif // CONFIG_SUPPORT_ORTHO_VIEW
+	{
+		UNROLL_N(4)
+		for (uint i = 0; i < 4; i++)
+		{
+			float BilinearWeight = GetSampleWeight(BilinearInter, i);
+			uint2 PixelPos = GetSamplePixelCoord(BilinearInter, i);
+
+			bool bValidPixelPos = all(and(PixelPos >= InputInfo_ViewportMin, PixelPos < InputInfo_ViewportMax));
+				
+			#if COMPRESS_PREV_USE_COUNT
+				uint2 CompressedUseCountPixelPos;
+				uint Shift;
+				ComputeCompressedUseCountPixelCoordinates(PixelPos, /* out */ CompressedUseCountPixelPos, /* out */ Shift);
+
+				uint HistoryUseCount = (HistoryUseCountSamples[i] >> (8 * Shift)) & 0xFF;
+			#else
+				uint HistoryUseCount = HistoryUseCountSamples[i];
+			#endif
+			
+			uint SampleHistoryClosestDeviceZ = HistoryClosestDeviceZSamples0[i];
+		
+			float HistoryClosestDeviceZ = f16tof32(SampleHistoryClosestDeviceZ >> VELOCITY_HOLE_FILLING_BITS);
+			float HistoryClosestWorldDepth = ConvertFromDeviceZ(HistoryClosestDeviceZ);
+
+			tsr_half DepthRejection;
+			{
+				const float PixelDepthError = 3.0;
+				const float RadiusToDiameter = 2.0;
+				float WorldDepthEpsilon = GetDepthPixelRadiusForProjectionType(HistoryClosestWorldDepth) * PixelDepthError * RadiusToDiameter;
+				float DeltaDepth = abs(HistoryClosestWorldDepth - WorldDepth);
+
+				DepthRejection = saturate(tsr_half(2.0) - tsr_half(DeltaDepth / WorldDepthEpsilon));
+			}
 
-		tsr_half HistoryGhostingRejection = saturate(tsr_half(1.0) - tsr_half(3.0) * abs(tsr_half(HistoryUseCount) * rcp(tsr_half(PREV_USE_COUNT_QUANTIZATION)) - tsr_half(1.0)));
+			tsr_half HistoryGhostingRejection = saturate(tsr_half(1.0) - tsr_half(3.0) * abs(tsr_half(HistoryUseCount) * rcp(tsr_half(PREV_USE_COUNT_QUANTIZATION)) - tsr_half(1.0)));
 		
-		bIsValidHoleFillingPixelVelocity = bIsValidHoleFillingPixelVelocity || bValidPixelPos;
+			bIsValidHoleFillingPixelVelocity = bIsValidHoleFillingPixelVelocity || bValidPixelPos;
 
-		EncodedHoleFillingVelocity = select(bValidPixelPos, max(EncodedHoleFillingVelocity, SampleHistoryClosestDeviceZ), EncodedHoleFillingVelocity);
-		ParallaxRejectionMask = select(bValidPixelPos, ParallaxRejectionMask + tsr_half(BilinearWeight) * max(DepthRejection, HistoryGhostingRejection), ParallaxRejectionMask);
+			EncodedHoleFillingVelocity = select(bValidPixelPos, max(EncodedHoleFillingVelocity, SampleHistoryClosestDeviceZ), EncodedHoleFillingVelocity);
+			ParallaxRejectionMask = select(bValidPixelPos, ParallaxRejectionMask + tsr_half(BilinearWeight) * max(DepthRejection, HistoryGhostingRejection), ParallaxRejectionMask);
+		}
 	}
 } // ProcessPrevAtomicTexturesSamples()
 
@@ -407,6 +478,111 @@ tsr_half4 CorrectMoireExposure(tsr_half4 ReprojectedHistoryMoireColor, const flo
 }
 #endif
 
+//ISOLATE
+void ReprojectAllPrevTextures(
+	uint2 GroupId,
+	uint GroupThreadIndex,
+	float2 EncodedScreenVelocity,
+	float DeviceZ,
+	out uint HistoryUseCountSamples[4],
+	out uint HistoryClosestDeviceZSamples0[4],
+	out uint HistoryClosestDeviceZSamples1[4],
+	out FCatmullRomFetches<tsr_half4> ReprojectedHistoryGuideSamples,
+	#if CONFIG_SCENE_COLOR_ALPHA
+		out FCatmullRomFetches<tsr_half4> ReprojectedHistoryGuideUncertaintySamples,
+	#endif
+	#if CONFIG_MOIRE_REPROJECTION
+		out tsr_half4 ReprojectedHistoryMoireColor,
+	#endif
+	#if CONFIG_RESURRECTION_REPROJECTION
+		out tsr_halfC ResurrectedHistoryGuideColor,
+		out tsr_half  ResurrectedHistoryGuideUncertainty,
+	#endif
+	out bool bIsOffScreen)
+{
+	tsr_ushort2 InputPixelPos = (
+		Map8x8Tile2x2Lane(GroupThreadIndex) +
+		(tsr_ushort2(InputInfo_ViewportMin) + tsr_ushort2(GroupId) * tsr_ushort2(TILE_SIZE, TILE_SIZE)));
+			
+	{
+		float2 ScreenPos = ApplyScreenTransform(float2(InputPixelPos), InputPixelPosToScreenPos);
+		float2 PrevScreenPos = ScreenPos - DecodeVelocityFromTexture(float4(EncodedScreenVelocity, 0.0, 0.0)).xy;
+		LoadPrevAtomicTexturesSamples(PrevScreenPos,
+			/* out */ HistoryUseCountSamples,
+			/* out */ HistoryClosestDeviceZSamples0,
+			/* out */ HistoryClosestDeviceZSamples1);
+	}
+			
+	{
+		float2 OutputScreenPos = ApplyScreenTransform(float2(InputPixelPos), InputPixelPosToReprojectScreenPos);
+	
+		float2 PrevOutputScreenPos = OutputScreenPos - DecodeVelocityFromTexture(float4(EncodedScreenVelocity, 0.0, 0.0)).xy;
+		bIsOffScreen = (
+			(PrevOutputScreenPos.x < -1.0) ||
+			(PrevOutputScreenPos.y < -1.0) ||
+			(PrevOutputScreenPos.x > 1.0) ||
+			(PrevOutputScreenPos.y > 1.0));
+				
+		float2 KernelHistoryBufferUV = ApplyScreenTransform(PrevOutputScreenPos, ScreenPosToPrevHistoryGuideBufferUV);
+
+		#if CONFIG_SCENE_COLOR_ALPHA
+			ReprojectedHistoryGuideSamples.FetchSamples(
+				PrevHistoryGuide,
+				KernelHistoryBufferUV,
+				PrevFrameIndex * 2.0 + 0.0,
+				PrevGuideInfo_Extent,
+				PrevGuideInfo_ExtentInverse,
+				PrevGuideInfo_UVViewportBilinearMin,
+				PrevGuideInfo_UVViewportBilinearMax);
+					
+			ReprojectedHistoryGuideUncertaintySamples.FetchSamples(
+				PrevHistoryGuide,
+				KernelHistoryBufferUV,
+				PrevFrameIndex * 2.0 + 1.0,
+				PrevGuideInfo_Extent,
+				PrevGuideInfo_ExtentInverse,
+				PrevGuideInfo_UVViewportBilinearMin,
+				PrevGuideInfo_UVViewportBilinearMax);
+		#else
+			ReprojectedHistoryGuideSamples.FetchSamples(
+				PrevHistoryGuide,
+				KernelHistoryBufferUV,
+				PrevFrameIndex,
+				PrevGuideInfo_Extent,
+				PrevGuideInfo_ExtentInverse,
+				PrevGuideInfo_UVViewportBilinearMin,
+				PrevGuideInfo_UVViewportBilinearMax);
+		#endif
+
+		#if CONFIG_MOIRE_REPROJECTION
+		{
+			ReprojectedHistoryMoireColor = ReprojectPrevHistoryMoire(InputPixelPos, KernelHistoryBufferUV);
+		}
+		#endif
+	}
+
+	#if CONFIG_RESURRECTION_REPROJECTION
+	{
+		float2 OutputScreenPos = ApplyScreenTransform(float2(InputPixelPos), InputPixelPosToReprojectScreenPos);
+
+		float4 ThisClip = float4(OutputScreenPos, DeviceZ, 1);
+		float4 ResurrectionClip = mul(ThisClip, ClipToResurrectionClip); 
+		float2 ResurrectionScreen = ResurrectionClip.xy / ResurrectionClip.w;
+
+		float2 KernelHistoryBufferUV = ApplyScreenTransform(ResurrectionScreen, ScreenPosToResurrectionGuideBufferUV);
+		KernelHistoryBufferUV = fastClamp(KernelHistoryBufferUV, ResurrectionGuideUVViewportBilinearMin, ResurrectionGuideUVViewportBilinearMax);
+
+		#if CONFIG_SCENE_COLOR_ALPHA
+			ResurrectedHistoryGuideColor       = PrevHistoryGuide.SampleLevel(GlobalBilinearClampedSampler, float3(KernelHistoryBufferUV, ResurrectionFrameIndex * 2.0 + 0.0), 0);
+			ResurrectedHistoryGuideUncertainty = PrevHistoryGuide.SampleLevel(GlobalBilinearClampedSampler, float3(KernelHistoryBufferUV, ResurrectionFrameIndex * 2.0 + 1.0), 0).r;
+		#else
+			tsr_half4 RawGuide                 = PrevHistoryGuide.SampleLevel(GlobalBilinearClampedSampler, float3(KernelHistoryBufferUV, ResurrectionFrameIndex), 0);
+			ResurrectedHistoryGuideColor = RawGuide.rgb;
+			ResurrectedHistoryGuideUncertainty = RawGuide.a;
+		#endif
+	}
+	#endif
+} // ReprojectAllPrevTextures()
 
 //------------------------------------------------------- ENTRY POINT
 
@@ -422,9 +598,7 @@ void MainCS(
 
 	// Sample current frame informations.
 	float2 EncodedScreenVelocity;
-	#if CONFIG_RESURRECTION_REPROJECTION
-		float DeviceZ;
-	#endif
+	float DeviceZ;
 	float PrevDeviceZ;
 	tsr_half VelocityEdge = tsr_half(1.0);
 	bool bHasPixelAnimation;
@@ -446,6 +620,7 @@ void MainCS(
 				DeviceZ = DeviceZAndPrevDeviceZ.g;
 			#else
 				float DeviceZAndPrevDeviceZ = ClosestDepthTexture[InputPixelPos];
+				DeviceZ = 0.0;
 			#endif
 			
 			PrevDeviceZ = abs(DeviceZAndPrevDeviceZ.r);
@@ -486,7 +661,8 @@ void MainCS(
 	
 	// Sample the previous frame information
 	uint HistoryUseCountSamples[4];
-	uint HistoryClosestDeviceZSamples[4];
+	uint HistoryClosestDeviceZSamples0[4];
+	uint HistoryClosestDeviceZSamples1[4];
 	
 
 	FCatmullRomFetches<tsr_half4> ReprojectedHistoryGuideSamples;
@@ -502,97 +678,33 @@ void MainCS(
 	#endif
 	bool bIsOffScreen;
 
-	// Issues all the history texture fetches
-	ISOLATE
+	// Issues all the overlapping texture fetches for previous frame texture using motion vectors
 	{
-		tsr_ushort2 InputPixelPos = (
-			Map8x8Tile2x2Lane(GroupThreadIndex) +
-			(tsr_ushort2(InputInfo_ViewportMin) + tsr_ushort2(GroupId) * tsr_ushort2(TILE_SIZE, TILE_SIZE)));
-			
-		{
-			float2 ScreenPos = ApplyScreenTransform(float2(InputPixelPos), InputPixelPosToScreenPos);
-			float2 PrevScreenPos = ScreenPos - DecodeVelocityFromTexture(float4(EncodedScreenVelocity, 0.0, 0.0)).xy;
-			LoadPrevAtomicTexturesSamples(PrevScreenPos,
-				/* out */ HistoryUseCountSamples,
-				/* out */ HistoryClosestDeviceZSamples);
-		}
-			
-		{
-			float2 OutputScreenPos = ApplyScreenTransform(float2(InputPixelPos), InputPixelPosToReprojectScreenPos);
-	
-			float2 PrevOutputScreenPos = OutputScreenPos - DecodeVelocityFromTexture(float4(EncodedScreenVelocity, 0.0, 0.0)).xy;
-			bIsOffScreen = (
-				(PrevOutputScreenPos.x < -1.0) ||
-				(PrevOutputScreenPos.y < -1.0) ||
-				(PrevOutputScreenPos.x > 1.0) ||
-				(PrevOutputScreenPos.y > 1.0));
-				
-			float2 KernelHistoryBufferUV = ApplyScreenTransform(PrevOutputScreenPos, ScreenPosToPrevHistoryGuideBufferUV);
-
+		ReprojectAllPrevTextures(
+			GroupId, GroupThreadIndex,
+			EncodedScreenVelocity, DeviceZ,
+			/* out */ HistoryUseCountSamples,
+			/* out */ HistoryClosestDeviceZSamples0,
+			/* out */ HistoryClosestDeviceZSamples1,
+			/* out */ ReprojectedHistoryGuideSamples,
 			#if CONFIG_SCENE_COLOR_ALPHA
-				ReprojectedHistoryGuideSamples.FetchSamples(
-					PrevHistoryGuide,
-					KernelHistoryBufferUV,
-					PrevFrameIndex * 2.0 + 0.0,
-					PrevGuideInfo_Extent,
-					PrevGuideInfo_ExtentInverse,
-					PrevGuideInfo_UVViewportBilinearMin,
-					PrevGuideInfo_UVViewportBilinearMax);
-					
-				ReprojectedHistoryGuideUncertaintySamples.FetchSamples(
-					PrevHistoryGuide,
-					KernelHistoryBufferUV,
-					PrevFrameIndex * 2.0 + 1.0,
-					PrevGuideInfo_Extent,
-					PrevGuideInfo_ExtentInverse,
-					PrevGuideInfo_UVViewportBilinearMin,
-					PrevGuideInfo_UVViewportBilinearMax);
-			#else
-				ReprojectedHistoryGuideSamples.FetchSamples(
-					PrevHistoryGuide,
-					KernelHistoryBufferUV,
-					PrevFrameIndex,
-					PrevGuideInfo_Extent,
-					PrevGuideInfo_ExtentInverse,
-					PrevGuideInfo_UVViewportBilinearMin,
-					PrevGuideInfo_UVViewportBilinearMax);
+				/* out */ ReprojectedHistoryGuideUncertaintySamples,
 			#endif
-
 			#if CONFIG_MOIRE_REPROJECTION
-			{
-				ReprojectedHistoryMoireColor = ReprojectPrevHistoryMoire(InputPixelPos, KernelHistoryBufferUV);
-			}
+				/* out */ ReprojectedHistoryMoireColor,
 			#endif
-		}
-
-		#if CONFIG_RESURRECTION_REPROJECTION
-		{
-			float2 OutputScreenPos = ApplyScreenTransform(float2(InputPixelPos), InputPixelPosToReprojectScreenPos);
-
-			float4 ThisClip = float4(OutputScreenPos, DeviceZ, 1);
-			float4 ResurrectionClip = mul(ThisClip, ClipToResurrectionClip); 
-			float2 ResurrectionScreen = ResurrectionClip.xy / ResurrectionClip.w;
-
-			float2 KernelHistoryBufferUV = ApplyScreenTransform(ResurrectionScreen, ScreenPosToResurrectionGuideBufferUV);
-			KernelHistoryBufferUV = fastClamp(KernelHistoryBufferUV, ResurrectionGuideUVViewportBilinearMin, ResurrectionGuideUVViewportBilinearMax);
-
-			#if CONFIG_SCENE_COLOR_ALPHA
-				ResurrectedHistoryGuideColor       = PrevHistoryGuide.SampleLevel(GlobalBilinearClampedSampler, float3(KernelHistoryBufferUV, ResurrectionFrameIndex * 2.0 + 0.0), 0);
-				ResurrectedHistoryGuideUncertainty = PrevHistoryGuide.SampleLevel(GlobalBilinearClampedSampler, float3(KernelHistoryBufferUV, ResurrectionFrameIndex * 2.0 + 1.0), 0).r;
-			#else
-				tsr_half4 RawGuide                 = PrevHistoryGuide.SampleLevel(GlobalBilinearClampedSampler, float3(KernelHistoryBufferUV, ResurrectionFrameIndex), 0);
-				ResurrectedHistoryGuideColor = RawGuide.rgb;
-				ResurrectedHistoryGuideUncertainty = RawGuide.a;
+			#if CONFIG_RESURRECTION_REPROJECTION
+				/* out */ ResurrectedHistoryGuideColor,
+				/* out */ ResurrectedHistoryGuideUncertainty,
 			#endif
-		}
-		#endif
+			/* out */ bIsOffScreen);
 	}
 	
 	// Process the history samples
 	tsr_halfC ReprojectedHistoryGuideColor;
 	tsr_half  ReprojectedHistoryGuideUncertainty;
 	tsr_half ParallaxRejectionMask;
-	ISOLATE
+	//ISOLATE
 	{
 		#if CONFIG_SCENE_COLOR_ALPHA
 		{
@@ -642,7 +754,8 @@ void MainCS(
 		uint EncodedHoleFillingVelocity = 0;
 		ProcessPrevAtomicTexturesSamples(
 			HistoryUseCountSamples,
-			HistoryClosestDeviceZSamples,
+			HistoryClosestDeviceZSamples0,
+			HistoryClosestDeviceZSamples1,
 			PrevDeviceZ,
 			PrevScreenPos,
 			/* out */ ParallaxRejectionMask,
diff --git a/Engine/Shaders/Private/TemporalSuperResolution/TSRDilateVelocity.usf b/Engine/Shaders/Private/TemporalSuperResolution/TSRDilateVelocity.usf
--- a/Engine/Shaders/Private/TemporalSuperResolution/TSRDilateVelocity.usf
+++ b/Engine/Shaders/Private/TemporalSuperResolution/TSRDilateVelocity.usf
@@ -51,15 +51,14 @@ Texture2D SceneVelocityTexture;
 
 RWTexture2D<float2> DilatedVelocityOutput;
 RWTexture2D<float2> ClosestDepthOutput;
-globallycoherent RWTexture2D<uint> PrevUseCountOutput;
-globallycoherent RWTexture2D<uint> PrevClosestDepthOutput;
+globallycoherent RWTexture2DArray<uint> PrevAtomicOutput;
 RWTexture2DArray<uint> R8Output;
 
 RWTexture2D<float3>	VelocityFlattenOutput;
 RWTexture2DArray<float4> VelocityTileArrayOutput;
 
 
-//------------------------------------------------------- ENTRY POINT
+//------------------------------------------------------- UTILITY
 
 uint PackDepthAndOffset(const int2 Offset, float DeviceZ)
 {
@@ -74,6 +73,90 @@ void UnpackDepthAndOffset(uint Packed, out int2 Offset, out float DeviceZ)
 }
 
 
+//------------------------------------------------------- FUNCTION
+
+void ScatterClosestOccluder(bool bValidOutputPixel, float2 ScreenPos, float2 FinalScreenVelocity, float PrevClosestDeviceZ)
+{
+	float2 PrevScreenPos = ScreenPos - FinalScreenVelocity.xy;
+	float2 PrevInputBufferUV = (InputInfo_ScreenPosToViewportScale * PrevScreenPos + InputInfo_ScreenPosToViewportBias) * InputInfo_ExtentInverse;
+
+	float2 PixelVelocity = FinalScreenVelocity.xy * 0.5 * InputInfo_ViewportSize;
+	uint EncodedHoleFillingVelocity = EncodeHoleFillingVelocity(PixelVelocity);
+
+	bool bValidHistoryCoord = bValidOutputPixel && all(and(PrevInputBufferUV > PrevOutputBufferUVMin, PrevInputBufferUV < PrevOutputBufferUVMax));
+
+	FBilinearSampleInfos BilinearInter = GetBilinearSampleLevelInfos(PrevInputBufferUV, InputInfo_Extent, InputInfo_ExtentInverse);
+	
+	#if CONFIG_SUPPORT_ORTHO_VIEW
+	BRANCH
+	if (IsOrthoProjection())
+	{
+		uint Depth = clamp(uint(round(PrevClosestDeviceZ * 8388607.0)), 0, 8388607) << 9;
+
+		uint PrevClosestDepth0 = Depth | ((EncodedHoleFillingVelocity >> 9u) & 0x1FFu);
+		uint PrevClosestDepth1 = Depth | ((EncodedHoleFillingVelocity >> 0u) & 0x1FFu);
+		
+		UNROLL_N(4)
+		for (uint i = 0; i < 4; i++)
+		{
+			float Weight = GetSampleWeight(BilinearInter, i);
+			uint2 PixelPos = GetSamplePixelCoord(BilinearInter, i);
+
+			uint FinalAdd = uint(round(Weight * PREV_USE_COUNT_QUANTIZATION));
+			
+			PixelPos.x = ((FinalAdd > 0) && bValidHistoryCoord) ? PixelPos.x : (~0);
+
+			#if COMPRESS_PREV_USE_COUNT
+				uint2 CompressedUseCountPixelPos;
+				uint Shift;
+				ComputeCompressedUseCountPixelCoordinates(PixelPos, /* out */ CompressedUseCountPixelPos, /* out */ Shift);
+
+				FinalAdd = FinalAdd << (8 * Shift);
+			#else
+				uint2 CompressedUseCountPixelPos = PixelPos;
+			#endif
+
+			InterlockedAdd(PrevAtomicOutput[tsr_ushort3(CompressedUseCountPixelPos, 0)], FinalAdd);
+
+			uint PreviousDepth0;
+			InterlockedMax(PrevAtomicOutput[tsr_ushort3(PixelPos, 1)], PrevClosestDepth0, /* out */ PreviousDepth0);
+
+			PixelPos.x = select(PreviousDepth0 < PrevClosestDepth0, PixelPos.x, ~0u);
+			InterlockedMax(PrevAtomicOutput[tsr_ushort3(PixelPos, 2)], PrevClosestDepth1);
+		}
+	}
+	else
+	#endif // CONFIG_SUPPORT_ORTHO_VIEW
+	{
+		uint PrevClosestDepth = (f32tof16(PrevClosestDeviceZ) << VELOCITY_HOLE_FILLING_BITS) | EncodedHoleFillingVelocity;
+			
+		UNROLL_N(4)
+		for (uint i = 0; i < 4; i++)
+		{
+			float Weight = GetSampleWeight(BilinearInter, i);
+			uint2 PixelPos = GetSamplePixelCoord(BilinearInter, i);
+
+			uint FinalAdd = uint(round(Weight * PREV_USE_COUNT_QUANTIZATION));
+			
+			PixelPos.x = ((FinalAdd > 0) && bValidHistoryCoord) ? PixelPos.x : (~0);
+
+			#if COMPRESS_PREV_USE_COUNT
+				uint2 CompressedUseCountPixelPos;
+				uint Shift;
+				ComputeCompressedUseCountPixelCoordinates(PixelPos, /* out */ CompressedUseCountPixelPos, /* out */ Shift);
+
+				FinalAdd = FinalAdd << (8 * Shift);
+			#else
+				uint2 CompressedUseCountPixelPos = PixelPos;
+			#endif
+
+			InterlockedAdd(PrevAtomicOutput[tsr_ushort3(CompressedUseCountPixelPos, 0)], FinalAdd);
+			InterlockedMax(PrevAtomicOutput[tsr_ushort3(PixelPos, 1)], PrevClosestDepth);
+		}
+	}
+}
+		
+
 //------------------------------------------------------- ENTRY POINT
 
 [numthreads(TILE_SIZE * TILE_SIZE, 1, 1)]
@@ -319,47 +402,6 @@ void MainCS(
 	{
 		bool bValidOutputPixel = all(InputPixelPos < InputInfo_ViewportMax);
 		
-		// Scatter parallax rejection
-		#if 1
-		{
-			float2 PrevScreenPos = ScreenPos - FinalScreenVelocity.xy;
-			float2 PrevInputBufferUV = (InputInfo_ScreenPosToViewportScale * PrevScreenPos + InputInfo_ScreenPosToViewportBias) * InputInfo_ExtentInverse;
-
-			float2 PixelVelocity = FinalScreenVelocity.xy * 0.5 * InputInfo_ViewportSize;
-			uint EncodedHoleFillingVelocity = EncodeHoleFillingVelocity(PixelVelocity);
-
-			uint PrevClosestDepth = (f32tof16(PrevClosestDeviceZ) << VELOCITY_HOLE_FILLING_BITS) | EncodedHoleFillingVelocity;
-			
-			bool bValidHistoryCoord = bValidOutputPixel && all(and(PrevInputBufferUV > PrevOutputBufferUVMin, PrevInputBufferUV < PrevOutputBufferUVMax));
-
-			FBilinearSampleInfos BilinearInter = GetBilinearSampleLevelInfos(PrevInputBufferUV, InputInfo_Extent, InputInfo_ExtentInverse);
-
-			UNROLL_N(4)
-			for (uint i = 0; i < 4; i++)
-			{
-				float Weight = GetSampleWeight(BilinearInter, i);
-				uint2 PixelPos = GetSamplePixelCoord(BilinearInter, i);
-
-				uint FinalAdd = uint(round(Weight * PREV_USE_COUNT_QUANTIZATION));
-			
-				PixelPos.x = ((FinalAdd > 0) && bValidHistoryCoord) ? PixelPos.x : (~0);
-
-				#if COMPRESS_PREV_USE_COUNT
-					uint2 CompressedUseCountPixelPos;
-					uint Shift;
-					ComputeCompressedUseCountPixelCoordinates(PixelPos, /* out */ CompressedUseCountPixelPos, /* out */ Shift);
-
-					FinalAdd = FinalAdd << (8 * Shift);
-				#else
-					uint2 CompressedUseCountPixelPos = PixelPos;
-				#endif
-
-				InterlockedAdd(PrevUseCountOutput[CompressedUseCountPixelPos], FinalAdd);
-				InterlockedMax(PrevClosestDepthOutput[PixelPos], PrevClosestDepth);
-			}
-		}
-		#endif
-		
 		float4 EncodedVelocity = SceneVelocityTexture[uint2(InputBufferUV * InputInfo_Extent)];
 		
 		uint2 OutputPixelPos = bValidOutputPixel ? InputPixelPos : uint(~0).xx;
@@ -383,7 +425,10 @@ void MainCS(
 			StoreVelocityRange(VelocityTileArrayOutput, OutputTilePos, PolarToCartesian(VelocityPolarRange));
 		}
 		#endif
-
+		
+		// Scatter parallax rejection
+		ScatterClosestOccluder(bValidOutputPixel, ScreenPos, FinalScreenVelocity.xy, PrevClosestDeviceZ);
+		
 		#if DEBUG_OUTPUT && 0 // running out of UAVs
 		{
 			DebugOutput[tsr_short3(OutputPixelPos, 0)] = Debug;
diff --git a/Engine/Source/Runtime/Renderer/Private/PostProcess/TemporalSuperResolution.cpp b/Engine/Source/Runtime/Renderer/Private/PostProcess/TemporalSuperResolution.cpp
--- a/Engine/Source/Runtime/Renderer/Private/PostProcess/TemporalSuperResolution.cpp
+++ b/Engine/Source/Runtime/Renderer/Private/PostProcess/TemporalSuperResolution.cpp
@@ -559,8 +559,7 @@ class FTSRClearPrevTexturesCS : public FTSRShader
 	BEGIN_SHADER_PARAMETER_STRUCT(FParameters, )
 		SHADER_PARAMETER_STRUCT_INCLUDE(FTSRCommonParameters, CommonParameters)
 
-		SHADER_PARAMETER_RDG_TEXTURE_UAV(RWTexture2D, PrevUseCountOutput)
-		SHADER_PARAMETER_RDG_TEXTURE_UAV(RWTexture2D, PrevClosestDepthOutput)
+		SHADER_PARAMETER_RDG_TEXTURE_UAV(RWTexture2DArray, PrevAtomicOutput)
 	END_SHADER_PARAMETER_STRUCT()
 }; // class FTSRClearPrevTexturesCS
 
@@ -588,8 +587,7 @@ class FTSRDilateVelocityCS : public FTSRShader
 
 		SHADER_PARAMETER_RDG_TEXTURE_UAV(RWTexture2D, DilatedVelocityOutput)
 		SHADER_PARAMETER_RDG_TEXTURE_UAV(RWTexture2D, ClosestDepthOutput)
-		SHADER_PARAMETER_RDG_TEXTURE_UAV(RWTexture2D, PrevUseCountOutput)
-		SHADER_PARAMETER_RDG_TEXTURE_UAV(RWTexture2D, PrevClosestDepthOutput)
+		SHADER_PARAMETER_RDG_TEXTURE_UAV(RWTexture2DArray, PrevAtomicOutput)
 		SHADER_PARAMETER_RDG_TEXTURE_UAV(RWTexture2DArray, R8Output)
 
 		SHADER_PARAMETER_RDG_TEXTURE_UAV(RWTexture2D, VelocityFlattenOutput)
@@ -635,8 +633,7 @@ class FTSRDecimateHistoryCS : public FTSRShader
 		SHADER_PARAMETER_RDG_TEXTURE(Texture2D, InputSceneColorTexture)
 		SHADER_PARAMETER_RDG_TEXTURE(Texture2D, DilatedVelocityTexture)
 		SHADER_PARAMETER_RDG_TEXTURE(Texture2D, ClosestDepthTexture)
-		SHADER_PARAMETER_RDG_TEXTURE(Texture2D, PrevUseCountTexture)
-		SHADER_PARAMETER_RDG_TEXTURE(Texture2D, PrevClosestDepthTexture)
+		SHADER_PARAMETER_RDG_TEXTURE(Texture2DArray, PrevAtomicTextureArray)
 
 		SHADER_PARAMETER_STRUCT_INCLUDE(FTSRPrevHistoryParameters, PrevHistoryParameters)
 		SHADER_PARAMETER_RDG_TEXTURE(Texture2D, PrevHistoryGuide)
@@ -1316,6 +1313,9 @@ FDefaultTemporalUpscaler::FOutputs AddTemporalSuperResolutionPasses(
 	// Maximum number sample for each output pixel in the history
 	const float MaxHistorySampleCount = FMath::Clamp(CVarTSRHistorySampleCount.GetValueOnRenderThread(), 8.0f, 32.0f);
 
+	// Whether the view is orthographic view
+	const bool bIsOrthoProjection = !View.IsPerspectiveProjection();
+
 	// whether TSR passes can run on async compute.
 	int32 AsyncComputePasses = GSupportsEfficientAsyncCompute ? CVarTSRAsyncCompute.GetValueOnRenderThread() : 0;
 
@@ -1734,24 +1734,22 @@ FDefaultTemporalUpscaler::FOutputs AddTemporalSuperResolutionPasses(
 	}
 
 	// Clear atomic scattered texture.
-	FRDGTextureRef PrevUseCountTexture;
-	FRDGTextureRef PrevClosestDepthTexture;
+	FRDGTextureRef PrevAtomicTextureArray;
 	{
 		{
-			FRDGTextureDesc Desc = FRDGTextureDesc::Create2D(
+			FRDGTextureDesc Desc = FRDGTextureDesc::Create2DArray(
 				InputExtent,
 				PF_R32_UINT,
 				FClearValueBinding::None,
-				/* InFlags = */ TexCreate_ShaderResource | TexCreate_UAV | TexCreate_AtomicCompatible);
+				/* InFlags = */ TexCreate_ShaderResource | TexCreate_UAV | TexCreate_AtomicCompatible,
+				/* ArraySize = */ bIsOrthoProjection ? 3 : 2);
 
-			PrevUseCountTexture = GraphBuilder.CreateTexture(Desc, TEXT("TSR.PrevUseCountTexture"));
-			PrevClosestDepthTexture = GraphBuilder.CreateTexture(Desc, TEXT("TSR.PrevClosestDepthTexture"));
+			PrevAtomicTextureArray = GraphBuilder.CreateTexture(Desc, TEXT("TSR.PrevAtomics"));
 		}
 
 		FTSRClearPrevTexturesCS::FParameters* PassParameters = GraphBuilder.AllocParameters<FTSRClearPrevTexturesCS::FParameters>();
 		PassParameters->CommonParameters = CommonParameters;
-		PassParameters->PrevUseCountOutput = GraphBuilder.CreateUAV(PrevUseCountTexture);
-		PassParameters->PrevClosestDepthOutput = GraphBuilder.CreateUAV(PrevClosestDepthTexture);
+		PassParameters->PrevAtomicOutput = GraphBuilder.CreateUAV(PrevAtomicTextureArray);
 
 		TShaderMapRef<FTSRClearPrevTexturesCS> ComputeShader(View.ShaderMap);
 		FComputeShaderUtils::AddPass(
@@ -1784,9 +1782,18 @@ FDefaultTemporalUpscaler::FOutputs AddTemporalSuperResolutionPasses(
 		}
 
 		{
+			EPixelFormat ClosestDepthFormat;
+			if (bIsOrthoProjection)
+			{
+				ClosestDepthFormat = bCanResurrectHistory ? PF_G32R32F : PF_R32_FLOAT;
+			}
+			else
+			{
+				ClosestDepthFormat = bCanResurrectHistory ? PF_G16R16F : PF_R16F;
+			}
 			FRDGTextureDesc Desc = FRDGTextureDesc::Create2D(
 				InputExtent,
-				bCanResurrectHistory ? PF_G16R16F : PF_R16F,
+				ClosestDepthFormat,
 				FClearValueBinding::None,
 				/* InFlags = */ TexCreate_ShaderResource | TexCreate_UAV);
 
@@ -1828,8 +1835,7 @@ FDefaultTemporalUpscaler::FOutputs AddTemporalSuperResolutionPasses(
 
 		PassParameters->DilatedVelocityOutput = GraphBuilder.CreateUAV(DilatedVelocityTexture);
 		PassParameters->ClosestDepthOutput = GraphBuilder.CreateUAV(ClosestDepthTexture);
-		PassParameters->PrevUseCountOutput = GraphBuilder.CreateUAV(PrevUseCountTexture);
-		PassParameters->PrevClosestDepthOutput = GraphBuilder.CreateUAV(PrevClosestDepthTexture);
+		PassParameters->PrevAtomicOutput = GraphBuilder.CreateUAV(PrevAtomicTextureArray);
 		PassParameters->R8Output = GraphBuilder.CreateUAV(R8OutputTexture);
 
 		// Setup up the motion blur's velocity flatten pass.
@@ -1932,8 +1938,7 @@ FDefaultTemporalUpscaler::FOutputs AddTemporalSuperResolutionPasses(
 		PassParameters->InputSceneColorTexture = PassInputs.SceneColor.Texture;
 		PassParameters->DilatedVelocityTexture = DilatedVelocityTexture;
 		PassParameters->ClosestDepthTexture = ClosestDepthTexture;
-		PassParameters->PrevUseCountTexture = PrevUseCountTexture;
-		PassParameters->PrevClosestDepthTexture = PrevClosestDepthTexture;
+		PassParameters->PrevAtomicTextureArray = PrevAtomicTextureArray;
 
 		PassParameters->PrevHistoryParameters = PrevHistoryParameters;
 
