VSM: disable screen ray and normal bias on subsurface materials
Can cause various artifacts with extreme opacity values

diff --git a/Engine/Shaders/Private/VirtualShadowMaps/VirtualShadowMapProjection.usf b/Engine/Shaders/Private/VirtualShadowMaps/VirtualShadowMapProjection.usf
--- a/Engine/Shaders/Private/VirtualShadowMaps/VirtualShadowMapProjection.usf
+++ b/Engine/Shaders/Private/VirtualShadowMaps/VirtualShadowMapProjection.usf
@@ -214,11 +214,6 @@ FVirtualShadowMapSampleResult ProjectLight(
 	const float Noise,
 	const FSubsurfaceOpacityMFP SubsurfaceOpacityMFP)
 {
-	const bool bIsHairInput = InputType == INPUT_TYPE_HAIRSTRANDS;
-
-	const float DistanceToCamera = GetDistanceToCameraFromViewVector(TranslatedWorldPosition - View.TranslatedWorldCameraOrigin);
-	const float NormalBiasLength = max(0.02f, NormalBias * DistanceToCamera / GetCotanHalfFieldOfView().x);
-
 #if DIRECTIONAL_LIGHT
 	float3 L = Light.Direction;
 	bool bInLightRegion = true;
@@ -244,6 +239,10 @@ FVirtualShadowMapSampleResult ProjectLight(
 	}
 #endif // DIRECTIONAL_LIGHT
 
+	const bool bIsHairInput = InputType == INPUT_TYPE_HAIRSTRANDS;
+	const float DistanceToCamera = GetDistanceToCameraFromViewVector(TranslatedWorldPosition - View.TranslatedWorldCameraOrigin);
+	const float NormalBiasLength = ShadingInfo.bIsSubsurface ? 0.0f : max(0.02f, NormalBias * DistanceToCamera / GetCotanHalfFieldOfView().x);
+
 	FVirtualShadowMapSampleResult Result = InitVirtualShadowMapSampleResult();
 
 	const bool bValidPixel = bIsHairInput || ShadingInfo.bIsValid;
@@ -264,7 +263,7 @@ FVirtualShadowMapSampleResult ProjectLight(
 		// Do not run contact shadow when computing the hair shadow mask (i.e. shadow mask applied on hair, has the scene 
 		// depth buffer contains fully opaque depth, which create false positive intersection resulting in wrong self shadowing)
 		float SMRTRayOffset = ScreenRayLengthWorld;
-		if (!bIsHairInput)
+		if (!bIsHairInput && !ShadingInfo.bIsSubsurface)
 		{
 			if (ScreenRayLengthWorld > 0.0f)
 			{
@@ -442,8 +441,6 @@ FVirtualShadowMapSampleResult ProjectLight(
 		#endif
 	}
 
-	FilterVirtualShadowMapSampleResult(PixelPos, Result);
-
 	return Result;
 }
 
@@ -670,6 +667,7 @@ void VirtualShadowMapProjection(
 				Noise,
 				SubsurfaceOpacityMFP);
 
+			FilterVirtualShadowMapSampleResult(PixelPos, Result);
 			PackShadowMask(ShadowMask, Result.ShadowFactor, Index);
 
 			if (VisualizeVirtualShadowMapId == VirtualShadowMapId)
diff --git a/Engine/Shaders/Private/VirtualShadowMaps/VirtualShadowMapProjectionCommon.ush b/Engine/Shaders/Private/VirtualShadowMaps/VirtualShadowMapProjectionCommon.ush
--- a/Engine/Shaders/Private/VirtualShadowMaps/VirtualShadowMapProjectionCommon.ush
+++ b/Engine/Shaders/Private/VirtualShadowMaps/VirtualShadowMapProjectionCommon.ush
@@ -477,9 +477,9 @@ FVirtualShadowMapSampleResult SampleVirtualShadowMapLocal(
 			Result.ShadowFactor = 0.0f;
 			// Shadow view matrix is rotation
 			float ReceiverDistance = length(ShadowTranslatedWorldPosition);
-			Result.OccluderDistance = RayStartDistance + ComputeOccluderDistancePerspective(
+			Result.OccluderDistance = ComputeOccluderDistancePerspective(
 				ProjectionData.ShadowViewToClipMatrix,
-				SmSample.Depth,
+				BiasedDepth,
 				ShadowUVz.z,
 				ReceiverDistance);
 		}
