Bumped material version. This fixes compile issues with stale materials using SVT nodes.

#rnx

diff --git a/Engine/Shaders/Private/MaterialTemplate.ush b/Engine/Shaders/Private/MaterialTemplate.ush
--- a/Engine/Shaders/Private/MaterialTemplate.ush
+++ b/Engine/Shaders/Private/MaterialTemplate.ush
@@ -41,7 +41,7 @@
 
 // Update this GUID to force all material shaders to recompile
 // Merge conflicts on this line should be resolved by generating a new GUID
-#pragma message("UESHADERMETADATA_VERSION 8899111D-F1E2-45E7-AB89-E1169159870B")
+#pragma message("UESHADERMETADATA_VERSION 1B0D44F1-B5B0-41FF-9E61-15BAF2C3A5EE")
 
 #if USES_SPEEDTREE
 	#include "/Engine/Private/SpeedTreeCommon.ush"
