Workaround for ShaderPrint compilation on Metal.

diff --git a/Engine/Shaders/Private/ShaderPrint.ush b/Engine/Shaders/Private/ShaderPrint.ush
--- a/Engine/Shaders/Private/ShaderPrint.ush
+++ b/Engine/Shaders/Private/ShaderPrint.ush
@@ -328,7 +328,7 @@ FShaderPrintContext Assert(bool InCondition, FShaderPrintText InText, uint2 InWr
 	if (!InCondition)
 	{
 		uint LineOffset = 0;
-		InterlockedAdd(SHADER_PRINT_RWENTRYBUFFER(OutCtx, 3) /* Free counter */, 1, LineOffset);
+		SHADER_PRINT_INTERLOCKEDADD(SHADER_PRINT_RWENTRYBUFFER(OutCtx, 3) /* Free counter */, 1, LineOffset);
 		InWritePos.y += LineOffset * 12;
 
 		OutCtx = InitShaderPrintContext(true, InWritePos);
@@ -344,7 +344,7 @@ FShaderPrintContext Assert(bool InCondition, FShaderPrintText InText, FShaderPri
 	if (!InCondition)
 	{
 		uint LineOffset = 0;
-		InterlockedAdd(SHADER_PRINT_RWENTRYBUFFER(OutCtx, 3) /* Free counter */, 1, LineOffset);
+		SHADER_PRINT_INTERLOCKEDADD(SHADER_PRINT_RWENTRYBUFFER(OutCtx, 3) /* Free counter */, 1, LineOffset);
 		InWritePos.y += LineOffset * 12;
 
 		OutCtx = InitShaderPrintContext(true, InWritePos);
diff --git a/Engine/Shaders/Private/ShaderPrintCommon.ush b/Engine/Shaders/Private/ShaderPrintCommon.ush
--- a/Engine/Shaders/Private/ShaderPrintCommon.ush
+++ b/Engine/Shaders/Private/ShaderPrintCommon.ush
@@ -165,6 +165,14 @@ void ClearCounters(RWBuffer<uint> InRWBuffer)
 #define SHADER_PRINT_RWSTATEBUFFER(Ctx, Idx) Ctx.Buffers.StateBuffer[Idx]
 #endif
 
+// Currently a bug Metal shader translation cause compilation to fails due to atomic operation.
+// As a workaround, disabling interlock add on Metal
+#if COMPILER_METAL
+#define SHADER_PRINT_INTERLOCKEDADD(InBuffer, InIncrement, OutOldValue) 
+#else
+#define SHADER_PRINT_INTERLOCKEDADD(InBuffer, InIncrement, OutOldValue) InterlockedAdd(InBuffer, InIncrement, OutOldValue)
+#endif
+
 ///////////////////////////////////////////////////////////////////////////////////////////////////
 // Font Color
 
@@ -373,7 +381,7 @@ void ShaderPrint_Internal(FShaderPrintContext Ctx, in FShaderPrintItem Item)
 
 	// Buffer counter is stored in first element .Value
 	int IndexToStore = 0;
-	InterlockedAdd(SHADER_PRINT_RWENTRYBUFFER(Ctx,SHADER_PRINT_COUNTER_OFFSET_SYMBOL), 1, IndexToStore);
+	SHADER_PRINT_INTERLOCKEDADD(SHADER_PRINT_RWENTRYBUFFER(Ctx,SHADER_PRINT_COUNTER_OFFSET_SYMBOL), 1, IndexToStore);
 
 	// Prevent writing off the buffer
 	// Note that counter still increases so need clamp when reading it in later passes
@@ -457,7 +465,7 @@ float2 ShaderPrintText_Internal(FShaderPrintContext Ctx, bool bIsActive, float2
 		// Prevent writing off the buffer
 		// Note that counter still increases so need clamp when reading it in later passes
 		int IndexToStore = 0;
-		InterlockedAdd(SHADER_PRINT_RWENTRYBUFFER(Ctx,SHADER_PRINT_COUNTER_OFFSET_SYMBOL), Count, IndexToStore);
+		SHADER_PRINT_INTERLOCKEDADD(SHADER_PRINT_RWENTRYBUFFER(Ctx,SHADER_PRINT_COUNTER_OFFSET_SYMBOL), Count, IndexToStore);
 		if (uint(IndexToStore + Count) < Ctx.Config.MaxCharacterCount)
 		{
 			for (uint i = Begin; i < End; ++i)
@@ -797,7 +805,7 @@ bool AllocateTriangleElement(FShaderPrintContext Ctx, uint Count, inout uint Out
 		return false;
 	}
 
-	InterlockedAdd(SHADER_PRINT_RWENTRYBUFFER(Ctx,SHADER_PRINT_COUNTER_OFFSET_TRIANGLE), Count, OutIndex);
+	SHADER_PRINT_INTERLOCKEDADD(SHADER_PRINT_RWENTRYBUFFER(Ctx,SHADER_PRINT_COUNTER_OFFSET_TRIANGLE), Count, OutIndex);
 	return (OutIndex + Count) < Ctx.Config.MaxTriangleCount;
 }
 
@@ -996,7 +1004,7 @@ bool AllocateLineElement(FShaderPrintContext Ctx, uint Count, inout uint OutInde
 		return false;
 	}
 
-	InterlockedAdd(SHADER_PRINT_RWENTRYBUFFER(Ctx,SHADER_PRINT_COUNTER_OFFSET_LINE), Count, OutIndex);
+	SHADER_PRINT_INTERLOCKEDADD(SHADER_PRINT_RWENTRYBUFFER(Ctx,SHADER_PRINT_COUNTER_OFFSET_LINE), Count, OutIndex);
 	return (OutIndex + Count) < Ctx.Config.MaxLineCount;
 }
 
