Include Engine/Shared/RayTracingTypes.h in ManyLightsHardwareRayTracing.usf

diff --git a/Engine/Shaders/Private/ManyLights/ManyLightsHardwareRayTracing.usf b/Engine/Shaders/Private/ManyLights/ManyLightsHardwareRayTracing.usf
--- a/Engine/Shaders/Private/ManyLights/ManyLightsHardwareRayTracing.usf
+++ b/Engine/Shaders/Private/ManyLights/ManyLightsHardwareRayTracing.usf
@@ -1,5 +1,6 @@
 // Copyright Epic Games, Inc. All Rights Reserved.
 
+#include "/Engine/Shared/RayTracingTypes.h"
 #include "../Common.ush"
 #include "../BlueNoise.ush"
 #include "../RayTracing/RayTracingCommon.ush"
