Fix hair interpolation issue.

When unpacking the data, the wrong variation was picked, causing the read data to be sometime invalid, and fetching out of bound value.

diff --git a/Engine/Plugins/Runtime/HairStrands/Source/HairStrandsCore/Private/HairStrandsDatas.cpp b/Engine/Plugins/Runtime/HairStrands/Source/HairStrandsCore/Private/HairStrandsDatas.cpp
--- a/Engine/Plugins/Runtime/HairStrands/Source/HairStrandsCore/Private/HairStrandsDatas.cpp
+++ b/Engine/Plugins/Runtime/HairStrands/Source/HairStrandsCore/Private/HairStrandsDatas.cpp
@@ -689,9 +689,7 @@ float FHairStrandsBulkData::GetCoverageScale(float InCurvePercentage /*[0..1]*/)
 
 void FHairStrandsInterpolationBulkData::Reset()
 {
-	Header.Flags = 0;
-	Header.PointCount = 0;
-	Header.CurveCount = 0;
+	Header = FHeader();
 	
 	// Deallocate memory if needed
 	Data.CurveInterpolation.RemoveBulkData();
diff --git a/Engine/Plugins/Runtime/HairStrands/Source/HairStrandsCore/Private/HairStrandsInterpolation.cpp b/Engine/Plugins/Runtime/HairStrands/Source/HairStrandsCore/Private/HairStrandsInterpolation.cpp
--- a/Engine/Plugins/Runtime/HairStrands/Source/HairStrandsCore/Private/HairStrandsInterpolation.cpp
+++ b/Engine/Plugins/Runtime/HairStrands/Source/HairStrandsCore/Private/HairStrandsInterpolation.cpp
@@ -54,8 +54,16 @@ static FAutoConsoleVariableRef CVarHairRTGeomForceRebuild(TEXT("r.HairStrands.St
 static int32 GHairStrands_MotionVector_CheckViewID = 1;
 static FAutoConsoleVariableRef CVarHairStrands_MotionVector_CheckViewID(TEXT("r.HairStrands.Strands.MotionVectorCheckViewID"), GHairStrands_MotionVector_CheckViewID, TEXT("Issue motion vector on hair strands only if updates happens on the same view"));
 
+static int32 GStrandHairInterpolationForceSingleGuide = 1;
+static FAutoConsoleVariableRef CVarHairStrands_Interpolation_ForceSingleGuide(TEXT("r.HairStrands.Strands.Interpolation.ForceSingleGuide"), GStrandHairInterpolationForceSingleGuide, TEXT("Force to use a single guide during the interpolation. This allows to reduce the cost of interpolation, but can create clumpy looks."));
+
 ////////////////////////////////////////////////////////////////////////////////////////////////////////////////////////
 
+uint32 GetHairStrandsForceSingleGuideInterpolation()
+{
+	return GStrandHairInterpolationForceSingleGuide > 0;
+}
+
 uint32 GetHairCardsInterpolationType()
 {
 	return FMath::Clamp(GHairCardsInterpolationType, 0, 2);
@@ -451,7 +459,7 @@ class FHairInterpolationCS : public FGlobalShader
 	SHADER_USE_PARAMETER_STRUCT(FHairInterpolationCS, FGlobalShader);
 
 	class FPointPerCurve : SHADER_PERMUTATION_SPARSE_INT("PERMUTATION_POINT_PER_CURVE", 4, 8, 16, 32);
-	class FDynamicGeometry : SHADER_PERMUTATION_INT("PERMUTATION_DYNAMIC_GEOMETRY", 5);
+	class FDynamicGeometry : SHADER_PERMUTATION_INT("PERMUTATION_DYNAMIC_GEOMETRY", 4); /* These is actually 5 DynamicGeometry permutations, but PERMUTATION_DYNAMIC_GEOMETRY==4 is never used (experimental) */
 	class FSimulation : SHADER_PERMUTATION_BOOL("PERMUTATION_SIMULATION");
 	class FSingleGuide : SHADER_PERMUTATION_BOOL("PERMUTATION_USE_SINGLE_GUIDE");
 	class FDeformer : SHADER_PERMUTATION_BOOL("PERMUTATION_DEFORMER");
@@ -461,12 +469,16 @@ class FHairInterpolationCS : public FGlobalShader
 		SHADER_PARAMETER_STRUCT_INCLUDE(ShaderPrint::FShaderParameters, ShaderDrawParameters)
 		SHADER_PARAMETER(uint32, VertexCount)
 		SHADER_PARAMETER(uint32, CurveCount)
+		SHADER_PARAMETER(uint32, SimPointCount)
+		SHADER_PARAMETER(uint32, SimCurveCount)
+
 		SHADER_PARAMETER(float, HairLengthScale)
 		SHADER_PARAMETER(FVector3f, InRenHairPositionOffset)
 		SHADER_PARAMETER(FVector3f, InSimHairPositionOffset)
 		SHADER_PARAMETER(FMatrix44f, LocalToWorldMatrix)
-		SHADER_PARAMETER(uint32,  DispatchCountX)
-	
+		SHADER_PARAMETER(uint32, DispatchCountX)
+		SHADER_PARAMETER(uint32, bUseSingleGuide)
+
 		SHADER_PARAMETER_RDG_BUFFER_SRV(ByteAddressBuffer, RenCurveBuffer)
 
 		SHADER_PARAMETER_RDG_BUFFER_SRV(ByteAddressBuffer, RenRestPosePositionBuffer)
@@ -575,7 +587,16 @@ void AddHairStrandsInterpolationPass(
 	Parameters->LocalToWorldMatrix = FMatrix44f(Instance->LocalToWorld.ToMatrixWithScale());		// LWC_TODO: Precision loss
 	Parameters->HairLengthScale = HairLengthScale;
 
+	Parameters->SimPointCount = 0;
+	Parameters->SimCurveCount = 0;
+	if (Instance->Guides.RestResource)
+	{		
+		Parameters->SimPointCount = Instance->Guides.RestResource->BulkData.Header.PointCount;
+		Parameters->SimCurveCount = Instance->Guides.RestResource->BulkData.Header.CurveCount;
+	}
+
 	Parameters->RenCurveBuffer = RenCurveBuffer;
+	Parameters->bUseSingleGuide = bUseSingleGuide;
 
 	if (ShaderPrintData)
 	{
@@ -650,16 +671,18 @@ void AddHairStrandsInterpolationPass(
 	const FPointPerCurveDispatchInfo DispatchInfo = GetPointPerCurveDispatchInfo(MaxPointPerCurve, CurveCount, FHairInterpolationCS::GetGroupSize());
 	Parameters->DispatchCountX = DispatchInfo.DispatchCount.X;
 
+	const bool bSingleGuidePermutation = bUseSingleGuide || GetHairStrandsForceSingleGuideInterpolation();
+
 	FHairInterpolationCS::FPermutationDomain PermutationVector;
 	PermutationVector.Set<FHairInterpolationCS::FDynamicGeometry>(DynamicGeometryType);
 	PermutationVector.Set<FHairInterpolationCS::FSimulation>(bHasLocalDeformation);
-	PermutationVector.Set<FHairInterpolationCS::FSingleGuide>(bUseSingleGuide);
+	PermutationVector.Set<FHairInterpolationCS::FSingleGuide>(bSingleGuidePermutation);
 	PermutationVector.Set<FHairInterpolationCS::FDeformer>(bSupportDeformer);
 	PermutationVector.Set<FHairInterpolationCS::FPointPerCurve>(DispatchInfo.PointPerCurve);
 	TShaderMapRef<FHairInterpolationCS> ComputeShader(ShaderMap, PermutationVector);
 	FComputeShaderUtils::AddPass(
 		GraphBuilder,
-		RDG_EVENT_NAME("HairStrands::Interpolation(culling=off, PerCurve=%d, SingleGuide=%d, DynGeom:%d)", DispatchInfo.PointPerCurve, bUseSingleGuide ? 1u : 0u, DynamicGeometryType > 0 ? 1u : 0u),
+		RDG_EVENT_NAME("HairStrands::Interpolation(culling=off, PerCurve=%d, SingleGuide=%d, DynGeom:%d)", DispatchInfo.PointPerCurve, bSingleGuidePermutation ? 1u : 0u, DynamicGeometryType > 0 ? 1u : 0u),
 		ComputeShader,
 		Parameters,
 		DispatchInfo.DispatchCount);
diff --git a/Engine/Shaders/Private/HairStrands/HairStrandsInterpolation.usf b/Engine/Shaders/Private/HairStrands/HairStrandsInterpolation.usf
--- a/Engine/Shaders/Private/HairStrands/HairStrandsInterpolation.usf
+++ b/Engine/Shaders/Private/HairStrands/HairStrandsInterpolation.usf
@@ -22,9 +22,6 @@
 
 #pragma warning (default:7203) // downgrade 'not enough registers available for the entire program.' to a warning instead of error
 
-// Currently disable multi-guide interpolation to reduce VGPR pressure and increase occupancy
-#define FORCE_USE_SINGLE_GUIDE 1
-
 ///////////////////////////////////////////////////////////////////////////////////////////////////
 // Guides
 
@@ -60,6 +57,7 @@ FGuideDataWithOffset GetGuideData(FGuideCurveData InCurveData, FGuidePointData I
 ///////////////////////////////////////////////////////////////////////////////////////////////////
 // Hair interpolation 
 
+// Interpolation methods
 #define INTERPOLATION_RIGID 0
 #define INTERPOLATION_SKINNING_OFFSET 1
 #define INTERPOLATION_SKINNING_TRANSLATION 2
@@ -79,6 +77,9 @@ float HairLengthScale;
 float3 InRenHairPositionOffset;
 float3 InSimHairPositionOffset;
 float4x4 LocalToWorldMatrix;
+uint bUseSingleGuide;
+uint SimPointCount;
+uint SimCurveCount;
 
 ByteAddressBuffer RenRestPosePositionBuffer;
 ByteAddressBuffer SimRestPosePositionBuffer;
@@ -239,7 +240,7 @@ float3 ApplyRigidInterpolation(FGuideDataWithOffset GuideData, float3 InPosition
 {
 	float3 CurrOffset = 0;
 #if PERMUTATION_SIMULATION == 1
-	#if PERMUTATION_USE_SINGLE_GUIDE == 1 || FORCE_USE_SINGLE_GUIDE
+	#if PERMUTATION_USE_SINGLE_GUIDE == 1
 	const uint K = 0;
 	#else
 	[unroll]
@@ -264,7 +265,7 @@ float3 ApplySkinningWithOffset(const FGuideDataWithOffset GuideData, const FRoot
 
 	// Compute the simulation offset in hair local space (i.e., triangle)
 #if PERMUTATION_SIMULATION == 1
-	#if PERMUTATION_USE_SINGLE_GUIDE == 1 || FORCE_USE_SINGLE_GUIDE
+	#if PERMUTATION_USE_SINGLE_GUIDE == 1
 	const uint K = 0;
 	#else
 	[unroll]
@@ -295,7 +296,7 @@ float3 ApplySkinningWithTranslation(const FGuideDataWithOffset GuideData, const
 	float3 CurrOffset = 0;
 
 #if PERMUTATION_SIMULATION == 1
-#if PERMUTATION_USE_SINGLE_GUIDE == 1 || FORCE_USE_SINGLE_GUIDE
+#if PERMUTATION_USE_SINGLE_GUIDE == 1
 	const uint K = 0;
 	#else
 	[unroll]
@@ -322,7 +323,7 @@ float3 ApplySkinningWithRotation(const FGuideDataWithOffset GuideData, const FRo
 	float3 CurrOffset = 0;
 
 #if PERMUTATION_SIMULATION == 1
-	#if PERMUTATION_USE_SINGLE_GUIDE == 1 || FORCE_USE_SINGLE_GUIDE
+	#if PERMUTATION_USE_SINGLE_GUIDE == 1
 	const uint K = 0;
 	#else
 	[unroll]
@@ -351,7 +352,7 @@ float3 ApplySkinningWithTranslationAndRotation(const FGuideDataWithOffset GuideD
 	float3 CurrOffset = 0;
 
 #if PERMUTATION_SIMULATION == 1
-	#if PERMUTATION_USE_SINGLE_GUIDE == 1 || FORCE_USE_SINGLE_GUIDE
+	#if PERMUTATION_USE_SINGLE_GUIDE == 1
 	const uint K = 0;
 	#else
 	[unroll]
@@ -449,8 +450,8 @@ void MainCS(uint3 DispatchThreadId : SV_DispatchThreadID, uint ThreadIndex : SV_
 
 			// Load curve interpolation data
 			#if PERMUTATION_SIMULATION
-			SimCurveData[LocalCurveIndex] = UnpackGuideCurveData(CurveInterpolationBuffer, GlobalCurveIndex, PERMUTATION_USE_SINGLE_GUIDE);
-			if (FORCE_USE_SINGLE_GUIDE)
+			SimCurveData[LocalCurveIndex] = UnpackGuideCurveData(CurveInterpolationBuffer, GlobalCurveIndex, bUseSingleGuide);
+			if (PERMUTATION_USE_SINGLE_GUIDE)
 			{
 				SimCurveData[LocalCurveIndex].CurveWeights.x = 1;
 				SimCurveData[LocalCurveIndex].CurveWeights.y = 0;
@@ -496,7 +497,7 @@ void MainCS(uint3 DispatchThreadId : SV_DispatchThreadID, uint ThreadIndex : SV_
 				const uint VertexIndex = CurrentPointIndex + CurvePointIndex[LocalCurveIndex];
 
 				#if PERMUTATION_SIMULATION
-				const FGuidePointData SimPointData = UnpackGuidePointData(PointInterpolationBuffer, VertexIndex, PERMUTATION_USE_SINGLE_GUIDE);
+				const FGuidePointData SimPointData = UnpackGuidePointData(PointInterpolationBuffer, VertexIndex, bUseSingleGuide);
 				const FGuideDataWithOffset GuideData = GetGuideData(SimCurveData[LocalCurveIndex], SimPointData, GetSimPositionOffset());
 				#else
 				const FGuidePointData SimPointData = (FGuidePointData)0;
@@ -613,8 +614,8 @@ void MainCS(uint3 DispatchThreadId : SV_DispatchThreadID, uint ThreadIndex : SV_
 
 			// Load curve interpolation data
 			#if PERMUTATION_SIMULATION
-			SimCurveData = UnpackGuideCurveData(CurveInterpolationBuffer, CurveIndex, PERMUTATION_USE_SINGLE_GUIDE);
-			if (FORCE_USE_SINGLE_GUIDE)
+			SimCurveData = UnpackGuideCurveData(CurveInterpolationBuffer, CurveIndex, bUseSingleGuide);
+			if (PERMUTATION_USE_SINGLE_GUIDE)
 			{
 				SimCurveData.CurveWeights.x = 1;
 				SimCurveData.CurveWeights.y = 0;
@@ -629,7 +630,7 @@ void MainCS(uint3 DispatchThreadId : SV_DispatchThreadID, uint ThreadIndex : SV_
 			#if PERMUTATION_SIMULATION
 			SimTris[0] = GetRootTriangleData(SimCurveData.CurveIndices[0], true/*bSim*/);
 			SimTris[1] = SimTris[0];
-			if (!PERMUTATION_USE_SINGLE_GUIDE && !FORCE_USE_SINGLE_GUIDE)
+			if (!PERMUTATION_USE_SINGLE_GUIDE)
 			{
 				SimTris[1] = GetRootTriangleData(SimCurveData.CurveIndices[1], true/*bSim*/);
 			}
@@ -654,7 +655,7 @@ void MainCS(uint3 DispatchThreadId : SV_DispatchThreadID, uint ThreadIndex : SV_
 				const uint VertexIndex = CurrentPointIndex + CurvePointIndex;
 
 				#if PERMUTATION_SIMULATION
-				const FGuidePointData SimPointData = UnpackGuidePointData(PointInterpolationBuffer, VertexIndex, PERMUTATION_USE_SINGLE_GUIDE);
+				const FGuidePointData SimPointData = UnpackGuidePointData(PointInterpolationBuffer, VertexIndex, bUseSingleGuide);
 				const FGuideDataWithOffset GuideData = GetGuideData(SimCurveData, SimPointData, GetSimPositionOffset());
 				#else
 				const FGuidePointData SimPointData = (FGuidePointData)0;
