HeterogeneousVolumes: Fixing compositing issue with multiple HVs and Height Fog or Volumetric Fog

Changed fog distance from t-min to the first scattering event.

diff --git a/Engine/Shaders/Private/HeterogeneousVolumes/HeterogeneousVolumesRayMarchingUtils.ush b/Engine/Shaders/Private/HeterogeneousVolumes/HeterogeneousVolumesRayMarchingUtils.ush
--- a/Engine/Shaders/Private/HeterogeneousVolumes/HeterogeneousVolumesRayMarchingUtils.ush
+++ b/Engine/Shaders/Private/HeterogeneousVolumes/HeterogeneousVolumesRayMarchingUtils.ush
@@ -311,6 +311,11 @@ void RayMarchSingleScattering(
 	inout float3 Transmittance
 )
 {
+#if APPLY_FOG_INSCATTERING
+	float LocalScatterTMin = RayMarchingContext.LocalRayTMax;
+	float3 PrevTransmittance = Transmittance;
+#endif // #if APPLY_FOG_INSCATTERING
+
 	float WorldShadowBias = CalcShadowBias();
 
 	for (uint StepIndex = 0; StepIndex < StepCount; ++StepIndex)
@@ -355,6 +360,10 @@ void RayMarchSingleScattering(
 				float IsotropicPhase = 1.0 / (4.0 * PI);
 				Radiance += Inscattering * ScatteringCoefficient * IsotropicPhase * Transmittance * RayMarchingContext.StepSize;
 			}
+
+#if APPLY_FOG_INSCATTERING
+			LocalScatterTMin = min(LocalHitT, LocalScatterTMin);
+#endif 
 		}
 
 		// Accumulate transmittance for the next evaluation
@@ -369,12 +378,13 @@ void RayMarchSingleScattering(
 	}
 
 #if APPLY_FOG_INSCATTERING
-	if (any(Transmittance < 1))
+	bool bApplyFog = LocalScatterTMin < RayMarchingContext.LocalRayTMax;
+	if (bApplyFog)
 	{
 		// TODO: Change in-scattering distance to be based on the first scattering event, instead of the
 		// intersection distance of the bounding box near-plane.
-		float RayT = RayMarchingContext.LocalRayTMin * RayMarchingContext.LocalToWorldScale;
-		float3 CameraRelative_WorldPosition = RayMarchingContext.WorldRayDirection * RayT;
+		float WorldRayT = LocalScatterTMin * RayMarchingContext.LocalToWorldScale;
+		float3 CameraRelative_WorldPosition = RayMarchingContext.WorldRayDirection * WorldRayT;
 		float4 HeightFog = float4(0, 0, 0, 1);
 		if (ShouldApplyHeightFog())
 		{
@@ -387,14 +397,19 @@ void RayMarchSingleScattering(
 			float3 WorldPosition = CameraRelative_WorldPosition + RayMarchingContext.WorldRayOrigin;
 			const uint EyeIndex = 0;
 			float3 VolumeUV = ComputeVolumeUV_DEPRECATED(WorldPosition, DFHackToFloat(PrimaryView.WorldToClip));
-			CombinedFog = CombineVolumetricFog(HeightFog, VolumeUV, EyeIndex, RayT);
+			CombinedFog = CombineVolumetricFog(HeightFog, VolumeUV, EyeIndex, WorldRayT);
 		}
 
+		// Simplifying terms to avoid vector division:
+		// float3 DeltaTransmittance = clamp(Transmittance / PrevTransmittance, 0.0, 1.0);
+		// float3 DeltaOpacity = 1.0 - DeltaTransmittance;
+		// float3 FogAlpha = PrevTransmittance * DeltaOpacity;
+		float3 FogAlpha = clamp(PrevTransmittance - Transmittance, 0.0, 1.0);
+
 		// Fogging in-scattering contribution is not currently attenuated by Heterogeneous Volumes
 		// for elements behind the subject. To prevent doubled contribution of in-scattering for 
 		// regions of low density (eg: silhouette edges), blending is interpreted as alpha-masking instead.
-		float Opacity = 1.0 - Transmittance;
-		Radiance = Radiance * CombinedFog.a + CombinedFog.rgb * Opacity;
+		Radiance = Radiance * CombinedFog.a + CombinedFog.rgb * FogAlpha;
 	}
 #endif // APPLY_FOG_INSCATTERING
 }
