Fix incorrect loading of actor and object world position in mobile scene data loader/writer.
Compact primitive data in primitive uniform shader parameter data.

diff --git a/Engine/Shaders/Private/SceneData.ush b/Engine/Shaders/Private/SceneData.ush
--- a/Engine/Shaders/Private/SceneData.ush
+++ b/Engine/Shaders/Private/SceneData.ush
@@ -112,6 +112,8 @@ struct FPrimitiveSceneData
 	float		ObjectBoundsX;
 	FDFVector3	ObjectWorldPosition;
 	FDFVector3	ActorWorldPosition;
+	float		MinMaterialDisplacement;
+	float		MaxMaterialDisplacement;
 #if PRIMITIVE_HAS_TILEOFFSET_DATA
 	FLWCVector3	ObjectWorldPositionTO;
 	FLWCVector3	ActorWorldPositionTO;
@@ -141,8 +143,6 @@ struct FPrimitiveSceneData
 	float		InstanceWPODisableDistanceSquared;
 	uint		NaniteRayTracingDataOffset;
 	float		MaxWPOExtent;
-	float		MinMaterialDisplacement;
-	float		MaxMaterialDisplacement;
 	uint		CustomStencilValueAndMask;
 	float4		CustomPrimitiveData[NUM_CUSTOM_PRIMITIVE_DATA]; // TODO: Move to associated array to shrink primitive data and pack cachelines more effectively
 };
@@ -428,6 +428,8 @@ FPrimitiveSceneData GetPrimitiveData(uint PrimitiveId)
 	PrimitiveData.ObjectWorldPosition					= MakeDFVector3(ObjectWorldPositionHighAndRadius.xyz, ObjectWorldPositionLow.xyz);
 	PrimitiveData.ActorWorldPosition					= MakeDFVector3(ActorWorldPositionHigh.xyz, ActorWorldPositionLow.xyz);
 #endif
+	PrimitiveData.MinMaterialDisplacement				= LoadPrimitivePrimitiveSceneDataElement(PrimitiveIndex, 19).w;
+	PrimitiveData.MaxMaterialDisplacement				= LoadPrimitivePrimitiveSceneDataElement(PrimitiveIndex, 20).w;
 	PrimitiveData.LightmapUVIndex						= asuint(LoadPrimitivePrimitiveSceneDataElement(PrimitiveIndex, 21).w);
 	PrimitiveData.ObjectOrientation						= LoadPrimitivePrimitiveSceneDataElement(PrimitiveIndex, 22).xyz;
 	PrimitiveData.LightmapDataIndex						= asuint(LoadPrimitivePrimitiveSceneDataElement(PrimitiveIndex, 22).w);
@@ -453,9 +455,7 @@ FPrimitiveSceneData GetPrimitiveData(uint PrimitiveId)
 	PrimitiveData.InstanceWPODisableDistanceSquared		= LoadPrimitivePrimitiveSceneDataElement(PrimitiveIndex, 31).z;
 	PrimitiveData.NaniteRayTracingDataOffset			= asuint(LoadPrimitivePrimitiveSceneDataElement(PrimitiveIndex, 31).w);
 	PrimitiveData.MaxWPOExtent							= LoadPrimitivePrimitiveSceneDataElement(PrimitiveIndex, 32).x;
-	PrimitiveData.MinMaterialDisplacement				= LoadPrimitivePrimitiveSceneDataElement(PrimitiveIndex, 32).y;
-	PrimitiveData.MaxMaterialDisplacement				= LoadPrimitivePrimitiveSceneDataElement(PrimitiveIndex, 32).z;
-	PrimitiveData.CustomStencilValueAndMask				= asuint(LoadPrimitivePrimitiveSceneDataElement(PrimitiveIndex, 32).w);
+	PrimitiveData.CustomStencilValueAndMask				= asuint(LoadPrimitivePrimitiveSceneDataElement(PrimitiveIndex, 32).y);
 
 	// NOTE: Please make sure GetPrimitiveDataFromUniformBuffer() gets updated as well when adding new members here
 
diff --git a/Engine/Shaders/Private/SceneDataMobileLoader.ush b/Engine/Shaders/Private/SceneDataMobileLoader.ush
--- a/Engine/Shaders/Private/SceneDataMobileLoader.ush
+++ b/Engine/Shaders/Private/SceneDataMobileLoader.ush
@@ -5,6 +5,10 @@
 	#define PLATFORM_USES_PRIMITIVE_UBO 0
 #endif
 
+#ifndef PRIMITIVE_HAS_TILEOFFSET_DATA
+#error "Missing PRIMITIVE_HAS_TILEOFFSET_DATA"
+#endif
+
 float4 LoadPrimitiveDataElementUBO(uint DrawInstanceId, uint DataIdx)
 {
 #if PLATFORM_USES_PRIMITIVE_UBO
@@ -62,7 +66,12 @@ FPrimitiveSceneData LoadPrimitiveDataUBO(uint DrawInstanceId)
 	{
 		float4 Data = LoadPrimitiveDataElementUBO(DrawInstanceId, DataIdx+0);
 		float4 Data2 = LoadPrimitiveDataElementUBO(DrawInstanceId, DataIdx+1);
+		#if PRIMITIVE_HAS_TILEOFFSET_DATA
+		PrimitiveData.ObjectWorldPositionTO = MakeLWCVector3(Data.xyz, Data2.xyz);
+		PrimitiveData.ObjectWorldPosition = DFFromTileOffset(PrimitiveData.ObjectWorldPositionTO);
+		#else
 		PrimitiveData.ObjectWorldPosition = MakeDFVector3(Data.xyz, Data2.xyz);
+		#endif
 		PrimitiveData.ObjectRadius = Data.w;
 		DataIdx += 2u;
 	}
@@ -71,7 +80,12 @@ FPrimitiveSceneData LoadPrimitiveDataUBO(uint DrawInstanceId)
 	{
 		float4 Data = LoadPrimitiveDataElementUBO(DrawInstanceId, DataIdx+0);
 		float4 Data2 = LoadPrimitiveDataElementUBO(DrawInstanceId, DataIdx+1);
+		#if PRIMITIVE_HAS_TILEOFFSET_DATA
+		PrimitiveData.ActorWorldPositionTO = MakeLWCVector3(Data.xyz, Data2.xyz);
+		PrimitiveData.ActorWorldPosition = DFFromTileOffset(PrimitiveData.ActorWorldPositionTO);
+		#else
 		PrimitiveData.ActorWorldPosition = MakeDFVector3(Data.xyz, Data2.xyz);
+		#endif
 		DataIdx += 2u;
 	}
 
diff --git a/Engine/Shaders/Private/SceneDataMobileWriter.ush b/Engine/Shaders/Private/SceneDataMobileWriter.ush
--- a/Engine/Shaders/Private/SceneDataMobileWriter.ush
+++ b/Engine/Shaders/Private/SceneDataMobileWriter.ush
@@ -5,6 +5,10 @@
 	#define PLATFORM_USES_PRIMITIVE_UBO 0
 #endif
 
+#ifndef PRIMITIVE_HAS_TILEOFFSET_DATA
+#error "Missing PRIMITIVE_HAS_TILEOFFSET_DATA"
+#endif
+
 void WriteDataElementUBO(RWStructuredBuffer<float4> Output, uint Offset, uint DataIdx, float4 Data)
 {
 #if	PLATFORM_USES_PRIMITIVE_UBO
@@ -64,15 +68,29 @@ void WritePrimitiveDataUBO(RWStructuredBuffer<float4> Output, uint Offset, FPrim
 
 	// ObjectWorldPosition, Radius
 	{
-		WriteDataElementUBO(Output, Offset, DataIdx + 0, float4(PrimitiveData.ObjectWorldPosition.High, PrimitiveData.ObjectRadius));
-		WriteDataElementUBO(Output, Offset, DataIdx + 0, float4(PrimitiveData.ObjectWorldPosition.Low, 0));
+		#if PRIMITIVE_HAS_TILEOFFSET_DATA
+		float3 Data = PrimitiveData.ObjectWorldPositionTO.Tile;
+		float3 Data2 = PrimitiveData.ObjectWorldPositionTO.Offset;
+		#else
+		float3 Data = PrimitiveData.ObjectWorldPosition.High;
+		float3 Data2 = PrimitiveData.ObjectWorldPosition.Low;
+		#endif
+		WriteDataElementUBO(Output, Offset, DataIdx + 0, float4(Data, PrimitiveData.ObjectRadius));
+		WriteDataElementUBO(Output, Offset, DataIdx + 1, float4(Data2, 0));
 		DataIdx += 2u;
 	}
 
 	// ActorWorldPosition, TODO .w
 	{
-		WriteDataElementUBO(Output, Offset, DataIdx + 0, float4(PrimitiveData.ActorWorldPosition.High, 0.0));
-		WriteDataElementUBO(Output, Offset, DataIdx + 0, float4(PrimitiveData.ActorWorldPosition.Low, 0.0));
+		#if PRIMITIVE_HAS_TILEOFFSET_DATA
+		float3 Data = PrimitiveData.ActorWorldPositionTO.Tile;
+		float3 Data2 = PrimitiveData.ActorWorldPositionTO.Offset;
+		#else
+		float3 Data = PrimitiveData.ActorWorldPosition.High;
+		float3 Data2 = PrimitiveData.ActorWorldPosition.Low;
+		#endif
+		WriteDataElementUBO(Output, Offset, DataIdx + 0, float4(Data, 0.0));
+		WriteDataElementUBO(Output, Offset, DataIdx + 1, float4(Data2, 0.0));
 		DataIdx += 2u;
 	}
 
diff --git a/Engine/Shaders/Public/ShaderVersion.ush b/Engine/Shaders/Public/ShaderVersion.ush
--- a/Engine/Shaders/Public/ShaderVersion.ush
+++ b/Engine/Shaders/Public/ShaderVersion.ush
@@ -3,4 +3,4 @@
 // Each time the console command is executed it generates a new GUID. As this file is included
 // in Platform.ush (which should be included in any shader) it allows to invalidate the shader DDC.
 // 
-#pragma message("UESHADERMETADATA_VERSION DC1DB8B4-F092-46B2-8B29-0A59A08B8D93")
+#pragma message("UESHADERMETADATA_VERSION A4430A10-A7A1-40B0-B194-763FCF1CAFFE")
diff --git a/Engine/Source/Runtime/Engine/Private/PrimitiveUniformShaderParameters.cpp b/Engine/Source/Runtime/Engine/Private/PrimitiveUniformShaderParameters.cpp
--- a/Engine/Source/Runtime/Engine/Private/PrimitiveUniformShaderParameters.cpp
+++ b/Engine/Source/Runtime/Engine/Private/PrimitiveUniformShaderParameters.cpp
@@ -202,9 +202,8 @@ void FPrimitiveSceneShaderData::Setup(const FPrimitiveUniformShaderParameters& P
 
 	OutData[17]	= FVector4f(PrimitiveUniformShaderParameters.InvNonUniformScale, PrimitiveUniformShaderParameters.ObjectBoundsX);
 	OutData[18]	= PrimitiveUniformShaderParameters.ObjectWorldPositionHighAndRadius;
-	OutData[19]	= PrimitiveUniformShaderParameters.ObjectWorldPositionLow;
-
-	OutData[20]	= FVector4f(PrimitiveUniformShaderParameters.ActorWorldPositionHigh, 0.0f);
+	OutData[19]	= FVector4f(PrimitiveUniformShaderParameters.ObjectWorldPositionLow, PrimitiveUniformShaderParameters.MinMaterialDisplacement);
+	OutData[20]	= FVector4f(PrimitiveUniformShaderParameters.ActorWorldPositionHigh, PrimitiveUniformShaderParameters.MaxMaterialDisplacement);
 
 	Store4(OutData, 21, PrimitiveUniformShaderParameters.ActorWorldPositionLow, PrimitiveUniformShaderParameters.LightmapUVIndex);
 	Store4(OutData, 22, PrimitiveUniformShaderParameters.ObjectOrientation, PrimitiveUniformShaderParameters.LightmapDataIndex);
@@ -231,12 +230,8 @@ void FPrimitiveSceneShaderData::Setup(const FPrimitiveUniformShaderParameters& P
 		PrimitiveUniformShaderParameters.InstanceDrawDistanceMinMaxSquared.Y, 
 		PrimitiveUniformShaderParameters.InstanceWPODisableDistanceSquared, 
 		PrimitiveUniformShaderParameters.NaniteRayTracingDataOffset);
-	
-	Store4(OutData, 32,
-		PrimitiveUniformShaderParameters.MaxWPOExtent,
-		PrimitiveUniformShaderParameters.MinMaterialDisplacement,
-		PrimitiveUniformShaderParameters.MaxMaterialDisplacement,
-		PrimitiveUniformShaderParameters.CustomStencilValueAndMask);
+
+	OutData[32]	= FVector4f(PrimitiveUniformShaderParameters.MaxWPOExtent, PrimitiveUniformShaderParameters.CustomStencilValueAndMask, 0.0, 0.0);
 
 	// Set all the custom primitive data float4. This matches the loop in SceneData.ush
 	const int32 CustomPrimitiveDataStartIndex = 33;
@@ -287,4 +282,4 @@ FPrimitiveUniformShaderParameters GetIdentityPrimitiveParameters()
 			.WorldBounds(FBoxSphereBounds(EForceInit::ForceInit))
 			.LocalBounds(FBoxSphereBounds(EForceInit::ForceInit))
 		.Build();
-}
\ No newline at end of file
+}
diff --git a/Engine/Source/Runtime/Renderer/Private/GPUScene.cpp b/Engine/Source/Runtime/Renderer/Private/GPUScene.cpp
--- a/Engine/Source/Runtime/Renderer/Private/GPUScene.cpp
+++ b/Engine/Source/Runtime/Renderer/Private/GPUScene.cpp
@@ -2324,7 +2324,7 @@ void FBatchedPrimitiveShaderData::Setup(const FPrimitiveUniformShaderParameters&
 	// ObjectWorldPosition, Radius
 	{
 		Data[i+0] = PrimitiveUniformShaderParameters.ObjectWorldPositionHighAndRadius;
-		Data[i+1] = PrimitiveUniformShaderParameters.ObjectWorldPositionLow;
+		Data[i+1] = FVector4f(PrimitiveUniformShaderParameters.ObjectWorldPositionLow, 0.0f);
 		i+=2;
 	}
 
diff --git a/Engine/Source/Runtime/Renderer/Private/InstanceCulling/InstanceCullingContext.cpp b/Engine/Source/Runtime/Renderer/Private/InstanceCulling/InstanceCullingContext.cpp
--- a/Engine/Source/Runtime/Renderer/Private/InstanceCulling/InstanceCullingContext.cpp
+++ b/Engine/Source/Runtime/Renderer/Private/InstanceCulling/InstanceCullingContext.cpp
@@ -545,6 +545,10 @@ public:
 		const FPermutationDomain PermutationVector(Parameters.PermutationId);
 		EBatchProcessingMode ProcessingMode = (PermutationVector.Get<FSingleInstanceModeDim>() ? EBatchProcessingMode::UnCulled : EBatchProcessingMode::Generic);
 		OutEnvironment.SetDefine(TEXT("INSTANCE_DATA_STRIDE_ELEMENTS"), GetInstanceDataStrideElements(Parameters.Platform, ProcessingMode));
+
+		static const auto CVarPrimitiveHasTileOffsetData = IConsoleManager::Get().FindTConsoleVariableDataInt(TEXT("r.PrimitiveHasTileOffsetData"));
+		const bool bPrimitiveHasTileOffsetData = CVarPrimitiveHasTileOffsetData->GetValueOnAnyThread() != 0;
+		OutEnvironment.SetDefine(TEXT("PRIMITIVE_HAS_TILEOFFSET_DATA"), bPrimitiveHasTileOffsetData ? 1 : 0);
 	}
 
 	BEGIN_SHADER_PARAMETER_STRUCT(FParameters, )
