Substrate - support for light complexity view.

diff --git a/Engine/Shaders/Private/DeferredLightPixelShaders.usf b/Engine/Shaders/Private/DeferredLightPixelShaders.usf
--- a/Engine/Shaders/Private/DeferredLightPixelShaders.usf
+++ b/Engine/Shaders/Private/DeferredLightPixelShaders.usf
@@ -263,6 +263,8 @@ void DeferredLightPixelMain(
 	FSubstrateAddressing SubstrateAddressing = GetSubstratePixelDataByteOffset(PixelPos, uint2(View.BufferSizeAndInvSize.xy), Substrate.MaxBytesPerPixel);
 	FSubstratePixelHeader SubstratePixelHeader = UnpackSubstrateHeaderIn(Substrate.MaterialTextureArray, SubstrateAddressing, Substrate.TopLayerTexture);
 
+	float EstimatedCost = 0.3; // running the PixelShader at all has a cost
+
 	BRANCH
 	if (SubstratePixelHeader.ClosureCount > 0	// This test is also enough to exclude sky pixels
 #if USE_LIGHTING_CHANNELS
@@ -310,14 +312,24 @@ void DeferredLightPixelMain(
 				SubstrateAddressing,
 				SubstratePixelHeader);
 
+			EstimatedCost += SubstrateLighting.EstimatedCost;
+
 			OutColor += SubstrateLighting.SceneColor;
-#if SUBSTRATE_OPAQUE_ROUGH_REFRACTION_ENABLED
+		#if SUBSTRATE_OPAQUE_ROUGH_REFRACTION_ENABLED
 			OutOpaqueRoughRefractionSceneColor += SubstrateLighting.OpaqueRoughRefractionSceneColor;
 			OutSubSurfaceSceneColor += SubstrateLighting.SubSurfaceSceneColor;
-#endif
+		#endif
 		}
 	}
 
+#if VISUALIZE_LIGHT_CULLING == 1
+	OutColor = ConvertEstimatedCostToColor(EstimatedCost);
+	#if SUBSTRATE_OPAQUE_ROUGH_REFRACTION_ENABLED
+	OutOpaqueRoughRefractionSceneColor = 0.0f;
+	OutSubSurfaceSceneColor = 0.0f;
+	#endif
+#endif
+
 #else // SUBSTRATE_ENABLED
 
 	FScreenSpaceData ScreenSpaceData = GetScreenSpaceData(InputParams.ScreenUV);
diff --git a/Engine/Shaders/Private/LightAccumulator.ush b/Engine/Shaders/Private/LightAccumulator.ush
--- a/Engine/Shaders/Private/LightAccumulator.ush
+++ b/Engine/Shaders/Private/LightAccumulator.ush
@@ -84,6 +84,11 @@ void LightAccumulator_Add(inout FLightAccumulator In, float3 TotalLight, float3
 	LightAccumulator_AddSplit(In, TotalLight, 0.0f, ScatterableLight, CommonMultiplier, bNeedsSeparateSubsurfaceLightAccumulation);
 }
 
+float4 ConvertEstimatedCostToColor(float EstimatedCost)
+{
+	return 0.1f * float4(1.0f, 0.25f, 0.075f, 0) * EstimatedCost;
+}
+
 //
 // compute final value to store in the MRT0
 // @retrun RGB:SceneColor Specular and Diffuse, A:Non Specular SceneColor Luminance
@@ -94,7 +99,7 @@ float4 LightAccumulator_GetResult(FLightAccumulator In)
 	if (VISUALIZE_LIGHT_CULLING == 1)
 	{
 		// a soft gradient from dark red to bright white, can be changed to be different
-		Ret = 0.1f * float4(1.0f, 0.25f, 0.075f, 0) * In.EstimatedCost;
+		Ret = ConvertEstimatedCostToColor(In.EstimatedCost);
 	}
 	else
 	{
@@ -137,8 +142,8 @@ FDeferredLightingSplit LightAccumulator_GetResultSplit(FLightAccumulator In)
 	if (VISUALIZE_LIGHT_CULLING == 1)
 	{
 		// a soft gradient from dark red to bright white, can be changed to be different
-		RetDiffuse = 0.1f * float4(1.0f, 0.25f, 0.075f, 0) * In.EstimatedCost;
-		RetSpecular = 0.1f * float4(1.0f, 0.25f, 0.075f, 0) * In.EstimatedCost;
+		RetDiffuse  = ConvertEstimatedCostToColor(In.EstimatedCost);
+		RetSpecular = RetDiffuse;
 	}
 	else
 	{
@@ -176,6 +181,10 @@ struct FSubstrateDeferredLighting
 	float3 OpaqueRoughRefractionSceneColor;	// Luminance that is going to be blurier by the top layer roughness
 	float3 SubSurfaceSceneColor;			// Scene color that should go through the sub surface post process. There is not alpha: not needed because it is separated from specular already.
 #endif
+
+	// only used for development (not compiled out with #if so we can use if() instead of #if for better readability and compiler error checking)
+	// assumed to be compiled out otherwise
+	float EstimatedCost;
 };
 
 FSubstrateDeferredLighting GetInitialisedSubstrateDeferredLighting()
@@ -204,4 +213,6 @@ void AccumulateSubstrateDeferredLighting(inout FSubstrateDeferredLighting Substr
 #else
 	SubstrateLighting.SceneColor += DiffSpec.DiffuseLighting + DiffSpec.SpecularLighting;
 #endif
+
+	SubstrateLighting.EstimatedCost += In.EstimatedCost;
 }
diff --git a/Engine/Shaders/Private/Substrate/SubstrateDeferredLighting.ush b/Engine/Shaders/Private/Substrate/SubstrateDeferredLighting.ush
--- a/Engine/Shaders/Private/Substrate/SubstrateDeferredLighting.ush
+++ b/Engine/Shaders/Private/Substrate/SubstrateDeferredLighting.ush
@@ -79,6 +79,8 @@ FSubstrateDeferredLighting SubstrateDeferredLighting(
 			LightData,
 			SubstrateShadowTermInputParameters.LightAttenuation,
 			ShadowTerms);
+
+		SubstrateLighting.EstimatedCost += 0.3;	// add the cost of getting the shadow terms
 	}
 	
 #if USE_LIGHT_FUNCTION_ATLAS
@@ -143,6 +145,8 @@ FSubstrateDeferredLighting SubstrateDeferredLighting(
 			}
 		}
 
+		SubstrateLighting.EstimatedCost += 0.4;	// add the cost of the lighting computations (should sum up to 1 form one light)
+
 		float3 DiffuseLuminance = BSDFEvaluate.IntegratedDiffuseValue;
 		float3 SpecularLuminance = BSDFEvaluate.IntegratedSpecularValue * LightData.SpecularScale;
 
