Fix minor shader warnings

diff --git a/Engine/Shaders/Private/DistanceFieldShadowing.usf b/Engine/Shaders/Private/DistanceFieldShadowing.usf
--- a/Engine/Shaders/Private/DistanceFieldShadowing.usf
+++ b/Engine/Shaders/Private/DistanceFieldShadowing.usf
@@ -305,7 +305,7 @@ FObjectViewSpaceBox GetObjectViewSpaceBox(uint ObjectIndex)
 
 #if (DISTANCEFIELD_PRIMITIVE_TYPE == DFPT_SignedDistanceField)
 		FDFObjectData DFObjectData = LoadDFObjectData(ObjectIndex);
-		float3 TranslatedWorldBoundsPosition = DFTransformLocalToTranslatedWorld(UnitBoxVertex * DFObjectData.VolumePositionExtent, DFObjectData.VolumeToWorld, PrimaryView.PreViewTranslation);
+		float3 TranslatedWorldBoundsPosition = DFTransformLocalToTranslatedWorld(UnitBoxVertex * DFObjectData.VolumePositionExtent, DFObjectData.VolumeToWorld, PrimaryView.PreViewTranslation).xyz;
 #else
 		FHeightfieldObjectBounds HeightfieldObjectBounds = LoadHeightfieldObjectBounds(ObjectIndex);
 		float3 TranslatedWorldBoundsPosition = UnitBoxVertex * HeightfieldObjectBounds.BoxExtent + DFFastToTranslatedWorld(HeightfieldObjectBounds.BoxOrigin, PrimaryView.PreViewTranslation);
diff --git a/Engine/Shaders/Private/ParticleSimulationShader.usf b/Engine/Shaders/Private/ParticleSimulationShader.usf
--- a/Engine/Shaders/Private/ParticleSimulationShader.usf
+++ b/Engine/Shaders/Private/ParticleSimulationShader.usf
@@ -64,7 +64,7 @@ void VertexMain(
 	float3 TileOffset = Input.TileOffsets.xyz;
 #endif
 	
-	float2 TileCoord = Input.TexCoord.xy * TexCoordScale.xy + TileOffset * TilePageScale.xy + GetTilePageOffset(TileOffset.z, TilePageScale, true);
+	float2 TileCoord = Input.TexCoord.xy * TexCoordScale.xy + TileOffset.xy * TilePageScale.xy + GetTilePageOffset(TileOffset.z, TilePageScale, true);
 
 	OutPosition = float4(
 		TileCoord.xy * float2(2.0f,-2.0f) + float2(-1.0f,1.0f),
diff --git a/Engine/Shaders/Private/PostProcessVisualizeHDR.usf b/Engine/Shaders/Private/PostProcessVisualizeHDR.usf
--- a/Engine/Shaders/Private/PostProcessVisualizeHDR.usf
+++ b/Engine/Shaders/Private/PostProcessVisualizeHDR.usf
@@ -240,24 +240,23 @@ void MainPS(noperspective float4 UVAndScreenPos : TEXCOORD0, float4 SvPosition :
 				OutColor = Texture2DSample(SceneColorTexture, SceneColorSampler, UV); // Influenced area default scene color
 			}
 
-			float3 LuminanceAvg = SceneLuminance(HDRSceneColorAvg);
-			float MaxLuminance = max(LuminanceAvg.r, max(LuminanceAvg.g, LuminanceAvg.b));
+			float LuminanceAvg = SceneLuminance(HDRSceneColorAvg);
 
 			// Luminance
 			int2 TopLeft = OutputViewportCenter + int2(11, 11);
-			if (MaxLuminance < 1000.0f)
+			if (LuminanceAvg < 1000.0f)
 			{
-				PrintFloat(PixelPos, OutColor.xyz, NumberColor, TopLeft, MaxLuminance);
+				PrintFloat(PixelPos, OutColor.xyz, NumberColor, TopLeft, LuminanceAvg);
 				TopLeft.x += 60;
 			}
-			else if (MaxLuminance < 999999.0f)
+			else if (LuminanceAvg < 999999.0f)
 			{
-				PrintFloatNoFractionLarge(PixelPos, OutColor.xyz, NumberColor, TopLeft, MaxLuminance);
+				PrintFloatNoFractionLarge(PixelPos, OutColor.xyz, NumberColor, TopLeft, LuminanceAvg);
 				TopLeft.x += 60;
 			}
 			else
 			{
-				PrintFloatNoFraction(PixelPos, OutColor.xyz, NumberColor, TopLeft, MaxLuminance, 10);
+				PrintFloatNoFraction(PixelPos, OutColor.xyz, NumberColor, TopLeft, LuminanceAvg, 10);
 				TopLeft.x += 100;
 			}
 			PrintCharacter(PixelPos, OutColor.xyz, NumberColor, TopLeft, _N_);
