Add depth penalty back to bilaterial filtering weight for transmission based on SSD to avoid bleeding over depth.

diff --git a/Engine/Shaders/Private/ScreenSpaceDenoise/SSDSignalFramework.ush b/Engine/Shaders/Private/ScreenSpaceDenoise/SSDSignalFramework.ush
--- a/Engine/Shaders/Private/ScreenSpaceDenoise/SSDSignalFramework.ush
+++ b/Engine/Shaders/Private/ScreenSpaceDenoise/SSDSignalFramework.ush
@@ -507,29 +507,42 @@ float2 ComputeBilateralWeight(
 		}
 		
 		float PlaneMisAlignementBilateral;
+		float TransmissionPlaneMisAlignementBilateral = 0.0f;
 		{
-			const float	NormalEncodingError = -0.25;
-			
 			float DistFromRefPlane = dot(RefNormal, -NeighborToRefVector);
-
 			float NeighborToRefPlaneDistance = length(-NeighborToRefVector - DistFromRefPlane * RefNormal);
-			float NeighborToRefError = NeighborToRefPlaneDistance * NormalEncodingError;
-
-			PlaneMisAlignementBilateral = saturate(1 - 6 * (abs(DistFromRefPlane) + NeighborToRefError) * InvMaxWorldBluringDistance);
-		}
-
-		PositionBilateralWeight = LightTangentBilateral * PlaneMisAlignementBilateral;
+			
+			// Regular shadow
+			{
+				const float	NormalEncodingError = -0.25;
+				float NeighborToRefError = NeighborToRefPlaneDistance * NormalEncodingError;
 
-		{
-			// Transmission of subsurface scattering should only be denoised between subsurface profile pixels to avoid bleeding.
-			const bool bApplyFilteringForSSS = RefSceneMetadata.ShadingModelID == SHADINGMODELID_SUBSURFACE_PROFILE
-				&& NeighborSceneMetadata.ShadingModelID == SHADINGMODELID_SUBSURFACE_PROFILE
-				&& ComputeLightSourceRadius(DomainKnowledge) > 0.0f;
+				PlaneMisAlignementBilateral = saturate(1 - 6 * (abs(DistFromRefPlane) + NeighborToRefError) * InvMaxWorldBluringDistance);
+			}
 
-			if (bApplyFilteringForSSS)
+			// Transmission
 			{
-				TransmissionPositionBilaterialWeight = PositionBilateralWeight;
+				// Transmission of subsurface scattering should only be denoised between subsurface profile pixels to avoid bleeding.
+				const bool bApplyFilteringForSSS = RefSceneMetadata.ShadingModelID == SHADINGMODELID_SUBSURFACE_PROFILE
+					&& NeighborSceneMetadata.ShadingModelID == SHADINGMODELID_SUBSURFACE_PROFILE
+					&& ComputeLightSourceRadius(DomainKnowledge) > 0.0f;
+
+				if (bApplyFilteringForSSS)
+				{
+					// Use source radius as encoding error to add more penalty to depth to avoid SSS backface transmission rim for spotlight with source radius > 0.
+					// It makes the scaling of the distance by (DistanceFromLight - HitT)/HitT instead of (DistanceFromLight - HitT)/HitT*SourceRadius
+					// When light is behind close to neighbor in view direction, the light tangent bilaterial will be close to 1. Little penalty
+					// is applied in depth. Adding the effect of SourceRadius on NeighborToRefError could add more penalty to depth to mitigate the
+					// bleeding error.
+					const float	NormalEncodingError = ComputeLightSourceRadius(DomainKnowledge);
+					float NeighborToRefError = NeighborToRefPlaneDistance * NormalEncodingError;
+
+					TransmissionPlaneMisAlignementBilateral = saturate(1 - 6 * (abs(DistFromRefPlane) + NeighborToRefError) * InvMaxWorldBluringDistance);
+				}
 			}
+
+			PositionBilateralWeight = LightTangentBilateral * PlaneMisAlignementBilateral;
+			TransmissionPositionBilaterialWeight = LightTangentBilateral * TransmissionPlaneMisAlignementBilateral;
 		}
 	}
 	else if ((BilateralSettings & 0xF) == 6)
