VulkanRHI: VK RT:  Disable GenerateCapsuleLightOcclusionRayWithSolidAngleSampling until driver crash can be resolved.

#rnx

diff --git a/Engine/Shaders/Private/RayTracing/RayTracingCapsuleLight.ush b/Engine/Shaders/Private/RayTracing/RayTracingCapsuleLight.ush
--- a/Engine/Shaders/Private/RayTracing/RayTracingCapsuleLight.ush
+++ b/Engine/Shaders/Private/RayTracing/RayTracingCapsuleLight.ush
@@ -17,6 +17,16 @@ bool GenerateCapsuleLightOcclusionRayWithSolidAngleSampling(
 	out float RayPdf
 )
 {
+// :todo-jn: Get around a driver crash until it is resolved
+#if VULKAN_PROFILE_SM6
+	const float3 LightDirection = LightParameters.TranslatedWorldPosition - TranslatedWorldPosition;
+	RayOrigin = TranslatedWorldPosition;
+	RayDirection = LightDirection;
+	RayTMin = 0;
+	RayTMax = 0;
+	RayPdf = 0;
+	return false;
+#else
 	float3 LightDirection = LightParameters.TranslatedWorldPosition - TranslatedWorldPosition;
 	float RayLength2 = dot(LightDirection, LightDirection);
 	float Radius = max(LightParameters.SourceRadius, 0.01); // clamp radius to a small value since the code below doesn't handle pure lines
@@ -36,4 +46,5 @@ bool GenerateCapsuleLightOcclusionRayWithSolidAngleSampling(
 	RayTMax = Distance;
 	RayPdf = Result.w;
 	return Distance > 0; // sample is only valid if we hit the light
+#endif
 }
