Fixed filtered VSM on SingleLayerWater being broken. The indirect draw tile list used to encode tile position in 16|16 bits but was recently changed to 12|12 bits to match Lumen. This fix now also makes the respective VSM shaders use the 12|12 format.

#rnx

diff --git a/Engine/Shaders/Private/VirtualShadowMaps/VirtualShadowMapProjection.usf b/Engine/Shaders/Private/VirtualShadowMaps/VirtualShadowMapProjection.usf
--- a/Engine/Shaders/Private/VirtualShadowMaps/VirtualShadowMapProjection.usf
+++ b/Engine/Shaders/Private/VirtualShadowMaps/VirtualShadowMapProjection.usf
@@ -597,7 +597,7 @@ void VirtualShadowMapProjection(
 	uint3	DispatchThreadId	: SV_DispatchThreadID )
 {
 #if USE_TILE_LIST
-	const uint2 TileCoord = UnpackTileCoord16bits(TileListData.Load(GroupId.x)); // Need to match SingleLayerWater tile encoding
+	const uint2 TileCoord = UnpackTileCoord12bits(TileListData.Load(GroupId.x)); // Need to match SingleLayerWater tile encoding
 #else
 	const uint2 TileCoord = GroupId.xy;
 #endif // USE_TILE_LIST
diff --git a/Engine/Shaders/Private/VirtualShadowMaps/VirtualShadowMapProjectionComposite.usf b/Engine/Shaders/Private/VirtualShadowMaps/VirtualShadowMapProjectionComposite.usf
--- a/Engine/Shaders/Private/VirtualShadowMaps/VirtualShadowMapProjectionComposite.usf
+++ b/Engine/Shaders/Private/VirtualShadowMaps/VirtualShadowMapProjectionComposite.usf
@@ -17,7 +17,7 @@ void VirtualShadowMapCompositeTileVS(
 	out float4 Position : SV_POSITION)
 {
 	const uint TileData = TileListData[InstanceId.x];
-	const uint2 TileOrigin = WORK_TILE_SIZE * uint2((TileData & 0xFFFF), (TileData >> 16u));
+	const uint2 TileOrigin = WORK_TILE_SIZE * UnpackTileCoord12bits(TileData);
 
 	uint2 TileVertex = TileOrigin;
 	TileVertex.x += VertexId == 1 || VertexId == 2 || VertexId == 4 ? WORK_TILE_SIZE : 0;
