Fixes to mobile multi view with mobile hdr.

diff --git a/Engine/Shaders/Private/InstancedStereo.ush b/Engine/Shaders/Private/InstancedStereo.ush
--- a/Engine/Shaders/Private/InstancedStereo.ush
+++ b/Engine/Shaders/Private/InstancedStereo.ush
@@ -85,3 +85,5 @@ uint GetInstanceId(uint InstanceId)
 	return InstanceId;
 #endif
 }
+
+#define USE_MULTI_VIEW_ID_SV (MOBILE_MULTI_VIEW && (COMPILER_VULKAN && !INSTANCED_STEREO))
\ No newline at end of file
diff --git a/Engine/Shaders/Private/PostProcessMobile.usf b/Engine/Shaders/Private/PostProcessMobile.usf
--- a/Engine/Shaders/Private/PostProcessMobile.usf
+++ b/Engine/Shaders/Private/PostProcessMobile.usf
@@ -9,7 +9,11 @@
 #include "PostprocessHistogramCommon.ush"
 #include "ScreenPass.ush"
 
+#if USE_MULTI_VIEW_ID_SV
+Texture2DArray SceneColorTexture;
+#else
 Texture2D SceneColorTexture;
+#endif
 SamplerState SceneColorSampler;
 
 Texture2D LastFrameSceneColorTexture;
@@ -132,6 +136,15 @@ half2 SunConstDepthMaskScaleBias()
 	return half2(Scale,-DepthMin * Scale);
 }
 
+float4 SampleSceneColor(float2 SceneUV, float ArrayIndex)
+{
+#if USE_MULTI_VIEW_ID_SV
+	return Texture2DArraySample(SceneColorTexture, SceneColorSampler, float3(SceneUV.xy, ArrayIndex));
+#else
+	return Texture2DSample(SceneColorTexture, SceneColorSampler, SceneUV);
+#endif
+}
+
 //
 // Convert depth in alpha into combined circle of confusion and sun intensity.
 //
@@ -140,13 +153,20 @@ half2 SunConstDepthMaskScaleBias()
 
 void SunMaskPS_Mobile(
 	float4 InUVPos : TEXCOORD0,
+#if USE_MULTI_VIEW_ID_SV
+	in uint MultiViewId : SV_ViewID,
+#endif
 	out HALF_TYPE OutSunShaftAndDof : SV_Target0
 #if MOBILE_USESUN && METAL_MSAA_HDR_DECODE
 	, out HALF4_TYPE OutColor : SV_Target1
 #endif
 	)
 {
-	half4 SceneColor = SceneColorTexture.Sample(SceneColorSampler, InUVPos.xy);
+	#if !USE_MULTI_VIEW_ID_SV
+		uint MultiViewId = 0;
+	#endif
+
+	half4 SceneColor = SampleSceneColor(InUVPos.xy,MultiViewId);
 
 	#if MOBILE_USEDEPTHTEXTURE 
 		half InDepth = ConvertFromDeviceZ(Texture2DSampleLevel(MobileSceneTextures.SceneDepthTexture, MobileSceneTextures.SceneDepthTextureSampler, InUVPos.xy, 0).r);
@@ -258,6 +278,9 @@ float HdrWeightInv4(float3 Color)
 
 void BloomPS_Mobile(
 	float2 InUVs[4] : TEXCOORD0
+#if USE_MULTI_VIEW_ID_SV
+	, in uint MultiViewId : SV_ViewID
+#endif
 #if MOBILE_USEBLOOM
 
 	, out HALF4_TYPE OutColor : SV_Target0
@@ -304,10 +327,14 @@ void BloomPS_Mobile(
 
 #if MOBILE_USEBLOOM || MOBILE_USESUN || MOBILE_USEEYEADAPTATION
 
-	float4 C0 = SceneColorTexture.Sample(SceneColorSampler, InUVs[0]);
-	float4 C1 = SceneColorTexture.Sample(SceneColorSampler, InUVs[1]);
-	float4 C2 = SceneColorTexture.Sample(SceneColorSampler, InUVs[2]);
-	float4 C3 = SceneColorTexture.Sample(SceneColorSampler, InUVs[3]);
+#if !USE_MULTI_VIEW_ID_SV
+	uint MultiViewId = 0;
+#endif
+	
+	float4 C0 = SampleSceneColor(InUVs[0],MultiViewId);
+	float4 C1 = SampleSceneColor(InUVs[1],MultiViewId);
+	float4 C2 = SampleSceneColor(InUVs[2],MultiViewId);
+	float4 C3 = SampleSceneColor(InUVs[3],MultiViewId);
 
 	C0 *= HdrWeight4(C0);
 	C1 *= HdrWeight4(C1);
@@ -681,6 +708,9 @@ void DofDownVS_Mobile(
 
 void DofDownPS_Mobile(
 	float2 InUVs[5] : TEXCOORD0,
+#if USE_MULTI_VIEW_ID_SV
+	in uint MultiViewId : SV_ViewID,
+#endif
 	out HALF4_TYPE OutColor : SV_Target0
 	)
 {
@@ -689,10 +719,14 @@ void DofDownPS_Mobile(
 	// Fetch near diolation and scale to (0 to 16384.0) range.
 	float N = DofNearTexture.Sample(DofNearSampler, InUVs[0]).r * 16384.0;
 
-	float4 A = SceneColorTexture.Sample(SceneColorSampler, InUVs[1]);
-	float4 B = SceneColorTexture.Sample(SceneColorSampler, InUVs[2]);
-	float4 C = SceneColorTexture.Sample(SceneColorSampler, InUVs[3]);
-	float4 D = SceneColorTexture.Sample(SceneColorSampler, InUVs[4]);
+#if !USE_MULTI_VIEW_ID_SV
+	uint MultiViewId = 0;
+#endif
+	
+	float4 A = SampleSceneColor(InUVs[1],MultiViewId);
+	float4 B = SampleSceneColor(InUVs[2],MultiViewId);
+	float4 C = SampleSceneColor(InUVs[3],MultiViewId);
+	float4 D = SampleSceneColor(InUVs[4],MultiViewId);
 
 	A.a = SunShaftAndDofTexture.Sample(SunShaftAndDofSampler, InUVs[1]).r;
 	B.a = SunShaftAndDofTexture.Sample(SunShaftAndDofSampler, InUVs[2]).r;
@@ -871,12 +905,19 @@ void IntegrateDOFVS_Mobile(
 void IntegrateDOFPS_Mobile(
 	in float4 TexCoords : TEXCOORD0,
 	in float2 FineDofGrain : TEXCOORD1,
+#if USE_MULTI_VIEW_ID_SV
+	in uint MultiViewId : SV_ViewID,
+#endif
 	out HALF4_TYPE OutColor : SV_Target0
 )
 {
-	half4 SceneColor = Texture2DSample(SceneColorTexture, SceneColorSampler, TexCoords.xy);
-
-	half4 DofFine = Texture2DSample(SceneColorTexture, SceneColorSampler, FineDofGrain.xy);
+	#if !USE_MULTI_VIEW_ID_SV
+		uint MultiViewId = 0;;
+	#endif
+	
+	half4 SceneColor = SampleSceneColor(TexCoords.xy,MultiViewId);
+	half4 DofFine = SampleSceneColor(FineDofGrain.xy,MultiViewId);
+	
 	half4 Dof = Texture2DSample(DofBlurTexture, DofBlurSampler, TexCoords.zw);
 	half DofCoc = Texture2DSample(SunShaftAndDofTexture, SunShaftAndDofSampler, TexCoords.xy).r;
 	// Convert alpha back into circle of confusion.
diff --git a/Engine/Shaders/Private/PostProcessTonemap.usf b/Engine/Shaders/Private/PostProcessTonemap.usf
--- a/Engine/Shaders/Private/PostProcessTonemap.usf
+++ b/Engine/Shaders/Private/PostProcessTonemap.usf
@@ -28,7 +28,11 @@
 SCREEN_PASS_TEXTURE_VIEWPORT(Color)
 SCREEN_PASS_TEXTURE_VIEWPORT(Output)
 
+#if USE_MULTI_VIEW_ID_SV
+Texture2DArray ColorTexture;
+#else
 Texture2D ColorTexture;
+#endif
 SamplerState ColorSampler;
 
 
@@ -249,20 +253,24 @@ float max4(float4 x)
 	return max(max(x.r, x.g), max(x.b, x.a));
 }
 
-float4 SampleSceneColor(float2 SceneUV)
+float4 SampleSceneColor(float2 SceneUV, float ArrayIndex)
 {
 	SceneUV = clamp(SceneUV, Color_UVViewportBilinearMin, Color_UVViewportBilinearMax);
+	#if USE_MULTI_VIEW_ID_SV
+	return Texture2DArraySample(ColorTexture, ColorSampler, float3(SceneUV.xy,ArrayIndex));
+	#else
 	return Texture2DSample(ColorTexture, ColorSampler, SceneUV);
+	#endif
 }
 
-half3 LookupSceneColor(float2 SceneUV, float2 PixelOffset)
+half3 LookupSceneColor(float2 SceneUV, float ArrayIndex, float2 PixelOffset)
 {
 	float2 SampleUV = SceneUV + PixelOffset * Color_ExtentInverse;
-	return SampleSceneColor(SampleUV).xyz;
+	return SampleSceneColor(SampleUV, ArrayIndex).xyz;
 }
 
 float4 TonemapCommonPS(
-	float2 UV,
+	float3 UV,
 	float2 Vignette,
 	float4 GrainUV,
 	float2 ScreenPos, // [-1, 1]x[-1, 1]
@@ -276,6 +284,7 @@ float4 TonemapCommonPS(
 	const float GlobalExposure = EyeAdaptationData.x;
 
 	float2 SceneUV = UV.xy;
+	float ArrayIndex = UV.z;
 
 	// Sample the scene color.
 	half4 SceneColor;
@@ -297,14 +306,14 @@ float4 TonemapCommonPS(
 				CAUV.xy = ScreenPosToUV(CAUV.xy, Color_ExtentInverse);
 				CAUV.zw = ScreenPosToUV(CAUV.zw, Color_ExtentInverse);
 
-			half4 SceneColorR = SampleSceneColor(CAUV.xy + SceneUVJitter.xy);
-			half SceneColorG = SampleSceneColor(CAUV.zw + SceneUVJitter.xy).g;
-			half SceneColorB = SampleSceneColor(SceneUV).b;
+			half4 SceneColorR = SampleSceneColor(CAUV.xy + SceneUVJitter.xy, ArrayIndex);
+			half SceneColorG = SampleSceneColor(CAUV.zw + SceneUVJitter.xy, ArrayIndex).g;
+			half SceneColorB = SampleSceneColor(SceneUV, ArrayIndex).b;
 			SceneColor = SceneColorR;
 			SceneColor.g = SceneColorG;
 			SceneColor.b = SceneColorB;
 		#else
-			SceneColor = SampleSceneColor(SceneUV);
+			SceneColor = SampleSceneColor(SceneUV, ArrayIndex);
 		#endif
 
 		#if METAL_MSAA_HDR_DECODE && !USE_GAMMA_ONLY
@@ -331,7 +340,7 @@ float4 TonemapCommonPS(
 		}
 		#else
 		{
-			BloomUV = ApplyScreenTransform(UV, ColorToBloom);
+			BloomUV = ApplyScreenTransform(UV.xy, ColorToBloom);
 			BloomUV = clamp(BloomUV, BloomUVViewportBilinearMin, BloomUVViewportBilinearMax);
 		}
 		#endif
@@ -404,9 +413,9 @@ float4 TonemapCommonPS(
 		FPQMPContext PQMP = PQMPInit(floor(SceneUV* Color_Extent));
 				
 		half4 LuminanceNeighbors;
-
-		half3 C1 = LookupSceneColor(SceneUV, float2(PQMP.QuadVector.x, 0)).rgb;
-		half3 C3 = LookupSceneColor(SceneUV, float2(0, PQMP.QuadVector.y)).rgb;
+		
+		half3 C1 = LookupSceneColor(SceneUV, ArrayIndex, float2(PQMP.QuadVector.x, 0)).rgb;
+		half3 C3 = LookupSceneColor(SceneUV, ArrayIndex, float2(0, PQMP.QuadVector.y)).rgb;
 		LuminanceNeighbors.x = Luminance(C1);
 		LuminanceNeighbors.y = Luminance(C3);
 		
@@ -415,10 +424,10 @@ float4 TonemapCommonPS(
 		LuminanceNeighbors.z =  PQMPReadX(PQMP,A0);
 		LuminanceNeighbors.w =  PQMPReadY(PQMP,A0);
 	#else
-		half3 C1 = LookupSceneColor(SceneUV, float2(-1,  0));
-		half3 C2 = LookupSceneColor(SceneUV, float2( 1,  0));
-		half3 C3 = LookupSceneColor(SceneUV, float2( 0, -1));
-		half3 C4 = LookupSceneColor(SceneUV, float2( 0,  1));
+		half3 C1 = LookupSceneColor(SceneUV, ArrayIndex, float2(-1,  0));
+		half3 C2 = LookupSceneColor(SceneUV, ArrayIndex, float2( 1,  0));
+		half3 C3 = LookupSceneColor(SceneUV, ArrayIndex, float2( 0, -1));
+		half3 C4 = LookupSceneColor(SceneUV, ArrayIndex, float2( 0,  1));
 		half4 LuminanceNeighbors = half4(Luminance(C1), Luminance(C2), Luminance(C3), Luminance(C4));
 	#endif
 
@@ -575,6 +584,9 @@ void MainPS(
 	in noperspective float4 GrainUV : TEXCOORD2,
 	in noperspective float2 ScreenPos : TEXCOORD3,
 	in noperspective float2 FullViewUV : TEXCOORD4,
+	#if USE_MULTI_VIEW_ID_SV
+	in uint MultiViewId : SV_ViewID,
+	#endif
 	float4 SvPosition : SV_POSITION,		// after all interpolators
 	out float4 OutColor : SV_Target0
 #if OUTPUT_LUMINANCE
@@ -583,7 +595,12 @@ void MainPS(
 	)
 {
 	float Luminance;
-    OutColor = TonemapCommonPS(UV, InVignette, GrainUV, ScreenPos, FullViewUV, SvPosition, Luminance);
+	#if USE_MULTI_VIEW_ID_SV
+		float3 UVI = float3(UV.xy,MultiViewId);
+	#else
+		float3 UVI = float3(UV.xy,0);
+	#endif
+    OutColor = TonemapCommonPS(UVI, InVignette, GrainUV, ScreenPos, FullViewUV, SvPosition, Luminance);
 #if OUTPUT_LUMINANCE
 	OutLuminance = Luminance;
 #endif
@@ -613,7 +630,7 @@ void MainCS(uint2 DispatchThreadId : SV_DispatchThreadID)
 	TonemapCommonVS(InScreenPos, UV, Vignette, GrainUV, ScreenPos, FullViewUV);
 
 	float Luminance;
-	float4 OutColor = TonemapCommonPS(UV, Vignette, GrainUV, ScreenPos, FullViewUV, SvPosition, Luminance);
+	float4 OutColor = TonemapCommonPS(float3(UV.xy,0), Vignette, GrainUV, ScreenPos, FullViewUV, SvPosition, Luminance);
 
 	uint2 PixelPos = DispatchThreadId + Output_ViewportMin;
 
@@ -630,7 +647,7 @@ void MainCS(uint2 DispatchThreadId : SV_DispatchThreadID)
 	#endif
 #endif
 
-half4 FetchAndResolveSceneColor(float2 UV)
+half4 FetchAndResolveSceneColor(float2 UV, float ArrayIndex)
 {
 	half4 SceneColor = half4(1.0, 1.0, 1.0, 1.0);
 
@@ -646,7 +663,11 @@ half4 FetchAndResolveSceneColor(float2 UV)
 		SceneColor = half4(VulkanSubpassFetch0());
 	#endif
 #else
-	SceneColor = Texture2DSample(ColorTexture, ColorSampler, UV);
+	#if USE_MULTI_VIEW_ID_SV
+		SceneColor = Texture2DArraySample(ColorTexture, ColorSampler, float3(UV.xy,ArrayIndex));
+	#else
+		SceneColor = Texture2DSample(ColorTexture, ColorSampler, UV);
+	#endif
 #endif
 	return SceneColor;
 }
@@ -654,13 +675,22 @@ half4 FetchAndResolveSceneColor(float2 UV)
 void MobileCustomResolve_MainPS(
 	in noperspective float2 UV						: TEXCOORD0,
 	in				 float4 SvPosition				: SV_POSITION,		// after all interpolators
-	out				 float4 OutColor				: SV_Target0
+	out				 float4 OutColor				: SV_Target0,
+#if USE_MULTI_VIEW_ID_SV
+	in nointerpolation uint MultiViewId : SV_ViewID
+#endif
 )
 {
+#if USE_MULTI_VIEW_ID_SV
+	ResolvedView = ResolveView(MultiViewId);
+#else
+	uint MultiViewId = 0;
 	ResolvedView = ResolveView();
+#endif // USE_MULTI_VIEW_ID_SV
+	
 	const float OneOverPreExposure = ResolvedView.OneOverPreExposure;
 
-	half4 SceneColor = FetchAndResolveSceneColor(UV);
+	half4 SceneColor = FetchAndResolveSceneColor(UV, MultiViewId);
 	SceneColor.a = 1.0;
 
 	half3 LinearColor = SceneColor.rgb * OneOverPreExposure;
diff --git a/Engine/Source/Runtime/Renderer/Private/PostProcess/PostProcessTonemap.cpp b/Engine/Source/Runtime/Renderer/Private/PostProcess/PostProcessTonemap.cpp
--- a/Engine/Source/Runtime/Renderer/Private/PostProcess/PostProcessTonemap.cpp
+++ b/Engine/Source/Runtime/Renderer/Private/PostProcess/PostProcessTonemap.cpp
@@ -932,7 +932,7 @@ FScreenPassTexture AddTonemapPass(FRDGBuilder& GraphBuilder, const FViewInfo& Vi
 
 		DesktopPermutationVector.Set<TonemapperPermutation::FTonemapperOutputDeviceDim>(EDisplayOutputFormat(CommonParameters.OutputDevice.OutputDevice));
 
-		DesktopPermutationVector.Set<TonemapperPermutation::FTonemapperOutputLuminance>(GVRSImageManager.IsVRSEnabledForFrame() && FVariableRateShadingImageManager::IsVRSCompatibleWithView(View));
+		DesktopPermutationVector.Set<TonemapperPermutation::FTonemapperOutputLuminance>(!View.bIsMobileMultiViewEnabled && GVRSImageManager.IsVRSEnabledForFrame() && FVariableRateShadingImageManager::IsVRSCompatibleWithView(View));
 
 		DesktopPermutationVector = TonemapperPermutation::RemapPermutation(DesktopPermutationVector, View.GetFeatureLevel());
 	}
@@ -995,6 +995,7 @@ FScreenPassTexture AddTonemapPass(FRDGBuilder& GraphBuilder, const FViewInfo& Vi
 		FTonemapPS::FParameters* PassParameters = GraphBuilder.AllocParameters<FTonemapPS::FParameters>();
 		PassParameters->Tonemap = CommonParameters;
 		PassParameters->RenderTargets[0] = Output.GetRenderTargetBinding();
+		PassParameters->RenderTargets.MultiViewCount = View.bIsMobileMultiViewEnabled ? 2 : 0;
 
 		if (OutputLuminance)
 		{
diff --git a/Engine/Source/Runtime/Renderer/Private/ScreenPass.cpp b/Engine/Source/Runtime/Renderer/Private/ScreenPass.cpp
--- a/Engine/Source/Runtime/Renderer/Private/ScreenPass.cpp
+++ b/Engine/Source/Runtime/Renderer/Private/ScreenPass.cpp
@@ -81,7 +81,6 @@ FScreenPassTextureSlice FScreenPassTextureSlice::CreateFromScreenPassTexture(FRD
 		return FScreenPassTextureSlice(nullptr, ScreenTexture.ViewRect);
 	}
 
-	check(!ScreenTexture.Texture->Desc.IsTextureArray());
 	return FScreenPassTextureSlice(GraphBuilder.CreateSRV(FRDGTextureSRVDesc(ScreenTexture.Texture)), ScreenTexture.ViewRect);
 }
 
diff --git a/Engine/Source/Runtime/Renderer/Public/ScreenPass.inl b/Engine/Source/Runtime/Renderer/Public/ScreenPass.inl
--- a/Engine/Source/Runtime/Renderer/Public/ScreenPass.inl
+++ b/Engine/Source/Runtime/Renderer/Public/ScreenPass.inl
@@ -121,7 +121,7 @@ inline FScreenPassTextureSlice::FScreenPassTextureSlice(FRDGTextureSRVRef InText
 {
 	if (TextureSRV && TextureSRV->Desc.Texture->Desc.IsTextureArray())
 	{
-		check(TextureSRV->Desc.NumArraySlices == 1);
+		check(TextureSRV->Desc.NumArraySlices > 0);
 	}
 }
 
