Remove RadiusScale buffer since now the scaling is uniform for an hair group.

The hair scaling is only applied for Manual LOD is is basd on scaling the hair radius on CPU side.

diff --git a/Engine/Plugins/Runtime/HairStrands/Source/HairStrandsCore/Private/GroomManager.cpp b/Engine/Plugins/Runtime/HairStrands/Source/HairStrandsCore/Private/GroomManager.cpp
--- a/Engine/Plugins/Runtime/HairStrands/Source/HairStrandsCore/Private/GroomManager.cpp
+++ b/Engine/Plugins/Runtime/HairStrands/Source/HairStrandsCore/Private/GroomManager.cpp
@@ -1630,6 +1630,7 @@ static void RunHairStrandsInterpolation_Strands(
 		{
 			if (InstanceData.bNeedRaytracing)
 			{
+				// Note: VFInput.Strands.Common.Radius already contains RadiusScale from discrete LOD setup, for backward compatibility.
 				const float CLODScale = InstanceData.Instance->HairGroupPublicData->ContinuousLODCoverageScale;
 				const float HairRadiusRT = InstanceData.Instance->HairGroupPublicData->VFInput.Strands.Common.RaytracingRadiusScale * InstanceData.Instance->HairGroupPublicData->VFInput.Strands.Common.Radius * CLODScale;
 				const float HairRootScaleRT = InstanceData.Instance->HairGroupPublicData->VFInput.Strands.Common.RootScale;
@@ -2373,6 +2374,27 @@ static float ComputeActiveCurveCoverageScale(const FHairStrandsBulkData& InData,
 	return InData.GetCoverageScale(CurveRatio);
 }
 
+static float ComputeActiveCurveRadiusScale(const FHairStrandsClusterResource* InResource, float InLODIndex)
+{
+	float Out = 1.f;
+	if (InResource && InResource->BulkData.IsValid())
+	{
+		const FHairStrandsClusterBulkData& InData = InResource->BulkData;
+		const uint32 MaxLODCount = InData.Header.LODInfos.Num();
+		InLODIndex = FMath::Clamp(InLODIndex, 0.f,MaxLODCount-1);
+		const int32 iLODIndex0 = FMath::Floor(InLODIndex);
+		const int32 iLODIndex1 = FMath::Min(iLODIndex0+1,int32(MaxLODCount-1));
+		const float S = InLODIndex - iLODIndex0;
+		check(InData.Header.LODInfos.IsValidIndex(iLODIndex0));
+		check(InData.Header.LODInfos.IsValidIndex(iLODIndex1));
+
+		const float RadiusScale0 = InData.Header.LODInfos[iLODIndex0].RadiusScale;
+		const float RadiusScale1 = InData.Header.LODInfos[iLODIndex1].RadiusScale;
+		Out = iLODIndex0 != iLODIndex1 ? FMath::LerpStable(RadiusScale0, RadiusScale1, S) : RadiusScale0;
+	}
+	return Out;
+}
+
 static uint32 ComputeActiveCurveCount(float InScreenSize, float InAutoLODBias, uint32 InCurveCount, uint32 InClusterCount)
 {
 	const float Power = 1.f;  // This could be exposed per asset
@@ -2398,6 +2420,7 @@ struct FHairLOD
 	FVector2f ContinuousLODScreenPos = FVector2f(0,0);
 	FBoxSphereBounds ContinuousLODBounds;
 	float ContinuousLODCoverageScale = 1.f;
+	float ContinuousLODRadiusScale = 1.f;
 };
 
 static FHairLOD ComputeHairLODIndex(const FHairGroupInstance* Instance, const TArray<const FSceneView*>& Views)
@@ -2451,6 +2474,7 @@ static FHairLOD ComputeHairLODIndex(const FHairGroupInstance* Instance, const TA
 	Out.ContinuousLODScreenPos = MaxContinuousLODScreenPos;
 	Out.ContinuousLODBounds = SphereBound;
 	Out.ContinuousLODCoverageScale = 1.f;
+	Out.ContinuousLODRadiusScale = 1.f;
 
 	// Extract the min/max LOD with strands. 
 	// This assumes that these LODs are contiguous
@@ -2492,6 +2516,7 @@ static FHairLOD ComputeHairLODIndex(const FHairGroupInstance* Instance, const TA
 		Out.ContinuousLODCurveCount = EffectiveCurveCount;
 		Out.ContinuousLODPointCount = EffectiveCurveCount > 0 ? Instance->Strands.GetData().Header.CurveToPointCount[EffectiveCurveCount - 1] : 0;
 		Out.ContinuousLODCoverageScale = ComputeActiveCurveCoverageScale(Instance->Strands.GetData(), EffectiveCurveCount);
+		Out.ContinuousLODRadiusScale = ComputeActiveCurveRadiusScale(Instance->Strands.ClusterResource, Out.HairLODIndex);
 	}
 
 	return Out;
@@ -2508,6 +2533,7 @@ static void ApplyHairLOD(const FHairLOD& In, FHairGroupInstance* OutInstance)
 	OutInstance->HairGroupPublicData->ContinuousLODScreenPos = In.ContinuousLODScreenPos;
 	OutInstance->HairGroupPublicData->ContinuousLODBounds = In.ContinuousLODBounds;
 	OutInstance->HairGroupPublicData->ContinuousLODCoverageScale = In.ContinuousLODCoverageScale;
+	OutInstance->HairGroupPublicData->ContinuousLODRadiusScale = In.ContinuousLODRadiusScale;
 }
 
 // Function for selecting, loading, & initializing LOD resources
@@ -2712,7 +2738,8 @@ static bool SelectValidLOD(
 		const uint32 EffectiveCurveCount = FMath::Min(ResourceStatus.AvailableCurveCount, Instance->HairGroupPublicData->ContinuousLODCurveCount);
 		Instance->HairGroupPublicData->ContinuousLODCurveCount = EffectiveCurveCount;
 		Instance->HairGroupPublicData->ContinuousLODPointCount = EffectiveCurveCount > 0 ? Instance->Strands.GetData().Header.CurveToPointCount[EffectiveCurveCount - 1] : 0;
-		Instance->HairGroupPublicData->ContinuousLODCoverageScale = ComputeActiveCurveCoverageScale(Instance->Strands.GetData(), EffectiveCurveCount);
+		Instance->HairGroupPublicData->ContinuousLODCoverageScale = ComputeActiveCurveCoverageScale(Instance->Strands.GetData(), EffectiveCurveCount); 
+		Instance->HairGroupPublicData->ContinuousLODRadiusScale = ComputeActiveCurveRadiusScale(Instance->Strands.ClusterResource, HairLODIndex); // default is 1.f - Only used for backward compatibility
 		check(Instance->HairGroupPublicData->ContinuousLODPointCount <= Instance->HairGroupPublicData->RestPointCount);
 	}
 	return true;
diff --git a/Engine/Plugins/Runtime/HairStrands/Source/HairStrandsCore/Private/GroomResources.cpp b/Engine/Plugins/Runtime/HairStrands/Source/HairStrandsCore/Private/GroomResources.cpp
--- a/Engine/Plugins/Runtime/HairStrands/Source/HairStrandsCore/Private/GroomResources.cpp
+++ b/Engine/Plugins/Runtime/HairStrands/Source/HairStrandsCore/Private/GroomResources.cpp
@@ -1498,8 +1498,6 @@ void FHairStrandsCullingResource::InternalAllocate(FRDGBuilder& GraphBuilder, ui
 
 		InternalCreateByteAddressBufferRDG(GraphBuilder, CurveCount * FHairStrandsUintFormat::SizeInByte, Resources.CulledCurveBuffer, TEXT("Hair.Cluster_CulledCurveBuffer"), GetOwnerName(), EHairResourceUsageType::Dynamic);
 		InternalCreateVertexBufferRDG<FHairStrandsUintFormat>(GraphBuilder, PointCount, Resources.CulledVertexIdBuffer, TEXT("Hair.Cluster_CulledVertexIdBuffer"), GetOwnerName(), EHairResourceUsageType::Dynamic);
-		InternalCreateVertexBufferRDG<FHairStrandsFloatFormat>(GraphBuilder, PointCount, Resources.CulledVertexRadiusScaleBuffer, TEXT("Hair.Cluster_CulledVertexRadiusScaleBuffer"), GetOwnerName(), EHairResourceUsageType::Dynamic);
-		AddClearUAVFloatPass(GraphBuilder, RegisterAsUAV(GraphBuilder, Resources.CulledVertexRadiusScaleBuffer), 0.f);
 	
 		// Only allocated once, as cluster count does not change
 		if (Resources.DrawIndirectBuffer.Buffer == nullptr)
@@ -1517,7 +1515,6 @@ void FHairStrandsCullingResource::InternalRelease()
 	Resources.DrawIndirectRasterComputeBuffer.Release();
 	Resources.CulledCurveBuffer.Release();
 	Resources.CulledVertexIdBuffer.Release();
-	Resources.CulledVertexRadiusScaleBuffer.Release();
 	MaxAvailableCurveCount = 0;
 }
 
diff --git a/Engine/Plugins/Runtime/HairStrands/Source/HairStrandsCore/Private/HairStrandsClusterCulling.cpp b/Engine/Plugins/Runtime/HairStrands/Source/HairStrandsCore/Private/HairStrandsClusterCulling.cpp
--- a/Engine/Plugins/Runtime/HairStrands/Source/HairStrandsCore/Private/HairStrandsClusterCulling.cpp
+++ b/Engine/Plugins/Runtime/HairStrands/Source/HairStrandsCore/Private/HairStrandsClusterCulling.cpp
@@ -36,7 +36,6 @@ class FHairClusterCullCS: public FGlobalShader
 	
 		SHADER_PARAMETER_RDG_BUFFER_UAV(RWStructuredBuffer<uint>, OutPointCounter)
 		SHADER_PARAMETER_RDG_BUFFER_UAV(RWBuffer<uint>, OutIndexBuffer)
-		SHADER_PARAMETER_RDG_BUFFER_UAV(RWBuffer<float>, OutRadiusScaleBuffer)
 		SHADER_PARAMETER_RDG_BUFFER_UAV(RWByteAddressBuffer, OutCulledCurveBuffer)
 
 		SHADER_PARAMETER_STRUCT_REF(FViewUniformShaderParameters, ViewUniformBuffer)
@@ -122,7 +121,6 @@ void AddClusterCullingPass(
 		Parameters->ViewUniformBuffer	 = View->ViewUniformBuffer;
 		Parameters->OutPointCounter 	 = PointCounterUAVSkipBarrier;
 		Parameters->OutIndexBuffer 		 = RegisterAsUAV(GraphBuilder, *ClusterData.GetCulledVertexIdBuffer());
-		Parameters->OutRadiusScaleBuffer = RegisterAsUAV(GraphBuilder, *ClusterData.GetCulledVertexRadiusScaleBuffer());
 		Parameters->OutCulledCurveBuffer = RegisterAsUAV(GraphBuilder, *ClusterData.GetCulledCurveBuffer()); // TODO: this could be changed to be transient buffer
 
 		if (ShaderPrintData)
diff --git a/Engine/Plugins/Runtime/HairStrands/Source/HairStrandsCore/Private/HairStrandsClusterCulling.h b/Engine/Plugins/Runtime/HairStrands/Source/HairStrandsCore/Private/HairStrandsClusterCulling.h
--- a/Engine/Plugins/Runtime/HairStrands/Source/HairStrandsCore/Private/HairStrandsClusterCulling.h
+++ b/Engine/Plugins/Runtime/HairStrands/Source/HairStrandsCore/Private/HairStrandsClusterCulling.h
@@ -37,7 +37,6 @@ struct FHairStrandClusterData
 		// Culling & LOD output
 		FRDGExternalBuffer* GetCulledCurveBuffer() const				{ return HairGroupPublicPtr ? &HairGroupPublicPtr->GetCulledCurveBuffer() : nullptr; }
 		FRDGExternalBuffer* GetCulledVertexIdBuffer() const				{ return HairGroupPublicPtr ? &HairGroupPublicPtr->GetCulledVertexIdBuffer() : nullptr; }
-		FRDGExternalBuffer* GetCulledVertexRadiusScaleBuffer() const	{ return HairGroupPublicPtr ? &HairGroupPublicPtr->GetCulledVertexRadiusScaleBuffer() : nullptr; }
 		bool GetCullingResultAvailable() const							{ return HairGroupPublicPtr ? HairGroupPublicPtr->GetCullingResultAvailable() : false; }
 		void SetCullingResultAvailable(bool b)							{ if (HairGroupPublicPtr) HairGroupPublicPtr->SetCullingResultAvailable(b); }
 
diff --git a/Engine/Plugins/Runtime/HairStrands/Source/HairStrandsCore/Private/HairStrandsInterpolation.cpp b/Engine/Plugins/Runtime/HairStrands/Source/HairStrandsCore/Private/HairStrandsInterpolation.cpp
--- a/Engine/Plugins/Runtime/HairStrands/Source/HairStrandsCore/Private/HairStrandsInterpolation.cpp
+++ b/Engine/Plugins/Runtime/HairStrands/Source/HairStrandsCore/Private/HairStrandsInterpolation.cpp
@@ -448,7 +448,6 @@ FRDGHairStrandsCullingData ImportCullingData(FRDGBuilder& GraphBuilder, FHairGro
 	{
 		Out.HairStrandsVF_CullingIndirectBuffer		= Register(GraphBuilder, In->GetDrawIndirectRasterComputeBuffer(), ERDGImportedBufferFlags::CreateViews);
 		Out.HairStrandsVF_CullingIndexBuffer		= Register(GraphBuilder, In->GetCulledVertexIdBuffer(), ERDGImportedBufferFlags::CreateViews);
-		Out.HairStrandsVF_CullingRadiusScaleBuffer	= Register(GraphBuilder, In->GetCulledVertexRadiusScaleBuffer(), ERDGImportedBufferFlags::CreateViews);
 	}
 
 	return Out;
@@ -510,7 +509,6 @@ class FHairInterpolationCS : public FGlobalShader
 
 		SHADER_PARAMETER_RDG_BUFFER_SRV(Buffer<uint>, HairStrandsVFTODO_CullingIndirectBuffer)
 		SHADER_PARAMETER_RDG_BUFFER_SRV(Buffer<uint>, HairStrandsVFTODO_CullingIndexBuffer)
-		SHADER_PARAMETER_RDG_BUFFER_SRV(Buffer<float>, HairStrandsVFTODO_CullingRadiusScaleBuffer)
 		RDG_BUFFER_ACCESS(HairStrandsVFTODO_CullingIndirectBufferArgs, ERHIAccess::IndirectArgs)
 		
 		END_SHADER_PARAMETER_STRUCT()
@@ -1181,7 +1179,6 @@ class FHairRaytracingGeometryCS : public FGlobalShader
 		SHADER_PARAMETER(uint32, HairStrandsVF_bCullingEnable)
 		SHADER_PARAMETER_RDG_BUFFER_SRV(Buffer,	HairStrandsVF_CullingIndirectBuffer)
 		SHADER_PARAMETER_RDG_BUFFER_SRV(Buffer,	HairStrandsVF_CullingIndexBuffer)
-		SHADER_PARAMETER_RDG_BUFFER_SRV(Buffer,	HairStrandsVF_CullingRadiusScaleBuffer)
 		RDG_BUFFER_ACCESS(HairStrandsVF_CullingIndirectBufferArgs, ERHIAccess::IndirectArgs)
 
 		SHADER_PARAMETER_RDG_BUFFER_SRV(StructuredBuffer, PositionOffsetBuffer)
@@ -1257,7 +1254,6 @@ void AddGenerateRaytracingGeometryPass(
 	{
 		Parameters->HairStrandsVF_CullingIndirectBuffer = CullingData.HairStrandsVF_CullingIndirectBuffer.SRV;
 		Parameters->HairStrandsVF_CullingIndexBuffer = CullingData.HairStrandsVF_CullingIndexBuffer.SRV;
-		Parameters->HairStrandsVF_CullingRadiusScaleBuffer = CullingData.HairStrandsVF_CullingRadiusScaleBuffer.SRV;
 		Parameters->HairStrandsVF_CullingIndirectBufferArgs = CullingData.HairStrandsVF_CullingIndirectBuffer.Buffer;
 	}
 
@@ -1755,6 +1751,7 @@ FHairGroupPublicData::FVertexFactoryInput InternalComputeHairStrandsVertexInputD
 
 	OutVFInput.Strands.Common.RegisteredIndex = Instance->RegisteredIndex;
 	OutVFInput.Strands.Common.Radius = (GStrandHairWidth > 0 ? GStrandHairWidth : Instance->Strands.Modifier.HairWidth) * 0.5f;
+	OutVFInput.Strands.Common.Radius *= Instance->HairGroupPublicData->GetActiveStrandsRadiusScale();
 	OutVFInput.Strands.Common.RootScale = Instance->Strands.Modifier.HairRootScale;
 	OutVFInput.Strands.Common.TipScale = Instance->Strands.Modifier.HairTipScale;
 	OutVFInput.Strands.Common.RaytracingRadiusScale = (GHairRaytracingRadiusScale > 0 ? GHairRaytracingRadiusScale : Instance->Strands.Modifier.HairRaytracingRadiusScale);
diff --git a/Engine/Plugins/Runtime/HairStrands/Source/HairStrandsCore/Private/HairStrandsVertexFactory.cpp b/Engine/Plugins/Runtime/HairStrands/Source/HairStrandsCore/Private/HairStrandsVertexFactory.cpp
--- a/Engine/Plugins/Runtime/HairStrands/Source/HairStrandsCore/Private/HairStrandsVertexFactory.cpp
+++ b/Engine/Plugins/Runtime/HairStrands/Source/HairStrandsCore/Private/HairStrandsVertexFactory.cpp
@@ -122,12 +122,10 @@ FHairStrandsVertexFactoryUniformShaderParameters FHairGroupInstance::GetHairStan
 	if (Out.Culling.bCullingEnable)
 	{
 		Out.Culling.CullingIndexBuffer = HairGroupPublicData->GetCulledVertexIdBuffer().SRV;
-		Out.Culling.CullingRadiusScaleBuffer = HairGroupPublicData->GetCulledVertexRadiusScaleBuffer().SRV;
 	}
 	else
 	{
 		Out.Culling.CullingIndexBuffer = GDummyCulledDispatchVertexIdsBuffer.SRVUint;
-		Out.Culling.CullingRadiusScaleBuffer = GDummyCulledDispatchVertexIdsBuffer.SRVFloat;
 	}
 	return Out;
 }
diff --git a/Engine/Plugins/Runtime/HairStrands/Source/HairStrandsCore/Public/GroomResources.h b/Engine/Plugins/Runtime/HairStrands/Source/HairStrandsCore/Public/GroomResources.h
--- a/Engine/Plugins/Runtime/HairStrands/Source/HairStrandsCore/Public/GroomResources.h
+++ b/Engine/Plugins/Runtime/HairStrands/Source/HairStrandsCore/Public/GroomResources.h
@@ -372,7 +372,6 @@ struct FHairStrandsCullingResource : public FHairCommonResource
 		Total += GetBufferTotalNumBytes(Resources.DrawIndirectRasterComputeBuffer);
 		Total += GetBufferTotalNumBytes(Resources.CulledCurveBuffer);
 		Total += GetBufferTotalNumBytes(Resources.CulledVertexIdBuffer);
-		Total += GetBufferTotalNumBytes(Resources.CulledVertexRadiusScaleBuffer);
 		return Total;
 	}
 
diff --git a/Engine/Plugins/Runtime/HairStrands/Source/HairStrandsCore/Public/HairStrandsInterpolation.h b/Engine/Plugins/Runtime/HairStrands/Source/HairStrandsCore/Public/HairStrandsInterpolation.h
--- a/Engine/Plugins/Runtime/HairStrands/Source/HairStrandsCore/Public/HairStrandsInterpolation.h
+++ b/Engine/Plugins/Runtime/HairStrands/Source/HairStrandsCore/Public/HairStrandsInterpolation.h
@@ -98,7 +98,6 @@ struct FRDGHairStrandsCullingData
 	bool bCullingResultAvailable = false;
 	FRDGImportedBuffer HairStrandsVF_CullingIndirectBuffer;
 	FRDGImportedBuffer HairStrandsVF_CullingIndexBuffer;
-	FRDGImportedBuffer HairStrandsVF_CullingRadiusScaleBuffer;
 };
 
 FRDGHairStrandsCullingData ImportCullingData(FRDGBuilder& GraphBuilder, FHairGroupPublicData* In);
diff --git a/Engine/Shaders/Private/HairStrands/HairStrandsClusterCulling.usf b/Engine/Shaders/Private/HairStrands/HairStrandsClusterCulling.usf
--- a/Engine/Shaders/Private/HairStrands/HairStrandsClusterCulling.usf
+++ b/Engine/Shaders/Private/HairStrands/HairStrandsClusterCulling.usf
@@ -54,7 +54,7 @@ void MainClearAABBCS(uint2 DispatchThreadId : SV_DispatchThreadID)
 #endif
 
 /////////////////////////////////////////////////////////////////////////////////////////////////////////////////////////////
-// Write out culled index buffex & radius scale buffer
+// Write out culled index buffex
 
 #if SHADER_CLUSTER_CULL
 
@@ -79,7 +79,6 @@ StructuredBuffer<FPackedHairClusterInfo> ClusterInfoBuffer;
 
 RWStructuredBuffer<uint>  OutPointCounter;
 RWBuffer<uint>  OutIndexBuffer;
-RWBuffer<float> OutRadiusScaleBuffer;
 RWByteAddressBuffer  OutCulledCurveBuffer;
 
 #ifndef GROUP_SIZE
@@ -110,7 +109,6 @@ uint GetGlobalCurveIndex(uint ThreadIndex, uint2 GroupIndex)
 groupshared uint CurvePointIndex[MAX_CURVE];
 groupshared uint CurvePointCount[MAX_CURVE];
 groupshared uint ClusterLOD[MAX_CURVE];
-groupshared float ClusterRadiusScale[MAX_CURVE];
 groupshared uint LODCount[MAX_CURVE];
 
 groupshared uint AllocationOffset;
@@ -118,8 +116,6 @@ groupshared uint AllocationCount;
 
 // Store locally the culled point index (encoded into 8 bits, and store 4 per uint)
 groupshared uint CulledPointIndices[MAX_CURVE * MAX_POINT];
-groupshared uint CulledPointRadiusScales[MAX_CURVE * MAX_POINT];
-groupshared float CulledPointFadeScales[MAX_CURVE * MAX_POINT];
 groupshared uint ValidPoints[MAX_CURVE][MAX_POINT];
 
 groupshared uint WriteOffset0[MAX_CURVE];
@@ -136,7 +132,7 @@ void Main(uint3 DispatchThreadId : SV_DispatchThreadID, uint ThreadIndex : SV_Gr
 		const bool bValidCurveIndex = GlobalCurveIndex < CurveCount;
 	
 		// 1.1 Cluster ID
-		// 1.3 Get cluster info (i.e. map all info per LOD for a given cluster ID) and derive effective LOD data (radius scale)
+		// 1.3 Get cluster info (i.e. map all info per LOD for a given cluster ID) and derive effective LOD data
 		FHairClusterLOD ClusterLODInfo = (FHairClusterLOD)0;
 		if (bValidCurveIndex)
 		{
@@ -144,7 +140,6 @@ void Main(uint3 DispatchThreadId : SV_DispatchThreadID, uint ThreadIndex : SV_Gr
 			ClusterLODInfo = GetHairClusterLOD(ClusterInfoBuffer[ClusterId], ClusterInfoParameters, LODIndex);
 		}
 		ClusterLOD[LocalCurveIndex] 		 = ClusterLODInfo.LOD;
-		ClusterRadiusScale[LocalCurveIndex]  = ClusterLODInfo.RadiusScale; // Curve are no longer scaled individually by ClusterLODInfo.RadiusScale, but scale globally based on CPU value. This is for legacy path, using discrete LOD with manual ThicknessScale
 		LODCount[LocalCurveIndex]			 = ClusterLODInfo.LODCount;
 
 		// 1.4 Curve's point count/offset
@@ -172,7 +167,6 @@ void Main(uint3 DispatchThreadId : SV_DispatchThreadID, uint ThreadIndex : SV_Gr
 		const uint CurrentPointIndex = ThreadIndex % uint(MAX_POINT);
 
 		// 2.2 Mark point having their LOD visibility smaller than the requested LOD level
-		float RadiusFade = 1.0f;
 		ValidPoints[CurveIndex][CurrentPointIndex] = 0; 
 		if (CurveIndex < CurveCount) // incorrect
 		{	
@@ -181,9 +175,6 @@ void Main(uint3 DispatchThreadId : SV_DispatchThreadID, uint ThreadIndex : SV_Gr
 				const uint PointIndex = CurrentPointIndex + CurvePointIndex[CurveIndex];
 				const uint MinLOD = GetHairControlPointMinLOD(PointIndex, PointLODBuffer);
 
-				// When a curve reach its last LOD, fade its radius progressively
-				RadiusFade = ceil(ClusterLOD[CurveIndex]) >= LODCount[CurveIndex]-1 ? 1.f : saturate(MinLOD+1 - ClusterLOD[CurveIndex]);
-
 				// Prune control point based on their MinLOD value
 				if (IsHairControlPointActive(MinLOD, ClusterLOD[CurveIndex]))
 				{
@@ -204,8 +195,6 @@ void Main(uint3 DispatchThreadId : SV_DispatchThreadID, uint ThreadIndex : SV_Gr
 			if (bIsValid)
 			{
 				CulledPointIndices[PrefixSum]		= CurrentPointIndex + CurvePointIndex[CurveIndex];
-				CulledPointRadiusScales[PrefixSum]	= ClusterRadiusScale[CurveIndex];
-				CulledPointFadeScales[PrefixSum] 	= RadiusFade;
 			}
 
 			if (CurveIndex < CurveCount)
@@ -227,7 +216,6 @@ void Main(uint3 DispatchThreadId : SV_DispatchThreadID, uint ThreadIndex : SV_Gr
 	if (ThreadIndex < AllocationCount)
 	{
 		OutIndexBuffer[AllocationOffset + ThreadIndex]       = CulledPointIndices[ThreadIndex];
-		OutRadiusScaleBuffer[AllocationOffset + ThreadIndex] = CulledPointRadiusScales[ThreadIndex];  // For fading out the curve which are about to be culled, one can use CulledPointFadeScales[ThreadIndex]
 	}
 
 	// 6. Write out curve data with visible points for subsequent passes
@@ -255,7 +243,6 @@ void Main(uint3 DispatchThreadId : SV_DispatchThreadID, uint ThreadIndex : SV_Gr
 groupshared uint CurvePointIndex;
 groupshared uint CurvePointCount;
 groupshared float ClusterLOD;
-groupshared float ClusterRadiusScale;
 groupshared uint LODCount;
 
 groupshared uint AllocationOffset;
@@ -264,7 +251,6 @@ groupshared uint AllocationCount;
 // Store locally the culled point index (encoded into 8 bits, and store 4 per uint)
 groupshared uint CulledPointIndices[HAIR_MAX_NUM_POINT_PER_CURVE]; 
 groupshared uint ValidPoints[HAIR_MAX_NUM_POINT_PER_CURVE];
-groupshared float CulledPointFadeScales[HAIR_MAX_NUM_POINT_PER_CURVE];
 
 // Version for curve with >32 points
 [numthreads(GROUP_SIZE, 1, 1)]
@@ -281,7 +267,7 @@ void Main(uint3 DispatchThreadId : SV_DispatchThreadID, uint ThreadIndex : SV_Gr
 	if (ThreadIndex == 0)
 	{
 		// 1.1 Cluster ID
-		// 1.2 Get cluster info (i.e. map all info per LOD for a given cluster ID) and derive effective LOD data (radius scale)
+		// 1.2 Get cluster info (i.e. map all info per LOD for a given cluster ID) and derive effective LOD data
 		FHairClusterLOD ClusterLODInfo = (FHairClusterLOD)0;
 		if (bIsCurveValid)
 		{
@@ -289,7 +275,6 @@ void Main(uint3 DispatchThreadId : SV_DispatchThreadID, uint ThreadIndex : SV_Gr
 			ClusterLODInfo = GetHairClusterLOD(ClusterInfoBuffer[ClusterId], ClusterInfoParameters, LODIndex);
 		}
 		ClusterLOD = ClusterLODInfo.LOD;
-		ClusterRadiusScale = ClusterLODInfo.RadiusScale; // Curve are no longer scaled individually by ClusterLODInfo.RadiusScale, but scale globally based on CPU value. This is for legacy path, using discrete LOD with manual ThicknessScale
 		LODCount = ClusterLODInfo.LODCount;
 
 		// 1.4 Curve's point count/offset
@@ -316,7 +301,6 @@ void Main(uint3 DispatchThreadId : SV_DispatchThreadID, uint ThreadIndex : SV_Gr
 			const uint CurrentPointIndex = ThreadIndex + LoopIt * GROUP_SIZE;
 	
 			// 2.2 Mark point having their LOD visibility smaller than the requested LOD level
-			float RadiusFade = 1.0f;
 			ValidPoints[CurrentPointIndex] = 0; 
 			if (CurveIndex < CurveCount)
 			{	
@@ -325,9 +309,6 @@ void Main(uint3 DispatchThreadId : SV_DispatchThreadID, uint ThreadIndex : SV_Gr
 					const uint PointIndex = CurrentPointIndex + CurvePointIndex;
 					const uint MinLOD = GetHairControlPointMinLOD(PointIndex, PointLODBuffer);
 
-					// When a curve reach its last LOD, fade its radius progressively
-					RadiusFade = ceil(ClusterLOD) >= LODCount-1 ? 1.f : saturate(MinLOD+1 - ClusterLOD);
-
 					// Prune control point based on their MinLOD value
 					if (IsHairControlPointActive(MinLOD, ClusterLOD))
 					{
@@ -345,7 +326,6 @@ void Main(uint3 DispatchThreadId : SV_DispatchThreadID, uint ThreadIndex : SV_Gr
 				if (bIsValid)
 				{
 					CulledPointIndices[IndexOffset + PrefixSum] = CurrentPointIndex;
-					CulledPointFadeScales[IndexOffset + PrefixSum] = RadiusFade;
 				}
 
 				if (ThreadIndex == 0)
@@ -373,10 +353,7 @@ void Main(uint3 DispatchThreadId : SV_DispatchThreadID, uint ThreadIndex : SV_Gr
 			if (CurrentPointIndex < AllocationCount)
 			{
 				const uint CullPointIndex = CulledPointIndices[CurrentPointIndex];
-				const float FadeScale = CulledPointFadeScales[CurrentPointIndex];
-
 				OutIndexBuffer[AllocationOffset + CurrentPointIndex]       = CurvePointIndex + CullPointIndex;
-				OutRadiusScaleBuffer[AllocationOffset + CurrentPointIndex] = ClusterRadiusScale;  // For fading out the curve which are about to be culled, one can use FadeScale
 			}
 		}
 	}
diff --git a/Engine/Shaders/Private/HairStrands/HairStrandsDebug.usf b/Engine/Shaders/Private/HairStrands/HairStrandsDebug.usf
--- a/Engine/Shaders/Private/HairStrands/HairStrandsDebug.usf
+++ b/Engine/Shaders/Private/HairStrands/HairStrandsDebug.usf
@@ -266,10 +266,8 @@ void MainCS(uint3 DispatchThreadId : SV_DispatchThreadID)
 		return;
 
 	uint VertexIndex = VertexId;
-	float ScaleR = 1;
 	if (HairStrandsVF.bCullingEnable)
 	{
-		ScaleR = HairStrandsVF.CullingRadiusScaleBuffer[VertexIndex];
 		VertexIndex = HairStrandsVF.CullingIndexBuffer[VertexIndex];
 	}
 
@@ -277,7 +275,7 @@ void MainCS(uint3 DispatchThreadId : SV_DispatchThreadID)
 		HairStrandsVF.PositionBuffer,
 		VertexIndex,
 		ReadRenPositionOffset(HairStrandsVF.PositionOffsetBuffer, HairStrandsVF.RegisteredIndex),
-		HairStrandsVF.Radius * ScaleR,
+		HairStrandsVF.Radius,
 		HairStrandsVF.RootScale,
 		HairStrandsVF.TipScale).Position;
 
diff --git a/Engine/Shaders/Private/HairStrands/HairStrandsInterpolation.usf b/Engine/Shaders/Private/HairStrands/HairStrandsInterpolation.usf
--- a/Engine/Shaders/Private/HairStrands/HairStrandsInterpolation.usf
+++ b/Engine/Shaders/Private/HairStrands/HairStrandsInterpolation.usf
@@ -89,7 +89,6 @@ uint HairStrandsVFTODO_bCullingEnable;
 #if PERMUTATION_CULLING == 1
 Buffer<uint>	HairStrandsVFTODO_CullingIndirectBuffer; 
 Buffer<uint>	HairStrandsVFTODO_CullingIndexBuffer;
-Buffer<float>	HairStrandsVFTODO_CullingRadiusScaleBuffer;
 #endif
 
 #if PERMUTATION_DEFORMER
diff --git a/Engine/Shaders/Private/HairStrands/HairStrandsRaytracingGeometry.usf b/Engine/Shaders/Private/HairStrands/HairStrandsRaytracingGeometry.usf
--- a/Engine/Shaders/Private/HairStrands/HairStrandsRaytracingGeometry.usf
+++ b/Engine/Shaders/Private/HairStrands/HairStrandsRaytracingGeometry.usf
@@ -29,7 +29,6 @@ uint HairStrandsVF_bCullingEnable;
 #if PERMUTATION_CULLING == 1
 Buffer<uint>	HairStrandsVF_CullingIndirectBuffer;
 Buffer<uint>	HairStrandsVF_CullingIndexBuffer;
-Buffer<float>	HairStrandsVF_CullingRadiusScaleBuffer;
 Buffer<uint>	HairStrandsVF_CullingIndirectBufferArgs;
 #endif
 
@@ -168,9 +167,6 @@ void MainCS(uint DispatchThreadId : SV_DispatchThreadID)
 		return;
 	}
 
-	float LodRadiusScale0 = 1;
-	float LodRadiusScale1 = 1;
-
 	uint SourceIndex0 = VertexIndex0;
 	uint SourceIndex1 = VertexIndex1;
 #if PERMUTATION_CULLING == 1
@@ -184,15 +180,7 @@ void MainCS(uint DispatchThreadId : SV_DispatchThreadID)
 		}
 
 		SourceIndex0 = HairStrandsVF_CullingIndexBuffer[VertexIndex0];
-		LodRadiusScale0 = HairStrandsVF_CullingRadiusScaleBuffer[VertexIndex0];
 		SourceIndex1 = HairStrandsVF_CullingIndexBuffer[VertexIndex1];
-		LodRadiusScale1 = HairStrandsVF_CullingRadiusScaleBuffer[VertexIndex1];
-
-		if (LodRadiusScale0 <= 0.0f)
-		{
-			SetInvalidSegment(VertexIndex0);
-			return;
-		}
 	}
 #endif
 
@@ -213,8 +201,8 @@ void MainCS(uint DispatchThreadId : SV_DispatchThreadID)
 		}
 		else
 		{
-			float4 A = float4(P0.Position, P0.WorldRadius * LodRadiusScale0);
-			float4 B = float4(P1.Position, P1.WorldRadius * LodRadiusScale1);
+			float4 A = float4(P0.Position, P0.WorldRadius);
+			float4 B = float4(P1.Position, P1.WorldRadius);
 			if ((A.w == 0 && B.w == 0) || length2(B.xyz - A.xyz) == 0 || any(!IsFinite(A)) || any(!IsFinite(B)))
 			{
 				// this segment is degenerate, don't bother adding it
@@ -262,10 +250,10 @@ void MainCS(uint DispatchThreadId : SV_DispatchThreadID)
 
 	const uint BaseOutIndex = VertexIndex0 * 4;
 	{
-		const float3 Px0 = P0.Position - N * P0.WorldRadius * LodRadiusScale0;
-		const float3 Px1 = P0.Position + N * P0.WorldRadius * LodRadiusScale0;
-		const float3 Py0 = P0.Position - B * P0.WorldRadius * LodRadiusScale0;
-		const float3 Py1 = P0.Position + B * P0.WorldRadius * LodRadiusScale0;
+		const float3 Px0 = P0.Position - N * P0.WorldRadius;
+		const float3 Px1 = P0.Position + N * P0.WorldRadius;
+		const float3 Py0 = P0.Position - B * P0.WorldRadius;
+		const float3 Py1 = P0.Position + B * P0.WorldRadius;
 
 		OutputPositionBuffer[BaseOutIndex    ] = float4(Px0, 1);
 		OutputPositionBuffer[BaseOutIndex + 1] = float4(Py0, 1);
diff --git a/Engine/Shaders/Private/HairStrands/HairStrandsVertexFactory.ush b/Engine/Shaders/Private/HairStrands/HairStrandsVertexFactory.ush
--- a/Engine/Shaders/Private/HairStrands/HairStrandsVertexFactory.ush
+++ b/Engine/Shaders/Private/HairStrands/HairStrandsVertexFactory.ush
@@ -106,7 +106,6 @@ struct FVertexInfo
 	uint IsTip;				// The vertex is on the quad side toward the tip of the strand.
 #if !RAYHITGROUPSHADER
 	uint IsLeft;
-	float LodRadiusScale;
 	bool bForceInvalidQuad;
 #endif
 };
@@ -133,7 +132,6 @@ FVertexInfo GetVertexInfo(FVertexFactoryInput Input)
 	VertexInfo.HairControlPointId = BaseIndex;
 	VertexInfo.VertexIndex = BaseIndex + VertexInfo.IsTip;
 #if !RAYHITGROUPSHADER
-	VertexInfo.LodRadiusScale = 1.0f;
 	VertexInfo.bForceInvalidQuad = false;
 #endif
 
@@ -141,10 +139,6 @@ FVertexInfo GetVertexInfo(FVertexFactoryInput Input)
 	if (HairStrandsVF.bCullingEnable)
 	{
 		VertexInfo.VertexIndex = HairStrandsVF.CullingIndexBuffer[BaseIndex + VertexInfo.IsTip];
-#if !RAYHITGROUPSHADER
-		VertexInfo.LodRadiusScale = HairStrandsVF.CullingRadiusScaleBuffer[BaseIndex + VertexInfo.IsTip];
-		VertexInfo.bForceInvalidQuad = VertexInfo.LodRadiusScale <= 0.0f;
-#endif
 		VertexInfo.HairControlPointId = VertexInfo.VertexIndex - VertexInfo.IsTip;
 	}
 #endif
@@ -212,11 +206,7 @@ FHairControlPoint GetVertexPosition(const FVertexInfo VertexInfo, bool bInvalidJ
 		HairStrandsVF.PositionBuffer,
 		VertexInfo.VertexIndex, 
 		PositionOffset,
-#if RAYHITGROUPSHADER
 		HairStrandsVF.Radius,
-#else
-		HairStrandsVF.Radius * VertexInfo.LodRadiusScale,
-#endif
 		HairStrandsVF.RootScale,
 		HairStrandsVF.TipScale);
 	
@@ -254,11 +244,7 @@ float3 GetVertexPreviousPosition(FVertexFactoryInput Input, bool bInvalidJointVe
 		HairStrandsVF.PreviousPositionBuffer,
 		VertexInfo.VertexIndex,
 		PreviousPositionOffset,
-#if RAYHITGROUPSHADER
 		HairStrandsVF.Radius,
-#else
-		HairStrandsVF.Radius * VertexInfo.LodRadiusScale,
-#endif
 		HairStrandsVF.RootScale,
 		HairStrandsVF.TipScale);
 
@@ -410,11 +396,8 @@ RAY_TRACING_ENTRY_INTERSECTION(MaterialIS)
 
 	uint VertexIndex0 = PrimitiveIndex() / HairStrandsVF.RaytracingProceduralSplits;
 	uint VertexIndex1 = VertexIndex0 + 1;
-	float ScaleR0 = 1, ScaleR1 = 1;
 	if (HairStrandsVF.bCullingEnable)
 	{
-		ScaleR0 = HairStrandsVF.CullingRadiusScaleBuffer[VertexIndex0];
-		ScaleR1 = HairStrandsVF.CullingRadiusScaleBuffer[VertexIndex1];
 		VertexIndex0 = HairStrandsVF.CullingIndexBuffer[VertexIndex0];
 		VertexIndex1 = HairStrandsVF.CullingIndexBuffer[VertexIndex1];
 	}
@@ -425,14 +408,14 @@ RAY_TRACING_ENTRY_INTERSECTION(MaterialIS)
 		HairStrandsVF.PositionBuffer,
 		VertexIndex0,
 		PositionOffset,
-		EffectiveRadius * ScaleR0,
+		EffectiveRadius,
 		HairStrandsVF.RootScale,
 		HairStrandsVF.TipScale);
 	FHairControlPoint P1 = ReadHairControlPoint(
 		HairStrandsVF.PositionBuffer,
 		VertexIndex1,
 		PositionOffset,
-		EffectiveRadius * ScaleR1,
+		EffectiveRadius,
 		HairStrandsVF.RootScale,
 		HairStrandsVF.TipScale);
 
diff --git a/Engine/Shaders/Private/HairStrands/HairStrandsVertexFactoryCommon.ush b/Engine/Shaders/Private/HairStrands/HairStrandsVertexFactoryCommon.ush
--- a/Engine/Shaders/Private/HairStrands/HairStrandsVertexFactoryCommon.ush
+++ b/Engine/Shaders/Private/HairStrands/HairStrandsVertexFactoryCommon.ush
@@ -86,7 +86,6 @@ float3 ReadCardPositionOffset(StructuredBuffer<float4> In, uint InIndex, uint In
 	ByteAddressBuffer 			Prefix##_PointAttributeBuffer; \
 	Buffer<uint>				Prefix##_CullingIndirectBuffer; \
 	Buffer<uint>				Prefix##_CullingIndexBuffer; \
-	Buffer<float>				Prefix##_CullingRadiusScaleBuffer; \
 	Buffer<uint> 				Prefix##_CullingIndirectBufferArgs; \
 	float3 Prefix##_GetHairInstancePositionOffset() { return ReadRenPositionOffset(Prefix##_PositionOffsetBuffer, Prefix##_RegisteredIndex); }
 
diff --git a/Engine/Shaders/Private/HairStrands/HairStrandsVisibilityRasterClassification.usf b/Engine/Shaders/Private/HairStrands/HairStrandsVisibilityRasterClassification.usf
--- a/Engine/Shaders/Private/HairStrands/HairStrandsVisibilityRasterClassification.usf
+++ b/Engine/Shaders/Private/HairStrands/HairStrandsVisibilityRasterClassification.usf
@@ -90,10 +90,8 @@ void CSMain(uint DispatchThreadID : SV_DispatchThreadID, uint GroupThreadID : SV
 
 			const uint VertexIndex0 = HairStrandsVF_CullingIndexBuffer[FetchIndex0];
 			const uint VertexIndex1 = HairStrandsVF_CullingIndexBuffer[FetchIndex1];
-			const float LodRadiusScale0 = HairStrandsVF_CullingRadiusScaleBuffer[FetchIndex0];
-			const float LodRadiusScale1 = HairStrandsVF_CullingRadiusScaleBuffer[FetchIndex1];
 
-			if ((LodRadiusScale0 <= 0.0f) || (VertexIndex1 != (VertexIndex0 + 1)))
+			if (VertexIndex1 != (VertexIndex0 + 1))
 			{
 				bSegValid = false;
 			}
diff --git a/Engine/Shaders/Private/HairStrands/HairStrandsVisibilityRasterCompute.usf b/Engine/Shaders/Private/HairStrands/HairStrandsVisibilityRasterCompute.usf
--- a/Engine/Shaders/Private/HairStrands/HairStrandsVisibilityRasterCompute.usf
+++ b/Engine/Shaders/Private/HairStrands/HairStrandsVisibilityRasterCompute.usf
@@ -439,10 +439,8 @@ void BinningCS(uint DispatchThreadID : SV_DispatchThreadID, uint GroupThreadID :
 
 			const uint VertexIndex0 = HairStrandsVF_CullingIndexBuffer[FetchIndex0];
 			const uint VertexIndex1 = HairStrandsVF_CullingIndexBuffer[FetchIndex1];
-			const float LodRadiusScale0 = HairStrandsVF_CullingRadiusScaleBuffer[FetchIndex0];
-			const float LodRadiusScale1 = HairStrandsVF_CullingRadiusScaleBuffer[FetchIndex1];
 
-			if ((LodRadiusScale0 <= 0.0f) || (VertexIndex1 != VertexIndex0 + 1))
+			if (VertexIndex1 != VertexIndex0 + 1)
 			{
 				bSegValid = false;
 			}
diff --git a/Engine/Shaders/Private/HairStrands/HairStrandsVisibilityRasterComputeNaive.usf b/Engine/Shaders/Private/HairStrands/HairStrandsVisibilityRasterComputeNaive.usf
--- a/Engine/Shaders/Private/HairStrands/HairStrandsVisibilityRasterComputeNaive.usf
+++ b/Engine/Shaders/Private/HairStrands/HairStrandsVisibilityRasterComputeNaive.usf
@@ -158,10 +158,8 @@ void CSMain(uint DispatchThreadID : SV_DispatchThreadID, uint GroupThreadID : SV
 
 			const uint VertexIndex0 = HairStrandsVF_CullingIndexBuffer[FetchIndex0];
 			const uint VertexIndex1 = HairStrandsVF_CullingIndexBuffer[FetchIndex1];
-			const float LodRadiusScale0 = HairStrandsVF_CullingRadiusScaleBuffer[FetchIndex0];
-			const float LodRadiusScale1 = HairStrandsVF_CullingRadiusScaleBuffer[FetchIndex1];
 
-			if ((LodRadiusScale0 <= 0.0f) || (VertexIndex1 != VertexIndex0 + 1))
+			if (VertexIndex1 != VertexIndex0 + 1)
 			{
 				bSegValid = false;
 			}
diff --git a/Engine/Shaders/Private/HairStrands/HairStrandsVisibilityRasterHW.usf b/Engine/Shaders/Private/HairStrands/HairStrandsVisibilityRasterHW.usf
--- a/Engine/Shaders/Private/HairStrands/HairStrandsVisibilityRasterHW.usf
+++ b/Engine/Shaders/Private/HairStrands/HairStrandsVisibilityRasterHW.usf
@@ -65,10 +65,8 @@ VSOutput VSMain(VSInput Input)
 
 		const uint VertexIndex0 = HairStrandsVF_CullingIndexBuffer[FetchIndex0];
 		const uint VertexIndex1 = HairStrandsVF_CullingIndexBuffer[FetchIndex1];
-		const float LodRadiusScale0 = HairStrandsVF_CullingRadiusScaleBuffer[FetchIndex0];
-		const float LodRadiusScale1 = HairStrandsVF_CullingRadiusScaleBuffer[FetchIndex1];
 
-		if ((LodRadiusScale0 <= 0.0f) || (VertexIndex1 != VertexIndex0 + 1))
+		if (VertexIndex1 != VertexIndex0 + 1)
 		{
 			return GetNullOutput();
 		}
diff --git a/Engine/Shaders/Private/HairStrands/HairStrandsVoxelRasterCompute.usf b/Engine/Shaders/Private/HairStrands/HairStrandsVoxelRasterCompute.usf
--- a/Engine/Shaders/Private/HairStrands/HairStrandsVoxelRasterCompute.usf
+++ b/Engine/Shaders/Private/HairStrands/HairStrandsVoxelRasterCompute.usf
@@ -678,8 +678,6 @@ void MainCS(uint2 DispatchThreadID : SV_DispatchThreadID)
 	if (!bIsValid)
 		return;
 
-	float LodRadiusScale0 = 1;
-	float LodRadiusScale1 = 1;
 #if PERMUTATION_CULLING == 1
 	if (HairStrandsVF_bCullingEnable)
 	{
@@ -694,15 +692,7 @@ void MainCS(uint2 DispatchThreadID : SV_DispatchThreadID)
 		FetchIndex1 = min(FetchIndex0 + 1, VertexCountAfterCulling - 1);
 
 		VertexIndex0 = HairStrandsVF_CullingIndexBuffer[FetchIndex0];
-		LodRadiusScale0 = HairStrandsVF_CullingRadiusScaleBuffer[FetchIndex0];
-
 		VertexIndex1 = HairStrandsVF_CullingIndexBuffer[FetchIndex1];
-		LodRadiusScale1 = HairStrandsVF_CullingRadiusScaleBuffer[FetchIndex1];
-
-		if (LodRadiusScale0 <= 0.0f)
-		{
-			return;
-		}
 	}
 #endif
 
@@ -712,7 +702,7 @@ void MainCS(uint2 DispatchThreadID : SV_DispatchThreadID)
 		HairStrandsVF_PositionBuffer,
 		VertexIndex0,
 		PositionOffset,
-		HairStrandsVF_Radius * LodRadiusScale0,
+		HairStrandsVF_Radius,
 		HairStrandsVF_RootScale,
 		HairStrandsVF_TipScale);
 
@@ -723,7 +713,7 @@ void MainCS(uint2 DispatchThreadID : SV_DispatchThreadID)
 		HairStrandsVF_PositionBuffer,
 		VertexIndex1,
 		PositionOffset,
-		HairStrandsVF_Radius * LodRadiusScale1,
+		HairStrandsVF_Radius,
 		HairStrandsVF_RootScale,
 		HairStrandsVF_TipScale);
 
diff --git a/Engine/Source/Runtime/Renderer/Private/HairStrands/HairStrandsInterface.cpp b/Engine/Source/Runtime/Renderer/Private/HairStrands/HairStrandsInterface.cpp
--- a/Engine/Source/Runtime/Renderer/Private/HairStrands/HairStrandsInterface.cpp
+++ b/Engine/Source/Runtime/Renderer/Private/HairStrands/HairStrandsInterface.cpp
@@ -346,6 +346,11 @@ float FHairGroupPublicData::GetActiveStrandsCoverageScale() const
 	return ContinuousLODCoverageScale; 
 }
 
+float FHairGroupPublicData::GetActiveStrandsRadiusScale() const
+{
+	return ContinuousLODRadiusScale; 
+}
+
 ///////////////////////////////////////////////////////////////////////////////////////////////////
 // Bookmark API
 THairStrandsBookmarkFunction  GHairStrandsBookmarkFunction = nullptr;
diff --git a/Engine/Source/Runtime/Renderer/Private/HairStrands/HairStrandsVisibility.cpp b/Engine/Source/Runtime/Renderer/Private/HairStrands/HairStrandsVisibility.cpp
--- a/Engine/Source/Runtime/Renderer/Private/HairStrands/HairStrandsVisibility.cpp
+++ b/Engine/Source/Runtime/Renderer/Private/HairStrands/HairStrandsVisibility.cpp
@@ -3102,11 +3102,9 @@ void AddMeshDrawTransitionPass(
 
 			FRDGBufferRef CulledCurveBuffer = Register(GraphBuilder, HairGroupPublicData->GetCulledCurveBuffer(), ERDGImportedBufferFlags::None).Buffer;
 			FRDGBufferRef CulledVertexIdBuffer = Register(GraphBuilder, HairGroupPublicData->GetCulledVertexIdBuffer(), ERDGImportedBufferFlags::None).Buffer;
-			FRDGBufferRef CulledVertexRadiusScaleBuffer = Register(GraphBuilder, HairGroupPublicData->GetCulledVertexRadiusScaleBuffer(), ERDGImportedBufferFlags::None).Buffer;
 			FRDGBufferRef DrawIndirectBuffer = Register(GraphBuilder, HairGroupPublicData->GetDrawIndirectBuffer(), ERDGImportedBufferFlags::None).Buffer;
 			ExternalAccessQueue.Add(CulledCurveBuffer);
 			ExternalAccessQueue.Add(CulledVertexIdBuffer);
-			ExternalAccessQueue.Add(CulledVertexRadiusScaleBuffer);
 			ExternalAccessQueue.Add(DrawIndirectBuffer, ERHIAccess::IndirectArgs);
 
 			const EHairVisibilityRenderMode RasterMode = GetHairVisibilityRenderMode(ViewInfo.GetShaderPlatform());
@@ -3232,7 +3230,6 @@ void GetHairStrandsInstanceCulling(FRDGBuilder& GraphBuilder, const FViewInfo& V
 		OutCulling.bCullingEnable = 1;
 		OutCulling.CullingIndirectBuffer = CullingIndirectBuffer.SRV;
 		OutCulling.CullingIndexBuffer = RegisterAsSRV(GraphBuilder, HairGroupPublicData->GetCulledVertexIdBuffer());
-		OutCulling.CullingRadiusScaleBuffer = RegisterAsSRV(GraphBuilder, HairGroupPublicData->GetCulledVertexRadiusScaleBuffer());
 		OutCulling.CullingIndirectBufferArgs = CullingIndirectBuffer.Buffer;
 	}
 }
diff --git a/Engine/Source/Runtime/Renderer/Private/HairStrands/HairStrandsVoxelization.cpp b/Engine/Source/Runtime/Renderer/Private/HairStrands/HairStrandsVoxelization.cpp
--- a/Engine/Source/Runtime/Renderer/Private/HairStrands/HairStrandsVoxelization.cpp
+++ b/Engine/Source/Runtime/Renderer/Private/HairStrands/HairStrandsVoxelization.cpp
@@ -1368,7 +1368,6 @@ void VoxelizeHairStrands(
 		if (InstanceData.Data->GetCullingResultAvailable())
 		{
 			Transitions.Add(RegisterAsSRV(GraphBuilder, InstanceData.Data->GetDrawIndirectRasterComputeBuffer()));
-			Transitions.Add(RegisterAsSRV(GraphBuilder, InstanceData.Data->GetCulledVertexRadiusScaleBuffer()));
 			Transitions.Add(RegisterAsSRV(GraphBuilder, InstanceData.Data->GetCulledVertexIdBuffer()));
 		}
 	}
diff --git a/Engine/Source/Runtime/Renderer/Public/HairStrandsInterface.h b/Engine/Source/Runtime/Renderer/Public/HairStrandsInterface.h
--- a/Engine/Source/Runtime/Renderer/Public/HairStrandsInterface.h
+++ b/Engine/Source/Runtime/Renderer/Public/HairStrandsInterface.h
@@ -84,7 +84,6 @@ BEGIN_SHADER_PARAMETER_STRUCT(FHairStrandsInstanceCullingParameters, RENDERER_AP
 	SHADER_PARAMETER(uint32, bCullingEnable)
 	SHADER_PARAMETER_RDG_BUFFER_SRV(Buffer, CullingIndirectBuffer)
 	SHADER_PARAMETER_RDG_BUFFER_SRV(Buffer, CullingIndexBuffer)
-	SHADER_PARAMETER_RDG_BUFFER_SRV(Buffer, CullingRadiusScaleBuffer)
 	RDG_BUFFER_ACCESS(CullingIndirectBufferArgs, ERHIAccess::IndirectArgs)
 END_SHADER_PARAMETER_STRUCT()
 
@@ -116,7 +115,6 @@ BEGIN_SHADER_PARAMETER_STRUCT(FHairStrandsInstanceCullingRawParameters, RENDERER
 	SHADER_PARAMETER(uint32, bCullingEnable)
 //	SHADER_PARAMETER_SRV(Buffer<uint>, CullingIndirectBuffer)
 	SHADER_PARAMETER_SRV(Buffer<uint>, CullingIndexBuffer)
-	SHADER_PARAMETER_SRV(Buffer<float>, CullingRadiusScaleBuffer)
 END_SHADER_PARAMETER_STRUCT()
 
 // Intermediate struct which can be referenced by FHairStrandsInstanceParameters for getting 
@@ -231,11 +229,9 @@ public:
 
 	const FRDGExternalBuffer& GetCulledCurveBuffer() const { return Culling->CulledCurveBuffer; }
 	const FRDGExternalBuffer& GetCulledVertexIdBuffer() const { return Culling->CulledVertexIdBuffer; }
-	const FRDGExternalBuffer& GetCulledVertexRadiusScaleBuffer() const { return Culling->CulledVertexRadiusScaleBuffer; }
 
 	FRDGExternalBuffer& GetCulledCurveBuffer() { return Culling->CulledCurveBuffer; }
 	FRDGExternalBuffer& GetCulledVertexIdBuffer() { return Culling->CulledVertexIdBuffer; }
-	FRDGExternalBuffer& GetCulledVertexRadiusScaleBuffer() { return Culling->CulledVertexRadiusScaleBuffer; }
 
 	bool GetCullingResultAvailable() const { return Culling->bCullingResultAvailable; }
 	void SetCullingResultAvailable(bool b) { Culling->bCullingResultAvailable = b; }
@@ -294,6 +290,7 @@ public:
 	RENDERER_API uint32 GetActiveStrandsPointCount() const;
 	RENDERER_API uint32 GetActiveStrandsCurveCount() const;
 	RENDERER_API float  GetActiveStrandsCoverageScale() const;
+	RENDERER_API float  GetActiveStrandsRadiusScale() const;
 
 	struct FVertexFactoryInput 
 	{
@@ -358,7 +355,6 @@ public:
 		/* Culling & LODing results for a hair group */ // Better to be transient?
 		FRDGExternalBuffer CulledCurveBuffer;
 		FRDGExternalBuffer CulledVertexIdBuffer;
-		FRDGExternalBuffer CulledVertexRadiusScaleBuffer;
 		bool bCullingResultAvailable = false;
 	};
 
@@ -404,6 +400,7 @@ public:
 	uint32 ContinuousLODCurveCount = 0;
 	float ContinuousLODScreenSize = 1.f;
 	float ContinuousLODCoverageScale = 1.f;
+	float ContinuousLODRadiusScale = 1.f;
 	FVector2f ContinuousLODScreenPos = FVector2f(0,0);
 	FBoxSphereBounds ContinuousLODBounds; 	//used by Continuous LOD
 
