Added VSMTexelDitherScale to the light proxy, analogous to r.Shadow.Virtual.SMRT.TexelDitherScale<Local/Directional>.
Controls the amount of temporal dither filtering applied to shadows for this light. Higher filtering can be used to hide aliasing from low-res shadows while lower filtering can be used to limit any filtering when we require crisp, detailed, high-res shadows.

diff --git a/Engine/Shaders/Private/VirtualShadowMaps/VirtualShadowMapProjectionDirectional.ush b/Engine/Shaders/Private/VirtualShadowMaps/VirtualShadowMapProjectionDirectional.ush
--- a/Engine/Shaders/Private/VirtualShadowMaps/VirtualShadowMapProjectionDirectional.ush
+++ b/Engine/Shaders/Private/VirtualShadowMaps/VirtualShadowMapProjectionDirectional.ush
@@ -242,11 +242,12 @@ FVirtualShadowMapSampleResult TraceDirectional(
 		// in some different way.
 		//
 		// 0.5 is arbitrary but makes it consistent with the previous dither scale cvar
+		float PerLightTexelDitherScale = TexelDitherScale * ProjectionData.TexelDitherScale;
 		float DitherScale = 0.0f;
-		if (TexelDitherScale > 0.0f)
+		if (PerLightTexelDitherScale > 0.0f)
 		{
 			DitherScale =
-				((0.5f / float(CalcLevelDimsTexels(0))) * TexelDitherScale * DistanceFromViewOrigin) /
+				((0.5f / float(CalcLevelDimsTexels(0))) * PerLightTexelDitherScale * DistanceFromViewOrigin) /
 				(exp2(ProjectionData.ClipmapLevel - ProjectionData.ResolutionLodBias));
 		}
 
diff --git a/Engine/Shaders/Private/VirtualShadowMaps/VirtualShadowMapProjectionSpot.ush b/Engine/Shaders/Private/VirtualShadowMaps/VirtualShadowMapProjectionSpot.ush
--- a/Engine/Shaders/Private/VirtualShadowMaps/VirtualShadowMapProjectionSpot.ush
+++ b/Engine/Shaders/Private/VirtualShadowMaps/VirtualShadowMapProjectionSpot.ush
@@ -370,7 +370,7 @@ FVirtualShadowMapSampleResult TraceLocalLight(
 		// We use a radially symmetric mapping here instead of the true function to minimize issues crossing cubemap faces.
 		const float ShadowDepthFunctionScale = abs(ProjectionData.ShadowViewToClipMatrix._43 / length2(ToLight));
 
-		const float DitherScale = TexelDitherScale * ShadowResolutionScale;
+		const float DitherScale = (TexelDitherScale * ProjectionData.TexelDitherScale) * ShadowResolutionScale;
 		const float MaxDitherSlope = MaxSlopeBias * DitherScale;
 		const float MaxDepthBias = MaxSlopeBias * ShadowResolutionScale * ShadowDepthFunctionScale;
 
diff --git a/Engine/Shaders/Private/VirtualShadowMaps/VirtualShadowMapProjectionStructs.ush b/Engine/Shaders/Private/VirtualShadowMaps/VirtualShadowMapProjectionStructs.ush
--- a/Engine/Shaders/Private/VirtualShadowMaps/VirtualShadowMapProjectionStructs.ush
+++ b/Engine/Shaders/Private/VirtualShadowMaps/VirtualShadowMapProjectionStructs.ush
@@ -33,6 +33,7 @@ struct FVirtualShadowMapProjectionShaderData
 	int VirtualShadowMapId;
 	bool bUnCached; // see VSM_PROJ_FLAG_UNCACHED
 	bool bUnreferenced; // see VSM_PROJ_FLAG_UNREFERENCED
+	float TexelDitherScale; // Per-light additional dither filtering
 };
 
 FVirtualShadowMapProjectionShaderData DecodeVirtualShadowMapProjectionData(ByteAddressBuffer ProjectionData, int VirtualShadowMapId)
@@ -80,6 +81,8 @@ FVirtualShadowMapProjectionShaderData DecodeVirtualShadowMapProjectionData(ByteA
 	Result.ClipmapLevelWPODistanceDisabledThresholdSquared
 		                                    = asfloat(ProjectionData.Load (Base + (16*17 + 4*1)));
 	
+	Result.TexelDitherScale					= asfloat(ProjectionData.Load (Base + (16*17 + 4*2)));
+	
 	// Computed data
 	Result.PreViewTranslation				= MakeDFVector3(PreViewTranslationHigh, PreViewTranslationLow);
 	Result.ClipmapWorldOriginOffset			= -NegativeClipmapWorldOriginOffset;
diff --git a/Engine/Source/Runtime/Engine/Private/LightSceneProxy.cpp b/Engine/Source/Runtime/Engine/Private/LightSceneProxy.cpp
--- a/Engine/Source/Runtime/Engine/Private/LightSceneProxy.cpp
+++ b/Engine/Source/Runtime/Engine/Private/LightSceneProxy.cpp
@@ -139,6 +139,7 @@ FLightSceneProxy::FLightSceneProxy(const ULightComponent* InLightComponent)
 
 		bTransmission = false;
 	}
+	VSMTexelDitherScale = 1.0f;
 }
 
 FLightSceneProxy::~FLightSceneProxy() = default;
diff --git a/Engine/Source/Runtime/Engine/Public/LightSceneProxy.h b/Engine/Source/Runtime/Engine/Public/LightSceneProxy.h
--- a/Engine/Source/Runtime/Engine/Public/LightSceneProxy.h
+++ b/Engine/Source/Runtime/Engine/Public/LightSceneProxy.h
@@ -110,6 +110,11 @@ public:
 	/** Whether this light should create CSM for dynamic objects only (forward renderer) */
 	ENGINE_API virtual bool UseCSMForDynamicObjects() const;
 
+	ENGINE_API float GetVSMTexelDitherScale() const
+	{
+		return VSMTexelDitherScale;
+	}
+
 	/** Returns the number of view dependent shadows this light will create, not counting distance field shadow cascades. */
 	virtual uint32 GetNumViewDependentWholeSceneShadows(const FSceneView& View, bool bPrecomputedLightingIsValid) const { return 0; }
 
@@ -457,6 +462,9 @@ protected:
 	/** The name of the level the light is in. */
 	FName LevelName;
 
+	/** Used to control the amount of additional dither filtering applied to shadows for each light. */
+	float VSMTexelDitherScale;
+
 	/** Only for whole scene directional lights, if FarShadowCascadeCount > 0 and FarShadowDistance >= WholeSceneDynamicShadowRadius, where far shadow cascade should end. */
 	float FarShadowDistance;
 
diff --git a/Engine/Source/Runtime/Renderer/Private/Shadows/ShadowSceneRenderer.cpp b/Engine/Source/Runtime/Renderer/Private/Shadows/ShadowSceneRenderer.cpp
--- a/Engine/Source/Runtime/Renderer/Private/Shadows/ShadowSceneRenderer.cpp
+++ b/Engine/Source/Runtime/Renderer/Private/Shadows/ShadowSceneRenderer.cpp
@@ -219,6 +219,7 @@ FVirtualShadowMapProjectionShaderData FShadowSceneRenderer::GetLocalLightProject
 	Data.LightRadius							= ProjectedShadowInfo->GetLightSceneInfo().Proxy->GetRadius();
 	Data.LightDirection                         = FVector3f(0, 0, 0);	// Unused for local lights
 	Data.Flags									= Flags;
+	Data.TexelDitherScale						= ProjectedShadowInfo->GetLightSceneInfo().Proxy->GetVSMTexelDitherScale();
 
 	return Data;
 }
diff --git a/Engine/Source/Runtime/Renderer/Private/VirtualShadowMaps/VirtualShadowMapArray.h b/Engine/Source/Runtime/Renderer/Private/VirtualShadowMaps/VirtualShadowMapArray.h
--- a/Engine/Source/Runtime/Renderer/Private/VirtualShadowMaps/VirtualShadowMapArray.h
+++ b/Engine/Source/Runtime/Renderer/Private/VirtualShadowMaps/VirtualShadowMapArray.h
@@ -110,9 +110,10 @@ struct FVirtualShadowMapProjectionShaderData
 	uint32 Flags = 0U;
 	// This clipmap level should allow WPO if this value is less than InstanceWPODisableDistanceSquared
 	float ClipmapLevelWPODistanceDisableThresholdSquared = 0.0f;
-
+	float TexelDitherScale;
+	
 	// Seems the FMatrix forces 16-byte alignment
-	float Padding[2];
+	float Padding;
 };
 static_assert(sizeof(FVirtualShadowMapProjectionShaderData) == (16*18), "FVirtualShadowMapProjectionShaderData does not match size in shader. See VirtualShadowMapProjectionStructs.ush.");
 
diff --git a/Engine/Source/Runtime/Renderer/Private/VirtualShadowMaps/VirtualShadowMapClipmap.cpp b/Engine/Source/Runtime/Renderer/Private/VirtualShadowMaps/VirtualShadowMapClipmap.cpp
--- a/Engine/Source/Runtime/Renderer/Private/VirtualShadowMaps/VirtualShadowMapClipmap.cpp
+++ b/Engine/Source/Runtime/Renderer/Private/VirtualShadowMaps/VirtualShadowMapClipmap.cpp
@@ -353,7 +353,7 @@ FVirtualShadowMapProjectionShaderData FVirtualShadowMapClipmap::ComputeProjectio
 	Data.ClipmapLevelWPODistanceDisableThresholdSquared = static_cast<float>(Level.WPODistanceDisableThresholdSquared);
 	Data.LightSourceRadius = GetLightSceneInfo().Proxy->GetSourceRadius();
 	Data.Flags = PerLightCacheEntry->IsUncached() ? VSM_PROJ_FLAG_UNCACHED : 0U;
-
+	Data.TexelDitherScale = GetLightSceneInfo().Proxy->GetVSMTexelDitherScale();
 	return Data;
 }
 
