Path Tracer: Fix possible NaNs in specular highlights when Anisotropy is set but Tangent is not

If the tangent frame basis construction produces a 0 vector when doing the cross product, we could wind up with NaNs in the basis which quickly propagated to the shaded result. Simply detect this case an use an arbitrary basis as a fallback, as this case is not well defined anyway.

diff --git a/Engine/Shaders/Private/PathTracing/Material/PathTracingGlossy.ush b/Engine/Shaders/Private/PathTracing/Material/PathTracingGlossy.ush
--- a/Engine/Shaders/Private/PathTracing/Material/PathTracingGlossy.ush
+++ b/Engine/Shaders/Private/PathTracing/Material/PathTracingGlossy.ush
@@ -37,22 +37,27 @@ float2 GetGGXAlpha(float Roughness, float Anisotropy)
 float3x3 GetGGXBasis(float Roughness, float Anisotropy, float3 WorldNormal, float3 WorldTangent, out float2 Alpha)
 {
 	Roughness = max(Roughness, GGX_MIN_ROUGHNESS);
-	float3x3 Basis;
 	if (abs(Anisotropy) >= GGX_MIN_ANISOTROPY)
 	{
 		// Material is anisotropic, our shading frame should use the tangent vector
 		GetAnisotropicRoughness(Roughness * Roughness, Anisotropy, Alpha.x, Alpha.y);
-		float3 X = WorldTangent;
-		float3 Y = normalize(cross(WorldNormal, X));
-		Basis = float3x3(cross(Y, WorldNormal), Y, WorldNormal);
+		float3 Y = cross(WorldNormal, WorldTangent);
+		float LenY2 = length2(Y);
+		if (LenY2 > 0)
+		{
+			// Tangent vector is valid
+			Y *= rsqrt(LenY2);
+			return float3x3(cross(Y, WorldNormal), Y, WorldNormal);
+		}
+		// Fallthrough to build a default basis if we couldn't build one
 	}
 	else
 	{
-		// No anisotropy, just pick an arbitrary basis
+		// No anisotropy, use simpler formula
 		Alpha = Roughness * Roughness;
-		Basis = GetTangentBasis(WorldNormal);
 	}
-	return Basis;
+	// pick an arbitrary basis (no anisotropy, or invalid tangent vector)
+	return GetTangentBasis(WorldNormal);
 }
 
 float3x3 GetGGXBasis(float Roughness, float3 WorldNormal, out float2 Alpha)
