Heterogeneous Volumes: Add partial final step to address harsh boundary when intersecting hard surfaces

diff --git a/Engine/Shaders/Private/HeterogeneousVolumes/HeterogeneousVolumesRayMarchingUtils.ush b/Engine/Shaders/Private/HeterogeneousVolumes/HeterogeneousVolumesRayMarchingUtils.ush
--- a/Engine/Shaders/Private/HeterogeneousVolumes/HeterogeneousVolumesRayMarchingUtils.ush
+++ b/Engine/Shaders/Private/HeterogeneousVolumes/HeterogeneousVolumesRayMarchingUtils.ush
@@ -317,10 +317,19 @@ void RayMarchSingleScattering(
 #endif // #if APPLY_FOG_INSCATTERING
 
 	float WorldShadowBias = CalcShadowBias();
+	float StepSize = RayMarchingContext.StepSize;
+	float JitterSize = StepSize * RayMarchingContext.Jitter;
 
+	float LocalHitT = RayMarchingContext.LocalRayTMin + JitterSize;
 	for (uint StepIndex = 0; StepIndex < StepCount; ++StepIndex)
 	{
-		float LocalHitT = RayMarchingContext.LocalRayTMin + RayMarchingContext.StepSize * (RayMarchingContext.Jitter + StepIndex);
+		// Correct the final step size
+		if (StepIndex == StepCount - 1)
+		{
+			float UnjitteredLocalHitT = LocalHitT - JitterSize;
+			StepSize = RayMarchingContext.LocalRayTMax - UnjitteredLocalHitT;
+		}
+
 		float3 LocalPosition = RayMarchingContext.LocalRayOrigin + RayMarchingContext.LocalRayDirection * LocalHitT;
 		float3 WorldPosition = RayMarchingContext.WorldRayOrigin + RayMarchingContext.WorldRayDirection * LocalHitT * RayMarchingContext.LocalToWorldScale;
 		FVolumeSampleContext SampleContext = CreateVolumeSampleContext(LocalPosition, WorldPosition, RayMarchingContext.MipLevel);
@@ -328,7 +337,7 @@ void RayMarchSingleScattering(
 		if (RayMarchingContext.bApplyEmissionAndTransmittance)
 		{
 			float3 Emission = SampleEmission(SampleContext);
-			Radiance += Emission * RayMarchingContext.StepSize * Transmittance;
+			Radiance += Emission * StepSize * Transmittance;
 
 #if DIM_USE_LUMEN_GI
 			FTwoBandSHVectorRGB TranslucencyGISH = GetTranslucencyGIVolumeLighting(DFPromote(WorldPosition), PrimaryView.WorldToClip, true);
@@ -344,7 +353,7 @@ void RayMarchSingleScattering(
 			float3 ScatteringCoefficient = Albedo * Extinction;
 
 			// Note: Phase is accounted for in spherical hamonic calculation
-			Radiance += IndirectInscattering * ScatteringCoefficient * Transmittance * RayMarchingContext.StepSize;
+			Radiance += IndirectInscattering * ScatteringCoefficient * Transmittance * StepSize;
 #endif // DIM_USE_LUMEN_GI
 		}
 
@@ -358,7 +367,7 @@ void RayMarchSingleScattering(
 
 				float3 ScatteringCoefficient = Albedo * Extinction;
 				float IsotropicPhase = 1.0 / (4.0 * PI);
-				Radiance += Inscattering * ScatteringCoefficient * IsotropicPhase * Transmittance * RayMarchingContext.StepSize;
+				Radiance += Inscattering * ScatteringCoefficient * IsotropicPhase * Transmittance * StepSize;
 			}
 
 #if APPLY_FOG_INSCATTERING
@@ -367,7 +376,7 @@ void RayMarchSingleScattering(
 		}
 
 		// Accumulate transmittance for the next evaluation
-		Transmittance *= exp(-Extinction * RayMarchingContext.StepSize);
+		Transmittance *= exp(-Extinction * StepSize);
 
 		float Epsilon = 1.0e-7;
 		if (all(Transmittance < Epsilon))
@@ -375,6 +384,8 @@ void RayMarchSingleScattering(
 			Transmittance = 0.0;
 			break;
 		}
+
+		LocalHitT += StepSize;
 	}
 
 #if APPLY_FOG_INSCATTERING
