Holdout - fixed Volumetric cloud to match expectations.

diff --git a/Engine/Shaders/Private/VolumetricCloud.usf b/Engine/Shaders/Private/VolumetricCloud.usf
--- a/Engine/Shaders/Private/VolumetricCloud.usf
+++ b/Engine/Shaders/Private/VolumetricCloud.usf
@@ -435,7 +435,9 @@ void MainCommon(in FMaterialPixelParameters MaterialParameters, in float4 SvPosi
 #if CLOUD_MIN_AND_MAX_DEPTH
 	, inout float4 OutColor1 
 #endif
-	, inout float4 OutDepth)
+	, inout float4 OutDepth
+	, inout float  OutHoldOutCoverage
+	, inout float  OutAboveCloudHoldOut)
 {
 	ResolvedView = ResolveView();
 
@@ -1415,6 +1417,8 @@ void MainCommon(in FMaterialPixelParameters MaterialParameters, in float4 SvPosi
 	#if CLOUD_MIN_AND_MAX_DEPTH
 		ClosestLuminance *= 0;
 	#endif
+
+		OutHoldOutCoverage = saturate(1.0f - dot(TransmittanceToView.MATSUF, 1.0f / 3.0f));
 	}
 #endif
 
@@ -1462,6 +1466,22 @@ void MainCommon(in FMaterialPixelParameters MaterialParameters, in float4 SvPosi
 		// skip if MeanTransmittance is 1
 		float MeanCoverage = 1.0 - MeanTransmittance;
 		Luminance = AerialPerspective.rgb * MeanCoverage + AerialPerspective.a * Luminance;
+
+	#if POST_PROCESS_ALPHA
+		if(View.RenderingReflectionCaptureMask == 0.0f)
+		{
+			const float  AerialPespectiveCoverage = saturate(1.0 - AerialPerspective.a);
+			if (!IsSkyAtmosphereHoldout(View.EnvironmentComponentsFlags))
+			{
+				// Fade out holdout coverage as a function of the aerial perspective above, if the sky atmosphere is not hold out
+				OutHoldOutCoverage   = lerp(OutHoldOutCoverage,   0.0f, AerialPespectiveCoverage);
+			}
+			else // if(IsSkyAtmosphereHoldout(View.EnvironmentComponentsFlags))
+			{
+				OutAboveCloudHoldOut = lerp(OutAboveCloudHoldOut, 1.0f, AerialPespectiveCoverage * MeanCoverage);
+			}
+		}
+	#endif
 	}
 
 
@@ -1504,10 +1524,26 @@ void MainCommon(in FMaterialPixelParameters MaterialParameters, in float4 SvPosi
 		// skip if MeanTransmittance is 1
 		float MeanCoverage = 1.0 - MeanTransmittance;
 		Luminance = HeightFogInscatteringAndTransmittance.rgb * MeanCoverage + HeightFogInscatteringAndTransmittance.a * Luminance;
+
+	#if POST_PROCESS_ALPHA
+		if(View.RenderingReflectionCaptureMask == 0.0f)
+		{
+			const float  ExponentialFogCoverage = saturate(1.0 - HeightFogInscatteringAndTransmittance.a);
+			if (!IsExponentialFogHoldout(View.EnvironmentComponentsFlags))
+			{
+				// Fade out holdout coverage as a function of the fog above, if the fog is not hold out
+				OutHoldOutCoverage   = lerp(OutHoldOutCoverage,   0.0f, ExponentialFogCoverage);
+			}
+			else //if(IsExponentialFogHoldout(View.EnvironmentComponentsFlags))
+			{
+				OutAboveCloudHoldOut = lerp(OutAboveCloudHoldOut, 1.0f, ExponentialFogCoverage * MeanCoverage);
+			}
+		}
+	#endif
 	}
 
 
-	// Apply atmospehre and and fog on closest
+	// Apply atmosphere and fog on closest if needed
 #if CLOUD_MIN_AND_MAX_DEPTH
 	{
 		float3 FogSampleWorldPositionRelativeToCameraCm = ClosestTAP * Raydir;
@@ -1630,7 +1666,9 @@ void MainPS(
 	FMaterialPixelParameters MaterialParameters = MakeInitializedMaterialPixelParameters();
 	
 	float4 OutDepth4 = 0;
-	MainCommon(MaterialParameters, SvPosition, OutColor0, OutDepth4);
+	float OutHoldOutCoverage = 0.0f;
+	float OutAboveCloudHoldOut = 0.0f;
+	MainCommon(MaterialParameters, SvPosition, OutColor0, OutDepth4, OutHoldOutCoverage, OutAboveCloudHoldOut);
 	OutDepth = float4(OutDepth4.xy, 0.0f, 0.0f);	// PS path never writes min and max tracing depth in zw.
 }
 
@@ -1664,11 +1702,15 @@ void MainCS(uint2 DispatchThreadId : SV_DispatchThreadID)
 
 		FMaterialPixelParameters MaterialParameters = MakeInitializedMaterialPixelParameters();
 
+		float OutHoldOutCoverage = 0.0f;
+		float OutAboveCloudHoldOut = 0.0f;
 		MainCommon(MaterialParameters, SvPosition, OutColor0
 #if CLOUD_MIN_AND_MAX_DEPTH
 			, OutColor1
 #endif
-			, OutDepth);
+			, OutDepth
+			, OutHoldOutCoverage
+			, OutAboveCloudHoldOut);
 
 		if (TargetCubeFace >= 0)
 		{
@@ -1688,7 +1730,24 @@ void MainCS(uint2 DispatchThreadId : SV_DispatchThreadID)
 			if (bBlendCloudColor)
 			{
 				float4 CurColor0 = OutCloudColor0[DispatchThreadId];
-				OutCloudColor0[DispatchThreadId] = float4(CurColor0.rgb * OutColor0.a + OutColor0.rgb, CurColor0.a * OutColor0.a);
+			#if POST_PROCESS_ALPHA
+				if (View.RenderingReflectionCaptureMask == 0.0f)
+				{
+					if(IsVolumetricCloudHoldout(View.EnvironmentComponentsFlags))
+					{
+						// Holdout for when we blend directly into the scene color Buffer.
+						OutCloudColor0[DispatchThreadId] = float4(CurColor0.rgb * OutColor0.a + OutColor0.rgb, CurColor0.a * (1.0 - OutHoldOutCoverage) + OutHoldOutCoverage);
+					}
+					else
+					{
+						OutCloudColor0[DispatchThreadId] = float4(CurColor0.rgb * OutColor0.a + OutColor0.rgb, CurColor0.a * OutColor0.a + OutAboveCloudHoldOut);
+					}
+				}
+			#else
+				{
+					OutCloudColor0[DispatchThreadId] = float4(CurColor0.rgb * OutColor0.a + OutColor0.rgb, CurColor0.a);
+				}
+			#endif
 			}
 			else
 			{
diff --git a/Engine/Source/Runtime/Renderer/Private/FogRendering.cpp b/Engine/Source/Runtime/Renderer/Private/FogRendering.cpp
--- a/Engine/Source/Runtime/Renderer/Private/FogRendering.cpp
+++ b/Engine/Source/Runtime/Renderer/Private/FogRendering.cpp
@@ -327,8 +327,8 @@ static void RenderViewFog(
 			}
 			else
 			{
-				// Need to store throughput into alpha so we multiply with transmittance=1-Coverage.
-				GraphicsPSOInit.BlendState = TStaticBlendState<CW_RGBA, BO_Add, BF_One, BF_InverseSourceAlpha, BO_Add, BF_Zero, BF_InverseSourceAlpha>::GetRHI();
+				// Same color blending as without alpha channel. Disabled alpha blending writes because alpha holdout from other component must be preserved.
+				GraphicsPSOInit.BlendState = TStaticBlendState<CW_RGB, BO_Add, BF_One, BF_InverseSourceAlpha>::GetRHI();
 			}
 		}
 		else
