Fix for out of bounds GPU memory accesses (on some platforms) in ThreadGroupPrefixSum and VSM directional projection (Short circuit VSM page table lookup with out of range clipmap levels, avoids an out of bounds access on some platforms)

diff --git a/Engine/Shaders/Private/ThreadGroupPrefixSum.ush b/Engine/Shaders/Private/ThreadGroupPrefixSum.ush
--- a/Engine/Shaders/Private/ThreadGroupPrefixSum.ush
+++ b/Engine/Shaders/Private/ThreadGroupPrefixSum.ush
@@ -52,7 +52,14 @@ groupshared uint ThreadGroupPrefixSumWorkspace[2][NUM_THREADS_PER_GROUP];
 			GroupMemoryBarrierWithGroupSync(); \
 		} \
 		GroupSum = CastFromUint(ThreadGroupPrefixSumWorkspace[Curr][NUM_THREADS_PER_GROUP - 1]); \
-		return (GroupThreadIndex == 0U) ? CastFromUint(0) : CastFromUint(ThreadGroupPrefixSumWorkspace[Curr][GroupThreadIndex - 1]); \
+		if (GroupThreadIndex == 0U) \
+		{ \
+			return CastFromUint(0); \
+		} \
+		else \
+		{ \
+			return CastFromUint(ThreadGroupPrefixSumWorkspace[Curr][GroupThreadIndex - 1]); \
+		} \
 	} \
 	ValType ThreadGroupPrefixSum(ValType Value, uint GroupThreadIndex) \
 	{ \
diff --git a/Engine/Shaders/Private/VirtualShadowMaps/VirtualShadowMapProjectionDirectional.ush b/Engine/Shaders/Private/VirtualShadowMaps/VirtualShadowMapProjectionDirectional.ush
--- a/Engine/Shaders/Private/VirtualShadowMaps/VirtualShadowMapProjectionDirectional.ush
+++ b/Engine/Shaders/Private/VirtualShadowMaps/VirtualShadowMapProjectionDirectional.ush
@@ -46,6 +46,11 @@ int GetMappedClipmapId(int VirtualShadowMapId, float3 RayOriginTranslatedWorld)
 	float ClipmapLevelFloat = CalcBiasedAbsoluteClipmapLevelForSampling(BaseProjectionData, RayOriginTranslatedWorld);
 	int ClipmapIndex = max(0, int(floor(ClipmapLevelFloat)) - BaseProjectionData.ClipmapLevel);
 
+	if (ClipmapIndex >= BaseProjectionData.ClipmapLevelCountRemaining)
+	{
+		return INDEX_NONE;
+	}
+
 	// Clipmap origin should be near shaded samples, so shadow translated world should be regular range
 	int ClipmapId = VirtualShadowMapId + ClipmapIndex;
 	const FVirtualShadowMapProjectionShaderData ProjectionData = GetVirtualShadowMapProjectionData(ClipmapId);
@@ -227,10 +232,11 @@ FVirtualShadowMapSampleResult TraceDirectional(
 	Result.bValid = true;	// TODO: false if all samples of all rays miss pages?
 	Result.ShadowFactor = 0.0f;
 
-	if (!bCullBackface || !IsBackfaceToDirectionalLight(WorldNormal, Light.Direction, Light.SourceRadius))
+	// Find the best resolution mapped clipmap at the ray origin location
+	const int ClipmapId = GetMappedClipmapId(VirtualShadowMapId, TranslatedWorldPosition);
+
+	if ((!bCullBackface || !IsBackfaceToDirectionalLight(WorldNormal, Light.Direction, Light.SourceRadius)) && ClipmapId >= 0)
 	{
-		// Find the best resolution mapped clipmap at the ray origin location
-		const int ClipmapId = GetMappedClipmapId(VirtualShadowMapId, TranslatedWorldPosition);
 		const FVirtualShadowMapProjectionShaderData ProjectionData = GetVirtualShadowMapProjectionData(ClipmapId);
 
 		// This function is designed to vary in depth in roughly the same way as shadow texel density does,
