Fix typo in particle sprite position variable type.

diff --git a/Engine/Shaders/Private/ParticleSpriteVertexFactory.ush b/Engine/Shaders/Private/ParticleSpriteVertexFactory.ush
--- a/Engine/Shaders/Private/ParticleSpriteVertexFactory.ush
+++ b/Engine/Shaders/Private/ParticleSpriteVertexFactory.ush
@@ -98,7 +98,7 @@ struct FVertexFactoryInterpolantsVSToPS
 struct FVertexFactoryIntermediates
 {
 	/** The position of the particle in translated world space. */
-	float ParticleTranslatedWorldPosition;
+	float3 ParticleTranslatedWorldPosition;
 	/** The position of the vertex in translated world space. */
 	float3 VertexWorldPosition;
 	/** Particle translated world space position and size. */
