Fix bad conversion function UnpackColorRGB24 from packed uints into floats from 30763411 that broke actor coloration and wireframe color.

diff --git a/Engine/Shaders/Private/SceneData.ush b/Engine/Shaders/Private/SceneData.ush
--- a/Engine/Shaders/Private/SceneData.ush
+++ b/Engine/Shaders/Private/SceneData.ush
@@ -150,7 +150,12 @@ struct FPrimitiveSceneData
 float3 UnpackColorRGB24(float PackedColorFloat)
 {
 	uint PackedColor = asuint(PackedColorFloat);
-	return float3((PackedColor >> 24) && 0xFF, (PackedColor >> 16) && 0xFF, (PackedColor >> 8) && 0xFF);
+	return float3
+	(
+		float((PackedColor >> 24u) & 0xFF),
+		float((PackedColor >> 16u) & 0xFF),
+		float((PackedColor >>  8u) & 0xFF)
+	) * (1.0f / 255.0f);
 }
 
 // Fetch from Primitive uniform buffer
