PR #9733: Fix up type of "OutColor" from "HALF4_TYPE" to "HALF_TYPE" in SunBlurPS_Mobile function ( PostProcessMobile.usf )

diff --git a/Engine/Shaders/Private/PostProcessMobile.usf b/Engine/Shaders/Private/PostProcessMobile.usf
--- a/Engine/Shaders/Private/PostProcessMobile.usf
+++ b/Engine/Shaders/Private/PostProcessMobile.usf
@@ -1066,7 +1066,7 @@ void SunBlurVS_Mobile(
 
 void SunBlurPS_Mobile(
 	float2 InUVs[8] : TEXCOORD0,
-	out HALF4_TYPE OutColor : SV_Target0
+	out HALF_TYPE OutColor : SV_Target0
 	)
 {
 	OutColor = 
