Apply sky lighting with ambient occlusion on mobile deferred.

diff --git a/Engine/Shaders/Private/MobileDeferredShading.usf b/Engine/Shaders/Private/MobileDeferredShading.usf
--- a/Engine/Shaders/Private/MobileDeferredShading.usf
+++ b/Engine/Shaders/Private/MobileDeferredShading.usf
@@ -91,7 +91,7 @@ half3 SkyLightDiffuseMobile(FGBufferData GBuffer, half3 DiffuseColor, half3 V, i
 	half3 DiffuseLookup = GetSkySHDiffuseSimple(SkyDiffuseLookUpNormal) * View.SkyLightColor.rgb;
 
 	// And accumulate the lighting
-	Lighting += DiffuseLookup  * DiffuseColor;
+	Lighting += DiffuseLookup  * DiffuseColor * GBuffer.GBufferAO;
 
 #if ALLOW_STATIC_LIGHTING
 	IndirectIrradiance += Luminance(DiffuseLookup);
