Fix incorrect shadow mask for hair when using
* VSM and cluster shading
* VSM and OnePassShadowmap

diff --git a/Engine/Shaders/Private/ClusteredDeferredShadingPixelShader.usf b/Engine/Shaders/Private/ClusteredDeferredShadingPixelShader.usf
--- a/Engine/Shaders/Private/ClusteredDeferredShadingPixelShader.usf
+++ b/Engine/Shaders/Private/ClusteredDeferredShadingPixelShader.usf
@@ -292,7 +292,6 @@ float4 GetLightGridLocalLighting(
 		}
 
 		#if USE_HAIR_LIGHTING
-		LightAttenuation = 1.0f; // TODO: use shadow mask bits (e.g., subpixel shadow mask)
 		if (LightAttenuation.x > 0)
 		{
 			const float3 L = normalize(LightData.TranslatedWorldPosition - TranslatedWorldPosition);
diff --git a/Engine/Shaders/Private/VirtualShadowMaps/VirtualShadowMapProjectionComposite.usf b/Engine/Shaders/Private/VirtualShadowMaps/VirtualShadowMapProjectionComposite.usf
--- a/Engine/Shaders/Private/VirtualShadowMaps/VirtualShadowMapProjectionComposite.usf
+++ b/Engine/Shaders/Private/VirtualShadowMaps/VirtualShadowMapProjectionComposite.usf
@@ -47,7 +47,7 @@ void VirtualShadowMapCompositeFromMaskBitsPS(
 {
 	const uint2 PixelPos = uint2(SvPosition.xy);
 
-	const float DeviceZ = SceneTexturesStruct.SceneDepthTexture.Load(int3(PixelPos, 0)).r;
+	const float DeviceZ = InputDepthTexture.Load(int3(PixelPos, 0)).r;
 	const float SceneDepth = ConvertFromDeviceZ(DeviceZ);
 
 	// Possibly needed to do fallback single sample VSM lookup
