Fix deep shadow map rendering.

This has recently broke with ortho view added for groom.

#swarm 31678185
#hardlock juan.canda

diff --git a/Engine/Shaders/Private/Common.ush b/Engine/Shaders/Private/Common.ush
--- a/Engine/Shaders/Private/Common.ush
+++ b/Engine/Shaders/Private/Common.ush
@@ -1294,9 +1294,9 @@ float2 GetScreenPositionForProjectionType(float2 ScreenPosition, float SceneDept
 	return select(IsOrthoProjection(), ScreenPosition, ScreenPosition * SceneDepth);
 }
 
-float ConvertGivenDepthRadiusForProjectionType(float Radius, float SceneDepth)
+float ConvertGivenDepthRadiusForProjectionType(float Radius, float SceneDepth, bool bForceOrthoView = false)
 {
-	return select(IsOrthoProjection(), Radius, Radius * SceneDepth);
+	return select(IsOrthoProjection() || bForceOrthoView, Radius, Radius * SceneDepth);
 }
 
 float GetDepthPixelRadiusForProjectionType(float SceneDepth)
diff --git a/Engine/Shaders/Private/HairStrands/HairStrandsDeepShadowVS.usf b/Engine/Shaders/Private/HairStrands/HairStrandsDeepShadowVS.usf
--- a/Engine/Shaders/Private/HairStrands/HairStrandsDeepShadowVS.usf
+++ b/Engine/Shaders/Private/HairStrands/HairStrandsDeepShadowVS.usf
@@ -46,6 +46,7 @@ void Main(
 	HairViewInfo.TranslatedWorldCameraOrigin = ResolvedView.TranslatedWorldCameraOrigin;
 	HairViewInfo.ViewForward = ResolvedView.ViewForward;
 	HairViewInfo.RadiusAtDepth1 = HairRenderInfo.RadiusAtDepth1Primary;
+	HairViewInfo.bIsOrthoView = HairRenderInfo.bIsOrthoView;
 
 	#if MESH_RENDER_MODE == MESH_RENDER_DEPTH || MESH_RENDER_MODE == MESH_RENDER_DOM
 	if (HairRenderInfo.bIsGPUDriven)
@@ -68,7 +69,7 @@ void Main(
 	float HairCoverage = 1;
   #ifdef HAIR_STRAND_MESH_FACTORY	
 	const float CurrentDepth = Output.Position.z / Output.Position.w;
-	const float PixelRadius = ConvertGivenDepthRadiusForProjectionType(HairRenderInfo.RadiusAtDepth1Primary, CurrentDepth);
+	const float PixelRadius = HairRenderInfo.bIsOrthoView ? HairRenderInfo.RadiusAtDepth1Primary : (CurrentDepth * HairRenderInfo.RadiusAtDepth1Primary);
 	const float StrandRealRadius = VFIntermediates.HairDimensions.y;
 	const float Coverage = StrandRealRadius / max(StrandRealRadius, PixelRadius);
 	HairCoverage = Coverage * VFIntermediates.HairDensity;
diff --git a/Engine/Shaders/Private/HairStrands/HairStrandsVertexFactory.ush b/Engine/Shaders/Private/HairStrands/HairStrandsVertexFactory.ush
--- a/Engine/Shaders/Private/HairStrands/HairStrandsVertexFactory.ush
+++ b/Engine/Shaders/Private/HairStrands/HairStrandsVertexFactory.ush
@@ -655,7 +655,8 @@ float4 ComputeViewAlignedWorldPosition(FVertexFactoryInput Input, float3 WorldTa
 
 	// Minimal radius to snap the strand to a sample/pixel center (to avoid aliasing)
 	const float DistanceToCamera = length(HairViewInfo.TranslatedWorldCameraOrigin - WorldPosition.xyz);
-	const float MinStrandHairRadius = ConvertGivenDepthRadiusForProjectionType(HairViewInfo.RadiusAtDepth1, DistanceToCamera);
+	const float MinStrandHairRadius = ConvertGivenDepthRadiusForProjectionType(HairViewInfo.RadiusAtDepth1, DistanceToCamera, HairViewInfo.bIsOrthoView);
+
 	const float3 ViewDir = -HairViewInfo.ViewForward;
 	const float3 Right = normalize(cross(WorldTangent, ViewDir));
 	const float3 OutWorldPosition = WorldPosition.xyz + (VertexInfo.IsLeft ? -Right : Right) * max(WorldStrandRadius, MinStrandHairRadius);
@@ -678,7 +679,7 @@ float4 VertexFactoryGetWorldPosition_Visibility(FVertexFactoryInput Input, FHair
 
 	// Minimal radius to snap the strand to a sample/pixel center (to avoid aliasing)
 	const float DistanceToCamera = length(HairViewInfo.TranslatedWorldCameraOrigin - WorldPosition.xyz);
-	const float MinStrandHairRadius = ConvertGivenDepthRadiusForProjectionType(HairViewInfo.RadiusAtDepth1, DistanceToCamera);
+	const float MinStrandHairRadius = ConvertGivenDepthRadiusForProjectionType(HairViewInfo.RadiusAtDepth1, DistanceToCamera, HairViewInfo.bIsOrthoView);
 	const float3 ViewDir = -HairViewInfo.ViewForward;
 	const float3 Right = normalize(cross(WorldTangent, ViewDir));
 	const float3 OutWorldPosition = WorldPosition.xyz + (VertexInfo.IsLeft ? -Right : Right) * max(ControlPoint.WorldRadius, MinStrandHairRadius);
@@ -729,6 +730,7 @@ float4 VertexFactoryGetWorldPosition(FVertexFactoryInput Input, FVertexFactoryIn
 	HairViewInfo.TranslatedWorldCameraOrigin = ResolvedView.TranslatedWorldCameraOrigin;
 	HairViewInfo.ViewForward = ResolvedView.ViewForward;
 	HairViewInfo.RadiusAtDepth1 = HairRenderInfo.RadiusAtDepth1Primary;
+	HairViewInfo.bIsOrthoView = HairRenderInfo.bIsOrthoView;
 
 	return VertexFactoryGetWorldPosition(Input, Intermediates, HairViewInfo);
 }
diff --git a/Engine/Shaders/Private/HairStrands/HairStrandsVisibilityCommon.ush b/Engine/Shaders/Private/HairStrands/HairStrandsVisibilityCommon.ush
--- a/Engine/Shaders/Private/HairStrands/HairStrandsVisibilityCommon.ush
+++ b/Engine/Shaders/Private/HairStrands/HairStrandsVisibilityCommon.ush
@@ -344,6 +344,7 @@ struct FHairRenderInfo
 	float VelocityMagnitudeScale;
 
 	bool bIsGPUDriven;
+	bool bIsOrthoView; // Only used for shadow view
 
 	float3 ViewForward;
 	float3 TranslatedWorldCameraOrigin;
@@ -357,6 +358,7 @@ FHairRenderInfo GetHairRenderInfo(float4 ViewHairRenderInfo, uint ViewHairRender
 	Info.VelocityMagnitudeScale = ViewHairRenderInfo.w;
 
 	const uint BitField = ViewHairRenderInfoBits;
+	Info.bIsOrthoView = (BitField & 0x1) != 0;
 	Info.bIsGPUDriven = (BitField & 0x2) != 0;
 
 	return Info;
@@ -367,6 +369,7 @@ struct FHairViewInfo
 	float3 TranslatedWorldCameraOrigin;
 	float3 ViewForward;
 	float  RadiusAtDepth1;
+	bool   bIsOrthoView;
 };
 
 ////////////////////////////////////////////////////////////////////////////////////////////////////////////////////////////////
