Holdout + RenderInMain for skyatmosphere, fog and volumetric cloud.

diff --git a/Engine/Shaders/Private/HeightFogCommon.ush b/Engine/Shaders/Private/HeightFogCommon.ush
--- a/Engine/Shaders/Private/HeightFogCommon.ush
+++ b/Engine/Shaders/Private/HeightFogCommon.ush
@@ -223,8 +223,12 @@ half4 GetExponentialHeightFog(float3 WorldPositionRelativeToCamera, float Exclud
 
 half4 CalculateHeightFog(float3 WorldPositionRelativeToCamera)
 {
-	float ExcludeDistance = 0;
+	if (View.RenderingReflectionCaptureMask == 0.0f && !IsExponentialFogRenderedInMain(View.EnvironmentComponentsFlags))
+	{
+		return float4(0.0f, 0.0f, 0.0f, 1.0f);
+	}
 
+	float ExcludeDistance = 0;
 #if PERMUTATION_SUPPORT_VOLUMETRIC_FOG
 	// Volumetric fog covers up to MaxDistance along ViewZ, exclude analytical fog from this range
 	float CosAngle = dot(normalize(WorldPositionRelativeToCamera), View.ViewForward);
@@ -245,6 +249,11 @@ half4 CalculateHeightFog(float3 WorldPositionRelativeToCamera)
 float4 CombineVolumetricFog(float4 GlobalFog, float3 VolumeUV, uint EyeIndex, float SceneDepth)
 {
 	float4 VolumetricFogLookup = float4(0, 0, 0, 1);
+	
+	if (View.RenderingReflectionCaptureMask == 0.0f && !IsExponentialFogRenderedInMain(View.EnvironmentComponentsFlags))
+	{
+		return VolumetricFogLookup;
+	}
 
 #if PERMUTATION_SUPPORT_VOLUMETRIC_FOG
 	float VolFogStartDistance = 0.0f;
diff --git a/Engine/Shaders/Private/MaterialTemplate.ush b/Engine/Shaders/Private/MaterialTemplate.ush
--- a/Engine/Shaders/Private/MaterialTemplate.ush
+++ b/Engine/Shaders/Private/MaterialTemplate.ush
@@ -33,6 +33,7 @@
 #include "/Engine/Private/HairStrands/HairStrandsAttributeCommon.ush"
 #include "/Engine/Private/DeferredShadingCommon.ush"
 #include "/Engine/Private/SparseVolumeTexture/SparseVolumeTextureCommon.ush"
+#include "/Engine/Shared/EnvironmentComponentsFlags.h"
 
 // Reduce latency of Nanite material base pass pixel shaders caused by register pressure due to the calculation of WPO
 #if PIXELSHADER && IS_NANITE_SHADING_PASS && IS_BASE_PASS
@@ -41,7 +42,7 @@
 
 // Update this GUID to force all material shaders to recompile
 // Merge conflicts on this line should be resolved by generating a new GUID
-#pragma message("UESHADERMETADATA_VERSION 1B0D44F1-B5B0-41FF-9E61-15BAF2C3A5EE")
+#pragma message("UESHADERMETADATA_VERSION A45831C8-997A-4387-907E-31EE08D781F7")
 
 #if USES_SPEEDTREE
 	#include "/Engine/Private/SpeedTreeCommon.ush"
@@ -2152,6 +2153,11 @@ float3 MaterialExpressionSkyAtmosphereLightDiskLuminance(FMaterialPixelParameter
 float3 MaterialExpressionSkyAtmosphereViewLuminance(FMaterialPixelParameters Parameters)
 {
 #if MATERIAL_SKY_ATMOSPHERE && PROJECT_SUPPORT_SKY_ATMOSPHERE
+	if (ResolvedView.RenderingReflectionCaptureMask == 0.0f && !IsSkyAtmosphereRenderedInMain(ResolvedView.EnvironmentComponentsFlags))
+	{
+		return float3(0.0f, 0.0f, 0.0f);
+	}
+	
 	const float ViewHeight = ResolvedView.SkyPlanetTranslatedWorldCenterAndViewHeight.w * CM_TO_SKY_UNIT;
 	const float3 ViewDir = -Parameters.CameraVector;
 
@@ -2171,6 +2177,14 @@ float3 MaterialExpressionSkyAtmosphereViewLuminance(FMaterialPixelParameters Par
 	float3 SkyAtmosphereViewLuminance = Texture2DSampleLevel(View.SkyViewLutTexture, View.SkyViewLutTextureSampler, SkyViewLutUv, 0.0f).rgb;
 	SkyAtmosphereViewLuminance *= ResolvedView.SkyAtmosphereSkyLuminanceFactor.rgb;
 	SkyAtmosphereViewLuminance *= ResolvedView.OneOverPreExposure;
+	
+#if POST_PROCESS_ALPHA
+	if (IsSkyAtmosphereHoldout(View.EnvironmentComponentsFlags) && View.RenderingReflectionCaptureMask == 0.0f)
+	{
+		SkyAtmosphereViewLuminance = 0.0f; // Disable sky dome luminance contribution. Throughput will be 0.
+	}
+#endif
+	
 	return SkyAtmosphereViewLuminance;
 #else
 	return float3(0.0f, 0.0f, 0.0f);
@@ -2199,6 +2213,14 @@ float4 MaterialExpressionSkyAtmosphereAerialPerspective(FMaterialPixelParameters
 		ResolvedView.SkyAtmosphereCameraAerialPerspectiveVolumeDepthSliceLengthKmInv,
 		ResolvedView.OneOverPreExposure,
 		NearFadeOutRangeInvDepthKm);
+	
+#if POST_PROCESS_ALPHA
+	if (IsSkyAtmosphereHoldout(View.EnvironmentComponentsFlags) && View.RenderingReflectionCaptureMask == 0.0f)
+	{
+		AerialPerspective.rgb = 0.0f;
+	}
+#endif
+	
 	return AerialPerspective;
 #else
 	return float4(0.0f, 0.0f, 0.0f, 1.0f); // RGB= null scattering, A= null transmittance
diff --git a/Engine/Shaders/Private/SkyAtmosphere.usf b/Engine/Shaders/Private/SkyAtmosphere.usf
--- a/Engine/Shaders/Private/SkyAtmosphere.usf
+++ b/Engine/Shaders/Private/SkyAtmosphere.usf
@@ -7,7 +7,6 @@
 
 #include "Common.ush"
 #include "SceneTexturesCommon.ush"
-#include "SkyAtmosphereCommon.ush"
 #include "/Engine/Shared/EnvironmentComponentsFlags.h"
 
 #ifndef COLORED_TRANSMITTANCE_ENABLED	// Never used, UE4 does not supports dual blending
@@ -104,6 +103,7 @@
 #include "VolumetricCloudCommon.ush" 
 #endif
 
+#include "SkyAtmosphereCommon.ush"
 
 #if SAMPLE_ATMOSPHERE_ON_CLOUDS
 Texture2D<float>	VolumetricCloudDepthTexture;
@@ -823,17 +823,21 @@ float4 PrepareOutput(float3 PreExposedLuminance, float3 Transmittance = float3(1
 	// Sky materials can result in high luminance values, e.g. the sun disk. 
 	// This is so we make sure to at least stay within the boundaries of fp10 and not cause NaN on some platforms.
 	// We also half that range to also make sure we have room for other additive elements such as bloom, clouds or particle visual effects.
-	float3 SafePreExposedLuminance = min(PreExposedLuminance, Max10BitsFloat.xxx * 0.5f);
-
+	const float3 SafePreExposedLuminance = min(PreExposedLuminance, Max10BitsFloat.xxx * 0.5f);
+	
+	const float GreyScaleTransmittance = dot(Transmittance, float3(1.0f / 3.0f, 1.0f / 3.0f, 1.0f / 3.0f));
+	float4 LuminanceAlpha = float4(SafePreExposedLuminance, GreyScaleTransmittance);
+	
 #if POST_PROCESS_ALPHA
 	if (IsSkyAtmosphereHoldout(View.EnvironmentComponentsFlags) && View.RenderingReflectionCaptureMask == 0.0f)
 	{
-		SafePreExposedLuminance *= 0;
+		LuminanceAlpha.rgb = 0.0f;
+		LuminanceAlpha.a = 1 - GreyScaleTransmittance;
 	}
 #endif
+	
+	return LuminanceAlpha;
 
-	const float GreyScaleTransmittance = dot(Transmittance, float3(1.0f / 3.0f, 1.0f / 3.0f, 1.0f / 3.0f));
-	return float4(SafePreExposedLuminance, GreyScaleTransmittance);
 }
 
 uint DepthReadDisabled;
@@ -928,7 +932,8 @@ void RenderSkyAtmosphereRayMarchingPS(
 
 	float ViewHeight = length(WorldPos);
 #if FASTSKY_ENABLED && RENDERSKY_ENABLED
-	if (ViewHeight < (Atmosphere.TopRadiusKm * PLANET_RADIUS_RATIO_SAFE_EDGE) && DeviceZ == FarDepthValue)
+	if (ViewHeight < (Atmosphere.TopRadiusKm * PLANET_RADIUS_RATIO_SAFE_EDGE) && DeviceZ == FarDepthValue
+		&& (View.RenderingReflectionCaptureMask > 0.0f || IsSkyAtmosphereRenderedInMain(View.EnvironmentComponentsFlags)))
 	{
 		float2 UV;
 
@@ -957,6 +962,7 @@ void RenderSkyAtmosphereRayMarchingPS(
 	#endif
 
 		PreExposedL += SkyLuminance * LuminanceScale * (ViewOneOverPreExposure * OutputPreExposure);
+	
 		OutLuminance = PrepareOutput(PreExposedL, SkyGreyTransmittance);
 		return;
 	}
diff --git a/Engine/Shaders/Private/SkyAtmosphereCommon.ush b/Engine/Shaders/Private/SkyAtmosphereCommon.ush
--- a/Engine/Shaders/Private/SkyAtmosphereCommon.ush
+++ b/Engine/Shaders/Private/SkyAtmosphereCommon.ush
@@ -8,6 +8,7 @@
 #pragma once
 
 #include "ParticipatingMediaCommon.ush"
+#include "/Engine/Shared/EnvironmentComponentsFlags.h"
 
 // Disabled for all mobile platforms, see USkyLightComponent::IsRealTimeCaptureEnabled
 #define AP_SUPPORT_REALTIME_REFLECTION_CAPTURE (FEATURE_LEVEL > FEATURE_LEVEL_ES3_1)
@@ -50,6 +51,11 @@ float4 GetAerialPerspectiveLuminanceTransmittance(
 	float OneOverExposure,
 	float NearFadeOutRangeInvDepthKm)
 {
+	if (View.RenderingReflectionCaptureMask == 0.0f && !IsSkyAtmosphereRenderedInMain(View.EnvironmentComponentsFlags))
+	{
+		return float4(0.0f, 0.0f, 0.0f, 1.0f);
+	}
+	
 	float2 ScreenUv = (NDC.xy / NDC.ww) * float2(0.5f, -0.5f) + 0.5f;
 
 	float tDepth = max(0.0f, length(WorldPositionRelativeToCamera) - AerialPerspectiveVolumeStartDepth);
diff --git a/Engine/Shaders/Private/VolumetricCloud.usf b/Engine/Shaders/Private/VolumetricCloud.usf
--- a/Engine/Shaders/Private/VolumetricCloud.usf
+++ b/Engine/Shaders/Private/VolumetricCloud.usf
@@ -1464,24 +1464,24 @@ void MainCommon(in FMaterialPixelParameters MaterialParameters, in float4 SvPosi
 
 		// Apply aerial perspective OVER the cloud, assiming coverage is from the full transmittance.
 		// skip if MeanTransmittance is 1
-		float MeanCoverage = 1.0 - MeanTransmittance;
-		Luminance = AerialPerspective.rgb * MeanCoverage + AerialPerspective.a * Luminance;
+		float CloudMeanCoverage = 1.0 - MeanTransmittance;
+		Luminance = AerialPerspective.rgb * CloudMeanCoverage + AerialPerspective.a * Luminance;
 
-	#if POST_PROCESS_ALPHA
+#if POST_PROCESS_ALPHA
 		if(View.RenderingReflectionCaptureMask == 0.0f)
 		{
-			const float  AerialPespectiveCoverage = saturate(1.0 - AerialPerspective.a);
-			if (!IsSkyAtmosphereHoldout(View.EnvironmentComponentsFlags))
-			{
-				// Fade out holdout coverage as a function of the aerial perspective above, if the sky atmosphere is not hold out
-				OutHoldOutCoverage   = lerp(OutHoldOutCoverage,   0.0f, AerialPespectiveCoverage);
-			}
-			else // if(IsSkyAtmosphereHoldout(View.EnvironmentComponentsFlags))
+			const float AerialPespectiveTransmittance = AerialPerspective.a;
+			const float AerialPespectiveCoverage = saturate(1.0 - AerialPespectiveTransmittance);
+
+			// Fade out holdout coverage as a function of the aerial perspective above, if the sky atmosphere is not hold out
+			OutHoldOutCoverage   *= saturate(1.0 - AerialPespectiveCoverage);
+
+			if(IsSkyAtmosphereHoldout(View.EnvironmentComponentsFlags))
 			{
-				OutAboveCloudHoldOut = lerp(OutAboveCloudHoldOut, 1.0f, AerialPespectiveCoverage * MeanCoverage);
+				OutAboveCloudHoldOut = OutAboveCloudHoldOut * AerialPespectiveTransmittance + CloudMeanCoverage * AerialPespectiveCoverage;
 			}
 		}
-	#endif
+#endif
 	}
 
 
@@ -1522,24 +1522,24 @@ void MainCommon(in FMaterialPixelParameters MaterialParameters, in float4 SvPosi
 
 		// Apply aerial perspective OVER the cloud, assiming coverage is from the full transmittance.
 		// skip if MeanTransmittance is 1
-		float MeanCoverage = 1.0 - MeanTransmittance;
-		Luminance = HeightFogInscatteringAndTransmittance.rgb * MeanCoverage + HeightFogInscatteringAndTransmittance.a * Luminance;
+		float CloudMeanCoverage = 1.0 - MeanTransmittance;
+		Luminance = HeightFogInscatteringAndTransmittance.rgb * CloudMeanCoverage + HeightFogInscatteringAndTransmittance.a * Luminance;
 
-	#if POST_PROCESS_ALPHA
+#if POST_PROCESS_ALPHA
 		if(View.RenderingReflectionCaptureMask == 0.0f)
 		{
-			const float  ExponentialFogCoverage = saturate(1.0 - HeightFogInscatteringAndTransmittance.a);
-			if (!IsExponentialFogHoldout(View.EnvironmentComponentsFlags))
-			{
-				// Fade out holdout coverage as a function of the fog above, if the fog is not hold out
-				OutHoldOutCoverage   = lerp(OutHoldOutCoverage,   0.0f, ExponentialFogCoverage);
-			}
-			else //if(IsExponentialFogHoldout(View.EnvironmentComponentsFlags))
+			const float ExponentialFogTransmittance = HeightFogInscatteringAndTransmittance.a;
+			const float ExponentialFogCoverage = saturate(1.0 - ExponentialFogTransmittance);
+
+			// Fade out holdout coverage as a function of the fog above, if the fog is not hold out
+			OutHoldOutCoverage *= ExponentialFogTransmittance;
+
+			if(IsExponentialFogHoldout(View.EnvironmentComponentsFlags))
 			{
-				OutAboveCloudHoldOut = lerp(OutAboveCloudHoldOut, 1.0f, ExponentialFogCoverage * MeanCoverage);
+				OutAboveCloudHoldOut = OutAboveCloudHoldOut * ExponentialFogTransmittance + CloudMeanCoverage * ExponentialFogCoverage;
 			}
 		}
-	#endif
+#endif
 	}
 
 
@@ -1679,6 +1679,7 @@ void MainPS(
 float4 OutputViewRect;
 int bBlendCloudColor;
 int TargetCubeFace;
+int bAccumulateAlphaHoldOut;
 
 RWTexture2D<float4> OutCloudColor0;
 #if CLOUD_MIN_AND_MAX_DEPTH
@@ -1733,14 +1734,16 @@ void MainCS(uint2 DispatchThreadId : SV_DispatchThreadID)
 			#if POST_PROCESS_ALPHA
 				if (View.RenderingReflectionCaptureMask == 0.0f)
 				{
-					if(IsVolumetricCloudHoldout(View.EnvironmentComponentsFlags))
+					if(IsVolumetricCloudHoldout(View.EnvironmentComponentsFlags) && bAccumulateAlphaHoldOut > 0)
 					{
-						// Holdout for when we blend directly into the scene color Buffer.
-						OutCloudColor0[DispatchThreadId] = float4(CurColor0.rgb * OutColor0.a + OutColor0.rgb, CurColor0.a * (1.0 - OutHoldOutCoverage) + OutHoldOutCoverage);
+						// Operations for Holdout for when we blend directly with scene color Buffer read to CurColor0.
+						// We maintain the color of the scene since it has alreayd been composed
+						// Then we simply add the holdout value on top. Since the alpha has already been impacted by cloud GetCloudLuminanceTransmittanceOverFog when the clouds ahve been composed.
+						OutCloudColor0[DispatchThreadId] = float4(CurColor0.rgb, CurColor0.a + OutHoldOutCoverage + OutAboveCloudHoldOut);
 					}
 					else
 					{
-						OutCloudColor0[DispatchThreadId] = float4(CurColor0.rgb * OutColor0.a + OutColor0.rgb, CurColor0.a * OutColor0.a + OutAboveCloudHoldOut);
+						OutCloudColor0[DispatchThreadId] = float4(CurColor0.rgb * OutColor0.a + OutColor0.rgb, CurColor0.a + OutAboveCloudHoldOut);
 					}
 				}
 			#else
diff --git a/Engine/Shaders/Shared/EnvironmentComponentsFlags.h b/Engine/Shaders/Shared/EnvironmentComponentsFlags.h
--- a/Engine/Shaders/Shared/EnvironmentComponentsFlags.h
+++ b/Engine/Shaders/Shared/EnvironmentComponentsFlags.h
@@ -13,9 +13,13 @@
 #endif
 
 
-#define ENVCOMP_FLAG_SKYATMOSPHERE_HOLDOUT 0x1
-#define ENVCOMP_FLAG_VOLUMETRICCLOUD_HOLDOUT 0x2
-#define ENVCOMP_FLAG_EXPONENTIALFOG_HOLDOUT 0x4
+#define ENVCOMP_FLAG_SKYATMOSPHERE_HOLDOUT			0x01
+#define ENVCOMP_FLAG_VOLUMETRICCLOUD_HOLDOUT		0x02
+#define ENVCOMP_FLAG_EXPONENTIALFOG_HOLDOUT			0x04
+
+#define ENVCOMP_FLAG_SKYATMOSPHERE_RENDERINMAIN		0x08
+#define ENVCOMP_FLAG_VOLUMETRICCLOUD_RENDERINMAIN	0x10
+#define ENVCOMP_FLAG_EXPONENTIALFOG_RENDERINMAIN	0x20
 
 
 // Using macro functions to avoid having to move the definition in cpp to avoid the multiple definition in different obj files.
@@ -23,3 +27,7 @@
 #define IsVolumetricCloudHoldout(EnvironmentComponentsFlags) ((EnvironmentComponentsFlags & ENVCOMP_FLAG_VOLUMETRICCLOUD_HOLDOUT) > 0)
 #define IsExponentialFogHoldout(EnvironmentComponentsFlags)  ((EnvironmentComponentsFlags & ENVCOMP_FLAG_EXPONENTIALFOG_HOLDOUT)  > 0)
 
+#define IsSkyAtmosphereRenderedInMain(EnvironmentComponentsFlags)   ((EnvironmentComponentsFlags & ENVCOMP_FLAG_SKYATMOSPHERE_RENDERINMAIN)   > 0)
+#define IsVolumetricCloudRenderedInMain(EnvironmentComponentsFlags) ((EnvironmentComponentsFlags & ENVCOMP_FLAG_VOLUMETRICCLOUD_RENDERINMAIN) > 0)
+#define IsExponentialFogRenderedInMain(EnvironmentComponentsFlags)  ((EnvironmentComponentsFlags & ENVCOMP_FLAG_EXPONENTIALFOG_RENDERINMAIN)  > 0)
+
diff --git a/Engine/Source/Runtime/Engine/Classes/Components/ExponentialHeightFogComponent.h b/Engine/Source/Runtime/Engine/Classes/Components/ExponentialHeightFogComponent.h
--- a/Engine/Source/Runtime/Engine/Classes/Components/ExponentialHeightFogComponent.h
+++ b/Engine/Source/Runtime/Engine/Classes/Components/ExponentialHeightFogComponent.h
@@ -204,6 +204,11 @@ class UExponentialHeightFogComponent : public USceneComponent
 	UPROPERTY(EditAnywhere, AdvancedDisplay, BlueprintReadOnly, Category = Rendering, Interp)
 	uint8 bHoldout : 1;
 
+	/** If true, this component will be rendered in the main pass (basepass, transparency) */
+	UPROPERTY(EditAnywhere, AdvancedDisplay, BlueprintReadOnly, Category = Rendering)
+	uint8 bRenderInMainPass : 1;
+
+
 public:
 	UFUNCTION(BlueprintCallable, Category="Rendering|Components|ExponentialHeightFog")
 	ENGINE_API void SetFogDensity(float Value);
@@ -280,6 +285,9 @@ public:
 	UFUNCTION(BlueprintCallable, Category = "Rendering")
 	ENGINE_API void SetHoldout(bool bNewHoldout);
 
+	UFUNCTION(BlueprintCallable, Category = "Rendering")
+	ENGINE_API void SetRenderInMainPass(bool bValue);
+
 protected:
 	//~ Begin UActorComponent Interface.
 	ENGINE_API virtual void CreateRenderState_Concurrent(FRegisterComponentContext* Context) override;
diff --git a/Engine/Source/Runtime/Engine/Classes/Components/SkyAtmosphereComponent.h b/Engine/Source/Runtime/Engine/Classes/Components/SkyAtmosphereComponent.h
--- a/Engine/Source/Runtime/Engine/Classes/Components/SkyAtmosphereComponent.h
+++ b/Engine/Source/Runtime/Engine/Classes/Components/SkyAtmosphereComponent.h
@@ -164,6 +164,10 @@ class USkyAtmosphereComponent : public USceneComponent
 	UPROPERTY(EditAnywhere, AdvancedDisplay, BlueprintReadOnly, Category = Rendering, Interp)
 	uint8 bHoldout : 1;
 
+	/** If true, this component will be rendered in the main pass (basepass, transparency) */
+	UPROPERTY(EditAnywhere, AdvancedDisplay, BlueprintReadOnly, Category = Rendering)
+	uint8 bRenderInMainPass : 1;
+
 
 	UFUNCTION(BlueprintCallable, Category = "Rendering")
 	ENGINE_API void OverrideAtmosphereLightDirection(int32 AtmosphereLightIndex, const FVector& LightDirection);
@@ -219,6 +223,9 @@ class USkyAtmosphereComponent : public USceneComponent
 	UFUNCTION(BlueprintCallable, Category = "Rendering")
 	ENGINE_API void SetHoldout(bool bNewHoldout);
 
+	UFUNCTION(BlueprintCallable, Category = "Rendering")
+	ENGINE_API void SetRenderInMainPass(bool bValue);
+
 	UFUNCTION(BlueprintCallable, Category = "Utilities", meta = (DisplayName = "Get Atmosphere Transmitance On Ground At Planet Top"))
 	ENGINE_API FLinearColor GetAtmosphereTransmitanceOnGroundAtPlanetTop(UDirectionalLightComponent* DirectionalLight);
 
diff --git a/Engine/Source/Runtime/Engine/Classes/Components/VolumetricCloudComponent.h b/Engine/Source/Runtime/Engine/Classes/Components/VolumetricCloudComponent.h
--- a/Engine/Source/Runtime/Engine/Classes/Components/VolumetricCloudComponent.h
+++ b/Engine/Source/Runtime/Engine/Classes/Components/VolumetricCloudComponent.h
@@ -143,6 +143,10 @@ class UVolumetricCloudComponent : public USceneComponent
 	UPROPERTY(EditAnywhere, AdvancedDisplay, BlueprintReadOnly, Category = Rendering, Interp)
 	uint8 bHoldout : 1;
 
+	/** If true, this component will be rendered in the main pass (basepass, transparency) */
+	UPROPERTY(EditAnywhere, AdvancedDisplay, BlueprintReadOnly, Category = Rendering)
+	uint8 bRenderInMainPass : 1;
+
 
 	UFUNCTION(BlueprintCallable, Category = "Rendering")
 	ENGINE_API void SetLayerBottomAltitude(float NewValue);
@@ -190,6 +194,9 @@ class UVolumetricCloudComponent : public USceneComponent
 	UFUNCTION(BlueprintCallable, Category = "Rendering")
 	ENGINE_API void SetHoldout(bool bNewHoldout);
 
+	UFUNCTION(BlueprintCallable, Category = "Rendering")
+	ENGINE_API void SetRenderInMainPass(bool bValue);
+
 
 protected:
 	//~ Begin UActorComponent Interface.
diff --git a/Engine/Source/Runtime/Engine/Private/Components/HeightFogComponent.cpp b/Engine/Source/Runtime/Engine/Private/Components/HeightFogComponent.cpp
--- a/Engine/Source/Runtime/Engine/Private/Components/HeightFogComponent.cpp
+++ b/Engine/Source/Runtime/Engine/Private/Components/HeightFogComponent.cpp
@@ -43,6 +43,9 @@ UExponentialHeightFogComponent::UExponentialHeightFogComponent(const FObjectInit
 	// disabled by default
 	FogCutoffDistance = 0;
 
+	bHoldout = false;
+	bRenderInMainPass = true;
+
 	VolumetricFogScatteringDistribution = .2f;
 	VolumetricFogAlbedo = FColor::White;
 	VolumetricFogExtinctionScale = 1.0f;
@@ -358,6 +361,15 @@ void UExponentialHeightFogComponent::SetHoldout(bool bNewHoldout)
 	}
 }
 
+void UExponentialHeightFogComponent::SetRenderInMainPass(bool bValue)
+{
+	if (bRenderInMainPass != bValue)
+	{
+		bRenderInMainPass = bValue;
+		MarkRenderStateDirty();
+	}
+}
+
 void UExponentialHeightFogComponent::Serialize(FArchive& Ar)
 {
 	Super::Serialize(Ar);
diff --git a/Engine/Source/Runtime/Engine/Private/Components/SkyAtmosphereComponent.cpp b/Engine/Source/Runtime/Engine/Private/Components/SkyAtmosphereComponent.cpp
--- a/Engine/Source/Runtime/Engine/Private/Components/SkyAtmosphereComponent.cpp
+++ b/Engine/Source/Runtime/Engine/Private/Components/SkyAtmosphereComponent.cpp
@@ -79,6 +79,9 @@ USkyAtmosphereComponent::USkyAtmosphereComponent(const FObjectInitializer& Objec
 
 	TraceSampleCountScale = 1.0f;
 
+	bHoldout = false;
+	bRenderInMainPass = true;
+
 	memset(OverrideAtmosphericLight, 0, sizeof(OverrideAtmosphericLight));
 
 	ValidateStaticLightingGUIDs();
@@ -396,6 +399,15 @@ void USkyAtmosphereComponent::SetHoldout(bool bNewHoldout)
 	}
 }
 
+void USkyAtmosphereComponent::SetRenderInMainPass(bool bValue)
+{
+	if (bRenderInMainPass != bValue)
+	{
+		bRenderInMainPass = bValue;
+		MarkRenderStateDirty();
+	}
+}
+
 FLinearColor USkyAtmosphereComponent::GetAtmosphereTransmitanceOnGroundAtPlanetTop(UDirectionalLightComponent* DirectionalLight)
 {
 	if(DirectionalLight != nullptr)
@@ -480,6 +492,7 @@ FSkyAtmosphereSceneProxy::FSkyAtmosphereSceneProxy(const USkyAtmosphereComponent
 	: bStaticLightingBuilt(false)
 	, AtmosphereSetup(*InComponent)
 	, bHoldout(InComponent->bHoldout > 0)
+	, bRenderInMainPass(InComponent->bRenderInMainPass > 0)
 {
 	SkyLuminanceFactor = InComponent->SkyLuminanceFactor;
 	AerialPespectiveViewDistanceScale = InComponent->AerialPespectiveViewDistanceScale;
diff --git a/Engine/Source/Runtime/Engine/Private/Components/VolumetricCloudComponent.cpp b/Engine/Source/Runtime/Engine/Private/Components/VolumetricCloudComponent.cpp
--- a/Engine/Source/Runtime/Engine/Private/Components/VolumetricCloudComponent.cpp
+++ b/Engine/Source/Runtime/Engine/Private/Components/VolumetricCloudComponent.cpp
@@ -53,6 +53,8 @@ UVolumetricCloudComponent::UVolumetricCloudComponent(const FObjectInitializer& O
 	, AerialPespectiveRayleighScatteringFadeDistance(0.0f)
 	, AerialPespectiveMieScatteringStartDistance(0.0f)
 	, AerialPespectiveMieScatteringFadeDistance(0.0f)
+	, bHoldout(false)
+	, bRenderInMainPass(true)
 	, VolumetricCloudSceneProxy(nullptr)
 {
 	static ConstructorHelpers::FObjectFinder<UMaterialInterface> VolumetricCloudDefaultMaterialRef(TEXT("/Engine/EngineSky/VolumetricClouds/m_SimpleVolumetricCloud_Inst.m_SimpleVolumetricCloud_Inst"));
@@ -175,6 +177,15 @@ void UVolumetricCloudComponent::SetHoldout(bool bNewHoldout)
 	}
 }
 
+void UVolumetricCloudComponent::SetRenderInMainPass(bool bValue)
+{
+	if (bRenderInMainPass != bValue)
+	{
+		bRenderInMainPass = bValue;
+		MarkRenderStateDirty();
+	}
+}
+
 void UVolumetricCloudComponent::SetReflectionViewSampleCountScale(float NewValue)
 {
 	if (AreDynamicDataChangesAllowed() && ReflectionViewSampleCountScaleValue != NewValue)
diff --git a/Engine/Source/Runtime/Engine/Private/VolumetricCloudProxy.cpp b/Engine/Source/Runtime/Engine/Private/VolumetricCloudProxy.cpp
--- a/Engine/Source/Runtime/Engine/Private/VolumetricCloudProxy.cpp
+++ b/Engine/Source/Runtime/Engine/Private/VolumetricCloudProxy.cpp
@@ -20,6 +20,7 @@ FVolumetricCloudSceneProxy::FVolumetricCloudSceneProxy(const UVolumetricCloudCom
 	, GroundAlbedo(InComponent->GroundAlbedo)
 	, bUsePerSampleAtmosphericLightTransmittance(InComponent->bUsePerSampleAtmosphericLightTransmittance)
 	, bHoldout(InComponent->bHoldout)
+	, bRenderInMainPass(InComponent->bRenderInMainPass)
 	, SkyLightCloudBottomOcclusion(InComponent->SkyLightCloudBottomOcclusion)
 	, ViewSampleCountScale(InComponent->ViewSampleCountScale)
 	, ReflectionViewSampleCountScale(InComponent->ReflectionViewSampleCountScaleValue)
diff --git a/Engine/Source/Runtime/Engine/Public/SceneManagement.h b/Engine/Source/Runtime/Engine/Public/SceneManagement.h
--- a/Engine/Source/Runtime/Engine/Public/SceneManagement.h
+++ b/Engine/Source/Runtime/Engine/Public/SceneManagement.h
@@ -1200,6 +1200,7 @@ public:
 	const FAtmosphereSetup& GetAtmosphereSetup() const { return AtmosphereSetup; }
 
 	bool IsHoldout() const { return bHoldout; }
+	bool IsRenderedInMainPass() const { return bRenderInMainPass; }
 
 	void UpdateTransform(const FTransform& ComponentTransform, uint8 TranformMode) { AtmosphereSetup.UpdateTransform(ComponentTransform, TranformMode); }
 	void ApplyWorldOffset(const FVector3f& InOffset) { AtmosphereSetup.ApplyWorldOffset((FVector)InOffset); }
@@ -1218,6 +1219,7 @@ private:
 	float AerialPerspectiveStartDepthKm;
 	float TraceSampleCountScale;
 	bool bHoldout;
+	bool bRenderInMainPass;
 
 	bool OverrideAtmosphericLight[NUM_ATMOSPHERE_LIGHTS];
 	FVector OverrideAtmosphericLightDirection[NUM_ATMOSPHERE_LIGHTS];
diff --git a/Engine/Source/Runtime/Engine/Public/VolumetricCloudProxy.h b/Engine/Source/Runtime/Engine/Public/VolumetricCloudProxy.h
--- a/Engine/Source/Runtime/Engine/Public/VolumetricCloudProxy.h
+++ b/Engine/Source/Runtime/Engine/Public/VolumetricCloudProxy.h
@@ -41,6 +41,7 @@ public:
 	FColor GroundAlbedo;
 	bool bUsePerSampleAtmosphericLightTransmittance;
 	bool bHoldout;
+	bool bRenderInMainPass;
 	float SkyLightCloudBottomOcclusion;
 	
 	float ViewSampleCountScale;
diff --git a/Engine/Source/Runtime/Renderer/Private/DeferredShadingRenderer.cpp b/Engine/Source/Runtime/Renderer/Private/DeferredShadingRenderer.cpp
--- a/Engine/Source/Runtime/Renderer/Private/DeferredShadingRenderer.cpp
+++ b/Engine/Source/Runtime/Renderer/Private/DeferredShadingRenderer.cpp
@@ -100,6 +100,7 @@
 #include "PostProcess/DebugAlphaChannel.h"
 #include "ManyLights/ManyLights.h"
 #include "Rendering/CustomRenderPass.h"
+#include "EnvironmentComponentsFlags.h"
 
 #if !UE_BUILD_SHIPPING
 #include "RenderCaptureInterface.h"
@@ -2300,7 +2301,8 @@ void FDeferredShadingSceneRenderer::Render(FRDGBuilder& GraphBuilder)
 
 			bool bSkipVolumetricRenderTarget = false;
 			bool bSkipPerPixelTracing = true;
-			bAsyncComputeVolumetricCloud = RenderVolumetricCloud(GraphBuilder, SceneTextures, bSkipVolumetricRenderTarget, bSkipPerPixelTracing, 
+			bool bAccumulateAlphaHoldOut = false;
+			RenderVolumetricCloud(GraphBuilder, SceneTextures, bSkipVolumetricRenderTarget, bSkipPerPixelTracing, bAccumulateAlphaHoldOut,
 				HalfResolutionDepthCheckerboardMinMaxTexture, QuarterResolutionDepthMinMaxTexture, true, InstanceCullingManager);
 		}
 		
@@ -2629,7 +2631,8 @@ void FDeferredShadingSceneRenderer::Render(FRDGBuilder& GraphBuilder)
 
 			bool bSkipVolumetricRenderTarget = false;
 			bool bSkipPerPixelTracing = true;
-			bAsyncComputeVolumetricCloud = RenderVolumetricCloud(GraphBuilder, SceneTextures, bSkipVolumetricRenderTarget, bSkipPerPixelTracing, 
+			bool bAccumulateAlphaHoldOut = false;
+			RenderVolumetricCloud(GraphBuilder, SceneTextures, bSkipVolumetricRenderTarget, bSkipPerPixelTracing, bAccumulateAlphaHoldOut,
 				HalfResolutionDepthCheckerboardMinMaxTexture, QuarterResolutionDepthMinMaxTexture, true, InstanceCullingManager);
 		}
 
@@ -2815,7 +2818,8 @@ void FDeferredShadingSceneRenderer::Render(FRDGBuilder& GraphBuilder)
 				// Generate the volumetric cloud render target
 				bool bSkipVolumetricRenderTarget = false;
 				bool bSkipPerPixelTracing = true;
-				RenderVolumetricCloud(GraphBuilder, SceneTextures, bSkipVolumetricRenderTarget, bSkipPerPixelTracing, 
+				bool bAccumulateAlphaHoldOut = false;
+				RenderVolumetricCloud(GraphBuilder, SceneTextures, bSkipVolumetricRenderTarget, bSkipPerPixelTracing, bAccumulateAlphaHoldOut,
 					HalfResolutionDepthCheckerboardMinMaxTexture, QuarterResolutionDepthMinMaxTexture, false, InstanceCullingManager);
 			}
 			// Reconstruct the volumetric cloud render target to be ready to compose it over the scene
@@ -2896,7 +2900,8 @@ void FDeferredShadingSceneRenderer::Render(FRDGBuilder& GraphBuilder)
 			{
 				bool bSkipVolumetricRenderTarget = true;
 				bool bSkipPerPixelTracing = false;
-				RenderVolumetricCloud(GraphBuilder, SceneTextures, bSkipVolumetricRenderTarget, bSkipPerPixelTracing,
+				bool bAccumulateAlphaHoldOut = false;
+				RenderVolumetricCloud(GraphBuilder, SceneTextures, bSkipVolumetricRenderTarget, bSkipPerPixelTracing, bAccumulateAlphaHoldOut,
 					HalfResolutionDepthCheckerboardMinMaxTexture, QuarterResolutionDepthMinMaxTexture, false, InstanceCullingManager);
 			}
 
@@ -2908,6 +2913,17 @@ void FDeferredShadingSceneRenderer::Render(FRDGBuilder& GraphBuilder)
 					GraphBuilder, Views, SceneTextures.Color.Target, SceneTextures.Depth.Target,
 					bComposeWithWater,
 					SceneWithoutWaterTextures, SceneTextures);
+
+				if (IsPostProcessingWithAlphaChannelSupported())
+				{
+					// When alpha is enabled to work with holdout. We need another full screen tracing pass to update the alpha channel containing the "holdout alpha throughput".
+					// Alpha hold out only works when using r.volumetricrendertarget.mode 3 which is the mode use by MRQ.
+					bool bSkipVolumetricRenderTarget = true;
+					bool bSkipPerPixelTracing = false;
+					bool bAccumulateAlphaHoldOut = true;
+					RenderVolumetricCloud(GraphBuilder, SceneTextures, bSkipVolumetricRenderTarget, bSkipPerPixelTracing, bAccumulateAlphaHoldOut,
+						HalfResolutionDepthCheckerboardMinMaxTexture, QuarterResolutionDepthMinMaxTexture, false, InstanceCullingManager);
+				}
 			}
 		};
 
@@ -3188,10 +3204,10 @@ void FDeferredShadingSceneRenderer::Render(FRDGBuilder& GraphBuilder)
 
 		if (bShouldVisualizeVolumetricCloud && !bHasRayTracedOverlay)
 		{
-			RenderVolumetricCloud(GraphBuilder, SceneTextures, false, true, HalfResolutionDepthCheckerboardMinMaxTexture, QuarterResolutionDepthMinMaxTexture, false, InstanceCullingManager);
+			RenderVolumetricCloud(GraphBuilder, SceneTextures, false, true, false, HalfResolutionDepthCheckerboardMinMaxTexture, QuarterResolutionDepthMinMaxTexture, false, InstanceCullingManager);
 			ReconstructVolumetricRenderTarget(GraphBuilder, Views, SceneTextures.Depth.Resolve, HalfResolutionDepthCheckerboardMinMaxTexture, false);
 			ComposeVolumetricRenderTargetOverSceneForVisualization(GraphBuilder, Views, SceneTextures.Color.Target, SceneTextures);
-			RenderVolumetricCloud(GraphBuilder, SceneTextures, true, false, HalfResolutionDepthCheckerboardMinMaxTexture, QuarterResolutionDepthMinMaxTexture, false, InstanceCullingManager);
+			RenderVolumetricCloud(GraphBuilder, SceneTextures, true, false, false, HalfResolutionDepthCheckerboardMinMaxTexture, QuarterResolutionDepthMinMaxTexture, false, InstanceCullingManager);
 		}
 
 		if (!bHasRayTracedOverlay)
diff --git a/Engine/Source/Runtime/Renderer/Private/FogRendering.cpp b/Engine/Source/Runtime/Renderer/Private/FogRendering.cpp
--- a/Engine/Source/Runtime/Renderer/Private/FogRendering.cpp
+++ b/Engine/Source/Runtime/Renderer/Private/FogRendering.cpp
@@ -327,8 +327,8 @@ static void RenderViewFog(
 			}
 			else
 			{
-				// Same color blending as without alpha channel. Disabled alpha blending writes because alpha holdout from other component must be preserved.
-				GraphicsPSOInit.BlendState = TStaticBlendState<CW_RGB, BO_Add, BF_One, BF_InverseSourceAlpha>::GetRHI();
+				// Same color blending. Alpha blending is kept so that the throughput of height fog can still be applied on the background (for instance when the sky or a mesh is holdout).
+				GraphicsPSOInit.BlendState = TStaticBlendState<CW_RGBA, BO_Add, BF_One, BF_InverseSourceAlpha, BO_Add, BF_Zero, BF_InverseSourceAlpha>::GetRHI();
 			}
 		}
 		else
diff --git a/Engine/Source/Runtime/Renderer/Private/SceneCore.cpp b/Engine/Source/Runtime/Renderer/Private/SceneCore.cpp
--- a/Engine/Source/Runtime/Renderer/Private/SceneCore.cpp
+++ b/Engine/Source/Runtime/Renderer/Private/SceneCore.cpp
@@ -419,6 +419,7 @@ FExponentialHeightFogSceneInfo::FExponentialHeightFogSceneInfo(const UExponentia
 	VolumetricFogStaticLightingScatteringIntensity = FMath::Max(InComponent->VolumetricFogStaticLightingScatteringIntensity, 0.0f);
 	bOverrideLightColorsWithFogInscatteringColors = InComponent->bOverrideLightColorsWithFogInscatteringColors;
 	bHoldout = InComponent->bHoldout;
+	bRenderInMainPass = InComponent->bRenderInMainPass;
 
 	VolumetricFogStartDistance = InComponent->VolumetricFogStartDistance;
 	VolumetricFogNearFadeInDistance = InComponent->VolumetricFogNearFadeInDistance;
diff --git a/Engine/Source/Runtime/Renderer/Private/SceneCore.h b/Engine/Source/Runtime/Renderer/Private/SceneCore.h
--- a/Engine/Source/Runtime/Renderer/Private/SceneCore.h
+++ b/Engine/Source/Runtime/Renderer/Private/SceneCore.h
@@ -168,6 +168,7 @@ public:
 	float VolumetricFogStaticLightingScatteringIntensity;
 	bool bOverrideLightColorsWithFogInscatteringColors;
 	bool bHoldout;
+	bool bRenderInMainPass;
 	float VolumetricFogStartDistance;
 	float VolumetricFogNearFadeInDistance;
 
diff --git a/Engine/Source/Runtime/Renderer/Private/SceneRendering.cpp b/Engine/Source/Runtime/Renderer/Private/SceneRendering.cpp
--- a/Engine/Source/Runtime/Renderer/Private/SceneRendering.cpp
+++ b/Engine/Source/Runtime/Renderer/Private/SceneRendering.cpp
@@ -1458,6 +1458,7 @@ void FViewInfo::SetupUniformBufferParameters(
 			const FSkyAtmosphereSceneProxy& SkyAtmosphereSceneProxy = SkyAtmosphere->GetSkyAtmosphereSceneProxy();
 
 			ViewUniformShaderParameters.EnvironmentComponentsFlags |= SkyAtmosphereSceneProxy.IsHoldout() ? ENVCOMP_FLAG_SKYATMOSPHERE_HOLDOUT : 0;
+			ViewUniformShaderParameters.EnvironmentComponentsFlags |= SkyAtmosphereSceneProxy.IsRenderedInMainPass() ? ENVCOMP_FLAG_SKYATMOSPHERE_RENDERINMAIN : 0;
 		}
 
 		if (ShouldRenderVolumetricCloud(Scene, Family->EngineShowFlags))
@@ -1466,6 +1467,7 @@ void FViewInfo::SetupUniformBufferParameters(
 			const FVolumetricCloudSceneProxy& VolumetricCloudSceneProxy = VolumetricCloud->GetVolumetricCloudSceneProxy();
 
 			ViewUniformShaderParameters.EnvironmentComponentsFlags |= VolumetricCloudSceneProxy.bHoldout ? ENVCOMP_FLAG_VOLUMETRICCLOUD_HOLDOUT : 0;
+			ViewUniformShaderParameters.EnvironmentComponentsFlags |= VolumetricCloudSceneProxy.bRenderInMainPass ? ENVCOMP_FLAG_VOLUMETRICCLOUD_RENDERINMAIN : 0;
 		}
 
 		if (Scene->ExponentialFogs.Num() > 0)
@@ -1473,6 +1475,7 @@ void FViewInfo::SetupUniformBufferParameters(
 			FExponentialHeightFogSceneInfo& Fog = Scene->ExponentialFogs[0];
 
 			ViewUniformShaderParameters.EnvironmentComponentsFlags |= Fog.bHoldout ? ENVCOMP_FLAG_EXPONENTIALFOG_HOLDOUT : 0;
+			ViewUniformShaderParameters.EnvironmentComponentsFlags |= Fog.bRenderInMainPass ? ENVCOMP_FLAG_EXPONENTIALFOG_RENDERINMAIN : 0;
 		}
 	}
 
diff --git a/Engine/Source/Runtime/Renderer/Private/SceneRendering.h b/Engine/Source/Runtime/Renderer/Private/SceneRendering.h
--- a/Engine/Source/Runtime/Renderer/Private/SceneRendering.h
+++ b/Engine/Source/Runtime/Renderer/Private/SceneRendering.h
@@ -2506,6 +2506,7 @@ protected:
 		const FMinimalSceneTextures& SceneTextures,
 		bool bSkipVolumetricRenderTarget,
 		bool bSkipPerPixelTracing,
+		bool bAccumulateAlphaHoldOut,
 		FRDGTextureRef HalfResolutionDepthCheckerboardMinMaxTexture,
 		FRDGTextureRef QuarterResolutionDepthMinMaxTexture,
 		bool bAsyncCompute,
diff --git a/Engine/Source/Runtime/Renderer/Private/SkyAtmosphereRendering.cpp b/Engine/Source/Runtime/Renderer/Private/SkyAtmosphereRendering.cpp
--- a/Engine/Source/Runtime/Renderer/Private/SkyAtmosphereRendering.cpp
+++ b/Engine/Source/Runtime/Renderer/Private/SkyAtmosphereRendering.cpp
@@ -1786,12 +1786,16 @@ void FSceneRenderer::RenderSkyAtmosphereInternal(
 		const float CloudShadowOnAtmosphereStrength1 = AtmosphereLight1Proxy ? AtmosphereLight1Proxy->GetCloudShadowOnAtmosphereStrength() : 0.0f;
 		const bool bShouldSampleCloudShadow = SkyRC.bShouldSampleCloudShadow && (CloudShadowOnAtmosphereStrength0 > 0.0f || CloudShadowOnAtmosphereStrength1 > 0.0f);
 
+		const bool bFastAerialPerspectiveDepthTest = SkyRC.bFastAerialPerspectiveDepthTest;
+		const bool SkyAtmosphereAlphaHoldOut = IsPostProcessingWithAlphaChannelSupported() && SkyAtmosphereSceneProxy.IsHoldout();
+		const bool bRenderSkyPixel = SkyRC.bRenderSkyPixel || SkyAtmosphereAlphaHoldOut; // In this case we need to write alpha holdout values in the sky pixels.
+
 		FRenderSkyAtmospherePS::FPermutationDomain PsPermutationVector;
 		PsPermutationVector.Set<FSampleCloudSkyAO>(SkyRC.bShouldSampleCloudSkyAO);
 		PsPermutationVector.Set<FFastSky>(SkyRC.bFastSky && !ForceRayMarching);
 		PsPermutationVector.Set<FFastAerialPespective>(SkyRC.bFastAerialPerspective && !ForceRayMarching);
 		PsPermutationVector.Set<FSecondAtmosphereLight>(SkyRC.bSecondAtmosphereLightEnabled);
-		PsPermutationVector.Set<FRenderSky>(SkyRC.bRenderSkyPixel);
+		PsPermutationVector.Set<FRenderSky>(bRenderSkyPixel);
 		PsPermutationVector.Set<FSampleOpaqueShadow>(SkyRC.bShouldSampleOpaqueShadow);
 		PsPermutationVector.Set<FSampleCloudShadow>(bShouldSampleCloudShadow);
 		PsPermutationVector.Set<FAtmosphereOnClouds>(SkyRC.bAPOnCloudMode);
@@ -1857,15 +1861,12 @@ void FSceneRenderer::RenderSkyAtmosphereInternal(
 			StartDepthZ = float(Projected.Z / Projected.W); // LWC_TODO: precision loss
 		}
 
-		const bool bFastAerialPerspectiveDepthTest = SkyRC.bFastAerialPerspectiveDepthTest;
-		const bool bRenderSkyPixel = SkyRC.bRenderSkyPixel;
-		const bool bSupportsAlpha = IsPostProcessingWithAlphaChannelSupported();
 		FIntRect Viewport = SkyRC.Viewport;
 		GraphBuilder.AddPass(
 			{},
 			PsPassParameters,
 			ERDGPassFlags::Raster,
-			[PsPassParameters, VertexShader, PixelShader, Viewport, bFastAerialPerspectiveDepthTest, bRenderSkyPixel, bDisableBlending, StartDepthZ, bSupportsAlpha](FRHICommandList& RHICmdListLambda)
+			[PsPassParameters, VertexShader, PixelShader, Viewport, bFastAerialPerspectiveDepthTest, bRenderSkyPixel, bDisableBlending, StartDepthZ, SkyAtmosphereAlphaHoldOut](FRHICommandList& RHICmdListLambda)
 		{
 			RHICmdListLambda.SetViewport(Viewport.Min.X, Viewport.Min.Y, 0.0f, Viewport.Max.X, Viewport.Max.Y, 1.0f);
 
@@ -1878,13 +1879,15 @@ void FSceneRenderer::RenderSkyAtmosphereInternal(
 			}
 			else
 			{
-				if (bSupportsAlpha)
+				if (SkyAtmosphereAlphaHoldOut)
 				{
-					GraphicsPSOInit.BlendState = TStaticBlendState<CW_RGBA, BO_Add, BF_One, BF_SourceAlpha, BO_Add, BF_Zero, BF_SourceAlpha>::GetRHI();
+					// We might need to write alpha holdout data. Since this is the sky writing to 
+					GraphicsPSOInit.BlendState = TStaticBlendState<CW_RGBA, BO_Add, BF_One, BF_InverseSourceAlpha, BO_Add, BF_One, BF_InverseSourceAlpha>::GetRHI();
 				}
 				else
 				{
-					GraphicsPSOInit.BlendState = TStaticBlendState<CW_RGB, BO_Add, BF_One, BF_SourceAlpha, BO_Add, BF_Zero, BF_One>::GetRHI(); 
+					// Maintain alpha untouched
+					GraphicsPSOInit.BlendState = TStaticBlendState<CW_RGB, BO_Add, BF_One, BF_SourceAlpha>::GetRHI(); 
 				}
 			}
 			if (bFastAerialPerspectiveDepthTest)
diff --git a/Engine/Source/Runtime/Renderer/Private/VolumetricCloudRendering.cpp b/Engine/Source/Runtime/Renderer/Private/VolumetricCloudRendering.cpp
--- a/Engine/Source/Runtime/Renderer/Private/VolumetricCloudRendering.cpp
+++ b/Engine/Source/Runtime/Renderer/Private/VolumetricCloudRendering.cpp
@@ -25,6 +25,7 @@
 #include "FogRendering.h"
 #include "SkyAtmosphereRendering.h"
 #include "BasePassRendering.h"
+#include "EnvironmentComponentsFlags.h"
 
 //  If this is enabled, you also need to touch VolumetricCloud.usf for shaders to be recompiled.
 #define CLOUD_DEBUG_SAMPLES 0 /*!Never check in enabled!*/
@@ -824,6 +825,7 @@ class FRenderVolumetricCloudRenderViewCS : public FMeshMaterialShader
 		SHADER_PARAMETER(FVector4f, OutputViewRect)
 		SHADER_PARAMETER(FVector2f, StartTracingDistanceTextureResolution)
 		SHADER_PARAMETER(float, StartTracingSampleVolumeDepth)
+		SHADER_PARAMETER(int32, bAccumulateAlphaHoldOut)
 		SHADER_PARAMETER(int32, bBlendCloudColor)
 		SHADER_PARAMETER(int32, TargetCubeFace)
 		SHADER_PARAMETER_RDG_UNIFORM_BUFFER(FRenderVolumetricCloudGlobalParameters, VolumetricCloudRenderViewParamsUB)
@@ -2060,6 +2062,7 @@ FCloudRenderContext::FCloudRenderContext()
 
 	bAsyncCompute = false;
 	bCloudDebugViewModeEnabled = false;
+	bAccumulateAlphaHoldOut = false;
 }
 
 static TAutoConsoleVariable<int32> CVarCloudDefaultTexturesNoFastClear(
@@ -2386,6 +2389,7 @@ void FSceneRenderer::RenderVolumetricCloudsInternal(FRDGBuilder& GraphBuilder, F
 		PassParameters->StartTracingDistanceTextureResolution = FVector2f(EmptySpaceSkippingTexture->Desc.Extent.X, EmptySpaceSkippingTexture->Desc.Extent.Y);
 		PassParameters->StartTracingDistanceTexture = EmptySpaceSkippingTexture;
 		PassParameters->StartTracingSampleVolumeDepth = GetEmptySpaceSkippingVolumeDepth();
+		PassParameters->bAccumulateAlphaHoldOut = CloudRC.bAccumulateAlphaHoldOut ? 1 : 0;
 
 #if CLOUD_DEBUG_SAMPLES
 		ShaderPrint::SetEnabled(true);
@@ -2478,6 +2482,7 @@ bool FSceneRenderer::RenderVolumetricCloud(
 	const FMinimalSceneTextures& SceneTextures,
 	bool bSkipVolumetricRenderTarget,
 	bool bSkipPerPixelTracing,
+	bool bAccumulateAlphaHoldOut,
 	FRDGTextureRef HalfResolutionDepthCheckerboardMinMaxTexture,
 	FRDGTextureRef QuarterResolutionDepthMinMaxTexture,
 	bool bAsyncCompute,
@@ -2498,7 +2503,7 @@ bool FSceneRenderer::RenderVolumetricCloud(
 		FMaterialRenderProxy* CloudVolumeMaterialProxy = CloudSceneProxy.GetCloudVolumeMaterial()->GetRenderProxy();
 		if (CloudVolumeMaterialProxy->GetIncompleteMaterialWithFallback(ViewFamily.GetFeatureLevel()).GetMaterialDomain() == MD_Volume)
 		{
-			RDG_EVENT_SCOPE(GraphBuilder, "VolumetricCloud");
+			RDG_EVENT_SCOPE(GraphBuilder, "%s", bAccumulateAlphaHoldOut ? "VolumetricCloudAlphaHoldout" : "VolumetricCloud");
 			RDG_GPU_STAT_SCOPE(GraphBuilder, VolumetricCloud);
 			SCOPED_NAMED_EVENT(VolumetricCloud, FColor::Emerald);
 
@@ -2527,7 +2532,14 @@ bool FSceneRenderer::RenderVolumetricCloud(
 			{
 				FViewInfo& ViewInfo = Views[ViewIndex];
 
-				bool bShouldViewRenderVolumetricCloudRenderTarget = ShouldViewRenderVolumetricCloudRenderTarget(ViewInfo); // not used by reflection captures for instance
+				if (bAccumulateAlphaHoldOut && ViewInfo.CachedViewUniformShaderParameters->RenderingReflectionCaptureMask > 0.0f)
+				{
+					continue;
+				}
+
+				// RenderTarget are not used by reflection captures for instance.
+				bool bShouldViewRenderVolumetricCloudRenderTarget = ShouldViewRenderVolumetricCloudRenderTarget(ViewInfo) && !bAccumulateAlphaHoldOut; // When rendering alpha holdout per pixel, we never skip even if a volumetric render target is used.; 
+
 				if ((bShouldViewRenderVolumetricCloudRenderTarget && bSkipVolumetricRenderTarget) || (!bShouldViewRenderVolumetricCloudRenderTarget && bSkipPerPixelTracing))
 				{
 					continue;
@@ -2651,6 +2663,11 @@ bool FSceneRenderer::RenderVolumetricCloud(
 				}
 				else
 				{
+					if (ViewInfo.CachedViewUniformShaderParameters->RenderingReflectionCaptureMask == 0 && !IsVolumetricCloudRenderedInMain(ViewInfo.CachedViewUniformShaderParameters->EnvironmentComponentsFlags))
+					{
+						continue;
+					}
+
 					DestinationRT = SceneTextures.Color.Target;
 					const FIntPoint RtSize = SceneTextures.Config.Extent;
 
@@ -2720,6 +2737,8 @@ bool FSceneRenderer::RenderVolumetricCloud(
 					CloudRC.VirtualShadowMapId0 = INDEX_NONE;
 				}
 
+				CloudRC.bAccumulateAlphaHoldOut = bAccumulateAlphaHoldOut;
+
 				// Cannot nest a global buffer into another one and we are limited to only one PassUniformBuffer on PassDrawRenderState.
 				//TUniformBufferRef<FVolumeShadowingShaderParametersGlobal0> LightShadowShaderParams0UniformBuffer = TUniformBufferRef<FVolumeShadowingShaderParametersGlobal0>::CreateUniformBufferImmediate(LightShadowShaderParams0, UniformBuffer_SingleFrame);
 
diff --git a/Engine/Source/Runtime/Renderer/Private/VolumetricCloudRendering.h b/Engine/Source/Runtime/Renderer/Private/VolumetricCloudRendering.h
--- a/Engine/Source/Runtime/Renderer/Private/VolumetricCloudRendering.h
+++ b/Engine/Source/Runtime/Renderer/Private/VolumetricCloudRendering.h
@@ -146,6 +146,7 @@ struct FCloudRenderContext
 
 	bool bAsyncCompute;
 	bool bCloudDebugViewModeEnabled;
+	bool bAccumulateAlphaHoldOut;
 
 	FUintVector4 TracingCoordToZbufferCoordScaleBias;
 	FUintVector4 TracingCoordToFullResPixelCoordScaleBias;
diff --git a/Engine/Source/Runtime/Renderer/Private/VolumetricRenderTarget.cpp b/Engine/Source/Runtime/Renderer/Private/VolumetricRenderTarget.cpp
--- a/Engine/Source/Runtime/Renderer/Private/VolumetricRenderTarget.cpp
+++ b/Engine/Source/Runtime/Renderer/Private/VolumetricRenderTarget.cpp
@@ -14,6 +14,7 @@
 #include "VolumetricCloudRendering.h"
 #include "RendererUtils.h"
 #include "PostProcess/PostProcessing.h" // IsPostProcessingWithAlphaChannelSupported
+#include "EnvironmentComponentsFlags.h"
 
 static TAutoConsoleVariable<int32> CVarVolumetricRenderTarget(
 	TEXT("r.VolumetricRenderTarget"), 1,
@@ -760,6 +761,8 @@ void ComposeVolumetricRenderTargetOverScene(
 	const bool bSupportsAlpha = IsPostProcessingWithAlphaChannelSupported();
 	if (bSupportsAlpha)
 	{
+		// When alpha channel is enabled, we always write transmittance to impact other alpha holdout values from sky or fog for instance.
+		// We will run a second pass later accumulating the cloud contribution to hold out
 		PreMultipliedColorTransmittanceBlend = TStaticBlendState<CW_RGBA, BO_Add, BF_One, BF_SourceAlpha, BO_Add, BF_Zero, BF_SourceAlpha>::GetRHI();
 	}
 	else
@@ -774,6 +777,11 @@ void ComposeVolumetricRenderTargetOverScene(
 		{
 			continue;
 		}
+		if (ViewInfo.CachedViewUniformShaderParameters->RenderingReflectionCaptureMask == 0 && !IsVolumetricCloudRenderedInMain(ViewInfo.CachedViewUniformShaderParameters->EnvironmentComponentsFlags))
+		{
+			continue;
+		}
+
 		FVolumetricRenderTargetViewStateData& VolumetricCloudRT = ViewInfo.ViewState->VolumetricCloudRenderTarget;
 
 		// When reconstructed and back buffer resolution matches, force using a pixel perfect upsampling.
