Fix LWCData not populated in slate vertex shaders

#tests PIE verify issue fixed

diff --git a/Engine/Shaders/Private/SlateVertexShader.usf b/Engine/Shaders/Private/SlateVertexShader.usf
--- a/Engine/Shaders/Private/SlateVertexShader.usf
+++ b/Engine/Shaders/Private/SlateVertexShader.usf
@@ -45,6 +45,7 @@ VertexToPixelInterpolants Main(
 	// Init primitive flags to a good default.
 	VertexParameters.SceneData.Primitive.Flags = GetDefaultPrimitiveSceneDataFlags();
 	VertexParameters.SceneData.Primitive.VisibilityFlags = GetDefaultPrimitiveVisibilityFlags();
+	VertexParameters.LWCData = MakeMaterialLWCData(VertexParameters);
 
 	bool bCalledEvaluateVertexMaterialAttributes = false;
 
