Reverse change to ortho material camera vector calculation to fix issues with ResolvedViews being skipped.

diff --git a/Engine/Shaders/Private/MaterialTemplate.ush b/Engine/Shaders/Private/MaterialTemplate.ush
--- a/Engine/Shaders/Private/MaterialTemplate.ush
+++ b/Engine/Shaders/Private/MaterialTemplate.ush
@@ -4101,7 +4101,7 @@ void CalcMaterialParametersEx(
 	#if RAYHITGROUPSHADER
 		Parameters.CameraVector = -WorldRayDirection();
 	#else
-		Parameters.CameraVector = -GetCameraVectorFromTranslatedWorldPosition(Parameters.WorldPosition_CamRelative.xyz);
+		Parameters.CameraVector = select(IsOrthoProjection(ResolvedView), -ResolvedView.ViewForward, normalize(-Parameters.WorldPosition_CamRelative.xyz));
 	#endif
 
 	Parameters.LightVector = 0;
