Swap RootUV and ColorXY channel in the new MeshCompact layout generate by the hair strands texture systemt.

This layout has been added for 5.4, but the current packing does not provide enough resolution for RootUV channel. After discussion with artists, the channels are swapped to address this.

diff --git a/Engine/Plugins/Runtime/HairStrands/Source/HairStrandsCore/Private/GroomAssetCards.cpp b/Engine/Plugins/Runtime/HairStrands/Source/HairStrandsCore/Private/GroomAssetCards.cpp
--- a/Engine/Plugins/Runtime/HairStrands/Source/HairStrandsCore/Private/GroomAssetCards.cpp
+++ b/Engine/Plugins/Runtime/HairStrands/Source/HairStrandsCore/Private/GroomAssetCards.cpp
@@ -145,8 +145,8 @@ const TCHAR* GetHairTextureLayoutTextureName(EHairTextureLayout InLayout, uint32
 			switch (InIndex)
 			{
 				case 0: return bDetail ? TEXT("Tangent | CoordU\n RGB8")					: TEXT("TangentCoordU");
-				case 1: return bDetail ? TEXT("RooUV | Depth | GroupID\n RGBA8")			: TEXT("RooUVDepthGroupID");
-				case 2: return bDetail ? TEXT("ColorXY | Seed | Coverage \n RGBA8")			: TEXT("ColorXYSeedCoverage");
+				case 1: return bDetail ? TEXT("ColorXY | Depth | GroupID\n RGBA8")			: TEXT("ColorXYDepthGroupID");
+				case 2: return bDetail ? TEXT("RooUV | Seed | Coverage \n RGBA8")			: TEXT("RooUVSeedCoverage");
 			}
 		} break;
 	}
diff --git a/Engine/Shaders/Private/HairStrands/HairCardsAttributeCommon.ush b/Engine/Shaders/Private/HairStrands/HairCardsAttributeCommon.ush
--- a/Engine/Shaders/Private/HairStrands/HairCardsAttributeCommon.ush
+++ b/Engine/Shaders/Private/HairStrands/HairCardsAttributeCommon.ush
@@ -42,8 +42,8 @@ uint GetHairCardTextureIndex(uint InLayoutIndex, uint InAttribute)
 		case EHairCardsVFAttribute_CoordU     : return LAYOUT_VALUES(InLayoutIndex, 3, 3, 0, 0);
 		case EHairCardsVFAttribute_GroupIndex : return LAYOUT_VALUES(InLayoutIndex, 9, 4, 9, 1); // Only used for mesh, not for cards
 		case EHairCardsVFAttribute_Roughness  : return LAYOUT_VALUES(InLayoutIndex, 9, 9, 9, 9);
-		case EHairCardsVFAttribute_RootUV     : return LAYOUT_VALUES(InLayoutIndex, 9, 3, 9, 1);
-		case EHairCardsVFAttribute_Color      : return LAYOUT_VALUES(InLayoutIndex, 9, 4, 9, 2);
+		case EHairCardsVFAttribute_RootUV     : return LAYOUT_VALUES(InLayoutIndex, 9, 3, 9, 2);
+		case EHairCardsVFAttribute_Color      : return LAYOUT_VALUES(InLayoutIndex, 9, 4, 9, 1);
 		case EHairCardsVFAttribute_Tangent    : return LAYOUT_VALUES(InLayoutIndex, 2, 2, 0, 0);
 		case EHairCardsVFAttribute_Coverage   : return LAYOUT_VALUES(InLayoutIndex, 1, 1, 1, 2);
 		case EHairCardsVFAttribute_Auxilary   : return LAYOUT_VALUES(InLayoutIndex, 5, 5, 9, 9); 
