Fixed translucent material using alpha to coverage when it is in fact disabled on the material.
MATERIAL_USE_ALPHA_TO_COVERAGE was not used to check for usage.

Added support for Alpha-to-coverage masking for up to MSAA 8X.

diff --git a/Engine/Shaders/Private/MaterialTemplate.ush b/Engine/Shaders/Private/MaterialTemplate.ush
--- a/Engine/Shaders/Private/MaterialTemplate.ush
+++ b/Engine/Shaders/Private/MaterialTemplate.ush
@@ -42,7 +42,7 @@
 
 // Update this GUID to force all material shaders to recompile
 // Merge conflicts on this line should be resolved by generating a new GUID
-#pragma message("UESHADERMETADATA_VERSION A45831C8-997A-4387-907E-31EE08D781F7")
+#pragma message("UESHADERMETADATA_VERSION 50476854-A3D1-477E-91EA-0ED4AF89E199")
 
 #if USES_SPEEDTREE
 	#include "/Engine/Private/SpeedTreeCommon.ush"
@@ -4006,16 +4006,20 @@ void GetMaterialClippingVelocity(FMaterialPixelParameters Parameters, FPixelMate
 #endif
 }
 
-#define MATERIALBLENDING_MASKED_USING_COVERAGE (FORWARD_SHADING && MATERIALBLENDING_MASKED && SUPPORTS_PIXEL_COVERAGE)
+#define MATERIALBLENDING_MASKED_USING_COVERAGE (FORWARD_SHADING && MATERIALBLENDING_MASKED && SUPPORTS_PIXEL_COVERAGE && MATERIAL_USE_ALPHA_TO_COVERAGE)
 #if MATERIALBLENDING_MASKED_USING_COVERAGE
 
 uint GetDerivativeCoverageFromMask(float MaterialMask)
 {
-	uint Coverage = 0x0;
-	if (MaterialMask > 0.01) Coverage = 0x8;
-	if (MaterialMask > 0.25) Coverage = 0x9;
-	if (MaterialMask > 0.50) Coverage = 0xD;
-	if (MaterialMask > 0.75) Coverage = 0xF;
+	uint Coverage = 0x0;						// MSAA sampler mask support up to 8X
+	if (MaterialMask > 0.010) Coverage = 0x08;	// 0000 1000
+	if (MaterialMask > 0.125) Coverage = 0x18;	// 0001 1000
+	if (MaterialMask > 0.250) Coverage = 0x19;	// 0001 1001
+	if (MaterialMask > 0.375) Coverage = 0x39;	// 0011 1001
+	if (MaterialMask > 0.500) Coverage = 0x3D;	// 0011 1101
+	if (MaterialMask > 0.625) Coverage = 0x7D;	// 0111 1101
+	if (MaterialMask > 0.750) Coverage = 0x7F;	// 0111 1111
+	if (MaterialMask > 0.875) Coverage = 0xFF;	// 1111 1111
 	return Coverage;
 }
 
diff --git a/Engine/Source/Runtime/Engine/Private/Materials/MaterialShared.cpp b/Engine/Source/Runtime/Engine/Private/Materials/MaterialShared.cpp
--- a/Engine/Source/Runtime/Engine/Private/Materials/MaterialShared.cpp
+++ b/Engine/Source/Runtime/Engine/Private/Materials/MaterialShared.cpp
@@ -2661,7 +2661,7 @@ void FMaterial::SetupMaterialEnvironment(
 	SET_SHADER_DEFINE(OutEnvironment, TRANSLUCENT_SHADOW_WITH_MASKED_OPACITY, GetCastDynamicShadowAsMasked());
 	SET_SHADER_DEFINE(OutEnvironment, TRANSLUCENT_WRITING_VELOCITY, IsTranslucencyWritingVelocity());
 	SET_SHADER_DEFINE(OutEnvironment, TRANSLUCENT_WRITING_FRONT_LAYER_TRANSPARENCY, IsTranslucencyWritingFrontLayerTransparency());
-	SET_SHADER_DEFINE(OutEnvironment, MATERIAL_USE_ALPHA_TO_COVERAGE, IsUsingAlphaToCoverage());
+	SET_SHADER_DEFINE(OutEnvironment, MATERIAL_USE_ALPHA_TO_COVERAGE, IsUsingAlphaToCoverage() ? 1 : 0);
 	SET_SHADER_DEFINE(OutEnvironment, MOBILE_HIGH_QUALITY_BRDF, IsMobileHighQualityBRDFEnabled());
 	SET_SHADER_DEFINE(OutEnvironment, MATERIAL_TRANSLUCENT_PASS_AFTERMOTIONBLUR, IsTranslucencyAfterMotionBlurEnabled() ? 1 : 0);
 
