VCloud - made mode 2 still use reprojection using cloud approximate depth as this help reduce latency/.

diff --git a/Engine/Shaders/Private/VolumetricCloud.usf b/Engine/Shaders/Private/VolumetricCloud.usf
--- a/Engine/Shaders/Private/VolumetricCloud.usf
+++ b/Engine/Shaders/Private/VolumetricCloud.usf
@@ -1646,7 +1646,7 @@ void MainCommon(in FMaterialPixelParameters MaterialParameters, in float4 SvPosi
 #endif // CLOUD_MIN_AND_MAX
 
 	OutDepth.x = MaxHalfFloat; // Default to far away depth to be flat and not intersect with any geometry.
-	if (RenderVolumetricCloudParameters.OpaqueIntersectionMode >= 1)
+	//if (RenderVolumetricCloudParameters.OpaqueIntersectionMode >= 1) => Always __XB_GDS_Write_U32 depth as this helps with reprojection in all modes, also mode 2 where this helps reduce latency.
 	{
 		OutDepth.x = ((GrayScaleTransmittance > 0.99) ? NoCloudDepth : tAP) * CENTIMETER_TO_KILOMETER; // using a small threshold on transmittance
 	}
