Fix failing engine hair test.

Force ortho flag was incorrectly applied on primary visibility.

diff --git a/Engine/Shaders/Private/HairStrands/HairStrandsVertexFactory.ush b/Engine/Shaders/Private/HairStrands/HairStrandsVertexFactory.ush
--- a/Engine/Shaders/Private/HairStrands/HairStrandsVertexFactory.ush
+++ b/Engine/Shaders/Private/HairStrands/HairStrandsVertexFactory.ush
@@ -679,7 +679,7 @@ float4 VertexFactoryGetWorldPosition_Visibility(FVertexFactoryInput Input, FHair
 
 	// Minimal radius to snap the strand to a sample/pixel center (to avoid aliasing)
 	const float DistanceToCamera = length(HairViewInfo.TranslatedWorldCameraOrigin - WorldPosition.xyz);
-	const float MinStrandHairRadius = ConvertGivenDepthRadiusForProjectionType(HairViewInfo.RadiusAtDepth1, DistanceToCamera, HairViewInfo.bIsOrthoView);
+	const float MinStrandHairRadius = ConvertGivenDepthRadiusForProjectionType(HairViewInfo.RadiusAtDepth1, DistanceToCamera);
 	const float3 ViewDir = -HairViewInfo.ViewForward;
 	const float3 Right = normalize(cross(WorldTangent, ViewDir));
 	const float3 OutWorldPosition = WorldPosition.xyz + (VertexInfo.IsLeft ? -Right : Right) * max(ControlPoint.WorldRadius, MinStrandHairRadius);
