Out of range clipmaps should be lit not shadowed

diff --git a/Engine/Shaders/Private/VirtualShadowMaps/VirtualShadowMapProjectionDirectional.ush b/Engine/Shaders/Private/VirtualShadowMaps/VirtualShadowMapProjectionDirectional.ush
--- a/Engine/Shaders/Private/VirtualShadowMaps/VirtualShadowMapProjectionDirectional.ush
+++ b/Engine/Shaders/Private/VirtualShadowMaps/VirtualShadowMapProjectionDirectional.ush
@@ -230,130 +230,138 @@ FVirtualShadowMapSampleResult TraceDirectional(
 
 	FVirtualShadowMapSampleResult Result = InitVirtualShadowMapSampleResult();
 	Result.bValid = true;	// TODO: false if all samples of all rays miss pages?
-	Result.ShadowFactor = 0.0f;
+	Result.ShadowFactor = 1.0f;
+
+	if (bCullBackface && IsBackfaceToDirectionalLight(WorldNormal, Light.Direction, Light.SourceRadius))
+	{
+		Result.ShadowFactor = 0.0f;		// Fully shadowed on backfaces
+		return Result;
+	}
 
 	// Find the best resolution mapped clipmap at the ray origin location
 	const int ClipmapId = GetMappedClipmapId(VirtualShadowMapId, TranslatedWorldPosition);
+	if (ClipmapId < 0)
+	{
+		// TODO: False for valid probably? But make sure that doesn't do anything weird as we have no fallback here
+		Result.ShadowFactor = 1.0f;		// Fully lit outside clipmap range
+		return Result;
+	}
 
-	if ((!bCullBackface || !IsBackfaceToDirectionalLight(WorldNormal, Light.Direction, Light.SourceRadius)) && ClipmapId >= 0)
+	const FVirtualShadowMapProjectionShaderData ProjectionData = GetVirtualShadowMapProjectionData(ClipmapId);
+
+	// This function is designed to vary in depth in roughly the same way as shadow texel density does,
+	// thus the dependence on Clipmap ResolutionLodBias. The goal here is to define a texel dither that
+	// is "smooth" and doesn't have obvious artifacts at page edges of different clipmap levels.
+	// 
+	// This could become a problem later if we break too far from that global resolution bias. That said,
+	// these cases would have unavoidable discontinuities at page edges, so that would have to be handled
+	// in some different way.
+	//
+	// 0.5 is arbitrary but makes it consistent with the previous dither scale cvar
+	float PerLightTexelDitherScale = TexelDitherScale * ProjectionData.TexelDitherScale;
+	float DitherScale = 0.0f;
+	if (PerLightTexelDitherScale > 0.0f)
 	{
-		const FVirtualShadowMapProjectionShaderData ProjectionData = GetVirtualShadowMapProjectionData(ClipmapId);
-
-		// This function is designed to vary in depth in roughly the same way as shadow texel density does,
-		// thus the dependence on Clipmap ResolutionLodBias. The goal here is to define a texel dither that
-		// is "smooth" and doesn't have obvious artifacts at page edges of different clipmap levels.
-		// 
-		// This could become a problem later if we break too far from that global resolution bias. That said,
-		// these cases would have unavoidable discontinuities at page edges, so that would have to be handled
-		// in some different way.
-		//
-		// 0.5 is arbitrary but makes it consistent with the previous dither scale cvar
-		float PerLightTexelDitherScale = TexelDitherScale * ProjectionData.TexelDitherScale;
-		float DitherScale = 0.0f;
-		if (PerLightTexelDitherScale > 0.0f)
-		{
-			DitherScale =
-				((0.5f / float(CalcLevelDimsTexels(0))) * PerLightTexelDitherScale * DistanceFromViewOrigin) /
-				(exp2(ProjectionData.ClipmapLevel - ProjectionData.ResolutionLodBias));
-		}
+		DitherScale =
+			((0.5f / float(CalcLevelDimsTexels(0))) * PerLightTexelDitherScale * DistanceFromViewOrigin) /
+			(exp2(ProjectionData.ClipmapLevel - ProjectionData.ResolutionLodBias));
+	}
 
-		float2 DepthSlopeUV = ComputeDepthSlopeDirectionalUV(ProjectionData, WorldNormal);
+	float2 DepthSlopeUV = ComputeDepthSlopeDirectionalUV(ProjectionData, WorldNormal);
 
-		// Compute max ray length based on view depth
-		// This affects both how far we are willing to trace across the shadow map (for clipmaps this is related to view depth)
-		// and the maximum size a penumbra can be.
-		// Too high values will cause shadows to detach from their contact points (unless more samples are used).
-		// Too low values will greatly restrict how large penumbras can be in screen space.
-		float RayLength = RayLengthScale * DistanceFromViewOrigin;
+	// Compute max ray length based on view depth
+	// This affects both how far we are willing to trace across the shadow map (for clipmaps this is related to view depth)
+	// and the maximum size a penumbra can be.
+	// Too high values will cause shadows to detach from their contact points (unless more samples are used).
+	// Too low values will greatly restrict how large penumbras can be in screen space.
+	float RayLength = RayLengthScale * DistanceFromViewOrigin;
 
-		// Clipmap origin should be near shaded samples, so shadow translated world should be regular range
-		float3 ViewToShadowTranslation = DFFastLocalSubtractDemote(ProjectionData.PreViewTranslation, PrimaryView.PreViewTranslation);
-		float3 RayOriginShadowTranslatedWorld = TranslatedWorldPosition + ViewToShadowTranslation;
+	// Clipmap origin should be near shaded samples, so shadow translated world should be regular range
+	float3 ViewToShadowTranslation = DFFastLocalSubtractDemote(ProjectionData.PreViewTranslation, PrimaryView.PreViewTranslation);
+	float3 RayOriginShadowTranslatedWorld = TranslatedWorldPosition + ViewToShadowTranslation;
 
-		uint RayMissCount = 0;
-		uint i = 0;
-		float OccluderDistanceSum = 0.0f;
-		float MaxOccluderDistance = -1.0f;
-		for ( ; i < MaxRayCount; i++)
-		{
-			// One sample for ray, one for texel dither
-			float4 RandSample = GetRandomSample(PixelPos, View.StateFrameIndex, i, MaxRayCount);
-			float3 RayDir = GetRandomDirectionalLightRayDir(Light, RandSample.xy);
-			float2 TexelOffset = (RandSample.zw - 0.5f) * DitherScale;
-
-			FSMRTClipmapRayState RayState = SMRTClipmapRayInitialize(
-				ProjectionData,
-				RayOriginShadowTranslatedWorld,
-				RayDir,
-				RayLength,
-				RayStartOffset,
-				DepthSlopeUV,
-				TexelOffset,
-				ExtrapolateSlope);
-			FSMRTResult SMRTResult = SMRTRayCast(RayState, SamplesPerRay, Noise);
-
-			// Debug output (DCE'd if not used)
-			Result.ClipmapOrMipLevel = GetVirtualShadowMapProjectionData(RayState.SampledClipmapId).ClipmapLevel;
-			Result.VirtualTexelAddress = RayState.VirtualTexelAddress;
-			Result.PhysicalTexelAddress =  RayState.PhysicalTexelAddress;
-			//Result.GeneralDebug = GreenToRedTurbo((RayState.SampledClipmapId - RayState.ClipmapId) / 2.0f);
+	uint RayMissCount = 0;
+	uint i = 0;
+	float OccluderDistanceSum = 0.0f;
+	float MaxOccluderDistance = -1.0f;
+	for ( ; i < MaxRayCount; i++)
+	{
+		// One sample for ray, one for texel dither
+		float4 RandSample = GetRandomSample(PixelPos, View.StateFrameIndex, i, MaxRayCount);
+		float3 RayDir = GetRandomDirectionalLightRayDir(Light, RandSample.xy);
+		float2 TexelOffset = (RandSample.zw - 0.5f) * DitherScale;
+
+		FSMRTClipmapRayState RayState = SMRTClipmapRayInitialize(
+			ProjectionData,
+			RayOriginShadowTranslatedWorld,
+			RayDir,
+			RayLength,
+			RayStartOffset,
+			DepthSlopeUV,
+			TexelOffset,
+			ExtrapolateSlope);
+		FSMRTResult SMRTResult = SMRTRayCast(RayState, SamplesPerRay, Noise);
+
+		// Debug output (DCE'd if not used)
+		Result.ClipmapOrMipLevel = GetVirtualShadowMapProjectionData(RayState.SampledClipmapId).ClipmapLevel;
+		Result.VirtualTexelAddress = RayState.VirtualTexelAddress;
+		Result.PhysicalTexelAddress =  RayState.PhysicalTexelAddress;
+		//Result.GeneralDebug = GreenToRedTurbo((RayState.SampledClipmapId - RayState.ClipmapId) / 2.0f);
 			
-			if (SMRTResult.bValidHit)
-			{
-				float OccluderDistance = ComputeOccluderDistanceOrtho(
-					GetVirtualShadowMapProjectionData(RayState.ClipmapId).ShadowViewToClipMatrix,
-					SMRTResult.HitDepth,
-					RayState.RayStartUVZ.z);
+		if (SMRTResult.bValidHit)
+		{
+			float OccluderDistance = ComputeOccluderDistanceOrtho(
+				GetVirtualShadowMapProjectionData(RayState.ClipmapId).ShadowViewToClipMatrix,
+				SMRTResult.HitDepth,
+				RayState.RayStartUVZ.z);
 
-				OccluderDistanceSum += OccluderDistance;
-				MaxOccluderDistance = max(MaxOccluderDistance, OccluderDistance);
-			}
-			else
-			{
-				++RayMissCount;
-			}
+			OccluderDistanceSum += OccluderDistance;
+			MaxOccluderDistance = max(MaxOccluderDistance, OccluderDistance);
+		}
+		else
+		{
+			++RayMissCount;
+		}
 
-			#if COMPUTESHADER
-			if (SMRTAdaptiveRayCount > 0)
+		#if COMPUTESHADER
+		if (SMRTAdaptiveRayCount > 0)
+		{
+			// TODO: Adapt this heuristic based on SMRTAdaptiveRayCount as well?
+			if( i == 0 )
 			{
-				// TODO: Adapt this heuristic based on SMRTAdaptiveRayCount as well?
-				if( i == 0 )
+				bool bHit = SMRTResult.bValidHit;
+
+				// All lanes missed
+				bool bAllLanesMiss = WaveActiveAllTrue( !bHit );
+				if( bAllLanesMiss )
 				{
-					bool bHit = SMRTResult.bValidHit;
-
-					// All lanes missed
-					bool bAllLanesMiss = WaveActiveAllTrue( !bHit );
-					if( bAllLanesMiss )
-					{
-						break;
-					}
+					break;
 				}
-				else if( i >= SMRTAdaptiveRayCount )
+			}
+			else if( i >= SMRTAdaptiveRayCount )
+			{
+				// After 2 iterations and all have hit, assume umbra
+				bool bAllLanesHit = WaveActiveAllTrue( RayMissCount == 0 );
+				if( bAllLanesHit )
 				{
-					// After 2 iterations and all have hit, assume umbra
-					bool bAllLanesHit = WaveActiveAllTrue( RayMissCount == 0 );
-					if( bAllLanesHit )
-					{
-						break;
-					}
+					break;
 				}
 			}
-			#endif
 		}
-		uint RayCount = min(i + 1U, MaxRayCount);		// break vs regular for loop exit
-
-		/* TODO: Experiment with max vs avg distance in various cases
-		float OccluderDistance = (View.GeneralPurposeTweak == 0)
-			? MaxOccluderDistance
-			: (OccluderDistanceSum / float(max(1, RayCount - RayMissCount)));
-		*/
-		float OccluderDistance = (OccluderDistanceSum / float(max(1U, RayCount - RayMissCount)));
-
-		Result.ShadowFactor = float(RayMissCount) / float(RayCount);
-		Result.OccluderDistance = OccluderDistance;
-		Result.RayCount = RayCount;
+		#endif
 	}
-
+	uint RayCount = min(i + 1U, MaxRayCount);		// break vs regular for loop exit
+
+	/* TODO: Experiment with max vs avg distance in various cases
+	float OccluderDistance = (View.GeneralPurposeTweak == 0)
+		? MaxOccluderDistance
+		: (OccluderDistanceSum / float(max(1, RayCount - RayMissCount)));
+	*/
+	float OccluderDistance = (OccluderDistanceSum / float(max(1U, RayCount - RayMissCount)));
+
+	Result.ShadowFactor = float(RayMissCount) / float(RayCount);
+	Result.OccluderDistance = OccluderDistance;
+	Result.RayCount = RayCount;
 	return Result;
 }
 
