VCloud - no history reprojction when using mode 2 now uses bilinear filtering.

diff --git a/Engine/Shaders/Private/VolumetricRenderTarget.usf b/Engine/Shaders/Private/VolumetricRenderTarget.usf
--- a/Engine/Shaders/Private/VolumetricRenderTarget.usf
+++ b/Engine/Shaders/Private/VolumetricRenderTarget.usf
@@ -198,8 +198,7 @@ void ReconstructVolumetricRenderTargetPS(
 		{
 			const bool bUseNewSampleMode2 = ((IntPixelPos.x - IntPixelPosDownsample.x * DownSampleFactor) == CurrentTracingPixelOffset.x) && ((IntPixelPos.y - IntPixelPosDownsample.y * DownSampleFactor) == CurrentTracingPixelOffset.y);
 
-			// Always use new sample, reproject previous frame samples/pixels
-			if (bUseNewSampleMode2 || !bValidPreviousUVs)
+			if (bUseNewSampleMode2)
 			{
 				// Load the new sample for this pixel we have just traced
 				RGBA = SafeLoadTracingVolumetricTexture(int2(SVPos.xy) / DownSampleFactor);// +float4(0.1, 0.0, 0, 0);
@@ -207,13 +206,13 @@ void ReconstructVolumetricRenderTargetPS(
 			}
 			else if(bValidPreviousUVs)
 			{
-				// Sample valid on screen history
+				// Sample valid on screen history using bilinear filtering
 				RGBA = SafeSamplePreviousFrameVolumetricTexture(PrevScreenUVs);// +float4(0, 0.1, 0, 0);
 				DepthData = SafeSamplePreviousFrameVolumetricDepthTexture(PrevScreenUVs);
 			}
 			else
 			{
-				// Bilinear sample new low resoltuion tracing
+				// Sample the new low resolution clouds we just traced using bilinear filtering 
 				RGBA = SafeSampleTracingVolumetricTexture(ScreenUV);
 				DepthData = SafeSampleTracingVolumetricDepthTexture(ScreenUV);
 			}
