VCloud - Mode 2 - Bias the sample such that the currently-traced screen pixels line up with the centre of their traced values in the tracing texture. This prevents clouds from appearing to shift position when rotating the camera quickly.

diff --git a/Engine/Shaders/Private/VolumetricRenderTarget.usf b/Engine/Shaders/Private/VolumetricRenderTarget.usf
--- a/Engine/Shaders/Private/VolumetricRenderTarget.usf
+++ b/Engine/Shaders/Private/VolumetricRenderTarget.usf
@@ -212,9 +212,14 @@ void ReconstructVolumetricRenderTargetPS(
 			}
 			else
 			{
+				// Bias the sample such that the currently-traced screen pixels line up with the centre of their traced values in the tracing texture
+				// This prevents clouds from appearing to shift position when rotating the camera quickly
+				float2 ScreenPixelsOffset = -(CurrentTracingPixelOffset - (DownSampleFactor - 1.0f) * 0.5f);
+				float2 TracingTextureUV = ScreenUV + ScreenPixelsOffset * DstVolumetricTextureSizeAndInvSize.zw;
+
 				// Sample the new low resolution clouds we just traced using bilinear filtering 
-				RGBA = SafeSampleTracingVolumetricTexture(ScreenUV);
-				DepthData = SafeSampleTracingVolumetricDepthTexture(ScreenUV);
+				RGBA = SafeSampleTracingVolumetricTexture(TracingTextureUV);
+				DepthData = SafeSampleTracingVolumetricDepthTexture(TracingTextureUV);
 			}
 			OutputRt0 = RGBA;
 			OutputRt1 = float2(DepthData.CloudFrontDepthFromViewKm, DepthData.SceneDepthFromViewKm);
