Path Tracer: Simplify implementation of MIS weights to make it easier for the compiler to optimize

Also cleanup duplication definitions of AddLobeWithMIS in Substrate by moving it to a common header.

Rename MISWeightRobust to MISWeightPower to make it clear this is using the power heuristic

diff --git a/Engine/Plugins/Experimental/GPULightmass/Shaders/Private/LightmapPathTracing.usf b/Engine/Plugins/Experimental/GPULightmass/Shaders/Private/LightmapPathTracing.usf
--- a/Engine/Plugins/Experimental/GPULightmass/Shaders/Private/LightmapPathTracing.usf
+++ b/Engine/Plugins/Experimental/GPULightmass/Shaders/Private/LightmapPathTracing.usf
@@ -612,8 +612,8 @@ void PathTracingKernel(
 						float3 BentNormalVector = LightSample.Direction / LightSample.Pdf / PI * Visibility;
 						if (MISMode == 2)
 						{
-							LightContrib *= MISWeightRobust(LightSample.Pdf, MaterialEval.Pdf);
-							BentNormalVector *= MISWeightRobust(LightSample.Pdf, MaterialEval.Pdf);
+							LightContrib *= MISWeightPower(LightSample.Pdf, MaterialEval.Pdf);
+							BentNormalVector *= MISWeightPower(LightSample.Pdf, MaterialEval.Pdf);
 						}
 
 						// Record the contribution
@@ -768,10 +768,10 @@ void PathTracingKernel(
 					}
 					float LightPickPdf = (LightPickingCdf[Index] - PreviousCdfValue) / LightPickingCdfSum;
 
-					LightContrib *= MISWeightRobust(MaterialSample.Pdf, LightResult.Pdf * LightPickPdf);
+					LightContrib *= MISWeightPower(MaterialSample.Pdf, LightResult.Pdf * LightPickPdf);
 
 					// Separate direct lighting (bounce 0) for GPU Lightmass
-					BentNormalVector *= MISWeightRobust(MaterialSample.Pdf, LightResult.Pdf * LightPickPdf);
+					BentNormalVector *= MISWeightPower(MaterialSample.Pdf, LightResult.Pdf * LightPickPdf);
 				}
 
 				if (any(LightContrib > 0))
diff --git a/Engine/Shaders/Private/MonteCarlo.ush b/Engine/Shaders/Private/MonteCarlo.ush
--- a/Engine/Shaders/Private/MonteCarlo.ush
+++ b/Engine/Shaders/Private/MonteCarlo.ush
@@ -450,78 +450,60 @@ float4 ImportanceSampleVisibleGGX(float2 E, float2 Alpha, float3 V, bool bLimitV
 	return float4(H, VisibleGGXPDF_aniso(V, H, Alpha));
 }
 
-// Multiple importance sampling power heuristic of two functions with a power of two. 
-// [Veach 1997, "Robust Monte Carlo Methods for Light Transport Simulation"]
-float MISWeight( uint Num, float PDF, uint OtherNum, float OtherPDF )
-{
-	float Weight = Num * PDF;
-	float OtherWeight = OtherNum * OtherPDF;
-	return Weight * Weight / (Weight * Weight + OtherWeight * OtherWeight);
-}
-
 // Multiple importance sampling balance heuristic
+// [Veach 1997, "Robust Monte Carlo Methods for Light Transport Simulation"]
 float MISWeightBalanced(float Pdf, float OtherPdf)
 {
-	// The straightforward implementation is prone to numerical overflow and divisions by 0
-	// and does not work well with +inf inputs.
-
-	if (Pdf == OtherPdf)
-	{
-		// Catch potential NaNs from (0,0) and (+inf, +inf)
-		return 0.5f;
-	}
-
-	// Evaluate the expression using the ratio of the smaller value to the bigger one for greater
-	// numerical stability. The math would also work using the ratio of bigger to smaller value,
-	// which would underflow less but would make the weights asymmetric. Underflow to 0 is not a
-	// bad property to have in rendering application as it ensures more weights are exactly 0
-	// which allows some evaluations to be skipped.
-	if (OtherPdf < Pdf)
-	{
-		float x = OtherPdf / Pdf;
-		return 1.0 / (1.0 + x);
-	}
-	else
-	{
-		// this form guarantees the weights add back up to one when arguments are swapped
-		float x = Pdf / OtherPdf;
-		return 1.0 - 1.0 / (1.0 + x);
-	}
-}
-
-// Multiple importance sampling power heuristic of two functions with a power of two. 
-float MISWeightRobust(float Pdf, float OtherPdf) {
-	// The straightforward implementation above is prone to numerical overflow and divisions by 0
+	// The straightforward implementation is prone to numerical overflow, divisions by 0
 	// and does not work well with +inf inputs.
+	// return Pdf / (Pdf + OtherPdf);
 
 	// We want this function to have the following properties:
 	//  0 <= w(a,b) <= 1 for all possible positive floats a and b (including 0 and +inf)
 	//  w(a, b) + w(b, a) == 1.0
-	
+
 	// The formulation below is much more stable across the range of all possible inputs
 	// and guarantees the sum always adds up to 1.0.
 
-	if (Pdf == OtherPdf)
-	{
-		// Catch potential NaNs from (0,0) and (+inf, +inf)
-		return 0.5f;
-	}
-
 	// Evaluate the expression using the ratio of the smaller value to the bigger one for greater
 	// numerical stability. The math would also work using the ratio of bigger to smaller value,
 	// which would underflow less but would make the weights asymmetric. Underflow to 0 is not a
 	// bad property to have in rendering application as it ensures more weights are exactly 0
 	// which allows some evaluations to be skipped.
-	if (OtherPdf < Pdf)
-	{
-		float x = OtherPdf / Pdf;
-		return 1.0 / (1.0 + x * x);
-	}
-	else
+	float X = min(Pdf, OtherPdf) / max(Pdf, OtherPdf); // This ratio is guaranteed to be in [0,1]
+	float Y = Pdf == OtherPdf ? 1.0 : X; // Guard against NaNs from 0/0 and Inf/Inf
+	float M = rcp(1.0 + Y);
+	return Pdf > OtherPdf ? M : 1.0 - M; // This ensures exchanging arguments will produce values that add back up to 1.0 exactly
+}
+
+// Multiple importance sampling power heuristic of two functions with a power of two. 
+// [Veach 1997, "Robust Monte Carlo Methods for Light Transport Simulation"]
+float MISWeightPower(float Pdf, float OtherPdf)
+{
+	// Naive code (which can overflow, divide by 0, etc ..)
+	// return Pdf * Pdf / (Pdf * Pdf + OtherPdf * OtherPdf);
+
+	// See function above for the explanation of how this works
+	float X = min(Pdf, OtherPdf) / max(Pdf, OtherPdf); // This ratio is guaranteed to be in [0,1]
+	float Y = Pdf == OtherPdf ? 1.0 : X; // Guard against NaNs from 0/0 and Inf/Inf
+	float M = rcp(1.0 + Y * Y);
+	return Pdf > OtherPdf ? M : 1.0 - M; // This ensures exchanging arguments will produce values that add back up to 1.0 exactly
+}
+
+// Takes as input the sample weight and pdf for a certain lobe of a mixed model, together with the probability of picking that lobe
+// This function then updates a running total Weight and Pdf value that represents the overall contribution of the BxDF
+// This function should be called when a BxDF is made up of multiple lobes combined with a sum to correctly account for the probability
+// of sampling directions via all lobes.
+// NOTE: this function also contains special logic to handle cases with infinite pdfs cleanly
+void AddLobeWithMIS(inout float3 Weight, inout float Pdf, float3 LobeWeight, float LobePdf, float LobeProb)
+{
+	const float MinLobeProb = 1.1754943508e-38; // smallest normal float
+	if (LobeProb > MinLobeProb)
 	{
-		// this form guarantees the weights add back up to one when arguments are swapped
-		float x = Pdf / OtherPdf;
-		return 1.0 - 1.0 / (1.0 + x * x);
+		LobePdf *= LobeProb;
+		LobeWeight *= rcp(LobeProb);
+		Weight = lerp(Weight, LobeWeight, MISWeightBalanced(LobePdf, Pdf));
+		Pdf += LobePdf;
 	}
 }
 
diff --git a/Engine/Shaders/Private/PathTracing/Material/PathTracingMaterialCommon.ush b/Engine/Shaders/Private/PathTracing/Material/PathTracingMaterialCommon.ush
--- a/Engine/Shaders/Private/PathTracing/Material/PathTracingMaterialCommon.ush
+++ b/Engine/Shaders/Private/PathTracing/Material/PathTracingMaterialCommon.ush
@@ -77,27 +77,6 @@ float4 LobeSelectionPdf(float3 LobeCdf)
 	return float4(LobeCdf.x, LobeCdf.y - LobeCdf.x, LobeCdf.z - LobeCdf.y, 1.0 - LobeCdf.z);
 }
 
-
-// Takes as input the sample weight and pdf for a certain lobe of a mixed model, together with the probability of picking that lobe
-// This function then updates a running total Weight and Pdf value that represents the overall contribution of the BxDF
-// This function should be called when a BxDF is made up of multiple lobes combined with a sum to correctly account for the probability
-// of sampling directions via all lobes.
-// NOTE: this function also contains special logic to handle cases with infinite pdfs cleanly
-void AddLobeWithMIS(inout float3 Weight, inout float Pdf, float3 LobeWeight, float LobePdf, float LobeProb)
-{
-	const float MinLobeProb = 1.1754943508e-38; // smallest normal float
-	if (LobeProb > MinLobeProb)
-	{
-		LobePdf *= LobeProb;
-		LobeWeight *= 1 / LobeProb;
-
-		float MISWeight = MISWeightBalanced(LobePdf, Pdf);
-
-		Weight = lerp(Weight, LobeWeight, MISWeight);
-		Pdf += LobePdf;
-	}
-}
-
 // The following structs are used as return types for the Eval/Sample methods
 
 #define PATHTRACER_SCATTER_CAMERA   0 // this is used to mark camera rays until we have actually scattered
diff --git a/Engine/Shaders/Private/PathTracing/PathTracingCore.ush b/Engine/Shaders/Private/PathTracing/PathTracingCore.ush
--- a/Engine/Shaders/Private/PathTracing/PathTracingCore.ush
+++ b/Engine/Shaders/Private/PathTracing/PathTracingCore.ush
@@ -1497,7 +1497,7 @@ bool PathTracingKernel(inout FPathState PathState, int Bounce)
 							float3 LightContrib = Contrib * LightSample.RadianceOverPdf * GetVolumetricScatteringIntensity(LightSampleLightId) * MaterialEval.Weight * MaterialEval.Pdf;
 							if (MISMode == 2)
 							{
-								LightContrib *= MISWeightRobust(LightSample.Pdf, MaterialEval.Pdf);
+								LightContrib *= MISWeightPower(LightSample.Pdf, MaterialEval.Pdf);
 							}
 							LightContrib *= TraceTransparentVisibilityRay(LightRay, 1.0, MissShaderIndex, bCastShadows, PathState.RandSequence);
 
@@ -1522,7 +1522,7 @@ bool PathTracingKernel(inout FPathState PathState, int Bounce)
 							float3 LightContrib = Contrib * MaterialSample.Weight * LightResult.Radiance * GetVolumetricScatteringIntensity(LightSampleLightId);
 							if (MISMode == 2)
 							{
-								LightContrib *= MISWeightRobust(MaterialSample.Pdf, LightResult.Pdf);
+								LightContrib *= MISWeightPower(MaterialSample.Pdf, LightResult.Pdf);
 							}
 							MaterialRay.TMax = LightResult.HitT;
 							LightContrib *= TraceTransparentVisibilityRay(MaterialRay, MaterialSample.Roughness, MissShaderIndex, bCastShadows, PathState.RandSequence);
@@ -1665,7 +1665,7 @@ bool PathTracingKernel(inout FPathState PathState, int Bounce)
 					float3 LightContrib = PathState.PathThroughput * LightSample.RadianceOverPdf * MaterialEval.Weight * MaterialEval.Pdf;
 					if (MISMode == 2)
 					{
-						LightContrib *= MISWeightRobust(LightSample.Pdf, MaterialEval.Pdf);
+						LightContrib *= MISWeightPower(LightSample.Pdf, MaterialEval.Pdf);
 					}
 					const bool bCastShadows = bIsVolumeSample ? CastsVolumeShadow(LightId) : CastsShadow(LightId);
 					const uint MissShaderIndex = GetLightMissShaderIndex(LightId);
@@ -1803,7 +1803,7 @@ bool PathTracingKernel(inout FPathState PathState, int Bounce)
 
 			if (bUseMIS)
 			{
-				LightContrib *= MISWeightRobust(MaterialSample.Pdf, LightResult.Pdf * LightPickPdf);
+				LightContrib *= MISWeightPower(MaterialSample.Pdf, LightResult.Pdf * LightPickPdf);
 			}
 
 			if (any(LightContrib > 0))
diff --git a/Engine/Shaders/Private/Substrate/SubstrateEvaluation.ush b/Engine/Shaders/Private/Substrate/SubstrateEvaluation.ush
--- a/Engine/Shaders/Private/Substrate/SubstrateEvaluation.ush
+++ b/Engine/Shaders/Private/Substrate/SubstrateEvaluation.ush
@@ -1181,33 +1181,6 @@ float SubstrateLobeSelectionProb(float3 A, float3 B)
 	return SumA / (SumA + SumB + 1e-6);
 }
 
-// Takes as input the sample weight and pdf for a certain lobe of a mixed model, together with the probability of picking that lobe
-// This function then updates a running total Weight and Pdf value that represents the overall contribution of the BxDF
-// This function should be called when a BxDF is made up of multiple lobes combined with a sum to correctly account for the probability
-// of sampling directions via all lobes.
-// NOTE: this function also contains special logic to handle cases with infinite pdfs cleanly
-void SubstrateAddLobeWithMIS(inout float3 Weight, inout float Pdf, float3 LobeWeight, float LobePdf, float LobeProb)
-{
-	const float MinLobeProb = 1.1754943508e-38; // smallest normal float
-	if (LobeProb > MinLobeProb)
-	{
-		LobePdf *= LobeProb;
-		LobeWeight *= 1 / LobeProb;
-
-		// See discussion in MISWeightRobust for why this is better than LobePdf / (Pdf + LobePdf)
-		float LocalMISWeight;
-		if (Pdf < LobePdf)
-			LocalMISWeight = 1 / (1 + Pdf / LobePdf);
-		else if (LobePdf < Pdf)
-			LocalMISWeight = 1 - 1 / (1 + LobePdf / Pdf);
-		else
-			LocalMISWeight = 0.5f; // avoid (rare) inf/inf
-
-		Weight = lerp(Weight, LobeWeight, LocalMISWeight);
-		Pdf += LobePdf;
-	}
-}
-
 // Select a term based based on a random value and a PDF, and returns the selected term
 // The random value is rescale/updated to reusable after evaluation
 #define SubstrateSelectLobe_Type(TDataType) TDataType SubstrateSelectLobe(inout float InE, float InPDF, TDataType InA, TDataType InB)\
@@ -1306,15 +1279,15 @@ FSubstrateImportanceSampleResult SubstrateImportanceSampleBSDF(FSubstrateBSDFCon
 			Out.Weight = 0;
 			if (InTermMask & SHADING_TERM_DIFFUSE)
 			{
-				SubstrateAddLobeWithMIS(Out.Weight, Out.PDF, Evaluate.DiffusePathValue / Evaluate.DiffusePDF, Evaluate.DiffusePDF, TermPDF);
+				AddLobeWithMIS(Out.Weight, Out.PDF, Evaluate.DiffusePathValue / Evaluate.DiffusePDF, Evaluate.DiffusePDF, TermPDF);
 			}
 
 			if (InTermMask & SHADING_TERM_SPECULAR)
 			{
-				SubstrateAddLobeWithMIS(Out.Weight, Out.PDF, Evaluate.SpecularPathValue / Evaluate.SpecularPDF, Evaluate.SpecularPDF, (1-TermPDF) * SpecularLobePDF);
+				AddLobeWithMIS(Out.Weight, Out.PDF, Evaluate.SpecularPathValue / Evaluate.SpecularPDF, Evaluate.SpecularPDF, (1-TermPDF) * SpecularLobePDF);
 				if (bHaziness)
 				{
-					SubstrateAddLobeWithMIS(Out.Weight, Out.PDF, Evaluate.SpecularHazePathValue / Evaluate.SpecularHazePDF, Evaluate.SpecularPDF, (1 - TermPDF) * (1-SpecularLobePDF));
+					AddLobeWithMIS(Out.Weight, Out.PDF, Evaluate.SpecularHazePathValue / Evaluate.SpecularHazePDF, Evaluate.SpecularPDF, (1 - TermPDF) * (1-SpecularLobePDF));
 				}
 			}
 
@@ -1401,12 +1374,12 @@ FSubstrateImportanceSampleResult SubstrateImportanceSampleBSDF(FSubstrateBSDFCon
 			Out.Weight = 0;
 			if (InTermMask & SHADING_TERM_DIFFUSE)
 			{
-				SubstrateAddLobeWithMIS(Out.Weight, Out.PDF, Evaluate.DiffusePathValue / Evaluate.DiffusePDF, Evaluate.DiffusePDF, TermPDF);
+				AddLobeWithMIS(Out.Weight, Out.PDF, Evaluate.DiffusePathValue / Evaluate.DiffusePDF, Evaluate.DiffusePDF, TermPDF);
 			}
 
 			if (InTermMask & SHADING_TERM_SPECULAR)
 			{
-				SubstrateAddLobeWithMIS(Out.Weight, Out.PDF, Evaluate.SpecularPathValue / Evaluate.SpecularPDF, Evaluate.SpecularPDF, (1-TermPDF));
+				AddLobeWithMIS(Out.Weight, Out.PDF, Evaluate.SpecularPathValue / Evaluate.SpecularPDF, Evaluate.SpecularPDF, (1-TermPDF));
 			}
 			break;
 		}
