Add DFFastMultiplyDemote and DFHash

diff --git a/Engine/Shaders/Private/DoubleFloat.ush b/Engine/Shaders/Private/DoubleFloat.ush
--- a/Engine/Shaders/Private/DoubleFloat.ush
+++ b/Engine/Shaders/Private/DoubleFloat.ush
@@ -304,12 +304,25 @@ float3 DFMultiplyDemote(FDFVector3 Vector, FDFInverseMatrix InMatrix)
 	return Result.xyz;
 }
 
+float3 DFFastMultiplyDemote(FDFVector3 Vector, FDFInverseMatrix InMatrix)
+{
+	float3 V = DFFastLocalSubtractDemote(Vector, InMatrix.PreTranslation);
+	float4 Result = mul(float4(V, 1), InMatrix.M);
+	return Result.xyz;
+}
+
 float4 DFMultiplyDemote(FDFVector4 Vector, FDFInverseMatrix InMatrix)
 {
 	float4 V = DFFastSubtractDemote(Vector, float4(InMatrix.PreTranslation, 0));
 	return mul(V, InMatrix.M);
 }
 
+float4 DFFastMultiplyDemote(FDFVector4 Vector, FDFInverseMatrix InMatrix)
+{
+	float4 V = DFFastLocalSubtractDemote(Vector, float4(InMatrix.PreTranslation, 0));
+	return mul(V, InMatrix.M);
+}
+
 FDFInverseMatrix DFMultiply(FDFInverseMatrix Lhs, float4x4 Rhs)
 {
 	float4x4 M = mul(Lhs.M, Rhs);
@@ -554,3 +567,75 @@ FDFInverseMatrix WaveReadLaneAt(FDFInverseMatrix In, uint SrcIndex)
 	return Result;
 }
 #endif
+
+// Include Random.ush
+uint3 Rand3DPCG32(int3);
+uint4 Rand4DPCG32(int4);
+
+uint4 DFHashPCG(FDFVector2 In)
+{
+	int4 V1 = asint(float4(In.High.x, In.High.y, In.Low.x, In.Low.y));
+	return Rand4DPCG32(V1);
+}
+
+uint DFHash_LCG(FDFVector2 In)
+{
+	const uint M = 1664525u;
+	const uint A = 1013904223u;
+	uint V1 = asuint(In.High.x) + asuint(In.High.y)
+			+ asuint(In.Low.x)  + asuint(In.Low.y);
+	uint V2 = V1 * M + A;
+	// second round improves visual quality significantly, even though the compiler premultiplies these values, so no extra ops
+	uint V3 = V2 * M + A;
+	return V3;
+}
+
+// High-quality hash, but more expensive.
+// If you need a statistically robust hash, choose this.
+// For noise that looks visually plausible, choose LCG.
+uint3 DFHashPCG(FDFVector3 In)
+{
+	int3 V1 = asint(In.High);
+	int3 V2 = asint(In.Low);
+
+	uint3 H1 = Rand3DPCG32(V1);
+	uint3 H2 = Rand3DPCG32(H1 + V2);
+
+	return H2;
+}
+
+// Low-quality hash, but instruction count is roughly 1/8th of PCG approach.
+uint DFHash_LCG(FDFVector3 In)
+{
+	const uint M = 1664525u;
+	const uint A = 1013904223u;
+	uint V1 = asuint(In.High.x) + asuint(In.High.y) + asuint(In.High.z) 
+			+ asuint(In.Low.x)  + asuint(In.Low.y)  + asuint(In.Low.z);
+	uint V2 = V1 * M + A;
+	// second round improves visual quality significantly, even though the compiler premultiplies these values, so no extra ops
+	uint V3 = V2 * M + A;
+	return V3;
+}
+
+uint4 DFHash_PCG(FDFVector4 In)
+{
+	int4 V1 = asint(In.High);
+	int4 V2 = asint(In.Low);
+
+	uint4 H1 = Rand4DPCG32(V1);
+	uint4 H2 = Rand4DPCG32(H1 + V2);
+
+	return H2;
+}
+
+uint DFHash_LCG(FDFVector4 In)
+{
+	const uint M = 1664525u;
+	const uint A = 1013904223u;
+	uint V1 = asuint(In.High.x) + asuint(In.High.y) + asuint(In.High.z) + asuint(In.High.w)
+			+ asuint(In.Low.x)  + asuint(In.Low.y)  + asuint(In.Low.z)  + asuint(In.Low.w);
+	uint V2 = V1 * M + A;
+	// second round improves visual quality significantly, even though the compiler premultiplies these values, so no extra ops
+	uint V3 = V2 * M + A;
+	return V3;
+}
\ No newline at end of file
