Fixed incorrect results for DFToTileOffset certain negative ranges of values.
Added non-fast-math flag path for DFFastToTileOffset that was missing.

diff --git a/Engine/Shaders/Private/DoubleFloat.ush b/Engine/Shaders/Private/DoubleFloat.ush
--- a/Engine/Shaders/Private/DoubleFloat.ush
+++ b/Engine/Shaders/Private/DoubleFloat.ush
@@ -493,7 +493,7 @@ float4 DFTransformLocalToTranslatedWorld(float3 LocalPosition, FDFMatrix LocalTo
 	float3 RotatedScaledPosition = FMA(LocalPosition.xxx, LocalToWorld.M[0].xyz, FMA(LocalPosition.yyy, LocalToWorld.M[1].xyz, LocalPosition.zzz * LocalToWorld.M[2].xyz));
 
 #if UE_DF_NO_FAST_MATH
-	float3 TranslatedWorldPositionOrigin = DFAddDemote(DFGetOrigin(LocalToWorld), ResolvedView.PreViewTranslation);
+	float3 TranslatedWorldPositionOrigin = DFAddDemote(DFGetOrigin(LocalToWorld), PreViewTranslation);
 #else
 	float3 TranslatedWorldPositionOrigin = DFFastLocalSubtractDemote(LocalToWorld.PostTranslation, DFNegate(PreViewTranslation)) + LocalToWorld.M[3].xyz;
 #endif
diff --git a/Engine/Shaders/Private/LargeWorldCoordinates.ush b/Engine/Shaders/Private/LargeWorldCoordinates.ush
--- a/Engine/Shaders/Private/LargeWorldCoordinates.ush
+++ b/Engine/Shaders/Private/LargeWorldCoordinates.ush
@@ -516,14 +516,24 @@ float4x4 LWCHackToFloat(FLWCInverseMatrix v) { return (float4x4)0.0f; }
 #define DEFINE_DF_TO_TILEOFFSET_OPERATOR(FDFType, FLWCType, LWCConstructor, FFloatType) \
 	FLWCType DFToTileOffset(FDFType In)																							\
 	{																															\
-		/* Used rounding, removes fractional data stored in Tile, so the resulting coordinate is 'normalized' */				\
-		/* T+C-C benchmarked slightly faster than round(), can also shift C right to retain some fractional bits */ 			\
-		const float C = 8388608; /* = 2^23 (23 mantisse bits for int part, 0 for fraction) */ 									\
-		FFloatType Tile = INVARIANT(FMA(In.High, 1 / UE_LWC_RENDER_TILE_SIZE, C) - C);											\
+		/* CPU side defines MakeTile as floor(Value / TILE_SIZE + 0.5) which is not quite the same as round but */				\
+		/* divergent cases are extremely rare and result simply in rollover from one tile edge to the edge of the next tile */	\
+		/* Values greater than 2^(24+floor(log2(TILE_SIZE-1))) are not converted correctly, because In.Low is assumed to be */	\
+		/* less than half a tile. These values by definition exceed WORLD_MAX, so this is not an issue. */						\
+		FFloatType Tile = INVARIANT(round(In.High / UE_LWC_RENDER_TILE_SIZE));													\
 		FFloatType Offset = INVARIANT(FMA(Tile, -UE_LWC_RENDER_TILE_SIZE, In.High) + In.Low);									\
 																																\
 		return LWCConstructor(Tile, Offset);																					\
-	}																															\
+	}	
+	
+#if UE_DF_NO_FAST_MATH
+#define DEFINE_DF_FAST_TO_TILEOFFSET_OPERATOR(FDFType, FLWCType, LWCConstructor, FFloatType) \
+	FLWCType DFFastToTileOffset(FDFType In)																						\
+	{																															\
+		return DFToTileOffset(In);																								\
+	}
+#else
+#define DEFINE_DF_FAST_TO_TILEOFFSET_OPERATOR(FDFType, FLWCType, LWCConstructor, FFloatType) \
 	FLWCType DFFastToTileOffset(FDFType In)																						\
 	{																															\
 		/* Produces a non-normalized tile-offset coordinate. */																	\
@@ -534,12 +544,18 @@ float4x4 LWCHackToFloat(FLWCInverseMatrix v) { return (float4x4)0.0f; }
 																																\
 		return LWCConstructor(Tile, Offset);																					\
 	}
+#endif
 
 DEFINE_DF_TO_TILEOFFSET_OPERATOR(FDFScalar, FLWCScalar, MakeLWCScalar, float)
 DEFINE_DF_TO_TILEOFFSET_OPERATOR(FDFVector2, FLWCVector2, MakeLWCVector2, float2)
 DEFINE_DF_TO_TILEOFFSET_OPERATOR(FDFVector3, FLWCVector3, MakeLWCVector3, float3)
 DEFINE_DF_TO_TILEOFFSET_OPERATOR(FDFVector4, FLWCVector4, MakeLWCVector4, float4)
+DEFINE_DF_FAST_TO_TILEOFFSET_OPERATOR(FDFScalar, FLWCScalar, MakeLWCScalar, float)
+DEFINE_DF_FAST_TO_TILEOFFSET_OPERATOR(FDFVector2, FLWCVector2, MakeLWCVector2, float2)
+DEFINE_DF_FAST_TO_TILEOFFSET_OPERATOR(FDFVector3, FLWCVector3, MakeLWCVector3, float3)
+DEFINE_DF_FAST_TO_TILEOFFSET_OPERATOR(FDFVector4, FLWCVector4, MakeLWCVector4, float4)
 #undef DEFINE_DF_TO_TILEOFFSET_OPERATOR
+#undef DEFINE_DF_FAST_TO_TILEOFFSET_OPERATOR
 
 FLWCMatrix DFToTileOffset(FDFMatrix In)
 {
@@ -552,9 +568,13 @@ FLWCMatrix DFToTileOffset(FDFMatrix In)
 
 FLWCMatrix DFFastToTileOffset(FDFMatrix In)
 {
+#if UE_DF_NO_FAST_MATH
+	return DFToTileOffset(In);
+#else
 	float4x4 M = In.M;
 	float3 Tile = In.PostTranslation / UE_LWC_RENDER_TILE_SIZE;
 	return MakeLWCMatrix(Tile, M);
+#endif
 }
 
 FLWCInverseMatrix DFToTileOffset(FDFInverseMatrix In)
@@ -567,9 +587,13 @@ FLWCInverseMatrix DFToTileOffset(FDFInverseMatrix In)
 
 FLWCInverseMatrix DFFastToTileOffset(FDFInverseMatrix In)
 {
+#if UE_DF_NO_FAST_MATH
+	return DFToTileOffset(In);
+#else
 	float4x4 M = In.M;
 	float3 Tile = In.PreTranslation / UE_LWC_RENDER_TILE_SIZE;
 	return MakeLWCInverseMatrix(Tile, M);
+#endif
 }
 
 #define DFToTileOffset_Hack DFToTileOffset
