Add DFFastMultiplyDemote and DFHash. Resubmit of CL 31853969 with forward declare removed. Instead, the relevant PCG functions are moved to a smaller header.

diff --git a/Engine/Shaders/Private/DoubleFloat.ush b/Engine/Shaders/Private/DoubleFloat.ush
--- a/Engine/Shaders/Private/DoubleFloat.ush
+++ b/Engine/Shaders/Private/DoubleFloat.ush
@@ -23,6 +23,8 @@
 #define UE_DF_NO_FAST_MATH 0
 #endif
 
+#include "RandomPCG.ush"
+
 /* FDFScalar */
 
 #define FDFType FDFScalar
@@ -304,12 +306,25 @@ float3 DFMultiplyDemote(FDFVector3 Vector, FDFInverseMatrix InMatrix)
 	return Result.xyz;
 }
 
+float3 DFFastMultiplyDemote(FDFVector3 Vector, FDFInverseMatrix InMatrix)
+{
+	float3 V = DFFastLocalSubtractDemote(Vector, InMatrix.PreTranslation);
+	float4 Result = mul(float4(V, 1), InMatrix.M);
+	return Result.xyz;
+}
+
 float4 DFMultiplyDemote(FDFVector4 Vector, FDFInverseMatrix InMatrix)
 {
 	float4 V = DFFastSubtractDemote(Vector, float4(InMatrix.PreTranslation, 0));
 	return mul(V, InMatrix.M);
 }
 
+float4 DFFastMultiplyDemote(FDFVector4 Vector, FDFInverseMatrix InMatrix)
+{
+	float4 V = DFFastLocalSubtractDemote(Vector, float4(InMatrix.PreTranslation, 0));
+	return mul(V, InMatrix.M);
+}
+
 FDFInverseMatrix DFMultiply(FDFInverseMatrix Lhs, float4x4 Rhs)
 {
 	float4x4 M = mul(Lhs.M, Rhs);
@@ -554,3 +569,71 @@ FDFInverseMatrix WaveReadLaneAt(FDFInverseMatrix In, uint SrcIndex)
 	return Result;
 }
 #endif
+
+uint4 DFHashPCG(FDFVector2 In)
+{
+	int4 V1 = asint(float4(In.High.x, In.High.y, In.Low.x, In.Low.y));
+	return Rand4DPCG32(V1);
+}
+
+uint DFHash_LCG(FDFVector2 In)
+{
+	const uint M = 1664525u;
+	const uint A = 1013904223u;
+	uint V1 = asuint(In.High.x) + asuint(In.High.y)
+			+ asuint(In.Low.x)  + asuint(In.Low.y);
+	uint V2 = V1 * M + A;
+	// second round improves visual quality significantly, even though the compiler premultiplies these values, so no extra ops
+	uint V3 = V2 * M + A;
+	return V3;
+}
+
+// High-quality hash, but more expensive.
+// If you need a statistically robust hash, choose this.
+// For noise that looks visually plausible, choose LCG.
+uint3 DFHashPCG(FDFVector3 In)
+{
+	int3 V1 = asint(In.High);
+	int3 V2 = asint(In.Low);
+
+	uint3 H1 = Rand3DPCG32(V1);
+	uint3 H2 = Rand3DPCG32(H1 + V2);
+
+	return H2;
+}
+
+// Low-quality hash, but instruction count is roughly 1/8th of PCG approach.
+uint DFHash_LCG(FDFVector3 In)
+{
+	const uint M = 1664525u;
+	const uint A = 1013904223u;
+	uint V1 = asuint(In.High.x) + asuint(In.High.y) + asuint(In.High.z) 
+			+ asuint(In.Low.x)  + asuint(In.Low.y)  + asuint(In.Low.z);
+	uint V2 = V1 * M + A;
+	// second round improves visual quality significantly, even though the compiler premultiplies these values, so no extra ops
+	uint V3 = V2 * M + A;
+	return V3;
+}
+
+uint4 DFHash_PCG(FDFVector4 In)
+{
+	int4 V1 = asint(In.High);
+	int4 V2 = asint(In.Low);
+
+	uint4 H1 = Rand4DPCG32(V1);
+	uint4 H2 = Rand4DPCG32(H1 + V2);
+
+	return H2;
+}
+
+uint DFHash_LCG(FDFVector4 In)
+{
+	const uint M = 1664525u;
+	const uint A = 1013904223u;
+	uint V1 = asuint(In.High.x) + asuint(In.High.y) + asuint(In.High.z) + asuint(In.High.w)
+			+ asuint(In.Low.x)  + asuint(In.Low.y)  + asuint(In.Low.z)  + asuint(In.Low.w);
+	uint V2 = V1 * M + A;
+	// second round improves visual quality significantly, even though the compiler premultiplies these values, so no extra ops
+	uint V3 = V2 * M + A;
+	return V3;
+}
\ No newline at end of file
diff --git a/Engine/Shaders/Private/Random.ush b/Engine/Shaders/Private/Random.ush
--- a/Engine/Shaders/Private/Random.ush
+++ b/Engine/Shaders/Private/Random.ush
@@ -1,11 +1,15 @@
 // Copyright Epic Games, Inc. All Rights Reserved.
 
 /*=============================================================================
-	Random.ush: A pseudo-random number generator.
+	Random.ush: Pseudo-random number generators.
 =============================================================================*/
 
 #pragma once
 
+// Most frequently used functions in separate smaller headers
+#include "RandomInterleavedGradientNoise.ush"
+#include "RandomPCG.ush"
+
 // @param xy should be a integer position (e.g. pixel position on the screen), repeats each 128x128 pixels
 // similar to a texture lookup but is only ALU
 // ~13 ALU operations (3 frac, 6 *, 4 mad)
@@ -17,21 +21,6 @@ float PseudoRandom(float2 xy)
 	return frac(dot(pos.xyx * pos.xyy, float3(20.390625f, 60.703125f, 2.4281209f)));
 }
 
-// high frequency dither pattern appearing almost random without banding steps
-//note: from "NEXT GENERATION POST PROCESSING IN CALL OF DUTY: ADVANCED WARFARE"
-//      http://advances.realtimerendering.com/s2014/index.html
-// Epic extended by FrameId
-// ~7 ALU operations (2 frac, 3 mad, 2 *)
-// @return 0..1
-float InterleavedGradientNoise( float2 uv, float FrameId )
-{
-	// magic values are found by experimentation
-	uv += FrameId * (float2(47, 17) * 0.695f);
-
-    const float3 magic = float3( 0.06711056f, 0.00583715f, 52.9829189f );
-    return frac(magic.z * frac(dot(uv, magic.xy)));
-}
-
 // [0, 1[
 // ~10 ALU operations (2 frac, 5 *, 3 mad)
 float RandFast( uint2 PixelPos, float Magic = 3571.0 )
@@ -61,112 +50,6 @@ float RandBBSfloat(float seed)
 	return s;
 }
 
-// 3D random number generator inspired by PCGs (permuted congruential generator)
-// Using a **simple** Feistel cipher in place of the usual xor shift permutation step
-// @param v = 3D integer coordinate
-// @return three elements w/ 16 random bits each (0-0xffff).
-// ~8 ALU operations for result.x    (7 mad, 1 >>)
-// ~10 ALU operations for result.xy  (8 mad, 2 >>)
-// ~12 ALU operations for result.xyz (9 mad, 3 >>)
-uint3 Rand3DPCG16(int3 p)
-{
-	// taking a signed int then reinterpreting as unsigned gives good behavior for negatives
-	uint3 v = uint3(p);
-
-	// Linear congruential step. These LCG constants are from Numerical Recipies
-	// For additional #'s, PCG would do multiple LCG steps and scramble each on output
-	// So v here is the RNG state
-	v = v * 1664525u + 1013904223u;
-
-	// PCG uses xorshift for the final shuffle, but it is expensive (and cheap
-	// versions of xorshift have visible artifacts). Instead, use simple MAD Feistel steps
-	//
-	// Feistel ciphers divide the state into separate parts (usually by bits)
-	// then apply a series of permutation steps one part at a time. The permutations
-	// use a reversible operation (usually ^) to part being updated with the result of
-	// a permutation function on the other parts and the key.
-	//
-	// In this case, I'm using v.x, v.y and v.z as the parts, using + instead of ^ for
-	// the combination function, and just multiplying the other two parts (no key) for 
-	// the permutation function.
-	//
-	// That gives a simple mad per round.
-	v.x += v.y*v.z;
-	v.y += v.z*v.x;
-	v.z += v.x*v.y;
-	v.x += v.y*v.z;
-	v.y += v.z*v.x;
-	v.z += v.x*v.y;
-
-	// only top 16 bits are well shuffled
-	return v >> 16u;
-}
-
-// 3D random number generator inspired by PCGs (permuted congruential generator)
-// Using a **simple** Feistel cipher in place of the usual xor shift permutation step
-// http://jcgt.org/published/0009/03/02/
-// @param v = 3D integer coordinate
-// @return three elements w/ 32 random bits each (0-0xffffffff).
-uint3 Rand3DPCG32(int3 p)
-{
-	// taking a signed int then reinterpreting as unsigned gives good behavior for negatives
-	uint3 v = uint3(p);
-
-	// Linear congruential step.
-	v = v * 1664525u + 1013904223u;
-
-	// shuffle
-	v.x += v.y*v.z;
-	v.y += v.z*v.x;
-	v.z += v.x*v.y;
-
-	// xoring high bits into low bits makes all 32 bits pretty good
-	v ^= v >> 16u;
-
-	// final shuffle
-	v.x += v.y*v.z;
-	v.y += v.z*v.x;
-	v.z += v.x*v.y;
-
-	return v;
-}
-
-
-
-// 4D random number generator inspired by PCGs (permuted congruential generator)
-// Using a **simple** Feistel cipher in place of the usual xor shift permutation step
-// http://jcgt.org/published/0009/03/02/
-// @param v = 4D integer coordinate
-// @return four elements w/ 32 random bits each (0-0xffffffff).
-uint4 Rand4DPCG32(int4 p)
-{
-	// taking a signed int then reinterpreting as unsigned gives good behavior for negatives
-	uint4 v = uint4(p);
-
-	// Linear congruential step.
-	v = v * 1664525u + 1013904223u;
-
-	// shuffle
-	v.x += v.y*v.w;
-	v.y += v.z*v.x;
-	v.z += v.x*v.y;
-	v.w += v.y*v.z;
-
-	// xoring high bits into low makes all 32 bits pretty good
-	v ^= (v >> 16u);
-
-	// final shuffle
-	v.x += v.y*v.w;
-	v.y += v.z*v.x;
-	v.z += v.x*v.y;
-	v.w += v.y*v.z;
-
-	return v;
-}
-
-
-
-
 /**
  * Find good arbitrary axis vectors to represent U and V axes of a plane,
  * given just the normal. Ported from UnMath.h
diff --git a/Engine/Shaders/Private/RandomInterleavedGradientNoise.ush b/Engine/Shaders/Private/RandomInterleavedGradientNoise.ush
new file mode 100644
--- /dev/null
+++ b/Engine/Shaders/Private/RandomInterleavedGradientNoise.ush
@@ -0,0 +1,18 @@
+// Copyright Epic Games, Inc. All Rights Reserved.
+
+#pragma once
+
+// high frequency dither pattern appearing almost random without banding steps
+//note: from "NEXT GENERATION POST PROCESSING IN CALL OF DUTY: ADVANCED WARFARE"
+//      http://advances.realtimerendering.com/s2014/index.html
+// Epic extended by FrameId
+// ~7 ALU operations (2 frac, 3 mad, 2 *)
+// @return 0..1
+float InterleavedGradientNoise( float2 uv, float FrameId )
+{
+	// magic values are found by experimentation
+	uv += FrameId * (float2(47, 17) * 0.695f);
+
+    const float3 magic = float3( 0.06711056f, 0.00583715f, 52.9829189f );
+    return frac(magic.z * frac(dot(uv, magic.xy)));
+}
diff --git a/Engine/Shaders/Private/RandomPCG.ush b/Engine/Shaders/Private/RandomPCG.ush
new file mode 100644
--- /dev/null
+++ b/Engine/Shaders/Private/RandomPCG.ush
@@ -0,0 +1,104 @@
+// Copyright Epic Games, Inc. All Rights Reserved.
+
+#pragma once
+
+// 3D random number generator inspired by PCGs (permuted congruential generator)
+// Using a **simple** Feistel cipher in place of the usual xor shift permutation step
+// @param v = 3D integer coordinate
+// @return three elements w/ 16 random bits each (0-0xffff).
+// ~8 ALU operations for result.x    (7 mad, 1 >>)
+// ~10 ALU operations for result.xy  (8 mad, 2 >>)
+// ~12 ALU operations for result.xyz (9 mad, 3 >>)
+uint3 Rand3DPCG16(int3 p)
+{
+	// taking a signed int then reinterpreting as unsigned gives good behavior for negatives
+	uint3 v = uint3(p);
+
+	// Linear congruential step. These LCG constants are from Numerical Recipies
+	// For additional #'s, PCG would do multiple LCG steps and scramble each on output
+	// So v here is the RNG state
+	v = v * 1664525u + 1013904223u;
+
+	// PCG uses xorshift for the final shuffle, but it is expensive (and cheap
+	// versions of xorshift have visible artifacts). Instead, use simple MAD Feistel steps
+	//
+	// Feistel ciphers divide the state into separate parts (usually by bits)
+	// then apply a series of permutation steps one part at a time. The permutations
+	// use a reversible operation (usually ^) to part being updated with the result of
+	// a permutation function on the other parts and the key.
+	//
+	// In this case, I'm using v.x, v.y and v.z as the parts, using + instead of ^ for
+	// the combination function, and just multiplying the other two parts (no key) for 
+	// the permutation function.
+	//
+	// That gives a simple mad per round.
+	v.x += v.y*v.z;
+	v.y += v.z*v.x;
+	v.z += v.x*v.y;
+	v.x += v.y*v.z;
+	v.y += v.z*v.x;
+	v.z += v.x*v.y;
+
+	// only top 16 bits are well shuffled
+	return v >> 16u;
+}
+
+// 3D random number generator inspired by PCGs (permuted congruential generator)
+// Using a **simple** Feistel cipher in place of the usual xor shift permutation step
+// http://jcgt.org/published/0009/03/02/
+// @param v = 3D integer coordinate
+// @return three elements w/ 32 random bits each (0-0xffffffff).
+uint3 Rand3DPCG32(int3 p)
+{
+	// taking a signed int then reinterpreting as unsigned gives good behavior for negatives
+	uint3 v = uint3(p);
+
+	// Linear congruential step.
+	v = v * 1664525u + 1013904223u;
+
+	// shuffle
+	v.x += v.y*v.z;
+	v.y += v.z*v.x;
+	v.z += v.x*v.y;
+
+	// xoring high bits into low bits makes all 32 bits pretty good
+	v ^= v >> 16u;
+
+	// final shuffle
+	v.x += v.y*v.z;
+	v.y += v.z*v.x;
+	v.z += v.x*v.y;
+
+	return v;
+}
+
+// 4D random number generator inspired by PCGs (permuted congruential generator)
+// Using a **simple** Feistel cipher in place of the usual xor shift permutation step
+// http://jcgt.org/published/0009/03/02/
+// @param v = 4D integer coordinate
+// @return four elements w/ 32 random bits each (0-0xffffffff).
+uint4 Rand4DPCG32(int4 p)
+{
+	// taking a signed int then reinterpreting as unsigned gives good behavior for negatives
+	uint4 v = uint4(p);
+
+	// Linear congruential step.
+	v = v * 1664525u + 1013904223u;
+
+	// shuffle
+	v.x += v.y*v.w;
+	v.y += v.z*v.x;
+	v.z += v.x*v.y;
+	v.w += v.y*v.z;
+
+	// xoring high bits into low makes all 32 bits pretty good
+	v ^= (v >> 16u);
+
+	// final shuffle
+	v.x += v.y*v.w;
+	v.y += v.z*v.x;
+	v.z += v.x*v.y;
+	v.w += v.y*v.z;
+
+	return v;
+}
