Holdout + RenderInMain for skyatmosphere, fog and volumetric cloud.

diff --git a/Engine/Shaders/Private/MaterialTemplate.ush b/Engine/Shaders/Private/MaterialTemplate.ush
--- a/Engine/Shaders/Private/MaterialTemplate.ush
+++ b/Engine/Shaders/Private/MaterialTemplate.ush
@@ -42,7 +42,7 @@
 
 // Update this GUID to force all material shaders to recompile
 // Merge conflicts on this line should be resolved by generating a new GUID
-#pragma message("UESHADERMETADATA_VERSION 50476854-A3D1-477E-91EA-0ED4AF89E199")
+#pragma message("UESHADERMETADATA_VERSION 381E7C4C-6229-4F2E-8F0E-EE692C2785A3")
 
 #if USES_SPEEDTREE
 	#include "/Engine/Private/SpeedTreeCommon.ush"
