Update reflection vector in of mobile renderer when dbuffer is used

diff --git a/Engine/Shaders/Private/MobileBasePassPixelShader.usf b/Engine/Shaders/Private/MobileBasePassPixelShader.usf
--- a/Engine/Shaders/Private/MobileBasePassPixelShader.usf
+++ b/Engine/Shaders/Private/MobileBasePassPixelShader.usf
@@ -485,6 +485,8 @@ void Main(
 			ApplyDBufferData(DBufferData, MaterialParameters.WorldNormal, SubsurfaceColor, Roughness, BaseColor, Metallic, Specular);
 			DBufferOpacity = (DBufferData.ColorOpacity + DBufferData.NormalOpacity + DBufferData.RoughnessOpacity) * (1.0f / 3.0f);
 		#endif // SUBSTRATE_ENABLED
+			
+			MaterialParameters.ReflectionVector = ReflectionAboutCustomWorldNormal(MaterialParameters, MaterialParameters.WorldNormal, false);
 		}
 	}
 #endif
