Fix sky atmosphere RenderInMain for high quality rendering using tracing.

diff --git a/Engine/Shaders/Private/SkyAtmosphere.usf b/Engine/Shaders/Private/SkyAtmosphere.usf
--- a/Engine/Shaders/Private/SkyAtmosphere.usf
+++ b/Engine/Shaders/Private/SkyAtmosphere.usf
@@ -1029,6 +1029,11 @@ void RenderSkyAtmosphereRayMarchingPS(
 		AerialPespectiveViewDistanceScale);
 
 	PreExposedL += ss.L * LuminanceScale;
+	
+	if (View.RenderingReflectionCaptureMask == 0.0f && !IsSkyAtmosphereRenderedInMain(View.EnvironmentComponentsFlags))
+	{
+		PreExposedL = 0.0f;
+	}
 
 #if SAMPLE_ATMOSPHERE_ON_CLOUDS
 	// We use gray scale transmittance to match the rendering when applying the AerialPerspective texture
