Holdouts: Ensure masked materials always output solid alpha. Some shading models using Opacity for other purposes could cause incorrect output

diff --git a/Engine/Shaders/Private/BasePassPixelShader.usf b/Engine/Shaders/Private/BasePassPixelShader.usf
--- a/Engine/Shaders/Private/BasePassPixelShader.usf
+++ b/Engine/Shaders/Private/BasePassPixelShader.usf
@@ -2153,7 +2153,7 @@ void FPixelShaderInOut_MainPS(
 	{
 		// Discard this pixel so that nothing gets written to the G-Buffers and alpha is kept as its default value (background)
 
-#if MATERIALBLENDING_SOLID
+#if MATERIALBLENDING_SOLID || MATERIALBLENDING_MASKED
 		// This catches the common case as well as handles cases where Opacity has a different meaning, such as for Subsurface materials
 		Out.MRT[0] = half4(0.0, 0.0, 0.0, 1.0);
 #else
