Path Tracer: Fix artifacts from solid angle sampling of small rectangles

Improve the math for solid angle rectangles by simplifying cross products (fewer terms are needed), using asin() instead of acos() and replacing the calls to the builtin asin() with a faster custom approximation that still has good accuracy near the end-points (unlike asinFast).

To further improve results, fallback to area sampling when the solid angle of the rectangle drops below a small threshold.

This also improves the capsule sampling case, which samples from a bounding rectangle that can be thin when the radius is small.

diff --git a/Engine/Shaders/Private/CapsuleLightIntegrate.ush b/Engine/Shaders/Private/CapsuleLightIntegrate.ush
--- a/Engine/Shaders/Private/CapsuleLightIntegrate.ush
+++ b/Engine/Shaders/Private/CapsuleLightIntegrate.ush
@@ -257,7 +257,7 @@ FDirectLighting IntegrateBxDF( FGBufferData GBuffer, half3 N, half3 V, FCapsuleL
 				{
 					NoL / PI,
 					D_GGX(Pow4(GBuffer.Roughness), NoH) * NoH / (4 * VoH),
-					1.0 / GetCapsuleBoundsSolidAngle(CapsuleBounds)
+					rcp(GetCapsuleBoundsInversePdf(L, CapsuleBounds))
 				};
 
 				if( Capsule.Length == 0 )
diff --git a/Engine/Shaders/Private/CapsuleLightSampling.ush b/Engine/Shaders/Private/CapsuleLightSampling.ush
--- a/Engine/Shaders/Private/CapsuleLightSampling.ush
+++ b/Engine/Shaders/Private/CapsuleLightSampling.ush
@@ -19,10 +19,15 @@ struct FCapsuleSphericalBounds
 	float ConeSolidAngle;
 };
 
-// return the smallest of the two possible bounding shapes
-float GetCapsuleBoundsSolidAngle(FCapsuleSphericalBounds Bounds)
+float GetCapsuleBoundsInversePdf(float3 Direction, FCapsuleSphericalBounds Bounds)
 {
-	return min(Bounds.SphericalRect.SolidAngle, Bounds.ConeSolidAngle);
+	if (Bounds.ConeSolidAngle < Bounds.SphericalRect.SolidAngle)
+	{
+		return Bounds.ConeSolidAngle;
+	}
+	float LocalDirZ = dot(Direction, Bounds.SphericalRect.Axis[2]);
+	float DistanceSquared = Square(Bounds.SphericalRect.z0 / LocalDirZ);
+	return GetSphericalRectInversePdf(Direction, DistanceSquared, Bounds.SphericalRect);
 }
 
 // return world space direction within the cone and solidangle
@@ -34,7 +39,8 @@ float4 SampleCapsuleBounds(FCapsuleSphericalBounds Bounds, float2 E)
 	}
 	else
 	{
-		return float4(normalize(UniformSampleSphericalRect(E, Bounds.SphericalRect)), Bounds.SphericalRect.SolidAngle);
+		FSphericalRectSample Result = UniformSampleSphericalRect(E, Bounds.SphericalRect);
+		return float4(Result.Direction, Result.InvPdf);
 	}
 }
 
diff --git a/Engine/Shaders/Private/PathTracing/Light/PathTracingCapsuleLight.ush b/Engine/Shaders/Private/PathTracing/Light/PathTracingCapsuleLight.ush
--- a/Engine/Shaders/Private/PathTracing/Light/PathTracingCapsuleLight.ush
+++ b/Engine/Shaders/Private/PathTracing/Light/PathTracingCapsuleLight.ush
@@ -108,7 +108,8 @@ FLightHit CapsuleLight_TraceLight(FRayDesc Ray, int LightId)
 		float3 LightDirection = TranslatedLightPosition - Ray.Origin;
 		float LightDistanceSquared = dot(LightDirection, LightDirection);
 #if CAPSULE_SOLIDANGLE_SAMPLING
-		float Pdf = 1.0 / GetCapsuleBoundsSolidAngle(CapsuleGetSphericalBounds(LightDirection, Axis, LightRadius, SourceLength));
+		FCapsuleSphericalBounds Bounds = CapsuleGetSphericalBounds(LightDirection, Axis, LightRadius, SourceLength);
+		float Pdf = 1.0 / GetCapsuleBoundsInversePdf(Ray.Direction, Bounds);
 #else
 		float3 Normal = CapsuleNormal(Ray.Origin + t * Ray.Direction, TranslatedLightPosition, Axis * SourceLength);
 
@@ -150,12 +151,12 @@ FLightSample CapsuleLight_SampleLight(
 	float3 LightRadiance = LightPower / (PI * (CapsArea + BodyArea));
 
 	float3 LightDirection = GetTranslatedPosition(LightId) - TranslatedWorldPos;
-	FCapsuleSphericalBounds CapsuleBounds = CapsuleGetSphericalBounds(LightDirection, Axis, Radius, SourceLength);
+	FCapsuleSphericalBounds Bounds = CapsuleGetSphericalBounds(LightDirection, Axis, Radius, SourceLength);
 
-	float4 Result = SampleCapsuleBounds(CapsuleBounds, RandSample);
+	float4 Result = SampleCapsuleBounds(Bounds, RandSample);
 
 	float3 L = Result.xyz;
-	float SolidAngle = Result.w;
+	float InvPdf = Result.w;
 		
 	// check direction to account for rays that hit the bounding shape but not the capsule
 #if 1
@@ -167,7 +168,7 @@ FLightSample CapsuleLight_SampleLight(
 #endif
 	if (Distance > 0)
 	{
-		return CreateLightSample(LightRadiance * SolidAngle, 1.0 / SolidAngle, L, Distance);
+		return CreateLightSample(LightRadiance * InvPdf, rcp(InvPdf), L, Distance);
 	}
 	// didn't pass the acceptance test -- reject this sample
 	return NullLightSample();
diff --git a/Engine/Shaders/Private/PathTracing/Light/PathTracingRectLight.ush b/Engine/Shaders/Private/PathTracing/Light/PathTracingRectLight.ush
--- a/Engine/Shaders/Private/PathTracing/Light/PathTracingRectLight.ush
+++ b/Engine/Shaders/Private/PathTracing/Light/PathTracingRectLight.ush
@@ -44,10 +44,10 @@ float3 BilinearQuadWarp(float2 uv, float W00, float W01, float W10, float W11) {
 	return float3(u, v, pdf);
 }
 
-// #dxr_todo: This method should be moved to RectLight.ush
-// #dxr_todo: Analyze the overall efficiency (MSE reduction vs. time)
+// TODO: This method should be moved to RectLight.ush
+// TODO: Analyze the overall efficiency (MSE reduction vs. time)
 //  Projected solid angle sampling converges faster at equal sample count, but is a bit slower
-// #dxr_todo: Should we switch to plain area sampling at some distance?
+// TODO: Unify with UniformSampleSphericalRect so both return the same struct, making them easier to interchange
 float4 SampleApproxProjectionSphericalRect(float2 Rand, FSphericalRect SphericalRect, float3 WorldNormal)
 {
 	// Construct the quad vertices
@@ -67,7 +67,7 @@ float4 SampleApproxProjectionSphericalRect(float2 Rand, FSphericalRect Spherical
 	float3 WarpedRand = BilinearQuadWarp(Rand, W00, W01, W10, W11);
 
 	// Proceed with uniform sampling of the spherical rectangle
-	float3 Direction = UniformSampleSphericalRect(WarpedRand.xy, SphericalRect);
+	float3 Direction = UniformSampleSphericalRect(WarpedRand.xy, SphericalRect).Direction;
 	float OutPdf = isfinite(SphericalRect.SolidAngle) ? 1.0 / SphericalRect.SolidAngle : 0.0;
 
 	// Return Direction vector and adjusted PDF
@@ -110,11 +110,12 @@ FLightHit RectLight_TraceLight(FRayDesc Ray, int LightId)
 				// stored color is radiance
 				float3 Radiance = GetColor(LightId) * ComputeIESAttenuation(LightId, Ray.Origin);
 				Radiance *= ComputeAttenuationFalloff(dot(LightDirection, LightDirection), LightId);
-				float SolidAngle = BuildSphericalRect(Rect).SolidAngle;
+				FSphericalRect SphericalRect = BuildSphericalRect(Rect);
 #if USE_RECT_LIGHT_TEXTURES
 				Radiance *= EvaluateTexture(LightId, (UV + LightExtent) / (LightExtent * 2));
 #endif
-				return CreateLightHit(Radiance, SolidAngle > 0 ? rcp(SolidAngle) : 0.0, t);
+				float Pdf = rcp(GetSphericalRectInversePdf(Ray.Direction, t * t, SphericalRect));
+				return CreateLightHit(Radiance, Pdf, t);
 			}
 		}
 	}
@@ -151,26 +152,15 @@ FLightSample RectLight_SampleLight(
 	}
 	FSphericalRect SphericalRect = BuildSphericalRect(Rect);
 
+	float3 Radiance = GetColor(LightId) * ComputeIESAttenuation(LightId, TranslatedWorldPos);
+	float DistanceSquared = length2(LightDirection);
+	Radiance *= ComputeAttenuationFalloff(DistanceSquared, LightId);
+	FSphericalRectSample Sample = UniformSampleSphericalRect(RandSample, SphericalRect);
 #if USE_RECT_LIGHT_TEXTURES
-	FSphericalRectSample Sample = UniformSampleSphericalRectWithUV(RandSample, SphericalRect);
-	float2 ClippedUV = Sample.UV;
-	float3 Direction = Sample.Direction;
-#else
-	float3 Direction = UniformSampleSphericalRect(RandSample, SphericalRect);
+	float2 UV = 0.5 * ((2 * Sample.UV - 1) * Rect.Extent + Rect.Offset) / Rect.FullExtent + 0.5;
+	Radiance *= EvaluateTexture(LightId, UV);
 #endif
-
-	float3 RadianceOverPdf = GetColor(LightId) * SphericalRect.SolidAngle * ComputeIESAttenuation(LightId, TranslatedWorldPos);
-	RadianceOverPdf *= ComputeAttenuationFalloff(dot(LightDirection, LightDirection), LightId);
-#if USE_RECT_LIGHT_TEXTURES
-	float2 UV = 0.5 * ((2 * ClippedUV - 1) * Rect.Extent + Rect.Offset) / Rect.FullExtent + 0.5;
-	RadianceOverPdf *= EvaluateTexture(LightId, UV);
-#endif
-
-	if (SphericalRect.SolidAngle > 0)
-	{
-		return CreateLightSample(RadianceOverPdf, 1.0 / SphericalRect.SolidAngle, normalize(Direction), length(Direction));
-	}
-	return NullLightSample();
+	return CreateLightSample(Radiance * Sample.InvPdf, rcp(Sample.InvPdf), Sample.Direction, Sample.Distance);
 }
 
 float RectLight_EstimateLight(
diff --git a/Engine/Shaders/Private/RayTracing/RayTracingRectLight.ush b/Engine/Shaders/Private/RayTracing/RayTracingRectLight.ush
--- a/Engine/Shaders/Private/RayTracing/RayTracingRectLight.ush
+++ b/Engine/Shaders/Private/RayTracing/RayTracingRectLight.ush
@@ -32,15 +32,15 @@ bool GenerateRectLightOcclusionRay(
 	RayPdf = 1.0 / SphericalRect.SolidAngle;
 
 	// 2. Uniformly sample from the spherical rectangle
-	float3 LightDirection = UniformSampleSphericalRect(RandSample, SphericalRect);
+	FSphericalRectSample Sample = UniformSampleSphericalRect(RandSample, SphericalRect);
 
 	// Light-normal culling
-	bool bIsValidRay = dot(-LightDirection, -LightParameters.Direction) >= 0.0;
+	bool bIsValidRay = dot(Sample.Direction, LightParameters.Direction) >= 0.0;
 	
 	// Apply normal perturbation when defining ray. normalize() resolves several banding artifacts
 	// at the cost of creating some bias under idea lighting condition but is neglectable.
-	RayDirection = normalize(LightDirection);
-	RayTMax = length(LightDirection);
+	RayDirection = Sample.Direction;
+	RayTMax = Sample.Distance;
 
 	return bIsValidRay;
 }
\ No newline at end of file
diff --git a/Engine/Shaders/Private/RectLight.ush b/Engine/Shaders/Private/RectLight.ush
--- a/Engine/Shaders/Private/RectLight.ush
+++ b/Engine/Shaders/Private/RectLight.ush
@@ -546,6 +546,27 @@ struct FSphericalRect
 	float		SolidAngle;
 };
 
+// The routine below needs more accuracy near 1.0 to accurately capture distant rectlights
+// The typical approach (for example see: asinFast) looses too much accuracy to cancelation near the endpoints
+float SphericalRectAsin(float x)
+{
+	const float HalfPI = PI / 2;
+	// Argument reduction to [0,0.5]
+	// using: asin(x) == pi/2-2*asin(sqrt(0.5-0.5*x))
+	float a = saturate(abs(x));
+	bool inner = a < 0.5f;
+	float a2 = inner ? a * a : 0.5 - 0.5 * a;
+	a = inner ? a : sqrt(a2);
+
+	// Fit using Mathematica -- max error for asin is 10^-4.87
+	float r = 0.100323f;
+	r = mad(r, a2, 0.163288f);
+	r = mad(r, a2, 1.00011f) * a;
+	r = inner ? r : HalfPI - 2 * r;
+
+	return asfloat(asuint(r) ^ (asuint(x) & 0x80000000u)); // propagate sign bit
+}
+
 // RectZ must face origin
 FSphericalRect BuildSphericalRect( FRect Rect )
 {
@@ -553,10 +574,7 @@ FSphericalRect BuildSphericalRect( FRect Rect )
 
 	SphericalRect.Axis = Rect.Axis;
 
-	float3 LocalPosition;
-	LocalPosition.x = dot( Rect.Axis[0], Rect.Origin );
-	LocalPosition.y = dot( Rect.Axis[1], Rect.Origin );
-	LocalPosition.z = dot( Rect.Axis[2], Rect.Origin );
+	float3 LocalPosition = mul(Rect.Axis, Rect.Origin);
 
 	SphericalRect.x0 = LocalPosition.x - Rect.Extent.x;
 	SphericalRect.x1 = LocalPosition.x + Rect.Extent.x;
@@ -566,86 +584,92 @@ FSphericalRect BuildSphericalRect( FRect Rect )
 
 	SphericalRect.Axis[2] *= LocalPosition.z > 0 ? -1 : 1;
 
-	float3 v0 = float3( SphericalRect.x0, SphericalRect.y0, SphericalRect.z0 );
-	float3 v1 = float3( SphericalRect.x1, SphericalRect.y0, SphericalRect.z0 );
-	float3 v2 = float3( SphericalRect.x1, SphericalRect.y1, SphericalRect.z0 );
-	float3 v3 = float3( SphericalRect.x0, SphericalRect.y1, SphericalRect.z0 );
-
-	float3 n0 = normalize( cross( v0, v1 ) );
-	float3 n1 = normalize( cross( v1, v2 ) );
-	float3 n2 = normalize( cross( v2, v3 ) );
-	float3 n3 = normalize( cross( v3, v0 ) );
-
-	float g0 = acos( -dot( n0, n1 ) );
-	float g1 = acos( -dot( n1, n2 ) );
-	float g2 = acos( -dot( n2, n3 ) );
-	float g3 = acos( -dot( n3, n0 ) );
-
-	SphericalRect.b0 = n0.z;
-	SphericalRect.b1 = n2.z;
-
-	SphericalRect.k = 2*PI - g2 - g3;
-	SphericalRect.SolidAngle = g0 + g1 - SphericalRect.k;
+	// Original paper creates all vertices, then all normals via cross products and finally computes angles via dot products.
+	// However since all the points lay on an axis-aligned plane, the math can be simplified by recognizing that most terms become
+	// zero, leaving only the product of the z components of the normals.
+	float z0sq = LocalPosition.z * LocalPosition.z;
+	float n0z = -SphericalRect.y0 * rsqrt(z0sq + SphericalRect.y0 * SphericalRect.y0);
+	float n1z =  SphericalRect.x1 * rsqrt(z0sq + SphericalRect.x1 * SphericalRect.x1);
+	float n2z =  SphericalRect.y1 * rsqrt(z0sq + SphericalRect.y1 * SphericalRect.y1);
+	float n3z = -SphericalRect.x0 * rsqrt(z0sq + SphericalRect.x0 * SphericalRect.x0);
+
+	// Original paper uses acos(-dot(...)). Here we use the identify acos(-x) == pi/2+asin(x) to simplify
+	// the math.
+	float G0G1 = SphericalRectAsin(n0z * n1z) + SphericalRectAsin(n1z * n2z);
+	float G2G3 = SphericalRectAsin(n2z * n3z) + SphericalRectAsin(n3z * n0z);
+
+	SphericalRect.b0 = n0z;
+	SphericalRect.b1 = n2z;
+	SphericalRect.k = G2G3;
+	SphericalRect.SolidAngle = G0G1 + SphericalRect.k;
 
 	return SphericalRect;
 }
 
 struct FSphericalRectSample {
 	float3 Direction;
+	float  Distance;
 	float2 UV;
+	float  InvPdf;
 };
 
-FSphericalRectSample UniformSampleSphericalRectWithUV(float2 E, FSphericalRect Rect)
-{
-	float au = E.x * Rect.SolidAngle + Rect.k;
-	float fu = (cos(au) * Rect.b0 - Rect.b1) / sin(au);
-	float cu = rsqrt(fu * fu + Rect.b0 * Rect.b0) * (fu > 0 ? 1 : -1);
-	cu = clamp(cu, -1, 1);				// avoid NaNs
-
-	float xu = -(cu * Rect.z0) * rsqrt(1 - cu * cu);
-	xu = clamp(xu, Rect.x0, Rect.x1);	// avoid Infs
-
-	float d = sqrt(xu * xu + Rect.z0 * Rect.z0);
-	float h0 = Rect.y0 * rsqrt(d * d + Rect.y0 * Rect.y0);
-	float h1 = Rect.y1 * rsqrt(d * d + Rect.y1 * Rect.y1);
-	float hv = h0 + E.y * (h1 - h0);
-	float yv = (hv * hv < 1 - 1e-6) ? (hv * d) * rsqrt(1 - hv * hv) : Rect.y1;
-
-	FSphericalRectSample Result;
-	Result.Direction = mul(float3(xu, yv, Rect.z0), Rect.Axis);
-	Result.UV = float2(xu - Rect.x0, yv - Rect.y0) / float2(Rect.x1 - Rect.x0, Rect.y1 - Rect.y0);
+#define SPHERICAL_RECT_MIN_SOLIDANGLE 1e-3
 
-	return Result;
-}
-
-float3 UniformSampleSphericalRect( float2 E, FSphericalRect Rect )
+float GetSphericalRectInversePdf(float3 Direction, float DistanceSquared, FSphericalRect Rect)
 {
-	return UniformSampleSphericalRectWithUV(E, Rect).Direction;
+	if (Rect.SolidAngle > SPHERICAL_RECT_MIN_SOLIDANGLE)
+	{
+		// Light is big enough to use solid angle sampling
+		return Rect.SolidAngle;
+	}
+	else
+	{
+		// Light is very small, area sampling is sufficient
+		float Area = (Rect.y1 - Rect.y0) * (Rect.x1 - Rect.x0);
+		float NoL = abs(dot(Direction, Rect.Axis[2]));
+		return Area * NoL / DistanceSquared;
+	}
 }
 
-float3 RectIrradianceRef( float3 N, FRect Rect, out float BaseIrradiance, out float NoL )
+FSphericalRectSample UniformSampleSphericalRect(float2 E, FSphericalRect Rect)
 {
-	FSphericalRect SphericalRect = BuildSphericalRect( Rect );
-
-	const uint NumSamples = 32;
-
-	float3 L = 0;
-	NoL = 0;
-
-	LOOP
-	for( uint i = 0; i < NumSamples; i++ )
+	float xu, yv;
+	if (Rect.SolidAngle > SPHERICAL_RECT_MIN_SOLIDANGLE)
 	{
-		float2 E = Hammersley( i, NumSamples, 0 );
-		float3 Ls = normalize( UniformSampleSphericalRect( E, SphericalRect ) );
-
-		L   += Ls;
-		NoL += saturate( dot(N, Ls) );
+		// Some signs are flipped from the original paper on account of the different calculation of SolidAngle and k
+		float au = E.x * Rect.SolidAngle - Rect.k;
+		float fu = (cos(au) * Rect.b0 + Rect.b1) / sin(au);
+		float cu = rsqrt(fu * fu + Rect.b0 * Rect.b0) * (fu > 0 ? 1 : -1);
+		cu = clamp(cu, -1, 1);				// avoid NaNs
+
+		xu = -(cu * Rect.z0) * rsqrt(1 - cu * cu);
+		xu = clamp(xu, Rect.x0, Rect.x1);	// avoid Infs
+
+		float d2 = xu * xu + Rect.z0 * Rect.z0;
+		float h0 = Rect.y0 * rsqrt(d2 + Rect.y0 * Rect.y0);
+		float h1 = Rect.y1 * rsqrt(d2 + Rect.y1 * Rect.y1);
+		float hv = h0 + E.y * (h1 - h0);
+		float rv = 1.0 - hv * hv;
+		yv = (rv > 0) ? (hv * d2 * rsqrt(rv * d2)) : Rect.y1;
 	}
+	else
+	{
+		// fallback to area sampling when solid angle is tiny to avoid numerical issues
+		xu = lerp(Rect.x0, Rect.x1, E.x);
+		yv = lerp(Rect.y0, Rect.y1, E.y);
+	}
+	FSphericalRectSample Result;
+	Result.Direction = mul(float3(xu, yv, Rect.z0), Rect.Axis);
+	Result.UV = float2(xu - Rect.x0, yv - Rect.y0) / float2(Rect.x1 - Rect.x0, Rect.y1 - Rect.y0);
 
-	BaseIrradiance = SphericalRect.SolidAngle;
-	NoL /= NumSamples;
+	float DistanceSquared = xu * xu + yv * yv + Rect.z0 * Rect.z0;
+	float InvDistance = rsqrt(DistanceSquared);
 
-	return normalize(L);
+	Result.Distance = DistanceSquared * InvDistance;
+	Result.Direction *= InvDistance;
+
+	Result.InvPdf = GetSphericalRectInversePdf(Result.Direction, DistanceSquared, Rect);
+	return Result;
 }
 
 FRect GetRect(
diff --git a/Engine/Shaders/Private/RectLightIntegrate.ush b/Engine/Shaders/Private/RectLightIntegrate.ush
--- a/Engine/Shaders/Private/RectLightIntegrate.ush
+++ b/Engine/Shaders/Private/RectLightIntegrate.ush
@@ -52,9 +52,7 @@ float IntegrateLight( FRect Rect )
 	float NoL;
 	float Falloff;
 
-#if 0	// Reference quality
-	float3 L = RectIrradianceRef( 0, Rect, Falloff, NoL );
-#elif !CHEAP_RECT	// Optimized Lambert
+#if !CHEAP_RECT	// Optimized Lambert
 	float3 L = RectIrradianceLambert( 0, Rect, Falloff, NoL );
 #elif 1	// Karis
 	float3 L = RectIrradianceApproxKaris( 0, Rect, Falloff, NoL );
@@ -74,9 +72,7 @@ FAreaLightIntegrateContext CreateRectIntegrateContext( float Roughness, half3 N,
 	
 	FAreaLightIntegrateContext Out = InitAreaLightIntegrateContext();
 
-#if 0	// Reference quality
-	float3 L = RectIrradianceRef( N, Rect, Falloff, NoL );
-#elif !CHEAP_RECT	// Optimized Lambert
+#if !CHEAP_RECT	// Optimized Lambert
 	float3 L = RectIrradianceLambert( N, Rect, Falloff, NoL );
 #elif 1	// Karis
 	float3 L = RectIrradianceApproxKaris( N, Rect, Falloff, NoL );
@@ -332,8 +328,7 @@ FDirectLighting IntegrateBxDF( FGBufferData GBuffer, half3 N, half3 V, FRect Rec
 				}
 				case 3:
 				{
-					float3 ToArea = UniformSampleSphericalRect( E, SphericalRect );
-					L = normalize( ToArea );
+					L = UniformSampleSphericalRect( E, SphericalRect ).Direction;
 					H = normalize( V + L );
 					break;
 				}
