Fix modulated shadows in mobile multi view.

diff --git a/Engine/Shaders/Private/DepthOnlyVertexShader.usf b/Engine/Shaders/Private/DepthOnlyVertexShader.usf
--- a/Engine/Shaders/Private/DepthOnlyVertexShader.usf
+++ b/Engine/Shaders/Private/DepthOnlyVertexShader.usf
@@ -35,14 +35,20 @@ void Main(
 #if USE_GLOBAL_CLIP_PLANE
 	, out float OutGlobalClipPlaneDistance : SV_ClipDistance
 #endif
-#if INSTANCED_STEREO
+#if INSTANCED_STEREO && MOBILE_MULTI_VIEW
+	, out uint LayerIndex : SV_RenderTargetArrayIndex
+#elif INSTANCED_STEREO
 	, out uint ViewportIndex : SV_ViewPortArrayIndex
 #endif
 	)
 {
 #if INSTANCED_STEREO
 	uint EyeIndex = GetEyeIndexFromVF(Input);
-	ViewportIndex = EyeIndex;
+	#if MOBILE_MULTI_VIEW
+		LayerIndex = EyeIndex;
+	#else
+		ViewportIndex = EyeIndex;
+	#endif
 #endif
 	ResolvedView = ResolveViewFromVF(Input);
 
diff --git a/Engine/Shaders/Private/SceneTexturesCommon.ush b/Engine/Shaders/Private/SceneTexturesCommon.ush
--- a/Engine/Shaders/Private/SceneTexturesCommon.ush
+++ b/Engine/Shaders/Private/SceneTexturesCommon.ush
@@ -188,13 +188,17 @@ float4 CalcFullSceneColor(float2 ScreenUV)
 
 
 /** Returns DeviceZ which is the z value stored in the depth buffer. */
-float LookupDeviceZ( float2 ScreenUV )
+float LookupDeviceZ( float2 ScreenUV, float ArrayIndex )
 {
 #if	SCENE_TEXTURES_DISABLED
 	return FarDepthValue;
 #elif FORCE_DEPTH_TEXTURE_READS || PLATFORM_NEEDS_DEPTH_TEXTURE_READS
 	// native Depth buffer lookup
-	return Texture2DSampleLevel(MobileSceneTextures.SceneDepthTexture, MobileSceneTextures.SceneDepthTextureSampler, ScreenUV, 0).r;
+	#if INSTANCED_STEREO || MOBILE_MULTI_VIEW
+		return Texture2DArraySample(MobileSceneTextures.SceneDepthTextureArray, MobileSceneTextures.SceneDepthTextureSampler, float3(ScreenUV.xy,ArrayIndex)).r;
+	#else
+		return Texture2DSampleLevel(MobileSceneTextures.SceneDepthTexture, MobileSceneTextures.SceneDepthTextureSampler, ScreenUV, 0).r;
+	#endif
 #elif MOBILE_DEPTHFECTH && COMPILER_GLSL_ES3_1
 	return DepthbufferFetchES2();
 #elif MOBILE_DEPTHFECTH && VULKAN_PROFILE
@@ -204,21 +208,44 @@ float LookupDeviceZ( float2 ScreenUV )
 #elif (USE_SCENE_DEPTH_AUX && !MOBILE_DEFERRED_SHADING)
 	// SceneDepth texture is discarded after BasePass (with forward shading) 
 	// instead fetch DeviceZ from SceneDepthAuxTexture
-	return Texture2DSampleLevel(MobileSceneTextures.SceneDepthAuxTexture, MobileSceneTextures.SceneDepthAuxTextureSampler, ScreenUV, 0).r;
+	#if INSTANCED_STEREO || MOBILE_MULTI_VIEW
+		return Texture2DArraySample(MobileSceneTextures.SceneDepthAuxTextureArray, MobileSceneTextures.SceneDepthAuxTextureSampler, float3(ScreenUV.xy,ArrayIndex)).r;
+	#else
+		return Texture2DSampleLevel(MobileSceneTextures.SceneDepthAuxTexture, MobileSceneTextures.SceneDepthAuxTextureSampler, ScreenUV, 0).r;
+	#endif
 #else
 	// native Depth buffer lookup
-	return Texture2DSampleLevel(MobileSceneTextures.SceneDepthTexture, MobileSceneTextures.SceneDepthTextureSampler, ScreenUV, 0).r;
+	#if INSTANCED_STEREO || MOBILE_MULTI_VIEW
+		return Texture2DArraySample(MobileSceneTextures.SceneDepthTextureArray, MobileSceneTextures.SceneDepthTextureSampler, float3(ScreenUV.xy,ArrayIndex)).r;
+	#else
+		return Texture2DSampleLevel(MobileSceneTextures.SceneDepthTexture, MobileSceneTextures.SceneDepthTextureSampler, ScreenUV, 0).r;
+	#endif
 #endif
 }
 
+/** Returns DeviceZ which is the z value stored in the depth buffer. */
+float LookupDeviceZ( float2 ScreenUV )
+{
+	return LookupDeviceZ(ScreenUV, 0);
+}
+
 /** Returns clip space W, which is world space distance along the View Z axis. Note if you need DeviceZ LookupDeviceZ() is the faster option */
-float CalcSceneDepth(float2 ScreenUV)
+float CalcSceneDepth(float2 ScreenUV, float ArrayIndex)
 {
 #if SCENE_TEXTURES_DISABLED
 	return SCENE_TEXTURES_DISABLED_SCENE_DEPTH_VALUE;
 #else
-	return ConvertFromDeviceZ(LookupDeviceZ(ScreenUV));
+	#if INSTANCED_STEREO || MOBILE_MULTI_VIEW
+			return ConvertFromDeviceZ(LookupDeviceZ(ScreenUV, ArrayIndex));
+	#else
+			return ConvertFromDeviceZ(LookupDeviceZ(ScreenUV));
+	#endif
 #endif
 }
 
+float CalcSceneDepth(float2 ScreenUV)
+{
+	return CalcSceneDepth(ScreenUV, 0);
+}
+
 #endif // SHADING_PATH_MOBILE
diff --git a/Engine/Shaders/Private/ShadowProjectionPixelShader.usf b/Engine/Shaders/Private/ShadowProjectionPixelShader.usf
--- a/Engine/Shaders/Private/ShadowProjectionPixelShader.usf
+++ b/Engine/Shaders/Private/ShadowProjectionPixelShader.usf
@@ -103,17 +103,37 @@ float2 ShadowNearAndFarDepth;
 EARLYDEPTHSTENCIL
 #endif
 void Main(
-	in float4 SVPos : SV_POSITION,
 	out float4 OutColor : SV_Target0
+#if INSTANCED_STEREO
+	, in nointerpolation uint EyeIndex : TEXCOORD0
+#elif MOBILE_MULTI_VIEW
+	, in nointerpolation uint ViewId : SV_ViewID
+#endif
+	, in float4 SVPos : SV_POSITION
 	)
 {
 	#if USE_FADE_PLANE
 	const bool bUseFadePlane = true;
 	#endif
+	
+	float ArrayIndex = 0;
+#if INSTANCED_STEREO
+	ResolvedView = ResolveView(EyeIndex);
+	ArrayIndex = EyeIndex;
+#elif MOBILE_MULTI_VIEW
+	ResolvedView = ResolveView(ViewId);
+	ArrayIndex = ViewId;;
+#else
+	ResolvedView = ResolveView();
+#endif
 
 	const FPQMPContext PQMPContext = PQMPInit(SVPos.xy);
-	float2 ScreenUV = float2(SVPos.xy * View.BufferSizeAndInvSize.zw);
+	float2 ScreenUV = float2(SVPos.xy * ResolvedView.BufferSizeAndInvSize.zw);
+#if (INSTANCED_STEREO || MOBILE_MULTI_VIEW) && SHADING_PATH_MOBILE
+	float SceneW = CalcSceneDepth(ScreenUV, ArrayIndex);
+#else
 	float SceneW = CalcSceneDepth(ScreenUV);
+#endif
 
 	bool bIsSubsurfaceCompatible = true;
 	const uint3 PixelCoord = uint3(floor(SVPos.xy), 0);
@@ -122,10 +142,17 @@ void Main(
 	SceneW = ConvertFromDeviceZ(HairStrands.HairOnlyDepthTexture.Load(PixelCoord).x);
 	bIsSubsurfaceCompatible = false;
 #endif
-
-	float4 ScreenPosition = float4(GetScreenPositionForProjectionType(((ScreenUV.xy - View.ScreenPositionScaleBias.wz ) / View.ScreenPositionScaleBias.xy),SceneW), SceneW, 1);
-	float4 ShadowPosition = mul(ScreenPosition, ScreenToShadowMatrix);
-	float3 TranslateWorldPosition = mul(ScreenPosition, View.ScreenToTranslatedWorld).xyz;
+	float4 ScreenPosition = float4(GetScreenPositionForProjectionType(((ScreenUV.xy - ResolvedView.ScreenPositionScaleBias.wz) / ResolvedView.ScreenPositionScaleBias.xy), SceneW), SceneW, 1);
+	
+	float4 ShadowPosition = float4(0, 0, 0, 0);
+	
+#if MOBILE_MULTI_VIEW || INSTANCED_STEREO
+	ShadowPosition = mul(ScreenPosition, ResolvedView.MobileMultiviewShadowTransform);
+	ShadowPosition = mul(ShadowPosition, ScreenToShadowMatrix);
+#else
+	ShadowPosition = mul(ScreenPosition, ScreenToShadowMatrix);
+#endif
+	float3 TranslateWorldPosition = mul(ScreenPosition, ResolvedView.ScreenToTranslatedWorld).xyz;
 
 	float ShadowZ = 1 - ShadowPosition.z;
 	ShadowPosition.xyz /= ShadowPosition.w;
@@ -256,7 +283,7 @@ void Main(
 		bool bHasSubsurfaceTransmission	= false;
 		bool bIsEyeOrHair 				= false;
 	#if SUBSTRATE_ENABLED
-		FSubstrateAddressing SubstrateAddressing = GetSubstratePixelDataByteOffset(PixelCoord.xy, uint2(View.BufferSizeAndInvSize.xy), Substrate.MaxBytesPerPixel);
+		FSubstrateAddressing SubstrateAddressing = GetSubstratePixelDataByteOffset(PixelCoord.xy, uint2(ResolvedView.BufferSizeAndInvSize.xy), Substrate.MaxBytesPerPixel);
 		FSubstratePixelHeader SubstrateHeader = UnpackSubstrateHeaderIn(Substrate.MaterialTextureArray, SubstrateAddressing, Substrate.TopLayerTexture);
 		if (SubstrateHeader.IsSubstrateMaterial())
 		{
diff --git a/Engine/Shaders/Private/ShadowProjectionVertexShader.usf b/Engine/Shaders/Private/ShadowProjectionVertexShader.usf
--- a/Engine/Shaders/Private/ShadowProjectionVertexShader.usf
+++ b/Engine/Shaders/Private/ShadowProjectionVertexShader.usf
@@ -7,9 +7,19 @@
 #include "Common.ush"
 
 void ShadowProjectionNoTransformVS(
-	in float4 InPosition : ATTRIBUTE0,
-	out float4 OutPosition : SV_POSITION)
+	in float4 InPosition : ATTRIBUTE0
+#if INSTANCED_STEREO
+	, in uint InstanceId : SV_InstanceID
+	, out nointerpolation uint EyeIndex : TEXCOORD0
+	, out nointerpolation uint LayerIndex : SV_RenderTargetArrayIndex
+#endif
+	, out float4 OutPosition : SV_POSITION
+)
 {
+#if INSTANCED_STEREO
+	EyeIndex = GetEyeIndex(InstanceId);
+	LayerIndex = EyeIndex;
+#endif
 	OutPosition = float4(InPosition.xyz, 1);
 }
 
@@ -19,14 +29,31 @@ float4x4 InvReceiverInnerMatrix;
 float4 PreShadowToPreViewTranslation;
 
 void ShadowVolumeBoundProjectionVS(
-	in float4 InPosition : ATTRIBUTE0,
-	out float4 OutPosition : SV_POSITION)
+	in float4 InPosition : ATTRIBUTE0
+#if INSTANCED_STEREO
+	, in uint InstanceId : SV_InstanceID
+	, out nointerpolation uint EyeIndex : TEXCOORD0
+	, out uint LayerIndex : SV_RenderTargetArrayIndex
+#elif MOBILE_MULTI_VIEW
+	, in nointerpolation uint ViewId : SV_ViewID
+#endif
+	, out float4 OutPosition : SV_POSITION
+)
 {
+#if INSTANCED_STEREO
+	EyeIndex = GetEyeIndex(InstanceId);
+	LayerIndex = EyeIndex;
+	ResolvedView = ResolveView(EyeIndex);
+#elif MOBILE_MULTI_VIEW
+	ResolvedView = ResolveView(ViewId);
+#else
+	ResolvedView = ResolveView();
+#endif
 	float4 UnprojectedVertex = mul(float4(InPosition.xyz, 1), InvReceiverInnerMatrix);
 	float3 OutVertex = UnprojectedVertex.xyz / UnprojectedVertex.w + PreShadowToPreViewTranslation.xyz;
 
 	float3 WorldPosition = OutVertex.xyz * StencilingGeometryPosAndScale.w + StencilingGeometryPosAndScale.xyz;
-	OutPosition = mul(float4(WorldPosition, 1), View.TranslatedWorldToClip);
+	OutPosition = mul(float4(WorldPosition, 1), ResolvedView.TranslatedWorldToClip);
 }
 
 
diff --git a/Engine/Source/Runtime/Engine/Public/SceneTexturesConfig.h b/Engine/Source/Runtime/Engine/Public/SceneTexturesConfig.h
--- a/Engine/Source/Runtime/Engine/Public/SceneTexturesConfig.h
+++ b/Engine/Source/Runtime/Engine/Public/SceneTexturesConfig.h
@@ -41,6 +41,7 @@ BEGIN_GLOBAL_SHADER_PARAMETER_STRUCT(FMobileSceneTextureUniformParameters, ENGIN
 	SHADER_PARAMETER_RDG_TEXTURE(Texture2D, SceneColorTexture)
 	SHADER_PARAMETER_SAMPLER(SamplerState, SceneColorTextureSampler)
 	SHADER_PARAMETER_RDG_TEXTURE(Texture2D, SceneDepthTexture)
+	SHADER_PARAMETER_RDG_TEXTURE(Texture2DArray, SceneDepthTextureArray)
 	SHADER_PARAMETER_SAMPLER(SamplerState, SceneDepthTextureSampler)
 	SHADER_PARAMETER_RDG_TEXTURE(Texture2D, ScenePartialDepthTexture)
 	SHADER_PARAMETER_SAMPLER(SamplerState, ScenePartialDepthTextureSampler)
@@ -55,6 +56,7 @@ BEGIN_GLOBAL_SHADER_PARAMETER_STRUCT(FMobileSceneTextureUniformParameters, ENGIN
 	SHADER_PARAMETER_RDG_TEXTURE(Texture2D, GBufferCTexture)
 	SHADER_PARAMETER_RDG_TEXTURE(Texture2D, GBufferDTexture)
 	SHADER_PARAMETER_RDG_TEXTURE(Texture2D, SceneDepthAuxTexture)
+	SHADER_PARAMETER_RDG_TEXTURE(Texture2DArray, SceneDepthAuxTextureArray)
 	SHADER_PARAMETER_RDG_TEXTURE(Texture2D, LocalLightTextureA)
 	SHADER_PARAMETER_RDG_TEXTURE(Texture2D, LocalLightTextureB)
 	SHADER_PARAMETER_SAMPLER(SamplerState, GBufferATextureSampler)
diff --git a/Engine/Source/Runtime/Renderer/Private/MobileShadingRenderer.cpp b/Engine/Source/Runtime/Renderer/Private/MobileShadingRenderer.cpp
--- a/Engine/Source/Runtime/Renderer/Private/MobileShadingRenderer.cpp
+++ b/Engine/Source/Runtime/Renderer/Private/MobileShadingRenderer.cpp
@@ -449,6 +449,9 @@ void FMobileSceneRenderer::InitViews(
 	for (int32 ViewIndex = 0; ViewIndex < Views.Num(); ++ViewIndex)
 	{
 		Views[ViewIndex].GPUSceneViewId = InstanceCullingManager.RegisterView(Views[ViewIndex]);
+		// Set the stereo instance factor as GetStereoPassInstanceFactor can't be called from inside an RDG pass.
+		uint32 InstanceFactor = Views[ViewIndex].GetStereoPassInstanceFactor();
+		Views[ViewIndex].InstanceFactor = InstanceFactor > 0 ? InstanceFactor : 1;
 	}
 
 	FILCUpdatePrimTaskData* ILCTaskData = nullptr;
diff --git a/Engine/Source/Runtime/Renderer/Private/SceneRendering.h b/Engine/Source/Runtime/Renderer/Private/SceneRendering.h
--- a/Engine/Source/Runtime/Renderer/Private/SceneRendering.h
+++ b/Engine/Source/Runtime/Renderer/Private/SceneRendering.h
@@ -1556,6 +1556,10 @@ public:
 	FTextureRHIRef PrimitiveSceneDataTextureOverrideRHI;
 
 	FShaderPrintData ShaderPrintData;
+
+	/** When using mobile multi view fallback path we need to instance draw calls ourselves to cover both eyes instead of letting the drivers do it for us. */
+	uint32 InstanceFactor = 1;
+
 private:
 	FLumenTranslucencyGIVolume LumenTranslucencyGIVolume;
 public:
diff --git a/Engine/Source/Runtime/Renderer/Private/SceneTextures.cpp b/Engine/Source/Runtime/Renderer/Private/SceneTextures.cpp
--- a/Engine/Source/Runtime/Renderer/Private/SceneTextures.cpp
+++ b/Engine/Source/Runtime/Renderer/Private/SceneTextures.cpp
@@ -1007,6 +1007,7 @@ void SetupMobileSceneTextureUniformParameters(
 	SceneTextureParameters.SceneColorTexture = SystemTextures.Black;
 	SceneTextureParameters.SceneColorTextureSampler = TStaticSamplerState<>::GetRHI();
 	SceneTextureParameters.SceneDepthTexture = SystemTextures.DepthDummy;
+	SceneTextureParameters.SceneDepthTextureArray = GSystemTextures.GetDefaultTexture(GraphBuilder, ETextureDimension::Texture2DArray, EPixelFormat::PF_B8G8R8A8, FClearValueBinding::Black);
 	SceneTextureParameters.SceneDepthTextureSampler = TStaticSamplerState<>::GetRHI();
 	SceneTextureParameters.ScenePartialDepthTexture = SystemTextures.DepthDummy;
 	SceneTextureParameters.ScenePartialDepthTextureSampler = TStaticSamplerState<>::GetRHI();
@@ -1022,6 +1023,7 @@ void SetupMobileSceneTextureUniformParameters(
 	SceneTextureParameters.GBufferDTexture = SystemTextures.Black;
 	// SceneDepthAuxTexture is a color texture on mobile, with DeviceZ values
 	SceneTextureParameters.SceneDepthAuxTexture = SystemTextures.Black;
+	SceneTextureParameters.SceneDepthAuxTextureArray = GSystemTextures.GetDefaultTexture(GraphBuilder, ETextureDimension::Texture2DArray, EPixelFormat::PF_B8G8R8A8, FClearValueBinding::Black);
 	SceneTextureParameters.LocalLightTextureA = SystemTextures.Black;
 	SceneTextureParameters.LocalLightTextureB = SystemTextures.Black;
 	SceneTextureParameters.GBufferATextureSampler = TStaticSamplerState<>::GetRHI();
@@ -1042,6 +1044,7 @@ void SetupMobileSceneTextureUniformParameters(
 			!EnumHasAnyFlags(SceneTextures->Depth.Resolve->Desc.Flags, TexCreate_Memoryless))
 		{
 			SceneTextureParameters.SceneDepthTexture = SceneTextures->Depth.Resolve;
+			SceneTextureParameters.SceneDepthTextureArray = SceneTextures->Depth.Resolve;
 		}
 
 		if (EnumHasAnyFlags(SetupMode, EMobileSceneTextureSetupMode::SceneDepth) &&
@@ -1079,6 +1082,7 @@ void SetupMobileSceneTextureUniformParameters(
 			if (HasBeenProduced(SceneTextures->DepthAux.Resolve))
 			{
 				SceneTextureParameters.SceneDepthAuxTexture = SceneTextures->DepthAux.Resolve;
+				SceneTextureParameters.SceneDepthAuxTextureArray = SceneTextures->DepthAux.Resolve;
 			}
 		}
 
diff --git a/Engine/Source/Runtime/Renderer/Private/ShadowRendering.cpp b/Engine/Source/Runtime/Renderer/Private/ShadowRendering.cpp
--- a/Engine/Source/Runtime/Renderer/Private/ShadowRendering.cpp
+++ b/Engine/Source/Runtime/Renderer/Private/ShadowRendering.cpp
@@ -1070,7 +1070,7 @@ void FProjectedShadowInfo::SetupProjectionStencilMask(
 		RHICmdList.SetStreamSource(0, GFrustumVertexBuffer.VertexBufferRHI, 0);
 
 		// Draw the frustum using the stencil buffer to mask just the pixels which are inside the shadow frustum.
-		RHICmdList.DrawIndexedPrimitive(GCubeIndexBuffer.IndexBufferRHI, 0, 0, 8, 0, 12, 1);
+		RHICmdList.DrawIndexedPrimitive(GCubeIndexBuffer.IndexBufferRHI, 0, 0, 8, 0, 12, View->InstanceFactor);
 
 		// if rendering modulated shadows mask out subject mesh elements to prevent self shadowing.
 		if (bMobileModulatedProjections && !CVarEnableModulatedSelfShadow.GetValueOnRenderThread())
@@ -1413,16 +1413,17 @@ void FProjectedShadowInfo::RenderProjectionInternal(
 		}
 	}
 
+	uint32 NumberOfInstances = View->InstanceFactor;
 	if (IsWholeSceneDirectionalShadow())
 	{
 		RHICmdList.SetStreamSource(0, GClearVertexBuffer.VertexBufferRHI, 0);
-		RHICmdList.DrawPrimitive(0, 2, 1);
+		RHICmdList.DrawPrimitive(0, 2, NumberOfInstances);
 	}
 	else
 	{
 		RHICmdList.SetStreamSource(0, GFrustumVertexBuffer.VertexBufferRHI, 0);
 		// Draw the frustum using the projection shader..
-		RHICmdList.DrawIndexedPrimitive(GCubeIndexBuffer.IndexBufferRHI, 0, 0, 8, 0, 12, 1);
+		RHICmdList.DrawIndexedPrimitive(GCubeIndexBuffer.IndexBufferRHI, 0, 0, 8, 0, 12, NumberOfInstances);
 	}
 
 	if (!bDepthBoundsTestEnabled && bStencilTestEnabled)
diff --git a/Engine/Source/Runtime/Renderer/Private/SimpleMeshDrawCommandPass.cpp b/Engine/Source/Runtime/Renderer/Private/SimpleMeshDrawCommandPass.cpp
--- a/Engine/Source/Runtime/Renderer/Private/SimpleMeshDrawCommandPass.cpp
+++ b/Engine/Source/Runtime/Renderer/Private/SimpleMeshDrawCommandPass.cpp
@@ -25,7 +25,7 @@ FSimpleMeshDrawCommandPass::FSimpleMeshDrawCommandPass(const FSceneView& View, F
 	static FName NAME_SimpleMDCPass("SimpleMDCPass");
 	InstanceCullingContext = FInstanceCullingContext(NAME_SimpleMDCPass, ViewInfo->GetShaderPlatform(), InstanceCullingManager, ViewIds, nullptr, bUsingStereo ? EInstanceCullingMode::Stereo : EInstanceCullingMode::Normal);
 
-	InstanceFactor = static_cast<uint32>(ViewIds.Num());
+	InstanceFactor = ViewInfo->InstanceFactor;
 }
 
 void FSimpleMeshDrawCommandPass::BuildRenderingCommands(FRDGBuilder& GraphBuilder, const FSceneView& View, const FGPUScene& GPUScene, FInstanceCullingDrawParams& OutInstanceCullingDrawParams)
