Fix failing engine test.
* Fix uninitialized variable causing incorrect hair strands rendering on Vulkan.
* Fix jitter mode 2 not working correct with hair blue noise random function, causing transmittance to not be stable in engine test.

diff --git a/Engine/Shaders/Private/HairStrands/HairStrandsDeepTransmittanceMask.usf b/Engine/Shaders/Private/HairStrands/HairStrandsDeepTransmittanceMask.usf
--- a/Engine/Shaders/Private/HairStrands/HairStrandsDeepTransmittanceMask.usf
+++ b/Engine/Shaders/Private/HairStrands/HairStrandsDeepTransmittanceMask.usf
@@ -282,7 +282,7 @@ float3 InternalGetHairVoxelJitter(uint2 PixelCoord, uint SampleIndex, uint MaxSa
 {
 	// Blue noise is cheaper to compute and h
 	#if USE_BLUE_NOISE
-	float3 RandomSample = GetHairBlueNoiseJitter(PixelCoord, SampleIndex, MaxSampleCount, TimeIndex).xyz;
+	float3 RandomSample = GetHairBlueNoiseJitter(PixelCoord, SampleIndex, MaxSampleCount, JitterMode > 1 ? 0 : TimeIndex).xyz;
 	#else
 	float3 RandomSample = GetHairVoxelJitter(PixelCoord, View.StateFrameIndexMod8, JitterMode);
 	#endif
diff --git a/Engine/Shaders/Private/HairStrands/HairStrandsVisibilityVS.usf b/Engine/Shaders/Private/HairStrands/HairStrandsVisibilityVS.usf
--- a/Engine/Shaders/Private/HairStrands/HairStrandsVisibilityVS.usf
+++ b/Engine/Shaders/Private/HairStrands/HairStrandsVisibilityVS.usf
@@ -43,6 +43,7 @@ void Main(
 	HairViewInfo.TranslatedWorldCameraOrigin = ResolvedView.TranslatedWorldCameraOrigin;
 	HairViewInfo.ViewForward = ResolvedView.ViewForward;
 	HairViewInfo.RadiusAtDepth1 = lerp(HairRenderInfo.RadiusAtDepth1Primary, HairRenderInfo.RadiusAtDepth1Velocity, VelocityScale);
+	HairViewInfo.bIsOrthoView = HairRenderInfo.bIsOrthoView;
 
 	uint HairControlPointId = 0;
 	const float4 WorldPosition = VertexFactoryGetWorldPosition_Visibility(Input, HairViewInfo, HairControlPointId);
@@ -91,6 +92,7 @@ void Main(
 	HairViewInfo.TranslatedWorldCameraOrigin = ResolvedView.TranslatedWorldCameraOrigin;
 	HairViewInfo.ViewForward = ResolvedView.ViewForward;
 	HairViewInfo.RadiusAtDepth1 = lerp(HairRenderInfo.RadiusAtDepth1Primary, HairRenderInfo.RadiusAtDepth1Velocity, VelocityScale);
+	HairViewInfo.bIsOrthoView = HairRenderInfo.bIsOrthoView;
 
 	float4 WorldPosition = VertexFactoryGetWorldPosition(Input, VFIntermediates, HairViewInfo);
 #else
