Fixing shader compilation warning

diff --git a/Engine/Config/ConsoleVariables.ini b/Engine/Config/ConsoleVariables.ini
--- a/Engine/Config/ConsoleVariables.ini
+++ b/Engine/Config/ConsoleVariables.ini
@@ -41,6 +41,9 @@
 ;  Mode 3: dump on compilation failure or warnings.
 r.DumpShaderDebugInfo=2
 
+; When this is enabled, shader compilation warnings will emit as errors
+r.Shaders.WarningsAsErrors=0 ;; 0: off, 1: global shaders, 2: all
+
 ; When this is enabled, dumped shader paths will get collapsed (in the cases where paths are longer than the OS's max)
 ;r.DumpShaderDebugShortNames=1
 ; When this is enabled, when dumping shaders an additional file to use with ShaderCompilerWorker -direct mode will be generated
diff --git a/Engine/Shaders/Private/HeterogeneousVolumes/HeterogeneousVolumesVoxelGridShadows.usf b/Engine/Shaders/Private/HeterogeneousVolumes/HeterogeneousVolumesVoxelGridShadows.usf
--- a/Engine/Shaders/Private/HeterogeneousVolumes/HeterogeneousVolumesVoxelGridShadows.usf
+++ b/Engine/Shaders/Private/HeterogeneousVolumes/HeterogeneousVolumesVoxelGridShadows.usf
@@ -424,7 +424,7 @@ void RenderVolumetricShadowMapForLightWithVoxelGridCS(
 			bool bFirstEntry = (ElementIndex == 0);
 			if (bFirstEntry && any(Extinction > 0.0))
 			{
-				FAVSMSampleData SampleData = AVSM_CreateSampleData(WorldHitT, Transmittance);
+				FAVSMSampleData SampleData = AVSM_CreateSampleData(WorldHitT, Luminance(Transmittance));
 				AddSample(SampleData, PixelIndex, PrevElementIndex, ElementIndex);
 
 				// Set expected values
@@ -439,7 +439,7 @@ void RenderVolumetricShadowMapForLightWithVoxelGridCS(
 				float RelativeError = length((Transmittance - ExpectedTransmittance) / ExpectedTransmittance);
 				if ((AbsoluteError > AbsoluteErrorThreshold) && (RelativeError > RelativeErrorThreshold))
 				{
-					FAVSMSampleData SampleData = AVSM_CreateSampleData(WorldHitT, Transmittance);
+					FAVSMSampleData SampleData = AVSM_CreateSampleData(WorldHitT, Luminance(Transmittance));
 					AddOrReplaceSample(SampleData, PixelIndex, PrevElementIndex, ElementIndex);
 
 					// Reset expected values
