Path Tracer: Fix inaccuracy in FastAsin due to re-using the input instead of its absolute value

This made the error closer to 10^-1 for negative arguments instead of 10^-3 as expected.

diff --git a/Engine/Shaders/Private/PathTracing/Material/PathTracingHair.ush b/Engine/Shaders/Private/PathTracing/Material/PathTracingHair.ush
--- a/Engine/Shaders/Private/PathTracing/Material/PathTracingHair.ush
+++ b/Engine/Shaders/Private/PathTracing/Material/PathTracingHair.ush
@@ -120,9 +120,8 @@ float LogI0(float x)
 float FastASin(float x) {
 	// Maximum error: 10^-3.35
 	// TODO: consider replacing acosFast/asinFast which have approximation error 10^-2.0 (with one less fma, but also with a slightly sub-optimal constant)
-	x = clamp(x, -1.0, 1.0);
-	float a = abs(x);
-	float r = (0.5 * PI) + x * (-0.207034 + 0.0531013 * x);
+	float a = saturate(abs(x));
+	float r = (0.5 * PI) + a * (-0.207034 + 0.0531013 * a);
 	r = (0.5 * PI) - r * sqrt(1 - a);
 	return x >= 0 ? r : -r;
 }
