Many Lights - remove debug View.GeneralPurposeTweak

diff --git a/Engine/Shaders/Private/ManyLights/ManyLightsDenoiserTemporal.usf b/Engine/Shaders/Private/ManyLights/ManyLightsDenoiserTemporal.usf
--- a/Engine/Shaders/Private/ManyLights/ManyLightsDenoiserTemporal.usf
+++ b/Engine/Shaders/Private/ManyLights/ManyLightsDenoiserTemporal.usf
@@ -145,7 +145,7 @@ void DenoiserTemporalCS(
 				bool bHistoryWasOnScreen = all(HistoryScreenUV >= HistoryUVMinMax.xy) && all(HistoryScreenUV <= HistoryUVMinMax.zw);
 				if (bHistoryWasOnScreen)
 				{
-					HistoryScreenUV = clamp(HistoryScreenUV, HistoryGatherUVMinMax.xy + View.GeneralPurposeTweak * View.BufferSizeAndInvSize.zw, HistoryGatherUVMinMax.zw - View.GeneralPurposeTweak * View.BufferSizeAndInvSize.zw);
+					HistoryScreenUV = clamp(HistoryScreenUV, HistoryGatherUVMinMax.xy, HistoryGatherUVMinMax.zw);
 
 					uint2 HistoryScreenCoord = floor(HistoryScreenUV * View.BufferSizeAndInvSize.xy - 0.5f);
 					float2 HistoryBilinearWeights = frac(HistoryScreenUV * View.BufferSizeAndInvSize.xy - 0.5f);
