Use invariant intrinsics for doublefloat operator implementations, as a workaround for compiler issues with fastmath.

diff --git a/Engine/Shaders/Private/DoubleFloatOperations.ush b/Engine/Shaders/Private/DoubleFloatOperations.ush
--- a/Engine/Shaders/Private/DoubleFloatOperations.ush
+++ b/Engine/Shaders/Private/DoubleFloatOperations.ush
@@ -48,12 +48,13 @@ FDFType DFTwoSum(FFloatType Lhs, FFloatType Rhs)
 #if UE_DF_FORCE_FP32_OPS
 	return DFPromote(Lhs + Rhs);
 #else
-	const FFloatType S = Lhs + Rhs;
-	const FFloatType V = INVARIANT(S - Lhs);
-	const FFloatType Q = INVARIANT(S - V);
-	const FFloatType R = INVARIANT(Lhs - Q);
-	const FFloatType T = INVARIANT(Rhs - V);
-	return DFConstructor(S, R + T);
+	const FFloatType S = INVARIANT_ADD(Lhs, Rhs);
+	const FFloatType V = INVARIANT_SUB(S, Lhs);
+	const FFloatType Q = INVARIANT_SUB(S, V);
+	const FFloatType R = INVARIANT_SUB(Lhs, Q);
+	const FFloatType T = INVARIANT_SUB(Rhs, V);
+	const FFloatType Y = INVARIANT_ADD(R, T);
+	return DFConstructor(S, Y);
 #endif
 }
 
@@ -66,8 +67,9 @@ FDFType DFFastTwoSum(FFloatType Lhs, FFloatType Rhs)
 #if UE_DF_FORCE_FP32_OPS
 	return DFPromote(Lhs + Rhs);
 #else
-	const FFloatType S = Lhs + Rhs;
-	const FFloatType E = INVARIANT(Rhs - (S - Lhs));
+	const FFloatType S = INVARIANT_ADD(Lhs, Rhs);
+	const FFloatType T = INVARIANT_SUB(S, Lhs);
+	const FFloatType E = INVARIANT_SUB(Rhs, T);
 	return DFConstructor(S, E);
 #endif
 }
@@ -216,9 +218,9 @@ FFloatType DFFastLocalSubtractDemote(FDFType Lhs, FDFType Rhs)
 #elif UE_DF_NO_FAST_MATH
 	return DFSubtractDemote(Lhs, Rhs);
 #else
-	const FFloatType High = INVARIANT(Lhs.High - Rhs.High);
-	const FFloatType Low = INVARIANT(Lhs.Low - Rhs.Low);
-	const FFloatType Sum = INVARIANT(High + Low);
+	const FFloatType High = INVARIANT_SUB(Lhs.High, Rhs.High);
+	const FFloatType Low = INVARIANT_SUB(Lhs.Low, Rhs.Low);
+	const FFloatType Sum = INVARIANT_ADD(High, Low);
 	return Sum;
 #endif
 }
@@ -230,8 +232,8 @@ FFloatType DFFastLocalSubtractDemote(FDFType Lhs, FFloatType Rhs)
 #elif UE_DF_NO_FAST_MATH
 	return DFSubtractDemote(Lhs, Rhs);
 #else
-	const FFloatType High = INVARIANT(Lhs.High - Rhs);
-	const FFloatType Sum = INVARIANT(High + Lhs.Low);
+	const FFloatType High = INVARIANT_SUB(Lhs.High, Rhs);
+	const FFloatType Sum = INVARIANT_ADD(High, Lhs.Low);
 	return Sum;
 #endif
 }
@@ -249,16 +251,16 @@ FFloatType DFFastLocalSubtractDemote(FFloatType Lhs, FDFType Rhs)
 FDFType DFTwoProductNoFMA(FFloatType Lhs, FFloatType Rhs)
 {
 	const FFloatType Constant = 4097; // = 2^(t - t/2) + 1, where t is the number of mantissa bits
-	const FFloatType Px = INVARIANT(Lhs * Constant);
-	const FFloatType Hx = INVARIANT(Lhs - Px + Px);
-	const FFloatType Tx = INVARIANT(Lhs - Hx);
-	const FFloatType Py = INVARIANT(Rhs * Constant);
-	const FFloatType Hy = INVARIANT(Rhs - Py + Py);
-	const FFloatType Ty = INVARIANT(Rhs - Hy);
-	const FFloatType P = INVARIANT(Hx * Hy);
-	const FFloatType Q = INVARIANT(Hx * Ty + Tx * Hy);
-	const FFloatType Z = INVARIANT(P + Q);
-	const FFloatType E = INVARIANT(P - Z + Q + Tx * Ty);
+	const FFloatType Px = INVARIANT_MUL(Lhs, Constant);
+	const FFloatType Hx = INVARIANT_ADD(INVARIANT_SUB(Lhs, Px), Px);
+	const FFloatType Tx = INVARIANT_SUB(Lhs, Hx);
+	const FFloatType Py = INVARIANT_MUL(Rhs, Constant);
+	const FFloatType Hy = INVARIANT_ADD(INVARIANT_SUB(Rhs, Py), Py);
+	const FFloatType Ty = INVARIANT_SUB(Rhs, Hy);
+	const FFloatType P = INVARIANT_MUL(Hx, Hy);
+	const FFloatType Q = INVARIANT_ADD(INVARIANT_MUL(Hx, Ty), INVARIANT_MUL(Tx, Hy));
+	const FFloatType Z = INVARIANT_ADD(P, Q);
+	const FFloatType E = INVARIANT_ADD(INVARIANT_ADD(INVARIANT_SUB(P, Z), Q), INVARIANT_MUL(Tx, Ty));
 	return DFFastTwoSum(Z, E);
 }
 
@@ -272,7 +274,7 @@ FDFType DFTwoProduct(FFloatType Lhs, FFloatType Rhs)
 	return DFTwoProductNoFMA(Lhs, Rhs);
 #else
 	const FFloatType P = Lhs * Rhs;
-	const FFloatType E = INVARIANT(FMA(Lhs, Rhs, -P));
+	const FFloatType E = INVARIANT_FMA(Lhs, Rhs, -P);
 	return DFConstructor(P, E);
 #endif
 }
@@ -288,9 +290,9 @@ FDFType DFMultiply(FDFType Lhs, FDFType Rhs)
 	return DFPromote(Lhs.High * Rhs.High);
 #else
 	const FDFType P = DFTwoProduct(Lhs.High, Rhs.High);
-	FFloatType T = INVARIANT(Lhs.High * Rhs.Low);
-	T = INVARIANT(FMA(Lhs.Low, Rhs.High, T));
-	T = INVARIANT(T + P.Low);
+	FFloatType T = INVARIANT_MUL(Lhs.High, Rhs.Low);
+	T = INVARIANT_FMA(Lhs.Low, Rhs.High, T);
+	T = INVARIANT_ADD(T, P.Low);
 
 	return DFFastTwoSum(P.High, T);
 #endif
@@ -304,7 +306,7 @@ FDFType DFMultiply(FDFType Lhs, FFloatType Rhs)
 	return DFPromote(Lhs.High * Rhs);
 #else
 	FDFType P = DFTwoProduct(Lhs.High, Rhs);
-	FFloatType T = INVARIANT(FMA(Lhs.Low, Rhs, P.Low));
+	FFloatType T = INVARIANT_FMA(Lhs.Low, Rhs, P.Low);
 	return DFFastTwoSum(P.High, T);
 #endif
 }
@@ -317,11 +319,11 @@ FFloatType DFMultiplyDemote(FDFType Lhs, FDFType Rhs)
 	return Lhs.High * Rhs.High;
 #else
 	const FDFType P = DFTwoProduct(Lhs.High, Rhs.High);
-	FFloatType T = INVARIANT(Lhs.High * Rhs.Low);
-	T = INVARIANT(FMA(Lhs.Low, Rhs.High, T));
-	T = INVARIANT(T + P.Low);
-
-	return P.High + T;
+	FFloatType T = INVARIANT_MUL(Lhs.High, Rhs.Low);
+	T = INVARIANT_FMA(Lhs.Low, Rhs.High, T);
+	T = INVARIANT_ADD(T, P.Low);
+	const FFloatType S = INVARIANT_ADD(P.High, T);
+	return S;
 #endif
 }
 
@@ -331,8 +333,9 @@ FFloatType DFMultiplyDemote(FDFType Lhs, FFloatType Rhs)
 	return Lhs.High * Rhs;
 #else
 	FDFType P = DFTwoProduct(Lhs.High, Rhs);
-	P.Low = FMA(Lhs.Low, Rhs, P.Low);
-	return P.High + P.Low;
+	P.Low = INVARIANT_FMA(Lhs.Low, Rhs, P.Low);
+	const FFloatType S = INVARIANT_ADD(P.High, P.Low);
+	return S;
 #endif
 }
 
@@ -367,9 +370,9 @@ FDFType DFDivide(FDFType Lhs, FDFType Rhs)
 #elif !PLATFORM_SUPPORTS_FMA
 	return DFFastDivide(Lhs, Rhs);
 #else
-	const FFloatType Th = INVARIANT(1.0 / Rhs.High);
-	const FFloatType Rh = INVARIANT(FMA(-Rhs.High, Th, (FFloatType)1));
-	const FFloatType Rl = INVARIANT(-Rhs.Low * Th);
+	const FFloatType Th = INVARIANT_DIV(1.0, Rhs.High);
+	const FFloatType Rh = INVARIANT_FMA(-Rhs.High, Th, (FFloatType)1);
+	const FFloatType Rl = INVARIANT_MUL(-Rhs.Low, Th);
 	const FDFType E = DFFastTwoSum(Rh, Rl);
 	const FDFType D = DFMultiply(E, Th);
 	const FDFType M = DFAdd(D, Th);
@@ -383,12 +386,12 @@ FDFType DFDivide(FDFType Lhs, FFloatType Rhs)
 #if UE_DF_FORCE_FP32_OPS
 	return DFPromote(Lhs.High / Rhs);
 #else
-	const FFloatType Th = INVARIANT(Lhs.High / Rhs);
+	const FFloatType Th = INVARIANT_DIV(Lhs.High, Rhs);
 	const FDFType P = DFTwoProduct(Th, Rhs);
-	const FFloatType Dh = INVARIANT(Lhs.High - P.High);
-	const FFloatType Dt = INVARIANT(Dh - P.Low);
-	const FFloatType D = INVARIANT(Dt + Lhs.Low);
-	const FFloatType Tl = INVARIANT(D / Rhs);
+	const FFloatType Dh = INVARIANT_SUB(Lhs.High, P.High);
+	const FFloatType Dt = INVARIANT_SUB(Dh, P.Low);
+	const FFloatType D = INVARIANT_ADD(Dt, Lhs.Low);
+	const FFloatType Tl = INVARIANT_DIV(D, Rhs);
 	return DFFastTwoSum(Th, Tl);
 #endif
 }
@@ -420,12 +423,12 @@ FDFType DFFastDivide(FDFType Lhs, FDFType Rhs)
 #elif UE_DF_NO_FAST_MATH
 	return DFDivide(Lhs, Rhs);
 #else
-	const FFloatType Th = INVARIANT(Lhs.High / Rhs.High);
+	const FFloatType Th = INVARIANT_DIV(Lhs.High, Rhs.High);
 	const FDFType R = DFMultiply(Rhs, Th);
-	const FFloatType Ph = INVARIANT(Lhs.High - R.High);
-	const FFloatType Dl = INVARIANT(Lhs.Low - R.Low);
-	const FFloatType D = INVARIANT(Ph + Dl);
-	const FFloatType Tl = INVARIANT(D / Rhs.High);
+	const FFloatType Ph = INVARIANT_SUB(Lhs.High, R.High);
+	const FFloatType Dl = INVARIANT_SUB(Lhs.Low, R.Low);
+	const FFloatType D = INVARIANT_ADD(Ph, Dl);
+	const FFloatType Tl = INVARIANT_DIV(D, Rhs.High);
 	return DFFastTwoSum(Th, Tl);
 #endif
 }
@@ -437,12 +440,12 @@ FDFType DFFastDivide(FDFType Lhs, FFloatType Rhs)
 #elif UE_DF_NO_FAST_MATH && PLATFORM_SUPPORTS_FMA
 	return DFDivide(Lhs, Rhs);
 #else
-	const FFloatType Th = INVARIANT(Lhs.High / Rhs);
+	const FFloatType Th = INVARIANT_DIV(Lhs.High, Rhs);
 	const FDFType R = DFTwoProduct(Rhs, Th);
-	const FFloatType Ph = INVARIANT(Lhs.High - R.High);
-	const FFloatType Dl = INVARIANT(Lhs.Low - R.Low);
-	const FFloatType D = INVARIANT(Ph + Dl);
-	const FFloatType Tl = INVARIANT(D / Rhs);
+	const FFloatType Ph = INVARIANT_SUB(Lhs.High, R.High);
+	const FFloatType Dl = INVARIANT_SUB(Lhs.Low, R.Low);
+	const FFloatType D = INVARIANT_ADD(Ph, Dl);
+	const FFloatType Tl = INVARIANT_DIV(D, Rhs);
 	return DFFastTwoSum(Th, Tl);
 #endif
 }
@@ -457,12 +460,12 @@ FFloatType DFFastDivideDemote(FDFType Lhs, FDFType Rhs)
 #elif UE_DF_NO_FAST_MATH
 	return DFDemote(DFDivide(Lhs, Rhs));
 #else
-	const FFloatType Th = INVARIANT(Lhs.High / Rhs.High);
+	const FFloatType Th = INVARIANT_DIV(Lhs.High, Rhs.High);
 	const FDFType R = DFMultiply(Rhs, Th);
-	const FFloatType Ph = INVARIANT(Lhs.High - R.High);
-	const FFloatType Dl = INVARIANT(Lhs.Low - R.Low);
-	const FFloatType D = INVARIANT(Ph + Dl);
-	const FFloatType Tl = INVARIANT(D / Rhs.High);
+	const FFloatType Ph = INVARIANT_SUB(Lhs.High, R.High);
+	const FFloatType Dl = INVARIANT_SUB(Lhs.Low, R.Low);
+	const FFloatType D = INVARIANT_ADD(Ph, Dl);
+	const FFloatType Tl = INVARIANT_DIV(D, Rhs.High);
 	return Th + Tl;
 #endif
 }
@@ -474,12 +477,12 @@ FFloatType DFFastDivideDemote(FDFType Lhs, FFloatType Rhs)
 #elif UE_DF_NO_FAST_MATH
 	return DFDemote(DFDivide(Lhs, Rhs));
 #else
-	const FFloatType Th = INVARIANT(Lhs.High / Rhs);
+	const FFloatType Th = INVARIANT_DIV(Lhs.High, Rhs);
 	const FDFType R = DFTwoProduct(Rhs, Th);
-	const FFloatType Ph = INVARIANT(Lhs.High - R.High);
-	const FFloatType Dl = INVARIANT(Lhs.Low - R.Low);
-	const FFloatType D = INVARIANT(Ph + Dl);
-	const FFloatType Tl = INVARIANT(D / Rhs);
+	const FFloatType Ph = INVARIANT_SUB(Lhs.High, R.High);
+	const FFloatType Dl = INVARIANT_SUB(Lhs.Low, R.Low);
+	const FFloatType D = INVARIANT_ADD(Ph, Dl);
+	const FFloatType Tl = INVARIANT_DIV(D, Rhs);
 	return Th + Tl;
 #endif
 }
@@ -654,10 +657,10 @@ FDFType DFSqrt(FDFType V)
 	return DFPromote(sqrt(V.High));
 #else
 	const FFloatType Sh = sqrt(V.High);
-	const FFloatType P1 = INVARIANT(FMA(-Sh, Sh, V.High)); // Note: precision loss here if FMA is not supported
-	const FFloatType P2 = INVARIANT(V.Low + P1);
+	const FFloatType P1 = INVARIANT_FMA(-Sh, Sh, V.High); // Note: precision loss here if FMA is not supported
+	const FFloatType P2 = INVARIANT_ADD(V.Low, P1);
 
-	const FFloatType Sl = INVARIANT(P2 / (2.0 * Sh));
+	const FFloatType Sl = INVARIANT_DIV(P2, (2.0 * Sh));
 	return DFFastTwoSum(Sh, Sl);
 
 	// equivalent to (remark 3.5)
@@ -676,10 +679,10 @@ FFloatType DFSqrtDemote(FDFType V)
 	return sqrt(V.High);
 #else
 	const FFloatType Sh = sqrt(V.High);
-	const FFloatType P1 = INVARIANT(FMA(-Sh, Sh, V.High)); // Note: precision loss here if FMA is not supported
-	const FFloatType P2 = INVARIANT(V.Low + P1);
+	const FFloatType P1 = INVARIANT_FMA(-Sh, Sh, V.High); // Note: precision loss here if FMA is not supported
+	const FFloatType P2 = INVARIANT_ADD(V.Low, P1);
 
-	const FFloatType Sl = INVARIANT(P2 / (2.0 * Sh));
+	const FFloatType Sl = INVARIANT_DIV(P2, (2.0 * Sh));
 	return Sh + Sl;
 #endif
 }
diff --git a/Engine/Shaders/Private/LargeWorldCoordinates.ush b/Engine/Shaders/Private/LargeWorldCoordinates.ush
--- a/Engine/Shaders/Private/LargeWorldCoordinates.ush
+++ b/Engine/Shaders/Private/LargeWorldCoordinates.ush
@@ -521,7 +521,7 @@ float4x4 LWCHackToFloat(FLWCInverseMatrix v) { return (float4x4)0.0f; }
 		/* Values greater than 2^(24+floor(log2(TILE_SIZE-1))) are not converted correctly, because In.Low is assumed to be */	\
 		/* less than half a tile. These values by definition exceed WORLD_MAX, so this is not an issue. */						\
 		FFloatType Tile = INVARIANT(round(In.High / UE_LWC_RENDER_TILE_SIZE));													\
-		FFloatType Offset = INVARIANT(FMA(Tile, -UE_LWC_RENDER_TILE_SIZE, In.High) + In.Low);									\
+		FFloatType Offset = INVARIANT_ADD(INVARIANT_FMA(Tile, -UE_LWC_RENDER_TILE_SIZE, In.High), In.Low);						\
 																																\
 		return LWCConstructor(Tile, Offset);																					\
 	}	
diff --git a/Engine/Shaders/Public/Platform.ush b/Engine/Shaders/Public/Platform.ush
--- a/Engine/Shaders/Public/Platform.ush
+++ b/Engine/Shaders/Public/Platform.ush
@@ -719,6 +719,26 @@ precise float4 MakePrecise(in precise float4 v) { precise float4 pv = v; return
 #define INVARIANT(X) MakePrecise(X)
 #endif
 
+#ifndef INVARIANT_ADD
+#define INVARIANT_ADD(Lhs, Rhs) INVARIANT((Lhs) + (Rhs))
+#endif
+
+#ifndef INVARIANT_SUB
+#define INVARIANT_SUB(Lhs, Rhs) INVARIANT((Lhs) - (Rhs))
+#endif
+
+#ifndef INVARIANT_MUL
+#define INVARIANT_MUL(Lhs, Rhs) INVARIANT((Lhs) * (Rhs))
+#endif
+
+#ifndef INVARIANT_DIV
+#define INVARIANT_DIV(Lhs, Rhs) INVARIANT((Lhs) / (Rhs))
+#endif
+
+#ifndef INVARIANT_FMA
+#define INVARIANT_FMA(A, B, C) INVARIANT(FMA((A), (B), (C)))
+#endif
+
 #ifndef INVARIANT_OUTPUT
 #define INVARIANT_OUTPUT
 #endif
