Path Tracer: Use stochastic blending to handle opacity on glass materials

This simplifies the logic for solid/thin glass greatly and makes it easier to support transparent thin glass via the new SurfaceCoverage material output (introduced in 32048301).

There is a new shading model id for the path tracer for the solid glass case. This and the existing SHADINGMODELID_THIN_TRANSLUCENT now refer only to the glass lobes. SHADINGMODELID_DEFAULT_LIT can be used for the rest which removes some duplicated logic and makes the common cases much faster.

diff --git a/Engine/Shaders/Private/PathTracing/Material/PathTracingMaterialSampling.ush b/Engine/Shaders/Private/PathTracing/Material/PathTracingMaterialSampling.ush
--- a/Engine/Shaders/Private/PathTracing/Material/PathTracingMaterialSampling.ush
+++ b/Engine/Shaders/Private/PathTracing/Material/PathTracingMaterialSampling.ush
@@ -93,10 +93,13 @@ float3 EstimateMaterialAlbedo(FPathTracingPayload Payload)
 			// glass portion is a mixture of pure specular and tinted transmission (when refracting)
 			// or glass portion is just the specular highlight with transmission accounted for by transparency
 			const float F0 = F0RGBToF0(Payload.SpecularColor);
-			const float3 GlassAlbedo = Payload.HasRefraction() ? lerp(Payload.GetTransmittanceColor(), 1.0, F0) : F0;
-			// surface is same as default-lit
-			const float3 SurfAlbedo = Payload.DiffuseColor + Payload.SpecularColor;
-			return lerp(GlassAlbedo, SurfAlbedo, Payload.BSDFOpacity);
+			Albedo = Payload.HasRefraction() ? lerp(Payload.GetTransmittanceColor(), 1.0, F0) : F0;
+			break;
+		}
+		case SHADINGMODELID_SOLID_GLASS:
+		{
+			Albedo = 1.0; // all energy is either reflected or transmitted
+			break;
 		}
 		case SHADINGMODELID_CLOTH:
 			Albedo = (Payload.DiffuseColor + lerp(Payload.SpecularColor, Payload.GetClothColor(), Payload.GetClothAmount()));
@@ -109,12 +112,6 @@ float3 EstimateMaterialAlbedo(FPathTracingPayload Payload)
 	}
 	Albedo *= Payload.BSDFOpacity;
 
-#if !SIMPLIFIED_MATERIAL_SHADER && !PATHTRACING_SUBSTRATE_PAYLOAD
-	if (Payload.HasRefraction())
-	{
-		Albedo += 1 - Payload.BSDFOpacity; // glass albedo is 1.0
-	}
-#endif
 	return Albedo;
 }
 
@@ -218,16 +215,7 @@ void AccumulateAlbedoNormal(inout FPathTracingPayload Payload, float3 PathThroug
 		Normal += BrdfWeight * NormThroughput * Payload.WorldNormal;
 		break;
 	case SHADINGMODELID_DEFAULT_LIT:
-	{
-		if (Payload.HasRefraction())
-		{
-			const float GlassAlbedo = 1.0; // solid glass has unit albedo
-			Albedo += (1 - BrdfWeight) * RoughnessWeight * PathThroughput * GlassAlbedo;
-			Normal += (1 - BrdfWeight) * RoughnessWeight * NormThroughput * Payload.WorldNormal;
-		}
-		// fallthrough
-	}
-	case SHADINGMODELID_CLEAR_COAT:
+	case SHADINGMODELID_CLEAR_COAT: // TODO: can we be more precise with clearcoat?
 		Albedo += BrdfWeight * PathThroughput * (Payload.DiffuseColor + RoughnessWeight * Payload.SpecularColor);
 		Normal += BrdfWeight * NormThroughput * Payload.WorldNormal;
 		break;
@@ -239,10 +227,15 @@ void AccumulateAlbedoNormal(inout FPathTracingPayload Payload, float3 PathThroug
 	{
 		const float F0 = F0RGBToF0(Payload.SpecularColor);
 		const float3 GlassAlbedo = Payload.HasRefraction() ? lerp(Payload.GetTransmittanceColor(), 1.0, F0) : F0;
-		Albedo += (1 - BrdfWeight) * RoughnessWeight * PathThroughput * GlassAlbedo;
-		Normal += (1 - BrdfWeight) * RoughnessWeight * NormThroughput * Payload.WorldNormal;
-		Albedo += BrdfWeight * PathThroughput * (Payload.DiffuseColor + RoughnessWeight * Payload.SpecularColor);
-		Normal += BrdfWeight * NormThroughput * Payload.WorldNormal;
+		Albedo += BrdfWeight * RoughnessWeight * PathThroughput * GlassAlbedo;
+		Normal += BrdfWeight * RoughnessWeight * NormThroughput * Payload.WorldNormal;
+		break;
+	}
+	case SHADINGMODELID_SOLID_GLASS:
+	{
+		const float GlassAlbedo = 1.0; // solid glass has unit albedo
+		Albedo += BrdfWeight * RoughnessWeight * PathThroughput * GlassAlbedo;
+		Normal += BrdfWeight * RoughnessWeight * NormThroughput * Payload.WorldNormal;
 		break;
 	}
 	default:
@@ -307,7 +300,6 @@ float GetAverageRoughness(FPathTracingPayload Payload)
 	}
 	return AvgRoughness;
 #else
-	const float GlassPercent = Payload.HasRefraction() ? 1 - Payload.BSDFOpacity : 0.0;
 	float Roughness = Payload.Roughness;
 	switch (Payload.ShadingModelID)
 	{
@@ -320,9 +312,11 @@ float GetAverageRoughness(FPathTracingPayload Payload)
 			Roughness = max(DualRoughnessData.x, DualRoughnessData.y);
 			break;
 		}
+		case SHADINGMODELID_THIN_TRANSLUCENT:
+		case SHADINGMODELID_SOLID_GLASS:
+			return Roughness; // both cases are purely specular, no diffuse
 	}
 	float AvgRoughness = lerp(1.0 /* diffuse lobe */, Roughness, Payload.Metallic);
-	AvgRoughness = lerp(AvgRoughness, Payload.Roughness, GlassPercent);
 	return AvgRoughness;
 #endif
 }
@@ -577,18 +571,13 @@ FMaterialSample SampleMaterial(
 			return SubstrateClearCoat_SampleMaterial(V_World, Payload, RandSample);
 #else
 		case SHADINGMODELID_DEFAULT_LIT:
-		{
-			if (Payload.HasRefraction() && Payload.BSDFOpacity < 1)
-			{
-				return RoughGlass_SampleMaterial(V_World, Payload, RandSample);
-			}
-			// fallthrough
-		}
 		case SHADINGMODELID_SUBSURFACE:
 		case SHADINGMODELID_PREINTEGRATED_SKIN:
 			return DefaultLit_SampleMaterial(V_World, Payload, RandSample);
 		case SHADINGMODELID_THIN_TRANSLUCENT:
 			return RoughThinGlass_SampleMaterial(V_World, Payload, RandSample);
+		case SHADINGMODELID_SOLID_GLASS:
+			return RoughGlass_SampleMaterial(V_World, Payload, RandSample);
 		case SHADINGMODELID_SUBSURFACE_PROFILE:
 			return SubsurfaceProfile_SampleMaterial(V_World, Payload, RandSample);
 		case SHADINGMODELID_CLOTH:
@@ -638,18 +627,13 @@ FMaterialEval EvalMaterial(
 			return SubstrateClearCoat_EvalMaterial(V_World, L_World, Payload, DiffuseSpecularScale);
 #else
 		case SHADINGMODELID_DEFAULT_LIT:
-		{
-			if (Payload.HasRefraction() && Payload.BSDFOpacity < 1)
-			{
-				return RoughGlass_EvalMaterial(V_World, L_World, Payload, DiffuseSpecularScale);
-			}
-			// fallthrough
-		}
 		case SHADINGMODELID_SUBSURFACE:
 		case SHADINGMODELID_PREINTEGRATED_SKIN:
 			return DefaultLit_EvalMaterial(V_World, L_World, Payload, DiffuseSpecularScale);
 		case SHADINGMODELID_THIN_TRANSLUCENT:
 			return RoughThinGlass_EvalMaterial(V_World, L_World, Payload, DiffuseSpecularScale);
+		case SHADINGMODELID_SOLID_GLASS:
+			return RoughGlass_EvalMaterial(V_World, L_World, Payload, DiffuseSpecularScale);
 		case SHADINGMODELID_SUBSURFACE_PROFILE:
 			return SubsurfaceProfile_EvalMaterial(V_World, L_World, Payload, DiffuseSpecularScale);
 		case SHADINGMODELID_CLOTH:
diff --git a/Engine/Shaders/Private/PathTracing/Material/PathTracingSolidGlass.ush b/Engine/Shaders/Private/PathTracing/Material/PathTracingSolidGlass.ush
--- a/Engine/Shaders/Private/PathTracing/Material/PathTracingSolidGlass.ush
+++ b/Engine/Shaders/Private/PathTracing/Material/PathTracingSolidGlass.ush
@@ -17,12 +17,6 @@ struct FRoughGlassData
 	float2 Alpha;
 	float3 V;
 
-	FBxDFEnergyTermsRGB Spec;
-	float3 DiffWeight;
-
-	float2 LobeCdf;
-	float3 LobePdf;
-
 	// glass lobe data
 	float Eta;
 	float F0;
@@ -38,16 +32,8 @@ FRoughGlassData PrepareRoughGlassData(FPathTracingPayload Payload, float3 V_Worl
 	const float NoV = saturate(Data.V.z);
 
 
-	Data.Spec = ComputeGGXSpecEnergyTermsRGB(Payload.Roughness, NoV, Payload.SpecularColor);
-	Data.DiffWeight = Payload.DiffuseColor * (1 - Data.Spec.E);
-
-	// Probability of picking diffuse vs spec vs glass
-	Data.LobeCdf = LobeSelectionCdf(Payload.BSDFOpacity * Data.DiffWeight, Payload.BSDFOpacity * Data.Spec.E, 1.0 - Payload.BSDFOpacity);
-	Data.LobePdf = LobeSelectionPdf(Data.LobeCdf);
-
-
 	// NOTE: IsFrontFace() determines polygon orientation, because the normal is always flipped towards in the incoming ray
-	// #dxr_todo: Maintain a refraction stack on the path tracing payload
+	// TODO: Maintain a refraction stack on the path tracing payload
 	Data.Eta = Payload.IsFrontFace() ? Payload.Ior : rcp(Payload.Ior);
 	Data.F0 = F0RGBToF0(Payload.SpecularColor);
 
@@ -71,87 +57,18 @@ FMaterialSample RoughGlass_SampleMaterial(
 		// invalid grazing angle
 		return NullMaterialSample();
 	}
+	const float3 H = ImportanceSampleVisibleGGX(RandSample.xy, Data.Alpha, V, false).xyz;
+	// Glass lobe (reflection and refraction)
+	float3 L = 0;
+	float F = 0;
 
-	// Randomly choose to sample diffuse or specular
-	if (RandSample.x < Data.LobeCdf.x)
-	{
-		RandSample.x = RescaleRandomNumber(RandSample.x, 0.0, Data.LobeCdf.x);
-
-		// Lambert
-		// TODO: evaluate CosineSampleHemisphereConcentric
-		const float4 Result = CosineSampleHemisphere(RandSample.xy);
-
-		const float3 L = Result.xyz;
-		const float3 H = normalize(L + V);
-
-		// transform to world space
-		const float3 L_World = normalize(mul(L, Data.Basis));
-
-		const float NoL = saturate(L.z);
-		const float VoH = saturate(dot(V, H));
-		const float NoH = saturate(H.z);
-
-		const float3 Diffuse = GetPathTracingDiffuseModel(Data.DiffWeight, Payload.Roughness, NoV, NoL, VoH, NoH);
-		const float3 DiffWeight = (Payload.BSDFOpacity / Data.LobePdf.x) * Diffuse * ShadowTerminatorTerm(L_World, Payload.WorldNormal, Payload.WorldSmoothNormal);
-		const float DiffPdf = Data.LobePdf.x * NoL / PI;
-
-		return CreateMaterialSample(L_World, DiffWeight, DiffPdf, 1.0, 1.0, PATHTRACER_SCATTER_DIFFUSE);
-	}
-	else
-	{
-		const bool bDoSpec = RandSample.x < Data.LobeCdf.y;
-
-		if (bDoSpec)
-		{
-			RandSample.x = RescaleRandomNumber(RandSample.x, Data.LobeCdf.x, Data.LobeCdf.y);
-		}
-		else
-		{
-			RandSample.x = RescaleRandomNumber(RandSample.x, Data.LobeCdf.y, 1.0);
-		}
-
-
-		const float3 H = ImportanceSampleVisibleGGX(RandSample.xy, Data.Alpha, V, bDoSpec).xyz;
-
-		if (bDoSpec)
-		{
-			const float3 L = reflect(-V, H);
-			if (L.z <= 0)
-			{
-				// invalid output direction, exit early
-				return NullMaterialSample();
-			}
-
-			// transform to world space
-			const float3 L_World = normalize(mul(L, Data.Basis));
-
-			const float NoL = saturate(L.z);
-			const float VoH = saturate(dot(V, H));
-			const float NoH = saturate(H.z);
-
-			// Specular lobe
-			const float2 GGXResult = GGXEvalReflection(L, V, H, Data.Alpha, true);
-			const float3 F = F_Schlick(Payload.SpecularColor, VoH);
-			const float3 SpecWeight = (Payload.BSDFOpacity / Data.LobePdf.y) * F * GGXResult.x * Data.Spec.W;
-			const float SpecPdf = GGXResult.y * Data.LobePdf.y;
-
-			return CreateMaterialSample(L_World, SpecWeight, SpecPdf, 1.0, Payload.Roughness, PATHTRACER_SCATTER_SPECULAR);
-		}
-		else
-		{
-			// Glass lobe (reflection and refraction)
-			float3 L = 0;
-			float F = 0;
-
-			const bool bRefract = SampleRefraction(-V, H, Data.Eta, Data.F0, RandSample.z, L, F);
-			// transform to world space
-			const float3 L_World = normalize(mul(L, Data.Basis));
+	const bool bRefract = SampleRefraction(-V, H, Data.Eta, Data.F0, RandSample.z, L, F);
+	// transform to world space
+	const float3 L_World = normalize(mul(L, Data.Basis));
 
-			const float2 GGXResult = bRefract ? GGXEvalRefraction(L, V, H, Data.Alpha, Data.Eta) : GGXEvalReflection(L, V, H, Data.Alpha, false);
+	const float2 GGXResult = bRefract ? GGXEvalRefraction(L, V, H, Data.Alpha, Data.Eta) : GGXEvalReflection(L, V, H, Data.Alpha, false);
 
-			return CreateMaterialSample(L_World, (1.0 - Payload.BSDFOpacity) * GGXResult.x / (Data.EGlass * Data.LobePdf.z), F * GGXResult.y * Data.LobePdf.z, bRefract ? -1.0 : 1.0, Payload.Roughness, bRefract ? PATHTRACER_SCATTER_REFRACT : PATHTRACER_SCATTER_SPECULAR);
-		}
-	}
+	return CreateMaterialSample(L_World, Payload.BSDFOpacity * GGXResult.x / Data.EGlass, F * GGXResult.y, bRefract ? -1.0 : 1.0, Payload.Roughness, bRefract ? PATHTRACER_SCATTER_REFRACT : PATHTRACER_SCATTER_SPECULAR);
 }
 
 FMaterialEval RoughGlass_EvalMaterial(
@@ -185,29 +102,14 @@ FMaterialEval RoughGlass_EvalMaterial(
 		const float VoH = saturate(dot(V, H));
 		const float NoH = saturate(H.z);
 
-		// Diffuse Lobe
-		const float3 Diffuse = GetPathTracingDiffuseModel(Data.DiffWeight, Payload.Roughness, NoV, NoL, VoH, NoH);
-		Result.AddLobeWithMIS(Payload.BSDFOpacity * Diffuse * ShadowTerminatorTerm(L_World, Payload.WorldNormal, Payload.WorldSmoothNormal) * DiffuseSpecularScale.x, NoL / PI, Data.LobePdf.x);
-
-		// Specular lobe
-		{
-			const float2 GGXResult = GGXEvalReflection(L, V, H, Data.Alpha, true);
-			const float3 Fs = F_Schlick(Payload.SpecularColor, VoH);
-			const float3 SpecWeight = Payload.BSDFOpacity * GGXResult.x * Data.Spec.W * Fs;
-			const float SpecPdf = GGXResult.y;
-			Result.AddLobeWithMIS(SpecWeight * DiffuseSpecularScale.y, SpecPdf, Data.LobePdf.y);
-		}
-
 		// Glass lobe
-		{
-			const float2 GGXResult = GGXEvalReflection(L, V, H, Data.Alpha, false);
-			const float Fg = FresnelReflectance(VoH, Data.Eta, Data.F0);
-			const float3 GlassWeight = GGXResult.x * (1 - Payload.BSDFOpacity) / Data.EGlass;
-			const float GlassPdf = GGXResult.y * Fg;
-			Result.AddLobeWithMIS(GlassWeight * DiffuseSpecularScale.y, GlassPdf, Data.LobePdf.z);
-		}
+		const float2 GGXResult = GGXEvalReflection(L, V, H, Data.Alpha, false);
+		const float Fg = FresnelReflectance(VoH, Data.Eta, Data.F0);
+		const float GlassWeight = GGXResult.x * Payload.BSDFOpacity / Data.EGlass;
+		const float GlassPdf = GGXResult.y * Fg;
+		Result.AddLobeWithMIS(GlassWeight * DiffuseSpecularScale.y, GlassPdf, 1.0);
 	}
-	else if (L.z < 0 && Data.LobePdf.z > 0)
+	else if (L.z < 0)
 	{
 		// refraction side
 		const float NoL = saturate(-L.z);
@@ -220,10 +122,10 @@ FMaterialEval RoughGlass_EvalMaterial(
 		{
 			const float2 GGXResult = GGXEvalRefraction(L, V, Ht, Data.Alpha, Data.Eta);
 
-			const float3 GlassWeight = GGXResult.x * (1 - Payload.BSDFOpacity) / Data.EGlass;
+			const float GlassWeight = GGXResult.x * Payload.BSDFOpacity / Data.EGlass;
 			const float GlassPdf = GGXResult.y * Fg;
 
-			Result.AddLobeWithMIS(GlassWeight * DiffuseSpecularScale.y, GlassPdf, Data.LobePdf.z);
+			Result.AddLobeWithMIS(GlassWeight * DiffuseSpecularScale.y, GlassPdf, 1.0);
 		}
 	}
 	return Result;
diff --git a/Engine/Shaders/Private/PathTracing/Material/PathTracingThinGlass.ush b/Engine/Shaders/Private/PathTracing/Material/PathTracingThinGlass.ush
--- a/Engine/Shaders/Private/PathTracing/Material/PathTracingThinGlass.ush
+++ b/Engine/Shaders/Private/PathTracing/Material/PathTracingThinGlass.ush
@@ -18,14 +18,11 @@ struct FRoughThinGlassData
 	float2 AlphaR, AlphaT;
 	float3 V;
 
-	FBxDFEnergyTermsRGB Spec;
 	FBxDFEnergyTermsA SpecR;
 	FBxDFEnergyTermsA SpecT;
 	float NudgeE;
-	float3 DiffWeight;
 
-	float3 LobeCdf;
-	float4 LobePdf;
+	float LobeProb; // R or T
 };
 
 FRoughThinGlassData PrepareRoughThinGlassData(FPathTracingPayload Payload, float3 V_World)
@@ -44,23 +41,16 @@ FRoughThinGlassData PrepareRoughThinGlassData(FPathTracingPayload Payload, float
 	Data.V = mul(Data.Basis, V_World);
 	const float NoV = saturate(Data.V.z);
 
-	Data.Spec = ComputeGGXSpecEnergyTermsRGB(Payload.Roughness, NoV, Payload.SpecularColor);
 	// because the roughnesses may not match, we need to measure the response for each individually, taking into account the reflected and transmitted fresnel respectively
 	// NOTE: 1.0 - F_Schlick(NoV, F0, F90) == F_Schlick(NoV, 1.0 - F0, 1.0 - F90)
 	Data.SpecR = ComputeGGXSpecEnergyTermsA(RoughnessR, NoV, Data.F0, Data.F90);
 	Data.SpecT = ComputeGGXSpecEnergyTermsA(RoughnessT, NoV, 1.0 - Data.F0, 1.0 - Data.F90);
 	Data.NudgeE = rcp(Data.SpecR.E + Data.SpecT.E);
-	Data.DiffWeight = (1 - Data.Spec.E) * Payload.DiffuseColor;
 
-	// Figure out probability of the 4 lobes: diffuse,spec,glass reflection, glass transmission (if enabled)
+	// Figure out probability of glass reflection vs glass transmission
 	// Approximate the transmission as just a plain tint by slab color because we can't depend on the half-vector for lobe selection
-	Data.LobeCdf = LobeSelectionCdf(
-		Payload.BSDFOpacity * Data.DiffWeight,
-		Payload.BSDFOpacity * Data.Spec.E,
-		(1 - Payload.BSDFOpacity) * Data.NudgeE * Data.SpecR.E,
-		(1 - Payload.BSDFOpacity) * Data.NudgeE * (Payload.HasRefraction() ? Data.SpecT.E * Payload.GetTransmittanceColor() : 0.0)
-	);
-	Data.LobePdf = LobeSelectionPdf(Data.LobeCdf);
+	// If we don't have any refraction, we will always pick the reflection lobe
+	Data.LobeProb = Payload.HasRefraction() ? LobeSelectionProb(Data.NudgeE * Data.SpecR.E, Data.NudgeE * Data.SpecT.E * Payload.GetTransmittanceColor()) : 1.0;
 
 	return Data;
 }
@@ -93,37 +83,21 @@ FMaterialEval RoughThinGlass_EvalMaterial(
 	const float VoH = saturate(dot(V, H));
 	const float NoH = saturate(H.z);
 
-
 	const FThinSlabWeights SlabResult = ComputeThinSlabWeights(Payload.GetTransmittanceColor(), VoH, Payload.Ior, Data.F0);
 	const float2 GGXResult = GGXEvalReflection(L, V, H, bIsReflection ? Data.AlphaR : Data.AlphaT);
-	FMaterialEval Result = NullMaterialEval();
+	const float GlassPdf = GGXResult.y;
+	float3 GlassWeight = Payload.BSDFOpacity * Data.NudgeE * GGXResult.x * DiffuseSpecularScale.y;
 	if (bIsReflection)
 	{
-		// Diffuse Lobe
-		const float3 Diffuse = GetPathTracingDiffuseModel(Data.DiffWeight, Payload.Roughness, NoV, NoL, VoH, NoH);
-		const float3 DiffWeight = Payload.BSDFOpacity * Diffuse * ShadowTerminatorTerm(L_World, Payload.WorldNormal, Payload.WorldSmoothNormal);
-		const float DiffPdf = NoL / PI;
-		Result.AddLobeWithMIS(DiffWeight * DiffuseSpecularScale.x, DiffPdf, Data.LobePdf.x);
-
-		// Specular lobe
-		const float2 GGXResult = GGXEvalReflection(L, V, H, Data.AlphaR);
-		const float3 F = F_Schlick(Payload.SpecularColor, VoH);
-		const float3 SpecWeight = Payload.BSDFOpacity * F * GGXResult.x * Data.Spec.W;
-		const float SpecPdf = GGXResult.y;
-		Result.AddLobeWithMIS(SpecWeight * DiffuseSpecularScale.y, SpecPdf, Data.LobePdf.y);
-
-
-		// Transmission reflection
-		const float3 GlassWeight = ((1 - Payload.BSDFOpacity) * Data.NudgeE * GGXResult.x * Data.SpecR.W) * SlabResult.Reflected;
-		const float GlassPdf = GGXResult.y;
-		Result.AddLobeWithMIS(GlassWeight * DiffuseSpecularScale.y, GlassPdf, Data.LobePdf.z);
+		GlassWeight *= Data.SpecR.W * SlabResult.Reflected;
 	}
 	else
 	{
-		const float3 GlassWeight = ((1 - Payload.BSDFOpacity) * Data.NudgeE * GGXResult.x * Data.SpecT.W) * SlabResult.Transmitted;
-		const float GlassPdf = GGXResult.y;
-		Result.AddLobeWithMIS(GlassWeight * DiffuseSpecularScale.y, GlassPdf, Data.LobePdf.w);
+		GlassWeight *= Data.SpecT.W * SlabResult.Transmitted;
 	}
+	FMaterialEval Result = NullMaterialEval();
+	// Transmission lobe will be automatically ignored when there is no refraction because we made LobeProbe==1 in that case
+	Result.AddLobeWithMIS(GlassWeight, GlassPdf, bIsReflection ? Data.LobeProb : 1.0 - Data.LobeProb);
 	return Result;
 }
 
@@ -144,100 +118,44 @@ FMaterialSample RoughThinGlass_SampleMaterial(
 		return NullMaterialSample();
 	}
 
-	if (RandSample.x < Data.LobeCdf.x)
+	// Specular/Glass lobes
+	const bool bIsReflection = RandSample.x < Data.LobeProb;
+	if (bIsReflection)
 	{
-		RandSample.x = RescaleRandomNumber(RandSample.x, 0.0, Data.LobeCdf.x);
-		// Diffuse lobe
-		const float3 L = CosineSampleHemisphere(RandSample.xy).xyz;
-		const float3 H = normalize(L + V);
-
-		// transform to world space
-		const float3 L_World = normalize(mul(L, Data.Basis));
-
-		const float NoL = saturate(L.z);
-		const float VoH = saturate(dot(V, H));
-		const float NoH = saturate(H.z);
-
-		// Diffuse Lobe
-		const float3 Diffuse = GetPathTracingDiffuseModel(Data.DiffWeight, Payload.Roughness, NoV, NoL, VoH, NoH);
-		const float3 DiffWeight = Payload.BSDFOpacity * Diffuse * ShadowTerminatorTerm(L_World, Payload.WorldNormal, Payload.WorldSmoothNormal);
-		const float DiffPdf = NoL / PI;
-
-		FMaterialSample Result = CreateMaterialSample(L_World, 0.0, 0.0, 1.0, 1.0, PATHTRACER_SCATTER_DIFFUSE);
-		Result.AddLobeWithMIS(DiffWeight, DiffPdf, Data.LobePdf.x);
-
-		const float SpecPdf = GGXEvalReflection(L, V, H, Data.AlphaR).y;
-		Result.Pdf += Data.LobePdf.y * SpecPdf;
-		Result.Pdf += Data.LobePdf.z * SpecPdf;
-		return Result;
+		RandSample.x = RescaleRandomNumber(RandSample.x, 0, Data.LobeProb);
 	}
 	else
 	{
-		// Specular/Glass lobes
-		const bool bIsReflection = RandSample.x < Data.LobeCdf.z;
-		if (RandSample.x < Data.LobeCdf.y)
-		{
-			RandSample.x = RescaleRandomNumber(RandSample.x, Data.LobeCdf.x, Data.LobeCdf.y);
-		}
-		else if (RandSample.x < Data.LobeCdf.z)
-		{
-			RandSample.x = RescaleRandomNumber(RandSample.x, Data.LobeCdf.y, Data.LobeCdf.z);
-		}
-		else
-		{
-			RandSample.x = RescaleRandomNumber(RandSample.x, Data.LobeCdf.z, 1.0);
-		}
-		const float3 H = ImportanceSampleVisibleGGX(RandSample.xy, bIsReflection ? Data.AlphaR : Data.AlphaT, V).xyz;
-		const float NoH = saturate(H.z);
-		const float VoH = saturate(dot(V, H));
-
-		const float3 L = reflect(-V, H);
-		if (L.z <= 0)
-		{
-			// invalid output direction, exit early
-			return NullMaterialSample();
-		}
-
-		const float NoL = saturate(L.z);
-
-		FThinSlabWeights SlabResult = ComputeThinSlabWeights(Payload.GetTransmittanceColor(), VoH, Payload.Ior, Data.F0);
-
-		const float2 GGXResult = GGXEvalReflection(L, V, H, bIsReflection ? Data.AlphaR : Data.AlphaT);
-
-		FMaterialSample Result = NullMaterialSample();
-		Result.Roughness = Payload.Roughness; // use common roughness for all lobes, even though transmission is squeezed
-		if (bIsReflection)
-		{
-			// Specular lobe
-			const float3 F = F_Schlick(Payload.SpecularColor, VoH);
-			const float3 SpecWeight = Payload.BSDFOpacity * F * GGXResult.x * Data.Spec.W;
-			const float SpecPdf = GGXResult.y;
-			Result.AddLobeWithMIS(SpecWeight, SpecPdf, Data.LobePdf.y);
-
-
-			// Transmission reflection
-			const float3 GlassWeight = ((1 - Payload.BSDFOpacity) * Data.NudgeE * GGXResult.x * Data.SpecR.W) * SlabResult.Reflected;
-			const float GlassPdf = GGXResult.y;
-			Result.AddLobeWithMIS(GlassWeight, GlassPdf, Data.LobePdf.z);
-
-			Result.Pdf += Data.LobePdf.x * NoL / PI;
-
-			Result.Direction = mul(L, Data.Basis);
-			Result.PositionBiasSign = 1.0;
-			Result.ScatterType = PATHTRACER_SCATTER_SPECULAR;
-		}
-		else
-		{
-			// transmitted ray
-			const float3 GlassWeight = ((1 - Payload.BSDFOpacity) * Data.NudgeE * GGXResult.x * Data.SpecT.W) * SlabResult.Transmitted;
-			const float GlassPdf = GGXResult.y;
-			Result.AddLobeWithMIS(GlassWeight, GlassPdf, Data.LobePdf.w);
-
-			Result.Direction = mul(float3(L.xy, -L.z), Data.Basis); // flip reflection to other side
-			Result.PositionBiasSign = -1.0;
-			Result.ScatterType = PATHTRACER_SCATTER_REFRACT;
-		}
-		Result.Direction = normalize(Result.Direction);
-		return Result;
+		RandSample.x = RescaleRandomNumber(RandSample.x, Data.LobeProb, 1.0);
+	}
+	const float3 H = ImportanceSampleVisibleGGX(RandSample.xy, bIsReflection ? Data.AlphaR : Data.AlphaT, V).xyz;
+	const float NoH = saturate(H.z);
+	const float VoH = saturate(dot(V, H));
+
+	const float3 L = reflect(-V, H);
+	if (L.z <= 0)
+	{
+		// invalid output direction, exit early
+		return NullMaterialSample();
 	}
+
+	const float NoL = saturate(L.z);
+
+	FThinSlabWeights SlabResult = ComputeThinSlabWeights(Payload.GetTransmittanceColor(), VoH, Payload.Ior, Data.F0);
+
+	const float2 GGXResult = GGXEvalReflection(L, V, H, bIsReflection ? Data.AlphaR : Data.AlphaT);
+
+	FMaterialSample Result = NullMaterialSample();
+	Result.Roughness = Payload.Roughness; // use common roughness for all lobes, even though transmission is squeezed
+
+	const float3 GlassWeight = Payload.BSDFOpacity * Data.NudgeE * GGXResult.x * (bIsReflection ? Data.SpecR.W * SlabResult.Reflected : Data.SpecT.W * SlabResult.Transmitted);
+	const float GlassPdf = GGXResult.y;
+
+	Result.AddLobeWithMIS(GlassWeight, GlassPdf, bIsReflection ? Data.LobeProb : 1.0 - Data.LobeProb);
+
+	Result.PositionBiasSign = bIsReflection ? 1.0 : -1.0;
+	Result.ScatterType = bIsReflection ? PATHTRACER_SCATTER_SPECULAR : PATHTRACER_SCATTER_REFRACT;
+	Result.Direction = mul(float3(L.xy, Result.PositionBiasSign * L.z), Data.Basis); // flip reflection to other side for refraction
+	Result.Direction = normalize(Result.Direction);
+	return Result;
 }
diff --git a/Engine/Shaders/Private/PathTracing/PathTracingCommon.ush b/Engine/Shaders/Private/PathTracing/PathTracingCommon.ush
--- a/Engine/Shaders/Private/PathTracing/PathTracingCommon.ush
+++ b/Engine/Shaders/Private/PathTracing/PathTracingCommon.ush
@@ -33,7 +33,11 @@
 #define PATH_TRACING_PAYLOAD_NORMAL_BITS  32
 
 // shading model extensions beyond the default count
-#define SHADINGMODELID_MEDIUM  (SHADINGMODELID_NUM + 1)
+#define SHADINGMODELID_MEDIUM  			(SHADINGMODELID_NUM + 1)
+#define SHADINGMODELID_SOLID_GLASS 		(SHADINGMODELID_NUM + 2)
+
+HLSL_STATIC_ASSERT(SHADINGMODELID_MEDIUM      <= SHADINGMODELID_MASK, "");
+HLSL_STATIC_ASSERT(SHADINGMODELID_SOLID_GLASS <= SHADINGMODELID_MASK, "");
 
 #define PATH_TRACING_ABSORPTION_SCALE  0.01        // inverse distance at which we hit the transmission color (Beer's law)
 
@@ -357,7 +361,6 @@ struct FPackedPathTracingPayload : FMinimalPayload
 		PackedData[13] |= f32tof16(DBufferC.w) << 16;
 	}
 	
-#if PATHTRACING_SUBSTRATE_PAYLOAD
 	void SetStochasticSlabRand(float SlabRand)
 	{
 		PackedData[14] = asuint(SlabRand);
@@ -367,7 +370,6 @@ struct FPackedPathTracingPayload : FMinimalPayload
 	{
 		return asfloat(PackedData[14]);
 	}
-#endif
 
 #endif // IS_PAYLOAD_ENABLED(RT_PAYLOAD_TYPE_GPULIGHTMASS)
 };
@@ -603,7 +605,7 @@ struct FPathTracingPayload : FMinimalPayload
 #if PATHTRACING_SUBSTRATE_PAYLOAD
 		return ShadingModelID == SHADINGMODELID_SUBSTRATE && HasRefraction();
 #else
-		return ShadingModelID == SHADINGMODELID_DEFAULT_LIT && HasRefraction() && BSDFOpacity < 1;
+		return ShadingModelID == SHADINGMODELID_SOLID_GLASS;
 #endif
 	}
 
@@ -612,7 +614,7 @@ struct FPathTracingPayload : FMinimalPayload
 #if PATHTRACING_SUBSTRATE_PAYLOAD
 		return ShadingModelID == SHADINGMODELID_THIN_TRANSLUCENT && HasRefraction();
 #else
-		return ShadingModelID == SHADINGMODELID_THIN_TRANSLUCENT && HasRefraction() && BSDFOpacity < 1;
+		return ShadingModelID == SHADINGMODELID_THIN_TRANSLUCENT && HasRefraction();
 #endif
 	}
 
diff --git a/Engine/Shaders/Private/PathTracing/PathTracingCore.ush b/Engine/Shaders/Private/PathTracing/PathTracingCore.ush
--- a/Engine/Shaders/Private/PathTracing/PathTracingCore.ush
+++ b/Engine/Shaders/Private/PathTracing/PathTracingCore.ush
@@ -455,7 +455,7 @@ FPackedPathTracingPayload TraceTransparentRay(inout FPathState PathState, int Bo
 
 
 	const uint MissShaderIndex = 0;
-	float2 RISRandSample = RandomSequence_GenerateSample2D(PathState.RandSequence);
+	float3 RISRandSample = RandomSequence_GenerateSample3D(PathState.RandSequence);
 	FRISContext HitSample = InitRISContext(RISRandSample.x);
 	FRISContext VolSample = InitRISContext(RISRandSample.y);
 
@@ -473,9 +473,7 @@ FPackedPathTracingPayload TraceTransparentRay(inout FPathState PathState, int Bo
 		PackedPayload.SetDBufferA(float4(0, 0, 0, 1));
 		PackedPayload.SetDBufferB(float4(0, 0, 0, 1));
 		PackedPayload.SetDBufferC(float4(0, 0, 0, 1));
-#if PATHTRACING_SUBSTRATE_PAYLOAD
-		PackedPayload.SetStochasticSlabRand(RandomSequence_GenerateSample1D(PathState.RandSequence));
-#endif
+		PackedPayload.SetStochasticSlabRand(RISRandSample.z);
 		// Trace the ray
 		TraceRay(
 			TLAS,
@@ -669,6 +667,7 @@ FPackedPathTracingPayload TraceTransparentRay(inout FPathState PathState, int Bo
 					PackedPayload.SetDBufferA(DecalPayload.GetBaseColor());
 					PackedPayload.SetDBufferB(DecalPayload.GetWorldNormal());
 					PackedPayload.SetDBufferC(DecalPayload.GetMetallicSpecularRoughness());
+					PackedPayload.SetStochasticSlabRand(RISRandSample.z); // Must use same random value to get a consistent result to the first call
 					TraceRay(
 						TLAS,
 						RayFlags,
diff --git a/Engine/Shaders/Private/PathTracing/PathTracingMaterialHitShader.usf b/Engine/Shaders/Private/PathTracing/PathTracingMaterialHitShader.usf
--- a/Engine/Shaders/Private/PathTracing/PathTracingMaterialHitShader.usf
+++ b/Engine/Shaders/Private/PathTracing/PathTracingMaterialHitShader.usf
@@ -924,14 +924,24 @@ RAY_TRACING_ENTRY_CLOSEST_HIT(PathTracingMaterialCHS,
 		const float3 TransmittanceColor = GetThinTranslucentMaterialOutput0(MaterialParameters);
 		Payload.SetTransmittanceColor(TransmittanceColor);
 
-		const float SurfaceCoverage = GetThinTranslucentMaterialOutput1(MaterialParameters);
-		Payload.BSDFOpacity *= SurfaceCoverage;
-
-		Payload.Ior = Ior;
+		// The material is a blend of glass and default_lit response, model it stochastically
+		if (Payload.BSDFOpacity < PackedPayload.GetStochasticSlabRand())
+		{
+			// Pick glass lobe - shade as glass (leave SHADINGMODELID_THIN_TRANSLUCENT on)
+			Payload.Ior = Ior;
+		}
+		else
+		{
+			// Pick solid lobe -- toggle back to default_lit response
+			// Ior stays 0.0 since this mode does not include any refraction
+			Payload.ShadingModelID = SHADINGMODELID_DEFAULT_LIT;
+		}
+		
+		float SurfaceCoverage = GetThinTranslucentMaterialOutput1(MaterialParameters);
 		if (Ior > 0.0)
 		{
-			// In this mode we get rough refraction, so don't apply any transparency
-			Payload.TransparencyColor = 0.0;
+			// In this mode we get rough refraction, so any transparency is only due to coverage (affecting both solid and glass lobes)
+			Payload.TransparencyColor = 1 - SurfaceCoverage;
 		}
 		else
 		{
@@ -943,11 +953,15 @@ RAY_TRACING_ENTRY_CLOSEST_HIT(PathTracingMaterialCHS,
 			const float  Metallic = Payload.Metallic;
 			const float  Specular = Payload.Specular;
 			const float  F0 = F0RGBToF0(ComputeF0(Specular, BaseColor, Metallic));
+			// This is the portion of transparency due to glass. We always compute it because we don't need to model the amount
+			// of transparency stochastically.
+			/// BSDFOpacity here still represents the blend between solid and glass behaviors.
 			Payload.TransparencyColor = (1 - Payload.BSDFOpacity) * ComputeThinSlabWeights(TransmittanceColor, VoN, 0.0, F0).Transmitted;
-
-			// Fade out the TransparencyColor as a function of the SurfaceCoverage.
-			Payload.TransparencyColor = lerp(1.0.xxx, Payload.TransparencyColor, SurfaceCoverage.xxx);
+			// Now account for overall coverage (allows the combination of solid/glass to become transparent)
+			Payload.TransparencyColor = lerp(1.0.xxx, Payload.TransparencyColor, SurfaceCoverage);
 		}
+		// Now we can replace the BSDFOpacity with coverage
+		Payload.BSDFOpacity = SurfaceCoverage;
 	}
 	else
 #endif
@@ -1042,11 +1056,23 @@ RAY_TRACING_ENTRY_CLOSEST_HIT(PathTracingMaterialCHS,
 		// Make the glass totally clear if no custom medium is set
 		Payload.SetExtinction(0.0);
 #endif
-		Payload.Ior = Ior;
 		if (Ior > 0.0)
 		{
-			// we are using actual refraction, so disable transparency
-			Payload.TransparencyColor = 0.0;
+			// Our material is a blend between glass and default_lit -- use stochastic blending between the two
+			if (PackedPayload.GetStochasticSlabRand() < Payload.BSDFOpacity)
+			{
+				// solid case - keep default_lit shading model
+			}
+			else
+			{
+				// glass
+				// we are using actual refraction, so disable transparency
+				Payload.Ior = Ior;
+				Payload.TransparencyColor = 0.0;
+				Payload.ShadingModelID = SHADINGMODELID_SOLID_GLASS;
+
+			}
+			Payload.BSDFOpacity = 1.0; // cancels out with the probability of selecting the lobe
 		}
 	}
 	else
@@ -1355,7 +1381,7 @@ RAY_TRACING_ENTRY_ANY_HIT(PathTracingMaterialAHS,
 
 			// Fade out the Transparency as a function of the SurfaceCoverage.
 			const float SurfaceCoverage = GetThinTranslucentMaterialOutput1(MaterialParameters);
-			Transparency = lerp(1.0.xxx, Transparency, SurfaceCoverage.xxx);
+			Transparency = lerp(1.0.xxx, Transparency, SurfaceCoverage);
 		}
 		else
 #endif // MATERIAL_SHADINGMODEL_THIN_TRANSLUCENT
