Forward shading - using more shared samplers when we have the possibility.
Exemple: a translucent forward material with Substrate enabled sees a reduction of sampler count from 12 to 8 when shared sampler is supported.

diff --git a/Engine/Shaders/Private/BRDF.ush b/Engine/Shaders/Private/BRDF.ush
--- a/Engine/Shaders/Private/BRDF.ush
+++ b/Engine/Shaders/Private/BRDF.ush
@@ -508,13 +508,19 @@ float3 F_Fresnel( float3 SpecularColor, float VoH )
 	return 0.5 * Square( (g - VoH) / (g + VoH) ) * ( 1 + Square( ((g+VoH)*VoH - 1) / ((g-VoH)*VoH + 1) ) );
 }
 
+#if SUPPORTS_INDEPENDENT_SAMPLERS
+	#define SharedSheenLTCSampler	View.SharedBilinearClampedSampler
+#else
+	#define SharedSheenLTCSampler	InSheenSampler
+#endif
+
 float4 SheenLTC_Cofficients(float NoV, float Roughness, Texture2D InSheenLTCTexture, SamplerState InSheenSampler)
 {
 	const float Alpha = sqrt(Roughness);
 	const float SatNoV = saturate(abs(NoV) + 1e-5);
 	float2 UV = float2(Alpha, SatNoV);
 	UV = UV * (31.0 / 32.0) + (0.5 / 32.0);
-	return InSheenLTCTexture.SampleLevel(InSheenSampler, UV, 0);
+	return InSheenLTCTexture.SampleLevel(SharedSheenLTCSampler, UV, 0);
 }
 
 float SheenLTC_DirectionalAlbedo(float NoV, float Roughness, Texture2D InSheenLTCTexture, SamplerState InSheenSampler)
diff --git a/Engine/Shaders/Private/ForwardLightingCommon.ush b/Engine/Shaders/Private/ForwardLightingCommon.ush
--- a/Engine/Shaders/Private/ForwardLightingCommon.ush
+++ b/Engine/Shaders/Private/ForwardLightingCommon.ush
@@ -535,7 +535,7 @@ half3 GetImageBasedReflectionSpecular(FMaterialPixelParameters MaterialParameter
 	if (abs(dot(PlanarReflectionStruct.ReflectionPlane.xyz, 1)) > .0001f)
 	{
 		// Reuse ReflectionCubemapSampler to avoid reducing the sampler count available to artists
-		float4 PlanarReflection = ComputePlanarReflections(MaterialParameters.WorldPosition_CamRelative, MaterialParameters.WorldNormal, Roughness, ReflectionStruct.ReflectionCubemapSampler);
+		float4 PlanarReflection = ComputePlanarReflections(MaterialParameters.WorldPosition_CamRelative, MaterialParameters.WorldNormal, Roughness, /*SharedClampSampler*/ View.SharedTrilinearClampedSampler);
 		// Planar reflections win over SSR and reflection environment
 		SpecularLighting = PlanarReflection.rgb + (1 - PlanarReflection.a) * SpecularLighting;
 	}
diff --git a/Engine/Shaders/Private/HeightFogCommon.ush b/Engine/Shaders/Private/HeightFogCommon.ush
--- a/Engine/Shaders/Private/HeightFogCommon.ush
+++ b/Engine/Shaders/Private/HeightFogCommon.ush
@@ -40,6 +40,13 @@
 	#endif
 #endif
 
+#if SUPPORTS_INDEPENDENT_SAMPLERS
+	#define SharedIntegratedLightScatteringSampler	View.SharedBilinearClampedSampler
+	#define SharedFogInColorScatteringSampler		View.SharedTrilinearClampedSampler
+#else
+	#define SharedIntegratedLightScatteringSampler	FogStruct.IntegratedLightScatteringSampler
+	#define SharedFogInColorScatteringSampler		FogStruct.FogInscatteringColorSampler
+#endif
 
 static const float FLT_EPSILON = 0.001f;
 static const float FLT_EPSILON2 = 0.01f;
@@ -61,8 +68,8 @@ float3 ComputeInscatteringColor(float3 CameraToReceiver, float CameraToReceiverL
 		float3 CubemapLookupVector = CameraToReceiver;
 		// Rotate around Z axis
 		CubemapLookupVector.xy = float2(dot(CubemapLookupVector.xy, float2(FogStruct.SinCosInscatteringColorCubemapRotation.y, -FogStruct.SinCosInscatteringColorCubemapRotation.x)), dot(CubemapLookupVector.xy, FogStruct.SinCosInscatteringColorCubemapRotation.xy));
-		float3 DirectionalColor = TextureCubeSampleLevel(FogStruct.FogInscatteringColorCubemap, FogStruct.FogInscatteringColorSampler, CubemapLookupVector, 0).xyz;
-		float3 NonDirectionalColor = TextureCubeSampleLevel(FogStruct.FogInscatteringColorCubemap, FogStruct.FogInscatteringColorSampler, CubemapLookupVector, FogStruct.FogInscatteringTextureParameters.z).xyz;
+		float3 DirectionalColor = TextureCubeSampleLevel(FogStruct.FogInscatteringColorCubemap, SharedFogInColorScatteringSampler, CubemapLookupVector, 0).xyz;
+		float3 NonDirectionalColor = TextureCubeSampleLevel(FogStruct.FogInscatteringColorCubemap, SharedFogInColorScatteringSampler, CubemapLookupVector, FogStruct.FogInscatteringTextureParameters.z).xyz;
 		Inscattering *= lerp(NonDirectionalColor, DirectionalColor, FadeAlpha);
 	}
 #endif
@@ -240,12 +247,6 @@ half4 CalculateHeightFog(float3 WorldPositionRelativeToCamera)
 	return FogInscatteringAndOpacity;
 }
 
-#if SUPPORTS_INDEPENDENT_SAMPLERS
-	#define SharedIntegratedLightScatteringSampler View.SharedBilinearClampedSampler
-#else
-	#define SharedIntegratedLightScatteringSampler FogStruct.IntegratedLightScatteringSampler
-#endif
-
 float4 CombineVolumetricFog(float4 GlobalFog, float3 VolumeUV, uint EyeIndex, float SceneDepth)
 {
 	float4 VolumetricFogLookup = float4(0, 0, 0, 1);
diff --git a/Engine/Shaders/Private/ParticipatingMediaCommon.ush b/Engine/Shaders/Private/ParticipatingMediaCommon.ush
--- a/Engine/Shaders/Private/ParticipatingMediaCommon.ush
+++ b/Engine/Shaders/Private/ParticipatingMediaCommon.ush
@@ -224,6 +224,14 @@ float IsotropicMediumSlabDirectionalAlbedoFadeLUT(float MFP)
 	return Fade;
 }
 
+#if SUPPORTS_INDEPENDENT_SAMPLERS
+	#define SharedSimpleVolumeTextureSampler	View.SharedBilinearClampedSampler
+	#define SharedSimpleVolumeEnvTextureSampler	View.SharedBilinearClampedSampler
+#else
+	#define SharedSimpleVolumeTextureSampler	View.SimpleVolumeTextureSampler
+	#define SharedSimpleVolumeEnvTextureSampler	View.SimpleVolumeEnvTextureSampler
+#endif
+
 float IsotropicMediumSlabPunctualDirectionalAlbedoLUT(float MFP, float BaseColor, float NoV, float NoL)
 {
 	// LUT content
@@ -263,8 +271,8 @@ float IsotropicMediumSlabPunctualDirectionalAlbedoLUT(float MFP, float BaseColor
 	float3 UVW0 = float3(UVs.x, UVs.y, EdgeOffset + FourthCoord0 * LUT4DUvStepW + LightAngleNorm * LUT4DUvDimZ);
 	float3 UVW1 = float3(UVs.x, UVs.y, EdgeOffset + FourthCoord1 * LUT4DUvStepW + LightAngleNorm * LUT4DUvDimZ);
 
-	float Value0 = View.SimpleVolumeTexture.SampleLevel(View.SimpleVolumeTextureSampler, UVW0, 0);
-	float Value1 = View.SimpleVolumeTexture.SampleLevel(View.SimpleVolumeTextureSampler, UVW1, 0);
+	float Value0 = View.SimpleVolumeTexture.SampleLevel(SharedSimpleVolumeTextureSampler, UVW0, 0);
+	float Value1 = View.SimpleVolumeTexture.SampleLevel(SharedSimpleVolumeTextureSampler, UVW1, 0);
 
 	float DirectionalAlbedo = lerp(Value0, Value1, FourthCoordLerp);
 
@@ -329,7 +337,7 @@ float IsotropicMediumSlabEnvDirectionalAlbedoLUT(float MFP, float BaseColor, flo
 	float3 UVs = float3(MFP / 10.0f, BaseColor, ViewAngleNorm);	// Input parameters in [0,1]
 	UVs = LUTMin + LUTLen * UVs;								// Correct in order to ignore texel border
 
-	float EnvDirectionalAlbedo = View.SimpleVolumeEnvTexture.SampleLevel(View.SimpleVolumeEnvTextureSampler, UVs, 0);
+	float EnvDirectionalAlbedo = View.SimpleVolumeEnvTexture.SampleLevel(SharedSimpleVolumeEnvTextureSampler, UVs, 0);
 
 	const float Fade = IsotropicMediumSlabDirectionalAlbedoFadeLUT(MFP);
 	return EnvDirectionalAlbedo * Fade;
diff --git a/Engine/Shaders/Private/RectLight.ush b/Engine/Shaders/Private/RectLight.ush
--- a/Engine/Shaders/Private/RectLight.ush
+++ b/Engine/Shaders/Private/RectLight.ush
@@ -404,6 +404,12 @@ FRectLTC GetRectLTC_GGX( float Roughness, float3 SpecularColor, half NoV)
 	return Out;
 }
 
+#if SUPPORTS_INDEPENDENT_SAMPLERS
+	#define SharedSheenLTCSampler	View.SharedBilinearClampedSampler
+#else
+	#define SharedSheenLTCSampler	View.SheenLTCSampler
+#endif
+
 // Return LTC coefficients for Sheen BSDF
 FRectLTC GetRectLTC_Sheen( float Roughness, half NoV)
 {	
@@ -411,7 +417,7 @@ FRectLTC GetRectLTC_Sheen( float Roughness, half NoV)
 	const float SatNoV = saturate(abs(NoV) + 1e-5);
 	float2 UV = float2(Alpha, SatNoV);
 	UV = UV * (31.0 / 32.0) + (0.5 / 32.0);
-	const float3 SheenLTC = View.SheenLTCTexture.SampleLevel(View.SheenLTCSampler, UV, 0).xyz;
+	const float3 SheenLTC = View.SheenLTCTexture.SampleLevel(SharedSheenLTCSampler, UV, 0).xyz;
 
 	const float aInv = SheenLTC.x;
 	const float bInv = SheenLTC.y;
diff --git a/Engine/Shaders/Private/ReflectionEnvironmentShared.ush b/Engine/Shaders/Private/ReflectionEnvironmentShared.ush
--- a/Engine/Shaders/Private/ReflectionEnvironmentShared.ush
+++ b/Engine/Shaders/Private/ReflectionEnvironmentShared.ush
@@ -16,6 +16,10 @@ float GetSkyLightCubemapBrightness()
 #define REFLECTION_CAPTURE_ROUGHEST_MIP 1
 #define REFLECTION_CAPTURE_ROUGHNESS_MIP_SCALE 1.2
 
+// This sampler cannot be shared because they come from the texture assets setup by users. We could enforce those to be min/mag/mip trilinear filtering.
+#define SharedSkyLightCubemapSampler					ReflectionStruct.SkyLightCubemapSampler
+#define SharedSkyLightBlendDestinationCubemapSampler	ReflectionStruct.SkyLightBlendDestinationCubemapSampler
+
 /** 
  * Compute absolute mip for a reflection capture cubemap given a roughness.
  */
@@ -39,7 +43,7 @@ float ComputeReflectionCaptureRoughnessFromMip(float Mip, half CubemapMaxMip)
 float3 GetSkyLightReflection(float3 ReflectionVector, float Roughness, out float OutSkyAverageBrightness)
 {
 	float AbsoluteSpecularMip = ComputeReflectionCaptureMipFromRoughness(Roughness, ReflectionStruct.SkyLightParameters.x);
-	float3 Reflection = TextureCubeSampleLevel(ReflectionStruct.SkyLightCubemap, ReflectionStruct.SkyLightCubemapSampler, ReflectionVector, AbsoluteSpecularMip).rgb;
+	float3 Reflection = TextureCubeSampleLevel(ReflectionStruct.SkyLightCubemap, SharedSkyLightCubemapSampler, ReflectionVector, AbsoluteSpecularMip).rgb;
 
 	OutSkyAverageBrightness = GetSkyLightCubemapBrightness() * Luminance(View.SkyLightColor.rgb);
 	return Reflection * View.SkyLightColor.rgb;
@@ -53,7 +57,7 @@ float3 GetSkyLightReflectionSupportingBlend(float3 ReflectionVector, float Rough
 	if (ReflectionStruct.SkyLightParameters.w > 0)
 	{
 		float AbsoluteSpecularMip = ComputeReflectionCaptureMipFromRoughness(Roughness, ReflectionStruct.SkyLightParameters.x);
-		float3 BlendDestinationReflection = TextureCubeSampleLevel(ReflectionStruct.SkyLightBlendDestinationCubemap, ReflectionStruct.SkyLightBlendDestinationCubemapSampler, ReflectionVector, AbsoluteSpecularMip).rgb;
+		float3 BlendDestinationReflection = TextureCubeSampleLevel(ReflectionStruct.SkyLightBlendDestinationCubemap, SharedSkyLightBlendDestinationCubemapSampler, ReflectionVector, AbsoluteSpecularMip).rgb;
 
 		Reflection = lerp(Reflection, BlendDestinationReflection * View.SkyLightColor.rgb, ReflectionStruct.SkyLightParameters.w);
 	}
diff --git a/Engine/Shaders/Private/ShadingEnergyConservation.ush b/Engine/Shaders/Private/ShadingEnergyConservation.ush
--- a/Engine/Shaders/Private/ShadingEnergyConservation.ush
+++ b/Engine/Shaders/Private/ShadingEnergyConservation.ush
@@ -13,6 +13,12 @@
 	#define USE_ENERGY_CONSERVATION 0
 #endif
 
+#if SUPPORTS_INDEPENDENT_SAMPLERS
+	#define SharedShadingEnergySampler			View.SharedBilinearClampedSampler
+#else
+	#define SharedShadingEnergySampler			View.ShadingEnergySampler
+#endif
+
 // Energy compensation/preservation for the various BxDFs. This file provides utility functions to tweak lobe weights
 // 
 // References:
@@ -42,7 +48,7 @@
 float2 GGXEnergyLookup(float Roughness, float NoV)
 {
 #if USE_ENERGY_CONSERVATION == 1
-	return View.ShadingEnergyGGXSpecTexture.SampleLevel(View.ShadingEnergySampler, float2(NoV, Roughness), 0);
+	return View.ShadingEnergyGGXSpecTexture.SampleLevel(SharedShadingEnergySampler, float2(NoV, Roughness), 0);
 #elif USE_ENERGY_CONSERVATION == 2
 	const float r = Roughness;
 	const float c = NoV;
@@ -60,7 +66,7 @@ float GGXEnergyLookup(float Roughness, float NoV, float Eta)
 	// This break reciprocity, but is much simpler than the approach outlined in [1] which requires substantially more computation.
 #if USE_ENERGY_CONSERVATION == 1
 	// NOTE: Eta is encoded for [1,3] range only, energy loss will happens above 3.0 - but this keeps more resolution on the visually significant portion of the range
-	float2 E = View.ShadingEnergyGGXGlassTexture.SampleLevel(View.ShadingEnergySampler, float3(NoV, Roughness, max(Eta, rcp(Eta)) * 0.5 - 0.5), 0);
+	float2 E = View.ShadingEnergyGGXGlassTexture.SampleLevel(SharedShadingEnergySampler, float3(NoV, Roughness, max(Eta, rcp(Eta)) * 0.5 - 0.5), 0);
 	return Eta >= 1.0 ? E.x : E.y;
 #else
 	// TODO: find an analytic fit for this case
@@ -71,7 +77,7 @@ float GGXEnergyLookup(float Roughness, float NoV, float Eta)
 float2 ClothEnergyLookup(float Roughness, float NoV)
 {
 #if USE_ENERGY_CONSERVATION == 1
-	return View.ShadingEnergyClothSpecTexture.SampleLevel(View.ShadingEnergySampler, float2(NoV, Roughness), 0);
+	return View.ShadingEnergyClothSpecTexture.SampleLevel(SharedShadingEnergySampler, float2(NoV, Roughness), 0);
 #elif USE_ENERGY_CONSERVATION == 2
 	// Simple analytical fit (avoids building a table)
 	// NOTE: this fit is not exact near grazing angles for roughness ~0.7
@@ -89,7 +95,7 @@ float2 ClothEnergyLookup(float Roughness, float NoV)
 float DiffuseEnergyLookup(float Roughness, float NoV)
 {
 #if USE_ENERGY_CONSERVATION == 1
-	//return View.ShadingEnergyDiffuseTexture.SampleLevel(View.ShadingEnergySampler, float2(NoV, Roughness), 0);
+	//return View.ShadingEnergyDiffuseTexture.SampleLevel(SharedShadingEnergySampler, float2(NoV, Roughness), 0);
 	// For now we do not apply Chan diffuse energy preservation on diffuse ambiant.
 	// This is because Chan is built for F=0.04 and unfortunately this causes ambient to darken a grazing angles.
 	// SUBSTRATE_TODO Apply the inverse of Fresnel with F=0.04 on Chan when building the table.
diff --git a/Engine/Source/Runtime/Engine/Private/SceneManagement.cpp b/Engine/Source/Runtime/Engine/Private/SceneManagement.cpp
--- a/Engine/Source/Runtime/Engine/Private/SceneManagement.cpp
+++ b/Engine/Source/Runtime/Engine/Private/SceneManagement.cpp
@@ -1208,9 +1208,9 @@ FViewUniformShaderParameters::FViewUniformShaderParameters()
 
 	// SimpleVolume
 	SimpleVolumeTexture = GBlackVolumeTexture->TextureRHI;
-	SimpleVolumeTextureSampler = TStaticSamplerState<SF_Trilinear, AM_Clamp, AM_Clamp, AM_Clamp>::GetRHI();
+	SimpleVolumeTextureSampler = TStaticSamplerState<SF_Bilinear, AM_Clamp, AM_Clamp, AM_Clamp>::GetRHI();
 	SimpleVolumeEnvTexture = GBlackVolumeTexture->TextureRHI;
-	SimpleVolumeEnvTextureSampler = TStaticSamplerState<SF_Trilinear, AM_Clamp, AM_Clamp, AM_Clamp>::GetRHI();
+	SimpleVolumeEnvTextureSampler = TStaticSamplerState<SF_Bilinear, AM_Clamp, AM_Clamp, AM_Clamp>::GetRHI();
 
 	// Rect light atlas
 	RectLightAtlasMaxMipLevel = 1;
diff --git a/Engine/Source/Runtime/Renderer/Private/SceneRendering.cpp b/Engine/Source/Runtime/Renderer/Private/SceneRendering.cpp
--- a/Engine/Source/Runtime/Renderer/Private/SceneRendering.cpp
+++ b/Engine/Source/Runtime/Renderer/Private/SceneRendering.cpp
@@ -2012,8 +2012,8 @@ void FViewInfo::SetupUniformBufferParameters(
 	}
 	ViewUniformShaderParameters.GlintTexture = OrBlack2DArrayIfNull(ViewUniformShaderParameters.GlintTexture);
 
-	ViewUniformShaderParameters.SimpleVolumeTextureSampler = TStaticSamplerState<SF_Trilinear, AM_Clamp, AM_Clamp, AM_Clamp>::GetRHI();
-	ViewUniformShaderParameters.SimpleVolumeEnvTextureSampler = TStaticSamplerState<SF_Trilinear, AM_Clamp, AM_Clamp, AM_Clamp>::GetRHI();
+	ViewUniformShaderParameters.SimpleVolumeTextureSampler = TStaticSamplerState<SF_Bilinear, AM_Clamp, AM_Clamp, AM_Clamp>::GetRHI();
+	ViewUniformShaderParameters.SimpleVolumeEnvTextureSampler = TStaticSamplerState<SF_Bilinear, AM_Clamp, AM_Clamp, AM_Clamp>::GetRHI();
 
 	if (GEngine->SimpleVolumeTexture)
 	{
