Add shared samplers for IES atlas, rect light integrator, light function atlas, hair LUT.

This reduces sampler usage on platforms supporting shared samplers.

diff --git a/Engine/Shaders/Private/ClusteredDeferredShadingPixelShader.usf b/Engine/Shaders/Private/ClusteredDeferredShadingPixelShader.usf
--- a/Engine/Shaders/Private/ClusteredDeferredShadingPixelShader.usf
+++ b/Engine/Shaders/Private/ClusteredDeferredShadingPixelShader.usf
@@ -316,7 +316,7 @@ float4 GetLightGridLocalLighting(
 				ScreenSpaceData.GBuffer,
 				TransmittanceMask,
 				View.HairScatteringLUTTexture,
-				View.HairScatteringLUTSampler,
+				HairScatteringLUTSampler,
 				View.HairComponents);
 
 			// This shouldn't be needed any more as the virtual shadow map should offer enough spatial precision
diff --git a/Engine/Shaders/Private/DeferredLightPixelShaders.usf b/Engine/Shaders/Private/DeferredLightPixelShaders.usf
--- a/Engine/Shaders/Private/DeferredLightPixelShaders.usf
+++ b/Engine/Shaders/Private/DeferredLightPixelShaders.usf
@@ -501,7 +501,7 @@ void DeferredLightPixelMain(
 			HairScreenSpaceData.GBuffer,
 			TransmittanceMask,
 			View.HairScatteringLUTTexture,
-			View.HairScatteringLUTSampler,
+			HairScatteringLUTSampler,
 			View.HairComponents);
 	}
 
diff --git a/Engine/Shaders/Private/HairStrands/HairStrandsDebug.usf b/Engine/Shaders/Private/HairStrands/HairStrandsDebug.usf
--- a/Engine/Shaders/Private/HairStrands/HairStrandsDebug.usf
+++ b/Engine/Shaders/Private/HairStrands/HairStrandsDebug.usf
@@ -2395,7 +2395,7 @@ void MainPS(
 				const float SinLightAngle = dot(L, T);
 
 				// Compute the transmittance based on precompute Hair transmittance LUT
-				FHairAverageScattering AverageScattering = SampleHairLUT(View.HairScatteringLUTTexture, View.HairScatteringLUTSampler, P.BaseColor, P.Roughness, SinLightAngle);
+				FHairAverageScattering AverageScattering = SampleHairLUT(View.HairScatteringLUTTexture, HairScatteringLUTSampler, P.BaseColor, P.Roughness, SinLightAngle);
 
 				TransmittanceData = ComputeDualScatteringTerms(
 					InitHairTransmittanceMask(),
diff --git a/Engine/Shaders/Private/HairStrands/HairStrandsDeepTransmittanceDualScattering.ush b/Engine/Shaders/Private/HairStrands/HairStrandsDeepTransmittanceDualScattering.ush
--- a/Engine/Shaders/Private/HairStrands/HairStrandsDeepTransmittanceDualScattering.ush
+++ b/Engine/Shaders/Private/HairStrands/HairStrandsDeepTransmittanceDualScattering.ush
@@ -9,6 +9,12 @@
 	#error HAIR_COMPONENT_GS is not defined, please include ShadingCommon.ush
 #endif
 
+#if SUPPORTS_INDEPENDENT_SAMPLERS
+#define HairScatteringLUTSampler View.SharedBilinearClampedSampler
+#else
+#define HairScatteringLUTSampler View.HairScatteringLUTSampler
+#endif
+
 float Hair_g2(float Variance,float Theta)
 {
 	//const float A = 1.f / sqrt(2 * PI * Variance);
@@ -164,7 +170,7 @@ FHairTransmittanceData EvaluateDualScattering(inout float3 BaseColor, float3 Wor
 
 	// Compute the transmittance based on precompute Hair transmittance LUT
 	const float SinLightAngle = dot(L, WorldNormal);
-	const FHairAverageScattering AverageScattering = SampleHairLUT(View.HairScatteringLUTTexture, View.HairScatteringLUTSampler, BaseColor, Roughness, SinLightAngle);
+	const FHairAverageScattering AverageScattering = SampleHairLUT(View.HairScatteringLUTTexture, HairScatteringLUTSampler, BaseColor, Roughness, SinLightAngle);
 
 	return ComputeDualScatteringTerms(
 		InitHairTransmittanceMask(),
diff --git a/Engine/Shaders/Private/HairStrands/HairStrandsEnvironmentLighting.usf b/Engine/Shaders/Private/HairStrands/HairStrandsEnvironmentLighting.usf
--- a/Engine/Shaders/Private/HairStrands/HairStrandsEnvironmentLighting.usf
+++ b/Engine/Shaders/Private/HairStrands/HairStrandsEnvironmentLighting.usf
@@ -355,7 +355,7 @@ float3 ReflectionEnvironment(FGBufferData GBuffer, float2 ScreenPosition, float2
 			const FHairTraversalResult Result = InternalComputeHairCountVoxel(NodeDesc, SampleTranslateWorldPosition, SampleTranslateLightPosition, 1, 0, MaxSearchDistance, PixelRadius, true, bDebugEnabled);
 
 			const uint ScatteringComponent = (View.HairComponents & HAIR_COMPONENT_R) | (View.HairComponents & HAIR_COMPONENT_TRT) | HAIR_COMPONENT_LS | HAIR_COMPONENT_GS;// | HAIR_MULTISCATTER;
-			const FHairTransmittanceData TransmittanceData = GetHairTransmittance(V, SampleL, GBuffer, InitHairTransmittanceMask(Result.HairCount, Result.Visibility), View.HairScatteringLUTTexture, View.HairScatteringLUTSampler, ScatteringComponent);
+			const FHairTransmittanceData TransmittanceData = GetHairTransmittance(V, SampleL, GBuffer, InitHairTransmittanceMask(Result.HairCount, Result.Visibility), View.HairScatteringLUTTexture, HairScatteringLUTSampler, ScatteringComponent);
 
 			const float Area = 0.2;
 			const float BackLit = 1;
@@ -424,7 +424,7 @@ float3 ReflectionEnvironment(FGBufferData GBuffer, float2 ScreenPosition, float2
 		#if 0
 		{
 			const float SinViewAngle = dot(T, R);
-			FHairAverageEnergy Energy = SampleHairEnergyLUT(HairEnergyLUTTexture, View.HairScatteringLUTSampler, GBuffer.BaseColor, GBuffer.Roughness, SinViewAngle);
+			FHairAverageEnergy Energy = SampleHairEnergyLUT(HairEnergyLUTTexture, HairScatteringLUTSampler, GBuffer.BaseColor, GBuffer.Roughness, SinViewAngle);
 			ColorMS.rgb *= Energy.A_R + Energy.A_TT + Energy.A_TRT;
 		}
 		#endif	
@@ -479,7 +479,7 @@ float3 ReflectionEnvironment(FGBufferData GBuffer, float2 ScreenPosition, float2
 
 			// Compute the number of hair count between light & shading point
 			const FHairTraversalResult Result = InternalComputeHairCountVoxel(NodeDesc, SampleTranslatedWorldPosition, SampleTranslatedLightPosition, 1, 0, MaxSearchDistance, PixelRadius, true, bDebugEnabled);
-			Transmittance = GetHairTransmittanceOnly(InitHairTransmittanceMask(Result.HairCount, Result.Visibility), GBuffer, SinLightAngle, View.HairScatteringLUTTexture, View.HairScatteringLUTSampler);
+			Transmittance = GetHairTransmittanceOnly(InitHairTransmittanceMask(Result.HairCount, Result.Visibility), GBuffer, SinLightAngle, View.HairScatteringLUTTexture, HairScatteringLUTSampler);
 			Visibility = Result.Visibility;
 		}
 
@@ -487,7 +487,7 @@ float3 ReflectionEnvironment(FGBufferData GBuffer, float2 ScreenPosition, float2
 		// Directional albedo
 		#if 0
 		{
-			FHairAverageEnergy Energy = SampleHairEnergyLUT(HairEnergyLUTTexture, View.HairScatteringLUTSampler, GBuffer.BaseColor, GBuffer.Roughness, SinViewAngle);
+			FHairAverageEnergy Energy = SampleHairEnergyLUT(HairEnergyLUTTexture, HairScatteringLUTSampler, GBuffer.BaseColor, GBuffer.Roughness, SinViewAngle);
 			ColorR.rgb *= (Energy.A_TRT + Energy.A_R) * Transmittance.xyz * Visibility;
 		}
 		#endif
@@ -546,7 +546,7 @@ float3 ReflectionEnvironment(FGBufferData GBuffer, float2 ScreenPosition, float2
 			{
 				Result = InternalComputeHairCountVoxel(NodeDesc, SampleTranslatedWorldPosition, SampleTranslatedLightPosition, TransmissionDensityScaleFactor, 0, DistanceThresholdTT, PixelRadius, false, bDebugEnabled);
 			}
-			Transmittance = GetHairTransmittanceOnly(Result.HairCount, Result.Visibility, GBuffer, SinLightAngle, View.HairScatteringLUTTexture, View.HairScatteringLUTSampler);
+			Transmittance = GetHairTransmittanceOnly(Result.HairCount, Result.Visibility, GBuffer, SinLightAngle, View.HairScatteringLUTTexture, HairScatteringLUTSampler);
 			Visibility = Result.Visibility;
 		}
 
@@ -555,7 +555,7 @@ float3 ReflectionEnvironment(FGBufferData GBuffer, float2 ScreenPosition, float2
 		{
 			Error the Transmittance.w does not contain the correct hair count when we hit opaque geometry: TODO
 			const uint HairComponent = HAIR_COMPONENT_TT | HAIR_COMPONENT_LS | HAIR_COMPONENT_GS;
-			const FHairTransmittanceData TransmittanceData = GetHairTransmittance(V, L, GBuffer, Transmittance, View.HairScatteringLUTTexture, View.HairScatteringLUTSampler, HairComponent);
+			const FHairTransmittanceData TransmittanceData = GetHairTransmittance(V, L, GBuffer, Transmittance, View.HairScatteringLUTTexture, HairScatteringLUTSampler, HairComponent);
 			const float Area = 0.2;
 			ColorTT.rgb *= HairShading(GBuffer, L, V, T, 1, TransmittanceData, 1, Area, uint2(0, 0)) * Visibility;
 		}
@@ -565,7 +565,7 @@ float3 ReflectionEnvironment(FGBufferData GBuffer, float2 ScreenPosition, float2
 		#if 1
 		{
 			const float Backlit = GBuffer.CustomData.z;
-			FHairAverageEnergy Energy = SampleHairEnergyLUT(HairEnergyLUTTexture, View.HairScatteringLUTSampler, GBuffer.BaseColor, GBuffer.Roughness, SinViewAngle);
+			FHairAverageEnergy Energy = SampleHairEnergyLUT(HairEnergyLUTTexture, HairScatteringLUTSampler, GBuffer.BaseColor, GBuffer.Roughness, SinViewAngle);
 			ColorTT.rgb *= Energy.A_TT * Transmittance * Visibility * Backlit;
 		}
 		#endif	
@@ -665,7 +665,7 @@ float3 ReflectionEnvironmentUniformSampling(
 			GBuffer,
 			InitHairTransmittanceMask(Result.HairCount, Result.Visibility),
 			View.HairScatteringLUTTexture, 
-			View.HairScatteringLUTSampler,
+			HairScatteringLUTSampler,
 			View.HairComponents);	
 
 		const float Area = 0;	
@@ -722,7 +722,7 @@ float3 ReflectionEnvironmentSH(
 		GBuffer,
 		InitHairTransmittanceMask(),
 		View.HairScatteringLUTTexture,
-		View.HairScatteringLUTSampler,
+		HairScatteringLUTSampler,
 		View.HairComponents);
 
 	// Note: We don't reuse EvaluateEnvHair() here as we need to estimate occlusion with GetHairTransmittance(), and the 2 * PI gives too bright result 
@@ -788,7 +788,7 @@ float3 ComputeThinTransmission(FGBufferData GBuffer, float2 ScreenPosition, floa
 
 	// The pre-exposure is not applied onto the output as the SceneColorTexture value are already preexposed.
 	const float SinViewAngle = dot(T, V);
-	const FHairAverageScattering Energy = SampleHairLUT(View.HairScatteringLUTTexture, View.HairScatteringLUTSampler, GBuffer.BaseColor, GBuffer.Roughness, SinViewAngle);
+	const FHairAverageScattering Energy = SampleHairLUT(View.HairScatteringLUTTexture, HairScatteringLUTSampler, GBuffer.BaseColor, GBuffer.Roughness, SinViewAngle);
 	float3 FinalColor = Energy.A_front * Luminance.xyz * PixelCoverage;
 
 	// Note: No need for pre-exposition, as we fetch & weight SceneColor which is already pre-exposed.
diff --git a/Engine/Shaders/Private/IESLightProfilesCommon.ush b/Engine/Shaders/Private/IESLightProfilesCommon.ush
--- a/Engine/Shaders/Private/IESLightProfilesCommon.ush
+++ b/Engine/Shaders/Private/IESLightProfilesCommon.ush
@@ -12,6 +12,12 @@
   #endif
 #endif
 
+#if SUPPORTS_INDEPENDENT_SAMPLERS
+#define IESAtlasSampler	View.SharedBilinearClampedSampler
+#else
+#define IESAtlasSampler	View.IESAtlasSampler
+#endif
+
 // Special case for IES texture preview, which require a standalone texture
 #if defined(USE_IES_STANDALONE_TEXTURE) && USE_IES_STANDALONE_TEXTURE == 1
 Texture2D		IESTexture;
@@ -46,7 +52,7 @@ float ComputeLightProfileMultiplier(float3 WorldPosition, float3 LightPosition,
 	#if USE_IES_STANDALONE_TEXTURE
 		Out = Texture2DSampleLevel(IESTexture, IESTextureSampler, float2(NormAngle, NormTangentAngle), 0).r;
 	#else
-		Out = View.IESAtlasTexture.SampleLevel(View.IESAtlasSampler, float3(NormAngle, NormTangentAngle, InIESAtlasIndex), 0);
+		Out = View.IESAtlasTexture.SampleLevel(IESAtlasSampler, float3(NormAngle, NormTangentAngle, InIESAtlasIndex), 0);
 	#endif
 	}
 #endif
diff --git a/Engine/Shaders/Private/LightFunctionAtlas/LightFunctionAtlasCommon.usf b/Engine/Shaders/Private/LightFunctionAtlas/LightFunctionAtlasCommon.usf
--- a/Engine/Shaders/Private/LightFunctionAtlas/LightFunctionAtlasCommon.usf
+++ b/Engine/Shaders/Private/LightFunctionAtlas/LightFunctionAtlasCommon.usf
@@ -23,6 +23,12 @@
 #define FLightFunctionColor float
 #endif
 
+#if SUPPORTS_INDEPENDENT_SAMPLERS
+#define LightFunctionAtlasSampler	View.SharedBilinearClampedSampler
+#else
+#define LightFunctionAtlasSampler	LightFunctionAtlas.LightFunctionAtlasSampler
+#endif
+
 FLightFunctionColor GetLocalLightFunctionCommon(float3 TranslatedWorldPosition, uint LightFunctionAtlasLightIndex)
 {
 	if (LightFunctionAtlasLightIndex == 0)
@@ -88,9 +94,9 @@ FLightFunctionColor GetLocalLightFunctionCommon(float3 TranslatedWorldPosition,
 	float2 AtlasUV = Slot_MinUV + saturate(LightFunctionUV) * LightFunctionAtlas.Slot_UVSize.xx;
 
 #if COLORED_LIGHT_FUNCTION_ATLAS
-	FLightFunctionColor LightFunctionColor = Texture2DSampleLevel(LightFunctionAtlas.LightFunctionAtlasTexture, LightFunctionAtlas.LightFunctionAtlasSampler, AtlasUV, 0).rgb;
+	FLightFunctionColor LightFunctionColor = Texture2DSampleLevel(LightFunctionAtlas.LightFunctionAtlasTexture, LightFunctionAtlasSampler, AtlasUV, 0).rgb;
 #else
-	FLightFunctionColor LightFunctionColor = Texture2DSampleLevel(LightFunctionAtlas.LightFunctionAtlasTexture, LightFunctionAtlas.LightFunctionAtlasSampler, AtlasUV, 0).x;
+	FLightFunctionColor LightFunctionColor = Texture2DSampleLevel(LightFunctionAtlas.LightFunctionAtlasTexture, LightFunctionAtlasSampler, AtlasUV, 0).x;
 #endif
 	LightFunctionColor = LightFunctionColor * LightFunctionColor; // Simple Gamma decompression
 
diff --git a/Engine/Shaders/Private/RectLight.ush b/Engine/Shaders/Private/RectLight.ush
--- a/Engine/Shaders/Private/RectLight.ush
+++ b/Engine/Shaders/Private/RectLight.ush
@@ -11,6 +11,14 @@
 #define USE_SOURCE_TEXTURE 1
 #endif
 
+#if SUPPORTS_INDEPENDENT_SAMPLERS
+#define GGXLTCMatSampler	View.SharedBilinearClampedSampler
+#define GGXLTCAmpSampler	View.SharedBilinearClampedSampler
+#else
+#define GGXLTCMatSampler	View.GGXLTCMatSampler
+#define GGXLTCAmpSampler	View.GGXLTCAmpSampler
+#endif
+
 struct FRect
 {
 	float3		Origin;
@@ -379,8 +387,8 @@ FRectLTC GetRectLTC_GGX( float Roughness, float3 SpecularColor, half NoV)
 	float2 UV = float2( Roughness, sqrt( 1 - NoV ) );
 	UV = UV * (63.0 / 64.0) + (0.5 / 64.0);
    
-	float4 LTCMat = View.GGXLTCMatTexture.SampleLevel( View.GGXLTCMatSampler, UV, 0 );
-	float4 LTCAmp = View.GGXLTCAmpTexture.SampleLevel( View.GGXLTCAmpSampler, UV, 0 );
+	float4 LTCMat = View.GGXLTCMatTexture.SampleLevel( GGXLTCMatSampler, UV, 0 );
+	float4 LTCAmp = View.GGXLTCAmpTexture.SampleLevel( GGXLTCAmpSampler, UV, 0 );
 
 	float3x3 LTC = {
 		float3( LTCMat.x, 0, LTCMat.z ),
