Add experiemental screen trace for hanlding occlusion test when integrating Lumen input for hair.
* This helps fixing light leak on eyebrows.
* This is enabled/disabled with r.HairStrands.SkyLighting.ScreenTraceOcclusion

diff --git a/Engine/Shaders/Private/HairStrands/HairStrandsEnvironmentLighting.usf b/Engine/Shaders/Private/HairStrands/HairStrandsEnvironmentLighting.usf
--- a/Engine/Shaders/Private/HairStrands/HairStrandsEnvironmentLighting.usf
+++ b/Engine/Shaders/Private/HairStrands/HairStrandsEnvironmentLighting.usf
@@ -271,6 +271,63 @@ FHairTraversalResult InternalComputeHairCountVoxel(
 		TraversalSettings);
 }
 
+#if PERMUTATION_INTEGRATION_TYPE == INTEGRATION_UNIFORM || PERMUTATION_INTEGRATION_TYPE == INTEGRATION_ADHOC || PERMUTATION_INTEGRATION_TYPE == INTEGRATION_SH
+
+#include "../SSRT/SSRTRayCast.ush"
+uint bScreenTrace;
+float SlopeCompareToleranceScale;
+float MaxScreenTraceFraction;
+Texture2D FurthestHZBTexture;
+#define FurthestHZBTextureSampler GlobalPointClampedSampler
+float4 HZBUvFactorAndInvFactor;
+
+float GetRayVisibility(float3 TranslatedWorldPosition,	float SceneDepth, float3 RayDirection, uint2 PixelCoord, uint PixelRayIndex, uint NumPixelSamples)
+{
+	float TraceDistance = MaxScreenTraceFraction * 2.0 * GetScreenRayLengthMultiplierForProjectionType(SceneDepth).x;
+	float DepthThresholdScale = View.ProjectionDepthThicknessScale;
+
+	FSSRTCastingSettings CastSettings = CreateDefaultCastSettings();
+	CastSettings.bStopWhenUncertain = true;
+
+	bool bHit = false;
+	float Level;
+	float3 HitUVz;
+	bool bRayWasClipped;
+	uint NumSteps = 4;
+	float StartMipLevel = 0;
+	float RayRoughness = .2f;
+	uint2 NoiseCoord = PixelCoord * uint2(NumPixelSamples, 1) + uint2(PixelRayIndex, 0);
+	float StepOffset = InterleavedGradientNoise(NoiseCoord + 0.5f, 0.0f);
+
+	FSSRTRay Ray = InitScreenSpaceRayFromWorldSpace(
+		TranslatedWorldPosition, RayDirection,
+		/* WorldTMax = */ TraceDistance,
+		/* SceneDepth = */ SceneDepth,
+		/* SlopeCompareToleranceScale */ SlopeCompareToleranceScale * DepthThresholdScale * (float)NumSteps,
+		/* bExtendRayToScreenBorder = */ false,
+		/* out */ bRayWasClipped);
+
+	bool bUncertain;
+	float3 DebugOutput;
+
+	CastScreenSpaceRay(
+		FurthestHZBTexture, FurthestHZBTextureSampler,
+		StartMipLevel,
+		CastSettings,
+		Ray, RayRoughness, NumSteps, StepOffset - .9f,
+		HZBUvFactorAndInvFactor, false,
+		/* out */ DebugOutput,
+		/* out */ HitUVz,
+		/* out */ Level,
+		/* out */ bHit,
+		/* out */ bUncertain);
+
+	bHit = bHit && !bUncertain;
+	return bHit ? 0.0f : 1.0f;
+}
+#endif
+
+
 //////////////////////////////////////////////////////////////////////////////////////////////////////////////////////////
 //////////////////////////////////////////////////////////////////////////////////////////////////////////////////////////
 //////////////////////////////////////////////////////////////////////////////////////////////////////////////////////////
@@ -648,7 +705,7 @@ float3 ReflectionEnvironmentUniformSampling(
 		#endif
 
 		const float SinLightAngle = dot(SampleL, T);
-		const FHairTraversalResult Result = InternalComputeHairCountVoxel(
+		FHairTraversalResult Result = InternalComputeHairCountVoxel(
 			NodeDesc,
 			SampleTranslatedWorldPosition, 
 			SampleTranslatedLightPosition, 
@@ -668,6 +725,11 @@ float3 ReflectionEnvironmentUniformSampling(
 			HairScatteringLUTSampler,
 			View.HairComponents);	
 
+		if (Result.Visibility > 0 && bScreenTrace)
+		{
+			Result.Visibility *= GetRayVisibility(TranslatedWorldPosition, PixelSceneDepth, SampleL, PixelPosition, SampleIt, SampleCount);
+		}
+
 		const float Area = 0;	
 		const float Backlit = 1;
 		const float RoughnessBias = 0.1f;
@@ -738,6 +800,12 @@ float3 ReflectionEnvironmentSH(
 		const uint2  PixelPosition = SvPosition - ResolvedView.ViewRectMin.xy;
 		Lighting = View.PreExposure * GatherRadiance(L, GBuffer.Roughness, PixelPosition, GBuffer.Depth, TranslatedWorldPosition) * IndirectDiffuseColor;
 	}
+
+	if (bScreenTrace)
+	{
+		Lighting *= GetRayVisibility(TranslatedWorldPosition, GBuffer.Depth, L, SvPosition, 0, 1);
+	}
+
 #else // PERMUTATION_LUMEN
 	const float AmbientOcclusion = 1.0f;
 	const float3 BentNormal = GBuffer.WorldNormal;
diff --git a/Engine/Source/Runtime/Renderer/Private/HairStrands/HairStrandsEnvironment.cpp b/Engine/Source/Runtime/Renderer/Private/HairStrands/HairStrandsEnvironment.cpp
--- a/Engine/Source/Runtime/Renderer/Private/HairStrands/HairStrandsEnvironment.cpp
+++ b/Engine/Source/Runtime/Renderer/Private/HairStrands/HairStrandsEnvironment.cpp
@@ -63,6 +63,9 @@ static FAutoConsoleVariableRef CVarHairStrandsSkyLighting_IntegrationType(TEXT("
 static int32 GHairStrandsSkyLighting_DebugSample = 0;
 static FAutoConsoleVariableRef CVarHairStrandsSkyLighting_DebugSample(TEXT("r.HairStrands.SkyLighting.DebugSample"), GHairStrandsSkyLighting_DebugSample, TEXT("Enable debug view for visualizing sample used for the sky integration"), ECVF_Scalability | ECVF_RenderThreadSafe);
 
+static int32 GHairStrandsSkyLighting_ScreenTraceOcclusion = 0;
+static FAutoConsoleVariableRef CVarHairStrandsSkyLighting_ScreenTraceOcclusion(TEXT("r.HairStrands.SkyLighting.ScreenTraceOcclusion"), GHairStrandsSkyLighting_ScreenTraceOcclusion, TEXT("Enable screen trace occlusion for during hair env. lighting integration"), ECVF_Scalability | ECVF_RenderThreadSafe);
+
 ///////////////////////////////////////////////////////////////////////////////////////////////////
 
 DECLARE_GPU_STAT_NAMED(HairSkyLighting, TEXT("Hair Sky lighting"));
@@ -251,6 +254,14 @@ class FHairEnvironmentLightingPS : public FGlobalShader
 		SHADER_PARAMETER(float,  HairDistanceThreshold)
 		SHADER_PARAMETER(uint32, bHairUseViewHairCount)
 
+		SHADER_PARAMETER(uint32, bScreenTrace)
+		SHADER_PARAMETER(float, SlopeCompareToleranceScale)
+		SHADER_PARAMETER(float, MaxScreenTraceFraction)
+		SHADER_PARAMETER(float, ScreenTraceNoFallbackThicknessScale)
+		SHADER_PARAMETER_RDG_TEXTURE(Texture2D, FurthestHZBTexture)
+		SHADER_PARAMETER_SAMPLER(SamplerState, FurthestHZBTextureSampler)
+		SHADER_PARAMETER(FVector4f, HZBUvFactorAndInvFactor)
+
 		SHADER_PARAMETER_TEXTURE(Texture2D, PreIntegratedGF)
 		SHADER_PARAMETER_SAMPLER(SamplerState, PreIntegratedGFSampler)
 
@@ -387,6 +398,33 @@ static void AddHairStrandsEnvironmentLightingPassPS(
 	ParametersPS->ReflectionCaptureData = View.ReflectionCaptureUniformBuffer;
 	ParametersPS->ReflectionsParameters = CreateReflectionUniformBuffer(GraphBuilder, View);
 	
+	const bool bScreenTrace = View.HZB && GHairStrandsSkyLighting_ScreenTraceOcclusion && IntegrationType != EHairLightingIntegrationType::SceneColor;
+	ParametersPS->FurthestHZBTexture = GSystemTextures.GetDepthDummy(GraphBuilder);
+	ParametersPS->FurthestHZBTextureSampler = TStaticSamplerState<SF_Point>::GetRHI();
+	ParametersPS->bScreenTrace = bScreenTrace ? 1u : 0u;
+	if (bScreenTrace)
+	{
+		// Default values, need more more testing/tweaking
+		const float MaxScreenTraceFraction = 128 * 2.0f / (float)View.ViewRect.Width();
+		const float ScreenTraceNoFallbackThicknessScale = View.ViewMatrices.GetPerProjectionDepthThicknessScale();
+
+		const FVector2D ViewportUVToHZBBufferUV(
+			float(View.ViewRect.Width()) / float(2 * View.HZBMipmap0Size.X),
+			float(View.ViewRect.Height()) / float(2 * View.HZBMipmap0Size.Y)
+		);
+
+		ParametersPS->HZBUvFactorAndInvFactor = FVector4f(
+			ViewportUVToHZBBufferUV.X,
+			ViewportUVToHZBBufferUV.Y,
+			1.0f / ViewportUVToHZBBufferUV.X,
+			1.0f / ViewportUVToHZBBufferUV.Y);
+		ParametersPS->FurthestHZBTexture = View.HZB;
+		ParametersPS->FurthestHZBTextureSampler = TStaticSamplerState<SF_Point>::GetRHI();		
+		ParametersPS->MaxScreenTraceFraction = MaxScreenTraceFraction;
+		ParametersPS->ScreenTraceNoFallbackThicknessScale = ScreenTraceNoFallbackThicknessScale;	
+		ParametersPS->SlopeCompareToleranceScale = 0.5f;
+	}
+
 	if (LightingType == EHairLightingSourceType::Lumen)
 	{
 		const FRadianceCacheState& RadianceCacheState = View.ViewState->Lumen.RadianceCacheState;
