HeterogeneousVolumes: Correctly composite with Thin Glass

diff --git a/Engine/Shaders/Private/BasePassPixelShader.usf b/Engine/Shaders/Private/BasePassPixelShader.usf
--- a/Engine/Shaders/Private/BasePassPixelShader.usf
+++ b/Engine/Shaders/Private/BasePassPixelShader.usf
@@ -1349,8 +1349,8 @@ void FPixelShaderInOut_MainPS(
 
 #if ADAPTIVE_VOLUMETRIC_SHADOW_MAP
 	float3 TranslatedWorldPosition = SvPositionToResolvedTranslatedWorld(In.SvPosition);
-	float Transmittance = saturate(AVSM_SampleTransmittance(TranslatedWorldPosition, ResolvedView.TranslatedWorldCameraOrigin));
-	Opacity *= Transmittance;
+	float HeterogeneousVolumeTransmittance = saturate(AVSM_SampleTransmittance(TranslatedWorldPosition, ResolvedView.TranslatedWorldCameraOrigin));
+	Opacity *= HeterogeneousVolumeTransmittance;
 
 #endif // ADAPTIVE_VOLUMETRIC_SHADOW_MAP
 
@@ -1526,6 +1526,9 @@ void FPixelShaderInOut_MainPS(
 
 #if MATERIAL_SHADINGMODEL_THIN_TRANSLUCENT && !SUBSTRATE_ENABLED
 	{
+#if ADAPTIVE_VOLUMETRIC_SHADOW_MAP
+		ColorSeparateSpecular *= HeterogeneousVolumeTransmittance;
+#endif
 		AccumulateThinTranslucentModel(
 								DualBlendSurfaceLuminancePostCoverage,
 								DualBlendSurfaceTransmittancePreCoverage,
@@ -1536,7 +1539,10 @@ void FPixelShaderInOut_MainPS(
 								ColorSeparateSpecular,
 								Emissive,
 								Opacity);
-
+#if ADAPTIVE_VOLUMETRIC_SHADOW_MAP
+		DualBlendSurfaceCoverage *= HeterogeneousVolumeTransmittance;
+		DualBlendSurfaceLuminancePostCoverage *= HeterogeneousVolumeTransmittance;
+#endif
 		Color = 0;
 		Opacity = 1.0f;
 	}
