Fix OnePass transmittance support for hair with light not casting shadow.

The light index computation for query a light within the light grid was incorrect if the scene contains a light not casting shadow. The CL ensures that the light indices used for transmittance computation and for lighting match.

diff --git a/Engine/Shaders/Private/ClusteredDeferredShadingPixelShader.usf b/Engine/Shaders/Private/ClusteredDeferredShadingPixelShader.usf
--- a/Engine/Shaders/Private/ClusteredDeferredShadingPixelShader.usf
+++ b/Engine/Shaders/Private/ClusteredDeferredShadingPixelShader.usf
@@ -232,7 +232,8 @@ float4 GetLightGridLocalLighting(
 	uint FirstNonSimpleLightIndex, 
 	uint SampleIndex=0, 
 	uint TotalSampleCount=0, 
-	uint2 PixelCoord=0)
+	uint2 PixelCoord=0,
+	uint HairLightChannelMask=0)
 {
 	float4 DirectLighting = 0;
 
@@ -243,10 +244,14 @@ float4 GetLightGridLocalLighting(
 	int2 PixelPos = int2( ScreenUV.xy * View.BufferSizeAndInvSize.xy );
 	uint4 ShadowMask = ShadowMaskBits[ PixelPos ];
 	uint LightChannelMask = LIGHTING_CHANNEL_MASK;
+	#if USE_HAIR_LIGHTING
+	LightChannelMask = HairLightChannelMask;
+	#else
 	if (bHasLightChannels)
 	{
 		LightChannelMask = LightingChannelsTexture.Load(uint3(PixelPos, 0)).x;
 	}
+	#endif
 	
 	LOOP
 	for (uint GridLightListIndex = FirstNonSimpleLightIndex; GridLightListIndex < NumLightsInGridCell; GridLightListIndex++)
@@ -281,6 +286,7 @@ float4 GetLightGridLocalLighting(
 		// Find packed VSM projection result if present
 		// NOTE: We have to do the search as in the deferred pass again currently as the shadow masks
 		// are packed according to a prune light list that contains only shadow-casting lights.
+		uint VSMGridLightListIndex = GridLightListIndex;
 		if (LocalLight.VirtualShadowMapId != INDEX_NONE)
 		{
 			LightAttenuation = GetVirtualShadowMapMaskForLight(
@@ -288,7 +294,8 @@ float4 GetLightGridLocalLighting(
 				uint2(PixelPos),
 				SceneDepth,
 				LocalLight.VirtualShadowMapId,
-				TranslatedWorldPosition).xxxx;
+				TranslatedWorldPosition,
+				VSMGridLightListIndex).xxxx;
 		}
 
 		#if USE_HAIR_LIGHTING
@@ -298,10 +305,12 @@ float4 GetLightGridLocalLighting(
 			const float3 V = normalize(-CameraVector);
 
 			// Fetch precomputed transmittance
+			// * Transmittance data are only computed for shadowed lights with VSM
+			// * Transmittance mask index is computed using VSM light index remapping. This needs to match HairStrandsDeepTransmittanceMask.usf
 			FHairTransmittanceMask TransmittanceMask = InitHairTransmittanceMask();
-			if (GridLightListIndex < GetPackedShadowMaskMaxLightCount())
+			if (LocalLight.VirtualShadowMapId != INDEX_NONE && VSMGridLightListIndex < GetPackedShadowMaskMaxLightCount())
 			{
-				const uint TransmittanceSampleIndex = SampleIndex + GridLightListIndex * TotalSampleCount;
+				const uint TransmittanceSampleIndex = SampleIndex + VSMGridLightListIndex * TotalSampleCount;
 				const uint TransmittanceSampleMaxCount = TotalSampleCount * GetPackedShadowMaskMaxLightCount();
 				if (TransmittanceSampleIndex < TransmittanceSampleMaxCount)
 				{
@@ -581,7 +590,7 @@ void ClusteredShadingPixelShader(
 
 	float4 CompositedLighting = 0;
 	uint FirstNonSimpleLightIndex = 0;
-	CompositedLighting += GetLightGridLocalLighting(ScreenSpaceData, CulledLightGridData, TranslatedWorldPosition, CameraVector, ScreenUV, SceneDepth, 0, Dither, FirstNonSimpleLightIndex, SampleIndex, TotalSampleCount, PixelCoord);
+	CompositedLighting += GetLightGridLocalLighting(ScreenSpaceData, CulledLightGridData, TranslatedWorldPosition, CameraVector, ScreenUV, SceneDepth, 0, Dither, FirstNonSimpleLightIndex, SampleIndex, TotalSampleCount, PixelCoord, Sample.LightChannelMask);
 #if !VISUALIZE_LIGHT_CULLING
 	CompositedLighting *= View.PreExposure;
 #endif
diff --git a/Engine/Shaders/Private/DeferredLightPixelShaders.usf b/Engine/Shaders/Private/DeferredLightPixelShaders.usf
--- a/Engine/Shaders/Private/DeferredLightPixelShaders.usf
+++ b/Engine/Shaders/Private/DeferredLightPixelShaders.usf
@@ -400,6 +400,35 @@ void DeferredLightPixelMain(
 
 #if USE_HAIR_LIGHTING == 1
 
+// Compute transmittance data from a transmittance mask
+FHairTransmittanceData GetTransmittanceDataFromTransmitttanceMask(
+	FDeferredLightData LightData, 
+	FGBufferData GBufferData, 
+	FHairTransmittanceMask TransmittanceMask,
+	float3 TranslatedWorldPosition, 
+	float3 CameraVector)
+{
+	float3 L = float3(0, 0, 1);
+	if (LightData.bRadialLight)
+	{
+		L = normalize(LightData.TranslatedWorldPosition - TranslatedWorldPosition);
+	}
+	else
+	{
+		L = LightData.Direction;
+	}
+
+	const float3 V = normalize(-CameraVector);
+	return GetHairTransmittance(
+		V,
+		L,
+		GBufferData,
+		TransmittanceMask,
+		View.HairScatteringLUTTexture,
+		HairScatteringLUTSampler,
+		View.HairComponents);
+}
+
 Texture2D<float4> ScreenShadowMaskSubPixelTexture;
 
 uint HairTransmittanceBufferMaxCount;
@@ -449,81 +478,55 @@ void DeferredLightPixelMain(
 
 	// todo add light culling/early out
 
-	// Compute Hair transmittance
-	{
-		float3 L = float3(0, 0, 1);
-		if (LightData.bRadialLight)
-			L = normalize(LightData.TranslatedWorldPosition - TranslatedWorldPosition);
-		else
-			L = LightData.Direction;
-
-		// Fetch precomputed transmittance
-		FHairTransmittanceMask TransmittanceMask = InitHairTransmittanceMask();
-
-		#if USE_VIRTUAL_SHADOW_MAP_MASK
-		{
-			const uint2 LocalPixelPos = PixelCoord.xy - uint2(View.ViewRectMinAndSize.xy);
-			const FCulledLightsGridData LightGridData = VirtualShadowMapGetLightsGrid(LocalPixelPos, SceneDepth);
-			uint LightCount = min(GetPackedShadowMaskMaxLightCount(), LightGridData.NumLights);
-
-			uint GridLightListIndex = 0;
-			for (; GridLightListIndex < LightCount; ++GridLightListIndex)
-			{
-				const FLocalLightData LightData = VirtualShadowMapGetLocalLightData(LightGridData, GridLightListIndex);
-				if (LightData.LightSceneId == LightSceneId)
-				{
-					break;
-				}
-			}
-			if (GridLightListIndex < GetPackedShadowMaskMaxLightCount())
-			{
-				const uint TotalSampleCount = NodeCount;
-				const uint TransmittanceSampleIndex = SampleIndex + GridLightListIndex * TotalSampleCount;
-				const uint TransmittanceSampleMaxCount = TotalSampleCount * GetPackedShadowMaskMaxLightCount();
-				if (TransmittanceSampleIndex < TransmittanceSampleMaxCount)
-				{
-					TransmittanceMask = UnpackTransmittanceMask(HairTransmittanceBuffer[TransmittanceSampleIndex]);
-				}
-			}
-		}
-		#else
-		if (SampleIndex < HairTransmittanceBufferMaxCount)
-		{
-			TransmittanceMask = UnpackTransmittanceMask(HairTransmittanceBuffer[SampleIndex]);
-		}
-		#endif
-
-
-		const float3 V = normalize(-CameraVector);
-		LightData.HairTransmittance = GetHairTransmittance(
-			V,
-			L,
-			HairScreenSpaceData.GBuffer,
-			TransmittanceMask,
-			View.HairScatteringLUTTexture,
-			HairScatteringLUTSampler,
-			View.HairComponents);
-	}
-
+	// Compute light attenuation from shadow mask, and fetch hair transmittance mask
+	// * Either use VSM shadow mask bits AND hair one pass transmittance buffer
+	// * Or use ShadowMask texture AND TransmittanceMask buffer
 	const float Dither = InterleavedGradientNoise(PixelCoord, View.StateFrameIndexMod8);
 	float4 LightAttenuation = 1.f;
+	FHairTransmittanceMask TransmittanceMask = InitHairTransmittanceMask();
 	#if USE_VIRTUAL_SHADOW_MAP_MASK
 	if (VirtualShadowMapId != INDEX_NONE)
 	{
+		// Fetch VSM shadow mask
+		uint GridLightListIndex = 0;
 		LightAttenuation = GetVirtualShadowMapMaskForLight(
 			ShadowMaskBits[PixelCoord],
 			PixelCoord,
 			SceneDepth,
 			VirtualShadowMapId,
-			TranslatedWorldPosition);
+			TranslatedWorldPosition,
+			GridLightListIndex);
+
+		// Fetch transmittance mask
+		if (GridLightListIndex < GetPackedShadowMaskMaxLightCount())
+		{
+			// Use VSM GridLightListIndex, which contains data only for shadowed light. 
+			// The index computation needs to be identical as in HairStrandsDeepTransmittanceMask.usf
+			const uint TotalSampleCount = NodeCount;
+			const uint TransmittanceSampleIndex = SampleIndex + GridLightListIndex * TotalSampleCount;
+			const uint TransmittanceSampleMaxCount = TotalSampleCount * GetPackedShadowMaskMaxLightCount();
+			if (TransmittanceSampleIndex < TransmittanceSampleMaxCount)
+			{
+				TransmittanceMask = UnpackTransmittanceMask(HairTransmittanceBuffer[TransmittanceSampleIndex]);
+			}
+		}
 	}
 	#else
 	if (HairShadowMaskValid)
 	{
 		LightAttenuation = ScreenShadowMaskSubPixelTexture.Load(uint3(PixelCoord, 0));
 	}
+	if (SampleIndex < HairTransmittanceBufferMaxCount)
+	{
+		TransmittanceMask = UnpackTransmittanceMask(HairTransmittanceBuffer[SampleIndex]);
+	}
 	#endif
 
+	// Compute hair transmittance data
+	{
+		LightData.HairTransmittance = GetTransmittanceDataFromTransmitttanceMask(LightData, HairScreenSpaceData.GBuffer, TransmittanceMask, TranslatedWorldPosition, CameraVector);
+	}
+
 	if (any(ShadowChannelMask < 1.0f))
 	{
 		LightAttenuation = dot(LightAttenuation, ShadowChannelMask).xxxx;
diff --git a/Engine/Shaders/Private/HairStrands/HairStrandsDeepTransmittanceMask.usf b/Engine/Shaders/Private/HairStrands/HairStrandsDeepTransmittanceMask.usf
--- a/Engine/Shaders/Private/HairStrands/HairStrandsDeepTransmittanceMask.usf
+++ b/Engine/Shaders/Private/HairStrands/HairStrandsDeepTransmittanceMask.usf
@@ -274,6 +274,8 @@ float4				ShadowChannelMask;
 Buffer<uint>		IndirectArgsBuffer;
 RWBuffer<uint>		OutTransmittanceMask;
 
+#include "../VirtualShadowMaps/VirtualShadowMapLightGrid.ush"
+
 #define USE_BLUE_NOISE 1
 #include "../BlueNoise.ush"
 DEFINE_HAIR_BLUE_NOISE_JITTER_FUNCTION
@@ -390,11 +392,9 @@ void MainCS(uint2 DispatchThreadId : SV_DispatchThreadID)
 
 	const uint2 LocalPixelPos = PixelCoord;
 
-	// TODO: Possibly make use of the new VSM culled light grid
-	const uint EyeIndex = 0; // TODO: Instanced stereo
-	const uint GridLinearIndex = ComputeLightGridCellIndex(LocalPixelPos, SceneDepth, EyeIndex);
-	const FCulledLightsGridData CulledLightGridData = GetCulledLightsGrid(GridLinearIndex, EyeIndex);
-
+	// Use VSM light grid index remapping to ensure we compute transmittance value only for shadowed light
+	// The logic to compute light index needs to be identical as in DeferredLightPixelShader.usf / ClusteredDeferredShadingPixelShader.usf
+	const FCulledLightsGridData CulledLightGridData = VirtualShadowMapGetLightsGrid(LocalPixelPos, SceneDepth);
 	const uint LightCount = min(GetPackedShadowMaskMaxLightCount(), CulledLightGridData.NumLights);
 
 	LOOP
@@ -403,7 +403,7 @@ void MainCS(uint2 DispatchThreadId : SV_DispatchThreadID)
 		const uint OutputIndex = SampleIndex + MaxVisibilityNodeCount * LightIndex;
 
 		// The Virtual Shadow Remap stores directional lights first
-		const FLocalLightData LightData 	 = GetLocalLightData(CulledLightGridData.DataStartIndex + LightIndex, EyeIndex);
+		const FLocalLightData LightData 	 = VirtualShadowMapGetLocalLightData(CulledLightGridData, LightIndex);
 		const float3 TranslatedLightPosition = LightData.LightPositionAndInvRadius.xyz; // LightPositionAndInvRadius is already in translated world space
 
 		// Early out when the shadow mask bits are zero
diff --git a/Engine/Shaders/Private/VirtualShadowMaps/VirtualShadowMapMaskBitsCommon.ush b/Engine/Shaders/Private/VirtualShadowMaps/VirtualShadowMapMaskBitsCommon.ush
--- a/Engine/Shaders/Private/VirtualShadowMaps/VirtualShadowMapMaskBitsCommon.ush
+++ b/Engine/Shaders/Private/VirtualShadowMaps/VirtualShadowMapMaskBitsCommon.ush
@@ -18,7 +18,8 @@ float GetVirtualShadowMapMaskForLight(
 	uint2 PixelPos,
 	float SceneDepth,
 	int VirtualShadowMapId,
-	float3 TranslatedWorldPosition)
+	float3 TranslatedWorldPosition,
+	inout uint OutLightIndex)
 {
 	// Light grid is relative to the view
 	const uint2 LocalPixelPos = PixelPos.xy - uint2(View.ViewRectMinAndSize.xy);
@@ -27,10 +28,10 @@ float GetVirtualShadowMapMaskForLight(
 	// We can only handle so many lights in our output encoding right now, so no purpose in computing more
 	uint LightCount = min(GetPackedShadowMaskMaxLightCount(), LightGridData.NumLights);
 
-	uint Index = 0;
-	for (; Index < LightCount; ++Index)
+	OutLightIndex = 0;
+	for (; OutLightIndex < LightCount; ++OutLightIndex)
 	{
-		const FLocalLightData LightData = VirtualShadowMapGetLocalLightData(LightGridData, Index);
+		const FLocalLightData LightData = VirtualShadowMapGetLocalLightData(LightGridData, OutLightIndex);
 		if (LightData.VirtualShadowMapId == VirtualShadowMapId)
 		{
 			break;
@@ -38,10 +39,10 @@ float GetVirtualShadowMapMaskForLight(
 	}
 
 	float LightAttenuation = 1.0f;
-	if (Index < LightCount)
+	if (OutLightIndex < LightCount)
 	{
 		// TODO: Dither!
-		LightAttenuation = UnpackShadowMask(ShadowMask, Index);
+		LightAttenuation = UnpackShadowMask(ShadowMask, OutLightIndex);
 	}
 	else
 	{
@@ -57,3 +58,14 @@ float GetVirtualShadowMapMaskForLight(
 
 	return LightAttenuation;
 }
+
+float GetVirtualShadowMapMaskForLight(
+	uint4 ShadowMask,
+	uint2 PixelPos,
+	float SceneDepth,
+	int VirtualShadowMapId,
+	float3 TranslatedWorldPosition)
+{
+	uint OutDummyLightIndex = 0;
+	return GetVirtualShadowMapMaskForLight(ShadowMask, PixelPos, SceneDepth, VirtualShadowMapId, TranslatedWorldPosition, OutDummyLightIndex);
+}
\ No newline at end of file
diff --git a/Engine/Source/Runtime/Renderer/Private/LightRendering.cpp b/Engine/Source/Runtime/Renderer/Private/LightRendering.cpp
--- a/Engine/Source/Runtime/Renderer/Private/LightRendering.cpp
+++ b/Engine/Source/Runtime/Renderer/Private/LightRendering.cpp
@@ -178,7 +178,7 @@ static TAutoConsoleVariable<float> CVarContactShadowsOverrideNonShadowCastingInt
 
 static TAutoConsoleVariable<int32> CVarHairStrandsAllowOneTransmittancePass(
 	TEXT("r.HairStrands.Lighting.AllowOneTransmittancePass"),
-	0,
+	1,
 	TEXT("Allows one transmittance pass for hair strands lighting to have better performance (experimental).\n"),
 	ECVF_RenderThreadSafe);
 
@@ -2108,6 +2108,9 @@ void FDeferredShadingSceneRenderer::RenderLights(
 						if (HairStrands::HasViewHairStrandsData(View))
 						{
 							// If the light elided the screen space shadow mask, sample directly from the packed shadow mask
+							// Note: this is only used when VSM one pass is enable AND hair one pass transmittance is enabled. 
+							// Hence the condition using bElideScreenShadowMaskSubPixel, instead of bElideScreenShadowMask for 
+							// computing VirtualShadowMapId
 							int32 VirtualShadowMapId = INDEX_NONE;
 							if (bElideScreenShadowMaskSubPixel)
 							{
