Fix tonemapper CS failing to properly account for ViewRect offset.

diff --git a/Engine/Shaders/Private/PostProcessTonemap.usf b/Engine/Shaders/Private/PostProcessTonemap.usf
--- a/Engine/Shaders/Private/PostProcessTonemap.usf
+++ b/Engine/Shaders/Private/PostProcessTonemap.usf
@@ -614,14 +614,13 @@ RWTexture2D<float> RWOutputLuminance;
 void MainCS(uint2 DispatchThreadId : SV_DispatchThreadID)
 {
 	float4 SvPosition = float4((float2)DispatchThreadId + Output_ViewportMin + 0.5f, 0.0f, 1.0f);
-	float2 UVScreen = SvPosition.xy * Output_ExtentInverse;
+	float2 UV = SvPosition.xy * Output_ExtentInverse;
 
-	if (IsComputeUVOutOfBounds(UVScreen))
+	if (IsComputeUVOutOfBounds(UV))
 	{
 		return;
 	}
 
-	float2 UV = ((float2)DispatchThreadId + 0.5f) * Output_ExtentInverse;
 	float4 InScreenPos = float4(UV * float2(2, -2) - float2(1, -1), 0, 1);
 
 	float2 Vignette;
