Fix compilation error when enabling diagnostic buffer output for shader print.

The buffer type was recently updated for accomodate some platform shader compiler error. This experimental pass with diagnostic buffer hasn't been updated, causing a compilation error.

diff --git a/Engine/Shaders/Private/ShaderPrint.ush b/Engine/Shaders/Private/ShaderPrint.ush
--- a/Engine/Shaders/Private/ShaderPrint.ush
+++ b/Engine/Shaders/Private/ShaderPrint.ush
@@ -115,7 +115,8 @@ FShaderPrintBuffers InitShaderPrintContextBuffers()
 {
 	FShaderPrintBuffers Buffers;
 #if SHADER_PRINT_USE_DIAGNOSTIC_BUFFER
-	Buffers.RWEntryBuffer = UEDiagnosticBuffer;
+	// Nothing to do as Buffers.RWEntryBuffer is never used when using diagnostic buffer. 
+	// This is ensured by the macro SHADER_PRINT_RWENTRYBUFFER, which defines how writing is handled.
 #else
 	Buffers.StateBuffer = ShaderPrint_StateBuffer;
 	Buffers.RWEntryBuffer = ShaderPrint_RWEntryBuffer;
