Path Tracer: Fix SLW rendering black due to recent changes to the glass BSDF

diff --git a/Engine/Shaders/Private/PathTracing/PathTracingMaterialHitShader.usf b/Engine/Shaders/Private/PathTracing/PathTracingMaterialHitShader.usf
--- a/Engine/Shaders/Private/PathTracing/PathTracingMaterialHitShader.usf
+++ b/Engine/Shaders/Private/PathTracing/PathTracingMaterialHitShader.usf
@@ -907,13 +907,13 @@ RAY_TRACING_ENTRY_CLOSEST_HIT(PathTracingMaterialCHS,
 		const float3 WaterExtinction = ScatteringCoeff + AbsorptionCoeff;
 		const float3 WaterAlbedo = select(WaterExtinction > 0.0, ScatteringCoeff / WaterExtinction, float3(0.0, 0.0, 0.0));
 
-		Payload.ShadingModelID = SHADINGMODELID_DEFAULT_LIT; // replace shading model to avoid extra cases in raygen
+		Payload.ShadingModelID = SHADINGMODELID_SOLID_GLASS; // replace shading model to avoid extra cases in raygen
 		Payload.BaseColor = WaterAlbedo; // TODO: unused at the moment
 		Payload.SetExtinction(WaterExtinction);
 		Payload.Ior = DielectricF0ToIor(DielectricSpecularToF0(Specular));
 		// we are using actual refraction, so don't apply any transparency
 		Payload.TransparencyColor = 0.0;
-		Payload.BSDFOpacity = 0.0; // only shade with the glass portion
+		Payload.BSDFOpacity = 1.0;
 	}
 	else
 #endif
