Prevent application of primitive holdouts during reflection captures.

#rnx

diff --git a/Engine/Shaders/Private/BasePassPixelShader.usf b/Engine/Shaders/Private/BasePassPixelShader.usf
--- a/Engine/Shaders/Private/BasePassPixelShader.usf
+++ b/Engine/Shaders/Private/BasePassPixelShader.usf
@@ -2165,21 +2165,25 @@ void FPixelShaderInOut_MainPS(
 	// Only include the holdout logic if the Renderer is forwarding its alpha channel to post-processing
 	if (GetPrimitiveData(MaterialParameters).Flags & PRIMITIVE_SCENE_DATA_FLAG_HOLDOUT)
 	{
-		// Discard this pixel so that nothing gets written to the G-Buffers and alpha is kept as its default value (background)
+		// Ignore primitive holdout during reflection captures
+		if(ResolvedView.RenderingReflectionCaptureMask == 0)
+		{
+			// Discard this pixel so that nothing gets written to the G-Buffers and alpha is kept as its default value (background)
 
 #if MATERIALBLENDING_SOLID || MATERIALBLENDING_MASKED
-		// This catches the common case as well as handles cases where Opacity has a different meaning, such as for Subsurface materials
-		Out.MRT[0] = half4(0.0, 0.0, 0.0, 1.0);
+			// This catches the common case as well as handles cases where Opacity has a different meaning, such as for Subsurface materials
+			Out.MRT[0] = half4(0.0, 0.0, 0.0, 1.0);
 #else
-		// TODO: The translucent case doesn't seem to get the right alpha (compared to the path traced reference)
-		Out.MRT[0] = half4(0.0, 0.0, 0.0, Opacity);
+			// TODO: The translucent case doesn't seem to get the right alpha (compared to the path traced reference)
+			Out.MRT[0] = half4(0.0, 0.0, 0.0, Opacity);
 #endif
 
 #if USES_GBUFFER
-		Out.MRT[1] = 0.0;
-		Out.MRT[2] = 0.0;
-		Out.MRT[3] = 0.0;
+			Out.MRT[1] = 0.0;
+			Out.MRT[2] = 0.0;
+			Out.MRT[3] = 0.0;
 #endif
+		}
 	}
 #endif
 
