Fix for subtractive blend when using local merge

diff --git a/Engine/Shaders/Private/Landscape/LandscapeLayersWeightmapsPS.usf b/Engine/Shaders/Private/Landscape/LandscapeLayersWeightmapsPS.usf
--- a/Engine/Shaders/Private/Landscape/LandscapeLayersWeightmapsPS.usf
+++ b/Engine/Shaders/Private/Landscape/LandscapeLayersWeightmapsPS.usf
@@ -188,7 +188,7 @@ void MergeEditLayers(in float4 SVPos : SV_POSITION, out float OutColor : SV_Targ
 	LOOP
 	for (uint i = 1; i < InNumEditLayers; ++i)
 	{
-		FEditLayerWeightmapMergeInfo LayerMergeInfo = InEditLayersMergeInfos[i];
+		FEditLayerWeightmapMergeInfo LayerMergeInfo = InEditLayersMergeInfos[InStartIndexInEditLayersMergeInfos + i];
 		float LayerWeight = SampleEditLayerWeightmap(TextureCoordinates, i);
 			
 		if (LayerMergeInfo.BlendMode == EEDITLAYERWEIGHTMAPBLENDMODE_ADDITIVE)
