Fix for GPU hang with work graph style nanite material shader bundles.
When we have binned quads that hang off the edge of the screen, the off screen helper lanes are shaded and cause GPU hangs.This is probably because of a lack of driver bounds checking in work graph shaders.
Solution is just to early out for the off screen helper lanes and take the hit of bad derivatives on odd edge pixels.

diff --git a/Engine/Shaders/Private/ComputeShaderOutputCommon.ush b/Engine/Shaders/Private/ComputeShaderOutputCommon.ush
--- a/Engine/Shaders/Private/ComputeShaderOutputCommon.ush
+++ b/Engine/Shaders/Private/ComputeShaderOutputCommon.ush
@@ -367,5 +367,17 @@ void MainCS(
 	}
 
 	const float2 SVPositionXY = PixelPos + ((1u << VRSShift) * 0.5f);
+
+#if WORKGRAPH_NODE
+	// Fix for GPU hang. Without this a lack of out of bounds checking in Work Graph shaders causes out of bounds access of VisBuffer.
+	// Note that derivatives can be incorrect on odd edge pixels, and VRS edge pixels may get rejected.
+	// But that's also true in general for the standard compute path.
+	BRANCH
+	if (any(SVPositionXY > float2(View.ViewRectMinAndSize.xy + View.ViewRectMinAndSize.zw)))
+	{
+		return;
+	}
+#endif
+
 	ProcessPixel(ShadingBin, PixelPos, SVPositionXY, ElementIndex, PixelIndex, PixelWriteMask, HelperLaneCount);
 }
