Fix overbright hair shading when roughness is close to 0.

The gaussian normalization causes the BSDF Mp term to reach very high value. Clamping this normalization limits the values. This change is equivalent to have min roughness value, but only for the normalization term. So the highlight remains sharp, but the energy is kept under sensible value.

diff --git a/Engine/Shaders/Private/HairBsdf.ush b/Engine/Shaders/Private/HairBsdf.ush
--- a/Engine/Shaders/Private/HairBsdf.ush
+++ b/Engine/Shaders/Private/HairBsdf.ush
@@ -9,7 +9,11 @@
 
 float Hair_g(float B, float Theta)
 {
-	return exp(-0.5 * Pow2(Theta) / (B * B)) / (sqrt(2 * PI) * B);
+	// Clamp B for the denominator term, as otherwise the Gaussian normalization returns too high value.
+	// This clamps allow to prevent large value for low roughness, while keeping the highlight shape/sharpness 
+	// similar.
+	const float DenominatorB = max(B, 0.01f);
+	return exp(-0.5 * Pow2(Theta) / (B * B)) / (sqrt(2 * PI) * DenominatorB);
 }
 
 float Hair_F(float CosTheta)
