Hair Strands: Fix shader warnings from new DXC compiler

diff --git a/Engine/Shaders/Private/HairStrands/HairStrandsComposition.usf b/Engine/Shaders/Private/HairStrands/HairStrandsComposition.usf
--- a/Engine/Shaders/Private/HairStrands/HairStrandsComposition.usf
+++ b/Engine/Shaders/Private/HairStrands/HairStrandsComposition.usf
@@ -308,6 +308,11 @@ void MainPS(
 	float4 OutGBufferD = 0;
 	float4 OutGBufferE = 0;
 	float  OutDepth = 0;
+#elif PERMUTATION_OUTPUT_TYPE == 1
+	OutGBufferC = 0;
+	OutGBufferD = 0;
+	OutGBufferE = 0;
+	OutDepth = 0;
 #endif
 
 	const uint3 PixelCoord = uint3(floor(SvPosition.xy), 0);
diff --git a/Engine/Shaders/Private/HairStrands/HairStrandsHitProxy.usf b/Engine/Shaders/Private/HairStrands/HairStrandsHitProxy.usf
--- a/Engine/Shaders/Private/HairStrands/HairStrandsHitProxy.usf
+++ b/Engine/Shaders/Private/HairStrands/HairStrandsHitProxy.usf
@@ -20,6 +20,8 @@ void EmitPS(
 	out float4 OutHitProxyId : SV_Target0,
 	out float OutDepth : SV_Depth)
 {
+	OutHitProxyId = 0;
+	OutDepth = 0;
 	const uint2 PixelCoord = (uint2)In.Position.xy;
 
 	// Take the first sample
@@ -65,7 +67,7 @@ void EmitPS(
 	in float4 InPosition : SV_POSITION,
 	out float OutDepth : SV_Depth)
 {
-
+	OutDepth = 0;
 	const float2 UV = InPosition.xy * InvViewportResolution;
 	const uint2 PixelCoord = UV * View.ViewSizeAndInvSize.xy;
 	const bool bValid = HairOnlyDepthTexture.Load(uint3(PixelCoord, 0)) > 0;
diff --git a/Engine/Shaders/Private/HairStrands/HairStrandsVoxelRasterCompute.usf b/Engine/Shaders/Private/HairStrands/HairStrandsVoxelRasterCompute.usf
--- a/Engine/Shaders/Private/HairStrands/HairStrandsVoxelRasterCompute.usf
+++ b/Engine/Shaders/Private/HairStrands/HairStrandsVoxelRasterCompute.usf
@@ -1013,6 +1013,7 @@ void MainPS(
 	in float3 InTranslatedWorldPosition : WORLD_POSITION,
 	out float OutDepth : SV_DEPTH)
 {
+	OutDepth = 0;
 	const float2 PixelCoord = InPosition.xy;
 	const float2 UV = PixelCoord / float2(OutputResolution); // todo view rect offset
 
