Substrate - fixed unlit view mode with clear coat.

diff --git a/Engine/Shaders/Private/BasePassPixelShader.usf b/Engine/Shaders/Private/BasePassPixelShader.usf
--- a/Engine/Shaders/Private/BasePassPixelShader.usf
+++ b/Engine/Shaders/Private/BasePassPixelShader.usf
@@ -1733,7 +1733,7 @@ void FPixelShaderInOut_MainPS(
 
 		// Unlit view mode
 	#if USE_DEVELOPMENT_SHADERS
-		Color  = lerp(Color, TopLayerData.BaseColor, View.UnlitViewmodeMask);
+		Color  = lerp(Color, TopLayerData.UnlitViewBaseColor, View.UnlitViewmodeMask);
 		EmissiveLuminance += OutOfBoundsMaskLuminance;
 	#endif
 
diff --git a/Engine/Shaders/Private/Substrate/Substrate.ush b/Engine/Shaders/Private/Substrate/Substrate.ush
--- a/Engine/Shaders/Private/Substrate/Substrate.ush
+++ b/Engine/Shaders/Private/Substrate/Substrate.ush
@@ -733,7 +733,7 @@ FSubstrateIrradianceAndOcclusion SubstrateUnpackIrradianceAndOcclusion(uint In /
 
 struct FSubstrateTopLayerData
 {
-	float3 BaseColor; // This data is only available during the base pass. It is not packed/stored otherwise.
+	float3 UnlitViewBaseColor; // This data is only available during the base pass. It is not packed/stored otherwise.
 	float3 WorldNormal;
 	float  Roughness;
 	uint   Material;
diff --git a/Engine/Shaders/Private/Substrate/SubstrateExport.ush b/Engine/Shaders/Private/Substrate/SubstrateExport.ush
--- a/Engine/Shaders/Private/Substrate/SubstrateExport.ush
+++ b/Engine/Shaders/Private/Substrate/SubstrateExport.ush
@@ -212,11 +212,11 @@ void PackSubstrateOut(
 							if (Haziness.bSimpleClearCoat)
 							{
 								TopLayerRoughnessContribution = BSDF.TopLayerDataWeight * lerp(SLAB_ROUGHNESS(BSDF), Haziness.Roughness, Haziness.Weight);
-								TopLayerBaseColorContribution = BSDF.TopLayerDataWeight * lerp(0.0, 0.04f, F0RGBToMetallic(0.04f));
+								//TopLayerBaseColorContribution = BSDF.TopLayerDataWeight * lerp(0.0, 0.04f, F0RGBToMetallic(0.04f)); Skip because, as for the legacy path, we only show the bottom BaseColor of Clear Coat materials.
 							}
 						}
-						TopLayerData.Roughness += TopLayerRoughnessContribution;
-						TopLayerData.BaseColor += TopLayerBaseColorContribution;
+						TopLayerData.Roughness			+= TopLayerRoughnessContribution;
+						TopLayerData.UnlitViewBaseColor += TopLayerBaseColorContribution;
 
 						TopLayerData.Material = TOP_LAYER_MATERIAL_VALID;
 
@@ -238,9 +238,9 @@ void PackSubstrateOut(
 						bIsSingleMaterial = false;
 						EmissiveColor += BSDF_GETEMISSIVE(BSDF) * BSDF.LuminanceWeightV;
 
-						TopLayerTotalWeight += BSDF.TopLayerDataWeight;
-						TopLayerData.Roughness += BSDF.TopLayerDataWeight * FullyRough;
-						TopLayerData.BaseColor += BSDF.TopLayerDataWeight * SubstrateGetBSDFBaseColor(BSDF);
+						TopLayerTotalWeight				+= BSDF.TopLayerDataWeight;
+						TopLayerData.Roughness			+= BSDF.TopLayerDataWeight * FullyRough;
+						TopLayerData.UnlitViewBaseColor += BSDF.TopLayerDataWeight * SubstrateGetBSDFBaseColor(BSDF);
 
 						TopLayerData.Material = TOP_LAYER_MATERIAL_VALID;
 
@@ -272,9 +272,9 @@ void PackSubstrateOut(
 						bIsSingleMaterial = false;
 						EmissiveColor += BSDF_GETEMISSIVE(BSDF) * BSDF.LuminanceWeightV;
 
-						TopLayerTotalWeight += BSDF.TopLayerDataWeight;
-						TopLayerData.Roughness += BSDF.TopLayerDataWeight * SLW_ROUGHNESS(BSDF);
-						TopLayerData.BaseColor += BSDF.TopLayerDataWeight * SubstrateGetBSDFBaseColor(BSDF);
+						TopLayerTotalWeight				+= BSDF.TopLayerDataWeight;
+						TopLayerData.Roughness			+= BSDF.TopLayerDataWeight * SLW_ROUGHNESS(BSDF);
+						TopLayerData.UnlitViewBaseColor += BSDF.TopLayerDataWeight * SubstrateGetBSDFBaseColor(BSDF);
 
 						TopLayerData.Material = TOP_LAYER_MATERIAL_SLWATER;
 
@@ -299,9 +299,9 @@ void PackSubstrateOut(
 		}
 
 		// Finalize top layer data
-		TopLayerData.WorldNormal= TopLayerTotalWeight > 0.0f	? normalize(TopLayerData.WorldNormal)			: 0.0f;
-		TopLayerData.Roughness	= TopLayerTotalWeight > 0.0f	? TopLayerData.Roughness / TopLayerTotalWeight	: 0.0f;
-		TopLayerData.BaseColor	= TopLayerTotalWeight > 0.0f	? TopLayerData.BaseColor / TopLayerTotalWeight	: 0.0f;
+		TopLayerData.WorldNormal		= TopLayerTotalWeight > 0.0f	? normalize(TopLayerData.WorldNormal)					: 0.0f;
+		TopLayerData.Roughness			= TopLayerTotalWeight > 0.0f	? TopLayerData.Roughness / TopLayerTotalWeight			: 0.0f;
+		TopLayerData.UnlitViewBaseColor	= TopLayerTotalWeight > 0.0f	? TopLayerData.UnlitViewBaseColor / TopLayerTotalWeight	: 0.0f;
 
 		// Set storage layout as either: fast(0), single(1), or complex(2)
 		// We check BSDFVisibleCount to make sure we do not set a MATERIALMODE that would enforce a BSDFVisibleCount==1.
