Fix vulkan tonemap subpass only rendering to left eye.

diff --git a/Engine/Shaders/Private/PostProcessTonemap.usf b/Engine/Shaders/Private/PostProcessTonemap.usf
--- a/Engine/Shaders/Private/PostProcessTonemap.usf
+++ b/Engine/Shaders/Private/PostProcessTonemap.usf
@@ -28,7 +28,9 @@
 SCREEN_PASS_TEXTURE_VIEWPORT(Color)
 SCREEN_PASS_TEXTURE_VIEWPORT(Output)
 
-#if USE_MULTI_VIEW_ID_SV
+#define MOBILE_MULTI_VIEW_FALLBACK (MOBILE_MULTI_VIEW && INSTANCED_STEREO)
+
+#if USE_MULTI_VIEW_ID_SV || MOBILE_MULTI_VIEW_FALLBACK
 Texture2DArray ColorTexture;
 #else
 Texture2D ColorTexture;
@@ -256,7 +258,7 @@ float max4(float4 x)
 float4 SampleSceneColor(float2 SceneUV, float ArrayIndex)
 {
 	SceneUV = clamp(SceneUV, Color_UVViewportBilinearMin, Color_UVViewportBilinearMax);
-	#if USE_MULTI_VIEW_ID_SV
+	#if USE_MULTI_VIEW_ID_SV || MOBILE_MULTI_VIEW_FALLBACK
 	return Texture2DArraySample(ColorTexture, ColorSampler, float3(SceneUV.xy,ArrayIndex));
 	#else
 	return Texture2DSample(ColorTexture, ColorSampler, SceneUV);
@@ -662,7 +664,7 @@ half4 FetchAndResolveSceneColor(float2 UV, float ArrayIndex)
 		SceneColor = half4(VulkanSubpassFetch0());
 	#endif
 #else
-	#if USE_MULTI_VIEW_ID_SV
+	#if USE_MULTI_VIEW_ID_SV || MOBILE_MULTI_VIEW_FALLBACK
 		SceneColor = Texture2DArraySample(ColorTexture, ColorSampler, float3(UV.xy,ArrayIndex));
 	#else
 		SceneColor = Texture2DSample(ColorTexture, ColorSampler, UV);
@@ -672,24 +674,31 @@ half4 FetchAndResolveSceneColor(float2 UV, float ArrayIndex)
 }
 
 void MobileCustomResolve_MainPS(
+	in float4 SvPosition							: SV_POSITION, 
 	in noperspective float2 UV						: TEXCOORD0,
-	in				 float4 SvPosition				: SV_POSITION,		// after all interpolators
+#if MOBILE_MULTI_VIEW_FALLBACK
+	in nointerpolation uint EyeIndex				: TEXCOORD1,
+#endif
 	out				 float4 OutColor				: SV_Target0
 #if USE_MULTI_VIEW_ID_SV
 	, in nointerpolation uint MultiViewId : SV_ViewID
 #endif
 )
 {
+	uint ArrayIndex = 0;
 #if USE_MULTI_VIEW_ID_SV
 	ResolvedView = ResolveView(MultiViewId);
+	ArrayIndex = MultiViewId;
+#elif MOBILE_MULTI_VIEW_FALLBACK
+	ResolvedView = ResolveView(EyeIndex);
+	ArrayIndex = EyeIndex;
 #else
-	uint MultiViewId = 0;
 	ResolvedView = ResolveView();
-#endif // USE_MULTI_VIEW_ID_SV
+#endif
 	
 	const float OneOverPreExposure = ResolvedView.OneOverPreExposure;
 
-	half4 SceneColor = FetchAndResolveSceneColor(UV, MultiViewId);
+	half4 SceneColor = FetchAndResolveSceneColor(UV, ArrayIndex);
 	SceneColor.a = 1.0;
 
 	half3 LinearColor = SceneColor.rgb * OneOverPreExposure;
@@ -700,3 +709,4 @@ void MobileCustomResolve_MainPS(
 	OutColor.a = 1.0;
 }
 
+
diff --git a/Engine/Shaders/Private/Tools/FullscreenVertexShader.usf b/Engine/Shaders/Private/Tools/FullscreenVertexShader.usf
--- a/Engine/Shaders/Private/Tools/FullscreenVertexShader.usf
+++ b/Engine/Shaders/Private/Tools/FullscreenVertexShader.usf
@@ -9,17 +9,29 @@
 #ifndef SCREEN_VS_FOR_INSTANCED_VIEWS
 	#define SCREEN_VS_FOR_INSTANCED_VIEWS 0
 #endif
+#ifndef SCREEN_VS_FOR_MOBILE_MULTI_VIEW
+	#define SCREEN_VS_FOR_MOBILE_MULTI_VIEW 0
+#endif
+
+#if SCREEN_VS_FOR_MOBILE_MULTI_VIEW
+	#include "../InstancedStereo.ush"
+#endif
 
 void MainVS(
 	float2 ViewportUV : ATTRIBUTE0,
 	float2 UV         : ATTRIBUTE1, // TODO: kill
-#if !SCREEN_VS_FOR_INSTANCED_VIEWS
 	out float4 OutPosition : SV_POSITION
-#else
-	uint InstanceId	  : SV_InstanceID,
-	out float4 OutPosition : SV_POSITION,
-	out uint OutViewportId	: SV_ViewportArrayIndex,
-	out nointerpolation uint OutViewId : VIEW_ID
+#if SCREEN_VS_FOR_INSTANCED_VIEWS
+	,	uint InstanceId	  : SV_InstanceID
+	,	out uint OutViewportId	: SV_ViewportArrayIndex
+	,	out nointerpolation uint OutViewId : VIEW_ID
+#elif SCREEN_VS_FOR_MOBILE_MULTI_VIEW
+	,	out noperspective float2 OutUV : TEXCOORD0
+	#if INSTANCED_STEREO // Mobile multi view fallback path
+	,	in uint InstanceId : SV_InstanceID
+	,	out nointerpolation uint EyeIndex : TEXCOORD1
+	,	out uint LayerIndex : SV_RenderTargetArrayIndex
+	#endif
 #endif
 	)
 {
@@ -29,6 +41,14 @@ void MainVS(
 	OutViewportId = InstanceId;
 	OutViewId = InstanceId;
 #endif
+	
+#if SCREEN_VS_FOR_MOBILE_MULTI_VIEW
+	OutUV = UV;
+	#if INSTANCED_STEREO // Mobile multi view fallback path
+		EyeIndex = GetEyeIndex(InstanceId);
+		LayerIndex = EyeIndex;
+	#endif
+#endif
 }
 
 float4 PreTransform;
diff --git a/Engine/Source/Runtime/RHI/Private/RHI.cpp b/Engine/Source/Runtime/RHI/Private/RHI.cpp
--- a/Engine/Source/Runtime/RHI/Private/RHI.cpp
+++ b/Engine/Source/Runtime/RHI/Private/RHI.cpp
@@ -1504,7 +1504,6 @@ void FRHIRenderPassInfo::ConvertToRenderTargetsInfo(FRHISetRenderTargetsInfo& Ou
 
 		OutRTInfo.bClearColor |= (LoadAction == ERenderTargetLoadAction::EClear);
 
-		ensure(!OutRTInfo.bHasResolveAttachments || ColorRenderTargets[Index].RenderTarget->GetNumSamples() == 1);
 		if (ColorRenderTargets[Index].ResolveTarget)
 		{
 			OutRTInfo.bHasResolveAttachments = true;
diff --git a/Engine/Source/Runtime/RenderCore/Private/CommonRenderResources.cpp b/Engine/Source/Runtime/RenderCore/Private/CommonRenderResources.cpp
--- a/Engine/Source/Runtime/RenderCore/Private/CommonRenderResources.cpp
+++ b/Engine/Source/Runtime/RenderCore/Private/CommonRenderResources.cpp
@@ -17,6 +17,7 @@ TGlobalResource<FScreenRectangleIndexBuffer, FRenderResource::EInitPhase::Pre> G
 
 IMPLEMENT_GLOBAL_SHADER(FScreenVertexShaderVS, "/Engine/Private/Tools/FullscreenVertexShader.usf", "MainVS", SF_Vertex);
 IMPLEMENT_GLOBAL_SHADER(FInstancedScreenVertexShaderVS, "/Engine/Private/Tools/FullscreenVertexShader.usf", "MainVS", SF_Vertex);
+IMPLEMENT_GLOBAL_SHADER(FMobileMultiViewVertexShaderVS, "/Engine/Private/Tools/FullscreenVertexShader.usf", "MainVS", SF_Vertex);
 IMPLEMENT_GLOBAL_SHADER(FCopyRectPS, "/Engine/Private/ScreenPass.usf", "CopyRectPS", SF_Pixel);
 IMPLEMENT_GLOBAL_SHADER(FImagePreTransformVS, "/Engine/Private/Tools/FullscreenVertexShader.usf", "MainForPreTransform", SF_Vertex);
 
@@ -26,6 +27,12 @@ bool FInstancedScreenVertexShaderVS::ShouldCompilePermutation(const FGlobalShade
 	return Aspects.IsInstancedMultiViewportEnabled();
 }
 
+bool FMobileMultiViewVertexShaderVS::ShouldCompilePermutation(const FGlobalShaderPermutationParameters& Parameters)
+{
+	UE::StereoRenderUtils::FStereoShaderAspects Aspects(Parameters.Platform);
+	return Aspects.IsMobileMultiViewEnabled();
+}
+
 void FScreenRectangleVertexBuffer::InitRHI(FRHICommandListBase& RHICmdList)
 {
 	TResourceArray<FFilterVertex, VERTEXBUFFER_ALIGNMENT> Vertices;
diff --git a/Engine/Source/Runtime/RenderCore/Public/CommonRenderResources.h b/Engine/Source/Runtime/RenderCore/Public/CommonRenderResources.h
--- a/Engine/Source/Runtime/RenderCore/Public/CommonRenderResources.h
+++ b/Engine/Source/Runtime/RenderCore/Public/CommonRenderResources.h
@@ -132,7 +132,26 @@ class FInstancedScreenVertexShaderVS : public FScreenVertexShaderVS
 	}
 
 	BEGIN_SHADER_PARAMETER_STRUCT(FParameters, )
-		END_SHADER_PARAMETER_STRUCT()
+	END_SHADER_PARAMETER_STRUCT()
+};
+
+/** Vertex shader to draw a quad covering all the viewports with mobile multi view (SV_RenderTargetArrayIndex is output for each SV_InstanceID if using the multi view fallback path). Does not have any shader parameters.
+ * The pixel shader should expect UV in TEXCOORD0 (and EyeIndex in TEXCOORD1 if using the mobile multi view fallback path). */
+class FMobileMultiViewVertexShaderVS : public FScreenVertexShaderVS
+{
+	DECLARE_EXPORTED_GLOBAL_SHADER(FMobileMultiViewVertexShaderVS, RENDERCORE_API);
+	SHADER_USE_PARAMETER_STRUCT(FMobileMultiViewVertexShaderVS, FScreenVertexShaderVS);
+
+	static bool ShouldCompilePermutation(const FGlobalShaderPermutationParameters& Parameters);
+
+	static void ModifyCompilationEnvironment(const FGlobalShaderPermutationParameters& Parameters, FShaderCompilerEnvironment& OutEnvironment)
+	{
+		FScreenVertexShaderVS::ModifyCompilationEnvironment(Parameters, OutEnvironment);
+		OutEnvironment.SetDefine(TEXT("SCREEN_VS_FOR_MOBILE_MULTI_VIEW"), 1);
+	}
+
+	BEGIN_SHADER_PARAMETER_STRUCT(FParameters, )
+	END_SHADER_PARAMETER_STRUCT()
 };
 
 /** Pixel shader to copy pixels from src to dst performing a format change that works on all platforms. */
diff --git a/Engine/Source/Runtime/Renderer/Private/PostProcess/PostProcessTonemap.cpp b/Engine/Source/Runtime/Renderer/Private/PostProcess/PostProcessTonemap.cpp
--- a/Engine/Source/Runtime/Renderer/Private/PostProcess/PostProcessTonemap.cpp
+++ b/Engine/Source/Runtime/Renderer/Private/PostProcess/PostProcessTonemap.cpp
@@ -23,6 +23,7 @@
 #include "HDRHelper.h"
 #include "VariableRateShadingImageManager.h"
 #include "DataDrivenShaderPlatformInfo.h"
+#include "CommonRenderResources.h"
 
 
 bool SupportsFilmGrain(EShaderPlatform Platform)
@@ -1091,7 +1092,7 @@ void RenderMobileCustomResolve(FRHICommandList& RHICmdList, const FViewInfo& Vie
 	const float LUTSize = ColorGradingLUT ? (float)ColorGradingLUT->GetDesc().GetSize().Y : /* unused (default): */ 32.0f;
 	const FPostProcessSettings& Settings = View.FinalPostProcessSettings;
 	
-	TShaderMapRef<FScreenVS> VertexShader(View.ShaderMap);
+	TShaderMapRef<FMobileMultiViewVertexShaderVS> VertexShader(View.ShaderMap);
 	
 	FMobileCustomResolvePS::FPermutationDomain PermutationVector;
 	PermutationVector.Set<FMobileCustomResolvePS::FTonemapperSubpassMsaaDim>(SubpassMSAASamples);
@@ -1137,7 +1138,8 @@ void RenderMobileCustomResolve(FRHICommandList& RHICmdList, const FViewInfo& Vie
 		TargetSize,
 		TargetSize,
 		VertexShader,
-		EDRF_UseTriangleOptimization);
+		EDRF_UseTriangleOptimization,
+		View.InstanceFactor);
 }
 
 BEGIN_SHADER_PARAMETER_STRUCT(FMobileCustomResolveParameters, )
@@ -1156,6 +1158,8 @@ void AddMobileCustomResolvePass(FRDGBuilder& GraphBuilder, const FViewInfo& View
 	PassParameters->ColorTexture = SceneTextures.Color.Resolve;
 	PassParameters->ColorGradingLUT = ColorGradingLUT;
 	PassParameters->RenderTargets[0] = FRenderTargetBinding(ViewFamilyTexture, ERenderTargetLoadAction::EClear);
+	// Need to specify multi view count for when the custom resolve pass is a separate pass and is not within the main pass.
+	PassParameters->RenderTargets.MultiViewCount = View.bIsMobileMultiViewEnabled ? 2 : 0;
 
 	GraphBuilder.AddPass(
 		RDG_EVENT_NAME("MobileCustomResolvePass"),
