Remove VRS specific coverage modifiers from the base pass pixel shader, now that masked materials (the only materials that output coverage) now have VRS automatically disabled.

diff --git a/Engine/Shaders/Private/BasePassPixelShader.usf b/Engine/Shaders/Private/BasePassPixelShader.usf
--- a/Engine/Shaders/Private/BasePassPixelShader.usf
+++ b/Engine/Shaders/Private/BasePassPixelShader.usf
@@ -2186,19 +2186,6 @@ void FPixelShaderInOut_MainPS(
 		}
 	}
 #endif
-
-// Duplicate coverage bit values in case multiple pixels are written by one invocation (VRS)
-#if OUTPUTS_COVERAGE
-	// If centroids are integers, it means we're at a non-1x1 shading rate
-	if(View.NumSceneColorMSAASamples > 1 && (frac(In.SvPosition.x) == 0 || frac(In.SvPosition.y) == 0)) 
-	{
-		Out.Coverage &= 0xF;
-		Out.Coverage |= (Out.Coverage << 4);
-		Out.Coverage |= (Out.Coverage << 8);
-		Out.Coverage |= (Out.Coverage << 12);
-	}
-#endif
-
 }
 
 #if SUBSTRATE_OPAQUE_DEFERRED
