Fixed crash on iOS when mobile deferred shading is enabled

diff --git a/Engine/Shaders/Private/MobileBasePassPixelShader.usf b/Engine/Shaders/Private/MobileBasePassPixelShader.usf
--- a/Engine/Shaders/Private/MobileBasePassPixelShader.usf
+++ b/Engine/Shaders/Private/MobileBasePassPixelShader.usf
@@ -80,7 +80,7 @@
 #define NONMETAL (MATERIAL_NONMETAL || MOBILE_QL_FORCE_NONMETAL)
 #define FORCE_VERTEX_NORMAL (MOBILE_QL_DISABLE_MATERIAL_NORMAL)
 #define SUPPORT_SPOTLIGHTS_SHADOW (MATERIALBLENDING_SOLID || MATERIALBLENDING_MASKED) && PROJECT_MOBILE_ENABLE_MOVABLE_SPOTLIGHT_SHADOWS
-#define DEFERRED_SHADING_PATH (MOBILE_DEFERRED_SHADING && ((MATERIALBLENDING_SOLID || MATERIALBLENDING_MASKED) && !MATERIAL_SHADINGMODEL_SINGLELAYERWATER && !MATERIAL_SHADINGMODEL_UNLIT))
+#define MOBILE_USE_GBUBUFFER (MOBILE_DEFERRED_SHADING && ((MATERIALBLENDING_SOLID || MATERIALBLENDING_MASKED) && !MATERIAL_SHADINGMODEL_SINGLELAYERWATER))
 #define MATERIAL_WORKS_WITH_DUAL_SOURCE_COLOR_BLENDING (MATERIAL_SHADINGMODEL_THIN_TRANSLUCENT || SUBSTRATE_BLENDING_TRANSLUCENT_COLOREDTRANSMITTANCE )
 
 // Force FORWARD_SHADING on mobile when MOBILE_DEFERRED_SHADING is disabled
@@ -94,9 +94,9 @@
 // SceneDepthAux always enabled for IOS, on other platforms only enabled when in forward and MobileHDR=true
 #define HAS_SCENE_DEPTH_AUX_OUTPUT (USE_SCENE_DEPTH_AUX && (MATERIALBLENDING_SOLID || MATERIALBLENDING_MASKED))
 
-#if DEFERRED_SHADING_PATH && MOBILE_EXTENDED_GBUFFER
+#if MOBILE_USE_GBUBUFFER && MOBILE_EXTENDED_GBUFFER
 	#define SV_TargetDepthAux SV_Target5	
-#elif DEFERRED_SHADING_PATH
+#elif MOBILE_USE_GBUBUFFER
 	#define SV_TargetDepthAux SV_Target4	
 #else
 	#define SV_TargetDepthAux SV_Target1
@@ -293,7 +293,7 @@ void Main(
 	, FMobileBasePassInterpolantsVSToPS BasePassInterpolants
 	, in float4 SvPosition : SV_Position
 	OPTIONAL_IsFrontFace
-#if DEFERRED_SHADING_PATH
+#if MOBILE_USE_GBUBUFFER
 	#if USE_GLES_FBF_DEFERRED
 	, out HALF4_TYPE OutProxy : SV_Target0
 	#elif MOBILE_TRANSLUCENT_COLOR_TRANSMITTANCE_PROGRAMMABLE_BLENDING
@@ -325,7 +325,7 @@ void Main(
 #endif
 	)
 {  
-#if DEFERRED_SHADING_PATH 
+#if MOBILE_USE_GBUBUFFER 
 	#if USE_GLES_FBF_DEFERRED && !MOBILE_TRANSLUCENT_COLOR_TRANSMITTANCE_PROGRAMMABLE_BLENDING
 		half4 OutColor;
 	#endif
@@ -841,7 +841,7 @@ void Main(
 
 	GBuffer.PrecomputedShadowFactors = GetPrimaryPrecomputedShadowMask(LightmapVTPageTableResult, Interpolants, MaterialParameters);
 
-#if DEFERRED_SHADING_PATH
+#if MOBILE_USE_GBUBUFFER && !MATERIAL_SHADINGMODEL_UNLIT
 	GBuffer.IndirectIrradiance = IndirectIrradiance;
 	MobileEncodeGBuffer(GBuffer, OutGBufferA, OutGBufferB, OutGBufferC, OutGBufferD);
 #else
@@ -1089,7 +1089,7 @@ void Main(
 	OutColor.rgb *= ResolvedView.PreExposure;
 #endif
 
-#if MATERIAL_IS_SKY || (!MATERIAL_SHADINGMODEL_UNLIT && !DEFERRED_SHADING_PATH)
+#if MATERIAL_IS_SKY || (!MATERIAL_SHADINGMODEL_UNLIT && !MOBILE_USE_GBUBUFFER)
 	// clamp lit color to avoid overflows
 	OutColor.rgb = SafeGetOutColor(OutColor.rgb);
 #endif
@@ -1118,7 +1118,7 @@ void Main(
 	);
 #endif
 
-#if DEFERRED_SHADING_PATH && USE_GLES_FBF_DEFERRED
+#if MOBILE_USE_GBUBUFFER && USE_GLES_FBF_DEFERRED
 	OutProxy.rgb = OutColor.rgb;
 #elif MOBILE_TRANSLUCENT_COLOR_TRANSMITTANCE_PROGRAMMABLE_BLENDING
 	OutProgrammableBlending = OutColor1 * FramebufferFetchColor0() + OutColor;
