Substrate - fixed shading model color for unlit and clear coat.

diff --git a/Engine/Shaders/Private/Substrate/Substrate.ush b/Engine/Shaders/Private/Substrate/Substrate.ush
--- a/Engine/Shaders/Private/Substrate/Substrate.ush
+++ b/Engine/Shaders/Private/Substrate/Substrate.ush
@@ -3255,7 +3255,7 @@ float SubstrateGetLegacyShadingModels(in FSubstrateBSDF BSDF)
 	case SUBSTRATE_BSDF_TYPE_SLAB:
 	{
 		// Base on features, we can only partially recovers the original shading models. Shading models like foliage, eye, clear coat are impossible/harder to infer.
-		if (BSDF_GETSSSTYPE(BSDF) == SSS_TYPE_DIFFUSION_PROFILE){ return SHADINGMODELID_SUBSURFACE_PROFILE; }
+		if (BSDF_GETSSSTYPE(BSDF) == SSS_TYPE_DIFFUSION_PROFILE)	{ return SHADINGMODELID_SUBSURFACE_PROFILE; }
 		else if (BSDF_GETSSSTYPE(BSDF) == SSS_TYPE_WRAP)		
 		{ 
 			if (BSDF_GETISTHIN(BSDF))
@@ -3264,12 +3264,24 @@ float SubstrateGetLegacyShadingModels(in FSubstrateBSDF BSDF)
 			}
 			return SHADINGMODELID_SUBSURFACE; 
 		}
-		else if (BSDF_GETHASFUZZ(BSDF))							{ return SHADINGMODELID_CLOTH; }
-		else													{ return SHADINGMODELID_DEFAULT_LIT; } // Also include SSS_TYPE_DIFFUSION and SSS_TYPE_SIMPLE_VOLUME
+		else if (BSDF_GETHASFUZZ(BSDF))								{ return SHADINGMODELID_CLOTH; }
+		else
+		{ 
+			if (BSDF_GETHASHAZINESS(BSDF))
+			{
+				FHaziness LocalHaziness = UnpackHaziness(SLAB_HAZINESS(BSDF));
+				if (LocalHaziness.bSimpleClearCoat)
+				{
+					return SHADINGMODELID_CLEAR_COAT;
+				}
+			}
+			return SHADINGMODELID_DEFAULT_LIT;
+		} // Also include SSS_TYPE_DIFFUSION and SSS_TYPE_SIMPLE_VOLUME
 	}
 	case SUBSTRATE_BSDF_TYPE_HAIR: 								{ return SHADINGMODELID_HAIR; }
-	case SUBSTRATE_BSDF_TYPE_EYE: 									{ return SHADINGMODELID_EYE; }
-	case SUBSTRATE_BSDF_TYPE_SINGLELAYERWATER:						{ return SHADINGMODELID_SINGLELAYERWATER; }
+	case SUBSTRATE_BSDF_TYPE_EYE: 								{ return SHADINGMODELID_EYE; }
+	case SUBSTRATE_BSDF_TYPE_SINGLELAYERWATER:					{ return SHADINGMODELID_SINGLELAYERWATER; }
+	case SUBSTRATE_BSDF_TYPE_UNLIT:								{ return SHADINGMODELID_UNLIT; } // This in fact does not show up because unlit does not write any information as it is not needed for any lighting passes down the line.
 	}
 	return 0.f;
 }
