Substrate - fixed hair shading model for gbuffer inspection when hair strand materials are composited later using visibility.

diff --git a/Engine/Shaders/Private/HairStrands/HairStrandsVisibilityDepthPS.usf b/Engine/Shaders/Private/HairStrands/HairStrandsVisibilityDepthPS.usf
--- a/Engine/Shaders/Private/HairStrands/HairStrandsVisibilityDepthPS.usf
+++ b/Engine/Shaders/Private/HairStrands/HairStrandsVisibilityDepthPS.usf
@@ -84,8 +84,10 @@ void MainPS(
 	Output0 = 0; // MaterialData
 	Output1 = 0; // TopLayerData
 
-	// Force PerObjectData contact shadow to respect hair settings, in order to have correct screen cast later
 	FSubstratePixelHeader DummyHeader = InitialiseSubstratePixelHeader();
+	// Set the material mode to hair for GBuffer/Material inspection
+	DummyHeader.SetMaterialMode(HEADER_MATERIALMODE_HAIR);
+	// Force PerObjectData contact shadow to respect hair settings, in order to have correct screen cast later
 	DummyHeader.SetCastContactShadow(SampleData.CastContactShadow ? 1 : 0);
 	Output0 = DummyHeader.State;
 
