Substrate - minor clean up

diff --git a/Engine/Shaders/Private/Substrate/SubstrateLegacyConversion.ush b/Engine/Shaders/Private/Substrate/SubstrateLegacyConversion.ush
--- a/Engine/Shaders/Private/Substrate/SubstrateLegacyConversion.ush
+++ b/Engine/Shaders/Private/Substrate/SubstrateLegacyConversion.ush
@@ -294,7 +294,6 @@ FSubstrateData SubstrateConvertLegacyMaterialDynamic(
 		LegacySlab.Roughness = Roughness;
 		LegacySlab.Anisotropy = Anisotropy;
 		LegacySlab.SSSMFP = max(1e-05f, MFPInCm); // Ensure the MFP is not null to force the material as have SSS
-		LegacySlab.SSSMFPScale = 1.0f;
 		LegacySlab.SSSProfileID = 0.f;
 		LegacySlab.SSSMFPScale = 1.f - Opacity;
 		LegacySlab.SSSPhaseAniso = 0.93f;
