Substrate - fixed
 - subsurface color of fuzz
 - base color of legacy SSSProfile when reading from scene texture in post processes material.

diff --git a/Engine/Shaders/Private/MaterialTemplate.ush b/Engine/Shaders/Private/MaterialTemplate.ush
--- a/Engine/Shaders/Private/MaterialTemplate.ush
+++ b/Engine/Shaders/Private/MaterialTemplate.ush
@@ -3096,7 +3096,8 @@ float4 SceneTextureLookup(float2 UV, int SceneTextureIndex, bool bFiltered)
 			const FSubstrateSubsurfaceHeader SSSHeader = SubstrateLoadSubsurfaceHeader(SubstrateSceneTexture.MaterialTextureArray, SubstrateSceneTexture.FirstSliceStoringSubstrateSSSData, SubstrateAddressing.PixelCoords);
 			const uint SSSType = SubstrateSubSurfaceHeaderGetSSSType(SSSHeader);
 
-			FSubstrateBSDF BSDF = UnpackSubstrateBSDFIn(SubstrateSceneTexture.MaterialTextureArray, SubstrateAddressing, SubstratePixelHeader);
+			const bool bSkipSSSMaterialOverride = true; // We double not want the material data to be affected by any override.
+			FSubstrateBSDF BSDF = UnpackSubstrateBSDFIn(SubstrateSceneTexture.MaterialTextureArray, SubstrateAddressing, SubstratePixelHeader, bSkipSSSMaterialOverride);
 			switch (SceneTextureIndex)
 			{
 				// order needs to match to ESceneTextureId
diff --git a/Engine/Shaders/Private/Substrate/Substrate.ush b/Engine/Shaders/Private/Substrate/Substrate.ush
--- a/Engine/Shaders/Private/Substrate/Substrate.ush
+++ b/Engine/Shaders/Private/Substrate/Substrate.ush
@@ -3155,6 +3155,10 @@ float3 SubstrateGetBSDFSubSurfaceColor(in FSubstrateBSDF BSDF)
 			const float3 Extinction = 1.0f / max(SUBSTRATE_EPSILON, SLAB_SSSMFP(BSDF) * CENTIMETER_TO_METER);
 			return ExtinctionToTransmittance(Extinction, SUBSTRATE_SIMPLEVOLUME_THICKNESS_CM * CENTIMETER_TO_METER);
 		}
+		else if (BSDF_GETHASFUZZ(BSDF))
+		{
+			return SLAB_FUZZ_COLOR(BSDF); // Match legacy subsurface color preview
+		}
 		return 0;
 	}
 	case SUBSTRATE_BSDF_TYPE_HAIR: 				return 0;
@@ -3989,7 +3993,7 @@ FSubstrateBSDF UnpackFastWaterPathSubstrateBSDFIn(FSubstrateMaterialContainer Su
 // Unpack a single BSDF
 // Note: All BSDF lobes needs to be unpack in a sequential manner since each BSDF has a variable footprint.
 // If this is changed, please update the compiler side material size evaluation in SubstrateMaterial.cpp
-FSubstrateBSDF UnpackSubstrateBSDFIn(FSubstrateMaterialContainer SubstrateBuffer, inout FSubstrateAddressing SubstrateAddressing, in FSubstratePixelHeader SubstrateHeader)
+FSubstrateBSDF UnpackSubstrateBSDFIn(FSubstrateMaterialContainer SubstrateBuffer, inout FSubstrateAddressing SubstrateAddressing, in FSubstratePixelHeader SubstrateHeader, const bool bSkipSSSMaterialOverride = false)
 {
 	BRANCH
 	if (SubstrateHeader.IsSimpleMaterial())
@@ -4164,6 +4168,7 @@ FSubstrateBSDF UnpackSubstrateBSDFIn(FSubstrateMaterialContainer SubstrateBuffer
 				// When loading Substrate BSDF which have SSS enable, BaseColor and (optionally) Specular value can be overriden 
 				// based on the SSS method used (screen space SSS / checherboarded screen space SSS / ...)
 				#if SUBSTRATE_SSS_MATERIAL_OVERRIDE
+				if(!bSkipSSSMaterialOverride)
 				{
 					const bool bChecker = CheckerFromPixelPos(SubstrateAddressing.PixelCoords);
 					float3 DiffuseAlbedo = SLAB_DIFFUSEALBEDO(OutBSDF);
@@ -4281,7 +4286,7 @@ FSubstrateBSDF UnpackSubstrateBSDFIn(FSubstrateMaterialContainer SubstrateBuffer
 			// When loading Substrate BSDF which have SSS enable, BaseColor and (optionally) Specular value can be overriden 
 			// based on the SSS method used (screen space SSS / checherboarded screen space SSS / ...)
 			#if SUBSTRATE_SSS_MATERIAL_OVERRIDE
-			if (BSDF_GETSSSTYPE(OutBSDF) == SSS_TYPE_DIFFUSION || BSDF_GETSSSTYPE(OutBSDF) == SSS_TYPE_DIFFUSION_PROFILE)
+			if (!bSkipSSSMaterialOverride && BSDF_GETSSSTYPE(OutBSDF) == SSS_TYPE_DIFFUSION || BSDF_GETSSSTYPE(OutBSDF) == SSS_TYPE_DIFFUSION_PROFILE)
 			{
 				const bool bChecker = CheckerFromPixelPos(SubstrateAddressing.PixelCoords);
 				float3 DiffuseAlbedo = SLAB_DIFFUSEALBEDO(OutBSDF);
@@ -4357,7 +4362,7 @@ FSubstrateBSDF UnpackSubstrateBSDFIn(FSubstrateMaterialContainer SubstrateBuffer
 		// When loading Substrate BSDF which have SSS enable, BaseColor and (optionally) Specular value can be overriden 
 		// based on the SSS method used (screen space SSS / checherboarded screen space SSS / ...)
 		#if SUBSTRATE_SSS_MATERIAL_OVERRIDE
-		if (BSDF_GETSSSTYPE(OutBSDF) == SSS_TYPE_DIFFUSION_PROFILE)
+		if (!bSkipSSSMaterialOverride && BSDF_GETSSSTYPE(OutBSDF) == SSS_TYPE_DIFFUSION_PROFILE)
 		{
 			const bool bChecker = CheckerFromPixelPos(SubstrateAddressing.PixelCoords);
 			float3 DiffuseAlbedo = EYE_DIFFUSEALBEDO(OutBSDF);
