Substrate - fixed shader compilation warning

diff --git a/Engine/Shaders/Private/Substrate/Substrate.ush b/Engine/Shaders/Private/Substrate/Substrate.ush
--- a/Engine/Shaders/Private/Substrate/Substrate.ush
+++ b/Engine/Shaders/Private/Substrate/Substrate.ush
@@ -4286,7 +4286,7 @@ FSubstrateBSDF UnpackSubstrateBSDFIn(FSubstrateMaterialContainer SubstrateBuffer
 			// When loading Substrate BSDF which have SSS enable, BaseColor and (optionally) Specular value can be overriden 
 			// based on the SSS method used (screen space SSS / checherboarded screen space SSS / ...)
 			#if SUBSTRATE_SSS_MATERIAL_OVERRIDE
-			if (!bSkipSSSMaterialOverride && BSDF_GETSSSTYPE(OutBSDF) == SSS_TYPE_DIFFUSION || BSDF_GETSSSTYPE(OutBSDF) == SSS_TYPE_DIFFUSION_PROFILE)
+			if (!bSkipSSSMaterialOverride && (BSDF_GETSSSTYPE(OutBSDF) == SSS_TYPE_DIFFUSION || BSDF_GETSSSTYPE(OutBSDF) == SSS_TYPE_DIFFUSION_PROFILE))
 			{
 				const bool bChecker = CheckerFromPixelPos(SubstrateAddressing.PixelCoords);
 				float3 DiffuseAlbedo = SLAB_DIFFUSEALBEDO(OutBSDF);
