Workaround crash in Xclipse shader compiler when compiling SM5 motion blur compute shader

diff --git a/Engine/Shaders/Private/MotionBlur/MotionBlurApply.usf b/Engine/Shaders/Private/MotionBlur/MotionBlurApply.usf
--- a/Engine/Shaders/Private/MotionBlur/MotionBlurApply.usf
+++ b/Engine/Shaders/Private/MotionBlur/MotionBlurApply.usf
@@ -940,6 +940,10 @@ void MainCS(
 			QuarterResOutput = 0.0;
 		}
 		
+		// Needed to avoid crash in shader compiler of Xclipse driver
+		#if !POST_PROCESS_ALPHA
+			QuarterResOutput.a = 0;
+		#endif
 		bool bIsValid = all(OutputPixelCoord < mb_short2(Color_ViewportMax));
 		mb_short2 OutputPixelCoord0Mip0 = (bIsValid) ? OutputPixelCoord : mb_short(~0).xx;
 		mb_short2 OutputPixelCoord1Mip0 = (bIsValid) ? OutputPixelCoord1 : mb_short(~0).xx;
@@ -1050,6 +1054,10 @@ void MainCS(
 			HalfResOutput = 0.0;
 		}
 		
+		// Needed to avoid crash in shader compiler of Xclipse driver
+		#if !POST_PROCESS_ALPHA
+			HalfResOutput.a = 0;
+		#endif
 		mb_half4 QuarterResOutput;
 		BRANCH
 		if (OutputMip2)
@@ -1078,6 +1086,10 @@ void MainCS(
 			QuarterResOutput = 0.0;
 		}
 
+		// Needed to avoid crash in shader compiler of Xclipse driver
+		#if !POST_PROCESS_ALPHA
+			QuarterResOutput.a = 0;
+		#endif
 		bool bIsValid = all(OutputPixelCoord < mb_short2(Color_ViewportMax));
 		mb_short2 OutputPixelCoordMip0 = (bIsValid) ? OutputPixelCoord : mb_short(~0).xx;
 		mb_short2 OutputPixelCoordMip1 = select(and(bIsValid, and(OutputMip1 != 0, (OutputPixelCoordMip0 & 0x1) == 0)), (OutputPixelCoord >> mb_short(1)), mb_short(~0).xx);
