Preserve hair BSDF value when generating hair LUT.

diff --git a/Engine/Shaders/Private/HairBsdf.ush b/Engine/Shaders/Private/HairBsdf.ush
--- a/Engine/Shaders/Private/HairBsdf.ush
+++ b/Engine/Shaders/Private/HairBsdf.ush
@@ -7,12 +7,12 @@
 ///////////////////////////////////////////////////////////////////////////////////////////////////
 // Utility functions
 
-float Hair_g(float B, float Theta)
+float Hair_g(float B, float Theta, bool bClampBSDFValue)
 {
 	// Clamp B for the denominator term, as otherwise the Gaussian normalization returns too high value.
 	// This clamps allow to prevent large value for low roughness, while keeping the highlight shape/sharpness 
 	// similar.
-	const float DenominatorB = max(B, 0.01f);
+	const float DenominatorB = bClampBSDFValue ? max(B, 0.01f) : B;
 	return exp(-0.5 * Pow2(Theta) / (B * B)) / (sqrt(2 * PI) * DenominatorB);
 }
 
@@ -130,7 +130,7 @@ float GaussianDetector(float Bp, float Phi)
 	for (int k = -4; k <= 4; k++)
 	{
 		// TODO use symmetry and detect for both Phi and -Phi
-		Dp += Hair_g(Bp, Phi - (2 * PI) * k);
+		Dp += Hair_g(Bp, Phi - (2 * PI) * k, false);
 	}
 	return Dp;
 }
@@ -329,7 +329,7 @@ float3 HairShading( FGBufferData GBuffer, float3 L, float3 V, half3 N, float Sha
 		const float ca = cos(Alpha[0]);
 		float ShiftR = 2 * sa * (ca * CosHalfPhi * sqrt(1 - SinThetaV * SinThetaV) + sa * SinThetaV);
 		float BScale = HairTransmittance.bUseSeparableR ? sqrt(2.0) * CosHalfPhi : 1;
-		float Mp = Hair_g(B[0] * BScale, SinThetaL + SinThetaV - ShiftR);
+		float Mp = Hair_g(B[0] * BScale, SinThetaL + SinThetaV - ShiftR, HairTransmittance.bClampBSDFValue);
 		float Np = 0.25 * CosHalfPhi;
 		float Fp = Hair_F(sqrt(saturate(0.5 + 0.5 * VoL)));
 		S += Mp * Np * Fp * (GBuffer.Specular * 2) * lerp(1, Backlit, saturate(-VoL));
@@ -338,7 +338,7 @@ float3 HairShading( FGBufferData GBuffer, float3 L, float3 V, half3 N, float Sha
 	// TT
 	if (HairTransmittance.ScatteringComponent & HAIR_COMPONENT_TT)
 	{
-		float Mp = Hair_g( B[1], SinThetaL + SinThetaV - Alpha[1] );
+		float Mp = Hair_g( B[1], SinThetaL + SinThetaV - Alpha[1], HairTransmittance.bClampBSDFValue);
 
 		float a = 1 / n_prime;
 		//float h = CosHalfPhi * rsqrt( 1 + a*a - 2*a * sqrt( 0.5 - 0.5 * CosPhi ) );
@@ -378,7 +378,7 @@ float3 HairShading( FGBufferData GBuffer, float3 L, float3 V, half3 N, float Sha
 	// TRT
 	if (HairTransmittance.ScatteringComponent & HAIR_COMPONENT_TRT)
 	{
-		float Mp = Hair_g( B[2], SinThetaL + SinThetaV - Alpha[2] );
+		float Mp = Hair_g( B[2], SinThetaL + SinThetaV - Alpha[2], HairTransmittance.bClampBSDFValue);
 		
 		//float h = 0.75;
 		float f = Hair_F( CosThetaD * 0.5 );
diff --git a/Engine/Shaders/Private/HairShadingCommon.ush b/Engine/Shaders/Private/HairShadingCommon.ush
--- a/Engine/Shaders/Private/HairShadingCommon.ush
+++ b/Engine/Shaders/Private/HairShadingCommon.ush
@@ -22,6 +22,7 @@ struct FHairTransmittanceData
 	bool bUseLegacyAbsorption;
 	bool bUseSeparableR;
 	bool bUseBacklit;
+	bool bClampBSDFValue;
 
 	float  OpaqueVisibility;
 	float3 LocalScattering;
@@ -36,6 +37,7 @@ FHairTransmittanceData InitHairTransmittanceData(bool bMultipleScatterEnable = t
 	o.bUseLegacyAbsorption = true;
 	o.bUseSeparableR = true;
 	o.bUseBacklit = false;
+	o.bClampBSDFValue = true;
 
 	o.OpaqueVisibility = 1;
 	o.LocalScattering = 0;
diff --git a/Engine/Shaders/Private/HairStrands/HairStrandsLUT.usf b/Engine/Shaders/Private/HairStrands/HairStrandsLUT.usf
--- a/Engine/Shaders/Private/HairStrands/HairStrandsLUT.usf
+++ b/Engine/Shaders/Private/HairStrands/HairStrandsLUT.usf
@@ -112,6 +112,7 @@ void MainCS(uint3 DispatchThreadId : SV_DispatchThreadID)
 	const float OpaqueVisibility = 1;
 	FHairTransmittanceData TransmittanceData = InitHairStrandsTransmittanceData();
 	TransmittanceData.bUseSeparableR = true;
+	TransmittanceData.bClampBSDFValue = false;
 
 	const float MaxCosThetaRadius = cos(0.25f * PI / float(ThetaCount)); // [0, Pi/2] / ThetaCount which is divided by 2 for getting the actual radius
 	float3x3 ToViewBasis = GetTangentBasis(V);
