Fix CIS shader compilation warning.

diff --git a/Engine/Shaders/Private/SpecularProfile.usf b/Engine/Shaders/Private/SpecularProfile.usf
--- a/Engine/Shaders/Private/SpecularProfile.usf
+++ b/Engine/Shaders/Private/SpecularProfile.usf
@@ -61,8 +61,8 @@ float3 EvaluateSpecularProfile(float NoV, float NoL, float VoH, float NoH, float
 {
 	const float3 ProfileCoord = GetSpecularProfileCoord(SpecularProfileId, NoV, NoL, VoH, NoH);
 	uint3 FetchCoord = uint3(ProfileCoord.xy * (Resolution - 1), ProfileCoord.z);
-	FetchCoord.x = clamp(FetchCoord.x, 0, Resolution.x - 1);
-	FetchCoord.y = clamp(FetchCoord.y, 0, Resolution.y - 1);
+	FetchCoord.x = clamp(FetchCoord.x, 0u, uint(Resolution.x - 1));
+	FetchCoord.y = clamp(FetchCoord.y, 0u, uint(Resolution.y - 1));
 	return TargetTexture[FetchCoord].xyz;
 }
 
