Planar reflection shared sampler parameter logic simplification.
Tested pc and mobile OGL preview to check shared sampler on/off.

diff --git a/Engine/Shaders/Private/ForwardLightingCommon.ush b/Engine/Shaders/Private/ForwardLightingCommon.ush
--- a/Engine/Shaders/Private/ForwardLightingCommon.ush
+++ b/Engine/Shaders/Private/ForwardLightingCommon.ush
@@ -535,7 +535,7 @@ half3 GetImageBasedReflectionSpecular(FMaterialPixelParameters MaterialParameter
 	if (abs(dot(PlanarReflectionStruct.ReflectionPlane.xyz, 1)) > .0001f)
 	{
 		// Reuse ReflectionCubemapSampler to avoid reducing the sampler count available to artists
-		float4 PlanarReflection = ComputePlanarReflections(MaterialParameters.WorldPosition_CamRelative, MaterialParameters.WorldNormal, Roughness, /*SharedClampSampler*/ View.SharedTrilinearClampedSampler);
+		float4 PlanarReflection = ComputePlanarReflections(MaterialParameters.WorldPosition_CamRelative, MaterialParameters.WorldNormal, Roughness);
 		// Planar reflections win over SSR and reflection environment
 		SpecularLighting = PlanarReflection.rgb + (1 - PlanarReflection.a) * SpecularLighting;
 	}
diff --git a/Engine/Shaders/Private/PlanarReflectionShaders.usf b/Engine/Shaders/Private/PlanarReflectionShaders.usf
--- a/Engine/Shaders/Private/PlanarReflectionShaders.usf
+++ b/Engine/Shaders/Private/PlanarReflectionShaders.usf
@@ -159,7 +159,7 @@ void PlanarReflectionPS(
 	BRANCH
 	if (!bNoMaterial)
 	{
-		OutColor = ComputePlanarReflections(TranslatedWorldPosition, WorldNormal, Roughness, PlanarReflectionStruct.PlanarReflectionSampler);
+		OutColor = ComputePlanarReflections(TranslatedWorldPosition, WorldNormal, Roughness);
 
 		OutColor.xyz *= View.PreExposure;
 	}
@@ -174,7 +174,7 @@ void PlanarReflectionPS(
 	BRANCH
 	if (GBuffer.ShadingModelID != SHADINGMODELID_UNLIT)
 	{
-		OutColor = ComputePlanarReflections(TranslatedWorldPosition, GBuffer.WorldNormal, GBuffer.Roughness, PlanarReflectionStruct.PlanarReflectionSampler);
+		OutColor = ComputePlanarReflections(TranslatedWorldPosition, GBuffer.WorldNormal, GBuffer.Roughness);
 
 		OutColor.xyz *= View.PreExposure;
 	}
diff --git a/Engine/Shaders/Private/PlanarReflectionShared.ush b/Engine/Shaders/Private/PlanarReflectionShared.ush
--- a/Engine/Shaders/Private/PlanarReflectionShared.ush
+++ b/Engine/Shaders/Private/PlanarReflectionShared.ush
@@ -4,7 +4,7 @@
  PlanarReflectionShared.usf
 =============================================================================*/
 
-half4 ComputePlanarReflections(float3 TranslatedWorldPosition, half3 WorldNormal, half Roughness, SamplerState SharedClampSampler)
+half4 ComputePlanarReflections(float3 TranslatedWorldPosition, half3 WorldNormal, half Roughness)
 {
 	half4 OutPlanarReflection = 0;
 
@@ -54,7 +54,7 @@ half4 ComputePlanarReflections(float3 TranslatedWorldPosition, half3 WorldNormal
 		half4 PlanarReflectionTextureValue = Texture2DSampleLevel(
 			PlanarReflectionStruct.PlanarReflectionTexture,
 #if SUPPORTS_INDEPENDENT_SAMPLERS
-			SharedClampSampler, 
+			View.SharedTrilinearClampedSampler, 
 #else
 			PlanarReflectionStruct.PlanarReflectionSampler,
 #endif
@@ -73,7 +73,7 @@ half4 ComputePlanarReflections(float3 TranslatedWorldPosition, half3 WorldNormal
 #if (FEATURE_LEVEL <= FEATURE_LEVEL_ES3_1)
 half4 GetPlanarReflection(float3 TranslatedWorldPosition, half3 WorldNormal, half Roughness)
 {
-	half4 PlanarReflection = ComputePlanarReflections(TranslatedWorldPosition, WorldNormal, Roughness, PlanarReflectionStruct.PlanarReflectionSampler);
+	half4 PlanarReflection = ComputePlanarReflections(TranslatedWorldPosition, WorldNormal, Roughness);
 #if OUTPUT_GAMMA_SPACE
 	// the capture will also be in gamma space, convert to linear:
 	PlanarReflection.rgb *= PlanarReflection.rgb;
