Fix a bug in the editor outline shader caused by reading an uninitialized variable.

diff --git a/Engine/Shaders/Private/PostProcessSelectionOutline.usf b/Engine/Shaders/Private/PostProcessSelectionOutline.usf
--- a/Engine/Shaders/Private/PostProcessSelectionOutline.usf
+++ b/Engine/Shaders/Private/PostProcessSelectionOutline.usf
@@ -260,7 +260,7 @@ void MainPS(
 	float4 OverlayColor = Texture2DSample(OverlayLookupTexture, OverlayLookupSampler, ColorUV);
 	OutColor = Texture2DSample(ColorTexture, ColorSampler, ColorUV); 
 	
-	float4 SelectionEffect;
+	float4 SelectionEffect = 0;
 	const int ScoutSum = dot(saturate(StencilScout), 1);
 
 	// If this sum is zero, none of our neighbors are selected pixels and we can skip all the heavy processing.
