Fix OutTextureMipColor's declaration: we create a cube texture, but the UAV bound is created with a DimensionOverride set to ETextureDimension::Texture2DArray: the HLSL needs to declare it as RWTexture2DArray

diff --git a/Engine/Shaders/Private/ReflectionEnvironmentShaders.usf b/Engine/Shaders/Private/ReflectionEnvironmentShaders.usf
--- a/Engine/Shaders/Private/ReflectionEnvironmentShaders.usf
+++ b/Engine/Shaders/Private/ReflectionEnvironmentShaders.usf
@@ -242,7 +242,7 @@ uint NumMips;
 #ifdef USE_COMPUTE
 int FaceThreadGroupSize;
 int2 ValidDispatchCoord;
-RWTextureCube<float4> OutTextureMipColor;	// All slices, i.e. faces, of a cube map mip level
+RWTexture2DArray<float4> OutTextureMipColor;	// All slices, i.e. faces, of a cube map mip level
 #endif 
 
 #ifdef USE_COMPUTE
