Fix validation error

diff --git a/Engine/Shaders/Private/HeterogeneousVolumes/HeterogeneousVolumesVoxelGridShadows.usf b/Engine/Shaders/Private/HeterogeneousVolumes/HeterogeneousVolumesVoxelGridShadows.usf
--- a/Engine/Shaders/Private/HeterogeneousVolumes/HeterogeneousVolumesVoxelGridShadows.usf
+++ b/Engine/Shaders/Private/HeterogeneousVolumes/HeterogeneousVolumesVoxelGridShadows.usf
@@ -483,7 +483,7 @@ void RenderVolumetricShadowMapForLightWithVoxelGridCS(
 
 StructuredBuffer<uint2> VolumetricShadowLinkedListBuffer;
 
-RWBuffer<int> RWVolumetricShadowIndirectionAllocatorBuffer;
+RWBuffer<uint> RWVolumetricShadowIndirectionAllocatorBuffer;
 RWStructuredBuffer<uint2> RWVolumetricShadowIndirectionBuffer;
 RWStructuredBuffer<uint> RWVolumetricShadowTransmittanceBuffer;
 
