Volumetric Cloud - try to force recompile cloud shader.

diff --git a/Engine/Shaders/Private/VolumetricCloud.usf b/Engine/Shaders/Private/VolumetricCloud.usf
--- a/Engine/Shaders/Private/VolumetricCloud.usf
+++ b/Engine/Shaders/Private/VolumetricCloud.usf
@@ -6,6 +6,9 @@
 
 
 
+// Change this to force recompilation of all shaders referencing this file.
+#pragma message("UESHADERMETADATA_VERSION B0DC9137-AC62-453E-A2F9-A01047CC479C")
+
 #include "Common.ush"
 #include "ColorMap.ush"
 #include "ColorSpace.ush"
