- Fix for inconsistent shader linkage when using instanced stereo

diff --git a/Engine/Shaders/Private/ShadowProjectionPixelShader.usf b/Engine/Shaders/Private/ShadowProjectionPixelShader.usf
--- a/Engine/Shaders/Private/ShadowProjectionPixelShader.usf
+++ b/Engine/Shaders/Private/ShadowProjectionPixelShader.usf
@@ -430,13 +430,29 @@ float2 PointLightProjParameters;
 
 /** Pixel shader for projecting a one pass point light shadow from a cube map. */
 void MainOnePassPointLightPS(
+#if INSTANCED_STEREO
+	in nointerpolation uint EyeIndex : TEXCOORD0,
+#endif
 	in float4 SVPos : SV_POSITION,
 	out float4 OutColor : SV_Target0
 	)
 {
-	float2 ScreenUV = float2( SVPos.xy * View.BufferSizeAndInvSize.zw );
-	float SceneW = CalcSceneDepth( ScreenUV );
+	
+	float ArrayIndex = 0;
+#if INSTANCED_STEREO
+	ResolvedView = ResolveView(EyeIndex);
+	ArrayIndex = EyeIndex;
+#else
+	ResolvedView = ResolveView();
+#endif
 
+	float2 ScreenUV = float2( SVPos.xy * ResolvedView.BufferSizeAndInvSize.zw );
+#if INSTANCED_STEREO && SHADING_PATH_MOBILE
+	float SceneW = CalcSceneDepth(ScreenUV, ArrayIndex);
+#else
+	float SceneW = CalcSceneDepth( ScreenUV );
+#endif
+	
 #if SUBPIXEL_SHADOW
 	const uint3 PixelCoord = uint3(floor(SVPos.xy), 0);
 	const float HairSampleDepth = HairStrands.HairOnlyDepthTexture.Load(PixelCoord).x;
@@ -450,8 +466,8 @@ void MainOnePassPointLightPS(
 	}
 #endif
 
-	float2 ScreenPosition = ( ScreenUV.xy - View.ScreenPositionScaleBias.wz ) / View.ScreenPositionScaleBias.xy;
-	float3 TranslateWorldPosition = mul(float4(GetScreenPositionForProjectionType(ScreenPosition, SceneW), SceneW, 1), View.ScreenToTranslatedWorld).xyz;
+	float2 ScreenPosition = ( ScreenUV.xy - ResolvedView.ScreenPositionScaleBias.wz ) / ResolvedView.ScreenPositionScaleBias.xy;
+	float3 TranslateWorldPosition = mul(float4(GetScreenPositionForProjectionType(ScreenPosition, SceneW), SceneW, 1), ResolvedView.ScreenToTranslatedWorld).xyz;
 
 	// For debugging
 #define OUTPUT_CUBE_SHADOW_DEPTH_NO_FILTERING 0
