ColorManagement: Ensure that the inverse tone curve approximation respects the working color space.

Small update for correctness despite this function not being used by our global shaders.

#rnx

diff --git a/Engine/Shaders/Private/TonemapCommon.ush b/Engine/Shaders/Private/TonemapCommon.ush
--- a/Engine/Shaders/Private/TonemapCommon.ush
+++ b/Engine/Shaders/Private/TonemapCommon.ush
@@ -230,9 +230,12 @@ half3 FilmToneMapInverse( half3 ToneColor )
 	// Pre desaturate
 	LinearColor = lerp( dot( LinearColor, AP1_RGB2Y ), LinearColor, 1.0 / 0.96 );
 
+#if WORKING_COLOR_SPACE_IS_SRGB
 	LinearColor = mul( AP1_2_sRGB, LinearColor );
-
-	// Returning positive sRGB values
+#else
+	LinearColor = mul( XYZ_TO_RGB_WORKING_COLOR_SPACE_MAT, mul( AP1_2_XYZ_MAT, LinearColor ) );
+#endif
+	// Returning positive working color space values
 	return max( 0, LinearColor );
 }
 
