Fix shader compilation issues in HeterogeneousVolumes.

diff --git a/Engine/Shaders/Private/HeterogeneousVolumes/HeterogeneousVolumesLiveShadingShadows.usf b/Engine/Shaders/Private/HeterogeneousVolumes/HeterogeneousVolumesLiveShadingShadows.usf
--- a/Engine/Shaders/Private/HeterogeneousVolumes/HeterogeneousVolumesLiveShadingShadows.usf
+++ b/Engine/Shaders/Private/HeterogeneousVolumes/HeterogeneousVolumesLiveShadingShadows.usf
@@ -164,6 +164,18 @@ float3 SampleExtinction(inout FVolumeSampleContext Context)
 	return Extinction;
 }
 
+float3 SampleEmission(inout FVolumeSampleContext Context)
+{
+	float3 Emission = SampleEmissive(Context.PixelMaterialInputs);
+	return Emission;
+}
+
+float3 SampleAlbedo(inout FVolumeSampleContext Context)
+{
+	float3 Albedo = SampleAlbedo(Context.PixelMaterialInputs);
+	return Albedo;
+}
+
 uint GetLinearIndex(uint Face, uint2 PixelCoord)
 {
 	uint LinearIndex = Face * ShadowResolution.x * ShadowResolution.y + PixelCoord.y * ShadowResolution.x + PixelCoord.x;
diff --git a/Engine/Shaders/Private/HeterogeneousVolumes/HeterogeneousVolumesRayMarchingUtils.ush b/Engine/Shaders/Private/HeterogeneousVolumes/HeterogeneousVolumesRayMarchingUtils.ush
--- a/Engine/Shaders/Private/HeterogeneousVolumes/HeterogeneousVolumesRayMarchingUtils.ush
+++ b/Engine/Shaders/Private/HeterogeneousVolumes/HeterogeneousVolumesRayMarchingUtils.ush
@@ -6,6 +6,8 @@
 
 #include "../ForwardShadowingCommon.ush"
 #include "../VolumeLightingCommon.ush"
+#include "../LightData.ush"
+#include "../DeferredLightingCommon.ush"
 
 #include "../SHCommon.ush"
 #define FrontLayerTranslucencyReflectionsStruct LumenGIVolumeStruct
