Sky atmosphere - fix for nan expoding light scattering in some cases.

diff --git a/Engine/Shaders/Private/SkyAtmosphere.usf b/Engine/Shaders/Private/SkyAtmosphere.usf
--- a/Engine/Shaders/Private/SkyAtmosphere.usf
+++ b/Engine/Shaders/Private/SkyAtmosphere.usf
@@ -736,12 +736,13 @@ SingleScatteringResult IntegrateSingleScatteredLuminance(
 		// The way to fix that is to use a proper analytical integration as porposed in slide 28 of http://www.frostbite.com/2015/08/physically-based-unified-volumetric-rendering-in-frostbite/ 
 		// However, it is possible to disable as it can also work using simple power serie sum unroll up to 5th order. The rest of the orders has a really low contribution. 
 #define MULTI_SCATTERING_POWER_SERIE 0 
+		const float3 SafeMediumExtinction = max(Medium.Extinction, 1.e-9);
 #if MULTI_SCATTERING_POWER_SERIE==0 
 		// 1 is the integration of luminance over the 4pi of a sphere, and assuming an isotropic phase function of 1.0/(4*PI) 
 		Result.MultiScatAs1 += Throughput * Medium.Scattering * 1.0f * dt;
 #else 
 		float3 MS = Medium.Scattering * 1;
-		float3 MSint = (MS - MS * SampleTransmittance) / Medium.Extinction;
+		float3 MSint = (MS - MS * SampleTransmittance) / SafeMediumExtinction;
 		Result.MultiScatAs1 += Throughput * MSint;
 #endif 
 
@@ -752,9 +753,9 @@ SingleScatteringResult IntegrateSingleScatteredLuminance(
 		Throughput	*= SampleTransmittance;
 #else
 		// See slide 28 at http://www.frostbite.com/2015/08/physically-based-unified-volumetric-rendering-in-frostbite/ 
-		float3 Sint			= (S        - S        * SampleTransmittance) / Medium.Extinction;	// integrate along the current step segment 
-		float3 SintMieOnly	= (SMieOnly - SMieOnly * SampleTransmittance) / Medium.Extinction;
-		float3 SintRayOnly	= (SRayOnly - SRayOnly * SampleTransmittance) / Medium.Extinction;
+		float3 Sint			= (S        - S        * SampleTransmittance) / SafeMediumExtinction;	// integrate along the current step segment 
+		float3 SintMieOnly	= (SMieOnly - SMieOnly * SampleTransmittance) / SafeMediumExtinction;
+		float3 SintRayOnly	= (SRayOnly - SRayOnly * SampleTransmittance) / SafeMediumExtinction;
 		L			+= Throughput * Sint;														// accumulate and also take into account the transmittance from previous steps
 		LMieOnly	+= Throughput * SintMieOnly;
 		LRayOnly	+= Throughput * SintRayOnly;
