Heterogeneous Volumes: Prefer storing transmittance to opacity.

diff --git a/Engine/Shaders/Private/HeterogeneousVolumes/HeterogeneousVolumesComposite.usf b/Engine/Shaders/Private/HeterogeneousVolumes/HeterogeneousVolumesComposite.usf
--- a/Engine/Shaders/Private/HeterogeneousVolumes/HeterogeneousVolumesComposite.usf
+++ b/Engine/Shaders/Private/HeterogeneousVolumes/HeterogeneousVolumesComposite.usf
@@ -18,7 +18,7 @@ void HeterogeneousVolumesCompositeCS(
 	uint2 PixelCoord = DispatchThreadId.xy + View.ViewRectMin.xy;
 
 	float4 Result;
-	float HeterogeneousVolumeTransmittance = saturate(1.0 - HeterogeneousVolumeRadiance[PixelCoord].a);
+	float HeterogeneousVolumeTransmittance = HeterogeneousVolumeRadiance[PixelCoord].a;
 	Result.rgb = HeterogeneousVolumeRadiance[PixelCoord].rgb + RWColorTexture[PixelCoord].rgb * HeterogeneousVolumeTransmittance;
 	Result.a = RWColorTexture[PixelCoord].a * HeterogeneousVolumeTransmittance;
 
diff --git a/Engine/Shaders/Private/HeterogeneousVolumes/HeterogeneousVolumesHardwareRayTracing.usf b/Engine/Shaders/Private/HeterogeneousVolumes/HeterogeneousVolumesHardwareRayTracing.usf
--- a/Engine/Shaders/Private/HeterogeneousVolumes/HeterogeneousVolumesHardwareRayTracing.usf
+++ b/Engine/Shaders/Private/HeterogeneousVolumes/HeterogeneousVolumesHardwareRayTracing.usf
@@ -592,6 +592,5 @@ RAY_TRACING_ENTRY_RAYGEN(RenderSingleScatteringWithPreshadingRGS)
 	}
 
 	// Output..
-	float3 Opacity = bApplyEmissionAndTransmittance ? 1.0 - Transmittance : 0.0;
-	RWLightingTexture[PixelCoord] += float4(Radiance * View.PreExposure, Luminance(Opacity));
+	RWLightingTexture[PixelCoord] += float4(Radiance * View.PreExposure, Luminance(Transmittance));
 }
diff --git a/Engine/Shaders/Private/HeterogeneousVolumes/HeterogeneousVolumesLiveShadingPipeline.usf b/Engine/Shaders/Private/HeterogeneousVolumes/HeterogeneousVolumesLiveShadingPipeline.usf
--- a/Engine/Shaders/Private/HeterogeneousVolumes/HeterogeneousVolumesLiveShadingPipeline.usf
+++ b/Engine/Shaders/Private/HeterogeneousVolumes/HeterogeneousVolumesLiveShadingPipeline.usf
@@ -315,7 +315,7 @@ void RenderSingleScatteringWithLiveShadingCS(
 	uint2 PixelCoord = DispatchThreadId.xy + View.ViewRectMin.xy;
 
 	float3 Radiance = 0.0;
-	float3 Transmittance = 1.0 - RWLightingTexture[PixelCoord].a;
+	float3 Transmittance = RWLightingTexture[PixelCoord].a;
 	float PrevOpacity = Luminance(1.0 - Transmittance);
 
 	// Extract depth
@@ -390,7 +390,7 @@ void RenderSingleScatteringWithLiveShadingCS(
 		if (bApplyEmissionAndTransmittance)
 		{
 			float Opacity = Luminance(1.0 - Transmittance);
-			RWLightingTexture[PixelCoord].a = Opacity;
+			RWLightingTexture[PixelCoord].a = Transmittance;
 
 #if DIM_WRITE_VELOCITY
 			bool bWriteVelocity =
diff --git a/Engine/Shaders/Private/HeterogeneousVolumes/HeterogeneousVolumesPreshadingPipeline.usf b/Engine/Shaders/Private/HeterogeneousVolumes/HeterogeneousVolumesPreshadingPipeline.usf
--- a/Engine/Shaders/Private/HeterogeneousVolumes/HeterogeneousVolumesPreshadingPipeline.usf
+++ b/Engine/Shaders/Private/HeterogeneousVolumes/HeterogeneousVolumesPreshadingPipeline.usf
@@ -517,7 +517,7 @@ void RenderSingleScatteringWithPreshadingCS(
 	float2 PixelCoord = RayMarchingTile.PixelOffset + TileCoord;
 
 	float3 Radiance = 0.0;
-	float3 Transmittance = 1.0 - RWLightingTexture[PixelCoord].a;
+	float3 Transmittance = RWLightingTexture[PixelCoord].a;
 	float PrevOpacity = Luminance(1.0 - Transmittance);
 
 	// Create screen ray
@@ -628,31 +628,35 @@ void RenderSingleScatteringWithPreshadingCS(
 	}
 
 	// Radiance
-	float Opacity = bApplyEmissionAndTransmittance ? Luminance(1.0 - Transmittance) : 0.0;
-	RWLightingTexture[PixelCoord] = float4(RWLightingTexture[PixelCoord].rgb + (Radiance * View.PreExposure), Opacity);
+	RWLightingTexture[PixelCoord].rgb += Radiance * View.PreExposure;
+	if (bApplyEmissionAndTransmittance)
+	{
+		float Opacity = Luminance(1.0 - Transmittance);
+		RWLightingTexture[PixelCoord].a = Transmittance;
 
 #if DIM_WRITE_VELOCITY
-	bool bWriteVelocity =
-		(WorldHitTMin < WorldRayLength) &&
-		(PrevOpacity < 0.5) &&
-		(Opacity > 0.5);
+		bool bWriteVelocity =
+			(WorldHitTMin < WorldRayLength) &&
+			(PrevOpacity < 0.5) &&
+			(Opacity > 0.5);
 
-	if (bWriteVelocity)
-	{
-		float3 Velocity = Calculate3DVelocityAtWorldPos(WorldRayOrigin, WorldRayDirection, WorldHitTMin);
-		RWVelocityTexture[PixelCoord] = EncodeVelocityToTexture(Velocity);
-	}
+		if (bWriteVelocity)
+		{
+			float3 Velocity = Calculate3DVelocityAtWorldPos(WorldRayOrigin, WorldRayDirection, WorldHitTMin);
+			RWVelocityTexture[PixelCoord] = EncodeVelocityToTexture(Velocity);
+		}
 #endif // DIM_WRITE_VELOCITY
 
 #if DIM_DEBUG
-	//if (!all(DispatchThreadId))
-	if (GroupId == 0 && GroupThreadId == 0)
-	{
-		RWVoxelOutputBuffer[0].MipLevel = GetNumVoxels();
-		for (uint VoxelIndex = 0; VoxelIndex < GetNumVoxels(); ++VoxelIndex)
+		//if (!all(DispatchThreadId))
+		if (GroupId == 0 && GroupThreadId == 0)
 		{
-			RWVoxelOutputBuffer[VoxelIndex + 1] = GetVoxelDataPacked(VoxelIndex);
+			RWVoxelOutputBuffer[0].MipLevel = GetNumVoxels();
+			for (uint VoxelIndex = 0; VoxelIndex < GetNumVoxels(); ++VoxelIndex)
+			{
+				RWVoxelOutputBuffer[VoxelIndex + 1] = GetVoxelDataPacked(VoxelIndex);
+			}
 		}
-	}
 #endif
+	}
 }
diff --git a/Engine/Shaders/Private/HeterogeneousVolumes/HeterogeneousVolumesVoxelGridRendering.usf b/Engine/Shaders/Private/HeterogeneousVolumes/HeterogeneousVolumesVoxelGridRendering.usf
--- a/Engine/Shaders/Private/HeterogeneousVolumes/HeterogeneousVolumesVoxelGridRendering.usf
+++ b/Engine/Shaders/Private/HeterogeneousVolumes/HeterogeneousVolumesVoxelGridRendering.usf
@@ -763,6 +763,5 @@ void RenderTransmittanceWithVoxelGridCS(
 #endif // DEBUG_MODE_DDA_MAJORANT
 
 	// Output..
-	float3 Opacity = 1.0 - Transmittance;
-	RWLightingTexture[PixelCoord] += float4(Radiance, Luminance(Opacity));
+	RWLightingTexture[PixelCoord] += float4(Radiance, Luminance(Transmittance));
 }
diff --git a/Engine/Source/Runtime/Renderer/Private/HeterogeneousVolumes/HeterogeneousVolumes.cpp b/Engine/Source/Runtime/Renderer/Private/HeterogeneousVolumes/HeterogeneousVolumes.cpp
--- a/Engine/Source/Runtime/Renderer/Private/HeterogeneousVolumes/HeterogeneousVolumes.cpp
+++ b/Engine/Source/Runtime/Renderer/Private/HeterogeneousVolumes/HeterogeneousVolumes.cpp
@@ -766,7 +766,7 @@ void FDeferredShadingSceneRenderer::RenderHeterogeneousVolumes(
 		Desc.Format = PF_FloatRGBA;
 		Desc.Flags &= ~(TexCreate_FastVRAM);
 		HeterogeneousVolumeRadiance = GraphBuilder.CreateTexture(Desc, TEXT("HeterogeneousVolumes"));
-		AddClearUAVPass(GraphBuilder, GraphBuilder.CreateUAV(HeterogeneousVolumeRadiance), FLinearColor::Transparent);
+		AddClearUAVPass(GraphBuilder, GraphBuilder.CreateUAV(HeterogeneousVolumeRadiance), FLinearColor::Black);
 	}
 
 	for (int32 ViewIndex = 0; ViewIndex < Views.Num(); ViewIndex++)
