Fixing shader warning

diff --git a/Engine/Shaders/Private/HeterogeneousVolumes/HeterogeneousVolumesLiveShadingPipeline.usf b/Engine/Shaders/Private/HeterogeneousVolumes/HeterogeneousVolumesLiveShadingPipeline.usf
--- a/Engine/Shaders/Private/HeterogeneousVolumes/HeterogeneousVolumesLiveShadingPipeline.usf
+++ b/Engine/Shaders/Private/HeterogeneousVolumes/HeterogeneousVolumesLiveShadingPipeline.usf
@@ -390,7 +390,7 @@ void RenderSingleScatteringWithLiveShadingCS(
 		if (bApplyEmissionAndTransmittance)
 		{
 			float Opacity = Luminance(1.0 - Transmittance);
-			RWLightingTexture[PixelCoord].a = Transmittance;
+			RWLightingTexture[PixelCoord].a = Luminance(Transmittance);
 
 #if DIM_WRITE_VELOCITY
 			bool bWriteVelocity =
diff --git a/Engine/Shaders/Private/HeterogeneousVolumes/HeterogeneousVolumesPreshadingPipeline.usf b/Engine/Shaders/Private/HeterogeneousVolumes/HeterogeneousVolumesPreshadingPipeline.usf
--- a/Engine/Shaders/Private/HeterogeneousVolumes/HeterogeneousVolumesPreshadingPipeline.usf
+++ b/Engine/Shaders/Private/HeterogeneousVolumes/HeterogeneousVolumesPreshadingPipeline.usf
@@ -632,7 +632,7 @@ void RenderSingleScatteringWithPreshadingCS(
 	if (bApplyEmissionAndTransmittance)
 	{
 		float Opacity = Luminance(1.0 - Transmittance);
-		RWLightingTexture[PixelCoord].a = Transmittance;
+		RWLightingTexture[PixelCoord].a = Luminance(Transmittance);
 
 #if DIM_WRITE_VELOCITY
 		bool bWriteVelocity =
