Fix volumetric shadows not rendering with instanced stereo (UE-210139).

diff --git a/Engine/Shaders/Private/ShadowProjectionPixelShader.usf b/Engine/Shaders/Private/ShadowProjectionPixelShader.usf
--- a/Engine/Shaders/Private/ShadowProjectionPixelShader.usf
+++ b/Engine/Shaders/Private/ShadowProjectionPixelShader.usf
@@ -146,7 +146,7 @@ void Main(
 	
 	float4 ShadowPosition = float4(0, 0, 0, 0);
 	
-#if MOBILE_MULTI_VIEW || INSTANCED_STEREO
+#if ((MOBILE_MULTI_VIEW || INSTANCED_STEREO) && SHADING_PATH_MOBILE)
 	ShadowPosition = mul(ScreenPosition, ResolvedView.MobileMultiviewShadowTransform);
 	ShadowPosition = mul(ShadowPosition, ScreenToShadowMatrix);
 #else
