Make sky atmosphere (without SkyDomeMesh) output opaque alpha (raster and PT).

diff --git a/Engine/Shaders/Private/PathTracing/PathTracingCore.ush b/Engine/Shaders/Private/PathTracing/PathTracingCore.ush
--- a/Engine/Shaders/Private/PathTracing/PathTracingCore.ush
+++ b/Engine/Shaders/Private/PathTracing/PathTracingCore.ush
@@ -648,6 +648,13 @@ FPackedPathTracingPayload TraceTransparentRay(inout FPathState PathState, int Bo
 
 		if (PackedPayload.IsMiss())
 		{
+			// This is for parity with the sky dome mesh behavior.
+			// Output fully opaque alpha if looking through outer space, if atmosphere is not in holdout.
+			if (bIsCameraRay && (VolumeFlags & PATH_TRACER_VOLUME_HOLDOUT_ATMOSPHERE) == 0)
+			{
+				PathState.Alpha += Luminance(PathState.PathThroughput);
+			}
+
 			// Ray didn't hit any real geometry
 			break;
 		}
diff --git a/Engine/Shaders/Private/SkyAtmosphere.usf b/Engine/Shaders/Private/SkyAtmosphere.usf
--- a/Engine/Shaders/Private/SkyAtmosphere.usf
+++ b/Engine/Shaders/Private/SkyAtmosphere.usf
@@ -965,6 +965,19 @@ void RenderSkyAtmosphereRayMarchingPS(
 		PreExposedL += SkyLuminance * LuminanceScale * (ViewOneOverPreExposure * OutputPreExposure);
 	
 		OutLuminance = PrepareOutput(PreExposedL, SkyGreyTransmittance);
+		
+	#if POST_PROCESS_ALPHA
+		if (View.RenderingReflectionCaptureMask == 0.0f)
+		{
+			// This is for parity with the sky dome mesh behavior.
+			// Output fully opaque alpha if looking through outer space, if atmosphere is not in holdout.
+			if (DeviceZ == FarDepthValue && !IsSkyAtmosphereHoldout(View.EnvironmentComponentsFlags))
+			{
+				OutLuminance.a = 1.0;
+			}
+		}
+	#endif
+		
 		return;
 	}
 #endif
