Workaround for Nvidia 700 series suspected compiler bug (DX11).

The flattened branch seems to result in corrupted texture samples. A proper branch works around the issue.

diff --git a/Engine/Shaders/Private/VirtualTextureCommon.ush b/Engine/Shaders/Private/VirtualTextureCommon.ush
--- a/Engine/Shaders/Private/VirtualTextureCommon.ush
+++ b/Engine/Shaders/Private/VirtualTextureCommon.ush
@@ -762,12 +762,28 @@ MaterialFloat4 TextureVirtualSample(
 {
 	float2 pUV = VTComputePhysicalUVs(PageTableResult, LayerIndex, Uniform);
 	bool bValid = IsValid(PageTableResult, LayerIndex);
-	const MaterialFloat4 Sample1 = bValid ? TextureVirtualSampleAniso(Physical, PhysicalSampler, pUV, PageTableResult.dUVdx, PageTableResult.dUVdy) : Uniform.FallbackValue;
+	MaterialFloat4 Sample1;
+	if (bValid)
+	{
+		Sample1 = TextureVirtualSampleAniso(Physical, PhysicalSampler, pUV, PageTableResult.dUVdx, PageTableResult.dUVdy);
+	}
+	else
+	{
+		Sample1 = Uniform.FallbackValue;
+	}
 
 #if VIRTUAL_TEXTURE_MANUAL_TRILINEAR_FILTERING
 	pUV = VTComputePhysicalUVs(PageTableResult, LayerIndex + 8u, Uniform);
 	bValid = IsValid(PageTableResult, LayerIndex + 8u);
-	const MaterialFloat4 Sample2 = bValid ? TextureVirtualSampleAniso(Physical, PhysicalSampler, pUV, PageTableResult.dUVdx, PageTableResult.dUVdy) : Sample1;
+	MaterialFloat4 Sample2;
+	if (bValid)
+	{
+		Sample2 = TextureVirtualSampleAniso(Physical, PhysicalSampler, pUV, PageTableResult.dUVdx, PageTableResult.dUVdy);
+	}
+	else
+	{
+		Sample2 = Sample1;
+	}
 	return lerp(Sample1, Sample2, PageTableResult.MipLevelFrac);
 #else
 	return Sample1;
@@ -783,8 +799,11 @@ MaterialFloat4 TextureVirtualSampleLevel(
 	const float2 pUV = VTComputePhysicalUVs(PageTableResult, LayerIndex, Uniform);
 	const bool bValid = IsValid(PageTableResult, LayerIndex);
 	// No need to apply dUVdx/dUVdy, don't support anisotropic when sampling a specific level
-
-	return bValid ? Physical.SampleLevel(PhysicalSampler, pUV, 0.0f) : Uniform.FallbackValue;
+	if (bValid)
+	{
+		return Physical.SampleLevel(PhysicalSampler, pUV, 0.0f);
+	}
+	return Uniform.FallbackValue;
 }
 
 
