Use VF_GPUSCENE_SET_INPUT_FOR_HAIR instead of reusing VF_GPUSCENE_SET_INPUT_FOR_RT in HairStrandsMaterialCommon.ush

diff --git a/Engine/Shaders/Private/HairStrands/HairStrandsMaterialCommon.ush b/Engine/Shaders/Private/HairStrands/HairStrandsMaterialCommon.ush
--- a/Engine/Shaders/Private/HairStrands/HairStrandsMaterialCommon.ush
+++ b/Engine/Shaders/Private/HairStrands/HairStrandsMaterialCommon.ush
@@ -63,7 +63,7 @@ FVertexFactoryInterpolantsVSToPS GetInterpolants(uint PrimitiveId, uint HairCont
 
 	FVertexFactoryInput Input;
 	// Assumes no actual instancing used for these draws. Will load the first instance data.
-	VF_GPUSCENE_SET_INPUT_FOR_RT(Input, PrimitiveId, 0U);
+	VF_GPUSCENE_SET_INPUT_FOR_HAIR(Input, PrimitiveId, 0U);
 	// * USE_HAIR_TRIANGLE_STRIP=0: the vertex factory expect triangles strip index. A quad is made of two triangles with indices: 0,1.
 	// * USE_HAIR_TRIANGLE_STRIP=1: the vertex factory expect triangles list index. A quad is made of two triangles with indices: 0,1,2,3,4,5.
 	Input.VertexId = HairControlPointId * HAIR_POINT_TO_VERTEX; 
diff --git a/Engine/Shaders/Private/SceneData.ush b/Engine/Shaders/Private/SceneData.ush
--- a/Engine/Shaders/Private/SceneData.ush
+++ b/Engine/Shaders/Private/SceneData.ush
@@ -1214,6 +1214,7 @@ FSceneDataIntermediates GetSceneDataIntermediates(uint InstanceIdOffset, uint Dr
 	//				 instead they mark the top bit in the primitive ID and disable auto instancing such that there is an 1:1:1
 	//               drawcmd:primitive:instance. Then we can just look up the primitive and fetch the instance data index.
 	// GPUCULL_TODO: Workaround also used for ray tracing interfacing with the VFs, that also supply a DrawInstanceId. 
+	// GPUCULL_TODO: Workaround also used by HairStrandsMaterialCommon.ush
 	// We mark the PrimitiveID with the top bit in dynamic draw passes
 	if ((InstanceIdOffset & VF_TREAT_INSTANCE_ID_OFFSET_AS_PRIMITIVE_ID_FLAG) != 0U)
 	{
diff --git a/Engine/Shaders/Private/VertexFactoryCommon.ush b/Engine/Shaders/Private/VertexFactoryCommon.ush
--- a/Engine/Shaders/Private/VertexFactoryCommon.ush
+++ b/Engine/Shaders/Private/VertexFactoryCommon.ush
@@ -48,6 +48,11 @@
 		_VFInputOut_.InstanceIdOffset = (_PrimitiveId_) | VF_TREAT_INSTANCE_ID_OFFSET_AS_PRIMITIVE_ID_FLAG;	\
 		_VFInputOut_.DrawInstanceId = _DrawInstanceId_;
 
+	// GPUCULL_TODO: Makes use of DynamicDraw workaround, should instead refactor VF interfaces and provide the instance ID data directly.
+	#define VF_GPUSCENE_SET_INPUT_FOR_HAIR(_VFInputOut_, _PrimitiveId_, _DrawInstanceId_) \
+		_VFInputOut_.InstanceIdOffset = (_PrimitiveId_) | VF_TREAT_INSTANCE_ID_OFFSET_AS_PRIMITIVE_ID_FLAG;	\
+		_VFInputOut_.DrawInstanceId = _DrawInstanceId_;
+
 	#define VF_INSTANCED_STEREO_DECLARE_INPUT_BLOCK()
 
 	#undef GetInstanceIdFromVF
@@ -58,6 +63,7 @@
 	#define VF_GPUSCENE_GET_LIGHTMAP_UV_SCALE_BIAS(VFInput, LightmapDataIndex) PrecomputedLightingBuffer.LightMapCoordinateScaleBias
 	#define VF_GPUSCENE_GET_SHADOWMAP_UV_SCALE_BIAS(VFInput, LightmapDataIndex) PrecomputedLightingBuffer.ShadowMapCoordinateScaleBias
 	#define VF_GPUSCENE_SET_INPUT_FOR_RT(_VFInputOut_, _PrimitiveId_, _DrawInstanceId_)
+	#define VF_GPUSCENE_SET_INPUT_FOR_HAIR(_VFInputOut_, _PrimitiveId_, _DrawInstanceId_)
 
 	#if INSTANCED_STEREO
 		#define VF_INSTANCED_STEREO_DECLARE_INPUT_BLOCK() \
