Changes to various Niagara particles and volumetrics to improve visuals with Ortho rendering:

- Resolved the camera + view forward positioning to fix skewed volumetrics
- Fixed PDO recalculations which resolved water positional issues. (Reverted previous hair PDO removal hack as this is working now)
- Manually updated the WaterSDF asset for ortho view camera vector compensation in the custom node.
- Minor fixes for certain particle calculations

#virtualized

diff --git a/Engine/Plugins/FX/Niagara/Shaders/Private/NiagaraDataInterfaceGBufferTemplate.ush b/Engine/Plugins/FX/Niagara/Shaders/Private/NiagaraDataInterfaceGBufferTemplate.ush
--- a/Engine/Plugins/FX/Niagara/Shaders/Private/NiagaraDataInterfaceGBufferTemplate.ush
+++ b/Engine/Plugins/FX/Niagara/Shaders/Private/NiagaraDataInterfaceGBufferTemplate.ush
@@ -97,7 +97,7 @@ void DecodeWorldVelocity_{ParameterName}(in float2 ScreenUV, in bool bApplyViewp
 		float SceneDeviceZ = LookupDeviceZ(ScreenUV);
 		float SceneDepth = ConvertFromDeviceZ(SceneDeviceZ);
  		float2 ScreenPosition = (ScreenUV - View.ScreenPositionScaleBias.wz) / View.ScreenPositionScaleBias.xy;
-		float3 TranslatedWorldPosition = mul(float4(ScreenPosition * SceneDepth, SceneDepth, 1), PrimaryView.ScreenToTranslatedWorld).xyz;
+		float3 TranslatedWorldPosition = mul(float4(GetScreenPositionForProjectionType(ScreenPosition, SceneDepth), SceneDepth, 1), PrimaryView.ScreenToTranslatedWorld).xyz;
  
 		float PrevSceneDepth = 0;
 	#if VELOCITY_ENCODE_DEPTH
@@ -107,7 +107,7 @@ void DecodeWorldVelocity_{ParameterName}(in float2 ScreenUV, in bool bApplyViewp
 		PrevSceneDepth = SceneDepth; // CalcPrevSceneDepth(PrevScreenUV);
 	#endif
 		float2 PrevScreenPosition = ScreenPosition - Velocity.xy;
-		float3 PrevTranslatedWorldPosition = mul(float4(PrevScreenPosition * PrevSceneDepth, PrevSceneDepth, 1), PrimaryView.PrevScreenToTranslatedWorld).xyz;
+		float3 PrevTranslatedWorldPosition = mul(float4(GetScreenPositionForProjectionType(PrevScreenPosition, PrevSceneDepth), PrevSceneDepth, 1), PrimaryView.PrevScreenToTranslatedWorld).xyz;
 
 		float3 PreViewTranslationOffset = DFFastLocalSubtractDemote(PrimaryView.PreViewTranslation, PrimaryView.PrevPreViewTranslation);
  
diff --git a/Engine/Shaders/Private/BasePassPixelShader.usf b/Engine/Shaders/Private/BasePassPixelShader.usf
--- a/Engine/Shaders/Private/BasePassPixelShader.usf
+++ b/Engine/Shaders/Private/BasePassPixelShader.usf
@@ -686,7 +686,14 @@ void ApplyPixelDepthOffsetForBasePass(inout FMaterialPixelParameters MaterialPar
 	float PixelDepthOffset = ApplyPixelDepthOffsetToMaterialParameters(MaterialParameters, PixelMaterialInputs, OutDepth);
 
 #if WRITES_VELOCITY_TO_GBUFFER && !IS_NANITE_PASS
-	BasePassInterpolants.VelocityPrevScreenPosition.w += PixelDepthOffset;
+	if(IsOrthoProjection(ResolvedView))
+	{
+		BasePassInterpolants.VelocityPrevScreenPosition.z += PixelDepthOffset;
+	}
+	else
+	{
+		BasePassInterpolants.VelocityPrevScreenPosition.w += PixelDepthOffset;
+	}
 #endif
 }
 
diff --git a/Engine/Shaders/Private/Common.ush b/Engine/Shaders/Private/Common.ush
--- a/Engine/Shaders/Private/Common.ush
+++ b/Engine/Shaders/Private/Common.ush
@@ -1196,15 +1196,6 @@ float4 GetPerPixelLightAttenuation(float2 UV)
 	return DecodeLightAttenuation(Texture2DSampleLevel(LightAttenuationTexture, LightAttenuationTextureSampler, UV, 0));
 }
 
-// also see ConvertToDeviceZ()
-// @param DeviceZ value that is stored in the depth buffer (Z/W)
-// @return SceneDepth (linear in world units, W)
-float ConvertFromDeviceZ(float DeviceZ)
-{
-	// Supports ortho and perspective, see CreateInvDeviceZToWorldZTransform()
-	return DeviceZ * View.InvDeviceZToWorldZTransform[0] + View.InvDeviceZToWorldZTransform[1] + 1.0f / (DeviceZ * View.InvDeviceZToWorldZTransform[2] - View.InvDeviceZToWorldZTransform[3]);
-}
-
 // Returns whether or not the given projection matrix is orthographic
 bool IsOrthoProjection(float4x4 ViewToClip)
 {
@@ -1217,7 +1208,6 @@ bool IsOrthoProjection(ViewState InView)
 	return IsOrthoProjection(InView.ViewToClip);
 }
 
-
 // Returns whether or not the default view is orthographic
 bool IsOrthoProjection()
 {
@@ -1248,9 +1238,28 @@ float ConvertToDeviceZ(float SceneDepth)
 	}
 }
 
+// also see ConvertToDeviceZ()
+// @param DeviceZ value that is stored in the depth buffer (Z/W)
+// @return SceneDepth (linear in world units, W)
+float ConvertFromDeviceZ(float DeviceZ)
+{
+	// Supports ortho and perspective, see CreateInvDeviceZToWorldZTransform()
+	return DeviceZ * View.InvDeviceZToWorldZTransform[0] + View.InvDeviceZToWorldZTransform[1] + 1.0f / (DeviceZ * View.InvDeviceZToWorldZTransform[2] - View.InvDeviceZToWorldZTransform[3]);
+}
+
+float GetScreenPositionDepth(float4 ScreenPosition)
+{
+	return select(IsOrthoProjection(ResolvedView), ConvertFromDeviceZ(ScreenPosition.z), ScreenPosition.w);
+}
+
+float3 GetCameraVector(ViewState InView, float3 Position, float3 CameraPosition)
+{
+	return select(IsOrthoProjection(InView), InView.ViewForward, normalize(Position - CameraPosition));
+}
+
 float3 GetCameraVector(float3 Position, float3 CameraPosition)
 {
-	return select(IsOrthoProjection(), View.ViewForward, normalize(Position - CameraPosition));
+	return GetCameraVector(PrimaryView, Position, CameraPosition);
 }
 
 float3 GetCameraVectorFromWorldPosition(float3 WorldPosition)
@@ -1258,9 +1267,14 @@ float3 GetCameraVectorFromWorldPosition(float3 WorldPosition)
 	return GetCameraVector(WorldPosition, DFHackToFloat(PrimaryView.WorldCameraOrigin));
 }
 
+float3 GetCameraVectorFromTranslatedWorldPosition(ViewState InView, float3 TranslatedWorldPosition)
+{
+	return GetCameraVector(TranslatedWorldPosition, InView.TranslatedWorldCameraOrigin);
+}
+
 float3 GetCameraVectorFromTranslatedWorldPosition(float3 TranslatedWorldPosition)
 {
-	return GetCameraVector(TranslatedWorldPosition, PrimaryView.TranslatedWorldCameraOrigin);
+	return GetCameraVectorFromTranslatedWorldPosition(PrimaryView, TranslatedWorldPosition);
 }
 
 float GetDistanceToCameraFromViewVector(float3 DistanceVector)
@@ -1394,6 +1408,24 @@ void SvPositionToResolvedScreenPositionDeriv(float4 SvPosition, float2 PPZ_DDX_D
 	ScreenPositionDDY = NDCPos * W_DDX_DDY.y + NDCPosDDY * SvPosition.w;
 }
 
+float3 GetTranslatedWorldCameraPosFromView(ViewState InView, float2 SvPosition)
+{
+	// Get the camera position per pixel in ortho, or just the camera location in perspective.
+	if (IsOrthoProjection(InView))
+	{
+		return SvPositionToTranslatedWorld(float4(SvPosition, 1.0f, 1.0f));
+	}
+	else
+	{
+		return InView.TranslatedWorldCameraOrigin;
+	}
+}
+
+float3 GetTranslatedWorldCameraPosFromView(float2 SvPosition)
+{
+	return GetTranslatedWorldCameraPosFromView(PrimaryView, SvPosition);
+}
+
 float2 SvPositionToViewportUV(float4 SvPosition)
 {
 	// can be optimized from 2SUB+2MUL to 2MAD
diff --git a/Engine/Shaders/Private/HairStrands/HairCardsAttributeCommon.ush b/Engine/Shaders/Private/HairStrands/HairCardsAttributeCommon.ush
--- a/Engine/Shaders/Private/HairStrands/HairCardsAttributeCommon.ush
+++ b/Engine/Shaders/Private/HairStrands/HairCardsAttributeCommon.ush
@@ -144,7 +144,9 @@ float4 GetAuxilary    (float2 InAtlasUV, float4 InDefault)      { float4 Out = I
 
 float GetHairCardsDepth(float2 InAtlasUV, float InDeviceZ)
 {
-	return select(IsOrthoProjection(), 0.0f, ConvertToDeviceZ(ConvertFromDeviceZ(InDeviceZ) + GetDepthOffset(InAtlasUV)));
+	const float SceneDepthOffset = GetDepthOffset(InAtlasUV);
+	const float SceneDepth = ConvertFromDeviceZ(InDeviceZ);
+	return ConvertToDeviceZ(SceneDepth + SceneDepthOffset);
 }
 
 float3 GetHairCardsTangent(float2 InAtlasUV, half3x3 TangentToWorld, bool bUseTangentSpace)
diff --git a/Engine/Shaders/Private/HeterogeneousVolumes/HeterogeneousVolumesHardwareRayTracing.usf b/Engine/Shaders/Private/HeterogeneousVolumes/HeterogeneousVolumesHardwareRayTracing.usf
--- a/Engine/Shaders/Private/HeterogeneousVolumes/HeterogeneousVolumesHardwareRayTracing.usf
+++ b/Engine/Shaders/Private/HeterogeneousVolumes/HeterogeneousVolumesHardwareRayTracing.usf
@@ -338,7 +338,7 @@ RAY_TRACING_ENTRY_RAYGEN(SparseVoxelsRayGenerationShader1)
 
 	// Intersect ray with bounding volume
 	// TODO: LWC integration..
-	float3 WorldRayOrigin = DFHackToFloat(DFFastSubtract(View.TranslatedWorldCameraOrigin, PrimaryView.PreViewTranslation));
+	float3 WorldRayOrigin = DFHackToFloat(DFFastSubtract(GetTranslatedWorldCameraPosFromView(PixelCoord + 0.5), PrimaryView.PreViewTranslation));
 	float3 WorldRayEnd = DFHackToFloat(SvPositionToWorld(float4(PixelCoord + 0.5, DeviceZ, 1)));
 	float3 WorldRayDirection = WorldRayEnd - WorldRayOrigin;
 	float WorldRayLength = length(WorldRayDirection); // SceneDepth
@@ -504,7 +504,7 @@ RAY_TRACING_ENTRY_RAYGEN(RenderSingleScatteringWithPreshadingRGS)
 
 	// Intersect ray with bounding volume
 	// TODO: LWC integration..
-	float3 WorldRayOrigin = DFHackToFloat(DFFastSubtract(View.TranslatedWorldCameraOrigin, PrimaryView.PreViewTranslation));
+	float3 WorldRayOrigin = DFHackToFloat(DFFastSubtract(GetTranslatedWorldCameraPosFromView(PixelCoord + 0.5), PrimaryView.PreViewTranslation));
 	float3 WorldRayEnd = DFHackToFloat(SvPositionToWorld(float4(PixelCoord + 0.5, DeviceZ, 1)));
 	float3 WorldRayDirection = WorldRayEnd - WorldRayOrigin;
 	float WorldRayLength = length(WorldRayDirection);
diff --git a/Engine/Shaders/Private/HeterogeneousVolumes/HeterogeneousVolumesLiveShadingPipeline.usf b/Engine/Shaders/Private/HeterogeneousVolumes/HeterogeneousVolumesLiveShadingPipeline.usf
--- a/Engine/Shaders/Private/HeterogeneousVolumes/HeterogeneousVolumesLiveShadingPipeline.usf
+++ b/Engine/Shaders/Private/HeterogeneousVolumes/HeterogeneousVolumesLiveShadingPipeline.usf
@@ -332,7 +332,7 @@ void RenderSingleScatteringWithLiveShadingCS(
 
 	// Intersect ray with bounding volume
 	// TODO: LWC integration..
-	float3 WorldRayOrigin = DFHackToFloat(DFFastSubtract(View.TranslatedWorldCameraOrigin, PrimaryView.PreViewTranslation));
+	float3 WorldRayOrigin = DFHackToFloat(DFFastSubtract(GetTranslatedWorldCameraPosFromView(PixelCoord + View.TemporalAAJitter.xy), PrimaryView.PreViewTranslation));
 	float3 WorldRayEnd = DFHackToFloat(SvPositionToWorld(float4(PixelCoord + View.TemporalAAJitter.xy, DeviceZ, 1)));
 	float3 WorldRayDirection = WorldRayEnd - WorldRayOrigin;
 	float WorldRayLength = length(WorldRayDirection);
diff --git a/Engine/Shaders/Private/HeterogeneousVolumes/HeterogeneousVolumesPreshadingPipeline.usf b/Engine/Shaders/Private/HeterogeneousVolumes/HeterogeneousVolumesPreshadingPipeline.usf
--- a/Engine/Shaders/Private/HeterogeneousVolumes/HeterogeneousVolumesPreshadingPipeline.usf
+++ b/Engine/Shaders/Private/HeterogeneousVolumes/HeterogeneousVolumesPreshadingPipeline.usf
@@ -541,7 +541,7 @@ void RenderSingleScatteringWithPreshadingCS(
 
 	// Intersect ray with bounding volume
 	// TODO: LWC integration..
-	float3 WorldRayOrigin = DFHackToFloat(DFFastSubtract(View.TranslatedWorldCameraOrigin, PrimaryView.PreViewTranslation));
+	float3 WorldRayOrigin = DFHackToFloat(DFFastSubtract(GetTranslatedWorldCameraPosFromView(PixelCoord + View.TemporalAAJitter.xy), PrimaryView.PreViewTranslation));
 	float3 WorldRayEnd = DFHackToFloat(SvPositionToWorld(float4(PixelCoord + View.TemporalAAJitter.xy, DeviceZ, 1)));
 	float3 WorldRayDirection = WorldRayEnd - WorldRayOrigin;
 	float WorldRayLength = length(WorldRayDirection);
diff --git a/Engine/Shaders/Private/HeterogeneousVolumes/HeterogeneousVolumesTracingUtils.ush b/Engine/Shaders/Private/HeterogeneousVolumes/HeterogeneousVolumesTracingUtils.ush
--- a/Engine/Shaders/Private/HeterogeneousVolumes/HeterogeneousVolumesTracingUtils.ush
+++ b/Engine/Shaders/Private/HeterogeneousVolumes/HeterogeneousVolumesTracingUtils.ush
@@ -34,7 +34,7 @@ float3 Calculate3DVelocityAtWorldPos(float3 WorldRayOrigin, float3 WorldRayDirec
 	ClipPos /= ClipPos.w;
 
 	float3 PreViewTranslationOffset = DFFastLocalSubtractDemote(PrimaryView.PreViewTranslation, PrimaryView.PrevPreViewTranslation);
-	float3 PrevTranslatedWorldPos = mul(float4(ClipPos.xy * WorldHitT, WorldHitT, 1), View.PrevScreenToTranslatedWorld).xyz + PreViewTranslationOffset;
+	float3 PrevTranslatedWorldPos = mul(float4(GetScreenPositionForProjectionType(ClipPos.xy, WorldHitT), WorldHitT, 1), View.PrevScreenToTranslatedWorld).xyz + PreViewTranslationOffset;
 	float4 PrevClipPos = mul(float4(PrevTranslatedWorldPos, 1), View.TranslatedWorldToClip);
 	PrevClipPos /= PrevClipPos.w;
 
diff --git a/Engine/Shaders/Private/HeterogeneousVolumes/HeterogeneousVolumesVoxelGridPipeline.usf b/Engine/Shaders/Private/HeterogeneousVolumes/HeterogeneousVolumesVoxelGridPipeline.usf
--- a/Engine/Shaders/Private/HeterogeneousVolumes/HeterogeneousVolumesVoxelGridPipeline.usf
+++ b/Engine/Shaders/Private/HeterogeneousVolumes/HeterogeneousVolumesVoxelGridPipeline.usf
@@ -181,7 +181,7 @@ void TopLevelGridCalculateVoxelSize(
 		float PixelWidth = TanHalfFOV.x / HalfWidth;
 
 		float3 WorldCameraOrigin = DFHackToFloat(DFFastSubtract(View.TranslatedWorldCameraOrigin, PrimaryView.PreViewTranslation));
-		float Distance = length(WorldVoxelCenter - WorldCameraOrigin) - length(WorldVoxelExtent);
+		float Distance = GetDistanceToCameraFromViewVector(WorldVoxelCenter - WorldCameraOrigin) - length(WorldVoxelExtent);
 		float VoxelWidth = Distance * PixelWidth * ShadingRate;
 
 		// Frustum intersection determines which min voxel-size to clip by
diff --git a/Engine/Shaders/Private/HeterogeneousVolumes/HeterogeneousVolumesVoxelGridRendering.usf b/Engine/Shaders/Private/HeterogeneousVolumes/HeterogeneousVolumesVoxelGridRendering.usf
--- a/Engine/Shaders/Private/HeterogeneousVolumes/HeterogeneousVolumesVoxelGridRendering.usf
+++ b/Engine/Shaders/Private/HeterogeneousVolumes/HeterogeneousVolumesVoxelGridRendering.usf
@@ -583,7 +583,7 @@ void RenderTransmittanceWithVoxelGridCS(
 	float SceneDepth = min(ConvertFromDeviceZ(DeviceZ), MaxTraceDistance);
 	DeviceZ = ConvertToDeviceZ(SceneDepth);
 
-	float3 RayOrigin = View.TranslatedWorldCameraOrigin;
+	float3 RayOrigin = GetTranslatedWorldCameraPosFromView(PixelCoord + 0.5);
 	float3 RayEnd = SvPositionToTranslatedWorld(float4(PixelCoord + 0.5, DeviceZ, 1));
 	float3 RayDirection = RayEnd - RayOrigin;
 	float RayTMin = 0.0;
diff --git a/Engine/Shaders/Private/MaterialTemplate.ush b/Engine/Shaders/Private/MaterialTemplate.ush
--- a/Engine/Shaders/Private/MaterialTemplate.ush
+++ b/Engine/Shaders/Private/MaterialTemplate.ush
@@ -1035,32 +1035,12 @@ float4 GetScreenPosition(FMaterialPixelParameters Parameters)
 // Returns the pixel's depth, in world units. Works for both orthographic and perspective projections:
 float GetPixelDepth(FMaterialVertexParameters Parameters)
 {
-	FLATTEN
-	if (View.ViewToClip[3][3] < 1.0f)
-	{
-		// Perspective
-		return GetScreenPosition(Parameters).w;
-	}
-	else
-	{
-		// Ortho
-		return ConvertFromDeviceZ(GetScreenPosition(Parameters).z);
-	}
+	return GetScreenPositionDepth(GetScreenPosition(Parameters));
 }
 
 float GetPixelDepth(FMaterialPixelParameters Parameters)
 {
-	FLATTEN
-	if (View.ViewToClip[3][3] < 1.0f)
-	{
-		// Perspective
-		return GetScreenPosition(Parameters).w;
-	}
-	else
-	{
-		// Ortho
-		return ConvertFromDeviceZ(GetScreenPosition(Parameters).z);
-	}
+	return GetScreenPositionDepth(GetScreenPosition(Parameters));
 }
 
 float2 GetSceneTextureUV(FMaterialVertexParameters Parameters)
@@ -4310,14 +4290,33 @@ float ApplyPixelDepthOffsetToMaterialParameters(inout FMaterialPixelParameters M
 	// RE_Z causes the HW to perform depth test twice, once before executing pixel shader, and once after
 	// The PreZ pass will write depth buffer using depth offset, then the base pass will test against this value using both modified and unmodifed depth
 	// If the unmodified depth is ever slightly less than the modified depth, the initial depth test will fail, which results in z-fighting/flickering type artifacts
-	float DeviceDepth = min(MaterialParameters.ScreenPosition.z / (MaterialParameters.ScreenPosition.w + PixelDepthOffset), MaterialParameters.SvPosition.z);
+	float DeviceDepth;
+	if(IsOrthoProjection(ResolvedView))
+	{
+		//Ortho projections have PDO applied to z rather than w. w should always stay as 1.0f
+		float OriginalDeviceZ = MaterialParameters.SvPosition.z;
+		float SceneDepth = ConvertFromDeviceZ(OriginalDeviceZ);
+		DeviceDepth = max(SceneDepth + PixelDepthOffset, SceneDepth);
+
+		PixelDepthOffset = SceneDepth - DeviceDepth;
+
+		MaterialParameters.ScreenPosition.z = ConvertToDeviceZ(SceneDepth + PixelDepthOffset);
+		MaterialParameters.SvPosition.z = MaterialParameters.ScreenPosition.z;
 
-	// Once we've computed our (clamped) device depth, recompute PixelDepthOffset again to take the potential clamp into account
-	PixelDepthOffset = (MaterialParameters.ScreenPosition.z - DeviceDepth * MaterialParameters.ScreenPosition.w) / DeviceDepth;
+		DeviceDepth = ConvertToDeviceZ(DeviceDepth);
+		PixelDepthOffset = MaterialParameters.SvPosition.z - OriginalDeviceZ;
+	}
+	else
+	{
+		DeviceDepth = min(MaterialParameters.ScreenPosition.z / (MaterialParameters.ScreenPosition.w + PixelDepthOffset), MaterialParameters.SvPosition.z);
 
-	// Update positions used for shading
-	MaterialParameters.ScreenPosition.w += PixelDepthOffset;
-	MaterialParameters.SvPosition.w = MaterialParameters.ScreenPosition.w;
+		// Once we've computed our (clamped) device depth, recompute PixelDepthOffset again to take the potential clamp into account
+		PixelDepthOffset = (MaterialParameters.ScreenPosition.z - DeviceDepth * MaterialParameters.ScreenPosition.w) / DeviceDepth;
+
+		// Update positions used for shading
+		MaterialParameters.ScreenPosition.w += PixelDepthOffset;
+		MaterialParameters.SvPosition.w = MaterialParameters.ScreenPosition.w;
+	}
 
 	MaterialParameters.LWCData.AbsoluteWorldPosition = WSAdd(MaterialParameters.LWCData.AbsoluteWorldPosition, -MaterialParameters.CameraVector * PixelDepthOffset);
 	MaterialParameters.AbsoluteWorldPosition = WSToDF(MaterialParameters.LWCData.AbsoluteWorldPosition);
diff --git a/Engine/Shaders/Private/ParticleBeamTrailVertexFactory.ush b/Engine/Shaders/Private/ParticleBeamTrailVertexFactory.ush
--- a/Engine/Shaders/Private/ParticleBeamTrailVertexFactory.ush
+++ b/Engine/Shaders/Private/ParticleBeamTrailVertexFactory.ush
@@ -175,7 +175,7 @@ void GetTangents(FVertexFactoryInput Input,out float4 Right,out float4 Up)
 	float4	Position			= SimulationToTranslatedWorld(Input.Position.xyz),
 			OldPosition			= SimulationToTranslatedWorld(Input.OldPosition.xyz);
 
-	float3	CameraDirection		= SafeNormalize(ResolvedView.TranslatedWorldCameraOrigin - Position.xyz),
+	float3	CameraDirection		= -GetCameraVectorFromTranslatedWorldPosition(ResolvedView, Position.xyz),
 			ParticleDirection	= SafeNormalize(Position.xyz - OldPosition.xyz);
 
 	float4	Right_Square		= BeamTrailVF.CameraRight,
diff --git a/Engine/Shaders/Private/ParticleGPUSpriteVertexFactory.ush b/Engine/Shaders/Private/ParticleGPUSpriteVertexFactory.ush
--- a/Engine/Shaders/Private/ParticleGPUSpriteVertexFactory.ush
+++ b/Engine/Shaders/Private/ParticleGPUSpriteVertexFactory.ush
@@ -189,7 +189,7 @@ void GetSpriteTangents(
 	float3 CameraUp = lerp(-ResolvedViewUp, EmitterDynamicUniforms.AxisLockUp.xyz, EmitterDynamicUniforms.AxisLockUp.w);
 
 	// Determine the vector from the particle to the camera and the particle's movement direction.
-	float3 CameraDirection	= SafeNormalize(ResolvedView.TranslatedWorldCameraOrigin - ParticleTranslatedWorldPosition.xyz);
+	float3 CameraDirection	= -GetCameraVectorFromTranslatedWorldPosition(ResolvedView, ParticleTranslatedWorldPosition.xyz);
 	float3 RightVector = CameraRight.xyz;
 	float3 UpVector = CameraUp.xyz;
 
@@ -197,8 +197,7 @@ void GetSpriteTangents(
 	if (EmitterUniforms.CameraFacingBlend.x > 0.f)
 	{
 		// Blend between PSA_FacingCamera and PSA_Square over distance
-		float3 CameraVector = ResolvedView.TranslatedWorldCameraOrigin.xyz - ParticleTranslatedWorldPosition.xyz;
-		float CameraDistanceSq = dot(CameraVector, CameraVector);
+		float CameraDistanceSq = GetDistanceToCameraFromViewVectorSqr(ResolvedView.TranslatedWorldCameraOrigin.xyz - ParticleTranslatedWorldPosition.xyz);
 		float AlignmentMode = saturate(CameraDistanceSq * EmitterUniforms.CameraFacingBlend.y - EmitterUniforms.CameraFacingBlend.z);
 
 		float3 CameraFacingRight = SafeNormalize(cross(CameraDirection, float3(0,0,1)));
diff --git a/Engine/Shaders/Private/ParticleSimulationShader.usf b/Engine/Shaders/Private/ParticleSimulationShader.usf
--- a/Engine/Shaders/Private/ParticleSimulationShader.usf
+++ b/Engine/Shaders/Private/ParticleSimulationShader.usf
@@ -267,7 +267,7 @@ float CollisionDepthBounds;
 /** TODO: Should be moved to Common.usf and used globaly. */
 float3 TranslatedWorldPositionFromSceneDepth(float2 ScreenPosition, float SceneDepth)
 {
-	float3 WorldPosition = mul(float4(ScreenPosition * SceneDepth, SceneDepth, 1), PrimaryView.ScreenToTranslatedWorld).xyz;
+	float3 WorldPosition = mul(float4(GetScreenPositionForProjectionType(ScreenPosition, SceneDepth), SceneDepth, 1), PrimaryView.ScreenToTranslatedWorld).xyz;
 	return WorldPosition;
 }
 
@@ -343,7 +343,7 @@ void CollideWithDepthBuffer(
 		// Sample the depth buffer to get a world position near the particle.
 		float2 ScreenUV = ScreenPosition * View.ScreenPositionScaleBias.xy + View.ScreenPositionScaleBias.wz;
 		float SceneDepth = CalcSceneDepth(ScreenUV);
-		if (abs(ClipPosition.w - SceneDepth) < CollisionDepthBounds)
+		if (abs(GetScreenPositionDepth(ClipPosition) - SceneDepth) < CollisionDepthBounds)
 		{
 #if FORWARD_SHADING_ACCURATE_NORMAL && FORWARD_SHADING_NORMAL_CALC
 			float SceneDepth0 = CalcSceneDepth(ScreenUV + float2(View.BufferSizeAndInvSize.z, 0.0));
diff --git a/Engine/Shaders/Private/ParticleSpriteVertexFactory.ush b/Engine/Shaders/Private/ParticleSpriteVertexFactory.ush
--- a/Engine/Shaders/Private/ParticleSpriteVertexFactory.ush
+++ b/Engine/Shaders/Private/ParticleSpriteVertexFactory.ush
@@ -365,7 +365,7 @@ void GetTangents(float3 TranslatedWorldPosition, float3 OldTranslatedWorldPositi
 	float3 CameraUp = lerp(-ResolvedViewUp, SpriteVF.AxisLockUp.xyz, SpriteVF.AxisLockUp.w);
 
 	// Determine the vector from the particle to the camera and the particle's movement direction.
-	float3 CameraDirection = SafeNormalize(ResolvedView.TranslatedWorldCameraOrigin - TranslatedWorldPosition);
+	float3 CameraDirection = -GetCameraVectorFromTranslatedWorldPosition(ResolvedView, TranslatedWorldPosition);
 	float3 RightVector = CameraRight.xyz;
 	float3 UpVector = CameraUp.xyz;
 
@@ -375,8 +375,7 @@ void GetTangents(float3 TranslatedWorldPosition, float3 OldTranslatedWorldPositi
 	if (SpriteVF.CameraFacingBlend.x > 0.f)
 	{
 		// Blend between PSA_FacingCamera and PSA_Square over distance
-		float3 CameraVector = ResolvedView.TranslatedWorldCameraOrigin - TranslatedWorldPosition;
-		float CameraDistanceSq = dot(CameraVector, CameraVector);
+		float CameraDistanceSq = GetDistanceToCameraFromViewVectorSqr(ResolvedView.TranslatedWorldCameraOrigin - TranslatedWorldPosition);
 		float AlignmentMode = saturate(CameraDistanceSq * SpriteVF.CameraFacingBlend.y - SpriteVF.CameraFacingBlend.z);
 
 		float3 CameraFacingRight = SafeNormalize(cross(CameraDirection, float3(0,0,1)));
diff --git a/Engine/Shaders/Private/SingleLayerWaterShading.ush b/Engine/Shaders/Private/SingleLayerWaterShading.ush
--- a/Engine/Shaders/Private/SingleLayerWaterShading.ush
+++ b/Engine/Shaders/Private/SingleLayerWaterShading.ush
@@ -57,7 +57,7 @@ WaterVolumeLightingOutput EvaluateWaterVolumeLighting(
 
 	if (WaterVisibility > 0.0f)
 	{
-		float WaterDepth = MaterialParameters.ScreenPosition.w; // Scene depth with Pixel-Depth-Offset taken into account is in ScreenPosition.w already
+		float WaterDepth = GetPixelDepth(MaterialParameters); // Scene depth with Pixel-Depth-Offset taken into account is in ScreenPosition.w already
 #if SINGLE_LAYER_WATER_MOBILE_SHADING
 
 #if SINGLE_LAYER_WATER_SHADING_QUALITY == SINGLE_LAYER_WATER_SHADING_QUALITY_MOBILE_WITH_DEPTH_TEXTURE 
