Fix shader errors when running with Substrate enabled

diff --git a/Engine/Shaders/Private/PathTracing/PathTracingCommon.ush b/Engine/Shaders/Private/PathTracing/PathTracingCommon.ush
--- a/Engine/Shaders/Private/PathTracing/PathTracingCommon.ush
+++ b/Engine/Shaders/Private/PathTracing/PathTracingCommon.ush
@@ -523,6 +523,13 @@ struct FPathTracingPayload : FMinimalPayload
 		SpecularColor = Weight;
 		PhaseG = G;
 	}
+	float3 GetDualHGWeight() { return SpecularEdgeColor; }
+	float3 GetDualHGPhaseData() { return RoughnessData; }
+	void SetDualHG(float3 Weight, float3 PhaseData)
+	{
+		SpecularEdgeColor = Weight;
+		RoughnessData = PhaseData;
+	}
 
 	// SHADINGMODELID_EYE
 	float GetEyeRoughness() { return RoughnessData.x; }
