Fix warning from path tracing material on substrate materials when substrate is turned off

diff --git a/Engine/Shaders/Private/PathTracing/PathTracingGPULightmassMaterialHitShader.usf b/Engine/Shaders/Private/PathTracing/PathTracingGPULightmassMaterialHitShader.usf
--- a/Engine/Shaders/Private/PathTracing/PathTracingGPULightmassMaterialHitShader.usf
+++ b/Engine/Shaders/Private/PathTracing/PathTracingGPULightmassMaterialHitShader.usf
@@ -194,7 +194,8 @@ RAY_TRACING_ENTRY_CLOSEST_HIT(GPULightmassMaterialCHS,
 #elif MATERIALBLENDING_MODULATE
 	Payload.BSDFOpacity = 0.0;
 	Payload.TransparencyColor = GetMaterialEmissive(PixelMaterialInputs);
-	HLSL_STATIC_ASSERT(MATERIAL_SHADINGMODEL_UNLIT == 1, "Modulate blend mode requires unlit shading model");
+	HLSL_STATIC_ASSERT(MATERIAL_SHADINGMODEL_UNLIT == 1 ||
+	                   MATERIAL_SHADINGMODEL_DEFAULT_LIT, "Modulate blend mode requires unlit shading model");
 	Payload.ShadingModelID = SHADINGMODELID_UNLIT;
 #elif MATERIALBLENDING_MASKED && MATERIAL_DITHER_OPACITY_MASK
 	// dithering emulates real transparency, so switch to translucent
diff --git a/Engine/Shaders/Private/PathTracing/PathTracingMaterialHitShader.usf b/Engine/Shaders/Private/PathTracing/PathTracingMaterialHitShader.usf
--- a/Engine/Shaders/Private/PathTracing/PathTracingMaterialHitShader.usf
+++ b/Engine/Shaders/Private/PathTracing/PathTracingMaterialHitShader.usf
@@ -286,7 +286,9 @@ RAY_TRACING_ENTRY_CLOSEST_HIT(PathTracingMaterialCHS,
 #elif MATERIALBLENDING_MODULATE
 	Payload.BSDFOpacity = 0.0;
 	Payload.TransparencyColor = GetMaterialEmissive(PixelMaterialInputs);
-	HLSL_STATIC_ASSERT(MATERIAL_SHADINGMODEL_UNLIT == 1, "Modulate blend mode requires unlit shading model");
+	// NOTE: in the case of substrate materials we could have DEFAULT_LIT here
+	HLSL_STATIC_ASSERT(MATERIAL_SHADINGMODEL_UNLIT == 1 ||
+	                   MATERIAL_SHADINGMODEL_DEFAULT_LIT == 1, "Modulate blend mode requires unlit shading model");
 	Payload.ShadingModelID = SHADINGMODELID_UNLIT;
 #elif MATERIALBLENDING_MASKED && MATERIAL_DITHER_OPACITY_MASK
 	// dithering emulates real transparency, so switch to translucent
