Path Tracer: Fix unlit materials not being able to emit more than ~64k radiance values.

This is problematic for cases when we want to encode the sun into a background texture, or when making a light source directly visible to camera rays (because the radiance data gets back into the payload).

For the unlit case, adopt an alternative encoding as the other fields are not used.

diff --git a/Engine/Shaders/Private/PathTracing/PathTracingCommon.ush b/Engine/Shaders/Private/PathTracing/PathTracingCommon.ush
--- a/Engine/Shaders/Private/PathTracing/PathTracingCommon.ush
+++ b/Engine/Shaders/Private/PathTracing/PathTracingCommon.ush
@@ -236,18 +236,30 @@ struct FPackedPathTracingPayload : FMinimalPayload
 	float3 UnpackWorldGeoNormal() 	 { return PayloadDecodeUnitVector(PackedData[1]); }
 	float3 UnpackWorldSmoothNormal() { return PayloadDecodeUnitVector(PackedData[2]); }
 	float3 UnpackWorldNormal()       { return PayloadDecodeUnitVector(PackedData[3]); }
-	float3 UnpackRadiance()          { return PayloadDecodeHDRColor(PackedData[4]); }
+	float3 UnpackRadiance()          { return (GetShadingModelID() == SHADINGMODELID_UNLIT) ? float3(asfloat(PackedData[4]), asfloat(PackedData[7]), asfloat(PackedData[8])) : PayloadDecodeHDRColor(PackedData[4]); }
 
 	void PackWorldNormal(float3 N)   { PackedData[3] = PayloadEncodeUnitVector(N); }
-	void PackRadiance(float3 R)      { PackedData[4] = PayloadEncodeHDRColor(R); }
+	void PackRadiance(float3 R)      {
+		if (GetShadingModelID() == SHADINGMODELID_UNLIT)
+		{
+			// encode with extra accuracy for UNLIT materials since these are also used for directly visible lights
+			PackedData[4] = asuint(R.x);
+			PackedData[7] = asuint(R.y);
+			PackedData[8] = asuint(R.z);
+		}
+		else
+		{
+			PackedData[4] = PayloadEncodeHDRColor(R);
+		}
+	}
 
 	
 #if PATHTRACING_SUBSTRATE_PAYLOAD
 
 	float3 UnpackDiffuseColor()       { return PayloadDecodeLDRColor(PackedData[5]); }
 	float3 UnpackSpecularColor()      { return PayloadDecodeLDRColor(PackedData[6]); }
-	float3 UnpackSpecularEdgeColor()  { return PayloadDecodeLDRColor(PackedData[7]); }
-	float4 UnpackRoughnessAniso()     { return PayloadDecodeRoughnessAniso(PackedData[8]); }
+	float3 UnpackSpecularEdgeColor()  { return (GetShadingModelID() == SHADINGMODELID_UNLIT) ? 0.0 : PayloadDecodeLDRColor(PackedData[7]); }
+	float4 UnpackRoughnessAniso()     { return (GetShadingModelID() == SHADINGMODELID_UNLIT) ? 0.0 : PayloadDecodeRoughnessAniso(PackedData[8]); }
 	float3 UnpackWorldTangent()       { return PayloadDecodeUnitVector(PackedData[9]); }
 	float3 UnpackTransparencyColor()  { return PayloadDecodeLDRColor(PackedData[10]); }
 	float3 UnpackMeanFreePath()       { return float3(f16tof32(PackedData[11]),
@@ -269,10 +281,10 @@ struct FPackedPathTracingPayload : FMinimalPayload
 #else
 	float3 UnpackBaseColor()         { return PayloadDecodeLDRColor(PackedData[5]); }
 	float3 UnpackTransparencyColor() { return PayloadDecodeLDRColor(PackedData[6]); }
-	float UnpackMetallic()           { return f16tof32(PackedData[7]); }
-	float UnpackSpecular()           { return f16tof32(PackedData[7] >> 16); }
-	float UnpackRoughness()          { return f16tof32(PackedData[8]); }
-	float UnpackAnisotropy()         { return f16tof32(PackedData[8] >> 16); }
+	float UnpackMetallic()           { return (GetShadingModelID() == SHADINGMODELID_UNLIT) ? 0.0 : f16tof32(PackedData[7]); }
+	float UnpackSpecular()           { return (GetShadingModelID() == SHADINGMODELID_UNLIT) ? 0.0 : f16tof32(PackedData[7] >> 16); }
+	float UnpackRoughness()          { return (GetShadingModelID() == SHADINGMODELID_UNLIT) ? 0.0 :  f16tof32(PackedData[8]); }
+	float UnpackAnisotropy()         { return (GetShadingModelID() == SHADINGMODELID_UNLIT) ? 0.0 : f16tof32(PackedData[8] >> 16); }
 	float3 UnpackWorldTangent()      { return PayloadDecodeUnitVector(PackedData[9]); }
 	float4 UnpackCustomData0()       { return float4(f16tof32(PackedData[10]),
 	                                                 f16tof32(PackedData[10] >> 16),
@@ -745,7 +757,7 @@ FPackedPathTracingPayload PackPathTracingPayload(FPathTracingPayload Input)
 	Output.PackedData[1] = PayloadEncodeUnitVector(Input.WorldGeoNormal);
 	Output.PackedData[2] = PayloadEncodeUnitVector(Input.WorldSmoothNormal);
 	Output.PackedData[3] = PayloadEncodeUnitVector(Input.WorldNormal);
-	Output.PackedData[4] = PayloadEncodeHDRColor(Input.Radiance);
+	// PackedData[4] is for Radiance (see PackRadiance call below)
 #if PATHTRACING_SUBSTRATE_PAYLOAD
 	Output.PackedData[5] = PayloadEncodeLDRColor(Input.DiffuseColor);
 	Output.PackedData[6] = PayloadEncodeLDRColor(Input.SpecularColor);
@@ -803,6 +815,8 @@ FPackedPathTracingPayload PackPathTracingPayload(FPathTracingPayload Input)
 	Output.PackedData[14] = f32tof16(Input.Ior);
 	// NOTE: 16 bits left in PackedData[14]
 #endif // PATHTRACING_SUBSTRATE_PAYLOAD
+	// handle this at the end since Radiance could be aliased with other fields
+	Output.PackRadiance(Input.Radiance);
 #endif // IS_PAYLOAD_ENABLED(RT_PAYLOAD_TYPE_GPULIGHTMASS)
 	return Output;
 }
