Fix coat energy conservation in ReflectionEnvironment shader

This manifested as missing reflections when a clear coat material at full intensity was used in the material editor (which uses the medium scalability mode).

diff --git a/Engine/Shaders/Private/ReflectionEnvironmentPixelShader.usf b/Engine/Shaders/Private/ReflectionEnvironmentPixelShader.usf
--- a/Engine/Shaders/Private/ReflectionEnvironmentPixelShader.usf
+++ b/Engine/Shaders/Private/ReflectionEnvironmentPixelShader.usf
@@ -212,7 +212,9 @@ float3 ReflectionEnvironment(FGBufferData GBuffer, float AmbientOcclusion, float
 	uint NumCulledReflectionCaptures = 0;
 #endif
 
+#if USE_ENERGY_CONSERVATION
 	const FBxDFEnergyTerms EnergyTerms = ComputeGGXSpecEnergyTerms(GBuffer.Roughness, NoV, GBuffer.SpecularColor);
+#endif
 
 	//Top of regular reflection or bottom layer of clear coat.
 	Color.rgb += View.PreExposure * GatherRadiance(Color.a, TranslatedWorldPosition, R, GBuffer.Roughness, BentNormal, IndirectIrradiance, GBuffer.ShadingModelID, NumCulledReflectionCaptures, DataStartIndex);
@@ -233,7 +235,7 @@ float3 ReflectionEnvironment(FGBufferData GBuffer, float AmbientOcclusion, float
 
 		// TODO EnvBRDF should have a mask param
 		#if USE_ENERGY_CONSERVATION
-		Color.rgb *= EnergyTerms.E * (1 - ClearCoat);
+		Color.rgb *= EnergyTerms.E;
 		#else
 		// Hack: Ensures when clear coat is >0, grazing angle does not get too much energy, 
 		//       but preserve response at normal incidence 
