Workaround for internal shader compiler error when compiling ray tracing shaders for some platforms

diff --git a/Engine/Shaders/Private/Common.ush b/Engine/Shaders/Private/Common.ush
--- a/Engine/Shaders/Private/Common.ush
+++ b/Engine/Shaders/Private/Common.ush
@@ -238,7 +238,11 @@ float VectorSum(float4 V) { return V.x + V.y + V.z + V.w; }
 // Using SV_ClipDistance has overhead (15% slower base pass in triangle bound test scene on PS4) so projects have to opt-in
 #define USE_GLOBAL_CLIP_PLANE (PROJECT_ALLOW_GLOBAL_CLIP_PLANE && !MATERIAL_DOMAIN_POSTPROCESS && !MATERIAL_DOMAIN_UI)
 
-#if RAYHITGROUPSHADER || RAYMISSHADER || RAYCALLABLESHADER
+#ifndef RAYTRACINGSHADER
+	#define RAYTRACINGSHADER (RAYHITGROUPSHADER || RAYMISSSHADER || RAYCALLABLESHADER)
+#endif
+
+#if RAYTRACINGSHADER
 
 // These built-ins are not available in ray tracing
 // Define dummy versions so that ray-tracing materials will at least compile
@@ -1249,7 +1253,12 @@ float ConvertFromDeviceZ(float DeviceZ)
 
 float GetScreenPositionDepth(float4 ScreenPosition)
 {
+#if RAYTRACINGSHADER
+	// Workaround for UE-212146: Accessing global ResolvedView in this case triggers internal compiler error on some platforms.
+	return select(IsOrthoProjection(ResolveView()), ConvertFromDeviceZ(ScreenPosition.z), ScreenPosition.w);
+#else
 	return select(IsOrthoProjection(ResolvedView), ConvertFromDeviceZ(ScreenPosition.z), ScreenPosition.w);
+#endif
 }
 
 float3 GetCameraVector(ViewState InView, float3 Position, float3 CameraPosition)
@@ -1931,7 +1940,7 @@ bool GetNaniteReplaceState()
 
 bool GetRayTracingQualitySwitch()
 {
-#if (RAYHITGROUPSHADER || RAYMISSSHADER || RAYCALLABLESHADER || LUMEN_CARD_CAPTURE) && !PATH_TRACING
+#if (RAYTRACINGSHADER || LUMEN_CARD_CAPTURE) && !PATH_TRACING
 	return true;
 #else
 	return false;
@@ -1950,7 +1959,7 @@ bool GetPathTracingIsIndirectVolume() { return false; }
 
 bool GetLightmassReplaceState()
 {
-#if (RAYHITGROUPSHADER || RAYMISSSHADER || RAYCALLABLESHADER) && PATH_TRACING && SIMPLIFIED_MATERIAL_SHADER
+#if RAYTRACINGSHADER && PATH_TRACING && SIMPLIFIED_MATERIAL_SHADER
 	return true;
 #else
 	return false;
