Use fma to implement invariant ops on specific platforms where the 'precise' keyword does not work as intended.

diff --git a/Engine/Shaders/Private/DoubleFloat.ush b/Engine/Shaders/Private/DoubleFloat.ush
--- a/Engine/Shaders/Private/DoubleFloat.ush
+++ b/Engine/Shaders/Private/DoubleFloat.ush
@@ -541,8 +541,8 @@ float4x4 DFFastToTranslatedWorld(FDFMatrix LocalToWorld, FDFVector3 PreViewTrans
 
 FDFVector3 DFMultiplyLHSAndFastTwoSum(float3 Lhs, float Factor, float3 Rhs)
 {
-	const float3 S = INVARIANT(FMA(Lhs, Factor, Rhs));
-	const float3 E = INVARIANT(Rhs - FMA(-Lhs, Factor, S));
+	const float3 S = INVARIANT_FMA(Lhs, Factor, Rhs);
+	const float3 E = INVARIANT_SUB(Rhs, INVARIANT_FMA(-Lhs, Factor, S));
 	return MakeDFVector3(S, E);
 }
 
diff --git a/Engine/Shaders/Private/LargeWorldCoordinates.ush b/Engine/Shaders/Private/LargeWorldCoordinates.ush
--- a/Engine/Shaders/Private/LargeWorldCoordinates.ush
+++ b/Engine/Shaders/Private/LargeWorldCoordinates.ush
@@ -520,7 +520,7 @@ float4x4 LWCHackToFloat(FLWCInverseMatrix v) { return (float4x4)0.0f; }
 		/* divergent cases are extremely rare and result simply in rollover from one tile edge to the edge of the next tile */	\
 		/* Values greater than 2^(24+floor(log2(TILE_SIZE-1))) are not converted correctly, because In.Low is assumed to be */	\
 		/* less than half a tile. These values by definition exceed WORLD_MAX, so this is not an issue. */						\
-		FFloatType Tile = INVARIANT(round(In.High / UE_LWC_RENDER_TILE_SIZE));													\
+		FFloatType Tile = round(In.High / UE_LWC_RENDER_TILE_SIZE);																\
 		FFloatType Offset = INVARIANT_ADD(INVARIANT_FMA(Tile, -UE_LWC_RENDER_TILE_SIZE, In.High), In.Low);						\
 																																\
 		return LWCConstructor(Tile, Offset);																					\
