PostProcessing: Fixed alpha channel update with after motion blur translucency composite.

diff --git a/Engine/Shaders/Private/MotionBlur/MotionBlurApply.usf b/Engine/Shaders/Private/MotionBlur/MotionBlurApply.usf
--- a/Engine/Shaders/Private/MotionBlur/MotionBlurApply.usf
+++ b/Engine/Shaders/Private/MotionBlur/MotionBlurApply.usf
@@ -882,6 +882,13 @@ void MainCS(
 			OutputColor1.rgb = OutputColor1.rgb * PostMotionBlurTranslucency1.a + PostMotionBlurTranslucency1.rgb;
 			OutputColor2.rgb = OutputColor2.rgb * PostMotionBlurTranslucency2.a + PostMotionBlurTranslucency2.rgb;
 			OutputColor3.rgb = OutputColor3.rgb * PostMotionBlurTranslucency3.a + PostMotionBlurTranslucency3.rgb;
+
+			#if POST_PROCESS_ALPHA
+				OutputColor0.a = OutputColor0.a * PostMotionBlurTranslucency0.a;
+				OutputColor1.a = OutputColor1.a * PostMotionBlurTranslucency1.a;
+				OutputColor2.a = OutputColor2.a * PostMotionBlurTranslucency2.a;
+				OutputColor3.a = OutputColor3.a * PostMotionBlurTranslucency3.a;
+			#endif
 		}
 		#endif
 		
@@ -1010,6 +1017,10 @@ void MainCS(
 		#if CONFIG_POST_MOTIONBLUR_TRANSLUCENCY
 		{
 			OutputColor.rgb = OutputColor.rgb * PostMotionBlurTranslucency.a + PostMotionBlurTranslucency.rgb;
+			
+			#if POST_PROCESS_ALPHA
+				OutputColor.a = OutputColor.a * PostMotionBlurTranslucency.a;
+			#endif
 		}
 		#endif
 
