RealTimeCapture - height fog code clean up.

diff --git a/Engine/Shaders/Private/ReflectionEnvironmentShaders.usf b/Engine/Shaders/Private/ReflectionEnvironmentShaders.usf
--- a/Engine/Shaders/Private/ReflectionEnvironmentShaders.usf
+++ b/Engine/Shaders/Private/ReflectionEnvironmentShaders.usf
@@ -886,8 +886,6 @@ void RenderRealTimeReflectionHeightFogPS(
 {
 	ResolvedView = ResolveView();
 
-	float4 WorldVector = mul(float4((SVPos.xy / 128.0f) * float2(2.0, -2.0) + float2(1.0, -1.0), 0.0, 0), ResolvedView.ScreenToTranslatedWorld);
-
 #if PERMUTATION_DEPTHTEXTURE
 	float DeviceZ = DepthTexture.Load(int3(SVPos.xy, 0)).x;
 	const float HeightFogDistance = ConvertFromDeviceZ(DeviceZ);
