Path Tracer: Revisit visibility of sky meshes and improve alpha precision

 * Hide the sky mesh when using path tracing
 * Fix alpha not accumulating to 1.0 due to clamping the alpha per sample instead of after averaging.
 * Remove code forcing alpha=1 for rays that do not hit the sky. This behavior is not desireable when reference atmosphere mode is used
 * Fix sunlight visiblity when atmosphere is marked as a holdout (as this is the only case where directional lights are made visible)
 * Fix conditions in the main path tracing shader that require ray marching

diff --git a/Engine/Shaders/Private/PathTracing/PathTracingCompositingPixelShader.usf b/Engine/Shaders/Private/PathTracing/PathTracingCompositingPixelShader.usf
--- a/Engine/Shaders/Private/PathTracing/PathTracingCompositingPixelShader.usf
+++ b/Engine/Shaders/Private/PathTracing/PathTracingCompositingPixelShader.usf
@@ -28,7 +28,9 @@ void CompositeMain(
 	float4 Radiance = RadianceTexture.Load(TexCoord);
 
 	// NOTE: UE has the convention of using 1.0 - Alpha
-	OutColor = float4(Radiance.rgb * View.PreExposure, 1.0 - Radiance.a);
+	// We must saturate the result because the path tracer's alpha is stochastic in the presence of volumetrics
+	// It is important that the clamp happens _after_ averaging, otherwise this can bias the estimate
+	OutColor = float4(Radiance.rgb * View.PreExposure, saturate(1.0 - Radiance.a));
 
 	if (AdaptiveSamplingVisualize == 1)
 	{
@@ -77,28 +79,6 @@ void CompositeMain(
 		float2 Center = float2(lerp(RectMin.x, RectMax.x, 0.5), RectMax.y) * View.BufferSizeAndInvSize.zz - float2(0, Size * 0.75);
 		P -= Center;
 		float Fraction = saturate(float(Iteration + 1) / float(MaxSamples));
-#if 0
-		// measure progress like a clock
-		float R2 = dot(P, P);
-		float ROuter2 = Square(0.50 * Size);
-		float RInner2 = Square(0.35 * Size);
-
-		float R2f = fwidth(R2);
-
-		// point is on the progress ring
-		if (R2 < ROuter2 + R2f && R2 > RInner2 - R2f)
-		{
-			float Angle = (atan2(-P.x, P.y) + PI) / (2 * PI);
-			// anti-alias progress bar edge
-			float Blend = smoothstep(Fraction - fwidth(Angle), Fraction + fwidth(Angle), Angle);
-			// Grey/Red blend
-			float4 ProgressBarColor = lerp(float4(1.0, 0.1, 0.1, 1), float4(0.5, 0.5, 0.5, 1), Blend);
-			// anti-alias ring edges
-			float BlendOuter = 1 - smoothstep(ROuter2 - R2f, ROuter2 + R2f, R2);
-			float BlendInner =     smoothstep(RInner2 - R2f, RInner2 + R2f, R2);
-			OutColor = lerp(OutColor, ProgressBarColor, 0.75 * BlendInner * BlendOuter);
-		}
-#else
 		// simple linear progress bar
 		float2 ProgressSize = float2(2.0 * Size, 0.15 * Size);
 		if (all(abs(P) < ProgressSize))
@@ -107,10 +87,9 @@ void CompositeMain(
 			P /= ProgressSize * 2;
 			float Blend = smoothstep(Fraction - fwidth(P.x), Fraction + fwidth(P.x), P.x);
 			// Grey/Red blend
-			float4 ProgressBarColor = lerp(float4(1.0, 0.1, 0.1, 1), float4(0.5, 0.5, 0.5, 1), Blend);
-			OutColor = lerp(OutColor, ProgressBarColor, 0.75);
+			float3 ProgressBarColor = lerp(float3(1.0, 0.1, 0.1), float3(0.5, 0.5, 0.5), Blend);
+			OutColor.xyz = lerp(OutColor.xyz, ProgressBarColor, 0.75);
 		}
-#endif
 	}
 
 	OutDepth = NormalDepthTexture.Load(TexCoord).w;
diff --git a/Engine/Shaders/Private/PathTracing/PathTracingCore.ush b/Engine/Shaders/Private/PathTracing/PathTracingCore.ush
--- a/Engine/Shaders/Private/PathTracing/PathTracingCore.ush
+++ b/Engine/Shaders/Private/PathTracing/PathTracingCore.ush
@@ -263,7 +263,7 @@ struct FPathState
 		float4 OldAlbedo = Iteration > 0 ? AlbedoTexture[TextureIndex]   : 0;
 		float4 OldNormal = Iteration > 0 ? NormalTexture[TextureIndex]   : 0;
 
-		const float4 PixelValue = float4(Radiance, saturate(Alpha));
+		const float4 PixelValue = float4(Radiance, Alpha);
 
 		const float3 Albedo = GetAlbedo();
 		const float3 Normal = GetNormal();
@@ -554,83 +554,99 @@ FPackedPathTracingPayload TraceTransparentRay(inout FPathState PathState, int Bo
 				}
 			}
 
-			// NOTE: when using analytic raymarching, we can skip all raymarching on the last bounce because we know we don't need a scattering point
 			{
 				// remember the start of our interval before ray marching messes with it
 				float StartVolumeTMin = Interval.VolumeTMin;
 				// Ray marching loop
 				float3 PathThroughput = PathState.PathThroughput;
-				float3 VolumeRadiance = 0;
-				// Limit number of steps to prevent timeouts
-				// TODO: This biases the result! Is there a better way to limit the number of steps?
-				for (int Step = 0; Step < MaxRaymarchSteps; Step++)
+
+				// In order to use analytic transmittance, the user must have requested it _and_ we must be on an interval with no clouds or heterogeneous volumes
+				const bool bUseAnalyticTransmittance =
+					(VolumeFlags & PATH_TRACER_VOLUME_USE_ANALYTIC_TRANSMITTANCE) != 0 &&
+					(Interval.VolumeMask & (PATH_TRACER_VOLUME_ENABLE_CLOUDS | PATH_TRACER_VOLUME_ENABLE_HETEROGENEOUS_VOLUMES)) == 0;
+				// are there any holdouts on the current segment?
+				const bool bHaveHoldouts =
+					((Interval.VolumeMask << PATH_TRACER_MAX_VOLUMES) & VolumeFlags & (
+						PATH_TRACER_VOLUME_HOLDOUT_ATMOSPHERE |
+						PATH_TRACER_VOLUME_HOLDOUT_CLOUDS |
+						PATH_TRACER_VOLUME_HOLDOUT_FOG |
+						PATH_TRACER_VOLUME_HOLDOUT_HETEROGENEOUS_VOLUMES)) != 0;
+
+				// reasons to enter the ray marching loop:
+				//  - we can't use analytic transmittance (need ray marching for throughput or emission)
+				//  - we are a camera ray and some volumes are flagged as holdouts
+				//  - this is not the last bounce, so we need to select a point for scattering
+				if (!bUseAnalyticTransmittance || (bIsCameraRay && bHaveHoldouts) || !bLastBounce)
 				{
-					FVolumeTrackingResult TrackingResult = VolumeSampleDistance(PathThroughput, PathState.Ray.Origin, PathState.Ray.Direction, Interval, bIsCameraRay, PathState.RandSequence);
-					if (TrackingResult.Distance < 0)
+					float3 VolumeRadiance = 0;
+					// Limit number of steps to prevent timeouts
+					// TODO: This biases the result! Is there a better way to limit the number of steps?
+					for (int Step = 0; Step < MaxRaymarchSteps; Step++)
 					{
-						break;
-					}
+						FVolumeTrackingResult TrackingResult = VolumeSampleDistance(PathThroughput, PathState.Ray.Origin, PathState.Ray.Direction, Interval, bIsCameraRay, PathState.RandSequence);
+						if (TrackingResult.Distance < 0)
+						{
+							break;
+						}
 
-					if (TrackingResult.bIsCollision)
-					{
-						Interval.VolumeTMin = TrackingResult.Distance;
-						PathThroughput = TrackingResult.Throughput;
+						if (TrackingResult.bIsCollision)
+						{
+							Interval.VolumeTMin = TrackingResult.Distance;
+							PathThroughput = TrackingResult.Throughput;
 
-						// find out how much volume exists at the current point
-						float3 Ro = PathState.Ray.Origin;
-						float3 Rd = PathState.Ray.Direction;
-						float3 TranslatedWorldPos = Ro + Interval.VolumeTMin * Rd;
-						FVolumeShadedResult Result = VolumeGetDensity(TranslatedWorldPos, Interval.VolumeMask, bIsCameraRay);
+							// find out how much volume exists at the current point
+							float3 Ro = PathState.Ray.Origin;
+							float3 Rd = PathState.Ray.Direction;
+							float3 TranslatedWorldPos = Ro + Interval.VolumeTMin * Rd;
+							FVolumeShadedResult Result = VolumeGetDensity(TranslatedWorldPos, Interval.VolumeMask, bIsCameraRay);
 
-						// clamp to make sure we never exceed the majorant (should not be the case, but need to avoid any possible numerical issues)
-						float3 SigmaT = min(Result.SigmaT, TrackingResult.SigmaBar);
-						float3 SigmaN = TrackingResult.SigmaBar - SigmaT;
+							// clamp to make sure we never exceed the majorant (should not be the case, but need to avoid any possible numerical issues)
+							float3 SigmaT = min(Result.SigmaT, TrackingResult.SigmaBar);
+							float3 SigmaN = TrackingResult.SigmaBar - SigmaT;
 
-						if (bIsCameraRay)
-						{
-							PathState.Alpha += Luminance(PathThroughput * (Result.SigmaT - Result.SigmaH));
-						}
+							if (bIsCameraRay)
+							{
+								float3 SigmaH = min(Result.SigmaH, SigmaT);
+								PathState.Alpha += Luminance(PathThroughput * (SigmaT - SigmaH));
+							}
 
-						// if this is not the last bounce, consider the possibility of scattering as giving up medium hits
-						if (!bLastBounce)
-						{
-							float3 SigmaS = min(Result.SigmaSRayleigh + Result.SigmaSHG + Result.SigmaSDualHG, SigmaT);
-							// accumulate a signal for the denoiser
-							float3 Albedo = PathState.GetAlbedo();
-							AccumulateAlbedo(SigmaS, PathThroughput, PathState.GetPathRoughness(), Albedo);
-							PathState.SetAlbedoNormal(Albedo, PathState.GetNormal());
-							float3 Contrib = PathThroughput * SigmaS;
-							float SelectionWeight = max3(Contrib.x, Contrib.y, Contrib.z);
-							if (HitSample.Accept(SelectionWeight))
+							// if this is not the last bounce, consider the possibility of scattering to give us medium hits
+							if (!bLastBounce)
 							{
-								// stash this hit for next time
-								Payload = PackPathTracingPayload(CreateMediumHitPayload(Interval.VolumeTMin, TranslatedWorldPos, Result));
-								PayloadThroughput = Contrib / SelectionWeight;
+								float3 SigmaS = min(Result.SigmaSRayleigh + Result.SigmaSHG + Result.SigmaSDualHG, SigmaT);
+								// accumulate a signal for the denoiser
+								float3 Albedo = PathState.GetAlbedo();
+								AccumulateAlbedo(SigmaS, PathThroughput, PathState.GetPathRoughness(), Albedo);
+								PathState.SetAlbedoNormal(Albedo, PathState.GetNormal());
+								float3 Contrib = PathThroughput * SigmaS;
+								float SelectionWeight = max3(Contrib.x, Contrib.y, Contrib.z);
+								if (HitSample.Accept(SelectionWeight))
+								{
+									// stash this hit for next time
+									Payload = PackPathTracingPayload(CreateMediumHitPayload(Interval.VolumeTMin, TranslatedWorldPos, Result));
+									PayloadThroughput = Contrib / SelectionWeight;
+								}
 							}
-						}
 
-						VolumeRadiance += Result.Emission * PathThroughput;
-						// keep tracing through the volume
-						PathThroughput *= SigmaN;
-					}
-					else
-					{
-						// update the path throughput, knowing that we escaped the medium
-						PathThroughput = TrackingResult.Throughput;
-						// exit the ray marching loop
-						break;
+							VolumeRadiance += Result.Emission * PathThroughput;
+							// keep tracing through the volume
+							PathThroughput *= SigmaN;
+						}
+						else
+						{
+							// update the path throughput, knowing that we escaped the medium
+							PathThroughput = TrackingResult.Throughput;
+							// exit the ray marching loop
+							break;
+						}
 					}
+					// accumulate all at once so clamping applies to entire volume contribution
+					AccumulateRadiance(PathState.Radiance, VolumeRadiance, bIsCameraRay);
 				}
-				// accumulate all at once so clamping applies to entire volume contribution
-				AccumulateRadiance(PathState.Radiance, VolumeRadiance, bIsCameraRay);
 
 				// restore interval start now that we are finished with ray marching
 				Interval.VolumeTMin = StartVolumeTMin;
 
-				// In order to use analytic transmittance, the user must have requested it _and_ we must be on an interval with no clouds or heterogeneous volumes
-				const bool bUseAnalyticTransmittance =
-					(VolumeFlags & PATH_TRACER_VOLUME_USE_ANALYTIC_TRANSMITTANCE) != 0 &&
-					(Interval.VolumeMask & (PATH_TRACER_VOLUME_ENABLE_CLOUDS | PATH_TRACER_VOLUME_ENABLE_HETEROGENEOUS_VOLUMES)) == 0;
 				if (bUseAnalyticTransmittance)
 				{
 					// analytically handle the transmittance to the next surface (which will be higher quality than the transmittance implicitly computed above and prevent noise with emissive surfaces in volumes)
@@ -651,13 +667,6 @@ FPackedPathTracingPayload TraceTransparentRay(inout FPathState PathState, int Bo
 
 		if (PackedPayload.IsMiss())
 		{
-			// This is for parity with the sky dome mesh behavior.
-			// Output fully opaque alpha if looking through outer space, if atmosphere is not in holdout.
-			if (bIsCameraRay && (VolumeFlags & PATH_TRACER_VOLUME_HOLDOUT_ATMOSPHERE) == 0)
-			{
-				PathState.Alpha += Luminance(PathState.PathThroughput);
-			}
-
 			// Ray didn't hit any real geometry
 			break;
 		}
diff --git a/Engine/Source/Runtime/Renderer/Private/PathTracing.cpp b/Engine/Source/Runtime/Renderer/Private/PathTracing.cpp
--- a/Engine/Source/Runtime/Renderer/Private/PathTracing.cpp
+++ b/Engine/Source/Runtime/Renderer/Private/PathTracing.cpp
@@ -2140,7 +2140,7 @@ RENDERER_API void PrepareSkyTexture_Internal(
 {
 	FRDGTextureDesc SkylightTextureDesc = FRDGTextureDesc::Create2D(
 		FIntPoint(Size, Size),
-		PF_A32B32G32R32F, // half precision might be ok?
+		PF_A32B32G32R32F, // Must use float as CubeMap * Color could have float range (could use half if we didn't include Color in the map)
 		FClearValueBinding::None,
 		TexCreate_ShaderResource | TexCreate_UAV);
 
@@ -2148,7 +2148,7 @@ RENDERER_API void PrepareSkyTexture_Internal(
 
 	FRDGTextureDesc SkylightPdfDesc = FRDGTextureDesc::Create2D(
 		FIntPoint(Size, Size),
-		PF_R32_FLOAT, // half precision might be ok?
+		PF_R32_FLOAT, // Must use float as CubeMap * Color could have float range (could use half if we didn't include Color in the map)
 		FClearValueBinding::None,
 		TexCreate_ShaderResource | TexCreate_UAV,
 		FMath::CeilLogTwo(Size) + 1);
@@ -2545,9 +2545,9 @@ void SetLightParameters(FRDGBuilder& GraphBuilder, FPathTracingRG::FParameters*
 		}
 	}
 
-	if (bUseAtmosphere)
+	if (bUseAtmosphere && (View.SkyAtmosphereUniformShaderParameters == nullptr ||  !IsSkyAtmosphereHoldout(View.CachedViewUniformShaderParameters->EnvironmentComponentsFlags)))
 	{
-		// show directional lights when atmosphere is enabled
+		// show directional lights when atmosphere is enabled and not marked as holdout
 		// NOTE: there cannot be any skydome in this case
 		PassParameters->SceneVisibleLightCount = NumLights;
 	}
diff --git a/Engine/Source/Runtime/Renderer/Private/RayTracing/RayTracing.cpp b/Engine/Source/Runtime/Renderer/Private/RayTracing/RayTracing.cpp
--- a/Engine/Source/Runtime/Renderer/Private/RayTracing/RayTracing.cpp
+++ b/Engine/Source/Runtime/Renderer/Private/RayTracing/RayTracing.cpp
@@ -1211,7 +1211,7 @@ namespace RayTracing
 						}
 
 						if ((GRayTracingExcludeTranslucent && RelevantPrimitive.bAllSegmentsTranslucent)
-							|| (GRayTracingExcludeSky && RelevantPrimitive.bIsSky && !bIsPathTracing))
+							|| (GRayTracingExcludeSky && RelevantPrimitive.bIsSky))
 						{
 							continue;
 						}
