Implemented support for analytic derivatives in Volumetric Cloud shaders. This is only compiled if the material signals that it uses explicit derivatives (DDX/DDY material graph nodes). Preflighted in 66202d0a4075e2bc66ae606f but SubmitTool refuses to accept the PF due to an unrelated error in the PF.

#rnx

diff --git a/Engine/Shaders/Private/Definitions.usf b/Engine/Shaders/Private/Definitions.usf
--- a/Engine/Shaders/Private/Definitions.usf
+++ b/Engine/Shaders/Private/Definitions.usf
@@ -147,8 +147,20 @@
 #define IS_NANITE_PASS (IS_NANITE_SHADING_PASS || IS_NANITE_RASTER_PASS)
 #endif
 
+// This is set by the material translator and signifies whether the material is using the DDX or DDY material graph nodes.
+// In that case the material shader might want to implement this functionality with analytic derivatives if hardware derivatives are not supported.
+#ifndef USES_EXPLICIT_DERIVATIVES
+#define USES_EXPLICIT_DERIVATIVES 0
+#endif
+
+// Whether USE_ANALYTIC_DERIVATIVES may be enabled when USES_EXPLICIT_DERIVATIVES is true
+#ifndef ALLOW_ANALYTIC_DERIVATIVES
+#define ALLOW_ANALYTIC_DERIVATIVES 0
+#endif
+
+// Whether to use analytic derivatives
 #ifndef USE_ANALYTIC_DERIVATIVES
-#define USE_ANALYTIC_DERIVATIVES 0
+#define USE_ANALYTIC_DERIVATIVES						(ALLOW_ANALYTIC_DERIVATIVES && USES_EXPLICIT_DERIVATIVES)
 #endif
 
 #ifndef GBUFFER_HAS_VELOCITY
diff --git a/Engine/Shaders/Private/VolumetricCloud.usf b/Engine/Shaders/Private/VolumetricCloud.usf
--- a/Engine/Shaders/Private/VolumetricCloud.usf
+++ b/Engine/Shaders/Private/VolumetricCloud.usf
@@ -96,6 +96,19 @@
 
 #if defined(SHADER_SHADOW_PS) || SHADER_RENDERVIEW_PS || SHADER_RENDERVIEW_CS || SHADER_EMPTY_SPACE_SKIPPING_CS
 
+void VolumetricCloudCalcPixelMaterialInputs(in out FMaterialPixelParameters Parameters, in out FPixelMaterialInputs PixelMaterialInputs)
+{
+	// This same logic is also used in CalcMaterialParameters(), which is also called in these shaders, so we keep it here for consistency.
+#if USE_ANALYTIC_DERIVATIVES
+	if (!TEXTURE_SAMPLE_DEBUG || View.GeneralPurposeTweak >= 1.0f)
+		CalcPixelMaterialInputsAnalyticDerivatives(Parameters, PixelMaterialInputs);
+	else
+#endif
+	{
+		CalcPixelMaterialInputs(Parameters, PixelMaterialInputs);
+	}
+}
+
 float3 SampleExtinctionCoefficients(in FPixelMaterialInputs PixelMaterialInputs)
 {
 	float3 Extinction = 0.0f;
@@ -461,11 +474,14 @@ void MainCommon(in FMaterialPixelParameters MaterialParameters, in float4 SvPosi
 	OutColor0					= float4(0.0f, 0.0f, 0.0f, 1.0f);
 	OutDepth					= MaxHalfFloat;
 
+	// CalcMaterialParameters() evaluates the material, but since the depth associated with SvPosition at this point doesn't correspond to any particular point along the ray,
+	// we leave WorldPosition_DDX/Y as 0, which forces mip 0 when compiled with USE_ANALYTIC_DERIVATIVES.
 	FPixelMaterialInputs PixelMaterialInputs	= (FPixelMaterialInputs)0;
 	CalcMaterialParameters(MaterialParameters, PixelMaterialInputs, SvPosition, true);
 
 	float3 RayWorldOriginKm				= DFHackToFloat(ResolvedView.WorldCameraOrigin) * CENTIMETER_TO_KILOMETER;
-	float3 RayTranslatedWorldOriginKm	= View.TranslatedWorldCameraOrigin * CENTIMETER_TO_KILOMETER;
+	float3 RayTranslatedWorldOrigin		= View.TranslatedWorldCameraOrigin;
+	float3 RayTranslatedWorldOriginKm	= RayTranslatedWorldOrigin * CENTIMETER_TO_KILOMETER;
 	float3 Raydir						= -MaterialParameters.CameraVector;
 	float TMin							= -999999999.0f;
 	float TMax							= -999999999.0f;
@@ -483,6 +499,18 @@ void MainCommon(in FMaterialPixelParameters MaterialParameters, in float4 SvPosi
 	float ClostestTAPWeightsSum = 0.0f;
 #endif // CLOUD_MIN_AND_MAX_DEPTH
 
+#if USE_ANALYTIC_DERIVATIVES
+	const float3 TranslatedWorldPositionOffsetX = SvPositionToResolvedTranslatedWorld(SvPosition + float4(1.0f, 0.0f, 0.0f, 0.0f));
+	const float3 TranslatedWorldPositionOffsetY = SvPositionToResolvedTranslatedWorld(SvPosition + float4(0.0f, 1.0f, 0.0f, 0.0f));
+	const float3 RayDirOffsetX = GetCameraVector(ResolvedView, TranslatedWorldPositionOffsetX, ResolvedView.TranslatedWorldCameraOrigin);
+	const float3 RayDirOffsetY = GetCameraVector(ResolvedView, TranslatedWorldPositionOffsetY, ResolvedView.TranslatedWorldCameraOrigin);
+
+	// TODO: The current cloud tracing code assumes the ray origin to be at the camera origin, which breaks on ortho projections.
+	// We still define these variables to reduce the risk of bugs in the future when the code is fixed to properly support ortho projections as well.
+	const float3 RayTranslatedWorldOriginOffsetX = RayTranslatedWorldOrigin;
+	const float3 RayTranslatedWorldOriginOffsetY = RayTranslatedWorldOrigin;
+#endif // USE_ANALYTIC_DERIVATIVES
+
 	uint4 TracingCoordToZbufferCoordScaleBias = RenderVolumetricCloudParameters.TracingCoordToZbufferCoordScaleBias;
 	uint2 SceneDepthTextureCoord = uint2(SvPosition.xy - 0.5) * TracingCoordToZbufferCoordScaleBias.xy + TracingCoordToZbufferCoordScaleBias.zw;
 	SceneDepthTextureCoord = clamp(SceneDepthTextureCoord, RenderVolumetricCloudParameters.SceneDepthTextureMinMaxCoord.xy, RenderVolumetricCloudParameters.SceneDepthTextureMinMaxCoord.zw);
@@ -788,16 +816,27 @@ void MainCommon(in FMaterialPixelParameters MaterialParameters, in float4 SvPosi
 		for (uint i = StartStepCount; i < IStepCount; ++i)
 		{
 			int ms;
-			float3	  SampleTranslatedWorldPosition = RayTranslatedWorldOriginKm + t * Raydir;
+			float3 SampleTranslatedWorldPosition = RayTranslatedWorldOrigin + t * Raydir;
 #if CLOUD_MIN_AND_MAX_DEPTH
 			const bool SampleContributeToClosest = t <= TClosestDepthBufferCm;
 #endif // CLOUD_MIN_AND_MAX
 			 
+#if USE_ANALYTIC_DERIVATIVES
+			float3 SampleWorldPositionDDX = (RayTranslatedWorldOriginOffsetX + t * RayDirOffsetX) - SampleTranslatedWorldPosition;
+			float3 SampleWorldPositionDDY = (RayTranslatedWorldOriginOffsetY + t * RayDirOffsetY) - SampleTranslatedWorldPosition;
+#else
+			float3 SampleWorldPositionDDX = 0.0f;
+			float3 SampleWorldPositionDDY = 0.0f;
+#endif // USE_ANALYTIC_DERIVATIVES
+
+			// Using the same derivatives for all shadow samples is an approximation only and could possibly be improved in the future
+			const float3 ShadowSampleWorldPositionDDX = SampleWorldPositionDDX;
+			const float3 ShadowSampleWorldPositionDDY = SampleWorldPositionDDY;
 
 			//////////////////////////////
 			// Update FMaterialPixelParameters according to current sample.
 			//////////////////////////////
-			UpdateMaterialCloudParam(MaterialParameters, SampleTranslatedWorldPosition,
+			UpdateMaterialCloudParam(MaterialParameters, SampleTranslatedWorldPosition, SampleWorldPositionDDX, SampleWorldPositionDDY,
 				ResolvedView, CloudLayerParams, TRACING_SHADOW_DISTANCE_OFF, RenderVolumetricCloudParameters.EmptySpaceSkippingSliceDepth);
 
 #if MATERIAL_VOLUMETRIC_ADVANCED_CONSERVATIVE_DENSITY && !CLOUD_DEBUG_VIEW_MODE
@@ -843,7 +882,7 @@ void MainCommon(in FMaterialPixelParameters MaterialParameters, in float4 SvPosi
 			//////////////////////////////
 			// Sample the participating media material at the sample position.
 			//////////////////////////////
-			CalcPixelMaterialInputs(MaterialParameters, PixelMaterialInputs);
+			VolumetricCloudCalcPixelMaterialInputs(MaterialParameters, PixelMaterialInputs);
 			ConvertCloudPixelMaterialInputsToWorkingColorSpace(PixelMaterialInputs);
 			const MATVEC ExtinctionCoefficients = SampleExtinctionCoefficients(PixelMaterialInputs);
 			const MATVEC EmissiveLuminance = USES_EMISSIVE_COLOR ? SampleEmissive(PixelMaterialInputs) : 0.0f;
@@ -930,7 +969,7 @@ void MainCommon(in FMaterialPixelParameters MaterialParameters, in float4 SvPosi
 						float CurrentNormT = ShadowT * ShadowT; // CurrentNormT is the end of the considered segment to integrate, PreviousNormT is its beginning.
 						const float DetlaNormT = CurrentNormT - PreviousNormT;
 						const float ShadowSampleDistance = ShadowLengthTest * (CurrentNormT - 0.5 * DetlaNormT);
-						UpdateMaterialCloudParam(MaterialParameters, SampleTranslatedWorldPosition + GroundDirection * ShadowSampleDistance,
+						UpdateMaterialCloudParam(MaterialParameters, SampleTranslatedWorldPosition + GroundDirection * ShadowSampleDistance, ShadowSampleWorldPositionDDX, ShadowSampleWorldPositionDDY,
 							ResolvedView, CloudLayerParams, ShadowSampleDistance, SPACE_SKIPPING_SLICE_DEPTH_OFF);
 						PreviousNormT = CurrentNormT;
 #if MATERIAL_VOLUMETRIC_ADVANCED_CONSERVATIVE_DENSITY
@@ -939,7 +978,7 @@ void MainCommon(in FMaterialPixelParameters MaterialParameters, in float4 SvPosi
 							continue; // Conservative density is 0 so skip and go to the next sample
 						}
 #endif
-						CalcPixelMaterialInputs(MaterialParameters, PixelMaterialInputs);
+						VolumetricCloudCalcPixelMaterialInputs(MaterialParameters, PixelMaterialInputs);
 						ConvertCloudPixelMaterialInputsToWorkingColorSpace(PixelMaterialInputs);
 						OpticalDepth += SampleExtinctionCoefficients(PixelMaterialInputs) * ShadowLengthTest * CENTIMETER_TO_METER * DetlaNormT;
 					}
@@ -951,7 +990,7 @@ void MainCommon(in FMaterialPixelParameters MaterialParameters, in float4 SvPosi
 					for (float ShadowT = 0.5; ShadowT < ShadowStepCount; ShadowT += 1.0f)
 					{
 						const float ShadowSampleDistance = ShadowLengthTest * (ShadowT * InvShadowStepCount);
-						UpdateMaterialCloudParam(MaterialParameters, SampleTranslatedWorldPosition + float3(0.0f, 0.0f, -1.0f) * ShadowSampleDistance,
+						UpdateMaterialCloudParam(MaterialParameters, SampleTranslatedWorldPosition + float3(0.0f, 0.0f, -1.0f) * ShadowSampleDistance, ShadowSampleWorldPositionDDX, ShadowSampleWorldPositionDDY,
 							ResolvedView, CloudLayerParams, ShadowSampleDistance, SPACE_SKIPPING_SLICE_DEPTH_OFF);
 #if MATERIAL_VOLUMETRIC_ADVANCED_CONSERVATIVE_DENSITY
 						if (MaterialParameters.VolumeSampleConservativeDensity.x <= 0.0f)
@@ -959,7 +998,7 @@ void MainCommon(in FMaterialPixelParameters MaterialParameters, in float4 SvPosi
 							continue; // Conservative density is 0 so skip and go to the next sample
 						}
 #endif
-						CalcPixelMaterialInputs(MaterialParameters, PixelMaterialInputs);
+						VolumetricCloudCalcPixelMaterialInputs(MaterialParameters, PixelMaterialInputs);
 						ConvertCloudPixelMaterialInputsToWorkingColorSpace(PixelMaterialInputs);
 						OpticalDepth += SampleExtinctionCoefficients(PixelMaterialInputs) * ShadowDtMeter;
 					}
@@ -1031,7 +1070,7 @@ void MainCommon(in FMaterialPixelParameters MaterialParameters, in float4 SvPosi
 				for (float ShadowT = ShadowJitterNorm; ShadowT < ShadowStepCount; ShadowT += 1.0f)
 				{
 					const float ShadowSampleDistance = ShadowLengthTest * (ShadowT * InvShadowStepCount);
-					UpdateMaterialCloudParam(MaterialParameters, SampleTranslatedWorldPosition + Light0Direction * ShadowSampleDistance,
+					UpdateMaterialCloudParam(MaterialParameters, SampleTranslatedWorldPosition + Light0Direction * ShadowSampleDistance, ShadowSampleWorldPositionDDX, ShadowSampleWorldPositionDDY,
 						ResolvedView, CloudLayerParams, ShadowSampleDistance, SPACE_SKIPPING_SLICE_DEPTH_OFF);
 					const float ExtinctionFactor = 1.0f;
 #else
@@ -1045,7 +1084,7 @@ void MainCommon(in FMaterialPixelParameters MaterialParameters, in float4 SvPosi
 					const float ExtinctionFactor = DetlaNormT;
 					const float ShadowSampleDistance = ShadowLengthTest * (PreviousNormT + DetlaNormT * ShadowJitterNorm);
 					const float3 ShadowSampleTranslatedWorldPos = SampleTranslatedWorldPosition + Light0Direction * ShadowSampleDistance;
-					UpdateMaterialCloudParam(MaterialParameters, ShadowSampleTranslatedWorldPos,
+					UpdateMaterialCloudParam(MaterialParameters, ShadowSampleTranslatedWorldPos, ShadowSampleWorldPositionDDX, ShadowSampleWorldPositionDDY,
 						ResolvedView, CloudLayerParams, ShadowSampleDistance, SPACE_SKIPPING_SLICE_DEPTH_OFF);
 					PreviousNormT = CurrentNormT;
 #endif
@@ -1071,7 +1110,7 @@ void MainCommon(in FMaterialPixelParameters MaterialParameters, in float4 SvPosi
 						break; // Ignore remaining samples since we have just traveld out of the cloud layer.
 					}
 
-					CalcPixelMaterialInputs(MaterialParameters, PixelMaterialInputs);
+					VolumetricCloudCalcPixelMaterialInputs(MaterialParameters, PixelMaterialInputs);
 					ConvertCloudPixelMaterialInputsToWorkingColorSpace(PixelMaterialInputs);
 					MATVEC ShadowExtinctionCoefficients = SampleExtinctionCoefficients(PixelMaterialInputs);
 
@@ -1120,7 +1159,7 @@ void MainCommon(in FMaterialPixelParameters MaterialParameters, in float4 SvPosi
 						const float DetlaNormT = CurrentNormT - PreviousNormT;
 						const float ExtinctionFactor = DetlaNormT;
 						const float ShadowSampleDistance = ShadowLengthTest * (PreviousNormT + DetlaNormT * ShadowJitterNorm);
-						UpdateMaterialCloudParam(MaterialParameters, SampleTranslatedWorldPosition + Light1Direction * ShadowSampleDistance,
+						UpdateMaterialCloudParam(MaterialParameters, SampleTranslatedWorldPosition + Light1Direction * ShadowSampleDistance, ShadowSampleWorldPositionDDX, ShadowSampleWorldPositionDDY,
 							ResolvedView, CloudLayerParams, ShadowSampleDistance, SPACE_SKIPPING_SLICE_DEPTH_OFF);
 						PreviousNormT = CurrentNormT;
 
@@ -1136,7 +1175,7 @@ void MainCommon(in FMaterialPixelParameters MaterialParameters, in float4 SvPosi
 							break; // Ignore remaining samples since we have just traveld out of the cloud layer.
 						}
 
-						CalcPixelMaterialInputs(MaterialParameters, PixelMaterialInputs);
+						VolumetricCloudCalcPixelMaterialInputs(MaterialParameters, PixelMaterialInputs);
 						ConvertCloudPixelMaterialInputsToWorkingColorSpace(PixelMaterialInputs);
 						MATVEC ShadowExtinctionCoefficients = SampleExtinctionCoefficients(PixelMaterialInputs);
 
@@ -1256,11 +1295,11 @@ void MainCommon(in FMaterialPixelParameters MaterialParameters, in float4 SvPosi
 							// Non-linear shadow sample distribution, hardcoded to x^2
 							for (float ShadowT = ShadowJitterNorm; ShadowT < ShadowStepCount; ShadowT += 1.0f)
 							{
-								UpdateMaterialCloudParam(MaterialParameters, SampleTranslatedWorldPosition + SampleToLightNorm * ShadowLengthTest * (ShadowT * InvShadowStepCount),
+								UpdateMaterialCloudParam(MaterialParameters, SampleTranslatedWorldPosition + SampleToLightNorm * ShadowLengthTest * (ShadowT * InvShadowStepCount), ShadowSampleWorldPositionDDX, ShadowSampleWorldPositionDDY,
 									ResolvedView, CloudLayerParams, TRACING_SHADOW_DISTANCE_OFF, SPACE_SKIPPING_SLICE_DEPTH_OFF);
 								const float ExtinctionFactor = 1.0f;
 
-								CalcPixelMaterialInputs(MaterialParameters, PixelMaterialInputs);
+								VolumetricCloudCalcPixelMaterialInputs(MaterialParameters, PixelMaterialInputs);
 								ConvertCloudPixelMaterialInputsToWorkingColorSpace(PixelMaterialInputs);
 								MATVEC ShadowExtinctionCoefficients = SampleExtinctionCoefficients(PixelMaterialInputs);
 
@@ -1816,6 +1855,9 @@ float GetSvPositionEmptySpaceSkippingValue(in ViewState ResolvedView, in uint3 D
 	float4 SvPosition = float4(PixelUV * ResolvedView.ViewSizeAndInvSize.xy, 0.1, 1);
 
 	FMaterialPixelParameters MaterialParameters = MakeInitializedMaterialPixelParameters();
+	
+	// CalcMaterialParameters() evaluates the material, but since the depth associated with SvPosition at this point doesn't correspond to any particular point along the ray,
+	// we leave WorldPosition_DDX/Y as 0, which forces mip 0 when compiled with USE_ANALYTIC_DERIVATIVES.
 	FPixelMaterialInputs PixelMaterialInputs;
 	CalcMaterialParameters(MaterialParameters, PixelMaterialInputs, SvPosition, true);
 
@@ -1825,7 +1867,21 @@ float GetSvPositionEmptySpaceSkippingValue(in ViewState ResolvedView, in uint3 D
 
 	const float SampleDistanceFromCamera = (float(GroupThreadId.z) + SampleFroxelOffset.z) * GetSliceDepth();
 	const float3 Raydir = -MaterialParameters.CameraVector;
-	UpdateMaterialCloudParam(MaterialParameters, ResolvedView.TranslatedWorldCameraOrigin + Raydir * SampleDistanceFromCamera,
+	const float3 SampleTranslatedWorldPosition = ResolvedView.TranslatedWorldCameraOrigin + Raydir * SampleDistanceFromCamera;
+
+#if USE_ANALYTIC_DERIVATIVES
+	const float3 TranslatedWorldPositionOffsetX = SvPositionToResolvedTranslatedWorld(SvPosition + float4(1.0f, 0.0f, 0.0f, 0.0f));
+	const float3 TranslatedWorldPositionOffsetY = SvPositionToResolvedTranslatedWorld(SvPosition + float4(0.0f, 1.0f, 0.0f, 0.0f));
+	const float3 RayDirOffsetX = GetCameraVector(ResolvedView, TranslatedWorldPositionOffsetX, ResolvedView.TranslatedWorldCameraOrigin);
+	const float3 RayDirOffsetY = GetCameraVector(ResolvedView, TranslatedWorldPositionOffsetY, ResolvedView.TranslatedWorldCameraOrigin);
+	float3 SampleWorldPositionDDX = (ResolvedView.TranslatedWorldCameraOrigin + t * RayDirOffsetX) - SampleTranslatedWorldPosition;
+	float3 SampleWorldPositionDDY = (ResolvedView.TranslatedWorldCameraOrigin + t * RayDirOffsetY) - SampleTranslatedWorldPosition;
+#else
+	float3 SampleWorldPositionDDX = 0.0f;
+	float3 SampleWorldPositionDDY = 0.0f;
+#endif // USE_ANALYTIC_DERIVATIVES
+
+	UpdateMaterialCloudParam(MaterialParameters, SampleTranslatedWorldPosition, SampleWorldPositionDDX, SampleWorldPositionDDY,
 		ResolvedView, CloudLayerParams, TRACING_SHADOW_DISTANCE_OFF, GetSliceDepth());
 
 	const float CloudEmptySpaceSkippingOutput0 = GetVolumetricCloudEmptySpaceSkippingOutput0(MaterialParameters);
@@ -1949,7 +2005,7 @@ void MainCS(uint3 DispatchThreadId : SV_DispatchThreadID, uint3 GroupThreadId :
 #endif // MATERIAL_VOLUMETRIC_CLOUD_EMPTY_SPACE_SKIPPING_OUTPUT
 }
 
-#endif // SHADER_SHADOW_PS
+#endif // SHADER_EMPTY_SPACE_SKIPPING_CS
 
 
 
@@ -2090,7 +2146,8 @@ void MainPS(
 		for (float ShadowT = 0.5; ShadowT < ShadowStepCount; ShadowT += 1.0f)
 		{
 			const float SampleT = ShadowLengthTest * (ShadowT * InvShadowStepCount);
-			UpdateMaterialCloudParam(MaterialParameters, (WorldPosOnLayer + LightDirection * SampleT) + DFHackToFloat(PrimaryView.PreViewTranslation),			// LWC TODO FIX ME
+			// WorldPositionDDX/Y only need to be set when using analytical derivatives, which is not needed in pixel shaders
+			UpdateMaterialCloudParam(MaterialParameters, (WorldPosOnLayer + LightDirection * SampleT) + DFHackToFloat(PrimaryView.PreViewTranslation), 0.0f /*WorldPositionDDX*/, 0.0f /*WorldPositionDDY*/,			// LWC TODO FIX ME
 				ResolvedView, CloudLayerParams, SampleT, SPACE_SKIPPING_SLICE_DEPTH_OFF);
 
 		#if MATERIAL_VOLUMETRIC_ADVANCED_CONSERVATIVE_DENSITY
@@ -2100,7 +2157,7 @@ void MainPS(
 			}
 		#endif
 
-			CalcPixelMaterialInputs(MaterialParameters, PixelMaterialInputs);
+			VolumetricCloudCalcPixelMaterialInputs(MaterialParameters, PixelMaterialInputs);
 			ConvertCloudPixelMaterialInputsToWorkingColorSpace(PixelMaterialInputs);
 			float3 ShadowExtinctionCoefficients = SampleExtinctionCoefficients(PixelMaterialInputs);
 
diff --git a/Engine/Shaders/Private/VolumetricCloudMaterialPixelCommon.ush b/Engine/Shaders/Private/VolumetricCloudMaterialPixelCommon.ush
--- a/Engine/Shaders/Private/VolumetricCloudMaterialPixelCommon.ush
+++ b/Engine/Shaders/Private/VolumetricCloudMaterialPixelCommon.ush
@@ -43,7 +43,8 @@ FCloudLayerParameters GetCloudLayerParams(
 
 // Function to update the material parameter structure for an evaluation
 void UpdateMaterialCloudParam(inout FMaterialPixelParameters MaterialParameters,
-                              float3 TranslatedWorldPosition, in ViewState InputView, in FCloudLayerParameters CloudLayerParams,
+                              float3 TranslatedWorldPosition, float3 WorldPositionDDX, float3 WorldPositionDDY,
+							  in ViewState InputView, in FCloudLayerParameters CloudLayerParams,
                               in float ShadowSampleDistance, in float EmptySpaceSkippingSliceDepth)
 {
 	MaterialParameters.AbsoluteWorldPosition = DFFastSubtract(TranslatedWorldPosition, InputView.PreViewTranslation);
@@ -53,6 +54,8 @@ void UpdateMaterialCloudParam(inout FMaterialPixelParameters MaterialParameters,
 	//MaterialParameters.LWCData.WorldPosition_NoOffsets;	// TODO
 	//MaterialParameters.WorldPosition_NoOffsets_CamRelative;	// TODO
 	MaterialParameters.CameraVector = normalize(InputView.TranslatedWorldCameraOrigin - TranslatedWorldPosition);
+	MaterialParameters.WorldPosition_DDX = WorldPositionDDX;
+	MaterialParameters.WorldPosition_DDY = WorldPositionDDY;
 
 	float3 AbsoluteWorldPositionTruncated = TranslatedWorldPosition - DFHackToFloat(InputView.PreViewTranslation);
 	float CloudSampleAltitude = length(AbsoluteWorldPositionTruncated - CloudLayerParams.CloudLayerCenter);
diff --git a/Engine/Shaders/Private/VolumetricFogVoxelization.usf b/Engine/Shaders/Private/VolumetricFogVoxelization.usf
--- a/Engine/Shaders/Private/VolumetricFogVoxelization.usf
+++ b/Engine/Shaders/Private/VolumetricFogVoxelization.usf
@@ -314,7 +314,8 @@ void VoxelizePS(
 		VoxelizeVolumePass.RenderVolumetricCloudParametersCloudLayerCenterKm, VoxelizeVolumePass.RenderVolumetricCloudParametersPlanetRadiusKm,
 		VoxelizeVolumePass.RenderVolumetricCloudParametersBottomRadiusKm, VoxelizeVolumePass.RenderVolumetricCloudParametersTopRadiusKm);
 	float3 TranslatedWorldPosition = SvPositionToTranslatedWorld(SvPosition);
-	UpdateMaterialCloudParam(MaterialParameters, TranslatedWorldPosition, ResolvedView, CloudLayerParams, TRACING_SHADOW_DISTANCE_OFF, SPACE_SKIPPING_SLICE_DEPTH_OFF);
+	// WorldPositionDDX/Y only need to be set when using analytical derivatives, which is not needed in pixel shaders
+	UpdateMaterialCloudParam(MaterialParameters, TranslatedWorldPosition, 0.0f /*WorldPositionDDX*/, 0.0f /*WorldPositionDDY*/, ResolvedView, CloudLayerParams, TRACING_SHADOW_DISTANCE_OFF, SPACE_SKIPPING_SLICE_DEPTH_OFF);
 #endif
 
 	// Scale rendered position to account for frustum difference between fog voxelization and rendered scene
diff --git a/Engine/Source/Runtime/Core/Private/UObject/FortniteMainBranchObjectVersion.cpp b/Engine/Source/Runtime/Core/Private/UObject/FortniteMainBranchObjectVersion.cpp
--- a/Engine/Source/Runtime/Core/Private/UObject/FortniteMainBranchObjectVersion.cpp
+++ b/Engine/Source/Runtime/Core/Private/UObject/FortniteMainBranchObjectVersion.cpp
@@ -14,6 +14,6 @@ TMap<FGuid, FGuid> FFortniteMainBranchObjectVersion::GetSystemGuids()
 	SystemGuids.Add(DevGuids.Niagara_LatestScriptCompileVersion, FGuid("E267F600682A4E94885E8AEC5CE1DC91"));
 	SystemGuids.Add(DevGuids.SkeletalMeshDerivedDataVersion, FGuid("B9AEAA3EE7AC4FCFB3860E8F3E06DD75"));
 	SystemGuids.Add(DevGuids.STATICMESH_DERIVEDDATA_VER, FGuid("543843B103794E0E9BA4BE80FB602F79"));
-	SystemGuids.Add(DevGuids.MaterialTranslationDDCVersion, FGuid("84B1DF444194438F8379C1FCE1C5EF9B"));
+	SystemGuids.Add(DevGuids.MaterialTranslationDDCVersion, FGuid("5F7AE5F7FE724C4CA448A63427BC4BB7"));
 	return SystemGuids;
 }
diff --git a/Engine/Source/Runtime/Engine/Private/Materials/HLSLMaterialTranslator.cpp b/Engine/Source/Runtime/Engine/Private/Materials/HLSLMaterialTranslator.cpp
--- a/Engine/Source/Runtime/Engine/Private/Materials/HLSLMaterialTranslator.cpp
+++ b/Engine/Source/Runtime/Engine/Private/Materials/HLSLMaterialTranslator.cpp
@@ -269,6 +269,7 @@ struct FHLSLMaterialTranslator::FEnvironmentDefines
 	bool bUsesDisplacement;
 	bool bUsesEmissiveColor;
 	bool bUsesDistortion;
+	bool bUsesExplicitDerivatives;
 	bool bDistortionAccountForCoverage;
 	bool bMaterialEnableTranslucencyFogging;
 	bool bMaterialEnableTranslucencyCloudFogging;
@@ -361,6 +362,7 @@ struct FHLSLMaterialTranslator::FEnvironmentDefines
 		Ar << bUsesDisplacement;
 		Ar << bUsesEmissiveColor;
 		Ar << bUsesDistortion;
+		Ar << bUsesExplicitDerivatives;
 		Ar << bDistortionAccountForCoverage;
 		Ar << bMaterialEnableTranslucencyFogging;
 		Ar << bMaterialEnableTranslucencyCloudFogging;
@@ -676,6 +678,7 @@ FHLSLMaterialTranslator::FHLSLMaterialTranslator(FMaterial* InMaterial,
 ,	bUsesCurvature(false)
 ,	bUsesPerInstanceFadeAmount(false)
 ,	bCullIntermediateUniformExpressions(GCullIntermediateUniformExpressions)
+,	bUsesExplicitDerivatives(false)
 ,	AddingUniformExpression(0)
 ,	AllocatedUserTexCoords()
 ,	AllocatedUserVertexTexCoords()
@@ -2546,6 +2549,8 @@ void FHLSLMaterialTranslator::GetMaterialEnvironment(EShaderPlatform InPlatform,
 	OutEnvironment.SetDefineAndCompileArgument(TEXT("USES_WORLD_POSITION_OFFSET"), EnvironmentDefines->bUsesWorldPositionOffset);
 	OutEnvironment.SetDefineAndCompileArgument(TEXT("USES_DISPLACEMENT"), EnvironmentDefines->bUsesDisplacement);
 
+	OutEnvironment.SetDefine(TEXT("USES_EXPLICIT_DERIVATIVES"), EnvironmentDefines->bUsesExplicitDerivatives);
+
 	OutEnvironment.SetDefine(TEXT("USES_EMISSIVE_COLOR"), EnvironmentDefines->bUsesEmissiveColor);
 	// Distortion uses tangent space transform 
 	OutEnvironment.SetDefine(TEXT("USES_DISTORTION"), EnvironmentDefines->bUsesDistortion);
@@ -11048,6 +11053,8 @@ int32 FHLSLMaterialTranslator::Derivative(int32 A, EDervativeComponent Component
 		return NonPixelShaderExpressionError();
 	}
 
+	bUsesExplicitDerivatives = true;
+
 	const FDerivInfo ADerivInfo = GetDerivInfo(A);
 	const EMaterialValueType ResultType = MakeNonLWCType(ADerivInfo.Type);
 	const bool bIsLWCType = IsLWCType(ADerivInfo.Type);
@@ -14703,6 +14710,7 @@ void FHLSLMaterialTranslator::PrepareEnvironmentDefines()
 	EnvironmentDefines->bUsesEmissiveColor = bUsesEmissiveColor;
 	// Distortion uses tangent space transform 
 	EnvironmentDefines->bUsesDistortion = Material->IsDistorted();
+	EnvironmentDefines->bUsesExplicitDerivatives = bUsesExplicitDerivatives;
 	EnvironmentDefines->bDistortionAccountForCoverage = bSubstrateEnabled && Material->GetRefractionCoverageMode() == RCM_CoverageAccountedFor ? 1 : 0;
 	EnvironmentDefines->bMaterialEnableTranslucencyFogging = Material->ShouldApplyFogging();
 	EnvironmentDefines->bMaterialEnableTranslucencyCloudFogging = Material->ShouldApplyCloudFogging();
diff --git a/Engine/Source/Runtime/Engine/Private/Materials/HLSLMaterialTranslator.h b/Engine/Source/Runtime/Engine/Private/Materials/HLSLMaterialTranslator.h
--- a/Engine/Source/Runtime/Engine/Private/Materials/HLSLMaterialTranslator.h
+++ b/Engine/Source/Runtime/Engine/Private/Materials/HLSLMaterialTranslator.h
@@ -429,6 +429,9 @@ protected:
 	uint32 bUsesPerInstanceFadeAmount : 1;
 
 	uint32 bCullIntermediateUniformExpressions : 1;
+	
+	/** The generated shader source uses explicit ddx()/ddy() calls */
+	uint32 bUsesExplicitDerivatives : 1;
 
 	/** Incremented and decremented by FAddUniformExpressionScope.  See comments on that class below */
 	int32 AddingUniformExpression;
diff --git a/Engine/Source/Runtime/Renderer/Private/VolumetricCloudRendering.cpp b/Engine/Source/Runtime/Renderer/Private/VolumetricCloudRendering.cpp
--- a/Engine/Source/Runtime/Renderer/Private/VolumetricCloudRendering.cpp
+++ b/Engine/Source/Runtime/Renderer/Private/VolumetricCloudRendering.cpp
@@ -236,6 +236,11 @@ static TAutoConsoleVariable<int32> CVarVolumetricCloudEmptySpaceSkippingSampleCo
 	TEXT("0 means center samples only, >0 means corner are also sampled."),
 	ECVF_RenderThreadSafe | ECVF_Scalability);
 
+static TAutoConsoleVariable<int32> CVarVolumetricCloudAllowAnalyticDerivatives(
+	TEXT("r.VolumetricCloud.AllowAnalyticDerivatives"), 1,
+	TEXT("Enables compiling cloud shaders with support for analytical derivatives when needed by the material."),
+	ECVF_ReadOnly | ECVF_RenderThreadSafe);
+
 ////////////////////////////////////////////////////////////////////////// 
 
 
@@ -902,6 +907,11 @@ class FRenderVolumetricCloudRenderViewCS : public FMeshMaterialShader
 		// If this is enabled
 		OutEnvironment.SetDefine(TEXT("CLOUD_DEBUG_SAMPLES"), CLOUD_DEBUG_SAMPLES);
 
+		if (CVarVolumetricCloudAllowAnalyticDerivatives.GetValueOnAnyThread() != 0)
+		{
+			OutEnvironment.SetDefine(TEXT("ALLOW_ANALYTIC_DERIVATIVES"), 1);
+		}
+
 		// This shader must support typed UAV load and we are testing if it is supported at runtime using RHIIsTypedUAVLoadSupported
 		OutEnvironment.CompilerFlags.Add(CFLAG_AllowTypedUAVLoads);
 
@@ -972,6 +982,11 @@ class FRenderVolumetricCloudEmptySpaceSkippingCS : public FMeshMaterialShader
 		OutEnvironment.SetDefine(TEXT("THREADGROUP_SIZEY"), ThreadGroupSizeY);
 		OutEnvironment.SetDefine(TEXT("THREADGROUP_SIZEZ"), ThreadGroupSizeZ);
 
+		if (CVarVolumetricCloudAllowAnalyticDerivatives.GetValueOnAnyThread() != 0)
+		{
+			OutEnvironment.SetDefine(TEXT("ALLOW_ANALYTIC_DERIVATIVES"), 1);
+		}
+
 		// This shader must support typed UAV load and we are testing if it is supported at runtime using RHIIsTypedUAVLoadSupported
 		OutEnvironment.CompilerFlags.Add(CFLAG_AllowTypedUAVLoads);
 
