Split screen:  Fix for incorrect UV calculation for compute version of Tonemap shader in split screen.

#rnx

diff --git a/Engine/Shaders/Private/PostProcessTonemap.usf b/Engine/Shaders/Private/PostProcessTonemap.usf
--- a/Engine/Shaders/Private/PostProcessTonemap.usf
+++ b/Engine/Shaders/Private/PostProcessTonemap.usf
@@ -615,8 +615,8 @@ RWTexture2D<float> RWOutputLuminance;
 [numthreads(THREADGROUP_SIZEX, THREADGROUP_SIZEY, 1)]
 void MainCS(uint2 DispatchThreadId : SV_DispatchThreadID)
 {
-	float4 SvPosition = float4((float2)DispatchThreadId + Output_ViewportMin + 0.5f, 0.0f, 1.0f);
-	float2 UV = SvPosition.xy * Output_ExtentInverse;
+	float4 SvPosition = float4((float2)DispatchThreadId + 0.5f + Output_ViewportMin, 0.0f, 1.0f);
+	float2 UV = ((float2)DispatchThreadId + 0.5f) * Output_ViewportSizeInverse;
 
 	if (IsComputeUVOutOfBounds(UV))
 	{
