Fixed unsigned int warning

diff --git a/Engine/Shaders/Private/SceneData.ush b/Engine/Shaders/Private/SceneData.ush
--- a/Engine/Shaders/Private/SceneData.ush
+++ b/Engine/Shaders/Private/SceneData.ush
@@ -223,7 +223,7 @@ FPrimitiveSceneData GetPrimitiveDataFromUniformBuffer()
 	PrimitiveData.MaterialDisplacementFadeOutSize	= Primitive.MaterialDisplacementFadeOutSize;
 	
 	UNROLL
-	for (int DataIndex = 0; DataIndex < NUM_CUSTOM_PRIMITIVE_DATA; ++DataIndex)
+	for (uint DataIndex = 0; DataIndex < NUM_CUSTOM_PRIMITIVE_DATA; ++DataIndex)
 	{
 		PrimitiveData.CustomPrimitiveData[DataIndex] = Primitive.CustomPrimitiveData[DataIndex];
 	}
@@ -474,7 +474,7 @@ FPrimitiveSceneData GetPrimitiveData(uint PrimitiveId)
 
 	// TODO: Move to associated array (and editor data) to shrink primitive data and better pack cachelines
 	UNROLL
-	for (int DataIndex = 0; DataIndex < NUM_CUSTOM_PRIMITIVE_DATA; ++DataIndex)
+	for (uint DataIndex = 0; DataIndex < NUM_CUSTOM_PRIMITIVE_DATA; ++DataIndex)
 	{
 		PrimitiveData.CustomPrimitiveData[DataIndex] = LoadPrimitivePrimitiveSceneDataElement(PrimitiveIndex,  33 + DataIndex);
 	}
