PR #11769: Fixing Position Driven Light Functions For Path Tracing And Ray Tracing

The original patch was not quite correct, but pointed in the right direction. The miss shaders were not setting up ResolvedView which is necessary for various MaterialTemplate fields to be populated properly. They were also using WorldRayOrigin() instead of TranslatedWorldRayOrigin() which is not quite the same in the multi-view case (split screen rendering).

diff --git a/Engine/Shaders/Private/PathTracing/PathTracingLightingMissShader.usf b/Engine/Shaders/Private/PathTracing/PathTracingLightingMissShader.usf
--- a/Engine/Shaders/Private/PathTracing/PathTracingLightingMissShader.usf
+++ b/Engine/Shaders/Private/PathTracing/PathTracingLightingMissShader.usf
@@ -17,8 +17,9 @@
 RAY_TRACING_ENTRY_MISS(PathTracingLightingMS, FPackedPathTracingPayload, Payload)
 {
 	Payload.SetMiss();
+	ResolvedView = ResolveView();
 	CurrentPayloadInputFlags = Payload.GetFlags();
-	float3 TranslatedWorldPosition = WorldRayOrigin(); // Shadow ray is traced from the shaded point
+	float3 TranslatedWorldPosition = TranslatedWorldRayOrigin(); // Shadow ray is traced from the shaded point
 	float3 Throughput = Payload.GetRayThroughput();
 	float3 LightFunctionMult = GetRayTracingLightFunction(TranslatedWorldPosition);
 
diff --git a/Engine/Shaders/Private/RayTracing/RayTracingLightingMS.usf b/Engine/Shaders/Private/RayTracing/RayTracingLightingMS.usf
--- a/Engine/Shaders/Private/RayTracing/RayTracingLightingMS.usf
+++ b/Engine/Shaders/Private/RayTracing/RayTracingLightingMS.usf
@@ -53,6 +53,8 @@ RAY_TRACING_ENTRY_MISS(RayTracingLightingMS, FPackedMaterialClosestHitPayload, P
 {
 	FRayDesc DummyRay = (FRayDesc)0;
 
+	ResolvedView = ResolveView();
+
 	FMaterialClosestHitPayload Payload = UnpackRayTracingPayload(PackedPayload, DummyRay);
 
 	// Light index is packed into HitT as this component is only accessed by closest hit or miss shaders.
@@ -67,9 +69,9 @@ RAY_TRACING_ENTRY_MISS(RayTracingLightingMS, FPackedMaterialClosestHitPayload, P
 
 	float4 LightAttenuation = 1.0f;
 
-	float3 WorldPosition = TranslatedWorldRayOrigin(); // Shadow ray is traced from the shaded point
+	float3 TranslatedWorldPosition = TranslatedWorldRayOrigin(); // Shadow ray is traced from the shaded point
 
-	float LightProfileMultiplier = ComputeLightProfileMultiplier(WorldPosition, LightData.TranslatedWorldPosition, -LightData.Direction, LightData.Tangent, LightData.IESAtlasIndex);
+	float LightProfileMultiplier = ComputeLightProfileMultiplier(TranslatedWorldPosition, LightData.TranslatedWorldPosition, -LightData.Direction, LightData.Tangent, LightData.IESAtlasIndex);
 
 	const float Dither = 0.5f;
 	const uint2 SVPos = uint2(0, 0);
@@ -116,11 +118,11 @@ RAY_TRACING_ENTRY_MISS(RayTracingLightingMS, FPackedMaterialClosestHitPayload, P
 	// LWC_TODO - lighting wants translated world space
 	float SurfaceShadow = 1.0f;
 	FGBufferData GBufferData = GetGBufferDataFromPayload(Payload);
-	float3 LightContribution = GetDynamicLighting(WorldPosition, CameraVector, GBufferData, 1.0f, GBufferData.ShadingModelID, LightData, LightAttenuation, Dither, SVPos, SurfaceShadow).xyz;
+	float3 LightContribution = GetDynamicLighting(TranslatedWorldPosition, CameraVector, GBufferData, 1.0f, GBufferData.ShadingModelID, LightData, LightAttenuation, Dither, SVPos, SurfaceShadow).xyz;
 #endif // SUBSTRATE_ENABLED
 
 #if SUPPORT_LIGHT_FUNCTION
-	float3 LightFunctionMult = GetRayTracingLightFunction(WorldPosition);
+	float3 LightFunctionMult = GetRayTracingLightFunction(TranslatedWorldPosition);
 	LightProfileMultiplier *= dot(LightFunctionMult, .3333f); // convert to greyscale (for compatibility with the appearance in the primary view)
 #endif
 
