[Nanite-Skinning] Fixed offset bug related to introduction of SkinningData in the per instance packing

diff --git a/Engine/Shaders/Private/SceneData.ush b/Engine/Shaders/Private/SceneData.ush
--- a/Engine/Shaders/Private/SceneData.ush
+++ b/Engine/Shaders/Private/SceneData.ush
@@ -937,7 +937,7 @@ FInstancePayloadDataOffsets GetInstancePayloadDataOffsets(uint PrimitiveId, uint
 	{
 		Offsets.LocalBounds = CurOffset;
 	}
-	CurOffset += CondMask(bHasLocalBounds, 2u, CondMask(bHasHierarchyOffset || bHasEditorData, 1u, 0u));
+	CurOffset += CondMask(bHasLocalBounds, 2u, CondMask(bHasHierarchyOffset || bHasSkinningData, 1u, 0u));
 
 	if (bHasDynamicData)
 	{
