Fix for Substrate permutation of RayTracingLightingMS.usf failing due to renamed variable.

diff --git a/Engine/Shaders/Private/RayTracing/RayTracingLightingMS.usf b/Engine/Shaders/Private/RayTracing/RayTracingLightingMS.usf
--- a/Engine/Shaders/Private/RayTracing/RayTracingLightingMS.usf
+++ b/Engine/Shaders/Private/RayTracing/RayTracingLightingMS.usf
@@ -84,7 +84,7 @@ RAY_TRACING_ENTRY_MISS(RayTracingLightingMS, FPackedMaterialClosestHitPayload, P
 	float LightMask = 1;
 	if (LightData.bRadialLight)
 	{
-		LightMask = GetLocalLightAttenuation(WorldPosition, LightData, ToLight, L);
+		LightMask = GetLocalLightAttenuation(TranslatedWorldPosition, LightData, ToLight, L);
 	}
 
 	float3 LightContribution = 0;
@@ -97,7 +97,7 @@ RAY_TRACING_ENTRY_MISS(RayTracingLightingMS, FPackedMaterialClosestHitPayload, P
 		SubstrateShadowTermInputParameters.bEvaluateShadowTerm					= true;
 		SubstrateShadowTermInputParameters.SceneDepth							= 1.0f; // Do not use depth
 		SubstrateShadowTermInputParameters.PrecomputedShadowFactors				= 1.0f; // Do not use precomputed factors
-		SubstrateShadowTermInputParameters.TranslatedWorldPosition				= WorldPosition;
+		SubstrateShadowTermInputParameters.TranslatedWorldPosition				= TranslatedWorldPosition;
 		SubstrateShadowTermInputParameters.LightAttenuation						= LightAttenuation;
 		SubstrateShadowTermInputParameters.Dither								= Dither;
 
