Fixed bug in tonemapper viewport calculations causing stretching of editor image when buffer size does not match viewport size

diff --git a/Engine/Shaders/Private/PostProcessTonemap.usf b/Engine/Shaders/Private/PostProcessTonemap.usf
--- a/Engine/Shaders/Private/PostProcessTonemap.usf
+++ b/Engine/Shaders/Private/PostProcessTonemap.usf
@@ -257,6 +257,7 @@ float max4(float4 x)
 
 float4 SampleSceneColor(float2 SceneUV, float ArrayIndex)
 {
+	SceneUV = ((SceneUV * Output_ViewportSize + Output_ViewportMin)) * Color_ExtentInverse;
 	SceneUV = clamp(SceneUV, Color_UVViewportBilinearMin, Color_UVViewportBilinearMax);
 	#if USE_MULTI_VIEW_ID_SV || MOBILE_MULTI_VIEW_FALLBACK
 	return Texture2DArraySample(ColorTexture, ColorSampler, float3(SceneUV.xy,ArrayIndex));
@@ -342,7 +343,8 @@ float4 TonemapCommonPS(
 		}
 		#else
 		{
-			BloomUV = ApplyScreenTransform(UV.xy, ColorToBloom);
+			BloomUV = ((SceneUV.xy * Output_ViewportSize + Output_ViewportMin)) * Color_ExtentInverse;
+			BloomUV = ApplyScreenTransform(BloomUV.xy, ColorToBloom);
 			BloomUV = clamp(BloomUV, BloomUVViewportBilinearMin, BloomUVViewportBilinearMax);
 		}
 		#endif
