Add alternative calculation route for RectLights with Substrate to ensure F90 approximation is included in the final colour output.

diff --git a/Engine/Shaders/Private/RectLight.ush b/Engine/Shaders/Private/RectLight.ush
--- a/Engine/Shaders/Private/RectLight.ush
+++ b/Engine/Shaders/Private/RectLight.ush
@@ -378,12 +378,8 @@ struct FRectLTC
 };
 
 // Return LTC coefficients for GGX-based BSDF
-FRectLTC GetRectLTC_GGX( float Roughness, float3 SpecularColor, half NoV)
+FRectLTC GetRectLTC_GGX(float Roughness, float3 F0, float3 F90, half NoV)
 {
-	// F90 is derived as, anything less than 2% is physically impossible and is instead considered to be shadowing
-	const float3 F0  = SpecularColor;
-	const float3 F90 = saturate(50.0 * SpecularColor);
-
 	float2 UV = float2( Roughness, sqrt( 1 - NoV ) );
 	UV = UV * (63.0 / 64.0) + (0.5 / 64.0);
    
@@ -412,6 +408,15 @@ FRectLTC GetRectLTC_GGX( float Roughness, float3 SpecularColor, half NoV)
 	return Out;
 }
 
+FRectLTC GetRectLTC_GGX( float Roughness, float3 SpecularColor, half NoV)
+{
+	// F90 is derived as, anything less than 2% is physically impossible and is instead considered to be shadowing
+	const float3 F0  = SpecularColor;
+	const float3 F90 = saturate(50.0 * SpecularColor);
+
+	return GetRectLTC_GGX(Roughness, F0, F90, NoV);
+}
+
 #if SUPPORTS_INDEPENDENT_SAMPLERS
 	#define SharedSheenLTCSampler	View.SharedBilinearClampedSampler
 #else
@@ -527,6 +532,23 @@ float3 RectGGXApproxLTC(float Roughness, float3 SpecularColor, half3 N, float3 V
 	return RectGGXApproxLTC(Roughness, SpecularColor, N, V, Rect, RectTexture, MeanLightWorldDirection);
 }
 
+float3 RectGGXApproxLTC(float Roughness, float3 F0, float3 F90, half3 N, float3 V, FRect Rect, FRectTexture RectTexture, inout float3 OutMeanLightWorldDirection)
+{
+	// No visibile rect light due to barn door occlusion
+	if (Rect.Extent.x == 0 || Rect.Extent.y == 0) return 0;
+
+	const float NoV = saturate(abs(dot(N, V)) + 1e-5);
+
+	const FRectLTC LTC = GetRectLTC_GGX(Roughness, F0, F90, NoV);
+	return RectApproxLTC(LTC, N, V, Rect, RectTexture, OutMeanLightWorldDirection);
+}
+
+float3 RectGGXApproxLTC(float Roughness, float3 F0, float3 F90, half3 N, float3 V, FRect Rect, FRectTexture RectTexture)
+{
+	float3 MeanLightWorldDirection = 0.0f;
+	return RectGGXApproxLTC(Roughness, F0, F90, N, V, Rect, RectTexture, MeanLightWorldDirection);
+}
+
 // Integrated a Sheen-based BSDF with a rect light using LTC
 float3 RectSheenApproxLTC( float Roughness, half3 N, float3 V, FRect Rect, FRectTexture RectTexture, inout float DirectionalAlbedo)
 {
diff --git a/Engine/Shaders/Private/Substrate/SubstrateEvaluation.ush b/Engine/Shaders/Private/Substrate/SubstrateEvaluation.ush
--- a/Engine/Shaders/Private/Substrate/SubstrateEvaluation.ush
+++ b/Engine/Shaders/Private/Substrate/SubstrateEvaluation.ush
@@ -460,7 +460,7 @@ FSubstrateEvaluateResult SubstrateEvaluateBSDFCommon(FSubstrateBSDFContext BSDFC
 						// In this case, we set D to 1 and Vis will contain the area light / GGX integration
 						{
 							float3 MeanLightWorldDirection = 0.0f;
-							RectLightSpec = RectGGXApproxLTC(SafeRoughness, F0, BSDFContext.N, BSDFContext.V, AreaLightContext.AreaLight.Rect, AreaLightContext.AreaLight.Texture, MeanLightWorldDirection);
+							RectLightSpec = RectGGXApproxLTC(SafeRoughness, F0, F90, BSDFContext.N, BSDFContext.V, AreaLightContext.AreaLight.Rect, AreaLightContext.AreaLight.Texture, MeanLightWorldDirection);
 							
 							// Now combine the rectangular area light with Glints.
 						#if SUBSTRATE_COMPLEXSPECIALPATH && SUBSTRATE_GLINTS_ALLOWED && SUBSTRATE_GLINTS_ENABLED
@@ -593,7 +593,7 @@ FSubstrateEvaluateResult SubstrateEvaluateBSDFCommon(FSubstrateBSDFContext BSDFC
 					if (bIsRectLight)
 					{
 						{
-							RectLightSpecHaze = RectGGXApproxLTC(HazeSafeRoughness, HazeF0, BSDFContext.N, BSDFContext.V, AreaLightContext.AreaLight.Rect, AreaLightContext.AreaLight.Texture);
+							RectLightSpecHaze = RectGGXApproxLTC(HazeSafeRoughness, HazeF0, HazeF90, BSDFContext.N, BSDFContext.V, AreaLightContext.AreaLight.Rect, AreaLightContext.AreaLight.Texture);
 					
 							const float H_PDF = VisibleGGXPDF(BSDFContext.TangentV, BSDFContext.TangentH, Pow4(HazeSafeRoughness));
 							HazePDF = RayPDFToReflectionRayPDF(BSDFContext.Context.VoH, H_PDF);
@@ -1090,7 +1090,7 @@ FSubstrateEvaluateResult SubstrateEvaluateBSDFCommon(FSubstrateBSDFContext BSDFC
 			if (bIsRectLight)
 			{
 				// In this case, we set D to 1 and Vis will contain the area light / GGX integration
-				RectLightSpec = RectGGXApproxLTC(SafeRoughness, F0, BSDFContext.N, BSDFContext.V, AreaLightContext.AreaLight.Rect, AreaLightContext.AreaLight.Texture);
+				RectLightSpec = RectGGXApproxLTC(SafeRoughness, F0, F90, BSDFContext.N, BSDFContext.V, AreaLightContext.AreaLight.Rect, AreaLightContext.AreaLight.Texture);
 
 				const float H_PDF = VisibleGGXPDF(BSDFContext.TangentV, BSDFContext.TangentH, Alpha2);
 				PDF = RayPDFToReflectionRayPDF(BSDFContext.Context.VoH, H_PDF);
