[Backout] - CL33136672
Original CL Desc
-----------------------------------------------------------------
Fixed bug in tonemapper viewport calculations causing stretching of editor image when buffer size does not match viewport size

diff --git a/Engine/Shaders/Private/PostProcessTonemap.usf b/Engine/Shaders/Private/PostProcessTonemap.usf
--- a/Engine/Shaders/Private/PostProcessTonemap.usf
+++ b/Engine/Shaders/Private/PostProcessTonemap.usf
@@ -257,7 +257,6 @@ float max4(float4 x)
 
 float4 SampleSceneColor(float2 SceneUV, float ArrayIndex)
 {
-	SceneUV = ((SceneUV * Output_ViewportSize + Output_ViewportMin)) * Color_ExtentInverse;
 	SceneUV = clamp(SceneUV, Color_UVViewportBilinearMin, Color_UVViewportBilinearMax);
 	#if USE_MULTI_VIEW_ID_SV || MOBILE_MULTI_VIEW_FALLBACK
 	return Texture2DArraySample(ColorTexture, ColorSampler, float3(SceneUV.xy,ArrayIndex));
@@ -343,8 +342,7 @@ float4 TonemapCommonPS(
 		}
 		#else
 		{
-			BloomUV = ((SceneUV.xy * Output_ViewportSize + Output_ViewportMin)) * Color_ExtentInverse;
-			BloomUV = ApplyScreenTransform(BloomUV.xy, ColorToBloom);
+			BloomUV = ApplyScreenTransform(UV.xy, ColorToBloom);
 			BloomUV = clamp(BloomUV, BloomUVViewportBilinearMin, BloomUVViewportBilinearMax);
 		}
 		#endif
