Fix alpha for FXAA shader when propagateAlpha is enabled.
#android

diff --git a/Engine/Shaders/Private/FXAAShader.usf b/Engine/Shaders/Private/FXAAShader.usf
--- a/Engine/Shaders/Private/FXAAShader.usf
+++ b/Engine/Shaders/Private/FXAAShader.usf
@@ -100,6 +100,8 @@ SamplerState Input_Sampler;
 
 		#if (POST_PROCESS_ALPHA != 2)
 			OutColor.a = 1.0;
+		#else
+			OutColor.a = 1.0f - OutColor.a;	// 	r.PostProcessing.PropagateAlpha
 		#endif
 	}
 
