Fix issue where thin translucent materials may be rendered fully translucent when thin translucency is disabled via r.Mobile.ShadingModelsMask

Requires CL 33135504

diff --git a/Engine/Shaders/Private/MobileBasePassPixelShader.usf b/Engine/Shaders/Private/MobileBasePassPixelShader.usf
--- a/Engine/Shaders/Private/MobileBasePassPixelShader.usf
+++ b/Engine/Shaders/Private/MobileBasePassPixelShader.usf
@@ -81,7 +81,19 @@
 #define FORCE_VERTEX_NORMAL (MOBILE_QL_DISABLE_MATERIAL_NORMAL)
 #define SUPPORT_SPOTLIGHTS_SHADOW (MATERIALBLENDING_SOLID || MATERIALBLENDING_MASKED) && PROJECT_MOBILE_ENABLE_MOVABLE_SPOTLIGHT_SHADOWS
 #define MOBILE_USE_GBUBUFFER (MOBILE_DEFERRED_SHADING && ((MATERIALBLENDING_SOLID || MATERIALBLENDING_MASKED) && !MATERIAL_SHADINGMODEL_SINGLELAYERWATER))
-#define MATERIAL_WORKS_WITH_DUAL_SOURCE_COLOR_BLENDING (MATERIAL_SHADINGMODEL_THIN_TRANSLUCENT || SUBSTRATE_BLENDING_TRANSLUCENT_COLOREDTRANSMITTANCE )
+
+#if MATERIAL_SHADINGMODEL_THIN_TRANSLUCENT || SUBSTRATE_BLENDING_TRANSLUCENT_COLOREDTRANSMITTANCE
+	#define MOBILE_TRANSLUCENT_COLOR_TRANSMITTANCE 1
+	
+#else
+	#define MOBILE_TRANSLUCENT_COLOR_TRANSMITTANCE 0
+	#undef MOBILE_TRANSLUCENT_COLOR_TRANSMITTANCE_DUAL_SRC_BLENDING
+	#undef MOBILE_TRANSLUCENT_COLOR_TRANSMITTANCE_PROGRAMMABLE_BLENDING
+	#undef MOBILE_TRANSLUCENT_COLOR_TRANSMITTANCE_SINGLE_SRC_BLENDING
+	#define MOBILE_TRANSLUCENT_COLOR_TRANSMITTANCE_DUAL_SRC_BLENDING 0
+	#define MOBILE_TRANSLUCENT_COLOR_TRANSMITTANCE_PROGRAMMABLE_BLENDING 0
+	#define MOBILE_TRANSLUCENT_COLOR_TRANSMITTANCE_SINGLE_SRC_BLENDING 0
+#endif
 
 // Force FORWARD_SHADING on mobile when MOBILE_DEFERRED_SHADING is disabled
 #ifndef FORWARD_SHADING
@@ -1039,7 +1051,7 @@ void Main(
 #endif // MATERIAL_SHADINGMODEL_THIN_TRANSLUCENT
 	
 	// MATERIAL_SHADINGMODEL_THIN_TRANSLUCENT must come first because it also has MATERIALBLENDING_TRANSLUCENT defined
-	#if MATERIAL_WORKS_WITH_DUAL_SOURCE_COLOR_BLENDING
+	#if MOBILE_TRANSLUCENT_COLOR_TRANSMITTANCE
 
 		// Add fog luminance according to surfacecoverage and reduce surface luminance according to fog coverage.
 		half3 AdjustedDualBlendAdd = DualBlendSurfaceCoverage * VertexFog.rgb + VertexFog.a * DualBlendSurfaceLuminancePostCoverage;
