LOD Coloration fix for geometry collections not being colored correctly in the debug view

diff --git a/Engine/Shaders/Private/DebugViewModePixelShader.usf b/Engine/Shaders/Private/DebugViewModePixelShader.usf
--- a/Engine/Shaders/Private/DebugViewModePixelShader.usf
+++ b/Engine/Shaders/Private/DebugViewModePixelShader.usf
@@ -499,10 +499,10 @@ float4 ApplyDebugColor(in FDebugPSIn DebugInputs, in bool bIsFrontFace, in float
 		CalcMaterialParameters(MaterialParameters, PixelMaterialInputs, DebugInputs.SvPosition, bIsFrontFace);
 
 		// Sample material properties. The results are not used, but the calls to StoreTexCoordScale are still be made.
-		BaseColor = GetMaterialBaseColorRaw(PixelMaterialInputs);
-		Opacity = GetMaterialOpacityRaw(PixelMaterialInputs);
+		BaseColor = GetMaterialBaseColor(PixelMaterialInputs);
+		Opacity = GetMaterialOpacity(PixelMaterialInputs);
 
-		half3 Emissive = GetMaterialEmissiveRaw(PixelMaterialInputs);
+		half3 Emissive = GetMaterialEmissive(PixelMaterialInputs);
 #if MATERIALBLENDING_ADDITIVE
 		Opacity = (.05 + .95 * Luminance(Emissive));
 #else
