Fix missing shadow casting from groom when light is too close.

The RayEnd point was scaled to 75% of the light surface to aleviate shadow map issue. This is not necessary for voxel groom traversal.

diff --git a/Engine/Shaders/Private/VirtualShadowMaps/VirtualShadowMapProjectionSpot.ush b/Engine/Shaders/Private/VirtualShadowMaps/VirtualShadowMapProjectionSpot.ush
--- a/Engine/Shaders/Private/VirtualShadowMaps/VirtualShadowMapProjectionSpot.ush
+++ b/Engine/Shaders/Private/VirtualShadowMaps/VirtualShadowMapProjectionSpot.ush
@@ -577,7 +577,7 @@ bool GenerateRayLocalLight(
 		// where the projection becomes extreme.
 
 		// Clip ray to max angle to limit penumbra sizes
-		float ClampedLength = DistToLight * 0.75 * saturate(1.5 / (CosTheta + CotMaxRayAngleFromLight * SinTheta));
+		float ClampedLength = DistToLight * saturate(1.5 / (CosTheta + CotMaxRayAngleFromLight * SinTheta));
 
 		OutRayStart = TranslatedWorldPosition;
 		OutRayEnd = TranslatedWorldPosition + Dir * ClampedLength;
