Path Tracer: Revist payload encoding of HDR colors

Change from the rgb 11,11,10 float format which is limited to ~65000 to a shared exponent format which has twice the range and better accuracy overall.

This avoid the need to special case the color encoding for unlit materials and makes regular materials benefit from the same expanded range which is better for consistency.

diff --git a/Engine/Shaders/Private/PackUnpack.ush b/Engine/Shaders/Private/PackUnpack.ush
--- a/Engine/Shaders/Private/PackUnpack.ush
+++ b/Engine/Shaders/Private/PackUnpack.ush
@@ -246,6 +246,68 @@ float3 UnpackR10G10B10F(uint rgb)
 	return float3(r, g, b);
 }
 
+//
+// float3 <-> uint
+//  using a RGBE encoding with 9,9,8 bits for RGB and 6 bits for the shared exponent
+// See a discussion of such encodings here: 
+//   https://cbloomrants.blogspot.com/2020/06/widespread-error-in-radiance-hdr-rgbe.html
+//   https://cbloomrants.blogspot.com/2020/06/followup-tidbits-on-rgbe.html 
+// Because we are not using these functions in the context of Radiance .hdr files,
+// we can deviate from the classic 888e8 format and apply the corrections suggested in
+// the post above, while changing the bit depths to get a bit more color resolution at
+// the expense of less dynamic range (exponent of 8 being excessive for colors anyway)
+// The RGBE_plus encoding suggested in the second blog post is slightly better in some
+// cases, but fairly close to the version below. This code has the advantage of being
+// slightly easier to decode.
+
+uint PackRGB998E6(float3 rgb)
+{
+	const float MAX_C = float((1 << 8) - 1) * asfloat((63 - 32 + 127) << 23);
+	uint3 cbits = asuint(min(rgb, MAX_C));
+	int3  exp = int3(cbits >> 23) - 127;
+	uint3 man = (cbits & ((1 << 23) - 1)) | (1 << 23);
+
+	int max_exp = max3(exp.x, exp.y, exp.z);
+
+	// change from [1,2) to [0.5,1)
+	max_exp++;
+	max_exp = max(max_exp, -32);
+
+	int3 shift = min(max_exp - exp - int3(9, 9, 8) + 23, 31);
+	// center on encode
+	man += 1u << (shift - 1);
+	man >>= shift;
+	if ((man.x >= (1 << 9)) ||
+		(man.y >= (1 << 9)) ||
+		(man.z >= (1 << 8)))
+	{
+		// overflow due to rounding, bump exponent so that color will fit
+		max_exp++;
+		man >>= 1;
+	}
+	return
+		(uint(max_exp + 32) << (9 + 9 + 8)) |
+		(uint(man.x   	  ) << (9 + 8    )) |
+		(uint(man.y   	  ) << (8        )) |
+		(uint(man.z       ));
+}
+
+float3 UnpackRGB998E6(uint bits)
+{
+	int e = ((bits >> (9 + 9 + 8)) & ((1 << 6) - 1)) - 32;
+	int r = ((bits >> (9 + 8    )) & ((1 << 9) - 1));
+	int g = ((bits >> (8        )) & ((1 << 9) - 1));
+	int b = ((bits               ) & ((1 << 8) - 1));
+	float3 f = {
+		asfloat(uint(e + 127 - 9) << 23),
+		asfloat(uint(e + 127 - 9) << 23),
+		asfloat(uint(e + 127 - 8) << 23)
+	};
+	float3 result = { float(r), float(g), float(b) };
+	return result * f;
+}
+
+
 //
 //	float4 <-> uint2 as f16.f16.f16.f16
 //
diff --git a/Engine/Shaders/Private/PathTracing/PathTracingCommon.ush b/Engine/Shaders/Private/PathTracing/PathTracingCommon.ush
--- a/Engine/Shaders/Private/PathTracing/PathTracingCommon.ush
+++ b/Engine/Shaders/Private/PathTracing/PathTracingCommon.ush
@@ -70,14 +70,12 @@ float3 PayloadDecodeUnitVector(uint Encoded)
 
 uint PayloadEncodeHDRColor(float3 rgb)
 {
-	// clamp to the largest representable value
-	rgb = clamp(rgb, 0.0, Max111110BitsFloat3);
-	return PackR11G11B10F(rgb);
+	return PackRGB998E6(rgb);
 }
 
 float3 PayloadDecodeHDRColor(uint rgb)
 {
-	return UnpackR11G11B10F(rgb);
+	return UnpackRGB998E6(rgb);
 }
 
 uint PayloadEncodeLDRColor(float3 rgb)
@@ -236,23 +234,10 @@ struct FPackedPathTracingPayload : FMinimalPayload
 	float3 UnpackWorldGeoNormal() 	 { return PayloadDecodeUnitVector(PackedData[1]); }
 	float3 UnpackWorldSmoothNormal() { return PayloadDecodeUnitVector(PackedData[2]); }
 	float3 UnpackWorldNormal()       { return PayloadDecodeUnitVector(PackedData[3]); }
-	float3 UnpackRadiance()          { return (GetShadingModelID() == SHADINGMODELID_UNLIT) ? float3(asfloat(PackedData[4]), asfloat(PackedData[7]), asfloat(PackedData[8])) : PayloadDecodeHDRColor(PackedData[4]); }
+	float3 UnpackRadiance()          { return PayloadDecodeHDRColor(PackedData[4]); }
 
 	void PackWorldNormal(float3 N)   { PackedData[3] = PayloadEncodeUnitVector(N); }
-	void PackRadiance(float3 R)      {
-		if (GetShadingModelID() == SHADINGMODELID_UNLIT)
-		{
-			// encode with extra accuracy for UNLIT materials since these are also used for directly visible lights
-			PackedData[4] = asuint(R.x);
-			PackedData[7] = asuint(R.y);
-			PackedData[8] = asuint(R.z);
-		}
-		else
-		{
-			PackedData[4] = PayloadEncodeHDRColor(R);
-		}
-	}
-
+	void PackRadiance(float3 R)      { PackedData[4] = PayloadEncodeHDRColor(R); }
 	
 #if PATHTRACING_SUBSTRATE_PAYLOAD
 
@@ -757,6 +742,7 @@ FPackedPathTracingPayload PackPathTracingPayload(FPathTracingPayload Input)
 	Output.PackedData[1] = PayloadEncodeUnitVector(Input.WorldGeoNormal);
 	Output.PackedData[2] = PayloadEncodeUnitVector(Input.WorldSmoothNormal);
 	Output.PackedData[3] = PayloadEncodeUnitVector(Input.WorldNormal);
+	Output.PackedData[4] = PayloadEncodeHDRColor(Input.Radiance);
 	// PackedData[4] is for Radiance (see PackRadiance call below)
 #if PATHTRACING_SUBSTRATE_PAYLOAD
 	Output.PackedData[5] = PayloadEncodeLDRColor(Input.DiffuseColor);
@@ -791,7 +777,6 @@ FPackedPathTracingPayload PackPathTracingPayload(FPathTracingPayload Input)
 #else // PATHTRACING_SUBSTRATE_PAYLOAD
 
 
-
 	Output.PackedData[5] = PayloadEncodeLDRColor(Input.BaseColor);
 	Output.PackedData[6] = PayloadEncodeLDRColor(Input.TransparencyColor);
 
@@ -815,8 +800,7 @@ FPackedPathTracingPayload PackPathTracingPayload(FPathTracingPayload Input)
 	Output.PackedData[14] = f32tof16(Input.Ior);
 	// NOTE: 16 bits left in PackedData[14]
 #endif // PATHTRACING_SUBSTRATE_PAYLOAD
-	// handle this at the end since Radiance could be aliased with other fields
-	Output.PackRadiance(Input.Radiance);
+
 #endif // IS_PAYLOAD_ENABLED(RT_PAYLOAD_TYPE_GPULIGHTMASS)
 	return Output;
 }
