Remove PostionNoOffsets interpolator from a landscape vertex factory, since landscape does not support WPO
Added VF_SUPPORTS_WORLD_POSITION_SHADER_OFFSETS define which can be set to 0 by other VFs if they do not support WPO

diff --git a/Engine/Shaders/Private/Common.ush b/Engine/Shaders/Private/Common.ush
--- a/Engine/Shaders/Private/Common.ush
+++ b/Engine/Shaders/Private/Common.ush
@@ -2049,10 +2049,6 @@ struct FWriteToSliceGeometryOutput
 	uint LayerIndex : SV_RenderTargetArrayIndex;
 };
 
-// Helper macro to globally ignore requests for non-offset world positions in materials when lower than shader model 4. We do this
-// because we are using an extra interpolator for this second world position, and in < ES31 there may not be enough
-// NOTE: Forced on in development shaders because the Out-Of-Bounds Pixels visualizer uses it
-#define USE_WORLD_POSITION_EXCLUDING_SHADER_OFFSETS	((NEEDS_WORLD_POSITION_EXCLUDING_SHADER_OFFSETS || USE_DEVELOPMENT_SHADERS) && FEATURE_LEVEL >= FEATURE_LEVEL_ES3_1)
 
 /** Used for calculating vertex positions and UVs when drawing with DrawRectangle */
 void DrawRectangle(
diff --git a/Engine/Shaders/Private/MaterialTemplate.ush b/Engine/Shaders/Private/MaterialTemplate.ush
--- a/Engine/Shaders/Private/MaterialTemplate.ush
+++ b/Engine/Shaders/Private/MaterialTemplate.ush
@@ -4121,6 +4121,14 @@ half GetFloatFacingSign(FIsFrontFace bIsFrontFace)
 	#define OPTIONAL_IsFrontFace , in FIsFrontFace bIsFrontFace : FrontFaceSemantic
 #endif
 
+#ifndef VF_SUPPORTS_WORLD_POSITION_SHADER_OFFSETS
+#define VF_SUPPORTS_WORLD_POSITION_SHADER_OFFSETS 1
+#endif 
+// Helper macro to globally ignore requests for non-offset world positions in materials when lower than shader model 4. We do this
+// because we are using an extra interpolator for this second world position, and in < ES31 there may not be enough
+// NOTE: Forced on in development shaders because the Out-Of-Bounds Pixels visualizer uses it
+#define USE_WORLD_POSITION_EXCLUDING_SHADER_OFFSETS	((NEEDS_WORLD_POSITION_EXCLUDING_SHADER_OFFSETS || USE_DEVELOPMENT_SHADERS) && VF_SUPPORTS_WORLD_POSITION_SHADER_OFFSETS)
+
 /** Initializes the subset of Parameters that was not set in GetMaterialPixelParameters. */
 void CalcMaterialParametersEx(
 	in out FMaterialPixelParameters Parameters,
diff --git a/Engine/Source/Runtime/Landscape/Private/LandscapeRender.cpp b/Engine/Source/Runtime/Landscape/Private/LandscapeRender.cpp
--- a/Engine/Source/Runtime/Landscape/Private/LandscapeRender.cpp
+++ b/Engine/Source/Runtime/Landscape/Private/LandscapeRender.cpp
@@ -3616,6 +3616,8 @@ void FLandscapeVertexFactory::ModifyCompilationEnvironment(const FVertexFactoryS
 	FVertexFactory::ModifyCompilationEnvironment(Parameters, OutEnvironment);
 
 	OutEnvironment.SetDefine(TEXT("VF_SUPPORTS_PRIMITIVE_SCENE_DATA"), Parameters.VertexFactoryType->SupportsPrimitiveIdStream() && UseGPUScene(Parameters.Platform, GetMaxSupportedFeatureLevel(Parameters.Platform)));
+	// landscape does not support WPO, so we can remove one interpolator
+	OutEnvironment.SetDefine(TEXT("VF_SUPPORTS_WORLD_POSITION_SHADER_OFFSETS"), TEXT("0"));
 
 	// Make sure landscape vertices go back to local space so that we have consistency between the transform on normals and geometry
 	OutEnvironment.SetDefine(TEXT("RAY_TRACING_DYNAMIC_MESH_IN_LOCAL_SPACE"), TEXT("1"));
