Forward UI's opacity in HDR composition mode as DLSSG needs it. Merged from github PR 11616

#11616

diff --git a/Engine/Shaders/Private/CompositeUIPixelShader.usf b/Engine/Shaders/Private/CompositeUIPixelShader.usf
--- a/Engine/Shaders/Private/CompositeUIPixelShader.usf
+++ b/Engine/Shaders/Private/CompositeUIPixelShader.usf
@@ -258,7 +258,8 @@ void Main(
 
 	OutColor.rgb = ApplyColorDeficiency(OutColor.rgb);
 
-	OutColor.w = 1.f;
+	// This is not compatible when alpha is preserved through post-processing with POST_PROCESS_ALPHA
+	OutColor.w = UIColor.w;
 
 #if !SCRGB_ENCODING
 	// Linear -> PQ
@@ -310,7 +311,9 @@ void CompositeUICS(uint3 DispatchThreadId : SV_DispatchThreadID)
 	SceneColor = ComputeHDRSceneColor(SceneColor);
 	float4 OutColor = 0;
 	OutColor.xyz = ComposeUIAndScene(SceneColor.xyz, UIColor, UILevel);
-	OutColor.w = 1.f;
+	
+	// This is not compatible when alpha is preserved through post-processing with POST_PROCESS_ALPHA
+	OutColor.w = UIColor.w;
 
 	OutColor.rgb = ApplyColorDeficiency(OutColor.rgb);
 
@@ -322,7 +325,7 @@ void CompositeUICS(uint3 DispatchThreadId : SV_DispatchThreadID)
 	OutColor.xyz = mul(Rec2020_2_sRGB, OutColor.xyz / ScRGBScaleFactor);
 #endif
 
-	RWSceneTexture[DispatchThreadId.xy] = float4(OutColor.xyz, 1);
+	RWSceneTexture[DispatchThreadId.xy] = OutColor;
 }
 
 #endif
