BinarySearch.ush - Fixed incorrect comparison in LowerBound(), causing the binary search to return invalid results
* This would break Nanite primitive filtering (ShowOnlyPrimitives / HiddenPrimitives)

diff --git a/Engine/Shaders/Private/BinarySearch.ush b/Engine/Shaders/Private/BinarySearch.ush
--- a/Engine/Shaders/Private/BinarySearch.ush
+++ b/Engine/Shaders/Private/BinarySearch.ush
@@ -20,7 +20,7 @@ T LowerBound(TBuffer SearchBuffer, uint BufferOffset, uint BufferCount, T Key)
 	
 	while (Width > 0)
 	{
-		Index += Comp.Compare(Key, SearchBuffer[Index + Width]) > 0 ? 0 : Width;
+		Index += Comp.Compare(Key, SearchBuffer[Index + Width]) < 0 ? 0 : Width;
 		Width = Width >> 1;
 	}
 	
