Fix Substrate with Nanite compute material.

When running base pass, Substrate output data through MRTs and additionally to UAV when the data doesn't fit into the available MRTs. With compute material, a lane can acts as an helper lane, running the pixel shader logic, but not writing output their data. This was correctly handled for the MRT output part, but the UAV writes were still issued. This CL fixes this by ensuring that helper lanes mask out their UAV writes.

diff --git a/Engine/Shaders/Private/BasePassPixelShader.usf b/Engine/Shaders/Private/BasePassPixelShader.usf
--- a/Engine/Shaders/Private/BasePassPixelShader.usf
+++ b/Engine/Shaders/Private/BasePassPixelShader.usf
@@ -798,7 +798,8 @@ void FPixelShaderInOut_MainPS(
 	FVertexFactoryInterpolantsVSToPS Interpolants,
 	FBasePassInterpolantsVSToPS BasePassInterpolants,
 	in FPixelShaderIn In,
-	inout FPixelShaderOut Out)
+	inout FPixelShaderOut Out,
+	bool bAllowWrite=true)
 {
 #if INSTANCED_STEREO
 	const uint EyeIndex = Interpolants.EyeIndex;
@@ -1696,7 +1697,7 @@ void FPixelShaderInOut_MainPS(
 
 		// Generate the Substrate material data to write out
 		FSubstrateAddressing SubstrateAddressing = GetSubstratePixelDataByteOffset(PixelPos, uint2(ResolvedView.BufferSizeAndInvSize.xy), SubstrateStruct.MaxBytesPerPixel);
-		FRWSubstrateMaterialContainer RWSubstrateMaterialContainer = InitialiseRWSubstrateMaterialContainer(SubstrateStruct.MaterialTextureArrayUAVWithoutRTs);
+		FRWSubstrateMaterialContainer RWSubstrateMaterialContainer = InitialiseRWSubstrateMaterialContainer(SubstrateStruct.MaterialTextureArrayUAVWithoutRTs, bAllowWrite);
 		PackSubstrateOut(
 			RWSubstrateMaterialContainer,
 			SubstrateStruct.MaterialTextureArrayUAVWithoutRTs,
@@ -1724,12 +1725,15 @@ void FPixelShaderInOut_MainPS(
 		Out.SubstrateTopLayerData = SubstratePackTopLayerData(TopLayerData);
 	#if SUBSTRATE_OUTPUT_ROUGH_REFRACTION
 		OpaqueRoughRefractionData.OpaqueRoughRefractionEnabled = true; // As long as we output to rough refraction, we want it to be enabled for the tiling process.
-		SubstrateStruct.OpaqueRoughRefractionTextureUAV[PixelPos] = SubstratePackOpaqueRoughRefractionData(OpaqueRoughRefractionData);
+		if (bAllowWrite)
+		{
+			SubstrateStruct.OpaqueRoughRefractionTextureUAV[PixelPos] = SubstratePackOpaqueRoughRefractionData(OpaqueRoughRefractionData);
+		}
 	#endif
 
 		// Only write SSS data if needed
 		BRANCH
-		if(SubstrateSubSurfaceHeaderGetIsValid(SSSData.Header))
+		if(SubstrateSubSurfaceHeaderGetIsValid(SSSData.Header) && bAllowWrite)
 		{
 			SubstrateStoreSubsurfaceHeader(SubstrateStruct.MaterialTextureArrayUAVWithoutRTs, SubstrateStruct.FirstSliceStoringSubstrateSSSDataWithoutMRT, PixelPos, SSSData.Header.Bytes);
 			if (SubstrateSubSurfaceHeaderHasExtras(SSSData.Header))
diff --git a/Engine/Shaders/Private/ComputeShaderOutputCommon.ush b/Engine/Shaders/Private/ComputeShaderOutputCommon.ush
--- a/Engine/Shaders/Private/ComputeShaderOutputCommon.ush
+++ b/Engine/Shaders/Private/ComputeShaderOutputCommon.ush
@@ -115,7 +115,7 @@ uint4 PassData;
 
 #endif // SUBSTRATE_OPAQUE_DEFERRED
 
-FPixelShaderOut ShadePixel(const float2 SVPositionXY, uint QuadIndex)
+FPixelShaderOut ShadePixel(const float2 SVPositionXY, uint QuadIndex, uint PixelWriteMask)
 {
 #if PIXELSHADEROUTPUT_INTERPOLANTS || PIXELSHADEROUTPUT_BASEPASS
 #if IS_NANITE_PASS
@@ -148,7 +148,7 @@ FPixelShaderOut ShadePixel(const float2 SVPositionXY, uint QuadIndex)
 
 #if PIXELSHADEROUTPUT_BASEPASS
 	FBasePassInterpolantsVSToPS BasePassInterpolants = (FBasePassInterpolantsVSToPS)0;
-	FPixelShaderInOut_MainPS(Interpolants, BasePassInterpolants, PixelShaderIn, PixelShaderOut);
+	FPixelShaderInOut_MainPS(Interpolants, BasePassInterpolants, PixelShaderIn, PixelShaderOut, PixelWriteMask != 0);
 #endif
 
 #if !HIGH_PRECISION_GBUFFERS
@@ -251,7 +251,7 @@ void ExportPixel(const uint2 PixelPos, FPixelShaderOut ShadedPixel)
 void ProcessPixel(uint ShadingBin, const uint2 PixelPos, const float2 SVPositionXY, uint QuadIndex, uint DispatchIndex, uint PixelWriteMask, uint HelperLaneCount)
 {
 	// All lanes shade (regardless of export - so ddx/ddy are valid)
-	FPixelShaderOut ShadedPixel = ShadePixel(SVPositionXY, QuadIndex);
+	FPixelShaderOut ShadedPixel = ShadePixel(SVPositionXY, QuadIndex, PixelWriteMask);
 
 #if VIS_HELPER_LANES
 	ShadedPixel.MRT[3].rgb = ColorMapTurbo(float(HelperLaneCount) / 4.0f);
diff --git a/Engine/Shaders/Private/DebugProbes.usf b/Engine/Shaders/Private/DebugProbes.usf
--- a/Engine/Shaders/Private/DebugProbes.usf
+++ b/Engine/Shaders/Private/DebugProbes.usf
@@ -161,7 +161,7 @@ void WriteOutMaterial(
 	const FSubstrateIntegrationSettings Settings = InitSubstrateIntegrationSettings(false /*bForceFullyRough*/, bRoughDiffuse, 0/*PeelLayersAboveDepth*/, false/*bRoughnessTracking*/);
 
 	FSubstrateAddressing SubstrateAddressing = GetSubstratePixelDataByteOffset(PixelCoord, uint2(ResolvedView.BufferSizeAndInvSize.xy), MaxBytesPerPixel);
-	FRWSubstrateMaterialContainer RWSubstrateMaterialContainer = InitialiseRWSubstrateMaterialContainer(MaterialTextureArrayUAV);
+	FRWSubstrateMaterialContainer RWSubstrateMaterialContainer = InitialiseRWSubstrateMaterialContainer(MaterialTextureArrayUAV, true /*bAllowWrite*/);
 	PackSubstrateOut(
 		RWSubstrateMaterialContainer,
 		MaterialTextureArrayUAV,
diff --git a/Engine/Shaders/Private/Substrate/Substrate.ush b/Engine/Shaders/Private/Substrate/Substrate.ush
--- a/Engine/Shaders/Private/Substrate/Substrate.ush
+++ b/Engine/Shaders/Private/Substrate/Substrate.ush
@@ -1760,13 +1760,14 @@ FSubstrateAddressing GetSubstratePixelDataByteOffset(uint2 PixelPos, uint2 ViewB
 struct FRWSubstrateMaterialContainerStruct
 {
 	uint MaterialRenderTargets[SUBSTRATE_BASE_PASS_MRT_OUTPUT_COUNT];
+	bool bAllowWrite;
 //	RWTexture2DArray<uint> ExtraMaterialDataUAV; // If this is used, we cannot run shaders on debug...
 };
 
 #if SUBSTRATE_MATERIALCONTAINER_IS_VIEWRESOURCE
 
 #define FRWSubstrateMaterialContainer	FRWSubstrateMaterialContainerStruct
-FRWSubstrateMaterialContainer InitialiseRWSubstrateMaterialContainer(RWTexture2DArray<uint> ExtraMaterialDataUAV)
+FRWSubstrateMaterialContainer InitialiseRWSubstrateMaterialContainer(RWTexture2DArray<uint> ExtraMaterialDataUAV, bool bInAllowWrite)
 {
 	FRWSubstrateMaterialContainer Out;
 	UNROLL
@@ -1774,6 +1775,7 @@ FRWSubstrateMaterialContainer InitialiseRWSubstrateMaterialContainer(RWTexture2D
 	{
 		Out.MaterialRenderTargets[i] = 0;
 	}
+	Out.bAllowWrite = bInAllowWrite;
 //	Out.ExtraMaterialDataUAV = ExtraMaterialDataUAV;
 	return Out;
 }
@@ -1791,7 +1793,7 @@ void SubstrateStoreUint1(inout FRWSubstrateMaterialContainer SubstrateBuffer, RW
 		#error Substrate SUBSTRATE_BASE_PASS_MRT_OUTPUT_COUNT has been update but not SubstrateStore function
 		#endif
 	}
-	else
+	else if (SubstrateBuffer.bAllowWrite) // Only write data to UAV if pixel is allowed to write (i.e., it is not an helper lane/pixel)
 	{
 		ExtraMaterialDataUAV[uint3(SubstrateAddressing.PixelCoords, SubstrateAddressing.CurrentIndex - SUBSTRATE_BASE_PASS_MRT_OUTPUT_COUNT)] = Data;
 	}
