Fixes to mobile multi view shaders.

INSTANCED_STEREO is currently used for both ISR and mobile multi view fallback but the two cases require different handling. INSTANCED_STEREO && MOBILE_MULTI_VIEW marks the mobile multi view fallback path used in editor d3d.
Also making EyeIndex consistently use VIEW_ID semantic instead of using TEXCOORD.

diff --git a/Engine/Shaders/Private/MobileBasePassVertexShader.usf b/Engine/Shaders/Private/MobileBasePassVertexShader.usf
--- a/Engine/Shaders/Private/MobileBasePassVertexShader.usf
+++ b/Engine/Shaders/Private/MobileBasePassVertexShader.usf
@@ -38,14 +38,14 @@ struct FMobileShadingBasePassVSToPS
 void Main(
 	FVertexFactoryInput Input
 	, out FMobileShadingBasePassVSOutput Output
-#if INSTANCED_STEREO
+#if INSTANCED_STEREO && MOBILE_MULTI_VIEW
 	, out uint LayerIndex : SV_RenderTargetArrayIndex
 #elif MOBILE_MULTI_VIEW
 	, in uint ViewId : SV_ViewID
 #endif
 	)
 {
-#if INSTANCED_STEREO
+#if INSTANCED_STEREO && MOBILE_MULTI_VIEW
 	const uint EyeIndex = GetEyeIndexFromVF(Input);
 	ResolvedView = ResolveView(EyeIndex);
 	LayerIndex = EyeIndex;
diff --git a/Engine/Shaders/Private/PostProcessAmbientOcclusionMobile.usf b/Engine/Shaders/Private/PostProcessAmbientOcclusionMobile.usf
--- a/Engine/Shaders/Private/PostProcessAmbientOcclusionMobile.usf
+++ b/Engine/Shaders/Private/PostProcessAmbientOcclusionMobile.usf
@@ -844,7 +844,7 @@ void GTAOHorizonSearchIntegralCS(
 	uint2 GroupId : SV_GroupID,
 	uint2 DispatchThreadId : SV_DispatchThreadID,
 	uint2 GroupThreadId : SV_GroupThreadID
-#if INSTANCED_STEREO
+#if INSTANCED_STEREO && MOBILE_MULTI_VIEW
 	, uint InstanceId : SV_InstanceID
 	, out uint LayerIndex : SV_RenderTargetArrayIndex
 #elif MOBILE_MULTI_VIEW
@@ -853,7 +853,7 @@ void GTAOHorizonSearchIntegralCS(
 #endif
 )
 {
-#if INSTANCED_STEREO
+#if INSTANCED_STEREO && MOBILE_MULTI_VIEW
 	const uint EyeIndex = GetEyeIndex(InstanceId);
 	ResolvedView = ResolveView(EyeIndex);
 	LayerIndex = EyeIndex;
diff --git a/Engine/Shaders/Private/PostProcessPixelProjectedReflectionMobile.usf b/Engine/Shaders/Private/PostProcessPixelProjectedReflectionMobile.usf
--- a/Engine/Shaders/Private/PostProcessPixelProjectedReflectionMobile.usf
+++ b/Engine/Shaders/Private/PostProcessPixelProjectedReflectionMobile.usf
@@ -126,7 +126,7 @@ void ProjectionPassCS(
 	uint2 GroupId : SV_GroupID,
 	uint2 DispatchThreadId : SV_DispatchThreadID,
 	uint2 GroupThreadId : SV_GroupThreadID
-#if INSTANCED_STEREO
+#if INSTANCED_STEREO && MOBILE_MULTI_VIEW
 	, uint InstanceId : SV_InstanceID
 	, out uint LayerIndex : SV_RenderTargetArrayIndex
 #elif MOBILE_MULTI_VIEW
@@ -135,7 +135,7 @@ void ProjectionPassCS(
 #endif
 )
 {
-#if INSTANCED_STEREO
+#if INSTANCED_STEREO && MOBILE_MULTI_VIEW
 	const uint EyeIndex = GetEyeIndex(InstanceId);
 	ResolvedView = ResolveView(EyeIndex);
 	LayerIndex = EyeIndex;
@@ -187,22 +187,21 @@ void ReflectionPassVS(
 	in float2 InPosition : ATTRIBUTE0,
 	out float4 OutPosition : SV_POSITION,
 	out float4 PixelPosition : TEXCOORD0
-#if INSTANCED_STEREO
+#if INSTANCED_STEREO && MOBILE_MULTI_VIEW
 	, uint InstanceId : SV_InstanceID
 	, out uint LayerIndex : SV_RenderTargetArrayIndex
+	, out nointerpolation uint EyeIndex : VIEW_ID
 #elif MOBILE_MULTI_VIEW
 	, in uint ViewId : SV_ViewID
-	, out nointerpolation uint MultiViewIndex : VIEW_ID
 #endif
 )
 {
-#if INSTANCED_STEREO
-	const uint EyeIndex = GetEyeIndex(InstanceId);
+#if INSTANCED_STEREO && MOBILE_MULTI_VIEW
+	EyeIndex = GetEyeIndex(InstanceId);
 	ResolvedView = ResolveView(EyeIndex);
 	LayerIndex = EyeIndex;
 #elif MOBILE_MULTI_VIEW
 	ResolvedView = ResolveView(ViewId);
-	MultiViewIndex = ViewId;
 #else
 	ResolvedView = ResolveView();
 #endif
@@ -280,13 +279,17 @@ half4 DecodeReflectionColor(int2 ReflectingPixel, uint EncodeValue)
 void ReflectionPassPS(
 	in float4 SvPosition : SV_Position,
 	in float4 PixelPosition : TEXCOORD0,
-#if MOBILE_MULTI_VIEW
-	in nointerpolation uint MultiViewId : VIEW_ID,
+#if INSTANCED_STEREO && MOBILE_MULTI_VIEW // Mobile multi view fallback path
+	in nointerpolation uint EyeIndex : VIEW_ID,
+#elif MOBILE_MULTI_VIEW
+	in nointerpolation uint ViewId : SV_ViewID,
 #endif
 	out HALF4_TYPE OutColor : SV_Target0)
 {
-#if MOBILE_MULTI_VIEW
-	ResolvedView = ResolveView(MultiViewId);
+#if INSTANCED_STEREO && MOBILE_MULTI_VIEW
+	ResolvedView = ResolveView(EyeIndex);
+#elif MOBILE_MULTI_VIEW
+	ResolvedView = ResolveView(ViewId);
 #else
 	ResolvedView = ResolveView();
 #endif
diff --git a/Engine/Shaders/Private/PostProcessTonemap.usf b/Engine/Shaders/Private/PostProcessTonemap.usf
--- a/Engine/Shaders/Private/PostProcessTonemap.usf
+++ b/Engine/Shaders/Private/PostProcessTonemap.usf
@@ -676,11 +676,10 @@ half4 FetchAndResolveSceneColor(float2 UV, float ArrayIndex)
 void MobileCustomResolve_MainPS(
 	in float4 SvPosition							: SV_POSITION, 
 	in noperspective float2 UV						: TEXCOORD0,
-#if MOBILE_MULTI_VIEW_FALLBACK
-	in nointerpolation uint EyeIndex				: TEXCOORD1,
-#endif
 	out				 float4 OutColor				: SV_Target0
-#if USE_MULTI_VIEW_ID_SV
+#if MOBILE_MULTI_VIEW_FALLBACK
+	, in nointerpolation uint EyeIndex				: VIEW_ID
+#elif USE_MULTI_VIEW_ID_SV
 	, in nointerpolation uint MultiViewId : SV_ViewID
 #endif
 )
diff --git a/Engine/Shaders/Private/SceneTexturesCommon.ush b/Engine/Shaders/Private/SceneTexturesCommon.ush
--- a/Engine/Shaders/Private/SceneTexturesCommon.ush
+++ b/Engine/Shaders/Private/SceneTexturesCommon.ush
@@ -204,7 +204,7 @@ float LookupDeviceZ( float2 ScreenUV, float ArrayIndex )
 	return SCENE_TEXTURES_DISABLED_SCENE_DEPTH_VALUE;
 #elif FORCE_DEPTH_TEXTURE_READS || PLATFORM_NEEDS_DEPTH_TEXTURE_READS
 	// native Depth buffer lookup
-	#if INSTANCED_STEREO || MOBILE_MULTI_VIEW
+	#if (INSTANCED_STEREO && MOBILE_MULTI_VIEW) || MOBILE_MULTI_VIEW
 		return Texture2DArraySample(MobileSceneTextures.SceneDepthTextureArray, MobileSceneTextures.SceneDepthTextureSampler, float3(ScreenUV.xy,ArrayIndex)).r;
 	#else
 		return Texture2DSampleLevel(MobileSceneTextures.SceneDepthTexture, MobileSceneTextures.SceneDepthTextureSampler, ScreenUV, 0).r;
@@ -218,14 +218,14 @@ float LookupDeviceZ( float2 ScreenUV, float ArrayIndex )
 #elif (USE_SCENE_DEPTH_AUX && !MOBILE_DEFERRED_SHADING)
 	// SceneDepth texture is discarded after BasePass (with forward shading) 
 	// instead fetch DeviceZ from SceneDepthAuxTexture
-	#if INSTANCED_STEREO || MOBILE_MULTI_VIEW
+	#if (INSTANCED_STEREO && MOBILE_MULTI_VIEW) || MOBILE_MULTI_VIEW
 		return Texture2DArraySample(MobileSceneTextures.SceneDepthAuxTextureArray, MobileSceneTextures.SceneDepthAuxTextureSampler, float3(ScreenUV.xy,ArrayIndex)).r;
 	#else
 		return Texture2DSampleLevel(MobileSceneTextures.SceneDepthAuxTexture, MobileSceneTextures.SceneDepthAuxTextureSampler, ScreenUV, 0).r;
 	#endif
 #else
 	// native Depth buffer lookup
-	#if INSTANCED_STEREO || MOBILE_MULTI_VIEW
+	#if (INSTANCED_STEREO && MOBILE_MULTI_VIEW) || MOBILE_MULTI_VIEW
 		return Texture2DArraySample(MobileSceneTextures.SceneDepthTextureArray, MobileSceneTextures.SceneDepthTextureSampler, float3(ScreenUV.xy,ArrayIndex)).r;
 	#else
 		return Texture2DSampleLevel(MobileSceneTextures.SceneDepthTexture, MobileSceneTextures.SceneDepthTextureSampler, ScreenUV, 0).r;
@@ -246,7 +246,7 @@ float4 GatherDeviceZ( float2 ScreenUV )
 	return SCENE_TEXTURES_DISABLED_SCENE_DEPTH_VALUE;
 #elif FORCE_DEPTH_TEXTURE_READS || PLATFORM_NEEDS_DEPTH_TEXTURE_READS
 	// native Depth buffer lookup
-	#if INSTANCED_STEREO || MOBILE_MULTI_VIEW
+	#if (INSTANCED_STEREO && MOBILE_MULTI_VIEW) || MOBILE_MULTI_VIEW
 	return MobileSceneTextures.SceneDepthTextureArray.GatherRed(MobileSceneTextures.SceneDepthTextureSampler, float3(ScreenUV.xy,ArrayIndex));
 	#else
 		return MobileSceneTextures.SceneDepthTexture.GatherRed(MobileSceneTextures.SceneDepthTextureSampler, ScreenUV);
@@ -254,14 +254,14 @@ float4 GatherDeviceZ( float2 ScreenUV )
 #elif (USE_SCENE_DEPTH_AUX && !MOBILE_DEFERRED_SHADING)
 	// SceneDepth texture is discarded after BasePass (with forward shading) 
 	// instead fetch DeviceZ from SceneDepthAuxTexture
-	#if INSTANCED_STEREO || MOBILE_MULTI_VIEW
+	#if (INSTANCED_STEREO && MOBILE_MULTI_VIEW) || MOBILE_MULTI_VIEW
 		return MobileSceneTextures.SceneDepthAuxTextureArray.GatherRed(MobileSceneTextures.SceneDepthAuxTextureSampler, float3(ScreenUV.xy,ArrayIndex));
 	#else
 		return MobileSceneTextures.SceneDepthAuxTexture.GatherRed(MobileSceneTextures.SceneDepthAuxTextureSampler, ScreenUV);
 	#endif
 #else
 	// native Depth buffer lookup
-	#if INSTANCED_STEREO || MOBILE_MULTI_VIEW
+	#if (INSTANCED_STEREO && MOBILE_MULTI_VIEW) || MOBILE_MULTI_VIEW
 		return MobileSceneTextures.SceneDepthTextureArray.GatherRed(MobileSceneTextures.SceneDepthTextureSampler, float3(ScreenUV.xy,ArrayIndex));
 	#else
 		return MobileSceneTextures.SceneDepthTexture.GatherRed(MobileSceneTextures.SceneDepthTextureSampler, ScreenUV);
@@ -275,7 +275,7 @@ float CalcSceneDepth(float2 ScreenUV, float ArrayIndex)
 #if SCENE_TEXTURES_DISABLED
 	return SCENE_TEXTURES_DISABLED_SCENE_DEPTH_VALUE;
 #else
-	#if INSTANCED_STEREO || MOBILE_MULTI_VIEW
+	#if (INSTANCED_STEREO && MOBILE_MULTI_VIEW) || MOBILE_MULTI_VIEW
 			return ConvertFromDeviceZ(LookupDeviceZ(ScreenUV, ArrayIndex));
 	#else
 			return ConvertFromDeviceZ(LookupDeviceZ(ScreenUV));
diff --git a/Engine/Shaders/Private/ShadowProjectionPixelShader.usf b/Engine/Shaders/Private/ShadowProjectionPixelShader.usf
--- a/Engine/Shaders/Private/ShadowProjectionPixelShader.usf
+++ b/Engine/Shaders/Private/ShadowProjectionPixelShader.usf
@@ -104,8 +104,8 @@ EARLYDEPTHSTENCIL
 #endif
 void Main(
 	out float4 OutColor : SV_Target0
-#if INSTANCED_STEREO
-	, in nointerpolation uint EyeIndex : TEXCOORD0
+#if INSTANCED_STEREO && MOBILE_MULTI_VIEW
+	, in nointerpolation uint EyeIndex : VIEW_ID
 #elif MOBILE_MULTI_VIEW
 	, in nointerpolation uint ViewId : SV_ViewID
 #endif
@@ -117,19 +117,19 @@ void Main(
 	#endif
 	
 	float ArrayIndex = 0;
-#if INSTANCED_STEREO
+#if INSTANCED_STEREO && MOBILE_MULTI_VIEW
 	ResolvedView = ResolveView(EyeIndex);
 	ArrayIndex = EyeIndex;
 #elif MOBILE_MULTI_VIEW
 	ResolvedView = ResolveView(ViewId);
-	ArrayIndex = ViewId;;
+	ArrayIndex = ViewId;
 #else
 	ResolvedView = ResolveView();
 #endif
 
 	const FPQMPContext PQMPContext = PQMPInit(SVPos.xy);
 	float2 ScreenUV = float2(SVPos.xy * ResolvedView.BufferSizeAndInvSize.zw);
-#if (INSTANCED_STEREO || MOBILE_MULTI_VIEW) && SHADING_PATH_MOBILE
+#if ((INSTANCED_STEREO && MOBILE_MULTI_VIEW) || MOBILE_MULTI_VIEW) && SHADING_PATH_MOBILE
 	float SceneW = CalcSceneDepth(ScreenUV, ArrayIndex);
 #else
 	float SceneW = CalcSceneDepth(ScreenUV);
@@ -146,7 +146,7 @@ void Main(
 	
 	float4 ShadowPosition = float4(0, 0, 0, 0);
 	
-#if ((MOBILE_MULTI_VIEW || INSTANCED_STEREO) && SHADING_PATH_MOBILE)
+#if ((INSTANCED_STEREO && MOBILE_MULTI_VIEW) || MOBILE_MULTI_VIEW) && SHADING_PATH_MOBILE
 	ShadowPosition = mul(ScreenPosition, ResolvedView.MobileMultiviewShadowTransform);
 	ShadowPosition = mul(ShadowPosition, ScreenToShadowMatrix);
 #else
diff --git a/Engine/Shaders/Private/ShadowProjectionVertexShader.usf b/Engine/Shaders/Private/ShadowProjectionVertexShader.usf
--- a/Engine/Shaders/Private/ShadowProjectionVertexShader.usf
+++ b/Engine/Shaders/Private/ShadowProjectionVertexShader.usf
@@ -8,15 +8,15 @@
 
 void ShadowProjectionNoTransformVS(
 	in float4 InPosition : ATTRIBUTE0
-#if INSTANCED_STEREO
+#if INSTANCED_STEREO && MOBILE_MULTI_VIEW
 	, in uint InstanceId : SV_InstanceID
-	, out nointerpolation uint EyeIndex : TEXCOORD0
+	, out nointerpolation uint EyeIndex : VIEW_ID
 	, out nointerpolation uint LayerIndex : SV_RenderTargetArrayIndex
 #endif
 	, out float4 OutPosition : SV_POSITION
 )
 {
-#if INSTANCED_STEREO
+#if INSTANCED_STEREO && MOBILE_MULTI_VIEW
 	EyeIndex = GetEyeIndex(InstanceId);
 	LayerIndex = EyeIndex;
 #endif
@@ -30,9 +30,9 @@ float4 PreShadowToPreViewTranslation;
 
 void ShadowVolumeBoundProjectionVS(
 	in float4 InPosition : ATTRIBUTE0
-#if INSTANCED_STEREO
+#if INSTANCED_STEREO && MOBILE_MULTI_VIEW
 	, in uint InstanceId : SV_InstanceID
-	, out nointerpolation uint EyeIndex : TEXCOORD0
+	, out nointerpolation uint EyeIndex : VIEW_ID
 	, out uint LayerIndex : SV_RenderTargetArrayIndex
 #elif MOBILE_MULTI_VIEW
 	, in nointerpolation uint ViewId : SV_ViewID
@@ -40,7 +40,7 @@ void ShadowVolumeBoundProjectionVS(
 	, out float4 OutPosition : SV_POSITION
 )
 {
-#if INSTANCED_STEREO
+#if INSTANCED_STEREO && MOBILE_MULTI_VIEW
 	EyeIndex = GetEyeIndex(InstanceId);
 	LayerIndex = EyeIndex;
 	ResolvedView = ResolveView(EyeIndex);
diff --git a/Engine/Shaders/Private/Tools/FullscreenVertexShader.usf b/Engine/Shaders/Private/Tools/FullscreenVertexShader.usf
--- a/Engine/Shaders/Private/Tools/FullscreenVertexShader.usf
+++ b/Engine/Shaders/Private/Tools/FullscreenVertexShader.usf
@@ -29,7 +29,7 @@ void MainVS(
 	,	out noperspective float2 OutUV : TEXCOORD0
 	#if INSTANCED_STEREO // Mobile multi view fallback path
 	,	in uint InstanceId : SV_InstanceID
-	,	out nointerpolation uint EyeIndex : TEXCOORD1
+	,	out nointerpolation uint EyeIndex : VIEW_ID
 	,	out uint LayerIndex : SV_RenderTargetArrayIndex
 	#endif
 #endif
