Fix minor shader warnings

diff --git a/Engine/Shaders/Private/HeterogeneousVolumes/HeterogeneousVolumesLiveShadingShadows.usf b/Engine/Shaders/Private/HeterogeneousVolumes/HeterogeneousVolumesLiveShadingShadows.usf
--- a/Engine/Shaders/Private/HeterogeneousVolumes/HeterogeneousVolumesLiveShadingShadows.usf
+++ b/Engine/Shaders/Private/HeterogeneousVolumes/HeterogeneousVolumesLiveShadingShadows.usf
@@ -396,7 +396,7 @@ void RenderVolumetricShadowMapForLightWithLiveShadingCS(
 			bool bFirstEntry = (ElementIndex == 0);
 			if (bFirstEntry && any(Extinction > 0.0))
 			{
-				FAVSMSampleData SampleData = AVSM_CreateSampleData(WorldHitT, Transmittance);
+				FAVSMSampleData SampleData = AVSM_CreateSampleData(WorldHitT, Luminance(Transmittance));
 				AddSample(SampleData, PixelIndex, PrevElementIndex, ElementIndex);
 
 				// Set expected values
@@ -411,7 +411,7 @@ void RenderVolumetricShadowMapForLightWithLiveShadingCS(
 				float RelativeError = length((Transmittance - ExpectedTransmittance) / ExpectedTransmittance);
 				if ((AbsoluteError > AbsoluteErrorThreshold) && (RelativeError > RelativeErrorThreshold))
 				{
-					FAVSMSampleData SampleData = AVSM_CreateSampleData(WorldHitT, Transmittance);
+					FAVSMSampleData SampleData = AVSM_CreateSampleData(WorldHitT, Luminance(Transmittance));
 					AddOrReplaceSample(SampleData, PixelIndex, PrevElementIndex, ElementIndex);
 
 					// Reset expected values
diff --git a/Engine/Shaders/Private/HeterogeneousVolumes/HeterogeneousVolumesMaterialBakingPipeline.usf b/Engine/Shaders/Private/HeterogeneousVolumes/HeterogeneousVolumesMaterialBakingPipeline.usf
--- a/Engine/Shaders/Private/HeterogeneousVolumes/HeterogeneousVolumesMaterialBakingPipeline.usf
+++ b/Engine/Shaders/Private/HeterogeneousVolumes/HeterogeneousVolumesMaterialBakingPipeline.usf
@@ -328,7 +328,7 @@ void RasterizeBottomLevelOrthoGridCS(
 		}
 
 		float ExtinctionVariance = GSExtinctionVariance[0] * rcp(GSAllocatedVoxelCount - 1);
-		float StdRelError = sqrt(ExtinctionVariance * rcp(GSAllocatedVoxelCount)) * rcp(ExtinctionMean);
+		float StdRelError = sqrt(ExtinctionVariance * rcp(GSAllocatedVoxelCount)) * rcp(Luminance(ExtinctionMean));
 
 		// Declare low-variance to constitute a region of homogeneity
 		if (all(GroupThreadId == 0) && (StdRelError < HomogeneousThreshold) && (GSZeroCrossings[0] < 4))
@@ -540,7 +540,8 @@ void RasterizeBottomLevelFrustumGridCS(
 				InterlockedExchange(RWBottomLevelGridAllocatorBuffer[0], BottomLevelGridBufferSize, Dummy);
 
 				BottomLevelIndex = EMPTY_VOXEL_INDEX;
-				GSAllocatedVoxelCount = AllocatedVoxelResolution = 0;
+				GSAllocatedVoxelCount = 0; 
+				AllocatedVoxelResolution = 0;
 			}
 
 			// Update the index with the properly allocated index
@@ -551,7 +552,8 @@ void RasterizeBottomLevelFrustumGridCS(
 		}
 		else
 		{
-			GSAllocatedVoxelCount = AllocatedVoxelResolution = 0;
+			GSAllocatedVoxelCount = 0; 
+			AllocatedVoxelResolution = 0;
 		}
 	}
 
diff --git a/Engine/Shaders/Private/PathTracing/PathTracingMaterialHitShader.usf b/Engine/Shaders/Private/PathTracing/PathTracingMaterialHitShader.usf
--- a/Engine/Shaders/Private/PathTracing/PathTracingMaterialHitShader.usf
+++ b/Engine/Shaders/Private/PathTracing/PathTracingMaterialHitShader.usf
@@ -904,7 +904,7 @@ RAY_TRACING_ENTRY_CLOSEST_HIT(PathTracingMaterialCHS,
 
 		const float3 ScatteringCoeff = max(0.0f, DFDemote(GetSingleLayerWaterMaterialOutput0(WaterMaterialParameters)));
 		const float3 AbsorptionCoeff = max(0.0f, DFDemote(GetSingleLayerWaterMaterialOutput1(WaterMaterialParameters)));
-		const float PhaseG = clamp(DFDemote(GetSingleLayerWaterMaterialOutput2(WaterMaterialParameters)), -0.99f, 0.99f);
+		const float PhaseG = clamp(DFDemote(GetSingleLayerWaterMaterialOutput2(WaterMaterialParameters).x), -0.99f, 0.99f);
 
 		const float3 WaterExtinction = ScatteringCoeff + AbsorptionCoeff;
 		const float3 WaterAlbedo = select(WaterExtinction > 0.0, ScatteringCoeff / WaterExtinction, float3(0.0, 0.0, 0.0));
diff --git a/Engine/Shaders/Private/RayTracing/RayTracingMaterialHitShaders.usf b/Engine/Shaders/Private/RayTracing/RayTracingMaterialHitShaders.usf
--- a/Engine/Shaders/Private/RayTracing/RayTracingMaterialHitShaders.usf
+++ b/Engine/Shaders/Private/RayTracing/RayTracingMaterialHitShaders.usf
@@ -414,7 +414,7 @@ void GetMaterialPayload(
 	{
 		const float3 ScatteringCoeff = max(0.0f, DFDemote(GetSingleLayerWaterMaterialOutput0(MaterialParameters)));
 		const float3 AbsorptionCoeff = max(0.0f, DFDemote(GetSingleLayerWaterMaterialOutput1(MaterialParameters)));
-		const float PhaseG = clamp(DFDemote(GetSingleLayerWaterMaterialOutput2(MaterialParameters)), -1.0f, 1.0f);
+		const float PhaseG = clamp(DFDemote(GetSingleLayerWaterMaterialOutput2(MaterialParameters).x), -1.0f, 1.0f);
 
 		const float3 WaterExtinction = ScatteringCoeff + AbsorptionCoeff;
 		const float3 WaterAlbedo = ScatteringCoeff / WaterExtinction;
diff --git a/Engine/Shaders/Private/SingleLayerWaterShading.ush b/Engine/Shaders/Private/SingleLayerWaterShading.ush
--- a/Engine/Shaders/Private/SingleLayerWaterShading.ush
+++ b/Engine/Shaders/Private/SingleLayerWaterShading.ush
@@ -145,7 +145,7 @@ WaterVolumeLightingOutput EvaluateWaterVolumeLighting(
 #else
 		const float3 ScatteringCoeff = max(0.0f, DFDemote(GetSingleLayerWaterMaterialOutput0(MaterialParameters)));
 		const float3 AbsorptionCoeff = max(0.0f, DFDemote(GetSingleLayerWaterMaterialOutput1(MaterialParameters)));
-		const float PhaseG = clamp(DFDemote(GetSingleLayerWaterMaterialOutput2(MaterialParameters)), -1.0f, 1.0f);
+		const float PhaseG = clamp(DFDemote(GetSingleLayerWaterMaterialOutput2(MaterialParameters).x), -1.0f, 1.0f);
 		//Sample the optional Material Input ColorScaleBehindWater and fade it out at shorelines to avoid hard edge intersections
 		float3 ColorScaleBehindWater = lerp(1.0f, max(0.0f, DFDemote(GetSingleLayerWaterMaterialOutput3(MaterialParameters))), saturate(BehindWaterDeltaDepth * 0.02f));
 		const float3 ExtinctionCoeff = ScatteringCoeff + AbsorptionCoeff;
