Shaders: Silence warnings around generated code in the MaterialTemplate.ush header so that we can enable warnings in material shaders

Individual shader compiler implementations must macros that locally save and restore the set of enabled warnings as well as a macro that silences undesired warnings. For now this is only implemented for the dxc compiler.

#tests Ran CompileShadersTestBed on EngineTest and Lyra
#rnx

diff --git a/Engine/Shaders/Private/MaterialTemplate.ush b/Engine/Shaders/Private/MaterialTemplate.ush
--- a/Engine/Shaders/Private/MaterialTemplate.ush
+++ b/Engine/Shaders/Private/MaterialTemplate.ush
@@ -308,12 +308,17 @@ float4 GetDynamicParameter(FMaterialParticleParameters Parameters, float4 Defaul
 
 }
 
+SHADER_PUSH_WARNINGS_STATE
+SHADER_DISABLE_WARNINGS
+
 /** Material declarations */
 %{material_declarations}
 
 /** FMaterialAttributes utilities */
 %{material_attributes_utilities}
 
+SHADER_POP_WARNINGS_STATE
+
 /** 
  * Parameters calculated from the pixel material inputs.
  */
@@ -3300,8 +3305,13 @@ float4 NeuralBufferOutput(FMaterialPixelParameters Parameters, float4 BufferIndi
 #define SwizzleDeriv4(_V, _MASK) ConstructFloatDeriv4(_V.Value._MASK, _V.Ddx._MASK, _V.Ddy._MASK)
 
 // Uniform material expressions.
+SHADER_PUSH_WARNINGS_STATE
+SHADER_DISABLE_WARNINGS
+
 %{uniform_material_expressions}
 
+SHADER_POP_WARNINGS_STATE
+
 #if USE_ANALYTIC_DERIVATIVES && TEXTURE_SAMPLE_DEBUG
 MaterialFloat4 DebugTextureCommon(const int Mode, float2 UV, MaterialFloat2 DDX, MaterialFloat2 DDY, MaterialFloat Scale)
 {
@@ -3423,7 +3433,10 @@ half3 GetMaterialEmissive(FPixelMaterialInputs PixelMaterialInputs)
 
 half3 GetMaterialEmissiveForCS(FMaterialPixelParameters Parameters)
 {
+SHADER_PUSH_WARNINGS_STATE
+SHADER_DISABLE_WARNINGS
 %{get_material_emissive_for_cs};
+SHADER_POP_WARNINGS_STATE
 }
 
 // Shading Model is an uint and represents a SHADINGMODELID_* in ShadingCommon.ush 
@@ -3494,6 +3507,9 @@ half GetMaterialAnisotropy(FPixelMaterialInputs PixelMaterialInputs)
 	return clamp(GetMaterialAnisotropyRaw(PixelMaterialInputs), -1.0f, 1.0f);
 }
 
+SHADER_PUSH_WARNINGS_STATE
+SHADER_DISABLE_WARNINGS
+
 half GetMaterialTranslucencyDirectionalLightingIntensity()
 {
 %{get_material_translucency_directional_lighting_intensity};
@@ -3536,6 +3552,9 @@ half GetMaterialOpacityMaskClipValue()
 %{get_material_opacity_mask_clip_value};
 }
 
+SHADER_POP_WARNINGS_STATE
+
+
 #if USES_DISPLACEMENT && IS_NANITE_PASS
 
 float GetMaterialDisplacement(FPixelMaterialInputs PixelMaterialInputs)
@@ -3611,7 +3630,10 @@ float3 ClampWorldPositionOffset(FMaterialVertexParameters Parameters, float3 InO
 
 float3 GetMaterialWorldPositionOffsetRaw(FMaterialVertexParameters Parameters)
 {
+SHADER_PUSH_WARNINGS_STATE
+SHADER_DISABLE_WARNINGS
 %{get_material_world_position_offset_raw};
+SHADER_POP_WARNINGS_STATE
 }
 
 float3 GetMaterialWorldPositionOffset(FMaterialVertexParameters Parameters)
@@ -3626,7 +3648,10 @@ float3 GetMaterialWorldPositionOffset(FMaterialVertexParameters Parameters)
 
 float3 GetMaterialPreviousWorldPositionOffsetRaw(FMaterialVertexParameters Parameters)
 {
+SHADER_PUSH_WARNINGS_STATE
+SHADER_DISABLE_WARNINGS
 %{get_material_previous_world_position_offset_raw};
+SHADER_POP_WARNINGS_STATE
 }
 
 float3 GetMaterialPreviousWorldPositionOffset(FMaterialVertexParameters Parameters)
@@ -3652,6 +3677,9 @@ half4 GetMaterialSubsurfaceData(FPixelMaterialInputs PixelMaterialInputs)
 	return OutSubsurface;
 }
 
+SHADER_PUSH_WARNINGS_STATE
+SHADER_DISABLE_WARNINGS
+
 half GetMaterialCustomData0(in out FMaterialPixelParameters Parameters)
 {
 %{get_material_custom_data0};
@@ -3662,6 +3690,8 @@ half GetMaterialCustomData1(in out FMaterialPixelParameters Parameters)
 %{get_material_custom_data1};
 }
 
+SHADER_POP_WARNINGS_STATE
+
 half GetMaterialAmbientOcclusionRaw(FPixelMaterialInputs PixelMaterialInputs)
 {
 	return PixelMaterialInputs.AmbientOcclusion;
@@ -3714,6 +3744,9 @@ FMaterialRefractionData GetMaterialRefraction(FPixelMaterialInputs PixelMaterial
 }
 
 #if NUM_TEX_COORD_INTERPOLATORS
+
+SHADER_PUSH_WARNINGS_STATE
+SHADER_DISABLE_WARNINGS
 void GetMaterialCustomizedUVs(FMaterialVertexParameters Parameters, inout float2 OutTexCoords[NUM_TEX_COORD_INTERPOLATORS])
 {
 %{get_material_customized_u_vs}
@@ -3723,6 +3756,8 @@ void GetCustomInterpolators(FMaterialVertexParameters Parameters, inout float2 O
 {
 %{get_custom_interpolators}
 }
+SHADER_POP_WARNINGS_STATE
+
 #endif
 
 float GetMaterialPixelDepthOffset(FPixelMaterialInputs PixelMaterialInputs)
@@ -3802,9 +3837,14 @@ float GetRoughnessFromNormalCurvature(FMaterialPixelParameters InMaterialParamet
 #endif
 }
 
+SHADER_PUSH_WARNINGS_STATE
+SHADER_DISABLE_WARNINGS
+
 // EvaluateMaterialAttributes
 %{evaluate_material_attributes}
 
+SHADER_POP_WARNINGS_STATE
+
 // We must to flip the normal of two sided SLW material with normal specified in world space.
 #define TWO_SIDED_WORLD_SPACE_SINGLELAYERWATER_NORMAL (MATERIAL_SHADINGMODEL_SINGLELAYERWATER && !MATERIAL_TANGENTSPACENORMAL && MATERIAL_TWOSIDED)
 
@@ -3813,10 +3853,13 @@ float GetRoughnessFromNormalCurvature(FMaterialPixelParameters InMaterialParamet
 // derivatives. If you change anything, make sure that both functions stay identical. 
 void CalcPixelMaterialInputs(in out FMaterialPixelParameters Parameters, in out FPixelMaterialInputs PixelMaterialInputs)
 {
+SHADER_PUSH_WARNINGS_STATE
+SHADER_DISABLE_WARNINGS
 	// Initial calculations (required for Normal)
 %{calc_pixel_material_inputs_initial_calculations}
 	// The Normal is a special case as it might have its own expressions and also be used to calculate other inputs, so perform the assignment here
 %{calc_pixel_material_inputs_normal}
+SHADER_POP_WARNINGS_STATE
 
 #if TEMPLATE_USES_SUBSTRATE
 	Parameters.SubstratePixelFootprint = SubstrateGetPixelFootprint(Parameters.WorldPosition_CamRelative, GetRoughnessFromNormalCurvature(Parameters));
@@ -3859,8 +3902,11 @@ void CalcPixelMaterialInputs(in out FMaterialPixelParameters Parameters, in out
 	Parameters.Particle.MotionBlurFade = 1.0f;
 #endif // !PARTICLE_SPRITE_FACTORY
 
+SHADER_PUSH_WARNINGS_STATE
+SHADER_DISABLE_WARNINGS
 	// Now the rest of the inputs
 %{calc_pixel_material_inputs_other_inputs}
+SHADER_POP_WARNINGS_STATE
 
 #if MATERIAL_USES_ANISOTROPY
 	Parameters.WorldTangent = CalculateAnisotropyTangent(Parameters, PixelMaterialInputs);
@@ -3875,10 +3921,13 @@ void CalcPixelMaterialInputs(in out FMaterialPixelParameters Parameters, in out
 #define DERIV_BASE_VALUE(_X) _X.Value
 void CalcPixelMaterialInputsAnalyticDerivatives(in out FMaterialPixelParameters Parameters, in out FPixelMaterialInputs PixelMaterialInputs)
 {
+SHADER_PUSH_WARNINGS_STATE
+SHADER_DISABLE_WARNINGS
 	// Initial calculations (required for Normal)
 %{calc_pixel_material_inputs_analytic_derivatives_initial}
 	// The Normal is a special case as it might have its own expressions and also be used to calculate other inputs, so perform the assignment here
 %{calc_pixel_material_inputs_analytic_derivatives_normal}
+SHADER_POP_WARNINGS_STATE
 
 #if TEMPLATE_USES_SUBSTRATE
 	Parameters.SubstratePixelFootprint = SubstrateGetPixelFootprint(Parameters.WorldPosition_CamRelative, GetRoughnessFromNormalCurvature(Parameters));
@@ -3921,9 +3970,14 @@ void CalcPixelMaterialInputsAnalyticDerivatives(in out FMaterialPixelParameters
 	Parameters.Particle.MotionBlurFade = 1.0f;
 #endif // !PARTICLE_SPRITE_FACTORY
 
+SHADER_PUSH_WARNINGS_STATE
+SHADER_DISABLE_WARNINGS
+
 	// Now the rest of the inputs
 %{calc_pixel_material_inputs_analytic_derivatives_other_inputs}
 
+SHADER_POP_WARNINGS_STATE
+
 #if MATERIAL_USES_ANISOTROPY
 	Parameters.WorldTangent = CalculateAnisotropyTangent(Parameters, PixelMaterialInputs);
 #else
diff --git a/Engine/Shaders/Public/Platform.ush b/Engine/Shaders/Public/Platform.ush
--- a/Engine/Shaders/Public/Platform.ush
+++ b/Engine/Shaders/Public/Platform.ush
@@ -285,6 +285,33 @@
 #define COMPILER_NEEDS_DETERMINANT 0
 #endif
 
+#if COMPILER_DXC == 1
+
+#define SHADER_PUSH_WARNINGS_STATE  _Pragma("dxc diagnostic push")
+#define SHADER_POP_WARNINGS_STATE   _Pragma("dxc diagnostic pop")
+
+// DXC requires several pragmas because there isn't a warning group that covers all cases
+#define SHADER_DISABLE_WARNINGS 								\
+	_Pragma("dxc diagnostic ignored \"-Wall\"") 				\
+	_Pragma("dxc diagnostic ignored \"-Wconversion\"") 			\
+	_Pragma("dxc diagnostic ignored \"-Wfor-redefinition\"") 	\
+	_Pragma("dxc diagnostic ignored \"-Winline-asm\"")
+
+#endif
+
+// If the shader compiler does not suport disabling warnings locally, these macros will be empty
+#ifndef SHADER_DISABLE_WARNINGS
+#define SHADER_DISABLE_WARNINGS
+#endif
+
+#ifndef SHADER_PUSH_WARNINGS_STATE
+#define SHADER_PUSH_WARNINGS_STATE
+#endif
+
+#ifndef SHADER_POP_WARNINGS_STATE
+#define SHADER_POP_WARNINGS_STATE
+#endif
+
 // ---------------------------------------------------- Alternative floating point types
 
 #ifndef FORCE_FLOATS
