Added a new mode for pixel depth offset that respects the perepsective projection when appliying the offset to world space an depth buffer. Now the worldspace position from depth more stable.

diff --git a/Engine/Shaders/Private/MaterialTemplate.ush b/Engine/Shaders/Private/MaterialTemplate.ush
--- a/Engine/Shaders/Private/MaterialTemplate.ush
+++ b/Engine/Shaders/Private/MaterialTemplate.ush
@@ -42,7 +42,7 @@
 
 // Update this GUID to force all material shaders to recompile
 // Merge conflicts on this line should be resolved by generating a new GUID
-#pragma message("UESHADERMETADATA_VERSION 5AE81BAA-6F06-4443-AEC8-98E87D189D36")
+#pragma message("UESHADERMETADATA_VERSION 25A8AD39-1BAE-4581-904F-FA19DE92F37E")
 
 #if USES_SPEEDTREE
 	#include "/Engine/Private/SpeedTreeCommon.ush"
@@ -4352,6 +4352,10 @@ half3x3 AssembleTangentToWorld( half3 TangentToWorld0, half4 TangentToWorld2 )
 	#define OPTIONAL_OutDepth
 #endif
 
+// Must match EPixelDepthOffsetMode from EngineTypes.h
+#define PODM_LEGACY			0
+#define PODM_CAMERAVECTOR	1
+
 float ApplyPixelDepthOffsetToMaterialParameters(inout FMaterialPixelParameters MaterialParameters, FPixelMaterialInputs PixelMaterialInputs, out float OutDepth)
 {
 	float PixelDepthOffset = GetMaterialPixelDepthOffset(PixelMaterialInputs);
@@ -4380,9 +4384,33 @@ float ApplyPixelDepthOffsetToMaterialParameters(inout FMaterialPixelParameters M
 
 		DeviceDepth = ConvertToDeviceZ(DeviceDepth);
 		PixelDepthOffset = MaterialParameters.SvPosition.z - OriginalDeviceZ;
+
+		// In this Ortho case, CameraVector is always ViewForward.
+		MaterialParameters.LWCData.AbsoluteWorldPosition = WSAdd(MaterialParameters.LWCData.AbsoluteWorldPosition, -MaterialParameters.CameraVector * PixelDepthOffset);
 	}
 	else
 	{
+		// We clamp PixelDepthOffset to 0 in order to not move the pixel towards the camera.
+		const float3 CameraPixelForwardVector = -MaterialParameters.CameraVector;
+
+	#if PIXEL_DEPTH_OFFSET_MODE == PODM_LEGACY
+
+		// The legacy behavior applying a world space offset independently from the perspective projection.
+		const float3 AbsoluteWorldOffsetVector = CameraPixelForwardVector * max(0.0, PixelDepthOffset);
+
+	#elif PIXEL_DEPTH_OFFSET_MODE == PODM_CAMERAVECTOR
+
+		// The change of depth will be less at the edge of the screen due to the perspective.
+		const float PixelDepthOffetForPerspective = PixelDepthOffset * dot(CameraPixelForwardVector, ResolvedView.ViewForward);
+		// This is the world space vector to offset the pixel world space position by according to the perspsective.
+		// It clamps the offset to 0 in order to only push the pixel foward and match the depth clamping math below.
+		const float3 AbsoluteWorldOffsetVector = CameraPixelForwardVector * max(0.0, PixelDepthOffetForPerspective);
+		// Update the new pixel depth offset accounting for perspective projection.
+		PixelDepthOffset = PixelDepthOffetForPerspective;
+
+	#endif
+
+		// Compute the new DeviceDepth accounting for the PixelDepthOffset.
 		DeviceDepth = min(MaterialParameters.ScreenPosition.z / (MaterialParameters.ScreenPosition.w + PixelDepthOffset), MaterialParameters.SvPosition.z);
 
 		// Once we've computed our (clamped) device depth, recompute PixelDepthOffset again to take the potential clamp into account
@@ -4391,9 +4419,10 @@ float ApplyPixelDepthOffsetToMaterialParameters(inout FMaterialPixelParameters M
 		// Update positions used for shading
 		MaterialParameters.ScreenPosition.w += PixelDepthOffset;
 		MaterialParameters.SvPosition.w = MaterialParameters.ScreenPosition.w;
+
+		MaterialParameters.LWCData.AbsoluteWorldPosition = WSAdd(MaterialParameters.LWCData.AbsoluteWorldPosition, AbsoluteWorldOffsetVector);
 	}
 
-	MaterialParameters.LWCData.AbsoluteWorldPosition = WSAdd(MaterialParameters.LWCData.AbsoluteWorldPosition, -MaterialParameters.CameraVector * PixelDepthOffset);
 	MaterialParameters.AbsoluteWorldPosition = WSToDF(MaterialParameters.LWCData.AbsoluteWorldPosition);
 
 	OutDepth = INVARIANT(DeviceDepth);
diff --git a/Engine/Source/Editor/UnrealEd/Private/Commandlets/DumpMaterialInfo.cpp b/Engine/Source/Editor/UnrealEd/Private/Commandlets/DumpMaterialInfo.cpp
--- a/Engine/Source/Editor/UnrealEd/Private/Commandlets/DumpMaterialInfo.cpp
+++ b/Engine/Source/Editor/UnrealEd/Private/Commandlets/DumpMaterialInfo.cpp
@@ -294,6 +294,7 @@ TArray<MaterialInfo::FPropertySet> GetMaterialInfoProperties()
 		FPropertySet::FromAccessor<bool, &FMaterialResource::GetCastDynamicShadowAsMasked>(TEXT("GetCastDynamicShadowAsMasked")),
 		FPropertySet::FromAccessor<bool, &FMaterialResource::IsDistorted>(TEXT("IsDistorted")),
 		FPropertySet::FromEnumAccessor<ERefractionCoverageMode, &FMaterialResource::GetRefractionCoverageMode>(TEXT("GetRefractionCoverageMode")),
+		FPropertySet::FromEnumAccessor<EPixelDepthOffsetMode, &FMaterialResource::GetPixelDepthOffsetMode>(TEXT("GetPixelDepthOffsetMode")),
 		FPropertySet::FromAccessor<float, &FMaterialResource::GetTranslucencyDirectionalLightingIntensity>(TEXT("GetTranslucencyDirectionalLightingIntensity")),
 		FPropertySet::FromAccessor<float, &FMaterialResource::GetTranslucentShadowDensityScale>(TEXT("GetTranslucentShadowDensityScale")),
 		FPropertySet::FromAccessor<float, &FMaterialResource::GetTranslucentSelfShadowDensityScale>(TEXT("GetTranslucentSelfShadowDensityScale")),
diff --git a/Engine/Source/Editor/UnrealEd/Private/MaterialGraphNode_Root.cpp b/Engine/Source/Editor/UnrealEd/Private/MaterialGraphNode_Root.cpp
--- a/Engine/Source/Editor/UnrealEd/Private/MaterialGraphNode_Root.cpp
+++ b/Engine/Source/Editor/UnrealEd/Private/MaterialGraphNode_Root.cpp
@@ -202,7 +202,7 @@ void UMaterialGraphNode_Root::CreateInputPins()
 	
 		FName MaterialInputName = *MaterialInput.GetName().ToString();
 		FName PinSubCategory;
-		FString RefractionMethodStr;
+		FString TempStr;
 		switch (Property)
 		{
 			case MP_Metallic:
@@ -214,33 +214,47 @@ void UMaterialGraphNode_Root::CreateInputPins()
 			case MP_CustomData0:
 			case MP_CustomData1:
 			case MP_AmbientOcclusion:
-			case MP_PixelDepthOffset:
 			case MP_Displacement:
 			case MP_SurfaceThickness:
 				PinSubCategory = UMaterialGraphSchema::PSC_Red;
 				break;
+
+			case MP_PixelDepthOffset:
+				PinSubCategory = UMaterialGraphSchema::PSC_Red;
+				switch (MaterialGraph->Material->PixelDepthOffsetMode)
+				{
+				case EPixelDepthOffsetMode::PDOM_Legacy:
+					TempStr = TEXT("Legacy");
+					break;
+				case EPixelDepthOffsetMode::PDOM_AlongCameraVector:
+					TempStr = TEXT("Camera Vector");
+					break;
+				default:
+					TempStr = TEXT("UNKNOWN EPixelDepthOffsetMode");
+				}
+				MaterialInputName = *FString::Printf(TEXT("%s (%s)"), *MaterialInputName.ToString(), *TempStr);
+				break;
 				
 			case MP_Refraction:
 				PinSubCategory = UMaterialGraphSchema::PSC_Red;
 				switch (MaterialGraph->Material->RefractionMethod)
 				{
 				case ERefractionMode::RM_None:
-					RefractionMethodStr = TEXT("Disabled");
+					TempStr = TEXT("Disabled");
 					break;
 				case ERefractionMode::RM_IndexOfRefraction:
-					RefractionMethodStr = TEXT("Index Of Refraction");
+					TempStr = TEXT("Index Of Refraction");
 					break;
 				case ERefractionMode::RM_PixelNormalOffset:
-					RefractionMethodStr = TEXT("Pixel Normal Offset");
+					TempStr = TEXT("Pixel Normal Offset");
 					break;
 				case ERefractionMode::RM_2DOffset:
-					RefractionMethodStr = TEXT("2D Offset");
+					TempStr = TEXT("2D Offset");
 					break;
 				default:
-					RefractionMethodStr = TEXT("UNKNOWN ERefractionMode");
+					TempStr = TEXT("UNKNOWN ERefractionMode");
 				}
-				
-				MaterialInputName = *FString::Printf(TEXT("%s (%s)"), *MaterialInput.GetName().ToString(), *RefractionMethodStr);
+				MaterialInputName = *FString::Printf(TEXT("%s (%s)"), *MaterialInputName.ToString(), *TempStr);
 				break;
 				
 			case MP_Normal:
diff --git a/Engine/Source/Runtime/Core/Private/UObject/FortniteMainBranchObjectVersion.cpp b/Engine/Source/Runtime/Core/Private/UObject/FortniteMainBranchObjectVersion.cpp
--- a/Engine/Source/Runtime/Core/Private/UObject/FortniteMainBranchObjectVersion.cpp
+++ b/Engine/Source/Runtime/Core/Private/UObject/FortniteMainBranchObjectVersion.cpp
@@ -1,4 +1,4 @@
-﻿// Copyright Epic Games, Inc. All Rights Reserved.
+// Copyright Epic Games, Inc. All Rights Reserved.
 #include "UObject/FortniteMainBranchObjectVersion.h"
 
 TMap<FGuid, FGuid> FFortniteMainBranchObjectVersion::GetSystemGuids()
@@ -14,6 +14,6 @@ TMap<FGuid, FGuid> FFortniteMainBranchObjectVersion::GetSystemGuids()
 	SystemGuids.Add(DevGuids.Niagara_LatestScriptCompileVersion, FGuid("E267F600682A4E94885E8AEC5CE1DC91"));
 	SystemGuids.Add(DevGuids.SkeletalMeshDerivedDataVersion, FGuid("B9AEAA3EE7AC4FCFB3860E8F3E06DD75"));
 	SystemGuids.Add(DevGuids.STATICMESH_DERIVEDDATA_VER, FGuid("543843B103794E0E9BA4BE80FB602F79"));
-	SystemGuids.Add(DevGuids.MaterialTranslationDDCVersion, FGuid("D83409BD645A4B86AF0ECC671B0962C9"));
+	SystemGuids.Add(DevGuids.MaterialTranslationDDCVersion, FGuid("7A75688F394C4F27837B35000234B087"));
 	return SystemGuids;
 }
diff --git a/Engine/Source/Runtime/Core/Public/UObject/FortniteMainBranchObjectVersion.h b/Engine/Source/Runtime/Core/Public/UObject/FortniteMainBranchObjectVersion.h
--- a/Engine/Source/Runtime/Core/Public/UObject/FortniteMainBranchObjectVersion.h
+++ b/Engine/Source/Runtime/Core/Public/UObject/FortniteMainBranchObjectVersion.h
@@ -471,6 +471,9 @@ struct FFortniteMainBranchObjectVersion
 		// AnimNext graphs now have public/private state
 		AnimNextGraphAccessSpecifiers,
 
+		// Added a more stable pixel depth offset mode.
+		MaterialPixelDepthOffsetMode,
+
 		// -----<new versions can be added above this line>-------------------------------------------------
 		VersionPlusOne,
 		LatestVersion = VersionPlusOne - 1
diff --git a/Engine/Source/Runtime/Engine/Classes/Engine/EngineTypes.h b/Engine/Source/Runtime/Engine/Classes/Engine/EngineTypes.h
--- a/Engine/Source/Runtime/Engine/Classes/Engine/EngineTypes.h
+++ b/Engine/Source/Runtime/Engine/Classes/Engine/EngineTypes.h
@@ -401,6 +401,21 @@ enum ERefractionCoverageMode : int
 	RCM_CoverageAccountedFor UMETA(DisplayName = "Coverage Accounted For"),
 };
 
+/** Determines how the pixel depth offset is evaluated and applied. Must match PODM_LEGACY in MaterialTemplace.ush.*/
+UENUM()
+enum EPixelDepthOffsetMode : int
+{
+	/**
+	 * This is the legacy mode where PDO is applied differently for Depth (along View Forward) and world position (along Camera Vector).
+	 */
+	PDOM_Legacy UMETA(DisplayName = "Legacy"),
+
+	/**
+	 * PDO is applied along the Camera Vector for Depth and World Position altogether.
+	 */
+	PDOM_AlongCameraVector UMETA(DisplayName = "Along Camera Vector"),
+};
+
 /**
  * Enumerates available options for the translucency sort policy.
  */
diff --git a/Engine/Source/Runtime/Engine/Private/Materials/HLSLMaterialTranslator.cpp b/Engine/Source/Runtime/Engine/Private/Materials/HLSLMaterialTranslator.cpp
--- a/Engine/Source/Runtime/Engine/Private/Materials/HLSLMaterialTranslator.cpp
+++ b/Engine/Source/Runtime/Engine/Private/Materials/HLSLMaterialTranslator.cpp
@@ -224,7 +224,7 @@ UE::DerivedData::FValueId EnvironmentDefinesId = UE::DerivedData::FValueId::From
 
 /** Data structure used to cache a part of material translation results. It contains all the generated
  *  defines that will be declared during the compilation of the generated material shader.
- *  Note: if you make a change to this structure, remember to FDevSystemGuids::MaterialTranslationDDCVersion.
+ *  Note: if you make a change to this structure, remember to bump FDevSystemGuids::MaterialTranslationDDCVersion.
  */
 struct FHLSLMaterialTranslator::FEnvironmentDefines
 {
@@ -279,6 +279,7 @@ struct FHLSLMaterialTranslator::FEnvironmentDefines
 	bool bMaterialFullyRough;
 	bool bMaterialUsesAnisotropy;
 	uint8 MaterialDecalReadMask;
+	int8 PixelDepthOffsetMode;
 	bool bMaterialUsesDecalLookup;
 	uint8 MaterialPathTracingBufferRead;
 	bool MaterialNeuralPostProcess;
@@ -372,6 +373,7 @@ struct FHLSLMaterialTranslator::FEnvironmentDefines
 		Ar << bMaterialFullyRough;
 		Ar << bMaterialUsesAnisotropy;
 		Ar << MaterialDecalReadMask;
+		Ar << PixelDepthOffsetMode;
 		Ar << bMaterialUsesDecalLookup;
 		Ar << MaterialPathTracingBufferRead;
 		Ar << MaterialNeuralPostProcess;
@@ -2581,6 +2583,7 @@ void FHLSLMaterialTranslator::GetMaterialEnvironment(EShaderPlatform InPlatform,
 	OutEnvironment.SetDefine(TEXT("NEEDS_PER_INSTANCE_RANDOM_PS"), EnvironmentDefines->bUsesPerInstanceRandomPS);
 	OutEnvironment.SetDefine(TEXT("USES_TRANSFORM_VECTOR"), EnvironmentDefines->bUsesTransformVector);
 	OutEnvironment.SetDefine(TEXT("WANT_PIXEL_DEPTH_OFFSET"), EnvironmentDefines->bUsesPixelDepthOffset);
+	OutEnvironment.SetDefine(TEXT("PIXEL_DEPTH_OFFSET_MODE"), EnvironmentDefines->PixelDepthOffsetMode);
 
 	// we want USES_WORLD_POSITION_OFFSET to be readable as a bool compile argument, hence the != 0 comparison 
 	// (bUsesWorldPositionOffset is actually a 1-bit uint32 bitfield member)
@@ -14842,6 +14845,7 @@ void FHLSLMaterialTranslator::PrepareEnvironmentDefines()
 	EnvironmentDefines->bMaterialUsesDecalLookup = MaterialCompilationOutput.bUsesDBufferTextureLookup;
 	EnvironmentDefines->MaterialPathTracingBufferRead = MaterialCompilationOutput.UsedPathTracingBufferTextures;
 	EnvironmentDefines->MaterialNeuralPostProcess = (MaterialCompilationOutput.bUsedWithNeuralNetworks || Material->IsUsedWithNeuralNetworks()) && Material->IsPostProcessMaterial();
+	EnvironmentDefines->PixelDepthOffsetMode = Material->GetPixelDepthOffsetMode();
 
 	// Count the number of VTStacks (each stack will allocate a feedback slot)
 	EnvironmentDefines->NumVirtualTextureSamples = VTStacks.Num();
diff --git a/Engine/Source/Runtime/Engine/Private/Materials/Material.cpp b/Engine/Source/Runtime/Engine/Private/Materials/Material.cpp
--- a/Engine/Source/Runtime/Engine/Private/Materials/Material.cpp
+++ b/Engine/Source/Runtime/Engine/Private/Materials/Material.cpp
@@ -1034,6 +1034,7 @@ UMaterial::UMaterial(const FObjectInitializer& ObjectInitializer)
 
 	RefractionMethod = RM_None;
 	RefractionCoverageMode = RCM_CoverageAccountedFor;
+	PixelDepthOffsetMode = PDOM_AlongCameraVector;
 
 	bAllowVariableRateShading = true;
 
@@ -2910,6 +2911,12 @@ void UMaterial::Serialize(FArchive& Ar)
 			TranslucencyPass = MTP_BeforeDOF;
 		}
 	}
+
+	if (Ar.IsLoading() && Ar.CustomVer(FFortniteMainBranchObjectVersion::GUID) < FFortniteMainBranchObjectVersion::MaterialPixelDepthOffsetMode)
+	{
+		// All previous material must use the legacy pixel depth offset mode. New material will use the new mode.
+		PixelDepthOffsetMode = PDOM_Legacy;
+	}
 }
 
 void UMaterial::PostDuplicate(bool bDuplicateForPIE)
diff --git a/Engine/Source/Runtime/Engine/Private/Materials/MaterialShared.cpp b/Engine/Source/Runtime/Engine/Private/Materials/MaterialShared.cpp
--- a/Engine/Source/Runtime/Engine/Private/Materials/MaterialShared.cpp
+++ b/Engine/Source/Runtime/Engine/Private/Materials/MaterialShared.cpp
@@ -2059,6 +2059,7 @@ bool FMaterialResource::AllowNegativeEmissiveColor() const
 
 bool FMaterialResource::IsDistorted() const { return Material->bUsesDistortion && IsTranslucentBlendMode(GetBlendMode()); }
 ERefractionCoverageMode FMaterialResource::GetRefractionCoverageMode() const { return Material->RefractionCoverageMode; }
+EPixelDepthOffsetMode FMaterialResource::GetPixelDepthOffsetMode() const { return Material->PixelDepthOffsetMode; }
 float FMaterialResource::GetTranslucencyDirectionalLightingIntensity() const { return Material->TranslucencyDirectionalLightingIntensity; }
 float FMaterialResource::GetTranslucentShadowDensityScale() const { return Material->TranslucentShadowDensityScale; }
 float FMaterialResource::GetTranslucentSelfShadowDensityScale() const { return Material->TranslucentSelfShadowDensityScale; }
diff --git a/Engine/Source/Runtime/Engine/Public/MaterialShared.h b/Engine/Source/Runtime/Engine/Public/MaterialShared.h
--- a/Engine/Source/Runtime/Engine/Public/MaterialShared.h
+++ b/Engine/Source/Runtime/Engine/Public/MaterialShared.h
@@ -2089,6 +2089,7 @@ public:
 	virtual bool GetCastDynamicShadowAsMasked() const = 0;
 	virtual bool IsDistorted() const { return false; };
 	virtual ERefractionCoverageMode GetRefractionCoverageMode() const { return RCM_CoverageIgnored; }
+	virtual EPixelDepthOffsetMode GetPixelDepthOffsetMode() const { return PDOM_Legacy; }
 	virtual float GetTranslucencyDirectionalLightingIntensity() const { return 1.0f; }
 	virtual float GetTranslucentShadowDensityScale() const { return 1.0f; }
 	virtual float GetTranslucentSelfShadowDensityScale() const { return 1.0f; }
@@ -2862,6 +2863,7 @@ public:
 	ENGINE_API virtual bool GetCastDynamicShadowAsMasked() const override;
 	ENGINE_API virtual bool IsDistorted() const override;
 	ENGINE_API virtual ERefractionCoverageMode GetRefractionCoverageMode() const override;
+	ENGINE_API virtual EPixelDepthOffsetMode GetPixelDepthOffsetMode() const override;
 	ENGINE_API virtual float GetTranslucencyDirectionalLightingIntensity() const override;
 	ENGINE_API virtual float GetTranslucentShadowDensityScale() const override;
 	ENGINE_API virtual float GetTranslucentSelfShadowDensityScale() const override;
diff --git a/Engine/Source/Runtime/Engine/Public/Materials/Material.h b/Engine/Source/Runtime/Engine/Public/Materials/Material.h
--- a/Engine/Source/Runtime/Engine/Public/Materials/Material.h
+++ b/Engine/Source/Runtime/Engine/Public/Materials/Material.h
@@ -1027,6 +1027,10 @@ public:
 	UPROPERTY(EditAnywhere, Category=Refraction)
 	TEnumAsByte<ERefractionCoverageMode> RefractionCoverageMode;
 
+	/** Controls whether refraction takes into account the material surface coverage, or not. */
+	UPROPERTY(EditAnywhere, Category=PixelDepthOffset)
+	TEnumAsByte<EPixelDepthOffsetMode> PixelDepthOffsetMode;
+
 	/** If multiple nodes with the same  type are inserted at the same point, this defined order and if they get combined, only used if domain is PostProcess */
 	UPROPERTY(EditAnywhere, BlueprintReadWrite, Category = PostProcessMaterial, meta = (DisplayName = "Blendable Priority"))
 	int32 BlendablePriority;
