Fix shader compilation issue with Mac SM5 preview in editor.

Shader conductor issues a shader compilation issue with MSAA, by failing to unroll the loop in PostProcesselectionOutline. Vanilla FXC manages to compile the shader though.

diff --git a/Engine/Shaders/Private/HairStrands/HairStrandsVoxelRasterCompute.usf b/Engine/Shaders/Private/HairStrands/HairStrandsVoxelRasterCompute.usf
--- a/Engine/Shaders/Private/HairStrands/HairStrandsVoxelRasterCompute.usf
+++ b/Engine/Shaders/Private/HairStrands/HairStrandsVoxelRasterCompute.usf
@@ -1168,6 +1168,7 @@ groupshared uint g_Density4[64][GROUP_SIZE_Z];
 groupshared uint g_Density2[8][GROUP_SIZE_Z];
 #endif
 
+// GroupSize is [64u,1u,16u]
 [numthreads(GROUP_SIZE_X, 1, GROUP_SIZE_Z)]
 void MainCS(uint3 DispatchThreadId : SV_DispatchThreadID, uint3 GroupThreadId : SV_GroupThreadID)
 {
@@ -1175,6 +1176,14 @@ void MainCS(uint3 DispatchThreadId : SV_DispatchThreadID, uint3 GroupThreadId :
 	const uint3 PageCoord = IndexToCoord(PageIndex, PageCountResolution);
 	const bool bValid = PageIndex < AllocatedPageCountBuffer[0];
 	
+	// Dummy initialization for pleasing compiler
+	#if PERMUTATION_MIP_AGGREGATE && COMPILER_FXC
+	if (GroupThreadId.x < 8 && GroupThreadId.z < GROUP_SIZE_Z)
+	{
+		g_Density2[GroupThreadId.x][GroupThreadId.z] = 0;
+	}
+	#endif
+
 	uint Total = 0;
 	if (bValid)
 	{
diff --git a/Engine/Shaders/Private/PostProcessSelectionOutline.usf b/Engine/Shaders/Private/PostProcessSelectionOutline.usf
--- a/Engine/Shaders/Private/PostProcessSelectionOutline.usf
+++ b/Engine/Shaders/Private/PostProcessSelectionOutline.usf
@@ -163,7 +163,8 @@ float4 PerSample(int2 PixelPos, int SampleID, FPixelInfo CenterPixelInfo, float3
 	bool bInnerBorder = BorderMinMax.y == 4;
 
 	// moving diagonal lines
-	float PatternMask = ((PixelPos.x/2 + PixelPos.y/2) % 2) * 0.6f;
+	const uint2 uHalfResPixelPos = uint2(PixelPos)>>1u;
+	float PatternMask = ((uHalfResPixelPos.x + uHalfResPixelPos.y) % 2) * 0.6f;
 
 	// the contants express the two opacity values we see when the primitive is hidden
 	float LowContrastPatternMask = lerp(0.2, 1.0f, PatternMask);
@@ -270,7 +271,12 @@ void MainPS(
 #if COMPILER_GLSL && FEATURE_LEVEL >= FEATURE_LEVEL_SM4
 	UNROLL
 #endif
-		for(int SampleID = 0; SampleID < MSAA_SAMPLE_COUNT; ++SampleID)
+		int SampleID = 0;
+		// Workaround for platform preview with MSAA & FXC
+		#if !(COMPILER_FXC && MOBILE_EMULATION)
+		[unroll(MSAA_SAMPLE_COUNT)]
+		for(SampleID = 0; SampleID < MSAA_SAMPLE_COUNT; ++SampleID)
+		#endif
 		{
 			Sum += PerSample(ColorPixelPos, SampleID, CenterPixelInfo, DeviceZPlane, DeviceZMinMax);
 		}
@@ -297,7 +303,8 @@ void MainPS(
 #else
 		OutColor.rgb = sRGBToLinear(OutColor.rgb);
 #endif
-		float PatternMask = ((ColorPixelPos.x/2 + ColorPixelPos.y/2) % 2) * 0.6f;
+		const uint2 uHalfResPixelPos = uint2(ColorPixelPos)>>1;
+		float PatternMask = ((uHalfResPixelPos.x + uHalfResPixelPos.y) % 2) * 0.6f;
 		float VisibleMask = lerp(PatternMask, 1.0f, CenterPixelInfo.fVisible);
 		OutColor.rgb = lerp(OutColor.rgb, OverlayColor.rgb, OverlayColor.a * VisibleMask);
 		OutColor.rgb = lerp(OutColor.rgb, SelectionEffect.rgb, SelectionEffect.a);
