VulkanRHI: Match the FScreenPassVS output in the PS

#rnx

diff --git a/Engine/Shaders/Private/PathTracing/PathTracingCompositingPixelShader.usf b/Engine/Shaders/Private/PathTracing/PathTracingCompositingPixelShader.usf
--- a/Engine/Shaders/Private/PathTracing/PathTracingCompositingPixelShader.usf
+++ b/Engine/Shaders/Private/PathTracing/PathTracingCompositingPixelShader.usf
@@ -17,11 +17,12 @@ Texture2D<float4> NormalDepthTexture;
 int AdaptiveSamplingVisualize;
 
 void CompositeMain(
-	in noperspective float2 UV : TEXCOORD0,
+	in noperspective float4 UVAndScreenPos : TEXCOORD0,
 	out float4 OutColor : SV_Target0,
 	out float OutDepth : SV_DEPTH
 )
 {
+	float2 UV = UVAndScreenPos.xy;
 	float2 BufferSize = View.BufferSizeAndInvSize.xy;
 	int3 TexCoord = int3(UV * BufferSize - View.ViewRectMin.xy, 0);
 
