Heterogeneous Volumes: Adding DownsampleFactor

r.HeterogeneousVolumes.DownsampleFactor (default=1)
Decouples the operating resolution of the single-scattering ray marching integrators from the active viewport resolution. A DownsampleFactor of 2 halves both dimensions of the rendering viewport, operating on a quarter of the pixels.

diff --git a/Engine/Shaders/Private/HeterogeneousVolumes/HeterogeneousVolumesComposite.usf b/Engine/Shaders/Private/HeterogeneousVolumes/HeterogeneousVolumesComposite.usf
--- a/Engine/Shaders/Private/HeterogeneousVolumes/HeterogeneousVolumesComposite.usf
+++ b/Engine/Shaders/Private/HeterogeneousVolumes/HeterogeneousVolumesComposite.usf
@@ -4,6 +4,7 @@
 
 Texture2D<float4> HeterogeneousVolumeRadiance;
 RWTexture2D<float4> RWColorTexture;
+int DownsampleFactor;
 
 [numthreads(THREADGROUP_SIZE_2D, THREADGROUP_SIZE_2D, 1)]
 void HeterogeneousVolumesCompositeCS(
@@ -15,12 +16,13 @@ void HeterogeneousVolumesCompositeCS(
 	{
 		return;
 	}
-	uint2 PixelCoord = DispatchThreadId.xy + View.ViewRectMin.xy;
+	uint2 ViewPixelCoord = DispatchThreadId.xy + View.ViewRectMin.xy;
+	uint2 PixelCoord = ViewPixelCoord / DownsampleFactor;
 
 	float4 Result;
 	float HeterogeneousVolumeTransmittance = HeterogeneousVolumeRadiance[PixelCoord].a;
-	Result.rgb = HeterogeneousVolumeRadiance[PixelCoord].rgb + RWColorTexture[PixelCoord].rgb * HeterogeneousVolumeTransmittance;
-	Result.a = RWColorTexture[PixelCoord].a * HeterogeneousVolumeTransmittance;
+	Result.rgb = HeterogeneousVolumeRadiance[PixelCoord].rgb + RWColorTexture[ViewPixelCoord].rgb * HeterogeneousVolumeTransmittance;
+	Result.a = RWColorTexture[ViewPixelCoord].a * HeterogeneousVolumeTransmittance;
 
-	RWColorTexture[PixelCoord] = Result;
+	RWColorTexture[ViewPixelCoord] = Result;
 }
diff --git a/Engine/Shaders/Private/HeterogeneousVolumes/HeterogeneousVolumesLiveShadingPipeline.usf b/Engine/Shaders/Private/HeterogeneousVolumes/HeterogeneousVolumesLiveShadingPipeline.usf
--- a/Engine/Shaders/Private/HeterogeneousVolumes/HeterogeneousVolumesLiveShadingPipeline.usf
+++ b/Engine/Shaders/Private/HeterogeneousVolumes/HeterogeneousVolumesLiveShadingPipeline.usf
@@ -72,6 +72,7 @@ int bJitter;
 
 // Dispatch data
 int3 GroupCount;
+int DownsampleFactor;
 
 // Output
 RWTexture2D<float4> RWLightingTexture;
@@ -298,6 +299,11 @@ void RenderLightingCacheWithLiveShadingCS(
 #endif // DIM_LIGHTING_CACHE_MODE
 }
 
+float2 GetDispatchSize()
+{
+	return View.ViewSizeAndInvSize.xy / DownsampleFactor;
+}
+
 [numthreads(THREADGROUP_SIZE_2D, THREADGROUP_SIZE_2D, 1)]
 void RenderSingleScatteringWithLiveShadingCS(
 	uint2 GroupThreadId : SV_GroupThreadID,
@@ -307,19 +313,20 @@ void RenderSingleScatteringWithLiveShadingCS(
 	FDebugOutput DebugOutput = (FDebugOutput)0;
 
 	// Create screen ray
-	if (any(DispatchThreadId.xy >= View.ViewSizeAndInvSize.xy))
+	if (any(DispatchThreadId.xy >= GetDispatchSize()))
 	{
 		return;
 	}
 	uint LinearIndex = DispatchThreadId.y * GroupCount.x * THREADGROUP_SIZE_2D + DispatchThreadId.x;
-	uint2 PixelCoord = DispatchThreadId.xy + View.ViewRectMin.xy;
+	uint2 PixelCoord = DispatchThreadId.xy + View.ViewRectMin.xy / DownsampleFactor;
+	uint2 ViewPixelCoord = PixelCoord * DownsampleFactor + View.ViewRectMin.xy;
 
 	float3 Radiance = 0.0;
 	float3 Transmittance = RWLightingTexture[PixelCoord].a;
 	float PrevOpacity = Luminance(1.0 - Transmittance);
 
 	// Extract depth
-	float DeviceZ = SceneDepthTexture.Load(int3(PixelCoord, 0)).r;
+	float DeviceZ = SceneDepthTexture.Load(int3(ViewPixelCoord, 0)).r;
 #if HAS_INVERTED_Z_BUFFER
 	DeviceZ = max(0.000000000001, DeviceZ);
 #endif // HAS_INVERTED_Z_BUFFER
@@ -328,12 +335,12 @@ void RenderSingleScatteringWithLiveShadingCS(
 	float SceneDepth = min(ConvertFromDeviceZ(DeviceZ), GetMaxTraceDistance());
 	DeviceZ = ConvertToDeviceZ(SceneDepth);
 
-	float Jitter = bJitter ? BlueNoiseScalar(PixelCoord, View.StateFrameIndexMod8) : 0.0;
+	float Jitter = bJitter ? BlueNoiseScalar(ViewPixelCoord, View.StateFrameIndexMod8) : 0.0;
 
 	// Intersect ray with bounding volume
 	// TODO: LWC integration..
-	float3 WorldRayOrigin = DFHackToFloat(DFFastSubtract(GetTranslatedWorldCameraPosFromView(PixelCoord + View.TemporalAAJitter.xy), PrimaryView.PreViewTranslation));
-	float3 WorldRayEnd = DFHackToFloat(SvPositionToWorld(float4(PixelCoord + View.TemporalAAJitter.xy, DeviceZ, 1)));
+	float3 WorldRayOrigin = DFHackToFloat(DFFastSubtract(GetTranslatedWorldCameraPosFromView(ViewPixelCoord + View.TemporalAAJitter.xy), PrimaryView.PreViewTranslation));
+	float3 WorldRayEnd = DFHackToFloat(SvPositionToWorld(float4(ViewPixelCoord + View.TemporalAAJitter.xy, DeviceZ, 1)));
 	float3 WorldRayDirection = WorldRayEnd - WorldRayOrigin;
 	float WorldRayLength = length(WorldRayDirection);
 	WorldRayDirection /= WorldRayLength;
diff --git a/Engine/Shaders/Private/HeterogeneousVolumes/HeterogeneousVolumesPreshadingPipeline.usf b/Engine/Shaders/Private/HeterogeneousVolumes/HeterogeneousVolumesPreshadingPipeline.usf
--- a/Engine/Shaders/Private/HeterogeneousVolumes/HeterogeneousVolumesPreshadingPipeline.usf
+++ b/Engine/Shaders/Private/HeterogeneousVolumes/HeterogeneousVolumesPreshadingPipeline.usf
@@ -62,6 +62,7 @@ int bJitter;
 
 // Dispatch data
 int3 GroupCount;
+int DownsampleFactor;
 
 // Volume data
 int MipLevel;
@@ -239,7 +240,11 @@ void GenerateRayMarchingTiles(
 
 	uint4 Viewport = uint4(DispatchThreadId.xy * THREADGROUP_SIZE_2D, (DispatchThreadId.xy + 1) * THREADGROUP_SIZE_2D);
 	uint2 PixelOffset = Viewport.xy;
-	if (any(PixelOffset >= View.ViewSizeAndInvSize.xy))
+
+	// Scale to full-res
+	Viewport *= DownsampleFactor;
+	uint2 ViewPixelOffset = Viewport.xy;
+	if (any(ViewPixelOffset >= View.ViewSizeAndInvSize.xy))
 	{
 		return;
 	}
@@ -503,6 +508,11 @@ RWStructuredBuffer<FVoxelDataPacked> RWVoxelOutputBuffer;
 StructuredBuffer<FRayMarchingTile> RayMarchingTilesBuffer;
 Buffer<uint> NumRayMarchingTilesBuffer;
 
+float2 GetDispatchSize()
+{
+	return View.ViewSizeAndInvSize.xy / DownsampleFactor;
+}
+
 [numthreads(THREADGROUP_SIZE_1D, 1, 1)]
 void RenderSingleScatteringWithPreshadingCS(
 	uint GroupId : SV_GroupID,
@@ -515,20 +525,21 @@ void RenderSingleScatteringWithPreshadingCS(
 	uint2 TileCoord = uint2(GroupThreadId % THREADGROUP_SIZE_2D, GroupThreadId / THREADGROUP_SIZE_2D);
 	//float2 PixelCoord = View.ViewRectMin.xy + RayMarchingTile.PixelOffset + TileCoord;
 	float2 PixelCoord = RayMarchingTile.PixelOffset + TileCoord;
+	float2 ViewPixelCoord = PixelCoord * DownsampleFactor;
 
 	float3 Radiance = 0.0;
 	float3 Transmittance = RWLightingTexture[PixelCoord].a;
 	float PrevOpacity = Luminance(1.0 - Transmittance);
 
 	// Create screen ray
-	if (any(PixelCoord.xy >= View.ViewSizeAndInvSize.xy))
+	if (any(PixelCoord.xy >= GetDispatchSize()))
 	{
 		return;
 	}
 	PixelCoord += View.ViewRectMin.xy;
 
 	// Extract depth
-	float DeviceZ = SceneDepthTexture.Load(int3(PixelCoord, 0)).r;
+	float DeviceZ = SceneDepthTexture.Load(int3(ViewPixelCoord, 0)).r;
 #if HAS_INVERTED_Z_BUFFER
 	DeviceZ = max(0.000000000001, DeviceZ);
 #endif // HAS_INVERTED_Z_BUFFER
@@ -537,12 +548,12 @@ void RenderSingleScatteringWithPreshadingCS(
 	float SceneDepth = min(ConvertFromDeviceZ(DeviceZ), MaxTraceDistance);
 	DeviceZ = ConvertToDeviceZ(SceneDepth);
 
-	float Jitter = bJitter ? InterleavedGradientNoise(PixelCoord, View.StateFrameIndexMod8) : 0.0;
+	float Jitter = bJitter ? InterleavedGradientNoise(ViewPixelCoord, View.StateFrameIndexMod8) : 0.0;
 
 	// Intersect ray with bounding volume
 	// TODO: LWC integration..
-	float3 WorldRayOrigin = DFHackToFloat(DFFastSubtract(GetTranslatedWorldCameraPosFromView(PixelCoord + View.TemporalAAJitter.xy), PrimaryView.PreViewTranslation));
-	float3 WorldRayEnd = DFHackToFloat(SvPositionToWorld(float4(PixelCoord + View.TemporalAAJitter.xy, DeviceZ, 1)));
+	float3 WorldRayOrigin = DFHackToFloat(DFFastSubtract(GetTranslatedWorldCameraPosFromView(ViewPixelCoord + View.TemporalAAJitter.xy), PrimaryView.PreViewTranslation));
+	float3 WorldRayEnd = DFHackToFloat(SvPositionToWorld(float4(ViewPixelCoord + View.TemporalAAJitter.xy, DeviceZ, 1)));
 	float3 WorldRayDirection = WorldRayEnd - WorldRayOrigin;
 	float WorldRayLength = length(WorldRayDirection);
 	WorldRayDirection /= WorldRayLength;
diff --git a/Engine/Source/Runtime/Renderer/Private/HeterogeneousVolumes/HeterogeneousVolumes.cpp b/Engine/Source/Runtime/Renderer/Private/HeterogeneousVolumes/HeterogeneousVolumes.cpp
--- a/Engine/Source/Runtime/Renderer/Private/HeterogeneousVolumes/HeterogeneousVolumes.cpp
+++ b/Engine/Source/Runtime/Renderer/Private/HeterogeneousVolumes/HeterogeneousVolumes.cpp
@@ -32,6 +32,13 @@ static TAutoConsoleVariable<int32> CVarTranslucencyHeterogeneousVolumes(
 	ECVF_RenderThreadSafe | ECVF_ReadOnly
 );
 
+static TAutoConsoleVariable<float> CVarHeterogeneousVolumesDownsampleFactor(
+	TEXT("r.HeterogeneousVolumes.DownsampleFactor"),
+	1.0,
+	TEXT("Downsamples the rendered viewport (Default = 1.0)"),
+	ECVF_Scalability | ECVF_RenderThreadSafe
+);
+
 static TAutoConsoleVariable<int32> CVarHeterogeneousVolumesComposition(
 	TEXT("r.HeterogeneousVolumes.Composition"),
 	0,
@@ -82,7 +89,7 @@ static TAutoConsoleVariable<int32> CVarHeterogeneousVolumesMaxStepCount(
 	TEXT("r.HeterogeneousVolumes.MaxStepCount"),
 	512,
 	TEXT("The maximum ray-marching step count (Default = 512)"),
-	ECVF_RenderThreadSafe
+	ECVF_Scalability | ECVF_RenderThreadSafe
 );
 
 static TAutoConsoleVariable<float> CVarHeterogeneousVolumesMaxTraceDistance(
@@ -367,6 +374,16 @@ bool ShouldRenderMeshBatchWithHeterogeneousVolumes(
 namespace HeterogeneousVolumes
 {
 	// CVars
+	int32 GetDownsampleFactor()
+	{
+		return FMath::Clamp(CVarHeterogeneousVolumesDownsampleFactor.GetValueOnRenderThread(), 1, 8);
+	}
+
+	FIntPoint GetScaledViewRect(FIntRect ViewRect)
+	{
+		return FIntPoint::DivideAndRoundUp(ViewRect.Size(), GetDownsampleFactor());
+	}
+
 	FIntVector GetVolumeResolution(const IHeterogeneousVolumeInterface* Interface)
 	{
 		FIntVector VolumeResolution = Interface->GetVoxelResolution();
@@ -947,6 +964,7 @@ class FHeterogeneousVolumesCompositeCS : public FGlobalShader
 
 		// Dispatch data
 		SHADER_PARAMETER(FIntVector, GroupCount)
+		SHADER_PARAMETER(int32, DownsampleFactor)
 
 		// Output
 		SHADER_PARAMETER_RDG_TEXTURE_UAV(RWTexture2D<float4>, RWColorTexture)
@@ -995,9 +1013,7 @@ void FDeferredShadingSceneRenderer::CompositeHeterogeneousVolumes(
 
 		if (ShouldRenderHeterogeneousVolumesForView(View))
 		{
-			uint32 GroupCountX = FMath::DivideAndRoundUp(View.ViewRect.Size().X, FHeterogeneousVolumesCompositeCS::GetThreadGroupSize2D());
-			uint32 GroupCountY = FMath::DivideAndRoundUp(View.ViewRect.Size().Y, FHeterogeneousVolumesCompositeCS::GetThreadGroupSize2D());
-			FIntVector GroupCount = FIntVector(GroupCountX, GroupCountY, 1);
+			FIntVector GroupCount = FComputeShaderUtils::GetGroupCount(View.ViewRect.Size(), FHeterogeneousVolumesCompositeCS::GetThreadGroupSize2D());
 
 			FHeterogeneousVolumesCompositeCS::FParameters* PassParameters = GraphBuilder.AllocParameters<FHeterogeneousVolumesCompositeCS::FParameters>();
 			{
@@ -1007,6 +1023,7 @@ void FDeferredShadingSceneRenderer::CompositeHeterogeneousVolumes(
 				PassParameters->HeterogeneousVolumeRadiance = View.HeterogeneousVolumeRadiance;
 				// Dispatch data
 				PassParameters->GroupCount = GroupCount;
+				PassParameters->DownsampleFactor = HeterogeneousVolumes::GetDownsampleFactor();
 				// Output
 				PassParameters->RWColorTexture = GraphBuilder.CreateUAV(SceneTextures.Color.Target);
 			}
diff --git a/Engine/Source/Runtime/Renderer/Private/HeterogeneousVolumes/HeterogeneousVolumes.h b/Engine/Source/Runtime/Renderer/Private/HeterogeneousVolumes/HeterogeneousVolumes.h
--- a/Engine/Source/Runtime/Renderer/Private/HeterogeneousVolumes/HeterogeneousVolumes.h
+++ b/Engine/Source/Runtime/Renderer/Private/HeterogeneousVolumes/HeterogeneousVolumes.h
@@ -61,6 +61,8 @@ namespace HeterogeneousVolumes
 	FIntVector GetVolumeResolution(const IHeterogeneousVolumeInterface*);
 	FIntVector GetLightingCacheResolution(const IHeterogeneousVolumeInterface*, float LODFactor);
 
+	int32 GetDownsampleFactor();
+	FIntPoint GetScaledViewRect(FIntRect ViewRect);
 	float GetShadowStepSize();
 	float GetMaxTraceDistance();
 	float GetMaxShadowTraceDistance();
diff --git a/Engine/Source/Runtime/Renderer/Private/HeterogeneousVolumes/HeterogeneousVolumesLiveShadingPipeline.cpp b/Engine/Source/Runtime/Renderer/Private/HeterogeneousVolumes/HeterogeneousVolumesLiveShadingPipeline.cpp
--- a/Engine/Source/Runtime/Renderer/Private/HeterogeneousVolumes/HeterogeneousVolumesLiveShadingPipeline.cpp
+++ b/Engine/Source/Runtime/Renderer/Private/HeterogeneousVolumes/HeterogeneousVolumesLiveShadingPipeline.cpp
@@ -290,6 +290,7 @@ class FRenderSingleScatteringWithLiveShadingCS : public FMeshMaterialShader
 
 		// Dispatch data
 		SHADER_PARAMETER(FIntVector, GroupCount)
+		SHADER_PARAMETER(int32, DownsampleFactor)
 
 		// Output
 		SHADER_PARAMETER_RDG_TEXTURE_UAV(RWTexture2D<float4>, RWLightingTexture)
@@ -644,9 +645,7 @@ static void RenderSingleScatteringWithLiveShading(
 	MaterialRenderProxy = MaterialRenderProxy ? MaterialRenderProxy : DefaultMaterialRenderProxy;
 	check(Material.GetMaterialDomain() == MD_Volume);
 
-	uint32 GroupCountX = FMath::DivideAndRoundUp(View.ViewRect.Size().X, FRenderSingleScatteringWithLiveShadingCS::GetThreadGroupSize2D());
-	uint32 GroupCountY = FMath::DivideAndRoundUp(View.ViewRect.Size().Y, FRenderSingleScatteringWithLiveShadingCS::GetThreadGroupSize2D());
-	FIntVector GroupCount = FIntVector(GroupCountX, GroupCountY, 1);
+	FIntVector GroupCount = FComputeShaderUtils::GetGroupCount(HeterogeneousVolumes::GetScaledViewRect(View.ViewRect), FRenderSingleScatteringWithLiveShadingCS::GetThreadGroupSize2D());
 
 	// Note must be done in the same scope as we add the pass otherwise the UB lifetime will not be guaranteed
 	FDeferredLightUniformStruct DeferredLightUniform;
@@ -764,6 +763,7 @@ static void RenderSingleScatteringWithLiveShading(
 
 		// Dispatch data
 		PassParameters->GroupCount = GroupCount;
+		PassParameters->DownsampleFactor = HeterogeneousVolumes::GetDownsampleFactor();
 
 		// Output
 		PassParameters->RWLightingTexture = GraphBuilder.CreateUAV(HeterogeneousVolumeTexture);
diff --git a/Engine/Source/Runtime/Renderer/Private/HeterogeneousVolumes/HeterogeneousVolumesPreshadingPipeline.cpp b/Engine/Source/Runtime/Renderer/Private/HeterogeneousVolumes/HeterogeneousVolumesPreshadingPipeline.cpp
--- a/Engine/Source/Runtime/Renderer/Private/HeterogeneousVolumes/HeterogeneousVolumesPreshadingPipeline.cpp
+++ b/Engine/Source/Runtime/Renderer/Private/HeterogeneousVolumes/HeterogeneousVolumesPreshadingPipeline.cpp
@@ -43,6 +43,7 @@ class FGenerateRayMarchingTiles : public FGlobalShader
 
 		// Dispatch data
 		SHADER_PARAMETER(FIntVector, GroupCount)
+		SHADER_PARAMETER(int32, DownsampleFactor)
 
 		// Debug Output
 		SHADER_PARAMETER_RDG_BUFFER_UAV(RWStructuredBuffer<Volumes::FRayMarchingDebug>, RWRayMarchingDebugBuffer)
@@ -118,10 +119,8 @@ void GenerateRayMarchingTiles(
 	FRDGBufferRef& VoxelsPerTileBuffer
 )
 {
-	uint32 GroupCountX = FMath::DivideAndRoundUp(View.ViewRect.Size().X, FGenerateRayMarchingTiles::GetThreadGroupSize2D());
-	uint32 GroupCountY = FMath::DivideAndRoundUp(View.ViewRect.Size().Y, FGenerateRayMarchingTiles::GetThreadGroupSize2D());
-	FIntVector GroupCount = FIntVector(GroupCountX, GroupCountY, 1);
-	uint32 NumTiles = GroupCountX * GroupCountY;
+	FIntVector GroupCount = FComputeShaderUtils::GetGroupCount(HeterogeneousVolumes::GetScaledViewRect(View.ViewRect), FGenerateRayMarchingTiles::GetThreadGroupSize2D());
+	uint32 NumTiles = GroupCount.X * GroupCount.Y;
 
 	NumRayMarchingTilesBuffer = GraphBuilder.CreateBuffer(
 		FRDGBufferDesc::CreateIndirectDesc<FRHIDispatchIndirectParameters>(1),
@@ -177,6 +176,7 @@ void GenerateRayMarchingTiles(
 
 		// Dispatch data
 		PassParameters->GroupCount = GroupCount;
+		PassParameters->DownsampleFactor = HeterogeneousVolumes::GetDownsampleFactor();
 
 		// Debug
 		PassParameters->RWRayMarchingDebugBuffer = GraphBuilder.CreateUAV(RayMarchingDebugBuffer);
@@ -341,6 +341,7 @@ class FRenderSingleScatteringWithPreshadingCS : public FGlobalShader
 
 		// Indirect args
 		RDG_BUFFER_ACCESS(IndirectArgs, ERHIAccess::IndirectArgs)
+		SHADER_PARAMETER(int32, DownsampleFactor)
 
 		// Output
 		SHADER_PARAMETER_RDG_TEXTURE_UAV(RWTexture2D<float4>, RWLightingTexture)
@@ -660,6 +661,7 @@ void RenderSingleScatteringWithPreshadingCompute(
 
 		// Dispatch data
 		PassParameters->IndirectArgs = NumRayMarchingTilesBuffer;
+		PassParameters->DownsampleFactor = HeterogeneousVolumes::GetDownsampleFactor();
 
 		// Output
 		PassParameters->RWLightingTexture = GraphBuilder.CreateUAV(HeterogeneousVolumeTexture);
diff --git a/Engine/Source/Runtime/Renderer/Private/HeterogeneousVolumes/HeterogeneousVolumesVoxelGridPipeline.cpp b/Engine/Source/Runtime/Renderer/Private/HeterogeneousVolumes/HeterogeneousVolumesVoxelGridPipeline.cpp
--- a/Engine/Source/Runtime/Renderer/Private/HeterogeneousVolumes/HeterogeneousVolumesVoxelGridPipeline.cpp
+++ b/Engine/Source/Runtime/Renderer/Private/HeterogeneousVolumes/HeterogeneousVolumesVoxelGridPipeline.cpp
@@ -207,7 +207,7 @@ static TAutoConsoleVariable<int32> CVarHeterogeneousVolumesShadowResolution(
 	TEXT("r.HeterogeneousVolumes.Shadows.Resolution"),
 	512,
 	TEXT("Resolution when building volumetric shadow map (Default = 512)\n"),
-	ECVF_RenderThreadSafe
+	ECVF_Scalability | ECVF_RenderThreadSafe
 );
 
 static TAutoConsoleVariable<float> CVarHeterogeneousVolumesStepSizeForShadows(
@@ -2885,9 +2885,7 @@ void RenderTransmittanceWithVoxelGrid(
 	FRDGTextureRef& HeterogeneousVolumeRadiance
 )
 {
-	uint32 GroupCountX = FMath::DivideAndRoundUp(View.ViewRect.Size().X, FRenderTransmittanceWithVoxelGridCS::GetThreadGroupSize2D());
-	uint32 GroupCountY = FMath::DivideAndRoundUp(View.ViewRect.Size().Y, FRenderTransmittanceWithVoxelGridCS::GetThreadGroupSize2D());
-	FIntVector GroupCount = FIntVector(GroupCountX, GroupCountY, 1);
+	FIntVector GroupCount = FComputeShaderUtils::GetGroupCount(HeterogeneousVolumes::GetScaledViewRect(View.ViewRect), FRenderTransmittanceWithVoxelGridCS::GetThreadGroupSize2D());
 
 	int32 BufferSize = View.ViewRect.Width() * View.ViewRect.Height();
 	FRDGBufferRef DebugBuffer = GraphBuilder.CreateBuffer(
