Fixed a wrong DepthAux output for a single water shading model on mobile

diff --git a/Engine/Shaders/Private/MobileBasePassPixelShader.usf b/Engine/Shaders/Private/MobileBasePassPixelShader.usf
--- a/Engine/Shaders/Private/MobileBasePassPixelShader.usf
+++ b/Engine/Shaders/Private/MobileBasePassPixelShader.usf
@@ -80,7 +80,7 @@
 #define NONMETAL (MATERIAL_NONMETAL || MOBILE_QL_FORCE_NONMETAL)
 #define FORCE_VERTEX_NORMAL (MOBILE_QL_DISABLE_MATERIAL_NORMAL)
 #define SUPPORT_SPOTLIGHTS_SHADOW (MATERIALBLENDING_SOLID || MATERIALBLENDING_MASKED) && PROJECT_MOBILE_ENABLE_MOVABLE_SPOTLIGHT_SHADOWS
-#define MOBILE_USE_GBUBUFFER (MOBILE_DEFERRED_SHADING && ((MATERIALBLENDING_SOLID || MATERIALBLENDING_MASKED) && !MATERIAL_SHADINGMODEL_SINGLELAYERWATER))
+#define MOBILE_USE_GBUFFER (MOBILE_DEFERRED_SHADING && ((MATERIALBLENDING_SOLID || MATERIALBLENDING_MASKED) && !MATERIAL_SHADINGMODEL_SINGLELAYERWATER))
 
 #if MATERIAL_SHADINGMODEL_THIN_TRANSLUCENT || SUBSTRATE_BLENDING_TRANSLUCENT_COLOREDTRANSMITTANCE
 	#define MOBILE_TRANSLUCENT_COLOR_TRANSMITTANCE 1
@@ -104,11 +104,11 @@
 #endif
 
 // SceneDepthAux always enabled for IOS, on other platforms only enabled when in forward and MobileHDR=true
-#define HAS_SCENE_DEPTH_AUX_OUTPUT (USE_SCENE_DEPTH_AUX && (MATERIALBLENDING_SOLID || MATERIALBLENDING_MASKED))
+#define HAS_SCENE_DEPTH_AUX_OUTPUT (USE_SCENE_DEPTH_AUX && (MATERIALBLENDING_SOLID || MATERIALBLENDING_MASKED) && !MATERIAL_SHADINGMODEL_SINGLELAYERWATER)
 
-#if MOBILE_USE_GBUBUFFER && MOBILE_EXTENDED_GBUFFER
+#if MOBILE_USE_GBUFFER && MOBILE_EXTENDED_GBUFFER
 	#define SV_TargetDepthAux SV_Target5	
-#elif MOBILE_USE_GBUBUFFER
+#elif MOBILE_USE_GBUFFER
 	#define SV_TargetDepthAux SV_Target4	
 #else
 	#define SV_TargetDepthAux SV_Target1
@@ -305,7 +305,7 @@ void Main(
 	, FMobileBasePassInterpolantsVSToPS BasePassInterpolants
 	, in float4 SvPosition : SV_Position
 	OPTIONAL_IsFrontFace
-#if MOBILE_USE_GBUBUFFER
+#if MOBILE_USE_GBUFFER
 	#if USE_GLES_FBF_DEFERRED
 	, out HALF4_TYPE OutProxy : SV_Target0
 	#elif MOBILE_TRANSLUCENT_COLOR_TRANSMITTANCE_PROGRAMMABLE_BLENDING
@@ -337,7 +337,7 @@ void Main(
 #endif
 	)
 {  
-#if MOBILE_USE_GBUBUFFER 
+#if MOBILE_USE_GBUFFER 
 	#if USE_GLES_FBF_DEFERRED && !MOBILE_TRANSLUCENT_COLOR_TRANSMITTANCE_PROGRAMMABLE_BLENDING
 		half4 OutColor;
 	#endif
@@ -853,7 +853,7 @@ void Main(
 
 	GBuffer.PrecomputedShadowFactors = GetPrimaryPrecomputedShadowMask(LightmapVTPageTableResult, Interpolants, MaterialParameters);
 
-#if MOBILE_USE_GBUBUFFER && !MATERIAL_SHADINGMODEL_UNLIT
+#if MOBILE_USE_GBUFFER && !MATERIAL_SHADINGMODEL_UNLIT
 	GBuffer.IndirectIrradiance = IndirectIrradiance;
 	MobileEncodeGBuffer(GBuffer, OutGBufferA, OutGBufferB, OutGBufferC, OutGBufferD);
 #else
@@ -1113,7 +1113,7 @@ void Main(
 	}
 #endif
 
-#if MATERIAL_IS_SKY || (!MATERIAL_SHADINGMODEL_UNLIT && !MOBILE_USE_GBUBUFFER)
+#if MATERIAL_IS_SKY || (!MATERIAL_SHADINGMODEL_UNLIT && !MOBILE_USE_GBUFFER)
 	// clamp lit color to avoid overflows
 	OutColor.rgb = SafeGetOutColor(OutColor.rgb);
 #endif
@@ -1142,7 +1142,7 @@ void Main(
 	);
 #endif
 
-#if MOBILE_USE_GBUBUFFER && USE_GLES_FBF_DEFERRED
+#if MOBILE_USE_GBUFFER && USE_GLES_FBF_DEFERRED
 	OutProxy.rgb = OutColor.rgb;
 #elif MOBILE_TRANSLUCENT_COLOR_TRANSMITTANCE_PROGRAMMABLE_BLENDING
 	OutProgrammableBlending = OutColor1 * FramebufferFetchColor0() + OutColor;
