Fix bug in TileBoundsIntersectDistanceFieldCaster causing artifacts when calculating indirect shadows.

diff --git a/Engine/Shaders/Private/CapsuleShadowShaders.usf b/Engine/Shaders/Private/CapsuleShadowShaders.usf
--- a/Engine/Shaders/Private/CapsuleShadowShaders.usf
+++ b/Engine/Shaders/Private/CapsuleShadowShaders.usf
@@ -714,7 +714,7 @@ bool TileBoundsIntersectDistanceFieldCaster(
 	float TileConeAngleSin,
 	float MaxOcclusionDistance)
 {
-	const float3 TranslatedBoundsCenter = DFFastToTranslatedWorld(DFObjectBounds.Center, ResolvedView.PreViewTranslation);
+	const float3 TranslatedBoundsCenter = DFFastToTranslatedWorld(DFObjectBounds.Center, PrimaryView.PreViewTranslation);
 	BRANCH
 	if (SphereIntersectSphere(float4(TranslatedBoundsCenter, DFObjectBounds.SphereRadius + MaxOcclusionDistance), TileTranslatedBoundingSphere)
 #if LIGHT_SOURCE_MODE != LIGHT_SOURCE_FROM_RECEIVER
