Fix editor crash using Rect Light with mobile multi view

diff --git a/Engine/Shaders/Private/ShadowProjectionPixelShader.usf b/Engine/Shaders/Private/ShadowProjectionPixelShader.usf
--- a/Engine/Shaders/Private/ShadowProjectionPixelShader.usf
+++ b/Engine/Shaders/Private/ShadowProjectionPixelShader.usf
@@ -430,8 +430,10 @@ float2 PointLightProjParameters;
 
 /** Pixel shader for projecting a one pass point light shadow from a cube map. */
 void MainOnePassPointLightPS(
-#if INSTANCED_STEREO
-	in nointerpolation uint EyeIndex : TEXCOORD0,
+#if INSTANCED_STEREO && MOBILE_MULTI_VIEW
+	in nointerpolation uint EyeIndex : VIEW_ID,
+#elif MOBILE_MULTI_VIEW
+	in nointerpolation uint ViewId : SV_ViewID,
 #endif
 	in float4 SVPos : SV_POSITION,
 	out float4 OutColor : SV_Target0
@@ -439,15 +441,18 @@ void MainOnePassPointLightPS(
 {
 	
 	float ArrayIndex = 0;
-#if INSTANCED_STEREO
+#if INSTANCED_STEREO && MOBILE_MULTI_VIEW
 	ResolvedView = ResolveView(EyeIndex);
 	ArrayIndex = EyeIndex;
+#elif MOBILE_MULTI_VIEW
+	ResolvedView = ResolveView(ViewId);
+	ArrayIndex = ViewId;
 #else
 	ResolvedView = ResolveView();
 #endif
 
 	float2 ScreenUV = float2( SVPos.xy * ResolvedView.BufferSizeAndInvSize.zw );
-#if INSTANCED_STEREO && SHADING_PATH_MOBILE
+#if ((INSTANCED_STEREO && MOBILE_MULTI_VIEW) || MOBILE_MULTI_VIEW) && SHADING_PATH_MOBILE
 	float SceneW = CalcSceneDepth(ScreenUV, ArrayIndex);
 #else
 	float SceneW = CalcSceneDepth( ScreenUV );
