LFV - fix nan appearing in the scene color due to divide by 0.

diff --git a/Engine/Shaders/Private/LocalFogVolumes/LocalFogVolumeCommon.ush b/Engine/Shaders/Private/LocalFogVolumes/LocalFogVolumeCommon.ush
--- a/Engine/Shaders/Private/LocalFogVolumes/LocalFogVolumeCommon.ush
+++ b/Engine/Shaders/Private/LocalFogVolumes/LocalFogVolumeCommon.ush
@@ -210,19 +210,24 @@ FFogData LocalFogVolumeEvaluateAnalyticalIntegral(FLocalFogVolumeGPUInstanceData
 
 		float StartHeight = RayStartU.z - FogInstance.HeightFogOffset; // account for the fog offset.
 
+		// We are going to divide by RayDirU.z so we need it to be never be 0.
+		const float SafeDirThreshold = 0.0001;
+		const float SignRayDirUZ = RayDirU.z >= 0.0f ? 1.0f : -1.0f;
+		const float SafeRayDirUZ = (abs(RayDirU.z) < SafeDirThreshold) ? SafeDirThreshold * SignRayDirUZ : RayDirU.z;
+
 	#if 0
 		// Original integral, less artefact
-		HeightOpticalDepth = (FogInstance.HeightFogExtinction / FogInstance.HeightFogFalloff) * exp(-StartHeight * FogInstance.HeightFogFalloff) * (1.0 - exp(-RayDirU.z * RayLengthU * FogInstance.HeightFogFalloff)) / RayDirU.z;
+		HeightOpticalDepth = (FogInstance.HeightFogExtinction / FogInstance.HeightFogFalloff) * exp(-StartHeight * FogInstance.HeightFogFalloff) * (1.0 - exp(-SafeRayDirUZ * RayLengthU * FogInstance.HeightFogFalloff)) / SafeRayDirUZ;
 	#elif 1
 		// Reworked to avoid large value as input to exp
 		//OpticalDepth = (FogInstance.HeightFogExtinction / FogInstance.HeightFogFalloff) * (exp(-StartHeight * FogInstance.HeightFogFalloff) - exp(-StartHeight * FogInstance.HeightFogFalloff - RayDir.z * RayLength * FogInstance.HeightFogFalloff)) / RayDir.z;
 
 		// But we use an updated formulation to avoid large input to exp leading to inf and nan.
 		float Factor0 = StartHeight * FogInstance.HeightFogFalloff;
-		float Factor1 = RayDirU.z * RayLengthU * FogInstance.HeightFogFalloff;
+		float Factor1 = SafeRayDirUZ * RayLengthU * FogInstance.HeightFogFalloff;
 		Factor0 = max(-80.0f, Factor0); // Simply clamp Factor0 to a valid range in order to not have exp explode.
 
-		HeightOpticalDepth += (FogInstance.HeightFogExtinction / (FogInstance.HeightFogFalloff * RayDirU.z)) * (exp(-Factor0) - exp(-(Factor0 + Factor1)));
+		HeightOpticalDepth += (FogInstance.HeightFogExtinction / (FogInstance.HeightFogFalloff * SafeRayDirUZ)) * (exp(-Factor0) - exp(-(Factor0 + Factor1)));
 	#endif
 
 		/*
