Fixed Custom Depth pass being culled by HZB by enabling opting out with a per-batch dynamic branch.
 - Also fixed incorrect use of ViewToClip rather than PrevViewToClip.

diff --git a/Engine/Shaders/Private/InstanceCulling/BuildInstanceDrawCommands.usf b/Engine/Shaders/Private/InstanceCulling/BuildInstanceDrawCommands.usf
--- a/Engine/Shaders/Private/InstanceCulling/BuildInstanceDrawCommands.usf
+++ b/Engine/Shaders/Private/InstanceCulling/BuildInstanceDrawCommands.usf
@@ -105,7 +105,7 @@ void WriteInstanceForCompaction(uint Offset, bool bVisible, uint InstanceId, FIn
 	CompactInstanceIdsBufferOut[Offset] = PackedId;	
 }
 
-bool IsInstanceVisible(FPrimitiveSceneData PrimitiveData, FInstanceSceneData InstanceData, uint InstanceId, uint ViewIdIndex, FDrawCommandDesc DrawCommandDesc, inout uint CullingFlags)
+bool IsInstanceVisible(FPrimitiveSceneData PrimitiveData, FInstanceSceneData InstanceData, uint InstanceId, uint ViewIdIndex, bool bAllowOcclusionCulling, FDrawCommandDesc DrawCommandDesc, inout uint CullingFlags)
 {
 	CullingFlags = INSTANCE_CULLING_FLAGS_DEFAULT;
 
@@ -182,10 +182,10 @@ bool IsInstanceVisible(FPrimitiveSceneData PrimitiveData, FInstanceSceneData Ins
 
 #if OCCLUSION_CULL_INSTANCES
 		BRANCH
-		if (Cull.bIsVisible)
+		if (Cull.bIsVisible && bAllowOcclusionCulling)
 		{
 			const bool bPrevIsOrtho = IsOrthoProjection(NaniteView.PrevViewToClip);
-			FFrustumCullData PrevCull = BoxCullFrustum(InstanceData.LocalBoundsCenter, InstanceData.LocalBoundsExtent, DynamicData.PrevLocalToTranslatedWorld, NaniteView.PrevTranslatedWorldToClip, NaniteView.ViewToClip, bPrevIsOrtho, Cull.bNearClip, true);
+			FFrustumCullData PrevCull = BoxCullFrustum(InstanceData.LocalBoundsCenter, InstanceData.LocalBoundsExtent, DynamicData.PrevLocalToTranslatedWorld, NaniteView.PrevTranslatedWorldToClip, NaniteView.PrevViewToClip, bPrevIsOrtho, Cull.bNearClip, true);
 			BRANCH
 			if (PrevCull.bIsVisible && !PrevCull.bCrossesNearPlane)
 			{
@@ -229,13 +229,15 @@ void InstanceCullBuildInstanceIdBufferCS(uint3 GroupId : SV_GroupID, int GroupTh
 
 #if ENABLE_BATCH_MODE
 	// Load Instance culling context batch info, indirection per group
-	FContextBatchInfo BatchInfo = BatchInfos[BatchInds[DispatchGroupId]];
+	FContextBatchInfo BatchInfo = LoadBatchInfo(DispatchGroupId);
 #else // !ENABLE_BATCH_MODE
 	// Single Instance culling context batch in the call, set up batch from the kernel parameters
 	FContextBatchInfo BatchInfo = (FContextBatchInfo)0;
 	BatchInfo.NumViewIds = NumViewIds;
 	BatchInfo.DynamicInstanceIdOffset = DynamicInstanceIdOffset;
 	BatchInfo.DynamicInstanceIdMax = DynamicInstanceIdMax;
+	// Note: for the unbatched case, the permutation will control HZB test, so we set to true
+	BatchInfo.bAllowOcclusionCulling = true;
 #endif // ENABLE_BATCH_MODE
 
 	FInstanceWorkSetup WorkSetup = InstanceCullingLoadBalancer_Setup(GroupId, GroupThreadIndex, GetItemDataOffset(BatchInfo, CurrentBatchProcessingMode));
@@ -288,8 +290,8 @@ void InstanceCullBuildInstanceIdBufferCS(uint3 GroupId : SV_GroupID, int GroupTh
 #if STEREO_CULLING_MODE
 	uint CullingFlags0 = 0;
 	uint CullingFlags1 = 0;
-	const bool bVisible = IsInstanceVisible(PrimitiveData, InstanceData, InstanceId, BatchInfo.ViewIdsOffset + 0U, DrawCommandDesc, CullingFlags0) ||
-					      IsInstanceVisible(PrimitiveData, InstanceData, InstanceId, BatchInfo.ViewIdsOffset + 1U, DrawCommandDesc, CullingFlags1);
+	const bool bVisible = IsInstanceVisible(PrimitiveData, InstanceData, InstanceId, BatchInfo.ViewIdsOffset + 0U, BatchInfo.bAllowOcclusionCulling, DrawCommandDesc, CullingFlags0) ||
+					      IsInstanceVisible(PrimitiveData, InstanceData, InstanceId, BatchInfo.ViewIdsOffset + 1U, BatchInfo.bAllowOcclusionCulling, DrawCommandDesc, CullingFlags1);
 	const uint CullingFlags = CullingFlags0 | CullingFlags1;
 	NumVisibleInstances += bVisible ? 2 : 0;
 	
@@ -313,7 +315,7 @@ void InstanceCullBuildInstanceIdBufferCS(uint3 GroupId : SV_GroupID, int GroupTh
 	{
 		// Culling is disabled for multi-view
 		uint CullingFlags = 0;
-		bool bVisible = IsInstanceVisible(PrimitiveData, InstanceData, InstanceId, BatchInfo.ViewIdsOffset + ViewIdIndex, DrawCommandDesc, CullingFlags);
+		bool bVisible = IsInstanceVisible(PrimitiveData, InstanceData, InstanceId, BatchInfo.ViewIdsOffset + ViewIdIndex, BatchInfo.bAllowOcclusionCulling, DrawCommandDesc, CullingFlags);
 
 		NumVisibleInstances += bVisible ? 1 : 0;
 		
diff --git a/Engine/Shaders/Private/InstanceCulling/InstanceCullingCommon.ush b/Engine/Shaders/Private/InstanceCulling/InstanceCullingCommon.ush
--- a/Engine/Shaders/Private/InstanceCulling/InstanceCullingCommon.ush
+++ b/Engine/Shaders/Private/InstanceCulling/InstanceCullingCommon.ush
@@ -32,6 +32,22 @@ FDrawCommandDesc UnpackDrawCommandDesc(uint PackedData)
 	return Desc;
 }
 
+struct FContextBatchInfoPacked
+{
+	uint IndirectArgOffset;
+	uint InstanceDataWriteOffset;
+	uint PayloadOffset;
+	uint CompactionDataOffset;
+	uint ViewIdsOffset;
+	uint NumViewIds_bAllowOcclusionCulling;
+	uint DynamicInstanceIdOffset;
+	uint DynamicInstanceIdMax;
+	// Must be equal to ItemDataOffset[BATCH_PROCESSING_MODE_NUM]
+	// Not using array here to avoid GL shader compiler bug on Adreno devices
+	uint ItemDataOffset0;
+	uint ItemDataOffset1;
+};
+
 struct FContextBatchInfo
 {
 	uint IndirectArgOffset;
@@ -40,6 +56,7 @@ struct FContextBatchInfo
 	uint CompactionDataOffset;
 	uint ViewIdsOffset;
 	uint NumViewIds;
+	bool bAllowOcclusionCulling;
 	uint DynamicInstanceIdOffset;
 	uint DynamicInstanceIdMax;
 	// Must be equal to ItemDataOffset[BATCH_PROCESSING_MODE_NUM]
@@ -58,8 +75,25 @@ uint GetItemDataOffset(FContextBatchInfo BatchInfo, uint BatchProcessingMode)
 }
 
 #if ENABLE_BATCH_MODE
-StructuredBuffer<FContextBatchInfo> BatchInfos;
+StructuredBuffer<FContextBatchInfoPacked> BatchInfos;
 StructuredBuffer<uint> BatchInds;
+FContextBatchInfo LoadBatchInfo(uint GroupIndex)
+{
+	FContextBatchInfoPacked ContextBatchInfoPacked = BatchInfos[BatchInds[GroupIndex]];
+	FContextBatchInfo ContextBatchInfo;
+	ContextBatchInfo.IndirectArgOffset = ContextBatchInfoPacked.IndirectArgOffset;
+	ContextBatchInfo.InstanceDataWriteOffset = ContextBatchInfoPacked.InstanceDataWriteOffset;
+	ContextBatchInfo.PayloadOffset = ContextBatchInfoPacked.PayloadOffset;
+	ContextBatchInfo.CompactionDataOffset = ContextBatchInfoPacked.CompactionDataOffset;
+	ContextBatchInfo.ViewIdsOffset = ContextBatchInfoPacked.ViewIdsOffset;
+	ContextBatchInfo.NumViewIds = ContextBatchInfoPacked.NumViewIds_bAllowOcclusionCulling >> 1u;
+	ContextBatchInfo.bAllowOcclusionCulling = (ContextBatchInfoPacked.NumViewIds_bAllowOcclusionCulling & 1u) != 0u;
+	ContextBatchInfo.DynamicInstanceIdOffset = ContextBatchInfoPacked.DynamicInstanceIdOffset;
+	ContextBatchInfo.DynamicInstanceIdMax = ContextBatchInfoPacked.DynamicInstanceIdMax;
+	ContextBatchInfo.ItemDataOffset0 = ContextBatchInfoPacked.ItemDataOffset0;
+	ContextBatchInfo.ItemDataOffset1 = ContextBatchInfoPacked.ItemDataOffset1;
+	return ContextBatchInfo;
+}
 #endif // ENABLE_BATCH_MODE
 
 struct FInstanceCullingPayload
diff --git a/Engine/Shaders/Private/VirtualShadowMaps/VirtualShadowMapBuildPerPageDrawCommands.usf b/Engine/Shaders/Private/VirtualShadowMaps/VirtualShadowMapBuildPerPageDrawCommands.usf
--- a/Engine/Shaders/Private/VirtualShadowMaps/VirtualShadowMapBuildPerPageDrawCommands.usf
+++ b/Engine/Shaders/Private/VirtualShadowMaps/VirtualShadowMapBuildPerPageDrawCommands.usf
@@ -184,13 +184,15 @@ void CullPerPageDrawCommandsCs(uint3 GroupId : SV_GroupID, int GroupThreadIndex
 
 #if ENABLE_BATCH_MODE
 	// Load Instance culling context batch info, indirection per group
-	FContextBatchInfo BatchInfo = BatchInfos[BatchInds[DispatchGroupId]];
+	FContextBatchInfo BatchInfo = LoadBatchInfo(DispatchGroupId);
 	FVSMCullingBatchInfo VSMCullingBatchInfo = VSMCullingBatchInfos[BatchInds[DispatchGroupId]];
 #else // !ENABLE_BATCH_MODE
 	// Single Instance culling context batch in the call, set up batch from the kernel parameters
 	FContextBatchInfo BatchInfo = (FContextBatchInfo)0;
 	BatchInfo.DynamicInstanceIdOffset = DynamicInstanceIdOffset;
 	BatchInfo.DynamicInstanceIdMax = DynamicInstanceIdMax;
+	// Note: for the unbatched case, the permutation will control HZB test, so we set to true
+	BatchInfo.bAllowOcclusionCulling = true;
 
 	FVSMCullingBatchInfo VSMCullingBatchInfo;
 	VSMCullingBatchInfo.FirstPrimaryView = FirstPrimaryView;
diff --git a/Engine/Source/Runtime/Renderer/Private/InstanceCulling/InstanceCullingMergedContext.cpp b/Engine/Source/Runtime/Renderer/Private/InstanceCulling/InstanceCullingMergedContext.cpp
--- a/Engine/Source/Runtime/Renderer/Private/InstanceCulling/InstanceCullingMergedContext.cpp
+++ b/Engine/Source/Runtime/Renderer/Private/InstanceCulling/InstanceCullingMergedContext.cpp
@@ -44,7 +44,7 @@ void FInstanceCullingMergedContext::MergeBatches()
 		check(InstanceCullingContext.HasCullingCommands());
 
 		int32 BatchInfoIndex = BatchInfos.Num();
-		FContextBatchInfo& BatchInfo = BatchInfos.AddDefaulted_GetRef();
+		FContextBatchInfoPacked& BatchInfo = BatchInfos.AddDefaulted_GetRef();
 
 		BatchInfo.IndirectArgsOffset = IndirectArgs.Num();
 		//BatchInfo.NumIndirectArgs = InstanceCullingContext.IndirectArgs.Num();
@@ -65,7 +65,11 @@ void FInstanceCullingMergedContext::MergeBatches()
 		}
 
 		BatchInfo.ViewIdsOffset = ViewIds.Num();
-		BatchInfo.NumViewIds = InstanceCullingContext.ViewIds.Num();
+		BatchInfo.NumViewIds_bAllowOcclusionCulling = uint32(InstanceCullingContext.ViewIds.Num()) << 1u;
+		if (InstanceCullingContext.PrevHZB.IsValid())
+		{
+			BatchInfo.NumViewIds_bAllowOcclusionCulling |= 1u;
+		}
 		ViewIds.Append(InstanceCullingContext.ViewIds);
 
 		check(InstanceCullingContext.DynamicInstanceIdOffset >= 0);
diff --git a/Engine/Source/Runtime/Renderer/Private/InstanceCulling/InstanceCullingMergedContext.h b/Engine/Source/Runtime/Renderer/Private/InstanceCulling/InstanceCullingMergedContext.h
--- a/Engine/Source/Runtime/Renderer/Private/InstanceCulling/InstanceCullingMergedContext.h
+++ b/Engine/Source/Runtime/Renderer/Private/InstanceCulling/InstanceCullingMergedContext.h
@@ -22,14 +22,14 @@ public:
 
 	// Info about a batch of culling work produced by a context, when part of a batched job
 	// Store once per context, provides start offsets to commands/etc for the context.
-	struct FContextBatchInfo
+	struct FContextBatchInfoPacked
 	{
 		uint32 IndirectArgsOffset;
 		uint32 InstanceDataWriteOffset;
 		uint32 PayloadDataOffset;
 		uint32 CompactionDataOffset;
 		uint32 ViewIdsOffset;
-		uint32 NumViewIds;
+		uint32 NumViewIds_bAllowOcclusionCulling;
 		uint32 DynamicInstanceIdOffset;
 		uint32 DynamicInstanceIdMax;
 		uint32 ItemDataOffset[uint32(EBatchProcessingMode::Num)];
@@ -55,7 +55,7 @@ public:
 
 	TStaticArray<TInstanceCullingLoadBalancer<SceneRenderingAllocator>, static_cast<uint32>(EBatchProcessingMode::Num)> LoadBalancers;
 	TStaticArray<TArray<uint32, SceneRenderingAllocator>, static_cast<uint32>(EBatchProcessingMode::Num)> BatchInds;
-	TArray<FContextBatchInfo, SceneRenderingAllocator> BatchInfos;
+	TArray<FContextBatchInfoPacked, SceneRenderingAllocator> BatchInfos;
 
 	EShaderPlatform ShaderPlatform = SP_NumPlatforms;
 	// if true, the contexts that are supplied through calling AddBatch must all have an 1:1 entry in the resulting merged Batches array
diff --git a/Engine/Source/Runtime/Renderer/Private/SceneRendering.cpp b/Engine/Source/Runtime/Renderer/Private/SceneRendering.cpp
--- a/Engine/Source/Runtime/Renderer/Private/SceneRendering.cpp
+++ b/Engine/Source/Runtime/Renderer/Private/SceneRendering.cpp
@@ -4390,11 +4390,14 @@ void FSceneRenderer::SetupMeshPass(FViewInfo& View, FExclusiveDepthStencil::Type
 			
 			EInstanceCullingFlags CullingFlags = EInstanceCullingFlags::None;
 
+			// TODO: Maybe this should be configured somewhere else?
+			const bool bAllowInstanceOcclusionCulling = PassType != EMeshPass::CustomDepth;
+
 			FName PassName(GetMeshPassName(PassType));
 			Pass.DispatchPassSetup(
 				Scene,
 				View,
-				FInstanceCullingContext(PassName, ShaderPlatform, &InstanceCullingManager, ViewIds, View.PrevViewInfo.HZB, InstanceCullingMode, CullingFlags),
+				FInstanceCullingContext(PassName, ShaderPlatform, &InstanceCullingManager, ViewIds, bAllowInstanceOcclusionCulling ? View.PrevViewInfo.HZB : nullptr, InstanceCullingMode, CullingFlags),
 				PassType,
 				BasePassDepthStencilAccess,
 				MeshPassProcessor,
diff --git a/Engine/Source/Runtime/Renderer/Private/VirtualShadowMaps/VirtualShadowMapArray.cpp b/Engine/Source/Runtime/Renderer/Private/VirtualShadowMaps/VirtualShadowMapArray.cpp
--- a/Engine/Source/Runtime/Renderer/Private/VirtualShadowMaps/VirtualShadowMapArray.cpp
+++ b/Engine/Source/Runtime/Renderer/Private/VirtualShadowMaps/VirtualShadowMapArray.cpp
@@ -2299,7 +2299,7 @@ static FCullingResult AddCullingPasses(FRDGBuilder& GraphBuilder,
 	const TConstArrayView<uint32>& DrawCommandDescs,
 	const TConstArrayView<uint32>& InstanceIdOffsets,
 	InstanceCullingLoadBalancerType *LoadBalancer,
-	const TConstArrayView<FInstanceCullingMergedContext::FContextBatchInfo> BatchInfos,
+	const TConstArrayView<FInstanceCullingMergedContext::FContextBatchInfoPacked> BatchInfos,
 	const TConstArrayView<FVSMCullingBatchInfo> VSMCullingBatchInfos,
 	const TConstArrayView<uint32> BatchInds,
 	uint32 TotalInstances,
@@ -3027,7 +3027,7 @@ void FVirtualShadowMapArray::RenderVirtualShadowMapsNonNanite(FRDGBuilder& Graph
 	{
 		const auto VSMCullingBatchInfo = UnBatchedVSMCullingBatchInfo[Index];
 		FProjectedShadowInfo* ProjectedShadowInfo = UnBatchedVirtualSmMeshCommandPasses[Index];
-		FInstanceCullingMergedContext::FContextBatchInfo CullingBatchInfo = FInstanceCullingMergedContext::FContextBatchInfo{ 0 };
+		FInstanceCullingMergedContext::FContextBatchInfoPacked CullingBatchInfo = FInstanceCullingMergedContext::FContextBatchInfoPacked{ 0 };
 
 		FParallelMeshDrawCommandPass& MeshCommandPass = ProjectedShadowInfo->GetShadowDepthPass();
 		const TSharedPtr<FVirtualShadowMapClipmap> Clipmap = ProjectedShadowInfo->VirtualShadowMapClipmap;
