Fix incorrect light shape permutation when running point light and incorrect initialization of IES atlas index.

Tested on EngineTest ShadingModels/Light test cases.

diff --git a/Engine/Shaders/Private/DeferredLightPixelShaders.usf b/Engine/Shaders/Private/DeferredLightPixelShaders.usf
--- a/Engine/Shaders/Private/DeferredLightPixelShaders.usf
+++ b/Engine/Shaders/Private/DeferredLightPixelShaders.usf
@@ -70,7 +70,7 @@
 #elif LIGHT_SOURCE_SHAPE == 2
 #define CURRENT_LIGHT_TYPE LIGHT_TYPE_RECT
 #else // LIGHT_SOURCE_SHAPE == 1
-#define CURRENT_LIGHT_TYPE LIGHT_TYPE_SPOT
+#define CURRENT_LIGHT_TYPE (IsDeferredSpotlight() ? LIGHT_TYPE_SPOT : LIGHT_TYPE_POINT)
 #endif
 
 struct FInputParams
diff --git a/Engine/Shaders/Private/LightDataUniforms.ush b/Engine/Shaders/Private/LightDataUniforms.ush
--- a/Engine/Shaders/Private/LightDataUniforms.ush
+++ b/Engine/Shaders/Private/LightDataUniforms.ush
@@ -40,6 +40,7 @@ FDeferredLightData InitDeferredLightFromUniforms(uint InLightType)
 	Out.RectLightData.AtlasData.AtlasUVOffset	= DeferredLightUniforms.RectLightAtlasUVOffset;
 	Out.RectLightData.AtlasData.AtlasUVScale	= DeferredLightUniforms.RectLightAtlasUVScale;
 
+	Out.IESAtlasIndex = DeferredLightUniforms.IESAtlasIndex;
 	Out.LightFunctionAtlasLightIndex = DeferredLightUniforms.LightFunctionAtlasLightIndex;
 
 	Out.bAffectsTranslucentLighting = DeferredLightUniforms.bAffectsTranslucentLighting;
@@ -76,4 +77,9 @@ FDeferredLightData SetupLightDataForStandardDeferred_Mobile(uint InLightType)
 	Out.HairTransmittance		= InitHairTransmittanceData();
 
 	return Out;
+}
+
+bool IsDeferredSpotlight()
+{
+	return DeferredLightUniforms.SpotAngles.x > -2.0f;
 }
\ No newline at end of file
