HeterogeneousVolumes: Optimizing AVSM query.

This change replaces the linear--scan AVSM query with a vectorized, two-level query. Loops are eliminated when the AVSM max sample count is less than or equal to 16 (default = 8).

The shadowing bias has also been fixed to much more closely align with the reference single-scattering integrator.

#tests performance and validation testing

diff --git a/Engine/Shaders/Private/HeterogeneousVolumes/HeterogeneousVolumesAdaptiveVolumetricShadowMapSampling.ush b/Engine/Shaders/Private/HeterogeneousVolumes/HeterogeneousVolumesAdaptiveVolumetricShadowMapSampling.ush
--- a/Engine/Shaders/Private/HeterogeneousVolumes/HeterogeneousVolumesAdaptiveVolumetricShadowMapSampling.ush
+++ b/Engine/Shaders/Private/HeterogeneousVolumes/HeterogeneousVolumesAdaptiveVolumetricShadowMapSampling.ush
@@ -4,28 +4,59 @@
 
 #include "HeterogeneousVolumesAdaptiveVolumetricShadowMapUtils.ush"
 
+#define AVSM_SAMPLE_MODE_TWO_LEVEL_VECTORIZED 0
+#define AVSM_SAMPLE_MODE_TWO_LEVEL_VECTORIZED_LOOP 1
+//#define AVSM_SAMPLE_MODE_TWO_LEVEL_REFERENCE 2
+//#define AVSM_SAMPLE_MODE_BINARY_SEARCH 3
+//#define AVSM_SAMPLE_MODE_REFERENCE 4
+
+#ifndef AVSM_SAMPLE_MODE
+#define AVSM_SAMPLE_MODE AVSM_SAMPLE_MODE_TWO_LEVEL_VECTORIZED_LOOP
+#endif // AVSM_SAMPLE_MODE
+
 uint GetMyLinearIndex(uint Face, uint2 PixelCoord)
 {
 	uint LinearIndex = Face * AVSM.Resolution.x * AVSM.Resolution.y + PixelCoord.y * AVSM.Resolution.x + PixelCoord.x;
 	return LinearIndex;
 }
 
-float AVSM_Sample(int Face, int2 Pixel, float WorldHitT)
+float SafeRcp(float Value)
+{
+	return Value > 0 ? rcp(Value) : 0.0;
+}
+
+#if 0
+float AVSM_Sample_Reference(int Face, int2 Pixel, float WorldHitT)
 {
 	uint LinearIndex = GetMyLinearIndex(Face, Pixel);
 	FAVSMIndirectionData IndirectionData = ASVM_UnpackIndirectionData(AVSM.IndirectionBuffer[LinearIndex]);
+	if (WorldHitT < IndirectionData.BeginX)
+	{
+		return 1.0;
+	}
 
 	FAVSMSampleData ShadowData[2] = {
 		AVSM_CreateSampleData(0.0, 1.0),
 		AVSM_CreateSampleData(0.0, 1.0)
 	};
 
-	if (IndirectionData.SampleCount > 0)
+	int TopLevelSampleCount = CalcTopLevelSampleCount(IndirectionData.SampleCount);
+	if (TopLevelSampleCount > 0)
 	{
+		// Early-out
+		if (WorldHitT >= IndirectionData.EndX)
+		{
+			int LastSampleIndex = IndirectionData.PixelOffset + Align4(TopLevelSampleCount) + IndirectionData.SampleCount - 1;
+			return AVSM_UnpackSampleData(AVSM.SampleBuffer[LastSampleIndex]).Tau;
+		}
+
+		// Initialize from the top-level
 		ShadowData[0] = AVSM_UnpackSampleData(AVSM.SampleBuffer[IndirectionData.PixelOffset]);
-		for (int ElementIndex = 1; ElementIndex < IndirectionData.SampleCount; ++ElementIndex)
+
+		// Iterate through the linear span of offset transmittance information
+		for (int ElementIndex = 0; ElementIndex < IndirectionData.SampleCount - 1; ++ElementIndex)
 		{
-			ShadowData[1] = AVSM_UnpackSampleData(AVSM.SampleBuffer[IndirectionData.PixelOffset + ElementIndex]);
+			ShadowData[1] = AVSM_UnpackSampleData(AVSM.SampleBuffer[IndirectionData.PixelOffset + Align4(TopLevelSampleCount) + ElementIndex]);
 			if (WorldHitT < ShadowData[1].X)
 			{
 				break;
@@ -35,14 +66,229 @@ float AVSM_Sample(int Face, int2 Pixel, float WorldHitT)
 		}
 	}
 
-	float Transmittance = 1.0;
-	if (WorldHitT > ShadowData[0].X)
+	float BlendWeight = saturate((WorldHitT - ShadowData[0].X) * SafeRcp(ShadowData[1].X - ShadowData[0].X));
+	float Transmittance = lerp(ShadowData[0].Tau, ShadowData[1].Tau, BlendWeight);
+	return Transmittance;
+}
+
+float AVSM_Sample_BinarySearch(int Face, int2 Pixel, float WorldHitT)
+{
+	uint LinearIndex = GetMyLinearIndex(Face, Pixel);
+	FAVSMIndirectionData IndirectionData = ASVM_UnpackIndirectionData(AVSM.IndirectionBuffer[LinearIndex]);
+	if (WorldHitT < IndirectionData.BeginX)
+	{
+		return 1.0;
+	}
+
+	FAVSMSampleData ShadowData[2] = {
+		AVSM_CreateSampleData(0.0, 1.0),
+		AVSM_CreateSampleData(0.0, 1.0)
+	};
+
+	int TopLevelSampleCount = CalcTopLevelSampleCount(IndirectionData.SampleCount);
+	if (TopLevelSampleCount > 0)
+	{
+		// Early-out
+		if (WorldHitT >= IndirectionData.EndX)
+		{
+			int LastSampleIndex = IndirectionData.PixelOffset + Align4(TopLevelSampleCount) + IndirectionData.SampleCount - 1;
+			return AVSM_UnpackSampleData(AVSM.SampleBuffer[LastSampleIndex]).Tau;
+		}
+
+		// Initialize from the top-level
+		ShadowData[0] = AVSM_UnpackSampleData(AVSM.SampleBuffer[IndirectionData.PixelOffset]);
+
+		// Account for the shifted bottom-level
+		int ElementIndexMin = -1;
+		int ElementIndexMax = IndirectionData.SampleCount - 1;
+		while (ElementIndexMax - ElementIndexMin > 1)
+		{
+			int ElementIndex = (ElementIndexMin + ElementIndexMax + 1) / 2;
+			FAVSMSampleData SampleData = AVSM_UnpackSampleData(AVSM.SampleBuffer[IndirectionData.PixelOffset + Align4(TopLevelSampleCount) + ElementIndex]);
+
+			if (WorldHitT < SampleData.X)
+			{
+				ElementIndexMax = ElementIndex;
+				ShadowData[1] = SampleData;
+			}
+			else
+			{
+				ElementIndexMin = ElementIndex;
+				ShadowData[0] = SampleData;
+			}
+		}
+	}
+
+	float BlendWeight = saturate((WorldHitT - ShadowData[0].X) * SafeRcp(ShadowData[1].X - ShadowData[0].X));
+	float Transmittance = lerp(ShadowData[0].Tau, ShadowData[1].Tau, BlendWeight);
+	return Transmittance;
+}
+
+float AVSM_Sample_TwoLevel_Reference(int Face, int2 Pixel, float WorldHitT)
+{
+	uint LinearIndex = GetMyLinearIndex(Face, Pixel);
+	FAVSMIndirectionData IndirectionData = ASVM_UnpackIndirectionData(AVSM.IndirectionBuffer[LinearIndex]);
+	if (WorldHitT < IndirectionData.BeginX)
+	{
+		return 1.0;
+	}
+
+	FAVSMSampleData ShadowData[2] = {
+		AVSM_CreateSampleData(0.0, 1.0),
+		AVSM_CreateSampleData(0.0, 1.0)
+	};
+
+	int TopLevelSampleIndex = 0;
+
+	// Level 1
+	int TopLevelSampleCount = CalcTopLevelSampleCount(IndirectionData.SampleCount);
+	if (TopLevelSampleCount > 0)
+	{
+		if (WorldHitT >= IndirectionData.EndX)
+		{
+			int LastSampleIndex = IndirectionData.PixelOffset + Align4(TopLevelSampleCount) + IndirectionData.SampleCount - 1;
+			return AVSM_UnpackSampleData(AVSM.SampleBuffer[LastSampleIndex]).Tau;
+		}
+
+		ShadowData[0] = AVSM_UnpackSampleData(AVSM.SampleBuffer[IndirectionData.PixelOffset]);
+
+		// Iterate through the top-level
+		for (TopLevelSampleIndex = 0; TopLevelSampleIndex < TopLevelSampleCount; ++TopLevelSampleIndex)
+		{
+			ShadowData[1] = AVSM_UnpackSampleData(AVSM.SampleBuffer[IndirectionData.PixelOffset + TopLevelSampleIndex]);
+			if (WorldHitT < ShadowData[1].X)
+			{
+				break;
+			}
+
+			ShadowData[0] = ShadowData[1];
+		}
+	}
+
+	// Level 2
+	int SampleIndexMin = max((TopLevelSampleIndex - 1) * 4, 0);
+	// Account for the index shifting in the bottom-level span
+	int SampleIndexMax = min(TopLevelSampleIndex * 4, IndirectionData.SampleCount - 1);
+	for (int SampleIndex = SampleIndexMin; SampleIndex < SampleIndexMax; ++SampleIndex)
 	{
-		float BlendWeight = saturate((WorldHitT - ShadowData[0].X) * rcp(ShadowData[1].X - ShadowData[0].X));
-		Transmittance = lerp(ShadowData[0].Tau, ShadowData[1].Tau, BlendWeight);
+		ShadowData[1] = AVSM_UnpackSampleData(AVSM.SampleBuffer[IndirectionData.PixelOffset + Align4(TopLevelSampleCount) + SampleIndex]);
+		if (WorldHitT < ShadowData[1].X)
+		{
+			break;
+		}
+
+		ShadowData[0] = ShadowData[1];
 	}
+
+	// Lerp the final value
+	float BlendWeight = saturate((WorldHitT - ShadowData[0].X) * SafeRcp(ShadowData[1].X - ShadowData[0].X));
+	float Transmittance = lerp(ShadowData[0].Tau, ShadowData[1].Tau, BlendWeight);
 	return Transmittance;
 }
+#endif
+
+float AVSM_Sample_TwoLevel_Vectorized(int Face, int2 Pixel, float WorldHitT)
+{
+	uint LinearIndex = GetMyLinearIndex(Face, Pixel);
+	FAVSMIndirectionData IndirectionData = ASVM_UnpackIndirectionData(AVSM.IndirectionBuffer[LinearIndex]);
+	if (WorldHitT < IndirectionData.BeginX)
+	{
+		return 1.0;
+	}
+
+	FAVSMSampleData ShadowData[2] = {
+		AVSM_CreateSampleData(0.0, 1.0),
+		AVSM_CreateSampleData(0.0, 1.0)
+	};
+
+	int TopLevelSampleIndex = 0;
+
+	// Level 1
+	int TopLevelSampleCount = CalcTopLevelSampleCount(IndirectionData.SampleCount);
+	if (TopLevelSampleCount > 0)
+	{
+		if (WorldHitT >= IndirectionData.EndX)
+		{
+			int LastIndex = IndirectionData.PixelOffset + Align4(TopLevelSampleCount) + IndirectionData.SampleCount - 1;
+			return AVSM_UnpackSampleData(AVSM.SampleBuffer[LastIndex]).Tau;
+		}
+
+		int TopLevelSampleOffset = 0;
+#if AVSM_SAMPLE_MODE == AVSM_SAMPLE_MODE_TWO_LEVEL_VECTORIZED_LOOP
+		for (; TopLevelSampleOffset < TopLevelSampleCount; TopLevelSampleOffset += 4)
+#endif
+		{
+			int4 IndexOffset = TopLevelSampleOffset + int4(0, 1, 2, 3);
+
+			int4 PackedShadowData4 = int4(
+				AVSM.SampleBuffer[IndirectionData.PixelOffset + IndexOffset[0]],
+				AVSM.SampleBuffer[IndirectionData.PixelOffset + IndexOffset[1]],
+				AVSM.SampleBuffer[IndirectionData.PixelOffset + IndexOffset[2]],
+				AVSM.SampleBuffer[IndirectionData.PixelOffset + IndexOffset[3]]
+			);
+
+			float4 Depths = float4(
+				AVSM_UnpackSampleData(PackedShadowData4[0]).X,
+				AVSM_UnpackSampleData(PackedShadowData4[1]).X,
+				AVSM_UnpackSampleData(PackedShadowData4[2]).X,
+				AVSM_UnpackSampleData(PackedShadowData4[3]).X
+			);
+
+			int4 Result = float4(WorldHitT, WorldHitT, WorldHitT, WorldHitT) > Depths;
+			int LocalSampleIndex = Result[0] + Result[1] + Result[2] + Result[3];
+
+			ShadowData[0] = AVSM_UnpackSampleData(PackedShadowData4[max(LocalSampleIndex - 1, 0)]);
+			ShadowData[1] = AVSM_UnpackSampleData(PackedShadowData4[min(LocalSampleIndex, 3)]);
+			TopLevelSampleIndex = min(TopLevelSampleOffset + LocalSampleIndex, TopLevelSampleCount);
+#if AVSM_SAMPLE_MODE == AVSM_SAMPLE_MODE_TWO_LEVEL_VECTORIZED_LOOP
+			if (LocalSampleIndex < 4)
+			{
+				break;
+			}
+#endif
+		}
+	}
+
+	// Level 2
+	if (TopLevelSampleCount > 0)
+	{
+		int4 IndexOffset = max((TopLevelSampleIndex - 1) * 4, 0);
+		IndexOffset += int4(0, 1, 2, 3);
+
+		int4 PackedShadowData4 = int4(
+			AVSM.SampleBuffer[IndirectionData.PixelOffset + Align4(TopLevelSampleCount) + IndexOffset.x],
+			AVSM.SampleBuffer[IndirectionData.PixelOffset + Align4(TopLevelSampleCount) + IndexOffset.y],
+			AVSM.SampleBuffer[IndirectionData.PixelOffset + Align4(TopLevelSampleCount) + IndexOffset.z],
+			AVSM.SampleBuffer[IndirectionData.PixelOffset + Align4(TopLevelSampleCount) + IndexOffset.w]
+		);
+
+		float4 Depths = float4(
+			AVSM_UnpackSampleData(PackedShadowData4.x).X,
+			AVSM_UnpackSampleData(PackedShadowData4.y).X,
+			AVSM_UnpackSampleData(PackedShadowData4.z).X,
+			AVSM_UnpackSampleData(PackedShadowData4.w).X
+		);
+
+		int4 Result = float4(WorldHitT, WorldHitT, WorldHitT, WorldHitT) > Depths;
+		int Index = Result.x + Result.y + Result.z + Result.w;
+
+		if (Index > 0)
+		{
+			ShadowData[0] = AVSM_UnpackSampleData(PackedShadowData4[max(Index - 1, 0)]);
+		}
+		ShadowData[1] = AVSM_UnpackSampleData(PackedShadowData4[min(Index, 3)]);
+	}
+
+	// Lerp the final value
+	float BlendWeight = saturate((WorldHitT - ShadowData[0].X) * SafeRcp(ShadowData[1].X - ShadowData[0].X));
+	float Transmittance = lerp(ShadowData[0].Tau, ShadowData[1].Tau, BlendWeight);
+	return Transmittance;
+}
+
+float AVSM_Sample(int Face, int2 Pixel, float WorldHitT)
+{
+	return AVSM_Sample_TwoLevel_Vectorized(Face, Pixel, WorldHitT);
+}
 
 float AVSM_SampleTransmittance(float3 TranslatedWorldPosition, float3 LightTranslatedWorldPosition)
 {
@@ -75,9 +321,8 @@ float AVSM_SampleTransmittance(float3 TranslatedWorldPosition, float3 LightTrans
 			ShadowPosition = ShadowPositionAsFloat4.xyz / ShadowPositionAsFloat4.w;
 			if (all(ShadowPosition.xyz > 0) || all(ShadowPosition.xy < 1))
 			{
-				float2 ShadowMapCoord = clamp(ShadowPosition.xy * AVSM.Resolution, 0, AVSM.Resolution - 1);
-				float ShadowBias = 0.5;
-				float ShadowZ = length(TranslatedWorldPosition - AVSM.TranslatedWorldOrigin) - ShadowBias;
+				float2 ShadowMapCoord = clamp(ShadowPosition.xy * AVSM.Resolution - 0.5, 0, AVSM.Resolution - 1);
+				float ShadowZ = length(TranslatedWorldPosition - AVSM.TranslatedWorldOrigin);
 				if (AVSM.bIsDirectionalLight)
 				{
 					ShadowZ = dot(AVSM.TranslatedWorldPlane.xyz, TranslatedWorldPosition) + AVSM.TranslatedWorldPlane.w;
@@ -89,8 +334,8 @@ float AVSM_SampleTransmittance(float3 TranslatedWorldPosition, float3 LightTrans
 				int2 ShadowMapCoordX[4] =
 				{
 					clamp(ShadowMapCoord, 0, AVSM.Resolution - 1),
-					clamp(ShadowMapCoord + int2(0, 1), 0, AVSM.Resolution - 1),
 					clamp(ShadowMapCoord + int2(1, 0), 0, AVSM.Resolution - 1),
+					clamp(ShadowMapCoord + int2(0, 1), 0, AVSM.Resolution - 1),
 					clamp(ShadowMapCoord + int2(1, 1), 0, AVSM.Resolution - 1)
 				};
 
@@ -103,8 +348,8 @@ float AVSM_SampleTransmittance(float3 TranslatedWorldPosition, float3 LightTrans
 				};
 
 				float2 Weight = frac(ShadowMapCoord);
-				float TransmittanceY0 = lerp(TransmittanceX[0], TransmittanceX[2], Weight.x);
-				float TransmittanceY1 = lerp(TransmittanceX[1], TransmittanceX[3], Weight.x);
+				float TransmittanceY0 = lerp(TransmittanceX[0], TransmittanceX[1], Weight.x);
+				float TransmittanceY1 = lerp(TransmittanceX[2], TransmittanceX[3], Weight.x);
 				Transmittance = lerp(TransmittanceY0, TransmittanceY1, Weight.y);
 			#else
 				Transmittance = AVSM_Sample(Face, ShadowMapCoord, ShadowZ);
diff --git a/Engine/Shaders/Private/HeterogeneousVolumes/HeterogeneousVolumesAdaptiveVolumetricShadowMapUtils.ush b/Engine/Shaders/Private/HeterogeneousVolumes/HeterogeneousVolumesAdaptiveVolumetricShadowMapUtils.ush
--- a/Engine/Shaders/Private/HeterogeneousVolumes/HeterogeneousVolumesAdaptiveVolumetricShadowMapUtils.ush
+++ b/Engine/Shaders/Private/HeterogeneousVolumes/HeterogeneousVolumesAdaptiveVolumetricShadowMapUtils.ush
@@ -122,3 +122,14 @@ FAVSMIndirectionData ASVM_UnpackIndirectionData(uint2 PackedData)
 	IndirectionData.EndX = f16tof32(PackedData.y);
 	return IndirectionData;
 }
+
+int Align4(int Value)
+{
+	return (Value + 3) & ~3;
+}
+
+int CalcTopLevelSampleCount(int SampleCount)
+{
+	int TopLevelSampleCount = Align4(SampleCount) >> 2;
+	return TopLevelSampleCount;
+}
diff --git a/Engine/Shaders/Private/HeterogeneousVolumes/HeterogeneousVolumesVoxelGridShadows.usf b/Engine/Shaders/Private/HeterogeneousVolumes/HeterogeneousVolumesVoxelGridShadows.usf
--- a/Engine/Shaders/Private/HeterogeneousVolumes/HeterogeneousVolumesVoxelGridShadows.usf
+++ b/Engine/Shaders/Private/HeterogeneousVolumes/HeterogeneousVolumesVoxelGridShadows.usf
@@ -491,6 +491,28 @@ groupshared int GSWarpSampleCount[THREADGROUP_SIZE_2D * THREADGROUP_SIZE_2D];
 groupshared int GSWarpPrefixSum[THREADGROUP_SIZE_2D * THREADGROUP_SIZE_2D];
 groupshared int GSWarpOffset;
 
+/*
+ * CompressVolumetricShadowMapCS
+ * 
+ * Builds a compressed Adaptive Volumetric Shadow Map. The structure utilizes two buffers:
+ * 
+ * 1) Per-pixel indirection buffer:
+ *    a) Indirection pointer to transmittance function.
+ *    b) Number of samples stored in the transmittance function (max = 64).
+ *    c) Bounding depth information: (z_0, z_n)
+ * 
+ * 2) Transmittance function atlas:
+ *    The atlas is a densely packed two-level structure, containing spans of transmittance
+ *    information for a given pixel. To optimize for query performance, only the top-level
+ *    contains the first element, offsetting the second-level elements by 1 to optimize
+ *    for 4-wide loads.
+ * 
+ *    Level 1) Strided information to accelerate the query:
+ *             (z_0, tau_0) -> (z_4, tau_4) -> (z_8, tau_8) -> (z_12, tau_12)
+ * 
+ *    Level 2) Linear span of transmittance information, offset by 1:
+ *             (z_1, tau_1) -> (z_2, tau_2) -> ... -> (z_15, tau_15)
+ */
 [numthreads(THREADGROUP_SIZE_2D, THREADGROUP_SIZE_2D, 1)]
 void CompressVolumetricShadowMapCS(
 	uint2 GroupThreadId : SV_GroupThreadID,
@@ -529,7 +551,10 @@ void CompressVolumetricShadowMapCS(
 			NextPtr = AVSM_GetNextPtr(VolumetricShadowLinkedListBuffer[PixelIndex + NextPtr]);
 		}
 	}
-	GSWarpSampleCount[GroupIndex] = SampleCount;
+
+	// Add top-level skip list
+	int TopLevelSampleCount = CalcTopLevelSampleCount(SampleCount);
+	GSWarpSampleCount[GroupIndex] = Align4(TopLevelSampleCount) + Align4(SampleCount);
 	GroupMemoryBarrierWithGroupSync();
 
 	// TODO: Parallel Prefix sum
@@ -550,20 +575,46 @@ void CompressVolumetricShadowMapCS(
 	}
 	GroupMemoryBarrierWithGroupSync();
 
+	uint TopLevelOffset = GSWarpOffset + GSWarpPrefixSum[GroupIndex];
+	uint BottomLevelOffset = GSWarpOffset + GSWarpPrefixSum[GroupIndex] + Align4(TopLevelSampleCount);
+
 	// Pack and fill data per thread
-	FAVSMIndirectionDataContext Context = AVSM_CreateIndirectionDataContext(DataBegin, DataEnd);
+	FAVSMSampleData ShadowData = AVSM_CreateSampleData(0.0, 1.0);
 	int NextPtr = 0;
 	for (int SampleIndex = 0; SampleIndex < SampleCount; ++SampleIndex)
 	{
-		FAVSMSampleData ShadowData = AVSM_UnpackSampleData(VolumetricShadowLinkedListBuffer[PixelIndex + NextPtr]);
+		ShadowData = AVSM_UnpackSampleData(VolumetricShadowLinkedListBuffer[PixelIndex + NextPtr]);
 		NextPtr = AVSM_GetNextPtr(VolumetricShadowLinkedListBuffer[PixelIndex + NextPtr]);
 
-		RWVolumetricShadowTransmittanceBuffer[GSWarpOffset + GSWarpPrefixSum[GroupIndex] + SampleIndex] = AVSM_PackSampleData(ShadowData);
+		// Store indexed list offset by 1 element
+		int BottomLevelIndex = max(SampleIndex - 1, 0);
+		RWVolumetricShadowTransmittanceBuffer[BottomLevelOffset + BottomLevelIndex] = AVSM_PackSampleData(ShadowData);
+
+		// Mark the current sample as an acceleration sample at the top-level
+		if (SampleIndex % 4 == 0)
+		{
+			uint TopLevelIndex = SampleIndex / 4;
+			RWVolumetricShadowTransmittanceBuffer[TopLevelOffset + TopLevelIndex] = AVSM_PackSampleData(ShadowData);
+		}
+	}
+
+
+	// Pad remainder
+	for (int BottomLevelIndex = SampleCount - 1; BottomLevelIndex < Align4(SampleCount); ++BottomLevelIndex)
+	{
+		RWVolumetricShadowTransmittanceBuffer[BottomLevelOffset + BottomLevelIndex] = AVSM_PackSampleData(AVSM_CreateSampleData(MaxHalfFloat, ShadowData.Tau));
+	}
+
+	// Pad remainder
+	for (int TopLevelIndex = TopLevelSampleCount; TopLevelIndex < Align4(TopLevelSampleCount); ++TopLevelIndex)
+	{
+		RWVolumetricShadowTransmittanceBuffer[TopLevelOffset + TopLevelIndex] = AVSM_PackSampleData(AVSM_CreateSampleData(MaxHalfFloat, ShadowData.Tau));
 	}
 
 	// Write indirection entry
 	if (all(PixelCoord < ShadowResolution))
 	{
+		FAVSMIndirectionDataContext Context = AVSM_CreateIndirectionDataContext(DataBegin, DataEnd);
 		FAVSMIndirectionData IndirectionData = AVSM_CreateIndirectionData(
 			Context,
 			GSWarpOffset + GSWarpPrefixSum[GroupIndex],
diff --git a/Engine/Source/Runtime/Renderer/Private/HeterogeneousVolumes/HeterogeneousVolumesLiveShadingPipeline.cpp b/Engine/Source/Runtime/Renderer/Private/HeterogeneousVolumes/HeterogeneousVolumesLiveShadingPipeline.cpp
--- a/Engine/Source/Runtime/Renderer/Private/HeterogeneousVolumes/HeterogeneousVolumesLiveShadingPipeline.cpp
+++ b/Engine/Source/Runtime/Renderer/Private/HeterogeneousVolumes/HeterogeneousVolumesLiveShadingPipeline.cpp
@@ -112,7 +112,8 @@ class FRenderLightingCacheWithLiveShadingCS : public FMeshMaterialShader
 
 	class FLightingCacheMode : SHADER_PERMUTATION_INT("DIM_LIGHTING_CACHE_MODE", 2);
 	class FUseAdaptiveVolumetricShadowMap : SHADER_PERMUTATION_BOOL("DIM_USE_ADAPTIVE_VOLUMETRIC_SHADOW_MAP");
-	using FPermutationDomain = TShaderPermutationDomain<FLightingCacheMode, FUseAdaptiveVolumetricShadowMap>;
+	class FAVSMSampleMode : SHADER_PERMUTATION_INT("AVSM_SAMPLE_MODE", 2);
+	using FPermutationDomain = TShaderPermutationDomain<FLightingCacheMode, FUseAdaptiveVolumetricShadowMap, FAVSMSampleMode>;
 
 	BEGIN_SHADER_PARAMETER_STRUCT(FParameters, )
 		// Scene data
@@ -235,8 +236,9 @@ class FRenderSingleScatteringWithLiveShadingCS : public FMeshMaterialShader
 	class FUseLumenGI : SHADER_PERMUTATION_BOOL("DIM_USE_LUMEN_GI");
 	class FWriteVelocity : SHADER_PERMUTATION_BOOL("DIM_WRITE_VELOCITY");
 	class FUseAdaptiveVolumetricShadowMap : SHADER_PERMUTATION_BOOL("DIM_USE_ADAPTIVE_VOLUMETRIC_SHADOW_MAP");
+	class FAVSMSampleMode : SHADER_PERMUTATION_INT("AVSM_SAMPLE_MODE", 2);
 	class FApplyFogInscattering : SHADER_PERMUTATION_INT("APPLY_FOG_INSCATTERING", 3);
-	using FPermutationDomain = TShaderPermutationDomain<FUseTransmittanceVolume, FUseInscatteringVolume, FUseLumenGI, FWriteVelocity, FUseAdaptiveVolumetricShadowMap, FApplyFogInscattering>;
+	using FPermutationDomain = TShaderPermutationDomain<FUseTransmittanceVolume, FUseInscatteringVolume, FUseLumenGI, FWriteVelocity, FUseAdaptiveVolumetricShadowMap, FApplyFogInscattering, FAVSMSampleMode>;
 
 	BEGIN_SHADER_PARAMETER_STRUCT(FParameters, )
 		// Scene data
@@ -603,10 +605,12 @@ static void RenderLightingCacheWithLiveShading(
 	GroupCount.Z = FMath::DivideAndRoundUp(GroupCount.Z, FRenderLightingCacheWithLiveShadingCS::GetThreadGroupSize3D());
 
 	bool bUseAVSM = HeterogeneousVolumes::UseAdaptiveVolumetricShadowMapForSelfShadowing(HeterogeneousVolumeInterface->GetPrimitiveSceneProxy());
+	int32 AVSMSampleMode = HeterogeneousVolumes::GetShadowMaxSampleCount() > 16;
 
 	FRenderLightingCacheWithLiveShadingCS::FPermutationDomain PermutationVector;
 	PermutationVector.Set<FRenderLightingCacheWithLiveShadingCS::FLightingCacheMode>(HeterogeneousVolumes::GetLightingCacheMode() - 1);
 	PermutationVector.Set<FRenderLightingCacheWithLiveShadingCS::FUseAdaptiveVolumetricShadowMap>(bUseAVSM);
+	PermutationVector.Set<FRenderLightingCacheWithLiveShadingCS::FAVSMSampleMode>(AVSMSampleMode);
 	TShaderRef<FRenderLightingCacheWithLiveShadingCS> ComputeShader = Material.GetShader<FRenderLightingCacheWithLiveShadingCS>(&FLocalVertexFactory::StaticType, PermutationVector, false);
 	if (!ComputeShader.IsNull())
 	{
@@ -655,6 +659,7 @@ static void RenderSingleScatteringWithLiveShading(
 	}
 	TUniformBufferRef<FDeferredLightUniformStruct> DeferredLightUB = CreateUniformBufferImmediate(DeferredLightUniform, UniformBuffer_SingleDraw);
 
+	bool bUseAVSM = HeterogeneousVolumes::UseAdaptiveVolumetricShadowMapForSelfShadowing(HeterogeneousVolumeInterface->GetPrimitiveSceneProxy());
 	bool bWriteVelocity = HeterogeneousVolumes::ShouldWriteVelocity() && HasBeenProduced(SceneTextures.Velocity);
 	FRenderSingleScatteringWithLiveShadingCS::FParameters* PassParameters = GraphBuilder.AllocParameters<FRenderSingleScatteringWithLiveShadingCS::FParameters>();
 	{
@@ -756,8 +761,16 @@ static void RenderSingleScatteringWithLiveShading(
 		}
 		else
 		{
-			PassParameters->LightingCache.LightingCacheResolution = FIntVector::ZeroValue;
-			PassParameters->LightingCache.LightingCacheVoxelBias = 0.0f;
+			if (bUseAVSM)
+			{
+				PassParameters->LightingCache.LightingCacheResolution = HeterogeneousVolumes::GetLightingCacheResolution(HeterogeneousVolumeInterface, LODFactor);
+				PassParameters->LightingCache.LightingCacheVoxelBias = HeterogeneousVolumeInterface->GetShadowBiasFactor();
+			}
+			else
+			{
+				PassParameters->LightingCache.LightingCacheResolution = FIntVector::ZeroValue;
+				PassParameters->LightingCache.LightingCacheVoxelBias = 0.0f;
+			}
 			PassParameters->LightingCache.LightingCacheTexture = FRDGSystemTextures::Get(GraphBuilder).VolumetricBlack;
 		}
 
@@ -787,7 +800,7 @@ static void RenderSingleScatteringWithLiveShading(
 	}
 #endif // WANTS_DRAW_MESH_EVENTS
 
-	bool bUseAVSM = HeterogeneousVolumes::UseAdaptiveVolumetricShadowMapForSelfShadowing(HeterogeneousVolumeInterface->GetPrimitiveSceneProxy());
+	int32 AVSMSampleMode = HeterogeneousVolumes::GetShadowMaxSampleCount() > 16;
 
 	FRenderSingleScatteringWithLiveShadingCS::FPermutationDomain PermutationVector;
 	PermutationVector.Set<FRenderSingleScatteringWithLiveShadingCS::FUseTransmittanceVolume>(HeterogeneousVolumes::UseLightingCacheForTransmittance() && PassParameters->bApplyShadowTransmittance);
@@ -795,6 +808,7 @@ static void RenderSingleScatteringWithLiveShading(
 	PermutationVector.Set<FRenderSingleScatteringWithLiveShadingCS::FUseLumenGI>(HeterogeneousVolumes::UseIndirectLighting() && View.GetLumenTranslucencyGIVolume().Texture0 != nullptr);
 	PermutationVector.Set<FRenderSingleScatteringWithLiveShadingCS::FWriteVelocity>(bWriteVelocity);
 	PermutationVector.Set<FRenderSingleScatteringWithLiveShadingCS::FUseAdaptiveVolumetricShadowMap>(bUseAVSM);
+	PermutationVector.Set<FRenderSingleScatteringWithLiveShadingCS::FAVSMSampleMode>(AVSMSampleMode);
 	PermutationVector.Set<FRenderSingleScatteringWithLiveShadingCS::FApplyFogInscattering>(static_cast<int32>(HeterogeneousVolumes::GetApplyFogInscattering()));
 	PermutationVector = FRenderSingleScatteringWithLiveShadingCS::RemapPermutation(PermutationVector);
 	TShaderRef<FRenderSingleScatteringWithLiveShadingCS> ComputeShader = Material.GetShader<FRenderSingleScatteringWithLiveShadingCS>(&FLocalVertexFactory::StaticType, PermutationVector, false);
diff --git a/Engine/Source/Runtime/Renderer/Private/HeterogeneousVolumes/HeterogeneousVolumesVoxelGridPipeline.cpp b/Engine/Source/Runtime/Renderer/Private/HeterogeneousVolumes/HeterogeneousVolumesVoxelGridPipeline.cpp
--- a/Engine/Source/Runtime/Renderer/Private/HeterogeneousVolumes/HeterogeneousVolumesVoxelGridPipeline.cpp
+++ b/Engine/Source/Runtime/Renderer/Private/HeterogeneousVolumes/HeterogeneousVolumesVoxelGridPipeline.cpp
@@ -3303,9 +3303,10 @@ void CompressVolumetricShadowMap(
 		TEXT("HeterogeneousVolume.VolumetricShadowIndirectionBuffer")
 	);
 
-	int32 FVolumetricShadowTransmittanceMaxCount = VolumetricShadowPixelCount * MaxSampleCount;
+	int32 TopLevelSampleCount = FMath::CeilToInt(MaxSampleCount / 4.0);
+	int32 VolumetricShadowTransmittanceMaxCount = VolumetricShadowPixelCount * (Align(TopLevelSampleCount, 4) + Align(MaxSampleCount, 4));
 	VolumetricShadowTransmittanceBuffer = GraphBuilder.CreateBuffer(
-		FRDGBufferDesc::CreateStructuredDesc(sizeof(FAVSMSamplePackedData), FVolumetricShadowTransmittanceMaxCount),
+		FRDGBufferDesc::CreateStructuredDesc(sizeof(FAVSMSamplePackedData), VolumetricShadowTransmittanceMaxCount),
 		TEXT("HeterogeneousVolume.VolumetricShadowTransmittanceBuffer")
 	);
 
