Path Tracer: Add an opacity cutoff to volume tracing so that ray marching can stop as we approach full opacity.

This saves ~13% with default settings and ~30% with all bounces enables. There are not significant visual changes even though this is technically slightly biased.

diff --git a/Engine/Shaders/Private/PathTracing/PathTracingCore.ush b/Engine/Shaders/Private/PathTracing/PathTracingCore.ush
--- a/Engine/Shaders/Private/PathTracing/PathTracingCore.ush
+++ b/Engine/Shaders/Private/PathTracing/PathTracingCore.ush
@@ -563,7 +563,7 @@ FPackedPathTracingPayload TraceTransparentRay(inout FPathState PathState, int Bo
 				float3 PathThroughput = PathState.PathThroughput;
 
 				// In order to use analytic transmittance, the user must have requested it _and_ we must be on an interval with no clouds or heterogeneous volumes
-				const bool bUseAnalyticTransmittance =
+				bool bUseAnalyticTransmittance =
 					(VolumeFlags & PATH_TRACER_VOLUME_USE_ANALYTIC_TRANSMITTANCE) != 0 &&
 					(Interval.VolumeMask & (PATH_TRACER_VOLUME_ENABLE_CLOUDS | PATH_TRACER_VOLUME_ENABLE_HETEROGENEOUS_VOLUMES)) == 0;
 				// are there any holdouts on the current segment?
@@ -633,6 +633,13 @@ FPackedPathTracingPayload TraceTransparentRay(inout FPathState PathState, int Bo
 							VolumeRadiance += Result.Emission * PathThroughput;
 							// keep tracing through the volume
 							PathThroughput *= SigmaN;
+							// early termination - fades off from 0.002 to 0.001 and then will be 0
+							PathThroughput *= saturate(1000.0 * max3(PathThroughput.x, PathThroughput.y, PathThroughput.z) - 1.0);
+							if (!any(PathThroughput > 0))
+							{
+								bUseAnalyticTransmittance = false; // avoid re-computing the transmittance after the loop
+								break;
+							}
 						}
 						else
 						{
diff --git a/Engine/Shaders/Private/PathTracing/Volume/PathTracingVolume.ush b/Engine/Shaders/Private/PathTracing/Volume/PathTracingVolume.ush
--- a/Engine/Shaders/Private/PathTracing/Volume/PathTracingVolume.ush
+++ b/Engine/Shaders/Private/PathTracing/Volume/PathTracingVolume.ush
@@ -99,7 +99,7 @@ FVolumeShadedResult VolumeGetDensity(float3 TranslatedWorldPos, uint VolumeMask,
 		/* take stochastic steps along the ray to estimate transmittance (null scattering) */		\
 		float3 ColorChannelPdf = Throughput;														\
 		/* Sample the distance to the next interaction */											\
-		float2 RandValue = RandomSequence_GenerateSample2D(RandSequence);							\
+		float RandValue = RandomSequence_GenerateSample1D(RandSequence);							\
 		float DeltaT = SampleSpectralTransmittance(RandValue.x, SigmaBar, ColorChannelPdf);			\
 		if (DeltaT < 0.0)																			\
 		{																							\
@@ -118,6 +118,12 @@ FVolumeShadedResult VolumeGetDensity(float3 TranslatedWorldPos, uint VolumeMask,
 			float3 SigmaN = SigmaBar - SigmaT;														\
 			/* keep tracing through the volume */													\
 			Throughput *= SigmaN;																	\
+			/* early termination - fades off from 0.002 to 0.001 and then will be 0 */				\
+			Throughput *= saturate(1000.0 * max3(Throughput.x, Throughput.y, Throughput.z) - 1.0);	\
+			if (!any(Throughput > 0))																\
+			{																						\
+				break;																				\
+			}																						\
 		}																							\
 		else                                                                                        \
 		{																							\
diff --git a/Engine/Shaders/Private/PathTracing/Volume/PathTracingVolumeCommon.ush b/Engine/Shaders/Private/PathTracing/Volume/PathTracingVolumeCommon.ush
--- a/Engine/Shaders/Private/PathTracing/Volume/PathTracingVolumeCommon.ush
+++ b/Engine/Shaders/Private/PathTracing/Volume/PathTracingVolumeCommon.ush
@@ -132,7 +132,7 @@ struct FVolumeIntersectionList
 				}
 				else
 				{
-					// segments could be disjoint (TMax is closer) or overlapping (TMax is closer)
+					// segments could be disjoint (TMax is closer) or overlapping (VolumeTMin is closer)
 					Result.VolumeTMax = min(TMax, Result.VolumeTMin);
 					Result.VolumeTMin = TMin; // we just established this one is closer
 				}
