Fix Segment Index in Ray Tracing Picker visualization.
- GeometryIndex was not being set in RayTracingDebugMainCHS.

diff --git a/Engine/Shaders/Private/RayTracing/RayTracingDebugCHS.usf b/Engine/Shaders/Private/RayTracing/RayTracingDebugCHS.usf
--- a/Engine/Shaders/Private/RayTracing/RayTracingDebugCHS.usf
+++ b/Engine/Shaders/Private/RayTracing/RayTracingDebugCHS.usf
@@ -23,7 +23,7 @@ RAY_TRACING_ENTRY_CLOSEST_HIT(RayTracingDebugMainCHS,
 
 	Payload.WorldNormal = GetGeometryNormalFromTriangleBaseAttributes(PrimitiveIndex());
 	Payload.InstanceIndex = InstanceIndex();
-	Payload.GeometryIndex = 0; // GeometryIndex();
+	Payload.GeometryIndex = GeometryIndex();
 	Payload.ScenePrimitiveIndex = InstanceSceneData.PrimitiveId;
 }
 
