[Backout] - CL33653653

GeometryIndex() is not currently supported in UE (see RayTracingCommon.ush)

Original CL Desc
-----------------------------------------------------------------
Fix Segment Index in Ray Tracing Picker visualization.
- GeometryIndex was not being set in RayTracingDebugMainCHS.

diff --git a/Engine/Shaders/Private/RayTracing/RayTracingDebugCHS.usf b/Engine/Shaders/Private/RayTracing/RayTracingDebugCHS.usf
--- a/Engine/Shaders/Private/RayTracing/RayTracingDebugCHS.usf
+++ b/Engine/Shaders/Private/RayTracing/RayTracingDebugCHS.usf
@@ -23,7 +23,7 @@ RAY_TRACING_ENTRY_CLOSEST_HIT(RayTracingDebugMainCHS,
 
 	Payload.WorldNormal = GetGeometryNormalFromTriangleBaseAttributes(PrimitiveIndex());
 	Payload.InstanceIndex = InstanceIndex();
-	Payload.GeometryIndex = GeometryIndex();
+	Payload.GeometryIndex = 0; // GeometryIndex();
 	Payload.ScenePrimitiveIndex = InstanceSceneData.PrimitiveId;
 }
 
