Tonemap:  Fix for editor viewport rendering bug in Tonemap compute shader code path.  UVs need to be computed based on the input texture rect, not the output.  Fixes smearing in editor viewports when running with split panes, and still keeps split screen working.

#rnx

diff --git a/Engine/Shaders/Private/PostProcessTonemap.usf b/Engine/Shaders/Private/PostProcessTonemap.usf
--- a/Engine/Shaders/Private/PostProcessTonemap.usf
+++ b/Engine/Shaders/Private/PostProcessTonemap.usf
@@ -616,7 +616,7 @@ RWTexture2D<float> RWOutputLuminance;
 void MainCS(uint2 DispatchThreadId : SV_DispatchThreadID)
 {
 	float4 SvPosition = float4((float2)DispatchThreadId + 0.5f + Output_ViewportMin, 0.0f, 1.0f);
-	float2 UV = ((float2)DispatchThreadId + 0.5f) * Output_ViewportSizeInverse;
+	float2 UV = ((float2)DispatchThreadId + 0.5f + Color_ViewportMin) * Color_ExtentInverse;
 
 	if (IsComputeUVOutOfBounds(UV))
 	{
