Fix mobile decal shader compile error

diff --git a/Engine/Shaders/Private/MobileLightingCommon.ush b/Engine/Shaders/Private/MobileLightingCommon.ush
--- a/Engine/Shaders/Private/MobileLightingCommon.ush
+++ b/Engine/Shaders/Private/MobileLightingCommon.ush
@@ -11,6 +11,7 @@
 #include "ShadowFilteringCommon.ush"
 #include "ClearCoatCommon.ush"
 #include "DeferredLightingCommon.ush"
+#include "SHCommon.ush"
 
 #ifndef TRANSLUCENCY_NON_DIRECTIONAL
 #define TRANSLUCENCY_NON_DIRECTIONAL 0
@@ -217,6 +218,8 @@ FTranslucencyLightingVector GetTranslucencyLightingVector(FDeferredLightData Lig
 	return Out;
 }
 
+#if TRANSLUCENCY_SH_LIGHTING
+
 FLightAccumulator GetTranslucencySHLighting(FGBufferData GBuffer, FDeferredLightData LightData, float3 TranslatedWorldPosition)
 {
 	float4 VolumeLighting;
@@ -239,6 +242,8 @@ FLightAccumulator GetTranslucencySHLighting(FGBufferData GBuffer, FDeferredLight
 	return LightAccumulator;
 }
 
+#endif
+
 /*------------------------------------------------------------------------------
 	Mobile Directional Light.
 ------------------------------------------------------------------------------*/
