Substrate - allow the SSS material attribute blending of legacy Subsurface and SubstrateDiffusion.

diff --git a/Engine/Shaders/Private/Substrate/SubstrateExport.ush b/Engine/Shaders/Private/Substrate/SubstrateExport.ush
--- a/Engine/Shaders/Private/Substrate/SubstrateExport.ush
+++ b/Engine/Shaders/Private/Substrate/SubstrateExport.ush
@@ -134,6 +134,14 @@ void PackSubstrateOut(
 	TopLayerData = (FSubstrateTopLayerData)0;
 	OpaqueRoughRefractionData = (FSubstrateOpaqueRoughRefractionData)0;
 
+	// The SSSData we accumulate to allow blending of Wrap and SubstrateDiffusion types.
+	uint SSSDataType = SSS_TYPE_INVALID;
+	float SSSDataAniso = 0.0f;
+	float SSSDataProfileID = 0.0f;
+	float SSSDataRadiusScale = 0.0f;
+	float3 SSSDataMFP = 0.0f;
+	float3 SSSDataBaseColor = 0.0f;
+
 	uint BSDFVisibleCount = 0;
 	
 	if (SubstratePixelHeader.SubstrateTree.BSDFCount > 0)
@@ -746,21 +754,28 @@ void PackSubstrateOut(
 								// All SSS types store side data (SSSHeader/SSSData), apart from SIMPLEVOLUME
 								if (BSDF_GETSSSTYPE(BSDF) != SSS_TYPE_SIMPLEVOLUME)
 								{
-									SubstrateSubSurfaceHeaderSetSSSType(SSSData.Header, BSDF_GETSSSTYPE(BSDF));
+									// Always override the type, we can only blend similar types of subsurface lighting effects.
+									SSSDataType = BSDF_GETSSSTYPE(BSDF);
+
+									// Now we are going to blend all the SSS attributes in order to be able to lerp legacy Subsurface and SubstrateDiffusion.
+									// Legacy SSSProfile ID is still set by the last encountered Slab.
+									// Attributes are blended using opacity only, no need to use the throughput to the view (SSS post process wants raw source data like BaseColor and the BSDF will otherwise run the correct lighting math).
+
 									if (BSDF_GETSSSTYPE(BSDF) == SSS_TYPE_WRAP)
 									{
-										SubstrateSubSurfaceHeaderSetWrap(SSSData.Header, SLAB_SSSPHASEANISOTROPY(BSDF));
+										SSSDataAniso += SLAB_SSSPHASEANISOTROPY(BSDF) * BSDF.Coverage;
 									}
 									else if (BSDF_GETSSSTYPE(BSDF) == SSS_TYPE_DIFFUSION_PROFILE)
 									{
-										SubstrateSubSurfaceHeaderSetProfile(SSSData.Header, SLAB_SSSPROFILERADIUSSCALE(BSDF), SubstrateSubsurfaceProfileIdTo8bits(SLAB_SSSPROFILEID(BSDF)));
+										SSSDataProfileID = SLAB_SSSPROFILEID(BSDF);
+										SSSDataRadiusScale = SLAB_SSSPROFILERADIUSSCALE(BSDF);
 									}
 									else
 									{
-										SubstrateSubSurfaceHeaderSetNonProfile(SSSData.Header, SLAB_SSSMFP(BSDF));
+										SSSDataMFP += SLAB_SSSMFP(BSDF) * BSDF.Coverage;
 									}
 
-									SubstrateSubsurfaceExtrasSetBaseColor(SSSData.Extras, SLAB_DIFFUSEALBEDO(BSDF));
+									SSSDataBaseColor += SLAB_DIFFUSEALBEDO(BSDF) * BSDF.Coverage;
 								}
 
 								if (BSDF_GETSSSTYPE(BSDF) == SSS_TYPE_DIFFUSION_PROFILE)
@@ -838,9 +853,10 @@ void PackSubstrateOut(
 							const bool bHasSSS = BSDF_GETSSSTYPE(BSDF) == SSS_TYPE_DIFFUSION_PROFILE;
 							if (bHasSSS)
 							{
-								SubstrateSubSurfaceHeaderSetSSSType(SSSData.Header, SSS_TYPE_DIFFUSION_PROFILE);
-								SubstrateSubSurfaceHeaderSetProfile(SSSData.Header, 1.0f, SubstrateSubsurfaceProfileIdTo8bits(EYE_SSSPROFILEID(BSDF)));
-								SubstrateSubsurfaceExtrasSetBaseColor(SSSData.Extras, EYE_DIFFUSEALBEDO(BSDF));
+								SSSDataType = SSS_TYPE_DIFFUSION_PROFILE;
+								SSSDataProfileID = EYE_SSSPROFILEID(BSDF);
+								SSSDataRadiusScale = 1.0f;
+								SSSDataBaseColor = EYE_DIFFUSEALBEDO(BSDF) * BSDF.Coverage;
 							}
 						}
 						break;
@@ -863,6 +879,32 @@ void PackSubstrateOut(
 		}
 
 		FinalizeWrites(SubstrateBuffer, SubstrateAddressing);
+
+		// Now generate the final SSSData structure from the parameter blended attributes
+		if (SSSDataType != SSS_TYPE_INVALID)
+		{
+			// All SSS types store side data (SSSHeader/SSSData), apart from SIMPLEVOLUME
+			if (SSSDataType != SSS_TYPE_SIMPLEVOLUME)
+			{
+				// Always override the type, we can only blend similar types of subsrface
+				SubstrateSubSurfaceHeaderSetSSSType(SSSData.Header, SSSDataType);
+
+				if (SSSDataType == SSS_TYPE_WRAP)
+				{
+					SubstrateSubSurfaceHeaderSetWrap(SSSData.Header, SSSDataAniso);
+				}
+				else if (SSSDataType == SSS_TYPE_DIFFUSION_PROFILE)
+				{
+					SubstrateSubSurfaceHeaderSetProfile(SSSData.Header, SSSDataRadiusScale, SubstrateSubsurfaceProfileIdTo8bits(SSSDataProfileID));
+				}
+				else
+				{
+					SubstrateSubSurfaceHeaderSetNonProfile(SSSData.Header, SSSDataMFP);
+				}
+
+				SubstrateSubsurfaceExtrasSetBaseColor(SSSData.Extras, SSSDataBaseColor);
+			}
+		}
 	}
 #endif // SUBSTRATE_INLINE_SHADING && SUBSTRATE_CLAMPED_CLOSURE_COUNT > 0
 
