Change groom skinning to use mesh's vertex normal instead of face vertex normal.

This helps smooth the groom look when the underlying mesh has rough triangulation.

diff --git a/Engine/Plugins/Runtime/HairStrands/Source/HairStrandsCore/Private/GroomBindingBuilder.cpp b/Engine/Plugins/Runtime/HairStrands/Source/HairStrandsCore/Private/GroomBindingBuilder.cpp
--- a/Engine/Plugins/Runtime/HairStrands/Source/HairStrandsCore/Private/GroomBindingBuilder.cpp
+++ b/Engine/Plugins/Runtime/HairStrands/Source/HairStrandsCore/Private/GroomBindingBuilder.cpp
@@ -60,7 +60,7 @@ static FAutoConsoleVariableRef CVarHairStrandsBindingBuilderWarningEnable(TEXT("
 FString FGroomBindingBuilder::GetVersion()
 {
 	// Important to update the version when groom building changes
-	return TEXT("4e");
+	return TEXT("4f");
 }
 
 ///////////////////////////////////////////////////////////////////////////////////////////////////
@@ -105,6 +105,18 @@ float FHairStrandsRootUtils::PackUVsToFloat(const FVector2f& UV)
 	return *((float*)(&Encoded));
 }
 
+static float PackNormalToFloat(const FVector3f& InN)
+{
+	FVector3f N = InN; 
+	N.Normalize();
+	const FVector3f NN = N * 0.5f + 0.5f;
+	const uint32 Encoded = 
+		(uint32(FMath::Clamp(NN.X, 0.f, 1.f) * 1023.0f) & 0x3FF)    |
+		(uint32(FMath::Clamp(NN.Y, 0.f, 1.f) * 1023.0f) & 0x3FF)<<10|
+		(uint32(FMath::Clamp(NN.Z, 0.f, 1.f) * 1023.0f) & 0x3FF)<<20;
+	return *((float*)(&Encoded));
+}
+
 //////////////////////////////////////////////////////////////////////////
 // Intermediate data struct
 
@@ -144,6 +156,7 @@ public:
 	virtual const IMeshSectionData& GetSection(uint32 SectionIndex) const = 0;
 	virtual const FVector3f& GetVertexPosition(uint32 VertexIndex) const = 0;
 	virtual FVector2f GetVertexUV(uint32 VertexIndex, uint32 ChannelIndex) const = 0;
+	virtual FVector3f GetVertexNormal(uint32 VertexIndex) const = 0;
 	virtual int32 GetSectionFromVertexIndex(uint32 InVertIndex) const = 0;
 	virtual ~IMeshLODData() {}
 };
@@ -276,6 +289,13 @@ public:
 		return FVector2f(MeshData->LODRenderData[LODIndex].StaticVertexBuffers.StaticMeshVertexBuffer.GetVertexUV(InVertexIndex, InChannelIndex));
 	}
 
+	virtual FVector3f GetVertexNormal(uint32 InVertexIndex) const override
+	{
+		check(MeshData);
+		check(MeshData->LODRenderData.IsValidIndex(LODIndex));
+		return FVector3f(MeshData->LODRenderData[LODIndex].StaticVertexBuffers.StaticMeshVertexBuffer.VertexTangentZ(InVertexIndex));
+	}
+
 	virtual int32 GetSectionFromVertexIndex(uint32 InVertIndex) const override
 	{
 		int32 OutSectionIndex = 0;
@@ -475,6 +495,11 @@ public:
 		return FVector2f(MeshData.TextureCoordinates[VertexIndex]);
 	}
 
+	virtual FVector3f GetVertexNormal(uint32 InVertexIndex) const override
+	{
+		return MeshData.TangentsZ[InVertexIndex].ToFVector3f();
+	}
+
 	virtual int32 GetSectionFromVertexIndex(uint32 InVertIndex) const override
 	{
 		for (int32 SectionIndex = 0; SectionIndex < SectionRanges.Num(); ++SectionIndex)
@@ -869,6 +894,10 @@ namespace GroomBinding_RootProjection
 			FVector3f P1;
 			FVector3f P2;
 
+			FVector3f N0;
+			FVector3f N1;
+			FVector3f N2;
+
 			FVector2f UV0;
 			FVector2f UV1;
 			FVector2f UV2;
@@ -1238,6 +1267,10 @@ namespace GroomBinding_RootProjection
 					T.UV1 = MeshLODData.GetVertexUV(T.I1, ChannelIndex);
 					T.UV2 = MeshLODData.GetVertexUV(T.I2, ChannelIndex);
 
+					T.N0 = MeshLODData.GetVertexNormal(T.I0);
+					T.N1 = MeshLODData.GetVertexNormal(T.I1);
+					T.N2 = MeshLODData.GetVertexNormal(T.I2);
+
 					MeshBound += T.P0;
 					MeshBound += T.P1;
 					MeshBound += T.P2;
@@ -1321,6 +1354,10 @@ namespace GroomBinding_RootProjection
 					T.UV1 = MeshLODData.GetVertexUV(T.I1, ChannelIndex);
 					T.UV2 = MeshLODData.GetVertexUV(T.I2, ChannelIndex);
 
+					T.N0 = MeshLODData.GetVertexNormal(T.I0);
+					T.N1 = MeshLODData.GetVertexNormal(T.I1);
+					T.N2 = MeshLODData.GetVertexNormal(T.I2);
+
 					bIsGridPopulated = Grid.Insert(T) || bIsGridPopulated;
 				}
 			}
@@ -1397,9 +1434,9 @@ namespace GroomBinding_RootProjection
 				OutRootData[MeshLODIt].RootBarycentricBuffer[CurveIndex] = EncodedBarycentrics;
 
 				RootTriangleIndexBuffer[CurveIndex] = EncodedTriangleIndex;
-				RestRootTrianglePositionBuffer[CurveIndex * 3 + 0] = FVector4f((FVector3f)ClosestTriangle.P0, FHairStrandsRootUtils::PackUVsToFloat(FVector2f(ClosestTriangle.UV0)));	// LWC_TODO: Precision loss
-				RestRootTrianglePositionBuffer[CurveIndex * 3 + 1] = FVector4f((FVector3f)ClosestTriangle.P1, FHairStrandsRootUtils::PackUVsToFloat(FVector2f(ClosestTriangle.UV1)));	// LWC_TODO: Precision loss
-				RestRootTrianglePositionBuffer[CurveIndex * 3 + 2] = FVector4f((FVector3f)ClosestTriangle.P2, FHairStrandsRootUtils::PackUVsToFloat(FVector2f(ClosestTriangle.UV2)));	// LWC_TODO: Precision loss
+				RestRootTrianglePositionBuffer[CurveIndex * 3 + 0] = FVector4f((FVector3f)ClosestTriangle.P0, PackNormalToFloat(ClosestTriangle.N0));
+				RestRootTrianglePositionBuffer[CurveIndex * 3 + 1] = FVector4f((FVector3f)ClosestTriangle.P1, PackNormalToFloat(ClosestTriangle.N1));
+				RestRootTrianglePositionBuffer[CurveIndex * 3 + 2] = FVector4f((FVector3f)ClosestTriangle.P2, PackNormalToFloat(ClosestTriangle.N2));
 			}
 		#if BINDING_PARALLEL_BUILDING
 			);
@@ -1793,6 +1830,11 @@ namespace GroomBinding_Transfer
 				T.UV1 = MeshLODData.GetVertexUV(T.I1, InChannelIndex);
 				T.UV2 = MeshLODData.GetVertexUV(T.I2, InChannelIndex);
 
+				// Needed??
+				//T.N0 = MeshLODData.GetVertexNormal(T.I0);
+				//T.N1 = MeshLODData.GetVertexNormal(T.I1);
+				//T.N2 = MeshLODData.GetVertexNormal(T.I2);
+
 				bIsGridPopulated = OutGrid.Insert(T) || bIsGridPopulated;
 			}
 
diff --git a/Engine/Plugins/Runtime/HairStrands/Source/HairStrandsCore/Private/GroomDebug.cpp b/Engine/Plugins/Runtime/HairStrands/Source/HairStrandsCore/Private/GroomDebug.cpp
--- a/Engine/Plugins/Runtime/HairStrands/Source/HairStrandsCore/Private/GroomDebug.cpp
+++ b/Engine/Plugins/Runtime/HairStrands/Source/HairStrandsCore/Private/GroomDebug.cpp
@@ -76,7 +76,7 @@ class FHairProjectionHairDebugCS : public FGlobalShader
 	DECLARE_GLOBAL_SHADER(FHairProjectionHairDebugCS);
 	SHADER_USE_PARAMETER_STRUCT(FHairProjectionHairDebugCS, FGlobalShader);	
 
-	class FInputType : SHADER_PERMUTATION_INT("PERMUTATION_INPUT_TYPE", 3);
+	class FInputType : SHADER_PERMUTATION_INT("PERMUTATION_INPUT_TYPE", 2);
 	using FPermutationDomain = TShaderPermutationDomain<FInputType>;
 
 	BEGIN_SHADER_PARAMETER_STRUCT(FParameters, )
@@ -131,8 +131,7 @@ IMPLEMENT_GLOBAL_SHADER(FHairProjectionHairDebugCS, "/Engine/Private/HairStrands
 
 enum class EDebugProjectionHairType
 {
-	HairFrame,
-	HairTriangle,
+	HairRoots,
 	HairSamples,
 };
 
@@ -165,7 +164,6 @@ static void AddDebugProjectionHairPass(
 		return;
 	}
 	
-	const EPrimitiveType PrimitiveType = GeometryType == EDebugProjectionHairType::HairFrame ? PT_LineList : GeometryType == EDebugProjectionHairType::HairTriangle ? PT_TriangleList : PT_LineList;
 	const uint32 RootCount = EDebugProjectionHairType::HairSamples == GeometryType ? 3 * RestRootResources->GetLOD(MeshLODIndex)->SampleCount : RestRootResources->GetRootCount();
 	const uint32 PrimitiveCount = RootCount;
 
@@ -179,10 +177,10 @@ static void AddDebugProjectionHairPass(
 		return;
 	}
 
-	ShaderPrint::RequestSpaceForLines(PrimitiveCount * 4 /* 1 hair root + 3 triangles */);
+	ShaderPrint::RequestSpaceForLines(PrimitiveCount * 10 /* 3+1 hair root + 3 triangles + 3 normals */);
 	ShaderPrint::RequestSpaceForTriangles(PrimitiveCount * 3);
 
-	if (EDebugProjectionHairType::HairFrame == GeometryType &&
+	if (EDebugProjectionHairType::HairRoots == GeometryType &&
 		!RestRootResources->GetLOD(MeshLODIndex)->RootBarycentricBuffer.Buffer)
 	{
 		return;
@@ -191,7 +189,7 @@ static void AddDebugProjectionHairPass(
 	if (EDebugProjectionHairType::HairSamples == GeometryType &&
 		!RestRootResources->GetLOD(MeshLODIndex)->RestSamplePositionsBuffer.Buffer)
 	{
-			return;
+		return;
 	}
 
 	const FHairStrandsLODRestRootResource& RestLODDatas = *RestRootResources->GetLOD(MeshLODIndex);
@@ -230,7 +228,7 @@ static void AddDebugProjectionHairPass(
 	Parameters->Sim_HairDeformedFrames		= GHairDebugMeshProjection_Sim_HairDeformedFrames > 0 ? 1u : 0u;
 	Parameters->Sim_HairDeformedSamples		= GHairDebugMeshProjection_Sim_HairDeformedSamples > 0 ? 1u : 0u;
 
-	if (EDebugProjectionHairType::HairFrame == GeometryType)
+	if (EDebugProjectionHairType::HairRoots == GeometryType)
 	{
 		Parameters->RootBarycentricBuffer = RegisterAsSRV(GraphBuilder, RestLODDatas.RootBarycentricBuffer);
 	}
@@ -247,14 +245,14 @@ static void AddDebugProjectionHairPass(
 	ShaderPrint::SetParameters(GraphBuilder, *ShaderPrintData, Parameters->ShaderPrintUniformBuffer);
 
 	FHairProjectionHairDebugCS::FPermutationDomain PermutationVector;
-	PermutationVector.Set<FHairProjectionHairDebugCS::FInputType>(GeometryType == EDebugProjectionHairType::HairFrame ? 0 : GeometryType == EDebugProjectionHairType::HairTriangle ? 1 : 2);
+	PermutationVector.Set<FHairProjectionHairDebugCS::FInputType>(GeometryType == EDebugProjectionHairType::HairRoots ? 0 : 1);
 
 	TShaderMapRef<FHairProjectionHairDebugCS> ComputeShader(ShaderMap, PermutationVector);
 	ClearUnusedGraphResources(ComputeShader, Parameters);
 
 	FComputeShaderUtils::AddPass(
 		GraphBuilder,
-		RDG_EVENT_NAME("HairStrands::MeshProjectionDebug(Hair, %s)", GeometryType == EDebugProjectionHairType::HairFrame ? TEXT("Frame") : GeometryType == EDebugProjectionHairType::HairTriangle ? TEXT("Triangle") : TEXT("Samples")),
+		RDG_EVENT_NAME("HairStrands::MeshProjectionDebug(Hair, %s)", GeometryType == EDebugProjectionHairType::HairRoots ? TEXT("Root") : TEXT("Samples")),
 		ComputeShader,
 		Parameters,
 		FIntVector(FMath::DivideAndRoundUp(PrimitiveCount, 256u), 1, 1));
@@ -1271,11 +1269,9 @@ void RunHairStrandsDebug(
 
 					const int32 MeshLODIndex = Instance->Debug.MeshLODIndex;
 
-					AddDebugProjectionHairPass(GraphBuilder, LocalView, ShaderMap, ShaderPrintData, Viewport, ViewUniformBuffer, MeshLODIndex, RestRootResource, DeformedRootResource, Instance->HairGroupPublicData->VFInput.LocalToWorldTransform, DrawIndex++, bSim, EDebugProjectionHairType::HairTriangle, HairStrandsTriangleType::RestPose);
-					AddDebugProjectionHairPass(GraphBuilder, LocalView, ShaderMap, ShaderPrintData, Viewport, ViewUniformBuffer, MeshLODIndex, RestRootResource, DeformedRootResource, Instance->HairGroupPublicData->VFInput.LocalToWorldTransform, DrawIndex++, bSim, EDebugProjectionHairType::HairFrame,    HairStrandsTriangleType::RestPose);
+					AddDebugProjectionHairPass(GraphBuilder, LocalView, ShaderMap, ShaderPrintData, Viewport, ViewUniformBuffer, MeshLODIndex, RestRootResource, DeformedRootResource, Instance->HairGroupPublicData->VFInput.LocalToWorldTransform, DrawIndex++, bSim, EDebugProjectionHairType::HairRoots,    HairStrandsTriangleType::RestPose);
 					AddDebugProjectionHairPass(GraphBuilder, LocalView, ShaderMap, ShaderPrintData, Viewport, ViewUniformBuffer, MeshLODIndex, RestRootResource, DeformedRootResource, Instance->HairGroupPublicData->VFInput.LocalToWorldTransform, DrawIndex++, bSim, EDebugProjectionHairType::HairSamples,  HairStrandsTriangleType::RestPose);
-					AddDebugProjectionHairPass(GraphBuilder, LocalView, ShaderMap, ShaderPrintData, Viewport, ViewUniformBuffer, MeshLODIndex, RestRootResource, DeformedRootResource, Instance->HairGroupPublicData->VFInput.LocalToWorldTransform, DrawIndex++, bSim, EDebugProjectionHairType::HairTriangle, HairStrandsTriangleType::DeformedPose);
-					AddDebugProjectionHairPass(GraphBuilder, LocalView, ShaderMap, ShaderPrintData, Viewport, ViewUniformBuffer, MeshLODIndex, RestRootResource, DeformedRootResource, Instance->HairGroupPublicData->VFInput.LocalToWorldTransform, DrawIndex++, bSim, EDebugProjectionHairType::HairFrame,    HairStrandsTriangleType::DeformedPose);
+					AddDebugProjectionHairPass(GraphBuilder, LocalView, ShaderMap, ShaderPrintData, Viewport, ViewUniformBuffer, MeshLODIndex, RestRootResource, DeformedRootResource, Instance->HairGroupPublicData->VFInput.LocalToWorldTransform, DrawIndex++, bSim, EDebugProjectionHairType::HairRoots,    HairStrandsTriangleType::DeformedPose);
 					AddDebugProjectionHairPass(GraphBuilder, LocalView, ShaderMap, ShaderPrintData, Viewport, ViewUniformBuffer, MeshLODIndex, RestRootResource, DeformedRootResource, Instance->HairGroupPublicData->VFInput.LocalToWorldTransform, DrawIndex++, bSim, EDebugProjectionHairType::HairSamples,  HairStrandsTriangleType::DeformedPose);
 				}
 			};
diff --git a/Engine/Plugins/Runtime/HairStrands/Source/HairStrandsCore/Private/GroomGeometryCache.cpp b/Engine/Plugins/Runtime/HairStrands/Source/HairStrandsCore/Private/GroomGeometryCache.cpp
--- a/Engine/Plugins/Runtime/HairStrands/Source/HairStrandsCore/Private/GroomGeometryCache.cpp
+++ b/Engine/Plugins/Runtime/HairStrands/Source/HairStrandsCore/Private/GroomGeometryCache.cpp
@@ -53,6 +53,7 @@
 				FCachedGeometry::Section OutSection;
 				OutSection.PositionBuffer = TrackProxy->PositionBuffers[PositionIndex].GetBufferSRV();
 				OutSection.UVsBuffer = TrackProxy->TextureCoordinatesBuffer.GetBufferSRV();
+				OutSection.TangentBuffer = TrackProxy->TangentZBuffer.GetBufferSRV();
 				OutSection.TotalVertexCount = TrackProxy->MeshData->Positions.Num();
 				OutSection.IndexBuffer = TrackProxy->IndexBuffer.GetBufferSRV();
 				OutSection.TotalIndexCount = TrackProxy->IndexBuffer.NumValidIndices;
@@ -64,6 +65,7 @@
 				OutSection.SectionIndex = SectionIdx;
 				OutSection.LODIndex = 0;
 				OutSection.UVsChannelOffset = 0;
+				OutSection.TangentFormat = 1u;
 			
 				if (OutSection.PositionBuffer && OutSection.IndexBuffer)
 				{
diff --git a/Engine/Plugins/Runtime/HairStrands/Source/HairStrandsCore/Private/GroomManager.cpp b/Engine/Plugins/Runtime/HairStrands/Source/HairStrandsCore/Private/GroomManager.cpp
--- a/Engine/Plugins/Runtime/HairStrands/Source/HairStrandsCore/Private/GroomManager.cpp
+++ b/Engine/Plugins/Runtime/HairStrands/Source/HairStrandsCore/Private/GroomManager.cpp
@@ -529,6 +529,7 @@ static void GetOrAllocateCachedGeometry(
 	// Create deformed position buffer (output)
 	FRDGBufferSRVRef DeformedPositionSRV = nullptr;
 	FRDGBufferSRVRef DeformedPreviousPositionSRV = nullptr;
+	FRDGBufferSRVRef DeformedTangentSRV = nullptr; // TODO
 	OutHairGeometryCache.GetOrAdd(GraphBuilder, SkeletalMeshObject, &LODData, LODIndex, UniqueSections, DeformedPositionSRV, DeformedPreviousPositionSRV);
 
 	// Add reference to be sure the data are not streamed out while they are used
@@ -540,6 +541,7 @@ static void GetOrAllocateCachedGeometry(
 		FCachedGeometry::Section& OutSection= Out.Sections.AddDefaulted_GetRef();
 		OutSection.RDGPositionBuffer 		= DeformedPositionSRV;
 		OutSection.RDGPreviousPositionBuffer= DeformedPreviousPositionSRV;
+		OutSection.RDGTangentBuffer			= DeformedTangentSRV;
 		OutSection.PositionBuffer 			= nullptr; // Do not use the SRV slot, but instead use the RDG buffer created above (DeformedPositionSRV)
 		OutSection.PreviousPositionBuffer 	= nullptr; // Do not use the SRV slot, but instead use the RDG buffer created above (DeformedPositionSRV)
 		OutSection.UVsBuffer 				= LODData.StaticVertexBuffers.StaticMeshVertexBuffer.GetTexCoordsSRV();
diff --git a/Engine/Plugins/Runtime/HairStrands/Source/HairStrandsCore/Private/HairStrandsMeshProjection.cpp b/Engine/Plugins/Runtime/HairStrands/Source/HairStrandsCore/Private/HairStrandsMeshProjection.cpp
--- a/Engine/Plugins/Runtime/HairStrands/Source/HairStrandsCore/Private/HairStrandsMeshProjection.cpp
+++ b/Engine/Plugins/Runtime/HairStrands/Source/HairStrandsCore/Private/HairStrandsMeshProjection.cpp
@@ -153,15 +153,17 @@ private:
 	using FPermutationDomain = TShaderPermutationDomain<FPositionType, FPrevious>;
 
 	BEGIN_SHADER_PARAMETER_STRUCT(FParameters, )
+		SHADER_PARAMETER(uint32, TangentFormat)
 		SHADER_PARAMETER(uint32, MaxSectionCount)
 		SHADER_PARAMETER(uint32, MaxUniqueTriangleCount)
 		SHADER_PARAMETER_RDG_BUFFER_SRV(StructuredBuffer, MeshSectionBuffer)
 		SHADER_PARAMETER_RDG_BUFFER_SRV(Buffer, RDGMeshPositionBuffer)
 		SHADER_PARAMETER_RDG_BUFFER_SRV(Buffer, RDGMeshPreviousPositionBuffer)
+		SHADER_PARAMETER_RDG_BUFFER_SRV(Buffer, RDGTangentBuffer)
 		SHADER_PARAMETER_SRV(Buffer, MeshPositionBuffer)
 		SHADER_PARAMETER_SRV(Buffer, MeshPreviousPositionBuffer)
 		SHADER_PARAMETER_SRV(Buffer, MeshIndexBuffer)
-		SHADER_PARAMETER_SRV(Buffer, MeshUVsBuffer)
+		SHADER_PARAMETER_SRV(Buffer, MeshTangentBuffer)
 
 		SHADER_PARAMETER_RDG_BUFFER_SRV(Buffer, UniqueTriangleIndices)
 		SHADER_PARAMETER_RDG_BUFFER_UAV(RWBuffer, OutUniqueTrianglePrevPosition)
@@ -204,12 +206,14 @@ static bool AddHairStrandUpdateMeshTrianglesPass(
 	FHairUpdateMeshTriangleCS::FParameters CommonParameters;
 	CommonParameters.MaxUniqueTriangleCount 		= UniqueTriangleCount;
 	CommonParameters.MaxSectionCount 				= TotalSectionCount;
+	CommonParameters.TangentFormat					= MeshLODData.Sections[0].TangentFormat;
 	CommonParameters.RDGMeshPositionBuffer			= MeshLODData.Sections[0].RDGPositionBuffer;
 	CommonParameters.RDGMeshPreviousPositionBuffer 	= MeshLODData.Sections[0].RDGPreviousPositionBuffer;
+	CommonParameters.RDGTangentBuffer				= MeshLODData.Sections[0].RDGTangentBuffer;
 	CommonParameters.MeshPositionBuffer				= MeshLODData.Sections[0].PositionBuffer;
 	CommonParameters.MeshPreviousPositionBuffer		= MeshLODData.Sections[0].PreviousPositionBuffer;
 	CommonParameters.MeshIndexBuffer				= MeshLODData.Sections[0].IndexBuffer;
-	CommonParameters.MeshUVsBuffer					= MeshLODData.Sections[0].UVsBuffer;
+	CommonParameters.MeshTangentBuffer				= MeshLODData.Sections[0].TangentBuffer;
 	CommonParameters.UniqueTriangleIndices 			= UniqueTriangleIndexSRV;
 	CommonParameters.OutUniqueTriangleCurrPosition	= OutputCurrUAV;
 	CommonParameters.OutUniqueTrianglePrevPosition	= OutputPrevUAV;
@@ -223,8 +227,9 @@ static bool AddHairStrandUpdateMeshTrianglesPass(
 		uint32 IndexBaseIndex;
 		uint32 UVsChannelOffset : 8;
 		uint32 UVsChannelCount : 8;
-		uint32 bIsSwapped : 8;
-		uint32 Pad : 8;
+		uint32 bIsSwapped : 1;
+		uint32 bUseFaceNormal : 1;
+		uint32 Pad = 14;
 	};
 
 	FSectionData Default;
@@ -234,6 +239,7 @@ static bool AddHairStrandUpdateMeshTrianglesPass(
 	Default.UVsChannelOffset  = 0;
 	Default.UVsChannelCount  = 0;
 	Default.bIsSwapped = 0;
+	Default.bUseFaceNormal = 0;
 	Default.Pad = 0;
 
 	// Allocate data for *all* sections, but only fill in the used/valid sections
@@ -251,6 +257,7 @@ static bool AddHairStrandUpdateMeshTrianglesPass(
 		SectionDatas[SectionIndex].bIsSwapped		= bUseRDGPositionBuffer ? 
 			(MeshSectionData.RDGPositionBuffer != CommonParameters.RDGMeshPositionBuffer ? 1u : 0u) :
 			(MeshSectionData.PositionBuffer    != CommonParameters.MeshPositionBuffer    ? 1u : 0u);
+		SectionDatas[SectionIndex].bUseFaceNormal	= 0u;
 
 		// Sanity check
 		check(MeshSectionData.UVsChannelOffset < 255);
@@ -258,7 +265,7 @@ static bool AddHairStrandUpdateMeshTrianglesPass(
 		check(CommonParameters.RDGMeshPositionBuffer == MeshSectionData.RDGPositionBuffer || CommonParameters.RDGMeshPositionBuffer == MeshSectionData.RDGPreviousPositionBuffer);
 		check(CommonParameters.MeshPositionBuffer    == MeshSectionData.PositionBuffer    || CommonParameters.MeshPositionBuffer    == MeshSectionData.PreviousPositionBuffer);
 		check(CommonParameters.MeshIndexBuffer		 == MeshSectionData.IndexBuffer);
-		check(CommonParameters.MeshUVsBuffer		 == MeshSectionData.UVsBuffer);
+		check(CommonParameters.MeshTangentBuffer	 == MeshSectionData.TangentBuffer);
 	}
 
 	// If no previous position buffer available, reusing the current position buffers
@@ -276,11 +283,12 @@ static bool AddHairStrandUpdateMeshTrianglesPass(
 
 	const bool bComputePreviousDeformedPosition = OutputPrevUAV != nullptr;
 	{
-		// UV stream is only used for debugging purpose. On some platform the TextureCoordinateSRV can be null as it is not create by default (requires CPU access flags).
-		// In such a case we bind the index buffer as a dummy data
-		if (CommonParameters.MeshUVsBuffer == nullptr)
+		// On some platform the TangentSRV can be null as it is not create by default (requires CPU access flags).
+		// In such a case we bind the index buffer as a dummy data, and we use face normal.
+		if (CommonParameters.MeshTangentBuffer == nullptr)
 		{
-			CommonParameters.MeshUVsBuffer = CommonParameters.MeshIndexBuffer;
+			CommonParameters.MeshTangentBuffer = CommonParameters.MeshIndexBuffer;
+			Default.bUseFaceNormal = 1u;
 		}
 
 		FHairUpdateMeshTriangleCS::FParameters* PassParameters = GraphBuilder.AllocParameters<FHairUpdateMeshTriangleCS::FParameters>();
diff --git a/Engine/Shaders/Private/HairStrands/HairCardsDeformation.usf b/Engine/Shaders/Private/HairStrands/HairCardsDeformation.usf
--- a/Engine/Shaders/Private/HairStrands/HairCardsDeformation.usf
+++ b/Engine/Shaders/Private/HairStrands/HairCardsDeformation.usf
@@ -44,16 +44,15 @@ Buffer<uint>    	GuideRootToUniqueTriangleIndexBuffer;
 float3 ComputeDynamicGeometryOffset(
 	uint GuideIndex,
 	float GuideVertexWeight,
-	FHairMeshTriangle RestTriangle,
-	FHairMeshTriangle DeformedTriangle,
-	float3 RootBarycentric,
+	FHairMeshBasis RestBasis,
+	FHairMeshBasis DeformedBasis,
 	float3 GuideDeformedPositionOffset)
 {
 	const float3 RestGuidePoint = ReadHairControlPointPosition(GuideRestPositionBuffer, GuideIndex, GuideRestPositionOffset);
-	const float3 LocalRestGuidePoint = ToTriangle(RestGuidePoint, RootBarycentric, RestTriangle);
+	const float3 LocalRestGuidePoint = ToTriangle(RestGuidePoint, RestBasis);
 
 	const float3 DeformedGuidePoint = ReadHairControlPointPosition(GuideDeformedPositionBuffer, GuideIndex, GuideDeformedPositionOffset);
-	const float3 LocalDeformedGuidePoint = ToTriangle(DeformedGuidePoint, RootBarycentric, DeformedTriangle);
+	const float3 LocalDeformedGuidePoint = ToTriangle(DeformedGuidePoint, DeformedBasis);
 
 	return (LocalDeformedGuidePoint - LocalRestGuidePoint) * GuideVertexWeight;
 }
@@ -61,12 +60,12 @@ float3 ComputeDynamicGeometryOffset(
 void TransformTangent(
 	uint InVertexId, 
 	uint InOffset,
-	FHairMeshTriangle InRestTriangle,
-	FHairMeshTriangle InDeformedTriangle)
+	FHairMeshBasis InRestBasis,
+	FHairMeshBasis InDeformedBasis)
 {
 	float4 Rest		= CardsRestTangentBuffer[InVertexId * 2 + InOffset];
-	float3 Local	= VectorToTriangle(Rest.xyz, InRestTriangle);
-	float3 Deformed = VectorToWorld(Local, InDeformedTriangle);
+	float3 Local	= VectorToTriangle(Rest.xyz, InRestBasis);
+	float3 Deformed = VectorToWorld(Local, InDeformedBasis);
 	CardsDeformedTangentBuffer[InVertexId * 2 + InOffset] = float4(Deformed, Rest.w);
 }
 #endif
@@ -149,24 +148,24 @@ void MainCS(uint DispatchThreadId : SV_DispatchThreadID)
 
 	const float3 RootBarycentric = UnpackBarycentrics(GuideRootBarycentricBuffer[GuideCurveIndex]);
 
-	const FHairMeshTriangle RestTriangle = GetTriangleTransformation(TriangleIndex, TriangleRestPositionBuffer);
-	const FHairMeshTriangle DeformedTriangle = GetTriangleTransformation(TriangleIndex, TriangleDeformedPositionBuffer);
+	const FHairMeshBasis RestBasis = GetHairMeshBasis(TriangleIndex, TriangleRestPositionBuffer, RootBarycentric);
+	const FHairMeshBasis DeformedBasis = GetHairMeshBasis(TriangleIndex, TriangleDeformedPositionBuffer, RootBarycentric);
 
-	const float3 Offset0 = ComputeDynamicGeometryOffset(GuidePointIndex0, 1, RestTriangle, DeformedTriangle, RootBarycentric, GuideDeformedPositionOffset);
-	const float3 Offset1 = ComputeDynamicGeometryOffset(GuidePointIndex1, 1, RestTriangle, DeformedTriangle, RootBarycentric, GuideDeformedPositionOffset);
-	const float3 CurrOffset = VectorToWorld(lerp(Offset0, Offset1, GuideData.VertexLerp), DeformedTriangle);
+	const float3 Offset0 = ComputeDynamicGeometryOffset(GuidePointIndex0, 1, RestBasis, DeformedBasis, GuideDeformedPositionOffset);
+	const float3 Offset1 = ComputeDynamicGeometryOffset(GuidePointIndex1, 1, RestBasis, DeformedBasis, GuideDeformedPositionOffset);
+	const float3 CurrOffset = VectorToWorld(lerp(Offset0, Offset1, GuideData.VertexLerp), DeformedBasis);
 
 	// Transform hair from rest post to deformed pose, based triangle deformation + simulation offset (optional depending of the simulation is actual running or not
 	const float4 RestControlPoint = CardsRestPositionBuffer[VertexId];
 	float3 ControlPoint = RestControlPoint.xyz;
-	ControlPoint = ToTriangle(ControlPoint, RootBarycentric, RestTriangle);
-	ControlPoint = ToWorld(ControlPoint, RootBarycentric, DeformedTriangle) + CurrOffset;
+	ControlPoint = ToTriangle(ControlPoint, RestBasis);
+	ControlPoint = ToWorld(ControlPoint, DeformedBasis) + CurrOffset;
 
 	#if DEBUG_ENABLE
 	{
-		float3 WP0 = mul(float4(DeformedTriangle.P0, 1), LocalToWorld).xyz;
-		float3 WP1 = mul(float4(DeformedTriangle.P1, 1), LocalToWorld).xyz;
-		float3 WP2 = mul(float4(DeformedTriangle.P2, 1), LocalToWorld).xyz;
+		float3 WP0 = mul(float4(DeformedBasis.P0, 1), LocalToWorld).xyz;
+		float3 WP1 = mul(float4(DeformedBasis.P1, 1), LocalToWorld).xyz;
+		float3 WP2 = mul(float4(DeformedBasis.P2, 1), LocalToWorld).xyz;
 		AddLineTriangle(WP0, WP1, WP2, ColorOrange);
 	}
 	#endif
@@ -174,8 +173,8 @@ void MainCS(uint DispatchThreadId : SV_DispatchThreadID)
 	CardsDeformedPositionBuffer[VertexId] = float4(ControlPoint, RestControlPoint.w);
 
 	// Transform tangent
-	TransformTangent(VertexId, 0, RestTriangle, DeformedTriangle);
-	TransformTangent(VertexId, 1, RestTriangle, DeformedTriangle);
+	TransformTangent(VertexId, 0, RestBasis, DeformedBasis);
+	TransformTangent(VertexId, 1, RestBasis, DeformedBasis);
 #endif // PERMUTATION_DYNAMIC_GEOMETRY
 }
 #endif // SHADER_DEFORM_CARDS
\ No newline at end of file
diff --git a/Engine/Shaders/Private/HairStrands/HairStrandsBindingCommon.ush b/Engine/Shaders/Private/HairStrands/HairStrandsBindingCommon.ush
--- a/Engine/Shaders/Private/HairStrands/HairStrandsBindingCommon.ush
+++ b/Engine/Shaders/Private/HairStrands/HairStrandsBindingCommon.ush
@@ -2,35 +2,25 @@
 
 #include "HairStrandsPack.ush"
 
+struct FHairMeshBasis
+{
+	float3 BaseO;
+	float3 BaseT;
+	float3 BaseB;
+	float3 BaseN;
+};
+
 struct FHairMeshTriangle
 {
 	float3 P0;
 	float3 P1;
 	float3 P2;
 
-	float3 E0;
-	float3 E1;
-
-	float3 N;
-};
-
-struct FHairMeshTriangleWithUVs : FHairMeshTriangle
-{
-	float2 UV0;
-	float2 UV1;
-	float2 UV2;
+	float3 N0;
+	float3 N1;
+	float3 N2;
 };
 
-float2 GetVertexUV(uint VId, Buffer<float2> InUVBuffer, Buffer<uint> InIndexBuffer, uint InIndexOffset, uint InMaxIndexCount, uint InMaxVertexCount, uint InMeshUVsChannelOffset, uint InMeshUVsChannelCount)
-{
-	const uint Index = min(InIndexOffset + VId, InMaxIndexCount - 1);
-	uint VertexIndex = InIndexBuffer.Load(Index);
-	VertexIndex = min(VertexIndex, InMaxVertexCount - 1);
-
-	const uint UVIndex = InMeshUVsChannelOffset + InMeshUVsChannelCount * VertexIndex;
-	return InUVBuffer.Load(UVIndex);
-}
-
 float3 GetVertexPosition(uint VId, Buffer<float> InPositionBuffer, Buffer<uint> InIndexBuffer, uint InIndexOffset, uint InMaxIndexCount, uint InMaxVertexCount)
 {
 	const uint Index = min(InIndexOffset + VId, InMaxIndexCount-1);
@@ -44,132 +34,133 @@ float3 GetVertexPosition(uint VId, Buffer<float> InPositionBuffer, Buffer<uint>
 	return P;
 }
 
-FHairMeshTriangle GetTriangleData(uint TriangleIndex, Buffer<float> InPositionBuffer, Buffer<uint> InIndexBuffer, uint InIndexOffset, uint InMaxIndexCount, uint InMaxVertexCount)
+float3 GetVertexNormal(uint VId, Buffer<float4> InTangentBuffer, Buffer<uint> InIndexBuffer, uint InIndexOffset, uint InMaxIndexCount, uint InMaxVertexCount, uint InTangentFormat)
 {
-	const uint BaseVertexId = TriangleIndex*3;
-	const uint VId0 = BaseVertexId;
-	const uint VId1 = BaseVertexId+1;
-	const uint VId2 = BaseVertexId+2;
-
-	FHairMeshTriangle Out;
-	Out.P0 = GetVertexPosition(VId0, InPositionBuffer, InIndexBuffer, InIndexOffset, InMaxIndexCount, InMaxVertexCount);
-	Out.P1 = GetVertexPosition(VId1, InPositionBuffer, InIndexBuffer, InIndexOffset, InMaxIndexCount, InMaxVertexCount);
-	Out.P2 = GetVertexPosition(VId2, InPositionBuffer, InIndexBuffer, InIndexOffset, InMaxIndexCount, InMaxVertexCount);
-	Out.E0 = Out.P1 - Out.P0;
-	Out.E1 = Out.P2 - Out.P0;
-	Out.N = normalize(cross(Out.E1, Out.E0));
+	const uint Index = min(InIndexOffset + VId, InMaxIndexCount-1);
+	uint VertexIndex = InIndexBuffer.Load(Index);
+	VertexIndex = min(VertexIndex, InMaxVertexCount-1);
 
-	return Out;
+	float4 TangentZ = 0;
+	if (InTangentFormat == 1)
+	{
+		TangentZ = float4(InTangentBuffer[VertexIndex].xyz,0);
+	}
+	else
+	{
+		TangentZ = TangentBias(InTangentBuffer[VertexIndex * 2 + 1].xyzw);
+	}
+	return TangentZ.xyz;
 }
 
-FHairMeshTriangleWithUVs GetTriangleData(uint TriangleIndex, Buffer<float> InPositionBuffer, Buffer<float2> InUVBuffer, Buffer<uint> InIndexBuffer, uint InIndexOffset, uint InMaxIndexCount, uint InMaxVertexCount, uint InMeshUVsChannelOffset, uint InMeshUVsChannelCount)
+// Return a mesh triangle (position + normal)
+FHairMeshTriangle GetHairMeshTriangle(uint TriangleIndex, Buffer<float> InPositionBuffer, Buffer<float4> InTangentBuffer, Buffer<uint> InIndexBuffer, uint InIndexOffset, uint InMaxIndexCount, uint InMaxVertexCount, uint InTangentFormat)
 {
 	const uint BaseVertexId = TriangleIndex * 3;
 	const uint VId0 = BaseVertexId;
 	const uint VId1 = BaseVertexId + 1;
 	const uint VId2 = BaseVertexId + 2;
 
-	FHairMeshTriangleWithUVs Out;
+	FHairMeshTriangle Out;
 	Out.P0 = GetVertexPosition(VId0, InPositionBuffer, InIndexBuffer, InIndexOffset, InMaxIndexCount, InMaxVertexCount);
 	Out.P1 = GetVertexPosition(VId1, InPositionBuffer, InIndexBuffer, InIndexOffset, InMaxIndexCount, InMaxVertexCount);
 	Out.P2 = GetVertexPosition(VId2, InPositionBuffer, InIndexBuffer, InIndexOffset, InMaxIndexCount, InMaxVertexCount);
-	Out.E0 = Out.P1 - Out.P0;
-	Out.E1 = Out.P2 - Out.P0;
-	Out.N = normalize(cross(Out.E1, Out.E0));
 
-	Out.UV0 = GetVertexUV(VId0, InUVBuffer, InIndexBuffer, InIndexOffset, InMaxIndexCount, InMaxVertexCount, InMeshUVsChannelOffset, InMeshUVsChannelCount);
-	Out.UV1 = GetVertexUV(VId1, InUVBuffer, InIndexBuffer, InIndexOffset, InMaxIndexCount, InMaxVertexCount, InMeshUVsChannelOffset, InMeshUVsChannelCount);
-	Out.UV2 = GetVertexUV(VId2, InUVBuffer, InIndexBuffer, InIndexOffset, InMaxIndexCount, InMaxVertexCount, InMeshUVsChannelOffset, InMeshUVsChannelCount);
+	Out.N0 = GetVertexNormal(VId0, InTangentBuffer, InIndexBuffer, InIndexOffset, InMaxIndexCount, InMaxVertexCount, InTangentFormat);
+	Out.N1 = GetVertexNormal(VId1, InTangentBuffer, InIndexBuffer, InIndexOffset, InMaxIndexCount, InMaxVertexCount, InTangentFormat);
+	Out.N2 = GetVertexNormal(VId2, InTangentBuffer, InIndexBuffer, InIndexOffset, InMaxIndexCount, InMaxVertexCount, InTangentFormat);
 	
 	return Out;
 }
 
-FHairMeshTriangle GetTriangleTransformation(
-	uint TriangleIndex,
-	Buffer<float4> InPositionBuffer)
+// Return a mesh triangle (position + normal)
+FHairMeshTriangle GetHairMeshTriangle(uint TriangleIndex, Buffer<float4> InPositionBuffer)
 {
 	FHairMeshTriangle Out;
-
-	// Build orthonormal frame
-	Out.P0 = InPositionBuffer[TriangleIndex * 3 + 0].xyz;
-	Out.P1 = InPositionBuffer[TriangleIndex * 3 + 1].xyz;
-	Out.P2 = InPositionBuffer[TriangleIndex * 3 + 2].xyz;
-	Out.E0 = Out.P1 - Out.P0;
-	Out.E1 = Out.P2 - Out.P0;
-
-	Out.N  = normalize(cross(Out.E1, Out.E0));
-	Out.E0 = normalize(Out.E0); 
-	Out.E1 = normalize(cross(Out.E0, Out.N));
-
+	const float4 V0 = InPositionBuffer[TriangleIndex * 3 + 0];
+	const float4 V1 = InPositionBuffer[TriangleIndex * 3 + 1];
+	const float4 V2 = InPositionBuffer[TriangleIndex * 3 + 2];
+	Out.P0 = V0.xyz;
+	Out.P1 = V1.xyz;
+	Out.P2 = V2.xyz;
+	Out.N0 = UnpackVertexNormal(V0.w);
+	Out.N1 = UnpackVertexNormal(V1.w);
+	Out.N2 = UnpackVertexNormal(V2.w);
 	return Out;
 }
 
-FHairMeshTriangleWithUVs GetTriangleWithUVsTransformation(
+// Return a hair basis (origin + tangent frame)
+FHairMeshBasis GetHairMeshBasis(
 	uint TriangleIndex,
-	Buffer<float4> InPositionBuffer)
+	Buffer<float4> InPositionBuffer,
+	float3 B/*Barycentrics*/)
 {
-	FHairMeshTriangleWithUVs Out;
+	FHairMeshBasis Out;
 
 	const float4 V0 = InPositionBuffer[TriangleIndex * 3 + 0];
 	const float4 V1 = InPositionBuffer[TriangleIndex * 3 + 1];
 	const float4 V2 = InPositionBuffer[TriangleIndex * 3 + 2];
 
-	// Build orthonormal frame
-	Out.P0 = V0.xyz;
-	Out.P1 = V1.xyz;
-	Out.P2 = V2.xyz;
-	Out.E0 = Out.P1 - Out.P0;
-	Out.E1 = Out.P2 - Out.P0;
+	const float3 N0 = UnpackVertexNormal(V0.w);
+	const float3 N1 = UnpackVertexNormal(V1.w);
+	const float3 N2 = UnpackVertexNormal(V2.w);
 
-	Out.N = normalize(cross(Out.E1, Out.E0));
-	Out.E0 = normalize(Out.E0);
-	Out.E1 = normalize(cross(Out.E0, Out.N));
+	const float3 E0 = (V1 - V0).xyz;
+	const float3 E1 = (V2 - V0).xyz;
 
-	Out.UV0 = UnPackUVs(V0.w);
-	Out.UV1 = UnPackUVs(V1.w);
-	Out.UV2 = UnPackUVs(V2.w);
+	// Build orthonormal frame
+	Out.BaseO = (V0 * B.x + V1 * B.y + V2 * B.z).xyz;
+	Out.BaseN = normalize(N0 * B.x + N1 * B.y + N2 * B.z);
+	Out.BaseT = normalize(cross(Out.BaseN, E1));
+	Out.BaseB = normalize(cross(Out.BaseT, Out.BaseN));
 
+	const bool bUseFaceNormal = false;
+	if (bUseFaceNormal)
+	{
+		Out.BaseN = normalize(cross(E1, E0));
+		Out.BaseT = normalize(E0);
+		Out.BaseB = normalize(cross(Out.BaseT, Out.BaseN));
+	}
 	return Out;
 }
 
-float3 ToTriangle(const float3 P, float3 B, const FHairMeshTriangle T)
+float3 ToTriangle(const float3 P, const FHairMeshBasis T)
 {
-	const float3 LocalP = P - (T.P0 * B.x + T.P1 * B.y + T.P2 * B.z);
+	const float3 LocalP = P - T.BaseO;
 
 	float3 Out;
-	Out.x = dot(LocalP, T.E0);
-	Out.y = dot(LocalP, T.E1);
-	Out.z = dot(LocalP, T.N);
+	Out.x = dot(LocalP, T.BaseT);
+	Out.y = dot(LocalP, T.BaseB);
+	Out.z = dot(LocalP, T.BaseN);
 	return Out;
 }
 
-float3 ToWorld(const float3 P, float3 B, const FHairMeshTriangle T)
+float3 ToWorld(const float3 P, const FHairMeshBasis T)
 {
-	const float3 Rotation = P.x * T.E0 + P.y * T.E1 + P.z * T.N;
-	const float3 Translation = T.P0 * B.x + T.P1 * B.y + T.P2 * B.z;
+	const float3 Rotation    = P.x * T.BaseT + P.y * T.BaseB + P.z * T.BaseN;
+	const float3 Translation = T.BaseO;
 	return Rotation + Translation;
 
 }
 
-float3 VectorToWorld(const float3 P, const FHairMeshTriangle T)
+float3 VectorToWorld(const float3 P, const FHairMeshBasis T)
 {
-	const float3 Rotation = P.x * T.E0 + P.y * T.E1 + P.z * T.N;
+	const float3 Rotation = P.x * T.BaseT + P.y * T.BaseB + P.z * T.BaseN;
 	return Rotation;
 }
 
-float3 VectorToTriangle(const float3 P, const FHairMeshTriangle T)
+float3 VectorToTriangle(const float3 P, const FHairMeshBasis T)
 {
 	float3 Out;
-	Out.x = dot(P, T.E0);
-	Out.y = dot(P, T.E1);
-	Out.z = dot(P, T.N);
+	Out.x = dot(P, T.BaseT);
+	Out.y = dot(P, T.BaseB);
+	Out.z = dot(P, T.BaseN);
 	return Out;
 }
 
-float3 TransformPoint(float3 P, float3 B, const FHairMeshTriangle RestTriangle, const FHairMeshTriangle DeformedTriangle)
+float3 TransformPoint(float3 P, const FHairMeshBasis RestBasis, const FHairMeshBasis DeformedBasis)
 {
-	float3 LocalP = ToTriangle(P, B, RestTriangle);
-	return ToWorld(LocalP, B, DeformedTriangle);
+	float3 LocalP = ToTriangle(P, RestBasis);
+	return ToWorld(LocalP, DeformedBasis);
 }
 
 // Apply RBF weight onto a control point
diff --git a/Engine/Shaders/Private/HairStrands/HairStrandsDebug.usf b/Engine/Shaders/Private/HairStrands/HairStrandsDebug.usf
--- a/Engine/Shaders/Private/HairStrands/HairStrandsDebug.usf
+++ b/Engine/Shaders/Private/HairStrands/HairStrandsDebug.usf
@@ -312,9 +312,8 @@ void MainCS(uint3 DispatchThreadId : SV_DispatchThreadID)
 #include "HairStrandsBindingCommon.ush" 
 #include "../ShaderPrint.ush"
 
-#define INPUT_ROOT_FRAME 0
-#define INPUT_ROOT_TRIANGLE 1
-#define INPUT_ROOT_SAMPLES 2
+#define INPUT_ROOTS 0
+#define INPUT_SAMPLES 1
 
 uint bSimPass;
 uint bRestPass;
@@ -337,18 +336,17 @@ uint Sim_HairDeformedTriangles;
 uint Sim_HairDeformedFrames;
 uint Sim_HairDeformedSamples;
 
-#if PERMUTATION_INPUT_TYPE == INPUT_ROOT_FRAME
+#if PERMUTATION_INPUT_TYPE == INPUT_ROOTS
 Buffer<uint>   RootBarycentricBuffer;
+Buffer<float4> RestPositionBuffer;
+Buffer<float4> DeformedPositionBuffer;
+Buffer<uint>   RootToUniqueTriangleIndexBuffer;
 #endif
 
-Buffer<uint>  RootToUniqueTriangleIndexBuffer;
-
-Buffer<float4>  RestPositionBuffer;
-Buffer<float4>  DeformedPositionBuffer;
-
+#if PERMUTATION_INPUT_TYPE == INPUT_SAMPLES
 StructuredBuffer<float4> RestSamplePositionsBuffer;
 StructuredBuffer<float4> DeformedSamplePositionsBuffer;
-
+#endif
 
 void Move(inout FShaderPrintContext Ctx, uint2 InPos)
 {
@@ -360,10 +358,51 @@ void Move(inout FShaderPrintContext Ctx, uint2 InPos)
 	}
 }
 
+float2 ToScreenUV(float3 TWSPos)
+{
+	const float4 NDC = mul(float4(TWSPos.xyz, 1), View.TranslatedWorldToClip);
+	const float2 ScreenUv = (NDC.xy / NDC.ww) * float2(0.5f, -0.5f) + 0.5f;
+	return ScreenUv;
+}
+
+float4 GetScreenMinMax(float3 WP0, float3 WP1, float3 WP2)
+{
+	const float3 PreViewTranslation = DFDemote(MakeDFVector3(View.PreViewTranslationHigh, View.PreViewTranslationLow));
+	const float3 TWSP0 = WP0 + PreViewTranslation;
+	const float3 TWSP1 = WP1 + PreViewTranslation;
+	const float3 TWSP2 = WP2 + PreViewTranslation;
+
+	const float2 UV0 = ToScreenUV(TWSP0);
+	const float2 UV1 = ToScreenUV(TWSP1);
+	const float2 UV2 = ToScreenUV(TWSP2);
+
+	float2 MinUV;
+	float2 MaxUV;
+	MinUV.x = min3(UV0.x, UV1.x, UV2.x);
+	MinUV.y = min3(UV0.y, UV1.y, UV2.y);
+	MaxUV.x = max3(UV0.x, UV1.x, UV2.x);
+	MaxUV.y = max3(UV0.y, UV1.y, UV2.y);
+	return float4(MinUV, MaxUV);
+}
+
+bool IsInRect(float4 MinMaxUV, float2 UV)
+{
+	return all(UV >= MinMaxUV.xy) && all(UV <= MinMaxUV.zw);
+}
+
 [numthreads(256, 1, 1)]
 void MainCS(uint3 DispatchThreadId : SV_DispatchThreadID)
 {
 	ResolvedView = ResolveView();
+	float NormalScale = 0.2f;
+	
+	const float2 CursorPixelCoord = GetCursorPos();
+	const float2 CursorUV = CursorPixelCoord * View.ViewSizeAndInvSize.zw;
+
+
+	bool bDrawTriangle= false;
+	bool bDrawNormal  = false;
+	bool bDrawRoot    = false; 
 
 	// Widgets for draw configuration
 	{	
@@ -376,7 +415,9 @@ void MainCS(uint3 DispatchThreadId : SV_DispatchThreadID)
 		{
 			bool bRest    = false;
 			bool bDeform  = false;
+
 			bool bTriangle= false;
+			bool bNormal  = false;
 			bool bRoot    = false; 
 
 			Print(Ctx, TEXT("Guides "), FontYellow); Newline(Ctx);
@@ -385,6 +426,7 @@ void MainCS(uint3 DispatchThreadId : SV_DispatchThreadID)
 				bRest    = Sim_HairRestTriangles     || Sim_HairRestFrames;
 				bDeform  = Sim_HairDeformedTriangles || Sim_HairDeformedFrames;
 				bTriangle= Sim_HairRestTriangles     || Sim_HairDeformedTriangles;
+				bNormal  = false;
 				bRoot    = Sim_HairRestFrames        || Sim_HairDeformedFrames;
 
 				Print(Ctx, TEXT("Sim Rest      "), FontWhite);	PrintBool(Ctx, bRest); 		Newline(Ctx);
@@ -392,6 +434,7 @@ void MainCS(uint3 DispatchThreadId : SV_DispatchThreadID)
 				Newline(Ctx);
 	
 				Print(Ctx, TEXT("Sim Triangles "), FontWhite);	PrintBool(Ctx, bTriangle); 	Newline(Ctx);
+				Print(Ctx, TEXT("Sim Normals   "), FontWhite);	PrintBool(Ctx, bNormal); 	Newline(Ctx);
 				Print(Ctx, TEXT("Sim Roots     "), FontWhite); 	PrintBool(Ctx, bRoot);		Newline(Ctx);
 				Newline(Ctx);
 	
@@ -403,29 +446,32 @@ void MainCS(uint3 DispatchThreadId : SV_DispatchThreadID)
 				Newline(Ctx);
 		
 				bTriangle = AddCheckbox(Ctx, TEXT("Sim Triangles"), false, FontWhite); Newline(Ctx);
+				bNormal   = AddCheckbox(Ctx, TEXT("Sim Normals"), false, FontWhite); Newline(Ctx);
 				bRoot     = AddCheckbox(Ctx, TEXT("Sim Roots"), false, FontWhite); Newline(Ctx);
 				Newline(Ctx);
 			}
 
 			if (bSimPass == 1)
 			{
-				if (PERMUTATION_INPUT_TYPE == INPUT_ROOT_TRIANGLE)
-				{
-					bValid = bTriangle  && ((bRestPass==1 && bRest) || (bRestPass==0 && bDeform));
-				}
-				if (PERMUTATION_INPUT_TYPE == INPUT_ROOT_FRAME)
+				if (PERMUTATION_INPUT_TYPE == INPUT_ROOTS)
 				{
-					bValid = bRoot && ((bRestPass==1 && bRest) || (bRestPass==0 && bDeform));
+					bValid = (bRoot || bNormal || bTriangle)  && ((bRestPass==1 && bRest) || (bRestPass==0 && bDeform));
+
+					bDrawTriangle = bDrawTriangle || bTriangle;
+					bDrawNormal   = bDrawNormal   || bNormal;
+					bDrawRoot     = bDrawRoot     || bRoot; 
+
 				}
 			}
 		}
-
 		// Strands/Cards
 		{
 			bool bRest    = false;
 			bool bDeform  = false;
+
 			bool bTriangle= false;
-			bool bRoot    = false;
+			bool bNormal  = false;
+			bool bRoot    = false; 
 
 			Move(Ctx, StartPos + OffsetPos);
 			Print(Ctx, TEXT("Strands / Cards"), FontYellow); Newline(Ctx);
@@ -434,6 +480,7 @@ void MainCS(uint3 DispatchThreadId : SV_DispatchThreadID)
 				bRest    = Render_HairRestTriangles     || Render_HairRestFrames;
 				bDeform  = Render_HairDeformedTriangles || Render_HairDeformedFrames;
 				bTriangle= Render_HairRestTriangles     || Render_HairDeformedTriangles;
+				bNormal  = false;
 				bRoot    = Render_HairRestFrames        || Render_HairDeformedFrames;
 	
 				Print(Ctx, TEXT("Ren Rest      "), FontWhite);	PrintBool(Ctx, bRest); 		Newline(Ctx);
@@ -441,6 +488,7 @@ void MainCS(uint3 DispatchThreadId : SV_DispatchThreadID)
 				Newline(Ctx);
 			
 				Print(Ctx, TEXT("Ren Triangles "), FontWhite);	PrintBool(Ctx, bTriangle); 	Newline(Ctx);
+				Print(Ctx, TEXT("Sim Normals   "), FontWhite);	PrintBool(Ctx, bNormal); 	Newline(Ctx);
 				Print(Ctx, TEXT("Ren Roots     "), FontWhite); 	PrintBool(Ctx, bRoot);		Newline(Ctx);
 				Newline(Ctx);
 			}
@@ -451,19 +499,20 @@ void MainCS(uint3 DispatchThreadId : SV_DispatchThreadID)
 				Newline(Ctx);
 				
 				bTriangle = AddCheckbox(Ctx, TEXT("Ren Triangles"), false, FontWhite); Newline(Ctx);
+				bNormal   = AddCheckbox(Ctx, TEXT("Ren Normals"), false, FontWhite); Newline(Ctx);
 				bRoot     = AddCheckbox(Ctx, TEXT("Ren Roots"), true, FontWhite); Newline(Ctx);
 				Newline(Ctx);
 			}
 
 			if (bSimPass == 0)
 			{
-				if (PERMUTATION_INPUT_TYPE == INPUT_ROOT_TRIANGLE)
-				{
-					bValid = bTriangle && ((bRestPass==1 && bRest) || (bRestPass==0 && bDeform));
-				}
-				if (PERMUTATION_INPUT_TYPE == INPUT_ROOT_FRAME)
+				if (PERMUTATION_INPUT_TYPE == INPUT_ROOTS)
 				{
-					bValid = bRoot && ((bRestPass==1 && bRest) || (bRestPass==0 && bDeform));
+					bValid = (bRoot || bNormal || bTriangle) && ((bRestPass==1 && bRest) || (bRestPass==0 && bDeform));
+
+					bDrawTriangle = bDrawTriangle || bTriangle;
+					bDrawNormal   = bDrawNormal   || bNormal;
+					bDrawRoot     = bDrawRoot     || bRoot; 
 				}
 			}
 		}
@@ -491,113 +540,137 @@ void MainCS(uint3 DispatchThreadId : SV_DispatchThreadID)
 				Newline(Ctx);	
 			}
 
-			if (bSimPass == 1 && (PERMUTATION_INPUT_TYPE == INPUT_ROOT_SAMPLES))
+
+			if (bSimPass == 1 && (PERMUTATION_INPUT_TYPE == INPUT_SAMPLES))
 			{
 				bValid = (bRestPass==1 && bRest) || (bRestPass==0 && bDeform);
 			}
 		}
 
+		// Settings
+		NormalScale = AddSlider(Ctx, TEXT("Size"), 0.02f, FontWhite) * 10.f; 
+
 		if (!bValid)
 		{
 			return;
 		}
 	}
 
-#if PERMUTATION_INPUT_TYPE == INPUT_ROOT_FRAME  
+#if PERMUTATION_INPUT_TYPE == INPUT_ROOTS  
 	const uint RootIndex = DispatchThreadId.x;
 	if (RootIndex < MaxRootCount)
 	{
+		const float3 B  = UnpackBarycentrics(RootBarycentricBuffer[RootIndex]);
 		const uint TriangleIndex = RootToUniqueTriangleIndexBuffer[RootIndex];
-		const FHairMeshTriangleWithUVs RestTriangle		= GetTriangleWithUVsTransformation(TriangleIndex, RestPositionBuffer);
-		const FHairMeshTriangleWithUVs DeformedTriangle = GetTriangleWithUVsTransformation(TriangleIndex, DeformedPositionBuffer);
 
-		FHairMeshTriangleWithUVs Tri;
+		FHairMeshBasis Basis;
+		FHairMeshTriangle Tri;
 		if (DeformedFrameEnable > 0)
-			Tri = DeformedTriangle;
+		{		
+			Tri   = GetHairMeshTriangle(TriangleIndex, DeformedPositionBuffer);
+			Basis = GetHairMeshBasis(TriangleIndex, DeformedPositionBuffer, B);
+		}
 		else
-			Tri = RestTriangle;
+		{
+			Tri   = GetHairMeshTriangle(TriangleIndex, RestPositionBuffer);
+			Basis = GetHairMeshBasis(TriangleIndex, RestPositionBuffer, B);
+		}
 
-		const float3 B  = UnpackBarycentrics(RootBarycentricBuffer[RootIndex]);
-		float3 P0 = B.x * Tri.P0 + B.y * Tri.P1 + B.z * Tri.P2;
-		float3 P1 = P0 + Tri.N;
+		// Triangle positions
+		const float3 WP0 = mul(float4(Tri.P0, 1), RootLocalToWorld).xyz;
+		const float3 WP1 = mul(float4(Tri.P1, 1), RootLocalToWorld).xyz;
+		const float3 WP2 = mul(float4(Tri.P2, 1), RootLocalToWorld).xyz;
 
-		const float2 UVs = 
-			B.x * Tri.UV0 +
-			B.y * Tri.UV1 +
-			B.z * Tri.UV2;
+		// Triangle normals
+		const float3 WN0 = mul(float4(Tri.N0, 0), RootLocalToWorld).xyz;
+		const float3 WN1 = mul(float4(Tri.N1, 0), RootLocalToWorld).xyz;
+		const float3 WN2 = mul(float4(Tri.N2, 0), RootLocalToWorld).xyz;
 
-		P0 = mul(float4(P0, 1), RootLocalToWorld).xyz;
-		P1 = mul(float4(P1, 1), RootLocalToWorld).xyz;
+		// Basis
+		const float3 WO = mul(float4(Basis.BaseO, 1), RootLocalToWorld).xyz;
+		const float3 WT = mul(float4(Basis.BaseT, 0), RootLocalToWorld).xyz;
+		const float3 WB = mul(float4(Basis.BaseB, 0), RootLocalToWorld).xyz;
+		const float3 WN = mul(float4(Basis.BaseN, 0), RootLocalToWorld).xyz;
 
-		const float3 N = normalize(P0 - P1);
-		const float4 Color = float4(UVs, 0, 1);
-		//Color = float4(N* 0.5f + 0.5f, 1);
-		AddLineWS(P0, P1, Color);
+		// Face normal
+		float3 WFaceN;
+		{
+			const float3 E0 = (Tri.P1 - Tri.P0).xyz;
+			const float3 E1 = (Tri.P2 - Tri.P0).xyz;
+			const float3 FaceN = normalize(cross(E1, E0));
+			WFaceN = mul(float4(FaceN, 0), RootLocalToWorld).xyz;
+		}
 
-		// Add triangle
-		const float4 TriColor = ColorWhite;
-		const float3 WP0 = mul(float4(Tri.P0, 1), RootLocalToWorld).xyz;
-		const float3 WP1 = mul(float4(Tri.P1, 1), RootLocalToWorld).xyz;
-		const float3 WP2 = mul(float4(Tri.P2, 1), RootLocalToWorld).xyz;
-		AddLineWS(WP0, WP1, TriColor);
-		AddLineWS(WP1, WP2, TriColor);
-		AddLineWS(WP2, WP0, TriColor);
-	}
-#endif // INPUT_ROOT_FRAME
+		// Selection
+		const float4 MinMaxUV = GetScreenMinMax(WP0, WP1, WP2);
+		const bool bIsInRect = IsInRect(MinMaxUV, CursorUV);
+		if (bIsInRect)
+		{
+			// Selected triangles
+			{
+				AddLineWS(WP0, WP1, ColorYellow);
+				AddLineWS(WP1, WP2, ColorYellow);
+				AddLineWS(WP2, WP0, ColorYellow);
+			}
 
-#if PERMUTATION_INPUT_TYPE == INPUT_ROOT_TRIANGLE
-	const uint TriangleIndex = DispatchThreadId.x;
-	if (TriangleIndex < MaxRootCount)
-	{
-		const FHairMeshTriangleWithUVs RestTriangle		= GetTriangleWithUVsTransformation(TriangleIndex, RestPositionBuffer);
-		const FHairMeshTriangleWithUVs DeformedTriangle = GetTriangleWithUVsTransformation(TriangleIndex, DeformedPositionBuffer);
-	
-		FHairMeshTriangleWithUVs Tri;
-		if (DeformedFrameEnable > 0)
-			Tri = DeformedTriangle;
-		else
-			Tri = RestTriangle;
+			// Vertex normal
+			{
+				AddLineWS(WP1, WP1 + WN1 * NormalScale, ColorBlue);
+				AddLineWS(WP2, WP2 + WN2 * NormalScale, ColorBlue);
+				AddLineWS(WP0, WP0 + WN0 * NormalScale, ColorBlue);
+			}
+
+			// Triangle normal
+			if (0)
+			{
+				AddLineWS(WP0, WP0 + WFaceN * NormalScale, ColorGreen);
+				AddLineWS(WP1, WP1 + WFaceN * NormalScale, ColorGreen);
+				AddLineWS(WP2, WP2 + WFaceN * NormalScale, ColorGreen);
+			}
 		
-	
-		// Debug: To test if the rest-triangle-to-deformed-triangle transformation works correctly
-		#if 0
-		if (DeformedFrameEnable > 0)
+			// Root basis
+			{
+				AddLineWS(WO, WO + WT * NormalScale * 0.3f, ColorRed);
+				AddLineWS(WO, WO + WB * NormalScale * 0.3f, ColorGreen);
+				AddLineWS(WO, WO + WN * NormalScale * 0.3f, ColorBlue);
+			}
+		}
+		else
 		{
-			float3 Barycentric = 0;
-			if (TriangleVertex == 0)
+			if (bDrawTriangle)
+			{
+				AddLineWS(WP0, WP1, ColorWhite);
+				AddLineWS(WP1, WP2, ColorWhite);
+				AddLineWS(WP2, WP0, ColorWhite);
+			}
+
+			// Vertex normal
+			if (bDrawNormal)
 			{
-				P = RestTriangle.P0;
-				Barycentric = float3(1, 0, 0);
+				AddLineWS(WP1, WP1 + WN1 * NormalScale, ColorBlue);
+				AddLineWS(WP2, WP2 + WN2 * NormalScale, ColorBlue);
+				AddLineWS(WP0, WP0 + WN0 * NormalScale, ColorBlue);
 			}
-			if (TriangleVertex == 1)
+
+			// Triangle normal
+			if (bDrawNormal && 0)
 			{
-				P = RestTriangle.P1;
-				Barycentric = float3(0, 1, 0);
+				AddLineWS(WP0, WP0 + WFaceN * NormalScale, ColorGreen);
+				AddLineWS(WP1, WP1 + WFaceN * NormalScale, ColorGreen);
+				AddLineWS(WP2, WP2 + WFaceN * NormalScale, ColorGreen);
 			}
-			if (TriangleVertex == 2)
+		
+			// Root basis
+			if (bDrawRoot)
 			{
-				P = RestTriangle.P2;
-				Barycentric = float3(0, 0, 1);
+				AddLineWS(WO, WO + WN * NormalScale * 0.3f, ColorYellow);
 			}
-			P = TransformPoint(P, Barycentric, RestTriangle, DeformedTriangle);
 		}
-		#endif
-	
-		float4 Color = float4(Tri.UV0, 0, 1);
-		// Color = float4((Tri.N + 1) * 0.5f, 1);
-	
-		const float3 WP0 = mul(float4(Tri.P0, 1), RootLocalToWorld).xyz;
-		const float3 WP1 = mul(float4(Tri.P1, 1), RootLocalToWorld).xyz; 
-		const float3 WP2 = mul(float4(Tri.P2, 1), RootLocalToWorld).xyz;
-	
-		FShaderPrintContext Ctx = InitShaderPrintContext(true, uint2(0, 0));
-		AddLineWS(WP0, WP1, Color);
-		AddLineWS(WP1, WP2, Color);
-		AddLineWS(WP2, WP0, Color);
 	}
-#endif // INPUT_ROOT_TRIANGLE
 
-#if PERMUTATION_INPUT_TYPE == INPUT_ROOT_SAMPLES 
+#endif // INPUT_ROOTS
+
+#if PERMUTATION_INPUT_TYPE == INPUT_SAMPLES
 	const uint PointIndex = DispatchThreadId.x;	
 	if (PointIndex < MaxRootCount)
 	{
@@ -607,7 +680,7 @@ void MainCS(uint3 DispatchThreadId : SV_DispatchThreadID)
 		const float4 RBFColor = bRestPass==1 ? ColorPurple : ColorYellow;
 		AddCrossWS(P, 1.f, RBFColor);
 	}
-#endif // INPUT_ROOT_SAMPLES
+#endif // INPUT_SAMPLES
 }
 #endif // SHADER_MESH_PROJECTION_HAIR
 
diff --git a/Engine/Shaders/Private/HairStrands/HairStrandsFollicleMask.usf b/Engine/Shaders/Private/HairStrands/HairStrandsFollicleMask.usf
--- a/Engine/Shaders/Private/HairStrands/HairStrandsFollicleMask.usf
+++ b/Engine/Shaders/Private/HairStrands/HairStrandsFollicleMask.usf
@@ -47,12 +47,14 @@ void MainVS(
 		UniqueTriangleIndex = RootToUniqueTriangleIndexBuffer[RootIndex];
 	}
 
+	// TODO -> follicle requires UVs
+
 	if (UniqueTriangleIndex > MaxUniqueTriangleIndex) return;
-	const FHairMeshTriangleWithUVs Tri = GetTriangleWithUVsTransformation(UniqueTriangleIndex, UniqueTrianglePositionBuffer);
+	const FHairMeshBasis Tri = GetHairMeshBasis(UniqueTriangleIndex, UniqueTrianglePositionBuffer, float3(0,0,0));
 
 	const uint EncodedBarycentric = RootBarycentricBuffer[RootIndex];
 	const float3 Barycentric = UnpackBarycentrics(EncodedBarycentric);
-	float2 RootUV = Barycentric.x * Tri.UV0 + Barycentric.y * Tri.UV1 + Barycentric.z * Tri.UV2;
+	float2 RootUV = float2(0, 0);  // TODO Barycentric.x * Tri.UV0 + Barycentric.y * Tri.UV1 + Barycentric.z * Tri.UV2;
 #endif
 
 #if PERMUTATION_UV_TYPE == UV_ROOT
diff --git a/Engine/Shaders/Private/HairStrands/HairStrandsGuideDeform.usf b/Engine/Shaders/Private/HairStrands/HairStrandsGuideDeform.usf
--- a/Engine/Shaders/Private/HairStrands/HairStrandsGuideDeform.usf
+++ b/Engine/Shaders/Private/HairStrands/HairStrandsGuideDeform.usf
@@ -33,13 +33,12 @@ Buffer<uint>	  SimRootToUniqueTriangleIndexBuffer;
 // Return world position of a guide vertex deformed by a skinned triangle
 float3 DeformGuideByTriangle(
 	uint GuideIndex,
-	FHairMeshTriangle RestTriangle,
-	FHairMeshTriangle DeformedTriangle,
-	float3 RootBarycentric)
+	FHairMeshBasis RestBasis,
+	FHairMeshBasis DeformedBasis)
 {
 	const float3 WorldRestGuidePoint = ReadHairControlPointPosition(SimRestPosePositionBuffer, GuideIndex, SimRestOffset);
-	const float3 LocalRestGuidePoint = ToTriangle(WorldRestGuidePoint, RootBarycentric, RestTriangle);
-	const float3 WorldDeformedGuidePoint = ToWorld(LocalRestGuidePoint, RootBarycentric, DeformedTriangle);
+	const float3 LocalRestGuidePoint = ToTriangle(WorldRestGuidePoint, RestBasis);
+	const float3 WorldDeformedGuidePoint = ToWorld(LocalRestGuidePoint, DeformedBasis);
 
 	return WorldDeformedGuidePoint;
 }
@@ -77,10 +76,10 @@ float3 DisplacePosition(uint VertexIndex, float3 Pos)
 	const float3 RootBarycentric = UnpackBarycentrics(SimRootBarycentricBuffer[RootIndex]);
 
 	const uint TriangleIndex = SimRootToUniqueTriangleIndexBuffer[RootIndex];
-	const FHairMeshTriangle RestTriangle = GetTriangleTransformation(TriangleIndex, SimRestPositionBuffer);
-	const FHairMeshTriangle DeformedTriangle = GetTriangleTransformation(TriangleIndex, SimDeformedPositionBuffer);
+	const FHairMeshBasis RestBasis = GetHairMeshBasis(TriangleIndex, SimRestPositionBuffer, RootBarycentric);
+	const FHairMeshBasis DeformedBasis = GetHairMeshBasis(TriangleIndex, SimDeformedPositionBuffer, RootBarycentric);
 
-	return DeformGuideByTriangle(VertexIndex, RestTriangle, DeformedTriangle, RootBarycentric) - OutSimHairPositionOffset;
+	return DeformGuideByTriangle(VertexIndex, RestBasis, DeformedBasis) - OutSimHairPositionOffset;
 #endif
 
 #if PERMUTATION_DEFORMATION == DEFORMATION_OFFSET_GUIDE_GLOBAL
diff --git a/Engine/Shaders/Private/HairStrands/HairStrandsInterpolation.usf b/Engine/Shaders/Private/HairStrands/HairStrandsInterpolation.usf
--- a/Engine/Shaders/Private/HairStrands/HairStrandsInterpolation.usf
+++ b/Engine/Shaders/Private/HairStrands/HairStrandsInterpolation.usf
@@ -161,16 +161,15 @@ float3 ComputeLocalGeometryOffset(
 float3 ComputeDynamicGeometryOffset(
 	FGuideIndex GuideIndex, 
 	float GuideVertexWeight, 
-	FHairMeshTriangle RestTriangle, 
-	FHairMeshTriangle DeformedTriangle,
-	float3 RootBarycentric,
+	FHairMeshBasis RestBasis, 
+	FHairMeshBasis DeformedBasis,
 	float3 OutSimHairPositionOffset)
 {
 	const float3 RestGuidePoint = ReadHairControlPointPosition(SimRestPosePositionBuffer, GuideIndex.PointIndex, InSimHairPositionOffset);
-	const float3 LocalRestGuidePoint = ToTriangle(RestGuidePoint, RootBarycentric, RestTriangle);
+	const float3 LocalRestGuidePoint = ToTriangle(RestGuidePoint, RestBasis);
 
 	const float3 DeformedGuidePoint = ReadHairControlPointPosition(SimDeformedPositionBuffer, GuideIndex.PointIndex, OutSimHairPositionOffset);
-	const float3 LocalDeformedGuidePoint = ToTriangle(DeformedGuidePoint, RootBarycentric, DeformedTriangle);
+	const float3 LocalDeformedGuidePoint = ToTriangle(DeformedGuidePoint, DeformedBasis);
 
 	return (LocalDeformedGuidePoint - LocalRestGuidePoint) * GuideVertexWeight;
 }
@@ -189,14 +188,14 @@ float3 ComputeTranslateGeometryOffset(
 float3 ComputeRotateGeometryOffset(
 	FGuideIndex GuideIndex, 
 	float GuideVertexWeight,
-	FHairMeshTriangle RestTriangle, 
-	FHairMeshTriangle DeformedTriangle)
+	FHairMeshBasis RestBasis, 
+	FHairMeshBasis DeformedBasis)
 {	
 	const float3 RestGuidePoint = ReadHairControlPointPosition(SimRestPosePositionBuffer, GuideIndex.PointIndex) - ReadHairControlPointPosition(SimRestPosePositionBuffer, GuideIndex.RootPointIndex);
-	const float3 LocalRestGuidePoint = VectorToTriangle(RestGuidePoint, RestTriangle);
+	const float3 LocalRestGuidePoint = VectorToTriangle(RestGuidePoint, RestBasis);
 
 	const float3 DeformedGuidePoint = ReadHairControlPointPosition(SimDeformedPositionBuffer, GuideIndex.PointIndex) - ReadHairControlPointPosition(SimDeformedPositionBuffer, GuideIndex.RootPointIndex);
-	const float3 LocalDeformedGuidePoint = VectorToTriangle(DeformedGuidePoint, DeformedTriangle);
+	const float3 LocalDeformedGuidePoint = VectorToTriangle(DeformedGuidePoint, DeformedBasis);
 
 	return (LocalDeformedGuidePoint - LocalRestGuidePoint) * GuideVertexWeight;
 }
@@ -204,30 +203,29 @@ float3 ComputeRotateGeometryOffset(
 ///////////////////////////////////////////////////////////////////////////////////////////////////
 // Root triangles
 
-struct FRootTriangleData
+struct FRootBasis
 {
-	float3 RootBarycentric;
-	FHairMeshTriangle RestTriangle;
-	FHairMeshTriangle DeformedTriangle;
+	FHairMeshBasis RestBasis;
+	FHairMeshBasis DeformedBasis;
 };
 
-FRootTriangleData GetRootTriangleData(uint CurveIndex, bool bSim)
+FRootBasis GetRootBasis(uint CurveIndex, bool bSim)
 {
-	FRootTriangleData Out = (FRootTriangleData)0;
+	FRootBasis Out = (FRootBasis)0;
 	#if PERMUTATION_DYNAMIC_GEOMETRY >= 1 
 	if (bSim)
 	{
 		const uint TriangleIndex = SimRootToUniqueTriangleIndexBuffer[CurveIndex];
-		Out.RootBarycentric  = UnpackBarycentrics(SimRootBarycentricBuffer[CurveIndex]);
-		Out.RestTriangle	 = GetTriangleTransformation(TriangleIndex, SimRootRestPositionBuffer);
-		Out.DeformedTriangle = GetTriangleTransformation(TriangleIndex, SimRootDeformedPositionBuffer);
+		const float3 RootBarycentric  = UnpackBarycentrics(SimRootBarycentricBuffer[CurveIndex]);
+		Out.RestBasis	  = GetHairMeshBasis(TriangleIndex, SimRootRestPositionBuffer, RootBarycentric);
+		Out.DeformedBasis = GetHairMeshBasis(TriangleIndex, SimRootDeformedPositionBuffer, RootBarycentric);
 	}
 	else
 	{
 		const uint TriangleIndex = RenRootToUniqueTriangleIndexBuffer[CurveIndex];
-		Out.RootBarycentric  = UnpackBarycentrics(RenRootBarycentricBuffer[CurveIndex]);
-		Out.RestTriangle	 = GetTriangleTransformation(TriangleIndex, RenRootRestPositionBuffer);
-		Out.DeformedTriangle = GetTriangleTransformation(TriangleIndex, RenRootDeformedPositionBuffer);
+		const float3 RootBarycentric  = UnpackBarycentrics(RenRootBarycentricBuffer[CurveIndex]);
+		Out.RestBasis	  = GetHairMeshBasis(TriangleIndex, RenRootRestPositionBuffer, RootBarycentric);
+		Out.DeformedBasis = GetHairMeshBasis(TriangleIndex, RenRootDeformedPositionBuffer, RootBarycentric);
 	}
 	#endif
 	return Out;
@@ -258,7 +256,7 @@ float3 ApplyRigidInterpolation(FGuideDataWithOffset GuideData, float3 InPosition
 
 ///////////////////////////////////////////////////////////////////////////////////////////////////
 // Skinning with offset
-float3 ApplySkinningWithOffset(const FGuideDataWithOffset GuideData, const FRootTriangleData RenTri, const FRootTriangleData SimTris[2], float3 InPosition)
+float3 ApplySkinningWithOffset(const FGuideDataWithOffset GuideData, const FRootBasis RenTri, const FRootBasis SimTris[2], float3 InPosition)
 {
 	float3 CurrOffset = 0;
 
@@ -271,26 +269,26 @@ float3 ApplySkinningWithOffset(const FGuideDataWithOffset GuideData, const FRoot
 	for (uint K = 0; K < HAIR_INTERPOLATION_MAX_GUIDE_COUNT; ++K)
 	#endif
 	{
-		const FRootTriangleData SimTri = SimTris[K];
+		const FRootBasis SimTri = SimTris[K];
 
-		const float3 Offset0 = ComputeDynamicGeometryOffset(ToGuideIndex(GuideData, K, 0), GuideData.PointWeights[K], SimTri.RestTriangle, SimTri.DeformedTriangle, SimTri.RootBarycentric, GuideData.OutPositionOffset);
-		const float3 Offset1 = ComputeDynamicGeometryOffset(ToGuideIndex(GuideData, K, 1), GuideData.PointWeights[K], SimTri.RestTriangle, SimTri.DeformedTriangle, SimTri.RootBarycentric, GuideData.OutPositionOffset);
+		const float3 Offset0 = ComputeDynamicGeometryOffset(ToGuideIndex(GuideData, K, 0), GuideData.PointWeights[K], SimTri.RestBasis, SimTri.DeformedBasis, GuideData.OutPositionOffset);
+		const float3 Offset1 = ComputeDynamicGeometryOffset(ToGuideIndex(GuideData, K, 1), GuideData.PointWeights[K], SimTri.RestBasis, SimTri.DeformedBasis, GuideData.OutPositionOffset);
 
-		CurrOffset += VectorToWorld(lerp(Offset0, Offset1, GuideData.PointLerps[K]), SimTri.DeformedTriangle);
+		CurrOffset += VectorToWorld(lerp(Offset0, Offset1, GuideData.PointLerps[K]), SimTri.DeformedBasis);
 	}
 #endif		
 
-	InPosition = ToTriangle(InPosition, RenTri.RootBarycentric, RenTri.RestTriangle);
-	InPosition = ToWorld(InPosition, RenTri.RootBarycentric, RenTri.DeformedTriangle) + CurrOffset;
+	InPosition = ToTriangle(InPosition, RenTri.RestBasis);
+	InPosition = ToWorld(InPosition, RenTri.DeformedBasis) + CurrOffset;
 
 	return InPosition;
 }
 
 ///////////////////////////////////////////////////////////////////////////////////////////////////
 // Skinning with translantion
-float3 ApplySkinningWithTranslation(const FGuideDataWithOffset GuideData, const FRootTriangleData RenTri, float3 InPosition)
+float3 ApplySkinningWithTranslation(const FGuideDataWithOffset GuideData, const FRootBasis RenTri, float3 InPosition)
 {
-	InPosition -= RenTri.RootBarycentric.x * RenTri.RestTriangle.P0 + RenTri.RootBarycentric.y * RenTri.RestTriangle.P1 + RenTri.RootBarycentric.z * RenTri.RestTriangle.P2;
+	InPosition -= RenTri.RestBasis.BaseO;
 
 	float3 CurrOffset = 0;
 
@@ -309,15 +307,15 @@ float3 ApplySkinningWithTranslation(const FGuideDataWithOffset GuideData, const
 	}
 #endif
 
-	InPosition += RenTri.RootBarycentric.x * RenTri.DeformedTriangle.P0 + RenTri.RootBarycentric.y * RenTri.DeformedTriangle.P1 + RenTri.RootBarycentric.z * RenTri.DeformedTriangle.P2;
+	InPosition += RenTri.DeformedBasis.BaseO;
 	return InPosition + CurrOffset;
 }
 
 ///////////////////////////////////////////////////////////////////////////////////////////////////
 // Skinning with rotation
-float3 ApplySkinningWithRotation(const FGuideDataWithOffset GuideData, const FRootTriangleData RenTri, const FRootTriangleData SimTris[2], float3 InPosition)
+float3 ApplySkinningWithRotation(const FGuideDataWithOffset GuideData, const FRootBasis RenTri, const FRootBasis SimTris[2], float3 InPosition)
 {
-	InPosition -= RenTri.RootBarycentric.x * RenTri.RestTriangle.P0 + RenTri.RootBarycentric.y * RenTri.RestTriangle.P1 + RenTri.RootBarycentric.z * RenTri.RestTriangle.P2;
+	InPosition -= RenTri.RestBasis.BaseO;
 
 	float3 CurrOffset = 0;
 
@@ -329,24 +327,24 @@ float3 ApplySkinningWithRotation(const FGuideDataWithOffset GuideData, const FRo
 	for (uint K = 0; K < HAIR_INTERPOLATION_MAX_GUIDE_COUNT; ++K)
 	#endif
 	{
-		const FRootTriangleData SimTri = SimTris[K];
+		const FRootBasis SimTri = SimTris[K];
 
-		const float3 LocalPoint = VectorToTriangle(InPosition, SimTri.RestTriangle);
+		const float3 LocalPoint = VectorToTriangle(InPosition, SimTri.RestBasis);
 
-		const float3 Offset0 = ComputeRotateGeometryOffset(ToGuideIndex(GuideData, K, 0), 1.0, SimTri.RestTriangle, SimTri.DeformedTriangle/*, SimHairPositionOffset*/);
-		const float3 Offset1 = ComputeRotateGeometryOffset(ToGuideIndex(GuideData, K, 1), 1.0, SimTri.RestTriangle, SimTri.DeformedTriangle/*, SimHairPositionOffset*/);
+		const float3 Offset0 = ComputeRotateGeometryOffset(ToGuideIndex(GuideData, K, 0), 1.0, SimTri.RestBasis, SimTri.DeformedBasis/*, SimHairPositionOffset*/);
+		const float3 Offset1 = ComputeRotateGeometryOffset(ToGuideIndex(GuideData, K, 1), 1.0, SimTri.RestBasis, SimTri.DeformedBasis/*, SimHairPositionOffset*/);
 
-		CurrOffset += GuideData.PointWeights[K] * VectorToWorld(lerp(Offset0, Offset1, GuideData.PointLerps[K]) + LocalPoint, SimTri.DeformedTriangle);
+		CurrOffset += GuideData.PointWeights[K] * VectorToWorld(lerp(Offset0, Offset1, GuideData.PointLerps[K]) + LocalPoint, SimTri.DeformedBasis);
 	}
 #endif
 
-	InPosition = RenTri.RootBarycentric.x * RenTri.DeformedTriangle.P0 + RenTri.RootBarycentric.y * RenTri.DeformedTriangle.P1 + RenTri.RootBarycentric.z * RenTri.DeformedTriangle.P2;
+	InPosition = RenTri.DeformedBasis.BaseO;
 	return InPosition + CurrOffset;
 }
 
 ///////////////////////////////////////////////////////////////////////////////////////////////////
 // Skinning with translation and rotation
-float3 ApplySkinningWithTranslationAndRotation(const FGuideDataWithOffset GuideData, const FRootTriangleData RenTri, const FRootTriangleData SimTris[2], float3 InPosition)
+float3 ApplySkinningWithTranslationAndRotation(const FGuideDataWithOffset GuideData, const FRootBasis RenTri, const FRootBasis SimTris[2], float3 InPosition)
 {
 	float3 CurrOffset = 0;
 
@@ -358,14 +356,14 @@ float3 ApplySkinningWithTranslationAndRotation(const FGuideDataWithOffset GuideD
 	for (uint K = 0; K < HAIR_INTERPOLATION_MAX_GUIDE_COUNT; ++K)
 	#endif
 	{
-		const FRootTriangleData SimTri = SimTris[K];  //GetRootTriangleData(GuideIndex0, true);
+		const FRootBasis SimTri = SimTris[K];  //GetRootTriangleData(GuideIndex0, true);
 
-		const float3 LocalPoint = ToTriangle(InPosition, SimTri.RootBarycentric, SimTri.RestTriangle);
+		const float3 LocalPoint = ToTriangle(InPosition, SimTri.RestBasis);
 
-		const float3 Offset0 = ComputeDynamicGeometryOffset(ToGuideIndex(GuideData, K, 0), 1.0, SimTri.RestTriangle, SimTri.DeformedTriangle, SimTri.RootBarycentric, GuideData.OutPositionOffset);
-		const float3 Offset1 = ComputeDynamicGeometryOffset(ToGuideIndex(GuideData, K, 1), 1.0, SimTri.RestTriangle, SimTri.DeformedTriangle, SimTri.RootBarycentric, GuideData.OutPositionOffset);
+		const float3 Offset0 = ComputeDynamicGeometryOffset(ToGuideIndex(GuideData, K, 0), 1.0, SimTri.RestBasis, SimTri.DeformedBasis, GuideData.OutPositionOffset);
+		const float3 Offset1 = ComputeDynamicGeometryOffset(ToGuideIndex(GuideData, K, 1), 1.0, SimTri.RestBasis, SimTri.DeformedBasis, GuideData.OutPositionOffset);
 
-		CurrOffset += GuideData.PointWeights[K] * ToWorld(lerp(Offset0, Offset1, GuideData.PointLerps[K]) + LocalPoint, SimTri.RootBarycentric, SimTri.DeformedTriangle);
+		CurrOffset += GuideData.PointWeights[K] * ToWorld(lerp(Offset0, Offset1, GuideData.PointLerps[K]) + LocalPoint, SimTri.DeformedBasis);
 	}
 #endif
 
@@ -374,19 +372,19 @@ float3 ApplySkinningWithTranslationAndRotation(const FGuideDataWithOffset GuideD
 
 ///////////////////////////////////////////////////////////////////////////////////////////////////
 // Deformer transformation
-float3 ApplyDeformer(const FRootTriangleData RenTri, float3 InRestPostion, float3 InDeformerPosition, float3 InDeformedPosition)
+float3 ApplyDeformer(const FRootBasis RenTri, float3 InRestPostion, float3 InDeformerPosition, float3 InDeformedPosition)
 {
 	float3 OutPosition = InDeformedPosition;
 #if PERMUTATION_DEFORMER
 	#if PERMUTATION_DYNAMIC_GEOMETRY == 0
 		OutPosition = InDeformedPosition + (InDeformerPosition - InRestPostion);
 	#else
-		const float3 LocalRestPoint     = ToTriangle(InRestPostion, RenTri.RootBarycentric, RenTri.RestTriangle);
-		const float3 LocalDeformerPoint = ToTriangle(InDeformerPosition, RenTri.RootBarycentric, RenTri.RestTriangle);
+		const float3 LocalRestPoint     = ToTriangle(InRestPostion, RenTri.RestBasis);
+		const float3 LocalDeformerPoint = ToTriangle(InDeformerPosition, RenTri.RestBasis);
 
 		const float3 LocalOffset = LocalDeformerPoint - LocalRestPoint;
-		const float3 LocalDeformedPoint = ToTriangle(InDeformedPosition, RenTri.RootBarycentric, RenTri.DeformedTriangle);
-		OutPosition = ToWorld(LocalDeformedPoint + LocalOffset, RenTri.RootBarycentric, RenTri.DeformedTriangle);
+		const float3 LocalDeformedPoint = ToTriangle(InDeformedPosition, RenTri.DeformedBasis);
+		OutPosition = ToWorld(LocalDeformedPoint + LocalOffset, RenTri.DeformedBasis);
 	#endif
 #endif
 	return OutPosition;
@@ -422,11 +420,11 @@ uint GetGlobalCurveIndex(uint ThreadIndex, uint2 GroupIndex)
 groupshared uint CurvePointIndex[MAX_CURVE];
 groupshared uint CurvePointCount[MAX_CURVE];
 
-groupshared FRootTriangleData g_RenTri[MAX_CURVE];
+groupshared FRootBasis g_RenTri[MAX_CURVE];
 groupshared float3 RenPositionOffset;
 
 groupshared FGuideCurveData SimCurveData[MAX_CURVE];
-groupshared FRootTriangleData g_SimTris[MAX_CURVE][2];
+groupshared FRootBasis g_SimTris[MAX_CURVE][2];
 
 [numthreads(GROUP_SIZE, 1, 1)]
 void MainCS(uint3 DispatchThreadId : SV_DispatchThreadID, uint ThreadIndex : SV_GroupIndex, uint2 GroupIndex : SV_GroupID)
@@ -460,16 +458,16 @@ void MainCS(uint3 DispatchThreadId : SV_DispatchThreadID, uint ThreadIndex : SV_
 			#endif
 
 			// Load ren triangles
-			g_RenTri[LocalCurveIndex] = GetRootTriangleData(GlobalCurveIndex, false/*bSim*/);
+			g_RenTri[LocalCurveIndex] = GetRootBasis(GlobalCurveIndex, false/*bSim*/);
 
 			// Load sim triangles
-			g_SimTris[LocalCurveIndex][0] = (FRootTriangleData)0;
-			g_SimTris[LocalCurveIndex][1] = (FRootTriangleData)0;
+			g_SimTris[LocalCurveIndex][0] = (FRootBasis)0;
+			g_SimTris[LocalCurveIndex][1] = (FRootBasis)0;
 			#if PERMUTATION_SIMULATION
-			g_SimTris[LocalCurveIndex][0] = GetRootTriangleData(SimCurveData[LocalCurveIndex].CurveIndices[0], true/*bSim*/);
+			g_SimTris[LocalCurveIndex][0] = GetRootBasis(SimCurveData[LocalCurveIndex].CurveIndices[0], true/*bSim*/);
 			if (!PERMUTATION_USE_SINGLE_GUIDE)
 			{
-				g_SimTris[LocalCurveIndex][1] = GetRootTriangleData(SimCurveData[LocalCurveIndex].CurveIndices[1], true/*bSim*/);
+				g_SimTris[LocalCurveIndex][1] = GetRootBasis(SimCurveData[LocalCurveIndex].CurveIndices[1], true/*bSim*/);
 			}
 			#endif
 		}
@@ -504,14 +502,14 @@ void MainCS(uint3 DispatchThreadId : SV_DispatchThreadID, uint ThreadIndex : SV_
 				#endif
 
 				// Remap from LDS to local variable to ease the setup
-				FRootTriangleData RenTri = g_RenTri[LocalCurveIndex];
-				FRootTriangleData SimTris[2];
+				FRootBasis RenTri = g_RenTri[LocalCurveIndex];
+				FRootBasis SimTris[2];
 				#if PERMUTATION_SIMULATION
 				SimTris[0] = g_SimTris[LocalCurveIndex][0];
 				SimTris[1] = g_SimTris[LocalCurveIndex][1];
 				#else
-				SimTris[0] = (FRootTriangleData)0;
-				SimTris[1] = (FRootTriangleData)0;
+				SimTris[0] = (FRootBasis)0;
+				SimTris[1] = (FRootBasis)0;
 				#endif
 
 				// Manual decoding of the rest position
@@ -581,24 +579,18 @@ void MainCS(uint3 DispatchThreadId : SV_DispatchThreadID, uint ThreadIndex : SV_
 #if PERMUTATION_POINT_PER_CURVE == TOTALMAXPOINTPERCURVE
 
 #if PERMUTATION_WAVEOPS
-FRootTriangleData WaveReadLaneFirstTriangle(FRootTriangleData In)
+FRootBasis WaveReadLaneFirstTriangle(FRootBasis In)
 {
-	FRootTriangleData Out;
-	Out.RootBarycentric = WaveReadLaneFirst(In.RootBarycentric);
-
-	Out.RestTriangle.P0 = WaveReadLaneFirst(In.RestTriangle.P0);
-	Out.RestTriangle.P1 = WaveReadLaneFirst(In.RestTriangle.P1);
-	Out.RestTriangle.P2 = WaveReadLaneFirst(In.RestTriangle.P2);
-	Out.RestTriangle.E0 = WaveReadLaneFirst(In.RestTriangle.E0);
-	Out.RestTriangle.E1 = WaveReadLaneFirst(In.RestTriangle.E1);
-	Out.RestTriangle.N  = WaveReadLaneFirst(In.RestTriangle.N) ;
-
-	Out.DeformedTriangle.P0 = WaveReadLaneFirst(In.DeformedTriangle.P0);
-	Out.DeformedTriangle.P1 = WaveReadLaneFirst(In.DeformedTriangle.P1);
-	Out.DeformedTriangle.P2 = WaveReadLaneFirst(In.DeformedTriangle.P2);
-	Out.DeformedTriangle.E0 = WaveReadLaneFirst(In.DeformedTriangle.E0);
-	Out.DeformedTriangle.E1 = WaveReadLaneFirst(In.DeformedTriangle.E1);
-	Out.DeformedTriangle.N  = WaveReadLaneFirst(In.DeformedTriangle.N) ;
+	FRootBasis Out;
+	Out.RestBasis.BaseO = WaveReadLaneFirst(In.RestBasis.BaseO);
+	Out.RestBasis.BaseT = WaveReadLaneFirst(In.RestBasis.BaseT);
+	Out.RestBasis.BaseB = WaveReadLaneFirst(In.RestBasis.BaseB);
+	Out.RestBasis.BaseN = WaveReadLaneFirst(In.RestBasis.BaseN);
+
+	Out.DeformedBasis.BaseO = WaveReadLaneFirst(In.DeformedBasis.BaseO);
+	Out.DeformedBasis.BaseT = WaveReadLaneFirst(In.DeformedBasis.BaseT);
+	Out.DeformedBasis.BaseB = WaveReadLaneFirst(In.DeformedBasis.BaseB);
+	Out.DeformedBasis.BaseN = WaveReadLaneFirst(In.DeformedBasis.BaseN);
 
 	return Out;
 }
@@ -615,11 +607,11 @@ FGuideCurveData WaveReadLaneFirstCurveData(FGuideCurveData In)
 groupshared uint CurvePointIndex;
 groupshared uint CurvePointCount;
 
-groupshared FRootTriangleData RenTri;
+groupshared FRootBasis RenTri;
 groupshared float3 RenPositionOffset;
 
 groupshared FGuideCurveData SimCurveData;
-groupshared FRootTriangleData SimTris[2];
+groupshared FRootBasis SimTris[2];
 #endif
 
 #if PERMUTATION_WAVEOPS && COMPILER_SUPPORTS_WAVE_SIZE
@@ -642,12 +634,12 @@ void MainCS(uint3 DispatchThreadId : SV_DispatchThreadID, uint ThreadIndex : SV_
 	uint CurvePointCount = 0;
 
 	float3 RenPositionOffset = 0;
-	FRootTriangleData RenTri = (FRootTriangleData)0;
+	FRootBasis RenTri = (FRootBasis)0;
 
 	FGuideCurveData SimCurveData = (FGuideCurveData)0;
-	FRootTriangleData SimTris[2];
-	SimTris[0] = (FRootTriangleData)0;
-	SimTris[1] = (FRootTriangleData)0;
+	FRootBasis SimTris[2];
+	SimTris[0] = (FRootBasis)0;
+	SimTris[1] = (FRootBasis)0;
 	#endif
 
 	// 2. Curve's point count/offset
@@ -678,18 +670,18 @@ void MainCS(uint3 DispatchThreadId : SV_DispatchThreadID, uint ThreadIndex : SV_
 			#endif
 
 			// Load sim/ren triangles
-			RenTri = GetRootTriangleData(CurveIndex, false/*bSim*/);
+			RenTri = GetRootBasis(CurveIndex, false/*bSim*/);
 
 			#if PERMUTATION_SIMULATION
-			SimTris[0] = GetRootTriangleData(SimCurveData.CurveIndices[0], true/*bSim*/);
+			SimTris[0] = GetRootBasis(SimCurveData.CurveIndices[0], true/*bSim*/);
 			SimTris[1] = SimTris[0];
 			if (!PERMUTATION_USE_SINGLE_GUIDE)
 			{
-				SimTris[1] = GetRootTriangleData(SimCurveData.CurveIndices[1], true/*bSim*/);
+				SimTris[1] = GetRootBasis(SimCurveData.CurveIndices[1], true/*bSim*/);
 			}
 			#else
-			SimTris[0] = (FRootTriangleData)0;
-			SimTris[1] = (FRootTriangleData)0;
+			SimTris[0] = (FRootBasis)0;
+			SimTris[1] = (FRootBasis)0;
 			#endif
 		}
 		RenPositionOffset = GetRenPositionOffset();
diff --git a/Engine/Shaders/Private/HairStrands/HairStrandsMesh.usf b/Engine/Shaders/Private/HairStrands/HairStrandsMesh.usf
--- a/Engine/Shaders/Private/HairStrands/HairStrandsMesh.usf
+++ b/Engine/Shaders/Private/HairStrands/HairStrandsMesh.usf
@@ -142,11 +142,13 @@ struct FMeshSectionData
 	uint UVsChannelOffset;
 	uint UVsChannelCount;
 	uint bIsSwapped;
+	uint bUseFaceNormal;
 };
 
+uint TangentFormat;
 uint MaxSectionCount;
 uint MaxUniqueTriangleCount;
-Buffer<float2> MeshUVsBuffer;
+Buffer<float4> MeshTangentBuffer;
 Buffer<uint> MeshIndexBuffer;
 Buffer<uint> UniqueTriangleIndices;
 StructuredBuffer<uint4> MeshSectionBuffer;
@@ -181,40 +183,46 @@ FMeshSectionData ReadMeshSectionData(uint InSectionIndex)
 	Out.UVsChannelOffset = BitFieldExtractU32(Packed.w, 8, 0);
 	Out.UVsChannelCount = BitFieldExtractU32(Packed.w, 8, 8);
 	Out.bIsSwapped = BitFieldExtractU32(Packed.w, 1, 16);
+	Out.bUseFaceNormal = BitFieldExtractU32(Packed.w, 1, 17);
 	return Out;
 }
 
-FHairMeshTriangleWithUVs GetTriangleData(FHairTriangleIndex TriangleIndex, FMeshSectionData MeshSectionData, bool bPrevious)
+FHairMeshTriangle GetHairMeshTriangle(uint InTriangleIndex, FMeshSectionData MeshSectionData, bool bPrevious)
 {
 	const bool bSamplePrev = (MeshSectionData.bIsSwapped > 0 && !bPrevious) || (MeshSectionData.bIsSwapped == 0 && bPrevious);
 	if (bSamplePrev)
 	{
-		return GetTriangleData(
-			TriangleIndex.TriangleIndex, 
+		return GetHairMeshTriangle(
+			InTriangleIndex, 
 			MeshPreviousPositionBuffer, 
-			MeshUVsBuffer, 
+			MeshTangentBuffer,
 			MeshIndexBuffer, 
 			MeshSectionData.IndexOffset,
 			MeshSectionData.MaxIndexCount,
 			MeshSectionData.MaxVertexCount,
-			MeshSectionData.UVsChannelOffset, 
-			MeshSectionData.UVsChannelCount);
+			TangentFormat);
 	}
 	else
 	{
-		return GetTriangleData(
-			TriangleIndex.TriangleIndex, 
+		return GetHairMeshTriangle(
+			InTriangleIndex, 
 			MeshPositionBuffer, 
-			MeshUVsBuffer, 
+			MeshTangentBuffer,
 			MeshIndexBuffer, 
 			MeshSectionData.IndexOffset,
 			MeshSectionData.MaxIndexCount,
 			MeshSectionData.MaxVertexCount,
-			MeshSectionData.UVsChannelOffset, 
-			MeshSectionData.UVsChannelCount);
+			TangentFormat);
 	}
 }
 
+float3 GetTriangleFaceNormal(FHairMeshTriangle Tri)
+{
+	const float3 E0 = Tri.P1 - Tri.P0;
+	const float3 E1 = Tri.P2 - Tri.P0;
+	return normalize(cross(E1, E0));
+}
+
 [numthreads(GROUP_SIZE, 1, 1)]
 void MainCS(uint3 DispatchThreadId : SV_DispatchThreadID)
 { 
@@ -231,17 +239,29 @@ void MainCS(uint3 DispatchThreadId : SV_DispatchThreadID)
 	
 		// Compute the slot index of the resources for the section index (this is done in multiple pass
 		{
-			const FHairMeshTriangleWithUVs Tri = GetTriangleData(TriangleIndex, MeshSectionData, false);
-			OutUniqueTriangleCurrPosition[UniqueTriangleIndex * 3 + 0] = float4(Tri.P0, PackUVsToFloat(Tri.UV0));
-			OutUniqueTriangleCurrPosition[UniqueTriangleIndex * 3 + 1] = float4(Tri.P1, PackUVsToFloat(Tri.UV1));
-			OutUniqueTriangleCurrPosition[UniqueTriangleIndex * 3 + 2] = float4(Tri.P2, PackUVsToFloat(Tri.UV2));
+			FHairMeshTriangle Tri = GetHairMeshTriangle(TriangleIndex.TriangleIndex, MeshSectionData, false);
+			if (MeshSectionData.bUseFaceNormal)
+			{
+				Tri.N0 = GetTriangleFaceNormal(Tri);
+				Tri.N1 = Tri.N0;
+				Tri.N2 = Tri.N0;
+			}
+			OutUniqueTriangleCurrPosition[UniqueTriangleIndex * 3 + 0] = float4(Tri.P0, PackVertexNormalToFloat(Tri.N0));
+			OutUniqueTriangleCurrPosition[UniqueTriangleIndex * 3 + 1] = float4(Tri.P1, PackVertexNormalToFloat(Tri.N1));
+			OutUniqueTriangleCurrPosition[UniqueTriangleIndex * 3 + 2] = float4(Tri.P2, PackVertexNormalToFloat(Tri.N2));
 		}
 		#if PERMUTATION_PREVIOUS
 		{
-			const FHairMeshTriangleWithUVs Tri = GetTriangleData(TriangleIndex, MeshSectionData, true);
-			OutUniqueTrianglePrevPosition[UniqueTriangleIndex * 3 + 0] = float4(Tri.P0, PackUVsToFloat(Tri.UV0));
-			OutUniqueTrianglePrevPosition[UniqueTriangleIndex * 3 + 1] = float4(Tri.P1, PackUVsToFloat(Tri.UV1));
-			OutUniqueTrianglePrevPosition[UniqueTriangleIndex * 3 + 2] = float4(Tri.P2, PackUVsToFloat(Tri.UV2));
+			FHairMeshTriangle Tri = GetHairMeshTriangle(TriangleIndex.TriangleIndex, MeshSectionData, true);
+			if (MeshSectionData.bUseFaceNormal)
+			{
+				Tri.N0 = GetTriangleFaceNormal(Tri);
+				Tri.N1 = Tri.N0;
+				Tri.N2 = Tri.N0;
+			}
+			OutUniqueTrianglePrevPosition[UniqueTriangleIndex * 3 + 0] = float4(Tri.P0, PackVertexNormalToFloat(Tri.N0));
+			OutUniqueTrianglePrevPosition[UniqueTriangleIndex * 3 + 1] = float4(Tri.P1, PackVertexNormalToFloat(Tri.N1));
+			OutUniqueTrianglePrevPosition[UniqueTriangleIndex * 3 + 2] = float4(Tri.P2, PackVertexNormalToFloat(Tri.N2));
 		}
 		#endif
 	}
diff --git a/Engine/Shaders/Private/HairStrands/HairStrandsPack.ush b/Engine/Shaders/Private/HairStrands/HairStrandsPack.ush
--- a/Engine/Shaders/Private/HairStrands/HairStrandsPack.ush
+++ b/Engine/Shaders/Private/HairStrands/HairStrandsPack.ush
@@ -439,6 +439,36 @@ float2 UnPackUVs(float InUV)
 	return UnPackUVs(asuint(InUV));
 }
 
+////////////////////////////////////////////////////////////////////////////////////////////////////////////////////////////////
+// Normal
+
+uint PackVertexNormal(float3 In)
+{
+	return
+		(PackUnorm10(In.x * 0.5f + 0.5f)    )|
+		(PackUnorm10(In.y * 0.5f + 0.5f)<<10)|
+		(PackUnorm10(In.z * 0.5f + 0.5f)<<20);
+}
+
+float PackVertexNormalToFloat(float3 In)
+{
+	return asfloat(PackVertexNormal(In));
+}
+
+float3 UnpackVertexNormal(uint In)
+{
+	const float3 N = float3(
+		UnpackUnorm10(In    ) * 2.f - 1.f,
+		UnpackUnorm10(In>>10) * 2.f - 1.f,
+		UnpackUnorm10(In>>20) * 2.f - 1.f);
+	return normalize(N);
+}
+
+float3 UnpackVertexNormal(float In)
+{
+	return UnpackVertexNormal(asuint(In));
+}
+
 ////////////////////////////////////////////////////////////////////////////////////////////////////////////////////////////////
 // Barycentrics
 
diff --git a/Engine/Source/Runtime/Engine/Private/GPUSkinCache.cpp b/Engine/Source/Runtime/Engine/Private/GPUSkinCache.cpp
--- a/Engine/Source/Runtime/Engine/Private/GPUSkinCache.cpp
+++ b/Engine/Source/Runtime/Engine/Private/GPUSkinCache.cpp
@@ -410,6 +410,7 @@ public:
 			MeshSection.PositionBuffer = DispatchData[SectionIndex].PositionBuffer->Buffer.SRV;
 			MeshSection.PreviousPositionBuffer = DispatchData[SectionIndex].PreviousPositionBuffer->Buffer.SRV;
 			MeshSection.UVsBuffer = DispatchData[SectionIndex].UVsBufferSRV;
+			MeshSection.TangentBuffer = DispatchData[SectionIndex].TangentBufferSRV;
 			MeshSection.TotalVertexCount = DispatchData[SectionIndex].PositionBuffer->Buffer.NumBytes / (sizeof(float) * 3);
 			MeshSection.NumPrimitives = Section.NumTriangles;
 			MeshSection.NumVertices = Section.NumVertices;
diff --git a/Engine/Source/Runtime/Engine/Public/CachedGeometry.h b/Engine/Source/Runtime/Engine/Public/CachedGeometry.h
--- a/Engine/Source/Runtime/Engine/Public/CachedGeometry.h
+++ b/Engine/Source/Runtime/Engine/Public/CachedGeometry.h
@@ -18,7 +18,8 @@ struct FCachedGeometry
 	{
 		FRDGBufferSRVRef RDGPositionBuffer = nullptr;
 		FRDGBufferSRVRef RDGPreviousPositionBuffer = nullptr;
-		
+		FRDGBufferSRVRef RDGTangentBuffer = nullptr;
+
 		FRHIShaderResourceView* PositionBuffer = nullptr;			// Valid when the input comes from the GPUSkinCache (since it is doesn't use RDG)
 		FRHIShaderResourceView* PreviousPositionBuffer = nullptr;	// Valid when the input comes from the GPUSkinCache (since it is doesn't use RDG)
 		FRHIShaderResourceView* TangentBuffer = nullptr;			// Valid when the input comes from the GPUSkinCache (since it is doesn't use RDG)
@@ -36,6 +37,7 @@ struct FCachedGeometry
 		uint32 TotalIndexCount = 0;
 		uint32 SectionIndex = 0;
 		int32 LODIndex = -1;
+		int32 TangentFormat = 0;  // 0: Interleaved TangentX/TangentZ, 1: TangentZ only
 	};
 
 	int32 LODIndex = -1;
