Remove Mobile Multi View inputs from Mobile AO CS

diff --git a/Engine/Shaders/Private/PostProcessAmbientOcclusionMobile.usf b/Engine/Shaders/Private/PostProcessAmbientOcclusionMobile.usf
--- a/Engine/Shaders/Private/PostProcessAmbientOcclusionMobile.usf
+++ b/Engine/Shaders/Private/PostProcessAmbientOcclusionMobile.usf
@@ -779,26 +779,8 @@ void GTAOHorizonSearchIntegralSpatialFilterCS(
 	uint2 GroupId : SV_GroupID,
 	uint2 DispatchThreadId : SV_DispatchThreadID,
 	uint2 GroupThreadId : SV_GroupThreadID
-#if INSTANCED_STEREO
-	, uint InstanceId : SV_InstanceID
-	, out uint LayerIndex : SV_RenderTargetArrayIndex
-#elif MOBILE_MULTI_VIEW
-	, in uint ViewId : SV_ViewID
-	, out nointerpolation uint MultiViewIndex : VIEW_ID
-#endif
 )
 {
-#if INSTANCED_STEREO
-	const uint EyeIndex = GetEyeIndex(InstanceId);
-	ResolvedView = ResolveView(EyeIndex);
-	LayerIndex = EyeIndex;
-#elif MOBILE_MULTI_VIEW
-	ResolvedView = ResolveView(ViewId);
-	MultiViewIndex = ViewId;
-#else
-	ResolvedView = ResolveView();
-#endif
-
 	int2 FullGroupOrigin = int2(GroupId.x * THREADGROUP_SIZEX, GroupId.y * THREADGROUP_SIZEY);
 	uint pixIdx = GroupIndex;
 
@@ -809,7 +791,7 @@ void GTAOHorizonSearchIntegralSpatialFilterCS(
 		pixIdx = GroupIndex * 2;
 		if (pixIdx < MAX_DEPTH_THREADS)
 		{
-			CacheZVal(FullGroupOriginDepth, pixIdx);
+			CacheZVal(FullGroupOriginDepth, pixIdx); 
 
 			CacheZVal(FullGroupOriginDepth, pixIdx + 1);
 		}
@@ -844,26 +826,8 @@ void GTAOHorizonSearchIntegralCS(
 	uint2 GroupId : SV_GroupID,
 	uint2 DispatchThreadId : SV_DispatchThreadID,
 	uint2 GroupThreadId : SV_GroupThreadID
-#if INSTANCED_STEREO && MOBILE_MULTI_VIEW
-	, uint InstanceId : SV_InstanceID
-	, out uint LayerIndex : SV_RenderTargetArrayIndex
-#elif MOBILE_MULTI_VIEW
-	, in uint ViewId : SV_ViewID
-	, out nointerpolation uint MultiViewIndex : VIEW_ID
-#endif
 )
 {
-#if INSTANCED_STEREO && MOBILE_MULTI_VIEW
-	const uint EyeIndex = GetEyeIndex(InstanceId);
-	ResolvedView = ResolveView(EyeIndex);
-	LayerIndex = EyeIndex;
-#elif MOBILE_MULTI_VIEW
-	ResolvedView = ResolveView(ViewId);
-	MultiViewIndex = ViewId;
-#else
-	ResolvedView = ResolveView();
-#endif
-
 	int2 FullGroupOrigin = int2(GroupId.x * THREADGROUP_SIZEX, GroupId.y * THREADGROUP_SIZEY);
 	int pixIdx = GroupIndex;
 
@@ -885,7 +849,7 @@ void GTAOHorizonSearchIntegralCS(
 	int2 ThreadPos = int2(GroupThreadId);
 	int2 TexturePos = int2(DispatchThreadId);
 
-	half2 TextureUV = TexturePosToBufferUV(TexturePos);
+	half2 TextureUV = TexturePosToBufferUV(TexturePos); 
 
 	half GTAO = CalculateGTAO(TextureUV.xy, TexturePos.xy, ThreadPos.xy);
 	half DeviceZ = GetDeviceZFromSharedMemory(ThreadPos + DEPTH_GROUP_THREAD_OFFSET);
diff --git a/Engine/Shaders/Private/PostProcessPixelProjectedReflectionMobile.usf b/Engine/Shaders/Private/PostProcessPixelProjectedReflectionMobile.usf
--- a/Engine/Shaders/Private/PostProcessPixelProjectedReflectionMobile.usf
+++ b/Engine/Shaders/Private/PostProcessPixelProjectedReflectionMobile.usf
@@ -39,12 +39,12 @@ float2 ProjectionPixelToBufferUV(int2 ReflectionPixel)
 {
 #if REFLECTION_PASS_PIXEL_SHADER
 	uint2 ViewportPos = (ReflectionPixel + float2(0.5f, 0.5f) - ViewRectMin.xy) * ViewSizeAndInvSize.zw * ResolvedView.ViewSizeAndInvSize.xy;
+	float2 BufferUV = ((float2)ViewportPos + float2(0.5f, 0.5f) + ResolvedView.ViewRectMin.xy) * ResolvedView.BufferSizeAndInvSize.zw;
 #else
-	uint2 ViewportPos = (ReflectionPixel + float2(0.5f, 0.5f)) * ViewSizeAndInvSize.zw * ResolvedView.ViewSizeAndInvSize.xy;
+	uint2 ViewportPos = (ReflectionPixel + float2(0.5f, 0.5f)) * ViewSizeAndInvSize.zw * ViewSizeAndInvSize.xy;
+	float2 BufferUV = ((float2) ViewportPos + float2(0.5f, 0.5f) + ViewRectMin.xy) * BufferSizeAndInvSize.zw;
 #endif
-
-	float2 BufferUV = ((float2)ViewportPos + float2(0.5f, 0.5f) + ResolvedView.ViewRectMin.xy) * ResolvedView.BufferSizeAndInvSize.zw;
-
+	
 	return BufferUV;
 }
 
@@ -126,26 +126,8 @@ void ProjectionPassCS(
 	uint2 GroupId : SV_GroupID,
 	uint2 DispatchThreadId : SV_DispatchThreadID,
 	uint2 GroupThreadId : SV_GroupThreadID
-#if INSTANCED_STEREO && MOBILE_MULTI_VIEW
-	, uint InstanceId : SV_InstanceID
-	, out uint LayerIndex : SV_RenderTargetArrayIndex
-#elif MOBILE_MULTI_VIEW
-	, in uint ViewId : SV_ViewID
-	, out nointerpolation uint MultiViewIndex : VIEW_ID
-#endif
 )
 {
-#if INSTANCED_STEREO && MOBILE_MULTI_VIEW
-	const uint EyeIndex = GetEyeIndex(InstanceId);
-	ResolvedView = ResolveView(EyeIndex);
-	LayerIndex = EyeIndex;
-#elif MOBILE_MULTI_VIEW
-	ResolvedView = ResolveView(ViewId);
-	MultiViewIndex = ViewId;
-#else
-	ResolvedView = ResolveView();
-#endif
-
 	int2 ReflectedPixel = DispatchThreadId.xy;
 
 	if (all(ReflectedPixel.xy < (int2)ViewSizeAndInvSize.xy))
@@ -162,7 +144,7 @@ void ProjectionPassCS(
 		{
 			float3 MirroredViewPosition = ViewPosition - ReflectionPlane.xyz * (2.f * PlaneDistance);
 
-			float4 MirroredProjectedPosition = mul(float4(MirroredViewPosition, 1.0f), ResolvedView.ViewToClip);
+			float4 MirroredProjectedPosition = mul(float4(MirroredViewPosition, 1.0f), View.ViewToClip);
 			half2 MirroredProjectedUV = MirroredProjectedPosition.xy / MirroredProjectedPosition.w;
 
 			if (all(abs(MirroredProjectedUV.xy) < 1.0f))
