Fixed volumetric fog for stereo rendering

diff --git a/Engine/Shaders/Private/HeightFogCommon.ush b/Engine/Shaders/Private/HeightFogCommon.ush
--- a/Engine/Shaders/Private/HeightFogCommon.ush
+++ b/Engine/Shaders/Private/HeightFogCommon.ush
@@ -33,7 +33,9 @@
 #define FOG_MATERIALBLENDING_OVERRIDE 0
 #endif
 
-#if INSTANCED_STEREO
+#if ((INSTANCED_STEREO && MOBILE_MULTI_VIEW) || MOBILE_MULTI_VIEW) && SHADING_PATH_MOBILE && PIXELSHADER
+	#define FogStructISR MobileBasePass.Fog
+#elif INSTANCED_STEREO
 	#if MATERIALBLENDING_ANY_TRANSLUCENT
 		#define FogStructISR TranslucentBasePass.Shared.FogISR
 	#else
@@ -57,31 +59,105 @@ static const float FLT_EPSILON2 = 0.01f;
 // FogStruct.ExponentialFogParameters3: FogDensity in x, FogHeight in y, whether to use cubemap fog color in z, FogCutoffDistance in w. 
 // FogStruct.FogInscatteringTextureParameters: mip distance scale in x, bias in y, num mips in z 
 
-float3 ComputeInscatteringColor(float3 CameraToReceiver, float CameraToReceiverLength)
+struct FogStructData
+{
+	float4 ExponentialFogParameters;
+	float4 ExponentialFogParameters2;
+	float4 ExponentialFogColorParameter;
+	float4 ExponentialFogParameters3;
+	float4 SkyAtmosphereAmbientContributionColorScale;
+	float4 InscatteringLightDirection;
+	float4 DirectionalInscatteringColor;
+	float2 SinCosInscatteringColorCubemapRotation;
+	float3 FogInscatteringTextureParameters;
+	float ApplyVolumetricFog;
+	float VolumetricFogStartDistance;
+	float VolumetricFogNearFadeInDistanceInv;
+	TextureCube FogInscatteringColorCubemap;
+	SamplerState FogInscatteringColorSampler;
+	Texture3D IntegratedLightScattering;
+	SamplerState IntegratedLightScatteringSampler;
+};
+
+FogStructData GetFogStructData(uint EyeIndex)
+{
+	FogStructData Out;
+#if INSTANCED_STEREO || MOBILE_MULTI_VIEW
+	BRANCH
+	if (EyeIndex == 0)
+	{
+#endif
+		Out.ExponentialFogParameters = FogStruct.ExponentialFogParameters;
+		Out.ExponentialFogParameters2 = FogStruct.ExponentialFogParameters2;
+		Out.ExponentialFogColorParameter = FogStruct.ExponentialFogColorParameter;
+		Out.ExponentialFogParameters3 = FogStruct.ExponentialFogParameters3;
+		Out.SkyAtmosphereAmbientContributionColorScale = FogStruct.SkyAtmosphereAmbientContributionColorScale;
+		Out.InscatteringLightDirection = FogStruct.InscatteringLightDirection;
+		Out.DirectionalInscatteringColor = FogStruct.DirectionalInscatteringColor;
+		Out.SinCosInscatteringColorCubemapRotation = FogStruct.SinCosInscatteringColorCubemapRotation;
+		Out.FogInscatteringTextureParameters = FogStruct.FogInscatteringTextureParameters;
+		Out.ApplyVolumetricFog = FogStruct.ApplyVolumetricFog;
+		Out.VolumetricFogStartDistance = FogStruct.VolumetricFogStartDistance;
+		Out.VolumetricFogNearFadeInDistanceInv = FogStruct.VolumetricFogNearFadeInDistanceInv;
+		Out.FogInscatteringColorCubemap = FogStruct.FogInscatteringColorCubemap;
+		Out.FogInscatteringColorSampler = FogStruct.FogInscatteringColorSampler;
+		Out.IntegratedLightScattering = FogStruct.IntegratedLightScattering;
+		Out.IntegratedLightScatteringSampler = FogStruct.IntegratedLightScatteringSampler;
+#if INSTANCED_STEREO || MOBILE_MULTI_VIEW
+	}
+	else
+	{	
+		Out.ExponentialFogParameters = FogStructISR.ExponentialFogParameters;
+		Out.ExponentialFogParameters2 = FogStructISR.ExponentialFogParameters2;
+		Out.ExponentialFogColorParameter = FogStructISR.ExponentialFogColorParameter;
+		Out.ExponentialFogParameters3 = FogStructISR.ExponentialFogParameters3;
+		Out.SkyAtmosphereAmbientContributionColorScale = FogStructISR.SkyAtmosphereAmbientContributionColorScale;
+		Out.InscatteringLightDirection = FogStructISR.InscatteringLightDirection;
+		Out.DirectionalInscatteringColor = FogStructISR.DirectionalInscatteringColor;
+		Out.SinCosInscatteringColorCubemapRotation = FogStructISR.SinCosInscatteringColorCubemapRotation;
+		Out.FogInscatteringTextureParameters = FogStructISR.FogInscatteringTextureParameters;
+		Out.ApplyVolumetricFog = FogStructISR.ApplyVolumetricFog;
+		Out.VolumetricFogStartDistance = FogStructISR.VolumetricFogStartDistance;
+		Out.VolumetricFogNearFadeInDistanceInv = FogStructISR.VolumetricFogNearFadeInDistanceInv;
+		Out.FogInscatteringColorCubemap = FogStructISR.FogInscatteringColorCubemap;
+		Out.FogInscatteringColorSampler = FogStructISR.FogInscatteringColorSampler;
+		Out.IntegratedLightScattering = FogStructISR.IntegratedLightScattering;
+		Out.IntegratedLightScatteringSampler = FogStructISR.IntegratedLightScatteringSampler;
+	}
+#endif
+	return Out;
+}
+
+float3 ComputeInscatteringColor(float3 CameraToReceiver, float CameraToReceiverLength, FogStructData FogData)
 {
-	half3 Inscattering = FogStruct.ExponentialFogColorParameter.xyz;
+	half3 Inscattering = FogData.ExponentialFogColorParameter.xyz;
 
 #if PERMUTATION_SUPPORT_FOG_INSCATTERING_TEXTURE
 	BRANCH
-	if (FogStruct.ExponentialFogParameters3.z > 0)
+	if (FogData.ExponentialFogParameters3.z > 0)
 	{
-		float FadeAlpha = saturate(CameraToReceiverLength * FogStruct.FogInscatteringTextureParameters.x + FogStruct.FogInscatteringTextureParameters.y);
+		float FadeAlpha = saturate(CameraToReceiverLength * FogData.FogInscatteringTextureParameters.x + FogData.FogInscatteringTextureParameters.y);
 		float3 CubemapLookupVector = CameraToReceiver;
 		// Rotate around Z axis
-		CubemapLookupVector.xy = float2(dot(CubemapLookupVector.xy, float2(FogStruct.SinCosInscatteringColorCubemapRotation.y, -FogStruct.SinCosInscatteringColorCubemapRotation.x)), dot(CubemapLookupVector.xy, FogStruct.SinCosInscatteringColorCubemapRotation.xy));
-		float3 DirectionalColor = TextureCubeSampleLevel(FogStruct.FogInscatteringColorCubemap, SharedFogInColorScatteringSampler, CubemapLookupVector, 0).xyz;
-		float3 NonDirectionalColor = TextureCubeSampleLevel(FogStruct.FogInscatteringColorCubemap, SharedFogInColorScatteringSampler, CubemapLookupVector, FogStruct.FogInscatteringTextureParameters.z).xyz;
+		CubemapLookupVector.xy = float2(dot(CubemapLookupVector.xy, float2(FogData.SinCosInscatteringColorCubemapRotation.y, -FogData.SinCosInscatteringColorCubemapRotation.x)), dot(CubemapLookupVector.xy, FogData.SinCosInscatteringColorCubemapRotation.xy));
+		float3 DirectionalColor = TextureCubeSampleLevel(FogData.FogInscatteringColorCubemap, SharedFogInColorScatteringSampler, CubemapLookupVector, 0).xyz;
+		float3 NonDirectionalColor = TextureCubeSampleLevel(FogData.FogInscatteringColorCubemap, SharedFogInColorScatteringSampler, CubemapLookupVector, FogData.FogInscatteringTextureParameters.z).xyz;
 		Inscattering *= lerp(NonDirectionalColor, DirectionalColor, FadeAlpha);
 	}
 #endif
 
 #if PROJECT_SUPPORT_SKY_ATMOSPHERE_AFFECTS_HEIGHFOG
-	Inscattering += FogStruct.SkyAtmosphereAmbientContributionColorScale.rgb * View.SkyAtmosphereHeightFogContribution.xxx * GetViewDistanceSkyLightColor();
+	Inscattering += FogData.SkyAtmosphereAmbientContributionColorScale.rgb * View.SkyAtmosphereHeightFogContribution.xxx * GetViewDistanceSkyLightColor();
 #endif
 
 	return Inscattering;
 }
 
+float3 ComputeInscatteringColor(float3 CameraToReceiver, float CameraToReceiverLength)
+{
+	return ComputeInscatteringColor(CameraToReceiver, CameraToReceiverLength, GetFogStructData(0));
+}
+
 // Calculate the line integral of the ray from the camera to the receiver position through the fog density function
 // The exponential fog density function is d = GlobalDensity * exp(-HeightFalloff * z)
 float CalculateLineIntegralShared(float FogHeightFalloff, float RayDirectionZ, float RayOriginTerms)
@@ -94,22 +170,24 @@ float CalculateLineIntegralShared(float FogHeightFalloff, float RayDirectionZ, f
 }
 
 // @param WorldPositionRelativeToCamera = WorldPosition - InCameraPosition
-half4 GetExponentialHeightFog(float3 WorldPositionRelativeToCamera, float ExcludeDistance)
+half4 GetExponentialHeightFog(float3 WorldPositionRelativeToCamera, float ExcludeDistance, uint EyeIndex)
 {
-	const half MinFogOpacity = FogStruct.ExponentialFogColorParameter.w;
-	const float MaxWorldObserverHeight = FogStruct.ExponentialFogParameters.z;
+	FogStructData FogData = GetFogStructData(EyeIndex);
+	
+	const half MinFogOpacity = FogData.ExponentialFogColorParameter.w;
+	const float MaxWorldObserverHeight = FogData.ExponentialFogParameters.z;
 
-	const float3 WorldObserverOrigin = float3(DFHackToFloat(PrimaryView.WorldCameraOrigin).xy, min(DFHackToFloat(PrimaryView.WorldCameraOrigin).z, MaxWorldObserverHeight)); // Clamp z to max height
+	const float3 WorldObserverOrigin = float3(DFHackToFloat(ResolvedView.WorldCameraOrigin).xy, min(DFHackToFloat(ResolvedView.WorldCameraOrigin).z, MaxWorldObserverHeight)); // Clamp z to max height
 
 	float3 CameraToReceiver = WorldPositionRelativeToCamera;
-	CameraToReceiver.z += DFHackToFloat(PrimaryView.WorldCameraOrigin).z - WorldObserverOrigin.z; // Compensate this vector for clamping the observer height 
+	CameraToReceiver.z += DFHackToFloat(ResolvedView.WorldCameraOrigin).z - WorldObserverOrigin.z; // Compensate this vector for clamping the observer height
 	float CameraToReceiverLengthSqr = dot(CameraToReceiver, CameraToReceiver);
 	float CameraToReceiverLengthInv = rsqrt(max(CameraToReceiverLengthSqr, 0.00000001f));
 	float CameraToReceiverLength = CameraToReceiverLengthSqr * CameraToReceiverLengthInv;
 	half3 CameraToReceiverNormalized = CameraToReceiver * CameraToReceiverLengthInv;
 
-	float RayOriginTerms = FogStruct.ExponentialFogParameters.x;
-	float RayOriginTermsSecond = FogStruct.ExponentialFogParameters2.x;
+	float RayOriginTerms = FogData.ExponentialFogParameters.x;
+	float RayOriginTermsSecond = FogData.ExponentialFogParameters2.x;
 	float RayLength = CameraToReceiverLength;
 	float RayDirectionZ = CameraToReceiver.z;
 
@@ -117,10 +195,10 @@ half4 GetExponentialHeightFog(float3 WorldPositionRelativeToCamera, float Exclud
 
 	BRANCH
 	// While rendering a planar reflection with a clip plane, we must compute analytical fog using a camera path starting from the plane, rather than the virtual camera origin
-	if (dot(View.GlobalClippingPlane.xyz, 1) > 0.0f)
+	if (dot(ResolvedView.GlobalClippingPlane.xyz, 1) > 0.0f)
 	{
-		float CameraOriginPlaneDistance = dot(View.GlobalClippingPlane, float4(WorldObserverOrigin + DFHackToFloat(PrimaryView.PreViewTranslation), 1));
-		float PlaneIntersectionTime = -CameraOriginPlaneDistance / dot(CameraToReceiver, View.GlobalClippingPlane.xyz);
+		float CameraOriginPlaneDistance = dot(ResolvedView.GlobalClippingPlane, float4(WorldObserverOrigin + DFHackToFloat(ResolvedView.PreViewTranslation), 1));
+		float PlaneIntersectionTime = -CameraOriginPlaneDistance / dot(CameraToReceiver, ResolvedView.GlobalClippingPlane.xyz);
 
 		// Only modify the start distance if the reflection plane is between the camera and receiver point
 		if (PlaneIntersectionTime > 0 && PlaneIntersectionTime < 1)
@@ -133,7 +211,7 @@ half4 GetExponentialHeightFog(float3 WorldPositionRelativeToCamera, float Exclud
 
 	// Factor in StartDistance
 #if PERMUTATION_SUPPORT_FOG_START_DISTANCE
-	ExcludeDistance = max(ExcludeDistance, FogStruct.ExponentialFogParameters.w);
+	ExcludeDistance = max(ExcludeDistance, FogData.ExponentialFogParameters.w);
 	if (ExcludeDistance > 0)
 	{
 		float ExcludeIntersectionTime = ExcludeDistance * CameraToReceiverLengthInv;
@@ -145,31 +223,31 @@ half4 GetExponentialHeightFog(float3 WorldPositionRelativeToCamera, float Exclud
 		RayLength = (1.0f - ExcludeIntersectionTime) * CameraToReceiverLength;
 		RayDirectionZ = ExclusionIntersectionToReceiverZ;
 
-		float Exponent = max(-127.0f, FogStruct.ExponentialFogParameters.y * (ExclusionIntersectionZ - FogStruct.ExponentialFogParameters3.y));
-		RayOriginTerms = FogStruct.ExponentialFogParameters3.x * exp2(-Exponent);
+		float Exponent = max(-127.0f, FogData.ExponentialFogParameters.y * (ExclusionIntersectionZ - FogData.ExponentialFogParameters3.y));
+		RayOriginTerms = FogData.ExponentialFogParameters3.x * exp2(-Exponent);
 		
-		float ExponentSecond = max(-127.0f, FogStruct.ExponentialFogParameters2.y * (ExclusionIntersectionZ - FogStruct.ExponentialFogParameters2.w)); 	 
-		RayOriginTermsSecond = FogStruct.ExponentialFogParameters2.z * exp2(-ExponentSecond);
+		float ExponentSecond = max(-127.0f, FogData.ExponentialFogParameters2.y * (ExclusionIntersectionZ - FogData.ExponentialFogParameters2.w));
+		RayOriginTermsSecond = FogData.ExponentialFogParameters2.z * exp2(-ExponentSecond);
 	}
 #endif
 
 	// Calculate the "shared" line integral (this term is also used for the directional light inscattering) by adding the two line integrals together (from two different height falloffs and densities)
-	float ExponentialHeightLineIntegralShared = CalculateLineIntegralShared(FogStruct.ExponentialFogParameters.y, RayDirectionZ, RayOriginTerms);
+	float ExponentialHeightLineIntegralShared = CalculateLineIntegralShared(FogData.ExponentialFogParameters.y, RayDirectionZ, RayOriginTerms);
 #if PERMUTATION_SUPPORT_FOG_SECOND_TERM
-	ExponentialHeightLineIntegralShared+= CalculateLineIntegralShared(FogStruct.ExponentialFogParameters2.y, RayDirectionZ, RayOriginTermsSecond);
+	ExponentialHeightLineIntegralShared += CalculateLineIntegralShared(FogData.ExponentialFogParameters2.y, RayDirectionZ, RayOriginTermsSecond);
 #endif
 
 	float ExponentialHeightLineIntegral = ExponentialHeightLineIntegralShared * RayLength;
 
-	half3 InscatteringColor = ComputeInscatteringColor(CameraToReceiver, CameraToReceiverLength);
+	half3 InscatteringColor = ComputeInscatteringColor(CameraToReceiver, CameraToReceiverLength, FogData);
 	half3 DirectionalInscattering = 0;
 
 #if PERMUTATION_SUPPORT_FOG_DIRECTIONAL_LIGHT_INSCATTERING
 	// if InscatteringLightDirection.w is negative then it's disabled, otherwise it holds directional inscattering start distance
 	BRANCH
-	if (FogStruct.InscatteringLightDirection.w >= 0
+	if (FogData.InscatteringLightDirection.w >= 0
 	#if PERMUTATION_SUPPORT_FOG_INSCATTERING_TEXTURE
-		&& FogStruct.ExponentialFogParameters3.z == 0
+		&& FogData.ExponentialFogParameters3.z == 0
 	#endif
 	)
 	{
@@ -179,21 +257,21 @@ half4 GetExponentialHeightFog(float3 WorldPositionRelativeToCamera, float Exclud
 		half3 DirectionalLightInscattering;
 
 		// No need to test View.AtmosphereLightIlluminanceOnGroundPostTransmittance[0].a because InscatteringLightDirection.w above is doing the same test already.
-		DirectionalInscatteringColor = FogStruct.DirectionalInscatteringColor.xyz + View.SkyAtmosphereHeightFogContribution * View.AtmosphereLightIlluminanceOnGroundPostTransmittance[0].rgb * UniformPhaseFunction;
-		DirectionalLightInscattering = DirectionalInscatteringColor * pow(saturate(dot(CameraToReceiverNormalized, View.AtmosphereLightDirection[0].xyz)), FogStruct.DirectionalInscatteringColor.w);
+		DirectionalInscatteringColor = FogData.DirectionalInscatteringColor.xyz + ResolvedView.SkyAtmosphereHeightFogContribution * ResolvedView.AtmosphereLightIlluminanceOnGroundPostTransmittance[0].rgb * UniformPhaseFunction;
+		DirectionalLightInscattering = DirectionalInscatteringColor * pow(saturate(dot(CameraToReceiverNormalized, ResolvedView.AtmosphereLightDirection[0].xyz)), FogData.DirectionalInscatteringColor.w);
 
-		if (View.AtmosphereLightIlluminanceOnGroundPostTransmittance[1].a > 0.0f) // Skip DirectionalInscatteringColor on the second light when disabled.
+		if (ResolvedView.AtmosphereLightIlluminanceOnGroundPostTransmittance[1].a > 0.0f) // Skip DirectionalInscatteringColor on the second light when disabled.
 		{
-			DirectionalInscatteringColor = FogStruct.DirectionalInscatteringColor.xyz + View.SkyAtmosphereHeightFogContribution * View.AtmosphereLightIlluminanceOnGroundPostTransmittance[1].rgb * UniformPhaseFunction;
-			DirectionalLightInscattering += DirectionalInscatteringColor * pow(saturate(dot(CameraToReceiverNormalized, View.AtmosphereLightDirection[1].xyz)), FogStruct.DirectionalInscatteringColor.w);
+			DirectionalInscatteringColor = FogData.DirectionalInscatteringColor.xyz + ResolvedView.SkyAtmosphereHeightFogContribution * ResolvedView.AtmosphereLightIlluminanceOnGroundPostTransmittance[1].rgb * UniformPhaseFunction;
+			DirectionalLightInscattering += DirectionalInscatteringColor * pow(saturate(dot(CameraToReceiverNormalized, ResolvedView.AtmosphereLightDirection[1].xyz)), FogData.DirectionalInscatteringColor.w);
 		}
 #else
 		// Setup a cosine lobe around the light direction to approximate inscattering from the directional light off of the ambient haze;
-		half3 DirectionalLightInscattering = FogStruct.DirectionalInscatteringColor.xyz * pow(saturate(dot(CameraToReceiverNormalized, FogStruct.InscatteringLightDirection.xyz)), FogStruct.DirectionalInscatteringColor.w);
+		half3 DirectionalLightInscattering = FogData.DirectionalInscatteringColor.xyz * pow(saturate(dot(CameraToReceiverNormalized, FogData.InscatteringLightDirection.xyz)), FogData.DirectionalInscatteringColor.w);
 #endif
 
 		// Calculate the line integral of the eye ray through the haze, using a special starting distance to limit the inscattering to the distance
-		float DirectionalInscatteringStartDistance = FogStruct.InscatteringLightDirection.w;
+		float DirectionalInscatteringStartDistance = FogData.InscatteringLightDirection.w;
 		float DirExponentialHeightLineIntegral = ExponentialHeightLineIntegralShared * max(RayLength - DirectionalInscatteringStartDistance, 0.0f);
 		// Calculate the amount of light that made it through the fog using the transmission equation
 		half DirectionalInscatteringFogFactor = saturate(exp2(-DirExponentialHeightLineIntegral));
@@ -206,7 +284,7 @@ half4 GetExponentialHeightFog(float3 WorldPositionRelativeToCamera, float Exclud
 	half ExpFogFactor = max(saturate(exp2(-ExponentialHeightLineIntegral)), MinFogOpacity);
 
 	FLATTEN
-	if (FogStruct.ExponentialFogParameters3.w > 0 && CameraToReceiverLength > FogStruct.ExponentialFogParameters3.w)
+	if (FogData.ExponentialFogParameters3.w > 0 && CameraToReceiverLength > FogData.ExponentialFogParameters3.w)
 	{
 		ExpFogFactor = 1;
 		DirectionalInscattering = 0;
@@ -229,9 +307,14 @@ half4 GetExponentialHeightFog(float3 WorldPositionRelativeToCamera, float Exclud
 	return half4(FogColor, ExpFogFactor);
 }
 
-half4 CalculateHeightFog(float3 WorldPositionRelativeToCamera)
+half4 GetExponentialHeightFog(float3 WorldPositionRelativeToCamera, float ExcludeDistance)
+{
+	return GetExponentialHeightFog(WorldPositionRelativeToCamera, ExcludeDistance, 0);
+}
+
+half4 CalculateHeightFog(float3 WorldPositionRelativeToCamera, uint EyeIndex)
 {
-	if (View.RenderingReflectionCaptureMask == 0.0f && !IsExponentialFogRenderedInMain(View.EnvironmentComponentsFlags))
+	if (ResolvedView.RenderingReflectionCaptureMask == 0.0f && !IsExponentialFogRenderedInMain(ResolvedView.EnvironmentComponentsFlags))
 	{
 		return float4(0.0f, 0.0f, 0.0f, 1.0f);
 	}
@@ -239,20 +322,26 @@ half4 CalculateHeightFog(float3 WorldPositionRelativeToCamera)
 	float ExcludeDistance = 0;
 #if PERMUTATION_SUPPORT_VOLUMETRIC_FOG
 	// Volumetric fog covers up to MaxDistance along ViewZ, exclude analytical fog from this range
-	float CosAngle = dot(normalize(WorldPositionRelativeToCamera), View.ViewForward);
+	float CosAngle = dot(normalize(WorldPositionRelativeToCamera), ResolvedView.ViewForward);
 	float InvCosAngle = (CosAngle > FLT_EPSILON) ? rcp(CosAngle) : 0;
-	ExcludeDistance = View.VolumetricFogMaxDistance * InvCosAngle;
+	ExcludeDistance = ResolvedView.VolumetricFogMaxDistance * InvCosAngle;
 #endif
 
-	half4 FogInscatteringAndOpacity = GetExponentialHeightFog(WorldPositionRelativeToCamera, ExcludeDistance);
+	half4 FogInscatteringAndOpacity = GetExponentialHeightFog(WorldPositionRelativeToCamera, ExcludeDistance, EyeIndex);
 	return FogInscatteringAndOpacity;
 }
 
+half4 CalculateHeightFog(float3 WorldPositionRelativeToCamera)
+{
+	return CalculateHeightFog(WorldPositionRelativeToCamera, 0);
+
+}
+
 float4 CombineVolumetricFog(float4 GlobalFog, float3 VolumeUV, uint EyeIndex, float SceneDepth)
 {
 	float4 VolumetricFogLookup = float4(0, 0, 0, 1);
 	
-	if (View.RenderingReflectionCaptureMask == 0.0f && !IsExponentialFogRenderedInMain(View.EnvironmentComponentsFlags))
+	if (ResolvedView.RenderingReflectionCaptureMask == 0.0f && !IsExponentialFogRenderedInMain(ResolvedView.EnvironmentComponentsFlags))
 	{
 		return VolumetricFogLookup;
 	}
@@ -261,7 +350,7 @@ float4 CombineVolumetricFog(float4 GlobalFog, float3 VolumeUV, uint EyeIndex, fl
 	float VolFogStartDistance = 0.0f;
 	if (FogStruct.ApplyVolumetricFog > 0)
 	{
-	#if INSTANCED_STEREO
+	#if INSTANCED_STEREO || MOBILE_MULTI_VIEW
 		if (EyeIndex == 0)
 		{
 			VolFogStartDistance = FogStruct.VolumetricFogStartDistance;
@@ -278,7 +367,7 @@ float4 CombineVolumetricFog(float4 GlobalFog, float3 VolumeUV, uint EyeIndex, fl
 	#endif
 
 		// IntegratedLightScattering is pre-exposed, remove pre exposure now so that it can correctly be applied later
-		VolumetricFogLookup.rgb *= View.OneOverPreExposure;
+		VolumetricFogLookup.rgb *= ResolvedView.OneOverPreExposure;
 	}
 
 	// Do not apply the Froxel volumetric texture in front of the fog start distance. (the soft fading occur in FinalIntegrationCS).
@@ -290,7 +379,7 @@ float4 CombineVolumetricFog(float4 GlobalFog, float3 VolumeUV, uint EyeIndex, fl
 	//VolumetricFogLookup.rgb += .1f * frac(min(ZSlice, 1.0f) / View.VolumetricFogInvGridSize.z);
 
 #if POST_PROCESS_ALPHA
-	if (IsExponentialFogHoldout(View.EnvironmentComponentsFlags) && View.RenderingReflectionCaptureMask == 0.0f)
+	if (IsExponentialFogHoldout(ResolvedView.EnvironmentComponentsFlags) && ResolvedView.RenderingReflectionCaptureMask == 0.0f)
 	{
 		VolumetricFogLookup.rgb *= 0;
 		GlobalFog.rgb *= 0;
@@ -302,14 +391,14 @@ float4 CombineVolumetricFog(float4 GlobalFog, float3 VolumeUV, uint EyeIndex, fl
 
 float ComputeNormalizedZSliceFromDepth(float SceneDepth)
 {
-	return log2(SceneDepth * View.VolumetricFogGridZParams.x + View.VolumetricFogGridZParams.y) * View.VolumetricFogGridZParams.z * View.VolumetricFogInvGridSize.z;
+	return log2(SceneDepth * ResolvedView.VolumetricFogGridZParams.x + ResolvedView.VolumetricFogGridZParams.y) * ResolvedView.VolumetricFogGridZParams.z * ResolvedView.VolumetricFogInvGridSize.z;
 }
 
 float3 ComputeVolumeUVFromNDC(float4 NDCPosition)
 {
 	NDCPosition.xy /= NDCPosition.w;
 	float3 VolumeUV = float3(NDCPosition.xy * float2(.5f, -.5f) + .5f, ComputeNormalizedZSliceFromDepth(NDCPosition.w));
-	return min(VolumeUV * float3(View.VolumetricFogScreenToResourceUV.xy, 1.0), float3(View.VolumetricFogUVMax, 1.0));
+	return min(VolumeUV * float3(ResolvedView.VolumetricFogScreenToResourceUV.xy, 1.0), float3(ResolvedView.VolumetricFogUVMax, 1.0));
 }
 float3 ComputeVolumeUVFromTranslatedPos(float3 TranslatedWorldPosition, float4x4 TranslatedWorldToClip)
 {
diff --git a/Engine/Shaders/Private/MobileFog.usf b/Engine/Shaders/Private/MobileFog.usf
--- a/Engine/Shaders/Private/MobileFog.usf
+++ b/Engine/Shaders/Private/MobileFog.usf
@@ -1,6 +1,9 @@
 // Copyright Epic Games, Inc. All Rights Reserved.
+#if SHADING_PATH_MOBILE
 #include "Common.ush"
 
+#define MOBILE_MULTI_VIEW_FALLBACK (INSTANCED_STEREO && MOBILE_MULTI_VIEW)
+
 // Reroute MobileSceneTextures uniform buffer references to the base pass uniform buffer 
 #define MobileSceneTextures MobileBasePass.SceneTextures
 #define FogStruct MobileBasePass.Fog
@@ -23,13 +26,28 @@ SamplerState HalfResLocalFogVolumeDepthSRVSampler;
 float StartDepthZ;
 
 void MobileFogVS(
+#if MOBILE_MULTI_VIEW_FALLBACK
+	in uint InstanceId : SV_InstanceID,
+	out nointerpolation uint EyeIndex : VIEW_ID,
+	out uint LayerIndex : SV_RenderTargetArrayIndex,
+#elif MOBILE_MULTI_VIEW
+	in nointerpolation uint ViewId : SV_ViewID,
+#endif
 	in float2 InPosition : ATTRIBUTE0,
 	out float2 OutTexCoord : TEXCOORD0,
 	out float3 OutScreenVector : TEXCOORD1,
 	out float4 OutPosition : SV_POSITION
 	)
 {
+#if MOBILE_MULTI_VIEW_FALLBACK
+	EyeIndex = GetEyeIndex(InstanceId);
+	LayerIndex = EyeIndex;
+	ResolvedView = ResolveView(EyeIndex);
+#elif MOBILE_MULTI_VIEW
+	ResolvedView = ResolveView(ViewId);
+#else
 	ResolvedView = ResolveView();
+#endif
 
 	// screenspace position from vb
 	OutPosition = float4(InPosition,StartDepthZ,1);
@@ -40,22 +58,38 @@ void MobileFogVS(
 }
 
 void MobileFogPS(
+#if MOBILE_MULTI_VIEW_FALLBACK
+	in nointerpolation uint EyeIndex : VIEW_ID,
+#elif MOBILE_MULTI_VIEW
+	in nointerpolation uint ViewId : SV_ViewID,
+#endif
 	float2 TexCoord : TEXCOORD0,
 	float3 ScreenVector : TEXCOORD1,
 	float4 SVPos : SV_POSITION,
 	out float4 OutColor : SV_Target0
 	)
 {
+	uint ArrayIndex = 0;
+#if MOBILE_MULTI_VIEW_FALLBACK
+	ResolvedView = ResolveView(EyeIndex);
+	ArrayIndex = EyeIndex;
+#elif MOBILE_MULTI_VIEW
+	ResolvedView = ResolveView(ViewId);
+	ArrayIndex = ViewId;
+#else
 	ResolvedView = ResolveView();
+#endif
+
 	half4 FinalFog = half4(0,0,0,1);
-	float DeviceZ = LookupDeviceZ(TexCoord);
+	
+	float DeviceZ = LookupDeviceZ(TexCoord, ArrayIndex);
 	float SceneDepth = ConvertFromDeviceZ(DeviceZ);
 	float3 WorldPositionRelativeToCamera = ScreenVector.xyz * SceneDepth;
 	float3 TranslatedWorldPosition = WorldPositionRelativeToCamera + ResolvedView.TranslatedWorldCameraOrigin;
 
 // at least one of these is active
 #if PERMUTATION_SUPPORT_HEIGHT_FOG
-	FinalFog = CalculateHeightFog(WorldPositionRelativeToCamera);
+	FinalFog = CalculateHeightFog(WorldPositionRelativeToCamera, ArrayIndex);
 #endif
 
 #if PERMUTATION_SUPPORT_LOCAL_FOG_VOLUME == 1
@@ -86,7 +120,7 @@ void MobileFogPS(
 	float VolumeFrontDepth3 = HalfResLocalFogVolumeDepthSRV.Load(uint3(VolumeCoord3, 0)).r * METER_TO_CENTIMETER;
 
 	float MaxFrontDepth4Diff = max(max(VolumeFrontDepth0, VolumeFrontDepth1), max(VolumeFrontDepth2, VolumeFrontDepth3));
-	float UpsampleSceneDepth = min(SceneDepth, MaxFrontDepth4Diff);	// needed because FP16 is limited in distance adn that fixes upsampling artefacts when depth is super large..
+	float UpsampleSceneDepth = min(SceneDepth, MaxFrontDepth4Diff);	// needed because FP16 is limited in distance and that fixes upsampling artefacts when depth is super large..
 
 	float4 Depth4 = float4(VolumeFrontDepth0, VolumeFrontDepth1, VolumeFrontDepth2, VolumeFrontDepth3);
 	float4 Depth4Diff = UpsampleSceneDepth - Depth4;
@@ -130,16 +164,15 @@ void MobileFogPS(
 #if PERMUTATION_SUPPORT_VOLUMETRIC_FOG
 	if (FogStruct.ApplyVolumetricFog > 0) 
 	{
-		const uint EyeIndex = 0;
 		float3 VolumeUV = ComputeVolumeUVFromTranslatedPos(TranslatedWorldPosition, ResolvedView.TranslatedWorldToClip);
-		FinalFog = CombineVolumetricFog(FinalFog, VolumeUV, EyeIndex, SceneDepth);
+		FinalFog = CombineVolumetricFog(FinalFog, VolumeUV, ArrayIndex, SceneDepth);
 	}
 #endif
 
 #if PERMUTATION_SUPPORT_AERIAL_PERSPECTIVE
 	if (ResolvedView.SkyAtmosphereApplyCameraAerialPerspectiveVolume > 0.0f && (DeviceZ != FarDepthValue))
 	{
-		float4 NDCPosition = mul(float4(TranslatedWorldPosition, 1.0f), PrimaryView.TranslatedWorldToClip);
+		float4 NDCPosition = mul(float4(TranslatedWorldPosition, 1.0f), ResolvedView.TranslatedWorldToClip);
 
 		// Sample the aerial perspective (AP). It is also blended under the VertexFog parameter.
 		FinalFog = GetAerialPerspectiveLuminanceTransmittanceWithFogOver(
@@ -162,3 +195,4 @@ void MobileFogPS(
 	FinalFog.rgb *= ResolvedView.PreExposure;
 	OutColor = RETURN_COLOR(half4(FinalFog.rgb, FinalFog.w)); 
 }
+#endif //SHADING_PATH_MOBILE
\ No newline at end of file
diff --git a/Engine/Source/Runtime/Renderer/Private/MobileBasePassRendering.cpp b/Engine/Source/Runtime/Renderer/Private/MobileBasePassRendering.cpp
--- a/Engine/Source/Runtime/Renderer/Private/MobileBasePassRendering.cpp
+++ b/Engine/Source/Runtime/Renderer/Private/MobileBasePassRendering.cpp
@@ -253,7 +253,17 @@ void SetupMobileBasePassUniformParameters(
 	const FMobileBasePassTextures& MobileBasePassTextures,
 	FMobileBasePassUniformParameters& BasePassParameters)
 {
+	const FViewInfo* InstancedView = View.GetInstancedView();
+
 	SetupFogUniformParameters(GraphBuilder, View, BasePassParameters.Fog);
+	if (InstancedView && (View.bIsMobileMultiViewEnabled || UE::StereoRenderUtils::FStereoShaderAspects(View.GetShaderPlatform()).IsMobileMultiViewEnabled()))
+	{
+		SetupFogUniformParameters(GraphBuilder, *InstancedView, BasePassParameters.FogMMV);
+	}
+	else
+	{
+		BasePassParameters.FogMMV = BasePassParameters.Fog;
+	}
 
 	if (View.ForwardLightingResources.ForwardLightData)
 	{
@@ -265,7 +275,6 @@ void SetupMobileBasePassUniformParameters(
 	}
 
 	// Setup forward light data for mobile multi-view secondary view if enabled and available
-	const FViewInfo* InstancedView = View.GetInstancedView();
 	const FForwardLightData* InstancedForwardLightData = InstancedView ? InstancedView->ForwardLightingResources.ForwardLightData : nullptr;
 	if (View.bIsMobileMultiViewEnabled && InstancedForwardLightData)
 	{
diff --git a/Engine/Source/Runtime/Renderer/Private/MobileBasePassRendering.h b/Engine/Source/Runtime/Renderer/Private/MobileBasePassRendering.h
--- a/Engine/Source/Runtime/Renderer/Private/MobileBasePassRendering.h
+++ b/Engine/Source/Runtime/Renderer/Private/MobileBasePassRendering.h
@@ -42,6 +42,7 @@ struct FMobileBasePassTextures
 BEGIN_GLOBAL_SHADER_PARAMETER_STRUCT(FMobileBasePassUniformParameters, )
 	SHADER_PARAMETER(float, AmbientOcclusionStaticFraction)
 	SHADER_PARAMETER_STRUCT(FFogUniformParameters, Fog)
+	SHADER_PARAMETER_STRUCT(FFogUniformParameters, FogMMV)
 	SHADER_PARAMETER_STRUCT(FLocalFogVolumeUniformParameters, LFV)
 	SHADER_PARAMETER_STRUCT(FForwardLightData, Forward)
 	SHADER_PARAMETER_STRUCT(FForwardLightData, ForwardMMV)
diff --git a/Engine/Source/Runtime/Renderer/Private/MobileFogRendering.cpp b/Engine/Source/Runtime/Renderer/Private/MobileFogRendering.cpp
--- a/Engine/Source/Runtime/Renderer/Private/MobileFogRendering.cpp
+++ b/Engine/Source/Runtime/Renderer/Private/MobileFogRendering.cpp
@@ -271,7 +271,7 @@ void FMobileSceneRenderer::RenderFog(FRHICommandList& RHICmdList, const FViewInf
 
 	// Draw a quad covering the view.
 	RHICmdList.SetStreamSource(0, GScreenSpaceVertexBuffer.VertexBufferRHI, 0);
-	RHICmdList.DrawIndexedPrimitive(GTwoTrianglesIndexBuffer.IndexBufferRHI, 0, 0, 4, 0, 2, 1);
+	RHICmdList.DrawIndexedPrimitive(GTwoTrianglesIndexBuffer.IndexBufferRHI, 0, 0, 4, 0, 2, View.GetStereoPassInstanceFactor());
 
 	RenderLocalFogVolumeMobileLambda();
 }
