Path Tracer: Fix dark bands in fog from early ray marching termination threshold

For cases that know how to compute analytic transmittance, its helpful to allow it to still be recomputed after the raymarch loop, even if we ended it early. This way the more precise calculation can be used for surfaces.

diff --git a/Engine/Shaders/Private/PathTracing/PathTracingCore.ush b/Engine/Shaders/Private/PathTracing/PathTracingCore.ush
--- a/Engine/Shaders/Private/PathTracing/PathTracingCore.ush
+++ b/Engine/Shaders/Private/PathTracing/PathTracingCore.ush
@@ -564,7 +564,7 @@ FPackedPathTracingPayload TraceTransparentRay(inout FPathState PathState, int Bo
 				float3 PathThroughput = PathState.PathThroughput;
 
 				// In order to use analytic transmittance, the user must have requested it _and_ we must be on an interval with no clouds or heterogeneous volumes
-				bool bUseAnalyticTransmittance =
+				const bool bUseAnalyticTransmittance =
 					(VolumeFlags & PATH_TRACER_VOLUME_USE_ANALYTIC_TRANSMITTANCE) != 0 &&
 					(Interval.VolumeMask & (PATH_TRACER_VOLUME_ENABLE_CLOUDS | PATH_TRACER_VOLUME_ENABLE_HETEROGENEOUS_VOLUMES)) == 0;
 				// are there any holdouts on the current segment?
@@ -638,7 +638,6 @@ FPackedPathTracingPayload TraceTransparentRay(inout FPathState PathState, int Bo
 							PathThroughput *= saturate(10000.0 * max3(PathThroughput.x, PathThroughput.y, PathThroughput.z) - 1.0);
 							if (!any(PathThroughput > 0))
 							{
-								bUseAnalyticTransmittance = false; // avoid re-computing the transmittance after the loop
 								break;
 							}
 						}
