fix mobile propagateAlpha  (invert alpha) for post processing.
#android

diff --git a/Engine/Shaders/Private/PostProcessTonemap.usf b/Engine/Shaders/Private/PostProcessTonemap.usf
--- a/Engine/Shaders/Private/PostProcessTonemap.usf
+++ b/Engine/Shaders/Private/PostProcessTonemap.usf
@@ -603,6 +603,11 @@ void MainPS(
 		float3 UVI = float3(UV.xy,0);
 	#endif
     OutColor = TonemapCommonPS(UVI, InVignette, GrainUV, ScreenPos, FullViewUV, SvPosition, Luminance);
+
+#if POST_PROCESS_ALPHA == 2 && SHADING_PATH_MOBILE
+	OutColor.a = 1.0f - OutColor.a;	// invert alpha when propagate alpha is used with ASIS
+#endif
+
 #if OUTPUT_LUMINANCE
 	OutLuminance = Luminance;
 #endif
