[Nanite-Skinning] Implemented BlendBoneTransformsWithScale and BlendBoneTransforms

diff --git a/Engine/Shaders/Private/BoneTransform.ush b/Engine/Shaders/Private/BoneTransform.ush
--- a/Engine/Shaders/Private/BoneTransform.ush
+++ b/Engine/Shaders/Private/BoneTransform.ush
@@ -36,10 +36,10 @@ FBoneTransform MakeBoneTransformIdentity()
 	);
 }
 
-FBoneTransform MultiplyBoneTransform(FBoneTransform LHS, FBoneTransform RHS)
+FBoneTransform MultiplyBoneTransform(FBoneTransform Atom1, FBoneTransform Atom2)
 {
-	FQuat Rotation	= QuatMultiply(RHS.Rotation, LHS.Rotation);
-	float3 Location	= QuatRotateVector(RHS.Rotation, LHS.Location) + RHS.Location;
+	FQuat Rotation	= QuatMultiply(Atom2.Rotation, Atom1.Rotation);
+	float3 Location	= QuatRotateVector(Atom2.Rotation, Atom1.Location) + Atom2.Location;
 	return MakeBoneTransform(Rotation, Location);
 }
 
@@ -70,6 +70,14 @@ float4x4 MatrixFromBoneTransform(FBoneTransform Transform)
 	return OutMatrix;
 }
 
+FBoneTransform BlendBoneTransforms(FBoneTransform Atom1, FBoneTransform Atom2, float Alpha)
+{
+	FBoneTransform Blended;
+	Blended.Location = lerp(Atom1.Location, Atom2.Location, Alpha);
+	Blended.Rotation = QuatSlerpApproximate(Atom1.Rotation, Atom2.Rotation, Alpha);
+	return Blended;
+}
+
 ////
 
 FBoneTransformWithScale MakeBoneTransformWithScale(FQuat Rotation, float3 Location, float3 Scale)
@@ -86,12 +94,12 @@ FBoneTransformWithScale MakeBoneTransformWithScaleIdentity()
 	return MakeBoneTransformWithScale(QuatIdentity(), float3(0.0f, 0.0f, 0.0f), float3(1.0f, 1.0f, 1.0f));
 }
 
-FBoneTransformWithScale MultiplyBoneTransformWithScale(FBoneTransformWithScale LHS, FBoneTransformWithScale RHS)
+FBoneTransformWithScale MultiplyBoneTransformWithScale(FBoneTransformWithScale Atom1, FBoneTransformWithScale Atom2)
 {
 	// NOTE: Assumes no negative scale! Otherwise should multiply using matrices
-	FQuat Rotation	= QuatMultiply(RHS.Rotation, LHS.Rotation);
-	float3 Location	= QuatRotateVector(RHS.Rotation, RHS.Scale * LHS.Location) + RHS.Location;
-	float3 Scale = LHS.Scale * RHS.Scale;
+	FQuat Rotation	= QuatMultiply(Atom2.Rotation, Atom1.Rotation);
+	float3 Location	= QuatRotateVector(Atom2.Rotation, Atom2.Scale * Atom1.Location) + Atom2.Location;
+	float3 Scale = Atom1.Scale * Atom2.Scale;
 	return MakeBoneTransformWithScale(Rotation, Location, Scale);
 }
 
@@ -109,4 +117,13 @@ float4x4 MatrixFromBoneTransformWithScale(FBoneTransformWithScale Transform)
 	OutMatrix[2] = float4(RotationMatrix[2], 0.0);
 	OutMatrix[3] = float4(Transform.Location, 1.0);
 	return OutMatrix;
+}
+
+FBoneTransformWithScale BlendBoneTransformsWithScale(FBoneTransformWithScale Atom1, FBoneTransformWithScale Atom2, float Alpha)
+{
+	FBoneTransformWithScale Blended;
+	Blended.Location = lerp(Atom1.Location, Atom2.Location, Alpha);
+	Blended.Rotation = QuatSlerpApproximate(Atom1.Rotation, Atom2.Rotation, Alpha);
+	Blended.Scale    = Lerp(Atom1.Scale, Atom2.Scale, Alpha);
+	return Blended;
 }
\ No newline at end of file
