Fix edge integral math causing NaN value with rect light under coplanar setup.

This fix ensures that rsqrt always get strictly positive values.

diff --git a/Engine/Shaders/Private/RectLight.ush b/Engine/Shaders/Private/RectLight.ush
--- a/Engine/Shaders/Private/RectLight.ush
+++ b/Engine/Shaders/Private/RectLight.ush
@@ -108,10 +108,10 @@ float3 RectIrradianceLambert( float3 N, FRect Rect, out float BaseIrradiance, ou
 	// normalize( cross( L0, L1 ) ) = cross( L0, L1 ) * rsqrt( 1 - dot( L0, L1 )^2 )
 	// acos( x ) ~= sqrt(1 - x) * (1.5708 - 0.175 * x)
 
-	float w01 = ( 1.5708 - 0.175 * c01 ) * rsqrt( c01 + 1 );	// 1 mad, 1 add, 1 rsqrt
-	float w12 = ( 1.5708 - 0.175 * c12 ) * rsqrt( c12 + 1 );	// 1 mad, 1 add, 1 rsqrt
-	float w23 = ( 1.5708 - 0.175 * c23 ) * rsqrt( c23 + 1 );	// 1 mad, 1 add, 1 rsqrt
-	float w30 = ( 1.5708 - 0.175 * c30 ) * rsqrt( c30 + 1 );	// 1 mad, 1 add, 1 rsqrt
+	float w01 = ( 1.5708 - 0.175 * c01 ) * rsqrt( max(c01 + 1, 1.0e-4f) );	// 1 mad, 1 add, 1 rsqrt
+	float w12 = ( 1.5708 - 0.175 * c12 ) * rsqrt( max(c12 + 1, 1.0e-4f) );	// 1 mad, 1 add, 1 rsqrt
+	float w23 = ( 1.5708 - 0.175 * c23 ) * rsqrt( max(c23 + 1, 1.0e-4f) );	// 1 mad, 1 add, 1 rsqrt
+	float w30 = ( 1.5708 - 0.175 * c30 ) * rsqrt( max(c30 + 1, 1.0e-4f) );	// 1 mad, 1 add, 1 rsqrt
 	// 8 alu, 4 rsqrt
 #endif
 
@@ -326,7 +326,7 @@ float IntegrateEdge( float3 L0, float3 L1 )
 	float w01 = ( 0.8543985 + (0.4965155 + 0.0145206 * abs(c01)) * abs(c01) ) /
 				( 3.4175940 + (4.1616724 + abs(c01)) * abs(c01) );
 	
-	w01 = c01 > 0 ? w01 : 0.5 * rsqrt( 1 - c01 * c01 ) - w01;
+	w01 = c01 > 0 ? w01 : 0.5 * rsqrt( max(1 - c01 * c01, 1.0e-4f)  ) - w01;
 #endif
 
 	return w01;
