Fix 'ghost' rect light silhouette in opposite surface direction, caused by negative irradiance vector.

diff --git a/Engine/Shaders/Private/RectLight.ush b/Engine/Shaders/Private/RectLight.ush
--- a/Engine/Shaders/Private/RectLight.ush
+++ b/Engine/Shaders/Private/RectLight.ush
@@ -478,6 +478,12 @@ float3 RectApproxLTC(FRectLTC In, half3 N, float3 V, FRect Rect, FRectTexture Re
 	// Vector irradiance
 	float3 L = PolygonIrradiance( Poly );
 
+	// Early out negative irradiance vector, causing 'ghost' rect light silhouette at low roughness
+	if (L.z <= 0)
+	{
+		return 0;
+	}
+
 	float LengthSqr = dot( L, L );
 	float InvLength = rsqrt( LengthSqr );
 	float Length = LengthSqr * InvLength;
