Fix Substrate shader compilation warning.

diff --git a/Engine/Shaders/Private/RayTracing/RayTracingMaterialHitShaders.usf b/Engine/Shaders/Private/RayTracing/RayTracingMaterialHitShaders.usf
--- a/Engine/Shaders/Private/RayTracing/RayTracingMaterialHitShaders.usf
+++ b/Engine/Shaders/Private/RayTracing/RayTracingMaterialHitShaders.usf
@@ -275,7 +275,7 @@ void GetMaterialPayload(
 	#if MATERIALBLENDING_ANY_TRANSLUCENT
 	if (SubstratePixelHeader.SubstrateTree.BSDFCount > 0)
 	{
-		float TransmittancePreCoverage = 0.f;
+		float3 TransmittancePreCoverage = 0.f;
 		SubstratePixelHeader.SubstrateUpdateTree(
 			SubstrateData, 
 			MaterialParameters.CameraVector, 
