Fix uninitialized rect light output vector causing NaN value on Substrate Glint material.

diff --git a/Engine/Shaders/Private/RectLight.ush b/Engine/Shaders/Private/RectLight.ush
--- a/Engine/Shaders/Private/RectLight.ush
+++ b/Engine/Shaders/Private/RectLight.ush
@@ -479,6 +479,7 @@ float3 RectApproxLTC(FRectLTC In, half3 N, float3 V, FRect Rect, FRectTexture Re
 	float3 L = PolygonIrradiance( Poly );
 
 	// Early out negative irradiance vector, causing 'ghost' rect light silhouette at low roughness
+	OutMeanLightWorldDirection = N;
 	if (L.z <= 0)
 	{
 		return 0;
