Fixing shader compilation error due to typo

#rnx

diff --git a/Engine/Shaders/Private/BoneTransform.ush b/Engine/Shaders/Private/BoneTransform.ush
--- a/Engine/Shaders/Private/BoneTransform.ush
+++ b/Engine/Shaders/Private/BoneTransform.ush
@@ -124,6 +124,6 @@ FBoneTransformWithScale BlendBoneTransformsWithScale(FBoneTransformWithScale Ato
 	FBoneTransformWithScale Blended;
 	Blended.Location = lerp(Atom1.Location, Atom2.Location, Alpha);
 	Blended.Rotation = QuatSlerpApproximate(Atom1.Rotation, Atom2.Rotation, Alpha);
-	Blended.Scale    = Lerp(Atom1.Scale, Atom2.Scale, Alpha);
+	Blended.Scale    = lerp(Atom1.Scale, Atom2.Scale, Alpha);
 	return Blended;
 }
\ No newline at end of file
