Add PrintLine(Ctx, Value) macro for reducing boiler plate when printing with ShaderPrint.

diff --git a/Engine/Shaders/Private/ShaderPrintCommon.ush b/Engine/Shaders/Private/ShaderPrintCommon.ush
--- a/Engine/Shaders/Private/ShaderPrintCommon.ush
+++ b/Engine/Shaders/Private/ShaderPrintCommon.ush
@@ -1448,4 +1448,10 @@ void SetShaderPrintCurrentPixelCoord(uint2 In)
 uint2 GetShaderPrintCurrentPixelCoord()
 {
 	return ShaderPrintCurrentPixelCoord;
-}
\ No newline at end of file
+}
+
+///////////////////////////////////////////////////////////////////////////////////////////////////
+// Helper functions
+
+// Print a value onto a single line
+#define PrintLine(Ctx, Value) Print(Ctx, TEXT(#Value), FontWhite); Print(Ctx, TEXT(" : "), FontWhite); Print(Ctx, Value, FontYellow); Newline(Ctx);
\ No newline at end of file
