Added test  to prevent merging of pages that are not requested this frame.

diff --git a/Engine/Shaders/Private/VirtualShadowMaps/VirtualShadowMapPhysicalPageManagement.usf b/Engine/Shaders/Private/VirtualShadowMaps/VirtualShadowMapPhysicalPageManagement.usf
--- a/Engine/Shaders/Private/VirtualShadowMaps/VirtualShadowMapPhysicalPageManagement.usf
+++ b/Engine/Shaders/Private/VirtualShadowMaps/VirtualShadowMapPhysicalPageManagement.usf
@@ -861,7 +861,8 @@ void SelectPagesToMergeCS(uint PhysicalPageIndex : SV_DispatchThreadID)
 	// An uncached view is always exclusively renders into the dynamic pages, and thus require no merging.
 	if ((MetaData.Flags & VSM_FLAG_ALLOCATED) != 0U &&
 		(MetaData.Flags & VSM_EXTENDED_FLAG_VIEW_UNCACHED) == 0U &&
-		(MetaData.Flags & VSM_EXTENDED_FLAG_DIRTY) != 0U)
+		(MetaData.Flags & VSM_EXTENDED_FLAG_DIRTY) != 0U &&
+		(MetaData.Flags & VSM_EXTENDED_FLAG_UNREFERENCED) == 0U)
 	{
 		StatsBufferInterlockedInc(VSM_STAT_NUM_PAGES_TO_MERGE);
 		EmitPageToProcess(OutMergePagesIndirectArgsBuffer, OutPhysicalPagesToMerge, PhysicalPageIndex);
