Fixed the page-flag check to also look at the physical page to determine if an unmapped page should be treated as occluded.

diff --git a/Engine/Shaders/Private/VirtualShadowMaps/VirtualShadowMapCacheInvalidation.ush b/Engine/Shaders/Private/VirtualShadowMaps/VirtualShadowMapCacheInvalidation.ush
--- a/Engine/Shaders/Private/VirtualShadowMaps/VirtualShadowMapCacheInvalidation.ush
+++ b/Engine/Shaders/Private/VirtualShadowMaps/VirtualShadowMapCacheInvalidation.ush
@@ -124,7 +124,7 @@ void InvalidateInstancePages(int VirtualShadowMapId, uint InstanceId, bool bForc
 							if ((PageFlags & VSM_FLAG_ALLOCATED) != 0)
 							{
 #if USE_HZB_OCCLUSION
-								if (!IsPageVisibleHZB(uint2(x, y), PageFlagOffset, HZBTestRect))
+								if (!IsPageVisibleHZB(uint2(x, y), PageFlagOffset, false, HZBTestRect))
 								{
 									continue;
 								}
diff --git a/Engine/Shaders/Private/VirtualShadowMaps/VirtualShadowMapPageOverlap.ush b/Engine/Shaders/Private/VirtualShadowMaps/VirtualShadowMapPageOverlap.ush
--- a/Engine/Shaders/Private/VirtualShadowMaps/VirtualShadowMapPageOverlap.ush
+++ b/Engine/Shaders/Private/VirtualShadowMaps/VirtualShadowMapPageOverlap.ush
@@ -125,7 +125,7 @@ FPageTestScreenRect SetupPageHZBRect(FScreenRect ScreenRect, uint ShadowMapID, u
 }
 
 
-bool IsPageVisibleHZB(uint2 vPage, uint PageFlagOffset, FPageTestScreenRect PageTestScreenRect)
+bool IsPageVisibleHZB(uint2 vPage, uint PageFlagOffset, bool bTreatUnmappedAsOccluded, FPageTestScreenRect PageTestScreenRect)
 {
 	FShadowPhysicalPage pPage = ShadowDecodePageTable(HZBPageTable[PageFlagOffset]);
 
@@ -141,11 +141,10 @@ bool IsPageVisibleHZB(uint2 vPage, uint PageFlagOffset, FPageTestScreenRect Page
 		// Translate to physical address space
 		HZBTestRect.HZBTexels += (PhysicalAddress << PageTestScreenRect.HZBLevelPageSizeShift).xyxy;
 
-		// Skip invalidation if the instance is occluded
 		return IsVisibleHZB(HZBTestRect, true);
 	}
 
-	return true;
+	return !bTreatUnmappedAsOccluded;
 }
 
 
@@ -191,7 +190,7 @@ bool IsVisibleMaskedHZB(uint PrevShadowMapID, uint MipLevel, FScreenRect Rect, b
 			if (!bTreatUnmappedAsOccluded || ((PageFlag & TestPageMask) != 0U))
 			{
 				// Treat pages with the VisiblePageMask as visible - can be used to select only cached pages
-				if ((PageFlag & VisiblePageMask) != 0U || IsPageVisibleHZB(uint2(x, y), PageFlagOffset, HZBTestRect))
+				if ((PageFlag & VisiblePageMask) != 0U || IsPageVisibleHZB(uint2(x, y), PageFlagOffset, bTreatUnmappedAsOccluded, HZBTestRect))
 				{
 					return true;
 				}
