Fix incorrect sRGB logic for some compression paths in the RVT format "Base Color, Normal, Roughness".
This format is specifically setup to tradeoff less memory for lower quality. It has a quantized but non-compressed path with 565 and 5651 formats. These don't support sRGB, so we disable that across the board, even for the uncompressed and actually compressed paths.
By mistake we had enabled sRGB for the BC1 compressed path.
Fixed by disabling sRGB and taking the LinearToSrgb() that we use when writing the virtual texture pages into account when compressing them (by not applying again when loading in the compression compute shader).

diff --git a/Engine/Shaders/Private/VirtualTextureCompress.usf b/Engine/Shaders/Private/VirtualTextureCompress.usf
--- a/Engine/Shaders/Private/VirtualTextureCompress.usf
+++ b/Engine/Shaders/Private/VirtualTextureCompress.usf
@@ -157,6 +157,8 @@ void CompressBaseColorNormalRoughnessCS(
 	float2 TexelUVSize = 1.f / float2(SourceRect.zw);
 	float2 SampleUV = (float2(SamplePos) + 0.5f) * TexelUVSize;
 
+	// Note that this format uses non sRGB BaseColor format with manual encode/decode.
+	// So that we don't use _SRGB forms of the block compression functions below.
 	float3 BlockBaseColor[16];
 	ReadBlockRGB(RenderTexture0, TextureSampler0, SampleUV, TexelUVSize, BlockBaseColor);
 
@@ -168,7 +170,7 @@ void CompressBaseColorNormalRoughnessCS(
 	OutCompressTexture0_64bit[ThreadId.xy + DestPos0] = PackRG32(CompressBlock_ETC2_RGB(BlockBaseColor));
 	OutCompressTexture1_128bit[ThreadId.xy + DestPos1] = CompressBlock_ETC2_RGBA(BlockNormalXRougnnessY, BlockNormalZ);
 #else
-	OutCompressTexture0_64bit[ThreadId.xy + DestPos0] = PackRG32(CompressBC1Block_SRGB(BlockBaseColor));
+	OutCompressTexture0_64bit[ThreadId.xy + DestPos0] = PackRG32(CompressBC1Block(BlockBaseColor));
 	OutCompressTexture1_128bit[ThreadId.xy + DestPos1] = CompressBC3Block(BlockNormalXRougnnessY, BlockNormalZ);
 #endif
 
diff --git a/Engine/Source/Runtime/Engine/Private/VT/RuntimeVirtualTexture.cpp b/Engine/Source/Runtime/Engine/Private/VT/RuntimeVirtualTexture.cpp
--- a/Engine/Source/Runtime/Engine/Private/VT/RuntimeVirtualTexture.cpp
+++ b/Engine/Source/Runtime/Engine/Private/VT/RuntimeVirtualTexture.cpp
@@ -484,14 +484,17 @@ bool URuntimeVirtualTexture::IsLayerSRGB(int32 LayerIndex) const
 	switch (MaterialType)
 	{
 	case ERuntimeVirtualTextureMaterialType::BaseColor:
-	case ERuntimeVirtualTextureMaterialType::BaseColor_Normal_Roughness:
 	case ERuntimeVirtualTextureMaterialType::BaseColor_Normal_Specular:
-		// Only BaseColor layer is sRGB
+		// Only BaseColor layer is sRGB.
 		return LayerIndex == 0;
+	case ERuntimeVirtualTextureMaterialType::BaseColor_Normal_Roughness:
+		// This format encodes/decodes sRGB on write/read.
 	case ERuntimeVirtualTextureMaterialType::BaseColor_Normal_Specular_YCoCg:
 	case ERuntimeVirtualTextureMaterialType::BaseColor_Normal_Specular_Mask_YCoCg:
+		// These formats have YCoCg packing which can't use sRGB.
 	case ERuntimeVirtualTextureMaterialType::WorldHeight:
 	case ERuntimeVirtualTextureMaterialType::Displacement:
+		// These formats require linear encoding.
 		return false;
 	default:
 		break;
