[Backout] - CL33776356 and move change to upscale instead of tonemapper
Original CL Desc
-----------------------------------------------------------------
fix mobile propagateAlpha  (invert alpha) for post processing.
#android

diff --git a/Engine/Shaders/Private/PostProcessTonemap.usf b/Engine/Shaders/Private/PostProcessTonemap.usf
--- a/Engine/Shaders/Private/PostProcessTonemap.usf
+++ b/Engine/Shaders/Private/PostProcessTonemap.usf
@@ -603,11 +603,6 @@ void MainPS(
 		float3 UVI = float3(UV.xy,0);
 	#endif
     OutColor = TonemapCommonPS(UVI, InVignette, GrainUV, ScreenPos, FullViewUV, SvPosition, Luminance);
-
-#if POST_PROCESS_ALPHA == 2 && SHADING_PATH_MOBILE
-	OutColor.a = 1.0f - OutColor.a;	// invert alpha when propagate alpha is used with ASIS
-#endif
-
 #if OUTPUT_LUMINANCE
 	OutLuminance = Luminance;
 #endif
diff --git a/Engine/Shaders/Private/PostProcessUpscale.usf b/Engine/Shaders/Private/PostProcessUpscale.usf
--- a/Engine/Shaders/Private/PostProcessUpscale.usf
+++ b/Engine/Shaders/Private/PostProcessUpscale.usf
@@ -306,6 +306,11 @@ void MainPS(noperspective float4 UVAndScreenPos : TEXCOORD0, float4 SvPosition :
 	{
 		OutColor.a = 0.0;
 	}
+
+#if POST_PROCESS_ALPHA == 2 && SHADING_PATH_MOBILE
+	OutColor.a = 1.0f - OutColor.a;	// invert alpha when propagate alpha is used
+#endif
+	
 #else
 	OutColor.a = 0; // Skip all computations related to alpha
 #endif
