* Fix incorrect hair voxel tracing. Trace were slightly offset compare to the shading point.
* Fix incorrect scene depth injection into voxels, causing light leak. The early out condition was counting page instead of voxels.

diff --git a/Engine/Shaders/Private/HairStrands/HairStrandsVoxelPageTraversal.ush b/Engine/Shaders/Private/HairStrands/HairStrandsVoxelPageTraversal.ush
--- a/Engine/Shaders/Private/HairStrands/HairStrandsVoxelPageTraversal.ush
+++ b/Engine/Shaders/Private/HairStrands/HairStrandsVoxelPageTraversal.ush
@@ -211,7 +211,7 @@ FHairTraversalResult ComputeHairCountVirtualVoxel(
 			AddStepAABB(InSettings.bDebugEnabled, PreviousP, HitP, float3(0, 1, 0));
 			#endif
 
-			const uint3 VolumeCoord    = PositionToCoord(HitP, InNodeDesc.TranslatedWorldMinAABB, InNodeDesc.TranslatedWorldMaxAABB, InNodeDesc.VirtualResolution);
+			const uint3 VolumeCoord = clamp((HitP - InNodeDesc.TranslatedWorldMinAABB) / InNodeDesc.VoxelWorldSize, 0, InNodeDesc.VirtualResolution - 1);
 			const uint3 PageIndexCoord = (VolumeCoord >> InCommonDesc.PageResolutionLog2);
 
 			// Update page index only when needed
diff --git a/Engine/Shaders/Private/HairStrands/HairStrandsVoxelRasterCompute.usf b/Engine/Shaders/Private/HairStrands/HairStrandsVoxelRasterCompute.usf
--- a/Engine/Shaders/Private/HairStrands/HairStrandsVoxelRasterCompute.usf
+++ b/Engine/Shaders/Private/HairStrands/HairStrandsVoxelRasterCompute.usf
@@ -769,7 +769,7 @@ void MainCS(uint2 DispatchThreadID : SV_DispatchThreadID)
 		const float Radius = lerp(CP0.WorldRadius, CP1.WorldRadius, U);
 
 		const float3 HitP = TranslatedWP0 + StepIt * StepD + Jitter * Radius;
-		const int3 VolumeCoord = PositionToCoord(HitP, NodeDesc.TranslatedWorldMinAABB, NodeDesc.TranslatedWorldMaxAABB, NodeDesc.VirtualResolution);
+		const int3 VolumeCoord = clamp((HitP - NodeDesc.TranslatedWorldMinAABB) / NodeDesc.VoxelWorldSize, 0, NodeDesc.VirtualResolution-1);
 		const int3 PageIndexCoord = VolumeCoord / CommonDesc.PageResolution;
 
 		// Update page index only when needed
@@ -842,13 +842,16 @@ void MainCS(uint3 DispatchThreadId : SV_DispatchThreadID)
 {
 	const uint LinearVoxelCoord = DispatchThreadId.x;
 	const uint AllocatedPageIndex = DispatchThreadId.z;
-	const bool bValidVoxel = AllocatedPageIndex < VirtualVoxelParams_AllocatedPageCountBuffer[0] && LinearVoxelCoord < VirtualVoxelParams_PageCount;
+	const uint VoxelCountPerPage = VirtualVoxelParams_PageResolution.x * VirtualVoxelParams_PageResolution.x * VirtualVoxelParams_PageResolution.x;
+	const bool bValidVoxel = AllocatedPageIndex < VirtualVoxelParams_AllocatedPageCountBuffer[0] && LinearVoxelCoord < VirtualVoxelParams_PageCount * VoxelCountPerPage;
 	if (!bValidVoxel)
 	{
 		return;
 	}
 
-	const uint3 VoxelCoordOffset = IndexToCoord(LinearVoxelCoord, VirtualVoxelParams_PageResolution.xxx);
+	uint3 VoxelCoordOffset = IndexToCoord(LinearVoxelCoord, VirtualVoxelParams_PageResolution.xxx);
+	VoxelCoordOffset.y = VoxelCoordOffset.y % VirtualVoxelParams_PageResolution.x;
+	VoxelCoordOffset.z = VoxelCoordOffset.z % VirtualVoxelParams_PageResolution.x;
 
 	const FVirtualVoxelNodeDesc NodeDesc = UnpackVoxelNode(VirtualVoxelParams_NodeDescBuffer[MacroGroupId], VirtualVoxelParams_PageResolution);
 
@@ -864,7 +867,7 @@ void MainCS(uint3 DispatchThreadId : SV_DispatchThreadID)
 		{
 			const uint3 VoxelCoordBase = PageIndexCoord.xyz * VirtualVoxelParams_PageResolution;
 			const uint3 VoxelCoord = VoxelCoordBase + VoxelCoordOffset;
-			const float3 TranslatedWorldPosition = VoxelCoord / float3(NodeDesc.VirtualResolution) * (NodeDesc.TranslatedWorldMaxAABB - NodeDesc.TranslatedWorldMinAABB) + NodeDesc.TranslatedWorldMinAABB;
+			const float3 TranslatedWorldPosition = VoxelCoord * NodeDesc.VoxelWorldSize + NodeDesc.TranslatedWorldMinAABB;
 
 			float4 ClipPos = mul(float4(TranslatedWorldPosition, 1), PrimaryView.TranslatedWorldToClip);
 			ClipPos /= ClipPos.w;
