Fixed HZB build group shared barrier to respect wave sizes smaller than 32.

diff --git a/Engine/Shaders/Private/HZB.usf b/Engine/Shaders/Private/HZB.usf
--- a/Engine/Shaders/Private/HZB.usf
+++ b/Engine/Shaders/Private/HZB.usf
@@ -214,16 +214,19 @@ void HZBBuildCS(
 	{
 		SharedMinDeviceZ[GroupThreadIndex] = MinDeviceZ;
 		SharedMaxDeviceZ[GroupThreadIndex] = MaxDeviceZ;
+		const uint LaneCount = WaveGetLaneCount();
 	
 		UNROLL
-		for (uint MipLevel = 1; MipLevel < DIM_MIP_LEVEL_COUNT; MipLevel++)
+		for (uint MipLevel = 1; MipLevel < DIM_MIP_LEVEL_COUNT; ++MipLevel)
 		{
 			const uint TileSize = uint(GROUP_TILE_SIZE) >> MipLevel;
 			const uint ReduceBankSize = TileSize * TileSize;
 			
-			// LDS has been written before.
-			if (MipLevel == 1)
+			// More waves than one wrote to LDS, need to sync.
+			if ((ReduceBankSize << 2u) > LaneCount)
+			{
 				GroupMemoryBarrierWithGroupSync();
+			}
 
 			BRANCH
 			if (GroupThreadIndex < ReduceBankSize)
