Fix minor shader warning

#rnx

diff --git a/Engine/Shaders/Private/VirtualShadowMaps/Desaturate.usf b/Engine/Shaders/Private/VirtualShadowMaps/Desaturate.usf
--- a/Engine/Shaders/Private/VirtualShadowMaps/Desaturate.usf
+++ b/Engine/Shaders/Private/VirtualShadowMaps/Desaturate.usf
@@ -16,6 +16,6 @@ float4 DesaturatePS(
 {
 	float2 UV = UVAndScreenPos.xy;
 	float4 Color = Texture2DSample(InputTexture, InputSampler, UV);
-	Color.rgb = Grayscale(Color);
+	Color.rgb = Grayscale(Color.rgb);
 	return Color;
 }
