Path Tracer: Fix indirect dispatch mode

The addressing logic was using DispatchDim to figure out where to write the pixel, but this was changing in the indirect dispatch case, leading to pixels being written in the wrong place.

Optimize index calculation to avoid modulo and divide instructions and ensure we only keep a single uint in live state for where to write the final pixel. This saves about 0.3ms of overhead.

Fix adaptive sampling when used without path compaction.

Indirect dispatch is still measurably slower (by over 1ms) so leave it off until we can track down the discrepancy.

diff --git a/Engine/Shaders/Private/PathTracing/PathTracing.usf b/Engine/Shaders/Private/PathTracing/PathTracing.usf
--- a/Engine/Shaders/Private/PathTracing/PathTracing.usf
+++ b/Engine/Shaders/Private/PathTracing/PathTracing.usf
@@ -84,6 +84,7 @@ RAY_TRACING_ENTRY_RAYGEN(PathTracingMainRG)
 	const uint2 DispatchDim = DispatchRaysDimensions().xy;
 
 	uint LinearPixelIndex = DispatchIdx.x + DispatchDim.x * DispatchIdx.y;
+	uint PathIndex = DispatchIdx.x + 65536 * DispatchIdx.y;
 	if (LinearPixelIndex == 0)
 	{
 		// write other dimensions for indirect dispatch of the next bounce
@@ -96,12 +97,13 @@ RAY_TRACING_ENTRY_RAYGEN(PathTracingMainRG)
 #if PATH_TRACER_USE_ADAPTIVE_SAMPLING
 		if (Iteration > 0)
 		{
-			if (LinearPixelIndex >= NumPathStates[3 * Bounce + 0])
+			if (LinearPixelIndex >= NumPathStates[0])
 			{
 				return; // nothing left to do on this thread
 			}
-			LinearPixelIndex = ActivePaths[LinearPixelIndex];
-			DispatchIdx = uint2(LinearPixelIndex % DispatchDim.x, LinearPixelIndex / DispatchDim.x);
+			PathIndex = ActivePaths[LinearPixelIndex];
+			DispatchIdx = uint2(PathIndex & 65535, PathIndex >> 16);
+			LinearPixelIndex = DispatchIdx.x + ScanlineWidth * DispatchIdx.y;
 		}
 #endif
 		PathState = CreatePathState(ComputePixelIndex(DispatchIdx), 
@@ -114,13 +116,12 @@ RAY_TRACING_ENTRY_RAYGEN(PathTracingMainRG)
 		{
 			return; // nothing left to do on this thread
 		}
-		LinearPixelIndex = ActivePaths[LinearPixelIndex];
-		PathState = LoadPathStateData(LinearPixelIndex);
+		PathIndex = ActivePaths[LinearPixelIndex];
+		PathState = LoadPathStateData((PathIndex & 65535) + ScanlineWidth * (PathIndex >> 16));
 	}
 
 	const bool KeepGoing = PathTracingKernel(PathState, Bounce);
-	const uint2 TextureIndex = ComputeTextureIndex(uint2(LinearPixelIndex % DispatchDim.x,
-	                                                     LinearPixelIndex / DispatchDim.x));
+	const uint2 TextureIndex = ComputeTextureIndex(uint2(PathIndex & 65535, PathIndex >> 16));
 	if (Bounce == 0)
 	{
 		// Camera ray must write depth right away because we don't want to have to carry DepthZ in the PackedPathState
@@ -131,8 +132,8 @@ RAY_TRACING_ENTRY_RAYGEN(PathTracingMainRG)
 	    // Add the path to the set of valid paths for the next bounce
 		uint PathStateIndex;
 		InterlockedAdd(NumPathStates[3 * Bounce + 3], 1, PathStateIndex);
-		NextActivePaths[PathStateIndex] = LinearPixelIndex;
-		StorePathStateData(PathState, LinearPixelIndex);
+		NextActivePaths[PathStateIndex] = PathIndex;
+		StorePathStateData(PathState, (PathIndex & 65535) + ScanlineWidth * (PathIndex >> 16));
 	}
 	else
 	{
@@ -145,8 +146,8 @@ RAY_TRACING_ENTRY_RAYGEN(PathTracingMainRG)
 #else // PATH_TRACER_USE_COMPACTION == 0
 
 #if PATH_TRACER_USE_ADAPTIVE_SAMPLING
-Buffer<int> ActivePaths;
-Buffer<uint> NumPathStates;
+RWBuffer<int> ActivePaths;
+RWBuffer<uint> NumPathStates;
 #endif
 
 RAY_TRACING_ENTRY_RAYGEN(PathTracingMainRG)
@@ -162,8 +163,9 @@ RAY_TRACING_ENTRY_RAYGEN(PathTracingMainRG)
 		{
 			return;
 		}
-		LinearPixelIndex = ActivePaths[LinearPixelIndex];
-		DispatchIdx = uint2(LinearPixelIndex % DispatchDim.x, LinearPixelIndex / DispatchDim.x);
+		uint PathIndex = ActivePaths[LinearPixelIndex];
+		DispatchIdx = uint2(PathIndex & 65535, PathIndex >> 16);
+		LinearPixelIndex = DispatchIdx.x + DispatchDim.x * DispatchIdx.y;
 	}
 #endif
 	FPathState PathState = CreatePathState(ComputePixelIndex(DispatchIdx), ComputeTextureIndex(DispatchIdx));
diff --git a/Engine/Shaders/Private/PathTracing/PathTracingAdaptiveStart.usf b/Engine/Shaders/Private/PathTracing/PathTracingAdaptiveStart.usf
--- a/Engine/Shaders/Private/PathTracing/PathTracingAdaptiveStart.usf
+++ b/Engine/Shaders/Private/PathTracing/PathTracingAdaptiveStart.usf
@@ -20,8 +20,8 @@ void PathTracingAdaptiveStartCS(uint2 DispatchThreadId : SV_DispatchThreadID)
 		return;
 	}
 
-	int LinearPixelIndex = DispatchIdx.x + DispatchDim.x * DispatchIdx.y;
-	if (LinearPixelIndex == 0)
+	uint PathIndex = DispatchIdx.x + 65536 * DispatchIdx.y;
+	if (PathIndex == 0)
 	{
 		// write other dimensions for indirect dispatch
 		NumPathStates[1] = 1;
@@ -34,6 +34,6 @@ void PathTracingAdaptiveStartCS(uint2 DispatchThreadId : SV_DispatchThreadID)
 		// This pixel needs more samples, allocate a path state for it
 		uint PathStateIndex;
 		InterlockedAdd(NumPathStates[0], 1, PathStateIndex);
-		NextActivePaths[PathStateIndex] = LinearPixelIndex;
+		NextActivePaths[PathStateIndex] = PathIndex;
 	}
 }
diff --git a/Engine/Source/Runtime/Renderer/Private/PathTracing.cpp b/Engine/Source/Runtime/Renderer/Private/PathTracing.cpp
--- a/Engine/Source/Runtime/Renderer/Private/PathTracing.cpp
+++ b/Engine/Source/Runtime/Renderer/Private/PathTracing.cpp
@@ -3603,6 +3603,11 @@ void FDeferredShadingSceneRenderer::RenderPathTracing(
 									PassParameters->PathTracingIndirectArgs = NumActivePaths;
 								}
 							}
+							else if (Config.UseAdaptiveSampling)
+							{
+								PassParameters->ActivePaths = GraphBuilder.CreateUAV(ActivePaths[0], PF_R32_UINT);
+								PassParameters->NumPathStates = GraphBuilder.CreateUAV(NumActivePaths, PF_R32_UINT);
+							}
 							ClearUnusedGraphResources(RayGenShader, PassParameters);
 							const bool bFlushRenderingCommands = FlushRenderingCommands == 1 || (FlushRenderingCommands == 2 && Bounce == MaxBounces);
 							GraphBuilder.AddPass(
@@ -3625,7 +3630,7 @@ void FDeferredShadingSceneRenderer::RenderPathTracing(
 											View.RayTracingMaterialPipeline,
 											RayGenShader.GetRayTracingShader(),
 											RayTracingSceneRHI, View.RayTracingSBT, GlobalResources,
-											PassParameters->PathTracingIndirectArgs->GetIndirectRHICallBuffer(), 3 * PassParameters->Bounce
+											PassParameters->PathTracingIndirectArgs->GetIndirectRHICallBuffer(), 3 * PassParameters->Bounce * sizeof(uint32)
 										);
 									}
 									else
