Fix exponential height fog not resolving view.

It cause the fog to render as completely black.

diff --git a/Engine/Shaders/Private/HeightFogPixelShader.usf b/Engine/Shaders/Private/HeightFogPixelShader.usf
--- a/Engine/Shaders/Private/HeightFogPixelShader.usf
+++ b/Engine/Shaders/Private/HeightFogPixelShader.usf
@@ -42,6 +42,7 @@ void ExponentialPixelMain(
 	out float4 OutColor : SV_Target0
 	)
 {
+	ResolvedView = ResolveView();
 	float SceneDepth;
 	float DeviceZ;
 #if PERMUTATION_SAMPLE_FOG_ON_CLOUDS
