Add option to enable/disable voxel fals color when visualization groom voxels.

diff --git a/Engine/Shaders/Private/HairStrands/HairStrandsVoxelDebug.usf b/Engine/Shaders/Private/HairStrands/HairStrandsVoxelDebug.usf
--- a/Engine/Shaders/Private/HairStrands/HairStrandsVoxelDebug.usf
+++ b/Engine/Shaders/Private/HairStrands/HairStrandsVoxelDebug.usf
@@ -247,10 +247,12 @@ void MainCS(uint3 DispatchThreadId : SV_DispatchThreadID)
 
 
 	// Draw widgets
+	bool bDrawVoxel = true;
 	bool bDrawPage = false;
 	int ForcedMipLevel = -1;
 	{
 		FShaderPrintContext WidgetContext = InitShaderPrintContext(MacroGroupId == 0 && bPrint, uint2(500, 100));
+		bDrawVoxel = AddCheckbox(WidgetContext, TEXT("Draw Voxels"), false, FontWhite); Newline(WidgetContext);
 		bDrawPage = AddCheckbox(WidgetContext, TEXT("Draw Pages"), false, FontWhite); Newline(WidgetContext);
 		bool  bForceMipLevel = AddCheckbox(WidgetContext, TEXT("Force MIP"), false, FontWhite); Newline(WidgetContext);
 		if (bForceMipLevel)
@@ -313,7 +315,7 @@ void MainCS(uint3 DispatchThreadId : SV_DispatchThreadID)
 			}
 		}
 	}
-	else 
+	else if (bDrawVoxel)
 	{
 		FHairTraversalResult TraversalResult = (FHairTraversalResult)0;
 
