Add VSM multiple sample support for hair strands voxel tracing.

diff --git a/Engine/Shaders/Private/VirtualShadowMaps/VirtualShadowMapProjection.usf b/Engine/Shaders/Private/VirtualShadowMaps/VirtualShadowMapProjection.usf
--- a/Engine/Shaders/Private/VirtualShadowMaps/VirtualShadowMapProjection.usf
+++ b/Engine/Shaders/Private/VirtualShadowMaps/VirtualShadowMapProjection.usf
@@ -184,6 +184,7 @@ float SMRTTexelDitherScale;			// Currently only directional lights
 float SMRTExtrapolateSlope;
 float SMRTMaxSlopeBias;				// Currently only local lights
 uint SMRTAdaptiveRayCount;
+uint SMRTHairRayCount;
 
 // One pass projection
 RWTexture2D< uint4 > OutShadowMaskBits;
@@ -358,93 +359,128 @@ FVirtualShadowMapSampleResult ProjectLight(
 
 		// Hair strands voxel traversal to apply hair shadow on top of opaque geometry
 		#if HAS_HAIR_STRANDS
-		if (!bIsHairInput)
+		if (!bIsHairInput && Result.ShadowFactor > 0)
 		{
-			float3 RayStart = 0;
-			float3 RayEnd = 0;
-			bool bCastHairRay = false;
-			float4 Random = 0;
-			if (Result.ShadowFactor > 0)
+			float AccShadowFactor = 0;
+			uint HairRayCount = min(SMRTHairRayCount, uint(SMRTRayCount));
+			uint RayIndex = 0; 
+			for (; RayIndex < HairRayCount; RayIndex++ )
 			{
+				float3 RayStart = 0;
+				float3 RayEnd = 0;
+
 				#if USE_BLUE_NOISE
-				Random = GetHairBlueNoiseJitter(PixelPos, 0 /*SampleIndex*/, 1 /*MaxSampleCount*/, View.StateFrameIndexMod8/*TimeIndex*/);
+				const float4 Random = GetHairBlueNoiseJitter(PixelPos, RayIndex, HairRayCount, View.StateFrameIndexMod8);
 				#else
-				Random = ComputeRandom4(PixelPos);
+				const float4 Random = ComputeRandom4(PixelPos);
 				#endif
 
-				uint RayIndex = min(Random.w * SMRTRayCount, SMRTRayCount - 1);
-
 			#if DIRECTIONAL_LIGHT
-				bCastHairRay = GenerateRayDirectional(
+				const bool bCastHairRay = GenerateRayDirectional(
 					Light,
 					PixelPos,
 					TranslatedWorldPosition,
 					SMRTRayLengthScale * 100,
 					RayIndex,
-					SMRTRayCount,
+					HairRayCount,
 					RayStart,
 					RayEnd);
 			#else
-				bCastHairRay = GenerateRayLocalLight(
+				const bool bCastHairRay = GenerateRayLocalLight(
 					Light,
 					PixelPos,
 					TranslatedWorldPosition,
 					WorldNormal,
 					SMRTCotMaxRayAngleFromLight,
 					RayIndex,
-					SMRTRayCount,
+					HairRayCount,
 					RayStart,
 					RayEnd);
 			#endif
-			}
-
-			if (bCastHairRay)
-			{
-				// Jitter start position to mask/compensate coarse voxel resolution
-				float3 NormalizedDepthBias = 0;
-				{
-					const float PositionBiasScale = 0.5f;
-					NormalizedDepthBias = (VirtualVoxel.DepthBiasScale_Shadow * L + PositionBiasScale * (Random.xyz * 2 - 1));
-				}
 
-				const float DistanceThreshold = 100000.0f;
-				const float CoverageThreshold = 0.995f; // When Coverage is high, we do not trace shadow on opaque since hair/fur is covering the background.
-
-				FVirtualVoxelCommonDesc CommonDesc;
-				CommonDesc.PageCountResolution		= VirtualVoxel.PageCountResolution;
-				CommonDesc.PageTextureResolution	= VirtualVoxel.PageTextureResolution;
-				CommonDesc.PageResolution			= VirtualVoxel.PageResolution;
-				CommonDesc.PageResolutionLog2		= VirtualVoxel.PageResolutionLog2;
-
-				FHairTraversalSettings TraversalSettings = InitHairTraversalSettings();
-				TraversalSettings.DensityScale		= VirtualVoxel.DensityScale_Shadow;
-				TraversalSettings.CountThreshold	= 1;
-				TraversalSettings.DistanceThreshold = DistanceThreshold;
-				TraversalSettings.bDebugEnabled		= false;
-				TraversalSettings.SteppingScale		= VirtualVoxel.SteppingScale_Shadow;
-				TraversalSettings.Random			= Random.xyz;
-				TraversalSettings.PixelRadius		= ConvertGivenDepthRadiusForProjectionType(VirtualVoxel.HairCoveragePixelRadiusAtDepth1, SceneDepth);
-				TraversalSettings.bUseOpaqueVisibility = false;
-				TraversalSettings.bCastShadow		= true;
-
-				const uint VoxelDescCount = VirtualVoxel.NodeDescCount;
-				for (uint VoxelDescIt=0; VoxelDescIt<VoxelDescCount; ++VoxelDescIt)
+				if (bCastHairRay)
 				{
-					const FVirtualVoxelNodeDesc NodeDesc = UnpackVoxelNode(VirtualVoxel.NodeDescBuffer[VoxelDescIt], VirtualVoxel.PageResolution);
-
-					FHairTraversalResult HairResult = InitHairTraversalResult();
-					HairResult = ComputeHairCountVirtualVoxel(
-						RayStart + NodeDesc.VoxelWorldSize * NormalizedDepthBias,
-						RayEnd,
-						CommonDesc,
-						NodeDesc,
-						VirtualVoxel.PageIndexBuffer,
-						VirtualVoxel.PageTexture,
-						TraversalSettings);
-
-					Result.ShadowFactor = min(Result.ShadowFactor, saturate(1 - HairResult.HairCount));
+					// Jitter start position to mask/compensate coarse voxel resolution
+					float3 NormalizedDepthBias = 0;
+					{
+						const float PositionBiasScale = 0.5f;
+						NormalizedDepthBias = (VirtualVoxel.DepthBiasScale_Shadow * L + PositionBiasScale * (Random.xyz * 2 - 1));
+					}
+
+					const float DistanceThreshold = 100000.0f;
+					const float CoverageThreshold = 0.995f; // When Coverage is high, we do not trace shadow on opaque since hair/fur is covering the background.
+
+					FVirtualVoxelCommonDesc CommonDesc;
+					CommonDesc.PageCountResolution		= VirtualVoxel.PageCountResolution;
+					CommonDesc.PageTextureResolution	= VirtualVoxel.PageTextureResolution;
+					CommonDesc.PageResolution			= VirtualVoxel.PageResolution;
+					CommonDesc.PageResolutionLog2		= VirtualVoxel.PageResolutionLog2;
+
+					FHairTraversalSettings TraversalSettings = InitHairTraversalSettings();
+					TraversalSettings.DensityScale		= VirtualVoxel.DensityScale_Shadow;
+					TraversalSettings.CountThreshold	= 1;
+					TraversalSettings.DistanceThreshold = DistanceThreshold;
+					TraversalSettings.bDebugEnabled		= false;
+					TraversalSettings.SteppingScale		= VirtualVoxel.SteppingScale_Shadow;
+					TraversalSettings.Random			= Random.xyz;
+					TraversalSettings.PixelRadius		= ConvertGivenDepthRadiusForProjectionType(VirtualVoxel.HairCoveragePixelRadiusAtDepth1, SceneDepth);
+					TraversalSettings.bUseOpaqueVisibility = false;
+					TraversalSettings.bCastShadow		= true;
+
+					float TraversalHairCount = 0;
+					const uint VoxelDescCount = VirtualVoxel.NodeDescCount;
+					for (uint VoxelDescIt=0; TraversalHairCount < 1 && VoxelDescIt<VoxelDescCount; ++VoxelDescIt)
+					{
+						const FVirtualVoxelNodeDesc NodeDesc = UnpackVoxelNode(VirtualVoxel.NodeDescBuffer[VoxelDescIt], VirtualVoxel.PageResolution);
+
+						FHairTraversalResult HairResult = InitHairTraversalResult();
+						HairResult = ComputeHairCountVirtualVoxel(
+							RayStart + NodeDesc.VoxelWorldSize * NormalizedDepthBias,
+							RayEnd,
+							CommonDesc,
+							NodeDesc,
+							VirtualVoxel.PageIndexBuffer,
+							VirtualVoxel.PageTexture,
+							TraversalSettings);
+						TraversalHairCount += HairResult.HairCount;
+					}
+					const float IterationShadowFactor = saturate(1 - TraversalHairCount);
+					AccShadowFactor += IterationShadowFactor;
+
+					#define USE_ADAPTATIVE 1
+					#if COMPUTESHADER && USE_ADAPTATIVE
+					if (SMRTAdaptiveRayCount > 0)
+					{
+						// TODO: Adapt this heuristic based on SMRTAdaptiveRayCount as well?
+						if( RayIndex == 0 )
+						{
+							bool bMiss = IterationShadowFactor >= 1.f;
+
+							// All lanes missed
+							bool bAllLanesMiss = WaveActiveAllTrue( bMiss );
+							if( bAllLanesMiss )
+							{
+								break;
+							}
+						}
+						else if( RayIndex >= SMRTAdaptiveRayCount )
+						{
+							bool bHit = IterationShadowFactor == 0;
+
+							// After N iterations and all have hit, assume umbra
+							bool bAllLanesHit = WaveActiveAllTrue( bHit );
+							if( bAllLanesHit )
+							{
+								break;
+							}
+						}
+					}
+					#endif
 				}
 			}
+
+			AccShadowFactor /= float(max(1, min(RayIndex, HairRayCount)));
+			Result.ShadowFactor = min(Result.ShadowFactor, AccShadowFactor);
 		}
 		#endif
 	}
diff --git a/Engine/Source/Runtime/Renderer/Private/VirtualShadowMaps/VirtualShadowMapProjection.cpp b/Engine/Source/Runtime/Renderer/Private/VirtualShadowMaps/VirtualShadowMapProjection.cpp
--- a/Engine/Source/Runtime/Renderer/Private/VirtualShadowMaps/VirtualShadowMapProjection.cpp
+++ b/Engine/Source/Runtime/Renderer/Private/VirtualShadowMaps/VirtualShadowMapProjection.cpp
@@ -163,6 +163,13 @@ static TAutoConsoleVariable<int32> CVarSubsurfaceShadowMinSourceAngle(
 	ECVF_RenderThreadSafe
 );
 
+static TAutoConsoleVariable<int32> CVarSMRTRayCountHair(
+	TEXT( "r.Shadow.Virtual.SMRT.SamplesPerRayHair" ),
+	1,
+	TEXT( "Shadow map samples per ray for hair" ),
+	ECVF_Scalability | ECVF_RenderThreadSafe
+);
+
 #if MAX_TEST_PERMUTATION > 0
 static TAutoConsoleVariable<int32> CVarTestPermutation(
 	TEXT( "r.Shadow.Virtual.ProjectionTestPermutation" ),
@@ -197,6 +204,7 @@ FVirtualShadowMapSMRTSettings GetVirtualShadowMapSMRTSettings(bool bDirectionalL
 		Out.SMRTTexelDitherScale = CVarSMRTTexelDitherScaleDirectional.GetValueOnRenderThread();
 		Out.SMRTExtrapolateSlope = CVarSMRTExtrapolateMaxSlopeDirectional.GetValueOnRenderThread();
 		Out.SMRTMaxSlopeBias = 0.0f; // unused in this path
+		Out.SMRTHairRayCount = CVarSMRTRayCountHair.GetValueOnRenderThread();
 	}
 	else
 	{
@@ -207,6 +215,7 @@ FVirtualShadowMapSMRTSettings GetVirtualShadowMapSMRTSettings(bool bDirectionalL
 		Out.SMRTTexelDitherScale = CVarSMRTTexelDitherScaleLocal.GetValueOnRenderThread();
 		Out.SMRTExtrapolateSlope = CVarSMRTExtrapolateMaxSlopeLocal.GetValueOnRenderThread();
 		Out.SMRTMaxSlopeBias = CVarSMRTMaxSlopeBiasLocal.GetValueOnRenderThread();
+		Out.SMRTHairRayCount = CVarSMRTRayCountHair.GetValueOnRenderThread();
 	}
 	return Out;
 }
diff --git a/Engine/Source/Runtime/Renderer/Private/VirtualShadowMaps/VirtualShadowMapProjection.h b/Engine/Source/Runtime/Renderer/Private/VirtualShadowMaps/VirtualShadowMapProjection.h
--- a/Engine/Source/Runtime/Renderer/Private/VirtualShadowMaps/VirtualShadowMapProjection.h
+++ b/Engine/Source/Runtime/Renderer/Private/VirtualShadowMaps/VirtualShadowMapProjection.h
@@ -28,6 +28,7 @@ BEGIN_SHADER_PARAMETER_STRUCT(FVirtualShadowMapSMRTSettings,)
 	SHADER_PARAMETER(float, SMRTExtrapolateSlope)
 	SHADER_PARAMETER(float, SMRTMaxSlopeBias)
 	SHADER_PARAMETER(uint32, SMRTAdaptiveRayCount)
+	SHADER_PARAMETER(uint32, SMRTHairRayCount)
 END_SHADER_PARAMETER_STRUCT()
 
 FVirtualShadowMapSMRTSettings GetVirtualShadowMapSMRTSettings(bool bDirectionalLight);
