Add ViewState param to fog functions

diff --git a/Engine/Shaders/Private/BasePassPixelShader.usf b/Engine/Shaders/Private/BasePassPixelShader.usf
--- a/Engine/Shaders/Private/BasePassPixelShader.usf
+++ b/Engine/Shaders/Private/BasePassPixelShader.usf
@@ -1333,7 +1333,7 @@ void FPixelShaderInOut_MainPS(
 	#if NEEDS_BASEPASS_VERTEX_FOGGING
 		float4 HeightFogging = BasePassInterpolants.VertexFog;
 	#elif NEEDS_BASEPASS_PIXEL_FOGGING
-		float4 HeightFogging = CalculateHeightFog(MaterialParameters.WorldPosition_CamRelative);
+		float4 HeightFogging = CalculateHeightFog(MaterialParameters.WorldPosition_CamRelative, EyeIndex, ResolvedView);
 		#if LOCAL_FOG_VOLUME_ON_TRANSLUCENT
 		const float4 LocalFogVolumeContrib = BasePassInterpolants.VertexFog;
 		HeightFogging = float4(LocalFogVolumeContrib.rgb + HeightFogging.rgb * LocalFogVolumeContrib.a, LocalFogVolumeContrib.a * HeightFogging.a);
@@ -1354,8 +1354,8 @@ void FPixelShaderInOut_MainPS(
 		HeightFogging = float4(LocalFogVolumeContrib.rgb + HeightFogging.rgb * LocalFogVolumeContrib.a, LocalFogVolumeContrib.a * HeightFogging.a);
 	#endif 
 
-		float3 VolumeUV = ComputeVolumeUV(MaterialParameters.AbsoluteWorldPosition, ResolvedView.WorldToClip);
-		Fogging = CombineVolumetricFog(HeightFogging, VolumeUV, EyeIndex, GBuffer.Depth);
+		float3 VolumeUV = ComputeVolumeUV(MaterialParameters.AbsoluteWorldPosition, ResolvedView.WorldToClip, ResolvedView);
+		Fogging = CombineVolumetricFog(HeightFogging, VolumeUV, EyeIndex, GBuffer.Depth, ResolvedView);
 	}
 #endif
 
diff --git a/Engine/Shaders/Private/BasePassVertexShader.usf b/Engine/Shaders/Private/BasePassVertexShader.usf
--- a/Engine/Shaders/Private/BasePassVertexShader.usf
+++ b/Engine/Shaders/Private/BasePassVertexShader.usf
@@ -94,7 +94,7 @@ void Main(
 
 // Calculate the fog needed for translucency
 #if NEEDS_BASEPASS_VERTEX_FOGGING
-	Output.BasePassInterpolants.VertexFog = CalculateHeightFog(WorldPosition.xyz - ResolvedView.TranslatedWorldCameraOrigin);
+	Output.BasePassInterpolants.VertexFog = CalculateHeightFog(WorldPosition.xyz - ResolvedView.TranslatedWorldCameraOrigin, EyeIndex, ResolvedView);
 
 	const float OneOverPreExposure = ResolvedView.OneOverPreExposure;
 
diff --git a/Engine/Shaders/Private/HeightFogCommon.ush b/Engine/Shaders/Private/HeightFogCommon.ush
--- a/Engine/Shaders/Private/HeightFogCommon.ush
+++ b/Engine/Shaders/Private/HeightFogCommon.ush
@@ -6,6 +6,7 @@
 
 #pragma once
 
+#include "Common.ush"
 #include "/Engine/Shared/EnvironmentComponentsFlags.h"
 #include "/Engine/Private/SkyAtmosphereCommon.ush"
 
@@ -184,17 +185,17 @@ float CalculateLineIntegralShared(float FogHeightFalloff, float RayDirectionZ, f
 }
 
 // @param WorldPositionRelativeToCamera = WorldPosition - InCameraPosition
-half4 GetExponentialHeightFog(float3 WorldPositionRelativeToCamera, float ExcludeDistance, uint EyeIndex)
+half4 GetExponentialHeightFog(float3 WorldPositionRelativeToCamera, float ExcludeDistance, uint EyeIndex, ViewState InView)
 {
 	FogStructData FogData = GetFogStructData(EyeIndex);
 	
 	const half MinFogOpacity = FogData.ExponentialFogColorParameter.w;
 	const float MaxWorldObserverHeight = FogData.ExponentialFogParameters.z;
 
-	const float3 WorldObserverOrigin = float3(DFHackToFloat(ResolvedView.WorldCameraOrigin).xy, min(DFHackToFloat(ResolvedView.WorldCameraOrigin).z, MaxWorldObserverHeight)); // Clamp z to max height
+	const float3 WorldObserverOrigin = float3(DFHackToFloat(InView.WorldCameraOrigin).xy, min(DFHackToFloat(InView.WorldCameraOrigin).z, MaxWorldObserverHeight)); // Clamp z to max height
 
 	float3 CameraToReceiver = WorldPositionRelativeToCamera;
-	CameraToReceiver.z += DFHackToFloat(ResolvedView.WorldCameraOrigin).z - WorldObserverOrigin.z; // Compensate this vector for clamping the observer height
+	CameraToReceiver.z += DFHackToFloat(InView.WorldCameraOrigin).z - WorldObserverOrigin.z; // Compensate this vector for clamping the observer height
 	float CameraToReceiverLengthSqr = dot(CameraToReceiver, CameraToReceiver);
 	float CameraToReceiverLengthInv = rsqrt(max(CameraToReceiverLengthSqr, 0.00000001f));
 	float CameraToReceiverLength = CameraToReceiverLengthSqr * CameraToReceiverLengthInv;
@@ -209,10 +210,10 @@ half4 GetExponentialHeightFog(float3 WorldPositionRelativeToCamera, float Exclud
 
 	BRANCH
 	// While rendering a planar reflection with a clip plane, we must compute analytical fog using a camera path starting from the plane, rather than the virtual camera origin
-	if (dot(ResolvedView.GlobalClippingPlane.xyz, 1) > 0.0f)
+	if (dot(InView.GlobalClippingPlane.xyz, 1) > 0.0f)
 	{
-		float CameraOriginPlaneDistance = dot(ResolvedView.GlobalClippingPlane, float4(WorldObserverOrigin + DFHackToFloat(ResolvedView.PreViewTranslation), 1));
-		float PlaneIntersectionTime = -CameraOriginPlaneDistance / dot(CameraToReceiver, ResolvedView.GlobalClippingPlane.xyz);
+		float CameraOriginPlaneDistance = dot(InView.GlobalClippingPlane, float4(WorldObserverOrigin + DFHackToFloat(InView.PreViewTranslation), 1));
+		float PlaneIntersectionTime = -CameraOriginPlaneDistance / dot(CameraToReceiver, InView.GlobalClippingPlane.xyz);
 
 		// Only modify the start distance if the reflection plane is between the camera and receiver point
 		if (PlaneIntersectionTime > 0 && PlaneIntersectionTime < 1)
@@ -271,13 +272,13 @@ half4 GetExponentialHeightFog(float3 WorldPositionRelativeToCamera, float Exclud
 		half3 DirectionalLightInscattering;
 
 		// No need to test View.AtmosphereLightIlluminanceOnGroundPostTransmittance[0].a because InscatteringLightDirection.w above is doing the same test already.
-		DirectionalInscatteringColor = FogData.DirectionalInscatteringColor.xyz + ResolvedView.SkyAtmosphereHeightFogContribution * ResolvedView.AtmosphereLightIlluminanceOnGroundPostTransmittance[0].rgb * UniformPhaseFunction;
-		DirectionalLightInscattering = DirectionalInscatteringColor * pow(saturate(dot(CameraToReceiverNormalized, ResolvedView.AtmosphereLightDirection[0].xyz)), FogData.DirectionalInscatteringColor.w);
+		DirectionalInscatteringColor = FogData.DirectionalInscatteringColor.xyz + InView.SkyAtmosphereHeightFogContribution * InView.AtmosphereLightIlluminanceOnGroundPostTransmittance[0].rgb * UniformPhaseFunction;
+		DirectionalLightInscattering = DirectionalInscatteringColor * pow(saturate(dot(CameraToReceiverNormalized, InView.AtmosphereLightDirection[0].xyz)), FogData.DirectionalInscatteringColor.w);
 
-		if (ResolvedView.AtmosphereLightIlluminanceOnGroundPostTransmittance[1].a > 0.0f) // Skip DirectionalInscatteringColor on the second light when disabled.
+		if (InView.AtmosphereLightIlluminanceOnGroundPostTransmittance[1].a > 0.0f) // Skip DirectionalInscatteringColor on the second light when disabled.
 		{
-			DirectionalInscatteringColor = FogData.DirectionalInscatteringColor.xyz + ResolvedView.SkyAtmosphereHeightFogContribution * ResolvedView.AtmosphereLightIlluminanceOnGroundPostTransmittance[1].rgb * UniformPhaseFunction;
-			DirectionalLightInscattering += DirectionalInscatteringColor * pow(saturate(dot(CameraToReceiverNormalized, ResolvedView.AtmosphereLightDirection[1].xyz)), FogData.DirectionalInscatteringColor.w);
+			DirectionalInscatteringColor = FogData.DirectionalInscatteringColor.xyz + InView.SkyAtmosphereHeightFogContribution * InView.AtmosphereLightIlluminanceOnGroundPostTransmittance[1].rgb * UniformPhaseFunction;
+			DirectionalLightInscattering += DirectionalInscatteringColor * pow(saturate(dot(CameraToReceiverNormalized, InView.AtmosphereLightDirection[1].xyz)), FogData.DirectionalInscatteringColor.w);
 		}
 #else
 		// Setup a cosine lobe around the light direction to approximate inscattering from the directional light off of the ambient haze;
@@ -323,12 +324,12 @@ half4 GetExponentialHeightFog(float3 WorldPositionRelativeToCamera, float Exclud
 
 half4 GetExponentialHeightFog(float3 WorldPositionRelativeToCamera, float ExcludeDistance)
 {
-	return GetExponentialHeightFog(WorldPositionRelativeToCamera, ExcludeDistance, 0);
+	return GetExponentialHeightFog(WorldPositionRelativeToCamera, ExcludeDistance, 0, GetPrimaryView());
 }
 
-half4 CalculateHeightFog(float3 WorldPositionRelativeToCamera, uint EyeIndex)
+half4 CalculateHeightFog(float3 WorldPositionRelativeToCamera, uint EyeIndex, ViewState InView)
 {
-	if (ResolvedView.RenderingReflectionCaptureMask == 0.0f && !IsExponentialFogRenderedInMain(ResolvedView.EnvironmentComponentsFlags))
+	if (InView.RenderingReflectionCaptureMask == 0.0f && !IsExponentialFogRenderedInMain(InView.EnvironmentComponentsFlags))
 	{
 		return float4(0.0f, 0.0f, 0.0f, 1.0f);
 	}
@@ -336,26 +337,25 @@ half4 CalculateHeightFog(float3 WorldPositionRelativeToCamera, uint EyeIndex)
 	float ExcludeDistance = 0;
 #if PERMUTATION_SUPPORT_VOLUMETRIC_FOG
 	// Volumetric fog covers up to MaxDistance along ViewZ, exclude analytical fog from this range
-	float CosAngle = dot(normalize(WorldPositionRelativeToCamera), ResolvedView.ViewForward);
+	float CosAngle = dot(normalize(WorldPositionRelativeToCamera), InView.ViewForward);
 	float InvCosAngle = (CosAngle > FLT_EPSILON) ? rcp(CosAngle) : 0;
-	ExcludeDistance = ResolvedView.VolumetricFogMaxDistance * InvCosAngle;
+	ExcludeDistance = InView.VolumetricFogMaxDistance * InvCosAngle;
 #endif
 
-	half4 FogInscatteringAndOpacity = GetExponentialHeightFog(WorldPositionRelativeToCamera, ExcludeDistance, EyeIndex);
+	half4 FogInscatteringAndOpacity = GetExponentialHeightFog(WorldPositionRelativeToCamera, ExcludeDistance, EyeIndex, InView);
 	return FogInscatteringAndOpacity;
 }
 
 half4 CalculateHeightFog(float3 WorldPositionRelativeToCamera)
 {
-	return CalculateHeightFog(WorldPositionRelativeToCamera, 0);
-
+	return CalculateHeightFog(WorldPositionRelativeToCamera, 0, GetPrimaryView());
 }
 
-float4 CombineVolumetricFog(float4 GlobalFog, float3 VolumeUV, uint EyeIndex, float SceneDepth)
+float4 CombineVolumetricFog(float4 GlobalFog, float3 VolumeUV, uint EyeIndex, float SceneDepth, ViewState InView)
 {
 	float4 VolumetricFogLookup = float4(0, 0, 0, 1);
 	
-	if (ResolvedView.RenderingReflectionCaptureMask == 0.0f && !IsExponentialFogRenderedInMain(ResolvedView.EnvironmentComponentsFlags))
+	if (InView.RenderingReflectionCaptureMask == 0.0f && !IsExponentialFogRenderedInMain(InView.EnvironmentComponentsFlags))
 	{
 		return VolumetricFogLookup;
 	}
@@ -381,7 +381,7 @@ float4 CombineVolumetricFog(float4 GlobalFog, float3 VolumeUV, uint EyeIndex, fl
 	#endif
 
 		// IntegratedLightScattering is pre-exposed, remove pre exposure now so that it can correctly be applied later
-		VolumetricFogLookup.rgb *= ResolvedView.OneOverPreExposure;
+		VolumetricFogLookup.rgb *= InView.OneOverPreExposure;
 	}
 
 	// Do not apply the Froxel volumetric texture in front of the fog start distance. (the soft fading occur in FinalIntegrationCS).
@@ -393,7 +393,7 @@ float4 CombineVolumetricFog(float4 GlobalFog, float3 VolumeUV, uint EyeIndex, fl
 	//VolumetricFogLookup.rgb += .1f * frac(min(ZSlice, 1.0f) / View.VolumetricFogInvGridSize.z);
 
 #if POST_PROCESS_ALPHA
-	if (IsExponentialFogHoldout(ResolvedView.EnvironmentComponentsFlags) && ResolvedView.RenderingReflectionCaptureMask == 0.0f)
+	if (IsExponentialFogHoldout(InView.EnvironmentComponentsFlags) && InView.RenderingReflectionCaptureMask == 0.0f)
 	{
 		VolumetricFogLookup.rgb *= 0;
 		GlobalFog.rgb *= 0;
@@ -403,31 +403,55 @@ float4 CombineVolumetricFog(float4 GlobalFog, float3 VolumeUV, uint EyeIndex, fl
 	return float4(VolumetricFogLookup.rgb + GlobalFog.rgb * VolumetricFogLookup.a, VolumetricFogLookup.a * GlobalFog.a);
 }
 
+float4 CombineVolumetricFog(float4 GlobalFog, float3 VolumeUV, uint EyeIndex, float SceneDepth)
+{
+	return CombineVolumetricFog(GlobalFog, VolumeUV, EyeIndex, SceneDepth, GetPrimaryView());
+}
+
+float ComputeNormalizedZSliceFromDepth(float SceneDepth, ViewState InView)
+{
+	return log2(SceneDepth * InView.VolumetricFogGridZParams.x + InView.VolumetricFogGridZParams.y) * InView.VolumetricFogGridZParams.z * InView.VolumetricFogInvGridSize.z;
+}
+
 float ComputeNormalizedZSliceFromDepth(float SceneDepth)
 {
-	return log2(SceneDepth * ResolvedView.VolumetricFogGridZParams.x + ResolvedView.VolumetricFogGridZParams.y) * ResolvedView.VolumetricFogGridZParams.z * ResolvedView.VolumetricFogInvGridSize.z;
+	return ComputeNormalizedZSliceFromDepth(SceneDepth, GetPrimaryView());
 }
 
-float3 ComputeVolumeUVFromNDC(float4 NDCPosition)
+float3 ComputeVolumeUVFromNDC(float4 NDCPosition, ViewState InView)
 {
 	NDCPosition.xy /= NDCPosition.w;
-	float3 VolumeUV = float3(NDCPosition.xy * float2(.5f, -.5f) + .5f, ComputeNormalizedZSliceFromDepth(NDCPosition.w));
-	return min(VolumeUV * float3(ResolvedView.VolumetricFogScreenToResourceUV.xy, 1.0), float3(ResolvedView.VolumetricFogUVMax, 1.0));
+	float3 VolumeUV = float3(NDCPosition.xy * float2(.5f, -.5f) + .5f, ComputeNormalizedZSliceFromDepth(NDCPosition.w, InView));
+	return min(VolumeUV * float3(InView.VolumetricFogScreenToResourceUV.xy, 1.0), float3(InView.VolumetricFogUVMax, 1.0));
 }
-float3 ComputeVolumeUVFromTranslatedPos(float3 TranslatedWorldPosition, float4x4 TranslatedWorldToClip)
+
+float3 ComputeVolumeUVFromNDC(float4 NDCPosition)
+{
+	return ComputeVolumeUVFromNDC(NDCPosition, GetPrimaryView());
+}
+
+float3 ComputeVolumeUVFromTranslatedPos(float3 TranslatedWorldPosition, float4x4 TranslatedWorldToClip, ViewState InView)
 {
 	float4 NDCPosition = mul(float4(TranslatedWorldPosition, 1), TranslatedWorldToClip);
-	return ComputeVolumeUVFromNDC(NDCPosition);
+	return ComputeVolumeUVFromNDC(NDCPosition, InView);
+}
+float3 ComputeVolumeUVFromTranslatedPos(float3 TranslatedWorldPosition, float4x4 TranslatedWorldToClip)
+{
+	return ComputeVolumeUVFromTranslatedPos(TranslatedWorldPosition, TranslatedWorldToClip, GetPrimaryView());
 }
 float3 ComputeVolumeUV_DEPRECATED(float3 WorldPosition, float4x4 WorldToClip)
 {
 	float4 NDCPosition = mul(float4(WorldPosition, 1), WorldToClip);
 	return ComputeVolumeUVFromNDC(NDCPosition);
 }
-float3 ComputeVolumeUV(FDFVector3 WorldPosition, FDFInverseMatrix WorldToClip)
+float3 ComputeVolumeUV(FDFVector3 WorldPosition, FDFInverseMatrix WorldToClip, ViewState InView) //ericado check
 {
 	float4 NDCPosition = DFMultiplyDemote(MakeDFVector(WorldPosition, 1.0f), WorldToClip);
-	return ComputeVolumeUVFromNDC(NDCPosition);
+	return ComputeVolumeUVFromNDC(NDCPosition, InView);
+}
+float3 ComputeVolumeUV(FDFVector3 WorldPosition, FDFInverseMatrix WorldToClip)
+{
+	return ComputeVolumeUV(WorldPosition, WorldToClip, GetPrimaryView());
 }
 
 float3 ComputeHistoryVolumeUVFromNDC(float4 NDCPosition)
diff --git a/Engine/Shaders/Private/HeightFogPixelShader.usf b/Engine/Shaders/Private/HeightFogPixelShader.usf
--- a/Engine/Shaders/Private/HeightFogPixelShader.usf
+++ b/Engine/Shaders/Private/HeightFogPixelShader.usf
@@ -42,7 +42,6 @@ void ExponentialPixelMain(
 	out float4 OutColor : SV_Target0
 	)
 {
-	ResolvedView = ResolveView();
 	float SceneDepth;
 	float DeviceZ;
 #if PERMUTATION_SAMPLE_FOG_ON_CLOUDS
diff --git a/Engine/Shaders/Private/MobileBasePassVertexShader.usf b/Engine/Shaders/Private/MobileBasePassVertexShader.usf
--- a/Engine/Shaders/Private/MobileBasePassVertexShader.usf
+++ b/Engine/Shaders/Private/MobileBasePassVertexShader.usf
@@ -51,9 +51,11 @@ void Main(
 	LayerIndex = EyeIndex;
 	Output.BasePassInterpolants.MultiViewId = EyeIndex;
 #elif MOBILE_MULTI_VIEW
+	const uint EyeIndex = ViewId;
 	ResolvedView = ResolveView(ViewId);
 	Output.BasePassInterpolants.MultiViewId = ViewId;
 #else
+	const uint EyeIndex = 0;
 	ResolvedView = ResolveView();
 #endif
 
@@ -90,7 +92,7 @@ void Main(
 #endif
 
 #if USE_VERTEX_FOG
-	half4 VertexFog = CalculateHeightFog(WorldPosition.xyz - ResolvedView.TranslatedWorldCameraOrigin);
+	half4 VertexFog = CalculateHeightFog(WorldPosition.xyz - ResolvedView.TranslatedWorldCameraOrigin, EyeIndex, ResolvedView);
 
 #if PROJECT_SUPPORT_SKY_ATMOSPHERE && MATERIAL_IS_SKY==0 // Do not apply aerial perpsective on sky materials
 	if (ResolvedView.SkyAtmosphereApplyCameraAerialPerspectiveVolume > 0.0f)
diff --git a/Engine/Shaders/Private/MobileFog.usf b/Engine/Shaders/Private/MobileFog.usf
--- a/Engine/Shaders/Private/MobileFog.usf
+++ b/Engine/Shaders/Private/MobileFog.usf
@@ -89,7 +89,7 @@ void MobileFogPS(
 
 // at least one of these is active
 #if PERMUTATION_SUPPORT_HEIGHT_FOG
-	FinalFog = CalculateHeightFog(WorldPositionRelativeToCamera, ArrayIndex);
+	FinalFog = CalculateHeightFog(WorldPositionRelativeToCamera, ArrayIndex, ResolvedView);
 #endif
 
 #if PERMUTATION_SUPPORT_LOCAL_FOG_VOLUME == 1
@@ -164,8 +164,8 @@ void MobileFogPS(
 #if PERMUTATION_SUPPORT_VOLUMETRIC_FOG
 	if (FogStruct.ApplyVolumetricFog > 0) 
 	{
-		float3 VolumeUV = ComputeVolumeUVFromTranslatedPos(TranslatedWorldPosition, ResolvedView.TranslatedWorldToClip);
-		FinalFog = CombineVolumetricFog(FinalFog, VolumeUV, ArrayIndex, SceneDepth);
+		float3 VolumeUV = ComputeVolumeUVFromTranslatedPos(TranslatedWorldPosition, ResolvedView.TranslatedWorldToClip, ResolvedView);
+		FinalFog = CombineVolumetricFog(FinalFog, VolumeUV, ArrayIndex, SceneDepth, ResolvedView);
 	}
 #endif
 
