Heterogeneous Volumes: Fixing monochromatic sampling when evaluating substrate

#tests tested color extinction material property with HVs and substrate

diff --git a/Engine/Shaders/Private/HeterogeneousVolumes/HeterogeneousVolumesLiveShadingUtils.ush b/Engine/Shaders/Private/HeterogeneousVolumes/HeterogeneousVolumesLiveShadingUtils.ush
--- a/Engine/Shaders/Private/HeterogeneousVolumes/HeterogeneousVolumesLiveShadingUtils.ush
+++ b/Engine/Shaders/Private/HeterogeneousVolumes/HeterogeneousVolumesLiveShadingUtils.ush
@@ -38,7 +38,7 @@ float3 SampleExtinctionCoefficients(in FPixelMaterialInputs PixelMaterialInputs)
 	float3 Extinction = 0.0f;
 #if SUBSTRATE_ENABLED
 	FSubstrateBSDF BSDF = PixelMaterialInputs.FrontMaterial.InlinedBSDF;
-	Extinction = VOLUMETRICFOGCLOUD_EXTINCTION(BSDF).r;
+	Extinction = VOLUMETRICFOGCLOUD_EXTINCTION(BSDF).rgb;
 #else
 #if !MATERIAL_SHADINGMODEL_UNLIT
 	Extinction = GetMaterialSubsurfaceDataRaw(PixelMaterialInputs).rgb;
