Fix inconsistent selection outlines of nanite enabled instanced static meshes.
* FEmitEditorSelectionDepthPS shader uses BinarySearch (with IsOnlySelected permutation) to determine if fragment cluster is selected.  Will write stencil if so, otherwise will discard to drop stencil write.
* BinarySearch will produce false negatives for non-power 2 buffer sizes.
* Fix is to resize the EditorSelectedNaniteHitProxyIds buffer to the next highest power of two.  This buffer is eventually uploaded to GPU for use in shader.
* Added comments to shader to make is more clear that the algorithms only work for power of two buffer sizes, hope the comment will avoid this problem in the future.
* Small spot-fix to BinarySearch to handle BufferOffset correctly. Not currently used anywhere in shaders as far as I can tell
* Fixes inconsistent selection outlines of nanite enabled instanced static meshes.

diff --git a/Engine/Shaders/Private/BinarySearch.ush b/Engine/Shaders/Private/BinarySearch.ush
--- a/Engine/Shaders/Private/BinarySearch.ush
+++ b/Engine/Shaders/Private/BinarySearch.ush
@@ -11,6 +11,7 @@ struct TDefaultCompare
 	}
 };
 
+// BufferCount must be a power of 2
 template<typename TBuffer, typename T, typename TCompare>
 T LowerBound(TBuffer SearchBuffer, uint BufferOffset, uint BufferCount, T Key)
 {
@@ -27,6 +28,7 @@ T LowerBound(TBuffer SearchBuffer, uint BufferOffset, uint BufferCount, T Key)
 	return Index;
 }
 
+// BufferCount must be a power of 2
 template<typename TBuffer, typename T, typename TCompare>
 bool BinarySearch(TBuffer SearchBuffer, uint BufferOffset, uint BufferCount, T Key, inout uint OutIndex)
 {
@@ -34,10 +36,11 @@ bool BinarySearch(TBuffer SearchBuffer, uint BufferOffset, uint BufferCount, T K
 	{
 		return false;
 	}
-	OutIndex = LowerBound<TBuffer, T, TCompare>(SearchBuffer, 0, BufferCount, Key);
+	OutIndex = LowerBound<TBuffer, T, TCompare>(SearchBuffer, BufferOffset, BufferCount, Key);
 	return SearchBuffer[OutIndex] == Key;
 }
 
+// BufferCount must be a power of 2
 template<typename TBuffer, typename T, typename TCompare>
 bool BinarySearch(TBuffer SearchBuffer, uint BufferOffset, uint BufferCount, T Key)
 {
@@ -45,12 +48,14 @@ bool BinarySearch(TBuffer SearchBuffer, uint BufferOffset, uint BufferCount, T K
 	return BinarySearch<TBuffer, T, TCompare>(SearchBuffer, BufferOffset, BufferCount, Key, Index);
 }
 
+// BufferCount must be a power of 2
 template<typename TBuffer, typename T>
 bool BinarySearch(TBuffer SearchBuffer, uint BufferOffset, uint BufferCount, T Key, inout uint OutIndex)
 {
 	return BinarySearch< TBuffer, T, TDefaultCompare<T> >(SearchBuffer, BufferOffset, BufferCount, Key, OutIndex);
 }
 
+// BufferCount must be a power of 2
 template<typename TBuffer, typename T>
 bool BinarySearch(TBuffer SearchBuffer, uint BufferOffset, uint BufferCount, T Key)
 {
diff --git a/Engine/Source/Runtime/Renderer/Private/SceneVisibility.cpp b/Engine/Source/Runtime/Renderer/Private/SceneVisibility.cpp
--- a/Engine/Source/Runtime/Renderer/Private/SceneVisibility.cpp
+++ b/Engine/Source/Runtime/Renderer/Private/SceneVisibility.cpp
@@ -4544,10 +4544,23 @@ void FVisibilityTaskData::SetupMeshPasses(FExclusiveDepthStencil::Type BasePassD
 
 	#if WITH_EDITOR
 		{
-			// Sort, uniquify and truncate the selected Nanite hit proxy IDs
+			// Sort the selected Nanite hit proxy IDs
 			Algo::Sort(View.EditorSelectedNaniteHitProxyIds);
-			int32 EndIndex = Algo::Unique(View.EditorSelectedNaniteHitProxyIds);
-			View.EditorSelectedNaniteHitProxyIds.RemoveAt(EndIndex, View.EditorSelectedNaniteHitProxyIds.Num() - EndIndex);
+
+			// Make sure it is power of 2 as it is searched in the shader with binary search that only works with power2 buffer sizes
+			int32 PreviousCount = View.EditorSelectedNaniteHitProxyIds.Num();
+			int32 NewCount = FMath::RoundUpToPowerOfTwo(PreviousCount);
+			
+			if (PreviousCount > 0)
+			{
+				const uint32 LastValue = View.EditorSelectedNaniteHitProxyIds.Last();
+				View.EditorSelectedNaniteHitProxyIds.SetNumUninitialized(NewCount);
+				// Pad the end with the last value to ensure still sorted
+				for (int32 Index = PreviousCount; Index < NewCount; ++Index)
+				{
+					View.EditorSelectedNaniteHitProxyIds[Index] = LastValue;
+				}
+			}
 		}
 	#endif
 
