Fix shader compilation warning when cooking Lyra.

diff --git a/Engine/Plugins/Runtime/HairStrands/Source/HairStrandsCore/Private/GroomTextureBuilder.cpp b/Engine/Plugins/Runtime/HairStrands/Source/HairStrandsCore/Private/GroomTextureBuilder.cpp
--- a/Engine/Plugins/Runtime/HairStrands/Source/HairStrandsCore/Private/GroomTextureBuilder.cpp
+++ b/Engine/Plugins/Runtime/HairStrands/Source/HairStrandsCore/Private/GroomTextureBuilder.cpp
@@ -10,6 +10,7 @@
 #include "RenderGraphUtils.h"
 #include "RHI.h"
 #include "RHIStaticStates.h"
+#include "ShaderCompilerCore.h"
 #include "ShaderParameterStruct.h"
 #include "GlobalShader.h"
 //#include "ShaderPrint.h"
@@ -480,6 +481,11 @@ public:
 	static void ModifyCompilationEnvironment(const FGlobalShaderPermutationParameters& Parameters, FShaderCompilerEnvironment& OutEnvironment)
 	{
 		FGlobalShader::ModifyCompilationEnvironment(Parameters, OutEnvironment);
+		// Work around issues with HLSLcc by switching to DXC
+		if (IsHlslccShaderPlatform(Parameters.Platform))
+		{
+			OutEnvironment.CompilerFlags.Add(CFLAG_ForceDXC);
+		}
 		OutEnvironment.SetDefine(TEXT("SHADER_TEXTURE_CLEAR"), 1);
 	}
 };
@@ -780,6 +786,11 @@ class FHairStrandsTextureCS : public FGlobalShader
 	static void ModifyCompilationEnvironment(const FGlobalShaderPermutationParameters& Parameters, FShaderCompilerEnvironment& OutEnvironment)
 	{
 		FGlobalShader::ModifyCompilationEnvironment(Parameters, OutEnvironment);
+		// Work around issues with HLSLcc by switching to DXC
+		if (IsHlslccShaderPlatform(Parameters.Platform))
+		{
+			OutEnvironment.CompilerFlags.Add(CFLAG_ForceDXC);
+		}
 		OutEnvironment.SetDefine(TEXT("SHADER_COMPUTE"), 1);
 	}
 };
diff --git a/Engine/Shaders/Private/HairStrands/HairStrandsBindingCommon.ush b/Engine/Shaders/Private/HairStrands/HairStrandsBindingCommon.ush
--- a/Engine/Shaders/Private/HairStrands/HairStrandsBindingCommon.ush
+++ b/Engine/Shaders/Private/HairStrands/HairStrandsBindingCommon.ush
@@ -60,7 +60,7 @@ FHairMeshTriangle GetHairMeshTriangle(uint TriangleIndex, Buffer<float> InPositi
 	const uint VId1 = BaseVertexId + 1;
 	const uint VId2 = BaseVertexId + 2;
 
-	FHairMeshTriangle Out;
+	FHairMeshTriangle Out = (FHairMeshTriangle)0;
 	Out.P0 = GetVertexPosition(VId0, InPositionBuffer, InIndexBuffer, InIndexOffset, InMaxIndexCount, InMaxVertexCount);
 	Out.P1 = GetVertexPosition(VId1, InPositionBuffer, InIndexBuffer, InIndexOffset, InMaxIndexCount, InMaxVertexCount);
 	Out.P2 = GetVertexPosition(VId2, InPositionBuffer, InIndexBuffer, InIndexOffset, InMaxIndexCount, InMaxVertexCount);
diff --git a/Engine/Shaders/Private/HairStrands/HairStrandsClusterCommon.ush b/Engine/Shaders/Private/HairStrands/HairStrandsClusterCommon.ush
--- a/Engine/Shaders/Private/HairStrands/HairStrandsClusterCommon.ush
+++ b/Engine/Shaders/Private/HairStrands/HairStrandsClusterCommon.ush
@@ -237,7 +237,7 @@ uint GetHairControlPointMinLOD(uint InPointIndex, Buffer<uint> InPointLODBuffer)
 
 bool IsHairControlPointActive(uint InMinLOD, float InLODIndex)
 {
-	return floor(InLODIndex) <= InMinLOD;
+	return uint(floor(InLODIndex)) <= InMinLOD;
 }
 
 bool IsHairControlPointActive(uint InPointIndex, Buffer<uint> InPointLODBuffer, float InLODIndex)
diff --git a/Engine/Shaders/Private/HairStrands/HairStrandsDebug.usf b/Engine/Shaders/Private/HairStrands/HairStrandsDebug.usf
--- a/Engine/Shaders/Private/HairStrands/HairStrandsDebug.usf
+++ b/Engine/Shaders/Private/HairStrands/HairStrandsDebug.usf
@@ -788,7 +788,7 @@ void VisibilityDebugPPLLCS(uint GroupIndex : SV_GroupIndex, uint3 DispatchThread
 
 #include "../PostProcessCommon.ush"
 
-float DomScale = 1;
+float DomScale;
 float2 DomAtlasOffset;
 float2 DomAtlasScale;
 float2 OutputResolution;
diff --git a/Engine/Shaders/Private/HairStrands/HairStrandsRaytracingGeometry.usf b/Engine/Shaders/Private/HairStrands/HairStrandsRaytracingGeometry.usf
--- a/Engine/Shaders/Private/HairStrands/HairStrandsRaytracingGeometry.usf
+++ b/Engine/Shaders/Private/HairStrands/HairStrandsRaytracingGeometry.usf
@@ -71,7 +71,7 @@ void SetInvalidSegment(uint InVertexIndex)
 {
 #if PERMUTATION_PROCEDURAL_PRIMITIVE
 	const float NaN = asfloat(0x7FC00000);
-	for (int i = 0; i < RaytracingProceduralSplits; i++)
+	for (uint i = 0; i < RaytracingProceduralSplits; i++)
 		SetAABB(RaytracingProceduralSplits * InVertexIndex + i, NaN, NaN);
 #else
 	const uint BasePoint = InVertexIndex * 4;
@@ -215,7 +215,7 @@ void MainCS(uint DispatchThreadId : SV_DispatchThreadID)
 
 				float4 SplitA = A;
 				float SumArea = 0.0;
-				for (int i = 0; i < RaytracingProceduralSplits; i++)
+				for (uint i = 0; i < RaytracingProceduralSplits; i++)
 				{
 					float4 SplitB = lerp(A, B, float(i + 1) / float(RaytracingProceduralSplits));
 					float3 Lo = min(SplitA.xyz - SplitA.w, SplitB.xyz - SplitB.w);
@@ -232,7 +232,7 @@ void MainCS(uint DispatchThreadId : SV_DispatchThreadID)
 					float3 Lo = min(A.xyz - A.w, B.xyz - B.w);
 					float3 Hi = max(A.xyz + A.w, B.xyz + B.w);
 					SetAABB(RaytracingProceduralSplits * VertexIndex0 + 0, Lo, Hi);
-					for (int i = 1; i < RaytracingProceduralSplits; i++)
+					for (uint i = 1; i < RaytracingProceduralSplits; i++)
 					{
 						const float NaN = asfloat(0x7FC00000);
 						SetAABB(RaytracingProceduralSplits * VertexIndex0 + i, NaN, NaN);
diff --git a/Engine/Shaders/Private/HairStrands/HairStrandsTexturesGeneration.usf b/Engine/Shaders/Private/HairStrands/HairStrandsTexturesGeneration.usf
--- a/Engine/Shaders/Private/HairStrands/HairStrandsTexturesGeneration.usf
+++ b/Engine/Shaders/Private/HairStrands/HairStrandsTexturesGeneration.usf
@@ -1,6 +1,6 @@
 // Copyright Epic Games, Inc. All Rights Reserved.
 
-#include "../Common.ush"
+#include "../Common.ush" 
 #include "../MonteCarlo.ush"
 
 #include "HairStrandsVisibilityCommon.ush"
@@ -18,6 +18,28 @@
 #error Max cards texture count is != 6. This codes needs to be updated.
 #endif
 
+void AssignValue1(inout float4 OutTexture, uint ChannelIndex, float Value)
+{
+	switch (ChannelIndex)
+	{
+		case 0: OutTexture.x = Value; break;
+		case 1: OutTexture.y = Value; break;
+		case 2: OutTexture.z = Value; break;
+		case 3: OutTexture.w = Value; break;
+		default: break;
+	}
+}
+
+void AssignValue2(inout float4 OutTexture, uint ChannelIndex, float2 Value)
+{
+	switch (ChannelIndex)
+	{
+		case 0: OutTexture.xy = Value; break;
+		case 1: OutTexture.zw = Value; break;
+		default: break;
+	}
+}
+
 #define WriteValue1(InLayout, InAttribute, In)\
 { \
 	if (IsHairCardTextureIndexValid(InLayout, InAttribute))\
@@ -25,13 +47,13 @@
 		const uint TexIndex = GetHairCardTextureIndex(InLayout, InAttribute);\
 		const uint ChannelIndex = GetHairCardChannelIndex(InLayout, InAttribute);\
 		switch (TexIndex)\
-		{\
-			case 0: ChannelIndex == 0 ? OutTexture0.x = In : (ChannelIndex == 1 ? OutTexture0.y = In : (ChannelIndex == 2 ? OutTexture0.z = In : OutTexture0.w = In)); break;\
-			case 1: ChannelIndex == 0 ? OutTexture1.x = In : (ChannelIndex == 1 ? OutTexture1.y = In : (ChannelIndex == 2 ? OutTexture1.z = In : OutTexture1.w = In)); break;\
-			case 2: ChannelIndex == 0 ? OutTexture2.x = In : (ChannelIndex == 1 ? OutTexture2.y = In : (ChannelIndex == 2 ? OutTexture2.z = In : OutTexture2.w = In)); break;\
-			case 3: ChannelIndex == 0 ? OutTexture3.x = In : (ChannelIndex == 1 ? OutTexture3.y = In : (ChannelIndex == 2 ? OutTexture3.z = In : OutTexture3.w = In)); break;\
-			case 4: ChannelIndex == 0 ? OutTexture4.x = In : (ChannelIndex == 1 ? OutTexture4.y = In : (ChannelIndex == 2 ? OutTexture4.z = In : OutTexture4.w = In)); break;\
-			case 5: ChannelIndex == 0 ? OutTexture5.x = In : (ChannelIndex == 1 ? OutTexture5.y = In : (ChannelIndex == 2 ? OutTexture5.z = In : OutTexture5.w = In)); break;\
+		{ \
+			case 0: AssignValue1(OutTexture0, ChannelIndex, In); break;\
+			case 1: AssignValue1(OutTexture1, ChannelIndex, In); break;\
+			case 2: AssignValue1(OutTexture2, ChannelIndex, In); break;\
+			case 3: AssignValue1(OutTexture3, ChannelIndex, In); break;\
+			case 4: AssignValue1(OutTexture4, ChannelIndex, In); break;\
+			case 5: AssignValue1(OutTexture5, ChannelIndex, In); break;\
 		}\
 	}\
 }
@@ -43,12 +65,12 @@
 		const uint ChannelIndex = GetHairCardChannelIndex(InLayout, InAttribute);\
 		switch (TexIndex)\
 		{\
-			case 0: ChannelIndex < 2 ? OutTexture0.xy = In : OutTexture0.zw = In; break;\
-			case 1: ChannelIndex < 2 ? OutTexture1.xy = In : OutTexture1.zw = In; break;\
-			case 2: ChannelIndex < 2 ? OutTexture2.xy = In : OutTexture2.zw = In; break;\
-			case 3: ChannelIndex < 2 ? OutTexture3.xy = In : OutTexture3.zw = In; break;\
-			case 4: ChannelIndex < 2 ? OutTexture4.xy = In : OutTexture4.zw = In; break;\
-			case 5: ChannelIndex < 2 ? OutTexture5.xy = In : OutTexture5.zw = In; break;\
+			case 0: AssignValue2(OutTexture0, ChannelIndex, In); break;\
+			case 1: AssignValue2(OutTexture1, ChannelIndex, In); break;\
+			case 2: AssignValue2(OutTexture2, ChannelIndex, In); break;\
+			case 3: AssignValue2(OutTexture3, ChannelIndex, In); break;\
+			case 4: AssignValue2(OutTexture4, ChannelIndex, In); break;\
+			case 5: AssignValue2(OutTexture5, ChannelIndex, In); break;\
 		}\
 	}\
 }
@@ -615,7 +637,7 @@ RWTexture2DArray<float4> TargetTexture;
 void MainCS(int2 DispatchThreadId : SV_DispatchThreadID)
 {
 	const uint2 PixelCoord = DispatchThreadId.xy;
-	if (any(PixelCoord >= Resolution))
+	if (any(PixelCoord >= uint2(Resolution)))
 	{
 		return;
 	}
diff --git a/Engine/Shaders/Private/HairStrands/HairStrandsVisibilityRasterCompute.usf b/Engine/Shaders/Private/HairStrands/HairStrandsVisibilityRasterCompute.usf
--- a/Engine/Shaders/Private/HairStrands/HairStrandsVisibilityRasterCompute.usf
+++ b/Engine/Shaders/Private/HairStrands/HairStrandsVisibilityRasterCompute.usf
@@ -833,7 +833,7 @@ void PlotInternal(int2 Coords, float AntiAliasingFactor, float4 P0, float4 P1, f
 {
 	const int2 IntraTileCoord = Coords - int2(((PackedTileMin >> 0) & 0xffff), ((PackedTileMin >> 16) & 0xffff));
 
-	if (all(IntraTileCoord >= 0) && all(IntraTileCoord < TileSize))
+	if (all(IntraTileCoord >= 0) && all(IntraTileCoord < int2(TileSize,TileSize)))
 	{
 		const float Alpha = ComputeLerpAlpha(Coords, P0.xy, P1.xy, SegmentLenSqRcp);
 		const float Depth = lerp(P0.z, P1.z, Alpha);
@@ -1250,7 +1250,7 @@ void Plot(int2 Coord, float4 P0, float4 P1, float Rad0, float Rad1, float Segmen
 {
 	const int2 IntraTileCoord = Coord - int2(((PackedHalfTileMin >> 0) & 0xffff), ((PackedHalfTileMin >> 16) & 0xffff));
 
-	if (all(IntraTileCoord >= 0) && all(IntraTileCoord < HalfTileSize))
+	if (all(IntraTileCoord >= 0) && all(IntraTileCoord < int2(HalfTileSize,HalfTileSize)))
 	{
 		const float Alpha = ComputeLerpAlpha(Coord, P0.xy, P1.xy, SegmentLenSqRcp);
 		const float Depth = lerp(P0.z, P1.z, Alpha);
@@ -1345,12 +1345,12 @@ void RasterMultiSampleCS(uint DispatchThreadID : SV_DispatchThreadID, uint Group
 		const uint Prim = uint((float(GroupThreadID) + 0.5f) / float(ThreadsPerSeg));
 		const uint PModTPS = GroupThreadID - (Prim * ThreadsPerSeg);
 
-		float4 SP0;
-		float4 SP1;
-		float Rad0;
-		float Rad1;
-		bool bIsEndPoint;
-		uint PrimMatID;
+		float4 SP0 = 0;
+		float4 SP1 = 0;
+		float Rad0 = 0;
+		float Rad1 = 0;
+		bool bIsEndPoint = false;
+		uint PrimMatID = ~0;
 
 		if (bThreadValid)
 		{
diff --git a/Engine/Shaders/Private/HairStrands/HairStrandsVoxelRasterCompute.usf b/Engine/Shaders/Private/HairStrands/HairStrandsVoxelRasterCompute.usf
--- a/Engine/Shaders/Private/HairStrands/HairStrandsVoxelRasterCompute.usf
+++ b/Engine/Shaders/Private/HairStrands/HairStrandsVoxelRasterCompute.usf
@@ -258,7 +258,7 @@ void MarkValid_PrepareCS(uint2 DispatchThreadId : SV_DispatchThreadID)
 void MarkValid_PrepareCS(uint3 DispatchThreadId : SV_DispatchThreadID)
 {
 	const uint4 PageIndexResolutionAndOffset = PageIndexResolutionAndOffsetBuffer.Load(MacroGroupId);
-	const int3 PageIndexResolution = PageIndexResolutionAndOffset.xyz;
+	const uint3 PageIndexResolution = PageIndexResolutionAndOffset.xyz;
 	const uint PageIndexOffset = PageIndexResolutionAndOffset.w;
 
 	const uint3 Coord = DispatchThreadId;
@@ -560,7 +560,7 @@ void FeedbackCS(uint3 DispatchThreadId : SV_DispatchThreadID)
 	// but the voxel size does not match the (more precise) target).
 	// In this case, we add a threshold/margin to to the target, so that there is no oscillation.
 	const float VolumeRatio_Thres = float(GPUAllocatedPageCount) / float(CPUAllocatedPageCount * AdaptiveCorrectionThreshold);
-	const float LinearRatio_Thres = pow(VolumeRatio_Thres, 1.f / 3.f);
+	const float LinearRatio_Thres = pow(max(VolumeRatio_Thres, 0.f), 1.f / 3.f);
 
 	// If the page pool is not large enough increase voxel size
 	float NextVoxelWorldSize = CPUMinVoxelWorldSize;
