Mobile: Pass light parameters as global constants instead of creating a temporary unfiorm buffer for each local light

diff --git a/Engine/Shaders/Private/LightShaderParameters.ush b/Engine/Shaders/Private/LightShaderParameters.ush
--- a/Engine/Shaders/Private/LightShaderParameters.ush
+++ b/Engine/Shaders/Private/LightShaderParameters.ush
@@ -34,6 +34,9 @@ float  Light_RectLightBarnLength;
 float2 Light_RectLightAtlasUVOffset;
 float2 Light_RectLightAtlasUVScale;
 float  Light_RectLightAtlasMaxLevel;
+float  Light_IESAtlasIndex;
+uint   Light_LightFunctionAtlasLightIndex;
+uint   Light_bAffectsTranslucentLighting;
 
 /** Returns the FLightShaderParameters from the root shader parameters. */
 FLightShaderParameters GetRootLightShaderParameters()
@@ -56,5 +59,8 @@ FLightShaderParameters GetRootLightShaderParameters()
 	Out.RectLightAtlasUVOffset	= Light_RectLightAtlasUVOffset;
 	Out.RectLightAtlasUVScale	= Light_RectLightAtlasUVScale;
 	Out.RectLightAtlasMaxLevel	= Light_RectLightAtlasMaxLevel;
+	Out.IESAtlasIndex			= Light_IESAtlasIndex;
+	Out.LightFunctionAtlasLightIndex = Light_LightFunctionAtlasLightIndex;
+	Out.bAffectsTranslucentLighting = Light_bAffectsTranslucentLighting;
 	return Out;
 }
diff --git a/Engine/Shaders/Private/MobileDeferredShading.usf b/Engine/Shaders/Private/MobileDeferredShading.usf
--- a/Engine/Shaders/Private/MobileDeferredShading.usf
+++ b/Engine/Shaders/Private/MobileDeferredShading.usf
@@ -22,6 +22,7 @@
 #include "LightFunctionCommon.ush"
 #endif
 #include "LightDataUniforms.ush"
+#include "LightShaderParameters.ush"
 
 #ifndef USE_HAIR_COMPLEX_TRANSMITTANCE
 #define USE_HAIR_COMPLEX_TRANSMITTANCE 0
@@ -215,6 +216,40 @@ SamplerState	LocalLightShadowSampler;
 float4			LocalLightShadowBufferSize;
 #endif
 
+FDeferredLightData InitDeferredLightFromLightParameters(FLightShaderParameters In, uint InLightType)
+{
+	FDeferredLightData Out = (FDeferredLightData)0;
+	Out.TranslatedWorldPosition	= In.TranslatedWorldPosition;
+	Out.InvRadius				= In.InvRadius;
+	Out.Color					= In.Color;
+	Out.FalloffExponent			= In.FalloffExponent;
+	Out.Direction				= In.Direction;
+	Out.Tangent					= In.Tangent;
+	Out.SpotAngles				= In.SpotAngles;
+	Out.SpecularScale			= In.SpecularScale;
+	Out.SourceRadius			= In.SourceRadius;
+	Out.SoftSourceRadius		= In.SoftSourceRadius;
+	Out.SourceLength			= In.SourceLength;
+	
+	Out.RectLightData.BarnCosAngle				= In.RectLightBarnCosAngle;
+	Out.RectLightData.BarnLength				= In.RectLightBarnLength;
+	Out.RectLightData.AtlasData.AtlasUVOffset	= In.RectLightAtlasUVOffset;
+	Out.RectLightData.AtlasData.AtlasUVScale	= In.RectLightAtlasUVScale;
+	Out.RectLightData.AtlasData.AtlasMaxLevel	= In.RectLightAtlasMaxLevel;
+
+	Out.IESAtlasIndex			= In.IESAtlasIndex;
+	Out.LightFunctionAtlasLightIndex = In.LightFunctionAtlasLightIndex;
+	Out.bAffectsTranslucentLighting = In.bAffectsTranslucentLighting;
+	
+	Out.bInverseSquared			= In.FalloffExponent == 0;
+	Out.bRadialLight			= true;
+	Out.bSpotLight				= InLightType == LIGHT_TYPE_SPOT;
+	Out.bRectLight				= InLightType == LIGHT_TYPE_RECT;
+	Out.HairTransmittance		= InitHairTransmittanceData();
+
+	return Out;
+}
+
 void MobileRadialLightPS(
 	float4 InScreenPosition : TEXCOORD0,
 	float4 SVPos			: SV_POSITION,
@@ -239,11 +274,11 @@ void MobileRadialLightPS(
 	half3 CameraVector = normalize(TranslatedWorldPosition);
 			
 	FLightAccumulator DirectLighting = (FLightAccumulator)0;
-	FDeferredLightData LightData = InitDeferredLightFromUniforms(IS_SPOT_LIGHT ? LIGHT_TYPE_SPOT : LIGHT_TYPE_POINT);
-	
+	FDeferredLightData LightData = InitDeferredLightFromLightParameters(GetRootLightShaderParameters(), IS_SPOT_LIGHT ? LIGHT_TYPE_SPOT : LIGHT_TYPE_POINT);
+		
 	// Affect the light color by any other shadow/transmittance before the lighting is evaluated
 	{
-		half Attenuation = ComputeLightProfileMultiplier(TranslatedWorldPosition, DeferredLightUniforms.TranslatedWorldPosition, -DeferredLightUniforms.Direction, DeferredLightUniforms.Tangent, DeferredLightUniforms.IESAtlasIndex);
+		half Attenuation = ComputeLightProfileMultiplier(TranslatedWorldPosition, LightData.TranslatedWorldPosition, -LightData.Direction, LightData.Tangent, LightData.IESAtlasIndex);
 		LightData.Color *= Attenuation;
 	}
 
diff --git a/Engine/Source/Runtime/Renderer/Private/MobileDeferredShadingPass.cpp b/Engine/Source/Runtime/Renderer/Private/MobileDeferredShadingPass.cpp
--- a/Engine/Source/Runtime/Renderer/Private/MobileDeferredShadingPass.cpp
+++ b/Engine/Source/Runtime/Renderer/Private/MobileDeferredShadingPass.cpp
@@ -191,11 +191,11 @@ public:
 	using FPermutationDomain = TShaderPermutationDomain<FEnableShadingModelSupport, FSpotLightDim, FIESProfileDim, FSpotLightShadowDim>;
 
 	BEGIN_SHADER_PARAMETER_STRUCT(FParameters, )
+		SHADER_PARAMETER_STRUCT(FLightShaderParameters, Light)
+		SHADER_PARAMETER_STRUCT_INCLUDE(FMobileMovableLocalLightShadowParameters, MobileMovableLocalLightShadow)
 		SHADER_PARAMETER(FMatrix44f, TranslatedWorldToLight)
 		SHADER_PARAMETER(FVector4f, LightFunctionParameters)
 		SHADER_PARAMETER(FVector2f, LightFunctionParameters2)
-		SHADER_PARAMETER_STRUCT_REF(FDeferredLightUniformStruct, DeferredLightUniforms)
-		SHADER_PARAMETER_STRUCT_INCLUDE(FMobileMovableLocalLightShadowParameters, MobileMovableLocalLightShadow)
 	END_SHADER_PARAMETER_STRUCT()
 
 	static bool ShouldCompilePermutation(const FMaterialShaderPermutationParameters& Parameters)
@@ -777,7 +777,7 @@ static void RenderLocalLight(
 	FMobileRadialLightFunctionPS::FParameters PassParameters;
 	const bool bShouldCastShadow = LightSceneInfo.SetupMobileMovableLocalLightShadowParameters(View, VisibleLightInfos, PassParameters.MobileMovableLocalLightShadow);
 
-	PassParameters.DeferredLightUniforms = TUniformBufferRef<FDeferredLightUniformStruct>::CreateUniformBufferImmediate(GetDeferredLightParameters(View, LightSceneInfo), EUniformBufferUsage::UniformBuffer_SingleFrame);
+	PassParameters.Light = GetDeferredLightParameters(View, LightSceneInfo).LightParameters;
 	const float TanOuterAngle = bIsSpotLight ? FMath::Tan(LightSceneInfo.Proxy->GetOuterConeAngle()) : 1.0f;
 	PassParameters.LightFunctionParameters = FVector4f(TanOuterAngle, 1.0f /*ShadowFadeFraction*/, bIsSpotLight ? 1.0f : 0.0f, bIsPointLight ? 1.0f : 0.0f);
 	PassParameters.LightFunctionParameters2 = FVector2f(LightSceneInfo.Proxy->GetLightFunctionFadeDistance(), LightSceneInfo.Proxy->GetLightFunctionDisabledBrightness());
@@ -929,8 +929,7 @@ static void RenderSimpleLights(
 
 		// Render light
 		FMobileRadialLightFunctionPS::FParameters PassParameters;
-		FDeferredLightUniformStruct DeferredLightUniformsValue = GetSimpleDeferredLightParameters(View, SimpleLight, SimpleLightPerViewData);
-		PassParameters.DeferredLightUniforms = TUniformBufferRef<FDeferredLightUniformStruct>::CreateUniformBufferImmediate(DeferredLightUniformsValue, EUniformBufferUsage::UniformBuffer_SingleFrame);
+		PassParameters.Light = GetSimpleDeferredLightParameters(View, SimpleLight, SimpleLightPerViewData).LightParameters;
 
 		for (int32 PassIndex = 0; PassIndex < NumPasses; ++PassIndex)
 		{
