Heterogeneous Volumes: Fixing issue when using Local Position in material graph.

diff --git a/Engine/Shaders/Private/HeterogeneousVolumes/HeterogeneousVolumesLiveShadingUtils.ush b/Engine/Shaders/Private/HeterogeneousVolumes/HeterogeneousVolumesLiveShadingUtils.ush
--- a/Engine/Shaders/Private/HeterogeneousVolumes/HeterogeneousVolumesLiveShadingUtils.ush
+++ b/Engine/Shaders/Private/HeterogeneousVolumes/HeterogeneousVolumesLiveShadingUtils.ush
@@ -14,8 +14,8 @@ FPrimitiveSceneData GetPrimitiveData(
 	FPrimitiveSceneData PrimitiveData = GetPrimitiveData(PrimitiveId);
 
 	// Replace local-space transforms..
-	PrimitiveData.LocalToWorld = MakeDFMatrix(PrimitiveData.PositionHigh, LocalToWorld);
-	PrimitiveData.WorldToLocal = MakeDFInverseMatrix(PrimitiveData.PositionHigh, WorldToLocal);
+	PrimitiveData.LocalToWorld = DFPromote(LocalToWorld);
+	PrimitiveData.WorldToLocal = DFPromoteInverse(WorldToLocal);
 	PrimitiveData.LocalObjectBoundsMin = LocalBoundsOrigin - LocalBoundsExtent;
 	PrimitiveData.LocalObjectBoundsMax = LocalBoundsOrigin + LocalBoundsExtent;
 
