Path Tracer: Add support for missing legacy shading models when using Substrate

Subsurface and two-sided foliage legacy shading models were not being handled properly.

Also refactor some of the sheen logic to be common among the various substrate shading models in the path tracer (so that the SheenQuality cvar applies the same to all).

Refactor substrate BSDFs to reduce the size of the temporary state needed during eval/sample. The energy terms E is only needed to setup probabilities, and the weight W is only needed during eval/sample. So we only need to keep the weight W in the temporary state. This optimization can be applied to non-substrate shading models as well, but can be done in a different CL.

diff --git a/Engine/Shaders/Private/PathTracing/Material/PathTracingMaterialSampling.ush b/Engine/Shaders/Private/PathTracing/Material/PathTracingMaterialSampling.ush
--- a/Engine/Shaders/Private/PathTracing/Material/PathTracingMaterialSampling.ush
+++ b/Engine/Shaders/Private/PathTracing/Material/PathTracingMaterialSampling.ush
@@ -13,6 +13,7 @@ PathTracingMaterialSampling.usf: Material sampling functions
 #include "PathTracingGlossy.ush"
 #if SUBSTRATE_ENABLED
 #include "PathTracingSubstrate.ush"
+#include "PathTracingSubstrateFoliage.ush"
 #include "PathTracingSubstrateSolidGlass.ush"
 #include "PathTracingSubstrateThinGlass.ush"
 #else
@@ -57,6 +58,12 @@ float3 EstimateMaterialAlbedo(FPathTracingPayload Payload)
 			Albedo *= Payload.GetMaxLobeWeight();
 			break;
 		}
+		case SHADINGMODELID_TWOSIDED_FOLIAGE:
+		{
+			Albedo = Payload.DiffuseColor + Payload.GetFoliageTransmissionColor() + Payload.SpecularColor + Payload.FuzzAmount * Payload.FuzzColor;
+			Albedo *= Payload.GetMaxLobeWeight();
+			break;
+		}
 		case SHADINGMODELID_THIN_TRANSLUCENT:
 		{
 			Albedo = 1.0 + Payload.FuzzAmount * Payload.FuzzColor;
@@ -162,6 +169,16 @@ void AccumulateAlbedoNormal(inout FPathTracingPayload Payload, float3 PathThroug
 			Normal += NormThroughput * Payload.WorldNormal;
 			break;
 		}
+		case SHADINGMODELID_TWOSIDED_FOLIAGE:
+		{
+			const float Roughness0 = Payload.RoughnessData.x;
+			const float Roughness1 = Payload.RoughnessData.y;
+			const float SpecBlend = Payload.RoughnessData.z;
+			const float RoughnessWeight = lerp(ROUGHNESS_WEIGHT(Roughness0), ROUGHNESS_WEIGHT(Roughness1), SpecBlend);
+			Albedo += PathThroughput * (Payload.DiffuseColor + Payload.GetFoliageTransmissionColor() + Payload.SpecularColor * RoughnessWeight + Payload.FuzzAmount * Payload.FuzzColor);
+			Normal += NormThroughput * Payload.WorldNormal;
+			break;
+		}
 		case SHADINGMODELID_THIN_TRANSLUCENT:
 		{
 			const float F0 = F0RGBToF0(Payload.SpecularColor);
@@ -355,6 +372,14 @@ void AdjustPayloadAfterUnpack(inout FPathTracingPayload HitPayload, bool bApplyO
 			}
 			break;
 		}
+		case SHADINGMODELID_TWOSIDED_FOLIAGE:
+		{
+			if (bApplyOverrides)
+			{
+				HitPayload.SetFoliageTransmissionColor(HitPayload.GetFoliageTransmissionColor() * View.DiffuseOverrideParameter.w + View.DiffuseOverrideParameter.xyz);
+			}
+			break;
+		}
 		case SHADINGMODELID_CLEAR_COAT:
 		{
 			if (bApplyOverrides)
@@ -565,6 +590,8 @@ FMaterialSample SampleMaterial(
 				return SubstrateSolidGlass_SampleMaterial(V_World, Payload, RandSample);
 			else
 				return Substrate_SampleMaterial(V_World, Payload, RandSample);
+		case SHADINGMODELID_TWOSIDED_FOLIAGE:
+			return SubstrateFoliage_SampleMaterial(V_World, Payload, RandSample);
 		case SHADINGMODELID_THIN_TRANSLUCENT:
 			return SubstrateThinGlass_SampleMaterial(V_World, Payload, RandSample);
 		case SHADINGMODELID_CLEAR_COAT:
@@ -621,6 +648,8 @@ FMaterialEval EvalMaterial(
 				return SubstrateSolidGlass_EvalMaterial(V_World, L_World, Payload, DiffuseSpecularScale);
 			else
 				return Substrate_EvalMaterial(V_World, L_World, Payload, DiffuseSpecularScale);
+		case SHADINGMODELID_TWOSIDED_FOLIAGE:
+			return SubstrateFoliage_EvalMaterial(V_World, L_World, Payload, DiffuseSpecularScale);
 		case SHADINGMODELID_THIN_TRANSLUCENT:
 			return SubstrateThinGlass_EvalMaterial(V_World, L_World, Payload, DiffuseSpecularScale);
 		case SHADINGMODELID_CLEAR_COAT:
diff --git a/Engine/Shaders/Private/PathTracing/Material/PathTracingSubstrate.ush b/Engine/Shaders/Private/PathTracing/Material/PathTracingSubstrate.ush
--- a/Engine/Shaders/Private/PathTracing/Material/PathTracingSubstrate.ush
+++ b/Engine/Shaders/Private/PathTracing/Material/PathTracingSubstrate.ush
@@ -6,10 +6,6 @@
 #error "This header should only be included when Substrate is enabled."
 #endif
 
-#ifndef SUBSTRATE_SHEEN_QUALITY
-#define SUBSTRATE_SHEEN_QUALITY 0
-#endif
-
 #include "PathTracingSubstrateCommon.ush"
 #include "PathTracingMaterialCommon.ush"
 #include "PathTracingGlossy.ush"
@@ -29,9 +25,9 @@ struct FSubstrateSingleSlabData
 	float FuzzRoughness;
 
 	float FuzzAttenuation;
-	FBxDFEnergyTermsRGB Spec0;
-	FBxDFEnergyTermsRGB Spec1;
-	FBxDFEnergyTermsA Fuzz;
+	float3 Spec0W;
+	float3 Spec1W;
+	float FuzzAmount;
 
 	float3 LobeCdf;
 	float4 LobePdf;
@@ -51,19 +47,15 @@ struct FSubstrateSingleSlabData
 
 		const float NoV = saturate(V.z);
 		// SUBSTRATE_TODO: Add support for F82 model?
-		Spec0 = ComputeGGXSpecEnergyTermsRGB(Payload.RoughnessData.x, NoV, F0, F90);
-		Spec1 = ComputeGGXSpecEnergyTermsRGB(Payload.RoughnessData.y, NoV, F0, F90);
-		#if SUBSTRATE_SHEEN_QUALITY == 1
-		{
-			Fuzz.E = SheenLTC_DirectionalAlbedo(NoV, FuzzRoughness, View.SheenLTCTexture, View.SheenLTCSampler);
-			Fuzz.W = 1;
-		}
-		#else
-		{
-			Fuzz = ComputeClothEnergyTermsA(FuzzRoughness, NoV);
-		}
-		#endif
-
+		FBxDFEnergyTermsRGB Spec0 = ComputeGGXSpecEnergyTermsRGB(Payload.RoughnessData.x, NoV, F0, F90);
+		FBxDFEnergyTermsRGB Spec1 = ComputeGGXSpecEnergyTermsRGB(Payload.RoughnessData.y, NoV, F0, F90);
+		FBxDFEnergyTermsA Fuzz = ComputeSheenEnergyTerms(NoV, FuzzRoughness);
+
+		// save the parts needed during eval/sampling
+		Spec0W = Spec0.W;
+		Spec1W = Spec1.W;
+		FuzzAmount = Payload.FuzzAmount * Fuzz.W;
+		
 		const float3 SpecE = lerp(Spec0.E, Spec1.E, Payload.RoughnessData.z);
 
 		FuzzAttenuation = 1.0 - Payload.FuzzAmount * Fuzz.E;
@@ -71,7 +63,7 @@ struct FSubstrateSingleSlabData
 
 		const float3 Spec0Albedo = FuzzAttenuation * (1.0 - Payload.RoughnessData.z) * Spec0.E;
 		const float3 Spec1Albedo = FuzzAttenuation * (      Payload.RoughnessData.z) * Spec1.E;
-		const float  FuzzAlbedo = Payload.FuzzAmount * Fuzz.E;
+		const float3 FuzzAlbedo = Payload.FuzzAmount * Fuzz.E * Payload.FuzzColor;
 
 		// Now prepare a cdf/pdf for lobe selection
 		float3 MaxLobeWeight = Payload.GetMaxLobeWeight();
@@ -177,8 +169,8 @@ FMaterialSample Substrate_SampleMaterial(
 		const float3 F = Data.FuzzAttenuation * F_Schlick(Data.F0, Data.F90, VoH) * FTint;
 #endif
 
-		Result.AddLobeWithMIS((1.0 - Payload.RoughnessData.z) * F * GGXResult0.x * Data.Spec0.W, Spec0Pdf, Data.LobePdf.y);
-		Result.AddLobeWithMIS((      Payload.RoughnessData.z) * F * GGXResult1.x * Data.Spec1.W, Spec1Pdf, Data.LobePdf.z);
+		Result.AddLobeWithMIS((1.0 - Payload.RoughnessData.z) * F * GGXResult0.x * Data.Spec0W, Spec0Pdf, Data.LobePdf.y);
+		Result.AddLobeWithMIS((      Payload.RoughnessData.z) * F * GGXResult1.x * Data.Spec1W, Spec1Pdf, Data.LobePdf.z);
 		Result.Pdf += Data.LobePdf.x * DiffPdf;
 		Result.Pdf += Data.LobePdf.w * ClothPdf;
 
@@ -188,21 +180,10 @@ FMaterialSample Substrate_SampleMaterial(
 	{
 		// Cloth Lobe
 		float3 ClothWeight = 0;
-		#if SUBSTRATE_SHEEN_QUALITY == 1
+		if (NoV > 0 && NoL > 0)
 		{
-			// Disney's Cloth Lobe
-			float FuzzDirectionalAlbedo = 1;
-			float3 FuzzSpecularPathValueNoL = Payload.FuzzColor * SheenLTC_Eval(V_World, L_World, N_World, NoV, Data.FuzzRoughness, View.SheenLTCTexture, View.SheenLTCSampler, FuzzDirectionalAlbedo);
-			ClothWeight = NoL > 0 && NoV > 0 ? (Payload.FuzzAmount * FuzzSpecularPathValueNoL) * (2 * PI) : 0.0;
+			ClothWeight = Data.FuzzAmount * ComputeSheenWeight(V_World, L_World, N_World, NoV, NoH, VoH, Data.FuzzRoughness, Payload.FuzzColor);
 		}
-		#else
-		{
-			// Charlie's Cloth Lobe
-			const float DCloth = D_Charlie(Data.FuzzRoughness, NoH);
-			const float3 FCloth = F_Schlick(Payload.FuzzColor, VoH);
-			ClothWeight = NoL > 0 && NoV > 0 ? (NoL * Payload.FuzzAmount * DCloth * Vis_Ashikhmin(NoV, NoL)) * FCloth * Data.Fuzz.W * (2 * PI) : 0.0;
-		}
-		#endif
 		Result.AddLobeWithMIS(ClothWeight * ShadowTerminator, 1 / (2.0 * PI), Data.LobePdf.w);
 		Result.ScatterType = PATHTRACER_SCATTER_SPECULAR;
 
@@ -255,25 +236,16 @@ FMaterialEval Substrate_EvalMaterial(
 #endif
 		SubstrateGlint(Payload, V, L, Data.Alpha0, GGXResult0.y /*Pdf*/, GGXResult0.x /*Weight*/);
 
-		Result.AddLobeWithMIS(DiffuseSpecularScale.y * (1.0 - Payload.RoughnessData.z) * F * GGXResult0.x * Data.Spec0.W, GGXResult0.y, Data.LobePdf.y);
-		Result.AddLobeWithMIS(DiffuseSpecularScale.y * (      Payload.RoughnessData.z) * F * GGXResult1.x * Data.Spec1.W, GGXResult1.y, Data.LobePdf.z);
+		Result.AddLobeWithMIS(DiffuseSpecularScale.y * (1.0 - Payload.RoughnessData.z) * F * GGXResult0.x * Data.Spec0W, GGXResult0.y, Data.LobePdf.y);
+		Result.AddLobeWithMIS(DiffuseSpecularScale.y * (      Payload.RoughnessData.z) * F * GGXResult1.x * Data.Spec1W, GGXResult1.y, Data.LobePdf.z);
 	}
 	{
 		// Cloth Lobe
 		float3 ClothWeight = 0;
-		#if SUBSTRATE_SHEEN_QUALITY == 1
-		{
-			float FuzzDirectionalAlbedo = 1;
-			float3 FuzzSpecularPathValueNoL = Payload.FuzzColor * SheenLTC_Eval(V_World, L_World, N_World, NoV, Data.FuzzRoughness, View.SheenLTCTexture, View.SheenLTCSampler, FuzzDirectionalAlbedo);
-			ClothWeight = NoL > 0 && NoV > 0 ? (Payload.FuzzAmount * FuzzSpecularPathValueNoL) * (2 * PI) : 0.0;
-		}
-		#else
+		if (NoV > 0 && NoL > 0)
 		{
-			const float DCloth = D_Charlie(Data.FuzzRoughness, NoH);
-			const float3 FCloth = F_Schlick(Payload.FuzzColor, VoH);
-			ClothWeight = NoL > 0 && NoV > 0 ? (NoL * Payload.FuzzAmount * DCloth * Vis_Ashikhmin(NoV, NoL)) * FCloth * Data.Fuzz.W * (2 * PI) : 0.0;
+			ClothWeight = Data.FuzzAmount * ComputeSheenWeight(V_World, L_World, N_World, NoV, NoH, VoH, Data.FuzzRoughness, Payload.FuzzColor);
 		}
-		#endif
 		Result.AddLobeWithMIS(DiffuseSpecularScale.y * ClothWeight * ShadowTerminator, 1 / (2.0 * PI), Data.LobePdf.w);
 	}
 
diff --git a/Engine/Shaders/Private/PathTracing/Material/PathTracingSubstrateCommon.ush b/Engine/Shaders/Private/PathTracing/Material/PathTracingSubstrateCommon.ush
--- a/Engine/Shaders/Private/PathTracing/Material/PathTracingSubstrateCommon.ush
+++ b/Engine/Shaders/Private/PathTracing/Material/PathTracingSubstrateCommon.ush
@@ -7,6 +7,7 @@
 #endif
 
 #include "../PathTracingCommon.ush"
+#include "../../BRDF.ush"
 
 #define SUBSTRATE_FRESNEL_F82  0
 
@@ -16,6 +17,10 @@
 #include "/Engine/Private/Substrate/Glint/GlintThirdParty.ush"
 #endif
 
+#ifndef SUBSTRATE_SHEEN_QUALITY
+#define SUBSTRATE_SHEEN_QUALITY 0
+#endif
+
 float3 SubstrateLobeWeight(FPathTracingPayload Payload, float SatNoL)
 {
 	// See reference in LuminanceWeight
@@ -48,4 +53,31 @@ void SubstrateGlint(FPathTracingPayload Payload, float3 InTangentV, float3 InTan
 		OutWeight = Fs / max(0.0001f, InPdf);
 	}
 #endif
-}
\ No newline at end of file
+}
+
+FBxDFEnergyTermsA ComputeSheenEnergyTerms(float NoV, float SheenRoughness)
+{
+#if SUBSTRATE_SHEEN_QUALITY == 1
+	FBxDFEnergyTermsA Result;
+	Result.E = SheenLTC_DirectionalAlbedo(NoV, SheenRoughness, View.SheenLTCTexture, View.SheenLTCSampler);
+	Result.W = 1;
+	return Result;
+#else
+	return ComputeClothEnergyTermsA(SheenRoughness, NoV);
+#endif
+}
+
+float3 ComputeSheenWeight(float3 V, float3 L, float3 N, float NoV, float NoH, float VoH, float SheenRoughness, float3 SheenColor)
+{
+#if SUBSTRATE_SHEEN_QUALITY == 1
+	// Disney's Cloth Lobe
+	float FuzzDirectionalAlbedo = 1;
+	float3 FuzzSpecularPathValueNoL = SheenColor * SheenLTC_Eval(V, L, N, NoV, SheenRoughness, View.SheenLTCTexture, View.SheenLTCSampler, FuzzDirectionalAlbedo);
+	return FuzzSpecularPathValueNoL * (2 * PI);
+#else
+	// Charlie's Cloth Lobe
+	const float DCloth = D_Charlie(SheenRoughness, NoH);
+	const float3 FCloth = F_Schlick(SheenColor, VoH);
+	return NoL * DCloth * Vis_Ashikhmin(NoV, NoL)) * FCloth * (2 * PI);
+#endif
+}
diff --git a/Engine/Shaders/Private/PathTracing/Material/PathTracingSubstrateFoliage.ush b/Engine/Shaders/Private/PathTracing/Material/PathTracingSubstrateFoliage.ush
new file mode 100644
--- /dev/null
+++ b/Engine/Shaders/Private/PathTracing/Material/PathTracingSubstrateFoliage.ush
@@ -0,0 +1,271 @@
+// Copyright Epic Games, Inc. All Rights Reserved.
+
+#pragma once
+
+#ifndef SUBSTRATE_ENABLED
+#error "This header should only be included when Substrate is enabled."
+#endif
+
+#include "PathTracingSubstrateCommon.ush"
+#include "PathTracingMaterialCommon.ush"
+#include "PathTracingGlossy.ush"
+#include "PathTracingFresnel.ush"
+#include "../../BRDF.ush"
+
+struct FSubstrateFoliageSlabData
+{
+	float3x3 Basis;
+	float3 V;
+
+	float  DiffuseWeight;
+	float  TransmissionProb; // for sampling the transmission lobe
+	float3 F0;
+	float3 F90;
+	float2 Alpha0;
+	float2 Alpha1;
+	float FuzzRoughness;
+	
+	float FuzzAttenuation;
+	float3 Spec0W;
+	float3 Spec1W;
+	float FuzzAmount;
+
+	float3 LobeCdf;
+	float4 LobePdf;
+
+	void PrepareBSDF(float3 V_World, FPathTracingPayload Payload)
+	{
+		Basis = GetGGXBasis(Payload.RoughnessData.x, Payload.Anisotropy, Payload.WorldNormal, Payload.WorldTangent, Alpha0);
+		Alpha1 = GetGGXAlpha(Payload.RoughnessData.y, Payload.Anisotropy);
+
+		V = mul(Basis, V_World);
+
+		// make sure F0=0.0 fades off the specular lobe completely
+		F0 = Payload.SpecularColor;
+		F90 = Payload.SpecularEdgeColor * F0RGBToMicroOcclusion(F0);
+
+		FuzzRoughness = MakeRoughnessSafe(Payload.FuzzRoughness, SUBSTRATE_MIN_FUZZ_ROUGHNESS);
+
+		const float NoV = saturate(V.z);
+		// SUBSTRATE_TODO: Add support for F82 model?
+		FBxDFEnergyTermsRGB Spec0 = ComputeGGXSpecEnergyTermsRGB(Payload.RoughnessData.x, NoV, F0, F90);
+		FBxDFEnergyTermsRGB Spec1 = ComputeGGXSpecEnergyTermsRGB(Payload.RoughnessData.y, NoV, F0, F90);
+		FBxDFEnergyTermsA Fuzz = ComputeSheenEnergyTerms(NoV, FuzzRoughness);
+
+		// save the parts needed during eval/sampling
+		Spec0W = Spec0.W;
+		Spec1W = Spec1.W;
+		FuzzAmount = Payload.FuzzAmount * Fuzz.W;
+		
+
+		const float3 SpecE = lerp(Spec0.E, Spec1.E, Payload.RoughnessData.z);
+
+		FuzzAttenuation = 1.0 - Payload.FuzzAmount * Fuzz.E;
+		DiffuseWeight = FuzzAttenuation * (1.0 - Luminance(SpecE));
+
+		const float3 Spec0Albedo = FuzzAttenuation * (1.0 - Payload.RoughnessData.z) * Spec0.E;
+		const float3 Spec1Albedo = FuzzAttenuation * (      Payload.RoughnessData.z) * Spec1.E;
+		const float3 FuzzAlbedo = Payload.FuzzAmount * Fuzz.E * Payload.FuzzColor;
+
+		// Now prepare a cdf/pdf for lobe selection
+		float3 MaxLobeWeight = Payload.GetMaxLobeWeight();
+		LobeCdf = LobeSelectionCdf(
+			MaxLobeWeight * DiffuseWeight * (Payload.GetFoliageTransmissionColor() + Payload.DiffuseColor),
+			MaxLobeWeight * Spec0Albedo,
+			MaxLobeWeight * Spec1Albedo,
+			MaxLobeWeight * FuzzAlbedo);
+		LobePdf = LobeSelectionPdf(LobeCdf);
+		TransmissionProb = LobeSelectionProb(Payload.GetFoliageTransmissionColor(), Payload.DiffuseColor);
+	}
+};
+
+
+FMaterialSample SubstrateFoliage_SampleMaterial(
+	float3 V_World,
+	FPathTracingPayload Payload,
+	float3 RandSample
+)
+{
+	FSubstrateFoliageSlabData Data = (FSubstrateFoliageSlabData)0;
+	Data.PrepareBSDF(V_World, Payload);
+
+
+	float3 L = 0, H = 0, V = Data.V;
+	float OutRoughness = 1;
+
+	const bool bSampledTransmit = RandSample.x < Data.LobeCdf.x * Data.TransmissionProb;
+	const bool bSampledDiffuse  = RandSample.x < Data.LobeCdf.x;
+	const bool bSampledSpecular = RandSample.x < Data.LobeCdf.z;
+	if (bSampledTransmit)
+	{
+		RandSample.x = RescaleRandomNumber(RandSample.x, 0.0, Data.LobeCdf.x * Data.TransmissionProb);
+		// diffuse lobe
+		L = CosineSampleHemisphere(RandSample.xy).xyz;
+		L.z = -L.z;
+	}
+	else if (bSampledDiffuse)
+	{
+		RandSample.x = RescaleRandomNumber(RandSample.x, Data.LobeCdf.x * Data.TransmissionProb, Data.LobeCdf.x);
+		// diffuse lobe
+		L = CosineSampleHemisphere(RandSample.xy).xyz;
+		H = normalize(L + V);
+	}
+	else if (bSampledSpecular)
+	{
+		// specular lobes
+		const bool bUseSpec0 = RandSample.x < Data.LobeCdf.y;
+		if (bUseSpec0)
+		{
+			RandSample.x = RescaleRandomNumber(RandSample.x, Data.LobeCdf.x, Data.LobeCdf.y);
+			OutRoughness = Payload.RoughnessData.x;
+		}
+		else
+		{
+			RandSample.x = RescaleRandomNumber(RandSample.x, Data.LobeCdf.y, Data.LobeCdf.z);
+			OutRoughness = Payload.RoughnessData.y;
+		}
+
+		H = ImportanceSampleVisibleGGX(RandSample.xy, bUseSpec0 ? Data.Alpha0 : Data.Alpha1, V).xyz;
+		L = reflect(-V, H);
+		if (L.z <= 0)
+		{
+			// invalid output direction, exit early
+			return NullMaterialSample();
+		}
+	}
+	else
+	{
+		// cloth lobe
+		RandSample.x = RescaleRandomNumber(RandSample.x, Data.LobeCdf.z, 1.0);
+		L = UniformSampleHemisphere(RandSample.xy).xyz;
+		H = normalize(L + V);
+	}
+
+	// transform to world space
+	const float3 L_World = normalize(mul(L, Data.Basis));
+	const float3 N_World = Payload.WorldNormal;
+	const float NoV = saturate(V.z);
+	const float NoL = abs(L.z);
+	const float VoH = saturate(dot(V, H));
+	const float NoH = saturate(H.z);
+
+	FMaterialSample Result = CreateMaterialSample(L_World, 0.0, 0.0, 1.0, OutRoughness, PATHTRACER_SCATTER_DIFFUSE);
+
+
+
+	const float DiffPdf = NoL / PI;
+	float2 GGXResult0 = GGXEvalReflection(L, V, H, Data.Alpha0);
+	const float2 GGXResult1 = GGXEvalReflection(L, V, H, Data.Alpha1);
+	const float Spec0Pdf = GGXResult0.y;
+	const float Spec1Pdf = GGXResult1.y;
+	const float ClothPdf = 1 / (2 * PI);
+
+	const float ShadowTerminator = ShadowTerminatorTerm(L_World, N_World, Payload.WorldSmoothNormal);
+	if (bSampledTransmit)
+	{
+		Result.AddLobeWithMIS(Data.DiffuseWeight * Payload.GetFoliageTransmissionColor(), abs(L.z) / PI, Data.LobePdf.x * Data.TransmissionProb);
+	}
+	else if (bSampledDiffuse)
+	{
+		const float3 Diffuse = Data.DiffuseWeight * GetPathTracingDiffuseModel(Payload.DiffuseColor, Payload.RoughnessData.x, NoV, NoL, VoH, NoH);
+
+		Result.AddLobeWithMIS(Diffuse * ShadowTerminator, NoL / PI, Data.LobePdf.x - Data.TransmissionProb * Data.LobePdf.x);
+
+		Result.Pdf += Data.LobePdf.y * Spec0Pdf;
+		Result.Pdf += Data.LobePdf.z * Spec1Pdf;
+		Result.Pdf += Data.LobePdf.w * ClothPdf;
+	}
+	else if (bSampledSpecular)
+	{
+		const float3 FTint = SubstrateSpecularTint(Payload, NoV, NoL, VoH, NoH);
+#if SUBSTRATE_FRESNEL_F82
+		// SUBSTRATE_TODO: Need to tune the logic for specular micro-occlusion ...
+		const float3 F = Data.FuzzAttenuation * F_AdobeF82(Data.F0, Data.F90, VoH) * FTint;
+#else
+		const float3 F = Data.FuzzAttenuation * F_Schlick(Data.F0, Data.F90, VoH) * FTint;
+#endif
+
+		Result.AddLobeWithMIS((1.0 - Payload.RoughnessData.z) * F * GGXResult0.x * Data.Spec0W, Spec0Pdf, Data.LobePdf.y);
+		Result.AddLobeWithMIS((      Payload.RoughnessData.z) * F * GGXResult1.x * Data.Spec1W, Spec1Pdf, Data.LobePdf.z);
+		Result.Pdf += Data.LobePdf.x * DiffPdf;
+		Result.Pdf += Data.LobePdf.w * ClothPdf;
+
+		Result.ScatterType = PATHTRACER_SCATTER_SPECULAR;
+	}
+	else
+	{
+		// Cloth Lobe
+		float3 ClothWeight = Data.FuzzAmount * ComputeSheenWeight(V_World, L_World, N_World, NoV, NoH, VoH, Data.FuzzRoughness, Payload.FuzzColor);
+		Result.AddLobeWithMIS(ClothWeight * ShadowTerminator, 1 / (2.0 * PI), Data.LobePdf.w);
+		Result.ScatterType = PATHTRACER_SCATTER_SPECULAR;
+
+		Result.Pdf += Data.LobePdf.x * DiffPdf;
+		Result.Pdf += Data.LobePdf.y * Spec0Pdf;
+		Result.Pdf += Data.LobePdf.z * Spec1Pdf;
+	}
+
+	Result.Weight *= SubstrateLobeWeight(Payload, NoL);
+
+	return Result;
+}
+
+FMaterialEval SubstrateFoliage_EvalMaterial(
+	float3 V_World,
+	float3 L_World,
+	FPathTracingPayload Payload,
+	float2 DiffuseSpecularScale
+)
+{
+	FSubstrateFoliageSlabData Data = (FSubstrateFoliageSlabData)0;
+	Data.PrepareBSDF(V_World, Payload);
+
+	const float3 N_World = Payload.WorldNormal;
+
+	// move vectors into right shading frame
+	const float3 V = Data.V;
+	const float3 L = mul(Data.Basis, L_World);
+	const float3 H = normalize(V + L);
+	const float NoV = saturate(V.z);
+	const float NoL = abs(L.z);
+	const float VoH = saturate(dot(V, H));
+	const float NoH = saturate(H.z);
+
+	FMaterialEval Result = NullMaterialEval();
+
+	const float ShadowTerminator = ShadowTerminatorTerm(L_World, N_World, Payload.WorldSmoothNormal);
+	if (L.z < 0)
+	{
+		const float3 Diffuse = Data.DiffuseWeight * Payload.GetFoliageTransmissionColor();
+		Result.AddLobeWithMIS(DiffuseSpecularScale.x * Diffuse, NoL / PI, Data.LobePdf.x * Data.TransmissionProb);
+	}
+	else if (L.z > 0)
+	{
+		{
+			const float3 Diffuse = Data.DiffuseWeight * GetPathTracingDiffuseModel(Payload.DiffuseColor, Payload.RoughnessData.x, NoV, NoL, VoH, NoH);
+			Result.AddLobeWithMIS(DiffuseSpecularScale.x * Diffuse * ShadowTerminator, NoL / PI, Data.LobePdf.x - Data.LobePdf.x * Data.TransmissionProb);
+		}
+
+		{
+			float2 GGXResult0 = GGXEvalReflection(L, V, H, Data.Alpha0);
+			const float2 GGXResult1 = GGXEvalReflection(L, V, H, Data.Alpha1);
+			const float3 FTint = SubstrateSpecularTint(Payload, NoV, NoL, VoH, NoH);
+			#if SUBSTRATE_FRESNEL_F82
+			const float3 F = Data.FuzzAttenuation * F_AdobeF82(Data.F0, Data.F90, VoH) * FTint;
+			#else
+			const float3 F = Data.FuzzAttenuation * F_Schlick(Data.F0, Data.F90, VoH) * FTint;
+			#endif
+
+			Result.AddLobeWithMIS(DiffuseSpecularScale.y * (1.0 - Payload.RoughnessData.z) * F * GGXResult0.x * Data.Spec0W, GGXResult0.y, Data.LobePdf.y);
+			Result.AddLobeWithMIS(DiffuseSpecularScale.y * (      Payload.RoughnessData.z) * F * GGXResult1.x * Data.Spec1W, GGXResult1.y, Data.LobePdf.z);
+		}
+		{
+			// Cloth Lobe
+			float3 ClothWeight = Data.FuzzAmount * ComputeSheenWeight(V_World, L_World, N_World, NoV, NoH, VoH, Data.FuzzRoughness, Payload.FuzzColor);
+			Result.AddLobeWithMIS(DiffuseSpecularScale.y * ClothWeight * ShadowTerminator, 1 / (2.0 * PI), Data.LobePdf.w);
+		}
+	}
+
+	Result.Weight *= SubstrateLobeWeight(Payload, NoL);
+
+	return Result;
+}
\ No newline at end of file
diff --git a/Engine/Shaders/Private/PathTracing/Material/PathTracingSubstrateSolidGlass.ush b/Engine/Shaders/Private/PathTracing/Material/PathTracingSubstrateSolidGlass.ush
--- a/Engine/Shaders/Private/PathTracing/Material/PathTracingSubstrateSolidGlass.ush
+++ b/Engine/Shaders/Private/PathTracing/Material/PathTracingSubstrateSolidGlass.ush
@@ -27,7 +27,7 @@ struct FSubstrateSolidGlassSlabData
 	float  FuzzAttenuation;
 	float  EGlass0;
 	float  EGlass1;
-	FBxDFEnergyTermsA Fuzz;
+	float FuzzAmount;
 
 	float2 LobeCdf;
 	float3 LobePdf;
@@ -51,9 +51,10 @@ struct FSubstrateSolidGlassSlabData
 		// correct for energy loss by scaling the whole BSDF
 		EGlass0 = GGXEnergyLookup(Payload.RoughnessData.x, NoV, Eta);
 		EGlass1 = GGXEnergyLookup(Payload.RoughnessData.y, NoV, Eta);
-		Fuzz = ComputeClothEnergyTermsA(FuzzRoughness, NoV);
+		FBxDFEnergyTermsA Fuzz = ComputeClothEnergyTermsA(FuzzRoughness, NoV);
 
 		FuzzAttenuation = 1.0 - Payload.FuzzAmount * Fuzz.E;
+		FuzzAmount = Payload.FuzzAmount * Fuzz.W;
 
 		const float3 Glass0Albedo = FuzzAttenuation * (1.0 - Payload.RoughnessData.z);
 		const float3 Glass1Albedo = FuzzAttenuation * (      Payload.RoughnessData.z);
@@ -168,7 +169,11 @@ FMaterialSample SubstrateSolidGlass_SampleMaterial(
 		const float ShadowTerminator = ShadowTerminatorTerm(L_World, N_World, Payload.WorldSmoothNormal);
 		const float DCloth = D_Charlie(Data.FuzzRoughness, NoH);
 		const float3 FCloth = F_Schlick(Payload.FuzzColor, VoH);
-		const float3 ClothWeight = NoL > 0 && NoV > 0 ? (NoL * Payload.FuzzAmount * DCloth * Vis_Ashikhmin(NoV, NoL)) * FCloth * Data.Fuzz.W * (2 * PI) : 0.0;
+		float3 ClothWeight = 0;
+		if (NoV > 0 && NoL > 0)
+		{
+			ClothWeight = Data.FuzzAmount * ComputeSheenWeight(V_World, L_World, N_World, NoV, NoH, VoH, Data.FuzzRoughness, Payload.FuzzColor);
+		}
 		Result.AddLobeWithMIS(ClothWeight * ShadowTerminator, 1 / (2.0 * PI), Data.LobePdf.z);
 		Result.ScatterType = PATHTRACER_SCATTER_SPECULAR;
 
@@ -222,9 +227,11 @@ FMaterialEval SubstrateSolidGlass_EvalMaterial(
 		{
 			// Cloth Lobe
 			const float ShadowTerminator = ShadowTerminatorTerm(L_World, N_World, Payload.WorldSmoothNormal);
-			const float DCloth = D_Charlie(Data.FuzzRoughness, NoH);
-			const float3 FCloth = F_Schlick(Payload.FuzzColor, VoH);
-			const float3 ClothWeight = NoL > 0 && NoV > 0 ? (NoL * Payload.FuzzAmount * DCloth * Vis_Ashikhmin(NoV, NoL)) * FCloth * Data.Fuzz.W * (2 * PI) : 0.0;
+			float3 ClothWeight = 0;
+			if (NoV > 0 && NoL > 0)
+			{
+				ClothWeight = Data.FuzzAmount * ComputeSheenWeight(V_World, L_World, N_World, NoV, NoH, VoH, Data.FuzzRoughness, Payload.FuzzColor);
+			}
 			Result.AddLobeWithMIS(DiffuseSpecularScale.y * ClothWeight * ShadowTerminator, 1 / (2.0 * PI), Data.LobePdf.z);
 		}
 
diff --git a/Engine/Shaders/Private/PathTracing/Material/PathTracingSubstrateThinGlass.ush b/Engine/Shaders/Private/PathTracing/Material/PathTracingSubstrateThinGlass.ush
--- a/Engine/Shaders/Private/PathTracing/Material/PathTracingSubstrateThinGlass.ush
+++ b/Engine/Shaders/Private/PathTracing/Material/PathTracingSubstrateThinGlass.ush
@@ -25,13 +25,13 @@ struct FSubstrateThinGlassSlabData
 	float3 TransmissionColor;
 	float2 AlphaR0, AlphaT0;
 	float2 AlphaR1, AlphaT1;
-	float  FuzzRoughness;
 
+	float  FuzzRoughness;
 	float  FuzzAttenuation;
-	float2  NudgeE;
-	FBxDFEnergyTermsA SpecR0, SpecR1;
-	FBxDFEnergyTermsA SpecT0, SpecT1;
-	FBxDFEnergyTermsA Fuzz;
+	float  FuzzAmount;
+
+	float SpecR0W, SpecR1W;
+	float SpecT0W, SpecT1W;
 
 	float2 LobeCdf;
 	float3 LobePdf;
@@ -56,14 +56,15 @@ struct FSubstrateThinGlassSlabData
 		FuzzRoughness = MakeRoughnessSafe(Payload.FuzzRoughness, SUBSTRATE_MIN_FUZZ_ROUGHNESS);
 
 		const float NoV = saturate(V.z);
-		SpecR0 = ComputeGGXSpecEnergyTermsA(RoughnessR.x, NoV, F0, F90);
-		SpecR1 = ComputeGGXSpecEnergyTermsA(RoughnessR.y, NoV, F0, F90);
-		SpecT0 = ComputeGGXSpecEnergyTermsA(RoughnessT.x, NoV, 1.0 - F0, 1.0 - F90);
-		SpecT1 = ComputeGGXSpecEnergyTermsA(RoughnessT.y, NoV, 1.0 - F0, 1.0 - F90);
-		NudgeE.x = rcp(SpecR0.E + SpecT0.E);
-		NudgeE.y = rcp(SpecR1.E + SpecT1.E);
+		FBxDFEnergyTermsA SpecR0 = ComputeGGXSpecEnergyTermsA(RoughnessR.x, NoV, F0, F90);
+		FBxDFEnergyTermsA SpecR1 = ComputeGGXSpecEnergyTermsA(RoughnessR.y, NoV, F0, F90);
+		FBxDFEnergyTermsA SpecT0 = ComputeGGXSpecEnergyTermsA(RoughnessT.x, NoV, 1.0 - F0, 1.0 - F90);
+		FBxDFEnergyTermsA SpecT1 = ComputeGGXSpecEnergyTermsA(RoughnessT.y, NoV, 1.0 - F0, 1.0 - F90);
+		float2 NudgeE = float2(
+			rcp(SpecR0.E + SpecT0.E),
+			rcp(SpecR1.E + SpecT1.E));
 
-		Fuzz = ComputeClothEnergyTermsA(FuzzRoughness, NoV);
+		FBxDFEnergyTermsA Fuzz = ComputeClothEnergyTermsA(FuzzRoughness, NoV);
 
 		FuzzAttenuation = 1.0 - Payload.FuzzAmount * Fuzz.E;
 
@@ -71,6 +72,11 @@ struct FSubstrateThinGlassSlabData
 		const float3 GlassT = FuzzAttenuation * lerp(SpecT0.E * NudgeE.x, SpecT1.E * NudgeE.y, Payload.RoughnessData.z);
 		const float  FuzzAlbedo = Payload.FuzzAmount * Fuzz.E;
 
+		SpecR0W = SpecR0.W * NudgeE.x;
+		SpecR1W = SpecR1.W * NudgeE.y;
+		SpecT0W = SpecT0.W * NudgeE.x;
+		SpecT1W = SpecT1.W * NudgeE.y;
+
 		// Now prepare a cdf/pdf for lobe selection
 		// NOTE: we put T first so that the reflected lobe (R and fuzz) are consecutive in the PDF
 		float3 MaxLobeWeight = Payload.GetMaxLobeWeight();
@@ -158,8 +164,8 @@ FMaterialSample SubstrateThinGlass_SampleMaterial(
 		SlabResult.Reflected *= SubstrateSpecularTint(Payload, NoV, NoL, VoH, NoH);
 
 		const float3 SlabTint = bIsReflection ? SlabResult.Reflected : SlabResult.Transmitted;
-		const float SlabW0 = (bIsReflection ? Data.SpecR0.W : Data.SpecT0.W) * Data.NudgeE.x;
-		const float SlabW1 = (bIsReflection ? Data.SpecR1.W : Data.SpecT1.W) * Data.NudgeE.y;
+		const float SlabW0 = (bIsReflection ? Data.SpecR0W : Data.SpecT0W);
+		const float SlabW1 = (bIsReflection ? Data.SpecR1W : Data.SpecT1W);
 
 
 		Result.AddLobeWithMIS(Data.FuzzAttenuation * (1 - Payload.RoughnessData.z) * GGXResult0.x * SlabW0 * SlabTint, GGXResult0.y, LobePdf * (1 - Payload.RoughnessData.z));
@@ -174,9 +180,11 @@ FMaterialSample SubstrateThinGlass_SampleMaterial(
 	{
 		// Cloth Lobe
 		const float ShadowTerminator = ShadowTerminatorTerm(L_World, N_World, Payload.WorldSmoothNormal);
-		const float DCloth = D_Charlie(Data.FuzzRoughness, NoH);
-		const float3 FCloth = F_Schlick(Payload.FuzzColor, VoH);
-		const float3 ClothWeight = NoL > 0 && NoV > 0 ? (NoL * Payload.FuzzAmount * DCloth * Vis_Ashikhmin(NoV, NoL)) * FCloth * Data.Fuzz.W * (2 * PI) : 0.0;
+		float3 ClothWeight = 0;
+		if (NoV > 0 && NoL > 0)
+		{
+			ClothWeight = Data.FuzzAmount * ComputeSheenWeight(V_World, L_World, N_World, NoV, NoH, VoH, Data.FuzzRoughness, Payload.FuzzColor);
+		}
 		Result.AddLobeWithMIS(ClothWeight * ShadowTerminator, ClothPdf, Data.LobePdf.z);
 
 		Result.Pdf += LobePdf * (1.0 - Payload.RoughnessData.z) * GGXResult0.y;
@@ -229,8 +237,8 @@ FMaterialEval SubstrateThinGlass_EvalMaterial(
 	SlabResult.Reflected *= SubstrateSpecularTint(Payload, NoV, NoL, VoH, NoH);
 
 	const float3 SlabTint = bIsReflection ? SlabResult.Reflected : SlabResult.Transmitted;
-	const float SlabW0 = (bIsReflection ? Data.SpecR0.W : Data.SpecT0.W) * Data.NudgeE.x;
-	const float SlabW1 = (bIsReflection ? Data.SpecR1.W : Data.SpecT1.W) * Data.NudgeE.y;
+	const float SlabW0 = (bIsReflection ? Data.SpecR0W : Data.SpecT0W);
+	const float SlabW1 = (bIsReflection ? Data.SpecR1W : Data.SpecT1W);
 
 	const float LobePdf = bIsReflection ? Data.LobePdf.y : Data.LobePdf.x;
 
@@ -245,9 +253,11 @@ FMaterialEval SubstrateThinGlass_EvalMaterial(
 		{
 			// Cloth Lobe
 			const float ShadowTerminator = ShadowTerminatorTerm(L_World, N_World, Payload.WorldSmoothNormal);
-			const float DCloth = D_Charlie(Data.FuzzRoughness, NoH);
-			const float3 FCloth = F_Schlick(Payload.FuzzColor, VoH);
-			const float3 ClothWeight = NoL > 0 && NoV > 0 ? (NoL * Payload.FuzzAmount * DCloth * Vis_Ashikhmin(NoV, NoL)) * FCloth * Data.Fuzz.W * (2 * PI) : 0.0;
+			float3 ClothWeight = 0;
+			if (NoV > 0 && NoL > 0)
+			{
+				ClothWeight = Data.FuzzAmount * ComputeSheenWeight(V_World, L_World, N_World, NoV, NoH, VoH, Data.FuzzRoughness, Payload.FuzzColor);
+			}
 			Result.AddLobeWithMIS(DiffuseSpecularScale.y * ClothWeight * ShadowTerminator, 1 / (2.0 * PI), Data.LobePdf.z);
 		}
 		Result.Weight *= SubstrateLobeWeight(Payload, NoL);
diff --git a/Engine/Shaders/Private/PathTracing/PathTracingCommon.ush b/Engine/Shaders/Private/PathTracing/PathTracingCommon.ush
--- a/Engine/Shaders/Private/PathTracing/PathTracingCommon.ush
+++ b/Engine/Shaders/Private/PathTracing/PathTracingCommon.ush
@@ -553,6 +553,10 @@ struct FPathTracingPayload : FMinimalPayload
 	// SHADINGMODELID_THIN_TRANSLUCENT
 	float3 GetTransmittanceColor() { return exp(-GetExtinction() * SUBSTRATE_SIMPLEVOLUME_THICKNESS_CM); }
 
+	// SHADINGMODELID_TWOSIDED_FOLIAGE
+	float3 GetFoliageTransmissionColor() { return MeanFreePath; }
+	void SetFoliageTransmissionColor(float3 C) { MeanFreePath = C; }
+
 
 #else // PATHTRACING_SUBSTRATE_PAYLOAD
 	// Various ways to interpret CustomData (depending on ShadingModelID)
@@ -578,7 +582,7 @@ struct FPathTracingPayload : FMinimalPayload
 	float GetSubsurfacePhaseFunction() { return CustomData0.a; }
 	void SetSubsurfacePhaseFunction(float G) { CustomData0.a = G; }
 
-	// SHADINGMODELID_SUBSURFACE, SHADINGMODELID_TWO_SIDED_FOLIAGE (note that SHADINGMODELID_SUBSURFACE_PROFILE and SHADINGMODELID_EYE don't use this as they use DiffuseColor directly)
+	// SHADINGMODELID_SUBSURFACE, SHADINGMODELID_TWOSIDED_FOLIAGE (note that SHADINGMODELID_SUBSURFACE_PROFILE and SHADINGMODELID_EYE don't use this as they use DiffuseColor directly)
 	float3 GetSubsurfaceColor() { return CustomData1.xyz; }
 	void SetSubsurfaceColor(float3 SubsurfaceColor) { CustomData1.xyz = SubsurfaceColor; }
 	
diff --git a/Engine/Shaders/Private/PathTracing/PathTracingMaterialHitShader.usf b/Engine/Shaders/Private/PathTracing/PathTracingMaterialHitShader.usf
--- a/Engine/Shaders/Private/PathTracing/PathTracingMaterialHitShader.usf
+++ b/Engine/Shaders/Private/PathTracing/PathTracingMaterialHitShader.usf
@@ -580,6 +580,7 @@ RAY_TRACING_ENTRY_CLOSEST_HIT(PathTracingMaterialCHS,
 					}
 					if (BSDF_GETSSSTYPE(BSDF) == SSS_TYPE_SIMPLEVOLUME)
 					{
+						// TODO: this case should be represented by a different BSDF
 						Payload.DiffuseColor = GetSimpleVolumeDiffuseColor(Payload.DiffuseColor, SLAB_SSSMFP(BSDF));
 					}
 					else if (BSDF_GETSSSTYPE(BSDF) == SSS_TYPE_DIFFUSION)
@@ -599,9 +600,47 @@ RAY_TRACING_ENTRY_CLOSEST_HIT(PathTracingMaterialCHS,
 						Payload.PhaseG = 0.0;
 #endif
 					}
+					else if (BSDF_GETSSSTYPE(BSDF) == SSS_TYPE_WRAP)
+					{
+						const bool bIsThin = BSDF_GETISTHIN(BSDF);
+						const float TransmittanceNoL = 1.0f;
+						const float3 SlabDiffuseColor = bIsThin ? Payload.DiffuseColor : float3(1, 1, 1);
+						const FParticipatingMedia PM = SubstrateSlabCreateParticipatingMedia(Payload.DiffuseColor, SLAB_SSSMFP(BSDF));
+						float3 SubsurfaceColor = IsotropicMediumSlabTransmittance(PM, SUBSTRATE_SIMPLEVOLUME_THICKNESS_M, TransmittanceNoL);
+						const float Opacity = 1.f - abs(SLAB_SSSPHASEANISOTROPY(BSDF));
+
+						float3 Front = Payload.DiffuseColor;
+						float3 Back = SubsurfaceColor;
+						float3 Sum = Front + Back;
+						float S = max(Sum.x, max(Sum.y, Sum.z));
+						if (S > 1)
+						{
+							// nudge the material back to something physically plausible
+							// NOTE: we ignore spec here since it should be accounted for by the brdf model itself
+							Payload.DiffuseColor = Front / S;
+							SubsurfaceColor = Back / S;
+						}
+						if (bIsThin)
+						{
+							// foliage case, only the color matters
+							Payload.ShadingModelID = SHADINGMODELID_TWOSIDED_FOLIAGE;
+							Payload.SetFoliageTransmissionColor(SubsurfaceColor);
+						}
+						else
+						{
+
+							// Match formula from non-Substrate handling of wrap SSS
+							const float SSSRadius = (1 - Opacity) * 10.0; // simple formula, up to 10cm radius
+							// The legacy behavior has both diffuse+sss responses separately, but we can combine them to avoid needing a custom lobe
+							// In practice the diffuse and SSS responses are similar, so the difference should be minimal (and is an approximation of the raster behavior anyway)
+							Payload.DiffuseColor += SubsurfaceColor;
+							Payload.MeanFreePath = SSSRadius;
+							Payload.PhaseG = 0;
+						}
+					}
 					else
 					{
-						// SUBSTRATE_TODO: Do something meaningful for the other SSS types ...
+						// TODO: Are there any other SSS types to consider?
 						Payload.MeanFreePath = 0;
 						Payload.PhaseG       = 0;
 					}
