Fix for moving LFVs when distance from orogin.
Now using translated world position.

diff --git a/Engine/Shaders/Private/LocalFogVolumes/LocalFogVolumeCommon.ush b/Engine/Shaders/Private/LocalFogVolumes/LocalFogVolumeCommon.ush
--- a/Engine/Shaders/Private/LocalFogVolumes/LocalFogVolumeCommon.ush
+++ b/Engine/Shaders/Private/LocalFogVolumes/LocalFogVolumeCommon.ush
@@ -13,19 +13,19 @@
 
 struct FLocalFogVolumeGPUInstanceData
 {
-	float4x3 InvTransform;
-
-	float RadialFogExtinction;
-	float HeightFogExtinction;
-	float HeightFogFalloff;
-	float HeightFogOffset;
-	float3 Emissive;
-
-	float  UniformScale;
-	float  UniformScaleInv;		// Derived from UniformScale
-
-	float3 Albedo;
-	float  PhaseG;
+	float3x3	PreTranslatedInvTransform;
+	float3		TranslatedWorlPos;
+	float		UniformScale;
+	float		UniformScaleInv;		// Derived from UniformScale
+
+	float		RadialFogExtinction;
+	float		HeightFogExtinction;
+	float		HeightFogFalloff;
+	float		HeightFogOffset;
+
+	float3		Emissive;
+	float3		Albedo;
+	float		PhaseG;
 };
 
 
@@ -89,12 +89,12 @@ FLocalFogVolumeGPUInstanceData GetLocalFogVolumeGPUInstanceData(uint Index)
 	float4 Data0 = LFVInsts[Offset + 0];
 	float4 Data1 = LFVInsts[Offset + 1];
 	float4 Data2 = LFVInsts[Offset + 2];
-	// See GetLocalFogVolumeSortingData for the packing mathod.
+	// See GetLocalFogVolumeViewSortingData for the packing mathod.
 
 	// Data0
-	float3 Trans		= Data0.xyz;
-	Data.UniformScale	= Data0.w;
-	Data.UniformScaleInv= 1.0f / Data.UniformScale;
+	Data.TranslatedWorlPos	= Data0.xyz;
+	Data.UniformScale		= Data0.w;
+	Data.UniformScaleInv	= 1.0f / Data.UniformScale;
 
 	// Data1
 	float3 XVec;
@@ -107,7 +107,7 @@ FLocalFogVolumeGPUInstanceData GetLocalFogVolumeGPUInstanceData(uint Index)
 
 	float3 ZVec = cross(XVec, YVec);
 
-	Data.InvTransform = float4x3(float3(XVec * Data.UniformScaleInv), float3(YVec * Data.UniformScaleInv), float3(ZVec * Data.UniformScaleInv), Trans);
+	Data.PreTranslatedInvTransform = float3x3(float3(XVec * Data.UniformScaleInv), float3(YVec * Data.UniformScaleInv), float3(ZVec * Data.UniformScaleInv));
 
 	// Data2
 	float3 Data2X			= UnpackR11G11B10F(asuint(Data2.x));
@@ -332,9 +332,9 @@ float3 LocalFogVolumeEvaluateInScattering(in FLocalFogVolumeGPUInstanceData FogI
 // Returns color transmittance
 float4 GetLFVInstanceContribution(
 	in uint InstanceIndex,
-	in float3 CamRayWorldOrigin,
-	in float3 CamRayWorldDir,
-	in float3 DepthBufferWorldPos,
+	in float3 CamRayTranslatedWorldOrigin,
+	in float3 CamRayTranslatedWorldDir,
+	in float3 DepthBufferTranslatedWorldPos,
 	in bool bPlatformSupportsVolumetricFogOntranslucent)
 {
 	FLocalFogVolumeGPUInstanceData FogInstance = GetLocalFogVolumeGPUInstanceData(InstanceIndex);
@@ -342,11 +342,11 @@ float4 GetLFVInstanceContribution(
 	float3 LuminanceColor = float3(0.0f, 0.0f, 0.0f);
 	float Transmittance = 1.0f;
 
-	// The "U" prefix is for all compuation done in the Unit Sphere space.
-	float3 DepthBufferPosU	= mul(float4(DepthBufferWorldPos, 1.0f),	FogInstance.InvTransform).xyz;
-	float3 RayPosU			= mul(float4(CamRayWorldOrigin,   1.0f),	FogInstance.InvTransform).xyz;
-	float3 RayDirU			= mul(float4(CamRayWorldDir,      0.0f),	FogInstance.InvTransform).xyz;
-	RayDirU.xyz = normalize(RayDirU);
+	// The "U" prefix is for all computation done in the Unit Sphere space.
+	float3 DepthBufferPosU	= mul(DepthBufferTranslatedWorldPos - FogInstance.TranslatedWorlPos,	FogInstance.PreTranslatedInvTransform);
+	float3 RayPosU			= mul(CamRayTranslatedWorldOrigin - FogInstance.TranslatedWorlPos,		FogInstance.PreTranslatedInvTransform);
+	float3 RayDirU			= mul(CamRayTranslatedWorldDir,											FogInstance.PreTranslatedInvTransform);
+	RayDirU.xyz				= normalize(RayDirU);
 
 	float2 TsU = RayIntersectSphere(RayPosU, RayDirU, float4(0.0, 0.0, 0.0, 1.0));
 
@@ -365,7 +365,7 @@ float4 GetLFVInstanceContribution(
 		if (LFVRenderInVolumetricFog && bPlatformSupportsVolumetricFogOntranslucent)
 		{
 			// Volumetric fog covers up to MaxDistance along ViewZ, exclude analytical fog from this range
-			float CosAngle = dot(CamRayWorldDir, View.ViewForward);
+			float CosAngle = dot(CamRayTranslatedWorldDir, View.ViewForward);
 			float InvCosAngle = (CosAngle > 0.0001) ? rcp(CosAngle) : 0;
 			float ExcludeDistance = View.VolumetricFogMaxDistance * InvCosAngle * FogInstance.UniformScaleInv;
 			TsU = max(TsU, ExcludeDistance);
@@ -377,7 +377,7 @@ float4 GetLFVInstanceContribution(
 		{
 			FFogData FogData = LocalFogVolumeEvaluateAnalyticalIntegral(FogInstance, RayPosU + RayDirU * TsU.x, RayDirU, RayTracedLengthU);
 
-			LuminanceColor = LocalFogVolumeEvaluateInScattering(FogInstance, FogData, CamRayWorldDir);
+			LuminanceColor = LocalFogVolumeEvaluateInScattering(FogInstance, FogData, CamRayTranslatedWorldDir);
 			Transmittance = 1.0 - FogData.Coverage;
 		}
 	}
@@ -396,11 +396,8 @@ float4 GetLFVContribution(
 	in bool bPlatformSupportsVolumetricFogOntranslucent,
 	inout uint OutLFVCount)
 {
-	const float3 PreViewTranslation			= DFHackToFloat(ResolvedView.PreViewTranslation);
 	const float3 CamRayTranslatedWorldOrigin= ResolvedView.TranslatedWorldCameraOrigin;
-	const float3 CamRayWorldOrigin			= CamRayTranslatedWorldOrigin - PreViewTranslation;
 	const float3 CamRayWorldDir				= normalize(DepthBufferTranslatedWorldPos - CamRayTranslatedWorldOrigin);
-	const float3 DepthBufferWorldPos		= DepthBufferTranslatedWorldPos - PreViewTranslation;
 
 	OutLFVCount = LFVTileDataTex[uint3(TilePos, 0)];
 	const float4 DefaultColorTransmittance = float4(0.0, 0.0, 0.0, 1.0);
@@ -411,9 +408,9 @@ float4 GetLFVContribution(
 
 		float4 LFVLuminanceColorTransmittance = GetLFVInstanceContribution(
 			InstanceIndex,
-			CamRayWorldOrigin,
+			CamRayTranslatedWorldOrigin,
 			CamRayWorldDir,
-			DepthBufferWorldPos,
+			DepthBufferTranslatedWorldPos,
 			bPlatformSupportsVolumetricFogOntranslucent);
 
 		LFVContribution = float4(LFVContribution.rgb * LFVLuminanceColorTransmittance.a + LFVLuminanceColorTransmittance.rgb, LFVContribution.a * LFVLuminanceColorTransmittance.a);
diff --git a/Engine/Shaders/Private/VolumetricFog.usf b/Engine/Shaders/Private/VolumetricFog.usf
--- a/Engine/Shaders/Private/VolumetricFog.usf
+++ b/Engine/Shaders/Private/VolumetricFog.usf
@@ -180,7 +180,7 @@ void MaterialSetupCS(
 			// LFV_TODO skip if out of sphere of influence
 			FLocalFogVolumeGPUInstanceData FogInstance = GetLocalFogVolumeGPUInstanceData(InstanceIndex);
 
-			float3 UnitSpacePos = mul(float4(WorldPosition, 1.0), FogInstance.InvTransform).xyz;
+			float3 UnitSpacePos = mul(TranslatedWorldPosition - FogInstance.TranslatedWorlPos, FogInstance.PreTranslatedInvTransform);
 
 			float DistToCenter = length(UnitSpacePos - float3(0.0f, 0.0f, 0.0f));
 			float SphereFade = saturate(1.0f - DistToCenter);
diff --git a/Engine/Source/Runtime/Renderer/Private/LocalFogVolumeRendering.cpp b/Engine/Source/Runtime/Renderer/Private/LocalFogVolumeRendering.cpp
--- a/Engine/Source/Runtime/Renderer/Private/LocalFogVolumeRendering.cpp
+++ b/Engine/Source/Runtime/Renderer/Private/LocalFogVolumeRendering.cpp
@@ -370,7 +370,8 @@ static void LocalFogVolumeViewTiledCullingPass(FViewInfo& View, FRDGBuilder& Gra
 	Local height fog rendering common function
 =============================================================================*/
 
-void GetLocalFogVolumeSortingData(const FScene* Scene, FRDGBuilder& GraphBuilder, FLocalFogVolumeSortingData& Out)
+// This is view data because the transform is expressed in translated world space
+void GetLocalFogVolumeViewSortingData(const FScene* Scene, const FViewInfo& View, FRDGBuilder& GraphBuilder, FLocalFogVolumeSortingData& Out)
 {
 	check(Scene->LocalFogVolumes.Num() > 0); // We should not get there if there is not any local fog volume.
 
@@ -394,12 +395,13 @@ void GetLocalFogVolumeSortingData(const FScene* Scene, FRDGBuilder& GraphBuilder
 		const float FalloffScaleUI = 0.01f;
 		const float SafeFallOff = FMath::Max(LHF->HeightFogFalloff, SafeFalloffThreshold) * FalloffScaleUI;
 
-		FMatrix44f Transform = FMatrix44f(LHF->FogTransform.ToMatrixWithScale());
-		FMatrix44f InvTransform = Transform.Inverse();
+		FTransform RotationScaleMatrix = LHF->FogTransform;
+		RotationScaleMatrix.SetTranslation(FVector::ZeroVector);
+		FMatrix44f InvTransform = FMatrix44f(RotationScaleMatrix.ToMatrixWithScale().Inverse());
 
 		FVector3f XVec(InvTransform.M[0][0], InvTransform.M[0][1], InvTransform.M[0][2]);
 		FVector3f YVec(InvTransform.M[1][0], InvTransform.M[1][1], InvTransform.M[1][2]);
-		FVector3f Tran(InvTransform.M[3][0], InvTransform.M[3][1], InvTransform.M[3][2]);
+		FVector3f TranslatedWorldPosition(View.ViewMatrices.GetPreViewTranslation() + LHF->FogTransform.GetTranslation());
 
 		// Normalization requires small tolerance for large volumes.
 		const float NormalizeTolerance = 1.e-32;
@@ -433,12 +435,12 @@ void GetLocalFogVolumeSortingData(const FScene* Scene, FRDGBuilder& GraphBuilder
 		};
 
 		// Translation and scale at fp32 for stability. Further optimization: we could use fp16 if translated/view space position would be sent.
-		LocalFogVolumeGPUInstanceDataIt->Data0[0] = AsUint32(Tran.X);
-		LocalFogVolumeGPUInstanceDataIt->Data0[1] = AsUint32(Tran.Y);
-		LocalFogVolumeGPUInstanceDataIt->Data0[2] = AsUint32(Tran.Z);
+		LocalFogVolumeGPUInstanceDataIt->Data0[0] = AsUint32(TranslatedWorldPosition.X);
+		LocalFogVolumeGPUInstanceDataIt->Data0[1] = AsUint32(TranslatedWorldPosition.Y);
+		LocalFogVolumeGPUInstanceDataIt->Data0[2] = AsUint32(TranslatedWorldPosition.Z);
 		LocalFogVolumeGPUInstanceDataIt->Data0[3] = AsUint32(LHF->FogUniformScale);
 
-		// Store X and Y from the rotation matrix and recover Z in the shader. Further optimization: could be fp16.
+		// Store X and Y from the rotation matrix and recover Z in the shader.
 		LocalFogVolumeGPUInstanceDataIt->Data1[0] = FVector2DHalf(XVec.X, XVec.Y).AsUInt32();
 		LocalFogVolumeGPUInstanceDataIt->Data1[1] = FVector2DHalf(XVec.Z, YVec.X).AsUInt32();
 		LocalFogVolumeGPUInstanceDataIt->Data1[2] = FVector2DHalf(YVec.Y, YVec.Z).AsUInt32();
@@ -634,11 +636,11 @@ void InitLocalFogVolumesForViews(
 	{
 		RDG_GPU_STAT_SCOPE(GraphBuilder, LocalFogVolumeVolumes);
 
-		FLocalFogVolumeSortingData SortingData;
-		GetLocalFogVolumeSortingData(Scene, GraphBuilder, SortingData);
-
 		for (FViewInfo& View : Views)
 		{
+			FLocalFogVolumeSortingData SortingData;
+			GetLocalFogVolumeViewSortingData(Scene, View, GraphBuilder, SortingData);
+
 			CreateViewLocalFogVolumeBufferSRV(Scene, View, GraphBuilder, SortingData, ShouldRenderLocalFogVolumeInVolumetricFog(Scene, SceneViewFamily, bShouldRenderVolumetricFog), bUseHalfResLocalFogVolume);
 
 			if (View.LocalFogVolumeViewData.GPUInstanceCount > 0)
diff --git a/Engine/Source/Runtime/Renderer/Private/LocalFogVolumeRendering.h b/Engine/Source/Runtime/Renderer/Private/LocalFogVolumeRendering.h
--- a/Engine/Source/Runtime/Renderer/Private/LocalFogVolumeRendering.h
+++ b/Engine/Source/Runtime/Renderer/Private/LocalFogVolumeRendering.h
@@ -135,7 +135,7 @@ bool ShouldRenderLocalFogVolumeInVolumetricFog(const FScene* Scene, const FScene
 float GetLocalFogVolumeGlobalStartDistance();
 bool IsLocalFogVolumeHalfResolution();
 
-void GetLocalFogVolumeSortingData(const FScene* Scene, FRDGBuilder& GraphBuilder, FLocalFogVolumeSortingData& Out);
+void GetLocalFogVolumeViewSortingData(const FScene* Scene, const FViewInfo& View, FRDGBuilder& GraphBuilder, FLocalFogVolumeSortingData& Out);
 
 void SetDummyLocalFogVolumeForViews(FRDGBuilder& GraphBuilder, TArray<FViewInfo>& Views);
 void SetDummyLocalFogVolumeForViews(FRDGBuilder& GraphBuilder, TArray<FViewInfo*>& Views);
