Fixed shader warnings

diff --git a/Engine/Shaders/Private/ShaderPrint.ush b/Engine/Shaders/Private/ShaderPrint.ush
--- a/Engine/Shaders/Private/ShaderPrint.ush
+++ b/Engine/Shaders/Private/ShaderPrint.ush
@@ -168,7 +168,7 @@ FShaderPrintContext InitShaderPrintContext(bool bActive, uint2 InStartCoord)
 FShaderPrintContext InitShaderPrintContextAtCursor(uint2 InCurrentCoord, uint2 InStartCoord) 
 {
 	FShaderPrintConfig Config = InitShaderPrintContextConfig();
-	const bool bActive = all(InCurrentCoord == Config.CursorCoord);
+	const bool bActive = all(int2(InCurrentCoord) == Config.CursorCoord);
 	return InitShaderPrintContext(bActive, float2(InStartCoord) / float2(Config.Resolution));
 }
 
diff --git a/Engine/Shaders/Private/SplineMeshCommon.ush b/Engine/Shaders/Private/SplineMeshCommon.ush
--- a/Engine/Shaders/Private/SplineMeshCommon.ush
+++ b/Engine/Shaders/Private/SplineMeshCommon.ush
@@ -420,7 +420,7 @@ FSplineMeshDeformedLocalBounds SplineMeshDeformLocalBounds(FSplineMeshDeformBoun
 
 	// Sample at evenly-spaced intervals from min->max, accumulating the convex combination of slice bounds
 	LOOP
-	for (int i = 0; i < NUM_SLICE_SAMPLES; ++i)
+	for (uint i = 0; i < NUM_SLICE_SAMPLES; ++i)
 	{
 		SplineMeshDeformLocalBoundsStep(Context, CurSplineDist);
 		CurSplineDist += SplineDistStep;
