Add rect light support to mobile forward renderer

diff --git a/Engine/Shaders/Private/DeferredLightingCommon.ush b/Engine/Shaders/Private/DeferredLightingCommon.ush
--- a/Engine/Shaders/Private/DeferredLightingCommon.ush
+++ b/Engine/Shaders/Private/DeferredLightingCommon.ush
@@ -260,7 +260,7 @@ float GetLocalLightAttenuation(
 		// This extra attenuation has been added for supporting existing 'legacy' shading model on mobile. 
 		// This is not needed for Substrate which unifies all lighting paths
 		#if SHADING_PATH_MOBILE && !SUBSTRATE_ENABLED
-		LightMask *= 1.0f / (DistanceSqr + 1.0f);
+		//LightMask *= 1.0f / (DistanceSqr + 1.0f);
 		#endif
 	}
 	else
@@ -364,7 +364,22 @@ FLightAccumulator AccumulateDynamicLighting(
 	#if TRANSLUCENCY_NON_DIRECTIONAL
 		NoL = 1.0f;
 	#endif
-		Lighting = EvaluateBxDF(GBuffer, N, V, L, NoL, Shadow);
+		BRANCH
+		if (LightData.bRectLight)
+		{
+			FRect Rect = GetRect( ToLight, LightData );
+			const FRectTexture SourceTexture = ConvertToRectTexture(LightData);
+
+		#if REFERENCE_QUALITY
+			Lighting = IntegrateBxDF( GBuffer, N, V, Rect, Shadow, SourceTexture, SVPos );
+		#else
+			Lighting = IntegrateBxDF( GBuffer, N, V, Rect, Shadow, SourceTexture);
+		#endif		
+		}
+		else
+		{
+			Lighting = EvaluateBxDF(GBuffer, N, V, L, NoL, Shadow);
+		}
 
 		Lighting.Specular *= LightData.SpecularScale;
 				
diff --git a/Engine/Shaders/Private/LightData.ush b/Engine/Shaders/Private/LightData.ush
--- a/Engine/Shaders/Private/LightData.ush
+++ b/Engine/Shaders/Private/LightData.ush
@@ -377,16 +377,23 @@ FDeferredLightData ConvertToDeferredLight_Mobile(const FLocalLightData In)
 	half4 OutPreviewShadowMapChannelMask = 0;
 	uint OutLightingChannelMask = LIGHTING_CHANNEL_MASK;
 	FDeferredLightData Out = ConvertToDeferredLight(In, 1.0f, OutPreviewShadowMapChannelMask, OutLightingChannelMask);	
-	
+
 	// Override value for performance on mobile
-	Out.SourceRadius	= 0;
-	Out.SourceLength	= 0;
-	Out.Tangent			= 0;
 	Out.SoftSourceRadius= 0;
 	Out.SpecularScale	= 1.0f;
-	//Render RectLight approximately as SpotLight
-	Out.bSpotLight		= Out.bSpotLight || Out.bRectLight;
-	Out.bRectLight		= false;
+#if RECT_LIGHT_AS_SPOTLIGHT
+	// Render RectLight approximately as SpotLight
+	Out.bSpotLight	= Out.bSpotLight || Out.bRectLight;
+	Out.bRectLight	= false;
+#endif
+
+	BRANCH
+	if (!Out.bRectLight)
+	{
+		Out.SourceRadius = 0;
+		Out.SourceLength = 0;
+		Out.Tangent = 0;
+	}
 
 	return Out;
 }
diff --git a/Engine/Source/Runtime/Core/Private/HAL/ConsoleManager.cpp b/Engine/Source/Runtime/Core/Private/HAL/ConsoleManager.cpp
--- a/Engine/Source/Runtime/Core/Private/HAL/ConsoleManager.cpp
+++ b/Engine/Source/Runtime/Core/Private/HAL/ConsoleManager.cpp
@@ -3772,6 +3772,13 @@ static TAutoConsoleVariable<int32> CVarMobileForwardEnableLocalLights(
 		"2: Local Lights Buffer Enabled\n"),
 	ECVF_ReadOnly | ECVF_RenderThreadSafe);
 
+static TAutoConsoleVariable<int32> CVarMobileRenderRectLightsAsSpotLights(
+	TEXT("r.Mobile.Forward.RenderRectLightsAsSpotLights"),
+	0,
+	TEXT("0: Rect Lights are rendered as area lights\n"
+		"1: Rect Lights are rendered as Spot Lights (default)\n"),
+	ECVF_ReadOnly | ECVF_MobileShaderChange);
+
 static TAutoConsoleVariable<int32> CVarMobileForwardEnableParticleLights(
 	TEXT("r.Mobile.Forward.EnableParticleLights"),
 	0,
diff --git a/Engine/Source/Runtime/Engine/Private/Components/RectLightComponent.cpp b/Engine/Source/Runtime/Engine/Private/Components/RectLightComponent.cpp
--- a/Engine/Source/Runtime/Engine/Private/Components/RectLightComponent.cpp
+++ b/Engine/Source/Runtime/Engine/Private/Components/RectLightComponent.cpp
@@ -263,7 +263,9 @@ void FRectLightSceneProxy::GetLightShaderParameters(FLightRenderParameters& Ligh
 	}
 	
 	// Render RectLight approximately as SpotLight if the requester does not support rect light (e.g., translucent light grid or mobile)
-	const bool bRenderAsSpotLight = !!(Flags & ELightShaderParameterFlags::RectAsSpotLight) || (SceneInterface && IsMobilePlatform(SceneInterface->GetShaderPlatform()));
+	static const auto CVarRenderRectLightAsSpotLight = IConsoleManager::Get().FindTConsoleVariableDataInt(TEXT("r.Mobile.Forward.RenderRectLightsAsSpotLights"));
+	const bool bRenderAsSpotLightCVar = (CVarRenderRectLightAsSpotLight && CVarRenderRectLightAsSpotLight->GetValueOnAnyThread() != 0);
+	const bool bRenderAsSpotLight = !!(Flags & ELightShaderParameterFlags::RectAsSpotLight) || (SceneInterface && IsMobilePlatform(SceneInterface->GetShaderPlatform()) && bRenderAsSpotLightCVar);
 	if (bRenderAsSpotLight)
 	{
 		float ClampedOuterConeAngle = FMath::DegreesToRadians(89.001f);
diff --git a/Engine/Source/Runtime/Engine/Private/ShaderCompiler/ShaderCompiler.cpp b/Engine/Source/Runtime/Engine/Private/ShaderCompiler/ShaderCompiler.cpp
--- a/Engine/Source/Runtime/Engine/Private/ShaderCompiler/ShaderCompiler.cpp
+++ b/Engine/Source/Runtime/Engine/Private/ShaderCompiler/ShaderCompiler.cpp
@@ -7806,6 +7806,12 @@ void GlobalBeginCompileShader(
 			SET_SHADER_DEFINE(Input.Environment, USE_GLES_FBF_DEFERRED, bGLESDeferredShading ? 1 : 0);
 			SET_SHADER_DEFINE(Input.Environment, MOBILE_EXTENDED_GBUFFER, MobileUsesExtenedGBuffer((EShaderPlatform)Target.Platform) ? 1 : 0);
 		}
+		else
+		{
+			static const auto CVarRenderRectLightAsSpotLight = IConsoleManager::Get().FindTConsoleVariableDataInt(TEXT("r.Mobile.Forward.RenderRectLightsAsSpotLights"));
+			const bool bRenderAsSpotLight = (CVarRenderRectLightAsSpotLight && CVarRenderRectLightAsSpotLight->GetValueOnAnyThread() != 0);
+			SET_SHADER_DEFINE(Input.Environment, RECT_LIGHT_AS_SPOTLIGHT, bRenderAsSpotLight ? 1 : 0);
+		}
 
 		SET_SHADER_DEFINE(Input.Environment, USE_SCENE_DEPTH_AUX, MobileRequiresSceneDepthAux(ShaderPlatform) ? 1 : 0);
 
