Fix mobile base pass compile error when having DBuffer decals enabled

diff --git a/Engine/Shaders/Private/MobileBasePassPixelShader.usf b/Engine/Shaders/Private/MobileBasePassPixelShader.usf
--- a/Engine/Shaders/Private/MobileBasePassPixelShader.usf
+++ b/Engine/Shaders/Private/MobileBasePassPixelShader.usf
@@ -478,14 +478,20 @@ void Main(
 		if (ValidDBufferTargetMask)
 		{
 			float2 BufferUV = SvPositionToBufferUV(SvPosition);
+		#if MOBILE_MULTI_VIEW
+			const uint EyeIndex = BasePassInterpolants.MultiViewId;
+		#else
+			const uint EyeIndex = 0;
+		#endif
+
 		#if SUBSTRATE_ENABLED 
 			#if SUBSTRATE_INLINE_SHADING && !SUBSTRATE_OPTIMIZED_UNLIT
-			const FSubstrateDBuffer SubstrateBufferData = GetDBufferData(BufferUV, ValidDBufferTargetMask, BasePassInterpolants.MultiViewId);
+			const FSubstrateDBuffer SubstrateBufferData = GetDBufferData(BufferUV, ValidDBufferTargetMask, EyeIndex);
 			ApplyDBufferData(SubstrateBufferData, SubstratePixelHeader, SubstrateData);
 			DBufferOpacity = SubstrateBufferData.Coverage;
 			#endif
 		#else
-			FDBufferData DBufferData = GetDBufferData(BufferUV, ValidDBufferTargetMask, BasePassInterpolants.MultiViewId);
+			FDBufferData DBufferData = GetDBufferData(BufferUV, ValidDBufferTargetMask, EyeIndex);
 			ApplyDBufferData(DBufferData, MaterialParameters.WorldNormal, SubsurfaceColor, Roughness, BaseColor, Metallic, Specular);
 			DBufferOpacity = (DBufferData.ColorOpacity + DBufferData.NormalOpacity + DBufferData.RoughnessOpacity) * (1.0f / 3.0f);
 		#endif // SUBSTRATE_ENABLED
