Fixes floating point precision problem in Panini's UndistortingDisplacementTexture causing artifact at the center of the screen

diff --git a/Engine/Shaders/Private/PaniniProjection.usf b/Engine/Shaders/Private/PaniniProjection.usf
--- a/Engine/Shaders/Private/PaniniProjection.usf
+++ b/Engine/Shaders/Private/PaniniProjection.usf
@@ -72,6 +72,15 @@ float2 PaniniInverseProjectionScreenPos(float2 InScreenPos)
 	float2 PaniniPosition = InScreenPos * ScreenSpaceToPaniniFactor * rcp(ScreenPosScale);
 	float2 PaniniDirection = PaniniInverseProjection(PaniniPosition, PaniniD, PaniniS);
 	
+	// Works arround floating point precision problem on the computation of OMx
+	const float PrecisionThreshold = 0.003;
+	BRANCH
+	if (abs(PaniniPosition.x) < PrecisionThreshold)
+	{
+		float Lerp = saturate(abs(PaniniPosition.x) / PrecisionThreshold);
+		PaniniDirection.x = Lerp * PaniniInverseProjection(float2(PrecisionThreshold * sign(PaniniPosition.x), PaniniPosition.y), PaniniD, PaniniS).x;
+	}
+
 	// Return the new position back in the screen space frame
 	return PaniniDirection * PaniniToScreenSpaceFactor;
 }
