Mobile: Re-enable extra attenuation for point and spot lights only

diff --git a/Engine/Shaders/Private/DeferredLightingCommon.ush b/Engine/Shaders/Private/DeferredLightingCommon.ush
--- a/Engine/Shaders/Private/DeferredLightingCommon.ush
+++ b/Engine/Shaders/Private/DeferredLightingCommon.ush
@@ -260,7 +260,10 @@ float GetLocalLightAttenuation(
 		// This extra attenuation has been added for supporting existing 'legacy' shading model on mobile. 
 		// This is not needed for Substrate which unifies all lighting paths
 		#if SHADING_PATH_MOBILE && !SUBSTRATE_ENABLED
-		//LightMask *= 1.0f / (DistanceSqr + 1.0f);
+		if (!LightData.bRectLight)
+		{ 
+			LightMask *= 1.0f / (DistanceSqr + 1.0f);
+		}
 		#endif
 	}
 	else
