Fix decal on mobile forward rendering with Substrate.

diff --git a/Engine/Shaders/Private/DecalCommon.ush b/Engine/Shaders/Private/DecalCommon.ush
--- a/Engine/Shaders/Private/DecalCommon.ush
+++ b/Engine/Shaders/Private/DecalCommon.ush
@@ -13,29 +13,29 @@
 #include "DBufferDecalShared.ush"
 
 #if SUBSTRATE_ENABLED  
-float DecalCommonOutput(
+FSubstrateDBuffer DecalCommonOutput(
 	inout FPixelShaderOut Out, 
 	in FSubstratePixelHeader InHeader,
 	in FSubstrateBSDF InBSDF,
 	in float FadeOutWeight,
 	in float CoverageOverBackground)
 {
-	FSubstrateDBuffer DBuffer = InHeader.SubstrateConvertToDBuffer(InBSDF);
+	FSubstrateDBuffer OutDBuffer = InHeader.SubstrateConvertToDBuffer(InBSDF);
 
-	DBuffer.Coverage *= FadeOutWeight;
+	OutDBuffer.Coverage *= FadeOutWeight;
 
-	float3 Color = DBuffer.Emissive;
+	float3 Color = OutDBuffer.Emissive;
 
 #if DECAL_RENDERSTAGE == DECAL_RENDERSTAGE_MOBILE
 	// Cheap mobile forward approximation uses both emissive and color.
-	Color += DBuffer.BaseColor;
+	Color += OutDBuffer.BaseColor;
 #if MATERIALBLENDING_MODULATE
-	Color *= DBuffer.Coverage;
+	Color *= OutDBuffer.Coverage;
 #endif
 #elif DECAL_RENDERSTAGE == DECAL_RENDERSTAGE_AO
 	// Replace color with AO for AO decals.
 #if SUBSTRATE_INLINE_SHADING
-	Color = lerp(1, InHeader.IrradianceAO.MaterialAO, DBuffer.Coverage);
+	Color = lerp(1, InHeader.IrradianceAO.MaterialAO, OutDBuffer.Coverage);
 #else
 #error Substrate - Decals must be evaluated inline.
 #endif
@@ -51,11 +51,11 @@ float DecalCommonOutput(
 #endif
 
 	// Outputing color to MRT0 (can be emissive or any other data set above)
-	Out.MRT[0] = float4(Color, DBuffer.Coverage);
+	Out.MRT[0] = float4(Color, OutDBuffer.Coverage);
 
 #if DECAL_RENDERTARGETMODE == DECAL_RENDERTARGETMODE_DBUFFER
-	// If we output to DBuffer, write to all the 3 RTs.
-	SubstratePackDBuffer(DBuffer, Out.MRT[0], Out.MRT[1], Out.MRT[2]);
+	// If we output to OutDBuffer, write to all the 3 RTs.
+	SubstratePackDBuffer(OutDBuffer, Out.MRT[0], Out.MRT[1], Out.MRT[2]);
 
 	#if DECAL_RENDERTARGET_COUNT == 4
 		Out.MRT[3] = 1.0f; // DBufferMask
@@ -82,7 +82,7 @@ float DecalCommonOutput(
 	Out.MRT[3] = 0;
 #endif
 
-	return DBuffer.Coverage;
+	return OutDBuffer;
 }
 #else // SUBSTRATE_ENABLED
 void DecalCommonOutput(inout FPixelShaderIn In, inout FPixelShaderOut Out, float3 Color, float Opacity, FGBufferData Data)
diff --git a/Engine/Shaders/Private/DeferredDecal.usf b/Engine/Shaders/Private/DeferredDecal.usf
--- a/Engine/Shaders/Private/DeferredDecal.usf
+++ b/Engine/Shaders/Private/DeferredDecal.usf
@@ -195,6 +195,7 @@ void FPixelShaderInOut_MainPS(inout FPixelShaderIn In, inout FPixelShaderOut Out
 	SubstrateHeader.IrradianceAO.MaterialAO = GetMaterialAmbientOcclusion(PixelMaterialInputs);
 
 	float Coverage = 0.0;
+	FSubstrateDBuffer DBufferData =(FSubstrateDBuffer)0;
 	if (SubstrateHeader.SubstrateTree.BSDFCount > 0)
 	{
 		const float3 V = MaterialParameters.CameraVector;
@@ -210,38 +211,50 @@ void FPixelShaderInOut_MainPS(inout FPixelShaderIn In, inout FPixelShaderOut Out
 #endif 
 
 		// Since the Convert-To-Decal node forces parameter blending for the entire Substrate tree, the BSDF we are interested in will be in 0.
-		Coverage = DecalCommonOutput(Out, SubstrateHeader, SubstrateHeader.SubstrateTree.BSDFs[0], DecalFading, CoverageOverBackground);
-
+		DBufferData = DecalCommonOutput(Out, SubstrateHeader, SubstrateHeader.SubstrateTree.BSDFs[0], DecalFading, CoverageOverBackground);
 	}
-	const float Opacity = Coverage;
+	const float Opacity = DBufferData.Coverage;
+	const float3 EmissiveColor = DBufferData.Emissive;
 #else // SUBSTRATE_ENABLED
-	float3 Color = GetMaterialEmissive(PixelMaterialInputs);
+	float3 EmissiveColor = GetMaterialEmissive(PixelMaterialInputs);
 	
 	float Opacity = GetMaterialOpacity(PixelMaterialInputs) * DecalFading;
 
-	FGBufferData GBufferData;
-	GBufferData.WorldNormal = MaterialParameters.WorldNormal;
-	GBufferData.BaseColor = GetMaterialBaseColor(PixelMaterialInputs);
-	GBufferData.Metallic = GetMaterialMetallic(PixelMaterialInputs);
-	GBufferData.Specular = GetMaterialSpecular(PixelMaterialInputs);
-	GBufferData.Roughness = GetMaterialRoughness(PixelMaterialInputs);
-	GBufferData.CustomData = 0;
-	GBufferData.IndirectIrradiance = 0;
-	GBufferData.PrecomputedShadowFactors = 1;
-	GBufferData.GBufferAO = GetMaterialAmbientOcclusion(PixelMaterialInputs);
-	GBufferData.ShadingModelID = SHADINGMODELID_DEFAULT_LIT;
-	GBufferData.SelectiveOutputMask = 0;
-	GBufferData.PerObjectGBufferData = 1; 
-	GBufferData.DiffuseIndirectSampleOcclusion = 0;
-	GBufferData.Velocity = 0;
-
-	DecalCommonOutput(In, Out, Color, Opacity, GBufferData);
+	FGBufferData DBufferData;
+	DBufferData.WorldNormal = MaterialParameters.WorldNormal;
+	DBufferData.BaseColor = GetMaterialBaseColor(PixelMaterialInputs);
+	DBufferData.Metallic = GetMaterialMetallic(PixelMaterialInputs);
+	DBufferData.Specular = GetMaterialSpecular(PixelMaterialInputs);
+	DBufferData.Roughness = GetMaterialRoughness(PixelMaterialInputs);
+	DBufferData.CustomData = 0;
+	DBufferData.IndirectIrradiance = 0;
+	DBufferData.PrecomputedShadowFactors = 1;
+	DBufferData.GBufferAO = GetMaterialAmbientOcclusion(PixelMaterialInputs);
+	DBufferData.ShadingModelID = SHADINGMODELID_DEFAULT_LIT;
+	DBufferData.SelectiveOutputMask = 0;
+	DBufferData.PerObjectGBufferData = 1; 
+	DBufferData.DiffuseIndirectSampleOcclusion = 0;
+	DBufferData.Velocity = 0;
+
+	DecalCommonOutput(In, Out, EmissiveColor, Opacity, DBufferData);
 #endif // SUBSTRATE_ENABLED
 
 #if DECAL_RENDERSTAGE == DECAL_RENDERSTAGE_MOBILE
 	#if DECAL_MOBILE_FORWARD_LIT
-		Out.MRT[0].rgb = ForwardDecalLighting(TranslatedWorldPosition, Color, Opacity, MaterialParameters.ScreenPosition.xy, MaterialParameters.ScreenPosition.w,
-			GBufferData, CameraVector, MaterialParameters.TangentToWorld[2], ReflectionAboutCustomWorldNormal(MaterialParameters, MaterialParameters.WorldNormal, false));
+		Out.MRT[0].rgb = ForwardDecalLighting(
+			TranslatedWorldPosition, 
+			MaterialParameters.ScreenPosition.xy, 
+			MaterialParameters.ScreenPosition.w,
+			CameraVector, 
+			MaterialParameters.TangentToWorld[2], 
+			ReflectionAboutCustomWorldNormal(MaterialParameters, MaterialParameters.WorldNormal, false),
+			DBufferData.WorldNormal,
+			DBufferData.Specular,
+			DBufferData.BaseColor,
+			DBufferData.Metallic,
+			DBufferData.Roughness,
+			EmissiveColor,
+			Opacity);
 	#endif
 #endif
 
diff --git a/Engine/Shaders/Private/MeshDecals.usf b/Engine/Shaders/Private/MeshDecals.usf
--- a/Engine/Shaders/Private/MeshDecals.usf
+++ b/Engine/Shaders/Private/MeshDecals.usf
@@ -162,12 +162,12 @@ void FPixelShaderInOut_MainPS(
 
 		// Either propagate the BSDF opacity, or use the one for AlphaComposite (premultiplied alpha)
 		float CoverageOverBackground = SubstrateHeader.SubstrateTree.BSDFs[0].Coverage;
-#if SUBSTRATE_USE_PREMULTALPHA_OVERRIDE // AlphaComposite - Premultiplied alpha blending
+	#if SUBSTRATE_USE_PREMULTALPHA_OVERRIDE // AlphaComposite - Premultiplied alpha blending
 		CoverageOverBackground = GetMaterialOpacity(PixelMaterialInputs);
-#endif 
-
+	#endif 
 		// Since the Convert-To-Decal node forces parameter blending for the entire Substrate tree, the BSDF we are interested in will be in 0.
-		Coverage = DecalCommonOutput(Out, SubstrateHeader, SubstrateHeader.SubstrateTree.BSDFs[0], 1.0f /*CoverageWeight*/, CoverageOverBackground);
+		const FSubstrateDBuffer DBuffer = DecalCommonOutput(Out, SubstrateHeader, SubstrateHeader.SubstrateTree.BSDFs[0], 1.0f /*CoverageWeight*/, CoverageOverBackground);
+		Coverage = DBuffer.Coverage;
 	}
 	const float Opacity = Coverage;
 #else // SUBSTRATE_ENABLED
diff --git a/Engine/Shaders/Private/MobileLightingCommon.ush b/Engine/Shaders/Private/MobileLightingCommon.ush
--- a/Engine/Shaders/Private/MobileLightingCommon.ush
+++ b/Engine/Shaders/Private/MobileLightingCommon.ush
@@ -688,25 +688,38 @@ void MergeLocalLights(const FCulledLightsGridData InLightGridData
 	Mobile Forward Decal Lighting
 ------------------------------------------------------------------------------*/
 
-half3 ForwardDecalLighting(float3 TranslatedWorldPosition, half3 EmissiveColor, half Opacity, float2 ScreenPosition, float SceneDepth, FGBufferData GBuffer, half3 CameraVector, half3 VertexNormal, half3 ReflectionVector)
+half3 ForwardDecalLighting(
+	float3 TranslatedWorldPosition, 
+	float2 ScreenPosition, 
+	float SceneDepth, 
+	half3 CameraVector, 
+	half3 VertexNormal, 
+	half3 ReflectionVector,
+	half3 WorldNormal,
+	half  Specular,
+	half3 BaseColor,
+	half  Metallic,
+	half  Roughness,
+	half3 EmissiveColor, 
+	half  Opacity)
 {
 	float ShadowPositionZ = 0;
 	half ShadowMap = MobileDirectionalLightCSM(ScreenPosition.xy, SceneDepth, ShadowPositionZ);
 	half3 DirectionalLightDirection = MobileDirectionalLight.DirectionalLightDirectionAndShadowTransition.xyz;
-	half NoL = saturate(dot(half3(GBuffer.WorldNormal), DirectionalLightDirection));
+	half NoL = saturate(dot(half3(WorldNormal), DirectionalLightDirection));
 	
-	half3 SpecularColor = ComputeF0(GBuffer.Specular, GBuffer.BaseColor, GBuffer.Metallic);
-	half3 DiffuseColor = GBuffer.BaseColor - GBuffer.BaseColor * GBuffer.Metallic;
+	half3 SpecularColor = ComputeF0(Specular, BaseColor, Metallic);
+	half3 DiffuseColor = BaseColor - BaseColor * Metallic;
 
 	half3 OutColor = EmissiveColor;
-	OutColor += MobileDirectionalLight.DirectionalLightColor.rgb * NoL * SimpleShading(DiffuseColor, SpecularColor, max(GBuffer.Roughness, .04f), DirectionalLightDirection, -CameraVector, GBuffer.WorldNormal) * ShadowMap;
+	OutColor += MobileDirectionalLight.DirectionalLightColor.rgb * NoL * SimpleShading(DiffuseColor, SpecularColor, max(Roughness, .04f), DirectionalLightDirection, -CameraVector, WorldNormal) * ShadowMap;
 
-	half3 SkyDiffuse = GetSkySHDiffuseSimple(GBuffer.WorldNormal) * ResolvedView.SkyLightColor.rgb;
+	half3 SkyDiffuse = GetSkySHDiffuseSimple(WorldNormal) * ResolvedView.SkyLightColor.rgb;
 	OutColor += DiffuseColor * SkyDiffuse;
 
-	half3 SpecularIBL = GetImageBasedReflectionSpecular(ReflectionVector, GBuffer.Roughness, TranslatedWorldPosition, GBuffer.WorldNormal, 1.0, Luminance(SkyDiffuse), 0);
-	half NoV = saturate(abs(dot(GBuffer.WorldNormal, -CameraVector)) + 1e-5);
-	OutColor += SpecularIBL * EnvBRDFApprox(SpecularColor, GBuffer.Roughness, NoV);
+	half3 SpecularIBL = GetImageBasedReflectionSpecular(ReflectionVector, Roughness, TranslatedWorldPosition, WorldNormal, 1.0, Luminance(SkyDiffuse), 0);
+	half NoV = saturate(abs(dot(WorldNormal, -CameraVector)) + 1e-5);
+	OutColor += SpecularIBL * EnvBRDFApprox(SpecularColor, Roughness, NoV);
 
 #if MATERIALBLENDING_MODULATE
 	OutColor *= Opacity;
