Add support for merged light data on mobile with Substrate.

diff --git a/Engine/Shaders/Private/DeferredLightingCommon.ush b/Engine/Shaders/Private/DeferredLightingCommon.ush
--- a/Engine/Shaders/Private/DeferredLightingCommon.ush
+++ b/Engine/Shaders/Private/DeferredLightingCommon.ush
@@ -23,6 +23,10 @@
 #define ADAPTIVE_VOLUMETRIC_SHADOW_MAP 0
 #endif
 
+#ifndef ALLOW_LOCAL_LIGHT_DISTANCE_ATTENUATION
+#define ALLOW_LOCAL_LIGHT_DISTANCE_ATTENUATION 0
+#endif
+
 #if ADAPTIVE_VOLUMETRIC_SHADOW_MAP
 #include "HeterogeneousVolumes/HeterogeneousVolumesAdaptiveVolumetricShadowMapSampling.ush"
 
@@ -259,7 +263,7 @@ float GetLocalLightAttenuation(
 		LightMask = Square( saturate( 1 - Square( DistanceSqr * Square(LightData.InvRadius) ) ) );
 		// This extra attenuation has been added for supporting existing 'legacy' shading model on mobile. 
 		// This is not needed for Substrate which unifies all lighting paths
-		#if SHADING_PATH_MOBILE && !SUBSTRATE_ENABLED
+		#if SHADING_PATH_MOBILE && (!SUBSTRATE_ENABLED || ALLOW_LOCAL_LIGHT_DISTANCE_ATTENUATION)
 		if (!LightData.bRectLight)
 		{ 
 			LightMask *= 1.0f / (DistanceSqr + 1.0f);
diff --git a/Engine/Shaders/Private/MobileLocalLightsBuffer.usf b/Engine/Shaders/Private/MobileLocalLightsBuffer.usf
--- a/Engine/Shaders/Private/MobileLocalLightsBuffer.usf
+++ b/Engine/Shaders/Private/MobileLocalLightsBuffer.usf
@@ -1,6 +1,7 @@
 // Copyright Epic Games, Inc. All Rights Reserved.
 
 #define SUPPORT_CONTACT_SHADOWS 0
+#define ALLOW_LOCAL_LIGHT_DISTANCE_ATTENUATION 1
 #define MAX_UINT 4294967295
 
 #if !COMPUTE_SHADER
diff --git a/Engine/Shaders/Private/Substrate/SubstrateMobileForwardLighting.ush b/Engine/Shaders/Private/Substrate/SubstrateMobileForwardLighting.ush
--- a/Engine/Shaders/Private/Substrate/SubstrateMobileForwardLighting.ush
+++ b/Engine/Shaders/Private/Substrate/SubstrateMobileForwardLighting.ush
@@ -272,7 +272,139 @@ float3 SubstrateMobileForwardLighting(
 				////
 				//// Evaluate local lights
 				////
-				#if ENABLE_CLUSTERED_LIGHTS && !DISABLE_FORWARD_LOCAL_LIGHTS && (FORWARD_SHADING || TRANSLUCENCY_LIGHTING_SURFACE_FORWARDSHADING) // Forbidden for Single Layer Water
+				#if MERGED_LOCAL_LIGHTS_MOBILE == 1
+				{
+					// Merged light data
+					FDeferredLightData LightData = (FDeferredLightData)0;
+					half3 LocalLightA = MobileSceneTextures.LocalLightTextureA.Load(int3(SvPosition.xy, 0)).xyz; 
+					half4 LocalLightB = MobileSceneTextures.LocalLightTextureB.Load(int3(SvPosition.xy, 0));
+
+					LightData.SourceRadius		= 0;
+					LightData.SourceLength		= 0;
+					LightData.Tangent			= 0;
+					LightData.SoftSourceRadius	= 0;
+					LightData.SpecularScale		= LocalLightB.a;
+					LightData.bSpotLight		= false;
+					LightData.bRectLight		= false;
+					LightData.Color 			= LocalLightA.rgb;
+					LightData.bRadialLight 		= false;
+					LightData.Direction 		= OctahedronToUnitVector(Pack888To1212(LocalLightB.rgb) * 2 - 1);
+
+					const half4 LightAttenuation 	= half4(1, 1, 1, 1); 
+					const float SurfaceShadow = 0;
+					const uint  LocalLightChannelMask = LIGHTING_CHANNEL_MASK;
+
+					const FRectTexture LocalRectTexture = ConvertToRectTexture(LightData);
+
+					float3 L = LightData.Direction;	// Already normalized
+					float3 ToLight = L;
+					float LightMask = 1;
+					if (LightData.bRadialLight)
+					{
+						LightMask = GetLocalLightAttenuation(TranslatedWorldPosition, LightData, ToLight, L);
+					}
+
+					FTranslucencyLightingVector LightingVector = GetTranslucencyLightingVector(LightData, TranslatedWorldPosition);
+					#if SUBSTRATE_TRANSLUCENCY_LIGHTING_VOLUMETRIC_DIRECTIONAL
+					GetVolumeLightingDirectional(LightingVector.AmbientLightingVector, LightingVector.DirectionalLightingVector, SubstrateBSDFContext.N, BSDFEvaluate.DiffuseColor, GetMaterialTranslucencyDirectionalLightingIntensity(), InterpolatedLighting, VolumeLighting);
+					#elif SUBSTRATE_TRANSLUCENCY_LIGHTING_VOLUMETRIC_NONDIRECTIONAL
+					GetVolumeLightingNonDirectional(LightingVector.AmbientLightingVector, BSDFEvaluate.DiffuseColor, InterpolatedLighting, VolumeLighting);
+					#else
+					// Update the Substrate BSDF context accordgin to the new L
+					SubstrateBSDFContext.SubstrateUpdateBSDFContext(ToLight);
+
+					FSubstrateEvaluateResult BSDFEvaluate = (FSubstrateEvaluateResult)0;
+					Color += SubstrateForwardLightingCommon(
+						Dither,
+						Settings,
+						LightData,
+						ToLight,
+						LightMask,
+						LightAttenuation,
+						LocalRectTexture,
+						LocalLightChannelMask,
+						PrimitiveLightingChannelMask,
+						PrecomputedShadowFactors,
+						TranslatedWorldPosition,
+						SceneDepth,
+						LuminanceWeight(SubstrateBSDFContext, CurrentBSDF),
+						SubstratePixelHeader,
+						SubstrateBSDFContext,
+						BSDFEvaluate);
+					#endif
+				}
+				#elif  MERGED_LOCAL_LIGHTS_MOBILE == 2
+				{
+					const FCulledLightsGridData CulledLightGridData = GetCulledLightsGrid(GridIndex, EyeIndex);
+
+					float3 MergedLightColor = uint3(0.f, 0.f , 0.f);
+					float3 MergedLightL = float3(0,0,0);
+					float MergedSpecularScale = 0.f;
+					float MergedTotalWeight = 0;
+					MergeLocalLights(CulledLightGridData, MaterialParameters.WorldPosition_CamRelative, EyeIndex, false/*bSupportLightFunctions*/, MergedLightColor, MergedLightL, MergedSpecularScale, MergedTotalWeight, SubstrateBSDFContext.N);
+
+					if (MergedTotalWeight > 0.f)
+					{
+						MergedLightL = MergedLightL / MergedTotalWeight;
+						MergedSpecularScale = MergedSpecularScale / MergedTotalWeight;
+
+						FDeferredLightData LightData = (FDeferredLightData)0;
+						LightData.SourceRadius		= 0;
+						LightData.SourceLength		= 0;
+						LightData.Tangent			= 0;
+						LightData.SoftSourceRadius	= 0;
+						LightData.SpecularScale		= MergedSpecularScale;
+						LightData.bSpotLight		= false;
+						LightData.bRectLight		= false;
+						LightData.Color 			= MergedLightColor;
+						LightData.bRadialLight 		= false;
+						LightData.Direction 		= MergedLightL;
+
+						const half4 LightAttenuation = half4(1, 1, 1, 1); 
+						const float SurfaceShadow = 0;
+						const uint  LocalLightChannelMask = LIGHTING_CHANNEL_MASK;
+
+						const FRectTexture LocalRectTexture = ConvertToRectTexture(LightData);			
+
+						float3 L = LightData.Direction;	// Already normalized
+						float3 ToLight = L;
+						float LightMask = 1;
+						if (LightData.bRadialLight)
+						{
+							LightMask = GetLocalLightAttenuation(TranslatedWorldPosition, LightData, ToLight, L);
+						}
+
+						FTranslucencyLightingVector LightingVector = GetTranslucencyLightingVector(LightData, TranslatedWorldPosition);
+						#if SUBSTRATE_TRANSLUCENCY_LIGHTING_VOLUMETRIC_DIRECTIONAL
+						GetVolumeLightingDirectional(LightingVector.AmbientLightingVector, LightingVector.DirectionalLightingVector, SubstrateBSDFContext.N, BSDFEvaluate.DiffuseColor, GetMaterialTranslucencyDirectionalLightingIntensity(), InterpolatedLighting, VolumeLighting);
+						#elif SUBSTRATE_TRANSLUCENCY_LIGHTING_VOLUMETRIC_NONDIRECTIONAL
+						GetVolumeLightingNonDirectional(LightingVector.AmbientLightingVector, BSDFEvaluate.DiffuseColor, InterpolatedLighting, VolumeLighting);
+						#else
+						// Update the Substrate BSDF context accordgin to the new L
+						SubstrateBSDFContext.SubstrateUpdateBSDFContext(ToLight);
+
+						FSubstrateEvaluateResult BSDFEvaluate = (FSubstrateEvaluateResult)0;
+						Color += SubstrateForwardLightingCommon(
+							Dither,
+							Settings,
+							LightData,
+							ToLight,
+							LightMask,
+							LightAttenuation,
+							LocalRectTexture,
+							LocalLightChannelMask,
+							PrimitiveLightingChannelMask,
+							PrecomputedShadowFactors,
+							TranslatedWorldPosition,
+							SceneDepth,
+							LuminanceWeight(SubstrateBSDFContext, CurrentBSDF),
+							SubstratePixelHeader,
+							SubstrateBSDFContext,
+							BSDFEvaluate);
+						#endif
+					}
+				}
+				#elif ENABLE_CLUSTERED_LIGHTS && !DISABLE_FORWARD_LOCAL_LIGHTS && (FORWARD_SHADING || TRANSLUCENCY_LIGHTING_SURFACE_FORWARDSHADING) // Forbidden for Single Layer Water
 				{
 					// Basic implementation for FORWARD_LOCAL_LIGHTS where local lights are reloaded per BSDF layer. It is fast by default when layer==1.
 					const FCulledLightsGridData CulledLightsGrid = GetCulledLightsGrid(GridIndex, EyeIndex);
