Fixes TSR's resurrection preview with lens distortion

diff --git a/Engine/Shaders/Private/TemporalSuperResolution/TSRVisualize.usf b/Engine/Shaders/Private/TemporalSuperResolution/TSRVisualize.usf
--- a/Engine/Shaders/Private/TemporalSuperResolution/TSRVisualize.usf
+++ b/Engine/Shaders/Private/TemporalSuperResolution/TSRVisualize.usf
@@ -39,6 +39,8 @@ float OutputToHistoryResolutionFractionSquare;
 float FlickeringFramePeriod;
 float PerceptionAdd;
 
+Texture2D<float2>    PrevDistortingDisplacementTexture;
+Texture2D<float2>    ResurrectedDistortingDisplacementTexture;
 Texture2D<float2>    UndistortingDisplacementTexture;
 Texture2D<tsr_half3> InputTexture;
 Texture2D<tsr_half>  InputMoireLumaTexture;
@@ -58,12 +60,12 @@ RWTexture2D<tsr_half4> Output;
 
 //------------------------------------------------------- FUNCTION
 
-float2 ApplyLensDistortionOnScreenPos(float2 DestScreenPos)
+float2 ApplyLensDistortionOnScreenPos(Texture2D<float2> DisplacementTexture, float2 DestScreenPos)
 #if TSR_SUPPORT_LENS_DISTORTION
 {
 	float2 DestViewportUV = ScreenPosToViewportUV(DestScreenPos);
 
-	float2 SrcViewportUV = DestViewportUV + UndistortingDisplacementTexture.SampleLevel(GlobalBilinearClampedSampler, DestViewportUV, 0);
+	float2 SrcViewportUV = DestViewportUV + DisplacementTexture.SampleLevel(GlobalBilinearClampedSampler, DestViewportUV, 0);
 			
 	return ViewportUVToScreenPos(SrcViewportUV);
 }
@@ -168,16 +170,15 @@ tsr_half3 DisplayResurrectionMask(FContext Ctx)
 	float ClosestDeviceZ = ClosestDepthTexture.SampleLevel(GlobalPointClampedSampler, Ctx.InputUV, 0).g;
 	bool ResurrectionMask = (tsr_ushort(HistoryRejectionTexture.SampleLevel(GlobalPointClampedSampler, Ctx.InputUV, 0).a * tsr_half(255.0)) & tsr_ushort(0x2)) != tsr_ushort(0x0);
 
-	float4 ThisClip = float4(Ctx.ScreenPos, ClosestDeviceZ, 1);
+	float4 ThisClip = float4(ApplyLensDistortionOnScreenPos(UndistortingDisplacementTexture, Ctx.ScreenPos), ClosestDeviceZ, 1);
 	float4 ResurrectionClip = mul(ThisClip, ClipToResurrectionClip); 
-	float2 ResurrectionScreen = ResurrectionClip.xy / ResurrectionClip.w;
+	float2 ResurrectionScreen = ApplyLensDistortionOnScreenPos(ResurrectedDistortingDisplacementTexture, ResurrectionClip.xy / ResurrectionClip.w);
 
-	if (any(ResurrectionScreen != clamp(ResurrectionScreen, -1.0, 1.0)))
-	{
-		return SampleTSROutputGray(Ctx) * kPink;
-	}
+	tsr_half3 Color = kWhite;
+	Color = select(any(ResurrectionScreen != clamp(ResurrectionScreen, -1.0, 1.0)), kPink, Color);
+	Color = select(ResurrectionMask, lerp(kRed, kGreen, LowFrequencyRejection), Color);
 
-	return SampleTSROutputGray(Ctx) * select(ResurrectionMask, lerp(kRed, kGreen, LowFrequencyRejection), kWhite);
+	return SampleTSROutputGray(Ctx) * Color;
 }
 
 tsr_half3 DisplayResurrectedColor(FContext Ctx)
@@ -191,10 +192,10 @@ tsr_half3 DisplayResurrectedColor(FContext Ctx)
 	float ClosestDeviceZ = ClosestDepthTexture.SampleLevel(GlobalPointClampedSampler, Ctx.InputUV, 0).g;
 	bool ResurrectionMask = (tsr_ushort(HistoryRejectionTexture.SampleLevel(GlobalPointClampedSampler, Ctx.InputUV, 0).a * tsr_half(255.0)) & tsr_ushort(0x2)) != tsr_ushort(0x0);
 
-	float4 ThisClip = float4(Ctx.ScreenPos, ClosestDeviceZ, 1);
+	float4 ThisClip = float4(ApplyLensDistortionOnScreenPos(UndistortingDisplacementTexture, Ctx.ScreenPos), ClosestDeviceZ, 1);
 	float4 ResurrectionClip = mul(ThisClip, ClipToResurrectionClip); 
-	float2 ResurrectionScreen = ResurrectionClip.xy / ResurrectionClip.w;
-	
+	float2 ResurrectionScreen = ApplyLensDistortionOnScreenPos(ResurrectedDistortingDisplacementTexture, ResurrectionClip.xy / ResurrectionClip.w);
+
 	float2 KernelHistoryBufferUV = ApplyScreenTransform(ResurrectionScreen, ScreenPosToPrevHistoryBufferUV);
 	KernelHistoryBufferUV = fastClamp(KernelHistoryBufferUV, PrevHistoryInfo_UVViewportBilinearMin, PrevHistoryInfo_UVViewportBilinearMax);
 
@@ -232,7 +233,7 @@ tsr_half3 DisplayAntiFlickering(FContext Ctx)
 	const tsr_half TotalVariationCountThreshold = tsr_half(1.0 / (1.0 - pow(1.0 - BlendFinal, tsr_half(FlickeringFramePeriod))));
 
 	// Texture coordinate
-	float2 MoireHistoryUV = ApplyScreenTransform(ApplyLensDistortionOnScreenPos(Ctx.ScreenPos), ScreenPosToMoireHistoryUV);
+	float2 MoireHistoryUV = ApplyScreenTransform(ApplyLensDistortionOnScreenPos(UndistortingDisplacementTexture, Ctx.ScreenPos), ScreenPosToMoireHistoryUV);
 	MoireHistoryUV = fastClamp(MoireHistoryUV, MoireHistoryUVBilinearMin, MoireHistoryUVBilinearMax);
 
 	// Texture fetches
@@ -304,9 +305,9 @@ void MainCS(
 		tsr_short2(GroupId) * tsr_short2(TILE_SIZE, TILE_SIZE) +
 		Map8x8Tile2x2Lane(GroupThreadIndex));
 	Ctx.ScreenPos = ApplyScreenTransform(float2(Ctx.OutputPixelPos), OutputPixelPosToScreenPos);
-	Ctx.InputUV = ApplyScreenTransform(ApplyLensDistortionOnScreenPos(Ctx.ScreenPos), ScreenPosToInputUV);
+	Ctx.InputUV = ApplyScreenTransform(ApplyLensDistortionOnScreenPos(UndistortingDisplacementTexture, Ctx.ScreenPos), ScreenPosToInputUV);
 	Ctx.InputUV = fastClamp(Ctx.InputUV, InputInfo_UVViewportBilinearMin, InputInfo_UVViewportBilinearMax);
-	Ctx.InputPixelPos = tsr_short2(ApplyScreenTransform(ApplyLensDistortionOnScreenPos(Ctx.ScreenPos), ScreenPosToInputPixelPos));
+	Ctx.InputPixelPos = tsr_short2(ApplyScreenTransform(ApplyLensDistortionOnScreenPos(UndistortingDisplacementTexture, Ctx.ScreenPos), ScreenPosToInputPixelPos));
 	Ctx.InputPixelPos = fastClamp(Ctx.InputPixelPos, tsr_short2(InputInfo_ViewportMin), tsr_short2(InputInfo_ViewportMax - 1));
 
 	tsr_half3 OutputColor = tsr_half(0.0);
diff --git a/Engine/Source/Runtime/Renderer/Private/PostProcess/TemporalSuperResolution.cpp b/Engine/Source/Runtime/Renderer/Private/PostProcess/TemporalSuperResolution.cpp
--- a/Engine/Source/Runtime/Renderer/Private/PostProcess/TemporalSuperResolution.cpp
+++ b/Engine/Source/Runtime/Renderer/Private/PostProcess/TemporalSuperResolution.cpp
@@ -1120,6 +1120,8 @@ class FTSRVisualizeCS : public FTSRShader
 		SHADER_PARAMETER(float, FlickeringFramePeriod)
 		SHADER_PARAMETER(float, PerceptionAdd)
 
+		SHADER_PARAMETER_RDG_TEXTURE(Texture2D, PrevDistortingDisplacementTexture)
+		SHADER_PARAMETER_RDG_TEXTURE(Texture2D, ResurrectedDistortingDisplacementTexture)
 		SHADER_PARAMETER_RDG_TEXTURE(Texture2D, UndistortingDisplacementTexture)
 		SHADER_PARAMETER_RDG_TEXTURE(Texture2D, InputTexture)
 		SHADER_PARAMETER_RDG_TEXTURE(Texture2D, InputMoireLumaTexture)
@@ -1722,6 +1724,8 @@ FDefaultTemporalUpscaler::FOutputs AddTemporalSuperResolutionPasses(
 	int32 CurrentFrameRollingIndex = 0;
 	FTSRHistoryTextures PrevHistory;
 	FTSRHistorySliceSequence PrevHistorySliceSequence;
+	FRDGTextureRef PrevDistortingDisplacementTexture = BlackDummy;
+	FRDGTextureRef ResurrectedDistortingDisplacementTexture = BlackDummy;
 	if (bCameraCut)
 	{
 		PrevHistory.ColorArray = BlackArrayDummy;
@@ -1787,6 +1791,15 @@ FDefaultTemporalUpscaler::FOutputs AddTemporalSuperResolutionPasses(
 		ResurrectionFrameSliceIndex = PrevHistorySliceSequence.RollingIndexToSliceIndex(ResurrectionFrameRollingIndex);
 		PrevFrameSliceIndex = PrevHistorySliceSequence.RollingIndexToSliceIndex(PrevFrameRollingIndex);
 		CurrentFrameSliceIndex = HistorySliceSequence.RollingIndexToSliceIndex(CurrentFrameRollingIndex);
+
+		if (InputHistory.DistortingDisplacementTextures[PrevFrameSliceIndex].IsValid())
+		{
+			PrevDistortingDisplacementTexture = GraphBuilder.RegisterExternalTexture(InputHistory.DistortingDisplacementTextures[PrevFrameSliceIndex]);
+		}
+		if (InputHistory.DistortingDisplacementTextures[ResurrectionFrameSliceIndex].IsValid() && ResurrectionFrameSliceIndex != PrevFrameSliceIndex)
+		{
+			ResurrectedDistortingDisplacementTexture = GraphBuilder.RegisterExternalTexture(InputHistory.DistortingDisplacementTextures[ResurrectionFrameSliceIndex]);
+		}
 	}
 
 	// Whether history Resurrection is possible at all 
@@ -2434,23 +2447,12 @@ FDefaultTemporalUpscaler::FOutputs AddTemporalSuperResolutionPasses(
 			PassParameters->PrevHistoryMetadataTexture = GraphBuilder.CreateSRV(PrevMetadataSRVDesc);
 		}
 
-		PassParameters->PrevDistortingDisplacementTexture = BlackDummy;
-		PassParameters->ResurrectedDistortingDisplacementTexture = BlackDummy;
+		PassParameters->PrevDistortingDisplacementTexture = PrevDistortingDisplacementTexture;
+		PassParameters->ResurrectedDistortingDisplacementTexture = ResurrectedDistortingDisplacementTexture;
 		PassParameters->UndistortingDisplacementTexture = BlackDummy;
-		if (bLensDistortion)
+		if (bLensDistortion && PassInputs.LensDistortionLUT.IsEnabled())
 		{
-			if (!bCameraCut && InputHistory.DistortingDisplacementTextures[PrevFrameSliceIndex].IsValid())
-			{
-				PassParameters->PrevDistortingDisplacementTexture = GraphBuilder.RegisterExternalTexture(InputHistory.DistortingDisplacementTextures[PrevFrameSliceIndex]);
-			}
-			if (!bCameraCut && InputHistory.DistortingDisplacementTextures[ResurrectionFrameSliceIndex].IsValid() && bCanResurrectHistory)
-			{
-				PassParameters->ResurrectedDistortingDisplacementTexture = GraphBuilder.RegisterExternalTexture(InputHistory.DistortingDisplacementTextures[ResurrectionFrameSliceIndex]);
-			}
-			if (PassInputs.LensDistortionLUT.UndistortingDisplacementTexture)
-			{
-				PassParameters->UndistortingDisplacementTexture = PassInputs.LensDistortionLUT.UndistortingDisplacementTexture;
-			}
+			PassParameters->UndistortingDisplacementTexture = PassInputs.LensDistortionLUT.UndistortingDisplacementTexture;
 		}
 
 		{
@@ -2769,6 +2771,8 @@ FDefaultTemporalUpscaler::FOutputs AddTemporalSuperResolutionPasses(
 			PassParameters->FlickeringFramePeriod = FlickeringFramePeriod;
 			PassParameters->PerceptionAdd = FMath::Pow(0.5f, CVarTSRShadingExposureOffset.GetValueOnRenderThread());
 
+			PassParameters->PrevDistortingDisplacementTexture = PrevDistortingDisplacementTexture;
+			PassParameters->ResurrectedDistortingDisplacementTexture = ResurrectedDistortingDisplacementTexture;
 			PassParameters->UndistortingDisplacementTexture = BlackDummy;
 			if (bLensDistortion && PassInputs.LensDistortionLUT.IsEnabled())
 			{
