Fixes shader compilations failure in SSR introduced by 34127283

diff --git a/Engine/Shaders/Private/SSRT/SSRTReflections.usf b/Engine/Shaders/Private/SSRT/SSRTReflections.usf
--- a/Engine/Shaders/Private/SSRT/SSRTReflections.usf
+++ b/Engine/Shaders/Private/SSRT/SSRTReflections.usf
@@ -483,7 +483,12 @@ void MobileScreenSpaceReflectionsPS(float4 SvPosition : SV_POSITION
 )
 {
 	float4 Color = 0;
-	ScreenSpaceReflections(SvPosition, Color);
+	float OutClosestHitDistance = 0;
+	ScreenSpaceReflections(SvPosition, OutColor
+#if SSR_OUTPUT_FOR_DENOISER
+		,OutClosestHitDistance
+#endif
+	);
 
 #if USE_GLES_FBF_DEFERRED
 	OutProxyAdditive = Color;
