Path Tracer: Optimize sheen LTC lobe in Substrate materials

Simplify logic for evaluating LTC sheen lobe, add importance sampling according to the LTC which is slightly cleaner for low roughness values.

The SheenLTC_Eval function no longer multiplies in the directional albedo, it must be applied afterwards if you need it. This is because it needs to be included in the weight, but not the pdf.

Make sure all Substrate BSDFs refer to the same shared functions, so they all agree about which cloth lobe and sampling technique is being used.

#rnx

diff --git a/Engine/Shaders/Private/BRDF.ush b/Engine/Shaders/Private/BRDF.ush
--- a/Engine/Shaders/Private/BRDF.ush
+++ b/Engine/Shaders/Private/BRDF.ush
@@ -534,7 +534,8 @@ float3 SheenLTC_Sample(float3 CosSample, float3 V, half3 N, float NoV, float Rou
 	const float4 LTC = SheenLTC_Cofficients(NoV, Roughness, InSheenLTCTexture, InSheenSampler);
 	const float aInv = LTC.x;
 	const float bInv = LTC.y;
-	const float3 LocalL = normalize(float3(CosSample.x / aInv - CosSample.z * bInv / aInv, CosSample.y / aInv, CosSample.z));
+	// NOTE: The vector is scaled by aInv compared to the reference implementation to avoid unecessary divides
+	const float3 LocalL = normalize(float3(CosSample.x - CosSample.z * bInv, CosSample.y, aInv * CosSample.z));
 
 	// Rotate L into tangent space
 	const float3 T1 = normalize(V - N * NoV);
@@ -545,6 +546,7 @@ float3 SheenLTC_Sample(float3 CosSample, float3 V, half3 N, float NoV, float Rou
 }
 
 // [Zeltner et al. 2022, "Practical Multiple-Scattering Sheen Using Linearly Transformed Cosines"]
+// The "OutDirectionalAlbedo" should be multiplied by the result (it is kept separate because in the sampling case it needs to be only included in the weight, not the pdf)
 float SheenLTC_Eval(float3 V, float3 L, half3 N, float NoV, float Roughness, Texture2D InSheenLTCTexture, SamplerState InSheenSampler, inout float OutDirectionalAlbedo)
 {
 	// Rotate L into tangent space
@@ -566,16 +568,16 @@ float SheenLTC_Eval(float3 V, float3 L, half3 N, float NoV, float Roughness, Tex
 		aInv * LocalL.y,
 		LocalL.z);
 
-	const float LenWrappedLocalL = length(WrappedLocalL);
+	const float InvLenWrappedLocalL2 = rcp(length2(WrappedLocalL));
 	// Jabobian is defined as: det(InvLTC) / ||InvLTC . L||
-	const float Jacobian = Pow2(aInv) / Pow3(LenWrappedLocalL);
+	// Jacobian should be aInv^2/Length^3 but since we skip the normalization
+	// we end up with aInv^2/Length^4 which is cheaper to compute
+	const float Jacobian = Pow2(aInv * InvLenWrappedLocalL2);
 
-	// Normalized wrapped direction
-	WrappedLocalL /= LenWrappedLocalL;
 
-	// Evaluation a normalized cos distribution
+	// Evaluation a normalized cos distribution (note: WrappedLocalL was not normalized, we folded the extra length division into the Jacobian)
 	const float CosDistribution = WrappedLocalL.z / PI;
-	const float Out = OutDirectionalAlbedo * CosDistribution * Jacobian; // Fs * CosTheta
+	const float Out = CosDistribution * Jacobian;
 
 	return max(Out, 0.f);
 }
diff --git a/Engine/Shaders/Private/PathTracing/Material/PathTracingSubstrate.ush b/Engine/Shaders/Private/PathTracing/Material/PathTracingSubstrate.ush
--- a/Engine/Shaders/Private/PathTracing/Material/PathTracingSubstrate.ush
+++ b/Engine/Shaders/Private/PathTracing/Material/PathTracingSubstrate.ush
@@ -125,7 +125,7 @@ FMaterialSample Substrate_SampleMaterial(
 	{
 		// cloth lobe
 		RandSample.x = RescaleRandomNumber(RandSample.x, Data.LobeCdf.z, 1.0);
-		L = UniformSampleHemisphere(RandSample.xy).xyz;
+		L = SampleSheenDirection(V, Data.FuzzRoughness, RandSample.xy);
 		H = normalize(L + V);
 	}
 
@@ -144,7 +144,7 @@ FMaterialSample Substrate_SampleMaterial(
 	const float2 GGXResult1 = GGXEvalReflection(L, V, H, Data.Alpha1);
 	const float Spec0Pdf = GGXResult0.y;
 	const float Spec1Pdf = GGXResult1.y;
-	const float ClothPdf = 1 / (2 * PI);
+	const float ClothPdf = ComputeSheenPdf(V_World, L_World, N_World, NoV, Data.FuzzRoughness);
 
 	const float ShadowTerminator = ShadowTerminatorTerm(L_World, N_World, Payload.WorldSmoothNormal);
 	if (bSampledDiffuse)
@@ -182,9 +182,9 @@ FMaterialSample Substrate_SampleMaterial(
 		float3 ClothWeight = 0;
 		if (NoV > 0 && NoL > 0)
 		{
-			ClothWeight = Data.FuzzAmount * ComputeSheenWeight(V_World, L_World, N_World, NoV, NoH, VoH, Data.FuzzRoughness, Payload.FuzzColor);
+			ClothWeight = Data.FuzzAmount * ComputeSheenWeight(V_World, L_World, N_World, NoV, NoL, NoH, VoH, Data.FuzzRoughness, Payload.FuzzColor);
 		}
-		Result.AddLobeWithMIS(ClothWeight * ShadowTerminator, 1 / (2.0 * PI), Data.LobePdf.w);
+		Result.AddLobeWithMIS(ClothWeight * ShadowTerminator, ClothPdf, Data.LobePdf.w);
 		Result.ScatterType = PATHTRACER_SCATTER_SPECULAR;
 
 		Result.Pdf += Data.LobePdf.x * DiffPdf;
@@ -242,11 +242,13 @@ FMaterialEval Substrate_EvalMaterial(
 	{
 		// Cloth Lobe
 		float3 ClothWeight = 0;
+		float ClothPdf = 0;
 		if (NoV > 0 && NoL > 0)
 		{
-			ClothWeight = Data.FuzzAmount * ComputeSheenWeight(V_World, L_World, N_World, NoV, NoH, VoH, Data.FuzzRoughness, Payload.FuzzColor);
+			ClothWeight = Data.FuzzAmount * ComputeSheenWeight(V_World, L_World, N_World, NoV, NoL, NoH, VoH, Data.FuzzRoughness, Payload.FuzzColor);
+			ClothPdf = ComputeSheenPdf(V_World, L_World, N_World, NoV, Data.FuzzRoughness);
 		}
-		Result.AddLobeWithMIS(DiffuseSpecularScale.y * ClothWeight * ShadowTerminator, 1 / (2.0 * PI), Data.LobePdf.w);
+		Result.AddLobeWithMIS(DiffuseSpecularScale.y * ClothWeight * ShadowTerminator, ClothPdf, Data.LobePdf.w);
 	}
 
 	Result.Weight *= SubstrateLobeWeight(Payload, NoL);
@@ -266,11 +268,11 @@ struct FSubstrateClearCoatData
 	float2 CoatAlpha;
 	float FuzzRoughness;
 
-	float FuzzAttenuation;
 	float CoatAttenuation;
 	FBxDFEnergyTermsRGB BaseSpec;
-	FBxDFEnergyTermsA CoatSpec;
-	FBxDFEnergyTermsA Fuzz;
+	FBxDFEnergyTermsA   CoatSpec;
+	float FuzzAttenuation;
+	float FuzzAmount;
 
 	float3 LobeCdf;
 	float4 LobePdf;
@@ -292,15 +294,16 @@ struct FSubstrateClearCoatData
 		// SUBSTRATE_TODO: Add support for F82 model?
 		BaseSpec = ComputeGGXSpecEnergyTermsRGB(Payload.RoughnessData.x, NoV, F0, F90);
 		CoatSpec = ComputeGGXSpecEnergyTermsA(Payload.RoughnessData.y, NoV, CLEAR_COAT_F0);
-		Fuzz = ComputeClothEnergyTermsA(FuzzRoughness, NoV);
+		FBxDFEnergyTermsA Fuzz = ComputeClothEnergyTermsA(FuzzRoughness, NoV);
 
 		FuzzAttenuation = 1.0 - Payload.FuzzAmount * Fuzz.E;
+		FuzzAmount      = Payload.FuzzAmount * Fuzz.W;
 		CoatAttenuation = 1.0 - Payload.RoughnessData.z * CoatSpec.E;
 		DiffuseWeight = FuzzAttenuation * CoatAttenuation * (1.0 - Luminance(BaseSpec.E));
 
 		const float3 BaseSpecAlbedo = FuzzAttenuation * CoatAttenuation * BaseSpec.E;
 		const float  CoatSpecAlbedo = FuzzAttenuation * Payload.RoughnessData.z * CoatSpec.E;
-		const float  FuzzAlbedo = Payload.FuzzAmount * Fuzz.E;
+		const float3 FuzzAlbedo = Payload.FuzzAmount * Fuzz.E * Payload.FuzzColor;
 
 		// Now prepare a cdf/pdf for lobe selection
 		float3 MaxLobeWeight = Payload.GetMaxLobeWeight();
@@ -363,7 +366,7 @@ FMaterialSample SubstrateClearCoat_SampleMaterial(
 	{
 		// cloth lobe
 		RandSample.x = RescaleRandomNumber(RandSample.x, Data.LobeCdf.z, 1.0);
-		L = UniformSampleHemisphere(RandSample.xy).xyz;
+		L = SampleSheenDirection(V, Data.FuzzRoughness, RandSample.xy);
 		H = normalize(L + V);
 	}
 
@@ -386,7 +389,7 @@ FMaterialSample SubstrateClearCoat_SampleMaterial(
 	const float2 CoatGGXResult = GGXEvalReflection(L, V, H, Data.CoatAlpha);
 	const float BaseSpecPdf = BaseGGXResult.y;
 	const float CoatSpecPdf = CoatGGXResult.y;
-	const float ClothPdf = 1 / (2 * PI);
+	const float ClothPdf = ComputeSheenPdf(V_World, L_World, N_World, NoV, Data.FuzzRoughness);
 
 	const float ShadowTerminator = ShadowTerminatorTerm(L_World, N_World, Payload.WorldSmoothNormal);
 	if (bSampledDiffuse)
@@ -420,10 +423,12 @@ FMaterialSample SubstrateClearCoat_SampleMaterial(
 	else
 	{
 		// Cloth Lobe
-		const float DCloth = D_Charlie(Data.FuzzRoughness, NoH);
-		const float3 FCloth = F_Schlick(Payload.FuzzColor, VoH);
-		const float3 ClothWeight = NoL > 0 && NoV > 0 ? (NoL * Payload.FuzzAmount * DCloth * Vis_Ashikhmin(NoV, NoL)) * FCloth * Data.Fuzz.W * (2 * PI) : 0.0;
-		Result.AddLobeWithMIS(ClothWeight * ShadowTerminator, 1 / (2.0 * PI), Data.LobePdf.w);
+		float3 ClothWeight = 0;
+		if (NoL > 0 && NoV > 0)
+		{
+			ClothWeight = Data.FuzzAmount * ComputeSheenWeight(V_World, L_World, N_World, NoV, NoL, NoH, VoH, Data.FuzzRoughness, Payload.FuzzColor);
+		}
+		Result.AddLobeWithMIS(ClothWeight * ShadowTerminator, ClothPdf, Data.LobePdf.w);
 		Result.ScatterType = PATHTRACER_SCATTER_SPECULAR;
 
 		Result.Pdf += Data.LobePdf.x * DiffPdf;
@@ -489,10 +494,14 @@ FMaterialEval SubstrateClearCoat_EvalMaterial(
 	}
 	{
 		// Cloth Lobe
-		const float DCloth = D_Charlie(Data.FuzzRoughness, NoH);
-		const float3 FCloth = F_Schlick(Payload.FuzzColor, VoH);
-		const float3 ClothWeight = NoL > 0 && NoV > 0 ? (NoL * Payload.FuzzAmount * DCloth * Vis_Ashikhmin(NoV, NoL)) * FCloth * Data.Fuzz.W * (2 * PI) : 0.0;
-		Result.AddLobeWithMIS(DiffuseSpecularScale.y * ClothWeight * ShadowTerminator, 1 / (2.0 * PI), Data.LobePdf.w);
+		float3 ClothWeight = 0;
+		float ClothPdf = 0;
+		if (NoV > 0 && NoL > 0)
+		{
+			ClothWeight = Data.FuzzAmount * ComputeSheenWeight(V_World, L_World, N_World, NoV, NoL, NoH, VoH, Data.FuzzRoughness, Payload.FuzzColor);
+			ClothPdf = ComputeSheenPdf(V_World, L_World, N_World, NoV, Data.FuzzRoughness);
+		}
+		Result.AddLobeWithMIS(DiffuseSpecularScale.y * ClothWeight * ShadowTerminator, ClothPdf, Data.LobePdf.w);
 	}
 
 	Result.Weight *= SubstrateLobeWeight(Payload, NoL);
diff --git a/Engine/Shaders/Private/PathTracing/Material/PathTracingSubstrateCommon.ush b/Engine/Shaders/Private/PathTracing/Material/PathTracingSubstrateCommon.ush
--- a/Engine/Shaders/Private/PathTracing/Material/PathTracingSubstrateCommon.ush
+++ b/Engine/Shaders/Private/PathTracing/Material/PathTracingSubstrateCommon.ush
@@ -60,24 +60,59 @@ FBxDFEnergyTermsA ComputeSheenEnergyTerms(float NoV, float SheenRoughness)
 #if SUBSTRATE_SHEEN_QUALITY == 1
 	FBxDFEnergyTermsA Result;
 	Result.E = SheenLTC_DirectionalAlbedo(NoV, SheenRoughness, View.SheenLTCTexture, View.SheenLTCSampler);
-	Result.W = 1;
+	Result.W = Result.E; // overall weight is the directional albedo
 	return Result;
 #else
 	return ComputeClothEnergyTermsA(SheenRoughness, NoV);
 #endif
 }
 
-float3 ComputeSheenWeight(float3 V, float3 L, float3 N, float NoV, float NoH, float VoH, float SheenRoughness, float3 SheenColor)
+#define USE_SHEEN_LTC_SAMPLER 1 // 0: fallback to uniform sampling (for reference)
+								// 1: use the LTC coefficients for importance sampling
+
+float ComputeSheenPdf(float3 V, float3 L, float3 N, float NoV, float SheenRoughness)
 {
-#if SUBSTRATE_SHEEN_QUALITY == 1
+#if USE_SHEEN_LTC_SAMPLER != 0 && SUBSTRATE_SHEEN_QUALITY == 1
 	// Disney's Cloth Lobe
 	float FuzzDirectionalAlbedo = 1;
-	float3 FuzzSpecularPathValueNoL = SheenColor * SheenLTC_Eval(V, L, N, NoV, SheenRoughness, View.SheenLTCTexture, View.SheenLTCSampler, FuzzDirectionalAlbedo);
-	return FuzzSpecularPathValueNoL * (2 * PI);
+	float FuzzSpecularPathValueNoL = SheenLTC_Eval(V, L, N, NoV, SheenRoughness, View.SheenLTCTexture, View.SheenLTCSampler, FuzzDirectionalAlbedo);
+	return FuzzSpecularPathValueNoL;
+#else
+	// uniform sampling (for legacy cloth lobe, or naive sampling of the LTC lobe)
+	const float Pdf = 1 / (2 * PI);
+	return Pdf;
+#endif
+
+}
+
+float3 ComputeSheenWeight(float3 V, float3 L, float3 N, float NoV, float NoL, float NoH, float VoH, float SheenRoughness, float3 SheenColor)
+{
+#if SUBSTRATE_SHEEN_QUALITY == 1
+	// Disney's Cloth Lobe
+	float FuzzDirectionalAlbedo = 1; // NOTE: Path Tracer includes this term in the weight applied outside
+	float FuzzSpecularPathValueNoL = SheenLTC_Eval(V, L, N, NoV, SheenRoughness, View.SheenLTCTexture, View.SheenLTCSampler, FuzzDirectionalAlbedo);
+#if USE_SHEEN_LTC_SAMPLER == 0
+	// using naive sampling, weight will include Brdf evaluation
+	// uniform sampling, pdf is just 1/2pi
+	return SheenColor * (FuzzSpecularPathValueNoL * 2 * PI);
+#else
+	// LTC sampling - pdf is equal to the brdf, so sample weight is just the color
+	return SheenColor;
+#endif
+
 #else
 	// Charlie's Cloth Lobe
 	const float DCloth = D_Charlie(SheenRoughness, NoH);
 	const float3 FCloth = F_Schlick(SheenColor, VoH);
-	return NoL * DCloth * Vis_Ashikhmin(NoV, NoL)) * FCloth * (2 * PI);
+	return (2 * PI * NoL * DCloth * Vis_Ashikhmin(NoV, NoL)) * FCloth;
+#endif
+}
+
+float3 SampleSheenDirection(float3 V, float SheenRoughness, float2 Rand)
+{
+#if USE_SHEEN_LTC_SAMPLER != 0 && SUBSTRATE_SHEEN_QUALITY == 1
+	return SheenLTC_Sample(CosineSampleHemisphere(Rand).xyz, V, float3(0, 0, 1), V.z, SheenRoughness, View.SheenLTCTexture, View.SheenLTCSampler);
+#else
+	return UniformSampleHemisphere(Rand).xyz;
 #endif
 }
diff --git a/Engine/Shaders/Private/PathTracing/Material/PathTracingSubstrateFoliage.ush b/Engine/Shaders/Private/PathTracing/Material/PathTracingSubstrateFoliage.ush
--- a/Engine/Shaders/Private/PathTracing/Material/PathTracingSubstrateFoliage.ush
+++ b/Engine/Shaders/Private/PathTracing/Material/PathTracingSubstrateFoliage.ush
@@ -137,7 +137,7 @@ FMaterialSample SubstrateFoliage_SampleMaterial(
 	{
 		// cloth lobe
 		RandSample.x = RescaleRandomNumber(RandSample.x, Data.LobeCdf.z, 1.0);
-		L = UniformSampleHemisphere(RandSample.xy).xyz;
+		L = SampleSheenDirection(V, Data.FuzzRoughness, RandSample.xy);
 		H = normalize(L + V);
 	}
 
@@ -151,14 +151,12 @@ FMaterialSample SubstrateFoliage_SampleMaterial(
 
 	FMaterialSample Result = CreateMaterialSample(L_World, 0.0, 0.0, 1.0, OutRoughness, PATHTRACER_SCATTER_DIFFUSE);
 
-
-
 	const float DiffPdf = NoL / PI;
 	float2 GGXResult0 = GGXEvalReflection(L, V, H, Data.Alpha0);
 	const float2 GGXResult1 = GGXEvalReflection(L, V, H, Data.Alpha1);
 	const float Spec0Pdf = GGXResult0.y;
 	const float Spec1Pdf = GGXResult1.y;
-	const float ClothPdf = 1 / (2 * PI);
+	const float ClothPdf = ComputeSheenPdf(V_World, L_World, N_World, NoV, Data.FuzzRoughness);
 
 	const float ShadowTerminator = ShadowTerminatorTerm(L_World, N_World, Payload.WorldSmoothNormal);
 	if (bSampledTransmit)
@@ -195,8 +193,12 @@ FMaterialSample SubstrateFoliage_SampleMaterial(
 	else
 	{
 		// Cloth Lobe
-		float3 ClothWeight = Data.FuzzAmount * ComputeSheenWeight(V_World, L_World, N_World, NoV, NoH, VoH, Data.FuzzRoughness, Payload.FuzzColor);
-		Result.AddLobeWithMIS(ClothWeight * ShadowTerminator, 1 / (2.0 * PI), Data.LobePdf.w);
+		float3 ClothWeight = 0;
+		if (NoV > 0 && NoL > 0)
+		{
+			ClothWeight = Data.FuzzAmount * ComputeSheenWeight(V_World, L_World, N_World, NoV, NoL, NoH, VoH, Data.FuzzRoughness, Payload.FuzzColor);
+		}
+		Result.AddLobeWithMIS(ClothWeight * ShadowTerminator, ClothPdf, Data.LobePdf.w);
 		Result.ScatterType = PATHTRACER_SCATTER_SPECULAR;
 
 		Result.Pdf += Data.LobePdf.x * DiffPdf;
@@ -260,8 +262,9 @@ FMaterialEval SubstrateFoliage_EvalMaterial(
 		}
 		{
 			// Cloth Lobe
-			float3 ClothWeight = Data.FuzzAmount * ComputeSheenWeight(V_World, L_World, N_World, NoV, NoH, VoH, Data.FuzzRoughness, Payload.FuzzColor);
-			Result.AddLobeWithMIS(DiffuseSpecularScale.y * ClothWeight * ShadowTerminator, 1 / (2.0 * PI), Data.LobePdf.w);
+			float3 ClothWeight = Data.FuzzAmount * ComputeSheenWeight(V_World, L_World, N_World, NoV, NoL, NoH, VoH, Data.FuzzRoughness, Payload.FuzzColor);
+			float ClothPdf = ComputeSheenPdf(V_World, L_World, N_World, NoV, Data.FuzzRoughness);
+			Result.AddLobeWithMIS(DiffuseSpecularScale.y * ClothWeight * ShadowTerminator, ClothPdf, Data.LobePdf.w);
 		}
 	}
 
diff --git a/Engine/Shaders/Private/PathTracing/Material/PathTracingSubstrateSolidGlass.ush b/Engine/Shaders/Private/PathTracing/Material/PathTracingSubstrateSolidGlass.ush
--- a/Engine/Shaders/Private/PathTracing/Material/PathTracingSubstrateSolidGlass.ush
+++ b/Engine/Shaders/Private/PathTracing/Material/PathTracingSubstrateSolidGlass.ush
@@ -58,7 +58,7 @@ struct FSubstrateSolidGlassSlabData
 
 		const float3 Glass0Albedo = FuzzAttenuation * (1.0 - Payload.RoughnessData.z);
 		const float3 Glass1Albedo = FuzzAttenuation * (      Payload.RoughnessData.z);
-		const float  FuzzAlbedo = Payload.FuzzAmount * Fuzz.E;
+		const float3 FuzzAlbedo = Payload.FuzzAmount * Fuzz.E * Payload.FuzzColor;
 
 		// Now prepare a cdf/pdf for lobe selection
 		float3 MaxLobeWeight = Payload.GetMaxLobeWeight();
@@ -133,7 +133,7 @@ FMaterialSample SubstrateSolidGlass_SampleMaterial(
 
 			// reflection side, apply full lobe weight and attenuation along L
 			// and account for probability of picking the fuzz layer
-			const float ClothPdf = 1 / (2 * PI);
+			const float ClothPdf = ComputeSheenPdf(V, L, float3(0, 0, 1), NoV, Data.FuzzRoughness);
 			Result.Pdf += Data.LobePdf.z * ClothPdf;
 			Result.Weight *= Data.FuzzAttenuation * SubstrateLobeWeight(Payload, saturate(L.z)) * FTint;
 		}
@@ -144,7 +144,7 @@ FMaterialSample SubstrateSolidGlass_SampleMaterial(
 	{
 		// cloth lobe
 		RandSample.x = RescaleRandomNumber(RandSample.x, Data.LobeCdf.y, 1.0);
-		L = UniformSampleHemisphere(RandSample.xy).xyz;
+		L = SampleSheenDirection(V, Data.FuzzRoughness, RandSample.xy);
 		H = normalize(L + V);
 
 		// transform to world space
@@ -161,20 +161,19 @@ FMaterialSample SubstrateSolidGlass_SampleMaterial(
 		const float Fg = FresnelReflectance(VoH, Data.Eta, Data.F0);
 		const float Spec0Pdf = Fg * GGXResult0.y;
 		const float Spec1Pdf = Fg * GGXResult1.y;
-		const float ClothPdf = 1 / (2 * PI);
 
 		FMaterialSample Result = CreateMaterialSample(L_World, 0.0, 0.0, 1.0, 1.0, PATHTRACER_SCATTER_SPECULAR);
 
 		// Cloth Lobe
 		const float ShadowTerminator = ShadowTerminatorTerm(L_World, N_World, Payload.WorldSmoothNormal);
-		const float DCloth = D_Charlie(Data.FuzzRoughness, NoH);
-		const float3 FCloth = F_Schlick(Payload.FuzzColor, VoH);
 		float3 ClothWeight = 0;
+		float ClothPdf = 0;
 		if (NoV > 0 && NoL > 0)
 		{
-			ClothWeight = Data.FuzzAmount * ComputeSheenWeight(V_World, L_World, N_World, NoV, NoH, VoH, Data.FuzzRoughness, Payload.FuzzColor);
+			ClothWeight = Data.FuzzAmount * ComputeSheenWeight(V_World, L_World, N_World, NoV, NoL, NoH, VoH, Data.FuzzRoughness, Payload.FuzzColor);
+			ClothPdf = ComputeSheenPdf(V_World, L_World, N_World, NoV, Data.FuzzRoughness);
 		}
-		Result.AddLobeWithMIS(ClothWeight * ShadowTerminator, 1 / (2.0 * PI), Data.LobePdf.z);
+		Result.AddLobeWithMIS(ClothWeight * ShadowTerminator, ClothPdf, Data.LobePdf.z);
 		Result.ScatterType = PATHTRACER_SCATTER_SPECULAR;
 
 		Result.Pdf += Data.LobePdf.y * Spec0Pdf;
@@ -228,11 +227,13 @@ FMaterialEval SubstrateSolidGlass_EvalMaterial(
 			// Cloth Lobe
 			const float ShadowTerminator = ShadowTerminatorTerm(L_World, N_World, Payload.WorldSmoothNormal);
 			float3 ClothWeight = 0;
+			float ClothPdf = 0;
 			if (NoV > 0 && NoL > 0)
 			{
-				ClothWeight = Data.FuzzAmount * ComputeSheenWeight(V_World, L_World, N_World, NoV, NoH, VoH, Data.FuzzRoughness, Payload.FuzzColor);
+				ClothWeight = Data.FuzzAmount * ComputeSheenWeight(V_World, L_World, N_World, NoV, NoL, NoH, VoH, Data.FuzzRoughness, Payload.FuzzColor);
+				ClothPdf = ComputeSheenPdf(V_World, L_World, N_World, NoV, Data.FuzzRoughness);
 			}
-			Result.AddLobeWithMIS(DiffuseSpecularScale.y * ClothWeight * ShadowTerminator, 1 / (2.0 * PI), Data.LobePdf.z);
+			Result.AddLobeWithMIS(DiffuseSpecularScale.y * ClothWeight * ShadowTerminator, ClothPdf, Data.LobePdf.z);
 		}
 
 		Result.Weight *= SubstrateLobeWeight(Payload, NoL);
diff --git a/Engine/Shaders/Private/PathTracing/Material/PathTracingSubstrateThinGlass.ush b/Engine/Shaders/Private/PathTracing/Material/PathTracingSubstrateThinGlass.ush
--- a/Engine/Shaders/Private/PathTracing/Material/PathTracingSubstrateThinGlass.ush
+++ b/Engine/Shaders/Private/PathTracing/Material/PathTracingSubstrateThinGlass.ush
@@ -70,7 +70,7 @@ struct FSubstrateThinGlassSlabData
 
 		const float3 GlassR = FuzzAttenuation * lerp(SpecR0.E * NudgeE.x, SpecR1.E * NudgeE.y, Payload.RoughnessData.z);
 		const float3 GlassT = FuzzAttenuation * lerp(SpecT0.E * NudgeE.x, SpecT1.E * NudgeE.y, Payload.RoughnessData.z);
-		const float  FuzzAlbedo = Payload.FuzzAmount * Fuzz.E;
+		const float3 FuzzAlbedo = Payload.FuzzAmount * Fuzz.E * Payload.FuzzColor;
 
 		SpecR0W = SpecR0.W * NudgeE.x;
 		SpecR1W = SpecR1.W * NudgeE.y;
@@ -137,7 +137,7 @@ FMaterialSample SubstrateThinGlass_SampleMaterial(
 	{
 		// cloth lobe
 		RandSample.x = RescaleRandomNumber(RandSample.x, Data.LobeCdf.y, 1.0);
-		L = UniformSampleHemisphere(RandSample.xy).xyz;
+		L = SampleSheenDirection(V, Data.FuzzRoughness, RandSample.xy);
 		H = normalize(L + V);
 	}
 	// transform to world space
@@ -157,7 +157,7 @@ FMaterialSample SubstrateThinGlass_SampleMaterial(
 	const float2 GGXResult0 = GGXEvalReflection(L, V, H, bIsReflection ? Data.AlphaR0 : Data.AlphaT0);
 	const float2 GGXResult1 = GGXEvalReflection(L, V, H, bIsReflection ? Data.AlphaR1 : Data.AlphaT1);
 	const float LobePdf = bIsReflection ? Data.LobePdf.y : Data.LobePdf.x;
-	const float ClothPdf = 1.0 / (2 * PI);
+	const float ClothPdf = ComputeSheenPdf(V_World, L_World, N_World, NoV, Data.FuzzRoughness);
 	if (bSampleGlass)
 	{
 		FThinSlabWeights SlabResult = ComputeThinSlabWeights(Data.TransmissionColor, VoH, Payload.Ior, Data.F0);
@@ -181,9 +181,10 @@ FMaterialSample SubstrateThinGlass_SampleMaterial(
 		// Cloth Lobe
 		const float ShadowTerminator = ShadowTerminatorTerm(L_World, N_World, Payload.WorldSmoothNormal);
 		float3 ClothWeight = 0;
+		float ClothPdf = 0;
 		if (NoV > 0 && NoL > 0)
 		{
-			ClothWeight = Data.FuzzAmount * ComputeSheenWeight(V_World, L_World, N_World, NoV, NoH, VoH, Data.FuzzRoughness, Payload.FuzzColor);
+			ClothWeight = Data.FuzzAmount * ComputeSheenWeight(V_World, L_World, N_World, NoV, NoL, NoH, VoH, Data.FuzzRoughness, Payload.FuzzColor);
 		}
 		Result.AddLobeWithMIS(ClothWeight * ShadowTerminator, ClothPdf, Data.LobePdf.z);
 
@@ -254,11 +255,13 @@ FMaterialEval SubstrateThinGlass_EvalMaterial(
 			// Cloth Lobe
 			const float ShadowTerminator = ShadowTerminatorTerm(L_World, N_World, Payload.WorldSmoothNormal);
 			float3 ClothWeight = 0;
+			float ClothPdf = 0;
 			if (NoV > 0 && NoL > 0)
 			{
-				ClothWeight = Data.FuzzAmount * ComputeSheenWeight(V_World, L_World, N_World, NoV, NoH, VoH, Data.FuzzRoughness, Payload.FuzzColor);
+				ClothWeight = Data.FuzzAmount * ComputeSheenWeight(V_World, L_World, N_World, NoV, NoL, NoH, VoH, Data.FuzzRoughness, Payload.FuzzColor);
+				ClothPdf = ComputeSheenPdf(V_World, L_World, N_World, NoV, Data.FuzzRoughness);
 			}
-			Result.AddLobeWithMIS(DiffuseSpecularScale.y * ClothWeight * ShadowTerminator, 1 / (2.0 * PI), Data.LobePdf.z);
+			Result.AddLobeWithMIS(DiffuseSpecularScale.y * ClothWeight * ShadowTerminator, ClothPdf, Data.LobePdf.z);
 		}
 		Result.Weight *= SubstrateLobeWeight(Payload, NoL);
 	}
diff --git a/Engine/Shaders/Private/PathTracing/PathTracingMaterialHitShader.usf b/Engine/Shaders/Private/PathTracing/PathTracingMaterialHitShader.usf
--- a/Engine/Shaders/Private/PathTracing/PathTracingMaterialHitShader.usf
+++ b/Engine/Shaders/Private/PathTracing/PathTracingMaterialHitShader.usf
@@ -43,6 +43,7 @@
 #if SUBSTRATE_ENABLED
 #define MATERIAL_SUBSTRATE_OPAQUE_PRECOMPUTED_LIGHTING 0
 #include "/Engine/Private/Substrate/SubstrateExport.ush"
+#include "/Engine/Private/PathTracing/Material/PathTracingSubstrateCommon.ush"
 #endif
 
 float AdjustMaterialRoughness(float Roughness, float PathRoughness)
@@ -470,7 +471,7 @@ RAY_TRACING_ENTRY_CLOSEST_HIT(PathTracingMaterialCHS,
 								FuzzRoughness = MakeRoughnessSafe(FuzzRoughness, SUBSTRATE_MIN_FUZZ_ROUGHNESS);
 							}
 							
-							const float FuzzE = ComputeClothEnergyTermsA(FuzzRoughness, NoV).E;
+							const float FuzzE = ComputeSheenEnergyTerms(FuzzRoughness, NoV).E;
 							const float FuzzAttenuation = 1.0 - FuzzAmount * FuzzE;
 							float DiffuseWeight = FuzzAttenuation;
 							float3 Spec0Albedo = FuzzAttenuation * Spec0E;
@@ -491,7 +492,7 @@ RAY_TRACING_ENTRY_CLOSEST_HIT(PathTracingMaterialCHS,
 								Spec1Albedo *=       SpecBlend;
 							}
 							// SUBSTRATE_TODO: Need to account for glass case here for lobe selection ....
-							const float  FuzzAlbedo = FuzzAmount * FuzzE;
+							const float3 FuzzAlbedo = FuzzAmount * FuzzE * FuzzColor;
 							const float3 SlabAlbedo = MaxLobeWeight * (DiffuseWeight * DiffuseColor + Spec0Albedo + Spec1Albedo + FuzzAlbedo);
 							Pdf = LobeColorToWeight(SlabAlbedo);
 						}
diff --git a/Engine/Shaders/Private/Substrate/SubstrateEvaluation.ush b/Engine/Shaders/Private/Substrate/SubstrateEvaluation.ush
--- a/Engine/Shaders/Private/Substrate/SubstrateEvaluation.ush
+++ b/Engine/Shaders/Private/Substrate/SubstrateEvaluation.ush
@@ -727,6 +727,7 @@ FSubstrateEvaluateResult SubstrateEvaluateBSDFCommon(FSubstrateBSDFContext BSDFC
 					else
 					{
 						ClothSpecularPathValueNoL = FuzzF0 * SheenLTC_Eval(BSDFContext.V, AreaLightContext.L, BSDFContext.N, BSDFContext.Context.NoV, FuzzRoughness, View.SheenLTCTexture, View.SheenLTCSampler, DirectionalAlbedo);
+						ClothSpecularPathValueNoL *= DirectionalAlbedo;
 					}
 					EnergyTerms.E = DirectionalAlbedo;
 					EnergyTerms.W = 1.f;
