Fix Mobile SSR

diff --git a/Engine/Shaders/Private/SSRT/SSRTReflections.usf b/Engine/Shaders/Private/SSRT/SSRTReflections.usf
--- a/Engine/Shaders/Private/SSRT/SSRTReflections.usf
+++ b/Engine/Shaders/Private/SSRT/SSRTReflections.usf
@@ -484,7 +484,7 @@ void MobileScreenSpaceReflectionsPS(float4 SvPosition : SV_POSITION
 {
 	float4 Color = 0;
 	float OutClosestHitDistance = 0;
-	ScreenSpaceReflections(SvPosition, OutColor
+	ScreenSpaceReflections(SvPosition, Color
 #if SSR_OUTPUT_FOR_DENOISER
 		,OutClosestHitDistance
 #endif
