Screen traces were incorrectly enabled for hair strands lighting.

diff --git a/Engine/Shaders/Private/DeferredLightPixelShaders.usf b/Engine/Shaders/Private/DeferredLightPixelShaders.usf
--- a/Engine/Shaders/Private/DeferredLightPixelShaders.usf
+++ b/Engine/Shaders/Private/DeferredLightPixelShaders.usf
@@ -4,7 +4,11 @@
 	DeferredLightPixelShaders.usf:
 =============================================================================*/
   
-#define SUPPORT_CONTACT_SHADOWS 1
+#if USE_HAIR_LIGHTING == 1
+  #define SUPPORT_CONTACT_SHADOWS 0
+#else
+  #define SUPPORT_CONTACT_SHADOWS 1
+#endif
 #define SUBSTRATE_SSS_TRANSMISSION USE_TRANSMISSION
 
 #if SUBSTRATE_ENABLED 
