Fix mobile merged light on translucent surfaces causing NaN due to missing distance scaling.

diff --git a/Engine/Shaders/Private/MobileBasePassPixelShader.usf b/Engine/Shaders/Private/MobileBasePassPixelShader.usf
--- a/Engine/Shaders/Private/MobileBasePassPixelShader.usf
+++ b/Engine/Shaders/Private/MobileBasePassPixelShader.usf
@@ -22,6 +22,10 @@
 #define MOBILE_QL_FORCE_DISABLE_PREINTEGRATEDGF 0
 #endif
 
+#if SUBSTRATE_ENABLED && MERGED_LOCAL_LIGHTS_MOBILE == 2
+#define ALLOW_LOCAL_LIGHT_DISTANCE_ATTENUATION 1
+#endif
+
 #define UseBasePassSkylightDiffuse  (1)
 #define UseBasePassSkylightSpecular (MobileBasePass.ReflectionsParameters.SkyLightParameters.y)
 
