Disable by default textured rect light on mobile platforms

diff --git a/Engine/Shaders/Private/RectLight.ush b/Engine/Shaders/Private/RectLight.ush
--- a/Engine/Shaders/Private/RectLight.ush
+++ b/Engine/Shaders/Private/RectLight.ush
@@ -8,7 +8,7 @@
 #include "LightData.ush"
 
 #ifndef USE_SOURCE_TEXTURE
-#define USE_SOURCE_TEXTURE 1
+#define USE_SOURCE_TEXTURE (FEATURE_LEVEL > FEATURE_LEVEL_ES3_1)
 #endif
 
 #if SUPPORTS_INDEPENDENT_SAMPLERS
