Substrate - added analytic version of Sheen LTC. disabled for now until we measure on platform.

diff --git a/Engine/Shaders/Private/BRDF.ush b/Engine/Shaders/Private/BRDF.ush
--- a/Engine/Shaders/Private/BRDF.ush
+++ b/Engine/Shaders/Private/BRDF.ush
@@ -516,11 +516,28 @@ float3 F_Fresnel( float3 SpecularColor, float VoH )
 
 float4 SheenLTC_Cofficients(float NoV, float Roughness, Texture2D InSheenLTCTexture, SamplerState InSheenSampler)
 {
+#if 1
 	const float Alpha = sqrt(Roughness);
 	const float SatNoV = saturate(abs(NoV) + 1e-5);
 	float2 UV = float2(Alpha, SatNoV);
 	UV = UV * (31.0 / 32.0) + (0.5 / 32.0);
 	return InSheenLTCTexture.SampleLevel(SharedSheenLTCSampler, UV, 0);
+#else
+	// Sheen model LUT approximation with ALU under the MaterialX TPS.
+	// SUBSTRATE_TODO: measure if this is worth it on some platforms.
+	float x = saturate(abs(NoV) + 1e-5);
+	float y = sqrt(Roughness);
+
+	float A = ((2.58126 * x + 0.813703 * y) * y) / (1.0 + 0.310327 * x * x + 2.60994 * x * y);
+	float B = sqrt(1.0 - x) * (y - 1.0) * y * y * y / (0.0000254053 + 1.71228 * x - 1.71506 * x * y + 1.34174 * y * y);
+
+	float s = y * (0.0206607 + 1.58491 * y) / (0.0379424 + y * (1.32227 + y));
+	float m = y * (-0.193854 + y * (-1.14885 + y * (1.7932 - 0.95943 * y * y))) / (0.046391 + y);
+	float o = y * (0.000654023 + (-0.0207818 + 0.119681 * y) * y) / (1.26264 + y * (-1.92021 + y));
+	float R = exp(-0.5 * Square((x - m) / s)) / (s * sqrt(2.0 * PI)) + o;
+
+	return float4(A, B, R, 0.0f);
+#endif
 }
 
 float SheenLTC_DirectionalAlbedo(float NoV, float Roughness, Texture2D InSheenLTCTexture, SamplerState InSheenSampler)
