Implemented proper support for the unlit view mode on SingleLayerWater both on mobile and desktop/console. Also fixed emissive being ignored in the unlit view mode on mobile.

#rnx

diff --git a/Engine/Shaders/Private/BasePassPixelShader.usf b/Engine/Shaders/Private/BasePassPixelShader.usf
--- a/Engine/Shaders/Private/BasePassPixelShader.usf
+++ b/Engine/Shaders/Private/BasePassPixelShader.usf
@@ -1488,8 +1488,12 @@ void FPixelShaderInOut_MainPS(
 #if MATERIAL_SHADINGMODEL_SINGLELAYERWATER || SUBSTRATE_INLINE_SINGLELAYERWATER
 	{
 		const bool CameraIsUnderWater = false;	// Fade out the material contribution over to water contribution according to material opacity.
-		const float3 SunIlluminance = ResolvedView.DirectionalLightColor.rgb * PI;	// times PI because it is divided by PI on CPU (=luminance) and we want illuminance here. 
-		const float3 WaterDiffuseIndirectIlluminance = DiffuseIndirectLighting * PI;// DiffuseIndirectLighting is luminance. So we need to multiply by PI to get illuminance.
+		float3 SunIlluminance = ResolvedView.DirectionalLightColor.rgb * PI;	// times PI because it is divided by PI on CPU (=luminance) and we want illuminance here. 
+		float3 WaterDiffuseIndirectIlluminance = DiffuseIndirectLighting * PI;// DiffuseIndirectLighting is luminance. So we need to multiply by PI to get illuminance.
+#if USE_DEVELOPMENT_SHADERS
+		SunIlluminance = lerp(SunIlluminance, 0.0f, View.UnlitViewmodeMask);
+		WaterDiffuseIndirectIlluminance = lerp(WaterDiffuseIndirectIlluminance, PI, View.UnlitViewmodeMask);
+#endif
 		const bool bSeparateWaterMainDirLightLuminance = (SINGLE_LAYER_WATER_SEPARATED_MAIN_LIGHT > 0) && SingleLayerWater.bSeparateMainDirLightLuminance;
 
 		// Evaluate Fresnel effect
diff --git a/Engine/Shaders/Private/MobileBasePassPixelShader.usf b/Engine/Shaders/Private/MobileBasePassPixelShader.usf
--- a/Engine/Shaders/Private/MobileBasePassPixelShader.usf
+++ b/Engine/Shaders/Private/MobileBasePassPixelShader.usf
@@ -851,6 +851,7 @@ void Main(
 	// We also keep the SpecularColor for sun/water interactions.
 	GBuffer.DiffuseColor *= BaseMaterialCoverageOverWater;
 	DiffuseColor *= BaseMaterialCoverageOverWater;
+	DiffuseColorForIndirect *= BaseMaterialCoverageOverWater;
 #endif
 
 	FLightAccumulator DirectLighting = (FLightAccumulator)0;
@@ -981,8 +982,13 @@ void Main(
 		float3 SeparatedWaterMainDirLightLuminance = float3(0, 0, 0);
 
 		const bool CameraIsUnderWater = false;	// Fade out the material contribution over to water contribution according to material opacity.
-		const float3 SunIlluminance = ResolvedView.DirectionalLightColor.rgb * PI;			// times PI because it is divided by PI on CPU (=luminance) and we want illuminance here. 
-		const float3 WaterDiffuseIndirectIlluminance = WaterDiffuseIndirectLuminance * PI;	// DiffuseIndirectLighting is luminance. So we need to multiply by PI to get illuminance.
+		float3 SunIlluminance = ResolvedView.DirectionalLightColor.rgb * PI;			// times PI because it is divided by PI on CPU (=luminance) and we want illuminance here. 
+		float3 WaterDiffuseIndirectIlluminance = WaterDiffuseIndirectLuminance * PI;	// DiffuseIndirectLighting is luminance. So we need to multiply by PI to get illuminance.
+#if MOBILE_EMULATION
+		SunIlluminance = lerp(SunIlluminance, 0.0f, View.UnlitViewmodeMask);
+		WaterDiffuseIndirectIlluminance = lerp(WaterDiffuseIndirectIlluminance, PI, View.UnlitViewmodeMask);
+#endif
+
 		half3 N = GBuffer.WorldNormal;
 		half3 V = MaterialParameters.CameraVector;
 		half NoV = saturate(abs(dot(N, V)) + 1e-5);
@@ -1000,6 +1006,9 @@ void Main(
 		// Accumulate luminance and occlude the background according to transmittance to view and mean transmittance to lights.
 		LightAccumulator_AddSplit(DirectLighting, WaterLighting.Luminance, 0.0f, WaterLighting.Luminance, 1.0f, false);
 		Opacity = 1.0 - ((1.0 - Opacity) * dot(WaterLighting.WaterToSceneToLightTransmittance, float3(1.0 / 3.0, 1.0 / 3.0, 1.0 / 3.0)));
+#if MOBILE_EMULATION
+		DiffuseColorForIndirect += WaterLighting.Luminance;
+#endif
 	}
 #endif // MATERIAL_SHADINGMODEL_SINGLELAYERWATER
 
@@ -1026,15 +1035,15 @@ void Main(
 #endif
 #endif
 	// NEEDS_BASEPASS_PIXEL_FOGGING is not allowed on mobile for the sake of performance.
-				 
-	half3 Emissive = GetMaterialEmissive(PixelMaterialInputs);
-
-	Color += Emissive;
-
+	
 #if !MATERIAL_SHADINGMODEL_UNLIT && MOBILE_EMULATION
 	Color = lerp(Color, DiffuseColorForIndirect, ResolvedView.UnlitViewmodeMask);
 #endif
 
+	half3 Emissive = GetMaterialEmissive(PixelMaterialInputs);
+
+	Color += Emissive;
+
 #if MATERIAL_SHADINGMODEL_THIN_TRANSLUCENT && !SUBSTRATE_ENABLED
 	half3 DualBlendSurfaceLuminancePostCoverage = 0.0f;
 	half3 DualBlendSurfaceTransmittancePreCoverage = 1.0f;
diff --git a/Engine/Source/Runtime/Renderer/Private/SingleLayerWaterRendering.cpp b/Engine/Source/Runtime/Renderer/Private/SingleLayerWaterRendering.cpp
--- a/Engine/Source/Runtime/Renderer/Private/SingleLayerWaterRendering.cpp
+++ b/Engine/Source/Runtime/Renderer/Private/SingleLayerWaterRendering.cpp
@@ -178,6 +178,7 @@ bool ShouldRenderScreenSpaceReflectionsWater(const FViewInfo& View)
 	const bool bSSRDefault = ReflectionsMethod == ESingleLayerWaterReflections::Enabled && View.FinalPostProcessSettings.ReflectionMethod != EReflectionMethod::None;
 
 	if (!View.Family->EngineShowFlags.ScreenSpaceReflections
+		|| !View.Family->EngineShowFlags.Lighting
 		|| (!bSSROverride && !bSSRDefault)
 		|| HasRayTracedOverlay(*View.Family)
 		|| !View.State /*no view state(e.g.thumbnail rendering ? ), no HZB(no screen space reflections or occlusion culling)*/
@@ -202,7 +203,7 @@ bool ShouldRenderScreenSpaceReflectionsWater(const FViewInfo& View)
 bool ShouldRenderLumenReflectionsWater(const FViewInfo& View, bool bSkipTracingDataCheck, bool bSkipProjectCheck)
 {
 	// This only returns true if using the default reflections method and having Lumen enabled in the scene. It can't be forced with r.Water.SingleLayer.Reflection.
-	return !View.bIsReflectionCapture 
+	return !View.bIsReflectionCapture  && View.Family->EngineShowFlags.Lighting
 		&& GetSingleLayerWaterReflectionTechnique() == ESingleLayerWaterReflections::Enabled
 		&& ShouldRenderLumenReflections(View, bSkipTracingDataCheck, bSkipProjectCheck);
 }
@@ -987,7 +988,7 @@ void FDeferredShadingSceneRenderer::RenderSingleLayerWaterReflections(
 		}
 
 		// ReflectionsMethodWater can also be Disabled when only reflection captures are requested, so check CVarWaterSingleLayerReflection directly before early exiting.
-		if (GetSingleLayerWaterReflectionTechnique() == ESingleLayerWaterReflections::Disabled)
+		if (GetSingleLayerWaterReflectionTechnique() == ESingleLayerWaterReflections::Disabled || !View.Family->EngineShowFlags.Lighting)
 		{
 			continue;
 		}
