Height Fog - added end distance on the XY plane to better control the heigh fog opacity without using maxopacity which does not look good.

diff --git a/Engine/Shaders/Private/HeightFogCommon.ush b/Engine/Shaders/Private/HeightFogCommon.ush
--- a/Engine/Shaders/Private/HeightFogCommon.ush
+++ b/Engine/Shaders/Private/HeightFogCommon.ush
@@ -71,6 +71,7 @@ struct FogStructData
 	float4 DirectionalInscatteringColor;
 	float2 SinCosInscatteringColorCubemapRotation;
 	float3 FogInscatteringTextureParameters;
+	float EndDistance;
 	float ApplyVolumetricFog;
 	float VolumetricFogStartDistance;
 	float VolumetricFogNearFadeInDistanceInv;
@@ -97,6 +98,7 @@ FogStructData GetFogStructData(uint EyeIndex)
 		Out.DirectionalInscatteringColor = FogStruct.DirectionalInscatteringColor;
 		Out.SinCosInscatteringColorCubemapRotation = FogStruct.SinCosInscatteringColorCubemapRotation;
 		Out.FogInscatteringTextureParameters = FogStruct.FogInscatteringTextureParameters;
+		Out.EndDistance = FogStruct.EndDistance;
 		Out.ApplyVolumetricFog = FogStruct.ApplyVolumetricFog;
 		Out.VolumetricFogStartDistance = FogStruct.VolumetricFogStartDistance;
 		Out.VolumetricFogNearFadeInDistanceInv = FogStruct.VolumetricFogNearFadeInDistanceInv;
@@ -116,6 +118,7 @@ FogStructData GetFogStructData(uint EyeIndex)
 		Out.DirectionalInscatteringColor = FogStructISR.DirectionalInscatteringColor;
 		Out.SinCosInscatteringColorCubemapRotation = FogStructISR.SinCosInscatteringColorCubemapRotation;
 		Out.FogInscatteringTextureParameters = FogStructISR.FogInscatteringTextureParameters;
+		Out.EndDistance = FogStructISR.EndDistance;
 		Out.ApplyVolumetricFog = FogStructISR.ApplyVolumetricFog;
 		Out.VolumetricFogStartDistance = FogStructISR.VolumetricFogStartDistance;
 		Out.VolumetricFogNearFadeInDistanceInv = FogStructISR.VolumetricFogNearFadeInDistanceInv;
@@ -218,6 +221,13 @@ half4 GetExponentialHeightFog(float3 WorldPositionRelativeToCamera, float Exclud
 
 	const float3 WorldObserverOrigin = float3(WorldCameraOrigin.xy, min(WorldCameraOrigin.z, MaxWorldObserverHeight)); // Clamp z to max height
 
+	// Apply end fog distance from view projected on the XY plane.
+	const float CameraToReceiverLenXYSqr = dot(CameraToReceiver.xy, CameraToReceiver.xy);
+	if (FogData.EndDistance > 0.0f && CameraToReceiverLenXYSqr > (FogData.EndDistance * FogData.EndDistance))
+	{
+		CameraToReceiver *= FogData.EndDistance / sqrt(max(1.0, CameraToReceiverLenXYSqr));
+	}
+
 	CameraToReceiver.z += WorldCameraOrigin.z - WorldObserverOrigin.z; // Compensate this vector for clamping the observer height
 	float CameraToReceiverLengthSqr = dot(CameraToReceiver, CameraToReceiver);
 	float CameraToReceiverLengthInv = rsqrt(max(CameraToReceiverLengthSqr, 0.00000001f));
diff --git a/Engine/Source/Runtime/Engine/Classes/Components/ExponentialHeightFogComponent.h b/Engine/Source/Runtime/Engine/Classes/Components/ExponentialHeightFogComponent.h
--- a/Engine/Source/Runtime/Engine/Classes/Components/ExponentialHeightFogComponent.h
+++ b/Engine/Source/Runtime/Engine/Classes/Components/ExponentialHeightFogComponent.h
@@ -132,6 +132,10 @@ class UExponentialHeightFogComponent : public USceneComponent
 	UPROPERTY(BlueprintReadOnly, interp, Category=ExponentialHeightFogComponent, meta=(UIMin = "0", UIMax = "5000"))
 	float StartDistance;
 
+	/** Distance from the camera, on the horizontal XY plane, that the fog will end integrating the lighting and transmittance. Disabled when 0. */
+	UPROPERTY(BlueprintReadOnly, interp, Category = ExponentialHeightFogComponent, meta = (UIMin = "0", UIMax = "500000"))
+	float EndDistance;
+
 	/** Scene elements past this distance will not have fog applied.  This is useful for excluding skyboxes which already have fog baked in. */
 	UPROPERTY(EditAnywhere, BlueprintReadOnly, Category=ExponentialHeightFogComponent, meta=(UIMin = "100000", UIMax = "20000000"))
 	float FogCutoffDistance;
@@ -257,6 +261,9 @@ public:
 	
 	UFUNCTION(BlueprintCallable, Category="Rendering|Components|ExponentialHeightFog")
 	ENGINE_API void SetStartDistance(float Value);
+	
+	UFUNCTION(BlueprintCallable, Category="Rendering|Components|ExponentialHeightFog")
+	ENGINE_API void SetEndDistance(float Value);
 
 	UFUNCTION(BlueprintCallable, Category="Rendering|Components|ExponentialHeightFog")
 	ENGINE_API void SetFogCutoffDistance(float Value);
diff --git a/Engine/Source/Runtime/Engine/Private/Components/HeightFogComponent.cpp b/Engine/Source/Runtime/Engine/Private/Components/HeightFogComponent.cpp
--- a/Engine/Source/Runtime/Engine/Private/Components/HeightFogComponent.cpp
+++ b/Engine/Source/Runtime/Engine/Private/Components/HeightFogComponent.cpp
@@ -39,6 +39,7 @@ UExponentialHeightFogComponent::UExponentialHeightFogComponent(const FObjectInit
 
 	FogMaxOpacity = 1.0f;
 	StartDistance = 0.0f;
+	EndDistance = 0.0f;
 
 	// disabled by default
 	FogCutoffDistance = 0;
@@ -125,6 +126,7 @@ void UExponentialHeightFogComponent::PostEditChangeProperty(FPropertyChangedEven
 	FogHeightFalloff = FMath::Clamp(FogHeightFalloff, 0.0f, 2.0f);
 	FogMaxOpacity = FMath::Clamp(FogMaxOpacity, 0.0f, 1.0f);
 	StartDistance = FMath::Clamp(StartDistance, 0.0f, (float)WORLD_MAX);
+	EndDistance = FMath::Clamp(EndDistance, 0.0f, (float)(10 * WORLD_MAX));
 	FogCutoffDistance = FMath::Clamp(FogCutoffDistance, 0.0f, (float)(10 * WORLD_MAX));
 	FullyDirectionalInscatteringColorDistance = FMath::Clamp(FullyDirectionalInscatteringColorDistance, 0.0f, (float)WORLD_MAX);
 	NonDirectionalInscatteringColorDistance = FMath::Clamp(NonDirectionalInscatteringColorDistance, 0.0f, FullyDirectionalInscatteringColorDistance);
@@ -278,6 +280,15 @@ void UExponentialHeightFogComponent::SetStartDistance(float Value)
 	}
 }
 
+void UExponentialHeightFogComponent::SetEndDistance(float Value)
+{
+	if(EndDistance != Value)
+	{
+		EndDistance = Value;
+		MarkRenderStateDirty();
+	}
+}
+
 void UExponentialHeightFogComponent::SetFogCutoffDistance(float Value)
 {
 	if(FogCutoffDistance != Value)
diff --git a/Engine/Source/Runtime/Renderer/Private/FogRendering.cpp b/Engine/Source/Runtime/Renderer/Private/FogRendering.cpp
--- a/Engine/Source/Runtime/Renderer/Private/FogRendering.cpp
+++ b/Engine/Source/Runtime/Renderer/Private/FogRendering.cpp
@@ -74,6 +74,7 @@ void SetupFogUniformParameters(FRDGBuilder& GraphBuilder, const FViewInfo& View,
 		OutParameters.DirectionalInscatteringColor = FVector4f(FVector3f(View.DirectionalInscatteringColor), FMath::Clamp(View.DirectionalInscatteringExponent, 0.000001f, 1000.0f));
 		OutParameters.FogInscatteringColorCubemap = Cubemap->TextureRHI;
 		OutParameters.FogInscatteringColorSampler = TStaticSamplerState<SF_Trilinear, AM_Clamp, AM_Clamp, AM_Clamp>::GetRHI();
+		OutParameters.EndDistance = View.FogEndDistance;
 	}
 
 	// Volumetric Fog
@@ -241,6 +242,7 @@ void FSceneRenderer::InitFogConstants()
 				const float InvRange = 1.0f / FMath::Max(FogInfo.FullyDirectionalInscatteringColorDistance - FogInfo.NonDirectionalInscatteringColorDistance, .00001f);
 				float NumMips = 1.0f;
 
+				View.FogEndDistance = FogInfo.EndDistance;
 				View.SkyAtmosphereAmbientContributionColorScale = FogInfo.SkyAtmosphereAmbientContributionColorScale;
 
 				if (FogInfo.InscatteringColorCubemap)
diff --git a/Engine/Source/Runtime/Renderer/Private/FogRendering.h b/Engine/Source/Runtime/Renderer/Private/FogRendering.h
--- a/Engine/Source/Runtime/Renderer/Private/FogRendering.h
+++ b/Engine/Source/Runtime/Renderer/Private/FogRendering.h
@@ -23,6 +23,7 @@ BEGIN_GLOBAL_SHADER_PARAMETER_STRUCT(FFogUniformParameters,)
 	SHADER_PARAMETER(FVector4f, DirectionalInscatteringColor)
 	SHADER_PARAMETER(FVector2f, SinCosInscatteringColorCubemapRotation)
 	SHADER_PARAMETER(FVector3f, FogInscatteringTextureParameters)
+	SHADER_PARAMETER(float, EndDistance)
 	SHADER_PARAMETER(float, ApplyVolumetricFog)
 	SHADER_PARAMETER(float, VolumetricFogStartDistance)
 	SHADER_PARAMETER(float, VolumetricFogNearFadeInDistanceInv)
diff --git a/Engine/Source/Runtime/Renderer/Private/SceneCore.cpp b/Engine/Source/Runtime/Renderer/Private/SceneCore.cpp
--- a/Engine/Source/Runtime/Renderer/Private/SceneCore.cpp
+++ b/Engine/Source/Runtime/Renderer/Private/SceneCore.cpp
@@ -389,6 +389,7 @@ FExponentialHeightFogSceneInfo::FExponentialHeightFogSceneInfo(const UExponentia
 	Component(InComponent),
 	FogMaxOpacity(InComponent->FogMaxOpacity),
 	StartDistance(InComponent->StartDistance),
+	EndDistance(InComponent->EndDistance),
 	FogCutoffDistance(InComponent->FogCutoffDistance),
 	DirectionalInscatteringExponent(InComponent->DirectionalInscatteringExponent),
 	DirectionalInscatteringStartDistance(InComponent->DirectionalInscatteringStartDistance),
diff --git a/Engine/Source/Runtime/Renderer/Private/SceneCore.h b/Engine/Source/Runtime/Renderer/Private/SceneCore.h
--- a/Engine/Source/Runtime/Renderer/Private/SceneCore.h
+++ b/Engine/Source/Runtime/Renderer/Private/SceneCore.h
@@ -161,7 +161,8 @@ public:
 	const UExponentialHeightFogComponent* Component;
 	FExponentialHeightFogSceneData FogData[NumFogs];
 	float FogMaxOpacity;
-	float StartDistance; 
+	float StartDistance;
+	float EndDistance;
 	float FogCutoffDistance;
 	FLinearColor FogColor;
 	float DirectionalInscatteringExponent; 
diff --git a/Engine/Source/Runtime/Renderer/Private/SceneRendering.cpp b/Engine/Source/Runtime/Renderer/Private/SceneRendering.cpp
--- a/Engine/Source/Runtime/Renderer/Private/SceneRendering.cpp
+++ b/Engine/Source/Runtime/Renderer/Private/SceneRendering.cpp
@@ -957,6 +957,7 @@ void FViewInfo::Init()
 	FogMaxOpacity = 1;
 	ExponentialFogParameters3 = FVector4f(0, 0, 0, 0);
 	SinCosInscatteringColorCubemapRotation = FVector2f::ZeroVector;
+	FogEndDistance = 0.0f;
 	FogInscatteringColorCubemap = nullptr;
 	FogInscatteringTextureParameters = FVector::ZeroVector;
 	VolumetricFogStartDistance = false;
diff --git a/Engine/Source/Runtime/Renderer/Private/SceneRendering.h b/Engine/Source/Runtime/Renderer/Private/SceneRendering.h
--- a/Engine/Source/Runtime/Renderer/Private/SceneRendering.h
+++ b/Engine/Source/Runtime/Renderer/Private/SceneRendering.h
@@ -1295,6 +1295,7 @@ public:
 	float FogMaxOpacity;
 	FVector4f ExponentialFogParameters3;
 	FVector4f SkyAtmosphereAmbientContributionColorScale;
+	float FogEndDistance;
 	bool bEnableVolumetricFog;
 	float VolumetricFogStartDistance;
 	float VolumetricFogNearFadeInDistanceInv;
