Path Tracer: Further refactor Sheen BSDF logic

Centralize the logic a bit more to avoid redundant lookups and to minimize the amount of temporary state. Also simplify the transform logic to avoid redundant transforms by always working in local space.

#rnx

diff --git a/Engine/Shaders/Private/PathTracing/Material/PathTracingSubstrate.ush b/Engine/Shaders/Private/PathTracing/Material/PathTracingSubstrate.ush
--- a/Engine/Shaders/Private/PathTracing/Material/PathTracingSubstrate.ush
+++ b/Engine/Shaders/Private/PathTracing/Material/PathTracingSubstrate.ush
@@ -22,12 +22,9 @@ struct FSubstrateSingleSlabData
 	float3 F90;
 	float2 Alpha0;
 	float2 Alpha1;
-	float FuzzRoughness;
-
-	float FuzzAttenuation;
 	float3 Spec0W;
 	float3 Spec1W;
-	float FuzzAmount;
+	FSubstrateSheenData Fuzz;
 
 	float3 LobeCdf;
 	float4 LobePdf;
@@ -43,27 +40,23 @@ struct FSubstrateSingleSlabData
 		F0 = Payload.SpecularColor;
 		F90 = Payload.SpecularEdgeColor * F0RGBToMicroOcclusion(F0);
 
-		FuzzRoughness = MakeRoughnessSafe(Payload.FuzzRoughness, SUBSTRATE_MIN_FUZZ_ROUGHNESS);
-
 		const float NoV = saturate(V.z);
 		// SUBSTRATE_TODO: Add support for F82 model?
 		FBxDFEnergyTermsRGB Spec0 = ComputeGGXSpecEnergyTermsRGB(Payload.RoughnessData.x, NoV, F0, F90);
 		FBxDFEnergyTermsRGB Spec1 = ComputeGGXSpecEnergyTermsRGB(Payload.RoughnessData.y, NoV, F0, F90);
-		FBxDFEnergyTermsA Fuzz = ComputeSheenEnergyTerms(NoV, FuzzRoughness);
+		Fuzz.Prepare(V, Payload.FuzzRoughness, Payload.FuzzAmount);
 
 		// save the parts needed during eval/sampling
 		Spec0W = Spec0.W;
 		Spec1W = Spec1.W;
-		FuzzAmount = Payload.FuzzAmount * Fuzz.W;
-		
+
 		const float3 SpecE = lerp(Spec0.E, Spec1.E, Payload.RoughnessData.z);
 
-		FuzzAttenuation = 1.0 - Payload.FuzzAmount * Fuzz.E;
-		DiffuseWeight = FuzzAttenuation * (1.0 - Luminance(SpecE));
+		DiffuseWeight = Fuzz.Attenuation * (1.0 - Luminance(SpecE));
 
-		const float3 Spec0Albedo = FuzzAttenuation * (1.0 - Payload.RoughnessData.z) * Spec0.E;
-		const float3 Spec1Albedo = FuzzAttenuation * (      Payload.RoughnessData.z) * Spec1.E;
-		const float3 FuzzAlbedo = Payload.FuzzAmount * Fuzz.E * Payload.FuzzColor;
+		const float3 Spec0Albedo = Fuzz.Attenuation * (1.0 - Payload.RoughnessData.z) * Spec0.E;
+		const float3 Spec1Albedo = Fuzz.Attenuation * (      Payload.RoughnessData.z) * Spec1.E;
+		const float3 FuzzAlbedo = Fuzz.Scale * Payload.FuzzColor;
 
 		// Now prepare a cdf/pdf for lobe selection
 		float3 MaxLobeWeight = Payload.GetMaxLobeWeight();
@@ -125,7 +118,7 @@ FMaterialSample Substrate_SampleMaterial(
 	{
 		// cloth lobe
 		RandSample.x = RescaleRandomNumber(RandSample.x, Data.LobeCdf.z, 1.0);
-		L = SampleSheenDirection(V, Data.FuzzRoughness, RandSample.xy);
+		L = Data.Fuzz.Sample(RandSample.xy);
 		H = normalize(L + V);
 	}
 
@@ -144,7 +137,8 @@ FMaterialSample Substrate_SampleMaterial(
 	const float2 GGXResult1 = GGXEvalReflection(L, V, H, Data.Alpha1);
 	const float Spec0Pdf = GGXResult0.y;
 	const float Spec1Pdf = GGXResult1.y;
-	const float ClothPdf = ComputeSheenPdf(V_World, L_World, N_World, NoV, Data.FuzzRoughness);
+	const float4 ClothResult = Data.Fuzz.Eval(L, V, H, Payload.FuzzColor);
+	const float ClothPdf = ClothResult.w;
 
 	const float ShadowTerminator = ShadowTerminatorTerm(L_World, N_World, Payload.WorldSmoothNormal);
 	if (bSampledDiffuse)
@@ -164,9 +158,9 @@ FMaterialSample Substrate_SampleMaterial(
 		const float3 FTint = SubstrateSpecularTint(Payload, NoV, NoL, VoH, NoH);
 #if SUBSTRATE_FRESNEL_F82
 		// SUBSTRATE_TODO: Need to tune the logic for specular micro-occlusion ...
-		const float3 F = Data.FuzzAttenuation * F_AdobeF82(Data.F0, Data.F90, VoH) * FTint;
+		const float3 F = Data.Fuzz.Attenuation * F_AdobeF82(Data.F0, Data.F90, VoH) * FTint;
 #else
-		const float3 F = Data.FuzzAttenuation * F_Schlick(Data.F0, Data.F90, VoH) * FTint;
+		const float3 F = Data.Fuzz.Attenuation * F_Schlick(Data.F0, Data.F90, VoH) * FTint;
 #endif
 
 		Result.AddLobeWithMIS((1.0 - Payload.RoughnessData.z) * F * GGXResult0.x * Data.Spec0W, Spec0Pdf, Data.LobePdf.y);
@@ -179,12 +173,7 @@ FMaterialSample Substrate_SampleMaterial(
 	else
 	{
 		// Cloth Lobe
-		float3 ClothWeight = 0;
-		if (NoV > 0 && NoL > 0)
-		{
-			ClothWeight = Data.FuzzAmount * ComputeSheenWeight(V_World, L_World, N_World, NoV, NoL, NoH, VoH, Data.FuzzRoughness, Payload.FuzzColor);
-		}
-		Result.AddLobeWithMIS(ClothWeight * ShadowTerminator, ClothPdf, Data.LobePdf.w);
+		Result.AddLobeWithMIS(ClothResult.xyz * ShadowTerminator, ClothPdf, Data.LobePdf.w);
 		Result.ScatterType = PATHTRACER_SCATTER_SPECULAR;
 
 		Result.Pdf += Data.LobePdf.x * DiffPdf;
@@ -230,9 +219,9 @@ FMaterialEval Substrate_EvalMaterial(
 		const float2 GGXResult1 = GGXEvalReflection(L, V, H, Data.Alpha1);
 		const float3 FTint = SubstrateSpecularTint(Payload, NoV, NoL, VoH, NoH);
 #if SUBSTRATE_FRESNEL_F82
-		const float3 F = Data.FuzzAttenuation * F_AdobeF82(Data.F0, Data.F90, VoH) * FTint;
+		const float3 F = Data.Fuzz.Attenuation * F_AdobeF82(Data.F0, Data.F90, VoH) * FTint;
 #else
-		const float3 F = Data.FuzzAttenuation * F_Schlick(Data.F0, Data.F90, VoH) * FTint;
+		const float3 F = Data.Fuzz.Attenuation * F_Schlick(Data.F0, Data.F90, VoH) * FTint;
 #endif
 		SubstrateGlint(Payload, V, L, Data.Alpha0, GGXResult0.y /*Pdf*/, GGXResult0.x /*Weight*/);
 
@@ -241,14 +230,9 @@ FMaterialEval Substrate_EvalMaterial(
 	}
 	{
 		// Cloth Lobe
-		float3 ClothWeight = 0;
-		float ClothPdf = 0;
-		if (NoV > 0 && NoL > 0)
-		{
-			ClothWeight = Data.FuzzAmount * ComputeSheenWeight(V_World, L_World, N_World, NoV, NoL, NoH, VoH, Data.FuzzRoughness, Payload.FuzzColor);
-			ClothPdf = ComputeSheenPdf(V_World, L_World, N_World, NoV, Data.FuzzRoughness);
-		}
-		Result.AddLobeWithMIS(DiffuseSpecularScale.y * ClothWeight * ShadowTerminator, ClothPdf, Data.LobePdf.w);
+		const float4 ClothResult = Data.Fuzz.Eval(L, V, H, Payload.FuzzColor);
+		const float ClothPdf = ClothResult.w;
+		Result.AddLobeWithMIS(DiffuseSpecularScale.y * ClothResult.xyz * ShadowTerminator, ClothPdf, Data.LobePdf.w);
 	}
 
 	Result.Weight *= SubstrateLobeWeight(Payload, NoL);
@@ -266,13 +250,11 @@ struct FSubstrateClearCoatData
 	float3 F90;
 	float2 BaseAlpha;
 	float2 CoatAlpha;
-	float FuzzRoughness;
 
 	float CoatAttenuation;
 	FBxDFEnergyTermsRGB BaseSpec;
 	FBxDFEnergyTermsA   CoatSpec;
-	float FuzzAttenuation;
-	float FuzzAmount;
+	FSubstrateSheenData Fuzz;
 
 	float3 LobeCdf;
 	float4 LobePdf;
@@ -288,22 +270,18 @@ struct FSubstrateClearCoatData
 		F0 = Payload.SpecularColor;
 		F90 = Payload.SpecularEdgeColor * F0RGBToMicroOcclusion(F0);
 
-		FuzzRoughness = MakeRoughnessSafe(Payload.FuzzRoughness, SUBSTRATE_MIN_FUZZ_ROUGHNESS);
-
 		const float NoV = saturate(V.z);
 		// SUBSTRATE_TODO: Add support for F82 model?
 		BaseSpec = ComputeGGXSpecEnergyTermsRGB(Payload.RoughnessData.x, NoV, F0, F90);
 		CoatSpec = ComputeGGXSpecEnergyTermsA(Payload.RoughnessData.y, NoV, CLEAR_COAT_F0);
-		FBxDFEnergyTermsA Fuzz = ComputeClothEnergyTermsA(FuzzRoughness, NoV);
+		Fuzz.Prepare(V, Payload.FuzzRoughness, Payload.FuzzAmount);
 
-		FuzzAttenuation = 1.0 - Payload.FuzzAmount * Fuzz.E;
-		FuzzAmount      = Payload.FuzzAmount * Fuzz.W;
 		CoatAttenuation = 1.0 - Payload.RoughnessData.z * CoatSpec.E;
-		DiffuseWeight = FuzzAttenuation * CoatAttenuation * (1.0 - Luminance(BaseSpec.E));
+		DiffuseWeight = Fuzz.Attenuation * CoatAttenuation * (1.0 - Luminance(BaseSpec.E));
 
-		const float3 BaseSpecAlbedo = FuzzAttenuation * CoatAttenuation * BaseSpec.E;
-		const float  CoatSpecAlbedo = FuzzAttenuation * Payload.RoughnessData.z * CoatSpec.E;
-		const float3 FuzzAlbedo = Payload.FuzzAmount * Fuzz.E * Payload.FuzzColor;
+		const float3 BaseSpecAlbedo = Fuzz.Attenuation * CoatAttenuation * BaseSpec.E;
+		const float  CoatSpecAlbedo = Fuzz.Attenuation * Payload.RoughnessData.z * CoatSpec.E;
+		const float3 FuzzAlbedo = Fuzz.Scale * Payload.FuzzColor;
 
 		// Now prepare a cdf/pdf for lobe selection
 		float3 MaxLobeWeight = Payload.GetMaxLobeWeight();
@@ -366,7 +344,7 @@ FMaterialSample SubstrateClearCoat_SampleMaterial(
 	{
 		// cloth lobe
 		RandSample.x = RescaleRandomNumber(RandSample.x, Data.LobeCdf.z, 1.0);
-		L = SampleSheenDirection(V, Data.FuzzRoughness, RandSample.xy);
+		L = Data.Fuzz.Sample(RandSample.xy);
 		H = normalize(L + V);
 	}
 
@@ -389,7 +367,8 @@ FMaterialSample SubstrateClearCoat_SampleMaterial(
 	const float2 CoatGGXResult = GGXEvalReflection(L, V, H, Data.CoatAlpha);
 	const float BaseSpecPdf = BaseGGXResult.y;
 	const float CoatSpecPdf = CoatGGXResult.y;
-	const float ClothPdf = ComputeSheenPdf(V_World, L_World, N_World, NoV, Data.FuzzRoughness);
+	const float4 ClothResult = Data.Fuzz.Eval(L, V, H, Payload.FuzzColor);
+	const float ClothPdf = ClothResult.w;
 
 	const float ShadowTerminator = ShadowTerminatorTerm(L_World, N_World, Payload.WorldSmoothNormal);
 	if (bSampledDiffuse)
@@ -407,14 +386,14 @@ FMaterialSample SubstrateClearCoat_SampleMaterial(
 		const float3 FTint = SubstrateSpecularTint(Payload, NoV, NoL, VoH, NoH);
 #if SUBSTRATE_FRESNEL_F82
 		// SUBSTRATE_TODO: Need to tune the logic for specular micro-occlusion ...
-		const float3 F = Data.FuzzAttenuation * F_AdobeF82(Data.F0, Data.F90, VoH) * FTint;
+		const float3 FBase = Data.Fuzz.Attenuation * F_AdobeF82(Data.F0, Data.F90, VoH) * FTint;
 #else
-		const float3 FBase = Data.FuzzAttenuation * F_Schlick(Data.F0, Data.F90, VoH) * FTint;
+		const float3 FBase = Data.Fuzz.Attenuation * F_Schlick(Data.F0, Data.F90, VoH) * FTint;
 #endif
 		const float3 FCoat = Payload.RoughnessData.z * ClearCoatFresnel(VoH) * FTint;
 
 		Result.AddLobeWithMIS(Data.CoatAttenuation * FBase * BaseGGXResult.x * Data.BaseSpec.W * Transmission, BaseSpecPdf, Data.LobePdf.y);
-		Result.AddLobeWithMIS(Data.FuzzAttenuation * FCoat * CoatGGXResult.x * Data.CoatSpec.W               , CoatSpecPdf, Data.LobePdf.z);
+		Result.AddLobeWithMIS(Data.Fuzz.Attenuation * FCoat * CoatGGXResult.x * Data.CoatSpec.W               , CoatSpecPdf, Data.LobePdf.z);
 		Result.Pdf += Data.LobePdf.x * DiffPdf;
 		Result.Pdf += Data.LobePdf.w * ClothPdf;
 
@@ -423,12 +402,7 @@ FMaterialSample SubstrateClearCoat_SampleMaterial(
 	else
 	{
 		// Cloth Lobe
-		float3 ClothWeight = 0;
-		if (NoL > 0 && NoV > 0)
-		{
-			ClothWeight = Data.FuzzAmount * ComputeSheenWeight(V_World, L_World, N_World, NoV, NoL, NoH, VoH, Data.FuzzRoughness, Payload.FuzzColor);
-		}
-		Result.AddLobeWithMIS(ClothWeight * ShadowTerminator, ClothPdf, Data.LobePdf.w);
+		Result.AddLobeWithMIS(ClothResult.xyz * ShadowTerminator, ClothPdf, Data.LobePdf.w);
 		Result.ScatterType = PATHTRACER_SCATTER_SPECULAR;
 
 		Result.Pdf += Data.LobePdf.x * DiffPdf;
@@ -483,25 +457,20 @@ FMaterialEval SubstrateClearCoat_EvalMaterial(
 		const float3 FTint = SubstrateSpecularTint(Payload, NoV, NoL, VoH, NoH);
 #if SUBSTRATE_FRESNEL_F82
 		// SUBSTRATE_TODO: Need to tune the logic for specular micro-occlusion ...
-		const float3 F = Data.FuzzAttenuation * F_AdobeF82(Data.F0, Data.F90, VoH) * FTint;
+		const float3 FBase = Data.Fuzz.Attenuation * F_AdobeF82(Data.F0, Data.F90, VoH) * FTint;
 #else
-		const float3 FBase = Data.FuzzAttenuation * F_Schlick(Data.F0, Data.F90, VoH) * FTint;
+		const float3 FBase = Data.Fuzz.Attenuation * F_Schlick(Data.F0, Data.F90, VoH) * FTint;
 #endif
 		const float3 FCoat = Payload.RoughnessData.z * ClearCoatFresnel(VoH) * FTint;
 
 		Result.AddLobeWithMIS(Data.CoatAttenuation * FBase * BaseGGXResult.x * Data.BaseSpec.W * Transmission, BaseSpecPdf, Data.LobePdf.y);
-		Result.AddLobeWithMIS(Data.FuzzAttenuation * FCoat * CoatGGXResult.x * Data.CoatSpec.W               , CoatSpecPdf, Data.LobePdf.z);
+		Result.AddLobeWithMIS(Data.Fuzz.Attenuation * FCoat * CoatGGXResult.x * Data.CoatSpec.W               , CoatSpecPdf, Data.LobePdf.z);
 	}
 	{
 		// Cloth Lobe
-		float3 ClothWeight = 0;
-		float ClothPdf = 0;
-		if (NoV > 0 && NoL > 0)
-		{
-			ClothWeight = Data.FuzzAmount * ComputeSheenWeight(V_World, L_World, N_World, NoV, NoL, NoH, VoH, Data.FuzzRoughness, Payload.FuzzColor);
-			ClothPdf = ComputeSheenPdf(V_World, L_World, N_World, NoV, Data.FuzzRoughness);
-		}
-		Result.AddLobeWithMIS(DiffuseSpecularScale.y * ClothWeight * ShadowTerminator, ClothPdf, Data.LobePdf.w);
+		const float4 ClothResult = Data.Fuzz.Eval(L, V, H, Payload.FuzzColor);
+		const float ClothPdf = ClothResult.w;
+		Result.AddLobeWithMIS(DiffuseSpecularScale.y * ClothResult.xyz * ShadowTerminator, ClothPdf, Data.LobePdf.w);
 	}
 
 	Result.Weight *= SubstrateLobeWeight(Payload, NoL);
diff --git a/Engine/Shaders/Private/PathTracing/Material/PathTracingSubstrateCommon.ush b/Engine/Shaders/Private/PathTracing/Material/PathTracingSubstrateCommon.ush
--- a/Engine/Shaders/Private/PathTracing/Material/PathTracingSubstrateCommon.ush
+++ b/Engine/Shaders/Private/PathTracing/Material/PathTracingSubstrateCommon.ush
@@ -55,64 +55,119 @@ void SubstrateGlint(FPathTracingPayload Payload, float3 InTangentV, float3 InTan
 #endif
 }
 
-FBxDFEnergyTermsA ComputeSheenEnergyTerms(float NoV, float SheenRoughness)
+#define USE_SHEEN_LTC_SAMPLER 1 // 0: fallback to uniform sampling (for reference)
+								// 1: use the LTC coefficients for importance sampling
+
+struct FSubstrateSheenData
 {
 #if SUBSTRATE_SHEEN_QUALITY == 1
-	FBxDFEnergyTermsA Result;
-	Result.E = SheenLTC_DirectionalAlbedo(NoV, SheenRoughness, View.SheenLTCTexture, View.SheenLTCSampler);
-	Result.W = Result.E; // overall weight is the directional albedo
-	return Result;
+	float2 Tangent; 	// Captures rotation within the local tangent frame
+	float aInv; 		// LTC Coefficients
+	float bInv;
 #else
-	return ComputeClothEnergyTermsA(SheenRoughness, NoV);
+	float Roughness;
 #endif
-}
-
-#define USE_SHEEN_LTC_SAMPLER 1 // 0: fallback to uniform sampling (for reference)
-								// 1: use the LTC coefficients for importance sampling
+	float Scale;      	// overall scale factor for the sheen lobe
+	float Attenuation; 	// amount by which other lobes have to be scaled
 
-float ComputeSheenPdf(float3 V, float3 L, float3 N, float NoV, float SheenRoughness)
-{
-#if USE_SHEEN_LTC_SAMPLER != 0 && SUBSTRATE_SHEEN_QUALITY == 1
-	// Disney's Cloth Lobe
-	float FuzzDirectionalAlbedo = 1;
-	float FuzzSpecularPathValueNoL = SheenLTC_Eval(V, L, N, NoV, SheenRoughness, View.SheenLTCTexture, View.SheenLTCSampler, FuzzDirectionalAlbedo);
-	return FuzzSpecularPathValueNoL;
+	void Prepare(float3 V, float FuzzRoughness, float FuzzAmount)
+	{
+#if SUBSTRATE_SHEEN_QUALITY == 1
+		Tangent = V.xy;
+		const float Len2 = length2(Tangent);
+		if (Len2 > 0)
+		{
+			Tangent *= rsqrt(Len2);
+		}
+		else
+		{
+			Tangent = float2(1, 0);
+		}
+		float4 LTC = SheenLTC_Cofficients(V.z, FuzzRoughness, View.SheenLTCTexture, View.SheenLTCSampler);
+		aInv = LTC.x;
+		bInv = LTC.y;
+		Scale = FuzzAmount * LTC.z;
+		Attenuation = 1 - FuzzAmount;
 #else
-	// uniform sampling (for legacy cloth lobe, or naive sampling of the LTC lobe)
-	const float Pdf = 1 / (2 * PI);
-	return Pdf;
+		Roughness = MakeRoughnessSafe(FuzzRoughness, SUBSTRATE_MIN_FUZZ_ROUGHNESS);
+		FBxDFEnergyTermsA Fuzz = ComputeClothEnergyTermsA(Roughness, V.z);
+		Scale = FuzzAmount * Fuzz.W;
+		Attenuation = 1 - FuzzAmount * Fuzz.E;
 #endif
+	}
 
-}
 
-float3 ComputeSheenWeight(float3 V, float3 L, float3 N, float NoV, float NoL, float NoH, float VoH, float SheenRoughness, float3 SheenColor)
-{
+
+	float4 Eval(float3 L, float3 V, float3 H, float3 FuzzColor)
+	{
 #if SUBSTRATE_SHEEN_QUALITY == 1
-	// Disney's Cloth Lobe
-	float FuzzDirectionalAlbedo = 1; // NOTE: Path Tracer includes this term in the weight applied outside
-	float FuzzSpecularPathValueNoL = SheenLTC_Eval(V, L, N, NoV, SheenRoughness, View.SheenLTCTexture, View.SheenLTCSampler, FuzzDirectionalAlbedo);
-#if USE_SHEEN_LTC_SAMPLER == 0
-	// using naive sampling, weight will include Brdf evaluation
-	// uniform sampling, pdf is just 1/2pi
-	return SheenColor * (FuzzSpecularPathValueNoL * 2 * PI);
+		if (L.z > 0)
+		{
+			const float3 LocalL = float3(
+				L.x * Tangent.x + L.y * Tangent.y,
+				L.y * Tangent.x - L.x * Tangent.y,
+				L.z);
+
+			const float3 WrappedLocalL = float3(
+				aInv * LocalL.x + bInv * LocalL.z,
+				aInv * LocalL.y,
+				LocalL.z);
+
+			const float InvLenWrappedLocalL2 = rcp(length2(WrappedLocalL));
+			// Jabobian is defined as: det(InvLTC) / ||InvLTC . L||
+			// Jacobian should be aInv^2/Length^3 but since we skip the normalization
+			// we end up with aInv^2/Length^4 which is cheaper to compute
+			const float Jacobian = Pow2(aInv * InvLenWrappedLocalL2);
+
+
+			// Evaluation a normalized cos distribution (note: WrappedLocalL was not normalized, we folded the extra length division into the Jacobian)
+#if USE_SHEEN_LTC_SAMPLER == 1
+			const float Weight = Scale;
+			const float Pdf = WrappedLocalL.z * Jacobian / PI;
 #else
-	// LTC sampling - pdf is equal to the brdf, so sample weight is just the color
-	return SheenColor;
+			const float Weight = Scale * WrappedLocalL.z * Jacobian * 2; // multiply by 2PI
+			const float Pdf = 1.0 / (2 * PI);
 #endif
+			return float4(Weight * FuzzColor, Pdf);
+		}
+#else
+		// Charlie's Cloth Lobe
+		if (V.z > 0 && L.z > 0)
+		{
+			const float DCloth = D_Charlie(Roughness, H.z);
+			const float3 FCloth = F_Schlick(FuzzColor, dot(V, H));
+			const float3 Weight = (2 * PI * L.z * DCloth * Vis_Ashikhmin(V.z, L.z)) * FCloth;
+			const float Pdf = 1.0 / (2 * PI);
+			return float4(Weight, Pdf);
+		}
+#endif
+		return 0;
+	}
 
+	float3 Sample(float2 Rand)
+	{
+#if SUBSTRATE_SHEEN_QUALITY == 1 && USE_SHEEN_LTC_SAMPLER == 1
+		const float3 CosSample = CosineSampleHemisphere(Rand).xyz;
+		// NOTE: The vector is scaled by aInv compared to the reference implementation to avoid unecessary divides
+		const float3 LocalL = normalize(float3(CosSample.x - CosSample.z * bInv, CosSample.y, aInv * CosSample.z));
+		return float3(
+			LocalL.x * Tangent.x - LocalL.y * Tangent.y,
+			LocalL.x * Tangent.y + LocalL.y * Tangent.x,
+			LocalL.z);
 #else
-	// Charlie's Cloth Lobe
-	const float DCloth = D_Charlie(SheenRoughness, NoH);
-	const float3 FCloth = F_Schlick(SheenColor, VoH);
-	return (2 * PI * NoL * DCloth * Vis_Ashikhmin(NoV, NoL)) * FCloth;
+		return UniformSampleHemisphere(Rand).xyz;
 #endif
-}
+	}
+};
 
-float3 SampleSheenDirection(float3 V, float SheenRoughness, float2 Rand)
+FBxDFEnergyTermsA ComputeSheenEnergyTerms(float NoV, float SheenRoughness)
 {
-#if USE_SHEEN_LTC_SAMPLER != 0 && SUBSTRATE_SHEEN_QUALITY == 1
-	return SheenLTC_Sample(CosineSampleHemisphere(Rand).xyz, V, float3(0, 0, 1), V.z, SheenRoughness, View.SheenLTCTexture, View.SheenLTCSampler);
+#if SUBSTRATE_SHEEN_QUALITY == 1
+	FBxDFEnergyTermsA Result;
+	Result.E = SheenLTC_DirectionalAlbedo(NoV, SheenRoughness, View.SheenLTCTexture, View.SheenLTCSampler);
+	Result.W = Result.E; // overall weight is the directional albedo
+	return Result;
 #else
-	return UniformSampleHemisphere(Rand).xyz;
+	return ComputeClothEnergyTermsA(SheenRoughness, NoV);
 #endif
 }
diff --git a/Engine/Shaders/Private/PathTracing/Material/PathTracingSubstrateFoliage.ush b/Engine/Shaders/Private/PathTracing/Material/PathTracingSubstrateFoliage.ush
--- a/Engine/Shaders/Private/PathTracing/Material/PathTracingSubstrateFoliage.ush
+++ b/Engine/Shaders/Private/PathTracing/Material/PathTracingSubstrateFoliage.ush
@@ -23,12 +23,9 @@ struct FSubstrateFoliageSlabData
 	float3 F90;
 	float2 Alpha0;
 	float2 Alpha1;
-	float FuzzRoughness;
-	
-	float FuzzAttenuation;
 	float3 Spec0W;
 	float3 Spec1W;
-	float FuzzAmount;
+	FSubstrateSheenData Fuzz;
 
 	float3 LobeCdf;
 	float4 LobePdf;
@@ -44,28 +41,24 @@ struct FSubstrateFoliageSlabData
 		F0 = Payload.SpecularColor;
 		F90 = Payload.SpecularEdgeColor * F0RGBToMicroOcclusion(F0);
 
-		FuzzRoughness = MakeRoughnessSafe(Payload.FuzzRoughness, SUBSTRATE_MIN_FUZZ_ROUGHNESS);
-
 		const float NoV = saturate(V.z);
 		// SUBSTRATE_TODO: Add support for F82 model?
 		FBxDFEnergyTermsRGB Spec0 = ComputeGGXSpecEnergyTermsRGB(Payload.RoughnessData.x, NoV, F0, F90);
 		FBxDFEnergyTermsRGB Spec1 = ComputeGGXSpecEnergyTermsRGB(Payload.RoughnessData.y, NoV, F0, F90);
-		FBxDFEnergyTermsA Fuzz = ComputeSheenEnergyTerms(NoV, FuzzRoughness);
+		Fuzz.Prepare(V, Payload.FuzzRoughness, Payload.FuzzAmount);
 
 		// save the parts needed during eval/sampling
 		Spec0W = Spec0.W;
 		Spec1W = Spec1.W;
-		FuzzAmount = Payload.FuzzAmount * Fuzz.W;
 		
 
 		const float3 SpecE = lerp(Spec0.E, Spec1.E, Payload.RoughnessData.z);
 
-		FuzzAttenuation = 1.0 - Payload.FuzzAmount * Fuzz.E;
-		DiffuseWeight = FuzzAttenuation * (1.0 - Luminance(SpecE));
+		DiffuseWeight = Fuzz.Attenuation * (1.0 - Luminance(SpecE));
 
-		const float3 Spec0Albedo = FuzzAttenuation * (1.0 - Payload.RoughnessData.z) * Spec0.E;
-		const float3 Spec1Albedo = FuzzAttenuation * (      Payload.RoughnessData.z) * Spec1.E;
-		const float3 FuzzAlbedo = Payload.FuzzAmount * Fuzz.E * Payload.FuzzColor;
+		const float3 Spec0Albedo = Fuzz.Attenuation * (1.0 - Payload.RoughnessData.z) * Spec0.E;
+		const float3 Spec1Albedo = Fuzz.Attenuation * (      Payload.RoughnessData.z) * Spec1.E;
+		const float3 FuzzAlbedo = Fuzz.Scale * Payload.FuzzColor;
 
 		// Now prepare a cdf/pdf for lobe selection
 		float3 MaxLobeWeight = Payload.GetMaxLobeWeight();
@@ -89,7 +82,6 @@ FMaterialSample SubstrateFoliage_SampleMaterial(
 	FSubstrateFoliageSlabData Data = (FSubstrateFoliageSlabData)0;
 	Data.PrepareBSDF(V_World, Payload);
 
-
 	float3 L = 0, H = 0, V = Data.V;
 	float OutRoughness = 1;
 
@@ -137,7 +129,7 @@ FMaterialSample SubstrateFoliage_SampleMaterial(
 	{
 		// cloth lobe
 		RandSample.x = RescaleRandomNumber(RandSample.x, Data.LobeCdf.z, 1.0);
-		L = SampleSheenDirection(V, Data.FuzzRoughness, RandSample.xy);
+		L = Data.Fuzz.Sample(RandSample.xy);
 		H = normalize(L + V);
 	}
 
@@ -156,7 +148,8 @@ FMaterialSample SubstrateFoliage_SampleMaterial(
 	const float2 GGXResult1 = GGXEvalReflection(L, V, H, Data.Alpha1);
 	const float Spec0Pdf = GGXResult0.y;
 	const float Spec1Pdf = GGXResult1.y;
-	const float ClothPdf = ComputeSheenPdf(V_World, L_World, N_World, NoV, Data.FuzzRoughness);
+	const float4 ClothResult = Data.Fuzz.Eval(L, V, H, Payload.FuzzColor);
+	const float ClothPdf = ClothResult.w;
 
 	const float ShadowTerminator = ShadowTerminatorTerm(L_World, N_World, Payload.WorldSmoothNormal);
 	if (bSampledTransmit)
@@ -178,9 +171,9 @@ FMaterialSample SubstrateFoliage_SampleMaterial(
 		const float3 FTint = SubstrateSpecularTint(Payload, NoV, NoL, VoH, NoH);
 #if SUBSTRATE_FRESNEL_F82
 		// SUBSTRATE_TODO: Need to tune the logic for specular micro-occlusion ...
-		const float3 F = Data.FuzzAttenuation * F_AdobeF82(Data.F0, Data.F90, VoH) * FTint;
+		const float3 F = Data.Fuzz.Attenuation * F_AdobeF82(Data.F0, Data.F90, VoH) * FTint;
 #else
-		const float3 F = Data.FuzzAttenuation * F_Schlick(Data.F0, Data.F90, VoH) * FTint;
+		const float3 F = Data.Fuzz.Attenuation * F_Schlick(Data.F0, Data.F90, VoH) * FTint;
 #endif
 
 		Result.AddLobeWithMIS((1.0 - Payload.RoughnessData.z) * F * GGXResult0.x * Data.Spec0W, Spec0Pdf, Data.LobePdf.y);
@@ -193,12 +186,7 @@ FMaterialSample SubstrateFoliage_SampleMaterial(
 	else
 	{
 		// Cloth Lobe
-		float3 ClothWeight = 0;
-		if (NoV > 0 && NoL > 0)
-		{
-			ClothWeight = Data.FuzzAmount * ComputeSheenWeight(V_World, L_World, N_World, NoV, NoL, NoH, VoH, Data.FuzzRoughness, Payload.FuzzColor);
-		}
-		Result.AddLobeWithMIS(ClothWeight * ShadowTerminator, ClothPdf, Data.LobePdf.w);
+		Result.AddLobeWithMIS(ClothResult.xyz * ShadowTerminator, ClothPdf, Data.LobePdf.w);
 		Result.ScatterType = PATHTRACER_SCATTER_SPECULAR;
 
 		Result.Pdf += Data.LobePdf.x * DiffPdf;
@@ -252,9 +240,9 @@ FMaterialEval SubstrateFoliage_EvalMaterial(
 			const float2 GGXResult1 = GGXEvalReflection(L, V, H, Data.Alpha1);
 			const float3 FTint = SubstrateSpecularTint(Payload, NoV, NoL, VoH, NoH);
 			#if SUBSTRATE_FRESNEL_F82
-			const float3 F = Data.FuzzAttenuation * F_AdobeF82(Data.F0, Data.F90, VoH) * FTint;
+			const float3 F = Data.Fuzz.Attenuation * F_AdobeF82(Data.F0, Data.F90, VoH) * FTint;
 			#else
-			const float3 F = Data.FuzzAttenuation * F_Schlick(Data.F0, Data.F90, VoH) * FTint;
+			const float3 F = Data.Fuzz.Attenuation * F_Schlick(Data.F0, Data.F90, VoH) * FTint;
 			#endif
 
 			Result.AddLobeWithMIS(DiffuseSpecularScale.y * (1.0 - Payload.RoughnessData.z) * F * GGXResult0.x * Data.Spec0W, GGXResult0.y, Data.LobePdf.y);
@@ -262,9 +250,9 @@ FMaterialEval SubstrateFoliage_EvalMaterial(
 		}
 		{
 			// Cloth Lobe
-			float3 ClothWeight = Data.FuzzAmount * ComputeSheenWeight(V_World, L_World, N_World, NoV, NoL, NoH, VoH, Data.FuzzRoughness, Payload.FuzzColor);
-			float ClothPdf = ComputeSheenPdf(V_World, L_World, N_World, NoV, Data.FuzzRoughness);
-			Result.AddLobeWithMIS(DiffuseSpecularScale.y * ClothWeight * ShadowTerminator, ClothPdf, Data.LobePdf.w);
+			const float4 ClothResult = Data.Fuzz.Eval(L, V, H, Payload.FuzzColor);
+			const float ClothPdf = ClothResult.w;
+			Result.AddLobeWithMIS(DiffuseSpecularScale.y * ClothResult.xyz * ShadowTerminator, ClothPdf, Data.LobePdf.w);
 		}
 	}
 
diff --git a/Engine/Shaders/Private/PathTracing/Material/PathTracingSubstrateSolidGlass.ush b/Engine/Shaders/Private/PathTracing/Material/PathTracingSubstrateSolidGlass.ush
--- a/Engine/Shaders/Private/PathTracing/Material/PathTracingSubstrateSolidGlass.ush
+++ b/Engine/Shaders/Private/PathTracing/Material/PathTracingSubstrateSolidGlass.ush
@@ -22,12 +22,11 @@ struct FSubstrateSolidGlassSlabData
 	float  Eta;
 	float2 Alpha0;
 	float2 Alpha1;
-	float  FuzzRoughness;
 
-	float  FuzzAttenuation;
 	float  EGlass0;
 	float  EGlass1;
-	float FuzzAmount;
+
+	FSubstrateSheenData Fuzz;
 
 	float2 LobeCdf;
 	float3 LobePdf;
@@ -44,21 +43,16 @@ struct FSubstrateSolidGlassSlabData
 		// NOTE: IsFrontFace() determines polygon orientation, because the normal is always flipped towards in the incoming ray
 		Eta = Payload.IsFrontFace() ? Payload.Ior : rcp(Payload.Ior);
 		
-
-		FuzzRoughness = MakeRoughnessSafe(Payload.FuzzRoughness, SUBSTRATE_MIN_FUZZ_ROUGHNESS);
-
 		const float NoV = saturate(V.z);
 		// correct for energy loss by scaling the whole BSDF
 		EGlass0 = GGXEnergyLookup(Payload.RoughnessData.x, NoV, Eta);
 		EGlass1 = GGXEnergyLookup(Payload.RoughnessData.y, NoV, Eta);
-		FBxDFEnergyTermsA Fuzz = ComputeClothEnergyTermsA(FuzzRoughness, NoV);
 
-		FuzzAttenuation = 1.0 - Payload.FuzzAmount * Fuzz.E;
-		FuzzAmount = Payload.FuzzAmount * Fuzz.W;
+		Fuzz.Prepare(V, Payload.FuzzRoughness, Payload.FuzzAmount);
 
-		const float3 Glass0Albedo = FuzzAttenuation * (1.0 - Payload.RoughnessData.z);
-		const float3 Glass1Albedo = FuzzAttenuation * (      Payload.RoughnessData.z);
-		const float3 FuzzAlbedo = Payload.FuzzAmount * Fuzz.E * Payload.FuzzColor;
+		const float3 Glass0Albedo = Fuzz.Attenuation * (1.0 - Payload.RoughnessData.z);
+		const float3 Glass1Albedo = Fuzz.Attenuation * (      Payload.RoughnessData.z);
+		const float3 FuzzAlbedo = Fuzz.Scale * Payload.FuzzColor;
 
 		// Now prepare a cdf/pdf for lobe selection
 		float3 MaxLobeWeight = Payload.GetMaxLobeWeight();
@@ -120,7 +114,7 @@ FMaterialSample SubstrateSolidGlass_SampleMaterial(
 		if (bRefract)
 		{
 			// only need to account for overall BSDF scaling
-			Result.Weight *= Data.FuzzAttenuation * Payload.BSDFOpacity * Payload.WeightV;
+			Result.Weight *= Data.Fuzz.Attenuation * Payload.BSDFOpacity * Payload.WeightV;
 		}
 		else
 		{
@@ -133,9 +127,9 @@ FMaterialSample SubstrateSolidGlass_SampleMaterial(
 
 			// reflection side, apply full lobe weight and attenuation along L
 			// and account for probability of picking the fuzz layer
-			const float ClothPdf = ComputeSheenPdf(V, L, float3(0, 0, 1), NoV, Data.FuzzRoughness);
+			const float ClothPdf = Data.Fuzz.Eval(L, V, H, Payload.FuzzColor).w;
 			Result.Pdf += Data.LobePdf.z * ClothPdf;
-			Result.Weight *= Data.FuzzAttenuation * SubstrateLobeWeight(Payload, saturate(L.z)) * FTint;
+			Result.Weight *= Data.Fuzz.Attenuation * SubstrateLobeWeight(Payload, saturate(L.z)) * FTint;
 		}
 
 		return Result;
@@ -144,7 +138,7 @@ FMaterialSample SubstrateSolidGlass_SampleMaterial(
 	{
 		// cloth lobe
 		RandSample.x = RescaleRandomNumber(RandSample.x, Data.LobeCdf.y, 1.0);
-		L = SampleSheenDirection(V, Data.FuzzRoughness, RandSample.xy);
+		L = Data.Fuzz.Sample(RandSample.xy);
 		H = normalize(L + V);
 
 		// transform to world space
@@ -166,14 +160,9 @@ FMaterialSample SubstrateSolidGlass_SampleMaterial(
 
 		// Cloth Lobe
 		const float ShadowTerminator = ShadowTerminatorTerm(L_World, N_World, Payload.WorldSmoothNormal);
-		float3 ClothWeight = 0;
-		float ClothPdf = 0;
-		if (NoV > 0 && NoL > 0)
-		{
-			ClothWeight = Data.FuzzAmount * ComputeSheenWeight(V_World, L_World, N_World, NoV, NoL, NoH, VoH, Data.FuzzRoughness, Payload.FuzzColor);
-			ClothPdf = ComputeSheenPdf(V_World, L_World, N_World, NoV, Data.FuzzRoughness);
-		}
-		Result.AddLobeWithMIS(ClothWeight * ShadowTerminator, ClothPdf, Data.LobePdf.z);
+		const float4 ClothResult = Data.Fuzz.Eval(L, V, H, Payload.FuzzColor);
+		const float ClothPdf = ClothResult.w;
+		Result.AddLobeWithMIS(ClothResult.xyz * ShadowTerminator, ClothPdf, Data.LobePdf.z);
 		Result.ScatterType = PATHTRACER_SCATTER_SPECULAR;
 
 		Result.Pdf += Data.LobePdf.y * Spec0Pdf;
@@ -220,20 +209,15 @@ FMaterialEval SubstrateSolidGlass_EvalMaterial(
 			const float Fg = FresnelReflectance(VoH, Data.Eta, Data.F0);
 			const float3 FTint = SubstrateSpecularTint(Payload, NoV, NoL, VoH, NoH);
 
-			Result.AddLobeWithMIS(FTint * DiffuseSpecularScale.y * Data.FuzzAttenuation * (1.0 - Payload.RoughnessData.z) * GGXResult0.x / Data.EGlass0, Fg * GGXResult0.y, Data.LobePdf.x);
-			Result.AddLobeWithMIS(FTint * DiffuseSpecularScale.y * Data.FuzzAttenuation * (      Payload.RoughnessData.z) * GGXResult1.x / Data.EGlass1, Fg * GGXResult1.y, Data.LobePdf.y);
+			Result.AddLobeWithMIS(FTint * DiffuseSpecularScale.y * Data.Fuzz.Attenuation * (1.0 - Payload.RoughnessData.z) * GGXResult0.x / Data.EGlass0, Fg * GGXResult0.y, Data.LobePdf.x);
+			Result.AddLobeWithMIS(FTint * DiffuseSpecularScale.y * Data.Fuzz.Attenuation * (      Payload.RoughnessData.z) * GGXResult1.x / Data.EGlass1, Fg * GGXResult1.y, Data.LobePdf.y);
 		}
 		{
 			// Cloth Lobe
 			const float ShadowTerminator = ShadowTerminatorTerm(L_World, N_World, Payload.WorldSmoothNormal);
-			float3 ClothWeight = 0;
-			float ClothPdf = 0;
-			if (NoV > 0 && NoL > 0)
-			{
-				ClothWeight = Data.FuzzAmount * ComputeSheenWeight(V_World, L_World, N_World, NoV, NoL, NoH, VoH, Data.FuzzRoughness, Payload.FuzzColor);
-				ClothPdf = ComputeSheenPdf(V_World, L_World, N_World, NoV, Data.FuzzRoughness);
-			}
-			Result.AddLobeWithMIS(DiffuseSpecularScale.y * ClothWeight * ShadowTerminator, ClothPdf, Data.LobePdf.z);
+			const float4 ClothResult = Data.Fuzz.Eval(L, V, H, Payload.FuzzColor);
+			const float ClothPdf = ClothResult.w;
+			Result.AddLobeWithMIS(DiffuseSpecularScale.y * ClothResult.xyz * ShadowTerminator, ClothPdf, Data.LobePdf.z);
 		}
 
 		Result.Weight *= SubstrateLobeWeight(Payload, NoL);
@@ -257,7 +241,7 @@ FMaterialEval SubstrateSolidGlass_EvalMaterial(
 			Result.AddLobeWithMIS((      Payload.RoughnessData.z) * GGXResult1.x / Data.EGlass1 * DiffuseSpecularScale.y, GGXResult1.y * Fg, Data.LobePdf.y);
 		}
 
-		Result.Weight *= Payload.BSDFOpacity * Payload.WeightV * Data.FuzzAttenuation;
+		Result.Weight *= Payload.BSDFOpacity * Payload.WeightV * Data.Fuzz.Attenuation;
 	}
 	return Result;
 }
diff --git a/Engine/Shaders/Private/PathTracing/Material/PathTracingSubstrateThinGlass.ush b/Engine/Shaders/Private/PathTracing/Material/PathTracingSubstrateThinGlass.ush
--- a/Engine/Shaders/Private/PathTracing/Material/PathTracingSubstrateThinGlass.ush
+++ b/Engine/Shaders/Private/PathTracing/Material/PathTracingSubstrateThinGlass.ush
@@ -26,13 +26,11 @@ struct FSubstrateThinGlassSlabData
 	float2 AlphaR0, AlphaT0;
 	float2 AlphaR1, AlphaT1;
 
-	float  FuzzRoughness;
-	float  FuzzAttenuation;
-	float  FuzzAmount;
-
 	float SpecR0W, SpecR1W;
 	float SpecT0W, SpecT1W;
 
+	FSubstrateSheenData Fuzz;
+
 	float2 LobeCdf;
 	float3 LobePdf;
 
@@ -53,7 +51,6 @@ struct FSubstrateThinGlassSlabData
 		F0 = F0RGBToF0(Payload.SpecularColor);
 		F90 = F0ToMicroOcclusion(F0);
 		TransmissionColor = Payload.GetTransmittanceColor();
-		FuzzRoughness = MakeRoughnessSafe(Payload.FuzzRoughness, SUBSTRATE_MIN_FUZZ_ROUGHNESS);
 
 		const float NoV = saturate(V.z);
 		FBxDFEnergyTermsA SpecR0 = ComputeGGXSpecEnergyTermsA(RoughnessR.x, NoV, F0, F90);
@@ -64,13 +61,12 @@ struct FSubstrateThinGlassSlabData
 			rcp(SpecR0.E + SpecT0.E),
 			rcp(SpecR1.E + SpecT1.E));
 
-		FBxDFEnergyTermsA Fuzz = ComputeClothEnergyTermsA(FuzzRoughness, NoV);
+		Fuzz.Prepare(V, Payload.FuzzRoughness, Payload.FuzzAmount);
 
-		FuzzAttenuation = 1.0 - Payload.FuzzAmount * Fuzz.E;
 
-		const float3 GlassR = FuzzAttenuation * lerp(SpecR0.E * NudgeE.x, SpecR1.E * NudgeE.y, Payload.RoughnessData.z);
-		const float3 GlassT = FuzzAttenuation * lerp(SpecT0.E * NudgeE.x, SpecT1.E * NudgeE.y, Payload.RoughnessData.z);
-		const float3 FuzzAlbedo = Payload.FuzzAmount * Fuzz.E * Payload.FuzzColor;
+		const float3 GlassR = Fuzz.Attenuation * lerp(SpecR0.E * NudgeE.x, SpecR1.E * NudgeE.y, Payload.RoughnessData.z);
+		const float3 GlassT = Fuzz.Attenuation * lerp(SpecT0.E * NudgeE.x, SpecT1.E * NudgeE.y, Payload.RoughnessData.z);
+		const float3 FuzzAlbedo = Fuzz.Scale * Payload.FuzzColor;
 
 		SpecR0W = SpecR0.W * NudgeE.x;
 		SpecR1W = SpecR1.W * NudgeE.y;
@@ -137,7 +133,7 @@ FMaterialSample SubstrateThinGlass_SampleMaterial(
 	{
 		// cloth lobe
 		RandSample.x = RescaleRandomNumber(RandSample.x, Data.LobeCdf.y, 1.0);
-		L = SampleSheenDirection(V, Data.FuzzRoughness, RandSample.xy);
+		L = Data.Fuzz.Sample(RandSample.xy);
 		H = normalize(L + V);
 	}
 	// transform to world space
@@ -157,7 +153,8 @@ FMaterialSample SubstrateThinGlass_SampleMaterial(
 	const float2 GGXResult0 = GGXEvalReflection(L, V, H, bIsReflection ? Data.AlphaR0 : Data.AlphaT0);
 	const float2 GGXResult1 = GGXEvalReflection(L, V, H, bIsReflection ? Data.AlphaR1 : Data.AlphaT1);
 	const float LobePdf = bIsReflection ? Data.LobePdf.y : Data.LobePdf.x;
-	const float ClothPdf = ComputeSheenPdf(V_World, L_World, N_World, NoV, Data.FuzzRoughness);
+	const float4 ClothResult = Data.Fuzz.Eval(L, V, H, Payload.FuzzColor);
+	const float ClothPdf = ClothResult.w;
 	if (bSampleGlass)
 	{
 		FThinSlabWeights SlabResult = ComputeThinSlabWeights(Data.TransmissionColor, VoH, Payload.Ior, Data.F0);
@@ -168,8 +165,8 @@ FMaterialSample SubstrateThinGlass_SampleMaterial(
 		const float SlabW1 = (bIsReflection ? Data.SpecR1W : Data.SpecT1W);
 
 
-		Result.AddLobeWithMIS(Data.FuzzAttenuation * (1 - Payload.RoughnessData.z) * GGXResult0.x * SlabW0 * SlabTint, GGXResult0.y, LobePdf * (1 - Payload.RoughnessData.z));
-		Result.AddLobeWithMIS(Data.FuzzAttenuation * (    Payload.RoughnessData.z) * GGXResult1.x * SlabW1 * SlabTint, GGXResult1.y, LobePdf * (    Payload.RoughnessData.z));
+		Result.AddLobeWithMIS(Data.Fuzz.Attenuation * (1 - Payload.RoughnessData.z) * GGXResult0.x * SlabW0 * SlabTint, GGXResult0.y, LobePdf * (1 - Payload.RoughnessData.z));
+		Result.AddLobeWithMIS(Data.Fuzz.Attenuation * (    Payload.RoughnessData.z) * GGXResult1.x * SlabW1 * SlabTint, GGXResult1.y, LobePdf * (    Payload.RoughnessData.z));
 
 		if (bIsReflection)
 		{
@@ -180,13 +177,7 @@ FMaterialSample SubstrateThinGlass_SampleMaterial(
 	{
 		// Cloth Lobe
 		const float ShadowTerminator = ShadowTerminatorTerm(L_World, N_World, Payload.WorldSmoothNormal);
-		float3 ClothWeight = 0;
-		float ClothPdf = 0;
-		if (NoV > 0 && NoL > 0)
-		{
-			ClothWeight = Data.FuzzAmount * ComputeSheenWeight(V_World, L_World, N_World, NoV, NoL, NoH, VoH, Data.FuzzRoughness, Payload.FuzzColor);
-		}
-		Result.AddLobeWithMIS(ClothWeight * ShadowTerminator, ClothPdf, Data.LobePdf.z);
+		Result.AddLobeWithMIS(ClothResult.xyz * ShadowTerminator, ClothPdf, Data.LobePdf.z);
 
 		Result.Pdf += LobePdf * (1.0 - Payload.RoughnessData.z) * GGXResult0.y;
 		Result.Pdf += LobePdf * (      Payload.RoughnessData.z) * GGXResult1.y;
@@ -245,8 +236,8 @@ FMaterialEval SubstrateThinGlass_EvalMaterial(
 
 	FMaterialEval Result = NullMaterialEval();
 
-	Result.AddLobeWithMIS(DiffuseSpecularScale.y * Data.FuzzAttenuation * (1 - Payload.RoughnessData.z) * GGXResult0.x * SlabW0 * SlabTint, GGXResult0.y, LobePdf * (1 - Payload.RoughnessData.z));
-	Result.AddLobeWithMIS(DiffuseSpecularScale.y * Data.FuzzAttenuation * (    Payload.RoughnessData.z) * GGXResult1.x * SlabW1 * SlabTint, GGXResult1.y, LobePdf * (    Payload.RoughnessData.z));
+	Result.AddLobeWithMIS(DiffuseSpecularScale.y * Data.Fuzz.Attenuation * (1 - Payload.RoughnessData.z) * GGXResult0.x * SlabW0 * SlabTint, GGXResult0.y, LobePdf * (1 - Payload.RoughnessData.z));
+	Result.AddLobeWithMIS(DiffuseSpecularScale.y * Data.Fuzz.Attenuation * (    Payload.RoughnessData.z) * GGXResult1.x * SlabW1 * SlabTint, GGXResult1.y, LobePdf * (    Payload.RoughnessData.z));
 
 
 	if (bIsReflection)
@@ -254,14 +245,9 @@ FMaterialEval SubstrateThinGlass_EvalMaterial(
 		{
 			// Cloth Lobe
 			const float ShadowTerminator = ShadowTerminatorTerm(L_World, N_World, Payload.WorldSmoothNormal);
-			float3 ClothWeight = 0;
-			float ClothPdf = 0;
-			if (NoV > 0 && NoL > 0)
-			{
-				ClothWeight = Data.FuzzAmount * ComputeSheenWeight(V_World, L_World, N_World, NoV, NoL, NoH, VoH, Data.FuzzRoughness, Payload.FuzzColor);
-				ClothPdf = ComputeSheenPdf(V_World, L_World, N_World, NoV, Data.FuzzRoughness);
-			}
-			Result.AddLobeWithMIS(DiffuseSpecularScale.y * ClothWeight * ShadowTerminator, ClothPdf, Data.LobePdf.z);
+			const float4 ClothResult = Data.Fuzz.Eval(L, V, H, Payload.FuzzColor);
+			const float ClothPdf = ClothResult.w;
+			Result.AddLobeWithMIS(DiffuseSpecularScale.y * ClothResult.xyz * ShadowTerminator, ClothPdf, Data.LobePdf.z);
 		}
 		Result.Weight *= SubstrateLobeWeight(Payload, NoL);
 	}
