Modularises TSR's spatial anti-aliaser

diff --git a/Engine/Shaders/Private/TemporalSuperResolution/TSRSpatialAntiAliasing.usf b/Engine/Shaders/Private/TemporalSuperResolution/TSRSpatialAntiAliasing.usf
--- a/Engine/Shaders/Private/TemporalSuperResolution/TSRSpatialAntiAliasing.usf
+++ b/Engine/Shaders/Private/TemporalSuperResolution/TSRSpatialAntiAliasing.usf
@@ -1,14 +1,12 @@
 // Copyright Epic Games, Inc. All Rights Reserved.
 
-#include "TSRCommon.ush"
+#include "TSRSpatialAntiAliasing.ush"
 
 
 //------------------------------------------------------- CONFIG
 
 #define TILE_SIZE 8
 
-#define CONFIG_COMPUTE_NOISE_FILTERING 1
-
 #if DIM_QUALITY_PRESET == 1
 	#define CONFIG_ITERATIONS 3
 
@@ -27,9 +25,6 @@
 #endif
 
 
-//------------------------------------------------------- INCLUDE FXAA
-
-
 //------------------------------------------------------- PARAMETERS
 
 Texture2D<tsr_ushort> AntiAliasMaskTexture;
@@ -58,12 +53,6 @@ tsr_ushort2x2 Map8x8Tile2x2LaneDPV(uint GroupThreadIndex)
 }
 #endif
 
-tsr_half2 InputToLDRLuma(tsr_half3x2 Color)
-{
-	tsr_half2 Luma = Luma4(Color);
-	return Luma * HdrWeightY(Luma);
-}
-
 tsr_half SampleInputLDRLuma(tsr_short2 KernelCenter, tsr_short2 Offset)
 {
 	tsr_short2 SampleInputPixelPos = KernelCenter + ClampPixelOffset(
@@ -88,100 +77,6 @@ tsr_half2 SampleInputLDRLuma(tsr_short2x2 KernelCenter, tsr_short2 Offset)
 	return InputLuma;
 }
 
-tsr_half2 SampleInputLDRLuma(tsr_short2x2 KernelCenter, tsr_short2x2 Offset)
-{
-	tsr_short2x2 SampleInputPixelPos = ClampPixelOffset(
-		KernelCenter + Offset,
-		InputPixelPosMin, InputPixelPosMax);
-	
-	tsr_half2 InputLuma = dpv_interleave_registers(
-		InputSceneColorLdrLumaTexture[dpv_lo(SampleInputPixelPos)],
-		InputSceneColorLdrLumaTexture[dpv_hi(SampleInputPixelPos)]);
-
-	return InputLuma;
-}
-
-tsr_half2 BilinearSampleInputLDRLuma(float2x2 KernelCenterUV, float2x2 UVOffset, const bool bIsNegativeDirection)
-{
-	float2x2 SampleUV = KernelCenterUV + UVOffset;
-
-	if (bIsNegativeDirection)
-		SampleUV = max(SampleUV, dpv_interleave_mono_registers(InputInfo_UVViewportBilinearMin));
-	else
-		SampleUV = min(SampleUV, dpv_interleave_mono_registers(InputInfo_UVViewportBilinearMax));
-
-	tsr_half2 InputLuma = dpv_interleave_registers(
-		InputSceneColorLdrLumaTexture.SampleLevel(GlobalBilinearClampedSampler, dpv_lo(SampleUV), 0),
-		InputSceneColorLdrLumaTexture.SampleLevel(GlobalBilinearClampedSampler, dpv_hi(SampleUV), 0));
-
-	return InputLuma;
-}
-
-void AccumulateEdgeLength(
-	tsr_half2 StartLuma, tsr_half2 StartEdgeLuma,
-	tsr_half2 SampleLuma, tsr_half2 SampleEdgeLuma,
-	tsr_half2 LumaDelta,
-	inout bool2 bEdgeStopedByIncrement,
-	inout bool2 bEdgeStopedByDecrement,
-	inout tsr_ushort2 EdgeLength)
-{
-	bool2 bStopNowIncrement = abs(StartLuma - SampleLuma) > LumaDelta;
-	bool2 bStopNowDecrement = abs(StartEdgeLuma - SampleEdgeLuma) > LumaDelta;
-	
-	bEdgeStopedByIncrement = or(bEdgeStopedByIncrement, bStopNowIncrement);
-	bEdgeStopedByDecrement = or(bEdgeStopedByDecrement, bStopNowDecrement);
-
-	EdgeLength += select(or(bEdgeStopedByIncrement, bEdgeStopedByDecrement), tsr_ushort(0), tsr_ushort(1));
-}
-
-tsr_half ComputeDistanceToEdge(
-	bool bEdgeStopedByIncrementN,
-	bool bEdgeStopedByIncrementP,
-	bool bEdgeStopedByDecrementN,
-	bool bEdgeStopedByDecrementP,
-	tsr_ushort EdgeLengthN,
-	tsr_ushort EdgeLengthP)
-{	
-	tsr_half fEdgeLengthN = tsr_half(EdgeLengthN);
-	tsr_half fEdgeLengthP = tsr_half(EdgeLengthP);
-	tsr_half fEdgeLength = tsr_half(1.0) + fEdgeLengthN + fEdgeLengthP;
-	tsr_half fEdgeInvLength = rcp(fEdgeLength);
-
-	
-	tsr_half TexelCenterOffsetToEdge = 0.0;
-
-	const float Min = -0.5;
-
-	FLATTEN
-	if (!bEdgeStopedByIncrementN && !bEdgeStopedByIncrementP && !bEdgeStopedByDecrementP && !bEdgeStopedByDecrementN)
-	{
-		// No aliasing in the neighborhood.
-		TexelCenterOffsetToEdge = tsr_half(0.0);
-	}
-	else if (bEdgeStopedByIncrementN && bEdgeStopedByIncrementP)
-	{
-		// Looks like InputC is concave detail
-		TexelCenterOffsetToEdge = tsr_half(0.5);
-	}
-	else if (bEdgeStopedByDecrementN && bEdgeStopedByDecrementP)
-	{
-		// Looks like InputC is convex detail
-		//TexelCenterOffsetToEdge = 0.5;
-	}
-	else if (bEdgeStopedByIncrementN) // && bEdgeStopedByDecrementP)
-	{
-		// Looks like staircasing from - to +
-		TexelCenterOffsetToEdge = tsr_half(0.5) - (EdgeLengthN + tsr_half(0.5)) * fEdgeInvLength;
-	}
-	else if (bEdgeStopedByIncrementP) // && bEdgeStopedByDecrementN)
-	{
-		// Looks like staircasing from + to -
-		TexelCenterOffsetToEdge = tsr_half(0.5) - (EdgeLengthP + tsr_half(0.5)) * fEdgeInvLength;
-	}
-
-	return TexelCenterOffsetToEdge;
-}
-
 
 //------------------------------------------------------- ENTRY POINT
 
@@ -236,218 +131,56 @@ void MainCS(
 			}
 
 			#define ReadCache(x, y) tsr_half2(InputCache[(x + 1) + 4 * (y + 1)], InputCache[(x + 2) + 4 * (y + 1)])
-
-			tsr_half2 InputC = ReadCache( 0,  0);
-
-			tsr_half2 InputN = ReadCache( 0, -1);
-			tsr_half2 InputS = ReadCache( 0, +1);
-			tsr_half2 InputE = ReadCache(+1,  0);
-			tsr_half2 InputW = ReadCache(-1,  0);
-	
-			tsr_half2 InputNE = ReadCache(+1, -1);
-			tsr_half2 InputNW = ReadCache(-1, -1);
-			tsr_half2 InputSE = ReadCache(+1, +1);
-			tsr_half2 InputSW = ReadCache(-1, +1);
 		#else
-			tsr_half2 InputC = SampleInputLDRLuma(InputPixelPos, tsr_short2( 0,  0));
-
-			tsr_half2 InputN = SampleInputLDRLuma(InputPixelPos, tsr_short2( 0, -1));
-			tsr_half2 InputS = SampleInputLDRLuma(InputPixelPos, tsr_short2( 0, +1));
-			tsr_half2 InputE = SampleInputLDRLuma(InputPixelPos, tsr_short2(+1,  0));
-			tsr_half2 InputW = SampleInputLDRLuma(InputPixelPos, tsr_short2(-1,  0));
-	
-			tsr_half2 InputNE = SampleInputLDRLuma(InputPixelPos, tsr_short2(+1, -1));
-			tsr_half2 InputNW = SampleInputLDRLuma(InputPixelPos, tsr_short2(-1, -1));
-			tsr_half2 InputSE = SampleInputLDRLuma(InputPixelPos, tsr_short2(+1, +1));
-			tsr_half2 InputSW = SampleInputLDRLuma(InputPixelPos, tsr_short2(-1, +1));
+			#define ReadCache(x, y) SampleInputLDRLuma(InputPixelPos, tsr_short2(x, y));
 		#endif
-
-		bool2 bBrowseVerticaly;
-		tsr_short2x2 BrowseDirection;
-		tsr_short2x2 EdgeSide = 0;
-		tsr_half2 EdgeLuma = 0;
-
-		tsr_half2 TotalVarHN = abs(0.5 * (InputNE + InputNW) - InputN);
-		tsr_half2 TotalVarH  = abs(0.5 * (InputE  + InputW ) - InputC);
-		tsr_half2 TotalVarHS = abs(0.5 * (InputSE + InputSW) - InputS);
-			
-		tsr_half2 TotalVarVE = abs(0.5 * (InputNE + InputSE) - InputE);
-		tsr_half2 TotalVarV  = abs(0.5 * (InputN  + InputS ) - InputC);
-		tsr_half2 TotalVarVW = abs(0.5 * (InputNW + InputSW) - InputW);
 		
-		tsr_half2 DiffN = abs(InputN - InputC);
-		tsr_half2 DiffS = abs(InputS - InputC);
-		tsr_half2 DiffE = abs(InputE - InputC);
-		tsr_half2 DiffW = abs(InputW - InputC);
+		tsr_half2 InputC = ReadCache( 0,  0);
 
-		#if CONFIG_COMPUTE_NOISE_FILTERING
-		{
-			NoiseFiltering = saturate(tsr_half(2.0) * TotalVarH - max(DiffE, DiffW));
-			NoiseFiltering = max(NoiseFiltering, saturate(tsr_half(2.0) * TotalVarV - max(DiffN, DiffS)));
-		}
-		#else
-		{
-			NoiseFiltering = tsr_half(0.0);
-		}
-		#endif
-
-		bBrowseVerticaly = (TotalVarHN + TotalVarH + TotalVarHS) > (TotalVarVE + TotalVarV + TotalVarVW);
-		BrowseDirection[0] = select(bBrowseVerticaly, tsr_short(0), tsr_short(1));
-		BrowseDirection[1] = select(bBrowseVerticaly, tsr_short(1), tsr_short(0));
+		tsr_half2 InputN = ReadCache( 0, -1);
+		tsr_half2 InputS = ReadCache( 0, +1);
+		tsr_half2 InputE = ReadCache(+1,  0);
+		tsr_half2 InputW = ReadCache(-1,  0);
+	
+		tsr_half2 InputNE = ReadCache(+1, -1);
+		tsr_half2 InputNW = ReadCache(-1, -1);
+		tsr_half2 InputSE = ReadCache(+1, +1);
+		tsr_half2 InputSW = ReadCache(-1, +1);
 
-		FLATTEN
-		if (bBrowseVerticaly[0])
-		{
-			EdgeSide[0][0] = DiffW[0] > DiffE[0] ? -1 : 1;
-			EdgeLuma[0] = DiffW[0] > DiffE[0] ? InputW[0] : InputE[0];
-		}
-		else
-		{
-			EdgeSide[1][0] = DiffN[0] > DiffS[0] ? -1 : 1;
-			EdgeLuma[0] = DiffN[0] > DiffS[0] ? InputN[0] : InputS[0];
-		}
-		
-		FLATTEN
-		if (bBrowseVerticaly[1])
-		{
-			EdgeSide[0][1] = DiffW[1] > DiffE[1] ? -1 : 1;
-			EdgeLuma[1] = DiffW[1] > DiffE[1] ? InputW[1] : InputE[1];
-		}
-		else
-		{
-			EdgeSide[1][1] = DiffN[1] > DiffS[1] ? -1 : 1;
-			EdgeLuma[1] = DiffN[1] > DiffS[1] ? InputN[1] : InputS[1];
-		}
+		tsr_short2x2 BrowseDirection;
+		tsr_short2x2 EdgeSide;
+		tsr_half2 EdgeLuma;
+		FindBrowsingDirection(
+			InputC,
+			InputN, InputS, InputE, InputW,
+			InputNE, InputNW, InputSE, InputSW,
+			/* out */ NoiseFiltering,
+			/* out */ BrowseDirection,
+			/* out */ EdgeSide,
+			/* out */ EdgeLuma);
 
 		tsr_half2 LumaDelta = abs(EdgeLuma - InputC) * 0.5;
-		
 		bool2 WorthBrowsing = LumaDelta > SPATIAL_ANTI_ALIASER_MIN_LUMIMANCE;
-
 		tsr_half2 BrowseDelta = saturate(LumaDelta * rcp(tsr_half(SPATIAL_ANTI_ALIASER_MIN_LUMIMANCE)));
 		
-		tsr_ushort2 EdgeLengthP = 0;
-		tsr_ushort2 EdgeLengthN = 0;
-		
-		/**
-		 * Increment:
-		 *     1 1 1 1 1
-		 *     0 0 X 0 1
-		 * 
-		 * Decrement:
-		 *     1 1 1 1 0
-		 *     0 0 X 0 0
-		 */
-		bool2 bEdgeStopedByIncrementP = false;
-		bool2 bEdgeStopedByIncrementN = false;
-		
-		bool2 bEdgeStopedByDecrementP = false;
-		bool2 bEdgeStopedByDecrementN = false;
-		
-#if 1 // Optimized for RDNA 2
-		const tsr_half BilinearInterp = 0.5;
-
-		float2x2 fBrowseDirection = float2x2(BrowseDirection) * dpv_interleave_mono_registers(InputInfo_ExtentInverse);
-		float2x2 KernelUV = (float2x2(InputPixelPos) + 0.5 + float2x2(EdgeSide) * BilinearInterp) * dpv_interleave_mono_registers(InputInfo_ExtentInverse);
-
-		tsr_half2 MergedEdgeLuma = BilinearInterp * (EdgeLuma + InputC);
-		
-		tsr_half2 MinLuma = MergedEdgeLuma - tsr_half(BilinearInterp * 0.5) * abs(EdgeLuma - InputC);
-		tsr_half2 MaxLuma = MergedEdgeLuma + tsr_half(BilinearInterp * 0.5) * abs(EdgeLuma - InputC);
-
-		bool2 bEdgeMoreLuminous = EdgeLuma > InputC;
-		
-		tsr_ushort2 EdgeStopedByMinLumaP = 0;
-		tsr_ushort2 EdgeStopedByMinLumaN = 0;
-		
-		tsr_ushort2 EdgeStopedByMaxLumaP = 0;
-		tsr_ushort2 EdgeStopedByMaxLumaN = 0;
-		
-		bool bStop = false;
-		
-		#if CONFIG_COMPILE_FP16
-			#define asuint(x) asuint16(x)
-			const tsr_ushort SignShift = 15;
-			const tsr_ushort SignMask = 0x8000;
-		#else
-			const tsr_ushort SignShift = 31;
-			const tsr_ushort SignMask = 0x80000000;
-		#endif
-			
-		EdgeLengthP = CONFIG_ITERATIONS;
-		EdgeLengthN = CONFIG_ITERATIONS;
-		
-		UNROLL_N(CONFIG_ITERATIONS)
-		for (uint i = 0; i < CONFIG_ITERATIONS; i++)
-		{
-			float2x2 NeighborOffset = fBrowseDirection * float(i + 1);
-
-			tsr_half2 SampleLumaP = BilinearSampleInputLDRLuma(KernelUV,  NeighborOffset, /* bIsNegativeDirection = */ false);
-			tsr_half2 SampleLumaN = BilinearSampleInputLDRLuma(KernelUV, -NeighborOffset, /* bIsNegativeDirection = */ true);
-				
-			//bEdgeStopedByMinLuma = bEdgeStopedByMinLuma || bool2(SampleLuma < MinLuma) && !bEdgeStopedByMaxLuma;
-			//bEdgeStopedByMaxLuma = bEdgeStopedByMaxLuma || bool2(SampleLuma > MaxLuma) && !bEdgeStopedByMinLuma;
-			tsr_ushort2 StopByMinLumaP = asuint(SampleLumaP - MinLuma) & ~EdgeStopedByMaxLumaP;
-			tsr_ushort2 StopByMaxLumaP = asuint(MaxLuma - SampleLumaP) & ~EdgeStopedByMinLumaP;
-			tsr_ushort2 StopByMinLumaN = asuint(SampleLumaN - MinLuma) & ~EdgeStopedByMaxLumaN;
-			tsr_ushort2 StopByMaxLumaN = asuint(MaxLuma - SampleLumaN) & ~EdgeStopedByMinLumaN;
-
-			EdgeStopedByMinLumaP |= StopByMinLumaP;
-			EdgeStopedByMaxLumaP |= StopByMaxLumaP;
-			EdgeStopedByMinLumaN |= StopByMinLumaN;
-			EdgeStopedByMaxLumaN |= StopByMaxLumaN;
-				
-			//EdgeLength += (bEdgeStopedByMaxLuma || bEdgeStopedByMinLuma) ? 0 : 1;
-			EdgeLengthP -= (EdgeStopedByMaxLumaP | EdgeStopedByMinLumaP) >> SignShift;
-			EdgeLengthN -= (EdgeStopedByMaxLumaN | EdgeStopedByMinLumaN) >> SignShift;
-		}
-		
-		bool2 bEdgeStopedByMinLumaP = (EdgeStopedByMinLumaP & SignMask) != 0;
-		bool2 bEdgeStopedByMinLumaN = (EdgeStopedByMinLumaN & SignMask) != 0;
-		
-		bool2 bEdgeStopedByMaxLumaP = (EdgeStopedByMaxLumaP & SignMask) != 0;
-		bool2 bEdgeStopedByMaxLumaN = (EdgeStopedByMaxLumaN & SignMask) != 0;
-		
-		bEdgeStopedByIncrementP = or(and(bEdgeMoreLuminous, bEdgeStopedByMaxLumaP), and(!bEdgeMoreLuminous, bEdgeStopedByMinLumaP));
-		bEdgeStopedByIncrementN = or(and(bEdgeMoreLuminous, bEdgeStopedByMaxLumaN), and(!bEdgeMoreLuminous, bEdgeStopedByMinLumaN));
-		
-		bEdgeStopedByDecrementP = or(and(bEdgeMoreLuminous, bEdgeStopedByMinLumaP), and(!bEdgeMoreLuminous, bEdgeStopedByMaxLumaP));
-		bEdgeStopedByDecrementN = or(and(bEdgeMoreLuminous, bEdgeStopedByMinLumaN), and(!bEdgeMoreLuminous, bEdgeStopedByMaxLumaN));
-
-#else // Reference implementation
-		UNROLL_N(CONFIG_ITERATIONS)
-		for (uint i = 0; i < CONFIG_ITERATIONS; i++)
-		{
-			bool2 bIsStoped = (bEdgeStopedByIncrementP || bEdgeStopedByDecrementP) && (bEdgeStopedByIncrementN || bEdgeStopedByDecrementP);
-
-			BRANCH
-			if (!bIsStoped.x)
-			{
-				tsr_short2x2 NeighborOffset = BrowseDirection * tsr_short(i + 1);
-
-				tsr_half2 NeighborP = SampleInputLDRLuma(InputPixelPos, NeighborOffset);
-				tsr_half2 EdgeP     = SampleInputLDRLuma(InputPixelPos, NeighborOffset + EdgeSide);
-			
-				tsr_half2 NeighborN = SampleInputLDRLuma(InputPixelPos, -NeighborOffset);
-				tsr_half2 EdgeN     = SampleInputLDRLuma(InputPixelPos, -NeighborOffset + EdgeSide);
-
-				AccumulateEdgeLength(
-					InputC, EdgeLuma,
-					NeighborP, EdgeP,
-					LumaDelta,
-					/* inout */ bEdgeStopedByIncrementP,
-					/* inout */ bEdgeStopedByDecrementP,
-					/* inout */ EdgeLengthP);
-			
-				AccumulateEdgeLength(
-					InputC, EdgeLuma,
-					NeighborN, EdgeN,
-					LumaDelta,
-					/* inout */ bEdgeStopedByIncrementN,
-					/* inout */ bEdgeStopedByDecrementN,
-					/* inout */ EdgeLengthN);
-			}
-		}
-#endif
+		tsr_ushort2 EdgeLengthP, EdgeLengthN;
+		bool2 bEdgeStopedByIncrementP, bEdgeStopedByIncrementN;
+		bool2 bEdgeStopedByDecrementP, bEdgeStopedByDecrementN;
+		
+		BrowseNeighborhoodBilinearOptimized(
+			/* Iterations = */ CONFIG_ITERATIONS,
+			/* Texture = */ InputSceneColorLdrLumaTexture,
+			InputC,
+			EdgeLuma,
+			InputPixelPos,
+			BrowseDirection,
+			EdgeSide,
+			/* out */ EdgeLengthP,
+			/* out */ EdgeLengthN,
+			/* out */ bEdgeStopedByIncrementP,
+			/* out */ bEdgeStopedByIncrementN,
+			/* out */ bEdgeStopedByDecrementP,
+			/* out */ bEdgeStopedByDecrementN);
 
 		tsr_half TexelOffsetLo = ComputeDistanceToEdge(
 			bEdgeStopedByIncrementN[0],
@@ -474,40 +207,6 @@ void MainCS(
 			TexelOffsetHi = 0.0;
 
 		TexelOffset = dpv_scale(tsr_half2x2(EdgeSide), BrowseDelta * dpv_interleave_registers(TexelOffsetLo, TexelOffsetHi));
-		
-		#if 0
-			if (WorthBrowsing.x)
-			{
-				Debug[0] = abs(TexelOffset[0]);
-				Debug[1] = abs(TexelOffset[1]);
-				//Debug[2] = abs(DecodeSpatialAntiAliasingOffset(EncodeSpatialAntiAliasingOffset(TexelOffset))[0]);
-				//Debug[3] = abs(DecodeSpatialAntiAliasingOffset(EncodeSpatialAntiAliasingOffset(TexelOffset))[1]);
-				//Debug[3] = bBrowseVerticaly ? 1 : 0;
-
-				//Debug[0] = EdgeLengthP * tsr_ushort2(WorthBrowsing && bBrowseVerticaly);
-				//Debug[1] = EdgeLengthN * tsr_ushort2(WorthBrowsing && bBrowseVerticaly);
-				//Debug[0] = (1 + EdgeLengthP + EdgeLengthN);
-				//Debug[0] = EdgeLengthN;
-		
-				//Debug[0] = WorthBrowsing ? 1 : 0;
-				//Debug[0] = EdgeLuma;
-				//Debug[0] = (EdgeSide[0] + 2) * 0.25;
-				//Debug[1] = (EdgeSide[1] + 2) * 0.25;
-			}
-
-			if (1)
-			{
-				Debug[2][0] = InputSceneColorLdrLumaTexture.SampleLevel(
-					GlobalBilinearClampedSampler,
-					(float2(dpv_lo(InputPixelPos)) + 0.5 + float2(dpv_lo(TexelOffset))) * InputInfo_ExtentInverse,
-					0);
-				
-				Debug[2][1] = InputSceneColorLdrLumaTexture.SampleLevel(
-					GlobalBilinearClampedSampler,
-					(float2(dpv_hi(InputPixelPos)) + 0.5 + float2(dpv_hi(TexelOffset))) * InputInfo_ExtentInverse,
-					0);
-			}
-		#endif
 	}
 
 	uint2 EncodedTexelOffset = EncodeSpatialAntiAliasingOffset(TexelOffset);
diff --git a/Engine/Shaders/Private/TemporalSuperResolution/TSRSpatialAntiAliasing.ush b/Engine/Shaders/Private/TemporalSuperResolution/TSRSpatialAntiAliasing.ush
new file mode 100644
--- /dev/null
+++ b/Engine/Shaders/Private/TemporalSuperResolution/TSRSpatialAntiAliasing.ush
@@ -0,0 +1,330 @@
+// Copyright Epic Games, Inc. All Rights Reserved.
+
+#include "TSRCommon.ush"
+
+
+//------------------------------------------------------- SAMPLING FUNCTION
+
+tsr_half2 SampleAAInput(Texture2D<tsr_half> Texture, tsr_short2x2 KernelCenter, tsr_short2x2 Offset)
+{
+	tsr_short2x2 SampleInputPixelPos = ClampPixelOffset(
+		KernelCenter + Offset,
+		InputPixelPosMin, InputPixelPosMax);
+	
+	tsr_half2 InputLuma = dpv_interleave_registers(
+		Texture[dpv_lo(SampleInputPixelPos)],
+		Texture[dpv_hi(SampleInputPixelPos)]);
+
+	return InputLuma;
+}
+
+tsr_half2 BilinearSampleAAInput(Texture2D<tsr_half> Texture, float2x2 KernelCenterUV, float2x2 UVOffset, const bool bIsNegativeDirection)
+{
+	float2x2 SampleUV = KernelCenterUV + UVOffset;
+
+	if (bIsNegativeDirection)
+		SampleUV = max(SampleUV, dpv_interleave_mono_registers(InputInfo_UVViewportBilinearMin));
+	else
+		SampleUV = min(SampleUV, dpv_interleave_mono_registers(InputInfo_UVViewportBilinearMax));
+
+	tsr_half2 InputLuma = dpv_interleave_registers(
+		Texture.SampleLevel(GlobalBilinearClampedSampler, dpv_lo(SampleUV), 0),
+		Texture.SampleLevel(GlobalBilinearClampedSampler, dpv_hi(SampleUV), 0));
+
+	return InputLuma;
+}
+
+
+//------------------------------------------------------- FUNCTIONS
+
+/** Finds the browsing direction of the spatial anti-aliaser out from the neighborhood. */
+void FindBrowsingDirection(
+	tsr_half2 InputC,
+	tsr_half2 InputN,
+	tsr_half2 InputS,
+	tsr_half2 InputE,
+	tsr_half2 InputW,
+	tsr_half2 InputNE,
+	tsr_half2 InputNW,
+	tsr_half2 InputSE,
+	tsr_half2 InputSW,
+	out tsr_half2 NoiseFiltering,
+	out tsr_short2x2 BrowseDirection,
+	out tsr_short2x2 EdgeSide,
+	out tsr_half2 EdgeLuma)
+{
+	tsr_half2 TotalVarHN = abs(0.5 * (InputNE + InputNW) - InputN);
+	tsr_half2 TotalVarH  = abs(0.5 * (InputE  + InputW ) - InputC);
+	tsr_half2 TotalVarHS = abs(0.5 * (InputSE + InputSW) - InputS);
+			
+	tsr_half2 TotalVarVE = abs(0.5 * (InputNE + InputSE) - InputE);
+	tsr_half2 TotalVarV  = abs(0.5 * (InputN  + InputS ) - InputC);
+	tsr_half2 TotalVarVW = abs(0.5 * (InputNW + InputSW) - InputW);
+		
+	tsr_half2 DiffN = abs(InputN - InputC);
+	tsr_half2 DiffS = abs(InputS - InputC);
+	tsr_half2 DiffE = abs(InputE - InputC);
+	tsr_half2 DiffW = abs(InputW - InputC);
+
+	bool2 bBrowseVerticaly = (TotalVarHN + TotalVarH + TotalVarHS) > (TotalVarVE + TotalVarV + TotalVarVW);
+	BrowseDirection[0] = select(bBrowseVerticaly, tsr_short(0), tsr_short(1));
+	BrowseDirection[1] = select(bBrowseVerticaly, tsr_short(1), tsr_short(0));
+	
+	{
+		NoiseFiltering = saturate(tsr_half(2.0) * TotalVarH - max(DiffE, DiffW));
+		NoiseFiltering = max(NoiseFiltering, saturate(tsr_half(2.0) * TotalVarV - max(DiffN, DiffS)));
+	}
+
+	EdgeSide = tsr_short(0);
+
+	FLATTEN
+	if (bBrowseVerticaly[0])
+	{
+		EdgeSide[0][0] = DiffW[0] > DiffE[0] ? -1 : 1;
+		EdgeLuma[0] = DiffW[0] > DiffE[0] ? InputW[0] : InputE[0];
+	}
+	else
+	{
+		EdgeSide[1][0] = DiffN[0] > DiffS[0] ? -1 : 1;
+		EdgeLuma[0] = DiffN[0] > DiffS[0] ? InputN[0] : InputS[0];
+	}
+		
+	FLATTEN
+	if (bBrowseVerticaly[1])
+	{
+		EdgeSide[0][1] = DiffW[1] > DiffE[1] ? -1 : 1;
+		EdgeLuma[1] = DiffW[1] > DiffE[1] ? InputW[1] : InputE[1];
+	}
+	else
+	{
+		EdgeSide[1][1] = DiffN[1] > DiffS[1] ? -1 : 1;
+		EdgeLuma[1] = DiffN[1] > DiffS[1] ? InputN[1] : InputS[1];
+	}
+} // FindBrowsingDirection()
+
+/** Accumulate an edge length */
+void AccumulateEdgeLength(
+	tsr_half2 StartLuma, tsr_half2 StartEdgeLuma,
+	tsr_half2 SampleLuma, tsr_half2 SampleEdgeLuma,
+	tsr_half2 LumaDelta,
+	inout bool2 bEdgeStopedByIncrement,
+	inout bool2 bEdgeStopedByDecrement,
+	inout tsr_ushort2 EdgeLength)
+{
+	bool2 bStopNowIncrement = abs(StartLuma - SampleLuma) > LumaDelta;
+	bool2 bStopNowDecrement = abs(StartEdgeLuma - SampleEdgeLuma) > LumaDelta;
+	
+	/**
+	 * Increment:
+	 *     1 1 1 1 1
+	 *     0 0 X 0 1
+	 * 
+	 * Decrement:
+	 *     1 1 1 1 0
+	 *     0 0 X 0 0
+	 */
+	bEdgeStopedByIncrement = or(bEdgeStopedByIncrement, bStopNowIncrement);
+	bEdgeStopedByDecrement = or(bEdgeStopedByDecrement, bStopNowDecrement);
+
+	EdgeLength += select(or(bEdgeStopedByIncrement, bEdgeStopedByDecrement), tsr_ushort(0), tsr_ushort(1));
+}
+
+void BrowseNeighborhood(
+	const uint Iterations,
+	Texture2D<tsr_half> Texture,
+	tsr_half2 InputC,
+	tsr_half2 EdgeLuma,
+	tsr_half2 LumaDelta,
+	tsr_short2x2 KernelCenter,
+	tsr_short2x2 BrowseDirection,
+	tsr_short2x2 EdgeSide,
+	out tsr_ushort2 EdgeLengthP,
+	out tsr_ushort2 EdgeLengthN,
+	out bool2 bEdgeStopedByIncrementP,
+	out bool2 bEdgeStopedByIncrementN,
+	out bool2 bEdgeStopedByDecrementP,
+	out bool2 bEdgeStopedByDecrementN)
+{
+	EdgeLengthP = 0;
+	EdgeLengthN = 0;
+	bEdgeStopedByIncrementP = false;
+	bEdgeStopedByIncrementN = false;
+	bEdgeStopedByDecrementP = false;
+	bEdgeStopedByDecrementN = false;
+	
+	UNROLL
+	for (uint i = 0; i < Iterations; i++)
+	{
+		bool2 bIsStoped = and(or(bEdgeStopedByIncrementP, bEdgeStopedByDecrementP), or(bEdgeStopedByIncrementN, bEdgeStopedByDecrementP));
+
+		BRANCH
+		if (!bIsStoped.x)
+		{
+			tsr_short2x2 NeighborOffset = BrowseDirection * tsr_short(i + 1);
+
+			tsr_half2 NeighborP = SampleAAInput(Texture, KernelCenter, NeighborOffset);
+			tsr_half2 EdgeP     = SampleAAInput(Texture, KernelCenter, NeighborOffset + EdgeSide);
+			
+			tsr_half2 NeighborN = SampleAAInput(Texture, KernelCenter, -NeighborOffset);
+			tsr_half2 EdgeN     = SampleAAInput(Texture, KernelCenter, -NeighborOffset + EdgeSide);
+
+			AccumulateEdgeLength(
+				InputC, EdgeLuma,
+				NeighborP, EdgeP,
+				LumaDelta,
+				/* inout */ bEdgeStopedByIncrementP,
+				/* inout */ bEdgeStopedByDecrementP,
+				/* inout */ EdgeLengthP);
+			
+			AccumulateEdgeLength(
+				InputC, EdgeLuma,
+				NeighborN, EdgeN,
+				LumaDelta,
+				/* inout */ bEdgeStopedByIncrementN,
+				/* inout */ bEdgeStopedByDecrementN,
+				/* inout */ EdgeLengthN);
+		}
+	}
+} // BrowseNeighborhood()
+
+void BrowseNeighborhoodBilinearOptimized(
+	const uint Iterations,
+	Texture2D<tsr_half> Texture,
+	tsr_half2 InputC,
+	tsr_half2 EdgeLuma,
+	tsr_short2x2 InputPixelPos,
+	tsr_short2x2 BrowseDirection,
+	tsr_short2x2 EdgeSide,
+	out tsr_ushort2 EdgeLengthP,
+	out tsr_ushort2 EdgeLengthN,
+	out bool2 bEdgeStopedByIncrementP,
+	out bool2 bEdgeStopedByIncrementN,
+	out bool2 bEdgeStopedByDecrementP,
+	out bool2 bEdgeStopedByDecrementN)
+{
+	EdgeLengthP = 0;
+	EdgeLengthN = 0;
+	bEdgeStopedByIncrementP = false;
+	bEdgeStopedByIncrementN = false;
+	bEdgeStopedByDecrementP = false;
+	bEdgeStopedByDecrementN = false;
+
+	const tsr_half BilinearInterp = 0.5;
+
+	float2x2 fBrowseDirection = float2x2(BrowseDirection) * dpv_interleave_mono_registers(InputInfo_ExtentInverse);
+	float2x2 KernelUV = (float2x2(InputPixelPos) + 0.5 + float2x2(EdgeSide) * BilinearInterp) * dpv_interleave_mono_registers(InputInfo_ExtentInverse);
+
+	tsr_half2 MergedEdgeLuma = BilinearInterp * (EdgeLuma + InputC);
+		
+	tsr_half2 MinLuma = MergedEdgeLuma - tsr_half(BilinearInterp * 0.5) * abs(EdgeLuma - InputC);
+	tsr_half2 MaxLuma = MergedEdgeLuma + tsr_half(BilinearInterp * 0.5) * abs(EdgeLuma - InputC);
+
+	bool2 bEdgeMoreLuminous = EdgeLuma > InputC;
+		
+	tsr_ushort2 EdgeStopedByMinLumaP = 0;
+	tsr_ushort2 EdgeStopedByMinLumaN = 0;
+		
+	tsr_ushort2 EdgeStopedByMaxLumaP = 0;
+	tsr_ushort2 EdgeStopedByMaxLumaN = 0;
+		
+	bool bStop = false;
+		
+	#if CONFIG_COMPILE_FP16
+		#define asuint(x) asuint16(x)
+		const tsr_ushort SignShift = 15;
+		const tsr_ushort SignMask = 0x8000;
+	#else
+		const tsr_ushort SignShift = 31;
+		const tsr_ushort SignMask = 0x80000000;
+	#endif
+			
+	EdgeLengthP = tsr_ushort(Iterations);
+	EdgeLengthN = tsr_ushort(Iterations);
+		
+	UNROLL
+	for (uint i = 0; i < Iterations; i++)
+	{
+		float2x2 NeighborOffset = fBrowseDirection * float(i + 1);
+
+		tsr_half2 SampleLumaP = BilinearSampleAAInput(Texture, KernelUV,  NeighborOffset, /* bIsNegativeDirection = */ false);
+		tsr_half2 SampleLumaN = BilinearSampleAAInput(Texture, KernelUV, -NeighborOffset, /* bIsNegativeDirection = */ true);
+				
+		//bEdgeStopedByMinLuma = bEdgeStopedByMinLuma || bool2(SampleLuma < MinLuma) && !bEdgeStopedByMaxLuma;
+		//bEdgeStopedByMaxLuma = bEdgeStopedByMaxLuma || bool2(SampleLuma > MaxLuma) && !bEdgeStopedByMinLuma;
+		tsr_ushort2 StopByMinLumaP = asuint(SampleLumaP - MinLuma) & ~EdgeStopedByMaxLumaP;
+		tsr_ushort2 StopByMaxLumaP = asuint(MaxLuma - SampleLumaP) & ~EdgeStopedByMinLumaP;
+		tsr_ushort2 StopByMinLumaN = asuint(SampleLumaN - MinLuma) & ~EdgeStopedByMaxLumaN;
+		tsr_ushort2 StopByMaxLumaN = asuint(MaxLuma - SampleLumaN) & ~EdgeStopedByMinLumaN;
+
+		EdgeStopedByMinLumaP |= StopByMinLumaP;
+		EdgeStopedByMaxLumaP |= StopByMaxLumaP;
+		EdgeStopedByMinLumaN |= StopByMinLumaN;
+		EdgeStopedByMaxLumaN |= StopByMaxLumaN;
+				
+		//EdgeLength += (bEdgeStopedByMaxLuma || bEdgeStopedByMinLuma) ? 0 : 1;
+		EdgeLengthP -= (EdgeStopedByMaxLumaP | EdgeStopedByMinLumaP) >> SignShift;
+		EdgeLengthN -= (EdgeStopedByMaxLumaN | EdgeStopedByMinLumaN) >> SignShift;
+	}
+		
+	bool2 bEdgeStopedByMinLumaP = (EdgeStopedByMinLumaP & SignMask) != 0;
+	bool2 bEdgeStopedByMinLumaN = (EdgeStopedByMinLumaN & SignMask) != 0;
+		
+	bool2 bEdgeStopedByMaxLumaP = (EdgeStopedByMaxLumaP & SignMask) != 0;
+	bool2 bEdgeStopedByMaxLumaN = (EdgeStopedByMaxLumaN & SignMask) != 0;
+		
+	bEdgeStopedByIncrementP = or(and(bEdgeMoreLuminous, bEdgeStopedByMaxLumaP), and(!bEdgeMoreLuminous, bEdgeStopedByMinLumaP));
+	bEdgeStopedByIncrementN = or(and(bEdgeMoreLuminous, bEdgeStopedByMaxLumaN), and(!bEdgeMoreLuminous, bEdgeStopedByMinLumaN));
+		
+	bEdgeStopedByDecrementP = or(and(bEdgeMoreLuminous, bEdgeStopedByMinLumaP), and(!bEdgeMoreLuminous, bEdgeStopedByMaxLumaP));
+	bEdgeStopedByDecrementN = or(and(bEdgeMoreLuminous, bEdgeStopedByMinLumaN), and(!bEdgeMoreLuminous, bEdgeStopedByMaxLumaN));
+
+
+} // BrowseNeighborhoodBilinearRDNAOptimized
+
+tsr_half ComputeDistanceToEdge(
+	bool bEdgeStopedByIncrementN,
+	bool bEdgeStopedByIncrementP,
+	bool bEdgeStopedByDecrementN,
+	bool bEdgeStopedByDecrementP,
+	tsr_ushort EdgeLengthN,
+	tsr_ushort EdgeLengthP)
+{	
+	tsr_half fEdgeLengthN = tsr_half(EdgeLengthN);
+	tsr_half fEdgeLengthP = tsr_half(EdgeLengthP);
+	tsr_half fEdgeLength = tsr_half(1.0) + fEdgeLengthN + fEdgeLengthP;
+	tsr_half fEdgeInvLength = rcp(fEdgeLength);
+
+	
+	tsr_half TexelCenterOffsetToEdge = 0.0;
+
+	const float Min = -0.5;
+
+	FLATTEN
+	if (!bEdgeStopedByIncrementN && !bEdgeStopedByIncrementP && !bEdgeStopedByDecrementP && !bEdgeStopedByDecrementN)
+	{
+		// No aliasing in the neighborhood.
+		TexelCenterOffsetToEdge = tsr_half(0.0);
+	}
+	else if (bEdgeStopedByIncrementN && bEdgeStopedByIncrementP)
+	{
+		// Looks like InputC is concave detail
+		TexelCenterOffsetToEdge = tsr_half(0.5);
+	}
+	else if (bEdgeStopedByDecrementN && bEdgeStopedByDecrementP)
+	{
+		// Looks like InputC is convex detail
+		//TexelCenterOffsetToEdge = 0.5;
+	}
+	else if (bEdgeStopedByIncrementN) // && bEdgeStopedByDecrementP)
+	{
+		// Looks like staircasing from - to +
+		TexelCenterOffsetToEdge = tsr_half(0.5) - (EdgeLengthN + tsr_half(0.5)) * fEdgeInvLength;
+	}
+	else if (bEdgeStopedByIncrementP) // && bEdgeStopedByDecrementN)
+	{
+		// Looks like staircasing from + to -
+		TexelCenterOffsetToEdge = tsr_half(0.5) - (EdgeLengthP + tsr_half(0.5)) * fEdgeInvLength;
+	}
+
+	return TexelCenterOffsetToEdge;
+} // ComputeDistanceToEdge()
