Fix shader warnings.

diff --git a/Engine/Shaders/Private/VirtualShadowMaps/VirtualShadowMapProjectionCommon.ush b/Engine/Shaders/Private/VirtualShadowMaps/VirtualShadowMapProjectionCommon.ush
--- a/Engine/Shaders/Private/VirtualShadowMaps/VirtualShadowMapProjectionCommon.ush
+++ b/Engine/Shaders/Private/VirtualShadowMaps/VirtualShadowMapProjectionCommon.ush
@@ -104,9 +104,9 @@ FVirtualShadowMapSample SampleVirtualShadowMapLevel(int VirtualShadowMapId, floa
 FVirtualShadowMapSample SampleVirtualShadowMap(int VirtualShadowMapId, float2 ShadowMapUV)
 {
 	FShadowPageTranslationResult Page = ShadowVirtualToPhysicalUV(VirtualShadowMapId, ShadowMapUV);
+	FVirtualShadowMapSample Result = InitVirtualShadowMapSample();
 	if (Page.bValid)
 	{
-		FVirtualShadowMapSample Result = InitVirtualShadowMapSample();
 		Result.bValid = true;
 		Result.MipLevel = Page.LODOffset;
 		Result.VirtualShadowMapId = VirtualShadowMapId;
@@ -114,10 +114,9 @@ FVirtualShadowMapSample SampleVirtualShadowMap(int VirtualShadowMapId, float2 Sh
 		Result.VirtualTexelAddressFloat = Page.VirtualTexelAddressFloat;
 		Result.PhysicalTexelAddress = Page.PhysicalTexelAddress;
 		Result.Depth = SampleVirtualShadowMapPhysicalDepth(Result.PhysicalTexelAddress);
-		return Result;
 	}
 
-	return InitVirtualShadowMapSample();
+	return Result;
 }
 
 
diff --git a/Engine/Shaders/Private/VolumetricFog.usf b/Engine/Shaders/Private/VolumetricFog.usf
--- a/Engine/Shaders/Private/VolumetricFog.usf
+++ b/Engine/Shaders/Private/VolumetricFog.usf
@@ -173,7 +173,7 @@ void MaterialSetupCS(
 
 		uint LFVCount = LFVTileDataTex[uint3(TilePos, 0)];
 
-		for (int LFVIndex = 0; LFVIndex < LFVCount; ++LFVIndex)
+		for (uint LFVIndex = 0; LFVIndex < LFVCount; ++LFVIndex)
 		{
 			uint InstanceIndex = LFVTileDataTex[uint3(TilePos, 1 + LFVIndex)];
 
