Fix missing preexposure on ReflectionEnvomnemntPixelShader, causing indirect lighting to be too dark.

diff --git a/Engine/Shaders/Private/ReflectionEnvironmentPixelShader.usf b/Engine/Shaders/Private/ReflectionEnvironmentPixelShader.usf
--- a/Engine/Shaders/Private/ReflectionEnvironmentPixelShader.usf
+++ b/Engine/Shaders/Private/ReflectionEnvironmentPixelShader.usf
@@ -494,10 +494,14 @@ void ReflectionEnvironmentSkyLighting(
 		#if REFLECTION_SOURCE_TYPE != REFLECTION_SOURCE_LUMEN_STANDALONE
 		OutColor += ApplySSRContribution(SSRContribution, SSRReflection);
 		#endif
+		OutColor.rgb *= View.PreExposure;
 
 		#if SUBSTRATE_OPAQUE_ROUGH_REFRACTION_ENABLED
 		OutOpaqueRoughRefractionSceneColor += SubstrateLighting.OpaqueRoughRefractionSceneColor;
 		OutSubSurfaceSceneColor += SubstrateLighting.SubSurfaceSceneColor;
+
+		OutOpaqueRoughRefractionSceneColor.rgb *= View.PreExposure;
+		OutSubSurfaceSceneColor.rgb *= View.PreExposure;
 		#endif
 	}
 
