Fix FXC compilation crashes when running with Substrate.

FXC doesn't support forward declaration of struct. This CL shuffle the code to avoid this, and optionally change member function into standalone function. This fixes crashes when compiling/previewing platform using FXC with Substrate enabled.

diff --git a/Engine/Shaders/Private/Substrate/Substrate.ush b/Engine/Shaders/Private/Substrate/Substrate.ush
--- a/Engine/Shaders/Private/Substrate/Substrate.ush
+++ b/Engine/Shaders/Private/Substrate/Substrate.ush
@@ -39,6 +39,8 @@
 #define SUBSTRATE_UNROLL_N(X) UNROLL_N(X)
 #endif
 
+#define SUBSTRATE_COMPILER_SUPPORTS_STRUCT_FORWARD_DECLARATION (!COMPILER_FXC)
+
 // During the base pass or forward rendering, shared local bases are simply available from registers.
 // If a BSDF is loaded for processing from the Substrate buffer, normals are unpacked on demand to reduce VGPR pressure and increase occupancy.
 #ifndef SUBSTRATE_INLINE_SHADING
@@ -134,13 +136,6 @@
 // Disable SSS diffusion based when the pixel's footprint is larger than x times the MFPs
 #define SUBSTRATE_DISABLE_SSS_DIFFUSION_BASED_ON_PIXEL_FOOTPRINT 0
 
-// Forward declaration
-struct FSubstratePixelHeader;			// only needed due to UpdateSingleBSDFOperatorCoverageTransmittance. We canb likely only senc the local bases then onstead ofthe full header....
-struct FSubstrateIntegrationSettings;	// can be moved up
-struct FSubstrateAddressing;			// can be moved up
-struct FSubstrateDBuffer;				// can be moved up
-struct FSubstrateData;					// can be moved up, only depend on Substrate BSDF
-
 ///////////////////////////////////////////////////////////////////////////////
 // BSDFs representation
 
@@ -991,6 +986,26 @@
 float SubstrateComputeHazeRoughness(float Roughness0)				{ float Roughness1=0; float MixLobe=0; SubstrateComputeHazeRoughnessAndWeight(Roughness0, 1.0f,     Roughness1, MixLobe);	return Roughness1; }
 float SubstrateComputeHazeWeight(float Roughness0, float Haziness)	{ float Roughness1=0; float MixLobe=0; SubstrateComputeHazeRoughnessAndWeight(Roughness0, Haziness, Roughness1, MixLobe);	return MixLobe; }
 
+
+///////////////////////////////////////////////////////////////////////////////////////////////////
+
+struct FSubstrateData
+{
+	FSubstrateBSDF InlinedBSDF;// Used for parameter blending or special BSDF such as hair, cloud or water.
+
+	int  OperatorIndex;		// Used as entry point ot the material graph with multiple BSDF and operators.
+
+	#if USE_DEVELOPMENT_SHADERS
+	float3 PreviewColor;	// Used as preview in the material graph within the editor.
+	#endif
+};
+
+FSubstrateData GetInitialisedSubstrateData()
+{
+	FSubstrateData SubstrateData = (FSubstrateData)0;
+	return SubstrateData;
+}
+
 ///////////////////////////////////////////////////////////////////////////////
 // Material & BSDF description
 
@@ -1043,6 +1058,13 @@
 	return SubstrateOp;
 }
 
+///////////////////////////////////////////////////////////////////////////////////////////////////
+
+#if SUBSTRATE_COMPILER_SUPPORTS_STRUCT_FORWARD_DECLARATION
+struct FSubstratePixelHeader;
+struct FSubstrateIntegrationSettings;
+#endif
+
 // Representes a Substrate material as BSDFs organised in a tree repreesnting its topology 
 struct FSubstrateTree
 {
@@ -1052,11 +1074,16 @@
 	FSubstrateBSDF BSDFs[SUBSTRATE_MAX_CLOSURE_COUNT];
 	FSubstrateOperator Operators[SUBSTRATE_MAX_OPERATOR_COUNT];
 
+	// Certain compiler don't support struct forward declaration. This function depends on FSubstratePixelHeader, 
+	// which depends on FSubstrateTree. To overcome circular dependency without forward declaration, a standalone 
+	// version of this function exists as well.
+	#if SUBSTRATE_COMPILER_SUPPORTS_STRUCT_FORWARD_DECLARATION
 	void UpdateSingleBSDFOperatorCoverageTransmittance(
 		FSubstratePixelHeader SubstratePixelHeader,
 		int BSDFIndex,
 		FSubstrateIntegrationSettings Settings,
 		float3 V);
+	#endif
 
 	void UpdateSingleOperatorCoverageTransmittance(
 		int OpIndex);
@@ -1085,10 +1112,14 @@
 
 FSubstrateTree GetInitialisedSubstrateTree()
 {
+#if !COMPILER_FXC
 	FSubstrateTree Out;
 	Out.BSDFCount = 0;
 	Out.OperatorCount = 0;
 	return Out;
+#else
+	return (FSubstrateTree)0;
+#endif
 }
 
 void SubstratePackOutSubstrateTree(
@@ -1271,24 +1302,6 @@
 #undef READ_FLOAT3
 #undef READ_FLOAT4
 
-
-struct FSubstrateData
-{
-	FSubstrateBSDF InlinedBSDF;// Used for parameter blending or special BSDF such as hair, cloud or water.
-
-	int  OperatorIndex;		// Used as entry point ot the material graph with multiple BSDF and operators.
-
-#if USE_DEVELOPMENT_SHADERS
-	float3 PreviewColor;	// Used as preview in the material graph within the editor.
-#endif
-};
-
-FSubstrateData GetInitialisedSubstrateData()
-{
-	FSubstrateData SubstrateData = (FSubstrateData)0;
-	return SubstrateData;
-}
-
 struct FSubstratePixelFootprint
 {
 	float PixelRadiusInWorldSpace;	// In cm
@@ -1933,5 +1946,68 @@
 #endif // SUBSTRATE_MATERIALCONTAINER_IS_VIEWRESOURCE
 
 ///////////////////////////////////////////////////////////////////////////////
+// Substrate decal
+
+// Dbuffer support only a single layer, with minimal set of feature (i.e., no anisotropy, ...);
+// This special structure is for easing data processing
+
+struct FSubstrateDBuffer
+{
+	float3 WorldNormal;
+	float  Roughness;
+	float3 BaseColor;
+	float  Specular;
+	float3 Emissive;
+	float  Metallic;
+	float  Coverage;
+
+	// Per-component coverage 
+	// Since DBuffer properties are blended independently, they can end-up with varying coverage per-property. 
+	// These values are only valid when *reading* DBuffer Value
+	float  OneMinusCoverage;
+	float  OneMinusCoverage_BaseColor;
+	float  OneMinusCoverage_WorldNormal;
+	float  OneMinusCoverage_Roughness;
+};
+
+void SubstratePackDBuffer(
+	in FSubstrateDBuffer In,
+	out float4 Out1,
+	out float4 Out2,
+	out float4 Out3)
+{
+	// DBuffer packing needs to be blendable
+	Out1 = float4(In.BaseColor, In.Coverage);
+	Out2 = float4(In.WorldNormal * 0.5f + 128.0f / 255.0f, In.Coverage);
+	Out3 = float4(In.Metallic, In.Specular, In.Roughness, In.Coverage);
+}
+
+FSubstrateDBuffer SubstrateUnpackDBuffer(
+	in float4 In1,
+	in float4 In2,
+	in float4 In3)
+{
+	FSubstrateDBuffer Out = (FSubstrateDBuffer)0;
+
+	// Per-component coverage
+	Out.OneMinusCoverage			 = (In1.w + In2.w + In3.w) / 3.0f;
+	Out.OneMinusCoverage_BaseColor   = In1.w;
+	Out.OneMinusCoverage_WorldNormal = In2.w;
+	Out.OneMinusCoverage_Roughness   = In3.w;
+
+	// Unpack data
+	Out.BaseColor   = In1.xyz;
+	Out.Coverage    = 1.f - Out.OneMinusCoverage;
+	Out.WorldNormal = In2.rgb * 2 - (256.0 / 255.0);
+	Out.Metallic    = In3.x;
+	Out.Specular    = In3.y;
+	Out.Roughness   = In3.z;
+	Out.Emissive    = 0.f;
+
+
+	return Out;
+}
+
+///////////////////////////////////////////////////////////////////////////////
 // Integration settings
 
@@ -2132,6 +2208,84 @@
 	return Out;
 }
 
+FSubstrateDBuffer FSubstratePixelHeader::SubstrateConvertToDBuffer(
+	in FSubstrateBSDF InBSDF)
+{
+	FSubstrateDBuffer Out = (FSubstrateDBuffer)0;
+
+	#if SUBSTRATE_INLINE_SHADING
+
+	const float3 DiffuseColor = SLAB_DIFFUSEALBEDO(InBSDF);
+	const float3 F0 = SLAB_F0(InBSDF);
+
+	Out.Roughness = SLAB_ROUGHNESS(InBSDF);
+	Out.Metallic = F0RGBToMetallic(F0);
+	Out.Specular = F0RGBToDielectricSpecular(F0);
+	Out.BaseColor = lerp(DiffuseColor, F0, Out.Metallic);
+	Out.Emissive = BSDF_GETEMISSIVE(InBSDF);
+	Out.WorldNormal = this.SharedLocalBases.Normals[BSDF_GETSHAREDLOCALBASISID(InBSDF)];
+	Out.Coverage = InBSDF.Coverage;
+
+	#else
+	Out.WorldNormal = float3(0,0,1);
+	Out.Coverage = 1.f;
+	#endif
+	return Out;
+}
+
+void FSubstratePixelHeader::SubstrateConvertFromDBuffer(
+	in FSubstrateDBuffer In,
+	inout FSubstrateData OutData)
+{
+	// Header
+	const uint SharedLocalBasisIndex = 0;
+	uint SharedLocalBasisTypes = 0;
+	this.ClosureCount = 1;
+	#if SUBSTRATE_INLINE_SHADING
+	this.SharedLocalBases = SubstrateInitialiseSharedLocalBases();
+	this.SharedLocalBases.Normals[SharedLocalBasisIndex] = In.WorldNormal;
+	this.SharedLocalBases.Tangents[SharedLocalBasisIndex] = 0.f;
+	this.IrradianceAO = InitIrradianceAndOcclusion();
+	#endif
+	this.State = 0;
+	this.SetMaterialMode(HEADER_MATERIALMODE_SLAB_SIMPLE);
+
+	// Decal only support a single Slab BSDF with no-anisotropy
+	const FSubstratePixelFootprint PixelFootprint = (FSubstratePixelFootprint)0;
+	const float DefaultThicknessCm = SUBSTRATE_LAYER_DEFAULT_THICKNESS_CM;
+	OutData = GetSubstrateSlabBSDF(
+		PixelFootprint,		// PixelFootprint
+		ComputeDiffuseAlbedo(In.BaseColor, In.Metallic),
+		ComputeF0(In.Specular, In.BaseColor, In.Metallic),
+		1.0f,				// F90
+		In.Roughness,		// Roughness
+		0.f,				// Anisotropy
+		0.f,				// SSSProfileID
+		0.f,				// SSSMFP
+		0.f,				// SSSMFPScale
+		0.f,				// SSSPhaseAniso
+		0.f,				// UseSSSDiffusion
+		0.f,				// Emissive
+		0.f,				// SecondRoughness
+		0.f,				// SecondRoughnessWeight
+		0.f,				// SecondRoughnessAsSimpleClearCoat
+		0.f,				// FuzzAmount
+		0.f,				// FuzzColor
+		In.Roughness,		// FuzzRoughness
+		0.f,				// GlintValue, 
+		(GLINT_UV_TYPE)0,	// GlintUV,
+		0.f, 				// SpecularProfileId
+		DefaultThicknessCm,
+		false,				// bIsThinSurface,
+		false,				// bIsAtTheBottomOfTopology, this will disable many unneeded features.
+		SharedLocalBasisIndex,
+		SharedLocalBasisTypes);
+	#if SUBSTRATE_INLINE_SHADING
+	OutData.InlinedBSDF.Emissive = In.Emissive;
+	OutData.InlinedBSDF.Coverage = In.Coverage;
+	#endif
+}
+
 float3x3 SubstrateGetBSDFSharedBasis_InlineShading(in FSubstratePixelHeader SubstratePixelHeader, in FSubstrateBSDF BSDF)
 {
 #if SUBSTRATE_INLINE_SHADING
@@ -2881,6 +3035,7 @@
 		#endif
 			UNLIT_TRANSMITTANCE(this) = saturate(UNLIT_TRANSMITTANCE(this));
 			UNLIT_NORMAL(this) = normalize(UNLIT_NORMAL(this));
+			break;
 		}
 	}
 }
@@ -2993,6 +3148,15 @@
 	return saturate(NewInput);
 }
 
+#if !SUBSTRATE_COMPILER_SUPPORTS_STRUCT_FORWARD_DECLARATION
+void UpdateSingleBSDFOperatorCoverageTransmittance(
+	inout FSubstrateTree SubstrateTree,
+	FSubstratePixelHeader SubstratePixelHeader,
+	int BSDFIndex,
+	FSubstrateIntegrationSettings Settings,
+	float3 V);
+#endif
+
 ///////////////////////////////////////////////////////////////////////////////
 // Other tools
 
@@ -3575,145 +3739,5 @@
 }
 
 ///////////////////////////////////////////////////////////////////////////////
-// Substrate decal
-
-// Dbuffer support only a single layer, with minimal set of feature (i.e., no anisotropy, ...);
-// This special structure is for easing data processing
-struct FSubstrateDBuffer
-{
-	float3 WorldNormal;
-	float  Roughness;
-	float3 BaseColor;
-	float  Specular;
-	float3 Emissive;
-	float  Metallic;
-	float  Coverage;
-
-	// Per-component coverage 
-	// Since DBuffer properties are blended independently, they can end-up with varying coverage per-property. 
-	// These values are only valid when *reading* DBuffer Value
-	float  OneMinusCoverage;
-	float  OneMinusCoverage_BaseColor;
-	float  OneMinusCoverage_WorldNormal;
-	float  OneMinusCoverage_Roughness;
-};
-
-void SubstratePackDBuffer(
-	in FSubstrateDBuffer In,
-	out float4 Out1,
-	out float4 Out2,
-	out float4 Out3)
-{
-	// DBuffer packing needs to be blendable
-	Out1 = float4(In.BaseColor, In.Coverage);
-	Out2 = float4(In.WorldNormal * 0.5f + 128.0f / 255.0f, In.Coverage);
-	Out3 = float4(In.Metallic, In.Specular, In.Roughness, In.Coverage);
-}
-
-FSubstrateDBuffer SubstrateUnpackDBuffer(
-	in float4 In1,
-	in float4 In2,
-	in float4 In3)
-{
-	FSubstrateDBuffer Out = (FSubstrateDBuffer)0;
-
-	// Per-component coverage
-	Out.OneMinusCoverage			 = (In1.w + In2.w + In3.w) / 3.0f;
-	Out.OneMinusCoverage_BaseColor   = In1.w;
-	Out.OneMinusCoverage_WorldNormal = In2.w;
-	Out.OneMinusCoverage_Roughness   = In3.w;
-
-	// Unpack data
-	Out.BaseColor   = In1.xyz;
-	Out.Coverage    = 1.f - Out.OneMinusCoverage;
-	Out.WorldNormal = In2.rgb * 2 - (256.0 / 255.0);
-	Out.Metallic    = In3.x;
-	Out.Specular    = In3.y;
-	Out.Roughness   = In3.z;
-	Out.Emissive    = 0.f;
-
-
-	return Out;
-}
-
-FSubstrateDBuffer FSubstratePixelHeader::SubstrateConvertToDBuffer(
-	in FSubstrateBSDF InBSDF)
-{
-	FSubstrateDBuffer Out = (FSubstrateDBuffer)0;
-
-#if SUBSTRATE_INLINE_SHADING
-
-	const float3 DiffuseColor = SLAB_DIFFUSEALBEDO(InBSDF);
-	const float3 F0 = SLAB_F0(InBSDF);
-
-	Out.Roughness = SLAB_ROUGHNESS(InBSDF);
-	Out.Metallic = F0RGBToMetallic(F0);
-	Out.Specular = F0RGBToDielectricSpecular(F0);
-	Out.BaseColor = lerp(DiffuseColor, F0, Out.Metallic);
-	Out.Emissive = BSDF_GETEMISSIVE(InBSDF);
-	Out.WorldNormal = this.SharedLocalBases.Normals[BSDF_GETSHAREDLOCALBASISID(InBSDF)];
-	Out.Coverage = InBSDF.Coverage;
-
-#else
-	Out.WorldNormal = float3(0,0,1);
-	Out.Coverage = 1.f;
-#endif
-	return Out;
-}
-
-void FSubstratePixelHeader::SubstrateConvertFromDBuffer(
-	in FSubstrateDBuffer In,
-	inout FSubstrateData OutData)
-{
-	// Header
-	const uint SharedLocalBasisIndex = 0;
-	uint SharedLocalBasisTypes = 0;
-	this.ClosureCount = 1;
-#if SUBSTRATE_INLINE_SHADING
-	this.SharedLocalBases = SubstrateInitialiseSharedLocalBases();
-	this.SharedLocalBases.Normals[SharedLocalBasisIndex] = In.WorldNormal;
-	this.SharedLocalBases.Tangents[SharedLocalBasisIndex] = 0.f;
-	this.IrradianceAO = InitIrradianceAndOcclusion();
-#endif
-	this.State = 0;
-	this.SetMaterialMode(HEADER_MATERIALMODE_SLAB_SIMPLE);
-
-	// Decal only support a single Slab BSDF with no-anisotropy
-	const FSubstratePixelFootprint PixelFootprint = (FSubstratePixelFootprint)0;
-	const float DefaultThicknessCm = SUBSTRATE_LAYER_DEFAULT_THICKNESS_CM;
-	OutData = GetSubstrateSlabBSDF(
-		PixelFootprint,		// PixelFootprint
-		ComputeDiffuseAlbedo(In.BaseColor, In.Metallic),
-		ComputeF0(In.Specular, In.BaseColor, In.Metallic),
-		1.0f,				// F90
-		In.Roughness,		// Roughness
-		0.f,				// Anisotropy
-		0.f,				// SSSProfileID
-		0.f,				// SSSMFP
-		0.f,				// SSSMFPScale
-		0.f,				// SSSPhaseAniso
-		0.f,				// UseSSSDiffusion
-		0.f,				// Emissive
-		0.f,				// SecondRoughness
-		0.f,				// SecondRoughnessWeight
-		0.f,				// SecondRoughnessAsSimpleClearCoat
-		0.f,				// FuzzAmount
-		0.f,				// FuzzColor
-		In.Roughness,		// FuzzRoughness
-		0.f,				// GlintValue, 
-		(GLINT_UV_TYPE)0,	// GlintUV,
-		0.f, 				// SpecularProfileId
-		DefaultThicknessCm,
-		false,				// bIsThinSurface,
-		false,				// bIsAtTheBottomOfTopology, this will disable many unneeded features.
-		SharedLocalBasisIndex,
-		SharedLocalBasisTypes);
-#if SUBSTRATE_INLINE_SHADING
-	OutData.InlinedBSDF.Emissive = In.Emissive;
-	OutData.InlinedBSDF.Coverage = In.Coverage;
-#endif
-}
-
-///////////////////////////////////////////////////////////////////////////////
 // Substrate deferred base-pass output
 
diff --git a/Engine/Shaders/Private/Substrate/SubstrateTree.ush b/Engine/Shaders/Private/Substrate/SubstrateTree.ush
--- a/Engine/Shaders/Private/Substrate/SubstrateTree.ush
+++ b/Engine/Shaders/Private/Substrate/SubstrateTree.ush
@@ -16,17 +16,30 @@
 ///////////////////////////////////////////////////////////////////////////////
 // Tree processing functions
 
+#if SUBSTRATE_COMPILER_SUPPORTS_STRUCT_FORWARD_DECLARATION
 void FSubstrateTree::UpdateSingleBSDFOperatorCoverageTransmittance(
 	FSubstratePixelHeader SubstratePixelHeader,
 	int BSDFIndex,
 	FSubstrateIntegrationSettings Settings,
 	float3 V)
+#else
+void UpdateSingleBSDFOperatorCoverageTransmittance(
+	inout FSubstrateTree InSubstrateTree,
+	FSubstratePixelHeader SubstratePixelHeader,
+	int BSDFIndex,
+	FSubstrateIntegrationSettings Settings,
+	float3 V)
+#endif
 {
 #if SUBSTRATE_INLINE_SHADING
 
-	
-#define CurrentBSDF this.BSDFs[BSDFIndex]
-#define Op this.Operators[CurrentBSDF.OperatorIndex]
+#if SUBSTRATE_COMPILER_SUPPORTS_STRUCT_FORWARD_DECLARATION
+	#define CurrentBSDF this.BSDFs[BSDFIndex]
+	#define Op this.Operators[CurrentBSDF.OperatorIndex]
+#else
+	#define CurrentBSDF InSubstrateTree.BSDFs[BSDFIndex]
+	#define Op InSubstrateTree.Operators[CurrentBSDF.OperatorIndex]
+#endif
 
 	// this feature is only needed for development/editor - we can compile it out for a shipping build (see r.CompileShadersForDevelopment cvar help)
 #if USE_DEVELOPMENT_SHADERS
diff --git a/Engine/Shaders/Private/Substrate/SubstrateVisualize.usf b/Engine/Shaders/Private/Substrate/SubstrateVisualize.usf
--- a/Engine/Shaders/Private/Substrate/SubstrateVisualize.usf
+++ b/Engine/Shaders/Private/Substrate/SubstrateVisualize.usf
@@ -815,7 +815,7 @@ float2 GetBranchOffset(FSubstrateOperator Op, bool bIsLeft)
 void MaterialDebugSubstrateTreeCS(uint3 DispatchThreadId : SV_DispatchThreadID)
 {
 // UE-194712 - Temporary workaround with Vulkan shader compilation
-#if !COMPILER_VULKAN
+#if !COMPILER_VULKAN && !COMPILER_FXC
 	const uint2 PixelPos = (float2(GetCursorPosition()) * View.ViewResolutionFraction);
 
 	FSubstrateAddressing SubstrateAddressing = GetSubstratePixelDataByteOffset(PixelPos, uint2(View.BufferSizeAndInvSize.xy), Substrate.MaxBytesPerPixel);
@@ -1423,7 +1423,7 @@ void MaterialDebugSubstrateTreePS(
 	OutColor = float4(0.0f, 0.0f, 0.0f, 1.0f);
 
 	// Workaround on Vulkan where this shader does not compile correctly
-#if !VULKAN_PROFILE & !VULKAN_PROFILE_SM5 && !VULKAN_PROFILE_SM6
+#if !VULKAN_PROFILE & !VULKAN_PROFILE_SM5 && !VULKAN_PROFILE_SM6 && !COMPILER_FXC
 
 	FSubstrateAddressing SubstrateAddressing = GetSubstratePixelDataByteOffset((float2(GetCursorPosition()) * View.ViewResolutionFraction), uint2(View.BufferSizeAndInvSize.xy), Substrate.MaxBytesPerPixel);
 	FSubstratePixelHeader Header = UnpackSubstrateHeaderIn(Substrate.MaterialTextureArray, SubstrateAddressing, Substrate.TopLayerTexture);
diff --git a/Engine/Source/Runtime/Core/Private/UObject/FortniteMainBranchObjectVersion.cpp b/Engine/Source/Runtime/Core/Private/UObject/FortniteMainBranchObjectVersion.cpp
--- a/Engine/Source/Runtime/Core/Private/UObject/FortniteMainBranchObjectVersion.cpp
+++ b/Engine/Source/Runtime/Core/Private/UObject/FortniteMainBranchObjectVersion.cpp
@@ -14,6 +14,6 @@ TMap<FGuid, FGuid> FFortniteMainBranchObjectVersion::GetSystemGuids()
 	SystemGuids.Add(DevGuids.Niagara_LatestScriptCompileVersion, FGuid("E486231DA967C64B95D4E73131056FA0"));
 	SystemGuids.Add(DevGuids.SkeletalMeshDerivedDataVersion, FGuid("B9AEAA3EE7AC4FCFB3860E8F3E06DD75"));
 	SystemGuids.Add(DevGuids.STATICMESH_DERIVEDDATA_VER, FGuid("A3C8F19042174D05A6AEB6E5D9A58E44"));
-	SystemGuids.Add(DevGuids.MaterialTranslationDDCVersion, FGuid("F43DE2FD43D353A9ABE260AD4CE9F3D4"));
+	SystemGuids.Add(DevGuids.MaterialTranslationDDCVersion, FGuid("BE7C92C6DC7D430B822D815C9D75906D"));
 	return SystemGuids;
 }
diff --git a/Engine/Source/Runtime/Engine/Private/Materials/HLSLMaterialTranslator.cpp b/Engine/Source/Runtime/Engine/Private/Materials/HLSLMaterialTranslator.cpp
--- a/Engine/Source/Runtime/Engine/Private/Materials/HLSLMaterialTranslator.cpp
+++ b/Engine/Source/Runtime/Engine/Private/Materials/HLSLMaterialTranslator.cpp
@@ -2342,7 +2342,11 @@ void FHLSLMaterialTranslator::TranslateMaterial()
 					ResourcesString += "{\n";
 					for (uint32 ClosureIndex = 0; ClosureIndex < SubstrateCtx.SubstrateMaterialEffectiveClosureCount; ++ClosureIndex)
 					{
+						ResourcesString += FString::Printf(TEXT("\t #if SUBSTRATE_COMPILER_SUPPORTS_STRUCT_FORWARD_DECLARATION\n"), ClosureIndex);
 						ResourcesString += FString::Printf(TEXT("\t SubstrateTree.UpdateSingleBSDFOperatorCoverageTransmittance(this, %d, Settings, V);\n"), ClosureIndex);
+						ResourcesString += FString::Printf(TEXT("\t #else\n"), ClosureIndex);
+						ResourcesString += FString::Printf(TEXT("\t UpdateSingleBSDFOperatorCoverageTransmittance(SubstrateTree, this, %d, Settings, V);\n"), ClosureIndex);
+						ResourcesString += FString::Printf(TEXT("\t #endif\n"), ClosureIndex);
 					}
 					ResourcesString += "}\n";
 				}
diff --git a/Engine/Source/Runtime/Renderer/Private/IndirectLightRendering.cpp b/Engine/Source/Runtime/Renderer/Private/IndirectLightRendering.cpp
--- a/Engine/Source/Runtime/Renderer/Private/IndirectLightRendering.cpp
+++ b/Engine/Source/Runtime/Renderer/Private/IndirectLightRendering.cpp
@@ -391,6 +391,12 @@ class FReflectionEnvironmentSkyLightingPS : public FGlobalShader
 		OutEnvironment.SetDefine(TEXT("SUBSTRATE_GLINTS_IS"), 0);
 		OutEnvironment.CompilerFlags.Add(CFLAG_StandardOptimization);
 		FForwardLightingParameters::ModifyCompilationEnvironment(Parameters.Platform, OutEnvironment);
+
+		// FXC is too slow at compiling Substrate lighting shaders
+		if (Substrate::IsSubstrateEnabled() && FDataDrivenShaderPlatformInfo::GetSupportsDxc(Parameters.Platform))
+		{
+			OutEnvironment.CompilerFlags.Add(CFLAG_ForceDXC);
+		}
 	}
 
 	BEGIN_SHADER_PARAMETER_STRUCT(FParameters, )
diff --git a/Engine/Source/Runtime/Renderer/Private/LightRendering.cpp b/Engine/Source/Runtime/Renderer/Private/LightRendering.cpp
--- a/Engine/Source/Runtime/Renderer/Private/LightRendering.cpp
+++ b/Engine/Source/Runtime/Renderer/Private/LightRendering.cpp
@@ -1010,6 +1010,12 @@ class FDeferredLightPS : public FGlobalShader
 		}
 
 		OutEnvironment.SetDefine(TEXT("USE_ADAPTIVE_VOLUMETRIC_SHADOW_MAP"), ShouldHeterogeneousVolumesCastShadows());
+
+		// FXC is too slow at compiling Substrate lighting shaders
+		if (Substrate::IsSubstrateEnabled() && FDataDrivenShaderPlatformInfo::GetSupportsDxc(Parameters.Platform))
+		{
+			OutEnvironment.CompilerFlags.Add(CFLAG_ForceDXC);
+		}
 	}
 
 	void SetParameters(FRHIBatchedShaderParameters& BatchedParameters, const FDeferredLightUniformStruct& DeferredLightUniformsValue)
