Fix cluster shading with substrate when viewport is resized.

The viewport UV was incorrectly computed causing the shadow and the lighting to shift.

diff --git a/Engine/Shaders/Private/ClusteredDeferredShadingPixelShader.usf b/Engine/Shaders/Private/ClusteredDeferredShadingPixelShader.usf
--- a/Engine/Shaders/Private/ClusteredDeferredShadingPixelShader.usf
+++ b/Engine/Shaders/Private/ClusteredDeferredShadingPixelShader.usf
@@ -454,7 +454,9 @@ void ClusteredShadingPixelShader(
 	BRANCH
 	if (SubstratePixelHeader.ClosureCount > 0 || CulledLightGridData.NumLights == 0)
 	{
-		const float2 ScreenPosition = ViewportUVToScreenPos(ScreenUV);
+		// Subtrate tile are expressed as buffer dimension rather than view dimension
+		const float2 ViewportScreenUV = ScreenUV * View.BufferSizeAndInvSize.xy * View.ViewSizeAndInvSize.zw;
+		const float2 ScreenPosition = ViewportUVToScreenPos(ViewportScreenUV);
 		const float3 TranslatedWorldPosition = mul(float4(GetScreenPositionForProjectionType(ScreenPosition, SceneDepth), SceneDepth, 1), PrimaryView.ScreenToTranslatedWorld).xyz;
 		const float3 CameraVector = normalize(TranslatedWorldPosition - PrimaryView.TranslatedWorldCameraOrigin);
 		uint FirstNonSimpleLightIndex = 0;
diff --git a/Engine/Source/Runtime/Renderer/Private/ClusteredDeferredShadingPass.cpp b/Engine/Source/Runtime/Renderer/Private/ClusteredDeferredShadingPass.cpp
--- a/Engine/Source/Runtime/Renderer/Private/ClusteredDeferredShadingPass.cpp
+++ b/Engine/Source/Runtime/Renderer/Private/ClusteredDeferredShadingPass.cpp
@@ -277,6 +277,7 @@ static void InternalAddClusteredDeferredShadingPass(
 		}
 		else if (bSubstrate)
 		{
+			InRHICmdList.SetViewport(View.ViewRect.Min.X, View.ViewRect.Min.Y, 0.0f, View.ViewRect.Max.X, View.ViewRect.Max.Y, 1.0f);
 			SetShaderParameters(InRHICmdList, TileVertexShader, TileVertexShader.GetVertexShader(), PassParameters->SubstrateTile);
 			InRHICmdList.DrawPrimitiveIndirect(PassParameters->SubstrateTile.TileIndirectBuffer->GetIndirectRHICallBuffer(), Substrate::TileTypeDrawIndirectArgOffset(TileType));
 		}
