Apply fog later, when upsampling, to volumetric cloud to reduce reconstruction artefacts.
It cannot be done in the reconstriction phase because the reprojection still ahs to happen and that will still cause artefact.

diff --git a/Engine/Shaders/Private/VolumetricRenderTarget.usf b/Engine/Shaders/Private/VolumetricRenderTarget.usf
--- a/Engine/Shaders/Private/VolumetricRenderTarget.usf
+++ b/Engine/Shaders/Private/VolumetricRenderTarget.usf
@@ -244,7 +244,7 @@ void ReconstructVolumetricRenderTargetPS(
 
 
 
-		#if CLOUD_MIN_AND_MAX_DEPTH
+		#if PERMUTATION_CLOUD_MIN_AND_MAX_DEPTH
 
 
 
@@ -324,7 +324,7 @@ void ReconstructVolumetricRenderTargetPS(
 
 
 
-		#else // CLOUD_MIN_AND_MAX_DEPTH
+		#else // PERMUTATION_CLOUD_MIN_AND_MAX_DEPTH
 
 
 
@@ -411,7 +411,7 @@ void ReconstructVolumetricRenderTargetPS(
 
 
 
-		#endif // CLOUD_MIN_AND_MAX_DEPTH
+		#endif // PERMUTATION_CLOUD_MIN_AND_MAX_DEPTH
 
 
 
@@ -453,6 +453,16 @@ void ReconstructVolumetricRenderTargetPS(
 
 #include "SceneTexturesCommon.ush"
 
+#if PERMUTATION_APPLY_FOG
+#include "HeightFogCommon.ush"
+#endif // PERMUTATION_APPLY_FOG
+
+#if PERMUTATION_APPLY_LOCAL_FOG_VOLUME
+#include "LocalFogVolumes/LocalFogVolumeCommon.ush"
+#endif // PERMUTATION_SUPPORT_LOCAL_FOG_VOLUME
+
+#define APPLY_FOG_LATE (PERMUTATION_APPLY_FOG || PERMUTATION_APPLY_LOCAL_FOG_VOLUME)
+
 SamplerState LinearTextureSampler;
 
 Texture2D<float4> VolumetricTexture;
@@ -525,6 +535,81 @@ float4 VolumetricTextureSizeAndInvSize;
 float UvOffsetSampleAcceptanceWeight;
 float2 FullResolutionToVolumetricBufferResolutionScale;
 
+float4 ApplyFogToCloudSampleFromVPosAndDepth(float2 SVPos, float DepthKilometer, float4 CloudRGBT)
+{
+#if APPLY_FOG_LATE
+
+	float4 HeightFogInscatteringAndTransmittance = float4(0, 0, 0, 1);
+	const float2 SVPosition = SVPos;
+	const float DistanceToPoint = DepthKilometer * KILOMETER_TO_CENTIMETER;
+
+	//float3 WorldPosition = DFHackToFloat(SvPositionToWorld(float4(SVPosition, 0.5, 1.0)));
+	//const float3 FogSampleWorldPositionRelativeToCameraCm = normalize(WorldPosition - DFHackToFloat(PrimaryView.WorldCameraOrigin)) * DistanceToPoint;// +DFHackToFloat(PrimaryView.WorldCameraOrigin);
+	const float3 FogSampleWorldPositionRelativeToCameraCm = normalize(SvPositionToTranslatedWorld(float4(SVPosition, 0.5f, 1.0f))) * DistanceToPoint;
+
+#if PERMUTATION_APPLY_FOG
+	{
+		HeightFogInscatteringAndTransmittance = CalculateHeightFog(FogSampleWorldPositionRelativeToCameraCm);
+	}
+#endif
+
+#if PERMUTATION_APPLY_LOCAL_FOG_VOLUME
+	uint2 TilePos = SVPosition / LFVTilePixelSize.xx;
+	float4 LFVContribution = GetLFVContribution(PrimaryView, TilePos, FogSampleWorldPositionRelativeToCameraCm);
+	if (LFVRenderInVolumetricFog > 0)
+	{
+		HeightFogInscatteringAndTransmittance = float4(LFVContribution.rgb + HeightFogInscatteringAndTransmittance.rgb * LFVContribution.a, LFVContribution.a * HeightFogInscatteringAndTransmittance.a);
+	}
+#endif
+
+#if PERMUTATION_APPLY_FOG
+	if (FogStruct.ApplyVolumetricFog > 0)
+	{
+		float4 ClipPos = mul(float4(FogSampleWorldPositionRelativeToCameraCm, 1.0f), PrimaryView.TranslatedWorldToClip);
+		float3 VolumeUV = ComputeVolumeUVFromNDC(ClipPos);
+		const uint EyeIndex = 0;
+		HeightFogInscatteringAndTransmittance = CombineVolumetricFog(HeightFogInscatteringAndTransmittance, VolumeUV, EyeIndex, DistanceToPoint);
+	}
+#endif
+
+#if PERMUTATION_APPLY_LOCAL_FOG_VOLUME
+	if (LFVRenderInVolumetricFog == 0)
+	{
+		HeightFogInscatteringAndTransmittance = float4(LFVContribution.rgb + HeightFogInscatteringAndTransmittance.rgb * LFVContribution.a, LFVContribution.a * HeightFogInscatteringAndTransmittance.a);
+	}
+#endif
+
+	HeightFogInscatteringAndTransmittance.rgb *= PrimaryView.PreExposure;
+
+	CloudRGBT.rgb = CloudRGBT.rgb * HeightFogInscatteringAndTransmittance.a + (1.0 - CloudRGBT.a) * HeightFogInscatteringAndTransmittance.rgb;
+
+#endif // APPLY_FOG_LATE
+
+	return CloudRGBT;
+}
+
+float4 ApplyFogToCloudFromUVsAndDepth(float2 UVWithOffset, float DepthKilometer, float4 CloudRGBT)
+{
+#if APPLY_FOG_LATE
+	float2 SVPosWithOffset = (UVWithOffset / (FullResolutionToVolumetricBufferResolutionScale.x *View.BufferSizeAndInvSize.xy * VolumetricTextureSizeAndInvSize.zw)) * View.BufferSizeAndInvSize.xy + View.ViewRectMin.xy;
+	CloudRGBT = ApplyFogToCloudSampleFromVPosAndDepth(SVPosWithOffset, DepthKilometer, CloudRGBT);
+#endif // APPLY_FOG_LATE
+
+	return CloudRGBT;
+}
+
+float4 ApplyFogToCloudFromUVs(float2 UVWithOffset, float4 CloudRGBT)
+{
+#if APPLY_FOG_LATE
+	float2 SVPosWithOffset = (UVWithOffset / (FullResolutionToVolumetricBufferResolutionScale.x * View.BufferSizeAndInvSize.xy * VolumetricTextureSizeAndInvSize.zw)) * View.BufferSizeAndInvSize.xy + View.ViewRectMin.xy;
+	float VolumeFrontDepthKilometer = SafeSampleVolumetricDepthTexture(UVWithOffset).x;
+	CloudRGBT = ApplyFogToCloudSampleFromVPosAndDepth(SVPosWithOffset, VolumeFrontDepthKilometer, CloudRGBT);
+#endif // APPLY_FOG_LATE
+
+	return CloudRGBT;
+}
+
+
 void ComposeVolumetricRTOverScenePS(
 	in float4 SVPos : SV_POSITION,
 	out float4 OutputRt0 : SV_Target0
@@ -566,17 +651,25 @@ void ComposeVolumetricRTOverScenePS(
 
 #if PERMUTATION_UPSAMPLINGMODE==0
 	// Single bilinear sample
-
 	OutputRt0 = SafeLoadVolumetricTexture(VolumeUV * VolumetricTextureSizeAndInvSize.xy); // SafeSampleVolumetricTexture(VolumeUV);
+	float2 VolumeFrontDepth0 = SafeSampleVolumetricDepthTexture(VolumeUV);
+	OutputRt0 = ApplyFogToCloudSampleFromVPosAndDepth(CurResPixelCoord, VolumeFrontDepth0.x, OutputRt0);
 	return;
 
 #elif PERMUTATION_UPSAMPLINGMODE==1
 	// Jitter the source sample to add high frequency that can be resolved by TAA - 4 samples
 
 	float4 Data0 = SafeSampleVolumetricTexture(VolumeUVOffset0);
+	Data0 = ApplyFogToCloudFromUVs(VolumeUVOffset0, Data0);
+
 	float4 Data1 = SafeSampleVolumetricTexture(VolumeUVOffset1);
+	Data1 = ApplyFogToCloudFromUVs(VolumeUVOffset1, Data1);
+
 	float4 Data2 = SafeSampleVolumetricTexture(VolumeUVOffset2);
+	Data2 = ApplyFogToCloudFromUVs(VolumeUVOffset2, Data2);
+
 	float4 Data3 = SafeSampleVolumetricTexture(VolumeUVOffset3);
+	Data3 = ApplyFogToCloudFromUVs(VolumeUVOffset3, Data3);
 
 	OutputRt0 = 0.25 * (Data0 + Data1 + Data2 + Data3); 
 	return;
@@ -636,7 +729,7 @@ void ComposeVolumetricRTOverScenePS(
 
 	if (PixelDistanceFromViewKm > VolumeFrontDepth)
 	{
-		OutputRt0 = VolumeRGBT;
+		OutputRt0 = ApplyFogToCloudFromUVsAndDepth(VolumeUV, VolumeFrontDepth.x, VolumeRGBT);
 	}
 	else
 	{
@@ -666,21 +759,22 @@ void ComposeVolumetricRTOverScenePS(
 	float2 VolumeFrontDepth3 = SafeLoadVolumetricDepthTexture(VolumeUVOffset3 * VolumetricTextureSizeAndInvSize.xy);
 #endif
 
+
 	float ValidSampleCount = 0.0f;
 	float4 DataAcc = 0.0f;
 #if 1
-	if (PixelDistanceFromViewKm > VolumeFrontDepth0.x) { DataAcc += VolumeRGBT0; ValidSampleCount += 1.0f; }
-	if (PixelDistanceFromViewKm > VolumeFrontDepth1.x) { DataAcc += VolumeRGBT1; ValidSampleCount += 1.0f; }
-	if (PixelDistanceFromViewKm > VolumeFrontDepth2.x) { DataAcc += VolumeRGBT2; ValidSampleCount += 1.0f; }
-	if (PixelDistanceFromViewKm > VolumeFrontDepth3.x) { DataAcc += VolumeRGBT3; ValidSampleCount += 1.0f; }
+	if (PixelDistanceFromViewKm > VolumeFrontDepth0.x) { VolumeRGBT0 = ApplyFogToCloudFromUVsAndDepth(VolumeUVOffset0, VolumeFrontDepth0.x, VolumeRGBT0); DataAcc += VolumeRGBT0; ValidSampleCount += 1.0f; }
+	if (PixelDistanceFromViewKm > VolumeFrontDepth1.x) { VolumeRGBT1 = ApplyFogToCloudFromUVsAndDepth(VolumeUVOffset1, VolumeFrontDepth1.x, VolumeRGBT1); DataAcc += VolumeRGBT1; ValidSampleCount += 1.0f; }
+	if (PixelDistanceFromViewKm > VolumeFrontDepth2.x) { VolumeRGBT2 = ApplyFogToCloudFromUVsAndDepth(VolumeUVOffset2, VolumeFrontDepth2.x, VolumeRGBT2); DataAcc += VolumeRGBT2; ValidSampleCount += 1.0f; }
+	if (PixelDistanceFromViewKm > VolumeFrontDepth3.x) { VolumeRGBT3 = ApplyFogToCloudFromUVsAndDepth(VolumeUVOffset3, VolumeFrontDepth3.x, VolumeRGBT3); DataAcc += VolumeRGBT3; ValidSampleCount += 1.0f; }
 #else
 	float ClostestDepth = 999999999.0f;
 	float ThisDepth;
 	PixelDistanceFromViewKm = min(PixelDistanceFromViewKm, max(max(VolumeFrontDepth0.y, VolumeFrontDepth1.y), max(VolumeFrontDepth2.y, VolumeFrontDepth3.y))); // clamp to the maximum of the read depth to avoid no depth matching
-	ThisDepth = abs(VolumeFrontDepth0.y - PixelDistanceFromViewKm); if (ThisDepth < ClostestDepth) { DataAcc = VolumeRGBT0; ValidSampleCount = 1.0f; ClostestDepth = ThisDepth; }
-	ThisDepth = abs(VolumeFrontDepth1.y - PixelDistanceFromViewKm); if (ThisDepth < ClostestDepth) { DataAcc = VolumeRGBT1; ValidSampleCount = 1.0f; ClostestDepth = ThisDepth; }
-	ThisDepth = abs(VolumeFrontDepth2.y - PixelDistanceFromViewKm); if (ThisDepth < ClostestDepth) { DataAcc = VolumeRGBT2; ValidSampleCount = 1.0f; ClostestDepth = ThisDepth; }
-	ThisDepth = abs(VolumeFrontDepth3.y - PixelDistanceFromViewKm); if (ThisDepth < ClostestDepth) { DataAcc = VolumeRGBT3; ValidSampleCount = 1.0f; ClostestDepth = ThisDepth; }
+	ThisDepth = abs(VolumeFrontDepth0.y - PixelDistanceFromViewKm); if (ThisDepth < ClostestDepth) { VolumeRGBT0 = ApplyFogToCloudFromUVsAndDepth(VolumeUVOffset0, VolumeFrontDepth0.x, VolumeRGBT0); DataAcc = VolumeRGBT0; ValidSampleCount = 1.0f; ClostestDepth = ThisDepth; }
+	ThisDepth = abs(VolumeFrontDepth1.y - PixelDistanceFromViewKm); if (ThisDepth < ClostestDepth) { VolumeRGBT1 = ApplyFogToCloudFromUVsAndDepth(VolumeUVOffset1, VolumeFrontDepth1.x, VolumeRGBT1); DataAcc = VolumeRGBT1; ValidSampleCount = 1.0f; ClostestDepth = ThisDepth; }
+	ThisDepth = abs(VolumeFrontDepth2.y - PixelDistanceFromViewKm); if (ThisDepth < ClostestDepth) { VolumeRGBT2 = ApplyFogToCloudFromUVsAndDepth(VolumeUVOffset2, VolumeFrontDepth2.x, VolumeRGBT2); DataAcc = VolumeRGBT2; ValidSampleCount = 1.0f; ClostestDepth = ThisDepth; }
+	ThisDepth = abs(VolumeFrontDepth3.y - PixelDistanceFromViewKm); if (ThisDepth < ClostestDepth) { VolumeRGBT3 = ApplyFogToCloudFromUVsAndDepth(VolumeUVOffset3, VolumeFrontDepth3.x, VolumeRGBT3); DataAcc = VolumeRGBT3; ValidSampleCount = 1.0f; ClostestDepth = ThisDepth; }
 #endif
 
 	if (ValidSampleCount > 0.0f)
@@ -725,10 +819,10 @@ void ComposeVolumetricRTOverScenePS(
 		const float CloudFrontDepthTinyOffset = 0.001;
 
 		// We are testing if the depth buffer is further than the cloud front depth and that the cloud front depth is actually in front of traced depth.
-		if (PixelDistanceFromViewKm > (VolumeFrontDepth0.x - CloudFrontDepthTinyOffset) && VolumeFrontDepth0.x < (VolumeFrontDepth0.y)) { DataAcc += VolumeRGBT0; ValidSampleCount += 1.0f; }
-		if (PixelDistanceFromViewKm > (VolumeFrontDepth1.x - CloudFrontDepthTinyOffset) && VolumeFrontDepth1.x < (VolumeFrontDepth1.y)) { DataAcc += VolumeRGBT1; ValidSampleCount += 1.0f; }
-		if (PixelDistanceFromViewKm > (VolumeFrontDepth2.x - CloudFrontDepthTinyOffset) && VolumeFrontDepth2.x < (VolumeFrontDepth2.y)) { DataAcc += VolumeRGBT2; ValidSampleCount += 1.0f; }
-		if (PixelDistanceFromViewKm > (VolumeFrontDepth3.x - CloudFrontDepthTinyOffset) && VolumeFrontDepth3.x < (VolumeFrontDepth3.y)) { DataAcc += VolumeRGBT3; ValidSampleCount += 1.0f; }
+		if (PixelDistanceFromViewKm > (VolumeFrontDepth0.x - CloudFrontDepthTinyOffset) && VolumeFrontDepth0.x < (VolumeFrontDepth0.y)) { VolumeRGBT0 = ApplyFogToCloudFromUVsAndDepth(VolumeUVOffset0, VolumeFrontDepth0.x, VolumeRGBT0); DataAcc += VolumeRGBT0; ValidSampleCount += 1.0f; }
+		if (PixelDistanceFromViewKm > (VolumeFrontDepth1.x - CloudFrontDepthTinyOffset) && VolumeFrontDepth1.x < (VolumeFrontDepth1.y)) { VolumeRGBT1 = ApplyFogToCloudFromUVsAndDepth(VolumeUVOffset1, VolumeFrontDepth1.x, VolumeRGBT1); DataAcc += VolumeRGBT1; ValidSampleCount += 1.0f; }
+		if (PixelDistanceFromViewKm > (VolumeFrontDepth2.x - CloudFrontDepthTinyOffset) && VolumeFrontDepth2.x < (VolumeFrontDepth2.y)) { VolumeRGBT2 = ApplyFogToCloudFromUVsAndDepth(VolumeUVOffset2, VolumeFrontDepth2.x, VolumeRGBT2); DataAcc += VolumeRGBT2; ValidSampleCount += 1.0f; }
+		if (PixelDistanceFromViewKm > (VolumeFrontDepth3.x - CloudFrontDepthTinyOffset) && VolumeFrontDepth3.x < (VolumeFrontDepth3.y)) { VolumeRGBT3 = ApplyFogToCloudFromUVsAndDepth(VolumeUVOffset3, VolumeFrontDepth3.x, VolumeRGBT3); DataAcc += VolumeRGBT3; ValidSampleCount += 1.0f; }
 		if (ValidSampleCount > 0.0f)
 		{
 			OutputRt0 = DataAcc / ValidSampleCount;
@@ -752,10 +846,10 @@ void ComposeVolumetricRTOverScenePS(
 			VolumeFrontDepth3 = SafeLoadVolumetricDepthTexture(VolumeUVOffset3 * VolumetricTextureSizeAndInvSize.xy);
 
 			// We are testing if the depth buffer is further than the cloud front depth and that the cloud front depth is actually in front of traced depth.
-			if (PixelDistanceFromViewKm > (VolumeFrontDepth0.x - CloudFrontDepthTinyOffset) && VolumeFrontDepth0.x < (VolumeFrontDepth0.y)) { DataAcc += VolumeRGBT0; ValidSampleCount += 1.0f; }
-			if (PixelDistanceFromViewKm > (VolumeFrontDepth1.x - CloudFrontDepthTinyOffset) && VolumeFrontDepth1.x < (VolumeFrontDepth1.y)) { DataAcc += VolumeRGBT1; ValidSampleCount += 1.0f; }
-			if (PixelDistanceFromViewKm > (VolumeFrontDepth2.x - CloudFrontDepthTinyOffset) && VolumeFrontDepth2.x < (VolumeFrontDepth2.y)) { DataAcc += VolumeRGBT2; ValidSampleCount += 1.0f; }
-			if (PixelDistanceFromViewKm > (VolumeFrontDepth3.x - CloudFrontDepthTinyOffset) && VolumeFrontDepth3.x < (VolumeFrontDepth3.y)) { DataAcc += VolumeRGBT3; ValidSampleCount += 1.0f; }
+			if (PixelDistanceFromViewKm > (VolumeFrontDepth0.x - CloudFrontDepthTinyOffset) && VolumeFrontDepth0.x < (VolumeFrontDepth0.y)) { VolumeRGBT0 = ApplyFogToCloudFromUVsAndDepth(VolumeUVOffset0, VolumeFrontDepth0.x, VolumeRGBT0); DataAcc += VolumeRGBT0; ValidSampleCount += 1.0f; }
+			if (PixelDistanceFromViewKm > (VolumeFrontDepth1.x - CloudFrontDepthTinyOffset) && VolumeFrontDepth1.x < (VolumeFrontDepth1.y)) { VolumeRGBT1 = ApplyFogToCloudFromUVsAndDepth(VolumeUVOffset1, VolumeFrontDepth1.x, VolumeRGBT1); DataAcc += VolumeRGBT1; ValidSampleCount += 1.0f; }
+			if (PixelDistanceFromViewKm > (VolumeFrontDepth2.x - CloudFrontDepthTinyOffset) && VolumeFrontDepth2.x < (VolumeFrontDepth2.y)) { VolumeRGBT2 = ApplyFogToCloudFromUVsAndDepth(VolumeUVOffset2, VolumeFrontDepth2.x, VolumeRGBT2); DataAcc += VolumeRGBT2; ValidSampleCount += 1.0f; }
+			if (PixelDistanceFromViewKm > (VolumeFrontDepth3.x - CloudFrontDepthTinyOffset) && VolumeFrontDepth3.x < (VolumeFrontDepth3.y)) { VolumeRGBT3 = ApplyFogToCloudFromUVsAndDepth(VolumeUVOffset3, VolumeFrontDepth3.x, VolumeRGBT3); DataAcc += VolumeRGBT3; ValidSampleCount += 1.0f; }
 			if (ValidSampleCount > 0.0f)
 			{
 				OutputRt0 = DataAcc / ValidSampleCount;// +float4(10.0, 0.0, 0.0, 0.0);
@@ -812,6 +906,11 @@ void ComposeVolumetricRTOverScenePS(
 		float4 VolumeRGBT2 = SafeLoadVolumetricTexture(VolumeCoord2);
 		float4 VolumeRGBT3 = SafeLoadVolumetricTexture(VolumeCoord3);
 
+		VolumeRGBT0 = ApplyFogToCloudFromUVsAndDepth(VolumeCoord0 * VolumetricTextureSizeAndInvSize.zw, VolumeFrontDepth0.x, VolumeRGBT0);
+		VolumeRGBT1 = ApplyFogToCloudFromUVsAndDepth(VolumeCoord1 * VolumetricTextureSizeAndInvSize.zw, VolumeFrontDepth1.x, VolumeRGBT1);
+		VolumeRGBT2 = ApplyFogToCloudFromUVsAndDepth(VolumeCoord2 * VolumetricTextureSizeAndInvSize.zw, VolumeFrontDepth2.x, VolumeRGBT2);
+		VolumeRGBT3 = ApplyFogToCloudFromUVsAndDepth(VolumeCoord3 * VolumetricTextureSizeAndInvSize.zw, VolumeFrontDepth3.x, VolumeRGBT3);
+
 		if (bIsWaterPixel) 
 		{
 			// Only do the following is the pixel is water
@@ -867,7 +966,10 @@ void ComposeVolumetricRTOverScenePS(
 	{
 		// Now do a bilinear sample to have a soft look region without depth edges.
 		ValidSampleCount = 1.0;
-		DataAcc = SafeSampleVolumetricTexture((float2(VolumeCoord0) + 0.25 + (float2(OffsetX, OffsetY) * 0.5 + 0.5) * 0.5) * VolumetricTextureSizeAndInvSize.zw);
+		float2 SampleUVs = (float2(VolumeCoord0) + 0.25 + (float2(OffsetX, OffsetY) * 0.5 + 0.5) * 0.5) * VolumetricTextureSizeAndInvSize.zw;
+		DataAcc = SafeSampleVolumetricTexture(SampleUVs);
+		DataAcc = ApplyFogToCloudFromUVs(SampleUVs, DataAcc);
+
 		//DataAcc = float4(0, 1.0, 0, 0);
 	}
 
diff --git a/Engine/Source/Runtime/Renderer/Private/VolumetricCloudRendering.cpp b/Engine/Source/Runtime/Renderer/Private/VolumetricCloudRendering.cpp
--- a/Engine/Source/Runtime/Renderer/Private/VolumetricCloudRendering.cpp
+++ b/Engine/Source/Runtime/Renderer/Private/VolumetricCloudRendering.cpp
@@ -74,6 +74,11 @@ static TAutoConsoleVariable<int32> CVarVolumetricCloudDisableCompute(
 	TEXT("Do not use compute shader for cloud tracing."),
 	ECVF_Scalability);
 
+static TAutoConsoleVariable<int32> CVarVolumetricCloudApplyFogLate(
+	TEXT("r.VolumetricCloud.ApplyFogLate"), 0,
+	TEXT("When true, fog will be applied on cloud later in the frame. Tis helps hiding latency in fog upsate on cloud due to reprojection and reconstruction. It is however more expenssive."),
+	ECVF_RenderThreadSafe);
+
 ////////////////////////////////////////////////////////////////////////// Shadow tracing
 
 static TAutoConsoleVariable<float> CVarVolumetricCloudShadowViewRaySampleMaxCount(
@@ -342,6 +347,17 @@ bool VolumetricCloudWantsSeparatedAtmosphereMieRayLeigh(const FScene* Scene)
 	return false;
 }
 
+bool ShouldVolumetricCloudsApplyFogDuringReconstruction(const FViewInfo& ViewInfo)
+{
+	if (ViewInfo.ViewState)
+	{
+		FVolumetricRenderTargetViewStateData& VRT = ViewInfo.ViewState->VolumetricCloudRenderTarget;
+		int32 VRTMode = VRT.GetMode();
+		return CVarVolumetricCloudApplyFogLate.GetValueOnRenderThread() > 0 && ShouldViewRenderVolumetricCloudRenderTarget(ViewInfo) && (VRTMode == 0 || VRTMode == 2); // Only restricting to a few mode which can exhibit update delay due to reconstruction/reprojection).
+	}
+	return false;
+}
+
 //////////////////////////////////////////////////////////////////////////
 
 static float GetVolumetricCloudShadowmapStrength(const FLightSceneProxy* AtmosphericLight)
@@ -674,6 +690,11 @@ void SetupDefaultRenderVolumetricCloudGlobalParameters(FRDGBuilder& GraphBuilder
 
 	VolumetricCloudParams.EnableHeightFog = ViewInfo.Family->Scene->HasAnyExponentialHeightFog() && ShouldRenderFog(*ViewInfo.Family);
 	SetupFogUniformParameters(GraphBuilder, ViewInfo, VolumetricCloudParams.FogStruct);
+	if (ShouldVolumetricCloudsApplyFogDuringReconstruction(ViewInfo))
+	{
+		VolumetricCloudParams.FogStruct.ApplyVolumetricFog = 0;
+		VolumetricCloudParams.EnableHeightFog = 0;
+	}
 
 	VolumetricCloudParams.LFV = ViewInfo.LocalFogVolumeViewData.UniformParametersStruct;
 
@@ -2242,6 +2263,14 @@ static TRDGUniformBufferRef<FRenderVolumetricCloudGlobalParameters> CreateCloudP
 		VolumetricCloudParams.FogStruct.ApplyVolumetricFog = 0;		// No valid camera froxel volume available.
 		VolumetricCloudParams.OpaqueIntersectionMode = 0;			// No depth buffer is available
 	}
+	else
+	{
+		if (ShouldVolumetricCloudsApplyFogDuringReconstruction(MainView))
+		{
+			VolumetricCloudParams.FogStruct.ApplyVolumetricFog = 0;
+			VolumetricCloudParams.EnableHeightFog = 0;
+		}
+	}
 
 	return GraphBuilder.CreateUniformBuffer(&VolumetricCloudParams);
 }
@@ -2472,7 +2501,7 @@ void FSceneRenderer::RenderVolumetricCloudsInternal(FRDGBuilder& GraphBuilder, F
 		PermutationVector.Set<typename FRenderVolumetricCloudRenderViewCS::FCloudSampleSecondLight>(!bCloudDebugViewModeEnabled && bSecondAtmosphereLightEnabled);
 		PermutationVector.Set<typename FRenderVolumetricCloudRenderViewCS::FCloudSampleLocalLights>(!bCloudDebugViewModeEnabled && bCloudEnableLocalLightSampling && !CloudRC.bIsSkyRealTimeReflectionRendering);
 		PermutationVector.Set<typename FRenderVolumetricCloudRenderViewCS::FCloudMinAndMaxDepth>(ShouldVolumetricCloudTraceWithMinMaxDepth(MainView) && CloudRC.SecondaryCloudTracingDataTexture!=nullptr && !CloudRC.bIsReflectionRendering);// Only for mode 0 in non reflection
-		PermutationVector.Set<typename FRenderVolumetricCloudRenderViewCS::FCloudLocalFogVolume>(MainView.LocalFogVolumeViewData.GPUInstanceCount > 0 && !CloudRC.bIsSkyRealTimeReflectionRendering); // LFVs are not culled for real time capture views.
+		PermutationVector.Set<typename FRenderVolumetricCloudRenderViewCS::FCloudLocalFogVolume>(MainView.LocalFogVolumeViewData.GPUInstanceCount > 0 && !CloudRC.bIsSkyRealTimeReflectionRendering && !ShouldVolumetricCloudsApplyFogDuringReconstruction(MainView)); // LFVs are not culled for real time capture views.
 		PermutationVector.Set<typename FRenderVolumetricCloudRenderViewCS::FCloudDebugViewMode>(bCloudDebugViewModeEnabled);
 
 		TShaderRef<FRenderVolumetricCloudRenderViewCS> ComputeShader = MaterialResource->GetShader<FRenderVolumetricCloudRenderViewCS>(&FLocalVertexFactory::StaticType, PermutationVector, false);
diff --git a/Engine/Source/Runtime/Renderer/Private/VolumetricCloudRendering.h b/Engine/Source/Runtime/Renderer/Private/VolumetricCloudRendering.h
--- a/Engine/Source/Runtime/Renderer/Private/VolumetricCloudRendering.h
+++ b/Engine/Source/Runtime/Renderer/Private/VolumetricCloudRendering.h
@@ -115,6 +115,8 @@ bool ShouldVolumetricCloudTraceWithMinMaxDepth(const FViewInfo& ViewInfo);
 bool ShouldVolumetricCloudTraceWithMinMaxDepth(const TArray<FViewInfo>& Views);
 bool VolumetricCloudWantsSeparatedAtmosphereMieRayLeigh(const FScene* Scene);
 
+bool ShouldVolumetricCloudsApplyFogDuringReconstruction(const FViewInfo& ViewInfo);
+
 // Structure with data necessary to specify a cloud render.
 struct FCloudRenderContext
 {
diff --git a/Engine/Source/Runtime/Renderer/Private/VolumetricRenderTarget.cpp b/Engine/Source/Runtime/Renderer/Private/VolumetricRenderTarget.cpp
--- a/Engine/Source/Runtime/Renderer/Private/VolumetricRenderTarget.cpp
+++ b/Engine/Source/Runtime/Renderer/Private/VolumetricRenderTarget.cpp
@@ -28,7 +28,7 @@ static TAutoConsoleVariable<float> CVarVolumetricRenderTargetUvNoiseSampleAccept
 
 static TAutoConsoleVariable<int32> CVarVolumetricRenderTargetMode(
 	TEXT("r.VolumetricRenderTarget.Mode"), 0,
-	TEXT("[0] trace quarter resolution + reconstruct at half resolution + upsample [1] trace half res + reconstruct full res + upsample [2] trace at quarter resolution + reconstruct full resolution (cannot intersect with opaque meshes and forces UpsamplingMode=2 [3] Cinematic mode with tracing done at full reoslution in render target so that clouds can also be applied on translucent.)"),
+	TEXT("[0] trace quarter resolution + reconstruct at half resolution + upsample [1] trace half res + upsample [2] trace at quarter resolution + reconstruct full resolution (cannot intersect with opaque meshes and forces UpsamplingMode=2 [3] Cinematic mode with tracing done at full reoslution in render target so that clouds can also be applied on translucent.)"),
 	ECVF_RenderThreadSafe | ECVF_Scalability);
 
 static TAutoConsoleVariable<int32> CVarVolumetricRenderTargetUpsamplingMode(
@@ -526,7 +526,7 @@ class FReconstructVolumetricRenderTargetPS : public FGlobalShader
 
 	class FHistoryAvailable : SHADER_PERMUTATION_BOOL("PERMUTATION_HISTORY_AVAILABLE");
 	class FReprojectionBoxConstraint : SHADER_PERMUTATION_BOOL("PERMUTATION_REPROJECTION_BOX_CONSTRAINT");
-	class FCloudMinAndMaxDepth : SHADER_PERMUTATION_BOOL("CLOUD_MIN_AND_MAX_DEPTH");
+	class FCloudMinAndMaxDepth : SHADER_PERMUTATION_BOOL("PERMUTATION_CLOUD_MIN_AND_MAX_DEPTH");
 	using FPermutationDomain = TShaderPermutationDomain<FHistoryAvailable, FReprojectionBoxConstraint, FCloudMinAndMaxDepth>;
 
 	BEGIN_SHADER_PARAMETER_STRUCT(FParameters, )
@@ -678,12 +678,16 @@ class FComposeVolumetricRTOverScenePS : public FGlobalShader
 
 	class FUpsamplingMode : SHADER_PERMUTATION_RANGE_INT("PERMUTATION_UPSAMPLINGMODE", 0, 5);
 	class FRenderUnderWaterBuffer : SHADER_PERMUTATION_BOOL("PERMUTATION_RENDER_UNDERWATER_BUFFER");	// Render into the water scene color buffer (used when rendering from water system)
-	class FRenderCameraComposeWithWater : SHADER_PERMUTATION_BOOL("PERMUTATION_COMPOSE_WITH_WATER");	// When water us used and the camera is under water, use that permutation (to handle camera intersection with water and double cloud composition)
+	class FRenderCameraComposeWithWater : SHADER_PERMUTATION_BOOL("PERMUTATION_COMPOSE_WITH_WATER");	// When water us used and the camera is under water, use that permutation (to handle camera intersection with water and double cloud composition) 
 	class FMSAASampleCount : SHADER_PERMUTATION_SPARSE_INT("MSAA_SAMPLE_COUNT", 1, 2, 4, 8);
-	using FPermutationDomain = TShaderPermutationDomain<FUpsamplingMode, FRenderUnderWaterBuffer, FRenderCameraComposeWithWater, FMSAASampleCount>;
+	class FApplyFog : SHADER_PERMUTATION_BOOL("PERMUTATION_APPLY_FOG");
+	class FApplyLocalFogVolume : SHADER_PERMUTATION_BOOL("PERMUTATION_APPLY_LOCAL_FOG_VOLUME");
+	using FPermutationDomain = TShaderPermutationDomain<FUpsamplingMode, FRenderUnderWaterBuffer, FRenderCameraComposeWithWater, FMSAASampleCount, FApplyFog, FApplyLocalFogVolume>;
 
 	BEGIN_SHADER_PARAMETER_STRUCT(FParameters, )
 		SHADER_PARAMETER_STRUCT_REF(FViewUniformShaderParameters, ViewUniformBuffer)
+		SHADER_PARAMETER_STRUCT(FLocalFogVolumeUniformParameters, LFV)
+		SHADER_PARAMETER_RDG_UNIFORM_BUFFER(FFogUniformParameters, FogStruct)
 		SHADER_PARAMETER_RDG_TEXTURE(Texture2D, VolumetricTexture)
 		SHADER_PARAMETER_RDG_TEXTURE(Texture2D, VolumetricDepthTexture)
 		SHADER_PARAMETER_RDG_TEXTURE(Texture2D, WaterLinearDepthTexture)
@@ -821,11 +825,15 @@ void ComposeVolumetricRenderTargetOverScene(
 		const bool bForwardShading = IsForwardShadingEnabled(ViewInfo.GetShaderPlatform());
 		check(bForwardShading || (!bForwardShading && SceneDepthTexture->Desc.NumSamples==1));
 
+		const bool bShouldVolumetricCloudsApplyFogDuringReconstruction = ShouldVolumetricCloudsApplyFogDuringReconstruction(ViewInfo);
+
 		FComposeVolumetricRTOverScenePS::FPermutationDomain PermutationVector;
 		PermutationVector.Set<FComposeVolumetricRTOverScenePS::FUpsamplingMode>(UpsamplingMode);
 		PermutationVector.Set<FComposeVolumetricRTOverScenePS::FRenderUnderWaterBuffer>(0);
 		PermutationVector.Set<FComposeVolumetricRTOverScenePS::FRenderCameraComposeWithWater>((bComposeWithWater && ViewInfo.IsUnderwater()) ? 1 : 0);
 		PermutationVector.Set<FComposeVolumetricRTOverScenePS::FMSAASampleCount>(SceneDepthTexture->Desc.NumSamples);
+		PermutationVector.Set<FComposeVolumetricRTOverScenePS::FApplyFog>(bShouldVolumetricCloudsApplyFogDuringReconstruction && ViewInfo.Family->Scene->HasAnyExponentialHeightFog() && ShouldRenderFog(*ViewInfo.Family));
+		PermutationVector.Set<FComposeVolumetricRTOverScenePS::FApplyLocalFogVolume>(bShouldVolumetricCloudsApplyFogDuringReconstruction && ViewInfo.LocalFogVolumeViewData.GPUInstanceCount > 0);
 		TShaderMapRef<FComposeVolumetricRTOverScenePS> PixelShader(ViewInfo.ShaderMap, PermutationVector);
 
 		FComposeVolumetricRTOverScenePS::FParameters* PassParameters = GraphBuilder.AllocParameters<FComposeVolumetricRTOverScenePS::FParameters>();
@@ -854,6 +862,9 @@ void ComposeVolumetricRenderTargetOverScene(
 		FVector2D VolumetricTextureSize = FVector2D(float(VolumetricTexture->Desc.GetSize().X), float(VolumetricTexture->Desc.GetSize().Y));
 		PassParameters->VolumetricTextureSizeAndInvSize = FVector4f(VolumetricTextureSize.X, VolumetricTextureSize.Y, 1.0f / VolumetricTextureSize.X, 1.0f / VolumetricTextureSize.Y);
 
+		PassParameters->FogStruct = CreateFogUniformBuffer(GraphBuilder, ViewInfo);
+		PassParameters->LFV = ViewInfo.LocalFogVolumeViewData.UniformParametersStruct;
+
 		FPixelShaderUtils::AddFullscreenPass<FComposeVolumetricRTOverScenePS>(
 			GraphBuilder, ViewInfo.ShaderMap, RDG_EVENT_NAME("VolumetricComposeOverScene"), PixelShader, PassParameters, ViewInfo.ViewRect,
 			PreMultipliedColorTransmittanceBlend);
