Fix for BitStreamReader not working with NumBits=32.
Adds an additional select instruction when CompileTimeMaxBits=32, which currently doesn't happen anywhere in the engine, but could happen in user code.

diff --git a/Engine/Shaders/Private/BitStreamReaderImplementation.ush b/Engine/Shaders/Private/BitStreamReaderImplementation.ush
--- a/Engine/Shaders/Private/BitStreamReaderImplementation.ush
+++ b/Engine/Shaders/Private/BitStreamReaderImplementation.ush
@@ -21,6 +21,7 @@
 #define CONCAT2(A, B) A##B
 #define CONCAT(A, B) CONCAT2(A, B)
 
+//TODO: Rewrite this with templates, so CompileMaxBits can be proper compile-time
 uint CONCAT(BitStreamReader_Read_, TYPE_SUFFIX)
 	(INPUT_BUFFER_TYPE InputBuffer, inout FBitStreamReaderState State, int NumBits, int CompileTimeMaxBits)
 {
@@ -65,7 +66,8 @@ uint CONCAT(BitStreamReader_Read_, TYPE_SUFFIX)
 		State.CompileTimeMinDwordBits = min(32, State.CompileTimeMaxRemainingBits);
 	}
 
-	const uint Result = BitFieldExtractU32(State.BufferBits.x, NumBits, State.BufferOffset); // BufferOffset implicitly &31
+	const bool bNumBits32 = (CompileTimeMaxBits >= 32) && NumBits == 32;	// This will be optimized away unless CompileTimeMaxBits is explicitly set to 32.
+	const uint Result = bNumBits32 ? State.BufferBits.x : BitFieldExtractU32(State.BufferBits.x, NumBits, State.BufferOffset); // BufferOffset implicitly &31
 	
 	State.BufferOffset += NumBits;
 	State.CompileTimeMinBufferBits    -= CompileTimeMaxBits;
