Added project option for Height fog to match volumetric fog.
Emissive of volumetric fog now follow the distribution of the exp height fog.
Exp fog follow the emissive and albedo of the volumetric fog. And it also uses the same PhaseG and phase function as the volumetric fog (the UI for in scattering is then disabled when the project setting is enabled)
Now volumetric fog looks mostly different due to the lact of shadow casting or difference integration of ambient (through lumen for instance).

This allows for lumen fog sampling to even better match the environment. (currently not including volumetric fog material and local fog volume as this would be a lot more expenssive).

diff --git a/Engine/Shaders/Private/HeightFogCommon.ush b/Engine/Shaders/Private/HeightFogCommon.ush
--- a/Engine/Shaders/Private/HeightFogCommon.ush
+++ b/Engine/Shaders/Private/HeightFogCommon.ush
@@ -9,6 +9,7 @@
 #include "Common.ush"
 #include "/Engine/Shared/EnvironmentComponentsFlags.h"
 #include "/Engine/Private/SkyAtmosphereCommon.ush"
+#include "ParticipatingMediaCommon.ush"
 
 #ifndef PERMUTATION_SUPPORT_FOG_INSCATTERING_TEXTURE
 #define PERMUTATION_SUPPORT_FOG_INSCATTERING_TEXTURE 1
@@ -75,6 +76,8 @@ struct FogStructData
 	float ApplyVolumetricFog;
 	float VolumetricFogStartDistance;
 	float VolumetricFogNearFadeInDistanceInv;
+	float3 VolumetricFogAlbedo;
+	float VolumetricFogPhaseG;
 	// Spir-v does not support copying texture types so 
 	// we can't use FogStructData to store them.
 	//TextureCube FogInscatteringColorCubemap;
@@ -102,6 +105,8 @@ FogStructData GetFogStructData(uint EyeIndex)
 		Out.ApplyVolumetricFog = FogStruct.ApplyVolumetricFog;
 		Out.VolumetricFogStartDistance = FogStruct.VolumetricFogStartDistance;
 		Out.VolumetricFogNearFadeInDistanceInv = FogStruct.VolumetricFogNearFadeInDistanceInv;
+		Out.VolumetricFogAlbedo = FogStruct.VolumetricFogAlbedo;
+		Out.VolumetricFogPhaseG = FogStruct.VolumetricFogPhaseG;
 		// Intentionally commented out to avoid spir-v generation issue - copying texture types is not supported
 		//Out.FogInscatteringColorCubemap = FogStruct.FogInscatteringColorCubemap;
 		//Out.IntegratedLightScattering = FogStruct.IntegratedLightScattering;
@@ -122,6 +127,8 @@ FogStructData GetFogStructData(uint EyeIndex)
 		Out.ApplyVolumetricFog = FogStructISR.ApplyVolumetricFog;
 		Out.VolumetricFogStartDistance = FogStructISR.VolumetricFogStartDistance;
 		Out.VolumetricFogNearFadeInDistanceInv = FogStructISR.VolumetricFogNearFadeInDistanceInv;
+		Out.VolumetricFogAlbedo = FogStructISR.VolumetricFogAlbedo;
+		Out.VolumetricFogPhaseG = FogStructISR.VolumetricFogPhaseG;
 		// Intentionally commented out to avoid spir-v generation issue - copying texture types is not supported
 		//Out.FogInscatteringColorCubemap = FogStructISR.FogInscatteringColorCubemap;
 		//Out.IntegratedLightScattering = FogStructISR.IntegratedLightScattering;
@@ -165,7 +172,11 @@ float3 ComputeInscatteringColor(float3 CameraToReceiver, float CameraToReceiverL
 #endif
 
 #if PROJECT_SUPPORT_SKY_ATMOSPHERE_AFFECTS_HEIGHFOG
-	Inscattering += FogData.SkyAtmosphereAmbientContributionColorScale.rgb * View.SkyAtmosphereHeightFogContribution.xxx * GetViewDistanceSkyLightColor();
+	float3 SkyAmbient = FogData.SkyAtmosphereAmbientContributionColorScale.rgb * View.SkyAtmosphereHeightFogContribution.xxx * GetViewDistanceSkyLightColor();
+	#if PROJECT_EXPFOG_MATCHES_VFOG
+	SkyAmbient *= FogData.VolumetricFogAlbedo;
+	#endif
+	Inscattering += SkyAmbient;
 #endif
 
 	return Inscattering;
@@ -300,26 +311,41 @@ half4 GetExponentialHeightFog(float3 WorldPositionRelativeToCamera, float Exclud
 	)
 	{
 #if PROJECT_SUPPORT_SKY_ATMOSPHERE_AFFECTS_HEIGHFOG
-		const float UniformPhaseFunction = 1.0f / (4.0f*PI);
+		const float UniformPhaseFunction = 1.0f / (4.0f * PI);
 		half3 DirectionalInscatteringColor;
 		half3 DirectionalLightInscattering;
 
 		// No need to test View.AtmosphereLightIlluminanceOnGroundPostTransmittance[0].a because InscatteringLightDirection.w above is doing the same test already.
-		DirectionalInscatteringColor = FogData.DirectionalInscatteringColor.xyz + InView.SkyAtmosphereHeightFogContribution * InView.AtmosphereLightIlluminanceOnGroundPostTransmittance[0].rgb * UniformPhaseFunction;
-		DirectionalLightInscattering = DirectionalInscatteringColor * pow(saturate(dot(CameraToReceiverNormalized, InView.AtmosphereLightDirection[0].xyz)), FogData.DirectionalInscatteringColor.w);
+		DirectionalInscatteringColor = FogData.DirectionalInscatteringColor.xyz + InView.SkyAtmosphereHeightFogContribution * InView.AtmosphereLightIlluminanceOnGroundPostTransmittance[0].rgb;
+		DirectionalLightInscattering = DirectionalInscatteringColor 
+	#if PROJECT_EXPFOG_MATCHES_VFOG
+			* HenyeyGreensteinPhase(FogData.VolumetricFogPhaseG, dot(-CameraToReceiverNormalized, InView.AtmosphereLightDirection[0].xyz));
+	#else
+			* pow(saturate(dot(CameraToReceiverNormalized, InView.AtmosphereLightDirection[0].xyz)), FogData.DirectionalInscatteringColor.w) * UniformPhaseFunction;
+	#endif
 
 		if (InView.AtmosphereLightIlluminanceOnGroundPostTransmittance[1].a > 0.0f) // Skip DirectionalInscatteringColor on the second light when disabled.
 		{
-			DirectionalInscatteringColor = FogData.DirectionalInscatteringColor.xyz + InView.SkyAtmosphereHeightFogContribution * InView.AtmosphereLightIlluminanceOnGroundPostTransmittance[1].rgb * UniformPhaseFunction;
-			DirectionalLightInscattering += DirectionalInscatteringColor * pow(saturate(dot(CameraToReceiverNormalized, InView.AtmosphereLightDirection[1].xyz)), FogData.DirectionalInscatteringColor.w);
+			DirectionalInscatteringColor = FogData.DirectionalInscatteringColor.xyz + InView.SkyAtmosphereHeightFogContribution * InView.AtmosphereLightIlluminanceOnGroundPostTransmittance[1].rgb;
+			DirectionalLightInscattering += DirectionalInscatteringColor 
+		#if PROJECT_EXPFOG_MATCHES_VFOG
+				* HenyeyGreensteinPhase(FogData.VolumetricFogPhaseG, dot(-CameraToReceiverNormalized, InView.AtmosphereLightDirection[1].xyz));
+		#else
+				* pow(saturate(dot(CameraToReceiverNormalized, InView.AtmosphereLightDirection[1].xyz)), FogData.DirectionalInscatteringColor.w) * UniformPhaseFunction;
+		#endif
 		}
 #else
 		// Setup a cosine lobe around the light direction to approximate inscattering from the directional light off of the ambient haze;
-		half3 DirectionalLightInscattering = FogData.DirectionalInscatteringColor.xyz * pow(saturate(dot(CameraToReceiverNormalized, FogData.InscatteringLightDirection.xyz)), FogData.DirectionalInscatteringColor.w);
+		half3 DirectionalLightInscattering = FogData.DirectionalInscatteringColor.xyz
+	#if PROJECT_EXPFOG_MATCHES_VFOG
+			* HenyeyGreensteinPhase(FogData.VolumetricFogPhaseG, dot(-CameraToReceiverNormalized, InView.AtmosphereLightDirection.xyz));
+	#else
+			* pow(saturate(dot(CameraToReceiverNormalized, FogData.InscatteringLightDirection.xyz)), FogData.DirectionalInscatteringColor.w) * UniformPhaseFunction;
+	#endif
 #endif
 
 		// Calculate the line integral of the eye ray through the haze, using a special starting distance to limit the inscattering to the distance
-		float DirectionalInscatteringStartDistance = FogData.InscatteringLightDirection.w;
+		float DirectionalInscatteringStartDistance = PROJECT_EXPFOG_MATCHES_VFOG ? 0.0 : FogData.InscatteringLightDirection.w;
 		float DirExponentialHeightLineIntegral = ExponentialHeightLineIntegralShared * max(RayLength - DirectionalInscatteringStartDistance, 0.0f);
 		// Calculate the amount of light that made it through the fog using the transmission equation
 		half DirectionalInscatteringFogFactor = saturate(exp2(-DirExponentialHeightLineIntegral));
@@ -328,6 +354,10 @@ half4 GetExponentialHeightFog(float3 WorldPositionRelativeToCamera, float Exclud
 	}
 #endif
 
+#if PROJECT_EXPFOG_MATCHES_VFOG
+	DirectionalInscattering *= FogData.VolumetricFogAlbedo;
+#endif
+
 	// Calculate the amount of light that made it through the fog using the transmission equation
 	half ExpFogFactor = max(saturate(exp2(-ExponentialHeightLineIntegral)), MinFogOpacity);
 
diff --git a/Engine/Shaders/Private/VolumetricFog.usf b/Engine/Shaders/Private/VolumetricFog.usf
--- a/Engine/Shaders/Private/VolumetricFog.usf
+++ b/Engine/Shaders/Private/VolumetricFog.usf
@@ -153,7 +153,8 @@ void MaterialSetupCS(
 
 	// Exponential height fog interprets density differently, match its behavior
 	float MatchHeightFogFactor = .5f;
-	float Extinction = max(GlobalDensity * GlobalExtinctionScale * MatchHeightFogFactor, 0);
+	GlobalDensity *= MatchHeightFogFactor;
+	float Extinction = max(GlobalDensity * GlobalExtinctionScale, 0);
 
 	float3 Scattering = Albedo * Extinction;
 	float Absorption = max(Extinction - Luminance(Scattering), 0.0f);
@@ -243,7 +244,11 @@ void MaterialSetupCS(
 	{
 		RWVBufferA[GridCoordinate] = float4(Scattering, Absorption);
 #if USE_EMISSIVE
-		RWVBufferB[GridCoordinate] = float4(GlobalEmissive.rgb + AdditionalEmissive, 0);
+		float3 HeightFogEmissive = GlobalEmissive.rgb;
+	#if PROJECT_EXPFOG_MATCHES_VFOG
+		HeightFogEmissive.rgb *= GlobalDensity; // Emissive needs to follow the height fog density for a better match
+	#endif 
+		RWVBufferB[GridCoordinate] = float4(HeightFogEmissive.rgb + AdditionalEmissive, 0);
 #endif
 	}
 }
@@ -792,7 +797,7 @@ void LightScatteringCS(
 		}
 
 		FTwoBandSHVector RotatedHGZonalHarmonic;
-		RotatedHGZonalHarmonic.V = float4(1.0f, CameraVector.y, CameraVector.z, CameraVector.x) * float4(1.0f, PhaseG, PhaseG, PhaseG);
+		RotatedHGZonalHarmonic.V = float4(1.0f, CameraVector.y, CameraVector.z, CameraVector.x) * float4(1.0f, PhaseG, PhaseG, PhaseG); // Note: I believe PhaseG here hsould be negated to match the sky & environment.
 
 		float3 BrickTextureUVs = 0;
 
@@ -829,6 +834,7 @@ void LightScatteringCS(
 			}
 			else
 			{
+				// NOTE it should be CameraVector * PhaseG. But we keep it like that for backward compatibility for now. (visible when r.Lumen.TranslucencyVolume.Enable 0)
 				SkyLighting = SampleSkyLightDiffuseEnvMap > 0 ? View.SkyLightColor.rgb * GetSkySHDiffuseSimple(CameraVector * -PhaseG) : 0.0f.xxx;
 			}
 
diff --git a/Engine/Source/Runtime/Engine/Classes/Components/ExponentialHeightFogComponent.h b/Engine/Source/Runtime/Engine/Classes/Components/ExponentialHeightFogComponent.h
--- a/Engine/Source/Runtime/Engine/Classes/Components/ExponentialHeightFogComponent.h
+++ b/Engine/Source/Runtime/Engine/Classes/Components/ExponentialHeightFogComponent.h
@@ -66,6 +66,9 @@ class UExponentialHeightFogComponent : public USceneComponent
 	UPROPERTY()
 	FLinearColor FogInscatteringColor_DEPRECATED;
 
+	/**
+	 * Note: when r.SupportExpFogMatchesVolumetricFog = 1, this value is ignored and the volumetric fog Emissive is used instead.
+	 */
 	UPROPERTY(BlueprintReadOnly, interp, Category=ExponentialHeightFogComponent, meta = (DisplayName = "Fog Inscattering Color"))
 	FLinearColor FogInscatteringLuminance;
 
@@ -98,14 +101,18 @@ class UExponentialHeightFogComponent : public USceneComponent
 
 	/** 
 	 * Controls the size of the directional inscattering cone, which is used to approximate inscattering from a directional light.  
-	 * Note: there must be a directional light with bUsedAsAtmosphereSunLight enabled for DirectionalInscattering to be used.
+	 * Note: 
+	 *   - there must be a directional light with bUsedAsAtmosphereSunLight enabled for DirectionalInscattering to be used.
+	 *   - When r.SupportExpFogMatchesVolumetricFog = 1, this value is ignored and the volumetric fog Scattering Distribution is used instead.
 	 */
 	UPROPERTY(BlueprintReadOnly, interp, Category=DirectionalInscattering, meta=(UIMin = "2", UIMax = "64"))
 	float DirectionalInscatteringExponent;
 
 	/** 
 	 * Controls the start distance from the viewer of the directional inscattering, which is used to approximate inscattering from a directional light. 
-	 * Note: there must be a directional light with bUsedAsAtmosphereSunLight enabled for DirectionalInscattering to be used.
+	 * Note: 
+	 *   - There must be a directional light with bUsedAsAtmosphereSunLight enabled for DirectionalInscattering to be used.
+	 *   - When r.SupportExpFogMatchesVolumetricFog = 1, this value is ignored.
 	 */
 	UPROPERTY(BlueprintReadOnly, interp, Category=DirectionalInscattering)
 	float DirectionalInscatteringStartDistance;
@@ -115,7 +122,9 @@ class UExponentialHeightFogComponent : public USceneComponent
 
 	/** 
 	 * Controls the color of the directional inscattering, which is used to approximate inscattering from a directional light. 
-	 * Note: there must be a directional light with bUsedAsAtmosphereSunLight enabled for DirectionalInscattering to be used.
+	 * Note:
+	 *   - there must be a directional light with bUsedAsAtmosphereSunLight enabled for DirectionalInscattering to be used.
+	 *   - When r.SupportExpFogMatchesVolumetricFog = 1, this value is ignored.
 	 */
 	UPROPERTY(BlueprintReadOnly, interp, Category=DirectionalInscattering, meta = (DisplayName = "Directional Inscattering Color"))
 	FLinearColor DirectionalInscatteringLuminance;
diff --git a/Engine/Source/Runtime/Engine/Classes/Engine/RendererSettings.h b/Engine/Source/Runtime/Engine/Classes/Engine/RendererSettings.h
--- a/Engine/Source/Runtime/Engine/Classes/Engine/RendererSettings.h
+++ b/Engine/Source/Runtime/Engine/Classes/Engine/RendererSettings.h
@@ -935,6 +935,15 @@ class URendererSettings : public UDeveloperSettings
 		ConfigRestartRequired = true))
 		uint32 bSupportSkyAtmosphereAffectsHeightFog : 1;
 
+	/**
+	"When enabled, the height fog scattering lob will match the volumetric fog phase function."
+	*/
+	UPROPERTY(config, EditAnywhere, Category = Optimizations, meta = (
+		ConsoleVariable = "r.SupportExpFogMatchesVolumetricFog", DisplayName = "Support Exponential Fog Matches Volumetric Fog",
+		ToolTip = "When enabled, the height fog scattering/ambient/emissive/phase will match the volumetric fog better.",
+		ConfigRestartRequired = true))
+		uint32 bSupportExpFogMatchesVolumetricFog : 1;
+
 	/**
 	"Local fog volume components can will need to be applied on translucent, and opaque in forward, so resources will need to be bound to apply aerial perspective on transparent surfaces (and all surfaces on mobile via per vertex evaluation)."
 	"It requires r.SupportLocalFogVolumes to be true."
diff --git a/Engine/Source/Runtime/Engine/Private/Components/HeightFogComponent.cpp b/Engine/Source/Runtime/Engine/Private/Components/HeightFogComponent.cpp
--- a/Engine/Source/Runtime/Engine/Private/Components/HeightFogComponent.cpp
+++ b/Engine/Source/Runtime/Engine/Private/Components/HeightFogComponent.cpp
@@ -93,6 +93,18 @@ bool UExponentialHeightFogComponent::CanEditChange(const FProperty* InProperty)
 	{
 		FString PropertyName = InProperty->GetName();
 
+		if (DoesProjectSupportExpFogMatchesVolumetricFog())
+		{
+			if (PropertyName == GET_MEMBER_NAME_STRING_CHECKED(UExponentialHeightFogComponent, DirectionalInscatteringExponent) ||
+				PropertyName == GET_MEMBER_NAME_STRING_CHECKED(UExponentialHeightFogComponent, DirectionalInscatteringStartDistance) ||
+				PropertyName == GET_MEMBER_NAME_STRING_CHECKED(UExponentialHeightFogComponent, FogInscatteringLuminance) ||
+				PropertyName == GET_MEMBER_NAME_STRING_CHECKED(UExponentialHeightFogComponent, DirectionalInscatteringLuminance))
+			{
+				// In this case, all the data will come from the volumetric fog.
+				return false;
+			}
+		}
+
 		if (PropertyName == GET_MEMBER_NAME_STRING_CHECKED(UExponentialHeightFogComponent, DirectionalInscatteringExponent) ||
 			PropertyName == GET_MEMBER_NAME_STRING_CHECKED(UExponentialHeightFogComponent, DirectionalInscatteringStartDistance) ||
 			PropertyName == GET_MEMBER_NAME_STRING_CHECKED(UExponentialHeightFogComponent, DirectionalInscatteringLuminance) ||
diff --git a/Engine/Source/Runtime/Engine/Private/ShaderCompiler/ShaderCompiler.cpp b/Engine/Source/Runtime/Engine/Private/ShaderCompiler/ShaderCompiler.cpp
--- a/Engine/Source/Runtime/Engine/Private/ShaderCompiler/ShaderCompiler.cpp
+++ b/Engine/Source/Runtime/Engine/Private/ShaderCompiler/ShaderCompiler.cpp
@@ -8428,6 +8428,10 @@ void GlobalBeginCompileShader(
 		SET_SHADER_DEFINE(Input.Environment, PROJECT_SUPPORT_SKY_ATMOSPHERE_AFFECTS_HEIGHFOG, (CVar && bSupportSkyAtmosphere) ? (CVar->GetInt() != 0) : 0);
 	}
 
+	{
+		SET_SHADER_DEFINE(Input.Environment, PROJECT_EXPFOG_MATCHES_VFOG, DoesProjectSupportExpFogMatchesVolumetricFog() ? 1 : 0);
+	}
+
 	{
 		static IConsoleVariable* CVar = IConsoleManager::Get().FindConsoleVariable(TEXT("r.PostProcessing.PropagateAlpha"));
 		int32 PropagateAlpha = CVar->GetInt();
diff --git a/Engine/Source/Runtime/Engine/Private/ShaderCompiler/ShaderGenerationUtil.cpp b/Engine/Source/Runtime/Engine/Private/ShaderCompiler/ShaderGenerationUtil.cpp
--- a/Engine/Source/Runtime/Engine/Private/ShaderCompiler/ShaderGenerationUtil.cpp
+++ b/Engine/Source/Runtime/Engine/Private/ShaderCompiler/ShaderGenerationUtil.cpp
@@ -69,6 +69,7 @@ void ApplyFetchEnvironmentInternal(FShaderGlobalDefines& SrcDefines, const Conta
 	FETCH_COMPILE_BOOL(EARLY_Z_PASS_ONLY_MATERIAL_MASKING);
 	FETCH_COMPILE_BOOL(PROJECT_SUPPORT_SKY_ATMOSPHERE);
 	FETCH_COMPILE_BOOL(PROJECT_SUPPORT_SKY_ATMOSPHERE_AFFECTS_HEIGHFOG);
+	FETCH_COMPILE_BOOL(PROJECT_EXPFOG_MATCHES_VFOG);
 	FETCH_COMPILE_BOOL(SUPPORT_CLOUD_SHADOW_ON_FORWARD_LIT_TRANSLUCENT);
 	FETCH_COMPILE_BOOL(SUPPORT_CLOUD_SHADOW_ON_SINGLE_LAYER_WATER);
 	FETCH_COMPILE_BOOL(POST_PROCESS_ALPHA);
@@ -416,6 +417,10 @@ static FShaderGlobalDefines FetchShaderGlobalDefines(EShaderPlatform TargetPlatf
 		Ret.PROJECT_SUPPORT_SKY_ATMOSPHERE_AFFECTS_HEIGHFOG = (CVar && bSupportSkyAtmosphere) ? (CVar->GetInt() != 0) : 0;
 	}
 
+	{
+		Ret.PROJECT_EXPFOG_MATCHES_VFOG = DoesProjectSupportExpFogMatchesVolumetricFog() ? 1 : 0;
+	}
+
 	{
 		static IConsoleVariable* CVar = IConsoleManager::Get().FindConsoleVariable(TEXT("r.SupportCloudShadowOnForwardLitTranslucent"));
 		const bool bSupportCloudShadowOnForwardLitTranslucent = CVar && CVar->GetInt() > 0;
diff --git a/Engine/Source/Runtime/RenderCore/Private/RenderUtils.cpp b/Engine/Source/Runtime/RenderCore/Private/RenderUtils.cpp
--- a/Engine/Source/Runtime/RenderCore/Private/RenderUtils.cpp
+++ b/Engine/Source/Runtime/RenderCore/Private/RenderUtils.cpp
@@ -2104,4 +2104,15 @@ bool HardwareVariableRateShadingSupportedByPlatform(EShaderPlatform ShaderPlatfo
 	}
 
 	return false;
-}
\ No newline at end of file
+}
+
+static TAutoConsoleVariable<int32> CVarSupportExpFogMatchesVolumetricFog(
+	TEXT("r.SupportExpFogMatchesVolumetricFog"),
+	0,
+	TEXT("When enabled, the height fog scattering/ambient/emissive/phase will match the volumetric fog better."),
+	ECVF_ReadOnly | ECVF_RenderThreadSafe);
+
+bool DoesProjectSupportExpFogMatchesVolumetricFog()
+{
+	return CVarSupportExpFogMatchesVolumetricFog.GetValueOnAnyThread() > 0;
+}
diff --git a/Engine/Source/Runtime/RenderCore/Private/Shader.cpp b/Engine/Source/Runtime/RenderCore/Private/Shader.cpp
--- a/Engine/Source/Runtime/RenderCore/Private/Shader.cpp
+++ b/Engine/Source/Runtime/RenderCore/Private/Shader.cpp
@@ -2286,6 +2286,14 @@ void ShaderMapAppendKey(EShaderPlatform Platform, FShaderKeyGenerator& KeyGen)
 		}
 	}
 
+	{
+		if (DoesProjectSupportExpFogMatchesVolumetricFog())
+		{
+			KeyGen.AppendSeparator();
+			KeyGen.Append(TEXT("EXPVFOG"));
+		}
+	}
+
 	const bool bNeedsSeparateMainDirLightTexture = IsWaterSeparateMainDirLightEnabled(Platform);
 	if (bNeedsSeparateMainDirLightTexture)
 	{
diff --git a/Engine/Source/Runtime/RenderCore/Public/RenderUtils.h b/Engine/Source/Runtime/RenderCore/Public/RenderUtils.h
--- a/Engine/Source/Runtime/RenderCore/Public/RenderUtils.h
+++ b/Engine/Source/Runtime/RenderCore/Public/RenderUtils.h
@@ -553,6 +553,8 @@ extern RENDERCORE_API bool DoesPlatformSupportLumenGI(EShaderPlatform Platform,
 
 extern RENDERCORE_API bool DoesProjectSupportLumenRayTracedTranslucentRefraction();
 
+RENDERCORE_API bool DoesProjectSupportExpFogMatchesVolumetricFog();
+
 /** Whether or not the platform supports the scene spline texture for spline meshes */
 RENDERCORE_API bool UseSplineMeshSceneResources(const FStaticShaderPlatform Platform);
 
@@ -635,4 +637,4 @@ FORCEINLINE float Log2ToEV100(float LuminanceMax, float Log2)
 }
 
 // Whether or not VRS is supported via r.VRS.Support and the current platform's DDPI
-RENDERCORE_API bool HardwareVariableRateShadingSupportedByPlatform(EShaderPlatform ShaderPlatform);
\ No newline at end of file
+RENDERCORE_API bool HardwareVariableRateShadingSupportedByPlatform(EShaderPlatform ShaderPlatform);
diff --git a/Engine/Source/Runtime/RenderCore/Public/ShaderMaterial.h b/Engine/Source/Runtime/RenderCore/Public/ShaderMaterial.h
--- a/Engine/Source/Runtime/RenderCore/Public/ShaderMaterial.h
+++ b/Engine/Source/Runtime/RenderCore/Public/ShaderMaterial.h
@@ -33,6 +33,7 @@ struct FShaderGlobalDefines
 	bool EARLY_Z_PASS_ONLY_MATERIAL_MASKING;
 	bool PROJECT_SUPPORT_SKY_ATMOSPHERE;
 	bool PROJECT_SUPPORT_SKY_ATMOSPHERE_AFFECTS_HEIGHFOG;
+	bool PROJECT_EXPFOG_MATCHES_VFOG;
 	bool SUPPORT_CLOUD_SHADOW_ON_FORWARD_LIT_TRANSLUCENT;
 	bool SUPPORT_CLOUD_SHADOW_ON_SINGLE_LAYER_WATER;
 	bool POST_PROCESS_ALPHA;
diff --git a/Engine/Source/Runtime/Renderer/Private/FogRendering.cpp b/Engine/Source/Runtime/Renderer/Private/FogRendering.cpp
--- a/Engine/Source/Runtime/Renderer/Private/FogRendering.cpp
+++ b/Engine/Source/Runtime/Renderer/Private/FogRendering.cpp
@@ -93,6 +93,8 @@ void SetupFogUniformParameters(FRDGBuilder& GraphBuilder, const FViewInfo& View,
 		OutParameters.IntegratedLightScatteringSampler = TStaticSamplerState<SF_Bilinear, AM_Clamp, AM_Clamp, AM_Clamp>::GetRHI();
 		OutParameters.VolumetricFogStartDistance = View.VolumetricFogStartDistance;
 		OutParameters.VolumetricFogNearFadeInDistanceInv = View.VolumetricFogNearFadeInDistanceInv;
+		OutParameters.VolumetricFogPhaseG = View.VolumetricFogPhaseG;
+		OutParameters.VolumetricFogAlbedo = View.VolumetricFogAlbedo;
 	}
 }
 
@@ -188,6 +190,8 @@ TGlobalResource<FFogVertexDeclaration> GFogVertexDeclaration;
 
 void FSceneRenderer::InitFogConstants()
 {
+	const bool bExpFogMatchesVolumetricFog = DoesProjectSupportExpFogMatchesVolumetricFog();
+
 	for(int32 ViewIndex = 0;ViewIndex < Views.Num();ViewIndex++)
 	{
 		FViewInfo& View = Views[ViewIndex];
@@ -265,6 +269,18 @@ void FSceneRenderer::InitFogConstants()
 				View.bEnableVolumetricFog = FogInfo.bEnableVolumetricFog;
 				View.VolumetricFogStartDistance = FogInfo.VolumetricFogStartDistance;
 				View.VolumetricFogNearFadeInDistanceInv = FogInfo.VolumetricFogNearFadeInDistance > 0.0f ? (1.0f / FogInfo.VolumetricFogNearFadeInDistance) : 100000000.0f;
+				View.VolumetricFogPhaseG = FogInfo.VolumetricFogScatteringDistribution;
+
+				View.VolumetricFogAlbedo = FVector3f::Zero(); // unused by default
+				if (bExpFogMatchesVolumetricFog)
+				{
+					// We make everything we can to get a good match between height fog and volumetric fog
+					View.DirectionalInscatteringStartDistance = 0.0f;						// No start distance for ExpFog as for VFog
+					View.DirectionalInscatteringExponent = 1.0f;							// Exponent is ununsed in this case
+					View.ExponentialFogColor = FVector3f(FogInfo.VolumetricFogEmissive);	// Emisive from ExpFog matches the VFog ExponentialFogColorParameter
+					View.VolumetricFogAlbedo = FVector3f(FogInfo.VolumetricFogAlbedo);		// Albedo is now supported by ExpFog and matches the VFog
+					View.DirectionalInscatteringColor = FLinearColor::Black;				// Directional scattering is only impacted by the atmospheric light to match
+				}
 			}
 		}
 	}
diff --git a/Engine/Source/Runtime/Renderer/Private/FogRendering.h b/Engine/Source/Runtime/Renderer/Private/FogRendering.h
--- a/Engine/Source/Runtime/Renderer/Private/FogRendering.h
+++ b/Engine/Source/Runtime/Renderer/Private/FogRendering.h
@@ -27,6 +27,8 @@ BEGIN_GLOBAL_SHADER_PARAMETER_STRUCT(FFogUniformParameters,)
 	SHADER_PARAMETER(float, ApplyVolumetricFog)
 	SHADER_PARAMETER(float, VolumetricFogStartDistance)
 	SHADER_PARAMETER(float, VolumetricFogNearFadeInDistanceInv)
+	SHADER_PARAMETER(FVector3f, VolumetricFogAlbedo)
+	SHADER_PARAMETER(float, VolumetricFogPhaseG)
 	SHADER_PARAMETER_TEXTURE(TextureCube, FogInscatteringColorCubemap)
 	SHADER_PARAMETER_SAMPLER(SamplerState, FogInscatteringColorSampler)
 	SHADER_PARAMETER_RDG_TEXTURE(Texture3D, IntegratedLightScattering)
diff --git a/Engine/Source/Runtime/Renderer/Private/SceneCore.cpp b/Engine/Source/Runtime/Renderer/Private/SceneCore.cpp
--- a/Engine/Source/Runtime/Renderer/Private/SceneCore.cpp
+++ b/Engine/Source/Runtime/Renderer/Private/SceneCore.cpp
@@ -416,11 +416,13 @@ FExponentialHeightFogSceneInfo::FExponentialHeightFogSceneInfo(const UExponentia
 	VolumetricFogAlbedo = FLinearColor(InComponent->VolumetricFogAlbedo);
 	VolumetricFogEmissive = InComponent->VolumetricFogEmissive;
 
-	// Apply a scale so artists don't have to work with tiny numbers.  
-	const float UnitScale = 1.0f / 10000.0f;
-	VolumetricFogEmissive.R = FMath::Max(VolumetricFogEmissive.R * UnitScale, 0.0f);
-	VolumetricFogEmissive.G = FMath::Max(VolumetricFogEmissive.G * UnitScale, 0.0f);
-	VolumetricFogEmissive.B = FMath::Max(VolumetricFogEmissive.B * UnitScale, 0.0f);
+	// Apply a scale so artists don't have to work with tiny numbers.
+	// The is only needed because emissive is by default not weighted by the height fog density distribution. 
+	// When we run "HeightFog matches FVog" that is no longer needed.
+	const float EmissiveUnitScale = DoesProjectSupportExpFogMatchesVolumetricFog() ? 1.0f : 1.0f / 10000.0f;
+	VolumetricFogEmissive.R = FMath::Max(VolumetricFogEmissive.R * EmissiveUnitScale, 0.0f);
+	VolumetricFogEmissive.G = FMath::Max(VolumetricFogEmissive.G * EmissiveUnitScale, 0.0f);
+	VolumetricFogEmissive.B = FMath::Max(VolumetricFogEmissive.B * EmissiveUnitScale, 0.0f);
 	VolumetricFogExtinctionScale = FMath::Max(InComponent->VolumetricFogExtinctionScale, 0.0f);
 	VolumetricFogDistance = FMath::Max(InComponent->VolumetricFogStartDistance + InComponent->VolumetricFogDistance, 0.0f);
 	VolumetricFogStaticLightingScatteringIntensity = FMath::Max(InComponent->VolumetricFogStaticLightingScatteringIntensity, 0.0f);
diff --git a/Engine/Source/Runtime/Renderer/Private/SceneRendering.cpp b/Engine/Source/Runtime/Renderer/Private/SceneRendering.cpp
--- a/Engine/Source/Runtime/Renderer/Private/SceneRendering.cpp
+++ b/Engine/Source/Runtime/Renderer/Private/SceneRendering.cpp
@@ -963,6 +963,8 @@ void FViewInfo::Init()
 	VolumetricFogStartDistance = false;
 	VolumetricFogStartDistance = 0.0f;
 	VolumetricFogNearFadeInDistanceInv = 100000000.0f;
+	VolumetricFogAlbedo = FVector3f::Zero();
+	VolumetricFogPhaseG = 0.0f;
 
 	SkyAtmosphereCameraAerialPerspectiveVolume = nullptr;
 	SkyAtmosphereCameraAerialPerspectiveVolumeMieOnly = nullptr;
diff --git a/Engine/Source/Runtime/Renderer/Private/SceneRendering.h b/Engine/Source/Runtime/Renderer/Private/SceneRendering.h
--- a/Engine/Source/Runtime/Renderer/Private/SceneRendering.h
+++ b/Engine/Source/Runtime/Renderer/Private/SceneRendering.h
@@ -1299,6 +1299,8 @@ public:
 	bool bEnableVolumetricFog;
 	float VolumetricFogStartDistance;
 	float VolumetricFogNearFadeInDistanceInv;
+	FVector3f VolumetricFogAlbedo;
+	float VolumetricFogPhaseG;
 	FVector2f SinCosInscatteringColorCubemapRotation;
 
 	UTexture* FogInscatteringColorCubemap;
