Fix for badly placed function from 34421615 which should have not been introduced in the first place for legacy compatibility.

diff --git a/Engine/Shaders/Private/HeightFogCommon.ush b/Engine/Shaders/Private/HeightFogCommon.ush
--- a/Engine/Shaders/Private/HeightFogCommon.ush
+++ b/Engine/Shaders/Private/HeightFogCommon.ush
@@ -340,7 +340,7 @@ half4 GetExponentialHeightFog(float3 WorldPositionRelativeToCamera, float Exclud
 	#if PROJECT_EXPFOG_MATCHES_VFOG
 			* HenyeyGreensteinPhase(FogData.VolumetricFogPhaseG, dot(-CameraToReceiverNormalized, InView.AtmosphereLightDirection.xyz));
 	#else
-			* pow(saturate(dot(CameraToReceiverNormalized, FogData.InscatteringLightDirection.xyz)), FogData.DirectionalInscatteringColor.w) * UniformPhaseFunction;
+			* pow(saturate(dot(CameraToReceiverNormalized, FogData.InscatteringLightDirection.xyz)), FogData.DirectionalInscatteringColor.w);
 	#endif
 #endif
 
