VulkanRHI: Fix ray tracing debug modes in Vulkan (temporarily disable LIGHT_GRID_COUNT and HITCOUNT_PER_INSTANCE modes in Vulkan)

#rnx

diff --git a/Engine/Shaders/Private/RayTracing/RayTracingDebug.usf b/Engine/Shaders/Private/RayTracing/RayTracingDebug.usf
--- a/Engine/Shaders/Private/RayTracing/RayTracingDebug.usf
+++ b/Engine/Shaders/Private/RayTracing/RayTracingDebug.usf
@@ -106,13 +106,18 @@ float4 DebugRayTracingMaterial(FRayDesc Ray, bool bIsFarField)
 		case RAY_TRACING_DEBUG_VIZ_INDIRECT_IRRADIANCE:		return float4(Payload.IndirectIrradiance, 0.0f);
 		case RAY_TRACING_DEBUG_VIZ_WORLD_POSITION:			return float4(Payload.TranslatedWorldPos, 0.0f);
 		case RAY_TRACING_DEBUG_VIZ_HITKIND:					return Payload.IsFrontFace() ? float4(0.0, 1.0, 0.0, 0.0f) : float4(1.0, 0.0, 0.0, 0.0f);
+
+		#if !VULKAN_PROFILE_SM6
+		// :todo-jn:  disabled until SPIRV codegen issue can be found (breaks all other debug modes when enabled). See UE-218212.
 		case RAY_TRACING_DEBUG_VIZ_LIGHT_GRID_COUNT:	    return float4(FCulledLightList::Create(Payload.TranslatedWorldPos).Visualize(1) * abs(dot(Payload.WorldNormal, Ray.Direction)), 0.0f);
+		#endif
+
 		#if PLATFORM_SUPPORTS_SHADER_TIMESTAMP
 		case RAY_TRACING_DEBUG_VIZ_PERFORMANCE:				return float4(ColorMapInferno(float(TimeEnd - TimeBegin) * TimingScale), 0.0f);
 		#endif // PLATFORM_SUPPORTS_SHADER_TIMESTAMP
 		}
 
-		// 2. Material depends view mode
+		// 2. Material dependent view mode
 	#if SUBSTRATE_ENABLED
 		FSubstrateAddressing SubstrateAddressing   = GetSubstratePixelDataByteOffset(0, 0, 0);
 		FSubstratePixelHeader SubstratePixelHeader = UnpackSubstrateHeaderIn(Payload.SubstrateData, SubstrateAddressing, Payload.SubstrateData);
@@ -210,9 +215,12 @@ float4 DebugRayTracingInstance(FRayDesc Ray, out float OutDepth)
 		case RAY_TRACING_DEBUG_VIZ_TRIANGLE_HITCOUNT:
 			Result.rgb = ColorMapTurbo(Payload.TriangleHitCountPerRay / TriangleHitCountMaxThreshold);
 			break;
+		#if !VULKAN_PROFILE_SM6
+		// :todo-jn:  disabled until SPIRV codegen issue can be found (breaks all other debug modes when enabled). See UE-218215.
 		case RAY_TRACING_DEBUG_VIZ_HITCOUNT_PER_INSTANCE:
 			Result.rgb = ColorMapTurbo(RayTracingDebugHitStatsUniformBuffer.HitStatsOutput[Payload.InstanceIndex].Count / (float)TriangleHitCountPerInstanceMaxThreshold);
 			break;
+		#endif
 		case RAY_TRACING_DEBUG_VIZ_INSTANCE_OVERLAP:
 			Result.rgb = 1.0; // We only use fake lighting as instance heatmap is blended in
 			break;
@@ -226,7 +234,7 @@ float4 DebugRayTracingInstance(FRayDesc Ray, out float OutDepth)
 			Result.rgb = IntToColor(InstancesDebugData[Payload.InstanceIndex].ProxyHash);
 			break;
 		}
-	
+
 		if (VisualizationMode == RAY_TRACING_DEBUG_VIZ_PICKER)
 		{
 			const float ColorScale = 0.1f;
