LensDistortion: Add support to after motion blur translucency composition.

diff --git a/Engine/Shaders/Private/ComposeSeparateTranslucency.usf b/Engine/Shaders/Private/ComposeSeparateTranslucency.usf
--- a/Engine/Shaders/Private/ComposeSeparateTranslucency.usf
+++ b/Engine/Shaders/Private/ComposeSeparateTranslucency.usf
@@ -1,6 +1,7 @@
 // Copyright Epic Games, Inc. All Rights Reserved.
 
 #include "Common.ush"
+#include "LensDistortion.ush"
 #include "ScreenPass.ush"
 
 #if PERMUTATION_NEARESTDEPTHNEIGHBOR 
@@ -26,12 +27,15 @@ SamplerState SeparateTranslucencyBilinearSampler;
 Texture2D    SeparateModulationBilinearTexture;
 SamplerState SeparateModulationBilinearSampler;
 
+Texture2D<float2> UndistortingDisplacementTexture;
+SamplerState      UndistortingDisplacementSampler;
+
 FScreenTransform ScreenPosToSceneColorUV;
 FScreenTransform ScreenPosToSeparateTranslucencyUV;
 float2 SeparateTranslucencyUVMin;
 float2 SeparateTranslucencyUVMax;
 float2 SeparateTranslucencyExtentInverse;
-
+uint bLensDistortion;
 
 //------------------------------------------------------- ENTRY POINT
 
@@ -43,6 +47,11 @@ void MainPS(
 
 	float2 SceneColorUV = ApplyScreenTransform(SvPosition.xy, ScreenPosToSceneColorUV);
 	float2 SeparateTranslucencyUV = ApplyScreenTransform(SvPosition.xy, ScreenPosToSeparateTranslucencyUV);
+	BRANCH
+	if (bLensDistortion)
+	{
+		SeparateTranslucencyUV = ApplyLensDistortionOnViewportUV(UndistortingDisplacementTexture, UndistortingDisplacementSampler, SeparateTranslucencyUV);
+	}
 
 	float4 SceneColorSample = SceneColorTexture.SampleLevel(SceneColorSampler, SceneColorUV, 0);
 	
diff --git a/Engine/Shaders/Private/MotionBlur/MotionBlurApply.usf b/Engine/Shaders/Private/MotionBlur/MotionBlurApply.usf
--- a/Engine/Shaders/Private/MotionBlur/MotionBlurApply.usf
+++ b/Engine/Shaders/Private/MotionBlur/MotionBlurApply.usf
@@ -4,6 +4,7 @@
 
 
 #include "MotionBlurCommon.ush"
+#include "../LensDistortion.ush"
 #include "../Random.ush"
 #include "/Engine/Public/WaveBroadcastIntrinsics.ush"
 
@@ -84,6 +85,7 @@ FScreenTransform ColorToVelocity;
 uint MaxSampleCount;
 uint OutputMip1;
 uint OutputMip2;
+uint bLensDistortion;
 
 uint TileListOffset;
 StructuredBuffer<uint> TileListsBuffer;
@@ -101,6 +103,9 @@ SamplerState VelocityTileSampler;
 SamplerState VelocityFlatSampler;
 SamplerState DepthSampler;
 
+Texture2D<float2> UndistortingDisplacementTexture;
+SamplerState      UndistortingDisplacementSampler;
+
 Texture2D TranslucencyTexture;
 SamplerState TranslucencySampler;
 FScreenTransform ColorToTranslucency;
@@ -734,6 +739,15 @@ void MainCS(
 			float2 PostMotionBlurTranslucencyUV2 = PostMotionBlurTranslucencyUV + float2(-0.5, +0.5) * TranslucencyExtentInverse;
 			float2 PostMotionBlurTranslucencyUV3 = PostMotionBlurTranslucencyUV + float2(+0.5, +0.5) * TranslucencyExtentInverse;
 		#endif
+
+		BRANCH
+		if (bLensDistortion)
+		{
+			PostMotionBlurTranslucencyUV0 = ApplyLensDistortionOnViewportUV(UndistortingDisplacementTexture, UndistortingDisplacementSampler, PostMotionBlurTranslucencyUV0);
+			PostMotionBlurTranslucencyUV1 = ApplyLensDistortionOnViewportUV(UndistortingDisplacementTexture, UndistortingDisplacementSampler, PostMotionBlurTranslucencyUV1);
+			PostMotionBlurTranslucencyUV2 = ApplyLensDistortionOnViewportUV(UndistortingDisplacementTexture, UndistortingDisplacementSampler, PostMotionBlurTranslucencyUV2);
+			PostMotionBlurTranslucencyUV3 = ApplyLensDistortionOnViewportUV(UndistortingDisplacementTexture, UndistortingDisplacementSampler, PostMotionBlurTranslucencyUV3);
+		}
 			
 		PostMotionBlurTranslucencyUV0 = clamp(PostMotionBlurTranslucencyUV0, TranslucencyUVMin, TranslucencyUVMax);
 		PostMotionBlurTranslucencyUV1 = clamp(PostMotionBlurTranslucencyUV1, TranslucencyUVMin, TranslucencyUVMax);
@@ -979,6 +993,11 @@ void MainCS(
 		mb_short2 OutputPixelCoord = mb_short2(ColorUV * Color_Extent);
 
 		float2 PostMotionBlurTranslucencyUV = ApplyScreenTransform(ColorUV, ColorToTranslucency);
+		BRANCH
+		if (bLensDistortion)
+		{
+			PostMotionBlurTranslucencyUV = ApplyLensDistortionOnViewportUV(UndistortingDisplacementTexture, UndistortingDisplacementSampler, PostMotionBlurTranslucencyUV);
+		}
 		PostMotionBlurTranslucencyUV = clamp(PostMotionBlurTranslucencyUV, TranslucencyUVMin, TranslucencyUVMax);
 		
 		mb_half4 CenterColor;
diff --git a/Engine/Source/Runtime/Renderer/Private/PostProcess/PostProcessMotionBlur.cpp b/Engine/Source/Runtime/Renderer/Private/PostProcess/PostProcessMotionBlur.cpp
--- a/Engine/Source/Runtime/Renderer/Private/PostProcess/PostProcessMotionBlur.cpp
+++ b/Engine/Source/Runtime/Renderer/Private/PostProcess/PostProcessMotionBlur.cpp
@@ -464,6 +464,7 @@ class FMotionBlurFilterCS : public FMotionBlurShader
 		SHADER_PARAMETER(int32, MaxSampleCount)
 		SHADER_PARAMETER(int32, OutputMip1)
 		SHADER_PARAMETER(int32, OutputMip2)
+		SHADER_PARAMETER(int32, bLensDistortion)
 
 		SHADER_PARAMETER_RDG_TEXTURE_SRV(Texture2D, ColorTexture)
 		SHADER_PARAMETER_RDG_TEXTURE(Texture2D, VelocityFlatTexture)
@@ -474,6 +475,9 @@ class FMotionBlurFilterCS : public FMotionBlurShader
 		SHADER_PARAMETER_SAMPLER(SamplerState, VelocityTileSampler)
 		SHADER_PARAMETER_SAMPLER(SamplerState, VelocityFlatSampler)
 
+		SHADER_PARAMETER_RDG_TEXTURE(Texture2D, UndistortingDisplacementTexture)
+		SHADER_PARAMETER_SAMPLER(SamplerState, UndistortingDisplacementSampler)
+
 		SHADER_PARAMETER_RDG_TEXTURE(Texture2D, TranslucencyTexture)
 		SHADER_PARAMETER_SAMPLER(SamplerState, TranslucencySampler)
 		SHADER_PARAMETER(FScreenTransform, ColorToTranslucency)
@@ -1006,6 +1010,10 @@ FMotionBlurOutputs AddMotionBlurFilterPass(
 		OriginalPassParameters.TileListsSizeBuffer = GraphBuilder.CreateSRV(TileListsSizeBuffer);
 		OriginalPassParameters.DispatchParameters = DispatchParameters;
 
+		OriginalPassParameters.UndistortingDisplacementTexture = GSystemTextures.GetBlackDummy(GraphBuilder);
+		OriginalPassParameters.UndistortingDisplacementSampler = TStaticSamplerState<SF_Bilinear, AM_Clamp, AM_Clamp, AM_Clamp>::GetRHI();
+		OriginalPassParameters.bLensDistortion = 0;
+
 		if (PostMotionBlurTranslucency != nullptr)
 		{
 			// TODO: broken with split screen
@@ -1023,6 +1031,12 @@ FMotionBlurOutputs AddMotionBlurFilterPass(
 			OriginalPassParameters.TranslucencyUVMin = FVector2f(0.0f, 0.0f);
 			OriginalPassParameters.TranslucencyUVMax = (FVector2f(PostMotionBlurTranslucencySize) - FVector2f(0.5f, 0.5f)) * PostMotionBlurTranslucencyExtentInv;
 			OriginalPassParameters.TranslucencyExtentInverse = PostMotionBlurTranslucencyExtentInv;
+
+			OriginalPassParameters.bLensDistortion = Inputs.LensDistortionLUT.IsEnabled();
+			if (Inputs.LensDistortionLUT.IsEnabled())
+			{
+				OriginalPassParameters.UndistortingDisplacementTexture = Inputs.LensDistortionLUT.UndistortingDisplacementTexture;
+			}
 		}
 		else
 		{
diff --git a/Engine/Source/Runtime/Renderer/Private/PostProcess/PostProcessing.cpp b/Engine/Source/Runtime/Renderer/Private/PostProcess/PostProcessing.cpp
--- a/Engine/Source/Runtime/Renderer/Private/PostProcess/PostProcessing.cpp
+++ b/Engine/Source/Runtime/Renderer/Private/PostProcess/PostProcessing.cpp
@@ -1002,6 +1002,10 @@ void AddPostProcessingPasses(
 			TranslucencyComposition.SceneColor = SceneColorSlice;
 			TranslucencyComposition.OutputViewport = FScreenPassTextureViewport(SceneColorSlice);
 			TranslucencyComposition.OutputPixelFormat = SceneColorFormat;
+			if (bApplyLensDistortionInTSR)
+			{
+				TranslucencyComposition.LensDistortionLUT = View.LensDistortionLUT;
+			}
 
 			SceneColorSlice = FScreenPassTextureSlice::CreateFromScreenPassTexture(GraphBuilder, TranslucencyComposition.AddPass(
 				GraphBuilder, View, PostMotionBlurTranslucencyResources));
diff --git a/Engine/Source/Runtime/Renderer/Private/TranslucentRendering.cpp b/Engine/Source/Runtime/Renderer/Private/TranslucentRendering.cpp
--- a/Engine/Source/Runtime/Renderer/Private/TranslucentRendering.cpp
+++ b/Engine/Source/Runtime/Renderer/Private/TranslucentRendering.cpp
@@ -389,6 +389,7 @@ class FComposeSeparateTranslucencyPS : public FGlobalShader
 		SHADER_PARAMETER(FVector2f, SeparateTranslucencyUVMin)
 		SHADER_PARAMETER(FVector2f, SeparateTranslucencyUVMax)
 		SHADER_PARAMETER(FVector2f, SeparateTranslucencyExtentInverse)
+		SHADER_PARAMETER(int32, bLensDistortion)
 
 		SHADER_PARAMETER_RDG_TEXTURE_SRV(Texture2D, SceneColorTexture)
 		SHADER_PARAMETER_SAMPLER(SamplerState,  SceneColorSampler)
@@ -409,6 +410,9 @@ class FComposeSeparateTranslucencyPS : public FGlobalShader
 		SHADER_PARAMETER_RDG_TEXTURE(Texture2D, FullResDepthTexture)
 		SHADER_PARAMETER_SAMPLER(SamplerState, FullResDepthSampler)
 
+		SHADER_PARAMETER_RDG_TEXTURE(Texture2D, UndistortingDisplacementTexture)
+		SHADER_PARAMETER_SAMPLER(SamplerState, UndistortingDisplacementSampler)
+
 		SHADER_PARAMETER_STRUCT_REF(FViewUniformShaderParameters, ViewUniformBuffer)
 		RENDER_TARGET_BINDING_SLOTS()
 	END_SHADER_PARAMETER_STRUCT()
@@ -586,6 +590,14 @@ FScreenPassTexture FTranslucencyComposition::AddPass(
 	PassParameters->SeparateModulationBilinearTexture = SeparateModulationTexture;
 	PassParameters->SeparateModulationBilinearSampler = TStaticSamplerState<SF_Bilinear>::GetRHI();
 
+	PassParameters->UndistortingDisplacementTexture = GSystemTextures.GetBlackDummy(GraphBuilder);
+	PassParameters->UndistortingDisplacementSampler = TStaticSamplerState<SF_Bilinear, AM_Clamp, AM_Clamp, AM_Clamp>::GetRHI();
+	PassParameters->bLensDistortion = LensDistortionLUT.IsEnabled();
+	if (LensDistortionLUT.IsEnabled())
+	{
+		PassParameters->UndistortingDisplacementTexture = LensDistortionLUT.UndistortingDisplacementTexture;
+	}
+
 	PassParameters->ViewUniformBuffer = View.ViewUniformBuffer;
 
 	if (Operation == EOperation::ComposeToExistingSceneColor)
diff --git a/Engine/Source/Runtime/Renderer/Private/TranslucentRendering.h b/Engine/Source/Runtime/Renderer/Private/TranslucentRendering.h
--- a/Engine/Source/Runtime/Renderer/Private/TranslucentRendering.h
+++ b/Engine/Source/Runtime/Renderer/Private/TranslucentRendering.h
@@ -79,6 +79,9 @@ struct FTranslucencyComposition
 	FScreenPassTextureViewport OutputViewport;
 	EPixelFormat OutputPixelFormat = PF_Unknown;
 
+	// [Optional] Lens distortion applied on the scene color.
+	FLensDistortionLUT LensDistortionLUT;
+
 	FScreenPassTexture AddPass(
 		FRDGBuilder& GraphBuilder,
 		const FViewInfo& View,
