Fixed shader runtime issues due to mismatch parameter sctructure after 34375713.

diff --git a/Engine/Shaders/Private/VolumetricCloud.usf b/Engine/Shaders/Private/VolumetricCloud.usf
--- a/Engine/Shaders/Private/VolumetricCloud.usf
+++ b/Engine/Shaders/Private/VolumetricCloud.usf
@@ -5,6 +5,8 @@
  */  
 
 
+ // Change this to force recompilation of all volumetric cloud material shaders
+#pragma message("UESHADERMETADATA_VERSION BFF66DE1-56ED-4BFA-AE26-D7FCF6E790CE")
 
 // Change this to force recompilation of all shaders referencing this file.
 #pragma message("UESHADERMETADATA_VERSION 58372ED6-B830-41C4-829E-A21186A68401")
