Remove comment accidentally ended up in commit

diff --git a/Engine/Shaders/Private/HeightFogCommon.ush b/Engine/Shaders/Private/HeightFogCommon.ush
--- a/Engine/Shaders/Private/HeightFogCommon.ush
+++ b/Engine/Shaders/Private/HeightFogCommon.ush
@@ -507,7 +507,7 @@ float3 ComputeVolumeUV_DEPRECATED(float3 WorldPosition, float4x4 WorldToClip)
 	float4 NDCPosition = mul(float4(WorldPosition, 1), WorldToClip);
 	return ComputeVolumeUVFromNDC(NDCPosition);
 }
-float3 ComputeVolumeUV(FDFVector3 WorldPosition, FDFInverseMatrix WorldToClip, ViewState InView) //ericado check
+float3 ComputeVolumeUV(FDFVector3 WorldPosition, FDFInverseMatrix WorldToClip, ViewState InView)
 {
 	float4 NDCPosition = DFMultiplyDemote(MakeDFVector(WorldPosition, 1.0f), WorldToClip);
 	return ComputeVolumeUVFromNDC(NDCPosition, InView);
