Fix a shader compilation error due to a typo in LIghtGridCommon.ush

diff --git a/Engine/Shaders/Private/LightGridCommon.ush b/Engine/Shaders/Private/LightGridCommon.ush
--- a/Engine/Shaders/Private/LightGridCommon.ush
+++ b/Engine/Shaders/Private/LightGridCommon.ush
@@ -219,7 +219,7 @@ FDirectionalLightData GetDirectionalLightData(uint EyeIndex)
 		Result.DirectionalLightSourceRadius			= ForwardLightDataISR.DirectionalLightSourceRadius;
 		Result.DirectionalLightSoftSourceRadius		= ForwardLightDataISR.DirectionalLightSoftSourceRadius;
 		Result.DirectionalLightSpecularScale		= ForwardLightDataISR.DirectionalLightSpecularScale;
-		Result.DirectionalLightDiffuseScale			= ForwardLightDataISR.DirectionalLightDiffuserScale;
+		Result.DirectionalLightDiffuseScale			= ForwardLightDataISR.DirectionalLightDiffuseScale;
 		Result.LightFunctionAtlasLightIndex			= ForwardLightDataISR.LightFunctionAtlasLightIndex;
 		Result.bAffectsTranslucentLighting			= ForwardLightDataISR.bAffectsTranslucentLighting;
 	}
