Fix analytical derivative with PT.

The clip positions were divided by .w prior to compute barycentric derivatives. This was causing barycentrics derivatives to be incorrect (i.e., constant) as w=1.

diff --git a/Engine/Shaders/Private/PathTracing/PathTracingMaterialHitShader.usf b/Engine/Shaders/Private/PathTracing/PathTracingMaterialHitShader.usf
--- a/Engine/Shaders/Private/PathTracing/PathTracingMaterialHitShader.usf
+++ b/Engine/Shaders/Private/PathTracing/PathTracingMaterialHitShader.usf
@@ -24,7 +24,8 @@
 // Ensure specular profile and glints are supported
 // TODO: the derivatives don't appear to be correct in all cases, causing incorrectly blurry lookups - so disable them for now unless using substrate
 //  We don't have an easy way to detect if the material is using glints or not, so use the complex special path which includes glints for now
-#define USE_ANALYTIC_DERIVATIVES (SUBSTRATE_ENABLED && SUBSTRATE_COMPLEXSPECIALPATH && PLATFORM_ENABLES_SUBSTRATE_GLINTS)
+#define SUBSTRATE_COMPLEXSPECIALPATH 1
+#define USE_ANALYTIC_DERIVATIVES 1
 
 #include "/Engine/Private/Common.ush"
 #include "/Engine/Private/RayTracing/RayTracingCommon.ush"
diff --git a/Engine/Shaders/Private/RayTracing/RayTracingCalcInterpolants.ush b/Engine/Shaders/Private/RayTracing/RayTracingCalcInterpolants.ush
--- a/Engine/Shaders/Private/RayTracing/RayTracingCalcInterpolants.ush
+++ b/Engine/Shaders/Private/RayTracing/RayTracingCalcInterpolants.ush
@@ -62,7 +62,6 @@ void CalcInterpolants(
 		#endif
 
 		ClipPositions[i] = mul(float4(WorldPositions[i], 1), View.TranslatedWorldToClip);
-		ClipPositions[i]/= ClipPositions[i].w;
 	}
 
 	OutInterpolants = VertexFactoryAssignInterpolants(Interpolated);
@@ -90,7 +89,7 @@ void CalcInterpolants(
 	const float2 PixelClip = (SvPosition.xy - View.ViewRectMin.xy) * View.ViewSizeAndInvSize.zw * float2(2, -2) + float2(-1, 1);
 	FBarycentrics Barycentrics = CalculateTriangleBarycentrics(PixelClip, ClipPositions[0], ClipPositions[1], ClipPositions[2], View.ViewSizeAndInvSize.zw);
 	OutDerivateData.WorldPosition = Lerp( WorldPositions[0], WorldPositions[1], WorldPositions[2], Barycentrics );
-	OutDerivateData.ScreenPosition = Lerp( ClipPositions[0], ClipPositions[1], ClipPositions[2], Barycentrics );
+	OutDerivateData.ScreenPosition = Lerp( ClipPositions[0] * rcp(ClipPositions[0].w), ClipPositions[1] * rcp(ClipPositions[1].w), ClipPositions[2] * rcp(ClipPositions[2].w), Barycentrics );
 
 #if NUM_TEX_COORD_INTERPOLATORS > 0
 	UNROLL
