Don't pull in InstancedView buffer into the params of fullscreen VS.

diff --git a/Engine/Shaders/Private/Tools/FullscreenVertexShader.usf b/Engine/Shaders/Private/Tools/FullscreenVertexShader.usf
--- a/Engine/Shaders/Private/Tools/FullscreenVertexShader.usf
+++ b/Engine/Shaders/Private/Tools/FullscreenVertexShader.usf
@@ -29,8 +29,8 @@ void MainVS(
 	,	out noperspective float2 OutUV : TEXCOORD0
 	#if INSTANCED_STEREO // Mobile multi view fallback path
 	,	in uint InstanceId : SV_InstanceID
-	,	out nointerpolation uint EyeIndex : VIEW_ID
-	,	out uint LayerIndex : SV_RenderTargetArrayIndex
+	,	out nointerpolation uint OutViewId : VIEW_ID
+	,	out uint OutLayerIndex : SV_RenderTargetArrayIndex
 	#endif
 #endif
 	)
@@ -38,15 +38,15 @@ void MainVS(
 	OutPosition = float4(ViewportUVToScreenPos(ViewportUV), 0, 1);
 
 #if SCREEN_VS_FOR_INSTANCED_VIEWS
-	OutViewportId = InstanceId;
 	OutViewId = InstanceId;
+	OutViewportId = OutViewId;
 #endif
 	
 #if SCREEN_VS_FOR_MOBILE_MULTI_VIEW
 	OutUV = UV;
 	#if INSTANCED_STEREO // Mobile multi view fallback path
-		EyeIndex = GetEyeIndex(InstanceId);
-		LayerIndex = EyeIndex;
+		OutViewId = InstanceId;
+		OutLayerIndex = OutViewId;
 	#endif
 #endif
 }
