Fix missing tiles when rendering groom with PPLL due to incorrect tile coordinate computation.

diff --git a/Engine/Shaders/Private/HairStrands/HairStrandsVisibilityCompaction.usf b/Engine/Shaders/Private/HairStrands/HairStrandsVisibilityCompaction.usf
--- a/Engine/Shaders/Private/HairStrands/HairStrandsVisibilityCompaction.usf
+++ b/Engine/Shaders/Private/HairStrands/HairStrandsVisibilityCompaction.usf
@@ -49,15 +49,32 @@ struct FSampleSetDesc
 #if PERMUTATION_PPLL
 #define PPLL_MAX_RENDER_NODE_COUNT PERMUTATION_PPLL
 
+Buffer<uint> TileCountBuffer;
+Buffer<uint2> TileDataBuffer;
+
+int2 TileCountXY;
 Texture2D<uint> PPLLCounter;
 Texture2D<uint> PPLLNodeIndex;
 StructuredBuffer<FPackedHairVisPPLL> PPLLNodeData;
 
 
 [numthreads(TILE_PIXEL_SIZE_X, TILE_PIXEL_SIZE_Y, 1)]
-void MainCS(uint GroupIndex : SV_GroupIndex, uint3 DispatchThreadId : SV_DispatchThreadID)
+void MainCS(
+	uint GroupIndex : SV_GroupIndex, 
+	uint3 GroupId : SV_GroupID, 
+	uint3 GroupThreadId : SV_GroupThreadID, 
+	uint3 DispatchThreadId : SV_DispatchThreadID)
 {
-	const uint2 PixelCoord = DispatchThreadId.xy + ResolutionOffset;
+	const uint TileSize = HAIR_TILE_SIZE;
+	const uint TileCount = TileCountBuffer[HAIRTILE_HAIR_ALL];
+	const uint LinearIndex = GroupId.x + GroupId.y * TileCountXY.x;
+	if (LinearIndex >= TileCount)
+	{
+		return;
+	}
+	const uint2 TileCoord = TileDataBuffer[LinearIndex];
+	const uint2 PixelCoord = TileCoord * TileSize + GroupThreadId.xy;
+
 	const uint FirstNodeIndex = PPLLNodeIndex[PixelCoord];
 	float PixelCoverage = 0;
 	
diff --git a/Engine/Shaders/Private/HairStrands/HairStrandsVisibilityTile.usf b/Engine/Shaders/Private/HairStrands/HairStrandsVisibilityTile.usf
--- a/Engine/Shaders/Private/HairStrands/HairStrandsVisibilityTile.usf
+++ b/Engine/Shaders/Private/HairStrands/HairStrandsVisibilityTile.usf
@@ -200,14 +200,14 @@ void FillIndirectBuffers(uint TileType)
 	// Indirect dispatch
 	WriteDispatchIndirectArgs(TileIndirectDispatchBuffer, TileType,
 		TileCount >= TileCountXY.x ? TileCountXY.x : TileCount,
-		(TileCount + TileCountXY.x - 1) / TileCountXY.x,
+		DivideAndRoundUp(TileCount, TileCountXY.x),
 		1);
 
 	// Indirect dispatch with one thread
 	// At 4k, with 8x8 tiles: 4096 x 4096 / 8 x 8 = 512 x 512 tiles 
 	// With a group size of 64 (8x8), the total dispatch count should be 512 x 512 / 8 x 8 = 64 x 64 = 4096, which is below the 65k limit per dimension
 	WriteDispatchIndirectArgs(TilePerThreadIndirectDispatchBuffer, TileType,
-		(TileCount + TilePerThread_GroupSize - 1) / TilePerThread_GroupSize,
+		DivideAndRoundUp(TileCount, TilePerThread_GroupSize),
 		1,
 		1);
 
