Fix another artifact when the shading point is nearly in coplanar configuration with the rect light plane.

diff --git a/Engine/Shaders/Private/RectLightIntegrate.ush b/Engine/Shaders/Private/RectLightIntegrate.ush
--- a/Engine/Shaders/Private/RectLightIntegrate.ush
+++ b/Engine/Shaders/Private/RectLightIntegrate.ush
@@ -250,9 +250,13 @@ FAreaLightIntegrateContext CreateRectIntegrateContext( float Roughness, half3 N,
 
 FDirectLighting IntegrateBxDF(FGBufferData GBuffer, half3 N, half3 V, FRect Rect, FShadowTerms Shadow, FRectTexture SourceTexture)
 {
+	// Compute distance to light plane and cull if not larger than threshold to 
+	// avoid numerical issue in (nearly) coplanar configuration
+	const float Distance = dot(Rect.Axis[2], Rect.Origin);
+
 	// No-visible rect light due to barn door occlusion
 	FDirectLighting Out = (FDirectLighting)0;
-	if (IsRectVisible(Rect))
+	if (IsRectVisible(Rect) && Distance > 0.001f)
 	{
 		FAreaLightIntegrateContext Context = CreateRectIntegrateContext(GBuffer.Roughness, N, V, Rect, SourceTexture);
 		GBuffer.Roughness = max(GBuffer.Roughness, 0.02);
