Many Lights - fixed shader compilation error on Vulkan
"/Engine/Private/ManyLights/ManyLightsShading.usf:28:43: error: use of undeclared identifier 'NUM_SAMPLES_PER_PIXEL_2D_X'

diff --git a/Engine/Shaders/Private/ManyLights/ManyLightsShading.usf b/Engine/Shaders/Private/ManyLights/ManyLightsShading.usf
--- a/Engine/Shaders/Private/ManyLights/ManyLightsShading.usf
+++ b/Engine/Shaders/Private/ManyLights/ManyLightsShading.usf
@@ -4,6 +4,12 @@
 	#define TILE_TYPE TILE_MODE_EMPTY
 #endif
 
+#ifndef NUM_SAMPLES_PER_PIXEL_1D
+	#define NUM_SAMPLES_PER_PIXEL_1D 1
+	#define NUM_SAMPLES_PER_PIXEL_2D_X 1
+	#define NUM_SAMPLES_PER_PIXEL_2D_Y 1
+#endif
+
 #include "../Common.ush"
 #include "../BlueNoise.ush"
 #include "ManyLights.ush"
