Fix for mobile fog not rendering correctly. Minor shader changes create parity with SM5/6.

diff --git a/Engine/Shaders/Private/DeferredLightPixelShaders.usf b/Engine/Shaders/Private/DeferredLightPixelShaders.usf
--- a/Engine/Shaders/Private/DeferredLightPixelShaders.usf
+++ b/Engine/Shaders/Private/DeferredLightPixelShaders.usf
@@ -101,7 +101,7 @@ FDerivedParams GetDerivedParams(in FInputParams Input, in float SceneDepth)
 	Out.TranslatedWorldPosition = mul(float4(ClipPosition, SceneDepth, 1), View.ScreenToTranslatedWorld).xyz;
 	Out.CameraVector = GetCameraVectorFromTranslatedWorldPosition(Out.TranslatedWorldPosition);
 #else
-	Out.TranslatedWorldPosition = Input.ScreenVector * SceneDepth + GetTranslatedWorldCameraPosFromView(Input.PixelPos + 0.5f);
+	Out.TranslatedWorldPosition = Input.ScreenVector * SceneDepth + GetTranslatedWorldCameraPosFromView(Input.PixelPos);
 	Out.CameraVector = normalize(Input.ScreenVector);
 #endif
 	return Out;
diff --git a/Engine/Shaders/Private/MobileFog.usf b/Engine/Shaders/Private/MobileFog.usf
--- a/Engine/Shaders/Private/MobileFog.usf
+++ b/Engine/Shaders/Private/MobileFog.usf
@@ -54,7 +54,7 @@ void MobileFogVS(
 	// texture coord from vb
 	OutTexCoord = InPosition * ResolvedView.ScreenPositionScaleBias.xy + ResolvedView.ScreenPositionScaleBias.wz;
 	// deproject to world space
-	OutScreenVector = mul(float4(InPosition,1,0), ResolvedView.ScreenToTranslatedWorld).xyz;
+	OutScreenVector = ScreenVectorFromScreenRect(float4(InPosition,1,0));
 }
 
 void MobileFogPS(
@@ -85,7 +85,7 @@ void MobileFogPS(
 	float DeviceZ = LookupDeviceZ(TexCoord, ArrayIndex);
 	float SceneDepth = ConvertFromDeviceZ(DeviceZ);
 	float3 WorldPositionRelativeToCamera = ScreenVector.xyz * SceneDepth;
-	float3 TranslatedWorldPosition = WorldPositionRelativeToCamera + ResolvedView.TranslatedWorldCameraOrigin;
+	float3 TranslatedWorldPosition = WorldPositionRelativeToCamera + GetTranslatedWorldCameraPosFromView(SVPos.xy);
 
 // at least one of these is active
 #if PERMUTATION_SUPPORT_HEIGHT_FOG
