Quick fix for 34646590 ray gen shaders in HeterogeneousVolumesHardwareRayTracing

#rnx

diff --git a/Engine/Shaders/Private/HeterogeneousVolumes/HeterogeneousVolumesHardwareRayTracing.usf b/Engine/Shaders/Private/HeterogeneousVolumes/HeterogeneousVolumesHardwareRayTracing.usf
--- a/Engine/Shaders/Private/HeterogeneousVolumes/HeterogeneousVolumesHardwareRayTracing.usf
+++ b/Engine/Shaders/Private/HeterogeneousVolumes/HeterogeneousVolumesHardwareRayTracing.usf
@@ -210,8 +210,6 @@ int MipLevel;
 // Output
 RWTexture2D<float4> RWLightingTexture;
 
-#endif
-
 FRayDesc CreateVolumeRay(
 	float3 Origin,
 	float3 Direction,
@@ -600,3 +598,5 @@ RAY_TRACING_ENTRY_RAYGEN(RenderSingleScatteringWithPreshadingRGS)
 	// Output..
 	RWLightingTexture[PixelCoord] += float4(Radiance * View.PreExposure, Luminance(Transmittance));
 }
+
+#endif  // RAYGENSHADER
