VulkanRHI: Prevent an extra binding from showing up (EyeAdaptationBuffer) because compiler doesn't recognize a dead path in the code of FDiffuseIndirectCompositePS.

#rnx

diff --git a/Engine/Shaders/Private/DiffuseIndirectComposite.usf b/Engine/Shaders/Private/DiffuseIndirectComposite.usf
--- a/Engine/Shaders/Private/DiffuseIndirectComposite.usf
+++ b/Engine/Shaders/Private/DiffuseIndirectComposite.usf
@@ -59,6 +59,10 @@ RWTexture2D<float3> OutOpaqueRoughRefractionSceneColor;
 RWTexture2D<float3> OutSubSurfaceSceneColor;
 #endif
 
+// Remove a branch of code that is unused where TransformSignal() should always early exit.
+// SPIRV does not recognize that the code should be remove which adds an unnecessary parameter requirement.
+#define FORCE_IDENTICAL_COLOR_SPACE 1
+
 //------------------------------------------------------- INCLUDES
 
 #include "Common.ush"
diff --git a/Engine/Shaders/Private/ScreenSpaceDenoise/SSDSpatialKernel.ush b/Engine/Shaders/Private/ScreenSpaceDenoise/SSDSpatialKernel.ush
--- a/Engine/Shaders/Private/ScreenSpaceDenoise/SSDSpatialKernel.ush
+++ b/Engine/Shaders/Private/ScreenSpaceDenoise/SSDSpatialKernel.ush
@@ -57,6 +57,9 @@
 	#define COMPILE_RAW_EXPERIMENTAL_KERNEL 0
 #endif
 
+#ifndef FORCE_IDENTICAL_COLOR_SPACE
+	#define FORCE_IDENTICAL_COLOR_SPACE 0
+#endif
 
 //------------------------------------------------------- STRUCTURES
 
@@ -518,12 +521,14 @@ FSSDSignalSample TransformSignalSampleForAccumulation(
 	uint2 SamplePixelCoord)
 {
 	// Transform the color space.
+	#if (!FORCE_IDENTICAL_COLOR_SPACE)
 	// TODO(Denoiser): could pass down information that this sample may be normalized.
 	Sample = TransformSignal(
 		Sample, 
 		/* SrcBasis  = */ KernelConfig.BufferColorSpace[MultiplexId], 
 		/* DestBasis = */ KernelConfig.AccumulatorColorSpace[MultiplexId]);
-	
+	#endif
+
 	// Compute the spherical harmonic of the sample.
 	#if COMPILE_SIGNAL_COLOR_SH && COMPILE_SIGNAL_COLOR
 	if (KernelConfig.bComputeSampleColorSH)
