Many Lights - fixed split screen bug caused by a missing View offset in temporal accumulation pass

diff --git a/Engine/Shaders/Private/ManyLights/ManyLightsDenoiserTemporal.usf b/Engine/Shaders/Private/ManyLights/ManyLightsDenoiserTemporal.usf
--- a/Engine/Shaders/Private/ManyLights/ManyLightsDenoiserTemporal.usf
+++ b/Engine/Shaders/Private/ManyLights/ManyLightsDenoiserTemporal.usf
@@ -184,7 +184,7 @@ void DenoiserTemporalCS(
 	{
 		for (uint SharedCoordX = GroupThreadId.x; SharedCoordX < SHARED_TILE_SIZE; SharedCoordX += THREADGROUP_SIZE)
 		{
-			int2 ResolvedCoord = GroupId.xy * THREADGROUP_SIZE + int2(SharedCoordX, SharedCoordY) - SHARED_TILE_BORDER;
+			int2 ResolvedCoord = GroupId.xy * THREADGROUP_SIZE + View.ViewRectMinAndSize.xy + int2(SharedCoordX, SharedCoordY) - SHARED_TILE_BORDER;
 			
 			float3 ResolvedDiffuseYCoCg = INVALID_LIGHTING;
 			float3 ResolvedSpecularYCoCg = INVALID_LIGHTING;
