Fixes a ghosting artifacts on the alpha channel with HasPixelAnimation=1

diff --git a/Engine/Shaders/Private/TemporalSuperResolution/TSRShadingAnalysis.ush b/Engine/Shaders/Private/TemporalSuperResolution/TSRShadingAnalysis.ush
--- a/Engine/Shaders/Private/TemporalSuperResolution/TSRShadingAnalysis.ush
+++ b/Engine/Shaders/Private/TemporalSuperResolution/TSRShadingAnalysis.ush
@@ -47,12 +47,18 @@ tsr_tensor_halfC ComposeTranslucencyForRejection(tsr_tensor_halfC OriginalOpaque
 	tsr_tensor_half4 BlurInput = OriginalTranslucencyInput;
 
 	{
+		#if CONFIG_SCENE_COLOR_ALPHA
+			const tsr_half4 TransparentSharpInput = tsr_half4(0.0, 0.0, 0.0, 1.0);
+		#else
+			const tsr_half3 TransparentSharpInput = tsr_half(0.0).xxx;
+		#endif
+
 		tsr_tensor_halfC CenterFinalSceneColor = ComposeTranslucency(OriginalOpaqueInput, OriginalTranslucencyInput);
 			
 		tsr_tensor_half4 HasAnimationBlurInput = ResizeChannels<4>(CenterFinalSceneColor);
 			
 		BlurInput  = select(tsr_tensor_bool4::Vectorize(bHasPixelAnimation), HasAnimationBlurInput, BlurInput);
-		SharpInput = select(tsr_tensor_boolC::Vectorize(bHasPixelAnimation), tsr_tensor_halfC::Const(0.0), SharpInput);
+		SharpInput = select(tsr_tensor_boolC::Vectorize(bHasPixelAnimation), tsr_tensor_halfC::Const(TransparentSharpInput), SharpInput);
 	}
 		
 	return ComposeTranslucency(SharpInput, Blur3x3(BlurInput));
