Applies the reprojection field's upscaling factor to the history rejection to take down history reprojection upscaling artifact on animations

diff --git a/Engine/Shaders/Private/TemporalSuperResolution/TSRUpdateHistory.usf b/Engine/Shaders/Private/TemporalSuperResolution/TSRUpdateHistory.usf
--- a/Engine/Shaders/Private/TemporalSuperResolution/TSRUpdateHistory.usf
+++ b/Engine/Shaders/Private/TemporalSuperResolution/TSRUpdateHistory.usf
@@ -601,6 +601,7 @@ void MainCS(
 	tsr_half2 DecreaseValidityMultiplier = tsr_half(0.0).xx;
 	tsr_half2 OutputPixelVelocity = tsr_half(0.0).xx;
 	tsr_half2 NoiseFiltering = tsr_half(0.0).xx;
+	tsr_half2 ReprojectionUpscaleCorrection = tsr_half(1.0).xx;
 	bool2 bIsParallaxRejected = false;
 	#if CONFIG_HISTORY_RESURRECTION
 		bool2 bIsHistoryResurrection = false;
@@ -699,9 +700,12 @@ void MainCS(
 						float2 ReprojectionPixelPosCorrection = float2(mul(JacobianCoordinate[PixelId], ReprojectionJacobian));
 						//ReprojectionScreenPosCorrection = ReprojectionPixelPosCorrection * (float2(2, -2) * InputInfo_ViewportSizeInverse);
 						ReprojectionScreenPosCorrection = ReprojectionPixelPosCorrection * InputPixelPosToScreenPos.xy;
+
+						ReprojectionUpscaleCorrection[PixelId] = rcp(max(ComputeReprojectionUpscaleFactorFromJacobian(ReprojectionJacobian), 1.0));
 					}
 					else
 					{
+						ReprojectionUpscaleCorrection[PixelId] = tsr_half(1.0);
 						ReprojectionScreenPosCorrection = float(0.0);
 					}
 					#else // !CONFIG_REPROJECTION_FIELD
@@ -988,7 +992,7 @@ void MainCS(
 				const tsr_half ActualHistorySampleCount = tsr_half(HistorySampleCount);
 				const tsr_half ActualHistoryHisteresis = tsr_half(HistoryHisteresis);
 
-				tsr_half2 MinRejectionBlendFactor = RejectionFactorToBlendFactor(LowFrequencyRejection);
+				tsr_half2 MinRejectionBlendFactor = RejectionFactorToBlendFactor(min(LowFrequencyRejection, ReprojectionUpscaleCorrection));
 				MinRejectionBlendFactor = select(bIsOffScreen, tsr_half(1.0).xx, MinRejectionBlendFactor);
 
 				tsr_half2 CoarseInputPixelAlignement = ComputeSampleWeigth(tsr_half(1.0), dInputKO, /* MinimalContribution = */ float(0.0));
