Improve NFOR denoiser quality:
1. Improve scene quality when the color is most red or blue. Change variance of luminance to that of the channel ranked lumiance so that the variance of (X, 0, 0) is equivilent to (0,X,0) and (0,0,X).
2. Preserve more details. Filter each MSE map with the corresponding bandwidth used to filter the image.

#tests The dawn of RenderGallary project is finally denoisable, other scene in hillside project also gets quality improvement.

diff --git a/Engine/Plugins/Experimental/NFORDenoise/Source/NFORDenoise/Private/NFORDenoiseCS.cpp b/Engine/Plugins/Experimental/NFORDenoise/Source/NFORDenoise/Private/NFORDenoiseCS.cpp
--- a/Engine/Plugins/Experimental/NFORDenoise/Source/NFORDenoise/Private/NFORDenoiseCS.cpp
+++ b/Engine/Plugins/Experimental/NFORDenoise/Source/NFORDenoise/Private/NFORDenoiseCS.cpp
@@ -246,6 +246,20 @@ namespace NFORDenoise
 		TEXT("(0,1]: Use a specific bandwidth."),
 		ECVF_RenderThreadSafe);
 
+	TAutoConsoleVariable<bool> CVarNFORBandwidthSelectionMSEPreserveDetail(
+		TEXT("r.NFOR.BandwidthSelection.MSE.PreserveDetail"),
+		true,
+		TEXT("false: Use bandwidth = 1.0 to filter MSE.")
+		TEXT("true: Use the corresponding bandwidth to filter MSE"),
+		ECVF_RenderThreadSafe);
+
+	TAutoConsoleVariable<bool> CVarNFORBandwidthSelectionMapPreserveDetail(
+		TEXT("r.NFOR.BandwidthSelection.Map.PreserveDetail"),
+		false,
+		TEXT("false: Use bandwidth = 1.0 to filter MSE.")
+		TEXT("true: Use the corresponding bandwidth to filter MSE"),
+		ECVF_RenderThreadSafe);
+
 	// Working in progress
 	TAutoConsoleVariable<int32> CVarNFORAlbedoDivideRecoverPhase(
 		TEXT("r.NFOR.AlbedoDivide.RecoverPhase"),
@@ -591,6 +605,17 @@ namespace NFORDenoise
 		return Bandwidths;
 	}
 
+	bool ShouldBandwidthSelectionMSEPreserveDetail()
+	{
+		return CVarNFORBandwidthSelectionMSEPreserveDetail.GetValueOnRenderThread();
+	}
+
+	bool ShouldBandwidthSelectionMapPreserveDetail()
+	{
+		return CVarNFORBandwidthSelectionMapPreserveDetail.GetValueOnRenderThread();
+	}
+
+
 	EPixelFormat GetFeaturePixelFormat()
 	{
 		EPixelFormat PixelFormat = PF_R32_FLOAT;
@@ -2729,7 +2754,8 @@ namespace NFORDenoise
 
 				// NLM filtering of MSE texture.
 				FRDGTextureRef FilteredMSETexure = GraphBuilder.CreateTexture(MSE.Image->Desc, TEXT("NFOR.FilteredMSE"));
-				FNonLocalMeanParameters MSENonLocalMeanParameters = GetNonLocalMeanParameters(1, RadiancePatchDistance, 1.0f);
+				FNonLocalMeanParameters MSENonLocalMeanParameters = GetNonLocalMeanParameters(1, RadiancePatchDistance, 
+					ShouldBandwidthSelectionMSEPreserveDetail() ? Bandwidths[i]: 1.0f);
 
 				ApplyNonLocalMeanFilterIfRequired(
 					GraphBuilder,
@@ -2780,7 +2806,8 @@ namespace NFORDenoise
 			// Filter the selection map with image variance
 			FRDGTextureRef FilteredSelectionMap = GraphBuilder.CreateTexture(Desc, TEXT("NFOR.FilteredSelectionMap"));
 			{
-				FNonLocalMeanParameters SelectionMapNonLocalMeanParameters = GetNonLocalMeanParameters(1, RadiancePatchDistance, 1.0f);
+				FNonLocalMeanParameters SelectionMapNonLocalMeanParameters = GetNonLocalMeanParameters(1, RadiancePatchDistance, 
+					ShouldBandwidthSelectionMapPreserveDetail() ? Bandwidths[0] : 1.0f);
 
 				ApplyNonLocalMeanFilterIfRequired(
 					GraphBuilder,
diff --git a/Engine/Shaders/Private/PathTracing/PathTracingSpatialTemporalDenoising.usf b/Engine/Shaders/Private/PathTracing/PathTracingSpatialTemporalDenoising.usf
--- a/Engine/Shaders/Private/PathTracing/PathTracingSpatialTemporalDenoising.usf
+++ b/Engine/Shaders/Private/PathTracing/PathTracingSpatialTemporalDenoising.usf
@@ -50,6 +50,10 @@ PathTracingSpatialTemporalDenoising.usf: Spatial temporal denoising for path tra
 #define PREPASS_PHASE_INIT				0
 #define PREPASS_PHASE_UPDATE			1
 
+#ifndef RANKED_LUMINANCE_VARIANCE
+#define RANKED_LUMINANCE_VARIANCE		0
+#endif
+
 #ifndef SPATIAL_DENOISING
 #define SPATIAL_DENOISING				0
 #endif
@@ -133,6 +137,12 @@ struct FPixelMaterialLightingFingerprint
 	float4 Var;
 };
 
+void Max3Index(in float3 Value, out int MaxIndex)
+{
+	MaxIndex = (Value.x > Value.y) ? 0 : 1;
+	MaxIndex = Value[MaxIndex] > Value.z ? MaxIndex : 2;
+}
+
 void Max3(in float3 Value, out float MaxValue, out int MaxIndex)
 {
 	MaxIndex = (Value.x > Value.y) ? 0 : 1;
@@ -140,6 +150,19 @@ void Max3(in float3 Value, out float MaxValue, out int MaxIndex)
 	MaxValue = Value[MaxIndex];
 }
 
+void Min3Index(in float3 Value, out int MinIndex)
+{
+	MinIndex = (Value.x <= Value.y) ? 0 : 1;
+	MinIndex = Value[MinIndex] <= Value.z ? MinIndex : 2;
+}
+
+void Min3(in float3 Value, out float MinValue, out int MinIndex)
+{
+	MinIndex = (Value.x <= Value.y) ? 0 : 1;
+	MinIndex = Value[MinIndex] <= Value.z ? MinIndex : 2;
+	MinValue = Value[MinIndex];
+}
+
 #if K_NUM_OF_TEXTURES_PER_PASS <= 7 && (TEMPORAL_REPROJECTION_ALIGN||TEMPORAL_REPROJECTION_MERGE)
 
 static int2 Shifts[25] = {
@@ -262,13 +285,6 @@ struct FFindMinContext
 	float Value;
 };
 
-void Min3(in float3 Value, out float MinValue, out int MinIndex)
-{
-	MinIndex = (Value.x <= Value.y) ? 0 : 1;
-	MinIndex = Value[MinIndex] <= Value.z ? MinIndex : 2;
-	MinValue = Value[MinIndex];
-}
-
 void Min4(in float4 Value, out float MinValue, out int MinIndex)
 {
 	int MinIndexXy = (Value.x <= Value.y) ? 0 : 1;
@@ -792,6 +808,20 @@ Texture2D<float4> NormalTexture;
 RWStructuredBuffer<FPixelMaterialLightingFingerprint> RWVarianceMap;
 int Iteration;
 
+// Channel ranked luminance places the max value to g, and minimal to b
+// In this way, the luminance is perception independent. The variance of (X,0,0) is the same to (0,X,0)
+// instead of being smaller. 
+float3 ChannelRankedLuminance(float3 Value)
+{
+	int Index = 0;
+	Max3Index(Value, Index);
+	Swap(Value[Index], Value.g); // swap max value to g
+	Min3Index(Value, Index);
+	Swap(Value[Index], Value.b); // swap min value to b
+
+	return Luminance(Value);
+}
+
 // Note that the radiance, albedo and normal are the mean value.
 float4 GetMeanFromTextures(int2 Position)
 {
@@ -799,9 +829,14 @@ float4 GetMeanFromTextures(int2 Position)
 #if VARIANCE_TYPE == RADIANCE_MULTI_CHANNEL
 	Value  = InputTexture[Position];
 #elif VARIANCE_TYPE == RADIANCE_ALBEDO_NORMAL_SINGLECHANNEL
+
+#if RANKED_LUMINANCE_VARIANCE
+	Value.x = ChannelRankedLuminance(InputTexture[Position].rgb);
+	Value.y = ChannelRankedLuminance(AlbedoTexture[Position].rgb);
+#else
 	Value.x = Luminance(InputTexture[Position].rgb);
 	Value.y = Luminance(AlbedoTexture[Position].rgb);
-	
+#endif
 	//The length of the average normal indicates the variance
 	//1: All pointing to the same direction
 	//0: Uniform in all directions
diff --git a/Engine/Source/Runtime/Renderer/Private/PathTracingSpatialTemporalDenoising.cpp b/Engine/Source/Runtime/Renderer/Private/PathTracingSpatialTemporalDenoising.cpp
--- a/Engine/Source/Runtime/Renderer/Private/PathTracingSpatialTemporalDenoising.cpp
+++ b/Engine/Source/Runtime/Renderer/Private/PathTracingSpatialTemporalDenoising.cpp
@@ -164,6 +164,15 @@ namespace {
 		TEXT("1: Output variance texture to denoisers, or the postprocess material usually used by MRQ")
 	);
 
+	TAutoConsoleVariable<int32> CVarPathTracingDenoiserPrepassRankedLuminanceVariance(
+		TEXT("r.PathTracing.Denoiser.Prepass.RankedLuminanceVariance"),
+		1,
+		TEXT("Select the luminance type when calculating the variance:\n")
+		TEXT("0: use default luminance to estimate variance.\n")
+		TEXT("1: Use channel ranked luminance when calculating variance.\n"),
+		ECVF_RenderThreadSafe
+	);
+
 	TAutoConsoleVariable<int32> CVarPathTracingSpatialDenoiser(
 		TEXT("r.PathTracing.SpatialDenoiser"),
 		1,
@@ -1796,7 +1805,13 @@ public:
 
 	class FPrepassPhase : SHADER_PERMUTATION_BOOL("PREPASS_PHASE");	// 0: initialize, 1: update
 	class FVarianceType : SHADER_PERMUTATION_ENUM_CLASS("VARIANCE_TYPE", EVarianceType);
-	using FPermutationDomain = TShaderPermutationDomain<FPrepassPhase,FVarianceType>;
+	class FRankedLuminanceVariance : SHADER_PERMUTATION_BOOL("RANKED_LUMINANCE_VARIANCE");
+	using FPermutationDomain = TShaderPermutationDomain<FPrepassPhase,FVarianceType, FRankedLuminanceVariance>;
+
+	static bool UseRankedLuminanceVariance()
+	{
+		return CVarPathTracingDenoiserPrepassRankedLuminanceVariance.GetValueOnRenderThread() != 0;
+	}
 
 	static EVarianceType GetVarianceType()
 	{
@@ -1896,6 +1911,7 @@ void PathTracingSpatialTemporalDenoisingPrePass(FRDGBuilder& GraphBuilder, const
 			SHADER::FPermutationDomain ComputeShaderPermutationVector;
 			ComputeShaderPermutationVector.Set<SHADER::FPrepassPhase>(bUpdateVarianceMapPhase);
 			ComputeShaderPermutationVector.Set<SHADER::FVarianceType>(SHADER::GetVarianceType());
+			ComputeShaderPermutationVector.Set<SHADER::FRankedLuminanceVariance>(SHADER::UseRankedLuminanceVariance());
 
 			TShaderMapRef<SHADER> ComputeShader(View.ShaderMap, ComputeShaderPermutationVector);
 			FComputeShaderUtils::AddPass(
