Fixes quality regression on shading change quality around geometric discontinuity with the reprojection field

diff --git a/Engine/Shaders/Private/TemporalSuperResolution/TSRUpdateHistory.usf b/Engine/Shaders/Private/TemporalSuperResolution/TSRUpdateHistory.usf
--- a/Engine/Shaders/Private/TemporalSuperResolution/TSRUpdateHistory.usf
+++ b/Engine/Shaders/Private/TemporalSuperResolution/TSRUpdateHistory.usf
@@ -625,6 +625,7 @@ void MainCS(
 		#if CONFIG_REPROJECTION_FIELD
 			uint RawEncodedReprojectionJacobian[DPV_PIXEL_PER_LANE];
 			tsr_half2 JacobianCoordinate[DPV_PIXEL_PER_LANE];
+			tsr_half4 RawDilatedHistoryRejection[DPV_PIXEL_PER_LANE];
 		#endif
 
 		// Issue overlapped texture fetches
@@ -668,9 +669,10 @@ void MainCS(
 				RawEncodedReprojectionVector[PixelId]   = ReprojectionVectorTexture[LocalInputPixelPosWithReprojectionAA];
 				#if CONFIG_REPROJECTION_FIELD
 					RawEncodedReprojectionJacobian[PixelId] = ReprojectionJacobianTexture[LocalInputPixelPosWithReprojectionAA];
+					RawDilatedHistoryRejection[PixelId] = HistoryRejectionTexture[LocalInputPixelPosWithReprojectionAA];
 				#endif
 
-				RawHistoryRejection[PixelId] = HistoryRejectionTexture[LocalInputPixelPosWithReprojectionAA];
+				RawHistoryRejection[PixelId] = HistoryRejectionTexture[LocalInputPixelPos];
 				
 				#if CONFIG_REJECTION_ANTI_ALIASING
 				{
@@ -725,7 +727,11 @@ void MainCS(
 			
 			// Unpacks bits masks
 			{
-				tsr_ushort2 BitMask = tsr_ushort2(round(dpv_interleave_registers_array(RawHistoryRejection)[3] * tsr_half(255.0)));
+				#if CONFIG_REPROJECTION_FIELD
+					tsr_ushort2 BitMask = tsr_ushort2(round(dpv_interleave_registers_array(RawDilatedHistoryRejection)[3] * tsr_half(255.0)));
+				#else
+					tsr_ushort2 BitMask = tsr_ushort2(round(dpv_interleave_registers_array(RawHistoryRejection)[3] * tsr_half(255.0)));
+				#endif
 				bIsParallaxRejected = (BitMask & tsr_ushort(0x1)) == tsr_ushort(0);
 				#if CONFIG_HISTORY_RESURRECTION
 					bIsHistoryResurrection = (BitMask & tsr_ushort(0x2)) != tsr_ushort(0);
@@ -736,6 +742,14 @@ void MainCS(
 			LowFrequencyRejection      = dpv_force_interleave_registers_array(RawHistoryRejection)[0];
 			DisableHistoryClamp        = dpv_force_interleave_registers_array(RawHistoryRejection)[1];
 			DecreaseValidityMultiplier = dpv_force_interleave_registers_array(RawHistoryRejection)[2];
+			
+			#if CONFIG_REPROJECTION_FIELD
+			{
+				LowFrequencyRejection      = min(LowFrequencyRejection     , dpv_force_interleave_registers_array(RawDilatedHistoryRejection)[0]);
+				DisableHistoryClamp        = min(DisableHistoryClamp       , dpv_force_interleave_registers_array(RawDilatedHistoryRejection)[1]);
+				DecreaseValidityMultiplier = max(DecreaseValidityMultiplier, dpv_force_interleave_registers_array(RawDilatedHistoryRejection)[2]);
+			}
+			#endif
 		}
 
 		#if CONFIG_REJECTION_ANTI_ALIASING
