Path Tracer: Fix minor shader warning in denoiser

#rnx

diff --git a/Engine/Shaders/Private/PathTracing/PathTracingSpatialTemporalDenoising.usf b/Engine/Shaders/Private/PathTracing/PathTracingSpatialTemporalDenoising.usf
--- a/Engine/Shaders/Private/PathTracing/PathTracingSpatialTemporalDenoising.usf
+++ b/Engine/Shaders/Private/PathTracing/PathTracingSpatialTemporalDenoising.usf
@@ -811,7 +811,7 @@ int Iteration;
 // Channel ranked luminance places the max value to g, and minimal to b
 // In this way, the luminance is perception independent. The variance of (X,0,0) is the same to (0,X,0)
 // instead of being smaller. 
-float3 ChannelRankedLuminance(float3 Value)
+float ChannelRankedLuminance(float3 Value)
 {
 	int Index = 0;
 	Max3Index(Value, Index);
