Fixed compilation issues after 34421615 when PROJECT_SUPPORT_SKY_ATMOSPHERE_AFFECTS_HEIGHFOG is 0.

diff --git a/Engine/Shaders/Private/HeightFogCommon.ush b/Engine/Shaders/Private/HeightFogCommon.ush
--- a/Engine/Shaders/Private/HeightFogCommon.ush
+++ b/Engine/Shaders/Private/HeightFogCommon.ush
@@ -338,7 +338,7 @@ half4 GetExponentialHeightFog(float3 WorldPositionRelativeToCamera, float Exclud
 		// Setup a cosine lobe around the light direction to approximate inscattering from the directional light off of the ambient haze;
 		half3 DirectionalLightInscattering = FogData.DirectionalInscatteringColor.xyz
 	#if PROJECT_EXPFOG_MATCHES_VFOG
-			* HenyeyGreensteinPhase(FogData.VolumetricFogPhaseG, dot(-CameraToReceiverNormalized, InView.AtmosphereLightDirection.xyz));
+			* HenyeyGreensteinPhase(FogData.VolumetricFogPhaseG, dot(-CameraToReceiverNormalized, InView.AtmosphereLightDirection[0].xyz));
 	#else
 			* pow(saturate(dot(CameraToReceiverNormalized, FogData.InscatteringLightDirection.xyz)), FogData.DirectionalInscatteringColor.w);
 	#endif
