Sunsets the parallax's HistoryUseCount atomic texture technic in favor of a DepthErrorTexture computed from the depth neighborhood

This notably:
* Have the same memory bandwidth of 1byte per rendering pixel in Dilate/Decimate passes.
* Saves 3 texture fetches per rendering pixel in the decimate pass.
* Saves LoadPrevAtomicTexturesSamples()'s register pressure.
* Fixes a bug where the parallax disocclusion could be detected on flat roads at grazing angle at great speeds.
* Fixes a bug where the reprojection hole filling could end up on moving non-disoccluded area as the history use count there was always > 1, introducing reprojection alignment error due its lower bitdepth, generating unnecessary bluring in motion.

This changes takes oportunity to further improve in motion clarity with the following:
* Changes ComputeReprojectionUpscaleFactorFromJacobian() to be square to further reject history where the reprojection is getting dilated both on x and y axes.
* Reduce the r.TSR.ReprojectionField.AntiAliasPixelSpeed's default from 0.75 to 0.125, thanks to manually clamping it 0.5 on disoccluding geometric discontinuity. It increases the costs of the depth buffer anti-aliasing search, but that is negliegable compared to the extra dependent texture fetch in HistoryUpdate.

diff --git a/Engine/Shaders/Private/TemporalSuperResolution/TSRClearPrevTextures.usf b/Engine/Shaders/Private/TemporalSuperResolution/TSRClearPrevTextures.usf
--- a/Engine/Shaders/Private/TemporalSuperResolution/TSRClearPrevTextures.usf
+++ b/Engine/Shaders/Private/TemporalSuperResolution/TSRClearPrevTextures.usf
@@ -33,22 +33,13 @@ void MainCS(
 		{
 			uint2 PixelOffset = uint2(x, y) * TILE_SIZE;
 			uint2 InputPixelPos = OutputPixelOrigin + PixelOffset;
-			uint2 PrevUseCountPixelPos = InputPixelPos;
 
 			InputPixelPos.x = select(all(InputPixelPos < InputInfo_ViewportMax), InputPixelPos.x, uint(~0));
 
-			#if COMPRESS_PREV_USE_COUNT
-				const uint2 PrevUseCountMaxPixelPos = InputInfo_ViewportMin + uint2(InputInfo_ViewportMax - InputInfo_ViewportMin + 1) / 2 + 8;
-			#else
-				const uint2 PrevUseCountMaxPixelPos = InputInfo_ViewportMax;
-			#endif
-			PrevUseCountPixelPos.x = select(all(PrevUseCountPixelPos < PrevUseCountMaxPixelPos), PrevUseCountPixelPos.x, uint(~0));
-
-			PrevAtomicOutput[uint3(PrevUseCountPixelPos, 0)] = 0;
-			PrevAtomicOutput[uint3(InputPixelPos, 1)] = 0;
+			PrevAtomicOutput[uint3(InputPixelPos, 0)] = 0;
 			
 			#if CONFIG_SUPPORT_ORTHO_VIEW
-				PrevAtomicOutput[uint3(InputPixelPos, 2)] = 0;
+				PrevAtomicOutput[uint3(InputPixelPos, 1)] = 0;
 			#endif
 		}
 	}
diff --git a/Engine/Shaders/Private/TemporalSuperResolution/TSRCommon.ush b/Engine/Shaders/Private/TemporalSuperResolution/TSRCommon.ush
--- a/Engine/Shaders/Private/TemporalSuperResolution/TSRCommon.ush
+++ b/Engine/Shaders/Private/TemporalSuperResolution/TSRCommon.ush
@@ -71,18 +71,10 @@
 #define SPATIAL_ANTI_ALIASING_OFFSET_RANGE 0.5
 #define SPATIAL_ANTI_ALIASING_OFFSET_BITDEPTH_PER_CHANNEL 4
 
-#define COMPRESS_PREV_USE_COUNT 1
-#if COMPRESS_PREV_USE_COUNT
-	// COMPRESS_PREV_USE_COUNT compress the use count in 8bit, so need to sum of a 2x2 to not overflow these 8bits
-	#define PREV_USE_COUNT_QUANTIZATION 63
-#else
-	#define PREV_USE_COUNT_QUANTIZATION 255
-#endif
+#define MAX_DEVICE_Z_ERROR 16.0
 
 #define PARALLAX_REJECTION_MASK_THRESHOLD 0.5
 
-#define MAX_PARALLAX_FACTOR 8.0
-
 // Minimum number of sample of each texel when reprojecting the high frequency history.
 // Value too low causes area of low sample count to have pixel dominating too much in reprojection causing input res aliasing on texture details
 #define MIN_REPROJECTION_SAMPLE_COUNT 2.0
@@ -666,74 +658,6 @@ tsr_half ComputeLowFrequencyContribution(tsr_half HistoryValidity)
 	return dpv_lo(ComputeLowFrequencyContribution(dpv_interleave_mono_registers(HistoryValidity)));
 }
 
-CALL_SITE_DEBUGLOC
-tsr_half3x2 RGBToYCoCg(tsr_half3x2 RGB)
-{
-    tsr_half2 R = RGB[0];
-    tsr_half2 G = RGB[1];
-    tsr_half2 B = RGB[2];       
-
-    tsr_half2 RplusB = R + B;
-    tsr_half2 Y = 2.0 * G + RplusB;
-    tsr_half2 Cg = 2.0 * G - RplusB;
-    tsr_half2 Co = 2.0 * (R - B);
-    tsr_half3x2 YCoCg = tsr_half3x2(Y, Co, Cg);
-
-    return YCoCg;
-}
-
-CALL_SITE_DEBUGLOC
-tsr_half3x2 YCoCgToRGB(tsr_half3x2 YCoCg)
-{
-	tsr_half2 Y =  YCoCg[0] * tsr_half(0.25);
-	tsr_half2 Co = YCoCg[1] * tsr_half(0.25);
-	tsr_half2 Cg = YCoCg[2] * tsr_half(0.25);
-
-	tsr_half2 R = Y + Co - Cg;
-	tsr_half2 G = Y + Cg;
-	tsr_half2 B = Y - Co - Cg;
-
-	tsr_half3x2 RGB = tsr_half3x2(R, G, B);
-	return RGB;
-}
-
-CALL_SITE_DEBUGLOC
-tsr_half4x2 RGBToYCoCg(tsr_half4x2 RGBA)
-{
-	return tsr_half4x2(RGBToYCoCg(tsr_half3x2(RGBA[0], RGBA[1], RGBA[2])), RGBA[3]);
-}
-
-CALL_SITE_DEBUGLOC
-tsr_half4x2 YCoCgToRGB(tsr_half4x2 YCoCgA)
-{
-	return tsr_half4x2(YCoCgToRGB(tsr_half3x2(YCoCgA[0], YCoCgA[1], YCoCgA[2])), YCoCgA[3]);
-}
-
-CALL_SITE_DEBUGLOC
-tsr_half3 RGBToYCoCg(tsr_half3 RGB)
-{
-	return dpv_lo(RGBToYCoCg(dpv_interleave_mono_registers(RGB)));
-}
-
-CALL_SITE_DEBUGLOC
-tsr_half3 YCoCgToRGB(tsr_half3 YCoCg)
-{
-	return dpv_lo(YCoCgToRGB(dpv_interleave_mono_registers(YCoCg)));
-}
-
-CALL_SITE_DEBUGLOC
-tsr_half4 RGBToYCoCg(tsr_half4 RGBA)
-{
-	return dpv_lo(RGBToYCoCg(dpv_interleave_mono_registers(RGBA)));
-}
-
-CALL_SITE_DEBUGLOC
-tsr_half4 YCoCgToRGB(tsr_half4 YCoCgA)
-{
-	return dpv_lo(YCoCgToRGB(dpv_interleave_mono_registers(YCoCgA)));
-}
-
-
 // Some bright pixel can cause HdrWeight to get nullified under fp16 representation. So clamping this to a value close to the minimum float float positive value (0.000061).
 #define HDR_WEIGHT_SAFE_MIN_VALUE 0.0001
 
@@ -877,16 +801,6 @@ tsr_half ComputePredictionCompleteness(tsr_half SampleHistoryValidity)
 	return saturate(SampleHistoryValidity * tsr_half(MAX_SAMPLE_COUNT) - tsr_half(0.2));
 }
 
-CALL_SITE_DEBUGLOC
-bool IsOffScreen(uint bCameraCut, float2 ScreenPos, tsr_half ParallaxRejectionMask)
-{
-	bool bIsCameraCut = bCameraCut != 0;
-	bool bIsOutOfBounds = max(abs(ScreenPos.x), abs(ScreenPos.y)) >= 1.0;
-	bool bIsParallaxRejected = ParallaxRejectionMask < tsr_half(PARALLAX_REJECTION_MASK_THRESHOLD);
-
-	return (bIsCameraCut || bIsOutOfBounds || bIsParallaxRejected);
-}
-
 CALL_SITE_DEBUGLOC
 void IsOffScreenOrDisoccluded(uint bCameraCut, float2x2 ScreenPos, bool2 bIsParallaxRejected, out bool2 bIsOffScreen, out bool2 bIsDisoccluded)
 {
@@ -916,57 +830,6 @@ tsr_half MeasureBackbufferLDRQuantizationErrorFromHDRLuma(tsr_half Luma)
 	return dpv_lo(MeasureBackbufferLDRQuantizationErrorFromHDRLuma(dpv_interleave_mono_registers(Luma)));
 }
 
-tsr_half2 MeasureRejectionFactor(
-	tsr_halfCx2 PrevYCoCg,
-	tsr_halfCx2 ClampedPrevYCoCg,
-	tsr_halfCx2 InputCenterYCoCg,
-	tsr_halfCx2 InputMinYCoCg, tsr_halfCx2 InputMaxYCoCg,
-	tsr_half2 BackbufferQuantizationError)
-{
-	const tsr_halfC Null = tsr_half(0.0);
-	
-	#if CONFIG_SCENE_COLOR_ALPHA
-		tsr_halfCx2 BackbufferQuantizationErrorVector = tsr_half4x2(BackbufferQuantizationError, BackbufferQuantizationError, BackbufferQuantizationError, BackbufferQuantizationError);
-	#else
-		tsr_halfCx2 BackbufferQuantizationErrorVector = tsr_half3x2(BackbufferQuantizationError, BackbufferQuantizationError, BackbufferQuantizationError);
-	#endif
-
-	tsr_halfCx2 BoxSize = InputMaxYCoCg - InputMinYCoCg;
-
-	tsr_halfCx2 ClampedEnergy = abs(ClampedPrevYCoCg - PrevYCoCg);
-	tsr_halfCx2 Delta = abs(InputCenterYCoCg - PrevYCoCg);
-
-	tsr_halfCx2 ClampError = max(
-		BoxSize * tsr_half(0.25),
-		BackbufferQuantizationErrorVector);
-
-	tsr_halfCx2 Factor = tsr_half(1.0) - saturate(max(ClampedEnergy - ClampError, dpv_interleave_mono_registers(Null)) * rcp(max(Delta, BoxSize)));
-
-	tsr_half2 Rejection = min3(Factor[0], Factor[1], Factor[2]);
-	
-	#if CONFIG_SCENE_COLOR_ALPHA
-		Rejection = min(Rejection, Factor[3]);
-	#endif
-
-	return Rejection;
-}
-
-tsr_half MeasureRejectionFactor(
-	tsr_halfC PrevYCoCg,
-	tsr_halfC ClampedPrevYCoCg,
-	tsr_halfC InputCenterYCoCg,
-	tsr_halfC InputMinYCoCg, tsr_halfC InputMaxYCoCg,
-	tsr_half BackbufferQuantizationError)
-{
-	return dpv_lo(MeasureRejectionFactor(
-		dpv_interleave_mono_registers(PrevYCoCg), 
-		dpv_interleave_mono_registers(ClampedPrevYCoCg), 
-		dpv_interleave_mono_registers(InputCenterYCoCg), 
-		dpv_interleave_mono_registers(InputMinYCoCg), 
-		dpv_interleave_mono_registers(InputMaxYCoCg), 
-		dpv_interleave_mono_registers(BackbufferQuantizationError)));
-}
-
 CALL_SITE_DEBUGLOC
 tsr_half3x2 QuantizeForUNormRenderTarget(tsr_half3x2 Color, tsr_half E, const float3 QuantizationError)
 {
@@ -1009,23 +872,6 @@ tsr_half3 QuantizeForFloatRenderTarget(tsr_half3 Color, tsr_half E, const float3
 	return dpv_lo(QuantizeForFloatRenderTarget(dpv_interleave_mono_registers(Color), E, QuantizationError));
 }
 
-
-#if COMPRESS_PREV_USE_COUNT
-
-// encode 4 use count that are 8bits shifted in the PF_R32_UINT
-void ComputeCompressedUseCountPixelCoordinates(uint2 PixelPos, out uint2 CompressedUseCountPixelPos, out uint Shift)
-{
-	// Overlap 2x2 consecutive blocks of 8x8 pixel to reduce atomic contention in the scattering pass
-	const uint kTileSize = 8;
-
-	uint2 PixelPosOffset = PixelPos - InputInfo_ViewportMin;
-
-	CompressedUseCountPixelPos = ((PixelPosOffset % kTileSize) | ((PixelPosOffset >> 1) & ~(kTileSize - 1))) + InputInfo_ViewportMin;
-	Shift = ((PixelPosOffset.x / kTileSize) % 2) | (((PixelPosOffset.y / kTileSize) % 2) << 1);
-}
-	
-#endif
-
 CALL_SITE_DEBUGLOC
 uint2 EncodeSpatialAntiAliasingOffset(tsr_half2x2 TexelOffset)
 {
@@ -1268,3 +1114,30 @@ FCatmullRomSamples GetBicubic2DCatmullRomSamples_Stubbe(float2 UV, float2 Size,
 	return Samples;
 }
 
+/** Compute the devize Z error due to the size of a pixel in world space at that given depth. */
+float ComputePixelDeviceZError(float DeviceZ)
+{
+	float Depth = ConvertFromDeviceZ(DeviceZ);
+	float ErrorCorrectedDeviceZ = ConvertToDeviceZ(GetDepthPixelRadiusForProjectionType(Depth) * 2.0f + Depth);
+	float DeviceZError = abs(ErrorCorrectedDeviceZ - DeviceZ);
+	return DeviceZError;
+}
+
+/** Encode the devize Z error relatively to the closest depth to squeeze to 8bit per pixel */
+uint EncodeDeviceZError(float PrevClosestDeviceZ, float ParallaxDepthError)
+{
+	float PixelDeviceZError = float(MAX_DEVICE_Z_ERROR) * ComputePixelDeviceZError(PrevClosestDeviceZ);
+	float RelativeZError = ParallaxDepthError / PixelDeviceZError;
+	
+	uint EncodedDeviceZError = uint(ceil(saturate(RelativeZError) * 255.0));
+	return EncodedDeviceZError;
+}
+
+/** Inverse of EncodeDeviceZError() */
+float DecodeDeviceZError(float PrevClosestDeviceZ, uint EncodedDeviceZError)
+{
+	float PixelDeviceZError = ComputePixelDeviceZError(PrevClosestDeviceZ);
+	float RelativeZError = float(EncodedDeviceZError) * (float(MAX_DEVICE_Z_ERROR) / 255.0);
+	float ParallaxDepthError = PixelDeviceZError * RelativeZError;
+	return ParallaxDepthError;
+}
diff --git a/Engine/Shaders/Private/TemporalSuperResolution/TSRDecimateHistory.usf b/Engine/Shaders/Private/TemporalSuperResolution/TSRDecimateHistory.usf
--- a/Engine/Shaders/Private/TemporalSuperResolution/TSRDecimateHistory.usf
+++ b/Engine/Shaders/Private/TemporalSuperResolution/TSRDecimateHistory.usf
@@ -42,6 +42,7 @@ Texture2D<uint> ReprojectionVectorTexture;
 #else
 	Texture2D<float> ClosestDepthTexture;
 #endif
+Texture2D<uint> DepthErrorTexture;
 Texture2DArray<uint> PrevAtomicTextureArray;
 
 Texture2DArray<tsr_half4> PrevHistoryGuide;
@@ -230,7 +231,6 @@ void ComputeReprojectionEdge(
 
 void LoadPrevAtomicTexturesSamples(
 	float2 PrevScreenPos,
-	out uint HistoryUseCountSamples[4],
 	out uint HistoryClosestDeviceZSamples0[4],
 	out uint HistoryClosestDeviceZSamples1[4])
 {
@@ -244,24 +244,11 @@ void LoadPrevAtomicTexturesSamples(
 		float BilinearWeight = GetSampleWeight(BilinearInter, i);
 		uint2 PixelPos = GetSamplePixelCoord(BilinearInter, i);
 
-		bool bValidPixelPos = all(and(PixelPos >= InputInfo_ViewportMin, PixelPos < InputInfo_ViewportMax));
-
-		#if COMPRESS_PREV_USE_COUNT
-			uint2 CompressedUseCountPixelPos;
-			uint Shift;
-			ComputeCompressedUseCountPixelCoordinates(PixelPos, /* out */ CompressedUseCountPixelPos, /* out */ Shift);
-
-			uint HistoryUseCount = PrevAtomicTextureArray[uint3(CompressedUseCountPixelPos, 0)];
-		#else
-			uint HistoryUseCount = PrevAtomicTextureArray[uint3(PixelPos, 0)];
-		#endif
-			
-		HistoryUseCountSamples[i] = HistoryUseCount;
-		HistoryClosestDeviceZSamples0[i] = PrevAtomicTextureArray[uint3(PixelPos, 1)];
+		HistoryClosestDeviceZSamples0[i] = PrevAtomicTextureArray[uint3(PixelPos, 0)];
 
 		#if CONFIG_SUPPORT_ORTHO_VIEW
 		{
-			HistoryClosestDeviceZSamples1[i] = PrevAtomicTextureArray[uint3(PixelPos, 2)];
+			HistoryClosestDeviceZSamples1[i] = PrevAtomicTextureArray[uint3(PixelPos, 1)];
 		}
 		#else
 		{
@@ -272,26 +259,27 @@ void LoadPrevAtomicTexturesSamples(
 } // LoadPrevAtomicTexturesSamples()
 
 void ProcessPrevAtomicTexturesSamples(
-	uint HistoryUseCountSamples[4],
 	uint HistoryClosestDeviceZSamples0[4],
 	uint HistoryClosestDeviceZSamples1[4],
+	float2 ScreenPos,
+	float2 ScreenVelocity,
 	float PrevDeviceZ,
-	float2 PrevScreenPos,
+	float DeviceZError,
 	out tsr_half ParallaxRejectionMask,
-	out bool bIsValidHoleFillingPixelVelocity,
-	out uint EncodedHoleFillingVelocity)
+	out float2 HoleFillingPixelVelocity,
+	out bool bCanHoleFill)
 {
-	ParallaxRejectionMask = tsr_half(0.0);
-	bIsValidHoleFillingPixelVelocity = false;
-	EncodedHoleFillingVelocity = 0;
-
 	float WorldDepth = ConvertFromDeviceZ(PrevDeviceZ);
+	float WorldDepthError = abs(WorldDepth - ConvertFromDeviceZ(PrevDeviceZ + DeviceZError));
 	
+	float2 PrevScreenPos = ScreenPos - ScreenVelocity;
 	float2 PrevInputBufferUV = (InputInfo_ScreenPosToViewportScale * PrevScreenPos + InputInfo_ScreenPosToViewportBias) * InputInfo_ExtentInverse;
 	PrevInputBufferUV = INVARIANT(PrevInputBufferUV);
 	FBilinearSampleInfos BilinearInter = GetBilinearSampleLevelInfos(PrevInputBufferUV, InputInfo_Extent, InputInfo_ExtentInverse);
-
 	
+	ParallaxRejectionMask = tsr_half(0.0);
+	bool bIsValidHoleFillingPixelVelocity = false;
+	uint EncodedHoleFillingVelocity = 0;
 	#if CONFIG_SUPPORT_ORTHO_VIEW
 	BRANCH
 	if (IsOrthoProjection())
@@ -305,16 +293,6 @@ void ProcessPrevAtomicTexturesSamples(
 			uint2 PixelPos = GetSamplePixelCoord(BilinearInter, i);
 
 			bool bValidPixelPos = all(and(PixelPos >= InputInfo_ViewportMin, PixelPos < InputInfo_ViewportMax));
-				
-			#if COMPRESS_PREV_USE_COUNT
-				uint2 CompressedUseCountPixelPos;
-				uint Shift;
-				ComputeCompressedUseCountPixelCoordinates(PixelPos, /* out */ CompressedUseCountPixelPos, /* out */ Shift);
-
-				uint HistoryUseCount = (HistoryUseCountSamples[i] >> (8 * Shift)) & 0xFF;
-			#else
-				uint HistoryUseCount = HistoryUseCountSamples[i];
-			#endif
 			
 			uint SampleHistoryClosestDeviceZ0 = HistoryClosestDeviceZSamples0[i];
 			uint SampleHistoryClosestDeviceZ1 = HistoryClosestDeviceZSamples1[i];
@@ -329,13 +307,12 @@ void ProcessPrevAtomicTexturesSamples(
 				const float PixelDepthError = 3.0;
 				const float RadiusToDiameter = 2.0;
 				float WorldDepthEpsilon = GetDepthPixelRadiusForProjectionType(HistoryClosestWorldDepth) * PixelDepthError * RadiusToDiameter;
+				WorldDepthEpsilon += WorldDepthError;
 				float DeltaDepth = abs(HistoryClosestWorldDepth - WorldDepth);
 
 				DepthRejection = saturate(tsr_half(2.0) - tsr_half(DeltaDepth / WorldDepthEpsilon));
 			}
 
-			tsr_half HistoryGhostingRejection = saturate(tsr_half(1.0) - tsr_half(3.0) * abs(tsr_half(HistoryUseCount) * rcp(tsr_half(PREV_USE_COUNT_QUANTIZATION)) - tsr_half(1.0)));
-		
 			bIsValidHoleFillingPixelVelocity = bIsValidHoleFillingPixelVelocity || bValidPixelPos;
 
 			FLATTEN
@@ -344,7 +321,7 @@ void ProcessPrevAtomicTexturesSamples(
 				EncodedHoleFillingVelocity = SampleEncodedHoleFillingVelocity;
 			}
 
-			ParallaxRejectionMask = select(bValidPixelPos, ParallaxRejectionMask + tsr_half(BilinearWeight) * max(DepthRejection, HistoryGhostingRejection), ParallaxRejectionMask);
+			ParallaxRejectionMask = select(bValidPixelPos, ParallaxRejectionMask + tsr_half(BilinearWeight) * DepthRejection, ParallaxRejectionMask);
 		}
 	}
 	else
@@ -357,16 +334,6 @@ void ProcessPrevAtomicTexturesSamples(
 			uint2 PixelPos = GetSamplePixelCoord(BilinearInter, i);
 
 			bool bValidPixelPos = all(and(PixelPos >= InputInfo_ViewportMin, PixelPos < InputInfo_ViewportMax));
-				
-			#if COMPRESS_PREV_USE_COUNT
-				uint2 CompressedUseCountPixelPos;
-				uint Shift;
-				ComputeCompressedUseCountPixelCoordinates(PixelPos, /* out */ CompressedUseCountPixelPos, /* out */ Shift);
-
-				uint HistoryUseCount = (HistoryUseCountSamples[i] >> (8 * Shift)) & 0xFF;
-			#else
-				uint HistoryUseCount = HistoryUseCountSamples[i];
-			#endif
 			
 			uint SampleHistoryClosestDeviceZ = HistoryClosestDeviceZSamples0[i];
 		
@@ -378,19 +345,32 @@ void ProcessPrevAtomicTexturesSamples(
 				const float PixelDepthError = 3.0;
 				const float RadiusToDiameter = 2.0;
 				float WorldDepthEpsilon = GetDepthPixelRadiusForProjectionType(HistoryClosestWorldDepth) * PixelDepthError * RadiusToDiameter;
+				WorldDepthEpsilon += WorldDepthError;
 				float DeltaDepth = abs(HistoryClosestWorldDepth - WorldDepth);
 
 				DepthRejection = saturate(tsr_half(2.0) - tsr_half(DeltaDepth / WorldDepthEpsilon));
 			}
 
-			tsr_half HistoryGhostingRejection = saturate(tsr_half(1.0) - tsr_half(3.0) * abs(tsr_half(HistoryUseCount) * rcp(tsr_half(PREV_USE_COUNT_QUANTIZATION)) - tsr_half(1.0)));
-		
 			bIsValidHoleFillingPixelVelocity = bIsValidHoleFillingPixelVelocity || bValidPixelPos;
 
 			EncodedHoleFillingVelocity = select(bValidPixelPos, max(EncodedHoleFillingVelocity, SampleHistoryClosestDeviceZ), EncodedHoleFillingVelocity);
-			ParallaxRejectionMask = select(bValidPixelPos, ParallaxRejectionMask + tsr_half(BilinearWeight) * max(DepthRejection, HistoryGhostingRejection), ParallaxRejectionMask);
+			ParallaxRejectionMask = select(bValidPixelPos, ParallaxRejectionMask + tsr_half(BilinearWeight) * DepthRejection, ParallaxRejectionMask);
 		}
 	}
+	
+	float HoleFillingOutOfRange;
+	DecodeHoleFillingVelocity(EncodedHoleFillingVelocity, /* out */ HoleFillingPixelVelocity, /* out */ HoleFillingOutOfRange);
+
+	bCanHoleFill = bIsValidHoleFillingPixelVelocity && HoleFillingOutOfRange > 0.5;
+
+	// Dim down the parallax rejection mask if the velocity matches
+	{
+		float2 PixelVelocity = ScreenVelocity * ScreenVelocityToInputPixelVelocity;
+		
+		tsr_half VelocityDiff = tsr_half(length(HoleFillingPixelVelocity - PixelVelocity));
+		
+		ParallaxRejectionMask = lerp(tsr_half(1.0), ParallaxRejectionMask, select(bCanHoleFill, saturate(VelocityDiff), tsr_half(1.0)));
+	}
 } // ProcessPrevAtomicTexturesSamples()
 
 template<typename Type>
@@ -478,7 +458,6 @@ void ReprojectAllPrevTextures(
 	uint GroupThreadIndex,
 	uint EncodedReprojectionVector,
 	float DeviceZ,
-	out uint HistoryUseCountSamples[4],
 	out uint HistoryClosestDeviceZSamples0[4],
 	out uint HistoryClosestDeviceZSamples1[4],
 	out FCatmullRomFetches<tsr_half4> ReprojectedHistoryGuideSamples,
@@ -502,7 +481,6 @@ void ReprojectAllPrevTextures(
 		float2 ScreenPos = ApplyScreenTransform(float2(InputPixelPos), InputPixelPosToScreenPos);
 		float2 PrevScreenPos = ScreenPos - DecodeReprojectionVector(EncodedReprojectionVector);
 		LoadPrevAtomicTexturesSamples(PrevScreenPos,
-			/* out */ HistoryUseCountSamples,
 			/* out */ HistoryClosestDeviceZSamples0,
 			/* out */ HistoryClosestDeviceZSamples1);
 	}
@@ -594,6 +572,7 @@ void MainCS(
 	uint DilatedEncodedReprojectionVector;
 	float DeviceZ;
 	float PrevDeviceZ;
+	float DeviceZError;
 	bool bHasPixelAnimation;
 	ISOLATE
 	{
@@ -606,6 +585,7 @@ void MainCS(
 				Map8x8Tile2x2Lane(GroupThreadIndex) +
 				(tsr_ushort2(InputInfo_ViewportMin) + tsr_ushort2(GroupId) * tsr_ushort2(TILE_SIZE, TILE_SIZE)));
 			
+			uint EncodedDeviceZError = DepthErrorTexture[InputPixelPos];
 			#if CONFIG_RESURRECTION_REPROJECTION
 				float2 DeviceZAndPrevDeviceZ = ClosestDepthTexture[InputPixelPos];
 				DeviceZ = DeviceZAndPrevDeviceZ.g;
@@ -616,6 +596,7 @@ void MainCS(
 			
 			PrevDeviceZ = abs(DeviceZAndPrevDeviceZ.r);
 			bHasPixelAnimation = DeviceZAndPrevDeviceZ.r < 0.0;
+			DeviceZError = DecodeDeviceZError(PrevDeviceZ, EncodedDeviceZError);
 
 			//LoadReprojectionNeighborhood(InputPixelPos, /* out */ EncodedReprojectionVectorNeighborhood);
 			//DilatedEncodedReprojectionVector = EncodedReprojectionVectorNeighborhood[0];
@@ -644,7 +625,6 @@ void MainCS(
 	}
 	
 	// Sample the previous frame information
-	uint HistoryUseCountSamples[4];
 	uint HistoryClosestDeviceZSamples0[4];
 	uint HistoryClosestDeviceZSamples1[4];
 	
@@ -667,7 +647,6 @@ void MainCS(
 		ReprojectAllPrevTextures(
 			GroupId, GroupThreadIndex,
 			DilatedEncodedReprojectionVector, DeviceZ,
-			/* out */ HistoryUseCountSamples,
 			/* out */ HistoryClosestDeviceZSamples0,
 			/* out */ HistoryClosestDeviceZSamples1,
 			/* out */ ReprojectedHistoryGuideSamples,
@@ -687,8 +666,8 @@ void MainCS(
 	// Process the history samples
 	tsr_halfC ReprojectedHistoryGuideColor;
 	tsr_half  ReprojectedHistoryGuideUncertainty;
-	tsr_half ParallaxRejectionMask;
-	bool bOverwriteReprojectionField = false;
+	bool bIsParallaxDisocclusion;
+	bool bReprojectionHollFill;
 	//ISOLATE
 	{
 		#if CONFIG_SCENE_COLOR_ALPHA
@@ -732,25 +711,24 @@ void MainCS(
 
 		float2 ScreenPos = ApplyScreenTransform(float2(InputPixelPos), InputPixelPosToScreenPos);
 		float2 ScreenVelocity = DecodeReprojectionVector(DilatedEncodedReprojectionVector);
-		float2 PrevScreenPos = ScreenPos - ScreenVelocity;
 		
 		// Compute the parralax rejection mask
-		bool bIsValidHoleFillingPixelVelocity = false;
-		uint EncodedHoleFillingVelocity = 0;
+		float2 HoleFillingPixelVelocity;
+		tsr_half ParallaxRejectionMask;
+		bool bCanHoleFill;
 		ProcessPrevAtomicTexturesSamples(
-			HistoryUseCountSamples,
 			HistoryClosestDeviceZSamples0,
 			HistoryClosestDeviceZSamples1,
+			ScreenPos,
+			ScreenVelocity,
 			PrevDeviceZ,
-			PrevScreenPos,
+			DeviceZError,
 			/* out */ ParallaxRejectionMask,
-			/* out */ bIsValidHoleFillingPixelVelocity,
-			/* out */ EncodedHoleFillingVelocity);
-
-		// Finish off the ParallaxRejectionMask
-		ParallaxRejectionMask = select(bIsOffScreen, tsr_half(0.0), ParallaxRejectionMask);
+			/* out */ HoleFillingPixelVelocity,
+			/* out */ bCanHoleFill);
 
-		bool bIsOffscreenOrParallaxRejected = ParallaxRejectionMask < tsr_half(PARALLAX_REJECTION_MASK_THRESHOLD);
+		bIsParallaxDisocclusion = !bIsOffScreen && ParallaxRejectionMask < tsr_half(PARALLAX_REJECTION_MASK_THRESHOLD);
+		bool bIsDisoccluded = bIsOffScreen || bIsParallaxDisocclusion;
 
 		// LDS despill
 		#if CONFIG_MANUAL_LDS_SPILL
@@ -759,36 +737,23 @@ void MainCS(
 		}
 		#endif
 	
-		// Hole fill the velocity.
-		{
-			float2 HoleFillingPixelVelocity;
-			float HoleFillingOutOfRange;
-			DecodeHoleFillingVelocity(EncodedHoleFillingVelocity, /* out */ HoleFillingPixelVelocity, /* out */ HoleFillingOutOfRange);
-			if (!bIsValidHoleFillingPixelVelocity)
-			{
-				HoleFillingOutOfRange = 0.0;
-			}
+		bReprojectionHollFill = bCanHoleFill && bIsParallaxDisocclusion;
 
-			// If the hole filling velocity has better result that the neighborhood, use that.
-			FLATTEN
-			if (tsr_half(HoleFillingOutOfRange) > 0.0 && bIsOffscreenOrParallaxRejected)
-			{
-				// Disable the reprojection's jacobian since this is to extrude disoccluded area we don't know how they look anyway.
-				float2 HoleFillingReprojectionVector = HoleFillingPixelVelocity * abs(InputPixelVelocityToScreenVelocity);
-
-				DilatedEncodedReprojectionVector = EncodeReprojectionVector(HoleFillingReprojectionVector);
-				bOverwriteReprojectionField = true;
-			}
+		// If the hole filling velocity has better result that the neighborhood, use that.
+		BRANCH
+		if (bReprojectionHollFill)
+		{
+			float2 HoleFillingReprojectionVector = HoleFillingPixelVelocity * InputPixelVelocityToScreenVelocity;
+			DilatedEncodedReprojectionVector = EncodeReprojectionVector(HoleFillingReprojectionVector);
 		}
 	}
 	
 	// Builds the decimate bitmask
-	tsr_ushort DecimateBitMask;
-	{
-		bool bIsDisoccluded = !(ParallaxRejectionMask > tsr_half(PARALLAX_REJECTION_MASK_THRESHOLD));
-
-		DecimateBitMask = select(bIsDisoccluded, tsr_ushort(0x1), tsr_ushort(0x0)) | select(bHasPixelAnimation, tsr_ushort(0x2), tsr_ushort(0x0));
-	}
+	tsr_ushort DecimateBitMask = 
+		select(bIsOffScreen,            tsr_ushort(0x1), tsr_ushort(0x0)) |
+		select(bIsParallaxDisocclusion, tsr_ushort(0x2), tsr_ushort(0x0)) |
+		select(bHasPixelAnimation,      tsr_ushort(0x4), tsr_ushort(0x0)) |
+		select(bReprojectionHollFill,   tsr_ushort(0x8), tsr_ushort(0x0));
 
 	// Output full res
 	ISOLATE
@@ -835,6 +800,9 @@ void MainCS(
 
 		// Selectively overwrite the reprojection field
 		{
+			bool bOverwriteReprojectionField = bReprojectionHollFill;
+
+			// Disable the reprojection's jacobian since this is to extrude disoccluded area we don't know how they look anyway.
 			const uint EncodedReprojectionJacobian = EncodeReprojectionJacobian(tsr_half2x2(tsr_half(0.0).xxxx));
 
 			tsr_short2 ReprojectionFieldOutputPixelPos = OutputPixelPos;
diff --git a/Engine/Shaders/Private/TemporalSuperResolution/TSRDilateVelocity.usf b/Engine/Shaders/Private/TemporalSuperResolution/TSRDilateVelocity.usf
--- a/Engine/Shaders/Private/TemporalSuperResolution/TSRDilateVelocity.usf
+++ b/Engine/Shaders/Private/TemporalSuperResolution/TSRDilateVelocity.usf
@@ -201,17 +201,9 @@ void FindClosestDepthOffset(
 	out tsr_short2 ReprojectionOffset)
 {
 	float DeviceZ = AccessNeighborhoodCenter(DevizeZNeighborhood);
+	float DeviceZError = ComputePixelDeviceZError(DeviceZ);
 	uint PackedDepthOffset = PackDepthAndOffset(/* Offset = */ int2(0, 0), DeviceZ);
 	
-	#if CONFIG_DEPTH_EDGE_DETECTION
-	float DeviceZError;
-	{
-		float Depth = ConvertFromDeviceZ(DeviceZ);
-		float ErrorCorrectedDeviceZ = ConvertToDeviceZ(GetDepthPixelRadiusForProjectionType(Depth) * 2.0f + Depth);
-		DeviceZError = abs(ErrorCorrectedDeviceZ - DeviceZ);
-	}
-	#endif
-	
 	#if 0
 		UNROLL_N(4)
 		for (uint i = 0; i < 4; i++)
@@ -256,123 +248,199 @@ void FindClosestDepthOffset(
 	#endif
 } // FindClosestDepthOffset()
 
-/** Compute the velocity differential along a screen space axis. */
-bool ComputeReprojectionDifferential(
-	float  DevizeZNeighborhood[CONFIG_BUTTERFLY_SAMPLES],
-	float2 PixelVelocityNeighborhood[CONFIG_BUTTERFLY_SAMPLES],
-	float DeviceZError,
-	const int2 PositiveDirection,
+/** Compute the velocity jacobian along a two orthogonal axes, taking a common closest depth. */
+bool ComputeDepthBilaterals(
+	const int2 DirectionE,
+	const int2 DirectionS,
 	bool bCompute,
-	inout tsr_half2 ReprojectionDifferential)
+	float DevizeZNeighborhood[CONFIG_BUTTERFLY_SAMPLES],
+	inout float DepthError,
+	out tsr_half BilateralWeights[4])
 {
 	float DevizeZC = AccessNeighborhoodCenter(DevizeZNeighborhood);
-	float DevizeZP = AccessNeighborhoodStaticOffset(DevizeZNeighborhood, +PositiveDirection);
-	float DevizeZN = AccessNeighborhoodStaticOffset(DevizeZNeighborhood, -PositiveDirection);
+	float DevizeZE = AccessNeighborhoodStaticOffset(DevizeZNeighborhood, +DirectionE);
+	float DevizeZW = AccessNeighborhoodStaticOffset(DevizeZNeighborhood, -DirectionE);
+	float DevizeZS = AccessNeighborhoodStaticOffset(DevizeZNeighborhood, +DirectionS);
+	float DevizeZN = AccessNeighborhoodStaticOffset(DevizeZNeighborhood, -DirectionS);
 	
-	float2 PixelVelocityC = AccessNeighborhoodCenter(PixelVelocityNeighborhood);
-	float2 PixelVelocityP = AccessNeighborhoodStaticOffset(PixelVelocityNeighborhood, +PositiveDirection);
-	float2 PixelVelocityN = AccessNeighborhoodStaticOffset(PixelVelocityNeighborhood, -PositiveDirection);
-	
-	bool bSuccess = false;
+	tsr_half WeightE;
+	tsr_half WeightW;
+	tsr_half WeightS;
+	tsr_half WeightN;
+		
+	bool bSuccess;
 	BRANCH
 	if (bCompute)
 	{
-		const tsr_half WeightThreshold = tsr_half(2.0); 
-		const tsr_half WeightThresholdSquare = WeightThreshold * WeightThreshold; 
-	
-		// Compute the velocity differential.
-		tsr_half2 VelocityDeltaP = tsr_half2(PixelVelocityP - PixelVelocityC);
-		tsr_half2 VelocityDeltaN = tsr_half2(PixelVelocityC - PixelVelocityN);
+		float VariationEW = abs(DevizeZC - (DevizeZE + DevizeZW) * 0.5);
+		float VariationSN = abs(DevizeZC - (DevizeZS + DevizeZN) * 0.5);
+		float VariationDiff = VariationEW - VariationSN;
+
+		float DepthDiffE = abs(DevizeZE - DevizeZC);
+		float DepthDiffW = abs(DevizeZW - DevizeZC);
+		float DepthDiffS = abs(DevizeZS - DevizeZC);
+		float DepthDiffN = abs(DevizeZN - DevizeZC);
 	
-		float DepthVariation = DevizeZC - (DevizeZP + DevizeZN) * 0.5;
-		float DepthMax = max3(DevizeZP, DevizeZN, DevizeZC);
-		float DepthMin = min3(DevizeZP, DevizeZN, DevizeZC);
-		float Delta = max(2.0 * DepthVariation - (DepthMax - DepthMin), 0.0);
-
-		float DepthDiffPN = abs(DevizeZP - DevizeZN);
-		float DepthDiffP  = abs(DevizeZP - DevizeZC);
-		float DepthDiffN  = abs(DevizeZN - DevizeZC);
+		float DepthDiffEW = abs(DevizeZE - DevizeZW);
+		float DepthDiffSN = abs(DevizeZS - DevizeZN);
 	
-		float MinDepthDerivative = clamp(min3(DepthDiffPN, DepthDiffP, DepthDiffN), tsr_half(0.125) * DeviceZError, tsr_half(8.0) * DeviceZError);
-
-		tsr_half WeightP = tsr_half(1.0);
-		tsr_half WeightN = tsr_half(1.0);
+		float MinDepthDerivativeE = min3(DepthDiffEW, DepthDiffE, DepthDiffW);
+		float MinDepthDerivativeS = min3(DepthDiffSN, DepthDiffS, DepthDiffN);
+		
+		float DepthDerivativeE = MinDepthDerivativeE;
+		float DepthDerivativeS = MinDepthDerivativeS;
 
-		WeightP *= tsr_half(saturate(1.1 - (DepthDiffP) / MinDepthDerivative));
-		WeightN *= tsr_half(saturate(1.1 - (DepthDiffN) / MinDepthDerivative));
+		WeightE = tsr_half(saturate(1.1 - min(DepthDiffE, VariationEW) / DepthDerivativeE));
+		WeightW = tsr_half(saturate(1.1 - min(DepthDiffW, VariationEW) / DepthDerivativeE));
 	
-		// Bilateral based on velicities differentials
-		//WeightP *= saturate(WeightThresholdSquare - dot(VelocityDeltaP, VelocityDeltaP));
-		//WeightN *= saturate(WeightThresholdSquare - dot(VelocityDeltaN, VelocityDeltaN));
+		WeightS = tsr_half(saturate(1.1 - min(DepthDiffS, VariationSN) / DepthDerivativeS));
+		WeightN = tsr_half(saturate(1.1 - min(DepthDiffN, VariationSN) / DepthDerivativeS));
 
-		tsr_half TotalWeight = WeightP + WeightN;
-		bSuccess = TotalWeight > tsr_half(0.0);
-		tsr_half InvTotalWeight = select(bSuccess, rcp(TotalWeight), tsr_half(0.0));
 	
-		ReprojectionDifferential = VelocityDeltaP * (WeightP * InvTotalWeight) + VelocityDeltaN * (WeightN * InvTotalWeight);
+		tsr_half TotalWeightEW = WeightE + WeightW;
+		tsr_half TotalWeightSN = WeightS + WeightN;
+	
+		bool bSuccessEW = TotalWeightEW > tsr_half(0.0);
+		bool bSuccessSN = TotalWeightSN > tsr_half(0.0);
+
+		tsr_half InvTotalWeightEW = select(bSuccessEW, rcp(TotalWeightEW), tsr_half(0.0));
+		tsr_half InvTotalWeightSN = select(bSuccessSN, rcp(TotalWeightSN), tsr_half(0.0));
+
+		WeightE *= InvTotalWeightEW;
+		WeightW *= InvTotalWeightEW;
+		WeightS *= InvTotalWeightSN;
+		WeightN *= InvTotalWeightSN;
+		bSuccess = bSuccessEW || bSuccessSN;
+		
+		DepthError = max(DepthDiffE * WeightE + DepthDiffW * WeightW, DepthDiffS * WeightS + DepthDiffN * WeightN);
 	}
+	else
+	{
+		WeightE = tsr_half(1.0);
+		WeightW = tsr_half(1.0);
+		WeightS = tsr_half(1.0);
+		WeightN = tsr_half(1.0);
+		bSuccess = false;
+	}
+
+	BilateralWeights[0] = WeightE;
+	BilateralWeights[1] = WeightW;
+	BilateralWeights[2] = WeightS;
+	BilateralWeights[3] = WeightN;
+
 	return bSuccess;
+} // ComputeDepthBilaterals()
+
+void ComputeVelocityJacobian(
+	const int2 DirectionE,
+	const int2 DirectionS,
+	bool bCompute,
+	float2 PixelVelocityNeighborhood[CONFIG_BUTTERFLY_SAMPLES],
+	tsr_half BilateralWeights[4],
+	inout tsr_half2x2 ReprojectionJacobian)
+{
+	
+	float2 PixelVelocityC = AccessNeighborhoodCenter(PixelVelocityNeighborhood);
+	float2 PixelVelocityE = AccessNeighborhoodStaticOffset(PixelVelocityNeighborhood, +DirectionE);
+	float2 PixelVelocityW = AccessNeighborhoodStaticOffset(PixelVelocityNeighborhood, -DirectionE);
+	float2 PixelVelocityS = AccessNeighborhoodStaticOffset(PixelVelocityNeighborhood, +DirectionS);
+	float2 PixelVelocityN = AccessNeighborhoodStaticOffset(PixelVelocityNeighborhood, -DirectionS);
+	
+	tsr_half WeightE = BilateralWeights[0];
+	tsr_half WeightW = BilateralWeights[1];
+	tsr_half WeightS = BilateralWeights[2];
+	tsr_half WeightN = BilateralWeights[3];
+
+	BRANCH
+	if (bCompute)
+	{
+		tsr_half2 VelocityDeltaE = tsr_half2(PixelVelocityE - PixelVelocityC);
+		tsr_half2 VelocityDeltaW = tsr_half2(PixelVelocityC - PixelVelocityW);
+		tsr_half2 VelocityDeltaS = tsr_half2(PixelVelocityS - PixelVelocityC);
+		tsr_half2 VelocityDeltaN = tsr_half2(PixelVelocityC - PixelVelocityN);
+
+		ReprojectionJacobian[0] = VelocityDeltaE * WeightE + VelocityDeltaW * WeightW;
+		ReprojectionJacobian[1] = VelocityDeltaS * WeightS + VelocityDeltaN * WeightN;
+	}
+} // ComputeVelocityJacobian()
+
+/** Compute the depth error on horizontal and vertical. */
+float ComputeDepthError(float DevizeZNeighborhood[CONFIG_BUTTERFLY_SAMPLES])
+{
+	float DepthError;
+	tsr_half BilateralWeights[4];
+
+	ComputeDepthBilaterals(
+		kOffsetE, kOffsetS,
+		/* bCompute = */ true,
+		DevizeZNeighborhood,
+		/* inout */ DepthError,
+		/* out */ BilateralWeights);
+
+	return DepthError;
 }
 
 /** Find the closest DevizeZ and its corresponding offset */
 void ComputeReprojectionJacobian(
 	float  DevizeZNeighborhood[CONFIG_BUTTERFLY_SAMPLES],
 	float2 PixelVelocityNeighborhood[CONFIG_BUTTERFLY_SAMPLES],
-	out tsr_half2x2 ReprojectionJacobian)
+	out tsr_half2x2 ReprojectionJacobian,
+	out float DepthError)
 {
 	float DeviceZ = AccessNeighborhoodCenter(DevizeZNeighborhood);
-	float DeviceZError;
-	{
-		float Depth = ConvertFromDeviceZ(DeviceZ);
-		float ErrorCorrectedDeviceZ = ConvertToDeviceZ(GetDepthPixelRadiusForProjectionType(Depth) * 2.0f + Depth);
-		DeviceZError = abs(ErrorCorrectedDeviceZ - DeviceZ);
-	}
 	
-	// Horizontal & vertical cardinals
+	ReprojectionJacobian = tsr_half(0.0);
+	DepthError = 0.0;
+
+	// Horizontal & vertical cardinals first as the pixels are closer than on diagonals
 	bool bSuccess;
 	{
-		ReprojectionJacobian = tsr_half(0.0);
-
-		bool bSuccessX = ComputeReprojectionDifferential(
+		const bool bCompute = true;
+		const int2 DirectionE = kOffsetE;
+		const int2 DirectionS = kOffsetS;
+
+		tsr_half BilateralWeights[4];
+		bSuccess = ComputeDepthBilaterals(
+			DirectionE, DirectionS,
+			bCompute,
 			DevizeZNeighborhood,
-			PixelVelocityNeighborhood,
-			DeviceZError,
-			/* PositiveDirection = */ kOffsetE,
-			/* bCompute = */ true,
-			/* inout */ ReprojectionJacobian[0]);
+			/* inout */ DepthError,
+			/* out */ BilateralWeights);
 		
-		bool bSuccessY = ComputeReprojectionDifferential(
-			DevizeZNeighborhood,
+		ComputeVelocityJacobian(
+			DirectionE, DirectionS,
+			bCompute,
 			PixelVelocityNeighborhood,
-			DeviceZError,
-			/* PositiveDirection = */ kOffsetS,
-			/* bCompute = */ true,
-			/* inout */ ReprojectionJacobian[1]);
-	
-		bSuccess = bSuccessX || bSuccessY;
+			BilateralWeights,
+			/* inout */ ReprojectionJacobian);
 	}
 	
-	// Diagonal cardinals
+	// Compute diagonal jacobian if the horizontal/vertical failed
+#if CONFIG_BUTTERFLY_KERNEL && COMPILER_SUPPORTS_WAVE_VOTE
+	if (WaveActiveAnyTrue(!bSuccess))
+#endif
 	{
-		tsr_half2 DifferentialSE;
-		ComputeReprojectionDifferential(
-			DevizeZNeighborhood,
-			PixelVelocityNeighborhood,
-			DeviceZError,
-			/* PositiveDirection = */ kOffsetSE,
-			/* bCompute = */ !bSuccess,
-			/* inout */ ReprojectionJacobian[0]);
-			
-		tsr_half2 DifferentialSW;
-		ComputeReprojectionDifferential(
+		const bool bCompute = !bSuccess;
+		const int2 DirectionE = kOffsetSE;
+		const int2 DirectionS = kOffsetSW;
+		
+		tsr_half BilateralWeights[4];
+		ComputeDepthBilaterals(
+			DirectionE, DirectionS,
+			bCompute,
 			DevizeZNeighborhood,
+			/* inout */ DepthError,
+			/* out */ BilateralWeights);
+		
+		ComputeVelocityJacobian(
+			DirectionE, DirectionS,
+			bCompute,
 			PixelVelocityNeighborhood,
-			DeviceZError,
-			/* PositiveDirection = */ kOffsetSW,
-			/* bCompute = */ !bSuccess,
-			/* inout */ ReprojectionJacobian[1]);
+			BilateralWeights,
+			/* inout */ ReprojectionJacobian);
 			
-		if (!bSuccess)
+		// Rotate the diagonal jacobian back to horizontal/vertical
+		if (bCompute)
 		{
 			tsr_half2 DifferentialSE = ReprojectionJacobian[0];
 			tsr_half2 DifferentialSW = ReprojectionJacobian[1];
@@ -381,6 +449,7 @@ void ComputeReprojectionJacobian(
 		}
 	}
 	
+	// Flips the signs due to the butterfly kernel.
 	#if CONFIG_BUTTERFLY_KERNEL && 1
 		tsr_half2 LaneOffsetSign = tsr_half2(GetLaneOffsetSign());
 		ReprojectionJacobian[0] *= LaneOffsetSign.x;
@@ -418,6 +487,7 @@ void FetchAndComputeScreenVelocity(
 /** Spatial anti-alias the depth buffer and return the boundary. */
 tsr_half2 ComputeReprojectionBoundary(
 	tsr_short2 InputPixelPos,
+	tsr_short2 ReprojectionOffset,
 	float2 ScreenVelocity,
 	float2 DilatedScreenVelocity,
 	float DevizeZNeighborhood[CONFIG_BUTTERFLY_SAMPLES])
@@ -441,7 +511,21 @@ tsr_half2 ComputeReprojectionBoundary(
 	float2 PrevScreenPos = ScreenPos - DilatedScreenVelocity.xy;
 	bool bOffscreen = any(PrevScreenPos != clamp(PrevScreenPos, -1.0, 1.0));
 
-	bool bDilationHasEnoughVelocityDiff = length2(InputInfo_ViewportSize * (ScreenVelocity.xy - DilatedScreenVelocity.xy)) > ReprojectionFieldAntiAliasVelocityThreshold;
+	float2 PixelVelocityDifferential = InputInfo_ViewportSize * (ScreenVelocity.xy - DilatedScreenVelocity.xy);
+
+	/* when disoccluding, increase the size. */
+	bool bIsDilatedDisoccluding = (
+		float(ReprojectionOffset.x) * PixelVelocityDifferential.x < 0.0 ||
+		float(ReprojectionOffset.y) * PixelVelocityDifferential.y > 0.0);
+
+	/** If Dilated velocity is coming towards this pixel, where an extrusion could possibly happen. */
+	bool bIsDilatedComing = (
+		float(ReprojectionOffset.x) * DilatedScreenVelocity.x < 0.0 || 
+		float(ReprojectionOffset.y) * DilatedScreenVelocity.y > 0.0);
+
+	float PixelVelocity = length2(PixelVelocityDifferential);
+	bool bDilationHasEnoughVelocityDiff = PixelVelocity > select(bIsDilatedDisoccluding, pow(0.5, 2.0), ReprojectionFieldAntiAliasVelocityThreshold);
+
 	bool bAntiAliasReprojectionField = bReprojectionField && bDilationHasEnoughVelocityDiff && !bOffscreen;
 
 	tsr_half2 ReprojectionBoundary;
@@ -500,11 +584,12 @@ tsr_half2 ComputeReprojectionBoundary(
 
 void ScatterClosestOccluder(bool bValidOutputPixel, float2 ScreenPos, float2 FinalScreenVelocity, float PrevClosestDeviceZ)
 {
+	const float PrevUseCountQuantization = 63.0;
+
 	float2 PrevScreenPos = ScreenPos - FinalScreenVelocity.xy;
 	float2 PrevInputBufferUV = (InputInfo_ScreenPosToViewportScale * PrevScreenPos + InputInfo_ScreenPosToViewportBias) * InputInfo_ExtentInverse;
 
-	float2 PixelVelocity = FinalScreenVelocity.xy * 0.5 * InputInfo_ViewportSize;
-	uint EncodedHoleFillingVelocity = EncodeHoleFillingVelocity(PixelVelocity);
+	uint EncodedHoleFillingVelocity = EncodeHoleFillingVelocity(FinalScreenVelocity.xy * ScreenVelocityToInputPixelVelocity);
 
 	bool bValidHistoryCoord = bValidOutputPixel && all(and(PrevInputBufferUV > PrevOutputBufferUVMin, PrevInputBufferUV < PrevOutputBufferUVMax));
 
@@ -525,27 +610,15 @@ void ScatterClosestOccluder(bool bValidOutputPixel, float2 ScreenPos, float2 Fin
 			float Weight = GetSampleWeight(BilinearInter, i);
 			uint2 PixelPos = GetSamplePixelCoord(BilinearInter, i);
 
-			uint FinalAdd = uint(round(Weight * PREV_USE_COUNT_QUANTIZATION));
+			uint FinalAdd = uint(round(Weight * PrevUseCountQuantization));
 			
 			PixelPos.x = ((FinalAdd > 0) && bValidHistoryCoord) ? PixelPos.x : (~0);
 
-			#if COMPRESS_PREV_USE_COUNT
-				uint2 CompressedUseCountPixelPos;
-				uint Shift;
-				ComputeCompressedUseCountPixelCoordinates(PixelPos, /* out */ CompressedUseCountPixelPos, /* out */ Shift);
-
-				FinalAdd = FinalAdd << (8 * Shift);
-			#else
-				uint2 CompressedUseCountPixelPos = PixelPos;
-			#endif
-
-			InterlockedAdd(PrevAtomicOutput[tsr_ushort3(CompressedUseCountPixelPos, 0)], FinalAdd);
-
 			uint PreviousDepth0;
-			InterlockedMax(PrevAtomicOutput[tsr_ushort3(PixelPos, 1)], PrevClosestDepth0, /* out */ PreviousDepth0);
+			InterlockedMax(PrevAtomicOutput[tsr_ushort3(PixelPos, 0)], PrevClosestDepth0, /* out */ PreviousDepth0);
 
 			PixelPos.x = select(PreviousDepth0 < PrevClosestDepth0, PixelPos.x, ~0u);
-			InterlockedMax(PrevAtomicOutput[tsr_ushort3(PixelPos, 2)], PrevClosestDepth1);
+			InterlockedMax(PrevAtomicOutput[tsr_ushort3(PixelPos, 1)], PrevClosestDepth1);
 		}
 	}
 	else
@@ -559,22 +632,11 @@ void ScatterClosestOccluder(bool bValidOutputPixel, float2 ScreenPos, float2 Fin
 			float Weight = GetSampleWeight(BilinearInter, i);
 			uint2 PixelPos = GetSamplePixelCoord(BilinearInter, i);
 
-			uint FinalAdd = uint(round(Weight * PREV_USE_COUNT_QUANTIZATION));
+			uint FinalAdd = uint(round(Weight * PrevUseCountQuantization));
 			
 			PixelPos.x = ((FinalAdd > 0) && bValidHistoryCoord) ? PixelPos.x : (~0);
 
-			#if COMPRESS_PREV_USE_COUNT
-				uint2 CompressedUseCountPixelPos;
-				uint Shift;
-				ComputeCompressedUseCountPixelCoordinates(PixelPos, /* out */ CompressedUseCountPixelPos, /* out */ Shift);
-
-				FinalAdd = FinalAdd << (8 * Shift);
-			#else
-				uint2 CompressedUseCountPixelPos = PixelPos;
-			#endif
-
-			InterlockedAdd(PrevAtomicOutput[tsr_ushort3(CompressedUseCountPixelPos, 0)], FinalAdd);
-			InterlockedMax(PrevAtomicOutput[tsr_ushort3(PixelPos, 1)], PrevClosestDepth);
+			InterlockedMax(PrevAtomicOutput[tsr_ushort3(PixelPos, 0)], PrevClosestDepth);
 		}
 	}
 } // ScatterClosestOccluder()
@@ -604,6 +666,7 @@ void MainCS(
 	float3 DilatedScreenVelocity;
 	uint EncodedReprojectionBoundary;
 	uint EncodedReprojectionJacobian;
+	uint EncodedDeviceZError;
 	bool bHasPixelAnimation;
 	tsr_half VelocityEdge;
 	BRANCH
@@ -636,12 +699,14 @@ void MainCS(
 		#endif
 	
 		tsr_half2x2 ReprojectionJacobian;
+		float ParallaxDepthError;
 		ISOLATE
 		{
 			ComputeReprojectionJacobian(
 				DevizeZNeighborhood,
 				PixelVelocityNeighborhood,
-				/* out */ ReprojectionJacobian);
+				/* out */ ReprojectionJacobian,
+				/* out */ ParallaxDepthError);
 			EncodedReprojectionJacobian = EncodeReprojectionJacobian(ReprojectionJacobian);
 		}
 		
@@ -679,12 +744,19 @@ void MainCS(
 		{
 			tsr_half2 ReprojectionBoundary = ComputeReprojectionBoundary(
 				InputPixelPos,
+				ReprojectionOffset,
 				ScreenVelocity,
 				DilatedScreenVelocity.xy,
 				DevizeZNeighborhood);
 
 			EncodedReprojectionBoundary = EncodeReprojectionBoundary(ReprojectionOffset, ReprojectionBoundary);
 		}
+		
+		// Encode the device z error.
+		{
+			float PrevClosestDeviceZ = ClosestDeviceZ - DilatedScreenVelocity.z;
+			EncodedDeviceZError = EncodeDeviceZError(PrevClosestDeviceZ, ParallaxDepthError);
+		}
 	}
 	else // if (!bReprojectionField)
 	{
@@ -701,8 +773,11 @@ void MainCS(
 			/* out */ ClosestDeviceZ,
 			/* out */ ReprojectionOffset);
 	
+		// Compute the parallax depth error
+		float ParallaxDepthError = ComputeDepthError(DevizeZNeighborhood);
+
 		// Compute final screen vecocity from pixel velocities to avoid duplicated register.
-		DeviceZ        = AccessNeighborhoodCenter(DevizeZNeighborhood);
+		DeviceZ = AccessNeighborhoodCenter(DevizeZNeighborhood);
 	
 		// Compute the screen vecocity.
 		float3 ScreenVelocity;
@@ -714,16 +789,22 @@ void MainCS(
 			/* out */ bHasPixelAnimation);
 
 		// Compute the dilated screen vecocity.
-		bool bDildatedHasPixelAnimation;
+		bool bDilatedHasPixelAnimation;
 		FetchAndComputeScreenVelocity(
 			InputPixelPos,
 			ClosestDeviceZ,
 			ReprojectionOffset,
 			/* out */ DilatedScreenVelocity,
-			/* out */ bDildatedHasPixelAnimation);
+			/* out */ bDilatedHasPixelAnimation);
 			
 		tsr_half2 NeighborDeltaPixelVelocity = tsr_half2((DilatedScreenVelocity.xy - ScreenVelocity.xy) * ScreenVelocityToInputPixelVelocity);
 		VelocityEdge = saturate(tsr_half(1.1) - dot(abs(NeighborDeltaPixelVelocity), tsr_half(1.1).xx));
+		
+		// Encode the device z error.
+		{
+			float PrevClosestDeviceZ = ClosestDeviceZ - DilatedScreenVelocity.z;
+			EncodedDeviceZError = EncodeDeviceZError(PrevClosestDeviceZ, ParallaxDepthError);
+		}
 
 		EncodedReprojectionBoundary = 0u;
 		EncodedReprojectionJacobian = 0u;
@@ -780,10 +861,11 @@ void MainCS(
 		ClosestDepthOutput[OutputPixelPos] = float2(select(bHasPixelAnimation, -PrevClosestDeviceZ, PrevClosestDeviceZ), ClosestDeviceZ);
 
 		R8Output[uint3(OutputPixelPos, 0)] = uint(round(VelocityEdge * 255.0));
+		R8Output[uint3(OutputPixelPos, 1)] = EncodedDeviceZError;
 		BRANCH
 		if (bOutputIsMovingTexture)
 		{
-			R8Output[uint3(OutputPixelPos, 1)] = uint(round(IsMovingMask * 255.0));
+			R8Output[uint3(OutputPixelPos, 2)] = uint(round(IsMovingMask * 255.0));
 		}
 
 		// Output the reprojection field
diff --git a/Engine/Shaders/Private/TemporalSuperResolution/TSRRejectShading.usf b/Engine/Shaders/Private/TemporalSuperResolution/TSRRejectShading.usf
--- a/Engine/Shaders/Private/TemporalSuperResolution/TSRRejectShading.usf
+++ b/Engine/Shaders/Private/TemporalSuperResolution/TSRRejectShading.usf
@@ -134,8 +134,8 @@ void FetchDecimateMask(
 		tsr_ushort BitMask = tsr_ushort(DecimateMask.r * tsr_half(255.0));
 
 		VelocityEdge.SetElement(ElementIndex, saturate(tsr_half(tsr_ushort(DilateMask)) * tsr_half(rcp(255.0))));
-		bIsDisoccluded.SetElement(ElementIndex,     (BitMask & tsr_ushort(0x1)) != tsr_ushort(0));
-		bHasPixelAnimation.SetElement(ElementIndex, (BitMask & tsr_ushort(0x2)) != tsr_ushort(0));
+		bIsDisoccluded.SetElement(ElementIndex,     (BitMask & tsr_ushort(0x3)) != tsr_ushort(0));
+		bHasPixelAnimation.SetElement(ElementIndex, (BitMask & tsr_ushort(0x4)) != tsr_ushort(0));
 	}
 } // FetchDecimateMask()
 
diff --git a/Engine/Shaders/Private/TemporalSuperResolution/TSRReprojectionField.ush b/Engine/Shaders/Private/TemporalSuperResolution/TSRReprojectionField.ush
--- a/Engine/Shaders/Private/TemporalSuperResolution/TSRReprojectionField.ush
+++ b/Engine/Shaders/Private/TemporalSuperResolution/TSRReprojectionField.ush
@@ -148,6 +148,6 @@ tsr_half ComputeReprojectionUpscaleFactorFromJacobian(tsr_half2x2 ReprojectionJa
 	tsr_half2 VectorOffsetE = tsr_half2(1.0, 0.0) + ReprojectionJacobian[0];
 	tsr_half2 VectorOffsetS = tsr_half2(0.0, 1.0) + ReprojectionJacobian[1];
 
-	tsr_half ReprojectionUpscaleFactor = max(length(VectorOffsetE), length(VectorOffsetS));
+	tsr_half ReprojectionUpscaleFactor = sqrt(max(dot(VectorOffsetE, VectorOffsetE), tsr_half(1.0)) * max(dot(VectorOffsetS, VectorOffsetS), tsr_half(1.0)));
 	return ReprojectionUpscaleFactor;
 }
diff --git a/Engine/Shaders/Private/TemporalSuperResolution/TSRUpdateHistory.usf b/Engine/Shaders/Private/TemporalSuperResolution/TSRUpdateHistory.usf
--- a/Engine/Shaders/Private/TemporalSuperResolution/TSRUpdateHistory.usf
+++ b/Engine/Shaders/Private/TemporalSuperResolution/TSRUpdateHistory.usf
@@ -224,12 +224,6 @@ tsr_ushort2x2 Map8x8Tile2x2LaneDPV(uint GroupThreadIndex)
 #endif
 }
 
-CALL_SITE_DEBUGLOC 
-tsr_halfCx2 TransformColorForClampingBox(tsr_halfCx2 Color)
-{
-	return RGBToYCoCg(Color);
-}
-
 CALL_SITE_DEBUGLOC 
 tsr_half3x2 GetColorChannel(tsr_half4x2 Color)
 {
diff --git a/Engine/Shaders/Private/TemporalSuperResolution/TSRVisualize.usf b/Engine/Shaders/Private/TemporalSuperResolution/TSRVisualize.usf
--- a/Engine/Shaders/Private/TemporalSuperResolution/TSRVisualize.usf
+++ b/Engine/Shaders/Private/TemporalSuperResolution/TSRVisualize.usf
@@ -57,6 +57,7 @@ Texture2D<uint>      ReprojectionBoundaryTexture;
 Texture2D<uint>      ReprojectionJacobianTexture;
 Texture2D<uint>      ReprojectionVectorTexture;
 Texture2D<tsr_ushort> IsMovingMaskTexture;
+Texture2D<tsr_half2> DecimateMaskTexture;
 Texture2D<tsr_half4> HistoryRejectionTexture;
 Texture2D<tsr_half4> MoireHistoryTexture;
 Texture2D<tsr_ushort> AntiAliasMaskTexture;
@@ -150,6 +151,30 @@ FReprojectionFieldSample SampleReprojectionField(FContext Ctx)
 	return Reprojection;
 }
 
+struct FDecimateMaskSample
+{
+	bool bIsOffScreen;
+	bool bIsParallaxDisocclusion;
+	bool bIsDisoccluded;
+	bool bHasPixelAnimation;
+	bool bReprojectionHollFill;
+};
+
+FDecimateMaskSample SampleDecimateMask(FContext Ctx)
+{
+	FDecimateMaskSample Decimate;
+
+	tsr_half2 DecimateMask = DecimateMaskTexture[Ctx.InputPixelPos];
+	
+	tsr_ushort BitMask = tsr_ushort(DecimateMask.r * tsr_half(255.0));
+	Decimate.bIsOffScreen            = (BitMask & tsr_ushort(0x1)) != tsr_ushort(0);
+	Decimate.bIsParallaxDisocclusion = (BitMask & tsr_ushort(0x2)) != tsr_ushort(0);
+	Decimate.bIsDisoccluded          = (BitMask & tsr_ushort(0x3)) != tsr_ushort(0);
+	Decimate.bHasPixelAnimation      = (BitMask & tsr_ushort(0x4)) != tsr_ushort(0);
+	Decimate.bReprojectionHollFill   = (BitMask & tsr_ushort(0x8)) != tsr_ushort(0);
+	return Decimate;
+}
+
 
 //------------------------------------------------------- ENTRY POINT
 
@@ -202,11 +227,12 @@ tsr_half3 DisplayHistoryClamp(FContext Ctx)
 
 tsr_half3 DisplayParallaxDisocclusionMask(FContext Ctx)
 {
-	bool ParallaxDisocclusionMask = (tsr_ushort(HistoryRejectionTexture.SampleLevel(GlobalPointClampedSampler, Ctx.InputUV, 0).a * tsr_half(255.0)) & tsr_ushort(0x1)) != tsr_ushort(0x0);
+	FDecimateMaskSample Decimate = SampleDecimateMask(Ctx);
 	
-	tsr_half3 Color = kWhite;
-	Color = select(bCameraCut == 0, Color, kPink);
-	Color = select(ParallaxDisocclusionMask, Color, kRed);
+	tsr_half3 Color;
+	Color = select(Decimate.bReprojectionHollFill && Decimate.bIsParallaxDisocclusion, kOrange, kWhite);
+	Color = select(!Decimate.bReprojectionHollFill && Decimate.bIsParallaxDisocclusion, kRed, Color);
+	Color = select(Decimate.bIsOffScreen, kPink, Color);
 	return SampleTSROutputGray(Ctx) * Color;
 }
 
diff --git a/Engine/Source/Runtime/Renderer/Private/PostProcess/TemporalSuperResolution.cpp b/Engine/Source/Runtime/Renderer/Private/PostProcess/TemporalSuperResolution.cpp
--- a/Engine/Source/Runtime/Renderer/Private/PostProcess/TemporalSuperResolution.cpp
+++ b/Engine/Source/Runtime/Renderer/Private/PostProcess/TemporalSuperResolution.cpp
@@ -291,10 +291,9 @@ TAutoConsoleVariable<int32> CVarTSRReprojectionField(
 	ECVF_Scalability | ECVF_RenderThreadSafe);
 
 TAutoConsoleVariable<float> CVarTSRReprojectionFieldAntiAliasPixelSpeed(
-	TEXT("r.TSR.ReprojectionField.AntiAliasPixelSpeed"), 0.75f,
+	TEXT("r.TSR.ReprojectionField.AntiAliasPixelSpeed"), 0.125f,
 	TEXT("Defines the output pixel velocity at which point the dilation should be spatial anti-aliased based of the depth buffer ")
-	TEXT("to avoid reprojection aliasing by extrusion on fast geometric edges (Default to 0.75 to make sure there is parralax ")
-	TEXT("disocclusion mask behind moving objects, best tuned with r.TSR.Visualize=11)."),
+	TEXT("to avoid reprojection aliasing by extrusion on fast geometric edges (Default to 0.125, best tuned with r.TSR.Visualize=11)."),
 	ECVF_RenderThreadSafe);
 
 TAutoConsoleVariable<float> CVarTSRWeightClampingSampleCount(
@@ -736,6 +735,7 @@ class FTSRDecimateHistoryCS : public FTSRShader
 		SHADER_PARAMETER_RDG_TEXTURE(Texture2D, InputSceneColorTexture)
 		SHADER_PARAMETER_RDG_TEXTURE_SRV(Texture2D, ReprojectionVectorTexture)
 		SHADER_PARAMETER_RDG_TEXTURE(Texture2D, ClosestDepthTexture)
+		SHADER_PARAMETER_RDG_TEXTURE_SRV(Texture2D, DepthErrorTexture)
 		SHADER_PARAMETER_RDG_TEXTURE(Texture2DArray, PrevAtomicTextureArray)
 
 		SHADER_PARAMETER_STRUCT_INCLUDE(FTSRPrevHistoryParameters, PrevHistoryParameters)
@@ -1165,6 +1165,7 @@ class FTSRVisualizeCS : public FTSRShader
 		SHADER_PARAMETER_RDG_TEXTURE_SRV(Texture2D, ReprojectionJacobianTexture)
 		SHADER_PARAMETER_RDG_TEXTURE_SRV(Texture2D, ReprojectionVectorTexture)
 		SHADER_PARAMETER_RDG_TEXTURE_SRV(Texture2D, IsMovingMaskTexture)
+		SHADER_PARAMETER_RDG_TEXTURE(Texture2D, DecimateMaskTexture)
 		SHADER_PARAMETER_RDG_TEXTURE(Texture2D, HistoryRejectionTexture)
 		SHADER_PARAMETER_RDG_TEXTURE_SRV(Texture2D, MoireHistoryTexture)
 		SHADER_PARAMETER_RDG_TEXTURE(Texture2D, AntiAliasMaskTexture)
@@ -1886,7 +1887,7 @@ FDefaultTemporalUpscaler::FOutputs AddTemporalSuperResolutionPasses(
 				PF_R32_UINT,
 				FClearValueBinding::None,
 				/* InFlags = */ TexCreate_ShaderResource | TexCreate_UAV | TexCreate_AtomicCompatible,
-				/* ArraySize = */ bIsOrthoProjection ? 3 : 2);
+				/* ArraySize = */ bIsOrthoProjection ? 2 : 1);
 
 			PrevAtomicTextureArray = GraphBuilder.CreateTexture(Desc, TEXT("TSR.PrevAtomics"));
 		}
@@ -1911,7 +1912,8 @@ FDefaultTemporalUpscaler::FOutputs AddTemporalSuperResolutionPasses(
 	FRDGTextureSRVRef ReprojectionBoundaryTexture = nullptr;
 	FRDGTextureSRVRef ReprojectionJacobianTexture = nullptr;
 	FRDGTextureRef ClosestDepthTexture;
-	FRDGTextureSRVRef DilateMaskTexture = nullptr;
+	FRDGTextureSRVRef DilateMaskTexture;
+	FRDGTextureSRVRef DepthErrorTexture;
 	FRDGTextureSRVRef IsMovingMaskTexture = nullptr;
 	FVelocityFlattenTextures VelocityFlattenTextures;
 	{
@@ -1943,13 +1945,14 @@ FDefaultTemporalUpscaler::FOutputs AddTemporalSuperResolutionPasses(
 				PF_R8_UINT,
 				FClearValueBinding::None,
 				/* InFlags = */ TexCreate_ShaderResource | TexCreate_UAV,
-				/* ArraySize = */ bOutputIsMovingTexture ? 2 : 1);
+				/* ArraySize = */ bOutputIsMovingTexture ? 3 : 2);
 
 			R8OutputTexture = GraphBuilder.CreateTexture(Desc, TEXT("TSR.DilateR8"));
 			DilateMaskTexture = GraphBuilder.CreateSRV(FRDGTextureSRVDesc::CreateForSlice(R8OutputTexture, 0));
+			DepthErrorTexture = GraphBuilder.CreateSRV(FRDGTextureSRVDesc::CreateForSlice(R8OutputTexture, 1));
 			if (bOutputIsMovingTexture)
 			{
-				IsMovingMaskTexture = GraphBuilder.CreateSRV(FRDGTextureSRVDesc::CreateForSlice(R8OutputTexture, 1));
+				IsMovingMaskTexture = GraphBuilder.CreateSRV(FRDGTextureSRVDesc::CreateForSlice(R8OutputTexture, 2));
 			}
 		}
 
@@ -2091,6 +2094,7 @@ FDefaultTemporalUpscaler::FOutputs AddTemporalSuperResolutionPasses(
 		PassParameters->InputSceneColorTexture = PassInputs.SceneColor.Texture;
 		PassParameters->ReprojectionVectorTexture = ReprojectionVectorTexture;
 		PassParameters->ClosestDepthTexture = ClosestDepthTexture;
+		PassParameters->DepthErrorTexture = DepthErrorTexture;
 		PassParameters->PrevAtomicTextureArray = PrevAtomicTextureArray;
 
 		PassParameters->PrevHistoryParameters = PrevHistoryParameters;
@@ -2878,6 +2882,7 @@ FDefaultTemporalUpscaler::FOutputs AddTemporalSuperResolutionPasses(
 			PassParameters->ReprojectionJacobianTexture = ReprojectionJacobianTexture ? ReprojectionJacobianTexture : GraphBuilder.CreateSRV(FRDGTextureSRVDesc(BlackUintDummy));
 			PassParameters->ReprojectionVectorTexture = ReprojectionVectorTexture;
 			PassParameters->IsMovingMaskTexture = IsMovingMaskTexture ? IsMovingMaskTexture : GraphBuilder.CreateSRV(FRDGTextureSRVDesc(BlackUintDummy));
+			PassParameters->DecimateMaskTexture = DecimateMaskTexture;
 			PassParameters->HistoryRejectionTexture = HistoryRejectionTexture;
 			PassParameters->MoireHistoryTexture = MoireHistoryTexture ? MoireHistoryTexture : GraphBuilder.CreateSRV(FRDGTextureSRVDesc(BlackDummy));
 			PassParameters->AntiAliasMaskTexture = AntiAliasMaskTexture ? AntiAliasMaskTexture : BlackUintDummy;
