Optimize y component of the RandomSequence API

A recent paper has shown how to generate this sequence more efficiently. The implementation used here is tweaked to account for the fact that the index is already reversed and that we want to generated a reversed output.

#tests Passes PT EngineTests
#rnx

diff --git a/Engine/Shaders/Private/PathTracing/Utilities/PathTracingRandomSequence.ush b/Engine/Shaders/Private/PathTracing/Utilities/PathTracingRandomSequence.ush
--- a/Engine/Shaders/Private/PathTracing/Utilities/PathTracingRandomSequence.ush
+++ b/Engine/Shaders/Private/PathTracing/Utilities/PathTracingRandomSequence.ush
@@ -301,19 +301,28 @@ float4 SobolSampler(uint SampleIndex, inout uint Seed)
 	// first scramble the index to decorelate from other 4-tuples
 	uint SobolIndex = FastOwenScrambling(SampleIndex, EvolveSobolSeed(Seed));
 	// now get Sobol' point from this index
-	uint4 Result = uint4(SobolIndex, 0, 0, 0);
+	uint4 Result = uint4(SobolIndex, SobolIndex, 0, 0);
+	// y component can be computed without iteration
+	// "An Implementation Algorithm of 2D Sobol Sequence Fast, Elegant, and Compact"
+	// Abdalla Ahmed, EGSR 2024
+	// See listing (19) in the paper
+	// The code is different here because we want the output to be bit-reversed, but
+	// the methodology is the same
+	Result.y ^=  Result.y               >> 16;
+	Result.y ^= (Result.y & 0xFF00FF00) >>  8;
+	Result.y ^= (Result.y & 0xF0F0F0F0) >>  4;
+	Result.y ^= (Result.y & 0xCCCCCCCC) >>  2;
+	Result.y ^= (Result.y & 0xAAAAAAAA) >>  1;
 #if RANDSEQ_UNROLL_SOBOL // unrolled version seems faster in most cases
-	UNROLL for (uint b = 0, v = 1; b < 32; b++)
+	UNROLL for (uint b = 0; b < 32; b++)
 	{
 		uint IndexBit = (SobolIndex >> b) & 1;		// bitfield extract
 #else // however a loop with early exit seems faster in the path tracing context (possibly due to low occupancy?)
-	for (uint b = 0, v = 1; SobolIndex; SobolIndex >>= 1, b++)
+	for (uint b = 0; SobolIndex; SobolIndex >>= 1, b++)
 	{
 		uint IndexBit = SobolIndex & 1;
 #endif
-		Result.y ^= IndexBit * v;
 		Result.zw ^= IndexBit * SobolMatrices[b];
-		v ^= v << 1;
 	}
 	// finally scramble the points to avoid structured artifacts
 	Result.x = FastOwenScrambling(Result.x, EvolveSobolSeed(Seed));
