Fix a bug that the scene is much brighter with reflection capture on mobile.

diff --git a/Engine/Shaders/Private/MobileLightingCommon.ush b/Engine/Shaders/Private/MobileLightingCommon.ush
--- a/Engine/Shaders/Private/MobileLightingCommon.ush
+++ b/Engine/Shaders/Private/MobileLightingCommon.ush
@@ -392,7 +392,7 @@ half3 GetImageBasedReflectionLighting_Mobile(
 
 	bool UsingSkyReflection = MobileReflectionCapture.Params.y > 0.0f;
 	// Normalize for static skylight types which mix with lightmaps.
-	bool bNormalize = UsingSkyReflection && MobileReflectionCapture.Params.z < 1;
+	bool bNormalize = !UsingSkyReflection || MobileReflectionCapture.Params.z < 1;
 	if (UsingSkyReflection)
 	{
 		// Apply sky colour if the reflection map is the sky.
