Fix LTC-based sheen integration with rect light.

The LTC matrices were inverted and had a missing scaling factor.

diff --git a/Engine/Shaders/Private/RectLight.ush b/Engine/Shaders/Private/RectLight.ush
--- a/Engine/Shaders/Private/RectLight.ush
+++ b/Engine/Shaders/Private/RectLight.ush
@@ -436,17 +436,17 @@ FRectLTC GetRectLTC_Sheen( float Roughness, half NoV)
 	const float bInv = SheenLTC.y;
 
 	float3x3 LTC = {
-		float3(1/aInv, 0,     -bInv/aInv),    
-		float3(0,      1/aInv, 0),
-		float3(0,      0,      1)
-	};
-
-	float3x3 InvLTC = {
 		float3(aInv, 0,    bInv),
 		float3(0,    aInv, 0),
 		float3(0,    0,    1)
 	};
 
+	float3x3 InvLTC = {
+		float3(1/aInv, 0,     -bInv/aInv),    
+		float3(0,      1/aInv, 0),
+		float3(0,      0,      1)
+	};
+
 	FRectLTC Out = (FRectLTC)0;
 	Out.LTC = LTC;
 	Out.InvLTC = InvLTC;
@@ -478,12 +478,14 @@ float3 RectApproxLTC(FRectLTC In, half3 N, float3 V, FRect Rect, FRectTexture Re
 	// Vector irradiance
 	float3 L = PolygonIrradiance( Poly );
 
+	#if 0
 	// Early out negative irradiance vector, causing 'ghost' rect light silhouette at low roughness
 	OutMeanLightWorldDirection = N;
 	if (L.z <= 0)
 	{
 		return 0;
 	}
+	#endif
 
 	float LengthSqr = dot( L, L );
 	float InvLength = rsqrt( LengthSqr );
@@ -567,8 +569,9 @@ float3 RectSheenApproxLTC( float Roughness, half3 N, float3 V, FRect Rect, FRect
 	const FRectLTC LTC = GetRectLTC_Sheen(Roughness, NoV);
 	DirectionalAlbedo = LTC.IrradianceScale.x;
 
+	const float Scale = 2 * PI; // Factor to match the path-tracer
 	float3 MeanLightWorldDirection = 0.0f;
-	return RectApproxLTC(LTC, N, V, Rect, RectTexture, MeanLightWorldDirection);
+	return RectApproxLTC(LTC, N, V, Rect, RectTexture, MeanLightWorldDirection) * Scale;
 }
 
 // Rectangle projected to a sphere
