Fix incorrect legacy model detection causing SSS material with Fuzz to be converted into legacy wrap SSS.

diff --git a/Engine/Shaders/Private/Substrate/SubstrateExport.ush b/Engine/Shaders/Private/Substrate/SubstrateExport.ush
--- a/Engine/Shaders/Private/Substrate/SubstrateExport.ush
+++ b/Engine/Shaders/Private/Substrate/SubstrateExport.ush
@@ -371,12 +371,12 @@ void PackSubstrateOut(
 				{
 					OptimisedLegacyMode = SINGLE_OPTLEGACYMODE_CLOTH;
 				}				
-				else if ((OptimisedModeMask & SingleIsWrapThin) == SingleIsWrapThin)
+				else if ((OptimisedModeMask & SingleIsWrapThin) == SingleIsWrapThin && !BSDF_GETHASFUZZ(BSDF))
 				{
 					OptimisedLegacyMode = SINGLE_OPTLEGACYMODE_SSSWRAP_THIN;
 					SSSData.Header.Bytes = 0; // Force to skip SSSData writing
 				}
-				else if ((OptimisedModeMask & SingleIsWrap) == SingleIsWrap)
+				else if ((OptimisedModeMask & SingleIsWrap) == SingleIsWrap && !BSDF_GETHASFUZZ(BSDF))
 				{
 					OptimisedLegacyMode = SINGLE_OPTLEGACYMODE_SSSWRAP;
 					SSSData.Header.Bytes = 0; // Force to skip SSSData writing
