Move Substrate Fuzz after SSS evaluation for getting correct fuzz coverage.

diff --git a/Engine/Shaders/Private/Substrate/SubstrateEvaluation.ush b/Engine/Shaders/Private/Substrate/SubstrateEvaluation.ush
--- a/Engine/Shaders/Private/Substrate/SubstrateEvaluation.ush
+++ b/Engine/Shaders/Private/Substrate/SubstrateEvaluation.ush
@@ -700,87 +700,5 @@
 			Sample.TransmittanceAlongN	= OpaqueBSDFThroughput;
 
 			////
-			//// Evaluate cloth fuzz layered on top of the slab.
-			////
-
-			#if SUBSTRATE_FASTPATH==0
-			BRANCH
-			if (BSDF_GETHASFUZZ(BSDFContext.BSDF))
-			{
-				const float  FuzzAmount = SLAB_FUZZ_AMOUNT(BSDFContext.BSDF);
-				const float3 FuzzF0  = SLAB_FUZZ_COLOR(BSDFContext.BSDF);
-				const float  FuzzRoughness = MakeRoughnessSafe(SLAB_FUZZ_ROUGHNESS(BSDFContext.BSDF), SUBSTRATE_MIN_FUZZ_ROUGHNESS);
-
-				// F0 is ignored, as energy preservation is handle solely based on the white directional albedo and FuzzAmount
-				FBxDFEnergyTermsA EnergyTerms = (FBxDFEnergyTermsA)1.f;
-				float3 ClothSpecularPathValueNoL = 0;
-
-				#if SUBSTRATE_SHEEN_QUALITY == 1
-				// Disney LTC
-				{
-					float DirectionalAlbedo = 1;
-					BRANCH
-					if (bIsRectLight)
-					{
-						ClothSpecularPathValueNoL = FuzzF0 * RectSheenApproxLTC(FuzzRoughness, BSDFContext.N, BSDFContext.V, AreaLightContext.AreaLight.Rect, AreaLightContext.AreaLight.Texture, DirectionalAlbedo);
-					}
-					else
-					{
-						ClothSpecularPathValueNoL = FuzzF0 * SheenLTC_Eval(BSDFContext.V, AreaLightContext.L, BSDFContext.N, BSDFContext.Context.NoV, FuzzRoughness, View.SheenLTCTexture, View.SheenLTCSampler, DirectionalAlbedo);
-						ClothSpecularPathValueNoL *= DirectionalAlbedo;
-					}
-					EnergyTerms.E = DirectionalAlbedo;
-					EnergyTerms.W = 1.f;
-				}
-				#else
-				// Charlie cloth modle for D and Ashikhmin for visibility (more efficient than Charlie).
-				{
-					// ComputeEnergyConservation(EnergyTerms) is not taken into account because we do not white white fuzz to reflect all incoming energy
-					// The fuzz should have some transmission to the lower lobes. In order to pass a furnace test, a combination of white diffuse+white fuzz is required
-					EnergyTerms = ComputeClothEnergyTermsA(FuzzRoughness, BSDFContext.Context.NoV);
-
-					float ClothD   = D_Charlie(FuzzRoughness, BSDFContext.Context.NoH);
-					float ClothVis = Vis_Ashikhmin(BSDFContext.Context.NoV, AreaLightContext.NoL);
-					float3 ClothF  = F_Schlick(FuzzF0, BSDFContext.Context.VoH);
-					ClothSpecularPathValueNoL = (ClothD * ClothVis) * ClothF * ComputeEnergyConservation(EnergyTerms) * AreaLightContext.NoL;
-				}
-				#endif //SUBSTRATE_SHEEN_QUALITY
-
-				// The specular and diffuse components below the fuzz are only attenuated linearly according to the amount of fuzz.
-				const float3 Cloth_DirectionalAlbedo_SpecularTransmission = lerp(1.0, ComputeEnergyPreservation(EnergyTerms), FuzzAmount);
-
-				// Area light are not supported by the cloth BRDF
-				float3 ClothIntegratedSpecularValue = (ShadowTerms.SurfaceShadow * AreaLightContext.Falloff * FuzzAmount) * AreaLightContext.AreaLight.FalloffColor * ClothSpecularPathValueNoL;
-				
-				// Apply specular transmittance to diffuse and specular lob from slab medium and interface sitting below the layer of fuzz.
-				Sample.DiffusePathValue			*= Cloth_DirectionalAlbedo_SpecularTransmission;
-				Sample.IntegratedDiffuseValue	*= Cloth_DirectionalAlbedo_SpecularTransmission;
-				Sample.SpecularPathValue		*= Cloth_DirectionalAlbedo_SpecularTransmission;
-				Sample.IntegratedSpecularValue	*= Cloth_DirectionalAlbedo_SpecularTransmission;
-				
-				Sample.SpecularPathValue		+= ClothSpecularPathValueNoL;
-				Sample.IntegratedSpecularValue  += ClothIntegratedSpecularValue;
-
-				// This is not good. We should really have a separate PDF for cloth it self associated with a probability.
-				// We should also output probability to sample cloth, diffuse, secular, etc.
-				// We will revisit when the path tracer is getting up with Substrate.
-				//Sample.SpecularPDF = lerp(Sample.SpecularPDF, BSDFContext.SatNoV / PI, FuzzAmount); // Per "Production Friendly Microfacet Sheen BRDF", hemispherical sampling give good result as the roughness is usually high.
-				// SUBSTRATE_TODO we should have a separated cloth PDF and Weight
-
-				if (bHaziness)
-				{
-					Sample.SpecularHazePathValue*= Cloth_DirectionalAlbedo_SpecularTransmission;
-					Sample.SpecularHazePathValue+= ClothSpecularPathValueNoL;
-
-					// SUBSTRATE_TODO we should have a separated cloth PDF and Weight
-					//Sample.SpecularHazePDF		= lerp(Sample.SpecularHazePDF,			BSDFContext.SatNoV / PI,	Fuzz);;
-				}
-
-				Sample.ThroughputV *= Cloth_DirectionalAlbedo_SpecularTransmission;
-				// The specular transmission is ignored towards the direction of the light. (TransmittanceAlongN can only store transmittance)
-			}
-			#endif
-
-			////
 			//// Evaluate approximated SSS (using wrap lighting). Temp code. SUBSTRATE_TODO: unify evaluation
 			////
@@ -978,6 +896,89 @@
 			}
 			#endif
 
+			
+			////
+			//// Evaluate cloth fuzz layered on top of the slab.
+			////
+
+			#if SUBSTRATE_FASTPATH==0
+			BRANCH
+			if (BSDF_GETHASFUZZ(BSDFContext.BSDF))
+			{
+				const float  FuzzAmount = SLAB_FUZZ_AMOUNT(BSDFContext.BSDF);
+				const float3 FuzzF0  = SLAB_FUZZ_COLOR(BSDFContext.BSDF);
+				const float  FuzzRoughness = MakeRoughnessSafe(SLAB_FUZZ_ROUGHNESS(BSDFContext.BSDF), SUBSTRATE_MIN_FUZZ_ROUGHNESS);
+
+				// F0 is ignored, as energy preservation is handle solely based on the white directional albedo and FuzzAmount
+				FBxDFEnergyTermsA EnergyTerms = (FBxDFEnergyTermsA)1.f;
+				float3 ClothSpecularPathValueNoL = 0;
+
+				#if SUBSTRATE_SHEEN_QUALITY == 1
+				// Disney LTC
+				{
+					float DirectionalAlbedo = 1;
+					BRANCH
+					if (bIsRectLight)
+					{
+						ClothSpecularPathValueNoL = FuzzF0 * RectSheenApproxLTC(FuzzRoughness, BSDFContext.N, BSDFContext.V, AreaLightContext.AreaLight.Rect, AreaLightContext.AreaLight.Texture, DirectionalAlbedo);
+					}
+					else
+					{
+						ClothSpecularPathValueNoL = FuzzF0 * SheenLTC_Eval(BSDFContext.V, AreaLightContext.L, BSDFContext.N, BSDFContext.Context.NoV, FuzzRoughness, View.SheenLTCTexture, View.SheenLTCSampler, DirectionalAlbedo);
+						ClothSpecularPathValueNoL *= DirectionalAlbedo;
+					}
+					EnergyTerms.E = DirectionalAlbedo;
+					EnergyTerms.W = 1.f;
+				}
+				#else
+				// Charlie cloth modle for D and Ashikhmin for visibility (more efficient than Charlie).
+				{
+					// ComputeEnergyConservation(EnergyTerms) is not taken into account because we do not white white fuzz to reflect all incoming energy
+					// The fuzz should have some transmission to the lower lobes. In order to pass a furnace test, a combination of white diffuse+white fuzz is required
+					EnergyTerms = ComputeClothEnergyTermsA(FuzzRoughness, BSDFContext.Context.NoV);
+
+					float ClothD   = D_Charlie(FuzzRoughness, BSDFContext.Context.NoH);
+					float ClothVis = Vis_Ashikhmin(BSDFContext.Context.NoV, AreaLightContext.NoL);
+					float3 ClothF  = F_Schlick(FuzzF0, BSDFContext.Context.VoH);
+					ClothSpecularPathValueNoL = (ClothD * ClothVis) * ClothF * ComputeEnergyConservation(EnergyTerms) * AreaLightContext.NoL;
+				}
+				#endif //SUBSTRATE_SHEEN_QUALITY
+
+				// The specular and diffuse components below the fuzz are only attenuated linearly according to the amount of fuzz.
+				const float3 Cloth_DirectionalAlbedo_SpecularTransmission = lerp(1.0, ComputeEnergyPreservation(EnergyTerms), FuzzAmount);
+
+				// Area light are not supported by the cloth BRDF
+				float3 ClothIntegratedSpecularValue = (ShadowTerms.SurfaceShadow * AreaLightContext.Falloff * FuzzAmount) * AreaLightContext.AreaLight.FalloffColor * ClothSpecularPathValueNoL;
+				
+				// Apply specular transmittance to diffuse and specular lob from slab medium and interface sitting below the layer of fuzz.
+				Sample.DiffusePathValue			*= Cloth_DirectionalAlbedo_SpecularTransmission;
+				Sample.IntegratedDiffuseValue	*= Cloth_DirectionalAlbedo_SpecularTransmission;
+				Sample.SpecularPathValue		*= Cloth_DirectionalAlbedo_SpecularTransmission;
+				Sample.IntegratedSpecularValue	*= Cloth_DirectionalAlbedo_SpecularTransmission;
+				
+				Sample.SpecularPathValue		+= ClothSpecularPathValueNoL;
+				Sample.IntegratedSpecularValue  += ClothIntegratedSpecularValue;
+
+				// This is not good. We should really have a separate PDF for cloth it self associated with a probability.
+				// We should also output probability to sample cloth, diffuse, secular, etc.
+				// We will revisit when the path tracer is getting up with Substrate.
+				//Sample.SpecularPDF = lerp(Sample.SpecularPDF, BSDFContext.SatNoV / PI, FuzzAmount); // Per "Production Friendly Microfacet Sheen BRDF", hemispherical sampling give good result as the roughness is usually high.
+				// SUBSTRATE_TODO we should have a separated cloth PDF and Weight
+
+				if (bHaziness)
+				{
+					Sample.SpecularHazePathValue*= Cloth_DirectionalAlbedo_SpecularTransmission;
+					Sample.SpecularHazePathValue+= ClothSpecularPathValueNoL;
+
+					// SUBSTRATE_TODO we should have a separated cloth PDF and Weight
+					//Sample.SpecularHazePDF		= lerp(Sample.SpecularHazePDF,			BSDFContext.SatNoV / PI,	Fuzz);;
+				}
+
+				Sample.ThroughputV *= Cloth_DirectionalAlbedo_SpecularTransmission;
+				// The specular transmission is ignored towards the direction of the light. (TransmittanceAlongN can only store transmittance)
+			}
+			#endif
+
 			break;
 		}
 
