Substrate
- Added support for SSR control from mateiral
- reflection environment now support skipping non SSR tiles pixels

#rnx

diff --git a/Engine/Shaders/Private/DefaultSSRTiles.usf b/Engine/Shaders/Private/DefaultSSRTiles.usf
--- a/Engine/Shaders/Private/DefaultSSRTiles.usf
+++ b/Engine/Shaders/Private/DefaultSSRTiles.usf
@@ -45,17 +45,51 @@ bool DoesPixelContainSSR(uint2 PixelPos)
 	FSubstrateAddressing SubstrateAddressing = GetSubstratePixelDataByteOffset(PixelPos, uint2(View.BufferSizeAndInvSize.xy), Substrate.MaxBytesPerPixel);
 	FSubstratePixelHeader SubstratePixelHeader = UnpackSubstrateHeaderIn(Substrate.MaterialTextureArray, SubstrateAddressing, Substrate.TopLayerTexture);
 	const FSubstrateTopLayerData TopLayerData = SubstrateUnpackTopLayerData(SubstratePixelHeader.PackedTopLayerData);
-	bool bContainValidPixels =
-		TopLayerData.Roughness <= MaxRoughness
-		// TODO_SUBSTRATE: need F0 value in order to have tge sane optimization as legacy
-		//&& TopLayerData.Specular >= MinSpecular
-		&& !SubstratePixelHeader.IsHair()
-		&& !SubstratePixelHeader.IsSingleLayerWater();
-	
-	//TODO: add substrate support
-	//bContainValidPixels = bContainValidPixels && (bEnableTwoSidedFoliage ? true : !SubstratePixelHeader.IsTwoSidedFoliage());
-
-	return (SubstratePixelHeader.ClosureCount > 0 && bContainValidPixels);
+
+	bool bRoughnessIsValid = TopLayerData.Roughness <= MaxRoughness;
+	bool bContainValidPixels = SubstratePixelHeader.ClosureCount > 0 && bRoughnessIsValid;
+	BRANCH
+	if (SubstratePixelHeader.ClosureCount == 1) // Only non complex materials will be interpreted for exclusion for now when it comes to Specular values or legacy shading models.
+	{
+		const FSubstrateSubsurfaceHeader SSSHeader = SubstrateLoadSubsurfaceHeader(Substrate.MaterialTextureArray, Substrate.FirstSliceStoringSubstrateSSSData, SubstrateAddressing.PixelCoords);
+
+		const bool bSkipSSSMaterialOverride = true; // We do not want the material data to be affected by any override.
+		FSubstrateBSDF BSDF = UnpackSubstrateBSDFIn(Substrate.MaterialTextureArray, SubstrateAddressing, SubstratePixelHeader, bSkipSSSMaterialOverride);
+
+		const float BSDFSpecular = F0ToDielectricSpecular(max3(SLAB_F0(BSDF).x, SLAB_F0(BSDF).y, SLAB_F0(BSDF).z));
+
+		bContainValidPixels =
+			bRoughnessIsValid
+			&& BSDFSpecular >= MinSpecular
+			&& !SubstratePixelHeader.IsHair()
+			&& !SubstratePixelHeader.IsSingleLayerWater();
+
+		bContainValidPixels = bContainValidPixels && (bEnableTwoSidedFoliage ? true : !BSDF_GETISTHIN(BSDF));
+
+		const uint BSDFType = BSDF_GETTYPE(BSDF);
+		switch (BSDFType)
+		{
+		case SUBSTRATE_BSDF_TYPE_SLAB:
+		{
+			// Cull special case where specular=0, and roughness=0
+			if (SLAB_ROUGHNESS(BSDF) == 0 && BSDFSpecular == 0)
+			{
+				bContainValidPixels = false;
+			}
+
+			// Clearcoat: if the material is clear coat it contains valid pixel.
+			const bool bHaziness = BSDF_GETHASHAZINESS(BSDF);
+			if (bHaziness)
+			{
+				const FHaziness Haziness = UnpackHaziness(SLAB_HAZINESS(BSDF));
+				const bool bHazeAsSimpleClearCoat = Haziness.bSimpleClearCoat;
+				bContainValidPixels |= bHazeAsSimpleClearCoat;
+			}
+		}
+		}
+	}
+
+	return bContainValidPixels;
 #else
 	float2 BufferUV = (PixelPos + 0.5f) * View.BufferSizeAndInvSize.zw;
 	FGBufferData GBuffer = GetGBufferDataFromSceneTextures(BufferUV);
diff --git a/Engine/Shaders/Private/ReflectionEnvironmentPixelShader.usf b/Engine/Shaders/Private/ReflectionEnvironmentPixelShader.usf
--- a/Engine/Shaders/Private/ReflectionEnvironmentPixelShader.usf
+++ b/Engine/Shaders/Private/ReflectionEnvironmentPixelShader.usf
@@ -402,10 +402,19 @@ void ReflectionEnvironmentSkyLighting(
 	#if REFLECTION_SOURCE_TYPE == REFLECTION_SOURCE_LUMEN_STANDALONE
 		const float TopLayerSpecularContributionFactor = 1.0f;
 	#else
-		const float4 SSRReflection = Texture2DSample(ReflectionTexture, ReflectionTextureSampler, BufferUV);
-		// SSR is computed for the top level by averaging BSDF components.
-		// This is the top level specular contribution factor according to areas.
-		const float TopLayerSpecularContributionFactor = 1.0f - SSRReflection.a;
+		float4 SSRReflection = 0.0f.xxxx;
+		float TopLayerSpecularContributionFactor = 1.0f;
+
+		#if REFLECTION_SOURCE_TYPE == REFLECTION_SOURCE_TILED_SSR
+		if (IsSSRTileUsed(PixelPos, SSRTiledViewRes, SSRTileMaskBuffer))
+		#endif
+		{
+			SSRReflection = Texture2DSample(ReflectionTexture, ReflectionTextureSampler, BufferUV);
+
+			// SSR is computed for the top level by averaging BSDF components.
+			// This is the top level specular contribution factor according to areas.
+			TopLayerSpecularContributionFactor = 1.0f - SSRReflection.a;
+		}
 
 		// A reduction of SSR contribution is needed when there is haziness because we only consider SSR for the sharpest lobe.
 		FSSRContribution SSRContribution = InitSSRContribution();
diff --git a/Engine/Shaders/Public/SubstratePublic.ush b/Engine/Shaders/Public/SubstratePublic.ush
--- a/Engine/Shaders/Public/SubstratePublic.ush
+++ b/Engine/Shaders/Public/SubstratePublic.ush
@@ -51,7 +51,7 @@ FSubstrateGBuffer SubstratePublic_GetSubstrateGBufferFromFirstBSDF(uint2 PixelCo
 	{
 		const FSubstrateSubsurfaceHeader SSSHeader = SubstrateLoadSubsurfaceHeader(SubstratePublicStruct.MaterialTextureArray, SubstratePublicStruct.FirstSliceStoringSubstrateSSSData, SubstrateAddressing.PixelCoords);
 
-		const bool bSkipSSSMaterialOverride = true; // We double not want the material data to be affected by any override.
+		const bool bSkipSSSMaterialOverride = true; // We do not want the material data to be affected by any override.
 		FSubstrateBSDF BSDF = UnpackSubstrateBSDFIn(SubstratePublicStruct.MaterialTextureArray, SubstrateAddressing, SubstratePixelHeader, bSkipSSSMaterialOverride);
 	
 		SGBuffer.bIsValid			= true;
