Substrate - fixed non Lumen SSR being broken due to exposure being applied twice (SSR buffer is already pre exposed).
#rnx

diff --git a/Engine/Shaders/Private/ReflectionEnvironmentPixelShader.usf b/Engine/Shaders/Private/ReflectionEnvironmentPixelShader.usf
--- a/Engine/Shaders/Private/ReflectionEnvironmentPixelShader.usf
+++ b/Engine/Shaders/Private/ReflectionEnvironmentPixelShader.usf
@@ -499,11 +499,10 @@ void ReflectionEnvironmentSkyLighting(
 			AccumulateSubstrateDeferredLighting(SubstrateLighting, Out, SubstrateEnvLight.bPostProcessSubsurface, BSDF_GETISTOPLAYER(BSDF));
 		}
 
-		OutColor = SubstrateLighting.SceneColor;
+		OutColor = SubstrateLighting.SceneColor * View.PreExposure;
 		#if REFLECTION_SOURCE_TYPE != REFLECTION_SOURCE_LUMEN_STANDALONE
 		OutColor += ApplySSRContribution(SSRContribution, SSRReflection);
 		#endif
-		OutColor.rgb *= View.PreExposure;
 
 		#if SUBSTRATE_OPAQUE_ROUGH_REFRACTION_ENABLED
 		OutOpaqueRoughRefractionSceneColor += SubstrateLighting.OpaqueRoughRefractionSceneColor;
