Fixes an issue where the parallax disocclusion mask threshold was rendering resolution dependent

diff --git a/Engine/Shaders/Private/TemporalSuperResolution/TSRCommon.ush b/Engine/Shaders/Private/TemporalSuperResolution/TSRCommon.ush
--- a/Engine/Shaders/Private/TemporalSuperResolution/TSRCommon.ush
+++ b/Engine/Shaders/Private/TemporalSuperResolution/TSRCommon.ush
@@ -71,9 +71,7 @@
 #define SPATIAL_ANTI_ALIASING_OFFSET_RANGE 0.5
 #define SPATIAL_ANTI_ALIASING_OFFSET_BITDEPTH_PER_CHANNEL 4
 
-#define MAX_DEVICE_Z_ERROR 16.0
-
-#define PARALLAX_REJECTION_MASK_THRESHOLD 0.5
+#define MAX_DEVICE_Z_ERROR 64.0
 
 // Minimum number of sample of each texel when reprojecting the high frequency history.
 // Value too low causes area of low sample count to have pixel dominating too much in reprojection causing input res aliasing on texture details
@@ -919,17 +917,46 @@ tsr_half2 DecodeSpatialAntiAliasingOffset(tsr_ushort EncodedInputTexelOffset)
 #define VELOCITY_HOLE_FILLING_ANGLE_BITS (VELOCITY_HOLE_FILLING_BITS - VELOCITY_HOLE_FILLING_LENGTH_BITS)
 #define VELOCITY_HOLE_FILLING_LENGTH_PRECISION (pow(2.0, 5))
 
+CALL_SITE_DEBUGLOC
+void PolarCoordinateHoleFillingVelocity(float2 PixelVelocity, out float PixelVelocityAngle, out float PixelVelocityLength)
+{
+	PixelVelocityAngle = atan2Fast(PixelVelocity.y, PixelVelocity.x);
+	PixelVelocityLength = length(PixelVelocity);
+}
+
+CALL_SITE_DEBUGLOC
+float2 CartesianHoleFillingVelocity(float PixelVelocityAngle, float PixelVelocityLength)
+{
+	float2 PixelDirection;
+	sincos(PixelVelocityAngle, /* out */ PixelDirection.y, /* out */ PixelDirection.x);
+	return PixelDirection * PixelVelocityLength;
+}
+
+CALL_SITE_DEBUGLOC
+void QuantizeHoleFillingVelocity(float2 PixelVelocity, out float EncodedAngle, out float EncodedLength)
+{
+	const uint AngleBitDepth = 1 << VELOCITY_HOLE_FILLING_ANGLE_BITS;
+	const float LengthPixelPrecision = VELOCITY_HOLE_FILLING_LENGTH_PRECISION;
+
+	float PixelVelocityAngle;
+	float PixelVelocityLength;
+	PolarCoordinateHoleFillingVelocity(PixelVelocity, /* out */ PixelVelocityAngle, /* out */ PixelVelocityLength);
+
+	EncodedAngle = PixelVelocityAngle * (0.5 * AngleBitDepth / PI) + (0.5 * AngleBitDepth);
+	EncodedLength = PixelVelocityLength * LengthPixelPrecision;
+}
+
 CALL_SITE_DEBUGLOC
 uint EncodeHoleFillingVelocity(float2 PixelVelocity)
 {
 	const uint AngleBitDepth = 1 << VELOCITY_HOLE_FILLING_ANGLE_BITS;
 	const uint LengthBitDepth = 1 << VELOCITY_HOLE_FILLING_LENGTH_BITS;
-	const float LengthPixelPrecision = VELOCITY_HOLE_FILLING_LENGTH_PRECISION;
 
-	float EncodedAngle = atan2Fast(PixelVelocity.y, PixelVelocity.x) * (0.5 * AngleBitDepth / PI) + (0.5 * AngleBitDepth);
-	uint iEncodedAngle = uint(round(EncodedAngle)) & (AngleBitDepth - 1);
+	float EncodedAngle;
+	float EncodedLength;
+	QuantizeHoleFillingVelocity(PixelVelocity, /* out */ EncodedAngle, /* out */ EncodedLength);
 
-	float EncodedLength = length(PixelVelocity) * LengthPixelPrecision;
+	uint iEncodedAngle = uint(round(EncodedAngle)) & (AngleBitDepth - 1);
 	uint iEncodedLength = fastClamp(uint(ceil(EncodedLength)), 0u, LengthBitDepth - 1u);
 
 	uint EncodedHoleFillingVelocity = iEncodedAngle * LengthBitDepth + iEncodedLength;
@@ -937,26 +964,29 @@ uint EncodeHoleFillingVelocity(float2 PixelVelocity)
 }
 
 CALL_SITE_DEBUGLOC
-void DecodeHoleFillingVelocity(uint EncodedHoleFillingVelocity, out float2 PixelVelocity, out float OutOfRangeFactor)
+void DecodeHoleFillingVelocity(uint EncodedHoleFillingVelocity, out float PixelAngle, out float PixelLength)
 {
 	const uint AngleBitDepth = 1 << VELOCITY_HOLE_FILLING_ANGLE_BITS;
 	const uint LengthBitDepth = 1 << VELOCITY_HOLE_FILLING_LENGTH_BITS;
 	const float LengthPixelPrecision = VELOCITY_HOLE_FILLING_LENGTH_PRECISION;
-	const float LengthRange = float(LengthBitDepth) / LengthPixelPrecision;
 
 	uint iEncodedLength = EncodedHoleFillingVelocity % LengthBitDepth;
 	uint iEncodedAngle = (EncodedHoleFillingVelocity >> VELOCITY_HOLE_FILLING_LENGTH_BITS) % AngleBitDepth;
 
 	float EncodedLength = float(iEncodedLength);
-	float Length = EncodedLength * rcp(LengthPixelPrecision);
+	PixelLength = EncodedLength * rcp(LengthPixelPrecision);
 
 	float EncodedAngle = float(iEncodedAngle);
-	float Angle = EncodedAngle * (PI * 2.0 / float(AngleBitDepth)) - PI;
-
-	sincos(Angle, /* out */ PixelVelocity.y, /* out */ PixelVelocity.x);
-	PixelVelocity *= Length;
+	PixelAngle = EncodedAngle * (PI * 2.0 / float(AngleBitDepth)) - PI;
+}
 
-	OutOfRangeFactor = saturate(LengthRange - Length);
+CALL_SITE_DEBUGLOC
+float GetMaxEncodableHoleFillingPixelLength()
+{
+	const uint LengthBitDepth = 1 << VELOCITY_HOLE_FILLING_LENGTH_BITS;
+	const float LengthPixelPrecision = VELOCITY_HOLE_FILLING_LENGTH_PRECISION;
+	const float LengthRange = (float(LengthBitDepth) - 1.5) / LengthPixelPrecision;
+	return LengthRange;
 }
 
 CALL_SITE_DEBUGLOC
diff --git a/Engine/Shaders/Private/TemporalSuperResolution/TSRDecimateHistory.usf b/Engine/Shaders/Private/TemporalSuperResolution/TSRDecimateHistory.usf
--- a/Engine/Shaders/Private/TemporalSuperResolution/TSRDecimateHistory.usf
+++ b/Engine/Shaders/Private/TemporalSuperResolution/TSRDecimateHistory.usf
@@ -61,6 +61,7 @@ FScreenTransform ScreenPosToResurrectionGuideBufferUV;
 float2 ResurrectionGuideUVViewportBilinearMin;
 float2 ResurrectionGuideUVViewportBilinearMax;
 float3 HistoryGuideQuantizationError;
+float ParallaxRejectionMaskThreshold;
 float PerceptionAdd;
 float ResurrectionFrameIndex;
 float PrevFrameIndex;
@@ -265,7 +266,8 @@ void ProcessPrevAtomicTexturesSamples(
 	float2 ScreenVelocity,
 	float PrevDeviceZ,
 	float DeviceZError,
-	out tsr_half ParallaxRejectionMask,
+	bool bIsOffScreen,
+	out bool bIsParallaxDisocclusion,
 	out float2 HoleFillingPixelVelocity,
 	out bool bCanHoleFill)
 {
@@ -277,7 +279,7 @@ void ProcessPrevAtomicTexturesSamples(
 	PrevInputBufferUV = INVARIANT(PrevInputBufferUV);
 	FBilinearSampleInfos BilinearInter = GetBilinearSampleLevelInfos(PrevInputBufferUV, InputInfo_Extent, InputInfo_ExtentInverse);
 	
-	ParallaxRejectionMask = tsr_half(0.0);
+	tsr_half ParallaxRejectionMask = tsr_half(0.0);
 	bool bIsValidHoleFillingPixelVelocity = false;
 	uint EncodedHoleFillingVelocity = 0;
 	#if CONFIG_SUPPORT_ORTHO_VIEW
@@ -358,19 +360,47 @@ void ProcessPrevAtomicTexturesSamples(
 		}
 	}
 	
-	float HoleFillingOutOfRange;
-	DecodeHoleFillingVelocity(EncodedHoleFillingVelocity, /* out */ HoleFillingPixelVelocity, /* out */ HoleFillingOutOfRange);
+	float HoleFillingPixelAngle;
+	float HoleFillingPixelLength;
+	DecodeHoleFillingVelocity(EncodedHoleFillingVelocity, /* out */ HoleFillingPixelAngle, /* out */ HoleFillingPixelLength);
 
-	bCanHoleFill = bIsValidHoleFillingPixelVelocity && HoleFillingOutOfRange > 0.5;
+	bool bIsEncodablePixelLength = HoleFillingPixelLength < GetMaxEncodableHoleFillingPixelLength();
 
-	// Dim down the parallax rejection mask if the velocity matches
+	HoleFillingPixelVelocity = CartesianHoleFillingVelocity(HoleFillingPixelAngle, HoleFillingPixelLength);
+	bCanHoleFill = bIsValidHoleFillingPixelVelocity && bIsEncodablePixelLength;
+	//bIsParallaxDisocclusion = !bIsOffScreen && ParallaxRejectionMask < tsr_half(ParallaxRejectionMaskThreshold);
+
+	// Dim down the parallax rejection mask if the hole fill velocity matches
+	FLATTEN
+	if (bIsEncodablePixelLength)
 	{
+		const float AnglePrecision = 2.0 * PI * pow(0.5, VELOCITY_HOLE_FILLING_ANGLE_BITS);
+		const float LengthPixelPrecision = VELOCITY_HOLE_FILLING_LENGTH_PRECISION;
+
 		float2 PixelVelocity = ScreenVelocity * ScreenVelocityToInputPixelVelocity;
+
+		float PixelVelocityAngle, PixelVelocityLength;
+		PolarCoordinateHoleFillingVelocity(PixelVelocity, /* out */ PixelVelocityAngle, /* out */ PixelVelocityLength);
 		
-		tsr_half VelocityDiff = tsr_half(length(HoleFillingPixelVelocity - PixelVelocity));
+		float CartesianCoordinatePrecision = AnglePrecision * PixelVelocityLength;
+
+		float PixelLengthDiff = abs(HoleFillingPixelLength - PixelVelocityLength) - 2.0;
+		float PixelAngleDiff = abs(PixelVelocityAngle - HoleFillingPixelAngle);
+		PixelAngleDiff = min(PixelAngleDiff, 2.0 * PI - PixelAngleDiff);
 		
-		ParallaxRejectionMask = lerp(tsr_half(1.0), ParallaxRejectionMask, select(bCanHoleFill, saturate(VelocityDiff), tsr_half(1.0)));
+		float VelocityDiff = length(HoleFillingPixelVelocity - PixelVelocity);
+		float PixelAngleCompare = saturate(2.0 - PixelAngleDiff / AnglePrecision);
+		float PixelLengthCompare = saturate(1.0 + rcp(LengthPixelPrecision) - PixelLengthDiff);
+
+		float MinParallaxRejectionMaskPolar = min(PixelLengthCompare, PixelAngleCompare);
+		float MinParallaxRejectionMaskCartesian = saturate(1.0 + CartesianCoordinatePrecision - VelocityDiff);
+		
+		float MinParallaxRejectionMask = lerp(MinParallaxRejectionMaskCartesian, MinParallaxRejectionMaskPolar, saturate(min(HoleFillingPixelLength, PixelVelocityLength) - 2.0));
+
+		ParallaxRejectionMask = max(ParallaxRejectionMask, tsr_half(MinParallaxRejectionMask));
 	}
+
+	bIsParallaxDisocclusion = !bIsOffScreen && ParallaxRejectionMask < tsr_half(ParallaxRejectionMaskThreshold);
 } // ProcessPrevAtomicTexturesSamples()
 
 template<typename Type>
@@ -714,7 +744,6 @@ void MainCS(
 		
 		// Compute the parralax rejection mask
 		float2 HoleFillingPixelVelocity;
-		tsr_half ParallaxRejectionMask;
 		bool bCanHoleFill;
 		ProcessPrevAtomicTexturesSamples(
 			HistoryClosestDeviceZSamples0,
@@ -723,13 +752,11 @@ void MainCS(
 			ScreenVelocity,
 			PrevDeviceZ,
 			DeviceZError,
-			/* out */ ParallaxRejectionMask,
+			bIsOffScreen,
+			/* out */ bIsParallaxDisocclusion,
 			/* out */ HoleFillingPixelVelocity,
 			/* out */ bCanHoleFill);
 
-		bIsParallaxDisocclusion = !bIsOffScreen && ParallaxRejectionMask < tsr_half(PARALLAX_REJECTION_MASK_THRESHOLD);
-		bool bIsDisoccluded = bIsOffScreen || bIsParallaxDisocclusion;
-
 		// LDS despill
 		#if CONFIG_MANUAL_LDS_SPILL
 		{
diff --git a/Engine/Shaders/Private/TemporalSuperResolution/TSRDilateVelocity.usf b/Engine/Shaders/Private/TemporalSuperResolution/TSRDilateVelocity.usf
--- a/Engine/Shaders/Private/TemporalSuperResolution/TSRDilateVelocity.usf
+++ b/Engine/Shaders/Private/TemporalSuperResolution/TSRDilateVelocity.usf
@@ -272,9 +272,11 @@ bool ComputeDepthBilaterals(
 	BRANCH
 	if (bCompute)
 	{
+		const float WeightOffset = 1.5;
+		const float PixelDistance = length(float2(DirectionE));
+
 		float VariationEW = abs(DevizeZC - (DevizeZE + DevizeZW) * 0.5);
 		float VariationSN = abs(DevizeZC - (DevizeZS + DevizeZN) * 0.5);
-		float VariationDiff = VariationEW - VariationSN;
 
 		float DepthDiffE = abs(DevizeZE - DevizeZC);
 		float DepthDiffW = abs(DevizeZW - DevizeZC);
@@ -283,18 +285,19 @@ bool ComputeDepthBilaterals(
 	
 		float DepthDiffEW = abs(DevizeZE - DevizeZW);
 		float DepthDiffSN = abs(DevizeZS - DevizeZN);
-	
+
 		float MinDepthDerivativeE = min3(DepthDiffEW, DepthDiffE, DepthDiffW);
 		float MinDepthDerivativeS = min3(DepthDiffSN, DepthDiffS, DepthDiffN);
 		
-		float DepthDerivativeE = MinDepthDerivativeE;
-		float DepthDerivativeS = MinDepthDerivativeS;
+		float PixelDeviceZError = ComputePixelDeviceZError(DevizeZC) * PixelDistance;
 
-		WeightE = tsr_half(saturate(1.1 - min(DepthDiffE, VariationEW) / DepthDerivativeE));
-		WeightW = tsr_half(saturate(1.1 - min(DepthDiffW, VariationEW) / DepthDerivativeE));
-	
-		WeightS = tsr_half(saturate(1.1 - min(DepthDiffS, VariationSN) / DepthDerivativeS));
-		WeightN = tsr_half(saturate(1.1 - min(DepthDiffN, VariationSN) / DepthDerivativeS));
+		float FinalDepthDerivativeE = min(max(MinDepthDerivativeE, PixelDeviceZError), max(MinDepthDerivativeS, 8.0 * PixelDeviceZError));
+		float FinalDepthDerivativeS = min(max(MinDepthDerivativeS, PixelDeviceZError), max(MinDepthDerivativeE, 8.0 * PixelDeviceZError));
+
+		WeightE = tsr_half(saturate(WeightOffset - min(DepthDiffE, VariationEW) / FinalDepthDerivativeE));
+		WeightW = tsr_half(saturate(WeightOffset - min(DepthDiffW, VariationEW) / FinalDepthDerivativeE));
+		WeightS = tsr_half(saturate(WeightOffset - min(DepthDiffS, VariationSN) / FinalDepthDerivativeS));
+		WeightN = tsr_half(saturate(WeightOffset - min(DepthDiffN, VariationSN) / FinalDepthDerivativeS));
 
 	
 		tsr_half TotalWeightEW = WeightE + WeightW;
@@ -513,18 +516,8 @@ tsr_half2 ComputeReprojectionBoundary(
 
 	float2 PixelVelocityDifferential = InputInfo_ViewportSize * (ScreenVelocity.xy - DilatedScreenVelocity.xy);
 
-	/* when disoccluding, increase the size. */
-	bool bIsDilatedDisoccluding = (
-		float(ReprojectionOffset.x) * PixelVelocityDifferential.x < 0.0 ||
-		float(ReprojectionOffset.y) * PixelVelocityDifferential.y > 0.0);
-
-	/** If Dilated velocity is coming towards this pixel, where an extrusion could possibly happen. */
-	bool bIsDilatedComing = (
-		float(ReprojectionOffset.x) * DilatedScreenVelocity.x < 0.0 || 
-		float(ReprojectionOffset.y) * DilatedScreenVelocity.y > 0.0);
-
 	float PixelVelocity = length2(PixelVelocityDifferential);
-	bool bDilationHasEnoughVelocityDiff = PixelVelocity > select(bIsDilatedDisoccluding, pow(0.5, 2.0), ReprojectionFieldAntiAliasVelocityThreshold);
+	bool bDilationHasEnoughVelocityDiff = PixelVelocity > ReprojectionFieldAntiAliasVelocityThreshold;
 
 	bool bAntiAliasReprojectionField = bReprojectionField && bDilationHasEnoughVelocityDiff && !bOffscreen;
 
diff --git a/Engine/Source/Runtime/Renderer/Private/PostProcess/TemporalSuperResolution.cpp b/Engine/Source/Runtime/Renderer/Private/PostProcess/TemporalSuperResolution.cpp
--- a/Engine/Source/Runtime/Renderer/Private/PostProcess/TemporalSuperResolution.cpp
+++ b/Engine/Source/Runtime/Renderer/Private/PostProcess/TemporalSuperResolution.cpp
@@ -748,6 +748,7 @@ class FTSRDecimateHistoryCS : public FTSRShader
 		SHADER_PARAMETER(FVector2f, ResurrectionGuideUVViewportBilinearMin)
 		SHADER_PARAMETER(FVector2f, ResurrectionGuideUVViewportBilinearMax)
 		SHADER_PARAMETER(FVector3f, HistoryGuideQuantizationError)
+		SHADER_PARAMETER(float, ParallaxRejectionMaskThreshold)
 		SHADER_PARAMETER(float, PerceptionAdd)
 		SHADER_PARAMETER(float, ResurrectionFrameIndex)
 		SHADER_PARAMETER(float, PrevFrameIndex)
@@ -2116,6 +2117,7 @@ FDefaultTemporalUpscaler::FOutputs AddTemporalSuperResolutionPasses(
 			PassParameters->ResurrectionGuideUVViewportBilinearMin = GetScreenPassTextureViewportParameters(ResurrectionGuideViewport).UVViewportBilinearMin;
 			PassParameters->ResurrectionGuideUVViewportBilinearMax = GetScreenPassTextureViewportParameters(ResurrectionGuideViewport).UVViewportBilinearMax;
 			PassParameters->HistoryGuideQuantizationError = ComputePixelFormatQuantizationError(PrevHistory.GuideArray->Desc.Format);
+			PassParameters->ParallaxRejectionMaskThreshold = 1.0f - 0.25f * OutputToInputResolutionFraction;
 			PassParameters->PerceptionAdd = FMath::Pow(0.5f, CVarTSRShadingExposureOffset.GetValueOnRenderThread());
 		}
 
