Fix crash when rendering mode than 32 lights with deep shadow.

diff --git a/Engine/Shaders/Private/HairStrands/HairStrandsDeepShadowAllocation.usf b/Engine/Shaders/Private/HairStrands/HairStrandsDeepShadowAllocation.usf
--- a/Engine/Shaders/Private/HairStrands/HairStrandsDeepShadowAllocation.usf
+++ b/Engine/Shaders/Private/HairStrands/HairStrandsDeepShadowAllocation.usf
@@ -35,9 +35,14 @@ float4x4 ComputeTranslatedWorldToLight(
 
 ////////////////////////////////////////////////////////////////////////////////
 
-float4				LightDirections[MAX_SLOT_COUNT];
-float4				TranslatedLightPositions[MAX_SLOT_COUNT];
-uint4				MacroGroupIndices[MAX_SLOT_COUNT];
+// Mirror of FLightData in HairStrands/HairStrandsDeepShadow.cpp
+struct FLightData
+{
+	float3 LightDirection;
+	uint MacroGroupId;
+	float3 TranslatedLightPosition;
+	uint bIsLightDirectional;
+};
 
 float				RasterizationScale;
 float				AABBScale;
@@ -52,6 +57,7 @@ float2				AtlasTexelSize;
 uint				MinAtlasTileResolution;
 uint				MinAtlasTileResolutionLog2;
 
+StructuredBuffer<FLightData> 			LightDataBuffer;
 Buffer<int>								MacroGroupAABBBuffer;
 RWStructuredBuffer<FDeepShadowViewInfo>	OutShadowViewInfoBuffer;
 
@@ -82,16 +88,12 @@ void ComputeTranslatedWorldToLightClip(
 	inout float		OutMinStrandRadiusAtDepth1,
 	inout float3	OutLightDirection,
 	const FHairAABB TranslatedWorld_AABB,
-	const uint		SlotIndex)
+	const FLightData LightData)
 {
-	const float3 LightDirection			  = LightDirections[SlotIndex].xyz;
-	const float3 TranslatedLightPosition  = TranslatedLightPositions[SlotIndex].xyz;
-	const bool bIsDirectional			  = TranslatedLightPositions[SlotIndex].w == 0;
-
 	float3 Center = GetCenter(TranslatedWorld_AABB);
-	OutLightDirection = bIsDirectional ? LightDirection : -normalize(TranslatedLightPosition - Center);
+	OutLightDirection = LightData.bIsLightDirectional ? LightData.LightDirection : -normalize(LightData.TranslatedLightPosition - Center);
 
-	const float4x4 Coarse_TranslatedWorldToLight = ComputeTranslatedWorldToLight(TranslatedWorld_AABB, LightDirection, TranslatedLightPosition, bIsDirectional);
+	const float4x4 Coarse_TranslatedWorldToLight = ComputeTranslatedWorldToLight(TranslatedWorld_AABB, LightData.LightDirection, LightData.TranslatedLightPosition, LightData.bIsLightDirectional);
 
 	const FHairAABB Light_AABB = Transform(TranslatedWorld_AABB, Coarse_TranslatedWorldToLight);
 	const float3 Light_Extents = GetExtents(Light_AABB);
@@ -104,9 +106,9 @@ void ComputeTranslatedWorldToLightClip(
 	OutMinStrandRadiusAtDepth1 = 1;
 	OutTranslatedWorldToClipTransform = 0;
 
-	if (bIsDirectional) // (LightType == LightType_Directional)
+	if (LightData.bIsLightDirectional)
 	{
-		const float4x4 TranslatedWorldToLight	= LookAtMatrix(Center - LightDirection * abs(MinZ), Center, float3(0, 0, 1));
+		const float4x4 TranslatedWorldToLight	= LookAtMatrix(Center - LightData.LightDirection * abs(MinZ), Center, float3(0, 0, 1));
 		const float4x4 ProjMatrix				= ReversedZOrthoMatrix(Radius, Radius, 1.f / (MaxZ-MinZ), 0);
 		OutTranslatedWorldToClipTransform		= mul(TranslatedWorldToLight, ProjMatrix);
 
@@ -115,15 +117,15 @@ void ComputeTranslatedWorldToLightClip(
 	}
 	else // if (LightType == LightType_Spot || LightType == LightType_Point || LightType == LightType_Rect)
 	{
-		const float LightDistanceToCenter = length(TranslatedLightPosition - Center);
+		const float LightDistanceToCenter = length(LightData.TranslatedLightPosition - Center);
 		MaxZ = max(0.2f, LightDistanceToCenter) + MaxZ;
 		MinZ = max(0.1f, LightDistanceToCenter) + MinZ;
 		MinZ = max(1.0f, MinZ);
 
-		const float SphereDistance			= length(TranslatedLightPosition - Center);
+		const float SphereDistance			= length(LightData.TranslatedLightPosition - Center);
 		const float HalfFov					= min(MaxHafFovInRad, atan(Radius / SphereDistance));
 
-		const float4x4 TranslatedWorldToLight	= LookAtMatrix(TranslatedLightPosition, Center, float3(0, 0, 1));
+		const float4x4 TranslatedWorldToLight	= LookAtMatrix(LightData.TranslatedLightPosition, Center, float3(0, 0, 1));
 		const float4x4 ProjMatrix				= ReversedZPerspectiveMatrix(HalfFov, 1, 1, MinZ, MaxZ);
 		OutTranslatedWorldToClipTransform		= mul(TranslatedWorldToLight, ProjMatrix);
 
@@ -305,15 +307,15 @@ void CreateViewInfo(uint2 DispatchThreadId : SV_DispatchThreadID)
 	FDeepShadowViewInfo ViewInfo;
 	if (SlotIndex < SlotIndexCount)
 	{
-		const uint MacroGroupId = MacroGroupIndices[SlotIndex].x;
+		const FLightData LightData = LightDataBuffer[SlotIndex];
 
 		FHairAABB TranslatedBound = InitHairAABB();
-		if (MacroGroupId < MacroGroupCount)
+		if (LightData.MacroGroupId < MacroGroupCount)
 		{
-			TranslatedBound = ReadHairAABB(MacroGroupId, MacroGroupAABBBuffer);
+			TranslatedBound = ReadHairAABB(LightData.MacroGroupId, MacroGroupAABBBuffer);
 		}
 
-		ComputeTranslatedWorldToLightClip(ViewInfo.TranslatedWorldToClip, ViewInfo.MinRadiusAtDepth1, ViewInfo.ViewForward, TranslatedBound, SlotIndex);
+		ComputeTranslatedWorldToLightClip(ViewInfo.TranslatedWorldToClip, ViewInfo.MinRadiusAtDepth1, ViewInfo.ViewForward, TranslatedBound, LightData);
 
 		float2 ProjectedShadowSize = ComputeProjectedScreenSize(ViewInfo.TranslatedWorldToClip);
 		AddToAtlasTileBin(ProjectedShadowSize, SlotIndex);
diff --git a/Engine/Source/Runtime/Renderer/Private/HairStrands/HairStrandsDeepShadow.cpp b/Engine/Source/Runtime/Renderer/Private/HairStrands/HairStrandsDeepShadow.cpp
--- a/Engine/Source/Runtime/Renderer/Private/HairStrands/HairStrandsDeepShadow.cpp
+++ b/Engine/Source/Runtime/Renderer/Private/HairStrands/HairStrandsDeepShadow.cpp
@@ -183,10 +183,6 @@ class FDeepShadowCreateViewInfoCS : public FGlobalShader
 
 	BEGIN_SHADER_PARAMETER_STRUCT(FParameters, )
 
-		SHADER_PARAMETER_ARRAY(FVector4f,	LightDirections,			[FHairStrandsDeepShadowData::MaxMacroGroupCount])
-		SHADER_PARAMETER_ARRAY(FVector4f,	TranslatedLightPositions,	[FHairStrandsDeepShadowData::MaxMacroGroupCount])
-		SHADER_PARAMETER_ARRAY(FIntVector4,	MacroGroupIndices,			[FHairStrandsDeepShadowData::MaxMacroGroupCount])
-
 		SHADER_PARAMETER(float, RasterizationScale)
 
 		SHADER_PARAMETER(FIntPoint, SlotResolution)
@@ -201,6 +197,7 @@ class FDeepShadowCreateViewInfoCS : public FGlobalShader
 		SHADER_PARAMETER(uint32, MinAtlasTileResolution)
 		SHADER_PARAMETER(uint32, MinAtlasTileResolutionLog2)
 
+		SHADER_PARAMETER_RDG_BUFFER_SRV(StructuredBuffer, LightDataBuffer)
 		SHADER_PARAMETER_RDG_BUFFER_SRV(Buffer<int>, MacroGroupAABBBuffer)
 		SHADER_PARAMETER_RDG_BUFFER_UAV(RWStructuredBuffer<FDeepShadowViewInfo>, OutShadowViewInfoBuffer)
 		SHADER_PARAMETER_STRUCT_REF(FViewUniformShaderParameters, ViewUniformBuffer)
@@ -368,23 +365,36 @@ void RenderHairStrandsDeepShadows(
 
 		DeepShadowResources.bIsGPUDriven = bDeepShadowGPUDriven;
 		{
-			check(TotalAtlasSlotIndex < FHairStrandsDeepShadowResources::MaxAtlasSlotCount);
-
 			// Allocate and create projection matrix and Min radius
 			// Stored FDeepShadowViewInfo structs
 			// See HairStrandsDeepShadowCommonStruct.ush for more details
 			FDeepShadowCreateViewInfoCS::FParameters* Parameters = GraphBuilder.AllocParameters<FDeepShadowCreateViewInfoCS::FParameters>();
 
+			struct FLightData
+			{
+				FVector3f LightDirection;
+				uint32 MacroGroupId;
+				FVector3f TranslatedLightPosition;
+				uint32 bIsLightDirectional;
+			};
+			static_assert(sizeof(FLightData) == 32u);
+			TArray<FLightData> LightData;
+			LightData.Reserve(DOMSlotCount);
 			for (FHairStrandsMacroGroupData& MacroGroup : MacroGroupDatas)
 			{
 				for (FHairStrandsDeepShadowData& DomData : MacroGroup.DeepShadowDatas)
-				{				
-					Parameters->LightDirections[DomData.AtlasSlotIndex]				= FVector4f(DomData.LightDirection.X, DomData.LightDirection.Y, DomData.LightDirection.Z, 0);
-					Parameters->TranslatedLightPositions[DomData.AtlasSlotIndex]	= FVector4f(DomData.TranslatedLightPosition.X, DomData.TranslatedLightPosition.Y, DomData.TranslatedLightPosition.Z, DomData.bIsLightDirectional ? 0 : 1);
-					Parameters->MacroGroupIndices[DomData.AtlasSlotIndex]			= FIntVector4(DomData.MacroGroupId, 0,0,0);
+				{
+					FLightData& Data = LightData.AddDefaulted_GetRef();
+					Data.LightDirection				= DomData.LightDirection;
+					Data.TranslatedLightPosition	= DomData.TranslatedLightPosition;
+					Data.MacroGroupId				= DomData.MacroGroupId;
+					Data.bIsLightDirectional		= DomData.bIsLightDirectional ? 1 : 0;
 				}
 			}
 
+			FRDGBufferRef LightDataBuffer= CreateStructuredBuffer(GraphBuilder, TEXT("Hair.DeepShadow.LightData"), sizeof(FLightData), LightData.Num(), LightData.GetData(), sizeof(FLightData) * LightData.Num());
+
+			Parameters->LightDataBuffer = GraphBuilder.CreateSRV(LightDataBuffer);
 			Parameters->SlotResolution = DeepShadowResources.AtlasSlotResolution;
 			Parameters->SlotIndexCount = DeepShadowResources.TotalAtlasSlotCount;
 			Parameters->MacroGroupCount = MacroGroupDatas.Num();
