Tightens the size of the TSR clamping box to all closest neighbors for better in motion clarity

diff --git a/Engine/Shaders/Private/TemporalSuperResolution/TSRUpdateHistory.usf b/Engine/Shaders/Private/TemporalSuperResolution/TSRUpdateHistory.usf
--- a/Engine/Shaders/Private/TemporalSuperResolution/TSRUpdateHistory.usf
+++ b/Engine/Shaders/Private/TemporalSuperResolution/TSRUpdateHistory.usf
@@ -9,49 +9,80 @@
 #define DEBUG_ARRAY_SIZE 8
 
 
+//------------------------------------------------------- ENUMS
+
+/** Shape of the scene color kernel. */
+	// One single sample
+	#define CONFIG_SAMPLES_1X1 0
+	
+	// Full 3x3
+	#define CONFIG_SAMPLES_3X3 1
+	
+	// 3x3 plus
+	#define CONFIG_SAMPLES_PLUS 2
+	
+	// 3x3 plus with a 6th sample in the closest corner.
+	#define CONFIG_SAMPLES_PLUS_AND_CORNER 3
+
+	// 3x3 plus but disable the furtherest samples for clamping box.
+	#define CONFIG_SAMPLES_PLUS_DISABLE_FAR 4
+	
+	// 3x3 plus but move the furtherest towards the closest corner.
+	#define CONFIG_SAMPLES_PLUS_MOVE_FAR 5
+
+
 //------------------------------------------------------- CONFIG
 
 #if DIM_UPDATE_QUALITY == 0 // Low
-	#define CONFIG_SAMPLES 5
+	#define CONFIG_SAMPLES CONFIG_SAMPLES_PLUS
 	#define CONFIG_REJECTION_ANTI_ALIASING 0
 
 #elif DIM_UPDATE_QUALITY == 1 // Medium
-	#define CONFIG_SAMPLES 5
+	#define CONFIG_SAMPLES CONFIG_SAMPLES_PLUS
 	#define CONFIG_REJECTION_ANTI_ALIASING 1
 
 #elif DIM_UPDATE_QUALITY == 2 // High
-	#define CONFIG_SAMPLES 6
+	#define CONFIG_SAMPLES CONFIG_SAMPLES_PLUS_MOVE_FAR
 	#define CONFIG_REJECTION_ANTI_ALIASING 1
 
 #elif DIM_UPDATE_QUALITY == 3 // Epic
-	#define CONFIG_SAMPLES 6
+	#define CONFIG_SAMPLES CONFIG_SAMPLES_PLUS_MOVE_FAR
 	#define CONFIG_REJECTION_ANTI_ALIASING 1
 
 #else
 	#error Unknown history update quality
 #endif
 
-#define CONFIG_COMPOSE_TRANSLUCENCY 0
-#define CONFIG_REJECT_TRANSLUCENCY_ON_RESPONSIVE_AA_ONLY 0
+
+//------------------------------------------------------- DEFAULTS
+
+/** Controls whether the shader is capable to do history resurrection. */
 #define CONFIG_HISTORY_RESURRECTION 1
+
+/** Controls whether the shader is capable to use the reprojection field. */
 #define CONFIG_REPROJECTION_FIELD 1
 
 /** Whether lens distortion is supported. */
 #define CONFIG_LENS_DISTORTION TSR_SUPPORT_LENS_DISTORTION
 
-
-//------------------------------------------------------- DEFAULTS
-
+/** Whether the history should do any rejection of any kind. */
 #define CONFIG_CLAMP 1
 
-#ifndef CONFIG_TRANSLUCENCY_REJECTION
-	#define CONFIG_TRANSLUCENCY_REJECTION 1
+/** Defines the number of samples of the scene color. */
+#if CONFIG_SAMPLES == CONFIG_SAMPLES_1X1
+	#define CONFIG_SAMPLES_COUNT 1
+#elif CONFIG_SAMPLES == CONFIG_SAMPLES_3X3
+	#define CONFIG_SAMPLES_COUNT 9
+#elif CONFIG_SAMPLES == CONFIG_SAMPLES_PLUS_6
+	#define CONFIG_SAMPLES_COUNT 6
+#else
+	#define CONFIG_SAMPLES_COUNT 5
 #endif
 
+/** Defines whether can completly skip the history reprojection when offscreen on fully rejected. */
 #define CONFIG_SELECTIVE_REPROJECTION 1
-#define CONFIG_TAA_RESPONSIVE_BEHAVIOR 1
-#define CONFIG_LOAD_TRANSLUCENCY_REJECTION 0
 
+/** Controls number of pixel to process per lane. */
 #if CONFIG_COMPILE_FP16
 	// Take advantage of RDNA's v_pk_*_{uif}16 instructions
 	#define CONFIG_ENABLE_DUAL_PIXEL_VECTORIZATION 1
@@ -84,7 +115,6 @@
 //------------------------------------------------------- PARAMETERS
 
 Texture2D<tsr_halfC> InputSceneColorTexture;
-Texture2D<tsr_half4> InputSceneTranslucencyTexture;
 
 Texture2D<tsr_half4> HistoryRejectionTexture;
 
@@ -101,14 +131,10 @@ Texture2D<tsr_ushort2> AntiAliasingTexture;
 	Texture2D<float2> UndistortingDisplacementTexture;
 #endif
 
-int2   TranslucencyPixelPosMin;
-int2   TranslucencyPixelPosMax;
-
 FScreenTransform HistoryPixelPosToViewportUV;
 FScreenTransform ViewportUVToInputPPCo;
 FScreenTransform HistoryPixelPosToScreenPos;
 FScreenTransform HistoryPixelPosToInputPPCo;
-FScreenTransform HistoryPixelPosToTranslucencyPPCo;
 float3 HistoryQuantizationError;
 float HistorySampleCount;
 float HistoryHisteresis;
@@ -127,7 +153,6 @@ float PrevFrameIndex;
 uint bGenerateOutputMip1;
 uint bGenerateOutputMip2;
 uint bGenerateOutputMip3;
-uint bHasSeparateTranslucency;
 uint HistoryArrayIndices_HighFrequency;
 
 Texture2DArray<tsr_halfC> PrevHistoryColorTexture;
@@ -235,14 +260,32 @@ tsr_half3x2 GetColorChannel(tsr_half4x2 Color)
 void ComputeInputKernelSamplePosition(
 	tsr_short2x2 PixelPos,
 	tsr_half2x2 dKO,
-	uint SampleId,
+	const uint SampleId,
 	int2 PixelPosMin,
 	int2 PixelPosMax,
 	out tsr_short2x2 SamplePixelPos,
-	out tsr_half2x2 PixelOffset)
+	out tsr_half2x2 PixelOffset,
+	out bool2 bIsDuplicatedSample,
+	out bool2 bIsDisabledSample)
 {
+	tsr_short2x2 iKOSign;
+	tsr_half2x2 KOSign;
+	#if CONFIG_COMPILE_FP16
+		iKOSign[0] = int16_t(1) - int16_t2((asuint16(dKO[0]) & uint16_t(0x8000)) >> uint16_t(14));
+		iKOSign[1] = int16_t(1) - int16_t2((asuint16(dKO[1]) & uint16_t(0x8000)) >> uint16_t(14));
+		KOSign[0] = asfloat16(asuint16(half(1.0)).xx | (asuint16(dKO[0]) & uint16_t(0x8000)));
+		KOSign[1] = asfloat16(asuint16(half(1.0)).xx | (asuint16(dKO[1]) & uint16_t(0x8000)));
+	#else
+		iKOSign = dpv_interleave_registers(
+			SignFastInt(dpv_lo(dKO)), SignFastInt(dpv_hi(dKO)));
+		KOSign[0] = asfloat(asuint(1.0).xx | (asuint(dKO[0]) & uint(0x80000000)));
+		KOSign[1] = asfloat(asuint(1.0).xx | (asuint(dKO[1]) & uint(0x80000000)));
+	#endif
 	
-	#if CONFIG_SAMPLES == 9
+	bIsDuplicatedSample = false;
+	bIsDisabledSample = false;
+
+	#if CONFIG_SAMPLES == CONFIG_SAMPLES_3X3
 	{
 		tsr_short2 iPixelOffset = tsr_short2(kOffsets3x3[kSquareIndexes3x3[SampleId]]);
 		PixelOffset = dpv_interleave_registers(tsr_half2(iPixelOffset), tsr_half2(iPixelOffset));
@@ -252,39 +295,99 @@ void ComputeInputKernelSamplePosition(
 			dpv_interleave_registers(iPixelOffset, iPixelOffset), iPixelOffset,
 			PixelPosMin, PixelPosMax);
 	}
-	#elif CONFIG_SAMPLES == 5 || CONFIG_SAMPLES == 6
+	#elif CONFIG_SAMPLES == CONFIG_SAMPLES_PLUS_MOVE_FAR
+	{
+		tsr_half2x2 SquareShape;
+		#if CONFIG_COMPILE_FP16
+			SquareShape[0] = asfloat16(asuint16(half(1.0)).xx | (asuint16(abs(dKO[0]) - tsr_half(0.25).xx) & uint16_t(0x8000)));
+			SquareShape[1] = asfloat16(asuint16(half(1.0)).xx | (asuint16(abs(dKO[1]) - tsr_half(0.25).xx) & uint16_t(0x8000)));
+		#else
+			SquareShape[0] = asfloat(asuint(1.0).xx | (asuint(abs(dKO[0]) - float(0.25).xx) & 0x8000u));
+			SquareShape[1] = asfloat(asuint(1.0).xx | (asuint(abs(dKO[1]) - float(0.25).xx) & 0x8000u));
+		#endif
+
+		PixelOffset = tsr_half(0);
+		SamplePixelPos = PixelPos;
+
+		if (SampleId == 0)
+		{
+			// NOP
+		}
+		else if (SampleId == 1)
+		{
+			PixelOffset[0] = KOSign[0];
+			
+			SamplePixelPos[0] += tsr_short2(PixelOffset[0]);
+			SamplePixelPos[0] = ClampPixelOffset(SamplePixelPos, PixelPosMin, PixelPosMax)[0];
+		}
+		else if (SampleId == 2)
+		{
+			PixelOffset[1] = KOSign[1];
+			
+			SamplePixelPos[1] += tsr_short2(PixelOffset[1]);
+			SamplePixelPos[1] = ClampPixelOffset(SamplePixelPos, PixelPosMin, PixelPosMax)[1];
+		}
+		else if (SampleId == 3)
+		{
+			PixelOffset[0] = SquareShape[0] * KOSign[0];
+			PixelOffset[1] = saturate(SquareShape[0]) * KOSign[1];
+			
+			SamplePixelPos[0] += tsr_short2(PixelOffset[0]);
+			SamplePixelPos[1] += tsr_short2(PixelOffset[1]);
+			SamplePixelPos = ClampPixelOffset(SamplePixelPos, PixelPosMin, PixelPosMax);
+		}
+		else // if (SampleId == 4)
+		{
+			PixelOffset[0] = saturate(SquareShape[1]) * KOSign[0];
+			PixelOffset[1] = SquareShape[1] * KOSign[1];
+
+			bIsDuplicatedSample = (SquareShape[0] + SquareShape[1]) == tsr_half(2.0);
+			
+			SamplePixelPos[0] += tsr_short2(PixelOffset[0]);
+			SamplePixelPos[1] += tsr_short2(PixelOffset[1]);
+			SamplePixelPos = ClampPixelOffset(SamplePixelPos, PixelPosMin, PixelPosMax);
+		}
+	}
+	#elif CONFIG_SAMPLES == CONFIG_SAMPLES_PLUS || CONFIG_SAMPLES == CONFIG_SAMPLES_PLUS_DISABLE_FAR || CONFIG_SAMPLES == CONFIG_SAMPLES_PLUS_AND_CORNER
 	{
-		#if CONFIG_SAMPLES == 6
+		#if CONFIG_SAMPLES == CONFIG_SAMPLES_PLUS_AND_CORNER
 		if (SampleId == 5)
 		{
-			tsr_short2x2 iPixelOffset;
-			#if CONFIG_COMPILE_FP16
-				iPixelOffset[0] = int16_t(1) - int16_t2((asuint16(dKO[0]) & uint16_t(0x8000)) >> uint16_t(14));
-				iPixelOffset[1] = int16_t(1) - int16_t2((asuint16(dKO[1]) & uint16_t(0x8000)) >> uint16_t(14));
-				PixelOffset[0] = asfloat16(asuint16(half(1.0)).xx | (asuint16(dKO[0]) & uint16_t(0x8000)));
-				PixelOffset[1] = asfloat16(asuint16(half(1.0)).xx | (asuint16(dKO[1]) & uint16_t(0x8000)));
-			#else
-				iPixelOffset = dpv_interleave_registers(
-					SignFastInt(dpv_lo(dKO)), SignFastInt(dpv_hi(dKO)));
-				PixelOffset[0] = asfloat(asuint(1.0).xx | (asuint(dKO[0]) & uint(0x80000000)));
-				PixelOffset[1] = asfloat(asuint(1.0).xx | (asuint(dKO[1]) & uint(0x80000000)));
-			#endif
-						
-			SamplePixelPos = ClampPixelOffset(PixelPos + iPixelOffset, PixelPosMin, PixelPosMax);
+			PixelOffset = KOSign;
+			SamplePixelPos = ClampPixelOffset(PixelPos + iKOSign, PixelPosMin, PixelPosMax);
 		}
 		else
 		#endif
 		{
-			tsr_short2 iPixelOffset = tsr_short2(kOffsets3x3[kPlusIndexes3x3[SampleId]]);
+			const tsr_short2 iPixelOffset = tsr_short2(kOffsets3x3[kPlusIndexes3x3[SampleId]]);
 			PixelOffset = dpv_interleave_registers(tsr_half2(iPixelOffset), tsr_half2(iPixelOffset));
 					
 			SamplePixelPos = PixelPos + ClampPixelOffset(
 				PixelPos,
 				dpv_interleave_registers(iPixelOffset, iPixelOffset), iPixelOffset,
 				PixelPosMin, PixelPosMax);
+				
+			#if CONFIG_SAMPLES == CONFIG_SAMPLES_PLUS_DISABLE_FAR
+			if (all(iPixelOffset == tsr_short2(+1, +0)))
+			{
+				bIsDisabledSample = dKO[0] < tsr_half(-0.25);
+			}
+			else if (all(iPixelOffset == tsr_short2(+0, +1)))
+			{
+				bIsDisabledSample = dKO[1] < tsr_half(-0.25);
+			}
+			else if (all(iPixelOffset == tsr_short2(-1, +0)))
+			{
+				bIsDisabledSample = dKO[0] > tsr_half(+0.25);
+			}
+			else if (all(iPixelOffset == tsr_short2(+0, -1)))
+			{
+				bIsDisabledSample = dKO[1] > tsr_half(+0.25);
+			}
+			#endif
 		}
 	}
-	#elif CONFIG_SAMPLES == 2
+	#elif CONFIG_SAMPLES == CONFIG_SAMPLES_1X1
 	{
 		SamplePixelPos = PixelPos;
 		PixelOffset = 0;
@@ -954,11 +1057,11 @@ void MainCS(
 	
 	PLATFORM_SPECIFIC_ISOLATE
 	{
-		tsr_halfC RawInputColorArray[CONFIG_SAMPLES][DPV_PIXEL_PER_LANE];
-		tsr_half4 RawTranslucencyColorArray[CONFIG_SAMPLES][DPV_PIXEL_PER_LANE];
+		tsr_halfC RawInputColorArray[CONFIG_SAMPLES_COUNT][DPV_PIXEL_PER_LANE];
 
-		tsr_half2x2 RawdPPArray[CONFIG_SAMPLES];
-		tsr_half2   RawSampleSpatialWeightArray[CONFIG_SAMPLES];
+		tsr_half2x2 RawdPPArray[CONFIG_SAMPLES_COUNT];
+		tsr_half2   RawSampleSpatialWeightArray[CONFIG_SAMPLES_COUNT];
+		bool2       bIsDisabledSampleArray[CONFIG_SAMPLES_COUNT];
 		
 		// Issues overlapped texture fetches
 		PLATFORM_SPECIFIC_ISOLATE
@@ -966,12 +1069,8 @@ void MainCS(
 			// Detect if HistoryBufferUV would be outside of the viewport.
 			tsr_half2 SpatialAntiAliasingLerp = select(or(bIsOffScreen, and(bIsDisoccluded, !bIsHistoryResurrection)), tsr_half(1.0).xx, saturate(tsr_half(1.0) - LowFrequencyRejection * tsr_half(4.0)));
 
-			// Pixel coordinate of the center of output pixel O in the input viewport.
-			float2x2 TranslucencyPPCo = ApplyScreenTransform(float2x2(HistoryPixelPos), HistoryPixelPosToTranslucencyPPCo);
-			
 			#if CONFIG_REJECTION_ANTI_ALIASING
 			{
-				TranslucencyPPCo += dpv_scale(SpatialAntiAliasingOffset, SpatialAntiAliasingLerp);
 				InputPPCo += dpv_scale(SpatialAntiAliasingOffset, SpatialAntiAliasingLerp);
 
 				InputPPCk = floor(InputPPCo) + 0.5;
@@ -981,16 +1080,6 @@ void MainCS(
 			}
 			#endif
 			
-			// Pixel coordinate of the center of the nearest input pixel K in the input viewport.
-			float2x2 TranslucencyPPCk = floor(TranslucencyPPCo) + 0.5;
-	
-			tsr_short2x2 TranslucencyPixelPos = ClampPixelOffset(
-				tsr_short2x2(TranslucencyPPCo),
-				TranslucencyPixelPosMin, TranslucencyPixelPosMax);
-				
-			// Vector in pixel between pixel K -> O.
-			tsr_half2x2 dTranslucencyKO = tsr_half2x2(TranslucencyPPCo - TranslucencyPPCk);
-			
 			// Vector in pixel between pixel K -> O.
 			tsr_half2x2 dInputKO = tsr_half2x2(InputPPCo - InputPPCk);
 			
@@ -1056,82 +1145,70 @@ void MainCS(
 				PrevWeight = min(PrevWeight, MaxWeight - CurrentWeight);	
 			}
 
-			UNROLL_N(CONFIG_SAMPLES)
-			for (uint SampleId = 0; SampleId < CONFIG_SAMPLES; SampleId++)
+			UNROLL_N(CONFIG_SAMPLES_COUNT)
+			for (uint SampleId = 0; SampleId < CONFIG_SAMPLES_COUNT; SampleId++)
 			{
 				tsr_short2x2 InputSamplePixelPos;
 				tsr_half2x2 InputPixelOffset;
+				bool2 bIsDuplicatedSample;
+				bool2 bIsDisabledSample;
 				ComputeInputKernelSamplePosition(
 					InputPixelPos, dInputKO, SampleId,
 					InputPixelPosMin, InputPixelPosMax,
 					/* out */ InputSamplePixelPos,
-					/* out */ InputPixelOffset);
-				
-				tsr_short2x2 TranslucencySamplePixelPos;
-				tsr_half2x2 TranslucencyPixelOffset;
-				ComputeInputKernelSamplePosition(
-					TranslucencyPixelPos, dTranslucencyKO, SampleId,
-					TranslucencyPixelPosMin, TranslucencyPixelPosMax,
-					/* out */ TranslucencySamplePixelPos,
-					/* out */ TranslucencyPixelOffset);
+					/* out */ InputPixelOffset,
+					/* out */ bIsDuplicatedSample,
+					/* out */ bIsDisabledSample);
 				
 				tsr_half2x2 dPP = InputPixelOffset - dInputKO;
 				tsr_half2 SampleSpatialWeight = ComputeSampleWeigth(KernelInputToHistoryFactor, dPP, /* MinimalContribution = */ float(0.005));
-				
-				tsr_half2x2 dTranslucencyPP = TranslucencyPixelOffset - dTranslucencyKO;
-				tsr_half2 TranslucencySampleSpatialWeight = ComputeSampleWeigth(KernelInputToHistoryFactor, dTranslucencyPP, /* MinimalContribution = */ float(0.005));
-				
-				RawdPPArray[SampleId] = dPP; 
+				SampleSpatialWeight = select(bIsDuplicatedSample, tsr_half(0.0), SampleSpatialWeight);
+
+				bIsDisabledSampleArray[SampleId]      = bIsDisabledSample;
+				RawdPPArray[SampleId]                 = dPP; 
 				RawSampleSpatialWeightArray[SampleId] = SampleSpatialWeight;
 
 				UNROLL_N(DPV_PIXEL_PER_LANE)
 				for (uint PixelId = 0; PixelId < DPV_PIXEL_PER_LANE; PixelId++)
 				{
 					RawInputColorArray[SampleId][PixelId] = InputSceneColorTexture[dpv_access_pixel(InputSamplePixelPos, PixelId)];
-					#if CONFIG_COMPOSE_TRANSLUCENCY
-						RawTranslucencyColorArray[SampleId][PixelId] = InputSceneTranslucencyTexture[dpv_access_pixel(TranslucencySamplePixelPos, PixelId)];
-					#endif
 				}
-			} // for (uint SampleId = 0; SampleId < CONFIG_SAMPLES; SampleId++)
+			} // for (uint SampleId = 0; SampleId < CONFIG_SAMPLES_COUNT; SampleId++)
 		}
 			
+		tsr_halfCx2 InputColorCenter = dpv_force_interleave_registers_array(RawInputColorArray[0]);
+
 		tsr_half2 FilteredInputColorWeight = tsr_half(0.0);
-		
 		FilteredInputColor = tsr_half(0.0);
-		#if CONFIG_SCENE_COLOR_ALPHA
-			InputMinColor = dpv_interleave_mono_registers(+LargestSceneColorYCoCgA);
-			InputMaxColor = dpv_interleave_mono_registers(-LargestSceneColorYCoCgA);
-		#else
-			InputMinColor = dpv_interleave_mono_registers(+LargestSceneColorYCoCg);
-			InputMaxColor = dpv_interleave_mono_registers(-LargestSceneColorYCoCg);
-		#endif
+		InputMinColor = InputColorCenter;
+		InputMaxColor = InputColorCenter;
 
-		UNROLL_N(CONFIG_SAMPLES)
-		for (uint SampleId = 0; SampleId < CONFIG_SAMPLES; SampleId++)
+		UNROLL_N(CONFIG_SAMPLES_COUNT)
+		for (uint SampleId = 0; SampleId < CONFIG_SAMPLES_COUNT; SampleId++)
 		{
+			bool2 bIsDisabledSample = bIsDisabledSampleArray[SampleId];
 			tsr_half2 SampleSpatialWeight = RawSampleSpatialWeightArray[SampleId];
 			tsr_halfCx2 InputColor = dpv_force_interleave_registers_array(RawInputColorArray[SampleId]);
-			tsr_half4x2 TranslucencyColor = dpv_force_interleave_registers_array(RawTranslucencyColorArray[SampleId]);
-
-			#if CONFIG_COMPOSE_TRANSLUCENCY
-			{
-				InputColor[0] = InputColor[0] * TranslucencyColor[3] + TranslucencyColor[0];
-				InputColor[1] = InputColor[1] * TranslucencyColor[3] + TranslucencyColor[1];
-				InputColor[2] = InputColor[2] * TranslucencyColor[3] + TranslucencyColor[2];
-				#if CONFIG_SCENE_COLOR_ALPHA
-					InputColor[3] = InputColor[3] * TranslucencyColor[3];
-				#endif
-			}
-			#endif
 
 			tsr_half2 ToneWeight = HdrWeight4(InputColor);
 
 			FilteredInputColor       += dpv_scale((SampleSpatialWeight * ToneWeight), InputColor);
 			FilteredInputColorWeight += (SampleSpatialWeight * ToneWeight);
 
-			InputMinColor = min(InputMinColor, InputColor);
-			InputMaxColor = max(InputMaxColor, InputColor);
-		}
+			if (SampleId != 0)
+			{
+				tsr_halfCx2 ClampBoxSample;
+				ClampBoxSample[0] = select(bIsDisabledSample, InputColorCenter[0], InputColor[0]);
+				ClampBoxSample[1] = select(bIsDisabledSample, InputColorCenter[1], InputColor[1]);
+				ClampBoxSample[2] = select(bIsDisabledSample, InputColorCenter[2], InputColor[2]);
+				#if CONFIG_SCENE_COLOR_ALPHA
+					ClampBoxSample[3] = select(bIsDisabledSample, InputColorCenter[3], InputColor[3]);
+				#endif
+			
+				InputMinColor = min(InputMinColor, ClampBoxSample);
+				InputMaxColor = max(InputMaxColor, ClampBoxSample);
+			}
+		} // // for (uint SampleId = 0; SampleId < CONFIG_SAMPLES_COUNT; SampleId++)
 		
 		FilteredInputColor = dpv_scale(FilteredInputColor, rcp(FilteredInputColorWeight));
 
diff --git a/Engine/Source/Runtime/Renderer/Private/PostProcess/TemporalSuperResolution.cpp b/Engine/Source/Runtime/Renderer/Private/PostProcess/TemporalSuperResolution.cpp
--- a/Engine/Source/Runtime/Renderer/Private/PostProcess/TemporalSuperResolution.cpp
+++ b/Engine/Source/Runtime/Renderer/Private/PostProcess/TemporalSuperResolution.cpp
@@ -958,7 +958,6 @@ class FTSRUpdateHistoryCS : public FTSRShader
 	BEGIN_SHADER_PARAMETER_STRUCT(FParameters, )
 		SHADER_PARAMETER_STRUCT_INCLUDE(FTSRCommonParameters, CommonParameters)
 		SHADER_PARAMETER_RDG_TEXTURE(Texture2D, InputSceneColorTexture)
-		SHADER_PARAMETER_RDG_TEXTURE(Texture2D, InputSceneTranslucencyTexture)
 
 		SHADER_PARAMETER_RDG_TEXTURE(Texture2D, HistoryRejectionTexture)
 		SHADER_PARAMETER_RDG_TEXTURE_SRV(Texture2D, ReprojectionBoundaryTexture)
@@ -966,15 +965,10 @@ class FTSRUpdateHistoryCS : public FTSRShader
 		SHADER_PARAMETER_RDG_TEXTURE_SRV(Texture2D, ReprojectionVectorTexture)
 		SHADER_PARAMETER_RDG_TEXTURE(Texture2D, AntiAliasingTexture)
 
-		SHADER_PARAMETER_STRUCT(FScreenPassTextureViewportParameters, TranslucencyInfo)
-		SHADER_PARAMETER(FIntPoint, TranslucencyPixelPosMin)
-		SHADER_PARAMETER(FIntPoint, TranslucencyPixelPosMax)
-
 		SHADER_PARAMETER(FScreenTransform, HistoryPixelPosToViewportUV)
 		SHADER_PARAMETER(FScreenTransform, ViewportUVToInputPPCo)
 		SHADER_PARAMETER(FScreenTransform, HistoryPixelPosToScreenPos)
 		SHADER_PARAMETER(FScreenTransform, HistoryPixelPosToInputPPCo)
-		SHADER_PARAMETER(FScreenTransform, HistoryPixelPosToTranslucencyPPCo)
 
 		SHADER_PARAMETER(FVector3f, HistoryQuantizationError)
 		SHADER_PARAMETER(float, HistorySampleCount)
@@ -991,7 +985,6 @@ class FTSRUpdateHistoryCS : public FTSRShader
 		SHADER_PARAMETER(int32, bGenerateOutputMip1)
 		SHADER_PARAMETER(int32, bGenerateOutputMip2)
 		SHADER_PARAMETER(int32, bGenerateOutputMip3)
-		SHADER_PARAMETER(int32, bHasSeparateTranslucency)
 
 		SHADER_PARAMETER_STRUCT(FTSRHistoryArrayIndices, HistoryArrayIndices)
 		SHADER_PARAMETER_STRUCT_INCLUDE(FTSRPrevHistoryParameters, PrevHistoryParameters)
@@ -2437,7 +2430,6 @@ FDefaultTemporalUpscaler::FOutputs AddTemporalSuperResolutionPasses(
 		FTSRUpdateHistoryCS::FParameters* PassParameters = GraphBuilder.AllocParameters<FTSRUpdateHistoryCS::FParameters>();
 		PassParameters->CommonParameters = CommonParameters;
 		PassParameters->InputSceneColorTexture = InputSceneColorTexture;
-		PassParameters->InputSceneTranslucencyTexture = SeparateTranslucencyTexture;
 		PassParameters->HistoryRejectionTexture = HistoryRejectionTexture;
 
 		PassParameters->ReprojectionBoundaryTexture = ReprojectionBoundaryTexture ? ReprojectionBoundaryTexture : GraphBuilder.CreateSRV(FRDGTextureSRVDesc(BlackUintDummy));
@@ -2445,15 +2437,11 @@ FDefaultTemporalUpscaler::FOutputs AddTemporalSuperResolutionPasses(
 		PassParameters->ReprojectionVectorTexture = ReprojectionVectorTexture;
 		PassParameters->AntiAliasingTexture = AntiAliasingTexture;
 
-		PassParameters->TranslucencyPixelPosMin = SeparateTranslucencyRect.Min;
-		PassParameters->TranslucencyPixelPosMax = SeparateTranslucencyRect.Max - 1;
-
 		FScreenTransform HistoryPixelPosToViewportUV = (FScreenTransform::Identity + 0.5f) * CommonParameters.HistoryInfo.ViewportSizeInverse;
 		PassParameters->HistoryPixelPosToViewportUV = HistoryPixelPosToViewportUV;
 		PassParameters->ViewportUVToInputPPCo = FScreenTransform::Identity * CommonParameters.InputInfo.ViewportSize + CommonParameters.InputJitter + CommonParameters.InputPixelPosMin;
 		PassParameters->HistoryPixelPosToScreenPos = HistoryPixelPosToViewportUV * FScreenTransform::ViewportUVToScreenPos;
 		PassParameters->HistoryPixelPosToInputPPCo = HistoryPixelPosToViewportUV * PassParameters->ViewportUVToInputPPCo;
-		PassParameters->HistoryPixelPosToTranslucencyPPCo = HistoryPixelPosToViewportUV * SeparateTranslucencyRect.Size() + CommonParameters.InputJitter * SeparateTranslucencyRect.Size() / CommonParameters.InputInfo.ViewportSize + SeparateTranslucencyRect.Min;
 		PassParameters->HistoryQuantizationError = ComputePixelFormatQuantizationError(HistoryColorFormat);
 
 		// All parameters to control the sample count in history.
@@ -2470,7 +2458,6 @@ FDefaultTemporalUpscaler::FOutputs AddTemporalSuperResolutionPasses(
 		PassParameters->bGenerateOutputMip1 = false;
 		PassParameters->bGenerateOutputMip2 = false;
 		PassParameters->bGenerateOutputMip3 = false;
-		PassParameters->bHasSeparateTranslucency = bHasSeparateTranslucency;
 
 		PassParameters->HistoryArrayIndices = HistoryArrayIndices;
 		PassParameters->PrevHistoryParameters = PrevHistoryParameters;
