Fixes TSR from computing the resurrection when it's offscreen

diff --git a/Engine/Shaders/Private/TemporalSuperResolution/TSRDecimateHistory.usf b/Engine/Shaders/Private/TemporalSuperResolution/TSRDecimateHistory.usf
--- a/Engine/Shaders/Private/TemporalSuperResolution/TSRDecimateHistory.usf
+++ b/Engine/Shaders/Private/TemporalSuperResolution/TSRDecimateHistory.usf
@@ -500,6 +500,7 @@ void ReprojectAllPrevTextures(
 	#if CONFIG_RESURRECTION_REPROJECTION
 		out tsr_halfC ResurrectedHistoryGuideColor,
 		out tsr_half  ResurrectedHistoryGuideUncertainty,
+		out bool      bIsResurrectionOffScreen,
 	#endif
 	out bool bIsOffScreen)
 {
@@ -571,6 +572,12 @@ void ReprojectAllPrevTextures(
 		float4 ResurrectionClip = mul(ThisClip, ClipToResurrectionClip); 
 		float2 ResurrectionScreen = ResurrectionClip.xy / ResurrectionClip.w;
 
+		bIsResurrectionOffScreen = (
+			(ResurrectionScreen.x < -1.0) ||
+			(ResurrectionScreen.y < -1.0) ||
+			(ResurrectionScreen.x > 1.0) ||
+			(ResurrectionScreen.y > 1.0));
+
 		float2 KernelHistoryBufferUV = ApplyScreenTransform(ResurrectionScreen, ScreenPosToResurrectionGuideBufferUV);
 		KernelHistoryBufferUV = fastClamp(KernelHistoryBufferUV, ResurrectionGuideUVViewportBilinearMin, ResurrectionGuideUVViewportBilinearMax);
 
@@ -669,6 +676,9 @@ void MainCS(
 	#if CONFIG_RESURRECTION_REPROJECTION
 		tsr_halfC ResurrectedHistoryGuideColor;
 		tsr_half  ResurrectedHistoryGuideUncertainty;
+		bool bIsResurrectionOffScreen;
+	#else
+		const bool bIsResurrectionOffScreen = false;
 	#endif
 	bool bIsOffScreen;
 
@@ -689,6 +699,7 @@ void MainCS(
 			#if CONFIG_RESURRECTION_REPROJECTION
 				/* out */ ResurrectedHistoryGuideColor,
 				/* out */ ResurrectedHistoryGuideUncertainty,
+				/* out */ bIsResurrectionOffScreen,
 			#endif
 			/* out */ bIsOffScreen);
 	}
@@ -777,10 +788,11 @@ void MainCS(
 	
 	// Builds the decimate bitmask
 	tsr_ushort DecimateBitMask = 
-		select(bIsOffScreen,            tsr_ushort(0x1), tsr_ushort(0x0)) |
-		select(bIsParallaxDisocclusion, tsr_ushort(0x2), tsr_ushort(0x0)) |
-		select(bHasPixelAnimation,      tsr_ushort(0x4), tsr_ushort(0x0)) |
-		select(bReprojectionHollFill,   tsr_ushort(0x8), tsr_ushort(0x0));
+		select(bIsOffScreen,             tsr_ushort(0x01), tsr_ushort(0x0)) |
+		select(bIsParallaxDisocclusion,  tsr_ushort(0x02), tsr_ushort(0x0)) |
+		select(bHasPixelAnimation,       tsr_ushort(0x04), tsr_ushort(0x0)) |
+		select(bReprojectionHollFill,    tsr_ushort(0x08), tsr_ushort(0x0)) |
+		select(bIsResurrectionOffScreen, tsr_ushort(0x10), tsr_ushort(0x0));
 
 	// Output full res
 	ISOLATE
diff --git a/Engine/Shaders/Private/TemporalSuperResolution/TSRRejectShading.usf b/Engine/Shaders/Private/TemporalSuperResolution/TSRRejectShading.usf
--- a/Engine/Shaders/Private/TemporalSuperResolution/TSRRejectShading.usf
+++ b/Engine/Shaders/Private/TemporalSuperResolution/TSRRejectShading.usf
@@ -121,7 +121,8 @@ void FetchDecimateMask(
 	tsr_tensor_short2 PixelPos,
 	out tsr_tensor_half VelocityEdge,
 	out tsr_tensor_bool bIsDisoccluded,
-	out tsr_tensor_bool bHasPixelAnimation)
+	out tsr_tensor_bool bHasPixelAnimation,
+	out tsr_tensor_bool bIsResurrectionOffScreen)
 {
 	UNROLL_N(SIMD_SIZE)
 	for (uint ElementIndex = 0; ElementIndex < SIMD_SIZE; ElementIndex++)
@@ -136,6 +137,7 @@ void FetchDecimateMask(
 		VelocityEdge.SetElement(ElementIndex, saturate(tsr_half(tsr_ushort(DilateMask)) * tsr_half(rcp(255.0))));
 		bIsDisoccluded.SetElement(ElementIndex,     (BitMask & tsr_ushort(0x3)) != tsr_ushort(0));
 		bHasPixelAnimation.SetElement(ElementIndex, (BitMask & tsr_ushort(0x4)) != tsr_ushort(0));
+		bIsResurrectionOffScreen.SetElement(ElementIndex, (BitMask & tsr_ushort(0x10)) != tsr_ushort(0));
 	}
 } // FetchDecimateMask()
 
@@ -237,7 +239,7 @@ void FetchResurrectedHistoryGuide(
 
 // Returns whether resurrected frame is closer than previous frame on per pixel basis.
 CALL_SITE_DEBUGLOC
-tsr_tensor_bool IsResurrectedFrameCloserThanPreviousFrame(tsr_tensor_halfC ExposedInput, tsr_tensor_halfC ExposedHistory, tsr_tensor_halfC ExposedResurrectedHistory)
+tsr_tensor_bool IsResurrectedFrameCloserThanPreviousFrame(tsr_tensor_bool bIsResurrectionOffScreen, tsr_tensor_halfC ExposedInput, tsr_tensor_halfC ExposedHistory, tsr_tensor_halfC ExposedResurrectedHistory)
 {
 	tsr_tensor_halfC PrevMatch         = abs(ExposedInput - ExposedHistory);
 	tsr_tensor_halfC ResurrectionMatch = abs(ExposedInput - ExposedResurrectedHistory);
@@ -251,7 +253,7 @@ tsr_tensor_bool IsResurrectedFrameCloserThanPreviousFrame(tsr_tensor_halfC Expos
 
 	tsr_tensor_half bResurrectedPixelIsCloser = saturate((TotalMatchDiff - tsr_tensor_half::Const(tsr_half(0.05 * 3.0 * 9.0))) * tsr_tensor_half::Const(tsr_half(POSITIVE_INFINITY)));
 
-	tsr_tensor_bool bResurrectionIsCloser = Sum3x3(bResurrectedPixelIsCloser) > tsr_tensor_half::Const(4.0);
+	tsr_tensor_bool bResurrectionIsCloser = and(Sum3x3(bResurrectedPixelIsCloser) > tsr_tensor_half::Const(4.0), not(bIsResurrectionOffScreen));
 
 	//tsr_tensor_bool bResurrectionIsCloser = MedianBool3x3((TotalMatchDiff - tsr_tensor_half::Const(tsr_half(0.05 * 3.0 * 9.0))));
 
@@ -399,8 +401,9 @@ void RefetchInputAndHistoryFromLDSOrCache(
 	tsr_tensor_half  VelocityEdge;
 	tsr_tensor_bool  bIsDisoccluded;
 	tsr_tensor_bool  bHasPixelAnimation;
+	tsr_tensor_bool  bIsResurrectionOffScreen;
 
-	FetchDecimateMask(PixelPos, /* out */ VelocityEdge, /* out */ bIsDisoccluded, /* out */ bHasPixelAnimation);
+	FetchDecimateMask(PixelPos, /* out */ VelocityEdge, /* out */ bIsDisoccluded, /* out */ bHasPixelAnimation, /* out */ bIsResurrectionOffScreen);
 	FetchSceneColorAndTranslucency(PixelPos, /* out */ OriginalOpaqueInput, /* out */ OriginalTranslucencyInput);
 	FetchReprojectedHistoryGuide(
 		PixelPos,
@@ -449,6 +452,7 @@ void MainCS(
 		tsr_tensor_half  VelocityEdge;
 		tsr_tensor_bool  bIsDisoccluded;
 		tsr_tensor_bool  bHasPixelAnimation;
+		tsr_tensor_bool  bIsResurrectionOffScreen;
 		tsr_tensor_halfC OriginalOpaqueInput;
 		tsr_tensor_half4 OriginalTranslucencyInput;
 		tsr_tensor_half  MoireInput;
@@ -458,7 +462,7 @@ void MainCS(
 		{
 			tsr_tensor_short2 PixelPos = ComputePixelPos();
 		
-			FetchDecimateMask(PixelPos, /* out */ VelocityEdge, /* out */ bIsDisoccluded, /* out */ bHasPixelAnimation);
+			FetchDecimateMask(PixelPos, /* out */ VelocityEdge, /* out */ bIsDisoccluded, /* out */ bHasPixelAnimation, /* out */ bIsResurrectionOffScreen);
 			FetchSceneColorAndTranslucency(PixelPos, /* out */ OriginalOpaqueInput, /* out */ OriginalTranslucencyInput);
 			FetchMoireInput(
 				PixelPos,
@@ -498,6 +502,7 @@ void MainCS(
 	tsr_tensor_half  VelocityEdge;
 	tsr_tensor_bool  bIsDisoccluded;
 	tsr_tensor_bool  bHasPixelAnimation;
+	tsr_tensor_bool  bIsResurrectionOffScreen;
 	tsr_tensor_halfC OriginalOpaqueInput;
 	tsr_tensor_half4 OriginalTranslucencyInput;
 	tsr_tensor_halfC OriginalHistory;
@@ -514,7 +519,7 @@ void MainCS(
 		{
 			tsr_tensor_short2 PixelPos = ComputePixelPos();
 
-			FetchDecimateMask(PixelPos, /* out */ VelocityEdge, /* out */ bIsDisoccluded, /* out */ bHasPixelAnimation);
+			FetchDecimateMask(PixelPos, /* out */ VelocityEdge, /* out */ bIsDisoccluded, /* out */ bHasPixelAnimation, /* out */ bIsResurrectionOffScreen);
 			FetchSceneColorAndTranslucency(PixelPos, /* out */ OriginalOpaqueInput, /* out */ OriginalTranslucencyInput);
 			FetchReprojectedHistoryGuide(
 				PixelPos,
@@ -531,7 +536,7 @@ void MainCS(
 		{
 			tsr_tensor_short2 PixelPos = ComputePixelPos();
 
-			FetchDecimateMask(PixelPos, /* out */ VelocityEdge, /* out */ bIsDisoccluded, /* out */ bHasPixelAnimation);
+			FetchDecimateMask(PixelPos, /* out */ VelocityEdge, /* out */ bIsDisoccluded, /* out */ bHasPixelAnimation, /* out */ bIsResurrectionOffScreen);
 			FetchSceneColorAndTranslucency(PixelPos, /* out */ OriginalOpaqueInput, /* out */ OriginalTranslucencyInput);
 			FetchReprojectedHistoryGuide(
 				PixelPos,
@@ -605,7 +610,7 @@ void MainCS(
 				/* SourcePerceptionAdd = */ kHistoryGuidePreceptionAdd,
 				/* DestPerceptionAdd   = */ PerceptionAdd);
 				
-			bResurrectionIsCloser = IsResurrectedFrameCloserThanPreviousFrame(ExposedInput, ExposedHistory, ExposedResurrectedHistory);
+			bResurrectionIsCloser = IsResurrectedFrameCloserThanPreviousFrame(bIsResurrectionOffScreen, ExposedInput, ExposedHistory, ExposedResurrectedHistory);
 		}
 		else
 		{
