Change opaque shadow mask to run into a single pass when using deep shadow map (instead of one pass per macro group.

diff --git a/Engine/Shaders/Private/HairStrands/HairStrandsDeepShadowMask.usf b/Engine/Shaders/Private/HairStrands/HairStrandsDeepShadowMask.usf
--- a/Engine/Shaders/Private/HairStrands/HairStrandsDeepShadowMask.usf
+++ b/Engine/Shaders/Private/HairStrands/HairStrandsDeepShadowMask.usf
@@ -42,16 +42,17 @@ float3 GetTranslatedWorldPosition(float2 UV, float SceneDepth)
 #if SHADER_SHADOWMASK_DEEPSHADOW
 
 StructuredBuffer<FDeepShadowViewInfo> DeepShadow_ViewInfoBuffer;
+StructuredBuffer<uint> DeepShadow_AtlasSlotIndexBuffer;
 
 int2				DeepShadow_SlotResolution;
 int2				DeepShadow_SlotOffset;
-uint				DeepShadow_SlotIndex;
 uint				DeepShadow_KernelType;
 float				DeepShadow_DepthBiasScale;
 float				DeepShadow_DensityScale;
 float4				DeepShadow_LayerDepths;
 float 				FadeAlpha;
 uint				EncodingType;
+uint				EffectiveAtlasSlotCount;
 
 Texture2D<float>	DeepShadow_FrontDepthTexture;
 Texture2D<float4>	DeepShadow_DomTexture;
@@ -63,9 +64,9 @@ SamplerState 		ShadowSampler;
 // XY: UV of the shadow space
 //  Z: Normalized depth value in light clip space
 //  W: Positive if the position is valid, negative otherwise
-float4 ToLightPosition(float3 InTranslatedWorldPosition, out float4 AtlasScaleBias)
+float4 ToLightPosition(float3 InTranslatedWorldPosition, uint InAtlasSlotIndex, out float4 AtlasScaleBias)
 {
-	FDeepShadowViewInfo ShadowViewInfo = DeepShadow_ViewInfoBuffer[DeepShadow_SlotIndex];
+	FDeepShadowViewInfo ShadowViewInfo = DeepShadow_ViewInfoBuffer[InAtlasSlotIndex];
 	float4x4 TranslatedWorldToLightTransform = ShadowViewInfo.TranslatedWorldToClip;
 	AtlasScaleBias = ShadowViewInfo.AtlasScaleBias;
 
@@ -338,51 +339,62 @@ void MainPS(
 	const bool bIsValidHairPixel = HairStrands.HairOnlyDepthTexture.Load(uint3(PixelCoord, 0)).x > 0;
 	const bool bIsHairPixel = bIsValidHairPixel && HairStrands.HairCoverageTexture.Load(uint3(PixelCoord, 0)) >= 1;
 
-	const float3 TranslatedWorldPosition = GetTranslatedWorldPosition(UV, SceneDepth);
-	float4 AtlasScaleBias;
-	float4 LightSpacePosition		 = ToLightPosition(TranslatedWorldPosition, AtlasScaleBias);
-
-	// PCF or simple point sampler
-	float Visibility = 1;
-	const bool bIsValid =
-		LightSpacePosition.w > 0 &&
-		(LightSpacePosition.x >= 0 && LightSpacePosition.x <= 1) &&
-		(LightSpacePosition.y >= 0 && LightSpacePosition.y <= 1);
-	if (!bIsHairPixel && bIsValid)
-	{
-		LightSpacePosition.xy = LightSpacePosition.xy * AtlasScaleBias.xy + AtlasScaleBias.zw;
-
-		uint2 ShadowAtlasResolution;
-		DeepShadow_FrontDepthTexture.GetDimensions(ShadowAtlasResolution.x, ShadowAtlasResolution.y);
-
-	#if OPAQUE_PCF == 1
-		FHairSamplingSettings ShadowSettings;
-		ShadowSettings.ShadowDepthTexture = DeepShadow_FrontDepthTexture;
-		ShadowSettings.ShadowDomTexture = DeepShadow_DomTexture;
-		ShadowSettings.ShadowDepthTextureSampler = ShadowSampler;
-		ShadowSettings.ShadowAtlasInvResolution = 1 / float2(ShadowAtlasResolution);
-		ShadowSettings.ShadowSlotResolution = DeepShadow_SlotResolution;
-		ShadowSettings.ShadowSlotOffset = DeepShadow_SlotOffset;
-		ShadowSettings.SceneDepth = saturate(LightSpacePosition.z);
-		#if PERMUTATION_KERNEL_TYPE == KERNEL_TYPE_4x4
-			Visibility = 1 - InternalManual4x4PCF(LightSpacePosition.xy, ShadowSettings);
-		#elif PERMUTATION_KERNEL_TYPE == KERNEL_TYPE_Gaussian8		
-			Visibility = 1 - InternalGaussianFilter(LightSpacePosition.xy, ShadowSettings, 8);
-		#elif PERMUTATION_KERNEL_TYPE == KERNEL_TYPE_Gaussian16
-			Visibility = 1 - InternalGaussianFilter(LightSpacePosition.xy, ShadowSettings, 16);
-		#elif PERMUTATION_KERNEL_TYPE == KERNEL_TYPE_GaussianTransmission8
-			Visibility = 1 - InternalGaussianTransmissionFilter(LightSpacePosition.xyz, ShadowSettings, 8);
-		#else//PERMUTATION_KERNEL_TYPE == KERNEL_TYPE_2x2
-			Visibility = 1 - InternalManual2x2PCF(LightSpacePosition.xy, ShadowSettings);
-		#endif	
+	float MinVisibility = 1.f;
+	#if 1
+	for (uint AtlasSlotIt=0;AtlasSlotIt<EffectiveAtlasSlotCount;++AtlasSlotIt)
 	#else
-		const float FrontDepth = DeepShadow_FrontDepthTexture.SampleLevel(LinearSampler, LightSpacePosition.xy, 0);
-		const float ShadowBias = 0;
-		Visibility = IsCloser(FrontDepth + ShadowBias, LightSpacePosition.z) ? 0 : 1;
+	const uint AtlasSlotIt = 0;
 	#endif
+	{
+		const uint AtlasSlotIndex = DeepShadow_AtlasSlotIndexBuffer[AtlasSlotIt];
+		const float3 TranslatedWorldPosition = GetTranslatedWorldPosition(UV, SceneDepth);
+		float4 AtlasScaleBias;
+		float4 LightSpacePosition = ToLightPosition(TranslatedWorldPosition, AtlasSlotIndex, AtlasScaleBias);
+
+		// PCF or simple point sampler
+		float Visibility = 1;
+		const bool bIsValid =
+			LightSpacePosition.w > 0 &&
+			(LightSpacePosition.x >= 0 && LightSpacePosition.x <= 1) &&
+			(LightSpacePosition.y >= 0 && LightSpacePosition.y <= 1);
+		if (!bIsHairPixel && bIsValid)
+		{
+			LightSpacePosition.xy = LightSpacePosition.xy * AtlasScaleBias.xy + AtlasScaleBias.zw;
+
+			uint2 ShadowAtlasResolution;
+			DeepShadow_FrontDepthTexture.GetDimensions(ShadowAtlasResolution.x, ShadowAtlasResolution.y);
+
+		#if OPAQUE_PCF == 1
+			FHairSamplingSettings ShadowSettings;
+			ShadowSettings.ShadowDepthTexture = DeepShadow_FrontDepthTexture;
+			ShadowSettings.ShadowDomTexture = DeepShadow_DomTexture;
+			ShadowSettings.ShadowDepthTextureSampler = ShadowSampler;
+			ShadowSettings.ShadowAtlasInvResolution = 1 / float2(ShadowAtlasResolution);
+			ShadowSettings.ShadowSlotResolution = DeepShadow_SlotResolution;
+			ShadowSettings.ShadowSlotOffset = DeepShadow_SlotOffset;
+			ShadowSettings.SceneDepth = saturate(LightSpacePosition.z);
+			#if PERMUTATION_KERNEL_TYPE == KERNEL_TYPE_4x4
+				Visibility = 1 - InternalManual4x4PCF(LightSpacePosition.xy, ShadowSettings);
+			#elif PERMUTATION_KERNEL_TYPE == KERNEL_TYPE_Gaussian8		
+				Visibility = 1 - InternalGaussianFilter(LightSpacePosition.xy, ShadowSettings, 8);
+			#elif PERMUTATION_KERNEL_TYPE == KERNEL_TYPE_Gaussian16
+				Visibility = 1 - InternalGaussianFilter(LightSpacePosition.xy, ShadowSettings, 16);
+			#elif PERMUTATION_KERNEL_TYPE == KERNEL_TYPE_GaussianTransmission8
+				Visibility = 1 - InternalGaussianTransmissionFilter(LightSpacePosition.xyz, ShadowSettings, 8);
+			#else//PERMUTATION_KERNEL_TYPE == KERNEL_TYPE_2x2
+				Visibility = 1 - InternalManual2x2PCF(LightSpacePosition.xy, ShadowSettings);
+			#endif	
+		#else
+			const float FrontDepth = DeepShadow_FrontDepthTexture.SampleLevel(LinearSampler, LightSpacePosition.xy, 0);
+			const float ShadowBias = 0;
+			Visibility = IsCloser(FrontDepth + ShadowBias, LightSpacePosition.z) ? 0 : 1;
+		#endif
+
+			MinVisibility = min(MinVisibility, Visibility);
+		}
 	}
 
-	const float FadedShadow = lerp(1.0f, Visibility, FadeAlpha);
+	const float FadedShadow = lerp(1.0f, MinVisibility, FadeAlpha);
 	OutColor = WriteShadowMaskOutput(FadedShadow, EncodingType);
 
 }
@@ -485,7 +497,6 @@ void MainPS(
 			NormalizedDepthBias = (VirtualVoxel.DepthBiasScale_Shadow * LightDirection + PositionBiasScale * (Random * 2 - 1));
 		}
 
-
 		FVirtualVoxelCommonDesc CommonDesc;
 		CommonDesc.PageCountResolution	= VirtualVoxel.PageCountResolution;
 		CommonDesc.PageTextureResolution= VirtualVoxel.PageTextureResolution;
diff --git a/Engine/Source/Runtime/Renderer/Private/HairStrands/HairStrandsTransmittance.cpp b/Engine/Source/Runtime/Renderer/Private/HairStrands/HairStrandsTransmittance.cpp
--- a/Engine/Source/Runtime/Renderer/Private/HairStrands/HairStrandsTransmittance.cpp
+++ b/Engine/Source/Runtime/Renderer/Private/HairStrands/HairStrandsTransmittance.cpp
@@ -601,17 +601,18 @@ class FHairStrandsDeepShadowMaskPS : public FGlobalShader
 		SHADER_PARAMETER_STRUCT_INCLUDE(ShaderPrint::FShaderParameters, ShaderPrintParameters)
 		
 		SHADER_PARAMETER(FIntPoint, DeepShadow_SlotOffset)
-		SHADER_PARAMETER(uint32, DeepShadow_SlotIndex)
 		SHADER_PARAMETER(FIntPoint, DeepShadow_SlotResolution)
 		SHADER_PARAMETER(float, DeepShadow_DepthBiasScale)
 		SHADER_PARAMETER(float, DeepShadow_DensityScale)
 		SHADER_PARAMETER(FVector4f, DeepShadow_LayerDepths)
 		SHADER_PARAMETER(float, FadeAlpha)
 		SHADER_PARAMETER(uint32, EncodingType)
+		SHADER_PARAMETER(uint32, EffectiveAtlasSlotCount)
 
 		SHADER_PARAMETER_RDG_TEXTURE(Texture2D, RayMarchMaskTexture)
 
 		SHADER_PARAMETER_RDG_BUFFER_SRV(StructuredBuffer, DeepShadow_ViewInfoBuffer)
+		SHADER_PARAMETER_RDG_BUFFER_SRV(StructuredBuffer, DeepShadow_AtlasSlotIndexBuffer)
 		SHADER_PARAMETER_RDG_TEXTURE(Texture2D, DeepShadow_FrontDepthTexture)
 		SHADER_PARAMETER_RDG_TEXTURE(Texture2D, DeepShadow_DomTexture)
 		SHADER_PARAMETER_RDG_TEXTURE(Texture2D, SceneDepthTexture)
@@ -640,12 +641,12 @@ struct FHairStrandsDeepShadowParams
 {
 	uint32			OutputChannel = ~0;
 	FRDGBufferSRVRef DeepShadow_ViewInfoBuffer = nullptr;
+	FRDGBufferSRVRef DeepShadow_AtlasSlotIndexBuffer = nullptr;
 	FIntRect		DeepShadow_AtlasRect;
 	FRDGTextureRef	DeepShadow_FrontDepthTexture = nullptr;
 	FRDGTextureRef	DeepShadow_LayerTexture = nullptr;
 	float			DeepShadow_DepthBiasScale = 1;
 	float			DeepShadow_DensityScale = 1;
-	uint32			DeepShadow_AtlasSlotIndex = 0;
 	FVector4f		DeepShadow_LayerDepths = FVector4f(0, 0, 0, 0);
 };
 
@@ -657,6 +658,7 @@ static void AddHairStrandsDeepShadowMaskPass(
 	const FHairStrandsDeepShadowParams& Params,
 	const float FadeAlpha,
 	const uint32 EncodingType,
+	const uint32 EffectiveAtlasSlotCount,
 	FRDGTextureRef& OutShadowMask)
 {
 	check(OutShadowMask);
@@ -674,10 +676,12 @@ static void AddHairStrandsDeepShadowMaskPass(
 	Parameters->DeepShadow_DensityScale = Params.DeepShadow_DensityScale;
 	Parameters->DeepShadow_LayerDepths = Params.DeepShadow_LayerDepths;
 	Parameters->RenderTargets[0] = FRenderTargetBinding(OutShadowMask, ERenderTargetLoadAction::ELoad);
-	Parameters->DeepShadow_LayerDepths = Params.DeepShadow_LayerDepths;
-	Parameters->DeepShadow_SlotIndex = Params.DeepShadow_AtlasSlotIndex;
+
 	Parameters->DeepShadow_SlotOffset = FIntPoint(Params.DeepShadow_AtlasRect.Min.X, Params.DeepShadow_AtlasRect.Min.Y);
 	Parameters->DeepShadow_SlotResolution = FIntPoint(Params.DeepShadow_AtlasRect.Max.X - Params.DeepShadow_AtlasRect.Min.X, Params.DeepShadow_AtlasRect.Max.Y - Params.DeepShadow_AtlasRect.Min.Y);
+
+	Parameters->EffectiveAtlasSlotCount = EffectiveAtlasSlotCount;
+	Parameters->DeepShadow_AtlasSlotIndexBuffer = Params.DeepShadow_AtlasSlotIndexBuffer;
 	Parameters->DeepShadow_ViewInfoBuffer = Params.DeepShadow_ViewInfoBuffer;
 	Parameters->HairStrands = HairStrands::BindHairStrandsViewUniformParameters(View);
 	if (ShaderPrint::IsValid(View.ShaderPrintData))
@@ -958,48 +962,58 @@ static void InternalRenderHairStrandsShadowMask(
 		EncodingType = 2;
 	}
 
+	// 1. Build list of all the macrogroup affected by this light
 	bool bHasDeepShadow = false;
+	TArray<uint32> AtlasSlotIndexData;
+	AtlasSlotIndexData.Reserve(InMacroGroupDatas.Num());
+	float LayerDistribution = 0;
+	FIntRect AtlasRect = FIntRect();
 	if (!IsHairStrandsForVoxelTransmittanceAndShadowEnable())
 	{
-		// TODO change this into a single pass, instead of iterating onto all macrogroup, as we do for transmittance mask
-		FRDGBufferSRVRef DeepShadow_ViewInfoBufferSRV = nullptr;
 		for (const FHairStrandsMacroGroupData& MacroGroupData : InMacroGroupDatas)
 		{
+			uint32 AtlasSlotIndex = 0;
 			for (const FHairStrandsDeepShadowData& DomData : MacroGroupData.DeepShadowDatas)
 			{
-				if (DomData.LightId != LightSceneInfo->Id)
-					continue;
-
-				if (DeepShadow_ViewInfoBufferSRV == nullptr)
+				if (DomData.LightId == LightSceneInfo->Id)
 				{
-					DeepShadow_ViewInfoBufferSRV = GraphBuilder.CreateSRV(DeepShadowResources.DeepShadowViewInfoBuffer);
+					bHasDeepShadow = true;
+					AtlasRect = DomData.AtlasRect;
+					LayerDistribution = DomData.LayerDistribution;
+					AtlasSlotIndexData.Add(DomData.AtlasSlotIndex);
 				}
-
-				bHasDeepShadow = true;
-
-				FHairStrandsDeepShadowParams Params;
-				Params.DeepShadow_AtlasSlotIndex = DomData.AtlasSlotIndex;
-				Params.DeepShadow_ViewInfoBuffer = DeepShadow_ViewInfoBufferSRV;
-				Params.DeepShadow_AtlasRect = DomData.AtlasRect;
-				Params.DeepShadow_FrontDepthTexture = DeepShadowResources.DepthAtlasTexture;
-				Params.DeepShadow_LayerTexture = DeepShadowResources.LayersAtlasTexture;
-				Params.DeepShadow_DepthBiasScale = GetDeepShadowDepthBiasScale();
-				Params.DeepShadow_DensityScale = GetDeepShadowDensityScale();
-				Params.DeepShadow_LayerDepths = ComputeDeepShadowLayerDepths(DomData.LayerDistribution);
-				Params.OutputChannel = bProjectingForForwardShading ? LightSceneInfo->GetDynamicShadowMapChannel() : ~0;
-				AddHairStrandsDeepShadowMaskPass(
-					GraphBuilder,
-					SceneTextures.Depth.Resolve,
-					View,
-					Params,
-					FadeAlpha,
-					EncodingType,
-					OutShadowMask);
 			}
 		}
 	}
 
-	// If there is no deep shadow for this light, fallback on the voxel representation
+	// 2. Render deep shadow mask if any
+	if (bHasDeepShadow)
+	{
+		FRDGBufferRef AtlasSlotIndexBuffer = CreateStructuredBuffer(GraphBuilder, TEXT("Hair.DeepShadow.AtlasSlotIndexBuffer"), 4u, AtlasSlotIndexData.Num(), AtlasSlotIndexData.GetData(), 4u * AtlasSlotIndexData.Num());			
+
+		FHairStrandsDeepShadowParams Params;
+		Params.DeepShadow_AtlasSlotIndexBuffer = GraphBuilder.CreateSRV(AtlasSlotIndexBuffer);
+		Params.DeepShadow_ViewInfoBuffer = GraphBuilder.CreateSRV(DeepShadowResources.DeepShadowViewInfoBuffer);
+		Params.DeepShadow_AtlasRect = AtlasRect;
+		Params.DeepShadow_FrontDepthTexture = DeepShadowResources.DepthAtlasTexture;
+		Params.DeepShadow_LayerTexture = DeepShadowResources.LayersAtlasTexture;
+		Params.DeepShadow_DepthBiasScale = GetDeepShadowDepthBiasScale();
+		Params.DeepShadow_DensityScale = GetDeepShadowDensityScale();
+		Params.DeepShadow_LayerDepths = ComputeDeepShadowLayerDepths(LayerDistribution);
+		Params.OutputChannel = bProjectingForForwardShading ? LightSceneInfo->GetDynamicShadowMapChannel() : ~0;
+
+		AddHairStrandsDeepShadowMaskPass(
+			GraphBuilder,
+			SceneTextures.Depth.Resolve,
+			View,
+			Params,
+			FadeAlpha,
+			EncodingType,
+			AtlasSlotIndexData.Num() /* EffectiveAtlasSlotCount */,
+			OutShadowMask);
+	}
+
+	// 3. If there is no deep shadow for this light, fallback on the voxel representation
 	if (!bHasDeepShadow && HairStrands::HasViewHairStrandsVoxelData(View))
 	{
 		AddHairStrandsVoxelShadowMaskPass(
