VulkanRHI: Fix validation error by matching the CompositeSkyLightPS input to the FScreenPassVS output (float4 OutUVAndScreenPos).

#rnx

diff --git a/Engine/Shaders/Private/RayTracing/CompositeSkyLightPS.usf b/Engine/Shaders/Private/RayTracing/CompositeSkyLightPS.usf
--- a/Engine/Shaders/Private/RayTracing/CompositeSkyLightPS.usf
+++ b/Engine/Shaders/Private/RayTracing/CompositeSkyLightPS.usf
@@ -8,10 +8,11 @@ Texture2D<float4> SkyLightTexture;
 SamplerState SkyLightTextureSampler;
 
 void CompositeSkyLightPS(
-	in noperspective float2 UV : TEXCOORD0,
+	in noperspective float4 OutUVAndScreenPos : TEXCOORD0,
 	out float4 OutColor : SV_Target0
 )
 {
+	const float2 UV = OutUVAndScreenPos.xy;
 	FGBufferData GBufferData = GetGBufferDataFromSceneTextures(UV);
 	float3 Albedo = GBufferData.StoredBaseColor - GBufferData.StoredBaseColor * GBufferData.Metallic;
 	float4 SkyLight = SkyLightTexture.Sample(SkyLightTextureSampler, UV);
