Heterogeneous Volumes: Initial overlapping volume support

The single-scattering integration can now optionally use the Camera Adaptive Volumetric Shadow Map to query global transmittance when ray marching through a singular volume. This has the benefit of correctly attenuating through other heterogeneous volumes that may be overlapping within the domain. The fidelity of the attenuation is dependent upon the number of shadow samples used with the Camera AVSM.

Note:
The project setting: "Heterogeneous Volumes -> Composite with Translucency" must be enable to support overlapping heterogeneous volumes.

Recommended settings to enable the feature:
r.HeterogeneousVolumes.SupportOverlappingVolumes 1 (default = 0)
r.HeterogeneousVolumes.Shadows.MaxSampleCount 16 (default = 8)

diff --git a/Engine/Shaders/Private/HeterogeneousVolumes/HeterogeneousVolumesAdaptiveVolumetricShadowMapSampler.ush b/Engine/Shaders/Private/HeterogeneousVolumes/HeterogeneousVolumesAdaptiveVolumetricShadowMapSampler.ush
new file mode 100644
--- /dev/null
+++ b/Engine/Shaders/Private/HeterogeneousVolumes/HeterogeneousVolumesAdaptiveVolumetricShadowMapSampler.ush
@@ -0,0 +1,441 @@
+#pragma once
+
+#include "HeterogeneousVolumesAdaptiveVolumetricShadowMapUtils.ush"
+#include "HeterogeneousVolumesAdaptiveVolumetricShadowMapSampling.ush"
+
+#ifndef AVSM_BILINEAR_INTERPOLATION
+#define AVSM_BILINEAR_INTERPOLATION 1
+#endif //
+
+#define INVALID_SAMPLE_INDEX 0xFF
+#define INVALID_SAMPLE_INDEX4 0xFFFFFFFF
+
+#if 0
+// Reference single-channel implementation
+struct FAVSM_Sampler
+{
+	FAdaptiveVolumetricShadowMap ShadowMap;
+	uint PixelOffset;
+	uint SampleCount;
+	uint SampleIndex;
+
+	// Query cache
+	FAVSMSampleData SampleData[2];
+
+	// Convenience creation data
+	int Face;
+	int2 Pixel;
+	float DepthOffset;
+};
+
+bool IsValid(FAVSM_Sampler Iterator)
+{
+	return Iterator.SampleIndex != INVALID_SAMPLE_INDEX;
+}
+
+FAVSM_Sampler AVSM_Sampler_Create(
+	inout FAdaptiveVolumetricShadowMap ShadowMap,
+	int Face,
+	int2 Pixel,
+	float DepthOffset
+)
+{
+	FAVSM_Sampler Iterator;
+	Iterator.ShadowMap = ShadowMap;
+
+	uint LinearIndex = AVSM_LinearIndex(Iterator.ShadowMap, Face, Pixel);
+	FAVSMIndirectionData IndirectionData = AVSM_UnpackIndirectionData(Iterator.ShadowMap.IndirectionBuffer[LinearIndex]);
+	Iterator.PixelOffset = IndirectionData.PixelOffset;
+	Iterator.SampleCount = IndirectionData.SampleCount;
+
+	if (Iterator.SampleCount > 0)
+	{
+		int TopLevelSampleCount = CalcTopLevelSampleCount(Iterator.SampleCount);
+		Iterator.SampleIndex = 0;
+		Iterator.SampleData[0] = AVSM_UnpackSampleData(Iterator.ShadowMap.SampleBuffer[Iterator.PixelOffset + Align4(TopLevelSampleCount) + Iterator.SampleIndex]);
+		Iterator.SampleData[1] = AVSM_UnpackSampleData(Iterator.ShadowMap.SampleBuffer[Iterator.PixelOffset + Align4(TopLevelSampleCount) + Iterator.SampleIndex + 1]);
+	}
+	else
+	{
+		Iterator.SampleIndex = INVALID_SAMPLE_INDEX;
+		Iterator.SampleData[0] = AVSM_CreateSampleData(0.0, 1.0);
+		Iterator.SampleData[1] = AVSM_CreateSampleData(0.0, 1.0);
+	}
+
+	Iterator.Face = Face;
+	Iterator.Pixel = Pixel;
+	Iterator.DepthOffset = DepthOffset;
+	return Iterator;
+}
+
+float AVSM_Sampler_Eval(inout FAVSM_Sampler Iterator, float X)
+{
+	X -= Iterator.DepthOffset;
+	if (IsValid(Iterator))
+	{
+		int TopLevelSampleCount = CalcTopLevelSampleCount(Iterator.SampleCount);
+
+		while ((X > Iterator.SampleData[1].X) && (Iterator.SampleIndex < Iterator.SampleCount - 1))
+		{
+			Iterator.SampleIndex++;
+			Iterator.SampleData[0] = Iterator.SampleData[1];
+			Iterator.SampleData[1] = AVSM_UnpackSampleData(Iterator.ShadowMap.SampleBuffer[Iterator.PixelOffset + Align4(TopLevelSampleCount) + Iterator.SampleIndex]);
+		}
+	}
+
+	float BlendWeight = saturate((X - Iterator.SampleData[0].X) * SafeRcp(Iterator.SampleData[1].X - Iterator.SampleData[0].X));
+	float Tau = lerp(Iterator.SampleData[0].Tau, Iterator.SampleData[1].Tau, BlendWeight);
+	return Tau;
+}
+
+FAVSMSampleData AVSM_Sampler_EvalInverse(inout FAVSM_Sampler Iterator, float Tau)
+{
+	if (IsValid(Iterator))
+	{
+		int TopLevelSampleCount = CalcTopLevelSampleCount(Iterator.SampleCount);
+
+		while ((Tau < Iterator.SampleData[1].Tau) && (Iterator.SampleIndex < Iterator.SampleCount - 1))
+		{
+			Iterator.SampleIndex++;
+			Iterator.SampleData[0] = Iterator.SampleData[1];
+			Iterator.SampleData[1] = AVSM_UnpackSampleData(Iterator.ShadowMap.SampleBuffer[Iterator.PixelOffset + Align4(TopLevelSampleCount) + Iterator.SampleIndex]);
+		}
+	}
+
+	float BlendWeight = saturate((Tau - Iterator.SampleData[0].Tau) * SafeRcp(Iterator.SampleData[1].Tau - Iterator.SampleData[0].Tau));
+	float X = lerp(Iterator.SampleData[0].X, Iterator.SampleData[1].X, BlendWeight);
+	X += Iterator.DepthOffset;
+	return AVSM_CreateSampleData(X, Tau);
+}
+
+float AVSM_Sampler_Last(inout FAVSM_Sampler Iterator)
+{
+	float Transmittance = 1.0;
+	if (IsValid(Iterator))
+	{
+		int TopLevelSampleCount = CalcTopLevelSampleCount(Iterator.SampleCount);
+		if (TopLevelSampleCount > 0)
+		{
+			uint SampleIndex = max(Iterator.SampleCount - 2, 0);
+			FAVSMSampleData SampleData = AVSM_UnpackSampleData(Iterator.ShadowMap.SampleBuffer[Iterator.PixelOffset + Align4(TopLevelSampleCount) + SampleIndex]);
+
+			Transmittance = SampleData.Tau;
+		}
+	}
+	return Transmittance;
+}
+#endif
+
+// Sampler4 implementation reduces VGPR usage
+struct FAVSM_Sampler4
+{
+	FAdaptiveVolumetricShadowMap ShadowMap;
+	uint4 PixelOffset4;
+	uint SampleCount4;
+	uint SampleIndex4;
+
+	// Query cache
+	uint4 SampleData0Packed4;
+	uint4 SampleData1Packed4;
+
+	// Convenience creation data
+	int Face;
+	float2 Pixel;
+	float DepthOffset;
+};
+
+FAVSMSampleData AVSM_Sampler4_GetSampleData0(inout FAVSM_Sampler4 Sampler4, uint Index)
+{
+	FAVSMSampleData SampleData0 = AVSM_UnpackSampleData(Sampler4.SampleData0Packed4[Index]);
+	SampleData0.X += Sampler4.DepthOffset;
+	return SampleData0;
+}
+
+void AVSM_Sampler4_SetSampleData0(inout FAVSM_Sampler4 Sampler4, uint4 SampleDataPacked4)
+{
+	Sampler4.SampleData0Packed4 = SampleDataPacked4;
+}
+
+void AVSM_Sampler4_SetSampleData0(inout FAVSM_Sampler4 Sampler4, uint Index, uint SampleDataPacked)
+{
+	Sampler4.SampleData0Packed4[Index] = SampleDataPacked;
+}
+
+FAVSMSampleData AVSM_Sampler4_GetSampleData1(inout FAVSM_Sampler4 Sampler4, uint Index)
+{
+	FAVSMSampleData SampleData1 = AVSM_UnpackSampleData(Sampler4.SampleData1Packed4[Index]);
+	SampleData1.X += Sampler4.DepthOffset;
+	return SampleData1;
+}
+
+void AVSM_Sampler4_SetSampleData1(inout FAVSM_Sampler4 Sampler4, uint4 SampleDataPacked4)
+{
+	Sampler4.SampleData1Packed4 = SampleDataPacked4;
+}
+
+void AVSM_Sampler4_SetSampleData1(inout FAVSM_Sampler4 Sampler4, uint Index, uint SampleDataPacked)
+{
+	Sampler4.SampleData1Packed4[Index] = SampleDataPacked;
+}
+
+uint AVSM_Sampler4_GetSampleIndex(inout FAVSM_Sampler4 Sampler4, uint Index)
+{
+	uint SampleIndex = (Sampler4.SampleIndex4 >> (Index * 8u)) & INVALID_SAMPLE_INDEX;
+	return SampleIndex;
+}
+
+void AVSM_Sampler4_SetSampleIndex(inout FAVSM_Sampler4 Sampler4, int4 SampleIndex4)
+{
+	int SampleIndexPacked = 0;
+	for (uint Index = 0; Index < 4; ++Index)
+	{
+		SampleIndexPacked |= ((SampleIndex4[Index] & INVALID_SAMPLE_INDEX) << (Index * 8u));
+	}
+	Sampler4.SampleIndex4 = SampleIndexPacked;
+}
+
+void AVSM_Sampler4_SetSampleIndex(inout FAVSM_Sampler4 Sampler4, uint Index, uint SampleIndex)
+{
+	uint WriteMask = INVALID_SAMPLE_INDEX << (Index * 8u);
+	Sampler4.SampleIndex4 &= ~WriteMask;
+	Sampler4.SampleIndex4 |= (SampleIndex & INVALID_SAMPLE_INDEX) << (Index * 8u);
+}
+
+uint AVSM_Sampler4_GetSampleCount(inout FAVSM_Sampler4 Sampler4, uint Index)
+{
+	uint SampleCount = (Sampler4.SampleCount4 >> (Index * 8u)) & INVALID_SAMPLE_INDEX;
+	return SampleCount;
+}
+
+bool AVSM_Sampler4_IsValid(inout FAVSM_Sampler4 Sampler4, uint Index)
+{
+	return AVSM_Sampler4_GetSampleIndex(Sampler4, Index) != INVALID_SAMPLE_INDEX;
+}
+
+void AVSM_Sampler4_SetSampleCount(inout FAVSM_Sampler4 Sampler4, int4 SampleCount4)
+{
+	uint SampleCountPacked = 0;
+	for (uint Index = 0; Index < 4; ++Index)
+	{
+		SampleCountPacked |= ((SampleCount4[Index] & INVALID_SAMPLE_INDEX) << (Index * 8u));
+	}
+	Sampler4.SampleCount4 = SampleCountPacked;
+}
+
+uint AVSM_Sampler4_GetPixelOffset(inout FAVSM_Sampler4 Sampler4, uint Index)
+{
+	uint PixelOffset = Sampler4.PixelOffset4[Index];
+	return PixelOffset;
+}
+
+void AVSM_Sampler4_SetPixelOffset(inout FAVSM_Sampler4 Sampler4, uint4 PixelOffset)
+{
+	Sampler4.PixelOffset4 = PixelOffset;
+}
+
+FAVSM_Sampler4 AVSM_Sampler4_Create(
+	inout FAdaptiveVolumetricShadowMap ShadowMap,
+	int Face,
+	float2 Pixel,
+	float DepthOffset
+)
+{
+	FAVSM_Sampler4 Sampler4;
+	Sampler4.ShadowMap = ShadowMap;
+
+	uint2 Pixel4[4] = {
+		clamp(Pixel, 0, ShadowMap.Resolution - 1),
+		clamp(Pixel + int2(1, 0), 0, ShadowMap.Resolution - 1),
+		clamp(Pixel + int2(0, 1), 0, ShadowMap.Resolution - 1),
+		clamp(Pixel + int2(1, 1), 0, ShadowMap.Resolution - 1)
+	};
+	uint4 LinearIndex4 = uint4(
+		AVSM_LinearIndex(Sampler4.ShadowMap, Face, Pixel4[0]),
+		AVSM_LinearIndex(Sampler4.ShadowMap, Face, Pixel4[1]),
+		AVSM_LinearIndex(Sampler4.ShadowMap, Face, Pixel4[2]),
+		AVSM_LinearIndex(Sampler4.ShadowMap, Face, Pixel4[3])
+	);
+
+	uint4 SampleCount4 = INVALID_SAMPLE_INDEX4;
+	uint4 SampleIndex4 = INVALID_SAMPLE_INDEX4;
+	uint4 SampleData0Packed4 = 0;
+	uint4 SampleData1Packed4 = 0;
+	uint4 PixelOffset4 = 0;
+	for (int Index = 0; Index < 4; ++Index)
+	{
+		FAVSMIndirectionData IndirectionData = AVSM_UnpackIndirectionData(ShadowMap.IndirectionBuffer[LinearIndex4[Index]]);
+		PixelOffset4[Index] = IndirectionData.PixelOffset;
+		SampleCount4[Index] = IndirectionData.SampleCount;
+
+		if (SampleCount4[Index] > 0)
+		{
+			int TopLevelSampleCount = CalcTopLevelSampleCount(SampleCount4[Index]);
+			SampleIndex4[Index] = 0;
+			SampleData0Packed4[Index] = ShadowMap.SampleBuffer[PixelOffset4[Index] + Align4(TopLevelSampleCount) + SampleIndex4[Index]];
+			SampleData1Packed4[Index] = ShadowMap.SampleBuffer[PixelOffset4[Index] + Align4(TopLevelSampleCount) + SampleIndex4[Index] + 1];
+		}
+		else
+		{
+			SampleIndex4[Index] = INVALID_SAMPLE_INDEX;
+			SampleData0Packed4[Index] = AVSM_PackSampleData(0.0, 1.0);
+			SampleData1Packed4[Index] = AVSM_PackSampleData(0.0, 1.0);
+		}
+	}
+	AVSM_Sampler4_SetSampleCount(Sampler4, SampleCount4);
+	AVSM_Sampler4_SetSampleIndex(Sampler4, SampleIndex4);
+	AVSM_Sampler4_SetSampleData0(Sampler4, SampleData0Packed4);
+	AVSM_Sampler4_SetSampleData1(Sampler4, SampleData1Packed4);
+	AVSM_Sampler4_SetPixelOffset(Sampler4, PixelOffset4);
+
+	Sampler4.Face = Face;
+	Sampler4.Pixel = Pixel;
+	Sampler4.DepthOffset = DepthOffset;
+	return Sampler4;
+}
+
+float AVSM_Sampler4_Eval_Internal(inout FAVSM_Sampler4 Sampler4, int Index, float X)
+{
+	X -= Sampler4.DepthOffset;
+
+	FAVSMSampleData SampleData[] = {
+		AVSM_Sampler4_GetSampleData0(Sampler4, Index),
+		AVSM_Sampler4_GetSampleData1(Sampler4, Index)
+	};
+
+	if (AVSM_Sampler4_IsValid(Sampler4, Index))
+	{
+		int SampleIndex = AVSM_Sampler4_GetSampleIndex(Sampler4, Index);
+		int SampleCount = AVSM_Sampler4_GetSampleCount(Sampler4, Index);
+		int TopLevelSampleCount = CalcTopLevelSampleCount(SampleCount);
+
+		while ((X > SampleData[1].X) && (SampleIndex < SampleCount - 1))
+		{
+			SampleIndex++;
+			SampleData[0] = SampleData[1];
+			uint PixelOffset = AVSM_Sampler4_GetPixelOffset(Sampler4, Index);
+			SampleData[1] = AVSM_UnpackSampleData(Sampler4.ShadowMap.SampleBuffer[PixelOffset + Align4(TopLevelSampleCount) + SampleIndex]);
+
+			AVSM_Sampler4_SetSampleIndex(Sampler4, Index, SampleIndex);
+			AVSM_Sampler4_SetSampleData0(Sampler4, Index, AVSM_PackSampleData(SampleData[0]));
+			AVSM_Sampler4_SetSampleData1(Sampler4, Index, AVSM_PackSampleData(SampleData[1]));
+		}
+	}
+
+	float BlendWeight = saturate((X - SampleData[0].X) * SafeRcp(SampleData[1].X - SampleData[0].X));
+	float Tau = lerp(SampleData[0].Tau, SampleData[1].Tau, BlendWeight);
+	return Tau;
+}
+
+float AVSM_Bilinear_Interpolate(float Value[4], float2 Weight)
+{
+	float ValueX[] = {
+		lerp(Value[0], Value[1], Weight.x),
+		lerp(Value[2], Value[3], Weight.x)
+	};
+	return lerp(ValueX[0], ValueX[1], Weight.y);
+}
+
+float AVSM_Sampler4_Eval(inout FAVSM_Sampler4 Sampler4, float X)
+{
+	float Tau = AVSM_Sampler4_Eval_Internal(Sampler4, 0, X);
+#if AVSM_BILINEAR_INTERPOLATION != 0
+	float Tau4[] = {
+			Tau,
+			AVSM_Sampler4_Eval_Internal(Sampler4, 1, X),
+			AVSM_Sampler4_Eval_Internal(Sampler4, 2, X),
+			AVSM_Sampler4_Eval_Internal(Sampler4, 3, X)
+	};
+	Tau = AVSM_Bilinear_Interpolate(Tau4, frac(Sampler4.Pixel));
+#endif // AVSM_BILINEAR_INTERPOLATION != 0
+	return Tau;
+}
+
+FAVSMSampleData AVSM_Sampler4_EvalInverse_Internal(inout FAVSM_Sampler4 Sampler4, int Index, float Tau)
+{
+	FAVSMSampleData SampleData[] = {
+		AVSM_Sampler4_GetSampleData0(Sampler4, Index),
+		AVSM_Sampler4_GetSampleData1(Sampler4, Index)
+	};
+
+	if (AVSM_Sampler4_IsValid(Sampler4, Index))
+	{
+		int SampleCount = AVSM_Sampler4_GetSampleCount(Sampler4, Index);
+		int TopLevelSampleCount = CalcTopLevelSampleCount(SampleCount);
+
+		int SampleIndex = AVSM_Sampler4_GetSampleIndex(Sampler4, Index);
+		uint PixelOffset = AVSM_Sampler4_GetPixelOffset(Sampler4, Index);
+		while ((Tau < SampleData[1].Tau) && (SampleIndex < SampleCount - 1))
+		{
+			SampleIndex++;
+			SampleData[0] = SampleData[1];
+			SampleData[1] = AVSM_UnpackSampleData(Sampler4.ShadowMap.SampleBuffer[PixelOffset + Align4(TopLevelSampleCount) + SampleIndex]);
+
+			AVSM_Sampler4_SetSampleData0(Sampler4, Index, AVSM_PackSampleData(SampleData[0]));
+			AVSM_Sampler4_SetSampleData1(Sampler4, Index, AVSM_PackSampleData(SampleData[1]));
+		}
+	}
+
+	float BlendWeight = saturate((Tau - SampleData[0].Tau) * SafeRcp(SampleData[1].Tau - SampleData[0].Tau));
+	float X = lerp(SampleData[0].X, SampleData[1].X, BlendWeight);
+	X += Sampler4.DepthOffset;
+	return AVSM_CreateSampleData(X, Tau);
+}
+
+FAVSMSampleData AVSM_Sampler4_EvalInverse(inout FAVSM_Sampler4 Sampler4, float Tau)
+{
+	FAVSMSampleData SampleData = AVSM_Sampler4_EvalInverse_Internal(Sampler4, 0, Tau);
+#if AVSM_BILINEAR_INTERPOLATION != 0
+	float X4[] = {
+		SampleData.X,
+		AVSM_Sampler4_EvalInverse_Internal(Sampler4, 1, Tau).X,
+		AVSM_Sampler4_EvalInverse_Internal(Sampler4, 2, Tau).X,
+		AVSM_Sampler4_EvalInverse_Internal(Sampler4, 3, Tau).X
+	};
+	SampleData.X = AVSM_Bilinear_Interpolate(X4, frac(Sampler4.Pixel));
+#endif // AVSM_BILINEAR_INTERPOLATION != 0
+	return SampleData;
+}
+
+FAVSMSampleData AVSM_Sampler4_Last_Internal(inout FAVSM_Sampler4 Sampler4, int Index)
+{
+	FAVSMSampleData SampleData = AVSM_CreateSampleData(0.0, 1.0);
+
+	if (AVSM_Sampler4_IsValid(Sampler4, Index))
+	{
+		int SampleCount = AVSM_Sampler4_GetSampleCount(Sampler4, Index);
+		int TopLevelSampleCount = CalcTopLevelSampleCount(SampleCount);
+		if (TopLevelSampleCount > 0)
+		{
+			uint SampleIndex = max(SampleCount - 2, 0);
+			uint PixelOffset = AVSM_Sampler4_GetPixelOffset(Sampler4, Index);
+			SampleData = AVSM_UnpackSampleData(Sampler4.ShadowMap.SampleBuffer[PixelOffset + Align4(TopLevelSampleCount) + SampleIndex]);
+			SampleData.X += Sampler4.DepthOffset;
+		}
+	}
+
+	return SampleData;
+}
+
+FAVSMSampleData AVSM_Sampler4_Last(inout FAVSM_Sampler4 Sampler4)
+{
+	FAVSMSampleData SampleData = AVSM_Sampler4_Last_Internal(Sampler4, 0);
+#if AVSM_BILINEAR_INTERPOLATION != 0
+	// TODO: Filter X when sentinel value is removed
+	//float Depth4[] = {
+	//	SampleData.X,
+	//	AVSM_Sampler4_Last_Internal(Sampler4, 1).X,
+	//	AVSM_Sampler4_Last_Internal(Sampler4, 2).X,
+	//	AVSM_Sampler4_Last_Internal(Sampler4, 3).X
+	//};
+	//SampleData.X = AVSM_Bilinear_Interpolate(Depth4, frac(Sampler4.Pixel));
+
+	float Tau4[] = {
+		SampleData.Tau,
+		AVSM_Sampler4_Last_Internal(Sampler4, 1).Tau,
+		AVSM_Sampler4_Last_Internal(Sampler4, 2).Tau,
+		AVSM_Sampler4_Last_Internal(Sampler4, 3).Tau
+	};
+	SampleData.Tau = AVSM_Bilinear_Interpolate(Tau4, frac(Sampler4.Pixel));
+#endif // AVSM_BILINEAR_INTERPOLATION != 0
+	return SampleData;
+}
diff --git a/Engine/Shaders/Private/HeterogeneousVolumes/HeterogeneousVolumesAdaptiveVolumetricShadowMapSampling.ush b/Engine/Shaders/Private/HeterogeneousVolumes/HeterogeneousVolumesAdaptiveVolumetricShadowMapSampling.ush
--- a/Engine/Shaders/Private/HeterogeneousVolumes/HeterogeneousVolumesAdaptiveVolumetricShadowMapSampling.ush
+++ b/Engine/Shaders/Private/HeterogeneousVolumes/HeterogeneousVolumesAdaptiveVolumetricShadowMapSampling.ush
@@ -14,9 +14,58 @@
 #define AVSM_SAMPLE_MODE AVSM_SAMPLE_MODE_TWO_LEVEL_VECTORIZED_LOOP
 #endif // AVSM_SAMPLE_MODE
 
-uint GetMyLinearIndex(uint Face, uint2 PixelCoord)
+struct FAdaptiveVolumetricShadowMap
 {
-	uint LinearIndex = Face * AVSM.Resolution.x * AVSM.Resolution.y + PixelCoord.y * AVSM.Resolution.x + PixelCoord.x;
+	float4x4 TranslatedWorldToShadow[6];
+	float3 TranslatedWorldOrigin;
+	float4 TranslatedWorldPlane;
+
+	int2 Resolution;
+	int NumShadowMatrices;
+	int MaxSampleCount;
+	bool bIsEmpty;
+	bool bIsDirectionalLight;
+
+	StructuredBuffer<uint2> IndirectionBuffer;
+	StructuredBuffer<uint> SampleBuffer;
+};
+
+bool AVSM_PixelAndDepth(
+	inout FAdaptiveVolumetricShadowMap ShadowMap,
+	int Face,
+	float3 TranslatedWorldPosition,
+	inout float2 Pixel,
+	inout float Depth
+)
+{
+	Depth = -65535.0;
+	Pixel = -1;
+	float4 ShadowPositionAsFloat4 = mul(float4(TranslatedWorldPosition, 1), ShadowMap.TranslatedWorldToShadow[Face]);
+	if (ShadowPositionAsFloat4.w > 0)
+	{
+		float3 ShadowPosition = ShadowPositionAsFloat4.xyz / ShadowPositionAsFloat4.w;
+		if (all(ShadowPosition.xyz > 0) || all(ShadowPosition.xy < 1))
+		{
+			Pixel = clamp(ShadowPosition.xy * ShadowMap.Resolution - 0.5, 0, ShadowMap.Resolution - 1);
+			Depth = length(TranslatedWorldPosition - ShadowMap.TranslatedWorldOrigin);
+			if (ShadowMap.bIsDirectionalLight)
+			{
+				Depth = dot(ShadowMap.TranslatedWorldPlane.xyz, TranslatedWorldPosition) + ShadowMap.TranslatedWorldPlane.w;
+			}
+			return true;
+		}
+	}
+
+	return false;
+}
+
+uint AVSM_LinearIndex(
+	inout FAdaptiveVolumetricShadowMap ShadowMap,
+	uint Face,
+	uint2 PixelCoord
+)
+{
+	uint LinearIndex = Face * ShadowMap.Resolution.x * ShadowMap.Resolution.y + PixelCoord.y * ShadowMap.Resolution.x + PixelCoord.x;
 	return LinearIndex;
 }
 
@@ -26,10 +75,15 @@ float SafeRcp(float Value)
 }
 
 #if 0
-float AVSM_Sample_Reference(int Face, int2 Pixel, float WorldHitT)
+float AVSM_Sample_Reference(
+	inout FAdaptiveVolumetricShadowMap ShadowMap,
+	int Face,
+	int2 Pixel,
+	float WorldHitT
+)
 {
-	uint LinearIndex = GetMyLinearIndex(Face, Pixel);
-	FAVSMIndirectionData IndirectionData = ASVM_UnpackIndirectionData(AVSM.IndirectionBuffer[LinearIndex]);
+	uint LinearIndex = AVSM_LinearIndex(ShadowMap, Face, Pixel);
+	FAVSMIndirectionData IndirectionData = AVSM_UnpackIndirectionData(ShadowMap.IndirectionBuffer[LinearIndex]);
 	if (WorldHitT < IndirectionData.BeginX)
 	{
 		return 1.0;
@@ -47,16 +101,16 @@ float AVSM_Sample_Reference(int Face, int2 Pixel, float WorldHitT)
 		if (WorldHitT >= IndirectionData.EndX)
 		{
 			int LastSampleIndex = IndirectionData.PixelOffset + Align4(TopLevelSampleCount) + IndirectionData.SampleCount - 1;
-			return AVSM_UnpackSampleData(AVSM.SampleBuffer[LastSampleIndex]).Tau;
+			return AVSM_UnpackSampleData(ShadowMap.SampleBuffer[LastSampleIndex]).Tau;
 		}
 
 		// Initialize from the top-level
-		ShadowData[0] = AVSM_UnpackSampleData(AVSM.SampleBuffer[IndirectionData.PixelOffset]);
+		ShadowData[0] = AVSM_UnpackSampleData(ShadowMap.SampleBuffer[IndirectionData.PixelOffset]);
 
 		// Iterate through the linear span of offset transmittance information
 		for (int ElementIndex = 0; ElementIndex < IndirectionData.SampleCount - 1; ++ElementIndex)
 		{
-			ShadowData[1] = AVSM_UnpackSampleData(AVSM.SampleBuffer[IndirectionData.PixelOffset + Align4(TopLevelSampleCount) + ElementIndex]);
+			ShadowData[1] = AVSM_UnpackSampleData(ShadowMap.SampleBuffer[IndirectionData.PixelOffset + Align4(TopLevelSampleCount) + ElementIndex]);
 			if (WorldHitT < ShadowData[1].X)
 			{
 				break;
@@ -71,10 +125,15 @@ float AVSM_Sample_Reference(int Face, int2 Pixel, float WorldHitT)
 	return Transmittance;
 }
 
-float AVSM_Sample_BinarySearch(int Face, int2 Pixel, float WorldHitT)
+float AVSM_Sample_BinarySearch(
+	inout FAdaptiveVolumetricShadowMap ShadowMap,
+	int Face,
+	int2 Pixel,
+	float WorldHitT
+)
 {
-	uint LinearIndex = GetMyLinearIndex(Face, Pixel);
-	FAVSMIndirectionData IndirectionData = ASVM_UnpackIndirectionData(AVSM.IndirectionBuffer[LinearIndex]);
+	uint LinearIndex = AVSM_LinearIndex(ShadowMap, Face, Pixel);
+	FAVSMIndirectionData IndirectionData = AVSM_UnpackIndirectionData(ShadowMap.IndirectionBuffer[LinearIndex]);
 	if (WorldHitT < IndirectionData.BeginX)
 	{
 		return 1.0;
@@ -92,11 +151,11 @@ float AVSM_Sample_BinarySearch(int Face, int2 Pixel, float WorldHitT)
 		if (WorldHitT >= IndirectionData.EndX)
 		{
 			int LastSampleIndex = IndirectionData.PixelOffset + Align4(TopLevelSampleCount) + IndirectionData.SampleCount - 1;
-			return AVSM_UnpackSampleData(AVSM.SampleBuffer[LastSampleIndex]).Tau;
+			return AVSM_UnpackSampleData(ShadowMap.SampleBuffer[LastSampleIndex]).Tau;
 		}
 
 		// Initialize from the top-level
-		ShadowData[0] = AVSM_UnpackSampleData(AVSM.SampleBuffer[IndirectionData.PixelOffset]);
+		ShadowData[0] = AVSM_UnpackSampleData(ShadowMap.SampleBuffer[IndirectionData.PixelOffset]);
 
 		// Account for the shifted bottom-level
 		int ElementIndexMin = -1;
@@ -104,7 +163,7 @@ float AVSM_Sample_BinarySearch(int Face, int2 Pixel, float WorldHitT)
 		while (ElementIndexMax - ElementIndexMin > 1)
 		{
 			int ElementIndex = (ElementIndexMin + ElementIndexMax + 1) / 2;
-			FAVSMSampleData SampleData = AVSM_UnpackSampleData(AVSM.SampleBuffer[IndirectionData.PixelOffset + Align4(TopLevelSampleCount) + ElementIndex]);
+			FAVSMSampleData SampleData = AVSM_UnpackSampleData(ShadowMap.SampleBuffer[IndirectionData.PixelOffset + Align4(TopLevelSampleCount) + ElementIndex]);
 
 			if (WorldHitT < SampleData.X)
 			{
@@ -124,10 +183,15 @@ float AVSM_Sample_BinarySearch(int Face, int2 Pixel, float WorldHitT)
 	return Transmittance;
 }
 
-float AVSM_Sample_TwoLevel_Reference(int Face, int2 Pixel, float WorldHitT)
+float AVSM_Sample_TwoLevel_Reference(
+	inout FAdaptiveVolumetricShadowMap ShadowMap,
+	int Face,
+	int2 Pixel,
+	float WorldHitT
+)
 {
-	uint LinearIndex = GetMyLinearIndex(Face, Pixel);
-	FAVSMIndirectionData IndirectionData = ASVM_UnpackIndirectionData(AVSM.IndirectionBuffer[LinearIndex]);
+	uint LinearIndex = AVSM_LinearIndex(ShadowMap, Face, Pixel);
+	FAVSMIndirectionData IndirectionData = AVSM_UnpackIndirectionData(ShadowMap.IndirectionBuffer[LinearIndex]);
 	if (WorldHitT < IndirectionData.BeginX)
 	{
 		return 1.0;
@@ -150,12 +214,12 @@ float AVSM_Sample_TwoLevel_Reference(int Face, int2 Pixel, float WorldHitT)
 			return AVSM_UnpackSampleData(AVSM.SampleBuffer[LastSampleIndex]).Tau;
 		}
 
-		ShadowData[0] = AVSM_UnpackSampleData(AVSM.SampleBuffer[IndirectionData.PixelOffset]);
+		ShadowData[0] = AVSM_UnpackSampleData(ShadowMap.SampleBuffer[IndirectionData.PixelOffset]);
 
 		// Iterate through the top-level
 		for (TopLevelSampleIndex = 0; TopLevelSampleIndex < TopLevelSampleCount; ++TopLevelSampleIndex)
 		{
-			ShadowData[1] = AVSM_UnpackSampleData(AVSM.SampleBuffer[IndirectionData.PixelOffset + TopLevelSampleIndex]);
+			ShadowData[1] = AVSM_UnpackSampleData(ShadowMap.SampleBuffer[IndirectionData.PixelOffset + TopLevelSampleIndex]);
 			if (WorldHitT < ShadowData[1].X)
 			{
 				break;
@@ -171,7 +235,7 @@ float AVSM_Sample_TwoLevel_Reference(int Face, int2 Pixel, float WorldHitT)
 	int SampleIndexMax = min(TopLevelSampleIndex * 4, IndirectionData.SampleCount - 1);
 	for (int SampleIndex = SampleIndexMin; SampleIndex < SampleIndexMax; ++SampleIndex)
 	{
-		ShadowData[1] = AVSM_UnpackSampleData(AVSM.SampleBuffer[IndirectionData.PixelOffset + Align4(TopLevelSampleCount) + SampleIndex]);
+		ShadowData[1] = AVSM_UnpackSampleData(ShadowMap.SampleBuffer[IndirectionData.PixelOffset + Align4(TopLevelSampleCount) + SampleIndex]);
 		if (WorldHitT < ShadowData[1].X)
 		{
 			break;
@@ -187,10 +251,15 @@ float AVSM_Sample_TwoLevel_Reference(int Face, int2 Pixel, float WorldHitT)
 }
 #endif
 
-float AVSM_Sample_TwoLevel_Vectorized(int Face, int2 Pixel, float WorldHitT)
+float AVSM_Sample_TwoLevel_Vectorized(
+	inout FAdaptiveVolumetricShadowMap ShadowMap,
+	int Face,
+	int2 Pixel,
+	float WorldHitT
+)
 {
-	uint LinearIndex = GetMyLinearIndex(Face, Pixel);
-	FAVSMIndirectionData IndirectionData = ASVM_UnpackIndirectionData(AVSM.IndirectionBuffer[LinearIndex]);
+	uint LinearIndex = AVSM_LinearIndex(ShadowMap, Face, Pixel);
+	FAVSMIndirectionData IndirectionData = AVSM_UnpackIndirectionData(ShadowMap.IndirectionBuffer[LinearIndex]);
 	if (WorldHitT < IndirectionData.BeginX)
 	{
 		return 1.0;
@@ -210,7 +279,7 @@ float AVSM_Sample_TwoLevel_Vectorized(int Face, int2 Pixel, float WorldHitT)
 		if (WorldHitT >= IndirectionData.EndX)
 		{
 			int LastIndex = IndirectionData.PixelOffset + Align4(TopLevelSampleCount) + IndirectionData.SampleCount - 1;
-			return AVSM_UnpackSampleData(AVSM.SampleBuffer[LastIndex]).Tau;
+			return AVSM_UnpackSampleData(ShadowMap.SampleBuffer[LastIndex]).Tau;
 		}
 
 		int TopLevelSampleOffset = 0;
@@ -221,10 +290,10 @@ float AVSM_Sample_TwoLevel_Vectorized(int Face, int2 Pixel, float WorldHitT)
 			int4 IndexOffset = TopLevelSampleOffset + int4(0, 1, 2, 3);
 
 			int4 PackedShadowData4 = int4(
-				AVSM.SampleBuffer[IndirectionData.PixelOffset + IndexOffset[0]],
-				AVSM.SampleBuffer[IndirectionData.PixelOffset + IndexOffset[1]],
-				AVSM.SampleBuffer[IndirectionData.PixelOffset + IndexOffset[2]],
-				AVSM.SampleBuffer[IndirectionData.PixelOffset + IndexOffset[3]]
+				ShadowMap.SampleBuffer[IndirectionData.PixelOffset + IndexOffset[0]],
+				ShadowMap.SampleBuffer[IndirectionData.PixelOffset + IndexOffset[1]],
+				ShadowMap.SampleBuffer[IndirectionData.PixelOffset + IndexOffset[2]],
+				ShadowMap.SampleBuffer[IndirectionData.PixelOffset + IndexOffset[3]]
 			);
 
 			float4 Depths = float4(
@@ -256,10 +325,10 @@ float AVSM_Sample_TwoLevel_Vectorized(int Face, int2 Pixel, float WorldHitT)
 		IndexOffset += int4(0, 1, 2, 3);
 
 		int4 PackedShadowData4 = int4(
-			AVSM.SampleBuffer[IndirectionData.PixelOffset + Align4(TopLevelSampleCount) + IndexOffset.x],
-			AVSM.SampleBuffer[IndirectionData.PixelOffset + Align4(TopLevelSampleCount) + IndexOffset.y],
-			AVSM.SampleBuffer[IndirectionData.PixelOffset + Align4(TopLevelSampleCount) + IndexOffset.z],
-			AVSM.SampleBuffer[IndirectionData.PixelOffset + Align4(TopLevelSampleCount) + IndexOffset.w]
+			ShadowMap.SampleBuffer[IndirectionData.PixelOffset + Align4(TopLevelSampleCount) + IndexOffset.x],
+			ShadowMap.SampleBuffer[IndirectionData.PixelOffset + Align4(TopLevelSampleCount) + IndexOffset.y],
+			ShadowMap.SampleBuffer[IndirectionData.PixelOffset + Align4(TopLevelSampleCount) + IndexOffset.z],
+			ShadowMap.SampleBuffer[IndirectionData.PixelOffset + Align4(TopLevelSampleCount) + IndexOffset.w]
 		);
 
 		float4 Depths = float4(
@@ -285,14 +354,23 @@ float AVSM_Sample_TwoLevel_Vectorized(int Face, int2 Pixel, float WorldHitT)
 	return Transmittance;
 }
 
-float AVSM_Sample(int Face, int2 Pixel, float WorldHitT)
+float AVSM_Sample(
+	inout FAdaptiveVolumetricShadowMap ShadowMap,
+	int Face,
+	int2 Pixel,
+	float WorldHitT
+)
 {
-	return AVSM_Sample_TwoLevel_Vectorized(Face, Pixel, WorldHitT);
+	return AVSM_Sample_TwoLevel_Vectorized(ShadowMap, Face, Pixel, WorldHitT);
 }
 
-float AVSM_SampleTransmittance(float3 TranslatedWorldPosition, float3 LightTranslatedWorldPosition)
+float AVSM_SampleTransmittance(
+	inout FAdaptiveVolumetricShadowMap ShadowMap,
+	float3 TranslatedWorldPosition,
+	float3 LightTranslatedWorldPosition
+)
 {
-	if (any(AVSM.Resolution == 0))
+	if (any(ShadowMap.Resolution == 0))
 	{
 		return 1.0;
 	}
@@ -302,7 +380,7 @@ float AVSM_SampleTransmittance(float3 TranslatedWorldPosition, float3 LightTrans
 	int Face = 0;
 	if (AVSM.NumShadowMatrices > 1)
 	{
-		float3 LightDir = TranslatedWorldPosition - AVSM.TranslatedWorldOrigin;
+		float3 LightDir = TranslatedWorldPosition - ShadowMap.TranslatedWorldOrigin;
 		float3 LightDirAbs = abs(LightDir);
 
 		Face = LightDirAbs[1] > LightDirAbs[0] ? 1 : 0;
@@ -315,17 +393,17 @@ float AVSM_SampleTransmittance(float3 TranslatedWorldPosition, float3 LightTrans
 	float4 ShadowPositionAsFloat4 = 0;
 	float3 ShadowPosition = 0;
 	{
-		ShadowPositionAsFloat4 = mul(float4(TranslatedWorldPosition, 1), AVSM.TranslatedWorldToShadow[Face]);
+		ShadowPositionAsFloat4 = mul(float4(TranslatedWorldPosition, 1), ShadowMap.TranslatedWorldToShadow[Face]);
 		if (ShadowPositionAsFloat4.w > 0)
 		{
 			ShadowPosition = ShadowPositionAsFloat4.xyz / ShadowPositionAsFloat4.w;
 			if (all(ShadowPosition.xyz > 0) || all(ShadowPosition.xy < 1))
 			{
-				float2 ShadowMapCoord = clamp(ShadowPosition.xy * AVSM.Resolution - 0.5, 0, AVSM.Resolution - 1);
-				float ShadowZ = length(TranslatedWorldPosition - AVSM.TranslatedWorldOrigin);
-				if (AVSM.bIsDirectionalLight)
+				float2 ShadowMapCoord = clamp(ShadowPosition.xy * ShadowMap.Resolution - 0.5, 0, ShadowMap.Resolution - 1);
+				float ShadowZ = length(TranslatedWorldPosition - ShadowMap.TranslatedWorldOrigin);
+				if (ShadowMap.bIsDirectionalLight)
 				{
-					ShadowZ = dot(AVSM.TranslatedWorldPlane.xyz, TranslatedWorldPosition) + AVSM.TranslatedWorldPlane.w;
+					ShadowZ = dot(ShadowMap.TranslatedWorldPlane.xyz, TranslatedWorldPosition) + ShadowMap.TranslatedWorldPlane.w;
 				}
 
 			#define BILINEAR_INTERPOLATION 1
@@ -333,18 +411,18 @@ float AVSM_SampleTransmittance(float3 TranslatedWorldPosition, float3 LightTrans
 				// Bilinear interpolation of shadow data
 				int2 ShadowMapCoordX[4] =
 				{
-					clamp(ShadowMapCoord, 0, AVSM.Resolution - 1),
-					clamp(ShadowMapCoord + int2(1, 0), 0, AVSM.Resolution - 1),
-					clamp(ShadowMapCoord + int2(0, 1), 0, AVSM.Resolution - 1),
-					clamp(ShadowMapCoord + int2(1, 1), 0, AVSM.Resolution - 1)
+					clamp(ShadowMapCoord, 0, ShadowMap.Resolution - 1),
+					clamp(ShadowMapCoord + int2(1, 0), 0, ShadowMap.Resolution - 1),
+					clamp(ShadowMapCoord + int2(0, 1), 0, ShadowMap.Resolution - 1),
+					clamp(ShadowMapCoord + int2(1, 1), 0, ShadowMap.Resolution - 1)
 				};
 
 				float TransmittanceX[4] =
 				{
-					AVSM_Sample(Face, ShadowMapCoordX[0], ShadowZ),
-					AVSM_Sample(Face, ShadowMapCoordX[1], ShadowZ),
-					AVSM_Sample(Face, ShadowMapCoordX[2], ShadowZ),
-					AVSM_Sample(Face, ShadowMapCoordX[3], ShadowZ)
+					AVSM_Sample(ShadowMap, Face, ShadowMapCoordX[0], ShadowZ),
+					AVSM_Sample(ShadowMap, Face, ShadowMapCoordX[1], ShadowZ),
+					AVSM_Sample(ShadowMap, Face, ShadowMapCoordX[2], ShadowZ),
+					AVSM_Sample(ShadowMap, Face, ShadowMapCoordX[3], ShadowZ)
 				};
 
 				float2 Weight = frac(ShadowMapCoord);
@@ -352,7 +430,7 @@ float AVSM_SampleTransmittance(float3 TranslatedWorldPosition, float3 LightTrans
 				float TransmittanceY1 = lerp(TransmittanceX[2], TransmittanceX[3], Weight.x);
 				Transmittance = lerp(TransmittanceY0, TransmittanceY1, Weight.y);
 			#else
-				Transmittance = AVSM_Sample(Face, ShadowMapCoord, ShadowZ);
+				Transmittance = AVSM_Sample(ShadowMap, Face, ShadowMapCoord, ShadowZ);
 			#endif
 			}
 		}
@@ -360,3 +438,56 @@ float AVSM_SampleTransmittance(float3 TranslatedWorldPosition, float3 LightTrans
 
 	return Transmittance;
 }
+
+FAdaptiveVolumetricShadowMap GetAdaptiveVolumetricShadowMap()
+{
+	FAdaptiveVolumetricShadowMap ShadowMap;
+	for (int i = 0; i < 6; ++i)
+	{
+		ShadowMap.TranslatedWorldToShadow[i] = AVSM.TranslatedWorldToShadow[i];
+	}
+	ShadowMap.TranslatedWorldOrigin = AVSM.TranslatedWorldOrigin;
+	ShadowMap.TranslatedWorldPlane = AVSM.TranslatedWorldPlane;
+
+	ShadowMap.Resolution = AVSM.Resolution;
+	ShadowMap.NumShadowMatrices = AVSM.NumShadowMatrices;
+	ShadowMap.MaxSampleCount = AVSM.MaxSampleCount;
+	ShadowMap.bIsEmpty = AVSM.bIsEmpty;
+	ShadowMap.bIsDirectionalLight = AVSM.bIsDirectionalLight;
+
+	ShadowMap.IndirectionBuffer = AVSM.IndirectionBuffer;
+	ShadowMap.SampleBuffer = AVSM.SampleBuffer;
+	return ShadowMap;
+}
+
+FAdaptiveVolumetricShadowMap GetAdaptiveVolumetricShadowMapForCamera()
+{
+	FAdaptiveVolumetricShadowMap ShadowMap;
+#if USE_CAMERA_AVSM
+	for (int i = 0; i < 6; ++i)
+	{
+		ShadowMap.TranslatedWorldToShadow[i] = CameraAVSM.TranslatedWorldToShadow[i];
+	}
+	ShadowMap.TranslatedWorldOrigin = CameraAVSM.TranslatedWorldOrigin;
+	ShadowMap.TranslatedWorldPlane = CameraAVSM.TranslatedWorldPlane;
+
+	ShadowMap.Resolution = CameraAVSM.Resolution;
+	ShadowMap.NumShadowMatrices = CameraAVSM.NumShadowMatrices;
+	ShadowMap.MaxSampleCount = CameraAVSM.MaxSampleCount;
+	ShadowMap.bIsEmpty = CameraAVSM.bIsEmpty;
+	ShadowMap.bIsDirectionalLight = CameraAVSM.bIsDirectionalLight;
+
+	ShadowMap.IndirectionBuffer = CameraAVSM.IndirectionBuffer;
+	ShadowMap.SampleBuffer = CameraAVSM.SampleBuffer;
+#endif // USE_CAMERA_AVSM
+	return ShadowMap;
+}
+
+float AVSM_SampleTransmittance(
+	float3 TranslatedWorldPosition,
+	float3 LightTranslatedWorldPosition
+)
+{
+	FAdaptiveVolumetricShadowMap AdaptiveVolumetricShadowMap = GetAdaptiveVolumetricShadowMap();
+	return AVSM_SampleTransmittance(AdaptiveVolumetricShadowMap, TranslatedWorldPosition, LightTranslatedWorldPosition);
+}
diff --git a/Engine/Shaders/Private/HeterogeneousVolumes/HeterogeneousVolumesAdaptiveVolumetricShadowMapUtils.ush b/Engine/Shaders/Private/HeterogeneousVolumes/HeterogeneousVolumesAdaptiveVolumetricShadowMapUtils.ush
--- a/Engine/Shaders/Private/HeterogeneousVolumes/HeterogeneousVolumesAdaptiveVolumetricShadowMapUtils.ush
+++ b/Engine/Shaders/Private/HeterogeneousVolumes/HeterogeneousVolumesAdaptiveVolumetricShadowMapUtils.ush
@@ -113,7 +113,7 @@ uint2 AVSM_PackIndirectionData(FAVSMIndirectionData IndirectionData)
 	return PackedData;
 }
 
-FAVSMIndirectionData ASVM_UnpackIndirectionData(uint2 PackedData)
+FAVSMIndirectionData AVSM_UnpackIndirectionData(uint2 PackedData)
 {
 	FAVSMIndirectionData IndirectionData;
 	IndirectionData.PixelOffset = (PackedData.x >> 6) & AVSM_INVALID_PIXEL_OFFSET;
diff --git a/Engine/Shaders/Private/HeterogeneousVolumes/HeterogeneousVolumesLiveShadingGlobalPipeline.usf b/Engine/Shaders/Private/HeterogeneousVolumes/HeterogeneousVolumesLiveShadingGlobalPipeline.usf
--- a/Engine/Shaders/Private/HeterogeneousVolumes/HeterogeneousVolumesLiveShadingGlobalPipeline.usf
+++ b/Engine/Shaders/Private/HeterogeneousVolumes/HeterogeneousVolumesLiveShadingGlobalPipeline.usf
@@ -29,9 +29,9 @@ RWStructuredBuffer<FScreenTile> RWScreenTileBuffer;
 
 groupshared uint GSNumRayHits;
 
-float2 GetDispatchSize()
+int2 GetDispatchSize()
 {
-	return View.ViewSizeAndInvSize.xy / DownsampleFactor;
+	return ceil(View.ViewSizeAndInvSize.xy / float(DownsampleFactor));
 }
 
 [numthreads(THREADGROUP_SIZE_2D, THREADGROUP_SIZE_2D, 1)]
diff --git a/Engine/Shaders/Private/HeterogeneousVolumes/HeterogeneousVolumesLiveShadingPipeline.usf b/Engine/Shaders/Private/HeterogeneousVolumes/HeterogeneousVolumesLiveShadingPipeline.usf
--- a/Engine/Shaders/Private/HeterogeneousVolumes/HeterogeneousVolumesLiveShadingPipeline.usf
+++ b/Engine/Shaders/Private/HeterogeneousVolumes/HeterogeneousVolumesLiveShadingPipeline.usf
@@ -11,6 +11,25 @@
 #define SHADOW_QUALITY 2
 #define NO_TRANSLUCENCY_AVAILABLE
 
+#ifndef SUPPORT_OVERLAPPING_VOLUMES
+#define SUPPORT_OVERLAPPING_VOLUMES 0
+#endif //
+
+#ifndef ADAPTIVE_MARCH
+#define ADAPTIVE_MARCH 0
+#endif //
+
+#ifndef REFERENCE_FAST_PATH
+#define REFERENCE_FAST_PATH 0
+#endif //
+
+#define USE_CAMERA_AVSM (DIM_USE_ADAPTIVE_VOLUMETRIC_SHADOW_MAP && (SUPPORT_OVERLAPPING_VOLUMES != 0))
+
+#if SUPPORT_OVERLAPPING_VOLUMES
+#define AVSM AVSMs.AVSM
+#define CameraAVSM AVSMs.CameraAVSM
+#endif // SUPPORT_OVERLAPPING_VOLUMES
+
 #include "/Engine/Generated/Material.ush"
 #include "HeterogeneousVolumesLiveShadingUtils.ush"
 #include "HeterogeneousVolumesTracingUtils.ush"
@@ -299,11 +318,109 @@ void RenderLightingCacheWithLiveShadingCS(
 #endif // DIM_LIGHTING_CACHE_MODE
 }
 
-float2 GetDispatchSize()
+int2 GetDispatchSize()
 {
-	return View.ViewSizeAndInvSize.xy / DownsampleFactor;
+	return ceil(View.ViewSizeAndInvSize.xy / (float)DownsampleFactor);
 }
 
+#if USE_CAMERA_AVSM
+
+void RayMarchSingleScatteringAdaptive(
+	inout FRayMarchingContext RayMarchingContext,
+	inout FAVSM_Sampler4 TransmittanceSampler,
+	uint StepCount,
+	inout float3 Radiance,
+	inout float3 Transmittance
+)
+{
+	FAVSMSampleData SampleData = AVSM_Sampler4_GetSampleData0(TransmittanceSampler, 0);
+	// Clip against first stored sample
+	float WorldRayTMin = max(RayMarchingContext.LocalRayTMin * RayMarchingContext.LocalToWorldScale, SampleData.X);
+	SampleData = AVSM_CreateSampleData(WorldRayTMin, AVSM_Sampler4_Eval(TransmittanceSampler, WorldRayTMin));
+
+	float WorldShadowBias = CalcShadowBias();
+	FAVSMSampleData PrevSampleData = AVSM_CreateSampleData(SampleData.X - WorldShadowBias, SampleData.Tau);
+
+	FAVSMSampleData LastSampleData = AVSM_Sampler4_Last(TransmittanceSampler);
+	// TODO: Sentinel values prevent efficient clipping against last stored sample
+	//float WorldRayTMax = min(RayMarchingContext.LocalRayTMax * RayMarchingContext.LocalToWorldScale, LastSampleData.X);
+	//LastSampleData = AVSM_CreateSampleData(WorldRayTMax, AVSM_Sampler4_Eval(TransmittanceSampler, WorldRayTMax));
+
+	// Clip step count to stored transmittance range
+	StepCount *= (SampleData.Tau - LastSampleData.Tau);
+	for (uint StepIndex = 0; StepIndex < StepCount; ++StepIndex)
+	{
+		float StepSize = (SampleData.X - PrevSampleData.X) / RayMarchingContext.LocalToWorldScale;
+		float WorldHitT = SampleData.X;
+		float LocalHitT = WorldHitT / RayMarchingContext.LocalToWorldScale;
+
+		float3 LocalPosition = RayMarchingContext.LocalRayOrigin + RayMarchingContext.LocalRayDirection * LocalHitT;
+		float3 WorldPosition = RayMarchingContext.WorldRayOrigin + RayMarchingContext.WorldRayDirection * WorldHitT;
+		FVolumeSampleContext SampleContext = CreateVolumeSampleContext(LocalPosition, WorldPosition, RayMarchingContext.MipLevel);
+
+		// Evaluate transmittance
+		Transmittance = PrevSampleData.Tau;
+
+		// Radiance from emission
+		if (RayMarchingContext.bApplyEmissionAndTransmittance)
+		{
+			float3 Emission = SampleEmission(SampleContext);
+			Radiance += Emission * StepSize * Transmittance;
+		}
+
+		// Radiance from in-scattering
+		float3 Extinction = SampleExtinction(SampleContext);
+		if (RayMarchingContext.bApplyDirectLighting && (any(Extinction > 0)))
+		{
+			float3 RadianceSample = 0;
+
+			// Evaluate in-scattering
+			float3 Albedo = SampleAlbedo(SampleContext);
+			if (any(Albedo > 0.0))
+			{
+				// Evaluate lighting cache
+				float3 LocalBoundsMin = GetLocalBoundsOrigin() - GetLocalBoundsExtent();
+				float3 UVW = saturate((LocalPosition - LocalBoundsMin) / (2.0 * GetLocalBoundsExtent()));
+				float3 Inscattering = SampleLightingCache(UVW, 0) * View.OneOverPreExposure;
+
+				// Evaluate scattering
+				float3 ScatteringCoefficient = Albedo * Extinction;
+				float IsotropicPhase = 1.0 / (4.0 * PI);
+				RadianceSample = Inscattering * ScatteringCoefficient * IsotropicPhase * Transmittance * StepSize;
+			}
+
+			Radiance += RadianceSample;
+		}
+
+		float Epsilon = 1.0e-7;
+		if (all(Transmittance < Epsilon))
+		{
+			Transmittance = 0.0;
+			break;
+		}
+
+		PrevSampleData = SampleData;
+		float DeltaTau = rcp(StepCount);
+		SampleData = AVSM_Sampler4_EvalInverse(TransmittanceSampler, SampleData.Tau - DeltaTau);
+
+		// Guard against steps finer than shadowing bias
+		if (SampleData.X - PrevSampleData.X < WorldShadowBias)
+		{
+			SampleData.X = PrevSampleData.X + WorldShadowBias;
+			SampleData.Tau = AVSM_Sampler4_Eval(TransmittanceSampler, SampleData.X);
+		}
+
+		// Guard against adaptively stepping beyond t-max
+		float WorldRayTMax = RayMarchingContext.LocalRayTMax * RayMarchingContext.LocalToWorldScale;
+		if (SampleData.X > WorldRayTMax)
+		{
+			break;
+		}
+	}
+}
+
+#endif // USE_CAMERA_AVSM
+
 void RenderSingleScatteringWithLiveShading(
 	uint2 DispatchThreadId
 )
@@ -320,8 +437,8 @@ void RenderSingleScatteringWithLiveShading(
 	uint2 ViewPixelCoord = PixelCoord * DownsampleFactor + View.ViewRectMin.xy;
 
 	float3 Radiance = 0.0;
-	float3 Transmittance = RWLightingTexture[PixelCoord].a;
-	float PrevOpacity = Luminance(1.0 - Transmittance);
+	float3 PrevTransmittance = RWLightingTexture[PixelCoord].a;
+	float PrevOpacity = Luminance(1.0 - PrevTransmittance);
 
 	// Extract depth
 	float DeviceZ = SceneDepthTexture.Load(int3(ViewPixelCoord, 0)).r;
@@ -380,8 +497,53 @@ void RenderSingleScatteringWithLiveShading(
 		uint StepCount = CalcStepCount(RayMarchingContext);
 
 		float ScatterHitT = HitT.y;
+
+		float3 Transmittance = PrevTransmittance;
+
+#if USE_CAMERA_AVSM
+		FAdaptiveVolumetricShadowMap GlobalTransmittanceMap = GetAdaptiveVolumetricShadowMapForCamera();
+		if (GlobalTransmittanceMap.bIsEmpty)
+		{
+			return;
+		}
+
+		// Define AVSM Sampler
+		float WorldRayTMin = max(RayMarchingContext.LocalRayTMin * RayMarchingContext.LocalToWorldScale, 0.01);
+		float3 WorldPosition = RayMarchingContext.WorldRayOrigin + RayMarchingContext.WorldRayDirection * WorldRayTMin;
+		float3 TranslatedWorldPosition = DFHackToFloat(DFFastAdd(WorldPosition, PrimaryView.PreViewTranslation));
+
+		int Face = 0;
+		float2 Pixel = 0;
+		float Depth = 0;
+		if (!AVSM_PixelAndDepth(GlobalTransmittanceMap, Face, TranslatedWorldPosition, Pixel, Depth))
+		{
+			return;
+		}
+
+		float DepthOffset = WorldRayTMin - Depth;
+		FAVSM_Sampler4 TransmittanceSampler = AVSM_Sampler4_Create(GlobalTransmittanceMap, Face, Pixel, DepthOffset);
+
+#if ADAPTIVE_MARCH
+		StepCount = View.GeneralPurposeTweak2;
+		RayMarchSingleScatteringAdaptive(
+			RayMarchingContext,
+			TransmittanceSampler,
+			StepCount,
+			Radiance,
+			Transmittance
+		);
+#elif REFERENCE_FAST_PATH
+		RayMarchSingleScatteringReferenceFastPath(
+			RayMarchingContext,
+			TransmittanceSampler,
+			StepCount,
+			Radiance,
+			Transmittance
+		);
+#else
 		RayMarchSingleScattering(
 			RayMarchingContext,
+			TransmittanceSampler,
 			LightData,
 			LightType,
 			StepCount,
@@ -389,6 +551,21 @@ void RenderSingleScatteringWithLiveShading(
 			Radiance,
 			Transmittance
 		);
+#endif
+		// Replace transmittance at WorldRayTMax with global transmittance at MaxTraceDistance
+		Transmittance = AVSM_Sampler4_Last(TransmittanceSampler).Tau;
+
+#else
+		RayMarchSingleScattering(
+			RayMarchingContext,
+			LightData,
+			LightType,
+			StepCount,
+			ScatterHitT,
+			Radiance,
+			Transmittance
+		);
+#endif
 
 		// Radiance
 		RWLightingTexture[PixelCoord].rgb += Radiance * View.PreExposure;
@@ -431,8 +608,6 @@ struct FScreenTile
 	uint Id;
 };
 
-
-
 StructuredBuffer<FScreenTile> ScreenTileBuffer;
 
 [numthreads(THREADGROUP_SIZE_2D, THREADGROUP_SIZE_2D, 1)]
@@ -446,4 +621,4 @@ void RenderSingleScatteringWithLiveShadingIndirectCS(
 	uint2 DispatchThreadId = GroupId_2D * THREADGROUP_SIZE_2D + GroupThreadId;
 
 	RenderSingleScatteringWithLiveShading(DispatchThreadId);
-}
\ No newline at end of file
+}
diff --git a/Engine/Shaders/Private/HeterogeneousVolumes/HeterogeneousVolumesRayMarchingUtils.ush b/Engine/Shaders/Private/HeterogeneousVolumes/HeterogeneousVolumesRayMarchingUtils.ush
--- a/Engine/Shaders/Private/HeterogeneousVolumes/HeterogeneousVolumesRayMarchingUtils.ush
+++ b/Engine/Shaders/Private/HeterogeneousVolumes/HeterogeneousVolumesRayMarchingUtils.ush
@@ -40,6 +40,7 @@
 
 #if DIM_USE_ADAPTIVE_VOLUMETRIC_SHADOW_MAP
 #include "HeterogeneousVolumesAdaptiveVolumetricShadowMapSampling.ush"
+#include "HeterogeneousVolumesAdaptiveVolumetricShadowMapSampler.ush"
 #endif // DIM_USE_ADAPTIVE_VOLUMETRIC_SHADOW_MAP
 
 #if APPLY_FOG_INSCATTERING
@@ -310,6 +311,9 @@ void RayMarchEmissionAbsorption(
 
 void RayMarchSingleScattering(
 	inout FRayMarchingContext RayMarchingContext,
+#if USE_CAMERA_AVSM
+	inout FAVSM_Sampler4 TransmittanceSampler,
+#endif
 	FDeferredLightData LightData,
 	uint LightType,
 	uint StepCount,
@@ -332,18 +336,19 @@ void RayMarchSingleScattering(
 	// Reuse jitter as random sample
 	float RandomSample = RayMarchingContext.Jitter;
 
-	float LocalHitT = RayMarchingContext.LocalRayTMin + JitterSize;
 	for (uint StepIndex = 0; StepIndex < StepCount; ++StepIndex)
 	{
 		// Correct the final step size
+		float LocalHitT = RayMarchingContext.LocalRayTMin + RayMarchingContext.StepSize * (RayMarchingContext.Jitter + StepIndex);
 		if (StepIndex == StepCount - 1)
 		{
 			float UnjitteredLocalHitT = LocalHitT - JitterSize;
 			StepSize = RayMarchingContext.LocalRayTMax - UnjitteredLocalHitT;
 		}
 
+		float WorldHitT = LocalHitT * RayMarchingContext.LocalToWorldScale;
 		float3 LocalPosition = RayMarchingContext.LocalRayOrigin + RayMarchingContext.LocalRayDirection * LocalHitT;
-		float3 WorldPosition = RayMarchingContext.WorldRayOrigin + RayMarchingContext.WorldRayDirection * LocalHitT * RayMarchingContext.LocalToWorldScale;
+		float3 WorldPosition = RayMarchingContext.WorldRayOrigin + RayMarchingContext.WorldRayDirection * WorldHitT;
 		FVolumeSampleContext SampleContext = CreateVolumeSampleContext(LocalPosition, WorldPosition, RayMarchingContext.MipLevel);
 
 		if (RayMarchingContext.bApplyEmissionAndTransmittance)
@@ -426,8 +431,12 @@ void RayMarchSingleScattering(
 			Radiance += RadianceSample;
 		}
 
+#if USE_CAMERA_AVSM
+		Transmittance = AVSM_Sampler4_Eval(TransmittanceSampler, WorldHitT);
+#else
 		// Accumulate transmittance for the next evaluation
 		Transmittance *= exp(-Extinction * StepSize);
+#endif
 
 		float Epsilon = 1.0e-7;
 		if (all(Transmittance < Epsilon))
@@ -479,3 +488,69 @@ void RayMarchSingleScattering(
 	Radiance += FogRadiance.rgb * FogAlpha;
 #endif // APPLY_FOG_INSCATTERING
 }
+
+void RayMarchSingleScatteringReferenceFastPath(
+	inout FRayMarchingContext RayMarchingContext,
+#if USE_CAMERA_AVSM
+	inout FAVSM_Sampler4 TransmittanceSampler,
+#endif
+	uint StepCount,
+	inout float3 Radiance,
+	inout float3 Transmittance
+)
+{
+	for (uint StepIndex = 0; StepIndex < StepCount; ++StepIndex)
+	{
+		float StepSize = RayMarchingContext.StepSize;
+		float LocalHitT = RayMarchingContext.LocalRayTMin + RayMarchingContext.StepSize * (RayMarchingContext.Jitter + StepIndex);
+		float WorldHitT = LocalHitT * RayMarchingContext.LocalToWorldScale;
+		float3 LocalPosition = RayMarchingContext.LocalRayOrigin + RayMarchingContext.LocalRayDirection * LocalHitT;
+		float3 WorldPosition = RayMarchingContext.WorldRayOrigin + RayMarchingContext.WorldRayDirection * WorldHitT;
+		FVolumeSampleContext SampleContext = CreateVolumeSampleContext(LocalPosition, WorldPosition, RayMarchingContext.MipLevel);
+
+		// Radiance from emission
+		if (RayMarchingContext.bApplyEmissionAndTransmittance)
+		{
+			float3 Emission = SampleEmission(SampleContext);
+			Radiance += Emission * StepSize * Transmittance;
+		}
+
+		// Radiance from in-scattering
+		float3 Extinction = SampleExtinction(SampleContext);
+		if (RayMarchingContext.bApplyDirectLighting && (any(Extinction > 0)))
+		{
+			float3 RadianceSample = 0;
+
+			// Evaluate in-scattering
+			float3 Albedo = SampleAlbedo(SampleContext);
+			if (any(Albedo > 0.0))
+			{
+				// Evaluate lighting cache
+				float3 LocalBoundsMin = GetLocalBoundsOrigin() - GetLocalBoundsExtent();
+				float3 UVW = saturate((LocalPosition - LocalBoundsMin) / (2.0 * GetLocalBoundsExtent()));
+				float3 Inscattering = SampleLightingCache(UVW, 0) * View.OneOverPreExposure;
+
+				// Evaluate scattering
+				float3 ScatteringCoefficient = Albedo * Extinction;
+				float IsotropicPhase = 1.0 / (4.0 * PI);
+				RadianceSample = Inscattering * ScatteringCoefficient * IsotropicPhase * Transmittance * StepSize;
+			}
+
+			Radiance += RadianceSample;
+		}
+
+		// Accumulate transmittance for the next evaluation
+#if USE_CAMERA_AVSM
+		Transmittance = AVSM_Sampler4_Eval(TransmittanceSampler, WorldHitT);
+#else
+		Transmittance *= exp(-Extinction * StepSize);
+#endif
+
+		float Epsilon = 1.0e-7;
+		if (all(Transmittance < Epsilon))
+		{
+			Transmittance = 0.0;
+			break;
+		}
+	}
+}
diff --git a/Engine/Source/Runtime/Renderer/Private/HeterogeneousVolumes/HeterogeneousVolumes.h b/Engine/Source/Runtime/Renderer/Private/HeterogeneousVolumes/HeterogeneousVolumes.h
--- a/Engine/Source/Runtime/Renderer/Private/HeterogeneousVolumes/HeterogeneousVolumes.h
+++ b/Engine/Source/Runtime/Renderer/Private/HeterogeneousVolumes/HeterogeneousVolumes.h
@@ -107,6 +107,7 @@ namespace HeterogeneousVolumes
 	bool UseAdaptiveVolumetricShadowMapForSelfShadowing(const FPrimitiveSceneProxy* PrimitiveSceneProxy);
 	bool ShouldApplyHeightFog();
 	bool ShouldApplyVolumetricFog();
+	bool SupportsOverlappingVolumes();
 
 	enum class EFogMode
 	{
@@ -396,6 +397,27 @@ BEGIN_UNIFORM_BUFFER_STRUCT(FAdaptiveVolumetricShadowMapUniformBufferParameters,
 	SHADER_PARAMETER_RDG_BUFFER_SRV(StructuredBuffer<uint>, SampleBuffer)
 END_UNIFORM_BUFFER_STRUCT()
 
+BEGIN_SHADER_PARAMETER_STRUCT(FAdaptiveVolumetricShadowMapParameters, RENDERER_API)
+	SHADER_PARAMETER_ARRAY(FMatrix44f, TranslatedWorldToShadow, [6])
+	SHADER_PARAMETER(FVector3f, TranslatedWorldOrigin)
+	SHADER_PARAMETER(FVector4f, TranslatedWorldPlane)
+
+	SHADER_PARAMETER(FIntPoint, Resolution)
+	SHADER_PARAMETER(int32, NumShadowMatrices)
+	SHADER_PARAMETER(int32, MaxSampleCount)
+	SHADER_PARAMETER(int32, bIsEmpty)
+	SHADER_PARAMETER(int32, bIsDirectionalLight)
+
+	SHADER_PARAMETER_RDG_BUFFER_SRV(StructuredBuffer<uint2>, LinkedListBuffer)
+	SHADER_PARAMETER_RDG_BUFFER_SRV(StructuredBuffer<uint2>, IndirectionBuffer)
+	SHADER_PARAMETER_RDG_BUFFER_SRV(StructuredBuffer<uint>, SampleBuffer)
+END_SHADER_PARAMETER_STRUCT()
+
+BEGIN_UNIFORM_BUFFER_STRUCT(FAdaptiveVolumetricShadowMaps, )
+	SHADER_PARAMETER_STRUCT(FAdaptiveVolumetricShadowMapParameters, AVSM)
+	SHADER_PARAMETER_STRUCT(FAdaptiveVolumetricShadowMapParameters, CameraAVSM)
+END_UNIFORM_BUFFER_STRUCT()
+
 namespace HeterogeneousVolumes {
 	struct FAdaptiveVolumetricShadowMapParameterCache
 	{
@@ -465,6 +487,12 @@ namespace HeterogeneousVolumes {
 
 } // namespace HeterogeneousVolumes
 
+TRDGUniformBufferRef<FAdaptiveVolumetricShadowMaps> CreateAdaptiveVolumetricShadowMapUniformBuffers(
+	FRDGBuilder& GraphBuilder,
+	FSceneViewState* ViewState,
+	const FLightSceneInfo* LightSceneInfo
+);
+
 void RenderWithLiveShading(
 	FRDGBuilder& GraphBuilder,
 	const FSceneTextures& SceneTextures,
diff --git a/Engine/Source/Runtime/Renderer/Private/HeterogeneousVolumes/HeterogeneousVolumesLiveShadingPipeline.cpp b/Engine/Source/Runtime/Renderer/Private/HeterogeneousVolumes/HeterogeneousVolumesLiveShadingPipeline.cpp
--- a/Engine/Source/Runtime/Renderer/Private/HeterogeneousVolumes/HeterogeneousVolumesLiveShadingPipeline.cpp
+++ b/Engine/Source/Runtime/Renderer/Private/HeterogeneousVolumes/HeterogeneousVolumesLiveShadingPipeline.cpp
@@ -31,6 +31,29 @@ static TAutoConsoleVariable<int32> CVarHeterogeneousLightingLiveShadingScreenTil
 	ECVF_RenderThreadSafe
 );
 
+static TAutoConsoleVariable<int32> CVarHeterogeneousVolumesSupportOverlappingVolumes(
+	TEXT("r.HeterogeneousVolumes.SupportOverlappingVolumes"),
+	1,
+	TEXT("Enables support for overlapping volumes (Default = 1)"),
+	ECVF_RenderThreadSafe
+);
+
+#if 0
+static TAutoConsoleVariable<int32> CVarHeterogeneousVolumesBilinearInterpolation(
+	TEXT("r.HeterogeneousVolumes.BilinearInterpolation"),
+	1,
+	TEXT("Enables bilinear interpolation when querying AVSM (Default = 1)"),
+	ECVF_RenderThreadSafe
+);
+
+static TAutoConsoleVariable<int32> CVarHeterogeneousVolumesAdaptiveMarching(
+	TEXT("r.HeterogeneousVolumes.AdaptiveMarching"),
+	0,
+	TEXT("Enables adaptive marching (Default = 0)"),
+	ECVF_RenderThreadSafe
+);
+#endif
+
 namespace HeterogeneousVolumes
 {
 	bool ShouldBoundsCull()
@@ -42,6 +65,22 @@ namespace HeterogeneousVolumes
 	{
 		return CVarHeterogeneousLightingLiveShadingScreenTileClassification.GetValueOnRenderThread() != 0;
 	}
+
+	bool SupportsOverlappingVolumes()
+	{
+		return ShouldCompositeHeterogeneousVolumesWithTranslucency() &&
+			(CVarHeterogeneousVolumesSupportOverlappingVolumes.GetValueOnRenderThread() != 0);
+	}
+#if 0
+	bool UseBilinearInterpolation()
+	{
+		return CVarHeterogeneousVolumesBilinearInterpolation.GetValueOnRenderThread() != 0;
+	}
+	bool ShouldAdaptiveMarch()
+	{
+		return CVarHeterogeneousVolumesAdaptiveMarching.GetValueOnRenderThread() != 0;
+	}
+#endif
 }
 
 //-OPT: Remove duplicate bindings
@@ -266,8 +305,22 @@ class FRenderSingleScatteringWithLiveShadingCS : public FMeshMaterialShader
 	class FWriteVelocity : SHADER_PERMUTATION_BOOL("DIM_WRITE_VELOCITY");
 	class FUseAdaptiveVolumetricShadowMap : SHADER_PERMUTATION_BOOL("DIM_USE_ADAPTIVE_VOLUMETRIC_SHADOW_MAP");
 	class FAVSMSampleMode : SHADER_PERMUTATION_INT("AVSM_SAMPLE_MODE", 2);
+	class FSupportOverlappingVolumes : SHADER_PERMUTATION_BOOL("SUPPORT_OVERLAPPING_VOLUMES");
+	//class FAdaptiveMarch : SHADER_PERMUTATION_INT("ADAPTIVE_MARCH", 2);
+	//class FBilinearInterpolation : SHADER_PERMUTATION_INT("AVSM_BILINEAR_INTERPOLATION", 2);
 	class FApplyFogInscattering : SHADER_PERMUTATION_INT("APPLY_FOG_INSCATTERING", 3);
-	using FPermutationDomain = TShaderPermutationDomain<FUseTransmittanceVolume, FUseInscatteringVolume, FUseLumenGI, FWriteVelocity, FUseAdaptiveVolumetricShadowMap, FApplyFogInscattering, FAVSMSampleMode>;
+	using FPermutationDomain = TShaderPermutationDomain<
+		FUseTransmittanceVolume, 
+		FUseInscatteringVolume, 
+		FUseLumenGI, 
+		FWriteVelocity, 
+		FUseAdaptiveVolumetricShadowMap, 
+		FApplyFogInscattering, 
+		FAVSMSampleMode, 
+		FSupportOverlappingVolumes
+		//FBilinearInterpolation,
+		//FAdaptiveMarch
+	>;
 
 	BEGIN_SHADER_PARAMETER_STRUCT(FParameters, )
 		// Scene data
@@ -291,6 +344,8 @@ class FRenderSingleScatteringWithLiveShadingCS : public FMeshMaterialShader
 		SHADER_PARAMETER_STRUCT_INCLUDE(FVolumeShadowingShaderParameters, VolumeShadowingShaderParameters)
 		SHADER_PARAMETER_STRUCT_INCLUDE(FVirtualShadowMapSamplingParameters, VirtualShadowMapSamplingParameters)
 		SHADER_PARAMETER(int32, VirtualShadowMapId)
+
+		SHADER_PARAMETER_RDG_UNIFORM_BUFFER(FAdaptiveVolumetricShadowMaps, AVSMs)
 		SHADER_PARAMETER_RDG_UNIFORM_BUFFER(FAdaptiveVolumetricShadowMapUniformBufferParameters, AVSM)
 
 		// Atmosphere
@@ -720,6 +775,54 @@ static void RenderLightingCacheWithLiveShading(
 	}
 }
 
+IMPLEMENT_UNIFORM_BUFFER_STRUCT(FAdaptiveVolumetricShadowMaps, "AVSMs");
+
+FAdaptiveVolumetricShadowMapParameters GetAdaptiveVolumetricShadowMapParametersFromUniformBuffer(
+	const TRDGUniformBufferRef<FAdaptiveVolumetricShadowMapUniformBufferParameters>& UniformBuffer
+)
+{
+	FAdaptiveVolumetricShadowMapParameters Parameters;
+	{
+		Parameters.NumShadowMatrices = UniformBuffer->GetParameters()->NumShadowMatrices;
+		for (int32 i = 0; i < Parameters.NumShadowMatrices; ++i)
+		{
+			Parameters.TranslatedWorldToShadow[i] = UniformBuffer->GetParameters()->TranslatedWorldToShadow[i];
+		}
+		Parameters.TranslatedWorldOrigin = UniformBuffer->GetParameters()->TranslatedWorldOrigin;
+		Parameters.TranslatedWorldPlane = UniformBuffer->GetParameters()->TranslatedWorldPlane;
+		Parameters.Resolution = UniformBuffer->GetParameters()->Resolution;
+		Parameters.MaxSampleCount = UniformBuffer->GetParameters()->MaxSampleCount;
+		Parameters.bIsEmpty = UniformBuffer->GetParameters()->bIsEmpty;
+		Parameters.bIsDirectionalLight = UniformBuffer->GetParameters()->bIsDirectionalLight;
+		Parameters.LinkedListBuffer = UniformBuffer->GetParameters()->LinkedListBuffer;
+		Parameters.IndirectionBuffer = UniformBuffer->GetParameters()->IndirectionBuffer;
+		Parameters.SampleBuffer = UniformBuffer->GetParameters()->SampleBuffer;
+	}
+
+	return Parameters;
+}
+
+TRDGUniformBufferRef<FAdaptiveVolumetricShadowMaps> CreateAdaptiveVolumetricShadowMapUniformBuffers(
+	FRDGBuilder& GraphBuilder,
+	FSceneViewState* ViewState,
+	const FLightSceneInfo* LightSceneInfo
+)
+{
+	FAdaptiveVolumetricShadowMaps* UniformBufferParameters = GraphBuilder.AllocParameters<FAdaptiveVolumetricShadowMaps>();
+	{
+		UniformBufferParameters->AVSM = GetAdaptiveVolumetricShadowMapParametersFromUniformBuffer(HeterogeneousVolumes::GetAdaptiveVolumetricShadowMapUniformBuffer(GraphBuilder, ViewState, LightSceneInfo));
+
+		TRDGUniformBufferRef<FAdaptiveVolumetricShadowMapUniformBufferParameters> CameraAVSM = ViewState->AdaptiveVolumetricCameraMapUniformBuffer;
+		if (!HeterogeneousVolumes::SupportsOverlappingVolumes())
+		{
+			CameraAVSM = HeterogeneousVolumes::CreateEmptyAdaptiveVolumetricShadowMapUniformBuffer(GraphBuilder);
+		}
+		UniformBufferParameters->CameraAVSM = GetAdaptiveVolumetricShadowMapParametersFromUniformBuffer(CameraAVSM);
+	}
+
+	return GraphBuilder.CreateUniformBuffer(UniformBufferParameters);
+}
+
 class FScreenTileClassificationCS : public FGlobalShader
 {
 	DECLARE_GLOBAL_SHADER(FScreenTileClassificationCS);
@@ -784,14 +887,14 @@ static void ScreenTileClassification(
 	// Object data
 	const IHeterogeneousVolumeInterface* HeterogeneousVolumeInterface,
 	// Output
-	FRDGBufferRef& NumScreenTilesBuffer,
+	FRDGBufferRef& ScreenTileIndirectArgsBuffer,
 	FRDGBufferRef& ScreenTileBuffer
 )
 {
 	FIntVector GroupCount = FComputeShaderUtils::GetGroupCount(HeterogeneousVolumes::GetScaledViewRect(View.ViewRect), FRenderSingleScatteringWithLiveShadingIndirectCS::GetThreadGroupSize2D());
 	int32 NumTiles = GroupCount.X * GroupCount.Y;
 
-	NumScreenTilesBuffer = GraphBuilder.CreateBuffer(
+	FRDGBufferRef NumScreenTilesBuffer = GraphBuilder.CreateBuffer(
 		FRDGBufferDesc::CreateIndirectDesc<FRHIDispatchIndirectParameters>(1),
 		TEXT("HeterogeneousVolume.NumScreenTilesBuffer")
 	);
@@ -837,11 +940,13 @@ static void ScreenTileClassification(
 	TShaderRef<FScreenTileClassificationCS> ComputeShader = View.ShaderMap->GetShader<FScreenTileClassificationCS>(PermutationVector);
 	FComputeShaderUtils::AddPass(
 		GraphBuilder,
-		RDG_EVENT_NAME("FScreenTileClassificationCS"),
+		RDG_EVENT_NAME("ScreenTileClassificationCS"),
 		ComputeShader,
 		PassParameters,
 		GroupCount
 	);
+
+	ScreenTileIndirectArgsBuffer = NumScreenTilesBuffer;
 }
 
 template <HeterogeneousVolumes::EDispatchMode DispatchMode>
@@ -978,7 +1083,14 @@ void RenderSingleScatteringWithLiveShading(
 			SetVolumeShadowingDefaultShaderParametersGlobal(GraphBuilder, PassParameters->VolumeShadowingShaderParameters);
 		}
 		PassParameters->VirtualShadowMapSamplingParameters = VirtualShadowMapArray.GetSamplingParameters(GraphBuilder);
-		PassParameters->AVSM = HeterogeneousVolumes::GetAdaptiveVolumetricShadowMapUniformBuffer(GraphBuilder, View.ViewState, LightSceneInfo);
+		if (HeterogeneousVolumes::SupportsOverlappingVolumes())
+		{
+			PassParameters->AVSMs = CreateAdaptiveVolumetricShadowMapUniformBuffers(GraphBuilder, View.ViewState, LightSceneInfo);
+		}
+		else
+		{
+			PassParameters->AVSM = HeterogeneousVolumes::GetAdaptiveVolumetricShadowMapUniformBuffer(GraphBuilder, View.ViewState, LightSceneInfo);
+		}
 
 		TRDGUniformBufferRef<FFogUniformParameters> FogBuffer = CreateFogUniformBuffer(GraphBuilder, View);
 		PassParameters->FogStruct = FogBuffer;
@@ -1058,6 +1170,9 @@ void RenderSingleScatteringWithLiveShading(
 	PermutationVector.template Set<typename FRenderSingleScatteringWithLiveShadingDispatchTypeCS::FWriteVelocity>(bWriteVelocity);
 	PermutationVector.template Set<typename FRenderSingleScatteringWithLiveShadingDispatchTypeCS::FUseAdaptiveVolumetricShadowMap>(bUseAVSM);
 	PermutationVector.template Set<typename FRenderSingleScatteringWithLiveShadingDispatchTypeCS::FAVSMSampleMode>(AVSMSampleMode);
+	PermutationVector.template Set<typename FRenderSingleScatteringWithLiveShadingDispatchTypeCS::FSupportOverlappingVolumes>(HeterogeneousVolumes::SupportsOverlappingVolumes());
+	//PermutationVector.template Set<typename FRenderSingleScatteringWithLiveShadingDispatchTypeCS::FBilinearInterpolation>(static_cast<int32>(HeterogeneousVolumes::UseBilinearInterpolation()));
+	//PermutationVector.template Set<typename FRenderSingleScatteringWithLiveShadingDispatchTypeCS::FAdaptiveMarch>(static_cast<int32>(HeterogeneousVolumes::ShouldAdaptiveMarch()));
 	PermutationVector.template Set<typename FRenderSingleScatteringWithLiveShadingDispatchTypeCS::FApplyFogInscattering>(static_cast<int32>(HeterogeneousVolumes::GetApplyFogInscattering()));
 	PermutationVector = FRenderSingleScatteringWithLiveShadingDispatchTypeCS::RemapPermutation(PermutationVector);
 	TShaderRef<FRenderSingleScatteringWithLiveShadingDispatchTypeCS> ComputeShader = Material.GetShader<FRenderSingleScatteringWithLiveShadingDispatchTypeCS>(&FLocalVertexFactory::StaticType, PermutationVector, false);
@@ -1618,6 +1733,7 @@ void CollectHeterogeneousVolumeMeshBatchesForLight(
 				for (int32 MeshBatchIndex = 0; MeshBatchIndex < MeshBatches.Num(); ++MeshBatchIndex)
 				{
 					const FMeshBatchAndRelevance& MeshBatch = MeshBatches[MeshBatchIndex];
+					check(MeshBatch.PrimitiveSceneProxy);
 					bool bIsShadowCast = MeshBatch.PrimitiveSceneProxy->IsShadowCast(ProjectedShadowInfo->ShadowDepthView);
 
 					// TODO: Is material determiniation too expensive?
@@ -2234,8 +2350,10 @@ void RenderAdaptiveVolumetricCameraMapWithLiveShading(
 		ShadowMapResolution.Y = FMath::Max(ShadowMapResolution.Y / DownsampleFactor, 1);
 
 		// Transform
+		const FMatrix ProjectionMatrix = View.ViewMatrices.GetProjectionMatrix();
+		float FOV = FMath::Atan(1.0f / ProjectionMatrix.M[0][0]);
 		FMatrix ViewToClip = FPerspectiveMatrix(
-			FMath::DegreesToRadians(View.FOV * 0.5),
+			FOV,
 			ShadowMapResolution.X,
 			ShadowMapResolution.Y,
 			1.0,
