Remove unused deep shadow atlas offset/resolution.

This was used for CPU allocation and is no longer used by the GPU allocator.

diff --git a/Engine/Shaders/Private/HairStrands/HairStrandsDebug.usf b/Engine/Shaders/Private/HairStrands/HairStrandsDebug.usf
--- a/Engine/Shaders/Private/HairStrands/HairStrandsDebug.usf
+++ b/Engine/Shaders/Private/HairStrands/HairStrandsDebug.usf
@@ -789,8 +789,6 @@ void VisibilityDebugPPLLCS(uint GroupIndex : SV_GroupIndex, uint3 DispatchThread
 #include "../PostProcessCommon.ush"
 
 float DomScale;
-float2 DomAtlasOffset;
-float2 DomAtlasScale;
 float2 OutputResolution;
 float2 InvOutputResolution;
 
@@ -849,13 +847,13 @@ void VisualizeDomPS(
 	float2 LocalUV = 0;
 	if (IsInsideRect(DepthRect, PixelPos, LocalUV))
 	{
-		const float2 AtlasUV = LocalUV * DomAtlasScale + DomAtlasOffset;
+		const float2 AtlasUV = LocalUV;
 		Color = DeepShadowDepthTexture.SampleLevel(LinearSampler, AtlasUV, 0).xxx;
 		Color += AddBorder(LocalUV, DepthRect);
 	}
 	else if (IsInsideRect(DomRect, PixelPos, LocalUV))
 	{
-		const float2 AtlasUV = LocalUV * DomAtlasScale + DomAtlasOffset;
+		const float2 AtlasUV = LocalUV;
 		Color = DeepShadowLayerTexture.SampleLevel(LinearSampler, AtlasUV, 0).xyz / DomScale;
 		Color += AddBorder(LocalUV, DomRect);
 #if 0 
diff --git a/Engine/Shaders/Private/HairStrands/HairStrandsDeepShadowMask.usf b/Engine/Shaders/Private/HairStrands/HairStrandsDeepShadowMask.usf
--- a/Engine/Shaders/Private/HairStrands/HairStrandsDeepShadowMask.usf
+++ b/Engine/Shaders/Private/HairStrands/HairStrandsDeepShadowMask.usf
@@ -44,8 +44,7 @@ float3 GetTranslatedWorldPosition(float2 UV, float SceneDepth)
 StructuredBuffer<FDeepShadowViewInfo> DeepShadow_ViewInfoBuffer;
 StructuredBuffer<uint> DeepShadow_AtlasSlotIndexBuffer;
 
-int2				DeepShadow_SlotResolution;
-int2				DeepShadow_SlotOffset;
+int2				DeepShadow_AtlasResolution;
 uint				DeepShadow_KernelType;
 float				DeepShadow_DepthBiasScale;
 float				DeepShadow_DensityScale;
@@ -91,6 +90,7 @@ struct FHairSamplingSettings
 	Texture2D<float4>ShadowDomTexture;
 	SamplerState	ShadowDepthTextureSampler;
 	float2			ShadowAtlasInvResolution;
+	float2			ShadowAtlasResolution;
 	float2			ShadowSlotResolution;
 	float2			ShadowSlotOffset;
 	float			SceneDepth;					// SceneDepth in lightspace.
@@ -156,7 +156,7 @@ float InternalPCF3x3(float2 Fraction, float4 Values0, float4 Values1, float4 Val
 
 float InternalManual2x2PCF(float2 ShadowPosition, FHairSamplingSettings Settings)
 {
-	float2 TexelPos = ShadowPosition * Settings.ShadowSlotResolution + Settings.ShadowSlotOffset;
+	float2 TexelPos = ShadowPosition * Settings.ShadowAtlasResolution;
 	float2 Fraction = frac(TexelPos);
 	float2 TexelCenter = floor(TexelPos) + 0.5f;
 	float2 Sample00TexelCenter = TexelCenter - float2(1, 1);
@@ -171,7 +171,7 @@ float InternalManual2x2PCF(float2 ShadowPosition, FHairSamplingSettings Settings
 
 float InternalManual4x4PCF(float2 ShadowPosition, FHairSamplingSettings Settings)
 {
-	float2 TexelPos = ShadowPosition * Settings.ShadowSlotResolution + Settings.ShadowSlotOffset - 0.5f;
+	float2 TexelPos = ShadowPosition * Settings.ShadowAtlasResolution - 0.5f;
 	float2 Fraction = frac(TexelPos);
 	float2 TexelCenter = floor(TexelPos) + 0.5f;	// bias to get reliable texel center content	
 	float2 Sample00TexelCenter = TexelCenter - float2(1, 1);
@@ -219,15 +219,15 @@ float4 InternalGather4(float2 SamplePos, int2 Offset, FHairSamplingSettings Sett
 float InternalManual7x7PCF(float2 ShadowPosition, FHairSamplingSettings Settings)
 {
 #if 1
-	float2 TexelPos = ShadowPosition * Settings.ShadowSlotResolution + Settings.ShadowSlotOffset - 0.5f;	// bias to be consistent with texture filtering hardware
+	float2 TexelPos = ShadowPosition * Settings.ShadowAtlasResolution - 0.5f;	// bias to be consistent with texture filtering hardware
 	float2 Fraction = frac(TexelPos);
 	float2 TexelCenter = floor(TexelPos);
-	float2 SamplePos = (TexelCenter + 0.5f) / Settings.ShadowSlotResolution;
+	float2 SamplePos = (TexelCenter + 0.5f) / Settings.ShadowAtlasResolution;
 #else
-	float2 TexelPos = ShadowPosition * Settings.ShadowSlotResolution + Settings.ShadowSlotOffset;
+	float2 TexelPos = ShadowPosition * Settings.ShadowAtlasResolution;
 	float2 Fraction = frac(TexelPos);
 	float2 TexelCenter = floor(TexelPos) + 0.5f;
-	float2 SamplePos = (TexelCenter - float2(1, 1)) / Settings.ShadowSlotResolution;
+	float2 SamplePos = (TexelCenter - float2(1, 1)) / Settings.ShadowAtlasResolution;
 #endif
 	float Results;
 
@@ -278,7 +278,7 @@ float InternalGaussianFilter(float2 ShadowPosition, FHairSamplingSettings Settin
 			float2(0.14383161, -0.14100790)
 	};
 
-	const float SampleRadius = 3.f / Settings.ShadowSlotResolution.x;
+	const float SampleRadius = 3.f * Settings.ShadowAtlasInvResolution.x;
 	float AccShadow = 0;
 	for (uint SampleIt = 0; SampleIt < SampleCount; ++SampleIt)
 	{
@@ -314,13 +314,13 @@ float InternalGaussianTransmissionFilter(float3 ShadowPosition, FHairSamplingSet
 			float2(0.14383161, -0.14100790)
 	};
 
-	const float SampleRadius = 3.f / Settings.ShadowSlotResolution.x;
+	const float SampleRadius = 3.f * Settings.ShadowAtlasInvResolution.x;
 	float AccHairCount = 0;
 	for (uint SampleIt = 0; SampleIt < SampleCount; ++SampleIt)
 	{
 		const float2 Offset = PoissonDisk[SampleIt] * SampleRadius;
 		const float DepthBias = DeepShadow_LayerDepths[0] * DeepShadow_DepthBiasScale;
-		const float3 LightSpacePositionInAtlas = (ShadowPosition+float3(Offset,0)) * float3(DeepShadow_SlotResolution, 1) + float3(DeepShadow_SlotOffset, 0);
+		const float3 LightSpacePositionInAtlas = (ShadowPosition+float3(Offset,0)) * float3(DeepShadow_AtlasResolution, 1);
 		const float HairCount = SampleDOM_PCF2x2(LightSpacePositionInAtlas.xyz, DepthBias, DeepShadow_LayerDepths, Settings.ShadowDepthTexture, Settings.ShadowDomTexture);
 		AccHairCount += saturate(HairCount);
 	}
@@ -369,9 +369,8 @@ void MainPS(
 			ShadowSettings.ShadowDepthTexture = DeepShadow_FrontDepthTexture;
 			ShadowSettings.ShadowDomTexture = DeepShadow_DomTexture;
 			ShadowSettings.ShadowDepthTextureSampler = ShadowSampler;
-			ShadowSettings.ShadowAtlasInvResolution = 1 / float2(ShadowAtlasResolution);
-			ShadowSettings.ShadowSlotResolution = DeepShadow_SlotResolution;
-			ShadowSettings.ShadowSlotOffset = DeepShadow_SlotOffset;
+			ShadowSettings.ShadowAtlasInvResolution = 1 / float2(DeepShadow_AtlasResolution);
+			ShadowSettings.ShadowAtlasResolution = DeepShadow_AtlasResolution;
 			ShadowSettings.SceneDepth = saturate(LightSpacePosition.z);
 			#if PERMUTATION_KERNEL_TYPE == KERNEL_TYPE_4x4
 				Visibility = 1 - InternalManual4x4PCF(LightSpacePosition.xy, ShadowSettings);
diff --git a/Engine/Shaders/Private/HairStrands/HairStrandsDeepTransmittanceMask.usf b/Engine/Shaders/Private/HairStrands/HairStrandsDeepTransmittanceMask.usf
--- a/Engine/Shaders/Private/HairStrands/HairStrandsDeepTransmittanceMask.usf
+++ b/Engine/Shaders/Private/HairStrands/HairStrandsDeepTransmittanceMask.usf
@@ -23,9 +23,9 @@
 // Utils function & structs
 struct FTransmittanceSettings
 {
-	int2	 DeepShadowAtlasResolution;
+	int2	 AtlasResolution;
 	float4x4 TranslatedWorldToLightTransform;
-	float4 AtlasScaleBias;
+	float4   AtlasScaleBias;
 	float3   LightDirection; // Direction from world position to the light
 	float3   TranslatedLightPosition;
 	bool	 bIsDirectional;
@@ -127,7 +127,7 @@ FHairTransmittanceMask ComputeTransmittance(
 			float HairCount = 0;
 			if (Settings.TransmittanceKernelType == 3)
 			{
-				HairCount = SampleDOM_PCSS(LightSpacePosition.xyz, Settings.DeepShadowAtlasResolution, DepthBias, InLayerDepths, FrontDepthTexture, DomTexture, Settings.TransmittanceKernelApertureInDegree, Settings.Random.xy);
+				HairCount = SampleDOM_PCSS(LightSpacePosition.xyz, Settings.AtlasResolution, DepthBias, InLayerDepths, FrontDepthTexture, DomTexture, Settings.TransmittanceKernelApertureInDegree, Settings.Random.xy);
 			}
 			else if (Settings.TransmittanceKernelType == 2)
 			{
@@ -190,7 +190,7 @@ void MainCS(uint2 DispatchThreadId : SV_DispatchThreadID)
 		const float3 TranslatedWorldPosition = GetTranslatedWorldPosition(UV, SceneDepth);
 
 		FTransmittanceSettings Settings;
-		Settings.DeepShadowAtlasResolution = DeepShadow_AtlasResolution;
+		Settings.AtlasResolution = DeepShadow_AtlasResolution;
 		Settings.TranslatedWorldToLightTransform = TranslatedWorldToLightTransform;
 		Settings.AtlasScaleBias = AtlasScaleBias;
 		Settings.bIsDirectional = IsLightDirectional(TranslatedLightPosition_LightDirection);
diff --git a/Engine/Source/Runtime/Renderer/Private/HairStrands/HairStrandsData.h b/Engine/Source/Runtime/Renderer/Private/HairStrands/HairStrandsData.h
--- a/Engine/Source/Runtime/Renderer/Private/HairStrands/HairStrandsData.h
+++ b/Engine/Source/Runtime/Renderer/Private/HairStrands/HairStrandsData.h
@@ -239,7 +239,6 @@ struct FMinHairRadiusAtDepth1
 struct FHairStrandsDeepShadowData
 {
 	FMinHairRadiusAtDepth1 CPU_MinStrandRadiusAtDepth1;
-	FIntRect AtlasRect;
 	uint32 MacroGroupId = ~0;
 	uint32 AtlasSlotIndex = 0;
 
diff --git a/Engine/Source/Runtime/Renderer/Private/HairStrands/HairStrandsDebug.cpp b/Engine/Source/Runtime/Renderer/Private/HairStrands/HairStrandsDebug.cpp
--- a/Engine/Source/Runtime/Renderer/Private/HairStrands/HairStrandsDebug.cpp
+++ b/Engine/Source/Runtime/Renderer/Private/HairStrands/HairStrandsDebug.cpp
@@ -581,8 +581,6 @@ class FDeepShadowVisualizePS : public FGlobalShader
 
 	BEGIN_SHADER_PARAMETER_STRUCT(FParameters, )
 		SHADER_PARAMETER(float, DomScale)
-		SHADER_PARAMETER(FVector2f, DomAtlasOffset)
-		SHADER_PARAMETER(FVector2f, DomAtlasScale)
 		SHADER_PARAMETER(FVector2f, OutputResolution)
 		SHADER_PARAMETER(FVector2f, InvOutputResolution)
 		SHADER_PARAMETER(FIntVector4, HairViewRect)
@@ -618,23 +616,11 @@ static void AddDebugDeepShadowTexturePass(
 {
 	check(OutTarget);
 
-	FIntPoint AtlasResolution(0, 0);
-	FVector2f AltasOffset(0, 0);
-	FVector2f AltasScale(0, 0);
-	if (ShadowData && Resources)
-	{
-		AtlasResolution = FIntPoint(Resources->DepthAtlasTexture->Desc.Extent.X, Resources->DepthAtlasTexture->Desc.Extent.Y);
-		AltasOffset = FVector2f(ShadowData->AtlasRect.Min.X / float(AtlasResolution.X), ShadowData->AtlasRect.Min.Y / float(AtlasResolution.Y));
-		AltasScale = FVector2f((ShadowData->AtlasRect.Max.X - ShadowData->AtlasRect.Min.X) / float(AtlasResolution.X), (ShadowData->AtlasRect.Max.Y - ShadowData->AtlasRect.Min.Y) / float(AtlasResolution.Y));
-	}
-
 	const FIntRect Viewport = View->ViewRect;
 	const FIntPoint Resolution(Viewport.Width(), Viewport.Height());
 
 	FDeepShadowVisualizePS::FParameters* Parameters = GraphBuilder.AllocParameters<FDeepShadowVisualizePS::FParameters>();
 	Parameters->DomScale = GDeepShadowDebugScale;
-	Parameters->DomAtlasOffset = AltasOffset;
-	Parameters->DomAtlasScale = AltasScale;
 	Parameters->OutputResolution = Resolution;
 	Parameters->InvOutputResolution = FVector2f(1.f / Resolution.X, 1.f / Resolution.Y);
 	Parameters->DeepShadowDepthTexture = Resources ? Resources->DepthAtlasTexture : nullptr;
diff --git a/Engine/Source/Runtime/Renderer/Private/HairStrands/HairStrandsDeepShadow.cpp b/Engine/Source/Runtime/Renderer/Private/HairStrands/HairStrandsDeepShadow.cpp
--- a/Engine/Source/Runtime/Renderer/Private/HairStrands/HairStrandsDeepShadow.cpp
+++ b/Engine/Source/Runtime/Renderer/Private/HairStrands/HairStrandsDeepShadow.cpp
@@ -82,15 +82,14 @@ void AddInjectHairVoxelShadowCaster(
 	const FViewInfo& View,
 	const bool bClear,
 	const FHairStrandsDeepShadowData& DomData,
-	FIntRect AtlasRect,
-	uint32 AtlasSlotIndex,
-	FIntPoint AtlasSlotResolution,
 	FHairStrandsVoxelResources& VoxelResources,
 	FRDGBufferSRVRef DeepShadowViewInfoBufferSRV,
 	FRDGTextureRef OutDepthTexture)
 {
+	const FIntPoint AtlasResolution = DomData.AtlasResolution;
+
 	FHairStransShadowDepthInjectionParameters* Parameters = GraphBuilder.AllocParameters<FHairStransShadowDepthInjectionParameters>();
-	Parameters->OutputResolution = AtlasSlotResolution;
+	Parameters->OutputResolution = AtlasResolution;
 	Parameters->ViewUniformBuffer = View.ViewUniformBuffer;
 	Parameters->RenderTargets.DepthStencil = FDepthStencilBinding(OutDepthTexture, bClear ? ERenderTargetLoadAction::EClear : ERenderTargetLoadAction::ELoad, ERenderTargetLoadAction::ENoAction, FExclusiveDepthStencil::DepthWrite_StencilNop);
 	Parameters->VirtualVoxel = VoxelResources.UniformBuffer;
@@ -99,7 +98,7 @@ void AddInjectHairVoxelShadowCaster(
 	Parameters->bIsDirectional = DomData.bIsLightDirectional ? 1 : 0;
 	Parameters->MacroGroupId = DomData.MacroGroupId;
 	Parameters->DeepShadowViewInfoBuffer = DeepShadowViewInfoBufferSRV;
-	Parameters->AtlasSlotIndex = AtlasSlotIndex;
+	Parameters->AtlasSlotIndex = DomData.AtlasSlotIndex;
 
 	TShaderMapRef<FHairStrandsShadowDepthInjectionVS> VertexShader(View.ShaderMap);
 	TShaderMapRef<FHairStrandsShadowDepthInjectionPS> PixelShader(View.ShaderMap);
@@ -112,7 +111,7 @@ void AddInjectHairVoxelShadowCaster(
 		RDG_EVENT_NAME("HairStrandsShadowDepthInjection"),
 		Parameters,
 		ERDGPassFlags::Raster,
-		[ParametersVS, ParametersPS, VertexShader, PixelShader, AtlasRect](FRHICommandList& RHICmdList)
+		[ParametersVS, ParametersPS, VertexShader, PixelShader, AtlasResolution](FRHICommandList& RHICmdList)
 		{
 
 			// Apply additive blending pipeline state.
@@ -131,7 +130,7 @@ void AddInjectHairVoxelShadowCaster(
 			SetShaderParameters(RHICmdList, PixelShader, PixelShader.GetPixelShader(), ParametersPS);
 
 			// Emit an instanced quad draw call on the order of the number of pixels on the screen.	
-			RHICmdList.SetViewport(AtlasRect.Min.X, AtlasRect.Min.Y, 0.0f, AtlasRect.Max.X, AtlasRect.Max.Y, 1.0f);
+			RHICmdList.SetViewport(0, 0, 0.0f, AtlasResolution.X, AtlasResolution.Y, 1.0f);
 			RHICmdList.DrawPrimitive(0, 12, 1);
 		});
 }
@@ -328,7 +327,6 @@ void RenderHairStrandsDeepShadows(
 				DomData.LightId = LightInfo->Id;
 				DomData.AtlasResolution = AtlasResolution;
 				DomData.Bounds = MacroGroupBounds;
-				DomData.AtlasRect = FIntRect(0, 0, AtlasResolution.X, AtlasResolution.Y);
 				DomData.MacroGroupId = MacroGroup.MacroGroupId;
 				DomData.CPU_MinStrandRadiusAtDepth1 = MinStrandRadiusAtDepth1;
 				DomData.AtlasSlotIndex = TotalAtlasSlotIndex;
@@ -424,9 +422,6 @@ void RenderHairStrandsDeepShadows(
 						View,
 						bClear,
 						DomData,
-						DomData.AtlasRect,
-						DomData.AtlasSlotIndex,
-						AtlasSlotResolution,
 						VirtualVoxelResources,
 						DeepShadowViewInfoBufferSRV,
 						FrontDepthAtlasTexture);
@@ -467,7 +462,7 @@ void RenderHairStrandsDeepShadows(
 						&View,
 						MacroGroup.PrimitivesInfos,
 						EHairStrandsRasterPassType::FrontDepth,
-						DomData.AtlasRect,
+						DomData.AtlasResolution,
 						HairRenderInfo,
 						HairRenderInfoBits,
 						DomData.LightDirection,
@@ -502,7 +497,7 @@ void RenderHairStrandsDeepShadows(
 						&View,
 						MacroGroup.PrimitivesInfos,
 						EHairStrandsRasterPassType::DeepOpacityMap,
-						DomData.AtlasRect,
+						DomData.AtlasResolution,
 						HairRenderInfo,
 						HairRenderInfoBits,
 						DomData.LightDirection,
diff --git a/Engine/Source/Runtime/Renderer/Private/HairStrands/HairStrandsRasterCommon.cpp b/Engine/Source/Runtime/Renderer/Private/HairStrands/HairStrandsRasterCommon.cpp
--- a/Engine/Source/Runtime/Renderer/Private/HairStrands/HairStrandsRasterCommon.cpp
+++ b/Engine/Source/Runtime/Renderer/Private/HairStrands/HairStrandsRasterCommon.cpp
@@ -378,7 +378,7 @@ void AddHairDeepShadowRasterPass(
 	const FViewInfo* ViewInfo,
 	const FHairStrandsMacroGroupData::TPrimitiveInfos& PrimitiveSceneInfos,
 	const EHairStrandsRasterPassType PassType,
-	const FIntRect& ViewportRect,
+	const FIntPoint& AtlasResolution,
 	const FVector4f& HairRenderInfo,
 	const uint32 HairRenderInfoBits,
 	const FVector3f& LightDirection,
@@ -393,7 +393,7 @@ void AddHairDeepShadowRasterPass(
 		ViewInfo, 
 		PrimitiveSceneInfos, 
 		PassType, 
-		ViewportRect, 
+		FIntRect(FIntPoint::ZeroValue, AtlasResolution), 
 		HairRenderInfo, 
 		HairRenderInfoBits,
 		LightDirection,
diff --git a/Engine/Source/Runtime/Renderer/Private/HairStrands/HairStrandsRasterCommon.h b/Engine/Source/Runtime/Renderer/Private/HairStrands/HairStrandsRasterCommon.h
--- a/Engine/Source/Runtime/Renderer/Private/HairStrands/HairStrandsRasterCommon.h
+++ b/Engine/Source/Runtime/Renderer/Private/HairStrands/HairStrandsRasterCommon.h
@@ -43,7 +43,7 @@ void AddHairDeepShadowRasterPass(
 	const FViewInfo* ViewInfo,
 	const FHairStrandsMacroGroupData::TPrimitiveInfos& PrimitiveSceneInfos,
 	const EHairStrandsRasterPassType ShadowPassType,
-	const FIntRect& ViewportRect,
+	const FIntPoint& ViewportRect,
 	const FVector4f& HairRenderInfo,
 	const uint32 HairRenderInfoBits,
 	const FVector3f& LightDirection,
diff --git a/Engine/Source/Runtime/Renderer/Private/HairStrands/HairStrandsTransmittance.cpp b/Engine/Source/Runtime/Renderer/Private/HairStrands/HairStrandsTransmittance.cpp
--- a/Engine/Source/Runtime/Renderer/Private/HairStrands/HairStrandsTransmittance.cpp
+++ b/Engine/Source/Runtime/Renderer/Private/HairStrands/HairStrandsTransmittance.cpp
@@ -600,8 +600,7 @@ class FHairStrandsDeepShadowMaskPS : public FGlobalShader
 	BEGIN_SHADER_PARAMETER_STRUCT(FParameters, )
 		SHADER_PARAMETER_STRUCT_INCLUDE(ShaderPrint::FShaderParameters, ShaderPrintParameters)
 		
-		SHADER_PARAMETER(FIntPoint, DeepShadow_SlotOffset)
-		SHADER_PARAMETER(FIntPoint, DeepShadow_SlotResolution)
+		SHADER_PARAMETER(FIntPoint, DeepShadow_AtlasResolution)
 		SHADER_PARAMETER(float, DeepShadow_DepthBiasScale)
 		SHADER_PARAMETER(float, DeepShadow_DensityScale)
 		SHADER_PARAMETER(FVector4f, DeepShadow_LayerDepths)
@@ -642,7 +641,7 @@ struct FHairStrandsDeepShadowParams
 	uint32			OutputChannel = ~0;
 	FRDGBufferSRVRef DeepShadow_ViewInfoBuffer = nullptr;
 	FRDGBufferSRVRef DeepShadow_AtlasSlotIndexBuffer = nullptr;
-	FIntRect		DeepShadow_AtlasRect;
+	FIntPoint		DeepShadow_AtlasResolution;
 	FRDGTextureRef	DeepShadow_FrontDepthTexture = nullptr;
 	FRDGTextureRef	DeepShadow_LayerTexture = nullptr;
 	float			DeepShadow_DepthBiasScale = 1;
@@ -676,10 +675,7 @@ static void AddHairStrandsDeepShadowMaskPass(
 	Parameters->DeepShadow_DensityScale = Params.DeepShadow_DensityScale;
 	Parameters->DeepShadow_LayerDepths = Params.DeepShadow_LayerDepths;
 	Parameters->RenderTargets[0] = FRenderTargetBinding(OutShadowMask, ERenderTargetLoadAction::ELoad);
-
-	Parameters->DeepShadow_SlotOffset = FIntPoint(Params.DeepShadow_AtlasRect.Min.X, Params.DeepShadow_AtlasRect.Min.Y);
-	Parameters->DeepShadow_SlotResolution = FIntPoint(Params.DeepShadow_AtlasRect.Max.X - Params.DeepShadow_AtlasRect.Min.X, Params.DeepShadow_AtlasRect.Max.Y - Params.DeepShadow_AtlasRect.Min.Y);
-
+	Parameters->DeepShadow_AtlasResolution = Params.DeepShadow_AtlasResolution;
 	Parameters->EffectiveAtlasSlotCount = EffectiveAtlasSlotCount;
 	Parameters->DeepShadow_AtlasSlotIndexBuffer = Params.DeepShadow_AtlasSlotIndexBuffer;
 	Parameters->DeepShadow_ViewInfoBuffer = Params.DeepShadow_ViewInfoBuffer;
@@ -967,7 +963,7 @@ static void InternalRenderHairStrandsShadowMask(
 	TArray<uint32> AtlasSlotIndexData;
 	AtlasSlotIndexData.Reserve(InMacroGroupDatas.Num());
 	float LayerDistribution = 0;
-	FIntRect AtlasRect = FIntRect();
+	FIntPoint AtlasResolution = FIntPoint::ZeroValue;
 	if (!IsHairStrandsForVoxelTransmittanceAndShadowEnable())
 	{
 		for (const FHairStrandsMacroGroupData& MacroGroupData : InMacroGroupDatas)
@@ -978,9 +974,9 @@ static void InternalRenderHairStrandsShadowMask(
 				if (DomData.LightId == LightSceneInfo->Id)
 				{
 					bHasDeepShadow = true;
-					AtlasRect = DomData.AtlasRect;
 					LayerDistribution = DomData.LayerDistribution;
 					AtlasSlotIndexData.Add(DomData.AtlasSlotIndex);
+					AtlasResolution = DomData.AtlasResolution;
 				}
 			}
 		}
@@ -994,7 +990,7 @@ static void InternalRenderHairStrandsShadowMask(
 		FHairStrandsDeepShadowParams Params;
 		Params.DeepShadow_AtlasSlotIndexBuffer = GraphBuilder.CreateSRV(AtlasSlotIndexBuffer);
 		Params.DeepShadow_ViewInfoBuffer = GraphBuilder.CreateSRV(DeepShadowResources.DeepShadowViewInfoBuffer);
-		Params.DeepShadow_AtlasRect = AtlasRect;
+		Params.DeepShadow_AtlasResolution = AtlasResolution;
 		Params.DeepShadow_FrontDepthTexture = DeepShadowResources.DepthAtlasTexture;
 		Params.DeepShadow_LayerTexture = DeepShadowResources.LayersAtlasTexture;
 		Params.DeepShadow_DepthBiasScale = GetDeepShadowDepthBiasScale();
