Modify virtual texture sampling shader code to branch on virtual texture page table uniform validity before doing work.
Invalid page table uniforms will skip work and the physical sample will return the virtual texture fallback color.
This is done to handle case where page table uniform comes from GPU Scene descriptor and it's valid that a primitive can use the null descriptor.

diff --git a/Engine/Shaders/Private/VirtualTextureCommon.ush b/Engine/Shaders/Private/VirtualTextureCommon.ush
--- a/Engine/Shaders/Private/VirtualTextureCommon.ush
+++ b/Engine/Shaders/Private/VirtualTextureCommon.ush
@@ -208,6 +208,13 @@ VTPageTableUniform VTPageTableUniform_Unpack(uint2 PackedPageTableUniform)
 	return Result;
 }
 
+bool IsValid(VTPageTableUniform PageTableUniform)
+{
+	// This is expected to happen when we have null MeshPaintTextureDescriptor values.
+	// In that case we want to skip page table sampling and use the FallbackValue.
+	return PageTableUniform.SizeInPages.x > 0;
+}
+
 /** Structure carrying page table read result. Is passed as inout parameter and partially filled by several functions. */
 struct VTPageTableResult
 {
@@ -437,14 +444,17 @@ VTPageTableResult TextureLoadVirtualPageTable(
 	in out FVirtualTextureFeedbackParams Feedback)
 {
 	VTPageTableResult Result = (VTPageTableResult)0;
-	UV = UV * PageTableUniform.UVScale;
-	int vLevel = GetGlobalVirtualTextureMipBias();
+	if (IsValid(PageTableUniform))
+	{
+		UV = UV * PageTableUniform.UVScale;
+		int vLevel = GetGlobalVirtualTextureMipBias();
 #if PIXELSHADER
-	vLevel = TextureComputeVirtualMipLevel(Result, ddx(UV), ddy(UV), MipBias, SvPositionXY, PageTableUniform);
+		vLevel = TextureComputeVirtualMipLevel(Result, ddx(UV), ddy(UV), MipBias, SvPositionXY, PageTableUniform);
 #endif // PIXELSHADER
-	UV = ApplyAddressMode(UV, AddressU, AddressV);
-	TextureLoadVirtualPageTableInternal(Result, PageTable0, PageTableUniform, UV, vLevel);
-	StoreVirtualTextureFeedback(Feedback, SampleIndex, Result.PackedRequest);
+		UV = ApplyAddressMode(UV, AddressU, AddressV);
+		TextureLoadVirtualPageTableInternal(Result, PageTable0, PageTableUniform, UV, vLevel);
+		StoreVirtualTextureFeedback(Feedback, SampleIndex, Result.PackedRequest);
+	}
 	return Result;
 }
 
@@ -456,13 +466,16 @@ VTPageTableResult TextureLoadVirtualPageTable(
 	MaterialFloat MipBias, float2 SvPositionXY)
 {
 	VTPageTableResult Result = (VTPageTableResult)0;
-	UV = UV * PageTableUniform.UVScale;
-	int vLevel = GetGlobalVirtualTextureMipBias();
+	if (IsValid(PageTableUniform))
+	{
+		UV = UV * PageTableUniform.UVScale;
+		int vLevel = GetGlobalVirtualTextureMipBias();
 #if PIXELSHADER
-	vLevel = TextureComputeVirtualMipLevel(Result, ddx(UV), ddy(UV), MipBias, SvPositionXY, PageTableUniform);
+		vLevel = TextureComputeVirtualMipLevel(Result, ddx(UV), ddy(UV), MipBias, SvPositionXY, PageTableUniform);
 #endif // PIXELSHADER
-	UV = ApplyAddressMode(UV, AddressU, AddressV);
-	TextureLoadVirtualPageTableInternal(Result, PageTable0, PageTableUniform, UV, vLevel);
+		UV = ApplyAddressMode(UV, AddressU, AddressV);
+		TextureLoadVirtualPageTableInternal(Result, PageTable0, PageTableUniform, UV, vLevel);
+	}
 	return Result;
 }
 
@@ -476,14 +489,17 @@ VTPageTableResult TextureLoadVirtualPageTable(
 	in out FVirtualTextureFeedbackParams Feedback)
 {
 	VTPageTableResult Result = (VTPageTableResult)0;
-	UV = UV * PageTableUniform.UVScale;
-	int vLevel = GetGlobalVirtualTextureMipBias();
+	if (IsValid(PageTableUniform))
+	{
+		UV = UV * PageTableUniform.UVScale;
+		int vLevel = GetGlobalVirtualTextureMipBias();
 #if PIXELSHADER
-	vLevel = TextureComputeVirtualMipLevel(Result, ddx(UV), ddy(UV), MipBias, SvPositionXY, PageTableUniform);
+		vLevel = TextureComputeVirtualMipLevel(Result, ddx(UV), ddy(UV), MipBias, SvPositionXY, PageTableUniform);
 #endif // PIXELSHADER
-	UV = ApplyAddressMode(UV, AddressU, AddressV);
-	TextureLoadVirtualPageTableInternal(Result, PageTable0, PageTable1, PageTableUniform, UV, vLevel);
-	StoreVirtualTextureFeedback(Feedback, SampleIndex, Result.PackedRequest);
+		UV = ApplyAddressMode(UV, AddressU, AddressV);
+		TextureLoadVirtualPageTableInternal(Result, PageTable0, PageTable1, PageTableUniform, UV, vLevel);
+		StoreVirtualTextureFeedback(Feedback, SampleIndex, Result.PackedRequest);
+	}
 	return Result;
 }
 
@@ -495,13 +511,16 @@ VTPageTableResult TextureLoadVirtualPageTable(
 	MaterialFloat MipBias, float2 SvPositionXY)
 {
 	VTPageTableResult Result = (VTPageTableResult)0;
-	UV = UV * PageTableUniform.UVScale;
-	int vLevel = GetGlobalVirtualTextureMipBias();
+	if (IsValid(PageTableUniform))
+	{
+		UV = UV * PageTableUniform.UVScale;
+		int vLevel = GetGlobalVirtualTextureMipBias();
 #if PIXELSHADER
-	vLevel = TextureComputeVirtualMipLevel(Result, ddx(UV), ddy(UV), MipBias, SvPositionXY, PageTableUniform);
+		vLevel = TextureComputeVirtualMipLevel(Result, ddx(UV), ddy(UV), MipBias, SvPositionXY, PageTableUniform);
 #endif // PIXELSHADER
-	UV = ApplyAddressMode(UV, AddressU, AddressV);
-	TextureLoadVirtualPageTableInternal(Result, PageTable0, PageTable1, PageTableUniform, UV, vLevel);
+		UV = ApplyAddressMode(UV, AddressU, AddressV);
+		TextureLoadVirtualPageTableInternal(Result, PageTable0, PageTable1, PageTableUniform, UV, vLevel);
+	}
 	return Result;
 }
 
@@ -515,11 +534,14 @@ VTPageTableResult TextureLoadVirtualPageTableGrad(
 	in out FVirtualTextureFeedbackParams Feedback)
 {
 	VTPageTableResult Result = (VTPageTableResult)0;
-	int vLevel = TextureComputeVirtualMipLevel(Result, dUVdx * PageTableUniform.UVScale, dUVdy * PageTableUniform.UVScale, 0, SvPositionXY, PageTableUniform);
-	UV = UV * PageTableUniform.UVScale;
-	UV = ApplyAddressMode(UV, AddressU, AddressV);
-	TextureLoadVirtualPageTableInternal(Result, PageTable0, PageTableUniform, UV, vLevel);
-	StoreVirtualTextureFeedback(Feedback, SampleIndex, Result.PackedRequest);
+	if (IsValid(PageTableUniform))
+	{
+		int vLevel = TextureComputeVirtualMipLevel(Result, dUVdx * PageTableUniform.UVScale, dUVdy * PageTableUniform.UVScale, 0, SvPositionXY, PageTableUniform);
+		UV = UV * PageTableUniform.UVScale;
+		UV = ApplyAddressMode(UV, AddressU, AddressV);
+		TextureLoadVirtualPageTableInternal(Result, PageTable0, PageTableUniform, UV, vLevel);
+		StoreVirtualTextureFeedback(Feedback, SampleIndex, Result.PackedRequest);
+	}
 	return Result;
 }
 
@@ -532,12 +554,15 @@ VTPageTableResult TextureLoadVirtualPageTableGrad(
 	uint SampleIndex,
 	in out FVirtualTextureFeedbackParams Feedback)
 {
-	VTPageTableResult Result = (VTPageTableResult)0.0f;
-	int vLevel = TextureComputeVirtualMipLevel(Result, dUVdx * PageTableUniform.UVScale, dUVdy * PageTableUniform.UVScale, 0, SvPositionXY, PageTableUniform);
-	UV = UV * PageTableUniform.UVScale;
-	UV = ApplyAddressMode(UV, AddressU, AddressV);
-	TextureLoadVirtualPageTableInternal(Result, PageTable0, PageTable1, PageTableUniform, UV, vLevel);
-	StoreVirtualTextureFeedback(Feedback, SampleIndex, Result.PackedRequest);
+	VTPageTableResult Result = (VTPageTableResult)0;
+	if (IsValid(PageTableUniform))
+	{
+		int vLevel = TextureComputeVirtualMipLevel(Result, dUVdx * PageTableUniform.UVScale, dUVdy * PageTableUniform.UVScale, 0, SvPositionXY, PageTableUniform);
+		UV = UV * PageTableUniform.UVScale;
+		UV = ApplyAddressMode(UV, AddressU, AddressV);
+		TextureLoadVirtualPageTableInternal(Result, PageTable0, PageTable1, PageTableUniform, UV, vLevel);
+		StoreVirtualTextureFeedback(Feedback, SampleIndex, Result.PackedRequest);
+	}
 	return Result;
 }
 
@@ -548,11 +573,14 @@ VTPageTableResult TextureLoadVirtualPageTableGrad(
 	float2 UV, uint AddressU, uint AddressV,
 	MaterialFloat2 dUVdx, MaterialFloat2 dUVdy, float2 SvPositionXY)
 {
-	VTPageTableResult Result = (VTPageTableResult)0.0f;
-	int vLevel = TextureComputeVirtualMipLevel(Result, dUVdx * PageTableUniform.UVScale, dUVdy * PageTableUniform.UVScale, 0, SvPositionXY, PageTableUniform);
-	UV = UV * PageTableUniform.UVScale;
-	UV = ApplyAddressMode(UV, AddressU, AddressV);
-	TextureLoadVirtualPageTableInternal(Result, PageTable0, PageTable1, PageTableUniform, UV, vLevel);
+	VTPageTableResult Result = (VTPageTableResult)0;
+	if (IsValid(PageTableUniform))
+	{
+		int vLevel = TextureComputeVirtualMipLevel(Result, dUVdx * PageTableUniform.UVScale, dUVdy * PageTableUniform.UVScale, 0, SvPositionXY, PageTableUniform);
+		UV = UV * PageTableUniform.UVScale;
+		UV = ApplyAddressMode(UV, AddressU, AddressV);
+		TextureLoadVirtualPageTableInternal(Result, PageTable0, PageTable1, PageTableUniform, UV, vLevel);
+	}
 	return Result;
 }
 
@@ -566,11 +594,14 @@ VTPageTableResult TextureLoadVirtualPageTableLevel(
 	in out FVirtualTextureFeedbackParams Feedback)
 {
 	VTPageTableResult Result = (VTPageTableResult)0;
-	UV = UV * PageTableUniform.UVScale;
-	int vLevel = (int)floor(Level + GetGlobalVirtualTextureMipBias());
-	UV = ApplyAddressMode(UV, AddressU, AddressV);
-	TextureLoadVirtualPageTableInternal(Result, PageTable0, PageTableUniform, UV, vLevel);
-	StoreVirtualTextureFeedback(Feedback, SampleIndex, Result.PackedRequest);
+	if (IsValid(PageTableUniform))
+	{
+		UV = UV * PageTableUniform.UVScale;
+		int vLevel = (int)floor(Level + GetGlobalVirtualTextureMipBias());
+		UV = ApplyAddressMode(UV, AddressU, AddressV);
+		TextureLoadVirtualPageTableInternal(Result, PageTable0, PageTableUniform, UV, vLevel);
+		StoreVirtualTextureFeedback(Feedback, SampleIndex, Result.PackedRequest);
+	}
 	return Result;
 }
 
@@ -582,10 +613,13 @@ VTPageTableResult TextureLoadVirtualPageTableLevel(
 	MaterialFloat Level)
 {
 	VTPageTableResult Result = (VTPageTableResult)0;
-	UV = UV * PageTableUniform.UVScale;
-	int vLevel = (int)floor(Level + GetGlobalVirtualTextureMipBias());
-	UV = ApplyAddressMode(UV, AddressU, AddressV);
-	TextureLoadVirtualPageTableInternal(Result, PageTable0, PageTableUniform, UV, vLevel);
+	if (IsValid(PageTableUniform))
+	{
+		UV = UV * PageTableUniform.UVScale;
+		int vLevel = (int)floor(Level + GetGlobalVirtualTextureMipBias());
+		UV = ApplyAddressMode(UV, AddressU, AddressV);
+		TextureLoadVirtualPageTableInternal(Result, PageTable0, PageTableUniform, UV, vLevel);
+	}
 	return Result;
 }
 
@@ -599,11 +633,14 @@ VTPageTableResult TextureLoadVirtualPageTableLevel(
 	in out FVirtualTextureFeedbackParams Feedback)
 {
 	VTPageTableResult Result = (VTPageTableResult)0;
-	UV = UV * PageTableUniform.UVScale;
-	int vLevel = (int)floor(Level + GetGlobalVirtualTextureMipBias());
-	UV = ApplyAddressMode(UV, AddressU, AddressV);
-	TextureLoadVirtualPageTableInternal(Result, PageTable0, PageTable1, PageTableUniform, UV, vLevel);
-	StoreVirtualTextureFeedback(Feedback, SampleIndex, Result.PackedRequest);
+	if (IsValid(PageTableUniform))
+	{
+		UV = UV * PageTableUniform.UVScale;
+		int vLevel = (int)floor(Level + GetGlobalVirtualTextureMipBias());
+		UV = ApplyAddressMode(UV, AddressU, AddressV);
+		TextureLoadVirtualPageTableInternal(Result, PageTable0, PageTable1, PageTableUniform, UV, vLevel);
+		StoreVirtualTextureFeedback(Feedback, SampleIndex, Result.PackedRequest);
+	}
 	return Result;
 }
 
@@ -776,11 +813,12 @@ MaterialFloat4 TextureVirtualSample(
 	VTPageTableResult PageTableResult, uint LayerIndex,
 	VTUniform Uniform)
 {
-	float2 pUV = VTComputePhysicalUVs(PageTableResult, LayerIndex, Uniform);
 	bool bValid = IsValid(PageTableResult, LayerIndex);
+	
 	MaterialFloat4 Sample1;
 	if (bValid)
 	{
+		float2 pUV = VTComputePhysicalUVs(PageTableResult, LayerIndex, Uniform);
 		Sample1 = TextureVirtualSampleAniso(Physical, PhysicalSampler, pUV, PageTableResult.dUVdx, PageTableResult.dUVdy);
 	}
 	else
@@ -789,11 +827,11 @@ MaterialFloat4 TextureVirtualSample(
 	}
 
 #if VIRTUAL_TEXTURE_MANUAL_TRILINEAR_FILTERING
-	pUV = VTComputePhysicalUVs(PageTableResult, LayerIndex + 8u, Uniform);
 	bValid = IsValid(PageTableResult, LayerIndex + 8u);
 	MaterialFloat4 Sample2;
 	if (bValid)
 	{
+		float2 pUV = VTComputePhysicalUVs(PageTableResult, LayerIndex + 8u, Uniform);
 		Sample2 = TextureVirtualSampleAniso(Physical, PhysicalSampler, pUV, PageTableResult.dUVdx, PageTableResult.dUVdy);
 	}
 	else
@@ -812,11 +850,11 @@ MaterialFloat4 TextureVirtualSampleLevel(
 	VTPageTableResult PageTableResult, uint LayerIndex,
 	VTUniform Uniform)
 {
-	const float2 pUV = VTComputePhysicalUVs(PageTableResult, LayerIndex, Uniform);
 	const bool bValid = IsValid(PageTableResult, LayerIndex);
-	// No need to apply dUVdx/dUVdy, don't support anisotropic when sampling a specific level
 	if (bValid)
 	{
+		const float2 pUV = VTComputePhysicalUVs(PageTableResult, LayerIndex, Uniform);
+		// No need to apply dUVdx/dUVdy, don't support anisotropic when sampling a specific level
 		return Physical.SampleLevel(PhysicalSampler, pUV, 0.0f);
 	}
 	return Uniform.FallbackValue;
