Fixes register pressure arround outputing the TSR.SceneColor

diff --git a/Engine/Shaders/Private/TemporalSuperResolution/TSRRejectShading.usf b/Engine/Shaders/Private/TemporalSuperResolution/TSRRejectShading.usf
--- a/Engine/Shaders/Private/TemporalSuperResolution/TSRRejectShading.usf
+++ b/Engine/Shaders/Private/TemporalSuperResolution/TSRRejectShading.usf
@@ -549,6 +549,14 @@ void MainCS(
 		}
 	}
 	
+	// Despill to LDS to make room in VGPRs for the outputing of InputSceneColorOutput
+	#if CONFIG_LDS_DESPILL
+	PLATFORM_SPECIFIC_ISOLATE
+	{
+		Concatenate(OriginalHistory, OriginalHistoryUncertainty).StoreLDSSpill(kDespillHistory);
+	}
+	#endif // CONFIG_LDS_DESPILL
+
 	// Compute AA Luma and compose separate translucency into scene color for optimised HistoryUpdate.
 	{
 		tsr_tensor_halfC CenterFinalSceneColor = ComposeTranslucency(OriginalOpaqueInput, OriginalTranslucencyInput);
@@ -571,12 +579,11 @@ void MainCS(
 
 	tsr_tensor_halfC OriginalInput = ComposeTranslucencyForRejection(OriginalOpaqueInput, OriginalTranslucencyInput, bHasPixelAnimation);
 	
-	// Despill to LDS to make room in VGPRs for the convolutions of ComputeMoireError()
+	// Despill to LDS to make room in VGPRs for the convolutions of MeasureRejection()
 	#if CONFIG_LDS_DESPILL
 	PLATFORM_SPECIFIC_ISOLATE
 	{
 		Concatenate(OriginalInput, VelocityEdge).StoreLDSSpill(kDespillInput);
-		Concatenate(OriginalHistory, OriginalHistoryUncertainty).StoreLDSSpill(kDespillHistory);
 	}
 	#endif // CONFIG_LDS_DESPILL
 
@@ -590,8 +597,14 @@ void MainCS(
 	PLATFORM_SPECIFIC_ISOLATE
 	{
 		// Reload all the data from LDS or cache with moire to avoid register pressure
-		tsr_tensor_halfC Input   = OriginalInput;
-		tsr_tensor_halfC History = OriginalHistory;
+		tsr_tensor_halfC Input = OriginalInput;
+		tsr_tensor_halfC History;
+		#if CONFIG_LDS_DESPILL
+			tsr_tensor_half HistoryUncertainty;
+			Deconcatenate(tsr_tensor_half4::LoadLDSSpill(kDespillHistory), /* out */ History, /* out */ HistoryUncertainty);
+		#else
+			History = OriginalHistory;
+		#endif
 
 		// Convert input to shading LDR
 		tsr_tensor_halfC ExposedInput = AddPerceptionAdd(Input, /* DestPerceptionAdd = */ PerceptionAdd);
