Substrate - Fixed only one light showing up in mobile forward.
Also fixed a compilation issue.

diff --git a/Engine/Shaders/Private/Substrate/SubstrateMobileForwardLighting.ush b/Engine/Shaders/Private/Substrate/SubstrateMobileForwardLighting.ush
--- a/Engine/Shaders/Private/Substrate/SubstrateMobileForwardLighting.ush
+++ b/Engine/Shaders/Private/Substrate/SubstrateMobileForwardLighting.ush
@@ -16,7 +16,7 @@
 
 // For debug or iteration purpose only
 #ifndef SUBSTRATE_FORCE_SINGLE_LOCAL_LIGHT
-#define SUBSTRATE_FORCE_SINGLE_LOCAL_LIGHT 1
+#define SUBSTRATE_FORCE_SINGLE_LOCAL_LIGHT 0
 #endif
 
 // Forward declarations
@@ -285,7 +285,6 @@ float3 SubstrateMobileForwardLighting(
 					LightData.SoftSourceRadius	= 0;
 					LightData.SpecularScale		= LocalLightB.a;
 					LightData.DiffuseScale		= 1.f;
-					LightData.DirectionScale	= 1.f;
 					LightData.bSpotLight		= false;
 					LightData.bRectLight		= false;
 					LightData.Color 			= LocalLightA.rgb;
@@ -499,6 +498,7 @@ float3 SubstrateMobileForwardLighting(
 							BSDFEvaluate);
 					#endif
 					}
+					#undef LOCAL_LIGHT_EARLY_OUT
 				}
 				#endif // !DISABLE_FORWARD_LOCAL_LIGHTS
 
