Add tangent visualization under mouse cursor when viewing hair instance debug information.

diff --git a/Engine/Shaders/Private/HairStrands/HairStrandsDebug.usf b/Engine/Shaders/Private/HairStrands/HairStrandsDebug.usf
--- a/Engine/Shaders/Private/HairStrands/HairStrandsDebug.usf
+++ b/Engine/Shaders/Private/HairStrands/HairStrandsDebug.usf
@@ -1351,6 +1351,16 @@ void MainCS(uint InGroupIndex : SV_GroupIndex, uint3 DispatchThreadId : SV_Dispa
 		Print(Context2, Sample.Specular,		FontEmerald);						Newline(Context2);
 		Print(Context2, Sample.Backlit,		FontOrange);						Newline(Context2);
 		Newline(Context2);
+
+		// Draw tangent vector (strands)
+		{
+			//const float2 PixelCoord = GetCursorPos();// DispatchThreadId.xy + 0.5f;
+			const float2 UV = PixelCoord / float2(MaxResolution);
+
+			const float SceneDepth = ConvertFromDeviceZ(Sample.Depth);
+			const float3 WorldPosition = ReconstructTranslatedWorldPositionFromDepth(UV, SceneDepth);
+			AddLineTWS(WorldPosition, WorldPosition + Sample.Tangent, ColorRed, ColorYellow);
+		}
 	#endif 
 
 	#if DETAIL_DATA
