Use blue noise for SMRT step offset and directional light ray.
Misc cleanup.

[RN] Use blue noise for SMRT step offset and directional light ray

diff --git a/Engine/Shaders/Private/VirtualShadowMaps/VirtualShadowMapProjection.usf b/Engine/Shaders/Private/VirtualShadowMaps/VirtualShadowMapProjection.usf
--- a/Engine/Shaders/Private/VirtualShadowMaps/VirtualShadowMapProjection.usf
+++ b/Engine/Shaders/Private/VirtualShadowMaps/VirtualShadowMapProjection.usf
@@ -1,11 +1,11 @@
 // Copyright Epic Games, Inc. All Rights Reserved.
 
-#define SUPPORT_CONTACT_SHADOWS 1
 #include "../Common.ush"
 #include "../SceneTexturesCommon.ush"
 #include "../Substrate/Substrate.ush"
 #include "../DeferredShadingCommon.ush"
-#include "../DeferredLightingCommon.ush"
+#include "../DynamicLightingCommon.ush"
+#include "../MortonCode.ush"
 #include "../LightShaderParameters.ush"
 #include "../LightGridCommon.ush"
 #include "../HairStrands/HairStrandsVisibilityCommon.ush"
@@ -187,12 +187,14 @@ uint SMRTAdaptiveRayCount;
 uint SMRTHairRayCount;
 
 // One pass projection
+#if ONE_PASS_PROJECTION
 RWTexture2D< uint4 > OutShadowMaskBits;
-
+#else
 // Single light per pass
 // Light parameters loaded via GetRootLightShaderParameters();
 int LightUniformVirtualShadowMapId;
 RWTexture2D< float2 > OutShadowFactor;
+#endif
 
 // Visualization output
 StructuredBuffer< FPhysicalPageMetaData > PhysicalPageMetaData;
@@ -677,7 +679,7 @@ void VirtualShadowMapProjection(
 	const float3 TranslatedWorldPosition = SvPositionToTranslatedWorld( SvPosition );
 
 	const float ScreenRayLengthWorld = GetScreenRayLengthMultiplierForProjectionType(ScreenRayLength * SceneDepth).y;
-	const float Noise = InterleavedGradientNoise( SvPosition.xy, View.StateFrameIndexMod8 );
+	const float Noise = BlueNoiseScalar(SvPosition.xy, View.StateFrameIndex);
 
 	const FProjectionShadingInfo ShadingInfo = GetProjectionShadingInfo(PixelPos);
 
diff --git a/Engine/Shaders/Private/VirtualShadowMaps/VirtualShadowMapProjectionDirectional.ush b/Engine/Shaders/Private/VirtualShadowMaps/VirtualShadowMapProjectionDirectional.ush
--- a/Engine/Shaders/Private/VirtualShadowMaps/VirtualShadowMapProjectionDirectional.ush
+++ b/Engine/Shaders/Private/VirtualShadowMaps/VirtualShadowMapProjectionDirectional.ush
@@ -8,15 +8,11 @@ VirtualShadowMapProjectionDirectional.ush:
 #include "../DeferredShadingCommon.ush"
 #include "../SceneTexturesCommon.ush"
 #include "../LightShaderParameters.ush"
-#include "../PathTracing/Utilities/PathTracingRandomSequence.ush"
 #include "../Visualization.ush"
-#include "../BlueNoise.ush"
 #include "VirtualShadowMapPageAccessCommon.ush"
 #include "VirtualShadowMapProjectionCommon.ush"
 #include "VirtualShadowMapSMRTCommon.ush"
 
-#define USE_BLUE_NOISE 1
-
 float2 ComputeDepthSlopeDirectionalUV(
 	FVirtualShadowMapProjectionShaderData ProjectionData,
 	float3 EstimatedGeoWorldNormal,
@@ -176,37 +172,6 @@ bool IsBackfaceToDirectionalLight(float3 Normal, float3 LightDirection, float Li
 	return dot(Normal, LightDirection) < -SinAlpha;
 }
 
-// See http://extremelearning.com.au/unreasonable-effectiveness-of-quasirandom-sequences/
-float2 R2Sequence(int n)
-{
-	return frac(float(n) * float2(0.754877669f, 0.569840296f));
-}
-
-float4 GetRandomSample(
-	uint2 PixelPos,
-	uint TimeIndex,
-	uint SampleIndex,
-	uint MaxSampleCount)
-{
-#if USE_BLUE_NOISE
-	{
-		// Blue noise is cheaper to compute
-		int2 Offset1 = int2(R2Sequence(SampleIndex) * BlueNoise.Dimensions.xy);
-		int2 Offset2 = int2(R2Sequence(SampleIndex + MaxSampleCount) * BlueNoise.Dimensions.xy);
-		float4 RandomSample;
-		RandomSample.xy = BlueNoiseVec2(PixelPos + Offset1, TimeIndex);
-		RandomSample.zw = BlueNoiseVec2(PixelPos + Offset2, TimeIndex);
-		return RandomSample;
-	}
-#else
-	{
-		RandomSequence RandSequence;
-		RandomSequence_Initialize(RandSequence, PixelPos, SampleIndex, TimeIndex, MaxSampleCount);
-		return RandomSequence_GenerateSample4D(RandSequence);
-	}
-#endif
-}
-
 FVirtualShadowMapSampleResult TraceDirectional(
 	int VirtualShadowMapId,
 	FLightShaderParameters Light,
@@ -287,7 +252,7 @@ FVirtualShadowMapSampleResult TraceDirectional(
 	for ( ; i < MaxRayCount; i++)
 	{
 		// One sample for ray, one for texel dither
-		float4 RandSample = GetRandomSample(PixelPos, View.StateFrameIndex, i, MaxRayCount);
+		float4 RandSample = VirtualShadowMapGetRandomSample(PixelPos, View.StateFrameIndex, i, MaxRayCount);
 		float3 RayDir = GetRandomDirectionalLightRayDir(Light, RandSample.xy);
 		float2 TexelOffset = (RandSample.zw - 0.5f) * DitherScale;
 
@@ -386,7 +351,7 @@ bool GenerateRayDirectional(
 	// Too low values will greatly restrict how large penumbras can be in screen space.
 	float RayLength = RayLengthScale * DistanceFromViewOrigin;
 
-	float2 E = GetRandomSample(PixelPos, View.StateFrameIndex, RayIndex, RayCount).xy;
+	float2 E = VirtualShadowMapGetRandomSample(PixelPos, View.StateFrameIndex, RayIndex, RayCount).xy;
 
 	float3 RayDir = GetRandomDirectionalLightRayDir(Light, E);
 	OutRayStart = TranslatedWorldPosition;
diff --git a/Engine/Shaders/Private/VirtualShadowMaps/VirtualShadowMapProjectionSpot.ush b/Engine/Shaders/Private/VirtualShadowMaps/VirtualShadowMapProjectionSpot.ush
--- a/Engine/Shaders/Private/VirtualShadowMaps/VirtualShadowMapProjectionSpot.ush
+++ b/Engine/Shaders/Private/VirtualShadowMaps/VirtualShadowMapProjectionSpot.ush
@@ -211,15 +211,13 @@ FSMRTResult ShadowRayCastPointLight(
 	// To get debug data out
 	inout FSMRTSingleRayState OutRayState)
 {
-	uint RayStartFace = VirtualShadowMapGetCubeFace(RayStartShadowTranslatedWorld);
-	uint RayEndFace   = VirtualShadowMapGetCubeFace(RayEndShadowTranslatedWorld);
+	uint RayStartFaceId = VirtualShadowMapId + VirtualShadowMapGetCubeFace(RayStartShadowTranslatedWorld);
+	uint RayEndFaceId   = VirtualShadowMapId + VirtualShadowMapGetCubeFace(RayEndShadowTranslatedWorld);
 
 	FSMRTResult Result;
-	if (WaveActiveAnyTrue(RayStartFace != RayEndFace))
+	if (WaveActiveAnyTrue(RayStartFaceId != RayEndFaceId))
 	{
-		FSMRTTwoCubeFaceRayState RayState = SMRTTwoCubeFaceRayInitialize(
-			VirtualShadowMapId + RayStartFace,
-			VirtualShadowMapId + RayEndFace,
+		FSMRTTwoCubeFaceRayState RayState = SMRTTwoCubeFaceRayInitialize(RayStartFaceId, RayEndFaceId,
 			RayStartShadowTranslatedWorld, RayEndShadowTranslatedWorld, ExtrapolateSlope, DepthBias);
 		Result = SMRTRayCast(RayState, NumSteps, StepOffset);
 
@@ -235,7 +233,7 @@ FSMRTResult ShadowRayCastPointLight(
 	else
 	{
 		// Fast path: ray stays on a single cube map face (effectively a spot light)
-		OutRayState = SMRTSingleRayInitialize(VirtualShadowMapId + RayEndFace, RayStartShadowTranslatedWorld, RayEndShadowTranslatedWorld, ExtrapolateSlope, DepthBias);
+		OutRayState = SMRTSingleRayInitialize(RayEndFaceId, RayStartShadowTranslatedWorld, RayEndShadowTranslatedWorld, ExtrapolateSlope, DepthBias);
 		Result = SMRTRayCast(OutRayState, NumSteps, StepOffset);
 	}
 
@@ -394,18 +392,8 @@ FVirtualShadowMapSampleResult TraceLocalLight(
 		float MaxOccluderDistance = -1.0f;
 		for( ; i < MaxRayCount; i++ )
 		{
-			RandomSequence RandSequence;
-#if 0
-			RandomSequence_Initialize( RandSequence, PixelPos, i, View.StateFrameIndex, MaxRayCount );
-			float4 RandSample = RandomSequence_GenerateSample4D( RandSequence );
-#else
-			// Combine pixel-level and sample-level bits into the sample index (visible structure will be hidden by owen scrambling of the index)
-			RandSequence.SampleIndex = reversebits((65536 * View.StateFrameIndex + TileID) * MaxRayCount + i);
-			RandSequence.SampleSeed = 0; // always use the same sequence
-
-			float4 RandSample = SobolSampler( RandSequence.SampleIndex, RandSequence.SampleSeed, 8 );
-#endif
-
+			float4 RandSample = VirtualShadowMapGetRandomSample(PixelPos, View.StateFrameIndex, i, MaxRayCount);
+			
 			float2 LightUV = UniformSampleDiskConcentricApprox( RandSample.xy ).xy;
 			LightUV *= ConeSin;
 
diff --git a/Engine/Shaders/Private/VirtualShadowMaps/VirtualShadowMapSMRTCommon.ush b/Engine/Shaders/Private/VirtualShadowMaps/VirtualShadowMapSMRTCommon.ush
--- a/Engine/Shaders/Private/VirtualShadowMaps/VirtualShadowMapSMRTCommon.ush
+++ b/Engine/Shaders/Private/VirtualShadowMaps/VirtualShadowMapSMRTCommon.ush
@@ -6,6 +6,9 @@ VirtualShadowMapSMRTCommon.ush:
 
 #pragma once
 
+//#include "../PathTracing/Utilities/PathTracingRandomSequence.ush"
+#include "../BlueNoise.ush"
+
 struct FSMRTSample
 {
 	bool bValid;	
@@ -31,3 +34,34 @@ struct FSMRTResult
 	bool bValidHit;
 	float HitDepth;	// Valid if bValidHit is true
 };
+
+// See http://extremelearning.com.au/unreasonable-effectiveness-of-quasirandom-sequences/
+float2 R2Sequence(int n)
+{
+	return frac(float(n) * float2(0.754877669f, 0.569840296f));
+}
+
+float4 VirtualShadowMapGetRandomSample(
+	uint2 PixelPos,
+	uint TimeIndex,
+	uint SampleIndex,
+	uint MaxSampleCount)
+{
+#if 1
+	{
+		// Blue noise is cheaper to compute
+		int2 Offset1 = int2(R2Sequence(SampleIndex) * BlueNoise.Dimensions.xy);
+		int2 Offset2 = int2(R2Sequence(SampleIndex + MaxSampleCount) * BlueNoise.Dimensions.xy);
+		float4 RandomSample;
+		RandomSample.xy = BlueNoiseVec2(PixelPos + Offset1, TimeIndex);
+		RandomSample.zw = BlueNoiseVec2(PixelPos + Offset2, TimeIndex);
+		return RandomSample;
+	}
+#else
+	{
+		RandomSequence RandSequence;
+		RandomSequence_Initialize(RandSequence, PixelPos, SampleIndex, TimeIndex, MaxSampleCount);
+		return RandomSequence_GenerateSample4D(RandSequence);
+	}
+#endif
+}
\ No newline at end of file
