Fixes a bug in DecodeReprojectionBoundary() when using 16bit ops

diff --git a/Engine/Shaders/Private/TemporalSuperResolution/TSRReprojectionField.ush b/Engine/Shaders/Private/TemporalSuperResolution/TSRReprojectionField.ush
--- a/Engine/Shaders/Private/TemporalSuperResolution/TSRReprojectionField.ush
+++ b/Engine/Shaders/Private/TemporalSuperResolution/TSRReprojectionField.ush
@@ -38,7 +38,7 @@ bool HasEncodedReprojectionOffset(uint EncodedReprojectionBoundary)
 /** Decodes EncodeReprojectionBoundary()'s ReprojectionBoundary. */
 tsr_half2 DecodeReprojectionBoundary(uint EncodedReprojectionBoundary)
 {
-	tsr_short2 BoundaryBits = tsr_short2(EncodedReprojectionBoundary >> 4u, EncodedReprojectionBoundary >> 15u) % tsr_short(2048) - tsr_short(1024);
+	tsr_short2 BoundaryBits = (tsr_short2(EncodedReprojectionBoundary >> 4u, EncodedReprojectionBoundary >> 15u) & tsr_short(0x7FF)) - tsr_short(1024);
 	tsr_half2 ReprojectionBoundary = fastClamp(tsr_half2(BoundaryBits) * rcp(tsr_half(1023)), tsr_half(-1.0).xx, tsr_half(1.0).xx);
 	return ReprojectionBoundary;
 }
