Fix shader compilation on Forward mobile when merged local light are enabled.

diff --git a/Engine/Shaders/Private/MobileLightingCommon.ush b/Engine/Shaders/Private/MobileLightingCommon.ush
--- a/Engine/Shaders/Private/MobileLightingCommon.ush
+++ b/Engine/Shaders/Private/MobileLightingCommon.ush
@@ -388,7 +388,7 @@ half3 GetImageBasedReflectionLighting_Mobile(
 	}
 
 	SpecularIBL *= CompositeAlpha;
-#else // !(ENABLE_CLUSTERED_REFLECTION || MOBILE_DEFERRED_LIGHTING)
+#elif IS_MOBILE_BASE_PASS || IS_DECAL
 
 	bool UsingSkyReflection = MobileReflectionCapture.Params.y > 0.0f;
 	// Normalize for static skylight types which mix with lightmaps.
