ManyLights - Remove some unused code and add extra debugging stuff

diff --git a/Engine/Shaders/Private/ManyLights/ManyLightsRayTracing.usf b/Engine/Shaders/Private/ManyLights/ManyLightsRayTracing.usf
--- a/Engine/Shaders/Private/ManyLights/ManyLightsRayTracing.usf
+++ b/Engine/Shaders/Private/ManyLights/ManyLightsRayTracing.usf
@@ -99,7 +99,6 @@ void CompactLightSampleTracesCS(
 	bool bTraceValid = false;
 	uint TraceTexelForThisThread = 0;
 
-	uint2 LocalTileMaskCoord = SampleCoordToDownsampledScreenCoord(SampleCoord - GroupId.xy * THREADGROUP_SIZE) / TILE_SIZE;
 	if (all(SampleCoord < SampleViewMin + SampleViewSize))
 	{
 		// #ml_todo: try to preserve z-order here for improved tracing coherency
diff --git a/Engine/Shaders/Private/ManyLights/ManyLightsSampling.usf b/Engine/Shaders/Private/ManyLights/ManyLightsSampling.usf
--- a/Engine/Shaders/Private/ManyLights/ManyLightsSampling.usf
+++ b/Engine/Shaders/Private/ManyLights/ManyLightsSampling.usf
@@ -12,6 +12,15 @@
 #include "../BlueNoise.ush"
 #include "ManyLights.ush"
 #include "ManyLightsRayTracing.ush"
+#include "../LightFunctionAtlas/LightFunctionAtlasCommon.usf"
+
+struct FDebug
+{
+#if DEBUG_MODE
+	bool bActive;
+	FShaderPrintContext Context;
+#endif
+};
 
 struct FLightTargetPDF
 {
@@ -29,7 +38,7 @@ FLightTargetPDF InitLightTargetPDF()
 	return LightTargetPDF;
 }
 
-FLightTargetPDF GetLocalLightTargetPDF(FDeferredLightData LightData, float3 TranslatedWorldPosition, FGBufferData GBuffer, uint2 ScreenCoord)
+FLightTargetPDF GetLocalLightTargetPDF(FDeferredLightData LightData, float3 TranslatedWorldPosition, FGBufferData GBuffer, uint2 ScreenCoord, inout FDebug Debug)
 {
 	FLightTargetPDF LightTargetPDF = InitLightTargetPDF();
 
@@ -166,11 +175,14 @@ void GenerateLightSamplesCS(
 		{
 			uint2 ScreenCoord = DownsampledScreenCoordToScreenCoord(DownsampledScreenCoord);
 
-	#if DEBUG_MODE
-			int2 DebugScreenCoord = View.CursorPosition.x >= 0 ? View.CursorPosition * View.ViewResolutionFraction : -1;
-			bool bDebug = all(DownsampledScreenCoord == DebugScreenCoord / DOWNSAMPLE_FACTOR);
-			FShaderPrintContext Context = InitShaderPrintContext(true, float2(0.05, 0.05));
-	#endif
+			FDebug Debug;
+			#if DEBUG_MODE
+			{
+				int2 DebugScreenCoord = View.CursorPosition.x >= 0 ? View.CursorPosition * View.ViewResolutionFraction : -1;
+				Debug.bActive = all(DownsampledScreenCoord == DebugScreenCoord / DOWNSAMPLE_FACTOR);
+				Debug.Context = InitShaderPrintContext(true, float2(0.05, 0.05));
+			}
+			#endif
 
 			FLightSample LightSamples[NUM_SAMPLES_PER_PIXEL_1D];
 			for (uint LightSampleIndex = 0; LightSampleIndex < NUM_SAMPLES_PER_PIXEL_1D; ++LightSampleIndex)
@@ -216,35 +228,35 @@ void GenerateLightSamplesCS(
 				}
 
 				#if DEBUG_MODE
-				if (bDebug)
+				if (Debug.bActive)
 				{
-					Print(Context, TEXT("GenerateSamplesCS"));
-					Newline(Context);
-					Print(Context, TEXT("ScreenCoord: "));
-					Print(Context, ScreenCoord.x);
-					Print(Context, ScreenCoord.y);
-					Newline(Context);
-					Print(Context, TEXT("TileType: "));
-					Print(Context, uint(TILE_TYPE));
-					Print(Context, DownsampledTileAllocator[TILE_TYPE]);
-					Newline(Context);
-					Print(Context, TEXT("View.PreExposure: "));
-					Print(Context, View.PreExposure);
-					Newline(Context);
-					Print(Context, TEXT("NumLightsInGridCell: "));
-					Print(Context, NumLightsInGridCell);
-					Print(Context, TEXT("NumLocalLights: "));
+					Print(Debug.Context, TEXT("GenerateSamplesCS"));
+					Newline(Debug.Context);
+					Print(Debug.Context, TEXT("ScreenCoord: "));
+					Print(Debug.Context, ScreenCoord.x);
+					Print(Debug.Context, ScreenCoord.y);
+					Newline(Debug.Context);
+					Print(Debug.Context, TEXT("TileType: "));
+					Print(Debug.Context, uint(TILE_TYPE));
+					Print(Debug.Context, DownsampledTileAllocator[TILE_TYPE]);
+					Newline(Debug.Context);
+					Print(Debug.Context, TEXT("View.PreExposure: "));
+					Print(Debug.Context, View.PreExposure);
+					Newline(Debug.Context);
+					Print(Debug.Context, TEXT("NumLightsInGridCell: "));
+					Print(Debug.Context, NumLightsInGridCell);
+					Print(Debug.Context, TEXT("NumLocalLights: "));
 					if (NumLocalLights <= MAX_LOCAL_LIGHT_INDEX)
 					{
-						Print(Context, NumLocalLights);
+						Print(Debug.Context, NumLocalLights);
 					}
 					else
 					{
-						Print(Context, NumLocalLights, FontRed);
+						Print(Debug.Context, NumLocalLights, FontRed);
 					}
-					AddCrossTWS(Context, TranslatedWorldPosition, 5.0f, float4(1, 1, 0, 1));
-					Newline(Context);
-					Print(Context, TEXT("LightId  | SampleWeight"));
+					AddCrossTWS(Debug.Context, TranslatedWorldPosition, 5.0f, float4(1, 1, 0, 1));
+					Newline(Debug.Context);
+					Print(Debug.Context, TEXT("LightId  | SampleWeight | LFAtlasIndex | IESAtlasIndex"));
 				}
 				#endif
 
@@ -268,14 +280,16 @@ void GenerateLightSamplesCS(
 						const FLocalLightData LocalLightData = GetLocalLightDataNonStereo(LocalLightIndex);
 						FDeferredLightData LightData = ConvertToDeferredLight(LocalLightData);
 					
-						FLightTargetPDF LightTargetPDF = GetLocalLightTargetPDF(LightData, TranslatedWorldPosition, GBuffer, ScreenCoord);
+						FLightTargetPDF LightTargetPDF = GetLocalLightTargetPDF(LightData, TranslatedWorldPosition, GBuffer, ScreenCoord, Debug);
 
 						#if DEBUG_MODE
-						if (bDebug)
+						if (Debug.bActive)
 						{
-							Newline(Context);
-							Print(Context, LocalLightData.LightSceneId);
-							Print(Context, LightTargetPDF.Weight);
+							Newline(Debug.Context);
+							Print(Debug.Context, LocalLightData.LightSceneId);
+							Print(Debug.Context, LightTargetPDF.Weight);
+							Print(Debug.Context, LightData.LightFunctionAtlasLightIndex);
+							Print(Debug.Context, LightData.IESAtlasIndex);
 						}
 						#endif
 
@@ -310,22 +324,22 @@ void GenerateLightSamplesCS(
 				DemodulateLighting(GBuffer, TranslatedWorldPosition, DiffuseSum, SpecularSum);
 
 				#if DEBUG_MODE
-				if (bDebug)
+				if (Debug.bActive)
 				{
-					Newline(Context);
-					Print(Context, TEXT("Weight sum: "));
-					Print(Context, WeightSum);
-					Newline(Context);
-					Print(Context, TEXT("Selected lights: "));
+					Newline(Debug.Context);
+					Print(Debug.Context, TEXT("Weight sum: "));
+					Print(Debug.Context, WeightSum);
+					Newline(Debug.Context);
+					Print(Debug.Context, TEXT("Selected lights: "));
 
 					for (uint LightSampleIndex = 0; LightSampleIndex < NUM_SAMPLES_PER_PIXEL_1D; ++LightSampleIndex)
 					{
 						FLightSample LightSample = LightSamples[LightSampleIndex];
 						
-						Newline(Context);
-						Print(Context, LightSample.LocalLightIndex); 
-						Print(Context, LightSample.Weight);
-						Print(Context, WeightSum /  LightSample.Weight);
+						Newline(Debug.Context);
+						Print(Debug.Context, LightSample.LocalLightIndex); 
+						Print(Debug.Context, LightSample.Weight);
+						Print(Debug.Context, WeightSum /  LightSample.Weight);
 
 						if (DebugMode == DEBUG_MODE_VISUALIZE_SAMPLING)
 						{
@@ -334,16 +348,16 @@ void GenerateLightSamplesCS(
 							const FLocalLightData LocalLightData = GetLocalLightDataNonStereo(LightSample.LocalLightIndex);
 							const FLightSampleTrace LightSampleTrace = GetLightSampleTrace(TranslatedWorldPosition, LightSample.LocalLightIndex, SampleCoord);
 
-							AddLineTWS(Context, TranslatedWorldPosition, TranslatedWorldPosition + LightSampleTrace.Direction * LightSampleTrace.Distance, float4(LocalLightData.LightColorAndIdAndFalloffExponent.xyz / Luminance(LocalLightData.LightColorAndIdAndFalloffExponent.xyz), 1.0f));
+							AddLineTWS(Debug.Context, TranslatedWorldPosition, TranslatedWorldPosition + LightSampleTrace.Direction * LightSampleTrace.Distance, float4(LocalLightData.LightColorAndIdAndFalloffExponent.xyz / Luminance(LocalLightData.LightColorAndIdAndFalloffExponent.xyz), 1.0f));
 						}
 					}
 
-					Newline(Context);
-					Print(Context, TEXT("Diffuse sum:  "));
-					Print(Context, DiffuseSum);
-					Newline(Context);
-					Print(Context, TEXT("Specular sum: "));
-					Print(Context, SpecularSum);
+					Newline(Debug.Context);
+					Print(Debug.Context, TEXT("Diffuse sum:  "));
+					Print(Debug.Context, DiffuseSum);
+					Newline(Debug.Context);
+					Print(Debug.Context, TEXT("Specular sum: "));
+					Print(Debug.Context, SpecularSum);
 				}
 				#endif
 			}
diff --git a/Engine/Source/Runtime/Renderer/Private/ManyLights/ManyLights.cpp b/Engine/Source/Runtime/Renderer/Private/ManyLights/ManyLights.cpp
--- a/Engine/Source/Runtime/Renderer/Private/ManyLights/ManyLights.cpp
+++ b/Engine/Source/Runtime/Renderer/Private/ManyLights/ManyLights.cpp
@@ -149,13 +149,6 @@ static TAutoConsoleVariable<int32> CVarManyLightsIESProfiles(
 	ECVF_Scalability | ECVF_RenderThreadSafe
 );
 
-static TAutoConsoleVariable<int32> CVarManyLightsTest(
-	TEXT("r.ManyLights.Test"),
-	1,
-	TEXT("Test."),
-	ECVF_Scalability | ECVF_RenderThreadSafe
-);
-
 namespace ManyLights
 {
 	// must match values in ManyLights.ush
