Add Substrate support for ManyLight.

diff --git a/Engine/Shaders/Private/LightAccumulator.ush b/Engine/Shaders/Private/LightAccumulator.ush
--- a/Engine/Shaders/Private/LightAccumulator.ush
+++ b/Engine/Shaders/Private/LightAccumulator.ush
@@ -185,6 +185,9 @@ struct FSubstrateDeferredLighting
 	// only used for development (not compiled out with #if so we can use if() instead of #if for better readability and compiler error checking)
 	// assumed to be compiled out otherwise
 	float EstimatedCost;
+
+	float3 TotalDiffuseLighting;
+	float3 TotalSpecularLighting;
 };
 
 FSubstrateDeferredLighting GetInitialisedSubstrateDeferredLighting()
diff --git a/Engine/Shaders/Private/ManyLights/ManyLights.usf b/Engine/Shaders/Private/ManyLights/ManyLights.usf
--- a/Engine/Shaders/Private/ManyLights/ManyLights.usf
+++ b/Engine/Shaders/Private/ManyLights/ManyLights.usf
@@ -39,16 +39,18 @@ void TileClassificationCS(
 	{
 #if DOWNSAMPLED_CLASSIFICATION
 		float2 ScreenUV = DownsampledScreenCoordToScreenUV(ScreenCoord);
+		uint2 EffectiveScreenCoord = DownsampledScreenCoordToScreenCoord(ScreenCoord);
 #else
 		float2 ScreenUV = (ScreenCoord + 0.5f) * View.BufferSizeAndInvSize.zw;
+		uint2 EffectiveScreenCoord = ScreenCoord;
 #endif
-		FGBufferData GBuffer = GetGBufferData(ScreenUV);
+		FManyLightMaterial Material = LoadMaterial(ScreenUV, EffectiveScreenCoord);
 
-		if (GBuffer.ShadingModelID == SHADINGMODELID_UNLIT)
+		if (!Material.IsValid())
 		{
 			SharedTileMode[LinearThreadIndex] = 0x1;
 		}
-		else if (GBuffer.ShadingModelID == SHADINGMODELID_DEFAULT_LIT)
+		else if (Material.IsSimple())
 		{
 			SharedTileMode[LinearThreadIndex] = 0x2;
 		}
@@ -175,8 +177,8 @@ void InitCompositeUpsampleWeightsCS(
 	if (all(ScreenCoord < View.ViewRectMinAndSize.xy + View.ViewRectMinAndSize.zw))
 	{
 		float2 ScreenUV = (ScreenCoord + 0.5f) * View.BufferSizeAndInvSize.zw;
-		FGBufferData GBuffer = GetGBufferData(ScreenUV);
-		float3 TranslatedWorldPosition = GetTranslatedWorldPositionFromScreenUV(ScreenUV, GBuffer.Depth);		
+		const FManyLightMaterial Material = LoadMaterial(ScreenUV, ScreenCoord);
+		float3 TranslatedWorldPosition = GetTranslatedWorldPositionFromScreenUV(ScreenUV, Material.Depth);
 
 		int2 DownsampledScreenCoord00 = floor(ScreenUV * View.BufferSizeAndInvSize.xy / DOWNSAMPLE_FACTOR - 0.5f);
 		float2 DownsampledGatherUV = (DownsampledScreenCoord00 + 1.0f) * DownsampledBufferInvSize;
@@ -201,7 +203,7 @@ void InitCompositeUpsampleWeightsCS(
 		#define PLANE_WEIGHTING 1
 		#if PLANE_WEIGHTING
 		{
-			float4 ScenePlane = float4(GBuffer.WorldNormal, dot(TranslatedWorldPosition, GBuffer.WorldNormal));
+			float4 ScenePlane = float4(Material.WorldNormal, dot(TranslatedWorldPosition, Material.WorldNormal));
 
 			float3 Position00 = GetTranslatedWorldPositionFromScreenUV(DownsampledScreenCoordToScreenUV(DownsampledScreenCoord00 + uint2(0, 0)), CornerDepths.x);
 			float3 Position10 = GetTranslatedWorldPositionFromScreenUV(DownsampledScreenCoordToScreenUV(DownsampledScreenCoord00 + uint2(1, 0)), CornerDepths.y);
@@ -214,14 +216,14 @@ void InitCompositeUpsampleWeightsCS(
 			PlaneDistances.z = abs(dot(float4(Position01, -1), ScenePlane));
 			PlaneDistances.w = abs(dot(float4(Position11, -1), ScenePlane));
 
-			float4 RelativeDepthDifference = PlaneDistances / GBuffer.Depth;
+			float4 RelativeDepthDifference = PlaneDistances / Material.Depth;
 
 			DepthWeights = select(CornerDepths > 0.0f, exp2(-10000.0f * (RelativeDepthDifference * RelativeDepthDifference)), 0.0f);
 		}
 		#else
 		{
-			float4 DepthDifference = abs(CornerDepths - GBuffer.Depth.xxxx);
-			float4 RelativeDepthDifference = DepthDifference / GBuffer.Depth;	
+			float4 DepthDifference = abs(CornerDepths - Material.Depth.xxxx);
+			float4 RelativeDepthDifference = DepthDifference / Material.Depth;
 
 			DepthWeights = select(CornerDepths > 0.0f, exp2(-100.0f * (RelativeDepthDifference * RelativeDepthDifference)), 0.0f);
 		}
diff --git a/Engine/Shaders/Private/ManyLights/ManyLights.ush b/Engine/Shaders/Private/ManyLights/ManyLights.ush
--- a/Engine/Shaders/Private/ManyLights/ManyLights.ush
+++ b/Engine/Shaders/Private/ManyLights/ManyLights.ush
@@ -20,6 +20,10 @@
 
 #include "../Lumen/LumenPosition.ush"
 
+#include "../Substrate/Substrate.ush"
+#include "../Substrate/SubstrateEvaluation.ush"
+#include "../Substrate/SubstrateDeferredLighting.ush"
+
 #define DEBUG_MODE_VISUALIZE_SAMPLING 1
 #define DEBUG_MODE_VISUALIZE_TRACING  2
 
@@ -113,28 +117,205 @@ float2 DownsampledScreenCoordToScreenUV(uint2 DownsampledScreenCoord)
 	return ScreenUV;
 }
 
+struct FManyLightMaterial
+{
+	float  Depth;
+	float3 WorldNormal;
+	float  Roughness;
+	bool   bIsValid;
+	bool   bIsSimple;
+	bool   bNeedsSeparateSubsurfaceLightAccumulation;
+
+	float3 DiffuseColor;
+	float3 SpecularColor;
+
+	bool IsValid()
+	{
+		return bIsValid;
+	}
+
+	bool IsSimple()
+	{
+		return bIsSimple;
+	}
+
+	void SetDepth(float In)
+	{
+		Depth = In;
+		#if !SUBSTRATE_ENABLED
+		GBuffer.Depth = In;
+		#endif
+	}
+
+	#if !SUBSTRATE_ENABLED
+	FGBufferData GBuffer;
+	#endif
+};
+
+FManyLightMaterial LoadMaterial(float2 ScreenUV, uint2 ScreenCoord, bool bForceSimpleShading=false)
+#if SUBSTRATE_ENABLED
+{
+	const FSubstrateTopLayerData TopLayerData = SubstrateUnpackTopLayerData(Substrate.TopLayerTexture.Load(uint3(ScreenCoord, 0)));
+	FManyLightMaterial Out = (FManyLightMaterial)0;
+	Out.WorldNormal   = TopLayerData.WorldNormal;
+	Out.Depth         = ConvertFromDeviceZ(SceneTexturesStruct.SceneDepthTexture.Load(int3(ScreenCoord, 0)).r);
+	Out.Roughness     = TopLayerData.Roughness;
+
+	#if SUBSTRATE_MATERIALCONTAINER_IS_VIEWRESOURCE
+
+	const float3 TranslatedWorldPosition = GetTranslatedWorldPositionFromScreenUV(ScreenUV, Out.Depth);
+	const float3 V = -GetCameraVectorFromTranslatedWorldPosition(TranslatedWorldPosition);
+
+	const FSubstrateIntegrationSettings Settings = InitSubstrateIntegrationSettings(false /*bForceFullyRough*/, Substrate.bRoughDiffuse, Substrate.PeelLayersAboveDepth, Substrate.bRoughnessTracking);
+	FSubstrateAddressing SubstrateAddressing	= GetSubstratePixelDataByteOffset(ScreenCoord, uint2(View.BufferSizeAndInvSize.xy), Substrate.MaxBytesPerPixel);
+	FSubstratePixelHeader SubstratePixelHeader= UnpackSubstrateHeaderIn(Substrate.MaterialTextureArray, SubstrateAddressing, Substrate.TopLayerTexture);
+
+	Out.bIsValid      = SubstratePixelHeader.IsSubstrateMaterial();
+	Out.bIsSimple     = SubstratePixelHeader.ClosureCount == 1;	
+	Substrate_for(uint ClosureIndex = 0, ClosureIndex < SubstratePixelHeader.ClosureCount, ++ClosureIndex)
+	{
+		FSubstrateBSDF BSDF = UnpackSubstrateBSDF(Substrate.MaterialTextureArray, SubstrateAddressing, SubstratePixelHeader);
+
+		// Create the BSDF context
+		FSubstrateBSDFContext SubstrateBSDFContext = SubstrateCreateBSDFContext(SubstratePixelHeader, BSDF, SubstrateAddressing, V);
+		const float3 BSDFThroughput = LuminanceWeight(SubstrateBSDFContext, BSDF);
+
+		// Evaluate environment lighting
+		FSubstrateEnvLightResult SubstrateEnvLight = SubstrateEvaluateForEnvLight(SubstrateBSDFContext, true /*bEnableSpecular*/, Settings);
+		
+		Out.DiffuseColor  += BSDFThroughput * SubstrateEnvLight.DiffuseColor;  //SubstrateEnvLight.DiffuseWeight;
+		Out.SpecularColor += BSDFThroughput * SubstrateEnvLight.SpecularColor; //SubstrateEnvLight.SpecularWeight;
+		#if SUBSTRATE_FASTPATH==0
+		if (any(SubstrateEnvLight.SpecularHazeWeight > 0.0f))
+		{
+			Out.SpecularColor += BSDFThroughput * SubstrateEnvLight.SpecularHazeWeight;
+		}
+		#endif 
+		if (SubstrateEnvLight.bPostProcessSubsurface)
+		{
+			Out.bNeedsSeparateSubsurfaceLightAccumulation = true;
+		}
+	}
+	#endif //SUBSTRATE_MATERIALCONTAINER_IS_VIEWRESOURCE
+
+	return Out;
+}
+#else
+{
+	FGBufferData GBuffer = GetGBufferData(ScreenUV);
+	GBuffer.Roughness = max(GBuffer.Roughness, View.MinRoughness);
+	if (GBuffer.ShadingModelID != SHADINGMODELID_UNLIT && bForceSimpleShading)
+	{
+		GBuffer.ShadingModelID = SHADINGMODELID_DEFAULT_LIT;
+	}
+
+	FManyLightMaterial Out = (FManyLightMaterial)0;
+	Out.bIsValid      = GBuffer.ShadingModelID != SHADINGMODELID_UNLIT;
+	Out.bIsSimple     = GBuffer.ShadingModelID == SHADINGMODELID_DEFAULT_LIT;
+	Out.bNeedsSeparateSubsurfaceLightAccumulation = UseSubsurfaceProfile(GBuffer.ShadingModelID);
+	Out.WorldNormal   = GBuffer.WorldNormal;
+	Out.Depth         = GBuffer.Depth;
+	Out.Roughness     = GBuffer.Roughness;
+	Out.DiffuseColor  = GBuffer.DiffuseColor;
+	Out.SpecularColor = GBuffer.SpecularColor;
+	Out.GBuffer       = GBuffer;
+	return Out;
+}
+#endif
+
 // Remove texture detail before denoising
-void DemodulateLighting(FGBufferData GBuffer, float3 TranslatedWorldPosition, inout float3 DiffuseLighting, inout float3 SpecularLighting)
+void DemodulateLighting(FManyLightMaterial Material, float3 TranslatedWorldPosition, inout float3 DiffuseLighting, inout float3 SpecularLighting)
 {
-	const float3 N = GBuffer.WorldNormal;
+	const float3 N = Material.WorldNormal;
 	const float3 V = normalize(View.TranslatedWorldCameraOrigin - TranslatedWorldPosition);
 	const float NoV = saturate(dot(N, V));
-	const float3 SpecularEnv = EnvBRDF(GBuffer.SpecularColor, GBuffer.Roughness, NoV);
+	const float3 SpecularEnv = EnvBRDF(Material.SpecularColor, Material.Roughness, NoV);
 
 	// #ml_todo: demodulate earlier during BRDF evaluation
-	DiffuseLighting = DiffuseLighting / max(GBuffer.DiffuseColor, 0.001f);
+	DiffuseLighting = DiffuseLighting / max(Material.DiffuseColor, 0.001f);
 	SpecularLighting = SpecularLighting / max(SpecularEnv, 0.001f);	
 }
 
 // Restore texture detail after denoising
-void ModulateLighting(FGBufferData GBuffer, float3 TranslatedWorldPosition, inout float3 DiffuseLighting, inout float3 SpecularLighting)
+void ModulateLighting(FManyLightMaterial Material, float3 TranslatedWorldPosition, inout float3 DiffuseLighting, inout float3 SpecularLighting)
 {
-	const float3 N = GBuffer.WorldNormal;
+	const float3 N = Material.WorldNormal;
 	const float3 V = normalize(View.TranslatedWorldCameraOrigin - TranslatedWorldPosition);
 	const float NoV = saturate(dot(N, V));
-	const float3 SpecularEnv = EnvBRDF(GBuffer.SpecularColor, GBuffer.Roughness, NoV);
+	const float3 SpecularEnv = EnvBRDF(Material.SpecularColor, Material.Roughness, NoV);
 
 	// Final pass outputs composites irradiance and outputs it to scene color
-	DiffuseLighting = DiffuseLighting * max(GBuffer.DiffuseColor, 0.001f);
+	DiffuseLighting = DiffuseLighting * max(Material.DiffuseColor, 0.001f);
 	SpecularLighting = SpecularLighting * max(SpecularEnv, 0.001f);
-}
\ No newline at end of file
+}
+
+FDeferredLightingSplit GetManyLightSplitLighting(
+	float3 TranslatedWorldPosition, 
+	float3 CameraVector, 
+	FManyLightMaterial Material, 
+	float AmbientOcclusion, 
+	FDeferredLightData LightData, 
+	float4 LightAttenuation, 
+	float Dither, 
+	uint2 ScreenCoord, 
+	inout float SurfaceShadow)
+#if SUBSTRATE_ENABLED
+{	
+	const uint2 PixelCoord = ScreenCoord;
+
+	FDeferredLightingSplit Out = (FDeferredLightingSplit)0;
+	FSubstrateAddressing SubstrateAddressing = GetSubstratePixelDataByteOffset(PixelCoord, uint2(View.BufferSizeAndInvSize.xy), Substrate.MaxBytesPerPixel);
+	FSubstratePixelHeader SubstratePixelHeader = UnpackSubstrateHeaderIn(Substrate.MaterialTextureArray, SubstrateAddressing, Substrate.TopLayerTexture);
+
+	BRANCH
+	if (SubstratePixelHeader.ClosureCount > 0) 
+	{
+		float3 V =-CameraVector;
+		float3 L = LightData.Direction;	// Already normalized
+		float3 ToLight = L;
+		float LightMask = 1;
+		if (LightData.bRadialLight)
+		{
+			LightMask = GetLocalLightAttenuation(TranslatedWorldPosition, LightData, ToLight, L);
+			#if ADAPTIVE_VOLUMETRIC_SHADOW_MAP
+			//LightAttenuation *= ComputeTransmittance(TranslatedWorldPosition, LightData.TranslatedWorldPosition, 256);
+			LightAttenuation *= AVSM_SampleTransmittance(TranslatedWorldPosition, LightData.TranslatedWorldPosition);
+			#endif // ADAPTIVE_VOLUMETRIC_SHADOW_MAP
+		}
+
+		if (LightMask > 0)
+		{
+			FSubstrateShadowTermInputParameters SubstrateShadowTermInputParameters	= GetInitialisedSubstrateShadowTermInputParameters();
+			SubstrateShadowTermInputParameters.bEvaluateShadowTerm					= true;
+			SubstrateShadowTermInputParameters.SceneDepth							= Material.Depth;
+			SubstrateShadowTermInputParameters.PrecomputedShadowFactors				= 1.f;  //SubstrateReadPrecomputedShadowFactors(SubstratePixelHeader, PixelCoord, SceneTexturesStruct.GBufferETexture);
+			SubstrateShadowTermInputParameters.TranslatedWorldPosition				= TranslatedWorldPosition;
+			SubstrateShadowTermInputParameters.LightAttenuation						= LightAttenuation;
+			SubstrateShadowTermInputParameters.Dither								= Dither;
+
+			FSubstrateDeferredLighting SubstrateLighting = SubstrateDeferredLighting(
+				LightData,
+				V,
+				L,
+				ToLight,
+				LightMask,
+				SubstrateShadowTermInputParameters,
+				Substrate.MaterialTextureArray,
+				SubstrateAddressing,
+				SubstratePixelHeader);
+
+			// SUBSTRATE_TODO - Add support for SUBSTRATE_OPAQUE_ROUGH_REFRACTION_ENABLED
+			Out.DiffuseLighting  = float4(SubstrateLighting.TotalDiffuseLighting, 0);
+			Out.SpecularLighting = float4(SubstrateLighting.TotalSpecularLighting, 0);
+		}
+	}
+	return Out;
+}
+#else // SUBSTRATE_ENABLED
+{
+	return GetDynamicLightingSplit(
+		TranslatedWorldPosition, CameraVector, Material.GBuffer, AmbientOcclusion, Material.GBuffer.ShadingModelID, 
+		LightData, LightAttenuation, Dither, ScreenCoord, 
+		SurfaceShadow);
+}
+#endif // SUBSTRATE_ENABLED
\ No newline at end of file
diff --git a/Engine/Shaders/Private/ManyLights/ManyLightsDenoiserSpatial.usf b/Engine/Shaders/Private/ManyLights/ManyLightsDenoiserSpatial.usf
--- a/Engine/Shaders/Private/ManyLights/ManyLightsDenoiserSpatial.usf
+++ b/Engine/Shaders/Private/ManyLights/ManyLightsDenoiserSpatial.usf
@@ -40,7 +40,7 @@ void DenoiserSpatialCS(
 
 		if (IsLightingValid(CenterDiffuseLightingAndSecondMoment.xyz))
 		{
-			FGBufferData GBuffer = GetGBufferData(ScreenUV);
+			FManyLightMaterial Material = LoadMaterial(ScreenUV, ScreenCoord);
 
 			float CenterDiffuseStdDev = sqrt(max(CenterDiffuseLightingAndSecondMoment.w - Pow2(Luminance(CenterDiffuseLightingAndSecondMoment.xyz)), 0.0f));
 			float CenterSpecularStdDev = sqrt(max(CenterSpecularLightingAndSecondMoment.w - Pow2(Luminance(CenterSpecularLightingAndSecondMoment.xyz)), 0.0f));
@@ -56,8 +56,8 @@ void DenoiserSpatialCS(
 
 			uint2 RandomSeed = Rand3DPCG16(int3(ScreenCoord, ManyLightsStateFrameIndex)).xy;
 
-			float3 TranslatedWorldPosition = GetTranslatedWorldPositionFromScreenUV(ScreenUV, GBuffer.Depth);
-			float4 ScenePlane = float4(GBuffer.WorldNormal, dot(TranslatedWorldPosition, GBuffer.WorldNormal));
+			float3 TranslatedWorldPosition = GetTranslatedWorldPositionFromScreenUV(ScreenUV, Material.Depth);
+			float4 ScenePlane = float4(Material.WorldNormal, dot(TranslatedWorldPosition, Material.WorldNormal));
 
 			// #ml_todo: expose as CVars
 			float MinStdDev = 0.1f;
@@ -75,15 +75,15 @@ void DenoiserSpatialCS(
 					{
 						// Depth weight
 						float2 NeighborScreenUV = (NeighborCoord + 0.5f) * View.BufferSizeAndInvSize.zw;
-						FGBufferData NeighborGBuffer = GetGBufferData(NeighborScreenUV);
-						const float3 NeighborTranslatedWorldPosition = GetTranslatedWorldPositionFromScreenUV(NeighborScreenUV, NeighborGBuffer.Depth);
+						const FManyLightMaterial NeighborMaterial = LoadMaterial(NeighborScreenUV, NeighborCoord);
+						const float3 NeighborTranslatedWorldPosition = GetTranslatedWorldPositionFromScreenUV(NeighborScreenUV, NeighborMaterial.Depth);
 						float PlaneDistance = abs(dot(float4(NeighborTranslatedWorldPosition, -1), ScenePlane));
-						float RelativeDepthDifference = PlaneDistance / GBuffer.Depth;
+						float RelativeDepthDifference = PlaneDistance / Material.Depth;
 						float DepthWeight = exp2(-SpatialFilterDepthWeightScale * (RelativeDepthDifference * RelativeDepthDifference));
 
 						// #sdl_todo: separate weight for specular accounting for roughness and expose as CVars
 						// Normal weight
-						float AngleBetweenNormals = acosFast(saturate(dot(ScenePlane.xyz, NeighborGBuffer.WorldNormal)));
+						float AngleBetweenNormals = acosFast(saturate(dot(ScenePlane.xyz, NeighborMaterial.WorldNormal)));
 						float NormalWeight = 1.0f - saturate(AngleBetweenNormals);
 
 						// Diffuse
@@ -125,12 +125,11 @@ void DenoiserSpatialCS(
 			// Final pass outputs composites irradiance and outputs it to scene color
 			float3 DiffuseLighting = DiffuseLightingSum / DiffuseLightingWeightSum;
 			float3 SpecularLighting = SpecularLightingSum / SpecularLightingWeightSum;
-			ModulateLighting(GBuffer, TranslatedWorldPosition, DiffuseLighting, SpecularLighting);
+			ModulateLighting(Material, TranslatedWorldPosition, DiffuseLighting, SpecularLighting);
 
 			// Composite diffuse and specular into scene color
-			const bool bNeedsSeparateSubsurfaceLightAccumulation = UseSubsurfaceProfile(GBuffer.ShadingModelID);
 			FLightAccumulator LightAccumulator = (FLightAccumulator)0;
-			LightAccumulator_AddSplit(LightAccumulator, DiffuseLighting, SpecularLighting, /*ScatterableLight*/ DiffuseLighting, /*CommonMultiplier*/ 1.0f, bNeedsSeparateSubsurfaceLightAccumulation);
+			LightAccumulator_AddSplit(LightAccumulator, DiffuseLighting, SpecularLighting, /*ScatterableLight*/ DiffuseLighting, /*CommonMultiplier*/ 1.0f, Material.bNeedsSeparateSubsurfaceLightAccumulation);
 			float4 AccumulatedSceneColor = LightAccumulator_GetResult(LightAccumulator);
 
 			#if DEBUG_MODE
diff --git a/Engine/Shaders/Private/ManyLights/ManyLightsDenoiserTemporal.usf b/Engine/Shaders/Private/ManyLights/ManyLightsDenoiserTemporal.usf
--- a/Engine/Shaders/Private/ManyLights/ManyLightsDenoiserTemporal.usf
+++ b/Engine/Shaders/Private/ManyLights/ManyLightsDenoiserTemporal.usf
@@ -216,10 +216,10 @@ void DenoiserTemporalCS(
 	if (all(ScreenCoord < View.ViewRectMinAndSize.xy + View.ViewRectMinAndSize.zw))
 	{
 		float2 ScreenUV = (ScreenCoord + 0.5f) * View.BufferSizeAndInvSize.zw;
-		FGBufferData GBuffer = GetGBufferData(ScreenUV);
+		FManyLightMaterial Material = LoadMaterial(ScreenUV, ScreenCoord);
 
 		const float DeviceZ = SceneDepthTexture[ScreenCoord].x;
-		GBuffer.Depth = ConvertFromDeviceZ(DeviceZ);
+		Material.SetDepth(ConvertFromDeviceZ(DeviceZ));
 
 		uint3 PackedLighting = SharedResolvedLightingYCoCg[GroupThreadId.x + SHARED_TILE_BORDER][GroupThreadId.y + SHARED_TILE_BORDER];
 		float3 DiffuseLighting = UnpackDiffuse(PackedLighting);
diff --git a/Engine/Shaders/Private/ManyLights/ManyLightsSampling.usf b/Engine/Shaders/Private/ManyLights/ManyLightsSampling.usf
--- a/Engine/Shaders/Private/ManyLights/ManyLightsSampling.usf
+++ b/Engine/Shaders/Private/ManyLights/ManyLightsSampling.usf
@@ -38,28 +38,21 @@ FLightTargetPDF InitLightTargetPDF()
 	return LightTargetPDF;
 }
 
-FLightTargetPDF GetLocalLightTargetPDF(FDeferredLightData LightData, float3 TranslatedWorldPosition, FGBufferData GBuffer, uint2 ScreenCoord, inout FDebug Debug)
+FLightTargetPDF GetLocalLightTargetPDF(FDeferredLightData LightData, float3 TranslatedWorldPosition, FManyLightMaterial Material, uint2 ScreenCoord, inout FDebug Debug)
 {
 	FLightTargetPDF LightTargetPDF = InitLightTargetPDF();
 
 	float3 CameraVector = normalize(TranslatedWorldPosition - View.TranslatedWorldCameraOrigin);
-	GBuffer.Roughness = max(GBuffer.Roughness, View.MinRoughness);
 
 	float4 LightAttenuation = 1.0f;
 	float Dither = 0.5f;
 	float SurfaceShadow = 1.0f;
 	float AmbientOcclusion = 1.0f;
 
-	// #ml_todo: Move more shading models to TILE_MODE_SIMPLE_SHADING
-	#if TILE_TYPE == TILE_MODE_SIMPLE_SHADING
-	{
-		GBuffer.ShadingModelID = SHADINGMODELID_DEFAULT_LIT;
-	}
-	#endif
 	LightData.ShadowedBits = 0;
 
-	FDeferredLightingSplit SplitLighting = GetDynamicLightingSplit(
-		TranslatedWorldPosition, CameraVector, GBuffer, AmbientOcclusion, GBuffer.ShadingModelID, 
+	FDeferredLightingSplit SplitLighting = GetManyLightSplitLighting(
+		TranslatedWorldPosition, CameraVector, Material, AmbientOcclusion, 
 		LightData, LightAttenuation, Dither, ScreenCoord, 
 		SurfaceShadow);
 
@@ -136,8 +129,8 @@ void GenerateLightSamplesCS(
 		{
 			uint2 ScreenCoord = DownsampledScreenCoordToScreenCoord(DownsampledScreenCoord);	
 			const float2 ScreenUV = (ScreenCoord + 0.5f) * View.BufferSizeAndInvSize.zw;
-			const FGBufferData GBuffer = GetGBufferData(ScreenUV);
-			const float SceneDepth = GBuffer.Depth;
+			const FManyLightMaterial Material = LoadMaterial(ScreenUV, ScreenCoord, TILE_TYPE == TILE_MODE_SIMPLE_SHADING/*bForceSimpleShading*/);
+			const float SceneDepth = Material.Depth;
 
 			float3 TranslatedWorldPosition = GetTranslatedWorldPositionFromScreenUV(ScreenUV, SceneDepth);		
 
@@ -199,8 +192,8 @@ void GenerateLightSamplesCS(
 			};
 
 			const float2 ScreenUV = (ScreenCoord + 0.5f) * View.BufferSizeAndInvSize.zw;
-			const FGBufferData GBuffer = GetGBufferData(ScreenUV);
-			const float SceneDepth = GBuffer.Depth;
+			const FManyLightMaterial Material = LoadMaterial(ScreenUV, ScreenCoord, TILE_TYPE == TILE_MODE_SIMPLE_SHADING/*bForceSimpleShading*/);
+			const float SceneDepth = Material.Depth;
 			float3 DiffuseSum = 0.0f;
 			float3 SpecularSum = 0.0f;
 			float WeightSum = 0.0f;
@@ -280,7 +273,7 @@ void GenerateLightSamplesCS(
 						const FLocalLightData LocalLightData = GetLocalLightDataNonStereo(LocalLightIndex);
 						FDeferredLightData LightData = ConvertToDeferredLight(LocalLightData);
 					
-						FLightTargetPDF LightTargetPDF = GetLocalLightTargetPDF(LightData, TranslatedWorldPosition, GBuffer, ScreenCoord, Debug);
+						FLightTargetPDF LightTargetPDF = GetLocalLightTargetPDF(LightData, TranslatedWorldPosition, Material, ScreenCoord, Debug);
 
 						#if DEBUG_MODE
 						if (Debug.bActive)
@@ -321,7 +314,7 @@ void GenerateLightSamplesCS(
 					}
 				}
 
-				DemodulateLighting(GBuffer, TranslatedWorldPosition, DiffuseSum, SpecularSum);
+				DemodulateLighting(Material, TranslatedWorldPosition, DiffuseSum, SpecularSum);
 
 				#if DEBUG_MODE
 				if (Debug.bActive)
@@ -363,7 +356,7 @@ void GenerateLightSamplesCS(
 			}
 
 			RWDownsampledSceneDepth[DownsampledScreenCoord] = SceneDepth;
-			RWDownsampledSceneWorldNormal[DownsampledScreenCoord] = EncodeNormal(GBuffer.WorldNormal);
+			RWDownsampledSceneWorldNormal[DownsampledScreenCoord] = EncodeNormal(Material.WorldNormal);
 			RWSampleLuminanceSum[DownsampledScreenCoord] = float2(Luminance(DiffuseSum), Luminance(SpecularSum));
 
 			for (uint LightSampleIndex = 0; LightSampleIndex < NUM_SAMPLES_PER_PIXEL_1D; ++LightSampleIndex)
diff --git a/Engine/Shaders/Private/ManyLights/ManyLightsShading.usf b/Engine/Shaders/Private/ManyLights/ManyLightsShading.usf
--- a/Engine/Shaders/Private/ManyLights/ManyLightsShading.usf
+++ b/Engine/Shaders/Private/ManyLights/ManyLightsShading.usf
@@ -135,7 +135,7 @@ void ShadeLightSamplesCS(
 	#endif
 
 			float2 ScreenUV = (ScreenCoord + 0.5f) * View.BufferSizeAndInvSize.zw;
-			FGBufferData GBuffer = GetGBufferData(ScreenUV);
+			FManyLightMaterial Material = LoadMaterial(ScreenUV, ScreenCoord, TILE_TYPE == TILE_MODE_SIMPLE_SHADING /*bForceSimpleShading*/);
 
 			// #ml_todo: handle this path when shadows are disabled
 			float ScreenSpaceAO = 1.0f;//Texture2DSampleLevel(SceneTexturesStruct.ScreenSpaceAOTexture, SceneTexturesStruct_ScreenSpaceAOTextureSampler, ScreenUV, 0).x;
@@ -143,12 +143,12 @@ void ShadeLightSamplesCS(
 			float3 DiffuseLighting = INVALID_LIGHTING;
 			float3 SpecularLighting = INVALID_LIGHTING;
 
-			if (GBuffer.Depth > 0.0f && GBuffer.ShadingModelID != SHADINGMODELID_UNLIT)
+			if (Material.Depth > 0.0f && Material.IsValid())
 			{
 				DiffuseLighting = 0.0f;
 				SpecularLighting = 0.0f;
 
-				float3 TranslatedWorldPosition = GetTranslatedWorldPositionFromScreenUV(ScreenUV, GBuffer.Depth);		
+				float3 TranslatedWorldPosition = GetTranslatedWorldPositionFromScreenUV(ScreenUV, Material.Depth);
 				float3 CameraVector = normalize(TranslatedWorldPosition - View.TranslatedWorldCameraOrigin);
 				float4 DownsampledScreenWeights = CompositeUpsampleWeights[ScreenCoord];
 
@@ -213,23 +213,14 @@ void ShadeLightSamplesCS(
 
 						if (SampleWeight > 0.0f)
 						{
-							GBuffer.Roughness = max(GBuffer.Roughness, View.MinRoughness);
-
 							float4 LightAttenuation = 1.0f;
 							float Dither = 0.5f;
 							float SurfaceShadow = 1;
 							float AmbientOcclusion = ScreenSpaceAO;
-
-							// #ml_todo: Move more shading models to TILE_MODE_SIMPLE_SHADING
-							#if TILE_TYPE == TILE_MODE_SIMPLE_SHADING
-							{
-								GBuffer.ShadingModelID = SHADINGMODELID_DEFAULT_LIT;
-							}
-							#endif
 							LightData.ShadowedBits = 0;
 
-							FDeferredLightingSplit SplitLighting = GetDynamicLightingSplit(
-								TranslatedWorldPosition, CameraVector, GBuffer, AmbientOcclusion, GBuffer.ShadingModelID, 
+							FDeferredLightingSplit SplitLighting = GetManyLightSplitLighting(
+								TranslatedWorldPosition, CameraVector, Material, AmbientOcclusion, 
 								LightData, LightAttenuation, Dither, ScreenCoord, 
 								SurfaceShadow);
 
@@ -251,7 +242,7 @@ void ShadeLightSamplesCS(
 				DiffuseLighting *= View.PreExposure;
 				SpecularLighting *= View.PreExposure;
 
-				DemodulateLighting(GBuffer, TranslatedWorldPosition, DiffuseLighting, SpecularLighting);
+				DemodulateLighting(Material, TranslatedWorldPosition, DiffuseLighting, SpecularLighting);
 			}
 
 			#if DEBUG_MODE
diff --git a/Engine/Shaders/Private/Substrate/SubstrateDeferredLighting.ush b/Engine/Shaders/Private/Substrate/SubstrateDeferredLighting.ush
--- a/Engine/Shaders/Private/Substrate/SubstrateDeferredLighting.ush
+++ b/Engine/Shaders/Private/Substrate/SubstrateDeferredLighting.ush
@@ -159,6 +159,9 @@ FSubstrateDeferredLighting SubstrateDeferredLighting(
 		FLightAccumulator Out = (FLightAccumulator)0;
 		LightAccumulator_AddSplit(Out, DiffuseLuminance, SpecularLuminance, DiffuseLuminance, CommonMultiplier, BSDFEvaluate.bPostProcessSubsurface);
 		AccumulateSubstrateDeferredLighting(SubstrateLighting, Out, BSDFEvaluate.bPostProcessSubsurface, BSDF_GETISTOPLAYER(BSDF));
+
+		SubstrateLighting.TotalDiffuseLighting += Out.TotalLightDiffuse;
+		SubstrateLighting.TotalSpecularLighting += Out.TotalLightSpecular;
 	}
 
 	return SubstrateLighting;
diff --git a/Engine/Source/Runtime/Renderer/Private/ManyLights/ManyLightsInternal.h b/Engine/Source/Runtime/Renderer/Private/ManyLights/ManyLightsInternal.h
--- a/Engine/Source/Runtime/Renderer/Private/ManyLights/ManyLightsInternal.h
+++ b/Engine/Source/Runtime/Renderer/Private/ManyLights/ManyLightsInternal.h
@@ -24,8 +24,9 @@ BEGIN_SHADER_PARAMETER_STRUCT(FManyLightsParameters, )
 	SHADER_PARAMETER(FIntPoint, DownsampledViewSize)
 	SHADER_PARAMETER(FIntPoint, NumSamplesPerPixel)
 	SHADER_PARAMETER(FIntPoint, NumSamplesPerPixelDivideShift)
-	SHADER_PARAMETER(uint32, ManyLightsStateFrameIndex)
 	SHADER_PARAMETER(FVector2f, DownsampledBufferInvSize)
+	SHADER_PARAMETER(uint32, DownsampleFactor)
+	SHADER_PARAMETER(uint32, ManyLightsStateFrameIndex)
 	SHADER_PARAMETER(float, SamplingMinWeight)
 	SHADER_PARAMETER(int32, TileDataStride)
 	SHADER_PARAMETER(int32, DownsampledTileDataStride)
