Fix shader compilation warning.

diff --git a/Engine/Shaders/Private/HeterogeneousVolumes/HeterogeneousVolumesAdaptiveVolumetricShadowMapSampler.ush b/Engine/Shaders/Private/HeterogeneousVolumes/HeterogeneousVolumesAdaptiveVolumetricShadowMapSampler.ush
--- a/Engine/Shaders/Private/HeterogeneousVolumes/HeterogeneousVolumesAdaptiveVolumetricShadowMapSampler.ush
+++ b/Engine/Shaders/Private/HeterogeneousVolumes/HeterogeneousVolumesAdaptiveVolumetricShadowMapSampler.ush
@@ -9,8 +9,8 @@
 #define AVSM_BILINEAR_INTERPOLATION 1
 #endif //
 
-#define INVALID_SAMPLE_INDEX 0xFF
-#define INVALID_SAMPLE_INDEX4 0xFFFFFFFF
+#define INVALID_SAMPLE_INDEX 0xFFu
+#define INVALID_SAMPLE_INDEX4 0xFFFFFFFFu
 
 #if 0
 // Reference single-channel implementation
diff --git a/Engine/Shaders/Private/Substrate/Substrate.ush b/Engine/Shaders/Private/Substrate/Substrate.ush
--- a/Engine/Shaders/Private/Substrate/Substrate.ush
+++ b/Engine/Shaders/Private/Substrate/Substrate.ush
@@ -2256,7 +2256,7 @@ void FSubstratePixelHeader::SubstrateConvertFromDBuffer(
 	OutData = GetSubstrateSlabBSDF(
 		PixelFootprint,		// PixelFootprint
 		ComputeDiffuseAlbedo(In.BaseColor, In.Metallic),
-		ComputeF0(In.Specular, In.BaseColor, In.Metallic),
+		ComputeF0(half(In.Specular), half3(In.BaseColor), half(In.Metallic)),
 		1.0f,				// F90
 		In.Roughness,		// Roughness
 		0.f,				// Anisotropy
