RayTraced Shadows: Fix incorrect shadowing for lights with non-zero source length

When the source length is small, there is a small probability of generating invalid samples which must be treated as shadowed.

When the shadow denoiser is enabled, these samples were still traced and incorrectly treated as illuminated. Make sure the ray is invalid so that it is correctly ignored.

#tests Verified on test scene from Jira

diff --git a/Engine/Shaders/Private/RayTracing/RayTracingCapsuleLight.ush b/Engine/Shaders/Private/RayTracing/RayTracingCapsuleLight.ush
--- a/Engine/Shaders/Private/RayTracing/RayTracingCapsuleLight.ush
+++ b/Engine/Shaders/Private/RayTracing/RayTracingCapsuleLight.ush
@@ -40,11 +40,24 @@ bool GenerateCapsuleLightOcclusionRayWithSolidAngleSampling(
 	float4 Result = SampleCapsuleBounds(CapsuleBounds, RandSample);
 	// compute distance to capsule (and check that we hit it)
 	float Distance = CapsuleTest(Result.xyz, LightDirection, Axis, Radius2, SourceLength);
-	RayOrigin = TranslatedWorldPosition;
-	RayDirection = Result.xyz;
-	RayTMin = 0;
-	RayTMax = Distance;
-	RayPdf = Result.w;
-	return Distance > 0; // sample is only valid if we hit the light
+	// sample is only valid if we hit the light
+	if (Distance > 0)
+	{
+		RayOrigin = TranslatedWorldPosition;
+		RayDirection = Result.xyz;
+		RayTMin = 0;
+		RayTMax = Distance;
+		RayPdf = Result.w;
+		return true;
+	}
+	else
+	{
+		RayOrigin = 0;
+		RayDirection = 0;
+		RayTMin = 0;
+		RayTMax = 0;
+		RayPdf = 0;
+		return false;
+	}
 #endif
 }
diff --git a/Engine/Shaders/Private/RayTracing/RayTracingLightingCommon.ush b/Engine/Shaders/Private/RayTracing/RayTracingLightingCommon.ush
--- a/Engine/Shaders/Private/RayTracing/RayTracingLightingCommon.ush
+++ b/Engine/Shaders/Private/RayTracing/RayTracingLightingCommon.ush
@@ -389,6 +389,10 @@ float3 ComputeDirectLighting(
 			// hard shadow or no shadow -- just trace toward light center
 			ShadowRayLength = length(ToLight);
 		}
+		if (ShadowRayLength == 0)
+		{
+			Lit = 0;
+		}
 
 		const bool ApplyShadow = (Lit && ReflectedShadowsType != 0);
 
