Changes to the TSR's shading measurement to non gama encoded color space that was causing it to be overprecise in the shadows

This introduces a linear Shading Measurement Color Space closely fitting ACES Film curve, killing off the need for r.TSR.ShadingRejection.ExposureOffset.
The TSR's history guide is also switched to a gama encoded Guide Color Space that is simply sqrt(SMCS) for the QuantizeForUNormRenderTarget() to operate within the MeasureBackbufferLDRQuantizationError() to reduce flickering cause by it.

diff --git a/Engine/Extras/GPUDumpViewer/GPUDumpViewer.html b/Engine/Extras/GPUDumpViewer/GPUDumpViewer.html
--- a/Engine/Extras/GPUDumpViewer/GPUDumpViewer.html
+++ b/Engine/Extras/GPUDumpViewer/GPUDumpViewer.html
@@ -1894,9 +1894,10 @@ FloatInfos decodeFloat(float x)
 #define isinf(x) decodeFloat(x).is_infinity
 #define isdenorm(x) decodeFloat(x).is_denormal
 
-vec3 LinearToGamma22(vec3 Color)
+// https://www.shadertoy.com/view/7tKXWR
+vec3 LinearToWebGLDisplayGamma(vec3 Color)
 {
-	return pow(Color, vec3(2.2));
+	return pow(Color, vec3(1.0/2.2));
 }
 
 vec3 LinearToRec709(vec3 lin) 
@@ -1929,6 +1930,17 @@ vec3 SRGBToLinear(vec3 Color)
 	return Color;
 }
 
+// https://knarkowicz.wordpress.com/2016/01/06/aces-filmic-tone-mapping-curve/
+vec3 ACESFilm(vec3 x)
+{
+	float a = 2.51f;
+	float b = 0.03f;
+	float c = 2.43f;
+	float d = 0.59f;
+	float e = 0.14f;
+	return clamp((x*(a*x+b))/(x*(c*x+d)+e), 0.0, 1.0);
+}
+
 `;
 	
 	var gl_shader = gl.createShader(shader_type);
@@ -2121,7 +2133,7 @@ class TextureView extends ResourceView
 		this.viewport_selected = false;
 		this.is_draging_canvas = false;
 		this.src_channels = 'RGB';
-		this.src_color_space = 'Gamut=sRGB Gamma 2.2';
+		this.src_color_space = 'Gamut=sRGB Pixel Values';
 		this.display_mode = 'Visualization';
 	}
 
@@ -2185,10 +2197,11 @@ class TextureView extends ResourceView
 						</div>
 						<div class="button_list" id="texture_visualization_change_src_color"><!---
 							---><span>Src color:</span><!---
-							---><input type="button" onclick="g_view.resource_view.change_src_color(this);" value="Gamut=sRGB Gamma 2.2" title="Performs no modifications to display to WebGL's sRGB gamut backbuffer (default)."/><!---
-							---><input type="button" onclick="g_view.resource_view.change_src_color(this);" value="Gamut=sRGB Linear[0;1]" title="Applies a LinearToGamma22(display.rgb); to output WebGL's sRGB gamut backbuffer."/><!---
-							---><input type="button" onclick="g_view.resource_view.change_src_color(this);" value="sRGB" title="Applies a LinearToGamma22(SRGBToLinear(display.rgba)) To output to WebGL's sRGB gamut backbuffer."/><!---
-							---><input type="button" onclick="g_view.resource_view.change_src_color(this);" value="Rec709" title="Applies LinearToGamma22(Rec709ToLinear(display.rgb)) to output WebGL's sRGB gamut backbuffer."/>
+							---><input type="button" onclick="g_view.resource_view.change_src_color(this);" value="Gamut=sRGB Pixel Values" title="Performs no modifications to display to WebGL's sRGB gamut backbuffer (default)."/><!---
+							---><input type="button" onclick="g_view.resource_view.change_src_color(this);" value="Gamut=sRGB Linear[0;1]" title="Applies a LinearToWebGLDisplayGamma(display.rgb); to output WebGL's sRGB gamut backbuffer."/><!---
+							---><input type="button" onclick="g_view.resource_view.change_src_color(this);" value="Gamut=sRGB Linear ACES Film" title="Applies a LinearToWebGLDisplayGamma(ACESFilm(display.rgb)); to output WebGL's sRGB gamut backbuffer."/><!---
+							---><input type="button" onclick="g_view.resource_view.change_src_color(this);" value="sRGB" title="Applies a LinearToWebGLDisplayGamma(SRGBToLinear(display.rgba)) To output to WebGL's sRGB gamut backbuffer."/><!---
+							---><input type="button" onclick="g_view.resource_view.change_src_color(this);" value="Rec709" title="Applies LinearToWebGLDisplayGamma(Rec709ToLinear(display.rgb)) to output WebGL's sRGB gamut backbuffer."/>
 						</div>
 						<div class="button_list" id="texture_visualization_change_display_modes"><!---
 							---><input type="button" onclick="g_view.resource_view.change_display_mode(this);" value="Visualization" title="Display the WebGL's visualization shader's display.rgb."/><!---
@@ -3783,21 +3796,25 @@ void main(void) {
 			}
 
 			// verify_backbuffer_color_space() check webgl is a sRGB gamut backbuffer
-			if (this.src_color_space == 'Gamut=sRGB Gamma 2.2')
+			if (this.src_color_space == 'Gamut=sRGB Pixel Values')
 			{
-				// NOP because verify_backbuffer_color_space() check webgl is a sRGB backbuffer
+				// NOP because verify_backbuffer_color_space() check webgl is a sRGB backbuffer, that directly export pixel values
 			}
 			else if (this.src_color_space == 'Gamut=sRGB Linear[0;1]')
 			{
-				final_frag_code += 'display.rgb = LinearToGamma22(display.rgb);';
+				final_frag_code += 'display.rgb = LinearToWebGLDisplayGamma(display.rgb);';
+			}
+			else if (this.src_color_space == 'Gamut=sRGB Linear ACES Film')
+			{
+				final_frag_code += 'display.rgb = LinearToWebGLDisplayGamma(ACESFilm(display.rgb));';
 			}
 			else if (this.src_color_space == 'sRGB')
 			{
-				final_frag_code += 'display.rgb = LinearToGamma22(SRGBToLinear(display.rgb));';
+				final_frag_code += 'display.rgb = LinearToWebGLDisplayGamma(SRGBToLinear(display.rgb));';
 			}
 			else if (this.src_color_space == 'Rec709')
 			{
-				final_frag_code += 'display.rgb = LinearToGamma22(Rec709ToLinear(display.rgb));';
+				final_frag_code += 'display.rgb = LinearToWebGLDisplayGamma(Rec709ToLinear(display.rgb));';
 			}
 			final_frag_code += '}';
 
diff --git a/Engine/Shaders/Private/TemporalSuperResolution/TSRColorSpace.ush b/Engine/Shaders/Private/TemporalSuperResolution/TSRColorSpace.ush
new file mode 100644
--- /dev/null
+++ b/Engine/Shaders/Private/TemporalSuperResolution/TSRColorSpace.ush
@@ -0,0 +1,286 @@
+// Copyright Epic Games, Inc. All Rights Reserved.
+
+#pragma once
+
+#include "TSRCommon.ush"
+
+
+//------------------------------------------------------- DEFINITION
+
+/**
+ * SMCS = Shading Measurement Color Space:
+ *  Goal of this color space is to match tone mapping curve as much as possible.
+ *  Note we want the shading measurement in a linear color as oposed to gama encoded color, to not increase the rejection measurements in the shadows.
+ *
+ *  UE's tonemaping curve is quite a bit complicated, so here is an aproximation:
+ *
+ *  // https://knarkowicz.wordpress.com/page/2/
+ *  float3 ACESFilm(float3 x)
+ *  {
+ *      float a = 2.51f;
+ *      float b = 0.03f;
+ *      float c = 2.43f;
+ *      float d = 0.59f;
+ *      float e = 0.14f;
+ *      return saturate((x * (a * x + b)) / (x * (c * x + d) + e));
+ *  }
+ *
+ *  One thing thaw is since TSR is running before bloom, wants a bit more precision in the highlight than ACES curve for bloom stability.
+ *  So the Linear -> SMCS curve chosen is SMCS(Linear) -> (Linear / (Linear + 0.17))^2
+ *
+ * GSC = Guide Color Space:
+ *  Color space used to encode SMCS in the TSR.History.Guide. Similarily to the back buffer being gama encoded, GSC is also gamma encoded to keep some precision in the shadows.
+ *  It also ensures the error introduced by the QuantizeForUNormRenderTarget() distributed through the gama curve as tonemapper does in the backbuffer.
+ *
+ *  For performance motivation of the GCSToSMCS() used in the RejectShading pass, we take SMCS(GSC) = GSC^2, leading to GSC(Linear) -> sqrt(SMCS(Linear)) -> Linear / (Linear + 0.17)
+ */
+
+
+//------------------------------------------------------- CONSTANTS
+
+// Largest encodable scene color value to not overflaw in Luma4() with halfs on console
+static const tsr_half3 LargestSceneColorRGB  = (LargestNormalNumber * tsr_half(0.25)).xxx;
+static const tsr_half4 LargestSceneColorRGBA = (LargestNormalNumber * tsr_half(0.25)).xxxx;
+
+// Weight of each channel for the computation of luminance.
+static const tsr_half kMoireLumaToChannel = tsr_half(3.0);
+static const tsr_half3 kMoireSMCSWeights = rcp(kMoireLumaToChannel).xxx;
+
+static const float kHistoryAccumulationPreceptionAdd = 1.0;
+static const float kGCSPreceptionAdd = 0.17;
+
+// Some bright pixel can cause HdrWeight to get nullified under fp16 representation. So clamping this to a value close to the minimum float float positive value (0.000061).
+#define HDR_WEIGHT_SAFE_MIN_VALUE 0.0001
+
+
+//------------------------------------------------------- LINEAR <-> GUIDE COLOR SPACE
+
+CALL_SITE_DEBUGLOC
+tsr_half3 LinearToGCS(tsr_half3 LinearColor)
+{
+	tsr_half3 GColor = LinearColor  * rcp(LinearColor + tsr_half(kGCSPreceptionAdd));
+	return GColor;
+}
+
+CALL_SITE_DEBUGLOC
+tsr_half4 LinearToGCS(tsr_half4 LinearColor)
+{
+	return tsr_half4(LinearToGCS(LinearColor.rgb), LinearColor.a);
+}
+
+CALL_SITE_DEBUGLOC
+tsr_half3 GCSToLinear(tsr_half3 GColor)
+{
+	tsr_half3 Minus = tsr_half(1.0) - GColor;
+	tsr_half3 LDRToLinear = min(tsr_half(kGCSPreceptionAdd) * rcp(Minus), tsr_half(MaxHalfFloat));
+	tsr_half3 LinearColor = GColor * LDRToLinear;
+	return LinearColor;
+}
+
+CALL_SITE_DEBUGLOC
+tsr_half4 GCSToLinear(tsr_half4 GColor)
+{
+	return tsr_half4(GCSToLinear(GColor.rgb), GColor.a);
+}
+
+CALL_SITE_DEBUGLOC
+tsr_half3 PreExposureCorrectGCS(tsr_half3 GColor, const float PreExposureCorrection)
+{
+	return LinearToGCS(GCSToLinear(GColor) * tsr_half(PreExposureCorrection));
+}
+
+CALL_SITE_DEBUGLOC
+tsr_half4 PreExposureCorrectGCS(tsr_half4 GColor, const float PreExposureCorrection)
+{
+	return tsr_half4(PreExposureCorrectGCS(GColor.rgb, PreExposureCorrection), GColor.a);
+}
+
+
+//------------------------------------------------------- GUIDE COLOR SPACE <-> SHADING MEASUREMENT COLOR SPACE
+
+CALL_SITE_DEBUGLOC
+tsr_half3 GCSToSMCS(tsr_half3 GColor)
+{
+	tsr_half3 SMColor = GColor * GColor;
+	return SMColor;
+}
+
+CALL_SITE_DEBUGLOC
+tsr_half4 GCSToSMCS(tsr_half4 GColor)
+{
+	return tsr_half4(GCSToSMCS(GColor.rgb), GColor.a);
+}
+
+CALL_SITE_DEBUGLOC
+tsr_half3 SMCSToGCS(tsr_half3 SMColor)
+{
+	tsr_half3 GColor = sqrt(SMColor);
+	return GColor;
+}
+
+CALL_SITE_DEBUGLOC
+tsr_half4 SMCSToGCS(tsr_half4 GColor)
+{
+	return tsr_half4(SMCSToGCS(GColor.rgb), GColor.a);
+}
+
+
+//------------------------------------------------------- LINEAR <-> SHADING MEASUREMENT COLOR SPACE
+
+CALL_SITE_DEBUGLOC
+tsr_half3 LinearToSMCS(tsr_half3 LinearColor)
+{
+	return GCSToSMCS(LinearToGCS(LinearColor));
+}
+
+CALL_SITE_DEBUGLOC
+tsr_half4 LinearToSMCS(tsr_half4 LinearColor)
+{
+	return tsr_half4(LinearToSMCS(LinearColor.rgb), LinearColor.a);
+}
+
+CALL_SITE_DEBUGLOC
+tsr_half3 SMCSToLinear(tsr_half3 SMColor)
+{
+	return GCSToLinear(SMCSToGCS(SMColor));
+}
+
+CALL_SITE_DEBUGLOC
+tsr_half4 SMCSToLinear(tsr_half4 SMColor)
+{
+	return tsr_half4(SMCSToLinear(SMColor.rgb), SMColor.a);
+}
+
+
+//------------------------------------------------------- KARIS HDR WEIGHTING
+
+// Faster but less accurate luma computation. 
+// Luma includes a scaling by 4.
+CALL_SITE_DEBUGLOC
+tsr_half2 Luma4(tsr_half3x2 Color)
+{
+	return (Color[1] * tsr_half(2.0)) + (Color[0] + Color[2]);
+}
+
+CALL_SITE_DEBUGLOC
+tsr_half2 Luma4(tsr_half4x2 Color)
+{
+	return Luma4(tsr_half3x2(Color[0], Color[1], Color[2]));
+}
+
+CALL_SITE_DEBUGLOC
+tsr_half Luma4(tsr_half3 Color)
+{
+	return dpv_lo(Luma4(dpv_interleave_mono_registers(Color)));
+}
+
+CALL_SITE_DEBUGLOC
+tsr_half Luma4(tsr_half4 Color)
+{
+	return Luma4(Color.rgb);
+}
+
+CALL_SITE_DEBUGLOC
+tsr_half2 HdrWeightY(tsr_half2 Luma)
+{
+	tsr_half Exposure = tsr_half(1.0);
+
+	return max(tsr_half(HDR_WEIGHT_SAFE_MIN_VALUE), rcp(Luma * Exposure + tsr_half(4.0)));
+}
+
+CALL_SITE_DEBUGLOC
+tsr_half HdrWeightY(tsr_half Luma)
+{
+	return dpv_lo(HdrWeightY(dpv_interleave_mono_registers(Luma)));
+}
+
+CALL_SITE_DEBUGLOC
+tsr_half2 HdrWeightInvY(tsr_half2 LDRLuma) 
+{
+	return tsr_half(4.0) * rcp(tsr_half(1.0) - LDRLuma);
+}
+
+CALL_SITE_DEBUGLOC
+tsr_half HdrWeightInvY(tsr_half LDRLuma) 
+{
+	return dpv_lo(HdrWeightInvY(dpv_interleave_mono_registers(LDRLuma)));
+}
+
+CALL_SITE_DEBUGLOC
+tsr_half2 HdrWeight4(tsr_half3x2 Color)
+{
+	return HdrWeightY(Luma4(Color));
+}
+
+CALL_SITE_DEBUGLOC
+tsr_half2 HdrWeight4(tsr_half4x2 Color)
+{
+	return HdrWeightY(Luma4(Color));
+}
+
+CALL_SITE_DEBUGLOC
+tsr_half HdrWeight4(tsr_half3 Color)
+{
+	return HdrWeightY(Luma4(Color));
+}
+
+CALL_SITE_DEBUGLOC
+tsr_half HdrWeight4(tsr_half4 Color)
+{
+	return HdrWeightY(Luma4(Color));
+}
+
+
+//------------------------------------------------------- BACK BUFFER ERROR
+
+CALL_SITE_DEBUGLOC
+tsr_half MeasureBackbufferLDRQuantizationError()
+{
+	// Assume the backbuffer is 10bit per channels
+	return tsr_half(0.5 / 1024.0);
+}
+
+
+//------------------------------------------------------- STOCASTIC QUANTIZATION
+
+CALL_SITE_DEBUGLOC
+tsr_half3x2 QuantizeForUNormRenderTarget(tsr_half3x2 Color, tsr_half E, const float3 QuantizationError)
+{
+	return Color + dpv_interleave_mono_registers(tsr_half3(QuantizationError)) * (E - tsr_half(0.5));
+}
+
+CALL_SITE_DEBUGLOC
+tsr_half3 QuantizeForUNormRenderTarget(tsr_half3 Color, tsr_half E, const float3 QuantizationError)
+{
+	return dpv_lo(QuantizeForUNormRenderTarget(dpv_interleave_mono_registers(Color), E, QuantizationError));
+}
+
+CALL_SITE_DEBUGLOC
+tsr_half3x2 QuantizeForFloatRenderTarget(tsr_half3x2 Color, tsr_half E, const float3 QuantizationError)
+{
+	tsr_half3x2 Error = dpv_mul(Color, tsr_half3(QuantizationError));
+
+	#if 0
+		// NOP
+	#elif CONFIG_COMPILE_FP16
+	{
+		Error[0] = asfloat16(asint16(Error[0]) & uint16_t(~0x03FF));
+		Error[1] = asfloat16(asint16(Error[1]) & uint16_t(~0x03FF));
+		Error[2] = asfloat16(asint16(Error[2]) & uint16_t(~0x03FF));
+	}
+	#else
+	{
+		Error[0] = asfloat(asuint(Error[0]) & ~0x007FFFFF);
+		Error[1] = asfloat(asuint(Error[1]) & ~0x007FFFFF);
+		Error[2] = asfloat(asuint(Error[2]) & ~0x007FFFFF);
+	}
+	#endif
+	
+	return Color + Error * E;
+}
+
+CALL_SITE_DEBUGLOC
+tsr_half3 QuantizeForFloatRenderTarget(tsr_half3 Color, tsr_half E, const float3 QuantizationError)
+{
+	return dpv_lo(QuantizeForFloatRenderTarget(dpv_interleave_mono_registers(Color), E, QuantizationError));
+}
diff --git a/Engine/Shaders/Private/TemporalSuperResolution/TSRCommon.ush b/Engine/Shaders/Private/TemporalSuperResolution/TSRCommon.ush
--- a/Engine/Shaders/Private/TemporalSuperResolution/TSRCommon.ush
+++ b/Engine/Shaders/Private/TemporalSuperResolution/TSRCommon.ush
@@ -235,16 +235,6 @@ static const int2 kOffsetSW = int2(-1, +1);
 // Largest encodable normal number in a half used on console.
 static const tsr_half LargestNormalNumber = tsr_half(65504.0);
 
-// Largest encodable scene color value to not overflaw in YCoCg with halfs on console
-static const tsr_half3 LargestSceneColorRGB = (LargestNormalNumber * tsr_half(0.25)).xxx;
-static const tsr_half4 LargestSceneColorRGBA = (LargestNormalNumber * tsr_half(0.25)).xxxx;
-static const tsr_half3 LargestSceneColorYCoCg = LargestNormalNumber.xxx;
-static const tsr_half4 LargestSceneColorYCoCgA = LargestNormalNumber.xxxx;
-
-// Weight of each channel for the computation of luminance.
-static const tsr_half kMoireLumaToChannel = tsr_half(3.0);
-static const tsr_half3 kMoireLumaWeights = rcp(kMoireLumaToChannel).xxx;
-
 
 //------------------------------------------------------- PARAMETERS
 
@@ -656,87 +646,6 @@ tsr_half ComputeLowFrequencyContribution(tsr_half HistoryValidity)
 	return dpv_lo(ComputeLowFrequencyContribution(dpv_interleave_mono_registers(HistoryValidity)));
 }
 
-// Some bright pixel can cause HdrWeight to get nullified under fp16 representation. So clamping this to a value close to the minimum float float positive value (0.000061).
-#define HDR_WEIGHT_SAFE_MIN_VALUE 0.0001
-
-// Faster but less accurate luma computation. 
-// Luma includes a scaling by 4.
-CALL_SITE_DEBUGLOC
-tsr_half2 Luma4(tsr_half3x2 Color)
-{
-	return (Color[1] * tsr_half(2.0)) + (Color[0] + Color[2]);
-}
-
-CALL_SITE_DEBUGLOC
-tsr_half2 Luma4(tsr_half4x2 Color)
-{
-	return Luma4(tsr_half3x2(Color[0], Color[1], Color[2]));
-}
-
-CALL_SITE_DEBUGLOC
-tsr_half Luma4(tsr_half3 Color)
-{
-	return dpv_lo(Luma4(dpv_interleave_mono_registers(Color)));
-}
-
-CALL_SITE_DEBUGLOC
-tsr_half Luma4(tsr_half4 Color)
-{
-	return Luma4(Color.rgb);
-}
-
-CALL_SITE_DEBUGLOC
-tsr_half2 HdrWeightY(tsr_half2 Luma)
-{
-	tsr_half Exposure = tsr_half(1.0);
-
-	return max(tsr_half(HDR_WEIGHT_SAFE_MIN_VALUE), rcp(Luma * Exposure + tsr_half(4.0)));
-}
-
-CALL_SITE_DEBUGLOC
-tsr_half HdrWeightY(tsr_half Luma)
-{
-	return dpv_lo(HdrWeightY(dpv_interleave_mono_registers(Luma)));
-}
-
-CALL_SITE_DEBUGLOC
-tsr_half2 HdrWeightInvY(tsr_half2 LDRLuma) 
-{
-	return tsr_half(4.0) * rcp(tsr_half(1.0) - LDRLuma);
-}
-
-CALL_SITE_DEBUGLOC
-tsr_half HdrWeightInvY(tsr_half LDRLuma) 
-{
-	return dpv_lo(HdrWeightInvY(dpv_interleave_mono_registers(LDRLuma)));
-}
-
-// Optimized HDR weighting function.
-CALL_SITE_DEBUGLOC
-tsr_half2 HdrWeight4(tsr_half3x2 Color)
-{
-	return HdrWeightY(Luma4(Color));
-}
-
-CALL_SITE_DEBUGLOC
-tsr_half2 HdrWeight4(tsr_half4x2 Color)
-{
-	return HdrWeightY(Luma4(Color));
-}
-
-CALL_SITE_DEBUGLOC
-tsr_half HdrWeight4(tsr_half3 Color)
-{
-	return HdrWeightY(Luma4(Color));
-}
-
-CALL_SITE_DEBUGLOC
-tsr_half HdrWeight4(tsr_half4 Color)
-{
-	return HdrWeightY(Luma4(Color));
-}
-
-
 // Returns the weight of a pixels at a coordinate <PixelDelta> from the PDF highest point.
 CALL_SITE_DEBUGLOC
 tsr_half2 ComputeSampleWeigth(tsr_half2 UpscaleFactor, tsr_half2x2 PixelDelta, const float MinimalContribution)
@@ -809,67 +718,6 @@ void IsOffScreenOrDisoccluded(uint bCameraCut, float2x2 ScreenPos, bool2 bIsPara
 	bIsDisoccluded = and(!bIsOffScreen, bIsParallaxRejected);
 }
 
-tsr_half MeasureBackbufferLDRQuantizationError()
-{
-	// Assume the backbuffer is 10bit per channels
-	return tsr_half(0.5 / 1024.0);
-}
-
-tsr_half2 MeasureBackbufferLDRQuantizationErrorFromHDRLuma(tsr_half2 Luma)
-{
-	tsr_half2 Tonemap = Luma * HdrWeightY(Luma) + MeasureBackbufferLDRQuantizationError().xx;
-	tsr_half2 PixelColorQuantizationError = abs(Tonemap * HdrWeightInvY(Tonemap) - Luma);
-
-	return PixelColorQuantizationError;
-}
-
-tsr_half MeasureBackbufferLDRQuantizationErrorFromHDRLuma(tsr_half Luma)
-{
-	return dpv_lo(MeasureBackbufferLDRQuantizationErrorFromHDRLuma(dpv_interleave_mono_registers(Luma)));
-}
-
-CALL_SITE_DEBUGLOC
-tsr_half3x2 QuantizeForUNormRenderTarget(tsr_half3x2 Color, tsr_half E, const float3 QuantizationError)
-{
-	return Color + dpv_interleave_mono_registers(tsr_half3(QuantizationError)) * (E - tsr_half(0.5));
-}
-
-CALL_SITE_DEBUGLOC
-tsr_half3 QuantizeForUNormRenderTarget(tsr_half3 Color, tsr_half E, const float3 QuantizationError)
-{
-	return dpv_lo(QuantizeForUNormRenderTarget(dpv_interleave_mono_registers(Color), E, QuantizationError));
-}
-
-CALL_SITE_DEBUGLOC
-tsr_half3x2 QuantizeForFloatRenderTarget(tsr_half3x2 Color, tsr_half E, const float3 QuantizationError)
-{
-	tsr_half3x2 Error = dpv_mul(Color, tsr_half3(QuantizationError));
-
-	#if 0
-		// NOP
-	#elif CONFIG_COMPILE_FP16
-	{
-		Error[0] = asfloat16(asint16(Error[0]) & uint16_t(~0x03FF));
-		Error[1] = asfloat16(asint16(Error[1]) & uint16_t(~0x03FF));
-		Error[2] = asfloat16(asint16(Error[2]) & uint16_t(~0x03FF));
-	}
-	#else
-	{
-		Error[0] = asfloat(asuint(Error[0]) & ~0x007FFFFF);
-		Error[1] = asfloat(asuint(Error[1]) & ~0x007FFFFF);
-		Error[2] = asfloat(asuint(Error[2]) & ~0x007FFFFF);
-	}
-	#endif
-	
-	return Color + Error * E;
-}
-
-CALL_SITE_DEBUGLOC
-tsr_half3 QuantizeForFloatRenderTarget(tsr_half3 Color, tsr_half E, const float3 QuantizationError)
-{
-	return dpv_lo(QuantizeForFloatRenderTarget(dpv_interleave_mono_registers(Color), E, QuantizationError));
-}
-
 CALL_SITE_DEBUGLOC
 uint2 EncodeSpatialAntiAliasingOffset(tsr_half2x2 TexelOffset)
 {
diff --git a/Engine/Shaders/Private/TemporalSuperResolution/TSRDecimateHistory.usf b/Engine/Shaders/Private/TemporalSuperResolution/TSRDecimateHistory.usf
--- a/Engine/Shaders/Private/TemporalSuperResolution/TSRDecimateHistory.usf
+++ b/Engine/Shaders/Private/TemporalSuperResolution/TSRDecimateHistory.usf
@@ -1,6 +1,6 @@
 // Copyright Epic Games, Inc. All Rights Reserved.
 
-#include "TSRCommon.ush"
+#include "TSRColorSpace.ush"
 #include "TSRReprojectionField.ush"
 
 
@@ -62,7 +62,6 @@ float2 ResurrectionGuideUVViewportBilinearMin;
 float2 ResurrectionGuideUVViewportBilinearMax;
 float3 HistoryGuideQuantizationError;
 float ParallaxRejectionMaskThreshold;
-float PerceptionAdd;
 float ResurrectionFrameIndex;
 float PrevFrameIndex;
 float4x4 ClipToResurrectionClip;
@@ -442,11 +441,7 @@ struct FCatmullRomFetches
 // Correct pre exposure of the guide color.
 tsr_halfC CorrectGuideColorExposure(tsr_halfC ReprojectedHistoryGuideColor, const float PreExposureCorrection)
 {
-	tsr_halfC LDRToLinear = min(rcp(saturate(tsr_half(1.0) - ReprojectedHistoryGuideColor)) * tsr_half(PreExposureCorrection), tsr_half(MaxHalfFloat));
-
-	tsr_halfC LinearReprojectedHistoryGuideColor = ReprojectedHistoryGuideColor * LDRToLinear;
-				
-	ReprojectedHistoryGuideColor = LinearReprojectedHistoryGuideColor * rcp(LinearReprojectedHistoryGuideColor + tsr_half(1.0));
+	ReprojectedHistoryGuideColor = PreExposureCorrectGCS(ReprojectedHistoryGuideColor, PreExposureCorrection);
 
 	// Removes NaN
 	ReprojectedHistoryGuideColor = -min(-ReprojectedHistoryGuideColor, tsr_half(0.0));
@@ -469,11 +464,7 @@ tsr_half4 CorrectMoireExposure(tsr_half4 ReprojectedHistoryMoireColor, const flo
 {
 	tsr_half PrevFlickeringHistory = ReprojectedHistoryMoireColor[0];
 
-	tsr_half LDRToLinear = min(rcp(saturate(tsr_half(1.0) - PrevFlickeringHistory)) * tsr_half(PerceptionAdd) * tsr_half(PreExposureCorrection), tsr_half(MaxHalfFloat));
-	
-	tsr_half LinearPrevFlickeringHistory = PrevFlickeringHistory * LDRToLinear;
-		
-	PrevFlickeringHistory = LinearPrevFlickeringHistory * rcp(LinearPrevFlickeringHistory + tsr_half(PerceptionAdd));
+	PrevFlickeringHistory = PreExposureCorrectGCS(PrevFlickeringHistory.xxx, PreExposureCorrection).r;
 
 	// Removes NaN
 	PrevFlickeringHistory = -min(-PrevFlickeringHistory, tsr_half(0.0));
@@ -521,6 +512,7 @@ void ReprojectAllPrevTextures(
 	
 		float2 PrevOutputScreenPos = OutputScreenPos - DecodeReprojectionVector(EncodedReprojectionVector);
 		bIsOffScreen = (
+			(bCameraCut != 0) ||
 			(PrevOutputScreenPos.x < -1.0) ||
 			(PrevOutputScreenPos.y < -1.0) ||
 			(PrevOutputScreenPos.x > 1.0) ||
@@ -713,7 +705,7 @@ void MainCS(
 	{
 		#if CONFIG_SCENE_COLOR_ALPHA
 		{
-			ReprojectedHistoryGuideColor = ReprojectedHistoryGuideSamples.AccumulateSamples();
+			ReprojectedHistoryGuideColor = saturate(ReprojectedHistoryGuideSamples.AccumulateSamples());
 			ReprojectedHistoryGuideUncertainty = ReprojectedHistoryGuideUncertaintySamples.AccumulateSamples().r;
 		}
 		#else
diff --git a/Engine/Shaders/Private/TemporalSuperResolution/TSRMeasureFlickeringLuma.usf b/Engine/Shaders/Private/TemporalSuperResolution/TSRMeasureFlickeringLuma.usf
--- a/Engine/Shaders/Private/TemporalSuperResolution/TSRMeasureFlickeringLuma.usf
+++ b/Engine/Shaders/Private/TemporalSuperResolution/TSRMeasureFlickeringLuma.usf
@@ -1,6 +1,6 @@
 // Copyright Epic Games, Inc. All Rights Reserved.
 
-#include "TSRCommon.ush"
+#include "TSRColorSpace.ush"
 
 
 //------------------------------------------------------- CONFIG
@@ -12,8 +12,6 @@
 
 //------------------------------------------------------- PARAMETERS
 
-float PerceptionAdd;
-
 Texture2D<tsr_half3> SceneColorTexture;
 RWTexture2D<tsr_half> FlickeringLumaOutput;
 
@@ -57,12 +55,14 @@ void MainCS(
 				const uint2 PixelOffset = uint2(x, y) * TILE_SIZE;
 				uint2 InputPixelPos = OutputPixelOrigin + PixelOffset;
 			
-				tsr_half3 Color = ColorArray[x + y * TILE_ROWS];
-				Color *= rcp(tsr_half(PerceptionAdd) + Color);
+				tsr_half3 LinearColor = ColorArray[x + y * TILE_ROWS];
+				tsr_half3 SMColor = LinearToSMCS(LinearColor);
+				tsr_half SMLuma = dot(SMColor, kMoireSMCSWeights);
+				tsr_half GLuma = SMCSToGCS(SMLuma.rrr).r;
 
 				InputPixelPos.x = select(all(InputPixelPos < InputInfo_ViewportMax), InputPixelPos.x, uint(~0));
-				
-				FlickeringLumaOutput[InputPixelPos] = dot(Color, kMoireLumaWeights);
+
+				FlickeringLumaOutput[InputPixelPos] = GLuma;
 			}
 		}
 	}
diff --git a/Engine/Shaders/Private/TemporalSuperResolution/TSRRejectShading.usf b/Engine/Shaders/Private/TemporalSuperResolution/TSRRejectShading.usf
--- a/Engine/Shaders/Private/TemporalSuperResolution/TSRRejectShading.usf
+++ b/Engine/Shaders/Private/TemporalSuperResolution/TSRRejectShading.usf
@@ -607,21 +607,15 @@ void MainCS(
 		#endif
 
 		// Convert input to shading LDR
-		tsr_tensor_halfC ExposedInput = AddPerceptionAdd(Input, /* DestPerceptionAdd = */ PerceptionAdd);
-		tsr_tensor_halfC ExposedHistory = ChangePerceptionAdd(
-			History,
-			/* SourcePerceptionAdd = */ kHistoryGuidePreceptionAdd,
-			/* DestPerceptionAdd   = */ PerceptionAdd);
+		tsr_tensor_halfC ExposedInput = LinearToSMCS(Input);
+		tsr_tensor_halfC ExposedHistory = GCSToSMCS(History);
 			
 		// Check whether the history resurrection is closer to input, to discard and no pay the register pressure of OriginalResurrectedHistory
 		#if CONFIG_HISTORY_RESURRECTION
 		BRANCH
 		if (bEnableResurrection)
 		{
-			tsr_tensor_halfC ExposedResurrectedHistory = ChangePerceptionAdd(
-				OriginalResurrectedHistory,
-				/* SourcePerceptionAdd = */ kHistoryGuidePreceptionAdd,
-				/* DestPerceptionAdd   = */ PerceptionAdd);
+			tsr_tensor_halfC ExposedResurrectedHistory = GCSToSMCS(OriginalResurrectedHistory);
 				
 			bResurrectionIsCloser = IsResurrectedFrameCloserThanPreviousFrame(bIsResurrectionOffScreen, ExposedInput, ExposedHistory, ExposedResurrectedHistory);
 		}
@@ -688,11 +682,8 @@ void MainCS(
 		#endif
 		
 		// Convert input to shading LDR
-		tsr_tensor_halfC ExposedInput = AddPerceptionAdd(Input, /* DestPerceptionAdd = */ PerceptionAdd);
-		tsr_tensor_halfC ExposedResurrectedHistory = ChangePerceptionAdd(
-			ResurrectedHistory,
-			/* SourcePerceptionAdd = */ kHistoryGuidePreceptionAdd,
-			/* DestPerceptionAdd   = */ PerceptionAdd);
+		tsr_tensor_halfC ExposedInput = LinearToSMCS(Input);
+		tsr_tensor_halfC ExposedResurrectedHistory = GCSToSMCS(ResurrectedHistory);
 			
 		// Mutual Annihilation
 		tsr_tensor_halfC ClampedResurrectedInput, ClampedResurrectedHistory;
@@ -774,10 +765,13 @@ void MainCS(
 		}
 		#endif // CONFIG_HISTORY_RESURRECTION
 
-		// Convert input to LDR
-		tsr_tensor_halfC ExposedInput = AddPerceptionAdd(Input, /* DestPerceptionAdd = */ kHistoryGuidePreceptionAdd);
-	
-		FinalGuide = UpdateGuide(ExposedInput, History, bIsDisoccluded, RejectionBlendFinal);
+		// Convert input to same color space as the accumulation
+		tsr_tensor_halfC ExposedInput   = AddPerceptionAdd(Input,                /* DestPerceptionAdd = */ kHistoryAccumulationPreceptionAdd);
+		tsr_tensor_halfC ExposedHistory = AddPerceptionAdd(GCSToLinear(History), /* DestPerceptionAdd = */ kHistoryAccumulationPreceptionAdd);
+		
+		tsr_tensor_halfC ExposedFinalGuide = UpdateGuide(ExposedInput, ExposedHistory, bIsDisoccluded, RejectionBlendFinal);
+
+		FinalGuide = LinearToGCS(RemovePerceptionAdd(ExposedFinalGuide, /* SourcePerceptionAdd = */ kHistoryAccumulationPreceptionAdd));
 	}
 
 	// Controls how the history and its validity is clamped.
diff --git a/Engine/Shaders/Private/TemporalSuperResolution/TSRReprojectionField.ush b/Engine/Shaders/Private/TemporalSuperResolution/TSRReprojectionField.ush
--- a/Engine/Shaders/Private/TemporalSuperResolution/TSRReprojectionField.ush
+++ b/Engine/Shaders/Private/TemporalSuperResolution/TSRReprojectionField.ush
@@ -1,5 +1,7 @@
 // Copyright Epic Games, Inc. All Rights Reserved.
 
+#pragma once
+
 #include "TSRCommon.ush"
 
 
diff --git a/Engine/Shaders/Private/TemporalSuperResolution/TSRResolveHistory.usf b/Engine/Shaders/Private/TemporalSuperResolution/TSRResolveHistory.usf
--- a/Engine/Shaders/Private/TemporalSuperResolution/TSRResolveHistory.usf
+++ b/Engine/Shaders/Private/TemporalSuperResolution/TSRResolveHistory.usf
@@ -1,6 +1,6 @@
 // Copyright Epic Games, Inc. All Rights Reserved.
 
-#include "TSRCommon.ush"
+#include "TSRColorSpace.ush"
 
 
 //------------------------------------------------------- CONFIG
diff --git a/Engine/Shaders/Private/TemporalSuperResolution/TSRShadingAnalysis.ush b/Engine/Shaders/Private/TemporalSuperResolution/TSRShadingAnalysis.ush
--- a/Engine/Shaders/Private/TemporalSuperResolution/TSRShadingAnalysis.ush
+++ b/Engine/Shaders/Private/TemporalSuperResolution/TSRShadingAnalysis.ush
@@ -5,6 +5,7 @@
 
 //------------------------------------------------------- INCLUDES
 
+#include "TSRColorSpace.ush"
 #include "TSRConvolutionNetworkPass.ush"
 
 
@@ -17,7 +18,6 @@ static const tsr_half FilteringWeight = rcp(tsr_half(1.0 + 4 * 0.5 + 4 * 0.25));
 
 //------------------------------------------------------- PARAMETERS
 
-float PerceptionAdd;
 float FlickeringFramePeriod;
 
 
@@ -83,12 +83,65 @@ tsr_tensor_halfC AddPerceptionAdd(tsr_tensor_halfC SrcColor, const float DestPer
 	return SrcColor;
 }
 
+
+//------------------------------------------------------- GUIDE COLOR SPACE
+
 CALL_SITE_DEBUGLOC
-tsr_tensor_halfC ChangePerceptionAdd(tsr_tensor_halfC SrcColor, const float SourcePerceptionAdd, const float DestPerceptionAdd)
+tsr_tensor_halfC LinearToGCS(tsr_tensor_halfC SrcColor)
 {
-	SrcColor = RemovePerceptionAdd(SrcColor, SourcePerceptionAdd);
-	SrcColor = AddPerceptionAdd(SrcColor, DestPerceptionAdd);
-	return SrcColor;
+	tsr_tensor_halfC GColor = SrcColor * rcp(SrcColor + tsr_tensor_halfC::Const(tsr_half(0.17)));
+	//SMColor = SMColor * SMColor;
+	#if CONFIG_SCENE_COLOR_ALPHA
+		GColor.SetComponent(3, SrcColor[3]);
+	#endif
+	return GColor;
+}
+
+CALL_SITE_DEBUGLOC
+tsr_tensor_halfC GCSToLinear(tsr_tensor_halfC GColor)
+{
+	//SrcColor = sqrt(SrcColor);
+
+	tsr_tensor_halfC Minus = tsr_tensor_halfC::Const(tsr_half(1.0)) - GColor;
+	tsr_tensor_halfC LDRToLinear = min(tsr_tensor_halfC::Const(tsr_half(0.17)) * rcp(Minus), tsr_tensor_halfC::Const(tsr_half(MaxHalfFloat)));
+
+	tsr_tensor_halfC LinearColor = GColor * LDRToLinear;
+	#if CONFIG_SCENE_COLOR_ALPHA
+		LinearColor.SetComponent(3, GColor[3]);
+	#endif
+	return LinearColor;
+}
+
+//------------------------------------------------------- SHADING MEASUREMENT COLOR SPACE
+
+CALL_SITE_DEBUGLOC
+tsr_tensor_half GCSToSMCS(tsr_tensor_half GColor)
+{
+	tsr_tensor_half SMColor = GColor * GColor;
+	return SMColor;
+}
+
+CALL_SITE_DEBUGLOC
+tsr_tensor_half SMCSToGCS(tsr_tensor_half SMColor)
+{
+	tsr_tensor_half GColor = sqrt(SMColor);
+	return GColor;
+}
+
+CALL_SITE_DEBUGLOC
+tsr_tensor_halfC GCSToSMCS(tsr_tensor_halfC GColor)
+{
+	tsr_tensor_halfC SMColor = GColor * GColor;
+	#if CONFIG_SCENE_COLOR_ALPHA
+		SMColor.SetComponent(3, GColor[3]);
+	#endif
+	return SMColor;
+}
+
+CALL_SITE_DEBUGLOC
+tsr_tensor_halfC LinearToSMCS(tsr_tensor_halfC LinearColor)
+{
+	return GCSToSMCS(LinearToGCS(LinearColor));
 }
 
 
@@ -157,7 +210,7 @@ void ComputeMoireError(
 	const tsr_half TotalVariationCountThreshold = tsr_half(1.0 / (1.0 - pow(1.0 - MinBlendFinal, tsr_half(FlickeringFramePeriod))));
 
 	// Input luminance to monitor flickering from.
-	tsr_tensor_half InputLuma = MoireInput;
+	tsr_tensor_half InputLuma = GCSToSMCS(MoireInput);
 
 	// How much input luminance from the TSRComputeMoireLuma.usf measurement has been modified since for instance with fog, water.
 	tsr_tensor_half InputLumaModification;
@@ -166,16 +219,16 @@ void ComputeMoireError(
 			
 		tsr_tensor_halfC OriginalInput = ComposeTranslucency(OriginalOpaqueInput, OriginalTranslucencyInput);
 
-		OriginalInput       = AddPerceptionAdd(OriginalInput, PerceptionAdd);
-		OriginalOpaqueInput = AddPerceptionAdd(OriginalOpaqueInput, PerceptionAdd);
+		OriginalInput       = LinearToSMCS(OriginalInput);
+		OriginalOpaqueInput = LinearToSMCS(OriginalOpaqueInput);
 
 		tsr_tensor_half SceneColorLuma;
 		tsr_tensor_half TranslucencyChange;
 		UNROLL_N(SIMD_SIZE)
 		for (uint ElementIndex = 0; ElementIndex < SIMD_SIZE; ElementIndex++)
 		{
-			SceneColorLuma.SetElement(ElementIndex, dot(OriginalOpaqueInput.GetElement(ElementIndex).rgb, kMoireLumaWeights));
-			TranslucencyChange.SetElement(ElementIndex, dot(abs(OriginalOpaqueInput.GetElement(ElementIndex).rgb - OriginalInput.GetElement(ElementIndex).rgb), kMoireLumaWeights));
+			SceneColorLuma.SetElement(ElementIndex, dot(OriginalOpaqueInput.GetElement(ElementIndex).rgb, kMoireSMCSWeights));
+			TranslucencyChange.SetElement(ElementIndex, dot(abs(OriginalOpaqueInput.GetElement(ElementIndex).rgb - OriginalInput.GetElement(ElementIndex).rgb), kMoireSMCSWeights));
 		}
 
 		InputLumaModification = tsr_tensor_half::Const(0.0);
@@ -185,7 +238,7 @@ void ComputeMoireError(
 	}
 
 	// Unpack history
-	tsr_tensor_half PrevFlickeringHistory = PrevMoireHistory[0];
+	tsr_tensor_half PrevFlickeringHistory = GCSToSMCS(PrevMoireHistory[0]);
 	tsr_tensor_half PrevGradient = PrevMoireHistory[1] * tsr_half(255.0 / 127.0) - tsr_half(1.0);
 	//tsr_tensor_half PrevTotalVariation = PrevMoireHistory[2] * tsr_tensor_half::Const(1.0 - MinBlendFinal);
 	tsr_tensor_half PrevTotalVariation = PrevMoireHistory[2] * tsr_tensor_half::Const(tsr_half(1.0) - MinBlendFinal);
@@ -200,7 +253,7 @@ void ComputeMoireError(
 	tsr_tensor_half OutRejectionBlendFinal;
 	tsr_tensor_half OutRejectionClampBlend;
 	{
-		tsr_tensor_half BackbufferQuantizationErrorVector = tsr_tensor_half::Const(MoireEncodingError);
+		tsr_tensor_half BackbufferQuantizationErrorVector = tsr_tensor_half::Const(MeasureBackbufferLDRQuantizationError());
 		
 		tsr_tensor_half InputC0 = InputLuma;
 		tsr_tensor_half HistoryC0 = PrevFlickeringHistory;
@@ -279,6 +332,8 @@ void ComputeMoireError(
 		CurrentGradient = FinalFlickeringHistory - GhostingHistory;
 	}
 
+	FinalFlickeringHistory = SMCSToGCS(FinalFlickeringHistory);
+
 	// Remove how much luma has changed from luma measurement to final scene color from the gradient.
 	for (uint ElementIndex = 0; ElementIndex < SIMD_SIZE; ElementIndex++)
 	{
diff --git a/Engine/Shaders/Private/TemporalSuperResolution/TSRUpdateHistory.usf b/Engine/Shaders/Private/TemporalSuperResolution/TSRUpdateHistory.usf
--- a/Engine/Shaders/Private/TemporalSuperResolution/TSRUpdateHistory.usf
+++ b/Engine/Shaders/Private/TemporalSuperResolution/TSRUpdateHistory.usf
@@ -1,6 +1,6 @@
 // Copyright Epic Games, Inc. All Rights Reserved.
 
-#include "TSRCommon.ush"
+#include "TSRColorSpace.ush"
 #include "TSRReprojectionField.ush"
 
 
diff --git a/Engine/Shaders/Private/TemporalSuperResolution/TSRVisualize.usf b/Engine/Shaders/Private/TemporalSuperResolution/TSRVisualize.usf
--- a/Engine/Shaders/Private/TemporalSuperResolution/TSRVisualize.usf
+++ b/Engine/Shaders/Private/TemporalSuperResolution/TSRVisualize.usf
@@ -1,6 +1,6 @@
 // Copyright Epic Games, Inc. All Rights Reserved.
 
-#include "TSRCommon.ush"
+#include "TSRColorSpace.ush"
 #include "TSRReprojectionField.ush"
 #if TSR_SUPPORT_LENS_DISTORTION
 	#include "../LensDistortion.ush"
@@ -43,7 +43,6 @@ uint bReprojectionField;
 float MaxHistorySampleCount;
 float OutputToHistoryResolutionFractionSquare;
 float FlickeringFramePeriod;
-float PerceptionAdd;
 
 Texture2D<float2>    PrevDistortingDisplacementTexture;
 Texture2D<float2>    ResurrectedDistortingDisplacementTexture;
@@ -332,11 +331,11 @@ tsr_half3 DisplayAntiFlickering(FContext Ctx)
 		
 		tsr_half3 OriginalInput = OriginalOpaqueInput * OriginalTranslucencyInput.a + OriginalTranslucencyInput.rgb;
 		
-		OriginalInput = OriginalInput * rcp(OriginalInput + tsr_half(PerceptionAdd));
-		OriginalOpaqueInput = OriginalOpaqueInput * rcp(OriginalOpaqueInput + tsr_half(PerceptionAdd));
+		OriginalInput = LinearToSMCS(OriginalInput);
+		OriginalOpaqueInput = LinearToSMCS(OriginalOpaqueInput);
 
-		tsr_half SceneColorLuma = dot(OriginalOpaqueInput.rgb, kMoireLumaWeights);
-		tsr_half TranslucencyChange = dot(abs(OriginalOpaqueInput.rgb - OriginalInput.rgb), kMoireLumaWeights);
+		tsr_half SceneColorLuma = dot(OriginalOpaqueInput.rgb, kMoireSMCSWeights);
+		tsr_half TranslucencyChange = dot(abs(OriginalOpaqueInput.rgb - OriginalInput.rgb), kMoireSMCSWeights);
 		
 		InputLumaModification = tsr_half(0.0);
 		InputLumaModification = max(InputLumaModification, abs(MoireInput - SceneColorLuma) - InputLumaEncodingError);
diff --git a/Engine/Shaders/Public/LaneVectorization.ush b/Engine/Shaders/Public/LaneVectorization.ush
--- a/Engine/Shaders/Public/LaneVectorization.ush
+++ b/Engine/Shaders/Public/LaneVectorization.ush
@@ -1612,6 +1612,11 @@ void Deconcatenate(
 		return R; \
 	} \
 	
+TVECTOR_FUNCTION_1PARAMS(ScalarType, log);
+TVECTOR_FUNCTION_1PARAMS(ScalarType, log2);
+TVECTOR_FUNCTION_1PARAMS(ScalarType, sqrt);
+TVECTOR_FUNCTION_1PARAMS(ScalarType, rsqrt);
+TVECTOR_FUNCTION_1PARAMS(ScalarType, exp);
 TVECTOR_FUNCTION_1PARAMS(ScalarType, rcp);
 TVECTOR_FUNCTION_1PARAMS(ScalarType, saturate);
 TVECTOR_FUNCTION_1PARAMS(ScalarType, abs);
@@ -1642,6 +1647,7 @@ TVECTOR_FUNCTION_1PARAMS(uint, asuint);
 		return R; \
 	} \
 
+TVECTOR_FUNCTION_2PARAMS(pow);
 TVECTOR_FUNCTION_2PARAMS(min);
 TVECTOR_FUNCTION_2PARAMS(max);
 TVECTOR_FUNCTION_2PARAMS(and_internal);
diff --git a/Engine/Source/Runtime/Renderer/Private/PostProcess/TemporalSuperResolution.cpp b/Engine/Source/Runtime/Renderer/Private/PostProcess/TemporalSuperResolution.cpp
--- a/Engine/Source/Runtime/Renderer/Private/PostProcess/TemporalSuperResolution.cpp
+++ b/Engine/Source/Runtime/Renderer/Private/PostProcess/TemporalSuperResolution.cpp
@@ -199,22 +199,6 @@ TAutoConsoleVariable<int32> CVarTSRShadingTileOverscan(
 	TEXT("and therefor need to overlap the tile by couple of padding to hide it. Higher means less prones to tiling artifacts, but performance loss."),
 	ECVF_RenderThreadSafe);
 
-TAutoConsoleVariable<float> CVarTSRShadingExposureOffset(
-	TEXT("r.TSR.ShadingRejection.ExposureOffset"), 3.0,
-	TEXT("The shading rejection needs to have a representative idea how bright a linear color pixel ends up displayed to the user. ")
-	TEXT("And the shading rejection detect if a color become to changed to be visible in the back buffer by comparing to MeasureBackbufferLDRQuantizationError().\n")
-	TEXT("\n")
-	TEXT("It is important to have TSR's MeasureBackbufferLDRQuantizationError() ends up distributed uniformly across ")
-	TEXT("the range of color intensity or it could otherwise disregard some subtle VFX causing ghostin.\n")
-	TEXT("\n")
-	TEXT("This controls adjusts the exposure of the linear color space solely in the TSR's rejection heuristic, such that higher value ")
-	TEXT("lifts the shadows's LDR intensity, meaning MeasureBackbufferLDRQuantizationError() is decreased in these shadows and increased in ")
-	TEXT("the highlights, control directly.\n")
-	TEXT("\n")
-	TEXT("The best TSR internal buffer to verify this is TSR.Flickering.Luminance, either with the \"show VisualizeTemporalUpscaler\" command or in DumpGPU ")
-	TEXT("with the RGB Linear[0;1] source color space against the Tonemaper's output in sRGB source color space.\n"),
-	ECVF_RenderThreadSafe);
-
 TAutoConsoleVariable<int32> CVarTSRLensDistortion(
 	TEXT("r.TSR.LensDistortion"), 1,
 	TEXT("Whether to apply lens distortion in TSR at runtime (enabled by default, requires r.TSR.Support.LensDistortion enabled at cook time)."),
@@ -646,7 +630,6 @@ class FTSRMeasureFlickeringLumaCS : public FTSRShader
 
 	BEGIN_SHADER_PARAMETER_STRUCT(FParameters, )
 		SHADER_PARAMETER_STRUCT(FScreenPassTextureViewportParameters, InputInfo)
-		SHADER_PARAMETER(float, PerceptionAdd)
 		SHADER_PARAMETER_RDG_TEXTURE(Texture2D, SceneColorTexture)
 		SHADER_PARAMETER_RDG_TEXTURE_UAV(RWTexture2D, FlickeringLumaOutput)
 	END_SHADER_PARAMETER_STRUCT()
@@ -749,7 +732,6 @@ class FTSRDecimateHistoryCS : public FTSRShader
 		SHADER_PARAMETER(FVector2f, ResurrectionGuideUVViewportBilinearMax)
 		SHADER_PARAMETER(FVector3f, HistoryGuideQuantizationError)
 		SHADER_PARAMETER(float, ParallaxRejectionMaskThreshold)
-		SHADER_PARAMETER(float, PerceptionAdd)
 		SHADER_PARAMETER(float, ResurrectionFrameIndex)
 		SHADER_PARAMETER(float, PrevFrameIndex)
 		SHADER_PARAMETER(FMatrix44f, ClipToResurrectionClip)
@@ -818,7 +800,6 @@ class FTSRRejectShadingCS : public FTSRConvolutionNetworkShader
 		SHADER_PARAMETER(float, FlickeringFramePeriod)
 		SHADER_PARAMETER(float, TheoricBlendFactor)
 		SHADER_PARAMETER(int32, TileOverscan)
-		SHADER_PARAMETER(float, PerceptionAdd)
 		SHADER_PARAMETER(int32, bEnableResurrection)
 		SHADER_PARAMETER(int32, bEnableFlickeringHeuristic)
 
@@ -1145,7 +1126,6 @@ class FTSRVisualizeCS : public FTSRShader
 		SHADER_PARAMETER(float, MaxHistorySampleCount)
 		SHADER_PARAMETER(float, OutputToHistoryResolutionFractionSquare)
 		SHADER_PARAMETER(float, FlickeringFramePeriod)
-		SHADER_PARAMETER(float, PerceptionAdd)
 
 		SHADER_PARAMETER_RDG_TEXTURE(Texture2D, PrevDistortingDisplacementTexture)
 		SHADER_PARAMETER_RDG_TEXTURE(Texture2D, ResurrectedDistortingDisplacementTexture)
@@ -1394,7 +1374,6 @@ FScreenPassTexture AddTSRMeasureFlickeringLuma(FRDGBuilder& GraphBuilder, FGloba
 	FTSRMeasureFlickeringLumaCS::FParameters* PassParameters = GraphBuilder.AllocParameters<FTSRMeasureFlickeringLumaCS::FParameters>();
 	PassParameters->InputInfo = GetScreenPassTextureViewportParameters(FScreenPassTextureViewport(
 		SceneColor.Texture->Desc.Extent, SceneColor.ViewRect));
-	PassParameters->PerceptionAdd = FMath::Pow(0.5f, CVarTSRShadingExposureOffset.GetValueOnRenderThread());
 	PassParameters->SceneColorTexture = SceneColor.Texture;
 	PassParameters->FlickeringLumaOutput = GraphBuilder.CreateUAV(FlickeringLuma.Texture);
 
@@ -2109,9 +2088,8 @@ FDefaultTemporalUpscaler::FOutputs AddTemporalSuperResolutionPasses(
 				FScreenTransform::ETextureBasis::TextureUV);
 			PassParameters->ResurrectionGuideUVViewportBilinearMin = GetScreenPassTextureViewportParameters(ResurrectionGuideViewport).UVViewportBilinearMin;
 			PassParameters->ResurrectionGuideUVViewportBilinearMax = GetScreenPassTextureViewportParameters(ResurrectionGuideViewport).UVViewportBilinearMax;
-			PassParameters->HistoryGuideQuantizationError = ComputePixelFormatQuantizationError(PrevHistory.GuideArray->Desc.Format);
+			PassParameters->HistoryGuideQuantizationError = ComputePixelFormatQuantizationError(ReprojectedHistoryGuideTexture->Desc.Format);
 			PassParameters->ParallaxRejectionMaskThreshold = 1.0f - 0.25f * OutputToInputResolutionFraction;
-			PassParameters->PerceptionAdd = FMath::Pow(0.5f, CVarTSRShadingExposureOffset.GetValueOnRenderThread());
 		}
 
 		PassParameters->ResurrectionFrameIndex = ResurrectionFrameSliceIndex;
@@ -2247,7 +2225,6 @@ FDefaultTemporalUpscaler::FOutputs AddTemporalSuperResolutionPasses(
 		PassParameters->FlickeringFramePeriod = FlickeringFramePeriod;
 		PassParameters->TheoricBlendFactor = 1.0f / (1.0f + MaxHistorySampleCount / OutputToInputResolutionFractionSquare);
 		PassParameters->TileOverscan = TileOverscan;
-		PassParameters->PerceptionAdd = FMath::Pow(0.5f, CVarTSRShadingExposureOffset.GetValueOnRenderThread());
 		PassParameters->bEnableResurrection = bCanResurrectHistory;
 		PassParameters->bEnableFlickeringHeuristic = FlickeringFramePeriod > 0.0f;
 
@@ -2844,7 +2821,6 @@ FDefaultTemporalUpscaler::FOutputs AddTemporalSuperResolutionPasses(
 			PassParameters->MaxHistorySampleCount = MaxHistorySampleCount;
 			PassParameters->OutputToHistoryResolutionFractionSquare = OutputToHistoryResolutionFractionSquare;
 			PassParameters->FlickeringFramePeriod = FlickeringFramePeriod;
-			PassParameters->PerceptionAdd = FMath::Pow(0.5f, CVarTSRShadingExposureOffset.GetValueOnRenderThread());
 
 			PassParameters->PrevDistortingDisplacementTexture = PrevDistortingDisplacementTexture;
 			PassParameters->ResurrectedDistortingDisplacementTexture = ResurrectedDistortingDisplacementTexture;
