Fix specular occlusion being applied on specular indirect lighting, while it should only be applied on the rough specular indirect lighting, during the indirect lighting composition with Substrate.

diff --git a/Engine/Shaders/Private/DiffuseIndirectComposite.usf b/Engine/Shaders/Private/DiffuseIndirectComposite.usf
--- a/Engine/Shaders/Private/DiffuseIndirectComposite.usf
+++ b/Engine/Shaders/Private/DiffuseIndirectComposite.usf
@@ -291,12 +291,12 @@ FSubstrateDeferredLighting SubstrateIndirectLighting(
 		const bool bHasBackfaceDiffuse = SubstrateGetBSDFType(BSDF) == SUBSTRATE_BSDF_TYPE_SLAB && BSDF.HasBackScattering();
 		const FShadingOcclusion Occlusion = GetShadingOcclusion(PixelPos, ClosureIndex, V, SubstrateBSDFContext.N, Roughness, DiffuseColor, InAO);
 
-		const float3 InGlossyLighting = MixSpecularAndRoughReflections(Roughness, bHasBackfaceDiffuse, InSpecularLighting, InRoughSpecularLighting);
+		const float3 InGlossyLighting = MixSpecularAndRoughReflections(Roughness, bHasBackfaceDiffuse, InSpecularLighting, InRoughSpecularLighting * Occlusion.SpecularOcclusion);
 
 
 		FDirectLighting OutBSDF = (FDirectLighting)0;
 		OutBSDF.Diffuse  += BSDFThroughput * InDiffuseLighting * SubstrateEnvLight.DiffuseWeight * Occlusion.DiffuseOcclusion;
-		OutBSDF.Specular += BSDFThroughput * InGlossyLighting  * SubstrateEnvLight.SpecularWeight * Occlusion.SpecularOcclusion;
+		OutBSDF.Specular += BSDFThroughput * InGlossyLighting  * SubstrateEnvLight.SpecularWeight; // Note: Occlusion.SpecularOcclusion is only applied on the rough specular part, as in legacy path.
 
 		#if DIM_APPLY_DIFFUSE_INDIRECT == DIM_APPLY_DIFFUSE_INDIRECT_SCREEN_PROBE_GATHER
 		if (bHasBackfaceDiffuse && bLumenSupportBackfaceDiffuse > 0)
@@ -309,8 +309,8 @@ FSubstrateDeferredLighting SubstrateIndirectLighting(
 		#if SUBSTRATE_FASTPATH==0
 		if (any(SubstrateEnvLight.SpecularHazeWeight > 0.0f))
 		{
-			const float3 InHazeGlossyLighting = MixSpecularAndRoughReflections(SubstrateEnvLight.SpecularHazeSafeRoughness, bHasBackfaceDiffuse, InSpecularLighting, InRoughSpecularLighting);
-			OutBSDF.Specular += BSDFThroughput * SubstrateEnvLight.SpecularHazeWeight * InHazeGlossyLighting * Occlusion.SpecularOcclusion;
+			const float3 InHazeGlossyLighting = MixSpecularAndRoughReflections(SubstrateEnvLight.SpecularHazeSafeRoughness, bHasBackfaceDiffuse, InSpecularLighting, InRoughSpecularLighting * Occlusion.SpecularOcclusion);
+			OutBSDF.Specular += BSDFThroughput * SubstrateEnvLight.SpecularHazeWeight * InHazeGlossyLighting; // Note: Occlusion.SpecularOcclusion is only applied on the rough specular part, as in legacy path.
 		}
 		#endif 
 
