Refactors all TSR's DilateVelocity & DecimateHistory functionalities into headers for better modularity between the two

diff --git a/Engine/Shaders/Private/TemporalSuperResolution/TSRClearPrevTextures.usf b/Engine/Shaders/Private/TemporalSuperResolution/TSRClearPrevTextures.usf
--- a/Engine/Shaders/Private/TemporalSuperResolution/TSRClearPrevTextures.usf
+++ b/Engine/Shaders/Private/TemporalSuperResolution/TSRClearPrevTextures.usf
@@ -1,6 +1,6 @@
 // Copyright Epic Games, Inc. All Rights Reserved.
 
-#include "TSRCommon.ush"
+#include "TSRClosestOccluder.ush"
 
 
 //------------------------------------------------------- CONFIG
diff --git a/Engine/Shaders/Private/TemporalSuperResolution/TSRClosestOccluder.ush b/Engine/Shaders/Private/TemporalSuperResolution/TSRClosestOccluder.ush
new file mode 100644
--- /dev/null
+++ b/Engine/Shaders/Private/TemporalSuperResolution/TSRClosestOccluder.ush
@@ -0,0 +1,333 @@
+// Copyright Epic Games, Inc. All Rights Reserved.
+
+#pragma once
+
+#include "TSRCommon.ush"
+
+
+//------------------------------------------------------- DEFINITIONS
+
+#define CONFIG_SUPPORT_ORTHO_VIEW 1
+
+#define VELOCITY_HOLE_FILLING_BITS 18
+#define VELOCITY_HOLE_FILLING_LENGTH_BITS 13
+#define VELOCITY_HOLE_FILLING_ANGLE_BITS (VELOCITY_HOLE_FILLING_BITS - VELOCITY_HOLE_FILLING_LENGTH_BITS)
+#define VELOCITY_HOLE_FILLING_LENGTH_PRECISION (pow(2.0, 5))
+
+
+//------------------------------------------------------- PARAMETERS
+
+float2 PrevOutputBufferUVMin;
+float2 PrevOutputBufferUVMax;
+float ParallaxRejectionMaskThreshold;
+
+
+//------------------------------------------------------- HOLE FILLING VELOCITY ENCODING
+
+CALL_SITE_DEBUGLOC
+void PolarCoordinateHoleFillingVelocity(float2 PixelVelocity, out float PixelVelocityAngle, out float PixelVelocityLength)
+{
+	PixelVelocityAngle = atan2Fast(PixelVelocity.y, PixelVelocity.x);
+	PixelVelocityLength = length(PixelVelocity);
+}
+
+CALL_SITE_DEBUGLOC
+float2 CartesianHoleFillingVelocity(float PixelVelocityAngle, float PixelVelocityLength)
+{
+	float2 PixelDirection;
+	sincos(PixelVelocityAngle, /* out */ PixelDirection.y, /* out */ PixelDirection.x);
+	return PixelDirection * PixelVelocityLength;
+}
+
+CALL_SITE_DEBUGLOC
+void QuantizeHoleFillingVelocity(float2 PixelVelocity, out float EncodedAngle, out float EncodedLength)
+{
+	const uint AngleBitDepth = 1 << VELOCITY_HOLE_FILLING_ANGLE_BITS;
+	const float LengthPixelPrecision = VELOCITY_HOLE_FILLING_LENGTH_PRECISION;
+
+	float PixelVelocityAngle;
+	float PixelVelocityLength;
+	PolarCoordinateHoleFillingVelocity(PixelVelocity, /* out */ PixelVelocityAngle, /* out */ PixelVelocityLength);
+
+	EncodedAngle = PixelVelocityAngle * (0.5 * AngleBitDepth / PI) + (0.5 * AngleBitDepth);
+	EncodedLength = PixelVelocityLength * LengthPixelPrecision;
+}
+
+CALL_SITE_DEBUGLOC
+uint EncodeHoleFillingVelocity(float2 PixelVelocity)
+{
+	const uint AngleBitDepth = 1 << VELOCITY_HOLE_FILLING_ANGLE_BITS;
+	const uint LengthBitDepth = 1 << VELOCITY_HOLE_FILLING_LENGTH_BITS;
+
+	float EncodedAngle;
+	float EncodedLength;
+	QuantizeHoleFillingVelocity(PixelVelocity, /* out */ EncodedAngle, /* out */ EncodedLength);
+
+	uint iEncodedAngle = uint(round(EncodedAngle)) & (AngleBitDepth - 1);
+	uint iEncodedLength = fastClamp(uint(ceil(EncodedLength)), 0u, LengthBitDepth - 1u);
+
+	uint EncodedHoleFillingVelocity = iEncodedAngle * LengthBitDepth + iEncodedLength;
+	return EncodedHoleFillingVelocity;
+}
+
+CALL_SITE_DEBUGLOC
+void DecodeHoleFillingVelocity(uint EncodedHoleFillingVelocity, out float PixelAngle, out float PixelLength)
+{
+	const uint AngleBitDepth = 1 << VELOCITY_HOLE_FILLING_ANGLE_BITS;
+	const uint LengthBitDepth = 1 << VELOCITY_HOLE_FILLING_LENGTH_BITS;
+	const float LengthPixelPrecision = VELOCITY_HOLE_FILLING_LENGTH_PRECISION;
+
+	uint iEncodedLength = EncodedHoleFillingVelocity % LengthBitDepth;
+	uint iEncodedAngle = (EncodedHoleFillingVelocity >> VELOCITY_HOLE_FILLING_LENGTH_BITS) % AngleBitDepth;
+
+	float EncodedLength = float(iEncodedLength);
+	PixelLength = EncodedLength * rcp(LengthPixelPrecision);
+
+	float EncodedAngle = float(iEncodedAngle);
+	PixelAngle = EncodedAngle * (PI * 2.0 / float(AngleBitDepth)) - PI;
+}
+
+CALL_SITE_DEBUGLOC
+float GetMaxEncodableHoleFillingPixelLength()
+{
+	const uint LengthBitDepth = 1 << VELOCITY_HOLE_FILLING_LENGTH_BITS;
+	const float LengthPixelPrecision = VELOCITY_HOLE_FILLING_LENGTH_PRECISION;
+	const float LengthRange = (float(LengthBitDepth) - 1.5) / LengthPixelPrecision;
+	return LengthRange;
+}
+
+
+//------------------------------------------------------- CLOSEST OCCLUDER SCATTERING
+
+void ScatterClosestOccluder(globallycoherent RWTexture2DArray<uint> PrevAtomicOutput, bool bValidOutputPixel, float2 ScreenPos, float2 FinalScreenVelocity, float PrevClosestDeviceZ)
+{
+	const float PrevUseCountQuantization = 63.0;
+
+	float2 PrevScreenPos = ScreenPos - FinalScreenVelocity.xy;
+	float2 PrevInputBufferUV = (InputInfo_ScreenPosToViewportScale * PrevScreenPos + InputInfo_ScreenPosToViewportBias) * InputInfo_ExtentInverse;
+
+	uint EncodedHoleFillingVelocity = EncodeHoleFillingVelocity(FinalScreenVelocity.xy * ScreenVelocityToInputPixelVelocity);
+
+	bool bValidHistoryCoord = bValidOutputPixel && all(and(PrevInputBufferUV > PrevOutputBufferUVMin, PrevInputBufferUV < PrevOutputBufferUVMax));
+
+	FBilinearSampleInfos BilinearInter = GetBilinearSampleLevelInfos(PrevInputBufferUV, InputInfo_Extent, InputInfo_ExtentInverse);
+	
+	#if CONFIG_SUPPORT_ORTHO_VIEW
+	BRANCH
+	if (IsOrthoProjection())
+	{
+		uint Depth = clamp(uint(round(PrevClosestDeviceZ * 8388607.0)), 0, 8388607) << 9;
+
+		uint PrevClosestDepth0 = Depth | ((EncodedHoleFillingVelocity >> 9u) & 0x1FFu);
+		uint PrevClosestDepth1 = Depth | ((EncodedHoleFillingVelocity >> 0u) & 0x1FFu);
+		
+		UNROLL_N(4)
+		for (uint i = 0; i < 4; i++)
+		{
+			float Weight = GetSampleWeight(BilinearInter, i);
+			uint2 PixelPos = GetSamplePixelCoord(BilinearInter, i);
+
+			uint FinalAdd = uint(round(Weight * PrevUseCountQuantization));
+			
+			PixelPos.x = ((FinalAdd > 0) && bValidHistoryCoord) ? PixelPos.x : (~0);
+
+			uint PreviousDepth0;
+			InterlockedMax(PrevAtomicOutput[tsr_ushort3(PixelPos, 0)], PrevClosestDepth0, /* out */ PreviousDepth0);
+
+			PixelPos.x = select(PreviousDepth0 < PrevClosestDepth0, PixelPos.x, ~0u);
+			InterlockedMax(PrevAtomicOutput[tsr_ushort3(PixelPos, 1)], PrevClosestDepth1);
+		}
+	}
+	else
+	#endif // CONFIG_SUPPORT_ORTHO_VIEW
+	{
+		uint PrevClosestDepth = (f32tof16(PrevClosestDeviceZ) << VELOCITY_HOLE_FILLING_BITS) | EncodedHoleFillingVelocity;
+			
+		UNROLL_N(4)
+		for (uint i = 0; i < 4; i++)
+		{
+			float Weight = GetSampleWeight(BilinearInter, i);
+			uint2 PixelPos = GetSamplePixelCoord(BilinearInter, i);
+
+			uint FinalAdd = uint(round(Weight * PrevUseCountQuantization));
+			
+			PixelPos.x = ((FinalAdd > 0) && bValidHistoryCoord) ? PixelPos.x : (~0);
+
+			InterlockedMax(PrevAtomicOutput[tsr_ushort3(PixelPos, 0)], PrevClosestDepth);
+		}
+	}
+} // ScatterClosestOccluder()
+
+void LoadPrevAtomicTexturesSamples(
+	Texture2DArray<uint> PrevAtomicTextureArray,
+	float2 PrevScreenPos,
+	out uint HistoryClosestDeviceZSamples0[4],
+	out uint HistoryClosestDeviceZSamples1[4])
+{
+	float2 PrevInputBufferUV = (InputInfo_ScreenPosToViewportScale * PrevScreenPos + InputInfo_ScreenPosToViewportBias) * InputInfo_ExtentInverse;
+	PrevInputBufferUV = INVARIANT(PrevInputBufferUV);
+	FBilinearSampleInfos BilinearInter = GetBilinearSampleLevelInfos(PrevInputBufferUV, InputInfo_Extent, InputInfo_ExtentInverse);
+
+	UNROLL_N(4)
+	for (uint i = 0; i < 4; i++)
+	{
+		float BilinearWeight = GetSampleWeight(BilinearInter, i);
+		uint2 PixelPos = GetSamplePixelCoord(BilinearInter, i);
+
+		HistoryClosestDeviceZSamples0[i] = PrevAtomicTextureArray[uint3(PixelPos, 0)];
+
+		#if CONFIG_SUPPORT_ORTHO_VIEW
+		{
+			HistoryClosestDeviceZSamples1[i] = PrevAtomicTextureArray[uint3(PixelPos, 1)];
+		}
+		#else
+		{
+			HistoryClosestDeviceZSamples1[i] = 0;
+		}
+		#endif
+	}
+} // LoadPrevAtomicTexturesSamples()
+
+void ProcessPrevAtomicTexturesSamples(
+	uint HistoryClosestDeviceZSamples0[4],
+	uint HistoryClosestDeviceZSamples1[4],
+	float2 ScreenPos,
+	float2 ScreenVelocity,
+	float PrevDeviceZ,
+	float DeviceZError,
+	bool bIsOffScreen,
+	out bool bIsParallaxDisocclusion,
+	out float2 HoleFillingPixelVelocity,
+	out bool bCanHoleFill)
+{
+	float WorldDepth = ConvertFromDeviceZ(PrevDeviceZ);
+	float WorldDepthError = abs(WorldDepth - ConvertFromDeviceZ(PrevDeviceZ + DeviceZError));
+	
+	float2 PrevScreenPos = ScreenPos - ScreenVelocity;
+	float2 PrevInputBufferUV = (InputInfo_ScreenPosToViewportScale * PrevScreenPos + InputInfo_ScreenPosToViewportBias) * InputInfo_ExtentInverse;
+	PrevInputBufferUV = INVARIANT(PrevInputBufferUV);
+	FBilinearSampleInfos BilinearInter = GetBilinearSampleLevelInfos(PrevInputBufferUV, InputInfo_Extent, InputInfo_ExtentInverse);
+	
+	tsr_half ParallaxRejectionMask = tsr_half(0.0);
+	bool bIsValidHoleFillingPixelVelocity = false;
+	uint EncodedHoleFillingVelocity = 0;
+	#if CONFIG_SUPPORT_ORTHO_VIEW
+	BRANCH
+	if (IsOrthoProjection())
+	{
+		float ClosestDeviceZ = 0.0f;
+
+		UNROLL_N(4)
+		for (uint i = 0; i < 4; i++)
+		{
+			float BilinearWeight = GetSampleWeight(BilinearInter, i);
+			uint2 PixelPos = GetSamplePixelCoord(BilinearInter, i);
+
+			bool bValidPixelPos = all(and(PixelPos >= InputInfo_ViewportMin, PixelPos < InputInfo_ViewportMax));
+			
+			uint SampleHistoryClosestDeviceZ0 = HistoryClosestDeviceZSamples0[i];
+			uint SampleHistoryClosestDeviceZ1 = HistoryClosestDeviceZSamples1[i];
+		
+			uint SampleEncodedHoleFillingVelocity = ((SampleHistoryClosestDeviceZ0 & 0x1FFu) << 9) | (SampleHistoryClosestDeviceZ1 & 0x1FFu);
+
+			float HistoryClosestDeviceZ = float((SampleHistoryClosestDeviceZ0 & 0xFFFFFE00u) >> 9u) * rcp(8388607.0);
+			float HistoryClosestWorldDepth = ConvertFromDeviceZ(HistoryClosestDeviceZ);
+
+			tsr_half DepthRejection;
+			{
+				const float PixelDepthError = 3.0;
+				const float RadiusToDiameter = 2.0;
+				float WorldDepthEpsilon = GetDepthPixelRadiusForProjectionType(HistoryClosestWorldDepth) * PixelDepthError * RadiusToDiameter;
+				WorldDepthEpsilon += WorldDepthError;
+				float DeltaDepth = abs(HistoryClosestWorldDepth - WorldDepth);
+
+				DepthRejection = saturate(tsr_half(2.0) - tsr_half(DeltaDepth / WorldDepthEpsilon));
+			}
+
+			bIsValidHoleFillingPixelVelocity = bIsValidHoleFillingPixelVelocity || bValidPixelPos;
+
+			FLATTEN
+			if (bValidPixelPos && HistoryClosestDeviceZ > ClosestDeviceZ)
+			{
+				EncodedHoleFillingVelocity = SampleEncodedHoleFillingVelocity;
+			}
+
+			ParallaxRejectionMask = select(bValidPixelPos, ParallaxRejectionMask + tsr_half(BilinearWeight) * DepthRejection, ParallaxRejectionMask);
+		}
+	}
+	else
+	#endif // CONFIG_SUPPORT_ORTHO_VIEW
+	{
+		UNROLL_N(4)
+		for (uint i = 0; i < 4; i++)
+		{
+			float BilinearWeight = GetSampleWeight(BilinearInter, i);
+			uint2 PixelPos = GetSamplePixelCoord(BilinearInter, i);
+
+			bool bValidPixelPos = all(and(PixelPos >= InputInfo_ViewportMin, PixelPos < InputInfo_ViewportMax));
+			
+			uint SampleHistoryClosestDeviceZ = HistoryClosestDeviceZSamples0[i];
+		
+			float HistoryClosestDeviceZ = f16tof32(SampleHistoryClosestDeviceZ >> VELOCITY_HOLE_FILLING_BITS);
+			float HistoryClosestWorldDepth = ConvertFromDeviceZ(HistoryClosestDeviceZ);
+
+			tsr_half DepthRejection;
+			{
+				const float PixelDepthError = 3.0;
+				const float RadiusToDiameter = 2.0;
+				float WorldDepthEpsilon = GetDepthPixelRadiusForProjectionType(HistoryClosestWorldDepth) * PixelDepthError * RadiusToDiameter;
+				WorldDepthEpsilon += WorldDepthError;
+				float DeltaDepth = abs(HistoryClosestWorldDepth - WorldDepth);
+
+				DepthRejection = saturate(tsr_half(2.0) - tsr_half(DeltaDepth / WorldDepthEpsilon));
+			}
+
+			bIsValidHoleFillingPixelVelocity = bIsValidHoleFillingPixelVelocity || bValidPixelPos;
+
+			EncodedHoleFillingVelocity = select(bValidPixelPos, max(EncodedHoleFillingVelocity, SampleHistoryClosestDeviceZ), EncodedHoleFillingVelocity);
+			ParallaxRejectionMask = select(bValidPixelPos, ParallaxRejectionMask + tsr_half(BilinearWeight) * DepthRejection, ParallaxRejectionMask);
+		}
+	}
+	
+	float HoleFillingPixelAngle;
+	float HoleFillingPixelLength;
+	DecodeHoleFillingVelocity(EncodedHoleFillingVelocity, /* out */ HoleFillingPixelAngle, /* out */ HoleFillingPixelLength);
+
+	bool bIsEncodablePixelLength = HoleFillingPixelLength < GetMaxEncodableHoleFillingPixelLength();
+
+	HoleFillingPixelVelocity = CartesianHoleFillingVelocity(HoleFillingPixelAngle, HoleFillingPixelLength);
+	bCanHoleFill = bIsValidHoleFillingPixelVelocity && bIsEncodablePixelLength;
+	//bIsParallaxDisocclusion = !bIsOffScreen && ParallaxRejectionMask < tsr_half(ParallaxRejectionMaskThreshold);
+
+	// Dim down the parallax rejection mask if the hole fill velocity matches
+	FLATTEN
+	if (bIsEncodablePixelLength)
+	{
+		const float AnglePrecision = 2.0 * PI * pow(0.5, VELOCITY_HOLE_FILLING_ANGLE_BITS);
+		const float LengthPixelPrecision = VELOCITY_HOLE_FILLING_LENGTH_PRECISION;
+
+		float2 PixelVelocity = ScreenVelocity * ScreenVelocityToInputPixelVelocity;
+
+		float PixelVelocityAngle, PixelVelocityLength;
+		PolarCoordinateHoleFillingVelocity(PixelVelocity, /* out */ PixelVelocityAngle, /* out */ PixelVelocityLength);
+		
+		float CartesianCoordinatePrecision = AnglePrecision * PixelVelocityLength;
+
+		float PixelLengthDiff = abs(HoleFillingPixelLength - PixelVelocityLength) - 2.0;
+		float PixelAngleDiff = abs(PixelVelocityAngle - HoleFillingPixelAngle);
+		PixelAngleDiff = min(PixelAngleDiff, 2.0 * PI - PixelAngleDiff);
+		
+		float VelocityDiff = length(HoleFillingPixelVelocity - PixelVelocity);
+		float PixelAngleCompare = saturate(2.0 - PixelAngleDiff / AnglePrecision);
+		float PixelLengthCompare = saturate(1.0 + rcp(LengthPixelPrecision) - PixelLengthDiff);
+
+		float MinParallaxRejectionMaskPolar = min(PixelLengthCompare, PixelAngleCompare);
+		float MinParallaxRejectionMaskCartesian = saturate(1.0 + CartesianCoordinatePrecision - VelocityDiff);
+		
+		float MinParallaxRejectionMask = lerp(MinParallaxRejectionMaskCartesian, MinParallaxRejectionMaskPolar, saturate(min(HoleFillingPixelLength, PixelVelocityLength) - 2.0));
+
+		ParallaxRejectionMask = max(ParallaxRejectionMask, tsr_half(MinParallaxRejectionMask));
+	}
+
+	bIsParallaxDisocclusion = !bIsOffScreen && ParallaxRejectionMask < tsr_half(ParallaxRejectionMaskThreshold);
+} // ProcessPrevAtomicTexturesSamples()
diff --git a/Engine/Shaders/Private/TemporalSuperResolution/TSRCommon.ush b/Engine/Shaders/Private/TemporalSuperResolution/TSRCommon.ush
--- a/Engine/Shaders/Private/TemporalSuperResolution/TSRCommon.ush
+++ b/Engine/Shaders/Private/TemporalSuperResolution/TSRCommon.ush
@@ -35,14 +35,6 @@
 	#define CONFIG_COMPILE_FP16 0
 #endif
 
-#if PLATFORM_SUPPORTS_WAVE_BROADCAST && 1
-	#define CONFIG_BUTTERFLY_KERNEL  1
-	#define CONFIG_BUTTERFLY_SAMPLES 4
-#else
-	#define CONFIG_BUTTERFLY_KERNEL  0
-	#define CONFIG_BUTTERFLY_SAMPLES 9
-#endif
-
 #define CONFIG_SCENE_COLOR_OVERFLOW 1
 
 #if defined(DIM_ALPHA_CHANNEL)
@@ -57,102 +49,12 @@
 // which on 16bit, if linear color spaces ends up close to half(65504.0), this can go +inf on if the compiler transform a * 0.5 + b * 0.5 into (a + b) * 0.5.
 #define CONFIG_FP16_PRECISE_MULTIPLY_ORDER 1
 
-#define CONFIG_SUPPORT_ORTHO_VIEW 1
-
 
 //------------------------------------------------------- CONSTANTS
 
 /* Maximum number of sample. */
 #define MAX_SAMPLE_COUNT 8
 
-#define MAX_FALLBACK_SAMPLE_COUNT 20
-
-// Spatial anti-aliasing encoding settings.
-#define SPATIAL_ANTI_ALIASING_OFFSET_RANGE 0.5
-#define SPATIAL_ANTI_ALIASING_OFFSET_BITDEPTH_PER_CHANNEL 4
-
-#define MAX_DEVICE_Z_ERROR 64.0
-
-// Minimum number of sample of each texel when reprojecting the high frequency history.
-// Value too low causes area of low sample count to have pixel dominating too much in reprojection causing input res aliasing on texture details
-#define MIN_REPROJECTION_SAMPLE_COUNT 2.0
-
-// Minimum weight of a pixel in history.
-#define MIN_IMPLICIT_HISTORY_WEIGHT ceil(tsr_half(255.0) * rcp(tsr_half(MAX_SAMPLE_COUNT)) * tsr_half(MIN_REPROJECTION_SAMPLE_COUNT)) * rcp(tsr_half(255.0))
-
-#define SPATIAL_ANTI_ALIASER_MIN_LUMIMANCE 0.05
-
-
-// K = Center of the nearest input pixel.
-// O = Center of the output pixel.
-//
-//          |           |
-//    0     |     1     |     2
-//          |           |
-//          |           |
-//  --------+-----------+--------
-//          |           |
-//          | O         |
-//    3     |     K     |     5
-//          |           |
-//          |           |
-//  --------+-----------+--------
-//          |           |
-//          |           |
-//    6     |     7     |     8
-//          |           |
-//
-static const int2 kOffsets3x3[9] =
-{
-	int2(-1, -1),
-	int2(0, -1),
-	int2(1, -1),
-	int2(-1,  0),
-	int2(0,  0), // K
-	int2(1,  0),
-	int2(-1,  1),
-	int2(0,  1),
-	int2(1,  1),
-};
-
-// T = Center of the nearest top left pixel input pixel.
-// O = Center of the output pixel.
-//
-//          | 
-//    T     |     .
-//          | 
-//       O  | 
-//  --------+--------
-//          | 
-//          | 
-//    .     |     .
-//          | 
-static const int2 Offsets2x2[4] =
-{
-	int2( 0,  0), // T
-	int2( 1,  0),
-	int2( 0,  1),
-	int2( 1,  1),
-};
-
-// Indexes of the 3x3 square.
-static const uint kSquareIndexes3x3[9] = { 4, 0, 1, 2, 3, 8, 7, 6, 5 };
-
-// Indexes of the offsets to have plus + shape.
-static const uint kPlusIndexes3x3[5] = { 4, 1, 3, 7, 5 };
-
-// Cardinal pixel directions;
-static const int2 kOffsetC = int2( 0,  0);
-static const int2 kOffsetN = int2( 0, -1);
-static const int2 kOffsetS = int2( 0, +1);
-static const int2 kOffsetE = int2(+1,  0);
-static const int2 kOffsetW = int2(-1,  0);
-static const int2 kOffsetNE = int2(+1, -1);
-static const int2 kOffsetNW = int2(-1, -1);
-static const int2 kOffsetSE = int2(+1, +1);
-static const int2 kOffsetSW = int2(-1, +1);
-
-
 #if CONFIG_COMPILE_FP16
 	#define tsr_half  half
 	#define tsr_half2 half2
@@ -384,148 +286,6 @@ uint GetGroupThreadIndex(uint GroupThreadIndex, uint GroupWaveIndex)
 }
 #endif
 
-// Clamp the offset to be shared across multiple samples
-CALL_SITE_DEBUGLOC
-tsr_short2x2 ClampPixelOffset(
-	tsr_short2x2 KernelCenterPixelPos,
-	tsr_short2x2 Offset,
-	const tsr_short2 OffsetDirection,
-	int2 MinPixelPos, int2 MaxPixelPos)
-{
-	tsr_short2x2 Min = dpv_sub(tsr_short2(MinPixelPos), KernelCenterPixelPos);
-	tsr_short2x2 Max = dpv_sub(tsr_short2(MaxPixelPos), KernelCenterPixelPos);
-
-	// Only do clamp based on the compile time known direction of the offset.
-	// This turns only 1 v_pk_max_i16 and v_pk_min_i16 for a 3x3 kernel.
-	tsr_short2x2 ClampedOffset = 0;
-
-	if (OffsetDirection.x > 0)
-	{
-		ClampedOffset[0] = min(Offset[0], Max[0]);
-	}
-	else if (OffsetDirection.x < 0)
-	{
-		ClampedOffset[0] = max(Offset[0], Min[0]);
-	}
-		
-	if (OffsetDirection.y > 0)
-	{
-		ClampedOffset[1] = min(Offset[1], Max[1]);
-	}
-	else if (OffsetDirection.y < 0)
-	{
-		ClampedOffset[1] = max(Offset[1], Min[1]);
-	}
-
-	return ClampedOffset;
-}
-
-CALL_SITE_DEBUGLOC
-tsr_short2 ClampPixelOffset(
-	tsr_short2 KernelCenterPixelPos,
-	tsr_short2 Offset,
-	const tsr_short2 OffsetDirection,
-	int2 MinPixelPos, int2 MaxPixelPos)
-{
-	tsr_short2 Min = tsr_short2(MinPixelPos) - KernelCenterPixelPos;
-	tsr_short2 Max = tsr_short2(MaxPixelPos) - KernelCenterPixelPos;
-
-	// Only do clamp based on the compile time known direction of the offset.
-	// This turns only 1 v_pk_max_i16 and v_pk_min_i16 for a 3x3 kernel.
-	tsr_short2 ClampedOffset = 0;
-
-	if (OffsetDirection.x > 0)
-		ClampedOffset.x = min(Offset.x, Max.x);
-	else if (OffsetDirection.x < 0)
-		ClampedOffset.x = max(Offset.x, Min.x);
-		
-	if (OffsetDirection.y > 0)
-		ClampedOffset.y = min(Offset.y, Max.y);
-	else if (OffsetDirection.y < 0)
-		ClampedOffset.y = max(Offset.y, Min.y);
-
-	return ClampedOffset;
-}
-
-CALL_SITE_DEBUGLOC
-tsr_short2x2 ClampPixelOffset(
-	tsr_short2x2 SamplePixelPos,
-	int2 MinPixelPos, int2 MaxPixelPos)
-{
-	SamplePixelPos[0] = fastClamp(SamplePixelPos[0], tsr_short2(MinPixelPos.xx), tsr_short2(MaxPixelPos.xx));
-	SamplePixelPos[1] = fastClamp(SamplePixelPos[1], tsr_short2(MinPixelPos.yy), tsr_short2(MaxPixelPos.yy));
-	return SamplePixelPos;
-}
-
-CALL_SITE_DEBUGLOC
-tsr_short2 ClampPixelOffset(
-	tsr_short2 SamplePixelPos,
-	int2 MinPixelPos, int2 MaxPixelPos)
-{
-	return fastClamp(SamplePixelPos, tsr_short2(MinPixelPos), tsr_short2(MaxPixelPos));
-}
-
-CALL_SITE_DEBUGLOC
-tsr_short2x2 AddAndClampPixelOffset(
-	tsr_short2x2 KernelCenterPixelPos,
-	tsr_short2x2 Offset,
-	const tsr_short2 OffsetDirection,
-	int2 MinPixelPos, int2 MaxPixelPos)
-{
-	tsr_short2x2 SamplePixelPos = KernelCenterPixelPos + Offset;
-	
-	tsr_short2 SamplePixelPos0 = dpv_lo(SamplePixelPos);
-	tsr_short2 SamplePixelPos1 = dpv_hi(SamplePixelPos);
-
-	// Only do clamp based on the compile time known direction of the offset.
-	// This turns only 1 v_pk_max_i16 and v_pk_min_i16 for a 3x3 kernel.
-	if (OffsetDirection.x > 0)
-		SamplePixelPos0.x = min(SamplePixelPos0.x, tsr_short(MaxPixelPos.x));
-	else if (OffsetDirection.x < 0)
-		SamplePixelPos0.x = max(SamplePixelPos0.x, tsr_short(MinPixelPos.x));
-		
-	if (OffsetDirection.y > 0)
-		SamplePixelPos0.y = min(SamplePixelPos0.y, tsr_short(MaxPixelPos.y));
-	else if (OffsetDirection.y < 0)
-		SamplePixelPos0.y = max(SamplePixelPos0.y, tsr_short(MinPixelPos.y));
-	
-	if (OffsetDirection.x > 0)
-		SamplePixelPos1.x = min(SamplePixelPos1.x, tsr_short(MaxPixelPos.x));
-	else if (OffsetDirection.x < 0)
-		SamplePixelPos1.x = max(SamplePixelPos1.x, tsr_short(MinPixelPos.x));
-		
-	if (OffsetDirection.y > 0)
-		SamplePixelPos1.y = min(SamplePixelPos1.y, tsr_short(MaxPixelPos.y));
-	else if (OffsetDirection.y < 0)
-		SamplePixelPos1.y = max(SamplePixelPos1.y, tsr_short(MinPixelPos.y));
-	
-	return dpv_interleave_registers(SamplePixelPos0, SamplePixelPos1);
-}
-
-CALL_SITE_DEBUGLOC
-tsr_short2 AddAndClampPixelOffset(
-	tsr_short2 KernelCenterPixelPos,
-	tsr_short2 Offset,
-	const tsr_short2 OffsetDirection,
-	int2 MinPixelPos, int2 MaxPixelPos)
-{
-	tsr_short2 SamplePixelPos = KernelCenterPixelPos + Offset;
-
-	// Only do clamp based on the compile time known direction of the offset.
-	// This turns only 1 v_pk_max_i16 and v_pk_min_i16 for a 3x3 kernel.
-	if (OffsetDirection.x > 0)
-		SamplePixelPos.x = min(SamplePixelPos.x, tsr_short(MaxPixelPos.x));
-	else if (OffsetDirection.x < 0)
-		SamplePixelPos.x = max(SamplePixelPos.x, tsr_short(MinPixelPos.x));
-		
-	if (OffsetDirection.y > 0)
-		SamplePixelPos.y = min(SamplePixelPos.y, tsr_short(MaxPixelPos.y));
-	else if (OffsetDirection.y < 0)
-		SamplePixelPos.y = max(SamplePixelPos.y, tsr_short(MinPixelPos.y));
-
-	return SamplePixelPos;
-}
-
 CALL_SITE_DEBUGLOC
 tsr_short2x2 InvalidateOutputPixelPos(tsr_short2x2 PixelPos, uint2 ViewportMax)
 #if 1
@@ -634,39 +394,6 @@ half2 SafeRcp(half2 x)
 
 #endif
 
-CALL_SITE_DEBUGLOC
-tsr_half2 ComputeLowFrequencyContribution(tsr_half2 HistoryValidity)
-{
-	return saturate((HistoryValidity - tsr_half(1.0 * rcp(MAX_SAMPLE_COUNT))) * tsr_half(6.0));
-}
-
-CALL_SITE_DEBUGLOC
-tsr_half ComputeLowFrequencyContribution(tsr_half HistoryValidity)
-{
-	return dpv_lo(ComputeLowFrequencyContribution(dpv_interleave_mono_registers(HistoryValidity)));
-}
-
-// Returns the weight of a pixels at a coordinate <PixelDelta> from the PDF highest point.
-CALL_SITE_DEBUGLOC
-tsr_half2 ComputeSampleWeigth(tsr_half2 UpscaleFactor, tsr_half2x2 PixelDelta, const float MinimalContribution)
-{
-	tsr_half2 u2 = UpscaleFactor * UpscaleFactor;
-
-	// 1 - 1.9 * x^2 + 0.9 * x^4
-	tsr_half2 x2 = saturate(u2 * dpv_length2(PixelDelta));
-	//return tsr_half(((float(0.9) + MinimalContribution) * x2 - float(1.9)) * x2 + float(1.0));
-	return saturate((tsr_half(0.9) * x2 - tsr_half(1.9)) * x2 + tsr_half(1.0 + MinimalContribution));
-}
-
-CALL_SITE_DEBUGLOC
-tsr_half ComputeSampleWeigth(tsr_half UpscaleFactor, tsr_half2 PixelDelta, const float MinimalContribution)
-{
-	return dpv_lo(ComputeSampleWeigth(
-		dpv_interleave_mono_registers(UpscaleFactor),
-		dpv_interleave_mono_registers(PixelDelta),
-		MinimalContribution));
-}
-
 CALL_SITE_DEBUGLOC
 tsr_half2x2 WeightedLerpFactors(tsr_half2 WeightA, tsr_half2 WeightB, tsr_half2 Blend)
 {
@@ -693,16 +420,6 @@ bool TakeOnlyOneSamplePair(float2 Offset)
 	return Offset.y > 0.0 || (Offset.x > 0.0 && Offset.y == 0.0);
 }
 
-float3 ComputeStaticVelocity(float2 ScreenPos, float DeviceZ)
-{
-	float3 PosN = float3(ScreenPos, DeviceZ);
-
-	float4 ThisClip = float4(PosN, 1);
-	float4 PrevClip = mul( ThisClip, View.ClipToPrevClip );
-	float3 PrevScreen = PrevClip.xyz / PrevClip.w;
-	return PosN - PrevScreen;
-}
-
 tsr_half ComputePredictionCompleteness(tsr_half SampleHistoryValidity)
 {
 	return saturate(SampleHistoryValidity * tsr_half(MAX_SAMPLE_COUNT) - tsr_half(0.2));
@@ -718,228 +435,6 @@ void IsOffScreenOrDisoccluded(uint bCameraCut, float2x2 ScreenPos, bool2 bIsPara
 	bIsDisoccluded = and(!bIsOffScreen, bIsParallaxRejected);
 }
 
-CALL_SITE_DEBUGLOC
-uint2 EncodeSpatialAntiAliasingOffset(tsr_half2x2 TexelOffset)
-{
-	const uint Mask = (1 << SPATIAL_ANTI_ALIASING_OFFSET_BITDEPTH_PER_CHANNEL) - 1;
-	const uint Zero = Mask / 2;
-
-	tsr_half Multiply = tsr_half(float(Zero) / SPATIAL_ANTI_ALIASING_OFFSET_RANGE);
-
-	uint2 EncodedTexelOffset = 0;
-	EncodedTexelOffset |= fastClamp(uint2(int2(round(TexelOffset[0] * Multiply)) + int(1 + Zero)), uint(1), Mask) << 0;
-	EncodedTexelOffset |= fastClamp(uint2(int2(round(TexelOffset[1] * Multiply)) + int(1 + Zero)), uint(1), Mask) << SPATIAL_ANTI_ALIASING_OFFSET_BITDEPTH_PER_CHANNEL;
-	return EncodedTexelOffset;
-}
-
-CALL_SITE_DEBUGLOC
-uint EncodeSpatialAntiAliasingOffset(tsr_half2 TexelOffset)
-{
-	return dpv_lo(EncodeSpatialAntiAliasingOffset(dpv_interleave_mono_registers(TexelOffset)));
-}
-
-CALL_SITE_DEBUGLOC
-tsr_half2x2 DecodeSpatialAntiAliasingOffset(tsr_ushort2 EncodedInputTexelOffset)
-{
-	const tsr_ushort Mask = tsr_ushort(1 << SPATIAL_ANTI_ALIASING_OFFSET_BITDEPTH_PER_CHANNEL) - tsr_ushort(1);
-	const tsr_ushort Zero = Mask / tsr_ushort(2);
-	
-	tsr_half Multiply = tsr_half(SPATIAL_ANTI_ALIASING_OFFSET_RANGE / float(Zero));
-
-	tsr_half2x2 TexelOffset;
-	TexelOffset[0] = tsr_half2((EncodedInputTexelOffset >> tsr_ushort(0                                                )) & Mask) * Multiply - (tsr_half(1 + Zero) * Multiply);
-	TexelOffset[1] = tsr_half2((EncodedInputTexelOffset >> tsr_ushort(SPATIAL_ANTI_ALIASING_OFFSET_BITDEPTH_PER_CHANNEL)) & Mask) * Multiply - (tsr_half(1 + Zero) * Multiply);
-
-	return TexelOffset;
-}
-
-CALL_SITE_DEBUGLOC
-tsr_half2 DecodeSpatialAntiAliasingOffset(tsr_ushort EncodedInputTexelOffset)
-{
-	return dpv_lo(DecodeSpatialAntiAliasingOffset(dpv_interleave_mono_registers(EncodedInputTexelOffset)));
-}
-
-
-#define VELOCITY_HOLE_FILLING_BITS 18
-#define VELOCITY_HOLE_FILLING_LENGTH_BITS 13
-#define VELOCITY_HOLE_FILLING_ANGLE_BITS (VELOCITY_HOLE_FILLING_BITS - VELOCITY_HOLE_FILLING_LENGTH_BITS)
-#define VELOCITY_HOLE_FILLING_LENGTH_PRECISION (pow(2.0, 5))
-
-CALL_SITE_DEBUGLOC
-void PolarCoordinateHoleFillingVelocity(float2 PixelVelocity, out float PixelVelocityAngle, out float PixelVelocityLength)
-{
-	PixelVelocityAngle = atan2Fast(PixelVelocity.y, PixelVelocity.x);
-	PixelVelocityLength = length(PixelVelocity);
-}
-
-CALL_SITE_DEBUGLOC
-float2 CartesianHoleFillingVelocity(float PixelVelocityAngle, float PixelVelocityLength)
-{
-	float2 PixelDirection;
-	sincos(PixelVelocityAngle, /* out */ PixelDirection.y, /* out */ PixelDirection.x);
-	return PixelDirection * PixelVelocityLength;
-}
-
-CALL_SITE_DEBUGLOC
-void QuantizeHoleFillingVelocity(float2 PixelVelocity, out float EncodedAngle, out float EncodedLength)
-{
-	const uint AngleBitDepth = 1 << VELOCITY_HOLE_FILLING_ANGLE_BITS;
-	const float LengthPixelPrecision = VELOCITY_HOLE_FILLING_LENGTH_PRECISION;
-
-	float PixelVelocityAngle;
-	float PixelVelocityLength;
-	PolarCoordinateHoleFillingVelocity(PixelVelocity, /* out */ PixelVelocityAngle, /* out */ PixelVelocityLength);
-
-	EncodedAngle = PixelVelocityAngle * (0.5 * AngleBitDepth / PI) + (0.5 * AngleBitDepth);
-	EncodedLength = PixelVelocityLength * LengthPixelPrecision;
-}
-
-CALL_SITE_DEBUGLOC
-uint EncodeHoleFillingVelocity(float2 PixelVelocity)
-{
-	const uint AngleBitDepth = 1 << VELOCITY_HOLE_FILLING_ANGLE_BITS;
-	const uint LengthBitDepth = 1 << VELOCITY_HOLE_FILLING_LENGTH_BITS;
-
-	float EncodedAngle;
-	float EncodedLength;
-	QuantizeHoleFillingVelocity(PixelVelocity, /* out */ EncodedAngle, /* out */ EncodedLength);
-
-	uint iEncodedAngle = uint(round(EncodedAngle)) & (AngleBitDepth - 1);
-	uint iEncodedLength = fastClamp(uint(ceil(EncodedLength)), 0u, LengthBitDepth - 1u);
-
-	uint EncodedHoleFillingVelocity = iEncodedAngle * LengthBitDepth + iEncodedLength;
-	return EncodedHoleFillingVelocity;
-}
-
-CALL_SITE_DEBUGLOC
-void DecodeHoleFillingVelocity(uint EncodedHoleFillingVelocity, out float PixelAngle, out float PixelLength)
-{
-	const uint AngleBitDepth = 1 << VELOCITY_HOLE_FILLING_ANGLE_BITS;
-	const uint LengthBitDepth = 1 << VELOCITY_HOLE_FILLING_LENGTH_BITS;
-	const float LengthPixelPrecision = VELOCITY_HOLE_FILLING_LENGTH_PRECISION;
-
-	uint iEncodedLength = EncodedHoleFillingVelocity % LengthBitDepth;
-	uint iEncodedAngle = (EncodedHoleFillingVelocity >> VELOCITY_HOLE_FILLING_LENGTH_BITS) % AngleBitDepth;
-
-	float EncodedLength = float(iEncodedLength);
-	PixelLength = EncodedLength * rcp(LengthPixelPrecision);
-
-	float EncodedAngle = float(iEncodedAngle);
-	PixelAngle = EncodedAngle * (PI * 2.0 / float(AngleBitDepth)) - PI;
-}
-
-CALL_SITE_DEBUGLOC
-float GetMaxEncodableHoleFillingPixelLength()
-{
-	const uint LengthBitDepth = 1 << VELOCITY_HOLE_FILLING_LENGTH_BITS;
-	const float LengthPixelPrecision = VELOCITY_HOLE_FILLING_LENGTH_PRECISION;
-	const float LengthRange = (float(LengthBitDepth) - 1.5) / LengthPixelPrecision;
-	return LengthRange;
-}
-
-CALL_SITE_DEBUGLOC
-tsr_short2 GetLaneOffsetSign()
-#if CONFIG_BUTTERFLY_KERNEL
-{
-	uint LaneIndex = WaveGetLaneIndex();
-
-	return tsr_short(-1) + tsr_short2(tsr_ushort2(tsr_ushort(LaneIndex) << tsr_ushort(1), LaneIndex) & tsr_ushort(0x2));
-}
-#else
-{
-	return tsr_short(1).xx;
-}
-#endif
-
-tsr_short2 GetNeighborhoodPixelPos(const uint ArrayIndex, tsr_short2 InputPixelPos)
-{
-	tsr_short2 LaneOffsetSign = GetLaneOffsetSign();
-
-	tsr_short2 SampleInputPixelPos;
-	#if CONFIG_BUTTERFLY_KERNEL
-	{
-		tsr_short2 Offset = tsr_short2(Offsets2x2[ArrayIndex]) * LaneOffsetSign;
-		SampleInputPixelPos = ClampPixelOffset(InputPixelPos + Offset, InputPixelPosMin, InputPixelPosMax);
-	}
-	#else
-	{
-		tsr_short2 Offset = kOffsets3x3[ArrayIndex];
-					
-		SampleInputPixelPos = InputPixelPos + ClampPixelOffset(
-			InputPixelPos,
-			Offset, Offset,
-			InputPixelPosMin, InputPixelPosMax);
-	}
-	#endif
-
-	return SampleInputPixelPos;
-}
-
-template<typename Type>
-Type AccessNeighborhoodStaticOffset(Type Array[CONFIG_BUTTERFLY_SAMPLES], const int2 Offset)
-{
-	const uint OffsetIndex = uint(dot(Offset + 1, int2(1, 3)));
-
-	#if CONFIG_BUTTERFLY_KERNEL
-		const uint ButterflyArray[] = {
-			0x3, 0x2, 0x2,
-			0x1, 0x0, 0x0,
-			0x1, 0x0, 0x0,
-		};
-		const uint IndexArray[] = {
-			0, 0, 1,
-			0, 0, 1,
-			2, 2, 3,
-		};
-
-		const FWaveBroadcastSettings BroadcastSetting = InitWaveXorButterfly(ButterflyArray[OffsetIndex]);
-		
-		return WaveBroadcast(BroadcastSetting, Array[IndexArray[OffsetIndex]]);
-	#else
-		return Array[OffsetIndex];
-	#endif
-}
-
-template<typename Type>
-Type AccessNeighborhoodDynamicOffset(Type Array[CONFIG_BUTTERFLY_SAMPLES], const tsr_short2 Offset)
-{
-	const uint OffsetIndex = uint(dot(Offset + 1, tsr_short2(1, 3)));
-
-	#if CONFIG_BUTTERFLY_KERNEL
-		const FWaveBroadcastSettings BroadcastVertical   = InitWaveXorButterfly(0x2);
-		const FWaveBroadcastSettings BroadcastHorizontal = InitWaveXorButterfly(0x1);
-		const FWaveBroadcastSettings BroadcastDiagonal   = InitWaveXorButterfly(0x3);
-
-		Type Return = Array[0];
-		Return = select(all(Offset == tsr_short2(-1, -1)), WaveBroadcast(BroadcastDiagonal,   Array[0]), Return);
-		Return = select(all(Offset == tsr_short2( 0, -1)), WaveBroadcast(BroadcastVertical,   Array[0]), Return);
-		Return = select(all(Offset == tsr_short2(+1, -1)), WaveBroadcast(BroadcastVertical,   Array[1]), Return);
-		
-		Return = select(all(Offset == tsr_short2(-1,  0)), WaveBroadcast(BroadcastHorizontal, Array[0]), Return);
-		Return = select(all(Offset == tsr_short2(+1,  0)), Array[1], Return);
-		
-		Return = select(all(Offset == tsr_short2(-1, +1)), WaveBroadcast(BroadcastHorizontal, Array[2]), Return);
-		Return = select(all(Offset == tsr_short2( 0, +1)), Array[2], Return);
-		Return = select(all(Offset == tsr_short2(+1, +1)), Array[3], Return);
-		
-		return Return;
-	#else
-		return Array[OffsetIndex];
-	#endif
-}
-
-template<typename Type>
-Type AccessNeighborhoodCenter(Type Array[CONFIG_BUTTERFLY_SAMPLES])
-#if CONFIG_BUTTERFLY_KERNEL
-{
-	return Array[0];
-}
-#else
-{
-	return Array[4];
-}
-#endif
-
 tsr_half2 ComputePrevWeightMultiplier(tsr_half2 BlendFactor)
 {
 	return (1.0 - BlendFactor) / BlendFactor;
@@ -949,73 +444,3 @@ tsr_half2 RejectionFactorToBlendFactor(tsr_half2 RejectionFactor)
 {
 	return tsr_half(1.0) - RejectionFactor;
 }
-
-FCatmullRomSamples GetBicubic2DCatmullRomSamples_Stubbe(float2 UV, float2 Size, in float2 InvSize)
-{
-	float2 PixelCoord = UV * Size;
-
-	float2 iuv = floor(PixelCoord - 0.5) + 1.0;
-	tsr_half2 f = tsr_half2(PixelCoord - iuv);
-	tsr_half2 f2 = f * f;
-	tsr_half2 f_2 = f + f;
-
-	tsr_half2 uv = (tsr_half(1.25) - f2) * f + tsr_half(0.5);
-	tsr_half2 t0 = (tsr_half(0.25) * f2 - tsr_half(0.0625)) * (tsr_half(1.125) - tsr_half(0.5) * f2.yx);
-	tsr_half2 w0 = t0 - f_2 * t0;
-	tsr_half2 w3 = t0 + f_2 * t0;
-
-	float2 cuv = (iuv - 0.5) * InvSize;
-	uv *= tsr_half2(InvSize);
-
-	FCatmullRomSamples Samples;
-	Samples.Count = BICUBIC_CATMULL_ROM_SAMPLES;
-
-	Samples.UV[0] = cuv + tsr_half2(               -InvSize.x,                      uv.y);
-	Samples.UV[1] = cuv + tsr_half2(tsr_half(2.0) * InvSize.x,                      uv.y);
-	Samples.UV[2] = cuv + uv;
-	Samples.UV[3] = cuv + tsr_half2(                     uv.x,                -InvSize.y);
-	Samples.UV[4] = cuv + tsr_half2(                     uv.x, tsr_half(2.0) * InvSize.y);
-	
-	Samples.Weight[0] = half(w0.x);
-	Samples.Weight[1] = half(w3.x);
-	Samples.Weight[2] = half(tsr_half(1.0) - w0.x - w3.x - w0.y - w3.y);
-	Samples.Weight[3] = half(w0.y);
-	Samples.Weight[4] = half(w3.y);
-
-	Samples.UVDir[0] = int2(-1,  0);
-	Samples.UVDir[1] = int2( 1,  0);
-	Samples.UVDir[2] = int2( 0,  0);
-	Samples.UVDir[3] = int2( 0, -1);
-	Samples.UVDir[4] = int2( 0,  1);
-	
-	Samples.FinalMultiplier = 1.0;
-	return Samples;
-}
-
-/** Compute the devize Z error due to the size of a pixel in world space at that given depth. */
-float ComputePixelDeviceZError(float DeviceZ)
-{
-	float Depth = ConvertFromDeviceZ(DeviceZ);
-	float ErrorCorrectedDeviceZ = ConvertToDeviceZ(GetDepthPixelRadiusForProjectionType(Depth) * 2.0f + Depth);
-	float DeviceZError = abs(ErrorCorrectedDeviceZ - DeviceZ);
-	return DeviceZError;
-}
-
-/** Encode the devize Z error relatively to the closest depth to squeeze to 8bit per pixel */
-uint EncodeDeviceZError(float PrevClosestDeviceZ, float ParallaxDepthError)
-{
-	float PixelDeviceZError = float(MAX_DEVICE_Z_ERROR) * ComputePixelDeviceZError(PrevClosestDeviceZ);
-	float RelativeZError = ParallaxDepthError / PixelDeviceZError;
-	
-	uint EncodedDeviceZError = uint(ceil(saturate(RelativeZError) * 255.0));
-	return EncodedDeviceZError;
-}
-
-/** Inverse of EncodeDeviceZError() */
-float DecodeDeviceZError(float PrevClosestDeviceZ, uint EncodedDeviceZError)
-{
-	float PixelDeviceZError = ComputePixelDeviceZError(PrevClosestDeviceZ);
-	float RelativeZError = float(EncodedDeviceZError) * (float(MAX_DEVICE_Z_ERROR) / 255.0);
-	float ParallaxDepthError = PixelDeviceZError * RelativeZError;
-	return ParallaxDepthError;
-}
diff --git a/Engine/Shaders/Private/TemporalSuperResolution/TSRConvolutionNetwork.ush b/Engine/Shaders/Private/TemporalSuperResolution/TSRConvolutionNetwork.ush
--- a/Engine/Shaders/Private/TemporalSuperResolution/TSRConvolutionNetwork.ush
+++ b/Engine/Shaders/Private/TemporalSuperResolution/TSRConvolutionNetwork.ush
@@ -2,7 +2,7 @@
 
 #pragma once
 
-#include "TSRCommon.ush"
+#include "TSRKernels.ush"
 #include "/Engine/Public/LaneVectorization.ush"
 
 
diff --git a/Engine/Shaders/Private/TemporalSuperResolution/TSRDecimateHistory.usf b/Engine/Shaders/Private/TemporalSuperResolution/TSRDecimateHistory.usf
--- a/Engine/Shaders/Private/TemporalSuperResolution/TSRDecimateHistory.usf
+++ b/Engine/Shaders/Private/TemporalSuperResolution/TSRDecimateHistory.usf
@@ -1,5 +1,7 @@
 // Copyright Epic Games, Inc. All Rights Reserved.
 
+#include "TSRDepthVelocityAnalysis.ush"
+#include "TSRClosestOccluder.ush"
 #include "TSRColorSpace.ush"
 #include "TSRReprojectionField.ush"
 
@@ -12,19 +14,6 @@
 
 #define CONFIG_MANUAL_LDS_SPILL 1
 
-#define CONFIG_VELOCITY_SAMPLES 9
-#define CONFIG_VELOCITY_SAMPLE_PAIRS ((CONFIG_VELOCITY_SAMPLES - 1) / 2)
-
-#if CONFIG_BUTTERFLY_KERNEL
-	#if CONFIG_VELOCITY_SAMPLES == 9
-		#define CONFIG_VELOCITY_PRE_FETCH_SAMPLES 4
-	#else
-		#define CONFIG_VELOCITY_PRE_FETCH_SAMPLES 3
-	#endif
-#else // !CONFIG_BUTTERFLY_KERNEL
-	#define CONFIG_VELOCITY_PRE_FETCH_SAMPLES (CONFIG_VELOCITY_SAMPLES)
-#endif
-
 /** Whether to compile the moire reprojection. */
 #define CONFIG_MOIRE_REPROJECTION (DIM_MOIRE_REPROJECTION)
 
@@ -61,7 +50,6 @@ FScreenTransform ScreenPosToResurrectionGuideBufferUV;
 float2 ResurrectionGuideUVViewportBilinearMin;
 float2 ResurrectionGuideUVViewportBilinearMax;
 float3 HistoryGuideQuantizationError;
-float ParallaxRejectionMaskThreshold;
 float ResurrectionFrameIndex;
 float PrevFrameIndex;
 float4x4 ClipToResurrectionClip;
@@ -90,354 +78,6 @@ static float4 Debug[DEBUG_ARRAY_SIZE];
 
 //------------------------------------------------------- FUNCTIONS
 
-void LoadReprojectionNeighborhood(
-	tsr_short2 InputPixelPos,
-	out uint EncodedReprojectionVectorNeighborhood[CONFIG_VELOCITY_PRE_FETCH_SAMPLES])
-{
-	tsr_short2 LaneOffsetSign = GetLaneOffsetSign();
-
-	UNROLL_N(CONFIG_VELOCITY_PRE_FETCH_SAMPLES)
-	for (uint i = 0; i < CONFIG_VELOCITY_PRE_FETCH_SAMPLES; i++)
-	{
-		tsr_short2 SampleInputPixelPos;
-		#if CONFIG_BUTTERFLY_KERNEL
-		{
-			tsr_short2 Offset = tsr_short2(Offsets2x2[i]) * LaneOffsetSign;
-
-			SampleInputPixelPos = InputPixelPos + Offset;
-			SampleInputPixelPos = ClampPixelOffset(SampleInputPixelPos, InputPixelPosMin, InputPixelPosMax);
-		}
-		#else
-		{
-			tsr_short2 Offset;
-			#if CONFIG_VELOCITY_SAMPLES == 9
-				Offset = tsr_short2(kOffsets3x3[kSquareIndexes3x3[i]]);
-			#elif CONFIG_VELOCITY_SAMPLES == 5
-				Offset = tsr_short2(kOffsets3x3[kPlusIndexes3x3[i]]);
-			#endif
-
-			SampleInputPixelPos = InputPixelPos + ClampPixelOffset(
-				InputPixelPos,
-				Offset, Offset,
-				InputPixelPosMin, InputPixelPosMax);
-		}
-		#endif
-
-		EncodedReprojectionVectorNeighborhood[i] = ReprojectionVectorTexture[SampleInputPixelPos];
-	}
-} // LoadReprojectionNeighborhood()
-
-void ComputePixelVelocityNeighborhood(
-	uint EncodedReprojectionVectorNeighborhood[CONFIG_VELOCITY_PRE_FETCH_SAMPLES],
-	out float2 PixelVelocityNeighborhood[CONFIG_VELOCITY_PRE_FETCH_SAMPLES])
-{
-	UNROLL_N(CONFIG_VELOCITY_PRE_FETCH_SAMPLES)
-	for (uint i = 0; i < CONFIG_VELOCITY_PRE_FETCH_SAMPLES; i++)
-	{
-		PixelVelocityNeighborhood[i] = DecodeReprojectionVector(EncodedReprojectionVectorNeighborhood[i]) * ScreenVelocityToInputPixelVelocity;
-	}
-} // ComputePixelVelocityNeighborhood()
-
-/** Compute the reprojection edge from the pixel velocities. */
-void ComputeReprojectionEdge(
-	float2 PixelVelocityNeighborhood[CONFIG_VELOCITY_PRE_FETCH_SAMPLES],
-	out tsr_half ReprojectionEdge)
-{
-	ReprojectionEdge = tsr_half(0.0);
-
-	// Sample
-	float2 ScreenPixelVelocity = PixelVelocityNeighborhood[0];
-	tsr_short2 LaneOffsetSign = GetLaneOffsetSign();
-
-	UNROLL_N(CONFIG_VELOCITY_SAMPLE_PAIRS)
-	for (uint i = 0; i < CONFIG_VELOCITY_SAMPLE_PAIRS; i++)
-	ISOLATE
-	{
-		tsr_short2 Offset0;
-		tsr_short2 Offset1;
-			
-		float2 NeighborPixelVelocity0;
-		float2 NeighborPixelVelocity1;
-
-		#if CONFIG_BUTTERFLY_KERNEL
-		{
-			const uint PairIndex = i;
-				
-			const uint ButterflyArray0[] = { 0x0, 0x0, 0x0, 0x1 };
-			const uint ButterflyArray1[] = { 0x1, 0x2, 0x3, 0x2 };
-				
-			FWaveBroadcastSettings BroadcastSettings[2];
-			BroadcastSettings[0] = InitWaveXorButterfly(ButterflyArray0[PairIndex]);
-			BroadcastSettings[1] = InitWaveXorButterfly(ButterflyArray1[PairIndex]);
-				
-			uint FetchIndex[2];
-			if (PairIndex == 3)
-			{
-				FetchIndex[0] = 2;
-				FetchIndex[1] = 1;
-				Offset0 = tsr_short2(-1, 1);
-			}
-			else
-			{
-				FetchIndex[0] = PairIndex + 1;
-				FetchIndex[1] = 0;
-				Offset0 = tsr_short2(Offsets2x2[PairIndex + 1]);
-			}
-
-			Offset0 *= LaneOffsetSign;
-			Offset1 = -Offset0;
-
-			NeighborPixelVelocity0 = WaveBroadcast(BroadcastSettings[0], PixelVelocityNeighborhood[FetchIndex[0]]);
-			NeighborPixelVelocity1 = WaveBroadcast(BroadcastSettings[1], PixelVelocityNeighborhood[FetchIndex[1]]);
-		}
-		#else
-		{
-			const uint NeightbordId0 = i + 1;
-			const uint NeightbordId1 = i + 1 + CONFIG_VELOCITY_SAMPLE_PAIRS;
-			
-			#if CONFIG_VELOCITY_SAMPLES == 9
-				Offset0 = tsr_short2(kOffsets3x3[kSquareIndexes3x3[NeightbordId0]]);
-				Offset1 = tsr_short2(kOffsets3x3[kSquareIndexes3x3[NeightbordId1]]);
-			#elif CONFIG_VELOCITY_SAMPLES == 5
-				Offset0 = tsr_short2(kOffsets3x3[kPlusIndexes3x3[NeightbordId0]]);
-				Offset1 = tsr_short2(kOffsets3x3[kPlusIndexes3x3[NeightbordId1]]);
-			#endif
-					
-			NeighborPixelVelocity0 = PixelVelocityNeighborhood[NeightbordId0];
-			NeighborPixelVelocity1 = PixelVelocityNeighborhood[NeightbordId1];
-		}
-		#endif
-
-		tsr_half2 NeighborDeltaPixelVelocity0 = tsr_half2(ScreenPixelVelocity - NeighborPixelVelocity0);
-		tsr_half2 NeighborDeltaPixelVelocity1 = tsr_half2(ScreenPixelVelocity - NeighborPixelVelocity1);
-			
-		// Reject high frequency on reprojection edge discontinuity.
-		{
-			tsr_half IsOverlapingEdge0 = tsr_half(1.0) - dot(saturate(abs(NeighborDeltaPixelVelocity0 * Offset0)), tsr_half(1.0).xx); // * saturate(ScreenPixelVelocityLength - NeighborPixelVelocityLength0);
-			tsr_half IsOverlapingEdge1 = tsr_half(1.0) - dot(saturate(abs(NeighborDeltaPixelVelocity1 * Offset1)), tsr_half(1.0).xx); // * saturate(ScreenPixelVelocityLength - NeighborPixelVelocityLength1);
-				
-			tsr_half IsOverlapingEdge = min(IsOverlapingEdge0, IsOverlapingEdge1);
-
-			if (i == 0)
-				ReprojectionEdge = IsOverlapingEdge;
-			else
-				ReprojectionEdge = min(ReprojectionEdge, IsOverlapingEdge);
-		}
-	} // for (uint i = 0; i < CONFIG_VELOCITY_SAMPLE_PAIRS; i++)
-
-	// Remove noise from high frequency rejections.
-	ReprojectionEdge = saturate(ReprojectionEdge * tsr_half(1.1));
-} // ComputeReprojectionEdge()
-
-void LoadPrevAtomicTexturesSamples(
-	float2 PrevScreenPos,
-	out uint HistoryClosestDeviceZSamples0[4],
-	out uint HistoryClosestDeviceZSamples1[4])
-{
-	float2 PrevInputBufferUV = (InputInfo_ScreenPosToViewportScale * PrevScreenPos + InputInfo_ScreenPosToViewportBias) * InputInfo_ExtentInverse;
-	PrevInputBufferUV = INVARIANT(PrevInputBufferUV);
-	FBilinearSampleInfos BilinearInter = GetBilinearSampleLevelInfos(PrevInputBufferUV, InputInfo_Extent, InputInfo_ExtentInverse);
-
-	UNROLL_N(4)
-	for (uint i = 0; i < 4; i++)
-	{
-		float BilinearWeight = GetSampleWeight(BilinearInter, i);
-		uint2 PixelPos = GetSamplePixelCoord(BilinearInter, i);
-
-		HistoryClosestDeviceZSamples0[i] = PrevAtomicTextureArray[uint3(PixelPos, 0)];
-
-		#if CONFIG_SUPPORT_ORTHO_VIEW
-		{
-			HistoryClosestDeviceZSamples1[i] = PrevAtomicTextureArray[uint3(PixelPos, 1)];
-		}
-		#else
-		{
-			HistoryClosestDeviceZSamples1[i] = 0;
-		}
-		#endif
-	}
-} // LoadPrevAtomicTexturesSamples()
-
-void ProcessPrevAtomicTexturesSamples(
-	uint HistoryClosestDeviceZSamples0[4],
-	uint HistoryClosestDeviceZSamples1[4],
-	float2 ScreenPos,
-	float2 ScreenVelocity,
-	float PrevDeviceZ,
-	float DeviceZError,
-	bool bIsOffScreen,
-	out bool bIsParallaxDisocclusion,
-	out float2 HoleFillingPixelVelocity,
-	out bool bCanHoleFill)
-{
-	float WorldDepth = ConvertFromDeviceZ(PrevDeviceZ);
-	float WorldDepthError = abs(WorldDepth - ConvertFromDeviceZ(PrevDeviceZ + DeviceZError));
-	
-	float2 PrevScreenPos = ScreenPos - ScreenVelocity;
-	float2 PrevInputBufferUV = (InputInfo_ScreenPosToViewportScale * PrevScreenPos + InputInfo_ScreenPosToViewportBias) * InputInfo_ExtentInverse;
-	PrevInputBufferUV = INVARIANT(PrevInputBufferUV);
-	FBilinearSampleInfos BilinearInter = GetBilinearSampleLevelInfos(PrevInputBufferUV, InputInfo_Extent, InputInfo_ExtentInverse);
-	
-	tsr_half ParallaxRejectionMask = tsr_half(0.0);
-	bool bIsValidHoleFillingPixelVelocity = false;
-	uint EncodedHoleFillingVelocity = 0;
-	#if CONFIG_SUPPORT_ORTHO_VIEW
-	BRANCH
-	if (IsOrthoProjection())
-	{
-		float ClosestDeviceZ = 0.0f;
-
-		UNROLL_N(4)
-		for (uint i = 0; i < 4; i++)
-		{
-			float BilinearWeight = GetSampleWeight(BilinearInter, i);
-			uint2 PixelPos = GetSamplePixelCoord(BilinearInter, i);
-
-			bool bValidPixelPos = all(and(PixelPos >= InputInfo_ViewportMin, PixelPos < InputInfo_ViewportMax));
-			
-			uint SampleHistoryClosestDeviceZ0 = HistoryClosestDeviceZSamples0[i];
-			uint SampleHistoryClosestDeviceZ1 = HistoryClosestDeviceZSamples1[i];
-		
-			uint SampleEncodedHoleFillingVelocity = ((SampleHistoryClosestDeviceZ0 & 0x1FFu) << 9) | (SampleHistoryClosestDeviceZ1 & 0x1FFu);
-
-			float HistoryClosestDeviceZ = float((SampleHistoryClosestDeviceZ0 & 0xFFFFFE00u) >> 9u) * rcp(8388607.0);
-			float HistoryClosestWorldDepth = ConvertFromDeviceZ(HistoryClosestDeviceZ);
-
-			tsr_half DepthRejection;
-			{
-				const float PixelDepthError = 3.0;
-				const float RadiusToDiameter = 2.0;
-				float WorldDepthEpsilon = GetDepthPixelRadiusForProjectionType(HistoryClosestWorldDepth) * PixelDepthError * RadiusToDiameter;
-				WorldDepthEpsilon += WorldDepthError;
-				float DeltaDepth = abs(HistoryClosestWorldDepth - WorldDepth);
-
-				DepthRejection = saturate(tsr_half(2.0) - tsr_half(DeltaDepth / WorldDepthEpsilon));
-			}
-
-			bIsValidHoleFillingPixelVelocity = bIsValidHoleFillingPixelVelocity || bValidPixelPos;
-
-			FLATTEN
-			if (bValidPixelPos && HistoryClosestDeviceZ > ClosestDeviceZ)
-			{
-				EncodedHoleFillingVelocity = SampleEncodedHoleFillingVelocity;
-			}
-
-			ParallaxRejectionMask = select(bValidPixelPos, ParallaxRejectionMask + tsr_half(BilinearWeight) * DepthRejection, ParallaxRejectionMask);
-		}
-	}
-	else
-	#endif // CONFIG_SUPPORT_ORTHO_VIEW
-	{
-		UNROLL_N(4)
-		for (uint i = 0; i < 4; i++)
-		{
-			float BilinearWeight = GetSampleWeight(BilinearInter, i);
-			uint2 PixelPos = GetSamplePixelCoord(BilinearInter, i);
-
-			bool bValidPixelPos = all(and(PixelPos >= InputInfo_ViewportMin, PixelPos < InputInfo_ViewportMax));
-			
-			uint SampleHistoryClosestDeviceZ = HistoryClosestDeviceZSamples0[i];
-		
-			float HistoryClosestDeviceZ = f16tof32(SampleHistoryClosestDeviceZ >> VELOCITY_HOLE_FILLING_BITS);
-			float HistoryClosestWorldDepth = ConvertFromDeviceZ(HistoryClosestDeviceZ);
-
-			tsr_half DepthRejection;
-			{
-				const float PixelDepthError = 3.0;
-				const float RadiusToDiameter = 2.0;
-				float WorldDepthEpsilon = GetDepthPixelRadiusForProjectionType(HistoryClosestWorldDepth) * PixelDepthError * RadiusToDiameter;
-				WorldDepthEpsilon += WorldDepthError;
-				float DeltaDepth = abs(HistoryClosestWorldDepth - WorldDepth);
-
-				DepthRejection = saturate(tsr_half(2.0) - tsr_half(DeltaDepth / WorldDepthEpsilon));
-			}
-
-			bIsValidHoleFillingPixelVelocity = bIsValidHoleFillingPixelVelocity || bValidPixelPos;
-
-			EncodedHoleFillingVelocity = select(bValidPixelPos, max(EncodedHoleFillingVelocity, SampleHistoryClosestDeviceZ), EncodedHoleFillingVelocity);
-			ParallaxRejectionMask = select(bValidPixelPos, ParallaxRejectionMask + tsr_half(BilinearWeight) * DepthRejection, ParallaxRejectionMask);
-		}
-	}
-	
-	float HoleFillingPixelAngle;
-	float HoleFillingPixelLength;
-	DecodeHoleFillingVelocity(EncodedHoleFillingVelocity, /* out */ HoleFillingPixelAngle, /* out */ HoleFillingPixelLength);
-
-	bool bIsEncodablePixelLength = HoleFillingPixelLength < GetMaxEncodableHoleFillingPixelLength();
-
-	HoleFillingPixelVelocity = CartesianHoleFillingVelocity(HoleFillingPixelAngle, HoleFillingPixelLength);
-	bCanHoleFill = bIsValidHoleFillingPixelVelocity && bIsEncodablePixelLength;
-	//bIsParallaxDisocclusion = !bIsOffScreen && ParallaxRejectionMask < tsr_half(ParallaxRejectionMaskThreshold);
-
-	// Dim down the parallax rejection mask if the hole fill velocity matches
-	FLATTEN
-	if (bIsEncodablePixelLength)
-	{
-		const float AnglePrecision = 2.0 * PI * pow(0.5, VELOCITY_HOLE_FILLING_ANGLE_BITS);
-		const float LengthPixelPrecision = VELOCITY_HOLE_FILLING_LENGTH_PRECISION;
-
-		float2 PixelVelocity = ScreenVelocity * ScreenVelocityToInputPixelVelocity;
-
-		float PixelVelocityAngle, PixelVelocityLength;
-		PolarCoordinateHoleFillingVelocity(PixelVelocity, /* out */ PixelVelocityAngle, /* out */ PixelVelocityLength);
-		
-		float CartesianCoordinatePrecision = AnglePrecision * PixelVelocityLength;
-
-		float PixelLengthDiff = abs(HoleFillingPixelLength - PixelVelocityLength) - 2.0;
-		float PixelAngleDiff = abs(PixelVelocityAngle - HoleFillingPixelAngle);
-		PixelAngleDiff = min(PixelAngleDiff, 2.0 * PI - PixelAngleDiff);
-		
-		float VelocityDiff = length(HoleFillingPixelVelocity - PixelVelocity);
-		float PixelAngleCompare = saturate(2.0 - PixelAngleDiff / AnglePrecision);
-		float PixelLengthCompare = saturate(1.0 + rcp(LengthPixelPrecision) - PixelLengthDiff);
-
-		float MinParallaxRejectionMaskPolar = min(PixelLengthCompare, PixelAngleCompare);
-		float MinParallaxRejectionMaskCartesian = saturate(1.0 + CartesianCoordinatePrecision - VelocityDiff);
-		
-		float MinParallaxRejectionMask = lerp(MinParallaxRejectionMaskCartesian, MinParallaxRejectionMaskPolar, saturate(min(HoleFillingPixelLength, PixelVelocityLength) - 2.0));
-
-		ParallaxRejectionMask = max(ParallaxRejectionMask, tsr_half(MinParallaxRejectionMask));
-	}
-
-	bIsParallaxDisocclusion = !bIsOffScreen && ParallaxRejectionMask < tsr_half(ParallaxRejectionMaskThreshold);
-} // ProcessPrevAtomicTexturesSamples()
-
-template<typename Type>
-struct FCatmullRomFetches
-{
-	Type Colors[BICUBIC_CATMULL_ROM_SAMPLES];
-	tsr_half Weights[BICUBIC_CATMULL_ROM_SAMPLES];
-	
-	void FetchSamples(Texture2DArray<Type> InTexture, float2 UV, float SliceIndex, float2 Size, float2 InvSize, float2 MinUV, float2 MaxUV)
-	{
-		FCatmullRomSamples Samples = GetBicubic2DCatmullRomSamples_Stubbe(UV, Size, InvSize);
-		
-		UNROLL_N(BICUBIC_CATMULL_ROM_SAMPLES)
-		for (uint i = 0; i < BICUBIC_CATMULL_ROM_SAMPLES; i++)
-		{
-			float2 SampleUV = Samples.UV[i];
-			SampleUV = fastClamp(
-				SampleUV,
-				MinUV,
-				MaxUV);
-
-			Colors[i] = InTexture.SampleLevel(GlobalBilinearClampedSampler, float3(SampleUV, SliceIndex), 0);
-			Weights[i] = tsr_half(Samples.Weight[i]);
-		}
-	}
-
-	Type AccumulateSamples()
-	{
-		Type Color = tsr_half(0.0);
-		UNROLL_N(BICUBIC_CATMULL_ROM_SAMPLES)
-		for (uint i = 0; i < BICUBIC_CATMULL_ROM_SAMPLES; i++)
-		{
-			Color += Colors[i] * Weights[i];
-		}
-		return Color;
-	}
-};
-
 // Correct pre exposure of the guide color.
 tsr_halfC CorrectGuideColorExposure(tsr_halfC ReprojectedHistoryGuideColor, const float PreExposureCorrection)
 {
@@ -502,7 +142,9 @@ void ReprojectAllPrevTextures(
 	{
 		float2 ScreenPos = ApplyScreenTransform(float2(InputPixelPos), InputPixelPosToScreenPos);
 		float2 PrevScreenPos = ScreenPos - DecodeReprojectionVector(EncodedReprojectionVector);
-		LoadPrevAtomicTexturesSamples(PrevScreenPos,
+		LoadPrevAtomicTexturesSamples(
+			PrevAtomicTextureArray,
+			PrevScreenPos,
 			/* out */ HistoryClosestDeviceZSamples0,
 			/* out */ HistoryClosestDeviceZSamples1);
 	}
@@ -627,7 +269,7 @@ void MainCS(
 			bHasPixelAnimation = DeviceZAndPrevDeviceZ.r < 0.0;
 			DeviceZError = DecodeDeviceZError(PrevDeviceZ, EncodedDeviceZError);
 
-			//LoadReprojectionNeighborhood(InputPixelPos, /* out */ EncodedReprojectionVectorNeighborhood);
+			//FetchReprojectionNeighborhood(ReprojectionVectorTexture, InputPixelPos, /* out */ EncodedReprojectionVectorNeighborhood);
 			//DilatedEncodedReprojectionVector = EncodedReprojectionVectorNeighborhood[0];
 			DilatedEncodedReprojectionVector = ReprojectionFieldOutput[tsr_ushort3(InputPixelPos, 0)];
 		}
diff --git a/Engine/Shaders/Private/TemporalSuperResolution/TSRDepthVelocityAnalysis.ush b/Engine/Shaders/Private/TemporalSuperResolution/TSRDepthVelocityAnalysis.ush
new file mode 100644
--- /dev/null
+++ b/Engine/Shaders/Private/TemporalSuperResolution/TSRDepthVelocityAnalysis.ush
@@ -0,0 +1,598 @@
+// Copyright Epic Games, Inc. All Rights Reserved.
+
+#pragma once
+
+#include "TSRKernels.ush"
+#include "TSRReprojectionField.ush"
+
+
+//------------------------------------------------------- CONFIG
+
+/** Whether to use the velocity from SceneVelocityTexture. */
+#define CONFIG_USE_SCENE_VELOCITY 1
+
+/** Whether the dilation should be only on depth edges, improving texture reprojection accuracy on flat surfaces. */
+#define CONFIG_DEPTH_EDGE_DETECTION 1
+
+
+//------------------------------------------------------- DEFINITIONS
+
+#define MAX_DEVICE_Z_ERROR 64.0
+
+#define CONFIG_VELOCITY_SAMPLES 9
+#define CONFIG_VELOCITY_SAMPLE_PAIRS ((CONFIG_VELOCITY_SAMPLES - 1) / 2)
+
+
+//------------------------------------------------------- DEPTH & VELOCITY UTILIES
+
+CALL_SITE_DEBUGLOC
+float3 ComputeStaticVelocity(float2 ScreenPos, float DeviceZ)
+{
+	float3 PosN = float3(ScreenPos, DeviceZ);
+
+	float4 ThisClip = float4(PosN, 1);
+	float4 PrevClip = mul( ThisClip, View.ClipToPrevClip );
+	float3 PrevScreen = PrevClip.xyz / PrevClip.w;
+	return PosN - PrevScreen;
+}
+
+/** Compute the devize Z error due to the size of a pixel in world space at that given depth. */
+CALL_SITE_DEBUGLOC
+float ComputePixelDeviceZError(float DeviceZ)
+{
+	float Depth = ConvertFromDeviceZ(DeviceZ);
+	float ErrorCorrectedDeviceZ = ConvertToDeviceZ(GetDepthPixelRadiusForProjectionType(Depth) * 2.0f + Depth);
+	float DeviceZError = abs(ErrorCorrectedDeviceZ - DeviceZ);
+	return DeviceZError;
+}
+
+/** Encode the devize Z error relatively to the closest depth to squeeze to 8bit per pixel */
+CALL_SITE_DEBUGLOC
+uint EncodeDeviceZError(float PrevClosestDeviceZ, float ParallaxDepthError)
+{
+	float PixelDeviceZError = float(MAX_DEVICE_Z_ERROR) * ComputePixelDeviceZError(PrevClosestDeviceZ);
+	float RelativeZError = ParallaxDepthError / PixelDeviceZError;
+	
+	uint EncodedDeviceZError = uint(ceil(saturate(RelativeZError) * 255.0));
+	return EncodedDeviceZError;
+}
+
+/** Inverse of EncodeDeviceZError() */
+CALL_SITE_DEBUGLOC
+float DecodeDeviceZError(float PrevClosestDeviceZ, uint EncodedDeviceZError)
+{
+	float PixelDeviceZError = ComputePixelDeviceZError(PrevClosestDeviceZ);
+	float RelativeZError = float(EncodedDeviceZError) * (float(MAX_DEVICE_Z_ERROR) / 255.0);
+	float ParallaxDepthError = PixelDeviceZError * RelativeZError;
+	return ParallaxDepthError;
+}
+
+
+//------------------------------------------------------- DEPTH & VELOCITY FETCHES & PROCESSING
+
+/* Loads all the 3x3's depth. */
+CALL_SITE_DEBUGLOC
+void FetchDepth3x3(
+	Texture2D<float> SceneDepthTexture,
+	tsr_short2 InputPixelPos,
+	out float DevizeZNeighborhood[CONFIG_BUTTERFLY_SAMPLES])
+{
+	UNROLL_N(CONFIG_BUTTERFLY_SAMPLES)
+	for (uint i = 0; i < CONFIG_BUTTERFLY_SAMPLES; i++)
+	{
+		tsr_short2 SampleInputPixelPos = GetNeighborhoodPixelPos(i, InputPixelPos);
+		
+		DevizeZNeighborhood[i] = SceneDepthTexture[SampleInputPixelPos];
+	}
+} // FetchDepthVelocity3x3()
+
+/* Loads all the 3x3's depth and velocities. */
+CALL_SITE_DEBUGLOC
+void FetchDepthVelocity3x3(
+	Texture2D<float> SceneDepthTexture,
+	Texture2D<float4> SceneVelocityTexture,
+	tsr_short2 InputPixelPos,
+	out float DevizeZNeighborhood[CONFIG_BUTTERFLY_SAMPLES],
+	out float4 VelocityNeighborhood[CONFIG_BUTTERFLY_SAMPLES])
+{
+	UNROLL_N(CONFIG_BUTTERFLY_SAMPLES)
+	for (uint i = 0; i < CONFIG_BUTTERFLY_SAMPLES; i++)
+	{
+		tsr_short2 SampleInputPixelPos = GetNeighborhoodPixelPos(i, InputPixelPos);
+		
+		DevizeZNeighborhood[i] = SceneDepthTexture[SampleInputPixelPos];
+		VelocityNeighborhood[i] = SceneVelocityTexture[SampleInputPixelPos];
+	}
+} // FetchDepthVelocity3x3()
+
+CALL_SITE_DEBUGLOC
+void FetchAndComputeScreenVelocity(
+	Texture2D<float4> SceneVelocityTexture,
+	tsr_short2 InputPixelPos,
+	float ClosestDeviceZ,
+	tsr_short2 ReprojectionOffset,
+	out float3 ScreenVelocity,
+	out bool bHasPixelAnimation)
+{
+	float4 EncodedVelocity = SceneVelocityTexture[InputPixelPos + ReprojectionOffset];
+
+	#if VELOCITY_ENCODE_HAS_PIXEL_ANIMATION
+		bHasPixelAnimation = DecodeHasPixelAnimationFromVelocityTexture(EncodedVelocity);
+	#endif
+	
+	#if CONFIG_USE_SCENE_VELOCITY
+	BRANCH
+	if (EncodedVelocity.x > 0.0)
+	{
+		ScreenVelocity = DecodeVelocityFromTexture(EncodedVelocity);
+	}
+	else
+	#endif
+	{
+		float2 ScreenPos = ApplyScreenTransform(float2(int2(InputPixelPos)), InputPixelPosToScreenPos);
+		ScreenVelocity = ComputeStaticVelocity(ScreenPos, ClosestDeviceZ);
+	}
+} // FetchAndComputeScreenVelocity()
+
+/** Convert the fetches velocities to pixel velocities. */
+void ComputePixelVelocityNeighborhood(
+	tsr_short2 InputPixelPos,
+	float DevizeZNeighborhood[CONFIG_BUTTERFLY_SAMPLES],
+	float4 VelocityNeighborhood[CONFIG_BUTTERFLY_SAMPLES],
+	out float2 PixelVelocityNeighborhood[CONFIG_BUTTERFLY_SAMPLES],
+	out float  DepthVelocityNeighborhood[CONFIG_BUTTERFLY_SAMPLES],
+	out bool   bIsDrawingVelocityNeighborhood[CONFIG_BUTTERFLY_SAMPLES])
+{
+	bool bHasAnyStaticPixels = false;
+	bool bHasAnyDynamicPixels = false;
+	{
+		UNROLL_N(CONFIG_BUTTERFLY_SAMPLES)
+		for (uint i = 0; i < CONFIG_BUTTERFLY_SAMPLES; i++)
+		{
+			#if CONFIG_USE_SCENE_VELOCITY
+				bool bIsDrawingVelocity = VelocityNeighborhood[i].x > 0.0;
+			#else
+				const bool bIsDrawingVelocity = false;
+			#endif
+			bIsDrawingVelocityNeighborhood[i] = bIsDrawingVelocity;
+			bHasAnyStaticPixels  |= !bIsDrawingVelocity;
+			bHasAnyDynamicPixels |= bIsDrawingVelocity;
+		}
+
+		#if COMPILER_SUPPORTS_WAVE_VOTE
+			bHasAnyStaticPixels  = WaveActiveAnyTrue(bHasAnyStaticPixels);
+			bHasAnyDynamicPixels = WaveActiveAnyTrue(bHasAnyDynamicPixels);
+		#else
+			bHasAnyStaticPixels  = true;
+			bHasAnyDynamicPixels = true;
+		#endif
+	}
+	
+	#if CONFIG_USE_SCENE_VELOCITY
+	BRANCH
+	if (bHasAnyStaticPixels && bHasAnyDynamicPixels)
+	{
+		UNROLL_N(CONFIG_BUTTERFLY_SAMPLES)
+		for (uint i = 0; i < CONFIG_BUTTERFLY_SAMPLES; i++)
+		{
+			tsr_short2 SampleInputPixelPos = GetNeighborhoodPixelPos(i, InputPixelPos);
+			float2 ScreenPos = ApplyScreenTransform(float2(int2(SampleInputPixelPos)), InputPixelPosToScreenPos);
+
+			float3 ScreenVelocity = ComputeStaticVelocity(ScreenPos, DevizeZNeighborhood[i]);
+			FLATTEN
+			if (VelocityNeighborhood[i].x > 0.0)
+			{
+				ScreenVelocity = DecodeVelocityFromTexture(VelocityNeighborhood[i]);
+			}
+
+			PixelVelocityNeighborhood[i] = ScreenVelocity.xy * ScreenVelocityToInputPixelVelocity;
+			DepthVelocityNeighborhood[i] = ScreenVelocity.z;
+		}
+	}
+	else if (bHasAnyDynamicPixels) // && !bHasAnyStaticPixels)
+	{
+		UNROLL_N(CONFIG_BUTTERFLY_SAMPLES)
+		for (uint i = 0; i < CONFIG_BUTTERFLY_SAMPLES; i++)
+		{
+			float3 ScreenVelocity = DecodeVelocityFromTexture(VelocityNeighborhood[i]);
+
+			PixelVelocityNeighborhood[i] = ScreenVelocity.xy * ScreenVelocityToInputPixelVelocity;
+			DepthVelocityNeighborhood[i] = ScreenVelocity.z;
+		}
+	}
+	else // if (bHasAnyStaticPixels && !bHasAnyDynamicPixels
+	#endif // !CONFIG_USE_SCENE_VELOCITY
+	{
+		UNROLL_N(CONFIG_BUTTERFLY_SAMPLES)
+		for (uint i = 0; i < CONFIG_BUTTERFLY_SAMPLES; i++)
+		{
+			tsr_short2 SampleInputPixelPos = GetNeighborhoodPixelPos(i, InputPixelPos);
+			float2 ScreenPos = ApplyScreenTransform(float2(int2(SampleInputPixelPos)), InputPixelPosToScreenPos);
+
+			float3 ScreenVelocity = ComputeStaticVelocity(ScreenPos, DevizeZNeighborhood[i]);
+
+			PixelVelocityNeighborhood[i] = ScreenVelocity.xy * ScreenVelocityToInputPixelVelocity;
+			DepthVelocityNeighborhood[i] = ScreenVelocity.z;
+		}
+	}
+} // ComputePixelVelocityNeighborhood()
+
+
+//------------------------------------------------------- REPROJECTION FIELD FETCH & PROCESSING
+
+CALL_SITE_DEBUGLOC
+void FetchReprojectionNeighborhood(
+	Texture2D<uint> ReprojectionVectorTexture,
+	tsr_short2 InputPixelPos,
+	out uint EncodedReprojectionVectorNeighborhood[CONFIG_BUTTERFLY_SAMPLES])
+{
+	tsr_short2 LaneOffsetSign = GetLaneOffsetSign();
+
+	UNROLL_N(CONFIG_BUTTERFLY_SAMPLES)
+	for (uint i = 0; i < CONFIG_BUTTERFLY_SAMPLES; i++)
+	{
+		tsr_short2 SampleInputPixelPos = GetNeighborhoodPixelPos(i, InputPixelPos);
+		EncodedReprojectionVectorNeighborhood[i] = ReprojectionVectorTexture[SampleInputPixelPos];
+	}
+} // FetchReprojectionNeighborhood()
+
+CALL_SITE_DEBUGLOC
+void ComputePixelVelocityNeighborhood(
+	uint EncodedReprojectionVectorNeighborhood[CONFIG_BUTTERFLY_SAMPLES],
+	out float2 PixelVelocityNeighborhood[CONFIG_BUTTERFLY_SAMPLES])
+{
+	UNROLL_N(CONFIG_BUTTERFLY_SAMPLES)
+	for (uint i = 0; i < CONFIG_BUTTERFLY_SAMPLES; i++)
+	{
+		PixelVelocityNeighborhood[i] = DecodeReprojectionVector(EncodedReprojectionVectorNeighborhood[i]) * ScreenVelocityToInputPixelVelocity;
+	}
+} // ComputePixelVelocityNeighborhood()
+
+
+//------------------------------------------------------- DEPTH ANALYSIS
+
+CALL_SITE_DEBUGLOC
+uint PackDepthAndOffset(const int2 Offset, float DeviceZ)
+{
+	return ((asuint(DeviceZ) << 2) & ~0xF) | (uint(1 + Offset.x) << 0) | (uint(1 + Offset.y) << 2);
+}
+
+CALL_SITE_DEBUGLOC
+void UnpackDepthAndOffset(uint Packed, out tsr_short2 Offset, out float DeviceZ)
+{
+	DeviceZ = asfloat((Packed & ~0xF) >> 2);
+	Offset.x = tsr_short((Packed >> 0u) & 0x3u) - tsr_short(1);
+	Offset.y = tsr_short((Packed >> 2u) & 0x3u) - tsr_short(1);
+}
+
+/** Find the closest DevizeZ and its corresponding offset */
+void FindClosestDepthOffset(
+	float DevizeZNeighborhood[CONFIG_BUTTERFLY_SAMPLES],
+	out float ClosestDeviceZ,
+	out tsr_short2 ReprojectionOffset)
+{
+	float DeviceZ = AccessNeighborhoodCenter(DevizeZNeighborhood);
+	float DeviceZError = ComputePixelDeviceZError(DeviceZ);
+	uint PackedDepthOffset = PackDepthAndOffset(/* Offset = */ int2(0, 0), DeviceZ);
+	
+	#if 0
+		UNROLL_N(4)
+		for (uint i = 0; i < 4; i++)
+	#elif 1
+		UNROLL_N(2)
+		for (uint i = 0; i < 4; i+= 2)
+	#else
+		UNROLL_N(2)
+		for (uint i = 1; i < 4; i+= 2)
+	#endif
+	{
+		const int2 PairOffset = kPairOffsets[i];
+
+		float SampleDeviceZ0 = AccessNeighborhoodStaticOffset(DevizeZNeighborhood, +PairOffset);
+		float SampleDeviceZ1 = AccessNeighborhoodStaticOffset(DevizeZNeighborhood, -PairOffset);
+
+		#if CONFIG_DEPTH_EDGE_DETECTION
+		{
+			float DepthDiff = abs(SampleDeviceZ1 - SampleDeviceZ0);
+			float DepthVariation = (SampleDeviceZ0 + SampleDeviceZ1) * 0.5 - DeviceZ;
+
+			FLATTEN
+			if (DepthVariation > max(DepthDiff * 0.25, DeviceZError))
+			{
+				PackedDepthOffset = max(PackedDepthOffset, PackDepthAndOffset(+PairOffset, SampleDeviceZ0));
+				PackedDepthOffset = max(PackedDepthOffset, PackDepthAndOffset(-PairOffset, SampleDeviceZ1));
+			}
+		}
+		#else
+		{
+			PackedDepthOffset = max(PackedDepthOffset, PackDepthAndOffset(+PairOffset, SampleDeviceZ0));
+			PackedDepthOffset = max(PackedDepthOffset, PackDepthAndOffset(-PairOffset, SampleDeviceZ1));
+		}
+		#endif
+	}
+
+	UnpackDepthAndOffset(PackedDepthOffset, /* out */ ReprojectionOffset, /* out */ ClosestDeviceZ);
+	
+	#if CONFIG_BUTTERFLY_KERNEL
+		tsr_short2 LaneOffsetSign = GetLaneOffsetSign();
+		ReprojectionOffset *= LaneOffsetSign;
+	#endif
+} // FindClosestDepthOffset()
+
+/** Compute the velocity jacobian along a two orthogonal axes, taking a common closest depth. */
+bool ComputeDepthBilaterals(
+	const int2 DirectionE,
+	const int2 DirectionS,
+	bool bCompute,
+	float DevizeZNeighborhood[CONFIG_BUTTERFLY_SAMPLES],
+	inout float DepthError,
+	out tsr_half BilateralWeights[4])
+{
+	float DevizeZC = AccessNeighborhoodCenter(DevizeZNeighborhood);
+	float DevizeZE = AccessNeighborhoodStaticOffset(DevizeZNeighborhood, +DirectionE);
+	float DevizeZW = AccessNeighborhoodStaticOffset(DevizeZNeighborhood, -DirectionE);
+	float DevizeZS = AccessNeighborhoodStaticOffset(DevizeZNeighborhood, +DirectionS);
+	float DevizeZN = AccessNeighborhoodStaticOffset(DevizeZNeighborhood, -DirectionS);
+	
+	tsr_half WeightE;
+	tsr_half WeightW;
+	tsr_half WeightS;
+	tsr_half WeightN;
+		
+	bool bSuccess;
+	BRANCH
+	if (bCompute)
+	{
+		const float WeightOffset = 1.5;
+		const float PixelDistance = length(float2(DirectionE));
+
+		float VariationEW = abs(DevizeZC - (DevizeZE + DevizeZW) * 0.5);
+		float VariationSN = abs(DevizeZC - (DevizeZS + DevizeZN) * 0.5);
+
+		float DepthDiffE = abs(DevizeZE - DevizeZC);
+		float DepthDiffW = abs(DevizeZW - DevizeZC);
+		float DepthDiffS = abs(DevizeZS - DevizeZC);
+		float DepthDiffN = abs(DevizeZN - DevizeZC);
+	
+		float DepthDiffEW = abs(DevizeZE - DevizeZW);
+		float DepthDiffSN = abs(DevizeZS - DevizeZN);
+
+		float MinDepthDerivativeE = min3(DepthDiffEW, DepthDiffE, DepthDiffW);
+		float MinDepthDerivativeS = min3(DepthDiffSN, DepthDiffS, DepthDiffN);
+		
+		float PixelDeviceZError = ComputePixelDeviceZError(DevizeZC) * PixelDistance;
+
+		float FinalDepthDerivativeE = min(max(MinDepthDerivativeE, PixelDeviceZError), max(MinDepthDerivativeS, 8.0 * PixelDeviceZError));
+		float FinalDepthDerivativeS = min(max(MinDepthDerivativeS, PixelDeviceZError), max(MinDepthDerivativeE, 8.0 * PixelDeviceZError));
+
+		WeightE = tsr_half(saturate(WeightOffset - min(DepthDiffE, VariationEW) / FinalDepthDerivativeE));
+		WeightW = tsr_half(saturate(WeightOffset - min(DepthDiffW, VariationEW) / FinalDepthDerivativeE));
+		WeightS = tsr_half(saturate(WeightOffset - min(DepthDiffS, VariationSN) / FinalDepthDerivativeS));
+		WeightN = tsr_half(saturate(WeightOffset - min(DepthDiffN, VariationSN) / FinalDepthDerivativeS));
+
+	
+		tsr_half TotalWeightEW = WeightE + WeightW;
+		tsr_half TotalWeightSN = WeightS + WeightN;
+	
+		bool bSuccessEW = TotalWeightEW > tsr_half(0.0);
+		bool bSuccessSN = TotalWeightSN > tsr_half(0.0);
+
+		tsr_half InvTotalWeightEW = select(bSuccessEW, rcp(TotalWeightEW), tsr_half(0.0));
+		tsr_half InvTotalWeightSN = select(bSuccessSN, rcp(TotalWeightSN), tsr_half(0.0));
+
+		WeightE *= InvTotalWeightEW;
+		WeightW *= InvTotalWeightEW;
+		WeightS *= InvTotalWeightSN;
+		WeightN *= InvTotalWeightSN;
+		bSuccess = bSuccessEW || bSuccessSN;
+		
+		DepthError = max(DepthDiffE * WeightE + DepthDiffW * WeightW, DepthDiffS * WeightS + DepthDiffN * WeightN);
+	}
+	else
+	{
+		WeightE = tsr_half(1.0);
+		WeightW = tsr_half(1.0);
+		WeightS = tsr_half(1.0);
+		WeightN = tsr_half(1.0);
+		bSuccess = false;
+	}
+
+	BilateralWeights[0] = WeightE;
+	BilateralWeights[1] = WeightW;
+	BilateralWeights[2] = WeightS;
+	BilateralWeights[3] = WeightN;
+
+	return bSuccess;
+} // ComputeDepthBilaterals()
+
+/** Compute the depth error on horizontal and vertical. */
+float ComputeDepthError(float DevizeZNeighborhood[CONFIG_BUTTERFLY_SAMPLES])
+{
+	float DepthError;
+	tsr_half BilateralWeights[4];
+
+	ComputeDepthBilaterals(
+		kOffsetE, kOffsetS,
+		/* bCompute = */ true,
+		DevizeZNeighborhood,
+		/* inout */ DepthError,
+		/* out */ BilateralWeights);
+
+	return DepthError;
+}
+
+
+//------------------------------------------------------- VELOCITY ANALYSIS
+
+/** Computes the velocity's jacobian from the depth bilateral weighs. */
+CALL_SITE_DEBUGLOC
+void ComputeVelocityJacobian(
+	const int2 DirectionE,
+	const int2 DirectionS,
+	bool bCompute,
+	float2 PixelVelocityNeighborhood[CONFIG_BUTTERFLY_SAMPLES],
+	tsr_half BilateralWeights[4],
+	inout tsr_half2x2 ReprojectionJacobian)
+{
+	
+	float2 PixelVelocityC = AccessNeighborhoodCenter(PixelVelocityNeighborhood);
+	float2 PixelVelocityE = AccessNeighborhoodStaticOffset(PixelVelocityNeighborhood, +DirectionE);
+	float2 PixelVelocityW = AccessNeighborhoodStaticOffset(PixelVelocityNeighborhood, -DirectionE);
+	float2 PixelVelocityS = AccessNeighborhoodStaticOffset(PixelVelocityNeighborhood, +DirectionS);
+	float2 PixelVelocityN = AccessNeighborhoodStaticOffset(PixelVelocityNeighborhood, -DirectionS);
+	
+	tsr_half WeightE = BilateralWeights[0];
+	tsr_half WeightW = BilateralWeights[1];
+	tsr_half WeightS = BilateralWeights[2];
+	tsr_half WeightN = BilateralWeights[3];
+
+	BRANCH
+	if (bCompute)
+	{
+		tsr_half2 VelocityDeltaE = tsr_half2(PixelVelocityE - PixelVelocityC);
+		tsr_half2 VelocityDeltaW = tsr_half2(PixelVelocityC - PixelVelocityW);
+		tsr_half2 VelocityDeltaS = tsr_half2(PixelVelocityS - PixelVelocityC);
+		tsr_half2 VelocityDeltaN = tsr_half2(PixelVelocityC - PixelVelocityN);
+
+		ReprojectionJacobian[0] = VelocityDeltaE * WeightE + VelocityDeltaW * WeightW;
+		ReprojectionJacobian[1] = VelocityDeltaS * WeightS + VelocityDeltaN * WeightN;
+	}
+} // ComputeVelocityJacobian()
+
+/** Find the closest DevizeZ and its corresponding offset */
+void ComputeReprojectionJacobian(
+	float  DevizeZNeighborhood[CONFIG_BUTTERFLY_SAMPLES],
+	float2 PixelVelocityNeighborhood[CONFIG_BUTTERFLY_SAMPLES],
+	out tsr_half2x2 ReprojectionJacobian,
+	out float DepthError)
+{
+	float DeviceZ = AccessNeighborhoodCenter(DevizeZNeighborhood);
+	
+	ReprojectionJacobian = tsr_half(0.0);
+	DepthError = 0.0;
+
+	// Horizontal & vertical cardinals first as the pixels are closer than on diagonals
+	bool bSuccess;
+	{
+		const bool bCompute = true;
+		const int2 DirectionE = kOffsetE;
+		const int2 DirectionS = kOffsetS;
+
+		tsr_half BilateralWeights[4];
+		bSuccess = ComputeDepthBilaterals(
+			DirectionE, DirectionS,
+			bCompute,
+			DevizeZNeighborhood,
+			/* inout */ DepthError,
+			/* out */ BilateralWeights);
+		
+		ComputeVelocityJacobian(
+			DirectionE, DirectionS,
+			bCompute,
+			PixelVelocityNeighborhood,
+			BilateralWeights,
+			/* inout */ ReprojectionJacobian);
+	}
+	
+	// Compute diagonal jacobian if the horizontal/vertical failed
+#if CONFIG_BUTTERFLY_KERNEL && COMPILER_SUPPORTS_WAVE_VOTE
+	if (WaveActiveAnyTrue(!bSuccess))
+#endif
+	{
+		const bool bCompute = !bSuccess;
+		const int2 DirectionE = kOffsetSE;
+		const int2 DirectionS = kOffsetSW;
+		
+		tsr_half BilateralWeights[4];
+		ComputeDepthBilaterals(
+			DirectionE, DirectionS,
+			bCompute,
+			DevizeZNeighborhood,
+			/* inout */ DepthError,
+			/* out */ BilateralWeights);
+		
+		ComputeVelocityJacobian(
+			DirectionE, DirectionS,
+			bCompute,
+			PixelVelocityNeighborhood,
+			BilateralWeights,
+			/* inout */ ReprojectionJacobian);
+			
+		// Rotate the diagonal jacobian back to horizontal/vertical
+		if (bCompute)
+		{
+			tsr_half2 DifferentialSE = ReprojectionJacobian[0];
+			tsr_half2 DifferentialSW = ReprojectionJacobian[1];
+			ReprojectionJacobian[0] = DifferentialSE * tsr_half(+0.5) + DifferentialSW * tsr_half(-0.5);
+			ReprojectionJacobian[1] = DifferentialSE * tsr_half(+0.5) + DifferentialSW * tsr_half(+0.5);
+		}
+	}
+	
+	// Flips the signs due to the butterfly kernel.
+	#if CONFIG_BUTTERFLY_KERNEL && 1
+		tsr_half2 LaneOffsetSign = tsr_half2(GetLaneOffsetSign());
+		ReprojectionJacobian[0] *= LaneOffsetSign.x;
+		ReprojectionJacobian[1] *= LaneOffsetSign.y;
+	#endif
+} // ComputeReprojectionJacobian()
+
+/** Compute the reprojection edge from non-dilated to dilated screen velocity. */
+CALL_SITE_DEBUGLOC
+tsr_half ComputeReprojectionEdge(
+	float2 ScreenVelocity,
+	float2 DilatedScreenVelocity)
+{
+	tsr_half2 NeighborDeltaPixelVelocity = tsr_half2((DilatedScreenVelocity - ScreenVelocity) * ScreenVelocityToInputPixelVelocity);
+	tsr_half ReprojectionEdge = saturate(tsr_half(1.1) - dot(abs(NeighborDeltaPixelVelocity), tsr_half(1.1).xx));
+	return ReprojectionEdge;
+}
+
+/** Compute the reprojection edge from the pixel velocity neighborhood. */
+tsr_half ComputeReprojectionEdge(
+	float2 PixelVelocityNeighborhood[CONFIG_BUTTERFLY_SAMPLES])
+{
+	tsr_half ReprojectionEdge = tsr_half(0.0);
+
+	// Sample
+	float2 ScreenPixelVelocity = AccessNeighborhoodCenter(PixelVelocityNeighborhood);
+	tsr_short2 LaneOffsetSign = GetLaneOffsetSign();
+
+	UNROLL_N(CONFIG_VELOCITY_SAMPLE_PAIRS)
+	for (uint i = 0; i < CONFIG_VELOCITY_SAMPLE_PAIRS; i++)
+	ISOLATE
+	{
+		const uint NeightbordId0 = i + 1;
+		const uint NeightbordId1 = i + 1 + CONFIG_VELOCITY_SAMPLE_PAIRS;
+			
+		#if CONFIG_VELOCITY_SAMPLES == 9
+			const int2 Offset0 = kOffsets3x3[kSquareIndexes3x3[NeightbordId0]];
+			const int2 Offset1 = kOffsets3x3[kSquareIndexes3x3[NeightbordId1]];
+		#elif CONFIG_VELOCITY_SAMPLES == 5
+			const int2 Offset0 = kOffsets3x3[kPlusIndexes3x3[NeightbordId0]];
+			const int2 Offset1 = kOffsets3x3[kPlusIndexes3x3[NeightbordId1]];
+		#endif
+
+		float2 NeighborPixelVelocity0 = AccessNeighborhoodStaticOffset(PixelVelocityNeighborhood, Offset0);
+		float2 NeighborPixelVelocity1 = AccessNeighborhoodStaticOffset(PixelVelocityNeighborhood, Offset1);
+
+		tsr_half2 NeighborDeltaPixelVelocity0 = tsr_half2(ScreenPixelVelocity - NeighborPixelVelocity0);
+		tsr_half2 NeighborDeltaPixelVelocity1 = tsr_half2(ScreenPixelVelocity - NeighborPixelVelocity1);
+			
+		// Reject high frequency on reprojection edge discontinuity.
+		{
+			tsr_half IsOverlapingEdge0 = tsr_half(1.0) - dot(saturate(abs(NeighborDeltaPixelVelocity0 * tsr_half2(tsr_short2(Offset0) * LaneOffsetSign))), tsr_half(1.0).xx); // * saturate(ScreenPixelVelocityLength - NeighborPixelVelocityLength0);
+			tsr_half IsOverlapingEdge1 = tsr_half(1.0) - dot(saturate(abs(NeighborDeltaPixelVelocity1 * tsr_half2(tsr_short2(Offset1) * LaneOffsetSign))), tsr_half(1.0).xx); // * saturate(ScreenPixelVelocityLength - NeighborPixelVelocityLength1);
+				
+			tsr_half IsOverlapingEdge = min(IsOverlapingEdge0, IsOverlapingEdge1);
+
+			if (i == 0)
+				ReprojectionEdge = IsOverlapingEdge;
+			else
+				ReprojectionEdge = min(ReprojectionEdge, IsOverlapingEdge);
+		}
+	} // for (uint i = 0; i < CONFIG_VELOCITY_SAMPLE_PAIRS; i++)
+
+	// Remove noise from high frequency rejections.
+	ReprojectionEdge = saturate(ReprojectionEdge * tsr_half(1.1));
+	return ReprojectionEdge;
+} // ComputeReprojectionEdge()
diff --git a/Engine/Shaders/Private/TemporalSuperResolution/TSRDilateVelocity.usf b/Engine/Shaders/Private/TemporalSuperResolution/TSRDilateVelocity.usf
--- a/Engine/Shaders/Private/TemporalSuperResolution/TSRDilateVelocity.usf
+++ b/Engine/Shaders/Private/TemporalSuperResolution/TSRDilateVelocity.usf
@@ -1,8 +1,9 @@
 // Copyright Epic Games, Inc. All Rights Reserved.
 
-#include "TSRCommon.ush"
+#include "TSRDepthVelocityAnalysis.ush"
 #include "TSRSpatialAntiAliasing.ush"
 #include "TSRReprojectionField.ush"
+#include "TSRClosestOccluder.ush"
 
 #if DIM_MOTION_BLUR_DIRECTIONS > 0
 	#define CONFIG_MAX_RANGE_SIZE (DIM_MOTION_BLUR_DIRECTIONS)
@@ -18,28 +19,10 @@
 	#define TILE_SIZE 8
 #endif
 
-/** Whether to use the velocity from SceneVelocityTexture. */
-#define CONFIG_USE_SCENE_VELOCITY 1
-
-/** Whether the dilation should be only on depth edges, improving texture reprojection accuracy on flat surfaces. */
-#define CONFIG_DEPTH_EDGE_DETECTION 1
-
-
-//------------------------------------------------------- CONSTANTS
-
-static const int2 kPairOffsets[4] = {
-	int2( 1,  0),
-	int2( 1,  1),
-	int2( 0,  1),
-	int2(-1,  1),
-};
-
 
 //------------------------------------------------------- PARAMETERS
 
 float4x4 RotationalClipToPrevClip;
-float2 PrevOutputBufferUVMin;
-float2 PrevOutputBufferUVMax;
 float InvFlickeringMaxParralaxVelocity;
 float ReprojectionFieldAntiAliasVelocityThreshold;
 uint bReprojectionField;
@@ -64,429 +47,8 @@ RWTexture2DArray<float4> VelocityTileArrayOutput;
 static float4 Debug[DEBUG_ARRAY_SIZE];
 
 
-//------------------------------------------------------- UTILITY
-
-uint PackDepthAndOffset(const int2 Offset, float DeviceZ)
-{
-	return ((asuint(DeviceZ) << 2) & ~0xF) | (uint(1 + Offset.x) << 0) | (uint(1 + Offset.y) << 2);
-}
-
-void UnpackDepthAndOffset(uint Packed, out tsr_short2 Offset, out float DeviceZ)
-{
-	DeviceZ = asfloat((Packed & ~0xF) >> 2);
-	Offset.x = tsr_short((Packed >> 0u) & 0x3u) - tsr_short(1);
-	Offset.y = tsr_short((Packed >> 2u) & 0x3u) - tsr_short(1);
-}
-
-
 //------------------------------------------------------- FUNCTIONS
 
-/* Loads all the 3x3's depth. */
-void FetchDepth3x3(
-	tsr_short2 InputPixelPos,
-	out float DevizeZNeighborhood[CONFIG_BUTTERFLY_SAMPLES])
-{
-	UNROLL_N(CONFIG_BUTTERFLY_SAMPLES)
-	for (uint i = 0; i < CONFIG_BUTTERFLY_SAMPLES; i++)
-	{
-		tsr_short2 SampleInputPixelPos = GetNeighborhoodPixelPos(i, InputPixelPos);
-		
-		DevizeZNeighborhood[i] = SceneDepthTexture[SampleInputPixelPos];
-	}
-} // FetchDepthVelocity3x3()
-
-/* Loads all the 3x3's depth and velocities. */
-void FetchDepthVelocity3x3(
-	tsr_short2 InputPixelPos,
-	out float DevizeZNeighborhood[CONFIG_BUTTERFLY_SAMPLES],
-	out float4 VelocityNeighborhood[CONFIG_BUTTERFLY_SAMPLES])
-{
-	UNROLL_N(CONFIG_BUTTERFLY_SAMPLES)
-	for (uint i = 0; i < CONFIG_BUTTERFLY_SAMPLES; i++)
-	{
-		tsr_short2 SampleInputPixelPos = GetNeighborhoodPixelPos(i, InputPixelPos);
-		
-		DevizeZNeighborhood[i] = SceneDepthTexture[SampleInputPixelPos];
-		VelocityNeighborhood[i] = SceneVelocityTexture[SampleInputPixelPos];
-	}
-} // FetchDepthVelocity3x3()
-
-/** Convert the fetches velocities to pixel velocities. */
-void ComputePixelVelocityNeighborhood(
-	tsr_short2 InputPixelPos,
-	float DevizeZNeighborhood[CONFIG_BUTTERFLY_SAMPLES],
-	float4 VelocityNeighborhood[CONFIG_BUTTERFLY_SAMPLES],
-	out float2 PixelVelocityNeighborhood[CONFIG_BUTTERFLY_SAMPLES],
-	out float  DepthVelocityNeighborhood[CONFIG_BUTTERFLY_SAMPLES],
-	out bool   bIsDrawingVelocityNeighborhood[CONFIG_BUTTERFLY_SAMPLES])
-{
-	bool bHasAnyStaticPixels = false;
-	bool bHasAnyDynamicPixels = false;
-	{
-		UNROLL_N(CONFIG_BUTTERFLY_SAMPLES)
-		for (uint i = 0; i < CONFIG_BUTTERFLY_SAMPLES; i++)
-		{
-			#if CONFIG_USE_SCENE_VELOCITY
-				bool bIsDrawingVelocity = VelocityNeighborhood[i].x > 0.0;
-			#else
-				const bool bIsDrawingVelocity = false;
-			#endif
-			bIsDrawingVelocityNeighborhood[i] = bIsDrawingVelocity;
-			bHasAnyStaticPixels  |= !bIsDrawingVelocity;
-			bHasAnyDynamicPixels |= bIsDrawingVelocity;
-		}
-
-		#if COMPILER_SUPPORTS_WAVE_VOTE
-			bHasAnyStaticPixels  = WaveActiveAnyTrue(bHasAnyStaticPixels);
-			bHasAnyDynamicPixels = WaveActiveAnyTrue(bHasAnyDynamicPixels);
-		#else
-			bHasAnyStaticPixels  = true;
-			bHasAnyDynamicPixels = true;
-		#endif
-	}
-	
-	#if CONFIG_USE_SCENE_VELOCITY
-	BRANCH
-	if (bHasAnyStaticPixels && bHasAnyDynamicPixels)
-	{
-		UNROLL_N(CONFIG_BUTTERFLY_SAMPLES)
-		for (uint i = 0; i < CONFIG_BUTTERFLY_SAMPLES; i++)
-		{
-			tsr_short2 SampleInputPixelPos = GetNeighborhoodPixelPos(i, InputPixelPos);
-			float2 ScreenPos = ApplyScreenTransform(float2(int2(SampleInputPixelPos)), InputPixelPosToScreenPos);
-
-			float3 ScreenVelocity = ComputeStaticVelocity(ScreenPos, DevizeZNeighborhood[i]);
-			FLATTEN
-			if (VelocityNeighborhood[i].x > 0.0)
-			{
-				ScreenVelocity = DecodeVelocityFromTexture(VelocityNeighborhood[i]);
-			}
-
-			PixelVelocityNeighborhood[i] = ScreenVelocity.xy * ScreenVelocityToInputPixelVelocity;
-			DepthVelocityNeighborhood[i] = ScreenVelocity.z;
-		}
-	}
-	else if (bHasAnyDynamicPixels) // && !bHasAnyStaticPixels)
-	{
-		UNROLL_N(CONFIG_BUTTERFLY_SAMPLES)
-		for (uint i = 0; i < CONFIG_BUTTERFLY_SAMPLES; i++)
-		{
-			float3 ScreenVelocity = DecodeVelocityFromTexture(VelocityNeighborhood[i]);
-
-			PixelVelocityNeighborhood[i] = ScreenVelocity.xy * ScreenVelocityToInputPixelVelocity;
-			DepthVelocityNeighborhood[i] = ScreenVelocity.z;
-		}
-	}
-	else // if (bHasAnyStaticPixels && !bHasAnyDynamicPixels
-	#endif // !CONFIG_USE_SCENE_VELOCITY
-	{
-		UNROLL_N(CONFIG_BUTTERFLY_SAMPLES)
-		for (uint i = 0; i < CONFIG_BUTTERFLY_SAMPLES; i++)
-		{
-			tsr_short2 SampleInputPixelPos = GetNeighborhoodPixelPos(i, InputPixelPos);
-			float2 ScreenPos = ApplyScreenTransform(float2(int2(SampleInputPixelPos)), InputPixelPosToScreenPos);
-
-			float3 ScreenVelocity = ComputeStaticVelocity(ScreenPos, DevizeZNeighborhood[i]);
-
-			PixelVelocityNeighborhood[i] = ScreenVelocity.xy * ScreenVelocityToInputPixelVelocity;
-			DepthVelocityNeighborhood[i] = ScreenVelocity.z;
-		}
-	}
-} // ComputePixelVelocityNeighborhood()
-
-/** Find the closest DevizeZ and its corresponding offset */
-void FindClosestDepthOffset(
-	float DevizeZNeighborhood[CONFIG_BUTTERFLY_SAMPLES],
-	out float ClosestDeviceZ,
-	out tsr_short2 ReprojectionOffset)
-{
-	float DeviceZ = AccessNeighborhoodCenter(DevizeZNeighborhood);
-	float DeviceZError = ComputePixelDeviceZError(DeviceZ);
-	uint PackedDepthOffset = PackDepthAndOffset(/* Offset = */ int2(0, 0), DeviceZ);
-	
-	#if 0
-		UNROLL_N(4)
-		for (uint i = 0; i < 4; i++)
-	#elif 1
-		UNROLL_N(2)
-		for (uint i = 0; i < 4; i+= 2)
-	#else
-		UNROLL_N(2)
-		for (uint i = 1; i < 4; i+= 2)
-	#endif
-	{
-		const int2 PairOffset = kPairOffsets[i];
-
-		float SampleDeviceZ0 = AccessNeighborhoodStaticOffset(DevizeZNeighborhood, +PairOffset);
-		float SampleDeviceZ1 = AccessNeighborhoodStaticOffset(DevizeZNeighborhood, -PairOffset);
-
-		#if CONFIG_DEPTH_EDGE_DETECTION
-		{
-			float DepthDiff = abs(SampleDeviceZ1 - SampleDeviceZ0);
-			float DepthVariation = (SampleDeviceZ0 + SampleDeviceZ1) * 0.5 - DeviceZ;
-
-			FLATTEN
-			if (DepthVariation > max(DepthDiff * 0.25, DeviceZError))
-			{
-				PackedDepthOffset = max(PackedDepthOffset, PackDepthAndOffset(+PairOffset, SampleDeviceZ0));
-				PackedDepthOffset = max(PackedDepthOffset, PackDepthAndOffset(-PairOffset, SampleDeviceZ1));
-			}
-		}
-		#else
-		{
-			PackedDepthOffset = max(PackedDepthOffset, PackDepthAndOffset(+PairOffset, SampleDeviceZ0));
-			PackedDepthOffset = max(PackedDepthOffset, PackDepthAndOffset(-PairOffset, SampleDeviceZ1));
-		}
-		#endif
-	}
-
-	UnpackDepthAndOffset(PackedDepthOffset, /* out */ ReprojectionOffset, /* out */ ClosestDeviceZ);
-	
-	#if CONFIG_BUTTERFLY_KERNEL
-		tsr_short2 LaneOffsetSign = GetLaneOffsetSign();
-		ReprojectionOffset *= LaneOffsetSign;
-	#endif
-} // FindClosestDepthOffset()
-
-/** Compute the velocity jacobian along a two orthogonal axes, taking a common closest depth. */
-bool ComputeDepthBilaterals(
-	const int2 DirectionE,
-	const int2 DirectionS,
-	bool bCompute,
-	float DevizeZNeighborhood[CONFIG_BUTTERFLY_SAMPLES],
-	inout float DepthError,
-	out tsr_half BilateralWeights[4])
-{
-	float DevizeZC = AccessNeighborhoodCenter(DevizeZNeighborhood);
-	float DevizeZE = AccessNeighborhoodStaticOffset(DevizeZNeighborhood, +DirectionE);
-	float DevizeZW = AccessNeighborhoodStaticOffset(DevizeZNeighborhood, -DirectionE);
-	float DevizeZS = AccessNeighborhoodStaticOffset(DevizeZNeighborhood, +DirectionS);
-	float DevizeZN = AccessNeighborhoodStaticOffset(DevizeZNeighborhood, -DirectionS);
-	
-	tsr_half WeightE;
-	tsr_half WeightW;
-	tsr_half WeightS;
-	tsr_half WeightN;
-		
-	bool bSuccess;
-	BRANCH
-	if (bCompute)
-	{
-		const float WeightOffset = 1.5;
-		const float PixelDistance = length(float2(DirectionE));
-
-		float VariationEW = abs(DevizeZC - (DevizeZE + DevizeZW) * 0.5);
-		float VariationSN = abs(DevizeZC - (DevizeZS + DevizeZN) * 0.5);
-
-		float DepthDiffE = abs(DevizeZE - DevizeZC);
-		float DepthDiffW = abs(DevizeZW - DevizeZC);
-		float DepthDiffS = abs(DevizeZS - DevizeZC);
-		float DepthDiffN = abs(DevizeZN - DevizeZC);
-	
-		float DepthDiffEW = abs(DevizeZE - DevizeZW);
-		float DepthDiffSN = abs(DevizeZS - DevizeZN);
-
-		float MinDepthDerivativeE = min3(DepthDiffEW, DepthDiffE, DepthDiffW);
-		float MinDepthDerivativeS = min3(DepthDiffSN, DepthDiffS, DepthDiffN);
-		
-		float PixelDeviceZError = ComputePixelDeviceZError(DevizeZC) * PixelDistance;
-
-		float FinalDepthDerivativeE = min(max(MinDepthDerivativeE, PixelDeviceZError), max(MinDepthDerivativeS, 8.0 * PixelDeviceZError));
-		float FinalDepthDerivativeS = min(max(MinDepthDerivativeS, PixelDeviceZError), max(MinDepthDerivativeE, 8.0 * PixelDeviceZError));
-
-		WeightE = tsr_half(saturate(WeightOffset - min(DepthDiffE, VariationEW) / FinalDepthDerivativeE));
-		WeightW = tsr_half(saturate(WeightOffset - min(DepthDiffW, VariationEW) / FinalDepthDerivativeE));
-		WeightS = tsr_half(saturate(WeightOffset - min(DepthDiffS, VariationSN) / FinalDepthDerivativeS));
-		WeightN = tsr_half(saturate(WeightOffset - min(DepthDiffN, VariationSN) / FinalDepthDerivativeS));
-
-	
-		tsr_half TotalWeightEW = WeightE + WeightW;
-		tsr_half TotalWeightSN = WeightS + WeightN;
-	
-		bool bSuccessEW = TotalWeightEW > tsr_half(0.0);
-		bool bSuccessSN = TotalWeightSN > tsr_half(0.0);
-
-		tsr_half InvTotalWeightEW = select(bSuccessEW, rcp(TotalWeightEW), tsr_half(0.0));
-		tsr_half InvTotalWeightSN = select(bSuccessSN, rcp(TotalWeightSN), tsr_half(0.0));
-
-		WeightE *= InvTotalWeightEW;
-		WeightW *= InvTotalWeightEW;
-		WeightS *= InvTotalWeightSN;
-		WeightN *= InvTotalWeightSN;
-		bSuccess = bSuccessEW || bSuccessSN;
-		
-		DepthError = max(DepthDiffE * WeightE + DepthDiffW * WeightW, DepthDiffS * WeightS + DepthDiffN * WeightN);
-	}
-	else
-	{
-		WeightE = tsr_half(1.0);
-		WeightW = tsr_half(1.0);
-		WeightS = tsr_half(1.0);
-		WeightN = tsr_half(1.0);
-		bSuccess = false;
-	}
-
-	BilateralWeights[0] = WeightE;
-	BilateralWeights[1] = WeightW;
-	BilateralWeights[2] = WeightS;
-	BilateralWeights[3] = WeightN;
-
-	return bSuccess;
-} // ComputeDepthBilaterals()
-
-void ComputeVelocityJacobian(
-	const int2 DirectionE,
-	const int2 DirectionS,
-	bool bCompute,
-	float2 PixelVelocityNeighborhood[CONFIG_BUTTERFLY_SAMPLES],
-	tsr_half BilateralWeights[4],
-	inout tsr_half2x2 ReprojectionJacobian)
-{
-	
-	float2 PixelVelocityC = AccessNeighborhoodCenter(PixelVelocityNeighborhood);
-	float2 PixelVelocityE = AccessNeighborhoodStaticOffset(PixelVelocityNeighborhood, +DirectionE);
-	float2 PixelVelocityW = AccessNeighborhoodStaticOffset(PixelVelocityNeighborhood, -DirectionE);
-	float2 PixelVelocityS = AccessNeighborhoodStaticOffset(PixelVelocityNeighborhood, +DirectionS);
-	float2 PixelVelocityN = AccessNeighborhoodStaticOffset(PixelVelocityNeighborhood, -DirectionS);
-	
-	tsr_half WeightE = BilateralWeights[0];
-	tsr_half WeightW = BilateralWeights[1];
-	tsr_half WeightS = BilateralWeights[2];
-	tsr_half WeightN = BilateralWeights[3];
-
-	BRANCH
-	if (bCompute)
-	{
-		tsr_half2 VelocityDeltaE = tsr_half2(PixelVelocityE - PixelVelocityC);
-		tsr_half2 VelocityDeltaW = tsr_half2(PixelVelocityC - PixelVelocityW);
-		tsr_half2 VelocityDeltaS = tsr_half2(PixelVelocityS - PixelVelocityC);
-		tsr_half2 VelocityDeltaN = tsr_half2(PixelVelocityC - PixelVelocityN);
-
-		ReprojectionJacobian[0] = VelocityDeltaE * WeightE + VelocityDeltaW * WeightW;
-		ReprojectionJacobian[1] = VelocityDeltaS * WeightS + VelocityDeltaN * WeightN;
-	}
-} // ComputeVelocityJacobian()
-
-/** Compute the depth error on horizontal and vertical. */
-float ComputeDepthError(float DevizeZNeighborhood[CONFIG_BUTTERFLY_SAMPLES])
-{
-	float DepthError;
-	tsr_half BilateralWeights[4];
-
-	ComputeDepthBilaterals(
-		kOffsetE, kOffsetS,
-		/* bCompute = */ true,
-		DevizeZNeighborhood,
-		/* inout */ DepthError,
-		/* out */ BilateralWeights);
-
-	return DepthError;
-}
-
-/** Find the closest DevizeZ and its corresponding offset */
-void ComputeReprojectionJacobian(
-	float  DevizeZNeighborhood[CONFIG_BUTTERFLY_SAMPLES],
-	float2 PixelVelocityNeighborhood[CONFIG_BUTTERFLY_SAMPLES],
-	out tsr_half2x2 ReprojectionJacobian,
-	out float DepthError)
-{
-	float DeviceZ = AccessNeighborhoodCenter(DevizeZNeighborhood);
-	
-	ReprojectionJacobian = tsr_half(0.0);
-	DepthError = 0.0;
-
-	// Horizontal & vertical cardinals first as the pixels are closer than on diagonals
-	bool bSuccess;
-	{
-		const bool bCompute = true;
-		const int2 DirectionE = kOffsetE;
-		const int2 DirectionS = kOffsetS;
-
-		tsr_half BilateralWeights[4];
-		bSuccess = ComputeDepthBilaterals(
-			DirectionE, DirectionS,
-			bCompute,
-			DevizeZNeighborhood,
-			/* inout */ DepthError,
-			/* out */ BilateralWeights);
-		
-		ComputeVelocityJacobian(
-			DirectionE, DirectionS,
-			bCompute,
-			PixelVelocityNeighborhood,
-			BilateralWeights,
-			/* inout */ ReprojectionJacobian);
-	}
-	
-	// Compute diagonal jacobian if the horizontal/vertical failed
-#if CONFIG_BUTTERFLY_KERNEL && COMPILER_SUPPORTS_WAVE_VOTE
-	if (WaveActiveAnyTrue(!bSuccess))
-#endif
-	{
-		const bool bCompute = !bSuccess;
-		const int2 DirectionE = kOffsetSE;
-		const int2 DirectionS = kOffsetSW;
-		
-		tsr_half BilateralWeights[4];
-		ComputeDepthBilaterals(
-			DirectionE, DirectionS,
-			bCompute,
-			DevizeZNeighborhood,
-			/* inout */ DepthError,
-			/* out */ BilateralWeights);
-		
-		ComputeVelocityJacobian(
-			DirectionE, DirectionS,
-			bCompute,
-			PixelVelocityNeighborhood,
-			BilateralWeights,
-			/* inout */ ReprojectionJacobian);
-			
-		// Rotate the diagonal jacobian back to horizontal/vertical
-		if (bCompute)
-		{
-			tsr_half2 DifferentialSE = ReprojectionJacobian[0];
-			tsr_half2 DifferentialSW = ReprojectionJacobian[1];
-			ReprojectionJacobian[0] = DifferentialSE * tsr_half(+0.5) + DifferentialSW * tsr_half(-0.5);
-			ReprojectionJacobian[1] = DifferentialSE * tsr_half(+0.5) + DifferentialSW * tsr_half(+0.5);
-		}
-	}
-	
-	// Flips the signs due to the butterfly kernel.
-	#if CONFIG_BUTTERFLY_KERNEL && 1
-		tsr_half2 LaneOffsetSign = tsr_half2(GetLaneOffsetSign());
-		ReprojectionJacobian[0] *= LaneOffsetSign.x;
-		ReprojectionJacobian[1] *= LaneOffsetSign.y;
-	#endif
-} // ComputeReprojectionJacobian()
-
-void FetchAndComputeScreenVelocity(
-	tsr_short2 InputPixelPos,
-	float ClosestDeviceZ,
-	tsr_short2 ReprojectionOffset,
-	out float3 ScreenVelocity,
-	out bool bHasPixelAnimation)
-{
-	float4 EncodedVelocity = SceneVelocityTexture[InputPixelPos + ReprojectionOffset];
-
-	#if VELOCITY_ENCODE_HAS_PIXEL_ANIMATION
-		bHasPixelAnimation = DecodeHasPixelAnimationFromVelocityTexture(EncodedVelocity);
-	#endif
-	
-	#if CONFIG_USE_SCENE_VELOCITY
-	BRANCH
-	if (EncodedVelocity.x > 0.0)
-	{
-		ScreenVelocity = DecodeVelocityFromTexture(EncodedVelocity);
-	}
-	else
-	#endif
-	{
-		float2 ScreenPos = ApplyScreenTransform(float2(int2(InputPixelPos)), InputPixelPosToScreenPos);
-		ScreenVelocity = ComputeStaticVelocity(ScreenPos, ClosestDeviceZ);
-	}
-} // FetchAndComputeScreenVelocity()
-
 /** Spatial anti-alias the depth buffer and return the boundary. */
 tsr_half2 ComputeReprojectionBoundary(
 	tsr_short2 InputPixelPos,
@@ -575,65 +137,6 @@ tsr_half2 ComputeReprojectionBoundary(
 	return ReprojectionBoundary;
 } // ComputeReprojectionBoundary()
 
-void ScatterClosestOccluder(bool bValidOutputPixel, float2 ScreenPos, float2 FinalScreenVelocity, float PrevClosestDeviceZ)
-{
-	const float PrevUseCountQuantization = 63.0;
-
-	float2 PrevScreenPos = ScreenPos - FinalScreenVelocity.xy;
-	float2 PrevInputBufferUV = (InputInfo_ScreenPosToViewportScale * PrevScreenPos + InputInfo_ScreenPosToViewportBias) * InputInfo_ExtentInverse;
-
-	uint EncodedHoleFillingVelocity = EncodeHoleFillingVelocity(FinalScreenVelocity.xy * ScreenVelocityToInputPixelVelocity);
-
-	bool bValidHistoryCoord = bValidOutputPixel && all(and(PrevInputBufferUV > PrevOutputBufferUVMin, PrevInputBufferUV < PrevOutputBufferUVMax));
-
-	FBilinearSampleInfos BilinearInter = GetBilinearSampleLevelInfos(PrevInputBufferUV, InputInfo_Extent, InputInfo_ExtentInverse);
-	
-	#if CONFIG_SUPPORT_ORTHO_VIEW
-	BRANCH
-	if (IsOrthoProjection())
-	{
-		uint Depth = clamp(uint(round(PrevClosestDeviceZ * 8388607.0)), 0, 8388607) << 9;
-
-		uint PrevClosestDepth0 = Depth | ((EncodedHoleFillingVelocity >> 9u) & 0x1FFu);
-		uint PrevClosestDepth1 = Depth | ((EncodedHoleFillingVelocity >> 0u) & 0x1FFu);
-		
-		UNROLL_N(4)
-		for (uint i = 0; i < 4; i++)
-		{
-			float Weight = GetSampleWeight(BilinearInter, i);
-			uint2 PixelPos = GetSamplePixelCoord(BilinearInter, i);
-
-			uint FinalAdd = uint(round(Weight * PrevUseCountQuantization));
-			
-			PixelPos.x = ((FinalAdd > 0) && bValidHistoryCoord) ? PixelPos.x : (~0);
-
-			uint PreviousDepth0;
-			InterlockedMax(PrevAtomicOutput[tsr_ushort3(PixelPos, 0)], PrevClosestDepth0, /* out */ PreviousDepth0);
-
-			PixelPos.x = select(PreviousDepth0 < PrevClosestDepth0, PixelPos.x, ~0u);
-			InterlockedMax(PrevAtomicOutput[tsr_ushort3(PixelPos, 1)], PrevClosestDepth1);
-		}
-	}
-	else
-	#endif // CONFIG_SUPPORT_ORTHO_VIEW
-	{
-		uint PrevClosestDepth = (f32tof16(PrevClosestDeviceZ) << VELOCITY_HOLE_FILLING_BITS) | EncodedHoleFillingVelocity;
-			
-		UNROLL_N(4)
-		for (uint i = 0; i < 4; i++)
-		{
-			float Weight = GetSampleWeight(BilinearInter, i);
-			uint2 PixelPos = GetSamplePixelCoord(BilinearInter, i);
-
-			uint FinalAdd = uint(round(Weight * PrevUseCountQuantization));
-			
-			PixelPos.x = ((FinalAdd > 0) && bValidHistoryCoord) ? PixelPos.x : (~0);
-
-			InterlockedMax(PrevAtomicOutput[tsr_ushort3(PixelPos, 0)], PrevClosestDepth);
-		}
-	}
-} // ScatterClosestOccluder()
-		
 
 //------------------------------------------------------- ENTRY POINT
 
@@ -661,7 +164,7 @@ void MainCS(
 	uint EncodedReprojectionJacobian;
 	uint EncodedDeviceZError;
 	bool bHasPixelAnimation;
-	tsr_half VelocityEdge;
+	tsr_half ReprojectionEdge;
 	BRANCH
 	if (bReprojectionField)
 	{
@@ -670,7 +173,7 @@ void MainCS(
 		float4 VelocityNeighborhood[CONFIG_BUTTERFLY_SAMPLES];
 		ISOLATE
 		{
-			FetchDepthVelocity3x3(InputPixelPos, /* out */ DevizeZNeighborhood, /* out */ VelocityNeighborhood);
+			FetchDepthVelocity3x3(SceneDepthTexture, SceneVelocityTexture, InputPixelPos, /* out */ DevizeZNeighborhood, /* out */ VelocityNeighborhood);
 		}
 
 		// Convert velocities to pixel velocities to take down register pressure.
@@ -722,6 +225,7 @@ void MainCS(
 			// Compute the dilated screen velocity.
 			bool bDildatedHasPixelAnimation;
 			FetchAndComputeScreenVelocity(
+				SceneVelocityTexture,
 				InputPixelPos,
 				ClosestDeviceZ,
 				ReprojectionOffset,
@@ -729,8 +233,7 @@ void MainCS(
 				/* out */ bDildatedHasPixelAnimation);
 		}
 	
-		tsr_half2 NeighborDeltaPixelVelocity = tsr_half2((DilatedScreenVelocity.xy - ScreenVelocity) * ScreenVelocityToInputPixelVelocity);
-		VelocityEdge = saturate(tsr_half(1.1) - dot(abs(NeighborDeltaPixelVelocity), tsr_half(1.1).xx));
+		ReprojectionEdge = ComputeReprojectionEdge(ScreenVelocity, DilatedScreenVelocity.xy);
 		
 		// Compute the reprojection boundary using a depth spatial anti-aliaser
 		ISOLATE
@@ -757,7 +260,7 @@ void MainCS(
 		float DevizeZNeighborhood[CONFIG_BUTTERFLY_SAMPLES];
 		ISOLATE
 		{
-			FetchDepth3x3(InputPixelPos, /* out */ DevizeZNeighborhood);
+			FetchDepth3x3(SceneDepthTexture, InputPixelPos, /* out */ DevizeZNeighborhood);
 		}
 		
 		tsr_short2 ReprojectionOffset;
@@ -775,6 +278,7 @@ void MainCS(
 		// Compute the screen vecocity.
 		float3 ScreenVelocity;
 		FetchAndComputeScreenVelocity(
+			SceneVelocityTexture,
 			InputPixelPos,
 			ClosestDeviceZ,
 			/* ReprojectionOffset = */ tsr_short2(0, 0),
@@ -784,14 +288,14 @@ void MainCS(
 		// Compute the dilated screen vecocity.
 		bool bDilatedHasPixelAnimation;
 		FetchAndComputeScreenVelocity(
+			SceneVelocityTexture,
 			InputPixelPos,
 			ClosestDeviceZ,
 			ReprojectionOffset,
 			/* out */ DilatedScreenVelocity,
 			/* out */ bDilatedHasPixelAnimation);
 			
-		tsr_half2 NeighborDeltaPixelVelocity = tsr_half2((DilatedScreenVelocity.xy - ScreenVelocity.xy) * ScreenVelocityToInputPixelVelocity);
-		VelocityEdge = saturate(tsr_half(1.1) - dot(abs(NeighborDeltaPixelVelocity), tsr_half(1.1).xx));
+		ReprojectionEdge = ComputeReprojectionEdge(ScreenVelocity.xy, DilatedScreenVelocity.xy);
 		
 		// Encode the device z error.
 		{
@@ -853,7 +357,7 @@ void MainCS(
 
 		ClosestDepthOutput[OutputPixelPos] = float2(select(bHasPixelAnimation, -PrevClosestDeviceZ, PrevClosestDeviceZ), ClosestDeviceZ);
 
-		R8Output[uint3(OutputPixelPos, 0)] = uint(round(VelocityEdge * 255.0));
+		R8Output[uint3(OutputPixelPos, 0)] = uint(round(ReprojectionEdge * 255.0));
 		R8Output[uint3(OutputPixelPos, 1)] = EncodedDeviceZError;
 		BRANCH
 		if (bOutputIsMovingTexture)
@@ -882,7 +386,7 @@ void MainCS(
 		
 		// Scatter parallax rejection
 		float2 ScreenPos = ApplyScreenTransform(float2(int2(OutputPixelPos)), InputPixelPosToScreenPos);
-		ScatterClosestOccluder(bValidOutputPixel, ScreenPos, DilatedScreenVelocity.xy, PrevClosestDeviceZ);
+		ScatterClosestOccluder(PrevAtomicOutput, bValidOutputPixel, ScreenPos, DilatedScreenVelocity.xy, PrevClosestDeviceZ);
 		
 		#if DEBUG_OUTPUT
 		for (uint DebugId = 0; DebugId < DEBUG_ARRAY_SIZE; DebugId++)
diff --git a/Engine/Shaders/Private/TemporalSuperResolution/TSRKernels.ush b/Engine/Shaders/Private/TemporalSuperResolution/TSRKernels.ush
new file mode 100644
--- /dev/null
+++ b/Engine/Shaders/Private/TemporalSuperResolution/TSRKernels.ush
@@ -0,0 +1,455 @@
+// Copyright Epic Games, Inc. All Rights Reserved.
+
+#pragma once
+
+#include "TSRCommon.ush"
+
+
+//------------------------------------------------------- CONSTANTS
+
+// K = Center of the nearest input pixel.
+// O = Center of the output pixel.
+//
+//          |           |
+//    0     |     1     |     2
+//          |           |
+//          |           |
+//  --------+-----------+--------
+//          |           |
+//          | O         |
+//    3     |     K     |     5
+//          |           |
+//          |           |
+//  --------+-----------+--------
+//          |           |
+//          |           |
+//    6     |     7     |     8
+//          |           |
+//
+static const int2 kOffsets3x3[9] =
+{
+	int2(-1, -1),
+	int2(0, -1),
+	int2(1, -1),
+	int2(-1,  0),
+	int2(0,  0), // K
+	int2(1,  0),
+	int2(-1,  1),
+	int2(0,  1),
+	int2(1,  1),
+};
+
+// T = Center of the nearest top left pixel input pixel.
+// O = Center of the output pixel.
+//
+//          | 
+//    T     |     .
+//          | 
+//       O  | 
+//  --------+--------
+//          | 
+//          | 
+//    .     |     .
+//          | 
+static const int2 Offsets2x2[4] =
+{
+	int2( 0,  0), // T
+	int2( 1,  0),
+	int2( 0,  1),
+	int2( 1,  1),
+};
+
+// Indexes of the 3x3 square.
+static const uint kSquareIndexes3x3[9] = { 4, 0, 1, 2, 3, 8, 7, 6, 5 };
+
+// Indexes of the offsets to have plus + shape.
+static const uint kPlusIndexes3x3[5] = { 4, 1, 3, 7, 5 };
+
+// Cardinal pixel directions;
+static const int2 kOffsetC = int2( 0,  0);
+static const int2 kOffsetN = int2( 0, -1);
+static const int2 kOffsetS = int2( 0, +1);
+static const int2 kOffsetE = int2(+1,  0);
+static const int2 kOffsetW = int2(-1,  0);
+static const int2 kOffsetNE = int2(+1, -1);
+static const int2 kOffsetNW = int2(-1, -1);
+static const int2 kOffsetSE = int2(+1, +1);
+static const int2 kOffsetSW = int2(-1, +1);
+
+static const int2 kPairOffsets[4] = {
+	int2( 1,  0),
+	int2( 1,  1),
+	int2( 0,  1),
+	int2(-1,  1),
+};
+
+
+//------------------------------------------------------- CONFIG
+
+#if PLATFORM_SUPPORTS_WAVE_BROADCAST && 1
+	#define CONFIG_BUTTERFLY_KERNEL  1
+	#define CONFIG_BUTTERFLY_SAMPLES 4
+#else
+	#define CONFIG_BUTTERFLY_KERNEL  0
+	#define CONFIG_BUTTERFLY_SAMPLES 9
+#endif
+
+
+//------------------------------------------------------- PIXEL OFFSET CLAMPING
+
+// Clamp the offset to be shared across multiple samples
+CALL_SITE_DEBUGLOC
+tsr_short2x2 ClampPixelOffset(
+	tsr_short2x2 KernelCenterPixelPos,
+	tsr_short2x2 Offset,
+	const tsr_short2 OffsetDirection,
+	int2 MinPixelPos, int2 MaxPixelPos)
+{
+	tsr_short2x2 Min = dpv_sub(tsr_short2(MinPixelPos), KernelCenterPixelPos);
+	tsr_short2x2 Max = dpv_sub(tsr_short2(MaxPixelPos), KernelCenterPixelPos);
+
+	// Only do clamp based on the compile time known direction of the offset.
+	// This turns only 1 v_pk_max_i16 and v_pk_min_i16 for a 3x3 kernel.
+	tsr_short2x2 ClampedOffset = 0;
+
+	if (OffsetDirection.x > 0)
+	{
+		ClampedOffset[0] = min(Offset[0], Max[0]);
+	}
+	else if (OffsetDirection.x < 0)
+	{
+		ClampedOffset[0] = max(Offset[0], Min[0]);
+	}
+		
+	if (OffsetDirection.y > 0)
+	{
+		ClampedOffset[1] = min(Offset[1], Max[1]);
+	}
+	else if (OffsetDirection.y < 0)
+	{
+		ClampedOffset[1] = max(Offset[1], Min[1]);
+	}
+
+	return ClampedOffset;
+}
+
+CALL_SITE_DEBUGLOC
+tsr_short2 ClampPixelOffset(
+	tsr_short2 KernelCenterPixelPos,
+	tsr_short2 Offset,
+	const tsr_short2 OffsetDirection,
+	int2 MinPixelPos, int2 MaxPixelPos)
+{
+	tsr_short2 Min = tsr_short2(MinPixelPos) - KernelCenterPixelPos;
+	tsr_short2 Max = tsr_short2(MaxPixelPos) - KernelCenterPixelPos;
+
+	// Only do clamp based on the compile time known direction of the offset.
+	// This turns only 1 v_pk_max_i16 and v_pk_min_i16 for a 3x3 kernel.
+	tsr_short2 ClampedOffset = 0;
+
+	if (OffsetDirection.x > 0)
+		ClampedOffset.x = min(Offset.x, Max.x);
+	else if (OffsetDirection.x < 0)
+		ClampedOffset.x = max(Offset.x, Min.x);
+		
+	if (OffsetDirection.y > 0)
+		ClampedOffset.y = min(Offset.y, Max.y);
+	else if (OffsetDirection.y < 0)
+		ClampedOffset.y = max(Offset.y, Min.y);
+
+	return ClampedOffset;
+}
+
+CALL_SITE_DEBUGLOC
+tsr_short2x2 ClampPixelOffset(
+	tsr_short2x2 SamplePixelPos,
+	int2 MinPixelPos, int2 MaxPixelPos)
+{
+	SamplePixelPos[0] = fastClamp(SamplePixelPos[0], tsr_short2(MinPixelPos.xx), tsr_short2(MaxPixelPos.xx));
+	SamplePixelPos[1] = fastClamp(SamplePixelPos[1], tsr_short2(MinPixelPos.yy), tsr_short2(MaxPixelPos.yy));
+	return SamplePixelPos;
+}
+
+CALL_SITE_DEBUGLOC
+tsr_short2 ClampPixelOffset(
+	tsr_short2 SamplePixelPos,
+	int2 MinPixelPos, int2 MaxPixelPos)
+{
+	return fastClamp(SamplePixelPos, tsr_short2(MinPixelPos), tsr_short2(MaxPixelPos));
+}
+
+CALL_SITE_DEBUGLOC
+tsr_short2x2 AddAndClampPixelOffset(
+	tsr_short2x2 KernelCenterPixelPos,
+	tsr_short2x2 Offset,
+	const tsr_short2 OffsetDirection,
+	int2 MinPixelPos, int2 MaxPixelPos)
+{
+	tsr_short2x2 SamplePixelPos = KernelCenterPixelPos + Offset;
+	
+	tsr_short2 SamplePixelPos0 = dpv_lo(SamplePixelPos);
+	tsr_short2 SamplePixelPos1 = dpv_hi(SamplePixelPos);
+
+	// Only do clamp based on the compile time known direction of the offset.
+	// This turns only 1 v_pk_max_i16 and v_pk_min_i16 for a 3x3 kernel.
+	if (OffsetDirection.x > 0)
+		SamplePixelPos0.x = min(SamplePixelPos0.x, tsr_short(MaxPixelPos.x));
+	else if (OffsetDirection.x < 0)
+		SamplePixelPos0.x = max(SamplePixelPos0.x, tsr_short(MinPixelPos.x));
+		
+	if (OffsetDirection.y > 0)
+		SamplePixelPos0.y = min(SamplePixelPos0.y, tsr_short(MaxPixelPos.y));
+	else if (OffsetDirection.y < 0)
+		SamplePixelPos0.y = max(SamplePixelPos0.y, tsr_short(MinPixelPos.y));
+	
+	if (OffsetDirection.x > 0)
+		SamplePixelPos1.x = min(SamplePixelPos1.x, tsr_short(MaxPixelPos.x));
+	else if (OffsetDirection.x < 0)
+		SamplePixelPos1.x = max(SamplePixelPos1.x, tsr_short(MinPixelPos.x));
+		
+	if (OffsetDirection.y > 0)
+		SamplePixelPos1.y = min(SamplePixelPos1.y, tsr_short(MaxPixelPos.y));
+	else if (OffsetDirection.y < 0)
+		SamplePixelPos1.y = max(SamplePixelPos1.y, tsr_short(MinPixelPos.y));
+	
+	return dpv_interleave_registers(SamplePixelPos0, SamplePixelPos1);
+}
+
+CALL_SITE_DEBUGLOC
+tsr_short2 AddAndClampPixelOffset(
+	tsr_short2 KernelCenterPixelPos,
+	tsr_short2 Offset,
+	const tsr_short2 OffsetDirection,
+	int2 MinPixelPos, int2 MaxPixelPos)
+{
+	tsr_short2 SamplePixelPos = KernelCenterPixelPos + Offset;
+
+	// Only do clamp based on the compile time known direction of the offset.
+	// This turns only 1 v_pk_max_i16 and v_pk_min_i16 for a 3x3 kernel.
+	if (OffsetDirection.x > 0)
+		SamplePixelPos.x = min(SamplePixelPos.x, tsr_short(MaxPixelPos.x));
+	else if (OffsetDirection.x < 0)
+		SamplePixelPos.x = max(SamplePixelPos.x, tsr_short(MinPixelPos.x));
+		
+	if (OffsetDirection.y > 0)
+		SamplePixelPos.y = min(SamplePixelPos.y, tsr_short(MaxPixelPos.y));
+	else if (OffsetDirection.y < 0)
+		SamplePixelPos.y = max(SamplePixelPos.y, tsr_short(MinPixelPos.y));
+
+	return SamplePixelPos;
+}
+
+
+//------------------------------------------------------- 3x3 BUTTERFLY KERNEL
+
+CALL_SITE_DEBUGLOC
+tsr_short2 GetLaneOffsetSign()
+#if CONFIG_BUTTERFLY_KERNEL
+{
+	uint LaneIndex = WaveGetLaneIndex();
+
+	return tsr_short(-1) + tsr_short2(tsr_ushort2(tsr_ushort(LaneIndex) << tsr_ushort(1), LaneIndex) & tsr_ushort(0x2));
+}
+#else
+{
+	return tsr_short(1).xx;
+}
+#endif
+
+CALL_SITE_DEBUGLOC
+tsr_short2 GetNeighborhoodPixelPos(const uint ArrayIndex, tsr_short2 InputPixelPos)
+{
+	tsr_short2 LaneOffsetSign = GetLaneOffsetSign();
+
+	tsr_short2 SampleInputPixelPos;
+	#if CONFIG_BUTTERFLY_KERNEL
+	{
+		tsr_short2 Offset = tsr_short2(Offsets2x2[ArrayIndex]) * LaneOffsetSign;
+		SampleInputPixelPos = ClampPixelOffset(InputPixelPos + Offset, InputPixelPosMin, InputPixelPosMax);
+	}
+	#else
+	{
+		tsr_short2 Offset = kOffsets3x3[ArrayIndex];
+					
+		SampleInputPixelPos = InputPixelPos + ClampPixelOffset(
+			InputPixelPos,
+			Offset, Offset,
+			InputPixelPosMin, InputPixelPosMax);
+	}
+	#endif
+
+	return SampleInputPixelPos;
+}
+
+CALL_SITE_DEBUGLOC
+template<typename Type>
+Type AccessNeighborhoodStaticOffset(Type Array[CONFIG_BUTTERFLY_SAMPLES], const int2 Offset)
+{
+	const uint OffsetIndex = uint(dot(Offset + 1, int2(1, 3)));
+
+	#if CONFIG_BUTTERFLY_KERNEL
+		const uint ButterflyArray[] = {
+			0x3, 0x2, 0x2,
+			0x1, 0x0, 0x0,
+			0x1, 0x0, 0x0,
+		};
+		const uint IndexArray[] = {
+			0, 0, 1,
+			0, 0, 1,
+			2, 2, 3,
+		};
+
+		const FWaveBroadcastSettings BroadcastSetting = InitWaveXorButterfly(ButterflyArray[OffsetIndex]);
+		
+		return WaveBroadcast(BroadcastSetting, Array[IndexArray[OffsetIndex]]);
+	#else
+		return Array[OffsetIndex];
+	#endif
+}
+
+CALL_SITE_DEBUGLOC
+template<typename Type>
+Type AccessNeighborhoodDynamicOffset(Type Array[CONFIG_BUTTERFLY_SAMPLES], const tsr_short2 Offset)
+{
+	const uint OffsetIndex = uint(dot(Offset + 1, tsr_short2(1, 3)));
+
+	#if CONFIG_BUTTERFLY_KERNEL
+		const FWaveBroadcastSettings BroadcastVertical   = InitWaveXorButterfly(0x2);
+		const FWaveBroadcastSettings BroadcastHorizontal = InitWaveXorButterfly(0x1);
+		const FWaveBroadcastSettings BroadcastDiagonal   = InitWaveXorButterfly(0x3);
+
+		Type Return = Array[0];
+		Return = select(all(Offset == tsr_short2(-1, -1)), WaveBroadcast(BroadcastDiagonal,   Array[0]), Return);
+		Return = select(all(Offset == tsr_short2( 0, -1)), WaveBroadcast(BroadcastVertical,   Array[0]), Return);
+		Return = select(all(Offset == tsr_short2(+1, -1)), WaveBroadcast(BroadcastVertical,   Array[1]), Return);
+		
+		Return = select(all(Offset == tsr_short2(-1,  0)), WaveBroadcast(BroadcastHorizontal, Array[0]), Return);
+		Return = select(all(Offset == tsr_short2(+1,  0)), Array[1], Return);
+		
+		Return = select(all(Offset == tsr_short2(-1, +1)), WaveBroadcast(BroadcastHorizontal, Array[2]), Return);
+		Return = select(all(Offset == tsr_short2( 0, +1)), Array[2], Return);
+		Return = select(all(Offset == tsr_short2(+1, +1)), Array[3], Return);
+		
+		return Return;
+	#else
+		return Array[OffsetIndex];
+	#endif
+}
+
+CALL_SITE_DEBUGLOC
+template<typename Type>
+Type AccessNeighborhoodCenter(Type Array[CONFIG_BUTTERFLY_SAMPLES])
+#if CONFIG_BUTTERFLY_KERNEL
+{
+	return Array[0];
+}
+#else
+{
+	return Array[4];
+}
+#endif
+
+
+//------------------------------------------------------- OPTIMISED CATMULL-ROM
+
+FCatmullRomSamples GetBicubic2DCatmullRomSamples_Stubbe(float2 UV, float2 Size, in float2 InvSize)
+{
+	float2 PixelCoord = UV * Size;
+
+	float2 iuv = floor(PixelCoord - 0.5) + 1.0;
+	tsr_half2 f = tsr_half2(PixelCoord - iuv);
+	tsr_half2 f2 = f * f;
+	tsr_half2 f_2 = f + f;
+
+	tsr_half2 uv = (tsr_half(1.25) - f2) * f + tsr_half(0.5);
+	tsr_half2 t0 = (tsr_half(0.25) * f2 - tsr_half(0.0625)) * (tsr_half(1.125) - tsr_half(0.5) * f2.yx);
+	tsr_half2 w0 = t0 - f_2 * t0;
+	tsr_half2 w3 = t0 + f_2 * t0;
+
+	float2 cuv = (iuv - 0.5) * InvSize;
+	uv *= tsr_half2(InvSize);
+
+	FCatmullRomSamples Samples;
+	Samples.Count = BICUBIC_CATMULL_ROM_SAMPLES;
+
+	Samples.UV[0] = cuv + tsr_half2(               -InvSize.x,                      uv.y);
+	Samples.UV[1] = cuv + tsr_half2(tsr_half(2.0) * InvSize.x,                      uv.y);
+	Samples.UV[2] = cuv + uv;
+	Samples.UV[3] = cuv + tsr_half2(                     uv.x,                -InvSize.y);
+	Samples.UV[4] = cuv + tsr_half2(                     uv.x, tsr_half(2.0) * InvSize.y);
+	
+	Samples.Weight[0] = half(w0.x);
+	Samples.Weight[1] = half(w3.x);
+	Samples.Weight[2] = half(tsr_half(1.0) - w0.x - w3.x - w0.y - w3.y);
+	Samples.Weight[3] = half(w0.y);
+	Samples.Weight[4] = half(w3.y);
+
+	Samples.UVDir[0] = int2(-1,  0);
+	Samples.UVDir[1] = int2( 1,  0);
+	Samples.UVDir[2] = int2( 0,  0);
+	Samples.UVDir[3] = int2( 0, -1);
+	Samples.UVDir[4] = int2( 0,  1);
+	
+	Samples.FinalMultiplier = 1.0;
+	return Samples;
+} // GetBicubic2DCatmullRomSamples_Stubbe()
+
+template<typename Type>
+struct FCatmullRomFetches
+{
+	Type Colors[BICUBIC_CATMULL_ROM_SAMPLES];
+	tsr_half Weights[BICUBIC_CATMULL_ROM_SAMPLES];
+	
+	void FetchSamples(Texture2DArray<Type> InTexture, float2 UV, float SliceIndex, float2 Size, float2 InvSize, float2 MinUV, float2 MaxUV)
+	{
+		FCatmullRomSamples Samples = GetBicubic2DCatmullRomSamples_Stubbe(UV, Size, InvSize);
+		
+		UNROLL_N(BICUBIC_CATMULL_ROM_SAMPLES)
+		for (uint i = 0; i < BICUBIC_CATMULL_ROM_SAMPLES; i++)
+		{
+			float2 SampleUV = Samples.UV[i];
+			SampleUV = fastClamp(
+				SampleUV,
+				MinUV,
+				MaxUV);
+
+			Colors[i] = InTexture.SampleLevel(GlobalBilinearClampedSampler, float3(SampleUV, SliceIndex), 0);
+			Weights[i] = tsr_half(Samples.Weight[i]);
+		}
+	}
+
+	Type AccumulateSamples()
+	{
+		Type Color = tsr_half(0.0);
+		UNROLL_N(BICUBIC_CATMULL_ROM_SAMPLES)
+		for (uint i = 0; i < BICUBIC_CATMULL_ROM_SAMPLES; i++)
+		{
+			Color += Colors[i] * Weights[i];
+		}
+		return Color;
+	}
+};
+
+
+//------------------------------------------------------- UPSCALING WEIGHT DISTRIBUTION
+
+// Returns the weight of a pixels at a coordinate <PixelDelta> from the PDF highest point.
+CALL_SITE_DEBUGLOC
+tsr_half2 ComputeSampleWeigth(tsr_half2 UpscaleFactor, tsr_half2x2 PixelDelta, const float MinimalContribution)
+{
+	tsr_half2 u2 = UpscaleFactor * UpscaleFactor;
+
+	// 1 - 1.9 * x^2 + 0.9 * x^4
+	tsr_half2 x2 = saturate(u2 * dpv_length2(PixelDelta));
+	//return tsr_half(((float(0.9) + MinimalContribution) * x2 - float(1.9)) * x2 + float(1.0));
+	return saturate((tsr_half(0.9) * x2 - tsr_half(1.9)) * x2 + tsr_half(1.0 + MinimalContribution));
+}
+
+CALL_SITE_DEBUGLOC
+tsr_half ComputeSampleWeigth(tsr_half UpscaleFactor, tsr_half2 PixelDelta, const float MinimalContribution)
+{
+	return dpv_lo(ComputeSampleWeigth(
+		dpv_interleave_mono_registers(UpscaleFactor),
+		dpv_interleave_mono_registers(PixelDelta),
+		MinimalContribution));
+}
diff --git a/Engine/Shaders/Private/TemporalSuperResolution/TSRShadingAnalysis.ush b/Engine/Shaders/Private/TemporalSuperResolution/TSRShadingAnalysis.ush
--- a/Engine/Shaders/Private/TemporalSuperResolution/TSRShadingAnalysis.ush
+++ b/Engine/Shaders/Private/TemporalSuperResolution/TSRShadingAnalysis.ush
@@ -6,6 +6,7 @@
 //------------------------------------------------------- INCLUDES
 
 #include "TSRColorSpace.ush"
+#include "TSRSpatialAntiAliasing.ush"
 #include "TSRConvolutionNetworkPass.ush"
 
 
diff --git a/Engine/Shaders/Private/TemporalSuperResolution/TSRSpatialAntiAliasing.ush b/Engine/Shaders/Private/TemporalSuperResolution/TSRSpatialAntiAliasing.ush
--- a/Engine/Shaders/Private/TemporalSuperResolution/TSRSpatialAntiAliasing.ush
+++ b/Engine/Shaders/Private/TemporalSuperResolution/TSRSpatialAntiAliasing.ush
@@ -1,6 +1,13 @@
 // Copyright Epic Games, Inc. All Rights Reserved.
 
-#include "TSRCommon.ush"
+#pragma once
+
+#include "TSRKernels.ush"
+
+
+//------------------------------------------------------- DEFINITIONS
+
+#define SPATIAL_ANTI_ALIASER_MIN_LUMIMANCE 0.05
 
 
 //------------------------------------------------------- SAMPLING FUNCTION
@@ -622,3 +629,51 @@ tsr_half2 ComputeReprojectionBoundary(
 	}
 	#endif
 } // ComputeReprojectionBoundary()
+
+
+//------------------------------------------------------- ENCODE ANTI-ALIASING OFFSET
+
+// Spatial anti-aliasing encoding settings.
+#define SPATIAL_ANTI_ALIASING_OFFSET_RANGE 0.5
+#define SPATIAL_ANTI_ALIASING_OFFSET_BITDEPTH_PER_CHANNEL 4
+
+CALL_SITE_DEBUGLOC
+uint2 EncodeSpatialAntiAliasingOffset(tsr_half2x2 TexelOffset)
+{
+	const uint Mask = (1 << SPATIAL_ANTI_ALIASING_OFFSET_BITDEPTH_PER_CHANNEL) - 1;
+	const uint Zero = Mask / 2;
+
+	tsr_half Multiply = tsr_half(float(Zero) / SPATIAL_ANTI_ALIASING_OFFSET_RANGE);
+
+	uint2 EncodedTexelOffset = 0;
+	EncodedTexelOffset |= fastClamp(uint2(int2(round(TexelOffset[0] * Multiply)) + int(1 + Zero)), uint(1), Mask) << 0;
+	EncodedTexelOffset |= fastClamp(uint2(int2(round(TexelOffset[1] * Multiply)) + int(1 + Zero)), uint(1), Mask) << SPATIAL_ANTI_ALIASING_OFFSET_BITDEPTH_PER_CHANNEL;
+	return EncodedTexelOffset;
+}
+
+CALL_SITE_DEBUGLOC
+uint EncodeSpatialAntiAliasingOffset(tsr_half2 TexelOffset)
+{
+	return dpv_lo(EncodeSpatialAntiAliasingOffset(dpv_interleave_mono_registers(TexelOffset)));
+}
+
+CALL_SITE_DEBUGLOC
+tsr_half2x2 DecodeSpatialAntiAliasingOffset(tsr_ushort2 EncodedInputTexelOffset)
+{
+	const tsr_ushort Mask = tsr_ushort(1 << SPATIAL_ANTI_ALIASING_OFFSET_BITDEPTH_PER_CHANNEL) - tsr_ushort(1);
+	const tsr_ushort Zero = Mask / tsr_ushort(2);
+	
+	tsr_half Multiply = tsr_half(SPATIAL_ANTI_ALIASING_OFFSET_RANGE / float(Zero));
+
+	tsr_half2x2 TexelOffset;
+	TexelOffset[0] = tsr_half2((EncodedInputTexelOffset >> tsr_ushort(0                                                )) & Mask) * Multiply - (tsr_half(1 + Zero) * Multiply);
+	TexelOffset[1] = tsr_half2((EncodedInputTexelOffset >> tsr_ushort(SPATIAL_ANTI_ALIASING_OFFSET_BITDEPTH_PER_CHANNEL)) & Mask) * Multiply - (tsr_half(1 + Zero) * Multiply);
+
+	return TexelOffset;
+}
+
+CALL_SITE_DEBUGLOC
+tsr_half2 DecodeSpatialAntiAliasingOffset(tsr_ushort EncodedInputTexelOffset)
+{
+	return dpv_lo(DecodeSpatialAntiAliasingOffset(dpv_interleave_mono_registers(EncodedInputTexelOffset)));
+}
diff --git a/Engine/Shaders/Private/TemporalSuperResolution/TSRUpdateHistory.usf b/Engine/Shaders/Private/TemporalSuperResolution/TSRUpdateHistory.usf
--- a/Engine/Shaders/Private/TemporalSuperResolution/TSRUpdateHistory.usf
+++ b/Engine/Shaders/Private/TemporalSuperResolution/TSRUpdateHistory.usf
@@ -1,5 +1,7 @@
 // Copyright Epic Games, Inc. All Rights Reserved.
 
+#include "TSRKernels.ush"
+#include "TSRSpatialAntiAliasing.ush"
 #include "TSRColorSpace.ush"
 #include "TSRReprojectionField.ush"
 
