Fix DFAO upsampling not using the correct UV on stereo rendering with Substrate.

diff --git a/Engine/Shaders/Private/ReflectionEnvironmentPixelShader.usf b/Engine/Shaders/Private/ReflectionEnvironmentPixelShader.usf
--- a/Engine/Shaders/Private/ReflectionEnvironmentPixelShader.usf
+++ b/Engine/Shaders/Private/ReflectionEnvironmentPixelShader.usf
@@ -427,7 +427,8 @@ void ReflectionEnvironmentSkyLighting(
 	
 	#if APPLY_SKY_SHADOWING
 		// Sample DFAO only once
-		FUpsampleDFAOOutput UpsampleDFAOOutput = UpsampleDFAO(BufferUV, SceneDepth);
+		const float2 BentNormalUV = (SvPosition.xy - View.ViewRectMin.xy) * View.BufferSizeAndInvSize.zw;
+		FUpsampleDFAOOutput UpsampleDFAOOutput = UpsampleDFAO(BentNormalUV, SceneDepth);
 	#endif
 	
 		float CombinedScreenAndMaterialAO = SubstrateGetAO(SubstratePixelHeader) * AmbientOcclusion;
