Fixing shader compilation error

diff --git a/Engine/Shaders/Private/HeterogeneousVolumes/HeterogeneousVolumesLiveShadingGlobalPipeline.usf b/Engine/Shaders/Private/HeterogeneousVolumes/HeterogeneousVolumesLiveShadingGlobalPipeline.usf
--- a/Engine/Shaders/Private/HeterogeneousVolumes/HeterogeneousVolumesLiveShadingGlobalPipeline.usf
+++ b/Engine/Shaders/Private/HeterogeneousVolumes/HeterogeneousVolumesLiveShadingGlobalPipeline.usf
@@ -4,6 +4,18 @@
 #include "../SceneTexturesCommon.ush"
 #include "HeterogeneousVolumesTracingUtils.ush"
 
+#ifndef THREADGROUP_SIZE_1D
+#define THREADGROUP_SIZE_1D 1
+#endif // THREADGROUP_SIZE_1D
+
+#ifndef THREADGROUP_SIZE_2D
+#define THREADGROUP_SIZE_2D 1
+#endif // THREADGROUP_SIZE_2D
+
+#ifndef THREADGROUP_SIZE_3D
+#define THREADGROUP_SIZE_3D 1
+#endif // THREADGROUP_SIZE_3D
+
 struct FScreenTile
 {
 	// TODO: This could be 16-bit..
