Fix SSS checkerboard with ManyLight.

diff --git a/Engine/Shaders/Private/DeferredShadingCommon.ush b/Engine/Shaders/Private/DeferredShadingCommon.ush
--- a/Engine/Shaders/Private/DeferredShadingCommon.ush
+++ b/Engine/Shaders/Private/DeferredShadingCommon.ush
@@ -17,6 +17,10 @@
 #define GBUFFER_REFACTOR 0
 #endif
 
+#ifndef ALLOW_SSS_MATERIAL_OVERRIDE
+#define ALLOW_SSS_MATERIAL_OVERRIDE 1
+#endif
+
 #if SHADING_PATH_DEFERRED
 
 #if !SUPPORTS_INDEPENDENT_SAMPLERS
@@ -926,7 +930,7 @@ void AdjustBaseColorAndSpecularColorForSubsurfaceProfileLighting(inout float3 Ba
 {
 // If this function is called in translucent pass, we should not modify the BaseColor, instead directly read the base color as it is. 
 // The translucent pass is called after the main post process subsurface scattering pass, it will not affect previous logic.
-#if MATERIALBLENDING_TRANSLUCENT
+#if MATERIALBLENDING_TRANSLUCENT || !ALLOW_SSS_MATERIAL_OVERRIDE
 	return;
 #endif
 
diff --git a/Engine/Shaders/Private/GBufferHelpers.ush b/Engine/Shaders/Private/GBufferHelpers.ush
--- a/Engine/Shaders/Private/GBufferHelpers.ush
+++ b/Engine/Shaders/Private/GBufferHelpers.ush
@@ -2,6 +2,9 @@
 
 #pragma once
 
+#ifndef ALLOW_SSS_MATERIAL_OVERRIDE
+#define ALLOW_SSS_MATERIAL_OVERRIDE 1
+#endif
 
 #include "SceneTexturesCommon.ush"
 #include "SceneTextureParameters.ush"
@@ -350,7 +353,7 @@ void AdjustBaseColorAndSpecularColorForSubsurfaceProfileLightingCopyHack(inout f
 {
 	// If this function is called in translucent pass, we should not modify the BaseColor, instead directly read the base color as it is. 
 	// The translucent pass is called after the main post process subsurface scattering pass, it will not affect previous logic.
-#if MATERIALBLENDING_TRANSLUCENT
+#if MATERIALBLENDING_TRANSLUCENT || !ALLOW_SSS_MATERIAL_OVERRIDE
 	return;
 #endif
 
diff --git a/Engine/Shaders/Private/ManyLights/ManyLights.usf b/Engine/Shaders/Private/ManyLights/ManyLights.usf
--- a/Engine/Shaders/Private/ManyLights/ManyLights.usf
+++ b/Engine/Shaders/Private/ManyLights/ManyLights.usf
@@ -1,5 +1,8 @@
 // Copyright Epic Games, Inc. All Rights Reserved.
 
+// When loading SSS checkerboard pixel, do not adjust DiffuseColor/SpecularColor to preserve specular and diffuse lighting values for each pixel
+#define ALLOW_SSS_MATERIAL_OVERRIDE 0
+
 #include "../Common.ush"
 #include "ManyLights.ush"
 
diff --git a/Engine/Shaders/Private/ManyLights/ManyLightsDenoiserSpatial.usf b/Engine/Shaders/Private/ManyLights/ManyLightsDenoiserSpatial.usf
--- a/Engine/Shaders/Private/ManyLights/ManyLightsDenoiserSpatial.usf
+++ b/Engine/Shaders/Private/ManyLights/ManyLightsDenoiserSpatial.usf
@@ -4,6 +4,10 @@
 #include "ManyLights.ush"
 #include "../Lumen/LumenReflectionDenoiserCommon.ush"
 
+// Spatial denoise is the last pass writing to the SceneColor buffer. We allow SSS checkerboard pixel to adjust 
+// DiffuseColor/SpecularColor to write specular & diffuse lighting in checkerboard pattern
+#define ALLOW_SSS_MATERIAL_OVERRIDE 1
+
 Texture2D<float4> DiffuseLightingAndSecondMomentTexture;
 Texture2D<float4> SpecularLightingAndSecondMomentTexture;
 Texture2D<UNORM float> NumFramesAccumulatedTexture;
diff --git a/Engine/Shaders/Private/ManyLights/ManyLightsDenoiserTemporal.usf b/Engine/Shaders/Private/ManyLights/ManyLightsDenoiserTemporal.usf
--- a/Engine/Shaders/Private/ManyLights/ManyLightsDenoiserTemporal.usf
+++ b/Engine/Shaders/Private/ManyLights/ManyLightsDenoiserTemporal.usf
@@ -1,5 +1,8 @@
 // Copyright Epic Games, Inc. All Rights Reserved.
 
+// When loading SSS checkerboard pixel, do not adjust DiffuseColor/SpecularColor to preserve specular and diffuse lighting values for each pixel
+#define ALLOW_SSS_MATERIAL_OVERRIDE 0
+
 #include "../Common.ush"
 #include "ManyLights.ush"
 #include "../Lumen/LumenReflectionDenoiserCommon.ush"
diff --git a/Engine/Shaders/Private/ManyLights/ManyLightsSampling.usf b/Engine/Shaders/Private/ManyLights/ManyLightsSampling.usf
--- a/Engine/Shaders/Private/ManyLights/ManyLightsSampling.usf
+++ b/Engine/Shaders/Private/ManyLights/ManyLightsSampling.usf
@@ -10,6 +10,9 @@
 	#define TILE_TYPE TILE_MODE_EMPTY
 #endif
 
+// When loading SSS checkerboard pixel, do not adjust DiffuseColor/SpecularColor to preserve specular and diffuse lighting values for each pixel
+#define ALLOW_SSS_MATERIAL_OVERRIDE 0
+
 #include "../Common.ush"
 #include "../BlueNoise.ush"
 #include "ManyLights.ush"
diff --git a/Engine/Shaders/Private/ManyLights/ManyLightsShading.usf b/Engine/Shaders/Private/ManyLights/ManyLightsShading.usf
--- a/Engine/Shaders/Private/ManyLights/ManyLightsShading.usf
+++ b/Engine/Shaders/Private/ManyLights/ManyLightsShading.usf
@@ -10,6 +10,9 @@
 	#define NUM_SAMPLES_PER_PIXEL_2D_Y 1
 #endif
 
+// When loading SSS checkerboard pixel, do not adjust DiffuseColor/SpecularColor to preserve specular and diffuse lighting values for each pixel
+#define ALLOW_SSS_MATERIAL_OVERRIDE 0
+
 #include "../Common.ush"
 #include "../BlueNoise.ush"
 #include "ManyLights.ush"
