ManyLights - fixed generate samples regression due to a bug causing TILE_MODE_SIMPLE_SHADING running with all the shading modes enabled.

Console, 1080p, Firefly: 0.66ms -> 0.35ms

diff --git a/Engine/Shaders/Private/ManyLights/ManyLights.ush b/Engine/Shaders/Private/ManyLights/ManyLights.ush
--- a/Engine/Shaders/Private/ManyLights/ManyLights.ush
+++ b/Engine/Shaders/Private/ManyLights/ManyLights.ush
@@ -204,13 +204,15 @@ FManyLightMaterial LoadMaterial(float2 ScreenUV, uint2 ScreenCoord, bool bForceS
 {
 	FGBufferData GBuffer = GetGBufferData(ScreenUV);
 	GBuffer.Roughness = max(GBuffer.Roughness, View.MinRoughness);
-	if (GBuffer.ShadingModelID != SHADINGMODELID_UNLIT && bForceSimpleShading)
+	const bool bIsLit = GBuffer.ShadingModelID != SHADINGMODELID_UNLIT;
+
+	if (bForceSimpleShading)
 	{
 		GBuffer.ShadingModelID = SHADINGMODELID_DEFAULT_LIT;
 	}
 
 	FManyLightMaterial Out = (FManyLightMaterial)0;
-	Out.bIsValid      = GBuffer.ShadingModelID != SHADINGMODELID_UNLIT;
+	Out.bIsValid      = bIsLit;
 	Out.bIsSimple     = GBuffer.ShadingModelID == SHADINGMODELID_DEFAULT_LIT;
 	Out.bNeedsSeparateSubsurfaceLightAccumulation = UseSubsurfaceProfile(GBuffer.ShadingModelID);
 	Out.WorldNormal   = GBuffer.WorldNormal;
