Make sky atmosphere (without SkyDomeMesh) output opaque alpha in all code path.
Followup to 32793646.
#rnx

diff --git a/Engine/Shaders/Private/SkyAtmosphere.usf b/Engine/Shaders/Private/SkyAtmosphere.usf
--- a/Engine/Shaders/Private/SkyAtmosphere.usf
+++ b/Engine/Shaders/Private/SkyAtmosphere.usf
@@ -843,6 +843,21 @@ float4 PrepareOutput(float3 PreExposedLuminance, float3 Transmittance = float3(1
 
 uint DepthReadDisabled;
 
+void UpdateVisibleSkyAlpha(in float DeviceZ, inout float4 OutLuminance)
+{
+#if POST_PROCESS_ALPHA
+	if (View.RenderingReflectionCaptureMask == 0.0f)
+	{
+		// This is for parity with the sky dome mesh behavior.
+		// Output fully opaque alpha if looking through outer space, if atmosphere is not in holdout.
+		if (DeviceZ == FarDepthValue && !IsSkyAtmosphereHoldout(View.EnvironmentComponentsFlags))
+		{
+			OutLuminance.a = 1.0;
+		}
+	}
+#endif
+}
+
 void RenderSkyAtmosphereRayMarchingPS(
 	in float4 SVPos : SV_POSITION,
 	out float4 OutLuminance : SV_Target0
@@ -969,19 +984,7 @@ void RenderSkyAtmosphereRayMarchingPS(
 		PreExposedL += SkyLuminance * LuminanceScale * (ViewOneOverPreExposure * OutputPreExposure);
 	
 		OutLuminance = PrepareOutput(PreExposedL, SkyGreyTransmittance);
-		
-	#if POST_PROCESS_ALPHA
-		if (View.RenderingReflectionCaptureMask == 0.0f)
-		{
-			// This is for parity with the sky dome mesh behavior.
-			// Output fully opaque alpha if looking through outer space, if atmosphere is not in holdout.
-			if (DeviceZ == FarDepthValue && !IsSkyAtmosphereHoldout(View.EnvironmentComponentsFlags))
-			{
-				OutLuminance.a = 1.0;
-			}
-		}
-	#endif
-		
+		UpdateVisibleSkyAlpha(DeviceZ, OutLuminance);
 		return;
 	}
 #endif
@@ -1012,6 +1015,7 @@ void RenderSkyAtmosphereRayMarchingPS(
 	float Transmittance = AP.a;
 
 	OutLuminance = PrepareOutput(PreExposedL, float3(Transmittance, Transmittance, Transmittance));
+	UpdateVisibleSkyAlpha(DeviceZ, OutLuminance);
 	return;
 #endif
 
@@ -1068,6 +1072,7 @@ void RenderSkyAtmosphereRayMarchingPS(
 	output.Transmittance = float4(ss.Transmittance, 1.0f);
 #else
 	OutLuminance = PrepareOutput(PreExposedL, ss.Transmittance);
+	UpdateVisibleSkyAlpha(DeviceZ, OutLuminance);
 	return;
 #endif
 
