Fixed shader compilation warning.

#rnx

diff --git a/Engine/Shaders/Private/HeterogeneousVolumes/HeterogeneousVolumesLiveShadingGlobalPipeline.usf b/Engine/Shaders/Private/HeterogeneousVolumes/HeterogeneousVolumesLiveShadingGlobalPipeline.usf
--- a/Engine/Shaders/Private/HeterogeneousVolumes/HeterogeneousVolumesLiveShadingGlobalPipeline.usf
+++ b/Engine/Shaders/Private/HeterogeneousVolumes/HeterogeneousVolumesLiveShadingGlobalPipeline.usf
@@ -155,7 +155,7 @@ void DilateExistenceMaskCS(
 		ExistenceTexture,
 		TextureSampler,
 		UVW,
-		MipLevel);
+		MipLevel).r;
 
 	RWDilatedExistenceTexture[DispatchThreadId] = Sample > 0.0 ? 1.0 : 0.0;
 }
