Substrate - disable energy conservation LUT on platfoms that do not support shared samplers in order to get material to compile.

#rnx

diff --git a/Engine/Shaders/Private/ShadingEnergyConservation.ush b/Engine/Shaders/Private/ShadingEnergyConservation.ush
--- a/Engine/Shaders/Private/ShadingEnergyConservation.ush
+++ b/Engine/Shaders/Private/ShadingEnergyConservation.ush
@@ -7,7 +7,13 @@
 // * If Substrate is disabled: OFF by default
 #if SUBSTRATE_ENABLED || LEGACY_MATERIAL_ENERGYCONSERVATION
 	#undef USE_ENERGY_CONSERVATION
-	#define USE_ENERGY_CONSERVATION 1
+	#if SUPPORTS_INDEPENDENT_SAMPLERS 
+		#define USE_ENERGY_CONSERVATION 1
+	#else
+		// Platforms that do not support shared sampler will in fact quickly reach the limit of samplers in material pixel shaders.
+		// This is mostly true for OpenGLES. In this case, we enforce disabling energy conservation to be able to compile substrate materials.
+		#define USE_ENERGY_CONSERVATION 0
+	#endif
 #else
 	#undef USE_ENERGY_CONSERVATION
 	#define USE_ENERGY_CONSERVATION 0
