Substrate - disable participating media energy "simple volume LUT" on platfoms that do not support shared samplers in order to get material to compile.

#rnx

diff --git a/Engine/Shaders/Private/ParticipatingMediaCommon.ush b/Engine/Shaders/Private/ParticipatingMediaCommon.ush
--- a/Engine/Shaders/Private/ParticipatingMediaCommon.ush
+++ b/Engine/Shaders/Private/ParticipatingMediaCommon.ush
@@ -212,8 +212,10 @@ float3 IsotropicMediumSlabPunctualDirectionalAlbedo(FParticipatingMedia PM)
 	return FinalEvaluate * IsotropicMediumSlabDirectionalAlbedoFade(PM.BaseColor, MFP);
 }
 
-// Wait before activating that feature for now.
-#define USE_NEW_SIMPLEVOLUME 1
+// Platforms that do not support shared sampler will in fact quickly reach the limit of samplers in material pixel shaders.
+// This is mostly true for OpenGLES. In this case, we enforce disabling SimpleVolume LUT and use the lower quality ALU based one.
+// SUBSTRATE_TODO Implement an ALU based solution that give the same view dependent high quality as the LUT.
+#define USE_NEW_SIMPLEVOLUME (SUPPORTS_INDEPENDENT_SAMPLERS)
 
 float IsotropicMediumSlabDirectionalAlbedoFadeLUT(float MFP)
 {
