Fixed shader compilation issue.

#rnx

diff --git a/Engine/Shaders/Private/AmbientCubemapComposite.usf b/Engine/Shaders/Private/AmbientCubemapComposite.usf
--- a/Engine/Shaders/Private/AmbientCubemapComposite.usf
+++ b/Engine/Shaders/Private/AmbientCubemapComposite.usf
@@ -80,7 +80,7 @@ void MainPS(
 		float SceneDepth = ConvertFromDeviceZ(DeviceZ);
 
 		float3 TranslatedWorldPosition = mul(float4(ScreenPosition * SceneDepth, SceneDepth, 1), View.ScreenToTranslatedWorld).xyz;
-		float3 CameraToPixel = normalize(WorldPosition - PrimaryView.TranslatedWorldCameraOrigin);
+		float3 CameraToPixel = normalize(TranslatedWorldPosition - PrimaryView.TranslatedWorldCameraOrigin);
 		float3 V = -CameraToPixel;
 
 		// Sample the ambient occlusion that is dynamically generated every frame.
