Add missing use case for sampling virtual texture page table.
SampleGrad with single page table and no feedback.
#rnx

diff --git a/Engine/Shaders/Private/VirtualTextureCommon.ush b/Engine/Shaders/Private/VirtualTextureCommon.ush
--- a/Engine/Shaders/Private/VirtualTextureCommon.ush
+++ b/Engine/Shaders/Private/VirtualTextureCommon.ush
@@ -545,6 +545,24 @@ VTPageTableResult TextureLoadVirtualPageTableGrad(
 	return Result;
 }
 
+// LoadPageTable: 1 Page table, SampleGrad, no feedback
+VTPageTableResult TextureLoadVirtualPageTableGrad(
+	Texture2D<uint4> PageTable0,
+	VTPageTableUniform PageTableUniform,
+	float2 UV, uint AddressU, uint AddressV,
+	MaterialFloat2 dUVdx, MaterialFloat2 dUVdy, float2 SvPositionXY)
+{
+	VTPageTableResult Result = (VTPageTableResult)0;
+	if (IsValid(PageTableUniform))
+	{
+		int vLevel = TextureComputeVirtualMipLevel(Result, dUVdx * PageTableUniform.UVScale, dUVdy * PageTableUniform.UVScale, 0, SvPositionXY, PageTableUniform);
+		UV = UV * PageTableUniform.UVScale;
+		UV = ApplyAddressMode(UV, AddressU, AddressV);
+		TextureLoadVirtualPageTableInternal(Result, PageTable0, PageTableUniform, UV, vLevel);
+	}
+	return Result;
+}
+
 // LoadPageTable: 2 Page tables, SampleGrad
 VTPageTableResult TextureLoadVirtualPageTableGrad(
 	Texture2D<uint4> PageTable0, Texture2D<uint4> PageTable1,
