Fix Legacy SSS profile having incorrect roughness.

diff --git a/Engine/Shaders/Private/PostProcessGBufferHints.usf b/Engine/Shaders/Private/PostProcessGBufferHints.usf
--- a/Engine/Shaders/Private/PostProcessGBufferHints.usf
+++ b/Engine/Shaders/Private/PostProcessGBufferHints.usf
@@ -120,6 +120,8 @@ void MainPS(noperspective float4 UVAndScreenPos : TEXCOORD0, float4 SvPosition :
 #include "ShaderPrint.ush"
 #include "SceneTextureParameters.ush"
 #include "DeferredShadingCommon.ush"
+#include "ShadingModels.ush"
+#include "SubsurfaceProfileCommon.ush"
 
 void PrintShadingMode(inout FShaderPrintContext Ctx, uint In)
 {
@@ -225,7 +227,19 @@ void MainCS(uint3 DispatchThreadId : SV_DispatchThreadID)
 		}
 		else if (Data.ShadingModelID == SHADINGMODELID_SUBSURFACE_PROFILE)
 		{
-		Print(Ctx, TEXT("SSS Profil ID     : "), FontWhite); Print(Ctx, ExtractSubsurfaceProfileInt(Data.CustomData.x)); Newline(Ctx, RectMax);	
+			const uint SSSProfileID = ExtractSubsurfaceProfileInt(Data.CustomData.x);
+			const half Opacity = Data.CustomData.a;
+			float LobeMix = 0;
+			float Lobe0Roughness = 0;
+			float Lobe1Roughness = 0;
+			GetProfileDualSpecular(SSSProfileID, Data.Roughness, Opacity, Lobe0Roughness, Lobe1Roughness, LobeMix);
+
+			Print(Ctx, TEXT("SSS Profil ID    : "), FontWhite); Print(Ctx, SSSProfileID); Newline(Ctx, RectMax);
+			Print(Ctx, TEXT("SSS Lobe Mix     : "), FontWhite); Print(Ctx, LobeMix); Newline(Ctx, RectMax);
+			Print(Ctx, TEXT("SSS Roughness0   : "), FontWhite); Print(Ctx, Lobe0Roughness); Newline(Ctx, RectMax);
+			Print(Ctx, TEXT("SSS Roughness1   : "), FontWhite); Print(Ctx, Lobe1Roughness); Newline(Ctx, RectMax);
+			Print(Ctx, TEXT("SSS SpecColor    : "), FontWhite); Print(Ctx, Data.SpecularColor); Newline(Ctx, RectMax);
+			
 		}
 		PrintBackgroundRect(Ctx, RectMin, RectMax);
 		Newline(Ctx);
diff --git a/Engine/Shaders/Private/Substrate/Substrate.ush b/Engine/Shaders/Private/Substrate/Substrate.ush
--- a/Engine/Shaders/Private/Substrate/Substrate.ush
+++ b/Engine/Shaders/Private/Substrate/Substrate.ush
@@ -4175,15 +4175,23 @@ FSubstrateBSDF UnpackSubstrateBSDFIn(FSubstrateMaterialContainer SubstrateBuffer
 				const uint PackedDiffuse20Bits	= Data1 & 0x000FFFFF;
 				const uint PackedF020Bits		= Data2 & 0x000FFFFF;
 				
-				const float Roughness				= UnpackR8(Data0 >> 24);
+				float Roughness						= UnpackR8(Data0 >> 24);
 				const float SSSProfileRadiusScale	= UnpackR8(Data1 >> 24);
 				const float SSSProfileId			= UnpackR8(Data2 >> 24);
+				FHaziness Haziness 					= InitialiseHaziness();
+
+				// Average roughness for dual specular.
+				// When storing data as SINGLE_OPTLEGACYMODE_SSSPROFILE, the dual roughness is not translated/stored as haziness in order to save storage space. 
+				// Thus we need to (re)compute them directly from the SSS profile data.
+				const uint SubsurfaceProfileUInt = SubstrateSubsurfaceProfileIdTo8bits(SSSProfileId);
+				GetSubsurfaceProfileDualSpecular(SubsurfaceProfileUInt, Roughness, SSSProfileRadiusScale, Roughness, Haziness.Roughness, Haziness.Weight);
 
 				SLAB_DIFFUSEALBEDO(OutBSDF)			= UnpackR7G7B6Gamma2(PackedDiffuse20Bits);
 				SLAB_F0(OutBSDF)					= UnpackR7G7B6Gamma2(PackedF020Bits);
 				SLAB_ROUGHNESS(OutBSDF)				= Roughness;
 				SLAB_SSSPROFILEID(OutBSDF)			= SSSProfileId;
 				SLAB_SSSPROFILERADIUSSCALE(OutBSDF)	= SSSProfileRadiusScale;
+				SLAB_HAZINESS(OutBSDF)				= PackHaziness(Haziness);
 
 				// When loading Substrate BSDF which have SSS enable, BaseColor and (optionally) Specular value can be overriden 
 				// based on the SSS method used (screen space SSS / checherboarded screen space SSS / ...)
@@ -4209,7 +4217,7 @@ FSubstrateBSDF UnpackSubstrateBSDFIn(FSubstrateMaterialContainer SubstrateBuffer
 				BSDF_SETISTOPLAYER(OutBSDF, 1);
 				BSDF_SETSSSTYPE(OutBSDF, SSS_TYPE_DIFFUSION_PROFILE);
 				BSDF_SETHASMFP(OutBSDF, 1);
-				BSDF_SETHASHAZINESS(OutBSDF, 0);
+				BSDF_SETHASHAZINESS(OutBSDF, 1);
 				BSDF_SETISTHIN(OutBSDF, 0);
 				BSDF_SETHASFUZZ(OutBSDF, 0);
 				BSDF_SETHASTRANSABOVE(OutBSDF, 0);
diff --git a/Engine/Shaders/Private/Substrate/SubstrateExport.ush b/Engine/Shaders/Private/Substrate/SubstrateExport.ush
--- a/Engine/Shaders/Private/Substrate/SubstrateExport.ush
+++ b/Engine/Shaders/Private/Substrate/SubstrateExport.ush
@@ -629,8 +629,15 @@ void PackSubstrateOut(
 							const uint PackedDiffuse20Bits		= PackR7G7B6Gamma2(	SLAB_DIFFUSEALBEDO(BSDF)	, Dither);
 							const uint PackedF020Bits			= PackR7G7B6Gamma2(	SLAB_F0(BSDF)				, Dither);
 
+							// When creating a slab with GetSubstrateSlabBSDF(), the SSS profile's dual roughness is expanded into haziness. 
+							// But when storing data with SINGLE_OPTLEGACYMODE_SSSPROFILE we need to store the original roughness. 
+							// The dual roughnesses will be recomputed based on the SSS profile data at load time. To retrieve the 
+							// original roughness, we need to invert the computation done into GetSubstrateSlabBSDF()
+							const uint SubsurfaceProfileUInt = SubstrateSubsurfaceProfileIdTo8bits(SLAB_SSSPROFILEID(BSDF));
+							const float OriginalRoughness = GetSubsurfaceProfileOriginalRoughness(SubsurfaceProfileUInt, SLAB_ROUGHNESS(BSDF), SLAB_SSSPROFILERADIUSSCALE(BSDF));
+
 							// MRT0 
-							SUBSTRATE_STORE_UINT1(Out					| (PackR8(SLAB_ROUGHNESS(BSDF)) << 24));
+							SUBSTRATE_STORE_UINT1(Out					| (PackR8(OriginalRoughness) << 24));
 							// MRT1 20 PackedDiffuse and Profile radius scale
 							SUBSTRATE_STORE_UINT1(PackedDiffuse20Bits  | (PackR8(SLAB_SSSPROFILERADIUSSCALE(BSDF)) << 24));
 							// MRT2 24 PackedF0      and Profile ID
diff --git a/Engine/Shaders/Private/SubsurfaceProfileCommon.ush b/Engine/Shaders/Private/SubsurfaceProfileCommon.ush
--- a/Engine/Shaders/Private/SubsurfaceProfileCommon.ush
+++ b/Engine/Shaders/Private/SubsurfaceProfileCommon.ush
@@ -79,6 +79,17 @@ half4 GetSubsurfaceProfileTexture(Texture2D InSSProfilesTexture, SamplerState In
 #endif 
 }
 
+float GetSubsurfaceProfileOriginalRoughness(uint SubsurfaceProfileInt, half ModifiedLobeRoughness0, half Opacity)
+{
+	const half4 Data = GetSubsurfaceProfileTexture(SSSS_DUAL_SPECULAR_OFFSET, SubsurfaceProfileInt);
+
+	// Smooth blend out dual specular when opacity is low, we have the extra SSSS_OPACITY_THRESHOLD_EPS so that we fade out by the time we
+	// get to 0.01, as opposed to 0.0.
+	half MaterialRoughnessToLobeRoughness0 = lerp(1.0f, Data.x * SSSS_MAX_DUAL_SPECULAR_ROUGHNESS, saturate((Opacity - SSSS_OPACITY_THRESHOLD_EPS) * 10.0f));
+
+	return ModifiedLobeRoughness0 / max(MaterialRoughnessToLobeRoughness0, 0.02f);
+}
+
 void GetSubsurfaceProfileDualSpecular(uint SubsurfaceProfileInt, half Roughness, half Opacity, out half LobeRoughness0, out half LobeRoughness1, out half LobeMix)
 {
 	const half4 Data = GetSubsurfaceProfileTexture(SSSS_DUAL_SPECULAR_OFFSET, SubsurfaceProfileInt);
