Fix shader compilation warning.

diff --git a/Engine/Shaders/Private/DeferredShadingCommon.ush b/Engine/Shaders/Private/DeferredShadingCommon.ush
--- a/Engine/Shaders/Private/DeferredShadingCommon.ush
+++ b/Engine/Shaders/Private/DeferredShadingCommon.ush
@@ -932,14 +932,14 @@ void AdjustBaseColorAndSpecularColorForSubsurfaceProfileLighting(inout float3 Ba
 // The translucent pass is called after the main post process subsurface scattering pass, it will not affect previous logic.
 #if MATERIALBLENDING_TRANSLUCENT || !ALLOW_SSS_MATERIAL_OVERRIDE
 	return;
-#endif
-
+#else
 	const bool bCheckerboardRequired = AdjustBaseColorAndSpecularColorForSubsurfaceProfileLighting(BaseColor, Specular, bChecker);
 	if (bCheckerboardRequired)
 	{
 		// in SubsurfaceRecombinePS() does not multiply with Specular so we do it here
 		SpecularColor *= !bChecker;
 	}
+#endif
 }
 
 /** Populates FGBufferData */
diff --git a/Engine/Shaders/Private/PostProcessSubsurface.usf b/Engine/Shaders/Private/PostProcessSubsurface.usf
--- a/Engine/Shaders/Private/PostProcessSubsurface.usf
+++ b/Engine/Shaders/Private/PostProcessSubsurface.usf
@@ -493,15 +493,14 @@ SDiffuseAndSpecular ReconstructLighting(float2 UVSceneColor, uint ReconstructMet
 	SDiffuseAndSpecular Ret;
 
 #if SUBSTRATE_OPAQUE_ROUGH_REFRACTION_ENABLED
-	// In this mode, we only read isolated diffuse lighting.
-	float3 CenterSample = SubsurfaceInput0_Texture.SampleLevel(SharedSubsurfaceSampler0, UVSceneColor, 0).rgb;
-	Ret.Diffuse = CenterSample;
-	Ret.Specular = 0.0f;
-	return Ret;
-#endif
-
+	{
+		// In this mode, we only read isolated diffuse lighting.
+		float3 CenterSample = SubsurfaceInput0_Texture.SampleLevel(SharedSubsurfaceSampler0, UVSceneColor, 0).rgb;
+		Ret.Diffuse = CenterSample;
+		Ret.Specular = 0.0f;
+	}
 	// If SUBSURFACE_CHANNEL_MODE is 0, checkerboard is forced on
-#if SUBSURFACE_PROFILE_CHECKERBOARD || SUBSURFACE_CHANNEL_MODE == 0
+#elif SUBSURFACE_PROFILE_CHECKERBOARD || SUBSURFACE_CHANNEL_MODE == 0
 	{
 	  bool bChecker = CheckerFromSceneColorUV(UVSceneColor);
 
