Fix incorrect legacy SSS profile classification

diff --git a/Engine/Shaders/Private/Substrate/SubstrateExport.ush b/Engine/Shaders/Private/Substrate/SubstrateExport.ush
--- a/Engine/Shaders/Private/Substrate/SubstrateExport.ush
+++ b/Engine/Shaders/Private/Substrate/SubstrateExport.ush
@@ -381,7 +381,7 @@ void PackSubstrateOut(
 					OptimisedLegacyMode = SINGLE_OPTLEGACYMODE_SSSWRAP;
 					SSSData.Header.Bytes = 0; // Force to skip SSSData writing
 				}
-				else if ((OptimisedModeMask & (SingleIsSSSProfile|SingleIsThin)) == SingleIsSSSProfile) // Force to skip SINGLE_OPTLEGACYMODE_SSSPROFILE when set to ThinSurface
+				else if ((OptimisedModeMask & (SingleIsSSSProfile|SingleIsThin)) == SingleIsSSSProfile && !BSDF_GETHASFUZZ(BSDF)) // Force to skip SINGLE_OPTLEGACYMODE_SSSPROFILE when set to ThinSurface
 				{
 					// SingleIsSSSProfile appears critically after SingleIsWrapThin because thin will disable subsurface profile.
 					OptimisedLegacyMode = SINGLE_OPTLEGACYMODE_SSSPROFILE;
