SMRT/VSM refactoring and cleanup

#rnx

diff --git a/Engine/Shaders/Private/VirtualShadowMaps/VirtualShadowMapProjection.usf b/Engine/Shaders/Private/VirtualShadowMaps/VirtualShadowMapProjection.usf
--- a/Engine/Shaders/Private/VirtualShadowMaps/VirtualShadowMapProjection.usf
+++ b/Engine/Shaders/Private/VirtualShadowMaps/VirtualShadowMapProjection.usf
@@ -218,13 +218,28 @@ FVirtualShadowMapSampleResult ProjectLight(
 	const float Noise,
 	const FSubsurfaceOpacityMFP SubsurfaceOpacityMFP)
 {
+	const bool bIsHairInput = InputType == INPUT_TYPE_HAIRSTRANDS;
+	
+	bool bBackfaceCull = (bCullBackfacingPixels > 0) && !bIsHairInput && !ShadingInfo.bIsSubsurface && !ShadingInfo.bIsHair;
+#if VISUALIZE
+	// Modes more related to the shadow depth rendering don't want to see the projection culling
+	bBackfaceCull = bBackfaceCull &&
+		VisualizeModeId != VIRTUAL_SHADOW_MAP_VISUALIZE_NANITE_OVERDRAW;
+#endif
+
 #if DIRECTIONAL_LIGHT
 	float3 L = Light.Direction;
 	bool bInLightRegion = true;
-#else	
+
+	if (bBackfaceCull && IsBackfaceToDirectionalLight(ShadingInfo.WorldNormal, Light.Direction, Light.SourceRadius))
+	{
+		bInLightRegion = false;
+	}
+
+#else	// Local light
+
 	float3 ToLight = Light.TranslatedWorldPosition - TranslatedWorldPosition;
-	float d2 = dot( ToLight, ToLight );
-	float InvDist = rsqrt( d2 );
+	float InvDist = rsqrt(dot(ToLight, ToLight));
 	float3 L = ToLight * InvDist;
 	
 	bool bInLightRadius = InvDist >= Light.InvRadius;
@@ -241,27 +256,27 @@ FVirtualShadowMapSampleResult ProjectLight(
 			bInLightRegion = false;
 		}
 	}
-#endif // DIRECTIONAL_LIGHT
 
-	const bool bIsHairInput = InputType == INPUT_TYPE_HAIRSTRANDS;
-	const float DistanceToCamera = GetDistanceToCameraFromViewVector(TranslatedWorldPosition - View.TranslatedWorldCameraOrigin);
-	const float NormalBiasLength = ShadingInfo.bIsSubsurface ? 0.0f : max(0.02f, NormalBias * DistanceToCamera / GetCotanHalfFieldOfView().x);
+	if (bInLightRegion && bBackfaceCull)
+	{
+		if (IsBackfaceToLocalLight(ToLight, ShadingInfo.WorldNormal, Light.SourceRadius))
+		{
+			bInLightRegion = false;
+		}
+	}
 
+#endif // DIRECTIONAL_LIGHT
+	
 	FVirtualShadowMapSampleResult Result = InitVirtualShadowMapSampleResult();
-
 	const bool bValidPixel = bIsHairInput || ShadingInfo.bIsValid;
 
 	BRANCH
 	if (bInLightRegion && bValidPixel)
 	{
-		bool bBackfaceCull = (bCullBackfacingPixels > 0) && !bIsHairInput && !ShadingInfo.bIsSubsurface;
-#if VISUALIZE
-		// Modes more related to the shadow depth rendering don't want to see the projection culling
-		bBackfaceCull = bBackfaceCull &&
-			VisualizeModeId != VIRTUAL_SHADOW_MAP_VISUALIZE_NANITE_OVERDRAW;
-#endif
+		const float3 WorldNormal = (bIsHairInput || ShadingInfo.bIsHair) ? L : ShadingInfo.WorldNormal;
 
-		float3 WorldNormal = (bIsHairInput || ShadingInfo.bIsHair) ? L : ShadingInfo.WorldNormal;
+		const float DistanceToCamera = GetDistanceToCameraFromViewVector(TranslatedWorldPosition - View.TranslatedWorldCameraOrigin);
+		const float NormalBiasLength = ShadingInfo.bIsSubsurface ? 0.0f : max(0.02f, NormalBias * DistanceToCamera / GetCotanHalfFieldOfView().x);
 		TranslatedWorldPosition += WorldNormal * NormalBiasLength;
 
 		// Do not run contact shadow when computing the hair shadow mask (i.e. shadow mask applied on hair, has the scene 
@@ -292,7 +307,6 @@ FVirtualShadowMapSampleResult ProjectLight(
 				TranslatedWorldPosition,
 				SMRTRayOffset,
 				Noise,
-				bBackfaceCull,
 				WorldNormal);
 		}
 		else
@@ -314,7 +328,6 @@ FVirtualShadowMapSampleResult ProjectLight(
 				TranslatedWorldPosition,
 				SMRTRayOffset,
 				Noise,
-				bBackfaceCull,
 				WorldNormal);
 		}
 		else
@@ -379,7 +392,6 @@ FVirtualShadowMapSampleResult ProjectLight(
 					PixelPos,
 					TranslatedWorldPosition,
 					WorldNormal,
-					VirtualShadowMap.SMRTCotMaxRayAngleFromLight,
 					RayIndex,
 					HairRayCount,
 					RayStart,
diff --git a/Engine/Shaders/Private/VirtualShadowMaps/VirtualShadowMapProjectionDirectional.ush b/Engine/Shaders/Private/VirtualShadowMaps/VirtualShadowMapProjectionDirectional.ush
--- a/Engine/Shaders/Private/VirtualShadowMaps/VirtualShadowMapProjectionDirectional.ush
+++ b/Engine/Shaders/Private/VirtualShadowMaps/VirtualShadowMapProjectionDirectional.ush
@@ -180,8 +180,8 @@ FVirtualShadowMapSampleResult TraceDirectional(
 	float3 TranslatedWorldPosition,
 	float RayStartOffset,
 	const float Noise,
-	bool bCullBackface,
-	float3 WorldNormal)
+	float3 WorldNormal,
+	const FSMRTTraceSettings Settings = GetSMRTTraceSettingsDirectional())
 {
 	float3 ViewPosition = mul(float4(TranslatedWorldPosition, 1.0f), View.TranslatedWorldToView).xyz;
 	float DistanceFromViewOrigin = length(ViewPosition);
@@ -190,12 +190,6 @@ FVirtualShadowMapSampleResult TraceDirectional(
 	Result.bValid = true;	// TODO: false if all samples of all rays miss pages?
 	Result.ShadowFactor = 1.0f;
 
-	if (bCullBackface && IsBackfaceToDirectionalLight(WorldNormal, Light.Direction, Light.SourceRadius))
-	{
-		Result.ShadowFactor = 0.0f;		// Fully shadowed on backfaces
-		return Result;
-	}
-
 	// Find the best resolution mapped clipmap at the ray origin location
 	const int ClipmapId = GetMappedClipmapId(VirtualShadowMapId, TranslatedWorldPosition);
 	if (ClipmapId < 0)
@@ -216,7 +210,7 @@ FVirtualShadowMapSampleResult TraceDirectional(
 	// in some different way.
 	//
 	// 0.5 is arbitrary but makes it consistent with the previous dither scale cvar
-	float PerLightTexelDitherScale = VirtualShadowMap.SMRTTexelDitherScaleDirectional * ProjectionData.TexelDitherScale;
+	float PerLightTexelDitherScale = Settings.TexelDitherScale * ProjectionData.TexelDitherScale;
 	float DitherScale = 0.0f;
 	if (PerLightTexelDitherScale > 0.0f)
 	{
@@ -242,7 +236,7 @@ FVirtualShadowMapSampleResult TraceDirectional(
 	uint i = 0;
 	float OccluderDistanceSum = 0.0f;
 	float MaxOccluderDistance = -1.0f;
-	const uint MaxRayCount = VirtualShadowMap.SMRTRayCountDirectional;
+	const uint MaxRayCount = Settings.RayCount;
 	for ( ; i < MaxRayCount; i++)
 	{
 		// One sample for ray, one for texel dither
@@ -258,8 +252,8 @@ FVirtualShadowMapSampleResult TraceDirectional(
 			RayStartOffset,
 			DepthSlopeUV,
 			TexelOffset,
-			VirtualShadowMap.SMRTExtrapolateMaxSlopeDirectional);
-		FSMRTResult SMRTResult = SMRTRayCast(RayState, VirtualShadowMap.SMRTSamplesPerRayDirectional, Noise);
+			Settings.ExtrapolateMaxSlope);
+		FSMRTResult SMRTResult = SMRTRayCast(RayState, Settings.SamplesPerRay, Noise);
 
 		// Debug output (DCE'd if not used)
 		Result.ClipmapOrMipLevel = GetVirtualShadowMapProjectionData(RayState.SampledClipmapId).ClipmapLevel;
@@ -283,7 +277,7 @@ FVirtualShadowMapSampleResult TraceDirectional(
 		}
 
 		#if COMPUTESHADER
-		if (VirtualShadowMap.SMRTAdaptiveRayCount > 0)
+		if (Settings.AdaptiveRayCount > 0)
 		{
 			// TODO: Adapt this heuristic based on SMRTAdaptiveRayCount as well?
 			if( i == 0 )
@@ -297,7 +291,7 @@ FVirtualShadowMapSampleResult TraceDirectional(
 					break;
 				}
 			}
-			else if( i >= VirtualShadowMap.SMRTAdaptiveRayCount )
+			else if (i >= Settings.AdaptiveRayCount)
 			{
 				// After 2 iterations and all have hit, assume umbra
 				bool bAllLanesHit = WaveActiveAllTrue( RayMissCount == 0 );
diff --git a/Engine/Shaders/Private/VirtualShadowMaps/VirtualShadowMapProjectionSpot.ush b/Engine/Shaders/Private/VirtualShadowMaps/VirtualShadowMapProjectionSpot.ush
--- a/Engine/Shaders/Private/VirtualShadowMaps/VirtualShadowMapProjectionSpot.ush
+++ b/Engine/Shaders/Private/VirtualShadowMaps/VirtualShadowMapProjectionSpot.ush
@@ -240,65 +240,21 @@ FSMRTResult ShadowRayCastPointLight(
 	return Result;
 }
 
-// MaxBits version reduces unrolled iterations
-float4 SobolSampler(uint SampleIndex, inout uint Seed, uint MaxBits)
-{
-	// first scramble the index to decorelate from other 4-tuples
-	uint SobolIndex = FastOwenScrambling(SampleIndex, EvolveSobolSeed(Seed)) & ( (1U << MaxBits) - 1U );
-	// now get Sobol' point from this index
-	uint4 Result = uint4(SobolIndex, 0, 0, 0);
-	UNROLL for (uint b = 0, v = 1; b < MaxBits; b++)
-	{
-		uint IndexBit = (SobolIndex >> b) & 1;
-		Result.y ^= IndexBit * v;
-		Result.zw ^= IndexBit * SobolMatrices[b];
-		v ^= v << 1;
-	}
-
-	// finally scramble the points to avoid structured artifacts
-	Result.x = FastOwenScrambling(Result.x, EvolveSobolSeed(Seed));
-	Result.y = FastOwenScrambling(Result.y, EvolveSobolSeed(Seed));
-	Result.z = FastOwenScrambling(Result.z, EvolveSobolSeed(Seed));
-	Result.w = FastOwenScrambling(Result.w, EvolveSobolSeed(Seed));
-
-	// output as float in [0,1) taking care not to skew the distribution
-	// due to the non-uniform spacing of floats in this range
-	return (Result >> 8) * 5.96046447754e-08; // * 2^-24
-}
-
 // Normal should be normalized
 bool IsBackfaceToLocalLight(
 	float3 ToLight,
 	float3 Normal,
-	float LightSourceRadius,
-	out float3 OutConeAxis,
-	out float OutConeSin,
-	out float OutDistToLight)
+	float LightSourceRadius)
 {
 	float DistSqr = dot(ToLight, ToLight);
 	float Falloff = rcp(DistSqr + 1);
 	float InvDist = rsqrt(DistSqr);
-
-	OutConeAxis = ToLight * InvDist;
-	OutConeSin = LightSourceRadius * InvDist;
-	OutDistToLight = DistSqr * InvDist;
-
 	float SinAlphaSqr = saturate(Pow2(LightSourceRadius) * Falloff);
 	float SinAlpha = sqrt(SinAlphaSqr);
-	bool bBackface = dot(Normal, OutConeAxis) < -SinAlpha;
-
+	bool bBackface = dot(Normal, ToLight * InvDist) < -SinAlpha;
 	return bBackface;
 }
 
-// Version without outputs
-bool IsBackfaceToLocalLight(float3 ToLight, float3 Normal, float LightSourceRadius)
-{
-	float3 Dummy1;
-	float Dummy2;
-	float Dummy3;
-	return IsBackfaceToLocalLight(ToLight, Normal, LightSourceRadius, Dummy1, Dummy2, Dummy3);
-}
-
 // Shadow resolution scale that roughly matches the LOD function in VSM page marking
 float ComputeShadowResolutionScale(FVirtualShadowMapProjectionShaderData ProjectionData, float SceneDepth)
 {
@@ -313,6 +269,17 @@ float ComputeShadowResolutionScale(FVirtualShadowMapProjectionShaderData Project
 	return ShadowResolutionScale;
 }
 
+float VirtualShadowMapGetClampedRayDistanceScaleLocal(float CosTheta)
+{
+	// The further afield our rays go the poorer our approximation is as the "bend" due to our testing
+	// against a shadow map instead of along the ray increases. Thus we avoid going all the way to the light
+	// where the projection becomes extreme.
+
+	// This clamping function is a bit arbitrary but is fairly smooth and avoids the worst artifacts in practice
+	float SinTheta = sqrt(1.0f - CosTheta * CosTheta);
+	return 0.75f * saturate(1.5f / (CosTheta + VirtualShadowMap.SMRTCotMaxRayAngleFromLight * SinTheta));
+}
+
 FVirtualShadowMapSampleResult TraceLocalLight(
 	int VirtualShadowMapId,
 	FLightShaderParameters Light,
@@ -321,199 +288,185 @@ FVirtualShadowMapSampleResult TraceLocalLight(
 	float3 TranslatedWorldPosition,
 	float RayStartOffset,
 	const float Noise,
-	const bool bCullBackface,
-	float3 WorldNormal)
+	float3 WorldNormal,
+	const FSMRTTraceSettings Settings = GetSMRTTraceSettingsLocal())
 {
 	const uint LightType = Light.SpotAngles.x > -2.0f ? LIGHT_TYPE_SPOT : LIGHT_TYPE_POINT;
+	
 	const float3 ToLight = Light.TranslatedWorldPosition - TranslatedWorldPosition;
-
-	float3 ConeAxis;
-	float ConeSin;
-	float DistToLight;
-	bool bBackface = IsBackfaceToLocalLight(ToLight, WorldNormal, Light.SourceRadius, ConeAxis, ConeSin, DistToLight);
+	const float3 ConeAxis = normalize(ToLight);
+	const float DistToLight = length(ToLight);
+	const float ConeSin = Light.SourceRadius / DistToLight;
 
 	FVirtualShadowMapSampleResult Result = InitVirtualShadowMapSampleResult();
 	Result.bValid = true;
 	Result.ShadowFactor = 0.0f;
 
-	if (!bBackface || !bCullBackface)
+	int OffsetVirtualShadowMapId = VirtualShadowMapId;
+	if (LightType == LIGHT_TYPE_POINT)
 	{
-		float StepOffset = Noise;
-
-		uint TileID = MortonEncode( PixelPos );
-
-		int OffsetVirtualShadowMapId = VirtualShadowMapId;
-		if (LightType == LIGHT_TYPE_POINT)
-		{
-			OffsetVirtualShadowMapId += VirtualShadowMapGetCubeFace(-ToLight);
-		}
+		OffsetVirtualShadowMapId += VirtualShadowMapGetCubeFace(-ToLight);
+	}
 
-		// Move from PrimaryView translated world to Shadow translated world
-		// NOTE: Cube map faces all share the same PreViewTranslation and projection depth scale
-		FVirtualShadowMapProjectionShaderData ProjectionData = GetVirtualShadowMapProjectionData(OffsetVirtualShadowMapId);
-		const float3 PrimaryToShadowTranslation = DFFastLocalSubtractDemote(ProjectionData.PreViewTranslation, PrimaryView.PreViewTranslation);
-		const float3 ShadowTranslatedWorld = TranslatedWorldPosition + PrimaryToShadowTranslation;
-
-		const float2 DepthSlopeUV = ComputeDepthSlopeLocalUV(OffsetVirtualShadowMapId, ShadowTranslatedWorld, WorldNormal);
-
-		const float ShadowResolutionScale = ComputeShadowResolutionScale(ProjectionData, SceneDepth);
-		// Rough scaling from world units -> post projection depth.
-		// We use a radially symmetric mapping here instead of the true function to minimize issues crossing cubemap faces.
-		const float ShadowDepthFunctionScale = abs(ProjectionData.ShadowViewToClipMatrix._43 / length2(ToLight));
-
-		const float DitherScale = (VirtualShadowMap.SMRTTexelDitherScaleLocal * ProjectionData.TexelDitherScale) * ShadowResolutionScale;
-		const float MaxDitherSlope = VirtualShadowMap.SMRTMaxSlopeBiasLocal * DitherScale;
-		const float MaxDepthBias = VirtualShadowMap.SMRTMaxSlopeBiasLocal * ShadowResolutionScale * ShadowDepthFunctionScale;
-
-		// Moving the ray origin means we need to attempt to account for local receiver geometry to avoid
-		// incorrect self-shadowing. To this end we move the point on a plane aligned with the shading normal.
-		// It would be much better to use geometric (faceted) normals - as with optimal bias - but we do not
-		// have access to those. Shading normals that stray too far from the real geometry will cause artifacts.
-		const float3x3 ConeTangentToWorldMatrix = GetTangentBasis(ConeAxis);
-		const float3 NormalPlaneTangent = mul(ConeTangentToWorldMatrix, WorldNormal);
-		const float2 DepthSlopeTangent = -NormalPlaneTangent.xy / NormalPlaneTangent.z;
-
-		// If source angle = 0, we don't need multiple samples on the same ray
-		uint SamplesPerRay = VirtualShadowMap.SMRTSamplesPerRayLocal;
-		if (ConeSin == 0.0f)
-		{
-			SamplesPerRay = 0;
-		}
+	// Move from PrimaryView translated world to Shadow translated world
+	// NOTE: Cube map faces all share the same PreViewTranslation and projection depth scale
+	FVirtualShadowMapProjectionShaderData ProjectionData = GetVirtualShadowMapProjectionData(OffsetVirtualShadowMapId);
+	const float3 PrimaryToShadowTranslation = DFFastLocalSubtractDemote(ProjectionData.PreViewTranslation, PrimaryView.PreViewTranslation);
+	const float3 ShadowTranslatedWorld = TranslatedWorldPosition + PrimaryToShadowTranslation;
+
+	const float2 DepthSlopeUV = ComputeDepthSlopeLocalUV(OffsetVirtualShadowMapId, ShadowTranslatedWorld, WorldNormal);
+
+	const float ShadowResolutionScale = ComputeShadowResolutionScale(ProjectionData, SceneDepth);
+	// Rough scaling from world units -> post projection depth.
+	// We use a radially symmetric mapping here instead of the true function to minimize issues crossing cubemap faces.
+	const float ShadowDepthFunctionScale = abs(ProjectionData.ShadowViewToClipMatrix._43 / length2(ToLight));
+
+	const float DitherScale = (Settings.TexelDitherScale * ProjectionData.TexelDitherScale) * ShadowResolutionScale;
+	const float MaxDitherSlope = VirtualShadowMap.SMRTMaxSlopeBiasLocal * DitherScale;
+	const float MaxDepthBias = VirtualShadowMap.SMRTMaxSlopeBiasLocal * ShadowResolutionScale * ShadowDepthFunctionScale;
+
+	// Moving the ray origin means we need to attempt to account for local receiver geometry to avoid
+	// incorrect self-shadowing. To this end we move the point on a plane aligned with the shading normal.
+	// It would be much better to use geometric (faceted) normals - as with optimal bias - but we do not
+	// have access to those. Shading normals that stray too far from the real geometry will cause artifacts.
+	const float3x3 ConeTangentToWorldMatrix = GetTangentBasis(ConeAxis);
+	const float3 NormalPlaneTangent = mul(ConeTangentToWorldMatrix, WorldNormal);
+	const float2 DepthSlopeTangent = -NormalPlaneTangent.xy / NormalPlaneTangent.z;
+
+	// If source angle = 0, we don't need multiple samples on the same ray
+	uint SamplesPerRay = Settings.SamplesPerRay;
+	if (ConeSin == 0.0f)
+	{
+		SamplesPerRay = 0;
+	}
 
-		uint i = 0;
-		uint RayMissCount = 0;
-		float OccluderDistanceSum = 0.0f;
-		float MaxOccluderDistance = -1.0f;
-		const uint MaxRayCount = VirtualShadowMap.SMRTRayCountLocal;
-		for( ; i < MaxRayCount; i++ )
-		{
-			float4 RandSample = VirtualShadowMapGetRandomSample(PixelPos, View.StateFrameIndex, i, MaxRayCount);
+	uint i = 0;
+	uint RayMissCount = 0;
+	float StepOffset = Noise;
+	float OccluderDistanceSum = 0.0f;
+	float MaxOccluderDistance = -1.0f;
+	const uint MaxRayCount = Settings.RayCount;
+	for( ; i < MaxRayCount; i++ )
+	{
+		float4 RandSample = VirtualShadowMapGetRandomSample(PixelPos, View.StateFrameIndex, i, MaxRayCount);
 			
-			float2 LightUV = UniformSampleDiskConcentricApprox( RandSample.xy ).xy;
-			LightUV *= ConeSin;
-
-			float SinTheta2 = dot( LightUV, LightUV );
-			float SinTheta = sqrt( SinTheta2 );
-			float CosTheta = sqrt( 1 - SinTheta2 );
-			float3 Dir = mul(float3(LightUV, CosTheta), ConeTangentToWorldMatrix);
+		float2 LightUV = UniformSampleDiskConcentricApprox( RandSample.xy ).xy;
+		LightUV *= ConeSin;
+		float SinTheta2 = dot(LightUV, LightUV);
+		float CosTheta = sqrt(1.0f - SinTheta2);
+		float3 Dir = mul(float3(LightUV, CosTheta), ConeTangentToWorldMatrix);
 
-			float3 RayStartWithDither = ShadowTranslatedWorld;
+		float3 RayStartWithDither = ShadowTranslatedWorld;
 
-			if (DitherScale > 0.0f)
-			{
-				// Offset following the plane of the receiver
-				float2 RandomOffset = (RandSample.zw - 0.5f) * DitherScale;
-				float OffsetZ = min(MaxDitherSlope, 2.0f * max(0.0f, dot(DepthSlopeTangent, RandomOffset)));
-				float3 Offset = mul(float3(RandomOffset, OffsetZ), ConeTangentToWorldMatrix);
-				RayStartWithDither += Offset;
-			}
+		if (DitherScale > 0.0f)
+		{
+			// Offset following the plane of the receiver
+			float2 RandomOffset = (RandSample.zw - 0.5f) * DitherScale;
+			float OffsetZ = min(MaxDitherSlope, 2.0f * max(0.0f, dot(DepthSlopeTangent, RandomOffset)));
+			float3 Offset = mul(float3(RandomOffset, OffsetZ), ConeTangentToWorldMatrix);
+			RayStartWithDither += Offset;
+		}
 
-			// The further afield our rays go the poorer our approximation is as the "bend" due to our testing
-			// against a shadow map instead of along the ray increases. Thus we avoid going all the way to the light
-			// where the projection becomes extreme.
+		float ClampedEnd = DistToLight * VirtualShadowMapGetClampedRayDistanceScaleLocal(CosTheta);
+		float Start = min(RayStartOffset, ClampedEnd - 1e-6f);
 
-			// Clip ray to max angle to limit penumbra sizes
-			float ClampedEnd = DistToLight * 0.75 * saturate( 1.5 / ( CosTheta + VirtualShadowMap.SMRTCotMaxRayAngleFromLight * SinTheta ) );
-			float Start = min(RayStartOffset, ClampedEnd - 1e-6f);
+		float3 RayStart = RayStartWithDither + Dir * Start;
+		float3 RayEnd = RayStartWithDither + Dir * ClampedEnd;
 
-			float3 RayStart = RayStartWithDither + Dir * Start;
-			float3 RayEnd = RayStartWithDither + Dir * ClampedEnd;
+		float DepthBias = 0.0f;
+		if (MaxDepthBias > 0.0f)
+		{
+			// Clamp depth bias to avoid excessive degenerate slope biases causing flickering lit pixels
+			DepthBias = min(MaxDepthBias, ComputeOptimalSlopeBiasLocal(OffsetVirtualShadowMapId, RayStart, DepthSlopeUV));
+		}
 
-			float DepthBias = 0.0f;
-			if (MaxDepthBias > 0.0f)
-			{
-				// Clamp depth bias to avoid excessive degenerate slope biases causing flickering lit pixels
-				DepthBias = min(MaxDepthBias, ComputeOptimalSlopeBiasLocal(OffsetVirtualShadowMapId, RayStart, DepthSlopeUV));
-			}
+		FSMRTResult SMRTResult;
+		FSMRTSingleRayState RayState;	// For debug output
+		if ( LightType == LIGHT_TYPE_SPOT )
+		{
+			SMRTResult = ShadowRayCastSpotLight(
+				VirtualShadowMapId,
+				RayStart,
+				RayEnd,
+				DepthBias,
+				SamplesPerRay,
+				StepOffset,
+				Settings.ExtrapolateMaxSlope,
+				RayState);
+		}
+		else
+		{
+			SMRTResult = ShadowRayCastPointLight(
+				VirtualShadowMapId,
+				RayStart,
+				RayEnd,
+				DepthBias,
+				SamplesPerRay,
+				StepOffset,
+				Settings.ExtrapolateMaxSlope,
+				RayState);
+		}
 
-			FSMRTResult SMRTResult;
-			FSMRTSingleRayState RayState;	// For debug output
-			if ( LightType == LIGHT_TYPE_SPOT )
-			{
-				SMRTResult = ShadowRayCastSpotLight(
-					VirtualShadowMapId,
-					RayStart,
-					RayEnd,
-					DepthBias,
-					SamplesPerRay,
-					StepOffset,
-					VirtualShadowMap.SMRTExtrapolateMaxSlopeLocal,
-					RayState);
-			}
-			else
-			{
-				SMRTResult = ShadowRayCastPointLight(
-					VirtualShadowMapId,
-					RayStart,
-					RayEnd,
-					DepthBias,
-					SamplesPerRay,
-					StepOffset,
-					VirtualShadowMap.SMRTExtrapolateMaxSlopeLocal,
-					RayState);
-			}
+		if (SMRTResult.bValidHit)
+		{
+			// TODO: Do we want to mess with this at all due to our offset?
+			float ReceiverDistance = length(RayStart);
+			float OccluderDistance = ComputeOccluderDistancePerspective(
+				// Re-fetch this here to avoid reg pressure
+				GetVirtualShadowMapProjectionData(RayState.VirtualShadowMapId).ShadowViewToClipMatrix,
+				SMRTResult.HitDepth,
+				saturate(RayState.RayStartUVz.z),
+				ReceiverDistance);
+
+			OccluderDistanceSum += OccluderDistance;
+			MaxOccluderDistance = max(MaxOccluderDistance, OccluderDistance);
+		}
+		else
+		{
+			++RayMissCount;
+		}
+		
+		// Debug output (DCE'd if not used)
+		Result.ClipmapOrMipLevel = RayState.MipLevel;
+		Result.VirtualTexelAddress = RayState.VirtualTexelAddress;
+		Result.PhysicalTexelAddress = RayState.PhysicalTexelAddress;
 
-			if (SMRTResult.bValidHit)
-			{
-				// TODO: Do we want to mess with this at all due to our offset?
-				float ReceiverDistance = length(RayStart);
-				float OccluderDistance = ComputeOccluderDistancePerspective(
-					// Re-fetch this here to avoid reg pressure
-					GetVirtualShadowMapProjectionData(RayState.VirtualShadowMapId).ShadowViewToClipMatrix,
-					SMRTResult.HitDepth,
-					saturate(RayState.RayStartUVz.z),
-					ReceiverDistance);
-
-				OccluderDistanceSum += OccluderDistance;
-				MaxOccluderDistance = max(MaxOccluderDistance, OccluderDistance);
-			}
-			else
+		if (MaxRayCount > 1 && Settings.AdaptiveRayCount > 0)
+		{
+			// TODO: Adapt this heuristic based on SMRTAdaptiveRayCount as well?
+			if( i == 0 )
 			{
-				++RayMissCount;
-			}
-			
-			// Debug output (DCE'd if not used)
-			Result.ClipmapOrMipLevel = RayState.MipLevel;
-			Result.VirtualTexelAddress = RayState.VirtualTexelAddress;
-			Result.PhysicalTexelAddress = RayState.PhysicalTexelAddress;
+				bool bHit = SMRTResult.bValidHit;
 
-			if (VirtualShadowMap.SMRTAdaptiveRayCount > 0)
-			{
-				// TODO: Adapt this heuristic based on SMRTAdaptiveRayCount as well?
-				if( i == 0 )
+				// All lanes missed
+				bool bAllLanesMiss = WaveActiveAllTrue( !bHit );
+				if( bAllLanesMiss )
 				{
-					bool bHit = SMRTResult.bValidHit;
-
-					// All lanes missed
-					bool bAllLanesMiss = WaveActiveAllTrue( !bHit );
-					if( bAllLanesMiss )
-					{
-						break;
-					}
+					break;
 				}
-				else if( i >= VirtualShadowMap.SMRTAdaptiveRayCount )
+			}
+			else if( i >= Settings.AdaptiveRayCount)
+			{
+				// After N iterations and all have hit, assume umbra
+				bool bAllLanesHit = WaveActiveAllTrue( RayMissCount == 0 );
+				if( bAllLanesHit )
 				{
-					// After N iterations and all have hit, assume umbra
-					bool bAllLanesHit = WaveActiveAllTrue( RayMissCount == 0 );
-					if( bAllLanesHit )
-					{
-						break;
-					}
+					break;
 				}
 			}
-
-			uint Seed = asuint( StepOffset );
-			StepOffset = ( EvolveSobolSeed( Seed ) >> 8 ) * 5.96046447754e-08; // * 2^-24;
 		}
-		uint RayCount = min(i + 1, MaxRayCount);		// break vs regular for loop exit
-
-		float OccluderDistance = (OccluderDistanceSum / float(max(1, RayCount - RayMissCount)));
-		//OccluderDistance = MaxOccluderDistance;
 
-		Result.ShadowFactor = float(RayMissCount) / float(RayCount);
-		Result.RayCount = RayCount;
-		Result.OccluderDistance = OccluderDistance;
+		uint Seed = asuint( StepOffset );
+		StepOffset = ( EvolveSobolSeed( Seed ) >> 8 ) * 5.96046447754e-08; // * 2^-24;
 	}
+	uint RayCount = min(i + 1, MaxRayCount);		// break vs regular for loop exit
+
+	float OccluderDistance = (OccluderDistanceSum / float(max(1, RayCount - RayMissCount)));
+	//OccluderDistance = MaxOccluderDistance;
+
+	Result.ShadowFactor = float(RayMissCount) / float(RayCount);
+	Result.RayCount = RayCount;
+	Result.OccluderDistance = OccluderDistance;
 
 	return Result;
 }
@@ -524,43 +477,27 @@ bool GenerateRayLocalLight(
 	uint2 PixelPos,
 	float3 TranslatedWorldPosition,
 	float3 WorldNormal,
-	float CotMaxRayAngleFromLight,
 	uint RayIndex,
 	uint RayCount,
 	inout float3 OutRayStart,
 	inout float3 OutRayEnd)
 {
 	const float3 ToLight = Light.TranslatedWorldPosition - TranslatedWorldPosition;
+	const float3 ConeAxis = normalize(ToLight);
+	const float DistToLight = length(ToLight);
+	const float ConeSin = Light.SourceRadius / DistToLight;
 
-	float3 ConeAxis;
-	float ConeSin;
-	float DistToLight;
-	bool bBackface = IsBackfaceToLocalLight(ToLight, WorldNormal, Light.SourceRadius, ConeAxis, ConeSin, DistToLight);
-	
+	bool bBackface = IsBackfaceToLocalLight(ToLight, WorldNormal, Light.SourceRadius);	
 	if (!bBackface)
 	{
-		uint TileID = MortonEncode(PixelPos);
-		RandomSequence RandSequence;
-		// Combine pixel-level and sample-level bits into the sample index (visible structure will be hidden by owen scrambling of the index)
-		RandSequence.SampleIndex = reversebits((65536 * View.StateFrameIndex + TileID) * RayCount + RayIndex);
-		RandSequence.SampleSeed = 0; // always use the same sequence
-
-		float2 E = SobolSampler(RandSequence.SampleIndex, RandSequence.SampleSeed, 8).xy;
-		float2 LightUV = UniformSampleDiskConcentricApprox(E).xy;
+		float2 E = VirtualShadowMapGetRandomSample(PixelPos, View.StateFrameIndex, RayIndex, RayCount).xy;
+		float2 LightUV = UniformSampleDiskConcentricApprox(E);
 		LightUV *= ConeSin;
-
 		float SinTheta2 = dot(LightUV, LightUV);
-		float SinTheta = sqrt(SinTheta2);
 		float CosTheta = sqrt(1 - SinTheta2);
-
 		float3 Dir = TangentToWorld(float3(LightUV, CosTheta), ConeAxis);
 
-		// The further afield our rays go the poorer our approximation is as the "bend" due to our testing
-		// against a shadow map instead of along the ray increases. Thus we avoid going all the way to the light
-		// where the projection becomes extreme.
-
-		// Clip ray to max angle to limit penumbra sizes
-		float ClampedLength = DistToLight * saturate(1.5 / (CosTheta + CotMaxRayAngleFromLight * SinTheta));
+		float ClampedLength = DistToLight * VirtualShadowMapGetClampedRayDistanceScaleLocal(CosTheta);
 
 		OutRayStart = TranslatedWorldPosition;
 		OutRayEnd = TranslatedWorldPosition + Dir * ClampedLength;
diff --git a/Engine/Shaders/Private/VirtualShadowMaps/VirtualShadowMapSMRTCommon.ush b/Engine/Shaders/Private/VirtualShadowMaps/VirtualShadowMapSMRTCommon.ush
--- a/Engine/Shaders/Private/VirtualShadowMaps/VirtualShadowMapSMRTCommon.ush
+++ b/Engine/Shaders/Private/VirtualShadowMaps/VirtualShadowMapSMRTCommon.ush
@@ -64,4 +64,36 @@ float4 VirtualShadowMapGetRandomSample(
 		return RandomSequence_GenerateSample4D(RandSequence);
 	}
 #endif
-}
\ No newline at end of file
+}
+
+// Handy structure for locally modifying some common SMRT parameters
+struct FSMRTTraceSettings
+{
+	int AdaptiveRayCount;
+	int RayCount;
+	int SamplesPerRay;
+	float ExtrapolateMaxSlope;
+	float TexelDitherScale;
+};
+
+FSMRTTraceSettings GetSMRTTraceSettingsDirectional()
+{
+	FSMRTTraceSettings Settings;
+	Settings.AdaptiveRayCount		= VirtualShadowMap.SMRTAdaptiveRayCount;
+	Settings.RayCount				= VirtualShadowMap.SMRTRayCountDirectional;
+	Settings.SamplesPerRay			= VirtualShadowMap.SMRTSamplesPerRayDirectional;
+	Settings.ExtrapolateMaxSlope	= VirtualShadowMap.SMRTExtrapolateMaxSlopeDirectional;
+	Settings.TexelDitherScale		= VirtualShadowMap.SMRTTexelDitherScaleDirectional;
+	return Settings;
+}
+
+FSMRTTraceSettings GetSMRTTraceSettingsLocal()
+{
+	FSMRTTraceSettings Settings;
+	Settings.AdaptiveRayCount		= VirtualShadowMap.SMRTAdaptiveRayCount;
+	Settings.RayCount				= VirtualShadowMap.SMRTRayCountLocal;
+	Settings.SamplesPerRay			= VirtualShadowMap.SMRTSamplesPerRayLocal;
+	Settings.ExtrapolateMaxSlope	= VirtualShadowMap.SMRTExtrapolateMaxSlopeLocal;
+	Settings.TexelDitherScale		= VirtualShadowMap.SMRTTexelDitherScaleLocal;
+	return Settings;
+}
