Alpha: Continuing POST_PROCESS_ALPHA removal in post-processing, which was missed during CL 35097313 resubmission.

diff --git a/Engine/Shaders/Private/MotionBlur/MotionBlurApply.usf b/Engine/Shaders/Private/MotionBlur/MotionBlurApply.usf
--- a/Engine/Shaders/Private/MotionBlur/MotionBlurApply.usf
+++ b/Engine/Shaders/Private/MotionBlur/MotionBlurApply.usf
@@ -909,7 +909,7 @@ void MainCS(
 			OutputColor2.rgb = OutputColor2.rgb * PostMotionBlurTranslucency2.a + PostMotionBlurTranslucency2.rgb;
 			OutputColor3.rgb = OutputColor3.rgb * PostMotionBlurTranslucency3.a + PostMotionBlurTranslucency3.rgb;
 
-			#if POST_PROCESS_ALPHA
+			#if CONFIG_SCENE_COLOR_ALPHA
 				OutputColor0.a = OutputColor0.a * PostMotionBlurTranslucency0.a;
 				OutputColor1.a = OutputColor1.a * PostMotionBlurTranslucency1.a;
 				OutputColor2.a = OutputColor2.a * PostMotionBlurTranslucency2.a;
@@ -974,7 +974,7 @@ void MainCS(
 		}
 		
 		// Needed to avoid crash in shader compiler of Xclipse driver
-		#if !POST_PROCESS_ALPHA
+		#if !CONFIG_SCENE_COLOR_ALPHA
 			QuarterResOutput.a = 0;
 		#endif
 		bool bIsValid = all(OutputPixelCoord < mb_short2(Color_ViewportMax));
@@ -1049,7 +1049,7 @@ void MainCS(
 		{
 			OutputColor.rgb = OutputColor.rgb * PostMotionBlurTranslucency.a + PostMotionBlurTranslucency.rgb;
 			
-			#if POST_PROCESS_ALPHA
+			#if CONFIG_SCENE_COLOR_ALPHA
 				OutputColor.a = OutputColor.a * PostMotionBlurTranslucency.a;
 			#endif
 		}
@@ -1097,7 +1097,7 @@ void MainCS(
 		}
 		
 		// Needed to avoid crash in shader compiler of Xclipse driver
-		#if !POST_PROCESS_ALPHA
+		#if !CONFIG_SCENE_COLOR_ALPHA
 			HalfResOutput.a = 0;
 		#endif
 		mb_half4 QuarterResOutput;
@@ -1129,7 +1129,7 @@ void MainCS(
 		}
 
 		// Needed to avoid crash in shader compiler of Xclipse driver
-		#if !POST_PROCESS_ALPHA
+		#if !CONFIG_SCENE_COLOR_ALPHA
 			QuarterResOutput.a = 0;
 		#endif
 		bool bIsValid = all(OutputPixelCoord < mb_short2(Color_ViewportMax));
