Small VSM refactor; move a few functions to different files

#rnx

diff --git a/Engine/Shaders/Private/VirtualShadowMaps/VirtualShadowMapProjection.usf b/Engine/Shaders/Private/VirtualShadowMaps/VirtualShadowMapProjection.usf
--- a/Engine/Shaders/Private/VirtualShadowMaps/VirtualShadowMapProjection.usf
+++ b/Engine/Shaders/Private/VirtualShadowMaps/VirtualShadowMapProjection.usf
@@ -17,6 +17,7 @@
 #include "VirtualShadowMapProjectionDirectional.ush"
 #include "VirtualShadowMapProjectionSpot.ush"
 #include "VirtualShadowMapTransmissionCommon.ush"
+#include "VirtualShadowMapScreenRayTrace.ush"
 #include "VirtualShadowMapLightGrid.ush"
 #include "/Engine/Shared/VirtualShadowMapDefinitions.h"
 
@@ -92,59 +93,6 @@ int GetSMRTRayCount()
 SCHEDULER_MIN_PRESSURE
 MAX_OCCUPANCY
 
-// Screen space ray trace to attempt to skip over ambiguous regions near the receiver surface
-// Returns length at which to start the virtual shadow map ray; usually this is where the screen ray ended or went behind a surface
-float VirtualShadowMapScreenRayCast(
-	float3 RayOriginTranslatedWorld,
-	float3 RayDirection,
-	float RayLength,
-	float Dither)
-{
-	float4 RayStartClip	= mul(float4(RayOriginTranslatedWorld, 1), View.TranslatedWorldToClip);
-	float4 RayDirClip	= mul(float4(RayDirection * RayLength, 0), View.TranslatedWorldToClip);
-	float4 RayEndClip	= RayStartClip + RayDirClip;
-
-	float3 RayStartScreen = RayStartClip.xyz / RayStartClip.w;
-	float3 RayEndScreen = RayEndClip.xyz / RayEndClip.w;
-
-	float3 RayStepScreen = RayEndScreen - RayStartScreen;
-
-	float3 RayStartUVz = float3(RayStartScreen.xy * View.ScreenPositionScaleBias.xy + View.ScreenPositionScaleBias.wz, RayStartScreen.z);
-	float3 RayStepUVz  = float3(RayStepScreen.xy * View.ScreenPositionScaleBias.xy, RayStepScreen.z);
-
-	float4 RayDepthClip	= RayStartClip + mul(float4(0, 0, RayLength, 0), View.ViewToClip);
-	float3 RayDepthScreen = RayDepthClip.xyz / RayDepthClip.w;
-
-	const int Steps = SCREEN_RAY_SAMPLES;
-	float StepOffset = Dither - 0.5f;
-	const float Step = 1.0 / Steps;
-	float SampleTime = StepOffset * Step + Step;
-
-	const float StartDepth = SceneTexturesStruct.SceneDepthTexture.SampleLevel(SceneTexturesStruct_SceneDepthTextureSampler, RayStartUVz.xy, 0).r;
-
-	UNROLL
-	for (int i = 0; i < Steps; i++)
-	{
-		float3 SampleUVz = RayStartUVz + RayStepUVz * SampleTime;
-		float SampleDepth = SceneTexturesStruct.SceneDepthTexture.SampleLevel(SceneTexturesStruct_SceneDepthTextureSampler, SampleUVz.xy, 0).r;
-
-		// Avoid self-intersection with the start pixel (exact comparison due to point sampling depth buffer)
-		if (SampleDepth != StartDepth)
-		{
-			if (SampleUVz.z < SampleDepth)
-			{
-				// Behind geometry. Back up a bit along the ray and do the VSM sample from there.
-				return RayLength * max(0.0, SampleTime - 1.5f * Step);
-			}
-		}
-
-		SampleTime += Step;
-	}
-
-	// Got to the end of the ray without going behind or hitting anything
-	return RayLength;
-}
-
 // NOTE: Not currently used by occasionally useful for testing
 float3 GetEstimatedGeoWorldNormal(float3 TranslatedWorldPosition, float3 ShadingNormal)
 {
@@ -181,7 +129,6 @@ float4 ComputeRandom4(uint2 PixelPosition)
 }
 
 int4 ProjectionRect;
-float NormalBias;
 float SubsurfaceMinSourceRadius;
 
 // One pass projection
@@ -275,9 +222,10 @@ FVirtualShadowMapSampleResult ProjectLight(
 	{
 		const float3 WorldNormal = (bIsHairInput || ShadingInfo.bIsHair) ? L : ShadingInfo.WorldNormal;
 
-		const float DistanceToCamera = GetDistanceToCameraFromViewVector(TranslatedWorldPosition - View.TranslatedWorldCameraOrigin);
-		const float NormalBiasLength = ShadingInfo.bIsSubsurface ? 0.0f : max(0.02f, NormalBias * DistanceToCamera / GetCotanHalfFieldOfView().x);
-		TranslatedWorldPosition += WorldNormal * NormalBiasLength;
+		if (!ShadingInfo.bIsSubsurface)
+		{
+			TranslatedWorldPosition += WorldNormal * VirtualShadowMapGetNormalBiasLength(TranslatedWorldPosition);
+		}
 
 		// Do not run contact shadow when computing the hair shadow mask (i.e. shadow mask applied on hair, has the scene 
 		// depth buffer contains fully opaque depth, which create false positive intersection resulting in wrong self shadowing)
diff --git a/Engine/Shaders/Private/VirtualShadowMaps/VirtualShadowMapProjectionCommon.ush b/Engine/Shaders/Private/VirtualShadowMaps/VirtualShadowMapProjectionCommon.ush
--- a/Engine/Shaders/Private/VirtualShadowMaps/VirtualShadowMapProjectionCommon.ush
+++ b/Engine/Shaders/Private/VirtualShadowMaps/VirtualShadowMapProjectionCommon.ush
@@ -119,6 +119,13 @@ FVirtualShadowMapSample SampleVirtualShadowMap(int VirtualShadowMapId, float2 Sh
 	return Result;
 }
 
+// VSM normal bias that smoothly scales with distance to camera
+float3 VirtualShadowMapGetNormalBiasLength(float3 TranslatedWorldPosition)
+{
+	float DistanceToCamera = GetDistanceToCameraFromViewVector(TranslatedWorldPosition - View.TranslatedWorldCameraOrigin);
+	float NormalBiasLength = max(0.02f, VirtualShadowMap.NormalBias * DistanceToCamera / GetCotanHalfFieldOfView().x);
+	return NormalBiasLength;
+}
 
 // Data to convert UV's and depths from one clipmap level basis to another
 struct FVirtualShadowMapClipmapRelativeTransform
@@ -517,3 +524,4 @@ FVirtualShadowMapSampleResult SampleVirtualShadowMapTranslatedWorld(int VirtualS
 		return SampleVirtualShadowMapLocal(VirtualShadowMapId, TranslatedWorldPosition, RayStartDistance);
 	}
 }
+
diff --git a/Engine/Shaders/Private/VirtualShadowMaps/VirtualShadowMapScreenRayTrace.ush b/Engine/Shaders/Private/VirtualShadowMaps/VirtualShadowMapScreenRayTrace.ush
new file mode 100644
--- /dev/null
+++ b/Engine/Shaders/Private/VirtualShadowMaps/VirtualShadowMapScreenRayTrace.ush
@@ -0,0 +1,61 @@
+// Copyright Epic Games, Inc. All Rights Reserved.
+
+#include "../Common.ush"
+#include "../SceneTexturesCommon.ush"
+
+#ifndef SCREEN_RAY_SAMPLES
+#define SCREEN_RAY_SAMPLES 4
+#endif
+
+// Screen space ray trace to attempt to skip over ambiguous regions near the receiver surface
+// Returns length at which to start the virtual shadow map ray; usually this is where the screen ray ended or went behind a surface
+float VirtualShadowMapScreenRayCast(
+	float3 RayOriginTranslatedWorld,
+	float3 RayDirection,
+	float RayLength,
+	float Dither)
+{
+	float4 RayStartClip	= mul(float4(RayOriginTranslatedWorld, 1), View.TranslatedWorldToClip);
+	float4 RayDirClip	= mul(float4(RayDirection * RayLength, 0), View.TranslatedWorldToClip);
+	float4 RayEndClip	= RayStartClip + RayDirClip;
+
+	float3 RayStartScreen = RayStartClip.xyz / RayStartClip.w;
+	float3 RayEndScreen = RayEndClip.xyz / RayEndClip.w;
+
+	float3 RayStepScreen = RayEndScreen - RayStartScreen;
+
+	float3 RayStartUVz = float3(RayStartScreen.xy * View.ScreenPositionScaleBias.xy + View.ScreenPositionScaleBias.wz, RayStartScreen.z);
+	float3 RayStepUVz  = float3(RayStepScreen.xy * View.ScreenPositionScaleBias.xy, RayStepScreen.z);
+
+	float4 RayDepthClip	= RayStartClip + mul(float4(0, 0, RayLength, 0), View.ViewToClip);
+	float3 RayDepthScreen = RayDepthClip.xyz / RayDepthClip.w;
+
+	const int Steps = SCREEN_RAY_SAMPLES;
+	float StepOffset = Dither - 0.5f;
+	const float Step = 1.0 / Steps;
+	float SampleTime = StepOffset * Step + Step;
+
+	const float StartDepth = SceneTexturesStruct.SceneDepthTexture.SampleLevel(SceneTexturesStruct_SceneDepthTextureSampler, RayStartUVz.xy, 0).r;
+
+	UNROLL
+	for (int i = 0; i < Steps; i++)
+	{
+		float3 SampleUVz = RayStartUVz + RayStepUVz * SampleTime;
+		float SampleDepth = SceneTexturesStruct.SceneDepthTexture.SampleLevel(SceneTexturesStruct_SceneDepthTextureSampler, SampleUVz.xy, 0).r;
+
+		// Avoid self-intersection with the start pixel (exact comparison due to point sampling depth buffer)
+		if (SampleDepth != StartDepth)
+		{
+			if (SampleUVz.z < SampleDepth)
+			{
+				// Behind geometry. Back up a bit along the ray and do the VSM sample from there.
+				return RayLength * max(0.0, SampleTime - 1.5f * Step);
+			}
+		}
+
+		SampleTime += Step;
+	}
+
+	// Got to the end of the ray without going behind or hitting anything
+	return RayLength;
+}
diff --git a/Engine/Source/Runtime/Renderer/Private/VirtualShadowMaps/VirtualShadowMapArray.cpp b/Engine/Source/Runtime/Renderer/Private/VirtualShadowMaps/VirtualShadowMapArray.cpp
--- a/Engine/Source/Runtime/Renderer/Private/VirtualShadowMaps/VirtualShadowMapArray.cpp
+++ b/Engine/Source/Runtime/Renderer/Private/VirtualShadowMaps/VirtualShadowMapArray.cpp
@@ -375,6 +375,14 @@ static TAutoConsoleVariable<float> CVarScreenRayLength(
 	ECVF_Scalability | ECVF_RenderThreadSafe
 );
 
+static TAutoConsoleVariable<float> CVarNormalBias(
+	TEXT("r.Shadow.Virtual.NormalBias"),
+	0.5f,
+	TEXT("Receiver offset along surface normal for shadow lookup. Scaled by distance to camera.")
+	TEXT("Higher values avoid artifacts on surfaces nearly parallel to the light, but also visibility offset shadows and increase the chance of hitting unmapped pages."),
+	ECVF_Scalability | ECVF_RenderThreadSafe
+);
+
 static TAutoConsoleVariable<int32> CVarSMRTAdaptiveRayCount(
 	TEXT("r.Shadow.Virtual.SMRT.AdaptiveRayCount"),
 	1,
@@ -502,6 +510,10 @@ FMatrix CalcTranslatedWorldToShadowUVNormalMatrix(
 	return CalcTranslatedWorldToShadowUVMatrix(TranslatedWorldToShadowView, ViewToClip).GetTransposed().Inverse();
 }
 
+static float GetNormalBiasForShader()
+{
+	return CVarNormalBias.GetValueOnRenderThread() / 1000.0f;
+}
 
 template <typename ShaderType>
 static bool SetStatsArgsAndPermutation(FRDGBufferUAVRef StatsBufferUAV, typename ShaderType::FParameters *OutPassParameters, typename ShaderType::FPermutationDomain& OutPermutationVector)
@@ -571,6 +583,8 @@ void FVirtualShadowMapArray::Initialize(
 	
 	// Global SMRT settings so they can be shared between different passes that call into them
 	UniformParameters.ScreenRayLength = CVarScreenRayLength.GetValueOnRenderThread();
+	UniformParameters.NormalBias = GetNormalBiasForShader();
+
 	UniformParameters.SMRTAdaptiveRayCount = CVarSMRTAdaptiveRayCount.GetValueOnRenderThread();
 
 	UniformParameters.SMRTRayCountLocal = CVarSMRTRayCountLocal.GetValueOnRenderThread();
diff --git a/Engine/Source/Runtime/Renderer/Private/VirtualShadowMaps/VirtualShadowMapArray.h b/Engine/Source/Runtime/Renderer/Private/VirtualShadowMaps/VirtualShadowMapArray.h
--- a/Engine/Source/Runtime/Renderer/Private/VirtualShadowMaps/VirtualShadowMapArray.h
+++ b/Engine/Source/Runtime/Renderer/Private/VirtualShadowMaps/VirtualShadowMapArray.h
@@ -157,6 +157,7 @@ BEGIN_GLOBAL_SHADER_PARAMETER_STRUCT(FVirtualShadowMapUniformParameters, )
 
 	// SMRT parameters that are sometimes used globally
 	SHADER_PARAMETER(float, ScreenRayLength)
+	SHADER_PARAMETER(float, NormalBias)
 	SHADER_PARAMETER(uint32, SMRTAdaptiveRayCount)
 	SHADER_PARAMETER(int32, SMRTRayCountLocal)
 	SHADER_PARAMETER(int32, SMRTSamplesPerRayLocal)
diff --git a/Engine/Source/Runtime/Renderer/Private/VirtualShadowMaps/VirtualShadowMapProjection.cpp b/Engine/Source/Runtime/Renderer/Private/VirtualShadowMaps/VirtualShadowMapProjection.cpp
--- a/Engine/Source/Runtime/Renderer/Private/VirtualShadowMaps/VirtualShadowMapProjection.cpp
+++ b/Engine/Source/Runtime/Renderer/Private/VirtualShadowMaps/VirtualShadowMapProjection.cpp
@@ -30,15 +30,6 @@
 
 #define MAX_TEST_PERMUTATION 0
 
-
-static TAutoConsoleVariable<float> CVarNormalBias(
-	TEXT( "r.Shadow.Virtual.NormalBias" ),
-	0.5f,
-	TEXT( "Receiver offset along surface normal for shadow lookup. Scaled by distance to camera." )
-	TEXT( "Higher values avoid artifacts on surfaces nearly parallel to the light, but also visibility offset shadows and increase the chance of hitting unmapped pages." ),
-	ECVF_Scalability | ECVF_RenderThreadSafe
-);
-
 static TAutoConsoleVariable<int32> CVarForcePerLightShadowMaskClear(
 	TEXT( "r.Shadow.Virtual.ForcePerLightShadowMaskClear" ),
 	0,
@@ -133,7 +124,6 @@ class FVirtualShadowMapProjectionCS : public FGlobalShader
 		SHADER_PARAMETER_STRUCT_REF(FViewUniformShaderParameters, View)
 		SHADER_PARAMETER_STRUCT_REF(FBlueNoise, BlueNoise)
 		SHADER_PARAMETER(FIntVector4, ProjectionRect)
-		SHADER_PARAMETER(float, NormalBias)
 		SHADER_PARAMETER(float, SubsurfaceMinSourceRadius)
 		SHADER_PARAMETER(uint32, InputType)
 		SHADER_PARAMETER(uint32, bCullBackfacingPixels)
@@ -202,11 +192,6 @@ class FVirtualShadowMapProjectionCS : public FGlobalShader
 };
 IMPLEMENT_GLOBAL_SHADER(FVirtualShadowMapProjectionCS, "/Engine/Private/VirtualShadowMaps/VirtualShadowMapProjection.usf", "VirtualShadowMapProjection", SF_Compute);
 
-static float GetNormalBiasForShader()
-{
-	return CVarNormalBias.GetValueOnRenderThread() / 1000.0f;
-}
-
 static void RenderVirtualShadowMapProjectionCommon(
 	FRDGBuilder& GraphBuilder,
 	const FMinimalSceneTextures& SceneTextures,
@@ -232,7 +217,6 @@ static void RenderVirtualShadowMapProjectionCommon(
 	PassParameters->SceneTexturesStruct = SceneTextures.UniformBuffer;
 	PassParameters->View = View.ViewUniformBuffer;
 	PassParameters->ProjectionRect = FIntVector4(ProjectionRect.Min.X, ProjectionRect.Min.Y, ProjectionRect.Max.X, ProjectionRect.Max.Y);
-	PassParameters->NormalBias = GetNormalBiasForShader();
 	PassParameters->SubsurfaceMinSourceRadius = FMath::Sin(0.5f * FMath::DegreesToRadians(CVarSubsurfaceShadowMinSourceAngle.GetValueOnRenderThread()));
 	PassParameters->InputType = uint32(InputType);
 	PassParameters->bCullBackfacingPixels = VirtualShadowMapArray.ShouldCullBackfacingPixels() ? 1 : 0;
