Fix minor shader warnings

#rnx

diff --git a/Engine/Shaders/Private/HeterogeneousVolumes/HeterogeneousVolumesLiveShadingGlobalPipeline.usf b/Engine/Shaders/Private/HeterogeneousVolumes/HeterogeneousVolumesLiveShadingGlobalPipeline.usf
--- a/Engine/Shaders/Private/HeterogeneousVolumes/HeterogeneousVolumesLiveShadingGlobalPipeline.usf
+++ b/Engine/Shaders/Private/HeterogeneousVolumes/HeterogeneousVolumesLiveShadingGlobalPipeline.usf
@@ -135,7 +135,7 @@ void ConvertTexture3DFixedPointToFloatCS(
 	RWFloatTexture3D[DispatchThreadId] = Value / 65535.0;
 }
 
-Texture3D ExistenceTexture;
+Texture3D<float> ExistenceTexture;
 SamplerState TextureSampler;
 RWTexture3D<float> RWDilatedExistenceTexture;
 
@@ -151,8 +151,7 @@ void DilateExistenceMaskCS(
 
 	float3 UVW = DispatchThreadId / (float3)TextureResolution;
 	int MipLevel = 0;
-	float Sample = Texture3DSampleLevel(
-		ExistenceTexture,
+	float Sample = ExistenceTexture.SampleLevel(
 		TextureSampler,
 		UVW,
 		MipLevel).r;
diff --git a/Engine/Shaders/Private/LandscapeGrassWeight.usf b/Engine/Shaders/Private/LandscapeGrassWeight.usf
--- a/Engine/Shaders/Private/LandscapeGrassWeight.usf
+++ b/Engine/Shaders/Private/LandscapeGrassWeight.usf
@@ -89,132 +89,132 @@ void PSMain(
 	case 0:
 	{
 #if HAVE_GetGrassWeight0
-		OutColor[2] = GetGrassWeight0(MaterialParameters);
+		OutColor[2] = GetGrassWeight0(MaterialParameters).x;
 #endif // HAVE_GetGrassWeight0
 #if HAVE_GetGrassWeight1
-		OutColor[3] = GetGrassWeight1(MaterialParameters);
+		OutColor[3] = GetGrassWeight1(MaterialParameters).x;
 #endif // HAVE_GetGrassWeight1
 		break;
 	}
 	case 1:
 	{
 #if HAVE_GetGrassWeight2
-		OutColor[0] = GetGrassWeight2(MaterialParameters);
+		OutColor[0] = GetGrassWeight2(MaterialParameters).x;
 #endif // HAVE_GetGrassWeight2
 #if HAVE_GetGrassWeight3
-		OutColor[1] = GetGrassWeight3(MaterialParameters);
+		OutColor[1] = GetGrassWeight3(MaterialParameters).x;
 #endif // HAVE_GetGrassWeight3
 #if HAVE_GetGrassWeight4
-		OutColor[2] = GetGrassWeight4(MaterialParameters);
+		OutColor[2] = GetGrassWeight4(MaterialParameters).x;
 #endif // HAVE_GetGrassWeight4
 #if HAVE_GetGrassWeight5
-		OutColor[3] = GetGrassWeight5(MaterialParameters);
+		OutColor[3] = GetGrassWeight5(MaterialParameters).x;
 #endif // HAVE_GetGrassWeight5
 		break;
 	}
 	case 2:
 	{
 #if HAVE_GetGrassWeight6
-		OutColor[0] = GetGrassWeight6(MaterialParameters);
+		OutColor[0] = GetGrassWeight6(MaterialParameters).x;
 #endif // HAVE_GetGrassWeight6
 #if HAVE_GetGrassWeight7
-		OutColor[1] = GetGrassWeight7(MaterialParameters);
+		OutColor[1] = GetGrassWeight7(MaterialParameters).x;
 #endif // HAVE_GetGrassWeight7
 #if HAVE_GetGrassWeight8
-		OutColor[2] = GetGrassWeight8(MaterialParameters);
+		OutColor[2] = GetGrassWeight8(MaterialParameters).x;
 #endif // HAVE_GetGrassWeight8
 #if HAVE_GetGrassWeight9
-		OutColor[3] = GetGrassWeight9(MaterialParameters);
+		OutColor[3] = GetGrassWeight9(MaterialParameters).x;
 #endif // HAVE_GetGrassWeight9
 		break;
 	}
 	case 3:
 	{
 #if HAVE_GetGrassWeight10
-		OutColor[0] = GetGrassWeight10(MaterialParameters);
+		OutColor[0] = GetGrassWeight10(MaterialParameters).x;
 #endif // HAVE_GetGrassWeight10
 #if HAVE_GetGrassWeight11
-		OutColor[1] = GetGrassWeight11(MaterialParameters);
+		OutColor[1] = GetGrassWeight11(MaterialParameters).x;
 #endif // HAVE_GetGrassWeight11
 #if HAVE_GetGrassWeight12
-		OutColor[2] = GetGrassWeight12(MaterialParameters);
+		OutColor[2] = GetGrassWeight12(MaterialParameters).x;
 #endif // HAVE_GetGrassWeight12
 #if HAVE_GetGrassWeight13
-		OutColor[3] = GetGrassWeight13(MaterialParameters);
+		OutColor[3] = GetGrassWeight13(MaterialParameters).x;
 #endif // HAVE_GetGrassWeight13
 		break;
 	}
 	case 4:
 	{
 #if HAVE_GetGrassWeight14
-		OutColor[0] = GetGrassWeight14(MaterialParameters);
+		OutColor[0] = GetGrassWeight14(MaterialParameters).x;
 #endif // HAVE_GetGrassWeight14
 #if HAVE_GetGrassWeight15
-		OutColor[1] = GetGrassWeight15(MaterialParameters);
+		OutColor[1] = GetGrassWeight15(MaterialParameters).x;
 #endif // HAVE_GetGrassWeight15
 #if HAVE_GetGrassWeight16
-		OutColor[2] = GetGrassWeight16(MaterialParameters);
+		OutColor[2] = GetGrassWeight16(MaterialParameters).x;
 #endif // HAVE_GetGrassWeight16
 #if HAVE_GetGrassWeight17
-		OutColor[3] = GetGrassWeight17(MaterialParameters);
+		OutColor[3] = GetGrassWeight17(MaterialParameters).x;
 #endif // HAVE_GetGrassWeight17
 		break;
 	}
 	case 5:
 	{
 #if HAVE_GetGrassWeight18
-		OutColor[0] = GetGrassWeight18(MaterialParameters);
+		OutColor[0] = GetGrassWeight18(MaterialParameters).x;
 #endif // HAVE_GetGrassWeight18
 #if HAVE_GetGrassWeight19
-		OutColor[1] = GetGrassWeight19(MaterialParameters);
+		OutColor[1] = GetGrassWeight19(MaterialParameters).x;
 #endif // HAVE_GetGrassWeight19
 #if HAVE_GetGrassWeight20
-		OutColor[2] = GetGrassWeight20(MaterialParameters);
+		OutColor[2] = GetGrassWeight20(MaterialParameters).x;
 #endif // HAVE_GetGrassWeight20
 #if HAVE_GetGrassWeight21
-		OutColor[3] = GetGrassWeight21(MaterialParameters);
+		OutColor[3] = GetGrassWeight21(MaterialParameters).x;
 #endif // HAVE_GetGrassWeight21
 		break;
 	}
 	case 6:
 	{
 #if HAVE_GetGrassWeight22
-		OutColor[0] = GetGrassWeight22(MaterialParameters);
+		OutColor[0] = GetGrassWeight22(MaterialParameters).x;
 #endif // HAVE_GetGrassWeight22
 #if HAVE_GetGrassWeight23
-		OutColor[1] = GetGrassWeight23(MaterialParameters);
+		OutColor[1] = GetGrassWeight23(MaterialParameters).x;
 #endif // HAVE_GetGrassWeight23
 #if HAVE_GetGrassWeight24
-		OutColor[2] = GetGrassWeight24(MaterialParameters);
+		OutColor[2] = GetGrassWeight24(MaterialParameters).x;
 #endif // HAVE_GetGrassWeight24
 #if HAVE_GetGrassWeight25
-		OutColor[3] = GetGrassWeight25(MaterialParameters);
+		OutColor[3] = GetGrassWeight25(MaterialParameters).x;
 #endif // HAVE_GetGrassWeight25
 		break;
 	}
 	case 7:
 	{
 #if HAVE_GetGrassWeight26
-		OutColor[0] = GetGrassWeight26(MaterialParameters);
+		OutColor[0] = GetGrassWeight26(MaterialParameters).x;
 #endif // HAVE_GetGrassWeight26
 #if HAVE_GetGrassWeight27
-		OutColor[1] = GetGrassWeight27(MaterialParameters);
+		OutColor[1] = GetGrassWeight27(MaterialParameters).x;
 #endif // HAVE_GetGrassWeight27
 #if HAVE_GetGrassWeight28
-		OutColor[2] = GetGrassWeight28(MaterialParameters);
+		OutColor[2] = GetGrassWeight28(MaterialParameters).x;
 #endif // HAVE_GetGrassWeight28
 #if HAVE_GetGrassWeight29
-		OutColor[3] = GetGrassWeight29(MaterialParameters);
+		OutColor[3] = GetGrassWeight29(MaterialParameters).x;
 #endif // HAVE_GetGrassWeight29
 		break;
 	}
 	case 8:
 	{
 #if HAVE_GetGrassWeight30
-		OutColor[0] = GetGrassWeight30(MaterialParameters);
+		OutColor[0] = GetGrassWeight30(MaterialParameters).x;
 #endif // HAVE_GetGrassWeight30
 #if HAVE_GetGrassWeight31
-		OutColor[1] = GetGrassWeight31(MaterialParameters);
+		OutColor[1] = GetGrassWeight31(MaterialParameters).x;
 #endif // HAVE_GetGrassWeight31
 		break;
 
diff --git a/Engine/Shaders/Private/PostProcessVisualizeHDR.usf b/Engine/Shaders/Private/PostProcessVisualizeHDR.usf
--- a/Engine/Shaders/Private/PostProcessVisualizeHDR.usf
+++ b/Engine/Shaders/Private/PostProcessVisualizeHDR.usf
@@ -435,7 +435,7 @@ void MainPS(noperspective float4 UVAndScreenPos : TEXCOORD0, float4 SvPosition :
 			float4 SceneColor = Texture2DSample(HDRSceneColorTexture, HDRSceneColorSampler, UV);
 			SceneColor.xyz *= View.OneOverPreExposure;
 
-			Luminance = CalculateEyeAdaptationLuminance(SceneColor);
+			Luminance = CalculateEyeAdaptationLuminance(SceneColor.xyz);
 		}
 
 		OutColor.rgb = Luminance * LuminanceScale;
