Substrate - fixed shader compilation error while cooking with Substrate enabled (due to shader rewritting).

diff --git a/Engine/Shaders/Private/DefaultSSRTiles.usf b/Engine/Shaders/Private/DefaultSSRTiles.usf
--- a/Engine/Shaders/Private/DefaultSSRTiles.usf
+++ b/Engine/Shaders/Private/DefaultSSRTiles.usf
@@ -83,8 +83,9 @@ bool DoesPixelContainSSR(uint2 PixelPos)
 			{
 				const FHaziness Haziness = UnpackHaziness(SLAB_HAZINESS(BSDF));
 				const bool bHazeAsSimpleClearCoat = Haziness.bSimpleClearCoat;
-				bContainValidPixels |= bHazeAsSimpleClearCoat;
+				bContainValidPixels = bContainValidPixels || bHazeAsSimpleClearCoat;
 			}
+			break;
 		}
 		}
 	}
diff --git a/Engine/Shaders/Private/HairStrands/HairStrandsVisibilityDepthPS.usf b/Engine/Shaders/Private/HairStrands/HairStrandsVisibilityDepthPS.usf
--- a/Engine/Shaders/Private/HairStrands/HairStrandsVisibilityDepthPS.usf
+++ b/Engine/Shaders/Private/HairStrands/HairStrandsVisibilityDepthPS.usf
@@ -41,7 +41,7 @@ FHairSampleData ExtractSampleData(uint3 PixelCoord)
 		const FHairSample Sample = UnpackHairSample(HairSampleData[SampleIndex]);
 		Out.DeviceZ = max(Out.DeviceZ, Sample.Depth); // Inverse-Z
 		Out.LightChannelMask |= Sample.LightChannelMask;
-		Out.CastContactShadow |= HasHairFlags(Sample.Flags, HAIR_FLAGS_CAST_CONTACT_SHADOW);
+		Out.CastContactShadow = Out.CastContactShadow || HasHairFlags(Sample.Flags, HAIR_FLAGS_CAST_CONTACT_SHADOW);
 	}
 	return Out;
 }
diff --git a/Engine/Shaders/Private/HeterogeneousVolumes/HeterogeneousVolumesPreshadingPipeline.usf b/Engine/Shaders/Private/HeterogeneousVolumes/HeterogeneousVolumesPreshadingPipeline.usf
--- a/Engine/Shaders/Private/HeterogeneousVolumes/HeterogeneousVolumesPreshadingPipeline.usf
+++ b/Engine/Shaders/Private/HeterogeneousVolumes/HeterogeneousVolumesPreshadingPipeline.usf
@@ -534,7 +534,8 @@ void RenderSingleScatteringWithPreshadingCS(
 	float2 ViewPixelCoord = PixelCoord * DownsampleFactor;
 
 	float3 Radiance = 0.0;
-	float3 Transmittance = RWLightingTexture[PixelCoord].a;
+	float4 LightingTexture = RWLightingTexture[PixelCoord];
+	float3 Transmittance = LightingTexture.aaa;
 	float PrevOpacity = Luminance(1.0 - Transmittance);
 
 	// Create screen ray
@@ -645,11 +646,11 @@ void RenderSingleScatteringWithPreshadingCS(
 	}
 
 	// Radiance
-	RWLightingTexture[PixelCoord].rgb += Radiance * View.PreExposure;
+	LightingTexture.rgb += Radiance * View.PreExposure;
 	if (bApplyEmissionAndTransmittance)
 	{
 		float Opacity = Luminance(1.0 - Transmittance);
-		RWLightingTexture[PixelCoord].a = Luminance(Transmittance);
+		LightingTexture.a = Luminance(Transmittance);
 
 #if DIM_WRITE_VELOCITY
 		bool bWriteVelocity =
@@ -676,4 +677,6 @@ void RenderSingleScatteringWithPreshadingCS(
 		}
 #endif
 	}
+
+	RWLightingTexture[PixelCoord] = LightingTexture;
 }
