Substrate - fixed missing lumen/SSR reflections when there is no sky in the level.

#rnx

diff --git a/Engine/Shaders/Private/Substrate/SubstrateForwardLighting.ush b/Engine/Shaders/Private/Substrate/SubstrateForwardLighting.ush
--- a/Engine/Shaders/Private/Substrate/SubstrateForwardLighting.ush
+++ b/Engine/Shaders/Private/Substrate/SubstrateForwardLighting.ush
@@ -284,8 +284,7 @@ float3 SubstrateForwardLighting(
 
 				// 2. Specular evaluation (using regular normal)
 				#if FORWARD_PER_PIXEL_SHADING || TRANSLUCENCY_LIGHTING_SURFACE_LIGHTINGVOLUME
-				const bool bEnableSpecular = ReflectionStruct.SkyLightParameters.y > 0.0f;
-				if (bEnableSpecular)
+				// Every components of forward shading reflection (Sky, SSR, Lumen) are enabled/disabled within GetImageBasedReflectionSpecular.
 				{
 					// Evaluate specular parameters without bent normal
 					FSubstrateEnvLightResult SubstrateEnvLight = SubstrateEvaluateForEnvLight(SubstrateBSDFContext, true/*bEnableSpecular*/, Settings);
diff --git a/Engine/Shaders/Private/Substrate/SubstrateMobileForwardLighting.ush b/Engine/Shaders/Private/Substrate/SubstrateMobileForwardLighting.ush
--- a/Engine/Shaders/Private/Substrate/SubstrateMobileForwardLighting.ush
+++ b/Engine/Shaders/Private/Substrate/SubstrateMobileForwardLighting.ush
@@ -177,7 +177,7 @@ float3 SubstrateMobileForwardLighting(
 
 				// 2. Specular evaluation (using regular normal)
 				#if (FORWARD_PER_PIXEL_SHADING || TRANSLUCENCY_LIGHTING_SURFACE_LIGHTINGVOLUME)
-				if (UseBasePassSkylightSpecular > 0)
+				// Every components of forward shading reflection (Sky, SSR, Lumen) are enabled/disabled within GetImageBasedReflectionSpecular.
 				{
 					// Evaluate specular parameters without bent normal
 					FSubstrateEnvLightResult SubstrateEnvLight = SubstrateEvaluateForEnvLight(SubstrateBSDFContext, true /*bEnableSpecular*/, Settings);
