Helper function GetCulledReflectionCapturesGrid(...)

diff --git a/Engine/Shaders/Private/ForwardLightingCommon.ush b/Engine/Shaders/Private/ForwardLightingCommon.ush
--- a/Engine/Shaders/Private/ForwardLightingCommon.ush
+++ b/Engine/Shaders/Private/ForwardLightingCommon.ush
@@ -434,30 +434,7 @@ half3 GetImageBasedReflectionSpecular(FMaterialPixelParameters MaterialParameter
 	else
 #endif
 	{
-		uint NumLocalReflectionCaptures = 0;
-		uint DataStartIndex = 0;
-
-#if REFLECTION_COMPOSITE_USE_BLENDED_REFLECTION_CAPTURES
-		#if (INSTANCED_STEREO || MOBILE_MULTI_VIEW)
-			BRANCH
-			if (EyeIndex == 0)
-			{
-		#endif
-
-				uint NumCulledEntryIndex = (ForwardLightData.NumGridCells + GridIndex) * NUM_CULLED_LIGHTS_GRID_STRIDE;
-				NumLocalReflectionCaptures = min(ForwardLightData.NumCulledLightsGrid[NumCulledEntryIndex + 0], ForwardLightData.NumReflectionCaptures);
-				DataStartIndex = ForwardLightData.NumCulledLightsGrid[NumCulledEntryIndex + 1];
-
-		#if (INSTANCED_STEREO || MOBILE_MULTI_VIEW)
-			}
-			else
-			{
-				uint NumCulledEntryIndex = (ForwardLightDataISR.NumGridCells + GridIndex) * NUM_CULLED_LIGHTS_GRID_STRIDE;
-				NumLocalReflectionCaptures = min(ForwardLightDataISR.NumCulledLightsGrid[NumCulledEntryIndex + 0], ForwardLightDataISR.NumReflectionCaptures);
-				DataStartIndex = ForwardLightDataISR.NumCulledLightsGrid[NumCulledEntryIndex + 1];
-			}
-		#endif
-#endif
+		FCulledReflectionCapturesGridHeader CulledReflectionCapturesGridHeader = GetCulledReflectionCapturesGrid(GridIndex, EyeIndex);
 
 		const bool bCompositeSkylight = true;
 		SpecularIBL = CompositeReflectionCapturesAndSkylightTWS(
@@ -468,8 +445,8 @@ half3 GetImageBasedReflectionSpecular(FMaterialPixelParameters MaterialParameter
 			IndirectIrradiance, 
 			1.0f, 
 			0.0f, 
-			NumLocalReflectionCaptures, 
-			DataStartIndex, 
+			CulledReflectionCapturesGridHeader.NumReflectionCaptures, 
+			CulledReflectionCapturesGridHeader.DataStartIndex, 
 			SingleCaptureIndex,
 			bCompositeSkylight,
 			EyeIndex);
diff --git a/Engine/Shaders/Private/HairStrands/HairStrandsEnvironmentLighting.usf b/Engine/Shaders/Private/HairStrands/HairStrandsEnvironmentLighting.usf
--- a/Engine/Shaders/Private/HairStrands/HairStrandsEnvironmentLighting.usf
+++ b/Engine/Shaders/Private/HairStrands/HairStrandsEnvironmentLighting.usf
@@ -165,9 +165,7 @@ float3 GatherRadiance(float3 RayDirection, float Roughness, uint2 PixelPosition,
 	const uint EyeIndex = 0;
 	const uint GridIndex = ComputeLightGridCellIndex(PixelPosition, PixelSceneDepth, EyeIndex);
 
-	const uint NumCulledEntryIndex = (ForwardLightData.NumGridCells + GridIndex) * NUM_CULLED_LIGHTS_GRID_STRIDE;
-	const uint NumLocalReflectionCaptures = min(ForwardLightData.NumCulledLightsGrid[NumCulledEntryIndex + 0], ForwardLightData.NumReflectionCaptures);
-	const uint DataStartIndex = ForwardLightData.NumCulledLightsGrid[NumCulledEntryIndex + 1];
+	FCulledReflectionCapturesGridHeader CulledReflectionCapturesGridHeader = GetCulledReflectionCapturesGrid(GridIndex, EyeIndex);
 
 	const bool bCompositeSkylight = true;
 	return CompositeReflectionCapturesAndSkylightTWS(
@@ -178,8 +176,8 @@ float3 GatherRadiance(float3 RayDirection, float Roughness, uint2 PixelPosition,
 		IndirectIrradiance,
 		1.0f,
 		0.0f,
-		NumLocalReflectionCaptures,
-		DataStartIndex,
+		CulledReflectionCapturesGridHeader.NumReflectionCaptures,
+		CulledReflectionCapturesGridHeader.DataStartIndex,
 		SingleCaptureIndex,
 		bCompositeSkylight,
 		EyeIndex);
diff --git a/Engine/Shaders/Private/LightGridCommon.ush b/Engine/Shaders/Private/LightGridCommon.ush
--- a/Engine/Shaders/Private/LightGridCommon.ush
+++ b/Engine/Shaders/Private/LightGridCommon.ush
@@ -157,6 +157,39 @@ FCulledLightsGridData GetCulledLightsGrid(uint GridIndex, uint EyeIndex)
 	return Result;
 }
 
+struct FCulledReflectionCapturesGridHeader
+{
+	uint NumReflectionCaptures;
+	uint DataStartIndex;
+};
+
+FCulledReflectionCapturesGridHeader GetCulledReflectionCapturesGrid(uint GridIndex, uint EyeIndex)
+{
+	FCulledReflectionCapturesGridHeader Result;
+
+#if (INSTANCED_STEREO || MOBILE_MULTI_VIEW)
+	BRANCH
+	if (EyeIndex == 0)
+	{
+#endif
+
+		const uint NumCulledEntryIndex = (ForwardLightData.NumGridCells + GridIndex) * NUM_CULLED_LIGHTS_GRID_STRIDE;
+		Result.NumReflectionCaptures = min(ForwardLightData.NumCulledLightsGrid[NumCulledEntryIndex + 0], ForwardLightData.NumReflectionCaptures);
+		Result.DataStartIndex = ForwardLightData.NumCulledLightsGrid[NumCulledEntryIndex + 1];
+
+#if (INSTANCED_STEREO || MOBILE_MULTI_VIEW)
+	}
+	else
+	{
+		const uint NumCulledEntryIndex = (ForwardLightDataISR.NumGridCells + GridIndex) * NUM_CULLED_LIGHTS_GRID_STRIDE;
+		Result.NumReflectionCaptures = min(ForwardLightDataISR.NumCulledLightsGrid[NumCulledEntryIndex + 0], ForwardLightDataISR.NumReflectionCaptures);
+		Result.DataStartIndex = ForwardLightDataISR.NumCulledLightsGrid[NumCulledEntryIndex + 1];
+	}
+#endif
+
+	return Result;
+}
+
 FDirectionalLightData GetDirectionalLightData(uint EyeIndex)
 {
 	FDirectionalLightData Result;
diff --git a/Engine/Shaders/Private/MobileLightingCommon.ush b/Engine/Shaders/Private/MobileLightingCommon.ush
--- a/Engine/Shaders/Private/MobileLightingCommon.ush
+++ b/Engine/Shaders/Private/MobileLightingCommon.ush
@@ -356,9 +356,7 @@ half3 GetImageBasedReflectionLighting_Mobile(
 	half3 SpecularIBL = (half3)0.0f;
 
 #if ENABLE_CLUSTERED_REFLECTION
-	uint NumCulledEntryIndex = (ForwardLightData.NumGridCells + GridIndex) * NUM_CULLED_LIGHTS_GRID_STRIDE;
-	uint NumLocalReflectionCaptures = min(ForwardLightData.NumCulledLightsGrid[NumCulledEntryIndex + 0], ForwardLightData.NumReflectionCaptures);
-	uint DataStartIndex = ForwardLightData.NumCulledLightsGrid[NumCulledEntryIndex + 1];
+	FCulledReflectionCapturesGridHeader CulledReflectionCapturesGridHeader = GetCulledReflectionCapturesGrid(GridIndex, 0);
 
 	SpecularIBL = CompositeReflectionCapturesAndSkylightTWS(
 		CompositeAlpha,
@@ -368,8 +366,8 @@ half3 GetImageBasedReflectionLighting_Mobile(
 		IndirectIrradiance,
 		1.0f,
 		0.0f,
-		NumLocalReflectionCaptures,
-		DataStartIndex,
+		CulledReflectionCapturesGridHeader.NumReflectionCaptures,
+		CulledReflectionCapturesGridHeader.DataStartIndex,
 		0,
 		true);
 
diff --git a/Engine/Shaders/Private/ReflectionEnvironmentPixelShader.usf b/Engine/Shaders/Private/ReflectionEnvironmentPixelShader.usf
--- a/Engine/Shaders/Private/ReflectionEnvironmentPixelShader.usf
+++ b/Engine/Shaders/Private/ReflectionEnvironmentPixelShader.usf
@@ -204,9 +204,10 @@ float3 ReflectionEnvironment(FGBufferData GBuffer, float AmbientOcclusion, float
 	float2 LocalPosition = SvPosition.xy - View.ViewRectMin.xy;
 
 	uint GridIndex = ComputeLightGridCellIndex(uint2(LocalPosition.x, LocalPosition.y), GBuffer.Depth);
-	uint NumCulledEntryIndex = (ForwardLightData.NumGridCells + GridIndex) * NUM_CULLED_LIGHTS_GRID_STRIDE;
-	uint NumCulledReflectionCaptures = min(ForwardLightData.NumCulledLightsGrid[NumCulledEntryIndex + 0], ForwardLightData.NumReflectionCaptures);
-	uint DataStartIndex = ForwardLightData.NumCulledLightsGrid[NumCulledEntryIndex + 1];
+
+	FCulledReflectionCapturesGridHeader CulledReflectionCapturesGridHeader = GetCulledReflectionCapturesGrid(GridIndex, 0);
+	uint NumCulledReflectionCaptures = CulledReflectionCapturesGridHeader.NumReflectionCaptures;
+	uint DataStartIndex = CulledReflectionCapturesGridHeader.DataStartIndex;
 #else
 	uint DataStartIndex = 0;
 	uint NumCulledReflectionCaptures = 0;
@@ -437,9 +438,10 @@ void ReflectionEnvironmentSkyLighting(
 		float2 LocalPosition = SvPosition.xy - View.ViewRectMin.xy;
 	
 		uint GridIndex = ComputeLightGridCellIndex(uint2(LocalPosition.x, LocalPosition.y), SceneDepth);
-		uint NumCulledEntryIndex = (ForwardLightData.NumGridCells + GridIndex) * NUM_CULLED_LIGHTS_GRID_STRIDE;
-		uint NumCulledReflectionCaptures = min(ForwardLightData.NumCulledLightsGrid[NumCulledEntryIndex + 0], ForwardLightData.NumReflectionCaptures);
-		uint CaptureDataStartIndex = ForwardLightData.NumCulledLightsGrid[NumCulledEntryIndex + 1];
+
+		FCulledReflectionCapturesGridHeader CulledReflectionCapturesGridHeader = GetCulledReflectionCapturesGrid(GridIndex, 0);
+		uint NumCulledReflectionCaptures = CulledReflectionCapturesGridHeader.NumReflectionCaptures;
+		uint DataStartIndex = CulledReflectionCapturesGridHeader.DataStartIndex;
 	#else
 		uint CaptureDataStartIndex = 0;
 		uint NumCulledReflectionCaptures = 0;
diff --git a/Engine/Shaders/Private/SingleLayerWaterComposite.usf b/Engine/Shaders/Private/SingleLayerWaterComposite.usf
--- a/Engine/Shaders/Private/SingleLayerWaterComposite.usf
+++ b/Engine/Shaders/Private/SingleLayerWaterComposite.usf
@@ -115,9 +115,10 @@ SingleLayerWaterCompositeOutput SingleLayerWaterComposite(float2 BufferUV, float
 				float2 LocalPosition = SvPosition.xy - View.ViewRectMin.xy;
 
 				uint GridIndex = ComputeLightGridCellIndex(uint2(LocalPosition.x, LocalPosition.y), WaterDepth);
-				uint NumCulledEntryIndex = (ForwardLightData.NumGridCells + GridIndex) * NUM_CULLED_LIGHTS_GRID_STRIDE;
-				uint NumCulledReflectionCaptures = min(ForwardLightData.NumCulledLightsGrid[NumCulledEntryIndex + 0], ForwardLightData.NumReflectionCaptures);
-				uint CaptureDataStartIndex = ForwardLightData.NumCulledLightsGrid[NumCulledEntryIndex + 1];
+
+				FCulledReflectionCapturesGridHeader CulledReflectionCapturesGridHeader = GetCulledReflectionCapturesGrid(GridIndex, 0);
+				uint NumCulledReflectionCaptures = CulledReflectionCapturesGridHeader.NumReflectionCaptures;
+				uint DataStartIndex = CulledReflectionCapturesGridHeader.DataStartIndex;
 			#else
 				uint NumCulledReflectionCaptures = 0;
 				uint CaptureDataStartIndex = 0;
@@ -190,9 +191,10 @@ SingleLayerWaterCompositeOutput SingleLayerWaterComposite(float2 BufferUV, float
 				float2 LocalPosition = SvPosition.xy - View.ViewRectMin.xy;
 
 				uint GridIndex = ComputeLightGridCellIndex(uint2(LocalPosition.x, LocalPosition.y), WaterDepth);
-				uint NumCulledEntryIndex = (ForwardLightData.NumGridCells + GridIndex) * NUM_CULLED_LIGHTS_GRID_STRIDE;
-				uint NumCulledReflectionCaptures = min(ForwardLightData.NumCulledLightsGrid[NumCulledEntryIndex + 0], ForwardLightData.NumReflectionCaptures);
-				uint CaptureDataStartIndex = ForwardLightData.NumCulledLightsGrid[NumCulledEntryIndex + 1];
+
+				FCulledReflectionCapturesGridHeader CulledReflectionCapturesGridHeader = GetCulledReflectionCapturesGrid(GridIndex, 0);
+				uint NumCulledReflectionCaptures = CulledReflectionCapturesGridHeader.NumReflectionCaptures;
+				uint CaptureDataStartIndex = CulledReflectionCapturesGridHeader.DataStartIndex;
 			#else
 				uint NumCulledReflectionCaptures = 0;
 				uint CaptureDataStartIndex = 0;
