Fixed Substrate shader compilation error.

#rnx

diff --git a/Engine/Shaders/Private/ReflectionEnvironmentPixelShader.usf b/Engine/Shaders/Private/ReflectionEnvironmentPixelShader.usf
--- a/Engine/Shaders/Private/ReflectionEnvironmentPixelShader.usf
+++ b/Engine/Shaders/Private/ReflectionEnvironmentPixelShader.usf
@@ -207,9 +207,9 @@ float3 ReflectionEnvironment(FGBufferData GBuffer, float AmbientOcclusion, float
 
 	FCulledReflectionCapturesGridHeader CulledReflectionCapturesGridHeader = GetCulledReflectionCapturesGrid(GridIndex, 0);
 	uint NumCulledReflectionCaptures = CulledReflectionCapturesGridHeader.NumReflectionCaptures;
-	uint DataStartIndex = CulledReflectionCapturesGridHeader.DataStartIndex;
+	uint CaptureDataStartIndex = CulledReflectionCapturesGridHeader.DataStartIndex;
 #else
-	uint DataStartIndex = 0;
+	uint CaptureDataStartIndex = 0;
 	uint NumCulledReflectionCaptures = 0;
 #endif
 
@@ -218,7 +218,7 @@ float3 ReflectionEnvironment(FGBufferData GBuffer, float AmbientOcclusion, float
 #endif
 
 	//Top of regular reflection or bottom layer of clear coat.
-	Color.rgb += View.PreExposure * GatherRadiance(Color.a, TranslatedWorldPosition, R, GBuffer.Roughness, BentNormal, IndirectIrradiance, GBuffer.ShadingModelID, NumCulledReflectionCaptures, DataStartIndex);
+	Color.rgb += View.PreExposure * GatherRadiance(Color.a, TranslatedWorldPosition, R, GBuffer.Roughness, BentNormal, IndirectIrradiance, GBuffer.ShadingModelID, NumCulledReflectionCaptures, CaptureDataStartIndex);
 
 	BRANCH
 	if( GBuffer.ShadingModelID == SHADINGMODELID_CLEAR_COAT)
@@ -264,7 +264,7 @@ float3 ReflectionEnvironment(FGBufferData GBuffer, float AmbientOcclusion, float
 		Color.a *= SpecularOcclusion;
 
 		float3 TopLayerR = 2 * dot( V, N ) * N - V;
-		Color.rgb += View.PreExposure * GatherRadiance(Color.a, TranslatedWorldPosition, TopLayerR, ClearCoatRoughness, BentNormal, IndirectIrradiance, GBuffer.ShadingModelID, NumCulledReflectionCaptures, DataStartIndex);
+		Color.rgb += View.PreExposure * GatherRadiance(Color.a, TranslatedWorldPosition, TopLayerR, ClearCoatRoughness, BentNormal, IndirectIrradiance, GBuffer.ShadingModelID, NumCulledReflectionCaptures, CaptureDataStartIndex);
 
 	}
 	else
@@ -441,7 +441,7 @@ void ReflectionEnvironmentSkyLighting(
 
 		FCulledReflectionCapturesGridHeader CulledReflectionCapturesGridHeader = GetCulledReflectionCapturesGrid(GridIndex, 0);
 		uint NumCulledReflectionCaptures = CulledReflectionCapturesGridHeader.NumReflectionCaptures;
-		uint DataStartIndex = CulledReflectionCapturesGridHeader.DataStartIndex;
+		uint CaptureDataStartIndex = CulledReflectionCapturesGridHeader.DataStartIndex;
 	#else
 		uint CaptureDataStartIndex = 0;
 		uint NumCulledReflectionCaptures = 0;
diff --git a/Engine/Shaders/Private/SingleLayerWaterComposite.usf b/Engine/Shaders/Private/SingleLayerWaterComposite.usf
--- a/Engine/Shaders/Private/SingleLayerWaterComposite.usf
+++ b/Engine/Shaders/Private/SingleLayerWaterComposite.usf
@@ -118,7 +118,7 @@ SingleLayerWaterCompositeOutput SingleLayerWaterComposite(float2 BufferUV, float
 
 				FCulledReflectionCapturesGridHeader CulledReflectionCapturesGridHeader = GetCulledReflectionCapturesGrid(GridIndex, 0);
 				uint NumCulledReflectionCaptures = CulledReflectionCapturesGridHeader.NumReflectionCaptures;
-				uint DataStartIndex = CulledReflectionCapturesGridHeader.DataStartIndex;
+				uint CaptureDataStartIndex = CulledReflectionCapturesGridHeader.DataStartIndex;
 			#else
 				uint NumCulledReflectionCaptures = 0;
 				uint CaptureDataStartIndex = 0;
