Substrate - bytes per pixel vis: do not show colors where there is 0 bytes read such as sky pixel.

#rnx

diff --git a/Engine/Shaders/Private/Substrate/SubstrateVisualize.usf b/Engine/Shaders/Private/Substrate/SubstrateVisualize.usf
--- a/Engine/Shaders/Private/Substrate/SubstrateVisualize.usf
+++ b/Engine/Shaders/Private/Substrate/SubstrateVisualize.usf
@@ -756,8 +756,11 @@ void VisualizeMaterialPS(
 			FSubstrateBSDF BSDF = UnpackSubstrateBSDF(Substrate.MaterialTextureArray, SubstrateAddressing, Header);
 			MaterialBytePerPixel = SubstrateAddressing.ReadBytes;
 		}
-			
-		OutColor.xyz = BytePerPixelToColor(MaterialBytePerPixel);
+		
+		if (MaterialBytePerPixel > 0)
+		{
+			OutColor.xyz = BytePerPixelToColor(MaterialBytePerPixel);
+		}
 	
 		// Print Legend
 		if (all(uint2(SVPos.xy) == uint2(0,0)))
