Fixing artifacts with Camera AVSM and intersecting geometry.

diff --git a/Engine/Shaders/Private/HeterogeneousVolumes/HeterogeneousVolumesAdaptiveVolumetricShadowMapSampler.ush b/Engine/Shaders/Private/HeterogeneousVolumes/HeterogeneousVolumesAdaptiveVolumetricShadowMapSampler.ush
--- a/Engine/Shaders/Private/HeterogeneousVolumes/HeterogeneousVolumesAdaptiveVolumetricShadowMapSampler.ush
+++ b/Engine/Shaders/Private/HeterogeneousVolumes/HeterogeneousVolumesAdaptiveVolumetricShadowMapSampler.ush
@@ -352,6 +352,12 @@ float AVSM_Sampler4_Eval(inout FAVSM_Sampler4 Sampler4, float X)
 	return Tau;
 }
 
+float AVSM_Sampler4_Eval_NoInterpolation(inout FAVSM_Sampler4 Sampler4, float X)
+{
+	float Tau = AVSM_Sampler4_Eval_Internal(Sampler4, 0, X);
+	return Tau;
+}
+
 FAVSMSampleData AVSM_Sampler4_EvalInverse_Internal(inout FAVSM_Sampler4 Sampler4, int Index, float Tau)
 {
 	FAVSMSampleData SampleData[] = {
diff --git a/Engine/Shaders/Private/HeterogeneousVolumes/HeterogeneousVolumesAdaptiveVolumetricShadowMapSampling.ush b/Engine/Shaders/Private/HeterogeneousVolumes/HeterogeneousVolumesAdaptiveVolumetricShadowMapSampling.ush
--- a/Engine/Shaders/Private/HeterogeneousVolumes/HeterogeneousVolumesAdaptiveVolumetricShadowMapSampling.ush
+++ b/Engine/Shaders/Private/HeterogeneousVolumes/HeterogeneousVolumesAdaptiveVolumetricShadowMapSampling.ush
@@ -38,7 +38,7 @@ bool AVSM_PixelAndDepth(
 	inout float Depth
 )
 {
-	Depth = -65535.0;
+	Depth = 0.0;
 	Pixel = -1;
 	float4 ShadowPositionAsFloat4 = mul(float4(TranslatedWorldPosition, 1), ShadowMap.TranslatedWorldToShadow[Face]);
 	if (ShadowPositionAsFloat4.w > 0)
@@ -47,7 +47,6 @@ bool AVSM_PixelAndDepth(
 		if (all(ShadowPosition.xyz > 0) || all(ShadowPosition.xy < 1))
 		{
 			Pixel = clamp(ShadowPosition.xy * ShadowMap.Resolution - 0.5, 0, ShadowMap.Resolution - 1);
-			Depth = length(TranslatedWorldPosition - ShadowMap.TranslatedWorldOrigin);
 			if (ShadowMap.bIsDirectionalLight)
 			{
 				Depth = dot(ShadowMap.TranslatedWorldPlane.xyz, TranslatedWorldPosition) + ShadowMap.TranslatedWorldPlane.w;
@@ -361,7 +360,14 @@ float AVSM_Sample(
 	float WorldHitT
 )
 {
-	return AVSM_Sample_TwoLevel_Vectorized(ShadowMap, Face, Pixel, WorldHitT);
+	return
+		//AVSM_Sample_Reference(
+		AVSM_Sample_TwoLevel_Vectorized(
+			ShadowMap,
+			Face,
+			Pixel,
+			WorldHitT
+		);
 }
 
 float AVSM_SampleTransmittance(
diff --git a/Engine/Shaders/Private/HeterogeneousVolumes/HeterogeneousVolumesLiveShadingPipeline.usf b/Engine/Shaders/Private/HeterogeneousVolumes/HeterogeneousVolumesLiveShadingPipeline.usf
--- a/Engine/Shaders/Private/HeterogeneousVolumes/HeterogeneousVolumesLiveShadingPipeline.usf
+++ b/Engine/Shaders/Private/HeterogeneousVolumes/HeterogeneousVolumesLiveShadingPipeline.usf
@@ -262,7 +262,7 @@ float EvalAmbientOcclusion(float3 WorldPosition)
 		AmbientOcclusionTexture,
 		View.SharedTrilinearClampedSampler,
 		UVW,
-		MipLevel);
+		MipLevel).r;
 }
 
 #include "HeterogeneousVolumesRayMarchingUtils.ush"
@@ -462,8 +462,8 @@ void RenderSingleScatteringWithLiveShading(
 		return;
 	}
 	uint LinearIndex = DispatchThreadId.y * GroupCount.x * THREADGROUP_SIZE_2D + DispatchThreadId.x;
-	uint2 PixelCoord = DispatchThreadId.xy + View.ViewRectMin.xy / DownsampleFactor;
-	uint2 ViewPixelCoord = PixelCoord * DownsampleFactor + View.ViewRectMin.xy;
+	uint2 PixelCoord = DispatchThreadId.xy;
+	uint2 ViewPixelCoord = PixelCoord * DownsampleFactor + View.ViewRectMin.xy / DownsampleFactor;
 
 	float3 Radiance = 0.0;
 	float3 PrevTransmittance = RWLightingTexture[PixelCoord].a;
@@ -543,13 +543,12 @@ void RenderSingleScatteringWithLiveShading(
 
 		int Face = 0;
 		float2 Pixel = 0;
-		float Depth = 0;
-		if (!AVSM_PixelAndDepth(GlobalTransmittanceMap, Face, TranslatedWorldPosition, Pixel, Depth))
+		float DepthOffset = 0;
+		if (!AVSM_PixelAndDepth(GlobalTransmittanceMap, Face, TranslatedWorldPosition, Pixel, DepthOffset))
 		{
 			return;
 		}
 
-		float DepthOffset = WorldRayTMin - Depth;
 		FAVSM_Sampler4 TransmittanceSampler = AVSM_Sampler4_Create(GlobalTransmittanceMap, Face, Pixel, DepthOffset);
 
 #if ADAPTIVE_MARCH
@@ -582,7 +581,8 @@ void RenderSingleScatteringWithLiveShading(
 		);
 #endif
 		// Replace transmittance at WorldRayTMax with global transmittance at MaxTraceDistance
-		Transmittance = AVSM_Sampler4_Last(TransmittanceSampler).Tau;
+		//Transmittance = AVSM_Sampler4_Eval_NoInterpolation(TransmittanceSampler, WorldRayLength);
+		Transmittance = AVSM_Sampler4_Eval(TransmittanceSampler, WorldRayLength);
 
 #else
 		RayMarchSingleScattering(
diff --git a/Engine/Shaders/Private/HeterogeneousVolumes/HeterogeneousVolumesLiveShadingShadows.usf b/Engine/Shaders/Private/HeterogeneousVolumes/HeterogeneousVolumesLiveShadingShadows.usf
--- a/Engine/Shaders/Private/HeterogeneousVolumes/HeterogeneousVolumesLiveShadingShadows.usf
+++ b/Engine/Shaders/Private/HeterogeneousVolumes/HeterogeneousVolumesLiveShadingShadows.usf
@@ -310,23 +310,24 @@ void RenderVolumetricShadowMapForLightWithLiveShadingCS(
 	uint Face = DispatchThreadId.z;
 	uint2 PixelCoord = DispatchThreadId.xy;
 	uint LinearIndex = GetLinearIndex(Face, PixelCoord);
-	//int PixelIndex = LinearIndex * MaxSampleCount;
 	int PixelIndex = LinearIndex * MaxSampleCount;
 
 	// Generate a light ray
 	float2 ShadowUV = float2(PixelCoord + 0.5f) / float2(ShadowResolution);
+	//float2 ShadowUV = float2(PixelCoord + View.TemporalAAJitter.xy) / float2(ShadowResolution);
 
 	float WorldRayLength = GetMaxTraceDistance();
 #if USE_CAMERA_SCENE_DEPTH
-	float DeviceZ = SampleDeviceZFromSceneTextures(ShadowUV);
+	float2 ViewPixelCoord = (PixelCoord + View.TemporalAAJitter.xy) * CameraDownsampleFactor + View.ViewRectMin.xy / CameraDownsampleFactor;
+	float DeviceZ = SceneDepthTexture.Load(int3(ViewPixelCoord, 0)).r;
+
 	#if HAS_INVERTED_Z_BUFFER
 		DeviceZ = max(0.000000000001, DeviceZ);
 	#endif // HAS_INVERTED_Z_BUFFER
 
 	// Project ray through view transform to trim ray-length by opaque depth
-	uint2 ViewPixelCoord = PixelCoord * CameraDownsampleFactor + View.ViewRectMin.xy;
-	FDFVector3 WorldViewRayOrigin = DFFastSubtract(GetTranslatedWorldCameraPosFromView(ViewPixelCoord + View.TemporalAAJitter.xy), PrimaryView.PreViewTranslation);
-	FDFVector3 WorldViewRayEnd = SvPositionToWorld(float4(ViewPixelCoord + View.TemporalAAJitter.xy, DeviceZ, 1));
+	FDFVector3 WorldViewRayOrigin = DFFastSubtract(GetTranslatedWorldCameraPosFromView(ViewPixelCoord), PrimaryView.PreViewTranslation);
+	FDFVector3 WorldViewRayEnd = SvPositionToWorld(float4(ViewPixelCoord, DeviceZ, 1));
 	float3 WorldViewRayDirection = DFFastSubtractDemote(WorldViewRayEnd, WorldViewRayOrigin);
 	float WorldViewRayLength = length(WorldViewRayDirection);
 
@@ -396,21 +397,30 @@ void RenderVolumetricShadowMapForLightWithLiveShadingCS(
 		int PrevElementIndex = AVSM_NULL_PTR;
 		int ElementIndex = 0;
 
+		float StepSize = RayMarchingContext.StepSize;
+		float JitterSize = StepSize * RayMarchingContext.Jitter;
+
 		float LocalHitT = HitT.x;
 		float WorldHitT = LocalHitT * RayMarchingContext.LocalToWorldScale;
 		for (uint StepIndex = 0; StepIndex < StepCount; ++StepIndex)
 		{
-			LocalHitT = min(RayMarchingContext.LocalRayTMin + RayMarchingContext.StepSize * (RayMarchingContext.Jitter + StepIndex), RayMarchingContext.LocalRayTMax);
+			LocalHitT = min(RayMarchingContext.LocalRayTMin + StepSize * (RayMarchingContext.Jitter + StepIndex), RayMarchingContext.LocalRayTMax);
+			if (StepIndex == StepCount - 1)
+			{
+				float UnjitteredLocalHitT = LocalHitT - JitterSize;
+				StepSize = RayMarchingContext.LocalRayTMax - UnjitteredLocalHitT;
+			}
+
 			WorldHitT = LocalHitT * RayMarchingContext.LocalToWorldScale;
 			float3 LocalPosition = RayMarchingContext.LocalRayOrigin + RayMarchingContext.LocalRayDirection * LocalHitT;
 			float3 WorldPosition = RayMarchingContext.WorldRayOrigin + RayMarchingContext.WorldRayDirection * WorldHitT;
 
 			FVolumeSampleContext SampleContext = CreateVolumeSampleContext(LocalPosition, WorldPosition, RayMarchingContext.MipLevel);
 			float3 Extinction = SampleExtinction(SampleContext);
-			Transmittance *= exp(-Extinction * RayMarchingContext.StepSize);
+			Transmittance *= exp(-Extinction * StepSize);
 
 			bool bFirstEntry = (ElementIndex == 0);
-			if (bFirstEntry && any(Extinction > 0.0))
+			if (bFirstEntry)
 			{
 				FAVSMSampleData SampleData = AVSM_CreateSampleData(WorldHitT, Luminance(Transmittance));
 				AddSample(SampleData, PixelIndex, PrevElementIndex, ElementIndex);
@@ -419,10 +429,10 @@ void RenderVolumetricShadowMapForLightWithLiveShadingCS(
 				ExpectedExtinction = Extinction;
 				ExpectedTransmittance = Transmittance;
 			}
-			else if (ElementIndex > 0 && any(Extinction > 0.0))
+			else
 			{
 				// If extrapolated transmittance is beyond relative error tolerance
-				ExpectedTransmittance *= exp(-ExpectedExtinction * RayMarchingContext.StepSize);
+				ExpectedTransmittance *= exp(-ExpectedExtinction * StepSize);
 				float AbsoluteError = length(Transmittance - ExpectedTransmittance);
 				float RelativeError = length((Transmittance - ExpectedTransmittance) / ExpectedTransmittance);
 				if ((AbsoluteError > AbsoluteErrorThreshold) && (RelativeError > RelativeErrorThreshold))
diff --git a/Engine/Shaders/Private/HeterogeneousVolumes/HeterogeneousVolumesRayMarchingUtils.ush b/Engine/Shaders/Private/HeterogeneousVolumes/HeterogeneousVolumesRayMarchingUtils.ush
--- a/Engine/Shaders/Private/HeterogeneousVolumes/HeterogeneousVolumesRayMarchingUtils.ush
+++ b/Engine/Shaders/Private/HeterogeneousVolumes/HeterogeneousVolumesRayMarchingUtils.ush
@@ -433,6 +433,7 @@ void RayMarchSingleScattering(
 		}
 
 #if USE_CAMERA_AVSM
+		//Transmittance = AVSM_Sampler4_Eval_NoInterpolation(TransmittanceSampler, WorldHitT);
 		Transmittance = AVSM_Sampler4_Eval(TransmittanceSampler, WorldHitT);
 #else
 		// Accumulate transmittance for the next evaluation
diff --git a/Engine/Source/Runtime/Renderer/Private/HeterogeneousVolumes/HeterogeneousVolumesLiveShadingPipeline.cpp b/Engine/Source/Runtime/Renderer/Private/HeterogeneousVolumes/HeterogeneousVolumesLiveShadingPipeline.cpp
--- a/Engine/Source/Runtime/Renderer/Private/HeterogeneousVolumes/HeterogeneousVolumesLiveShadingPipeline.cpp
+++ b/Engine/Source/Runtime/Renderer/Private/HeterogeneousVolumes/HeterogeneousVolumesLiveShadingPipeline.cpp
@@ -46,6 +46,13 @@ static TAutoConsoleVariable<int32> CVarHeterogeneousVolumesUseExistenceMask(
 	ECVF_RenderThreadSafe
 );
 
+static TAutoConsoleVariable<int32> CVarHeterogeneousVolumesShadowsUseCameraSceneDepth(
+	TEXT("r.HeterogeneousVolumes.Shadows.UseCameraSceneDepth"),
+	0,
+	TEXT("Culls Camera AVSM by SceneDepth (Default = 0)"),
+	ECVF_RenderThreadSafe
+);
+
 #if 0
 static TAutoConsoleVariable<int32> CVarHeterogeneousVolumesBilinearInterpolation(
 	TEXT("r.HeterogeneousVolumes.BilinearInterpolation"),
@@ -94,6 +101,11 @@ namespace HeterogeneousVolumes
 		return CVarHeterogeneousVolumesAdaptiveMarching.GetValueOnRenderThread() != 0;
 	}
 #endif
+
+	bool ShadowsUseCameraSceneDepth()
+	{
+		return CVarHeterogeneousVolumesShadowsUseCameraSceneDepth.GetValueOnRenderThread() != 0;
+	}
 }
 
 //-OPT: Remove duplicate bindings
@@ -1975,7 +1987,7 @@ bool RenderVolumetricShadowMapForLightForHeterogeneousVolumeWithLiveShading(
 
 	FRenderVolumetricShadowMapForLightWithLiveShadingCS::FPermutationDomain PermutationVector;
 	PermutationVector.Set<FRenderVolumetricShadowMapForLightWithLiveShadingCS::FUseAVSMCompression>(HeterogeneousVolumes::UseAVSMCompression());
-	PermutationVector.Set<FRenderVolumetricShadowMapForLightWithLiveShadingCS::FUseCameraSceneDepth>(bUseCameraSceneDepth);
+	PermutationVector.Set<FRenderVolumetricShadowMapForLightWithLiveShadingCS::FUseCameraSceneDepth>(bUseCameraSceneDepth && HeterogeneousVolumes::ShadowsUseCameraSceneDepth());
 	TShaderRef<FRenderVolumetricShadowMapForLightWithLiveShadingCS> ComputeShader = Material.GetShader<FRenderVolumetricShadowMapForLightWithLiveShadingCS>(&FLocalVertexFactory::StaticType, PermutationVector, false);
 	if (!ComputeShader.IsNull())
 	{
