Substrate - Fix for skypass and render target mismatch. We now simple skip writting to the gbuffer and only write to the SceneColor. Should help fix MAC.

#rnx

diff --git a/Engine/Shaders/Private/PixelShaderOutputCommon.ush b/Engine/Shaders/Private/PixelShaderOutputCommon.ush
--- a/Engine/Shaders/Private/PixelShaderOutputCommon.ush
+++ b/Engine/Shaders/Private/PixelShaderOutputCommon.ush
@@ -58,7 +58,9 @@ void MainPS
 		, out float4 OutTarget3 : SV_Target3
 #endif
 
-#if SUBSTRATE_OPAQUE_DEFERRED
+#define SUBSTRATE_OUTPUT_GBUFFER  (SUBSTRATE_OPAQUE_DEFERRED && !MATERIAL_IS_SKY)
+
+#if SUBSTRATE_OUTPUT_GBUFFER
 #if SUBSTRATE_BASE_PASS_MRT_OUTPUT_COUNT != 3
 #error Substrate SUBSTRATE_BASE_PASS_MRT_OUTPUT_COUNT has been updated but not the uint MRTs
 #endif
@@ -86,7 +88,7 @@ void MainPS
 		, out uint SubstrateOutput3 : SV_Target3
 		, out SUBSTRATE_TOP_LAYER_TYPE SubstrateOutput4 : SV_Target4
 #endif
-#else // SUBSTRATE_OPAQUE_DEFERRED
+#else // SUBSTRATE_OUTPUT_GBUFFER
 
 #if PIXELSHADEROUTPUT_MRT4
 		, out float4 OutTarget4 : SV_Target4
@@ -104,7 +106,7 @@ void MainPS
 		, out float4 OutTarget7 : SV_Target7
 #endif
 
-#endif // SUBSTRATE_OPAQUE_DEFERRED
+#endif // SUBSTRATE_OUTPUT_GBUFFER
 
 		// todo: if we are not inside a volumetric decal we could use OPTIONAL_OutDepthConservative for better performance, this would require a shader permutation
 		OPTIONAL_OutDepthConservative
@@ -176,7 +178,7 @@ uint ArrayIndex = 0;
 	OutTarget3 = PixelShaderOut.MRT[3];
 #endif
 
-#if SUBSTRATE_OPAQUE_DEFERRED
+#if SUBSTRATE_OUTPUT_GBUFFER
 
 	// In this case, here is the gbuffer pattern
 	// MRT0 is pixel color
@@ -206,7 +208,7 @@ uint ArrayIndex = 0;
 	SubstrateOutput4 = PixelShaderOut.SubstrateTopLayerData;
 #endif
 
-#else // SUBSTRATE_OPAQUE_DEFERRED
+#else // SUBSTRATE_OUTPUT_GBUFFER
 
 #if PIXELSHADEROUTPUT_MRT4
 	OutTarget4 = PixelShaderOut.MRT[4];
@@ -224,7 +226,7 @@ uint ArrayIndex = 0;
 	OutTarget7 = PixelShaderOut.MRT[7];
 #endif
 
-#endif // SUBSTRATE_OPAQUE_DEFERRED
+#endif // SUBSTRATE_OUTPUT_GBUFFER
 
 #if PIXELSHADEROUTPUT_COVERAGE || PIXELSHADEROUTPUT_A2C
 	OutCoverage = PixelShaderOut.Coverage;
diff --git a/Engine/Source/Runtime/Renderer/Private/BasePassRendering.cpp b/Engine/Source/Runtime/Renderer/Private/BasePassRendering.cpp
--- a/Engine/Source/Runtime/Renderer/Private/BasePassRendering.cpp
+++ b/Engine/Source/Runtime/Renderer/Private/BasePassRendering.cpp
@@ -1511,6 +1511,12 @@ void FDeferredShadingSceneRenderer::RenderBasePassInternal(
 					SkyPassPassParameters->View = View.GetShaderParameters();
 					SkyPassPassParameters->ReflectionCapture = View.ReflectionCaptureUniformBuffer;
 
+					// Remove all but the SceneColor
+					for (uint32 i = 1; i < MaxSimultaneousRenderTargets; ++i)
+					{
+						SkyPassPassParameters->RenderTargets[i] = FRenderTargetBinding();
+					}
+
 					View.ParallelMeshDrawCommandPasses[EMeshPass::SkyPass].BuildRenderingCommands(GraphBuilder, Scene->GPUScene, SkyPassPassParameters->InstanceCullingDrawParams);
 
 					GraphBuilder.AddPass(
@@ -1582,6 +1588,12 @@ void FDeferredShadingSceneRenderer::RenderBasePassInternal(
 					SkyPassParameters->View = View.GetShaderParameters();
 					SkyPassParameters->ReflectionCapture = View.ReflectionCaptureUniformBuffer;
 
+					// Remove all but the SceneColor
+					for (uint32 i = 1; i < MaxSimultaneousRenderTargets; ++i)
+					{
+						SkyPassParameters->RenderTargets[i] = FRenderTargetBinding();
+					}
+
 					View.ParallelMeshDrawCommandPasses[EMeshPass::SkyPass].BuildRenderingCommands(GraphBuilder, Scene->GPUScene, SkyPassParameters->InstanceCullingDrawParams);
 
 					GraphBuilder.AddPass(
