Fix VolumetricFog not applying all lights when r.MegaLights.Volume is disabled.

diff --git a/Engine/Shaders/Private/VolumetricFog.usf b/Engine/Shaders/Private/VolumetricFog.usf
--- a/Engine/Shaders/Private/VolumetricFog.usf
+++ b/Engine/Shaders/Private/VolumetricFog.usf
@@ -850,12 +850,21 @@ void LightScatteringCS(
 		float DistanceBiasSqr = max(CellRadius * InverseSquaredLightDistanceBiasScale, 1);
 		DistanceBiasSqr *= DistanceBiasSqr;
 
+#if USE_MEGA_LIGHTS
+		// when using Mega Lights Volume only need to apply "non-mega" lights here
+		const uint NumLights = CulledLightsGridHeader.NumLights - CulledLightsGridHeader.NumMegaLights;
+#else
+		const uint NumLights = CulledLightsGridHeader.NumLights;
+#endif
+
 		// Forward lighting of unshadowed point and spot lights
 		LOOP
-		for (uint GridLightListIndex = 0; GridLightListIndex < CulledLightsGridHeader.NumLights - CulledLightsGridHeader.NumMegaLights; GridLightListIndex++)
+		for (uint GridLightListIndex = 0; GridLightListIndex < NumLights; GridLightListIndex++)
 		{
 			const FLocalLightData LocalLight = GetLocalLightData(CulledLightsGridHeader.DataStartIndex + GridLightListIndex, 0);
+#if USE_MEGA_LIGHTS
 			checkSlow(!LocalLight.bMegaLightsSupported);
+#endif
 
 			const float VolumetricScatteringIntensity = UnpackVolumetricScatteringIntensity(LocalLight);
 
