Substrate - Fixed display of UI element using slate. Fixed UI node opacity not clamped correctly.

#rnx

diff --git a/Engine/Shaders/Private/LightFunctionCommon.ush b/Engine/Shaders/Private/LightFunctionCommon.ush
--- a/Engine/Shaders/Private/LightFunctionCommon.ush
+++ b/Engine/Shaders/Private/LightFunctionCommon.ush
@@ -79,7 +79,7 @@ float3 GetLightFunctionColor(float3 LightVector, float3 TranslatedWorldSpace
 
 #if SUBSTRATE_ENABLED
 
-	// Light function node forces a single unlit BSDF. We get the raw data from light function node from the raw BSDF on the tree.
+	// Light function node forces a single unlit BSDF. We get the raw data from the Substrate light function node from the raw BSDF on the tree.
 	FSubstratePixelHeader SubstratePixelHeader = MaterialParameters.GetFrontSubstrateHeader();
 	FSubstrateBSDF UnlitBSDF = SubstratePixelHeader.SubstrateTree.BSDFs[0];
 	UnlitBSDF.SubstrateSanitizeBSDF(); // required wehen not calling SubstrateUpdateTreeUnlit
diff --git a/Engine/Shaders/Private/PostProcessMaterialShaders.usf b/Engine/Shaders/Private/PostProcessMaterialShaders.usf
--- a/Engine/Shaders/Private/PostProcessMaterialShaders.usf
+++ b/Engine/Shaders/Private/PostProcessMaterialShaders.usf
@@ -308,7 +308,7 @@ void MainPS(
 	float3 EmissiveColor = 0.0f;
 #if SUBSTRATE_ENABLED
 
-	// Post process node forces a single unlit BSDF. We get the raw data from light function node from the raw BSDF on the tree.
+	// Post process node forces a single unlit BSDF. We get the raw data from the Substrate post process node from the raw BSDF on the tree.
 	FSubstratePixelHeader SubstratePixelHeader = Parameters.GetFrontSubstrateHeader();
 	FSubstrateBSDF UnlitBSDF = SubstratePixelHeader.SubstrateTree.BSDFs[0];
 	UnlitBSDF.SubstrateSanitizeBSDF(); // required wehen not calling SubstrateUpdateTreeUnlit
diff --git a/Engine/Shaders/Private/SlateElementPixelShader.usf b/Engine/Shaders/Private/SlateElementPixelShader.usf
--- a/Engine/Shaders/Private/SlateElementPixelShader.usf
+++ b/Engine/Shaders/Private/SlateElementPixelShader.usf
@@ -164,33 +164,36 @@ half4 GetMaterialColor( VertexToPixelInterpolants InVertex, float2 MeshUV )
 	
 #if SUBSTRATE_ENABLED
 
-	// UI node forces a single unlit BSDF. We get the raw data from light function node from the raw BSDF on the tree.
+	// UI node forces a single unlit BSDF. We get the raw data from the Substrate UI node from the raw BSDF on the tree.
 	FSubstratePixelHeader SubstratePixelHeader = Parameters.GetFrontSubstrateHeader();
 	FSubstrateBSDF UnlitBSDF = SubstratePixelHeader.SubstrateTree.BSDFs[0];
 	UnlitBSDF.SubstrateSanitizeBSDF(); // required wehen not calling SubstrateUpdateTreeUnlit
 
+	float3 UnlitLuminance = BSDF_GETEMISSIVE(UnlitBSDF);
+	float3 UnlitTransmittance= UNLIT_TRANSMITTANCE(UnlitBSDF);
+	float UnlitGreyCoverage = 1.0 - UnlitTransmittance.x;
+
 #if ((SUBSTRATE_OPAQUE_DEFERRED || (SUBSTRATE_OPAQUE_MATERIAL && SUBSTRATE_FORWARD_SHADING)))
 
-	OutColor = half4(BSDF_GETEMISSIVE(UnlitBSDF), 1.0f);
+	OutColor = half4(UnlitLuminance, 1.0f);
 
 #else // SUBSTRATE_OPAQUE_DEFERRED
 
-	float3 Color = BSDF_GETEMISSIVE(UnlitBSDF);
-	float Coverage = saturate(1.0f - UNLIT_TRANSMITTANCE(UnlitBSDF).x);	// Unlit node transmittance is only generated from opacity
-
 #if MATERIALBLENDING_ALPHACOMPOSITE
 
 #if SUBSTRATE_USE_PREMULTALPHA_OVERRIDE // AlphaComposite - Premultiplied alpha blending
-	Coverage = GetMaterialOpacity(PixelMaterialInputs);
+	UnlitGreyCoverage = GetMaterialOpacity(PixelMaterialInputs);
 #endif
-	OutColor = half4(Color, Coverage);
+	OutColor = half4(UnlitLuminance, UnlitGreyCoverage);
 
 #elif MATERIALBLENDING_MODULATE
-	OutColor = half4(Color, 0.0f);
+	OutColor = half4(UnlitTransmittance, 0.0f);
 #elif MATERIALBLENDING_ALPHAHOLDOUT
-	OutColor = half4(0.0f, 0.0f, 0.0f, Coverage);
-#else // MATERIALBLENDING_ADDITIVE || MATERIALBLENDING_TRANSLUCENT
-	OutColor = half4(Color * Coverage, Coverage);
+	OutColor = half4(0.0f, 0.0f, 0.0f, UnlitGreyCoverage);
+#elif MATERIALBLENDING_ADDITIVE
+	OutColor = half4(UnlitLuminance, 1.0);
+#else // MATERIALBLENDING_TRANSLUCENT
+	OutColor = half4(UnlitLuminance, UnlitBSDF.Coverage);
 #endif // MATERIALBLENDING_xxx
 
 #endif// SUBSTRATE_OPAQUE_DEFERRED
diff --git a/Engine/Shaders/Private/Substrate/SubstrateLegacyConversion.ush b/Engine/Shaders/Private/Substrate/SubstrateLegacyConversion.ush
--- a/Engine/Shaders/Private/Substrate/SubstrateLegacyConversion.ush
+++ b/Engine/Shaders/Private/Substrate/SubstrateLegacyConversion.ush
@@ -595,7 +595,8 @@ FSubstrateData SubstrateCreateUIMaterial(
 	FSubstrateData Out = GetInitialisedSubstrateData();
 
 	const float3 RawNormal = float3(0, 0, 1);	// Unused in this case.
-	const float GreyTransmittance = saturate(1.0f - Opacity);
+	Opacity = saturate(Opacity);
+	const float GreyTransmittance = 1.0f - Opacity;
 
 #if MATERIALBLENDING_SOLID || MATERIALBLENDING_MASKED
 	// Opaque materials only write emissive.
