Add optional shader print debug capability to the PT to ease debugging.

To enable this you needs to set r.PathTracing.Debug=1 in your ConsoleVariables.ini/DefaultEngine.ini. This CL does not introduce a extra permutation as the current RGS is already very long to compile.

diff --git a/Engine/Shaders/Private/PathTracing/PathTracingCore.ush b/Engine/Shaders/Private/PathTracing/PathTracingCore.ush
--- a/Engine/Shaders/Private/PathTracing/PathTracingCore.ush
+++ b/Engine/Shaders/Private/PathTracing/PathTracingCore.ush
@@ -11,6 +11,11 @@
 // In ordinary scenes, the most likely occurance of this problem is when combining hair and volumetric atmosphere as rays from several Km can travel back towards the hair and cause long delays or even TDRs
 #define PATH_TRACING_SKIP_HAIR_DISTANCE      1e6
 
+// Debug printing is enabled from config file with r.PathTracing.Debug
+#ifndef DEBUG_ENABLE
+#define DEBUG_ENABLE 0
+#endif
+
 #include "../HeterogeneousVolumes/HeterogeneousVolumesVoxelGridTypes.ush"
 #include "../Common.ush"
 #include "../PostProcessCommon.ush"
@@ -30,6 +35,10 @@
 #include "./Volume/PathTracingVolume.ush"
 #include "./Volume/PathTracingVolumeSampling.ush"
 
+#if DEBUG_ENABLE
+#include "../ShaderPrint.ush"
+#endif
+
 float BlendFactor;
 float BaseExposure;
 uint Iteration;
@@ -97,6 +106,7 @@ struct FPathState
 
 	// Temporary parameters not stored				
 	float DepthZ;                   // Camera ray depth
+	uint2 PixelCoord;               // Pixel index/coord
 
 	float3 GetAlbedo()
 	{
@@ -1436,6 +1446,10 @@ FPathState CreatePathState(int2 PixelIndex, int2 TextureIndex)
 	Output.SetSigmaT(float3(StartingExtinctionCoefficient[0], StartingExtinctionCoefficient[1], StartingExtinctionCoefficient[2]));
 	Output.SetFirstScatterType(PATHTRACER_SCATTER_CAMERA);
 
+	#if DEBUG_ENABLE
+	Output.PixelCoord = LaunchIndex;
+	#endif
+
 	return Output;
 }
 
@@ -1642,6 +1656,24 @@ bool PathTracingKernel(inout FPathState PathState, int Bounce)
 	// Sample material
 	FMaterialSample MaterialSample = SampleMaterial(-PathState.Ray.Direction, Payload, RandSample.xyz);
 
+	// Debug Printing
+	#if DEBUG_ENABLE
+	FShaderPrintContext Ctx = InitShaderPrintContextAtCursor(PathState.PixelCoord, uint2(50, 50));
+	if (Ctx.bIsActive)
+	{
+		Print(Ctx, TEXT("Debug"), FontRed); Newline(Ctx);
+		Print(Ctx, TEXT("Pixel Coord : "), FontWhite); Print(Ctx, PathState.PixelCoord, FontYellow); Newline(Ctx);
+
+		// Example of debugging :
+		#if SUBSTRATE_ENABLED
+		PrintLine(Ctx, Payload.GlintValue);
+		PrintLine(Ctx, Payload.GlintUV);
+		PrintLine(Ctx, Payload.GlintUVdx);
+		PrintLine(Ctx, Payload.GlintUVdy);
+		#endif
+	}
+	#endif
+
 	// If we are using Light sampling and the material can use it ...
 	if (MISMode != 0 && bDoLightLoop)
 	{
diff --git a/Engine/Source/Runtime/Renderer/Private/PathTracing.cpp b/Engine/Source/Runtime/Renderer/Private/PathTracing.cpp
--- a/Engine/Source/Runtime/Renderer/Private/PathTracing.cpp
+++ b/Engine/Source/Runtime/Renderer/Private/PathTracing.cpp
@@ -503,6 +503,12 @@ TAutoConsoleVariable<float> CVarPathTracingBackgroundAlpha(
 	ECVF_RenderThreadSafe
 );
 
+TAutoConsoleVariable<int32> CVarPathTracingDebug(
+	TEXT("r.PathTracing.Debug"),
+	0,
+	TEXT("Enable debug rendering for path tracer. Used for only development and needs to be enabled before starting the engine.\n"),
+	ECVF_RenderThreadSafe | ECVF_ReadOnly
+);
 
 BEGIN_SHADER_PARAMETER_STRUCT(FPathTracingData, )
 	SHADER_PARAMETER(float, BlendFactor)
@@ -1113,6 +1119,7 @@ class FPathTracingRG : public FGlobalShader
 	{
 		FGlobalShader::ModifyCompilationEnvironment(Parameters, OutEnvironment);
 		OutEnvironment.SetDefine(TEXT("USE_RECT_LIGHT_TEXTURES"), 1);
+		OutEnvironment.SetDefine(TEXT("DEBUG_ENABLE"), CVarPathTracingDebug.GetValueOnAnyThread() > 0 ? 1u : 0u);
 		OutEnvironment.CompilerFlags.Add(CFLAG_WarningsAsErrors);
 	}
 
@@ -1132,6 +1139,7 @@ class FPathTracingRG : public FGlobalShader
 
 		SHADER_PARAMETER_STRUCT_REF(FViewUniformShaderParameters, ViewUniformBuffer)
 		SHADER_PARAMETER_STRUCT_INCLUDE(FPathTracingData, PathTracingData)
+		SHADER_PARAMETER_STRUCT_INCLUDE(ShaderPrint::FShaderParameters, ShaderPrint)
 
 		// scene lights
 		SHADER_PARAMETER_RDG_BUFFER_SRV(StructuredBuffer<FPathTracingLight>, SceneLights)
@@ -3622,6 +3630,15 @@ void FDeferredShadingSceneRenderer::RenderPathTracing(
 								PassParameters->ActivePaths = GraphBuilder.CreateUAV(ActivePaths[0], PF_R32_UINT);
 								PassParameters->NumPathStates = GraphBuilder.CreateUAV(NumActivePaths, PF_R32_UINT);
 							}
+							const bool bEnableDebug = CVarPathTracingDebug.GetValueOnRenderThread() > 0;
+							if (bEnableDebug)
+							{
+								ShaderPrint::SetEnabled(true);
+								ShaderPrint::RequestSpaceForCharacters(1024);
+								ShaderPrint::RequestSpaceForLines(1024);
+								ShaderPrint::SetParameters(GraphBuilder, View.ShaderPrintData, PassParameters->ShaderPrint);
+							}
+
 							ClearUnusedGraphResources(RayGenShader, PassParameters);
 							const bool bFlushRenderingCommands = FlushRenderingCommands == 1 || (FlushRenderingCommands == 2 && Bounce == MaxBounces);
 							const bool bUse1DDispatch = Config.UseAdaptiveSampling || (bUseCompaction && Bounce > 0);
