Prevent artifacts when using Camera AVSM and stepping into the near field.

diff --git a/Engine/Shaders/Private/HeterogeneousVolumes/HeterogeneousVolumesLiveShadingPipeline.usf b/Engine/Shaders/Private/HeterogeneousVolumes/HeterogeneousVolumesLiveShadingPipeline.usf
--- a/Engine/Shaders/Private/HeterogeneousVolumes/HeterogeneousVolumesLiveShadingPipeline.usf
+++ b/Engine/Shaders/Private/HeterogeneousVolumes/HeterogeneousVolumesLiveShadingPipeline.usf
@@ -537,8 +537,8 @@ void RenderSingleScatteringWithLiveShading(
 		}
 
 		// Define AVSM Sampler
-		float WorldRayTMin = max(RayMarchingContext.LocalRayTMin * RayMarchingContext.LocalToWorldScale, 0.01);
-		float3 WorldPosition = RayMarchingContext.WorldRayOrigin + RayMarchingContext.WorldRayDirection * WorldRayTMin;
+		float MidpointT = (HitT.x + HitT.y) * 0.5;
+		float3 WorldPosition = RayMarchingContext.WorldRayOrigin + RayMarchingContext.WorldRayDirection * MidpointT;
 		float3 TranslatedWorldPosition = DFHackToFloat(DFFastAdd(WorldPosition, PrimaryView.PreViewTranslation));
 
 		int Face = 0;
