Clean up SMRT/VSM for forward lighting

diff --git a/Engine/Shaders/Private/ForwardLightingCommon.ush b/Engine/Shaders/Private/ForwardLightingCommon.ush
--- a/Engine/Shaders/Private/ForwardLightingCommon.ush
+++ b/Engine/Shaders/Private/ForwardLightingCommon.ush
@@ -198,7 +198,6 @@ FDeferredLightingSplit GetForwardDirectLightingSplit(
 				{
 					const FLightShaderParameters DirectionalLight = ConvertToLightShaderParameters(LightData);
 
-					const bool bBackfaceCull = !IsSubsurfaceModel(GBufferData.ShadingModelID);
 					const float ScreenRayLengthWorld = VirtualShadowMap.ScreenRayLength * View.ClipToView[1][1] * GBufferData.Depth;
 					float SMRTRayOffset = ScreenRayLengthWorld;
 
@@ -210,7 +209,6 @@ FDeferredLightingSplit GetForwardDirectLightingSplit(
 						TranslatedWorldPosition,
 						SMRTRayOffset,
 						Dither,
-						bBackfaceCull,
 						GBufferData.WorldNormal);
 
 					FilterVirtualShadowMapSampleResult(PixelPos, VirtualShadowMapSample);
diff --git a/Engine/Shaders/Private/Substrate/SubstrateForwardLighting.ush b/Engine/Shaders/Private/Substrate/SubstrateForwardLighting.ush
--- a/Engine/Shaders/Private/Substrate/SubstrateForwardLighting.ush
+++ b/Engine/Shaders/Private/Substrate/SubstrateForwardLighting.ush
@@ -120,7 +120,6 @@ float3 SubstrateForwardLighting(
 		{
 			const FLightShaderParameters DirectionalLight = ConvertToLightShaderParameters(DirLightData);
 
-			const bool bBackfaceCull = true; // SUBSTRATE_TODO: this needs to be fixed but require to extract hair/SSS/diffusebackface prior to BSDF looping
 			const float ScreenRayLengthWorld = VirtualShadowMap.ScreenRayLength * View.ClipToView[1][1] * SceneDepth;
 			float SMRTRayOffset = ScreenRayLengthWorld;
 
@@ -135,7 +134,6 @@ float3 SubstrateForwardLighting(
 				TranslatedWorldPosition,
 				SMRTRayOffset,
 				Dither,
-				bBackfaceCull,
 				WorldNormal);
 
 			FilterVirtualShadowMapSampleResult(PixelPos, VirtualShadowMapSample);
