Fix overlay flickering

Caused by not enough depth biasing, but only looked like flickering when combined with TAA

diff --git a/Engine/Shaders/Private/PostProcessSelectionOutline.usf b/Engine/Shaders/Private/PostProcessSelectionOutline.usf
--- a/Engine/Shaders/Private/PostProcessSelectionOutline.usf
+++ b/Engine/Shaders/Private/PostProcessSelectionOutline.usf
@@ -100,7 +100,7 @@ FPixelInfo GetPixelInfo(int2 PixelPos, int SampleID, int OffsetX, int OffsetY, f
 #endif
 
 	// Soft Bias with DeviceZ for best quality (larger bias than usual because SceneDeviceZ is only approximated)
-	const float DeviceDepthFade = 0.00005f;
+	const float DeviceDepthFade = 0.01f;
 
 	// 2 to always bias over the current plane
 	Ret.fVisible = saturate(2.0f - (SceneDeviceZ - DeviceZ) / DeviceDepthFade);
@@ -268,7 +268,7 @@ void MainPS(
 
 	FPixelInfo CenterPixelInfo = GetPixelInfo(ColorPixelPos, 0, SecondaryViewportOffset, 0, DeviceZPlane, DeviceZMinMax);
 
-	float4 OverlayColor = Texture2DSample(OverlayLookupTexture, OverlayLookupSampler, ColorUV);
+	float4 OverlayColor = Texture2DSample(OverlayLookupTexture, OverlayLookupSampler, DistortedColorUV);
 	OutColor = Texture2DSample(ColorTexture, ColorSampler, DistortedColorUV); 
 	
 	float4 SelectionEffect = 0;
