Fixed VS/PS interpolator mismatch from changes in Slate

diff --git a/Engine/Shaders/Private/SlatePostProcessPixelShader.usf b/Engine/Shaders/Private/SlatePostProcessPixelShader.usf
--- a/Engine/Shaders/Private/SlatePostProcessPixelShader.usf
+++ b/Engine/Shaders/Private/SlatePostProcessPixelShader.usf
@@ -33,15 +33,15 @@ float4 GetSample(float Weight, float Offset, float2 UV)
 		+ Texture2DSample(ElementTexture, ElementTextureSampler, clamp(UV - UVOffset, MinUV, MaxUV)) * Weight;
 }
 
-float4 GaussianBlurMain( FScreenVertexOutput Input ) : SV_Target0
+float4 GaussianBlurMain(in noperspective float4 InUVAndScreenPos : TEXCOORD0) : SV_Target0
 {
-	float4 OutColor = Texture2DSample(ElementTexture, ElementTextureSampler, clamp(Input.UV, UVBounds.xy, UVBounds.zw)) * WeightAndOffsets[0].x;
+	float4 OutColor = Texture2DSample(ElementTexture, ElementTextureSampler, clamp(InUVAndScreenPos.xy, UVBounds.xy, UVBounds.zw)) * WeightAndOffsets[0].x;
 	// First offset is in zw
 	{
 		float Weight = WeightAndOffsets[0].z;
 		float Offset = WeightAndOffsets[0].w;
 
-		OutColor += GetSample(Weight, Offset, Input.UV);
+		OutColor += GetSample(Weight, Offset, InUVAndScreenPos.xy);
 	}
 
 	for (int i = 2; i<SampleCount; i+=2)
@@ -51,14 +51,14 @@ float4 GaussianBlurMain( FScreenVertexOutput Input ) : SV_Target0
 			float Weight = WeightAndOffsets[Index].x;
 			float Offset = WeightAndOffsets[Index].y;
 
-			OutColor += GetSample(Weight, Offset, Input.UV);
+			OutColor += GetSample(Weight, Offset, InUVAndScreenPos.xy);
 		}
 
 		{
 			float Weight = WeightAndOffsets[Index].z;
 			float Offset = WeightAndOffsets[Index].w;
 
-			OutColor += GetSample(Weight, Offset, Input.UV);
+			OutColor += GetSample(Weight, Offset, InUVAndScreenPos.xy);
 		}
 
 	}
@@ -66,7 +66,7 @@ float4 GaussianBlurMain( FScreenVertexOutput Input ) : SV_Target0
 	return float4(OutColor.rgb, 1);
 }
 
-float4 DownsampleMain(FScreenVertexOutput Input) : SV_Target0
+float4 DownsampleMain(in noperspective float4 InUVAndScreenPos : TEXCOORD0) : SV_Target0
 {
 	float2 UV[4];
 
@@ -74,10 +74,10 @@ float4 DownsampleMain(FScreenVertexOutput Input) : SV_Target0
 	float2 MaxUV = UVBounds.zw;
 
 	// Shader params X/Y stores the UV offset in each direction
-	UV[0] = clamp(Input.UV + ShaderParams.xy * float2(-1, -1), MinUV, MaxUV);
-	UV[1] = clamp(Input.UV + ShaderParams.xy * float2(1, -1), MinUV, MaxUV);
-	UV[2] = clamp(Input.UV + ShaderParams.xy * float2(-1, 1), MinUV, MaxUV);
-	UV[3] = clamp(Input.UV + ShaderParams.xy * float2(1, 1), MinUV, MaxUV);
+	UV[0] = clamp(InUVAndScreenPos.xy + ShaderParams.xy * float2(-1, -1), MinUV, MaxUV);
+	UV[1] = clamp(InUVAndScreenPos.xy + ShaderParams.xy * float2(1, -1), MinUV, MaxUV);
+	UV[2] = clamp(InUVAndScreenPos.xy + ShaderParams.xy * float2(-1, 1), MinUV, MaxUV);
+	UV[3] = clamp(InUVAndScreenPos.xy + ShaderParams.xy * float2(1, 1), MinUV, MaxUV);
 
 	float4 Sample[4];
 
@@ -98,7 +98,7 @@ float2 Map(float2 value, float2 min1, float2 max1, float2 min2, float2 max2)
 #define UPSAMPLE_OUTPUT_FORMAT_HDR_SCRGB 1
 #define UPSAMPLE_OUTPUT_FORMAT_HDR_PQ10 2
 
-void UpsampleMain(FScreenVertexOutput Input
+void UpsampleMain(in noperspective float4 InUVAndScreenPos : TEXCOORD0
 	, out float4 Color0 : SV_Target0
 #if UPSAMPLE_OUTPUT_FORMAT != UPSAMPLE_OUTPUT_FORMAT_SDR
 	, out float4 Color1 : SV_Target1
@@ -106,10 +106,10 @@ void UpsampleMain(FScreenVertexOutput Input
 )
 {
 	const float2 LocalSize = ShaderParams.xy;
-	const float2 UV = Input.UV;
+	const float2 UV = InUVAndScreenPos.xy;
 	const float Thickness = 0;
 	const float4 CornerRadii = ShaderParams2;
-	const float4 FillColor = Texture2DSample(ElementTexture, ElementTextureSampler, Input.UV);
+	const float4 FillColor = Texture2DSample(ElementTexture, ElementTextureSampler, UV);
 	Color0 = GetRoundedBoxElementColorInternal(LocalSize, Map(UV, float2(0,0), ShaderParams.zw, float2(0,0),float2(1,1)), Thickness, CornerRadii, FillColor, FillColor);
 
 #if UPSAMPLE_OUTPUT_FORMAT != UPSAMPLE_OUTPUT_FORMAT_SDR
