Fix shader compilation warning.

diff --git a/Engine/Shaders/Private/PostProcessTonemap.usf b/Engine/Shaders/Private/PostProcessTonemap.usf
--- a/Engine/Shaders/Private/PostProcessTonemap.usf
+++ b/Engine/Shaders/Private/PostProcessTonemap.usf
@@ -434,7 +434,7 @@ float4 TonemapCommonPS(
 		float2 ExposureFusionUV = ApplyScreenTransform(UV.xy, ColorToExposureFusion);
 		ExposureFusionUV = clamp(ExposureFusionUV, ExposureFusion_UVViewportBilinearMin, ExposureFusion_UVViewportBilinearMax);
 
-		const float FusionVal = Texture2DSample(ExposureFusionTexture, ExposureFusionSampler, ExposureFusionUV);
+		const float FusionVal = Texture2DSample(ExposureFusionTexture, ExposureFusionSampler, ExposureFusionUV).x;
 
 		const float LuminanceVal = CalculateEyeAdaptationLuminance(SceneColor.rgb * OneOverPreExposure * GlobalExposure);
 		
