Fix hair strands tangent output when using hair attribute output in tangent space with Substrate.

Unlike the legacy path we only store the tangent in world space when converting hair material, or when using Substrate hair node. Due to this, when manipulating the hair tangent, we first need to transform it into local space.

diff --git a/Engine/Shaders/Private/HairStrands/HairStrandsMaterialCommon.ush b/Engine/Shaders/Private/HairStrands/HairStrandsMaterialCommon.ush
--- a/Engine/Shaders/Private/HairStrands/HairStrandsMaterialCommon.ush
+++ b/Engine/Shaders/Private/HairStrands/HairStrandsMaterialCommon.ush
@@ -239,7 +239,9 @@ FEvaluateOutput Evaluate(
 		// Convert the tangent space input tangent 
 		// * from a frame with Z/(0,0,1) being up and Y/(0,1,0) being the main tangent direction (called green)
 		// * to   a frame with Y/(0,1,0) being up and Z/(0,0,1) being the main tangent direction (called blue)
-		const float3 TangentSpaceGreen = Tangent;
+		// Tangent is in world space. First transform it into local space prior rotation)
+		float3 LocalTangent = mul(Tangent, transpose(MaterialParameters.TangentToWorld));
+		const float3 TangentSpaceGreen = LocalTangent;
 		const float3 TangentSpaceBlue = float3(-TangentSpaceGreen.x, TangentSpaceGreen.z, TangentSpaceGreen.y);
 		OutSample.Tangent = mul(TangentSpaceBlue, MaterialParameters.TangentToWorld);
 		OutSample.Tangent = normalize(OutSample.Tangent);
