Substrate - disable GLINTAA for compute shaders since we do not have ddx/ddy on some platforms.
The visual difference is minimal.

#rnx

diff --git a/Engine/Shaders/Private/Substrate/Glint/GlintThirdParty.ush b/Engine/Shaders/Private/Substrate/Glint/GlintThirdParty.ush
--- a/Engine/Shaders/Private/Substrate/Glint/GlintThirdParty.ush
+++ b/Engine/Shaders/Private/Substrate/Glint/GlintThirdParty.ush
@@ -44,7 +44,11 @@ FGlintLutParams GetGlintParams()
 	return GlintLutParams;
 }
 
+// GLINT AA must be skipped for some shaders who do not support ddx/ddy, such as ray tracing shaders.
+// It is also skiped for Compute shader (MegaLights). (available on SM6.6 but not on the Metal API yet).
+// The good news is that the visual difference is minimal.
 #define USE_GLINT_AA (!RAYCALLABLESHADER && !RAYHITGROUPSHADER && !RAYGENSHADER && !RAYMISSSHADER)
+//&& !COMPUTESHADER)
 
 // Geometric Glint Anti-Aliasing parameters
 #define GGAAKernelSize						0.5
