Fix NoL being applied twice on eye shading causing a darkerning compare to legacy shading.

diff --git a/Engine/Shaders/Private/Substrate/SubstrateEvaluation.ush b/Engine/Shaders/Private/Substrate/SubstrateEvaluation.ush
--- a/Engine/Shaders/Private/Substrate/SubstrateEvaluation.ush
+++ b/Engine/Shaders/Private/Substrate/SubstrateEvaluation.ush
@@ -1052,10 +1052,10 @@ FSubstrateEvaluateResult SubstrateEvaluateBSDFCommon(FSubstrateBSDFContext BSDFC
 			const float Caustic = 0.8 + 0.2 * (Power + 1) * pow(saturate(dot(CausticNormal, BSDFContext.L)), Power);
 
 			const float Iris   = IrisNoL * Caustic;
-			const float Sclera = BSDFContext.Context.NoL;
+			const float Sclera = AreaLightContext.NoL;
 
 			Sample.DiffusePathValue = Diffuse_Lambert(DiffuseColor) * lerp(Sclera, Iris, EYE_IRISMASK(BSDFContext.BSDF));
-			Sample.IntegratedDiffuseValue += (ShadowTerms.SurfaceShadow * AreaLightContext.NoL * AreaLightContext.Falloff) * AreaLightContext.AreaLight.FalloffColor * Sample.DiffusePathValue;
+			Sample.IntegratedDiffuseValue += (ShadowTerms.SurfaceShadow * /* We don't AreaLightContext.NoL as it is already applied through the Sclera/Iris value above */ AreaLightContext.Falloff) * AreaLightContext.AreaLight.FalloffColor * Sample.DiffusePathValue;
 
 			Sample.DiffuseColor		= DiffuseColor;
 			Sample.DiffusePDF		= BSDFContext.SatNoL / PI;
