* Add Gbuffer normal/tangent visualization
* Fix typos

diff --git a/Engine/Shaders/Private/DeferredShadingCommon.ush b/Engine/Shaders/Private/DeferredShadingCommon.ush
--- a/Engine/Shaders/Private/DeferredShadingCommon.ush
+++ b/Engine/Shaders/Private/DeferredShadingCommon.ush
@@ -915,9 +915,9 @@ bool AdjustBaseColorAndSpecularColorForSubsurfaceProfileLighting(inout float3 Ba
 	const bool bCheckerboardRequired = View.bSubsurfacePostprocessEnabled > 0;
 #else
 	const bool bCheckerboardRequired = View.bSubsurfacePostprocessEnabled > 0 && View.bCheckerboardSubsurfaceProfileRendering > 0;
-	BaseColor = View.bSubsurfacePostprocessEnabled ? float3(1, 1, 1) : BaseColor;
+	BaseColor = (ALLOW_SSS_MATERIAL_OVERRIDE && View.bSubsurfacePostprocessEnabled) ? float3(1, 1, 1) : BaseColor;
 #endif
-	if (bCheckerboardRequired)
+	if (bCheckerboardRequired && ALLOW_SSS_MATERIAL_OVERRIDE)
 	{
 		// because we adjust the BaseColor here, we need StoredBaseColor
 		// we apply the base color later in SubsurfaceRecombinePS()
diff --git a/Engine/Shaders/Private/PostProcessGBufferHints.usf b/Engine/Shaders/Private/PostProcessGBufferHints.usf
--- a/Engine/Shaders/Private/PostProcessGBufferHints.usf
+++ b/Engine/Shaders/Private/PostProcessGBufferHints.usf
@@ -1,5 +1,9 @@
 // Copyright Epic Games, Inc. All Rights Reserved.
 
+#if SHADER_GBUFFER_PICKING
+#define ALLOW_SSS_MATERIAL_OVERRIDE 0
+#endif
+
 #include "Common.ush"
 #include "DeferredShadingCommon.ush"
 #include "ScreenPass.ush"
@@ -122,6 +126,7 @@ void MainPS(noperspective float4 UVAndScreenPos : TEXCOORD0, float4 SvPosition :
 #include "DeferredShadingCommon.ush"
 #include "ShadingModels.ush"
 #include "SubsurfaceProfileCommon.ush"
+#include "PositionReconstructionCommon.ush"
 
 void PrintShadingMode(inout FShaderPrintContext Ctx, uint In)
 {
@@ -234,15 +239,22 @@ void MainCS(uint3 DispatchThreadId : SV_DispatchThreadID)
 			float Lobe1Roughness = 0;
 			GetProfileDualSpecular(SSSProfileID, Data.Roughness, Opacity, Lobe0Roughness, Lobe1Roughness, LobeMix);
 
-			Print(Ctx, TEXT("SSS Profil ID    : "), FontWhite); Print(Ctx, SSSProfileID); Newline(Ctx, RectMax);
-			Print(Ctx, TEXT("SSS Lobe Mix     : "), FontWhite); Print(Ctx, LobeMix); Newline(Ctx, RectMax);
-			Print(Ctx, TEXT("SSS Roughness0   : "), FontWhite); Print(Ctx, Lobe0Roughness); Newline(Ctx, RectMax);
-			Print(Ctx, TEXT("SSS Roughness1   : "), FontWhite); Print(Ctx, Lobe1Roughness); Newline(Ctx, RectMax);
-			Print(Ctx, TEXT("SSS SpecColor    : "), FontWhite); Print(Ctx, Data.SpecularColor); Newline(Ctx, RectMax);
+			Print(Ctx, TEXT("SSS Profil ID     : "), FontWhite); Print(Ctx, SSSProfileID); Newline(Ctx, RectMax);
+			Print(Ctx, TEXT("SSS Lobe Mix      : "), FontWhite); Print(Ctx, LobeMix); Newline(Ctx, RectMax);
+			Print(Ctx, TEXT("SSS Roughness0    : "), FontWhite); Print(Ctx, Lobe0Roughness); Newline(Ctx, RectMax);
+			Print(Ctx, TEXT("SSS Roughness1    : "), FontWhite); Print(Ctx, Lobe1Roughness); Newline(Ctx, RectMax);
+			Print(Ctx, TEXT("SSS SpecColor     : "), FontWhite); Print(Ctx, Data.SpecularColor); Newline(Ctx, RectMax);
 			
 		}
 		PrintBackgroundRect(Ctx, RectMin, RectMax);
 		Newline(Ctx);
+		
+		// Draw tangent vector (strands)
+		{
+			const float2 UV = (ShaderPrintData.CursorCoord + 0.5f) * View.BufferSizeAndInvSize.zw;
+			const float3 WorldPosition = ReconstructTranslatedWorldPositionFromDepth(UV, Data.Depth);
+			AddLineTWS(WorldPosition, WorldPosition + Data.WorldNormal, ColorRed, ColorYellow);
+		}
 	}
 #endif
 }
diff --git a/Engine/Shaders/Private/Substrate/SubstrateDBuffer.usf b/Engine/Shaders/Private/Substrate/SubstrateDBuffer.usf
--- a/Engine/Shaders/Private/Substrate/SubstrateDBuffer.usf
+++ b/Engine/Shaders/Private/Substrate/SubstrateDBuffer.usf
@@ -95,9 +95,6 @@ void MainCS(
 	uint3 GroupId : SV_GroupID,
 	uint2 GroupThreadId : SV_GroupThreadID)
 {
-	if (View.GeneralPurposeTweak == 0)
-		return;
-
 	const uint2 TileCoord = SubstrateUnpackTile(TileListBuffer[TileListBufferOffset + GroupId.x], TileEncoding);
 	const uint2 PixelCoord = TileCoord * SUBSTRATE_TILE_SIZE + GroupThreadId;
 	const bool bIsValid = all(PixelCoord < uint2(View.ViewSizeAndInvSize.xy));
@@ -353,4 +350,4 @@ void MainCS(
 		#endif // !SUBSTRATE_FASTPATH
 	}
 }
-#endif // SHADER_DBUFFER
\ No newline at end of file
+#endif // SHADER_DBUFFER
diff --git a/Engine/Shaders/Private/Substrate/SubstrateVisualize.usf b/Engine/Shaders/Private/Substrate/SubstrateVisualize.usf
--- a/Engine/Shaders/Private/Substrate/SubstrateVisualize.usf
+++ b/Engine/Shaders/Private/Substrate/SubstrateVisualize.usf
@@ -430,8 +430,8 @@ void PrintBSDF(inout FShaderPrintContext Context, inout FSubstrateAddressing Sub
 			Print(Context, TEXT(" Diffuse Albedo    "), FontBSDFPropName);	Print(Context, EYE_DIFFUSEALBEDO(BSDF), FontBSDFPropValu);			Newline(Context, RectMax);
 			Print(Context, TEXT(" F0                "), FontBSDFPropName);	Print(Context, EYE_F0(BSDF), FontBSDFPropValu);						Newline(Context, RectMax);
 			Print(Context, TEXT(" Roughness         "), FontBSDFPropName);	Print(Context, EYE_ROUGHNESS(BSDF), FontBSDFPropValu);				Newline(Context, RectMax);
-			Print(Context, TEXT(" Iris Mask         "), FontBSDFPropName);	Print(Context, EYE_IRISDISTANCE(BSDF), FontBSDFPropValu);			Newline(Context, RectMax);
-			Print(Context, TEXT(" Iris Distance     "), FontBSDFPropName);	Print(Context, EYE_IRISMASK(BSDF), FontBSDFPropValu);				Newline(Context, RectMax);
+			Print(Context, TEXT(" Iris Mask         "), FontBSDFPropName);	Print(Context, EYE_IRISMASK(BSDF), FontBSDFPropValu);			Newline(Context, RectMax);
+			Print(Context, TEXT(" Iris Distance     "), FontBSDFPropName);	Print(Context, EYE_IRISDISTANCE(BSDF), FontBSDFPropValu);				Newline(Context, RectMax);
 
 			// Tangent basis for the iris
 			const float3x3 IrisTangentBasis = GetTangentBasis(EYE_IRISNORMAL(BSDF));
