* Fix missing SSS profile on eye when using material attributes
* Fix incorrect shadow with VSM with SSS profile

diff --git a/Engine/Shaders/Private/Substrate/SubstrateSubsurface.ush b/Engine/Shaders/Private/Substrate/SubstrateSubsurface.ush
--- a/Engine/Shaders/Private/Substrate/SubstrateSubsurface.ush
+++ b/Engine/Shaders/Private/Substrate/SubstrateSubsurface.ush
@@ -246,20 +246,20 @@ float SubstrateShadowColoredMFPToGreyScaleMFP(float3 MFP)
 }
 
 // Legacy conversion function, which translated Substrate sub-surface data into 'opacity' for shadow transmission purpose. 
-FSubsurfaceOpacityMFP SubstrateGetSubsurfaceOpacityMFP(FSubstrateSubsurfaceHeader SSSHeader)
+FSubsurfaceOpacityMFP SubstrateGetSubsurfaceOpacityMFP(FSubstrateSubsurfaceHeader SSSHeader, bool bAllowDiffuse = true, bool bAllowDiffuseProfile = true)
 {
 	FSubsurfaceOpacityMFP SubsurfaceOpacityMFP = GetInitialisedSubsurfaceOpacityMFP();
 
 	const uint SSSType = SubstrateSubSurfaceHeaderGetSSSType(SSSHeader);
 
-	if (SSSType == SSS_TYPE_DIFFUSION)
+	if (bAllowDiffuse && SSSType == SSS_TYPE_DIFFUSION)
 	{
 		SubsurfaceOpacityMFP.bDataIsOpacity = false;
 		const float3 ColoredMFP = SubstrateSubSurfaceHeaderGetMFP(SSSHeader);
 		SubsurfaceOpacityMFP.Data = SubstrateShadowColoredMFPToGreyScaleMFP(ColoredMFP);
 		SubsurfaceOpacityMFP.Density = 0.0f;
 	}
-	else if (SSSType == SSS_TYPE_DIFFUSION_PROFILE)
+	else if (bAllowDiffuseProfile && SSSType == SSS_TYPE_DIFFUSION_PROFILE)
 	{
 		SubsurfaceOpacityMFP.bDataIsOpacity = true;
 
diff --git a/Engine/Shaders/Private/VirtualShadowMaps/VirtualShadowMapTransmissionCommon.ush b/Engine/Shaders/Private/VirtualShadowMaps/VirtualShadowMapTransmissionCommon.ush
--- a/Engine/Shaders/Private/VirtualShadowMaps/VirtualShadowMapTransmissionCommon.ush
+++ b/Engine/Shaders/Private/VirtualShadowMaps/VirtualShadowMapTransmissionCommon.ush
@@ -20,7 +20,8 @@ FSubsurfaceOpacityMFP GetSubsurfaceOpacityFromGbuffer(uint2 PixelPos)
 
 	FSubstrateSubsurfaceData LoadStrat = SubstrateLoadSubsurfaceData(Substrate.MaterialTextureArray, Substrate.FirstSliceStoringSubstrateSSSData, PixelPos);
 
-	SubsurfaceOpacityMFP = SubstrateGetSubsurfaceOpacityMFP(LoadStrat.Header);
+	// Diffusion and DiffuseProfile are not allowed to translate their MFP into opacity in order to match legacy behavior
+	SubsurfaceOpacityMFP = SubstrateGetSubsurfaceOpacityMFP(LoadStrat.Header, false /*bAllowDiffuse*/, false /*bAllowDiffuseProfile*/);
 
 #else // SUBSTRATE_ENABLED
 	FGBufferData GBufferData = GetGBufferDataUint(PixelPos);
diff --git a/Engine/Source/Runtime/Engine/Private/Materials/Material.cpp b/Engine/Source/Runtime/Engine/Private/Materials/Material.cpp
--- a/Engine/Source/Runtime/Engine/Private/Materials/Material.cpp
+++ b/Engine/Source/Runtime/Engine/Private/Materials/Material.cpp
@@ -5399,7 +5399,7 @@ void UMaterial::RebuildShadingModelField()
 
 		// Now, reset the subsurface profile (in case it has been removed from any slab before) and set it only if needed.
 		SubsurfaceProfile = nullptr;
-		if ((SubstrateMaterialInfo.HasOnlyShadingModel(SSM_Eye) || SubstrateMaterialInfo.HasOnlyShadingModel(SSM_SubsurfaceProfile)) && SubstrateMaterialInfo.CountSubsurfaceProfiles() > 0)
+		if ((SubstrateMaterialInfo.HasShadingModel(SSM_Eye) || SubstrateMaterialInfo.HasShadingModel(SSM_SubsurfaceProfile)) && SubstrateMaterialInfo.CountSubsurfaceProfiles() > 0)
 		{
 			SubsurfaceProfile = SubstrateMaterialInfo.GetSubsurfaceProfile();
 		}
