PR #12245: IES Texture Thumbnails: Fix <0 InIESAtlasIndex when standalone IES textures are used.

diff --git a/Engine/Shaders/Private/IESLightProfilesCommon.ush b/Engine/Shaders/Private/IESLightProfilesCommon.ush
--- a/Engine/Shaders/Private/IESLightProfilesCommon.ush
+++ b/Engine/Shaders/Private/IESLightProfilesCommon.ush
@@ -29,7 +29,10 @@ float ComputeLightProfileMultiplier(float3 WorldPosition, float3 LightPosition,
 {
 	float Out = 1.0f;
 #if USE_IES_PROFILE && !VISUALIZE_LIGHT_CULLING
+
+#if !defined(USE_IES_STANDALONE_TEXTURE) || !USE_IES_STANDALONE_TEXTURE
 	if (InIESAtlasIndex >= 0)
+#endif
 	{
 		float3 LightBitangent = normalize( cross( LightTangent, LightDirection ) );
 
@@ -49,7 +52,7 @@ float ComputeLightProfileMultiplier(float3 WorldPosition, float3 LightPosition,
 		float TangentAngle = atan2( -LocalToLight.z, -LocalToLight.y ); // -Y represents 0/360 horizontal angle and we're rotating counter-clockwise
 		float NormTangentAngle = TangentAngle / (PI * 2.f) + 0.5f;
 
-	#if USE_IES_STANDALONE_TEXTURE
+	#if defined(USE_IES_STANDALONE_TEXTURE) && USE_IES_STANDALONE_TEXTURE == 1
 		Out = Texture2DSampleLevel(IESTexture, IESTextureSampler, float2(NormAngle, NormTangentAngle), 0).r;
 	#else
 		Out = View.IESAtlasTexture.SampleLevel(IESAtlasSampler, float3(NormAngle, NormTangentAngle, InIESAtlasIndex), 0);
@@ -57,4 +60,4 @@ float ComputeLightProfileMultiplier(float3 WorldPosition, float3 LightPosition,
 	}
 #endif
 	return Out;
-}
\ No newline at end of file
+}
