Added IsDrawLocked to FShaderPrintConfig such that GPU can stop drawing when locked if desired and made it possible to use projecting transforms with AddOBB*.

diff --git a/Engine/Shaders/Private/ShaderPrint.ush b/Engine/Shaders/Private/ShaderPrint.ush
--- a/Engine/Shaders/Private/ShaderPrint.ush
+++ b/Engine/Shaders/Private/ShaderPrint.ush
@@ -77,7 +77,7 @@ void PackShaderPrintConfig(FShaderPrintConfig In)
 	UEDiagnosticBuffer[Offset+7] = In.MaxCharacterCount;
 	UEDiagnosticBuffer[Offset+8] = In.MaxSymbolCount;
 	UEDiagnosticBuffer[Offset+9] = In.MaxStateCount;
-	UEDiagnosticBuffer[Offset+10]= In.MaxLineCount;
+	UEDiagnosticBuffer[Offset+10]= (In.MaxLineCount << 1) | (In.bIsDrawLocked ? 1 : 0);
 	UEDiagnosticBuffer[Offset+11]= In.MaxTriangleCount;
 }
 
@@ -96,7 +96,8 @@ FShaderPrintConfig UnpackShaderPrintConfig()
 	Out.MaxCharacterCount		= UEDiagnosticBuffer[Offset+7];
 	Out.MaxSymbolCount			= UEDiagnosticBuffer[Offset+8];
 	Out.MaxStateCount			= UEDiagnosticBuffer[Offset+9];
-	Out.MaxLineCount			= UEDiagnosticBuffer[Offset+10];
+	Out.MaxLineCount			= UEDiagnosticBuffer[Offset+10] >> 1;
+	Out.bIsDrawLocked			= (UEDiagnosticBuffer[Offset+10] & 1) != 0;
 	Out.MaxTriangleCount		= UEDiagnosticBuffer[Offset+11];
 
 	// Dummy write to avoid the pass to be culled when using explicit ShaderPrint parameter bindings
@@ -141,6 +142,7 @@ FShaderPrintConfig InitShaderPrintContextConfig()
 	Config.MaxStateCount = ShaderPrintData.MaxStateCount;
 	Config.MaxLineCount = ShaderPrintData.MaxLineCount;
 	Config.MaxTriangleCount = ShaderPrintData.MaxTriangleCount;
+	Config.bIsDrawLocked = ShaderPrintData.IsDrawLocked;
 	return Config;
 #endif
 }
diff --git a/Engine/Shaders/Private/ShaderPrintCommon.ush b/Engine/Shaders/Private/ShaderPrintCommon.ush
--- a/Engine/Shaders/Private/ShaderPrintCommon.ush
+++ b/Engine/Shaders/Private/ShaderPrintCommon.ush
@@ -335,6 +335,7 @@ struct FShaderPrintConfig
 	uint MaxStateCount;
 	uint MaxLineCount;
 	uint MaxTriangleCount;
+	bool bIsDrawLocked;
 };
 
 struct FShaderPrintBuffers
@@ -351,6 +352,8 @@ struct FShaderPrintContext
 
 	FShaderPrintConfig Config;
 
+	bool IsDrawLocked() { return Config.bIsDrawLocked; }
+
 #if !SHADER_PRINT_USE_GLOBAL_RESOURCE
 	FShaderPrintBuffers Buffers;
 #endif
@@ -1320,46 +1323,57 @@ void AddLineTriangle   (FShaderPrintContext Ctx, float3 P0, float3 P1, float3 P2
 ///////////////////////////////////////////////////////////////////////////////////////////////////
 // OBB
 
-void AddOBBTWS(FShaderPrintContext Ctx, float3 Min, float3 Max, float4 Color, float4x4 Transform)
+void AddOBBTWS(FShaderPrintContext Ctx, float3 Min, float3 Max, float4 Color, float4x4 Transform, bool bHomogeneousDivide = false)
 {
 	if (Ctx.bIsActive) 
 	{
 		uint Offset = 0;
 		if (AllocateLineElement(Ctx, 12, Offset))
 		{
-			float3 P0 = mul(float4(Min.x, Min.y, Min.z, 1.0f), Transform).xyz;
-			float3 P1 = mul(float4(Max.x, Min.y, Min.z, 1.0f), Transform).xyz;
-			float3 P2 = mul(float4(Max.x, Max.y, Min.z, 1.0f), Transform).xyz;
-			float3 P3 = mul(float4(Min.x, Max.y, Min.z, 1.0f), Transform).xyz;
-			float3 P4 = mul(float4(Min.x, Min.y, Max.z, 1.0f), Transform).xyz;
-			float3 P5 = mul(float4(Max.x, Min.y, Max.z, 1.0f), Transform).xyz;
-			float3 P6 = mul(float4(Max.x, Max.y, Max.z, 1.0f), Transform).xyz;
-			float3 P7 = mul(float4(Min.x, Max.y, Max.z, 1.0f), Transform).xyz;
-
+			float4 P0 = mul(float4(Min.x, Min.y, Min.z, 1.0f), Transform);
+			float4 P1 = mul(float4(Max.x, Min.y, Min.z, 1.0f), Transform);
+			float4 P2 = mul(float4(Max.x, Max.y, Min.z, 1.0f), Transform);
+			float4 P3 = mul(float4(Min.x, Max.y, Min.z, 1.0f), Transform);
+			float4 P4 = mul(float4(Min.x, Min.y, Max.z, 1.0f), Transform);
+			float4 P5 = mul(float4(Max.x, Min.y, Max.z, 1.0f), Transform);
+			float4 P6 = mul(float4(Max.x, Max.y, Max.z, 1.0f), Transform);
+			float4 P7 = mul(float4(Min.x, Max.y, Max.z, 1.0f), Transform);
+
+			if (bHomogeneousDivide)
+			{
+				P0.xyz /= P0.w;
+				P1.xyz /= P1.w;
+				P2.xyz /= P2.w;
+				P3.xyz /= P3.w;
+				P4.xyz /= P4.w;
+				P5.xyz /= P5.w;
+				P6.xyz /= P6.w;
+				P7.xyz /= P7.w;
+			}
 			FLineElement Element;
 			Element.bIsScreenSpace = false;
 			Element.Color0 = Element.Color1 = Color;
 
-			Element.Pos0 = P0; Element.Pos1 = P1; AddLineElement(Ctx, Element, Offset + 0);
-			Element.Pos0 = P1; Element.Pos1 = P2; AddLineElement(Ctx, Element, Offset + 1);
-			Element.Pos0 = P2; Element.Pos1 = P3; AddLineElement(Ctx, Element, Offset + 2);
-			Element.Pos0 = P3; Element.Pos1 = P0; AddLineElement(Ctx, Element, Offset + 3);
+			Element.Pos0 = P0.xyz; Element.Pos1 = P1.xyz; AddLineElement(Ctx, Element, Offset + 0);
+			Element.Pos0 = P1.xyz; Element.Pos1 = P2.xyz; AddLineElement(Ctx, Element, Offset + 1);
+			Element.Pos0 = P2.xyz; Element.Pos1 = P3.xyz; AddLineElement(Ctx, Element, Offset + 2);
+			Element.Pos0 = P3.xyz; Element.Pos1 = P0.xyz; AddLineElement(Ctx, Element, Offset + 3);
 
-			Element.Pos0 = P4; Element.Pos1 = P5; AddLineElement(Ctx, Element, Offset + 4);
-			Element.Pos0 = P5; Element.Pos1 = P6; AddLineElement(Ctx, Element, Offset + 5);
-			Element.Pos0 = P6; Element.Pos1 = P7; AddLineElement(Ctx, Element, Offset + 6);
-			Element.Pos0 = P7; Element.Pos1 = P4; AddLineElement(Ctx, Element, Offset + 7);
+			Element.Pos0 = P4.xyz; Element.Pos1 = P5.xyz; AddLineElement(Ctx, Element, Offset + 4);
+			Element.Pos0 = P5.xyz; Element.Pos1 = P6.xyz; AddLineElement(Ctx, Element, Offset + 5);
+			Element.Pos0 = P6.xyz; Element.Pos1 = P7.xyz; AddLineElement(Ctx, Element, Offset + 6);
+			Element.Pos0 = P7.xyz; Element.Pos1 = P4.xyz; AddLineElement(Ctx, Element, Offset + 7);
 
-			Element.Pos0 = P0; Element.Pos1 = P4; AddLineElement(Ctx, Element, Offset + 8);
-			Element.Pos0 = P1; Element.Pos1 = P5; AddLineElement(Ctx, Element, Offset + 9);
-			Element.Pos0 = P2; Element.Pos1 = P6; AddLineElement(Ctx, Element, Offset + 10);
-			Element.Pos0 = P3; Element.Pos1 = P7; AddLineElement(Ctx, Element, Offset + 11);
+			Element.Pos0 = P0.xyz; Element.Pos1 = P4.xyz; AddLineElement(Ctx, Element, Offset + 8);
+			Element.Pos0 = P1.xyz; Element.Pos1 = P5.xyz; AddLineElement(Ctx, Element, Offset + 9);
+			Element.Pos0 = P2.xyz; Element.Pos1 = P6.xyz; AddLineElement(Ctx, Element, Offset + 10);
+			Element.Pos0 = P3.xyz; Element.Pos1 = P7.xyz; AddLineElement(Ctx, Element, Offset + 11);
 		}
 	}
 }
 
-void AddOBBWS (FShaderPrintContext Ctx, float3 Min, float3 Max, float4 Color, float4x4 Transform)	{ Transform[3].xyz += Ctx.Config.TranslatedWorldOffset; AddOBBTWS(Ctx, Min, Max, Color, Transform); }
-void AddOBB   (FShaderPrintContext Ctx, float3 Min, float3 Max, float4 Color, float4x4 Transform)	{ Transform[3].xyz += Ctx.Config.TranslatedWorldOffset; AddOBBTWS(Ctx, Min, Max, Color, Transform); }
+void AddOBBWS (FShaderPrintContext Ctx, float3 Min, float3 Max, float4 Color, float4x4 Transform, bool bHomogeneousDivide = false)	{ Transform[3].xyz += Ctx.Config.TranslatedWorldOffset; AddOBBTWS(Ctx, Min, Max, Color, Transform, bHomogeneousDivide); }
+void AddOBB   (FShaderPrintContext Ctx, float3 Min, float3 Max, float4 Color, float4x4 Transform, bool bHomogeneousDivide = false)	{ Transform[3].xyz += Ctx.Config.TranslatedWorldOffset; AddOBBTWS(Ctx, Min, Max, Color, Transform, bHomogeneousDivide); }
 
 ///////////////////////////////////////////////////////////////////////////////////////////////////
 // Strings
diff --git a/Engine/Shaders/Private/ShaderPrintDraw.usf b/Engine/Shaders/Private/ShaderPrintDraw.usf
--- a/Engine/Shaders/Private/ShaderPrintDraw.usf
+++ b/Engine/Shaders/Private/ShaderPrintDraw.usf
@@ -441,6 +441,7 @@ void UploadCS(int2 DispatchThreadId : SV_DispatchThreadID)
 	In.MaxStateCount = ShaderPrintData.MaxStateCount;
 	In.MaxLineCount = ShaderPrintData.MaxLineCount;
 	In.MaxTriangleCount = ShaderPrintData.MaxTriangleCount;
+	In.bIsDrawLocked = ShaderPrintData.IsDrawLocked;
 	PackShaderPrintConfig(In);
 
 	// Clear counters
