Heterogeneous Volumes: Fixing buffer overrun

diff --git a/Engine/Shaders/Private/HeterogeneousVolumes/HeterogeneousVolumesVoxelGridShadows.usf b/Engine/Shaders/Private/HeterogeneousVolumes/HeterogeneousVolumesVoxelGridShadows.usf
--- a/Engine/Shaders/Private/HeterogeneousVolumes/HeterogeneousVolumesVoxelGridShadows.usf
+++ b/Engine/Shaders/Private/HeterogeneousVolumes/HeterogeneousVolumesVoxelGridShadows.usf
@@ -598,9 +598,8 @@ void CompressVolumetricShadowMapCS(
 		}
 	}
 
-
 	// Pad remainder
-	for (int BottomLevelIndex = SampleCount - 1; BottomLevelIndex < Align4(SampleCount); ++BottomLevelIndex)
+	for (int BottomLevelIndex = max(SampleCount - 1, 0); BottomLevelIndex < Align4(SampleCount); ++BottomLevelIndex)
 	{
 		RWVolumetricShadowTransmittanceBuffer[BottomLevelOffset + BottomLevelIndex] = AVSM_PackSampleData(AVSM_CreateSampleData(MaxHalfFloat, ShadowData.Tau));
 	}
