fix helper lane visualization for nanite

diff --git a/Engine/Shaders/Private/ComputeShaderOutputCommon.ush b/Engine/Shaders/Private/ComputeShaderOutputCommon.ush
--- a/Engine/Shaders/Private/ComputeShaderOutputCommon.ush
+++ b/Engine/Shaders/Private/ComputeShaderOutputCommon.ush
@@ -273,8 +273,8 @@ void MainCS(
 		PixelWriteMask		= BitFieldExtractU32(PackedElement.y, 4, CornerIndex * 4u);
 
 		PixelPos			= TopLeft + (uint2(CornerIndex & 1u, CornerIndex >> 1u) << VRSShift);
-
-		HelperLaneCount		= (VIS_HELPER_LANES && bQuadBinning) ? (4u - countbits(PixelWriteMask)) : 0u;
+		uint NumActiveLanes = ((PackedElement.y & 0xF) != 0) + ((PackedElement.y & 0xF0) != 0) + ((PackedElement.y & 0xF00) != 0) + ((PackedElement.y & 0xF000) != 0);
+		HelperLaneCount		= (VIS_HELPER_LANES && bQuadBinning) ? (4u - NumActiveLanes) : 0u;
 	}
 	else
 	{
@@ -282,6 +282,7 @@ void MainCS(
 		PixelPos				= uint2(BitFieldExtractU32(PackedElement, 13, 0), BitFieldExtractU32(PackedElement, 13, 13));
 		VRSShift				= uint2(BitFieldExtractU32(PackedElement, 1, 26), BitFieldExtractU32(PackedElement, 1, 27));
 		PixelWriteMask			= PackedElement >> 28;
+		HelperLaneCount		= 0;
 	}
 
 	const float2 SVPositionXY = PixelPos + ((1u << VRSShift) * 0.5f);
