Fix missing backlight on SSS surface with Lumen.

diff --git a/Engine/Shaders/Private/Substrate/Substrate.ush b/Engine/Shaders/Private/Substrate/Substrate.ush
--- a/Engine/Shaders/Private/Substrate/Substrate.ush
+++ b/Engine/Shaders/Private/Substrate/Substrate.ush
@@ -3113,13 +3113,13 @@ void FSubstrateBSDF::PostProcessBSDFBeforeLighting(bool bIsSubstrateOpaqueMateri
 
 bool FSubstrateBSDF::HasScatteringData()
 {
-	const uint SSSType = BSDF_GETTYPE(this);
+	const uint SSSType = BSDF_GETSSSTYPE(this);
 	return SSSType != SSS_TYPE_INVALID && SSSType != SSS_TYPE_SIMPLEVOLUME;
 }
 
 bool FSubstrateBSDF::HasBackScattering()
 {
-	const uint SSSType = BSDF_GETTYPE(this);
+	const uint SSSType = BSDF_GETSSSTYPE(this);
 	return SSSType != SSS_TYPE_INVALID && SSSType != SSS_TYPE_SIMPLEVOLUME;
 }
 
