Change ThinTranslucent legacy shading conversion to fix certain case.

The conversion is implified, and now use a simple transmittance color lerp, rather than using LAB + MFP lerp.

diff --git a/Engine/Shaders/Private/Substrate/SubstrateLegacyConversion.ush b/Engine/Shaders/Private/Substrate/SubstrateLegacyConversion.ush
--- a/Engine/Shaders/Private/Substrate/SubstrateLegacyConversion.ush
+++ b/Engine/Shaders/Private/Substrate/SubstrateLegacyConversion.ush
@@ -392,32 +392,18 @@ FSubstrateData SubstrateConvertLegacyMaterialDynamic(
 	{
 		const float Thickness = DefaultThickness;
 
-		// If the material is dieletric, bias the opacity so that the dieletric part extend further. This better matches the two layer approaches.
-		Opacity = lerp(1.f - Square(1.f - Opacity), Opacity, Metallic);
-
-		// Because we are using a single slab for simulating a 'coverage mix' between two BSDFs (a opaque metallic/conductor and a non-scattering translucent medium), 
-		// we need to do some approximation. In order to reduce avoid hue shift due to difference between 'solid' base-color/diffuse-albedo and 'transparent' transmission 
-		// color, we lerp BaseColor with transmittance color in LAB space so that the interpolated color minized its 'distance' 
-		const float3 BaseColor_LAB			= LinearRGB_2_LAB(BaseColor);
-		const float3 TransmittanceColor_LAB = LinearRGB_2_LAB(ThinTranslucentTransmittanceColor);
-		const float3 Color_LAB			    = lerp(TransmittanceColor_LAB, BaseColor_LAB, Opacity);
-		const float3 Color					= LAB_2_LinearRGB(Color_LAB);
-
-		// Lerp the transmittance color rather than MFP as it is more perceptually linear
-		const float3 MFP = TransmittanceToMeanFreePath(lerp(Color, Zeros, Opacity), Thickness * CENTIMETER_TO_METER) * METER_TO_CENTIMETER;
-
+		const float3 TransColor = lerp(ThinTranslucentTransmittanceColor, Zeros, Opacity);
+		const float3 MFP = TransmittanceToMeanFreePath(TransColor, Thickness * CENTIMETER_TO_METER) * METER_TO_CENTIMETER;
 		const float F0Dieletrict = DielectricSpecularToF0(Specular);
 
 		const float3 TopF0     = lerp(F0Dieletrict, BaseColor, Metallic);
-		const float3 TopAlbedo = Color * (1.f - Metallic);
-		const float3 TopMFP    = 0;
+		const float3 TopAlbedo = BaseColor * (1.f - Metallic);
 
 		const float3 BotF0     = F0Dieletrict;
 		const float3 BotAlbedo = Zeros;
-		const float3 BotMFP    = MFP;
 
 		LegacySlab.DiffuseAlbedo	= lerp(BotAlbedo, TopAlbedo, Opacity);
-		LegacySlab.SSSMFP			= lerp(BotMFP,    TopMFP,    Opacity);
+		LegacySlab.SSSMFP			= MFP;
 		LegacySlab.F0				= lerp(BotF0,     TopF0,     Opacity);
 		LegacySlab.F90				= Ones;
 		LegacySlab.UseMetalness		= false;
