Track whether light grid and it's cells contain any rect or textured lights.
- store per-view flags in FViewInfo.
- store per-cell flags in FCulledLightsGridHeader.
- will be used to speed up tile classification in Mega Lights.

diff --git a/Engine/Shaders/Private/LightGridCommon.ush b/Engine/Shaders/Private/LightGridCommon.ush
--- a/Engine/Shaders/Private/LightGridCommon.ush
+++ b/Engine/Shaders/Private/LightGridCommon.ush
@@ -126,12 +126,37 @@ uint GetMaxLightsPerCell(uint EyeIndex)
 #endif
 }
 
+uint PackCulledLightsGridHeader0(uint NumVisibleLights, uint NumVisibleMegaLights)
+{
+	return (NumVisibleMegaLights << 16) | (NumVisibleLights & 0xFFFF);
+}
+
+uint PackCulledLightsGridHeader1(uint CulledLightDataStart, bool bHasRectLight, bool bHasTexturedLight)
+{
+	return (bHasTexturedLight << 31) | (bHasRectLight << 30) | (CulledLightDataStart & 0x3FFFFFFF);
+}
+
+void UnpackCulledLightsGridHeader0(uint PackedData0, out uint NumVisibleLights, out uint NumVisibleMegaLights)
+{
+	NumVisibleLights = PackedData0 & 0xFFFF;
+	NumVisibleMegaLights = (PackedData0 >> 16) & 0xFFFF;
+}
+
+void UnpackCulledLightsGridHeader1(uint PackedData1, out uint CulledLightDataStart, out bool bHasRectLight, out bool bHasTexturedLight)
+{
+	CulledLightDataStart = (PackedData1 & 0x3FFFFFFF);
+	bHasRectLight = (PackedData1 & 0x40000000) != 0;
+	bHasTexturedLight = (PackedData1 & 0x80000000) != 0;
+}
+
 struct FCulledLightsGridHeader
 {
 	uint NumLights;
 	uint NumMegaLights;
 	uint DataStartIndex;
 	uint MegaLightsDataStartIndex;
+	bool bHasRectLight;
+	bool bHasTexturedLight;
 };
 
 FCulledLightsGridHeader GetCulledLightsGridHeader(uint GridIndex, uint EyeIndex)
@@ -144,20 +169,28 @@ FCulledLightsGridHeader GetCulledLightsGridHeader(uint GridIndex, uint EyeIndex)
 	{
 #endif
 
-		const uint PackedNums = ForwardLightData.NumCulledLightsGrid[GridIndex * NUM_CULLED_LIGHTS_GRID_STRIDE + 0];
-		Result.NumLights = min(PackedNums & 0xFFFF, ForwardLightData.NumLocalLights);
-		Result.NumMegaLights = min((PackedNums >> 16) & 0xFFFF, ForwardLightData.NumLocalLights);
-		Result.DataStartIndex = ForwardLightData.NumCulledLightsGrid[GridIndex * NUM_CULLED_LIGHTS_GRID_STRIDE + 1];
+		const uint PackedData0 = ForwardLightData.NumCulledLightsGrid[GridIndex * NUM_CULLED_LIGHTS_GRID_STRIDE + 0];
+		UnpackCulledLightsGridHeader0(PackedData0, Result.NumLights, Result.NumMegaLights);
+		Result.NumLights = min(Result.NumLights, ForwardLightData.NumLocalLights);
+		Result.NumMegaLights = min(Result.NumMegaLights, ForwardLightData.NumLocalLights);
+
+		const uint PackedData1 = ForwardLightData.NumCulledLightsGrid[GridIndex * NUM_CULLED_LIGHTS_GRID_STRIDE + 1];
+		UnpackCulledLightsGridHeader1(PackedData1, Result.DataStartIndex, Result.bHasRectLight, Result.bHasTexturedLight);
+		
 		Result.MegaLightsDataStartIndex = Result.DataStartIndex + Result.NumLights - Result.NumMegaLights;
 
 #if (INSTANCED_STEREO || MOBILE_MULTI_VIEW)
 	}
 	else
 	{
-		const uint PackedNums = ForwardLightDataISR.NumCulledLightsGrid[GridIndex * NUM_CULLED_LIGHTS_GRID_STRIDE + 0];
-		Result.NumLights = min(PackedNums & 0xFFFF, ForwardLightDataISR.NumLocalLights);
-		Result.NumMegaLights = min((PackedNums >> 16) & 0xFFFF, ForwardLightDataISR.NumLocalLights);
-		Result.DataStartIndex = ForwardLightDataISR.NumCulledLightsGrid[GridIndex * NUM_CULLED_LIGHTS_GRID_STRIDE + 1];
+		const uint PackedData0 = ForwardLightDataISR.NumCulledLightsGrid[GridIndex * NUM_CULLED_LIGHTS_GRID_STRIDE + 0];
+		UnpackCulledLightsGridHeader0(PackedData0, Result.NumLights, Result.NumMegaLights);
+		Result.NumLights = min(Result.NumLights, ForwardLightDataISR.NumLocalLights);
+		Result.NumMegaLights = min(Result.NumMegaLights, ForwardLightDataISR.NumLocalLights);
+
+		const uint PackedData1 = ForwardLightDataISR.NumCulledLightsGrid[GridIndex * NUM_CULLED_LIGHTS_GRID_STRIDE + 1];
+		UnpackCulledLightsGridHeader1(PackedData1, Result.DataStartIndex, Result.bHasRectLight, Result.bHasTexturedLight);
+		
 		Result.MegaLightsDataStartIndex = Result.DataStartIndex + Result.NumLights - Result.NumMegaLights;
 	}
 #endif
diff --git a/Engine/Shaders/Private/LightGridInjection.usf b/Engine/Shaders/Private/LightGridInjection.usf
--- a/Engine/Shaders/Private/LightGridInjection.usf
+++ b/Engine/Shaders/Private/LightGridInjection.usf
@@ -197,6 +197,8 @@ void LightGridInjectionCS(
 
 		uint NumVisibleLights = 0;
 		uint NumVisibleMegaLights = 0;
+		uint bHasRectLights = false;
+		uint bHasTexturedLights = false;
 
 		LOOP
 		for (uint LocalLightIndex = 0; LocalLightIndex < NumLocalLights; LocalLightIndex++)
@@ -213,11 +215,12 @@ void LightGridInjectionCS(
 			{
 				float4 ViewSpaceDirAndPreprocAngle = LightViewSpaceDirAndPreprocAngle[LocalLightIndex];
 				bool bIsRectLight = asuint(ViewSpaceDirAndPreprocAngle.w) & 0x1;
+				bool bIsTexturedLight = asuint(ViewSpaceDirAndPreprocAngle.w) & 0x2;
 
 				bool bPassSpotLightTest = true;
 #if REFINE_SPOTLIGHT_BOUNDS
 				{
-					float TanConeAngle = asfloat(asuint(ViewSpaceDirAndPreprocAngle.w) & 0xFFFFFFFE);
+					float TanConeAngle = asfloat(asuint(ViewSpaceDirAndPreprocAngle.w) & 0xFFFFFFFC);
 
 					// Set to 0 for non-acute cones, or non-spot lights.
 					if (TanConeAngle > 0.0f)
@@ -243,6 +246,8 @@ void LightGridInjectionCS(
 #endif // REFINE_RECTLIGHT_BOUNDS
 				if (bPassSpotLightTest && bPassRectLightTest)
 				{
+					bHasRectLights |= bIsRectLight;
+					bHasTexturedLights |= bIsTexturedLight;
 
 #if USE_LINKED_CULL_LIST
 					uint NextLink;
@@ -272,8 +277,8 @@ void LightGridInjectionCS(
 			}
 		}
 
-		RWNumCulledLightsGrid[GridIndex * NUM_CULLED_LIGHTS_GRID_STRIDE + 0] = (NumVisibleMegaLights << 16) | NumVisibleLights;
-		RWNumCulledLightsGrid[GridIndex * NUM_CULLED_LIGHTS_GRID_STRIDE + 1] = GridIndex * MaxCulledLightsPerCell;
+		RWNumCulledLightsGrid[GridIndex * NUM_CULLED_LIGHTS_GRID_STRIDE + 0] = PackCulledLightsGridHeader0(NumVisibleLights, NumVisibleMegaLights);
+		RWNumCulledLightsGrid[GridIndex * NUM_CULLED_LIGHTS_GRID_STRIDE + 1] = PackCulledLightsGridHeader1(GridIndex * MaxCulledLightsPerCell, bHasRectLights, bHasTexturedLights);
 		
 		LOOP
 		for (uint ReflectionCaptureIndex = 0; ReflectionCaptureIndex < NumReflectionCaptures; ReflectionCaptureIndex++)
@@ -322,8 +327,8 @@ void LightGridInjectionCS(
 			}
 		}
 
-		RWNumCulledLightsGrid[GridIndex * NUM_CULLED_LIGHTS_GRID_STRIDE + 0] = (NumLocalLights - MegaLightsSupportedStartIndex) | NumLocalLights;
-		RWNumCulledLightsGrid[GridIndex * NUM_CULLED_LIGHTS_GRID_STRIDE + 1] = GridIndex * MaxCulledLightsPerCell;
+		RWNumCulledLightsGrid[GridIndex * NUM_CULLED_LIGHTS_GRID_STRIDE + 0] = PackCulledLightsGridHeader0(NumLocalLights, NumLocalLights - MegaLightsSupportedStartIndex);
+		RWNumCulledLightsGrid[GridIndex * NUM_CULLED_LIGHTS_GRID_STRIDE + 1] = PackCulledLightsGridHeader1(GridIndex * MaxCulledLightsPerCell, true, true); // TODO
 		
 		LOOP
 		for (uint ReflectionCaptureIndex = 0; ReflectionCaptureIndex < NumReflectionCaptures; ReflectionCaptureIndex++)
@@ -424,8 +429,14 @@ void CompactReverseLinkedList(uint GridIndex, uint SceneMax, bool FirstThread, b
 
 		if (bLocalLights)
 		{
-			RWNumCulledLightsGrid[GridIndex * NUM_CULLED_LIGHTS_GRID_STRIDE + 0] = (NumMegaLights << 16) | NumCulledLights;
-			RWNumCulledLightsGrid[GridIndex * NUM_CULLED_LIGHTS_GRID_STRIDE + 1] = CulledLightDataStart;
+			// Get bHasRectLights / bHasTexturedLights written by LightGridInjectionCS before updating CulledLightDataStart with packed offset
+			uint Unused;
+			bool bHasRectLights;
+			bool bHasTexturedLights;
+			UnpackCulledLightsGridHeader1(RWNumCulledLightsGrid[GridIndex * NUM_CULLED_LIGHTS_GRID_STRIDE + 1], Unused, bHasRectLights, bHasTexturedLights);
+
+			RWNumCulledLightsGrid[GridIndex * NUM_CULLED_LIGHTS_GRID_STRIDE + 0] = PackCulledLightsGridHeader0(NumCulledLights, NumMegaLights);
+			RWNumCulledLightsGrid[GridIndex * NUM_CULLED_LIGHTS_GRID_STRIDE + 1] = PackCulledLightsGridHeader1(CulledLightDataStart, bHasRectLights, bHasTexturedLights);
 		}
 		else
 		{
diff --git a/Engine/Source/Runtime/Renderer/Private/LightGridInjection.cpp b/Engine/Source/Runtime/Renderer/Private/LightGridInjection.cpp
--- a/Engine/Source/Runtime/Renderer/Private/LightGridInjection.cpp
+++ b/Engine/Source/Runtime/Renderer/Private/LightGridInjection.cpp
@@ -490,6 +490,9 @@ FComputeLightGridOutput FSceneRenderer::ComputeLightGrid(FRDGBuilder& GraphBuild
 		int32 ClusteredSupportedEnd = 0;
 		int32 MegaLightsSupportedStart = 0;
 
+		bool bHasRectLights = false;
+		bool bHasTexturedLights = false;
+
 		const float Exposure = View.GetLastEyeAdaptationExposure();
 
 		if (bCullLightsToGrid)
@@ -631,9 +634,12 @@ FComputeLightGridOutput FSceneRenderer::ComputeLightGrid(FRDGBuilder& GraphBuild
 
 						// Pack flags in the LSB of PreProcAngle
 						const float PreProcAngle = SortedLightInfo.SortKey.Fields.LightType == LightType_Spot ? GetTanRadAngleOrZero(LightSceneInfo->Proxy->GetOuterConeAngle()) : 0.0f;
-						const uint32 PackedPreProcAngleAndFlags = (FMath::AsUInt(PreProcAngle) & 0xFFFFFFFE) | (LightProxy->IsRectLight() ? 0x1 : 0);
+						const uint32 PackedPreProcAngleAndFlags = (FMath::AsUInt(PreProcAngle) & 0xFFFFFFFC) | (LightProxy->HasSourceTexture() ? 0x2 : 0) | (LightProxy->IsRectLight() ? 0x1 : 0);
 						FVector4f ViewSpaceDirAndPreprocAngleAndFlags(FVector4f(View.ViewMatrices.GetViewMatrix().TransformVector((FVector)LightParameters.Direction)), FMath::AsFloat(PackedPreProcAngleAndFlags)); // LWC_TODO: precision loss
 						ViewSpaceDirAndPreprocAngleData.Add(ViewSpaceDirAndPreprocAngleAndFlags);
+
+						bHasRectLights |= LightProxy->IsRectLight();
+						bHasTexturedLights |= LightProxy->HasSourceTexture();
 					}
 					// On mobile there is a separate FMobileDirectionalLightShaderParameters UB which holds all directional light data.
 					else if (SortedLightInfo.SortKey.Fields.LightType == LightType_Directional && ViewFamily.EngineShowFlags.DirectionalLights && !IsMobilePlatform(View.GetShaderPlatform()))
@@ -769,6 +775,9 @@ FComputeLightGridOutput FSceneRenderer::ComputeLightGrid(FRDGBuilder& GraphBuild
 
 		View.ForwardLightingResources.LocalLightVisibleLightInfosIndex = LocalLightVisibleLightInfosIndex;
 
+		View.bLightGridHasRectLights = bHasRectLights;
+		View.bLightGridHasTexturedLights = bHasTexturedLights;
+
 		const FIntPoint LightGridSizeXY = FIntPoint::DivideAndRoundUp(View.ViewRect.Size(), GLightGridPixelSize);
 		ForwardLightData->ForwardLocalLightBuffer = GraphBuilder.CreateSRV(FRDGBufferSRVDesc(ForwardLocalLightBuffer));
 		ForwardLightData->NumLocalLights = NumLocalLightsFinal;
diff --git a/Engine/Source/Runtime/Renderer/Private/SceneRendering.h b/Engine/Source/Runtime/Renderer/Private/SceneRendering.h
--- a/Engine/Source/Runtime/Renderer/Private/SceneRendering.h
+++ b/Engine/Source/Runtime/Renderer/Private/SceneRendering.h
@@ -1426,6 +1426,9 @@ public:
 	FForwardLightingViewResources ForwardLightingResources;
 	FVolumetricFogViewResources VolumetricFogResources;
 
+	bool bLightGridHasRectLights = false;
+	bool bLightGridHasTexturedLights = false;
+
 	FRDGTextureRef HeterogeneousVolumeRadiance = nullptr;
 	FRDGTextureRef HeterogeneousVolumeHoldout = nullptr;
 
