Fix incorrect logic in UpsampleShadowFactors(...)

diff --git a/Engine/Shaders/Private/ShadowFactorsUpsampleCommon.ush b/Engine/Shaders/Private/ShadowFactorsUpsampleCommon.ush
--- a/Engine/Shaders/Private/ShadowFactorsUpsampleCommon.ush
+++ b/Engine/Shaders/Private/ShadowFactorsUpsampleCommon.ush
@@ -50,7 +50,7 @@ void UpsampleShadowFactors(
 	float2 TextureValues00 = Texture2DSampleLevel(ShadowFactorsTexture, ShadowFactorsSampler, min(Corner00SampleUV, ShadowFactorsUVBilinearMax), 0).xy;
 	float2 TextureValues10 = Texture2DSampleLevel(ShadowFactorsTexture, ShadowFactorsSampler, min(Corner00SampleUV + float2(LowResTexelSize.x, 0), ShadowFactorsUVBilinearMax), 0).xy;
 	float2 TextureValues01 = Texture2DSampleLevel(ShadowFactorsTexture, ShadowFactorsSampler, min(Corner00SampleUV + float2(0, LowResTexelSize.y), ShadowFactorsUVBilinearMax), 0).xy;
-	float2 TextureValues11 = Texture2DSampleLevel(ShadowFactorsTexture, ShadowFactorsSampler, min(Corner00SampleUV + LowResTexelSize, 0), ShadowFactorsUVBilinearMax).xy;
+	float2 TextureValues11 = Texture2DSampleLevel(ShadowFactorsTexture, ShadowFactorsSampler, min(Corner00SampleUV + LowResTexelSize, ShadowFactorsUVBilinearMax), 0).xy;
 
 	float4 CornerWeights = float4(
 		(1 - BilinearWeights.y) * (1 - BilinearWeights.x), 
