Fix editor HDR. Cooked HDR/SDR wasn't impacted since the output of ColorLookupTable is already in the correct output format, but editor is not

#rnx

diff --git a/Engine/Shaders/Private/PostProcessTonemap.usf b/Engine/Shaders/Private/PostProcessTonemap.usf
--- a/Engine/Shaders/Private/PostProcessTonemap.usf
+++ b/Engine/Shaders/Private/PostProcessTonemap.usf
@@ -563,7 +563,7 @@ float4 TonemapCommonPS(
 		{
 			OutColor.xyz = ST2084ToScRGB(OutColor.xyz, GetOutputDevice(), OutputMaxLuminance);
 		}
-		if(GetOutputDevice() == TONEMAPPER_OUTPUT_LinearEXR)
+		else if(GetOutputDevice() == TONEMAPPER_OUTPUT_LinearEXR)
 		{
 			OutColor.xyz = ST2084ToLinear(OutColor.xyz);
 		}
