Modified PackCulledLightsGridHeader1(...) to not use bit shift on bool type.

diff --git a/Engine/Shaders/Private/LightGridCommon.ush b/Engine/Shaders/Private/LightGridCommon.ush
--- a/Engine/Shaders/Private/LightGridCommon.ush
+++ b/Engine/Shaders/Private/LightGridCommon.ush
@@ -133,7 +133,7 @@ uint PackCulledLightsGridHeader0(uint NumVisibleLights, uint NumVisibleMegaLight
 
 uint PackCulledLightsGridHeader1(uint CulledLightDataStart, bool bHasRectLight, bool bHasTexturedLight)
 {
-	return (bHasTexturedLight << 31) | (bHasRectLight << 30) | (CulledLightDataStart & 0x3FFFFFFF);
+	return (bHasTexturedLight ? 0x80000000 : 0) | (bHasRectLight ? 0x40000000 : 0) | (CulledLightDataStart & 0x3FFFFFFF);
 }
 
 void UnpackCulledLightsGridHeader0(uint PackedData0, out uint NumVisibleLights, out uint NumVisibleMegaLights)
