Fix warning of Output SV_Target0 not fully written.

diff --git a/Engine/Shaders/Private/ComposeSeparateTranslucency.usf b/Engine/Shaders/Private/ComposeSeparateTranslucency.usf
--- a/Engine/Shaders/Private/ComposeSeparateTranslucency.usf
+++ b/Engine/Shaders/Private/ComposeSeparateTranslucency.usf
@@ -133,6 +133,7 @@ void CopyBackgroundVisibilityPS(
 	OutColor.a = 1.0f;
 #elif VISIBILITY_COPY_TYPE == VISIBILITY_COPY_TO_SCENECOLOR
 	float4 BackgroundVisibilityWithAlpha = TranslucentHoldoutPointTexture.Load(uint3(PixelPos, 0));
+	OutColor.rgb = 0.0f;
 	OutColor.a =  saturate(BackgroundVisibilityWithAlpha.r);// We can copy any rgb channel.
 #endif 
 }
