Fix missing translucent in Lumen's reflection with Substrate.

diff --git a/Engine/Shaders/Private/RayTracing/RayTracingMaterialHitShaders.usf b/Engine/Shaders/Private/RayTracing/RayTracingMaterialHitShaders.usf
--- a/Engine/Shaders/Private/RayTracing/RayTracingMaterialHitShaders.usf
+++ b/Engine/Shaders/Private/RayTracing/RayTracingMaterialHitShaders.usf
@@ -13,9 +13,9 @@
 // This should be good enough for ray tracing and avoids having to bind an extra buffer
 #define EYE_ADAPTATION_PREV_FRAME_EXPOSURE 1
 
-
 #if SUBSTRATE_ENABLED
 #define SUBSTRATE_INLINE_SHADING 1
+#define SUBSTRATE_USES_FORCED_COMPLEXITY 0
 #endif
 
 #include "/Engine/Private/Common.ush"
diff --git a/Engine/Shaders/Private/Substrate/Substrate.ush b/Engine/Shaders/Private/Substrate/Substrate.ush
--- a/Engine/Shaders/Private/Substrate/Substrate.ush
+++ b/Engine/Shaders/Private/Substrate/Substrate.ush
@@ -103,6 +103,11 @@
 #define SUBSTRATE_SINGLEPATH 0
 #endif
 
+// Allow SUBSTRATE_FASTPATH and SUBSTRATE_SINGLEPATH to drive the actual material's complexity, rather than using the material configuration/header bits
+#ifndef SUBSTRATE_USES_FORCED_COMPLEXITY
+#define SUBSTRATE_USES_FORCED_COMPLEXITY 1
+#endif
+
 #define SUBSTRATE_COMPLEXPATH (!SUBSTRATE_FASTPATH && !SUBSTRATE_SINGLEPATH)
 
 // SUBSTRATE_COMPLEXSPECIALPATH is defined by the material or lighting shaders as they see fit.
@@ -2142,9 +2147,15 @@ struct FSubstratePixelHeader
 
 	// Material mode & BSDF type
 	uint GetMaterialMode()										{ return (this.State & HEADER_MASK_MATERIALMODE); }
+	#if SUBSTRATE_USES_FORCED_COMPLEXITY
 	bool IsSimpleMaterial()										{ return (this.State & HEADER_MASK_MATERIALMODE) == HEADER_MATERIALMODE_SLAB_SIMPLE || SUBSTRATE_FASTPATH; }
 	bool IsSingleMaterial()										{ return (this.State & HEADER_MASK_MATERIALMODE) == HEADER_MATERIALMODE_SLAB_SINGLE || SUBSTRATE_SINGLEPATH; }
 	bool IsComplexMaterial()									{ return (this.State & HEADER_MASK_MATERIALMODE) == HEADER_MATERIALMODE_SLAB_COMPLEX && !SUBSTRATE_FASTPATH && !SUBSTRATE_SINGLEPATH; }
+	#else
+	bool IsSimpleMaterial()										{ return (this.State & HEADER_MASK_MATERIALMODE) == HEADER_MATERIALMODE_SLAB_SIMPLE; }
+	bool IsSingleMaterial()										{ return (this.State & HEADER_MASK_MATERIALMODE) == HEADER_MATERIALMODE_SLAB_SINGLE; }
+	bool IsComplexMaterial()									{ return (this.State & HEADER_MASK_MATERIALMODE) == HEADER_MATERIALMODE_SLAB_COMPLEX; }
+	#endif
 	bool IsComplexSpecialMaterial()								{ return (this.State & HEADER_COMPLEXSPECIALPATH_MASK) != 0; }
 	bool IsSingleLayerWater()									{ return (this.State & HEADER_MASK_MATERIALMODE) == HEADER_MATERIALMODE_SLWATER; }
 	bool IsHair()												{ return (this.State & HEADER_MASK_MATERIALMODE) == HEADER_MATERIALMODE_HAIR; }
