Alpha: New (deferred) renderer project setting for primitive alpha holdout support, separating it regular propagate alpha which now becomes runtime-toggleable.

POST_PROCESS_ALPHA is completely removed from the engine, and the holdout-specific shader code is conditionally compiled based on a new SUPPORT_PRIMITIVE_ALPHA_HOLDOUT definition to isolate its added cost in the deferred renderer.

The new holdout project setting is called "Support Primitive Alpha Holdout (Deferred)" and can be found in under the "Optimizations" section of the renderer settings. Engine component, MRQ & HoldoutComposite plugin toast notifications for setting activation is left as separate CL.

diff --git a/Engine/Plugins/MovieScene/MovieRenderPipeline/Source/MovieRenderPipelineCore/Public/Graph/MovieGraphRenderLayerSubsystem.h b/Engine/Plugins/MovieScene/MovieRenderPipeline/Source/MovieRenderPipelineCore/Public/Graph/MovieGraphRenderLayerSubsystem.h
--- a/Engine/Plugins/MovieScene/MovieRenderPipeline/Source/MovieRenderPipelineCore/Public/Graph/MovieGraphRenderLayerSubsystem.h
+++ b/Engine/Plugins/MovieScene/MovieRenderPipeline/Source/MovieRenderPipelineCore/Public/Graph/MovieGraphRenderLayerSubsystem.h
@@ -877,7 +877,7 @@ public:
 
 	/**
 	 * If true, the primitive will render black with an alpha of 0, but all secondary effects (shadows, reflections,
-	 * indirect lighting) remain. This feature is currently only implemented in the Path Tracer.
+	 * indirect lighting) remain. This feature requires activating the project setting(s) "Enable Alpha Output", and "Support Primitive Alpha Holdout" if using the deferred renderer.
 	 */
 	UPROPERTY(BlueprintReadWrite, EditAnywhere, Category = "Settings", meta = (EditCondition = "bOverride_bHoldout"))
 	uint8 bHoldout : 1;
diff --git a/Engine/Shaders/Private/BasePassPixelShader.usf b/Engine/Shaders/Private/BasePassPixelShader.usf
--- a/Engine/Shaders/Private/BasePassPixelShader.usf
+++ b/Engine/Shaders/Private/BasePassPixelShader.usf
@@ -2180,7 +2180,7 @@ void FPixelShaderInOut_MainPS(
 	Out.MRT[0] = SubstrateMaterialExport;
 #endif
 
-#if POST_PROCESS_ALPHA != 0 && MATERIAL_DOMAIN_SURFACE
+#if SUPPORT_PRIMITIVE_ALPHA_HOLDOUT && MATERIAL_DOMAIN_SURFACE
 	const bool bIsHoldout = GetPrimitiveData(MaterialParameters).Flags & PRIMITIVE_SCENE_DATA_FLAG_HOLDOUT;
 	// Only include the holdout logic if the Renderer is forwarding its alpha channel to post-processing
 	if (bIsHoldout)
@@ -2221,9 +2221,7 @@ void FPixelShaderInOut_MainPS(
 		}
 	}
 
-#endif
-
-#if (MATERIALBLENDING_ANY_TRANSLUCENT || MATERIALBLENDING_ALPHACOMPOSITE || MATERIALBLENDING_ALPHAHOLDOUT) && MATERIAL_DOMAIN_SURFACE && POST_PROCESS_ALPHA!= 0
+#if MATERIALBLENDING_ANY_TRANSLUCENT || MATERIALBLENDING_ALPHACOMPOSITE || MATERIALBLENDING_ALPHAHOLDOUT
 	// If this is translucency holdout pass, calculate the background visibility and path throughput.
 	if (IsTranslucentHoldoutPass &&
         ResolvedView.RenderingReflectionCaptureMask == 0)
@@ -2286,6 +2284,7 @@ void FPixelShaderInOut_MainPS(
 		
 	}
 #endif
+#endif // SUPPORT_PRIMITIVE_ALPHA_HOLDOUT && MATERIAL_DOMAIN_SURFACE
 
 }
 
diff --git a/Engine/Shaders/Private/CompositeUIPixelShader.usf b/Engine/Shaders/Private/CompositeUIPixelShader.usf
--- a/Engine/Shaders/Private/CompositeUIPixelShader.usf
+++ b/Engine/Shaders/Private/CompositeUIPixelShader.usf
@@ -153,7 +153,7 @@ void Main(
 
 	OutColor.rgb = ApplyColorDeficiency(OutColor.rgb);
 
-	// This is not compatible when alpha is preserved through post-processing with POST_PROCESS_ALPHA
+	// This is not compatible when alpha is preserved through post-processing
 	OutColor.w = UIColor.w;
 
 #if !SCRGB_ENCODING
@@ -207,7 +207,7 @@ void CompositeUICS(uint3 DispatchThreadId : SV_DispatchThreadID)
 	float4 OutColor = 0;
 	OutColor.xyz = ComposeUIAndScene(SceneColor.xyz, UIColor, UILevel);
 	
-	// This is not compatible when alpha is preserved through post-processing with POST_PROCESS_ALPHA
+	// This is not compatible when alpha is preserved through post-processing
 	OutColor.w = UIColor.w;
 
 	OutColor.rgb = ApplyColorDeficiency(OutColor.rgb);
diff --git a/Engine/Shaders/Private/HairStrands/HairStrandsComposition.usf b/Engine/Shaders/Private/HairStrands/HairStrandsComposition.usf
--- a/Engine/Shaders/Private/HairStrands/HairStrandsComposition.usf
+++ b/Engine/Shaders/Private/HairStrands/HairStrandsComposition.usf
@@ -109,7 +109,7 @@ void ComposeSamplePS(
 	LOOP
 	for (uint SampleIt = 0; SampleIt < NodeDesc.Count; SampleIt++)
 	{
-		#if POST_PROCESS_ALPHA
+		#if SUPPORT_PRIMITIVE_ALPHA_HOLDOUT
 		if (bHasHoldout)
 		{
 			const FHairSample Sample = UnpackHairSample(HairStrands.HairSampleData[NodeDesc.Offset + SampleIt]);
@@ -181,7 +181,7 @@ void HoldoutPS(
 	out float4 OutColor : SV_Target0,
 	out float  OutDepth : SV_DEPTH)
 {
-#if POST_PROCESS_ALPHA
+#if SUPPORT_PRIMITIVE_ALPHA_HOLDOUT
 	OutColor = 0.0f;
 	const uint3 PixelCoord = uint3(floor(Input.Position.xy), 0);
 	
diff --git a/Engine/Shaders/Private/HeightFogCommon.ush b/Engine/Shaders/Private/HeightFogCommon.ush
--- a/Engine/Shaders/Private/HeightFogCommon.ush
+++ b/Engine/Shaders/Private/HeightFogCommon.ush
@@ -375,7 +375,7 @@ half4 GetExponentialHeightFog(float3 WorldPositionRelativeToCamera, float Exclud
 		half3 FogColor = (InscatteringColor) * (1 - ExpFogFactor) + DirectionalInscattering;
 	#endif
 
-#if POST_PROCESS_ALPHA
+#if SUPPORT_PRIMITIVE_ALPHA_HOLDOUT
 	if (IsExponentialFogHoldout(View.EnvironmentComponentsFlags) && View.RenderingReflectionCaptureMask == 0.0f)
 	{
 		FogColor.rgb *= 0;
@@ -455,7 +455,7 @@ float4 CombineVolumetricFog(float4 GlobalFog, float3 VolumeUV, uint EyeIndex, fl
 	// Visualize depth distribution
 	//VolumetricFogLookup.rgb += .1f * frac(min(ZSlice, 1.0f) / View.VolumetricFogInvGridSize.z);
 
-#if POST_PROCESS_ALPHA
+#if SUPPORT_PRIMITIVE_ALPHA_HOLDOUT
 	if (IsExponentialFogHoldout(InView.EnvironmentComponentsFlags) && InView.RenderingReflectionCaptureMask == 0.0f)
 	{
 		VolumetricFogLookup.rgb *= 0;
diff --git a/Engine/Shaders/Private/HeightFogPixelShader.usf b/Engine/Shaders/Private/HeightFogPixelShader.usf
--- a/Engine/Shaders/Private/HeightFogPixelShader.usf
+++ b/Engine/Shaders/Private/HeightFogPixelShader.usf
@@ -18,6 +18,7 @@ float UpsampleJitterMultiplier;
 
 float bOnlyOnRenderedOpaque;
 uint bUseWaterDepthTexture;
+uint bPropagateAlpha;
 
 float4 WaterDepthTextureMinMaxUV;
 float4 OcclusionTextureMinMaxUV;
@@ -128,12 +129,16 @@ void ExponentialPixelMain(
 		FogInscatteringAndOpacity.a = 1;
 	}
 	
-#if POST_PROCESS_ALPHA
-	// We need to output coverage so that we can potentially apply alpha hold out blending. See RenderViewFog.
-	OutColor = RETURN_COLOR(float4(FogInscatteringAndOpacity.rgb, 1.0 - FogInscatteringAndOpacity.w));
-#else
-	OutColor = RETURN_COLOR(float4(FogInscatteringAndOpacity.rgb, FogInscatteringAndOpacity.w));
-#endif
+	FLATTEN
+	if (bPropagateAlpha)
+	{
+		// We need to output coverage so that we can potentially apply alpha hold out blending. See RenderViewFog.
+		OutColor = RETURN_COLOR(float4(FogInscatteringAndOpacity.rgb, 1.0 - FogInscatteringAndOpacity.w));
+	}
+	else
+	{
+		OutColor = RETURN_COLOR(float4(FogInscatteringAndOpacity.rgb, FogInscatteringAndOpacity.w));
+	}
 	
 #endif // PERMUTATION_SAMPLE_FOG_ON_CLOUDS
 }
diff --git a/Engine/Shaders/Private/HeterogeneousVolumes/HeterogeneousVolumesComposite.usf b/Engine/Shaders/Private/HeterogeneousVolumes/HeterogeneousVolumesComposite.usf
--- a/Engine/Shaders/Private/HeterogeneousVolumes/HeterogeneousVolumesComposite.usf
+++ b/Engine/Shaders/Private/HeterogeneousVolumes/HeterogeneousVolumesComposite.usf
@@ -23,7 +23,7 @@ void HeterogeneousVolumesCompositeCS(
 	float4 Result;
 	float HeterogeneousVolumeTransmittance = HeterogeneousVolumeRadiance[PixelCoord].a;
 	Result.rgb = HeterogeneousVolumeRadiance[PixelCoord].rgb + RWColorTexture[ViewPixelCoord].rgb * HeterogeneousVolumeTransmittance;
-#if POST_PROCESS_ALPHA
+#if SUPPORT_PRIMITIVE_ALPHA_HOLDOUT
 	if (View.RenderingReflectionCaptureMask == 0.0f)
 	{
 		float SurfaceOpacity = RWColorTexture[ViewPixelCoord].a * HeterogeneousVolumeTransmittance;
diff --git a/Engine/Shaders/Private/HeterogeneousVolumes/HeterogeneousVolumesLiveShadingPipeline.usf b/Engine/Shaders/Private/HeterogeneousVolumes/HeterogeneousVolumesLiveShadingPipeline.usf
--- a/Engine/Shaders/Private/HeterogeneousVolumes/HeterogeneousVolumesLiveShadingPipeline.usf
+++ b/Engine/Shaders/Private/HeterogeneousVolumes/HeterogeneousVolumesLiveShadingPipeline.usf
@@ -597,7 +597,7 @@ void RenderSingleScatteringWithLiveShading(
 #endif
 
 		// Radiance
-#if POST_PROCESS_ALPHA
+#if SUPPORT_PRIMITIVE_ALPHA_HOLDOUT
 		if (bHoldout && View.RenderingReflectionCaptureMask == 0.0f)
 		{
 			Radiance *= 0;
@@ -610,7 +610,7 @@ void RenderSingleScatteringWithLiveShading(
 			float Opacity = Luminance(1.0 - Transmittance);
 			RWLightingTexture[PixelCoord].a = Luminance(Transmittance);
 
-#if POST_PROCESS_ALPHA
+#if SUPPORT_PRIMITIVE_ALPHA_HOLDOUT
 			if (bHoldout && View.RenderingReflectionCaptureMask == 0.0f)
 			{
 				RWHoldoutTexture[PixelCoord] += Luminance(PrevTransmittance) - Luminance(Transmittance);
diff --git a/Engine/Shaders/Private/MaterialTemplate.ush b/Engine/Shaders/Private/MaterialTemplate.ush
--- a/Engine/Shaders/Private/MaterialTemplate.ush
+++ b/Engine/Shaders/Private/MaterialTemplate.ush
@@ -2267,7 +2267,7 @@ float3 MaterialExpressionSkyAtmosphereViewLuminance(FMaterialPixelParameters Par
 	SkyAtmosphereViewLuminance *= ResolvedView.SkyAtmosphereSkyLuminanceFactor.rgb;
 	SkyAtmosphereViewLuminance *= ResolvedView.OneOverPreExposure;
 	
-#if POST_PROCESS_ALPHA
+#if SUPPORT_PRIMITIVE_ALPHA_HOLDOUT
 	if (IsSkyAtmosphereHoldout(View.EnvironmentComponentsFlags) && View.RenderingReflectionCaptureMask == 0.0f)
 	{
 		SkyAtmosphereViewLuminance = 0.0f; // Disable sky dome luminance contribution. Throughput will be 0.
@@ -2303,7 +2303,7 @@ float4 MaterialExpressionSkyAtmosphereAerialPerspective(FMaterialPixelParameters
 		ResolvedView.OneOverPreExposure,
 		NearFadeOutRangeInvDepthKm);
 	
-#if POST_PROCESS_ALPHA
+#if SUPPORT_PRIMITIVE_ALPHA_HOLDOUT
 	if (IsSkyAtmosphereHoldout(View.EnvironmentComponentsFlags) && View.RenderingReflectionCaptureMask == 0.0f)
 	{
 		AerialPerspective.rgb = 0.0f;
diff --git a/Engine/Shaders/Private/SkyAtmosphere.usf b/Engine/Shaders/Private/SkyAtmosphere.usf
--- a/Engine/Shaders/Private/SkyAtmosphere.usf
+++ b/Engine/Shaders/Private/SkyAtmosphere.usf
@@ -124,6 +124,8 @@ Texture2DMS<float, MSAA_SAMPLE_COUNT> MSAADepthTexture;
 // Only available for shaders ran per view, so this constant should not be accessed in common code such as IntegrateSingleScatteredLuminance for instance.
 float AerialPerspectiveStartDepthKm;
 
+// Propagate alpha with (View.RenderingReflectionCaptureMask == 0.0f) guarantee
+uint bPropagateAlphaNonReflection;
 
 // - RayOrigin: ray origin
 // - RayDir: normalized ray direction
@@ -829,13 +831,14 @@ float4 PrepareOutput(float3 PreExposedLuminance, float3 Transmittance = float3(1
 	const float GreyScaleTransmittance = dot(Transmittance, float3(1.0f / 3.0f, 1.0f / 3.0f, 1.0f / 3.0f));
 	float4 LuminanceAlpha = float4(SafePreExposedLuminance, GreyScaleTransmittance);
 	
-#if POST_PROCESS_ALPHA
-	if (View.RenderingReflectionCaptureMask == 0.0f)
+	FLATTEN
+	if(bPropagateAlphaNonReflection > 0)
 	{
+#if SUPPORT_PRIMITIVE_ALPHA_HOLDOUT
 		LuminanceAlpha.rgb	= IsSkyAtmosphereHoldout(View.EnvironmentComponentsFlags) ? 0.0f : LuminanceAlpha.rgb;
+#endif
 		LuminanceAlpha.a	= 1 - GreyScaleTransmittance;
 	}
-#endif
 	
 	return LuminanceAlpha;
 
@@ -845,17 +848,18 @@ uint DepthReadDisabled;
 
 void UpdateVisibleSkyAlpha(in float DeviceZ, inout float4 OutLuminance)
 {
-#if POST_PROCESS_ALPHA
-	if (View.RenderingReflectionCaptureMask == 0.0f)
+	// This is for parity with the sky dome mesh behavior.
+	// Output fully opaque alpha if looking through outer space, if atmosphere is not in holdout.
+	bool bOutputFullyOpaque = (bPropagateAlphaNonReflection > 0) && (DeviceZ == FarDepthValue);
+#if SUPPORT_PRIMITIVE_ALPHA_HOLDOUT
+	bOutputFullyOpaque &= !IsSkyAtmosphereHoldout(View.EnvironmentComponentsFlags);	
+#endif
+
+	FLATTEN
+	if (bOutputFullyOpaque)
 	{
-		// This is for parity with the sky dome mesh behavior.
-		// Output fully opaque alpha if looking through outer space, if atmosphere is not in holdout.
-		if (DeviceZ == FarDepthValue && !IsSkyAtmosphereHoldout(View.EnvironmentComponentsFlags))
-		{
-			OutLuminance.a = 1.0;
-		}
+		OutLuminance.a = 1.0;
 	}
-#endif
 }
 
 void RenderSkyAtmosphereRayMarchingPS(
@@ -974,12 +978,11 @@ void RenderSkyAtmosphereRayMarchingPS(
 		float3 SkyLuminance = SkyLuminanceTransmittance.rgb;
 
 		float3 SkyGreyTransmittance = 1.0f;
-	#if POST_PROCESS_ALPHA
-		if (View.RenderingReflectionCaptureMask == 0.0f)
+		FLATTEN
+		if(bPropagateAlphaNonReflection > 0)
 		{
 			SkyGreyTransmittance = SkyLuminanceTransmittance.aaa;
 		}
-	#endif
 
 		PreExposedL += SkyLuminance * LuminanceScale * (ViewOneOverPreExposure * OutputPreExposure);
 	
@@ -1617,7 +1620,7 @@ void RenderCameraAerialPerspectiveVolumeCS(uint3 ThreadId : SV_DispatchThreadID)
 		AerialPespectiveViewDistanceScale,
 		tMaxMax);
 
-#if POST_PROCESS_ALPHA
+#if SUPPORT_PRIMITIVE_ALPHA_HOLDOUT
 	if (IsSkyAtmosphereHoldout(View.EnvironmentComponentsFlags) && !RealTimeReflection360Mode)
 	{
 		ss.L *= 0;
diff --git a/Engine/Shaders/Private/VolumetricCloud.usf b/Engine/Shaders/Private/VolumetricCloud.usf
--- a/Engine/Shaders/Private/VolumetricCloud.usf
+++ b/Engine/Shaders/Private/VolumetricCloud.usf
@@ -1458,7 +1458,7 @@ void MainCommon(in FMaterialPixelParameters MaterialParameters, in float4 SvPosi
 	// When rendering a real time reflection capture (sky envmap) whe use a different exposure
 	const float OutputPreExposure = (ResolvedView.RealTimeReflectionCapture ? ResolvedView.RealTimeReflectionCapturePreExposure : ResolvedView.PreExposure);
 
-#if POST_PROCESS_ALPHA
+#if SUPPORT_PRIMITIVE_ALPHA_HOLDOUT
 	if (IsVolumetricCloudHoldout(View.EnvironmentComponentsFlags) && View.RenderingReflectionCaptureMask == 0.0f)
 	{
 		Luminance *= 0;
@@ -1515,7 +1515,7 @@ void MainCommon(in FMaterialPixelParameters MaterialParameters, in float4 SvPosi
 		float CloudMeanCoverage = 1.0 - MeanTransmittance;
 		Luminance = AerialPerspective.rgb * CloudMeanCoverage + AerialPerspective.a * Luminance;
 
-#if POST_PROCESS_ALPHA
+#if SUPPORT_PRIMITIVE_ALPHA_HOLDOUT
 		if(View.RenderingReflectionCaptureMask == 0.0f)
 		{
 			const float AerialPespectiveTransmittance = AerialPerspective.a;
@@ -1573,7 +1573,7 @@ void MainCommon(in FMaterialPixelParameters MaterialParameters, in float4 SvPosi
 		float CloudMeanCoverage = 1.0 - MeanTransmittance;
 		Luminance = HeightFogInscatteringAndTransmittance.rgb * CloudMeanCoverage + HeightFogInscatteringAndTransmittance.a * Luminance;
 
-#if POST_PROCESS_ALPHA
+#if SUPPORT_PRIMITIVE_ALPHA_HOLDOUT
 		if(View.RenderingReflectionCaptureMask == 0.0f)
 		{
 			const float ExponentialFogTransmittance = HeightFogInscatteringAndTransmittance.a;
@@ -1779,7 +1779,7 @@ void MainCS(uint2 DispatchThreadId : SV_DispatchThreadID)
 			if (bBlendCloudColor)
 			{
 				float4 CurColor0 = OutCloudColor0[DispatchThreadId];
-			#if POST_PROCESS_ALPHA
+			#if SUPPORT_PRIMITIVE_ALPHA_HOLDOUT
 				if (View.RenderingReflectionCaptureMask == 0.0f)
 				{
 					if(IsVolumetricCloudHoldout(View.EnvironmentComponentsFlags) && bAccumulateAlphaHoldOut > 0)
diff --git a/Engine/Shaders/Public/ShaderVersion.ush b/Engine/Shaders/Public/ShaderVersion.ush
--- a/Engine/Shaders/Public/ShaderVersion.ush
+++ b/Engine/Shaders/Public/ShaderVersion.ush
@@ -3,4 +3,4 @@
 // Each time the console command is executed it generates a new GUID. As this file is included
 // in Platform.ush (which should be included in any shader) it allows to invalidate the shader DDC.
 // 
-#pragma message("UESHADERMETADATA_VERSION 10C02C57-1251-46A8-A27D-B5C80CD940EF")
\ No newline at end of file
+#pragma message("UESHADERMETADATA_VERSION 1B9EB545-4184-F9A2-094F-A0AA4BD39EC4")
\ No newline at end of file
diff --git a/Engine/Source/Runtime/Engine/Classes/Components/ExponentialHeightFogComponent.h b/Engine/Source/Runtime/Engine/Classes/Components/ExponentialHeightFogComponent.h
--- a/Engine/Source/Runtime/Engine/Classes/Components/ExponentialHeightFogComponent.h
+++ b/Engine/Source/Runtime/Engine/Classes/Components/ExponentialHeightFogComponent.h
@@ -213,7 +213,7 @@ class UExponentialHeightFogComponent : public USceneComponent
 	UPROPERTY(EditAnywhere, BlueprintReadOnly, Category = VolumetricFog, AdvancedDisplay)
 	bool bOverrideLightColorsWithFogInscatteringColors;
 
-	/** If this is True, this primitive will render black with an alpha of 0, but all secondary effects (shadows, reflections, indirect lighting) remain. This feature required the project setting "Enable Alpha Output". */
+	/** If this is True, this primitive will render black with an alpha of 0, but all secondary effects (shadows, reflections, indirect lighting) remain. This feature requires activating the project setting(s) "Enable Alpha Output", and "Support Primitive Alpha Holdout" if using the deferred renderer. */
 	UPROPERTY(EditAnywhere, AdvancedDisplay, BlueprintReadOnly, Category = Rendering, Interp)
 	uint8 bHoldout : 1;
 
diff --git a/Engine/Source/Runtime/Engine/Classes/Components/PrimitiveComponent.h b/Engine/Source/Runtime/Engine/Classes/Components/PrimitiveComponent.h
--- a/Engine/Source/Runtime/Engine/Classes/Components/PrimitiveComponent.h
+++ b/Engine/Source/Runtime/Engine/Classes/Components/PrimitiveComponent.h
@@ -428,7 +428,7 @@ public:
 	UPROPERTY(EditAnywhere, AdvancedDisplay, BlueprintReadOnly, Category=Rendering)
 	uint8 bReceivesDecals:1;
 
-	/** If this is True, this primitive will render black with an alpha of 0, but all secondary effects (shadows, reflections, indirect lighting) remain. This feature required the project setting "Enable Alpha Output". */
+	/** If this is True, this primitive will render black with an alpha of 0, but all secondary effects (shadows, reflections, indirect lighting) remain. This feature requires activating the project setting(s) "Enable Alpha Output", and "Support Primitive Alpha Holdout" if using the deferred renderer. */
 	UPROPERTY(EditAnywhere, AdvancedDisplay, BlueprintReadOnly, Category = Rendering, Interp)
 	uint8 bHoldout : 1;
 
diff --git a/Engine/Source/Runtime/Engine/Classes/Components/SkyAtmosphereComponent.h b/Engine/Source/Runtime/Engine/Classes/Components/SkyAtmosphereComponent.h
--- a/Engine/Source/Runtime/Engine/Classes/Components/SkyAtmosphereComponent.h
+++ b/Engine/Source/Runtime/Engine/Classes/Components/SkyAtmosphereComponent.h
@@ -160,7 +160,7 @@ class USkyAtmosphereComponent : public USceneComponent
 	UPROPERTY(EditAnywhere, BlueprintReadOnly, interp, Category = "Art Direction", meta = (UIMin = 0.001f, UIMax = 10.0f, ClampMin = 0.001f))
 	float AerialPerspectiveStartDepth;
 
-	/** If this is True, this primitive will render black with an alpha of 0, but all secondary effects (shadows, reflections, indirect lighting) remain. This feature required the project setting "Enable Alpha Output". */
+	/** If this is True, this primitive will render black with an alpha of 0, but all secondary effects (shadows, reflections, indirect lighting) remain. This feature requires activating the project setting(s) "Enable Alpha Output", and "Support Primitive Alpha Holdout" if using the deferred renderer. */
 	UPROPERTY(EditAnywhere, AdvancedDisplay, BlueprintReadOnly, Category = Rendering, Interp)
 	uint8 bHoldout : 1;
 
diff --git a/Engine/Source/Runtime/Engine/Classes/Components/VolumetricCloudComponent.h b/Engine/Source/Runtime/Engine/Classes/Components/VolumetricCloudComponent.h
--- a/Engine/Source/Runtime/Engine/Classes/Components/VolumetricCloudComponent.h
+++ b/Engine/Source/Runtime/Engine/Classes/Components/VolumetricCloudComponent.h
@@ -139,7 +139,7 @@ class UVolumetricCloudComponent : public USceneComponent
 	UPROPERTY(EditAnywhere, BlueprintReadOnly, interp, Category = "Art Direction", meta = (UIMin = 0.0, UIMax = 100.0, ClampMin = 0.0, SliderExponent = 2.0))
 	float AerialPespectiveMieScatteringFadeDistance;
 
-	/** If this is True, this primitive will render black with an alpha of 0, but all secondary effects (shadows, reflections, indirect lighting) remain. This feature required the project setting "Enable Alpha Output". */
+	/** If this is True, this primitive will render black with an alpha of 0, but all secondary effects (shadows, reflections, indirect lighting) remain. This feature requires activating the project setting(s) "Enable Alpha Output", and "Support Primitive Alpha Holdout" if using the deferred renderer. */
 	UPROPERTY(EditAnywhere, AdvancedDisplay, BlueprintReadOnly, Category = Rendering, Interp)
 	uint8 bHoldout : 1;
 
diff --git a/Engine/Source/Runtime/Engine/Classes/Engine/RendererSettings.h b/Engine/Source/Runtime/Engine/Classes/Engine/RendererSettings.h
--- a/Engine/Source/Runtime/Engine/Classes/Engine/RendererSettings.h
+++ b/Engine/Source/Runtime/Engine/Classes/Engine/RendererSettings.h
@@ -742,9 +742,8 @@ class URendererSettings : public UDeveloperSettings
 
 	UPROPERTY(config, EditAnywhere, Category = Postprocessing, meta = (
 		ConsoleVariable = "r.PostProcessing.PropagateAlpha", DisplayName = "Enable Alpha Output",
-		ToolTip = "Enable r.PostProcessing.PropagateAlpha to enforce alpha in scene color (overriding r.SceneColorFormat if necessary) and propagate it through the renderer's post-processing chain.  The legacy \"Linear color space only\" and \"Allow through tonemapper\" options now map to true and the engine keeps full alpha precision throughout post-processing.",
-		ConfigRestartRequired = true))
-	bool bEnableAlphaChannelInPostProcessing;
+		ToolTip = "Enable r.PostProcessing.PropagateAlpha to enforce alpha in scene color (overriding r.SceneColorFormat if necessary) and propagate it through the renderer's post-processing chain. The legacy \"Linear color space only\" and \"Allow through tonemapper\" options now map to True and the engine keeps full alpha precision throughout post-processing. This feature can now be toggled without an engine restart."))
+	uint32 bEnableAlphaChannelInPostProcessing : 1;
 
 	UPROPERTY(config, EditAnywhere, Category = DefaultSettings, meta = (
 		ConsoleVariable = "r.DefaultFeature.Bloom", DisplayName = "Bloom",
@@ -997,6 +996,12 @@ class URendererSettings : public UDeveloperSettings
 		ConfigRestartRequired = true))
 		uint32 bSupportCloudShadowOnForwardLitTranslucent : 1;
 
+	UPROPERTY(config, EditAnywhere, Category = Optimizations, meta = (
+		ConsoleVariable = "r.Deferred.SupportPrimitiveAlphaHoldout", DisplayName = "Support Primitive Alpha Holdout (Deferred)",
+		ToolTip = "Enable primitive alpha holdout support in multiple deferred renderer passes. If primitive holdout masks are not to be used, keep this setting disabled for increased performance.",
+		ConfigRestartRequired = true))
+	uint32 bDeferredSupportPrimitiveAlphaHoldout : 1;
+
 	/**
 	"Select the format of the light function atlas texture."
 	*/
diff --git a/Engine/Source/Runtime/Engine/Private/ShaderCompiler/ShaderCompiler.cpp b/Engine/Source/Runtime/Engine/Private/ShaderCompiler/ShaderCompiler.cpp
--- a/Engine/Source/Runtime/Engine/Private/ShaderCompiler/ShaderCompiler.cpp
+++ b/Engine/Source/Runtime/Engine/Private/ShaderCompiler/ShaderCompiler.cpp
@@ -3820,16 +3820,10 @@ void GlobalBeginCompileShader(
 	}
 
 	{
-		static IConsoleVariable* CVar = IConsoleManager::Get().FindConsoleVariable(TEXT("r.PostProcessing.PropagateAlpha"));
-		bool bPropagateAlpha = CVar->GetBool();
+		static IConsoleVariable* CVar = IConsoleManager::Get().FindConsoleVariable(TEXT("r.Deferred.SupportPrimitiveAlphaHoldout"));
+		const bool bDeferredSupportPrimitiveAlphaHoldout = CVar->GetBool();
 
-		if (bIsMobilePlatform)
-		{
-			static FShaderPlatformCachedIniValue<int32> MobilePropagateAlphaIniValue(TEXT("r.Mobile.PropagateAlpha"));
-			bPropagateAlpha = MobilePropagateAlphaIniValue.Get((EShaderPlatform)ShaderPlatform) > 0;
-		}
-
-		SET_SHADER_DEFINE(Input.Environment, POST_PROCESS_ALPHA, bPropagateAlpha ? 1 : 0);
+		SET_SHADER_DEFINE(Input.Environment, SUPPORT_PRIMITIVE_ALPHA_HOLDOUT, bDeferredSupportPrimitiveAlphaHoldout ? 1 : 0);
 	}
 
 	if (TargetPlatform && 
diff --git a/Engine/Source/Runtime/Engine/Private/ShaderCompiler/ShaderGenerationUtil.cpp b/Engine/Source/Runtime/Engine/Private/ShaderCompiler/ShaderGenerationUtil.cpp
--- a/Engine/Source/Runtime/Engine/Private/ShaderCompiler/ShaderGenerationUtil.cpp
+++ b/Engine/Source/Runtime/Engine/Private/ShaderCompiler/ShaderGenerationUtil.cpp
@@ -72,7 +72,7 @@ void ApplyFetchEnvironmentInternal(FShaderGlobalDefines& SrcDefines, const Conta
 	FETCH_COMPILE_BOOL(PROJECT_EXPFOG_MATCHES_VFOG);
 	FETCH_COMPILE_BOOL(SUPPORT_CLOUD_SHADOW_ON_FORWARD_LIT_TRANSLUCENT);
 	FETCH_COMPILE_BOOL(SUPPORT_CLOUD_SHADOW_ON_SINGLE_LAYER_WATER);
-	FETCH_COMPILE_BOOL(POST_PROCESS_ALPHA);
+	FETCH_COMPILE_BOOL(SUPPORT_PRIMITIVE_ALPHA_HOLDOUT);
 	FETCH_COMPILE_BOOL(PLATFORM_SUPPORTS_SHADER_ROOT_CONSTANTS);
 	FETCH_COMPILE_BOOL(PLATFORM_SUPPORTS_SHADER_BUNDLE_DISPATCH);
 	FETCH_COMPILE_BOOL(PLATFORM_SUPPORTS_RENDERTARGET_WRITE_MASK);
@@ -434,8 +434,8 @@ static FShaderGlobalDefines FetchShaderGlobalDefines(EShaderPlatform TargetPlatf
 	}
 
 	{
-		static IConsoleVariable* CVar = IConsoleManager::Get().FindConsoleVariable(TEXT("r.PostProcessing.PropagateAlpha"));
-		Ret.POST_PROCESS_ALPHA = CVar->GetBool();
+		static IConsoleVariable* CVar = IConsoleManager::Get().FindConsoleVariable(TEXT("r.Deferred.SupportPrimitiveAlphaHoldout"));
+		Ret.SUPPORT_PRIMITIVE_ALPHA_HOLDOUT = CVar->GetBool();
 	}
 
 	Ret.PLATFORM_SUPPORTS_SHADER_ROOT_CONSTANTS = RHISupportsShaderRootConstants(EShaderPlatform(TargetPlatform)) ? 1 : 0;
diff --git a/Engine/Source/Runtime/Engine/Public/PrimitiveSceneProxy.h b/Engine/Source/Runtime/Engine/Public/PrimitiveSceneProxy.h
--- a/Engine/Source/Runtime/Engine/Public/PrimitiveSceneProxy.h
+++ b/Engine/Source/Runtime/Engine/Public/PrimitiveSceneProxy.h
@@ -1457,7 +1457,7 @@ protected:
 
 	/**
 	 * If this is True, this primitive should render black with an alpha of 0, but all secondary effects (shadows, refletions, indirect lighting)
-	 * should behave as usual. This feature is currently only implemented in the Path Tracer.
+	 * should behave as usual. This feature requires activating the project setting(s) "Enable Alpha Output", and "Support Primitive Alpha Holdout" if using the deferred renderer.
 	 */
 	uint8 bHoldout : 1;
 
diff --git a/Engine/Source/Runtime/Landscape/Classes/LandscapeProxy.h b/Engine/Source/Runtime/Landscape/Classes/LandscapeProxy.h
--- a/Engine/Source/Runtime/Landscape/Classes/LandscapeProxy.h
+++ b/Engine/Source/Runtime/Landscape/Classes/LandscapeProxy.h
@@ -795,7 +795,7 @@ public:
 	UPROPERTY(EditAnywhere, AdvancedDisplay, Category = Lighting, meta = (LandscapeOverridable, ClampMin = "0.0", ClampMax = "10.0", UIMin = "0.0", UIMax = "10.0"))
 	float NonNaniteVirtualShadowMapInvalidationScreenSizeLimit = 0.2f;
 
-	/** If this is True, this primitive will render black with an alpha of 0, but all secondary effects (shadows, reflections, indirect lighting) remain. This feature required the project setting "Enable Alpha Output". */
+	/** If this is True, this primitive will render black with an alpha of 0, but all secondary effects (shadows, reflections, indirect lighting) remain. This feature requires the project settings "Enable Alpha Output" and "Support Primitive Alpha Holdout". */
 	UPROPERTY(EditAnywhere, AdvancedDisplay, Category = Rendering, meta=(LandscapeOverridable), Interp)
 	uint8 bHoldout : 1;
 
diff --git a/Engine/Source/Runtime/RenderCore/Private/Shader.cpp b/Engine/Source/Runtime/RenderCore/Private/Shader.cpp
--- a/Engine/Source/Runtime/RenderCore/Private/Shader.cpp
+++ b/Engine/Source/Runtime/RenderCore/Private/Shader.cpp
@@ -2192,25 +2192,13 @@ void ShaderMapAppendKey(EShaderPlatform Platform, FShaderKeyGenerator& KeyGen)
 	}
 
 	{
-		bool bPropagateAlpha = false;
-		if (IsMobilePlatform(Platform))
-		{
-			static FShaderPlatformCachedIniValue<int32> MobilePropagateAlphaIniValue(TEXT("r.Mobile.PropagateAlpha"));
-			bPropagateAlpha = MobilePropagateAlphaIniValue.Get((EShaderPlatform)Platform) > 0;
-		}
-		else
-		{
-			static const auto CVarPropagateAlpha = IConsoleManager::Get().FindConsoleVariable(TEXT("r.PostProcessing.PropagateAlpha"));
-			if (CVarPropagateAlpha)
-			{
-				bPropagateAlpha = CVarPropagateAlpha->GetBool();
-			}
-		}
+		static const auto CVar = IConsoleManager::Get().FindConsoleVariable(TEXT("r.Deferred.SupportPrimitiveAlphaHoldout"));
+		const bool bDeferredSupportPrimitiveAlphaHoldout = CVar ? CVar->GetBool() : false;
 
-		if (bPropagateAlpha)
+		if (bDeferredSupportPrimitiveAlphaHoldout)
 		{
 			KeyGen.AppendSeparator();
-			KeyGen.Append(TEXT("PA"));
+			KeyGen.Append(TEXT("PAH"));
 		}
 	}
 
diff --git a/Engine/Source/Runtime/RenderCore/Public/ShaderMaterial.h b/Engine/Source/Runtime/RenderCore/Public/ShaderMaterial.h
--- a/Engine/Source/Runtime/RenderCore/Public/ShaderMaterial.h
+++ b/Engine/Source/Runtime/RenderCore/Public/ShaderMaterial.h
@@ -36,7 +36,7 @@ struct FShaderGlobalDefines
 	bool PROJECT_EXPFOG_MATCHES_VFOG;
 	bool SUPPORT_CLOUD_SHADOW_ON_FORWARD_LIT_TRANSLUCENT;
 	bool SUPPORT_CLOUD_SHADOW_ON_SINGLE_LAYER_WATER;
-	bool POST_PROCESS_ALPHA;
+	bool SUPPORT_PRIMITIVE_ALPHA_HOLDOUT;
 	bool PLATFORM_SUPPORTS_SHADER_ROOT_CONSTANTS;
 	bool PLATFORM_SUPPORTS_SHADER_BUNDLE_DISPATCH;
 	bool PLATFORM_SUPPORTS_RENDERTARGET_WRITE_MASK;
diff --git a/Engine/Source/Runtime/Renderer/Private/FogRendering.cpp b/Engine/Source/Runtime/Renderer/Private/FogRendering.cpp
--- a/Engine/Source/Runtime/Renderer/Private/FogRendering.cpp
+++ b/Engine/Source/Runtime/Renderer/Private/FogRendering.cpp
@@ -146,6 +146,7 @@ class FExponentialHeightFogPS : public FGlobalShader
 		SHADER_PARAMETER_SAMPLER(SamplerState, WaterDepthSampler)
 		SHADER_PARAMETER(float, bOnlyOnRenderedOpaque)
 		SHADER_PARAMETER(uint32, bUseWaterDepthTexture)
+		SHADER_PARAMETER(uint32, bPropagateAlpha)
 		SHADER_PARAMETER(float, UpsampleJitterMultiplier)
 		SHADER_PARAMETER(FVector4f, WaterDepthTextureMinMaxUV)
 		SHADER_PARAMETER(FVector4f, OcclusionTextureMinMaxUV)
@@ -318,6 +319,7 @@ static FFogPassParameters* CreateDefaultFogPassParameters(
 	PassParameters->PS.UpsampleJitterMultiplier = CVarUpsampleJitterMultiplier.GetValueOnRenderThread() * GVolumetricFogGridPixelSize;
 	PassParameters->PS.bOnlyOnRenderedOpaque = View.bFogOnlyOnRenderedOpaque;
 	PassParameters->PS.bUseWaterDepthTexture = false;
+	PassParameters->PS.bPropagateAlpha = IsPostProcessingWithAlphaChannelSupported();
 
 	PassParameters->PS.SrcCloudDepthTexture = GSystemTextures.GetWhiteDummy(GraphBuilder);
 	PassParameters->PS.SrcCloudDepthSampler = TStaticSamplerState<SF_Point, AM_Clamp, AM_Clamp, AM_Clamp>::GetRHI();
diff --git a/Engine/Source/Runtime/Renderer/Private/PostProcess/PostProcessing.cpp b/Engine/Source/Runtime/Renderer/Private/PostProcess/PostProcessing.cpp
--- a/Engine/Source/Runtime/Renderer/Private/PostProcess/PostProcessing.cpp
+++ b/Engine/Source/Runtime/Renderer/Private/PostProcess/PostProcessing.cpp
@@ -118,7 +118,7 @@ FAutoConsoleVariableRef CVarPostProcessingPropagateAlpha(
 	TEXT("r.PostProcessing.PropagateAlpha"),
 	GPostProcessingPropagateAlpha,
 	TEXT("Enforce alpha in scene color (overriding r.SceneColorFormat if necessary) and propagate it through the renderer's post-processing chain, default: false"),
-	ECVF_ReadOnly | ECVF_RenderThreadSafe);
+	ECVF_RenderThreadSafe);
 
 TAutoConsoleVariable<int32> CVarPostProcessingPreferCompute(
 	TEXT("r.PostProcessing.PreferCompute"),
diff --git a/Engine/Source/Runtime/Renderer/Private/SceneRendering.cpp b/Engine/Source/Runtime/Renderer/Private/SceneRendering.cpp
--- a/Engine/Source/Runtime/Renderer/Private/SceneRendering.cpp
+++ b/Engine/Source/Runtime/Renderer/Private/SceneRendering.cpp
@@ -342,6 +342,12 @@ static TAutoConsoleVariable<int32> CVarPrimitiveHasTileOffsetData(
 	TEXT("1 to upload lower-precision tileoffset primitive data to gpu, 0 to use higher-precision double float.\n"),
 	ECVF_ReadOnly | ECVF_RenderThreadSafe);
 
+static TAutoConsoleVariable<bool> CVarPrimitiveAlphaHoldoutSupport(
+	TEXT("r.Deferred.SupportPrimitiveAlphaHoldout"),
+	false,
+	TEXT("True to enable deferred renderer support for primitive alpha holdout (disabled by default).\n"),
+	ECVF_ReadOnly | ECVF_RenderThreadSafe);
+
 #if !(UE_BUILD_SHIPPING || UE_BUILD_TEST)
 static TAutoConsoleVariable<float> CVarGeneralPurposeTweak(
 	TEXT("r.GeneralPurposeTweak"),
diff --git a/Engine/Source/Runtime/Renderer/Private/SkyAtmosphereRendering.cpp b/Engine/Source/Runtime/Renderer/Private/SkyAtmosphereRendering.cpp
--- a/Engine/Source/Runtime/Renderer/Private/SkyAtmosphereRendering.cpp
+++ b/Engine/Source/Runtime/Renderer/Private/SkyAtmosphereRendering.cpp
@@ -697,6 +697,7 @@ class FRenderSkyAtmospherePS : public FGlobalShader
 		SHADER_PARAMETER_SAMPLER(SamplerState, VolumetricCloudShadowMapTexture1Sampler)
 		SHADER_PARAMETER_SAMPLER(SamplerState, VolumetricCloudSkyAOTextureSampler)
 		SHADER_PARAMETER(float, AerialPerspectiveStartDepthKm)
+		SHADER_PARAMETER(uint32, bPropagateAlphaNonReflection)
 		SHADER_PARAMETER(uint32, SourceDiskEnabled)
 		SHADER_PARAMETER(uint32, DepthReadDisabled)
 		SHADER_PARAMETER_RDG_UNIFORM_BUFFER(FVolumeShadowingShaderParametersGlobal0, Light0Shadow)
@@ -1868,6 +1869,8 @@ void FSceneRenderer::RenderSkyAtmosphereInternal(
 		PsPassParameters->VolumetricCloudDepthTexture = SkyRC.VolumetricCloudDepthTexture;
 		PsPassParameters->InputCloudLuminanceTransmittanceTexture = SkyRC.InputCloudLuminanceTransmittanceTexture;
 		PsPassParameters->AerialPerspectiveStartDepthKm = AerialPerspectiveStartDepthInCm * CM_TO_KM;
+		// NOTE: bLightDiskEnabled == !View.bIsReflectionCapture
+		PsPassParameters->bPropagateAlphaNonReflection = (IsPostProcessingWithAlphaChannelSupported() && SkyRC.bLightDiskEnabled) ? 1 : 0;
 		PsPassParameters->SourceDiskEnabled = SkyRC.bLightDiskEnabled ? 1 : 0;
 		PsPassParameters->DepthReadDisabled = SkyRC.bDepthReadDisabled ? 1 : 0;
 		if (bShouldSampleCloudShadow || SkyRC.bShouldSampleCloudSkyAO)
