Fix minor shader warnings

#rnx

diff --git a/Engine/Shaders/Private/GBufferHelpers.ush b/Engine/Shaders/Private/GBufferHelpers.ush
--- a/Engine/Shaders/Private/GBufferHelpers.ush
+++ b/Engine/Shaders/Private/GBufferHelpers.ush
@@ -355,7 +355,7 @@ void AdjustBaseColorAndSpecularColorForSubsurfaceProfileLightingCopyHack(inout f
 	// The translucent pass is called after the main post process subsurface scattering pass, it will not affect previous logic.
 #if MATERIALBLENDING_TRANSLUCENT || !ALLOW_SSS_MATERIAL_OVERRIDE
 	return;
-#endif
+#else
 
 #if SUBSURFACE_CHANNEL_MODE == 0
 	// If SUBSURFACE_CHANNEL_MODE is 0, we can't support full-resolution lighting, so we 
@@ -375,6 +375,7 @@ void AdjustBaseColorAndSpecularColorForSubsurfaceProfileLightingCopyHack(inout f
 		SpecularColor *= !bChecker;
 		Specular *= !bChecker;
 	}
+#endif
 }
 
 #ifndef INDIRECT_SAMPLE_COUNT
diff --git a/Engine/Shaders/Private/GPUScene/GPUSceneWriter.ush b/Engine/Shaders/Private/GPUScene/GPUSceneWriter.ush
--- a/Engine/Shaders/Private/GPUScene/GPUSceneWriter.ush
+++ b/Engine/Shaders/Private/GPUScene/GPUSceneWriter.ush
@@ -246,7 +246,7 @@ void WriteInstanceDynamicDataWS(uint InstanceId, uint PrimitiveId, FDFMatrix Pre
 void WriteInstanceDynamicDataLS(FInstancePayloadDataOffsets PayloadOffsets, uint PrimitiveId, float4x4 PrevLocalToPrimitive)
 {
 	FDFMatrix PrevLocalToWorld = DFMultiply(PrevLocalToPrimitive, GetPrimitiveData(PrimitiveId).PreviousLocalToWorld);
-	return WriteInstanceDynamicDataWS(PayloadOffsets, PrimitiveId, PrevLocalToWorld);
+	WriteInstanceDynamicDataWS(PayloadOffsets, PrimitiveId, PrevLocalToWorld);
 }
 
 void WriteInstanceDynamicDataLS(uint InstanceId, uint PrimitiveId, float4x4 PrevLocalToPrimitive)
diff --git a/Engine/Shaders/Private/VirtualShadowMaps/VirtualShadowMapProjection.usf b/Engine/Shaders/Private/VirtualShadowMaps/VirtualShadowMapProjection.usf
--- a/Engine/Shaders/Private/VirtualShadowMaps/VirtualShadowMapProjection.usf
+++ b/Engine/Shaders/Private/VirtualShadowMaps/VirtualShadowMapProjection.usf
@@ -559,7 +559,7 @@ void OutputVisualize(
 			if (CacheColor.a > 0)
 			{
 				Output.rgb = CacheColor.a*CacheColor.rgb + VPageAlpha*vPageColor + ShadowAlpha*Result.ShadowFactor.xxx;
-				Output = float4(Output + Accumulator.rgb, 0.9f);
+				Output = float4(Output.rgb + Accumulator.rgb, 0.9f);
 			}
 			else
 			{
diff --git a/Engine/Shaders/Public/Platform.ush b/Engine/Shaders/Public/Platform.ush
--- a/Engine/Shaders/Public/Platform.ush
+++ b/Engine/Shaders/Public/Platform.ush
@@ -299,7 +299,8 @@
 	_Pragma("dxc diagnostic ignored \"-Wall\"") 				\
 	_Pragma("dxc diagnostic ignored \"-Wconversion\"") 			\
 	_Pragma("dxc diagnostic ignored \"-Wfor-redefinition\"") 	\
-	_Pragma("dxc diagnostic ignored \"-Winline-asm\"")
+	_Pragma("dxc diagnostic ignored \"-Winline-asm\"")			\
+	_Pragma("dxc diagnostic ignored \"-Wunsequenced\"")
 
 #endif
 
