Few minor misc optimizations:
- Skip translucent volume clear if disabled
- Skip VSM OPP pass if there are no local lights
- Track range of VSM lights during first pruning pass

#rnx

diff --git a/Engine/Shaders/Private/VirtualShadowMaps/VirtualShadowMapPageMarking.usf b/Engine/Shaders/Private/VirtualShadowMaps/VirtualShadowMapPageMarking.usf
--- a/Engine/Shaders/Private/VirtualShadowMaps/VirtualShadowMapPageMarking.usf
+++ b/Engine/Shaders/Private/VirtualShadowMaps/VirtualShadowMapPageMarking.usf
@@ -158,17 +158,25 @@ void PruneLightGridCS(uint GridLinearIndex : SV_DispatchThreadID)
 	const FCulledLightsGridHeader CulledLightGridHeader = GetCulledLightsGridHeader(GridLinearIndex, EyeIndex);
 
 	uint NumRetainedLights = 0U;
+	uint FirstLightWithVSM = CulledLightGridHeader.NumLights;
+	uint LastLightWithVSM = 0U;
 
 	// First pass add the non-distant lights with VSMs
 	LOOP
 	for (uint Index = 0; Index < CulledLightGridHeader.NumLights; ++Index)
 	{
 		int VirtualShadowMapId = GetLocalLightData(CulledLightGridHeader.DataStartIndex + Index, EyeIndex).VirtualShadowMapId;
-		if (VirtualShadowMapId != INDEX_NONE && !IsSinglePageVirtualShadowMap(VirtualShadowMapId))
+		if (VirtualShadowMapId != INDEX_NONE)
 		{
-			// Copy light grid data index
-			checkStructuredBufferAccessSlow(OutPrunedLightGridData, CulledLightGridHeader.DataStartIndex + NumRetainedLights);
-			OutPrunedLightGridData[CulledLightGridHeader.DataStartIndex + NumRetainedLights++] = GetCulledLightDataGrid(CulledLightGridHeader.DataStartIndex + Index);
+			FirstLightWithVSM = min(FirstLightWithVSM, Index);
+			LastLightWithVSM = max(LastLightWithVSM, Index);
+
+			if (!IsSinglePageVirtualShadowMap(VirtualShadowMapId))
+			{
+				// Copy light grid data index
+				checkStructuredBufferAccessSlow(OutPrunedLightGridData, CulledLightGridHeader.DataStartIndex + NumRetainedLights);
+				OutPrunedLightGridData[CulledLightGridHeader.DataStartIndex + NumRetainedLights++] = GetCulledLightDataGrid(CulledLightGridHeader.DataStartIndex + Index);
+			}
 		}
 	}
 
@@ -176,7 +184,7 @@ void PruneLightGridCS(uint GridLinearIndex : SV_DispatchThreadID)
 	// NOTE: We could add these to the end in the first pass and then re-pack them instead, but this works
 	// well enough for moderate light counts.
 	LOOP
-	for (uint Index = 0; Index < CulledLightGridHeader.NumLights; ++Index)
+	for (uint Index = FirstLightWithVSM; Index <= LastLightWithVSM; ++Index)
 	{
 		int VirtualShadowMapId = GetLocalLightData(CulledLightGridHeader.DataStartIndex + Index, EyeIndex).VirtualShadowMapId;
 		// Discard any light without a VSM
diff --git a/Engine/Source/Runtime/Renderer/Private/Shadows/ShadowSceneRenderer.cpp b/Engine/Source/Runtime/Renderer/Private/Shadows/ShadowSceneRenderer.cpp
--- a/Engine/Source/Runtime/Renderer/Private/Shadows/ShadowSceneRenderer.cpp
+++ b/Engine/Source/Runtime/Renderer/Private/Shadows/ShadowSceneRenderer.cpp
@@ -481,12 +481,12 @@ void FShadowSceneRenderer::RenderVirtualShadowMapProjectionMaskBits(
 	FRDGBuilder& GraphBuilder,
 	FMinimalSceneTextures& SceneTextures)
 {
-	// VSM one pass projection (done first as it may be needed by clustered shading)
 	bShouldUseVirtualShadowMapOnePassProjection =
 		VirtualShadowMapArray.IsAllocated() &&
 		IsVSMOnePassProjectionEnabled(SceneRenderer.ViewFamily.EngineShowFlags);
 
-	if (!VirtualShadowMapArray.HasAnyShadowData())
+	// One pass projection is only for local lights, so not needed if we don't have any
+	if (!VirtualShadowMapArray.HasAnyShadowData() || LocalLights.Num() == 0)
 	{
 		return;
 	}
diff --git a/Engine/Source/Runtime/Renderer/Private/TranslucentLighting.cpp b/Engine/Source/Runtime/Renderer/Private/TranslucentLighting.cpp
--- a/Engine/Source/Runtime/Renderer/Private/TranslucentLighting.cpp
+++ b/Engine/Source/Runtime/Renderer/Private/TranslucentLighting.cpp
@@ -861,6 +861,13 @@ int32 FTranslucencyLightingVolumeTextures::GetIndex(const FViewInfo& View, int32
 
 void FTranslucencyLightingVolumeTextures::Init(FRDGBuilder& GraphBuilder, TArrayView<const FViewInfo> Views, ERDGPassFlags PassFlags)
 {
+	// Skip init/clear if disabled
+	// GetTranslucencyLightingVolumeParameters will return black system textures
+	if (!GUseTranslucentLightingVolumes)
+	{
+		return;
+	}
+
 	check(PassFlags == ERDGPassFlags::Compute || PassFlags == ERDGPassFlags::AsyncCompute);
 
 	RDG_GPU_STAT_SCOPE(GraphBuilder, TranslucentLighting);
