Fix a shader compiling issue if a traslucent material enable PreIntegratedGF in the ForwardShading.

diff --git a/Engine/Shaders/Private/BasePassPixelShader.usf b/Engine/Shaders/Private/BasePassPixelShader.usf
--- a/Engine/Shaders/Private/BasePassPixelShader.usf
+++ b/Engine/Shaders/Private/BasePassPixelShader.usf
@@ -41,16 +41,16 @@
 	#define UseBasePassSkylight OpaqueBasePass.Shared.UseBasePassSkylight
 	#define SubstrateStruct OpaqueBasePass.Substrate
 	#define LightFunctionAtlasStruct OpaqueBasePass.Shared.LightFunctionAtlas
-#endif
 
-// Material setting to allow forward shading (including mobile) to use preintegrated GF lut for simple IBL.
-#if MATERIAL_SHADINGMODEL_SINGLELAYERWATER || (FORWARD_SHADING && MATERIAL_USE_PREINTEGRATED_GF)
-	#define PreIntegratedGF			OpaqueBasePass.PreIntegratedGFTexture
+	// Material setting to allow forward shading (including mobile) to use preintegrated GF lut for simple IBL.
+	#if MATERIAL_SHADINGMODEL_SINGLELAYERWATER || (FORWARD_SHADING && MATERIAL_USE_PREINTEGRATED_GF)
+		#define PreIntegratedGF			OpaqueBasePass.PreIntegratedGFTexture
 
-	#if SUPPORTS_INDEPENDENT_SAMPLERS
-		#define PreIntegratedGFSampler	View.SharedBilinearClampedSampler
-	#else
-		#define PreIntegratedGFSampler	OpaqueBasePass.PreIntegratedGFSampler
+		#if SUPPORTS_INDEPENDENT_SAMPLERS
+			#define PreIntegratedGFSampler	View.SharedBilinearClampedSampler
+		#else
+			#define PreIntegratedGFSampler	OpaqueBasePass.PreIntegratedGFSampler
+		#endif
 	#endif
 #endif
 
