* Add group color to CV visualization
* Add reference point / position offset.
* Add large point option.

diff --git a/Engine/Plugins/Runtime/HairStrands/Source/HairStrandsCore/Private/GroomDebug.cpp b/Engine/Plugins/Runtime/HairStrands/Source/HairStrandsCore/Private/GroomDebug.cpp
--- a/Engine/Plugins/Runtime/HairStrands/Source/HairStrandsCore/Private/GroomDebug.cpp
+++ b/Engine/Plugins/Runtime/HairStrands/Source/HairStrandsCore/Private/GroomDebug.cpp
@@ -447,11 +447,14 @@ class FDrawDebugStrandsCVsCS : public FGlobalShader
 	BEGIN_SHADER_PARAMETER_STRUCT(FParameters, )
 		SHADER_PARAMETER_STRUCT_REF(FViewUniformShaderParameters, ViewUniformBuffer)
 		SHADER_PARAMETER(uint32, MaxVertexCount)
+		SHADER_PARAMETER(uint32, InstanceIt)
+		SHADER_PARAMETER(FVector3f, GroupColor)
 		SHADER_PARAMETER(FMatrix44f, LocalToWorld)
 		SHADER_PARAMETER_SAMPLER(SamplerState, LinearSampler)
 		SHADER_PARAMETER_RDG_TEXTURE(Texture2D<float>, DepthTexture)
 		SHADER_PARAMETER_RDG_TEXTURE_UAV(RWTexture2D<float4>, ColorTexture)
 		SHADER_PARAMETER_STRUCT_REF(FHairStrandsVertexFactoryUniformShaderParameters, HairStrandsVF)
+		SHADER_PARAMETER_STRUCT_INCLUDE(ShaderPrint::FShaderParameters, ShaderPrintParameters)
 	END_SHADER_PARAMETER_STRUCT()
 
 public:
@@ -466,15 +469,28 @@ public:
 
 IMPLEMENT_GLOBAL_SHADER(FDrawDebugStrandsCVsCS, "/Engine/Private/HairStrands/HairStrandsDebug.usf", "MainCS", SF_Compute);
 
+const FLinearColor GetHairGroupDebugColor(int32 GroupIt);
+
 static void AddDrawDebugStrandsCVsPass(
 	FRDGBuilder& GraphBuilder,
 	const FSceneView& View,
 	FGlobalShaderMap* ShaderMap,
 	const FHairGroupInstance* Instance,
 	const FShaderPrintData* ShaderPrintData,
+	const uint32 InIndex,
 	FRDGTextureRef ColorTexture,
 	FRDGTextureRef DepthTexture)
 {
+	if (!ShaderPrint::IsSupported(View.GetShaderPlatform()))
+	{
+		return;
+	}
+
+	// Force ShaderPrint on.
+	ShaderPrint::SetEnabled(true);
+	ShaderPrint::RequestSpaceForLines(256);
+	ShaderPrint::RequestSpaceForCharacters(256);
+
 	if (Instance->HairGroupPublicData->VFInput.GeometryType != EHairGeometryType::Strands || ShaderPrintData == nullptr)
 	{
 		return;
@@ -494,6 +510,10 @@ static void AddDrawDebugStrandsCVsPass(
 	Parameters->ColorTexture = GraphBuilder.CreateUAV(ColorTexture);
 	Parameters->DepthTexture = DepthTexture;
 	Parameters->LinearSampler = TStaticSamplerState<SF_Bilinear, AM_Clamp, AM_Clamp, AM_Clamp>::GetRHI();
+	Parameters->InstanceIt = InIndex;
+	Parameters->GroupColor = FVector3f(GetHairGroupDebugColor(Instance->Debug.GroupIndex));
+
+	ShaderPrint::SetParameters(GraphBuilder, *ShaderPrintData, Parameters->ShaderPrintParameters);
 
 	const uint32 PointCount = Instance->HairGroupPublicData->VFInput.Strands.Common.PointCount;
 	FComputeShaderUtils::AddPass(
@@ -1298,6 +1318,7 @@ void RunHairStrandsDebug(
 
 	if (bEnabled)
 	{	
+		uint32 Index = 0;
 		for (FHairStrandsInstance* AbstractInstance : Instances)
 		{
 			FHairGroupInstance* Instance = static_cast<FHairGroupInstance*>(AbstractInstance);
@@ -1314,8 +1335,9 @@ void RunHairStrandsDebug(
 	
 			if (ViewMode == EGroomViewMode::ControlPoints)
 			{
-				AddDrawDebugStrandsCVsPass(GraphBuilder, View, ShaderMap, Instance, ShaderPrintData, SceneColorTexture, SceneDepthTexture);
+				AddDrawDebugStrandsCVsPass(GraphBuilder, View, ShaderMap, Instance, ShaderPrintData, Index, SceneColorTexture, SceneDepthTexture);
 			}
+			++Index;
 		}
 	}
 
diff --git a/Engine/Shaders/Private/HairStrands/HairStrandsDebug.usf b/Engine/Shaders/Private/HairStrands/HairStrandsDebug.usf
--- a/Engine/Shaders/Private/HairStrands/HairStrandsDebug.usf
+++ b/Engine/Shaders/Private/HairStrands/HairStrandsDebug.usf
@@ -250,9 +250,12 @@ void MainPS(
 
 #if SHADER_CVS
 #include "HairStrandsVertexFactoryCommon.ush"
+#include "../ShaderPrint.ush"
 
 float4x4 LocalToWorld;
 uint MaxVertexCount;
+uint InstanceIt;
+float3 GroupColor;
 
 SamplerState LinearSampler;
 Texture2D<float> DepthTexture;
@@ -271,30 +274,66 @@ void MainCS(uint3 DispatchThreadId : SV_DispatchThreadID)
 		VertexIndex = HairStrandsVF.CullingIndexBuffer[VertexIndex];
 	}
 
-	const float3 P_Local = ReadHairControlPoint(
-		HairStrandsVF.PositionBuffer,
-		VertexIndex,
-		ReadRenPositionOffset(HairStrandsVF.PositionOffsetBuffer, HairStrandsVF.RegisteredIndex),
-		HairStrandsVF.Radius,
-		HairStrandsVF.RootScale,
-		HairStrandsVF.TipScale).Position;
-
-	const float3 P_World = mul(float4(P_Local, 1), LocalToWorld).xyz;
-	float4 SamplePosition = float4(P_World + DFHackToFloat(PrimaryView.PreViewTranslation), 1);
-	float4 ClipPosition = mul(SamplePosition, View.TranslatedWorldToClip);
-	float2 ScreenPosition = ClipPosition.xy / ClipPosition.w;
-	float2 ScreenUV = ScreenPosition * View.ScreenPositionScaleBias.xy + View.ScreenPositionScaleBias.wz;
-	float SampleSceneDepth = ClipPosition.w;
-
-	const float S = 1;
-	const float4 Color = float4(S, S, 0, 1);
-	const float SceneDepth = ConvertFromDeviceZ(DepthTexture.SampleLevel(LinearSampler, ScreenUV, 0));
-	const bool bIsVisible = SampleSceneDepth <= SceneDepth;
-	if (bIsVisible)
-	{
-		int2 OutCoord = ScreenUV * View.BufferSizeAndInvSize.xy;
-		OutCoord = clamp(OutCoord, int2(0, 0), View.BufferSizeAndInvSize.xy);
-		ColorTexture[OutCoord] = Color;
+	const float3 HairStrandsVF_RenPositionOffset = ReadRenPositionOffset(HairStrandsVF.PositionOffsetBuffer, HairStrandsVF.RegisteredIndex);
+
+	FShaderPrintContext CtxWidget = InitShaderPrintContext(InstanceIt == 0 && VertexId == 0, uint2(50, 50));
+	Print(CtxWidget, TEXT("Controls"), FontOrange); Newline(CtxWidget);
+	const bool bOffsetPosition = AddCheckbox(CtxWidget, TEXT("Position Offset"), false, FontWhite); Newline(CtxWidget);
+	const bool bControlPoints  = AddCheckbox(CtxWidget, TEXT("Control Points"), true, FontWhite); Newline(CtxWidget);
+	const bool bGroupColor     = AddCheckbox(CtxWidget, TEXT("Group Color"), true, FontWhite); Newline(CtxWidget);
+	const bool bLargePoint     = AddCheckbox(CtxWidget, TEXT("Large Point"), false, FontWhite); Newline(CtxWidget);
+//	int PointSize              = AddSlider(CtxWidget,   TEXT("Point Size"), 1, FontWhite, 1, 3); Print(CtxWidget, PointSize, FontYellow); Newline(CtxWidget);	
+
+	FShaderPrintContext Ctx = InitShaderPrintContext(VertexId == 0, uint2(50, 50));
+	if (Ctx.bIsActive && bOffsetPosition)
+	{
+		const float3 P_Local = HairStrandsVF_RenPositionOffset;
+		const float3 P_World = mul(float4(P_Local, 1), LocalToWorld).xyz;
+		AddCrossWS(P_World, 10.f, ColorRed);
+	}
+
+	if (bControlPoints)
+	{
+		const float3 P_Local = ReadHairControlPoint(
+			HairStrandsVF.PositionBuffer,
+			VertexIndex,
+			HairStrandsVF_RenPositionOffset,
+			HairStrandsVF.Radius,
+			HairStrandsVF.RootScale,
+			HairStrandsVF.TipScale).Position;
+
+		const float3 P_World = mul(float4(P_Local, 1), LocalToWorld).xyz;
+		float4 SamplePosition = float4(P_World + DFHackToFloat(PrimaryView.PreViewTranslation), 1);
+		float4 ClipPosition = mul(SamplePosition, View.TranslatedWorldToClip);
+		float2 ScreenPosition = ClipPosition.xy / ClipPosition.w;
+		float2 ScreenUV = ScreenPosition * View.ScreenPositionScaleBias.xy + View.ScreenPositionScaleBias.wz;
+		float SampleSceneDepth = ClipPosition.w;
+
+		const float S = 1;
+		const float4 Color = bGroupColor ? float4(GroupColor, 1) : float4(S, S, 0, 1);
+		const float SceneDepth = ConvertFromDeviceZ(DepthTexture.SampleLevel(LinearSampler, ScreenUV, 0));
+		const bool bIsVisible = SampleSceneDepth <= SceneDepth;
+		if (bIsVisible)
+		{
+			int2 OutCoord = ScreenUV * View.BufferSizeAndInvSize.xy;
+			OutCoord = clamp(OutCoord, int2(0, 0), View.BufferSizeAndInvSize.xy-1);
+
+			if (bLargePoint)
+			{
+				int PointSize = 1;
+				for (int y = -PointSize; y <= PointSize; ++y)
+				for (int x = -PointSize; x <= PointSize; ++x)
+				{
+					int2 SplatCoord = OutCoord + int2(x, y);
+					SplatCoord = clamp(SplatCoord, int2(0, 0), View.BufferSizeAndInvSize.xy-1);
+					ColorTexture[SplatCoord] = Color;
+				}
+			}
+			else
+			{
+				ColorTexture[OutCoord] = Color;
+			}
+		}
 	}
 }
 
