Path Tracer: Fix substrate support for new debug mode

This just implements the basic operations without new enums. Eventually more modes should  be added to better visualize custom substrate only fields.

#rnx

diff --git a/Engine/Shaders/Private/PathTracing/PathTracingDebug.usf b/Engine/Shaders/Private/PathTracing/PathTracingDebug.usf
--- a/Engine/Shaders/Private/PathTracing/PathTracingDebug.usf
+++ b/Engine/Shaders/Private/PathTracing/PathTracingDebug.usf
@@ -136,22 +136,33 @@ PathTracingResult PathTracingDebugCore(FRayDesc Ray, uint2 LaunchIndex)
 			case PATH_TRACER_DEBUG_VIZ_WORLD_SMOOTH_NORMAL:			return MakeResult(PackedPayload.UnpackWorldSmoothNormal() * 0.5 + 0.5, 0.0, Z);
 			case PATH_TRACER_DEBUG_VIZ_WORLD_GEO_NORMAL:			return MakeResult(PackedPayload.UnpackWorldGeoNormal() * 0.5 + 0.5, 0.0, Z);
 			default:
+#if SUBSTRATE_ENABLED
+			// TODO: implement more visualizations for this case
+			case PATH_TRACER_DEBUG_VIZ_BASE_COLOR:					return MakeResult(lerp(PackedPayload.UnpackDiffuseColor(), PackedPayload.UnpackSpecularColor(), F0RGBToMetallic(PackedPayload.UnpackSpecularColor())), 0.0, Z);
+			case PATH_TRACER_DEBUG_VIZ_DIFFUSE_COLOR:				return MakeResult(PackedPayload.UnpackDiffuseColor(), 0.0, Z);
+			case PATH_TRACER_DEBUG_VIZ_SPECULAR_COLOR:				return MakeResult(PackedPayload.UnpackSpecularColor(), 0.0, Z);
+			case PATH_TRACER_DEBUG_VIZ_METALLIC:					return MakeResult(F0RGBToMetallic(PackedPayload.UnpackSpecularColor()), 0.0, Z);
+			case PATH_TRACER_DEBUG_VIZ_SPECULAR:					return MakeResult(F0RGBToDielectricSpecular(PackedPayload.UnpackSpecularColor()), 0.0, Z);
+			case PATH_TRACER_DEBUG_VIZ_ROUGHNESS:					return MakeResult(PackedPayload.UnpackRoughnessAniso().xyz, 0.0, Z);
+			case PATH_TRACER_DEBUG_VIZ_ANISOTROPY:					return MakeResult(AnisoDebugColor(PackedPayload.UnpackRoughnessAniso().w), 0.0, Z);
+#else
 			case PATH_TRACER_DEBUG_VIZ_BASE_COLOR:					return MakeResult(PackedPayload.UnpackBaseColor(), 0.0, Z);
 			case PATH_TRACER_DEBUG_VIZ_DIFFUSE_COLOR:				return MakeResult(PackedPayload.UnpackBaseColor() * (1 - PackedPayload.UnpackMetallic()), 0.0, Z);
 			case PATH_TRACER_DEBUG_VIZ_SPECULAR_COLOR:				return MakeResult(lerp(0.08 * PackedPayload.UnpackSpecular(), PackedPayload.UnpackBaseColor(), PackedPayload.UnpackMetallic()), 0.0, Z);
-			case PATH_TRACER_DEBUG_VIZ_OPACITY:						return MakeResult(PackedPayload.UnpackTransparencyColor(), 0.0, Z);
 			case PATH_TRACER_DEBUG_VIZ_METALLIC:					return MakeResult(PackedPayload.UnpackMetallic(), 0.0, Z);
 			case PATH_TRACER_DEBUG_VIZ_SPECULAR:					return MakeResult(PackedPayload.UnpackSpecular(), 0.0, Z);
 			case PATH_TRACER_DEBUG_VIZ_ROUGHNESS:					return MakeResult(PackedPayload.UnpackRoughness(), 0.0, Z);
+			case PATH_TRACER_DEBUG_VIZ_ANISOTROPY:					return MakeResult(AnisoDebugColor(PackedPayload.UnpackAnisotropy()), 0.0, Z);
+			case PATH_TRACER_DEBUG_VIZ_CUSTOM_DATA0:				return MakeResult(PackedPayload.UnpackCustomData0().rgb * View.PreExposure, PackedPayload.UnpackCustomData0().a, Z);
+			case PATH_TRACER_DEBUG_VIZ_CUSTOM_DATA1:				return MakeResult(PackedPayload.UnpackCustomData1().rgb * View.PreExposure, PackedPayload.UnpackCustomData1().a, Z);
+#endif
+			case PATH_TRACER_DEBUG_VIZ_OPACITY:						return MakeResult(PackedPayload.UnpackTransparencyColor(), 0.0, Z);
 			case PATH_TRACER_DEBUG_VIZ_IOR:							return MakeResult(PackedPayload.UnpackIor(), 0.0, Z);
 			case PATH_TRACER_DEBUG_VIZ_SHADING_MODEL:				return MakeResult(GetShadingModelColor(PackedPayload.GetShadingModelID()), 0.0, Z);
 			case PATH_TRACER_DEBUG_VIZ_LIGHTING_CHANNEL_MASK:		return MakeResult(GetLightingChannelMaskDebugColor(PackedPayload.GetPrimitiveLightingChannelMask()), 0.0, Z);
-			case PATH_TRACER_DEBUG_VIZ_CUSTOM_DATA0:				return MakeResult(PackedPayload.UnpackCustomData0().rgb * View.PreExposure, PackedPayload.UnpackCustomData0().a, Z);
-			case PATH_TRACER_DEBUG_VIZ_CUSTOM_DATA1:				return MakeResult(PackedPayload.UnpackCustomData1().rgb * View.PreExposure, PackedPayload.UnpackCustomData1().a, Z);
 			case PATH_TRACER_DEBUG_VIZ_WORLD_POSITION:				return MakeResult(TranslatedWorldPos * View.PreExposure, 0.0, Z);
 			case PATH_TRACER_DEBUG_VIZ_PRIMARY_RAYS: break; // fallthrough to implementation below
 			case PATH_TRACER_DEBUG_VIZ_WORLD_TANGENT: 				return MakeResult(PackedPayload.UnpackWorldTangent() * 0.5 + 0.5, 0.0, Z);
-			case PATH_TRACER_DEBUG_VIZ_ANISOTROPY:					return MakeResult(AnisoDebugColor(PackedPayload.UnpackAnisotropy()), 0.0, Z);
 			case PATH_TRACER_DEBUG_VIZ_LIGHT_GRID_COUNT:
 			{
 				float3 Color = LightGridVisualize(LightLoopCount, 1);
@@ -219,7 +230,11 @@ PathTracingResult PathTracingDebugCore(FRayDesc Ray, uint2 LaunchIndex)
 					LightSample.RadianceOverPdf /= LightPickPdf;
 					LightSample.Pdf *= LightPickPdf;
 
+#if SUBSTRATE_ENABLED
+					float3 MaterialEvalWeight = PackedPayload.UnpackDiffuseColor();
+#else
 					float3 MaterialEvalWeight = PackedPayload.UnpackBaseColor();
+#endif
 					float  MaterialEvalPdf = saturate(dot(WorldNormal, LightSample.Direction)) / PI;
 					float3 LightContrib = PathThroughput * LightSample.RadianceOverPdf * MaterialEvalWeight * MaterialEvalPdf;
 					{
