Path Tracer: Fix inconsistency in the handling of skylight intensity between path tracer and raster

The sky visible to camera should not be scaled by the intensity in this case. It was previously being accounted for by the color.

diff --git a/Engine/Shaders/Private/PathTracing/PathTracingCore.ush b/Engine/Shaders/Private/PathTracing/PathTracingCore.ush
--- a/Engine/Shaders/Private/PathTracing/PathTracingCore.ush
+++ b/Engine/Shaders/Private/PathTracing/PathTracingCore.ush
@@ -526,6 +526,11 @@ FPackedPathTracingPayload TraceTransparentRay(inout FPathState PathState, int Bo
 				FPathTracingPayload LightHitPayload = (FPathTracingPayload)0; // clear all fields
 				LightHitPayload.HitT = LightHit.HitT;
 				LightHitPayload.Radiance = LightHit.Radiance;
+				if ((SceneLights[LightId].Flags & PATHTRACER_FLAG_TYPE_MASK) == PATHTRACING_LIGHT_SKY)
+				{
+					// UE-222576: cancel out scale factor that is baked into the texture for camera rays only
+					LightHitPayload.Radiance /= GetColor(LightId);
+				}
 				LightHitPayload.ShadingModelID = SHADINGMODELID_UNLIT;
 				// light hits are fully transparent to match how the light loop works (the light surface does not cast a shadow)
 				LightHitPayload.BSDFOpacity = 1.0;
diff --git a/Engine/Source/Runtime/Renderer/Private/PathTracing.cpp b/Engine/Source/Runtime/Renderer/Private/PathTracing.cpp
--- a/Engine/Source/Runtime/Renderer/Private/PathTracing.cpp
+++ b/Engine/Source/Runtime/Renderer/Private/PathTracing.cpp
@@ -2672,6 +2672,13 @@ void SetLightParameters(
 
 			// Also allow seeing just the sky via a cvar for debugging purposes
 			*SceneVisibleLightCount = 1;
+
+			if (Scene->SkyLight->bRealTimeCaptureEnabled)
+			{
+				// NOTE: this color is already baked into the skylight texture so that importance sampling takes it into account, we pass it in here so that camera rays can factor it out
+				// This is only for the realtime capture case, because otherwise (specified cube map case) we want the displayed texture and lighting to match
+				DestLight.Color = FVector3f(Scene->SkyLight->GetEffectiveLightColor());
+			}
 		}
 	}
 
