Fix incorrect flag logic in GPU Scene Debug Render shader.

#rnx

diff --git a/Engine/Shaders/Private/GPUSceneDebugRender.usf b/Engine/Shaders/Private/GPUSceneDebugRender.usf
--- a/Engine/Shaders/Private/GPUSceneDebugRender.usf
+++ b/Engine/Shaders/Private/GPUSceneDebugRender.usf
@@ -176,8 +176,8 @@ void GPUSceneDebugRenderCS(uint InstanceID : SV_DispatchThreadID)
 		FPrimitiveSceneData PrimitiveData = GetPrimitiveData(InstanceData.PrimitiveId);
 		AddFlag(Context, 0u != (PrimitiveData.Flags & PRIMITIVE_SCENE_DATA_FLAG_CAST_SHADOWS));
 		AddFlag(Context, 0u != (PrimitiveData.Flags & PRIMITIVE_SCENE_DATA_FLAG_OUTPUT_VELOCITY));
-		AddFlag(Context, 0u != (PrimitiveData.Flags & INSTANCE_SCENE_DATA_FLAG_HAS_CUSTOM_DATA));
-		AddFlag(Context, 0u != (PrimitiveData.Flags & INSTANCE_SCENE_DATA_FLAG_HAS_DYNAMIC_DATA));
+		AddFlag(Context, 0u != (InstanceData.Flags & INSTANCE_SCENE_DATA_FLAG_HAS_CUSTOM_DATA));
+		AddFlag(Context, 0u != (InstanceData.Flags & INSTANCE_SCENE_DATA_FLAG_HAS_DYNAMIC_DATA));
 		
 		PrintInstanceName(Context, InstanceData.PrimitiveId, FontSilver); Newline(Context);
 	}
