Fix HZB culling in first slice of light grid.

diff --git a/Engine/Shaders/Private/LightGridInjection.usf b/Engine/Shaders/Private/LightGridInjection.usf
--- a/Engine/Shaders/Private/LightGridInjection.usf
+++ b/Engine/Shaders/Private/LightGridInjection.usf
@@ -74,7 +74,7 @@ float ComputeCellNearViewDepthFromZSlice(uint ZSlice)
 		// The exponential distribution of z slices contains an offset, but some screen pixels
 		// may be nearer to the camera than this offset. To avoid false light rejection, we set the
 		// first depth slice to zero to ensure that the AABB includes the [0, offset] depth range.
-		SliceDepth = 0.0f;
+		SliceDepth = View.NearPlane;
 	}
 
 	return SliceDepth;
