Substrate - fixed missing participating media albedo update of top layers when using parameter blending.

#rnx
#tests Fix for the Substrate tests already in place.

diff --git a/Engine/Shaders/Private/ParticipatingMediaCommon.ush b/Engine/Shaders/Private/ParticipatingMediaCommon.ush
--- a/Engine/Shaders/Private/ParticipatingMediaCommon.ush
+++ b/Engine/Shaders/Private/ParticipatingMediaCommon.ush
@@ -198,7 +198,8 @@ float3 IsotropicMediumSlabDirectionalAlbedoFade(float3 BaseColor, float3 MFP)
 	return Fade;
 }
 
-float3 IsotropicMediumSlabPunctualDirectionalAlbedo(FParticipatingMedia PM)
+// The ALU version is a lot simpler as compared to the LUT version: it is not view and not light dependent.
+float3 IsotropicMediumSlabPunctualDirectionalAlbedoALU(FParticipatingMedia PM)
 {
 	// Our measurement are only valid for a MFP of 0.01 meter so we clamp input MFP to that.
 	// This is relatively ok since all computation are done for a slab of 1 meter.
@@ -215,7 +216,7 @@ float3 IsotropicMediumSlabPunctualDirectionalAlbedo(FParticipatingMedia PM)
 // Platforms that do not support shared sampler will in fact quickly reach the limit of samplers in material pixel shaders.
 // This is mostly true for OpenGLES. In this case, we enforce disabling SimpleVolume LUT and use the lower quality ALU based one.
 // SUBSTRATE_TODO Implement an ALU based solution that give the same view dependent high quality as the LUT.
-#define USE_NEW_SIMPLEVOLUME (SUPPORTS_INDEPENDENT_SAMPLERS)
+#define USE_LUT_SIMPLEVOLUME (SUPPORTS_INDEPENDENT_SAMPLERS)
 
 float IsotropicMediumSlabDirectionalAlbedoFadeLUT(float MFP)
 {
@@ -291,7 +292,20 @@ float3 IsotropicMediumSlabPunctualDirectionalAlbedoLUT(FParticipatingMedia PM, f
 	);
 }
 
-float3 IsotropicMediumSlabEnvDirectionalAlbedo(FParticipatingMedia PM)
+float IsotropicMediumSlabPunctualDirectionalAlbedo(FParticipatingMedia PM, float NoV, float NoL)
+{
+#if USE_LUT_SIMPLEVOLUME
+	const float3 SlabDirectionalAlbedo = IsotropicMediumSlabPunctualDirectionalAlbedoLUT(PM, NoV, NoL);
+#else
+	const float3 SlabDirectionalAlbedo = IsotropicMediumSlabPunctualDirectionalAlbedoALU(PM);
+#endif
+	return SlabDirectionalAlbedo;
+}
+
+
+
+// The ALU version is a lot simpler as compared to the LUT version: it is not view dependent.
+float3 IsotropicMediumSlabEnvDirectionalAlbedoALU(FParticipatingMedia PM)
 {
 	// Our measurement are only valid for a MFP of 0.01 meter so we clamp input MFP to that.
 	// This is relatively ok since all computation are done for a slab of 1 meter.
@@ -353,3 +367,20 @@ float3 IsotropicMediumSlabEnvDirectionalAlbedoLUT(FParticipatingMedia PM, float
 		IsotropicMediumSlabEnvDirectionalAlbedoLUT(PM.MeanFreePath.b, PM.BaseColor.b, NoV)
 	);
 }
+
+float3 IsotropicMediumSlabEnvDirectionalAlbedo(FParticipatingMedia PM, float NoV)
+{
+#if USE_LUT_SIMPLEVOLUME
+	float3 SlabEnvAlbedo = IsotropicMediumSlabEnvDirectionalAlbedoLUT(PM, NoV);
+#else
+	// The environment response of has been measured for an hemisphere having a uniform radiance of 1. 
+	// This DirectionalAlbedo also contains the integral over the hemisphere according to the phase function.
+	// The spherical harmonic used for environment lighting contains the diffuse integral over the hemisphere as DiffLambert = int_{\omega} {1 * NoL d\Omega} = PI.
+	// We do not want to be affected by the NoL term here since The DirectionalAlbedo also contains the integral over the hemisphere according to the phase function.
+	// Integral of a value=1 over the Hemisphere is int_{\omega} {1  d\Omega} = 2PI.
+	// So, as an approximation, we recover the uniform hemisphere luminance as UniformEnvL = DiffLambert * (1 / PI) * (2 * PI).
+	// We consider that the diffuse integration of the environment SH over the hemisphere multiplied with the light response will be a good approximation similar to the "split sum integral".
+	float3 SlabEnvAlbedo = IsotropicMediumSlabEnvDirectionalAlbedoALU(PM) * ((1 / PI) * (2 * PI));
+#endif
+	return SlabEnvAlbedo;
+}
diff --git a/Engine/Shaders/Private/PathTracing/PathTracingMaterialHitShader.usf b/Engine/Shaders/Private/PathTracing/PathTracingMaterialHitShader.usf
--- a/Engine/Shaders/Private/PathTracing/PathTracingMaterialHitShader.usf
+++ b/Engine/Shaders/Private/PathTracing/PathTracingMaterialHitShader.usf
@@ -157,7 +157,7 @@ float3 GetSimpleVolumeDiffuseColor(float3 DiffuseColor, float3 MeanFreePath)
 	// SUBSTRATE_TODO: In the case of thin materials, we should also include a backscattering diffuse portion
 	FParticipatingMedia PM = SubstrateSlabCreateParticipatingMedia(DiffuseColor, MeanFreePath);
 	const float DiffuseToVolumeBlend = SubstrateSlabDiffuseToVolumeBlend(PM);
-	const float3 SlabDirectionalAlbedo = IsotropicMediumSlabEnvDirectionalAlbedo(PM);
+	const float3 SlabDirectionalAlbedo = IsotropicMediumSlabEnvDirectionalAlbedoALU(PM);
 	return lerp(DiffuseColor, SlabDirectionalAlbedo, DiffuseToVolumeBlend);
 }
 
diff --git a/Engine/Shaders/Private/Substrate/Substrate.ush b/Engine/Shaders/Private/Substrate/Substrate.ush
--- a/Engine/Shaders/Private/Substrate/Substrate.ush
+++ b/Engine/Shaders/Private/Substrate/Substrate.ush
@@ -2755,16 +2755,17 @@ FSubstrateBSDF SubstrateVerticalLayeringParameterBlending(in FSubstrateBSDF TopB
 	// We do not check BSDF_GETISSIMPLEVOLUME(BSDFContext.BSDF) because that is set when we know the BSDF is not at the bottom later. And here we already know it is not the case.
 	TopBSDF.EnableSlabBSDFSimpleVolumetric();
 	BaseBSDF.EnableSlabBSDFSimpleVolumetric();
-	const float3 UnusedDiffuseColor		= 0;	// Unused when only transmittance is needed
-	FParticipatingMedia PM				= SubstrateSlabCreateParticipatingMedia(UnusedDiffuseColor, SLAB_SSSMFP(TopBSDF));
+	FParticipatingMedia PM				= SubstrateSlabCreateParticipatingMedia(SLAB_DIFFUSEALBEDO(TopBSDF), SLAB_SSSMFP(TopBSDF));
 	const float3 SlabTransmittanceV		= IsotropicMediumSlabTransmittance(PM, SUBSTRATE_SIMPLEVOLUME_THICKNESS_M, NoV);
 	const float3 SlabTransmittanceN		= IsotropicMediumSlabTransmittance(PM, SUBSTRATE_SIMPLEVOLUME_THICKNESS_M, 1.0f /*Normal incidence i.e., NoL == 1.f*/);
 	const float3 SlabTransmittance		= SlabTransmittanceV * SlabTransmittanceN;
 
+	// Using the environment version as an approximation for albedo resulting from scattering in a top layer without light source direction.
+	float3 TopSlabAlbedo = IsotropicMediumSlabEnvDirectionalAlbedo(PM, NoV);
 	// SUBSTRATE_TODO account for energy conservation here.
 
 	FVerticalLayeringInfo Info			= GetVerticalLayeringInfo(SafeTopCoverage, SafeBaseCoverage);
-	// Compute some ways to intepoldate and accumulate parameters, but always normalise because coverage is applied later in UpdateSingleBSDFOperatorCoverageTransmittance.
+	// Compute some ways to interpolate and accumulate parameters, but always normalise because coverage is applied later in UpdateSingleBSDFOperatorCoverageTransmittance.
 	const float  TopToBottomLerp		= saturate((Info.TransmittanceTopAndBottom * dot(SlabTransmittance,0.33.xxx) + Info.SurfaceBottom) / max(SUBSTRATE_EPSILON, Info.Coverage));
 	const float3 BottomColorFactor		= (Info.TransmittanceTopAndBottom * SlabTransmittance + Info.SurfaceBottom) / max(SUBSTRATE_EPSILON, Info.Coverage);
 	const float  TopFactor				= Info.SurfaceTop / max(SUBSTRATE_EPSILON, Info.Coverage);
@@ -2777,7 +2778,7 @@ FSubstrateBSDF SubstrateVerticalLayeringParameterBlending(in FSubstrateBSDF TopB
 	// Some parameters are not contributing to added luminance: in this case we simply lerp them
 	BSDF_SETEMISSIVE(ResultBSDF,		BottomColorFactor * BSDF_GETEMISSIVE(BaseBSDF)								+ TopFactor * BSDF_GETEMISSIVE(TopBSDF));
 
-	SLAB_DIFFUSEALBEDO(ResultBSDF) =	saturate(BottomColorFactor * SLAB_DIFFUSEALBEDO(BaseBSDF)					+ TopFactor * SLAB_DIFFUSEALBEDO(TopBSDF));
+	SLAB_DIFFUSEALBEDO(ResultBSDF) =	saturate(BottomColorFactor * SLAB_DIFFUSEALBEDO(BaseBSDF)					+ TopFactor * TopSlabAlbedo);
 	SLAB_F0(ResultBSDF) =				saturate(BottomColorFactor * SLAB_F0(BaseBSDF)								+ TopFactor * SLAB_F0(TopBSDF));
 	SLAB_F90(ResultBSDF) =				saturate(BottomColorFactor * SLAB_F90(BaseBSDF)								+ TopFactor * SLAB_F90(TopBSDF));
 	SLAB_FUZZ_COLOR(ResultBSDF) =		saturate(BottomColorFactor * SLAB_FUZZ_COLOR(BaseBSDF)						+ TopFactor * SLAB_FUZZ_COLOR(TopBSDF));
diff --git a/Engine/Shaders/Private/Substrate/SubstrateEvaluation.ush b/Engine/Shaders/Private/Substrate/SubstrateEvaluation.ush
--- a/Engine/Shaders/Private/Substrate/SubstrateEvaluation.ush
+++ b/Engine/Shaders/Private/Substrate/SubstrateEvaluation.ush
@@ -842,11 +842,7 @@ FSubstrateEvaluateResult SubstrateEvaluateBSDFCommon(FSubstrateBSDFContext BSDFC
 				const float HGPhaseFunction	= HenyeyGreensteinPhase(SLAB_SSSPHASEANISOTROPY(BSDFContext.BSDF), CosTheta);
 				const float PhaseFunction	= HGPhaseFunction / IsotropicPhase();
 
-			#if USE_NEW_SIMPLEVOLUME
-				const float3 SlabDirectionalAlbedo = PhaseFunction * IsotropicMediumSlabPunctualDirectionalAlbedoLUT(PM, BSDFContext.Context.NoV, BSDFContext.Context.NoL);
-			#else
-				const float3 SlabDirectionalAlbedo = PhaseFunction * IsotropicMediumSlabPunctualDirectionalAlbedo(PM);
-			#endif
+				const float3 SlabDirectionalAlbedo = PhaseFunction * IsotropicMediumSlabPunctualDirectionalAlbedo(PM, BSDFContext.Context.NoV, BSDFContext.Context.NoL);
 
 			#if 1
 				// For transmittance, we must account the fact that the view ray is going to be refracted at the interface of the slab.
@@ -1474,18 +1470,7 @@ FSubstrateEnvLightResult SubstrateEvaluateForEnvLight(FSubstrateBSDFContext BSDF
 				FParticipatingMedia PM = SubstrateSlabCreateParticipatingMedia(DiffuseColor, SLAB_SSSMFP(BSDFContext.BSDF));
 				const float DiffuseToVolumeBlend = SubstrateSlabDiffuseToVolumeBlend(PM);
 
-			#if USE_NEW_SIMPLEVOLUME
-				float3 SlabDirectionalAlbedo = IsotropicMediumSlabEnvDirectionalAlbedoLUT(PM, BSDFContext.Context.NoV);
-			#else
-				// The environment response of has been measured for an hemisphere having a uniform radiance of 1. 
-				// This DirectionalAlbedo also contains the integral over the hemisphere according to the phase function.
-				// The spherical harmonic used for environment lighting contains the diffuse integral over the hemisphere as DiffLambert = int_{\omega} {1 * NoL d\Omega} = PI.
-				// We do not want to be affected by the NoL term here since The DirectionalAlbedo also contains the integral over the hemisphere according to the phase function.
-				// Integral of a value=1 over the Hemisphere is int_{\omega} {1  d\Omega} = 2PI.
-				// So, as an approximation, we recover the uniform hemisphere luminance as UniformEnvL = DiffLambert * (1 / PI) * (2 * PI).
-				// We consider that the diffuse integration of the environment SH over the hemisphere multiplied with the light response will be a good approximation similar to the "split sum integral".
-				float3 SlabDirectionalAlbedo = IsotropicMediumSlabEnvDirectionalAlbedo(PM) * ((1 / PI)* (2 * PI));
-			#endif
+				float3 SlabDirectionalAlbedo = IsotropicMediumSlabEnvDirectionalAlbedo(PM, BSDFContext.Context.NoV);
 
 			#if 1
 				// Because the default EnvBRDF does not account for multiple scattering, we rescale the participating media response according to EnvBRDF.
