Subsutrate - avoid diffusion MFP explosion. The top layer MFP is usually large to keep the coat clear and this would result is wide scattering radius the SSS post process cannot resolve. So parameter blendign of slab only account for the MFP of the bottom layer now. Also we only save the previous MFP.

#rnx

diff --git a/Engine/Shaders/Private/Substrate/Substrate.ush b/Engine/Shaders/Private/Substrate/Substrate.ush
--- a/Engine/Shaders/Private/Substrate/Substrate.ush
+++ b/Engine/Shaders/Private/Substrate/Substrate.ush
@@ -2751,6 +2751,10 @@ FSubstrateBSDF SubstrateVerticalLayeringParameterBlending(in FSubstrateBSDF TopB
 	BSDF_SETHASSPECPROFILE(ResultBSDF,	0); // Disable for now
 #endif
 
+	// Saved original Base BBSDF MFP to not be changed be changed by EnableSlabBSDFSimpleVolumetric used to compute medium transmittance and scattering.
+	// EnableSlabBSDFSimpleVolumetric will be called at then end of the SubstrateTree as needed later.
+	const float3 OriginalBaseMFP = SLAB_SSSMFP(BaseBSDF);
+
 	// Compute the top slab transmittance that will be used to blend the top and base material parameters.
 	// We do not check BSDF_GETISSIMPLEVOLUME(BSDFContext.BSDF) because that is set when we know the BSDF is not at the bottom later. And here we already know it is not the case.
 	TopBSDF.EnableSlabBSDFSimpleVolumetric();
@@ -2805,13 +2809,15 @@ FSubstrateBSDF SubstrateVerticalLayeringParameterBlending(in FSubstrateBSDF TopB
 	// Anisotropy is disabled with vertical blending. 
 	SLAB_ANISOTROPY(ResultBSDF) =		0.f;
 
-	// Keep the MFP form the bottom layer only (top layer simple volume is backed in parameters).
-	SLAB_SSSMFP(ResultBSDF) =			lerp(SLAB_SSSMFP(TopBSDF),						SLAB_SSSMFP(BaseBSDF),		TopToBottomLerp);
+	// Keep the MFP form the bottom layer only and reduce it according to visibility. 
+	// The top layer simple volume coat scattering/transmittance is baked with parameter blending.
+	// Only the base is kept because the top layer MFP is usually large to keep the coat clear and this would result is wide scattering radius the SSS post process cannot resolve.
+	SLAB_SSSMFP(ResultBSDF) =			OriginalBaseMFP * TopToBottomLerp;
 
 #if SUBSTRATE_INLINE_SHADING
 #if SUBSTRATE_OPAQUE_MATERIAL
 	// Only the bottom layer SSS / Thin model is used so we only keep the bottom layer data there. Top layer is converted to throughput.
-	ResultBSDF.TmpMFP =					BaseBSDF.TmpMFP;
+	ResultBSDF.TmpMFP =					BaseBSDF.TmpMFP * TopToBottomLerp;
 	ResultBSDF.Coverage =				Info.Coverage;
 #else
 
@@ -2828,7 +2834,7 @@ FSubstrateBSDF SubstrateVerticalLayeringParameterBlending(in FSubstrateBSDF TopB
 	const float3 NewExtinction			= -log(saturate(max(0.0, NewTransmittance)));
 	const float3 NewMFP					= 1.0f / NewExtinction;
 #else
-	// MFP looks more prceptually linear even though it is not perfect
+	// MFP looks more perceptually linear even though it is not perfect
 	const float3 NewMFP = Info2.TransmittanceOnlyBottom * BaseBSDF.TmpMFP + Info2.TransmittanceOnlyTop * TopBSDF.TmpMFP + Info2.TransmittanceTopAndBottom * min(TopBSDF.TmpMFP, BaseBSDF.TmpMFP);
 #endif
 
