Add raytracing debug view for visualizating Substrate's material properties.

diff --git a/Engine/Shaders/Private/RayTracing/RayTracingDebug.usf b/Engine/Shaders/Private/RayTracing/RayTracingDebug.usf
--- a/Engine/Shaders/Private/RayTracing/RayTracingDebug.usf
+++ b/Engine/Shaders/Private/RayTracing/RayTracingDebug.usf
@@ -1,6 +1,10 @@
 // Copyright Epic Games, Inc. All Rights Reserved.
 
 #define SUBSTRATE_MATERIALCONTAINER_IS_VIEWRESOURCE 0
+#define SUBSTRATE_SSS_MATERIAL_OVERRIDE 0
+#define SUBSTRATE_SSS_TRANSMISSION 0				// Disable for now, as the SSS profile texture need to be bound
+#define SUBSTRATE_COMPLEXSPECIALPATH 1				// Allow to read data from complex special path
+#define SUBSTRATE_INLINE_SHADING 0
 
 #include "/Engine/Shared/RayTracingDebugTypes.h"
 #include "/Engine/Shared/RayTracingTypes.h"
@@ -9,7 +13,9 @@
 #include "../ColorMap.ush"
 #include "../Hash.ush"
 #include "../Visualization.ush"
+#include "../ShaderPrint.ush"
 #include "../Substrate/Substrate.ush"
+#include "../Substrate/SubstrateVisualizeCommon.ush"
 
 #include "RayTracingDebugUtils.ush"
 #include "RayTracingLightCullingCommon.ush"
@@ -54,7 +60,7 @@ float4 RoughnessDebugColor(float InRoughness)
 	return float4(pow(InRoughness, 2.2f).xxx, 0.0f);
 }
 
-float4 DebugRayTracingMaterial(FRayDesc Ray, bool bIsFarField)
+float4 DebugRayTracingMaterial(uint2 PixelCoord, FRayDesc Ray, bool bIsFarField, inout bool bWriteOuput)
 {
 	FRayCone RayCone = (FRayCone)0;
 	RayCone.SpreadAngle = View.EyeToPixelSpreadAngle;
@@ -123,7 +129,23 @@ float4 DebugRayTracingMaterial(FRayDesc Ray, bool bIsFarField)
 		FSubstratePixelHeader SubstratePixelHeader = UnpackSubstrateHeaderIn(Payload.SubstrateData, SubstrateAddressing, Payload.SubstrateData);
 
 		BRANCH
-		if (SubstratePixelHeader.ClosureCount > 0)
+		if (VisualizationMode == RAY_TRACING_DEBUG_VIZ_SUBSTRATE_DATA)
+		{
+			uint2 CursorPixelCoord = GetCursorPos() * View.ViewResolutionFraction;
+			FShaderPrintContext Ctx = InitShaderPrintContext(all(CursorPixelCoord == PixelCoord), uint2(50, 50));
+			if (Ctx.bIsActive)
+			{
+				//const FSubstrateBSDF BSDF = UnpackSubstrateBSDFIn(Payload.SubstrateData, SubstrateAddressing, SubstratePixelHeader); // Take the parameter of the first BSDF
+				const float3 TranslatedWorldPosition = Ray.Origin + Payload.HitT * Ray.Direction;
+				const float3 V = -normalize(TranslatedWorldPosition - PrimaryView.TranslatedWorldCameraOrigin);
+				const uint ClosureIndex = 0;
+				SubstratePrintMaterialProperties(Ctx, PixelCoord, TranslatedWorldPosition, V, ClosureIndex, Payload.SubstrateData);
+				bWriteOuput = false;
+				return 0; // No scene color writing
+			}
+
+		}
+		else if (SubstratePixelHeader.ClosureCount > 0)
 		{
 			const FSubstrateBSDF BSDF = UnpackSubstrateBSDFIn(Payload.SubstrateData, SubstrateAddressing, SubstratePixelHeader); // Take the parameter of the first BSDF
 			const float3x3 TangentBasis = SubstrateGetBSDFSharedBasis(SubstratePixelHeader, BSDF, SubstrateAddressing);
@@ -301,6 +323,8 @@ RAY_TRACING_ENTRY_RAYGEN(RayTracingDebugMainRGS)
 	float OutDepth = 1.0f;
 #endif
 
+	bool bWriteColorOuput = false;
+	bool bWriteDepthOutput = VisualizationMode == RAY_TRACING_DEBUG_VIZ_INSTANCE_OVERLAP;
 #if USE_DEBUG_CHS
 	{
 		Result = DebugRayTracingInstance(Ray, OutDepth);
@@ -308,14 +332,14 @@ RAY_TRACING_ENTRY_RAYGEN(RayTracingDebugMainRGS)
 #else // USE_DEBUG_CHS
 	{
 		bool bIsFarField = false;
-		Result = DebugRayTracingMaterial(Ray, bIsFarField);
+		Result = DebugRayTracingMaterial(PixelCoord, Ray, bIsFarField, bWriteColorOuput);
 		
 		const bool bNearFieldMiss = Result.w < 0.0;
 		if (VisualizationMode == RAY_TRACING_DEBUG_VIZ_FAR_FIELD && bNearFieldMiss)
 		{
 			bIsFarField = true;
 			const float4 FarFieldTint = float4(1.0f, 0.5f, 0.5f, 1.0f);
-			Result = DebugRayTracingMaterial(Ray, bIsFarField) * FarFieldTint;
+			Result = DebugRayTracingMaterial(PixelCoord, Ray, bIsFarField, bWriteColorOuput) * FarFieldTint;
 		}
 	}
 #endif // USE_DEBUG_CHS
@@ -335,10 +359,13 @@ RAY_TRACING_ENTRY_RAYGEN(RayTracingDebugMainRGS)
 		}
 	}
 	
-	Output[PixelCoord] = Result;
+	if (bWriteColorOuput)
+	{
+		Output[PixelCoord] = Result;
+	}
 
-	if (VisualizationMode == RAY_TRACING_DEBUG_VIZ_INSTANCE_OVERLAP)
+	if (bWriteDepthOutput)
 	{
 		OutputDepth[PixelCoord] = OutDepth;
-	}	
+	}
 }
diff --git a/Engine/Shaders/Private/Substrate/SubstrateVisualizeCommon.ush b/Engine/Shaders/Private/Substrate/SubstrateVisualizeCommon.ush
new file mode 100644
--- /dev/null
+++ b/Engine/Shaders/Private/Substrate/SubstrateVisualizeCommon.ush
@@ -0,0 +1,681 @@
+// Copyright Epic Games, Inc. All Rights Reserved.
+
+#include "/Engine/Private/Common.ush"
+#include "/Engine/Private/ShaderPrint.ush"
+
+#ifndef SUBSTRATE_MATERIALCONTAINER_IS_VIEWRESOURCE
+#error SUBSTRATE_MATERIALCONTAINER_IS_VIEWRESOURCE needs to be defined
+#endif
+
+void DrawReferentialTWS(float3 P /*In Translated World Space*/, float3 X, float3 Y, float3 N, float3 InColor)
+{
+	const float Size = 10.f;
+	const float SizeWithTip = 12.0f;
+	const float4 Color = float4(InColor, 1);
+	const float4 ColorX = float4(1, 0, 0, 1);
+	const float4 ColorY = float4(0, 1, 0, 1);
+	const float4 ColorZ = float4(0, 0, 1, 1);
+
+	// Core
+	AddLineTWS(P, P + X * Size, Color, Color);
+	AddLineTWS(P, P + Y * Size, Color, Color);
+	AddLineTWS(P, P + N * Size, Color, Color);
+	// Tips
+	AddLineTWS(P + X * Size, P + X * SizeWithTip, ColorX, ColorX);
+	AddLineTWS(P + Y * Size, P + Y * SizeWithTip, ColorY, ColorY);
+	AddLineTWS(P + N * Size, P + N * SizeWithTip, ColorZ, ColorZ);
+}
+
+void AddDrawPixelFootprint(float3 P, float3 dPdx, float3 dPdy, float2 Scale, bool bNormalize, float4 Color)
+{
+	const float3 T = (bNormalize ? normalize(dPdx) : dPdx) * Scale.x;
+	const float3 B = (bNormalize ? normalize(dPdy) : dPdy) * Scale.y;
+	const float3 N = normalize(cross(T, B));
+
+	const float3 WP0 = P - T - B;
+	const float3 WP1 = P + T - B;
+	const float3 WP2 = P + T + B;
+	const float3 WP3 = P - T + B;
+
+	AddLineTWS(WP0, WP1, Color);
+	AddLineTWS(WP1, WP2, Color);
+	AddLineTWS(WP2, WP3, Color);
+	AddLineTWS(WP3, WP0, Color);
+}
+
+#if SUBSTRATE_MATERIALCONTAINER_IS_VIEWRESOURCE
+void DrawPixelFootprint(float3 P, float3 dPdx, float3 dPdy, uint2 PixelCoord)
+{
+	const FSubstrateSubsurfaceHeader SSSHeader = SubstrateLoadSubsurfaceHeader(Substrate.MaterialTextureArray, Substrate.FirstSliceStoringSubstrateSSSData, PixelCoord);
+	const bool bIsValid = SubstrateSubSurfaceHeaderGetIsValid(SSSHeader);
+	const bool bIsProfile = SubstrateSubSurfaceHeaderGetIsProfile(SSSHeader);
+
+	if (bIsValid)
+	{
+		float3 MFP = 0;
+		if (bIsProfile)
+		{
+			MFP = GetSubsurfaceProfileMFPInCm(SubstrateSubSurfaceHeaderGetProfileId(SSSHeader)).xyz * SubstrateSubSurfaceHeaderGetProfileRadiusScale(SSSHeader);
+		}
+		else
+		{
+			MFP = SubstrateSubSurfaceHeaderGetMFP(SSSHeader);
+		}
+
+		FSubstratePixelFootprint Footprint = SubstrateGetPixelFootprint(dPdx, dPdy, 0.f /*InNormalCurvatureRoughness*/);
+		AddDrawPixelFootprint(P, dPdx, dPdy, 0.5f, false, ColorRed);
+		AddDrawPixelFootprint(P, dPdx, dPdy, Footprint.PixelRadiusInWorldSpace, true, ColorOrange);
+		AddDrawPixelFootprint(P, dPdx, dPdy, max3(MFP.x, MFP.y, MFP.z), true, ColorCyan);
+	}
+}
+#endif
+
+///////////////////////////////////////////////////////////////////////////////////////////////////
+// Material Print
+
+void PrintPixelType(inout FShaderPrintContext Ctx, in FSubstratePixelHeader Header, FFontColor InColor)
+{
+	if (Header.IsSimpleMaterial())						Print(Ctx, TEXT("Simple  "), InColor);
+	else if (Header.IsSingleMaterial())					Print(Ctx, TEXT("Single  "), InColor);
+	else if (Header.IsComplexSpecialMaterial())			Print(Ctx, TEXT("Complex Special "), InColor);
+	else if (Header.IsComplexMaterial())				Print(Ctx, TEXT("Complex "), InColor);
+	else if (Header.IsSingleLayerWater())				Print(Ctx, TEXT("SLWater "), InColor);
+	else if (Header.IsHair())							Print(Ctx, TEXT("Hair    "), InColor);
+	else if (Header.IsEye())							Print(Ctx, TEXT("Eye     "), InColor);
+	else												Print(Ctx, TEXT("Error - Unkown pixel type"), FontRed);
+}
+
+void PrintSSSType(inout FShaderPrintContext Ctx, uint SSSType, bool bThin, FFontColor InColor)
+{
+		 if (SSSType == SSS_TYPE_WRAP && bThin)		Print(Ctx, TEXT("Wrap & Thin"), InColor);
+	else if (SSSType == SSS_TYPE_WRAP)				Print(Ctx, TEXT("Wrap"), InColor);
+	else if (SSSType == SSS_TYPE_DIFFUSION)			Print(Ctx, TEXT("Diffusion"), InColor);
+	else if (SSSType == SSS_TYPE_DIFFUSION_PROFILE)	Print(Ctx, TEXT("Diffusion Profile"), InColor);
+	else if (SSSType == SSS_TYPE_SIMPLEVOLUME)		Print(Ctx, TEXT("Simple Volume"), InColor);
+	else if (bThin)									Print(Ctx, TEXT("Thin"), InColor);
+	else /*	(SSSType == SSS_TYPE_INVALID) */		Print(Ctx, TEXT("Invalid"), InColor);
+}
+
+void PrintBackgroundRect(inout FShaderPrintContext Ctx, float2 StartRect, float2 EndRect)
+{
+	const float4 RectBackgroundColor = float4(1, 1, 1, 0.25f);
+	AddFilledQuadSS(
+		(StartRect - ShaderPrintData.FontSize) * ShaderPrintData.Resolution,
+		(EndRect + ShaderPrintData.FontSize) * ShaderPrintData.Resolution,
+		RectBackgroundColor);
+}
+
+void Newline(inout FShaderPrintContext Context, inout float2 MaxRect)
+{
+	MaxRect.x = max(MaxRect.x, Context.Pos.x);
+	MaxRect.y = Context.Pos.y;
+	Newline(Context);
+}
+
+void PrintAddress(inout FShaderPrintContext Context, inout FSubstrateAddressing SubstrateAddressing)
+{
+	FFontColor Font;
+	Font.Color = lerp(FontEmerald.Color, FontSilver.Color, 0.75f);
+	Print(Context, TEXT("[Address="), Font);
+	Print(Context, SubstrateAddressing.CurrentIndex, Font, 2, 0);
+	Print(Context, TEXT("]"), Font);
+}
+
+void SubstratePrintBSDF(
+	inout FShaderPrintContext Context, 
+	inout FSubstrateAddressing SubstrateAddressing, 
+	FSubstratePixelHeader Header, 
+	FSubstrateSubsurfaceHeader SSSHeader,
+	float3 WorldPosition, 
+	float3 V, 
+	inout float2 RectMax, 
+	FSubstrateRaytracingPayload Payload)
+{
+	const FFontColor FontBSDF         = FontSilver;
+	const FFontColor FontBSDFType     = FontEmerald;
+	const FFontColor FontBSDFStateName= FontWhite;
+	const FFontColor FontBSDFPropName = FontOrange;
+	const FFontColor FontBSDFPropValu = FontWhite;
+
+	Newline(Context, RectMax);
+
+	#if SUBSTRATE_MATERIALCONTAINER_IS_VIEWRESOURCE
+	const FSubstrateBSDF BSDF = UnpackSubstrateBSDFIn(Substrate.MaterialTextureArray, SubstrateAddressing, Header);
+	#else
+	const FSubstrateBSDF BSDF = UnpackSubstrateBSDFIn(Payload, SubstrateAddressing, Header); // Take the parameter of the first BSDF
+	#endif
+
+	const float LineSize = 5.f;
+	// Draw Referential
+	#if SUBSTRATE_MATERIALCONTAINER_IS_VIEWRESOURCE
+	{
+		const float3 DummyL = float3(0, 0, 1);
+		FSubstrateBSDFContext BSDFContext = SubstrateCreateBSDFContext(Header, BSDF, SubstrateAddressing, V, DummyL);
+		DrawReferentialTWS(WorldPosition, BSDFContext.X, BSDFContext.Y, BSDFContext.N, float3(1, 1, 0));
+	}
+	#else
+	{
+		const float3x3 TangentBasis = SubstrateGetBSDFSharedBasis(Header, BSDF, SubstrateAddressing);
+		DrawReferentialTWS(WorldPosition, TangentBasis[0], TangentBasis[1], TangentBasis[2], float3(1, 1, 0));
+	}
+	#endif
+
+	switch (BSDF_GETTYPE(BSDF))
+	{
+		case SUBSTRATE_BSDF_TYPE_SLAB:
+		{
+			const bool bWeightL = BSDF_GETHASTRANSABOVE(BSDF);
+			const bool bGreyWeightV = BSDF_GETHASGREYWEIGHT_V(BSDF);
+			const bool bGreyTopTrans = BSDF.TransmittanceAboveAlongN.x == BSDF.TransmittanceAboveAlongN.y && BSDF.TransmittanceAboveAlongN.x == BSDF.TransmittanceAboveAlongN.z;
+
+			Print(Context, TEXT("Slab "), FontBSDFType);
+			PrintAddress(Context, SubstrateAddressing);
+			Newline(Context, RectMax);
+			Print(Context, TEXT(" NormalID BasisType Aniso TopLayer Scattering IsThin WeightV "), FontBSDFStateName);
+			if (!bGreyWeightV)
+			{
+				Print(Context, TEXT("            "), FontBSDFStateName);
+			}
+			if (bWeightL)
+			{
+				Print(Context, TEXT("TransToL     "), FontBSDFStateName);
+				if (!bGreyTopTrans)
+				{
+					Print(Context, TEXT("       "), FontBSDFStateName);
+				}
+				Print(Context, TEXT("CoverToL"), FontBSDFStateName);
+			}
+			Newline(Context, RectMax);
+
+			const uint BasisType  = SubstrateGetSharedLocalBasisType(Header.SharedLocalBasesTypes_PackedHeader, BSDF_GETSHAREDLOCALBASISID(BSDF));
+
+			Print(Context, TEXT(" "), FontSilver);
+			Print(Context, BSDF_GETSHAREDLOCALBASISID(BSDF), FontBSDF, 2, 1);	Print(Context, TEXT("       "), FontSilver);
+			if (BasisType == 0)
+			{
+				Print(Context, TEXT("Normal"), FontBSDF);	Print(Context, TEXT("    "), FontSilver);
+			}
+			else
+			{
+				Print(Context, TEXT("Tangent"), FontBSDF);	Print(Context, TEXT("   "), FontSilver);
+			}
+			PrintBool(Context, BSDF_GETHASANISOTROPY(BSDF));					Print(Context, TEXT("  "));
+			PrintBool(Context, BSDF_GETISTOPLAYER(BSDF));						Print(Context, TEXT("     "));
+			PrintBool(Context, BSDF_GETSSSTYPE(BSDF) != SSS_TYPE_INVALID);		Print(Context, TEXT("       "));
+			PrintBool(Context, BSDF_GETISTHIN(BSDF));							Print(Context, TEXT("   "));
+
+			// View Weight
+			{
+				if (bGreyWeightV)
+				{
+					Print(Context, BSDF.LuminanceWeightV.x, FontSilver, 6, 4);
+				}
+				else
+				{
+					Print(Context, BSDF.LuminanceWeightV.r, FontLightRed,   6, 4);
+					Print(Context, BSDF.LuminanceWeightV.g, FontLightGreen, 6, 4);
+					Print(Context, BSDF.LuminanceWeightV.b, FontLightBlue,  6, 4);
+				}
+				Print(Context, TEXT(" "), FontWhite);
+			}
+
+			// Light Weight
+			if (bWeightL)
+			{
+				Print(Context, TEXT(" "), FontSilver);
+				if (bGreyTopTrans)
+				{
+					Print(Context, BSDF.TransmittanceAboveAlongN.x, FontSilver);
+				}
+				else	
+				{
+					Print(Context, BSDF.TransmittanceAboveAlongN.r, FontLightRed,   6, 4);
+					Print(Context, BSDF.TransmittanceAboveAlongN.g, FontLightGreen, 6, 4);
+					Print(Context, BSDF.TransmittanceAboveAlongN.b, FontLightBlue,  6, 4);
+				}
+
+				Print(Context, TEXT(" "), FontSilver);
+				Print(Context, BSDF.CoverageAboveAlongN, FontSilver, 6, 4);
+			}
+			Newline(Context, RectMax);
+
+			Print(Context, TEXT(" Diffuse     "), FontBSDFPropName);	Print(Context, SLAB_DIFFUSEALBEDO(BSDF), FontBSDFPropValu);			Newline(Context, RectMax);
+
+			const bool bGreyF0 = SLAB_F0(BSDF).x == SLAB_F0(BSDF).y && SLAB_F0(BSDF).x == SLAB_F0(BSDF).z;
+			if (bGreyF0)
+			{
+				Print(Context, TEXT(" F0          "), FontBSDFPropName);	Print(Context, SLAB_F0(BSDF).x, FontBSDFPropValu);				Newline(Context, RectMax);
+			}
+			else
+			{
+				Print(Context, TEXT(" F0          "), FontBSDFPropName);	Print(Context, SLAB_F0(BSDF), FontBSDFPropValu);				Newline(Context, RectMax);
+			}
+
+
+			if (BSDF_GETHASF90(BSDF))
+			{
+				const bool bGreyF90 = SLAB_F90(BSDF).x == SLAB_F90(BSDF).y && SLAB_F90(BSDF).x == SLAB_F90(BSDF).z;
+				if (bGreyF90)
+				{
+					Print(Context, TEXT(" F90         "), FontBSDFPropName);	Print(Context, SLAB_F90(BSDF).x, FontBSDFPropValu);			Newline(Context, RectMax);
+				}
+				else
+				{
+					Print(Context, TEXT(" F90         "), FontBSDFPropName);	Print(Context, SLAB_F90(BSDF), FontBSDFPropValu);			Newline(Context, RectMax);
+				}
+			}
+
+			Print(Context, TEXT(" Roughness   "), FontBSDFPropName);	Print(Context, SLAB_ROUGHNESS(BSDF), FontBSDFPropValu);				Newline(Context, RectMax);
+
+			if (BSDF_GETHASANISOTROPY(BSDF))
+			{
+				Print(Context, TEXT(" Anisotropy  "), FontBSDFPropName);	Print(Context, SLAB_ANISOTROPY(BSDF), FontBSDFPropValu);		Newline(Context, RectMax);
+			}
+
+			if (BSDF_GETHASHAZINESS(BSDF))
+			{
+				FHaziness Haziness = UnpackHaziness(SLAB_HAZINESS(BSDF));
+
+				Print(Context, TEXT(" Haziness    "),			FontBSDFPropName);																		Newline(Context, RectMax);
+				Print(Context, TEXT("  - Roughness       "),	FontBSDFPropName);	Print(Context, Haziness.Roughness,				FontBSDFPropValu);	Newline(Context, RectMax);
+				Print(Context, TEXT("  - Weight          "),	FontBSDFPropName);	Print(Context, Haziness.Weight,					FontBSDFPropValu);	Newline(Context, RectMax);
+				Print(Context, TEXT("  - SimpleClearCoat "),	FontBSDFPropName);	Print(Context, Haziness.bSimpleClearCoat?1u:0u,	FontBSDFPropValu);	Newline(Context, RectMax);
+			}
+
+			if (BSDF_GETSSSTYPE(BSDF) != SSS_TYPE_INVALID || BSDF_GETISTHIN(BSDF))
+			{
+				const bool bIsThin = BSDF_GETISTHIN(BSDF);
+				Print(Context, TEXT(" SSS Type    "), FontBSDFPropName);	PrintSSSType(Context, BSDF_GETSSSTYPE(BSDF), BSDF_GETISTHIN(BSDF), FontYellow);	Newline(Context, RectMax);
+				if (BSDF_GETSSSTYPE(BSDF) == SSS_TYPE_WRAP)
+				{
+					const float Opacity = SubstrateSubSurfaceHeaderGetWrapOpacity(SSSHeader);
+					Print(Context, TEXT(" SSS Opacity "), FontBSDFPropName);	Print(Context, Opacity, FontBSDFPropValu);						Newline(Context, RectMax);
+					Print(Context, TEXT(" SSS MFP     "), FontBSDFPropName);	Print(Context, SLAB_SSSMFP(BSDF), FontBSDFPropValu);			Newline(Context, RectMax);
+					Print(Context, TEXT(" SSS Phase   "), FontBSDFPropName);	Print(Context, SLAB_SSSPHASEANISOTROPY(BSDF), FontBSDFPropValu);Newline(Context, RectMax);
+					Print(Context, TEXT(" SSS Thicknes"), FontBSDFPropName);	Print(Context, BSDF_GETTHICKNESSCM(BSDF), FontBSDFPropValu);	Newline(Context, RectMax);
+				}
+				else if (BSDF_GETSSSTYPE(BSDF) == SSS_TYPE_SIMPLEVOLUME)
+				{
+					const float Opacity = SubstrateSubSurfaceHeaderGetWrapOpacity(SSSHeader);
+					Print(Context, TEXT(" SSS Opacity "), FontBSDFPropName);	Print(Context, Opacity, FontBSDFPropValu);						Newline(Context, RectMax);
+					Print(Context, TEXT(" SSS MFP     "), FontBSDFPropName);	Print(Context, SLAB_SSSMFP(BSDF), FontBSDFPropValu);			Newline(Context, RectMax);
+					Print(Context, TEXT(" SSS Phase   "), FontBSDFPropName);	Print(Context, SLAB_SSSPHASEANISOTROPY(BSDF), FontBSDFPropValu);Newline(Context, RectMax);
+					Print(Context, TEXT(" SSS Thicknes"), FontBSDFPropName);	Print(Context, BSDF_GETTHICKNESSCM(BSDF), FontBSDFPropValu);	Newline(Context, RectMax);
+				}
+				else if (BSDF_GETSSSTYPE(BSDF) == SSS_TYPE_DIFFUSION_PROFILE)
+				{
+					const uint ProfileId = SubstrateSubSurfaceHeaderGetProfileId(SSSHeader);
+					const float RadiusScale = SubstrateSubSurfaceHeaderGetProfileRadiusScale(SSSHeader);
+					const float3 DiffuseMFP = GetSubsurfaceProfileMFPInCm(ProfileId).xyz * RadiusScale;
+					const float Phase = GetTransmissionProfileParams(ProfileId).ScatteringDistribution;
+
+					Print(Context, TEXT(" SSS ID.     "), FontBSDFPropName);	Print(Context, ProfileId, FontBSDFPropValu);	Newline(Context, RectMax);
+					Print(Context, TEXT(" SSS Radius  "), FontBSDFPropName);	Print(Context, RadiusScale, FontBSDFPropValu);	Newline(Context, RectMax);
+					Print(Context, TEXT(" SSS MFP     "), FontBSDFPropName);	Print(Context, DiffuseMFP, FontBSDFPropValu);	Newline(Context, RectMax);
+					Print(Context, TEXT(" SSS Phase   "), FontBSDFPropName);	Print(Context, Phase, FontBSDFPropValu);		Newline(Context, RectMax);
+					if (bIsThin)
+					{
+					Print(Context, TEXT(" SSS Thicknes"), FontBSDFPropName);	Print(Context, SLAB_SSSPROFILETHICKNESSCM(BSDF), FontBSDFPropValu);		Newline(Context, RectMax);
+					}
+				}
+				else if (BSDF_GETSSSTYPE(BSDF) == SSS_TYPE_DIFFUSION)
+				{
+					const float3 OriginalMFP = SubstrateSubSurfaceHeaderGetMFP(SSSHeader);
+					const float3 RescaledMFP = SLAB_SSSMFP(BSDF);
+
+					if (bIsThin)
+					{
+						Print(Context, TEXT(" SSS MFP     "), FontBSDFPropName);	Print(Context, OriginalMFP, FontBSDFPropValu);				Newline(Context, RectMax);
+						Print(Context, TEXT(" SSS MFP Norm"), FontBSDFPropName);	Print(Context, RescaledMFP, FontBSDFPropValu);				Newline(Context, RectMax);
+						Print(Context, TEXT(" SSS Thicknes"), FontBSDFPropName);	Print(Context, BSDF_GETTHICKNESSCM(BSDF), FontBSDFPropValu);Newline(Context, RectMax);
+					}
+					else
+					{
+						Print(Context, TEXT(" SSS MFP     "), FontBSDFPropName);	Print(Context, OriginalMFP, FontBSDFPropValu);				Newline(Context, RectMax);
+					}
+
+					Print(Context, TEXT(" SSS Phase   "), FontBSDFPropName);	Print(Context, SLAB_SSSPHASEANISOTROPY(BSDF), FontBSDFPropValu);Newline(Context, RectMax);
+				}
+			} // SSS
+
+			if (BSDF_GETHASFUZZ(BSDF))
+			{
+				Print(Context, TEXT(" Fuzz Amount "), FontBSDFPropName);	Print(Context, SLAB_FUZZ_AMOUNT(BSDF), FontBSDFPropValu);				Newline(Context, RectMax);
+				Print(Context, TEXT(" Fuzz Color  "), FontBSDFPropName);	Print(Context, SLAB_FUZZ_COLOR(BSDF), FontBSDFPropValu);				Newline(Context, RectMax);
+				Print(Context, TEXT(" Fuzz Rough. "), FontBSDFPropName);	Print(Context, SLAB_FUZZ_ROUGHNESS(BSDF), FontBSDFPropValu);			Newline(Context, RectMax);
+			}
+		#if SUBSTRATE_COMPLEXSPECIALPATH
+			if (BSDF_GETHASGLINT(BSDF))
+			{
+				Print(Context, TEXT(" Glint Value "), FontBSDFPropName);	Print(Context, SLAB_GLINT_VALUE(BSDF), FontBSDFPropValu);				Newline(Context, RectMax);
+				Print(Context, TEXT(" Glint UVs   "), FontBSDFPropName);	Print(Context, SLAB_GLINT_UV(BSDF), FontBSDFPropValu);					Newline(Context, RectMax);
+				#if 1
+				Print(Context, TEXT(" Glint UVsDx "), FontBSDFPropName);	Print(Context, SLAB_GLINT_UVDDX(BSDF), FontBSDFPropValu);				Newline(Context, RectMax);
+				Print(Context, TEXT(" Glint UVsDy "), FontBSDFPropName);	Print(Context, SLAB_GLINT_UVDDY(BSDF), FontBSDFPropValu);				Newline(Context, RectMax);
+				#endif
+			}
+			if (BSDF_GETHASSPECPROFILE(BSDF))
+			{
+				Print(Context, TEXT(" SpecID      "), FontBSDFPropName);	Print(Context, GetSpecularProfileId(SLAB_SPECPROFILEID(BSDF)), FontBSDFPropValu);Newline(Context, RectMax);
+				if (GetSpecularProfileParameterization(SLAB_SPECPROFILEID(BSDF)) == 0)
+				{
+				Print(Context, TEXT(" Spec Param. "), FontBSDFPropName);	Print(Context, TEXT("View/Light angles"), FontBSDFPropValu);			Newline(Context, RectMax);
+				}
+				else
+				{
+				Print(Context, TEXT(" Spec Param. "), FontBSDFPropName);	Print(Context, TEXT("Half angles"), FontBSDFPropValu);					Newline(Context, RectMax);
+				}
+			}
+		#endif
+		}
+		break;
+
+		case SUBSTRATE_BSDF_TYPE_HAIR:
+		{
+			Print(Context, TEXT("Hair"), FontBSDFType);
+			PrintAddress(Context, SubstrateAddressing);
+			Newline(Context, RectMax);
+			Print(Context, TEXT(" NormalID Aniso IsTopLayer LuminanceWeight"), FontBSDFStateName);
+			Newline(Context, RectMax);
+
+			Print(Context, TEXT(" "), FontSilver);
+			Print(Context, BSDF_GETSHAREDLOCALBASISID(BSDF), FontBSDF);
+			PrintBool(Context, BSDF_GETHASANISOTROPY(BSDF));
+			PrintBool(Context, BSDF_GETISTOPLAYER(BSDF));
+			if (BSDF_GETHASGREYWEIGHT_V(BSDF))
+			{
+				Print(Context, BSDF.LuminanceWeightV.x, FontSilver);
+			}
+			else
+			{
+				Print(Context, BSDF.LuminanceWeightV.r, FontLightRed);
+				Print(Context, BSDF.LuminanceWeightV.g, FontLightGreen);
+				Print(Context, BSDF.LuminanceWeightV.b, FontLightBlue);
+			}
+			Newline(Context, RectMax);
+
+			Print(Context, TEXT(" BaseColor         "), FontBSDFPropName);	Print(Context, HAIR_BASECOLOR(BSDF), FontBSDFPropValu);				Newline(Context, RectMax);
+			Print(Context, TEXT(" Specular          "), FontBSDFPropName);	Print(Context, HAIR_SPECULAR(BSDF), FontBSDFPropValu);				Newline(Context, RectMax);
+			Print(Context, TEXT(" Roughness         "), FontBSDFPropName);	Print(Context, HAIR_ROUGHNESS(BSDF), FontBSDFPropValu);				Newline(Context, RectMax);
+			Print(Context, TEXT(" Scatter           "), FontBSDFPropName);	Print(Context, HAIR_SCATTER(BSDF), FontBSDFPropValu);				Newline(Context, RectMax);
+			Print(Context, TEXT(" Backlit           "), FontBSDFPropName);	Print(Context, HAIR_BACKLIT(BSDF), FontBSDFPropValu);				Newline(Context, RectMax);
+			Print(Context, TEXT(" Has Transmittance "), FontBSDFPropName);	Print(Context, HAIR_COMPLEXTRANSMITTANCE(BSDF), FontBSDFPropValu);	Newline(Context, RectMax);
+		}
+		break;
+
+		case SUBSTRATE_BSDF_TYPE_EYE:
+		{
+			Print(Context, TEXT("Eye"), FontBSDFType);
+			PrintAddress(Context, SubstrateAddressing);
+			Newline(Context, RectMax);
+			Print(Context, TEXT(" NormalID IsTopLayer LuminanceWeight"), FontBSDFStateName);
+			Newline(Context, RectMax);
+
+			Print(Context, TEXT(" "), FontSilver);
+			Print(Context, BSDF_GETSHAREDLOCALBASISID(BSDF), FontBSDF);
+			PrintBool(Context, BSDF_GETISTOPLAYER(BSDF));
+			if (BSDF_GETHASGREYWEIGHT_V(BSDF))
+			{
+				Print(Context, BSDF.LuminanceWeightV.x, FontSilver);
+			}
+			else
+			{
+				Print(Context, BSDF.LuminanceWeightV.r, FontLightRed);
+				Print(Context, BSDF.LuminanceWeightV.g, FontLightGreen);
+				Print(Context, BSDF.LuminanceWeightV.b, FontLightBlue);
+			}
+			Newline(Context, RectMax);
+
+			Print(Context, TEXT(" Diffuse Albedo    "), FontBSDFPropName);	Print(Context, EYE_DIFFUSEALBEDO(BSDF), FontBSDFPropValu);			Newline(Context, RectMax);
+			Print(Context, TEXT(" F0                "), FontBSDFPropName);	Print(Context, EYE_F0(BSDF), FontBSDFPropValu);						Newline(Context, RectMax);
+			Print(Context, TEXT(" Roughness         "), FontBSDFPropName);	Print(Context, EYE_ROUGHNESS(BSDF), FontBSDFPropValu);				Newline(Context, RectMax);
+			Print(Context, TEXT(" Iris Mask         "), FontBSDFPropName);	Print(Context, EYE_IRISMASK(BSDF), FontBSDFPropValu);			Newline(Context, RectMax);
+			Print(Context, TEXT(" Iris Distance     "), FontBSDFPropName);	Print(Context, EYE_IRISDISTANCE(BSDF), FontBSDFPropValu);				Newline(Context, RectMax);
+
+			const float LineSize = 10.f;
+
+			// Tangent basis for the iris
+			const float3x3 IrisTangentBasis = GetTangentBasis(EYE_IRISNORMAL(BSDF));
+			//DrawReferentialTWS(WorldPosition, IrisTangentBasis[0], IrisTangentBasis[1], IrisTangentBasis[2], float3(0, 1, 1));
+			AddLineTWS(WorldPosition, WorldPosition + EYE_IRISNORMAL(BSDF) * LineSize, float4(0, 1, 1, 1));
+
+			// Tangent basis for the iris plane
+			const float3x3 IrisPlaneTangentBasis = GetTangentBasis(EYE_IRISPLANENORMAL(BSDF));
+			//DrawReferentialTWS(WorldPosition, IrisPlaneTangentBasis[0], IrisPlaneTangentBasis[1], IrisPlaneTangentBasis[2], float3(1, 0, 1));
+			AddLineTWS(WorldPosition, WorldPosition + EYE_IRISPLANENORMAL(BSDF) * LineSize, float4(1, 0, 1, 1));
+
+
+			Print(Context, TEXT(" Iris N            "), FontBSDFPropName);	Print(Context, EYE_IRISNORMAL(BSDF), FontBSDFPropValu);			Newline(Context, RectMax);
+			Print(Context, TEXT(" PlaneN            "), FontBSDFPropName);	Print(Context, EYE_IRISPLANENORMAL(BSDF), FontBSDFPropValu);			Newline(Context, RectMax);
+
+		#if SUBSTRATE_MATERIALCONTAINER_IS_VIEWRESOURCE
+			const bool bHasSSS = Header.HasSubsurface();
+			const bool bIsValid = SubstrateSubSurfaceHeaderGetIsValid(SSSHeader);
+			const bool bIsProfile = SubstrateSubSurfaceHeaderGetIsProfile(SSSHeader);
+			if (bHasSSS && bIsValid && bIsProfile)
+			{
+				const uint SubsurfaceProfileInt = SubstrateSubSurfaceHeaderGetProfileId(SSSHeader);
+				Print(Context, TEXT(" SSS ID.           "), FontBSDFPropName);	Print(Context, SubsurfaceProfileInt, FontBSDFPropValu);			Newline(Context, RectMax);
+			}
+		#endif
+		}
+		break;
+
+		case SUBSTRATE_BSDF_TYPE_SINGLELAYERWATER:
+		{
+			Print(Context, TEXT("Water"), FontBSDFType);
+			PrintAddress(Context, SubstrateAddressing);
+			Newline(Context);
+
+			Print(Context, TEXT(" BaseColor            "), FontBSDFPropName);	Print(Context, SLW_BASECOLOR(BSDF), FontBSDFPropValu);			Newline(Context, RectMax);
+			Print(Context, TEXT(" Metallic             "), FontBSDFPropName);	Print(Context, SLW_METALLIC(BSDF), FontBSDFPropValu);			Newline(Context, RectMax);
+			Print(Context, TEXT(" Specular             "), FontBSDFPropName);	Print(Context, SLW_SPECULAR(BSDF), FontBSDFPropValu);			Newline(Context, RectMax);
+			Print(Context, TEXT(" Roughness            "), FontBSDFPropName);	Print(Context, SLW_ROUGHNESS(BSDF), FontBSDFPropValu);			Newline(Context, RectMax);
+			Print(Context, TEXT(" Top Material Opacity "), FontBSDFPropName);	Print(Context, SLW_TOPMATERIALOPACITY(BSDF), FontBSDFPropValu);	Newline(Context, RectMax);
+		}
+		break;
+
+	default:
+		{
+			Print(Context, TEXT("Error - Uknown BSDF type"), FontRed);
+		}
+		break;
+	}
+
+	Newline(Context);
+}
+
+void SubstratePrintMaterialProperties_Internal(
+	inout FShaderPrintContext Context, 
+	uint2 InCoord, 
+	float3 InWorldPosition, 
+	float3 V, 
+	uint InClosureIndex, 
+	inout FSubstrateAddressing SubstrateAddressing, 
+	FSubstratePixelHeader Header,
+	FSubstrateSubsurfaceHeader SSSHeader,
+	FSubstrateTopLayerData TopLayerData,
+	FSubstrateRaytracingPayload Payload)
+{
+	const uint FootPrint_Start = SubstrateAddressing.ReadBytes;
+	const uint FootPrint_PostHeader = SubstrateAddressing.ReadBytes;
+
+	const FFontColor FontBSDFCount = FontEmerald;
+	const FFontColor FontHeaderStateName = FontWhite;
+	const FFontColor FontHeaderPropName = FontWhite;
+	const FFontColor FontHeaderPropValue = FontSilver;
+
+	// Backgroun rect. for better print visibility. Track min/max point
+	float2 RectMin = Context.Pos;
+	float2 RectMax = Context.StartPos;
+
+	// Header
+	if (InClosureIndex < Header.ClosureCount && InClosureIndex == 0)
+	{	
+		const bool bSubstrateMaterial = Header.ClosureCount > 0;
+		const bool bIsSimpleMaterial = Header.IsSimpleMaterial() || Header.ClosureCount == 0;
+		const bool bIsSingleMaterial = !Header.IsSimpleMaterial() && Header.IsSingleMaterial();
+
+		FSubstrateIrradianceAndOcclusion IrradianceAndOcclusion = SubstrateGetIrradianceAndAO(Header);
+
+		RectMin = Context.Pos;
+		Print(Context, Header.ClosureCount, FontBSDFCount, 2, 0);
+		Print(Context, TEXT("BSDF"), FontBSDFCount);
+		Print(Context, TEXT(" - "), FontWhite);
+		PrintPixelType(Context, Header, FontLightRed);
+		if (bIsSingleMaterial)
+		{
+		#if SUBSTRATE_DEFERRED_SHADING
+			uint OptimisedLegacyMode = (Header.SharedLocalBasesTypes_PackedHeader >> (HEADER_SINGLEENCODING_BIT_COUNT)) & HEADER_SINGLE_OPTLEGACYMODE_BIT_MASK;
+			if (OptimisedLegacyMode ==SINGLE_OPTLEGACYMODE_CLEARCOAT)
+			{
+				Print(Context, TEXT(" - Legacy Clear-Coat"), FontLightRed);
+			}
+			else if (OptimisedLegacyMode ==SINGLE_OPTLEGACYMODE_CLOTH)
+			{
+				Print(Context, TEXT(" - Legacy Cloth"), FontLightRed);
+			}
+			else if (OptimisedLegacyMode ==SINGLE_OPTLEGACYMODE_SSSWRAP)
+			{
+				Print(Context, TEXT(" - Legacy SSS-Wrap"), FontLightRed);
+			}
+			else if (OptimisedLegacyMode == SINGLE_OPTLEGACYMODE_SSSWRAP_THIN)
+			{
+				Print(Context, TEXT(" - Legacy Foliage"), FontLightRed);
+			}
+			else if (OptimisedLegacyMode == SINGLE_OPTLEGACYMODE_SSSPROFILE)
+			{
+				Print(Context, TEXT(" - Legacy SSS-Profile"), FontLightRed);
+			}
+		#endif
+		}
+		Newline(Context, RectMax);
+
+		Print(Context, TEXT(" AO    IndIrr TopRoughness PreShadow ZeroShadow ContacShadow Ind.Occluder HasSSS"), FontHeaderStateName);
+		Newline(Context, RectMax);
+
+		Print(Context, TEXT(" "));
+		Print(Context, IrradianceAndOcclusion.MaterialAO, FontHeaderPropValue, 5, 3);					Print(Context, TEXT(" "));
+		Print(Context, IrradianceAndOcclusion.IndirectIrradiance, FontHeaderPropValue, 6, 3);			Print(Context, TEXT(" "));
+		Print(Context, TopLayerData.Roughness, FontHeaderPropValue, 5, 3);								Print(Context, TEXT("        "));
+		PrintBool(Context, Header.HasPrecShadowMask());													Print(Context, TEXT("      "));
+		PrintBool(Context, Header.HasZeroPrecShadowMask());												Print(Context, TEXT("       "));
+		PrintBool(Context, Header.DoesCastContactShadow());												Print(Context, TEXT("         "));
+		PrintBool(Context, Header.HasDynamicIndirectShadowCasterRepresentation());						Print(Context, TEXT("         "));
+		PrintBool(Context, Header.HasSubsurface());
+		Newline(Context, RectMax);
+		Newline(Context);
+		
+		PrintBackgroundRect(Context, RectMin, RectMax);
+	}
+
+	// BSDFs
+	if (InClosureIndex < Header.ClosureCount)
+	{
+		#if SUBSTRATE_MATERIALCONTAINER_IS_VIEWRESOURCE
+		if (InClosureIndex > 0)
+		{
+			const uint AddressOffset = UnpackClosureOffsetAtIndex(Substrate.ClosureOffsetTexture[InCoord], InClosureIndex, Header.ClosureCount);
+			SubstrateSeekClosure(SubstrateAddressing, AddressOffset);
+		}
+		#endif
+
+		RectMin = Context.Pos;
+		RectMax = Context.StartPos;
+		SubstratePrintBSDF(Context, SubstrateAddressing, Header, SSSHeader, InWorldPosition, V, RectMax, Payload);
+		PrintBackgroundRect(Context, RectMin, RectMax);
+	}
+
+	// Memory footprint
+	if (InClosureIndex+1 == Header.ClosureCount)
+	{
+		const bool bSubstrateMaterial = Header.ClosureCount > 0;
+		const bool bIsSimpleMaterial = Header.IsSimpleMaterial() || Header.ClosureCount == 0;
+		const bool bIsSingleMaterial = !Header.IsSimpleMaterial() && Header.IsSingleMaterial();
+
+		const uint FootPrint_PostBSDFs = SubstrateAddressing.ReadBytes;
+
+		const bool bHasSSSData = SubstrateSubSurfaceHeaderGetIsValid(SSSHeader);
+
+		const uint TopLayerDataBytes = 4;	// bytes 
+		const uint SSSDataBytes = 8;
+
+		const uint HeaderSize =  FootPrint_PostHeader - FootPrint_Start;
+		const uint BSDFsSize  =  FootPrint_PostBSDFs - FootPrint_PostHeader;
+		const uint TotalSize  = (FootPrint_PostBSDFs - FootPrint_Start) + (bHasSSSData ? SSSDataBytes : 0);
+
+		RectMin = Context.Pos;
+		RectMax = Context.StartPos;
+		Print(Context, TEXT("Memory"), FontEmerald); Newline(Context, RectMax);
+		if (bIsSingleMaterial || Header.IsEye() || Header.IsHair())
+		{
+			Print(Context, TEXT(" Header       "), FontSilver); Print(Context, HeaderSize - TopLayerDataBytes,				FontSilver, 3, 0); Print(Context, TEXT("bytes"), FontSilver); Newline(Context, RectMax);
+			Print(Context, TEXT(" TopNormalTex "), FontSilver); Print(Context, TopLayerDataBytes,							FontSilver, 3, 0); Print(Context, TEXT("bytes"), FontSilver); Newline(Context, RectMax);
+			Print(Context, TEXT(" BSDF         "), FontSilver); Print(Context, BSDFsSize,									FontSilver, 3, 0); Print(Context, TEXT("bytes"), FontSilver); Newline(Context, RectMax);
+		}
+		else if (bIsSimpleMaterial || Header.IsSingleLayerWater())
+		{
+			Print(Context, TEXT(" Header+BSDF  "), FontSilver); Print(Context, HeaderSize + BSDFsSize - TopLayerDataBytes,	FontSilver, 3, 0); Print(Context, TEXT("bytes"), FontSilver); Newline(Context, RectMax);
+			Print(Context, TEXT(" TopNormalTex "), FontSilver); Print(Context, TopLayerDataBytes,							FontSilver, 3, 0); Print(Context, TEXT("bytes"), FontSilver); Newline(Context, RectMax);
+		}
+		else
+		{
+			Print(Context, TEXT(" Header+Norm  "), FontSilver); Print(Context, HeaderSize,									FontSilver, 3, 0); Print(Context, TEXT("bytes"), FontSilver); Newline(Context, RectMax);
+			Print(Context, TEXT(" BSDFs        "), FontSilver); Print(Context, BSDFsSize,									FontSilver, 3, 0); Print(Context, TEXT("bytes"), FontSilver); Newline(Context, RectMax);
+		}
+
+		if (bHasSSSData)
+		{
+			Print(Context, TEXT(" SSS Data     "), FontSilver); Print(Context, SSSDataBytes, FontSilver, 3, 0); Print(Context, TEXT("bytes"), FontSilver); Newline(Context, RectMax);
+		}
+
+		Print(Context, TEXT(" Total        "), FontOrange); Print(Context, TotalSize,  FontOrange, 3, 0); Print(Context, TEXT("bytes"), FontOrange); Newline(Context, RectMax);
+		PrintBackgroundRect(Context, RectMin, RectMax);
+	}
+}
+
+#if SUBSTRATE_MATERIALCONTAINER_IS_VIEWRESOURCE
+void SubstratePrintMaterialProperties(inout FShaderPrintContext Context, uint2 InCoord, float3 InWorldPosition, float3 V, uint InClosureIndex)
+{
+	FSubstrateAddressing SubstrateAddressing = GetSubstratePixelDataByteOffset(InCoord, uint2(View.BufferSizeAndInvSize.xy), Substrate.MaxBytesPerPixel);
+	FSubstratePixelHeader Header = UnpackSubstrateHeaderIn(Substrate.MaterialTextureArray, SubstrateAddressing, Substrate.TopLayerTexture);
+	FSubstrateSubsurfaceHeader SSSHeader = SubstrateLoadSubsurfaceHeader(Substrate.MaterialTextureArray, Substrate.FirstSliceStoringSubstrateSSSData, SubstrateAddressing.PixelCoords);
+	FSubstrateRaytracingPayload DummyPayload;
+	FSubstrateTopLayerData TopLayerData = SubstrateUnpackTopLayerData(Substrate.TopLayerTexture.Load(uint3(InCoord, 0)));
+
+	SubstratePrintMaterialProperties_Internal(
+		Context, 
+		InCoord, 
+		InWorldPosition, 
+		V, 
+		InClosureIndex, 
+		SubstrateAddressing,
+		Header,
+		SSSHeader,
+		TopLayerData,
+		DummyPayload);
+}
+#else
+void SubstratePrintMaterialProperties(inout FShaderPrintContext Context, uint2 InCoord, float3 InWorldPosition, float3 V, uint InClosureIndex, FSubstrateRaytracingPayload Payload)
+{
+	FSubstrateAddressing SubstrateAddressing = GetSubstratePixelDataByteOffset(InCoord, uint2(View.BufferSizeAndInvSize.xy), 0);	
+	FSubstratePixelHeader Header = UnpackSubstrateHeaderIn(Payload, SubstrateAddressing, Payload);
+	FSubstrateSubsurfaceHeader SSSHeader = (FSubstrateSubsurfaceHeader)0;
+	FSubstrateTopLayerData TopLayerData = SubstrateUnpackTopLayerData(Payload.PackedTopLayerData);
+
+	for (uint ClosureIndex=0;ClosureIndex<Header.ClosureCount;ClosureIndex++)
+	{
+		SubstratePrintMaterialProperties_Internal(
+			Context, 
+			InCoord, 
+			InWorldPosition, 
+			V, 
+			InClosureIndex, 
+			SubstrateAddressing,
+			Header,
+			SSSHeader,
+			TopLayerData,
+			Payload);
+	}
+}
+#endif
\ No newline at end of file
diff --git a/Engine/Shaders/Shared/RaytracingDebugDefinitions.h b/Engine/Shaders/Shared/RaytracingDebugDefinitions.h
--- a/Engine/Shaders/Shared/RaytracingDebugDefinitions.h
+++ b/Engine/Shaders/Shared/RaytracingDebugDefinitions.h
@@ -44,7 +44,8 @@
 #define RAY_TRACING_DEBUG_VIZ_INSTANCE_OVERLAP			34
 #define RAY_TRACING_DEBUG_VIZ_TRIANGLE_HITCOUNT			35
 #define RAY_TRACING_DEBUG_VIZ_HITCOUNT_PER_INSTANCE		36
-#define RAY_TRACING_DEBUG_VIZ_LIGHT_GRID_COUNT          37
+#define RAY_TRACING_DEBUG_VIZ_LIGHT_GRID_COUNT			37
+#define RAY_TRACING_DEBUG_VIZ_SUBSTRATE_DATA			38
 
 #define RAY_TRACING_DEBUG_PICKER_DOMAIN_TRIANGLE		0
 #define RAY_TRACING_DEBUG_PICKER_DOMAIN_INSTANCE		1
diff --git a/Engine/Source/Editor/UnrealEd/Private/RayTracingDebugVisualizationMenuCommands.cpp b/Engine/Source/Editor/UnrealEd/Private/RayTracingDebugVisualizationMenuCommands.cpp
--- a/Engine/Source/Editor/UnrealEd/Private/RayTracingDebugVisualizationMenuCommands.cpp
+++ b/Engine/Source/Editor/UnrealEd/Private/RayTracingDebugVisualizationMenuCommands.cpp
@@ -80,6 +80,10 @@ void FRayTracingDebugVisualizationMenuCommands::CreateRayTracingDebugVisualizati
 		RayTracingDebugModeNames.Add(LOCTEXT("Proxy Type", "Proxy Type"));
 		RayTracingDebugModeNames.Add(LOCTEXT("Picker", "Picker"));
 		RayTracingDebugModeNames.Add(LOCTEXT("Light Grid Occupancy", "Light Grid Occupancy"));
+		if (Substrate::IsSubstrateEnabled())
+		{
+			RayTracingDebugModeNames.Emplace(LOCTEXT("Subtrate Material Properties", "Subtrate Material Properties"));
+		}
 	}
 	for ( int32 RayTracingDebugIndex = 0; RayTracingDebugIndex < RayTracingDebugModeNames.Num(); ++RayTracingDebugIndex)
 	{
diff --git a/Engine/Source/Runtime/Renderer/Private/DeferredShadingRenderer.cpp b/Engine/Source/Runtime/Renderer/Private/DeferredShadingRenderer.cpp
--- a/Engine/Source/Runtime/Renderer/Private/DeferredShadingRenderer.cpp
+++ b/Engine/Source/Runtime/Renderer/Private/DeferredShadingRenderer.cpp
@@ -3554,13 +3554,14 @@ bool ShouldRenderRayTracingEffect(bool bEffectEnabled, ERayTracingPipelineCompat
 	}
 }
 
+bool HasRaytracingDebugViewModeRaytracedOverlay(const FSceneViewFamily& ViewFamily);
 bool HasRayTracedOverlay(const FSceneViewFamily& ViewFamily)
 {
 	// Return true if a full screen ray tracing pass will be displayed on top of the raster pass
 	// This can be used to skip certain calculations
 	return
 		ViewFamily.EngineShowFlags.PathTracing ||
-		ViewFamily.EngineShowFlags.RayTracingDebug;
+		(ViewFamily.EngineShowFlags.RayTracingDebug && HasRaytracingDebugViewModeRaytracedOverlay(ViewFamily));
 }
 
 void FDeferredShadingSceneRenderer::InitializeRayTracingFlags_RenderThread()
diff --git a/Engine/Source/Runtime/Renderer/Private/RayTracing/RayTracingDebug.cpp b/Engine/Source/Runtime/Renderer/Private/RayTracing/RayTracingDebug.cpp
--- a/Engine/Source/Runtime/Renderer/Private/RayTracing/RayTracingDebug.cpp
+++ b/Engine/Source/Runtime/Renderer/Private/RayTracing/RayTracingDebug.cpp
@@ -159,6 +159,7 @@ class FRayTracingDebugRGS : public FGlobalShader
 		SHADER_PARAMETER_RDG_UNIFORM_BUFFER(FSceneUniformParameters, SceneUniformBuffer)
 		SHADER_PARAMETER_RDG_UNIFORM_BUFFER(FRayTracingDebugHitStatsUniformBufferParameters, RayTracingDebugHitStatsUniformBuffer)
 		SHADER_PARAMETER_RDG_UNIFORM_BUFFER(FRayTracingLightGrid, LightGridPacked)
+		SHADER_PARAMETER_STRUCT_INCLUDE(ShaderPrint::FShaderParameters, ShaderPrintUniformBuffer)
 	END_SHADER_PARAMETER_STRUCT()
 
 	static bool ShouldCompilePermutation(const FGlobalShaderPermutationParameters& Parameters)
@@ -1329,7 +1330,7 @@ static void DrawInstanceOverlap(FRDGBuilder& GraphBuilder, const FScene* Scene,
 	}
 }
 
-void FDeferredShadingSceneRenderer::RenderRayTracingDebug(FRDGBuilder& GraphBuilder, const FViewInfo& View, FRDGTextureRef SceneColorTexture, FRayTracingPickingFeedback& PickingFeedback)
+static uint32 GetRaytracingDebugViewMode(const FSceneView& View)
 {
 	static TMap<FName, uint32> RayTracingDebugVisualizationModes;
 	if (RayTracingDebugVisualizationModes.Num() == 0)
@@ -1371,27 +1372,50 @@ void FDeferredShadingSceneRenderer::RenderRayTracingDebug(FRDGBuilder& GraphBuil
 		RayTracingDebugVisualizationModes.Emplace(FName(*LOCTEXT("Instance Overlap", "Instance Overlap").ToString()),							RAY_TRACING_DEBUG_VIZ_INSTANCE_OVERLAP);
 		RayTracingDebugVisualizationModes.Emplace(FName(*LOCTEXT("Triangle Hit Count", "Triangle Hit Count").ToString()),						RAY_TRACING_DEBUG_VIZ_TRIANGLE_HITCOUNT);
 		RayTracingDebugVisualizationModes.Emplace(FName(*LOCTEXT("Hit Count Per Instance", "Hit Count Per Instance").ToString()),				RAY_TRACING_DEBUG_VIZ_HITCOUNT_PER_INSTANCE);
-		RayTracingDebugVisualizationModes.Emplace(FName(*LOCTEXT("Light Grid Occupancy", "Light Grid Occupancy").ToString()),                   RAY_TRACING_DEBUG_VIZ_LIGHT_GRID_COUNT);
+		RayTracingDebugVisualizationModes.Emplace(FName(*LOCTEXT("Light Grid Occupancy", "Light Grid Occupancy").ToString()),					RAY_TRACING_DEBUG_VIZ_LIGHT_GRID_COUNT);
+		if (Substrate::IsSubstrateEnabled())
+		{
+			RayTracingDebugVisualizationModes.Emplace(FName(*LOCTEXT("Subtrate Material Properties", "Subtrate Material Properties").ToString()),RAY_TRACING_DEBUG_VIZ_SUBSTRATE_DATA);
+		}
 	}
 
-	uint32 DebugVisualizationMode;
+	uint32 OutMode;
 	
 	FString ConsoleViewMode = CVarRayTracingDebugMode.GetValueOnRenderThread();
 
 	if (!ConsoleViewMode.IsEmpty())
 	{
-		DebugVisualizationMode = RayTracingDebugVisualizationModes.FindRef(FName(*ConsoleViewMode));
+		OutMode = RayTracingDebugVisualizationModes.FindRef(FName(*ConsoleViewMode));
 	}
 	else if(View.CurrentRayTracingDebugVisualizationMode != NAME_None)
 	{
-		DebugVisualizationMode = RayTracingDebugVisualizationModes.FindRef(View.CurrentRayTracingDebugVisualizationMode);
+		OutMode = RayTracingDebugVisualizationModes.FindRef(View.CurrentRayTracingDebugVisualizationMode);
 	}
 	else
 	{
 		// Set useful default value
-		DebugVisualizationMode = RAY_TRACING_DEBUG_VIZ_BARYCENTRICS;
+		OutMode = RAY_TRACING_DEBUG_VIZ_BARYCENTRICS;
 	}
 
+	return OutMode;
+}
+
+bool HasRaytracingDebugViewModeRaytracedOverlay(const FSceneViewFamily& ViewFamily)
+{
+	for (const FSceneView* View : ViewFamily.Views)
+	{
+		if (View && GetRaytracingDebugViewMode(*View) != RAY_TRACING_DEBUG_VIZ_SUBSTRATE_DATA)
+		{
+			return true;
+		}
+	}
+	return false;
+}
+
+void FDeferredShadingSceneRenderer::RenderRayTracingDebug(FRDGBuilder& GraphBuilder, const FViewInfo& View, FRDGTextureRef SceneColorTexture, FRayTracingPickingFeedback& PickingFeedback)
+{
+	const uint32 DebugVisualizationMode = GetRaytracingDebugViewMode(View);
+
 	if (DebugVisualizationMode == RAY_TRACING_DEBUG_VIZ_BARYCENTRICS)
 	{
 		return RenderRayTracingBarycentrics(GraphBuilder, View, SceneColorTexture, (bool)GVisualizeProceduralPrimitives);
@@ -1568,7 +1592,7 @@ void FDeferredShadingSceneRenderer::RenderRayTracingDebug(FRDGBuilder& GraphBuil
 	RayGenParameters->LightGridPacked = View.RayTracingLightGridUniformBuffer;
 	RayGenParameters->TopKMostHitInstances = CVarRayTracingDebugHitCountTopKHits.GetValueOnRenderThread();
 	RayGenParameters->NumTotalInstances = NumInstances;
-
+	ShaderPrint::SetParameters(GraphBuilder, View.ShaderPrintData, RayGenParameters->ShaderPrintUniformBuffer);
 	
 	// If we don't output depth, create dummy 1x1 texture
 	const bool bOutputDepth = DebugVisualizationMode == RAY_TRACING_DEBUG_VIZ_INSTANCE_OVERLAP;
