Fix typo in return type.

diff --git a/Engine/Shaders/Private/ParticipatingMediaCommon.ush b/Engine/Shaders/Private/ParticipatingMediaCommon.ush
--- a/Engine/Shaders/Private/ParticipatingMediaCommon.ush
+++ b/Engine/Shaders/Private/ParticipatingMediaCommon.ush
@@ -292,7 +292,7 @@ float3 IsotropicMediumSlabPunctualDirectionalAlbedoLUT(FParticipatingMedia PM, f
 	);
 }
 
-float IsotropicMediumSlabPunctualDirectionalAlbedo(FParticipatingMedia PM, float NoV, float NoL)
+float3 IsotropicMediumSlabPunctualDirectionalAlbedo(FParticipatingMedia PM, float NoV, float NoL)
 {
 #if USE_LUT_SIMPLEVOLUME
 	const float3 SlabDirectionalAlbedo = IsotropicMediumSlabPunctualDirectionalAlbedoLUT(PM, NoV, NoL);
