Substrate - now prefer IsotropicMediumSlabPunctualDirectionalAlbedo over the env version which look more reasxonable after empirical test. (there is no right answer here anyway).

#tests none
#rnx

diff --git a/Engine/Shaders/Private/Substrate/Substrate.ush b/Engine/Shaders/Private/Substrate/Substrate.ush
--- a/Engine/Shaders/Private/Substrate/Substrate.ush
+++ b/Engine/Shaders/Private/Substrate/Substrate.ush
@@ -2764,8 +2764,11 @@ FSubstrateBSDF SubstrateVerticalLayeringParameterBlending(in FSubstrateBSDF TopB
 	const float3 SlabTransmittanceN		= IsotropicMediumSlabTransmittance(PM, SUBSTRATE_SIMPLEVOLUME_THICKNESS_M, 1.0f /*Normal incidence i.e., NoL == 1.f*/);
 	const float3 SlabTransmittance		= SlabTransmittanceV * SlabTransmittanceN;
 
-	// Using the environment version as an approximation for albedo resulting from scattering in a top layer without light source direction.
-	float3 TopSlabAlbedo = IsotropicMediumSlabEnvDirectionalAlbedo(PM, NoV);
+	// We need to convert the scattering from the top layer to an additional albedo without light source direction.
+	// There is no right answer here since we are adding energy without accounting for the directionality of directional or other local lights.
+	// IsotropicMediumSlabEnvDirectionalAlbedo seems to set too much energy so instead the one for punctual light is used instead and looks more reasonable after some empirical tests.
+	float3 TopSlabAlbedo = IsotropicMediumSlabPunctualDirectionalAlbedo(PM, NoV, 1.0f);
+
 	// SUBSTRATE_TODO account for energy conservation here.
 
 	FVerticalLayeringInfo Info			= GetVerticalLayeringInfo(SafeTopCoverage, SafeBaseCoverage);
