Fix SSS profile conversion having larger influence than legacy material.

diff --git a/Engine/Shaders/Private/Substrate/SubstrateLegacyConversion.ush b/Engine/Shaders/Private/Substrate/SubstrateLegacyConversion.ush
--- a/Engine/Shaders/Private/Substrate/SubstrateLegacyConversion.ush
+++ b/Engine/Shaders/Private/Substrate/SubstrateLegacyConversion.ush
@@ -4,6 +4,10 @@
 
 #include "../Common.ush"
 
+// Optimization: if opacity is 0 then revert to default shading model 
+// This is for matching legacy behavior. See ShadingModelsMaterial.ush.
+#define SUBSURFACE_PROFILE_OPACITY_THRESHOLD 1
+
 ///////////////////////////////////////////////////////////////////////////////////////////////////
 // Helpers
 
@@ -312,8 +316,13 @@ FSubstrateData SubstrateConvertLegacyMaterialDynamic(
 		LegacySlab.Metallic = Metallic;
 		LegacySlab.Roughness = Roughness;
 		LegacySlab.Anisotropy = Anisotropy;
-		LegacySlab.SSSProfileID = SubSurfaceProfileId;
-		LegacySlab.SSSMFPScale = Opacity;
+		#if SUBSURFACE_PROFILE_OPACITY_THRESHOLD
+		if (Opacity > SSSS_OPACITY_THRESHOLD_EPS)
+		#endif
+		{
+			LegacySlab.SSSProfileID = SubSurfaceProfileId;
+			LegacySlab.SSSMFPScale = Opacity;
+		}
 		LegacySlab.SSSPhaseAniso = 0.93f;
 		LegacySlab.UseSSSDiffusion = true;
 		LegacySlab.Emissive = Emissive;
