Fix shader compilation error when substrate is disabled.

diff --git a/Engine/Shaders/Private/RayTracing/RayTracingDebug.usf b/Engine/Shaders/Private/RayTracing/RayTracingDebug.usf
--- a/Engine/Shaders/Private/RayTracing/RayTracingDebug.usf
+++ b/Engine/Shaders/Private/RayTracing/RayTracingDebug.usf
@@ -323,7 +323,7 @@ RAY_TRACING_ENTRY_RAYGEN(RayTracingDebugMainRGS)
 	float OutDepth = 1.0f;
 #endif
 
-	bool bWriteColorOuput = false;
+	bool bWriteColorOuput = true;
 	bool bWriteDepthOutput = VisualizationMode == RAY_TRACING_DEBUG_VIZ_INSTANCE_OVERLAP;
 #if USE_DEBUG_CHS
 	{
diff --git a/Engine/Shaders/Private/Substrate/SubstrateVisualizeCommon.ush b/Engine/Shaders/Private/Substrate/SubstrateVisualizeCommon.ush
--- a/Engine/Shaders/Private/Substrate/SubstrateVisualizeCommon.ush
+++ b/Engine/Shaders/Private/Substrate/SubstrateVisualizeCommon.ush
@@ -1,6 +1,9 @@
 // Copyright Epic Games, Inc. All Rights Reserved.
 
 #include "/Engine/Private/Common.ush"
+
+#if SUBSTRATE_ENABLED
+
 #include "/Engine/Private/ShaderPrint.ush"
 
 #ifndef SUBSTRATE_MATERIALCONTAINER_IS_VIEWRESOURCE
@@ -678,4 +681,6 @@ void SubstratePrintMaterialProperties(inout FShaderPrintContext Context, uint2 I
 			Payload);
 	}
 }
-#endif
\ No newline at end of file
+#endif
+
+#endif // SUBSTRATE_ENABLED
\ No newline at end of file
