LFV - now respects PerPixelFog option

#tests could be tested by enabling per pixel fog on material behind a LFV with low triangle count (simple plane)
[RN] LFV can now be applied per pixel on transparent when per pixel fog is enable.

diff --git a/Engine/Shaders/Private/BasePassCommon.ush b/Engine/Shaders/Private/BasePassCommon.ush
--- a/Engine/Shaders/Private/BasePassCommon.ush
+++ b/Engine/Shaders/Private/BasePassCommon.ush
@@ -67,13 +67,13 @@
 #endif
 
 // Local fog volumes are always evaluated per vertex for now, even with forward shading. But never for opaque meshes rendered in forward.
-#define LOCAL_FOG_VOLUME_ON_TRANSLUCENT (PROJECT_LOCALFOGVOLUME_APPLYONTRANSLUCENT && ((NEEDS_BASEPASS_VERTEX_FOGGING && !(OPAQUE_NEEDS_BASEPASS_FOGGING && PROJECT_VERTEX_FOGGING_FOR_OPAQUE)) || NEEDS_BASEPASS_PIXEL_FOGGING) && !MATERIAL_IS_SKY)
+#define LOCAL_FOG_VOLUME_ON_TRANSLUCENT (PROJECT_LOCALFOGVOLUME_APPLYONTRANSLUCENT && (NEEDS_BASEPASS_VERTEX_FOGGING && !(OPAQUE_NEEDS_BASEPASS_FOGGING && PROJECT_VERTEX_FOGGING_FOR_OPAQUE)) && !MATERIAL_IS_SKY)
 
 struct FSharedBasePassInterpolants
 {
-	// This is enabled for vertex fogging or for local fog volume that are better evaluated per vertex for now.
+	// This is enabled for vertex fogging.
 	// For texture-lightmapped translucency we can pass the vertex fog in its own interpolator
-#if NEEDS_BASEPASS_VERTEX_FOGGING || LOCAL_FOG_VOLUME_ON_TRANSLUCENT
+#if NEEDS_BASEPASS_VERTEX_FOGGING
 	float4 VertexFog		: TEXCOORD7;
 #endif
 
diff --git a/Engine/Shaders/Private/BasePassPixelShader.usf b/Engine/Shaders/Private/BasePassPixelShader.usf
--- a/Engine/Shaders/Private/BasePassPixelShader.usf
+++ b/Engine/Shaders/Private/BasePassPixelShader.usf
@@ -209,8 +209,9 @@
 #include "/Engine/Private/Substrate/SubstrateExport.ush"
 #endif
 
-#define ENABLE_LOCAL_FOG_VOLUMES_ON_OPAQUE_FORWARD (PROJECT_SUPPORTS_LOCALFOGVOLUME && NEEDS_BASEPASS_PIXEL_VOLUMETRIC_FOGGING && OPAQUE_NEEDS_BASEPASS_FOGGING)
-#if ENABLE_LOCAL_FOG_VOLUMES_ON_OPAQUE_FORWARD
+#define ENABLE_LOCAL_FOG_VOLUMES_ON_OPAQUE_FORWARD  (PROJECT_SUPPORTS_LOCALFOGVOLUME && NEEDS_BASEPASS_PIXEL_VOLUMETRIC_FOGGING && OPAQUE_NEEDS_BASEPASS_FOGGING)
+#define LOCAL_FOG_VOLUME_PER_PIXEL_ON_TRANSLUCENT (PROJECT_SUPPORTS_LOCALFOGVOLUME && NEEDS_BASEPASS_PIXEL_FOGGING)
+#if ENABLE_LOCAL_FOG_VOLUMES_ON_OPAQUE_FORWARD || LOCAL_FOG_VOLUME_PER_PIXEL_ON_TRANSLUCENT
 #include "LocalFogVolumes/LocalFogVolumeCommon.ush"
 #endif
 
@@ -1348,6 +1349,15 @@ void FPixelShaderInOut_MainPS(
 
 	float4 Fogging = HeightFogging;
 
+#if LOCAL_FOG_VOLUME_PER_PIXEL_ON_TRANSLUCENT
+	{
+		float3 TranslatedWorldPosition = SvPositionToResolvedTranslatedWorld(In.SvPosition);
+		uint2 TilePos = clamp(uint2(In.SvPosition.xy / float(LFVTilePixelSize)), uint2(0, 0), LFVTileDataResolution - 1);
+		float4 LocalFogVolumeContrib = GetLFVContribution(ResolvedView, TilePos, TranslatedWorldPosition);
+		Fogging = float4(LocalFogVolumeContrib.rgb + Fogging.rgb * LocalFogVolumeContrib.a, LocalFogVolumeContrib.a * Fogging.a);
+	}
+#endif
+
 #if NEEDS_BASEPASS_PIXEL_VOLUMETRIC_FOGGING
 	if (FogStruct.ApplyVolumetricFog > 0) 
 	{
@@ -1355,11 +1365,11 @@ void FPixelShaderInOut_MainPS(
 		float3 TranslatedWorldPosition = SvPositionToResolvedTranslatedWorld(In.SvPosition);
 		uint2 TilePos = clamp(uint2(In.SvPosition.xy / float(LFVTilePixelSize)), uint2(0, 0), LFVTileDataResolution - 1);
 		float4 LocalFogVolumeContrib = GetLFVContribution(ResolvedView, TilePos, TranslatedWorldPosition);
-		HeightFogging = float4(LocalFogVolumeContrib.rgb + HeightFogging.rgb * LocalFogVolumeContrib.a, LocalFogVolumeContrib.a * HeightFogging.a);
+		Fogging = float4(LocalFogVolumeContrib.rgb + Fogging.rgb * LocalFogVolumeContrib.a, LocalFogVolumeContrib.a * Fogging.a);
 	#endif 
 
 		float3 VolumeUV = ComputeVolumeUV(MaterialParameters.AbsoluteWorldPosition, ResolvedView.WorldToClip, ResolvedView);
-		Fogging = CombineVolumetricFog(HeightFogging, VolumeUV, EyeIndex, GBuffer.Depth, ResolvedView);
+		Fogging = CombineVolumetricFog(Fogging, VolumeUV, EyeIndex, GBuffer.Depth, ResolvedView);
 	}
 #endif
 
diff --git a/Engine/Shaders/Private/BasePassVertexShader.usf b/Engine/Shaders/Private/BasePassVertexShader.usf
--- a/Engine/Shaders/Private/BasePassVertexShader.usf
+++ b/Engine/Shaders/Private/BasePassVertexShader.usf
@@ -116,21 +116,17 @@ void Main(
 #endif
 #endif // NEEDS_BASEPASS_VERTEX_FOGGING
 
+#if NEEDS_BASEPASS_VERTEX_FOGGING
+
 	// Apply local fog volume to if enabled. It is always enabled per vertex as of today, that is why it is out ofr the NEEDS_BASEPASS_VERTEX_FOGGING defines.
 #if LOCAL_FOG_VOLUME_ON_TRANSLUCENT
 	float4 VertexClipSpacePosition			= mul(float4(WorldPosition.xyz, 1), ResolvedView.TranslatedWorldToClip);
 	float2 SvPosition						= (VertexClipSpacePosition.xy / VertexClipSpacePosition.w * float2(.5f, -.5f) + .5f) * ResolvedView.ViewSizeAndInvSize.xy;
 	uint2 TilePos							= clamp(uint2(SvPosition.xy / float(LFVTilePixelSize)), uint2(0, 0), LFVTileDataResolution - 1);
 	float4 LFVContribution					= GetLFVContribution(ResolvedView, TilePos, WorldPosition.xyz);
-	#if NEEDS_BASEPASS_VERTEX_FOGGING
 	Output.BasePassInterpolants.VertexFog	= float4(LFVContribution.rgb + Output.BasePassInterpolants.VertexFog.rgb * LFVContribution.a, LFVContribution.a * Output.BasePassInterpolants.VertexFog.a);
-	#else
-	Output.BasePassInterpolants.VertexFog	= LFVContribution; // path when Output.BasePassInterpolants.VertexFog has not been initialised
-	#endif
 #endif // LOCAL_FOG_VOLUME_ON_TRANSLUCENT
 
-#if NEEDS_BASEPASS_VERTEX_FOGGING
-
 #if MATERIAL_ENABLE_TRANSLUCENCY_CLOUD_FOGGING
 	if (TranslucentBasePass.ApplyVolumetricCloudOnTransparent > 0.0f)
 	{
