Substrate - fix compilation on a console.

#rnx
#tests none

diff --git a/Engine/Shaders/Private/MobileBasePassPixelShader.usf b/Engine/Shaders/Private/MobileBasePassPixelShader.usf
--- a/Engine/Shaders/Private/MobileBasePassPixelShader.usf
+++ b/Engine/Shaders/Private/MobileBasePassPixelShader.usf
@@ -686,6 +686,7 @@ void Main(
 	half3 CameraVector = -MaterialParameters.CameraVector;
 	float2 ScreenUV = ScreenPositionToBufferUV(MaterialParameters.ScreenPosition);
 
+#if (SUBSTRATE_TRANSLUCENT_FORWARD || SUBSTRATE_FORWARD_SHADING)
 	#define SUBSTRATE_DEBUG_ENABLE_LIGHTING 1
 	#if SUBSTRATE_DEBUG_ENABLE_LIGHTING
 	if (SubstratePixelHeader.ClosureCount > 0)
@@ -731,6 +732,7 @@ void Main(
 	#endif
 	}
 	#endif // SUBSTRATE_DEBUG_ENABLE_LIGHTING
+#endif // SUBSTRATE_TRANSLUCENT_FORWARD || SUBSTRATE_FORWARD_SHADING
 
   #else // SUBSTRATE_OPTIMIZED_UNLIT
 
