Heterogeneous Volumes: Disable vectorized query for AVSMs larger than 16 samples

Currently using binary search algorithm for AVSMs beyond 16 maximum samples.

diff --git a/Engine/Shaders/Private/HeterogeneousVolumes/HeterogeneousVolumesAdaptiveVolumetricShadowMapSampling.ush b/Engine/Shaders/Private/HeterogeneousVolumes/HeterogeneousVolumesAdaptiveVolumetricShadowMapSampling.ush
--- a/Engine/Shaders/Private/HeterogeneousVolumes/HeterogeneousVolumesAdaptiveVolumetricShadowMapSampling.ush
+++ b/Engine/Shaders/Private/HeterogeneousVolumes/HeterogeneousVolumesAdaptiveVolumetricShadowMapSampling.ush
@@ -73,7 +73,6 @@ float SafeRcp(float Value)
 	return Value > 0 ? rcp(Value) : 0.0;
 }
 
-#if 0
 float AVSM_Sample_Reference(
 	inout FAdaptiveVolumetricShadowMap ShadowMap,
 	int Face,
@@ -181,7 +180,7 @@ float AVSM_Sample_BinarySearch(
 	float Transmittance = lerp(ShadowData[0].Tau, ShadowData[1].Tau, BlendWeight);
 	return Transmittance;
 }
-
+#if 0
 float AVSM_Sample_TwoLevel_Reference(
 	inout FAdaptiveVolumetricShadowMap ShadowMap,
 	int Face,
@@ -361,8 +360,12 @@ float AVSM_Sample(
 )
 {
 	return
-		//AVSM_Sample_Reference(
+// Temporarily disabling vectorized loop for reference
+#if AVSM_SAMPLE_MODE == AVSM_SAMPLE_MODE_TWO_LEVEL_VECTORIZED_LOOP
+		AVSM_Sample_BinarySearch(
+#else
 		AVSM_Sample_TwoLevel_Vectorized(
+#endif
 			ShadowMap,
 			Face,
 			Pixel,
