Followup fix for complexity debug viewmode failing when viewrect is offset. Ensure that appropriate viewport info is used when input and output viewports differ.

diff --git a/Engine/Shaders/Private/ShaderComplexityApplyPixelShader.usf b/Engine/Shaders/Private/ShaderComplexityApplyPixelShader.usf
--- a/Engine/Shaders/Private/ShaderComplexityApplyPixelShader.usf
+++ b/Engine/Shaders/Private/ShaderComplexityApplyPixelShader.usf
@@ -7,6 +7,7 @@
 #include "QuadOverdraw.ush"
  
 SCREEN_PASS_TEXTURE_VIEWPORT(Input)
+SCREEN_PASS_TEXTURE_VIEWPORT(Output)
 
 Texture2D InputTexture;
 SamplerState InputSampler;
@@ -144,7 +145,8 @@ void Main(
 	int2 PixelPos = (int2)SvPosition.xy;
 	float2 ViewLocalUV = float2(UVAndScreenPos.z * 0.5f + 0.5f, 0.5f - 0.5f * UVAndScreenPos.w);
 
-	int2 ViewportCenter = (int2)(Input_ViewportMin + Input_ViewportSize / 2);
+	int2 InputViewportCenter = (int2)(Input_ViewportMin + Input_ViewportSize / 2);
+	int2 OutputViewportCenter = (int2)(Output_ViewportMin + Output_ViewportSize / 2);
 
 	// float3(PS/ShaderBudget, VS/ShaderBudget, overdraw)
 	float3 SceneColor = Texture2DSampleLevel(InputTexture, InputSampler, UV, 0).rgb;
@@ -153,7 +155,7 @@ void Main(
 	[branch]
 	if (DebugViewShaderMode == DVSM_QuadComplexity || DebugViewShaderMode == DVSM_ShaderComplexityContainedQuadOverhead || DebugViewShaderMode == DVSM_ShaderComplexityBleedingQuadOverhead)
 	{
-		float2 QuadID = UsedQuadTextureSize * Input_ViewportSizeInverse * SvPosition.xy;
+		float2 QuadID = UsedQuadTextureSize * Output_ViewportSizeInverse * SvPosition.xy;
 		SceneColor.r += ComplexityToFloat(QuadOverdrawTexture.Load(int3(QuadID.xy + QuadOverdrawOffset(), 0)));
 	}
 #endif // READ_QUAD_OVERDRAW
@@ -181,8 +183,8 @@ void Main(
 
 	// crosshair
 	{
-		float CrossHairMask = PixelPos.x == ViewportCenter.x || PixelPos.y == ViewportCenter.y;
-		float2 DistAbs = abs(PixelPos - ViewportCenter);
+		float CrossHairMask = PixelPos.x == OutputViewportCenter.x || PixelPos.y == OutputViewportCenter.y;
+		float2 DistAbs = abs(PixelPos - OutputViewportCenter);
 		float Dist = max(DistAbs.x, DistAbs.y);
 		float DistMask = Dist >= 2 && Dist < 17;
 
@@ -190,13 +192,13 @@ void Main(
 	}
 	
 	// value under crosshair
-	float3 CenterViewportHDRColor = InputTexture.Load(int3(ViewportCenter, 0)).rgb;
+	float3 CenterViewportHDRColor = InputTexture.Load(int3(InputViewportCenter, 0)).rgb;
 
 #if READ_QUAD_OVERDRAW
 	[branch]
 	if (DebugViewShaderMode == DVSM_QuadComplexity || DebugViewShaderMode == DVSM_ShaderComplexityContainedQuadOverhead || DebugViewShaderMode == DVSM_ShaderComplexityBleedingQuadOverhead)
 	{
-		float2 QuadID = UsedQuadTextureSize * Input_ViewportSizeInverse * ViewportCenter.xy;
+		float2 QuadID = UsedQuadTextureSize * Input_ViewportSizeInverse * InputViewportCenter.xy;
 		SceneColor.r += ComplexityToFloat(QuadOverdrawTexture.Load(int3(QuadID.xy + QuadOverdrawOffset(), 0)));
 	}
 #endif
@@ -208,11 +210,11 @@ void Main(
 		const int OffsetFromBottom = 52;
 		const int Height = 18;
 
-		const int2 Size = int2(Input_ViewportMax.x - Input_ViewportMin.x - OffsetFromLeft * 2, Height);
+		const int2 Size = int2(Output_ViewportMax.x - Output_ViewportMin.x - OffsetFromLeft * 2, Height);
 
 		const int2 LeftTop = int2(
-			Input_ViewportMin.x + OffsetFromLeft,
-			Input_ViewportMax.y - OffsetFromBottom - Size.y);
+			Output_ViewportMin.x + OffsetFromLeft,
+			Output_ViewportMax.y - OffsetFromBottom - Size.y);
 
 		const int OuterBorder = 1;
 
diff --git a/Engine/Source/Runtime/Renderer/Private/PostProcess/PostProcessVisualizeComplexity.cpp b/Engine/Source/Runtime/Renderer/Private/PostProcess/PostProcessVisualizeComplexity.cpp
--- a/Engine/Source/Runtime/Renderer/Private/PostProcess/PostProcessVisualizeComplexity.cpp
+++ b/Engine/Source/Runtime/Renderer/Private/PostProcess/PostProcessVisualizeComplexity.cpp
@@ -15,6 +15,7 @@ public:
 
 	BEGIN_SHADER_PARAMETER_STRUCT(FParameters, )
 		SHADER_PARAMETER_STRUCT(FScreenPassTextureViewportParameters, Input)
+		SHADER_PARAMETER_STRUCT(FScreenPassTextureViewportParameters, Output)
 		SHADER_PARAMETER_RDG_TEXTURE(Texture2D, InputTexture)
 		SHADER_PARAMETER_SAMPLER(SamplerState, InputSampler)
 		SHADER_PARAMETER_RDG_TEXTURE(Texture2D<uint>, QuadOverdrawTexture)
@@ -93,10 +94,12 @@ FScreenPassTexture AddVisualizeComplexityPass(FRDGBuilder& GraphBuilder, const F
 	}
 
 	const FScreenPassTextureViewport InputViewport(Inputs.SceneColor);
+	const FScreenPassTextureViewport OutputViewport(Output);
 
 	FVisualizeComplexityApplyPS::FParameters* PassParameters = GraphBuilder.AllocParameters<FVisualizeComplexityApplyPS::FParameters>();
 	PassParameters->RenderTargets[0] = Output.GetRenderTargetBinding();
 	PassParameters->Input = GetScreenPassTextureViewportParameters(InputViewport);
+	PassParameters->Output = GetScreenPassTextureViewportParameters(OutputViewport);
 	PassParameters->InputTexture = Inputs.SceneColor.Texture;
 	PassParameters->InputSampler = TStaticSamplerState<SF_Bilinear, AM_Clamp, AM_Clamp, AM_Clamp>::GetRHI();
 
