Substrate - fixed overexposed lumen reflection in some cases when it is already pre exposed.

#tests none
#rnx

diff --git a/Engine/Shaders/Private/ReflectionEnvironmentPixelShader.usf b/Engine/Shaders/Private/ReflectionEnvironmentPixelShader.usf
--- a/Engine/Shaders/Private/ReflectionEnvironmentPixelShader.usf
+++ b/Engine/Shaders/Private/ReflectionEnvironmentPixelShader.usf
@@ -501,7 +501,7 @@ void ReflectionEnvironmentSkyLighting(
 			AccumulateSubstrateDeferredLighting(SubstrateLighting, Out, SubstrateEnvLight.bPostProcessSubsurface, BSDF_GETISTOPLAYER(BSDF));
 		}
 
-		OutColor = SubstrateLighting.SceneColor * View.PreExposure;
+		OutColor = SubstrateLighting.SceneColor * (REFLECTION_SOURCE_TYPE == REFLECTION_SOURCE_LUMEN_STANDALONE ? 1.0f : View.PreExposure);
 		#if REFLECTION_SOURCE_TYPE != REFLECTION_SOURCE_LUMEN_STANDALONE
 		OutColor += ApplySSRContribution(SSRContribution, SSRReflection);
 		#endif
