Substrate - glint envlighting now do not modify roughness based on glint presence. this matches the lumen behavior and is more pleasing when glint becomes infinitely small.

#tests none
#rnx

diff --git a/Engine/Shaders/Private/Substrate/SubstrateEvaluation.ush b/Engine/Shaders/Private/Substrate/SubstrateEvaluation.ush
--- a/Engine/Shaders/Private/Substrate/SubstrateEvaluation.ush
+++ b/Engine/Shaders/Private/Substrate/SubstrateEvaluation.ush
@@ -1546,7 +1546,7 @@ FSubstrateEnvLightResult SubstrateEvaluateForEnvLight(FSubstrateBSDFContext BSDF
 					// We need a way to lerp to the glint direction with a roughness of 0. 
 					// SUBSTRATE_TODO recover glint direction.
 					SubstrateEnvLightResult.SpecularWeight *= saturate(GlintD);
-					SubstrateEnvLightResult.SpecularSafeRoughness = lerp(SubstrateEnvLightResult.SpecularSafeRoughness, 0.001f, saturate(GlintD));
+					// We keep the same SpecularSafeRoughness for single glint environment sampling to match the roughness of the surface overall, especially when the glints becomes infinitely small and thus invisible.
 				}
 			#endif // SUBSTRATE_COMPLEXSPECIALPATH && SUBSTRATE_GLINTS_ALLOWED && SUBSTRATE_GLINTS_ENABLED
 
