Readd energy conservation correction rough diffuse.

Without this correction the diffuse value at grazing angle are too low, which is problematic when doing accurate material matching. This correction is disabled for lower platforms.

diff --git a/Engine/Shaders/Private/Substrate/SubstrateEvaluation.ush b/Engine/Shaders/Private/Substrate/SubstrateEvaluation.ush
--- a/Engine/Shaders/Private/Substrate/SubstrateEvaluation.ush
+++ b/Engine/Shaders/Private/Substrate/SubstrateEvaluation.ush
@@ -410,6 +410,23 @@ FSubstrateEvaluateResult SubstrateEvaluateBSDFCommon(FSubstrateBSDFContext BSDFC
 					DirectionalAlbedo_SpecularTransmission = ComputeEnergyPreservation(EnergyTerms);
 					MSScale = ComputeEnergyConservation(EnergyTerms);
 				}
+
+				#if MATERIAL_ROUGHDIFFUSE && SUBSTRATE_SHADING_QUALITY <= 1
+				// We remove the Chan model baked transmittance specular directional albedo assuming F=0.04.
+				// Then we apply the transmittance according to the current surface setup.
+				// To make this a safe operation:
+				//	- We clamp SpecularTChan to a threshold to avoid div by 0 when the material is a metal. 
+				//	- When the material tends to be metal, DirectionalAlbedo_SpecularTransmission will naturally tend to be 0
+				// Bypass this correction when SUBSTRATE_SHADING_QUALITY > 1 (i.e., med/low quality), for trading accuracy for speed
+				if (Settings.bRoughDiffuseEnabled)
+				{
+					const float ChanF0 = 0.04;
+					FBxDFEnergyTerms EnergyTerms = ComputeGGXSpecEnergyTerms(SafeRoughness, BSDFContext.Context.NoV, ChanF0, 1.0f);
+					const float SpecularTransmitionChanF0ToMaterialF0 = DirectionalAlbedo_SpecularTransmission / max(SUBSTRATE_EPSILON, ComputeEnergyPreservation(EnergyTerms));
+					DirectionalAlbedo_SpecularTransmission = SpecularTransmitionChanF0ToMaterialF0;
+				}
+				#endif
+
 				// Apply energy conservation on the diffuse component
 				// If the specular layer is anisotropic, the energy term is computed onto the 'main' roughness [Kulla 2019]
 				//
