Substrate - verified that bent normal are working, implemented before.
A single bent normal can be used per material.

#rnx
#tests none

diff --git a/Engine/Shaders/Private/BasePassPixelShader.usf b/Engine/Shaders/Private/BasePassPixelShader.usf
--- a/Engine/Shaders/Private/BasePassPixelShader.usf
+++ b/Engine/Shaders/Private/BasePassPixelShader.usf
@@ -1678,11 +1678,9 @@ void FPixelShaderInOut_MainPS(
 	Color = 0;
 #endif
 
-	// SUBSTRATE_TODO GBUFFER_HAS_DIFFUSE_SAMPLE_OCCLUSION and ALLOW_STATIC_LIGHTING cases
-	// const FShadingOcclusion ShadingOcclusion = ApplyBentNormal(MaterialParameters.CameraVector, MaterialParameters.WorldNormal, MaterialParameters.TangentToWorld, GetBentNormalZero(MaterialParameters), GBuffer.Roughness, MaterialAO);
 	{
 		SubstratePixelHeader.IrradianceAO									= InitIrradianceAndOcclusion();
-		SubstratePixelHeader.IrradianceAO.MaterialAO						= MaterialAO;			// This is the material AO, it will be converted to multibounce AO using ShadingOcclusion.SpecOcclusion in the PackSubstrateOut function.
+		SubstratePixelHeader.IrradianceAO.MaterialAO						= MaterialAO;			// This is the material AO, it will be converted to multibounce AO using ShadingOcclusion.SpecOcclusion in the PackSubstrateOut function, also accounting for bent normal.
 		SubstratePixelHeader.IrradianceAO.IndirectIrradiance				= IndirectIrradiance;
 		SubstratePixelHeader.IrradianceAO.DiffuseIndirectSampleOcclusion	= GetDiffuseIndirectSampleOcclusion(SubstratePixelHeader.SharedLocalBases, MaterialParameters.CameraVector, In.SvPosition.xy, SubstratePixelHeader.IrradianceAO.MaterialAO);
 	}
