Updated shader files with "pragma once", to prevent compilation errors by niagara data interfaces

#tests added a static mesh DI and rigid mesh DI to a scratch pad in a gpu emitter, which fails compilation without the fix

diff --git a/Engine/Shaders/Private/DistanceFieldLightingShared.ush b/Engine/Shaders/Private/DistanceFieldLightingShared.ush
--- a/Engine/Shaders/Private/DistanceFieldLightingShared.ush
+++ b/Engine/Shaders/Private/DistanceFieldLightingShared.ush
@@ -4,6 +4,8 @@
 	DistanceFieldLightingShared.usf
 =============================================================================*/
 
+#pragma once
+
 #ifndef THREADGROUP_SIZEX
 #define THREADGROUP_SIZEX 1
 #endif   
diff --git a/Engine/Shaders/Private/MeshDistanceFieldCommon.ush b/Engine/Shaders/Private/MeshDistanceFieldCommon.ush
--- a/Engine/Shaders/Private/MeshDistanceFieldCommon.ush
+++ b/Engine/Shaders/Private/MeshDistanceFieldCommon.ush
@@ -4,6 +4,8 @@
 	MeshDistanceFieldCommon.ush
 =============================================================================*/
 
+#pragma once
+
 float DistanceToNearestSurfaceForObject(FDFObjectData DFObjectData, FDFVector3 WorldPosition, float MaxDistance)
 {
 	const float3 VolumePosition = DFMultiplyDemote(WorldPosition, DFObjectData.WorldToVolume);
