Fix incorrect clamping of ShadowFactor weights in ShadowFactorsUpsample pass, used in distance field and capsule shadows, leading to screenspace blocky shadow artifacts.

diff --git a/Engine/Shaders/Private/ShadowFactorsUpsampleCommon.ush b/Engine/Shaders/Private/ShadowFactorsUpsampleCommon.ush
--- a/Engine/Shaders/Private/ShadowFactorsUpsampleCommon.ush
+++ b/Engine/Shaders/Private/ShadowFactorsUpsampleCommon.ush
@@ -72,8 +72,13 @@ void UpsampleShadowFactors(
 		);
 	float4 DepthWeights = 1.0f / (abs(CornerDepths - SceneDepth.xxxx) + Epsilon);
 	float4 FinalWeights = CornerWeights * CornerDepthsValid * DepthWeights;
+	float TotalWeight = dot(FinalWeights, 1);
 
-	float InterpolatedResult = dot(ShadowFactors, FinalWeights) / max(dot(FinalWeights, 1), Epsilon);
+	float InterpolatedResult = 1;
+	if (TotalWeight != 0)
+	{
+		InterpolatedResult = dot(ShadowFactors, FinalWeights) / TotalWeight;
+	}
 
 	ShadowFactor = InterpolatedResult;
 #else
