Swap FCalculateShadingRateImageCS to use 1-dimensional thread groups, in order to support platforms which can only support 1-D quad operations.

diff --git a/Engine/Shaders/Private/VariableRateShading/VRSShadingRateCalculate.usf b/Engine/Shaders/Private/VariableRateShading/VRSShadingRateCalculate.usf
--- a/Engine/Shaders/Private/VariableRateShading/VRSShadingRateCalculate.usf
+++ b/Engine/Shaders/Private/VariableRateShading/VRSShadingRateCalculate.usf
@@ -12,7 +12,7 @@ Texture2D<float> LuminanceTexture;
 RWTexture2D<uint> HardwareShadingRateImage;
 RWTexture2D<uint> SoftwareShadingRateImage;
 
-groupshared half LumaMatrix[THREADGROUP_SIZE_XY][THREADGROUP_SIZE_XY];
+groupshared half LumaMatrix[HARDWARE_TILE_SIZE][HARDWARE_TILE_SIZE];
 groupshared uint TileSupportedShadingRate;
 
 #define SOBEL_DIMENSION 3
@@ -40,8 +40,8 @@ void CalculateShadingRateImageInner(
 	bool InvalidThread = (
 #if !OUTPUT_SOFTWARE_IMAGE
 		// When in HW only mode, we only read from the groupshared LumaMatrix, and so skip sobel operations on pixels at the edge of a group
-		GroupThreadId.x == THREADGROUP_SIZE_XY - 1 ||
-		GroupThreadId.y == THREADGROUP_SIZE_XY - 1 ||
+		GroupThreadId.x == HARDWARE_TILE_SIZE - 1 ||
+		GroupThreadId.y == HARDWARE_TILE_SIZE - 1 ||
 		GroupThreadId.x == 0 ||
 		GroupThreadId.y == 0 ||
 #endif
@@ -165,7 +165,7 @@ uint2 LinearToSwizzled(uint LinearIndex)
 	//  6	14	22	30	38	46	54	62
 	//  7	15	23	31	39	47	55	63
 
-#if THREADGROUP_SIZE_XY == 8
+#if HARDWARE_TILE_SIZE == 8
 	// Quads organized in squares (shift Y, then X, then X, then Y)
 	// Original index has bits 5-4-3-2-1-0
 	// SwizzleX = 4-3-0
@@ -184,29 +184,23 @@ uint2 LinearToSwizzled(uint LinearIndex)
 	return uint2(SwizzleX, SwizzleY);
 }
 
+#define CALCULATE_SHADING_RATE_IMAGE_THREADS	(HARDWARE_TILE_SIZE * HARDWARE_TILE_SIZE)
 
-
-[numthreads(THREADGROUP_SIZE_XY, THREADGROUP_SIZE_XY, 1)]
+[numthreads(CALCULATE_SHADING_RATE_IMAGE_THREADS, 1, 1)]
 void CalculateShadingRateImage(
 	uint3 DispatchThreadId : SV_DispatchThreadID,
 	uint3 GroupThreadId : SV_GroupThreadID,
 	uint3 GroupId : SV_GroupID,
 	uint GroupIndex : SV_GroupIndex)
 {
-#if OUTPUT_SOFTWARE_IMAGE
 	// Remap lanes so that 4 adjacent lanes (e.g. 0,1,2,3) correspond to a pixel quad in the luma texture
 	// Required to use QuadReadAcross functions
 	uint2 RemappedGroupThreadId = LinearToSwizzled(GroupIndex);
-	uint2 RemappeddDispatchThreadId = GroupId.xy * uint2(THREADGROUP_SIZE_XY, THREADGROUP_SIZE_XY) + RemappedGroupThreadId;
+	uint2 RemappedDispatchThreadId = GroupId.xy * uint2(HARDWARE_TILE_SIZE, HARDWARE_TILE_SIZE) + RemappedGroupThreadId;
 
 	CalculateShadingRateImageInner(
-		RemappeddDispatchThreadId, 
+		RemappedDispatchThreadId,
 		RemappedGroupThreadId,
 		GroupId.xy);
-#else
-	CalculateShadingRateImageInner(
-		DispatchThreadId.xy, 
-		GroupThreadId.xy,
-		GroupId.xy);
-#endif
+
 }
\ No newline at end of file
diff --git a/Engine/Source/Runtime/Renderer/Private/VariableRateShading/ContrastAdaptiveImageGenerator.cpp b/Engine/Source/Runtime/Renderer/Private/VariableRateShading/ContrastAdaptiveImageGenerator.cpp
--- a/Engine/Source/Runtime/Renderer/Private/VariableRateShading/ContrastAdaptiveImageGenerator.cpp
+++ b/Engine/Source/Runtime/Renderer/Private/VariableRateShading/ContrastAdaptiveImageGenerator.cpp
@@ -73,10 +73,10 @@ class FCalculateShadingRateImageCS : public FGlobalShader
 	DECLARE_GLOBAL_SHADER(FCalculateShadingRateImageCS);
 	SHADER_USE_PARAMETER_STRUCT(FCalculateShadingRateImageCS, FGlobalShader);
 
-	class FThreadGroupSizeXY : SHADER_PERMUTATION_SPARSE_INT("THREADGROUP_SIZE_XY", 8, 16);
+	class FHardwareTileSize : SHADER_PERMUTATION_SPARSE_INT("HARDWARE_TILE_SIZE", 8, 16);
 	class FOutputHardwareImage : SHADER_PERMUTATION_BOOL("OUTPUT_HARDWARE_IMAGE");
 	class FOutputSoftwareImage : SHADER_PERMUTATION_BOOL("OUTPUT_SOFTWARE_IMAGE");
-	using FPermutationDomain = TShaderPermutationDomain<FThreadGroupSizeXY, FOutputHardwareImage, FOutputSoftwareImage>;
+	using FPermutationDomain = TShaderPermutationDomain<FHardwareTileSize, FOutputHardwareImage, FOutputSoftwareImage>;
 
 	BEGIN_SHADER_PARAMETER_STRUCT(FParameters, )
 		SHADER_PARAMETER_RDG_TEXTURE(Texture2D<float>, LuminanceTexture)
@@ -395,7 +395,7 @@ void AddCreateShadingRateImagePass(
 			check(TileSize.X == TileSize.Y);
 
 			FCalculateShadingRateImageCS::FPermutationDomain PermutationVector;
-			PermutationVector.Set<FCalculateShadingRateImageCS::FThreadGroupSizeXY>(TileSize.X);
+			PermutationVector.Set<FCalculateShadingRateImageCS::FHardwareTileSize>(TileSize.X);
 			PermutationVector.Set<FCalculateShadingRateImageCS::FOutputHardwareImage>(bCreateHardwareImages);
 			PermutationVector.Set<FCalculateShadingRateImageCS::FOutputSoftwareImage>(bCreateSoftwareImages);
 
