Substrate - fixed two sided shading model with translucent blend mode.

#rnx
#tests None

diff --git a/Engine/Shaders/Private/Substrate/Substrate.ush b/Engine/Shaders/Private/Substrate/Substrate.ush
--- a/Engine/Shaders/Private/Substrate/Substrate.ush
+++ b/Engine/Shaders/Private/Substrate/Substrate.ush
@@ -1371,8 +1371,8 @@ FSubstrateData InternalGetSubstrateSlabBSDF(
 	const bool bSSSProfilePlugged = SSSProfileID > 0.0f;
 	const bool bMFPPlugged = any(SSSMFP > 0);
 	bool bIsThin		= bIsAtTheBottomOfTopology && (bMFPPlugged && bIsThinSurface); // Enable thin only if MFP is plugged, as otherwise there will be no back-tramission
-	bool bHasSSSWrap	= bIsAtTheBottomOfTopology && (bMFPPlugged && !UseSSSDiffusion);
-	bool bHasSSSProfile	= bIsAtTheBottomOfTopology && (bSSSProfilePlugged && SSSMFPScale > 0.f);
+	bool bHasSSSWrap	= bIsThin || bIsAtTheBottomOfTopology && (bMFPPlugged && !UseSSSDiffusion);
+	bool bHasSSSProfile	= !bIsThin && bIsAtTheBottomOfTopology && (bSSSProfilePlugged && SSSMFPScale > 0.f);
 	bool bHasSSS		= bMFPPlugged || (bIsAtTheBottomOfTopology && (bHasSSSProfile || bHasSSSWrap)); // When MFP is plugged, that can be diffusion (bottom) or Simple volume (any position)
 
 	FHaziness Haziness = InitialiseHaziness();
@@ -3090,7 +3090,7 @@ void FSubstrateBSDF::PostProcessSlabBeforeLighting(bool bIsSubstrateOpaqueMateri
 {
 	// Two sided lighting can only be evaluated on the bottom layer, together with option SSS. Otherwise, the simple volume will be used.
 	// At the moment, a material can't be thin & simple volume at the same time. This should be fixed by unifying SSSDiffusion|SSSWrap|SimpleVolume into a single evaluation function
-	const bool bIsThin = this.bIsBottom && BSDF_GETISTHIN(this) && bIsSubstrateOpaqueMaterial;
+	const bool bIsThin = this.bIsBottom && BSDF_GETISTHIN(this);
 
 	// We want simple volumetric only if the MFP input is pluged in, otherwise we avoid dealing with simple volumetric.
 	// We also do not want to use the simple volumetric path if we are dealing with an opaque material and this slab is not at the bottom.
