Fix shader compilation warning with Substrate.

diff --git a/Engine/Shaders/Private/Substrate/Substrate.ush b/Engine/Shaders/Private/Substrate/Substrate.ush
--- a/Engine/Shaders/Private/Substrate/Substrate.ush
+++ b/Engine/Shaders/Private/Substrate/Substrate.ush
@@ -1371,7 +1371,7 @@ FSubstrateData InternalGetSubstrateSlabBSDF(
 	const bool bSSSProfilePlugged = SSSProfileID > 0.0f;
 	const bool bMFPPlugged = any(SSSMFP > 0);
 	bool bIsThin		= bIsAtTheBottomOfTopology && (bMFPPlugged && bIsThinSurface); // Enable thin only if MFP is plugged, as otherwise there will be no back-tramission
-	bool bHasSSSWrap	= bIsThin || bIsAtTheBottomOfTopology && (bMFPPlugged && !UseSSSDiffusion);
+	bool bHasSSSWrap	= bIsThin || (bIsAtTheBottomOfTopology && (bMFPPlugged && !UseSSSDiffusion));
 	bool bHasSSSProfile	= !bIsThin && bIsAtTheBottomOfTopology && (bSSSProfilePlugged && SSSMFPScale > 0.f);
 	bool bHasSSS		= bMFPPlugged || (bIsAtTheBottomOfTopology && (bHasSSSProfile || bHasSSSWrap)); // When MFP is plugged, that can be diffusion (bottom) or Simple volume (any position)
 
