Add opacity approximation when translucent material has IoR enabled:

* If IoR is 0, translucency behaves the same to the current rasterizer. Use opacity directly to calculate the translucency holdout.
* When IoR is larger than 0, approximate the result of path tracing that perform MIS based on opacity between solid glass and translucency with default light material. Have match for first intersection, biased opacity for translucency overlays.

#tests local test with translucency IOR from 0 to 1.0

diff --git a/Engine/Shaders/Private/BasePassPixelShader.usf b/Engine/Shaders/Private/BasePassPixelShader.usf
--- a/Engine/Shaders/Private/BasePassPixelShader.usf
+++ b/Engine/Shaders/Private/BasePassPixelShader.usf
@@ -2189,6 +2189,28 @@ void FPixelShaderInOut_MainPS(
 #endif
 
 #if SUPPORT_PRIMITIVE_ALPHA_HOLDOUT && MATERIAL_DOMAIN_SURFACE
+		
+	//Calculate IoR for translucent material with ior and update holdout property
+	#if (MATERIALBLENDING_TRANSLUCENT || SUBSTRATE_TRANSLUCENT_MATERIAL) && REFRACTION_USE_INDEX_OF_REFRACTION
+		#if REFRACTION_ROOT_NODE_OVERRIDES_DEFAULT
+			FMaterialRefractionData RefractionData = GetMaterialRefraction(PixelMaterialInputs);
+			float MaterialIOR = GetMaterialRefractionIOR(RefractionData);
+		#else
+			float MaterialIOR = DielectricF0RGBToIor(ComputeF0(Specular, BaseColor, Metallic));
+		#endif
+		
+		// Match path tracing first intersection. Multiple intersections are biased.
+		#if SUBSTRATE_TRANSLUCENT_MATERIAL
+			#if !MATERIAL_SHADINGMODEL_THIN_TRANSLUCENT && !SUBSTRATE_BLENDING_ALPHAHOLDOUT
+				float DualBlendSurfaceCoverageMix = 1 - DualBlendSurfaceCoverage * (1 - DualBlendSurfaceCoverage);
+				DualBlendSurfaceCoverage = select(MaterialIOR != 0.0, DualBlendSurfaceCoverageMix, DualBlendSurfaceCoverage);
+			#endif
+		#elif MATERIALBLENDING_TRANSLUCENT
+			float OpacityMix =  1 - Opacity * (1 - Opacity);
+			Opacity = select(MaterialIOR != 0.0, OpacityMix, Opacity);
+		#endif
+	#endif 
+	
 	const bool bIsHoldout = GetPrimitiveData(MaterialParameters).Flags & PRIMITIVE_SCENE_DATA_FLAG_HOLDOUT;
 	// Only include the holdout logic if the Renderer is forwarding its alpha channel to post-processing
 	if (bIsHoldout)
