Substrate - bottom layer not account for some refraction direction of the L dir when computing thoughput to the current BSDF.

#rnx
#tests none

diff --git a/Engine/Shaders/Private/Substrate/SubstrateEvaluation.ush b/Engine/Shaders/Private/Substrate/SubstrateEvaluation.ush
--- a/Engine/Shaders/Private/Substrate/SubstrateEvaluation.ush
+++ b/Engine/Shaders/Private/Substrate/SubstrateEvaluation.ush
@@ -230,7 +230,21 @@ float3 LuminanceWeight(in float SatNoL, in FSubstrateBSDF InBSDF)
 
 float3 LuminanceWeight(in FSubstrateBSDFContext InContext, in FSubstrateBSDF InBSDF)
 {
-	return LuminanceWeight(InContext.SatNoL, InBSDF);
+	float SatNoL = InContext.SatNoL;
+
+#if SUBSTRATE_COMPLEXPATH
+	if (BSDF_GETTYPE(InBSDF) == SUBSTRATE_BSDF_TYPE_SLAB && !BSDF_GETISTOPLAYER(InBSDF))
+	{
+		// We do not have any information about the layers above at this stage so we use a default IOR mapping to water behavior (for all fo the layers).
+		const float AirIOR = 1.0f;
+		const float LayersAboveIOR = 1.5f;
+		// Refract L coming toward the surface from air to matter according to IOR, and assuming all surface above have the same IOR.
+		float3 RefractedL = refract(-InContext.L, InContext.N, AirIOR / LayersAboveIOR);
+		SatNoL = saturate(dot(InContext.N, -RefractedL));
+	}
+#endif
+
+	return LuminanceWeight(SatNoL, InBSDF);
 }
 
 half Substrate_D_GGX(half Roughness, half a2, half NoH)
