Fixed shader compilation warning

diff --git a/Engine/Shaders/Private/ForwardLightingCommon.ush b/Engine/Shaders/Private/ForwardLightingCommon.ush
--- a/Engine/Shaders/Private/ForwardLightingCommon.ush
+++ b/Engine/Shaders/Private/ForwardLightingCommon.ush
@@ -241,7 +241,7 @@ FDeferredLightingSplit GetForwardDirectLightingSplit(
 		#endif
 			NewLighting.DiffuseLighting.rgb *= InDirectionalLightAtmosphereTransmittance;
 			NewLighting.SpecularLighting.rgb *= InDirectionalLightAtmosphereTransmittance;
-			NewLighting.LightingLuminance *= InDirectionalLightAtmosphereTransmittance;
+			NewLighting.LightingLuminance *= Luminance(InDirectionalLightAtmosphereTransmittance);
 
 			if (bSeparateMainDirLightLuminance)
 			{
