Fix hair strands texture dilation pass not applied on the inside of the mesh part (on all attributes apart coverage).

diff --git a/Engine/Shaders/Private/HairStrands/HairStrandsTexturesGeneration.usf b/Engine/Shaders/Private/HairStrands/HairStrandsTexturesGeneration.usf
--- a/Engine/Shaders/Private/HairStrands/HairStrandsTexturesGeneration.usf
+++ b/Engine/Shaders/Private/HairStrands/HairStrandsTexturesGeneration.usf
@@ -642,7 +642,27 @@ void MainCS(int2 DispatchThreadId : SV_DispatchThreadID)
 		return;
 	}
 
-	const bool bIsValid = SourceTriangleMaskTexture.Load(uint3(PixelCoord,0)) > 0;
+	bool bIsValid = SourceTriangleMaskTexture.Load(uint3(PixelCoord,0)) > 0;
+
+	// Pixel within the source, which are not written due to lack of data, will be dialated as well
+	{
+		const uint CoverageTextureIndex = GetHairCardTextureIndex(LayoutIndex, EHairCardsVFAttribute_Coverage);
+		const uint CoverageChannelIndex = GetHairCardChannelIndex(LayoutIndex, EHairCardsVFAttribute_Coverage);
+		const float4 Value = SourceTexture[uint3(PixelCoord, CoverageTextureIndex)];
+		float CoverageValue = 1.f;
+		switch (CoverageChannelIndex)
+		{
+			case 0 : CoverageValue = Value.x; break;
+			case 1 : CoverageValue = Value.y; break;
+			case 2 : CoverageValue = Value.z; break;
+			case 3 : CoverageValue = Value.z; break;
+		}
+		if (CoverageValue <= 0.f)
+		{
+			bIsValid = false;
+		}
+	}
+
 	if (bClearPass || bIsValid)
 	{
 		// Fast path: the pixel is valid, we copy value
