Fixed dilation not being applied correctly on all valid pixels.

diff --git a/Engine/Shaders/Private/HairStrands/HairStrandsTexturesGeneration.usf b/Engine/Shaders/Private/HairStrands/HairStrandsTexturesGeneration.usf
--- a/Engine/Shaders/Private/HairStrands/HairStrandsTexturesGeneration.usf
+++ b/Engine/Shaders/Private/HairStrands/HairStrandsTexturesGeneration.usf
@@ -642,7 +642,7 @@ void MainCS(int2 DispatchThreadId : SV_DispatchThreadID)
 		return;
 	}
 
-	bool bIsValid = SourceTriangleMaskTexture.Load(uint3(PixelCoord,0)) > 0;
+	bool bIsValid = false;
 
 	// Pixel within the source, which are not written due to lack of data, will be dialated as well
 	{
@@ -657,9 +657,16 @@ void MainCS(int2 DispatchThreadId : SV_DispatchThreadID)
 			case 2 : CoverageValue = Value.z; break;
 			case 3 : CoverageValue = Value.z; break;
 		}
-		if (CoverageValue <= 0.f)
+
+		// * When doing the clear pass, only mark as valid, pixels whose coverage is > 0
+		// * For other passes, use the source mask (i.e, the mask updated at every pass) for identifying valid/already processed pixels
+		if (bClearPass)
+		{
+			bIsValid = CoverageValue > 0.f;
+		}
+		else
 		{
-			bIsValid = false;
+			bIsValid = SourceTriangleMaskTexture.Load(uint3(PixelCoord,0)) > 0;
 		}
 	}
 
@@ -673,7 +680,7 @@ void MainCS(int2 DispatchThreadId : SV_DispatchThreadID)
 
 		if (bClearPass)
 		{
-			TargetTriangleMaskTexture[PixelCoord] = DepthTestTexture.Load(uint3(PixelCoord,0)) < 1.f ? 1u : 0u;
+			TargetTriangleMaskTexture[PixelCoord] = (bIsValid && OriginalTriangleMaskTexture.Load(uint3(PixelCoord,0)) > 0.5f) ? 1u : 0u;
 		}
 		else
 		{
