[LightGrid] Refactor HZB test logic.

diff --git a/Engine/Shaders/Private/LightGridInjection.usf b/Engine/Shaders/Private/LightGridInjection.usf
--- a/Engine/Shaders/Private/LightGridInjection.usf
+++ b/Engine/Shaders/Private/LightGridInjection.usf
@@ -80,7 +80,7 @@ float ComputeCellNearViewDepthFromZSlice(uint ZSlice)
 	return SliceDepth;
 }
 
-void ComputeCellViewAABB(uint3 GridCoordinate, out float3 ViewTileMin, out float3 ViewTileMax, out float3 CullRectMin, out float3 CullRectMax)
+void ComputeCellCullRect(uint3 GridCoordinate, float MinTileZ, float MaxTileZ, out float3 CullRectMin, out float3 CullRectMax)
 {
 	// Compute extent of tiles in clip-space. Note that the last tile may extend a bit outside of view if view size is not evenly divisible tile size.
 	const float2 InvCulledGridSizeF = (1u << LightGridPixelSizeShift) * View.ViewSizeAndInvSize.zw;
@@ -90,8 +90,27 @@ void ComputeCellViewAABB(uint3 GridCoordinate, out float3 ViewTileMin, out float
 	float2 UnitPlaneTileMin = GridCoordinate.xy * TileSize + UnitPlaneMin;
 	float2 UnitPlaneTileMax = (GridCoordinate.xy + 1) * TileSize + UnitPlaneMin;
 
-	float MinTileZ = ComputeCellNearViewDepthFromZSlice(GridCoordinate.z);
-	float MaxTileZ = ComputeCellNearViewDepthFromZSlice(GridCoordinate.z + 1);
+	float MinTileDeviceZ = ConvertToDeviceZ(MinTileZ);
+	float MaxTileDeviceZ = ConvertToDeviceZ(MaxTileZ);
+
+	CullRectMin.x = min(UnitPlaneTileMin.x, UnitPlaneTileMax.x);
+	CullRectMin.y = min(UnitPlaneTileMin.y, UnitPlaneTileMax.y);
+	CullRectMin.z = min(MinTileDeviceZ, MaxTileDeviceZ);
+
+	CullRectMax.x = max(UnitPlaneTileMin.x, UnitPlaneTileMax.x);
+	CullRectMax.y = max(UnitPlaneTileMin.y, UnitPlaneTileMax.y);
+	CullRectMax.z = max(MinTileDeviceZ, MaxTileDeviceZ);
+}
+
+void ComputeCellViewAABB(uint3 GridCoordinate, float MinTileZ, float MaxTileZ, out float3 ViewTileMin, out float3 ViewTileMax)
+{
+	// Compute extent of tiles in clip-space. Note that the last tile may extend a bit outside of view if view size is not evenly divisible tile size.
+	const float2 InvCulledGridSizeF = (1u << LightGridPixelSizeShift) * View.ViewSizeAndInvSize.zw;
+	const float2 TileSize = float2(2.0f, -2.0f) * InvCulledGridSizeF.xy;
+	const float2 UnitPlaneMin = float2(-1.0f, 1.0f);
+
+	float2 UnitPlaneTileMin = GridCoordinate.xy * TileSize + UnitPlaneMin;
+	float2 UnitPlaneTileMax = (GridCoordinate.xy + 1) * TileSize + UnitPlaneMin;
 
 	float MinTileDeviceZ = ConvertToDeviceZ(MinTileZ);
 	float4 MinDepthCorner0 = mul(float4(UnitPlaneTileMin.x, UnitPlaneTileMin.y, MinTileDeviceZ, 1), View.ClipToView);
@@ -133,14 +152,6 @@ void ComputeCellViewAABB(uint3 GridCoordinate, out float3 ViewTileMin, out float
 
 	ViewTileMin.z = MinTileZ;
 	ViewTileMax.z = MaxTileZ;
-
-	CullRectMin.x = min(UnitPlaneTileMin.x, UnitPlaneTileMax.x);
-	CullRectMin.y = min(UnitPlaneTileMin.y, UnitPlaneTileMax.y);
-	CullRectMin.z = min(MinTileDeviceZ, MaxTileDeviceZ);
-
-	CullRectMax.x = max(UnitPlaneTileMin.x, UnitPlaneTileMax.x);
-	CullRectMax.y = max(UnitPlaneTileMin.y, UnitPlaneTileMax.y);
-	CullRectMax.z = max(MinTileDeviceZ, MaxTileDeviceZ);
 }
 
 bool IntersectConeWithSphere(float3 ConeVertex, float3 ConeAxis, float ConeRadius, float2 CosSinAngle, float4 SphereToTest)
@@ -204,16 +215,15 @@ void LightGridInjectionCS(
 // Disable to pass all lights through for debugging, will hit limits quickly though
 #define CULL_LIGHTS 1
 	#if CULL_LIGHTS
-		float3 ViewTileMin;
-		float3 ViewTileMax;
-		float3 CullRectMin;
-		float3 CullRectMax;
-		ComputeCellViewAABB(GridCoordinate, ViewTileMin, ViewTileMax, CullRectMin, CullRectMax);
 
-		float3 ViewTileCenter = .5f * (ViewTileMin + ViewTileMax);
-		float3 ViewTileExtent = ViewTileMax - ViewTileCenter;
+		float MinTileZ = ComputeCellNearViewDepthFromZSlice(GridCoordinate.z);
+		float MaxTileZ = ComputeCellNearViewDepthFromZSlice(GridCoordinate.z + 1);
 
 #if USE_HZB_CULL
+		float3 CullRectMin;
+		float3 CullRectMax;
+		ComputeCellCullRect(GridCoordinate, MinTileZ, MaxTileZ, CullRectMin, CullRectMax);
+
 		FScreenRect Rect = GetScreenRect(int4(0, 0, HZBViewSize), CullRectMin, CullRectMax, 4);
 
  		if (!IsVisibleHZB(Rect, true /*bSample4x4*/))
@@ -224,6 +234,13 @@ void LightGridInjectionCS(
  		}
 #endif
 		
+		float3 ViewTileMin;
+		float3 ViewTileMax;
+		ComputeCellViewAABB(GridCoordinate, MinTileZ, MaxTileZ, ViewTileMin, ViewTileMax);
+
+		float3 ViewTileCenter = .5f * (ViewTileMin + ViewTileMax);
+		float3 ViewTileExtent = ViewTileMax - ViewTileCenter;
+		
 		uint NumAvailableLinks = NumGridCells * MaxCulledLightsPerCell * NUM_CULLED_GRID_PRIMITIVE_TYPES;
 
 		uint NumVisibleLights = 0;
