[LightGrid] Tighten cell bounds using HZB depth.

diff --git a/Engine/Shaders/Private/LightGridInjection.usf b/Engine/Shaders/Private/LightGridInjection.usf
--- a/Engine/Shaders/Private/LightGridInjection.usf
+++ b/Engine/Shaders/Private/LightGridInjection.usf
@@ -232,6 +232,14 @@ void LightGridInjectionCS(
  			RWNumCulledLightsGrid[GridIndex * NUM_CULLED_LIGHTS_GRID_STRIDE + 1] = PackCulledLightsGridHeader1(GridIndex * MaxCulledLightsPerCell, false, false);
  			return;
  		}
+		
+		// tighten cell bounds based on farthest depth in cell
+		{
+			const float FarthestDeviceZ = GetMinDepthFromHZB(Rect, true /*bSample4x4*/);
+			const float FarthestSceneDepth = ConvertFromDeviceZ(FarthestDeviceZ);
+
+			MaxTileZ = min(MaxTileZ, FarthestSceneDepth);
+		}
 #endif
 		
 		float3 ViewTileMin;
