Disable back face VSM optim for ClearCoat.
The normals used for ClearCoat BRDF will be different from the WorldNormal used for the VSM back face check.
That could lead to us taking the no-shadow fast path even when the BRDF is non-zero.
This issue appeard after CL35466136. Before that the back face path would still output a valid fully occluded shadow value.

diff --git a/Engine/Shaders/Private/VirtualShadowMaps/VirtualShadowMapProjection.usf b/Engine/Shaders/Private/VirtualShadowMaps/VirtualShadowMapProjection.usf
--- a/Engine/Shaders/Private/VirtualShadowMaps/VirtualShadowMapProjection.usf
+++ b/Engine/Shaders/Private/VirtualShadowMaps/VirtualShadowMapProjection.usf
@@ -47,6 +47,7 @@ struct FProjectionShadingInfo
 	bool bIsValid;
 	bool bIsHair;
 	bool bIsSubsurface;
+	bool bIsClearCoat;
 
 	float3 WorldNormal;
 };
@@ -71,12 +72,14 @@ FProjectionShadingInfo GetProjectionShadingInfo(uint2 PixelPos)
 	Out.bIsHair = BSDFType == SUBSTRATE_BSDF_TYPE_HAIR;
 	Out.WorldNormal = TopLayerData.WorldNormal;
 	Out.bIsSubsurface = SSSType != SSS_TYPE_INVALID; // legacy IsSubsurfaceModel is true for subsurface profile, eyes, etc.
+	Out.bIsClearCoat = false;
 
 #else
 	const FGBufferData GBufferData = GetGBufferDataUint(PixelPos, true);
 	Out.bIsValid = GBufferData.ShadingModelID != SHADINGMODELID_UNLIT;
 	Out.bIsHair = GBufferData.ShadingModelID == SHADINGMODELID_HAIR;
 	Out.bIsSubsurface = IsSubsurfaceModel(GBufferData.ShadingModelID);
+	Out.bIsClearCoat = GBufferData.ShadingModelID == SHADINGMODELID_CLEAR_COAT;
 	Out.WorldNormal = GBufferData.WorldNormal;
 #endif
 	return Out;
@@ -168,7 +171,7 @@ FVirtualShadowMapSampleResult ProjectLight(
 {
 	const bool bIsHairInput = InputType == INPUT_TYPE_HAIRSTRANDS;
 	
-	bool bBackfaceCull = (bCullBackfacingPixels > 0) && !bIsHairInput && !ShadingInfo.bIsSubsurface && !ShadingInfo.bIsHair;
+	bool bBackfaceCull = (bCullBackfacingPixels > 0) && !bIsHairInput && !ShadingInfo.bIsSubsurface && !ShadingInfo.bIsHair && !ShadingInfo.bIsClearCoat;
 #if VISUALIZE
 	// Modes more related to the shadow depth rendering don't want to see the projection culling
 	bBackfaceCull = bBackfaceCull &&
