Mobile: Added an option to not use PLS and FBF extension on Android OpenGL
Enabling it will force a multipass rendering where GBuffer and SceneColor are resolved to a system memory

diff --git a/Engine/Shaders/Private/DeferredShadingCommon.ush b/Engine/Shaders/Private/DeferredShadingCommon.ush
--- a/Engine/Shaders/Private/DeferredShadingCommon.ush
+++ b/Engine/Shaders/Private/DeferredShadingCommon.ush
@@ -748,39 +748,39 @@ FGBufferData MobileDecodeGBuffer(half4 InGBufferA, half4 InGBufferB, half4 InGBu
 	return GBuffer;
 }
 
-#if MOBILE_DEFERRED_SHADING && PIXELSHADER
+#if MOBILE_DEFERRED_SHADING
 
-#if METAL_PROFILE
+#define	FRAMEBUFFER_FETCH_ACCESS (IS_MOBILE_DEFERREDSHADING_SUBPASS && PIXELSHADER && ALLOW_FRAMEBUFFER_FETCH)
+
+#if (METAL_PROFILE && FRAMEBUFFER_FETCH_ACCESS)
 #include "/Engine/Public/Platform/Metal/MetalSubpassSupport.ush"
-#elif VULKAN_PROFILE
+#elif (VULKAN_PROFILE && FRAMEBUFFER_FETCH_ACCESS)
 #include "/Engine/Public/Platform/Vulkan/VulkanSubpassSupport.ush"
-#elif COMPILER_GLSL_ES3_1
+#elif (USE_GLES_FBF_DEFERRED && FRAMEBUFFER_FETCH_ACCESS)
 #include "/Engine/Public/Platform/GL/GLSubpassSupport.ush"
 #endif
 
 void MobileFetchGBuffer(in float2 UV, out half4 GBufferA, out half4 GBufferB, out half4 GBufferC, out half4 GBufferD, out float SceneDepth)
 {
-	GBufferD = 0;
-
-#if VULKAN_PROFILE
+#if (VULKAN_PROFILE && FRAMEBUFFER_FETCH_ACCESS)
 	GBufferA = VulkanSubpassFetch1(); 
 	GBufferB = VulkanSubpassFetch2(); 
 	GBufferC = VulkanSubpassFetch3();
-#if MOBILE_EXTENDED_GBUFFER
-	GBufferD = VulkanSubpassFetch4();
-#endif
+	#if MOBILE_EXTENDED_GBUFFER
+		GBufferD = VulkanSubpassFetch4();
+	#endif
 	SceneDepth = ConvertFromDeviceZ(VulkanSubpassDepthFetch());
-#elif (METAL_PROFILE && !MAC)
+#elif (METAL_PROFILE && !MAC && FRAMEBUFFER_FETCH_ACCESS)
 	GBufferA = SubpassFetchRGBA_1(); 
 	GBufferB = SubpassFetchRGBA_2(); 
 	GBufferC = SubpassFetchRGBA_3(); 
-#if MOBILE_EXTENDED_GBUFFER
-	GBufferD = SubpassFetchRGBA_4();
-	SceneDepth = ConvertFromDeviceZ(SubpassFetchR_5());
-#else
-	SceneDepth = ConvertFromDeviceZ(SubpassFetchR_4());
-#endif
-#elif USE_GLES_FBF_DEFERRED
+	#if MOBILE_EXTENDED_GBUFFER
+		GBufferD = SubpassFetchRGBA_4();
+		SceneDepth = ConvertFromDeviceZ(SubpassFetchR_5());
+	#else
+		SceneDepth = ConvertFromDeviceZ(SubpassFetchR_4());
+	#endif
+#elif (USE_GLES_FBF_DEFERRED && FRAMEBUFFER_FETCH_ACCESS)
 	GBufferA = GLSubpassFetch1(); 
 	GBufferB = GLSubpassFetch2(); 
 	GBufferC = GLSubpassFetch3();  
@@ -790,13 +790,9 @@ void MobileFetchGBuffer(in float2 UV, out half4 GBufferA, out half4 GBufferB, ou
 	GBufferA = Texture2DSampleLevel(MobileSceneTextures.GBufferATexture, MobileSceneTextures.GBufferATextureSampler, UV, 0); 
 	GBufferB = Texture2DSampleLevel(MobileSceneTextures.GBufferBTexture, MobileSceneTextures.GBufferBTextureSampler, UV, 0);
 	GBufferC = Texture2DSampleLevel(MobileSceneTextures.GBufferCTexture, MobileSceneTextures.GBufferCTextureSampler, UV, 0);
-#if MOBILE_EXTENDED_GBUFFER
-	GBufferD = Texture2DSampleLevel(MobileSceneTextures.GBufferDTexture, MobileSceneTextures.GBufferDTextureSampler, UV, 0);
-#endif
-	SceneDepth = ConvertFromDeviceZ(Texture2DSampleLevel(MobileSceneTextures.SceneDepthTexture, MobileSceneTextures.SceneDepthTextureSampler, UV, 0).r);
-#endif
-
-#if FORCE_DEPTH_TEXTURE_READS || PLATFORM_NEEDS_DEPTH_TEXTURE_READS
+	#if MOBILE_EXTENDED_GBUFFER
+		GBufferD = Texture2DSampleLevel(MobileSceneTextures.GBufferDTexture, MobileSceneTextures.GBufferDTextureSampler, UV, 0);
+	#endif
 	SceneDepth = ConvertFromDeviceZ(Texture2DSampleLevel(MobileSceneTextures.SceneDepthTexture, MobileSceneTextures.SceneDepthTextureSampler, UV, 0).r);
 #endif
 }
@@ -805,7 +801,7 @@ void MobileFetchGBuffer(in float2 UV, out half4 GBufferA, out half4 GBufferB, ou
 FGBufferData MobileFetchAndDecodeGBuffer(in float2 UV, in float2 PixelPos, uint ViewId)
 {
 	FGBufferData GBuffer = (FGBufferData)0;
-#if (MOBILE_DEFERRED_SHADING && IS_MOBILE_DEFERREDSHADING_SUBPASS && PIXELSHADER)
+#if MOBILE_DEFERRED_SHADING
 	float SceneDepth = 0; 
 	half4 GBufferA = 0;
 	half4 GBufferB = 0;
diff --git a/Engine/Shaders/Private/SceneTexturesCommon.ush b/Engine/Shaders/Private/SceneTexturesCommon.ush
--- a/Engine/Shaders/Private/SceneTexturesCommon.ush
+++ b/Engine/Shaders/Private/SceneTexturesCommon.ush
@@ -188,7 +188,7 @@ float4 CalcFullSceneColor(float2 ScreenUV)
 #endif
 
 // Special intrinsic to read from the current depth buffer
-#define MOBILE_DEPTHFECTH (IS_MOBILE_DEPTHREAD_SUBPASS || POST_PROCESS_AR_PASSTHROUGH) && PIXELSHADER
+#define MOBILE_DEPTHFECTH (IS_MOBILE_DEPTHREAD_SUBPASS || POST_PROCESS_AR_PASSTHROUGH) && PIXELSHADER && ALLOW_FRAMEBUFFER_FETCH
 
 #if MOBILE_DEPTHFECTH && VULKAN_PROFILE
 #include "/Engine/Public/Platform/Vulkan/VulkanSubpassSupport.ush"
diff --git a/Engine/Source/Developer/ShaderFormatOpenGL/Private/OpenGLShaderCompiler.cpp b/Engine/Source/Developer/ShaderFormatOpenGL/Private/OpenGLShaderCompiler.cpp
--- a/Engine/Source/Developer/ShaderFormatOpenGL/Private/OpenGLShaderCompiler.cpp
+++ b/Engine/Source/Developer/ShaderFormatOpenGL/Private/OpenGLShaderCompiler.cpp
@@ -2998,10 +2998,8 @@ static bool CompileToGlslWithShaderConductor(
 		std::string GlslSource;
 
 		// Handle PLS and FBF in OpenGL
-
-		if (Input.Environment.GetCompileArgument(TEXT("SHADING_PATH_MOBILE"), false) &&
-			Input.Environment.GetCompileArgument(TEXT("MOBILE_DEFERRED_SHADING"), false) &&
-			Version == GLSL_ES3_1_ANDROID)
+		if (Version == GLSL_ES3_1_ANDROID &&
+			Input.Environment.GetCompileArgument(TEXT("USE_GLES_FBF_DEFERRED"), false))
 		{
 			bCompilationFailed = !GenerateDeferredMobileShaders(GlslSource, GLSLCompileParams, SourceData, ReflectData, true, false, bEmulatedUBs, BlendFlags);
 		}
diff --git a/Engine/Source/Developer/ShaderFormatOpenGL/Private/ShaderFormatOpenGL.cpp b/Engine/Source/Developer/ShaderFormatOpenGL/Private/ShaderFormatOpenGL.cpp
--- a/Engine/Source/Developer/ShaderFormatOpenGL/Private/ShaderFormatOpenGL.cpp
+++ b/Engine/Source/Developer/ShaderFormatOpenGL/Private/ShaderFormatOpenGL.cpp
@@ -26,7 +26,7 @@ extern void CompileOpenGLShader(
 	GLSLVersion Version);
 
 /** Version for shader format, this becomes part of the DDC key. */
-static const FGuid UE_SHADER_GLSL_VER = FGuid("B1107259-7ED0-4122-848E-63140897A56F");
+static const FGuid UE_SHADER_GLSL_VER = FGuid("16563EF4-3474-4296-8DA3-AEFF2B0547F1");
 
 class FShaderFormatGLSL : public UE::ShaderCompilerCommon::FBaseShaderFormat 
 {
diff --git a/Engine/Source/Runtime/Apple/MetalRHI/Private/MetalRHI.cpp b/Engine/Source/Runtime/Apple/MetalRHI/Private/MetalRHI.cpp
--- a/Engine/Source/Runtime/Apple/MetalRHI/Private/MetalRHI.cpp
+++ b/Engine/Source/Runtime/Apple/MetalRHI/Private/MetalRHI.cpp
@@ -636,7 +636,7 @@ FMetalDynamicRHI::FMetalDynamicRHI(ERHIFeatureLevel::Type RequestedFeatureLevel)
 	SetEmitDrawEvents(GetEmitDrawEvents() | ENABLE_METAL_GPUEVENTS);
 #endif
 
-	GSupportsShaderFramebufferFetch = !PLATFORM_MAC && GMaxRHIShaderPlatform != SP_METAL_MRT && GMaxRHIShaderPlatform != SP_METAL_MRT_TVOS;
+	GSupportsShaderFramebufferFetch = !PLATFORM_MAC && GMaxRHIShaderPlatform != SP_METAL_MRT && GMaxRHIShaderPlatform != SP_METAL_MRT_TVOS && MobileAllowFramebufferFetch(GMaxRHIShaderPlatform);
 	GSupportsShaderMRTFramebufferFetch = GSupportsShaderFramebufferFetch;
 	GHardwareHiddenSurfaceRemoval = true;
 	GSupportsRenderTargetFormat_PF_G8 = false;
diff --git a/Engine/Source/Runtime/Core/Private/HAL/ConsoleManager.cpp b/Engine/Source/Runtime/Core/Private/HAL/ConsoleManager.cpp
--- a/Engine/Source/Runtime/Core/Private/HAL/ConsoleManager.cpp
+++ b/Engine/Source/Runtime/Core/Private/HAL/ConsoleManager.cpp
@@ -3685,6 +3685,13 @@ static TAutoConsoleVariable<int32> CVarMobileAllowDeferredShadingOpenGL(
 		 "1: Allow Deferred Shading on OpenGL"),
 	ECVF_RenderThreadSafe | ECVF_ReadOnly);
 
+static TAutoConsoleVariable<int32> CVarMobileAllowFramebufferFetchOpenGL(
+	TEXT("r.Mobile.AllowFramebufferFetchOpenGL"),
+	1,
+	TEXT("0: Use multi-pass rendering without FBF and PLS extensions\n"
+		"1: Allow use of framebuffer fetch and PLS extensions (default)"),
+	ECVF_RenderThreadSafe | ECVF_ReadOnly);
+
 static TAutoConsoleVariable<int32> CVarMobileEnableStaticAndCSMShadowReceivers(
 	TEXT("r.Mobile.EnableStaticAndCSMShadowReceivers"),
 	1,
diff --git a/Engine/Source/Runtime/Engine/Private/ShaderCompiler/ShaderCompiler.cpp b/Engine/Source/Runtime/Engine/Private/ShaderCompiler/ShaderCompiler.cpp
--- a/Engine/Source/Runtime/Engine/Private/ShaderCompiler/ShaderCompiler.cpp
+++ b/Engine/Source/Runtime/Engine/Private/ShaderCompiler/ShaderCompiler.cpp
@@ -3182,10 +3182,13 @@ void GlobalBeginCompileShader(
 		const bool bMobileDeferredShading = IsMobileDeferredShadingEnabled((EShaderPlatform)Target.Platform);
 		SET_SHADER_DEFINE_AND_COMPILE_ARGUMENT(Input.Environment, MOBILE_DEFERRED_SHADING, bMobileDeferredShading);
 
+		const bool bAllowFramebufferFetch = MobileAllowFramebufferFetch((EShaderPlatform)Target.Platform);
+		SET_SHADER_DEFINE(Input.Environment, ALLOW_FRAMEBUFFER_FETCH, bAllowFramebufferFetch);
+	
 		if (bMobileDeferredShading)
 		{
-			bool bGLESDeferredShading = Target.Platform == SP_OPENGL_ES3_1_ANDROID;
-			SET_SHADER_DEFINE(Input.Environment, USE_GLES_FBF_DEFERRED, bGLESDeferredShading ? 1 : 0);
+			const bool bGLESDeferredShading = (Target.Platform == SP_OPENGL_ES3_1_ANDROID && bAllowFramebufferFetch);
+			SET_SHADER_DEFINE_AND_COMPILE_ARGUMENT(Input.Environment, USE_GLES_FBF_DEFERRED, bGLESDeferredShading);
 			SET_SHADER_DEFINE(Input.Environment, MOBILE_EXTENDED_GBUFFER, MobileUsesExtenedGBuffer((EShaderPlatform)Target.Platform) ? 1 : 0);
 		}
 		else
diff --git a/Engine/Source/Runtime/OpenGLDrv/Private/OpenGLES.cpp b/Engine/Source/Runtime/OpenGLDrv/Private/OpenGLES.cpp
--- a/Engine/Source/Runtime/OpenGLDrv/Private/OpenGLES.cpp
+++ b/Engine/Source/Runtime/OpenGLDrv/Private/OpenGLES.cpp
@@ -170,16 +170,19 @@ void FOpenGLES::ProcessExtensions(const FString& ExtensionsString)
 	bSupportsBGRA8888 = ExtensionsString.Contains(TEXT("GL_APPLE_texture_format_BGRA8888")) || ExtensionsString.Contains(TEXT("GL_IMG_texture_format_BGRA8888")) || ExtensionsString.Contains(TEXT("GL_EXT_texture_format_BGRA8888"));
 	bSupportsColorBufferFloat = ExtensionsString.Contains(TEXT("GL_EXT_color_buffer_float"));
 	bSupportsColorBufferHalfFloat = ExtensionsString.Contains(TEXT("GL_EXT_color_buffer_half_float"));
-	bSupportsShaderFramebufferFetch = ExtensionsString.Contains(TEXT("GL_EXT_shader_framebuffer_fetch")) || ExtensionsString.Contains(TEXT("GL_NV_shader_framebuffer_fetch"))
-		|| ExtensionsString.Contains(TEXT("GL_ARM_shader_framebuffer_fetch ")); // has space at the end to exclude GL_ARM_shader_framebuffer_fetch_depth_stencil match
-	bSupportsShaderMRTFramebufferFetch = ExtensionsString.Contains(TEXT("GL_EXT_shader_framebuffer_fetch")) || ExtensionsString.Contains(TEXT("GL_NV_shader_framebuffer_fetch"));
-	bSupportsPixelLocalStorage = ExtensionsString.Contains(TEXT("GL_EXT_shader_pixel_local_storage"));
-	bSupportsShaderDepthStencilFetch = ExtensionsString.Contains(TEXT("GL_ARM_shader_framebuffer_fetch_depth_stencil"));
 	bSupportsMultisampledRenderToTexture = ExtensionsString.Contains(TEXT("GL_EXT_multisampled_render_to_texture"));
 	bSupportsNVFrameBufferBlit = ExtensionsString.Contains(TEXT("GL_NV_framebuffer_blit"));
 	bSupportsBufferStorage = ExtensionsString.Contains(TEXT("GL_EXT_buffer_storage"));
 	bSupportsDepthClamp = ExtensionsString.Contains(TEXT("GL_EXT_depth_clamp"));
 	bSupportsASTCDecodeMode = ExtensionsString.Contains(TEXT("GL_EXT_texture_compression_astc_decode_mode"));
+	
+	if (MobileAllowFramebufferFetch(GMaxRHIShaderPlatform))
+	{
+		bSupportsShaderFramebufferFetch = ExtensionsString.Contains(TEXT("GL_EXT_shader_framebuffer_fetch")) || ExtensionsString.Contains(TEXT("GL_NV_shader_framebuffer_fetch")) || ExtensionsString.Contains(TEXT("GL_ARM_shader_framebuffer_fetch ")); // has space at the end to exclude GL_ARM_shader_framebuffer_fetch_depth_stencil match
+		bSupportsShaderMRTFramebufferFetch = ExtensionsString.Contains(TEXT("GL_EXT_shader_framebuffer_fetch")) || ExtensionsString.Contains(TEXT("GL_NV_shader_framebuffer_fetch"));
+		bSupportsPixelLocalStorage = ExtensionsString.Contains(TEXT("GL_EXT_shader_pixel_local_storage"));
+		bSupportsShaderDepthStencilFetch = ExtensionsString.Contains(TEXT("GL_ARM_shader_framebuffer_fetch_depth_stencil"));
+	}
 
 	// Report shader precision
 	int Range[2];
diff --git a/Engine/Source/Runtime/RenderCore/Private/RenderUtils.cpp b/Engine/Source/Runtime/RenderCore/Private/RenderUtils.cpp
--- a/Engine/Source/Runtime/RenderCore/Private/RenderUtils.cpp
+++ b/Engine/Source/Runtime/RenderCore/Private/RenderUtils.cpp
@@ -482,6 +482,16 @@ RENDERCORE_API bool MobileRequiresSceneDepthAux(const FStaticShaderPlatform Plat
 	return false;
 }
 
+RENDERCORE_API bool MobileAllowFramebufferFetch(const FStaticShaderPlatform Platform)
+{
+	if (Platform == SP_OPENGL_ES3_1_ANDROID)
+	{
+		const static IConsoleVariable* CVarAllowFramebufferFetchOpenGL = IConsoleManager::Get().FindConsoleVariable(TEXT("r.Mobile.AllowFramebufferFetchOpenGL"));
+		return CVarAllowFramebufferFetchOpenGL->GetBool();
+	}
+	return true;
+}
+
 RENDERCORE_API bool SupportsTextureCubeArray(ERHIFeatureLevel::Type FeatureLevel)
 {
 	return FeatureLevel >= ERHIFeatureLevel::SM5 
diff --git a/Engine/Source/Runtime/RenderCore/Private/Shader.cpp b/Engine/Source/Runtime/RenderCore/Private/Shader.cpp
--- a/Engine/Source/Runtime/RenderCore/Private/Shader.cpp
+++ b/Engine/Source/Runtime/RenderCore/Private/Shader.cpp
@@ -1998,6 +1998,12 @@ void ShaderMapAppendKey(EShaderPlatform Platform, FShaderKeyGenerator& KeyGen)
 			KeyGen.AppendSeparator();
 			KeyGen.Append(TEXT("MobSSR"));
 		}
+
+		if (MobileAllowFramebufferFetch(Platform) == false)
+		{
+			KeyGen.AppendSeparator();
+			KeyGen.Append(TEXT("NoFBF"));
+		}
 	}
 	else
 	{
diff --git a/Engine/Source/Runtime/RenderCore/Public/RenderUtils.h b/Engine/Source/Runtime/RenderCore/Public/RenderUtils.h
--- a/Engine/Source/Runtime/RenderCore/Public/RenderUtils.h
+++ b/Engine/Source/Runtime/RenderCore/Public/RenderUtils.h
@@ -300,6 +300,8 @@ RENDERCORE_API bool PlatformGPUSceneUsesUniformBufferView(const FStaticShaderPla
 
 RENDERCORE_API bool MobileRequiresSceneDepthAux(const FStaticShaderPlatform Platform);
 
+RENDERCORE_API bool MobileAllowFramebufferFetch(const FStaticShaderPlatform Platform);
+
 RENDERCORE_API bool SupportsTextureCubeArray(ERHIFeatureLevel::Type FeatureLevel);
 
 RENDERCORE_API bool MaskedInEarlyPass(const FStaticShaderPlatform Platform);
diff --git a/Engine/Source/Runtime/Renderer/Private/MobileShadingRenderer.cpp b/Engine/Source/Runtime/Renderer/Private/MobileShadingRenderer.cpp
--- a/Engine/Source/Runtime/Renderer/Private/MobileShadingRenderer.cpp
+++ b/Engine/Source/Runtime/Renderer/Private/MobileShadingRenderer.cpp
@@ -2070,14 +2070,7 @@ void FMobileSceneRenderer::RenderDeferredMultiPass(FRDGBuilder& GraphBuilder, FS
 		const bool bDoOcclusionQueires = (!bIsFullDepthPrepassEnabled && ViewContext.bIsLastView && DoOcclusionQueries());
 		PassParameters->RenderTargets.NumOcclusionQueries = bDoOcclusionQueires ? ComputeNumOcclusionQueriesToBatch() : 0u;
 
-		if (!bIsFullDepthPrepassEnabled)
-		{
-			View.ParallelMeshDrawCommandPasses[EMeshPass::DepthPass].BuildRenderingCommands(GraphBuilder, Scene->GPUScene, DepthPassInstanceCullingDrawParams);
-		}
-		View.ParallelMeshDrawCommandPasses[EMeshPass::BasePass].BuildRenderingCommands(GraphBuilder, Scene->GPUScene, PassParameters->InstanceCullingDrawParams);
-		View.ParallelMeshDrawCommandPasses[EMeshPass::SkyPass].BuildRenderingCommands(GraphBuilder, Scene->GPUScene, SkyPassInstanceCullingDrawParams);
-		View.ParallelMeshDrawCommandPasses[EMeshPass::DebugViewMode].BuildRenderingCommands(GraphBuilder, Scene->GPUScene, DebugViewModeInstanceCullingDrawParams);
-		View.ParallelMeshDrawCommandPasses[EMeshPass::MeshDecal_SceneColorAndGBuffer].BuildRenderingCommands(GraphBuilder, Scene->GPUScene, MeshDecalSceneColorAndGBufferInstanceCullingDrawParams);
+		BuildInstanceCullingDrawParams(GraphBuilder, View, PassParameters);
 
 		GraphBuilder.AddPass(
 			RDG_EVENT_NAME("BasePass"),
@@ -2136,6 +2129,7 @@ void FMobileSceneRenderer::RenderDeferredMultiPass(FRDGBuilder& GraphBuilder, FS
 			PassParameters->View = View.GetShaderParameters();
 			PassParameters->MobileBasePass = CreateMobileBasePassUniformBuffer(GraphBuilder, View, EMobileBasePass::Opaque, EMobileSceneTextureSetupMode::SceneDepth);
 			PassParameters->RenderTargets = BasePassRenderTargets;
+			PassParameters->InstanceCullingDrawParams = MeshDecalSceneColorInstanceCullingDrawParams;
 
 			GraphBuilder.AddPass(
 				RDG_EVENT_NAME("Decals"),
@@ -2177,8 +2171,6 @@ void FMobileSceneRenderer::RenderDeferredMultiPass(FRDGBuilder& GraphBuilder, FS
 		// Only SceneColor and Depth
 		PassParameters->RenderTargets[0] = BasePassRenderTargets[0];
 		PassParameters->RenderTargets.DepthStencil = BasePassRenderTargets.DepthStencil;
-
-		View.ParallelMeshDrawCommandPasses[StandardTranslucencyMeshPass].BuildRenderingCommands(GraphBuilder, Scene->GPUScene, TranslucencyInstanceCullingDrawParams);
 		PassParameters->InstanceCullingDrawParams = TranslucencyInstanceCullingDrawParams;
 		const EMobileSSRQuality MobileSSRQuality = ActiveMobileSSRQuality(View, bShouldRenderVelocities);
 
