SkyAtmosphere - added art direction control Sky And Aerial Perspective illumiance scale factor.

#rnx
#tests Could add a test where the new SkyAndAerialPerspectiveLuminance factor is tweaked.

diff --git a/Engine/Shaders/Private/SkyAtmosphere.usf b/Engine/Shaders/Private/SkyAtmosphere.usf
--- a/Engine/Shaders/Private/SkyAtmosphere.usf
+++ b/Engine/Shaders/Private/SkyAtmosphere.usf
@@ -1050,7 +1050,8 @@ void RenderSkyAtmosphereRayMarchingPS(
 		SVPos, WorldPos, WorldDir,
 		Ground, Sampling, DeviceZ, MieRayPhase,
 		View.AtmosphereLightDirection[0].xyz, View.AtmosphereLightDirection[1].xyz, 
-		View.AtmosphereLightIlluminanceOuterSpace[0].rgb, View.AtmosphereLightIlluminanceOuterSpace[1].rgb,
+		View.AtmosphereLightIlluminanceOuterSpace[0].rgb * SkyAtmosphere.SkyAndAerialPerspectiveLuminanceFactor, 
+		View.AtmosphereLightIlluminanceOuterSpace[1].rgb * SkyAtmosphere.SkyAndAerialPerspectiveLuminanceFactor,
 		AerialPespectiveViewDistanceScale);
 
 	PreExposedL += ss.L * LuminanceScale;
@@ -1333,7 +1334,8 @@ void RenderSkyViewLutCS(uint3 ThreadId : SV_DispatchThreadID)
 		float4(PixPos, 0.0f, 1.0f), WorldPos, WorldDir,
 		Ground, Sampling, DeviceZ, MieRayPhase,
 		AtmosphereLightDirection0, AtmosphereLightDirection1,
-		View.AtmosphereLightIlluminanceOuterSpace[0].rgb, View.AtmosphereLightIlluminanceOuterSpace[1].rgb,
+		View.AtmosphereLightIlluminanceOuterSpace[0].rgb * SkyAtmosphere.SkyAndAerialPerspectiveLuminanceFactor,
+		View.AtmosphereLightIlluminanceOuterSpace[1].rgb * SkyAtmosphere.SkyAndAerialPerspectiveLuminanceFactor,
 		AerialPespectiveViewDistanceScale);
 
 	const float Transmittance = dot(ss.Transmittance, float3(1.0f / 3.0f, 1.0f / 3.0f, 1.0f / 3.0f));
@@ -1392,7 +1394,8 @@ void RenderDistantSkyLightLutCS(uint3 ThreadId : SV_DispatchThreadID)
 		float4(PixPos, 0.0f, 1.0f), SamplePos, SampleDir,
 		Ground, Sampling, DeviceZ, MieRayPhase,
 		AtmosphereLightDirection0.xyz, AtmosphereLightDirection1.xyz, 
-		AtmosphereLightIlluminanceOuterSpace0.rgb, AtmosphereLightIlluminanceOuterSpace1.rgb,
+		AtmosphereLightIlluminanceOuterSpace0.rgb * SkyAtmosphere.SkyAndAerialPerspectiveLuminanceFactor,
+		AtmosphereLightIlluminanceOuterSpace1.rgb * SkyAtmosphere.SkyAndAerialPerspectiveLuminanceFactor,
 		AerialPespectiveViewDistanceScale);
 
 	GroupSkyLuminanceSamples[LinearIndex] = ss.L * SkyAtmosphere.SkyLuminanceFactor;
@@ -1616,7 +1619,8 @@ void RenderCameraAerialPerspectiveVolumeCS(uint3 ThreadId : SV_DispatchThreadID)
 		float4(PixPos, 0.0f, 1.0f), RayStartWorldPos, WorldDir,
 		Ground, Sampling, DeviceZ, MieRayPhase,
 		View.AtmosphereLightDirection[0].xyz, View.AtmosphereLightDirection[1].xyz, 
-		View.AtmosphereLightIlluminanceOuterSpace[0].rgb, View.AtmosphereLightIlluminanceOuterSpace[1].rgb,
+		View.AtmosphereLightIlluminanceOuterSpace[0].rgb * SkyAtmosphere.SkyAndAerialPerspectiveLuminanceFactor,
+		View.AtmosphereLightIlluminanceOuterSpace[1].rgb * SkyAtmosphere.SkyAndAerialPerspectiveLuminanceFactor,
 		AerialPespectiveViewDistanceScale,
 		tMaxMax);
 
@@ -1704,7 +1708,8 @@ void RenderSkyAtmosphereDebugPS(
 				SVPos, WorldPos, WorldDir,
 				Ground, Sampling, DeviceZ, MieRayPhase,
 				View.AtmosphereLightDirection[0].xyz, NullLightDirection, 
-				View.AtmosphereLightIlluminanceOuterSpace[0].rgb, NullLightIlluminance,
+				View.AtmosphereLightIlluminanceOuterSpace[0].rgb * SkyAtmosphere.SkyAndAerialPerspectiveLuminanceFactor, 
+				NullLightIlluminance,
 				AerialPespectiveViewDistanceScale);
 			OutLuminance = float4(SimpleTonemap(SS.L), 0.0f);
 			return;
@@ -1744,7 +1749,8 @@ void RenderSkyAtmosphereDebugPS(
 			SVPos, WorldPos, WorldDir,
 			Ground, Sampling, DeviceZ, MieRayPhase,
 			LightDir, NullLightDirection, 
-			View.AtmosphereLightIlluminanceOuterSpace[0].rgb, NullLightIlluminance,
+			View.AtmosphereLightIlluminanceOuterSpace[0].rgb * SkyAtmosphere.SkyAndAerialPerspectiveLuminanceFactor, 
+			NullLightIlluminance,
 			AerialPespectiveViewDistanceScale);
 
 		OutLuminance = float4(SimpleTonemap(SS.L), 0.0f);
diff --git a/Engine/Source/Editor/EnvironmentLightingViewer/Private/SEnvironmentLightingViewer.cpp b/Engine/Source/Editor/EnvironmentLightingViewer/Private/SEnvironmentLightingViewer.cpp
--- a/Engine/Source/Editor/EnvironmentLightingViewer/Private/SEnvironmentLightingViewer.cpp
+++ b/Engine/Source/Editor/EnvironmentLightingViewer/Private/SEnvironmentLightingViewer.cpp
@@ -537,6 +537,7 @@ bool SEnvironmentLightingViewer::GetIsPropertyVisible(const FPropertyAndParent&
 				|| PropertyAndParent.Property.GetNameCPP().Equals(TEXT("OtherAbsorptionScale"))
 				|| PropertyAndParent.Property.GetNameCPP().Equals(TEXT("OtherAbsorption"))
 				|| PropertyAndParent.Property.GetNameCPP().Equals(TEXT("SkyLuminanceFactor"))
+				|| PropertyAndParent.Property.GetNameCPP().Equals(TEXT("SkyAndAerialPerspectiveLuminanceFactor"))
 				|| PropertyAndParent.Property.GetNameCPP().Equals(TEXT("HeightFogContribution"))
 				|| PropertyAndParent.Property.GetNameCPP().Equals(TEXT("AerialPespectiveViewDistanceScale"))
 				|| PropertyAndParent.Property.GetNameCPP().Equals(TEXT("AerialPerspectiveStartDepth"));
diff --git a/Engine/Source/Runtime/Engine/Classes/Components/SkyAtmosphereComponent.h b/Engine/Source/Runtime/Engine/Classes/Components/SkyAtmosphereComponent.h
--- a/Engine/Source/Runtime/Engine/Classes/Components/SkyAtmosphereComponent.h
+++ b/Engine/Source/Runtime/Engine/Classes/Components/SkyAtmosphereComponent.h
@@ -144,6 +144,10 @@ class USkyAtmosphereComponent : public USceneComponent
 	UPROPERTY(EditAnywhere, BlueprintReadOnly, interp, Category = "Art Direction", meta = (HideAlphaChannel))
 	FLinearColor SkyLuminanceFactor;
 
+	/** Scales the luminance of pixels representing the sky and the aerial perspective. This will impact the captured sky light. */
+	UPROPERTY(EditAnywhere, BlueprintReadOnly, interp, Category = "Art Direction", meta = (HideAlphaChannel))
+	FLinearColor SkyAndAerialPerspectiveLuminanceFactor;
+
 	/** Makes the aerial perspective look thicker by scaling distances from view to surfaces (opaque and translucent). */
 	UPROPERTY(EditAnywhere, BlueprintReadOnly, interp, Category = "Art Direction", meta = (DisplayName = "Aerial Perspective View Distance Scale", UIMin = 0.0, UIMax = 3.0, ClampMin = 0.0, SliderExponent = 2.0))
 	float AerialPespectiveViewDistanceScale;
@@ -216,6 +220,8 @@ class USkyAtmosphereComponent : public USceneComponent
 	UFUNCTION(BlueprintCallable, Category = "Rendering")
 	ENGINE_API void SetSkyLuminanceFactor(FLinearColor NewValue);
 	UFUNCTION(BlueprintCallable, Category = "Rendering")
+	ENGINE_API void SetSkyAndAerialPerspectiveLuminanceFactor(FLinearColor NewValue);
+	UFUNCTION(BlueprintCallable, Category = "Rendering")
 	ENGINE_API void SetAerialPespectiveViewDistanceScale(float NewValue);
 	UFUNCTION(BlueprintCallable, Category = "Rendering")
 	ENGINE_API void SetAerialPerspectiveStartDepth(float NewValue);
diff --git a/Engine/Source/Runtime/Engine/Private/Components/SkyAtmosphereComponent.cpp b/Engine/Source/Runtime/Engine/Private/Components/SkyAtmosphereComponent.cpp
--- a/Engine/Source/Runtime/Engine/Private/Components/SkyAtmosphereComponent.cpp
+++ b/Engine/Source/Runtime/Engine/Private/Components/SkyAtmosphereComponent.cpp
@@ -72,6 +72,7 @@ USkyAtmosphereComponent::USkyAtmosphereComponent(const FObjectInitializer& Objec
 	OtherTentDistribution.Width       = 15.0f;
 
 	SkyLuminanceFactor = FLinearColor(FLinearColor::White);
+	SkyAndAerialPerspectiveLuminanceFactor = FLinearColor(FLinearColor::White);
 	MultiScatteringFactor = 1.0f;
 	AerialPespectiveViewDistanceScale = 1.0f;
 	HeightFogContribution = 1.0f;
@@ -394,6 +395,7 @@ SKY_DECLARE_BLUEPRINT_SETFUNCTION(float, OtherAbsorptionScale);
 SKY_DECLARE_BLUEPRINT_SETFUNCTION_LINEARCOEFFICIENT(OtherAbsorption);
 
 SKY_DECLARE_BLUEPRINT_SETFUNCTION_LINEARCOEFFICIENT(SkyLuminanceFactor);
+SKY_DECLARE_BLUEPRINT_SETFUNCTION_LINEARCOEFFICIENT(SkyAndAerialPerspectiveLuminanceFactor);
 SKY_DECLARE_BLUEPRINT_SETFUNCTION(float, AerialPespectiveViewDistanceScale);
 SKY_DECLARE_BLUEPRINT_SETFUNCTION(float, AerialPerspectiveStartDepth);
 SKY_DECLARE_BLUEPRINT_SETFUNCTION(float, HeightFogContribution);
@@ -512,6 +514,7 @@ FSkyAtmosphereSceneProxy::FSkyAtmosphereSceneProxy(const USkyAtmosphereComponent
 	, bRenderInMainPass(InComponent->bRenderInMainPass > 0)
 {
 	SkyLuminanceFactor = InComponent->SkyLuminanceFactor;
+	SkyAndAerialPerspectiveLuminanceFactor = InComponent->SkyAndAerialPerspectiveLuminanceFactor;
 	AerialPespectiveViewDistanceScale = InComponent->AerialPespectiveViewDistanceScale;
 	HeightFogContribution = InComponent->HeightFogContribution;
 	AerialPerspectiveStartDepthKm = InComponent->AerialPerspectiveStartDepth;
diff --git a/Engine/Source/Runtime/Engine/Public/SceneManagement.h b/Engine/Source/Runtime/Engine/Public/SceneManagement.h
--- a/Engine/Source/Runtime/Engine/Public/SceneManagement.h
+++ b/Engine/Source/Runtime/Engine/Public/SceneManagement.h
@@ -1192,6 +1192,7 @@ public:
 	ENGINE_API ~FSkyAtmosphereSceneProxy();
 
 	FLinearColor GetSkyLuminanceFactor() const { return SkyLuminanceFactor; }
+	FLinearColor GetSkyAndAerialPerspectiveLuminanceFactor() const { return SkyAndAerialPerspectiveLuminanceFactor; }
 	float GetAerialPespectiveViewDistanceScale() const { return AerialPespectiveViewDistanceScale; }
 	float GetHeightFogContribution() const { return HeightFogContribution; }
 	float GetAerialPerspectiveStartDepthKm() const { return AerialPerspectiveStartDepthKm; }
@@ -1214,6 +1215,7 @@ private:
 	FAtmosphereSetup AtmosphereSetup;
 
 	FLinearColor SkyLuminanceFactor;
+	FLinearColor SkyAndAerialPerspectiveLuminanceFactor;
 	float AerialPespectiveViewDistanceScale;
 	float HeightFogContribution;
 	float AerialPerspectiveStartDepthKm;
diff --git a/Engine/Source/Runtime/Renderer/Private/SkyAtmosphereRendering.cpp b/Engine/Source/Runtime/Renderer/Private/SkyAtmosphereRendering.cpp
--- a/Engine/Source/Runtime/Renderer/Private/SkyAtmosphereRendering.cpp
+++ b/Engine/Source/Runtime/Renderer/Private/SkyAtmosphereRendering.cpp
@@ -255,6 +255,7 @@ BEGIN_GLOBAL_SHADER_PARAMETER_STRUCT(FSkyAtmosphereInternalCommonParameters, )
 	SHADER_PARAMETER(float, FogShowFlagFactor)
 
 	SHADER_PARAMETER(FVector3f, SkyLuminanceFactor)
+	SHADER_PARAMETER(FVector3f, SkyAndAerialPerspectiveLuminanceFactor)
 
 END_GLOBAL_SHADER_PARAMETER_STRUCT()
 
@@ -1292,6 +1293,7 @@ static void SetupSkyAtmosphereInternalCommonParameters(
 
 	const FSkyAtmosphereSceneProxy& SkyAtmosphereSceneProxy = SkyInfo.GetSkyAtmosphereSceneProxy();
 	InternalCommonParameters.SkyLuminanceFactor = FVector3f(SkyAtmosphereSceneProxy.GetSkyLuminanceFactor());
+	InternalCommonParameters.SkyAndAerialPerspectiveLuminanceFactor = FVector3f(SkyAtmosphereSceneProxy.GetSkyAndAerialPerspectiveLuminanceFactor());
 	InternalCommonParameters.AerialPespectiveViewDistanceScale = SkyAtmosphereSceneProxy.GetAerialPespectiveViewDistanceScale();
 	InternalCommonParameters.FogShowFlagFactor = ViewFamily.EngineShowFlags.Fog ? 1.0f : 0.0f;
 
