Fix incorrect header size in Substrate's material visualization.

diff --git a/Engine/Shaders/Private/Substrate/SubstrateVisualizeCommon.ush b/Engine/Shaders/Private/Substrate/SubstrateVisualizeCommon.ush
--- a/Engine/Shaders/Private/Substrate/SubstrateVisualizeCommon.ush
+++ b/Engine/Shaders/Private/Substrate/SubstrateVisualizeCommon.ush
@@ -498,11 +498,10 @@ void SubstratePrintMaterialProperties_Internal(
 	FSubstratePixelHeader Header,
 	FSubstrateSubsurfaceHeader SSSHeader,
 	FSubstrateTopLayerData TopLayerData,
-	FSubstrateRaytracingPayload Payload)
+	FSubstrateRaytracingPayload Payload,
+	const uint FootPrint_Start,
+	const uint FootPrint_PostHeader)
 {
-	const uint FootPrint_Start = SubstrateAddressing.ReadBytes;
-	const uint FootPrint_PostHeader = SubstrateAddressing.ReadBytes;
-
 	const FFontColor FontBSDFCount = FontEmerald;
 	const FFontColor FontHeaderStateName = FontWhite;
 	const FFontColor FontHeaderPropName = FontWhite;
@@ -641,7 +640,11 @@ void SubstratePrintMaterialProperties_Internal(
 void SubstratePrintMaterialProperties(inout FShaderPrintContext Context, uint2 InCoord, float3 InWorldPosition, float3 V, uint InClosureIndex)
 {
 	FSubstrateAddressing SubstrateAddressing = GetSubstratePixelDataByteOffset(InCoord, uint2(View.BufferSizeAndInvSize.xy), Substrate.MaxBytesPerPixel);
+	const uint FootPrint_Start = SubstrateAddressing.ReadBytes;
+
 	FSubstratePixelHeader Header = UnpackSubstrateHeaderIn(Substrate.MaterialTextureArray, SubstrateAddressing, Substrate.TopLayerTexture);
+	const uint FootPrint_PostHeader = SubstrateAddressing.ReadBytes;
+
 	FSubstrateSubsurfaceHeader SSSHeader = SubstrateLoadSubsurfaceHeader(Substrate.MaterialTextureArray, Substrate.FirstSliceStoringSubstrateSSSData, SubstrateAddressing.PixelCoords);
 	FSubstrateRaytracingPayload DummyPayload;
 	FSubstrateTopLayerData TopLayerData = SubstrateUnpackTopLayerData(Substrate.TopLayerTexture.Load(uint3(InCoord, 0)));
@@ -656,13 +659,19 @@ void SubstratePrintMaterialProperties(inout FShaderPrintContext Context, uint2 I
 		Header,
 		SSSHeader,
 		TopLayerData,
-		DummyPayload);
+		DummyPayload, 
+		FootPrint_Start, 
+		FootPrint_PostHeader);
 }
 #else
 void SubstratePrintMaterialProperties(inout FShaderPrintContext Context, uint2 InCoord, float3 InWorldPosition, float3 V, uint InClosureIndex, FSubstrateRaytracingPayload Payload)
 {
 	FSubstrateAddressing SubstrateAddressing = GetSubstratePixelDataByteOffset(InCoord, uint2(View.BufferSizeAndInvSize.xy), 0);	
+	const uint FootPrint_Start = SubstrateAddressing.ReadBytes;
+	
 	FSubstratePixelHeader Header = UnpackSubstrateHeaderIn(Payload, SubstrateAddressing, Payload);
+	const uint FootPrint_PostHeader = SubstrateAddressing.ReadBytes;
+
 	FSubstrateSubsurfaceHeader SSSHeader = (FSubstrateSubsurfaceHeader)0;
 	FSubstrateTopLayerData TopLayerData = SubstrateUnpackTopLayerData(Payload.PackedTopLayerData);
 
@@ -678,7 +687,9 @@ void SubstratePrintMaterialProperties(inout FShaderPrintContext Context, uint2 I
 			Header,
 			SSSHeader,
 			TopLayerData,
-			Payload);
+			Payload,
+			FootPrint_Start,
+			FootPrint_PostHeader);
 	}
 }
 #endif
