Shadows: Directional light penumbra size is incorrect

The sampling of the directional light was using a disk sampling routine with a radius of Sin(theta/2) which resulted in a shadow that is slightly wider than expected. The result now matches much more closely with the reference computed with the path tracer.

#tests Tested RT shadows with source angle above 20 degrees

diff --git a/Engine/Shaders/Private/RayTracing/RayTracingDirectionalLight.ush b/Engine/Shaders/Private/RayTracing/RayTracingDirectionalLight.ush
--- a/Engine/Shaders/Private/RayTracing/RayTracingDirectionalLight.ush
+++ b/Engine/Shaders/Private/RayTracing/RayTracingDirectionalLight.ush
@@ -7,30 +7,6 @@
 #include "../LightData.ush"
 #include "RayTracingCommon.ush"
 
-// Adapted from "A Low Distortion Map Between Disk and Square." Improved based
-// on comments from Peter Shirley's blog post: "Improved code for concentric map."
-float2 ToConcentricMap(float2 RectangularCoords)
-{
-	float R;
-	float Phi;
-
-	RectangularCoords = 2.0 * RectangularCoords - 1.0;
-	float2 RectangularCoordsSquared = RectangularCoords * RectangularCoords;
-	if (RectangularCoordsSquared.x > RectangularCoordsSquared.y)
-	{
-		R = RectangularCoords.x;
-		Phi = (PI / 4.0) * (RectangularCoords.y / RectangularCoords.x);
-	}
-	else
-	{
-		R = RectangularCoords.y;
-		Phi = (PI / 2.0) - (PI / 4.0) * (RectangularCoords.x / RectangularCoords.y);
-	}
-
-	float2 PolarCoords = float2(R, Phi);
-	return PolarCoords;
-}
-
 void GenerateDirectionalLightOcclusionRay(
 	FLightShaderParameters LightParameters,
 	float3 TranslatedWorldPosition,
@@ -41,28 +17,10 @@ void GenerateDirectionalLightOcclusionRay(
 	out float RayTMin,
 	out float RayTMax)
 {
-	// Draw random variable and choose a point on a unit disk
-	float2 BufferSize = View.BufferSizeAndInvSize.xy;
-
-	float2 DiskUV = UniformSampleDiskConcentric(RandSample) * LightParameters.SourceRadius;
-
-	// Permute light direction by user-defined radius on unit sphere
-	float3 LightDirection = LightParameters.Direction;
-	float3 N = LightDirection;
-	float3 dPdu = float3(1, 0, 0);
-	if (abs(dot(N, dPdu)) != 1)
-	{
-		dPdu = cross(N, dPdu);
-	}
-	else
-	{
-		dPdu = cross(N, float3(0, 1, 0));
-	}
-	float3 dPdv = cross(dPdu, N);
-	LightDirection += dPdu * DiskUV.x + dPdv * DiskUV.y;
-	
-    RayOrigin = TranslatedWorldPosition;
-    RayDirection = normalize(LightDirection);
+	const float SinThetaMax = LightParameters.SourceRadius;
+	const float4 DirAndPdf = UniformSampleConeRobust(RandSample, SinThetaMax * SinThetaMax);
+	RayOrigin = TranslatedWorldPosition;
+	RayDirection = TangentToWorld(DirAndPdf.xyz, LightParameters.Direction);
 	RayTMin = 0.0;
 	RayTMax = 1.0e27;
 }
