Substrate - made refraction of L vector progressive with respect to coverage of surfaces above.

#rnx
#tests none

diff --git a/Engine/Shaders/Private/Substrate/SubstrateEvaluation.ush b/Engine/Shaders/Private/Substrate/SubstrateEvaluation.ush
--- a/Engine/Shaders/Private/Substrate/SubstrateEvaluation.ush
+++ b/Engine/Shaders/Private/Substrate/SubstrateEvaluation.ush
@@ -239,7 +239,8 @@ float3 LuminanceWeight(in FSubstrateBSDFContext InContext, in FSubstrateBSDF InB
 		const float AirIOR = 1.0f;
 		const float LayersAboveIOR = 1.5f;
 		// Refract L coming toward the surface from air to matter according to IOR, and assuming all surface above have the same IOR.
-		float3 RefractedL = refract(-InContext.L, InContext.N, AirIOR / LayersAboveIOR);
+		// Refraction also fades out as a function of the coverage of the surfaces above.
+		float3 RefractedL = refract(-InContext.L, InContext.N, lerp(1.0f, AirIOR / LayersAboveIOR, InBSDF.CoverageAboveAlongN));
 		SatNoL = saturate(dot(InContext.N, -RefractedL));
 	}
 #endif
