Substrate - Fixed missing energy conservation on translucent due to a recent change.

#rnx
#tests Already mentioned to Alan in a thread.

diff --git a/Engine/Shaders/Private/Substrate/SubstrateEvaluation.ush b/Engine/Shaders/Private/Substrate/SubstrateEvaluation.ush
--- a/Engine/Shaders/Private/Substrate/SubstrateEvaluation.ush
+++ b/Engine/Shaders/Private/Substrate/SubstrateEvaluation.ush
@@ -433,7 +433,10 @@ FSubstrateEvaluateResult SubstrateEvaluateBSDFCommon(FSubstrateBSDFContext BSDFC
 				//	- We clamp SpecularTChan to a threshold to avoid div by 0 when the material is a metal. 
 				//	- When the material tends to be metal, DirectionalAlbedo_SpecularTransmission will naturally tend to be 0
 				// Bypass this correction when SUBSTRATE_SHADING_QUALITY > 1 (i.e., med/low quality), for trading accuracy for speed
-				if (Settings.bRoughDiffuseEnabled)
+				// We also do not apply that correction when rendering SIMPLE_VOLUME slabs:
+				//	- In this case there is no specular transmission baked in the participating media scattering model we use for coating layer.
+				//	- And specular throughput needs to be applied on the surface transmittance too to be energy conservative.
+				if (Settings.bRoughDiffuseEnabled && BSDF_GETSSSTYPE(BSDFContext.BSDF) != SSS_TYPE_SIMPLEVOLUME)
 				{
 					const float ChanF0 = 0.04;
 					FBxDFEnergyTerms EnergyTerms = ComputeGGXSpecEnergyTerms(SafeRoughness, BSDFContext.Context.NoV, ChanF0, 1.0f);
