Actually use shared GetMipLevelLocal function from the shader side of things.

#rnx

diff --git a/Engine/Shaders/Private/VirtualShadowMaps/VirtualShadowMapPageMarking.usf b/Engine/Shaders/Private/VirtualShadowMaps/VirtualShadowMapPageMarking.usf
--- a/Engine/Shaders/Private/VirtualShadowMaps/VirtualShadowMapPageMarking.usf
+++ b/Engine/Shaders/Private/VirtualShadowMaps/VirtualShadowMapPageMarking.usf
@@ -54,23 +54,7 @@ uint GetMipLevelLocal(int VirtualShadowMapId, float3 TranslatedWorldPosition, fl
 	float3 ShadowTranslatedWorldPosition = TranslatedWorldPosition + ViewToShadowTranslation;
 
 	float Footprint = VirtualShadowMapCalcPixelFootprintLocal(ProjectionData, ShadowTranslatedWorldPosition, SceneDepth);
-
-	float MipLevelFloat = log2(Footprint) + ProjectionData.ResolutionLodBias + VirtualShadowMap.GlobalResolutionLodBias;
-	uint MipLevel = uint(max(floor(MipLevelFloat), 0.0f));
-	MipLevel = min(MipLevel, (VSM_MAX_MIP_LEVELS - 1U));
-
-	if (MipModeLocal == 1)
-	{
-		// Even mips
-		MipLevel &= ~(1U);
-	}
-	else if (MipModeLocal == 2)
-	{
-		// Odd mips
-		MipLevel |= 1U;
-	}
-
-	return MipLevel;
+	return GetMipLevelLocal(Footprint, MipModeLocal, ProjectionData.ResolutionLodBias, VirtualShadowMap.GlobalResolutionLodBias);
 }
 
 void MarkPageAddress(uint PageOffset, uint Flags)
