Fixed accidental double up of adding MinMipLevel to the VSM shadow result, it's already added in ShadowVirtualToPhysicalUV.

#rnx

diff --git a/Engine/Shaders/Private/VirtualShadowMaps/VirtualShadowMapProjectionCommon.ush b/Engine/Shaders/Private/VirtualShadowMaps/VirtualShadowMapProjectionCommon.ush
--- a/Engine/Shaders/Private/VirtualShadowMaps/VirtualShadowMapProjectionCommon.ush
+++ b/Engine/Shaders/Private/VirtualShadowMaps/VirtualShadowMapProjectionCommon.ush
@@ -108,7 +108,7 @@ FVirtualShadowMapSample SampleVirtualShadowMap(int VirtualShadowMapId, float2 Sh
 	if (Page.bValid)
 	{
 		Result.bValid = true;
-		Result.MipLevel = Page.LODOffset + MinMipLevel;
+		Result.MipLevel = Page.LODOffset;
 		Result.VirtualShadowMapId = VirtualShadowMapId;
 		Result.VirtualTexelAddress = Page.VirtualTexelAddress;
 		Result.VirtualTexelAddressFloat = Page.VirtualTexelAddressFloat;
