[Light Grid] Clean up shader logic
- return FScreenRect directly from ComputeCellCullRect(...).
- use FCellBounds instead of individual min/max parameters.
- moved light/reflection capture overlap logic to individual functions.
- Removed unused ForwardLocalLightBuffer parameter.

diff --git a/Engine/Shaders/Private/LightGridInjection.usf b/Engine/Shaders/Private/LightGridInjection.usf
--- a/Engine/Shaders/Private/LightGridInjection.usf
+++ b/Engine/Shaders/Private/LightGridInjection.usf
@@ -28,8 +28,6 @@ RWStructuredBuffer<uint> RWNextCulledLightLink;
 RWStructuredBuffer<uint> RWStartOffsetGrid;
 RWStructuredBuffer<uint> RWCulledLightLinks;
 
-StructuredBuffer<float4> ForwardLocalLightBuffer;
-
 uint NumReflectionCaptures;
 uint NumLocalLights;
 uint NumGridCells;
@@ -80,7 +78,7 @@ float ComputeCellNearViewDepthFromZSlice(uint ZSlice)
 	return SliceDepth;
 }
 
-void ComputeCellCullRect(uint3 GridCoordinate, float MinTileZ, float MaxTileZ, out float3 CullRectMin, out float3 CullRectMax)
+FScreenRect ComputeCellCullRect(uint3 GridCoordinate, float MinTileZ, float MaxTileZ)
 {
 	// Compute extent of tiles in clip-space. Note that the last tile may extend a bit outside of view if view size is not evenly divisible tile size.
 	const float2 InvCulledGridSizeF = (1u << LightGridPixelSizeShift) * View.ViewSizeAndInvSize.zw;
@@ -93,16 +91,28 @@ void ComputeCellCullRect(uint3 GridCoordinate, float MinTileZ, float MaxTileZ, o
 	float MinTileDeviceZ = ConvertToDeviceZ(MinTileZ);
 	float MaxTileDeviceZ = ConvertToDeviceZ(MaxTileZ);
 
+	float3 CullRectMin;
 	CullRectMin.x = min(UnitPlaneTileMin.x, UnitPlaneTileMax.x);
 	CullRectMin.y = min(UnitPlaneTileMin.y, UnitPlaneTileMax.y);
 	CullRectMin.z = min(MinTileDeviceZ, MaxTileDeviceZ);
 
+	float3 CullRectMax;
 	CullRectMax.x = max(UnitPlaneTileMin.x, UnitPlaneTileMax.x);
 	CullRectMax.y = max(UnitPlaneTileMin.y, UnitPlaneTileMax.y);
 	CullRectMax.z = max(MinTileDeviceZ, MaxTileDeviceZ);
+
+	return GetScreenRect(int4(0, 0, HZBViewSize), CullRectMin, CullRectMax, 4);
 }
 
-void ComputeCellViewAABB(uint3 GridCoordinate, float MinTileZ, float MaxTileZ, out float3 ViewTileMin, out float3 ViewTileMax)
+struct FCellBounds
+{
+	float3 ViewAabbMin;
+	float3 ViewAabbMax;
+	float3 ViewCenter;
+	float3 ViewExtent;
+};
+
+FCellBounds ComputeCellBounds(uint3 GridCoordinate, float MinTileZ, float MaxTileZ)
 {
 	// Compute extent of tiles in clip-space. Note that the last tile may extend a bit outside of view if view size is not evenly divisible tile size.
 	const float2 InvCulledGridSizeF = (1u << LightGridPixelSizeShift) * View.ViewSizeAndInvSize.zw;
@@ -133,25 +143,29 @@ void ComputeCellViewAABB(uint3 GridCoordinate, float MinTileZ, float MaxTileZ, o
 	float2 ViewMaxDepthCorner2 = MaxDepthCorner2.xy / MaxDepthCorner2.w;
 	float2 ViewMaxDepthCorner3 = MaxDepthCorner3.xy / MaxDepthCorner3.w;
 
+	FCellBounds CellBounds;
+
 	//@todo - derive min and max from quadrant
-	ViewTileMin.xy = min(ViewMinDepthCorner0, ViewMinDepthCorner1);
-	ViewTileMin.xy = min(ViewTileMin.xy, ViewMinDepthCorner2);
-	ViewTileMin.xy = min(ViewTileMin.xy, ViewMinDepthCorner3);
-	ViewTileMin.xy = min(ViewTileMin.xy, ViewMaxDepthCorner0);
-	ViewTileMin.xy = min(ViewTileMin.xy, ViewMaxDepthCorner1);
-	ViewTileMin.xy = min(ViewTileMin.xy, ViewMaxDepthCorner2);
-	ViewTileMin.xy = min(ViewTileMin.xy, ViewMaxDepthCorner3);
-
-	ViewTileMax.xy = max(ViewMinDepthCorner0, ViewMinDepthCorner1);
-	ViewTileMax.xy = max(ViewTileMax.xy, ViewMinDepthCorner2);
-	ViewTileMax.xy = max(ViewTileMax.xy, ViewMinDepthCorner3);
-	ViewTileMax.xy = max(ViewTileMax.xy, ViewMaxDepthCorner0);
-	ViewTileMax.xy = max(ViewTileMax.xy, ViewMaxDepthCorner1);
-	ViewTileMax.xy = max(ViewTileMax.xy, ViewMaxDepthCorner2);
-	ViewTileMax.xy = max(ViewTileMax.xy, ViewMaxDepthCorner3);
-
-	ViewTileMin.z = MinTileZ;
-	ViewTileMax.z = MaxTileZ;
+	CellBounds.ViewAabbMin.xy = min(ViewMinDepthCorner0, ViewMinDepthCorner1);
+	CellBounds.ViewAabbMin.xy = min(CellBounds.ViewAabbMin.xy, ViewMinDepthCorner2);
+	CellBounds.ViewAabbMin.xy = min(CellBounds.ViewAabbMin.xy, ViewMinDepthCorner3);
+	CellBounds.ViewAabbMin.xy = min(CellBounds.ViewAabbMin.xy, ViewMaxDepthCorner0);
+	CellBounds.ViewAabbMin.xy = min(CellBounds.ViewAabbMin.xy, ViewMaxDepthCorner1);
+	CellBounds.ViewAabbMin.xy = min(CellBounds.ViewAabbMin.xy, ViewMaxDepthCorner2);
+	CellBounds.ViewAabbMin.xy = min(CellBounds.ViewAabbMin.xy, ViewMaxDepthCorner3);
+	CellBounds.ViewAabbMax.xy = max(ViewMinDepthCorner0, ViewMinDepthCorner1);
+	CellBounds.ViewAabbMax.xy = max(CellBounds.ViewAabbMax.xy, ViewMinDepthCorner2);
+	CellBounds.ViewAabbMax.xy = max(CellBounds.ViewAabbMax.xy, ViewMinDepthCorner3);
+	CellBounds.ViewAabbMax.xy = max(CellBounds.ViewAabbMax.xy, ViewMaxDepthCorner0);
+	CellBounds.ViewAabbMax.xy = max(CellBounds.ViewAabbMax.xy, ViewMaxDepthCorner1);
+	CellBounds.ViewAabbMax.xy = max(CellBounds.ViewAabbMax.xy, ViewMaxDepthCorner2);
+	CellBounds.ViewAabbMax.xy = max(CellBounds.ViewAabbMax.xy, ViewMaxDepthCorner3);
+	CellBounds.ViewAabbMin.z = MinTileZ;
+	CellBounds.ViewAabbMax.z = MaxTileZ;
+	CellBounds.ViewCenter = .5f * (CellBounds.ViewAabbMin + CellBounds.ViewAabbMax);
+	CellBounds.ViewExtent = CellBounds.ViewAabbMax - CellBounds.ViewCenter;
+
+	return CellBounds;
 }
 
 bool IntersectConeWithSphere(float3 ConeVertex, float3 ConeAxis, float ConeRadius, float2 CosSinAngle, float4 SphereToTest)
@@ -167,17 +181,16 @@ bool IntersectConeWithSphere(float3 ConeVertex, float3 ConeAxis, float ConeRadiu
 	return !(bSphereTooFarFromCone || bSpherePastConeEnd || bSphereBehindVertex);
 }
 
-
-
-bool AabbOutsidePlane(float3 center, float3 extents, float4 plane)
+/**
+ * Returns true if the aabb defined by Center and Extents is fully in front of the plane.
+ */
+bool IsAabbInFrontOfPlane(float3 Center, float3 Extents, float4 Plane)
 {
-	float dist = dot(float4(center, 1.0), plane);
-	float radius = dot(extents, abs(plane.xyz));
-
-	return dist > radius;
+	float Dist = dot(float4(Center, 1.0), Plane);
+	float Radius = dot(Extents, abs(Plane.xyz));
+	return Dist > Radius;
 }
 
-
 /**
  * Approximate cone / aabb test that creates a single separating plane that lies in the cone on the side facing the centre of the Aabb
  * Returns false if the Aabb is outside (entirely on the positive side) of this plane.
@@ -195,7 +208,77 @@ bool IsAabbOutsideInfiniteAcuteConeApprox(float3 ConeVertex, float3 ConeAxis, fl
 	float3 N = -TanConeAngle * ConeAxis + M;
 	float4 Plane = float4(N, 0.0);
 
-	return AabbOutsidePlane(D, AabbExt, Plane);
+	return IsAabbInFrontOfPlane(D, AabbExt, Plane);
+}
+
+bool OverlapsLight(uint LocalLightIndex, FCellBounds CellBounds, out bool bIsRectLight, out bool bIsTexturedLight)
+{
+	float4 LightPositionAndRadius = LightViewSpacePositionAndRadius[LocalLightIndex];
+	float3 ViewSpaceLightPosition = LightPositionAndRadius.xyz;
+	float LightRadius = LightPositionAndRadius.w;
+	float BoxDistanceSq = ComputeSquaredDistanceFromBoxToPoint(CellBounds.ViewCenter, CellBounds.ViewExtent, ViewSpaceLightPosition);
+	if (BoxDistanceSq < LightRadius * LightRadius)
+	{
+		float4 ViewSpaceDirAndPreprocAngle = LightViewSpaceDirAndPreprocAngle[LocalLightIndex];
+		bIsRectLight = asuint(ViewSpaceDirAndPreprocAngle.w) & 0x1;
+		bool bUseTightRectLightCulling = asuint(ViewSpaceDirAndPreprocAngle.w) & 0x2;
+		bIsTexturedLight = asuint(ViewSpaceDirAndPreprocAngle.w) & 0x4;
+
+		bool bPassSpotLightTest = true;
+#if REFINE_SPOTLIGHT_BOUNDS
+		{
+			float TanConeAngle = asfloat(asuint(ViewSpaceDirAndPreprocAngle.w) & 0xFFFFFFF8);
+
+			// Set to 0 for non-acute cones, or non-spot lights.
+			if (TanConeAngle > 0.0f)
+			{
+				float3 ViewSpaceLightDirection = -ViewSpaceDirAndPreprocAngle.xyz;
+				bPassSpotLightTest = !IsAabbOutsideInfiniteAcuteConeApprox(ViewSpaceLightPosition, ViewSpaceLightDirection, TanConeAngle, CellBounds.ViewCenter, CellBounds.ViewExtent);
+			}
+		}
+#endif // REFINE_SPOTLIGHT_BOUNDS
+		bool bPassRectLightTest = true;
+#if REFINE_RECTLIGHT_BOUNDS
+		{
+			if (bIsRectLight)
+			{
+				float3 D = CellBounds.ViewCenter - ViewSpaceLightPosition;
+				float4 Plane = float4(ViewSpaceDirAndPreprocAngle.xyz, 0.0);
+
+				bPassRectLightTest = !IsAabbInFrontOfPlane(D, CellBounds.ViewExtent, Plane);
+
+				if(bPassRectLightTest && bUseTightRectLightCulling)
+				{
+					for (uint PlaneIndex = 0; PlaneIndex < NUM_PLANES_PER_RECT_LIGHT; PlaneIndex++)
+					{
+						if (IsAabbInFrontOfPlane(CellBounds.ViewCenter, CellBounds.ViewExtent, LightViewSpaceRectPlanes[LocalLightIndex * NUM_PLANES_PER_RECT_LIGHT + PlaneIndex]))
+						{
+							bPassRectLightTest = false;
+						}
+					}
+				}
+			}
+		}
+#endif // REFINE_RECTLIGHT_BOUNDS
+		return bPassSpotLightTest && bPassRectLightTest;
+	}
+
+	bIsRectLight = false;
+	bIsTexturedLight = false;
+
+	return false;
+}
+
+bool OverlapsReflectionCapture(uint ReflectionCaptureIndex, FCellBounds CellBounds)
+{
+	FDFVector3 CaptureWorldPosition = MakeDFVector3(GetReflectionPositionAndRadius(ReflectionCaptureIndex).xyz, GetReflectionPositionLow(ReflectionCaptureIndex).xyz);
+	float3 CaptureTranslatedWorldPosition = DFFastAddDemote(CaptureWorldPosition, PrimaryView.PreViewTranslation);
+	float3 ViewSpaceCapturePosition = mul(float4(CaptureTranslatedWorldPosition, 1), View.TranslatedWorldToView).xyz;
+	float CaptureRadius = GetReflectionPositionAndRadius(ReflectionCaptureIndex).w;
+
+	float BoxDistanceSq = ComputeSquaredDistanceFromBoxToPoint(CellBounds.ViewCenter, CellBounds.ViewExtent, ViewSpaceCapturePosition);
+
+	return (BoxDistanceSq < CaptureRadius * CaptureRadius);
 }
 
 #if SHADER_LIGHT_GRID_INJECTION_CS
@@ -214,17 +297,13 @@ void LightGridInjectionCS(
 
 // Disable to pass all lights through for debugging, will hit limits quickly though
 #define CULL_LIGHTS 1
-	#if CULL_LIGHTS
+#if CULL_LIGHTS
 
 		float MinTileZ = ComputeCellNearViewDepthFromZSlice(GridCoordinate.z);
 		float MaxTileZ = ComputeCellNearViewDepthFromZSlice(GridCoordinate.z + 1);
 
 #if USE_HZB_CULL
-		float3 CullRectMin;
-		float3 CullRectMax;
-		ComputeCellCullRect(GridCoordinate, MinTileZ, MaxTileZ, CullRectMin, CullRectMax);
-
-		FScreenRect Rect = GetScreenRect(int4(0, 0, HZBViewSize), CullRectMin, CullRectMax, 4);
+		FScreenRect Rect = ComputeCellCullRect(GridCoordinate, MinTileZ, MaxTileZ);
 
  		if (!IsVisibleHZB(Rect, true /*bSample4x4*/))
  		{
@@ -240,14 +319,9 @@ void LightGridInjectionCS(
 
 			MaxTileZ = min(MaxTileZ, FarthestSceneDepth);
 		}
-#endif
+#endif // USE_HZB_CULL
 		
-		float3 ViewTileMin;
-		float3 ViewTileMax;
-		ComputeCellViewAABB(GridCoordinate, MinTileZ, MaxTileZ, ViewTileMin, ViewTileMax);
-
-		float3 ViewTileCenter = .5f * (ViewTileMin + ViewTileMax);
-		float3 ViewTileExtent = ViewTileMax - ViewTileCenter;
+		FCellBounds CellBounds = ComputeCellBounds(GridCoordinate, MinTileZ, MaxTileZ);
 		
 		uint NumAvailableLinks = NumGridCells * MaxCulledLightsPerCell * NUM_CULLED_GRID_PRIMITIVE_TYPES;
 
@@ -259,87 +333,36 @@ void LightGridInjectionCS(
 		LOOP
 		for (uint LocalLightIndex = 0; LocalLightIndex < NumLocalLights; LocalLightIndex++)
 		{
-			uint LocalLightBaseIndex = LocalLightIndex * LOCAL_LIGHT_DATA_STRIDE;
-
-			float4 LightPositionAndRadius = LightViewSpacePositionAndRadius[LocalLightIndex];
-			float3 ViewSpaceLightPosition = LightPositionAndRadius.xyz;
-			float LightRadius = LightPositionAndRadius.w;
-
-			float BoxDistanceSq = ComputeSquaredDistanceFromBoxToPoint(ViewTileCenter, ViewTileExtent, ViewSpaceLightPosition);
-
-			if (BoxDistanceSq < LightRadius * LightRadius)
+			bool bIsRectLight;
+			bool bIsTexturedLight;
+			if (OverlapsLight(LocalLightIndex, CellBounds, bIsRectLight, bIsTexturedLight))
 			{
-				float4 ViewSpaceDirAndPreprocAngle = LightViewSpaceDirAndPreprocAngle[LocalLightIndex];
-				bool bIsRectLight = asuint(ViewSpaceDirAndPreprocAngle.w) & 0x1;
-				bool bUseTightRectLightCulling = asuint(ViewSpaceDirAndPreprocAngle.w) & 0x2;
-				bool bIsTexturedLight = asuint(ViewSpaceDirAndPreprocAngle.w) & 0x4;
+				bHasRectLights |= bIsRectLight;
+				bHasTexturedLights |= bIsTexturedLight;
 
-				bool bPassSpotLightTest = true;
-#if REFINE_SPOTLIGHT_BOUNDS
-				{
-					float TanConeAngle = asfloat(asuint(ViewSpaceDirAndPreprocAngle.w) & 0xFFFFFFF8);
+#if USE_LINKED_CULL_LIST
+				uint NextLink;
+				InterlockedAdd(RWNextCulledLightLink[0], 1U, NextLink);
 
-					// Set to 0 for non-acute cones, or non-spot lights.
-					if (TanConeAngle > 0.0f)
-					{
-						float3 ViewSpaceLightDirection = -ViewSpaceDirAndPreprocAngle.xyz;
-						bPassSpotLightTest = !IsAabbOutsideInfiniteAcuteConeApprox(ViewSpaceLightPosition, ViewSpaceLightDirection, TanConeAngle, ViewTileCenter, ViewTileExtent);
-					}
-				}
-#endif // REFINE_SPOTLIGHT_BOUNDS
-				bool bPassRectLightTest = true;
-#if REFINE_RECTLIGHT_BOUNDS
+				if (NextLink < NumAvailableLinks)
 				{
-					if (bIsRectLight)
-					{
-						float3 D = ViewTileCenter - ViewSpaceLightPosition;
-						float4 Plane = float4(ViewSpaceDirAndPreprocAngle.xyz, 0.0);
-
-						bPassRectLightTest = !AabbOutsidePlane(D, ViewTileExtent, Plane);
-
-						if(bPassRectLightTest && bUseTightRectLightCulling)
-						{
-							for (uint PlaneIndex = 0; PlaneIndex < NUM_PLANES_PER_RECT_LIGHT; PlaneIndex++)
-							{
-								if (AabbOutsidePlane(ViewTileCenter, ViewTileExtent, LightViewSpaceRectPlanes[LocalLightIndex * NUM_PLANES_PER_RECT_LIGHT + PlaneIndex]))
-								{
-									bPassRectLightTest = false;
-								}
-							}
-						}
-					}
+					uint PreviousLink;
+					InterlockedExchange(RWStartOffsetGrid[GridIndex], NextLink, PreviousLink);
+					RWCulledLightLinks[NextLink * LIGHT_LINK_STRIDE + 0] = LocalLightIndex;
+					RWCulledLightLinks[NextLink * LIGHT_LINK_STRIDE + 1] = PreviousLink;
 				}
-#endif // REFINE_RECTLIGHT_BOUNDS
-				if (bPassSpotLightTest && bPassRectLightTest)
+#else // !USE_LINKED_CULL_LIST
+				if (NumVisibleLights < MaxCulledLightsPerCell)
 				{
-					bHasRectLights |= bIsRectLight;
-					bHasTexturedLights |= bIsTexturedLight;
-
-#if USE_LINKED_CULL_LIST
-					uint NextLink;
-					InterlockedAdd(RWNextCulledLightLink[0], 1U, NextLink);
-
-					if (NextLink < NumAvailableLinks)
-					{
-						uint PreviousLink;
-						InterlockedExchange(RWStartOffsetGrid[GridIndex], NextLink, PreviousLink);
-						RWCulledLightLinks[NextLink * LIGHT_LINK_STRIDE + 0] = LocalLightIndex;
-						RWCulledLightLinks[NextLink * LIGHT_LINK_STRIDE + 1] = PreviousLink;
-					}
-
-#else
-					if (NumVisibleLights < MaxCulledLightsPerCell)
-					{
-						RWCulledLightDataGrid[GridIndex * MaxCulledLightsPerCell + NumVisibleLights] = LocalLightIndex;
-					}
+					RWCulledLightDataGrid[GridIndex * MaxCulledLightsPerCell + NumVisibleLights] = LocalLightIndex;
+				}
 
-					++NumVisibleLights;
-					if (LocalLightIndex >= MegaLightsSupportedStartIndex)
-					{
-						++NumVisibleMegaLights;
-					}
-#endif
+				++NumVisibleLights;
+				if (LocalLightIndex >= MegaLightsSupportedStartIndex)
+				{
+					++NumVisibleMegaLights;
 				}
+#endif // !USE_LINKED_CULL_LIST
 			}
 		}
 
@@ -349,15 +372,8 @@ void LightGridInjectionCS(
 		LOOP
 		for (uint ReflectionCaptureIndex = 0; ReflectionCaptureIndex < NumReflectionCaptures; ReflectionCaptureIndex++)
 		{
-			FDFVector3 CaptureWorldPosition = MakeDFVector3(GetReflectionPositionAndRadius(ReflectionCaptureIndex).xyz, GetReflectionPositionLow(ReflectionCaptureIndex).xyz);
-			float3 CaptureTranslatedWorldPosition = DFFastAddDemote(CaptureWorldPosition, PrimaryView.PreViewTranslation);
-			float3 ViewSpaceCapturePosition = mul(float4(CaptureTranslatedWorldPosition, 1), View.TranslatedWorldToView).xyz;
-			float CaptureRadius = GetReflectionPositionAndRadius(ReflectionCaptureIndex).w;
-
-			float BoxDistanceSq = ComputeSquaredDistanceFromBoxToPoint(ViewTileCenter, ViewTileExtent, ViewSpaceCapturePosition);
-
-			if (BoxDistanceSq < CaptureRadius * CaptureRadius)
-			{
+			if (OverlapsReflectionCapture(ReflectionCaptureIndex, CellBounds))
+			{			
 				#if USE_LINKED_CULL_LIST
 					uint NextLink;
 					InterlockedAdd(RWNextCulledLightLink[0], 1U, NextLink);
@@ -382,8 +398,7 @@ void LightGridInjectionCS(
 				#endif
 			}
 		}
-#else
-
+#else // !CULL_LIGHTS
 		LOOP
 		for (uint LocalLightIndex = 0; LocalLightIndex < NumLocalLights; LocalLightIndex++)
 		{
@@ -407,7 +422,7 @@ void LightGridInjectionCS(
 
 		RWNumCulledLightsGrid[(NumGridCells + GridIndex) * NUM_CULLED_LIGHTS_GRID_STRIDE + 0] = NumReflectionCaptures;
 		RWNumCulledLightsGrid[(NumGridCells + GridIndex) * NUM_CULLED_LIGHTS_GRID_STRIDE + 1] = (NumGridCells + GridIndex) * MaxCulledLightsPerCell;
-#endif
+#endif // !CULL_LIGHTS
 	}
 }
 
@@ -434,7 +449,7 @@ void CompactReverseLinkedList(uint GridIndex, uint SceneMax, bool FirstThread, b
 	uint LinkOffset				= 0;
 	uint CulledLightDataStart	= 0;
 
-	if(bThreadValid)
+	if (bThreadValid)
 	{
 		StartLinkOffset = StartOffsetGrid[GridIndex];
 		LinkOffset = StartLinkOffset;
@@ -447,9 +462,7 @@ void CompactReverseLinkedList(uint GridIndex, uint SceneMax, bool FirstThread, b
 		}
 	}
 
-
 #if USE_LOCALSTORE_FOR_ATOMICS
-
 	LDSNextCulledLightData	= 0;
 	LDSBaseStart			= 0;
 
@@ -467,13 +480,13 @@ void CompactReverseLinkedList(uint GridIndex, uint SceneMax, bool FirstThread, b
 	GroupMemoryBarrierWithGroupSync();
 	
 	CulledLightDataStart += LDSBaseStart;
-#else
+#else // !USE_LOCALSTORE_FOR_ATOMICS
 	InterlockedAdd(RWNextCulledLightData[0], NumCulledLights, CulledLightDataStart);
-#endif
+#endif // !USE_LOCALSTORE_FOR_ATOMICS
+
 
-	if( bThreadValid)
+	if (bThreadValid)
 	{
-	
 		LinkOffset = StartLinkOffset;
 		uint CulledLightIndex = 0;
 		uint NumMegaLights = 0;
diff --git a/Engine/Source/Runtime/Renderer/Private/LightGridInjection.cpp b/Engine/Source/Runtime/Renderer/Private/LightGridInjection.cpp
--- a/Engine/Source/Runtime/Renderer/Private/LightGridInjection.cpp
+++ b/Engine/Source/Runtime/Renderer/Private/LightGridInjection.cpp
@@ -221,7 +221,6 @@ public:
 		SHADER_PARAMETER_RDG_BUFFER_SRV(StructuredBuffer<float4>, LightViewSpacePositionAndRadius)
 		SHADER_PARAMETER_RDG_BUFFER_SRV(StructuredBuffer<float4>, LightViewSpaceDirAndPreprocAngle)
 		SHADER_PARAMETER_RDG_BUFFER_SRV(StructuredBuffer<float4>, LightViewSpaceRectPlanes)
-		SHADER_PARAMETER_RDG_BUFFER_SRV(StructuredBuffer<float4>, ForwardLocalLightBuffer)
 
 		SHADER_PARAMETER(FIntVector, CulledGridSize)
 		SHADER_PARAMETER(uint32, NumReflectionCaptures)
@@ -963,7 +962,6 @@ FComputeLightGridOutput FSceneRenderer::ComputeLightGrid(FRDGBuilder& GraphBuild
 			PassParameters->RWNextCulledLightLink   = GraphBuilder.CreateUAV(NextCulledLightLinkBuffer);
 			PassParameters->RWStartOffsetGrid       = GraphBuilder.CreateUAV(StartOffsetGridBuffer);
 			PassParameters->RWCulledLightLinks      = GraphBuilder.CreateUAV(CulledLightLinksBuffer);
-			PassParameters->ForwardLocalLightBuffer = ForwardLightData->ForwardLocalLightBuffer;
 			PassParameters->CulledGridSize          = ForwardLightData->CulledGridSize;
 			PassParameters->LightGridZParams        = ForwardLightData->LightGridZParams;
 			PassParameters->NumReflectionCaptures   = ForwardLightData->NumReflectionCaptures;
