Update how hair strands coverage computation is done.
* Compute actual coverage per pixel for each segment.
* Clip segment within its voxel to ensure that the segment's coverage is only accounted once.

diff --git a/Engine/Plugins/Runtime/HairStrands/Source/HairStrandsCore/Private/GroomTextureBuilder.cpp b/Engine/Plugins/Runtime/HairStrands/Source/HairStrandsCore/Private/GroomTextureBuilder.cpp
--- a/Engine/Plugins/Runtime/HairStrands/Source/HairStrandsCore/Private/GroomTextureBuilder.cpp
+++ b/Engine/Plugins/Runtime/HairStrands/Source/HairStrandsCore/Private/GroomTextureBuilder.cpp
@@ -958,13 +958,7 @@ static void InternalGenerateHairStrandsTextures(
 				RHICmdList.SetStreamSource(0, nullptr, 0);
 				RHICmdList.SetViewport(0, 0, 0.0f, OutputResolution.X, OutputResolution.Y, 1.0f);
 			
-				// Divide the rendering work into small batches to reduce risk of TDR as the texture projection implies heavy works 
-				// (i.e. long thread running due the the large amount of strands a groom can have)
-				const int32 TileSize = HairStrandsTextureTileSize;
-				if (OutputResolution.X > TileSize)
-				{
-					RHICmdList.SetScissorRect(true, TileOffsetInPixels.X, TileOffsetInPixels.Y, TileOffsetInPixels.X + TileSize, TileOffsetInPixels.Y + TileSize);
-				}
+				// The visibility pass is no longer tile in order to avoid issue at tile borde, preventing to compute derivative during the tracing step
 				RHICmdList.DrawIndexedPrimitive(InMeshIndexBuffer, VertexBaseIndex, 0, VertexCount, IndexBaseIndex, PrimitiveCount, 1);
 			});
 	}
diff --git a/Engine/Shaders/Private/HairStrands/HairStrandsTexturesGeneration.usf b/Engine/Shaders/Private/HairStrands/HairStrandsTexturesGeneration.usf
--- a/Engine/Shaders/Private/HairStrands/HairStrandsTexturesGeneration.usf
+++ b/Engine/Shaders/Private/HairStrands/HairStrandsTexturesGeneration.usf
@@ -244,10 +244,88 @@ uint CoordToIndex(uint3 InCoord, uint3 InResolution, uint LinearOffset)
 		LinearOffset;
 }
 
+
+bool ClipRaySegment(float3 AABBMin, float3 AABBMax, float4 P0, float4 P1, out float2 T, out bool3 bClipped)
+{
+	bClipped = false;
+	T = float2(0.0f, 1.0f);
+	const bool bP0Outside = any(P0.xyz < AABBMin) || any(P0.xyz > AABBMax);
+	const bool bP1Outside = any(P1.xyz < AABBMin) || any(P1.xyz > AABBMax);
+	if (!bP0Outside && !bP1Outside)
+	{
+		return true;
+	}
+
+	const float3 Origin = P0.xyz;
+	const float3 Dir = P1.xyz - P0.xyz;
+	const float3 RcpDir = 1.0f / Dir;
+
+	const float3 T0 = (AABBMin - Origin) * RcpDir;
+	const float3 T1 = (AABBMax - Origin) * RcpDir;
+
+	T.x = max3(min(T0.x, T1.x), min(T0.y, T1.y), min(T0.z, T1.z));
+	T.y = min3(max(T0.x, T1.x), max(T0.y, T1.y), max(T0.z, T1.z));
+
+	// Ray intersects the AABB but the segment is completely outside or no intersection at all.
+	if (T.y < 0.0f || T.x > T.y)
+	{
+		bClipped = true;
+		return false;
+	}
+
+	if (bP0Outside && T.x > 0.0f && T.x < 1.0f)
+	{
+		bClipped.x = true;
+	}
+	if (bP1Outside && T.y > 0.0f && T.y < 1.0f)
+	{
+		bClipped.y = true;
+	}
+
+	return true;
+}
+
+bool ClipRaySegment(float3 AABBMin, float3 AABBMax, inout float4 P0, inout float4 P1, inout float Rad0, inout float Rad1, out bool3 bClipped)
+{
+	float2 T;
+	bool bIsValid = ClipRaySegment(AABBMin, AABBMax, P0, P1, T, bClipped);
+	
+	if (bIsValid)
+	{
+		const bool bP0Outside = any(P0.xyz < AABBMin) || any(P0.xyz > AABBMax);
+		const bool bP1Outside = any(P1.xyz < AABBMin) || any(P1.xyz > AABBMax);
+
+		float4 P0New = P0;
+		float4 P1New = P1;
+		float Rad0New = Rad0;
+		float Rad1New = Rad1;
+		if (bP0Outside && T.x > 0.0f && T.x < 1.0f)
+		{
+			P0New = lerp(P0, P1, T.x);
+			Rad0New = lerp(Rad0, Rad1, T.x);
+			bClipped.x = true;
+		}
+		if (bP1Outside && T.y > 0.0f && T.y < 1.0f)
+		{
+			P1New = lerp(P0, P1, T.y);
+			Rad1New = lerp(Rad0, Rad1, T.y);
+			bClipped.y = true;
+		}
+
+		P0 = P0New;
+		P1 = P1New;
+		Rad0 = Rad0New;
+		Rad1 = Rad1New;
+	}
+
+	return bIsValid;
+}
+
 bool Trace(
 	float3 InRayO,
 	float3 InRayD,
 	float InMaxDistance,
+	float InPixelRadius,
 	bool bDebugEnable,
 	inout float  OutDepth,
 	inout float  OutCoverage,
@@ -340,7 +418,7 @@ bool Trace(
 			}
 
 			const float3 WPCurr = (PageIndexCoord + 0.5f) * Voxel_Size  + Voxel_MinBound;
-			const bool bExit = dot((WPCurr - O), InRayD) > (MaxDistance + Voxel_Size);
+			const bool bExit = dot((WPCurr - O), InRayD) > (InMaxDistance + Voxel_Size);
 			if (bExit)
 			{
 				return bIsValid;
@@ -379,8 +457,31 @@ bool Trace(
 						const uint P0 = P0_P1.x;
 						const uint P1 = P0_P1.y;
 
-						const FHairControlPoint CP_Curr = ReadHairControlPoint(InVF_PositionBuffer, P0, InVF_PositionOffset, InVF_Radius, InVF_RootScale, InVF_TipScale);
-						const FHairControlPoint CP_Next = ReadHairControlPoint(InVF_PositionBuffer, P1, InVF_PositionOffset, InVF_Radius, InVF_RootScale, InVF_TipScale);
+						FHairControlPoint CP_Curr = ReadHairControlPoint(InVF_PositionBuffer, P0, InVF_PositionOffset, InVF_Radius, InVF_RootScale, InVF_TipScale);
+						FHairControlPoint CP_Next = ReadHairControlPoint(InVF_PositionBuffer, P1, InVF_PositionOffset, InVF_Radius, InVF_RootScale, InVF_TipScale);
+
+						// Clip segment to voxel boundary in order to avoid counting multiple time the same segment which would cause the coverage to be incorrect 
+						// as the ray trace several voxels
+						{
+							const float3 MinVoxelAABB = Voxel_MinBound + PageIndexCoord * Voxel_Size;
+							const float3 MaxVoxelAABB = Voxel_MinBound + (PageIndexCoord + 1) * Voxel_Size;
+						
+							float4 ClippedP0  = float4(CP_Curr.Position, 1);
+							float4 ClippedP1  = float4(CP_Next.Position, 1);
+							float ClippedRad0 = CP_Curr.WorldRadius;
+							float ClippedRad1 = CP_Next.WorldRadius;
+							bool3 bClipped    = false;
+
+							ClipRaySegment(MinVoxelAABB, MaxVoxelAABB, ClippedP0, ClippedP1, ClippedRad0, ClippedRad1, bClipped);
+
+							if (any(bClipped))
+							{
+								CP_Curr.Position    = ClippedP0.xyz;
+								CP_Next.Position    = ClippedP1.xyz;
+								CP_Curr.WorldRadius = ClippedRad0;
+								CP_Next.WorldRadius = ClippedRad1;
+							}
+						}
 
 						const float3 Tangent = ComputeTangent(CP_Curr, CP_Next);
 
@@ -398,7 +499,10 @@ bool Trace(
 
 						if (bIntersect)
 						{
-							OutCoverage = 1;
+							const float SegmentCoverage = saturate(WorldRadius / InPixelRadius);
+							//OutCoverage = saturate(OutCoverage + saturate(1.f-OutCoverage) * SegmentCoverage);
+							OutCoverage = saturate(OutCoverage + SegmentCoverage);
+							//OutCoverage = 1;
 						}
 
 						if (bIntersectAndClosest)
@@ -474,14 +578,15 @@ void MainCS(
 	// Manual ddx/ddy
 	float3 PixelAxisX = 0;
 	float3 PixelAxisY = 0; 
-	if (PixelCoord.x + 1 < TileOffsetInPixels.x)
+	if (PixelCoord.x + 1 < OutputResolution.x)
 	{
 		PixelAxisX = PositionTexture[float2(PixelCoord.x+1, PixelCoord.y)].xyz - InPosition;
 	}
-	if (PixelCoord.y + 1 < TileOffsetInPixels.y)
+	if (PixelCoord.y + 1 < OutputResolution.y)
 	{
 		PixelAxisY = PositionTexture[float2(PixelCoord.x, PixelCoord.y+1)].xyz - InPosition;
 	}
+	const float PixelWorldRadius = 0.25f * (length(PixelAxisY) + length(PixelAxisX)); // Average + pixel radius (vs. diameter)
 
 	float4 OutTexture0 = 0;
 	float4 OutTexture1	= 0;
@@ -513,7 +618,7 @@ void MainCS(
 	const uint SampleCount = 16;
 	for (uint SampleIt = 0; SampleIt < SampleCount; ++SampleIt)
 	{
-		const float2 Jitter = Hammersley(SampleIt, SampleCount, 0) * 2 - 1;
+		const float2 Jitter = 0.5f * (Hammersley(SampleIt, SampleCount, 0) * 2 - 1);
 
 		float3 RayO = InPosition + Jitter.x * PixelAxisX + Jitter.y * PixelAxisY;
 		float3 RayD = N;
@@ -544,7 +649,7 @@ void MainCS(
 		float  SampleU = 0;
 		float  SampleSeed = 0;
 		float4 SampleMaterial = 0;
-		if (Trace(RayO, RayD, MaxDistance, bDebugEnable, SampleDepth, SampleCoverage, SampleTangent, SampleRootUV, SampleU, SampleSeed, SampleMaterial))
+		if (Trace(RayO, RayD, MaxDistance, PixelWorldRadius, bDebugEnable, SampleDepth, SampleCoverage, SampleTangent, SampleRootUV, SampleU, SampleSeed, SampleMaterial))
 		{
 			if (SampleDepth < OutDepth)
 			{
