HeterogeneousVolumes: Stochastic Filtering for cinematics

Stochastic filtering allows for cubic interpolation when performing shader evaluations. Since the operation is inherently stochastic, multiple samples-per-pixel are necessary. For this reason, stochastic filtering is only enabled for use with Movie Render Queue for now.

r.HeterogeneousVolumes.StochasticFiltering:
Configures the stochastic filtering kernel (Default = 3)
0: Disabled
1: Constant
2: Linear
3: Cubic

diff --git a/Engine/Shaders/Private/HeterogeneousVolumes/HeterogeneousVolumesAmbientOcclusionPipeline.usf b/Engine/Shaders/Private/HeterogeneousVolumes/HeterogeneousVolumesAmbientOcclusionPipeline.usf
--- a/Engine/Shaders/Private/HeterogeneousVolumes/HeterogeneousVolumesAmbientOcclusionPipeline.usf
+++ b/Engine/Shaders/Private/HeterogeneousVolumes/HeterogeneousVolumesAmbientOcclusionPipeline.usf
@@ -5,6 +5,7 @@
 #define SUPPORT_CONTACT_SHADOWS 0
 
 #include "/Engine/Generated/Material.ush"
+#include "HeterogeneousVolumesStochasticFiltering.ush"
 #include "HeterogeneousVolumesLiveShadingUtils.ush"
 #include "HeterogeneousVolumesTracingUtils.ush"
 #include "HeterogeneousVolumesLightingUtils.ush"
@@ -49,6 +50,7 @@ float StepSize;
 float StepFactor;
 int MaxStepCount;
 int bJitter;
+int StochasticFilteringMode;
 
 Texture3D ExistenceMaskTexture;
 
@@ -120,6 +122,11 @@ float GetMaxShadowTraceDistance()
 	return MaxShadowTraceDistance;
 }
 
+int GetStochasticFilteringMode()
+{
+	return StochasticFilteringMode;
+}
+
 // Material sampling hooks..
 struct FVolumeSampleContext
 {
@@ -140,10 +147,33 @@ FVolumeSampleContext CreateVolumeSampleContext(float3 LocalPosition, float3 Worl
 
 	// Evaluate Material graph
 	CalcPixelMaterialInputs(VolumeSampleContext.MaterialParameters, VolumeSampleContext.PixelMaterialInputs);
-
 	return VolumeSampleContext;
 }
 
+FVolumeSampleContext CreateVolumeSampleContext(float3 LocalPosition, float3 WorldPosition, float3 FilterWidth, float MipLevel, float Rand, int StochasticFilteringMode)
+{
+	// TODO: Use FilterWidth determine MipLevel
+	//float3 WorldPosition_DDX = mul(float4(LocalPosition + float3(1.0, 0.0, 0.0), 1.0), GetLocalToWorld()).xyz;
+	//float3 WorldPosition_DDY = mul(float4(LocalPosition + float3(0.0, 1.0, 0.0), 1.0), GetLocalToWorld()).xyz;
+	//float3 WorldPosition_DDZ = mul(float4(LocalPosition + float3(0.0, 0.0, 1.0), 1.0), GetLocalToWorld()).xyz;
+
+	if (StochasticFilteringMode == STOCHASTIC_FILTERING_CONSTANT)
+	{
+		LocalPosition = StochasticFilteringConstant(LocalPosition, Rand);
+	}
+	else if (StochasticFilteringMode == STOCHASTIC_FILTERING_TRILINEAR)
+	{
+		LocalPosition = StochasticFilteringTrilinear(LocalPosition, Rand);
+	}
+	else if (StochasticFilteringMode == STOCHASTIC_FILTERING_TRICUBIC)
+	{
+		LocalPosition = StochasticFilteringTricubic(LocalPosition, Rand);
+	}
+
+	WorldPosition = mul(float4(LocalPosition, 1.0), GetLocalToWorld()).xyz;
+	return CreateVolumeSampleContext(LocalPosition, WorldPosition, MipLevel);
+}
+
 float3 SampleExtinction(inout FVolumeSampleContext Context)
 {
 	float3 Extinction = SampleExtinctionCoefficients(Context.PixelMaterialInputs);
@@ -292,7 +322,17 @@ void RenderExistenceMaskWithLiveShadingCS(
 	float3 WorldPosition = mul(float4(LocalPosition, 1.0), GetLocalToWorld()).xyz;
 
 	int MipLevel = 0;
+#if IS_OFFLINE_RENDER
+	float3 FilterWidth = 1.0;
+	RandomSequence RandSequence;
+	int LinearIndex = VoxelIndex.x * (VoxelResolution.y * VoxelResolution.z) + VoxelIndex.y * VoxelResolution.z + VoxelIndex.z;
+	RandomSequence_Initialize(RandSequence, LinearIndex, View.StateFrameIndex);
+	float Rand = RandomSequence_GenerateSample1D(RandSequence);
+	int StochasticFilteringMode = GetStochasticFilteringMode();
+	FVolumeSampleContext SampleContext = CreateVolumeSampleContext(LocalPosition, WorldPosition, FilterWidth, MipLevel, Rand, StochasticFilteringMode);
+#else
 	FVolumeSampleContext SampleContext = CreateVolumeSampleContext(LocalPosition, WorldPosition, MipLevel);
+#endif // IS_OFFLINE_RENDER
 	float3 Extinction = SampleExtinction(SampleContext);
 	RWExistenceMaskTexture[VoxelIndex] = any(Extinction > 0.0) ? 1.0 : 0.0;
 }
diff --git a/Engine/Shaders/Private/HeterogeneousVolumes/HeterogeneousVolumesLiveShadingPipeline.usf b/Engine/Shaders/Private/HeterogeneousVolumes/HeterogeneousVolumesLiveShadingPipeline.usf
--- a/Engine/Shaders/Private/HeterogeneousVolumes/HeterogeneousVolumesLiveShadingPipeline.usf
+++ b/Engine/Shaders/Private/HeterogeneousVolumes/HeterogeneousVolumesLiveShadingPipeline.usf
@@ -35,6 +35,7 @@
 #endif // SUPPORT_OVERLAPPING_VOLUMES
 
 #include "/Engine/Generated/Material.ush"
+#include "HeterogeneousVolumesStochasticFiltering.ush"
 #include "HeterogeneousVolumesLiveShadingUtils.ush"
 #include "HeterogeneousVolumesTracingUtils.ush"
 #include "HeterogeneousVolumesLightingUtils.ush"
@@ -93,6 +94,7 @@ float StepSize;
 float StepFactor;
 int MaxStepCount;
 int bJitter;
+int StochasticFilteringMode;
 
 // Dispatch data
 int3 GroupCount;
@@ -197,6 +199,11 @@ float GetMaxShadowTraceDistance()
 	return MaxShadowTraceDistance;
 }
 
+int GetStochasticFilteringMode()
+{
+	return StochasticFilteringMode;
+}
+
 int ShouldApplyHeightFog()
 {
 	return bApplyHeightFog;
@@ -230,6 +237,30 @@ FVolumeSampleContext CreateVolumeSampleContext(float3 LocalPosition, float3 Worl
 	return VolumeSampleContext;
 }
 
+FVolumeSampleContext CreateVolumeSampleContext(float3 LocalPosition, float3 WorldPosition, float3 FilterWidth, float MipLevel, float Rand, int StochasticFilteringMode)
+{
+	// TODO: Use FilterWidth determine MipLevel
+	//float3 WorldPosition_DDX = mul(float4(LocalPosition + float3(1.0, 0.0, 0.0), 1.0), GetLocalToWorld()).xyz;
+	//float3 WorldPosition_DDY = mul(float4(LocalPosition + float3(0.0, 1.0, 0.0), 1.0), GetLocalToWorld()).xyz;
+	//float3 WorldPosition_DDZ = mul(float4(LocalPosition + float3(0.0, 0.0, 1.0), 1.0), GetLocalToWorld()).xyz;
+
+	if (StochasticFilteringMode == STOCHASTIC_FILTERING_CONSTANT)
+	{
+		LocalPosition = StochasticFilteringConstant(LocalPosition, Rand);
+	}
+	else if (StochasticFilteringMode == STOCHASTIC_FILTERING_TRILINEAR)
+	{
+		LocalPosition = StochasticFilteringTrilinear(LocalPosition, Rand);
+	}
+	else if (StochasticFilteringMode == STOCHASTIC_FILTERING_TRICUBIC)
+	{
+		LocalPosition = StochasticFilteringTricubic(LocalPosition, Rand);
+	}
+
+	WorldPosition = mul(float4(LocalPosition, 1.0), GetLocalToWorld()).xyz;
+	return CreateVolumeSampleContext(LocalPosition, WorldPosition, MipLevel);
+}
+
 float3 SampleExtinction(inout FVolumeSampleContext Context)
 {
 	float3 Extinction = SampleExtinctionCoefficients(Context.PixelMaterialInputs);
@@ -283,15 +314,7 @@ void RenderLightingCacheWithLiveShadingCS(
 		return;
 	}
 
-	float3 VoxelJitter = 0.5;
-	if (bJitter)
-	{
-		uint3 Rand32Bits = Rand4DPCG32(uint4(VoxelIndex, View.StateFrameIndexMod8)).xyz;
-		VoxelJitter = float3(Rand32Bits) / float(uint(0xffffffff));
-	}
-	//float3 UVW = (VoxelIndex + VoxelJitter) / float3(GetLightingCacheResolution());
 	float3 UVW = (VoxelIndex + 0.5) / float3(GetLightingCacheResolution());
-
 	float3 LocalBoundsMin = GetLocalBoundsOrigin() - GetLocalBoundsExtent();
 	float3 LocalBoundsMax = GetLocalBoundsOrigin() + GetLocalBoundsExtent();
 	float3 LocalRayOrigin = UVW * GetLocalBoundsExtent() * 2.0 + LocalBoundsMin;
@@ -302,12 +325,32 @@ void RenderLightingCacheWithLiveShadingCS(
 #if USE_EXISTENCE_MASK
 	float Existence = ExistenceMaskTexture.SampleLevel(View.SharedBilinearClampedSampler, UVW, MipLevel).r;
 	if (Existence <= 0.0)
+	{
+#if DIM_LIGHTING_CACHE_MODE == LIGHTING_CACHE_TRANSMITTANCE
+		RWLightingCacheTexture[VoxelIndex] = 1.0f;
+
+#elif DIM_LIGHTING_CACHE_MODE == LIGHTING_CACHE_INSCATTERING
+		RWLightingCacheTexture[VoxelIndex] = 0.0f;
+#endif // DIM_LIGHTING_CACHE_MODE
+		return;
+	}
+#endif
+
+#if IS_OFFLINE_RENDER
+	float3 FilterWidth = 1.0;
+	RandomSequence RandSequence;
+	int LinearIndex = VoxelIndex.x * (GetLightingCacheResolution().y * GetLightingCacheResolution().z) + VoxelIndex.y * GetLightingCacheResolution().z + VoxelIndex.z;
+	RandomSequence_Initialize(RandSequence, LinearIndex, View.StateFrameIndex);
+	float Rand = RandomSequence_GenerateSample1D(RandSequence);
+	int StochasticFilteringMode = GetStochasticFilteringMode();
+	FVolumeSampleContext SampleContext = CreateVolumeSampleContext(LocalRayOrigin, WorldRayOrigin, FilterWidth, MipLevel, Rand, StochasticFilteringMode);
 #else
 	FVolumeSampleContext SampleContext = CreateVolumeSampleContext(LocalRayOrigin, WorldRayOrigin, MipLevel);
+#endif // IS_OFFLINE_RENDER
+
 	float3 Extinction = SampleExtinction(SampleContext);
 	float Epsilon = 1.0e-7;
 	if (all(Extinction < Epsilon))
-#endif
 	{
 #if DIM_LIGHTING_CACHE_MODE == LIGHTING_CACHE_TRANSMITTANCE
 		RWLightingCacheTexture[VoxelIndex] = 1.0f;
@@ -480,11 +523,12 @@ void RenderSingleScatteringWithLiveShading(
 	DeviceZ = ConvertToDeviceZ(SceneDepth);
 
 	float Jitter = bJitter ? BlueNoiseScalar(ViewPixelCoord, View.StateFrameIndexMod8) : 0.0;
+	float2 ScreenJitter = IS_OFFLINE_RENDER ? BlueNoiseVec2(ViewPixelCoord, View.StateFrameIndex) : View.TemporalAAJitter.xy;
 
 	// Intersect ray with bounding volume
 	// TODO: LWC integration..
-	float3 WorldRayOrigin = DFHackToFloat(DFFastSubtract(GetTranslatedWorldCameraPosFromView(ViewPixelCoord + View.TemporalAAJitter.xy), PrimaryView.PreViewTranslation));
-	float3 WorldRayEnd = DFHackToFloat(SvPositionToWorld(float4(ViewPixelCoord + View.TemporalAAJitter.xy, DeviceZ, 1)));
+	float3 WorldRayOrigin = DFHackToFloat(DFFastSubtract(GetTranslatedWorldCameraPosFromView(ViewPixelCoord + ScreenJitter), PrimaryView.PreViewTranslation));
+	float3 WorldRayEnd = DFHackToFloat(SvPositionToWorld(float4(ViewPixelCoord + ScreenJitter, DeviceZ, 1)));
 	float3 WorldRayDirection = WorldRayEnd - WorldRayOrigin;
 	float WorldRayLength = length(WorldRayDirection);
 	WorldRayDirection /= WorldRayLength;
@@ -522,6 +566,10 @@ void RenderSingleScatteringWithLiveShading(
 			bApplyShadowTransmittance
 		);
 
+		RandomSequence RandSequence;
+		RandomSequence_Initialize(RandSequence, LinearIndex, View.StateFrameIndex);
+		RayMarchingContext.RandSequence = RandSequence;
+
 		const FDeferredLightData LightData = InitDeferredLightFromUniforms(LightType, VolumetricScatteringIntensity);
 		uint StepCount = CalcStepCount(RayMarchingContext);
 
diff --git a/Engine/Shaders/Private/HeterogeneousVolumes/HeterogeneousVolumesLiveShadingShadows.usf b/Engine/Shaders/Private/HeterogeneousVolumes/HeterogeneousVolumesLiveShadingShadows.usf
--- a/Engine/Shaders/Private/HeterogeneousVolumes/HeterogeneousVolumesLiveShadingShadows.usf
+++ b/Engine/Shaders/Private/HeterogeneousVolumes/HeterogeneousVolumesLiveShadingShadows.usf
@@ -11,6 +11,7 @@
 #define NO_TRANSLUCENCY_AVAILABLE
 
 #include "/Engine/Generated/Material.ush"
+#include "HeterogeneousVolumesStochasticFiltering.ush"
 #include "HeterogeneousVolumesLiveShadingUtils.ush"
 #include "HeterogeneousVolumesTracingUtils.ush"
 #include "HeterogeneousVolumesAdaptiveVolumetricShadowMapUtils.ush"
@@ -47,6 +48,7 @@ float ShadowStepSize;
 float ShadowStepFactor;
 int MaxStepCount;
 int bJitter;
+int StochasticFilteringMode;
 
 // Shadow data
 int NumShadowMatrices;
@@ -135,6 +137,11 @@ float GetMaxShadowTraceDistance()
 	return MaxShadowTraceDistance;
 }
 
+int GetStochasticFilteringMode()
+{
+	return StochasticFilteringMode;
+}
+
 struct FVolumeSampleContext
 {
 	FMaterialPixelParameters MaterialParameters;
@@ -158,6 +165,30 @@ FVolumeSampleContext CreateVolumeSampleContext(float3 LocalPosition, float3 Worl
 	return VolumeSampleContext;
 }
 
+FVolumeSampleContext CreateVolumeSampleContext(float3 LocalPosition, float3 WorldPosition, float3 FilterWidth, float MipLevel, float Rand, int StochasticFilteringMode)
+{
+	// TODO: Use FilterWidth determine MipLevel
+	//float3 WorldPosition_DDX = mul(float4(LocalPosition + float3(1.0, 0.0, 0.0), 1.0), GetLocalToWorld()).xyz;
+	//float3 WorldPosition_DDY = mul(float4(LocalPosition + float3(0.0, 1.0, 0.0), 1.0), GetLocalToWorld()).xyz;
+	//float3 WorldPosition_DDZ = mul(float4(LocalPosition + float3(0.0, 0.0, 1.0), 1.0), GetLocalToWorld()).xyz;
+
+	if (StochasticFilteringMode == STOCHASTIC_FILTERING_CONSTANT)
+	{
+		LocalPosition = StochasticFilteringConstant(LocalPosition, Rand);
+	}
+	else if (StochasticFilteringMode == STOCHASTIC_FILTERING_TRILINEAR)
+	{
+		LocalPosition = StochasticFilteringTrilinear(LocalPosition, Rand);
+	}
+	else if (StochasticFilteringMode == STOCHASTIC_FILTERING_TRICUBIC)
+	{
+		LocalPosition = StochasticFilteringTricubic(LocalPosition, Rand);
+	}
+
+	WorldPosition = mul(float4(LocalPosition, 1.0), GetLocalToWorld()).xyz;
+	return CreateVolumeSampleContext(LocalPosition, WorldPosition, MipLevel);
+}
+
 float3 SampleExtinction(inout FVolumeSampleContext Context)
 {
 	float3 Extinction = SampleExtinctionCoefficients(Context.PixelMaterialInputs);
@@ -313,8 +344,9 @@ void RenderVolumetricShadowMapForLightWithLiveShadingCS(
 	int PixelIndex = LinearIndex * MaxSampleCount;
 
 	// Generate a light ray
+	float2 ScreenJitter = IS_OFFLINE_RENDER ? BlueNoiseVec2(PixelCoord, View.StateFrameIndex) : View.TemporalAAJitter.xy;
 	//float2 ShadowUV = float2(PixelCoord + 0.5f) / float2(ShadowResolution);
-	float2 ShadowUV = float2(PixelCoord + View.TemporalAAJitter.xy) / float2(ShadowResolution);
+	float2 ShadowUV = float2(PixelCoord + ScreenJitter) / float2(ShadowResolution);
 
 	float WorldRayLength = GetMaxTraceDistance();
 #if USE_CAMERA_SCENE_DEPTH
@@ -363,7 +395,8 @@ void RenderVolumetricShadowMapForLightWithLiveShadingCS(
 	float HitSpan = HitT.y - HitT.x;
 	if (HitSpan > 0.0)
 	{
-		float Jitter = bJitter ? BlueNoiseScalar(PixelCoord, View.StateFrameIndexMod8) : 0.0;
+		//float Jitter = bJitter ? BlueNoiseScalar(PixelCoord, View.StateFrameIndexMod8) : 0.0;
+		float Jitter = IS_OFFLINE_RENDER ? BlueNoiseScalar(PixelCoord, View.StateFrameIndexMod8) : 0.0;
 		int bApplyEmissionAndTransmittance = 0;
 		int bApplyDirectLighting = 0;
 		int bApplyShadowTransmittance = 0;
diff --git a/Engine/Shaders/Private/HeterogeneousVolumes/HeterogeneousVolumesRayMarchingUtils.ush b/Engine/Shaders/Private/HeterogeneousVolumes/HeterogeneousVolumesRayMarchingUtils.ush
--- a/Engine/Shaders/Private/HeterogeneousVolumes/HeterogeneousVolumesRayMarchingUtils.ush
+++ b/Engine/Shaders/Private/HeterogeneousVolumes/HeterogeneousVolumesRayMarchingUtils.ush
@@ -38,6 +38,10 @@
 #define APPLY_FOG_INSCATTERING APPLY_FOG_INSCATTERING_OFF
 #endif // APPLY_FOG_INSCATTERING
 
+#ifndef IS_OFFLINE_RENDER
+#define IS_OFFLINE_RENDER 0
+#endif // IS_OFFLINE_RENDER
+
 #if DIM_USE_ADAPTIVE_VOLUMETRIC_SHADOW_MAP
 #include "HeterogeneousVolumesAdaptiveVolumetricShadowMapSampling.ush"
 #include "HeterogeneousVolumesAdaptiveVolumetricShadowMapSampler.ush"
@@ -186,7 +190,7 @@ float3 ComputeTransmittance(
 	#if DIM_USE_ADAPTIVE_VOLUMETRIC_SHADOW_MAP
 		if (!AVSM.bIsEmpty)
 		{
-			float3 TranslatedWorldPosition = DFHackToFloat(DFFastAdd(WorldRayOrigin, PrimaryView.PreViewTranslation));
+			float3 TranslatedWorldPosition = DFFastAddDemote(WorldRayOrigin, PrimaryView.PreViewTranslation);
 			float3 LightTranslatedWorldPosition = TranslatedWorldPosition + ToLight;
 			Transmittance = AVSM_SampleTransmittance(TranslatedWorldPosition, LightTranslatedWorldPosition);
 		}
@@ -349,7 +353,14 @@ void RayMarchSingleScattering(
 		float WorldHitT = LocalHitT * RayMarchingContext.LocalToWorldScale;
 		float3 LocalPosition = RayMarchingContext.LocalRayOrigin + RayMarchingContext.LocalRayDirection * LocalHitT;
 		float3 WorldPosition = RayMarchingContext.WorldRayOrigin + RayMarchingContext.WorldRayDirection * WorldHitT;
+#if IS_OFFLINE_RENDER
+		float3 FilterWidth = 1.0;
+		float Rand = RandomSequence_GenerateSample1D(RayMarchingContext.RandSequence);
+		int StochasticFilteringMode = GetStochasticFilteringMode();
+		FVolumeSampleContext SampleContext = CreateVolumeSampleContext(LocalPosition, WorldPosition, FilterWidth, RayMarchingContext.MipLevel, Rand, StochasticFilteringMode);
+#else
 		FVolumeSampleContext SampleContext = CreateVolumeSampleContext(LocalPosition, WorldPosition, RayMarchingContext.MipLevel);
+#endif // IS_OFFLINE_RENDER
 
 		if (RayMarchingContext.bApplyEmissionAndTransmittance)
 		{
diff --git a/Engine/Shaders/Private/HeterogeneousVolumes/HeterogeneousVolumesSparseVoxelUniformBufferUtils.ush b/Engine/Shaders/Private/HeterogeneousVolumes/HeterogeneousVolumesSparseVoxelUniformBufferUtils.ush
--- a/Engine/Shaders/Private/HeterogeneousVolumes/HeterogeneousVolumesSparseVoxelUniformBufferUtils.ush
+++ b/Engine/Shaders/Private/HeterogeneousVolumes/HeterogeneousVolumesSparseVoxelUniformBufferUtils.ush
@@ -1,5 +1,7 @@
 // Copyright Epic Games, Inc. All Rights Reserved.
 
+#include "HeterogeneousVolumesStochasticFiltering.ush"
+
 float4x4 GetLocalToWorld()
 {
 	return SparseVoxelUniformBuffer.LocalToWorld;
@@ -50,6 +52,11 @@ float GetShadowStepFactor()
 	return SparseVoxelUniformBuffer.ShadowStepFactor;
 }
 
+int GetStochasticFilteringMode()
+{
+	return STOCHASTIC_FILTERING_TRILINEAR;
+}
+
 int ShouldApplyHeightFog()
 {
 	return SparseVoxelUniformBuffer.bApplyHeightFog;
@@ -93,6 +100,30 @@ FVolumeSampleContext CreateVolumeSampleContext(float3 LocalPosition, float3 Worl
 	return VolumeSampleContext;
 }
 
+FVolumeSampleContext CreateVolumeSampleContext(float3 LocalPosition, float3 WorldPosition, float3 FilterWidth, float MipLevel, float Rand, int SamplingMethod)
+{
+	// TODO: Use FilterWidth determine MipLevel
+	//float3 WorldPosition_DDX = mul(float4(LocalPosition + float3(1.0, 0.0, 0.0), 1.0), GetLocalToWorld()).xyz;
+	//float3 WorldPosition_DDY = mul(float4(LocalPosition + float3(0.0, 1.0, 0.0), 1.0), GetLocalToWorld()).xyz;
+	//float3 WorldPosition_DDZ = mul(float4(LocalPosition + float3(0.0, 0.0, 1.0), 1.0), GetLocalToWorld()).xyz;
+
+	if (SamplingMethod == STOCHASTIC_FILTERING_CONSTANT)
+	{
+		LocalPosition = StochasticFilteringConstant(LocalPosition, Rand);
+	}
+	else if (SamplingMethod == STOCHASTIC_FILTERING_TRILINEAR)
+	{
+		LocalPosition = StochasticFilteringTrilinear(LocalPosition, Rand);
+	}
+	else if (SamplingMethod == STOCHASTIC_FILTERING_TRICUBIC)
+	{
+		LocalPosition = StochasticFilteringTricubic(LocalPosition, Rand);
+	}
+
+	WorldPosition = mul(float4(LocalPosition, 1.0), GetLocalToWorld()).xyz;
+	return CreateVolumeSampleContext(LocalPosition, WorldPosition, MipLevel);
+}
+
 float3 SampleExtinction(inout FVolumeSampleContext Context)
 {
 	float3 Extinction = Texture3DSampleLevel(
diff --git a/Engine/Shaders/Private/HeterogeneousVolumes/HeterogeneousVolumesStochasticFiltering.ush b/Engine/Shaders/Private/HeterogeneousVolumes/HeterogeneousVolumesStochasticFiltering.ush
new file mode 100644
--- /dev/null
+++ b/Engine/Shaders/Private/HeterogeneousVolumes/HeterogeneousVolumesStochasticFiltering.ush
@@ -0,0 +1,71 @@
+// Copyright Epic Games, Inc. All Rights Reserved.
+
+int3 WeightedReservoir(int i, int j, float Prob, int3 Index, inout float Rand)
+{
+	if (Rand < Prob)
+	{
+		Index[i] = j;
+		Rand /= Prob;
+	}
+	else
+	{
+		Rand = (Rand - Prob) / (1.0 - Prob);
+	}
+
+	return Index;
+}
+
+#define STOCHASTIC_FILTERING_DISABLED 0
+#define STOCHASTIC_FILTERING_CONSTANT 1
+#define STOCHASTIC_FILTERING_TRILINEAR 2
+#define STOCHASTIC_FILTERING_TRICUBIC 3
+
+int3 StochasticFilteringConstant(float3 Position, inout float Rand)
+{
+	return Position;
+}
+
+int3 StochasticFilteringTrilinear(float3 Position, inout float Rand)
+{
+	int3 PositionInt = floor(Position);
+	float3 Delta = Position - PositionInt;
+
+	int3 Index = 0;
+	for (int i = 0; i < 3; ++i)
+	{
+		float Prob = Delta[i];
+		Index = WeightedReservoir(i, 1, Prob, Index, Rand);
+	}
+
+	return PositionInt + Index;
+}
+
+int3 StochasticFilteringTricubic(float3 Position, inout float Rand)
+{
+	int3 PositionInt = floor(Position);
+	float3 Delta = Position - PositionInt;
+
+	int3 Index = 0;
+	for (int i = 0; i < 3; ++i)
+	{
+		float t = Delta[i];
+		float t2 = t * t;
+
+		float4 Weights = float4(
+			(1.0 / 6.0) * (-t*t2 + 3*t2 - 3*t + 1),
+			(1.0 / 6.0) * (3*t*t2 - 6*t2 + 4),
+			(1.0 / 6.0) * (-3*t*t2 + 3*t2 + 3*t + 1),
+			(1.0 / 6.0) * t*t2
+		);
+
+		float Sum = Weights[0];
+		for (int j = 0; j < 3; ++j)
+		{
+			Sum += Weights[j + 1];
+			float Prob = Weights[j + 1] / Sum;
+			Index = WeightedReservoir(i, j + 1, Prob, Index, Rand);
+		}
+	}
+
+	return PositionInt - 1 + Index;
+}
diff --git a/Engine/Shaders/Private/HeterogeneousVolumes/HeterogeneousVolumesVoxelGridShadows.usf b/Engine/Shaders/Private/HeterogeneousVolumes/HeterogeneousVolumesVoxelGridShadows.usf
--- a/Engine/Shaders/Private/HeterogeneousVolumes/HeterogeneousVolumesVoxelGridShadows.usf
+++ b/Engine/Shaders/Private/HeterogeneousVolumes/HeterogeneousVolumesVoxelGridShadows.usf
@@ -5,6 +5,7 @@
 #include "../Common.ush"
 #include "../DeferredShadingCommon.ush"
 #include "../ComputeShaderUtils.ush"
+#include "HeterogeneousVolumesStochasticFiltering.ush"
 #include "HeterogeneousVolumesTracingUtils.ush"
 #include "HeterogeneousVolumesVoxelGridTraversal.ush"
 #include "HeterogeneousVolumesVoxelGridRenderingUtils.ush"
@@ -21,6 +22,7 @@ int3 GroupCount;
 int ShadowDebugTweak;
 
 int MaxStepCount;
+int StochasticFilteringMode;
 
 int NumShadowMatrices;
 float4x4 TranslatedWorldToShadow[6];
@@ -132,6 +134,11 @@ float GetShadowStepFactor()
 	return 1.0;
 }
 
+int GetStochasticFilteringMode()
+{
+	return StochasticFilteringMode;
+}
+
 uint3 GetVolumeResolution()
 {
 	return OrthoGridUniformBuffer.TopLevelGridResolution;
@@ -142,6 +149,32 @@ uint GetNumVoxels()
 	return GetVolumeResolution().x * GetVolumeResolution().y * GetVolumeResolution().z;
 }
 
+FVolumeSampleContext CreateVolumeSampleContext(float3 LocalPosition, float3 WorldPosition, float3 FilterWidth, float MipLevel, float Rand, int StochasticFilteringMode)
+{
+	float3 WorldPosition_DDX = mul(float4(LocalPosition + float3(1.0, 0.0, 0.0), 1.0), GetLocalToWorld()).xyz;
+	float3 WorldPosition_DDY = mul(float4(LocalPosition + float3(0.0, 1.0, 0.0), 1.0), GetLocalToWorld()).xyz;
+	float3 WorldPosition_DDZ = mul(float4(LocalPosition + float3(0.0, 0.0, 1.0), 1.0), GetLocalToWorld()).xyz;
+
+	if (StochasticFilteringMode == STOCHASTIC_FILTERING_CONSTANT)
+	{
+		LocalPosition = StochasticFilteringConstant(LocalPosition, Rand);
+	}
+	else if (StochasticFilteringMode == STOCHASTIC_FILTERING_TRILINEAR)
+	{
+		LocalPosition = StochasticFilteringTrilinear(LocalPosition, Rand);
+	}
+	else if (StochasticFilteringMode == STOCHASTIC_FILTERING_TRICUBIC)
+	{
+		LocalPosition = StochasticFilteringTricubic(LocalPosition, Rand);
+	}
+
+	float3 WorldPosition2 = mul(float4(LocalPosition, 1.0), GetLocalToWorld()).xyz;
+
+	FVolumeSampleContext VolumeSampleContext;
+	VolumeSampleContext.WorldPosition = WorldPosition2;
+	return VolumeSampleContext;
+}
+
 FVolumeSampleContext CreateVolumeSampleContext(float3 LocalPosition, float3 WorldPosition, float MipLevel)
 {
 	FVolumeSampleContext VolumeSampleContext;
diff --git a/Engine/Source/Runtime/Renderer/Private/HeterogeneousVolumes/HeterogeneousVolumes.cpp b/Engine/Source/Runtime/Renderer/Private/HeterogeneousVolumes/HeterogeneousVolumes.cpp
--- a/Engine/Source/Runtime/Renderer/Private/HeterogeneousVolumes/HeterogeneousVolumes.cpp
+++ b/Engine/Source/Runtime/Renderer/Private/HeterogeneousVolumes/HeterogeneousVolumes.cpp
@@ -79,6 +79,17 @@ static TAutoConsoleVariable<int32> CVarHeterogeneousVolumesIndirectLighting(
 	ECVF_RenderThreadSafe
 );
 
+static TAutoConsoleVariable<int32> CVarHeterogeneousVolumesStochasticFiltering(
+	TEXT("r.HeterogeneousVolumes.StochasticFiltering"),
+	3,
+	TEXT("Configures the stochastic filtering kernel (Default = 3)\n")
+	TEXT("0: Disabled\n")
+	TEXT("1: Constant\n")
+	TEXT("2: Linear\n")
+	TEXT("3: Cubic"),
+	ECVF_RenderThreadSafe
+);
+
 static TAutoConsoleVariable<int32> CVarHeterogeneousVolumesJitter(
 	TEXT("r.HeterogeneousVolumes.Jitter"),
 	1,
@@ -479,6 +490,11 @@ namespace HeterogeneousVolumes
 		return static_cast<EShadowMode>(CVarHeterogeneousVolumesShadowMode.GetValueOnRenderThread());
 	}
 
+	EStochasticFilteringMode GetStochasticFilteringMode()
+	{
+		return static_cast<EStochasticFilteringMode>(CVarHeterogeneousVolumesStochasticFiltering.GetValueOnRenderThread());
+	}
+
 	bool UseSparseVoxelPipeline()
 	{
 		return CVarHeterogeneousVolumesSparseVoxel.GetValueOnAnyThread() != 0;
diff --git a/Engine/Source/Runtime/Renderer/Private/HeterogeneousVolumes/HeterogeneousVolumes.h b/Engine/Source/Runtime/Renderer/Private/HeterogeneousVolumes/HeterogeneousVolumes.h
--- a/Engine/Source/Runtime/Renderer/Private/HeterogeneousVolumes/HeterogeneousVolumes.h
+++ b/Engine/Source/Runtime/Renderer/Private/HeterogeneousVolumes/HeterogeneousVolumes.h
@@ -100,6 +100,14 @@ namespace HeterogeneousVolumes
 	uint32 GetSparseVoxelMipBias();
 	int32 GetBottomLevelGridResolution();
 	int32 GetIndirectionGridResolution();
+	enum class EStochasticFilteringMode
+	{
+		Disabled,
+		Constant,
+		Linear,
+		Cubic
+	};
+	EStochasticFilteringMode GetStochasticFilteringMode();
 	
 	bool ShouldJitter();
 	bool ShouldRefineSparseVoxels();
diff --git a/Engine/Source/Runtime/Renderer/Private/HeterogeneousVolumes/HeterogeneousVolumesAmbientOcclusionPipeline.cpp b/Engine/Source/Runtime/Renderer/Private/HeterogeneousVolumes/HeterogeneousVolumesAmbientOcclusionPipeline.cpp
--- a/Engine/Source/Runtime/Renderer/Private/HeterogeneousVolumes/HeterogeneousVolumesAmbientOcclusionPipeline.cpp
+++ b/Engine/Source/Runtime/Renderer/Private/HeterogeneousVolumes/HeterogeneousVolumesAmbientOcclusionPipeline.cpp
@@ -440,7 +440,8 @@ class FRenderAmbientOcclusionWithLiveShadingCS : public FMeshMaterialShader
 	DECLARE_SHADER_TYPE(FRenderAmbientOcclusionWithLiveShadingCS, MeshMaterial);
 
 	class FUseExistenceMask : SHADER_PERMUTATION_INT("USE_EXISTENCE_MASK", 2);
-	using FPermutationDomain = TShaderPermutationDomain<FUseExistenceMask>;
+	class FIsOfflineRender : SHADER_PERMUTATION_INT("IS_OFFLINE_RENDER", 2);
+	using FPermutationDomain = TShaderPermutationDomain<FUseExistenceMask, FIsOfflineRender>;
 
 	BEGIN_SHADER_PARAMETER_STRUCT(FParameters, )
 		// Scene data
@@ -466,6 +467,7 @@ class FRenderAmbientOcclusionWithLiveShadingCS : public FMeshMaterialShader
 		SHADER_PARAMETER(float, StepFactor)
 		SHADER_PARAMETER(int, MaxStepCount)
 		SHADER_PARAMETER(int, bJitter)
+		SHADER_PARAMETER(int, StochasticFilteringMode)
 
 		// Volume data
 		SHADER_PARAMETER(FIntVector, VoxelResolution)
@@ -595,6 +597,7 @@ void RenderAmbientOcclusionWithLiveShadingAsFixedPoint(
 		PassParameters->StepFactor = HeterogeneousVolumeInterface->GetStepFactor();
 		PassParameters->MaxStepCount = HeterogeneousVolumes::GetAmbientOcclusionMaxStepCount();
 		PassParameters->bJitter = 0;
+		PassParameters->StochasticFilteringMode = static_cast<int32>(HeterogeneousVolumes::GetStochasticFilteringMode());
 
 		FMatrix LocalToInstance = InstanceToLocal.Inverse();
 		FBoxSphereBounds InstanceBoxSphereBounds = LocalBoxSphereBounds.TransformBy(LocalToInstance);
@@ -626,6 +629,7 @@ void RenderAmbientOcclusionWithLiveShadingAsFixedPoint(
 
 	FRenderAmbientOcclusionWithLiveShadingCS::FPermutationDomain PermutationVector;
 	PermutationVector.Set<FRenderAmbientOcclusionWithLiveShadingCS::FUseExistenceMask>(HeterogeneousVolumes::UseExistenceMask());
+	PermutationVector.Set<FRenderAmbientOcclusionWithLiveShadingCS::FIsOfflineRender>(View.bIsOfflineRender);
 	TShaderRef<FRenderAmbientOcclusionWithLiveShadingCS> ComputeShader = Material.GetShader<FRenderAmbientOcclusionWithLiveShadingCS>(&FLocalVertexFactory::StaticType, PermutationVector, false);
 	GraphBuilder.AddPass(
 		RDG_EVENT_NAME("AmbientOcclusion"),
diff --git a/Engine/Source/Runtime/Renderer/Private/HeterogeneousVolumes/HeterogeneousVolumesLiveShadingPipeline.cpp b/Engine/Source/Runtime/Renderer/Private/HeterogeneousVolumes/HeterogeneousVolumesLiveShadingPipeline.cpp
--- a/Engine/Source/Runtime/Renderer/Private/HeterogeneousVolumes/HeterogeneousVolumesLiveShadingPipeline.cpp
+++ b/Engine/Source/Runtime/Renderer/Private/HeterogeneousVolumes/HeterogeneousVolumesLiveShadingPipeline.cpp
@@ -202,7 +202,8 @@ class FRenderLightingCacheWithLiveShadingCS : public FMeshMaterialShader
 	class FUseAdaptiveVolumetricShadowMap : SHADER_PERMUTATION_BOOL("DIM_USE_ADAPTIVE_VOLUMETRIC_SHADOW_MAP");
 	class FAVSMSampleMode : SHADER_PERMUTATION_INT("AVSM_SAMPLE_MODE", 2);
 	class FUseExistenceMask : SHADER_PERMUTATION_INT("USE_EXISTENCE_MASK", 2);
-	using FPermutationDomain = TShaderPermutationDomain<FLightingCacheMode, FUseAdaptiveVolumetricShadowMap, FAVSMSampleMode, FUseExistenceMask>;
+	class FIsOfflineRender : SHADER_PERMUTATION_INT("IS_OFFLINE_RENDER", 2);
+	using FPermutationDomain = TShaderPermutationDomain<FLightingCacheMode, FUseAdaptiveVolumetricShadowMap, FAVSMSampleMode, FUseExistenceMask, FIsOfflineRender>;
 
 	BEGIN_SHADER_PARAMETER_STRUCT(FParameters, )
 		// Scene data
@@ -241,6 +242,7 @@ class FRenderLightingCacheWithLiveShadingCS : public FMeshMaterialShader
 		SHADER_PARAMETER(float, StepFactor)
 		SHADER_PARAMETER(int, MaxStepCount)
 		SHADER_PARAMETER(int, bJitter)
+		SHADER_PARAMETER(int, StochasticFilteringMode)
 
 		// Volume data
 		SHADER_PARAMETER(FIntVector, VoxelResolution)
@@ -347,6 +349,7 @@ class FRenderSingleScatteringWithLiveShadingCS : public FMeshMaterialShader
 	class FSupportOverlappingVolumes : SHADER_PERMUTATION_BOOL("SUPPORT_OVERLAPPING_VOLUMES");
 	//class FAdaptiveMarch : SHADER_PERMUTATION_INT("ADAPTIVE_MARCH", 2);
 	//class FBilinearInterpolation : SHADER_PERMUTATION_INT("AVSM_BILINEAR_INTERPOLATION", 2);
+	class FIsOfflineRender : SHADER_PERMUTATION_INT("IS_OFFLINE_RENDER", 2);
 	class FApplyFogInscattering : SHADER_PERMUTATION_INT("APPLY_FOG_INSCATTERING", 3);
 	using FPermutationDomain = TShaderPermutationDomain<
 		FUseTransmittanceVolume, 
@@ -356,7 +359,8 @@ class FRenderSingleScatteringWithLiveShadingCS : public FMeshMaterialShader
 		FUseAdaptiveVolumetricShadowMap, 
 		FApplyFogInscattering, 
 		FAVSMSampleMode, 
-		FSupportOverlappingVolumes
+		FSupportOverlappingVolumes,
+		FIsOfflineRender
 		//FBilinearInterpolation,
 		//FAdaptiveMarch
 	>;
@@ -416,6 +420,7 @@ class FRenderSingleScatteringWithLiveShadingCS : public FMeshMaterialShader
 		SHADER_PARAMETER(float, StepFactor)
 		SHADER_PARAMETER(int, MaxStepCount)
 		SHADER_PARAMETER(int, bJitter)
+		SHADER_PARAMETER(int, StochasticFilteringMode)
 
 		// Dispatch data
 		SHADER_PARAMETER(FIntVector, GroupCount)
@@ -738,6 +743,7 @@ static void RenderLightingCacheWithLiveShading(
 		PassParameters->ShadowStepFactor = HeterogeneousVolumeInterface->GetShadowStepFactor() * LODFactor;
 		PassParameters->MaxStepCount = HeterogeneousVolumes::GetMaxStepCount();
 		PassParameters->bJitter = HeterogeneousVolumes::ShouldJitter();
+		PassParameters->StochasticFilteringMode = static_cast<int32>(HeterogeneousVolumes::GetStochasticFilteringMode());
 
 		// Shadow data
 		PassParameters->ForwardLightData = View.ForwardLightingResources.ForwardLightUniformBuffer;
@@ -846,6 +852,7 @@ static void RenderLightingCacheWithLiveShading(
 	PermutationVector.Set<FRenderLightingCacheWithLiveShadingCS::FUseAdaptiveVolumetricShadowMap>(bUseAVSM);
 	PermutationVector.Set<FRenderLightingCacheWithLiveShadingCS::FAVSMSampleMode>(AVSMSampleMode);
 	PermutationVector.Set<FRenderLightingCacheWithLiveShadingCS::FUseExistenceMask>(HeterogeneousVolumes::UseExistenceMask());
+	PermutationVector.Set<FRenderLightingCacheWithLiveShadingCS::FIsOfflineRender>(View.bIsOfflineRender);
 	TShaderRef<FRenderLightingCacheWithLiveShadingCS> ComputeShader = Material.GetShader<FRenderLightingCacheWithLiveShadingCS>(&FLocalVertexFactory::StaticType, PermutationVector, false);
 	if (!ComputeShader.IsNull())
 	{
@@ -1136,6 +1143,7 @@ void RenderSingleScatteringWithLiveShading(
 		PassParameters->StepFactor = HeterogeneousVolumeInterface->GetStepFactor() * LODFactor;
 		PassParameters->MaxStepCount = HeterogeneousVolumes::GetMaxStepCount();
 		PassParameters->bJitter = HeterogeneousVolumes::ShouldJitter();
+		PassParameters->StochasticFilteringMode = static_cast<int32>(HeterogeneousVolumes::GetStochasticFilteringMode());
 
 		// Shadow data
 		PassParameters->ForwardLightData = View.ForwardLightingResources.ForwardLightUniformBuffer;
@@ -1255,6 +1263,7 @@ void RenderSingleScatteringWithLiveShading(
 	PermutationVector.template Set<typename FRenderSingleScatteringWithLiveShadingDispatchTypeCS::FSupportOverlappingVolumes>(HeterogeneousVolumes::SupportsOverlappingVolumes());
 	//PermutationVector.template Set<typename FRenderSingleScatteringWithLiveShadingDispatchTypeCS::FBilinearInterpolation>(static_cast<int32>(HeterogeneousVolumes::UseBilinearInterpolation()));
 	//PermutationVector.template Set<typename FRenderSingleScatteringWithLiveShadingDispatchTypeCS::FAdaptiveMarch>(static_cast<int32>(HeterogeneousVolumes::ShouldAdaptiveMarch()));
+	PermutationVector.template Set<typename FRenderSingleScatteringWithLiveShadingDispatchTypeCS::FIsOfflineRender>(View.bIsOfflineRender);
 	PermutationVector.template Set<typename FRenderSingleScatteringWithLiveShadingDispatchTypeCS::FApplyFogInscattering>(static_cast<int32>(HeterogeneousVolumes::GetApplyFogInscattering()));
 	PermutationVector = FRenderSingleScatteringWithLiveShadingDispatchTypeCS::RemapPermutation(PermutationVector);
 	TShaderRef<FRenderSingleScatteringWithLiveShadingDispatchTypeCS> ComputeShader = Material.GetShader<FRenderSingleScatteringWithLiveShadingDispatchTypeCS>(&FLocalVertexFactory::StaticType, PermutationVector, false);
diff --git a/Engine/Source/Runtime/Renderer/Private/HeterogeneousVolumes/HeterogeneousVolumesVoxelGridPipeline.cpp b/Engine/Source/Runtime/Renderer/Private/HeterogeneousVolumes/HeterogeneousVolumesVoxelGridPipeline.cpp
--- a/Engine/Source/Runtime/Renderer/Private/HeterogeneousVolumes/HeterogeneousVolumesVoxelGridPipeline.cpp
+++ b/Engine/Source/Runtime/Renderer/Private/HeterogeneousVolumes/HeterogeneousVolumesVoxelGridPipeline.cpp
@@ -2946,7 +2946,8 @@ class FRenderVolumetricShadowMapForLightWithVoxelGridCS : public FGlobalShader
 	SHADER_USE_PARAMETER_STRUCT(FRenderVolumetricShadowMapForLightWithVoxelGridCS, FGlobalShader);
 
 	class FUseAVSMCompression : SHADER_PERMUTATION_BOOL("USE_AVSM_COMPRESSION");
-	using FPermutationDomain = TShaderPermutationDomain<FUseAVSMCompression>;
+	class FIsOfflineRender : SHADER_PERMUTATION_INT("IS_OFFLINE_RENDER", 2);
+	using FPermutationDomain = TShaderPermutationDomain<FUseAVSMCompression, FIsOfflineRender>;
 
 	BEGIN_SHADER_PARAMETER_STRUCT(FParameters, )
 		// Scene data
@@ -2962,6 +2963,7 @@ class FRenderVolumetricShadowMapForLightWithVoxelGridCS : public FGlobalShader
 		SHADER_PARAMETER(FVector3f, TranslatedWorldOrigin)
 		SHADER_PARAMETER(FIntPoint, ShadowResolution)
 		SHADER_PARAMETER(int, MaxSampleCount)
+		SHADER_PARAMETER(int, StochasticFilteringMode)
 		SHADER_PARAMETER(float, AbsoluteErrorThreshold)
 		SHADER_PARAMETER(float, RelativeErrorThreshold)
 
@@ -3496,6 +3498,7 @@ void RenderVolumetricShadowMapForLightWithVoxelGrid(
 		PassParameters->MaxTraceDistance = HeterogeneousVolumes::GetMaxTraceDistance();
 		PassParameters->MaxStepCount = HeterogeneousVolumes::GetMaxStepCount();
 		PassParameters->bJitter = HeterogeneousVolumes::ShouldJitter();
+		PassParameters->StochasticFilteringMode = static_cast<int32>(HeterogeneousVolumes::GetStochasticFilteringMode());
 
 		PassParameters->NumShadowMatrices = ProjectedShadowInfo->OnePassShadowViewProjectionMatrices.Num();
 		if (PassParameters->NumShadowMatrices > 0)
@@ -3619,6 +3622,7 @@ void RenderVolumetricShadowMapForLightWithVoxelGrid(
 
 	FRenderVolumetricShadowMapForLightWithVoxelGridCS::FPermutationDomain PermutationVector;
 	PermutationVector.Set<FRenderVolumetricShadowMapForLightWithVoxelGridCS::FUseAVSMCompression>(HeterogeneousVolumes::UseAVSMCompression());
+	PermutationVector.Set<FRenderVolumetricShadowMapForLightWithVoxelGridCS::FIsOfflineRender>(View.bIsOfflineRender);
 	TShaderRef<FRenderVolumetricShadowMapForLightWithVoxelGridCS> ComputeShader = View.ShaderMap->GetShader<FRenderVolumetricShadowMapForLightWithVoxelGridCS>(PermutationVector);
 	FComputeShaderUtils::AddPass(
 		GraphBuilder,
