Fix error compiling LightGridInjectionCS due to FScreenRect not being defined.

diff --git a/Engine/Shaders/Private/LightGridInjection.usf b/Engine/Shaders/Private/LightGridInjection.usf
--- a/Engine/Shaders/Private/LightGridInjection.usf
+++ b/Engine/Shaders/Private/LightGridInjection.usf
@@ -78,6 +78,7 @@ float ComputeCellNearViewDepthFromZSlice(uint ZSlice)
 	return SliceDepth;
 }
 
+#if USE_HZB_CULL
 FScreenRect ComputeCellCullRect(uint3 GridCoordinate, float MinTileZ, float MaxTileZ)
 {
 	// Compute extent of tiles in clip-space. Note that the last tile may extend a bit outside of view if view size is not evenly divisible tile size.
@@ -103,6 +104,7 @@ FScreenRect ComputeCellCullRect(uint3 GridCoordinate, float MinTileZ, float MaxT
 
 	return GetScreenRect(int4(0, 0, HZBViewSize), CullRectMin, CullRectMax, 4);
 }
+#endif
 
 struct FCellBounds
 {
