Add fog support to translucency holdout, including fog affects translucency and fog as holdout.

diff --git a/Engine/Shaders/Private/BasePassPixelShader.usf b/Engine/Shaders/Private/BasePassPixelShader.usf
--- a/Engine/Shaders/Private/BasePassPixelShader.usf
+++ b/Engine/Shaders/Private/BasePassPixelShader.usf
@@ -2211,7 +2211,13 @@ void FPixelShaderInOut_MainPS(
 		#endif
 	#endif 
 	
-	const bool bIsHoldout = GetPrimitiveData(MaterialParameters).Flags & PRIMITIVE_SCENE_DATA_FLAG_HOLDOUT;
+	
+#if (MATERIALBLENDING_ALPHAHOLDOUT || SUBSTRATE_BLENDING_ALPHAHOLDOUT)
+	const bool bIsHoldout = true;
+#else
+	const bool bIsHoldout = (GetPrimitiveData(MaterialParameters).Flags & PRIMITIVE_SCENE_DATA_FLAG_HOLDOUT);
+#endif
+	
 	// Only include the holdout logic if the Renderer is forwarding its alpha channel to post-processing
 	if (bIsHoldout)
 	{
@@ -2238,9 +2244,12 @@ void FPixelShaderInOut_MainPS(
 	#endif
 			
 #else
-			// TODO: Not all translucent cases get the right alpha (compared to the path traced reference)
-
-			Out.MRT[0] = half4(0.0, 0.0, 0.0, Opacity);
+	#if MATERIAL_SHADINGMODEL_THIN_TRANSLUCENT
+		Out.MRT[0] = half4(0.0, 0.0, 0.0, DualBlendSurfaceCoverage);
+	#else
+		// TODO: Not all translucent cases get the right alpha (compared to the path traced reference)
+		Out.MRT[0] = half4(0.0, 0.0, 0.0, Opacity);
+	#endif
 #endif
 
 #if USES_GBUFFER
@@ -2258,21 +2267,26 @@ void FPixelShaderInOut_MainPS(
 	{
 		// Default behavior is to preserve the destination alpha.
 		half BackgroundVisibilityAdd = 0.0f;
-		half PathThroughputMul = 1.0f;
-
+		half PathThroughputMul = 1.0f;		
+		half FogPathThroughput = Fogging.a;
+		half PrimitiveCoverageForFog = 0.0f;
+	
 #if SUBSTRATE_TRANSLUCENT_MATERIAL
 	#if SUBSTRATE_BLENDING_TRANSLUCENT_GREYTRANSMITTANCE
 		BackgroundVisibilityAdd = bIsHoldout ? DualBlendSurfaceCoverage : 0;
 		PathThroughputMul = 1 - DualBlendSurfaceCoverage;
+		PrimitiveCoverageForFog = DualBlendSurfaceCoverage;
 	#elif SUBSTRATE_BLENDING_TRANSLUCENT_COLOREDTRANSMITTANCE
 		BackgroundVisibilityAdd = bIsHoldout ? DualBlendSurfaceCoverage : 0;
 		PathThroughputMul = 1 - DualBlendSurfaceCoverage;
+		PrimitiveCoverageForFog = DualBlendSurfaceCoverage;
 	#elif SUBSTRATE_BLENDING_COLOREDTRANSMITTANCEONLY
 		// TODO
 	#elif MATERIALBLENDING_ALPHAHOLDOUT
 		//alpha holdout should always contribute to the background visibility.
 		BackgroundVisibilityAdd = 1 - AdjustedAlpha;
 		PathThroughputMul = AdjustedAlpha;
+		PrimitiveCoverageForFog = 1 - AdjustedAlpha;
 	#endif 
 #else
 
@@ -2280,14 +2294,17 @@ void FPixelShaderInOut_MainPS(
 		//thin translucent shading model uses surface coverage as opacity.
 		BackgroundVisibilityAdd = bIsHoldout ? DualBlendSurfaceCoverage : 0;
 		PathThroughputMul = 1 - DualBlendSurfaceCoverage;
+		PrimitiveCoverageForFog = DualBlendSurfaceCoverage;
 	#elif MATERIALBLENDING_ALPHAHOLDOUT
 		//alpha holdout should always contribute to the background visibility.
 		BackgroundVisibilityAdd = Opacity;
 		PathThroughputMul = (1 - Opacity);
+		PrimitiveCoverageForFog = Opacity;
 	#elif MATERIALBLENDING_TRANSLUCENT
 		// Use Opacity as it is.
 		BackgroundVisibilityAdd = bIsHoldout ? Opacity : 0;
 		PathThroughputMul = 1 - Opacity;
+		PrimitiveCoverageForFog = Opacity;
 	#endif
 #endif
 
@@ -2295,10 +2312,12 @@ void FPixelShaderInOut_MainPS(
 		// Not affecting alpha
 		BackgroundVisibilityAdd = 0.0f;
 		PathThroughputMul = 1.0f;
+		PrimitiveCoverageForFog = 0.0f;
 #elif MATERIALBLENDING_MODULATE
-
+		FogPathThroughput = Fogging.a * Fogging.a; // Match rendering pass
 		BackgroundVisibilityAdd = bIsHoldout ? Opacity : 0;
 		PathThroughputMul = (1 - Opacity);
+		PrimitiveCoverageForFog = Opacity;
 		//---------------------------------------------------------------------
 		// Another solution: use fogged color to calcualte throughput and background visibility
 		// C = C * L + C * (1-L).
@@ -2308,8 +2327,13 @@ void FPixelShaderInOut_MainPS(
 		// Otherwise, we can only have one modulate in a scene that can be holdout as well.
 #endif
 
+		// Fog holdout parameters
+		const bool bIsExponentialFogHoldout = IsExponentialFogHoldout(View.EnvironmentComponentsFlags); 
+		half FogBackgroundVisibilityAdd = bIsExponentialFogHoldout ? (PrimitiveCoverageForFog * (1 - FogPathThroughput)) : 0;
+	
 		// Render far to close logic
-		Out.MRT[0].xyz = BackgroundVisibilityAdd;
+		// Holdout = Holdout * PrimitivePathThroughput + (FogHoldout * PrimCoverage + HoldoutPrimCoverage * FogPathThroughput);
+		Out.MRT[0].xyz = FogBackgroundVisibilityAdd + BackgroundVisibilityAdd * FogPathThroughput;
 		Out.MRT[0].w = PathThroughputMul;
 		
 	}
