Fix thin translucency does not blend in the scenecolor.
TODO: Need to consider both refraction and translucency.

#tests local tests

diff --git a/Engine/Shaders/Private/RayTracing/RayTracingPrimaryRays.usf b/Engine/Shaders/Private/RayTracing/RayTracingPrimaryRays.usf
--- a/Engine/Shaders/Private/RayTracing/RayTracingPrimaryRays.usf
+++ b/Engine/Shaders/Private/RayTracing/RayTracingPrimaryRays.usf
@@ -382,7 +382,7 @@ RAY_TRACING_ENTRY_RAYGEN(RayTracingPrimaryRaysRGS)
 			const float3 Transmittance = Payload.CustomData.xyz;
 			const float PathLength = rcp(saturate(abs(NoV) + 1e-5));
 			PathThroughput *= pow(Transmittance, PathLength) * Pow2(1 - F);
-			bHasScattered = true; // treat the thin translucency as refraction for consistency with solid glass, even though the ray is not bent
+			bHasScattered = false; // need to consider both refraction and translucency through the material by coverage.
 			if (bIsHoldout)
 			{
 				// In the refraction case, the non-opaque part should also contribute to the holdout
