Substrate - now works with RTTranslucency.
It does not match rasteriser, it has the same visual difference as for the legacy path RTTranslucency, so at least that is on par. This is a starting point to be improved.

[RN] Substrate now has basic support for RTTranslucency

diff --git a/Engine/Shaders/Private/RayTracing/RayTracingCommon.ush b/Engine/Shaders/Private/RayTracing/RayTracingCommon.ush
--- a/Engine/Shaders/Private/RayTracing/RayTracingCommon.ush
+++ b/Engine/Shaders/Private/RayTracing/RayTracingCommon.ush
@@ -566,6 +566,7 @@ struct FMaterialClosestHitPayload : FMinimalPayload
 	float3 Radiance;
 	float3 WorldNormal; // Top Layer Normal
 	float Opacity;
+	float3 TransmittancePreCoverage;
 	uint BlendingMode;
 	uint PrimitiveLightingChannelMask;
 	float3 IndirectIrradiance;
@@ -602,12 +603,13 @@ struct FPackedMaterialClosestHitPayload : FMinimalPayload
 	// float FMinimalPayload::HitT                       // 4  bytes
 	uint PackedRayCone;                                  // 4  bytes (fp16 width, fp16 spread) aliased with RayDirection for MS lighting payload
 	uint Flags;											 // 4  bytes (see layout above)
+	uint TransmittancePreCoverage;						 // 4  bytes
 	uint Radiance;										 // 4  bytes
 	uint PackedGeometryNormal;							 // 4  bytes
 	uint SceneInstanceIndex;                             // 4  bytes - GPUSceneInstanceIndex
-														 //=24 bytes - must be in sync with InternalGetRayTracingMaterialPayloadSizeInBytes()
-	FSubstrateRaytracingPayload SubstrateData;				 //  ? bytes - must be in sync with GetRaytracingMaterialPayloadSize (see FSubstrateRaytracingPayload)
-														 //=24+? bytes total
+														 //=28 bytes - must be in sync with InternalGetRayTracingMaterialPayloadSizeInBytes()
+	FSubstrateRaytracingPayload SubstrateData;			 //  ? bytes - must be in sync with GetRaytracingMaterialPayloadSize (see FSubstrateRaytracingPayload)
+														 //=28+? bytes total
 
 	void SetMinimalPayloadMode()			{ Flags |= RAY_TRACING_PAYLOAD_INPUT_FLAG_MINIMAL_PAYLOAD; }
 	bool IsMinimalPayloadMode()				{ return (Flags & RAY_TRACING_PAYLOAD_INPUT_FLAG_MINIMAL_PAYLOAD) != 0; }
@@ -662,11 +664,13 @@ struct FPackedMaterialClosestHitPayload : FMinimalPayload
 	uint  GetBlendingMode()					{ return (Flags >> 16u) & 0x7; }
 	uint  GetPrimitiveLightingChannelMask() { return (Flags >> 19u) & 0x7; }
 	float GetOpacity()						{ return UnpackR8(Flags >> 24u); }
+	float3 GetTransmittancePreCoverage()		{ return UnpackR11G11B10F(TransmittancePreCoverage); }
 
 	void SetFlags(uint In)							{ Flags |= ((In & 0xFF)); }
 	void SetBlendingMode(uint In)					{ Flags  = (Flags & 0xFFF8FFFF) | ((In & 0x7)<<16u); }
 	void SetPrimitiveLightingChannelMask(uint In)	{ Flags  = (Flags & 0xFFC7FFFF) | ((In & 0x7)<<19u); }
 	void SetOpacity(float In)						{ Flags  = (Flags & 0x00FFFFFF) | (PackR8(In) << 24u); }
+	void SetTransmittancePreCoverage(float3 In)		{ TransmittancePreCoverage  = PackR11G11B10F(In); }
 
 	void SetRayDirection(float3 In)			{ PackedRayCone = SubstratePackNormal(In); }
 	float3 GetRayDirection()				{ return SubstrateUnpackNormal(PackedRayCone); }
@@ -701,6 +705,7 @@ FPackedMaterialClosestHitPayload PackRayTracingPayload(FMaterialClosestHitPayloa
 	Out.SetRadiance(In.Radiance);
 	Out.SubstrateData = In.SubstrateData;
 	Out.SetOpacity(In.Opacity);
+	Out.SetTransmittancePreCoverage(In.TransmittancePreCoverage);
 
 	return Out;
 }
@@ -721,6 +726,7 @@ FMaterialClosestHitPayload UnpackRayTracingPayload(FPackedMaterialClosestHitPayl
 	Out.TranslatedWorldPos = Ray.Origin + Out.HitT * Ray.Direction;
 	Out.RayDirection = In.GetRayDirection();
 	Out.Opacity = In.GetOpacity();
+	Out.TransmittancePreCoverage = In.GetTransmittancePreCoverage();
 	return Out;
 }
 
diff --git a/Engine/Shaders/Private/RayTracing/RayTracingMaterialHitShaders.usf b/Engine/Shaders/Private/RayTracing/RayTracingMaterialHitShaders.usf
--- a/Engine/Shaders/Private/RayTracing/RayTracingMaterialHitShaders.usf
+++ b/Engine/Shaders/Private/RayTracing/RayTracingMaterialHitShaders.usf
@@ -233,8 +233,10 @@ void GetMaterialPayload(
 	const float Dither = 0.5f;
 	float3 EmissiveLuminance = 0.0f;
 	float TotalCoverage = 1.f;
+	float3 TransmittancePreCoverage = 0.0f;
 	bool bSubstrateSubsurfaceEnable = false;
 	const float3 WorldBentNormal0 = GetWorldBentNormalZero(MaterialParameters);
+	const float3 CamVector = MaterialParameters.CameraVector;
 
 	// Initialise a Substrate header with normal in registers
 	FSubstrateData SubstrateData = PixelMaterialInputs.GetFrontSubstrateData();
@@ -248,16 +250,19 @@ void GetMaterialPayload(
 	float3	UnlitSurfaceNormal = 0.0f;
 	SubstratePixelHeader.SubstrateUpdateTreeUnlit(
 		uint2(MaterialParameters.SvPosition.xy),
-		MaterialParameters.CameraVector,
+		CamVector,
 		SubstrateData,
 		UnlitSurfaceLuminancePostCoverage,
 		UnlitSurfaceCoverage,
 		UnlitSurfaceTransmittancePreCoverage,
 		UnlitSurfaceNormal);
 	EmissiveLuminance = UnlitSurfaceLuminancePostCoverage;
+
 	#if !SUBSTRATE_OPAQUE_MATERIAL
 	TotalCoverage = UnlitSurfaceCoverage;
+	TransmittancePreCoverage = UnlitSurfaceTransmittancePreCoverage;
 	#endif
+
   #else // SUBSTRATE_OPTIMIZED_UNLIT
 	SubstratePixelHeader.IrradianceAO.MaterialAO = GetMaterialAmbientOcclusion(PixelMaterialInputs);
 	SubstratePixelHeader.SetCastContactShadow(GetPrimitiveData(MaterialParameters).Flags & PRIMITIVE_SCENE_DATA_FLAG_HAS_CAST_CONTACT_SHADOW);
@@ -275,18 +280,16 @@ void GetMaterialPayload(
 	#if MATERIALBLENDING_ANY_TRANSLUCENT
 	if (SubstratePixelHeader.SubstrateTree.BSDFCount > 0)
 	{
-		float3 TransmittancePreCoverage = 0.f;
 		SubstratePixelHeader.SubstrateUpdateTree(
 			SubstrateData, 
-			MaterialParameters.CameraVector, 
+			CamVector,
 			Settings, 
 			TotalCoverage, 
 			TransmittancePreCoverage);
-		TotalCoverage *= (1.f-dot(TransmittancePreCoverage, float(1.f/3.f).xxx));
+
 	}
 	#endif
 
-	const float3 CamVector = MaterialParameters.CameraVector;
 	Payload.SubstrateData = InitialiseRWSubstrateMaterialContainer();
 	PackSubstrateOut(
 		Payload.SubstrateData,
@@ -311,6 +314,8 @@ void GetMaterialPayload(
 
 	// Opacity
 	Payload.Opacity = TotalCoverage;
+	Payload.TransmittancePreCoverage = TransmittancePreCoverage;
+
 	#if MATERIALBLENDING_MASKED
 	// regular masked mode - binary decision
 	if (GetMaterialMask(PixelMaterialInputs) < 0)
@@ -318,6 +323,8 @@ void GetMaterialPayload(
 		Payload.Opacity = 0;
 	}
 	#endif // MATERIALBLENDING_MASKED
+
+
   #endif // MATERIAL_IS_SUBSTRATE
 }
 #else // SUBSTRATE_ENABLED
@@ -610,7 +617,11 @@ RAY_TRACING_ENTRY_CLOSEST_HIT(MaterialCHS,
 	if (bIsHoldout && bIsCameraRay)
 	{
 #if SUBSTRATE_ENABLED
-		// SUBSTRATE_TODO: Need to clear the material data and set a black unlit material in this case
+		Payload.SubstrateData = (FSubstrateRaytracingPayload)0;
+		Payload.ShadingModelID = SHADINGMODELID_UNLIT;
+		Payload.Radiance = 0.0;
+		Payload.Opacity = 1.0;
+		Payload.TransmittancePreCoverage = 0.0;
 #else
 		Payload.ShadingModelID = SHADINGMODELID_UNLIT;
 		Payload.Radiance = 0.0;
diff --git a/Engine/Shaders/Private/RayTracing/RayTracingPrimaryRays.usf b/Engine/Shaders/Private/RayTracing/RayTracingPrimaryRays.usf
--- a/Engine/Shaders/Private/RayTracing/RayTracingPrimaryRays.usf
+++ b/Engine/Shaders/Private/RayTracing/RayTracingPrimaryRays.usf
@@ -28,6 +28,9 @@
 #define	ERayTracingPrimaryRaysFlag_PrimaryView                 (1u << 1)
 #define	ERayTracingPrimaryRaysFlag_AllowSkipSkySample          (1u << 2)
 
+#define SUBSTRATE_MATERIALCONTAINER_IS_VIEWRESOURCE 0
+#include "../Substrate/Substrate.ush"
+#include "../Substrate/SubstrateEvaluation.ush"
 
 int SamplesPerPixel;
 int MaxRefractionRays;
@@ -82,24 +85,15 @@ float3 GetSkyRadiance(float3 Direction, float Roughness)
 
 RAY_TRACING_ENTRY_RAYGEN(RayTracingPrimaryRaysRGS)
 {
-	// SUBSTRATE_TODO
-#if !SUBSTRATE_ENABLED
 	uint2 DispatchThreadId = DispatchRaysIndex().xy + View.ViewRectMin.xy;
 	uint2 PixelCoord = GetPixelCoord(DispatchThreadId, UpscaleFactor);
 	uint LinearIndex = PixelCoord.y * View.BufferSizeAndInvSize.x + PixelCoord.x; // TODO(Denoiser): PixelCoord or DispatchThreadId
 	
 	float2 InvBufferSize = View.BufferSizeAndInvSize.zw;
 	float2 UV = (float2(PixelCoord) + 0.5) * InvBufferSize;
-
-#if 0
-	FGBufferData GBufferData = GetGBufferDataFromSceneTextures(UV);
-#else
-	//#dxr-todo: workaround for flickering. UE-87281
-	FGBufferData GBufferData = GetGBufferDataFromSceneTexturesLoad(PixelCoord);
-#endif
-
-	float Depth = GBufferData.Depth;
-	float3 TranslatedWorldPosition = ReconstructTranslatedWorldPositionFromDepth(UV, Depth);
+	
+	float GBufferDepth = ConvertFromDeviceZ(SampleDeviceZFromSceneTextures(UV)); // CalcSceneDepth SampleDeviceZFromSceneTextures
+	float3 TranslatedWorldPosition = ReconstructTranslatedWorldPositionFromDepth(UV, GBufferDepth);
 
 	// Trace rays from camera origin to Gbuffer (only translucent objects are considered on the first ray, so a bias is not needed)
 	FRayDesc Ray = CreatePrimaryRay(UV);
@@ -174,7 +168,55 @@ RAY_TRACING_ENTRY_RAYGEN(RayTracingPrimaryRaysRGS)
 			RayCone,
 			bRefractionEnableSkyLightContribution,
 			PathVertexRadiance);
-		LastRoughness = Payload.Roughness;
+
+		float PayloadRoughness = 0;
+		float3 PayloadDiffuseColor = 0;
+		float3 PayloadSpecularColor = 0;
+		float3 PayloadWorldTangent = 0;
+		float PayloadAnisotropy = 0;
+		float3 PayloadLuminanceWeight = 1.0;
+	#if SUBSTRATE_ENABLED
+		
+		float3 SubstrateMaterialThroughput = 1.0;
+		FSubstrateAddressing SubstrateAddressing = GetSubstratePixelDataByteOffset(0, 0, 0);
+		FSubstratePixelHeader SubstratePixelHeader = UnpackSubstrateHeaderIn(Payload.SubstrateData, SubstrateAddressing, Payload.SubstrateData);
+
+		BRANCH
+		if (SubstratePixelHeader.ClosureCount > 0) // With ray tracing, we only account for the first BSDF since all materials should be simplified to a single Closure.
+		{
+			const bool bSkipSSSMaterialOverride = true;
+			const FSubstrateBSDF BSDF = UnpackSubstrateBSDFIn(Payload.SubstrateData, SubstrateAddressing, SubstratePixelHeader, bSkipSSSMaterialOverride);
+			
+			if (SubstrateIsBSDFVisible(BSDF))
+			{
+				// Create the BSDF context
+				FSubstrateBSDFContext SubstrateBSDFContext = SubstrateCreateBSDFContext(SubstratePixelHeader, BSDF, SubstrateAddressing, -Ray.Direction);
+
+				PayloadRoughness = SLAB_ROUGHNESS(BSDF);
+				PayloadDiffuseColor = SubstrateGetBSDFDiffuseColor(BSDF);
+				PayloadSpecularColor = SubstrateGetBSDFSpecularColor(BSDF);
+				PayloadWorldTangent = SubstrateBSDFContext.X;
+				PayloadAnisotropy = SLAB_ANISOTROPY(BSDF);
+
+				PayloadLuminanceWeight = BSDF.LuminanceWeightV;
+				const float Coverage = saturate(Payload.Opacity);
+				SubstrateMaterialThroughput = (1.0 - Coverage) + Coverage * Payload.TransmittancePreCoverage;
+			}
+
+		}
+
+	#else
+
+		PayloadRoughness = Payload.Roughness;
+		PayloadDiffuseColor = Payload.DiffuseColor;
+		PayloadSpecularColor = Payload.SpecularColor;
+		PayloadWorldTangent = Payload.WorldTangent;
+		PayloadAnisotropy = Payload.Anisotropy;
+		PayloadLuminanceWeight = Payload.Opacity;
+
+	#endif
+
+		LastRoughness = PayloadRoughness;
 
 		//
 		// Handle no hit condition
@@ -194,8 +236,9 @@ RAY_TRACING_ENTRY_RAYGEN(RayTracingPrimaryRaysRGS)
 		//
 		// Handle surface lighting
 		//
+		
+		const float3 vertexRadianceWeight = PayloadLuminanceWeight;	// Opacity as coverage. This works for RAY_TRACING_BLEND_MODE_OPAQUE and RAY_TRACING_BLEND_MODE_TRANSLUCENT.
 
-		const float vertexRadianceWeight = Payload.Opacity;	// Opacity as coverage. This works for RAY_TRACING_BLEND_MODE_OPAQUE and RAY_TRACING_BLEND_MODE_TRANSLUCENT.
 		// It is also needed for RAY_TRACING_BLEND_MODE_ADDITIVE and  RAY_TRACING_BLEND_MODE_ALPHA_COMPOSITE: radiance contribution is alway weighted by coverage.
 		// #dxr_todo: I have not been able to setup a material using RAY_TRACING_BLEND_MODE_MODULATE.
 
@@ -219,7 +262,7 @@ RAY_TRACING_ENTRY_RAYGEN(RayTracingPrimaryRaysRGS)
 				float4 ClipPos = mul(float4(TranslatedHitPoint, 1.0f), PrimaryView.TranslatedWorldToClip);
 				float3 VolumeUV = ComputeVolumeUVFromNDC(ClipPos);
 
-				ViewFogInScatteringAndTransmittance = CombineVolumetricFog(ViewFogInScatteringAndTransmittance, VolumeUV, 0/*EyeIndex*/, GBufferData.Depth);
+				ViewFogInScatteringAndTransmittance = CombineVolumetricFog(ViewFogInScatteringAndTransmittance, VolumeUV, 0 /*EyeIndex*/, GBufferDepth);
 			}
 		}
 
@@ -229,7 +272,7 @@ RAY_TRACING_ENTRY_RAYGEN(RayTracingPrimaryRaysRGS)
 		if (bIsHoldout)
 		{
 			// A holdout object acts as if it was a background pixel so that it influences the alpha channel
-			BackgroundVisibility += Luminance(PathThroughput) * vertexRadianceWeight;
+			BackgroundVisibility += Luminance(PathThroughput * vertexRadianceWeight);
 		}
 
 		//
@@ -238,8 +281,8 @@ RAY_TRACING_ENTRY_RAYGEN(RayTracingPrimaryRaysRGS)
 
 		// Shorten the rays on rougher surfaces between user-provided min and max ray lengths.
 		// When a shortened ray misses the geometry, we fall back to local reflection capture sampling (similar to SSR).
-		const float LocalMaxRayDistance = bAllowSkySampling ? 1e27f : lerp(TranslucencyMaxRayDistance, TranslucencyMinRayDistance, Payload.Roughness);
-		if (Payload.Roughness < TranslucencyMaxRoughness)
+		const float LocalMaxRayDistance = bAllowSkySampling ? 1e27f : lerp(TranslucencyMaxRayDistance, TranslucencyMinRayDistance, PayloadRoughness);
+		if (PayloadRoughness < TranslucencyMaxRoughness)
 		{
 			// Trace reflection ray 
 			float2 RandSample = RandomSequence_GenerateSample2D(RandSequence);
@@ -249,9 +292,9 @@ RAY_TRACING_ENTRY_RAYGEN(RayTracingPrimaryRaysRGS)
 			ReflectionRay.TMax = LocalMaxRayDistance;
 			ReflectionRay.Origin = TranslatedHitPoint;
 
-			ModifyGGXAnisotropicNormalRoughness(Payload.WorldTangent, Payload.Anisotropy, Payload.Roughness, Payload.WorldNormal, Ray.Direction);
+			ModifyGGXAnisotropicNormalRoughness(PayloadWorldTangent, PayloadAnisotropy, PayloadRoughness, Payload.WorldNormal, Ray.Direction);
 
-			ReflectionRay.Direction = GenerateReflectedRayDirection(Ray.Direction, Payload.WorldNormal, Payload.Roughness, RandSample);
+			ReflectionRay.Direction = GenerateReflectedRayDirection(Ray.Direction, Payload.WorldNormal, PayloadRoughness, RandSample);
 			ApplyPositionBias(ReflectionRay, Payload.WorldNormal, MaxNormalBias);
 
 			const uint ReflectionRayFlags = RAY_FLAG_CULL_BACK_FACING_TRIANGLES;
@@ -288,7 +331,7 @@ RAY_TRACING_ENTRY_RAYGEN(RayTracingPrimaryRaysRGS)
 
 			// #dxr_todo: reflection IOR and clear coat also? This only handles default material.
 			float NoV = saturate(dot(-Ray.Direction, Payload.WorldNormal));
-			const float3 ReflectionThroughput = EnvBRDF(Payload.SpecularColor, Payload.Roughness, NoV);
+			const float3 ReflectionThroughput = EnvBRDF(PayloadSpecularColor, PayloadRoughness, NoV);
 
 			// For reflections, we can only sample the height fog (volumetric fog is only for primary rays)
 			if (HeightFog > 0)
@@ -331,8 +374,25 @@ RAY_TRACING_ENTRY_RAYGEN(RayTracingPrimaryRaysRGS)
 				// noop => if lumen is not enabled: bRefractionEnableSkyLightContribution == true
 			}
 
-			PathRadiance += PathThroughput * vertexRadianceWeight * VolumeLighting * Payload.DiffuseColor;
+			PathRadiance += PathThroughput * vertexRadianceWeight * VolumeLighting * PayloadDiffuseColor;
 		}
+		
+		float3 RefractedDirection = Ray.Direction;
+	#if SUBSTRATE_ENABLED
+
+		bHasScattered = false;
+		PathThroughput *= SubstrateMaterialThroughput;
+
+		if (bIsHoldout)
+		{
+			// In the refraction case, the non-opaque part should also contribute to the holdout
+			// The net result is that the entire object becomes opaque for the alpha channel.
+			BackgroundVisibility += Luminance(PathThroughput) * (1.0 - Payload.Opacity);
+			PathThroughput = 0.0;
+			break;
+		}
+
+	#else // SUBSTRATE_ENABLED
 
 		// Update the refraction/transparency path transmittance and check stop condition
 		switch (Payload.BlendingMode)
@@ -370,7 +430,6 @@ RAY_TRACING_ENTRY_RAYGEN(RayTracingPrimaryRaysRGS)
 		//
 
 		// Set refraction ray for next iteration
-		float3 RefractedDirection = Ray.Direction;
 		if (Payload.ShadingModelID == SHADINGMODELID_THIN_TRANSLUCENT)
 		{
 			const float3 N = Payload.WorldNormal;
@@ -432,6 +491,8 @@ RAY_TRACING_ENTRY_RAYGEN(RayTracingPrimaryRaysRGS)
 			NextMaxRayDistance = LocalMaxRayDistance;
 		}
 
+	#endif // SUBSTRATE_ENABLED
+
 
 		//
 		// Setup refracted ray to be traced
@@ -442,7 +503,7 @@ RAY_TRACING_ENTRY_RAYGEN(RayTracingPrimaryRaysRGS)
 		Ray.TMax = NextMaxRayDistance;
 		Ray.Direction = RefractedDirection;
 		float SurfaceCurvature = 0.0f; /* #todo_dxr assume no curvature */
-		RayCone = PropagateRayCone(RayCone, SurfaceCurvature, Depth);
+		RayCone = PropagateRayCone(RayCone, SurfaceCurvature, GBufferDepth);
 	}
 
 	if (!bHasScattered)
@@ -468,5 +529,5 @@ RAY_TRACING_ENTRY_RAYGEN(RayTracingPrimaryRaysRGS)
 	PathRadiance = ClampToHalfFloatRange(PathRadiance);
 	ColorOutput[DispatchThreadId] = float4(PathRadiance, FinalAlpha);
 	RayHitDistanceOutput[DispatchThreadId] = HitDistance;
-#endif // !SUBSTRATE_ENABLED
+	
 }
diff --git a/Engine/Shaders/Public/ShaderVersion.ush b/Engine/Shaders/Public/ShaderVersion.ush
--- a/Engine/Shaders/Public/ShaderVersion.ush
+++ b/Engine/Shaders/Public/ShaderVersion.ush
@@ -3,4 +3,4 @@
 // Each time the console command is executed it generates a new GUID. As this file is included
 // in Platform.ush (which should be included in any shader) it allows to invalidate the shader DDC.
 // 
-#pragma message("UESHADERMETADATA_VERSION 1115A78C-1E1B-4E96-A416-DA839794A1C5")
\ No newline at end of file
+#pragma message("UESHADERMETADATA_VERSION 8F12FB5A-D796-45C7-88A4-95639B056795")
diff --git a/Engine/Source/Runtime/Engine/Private/Materials/MaterialExpressionSubstrate.cpp b/Engine/Source/Runtime/Engine/Private/Materials/MaterialExpressionSubstrate.cpp
--- a/Engine/Source/Runtime/Engine/Private/Materials/MaterialExpressionSubstrate.cpp
+++ b/Engine/Source/Runtime/Engine/Private/Materials/MaterialExpressionSubstrate.cpp
@@ -630,7 +630,7 @@ FSubstrateOperator* UMaterialExpressionSubstrateShadingModels::SubstrateGenerate
 		}
 		else if (ShadingModels.HasShadingModel(MSM_ThinTranslucent))
 		{
-			return AddDefaultWorstCase(ESubstrateBsdfFeature_None);
+			return AddDefaultWorstCase(ESubstrateBsdfFeature_MFPPluggedIn);
 		}
 		else if (ShadingModels.HasShadingModel(MSM_SubsurfaceProfile))
 		{
diff --git a/Engine/Source/Runtime/RenderCore/Private/BuiltInRayTracingShaders.cpp b/Engine/Source/Runtime/RenderCore/Private/BuiltInRayTracingShaders.cpp
--- a/Engine/Source/Runtime/RenderCore/Private/BuiltInRayTracingShaders.cpp
+++ b/Engine/Source/Runtime/RenderCore/Private/BuiltInRayTracingShaders.cpp
@@ -14,7 +14,7 @@ uint32 GetRaytracingMaterialPayloadSizeFullySimplified()
 	if (Substrate::IsSubstrateEnabled())
 	{
 		// All the data from FPackedMaterialClosestHitPayload except FSubstrateRaytracingPayload (see RayTracingCommon.ush)
-		uint32 PayloadSizeBytes = 6 * sizeof(uint32);
+		uint32 PayloadSizeBytes = 7 * sizeof(uint32);
 
 		// The remaining data from FSubstrateRaytracingPayload.
 		const bool bFullySimplifiedMaterial = true;	// This is needed because ERayTracingPayloadType::RayTracingMaterial will be fully simplified, see FShaderType::ModifyCompilationEnvironment.
