Substrate - fixed legacy shading model conversion when Opacity is >1, especially ThinTranslucent.

#rnx
#tests none

diff --git a/Engine/Shaders/Private/Substrate/SubstrateLegacyConversion.ush b/Engine/Shaders/Private/Substrate/SubstrateLegacyConversion.ush
--- a/Engine/Shaders/Private/Substrate/SubstrateLegacyConversion.ush
+++ b/Engine/Shaders/Private/Substrate/SubstrateLegacyConversion.ush
@@ -161,6 +161,8 @@ FSubstrateData SubstrateConvertLegacyMaterialDynamic(
 	const float3 Zeros = float3(0, 0, 0);
 	const float3 Ones  = float3(1, 1, 1);
 
+	const float OneMinusOpacity = 1.f - Opacity;
+
 	// Can only mix Unlit / DefaultLit / Sub-Surface / Preintegrated-Skin / Subsurface-Profile / ClearCoat / Foliage / Cloth / Eye / Thin
 
 	// Note: If hair or single layer water are enabled with other shading model, the code generation won't be ideal, as the compiler will struggle with unrolling the BSDFs
@@ -265,7 +267,7 @@ FSubstrateData SubstrateConvertLegacyMaterialDynamic(
 		LegacySlab.Anisotropy = Anisotropy;
 		LegacySlab.SSSMFP = max(1e-05f, MFPInCm); // Ensure the MFP is not null to force the material as have SSS
 		LegacySlab.SSSMFPScale = 1.0f;
-		LegacySlab.SSSPhaseAniso = saturate(1.f - Opacity); // Opaque-Thick: isotropic phase function, Thin: forward phase scattering function
+		LegacySlab.SSSPhaseAniso = OneMinusOpacity; // Opaque-Thick: isotropic phase function, Thin: forward phase scattering function
 		LegacySlab.Emissive = Emissive;
 		LegacySlab.UseSSSDiffusion = false;
 		LegacySlab.Thickness = ThicknessInCm;
@@ -277,7 +279,7 @@ FSubstrateData SubstrateConvertLegacyMaterialDynamic(
 		// If we try to reduce the diffusion according to transmittance that will increase mfp and make the material disapear (make it optically thin).
 		// Instead for now we use a mfp of 0 and make the material vanish according to opacity.
 		// All this will be fixed when this material will be converted to LUT.
-		const float Transmittance = saturate(1.f - Opacity);
+		const float Transmittance = OneMinusOpacity;
 		LegacySlab.SSSMFP = 0;// MFPInCm;
 		LegacySlab.SSSMFPScale = 1.0f;
 		LegacySlab.Thickness = ThicknessInCm;
@@ -299,7 +301,7 @@ FSubstrateData SubstrateConvertLegacyMaterialDynamic(
 		LegacySlab.Anisotropy = Anisotropy;
 		LegacySlab.SSSMFP = max(1e-05f, MFPInCm); // Ensure the MFP is not null to force the material as have SSS
 		LegacySlab.SSSProfileID = 0.f;
-		LegacySlab.SSSMFPScale = 1.f - Opacity;
+		LegacySlab.SSSMFPScale = OneMinusOpacity;
 		LegacySlab.SSSPhaseAniso = 0.93f;
 		LegacySlab.UseSSSDiffusion = true;
 		LegacySlab.Emissive = Emissive;
@@ -367,7 +369,7 @@ FSubstrateData SubstrateConvertLegacyMaterialDynamic(
 		LegacySlab.SSSProfileID = 0.f;
 		LegacySlab.SSSMFP = max(1e-05f, MFPInCm); // Ensure the MFP is not null to force the material as have SSS
 		LegacySlab.SSSMFPScale = 1.0f;
-		LegacySlab.SSSPhaseAniso = 1.0f - saturate(Opacity); // Opaque-Thick: isotropic phase function, Thin: forward phase scattering function
+		LegacySlab.SSSPhaseAniso = OneMinusOpacity; // Opaque-Thick: isotropic phase function, Thin: forward phase scattering function
 		LegacySlab.Emissive = Emissive;
 		LegacySlab.UseSSSDiffusion = false;
 		LegacySlab.Thickness = Thickness;
@@ -468,6 +470,13 @@ FSubstrateData SubstrateConvertLegacyMaterialStatic(
 	const float3 Zeros = float3(0, 0, 0);
 	const float3 Ones = float3(1, 1, 1);
 
+	// We need to make sure a few values are safe for the transformations into slab input we do.
+	Emissive = max(Emissive, 0.0f.xxx);
+	Metallic = saturate(Metallic);
+	ThinTranslucentTransmittanceColor = saturate(ThinTranslucentTransmittanceColor);
+	ThinTranslucentSurfaceCoverage = saturate(ThinTranslucentSurfaceCoverage);
+	Opacity = saturate(Opacity);
+
 #if SUBSTRATE_USE_PREMULTALPHA_OVERRIDE
 	// Opacity is applied using the "Opacity Override" pin on the root node when using alpha composite (pre multiplied alpha).
 	// This is needed because some material might have the pin plugged in and the blending mode overriden on the material instance.
