Substrate - fixed vertical layering parameter blening that can explode sometimes.

#rnx
#tests none

diff --git a/Engine/Shaders/Private/Substrate/Substrate.ush b/Engine/Shaders/Private/Substrate/Substrate.ush
--- a/Engine/Shaders/Private/Substrate/Substrate.ush
+++ b/Engine/Shaders/Private/Substrate/Substrate.ush
@@ -2825,12 +2825,12 @@ FSubstrateBSDF SubstrateVerticalLayeringParameterBlending(in FSubstrateBSDF TopB
 #else
 
 	// We compute the total transmittance of the medium over the max coverage
-	const float MaxCoverage = max(TopBSDF.Coverage, BaseBSDF.Coverage);
+	const float MaxCoverage = max(0.000001, max(TopBSDF.Coverage, BaseBSDF.Coverage));
 	FVerticalLayeringInfo Info2 = GetVerticalLayeringInfo(TopBSDF.Coverage / MaxCoverage, BaseBSDF.Coverage / MaxCoverage);
 #if 0
 	// Using transmittance as a way to interpolate should be better but there are issues with large values
-	const float3 TopExtinction			= 1.0f / max(0.0, TopBSDF.TmpMFP);
-	const float3 BaseExtinction			= 1.0f / max(0.0, BaseBSDF.TmpMFP);
+	const float3 TopExtinction			= 1.0f / max(0.000001, TopBSDF.TmpMFP);
+	const float3 BaseExtinction			= 1.0f / max(0.000001, BaseBSDF.TmpMFP);
 	const float3 TopTransmitance		= exp(-TopExtinction);	// extinction = optical depth, we assume here d = 1 meter
 	const float3 BaseTransmitance		= exp(-BaseExtinction);
 	const float3 NewTransmittance		= Info2.TransmittanceOnlyBottom * BaseTransmitance + Info2.TransmittanceOnlyTop * TopTransmitance + Info2.TransmittanceTopAndBottom * TopTransmitance * BaseTransmitance;
