[Light Grid] Merge CompactLinks pass into main light grid injection pass
- avoids extra linked list iteration to count number of links and separate compute pass
- 0.36ms -> 0.31ms

diff --git a/Engine/Shaders/Private/LightGridInjection.usf b/Engine/Shaders/Private/LightGridInjection.usf
--- a/Engine/Shaders/Private/LightGridInjection.usf
+++ b/Engine/Shaders/Private/LightGridInjection.usf
@@ -9,6 +9,7 @@
 #include "LargeWorldCoordinates.ush"
 #include "ReflectionEnvironmentShared.ush"
 #include "LightGridCommon.ush"
+#include "WaveOpUtil.ush"
 
 #if USE_HZB_CULL
 	#include "Nanite/NaniteHZBCull.ush"
@@ -25,7 +26,6 @@ RWStructuredBuffer<uint> RWNumCulledLightsGrid;
 #endif
 
 RWStructuredBuffer<uint> RWNextCulledLightLink;
-RWStructuredBuffer<uint> RWStartOffsetGrid;
 RWStructuredBuffer<uint> RWCulledLightLinks;
 
 uint NumReflectionCaptures;
@@ -283,8 +283,24 @@ bool OverlapsReflectionCapture(uint ReflectionCaptureIndex, FCellBounds CellBoun
 	return (BoxDistanceSq < CaptureRadius * CaptureRadius);
 }
 
+void WriteReversedLinkedList(uint LinkOffset, uint NumPrimitives, uint CulledLightDataStart)
+{
+	uint CurrentIndex = 0;
+
+	while (LinkOffset != 0xFFFFFFFF && CurrentIndex < NumPrimitives)
+	{
+		// Reverse the order as we write them out, which restores the original order before the reverse linked list was built
+		const uint PrimitiveIndex = RWCulledLightLinks[LinkOffset * LIGHT_LINK_STRIDE + 0];
+		RWCulledLightDataGrid[CulledLightDataStart + NumPrimitives - CurrentIndex - 1] = PrimitiveIndex;
+		CurrentIndex++;
+		LinkOffset = RWCulledLightLinks[LinkOffset * LIGHT_LINK_STRIDE + 1];
+	}
+}
+
 #if SHADER_LIGHT_GRID_INJECTION_CS
 
+RWStructuredBuffer<uint> RWNextCulledLightData;
+
 [numthreads(THREADGROUP_SIZE, THREADGROUP_SIZE, THREADGROUP_SIZE)]
 void LightGridInjectionCS(
 	uint3 GroupId : SV_GroupID,
@@ -332,6 +348,8 @@ void LightGridInjectionCS(
 		uint bHasRectLights = false;
 		uint bHasTexturedLights = false;
 
+		uint LinkOffset = 0xFFFFFFFF;
+
 		LOOP
 		for (uint LocalLightIndex = 0; LocalLightIndex < NumLocalLights; LocalLightIndex++)
 		{
@@ -348,29 +366,48 @@ void LightGridInjectionCS(
 
 				if (NextLink < NumAvailableLinks)
 				{
-					uint PreviousLink;
-					InterlockedExchange(RWStartOffsetGrid[GridIndex], NextLink, PreviousLink);
 					RWCulledLightLinks[NextLink * LIGHT_LINK_STRIDE + 0] = LocalLightIndex;
-					RWCulledLightLinks[NextLink * LIGHT_LINK_STRIDE + 1] = PreviousLink;
+					RWCulledLightLinks[NextLink * LIGHT_LINK_STRIDE + 1] = LinkOffset;
+
+					LinkOffset = NextLink;
+
+					++NumVisibleLights;
+					if (LocalLightIndex >= MegaLightsSupportedStartIndex)
+					{
+						++NumVisibleMegaLights;
+					}
 				}
 #else // !USE_LINKED_CULL_LIST
 				if (NumVisibleLights < MaxCulledLightsPerCell)
 				{
 					RWCulledLightDataGrid[GridIndex * MaxCulledLightsPerCell + NumVisibleLights] = LocalLightIndex;
-				}
 
-				++NumVisibleLights;
-				if (LocalLightIndex >= MegaLightsSupportedStartIndex)
-				{
-					++NumVisibleMegaLights;
+					++NumVisibleLights;
+					if (LocalLightIndex >= MegaLightsSupportedStartIndex)
+					{
+						++NumVisibleMegaLights;
+					}
 				}
 #endif // !USE_LINKED_CULL_LIST
 			}
 		}
 
+#if USE_LINKED_CULL_LIST
+		uint CulledLightDataStart;
+		WaveInterlockedAdd_(RWNextCulledLightData[0], NumVisibleLights, CulledLightDataStart);
+
+		// Mega Lights are order dependent
+		WriteReversedLinkedList(LinkOffset, NumVisibleLights, CulledLightDataStart);
+#else
+		uint CulledLightDataStart = NumGridCells * MaxCulledLightsPerCell;
+#endif
+
 		RWNumCulledLightsGrid[GridIndex * NUM_CULLED_LIGHTS_GRID_STRIDE + 0] = PackCulledLightsGridHeader0(NumVisibleLights, NumVisibleMegaLights);
-		RWNumCulledLightsGrid[GridIndex * NUM_CULLED_LIGHTS_GRID_STRIDE + 1] = PackCulledLightsGridHeader1(GridIndex * MaxCulledLightsPerCell, bHasRectLights, bHasTexturedLights);
+		RWNumCulledLightsGrid[GridIndex * NUM_CULLED_LIGHTS_GRID_STRIDE + 1] = PackCulledLightsGridHeader1(CulledLightDataStart, bHasRectLights, bHasTexturedLights);
 		
+		uint NumVisibleCaptures = 0;
+		uint CaptureLinkOffset = 0xFFFFFFFF;
+
 		LOOP
 		for (uint ReflectionCaptureIndex = 0; ReflectionCaptureIndex < NumReflectionCaptures; ReflectionCaptureIndex++)
 		{
@@ -382,24 +419,38 @@ void LightGridInjectionCS(
 
 					if (NextLink < NumAvailableLinks)
 					{
-						uint PreviousLink;
-						InterlockedExchange(RWStartOffsetGrid[NumGridCells + GridIndex], NextLink, PreviousLink);
 						RWCulledLightLinks[NextLink * LIGHT_LINK_STRIDE + 0] = ReflectionCaptureIndex;
-						RWCulledLightLinks[NextLink * LIGHT_LINK_STRIDE + 1] = PreviousLink;
-					}
+						RWCulledLightLinks[NextLink * LIGHT_LINK_STRIDE + 1] = CaptureLinkOffset;
 
-				#else
-					uint CulledCaptureIndex;
-					InterlockedAdd(RWNumCulledLightsGrid[(NumGridCells + GridIndex) * NUM_CULLED_LIGHTS_GRID_STRIDE + 0], 1U, CulledCaptureIndex);
-					RWNumCulledLightsGrid[(NumGridCells + GridIndex) * NUM_CULLED_LIGHTS_GRID_STRIDE + 1] = (NumGridCells + GridIndex) * MaxCulledLightsPerCell;
+						CaptureLinkOffset = NextLink;
 
-					if (CulledCaptureIndex < MaxCulledLightsPerCell)
+						++NumVisibleCaptures;
+					}
+
+				#else // !USE_LINKED_CULL_LIST
+					if (NumVisibleCaptures < MaxCulledLightsPerCell)
 					{
-						RWCulledLightDataGrid[(NumGridCells + GridIndex) * MaxCulledLightsPerCell + CulledCaptureIndex] = ReflectionCaptureIndex;
+						RWCulledLightDataGrid[(NumGridCells + GridIndex) * MaxCulledLightsPerCell + NumVisibleCaptures] = ReflectionCaptureIndex;
+
+						++NumVisibleCaptures;
 					}
-				#endif
+				#endif // !USE_LINKED_CULL_LIST
 			}
 		}
+
+#if USE_LINKED_CULL_LIST
+		uint CulledReflectionCaptureDataStart;
+		WaveInterlockedAdd_(RWNextCulledLightData[0], NumVisibleCaptures, CulledReflectionCaptureDataStart);
+
+		// Reflection captures are order dependent
+		WriteReversedLinkedList(CaptureLinkOffset, NumVisibleCaptures, CulledReflectionCaptureDataStart);
+#else
+		uint CulledReflectionCaptureDataStart = NumGridCells * MaxCulledLightsPerCell;
+#endif
+
+		RWNumCulledLightsGrid[(NumGridCells + GridIndex) * NUM_CULLED_LIGHTS_GRID_STRIDE + 0] = NumVisibleCaptures;
+		RWNumCulledLightsGrid[(NumGridCells + GridIndex) * NUM_CULLED_LIGHTS_GRID_STRIDE + 1] = CulledReflectionCaptureDataStart;
+
 #else // !CULL_LIGHTS
 		LOOP
 		for (uint LocalLightIndex = 0; LocalLightIndex < NumLocalLights; LocalLightIndex++)
@@ -431,127 +482,6 @@ void LightGridInjectionCS(
 #endif // SHADER_LIGHT_GRID_INJECTION_CS
 
 
-
-#if SHADER_LIGHT_GRID_COMPACT_CS
-
-RWStructuredBuffer<uint> RWNextCulledLightData;
-StructuredBuffer<uint> StartOffsetGrid;
-StructuredBuffer<uint> CulledLightLinks;
-
-#if USE_LOCALSTORE_FOR_ATOMICS
-groupshared uint LDSNextCulledLightData;
-groupshared uint LDSBaseStart;
-#endif
-
-
-void CompactReverseLinkedList(uint GridIndex, uint SceneMax, bool FirstThread, bool bThreadValid, bool bLocalLights)
-{
-	uint NumCulledLights		= 0;
-	uint StartLinkOffset		= 0;
-	uint LinkOffset				= 0;
-	uint CulledLightDataStart	= 0;
-
-	if (bThreadValid)
-	{
-		StartLinkOffset = StartOffsetGrid[GridIndex];
-		LinkOffset = StartLinkOffset;
-
-		// Traverse the linked list to count how many culled indices we have
-		while (LinkOffset != 0xFFFFFFFF && NumCulledLights < SceneMax)
-		{
-			NumCulledLights++;
-			LinkOffset = CulledLightLinks[LinkOffset * LIGHT_LINK_STRIDE + 1];
-		}
-	}
-
-#if USE_LOCALSTORE_FOR_ATOMICS
-	LDSNextCulledLightData	= 0;
-	LDSBaseStart			= 0;
-
-	GroupMemoryBarrierWithGroupSync();
-
-	// Firstly do an interlocked add with LDS
-	InterlockedAdd(LDSNextCulledLightData, NumCulledLights, CulledLightDataStart);
-	GroupMemoryBarrierWithGroupSync();
-
-	// Then the first thread can do an interlocked add with the buffer
-	if(FirstThread)
-	{
-		InterlockedAdd(RWNextCulledLightData[0], LDSNextCulledLightData, LDSBaseStart);
-	}
-	GroupMemoryBarrierWithGroupSync();
-	
-	CulledLightDataStart += LDSBaseStart;
-#else // !USE_LOCALSTORE_FOR_ATOMICS
-	InterlockedAdd(RWNextCulledLightData[0], NumCulledLights, CulledLightDataStart);
-#endif // !USE_LOCALSTORE_FOR_ATOMICS
-
-
-	if (bThreadValid)
-	{
-		LinkOffset = StartLinkOffset;
-		uint CulledLightIndex = 0;
-		uint NumMegaLights = 0;
-
-		while (LinkOffset != 0xFFFFFFFF && CulledLightIndex < NumCulledLights)
-		{
-			// Reverse the order as we write them out, which restores the original order before the reverse linked list was built
-			// Reflection captures are order dependent
-			const uint PrimitiveIndex = CulledLightLinks[LinkOffset * LIGHT_LINK_STRIDE + 0];
-			RWCulledLightDataGrid[CulledLightDataStart + NumCulledLights - CulledLightIndex - 1] = PrimitiveIndex;
-			CulledLightIndex++;
-			LinkOffset = CulledLightLinks[LinkOffset * LIGHT_LINK_STRIDE + 1];
-
-			if (bLocalLights && PrimitiveIndex >= MegaLightsSupportedStartIndex)
-			{
-				++NumMegaLights;
-			}
-		}
-
-		if (bLocalLights)
-		{
-			// Get bHasRectLights / bHasTexturedLights written by LightGridInjectionCS before updating CulledLightDataStart with packed offset
-			uint Unused;
-			bool bHasRectLights;
-			bool bHasTexturedLights;
-			UnpackCulledLightsGridHeader1(RWNumCulledLightsGrid[GridIndex * NUM_CULLED_LIGHTS_GRID_STRIDE + 1], Unused, bHasRectLights, bHasTexturedLights);
-
-			RWNumCulledLightsGrid[GridIndex * NUM_CULLED_LIGHTS_GRID_STRIDE + 0] = PackCulledLightsGridHeader0(NumCulledLights, NumMegaLights);
-			RWNumCulledLightsGrid[GridIndex * NUM_CULLED_LIGHTS_GRID_STRIDE + 1] = PackCulledLightsGridHeader1(CulledLightDataStart, bHasRectLights, bHasTexturedLights);
-		}
-		else
-		{
-			RWNumCulledLightsGrid[GridIndex * NUM_CULLED_LIGHTS_GRID_STRIDE + 0] = NumCulledLights;
-			RWNumCulledLightsGrid[GridIndex * NUM_CULLED_LIGHTS_GRID_STRIDE + 1] = CulledLightDataStart;
-		}
-	}
-}
-
-[numthreads(THREADGROUP_SIZE, THREADGROUP_SIZE, THREADGROUP_SIZE)]
-void LightGridCompactCS(
-	uint3 GroupId : SV_GroupID,
-	uint3 DispatchThreadId : SV_DispatchThreadID,
-	uint3 GroupThreadId : SV_GroupThreadID,
-	uint  GroupIndex : SV_GroupIndex)
-{
-	uint3 GridCoordinate = DispatchThreadId;
-
-	bool bThreadValid = all(GridCoordinate < CulledGridSize);
-	bool bFirstThread =  GroupIndex==0;
-	
-	uint GridIndex = (GridCoordinate.z * CulledGridSize.y + GridCoordinate.y) * CulledGridSize.x + GridCoordinate.x;
-
-	// Compact lights
-	CompactReverseLinkedList(GridIndex, NumLocalLights, bFirstThread, bThreadValid, /*bLocalLights*/ true);
-
-	// Compact reflection captures
-	CompactReverseLinkedList(NumGridCells + GridIndex, NumReflectionCaptures, bFirstThread, bThreadValid, /*bLocalLights*/ false);
-}
-
-#endif // SHADER_LIGHT_GRID_COMPACT_CS
-
-
-
 #ifdef SHADER_DEBUG_LIGHT_GRID_PS
 
 #define SUPPORT_CONTACT_SHADOWS 0
diff --git a/Engine/Source/Runtime/Renderer/Private/LightGridInjection.cpp b/Engine/Source/Runtime/Renderer/Private/LightGridInjection.cpp
--- a/Engine/Source/Runtime/Renderer/Private/LightGridInjection.cpp
+++ b/Engine/Source/Runtime/Renderer/Private/LightGridInjection.cpp
@@ -224,7 +224,7 @@ public:
 		SHADER_PARAMETER_RDG_BUFFER_UAV(RWStructuredBuffer<uint>, RWCulledLightDataGrid32Bit)
 		SHADER_PARAMETER_RDG_BUFFER_UAV(RWBuffer<uint>, RWCulledLightDataGrid16Bit)
 		SHADER_PARAMETER_RDG_BUFFER_UAV(RWStructuredBuffer<uint>, RWNextCulledLightLink)
-		SHADER_PARAMETER_RDG_BUFFER_UAV(RWStructuredBuffer<uint>, RWStartOffsetGrid)
+		SHADER_PARAMETER_RDG_BUFFER_UAV(RWStructuredBuffer<uint>, RWNextCulledLightData)
 		SHADER_PARAMETER_RDG_BUFFER_UAV(RWStructuredBuffer<uint>, RWCulledLightLinks)
 		SHADER_PARAMETER_RDG_BUFFER_SRV(StructuredBuffer<float4>, LightViewSpacePositionAndRadius)
 		SHADER_PARAMETER_RDG_BUFFER_SRV(StructuredBuffer<float4>, LightViewSpaceDirAndPreprocAngle)
@@ -265,45 +265,6 @@ public:
 IMPLEMENT_GLOBAL_SHADER(FLightGridInjectionCS, "/Engine/Private/LightGridInjection.usf", "LightGridInjectionCS", SF_Compute);
 
 
-class FLightGridCompactCS : public FGlobalShader
-{
-	DECLARE_GLOBAL_SHADER(FLightGridCompactCS)
-	SHADER_USE_PARAMETER_STRUCT(FLightGridCompactCS, FGlobalShader)
-public:
-	BEGIN_SHADER_PARAMETER_STRUCT(FParameters, )
-		SHADER_PARAMETER_STRUCT_REF(FViewUniformShaderParameters, View)
-		SHADER_PARAMETER_RDG_BUFFER_UAV(RWStructuredBuffer<uint>, RWNumCulledLightsGrid)
-		SHADER_PARAMETER_RDG_BUFFER_UAV(RWStructuredBuffer<uint>, RWCulledLightDataGrid32Bit)
-		SHADER_PARAMETER_RDG_BUFFER_UAV(RWBuffer<uint>, RWCulledLightDataGrid16Bit)
-		SHADER_PARAMETER_RDG_BUFFER_UAV(RWStructuredBuffer<uint>, RWNextCulledLightData)
-		SHADER_PARAMETER_RDG_BUFFER_SRV(StructuredBuffer<uint>, StartOffsetGrid)
-		SHADER_PARAMETER_RDG_BUFFER_SRV(StructuredBuffer<uint>, CulledLightLinks)
-
-		SHADER_PARAMETER(FIntVector, CulledGridSize)
-		SHADER_PARAMETER(uint32, NumReflectionCaptures)
-		SHADER_PARAMETER(uint32, NumLocalLights)
-		SHADER_PARAMETER(uint32, NumGridCells)
-		SHADER_PARAMETER(uint32, MegaLightsSupportedStartIndex)
-	END_SHADER_PARAMETER_STRUCT()
-
-	static bool ShouldCompilePermutation(const FGlobalShaderPermutationParameters& Parameters)
-	{
-		return true;
-	}
-
-	static void ModifyCompilationEnvironment(const FGlobalShaderPermutationParameters& Parameters, FShaderCompilerEnvironment& OutEnvironment)
-	{
-		FGlobalShader::ModifyCompilationEnvironment(Parameters, OutEnvironment);
-		OutEnvironment.SetDefine(TEXT("SHADER_LIGHT_GRID_COMPACT_CS"), 1);
-		OutEnvironment.SetDefine(TEXT("THREADGROUP_SIZE"), LightGridInjectionGroupSize);
-		FForwardLightingParameters::ModifyCompilationEnvironment(Parameters.Platform, OutEnvironment);
-		OutEnvironment.SetDefine(TEXT("LIGHT_LINK_STRIDE"), LightLinkStride);
-		OutEnvironment.SetDefine(TEXT("MAX_CAPTURES"), GetMaxNumReflectionCaptures(Parameters.Platform));
-	}
-};
-
-IMPLEMENT_GLOBAL_SHADER(FLightGridCompactCS, "/Engine/Private/LightGridInjection.usf", "LightGridCompactCS", SF_Compute);
-
 /**
  */
 FORCEINLINE float GetTanRadAngleOrZero(float coneAngle)
@@ -925,7 +886,6 @@ FComputeLightGridOutput FSceneRenderer::ComputeLightGrid(FRDGBuilder& GraphBuild
 
 
 			FRDGBufferRef CulledLightLinksBuffer = GraphBuilder.CreateBuffer(FRDGBufferDesc::CreateStructuredDesc(sizeof(uint32), CulledLightLinksElements), TEXT("CulledLightLinks"));
-			FRDGBufferRef StartOffsetGridBuffer = GraphBuilder.CreateBuffer(FRDGBufferDesc::CreateStructuredDesc(sizeof(uint32), MaxNumCells), TEXT("StartOffsetGrid"));
 			FRDGBufferRef NextCulledLightLinkBuffer = GraphBuilder.CreateBuffer(FRDGBufferDesc::CreateStructuredDesc(sizeof(uint32), 1), TEXT("NextCulledLightLink"));
 			FRDGBufferRef NextCulledLightDataBuffer = GraphBuilder.CreateBuffer(FRDGBufferDesc::CreateStructuredDesc(sizeof(uint32), 1), TEXT("NextCulledLightData"));
 			FRDGBufferUAVRef NextCulledLightDataUAV = GraphBuilder.CreateUAV(NextCulledLightDataBuffer);
@@ -968,7 +928,7 @@ FComputeLightGridOutput FSceneRenderer::ComputeLightGrid(FRDGBuilder& GraphBuild
 			PassParameters->RWCulledLightDataGrid32Bit   = CulledLightDataGridUAV;
 			PassParameters->RWCulledLightDataGrid16Bit = CulledLightDataGridUAV;
 			PassParameters->RWNextCulledLightLink   = GraphBuilder.CreateUAV(NextCulledLightLinkBuffer);
-			PassParameters->RWStartOffsetGrid       = GraphBuilder.CreateUAV(StartOffsetGridBuffer);
+			PassParameters->RWNextCulledLightData   = NextCulledLightDataUAV;
 			PassParameters->RWCulledLightLinks      = GraphBuilder.CreateUAV(CulledLightLinksBuffer);
 			PassParameters->CulledGridSize          = ForwardLightData->CulledGridSize;
 			PassParameters->LightGridZParams        = ForwardLightData->LightGridZParams;
@@ -1007,32 +967,10 @@ FComputeLightGridOutput FSceneRenderer::ComputeLightGrid(FRDGBuilder& GraphBuild
 
 			if (GLightLinkedListCulling != 0)
 			{
-				AddClearUAVPass(GraphBuilder, PassParameters->RWStartOffsetGrid, 0xFFFFFFFF);
 				AddClearUAVPass(GraphBuilder, PassParameters->RWNextCulledLightLink, 0);
 				AddClearUAVPass(GraphBuilder, NextCulledLightDataUAV, 0);
 				AddClearUAVPass(GraphBuilder, NumCulledLightsGridUAV, 0);
 				FComputeShaderUtils::AddPass(GraphBuilder, RDG_EVENT_NAME("LightGridInject:LinkedList"), ComputeShader, PassParameters, NumGroups);
-
-				{
-					TShaderMapRef<FLightGridCompactCS> ComputeShaderCompact(View.ShaderMap);
-					FLightGridCompactCS::FParameters *PassParametersCompact = GraphBuilder.AllocParameters<FLightGridCompactCS::FParameters>();
-					PassParametersCompact->View = View.ViewUniformBuffer;
-
-					PassParametersCompact->CulledLightLinks = GraphBuilder.CreateSRV(CulledLightLinksBuffer);
-					PassParametersCompact->RWNumCulledLightsGrid = NumCulledLightsGridUAV;
-					PassParametersCompact->RWCulledLightDataGrid32Bit = CulledLightDataGridUAV;
-					PassParametersCompact->RWCulledLightDataGrid16Bit = CulledLightDataGridUAV;
-					PassParametersCompact->RWNextCulledLightData = NextCulledLightDataUAV;
-					PassParametersCompact->StartOffsetGrid = GraphBuilder.CreateSRV(StartOffsetGridBuffer);
-
-					PassParametersCompact->CulledGridSize = ForwardLightData->CulledGridSize;
-					PassParametersCompact->NumReflectionCaptures = ForwardLightData->NumReflectionCaptures;
-					PassParametersCompact->NumLocalLights = ForwardLightData->NumLocalLights;
-					PassParametersCompact->NumGridCells = ForwardLightData->NumGridCells;
-					PassParametersCompact->MegaLightsSupportedStartIndex = ForwardLightData->MegaLightsSupportedStartIndex;
-
-					Result.CompactLinksPass = FComputeShaderUtils::AddPass(GraphBuilder, RDG_EVENT_NAME("CompactLinks"), ComputeShaderCompact, PassParametersCompact, NumGroups);
-				}
 			}
 			else
 			{
