Post Process Alpha is preserved in SceneColor where possible for post process materials that don't write alpha.  This includes making it so materials that read and write from SceneColor, but don't write alpha, automatically read and write alpha.  For materials that only write to SceneColor, a Color Write mask is set to mask the alpha channel in the destination.  An important bug with UserSceneTextures was also fixed where a final pass of a post process chain wasn't properly identified if it wrote to a UserSceneTexture overridden to "SceneColor".

Propagated alpha is handled by storing the alpha to a local variable when PostProcessInput0 (SceneColor) is fetched from either a SceneTexture or UserSceneTexture node, then writing the fetched value to the output.  The propagation of alpha is free in any material which reads from a SceneTexture, but an extra scalar and vector ALU is generated per UserSceneTexture, since the alpha propagation is conditional based on whether a given input has SceneColor assigned as its input (based on the value of the UserSceneTextureSceneColorInput parameter).  The first UserSceneTexture vector ALU will be free, rolled into the output assignment.  The material can be set up to output opacity to optimize that minor overhead out completely if desired.

An advantage of this approach is that it even works with custom HLSL, since it just requires SceneTextureLookup to be reached by any means, and it also straightforwardly handles UserSceneTextures reading and writing SceneColor.  Some complexity is added in that SceneTexture or UserSceneTexture nodes feeding into custom HLSL, where the input pin value itself isn't used, now need to be explicitly dead stripped, as the side effect of writing to PropagatedAlpha prevents the shader compiler from optimizing out calls to SceneTextureLookup based on the return value being unused in downstream code.  See UMaterialExpressionCustom::Compile for where that is handled (with detailed comments).

diff --git a/Engine/Shaders/Private/MaterialTemplate.ush b/Engine/Shaders/Private/MaterialTemplate.ush
--- a/Engine/Shaders/Private/MaterialTemplate.ush
+++ b/Engine/Shaders/Private/MaterialTemplate.ush
@@ -80,6 +80,11 @@
 #define PPI_WorldTangent 29
 #define PPI_Anisotropy 30
 
+#ifdef POST_PROCESS_PROPAGATE_ALPHA_INPUT
+// Arbitrary internal value to fetch local variable propagated alpha stored in the SceneTextureLookup function
+#define PPI_PropagatedAlpha 64
+#endif
+
 // Defines to remap PPI_UserSceneTexture0-6 to PPI_PostProcessInput0-6, generated via GenerateUserSceneTextureRemapDefines
 %{user_scene_texture_remap}
 
@@ -2935,6 +2940,11 @@ uint bMetalMSAAHDRDecode;
 
 #if NEEDS_SCENE_TEXTURES
 
+#if POST_PROCESS_MATERIAL
+// SceneTextureId of UserSceneTexture SceneColor input, to allow SceneColor alpha to be dynamically propagated
+uint UserSceneTextureSceneColorInput;
+#endif
+
 #if SHADING_PATH_MOBILE
 
 MaterialFloat4 MobileSceneTextureLookup(inout FMaterialPixelParameters Parameters, int SceneTextureId, float2 UV)
@@ -2953,6 +2963,11 @@ MaterialFloat4 MobileSceneTextureLookup(inout FMaterialPixelParameters Parameter
 		return 0.0f;
 	}
 #else
+
+	#ifdef POST_PROCESS_PROPAGATE_ALPHA_INPUT
+	static float PropagatedAlpha = 0.0f;
+	#endif
+
 	const float2 PixelPos = UV * View.BufferSizeAndInvSize.xy;
 	FGBufferData GBuffer = MobileFetchAndDecodeGBuffer(UV, PixelPos, GetViewId(Parameters));
 
@@ -2997,6 +3012,15 @@ MaterialFloat4 MobileSceneTextureLookup(inout FMaterialPixelParameters Parameter
 			Input0.rgb = Input0.rgb * rcp(Input0.r*(-0.299) + Input0.g*(-0.587) + Input0.b*(-0.114) + 1.0);
 		}
 		#endif
+		#ifdef POST_PROCESS_PROPAGATE_ALPHA_INPUT
+		// We do a comparison of the POST_PROCESS_PROPAGATE_ALPHA_INPUT define with the PPI input index, to support propagation of alpha for shaders with UserSceneTextures,
+		// where the specific PPI_PostProcessInput that has SceneColor as an input is dynamic.  Usually, POST_PROCESS_PROPAGATE_ALPHA_INPUT will be a constant, when a normal
+		// SceneTexture node is used.
+		if (POST_PROCESS_PROPAGATE_ALPHA_INPUT == PPI_PostProcessInput0)
+		{
+			PropagatedAlpha = Input0.a;
+		}
+		#endif
 		return Input0;
 	}
 #endif // POST_PROCESS_MATERIAL
@@ -3012,6 +3036,10 @@ MaterialFloat4 MobileSceneTextureLookup(inout FMaterialPixelParameters Parameter
 		return GBuffer.CustomStencil;
 	case PPI_WorldTangent:
 		return MaterialFloat4(GBuffer.WorldTangent, 0);
+#ifdef POST_PROCESS_PROPAGATE_ALPHA_INPUT
+	case PPI_PropagatedAlpha:
+		return MaterialFloat4(PropagatedAlpha, 0, 0, 0);
+#endif
 	default:
 		return MaterialFloat4(0, 0, 0, 0);
 	}
@@ -3080,6 +3108,42 @@ float4 SceneTextureLookup(float2 UV, int SceneTextureIndex, bool bFiltered)
 
 	const uint2 PixelPos = uint2(UV * View.BufferSizeAndInvSize.xy);
 
+#ifdef POST_PROCESS_PROPAGATE_ALPHA_INPUT
+	// Propagation of SceneColor alpha, to preserve it in materials that don't themselves write to alpha, is supported by storing the propagated alpha to a local variable,
+	// which can be read from with a call to this function at the end of the shader.  The big advantage of this approach is that it even works with Custom HLSL, since it
+	// just requires the SceneTextureLookup code to be reached by any means.  It also works well with UserSceneTexture fetches, where a given input can be dynamically set
+	// to SceneColor.  Some complexity is added in that SceneTexture or UserSceneTexture nodes feeding into custom HLSL, where the input pin value itself isn't used, now
+	// need to be explicitly dead stripped, as the side effect of writing to PropagatedAlpha prevents the shader compiler from optimizing out calls to SceneTextureLookup
+	// based on the return value being unused in downstream code.  See UMaterialExpressionCustom::Compile for where that is handled.
+	static float PropagatedAlpha = 0.0f;
+
+	// There are three parts to the condition to detect whether we should automatically propagate SceneColor alpha:
+	//		1.  To make the shader compiler optimize out repeated reads of alpha, we set a local bool to ensure we only read alpha from a given scene texture index once.
+	//		2.  We want to support propagation of alpha for UserSceneTextures.  A given input is a UserSceneTexture if it's not used as a regular SceneTexture, as
+	//			specified by the POST_PROCESS_USED_SCENE_TEXTURES mask compile define, so we skip propagation for any input with a bit set there.
+	//			PPI_PostProcessInput0 is an exception, which is the SceneColor itself when used as a SceneTexture, and should always support propagation.
+	//		3.  Finally, we test if the given input equals POST_PROCESS_PROPAGATE_ALPHA_INPUT, which will either be PPI_PostProcessInput0 if the material in question
+	//			directly uses that as a SceneTexture, or the dynamic UserSceneTextureSceneColorInput variable, if UserSceneTexture nodes are in use, and one of them
+	//			may be mapped to SceneColor.
+	// If PostProcessInput0 is directly fetched in a SceneTexture node, the conditional will be compile time invariant, and no extra code is generated, making the alpha
+	// propagation free.  Otherwise, for UserSceneTextures, only the final part of the comparison is dynamic, and the rest will be compile time invariant, so it adds one
+	// scalar ALU (equals) and one vector ALU (conditional move), although the first UserSceneTexture vector ALU is free, rolled into the output assignment.  The material
+	// can be set up to output opacity to optimize that minor overhead out completely if desired.
+	//
+	#define RETURN_POST_PROCESS_INPUT_SAMPLE(Index) {  \
+			float4 Result = Texture2DSample(PostProcessInput_##Index##_Texture, bFiltered ? PostProcessInput_BilinearSampler : PostProcessInput_##Index##_SharedSampler, UV);  \
+			static bool bFirstRead = true;  \
+			if (bFirstRead && !(POST_PROCESS_USED_SCENE_TEXTURES & (1u << (Index + PPI_PostProcessInput0)) & ~(1u<<PPI_PostProcessInput0)) && POST_PROCESS_PROPAGATE_ALPHA_INPUT == PPI_PostProcessInput##Index)  \
+			{ PropagatedAlpha = Result.a; }  \
+			bFirstRead = false;  \
+			return Result;  \
+		}
+
+#else
+	#define RETURN_POST_PROCESS_INPUT_SAMPLE(Index) \
+		return Texture2DSample(PostProcessInput_##Index##_Texture, bFiltered ? PostProcessInput_BilinearSampler : PostProcessInput_##Index##_SharedSampler, UV)
+#endif
+
 	// 1. Post-process inputs which are independent on material data
 	switch (SceneTextureIndex)
 	{
@@ -3088,18 +3152,21 @@ float4 SceneTextureLookup(float2 UV, int SceneTextureIndex, bool bFiltered)
 		case PPI_SeparateTranslucency:	return float4(1, 1, 1, 1);	// todo
 		case PPI_CustomDepth:			return CalcSceneCustomDepth(UV);
 	#if POST_PROCESS_MATERIAL
-		case PPI_PostProcessInput0:		return Texture2DSample(PostProcessInput_0_Texture, bFiltered ? PostProcessInput_BilinearSampler : PostProcessInput_0_SharedSampler, UV);
-		case PPI_PostProcessInput1:		return Texture2DSample(PostProcessInput_1_Texture, bFiltered ? PostProcessInput_BilinearSampler : PostProcessInput_1_SharedSampler, UV);
-		case PPI_PostProcessInput2:		return Texture2DSample(PostProcessInput_2_Texture, bFiltered ? PostProcessInput_BilinearSampler : PostProcessInput_2_SharedSampler, UV);
-		case PPI_PostProcessInput3:		return Texture2DSample(PostProcessInput_3_Texture, bFiltered ? PostProcessInput_BilinearSampler : PostProcessInput_3_SharedSampler, UV);
-		case PPI_PostProcessInput4:		return Texture2DSample(PostProcessInput_4_Texture, bFiltered ? PostProcessInput_BilinearSampler : PostProcessInput_4_SharedSampler, UV);
-	#endif // __POST_PROCESS_COMMON__
+		case PPI_PostProcessInput0:		RETURN_POST_PROCESS_INPUT_SAMPLE(0);
+		case PPI_PostProcessInput1:		RETURN_POST_PROCESS_INPUT_SAMPLE(1);
+		case PPI_PostProcessInput2:		RETURN_POST_PROCESS_INPUT_SAMPLE(2);
+		case PPI_PostProcessInput3:		RETURN_POST_PROCESS_INPUT_SAMPLE(3);
+		case PPI_PostProcessInput4:		RETURN_POST_PROCESS_INPUT_SAMPLE(4);
+	#endif // POST_PROCESS_MATERIAL
 		case PPI_DecalMask:				return 0;  // material compiler will return an error
 		case PPI_AmbientOcclusion:		return CalcSceneAO(UV);
 		case PPI_CustomStencil:			return CalcSceneCustomStencil(PixelPos);
 	#if POST_PROCESS_MATERIAL
 		case PPI_Velocity:				return float4(PostProcessVelocityLookup(ConvertToDeviceZ(CalcSceneDepth(UV)), UV), 0, 0);
 	#endif
+	#ifdef POST_PROCESS_PROPAGATE_ALPHA_INPUT
+		case PPI_PropagatedAlpha:		return float4(PropagatedAlpha, 0, 0, 0);
+	#endif
 	}
 
 	// 2. Post-process inputs which are dependent on material data
diff --git a/Engine/Shaders/Private/PostProcessMaterialShaders.usf b/Engine/Shaders/Private/PostProcessMaterialShaders.usf
--- a/Engine/Shaders/Private/PostProcessMaterialShaders.usf
+++ b/Engine/Shaders/Private/PostProcessMaterialShaders.usf
@@ -350,6 +350,10 @@ void MainPS(
 	EmissiveColor = GetMaterialEmissive(PixelMaterialInputs);
 #endif // SUBSTRATE_ENABLED
 
+#ifdef POST_PROCESS_PROPAGATE_ALPHA_INPUT
+	Alpha = SceneTextureLookup(float2(0,0), PPI_PropagatedAlpha, false).r;
+#endif
+
 	OutColor = float4(EmissiveColor, Alpha );
 
 #if !MATERIALBLENDING_MODULATE && POST_PROCESS_MATERIAL_BEFORE_TONEMAP && !POST_PROCESS_DISABLE_PRE_EXPOSURE_SCALE
diff --git a/Engine/Source/Runtime/Engine/Private/Materials/HLSLMaterialTranslator.cpp b/Engine/Source/Runtime/Engine/Private/Materials/HLSLMaterialTranslator.cpp
--- a/Engine/Source/Runtime/Engine/Private/Materials/HLSLMaterialTranslator.cpp
+++ b/Engine/Source/Runtime/Engine/Private/Materials/HLSLMaterialTranslator.cpp
@@ -259,6 +259,9 @@ struct FHLSLMaterialTranslator::FEnvironmentDefines
 	bool bNeedsParticleSize;
 	bool bNeedsParticleSpriteRotation;
 	bool bNeedsSceneTextures;
+	bool bAlphaPropagatePostProcessInput0;
+	bool bAlphaPropagateUserSceneTexture;
+	uint32 UsedSceneTextures;
 	bool bUsesEyeAdaptation;
 	bool bVirtualTextureOutput;
 	bool bUsesPerInstanceCustomData;
@@ -355,6 +358,9 @@ struct FHLSLMaterialTranslator::FEnvironmentDefines
 		Ar << bNeedsParticleSize;
 		Ar << bNeedsParticleSpriteRotation;
 		Ar << bNeedsSceneTextures;
+		Ar << bAlphaPropagatePostProcessInput0;
+		Ar << bAlphaPropagateUserSceneTexture;
+		Ar << UsedSceneTextures;
 		Ar << bUsesEyeAdaptation;
 		Ar << bVirtualTextureOutput;
 		Ar << bUsesPerInstanceCustomData;
@@ -2634,6 +2640,12 @@ void FHLSLMaterialTranslator::GetMaterialEnvironment(EShaderPlatform InPlatform,
 		OutEnvironment.SetDefine(TEXT("NEEDS_SCENE_TEXTURES"), TEXT("1"));
 	}
 
+	if (EnvironmentDefines->bAlphaPropagatePostProcessInput0 || EnvironmentDefines->bAlphaPropagateUserSceneTexture)
+	{
+		OutEnvironment.SetDefine(TEXT("POST_PROCESS_PROPAGATE_ALPHA_INPUT"), EnvironmentDefines->bAlphaPropagatePostProcessInput0 ? TEXT("PPI_PostProcessInput0") : TEXT("UserSceneTextureSceneColorInput"));
+		OutEnvironment.SetDefine(TEXT("POST_PROCESS_USED_SCENE_TEXTURES"), EnvironmentDefines->UsedSceneTextures);
+	}
+
 	if (EnvironmentDefines->bUsesEyeAdaptation)
 	{
 		OutEnvironment.SetDefine(TEXT("USES_EYE_ADAPTATION"), TEXT("1"));
@@ -7930,7 +7942,7 @@ int32 FHLSLMaterialTranslator::SceneDepth(int32 Offset, int32 ViewportUV, bool b
 }
 
 // @param SceneTextureId of type ESceneTextureId e.g. PPI_SubsurfaceColor
-int32 FHLSLMaterialTranslator::SceneTextureLookup(int32 ViewportUV, uint32 InSceneTextureId, bool bFiltered, bool bClamped)
+int32 FHLSLMaterialTranslator::SceneTextureLookup(int32 ViewportUV, uint32 InSceneTextureId, bool bFiltered, bool bClamped, bool bUnused)
 {
 	ESceneTextureId SceneTextureId = (ESceneTextureId)InSceneTextureId;
 
@@ -7954,6 +7966,14 @@ int32 FHLSLMaterialTranslator::SceneTextureLookup(int32 ViewportUV, uint32 InSce
 
 	UseSceneTextureId(SceneTextureId, true);
 
+	if (bUnused)
+	{
+		// If the output is unused, we just need to reach the call to UseSceneTextureId above so the shader compiler includes the scene texture, then
+		// emit an arbitrary unused constant expression.  This code path is for Custom HLSL references to scene textures, where custom code is calling
+		// SceneTextureFetch or SceneTextureLookup, and the result of the input pin itself isn't used.
+		return Constant4(0.0f, 0.0f, 0.0f, 0.0f);
+	}
+
 	FString SceneTextureIdString = SceneTextureIdToHLSLString(SceneTextureId);
 
 	int32 BufferUV;
@@ -14836,7 +14856,9 @@ static const TCHAR* FindHlslIdentifierInCode(const TCHAR*& Code, const TCHAR* Id
 
 // Fixup for SceneTexture and UserSceneTexture inputs to custom HLSL.  Returns a new string, rather than modifying a Code FString
 // in place, to allow the code to be shared with the HLSL tree code path (MaterialExpressionHLSL.cpp), which uses TStringBuilder.
-// An empty string is returned if no fixup was required (the common case), avoiding reallocation.
+// An empty string is returned if no fixup was required (the common case), avoiding reallocation.  OutSceneTextureInfo tracks
+// which inputs are SceneTexture or UserSceneTexture (left empty if no scene textures).  A -1 indicates a scene texture input not
+// directly used as a symbol in the custom HLSL, allowing the fetch to be dead stripped, while 1 is an actually used input.
 //
 // The goal is to allow the scene texture ID to be automatically generated in the code, without the user needing to manually
 // insert the correct identifier based on settings in the connected node.  Without this fixup, the user needs to change the
@@ -14874,13 +14896,15 @@ static const TCHAR* FindHlslIdentifierInCode(const TCHAR*& Code, const TCHAR* Id
 //		#define WEIGHTED_FETCH_SCENE_TEXTURE(SampleIndex) \
 //			(InputName.Fetch(Offsets[SampleIndex], 0) * Weights[SampleIndex])
 //
-FString CustomExpressionSceneTextureInputFixup(const UMaterialExpressionCustom* Custom, const TCHAR* Code)
+FString CustomExpressionSceneTextureInputFixup(const UMaterialExpressionCustom* Custom, const TCHAR* Code, TArray<int8>& OutSceneTextureInfo)
 {
 	// Allocated when required fixup is first encountered
 	FString ModifiedCode;
 
-	for (const FCustomInput& Input : Custom->Inputs)
+	for (int32 InputIndex = 0; InputIndex < Custom->Inputs.Num(); InputIndex++)
 	{
+		const FCustomInput& Input = Custom->Inputs[InputIndex];
+
 		if (Input.InputName.IsNone())
 		{
 			continue;
@@ -14906,6 +14930,15 @@ FString CustomExpressionSceneTextureInputFixup(const UMaterialExpressionCustom*
 			InputIDString = TEXT("PPIUser_") + UserSceneTextureExpression->UserSceneTexture.ToString();
 		}
 
+		// We've encountered a scene texture, allocate our info array
+		if (!OutSceneTextureInfo.Num())
+		{
+			OutSceneTextureInfo.SetNumZeroed(Custom->Inputs.Num());
+		}
+
+		// Mark with -1 to indicate that it's a scene texture input, but might not be directly used in the source.
+		OutSceneTextureInfo[InputIndex] = -1;
+
 		int32 CodeOffset = 0;
 		while (1)
 		{
@@ -14964,6 +14997,9 @@ FString CustomExpressionSceneTextureInputFixup(const UMaterialExpressionCustom*
 			if (CodeOffset == OriginalCodeOffset)
 			{
 				CodeOffset = CodeString - CodeStart;
+
+				// Scene Texture input pin symbol was directly used by the custom HLSL, mark it with a 1, as the input pin expression needs to be compiled in
+				OutSceneTextureInfo[InputIndex] = 1;
 			}
 		}
 	}
@@ -15178,7 +15214,8 @@ int32 FHLSLMaterialTranslator::CustomExpression( class UMaterialExpressionCustom
 		}
 		Code.ReplaceInline(TEXT("\r\n"), TEXT("\n"), ESearchCase::CaseSensitive);
 
-		FString ModifiedCode = CustomExpressionSceneTextureInputFixup(Custom, *Code);
+		TArray<int8> SceneTextureInfoIgnored;
+		FString ModifiedCode = CustomExpressionSceneTextureInputFixup(Custom, *Code, SceneTextureInfoIgnored);
 		if (!ModifiedCode.IsEmpty())
 		{
 			Code = ModifiedCode;
@@ -15681,6 +15718,9 @@ void FHLSLMaterialTranslator::PrepareEnvironmentDefines()
 	EnvironmentDefines->bNeedsParticleSize = bNeedsParticleSize;
 	EnvironmentDefines->bNeedsParticleSpriteRotation = bNeedsParticleSpriteRotation;
 	EnvironmentDefines->bNeedsSceneTextures = MaterialCompilationOutput.bNeedsSceneTextures;
+	EnvironmentDefines->bAlphaPropagatePostProcessInput0 = Material->GetMaterialDomain() == MD_PostProcess && !Material->GetBlendableOutputAlpha() && MaterialCompilationOutput.IsSceneTextureUsed(PPI_PostProcessInput0);
+	EnvironmentDefines->bAlphaPropagateUserSceneTexture = Material->GetMaterialDomain() == MD_PostProcess && !Material->GetBlendableOutputAlpha() && !MaterialCompilationOutput.UserSceneTextureInputs.IsEmpty();
+	EnvironmentDefines->UsedSceneTextures = MaterialCompilationOutput.UsedSceneTextures;
 	EnvironmentDefines->bUsesEyeAdaptation = MaterialCompilationOutput.bUsesEyeAdaptation;
 	EnvironmentDefines->bVirtualTextureOutput = MaterialCompilationOutput.bHasRuntimeVirtualTextureOutputNode;
 	EnvironmentDefines->bUsesPerInstanceCustomData = MaterialCompilationOutput.bUsesPerInstanceCustomData && Material->IsUsedWithInstancedStaticMeshes();
diff --git a/Engine/Source/Runtime/Engine/Private/Materials/HLSLMaterialTranslator.h b/Engine/Source/Runtime/Engine/Private/Materials/HLSLMaterialTranslator.h
--- a/Engine/Source/Runtime/Engine/Private/Materials/HLSLMaterialTranslator.h
+++ b/Engine/Source/Runtime/Engine/Private/Materials/HLSLMaterialTranslator.h
@@ -1060,7 +1060,7 @@ protected:
 	virtual int32 SceneDepth(int32 Offset, int32 ViewportUV, bool bUseOffset) override;
 	
 	// @param SceneTextureId of type ESceneTextureId e.g. PPI_SubsurfaceColor
-	virtual int32 SceneTextureLookup(int32 ViewportUV, uint32 InSceneTextureId, bool bFiltered, bool bClamped) override;
+	virtual int32 SceneTextureLookup(int32 ViewportUV, uint32 InSceneTextureId, bool bFiltered, bool bClamped, bool bUnused) override;
 
 	virtual int32 GetSceneTextureViewSize(int32 SceneTextureId, bool InvProperty) override;
 
diff --git a/Engine/Source/Runtime/Engine/Private/Materials/MaterialExpressionHLSL.cpp b/Engine/Source/Runtime/Engine/Private/Materials/MaterialExpressionHLSL.cpp
--- a/Engine/Source/Runtime/Engine/Private/Materials/MaterialExpressionHLSL.cpp
+++ b/Engine/Source/Runtime/Engine/Private/Materials/MaterialExpressionHLSL.cpp
@@ -4615,7 +4615,7 @@ bool UMaterialExpressionObjectRadius::GenerateHLSLExpression(FMaterialHLSLGenera
 	return true;
 }
 
-extern FString CustomExpressionSceneTextureInputFixup(const UMaterialExpressionCustom* Custom, const TCHAR* Code);
+extern FString CustomExpressionSceneTextureInputFixup(const UMaterialExpressionCustom* Custom, const TCHAR* Code, TArray<int8>& OutSceneTextureInfo);
 
 bool UMaterialExpressionCustom::GenerateHLSLExpression(FMaterialHLSLGenerator& Generator, UE::HLSLTree::FScope& Scope, int32 OutputIndex, UE::HLSLTree::FExpression const*& OutExpression) const
 {
@@ -4691,7 +4691,8 @@ bool UMaterialExpressionCustom::GenerateHLSLExpression(FMaterialHLSLGenerator& G
 		FormattedCode.Appendf(TEXT("return %s;"), *Code);
 	}
 
-	FString SceneTextureFixupCode = CustomExpressionSceneTextureInputFixup(this, FormattedCode.Len() ? FormattedCode.ToString() : *Code);
+	TArray<int8> SceneTextureInfoIgnored;
+	FString SceneTextureFixupCode = CustomExpressionSceneTextureInputFixup(this, FormattedCode.Len() ? FormattedCode.ToString() : *Code, SceneTextureInfoIgnored);
 
 	FStringView FunctionCode;
 	if (SceneTextureFixupCode.Len())
diff --git a/Engine/Source/Runtime/Engine/Private/Materials/MaterialExpressions.cpp b/Engine/Source/Runtime/Engine/Private/Materials/MaterialExpressions.cpp
--- a/Engine/Source/Runtime/Engine/Private/Materials/MaterialExpressions.cpp
+++ b/Engine/Source/Runtime/Engine/Private/Materials/MaterialExpressions.cpp
@@ -12592,13 +12592,15 @@ int32 UMaterialExpressionSceneTexture::Compile(class FMaterialCompiler* Compiler
 		ViewportUV = Coordinates.Compile(Compiler);
 	}
 
-	if(OutputIndex == 0)
+	if(OutputIndex == 0 || OutputIndex == 3)
 	{
 		// Color.  Note that clamping support is not necessary for regular SceneTexture, because it's only useful when sampling from lower resolution
 		// maps with filtering, where bilinear blending of a higher resolution UV sample can end up interpolating with pixels outside the valid UV
 		// range on a lower resolution map.  All SceneTextures are full resolution, while UserSceneTextures can be lower resolution (see
-		// UMaterialExpressionUserSceneTexture::Compile below), so those support a user specified clamp flag.
-		return Compiler->SceneTextureLookup(ViewportUV, SceneTextureId, bFiltered, /*bClamped=*/ false);
+		// UMaterialExpressionUserSceneTexture::Compile below), so those support a user specified clamp flag.  The special OutputIndex of 3 (not
+		// user facing) indicates an input pin to custom HLSL that isn't used in the code, meaning the scene texture input should be compiled in,
+		// but the input pin's expression should be dead stripped to avoid an unnecessary texture fetch.
+		return Compiler->SceneTextureLookup(ViewportUV, SceneTextureId, bFiltered, /*bClamped=*/ false, /*bUnused=*/ OutputIndex == 3);
 	}
 	else if(OutputIndex == 1 || OutputIndex == 2)
 	{
@@ -12676,10 +12678,11 @@ int32 UMaterialExpressionUserSceneTexture::Compile(class FMaterialCompiler* Comp
 		ViewportUV = Coordinates.Compile(Compiler);
 	}
 
-	if (OutputIndex == 0)
+	if (OutputIndex == 0 || OutputIndex == 3)
 	{
-		// Color
-		return Compiler->SceneTextureLookup(ViewportUV, SceneTextureId, bFiltered, bClamped);
+		// Color.    The special OutputIndex of 3 (not user facing) indicates an input pin to custom HLSL that isn't used in the code, meaning the
+		// scene texture input should be compiled in, but the input pin's expression should be dead stripped to avoid an unnecessary texture fetch.
+		return Compiler->SceneTextureLookup(ViewportUV, SceneTextureId, bFiltered, bClamped, /*bUnused=*/ OutputIndex == 3);
 	}
 	else if (OutputIndex == 1 || OutputIndex == 2)
 	{
@@ -15383,10 +15386,19 @@ UMaterialExpressionCustom::UMaterialExpressionCustom(const FObjectInitializer& O
 }
 
 #if WITH_EDITOR
+extern FString CustomExpressionSceneTextureInputFixup(const UMaterialExpressionCustom* Custom, const TCHAR* Code, TArray<int8>& OutSceneTextureInfo);
+
 int32 UMaterialExpressionCustom::Compile(class FMaterialCompiler* Compiler, int32 OutputIndex)
 {
 	TArray<int32> CompiledInputs;
 
+	// We're not using the fixed up code here, just the SceneTextureInfo, which tracks whether the value of SceneTexture / UserSceneTexture input pins are
+	// used in the custom HLSL code.  In many cases, the scene textures will be fetched from using SceneTextureLookup, SceneTextureFetch, or *.Fetch calls
+	// in the custom HLSL, rather than using the input pin value.  The fixup function has a parser that is aware of HLSL syntax, and able to tokenize out
+	// identifiers, and handle symbol sequences (such as Input.ID or Input.Fetch) that will be substituted into SceneTextureFetch calls.
+	TArray<int8> SceneTextureInfo;
+	CustomExpressionSceneTextureInputFixup(this, *Code, SceneTextureInfo);
+
 	for( int32 i=0;i<Inputs.Num();i++ )
 	{
 		// skip over unnamed inputs
@@ -15400,7 +15412,24 @@ int32 UMaterialExpressionCustom::Compile(class FMaterialCompiler* Compiler, int3
 			{
 				return Compiler->Errorf(TEXT("Custom material %s missing input %d (%s)"), *Description, i+1, *Inputs[i].InputName.ToString());
 			}
-			int32 InputCode = Inputs[i].Input.Compile(Compiler);
+
+			int32 InputCode;
+			if (SceneTextureInfo.Num() && SceneTextureInfo[i] == -1)
+			{
+				// Scene texture reference, not actually used in the custom HLSL.  The special output index "3" (not present in the user interface) specifies
+				// that the scene texture should be compiled into the shader for use by custom HLSL, but the input pin value is not actually used in code, so
+				// its expression shouldn't be compiled in.  It's necessary to explicitly remove the input pin's SceneTexture expression, because the SceneColor
+				// alpha propagation feature means the SceneTextureLookup function now has a side effect of caching propagated alpha, and the compiler can no
+				// longer dead strip calls to that function.  This early removal of the fetch makes stats more accurate as a side bonus.
+				FExpressionInput LocalInput = Inputs[i].Input;
+				LocalInput.OutputIndex = 3;
+				InputCode = LocalInput.Compile(Compiler);
+			}
+			else
+			{
+				InputCode = Inputs[i].Input.Compile(Compiler);
+			}
+
 			if( InputCode < 0 )
 			{
 				return InputCode;
diff --git a/Engine/Source/Runtime/Engine/Private/Materials/MaterialHLSLEmitter.cpp b/Engine/Source/Runtime/Engine/Private/Materials/MaterialHLSLEmitter.cpp
--- a/Engine/Source/Runtime/Engine/Private/Materials/MaterialHLSLEmitter.cpp
+++ b/Engine/Source/Runtime/Engine/Private/Materials/MaterialHLSLEmitter.cpp
@@ -432,6 +432,14 @@ static void GetMaterialEnvironment(EShaderPlatform InPlatform,
 	{
 		OutEnvironment.SetDefine(TEXT("NEEDS_SCENE_TEXTURES"), TEXT("1"));
 	}
+
+	// If post process material doesn't already output alpha, and reads from PostProcessInput0 or a UserSceneTexture, enable automatic alpha propagation for SceneColor
+	if (InMaterial.GetMaterialDomain() == MD_PostProcess && InMaterial.GetBlendableOutputAlpha() == false && (MaterialCompilationOutput.IsSceneTextureUsed(PPI_PostProcessInput0) || !MaterialCompilationOutput.UserSceneTextureInputs.IsEmpty()))
+	{
+		OutEnvironment.SetDefine(TEXT("POST_PROCESS_PROPAGATE_ALPHA_INPUT"), MaterialCompilationOutput.IsSceneTextureUsed(PPI_PostProcessInput0) ? TEXT("PPI_PostProcessInput0") : TEXT("UserSceneTextureSceneColorInput"));
+		OutEnvironment.SetDefine(TEXT("POST_PROCESS_USED_SCENE_TEXTURES"), MaterialCompilationOutput.UsedSceneTextures);
+	}
+
 	if (MaterialCompilationOutput.bUsesEyeAdaptation)
 	{
 		OutEnvironment.SetDefine(TEXT("USES_EYE_ADAPTATION"), TEXT("1"));
diff --git a/Engine/Source/Runtime/Engine/Public/MaterialCompiler.h b/Engine/Source/Runtime/Engine/Public/MaterialCompiler.h
--- a/Engine/Source/Runtime/Engine/Public/MaterialCompiler.h
+++ b/Engine/Source/Runtime/Engine/Public/MaterialCompiler.h
@@ -436,7 +436,7 @@ public:
 	virtual int32 SceneDepth(int32 Offset, int32 ViewportUV, bool bUseOffset) = 0;
 	virtual int32 SceneColor(int32 Offset, int32 ViewportUV, bool bUseOffset) = 0;
 	// @param SceneTextureId of type ESceneTextureId e.g. PPI_SubsurfaceColor
-	virtual int32 SceneTextureLookup(int32 ViewportUV, uint32 SceneTextureId, bool bFiltered, bool bClamped) = 0;
+	virtual int32 SceneTextureLookup(int32 ViewportUV, uint32 SceneTextureId, bool bFiltered, bool bClamped, bool bUnused) = 0;
 	virtual int32 GetSceneTextureViewSize(int32 SceneTextureId, bool InvProperty) = 0;
 	virtual int32 FindOrAddUserSceneTexture(FName UserSceneTextureName) = 0;
 	virtual int32 DBufferTextureLookup(int32 ViewportUV, uint32 DBufferTextureIndex) = 0;
@@ -990,7 +990,7 @@ public:
 	virtual	int32 PixelDepth() override { return Compiler->PixelDepth(); }
 	virtual int32 SceneDepth(int32 Offset, int32 ViewportUV, bool bUseOffset) override { return Compiler->SceneDepth(Offset, ViewportUV, bUseOffset); }
 	virtual int32 SceneColor(int32 Offset, int32 ViewportUV, bool bUseOffset) override { return Compiler->SceneColor(Offset, ViewportUV, bUseOffset); }
-	virtual int32 SceneTextureLookup(int32 ViewportUV, uint32 InSceneTextureId, bool bFiltered, bool bClamped) override { return Compiler->SceneTextureLookup(ViewportUV, InSceneTextureId, bFiltered, bClamped); }
+	virtual int32 SceneTextureLookup(int32 ViewportUV, uint32 InSceneTextureId, bool bFiltered, bool bClamped, bool bUnused) override { return Compiler->SceneTextureLookup(ViewportUV, InSceneTextureId, bFiltered, bClamped, bUnused); }
 	virtual int32 GetSceneTextureViewSize(int32 SceneTextureId, bool InvProperty) override { return Compiler->GetSceneTextureViewSize(SceneTextureId, InvProperty); }
 	virtual int32 FindOrAddUserSceneTexture(FName UserSceneTextureName) override { return Compiler->FindOrAddUserSceneTexture(UserSceneTextureName); }
 	virtual int32 DBufferTextureLookup(int32 ViewportUV, uint32 DBufferTextureIndex) override { return Compiler->DBufferTextureLookup(ViewportUV, DBufferTextureIndex); }
diff --git a/Engine/Source/Runtime/Renderer/Private/PostProcess/PostProcessMaterial.cpp b/Engine/Source/Runtime/Renderer/Private/PostProcess/PostProcessMaterial.cpp
--- a/Engine/Source/Runtime/Renderer/Private/PostProcess/PostProcessMaterial.cpp
+++ b/Engine/Source/Runtime/Renderer/Private/PostProcess/PostProcessMaterial.cpp
@@ -536,6 +536,10 @@ FPostProcessMaterialParameters* GetPostProcessMaterialParameters(
 
 	PostProcessMaterialParameters->NeuralPostProcessParameters = GetDefaultNeuralPostProcessShaderParameters(GraphBuilder);
 
+	// UserSceneTextureSceneColorInput is used for automatic scene color alpha propagation.  Alpha propagation only occurs if the output is scene color
+	// (meaning not a user scene texture), so set this to INDEX_NONE if writing to a UserSceneTexture output instead.
+	PostProcessMaterialParameters->UserSceneTextureSceneColorInput = Inputs.bUserSceneTextureOutput ? INDEX_NONE : Inputs.UserSceneTextureSceneColorInput;
+
 	return PostProcessMaterialParameters;
 }
 
@@ -764,14 +768,18 @@ FScreenPassTexture AddPostProcessMaterialPass(
 	FScreenPassRenderTarget Output = Inputs.OverrideOutput;
 
 	// We can re-use the scene color texture as the render target if we're not simultaneously reading from it.
-	// This is only necessary to do if we're going to be priming content from the render target since it avoids
-	// the copy. Otherwise, we just allocate a new render target.
 	const bool bValidShaderPlatform = (GMaxRHIShaderPlatform != SP_PCD3D_ES3_1);
-	if (!Output.IsValid() && !MaterialShaderMap->UsesSceneTexture(PPI_PostProcessInput0) && bPrimeOutputColor && !bForceIntermediateTarget && Inputs.bAllowSceneColorInputAsOutput && bValidShaderPlatform && !Inputs.bUserSceneTextureOutput)
+	if (!Output.IsValid() && !MaterialShaderMap->UsesSceneTexture(PPI_PostProcessInput0) && !bForceIntermediateTarget && Inputs.bAllowSceneColorInputAsOutput && bValidShaderPlatform && !Inputs.bUserSceneTextureOutput)
 	{
 		FScreenPassTexture SceneColor = FScreenPassTexture::CopyFromSlice(GraphBuilder, SceneColorOutput);
 
 		Output = FScreenPassRenderTarget(SceneColor, ERenderTargetLoadAction::ELoad);
+
+		// If material doesn't output alpha, and we are writing to an existing scene color, preserve its alpha by masking out writes
+		if (!Material->GetBlendableOutputAlpha())
+		{
+			BlendState = TStaticBlendState<CW_RGB>::GetRHI();
+		}
 	}
 	else
 	{
@@ -1090,9 +1098,9 @@ FScreenPassTexture AddPostProcessMaterialChain(
 			UserSceneTextureOutput = MaterialRenderProxy->GetUserSceneTextureOutput(Material);
 		}
 
-		if (UserSceneTextureOutput.IsNone())
+		if (UserSceneTextureOutput.IsNone() || UserSceneTextureOutput == NAME_SceneColor)
 		{
-			// Doesn't write to a UserSceneTexture, so it writes to the default SceneColor output
+			// If it doesn't write to a UserSceneTexture, it writes to the default SceneColor output (or it could be a UserSceneTexture set to write to SceneColor)
 			LastOutputWrite = Materials[MaterialIndex];
 			break;
 		}
@@ -1157,6 +1165,9 @@ FScreenPassTexture AddPostProcessMaterialChain(
 							// the use of SceneColor as an input is detected by the flags on the original FMaterialShaderMap (accessed via the UsesSceneTexture
 							// function), but those flags won't be set if a UserSceneTexture input is overridden to point at SceneColor.
 							Inputs.bAllowSceneColorInputAsOutput = false;
+
+							// Handle automatic propagation of scene color alpha from a UserSceneTexture input
+							Inputs.UserSceneTextureSceneColorInput = PostProcessIndex + (uint32)PPI_PostProcessInput0;
 						}
 						else
 						{
diff --git a/Engine/Source/Runtime/Renderer/Private/PostProcess/PostProcessMaterial.h b/Engine/Source/Runtime/Renderer/Private/PostProcess/PostProcessMaterial.h
--- a/Engine/Source/Runtime/Renderer/Private/PostProcess/PostProcessMaterial.h
+++ b/Engine/Source/Runtime/Renderer/Private/PostProcess/PostProcessMaterial.h
@@ -34,6 +34,7 @@ BEGIN_SHADER_PARAMETER_STRUCT(FPostProcessMaterialParameters, )
 	SHADER_PARAMETER(FVector2f, SceneWithoutSingleLayerWaterInvTextureSize)
 	SHADER_PARAMETER(uint32, ManualStencilReferenceValue)
 	SHADER_PARAMETER(uint32, ManualStencilTestMask)
+	SHADER_PARAMETER(uint32, UserSceneTextureSceneColorInput)
 	RENDER_TARGET_BINDING_SLOTS()
 END_SHADER_PARAMETER_STRUCT()
 
diff --git a/Engine/Source/Runtime/Renderer/Public/PostProcess/PostProcessMaterialInputs.h b/Engine/Source/Runtime/Renderer/Public/PostProcess/PostProcessMaterialInputs.h
--- a/Engine/Source/Runtime/Renderer/Public/PostProcess/PostProcessMaterialInputs.h
+++ b/Engine/Source/Runtime/Renderer/Public/PostProcess/PostProcessMaterialInputs.h
@@ -171,6 +171,7 @@ struct FPostProcessMaterialInputs
 
 	bool bUserSceneTextureOutput = false;
 	bool bUserSceneTextureFirstRender = false;
+	uint32 UserSceneTextureSceneColorInput = INDEX_NONE;
 };
 
 class UMaterialInterface;
