Slate SDF fonts - Material support
#tests Tested in editor
#rnx <Release Notes exclude>

diff --git a/Engine/Shaders/Private/MaterialTemplate.ush b/Engine/Shaders/Private/MaterialTemplate.ush
--- a/Engine/Shaders/Private/MaterialTemplate.ush
+++ b/Engine/Shaders/Private/MaterialTemplate.ush
@@ -568,6 +568,10 @@ struct FMaterialPixelParameters
 	uint WaterWaveParamIndex;
 #endif
 
+#if SHADER_TYPE == ST_GrayscaleFont || SHADER_TYPE == ST_ColorFont || SHADER_TYPE == ST_SdfFont || SHADER_TYPE == ST_MsdfFont
+	float4 FontSignedDistanceData;
+#endif
+
 #if CLOUD_LAYER_PIXEL_SHADER
 	float CloudSampleAltitude;
 	float CloudSampleAltitudeInLayer;
@@ -1918,6 +1922,15 @@ MaterialFloat GetDistanceCullFade()
 #endif
 }
 
+MaterialFloat4 GetFontSignedDistanceData(FMaterialPixelParameters Parameters)
+{
+#if SHADER_TYPE == ST_GrayscaleFont || SHADER_TYPE == ST_ColorFont || SHADER_TYPE == ST_SdfFont || SHADER_TYPE == ST_MsdfFont
+	return Parameters.FontSignedDistanceData;
+#else
+	return float4(0.0f, 0.0f, 0.0f, 0.0f);
+#endif
+}
+
 /** Rotates Position about the given axis by the given angle, in radians, and returns the offset to Position. */
 float3 RotateAboutAxis(float4 NormalizedRotationAxisAndAngle, float3 PositionOnAxis, float3 Position)
 {
diff --git a/Engine/Shaders/Private/SlateElementPixelShader.usf b/Engine/Shaders/Private/SlateElementPixelShader.usf
--- a/Engine/Shaders/Private/SlateElementPixelShader.usf
+++ b/Engine/Shaders/Private/SlateElementPixelShader.usf
@@ -144,7 +144,7 @@ void FillTexCoords( inout FMaterialPixelParameters Parameters, VertexToPixelInte
 #endif	// SHADER_TYPE
 }
 
-half4 GetMaterialColor( VertexToPixelInterpolants InVertex, float2 MeshUV )
+half4 GetMaterialColor( VertexToPixelInterpolants InVertex, float2 MeshUV, float4 InFontSignedDistanceData = float4(0.0, 0.0, 0.0, 0.0) )
 {
 	ResolvedView = ResolveView();
 	FMaterialPixelParameters Parameters = MakeInitializedMaterialPixelParameters();
@@ -153,8 +153,10 @@ half4 GetMaterialColor( VertexToPixelInterpolants InVertex, float2 MeshUV )
 	FillTexCoords( Parameters, InVertex, MeshUV );
 
 	Parameters.VertexColor = InVertex.Color;
+#if SHADER_TYPE == ST_GrayscaleFont || SHADER_TYPE == ST_ColorFont || SHADER_TYPE == ST_SdfFont || SHADER_TYPE == ST_MsdfFont
+	Parameters.FontSignedDistanceData = InFontSignedDistanceData;
+#endif
 
-   
 	FPixelMaterialInputs PixelMaterialInputs;
 	CalcMaterialParameters( Parameters, PixelMaterialInputs, InVertex.Position, true );
 
@@ -273,53 +275,95 @@ half4 GetColor( VertexToPixelInterpolants InVertex, float2 UV )
 }
 
 #if SHADER_TYPE == ST_SdfFont || SHADER_TYPE == ST_MsdfFont
-half4 GetSdfTextElementColor( VertexToPixelInterpolants InVertex )
+
+struct FontSignedDistanceSample
 {
-	const float2 AtlasTextureCoordinates = SLATE_UV0( InVertex );
-#if USE_MATERIALS
-	half4 OutColor = GetMaterialColor( InVertex, float2(0,0) ) * InVertex.Color;
-#else
-	half4 OutColor = InVertex.Color;
-#endif
+	MaterialFloat DistanceSample;
+	MaterialFloat SmoothDistanceSample;
+	float ScreenSpaceSpread;
+};
 
-	const float4 MultiDistance = float4(Texture2DSample( ElementTexture, ElementTextureSampler, AtlasTextureCoordinates ));
+FontSignedDistanceSample SampleFontSignedDistance(float2 TextureCoordinate)
+{
+	FontSignedDistanceSample OutSample;
+	const MaterialFloat4 MultiDistance = Texture2DSample(ElementTexture, ElementTextureSampler, TextureCoordinate);
 #if SHADER_TYPE == ST_MsdfFont
-	const float SampledDistance = max(min(MultiDistance.r, MultiDistance.g), min(max(MultiDistance.r, MultiDistance.g), MultiDistance.b));
-#else
-	const float SampledDistance = lerp(MultiDistance A8_SAMPLE_MASK, MultiDistance.a, ShaderParams.w);
+	OutSample.DistanceSample = max(min(MultiDistance.r, MultiDistance.g), min(max(MultiDistance.r, MultiDistance.g), MultiDistance.b));
+	OutSample.SmoothDistanceSample = MultiDistance.a;
+#else // SHADER_TYPE == ST_SdfFont
+	OutSample.DistanceSample = lerp(MultiDistance A8_SAMPLE_MASK, MultiDistance.a, ShaderParams2.x);
+	OutSample.SmoothDistanceSample = OutSample.DistanceSample;
+#endif
+	const float2 ScreenSpaceTextureSize = float2(1.0, 1.0) / fwidth(TextureCoordinate);
+	OutSample.ScreenSpaceSpread = max(dot(ShaderParams.xy, ScreenSpaceTextureSize), 1.0);
+	return OutSample;
+}
+
+float GetFontSignedDistanceOpacity(FontSignedDistanceSample InSample)
+{
+	const float ScreenSpaceDistance = (InSample.DistanceSample - ShaderParams.z) * InSample.ScreenSpaceSpread;
+	return saturate(ScreenSpaceDistance + 0.5);
+}
+
+#if USE_MATERIALS
+float4 GetFontSignedDistanceData(FontSignedDistanceSample InSample)
+{
+	const float EmSignedDistance = (InSample.DistanceSample - ShaderParams.z) * ShaderParams.w;
+	const float EmSmoothDistance = (InSample.SmoothDistanceSample - ShaderParams.z) * ShaderParams.w;
+	const float EmPxFactor = InSample.ScreenSpaceSpread / ShaderParams.w;
+	return float4(EmSignedDistance, EmSmoothDistance, EmPxFactor, GetFontSignedDistanceOpacity(InSample));
+}
 #endif
-	const float2 ScreenSpaceTextureSize = float2(1.0, 1.0) / fwidth(AtlasTextureCoordinates);
-	const float ScreenSpaceSpread = max(dot(ShaderParams.xy, ScreenSpaceTextureSize), 1.0);
-	const float ScreenSpaceDistance = (SampledDistance - ShaderParams.z) * ScreenSpaceSpread;
-	#if USE_MATERIALS && (MATERIALBLENDING_ADDITIVE || MATERIALBLENDING_ALPHACOMPOSITE)
-		OutColor *= half(saturate(ScreenSpaceDistance + 0.5));
+
+half4 GetSdfTextElementColor(VertexToPixelInterpolants InVertex)
+{
+	const FontSignedDistanceSample SdfSample = SampleFontSignedDistance(SLATE_UV0(InVertex));
+	half4 OutColor = InVertex.Color;
+#if USE_MATERIALS
+	const float4 SdfData = GetFontSignedDistanceData(SdfSample);
+	OutColor *= GetMaterialColor(InVertex, float2(0, 0), SdfData);
+	if (ShaderParams2.y < 0.5) // ShaderParams2.y is bMaterialIsStencil
+	{
+	#if MATERIALBLENDING_ADDITIVE || MATERIALBLENDING_ALPHACOMPOSITE
+		OutColor *= SdfData.a;
 	#else
-		OutColor.a *= half(saturate(ScreenSpaceDistance + 0.5));
+		OutColor.a *= SdfData.a;
 	#endif
-
+	}
+#else // !USE_MATERIALS
+	OutColor.a *= GetFontSignedDistanceOpacity(SdfSample);
+#endif
 	return OutColor;
 }
+
 #endif
 
+
 #if SHADER_TYPE == ST_GrayscaleFont || SHADER_TYPE == ST_ColorFont
 half4 GetFontElementColor( VertexToPixelInterpolants InVertex )
 {
 	float2 AtlasTextureCoordinates = SLATE_UV0( InVertex );
-#if USE_MATERIALS
-	half4 BaseColor = GetMaterialColor( InVertex, float2(0,0) ) * InVertex.Color;
-#else
-	half4 BaseColor = InVertex.Color;
-#endif
-	half4 OutColor = BaseColor;
+	MaterialFloat4 SampledColor = MaterialFloat4(1,1,1,1);
 
 #if SHADER_TYPE == ST_GrayscaleFont
 	#if USE_MATERIALS && (MATERIALBLENDING_ADDITIVE || MATERIALBLENDING_ALPHACOMPOSITE)
-		OutColor *= Texture2DSample_A8( ElementTexture, ElementTextureSampler, AtlasTextureCoordinates );
+		SampledColor *= Texture2DSample_A8( ElementTexture, ElementTextureSampler, AtlasTextureCoordinates );
 	#else
-		OutColor.a *= Texture2DSample_A8( ElementTexture, ElementTextureSampler, AtlasTextureCoordinates );
+		SampledColor.a = Texture2DSample_A8( ElementTexture, ElementTextureSampler, AtlasTextureCoordinates );
 	#endif
 #else
-	OutColor *= Texture2DSample( ElementTexture, ElementTextureSampler, AtlasTextureCoordinates );
+	SampledColor = Texture2DSample( ElementTexture, ElementTextureSampler, AtlasTextureCoordinates );
+#endif
+
+	half4 OutColor = InVertex.Color;
+#if USE_MATERIALS
+	OutColor *= GetMaterialColor( InVertex, float2(0,0), float4(0, 0, 0, SampledColor.a) );
+	if (ShaderParams2.y < 0.5) // ShaderParams2.y is bMaterialIsStencil
+	{
+		OutColor *= half4(SampledColor);
+	}
+#else
+	OutColor *= half4(SampledColor);
 #endif
 
 	return OutColor;
diff --git a/Engine/Source/Editor/DetailCustomizations/Private/SlateFontInfoCustomization.cpp b/Engine/Source/Editor/DetailCustomizations/Private/SlateFontInfoCustomization.cpp
--- a/Engine/Source/Editor/DetailCustomizations/Private/SlateFontInfoCustomization.cpp
+++ b/Engine/Source/Editor/DetailCustomizations/Private/SlateFontInfoCustomization.cpp
@@ -143,6 +143,8 @@ void FSlateFontInfoStructCustomization::CustomizeChildren(TSharedRef<IPropertyHa
 
 	InStructBuilder.AddProperty(InStructPropertyHandle->GetChildHandle(GET_MEMBER_NAME_CHECKED(FSlateFontInfo, FontMaterial)).ToSharedRef());
 
+	InStructBuilder.AddProperty(InStructPropertyHandle->GetChildHandle(GET_MEMBER_NAME_CHECKED(FSlateFontInfo, bMaterialIsStencil)).ToSharedRef());
+
 	InStructBuilder.AddProperty(InStructPropertyHandle->GetChildHandle(GET_MEMBER_NAME_CHECKED(FSlateFontInfo, OutlineSettings)).ToSharedRef());
 }
 
diff --git a/Engine/Source/Runtime/Engine/Private/Materials/HLSLMaterialTranslator.cpp b/Engine/Source/Runtime/Engine/Private/Materials/HLSLMaterialTranslator.cpp
--- a/Engine/Source/Runtime/Engine/Private/Materials/HLSLMaterialTranslator.cpp
+++ b/Engine/Source/Runtime/Engine/Private/Materials/HLSLMaterialTranslator.cpp
@@ -8855,6 +8855,15 @@ int32 FHLSLMaterialTranslator::StaticTerrainLayerWeight(FName LayerName,int32 De
 	}
 }
 
+int32 FHLSLMaterialTranslator::FontSignedDistanceData()
+{
+	if (ShaderFrequency != SF_Pixel)
+	{
+		return Errorf(TEXT("Font Signed Distance node is only available in pixel shader input."));
+	}
+	return AddCodeChunk(MCT_Float4, TEXT("GetFontSignedDistanceData(Parameters)"));
+}
+
 int32 FHLSLMaterialTranslator::VertexColor()
 {
 	bUsesVertexColor |= (ShaderFrequency != SF_Vertex);
diff --git a/Engine/Source/Runtime/Engine/Private/Materials/HLSLMaterialTranslator.h b/Engine/Source/Runtime/Engine/Private/Materials/HLSLMaterialTranslator.h
--- a/Engine/Source/Runtime/Engine/Private/Materials/HLSLMaterialTranslator.h
+++ b/Engine/Source/Runtime/Engine/Private/Materials/HLSLMaterialTranslator.h
@@ -1119,6 +1119,8 @@ protected:
 
 	virtual int32 StaticTerrainLayerWeight(FName ParameterName, int32 Default, bool bTextureArray = false) override;
 
+	virtual int32 FontSignedDistanceData() override;
+
 	virtual int32 VertexColor() override;
 	virtual int32 MeshPaintTextureDescriptor() override;
 	virtual int32 MeshPaintTextureReplace(int32 Invalid, int32 Valid) override;
diff --git a/Engine/Source/Runtime/Engine/Private/Materials/MaterialExpressions.cpp b/Engine/Source/Runtime/Engine/Private/Materials/MaterialExpressions.cpp
--- a/Engine/Source/Runtime/Engine/Private/Materials/MaterialExpressions.cpp
+++ b/Engine/Source/Runtime/Engine/Private/Materials/MaterialExpressions.cpp
@@ -107,6 +107,7 @@
 #include "Materials/MaterialExpressionFloor.h"	
 #include "Materials/MaterialExpressionFloatToUInt.h"	
 #include "Materials/MaterialExpressionFmod.h"
+#include "Materials/MaterialExpressionFontSignedDistance.h"
 #include "Materials/MaterialExpressionFontSample.h"
 #include "Materials/MaterialExpressionFontSampleParameter.h"
 #include "Materials/MaterialExpressionFrac.h"
@@ -11626,6 +11627,51 @@ void UMaterialExpressionVertexColor::GetCaption(TArray<FString>& OutCaptions) co
 }
 #endif // WITH_EDITOR
 
+UMaterialExpressionFontSignedDistance::UMaterialExpressionFontSignedDistance(const FObjectInitializer& ObjectInitializer)
+	: Super(ObjectInitializer)
+{
+#if WITH_EDITORONLY_DATA
+	// Structure to hold one-time initialization
+	struct FConstructorStatics
+	{
+		FText NAME_Constants;
+		FConstructorStatics()
+			: NAME_Constants(LOCTEXT("Constants", "Constants"))
+		{
+		}
+	};
+	static FConstructorStatics ConstructorStatics;
+
+	MenuCategories.Add(ConstructorStatics.NAME_Constants);
+
+	Outputs.Reset();
+	Outputs.Add(FExpressionOutput(TEXT("Signed Distance"), 1, 1, 0, 0, 0));
+	Outputs.Add(FExpressionOutput(TEXT("Smooth Signed Distance"), 1, 0, 1, 0, 0));
+	Outputs.Add(FExpressionOutput(TEXT("Pixel Distance Factor"), 1, 0, 0, 1, 0));
+	Outputs.Add(FExpressionOutput(TEXT("Implicit Opacity"), 1, 0, 0, 0, 1));
+
+	bShowOutputNameOnPin = true;
+	bShowMaskColorsOnPin = false;
+#endif
+}
+
+#if WITH_EDITOR
+int32 UMaterialExpressionFontSignedDistance::Compile(class FMaterialCompiler* Compiler, int32 OutputIndex)
+{
+	if(!Material || Material->MaterialDomain != MD_UI)
+	{
+		return Compiler->Errorf(TEXT("Font Signed Distance node is only compatible with UI materials"));
+	}
+
+	return Compiler->FontSignedDistanceData();
+}
+
+void UMaterialExpressionFontSignedDistance::GetCaption(TArray<FString>& OutCaptions) const
+{
+	OutCaptions.Add(TEXT("Font Signed Distance"));
+}
+#endif // WITH_EDITOR
+
 UMaterialExpressionParticleColor::UMaterialExpressionParticleColor(const FObjectInitializer& ObjectInitializer)
 	: Super(ObjectInitializer)
 {
diff --git a/Engine/Source/Runtime/Engine/Public/MaterialCompiler.h b/Engine/Source/Runtime/Engine/Public/MaterialCompiler.h
--- a/Engine/Source/Runtime/Engine/Public/MaterialCompiler.h
+++ b/Engine/Source/Runtime/Engine/Public/MaterialCompiler.h
@@ -450,6 +450,8 @@ public:
 	virtual bool GetStaticBoolValue(int32 BoolIndex, bool& bSucceeded) = 0;
 	virtual int32 StaticTerrainLayerWeight(FName ParameterName,int32 Default, bool bTextureArray = false) = 0;
 
+	virtual int32 FontSignedDistanceData() = 0;
+
 	virtual int32 VertexColor() = 0;
 	virtual int32 MeshPaintTextureDescriptor() = 0;
 	virtual int32 MeshPaintTextureReplace(int32 Invalid, int32 Valid) = 0;
@@ -1004,6 +1006,8 @@ public:
 	virtual bool GetStaticBoolValue(int32 BoolIndex, bool& bSucceeded) override { return Compiler->GetStaticBoolValue(BoolIndex, bSucceeded); }
 	virtual int32 StaticTerrainLayerWeight(FName ParameterName, int32 Default, bool bTextureArray = false) override { return Compiler->StaticTerrainLayerWeight(ParameterName, Default, bTextureArray); }
 
+	virtual int32 FontSignedDistanceData() override { return Compiler->FontSignedDistanceData(); }
+
 	virtual int32 VertexColor() override { return Compiler->VertexColor(); }
 	virtual int32 MeshPaintTextureDescriptor() override { return Compiler->MeshPaintTextureDescriptor(); }
 	virtual int32 MeshPaintTextureReplace(int32 Invalid, int32 Valid) override { return Compiler->MeshPaintTextureReplace(Invalid, Valid); }
diff --git a/Engine/Source/Runtime/Engine/Public/Materials/MaterialExpressionFontSignedDistance.h b/Engine/Source/Runtime/Engine/Public/Materials/MaterialExpressionFontSignedDistance.h
new file mode 100644
--- /dev/null
+++ b/Engine/Source/Runtime/Engine/Public/Materials/MaterialExpressionFontSignedDistance.h
@@ -0,0 +1,20 @@
+// Copyright Epic Games, Inc. All Rights Reserved.
+#pragma once
+
+#include "CoreMinimal.h"
+#include "UObject/ObjectMacros.h"
+#include "Materials/MaterialExpression.h"
+#include "MaterialExpressionFontSignedDistance.generated.h"
+
+UCLASS(collapsecategories, hidecategories=Object, MinimalAPI)
+class UMaterialExpressionFontSignedDistance : public UMaterialExpression
+{
+	GENERATED_UCLASS_BODY()
+
+	//~ Begin UMaterialExpression Interface
+#if WITH_EDITOR
+	virtual int32 Compile(class FMaterialCompiler* Compiler, int32 OutputIndex) override;
+	virtual void GetCaption(TArray<FString>& OutCaptions) const override;
+#endif
+	//~ End UMaterialExpression Interface
+};
diff --git a/Engine/Source/Runtime/SlateCore/Private/Fonts/FontCache.cpp b/Engine/Source/Runtime/SlateCore/Private/Fonts/FontCache.cpp
--- a/Engine/Source/Runtime/SlateCore/Private/Fonts/FontCache.cpp
+++ b/Engine/Source/Runtime/SlateCore/Private/Fonts/FontCache.cpp
@@ -174,10 +174,11 @@ FSdfGlyphTaskKey::FSdfGlyphTaskKey(const FSdfGlyphEntryKey& InSdfGlyphEntryKey,
 * 
 */
 
-FShapedGlyphSequence::FShapedGlyphSequence(TArray<FShapedGlyphEntry> InGlyphsToRender, const int16 InTextBaseline, const uint16 InMaxTextHeight, const UObject* InFontMaterial, const FFontOutlineSettings& InOutlineSettings, const FSourceTextRange& InSourceTextRange)
+FShapedGlyphSequence::FShapedGlyphSequence(TArray<FShapedGlyphEntry> InGlyphsToRender, const int16 InTextBaseline, const uint16 InMaxTextHeight, const UObject* InFontMaterial, const bool InMaterialIsStencil, const FFontOutlineSettings& InOutlineSettings, const FSourceTextRange& InSourceTextRange)
 	: GlyphsToRender(MoveTemp(InGlyphsToRender))
 	, TextBaseline(InTextBaseline)
 	, MaxTextHeight(InMaxTextHeight)
+	, bMaterialIsStencil(InMaterialIsStencil)
 	, FontMaterial(InFontMaterial)
 	, OutlineSettings(InOutlineSettings)
 	, SequenceWidth(0)
@@ -507,7 +508,7 @@ FShapedGlyphSequencePtr FShapedGlyphSequence::GetSubSequence(const int32 InStart
 
 	if (EnumerateVisualGlyphsInSourceRange(InStartIndex, InEndIndex, GlyphCallback) == EEnumerateGlyphsResult::EnumerationComplete)
 	{
-		return MakeShared<FShapedGlyphSequence>(MoveTemp(SubGlyphsToRender), TextBaseline, MaxTextHeight, FontMaterial, OutlineSettings, SubSequenceRange);
+		return MakeShared<FShapedGlyphSequence>(MoveTemp(SubGlyphsToRender), TextBaseline, MaxTextHeight, FontMaterial, bMaterialIsStencil, OutlineSettings, SubSequenceRange);
 	}
 
 	return nullptr;
diff --git a/Engine/Source/Runtime/SlateCore/Private/Fonts/SlateTextShaper.cpp b/Engine/Source/Runtime/SlateCore/Private/Fonts/SlateTextShaper.cpp
--- a/Engine/Source/Runtime/SlateCore/Private/Fonts/SlateTextShaper.cpp
+++ b/Engine/Source/Runtime/SlateCore/Private/Fonts/SlateTextShaper.cpp
@@ -312,6 +312,7 @@ FShapedGlyphSequenceRef FSlateTextShaper::FinalizeTextShaping(TArray<FShapedGlyp
 											TextBaseline, 
 											MaxHeight, 
 											InFontInfo.FontMaterial.Get(),
+											InFontInfo.bMaterialIsStencil,
 											InFontInfo.OutlineSettings,
 											InSourceTextRange);
 }
diff --git a/Engine/Source/Runtime/SlateCore/Private/Rendering/ElementBatcher.cpp b/Engine/Source/Runtime/SlateCore/Private/Rendering/ElementBatcher.cpp
--- a/Engine/Source/Runtime/SlateCore/Private/Rendering/ElementBatcher.cpp
+++ b/Engine/Source/Runtime/SlateCore/Private/Rendering/ElementBatcher.cpp
@@ -3295,8 +3295,10 @@ int32 FSlateElementBatcher::BuildShapedTextSequence(const FShapedTextBuildContex
 	bool bNeedSpaceForEllipsis = false;
 	const bool bNonZeroSkew = !FMath::IsNearlyEqual(GlyphSequenceToRender->GetFontSkew(), 0.f);
 	bool bPrevSdfGlyph = false;
+	float SdfEmSpread = 0;
 	float SdfPixelSpread = 0;
 	float SdfBias = 0;
+	const bool bMaterialIsStencil = GlyphSequenceToRender->IsMaterialStencil();
 	// For left to right overflow direction - Sum of total whitespace we're currently advancing through. Once a non-whitespace glyph is detected this will return to 0
 	// For right to left this value is unused. We just skip all leading whitespace
 	float PreviousWhitespaceAdvance = 0;
@@ -3371,6 +3373,7 @@ int32 FSlateElementBatcher::BuildShapedTextSequence(const FShapedTextBuildContex
 		{
 			// Get Sizing and atlas info
 			int8 NextAtlasDataTextureIndex = -1;
+			float NextSdfEmSpread = 0;
 			float NextSdfPixelSpread = 0;
 			float NextSdfBias = 0;
 
@@ -3395,7 +3398,8 @@ int32 FSlateElementBatcher::BuildShapedTextSequence(const FShapedTextBuildContex
 					// SdfGlyphAtlasData.Metrics values assume this operation, so QuadMeshSize already accounts for this
 					SpriteSize = FVector2f(SdfGlyphAtlasData.USize-1, SdfGlyphAtlasData.VSize-1);
 					SpriteOffset = FVector2f(SdfGlyphAtlasData.StartU+.5f, SdfGlyphAtlasData.StartV+.5f);
-					NextSdfPixelSpread = (SdfGlyphAtlasData.EmInnerSpread+SdfGlyphAtlasData.EmOuterSpread)*float(GlyphToRender.FontFaceData->SdfPpem);
+					NextSdfEmSpread = SdfGlyphAtlasData.EmInnerSpread + SdfGlyphAtlasData.EmOuterSpread;
+					NextSdfPixelSpread = NextSdfEmSpread * float(GlyphToRender.FontFaceData->SdfPpem);
 					// Value representing zero distance
 					NextSdfBias = (SdfGlyphAtlasData.EmOuterSpread-EmOutlineSize)/(SdfGlyphAtlasData.EmInnerSpread+SdfGlyphAtlasData.EmOuterSpread);
 				}
@@ -3441,7 +3445,7 @@ int32 FSlateElementBatcher::BuildShapedTextSequence(const FShapedTextBuildContex
 
 				// Initiate new render batch for first glyph and whenever texture / shader / shader parameters change between previous and current glyph
 				check(NextAtlasDataTextureIndex >= 0);
-				if (FontAtlasTexture == nullptr || NextAtlasDataTextureIndex != FontTextureIndex || bIsSdfGlyph != bPrevSdfGlyph || (bIsSdfGlyph && (NextSdfPixelSpread != SdfPixelSpread || NextSdfBias != SdfBias)))
+				if (FontAtlasTexture == nullptr || NextAtlasDataTextureIndex != FontTextureIndex || bIsSdfGlyph != bPrevSdfGlyph || (bIsSdfGlyph && (NextSdfEmSpread != SdfEmSpread || NextSdfPixelSpread != SdfPixelSpread || NextSdfBias != SdfBias)))
 				{
 					// Font has a new texture for this glyph or shader parameters changed. Refresh the batch we use and the index we are currently using
 					FontTextureIndex = NextAtlasDataTextureIndex;
@@ -3486,6 +3490,7 @@ int32 FSlateElementBatcher::BuildShapedTextSequence(const FShapedTextBuildContex
 					FShaderParams ShaderParams;
 					if (bIsSdfGlyph)
 					{
+						SdfEmSpread = NextSdfEmSpread;
 						SdfPixelSpread = NextSdfPixelSpread;
 						SdfBias = NextSdfBias;
 						// Note - it would be much better to pass the SDF shader params as per-vertex attributes instead to avoid having to switch batches too often
@@ -3495,10 +3500,19 @@ int32 FSlateElementBatcher::BuildShapedTextSequence(const FShapedTextBuildContex
 							.5f*InvTextureSizeY*SdfPixelSpread,
 							// Signed distance sample bias, the color value (between 0 to 1) representing zero distance
 							SdfBias,
-							// The last parameter needs to be 0 for alpha texture single-channel SDF, 1 for BGRA/RGBA MTSDF
-							ContentType == ESlateFontAtlasContentType::Msdf ? 1.f : 0.f
+							// Signed distance spread in em
+							SdfEmSpread
+						), FVector4f(
+							// This parameter needs to be 0 for alpha texture single-channel SDF, 1 for BGRA/RGBA MTSDF
+							ContentType == ESlateFontAtlasContentType::Msdf ? 1.f : 0.f,
+							bMaterialIsStencil ? 1.f : 0.f,
+							0.f, 0.f
 						));
 					}
+					else
+					{
+						ShaderParams = FShaderParams::MakePixelShaderParams(FVector4f(), FVector4f(0.f, bMaterialIsStencil ? 1.f : 0.f, 0.f, 0.f));
+					}
 
 					RenderBatch = &CreateRenderBatch(Context.LayerId,
 						ShaderParams,
diff --git a/Engine/Source/Runtime/SlateCore/Public/Fonts/FontCache.h b/Engine/Source/Runtime/SlateCore/Public/Fonts/FontCache.h
--- a/Engine/Source/Runtime/SlateCore/Public/Fonts/FontCache.h
+++ b/Engine/Source/Runtime/SlateCore/Public/Fonts/FontCache.h
@@ -349,6 +349,7 @@ public:
 		: GlyphsToRender()
 		, TextBaseline(0)
 		, MaxTextHeight(0)
+		, bMaterialIsStencil(false)
 		, FontMaterial(nullptr)
 		, OutlineSettings()
 		, SequenceWidth(0)
@@ -361,6 +362,7 @@ public:
 						 const int16 InTextBaseline,
 						 const uint16 InMaxTextHeight,
 						 const UObject* InFontMaterial,
+						 const bool InMaterialIsStencil,
 						 const FFontOutlineSettings& InOutlineSettings,
 						 const FSourceTextRange& InSourceTextRange);
 	SLATECORE_API ~FShapedGlyphSequence();
@@ -392,6 +394,12 @@ public:
 		return FontMaterial;
 	}
 
+	/** Get "Material is stencil" flag */
+	bool IsMaterialStencil() const
+	{
+		return bMaterialIsStencil;
+	}
+
 	/** Get the font outline settings to use when rendering these glyphs */
 	const FFontOutlineSettings& GetFontOutlineSettings() const
 	{
@@ -579,6 +587,8 @@ private:
 	int16 TextBaseline;
 	/** The maximum height of any glyph in the font we're using */
 	uint16 MaxTextHeight;
+	/** Toggles between stenciling text onto material output (false) or filling entire quads with material (true) */
+	bool bMaterialIsStencil;
 	/** The material to use when rendering these glyphs */
 	TObjectPtr<const UObject> FontMaterial;
 	/** Outline settings to use when rendering these glyphs */
diff --git a/Engine/Source/Runtime/SlateCore/Public/Fonts/SlateFontInfo.h b/Engine/Source/Runtime/SlateCore/Public/Fonts/SlateFontInfo.h
--- a/Engine/Source/Runtime/SlateCore/Public/Fonts/SlateFontInfo.h
+++ b/Engine/Source/Runtime/SlateCore/Public/Fonts/SlateFontInfo.h
@@ -189,6 +189,10 @@ struct FSlateFontInfo
 	UPROPERTY(EditAnywhere, BlueprintReadWrite, Category=SlateStyleRules, meta=(DisplayName="Monospacing"))
 	bool bForceMonospaced = false;
 
+	/** When enabled, whole quads are filled by the material without automatically stenciling the text - this needs to be done within the material (using Font Signed Distance node). */
+	UPROPERTY(EditAnywhere, BlueprintReadWrite, Category=SlateStyleRules)
+	bool bMaterialIsStencil = false;
+
 	/** The uniform width to apply to all characters when bForceMonospaced is enabled, proportional of the font Size. */
 	UPROPERTY(EditAnywhere, BlueprintReadWrite, Category=SlateStyleRules, meta=(ClampMin=0, EditCondition = "bForceMonospaced==true", EditConditionHides))
 	float MonospacedWidth = 1.0f;
