MegaLights - fast out on NdotL in order to skip complex calculations in DefaultLitBxDF if we know that they will be anyway multiplied by 0.

~0.07ms performance improvement on console

diff --git a/Engine/Shaders/Private/ShadingModels.ush b/Engine/Shaders/Private/ShadingModels.ush
--- a/Engine/Shaders/Private/ShadingModels.ush
+++ b/Engine/Shaders/Private/ShadingModels.ush
@@ -269,78 +269,86 @@ FDirectLighting DefaultLitBxDF( FGBufferData GBuffer, half3 N, half3 V, half3 L,
 {
 	BxDFContext Context;
 	FDirectLighting Lighting;
+	Lighting.Diffuse = 0;
+	Lighting.Specular = 0;
+	Lighting.Transmission = 0;
 
+	BRANCH
+	if (NoL > 0.0f)
+	{
 #if SUPPORTS_ANISOTROPIC_MATERIALS
-	bool bHasAnisotropy = HasAnisotropy(GBuffer.SelectiveOutputMask);
+		bool bHasAnisotropy = HasAnisotropy(GBuffer.SelectiveOutputMask);
 #else
-	bool bHasAnisotropy = false;
+		bool bHasAnisotropy = false;
 #endif
 
-	float NoV, VoH, NoH;
-	BRANCH
-	if (bHasAnisotropy)
-	{
-		half3 X = GBuffer.WorldTangent;
-		half3 Y = normalize(cross(N, X));
-		Init(Context, N, X, Y, V, L);
+		float NoV, VoH, NoH;
+		BRANCH
+		if (bHasAnisotropy)
+		{
+			half3 X = GBuffer.WorldTangent;
+			half3 Y = normalize(cross(N, X));
+			Init(Context, N, X, Y, V, L);
 
-		NoV = Context.NoV;
-		VoH = Context.VoH;
-		NoH = Context.NoH;
-	}
-	else
-	{
+			NoV = Context.NoV;
+			VoH = Context.VoH;
+			NoH = Context.NoH;
+		}
+		else
+		{
 #if SHADING_PATH_MOBILE
-		InitMobile(Context, N, V, L, NoL);
+			InitMobile(Context, N, V, L, NoL);
 #else
-		Init(Context, N, V, L);
+			Init(Context, N, V, L);
 #endif
 
-		NoV = Context.NoV;
-		VoH = Context.VoH;
-		NoH = Context.NoH;
+			NoV = Context.NoV;
+			VoH = Context.VoH;
+			NoH = Context.NoH;
 
-		SphereMaxNoH(Context, AreaLight.SphereSinAlpha, true);
-	}
+			SphereMaxNoH(Context, AreaLight.SphereSinAlpha, true);
+		}
 
-	Context.NoV = saturate(abs( Context.NoV ) + 1e-5);
+		Context.NoV = saturate(abs( Context.NoV ) + 1e-5);
 
 #if MATERIAL_ROUGHDIFFUSE
-	// Chan diffuse model with roughness == specular roughness. This is not necessarily a good modelisation of reality because when the mean free path is super small, the diffuse can in fact looks rougher. But this is a start.
-	// Also we cannot use the morphed context maximising NoH as this is causing visual artefact when interpolating rough/smooth diffuse response. 
-	Lighting.Diffuse = Diffuse_Chan(GBuffer.DiffuseColor, Pow4(GBuffer.Roughness), NoV, NoL, VoH, NoH, GetAreaLightDiffuseMicroReflWeight(AreaLight));
+		// Chan diffuse model with roughness == specular roughness. This is not necessarily a good modelisation of reality because when the mean free path is super small, the diffuse can in fact looks rougher. But this is a start.
+		// Also we cannot use the morphed context maximising NoH as this is causing visual artefact when interpolating rough/smooth diffuse response. 
+		Lighting.Diffuse = Diffuse_Chan(GBuffer.DiffuseColor, Pow4(GBuffer.Roughness), NoV, NoL, VoH, NoH, GetAreaLightDiffuseMicroReflWeight(AreaLight));
 #else
-	Lighting.Diffuse = Diffuse_Lambert(GBuffer.DiffuseColor);
+		Lighting.Diffuse = Diffuse_Lambert(GBuffer.DiffuseColor);
 #endif
-	Lighting.Diffuse *= AreaLight.FalloffColor * (Falloff * NoL);
+		Lighting.Diffuse *= AreaLight.FalloffColor * (Falloff * NoL);
 
-	BRANCH
-	if (bHasAnisotropy)
-	{
-		//Lighting.Specular = GBuffer.WorldTangent * .5f + .5f;
-		Lighting.Specular = AreaLight.FalloffColor * (Falloff * NoL) * SpecularGGX(GBuffer.Roughness, GBuffer.Anisotropy, GBuffer.SpecularColor, Context, NoL, AreaLight);
-	}
-	else
-	{
-		if( IsRectLight(AreaLight) )
+		BRANCH
+		if (bHasAnisotropy)
 		{
-			Lighting.Specular = RectGGXApproxLTC(GBuffer.Roughness, GBuffer.SpecularColor, N, V, AreaLight.Rect, AreaLight.Texture);
+			//Lighting.Specular = GBuffer.WorldTangent * .5f + .5f;
+			Lighting.Specular = AreaLight.FalloffColor * (Falloff * NoL) * SpecularGGX(GBuffer.Roughness, GBuffer.Anisotropy, GBuffer.SpecularColor, Context, NoL, AreaLight);
 		}
 		else
 		{
-			Lighting.Specular = AreaLight.FalloffColor * (Falloff * NoL) * SpecularGGX(GBuffer.Roughness, GBuffer.SpecularColor, Context, NoL, AreaLight);
+			if( IsRectLight(AreaLight) )
+			{
+				Lighting.Specular = RectGGXApproxLTC(GBuffer.Roughness, GBuffer.SpecularColor, N, V, AreaLight.Rect, AreaLight.Texture);
+			}
+			else
+			{
+				Lighting.Specular = AreaLight.FalloffColor * (Falloff * NoL) * SpecularGGX(GBuffer.Roughness, GBuffer.SpecularColor, Context, NoL, AreaLight);
+			}
 		}
-	}
 
-	FBxDFEnergyTerms EnergyTerms = ComputeGGXSpecEnergyTerms(GBuffer.Roughness, Context.NoV, GBuffer.SpecularColor);
+		FBxDFEnergyTerms EnergyTerms = ComputeGGXSpecEnergyTerms(GBuffer.Roughness, Context.NoV, GBuffer.SpecularColor);
 
-	// Add energy presevation (i.e. attenuation of the specular layer onto the diffuse component
-	Lighting.Diffuse *= ComputeEnergyPreservation(EnergyTerms);
+		// Add energy presevation (i.e. attenuation of the specular layer onto the diffuse component
+		Lighting.Diffuse *= ComputeEnergyPreservation(EnergyTerms);
 
-	// Add specular microfacet multiple scattering term (energy-conservation)
-	Lighting.Specular *= ComputeEnergyConservation(EnergyTerms);
+		// Add specular microfacet multiple scattering term (energy-conservation)
+		Lighting.Specular *= ComputeEnergyConservation(EnergyTerms);
+
+		Lighting.Transmission = 0;
+	}
 
-	Lighting.Transmission = 0;
 	return Lighting;
 }
 
