Substrate - fix for RT shadow leaking when multiple layers have normal conflicting with the top normal.

[RN]
#tests none

diff --git a/Engine/Shaders/Private/RayTracing/RayTracingOcclusionRGS.usf b/Engine/Shaders/Private/RayTracing/RayTracingOcclusionRGS.usf
--- a/Engine/Shaders/Private/RayTracing/RayTracingOcclusionRGS.usf
+++ b/Engine/Shaders/Private/RayTracing/RayTracingOcclusionRGS.usf
@@ -734,6 +734,17 @@ RAY_TRACING_ENTRY_RAYGEN_OR_INLINE(Occlusion)
 			ShadingParameters.bHasSubsurfaceProfile = bHasSSS ? (SSSType == SSS_TYPE_DIFFUSION_PROFILE || SSSType == SSS_TYPE_DIFFUSION) : false;
 			ShadingParameters.bBackfaceCulling = !(BSDFType == SUBSTRATE_BSDF_TYPE_EYE || (bHasSSS && SSSType == SSS_TYPE_WRAP)); // isthin?
 
+			// The SSDenoiser reads the top layer normal in order to know how to process the data.
+			// Shadow is traced using the top layer normal when NoL>0 only. So if normals for slabs below do not match, light leaking will occur.
+			// For instance: a flat translucent top layer above a normal mapped bottom layer.
+			// That is why, for complex materials with multiple BSDFs, we also must not do backface culling.
+			if (SubstratePixelHeader.ClosureCount > 1)
+			{
+				// Technically, this could only be done when vertical layering is encountered only. 
+				// But we do not have that data available (or it would be costly to infer from closures)
+				ShadingParameters.bBackfaceCulling = false;
+			}
+
 			WorldNormal = TopLayerData.WorldNormal;
 
 
