Substrate - fix for VSM  leaking when multiple layers have normal conflicting with the top normal.

#rnx
#tests none

diff --git a/Engine/Shaders/Private/VirtualShadowMaps/VirtualShadowMapProjection.usf b/Engine/Shaders/Private/VirtualShadowMaps/VirtualShadowMapProjection.usf
--- a/Engine/Shaders/Private/VirtualShadowMaps/VirtualShadowMapProjection.usf
+++ b/Engine/Shaders/Private/VirtualShadowMaps/VirtualShadowMapProjection.usf
@@ -47,7 +47,7 @@ struct FProjectionShadingInfo
 	bool bIsValid;
 	bool bIsHair;
 	bool bIsSubsurface;
-	bool bIsClearCoat;
+	bool bDisableBackFaceCulling;
 
 	float3 WorldNormal;
 };
@@ -55,6 +55,8 @@ struct FProjectionShadingInfo
 FProjectionShadingInfo GetProjectionShadingInfo(uint2 PixelPos)
 {
 	FProjectionShadingInfo Out;
+	Out.bDisableBackFaceCulling = false;
+
 #if SUBSTRATE_ENABLED
 	FSubstrateAddressing SubstrateAddressing = GetSubstratePixelDataByteOffset(PixelPos, uint2(View.BufferSizeAndInvSize.xy), Substrate.MaxBytesPerPixel);
 	FSubstratePixelHeader SubstratePixelHeader = UnpackSubstrateHeaderIn(Substrate.MaterialTextureArray, SubstrateAddressing, Substrate.TopLayerTexture);
@@ -72,16 +74,26 @@ FProjectionShadingInfo GetProjectionShadingInfo(uint2 PixelPos)
 	Out.bIsHair = BSDFType == SUBSTRATE_BSDF_TYPE_HAIR;
 	Out.WorldNormal = TopLayerData.WorldNormal;
 	Out.bIsSubsurface = SSSType != SSS_TYPE_INVALID; // legacy IsSubsurfaceModel is true for subsurface profile, eyes, etc.
-	Out.bIsClearCoat = false;
 
+	// VSM traces traces using the top layer normal when NoL>0 only. So if normals for slabs below do not match, light leaking will occur.
+	// For instance: a flat translucent top layer above a normal mapped bottom layer.
+	// That is why, for complex materials with multiple BSDFs, we also must not do backface culling.
+	if (SubstratePixelHeader.ClosureCount > 1)
+	{
+		// Technically, this could only be done when vertical layering is encountered only. 
+		// But we do not have that data available (or it would be costly to infer from closures)
+		Out.bDisableBackFaceCulling = true;
+}
 #else
 	const FGBufferData GBufferData = GetGBufferDataUint(PixelPos, true);
 	Out.bIsValid = GBufferData.ShadingModelID != SHADINGMODELID_UNLIT;
 	Out.bIsHair = GBufferData.ShadingModelID == SHADINGMODELID_HAIR;
 	Out.bIsSubsurface = IsSubsurfaceModel(GBufferData.ShadingModelID);
-	Out.bIsClearCoat = GBufferData.ShadingModelID == SHADINGMODELID_CLEAR_COAT;
 	Out.WorldNormal = GBufferData.WorldNormal;
 #endif
+
+	Out.bDisableBackFaceCulling = Out.bDisableBackFaceCulling || Out.bIsSubsurface || Out.bIsHair;
+
 	return Out;
 }
 
@@ -171,7 +183,7 @@ FVirtualShadowMapSampleResult ProjectLight(
 {
 	const bool bIsHairInput = InputType == INPUT_TYPE_HAIRSTRANDS;
 	
-	bool bBackfaceCull = (bCullBackfacingPixels > 0) && !bIsHairInput && !ShadingInfo.bIsSubsurface && !ShadingInfo.bIsHair && !ShadingInfo.bIsClearCoat;
+	bool bBackfaceCull = (bCullBackfacingPixels > 0) && !bIsHairInput && !ShadingInfo.bDisableBackFaceCulling;
 #if VISUALIZE
 	// Modes more related to the shadow depth rendering don't want to see the projection culling
 	bBackfaceCull = bBackfaceCull &&
