Fix Substrate shader compilation error when dead code stripping is disabled.

diff --git a/Engine/Shaders/Private/Substrate/SubstrateVisualize.usf b/Engine/Shaders/Private/Substrate/SubstrateVisualize.usf
--- a/Engine/Shaders/Private/Substrate/SubstrateVisualize.usf
+++ b/Engine/Shaders/Private/Substrate/SubstrateVisualize.usf
@@ -15,7 +15,10 @@
 #define SUBSTRATE_COMPLEXSPECIALPATH 1				// Allow to read data from complex special path
 #include "/Engine/Private/Substrate/SubstrateEvaluation.ush"
 #include "/Engine/Private/Substrate/SubstrateTile.ush"
+
+#if SUBSTRATE_INLINE_SHADING == 0
 #include "/Engine/Private/Substrate/SubstrateVisualizeCommon.ush"
+#endif
 
 #include "/Engine/Private/ShaderPrint.ush"
 #include "../ColorMap.ush"
