Fix DOF in right view when using stereo rendering.
During the DOF setup, the full res output was viewrect-aligned, while the half res output was top-left aligned. The temporal AA pass explicitly expected top-left aligned input, based on the TopLeftCornerViewRects() call introduced in CL 4002919 for dynres. This caused out-of-bounds sampling, leading to artifacts.
With this CL, the setup output and TAA input are now always viewrect-aligned.

diff --git a/Engine/Shaders/Private/DiaphragmDOF/DOFSetup.usf b/Engine/Shaders/Private/DiaphragmDOF/DOFSetup.usf
--- a/Engine/Shaders/Private/DiaphragmDOF/DOFSetup.usf
+++ b/Engine/Shaders/Private/DiaphragmDOF/DOFSetup.usf
@@ -148,7 +148,7 @@ void SetupCS(
 		float OutCocRadius;
 		DownsampleSceneColorWithCoc(DownsampleParams, SceneColors, CocRadii, OutColor, OutCocRadius);
 	
-		uint2 HalfResOutputPixelPosition = GroupId * (GROUP_BORDER_SIZE / 2) + SubGroupPos;
+		uint2 HalfResOutputPixelPosition = (View.ViewRectMin.xy/2) + GroupId * (GROUP_BORDER_SIZE / 2) + SubGroupPos;
 
 		//if (all(HalfResOutputPixelPosition < ViewportRect.zw))
 		#if CONFIG_DOF_ALPHA
@@ -272,7 +272,7 @@ void SetupCS(
 		#if 0 // Lower VGPR footprint.
 			uint2 OutputPixelPosition = 0.5 * SceneBufferUV * View.BufferSizeAndInvSize.xy;
 		#else
-			uint2 OutputPixelPosition = DispatchThreadId;
+			uint2 OutputPixelPosition = (View.ViewRectMin.xy/2) + DispatchThreadId;
 		#endif
 
 		if (all(OutputPixelPosition < ViewportRect.zw))
diff --git a/Engine/Source/Runtime/Renderer/Private/PostProcess/DiaphragmDOF.cpp b/Engine/Source/Runtime/Renderer/Private/PostProcess/DiaphragmDOF.cpp
--- a/Engine/Source/Runtime/Renderer/Private/PostProcess/DiaphragmDOF.cpp
+++ b/Engine/Source/Runtime/Renderer/Private/PostProcess/DiaphragmDOF.cpp
@@ -1490,7 +1490,6 @@ bool DiaphragmDOF::AddPasses(
 		}
 		
 		TAAParameters.SetupViewRect(View, PrefilteringResolutionDivisor);
-		TAAParameters.TopLeftCornerViewRects();
 
 		TAAParameters.Quality = CVarDOFTemporalAAQuality.GetValueOnRenderThread() == 0 ? ETAAQuality::Medium : ETAAQuality::High;
 	}
diff --git a/Engine/Source/Runtime/Renderer/Private/PostProcess/TemporalAA.h b/Engine/Source/Runtime/Renderer/Private/PostProcess/TemporalAA.h
--- a/Engine/Source/Runtime/Renderer/Private/PostProcess/TemporalAA.h
+++ b/Engine/Source/Runtime/Renderer/Private/PostProcess/TemporalAA.h
@@ -150,15 +150,6 @@ struct FTAAPassParameters
 			OutputViewRect = InputViewRect;
 		}
 	}
-
-	// Shifts input and output view rect to top left corner
-	FORCEINLINE void TopLeftCornerViewRects()
-	{
-		InputViewRect.Max -= InputViewRect.Min;
-		InputViewRect.Min = FIntPoint::ZeroValue;
-		OutputViewRect.Max -= OutputViewRect.Min;
-		OutputViewRect.Min = FIntPoint::ZeroValue;
-	}
 	
 	/** Returns the texture resolution that will be output. */
 	FIntPoint GetOutputExtent() const;
