Substrate - fixed holdout on opaque primitives.

#rnx

diff --git a/Engine/Shaders/Private/BasePassPixelShader.usf b/Engine/Shaders/Private/BasePassPixelShader.usf
--- a/Engine/Shaders/Private/BasePassPixelShader.usf
+++ b/Engine/Shaders/Private/BasePassPixelShader.usf
@@ -70,6 +70,8 @@
 #define USE_DBUFFER 0
 #endif
 
+#define SUPPORT_MATERIAL_PRIMITIVE_ALPHA_HOLDOUT (SUPPORT_PRIMITIVE_ALPHA_HOLDOUT && MATERIAL_DOMAIN_SURFACE)
+
 // Enable IES profile evaluation based on project settings
 #if SUPPORT_IESPROFILE_ON_FORWARD_LIT_TRANSLUCENT
 #define USE_IES_PROFILE 1
@@ -740,6 +742,16 @@ uint GetDiffuseIndirectSampleOcclusion(FSharedLocalBases SharedLocalBases, float
 #endif
 	return DiffuseIndirectSampleOcclusion;
 }
+
+#if SUPPORT_MATERIAL_PRIMITIVE_ALPHA_HOLDOUT
+void ApplyAlphaHoldOutToSubstrateOpaque(inout FSubstratePixelHeader SubstratePixelHeader, inout FSubstrateData SubstrateData)
+{
+	// Set a default null material to appear black. Alpha will remain correct.
+	SubstrateData = GetInitialisedSubstrateData();
+	SubstratePixelHeader = InitialiseSubstratePixelHeader();
+}
+#endif
+
 #else // SUBSTRATE_ENABLED
 uint GetDiffuseIndirectSampleOcclusion(FGBufferData GBuffer, float3 V, float3 WorldNormal, float3 WorldBentNormal, float2 SvPosition, float MaterialAO)
 {
@@ -894,6 +906,14 @@ void FPixelShaderInOut_MainPS(
 
 	const float Dither = InterleavedGradientNoise(MaterialParameters.SvPosition.xy, View.StateFrameIndexMod8);
 
+#if SUPPORT_MATERIAL_PRIMITIVE_ALPHA_HOLDOUT
+#if (MATERIALBLENDING_ALPHAHOLDOUT || SUBSTRATE_BLENDING_ALPHAHOLDOUT)
+	const bool bIsHoldout = true;
+#else
+	const bool bIsHoldout = (GetPrimitiveData(MaterialParameters).Flags & PRIMITIVE_SCENE_DATA_FLAG_HOLDOUT);
+#endif
+#endif
+
 #if !SUBSTRATE_ENABLED
 	// Store the results in local variables and reuse instead of calling the functions multiple times.
 	half3 BaseColor = GetMaterialBaseColor(PixelMaterialInputs);
@@ -1711,6 +1731,13 @@ void FPixelShaderInOut_MainPS(
 		Settings.SliceStoringDebugSubstrateTreeData = all(uint2(float2(View.CursorPosition) * View.ViewResolutionFraction) == PixelPos) ? SubstrateStruct.SliceStoringDebugSubstrateTreeDataWithoutMRT : Settings.SliceStoringDebugSubstrateTreeData;
 	#endif
 
+	#if SUPPORT_MATERIAL_PRIMITIVE_ALPHA_HOLDOUT
+		if (bIsHoldout)
+		{
+			ApplyAlphaHoldOutToSubstrateOpaque(SubstratePixelHeader, SubstrateData);
+		}
+	#endif
+
 		// Generate the Substrate material data to write out
 		FSubstrateAddressing SubstrateAddressing = GetSubstratePixelDataByteOffset(PixelPos, uint2(ResolvedView.BufferSizeAndInvSize.xy), SubstrateStruct.MaxBytesPerPixel);
 		FRWSubstrateMaterialContainer RWSubstrateMaterialContainer = InitialiseRWSubstrateMaterialContainer(SubstrateStruct.MaterialTextureArrayUAVWithoutRTs, bAllowWrite);
@@ -2188,7 +2215,7 @@ void FPixelShaderInOut_MainPS(
 	Out.MRT[0] = SubstrateMaterialExport;
 #endif
 
-#if SUPPORT_PRIMITIVE_ALPHA_HOLDOUT && MATERIAL_DOMAIN_SURFACE
+#if SUPPORT_MATERIAL_PRIMITIVE_ALPHA_HOLDOUT
 		
 	//Calculate IoR for translucent material with ior and update holdout property
 	#if (MATERIALBLENDING_TRANSLUCENT || SUBSTRATE_TRANSLUCENT_MATERIAL) && REFRACTION_USE_INDEX_OF_REFRACTION
@@ -2211,13 +2238,6 @@ void FPixelShaderInOut_MainPS(
 		#endif
 	#endif 
 	
-	
-#if (MATERIALBLENDING_ALPHAHOLDOUT || SUBSTRATE_BLENDING_ALPHAHOLDOUT)
-	const bool bIsHoldout = true;
-#else
-	const bool bIsHoldout = (GetPrimitiveData(MaterialParameters).Flags & PRIMITIVE_SCENE_DATA_FLAG_HOLDOUT);
-#endif
-	
 	// Only include the holdout logic if the Renderer is forwarding its alpha channel to post-processing
 	if (bIsHoldout)
 	{
@@ -2338,7 +2358,7 @@ void FPixelShaderInOut_MainPS(
 		
 	}
 #endif
-#endif // SUPPORT_PRIMITIVE_ALPHA_HOLDOUT && MATERIAL_DOMAIN_SURFACE
+#endif // SUPPORT_MATERIAL_PRIMITIVE_ALPHA_HOLDOUT
 
 }
 
