Substrate - fixed diffuse sampling with bent normal to match legacy.

#rnx
#tests none

diff --git a/Engine/Shaders/Private/Substrate/SubstrateLightingCommon.ush b/Engine/Shaders/Private/Substrate/SubstrateLightingCommon.ush
--- a/Engine/Shaders/Private/Substrate/SubstrateLightingCommon.ush
+++ b/Engine/Shaders/Private/Substrate/SubstrateLightingCommon.ush
@@ -233,7 +233,7 @@ void SubstrateEnvLightingCommon(
 		if (bProcessFrontFaceDiffuse)
 		{
 			// Finally sample the sky diffuse contribution (spherical harmonic, Lambert BRDF)
-			float3 DiffuseLookup = GetEnvDiffuseLighting(BentNormal);
+			float3 DiffuseLookup = GetEnvDiffuseLighting(SkyVisData.SkyDiffuseLookUpNormal);
 			// And accumulate
 			// Note: Use diffuse directional albedo (i.e., DiffuseWeight) as first order approximation for env. integration (SUBSTRATE_TODO instead compute SH coefficients for Chan)
 			DiffuseLighting = CloudVolumetricAOShadow * BSDFThroughput * (SkyVisData.SkyDiffuseLookUpMul * DiffuseLookup + SkyVisData.SkyDiffuseLookUpAdd) * SubstrateEnvLight.DiffuseWeight;
@@ -244,7 +244,7 @@ void SubstrateEnvLightingCommon(
 			// However, we do apply SkyVisData.SkyDiffuseLookUpMul for scaling consistency.
 
 			// Finally sample the sky diffuse contribution (spherical harmonic, Lambert BRDF) along the opposite normal direction
-			float3 DiffuseLookup = GetEnvDiffuseLighting(-BentNormal);
+			float3 DiffuseLookup = GetEnvDiffuseLighting(-SkyVisData.SkyDiffuseLookUpNormal);
 			// And accumulate
 			// Note: Use diffuse directional albedo (i.e., DiffuseWeight) as first order approximation for env. integration (SUBSTRATE_TODO instead compute SH coefficients for Chan)
 			DiffuseLighting += CloudVolumetricAOShadow * BSDFThroughput * (SkyVisData.SkyDiffuseLookUpMul * DiffuseLookup) * SubstrateEnvLight.DiffuseBackFaceWeight;
