Substrate - fixed upsampling of capsule shadow.

#rnx
#tests none

diff --git a/Engine/Shaders/Private/CapsuleShadowShaders.usf b/Engine/Shaders/Private/CapsuleShadowShaders.usf
--- a/Engine/Shaders/Private/CapsuleShadowShaders.usf
+++ b/Engine/Shaders/Private/CapsuleShadowShaders.usf
@@ -926,8 +926,7 @@ void ApplySelfShadowingIntensityForDeferred(float2 ScreenUV, inout float Visibil
 	if (View.IndirectCapsuleSelfShadowingIntensity < 1)
 	{
 #if SUBSTRATE_ENABLED
-		uint2 PixelPos = View.ViewRectMin.xy + View.ViewSizeAndInvSize.xy * ScreenUV.xy;
-		FSubstrateAddressing SubstrateAddressing = GetSubstratePixelDataByteOffset(PixelPos, uint2(View.BufferSizeAndInvSize.xy), Substrate.MaxBytesPerPixel);
+		FSubstrateAddressing SubstrateAddressing = GetSubstratePixelDataByteOffset(SVPos.xy, uint2(View.BufferSizeAndInvSize.xy), Substrate.MaxBytesPerPixel);
 		FSubstratePixelHeader SubstratePixelHeader = UnpackSubstrateHeaderIn(Substrate.MaterialTextureArray, SubstrateAddressing, Substrate.TopLayerTexture);
 		const bool bHasDynamicIndirectShadowCasterRepresentation = SubstratePixelHeader.HasDynamicIndirectShadowCasterRepresentation();
 #else
