Path Tracer: Refactor transparency handling so that opaque surfaces are traced first, enabling support for DepthFade nodes

There are some material features that depend on the ability ot read from the SceneDepth (either explicitly or through the DepthFade node). To support this, change the path tracer to trace opaque surfaces first, and then go into the loop to handle transparencies and volumes.

Modify ray flags to account for transparent blend modes vs opaque/masked ones. Split out dithered masking since this gets converted to transparency in the path tracer. Note that since there are cases that involve multiple materials on a single mesh - we might still have some transparency during the opaque trace. Luckily these materials all have an AHS and we can set a payload flag to ignore them for the opaque trace.

When tracing opaque surfaces first, the decal logic can be moved out of the transparency loop which makes it a bit lighter, and simplifies the control flow.

Readjust cvar defaults based on testing more scenes. CompactionDepth of 4 is usually slightly better, but occasionally much worse. It seems like the slowdowns outweight the slight improvements, and this reduces RTPSO compilation time a bit.

diff --git a/Engine/Plugins/Experimental/GPULightmass/Source/GPULightmass/Private/LightmapRenderer.cpp b/Engine/Plugins/Experimental/GPULightmass/Source/GPULightmass/Private/LightmapRenderer.cpp
--- a/Engine/Plugins/Experimental/GPULightmass/Source/GPULightmass/Private/LightmapRenderer.cpp
+++ b/Engine/Plugins/Experimental/GPULightmass/Source/GPULightmass/Private/LightmapRenderer.cpp
@@ -37,7 +37,7 @@
 #include "PrimitiveUniformShaderParametersBuilder.h"
 #include "PrimitiveSceneShaderData.h"
 
-RENDERER_API uint8 BlendModeToRayTracingInstanceMask(const EBlendMode BlendMode, bool bCastShadow, ERayTracingViewMaskMode MaskMode);
+RENDERER_API uint8 BlendModeToRayTracingInstanceMask(const EBlendMode BlendMode, bool bIsDitherMasked, bool bCastShadow, ERayTracingViewMaskMode MaskMode);
 
 class FCopyConvergedLightmapTilesCS : public FGlobalShader
 {
@@ -640,7 +640,7 @@ struct FRayTracingMaskAndStatus
 			bAllSegmentsUnlit &= Material.GetShadingModels().HasOnlyShadingModel(MSM_Unlit) || !MeshBatch.CastShadow;
 			bAllSegmentsOpaque &= Material.GetBlendMode() == EBlendMode::BLEND_Opaque;
 			bAnySegmentsCastShadow |= bSegmentCastsShadow;
-			InstanceMask |= BlendModeToRayTracingInstanceMask(Material.GetBlendMode(), bSegmentCastsShadow, MaskMode);
+			InstanceMask |= BlendModeToRayTracingInstanceMask(Material.GetBlendMode(), Material.IsDitherMasked(), bSegmentCastsShadow, MaskMode);
 		}
 	}
 };
diff --git a/Engine/Shaders/Private/PathTracing/PathTracingCommon.ush b/Engine/Shaders/Private/PathTracing/PathTracingCommon.ush
--- a/Engine/Shaders/Private/PathTracing/PathTracingCommon.ush
+++ b/Engine/Shaders/Private/PathTracing/PathTracingCommon.ush
@@ -30,6 +30,8 @@
 #define PATH_TRACING_PAYLOAD_INPUT_FLAG_RAY_TYPE_INDIRECT_DIFFUSE      2
 #define PATH_TRACING_PAYLOAD_INPUT_FLAG_RAY_TYPE_INDIRECT_SPECULAR     3
 #define PATH_TRACING_PAYLOAD_INPUT_FLAG_RAY_TYPE_INDIRECT_VOLUME       4
+#define PATH_TRACING_PAYLOAD_INPUT_FLAG_HAS_SCENE_DEPTH               (1 << 4) // Indicates the payload contains a scene depth value (exclusive with dbuffer usage)
+#define PATH_TRACING_PAYLOAD_INPUT_FLAG_IGNORE_TRANSLUCENT            (1 << 5) // Indicates this ray should not process any translucent geometry (so we can skip them in our opaque traces)
 
 #define PATH_TRACING_PAYLOAD_NORMAL_BITS  32
 
@@ -369,6 +371,17 @@ struct FPackedPathTracingPayload : FMinimalPayload
 		return asfloat(PackedData[14]);
 	}
 
+	float GetSceneDepth()
+	{
+		// NOTE: This is aliased with GetDBufferA, but this is ok as the two features may not be used together
+		return asfloat(PackedData[8]);
+	}
+
+	void SetSceneDepth(float SceneZ)
+	{
+		PackedData[8] = asuint(SceneZ);
+	}
+
 	// Volumetric callable shaders
 	void SetVolumetricCallableShaderInput(FDFVector3 AbsolutePosition, float CloudHeightKm)
 	{
diff --git a/Engine/Shaders/Private/PathTracing/PathTracingCore.ush b/Engine/Shaders/Private/PathTracing/PathTracingCore.ush
--- a/Engine/Shaders/Private/PathTracing/PathTracingCore.ush
+++ b/Engine/Shaders/Private/PathTracing/PathTracingCore.ush
@@ -422,7 +422,7 @@ void EvaluateDecals(FPathState PathState, float HitT, inout FDecalShaderPayload
 		DecalRay.TMin = 0.0f;
 		DecalRay.TMax = HitT + MeshDecalBias;
 
-		const uint InstanceInclusionMask = PathState.GetFirstScatterType() == PATHTRACER_SCATTER_CAMERA ? PATHTRACER_MASK_CAMERA : PATHTRACER_MASK_INDIRECT;
+		const uint InstanceInclusionMask = PathState.GetFirstScatterType() == PATHTRACER_SCATTER_CAMERA ? (PATHTRACER_MASK_CAMERA | PATHTRACER_MASK_CAMERA_TRANSLUCENT) : (PATHTRACER_MASK_INDIRECT | PATHTRACER_MASK_INDIRECT_TRANSLUCENT);
 		const uint MissShaderIndex = 0; // TODO
 
 		for (;;)
@@ -484,6 +484,97 @@ FPackedPathTracingPayload TraceTransparentRay(inout FPathState PathState, int Bo
 	FRISContext HitSample = InitRISContext(RISRandSample.x);
 	FRISContext VolSample = InitRISContext(RISRandSample.y);
 
+
+#if PATH_TRACER_TRACE_OPAQUE_FIRST
+	FPackedPathTracingPayload OpaquePayload = InitPathTracingPayload(PathState.GetFirstScatterType(), ApproximateCaustics ? PathState.GetPathRoughness() : 0.0);
+	OpaquePayload.SetDBufferA(float4(0, 0, 0, 1));
+	OpaquePayload.SetDBufferB(float4(0, 0, 0, 1));
+	OpaquePayload.SetDBufferC(float4(0, 0, 0, 1));
+	OpaquePayload.SetStochasticSlabRand(RISRandSample.z);
+	OpaquePayload.SetFlag(PATH_TRACING_PAYLOAD_INPUT_FLAG_IGNORE_TRANSLUCENT);
+#if PATH_TRACER_USE_SER 
+	{
+		NvHitObject Hit;
+		NvTraceRayHitObject(
+			TLAS,
+			RayFlags,
+			InstanceInclusionMask,
+			RAY_TRACING_SHADER_SLOT_MATERIAL,
+			RAY_TRACING_NUM_SHADER_SLOTS,
+			MissShaderIndex,
+			PathState.Ray.GetNativeDesc(),
+			OpaquePayload,
+			Hit);
+		if (!bIsCameraRay)
+		{
+			NvReorderThread(Hit);
+		}
+		NvInvokeHitObject(TLAS, Hit, OpaquePayload);
+	}
+#else
+	// Trace the ray
+	TraceRay(
+		TLAS,
+		RayFlags,
+		InstanceInclusionMask,
+		RAY_TRACING_SHADER_SLOT_MATERIAL,
+		RAY_TRACING_NUM_SHADER_SLOTS,
+		MissShaderIndex,
+		PathState.Ray.GetNativeDesc(),
+		OpaquePayload);
+#endif
+	if (OpaquePayload.IsDecalReceiver() && PathState.GetPathRoughness() < max(MeshDecalRoughnessCutoff, DecalRoughnessCutoff))
+	{
+		FDecalShaderPayload DecalPayload = InitDecalShaderPayload(PathState.Ray.Origin + OpaquePayload.HitT * PathState.Ray.Direction);
+		EvaluateDecals(PathState, OpaquePayload.HitT, DecalPayload);
+#if USE_DBUFFER
+		if (OpaquePayload.UsesDBufferLookup() && HasDecalData(DecalPayload))
+		{
+			if (EnableDBuffer)
+			{
+				// Retrace ray with DBuffer data
+				OpaquePayload = InitPathTracingPayload(PathState.GetFirstScatterType(), ApproximateCaustics ? PathState.GetPathRoughness() : 0.0);
+
+				OpaquePayload.SetDBufferA(DecalPayload.GetBaseColor());
+				OpaquePayload.SetDBufferB(DecalPayload.GetWorldNormal());
+				OpaquePayload.SetDBufferC(DecalPayload.GetMetallicSpecularRoughness());
+				OpaquePayload.SetStochasticSlabRand(RISRandSample.z); // Must use same random value to get a consistent result to the first call
+				OpaquePayload.SetFlag(PATH_TRACING_PAYLOAD_INPUT_FLAG_IGNORE_TRANSLUCENT);
+				TraceRay(
+					TLAS,
+					RayFlags,
+					InstanceInclusionMask,
+					RAY_TRACING_SHADER_SLOT_MATERIAL,
+					RAY_TRACING_NUM_SHADER_SLOTS,
+					MissShaderIndex,
+					PathState.Ray.GetNativeDesc(),
+					OpaquePayload);
+			}
+			else
+			{
+				// The user asked for a dbuffer lookup, but we have disabled the double-trace.
+				// Set all the decal response flags so we at least pickup some kind of response
+				// even if it won't be exactly what the user had in mind.
+				OpaquePayload.SetFlag(PATH_TRACING_PAYLOAD_OUTPUT_FLAG_DECAL_RESPONSE_COLOR | PATH_TRACING_PAYLOAD_OUTPUT_FLAG_DECAL_RESPONSE_NORMAL | PATH_TRACING_PAYLOAD_OUTPUT_FLAG_DECAL_RESPONSE_ROUGHNESS);
+			}
+		}
+#endif
+		// Now that we know the real decal data, apply it
+		ApplyDecal(DecalPayload, OpaquePayload, PathState.Ray.Direction, PathState.GetPathRoughness());
+	}
+
+	InstanceInclusionMask = bUseCameraRayType ? PATHTRACER_MASK_CAMERA_TRANSLUCENT : PATHTRACER_MASK_INDIRECT_TRANSLUCENT;
+
+	if (OpaquePayload.IsHit())
+	{
+		PathState.Ray.TMax = asfloat(asuint(OpaquePayload.HitT) - 1);
+	}
+	
+
+	float3 PayloadThroughput = PathState.PathThroughput;
+	FPackedPathTracingPayload Payload;
+#else // PATH_TRACER_TRACE_OPAQUE_FIRST
+	InstanceInclusionMask |= bUseCameraRayType ? PATHTRACER_MASK_CAMERA_TRANSLUCENT : PATHTRACER_MASK_INDIRECT_TRANSLUCENT;
 	float3 PayloadThroughput = PathState.PathThroughput;
 	FPackedPathTracingPayload Payload;
 	if (bLastBounce && !PathState.ShouldAccumulateEmissive())
@@ -492,6 +583,8 @@ FPackedPathTracingPayload TraceTransparentRay(inout FPathState PathState, int Bo
 		PathState.PathThroughput = 0;
 		return Payload;
 	}
+#endif
+
 	for (;;)
 	{
 		FPackedPathTracingPayload PackedPayload = InitPathTracingPayload(PathState.GetFirstScatterType(), ApproximateCaustics ? PathState.GetPathRoughness() : 0.0);
@@ -499,6 +592,16 @@ FPackedPathTracingPayload TraceTransparentRay(inout FPathState PathState, int Bo
 		PackedPayload.SetDBufferB(float4(0, 0, 0, 1));
 		PackedPayload.SetDBufferC(float4(0, 0, 0, 1));
 		PackedPayload.SetStochasticSlabRand(RISRandSample.z);
+#if PATH_TRACER_TRACE_OPAQUE_FIRST
+		if (OpaquePayload.IsHit())
+		{
+			const float3 ViewZ = View.ViewToTranslatedWorld[2].xyz;
+			const float3 TranslatedWorldPos = PathState.Ray.Origin + OpaquePayload.HitT * PathState.Ray.Direction;
+			PackedPayload.SetFlag(PATH_TRACING_PAYLOAD_INPUT_FLAG_HAS_SCENE_DEPTH);
+			PackedPayload.SetSceneDepth(dot(TranslatedWorldPos, ViewZ));
+		}
+#endif
+
 #if PATH_TRACER_USE_SER 
 		{
 			NvHitObject Hit;
@@ -530,6 +633,14 @@ FPackedPathTracingPayload TraceTransparentRay(inout FPathState PathState, int Bo
 			PathState.Ray.GetNativeDesc(),
 			PackedPayload);
 #endif
+
+#if PATH_TRACER_TRACE_OPAQUE_FIRST
+		if (PackedPayload.IsMiss())
+		{
+			// If we miss the translucent geometry, the next hit is the opaque one
+			PackedPayload = OpaquePayload;
+		}
+#endif
 		FVolumeIntersectionList VolumeIsectList = VolumeIntersect(PathState.Ray.Origin, PathState.Ray.Direction, PathState.Ray.TMin, PathState.Ray.TMax, PathState.GetPathRoughness(), VolumeFlags);
 		if (VolumeIsectList.HitBlocker() && (PackedPayload.IsMiss() || VolumeIsectList.BlockerHitT < PackedPayload.HitT))
 		{
@@ -731,6 +842,7 @@ FPackedPathTracingPayload TraceTransparentRay(inout FPathState PathState, int Bo
 			break;
 		}
 
+#if PATH_TRACER_TRACE_OPAQUE_FIRST == 0
 		if (PackedPayload.IsDecalReceiver() && PathState.GetPathRoughness() < max(MeshDecalRoughnessCutoff, DecalRoughnessCutoff))
 		{
 			FDecalShaderPayload DecalPayload = InitDecalShaderPayload(PathState.Ray.Origin + PackedPayload.HitT * PathState.Ray.Direction);
@@ -769,7 +881,7 @@ FPackedPathTracingPayload TraceTransparentRay(inout FPathState PathState, int Bo
 			// Now that we know the real decal data, apply it
 			ApplyDecal(DecalPayload, PackedPayload, PathState.Ray.Direction, PathState.GetPathRoughness());
 		}
-
+#endif
 		{
 			// account for Beer's law through the currently active medium
 			// TODO: merge this with the volume raymarching?
@@ -1908,5 +2020,7 @@ bool PathTracingKernel(inout FPathState PathState, int Bounce)
 	}
 
 	// continue the path only if the relevant indirect contribution is enabled
-	return PathState.UpdateScatterType(MaterialSample.ScatterType);
+	bool bKeepGoing = PathState.UpdateScatterType(MaterialSample.ScatterType);
+	// also check for last bounce if emissive contributions are not needed
+	return bKeepGoing && (Bounce + 1 < MaxBounces || PathState.ShouldAccumulateEmissive());
 }
diff --git a/Engine/Shaders/Private/PathTracing/PathTracingMaterialHitShader.usf b/Engine/Shaders/Private/PathTracing/PathTracingMaterialHitShader.usf
--- a/Engine/Shaders/Private/PathTracing/PathTracingMaterialHitShader.usf
+++ b/Engine/Shaders/Private/PathTracing/PathTracingMaterialHitShader.usf
@@ -175,11 +175,15 @@ RAY_TRACING_ENTRY_CLOSEST_HIT(PathTracingMaterialCHS,
 
 
 	CurrentPayloadInputFlags = PackedPayload.GetFlags();
+	PathTracerSceneDepth = (CurrentPayloadInputFlags & PATH_TRACING_PAYLOAD_INPUT_FLAG_HAS_SCENE_DEPTH) ? PackedPayload.GetSceneDepth() : SCENE_TEXTURES_DISABLED_SCENE_DEPTH_VALUE;
 
 #if USE_DBUFFER && MATERIAL_USES_DECAL_LOOKUP
-	CurrentPayloadDBufferA = PackedPayload.GetDBufferA();
-	CurrentPayloadDBufferB = PackedPayload.GetDBufferB();
-	CurrentPayloadDBufferC = PackedPayload.GetDBufferC();
+	if ((CurrentPayloadInputFlags & PATH_TRACING_PAYLOAD_INPUT_FLAG_HAS_SCENE_DEPTH) == 0)
+	{
+		CurrentPayloadDBufferA = PackedPayload.GetDBufferA();
+		CurrentPayloadDBufferB = PackedPayload.GetDBufferB();
+		CurrentPayloadDBufferC = PackedPayload.GetDBufferC();
+	}
 #endif
 
 #if VF_SUPPORTS_RAYTRACING_PREPARE_MATERIAL_PIXEL_PARAMETERS
@@ -1207,6 +1211,10 @@ RAY_TRACING_ENTRY_ANY_HIT(PathTracingMaterialAHS,
 #if MATERIALBLENDING_MODULATE || (MATERIALBLENDING_MASKED && MATERIAL_DITHER_OPACITY_MASK) || MATERIALBLENDING_ALPHACOMPOSITE || MATERIALBLENDING_TRANSLUCENT
 	if (!PackedPayload.IsVisibilityRay())
 	{
+		if (PackedPayload.GetFlags() & PATH_TRACING_PAYLOAD_INPUT_FLAG_IGNORE_TRANSLUCENT)
+		{
+			IgnoreHit();
+		}
 		// not a shadow ray -- don't need to run the AHS logic
 		return;
 	}
@@ -1224,6 +1232,7 @@ RAY_TRACING_ENTRY_ANY_HIT(PathTracingMaterialAHS,
 	const float4 SvPosition = TranslatedWorldPositionToSvPosition(TranslatedWorldPosition);
 
 	CurrentPayloadInputFlags = PackedPayload.GetFlags();
+	PathTracerSceneDepth = (CurrentPayloadInputFlags & PATH_TRACING_PAYLOAD_INPUT_FLAG_HAS_SCENE_DEPTH) ? PackedPayload.GetSceneDepth() : SCENE_TEXTURES_DISABLED_SCENE_DEPTH_VALUE;
 
 #if VF_SUPPORTS_RAYTRACING_PREPARE_MATERIAL_PIXEL_PARAMETERS
 	// this is a newer codepath that is both more flexible and allows for more direct calculation compared to the other codepath
diff --git a/Engine/Shaders/Private/PathTracing/PathTracingMissShader.usf b/Engine/Shaders/Private/PathTracing/PathTracingMissShader.usf
--- a/Engine/Shaders/Private/PathTracing/PathTracingMissShader.usf
+++ b/Engine/Shaders/Private/PathTracing/PathTracingMissShader.usf
@@ -5,4 +5,5 @@
 RAY_TRACING_ENTRY_MISS(PathTracingDefaultMS, FPackedPathTracingPayload, Payload)
 {
 	Payload.SetMiss();
+	Payload.SetFlag(0); // clear input flags
 }
diff --git a/Engine/Shaders/Private/SceneTexturesCommon.ush b/Engine/Shaders/Private/SceneTexturesCommon.ush
--- a/Engine/Shaders/Private/SceneTexturesCommon.ush
+++ b/Engine/Shaders/Private/SceneTexturesCommon.ush
@@ -49,10 +49,16 @@ float4 CalcFullSceneColor(float2 ScreenUV)
 #endif
 }
 
+#ifdef PATH_TRACING
+static float PathTracerSceneDepth = 0.0;
+#endif
+
 /** Fetches device depth and converts to linear. */
 float CalcSceneDepth(float2 ScreenUV)
 {
-#if SCENE_TEXTURES_DISABLED
+#if PATH_TRACING
+	return PathTracerSceneDepth;
+#elif SCENE_TEXTURES_DISABLED
 	return SCENE_TEXTURES_DISABLED_SCENE_DEPTH_VALUE;
 #else
 	return ConvertFromDeviceZ(Texture2DSampleLevel(SceneTexturesStruct.SceneDepthTexture, SceneTexturesStruct_SceneDepthTextureSampler, ScreenUV, 0).r);
diff --git a/Engine/Shaders/Shared/PathTracingDefinitions.h b/Engine/Shaders/Shared/PathTracingDefinitions.h
--- a/Engine/Shaders/Shared/PathTracingDefinitions.h
+++ b/Engine/Shaders/Shared/PathTracingDefinitions.h
@@ -28,8 +28,8 @@
 #define PATHTRACER_MASK_HAIR_SHADOW							0x08			// Whether hair is visible for shadow rays
 #define PATHTRACER_MASK_INDIRECT							0x10			// opaque and alpha tested meshes and particles as a whole (indirect ray) excluding hairs
 #define PATHTRACER_MASK_HAIR_INDIRECT						0x20			// For indirect ray tracing against hair
-#define PATHTRACER_MASK_EMPTY_SLOT1							0x40			
-#define PATHTRACER_MASK_EMPTY_SLOT2							0x80			
+#define PATHTRACER_MASK_CAMERA_TRANSLUCENT					0x40			
+#define PATHTRACER_MASK_INDIRECT_TRANSLUCENT				0x80			
 
 #define PATHTRACER_MASK_IGNORE								0x00			// used when mapping general tracing mask to path tracing mask
 #define PATHTRACER_MASK_UNUSED	(PATHTRACER_MASK_EMPTY_SLOT1|PATHTRACER_MASK_EMPTY_SLOT2)	
diff --git a/Engine/Source/Runtime/Renderer/Private/PathTracing.cpp b/Engine/Source/Runtime/Renderer/Private/PathTracing.cpp
--- a/Engine/Source/Runtime/Renderer/Private/PathTracing.cpp
+++ b/Engine/Source/Runtime/Renderer/Private/PathTracing.cpp
@@ -62,8 +62,15 @@ TAutoConsoleVariable<int32> CVarPathTracingDebugVisualizeLightSamples(
 
 TAutoConsoleVariable<int32> CVarPathTracingCompactionDepth(
 	TEXT("r.PathTracing.CompactionDepth"),
-	4,
-	TEXT("Enables path compaction to improve GPU occupancy for the path tracer. The value sets the bounce up to which compaction will happen, beyond that point the path tracer will handle bounces within the dispatch."),
+	-1,
+	TEXT("Enables path compaction to improve GPU occupancy for the path tracer. The value sets the bounce up to which compaction will happen, beyond that point the path tracer will handle bounces within the dispatch. (default: -1, all bounces)"),
+	ECVF_RenderThreadSafe
+);
+
+TAutoConsoleVariable<bool> CVarPathTracingTraceOpaqueFirst(
+	TEXT("r.PathTracing.TraceOpaqueFirst"),
+	true,
+	TEXT("Trace opaque geometry before translucent geometry. This allows the path tracer to setup a correct depth for DepthFade based effects (default: 0)"),
 	ECVF_RenderThreadSafe
 );
 
@@ -1109,7 +1116,8 @@ class FPathTracingRG : public FGlobalShader
 	class FCloudShader : SHADER_PERMUTATION_BOOL("PATH_TRACER_USE_CLOUD_SHADER");
 	class FSubstrateComplexSpecialMaterial : SHADER_PERMUTATION_BOOL("PATH_TRACER_USE_SUBSTRATE_SPECIAL_COMPLEX_MATERIAL");
 	class FUseSER : SHADER_PERMUTATION_BOOL("PATH_TRACER_USE_SER");
-	using FPermutationDomain = TShaderPermutationDomain<FCompactionType, FAdaptiveSampling, FCloudShader, FSubstrateComplexSpecialMaterial, FUseSER>;
+	class FTraceOpaqueFirst : SHADER_PERMUTATION_BOOL("PATH_TRACER_TRACE_OPAQUE_FIRST");
+	using FPermutationDomain = TShaderPermutationDomain<FCompactionType, FAdaptiveSampling, FCloudShader, FSubstrateComplexSpecialMaterial, FUseSER, FTraceOpaqueFirst>;
 
 	static bool ShouldCompilePermutation(const FGlobalShaderPermutationParameters& Parameters)
 	{
@@ -1164,6 +1172,7 @@ class FPathTracingRG : public FGlobalShader
 		OutEnvironment.SetDefine(TEXT("USE_RECT_LIGHT_TEXTURES"), 1);
 		OutEnvironment.SetDefine(TEXT("DEBUG_ENABLE"), CVarPathTracingDebug.GetValueOnAnyThread() > 0 ? 1u : 0u);
 		OutEnvironment.CompilerFlags.Add(CFLAG_WarningsAsErrors);
+		OutEnvironment.CompilerFlags.Add(CFLAG_HLSL2021);
 	}
 
 	static ERayTracingPayloadType GetRayTracingPayloadType(const int32 PermutationId)
@@ -2938,6 +2947,7 @@ static FPathTracingRG::FPermutationDomain GetPathTracingRGPermutation(const FSce
 	const bool bUseCloudShader = bUseExperimental && ShouldCompileRayTracingCallableShadersForProject(Scene.GetShaderPlatform()) && !CVarPathTracingCloudMapEnabled.GetValueOnRenderThread();
 	const bool bHasComplexSpecialRenderPath = Substrate::IsSubstrateEnabled() && Scene.SubstrateSceneData.bUsesComplexSpecialRenderPath;
 	const bool bUseSER = GRHIGlobals.SupportsShaderExecutionReordering && CVarPathTracingShaderExecutionReordering.GetValueOnRenderThread();
+	const bool bUseTraceOpaqueFirst = CVarPathTracingTraceOpaqueFirst.GetValueOnRenderThread();
 
 	FPathTracingRG::FPermutationDomain Out;
 	Out.Set<FPathTracingRG::FCompactionType>(bUseCompaction);
@@ -2945,6 +2955,7 @@ static FPathTracingRG::FPermutationDomain GetPathTracingRGPermutation(const FSce
 	Out.Set<FPathTracingRG::FCloudShader>(bUseCloudShader);
 	Out.Set<FPathTracingRG::FSubstrateComplexSpecialMaterial>(bHasComplexSpecialRenderPath);
 	Out.Set<FPathTracingRG::FUseSER>(bUseSER);
+	Out.Set<FPathTracingRG::FTraceOpaqueFirst>(bUseTraceOpaqueFirst);
 	return Out;
 }
 
@@ -2960,12 +2971,15 @@ void FDeferredShadingSceneRenderer::PreparePathTracing(const FSceneViewFamily& V
 		}
 		else
 		{
+			const int CompactionDepth = CVarPathTracingCompactionDepth.GetValueOnRenderThread();
 			// Declare all RayGen shaders that require material closest hit shaders to be bound
+			if (CompactionDepth >= 0)
 			{
 				FPathTracingRG::FPermutationDomain PermutationVector = GetPathTracingRGPermutation(Scene, false);
 				auto RayGenShader = GetGlobalShaderMap(ViewFamily.GetShaderPlatform())->GetShader<FPathTracingRG>(PermutationVector);
 				OutRayGenShaders.Add(RayGenShader.GetRayTracingShader());
 			}
+			if (CompactionDepth != 0)
 			{
 				FPathTracingRG::FPermutationDomain PermutationVector = GetPathTracingRGPermutation(Scene, true);
 				auto RayGenShader = GetGlobalShaderMap(ViewFamily.GetShaderPlatform())->GetShader<FPathTracingRG>(PermutationVector);
@@ -3000,12 +3014,6 @@ void PreparePathTracingRTPSO()
 
 				FPathTracingRG::FPermutationDomain PermutationVector(PermutationId);
 
-				if (PermutationVector.Get<FPathTracingRG::FCompactionType>())
-				{
-					// skip over these so we don't double count them
-					continue;
-				}
-
 				if (PermutationVector.Get<FPathTracingRG::FUseSER>() && !GRHIGlobals.SupportsShaderExecutionReordering)
 				{
 					// we can safely skip these if the current machine does not support SER
@@ -3014,12 +3022,8 @@ void PreparePathTracingRTPSO()
 					continue;
 				}
 				
-				FPathTracingRG::FPermutationDomain PermutationVectorCompaction(PermutationId);
-				PermutationVectorCompaction.Set<FPathTracingRG::FCompactionType>(true);
-
 				FRHIRayTracingShader* RayGenShaderTable[] = {
 					ShaderMap->GetShader<FPathTracingRG>(PermutationVector).GetRayTracingShader(),
-					ShaderMap->GetShader<FPathTracingRG>(PermutationVectorCompaction).GetRayTracingShader(),
 					ShaderMap->GetShader<FPathTracingInitExtinctionCoefficientRG>().GetRayTracingShader(),
 				};
 				FRHIRayTracingShader* MissShaderTable[] = {
diff --git a/Engine/Source/Runtime/Renderer/Private/RayTracing/RayTracingInstanceMask.cpp b/Engine/Source/Runtime/Renderer/Private/RayTracing/RayTracingInstanceMask.cpp
--- a/Engine/Source/Runtime/Renderer/Private/RayTracing/RayTracingInstanceMask.cpp
+++ b/Engine/Source/Runtime/Renderer/Private/RayTracing/RayTracingInstanceMask.cpp
@@ -61,7 +61,7 @@ uint8 ComputeRayTracingInstanceMask(ERayTracingInstanceMaskType MaskType, ERayTr
 			Mask = PATHTRACER_MASK_CAMERA | PATHTRACER_MASK_INDIRECT;
 			break;
 		case ERayTracingInstanceMaskType::Translucent:
-			Mask = PATHTRACER_MASK_CAMERA | PATHTRACER_MASK_INDIRECT;
+			Mask = PATHTRACER_MASK_CAMERA_TRANSLUCENT | PATHTRACER_MASK_INDIRECT_TRANSLUCENT;
 			break;
 		case ERayTracingInstanceMaskType::OpaqueShadow:
 			Mask = PATHTRACER_MASK_SHADOW;
@@ -79,10 +79,10 @@ uint8 ComputeRayTracingInstanceMask(ERayTracingInstanceMaskType MaskType, ERayTr
 			Mask = PATHTRACER_MASK_HAIR_CAMERA | PATHTRACER_MASK_HAIR_INDIRECT;
 			break;
 		case ERayTracingInstanceMaskType::VisibleInPrimaryRay:
-			Mask = PATHTRACER_MASK_CAMERA | PATHTRACER_MASK_HAIR_CAMERA;
+			Mask = PATHTRACER_MASK_CAMERA | PATHTRACER_MASK_HAIR_CAMERA | PATHTRACER_MASK_CAMERA_TRANSLUCENT;
 			break;
 		case ERayTracingInstanceMaskType::VisibleInIndirectRay:
-			Mask = PATHTRACER_MASK_INDIRECT | PATHTRACER_MASK_HAIR_INDIRECT;
+			Mask = PATHTRACER_MASK_INDIRECT | PATHTRACER_MASK_HAIR_INDIRECT | PATHTRACER_MASK_INDIRECT_TRANSLUCENT;
 			break;
 		default:
 			checkNoEntry();
@@ -93,11 +93,11 @@ uint8 ComputeRayTracingInstanceMask(ERayTracingInstanceMaskType MaskType, ERayTr
 	return Mask;
 }
 
-uint8 BlendModeToRayTracingInstanceMask(const EBlendMode BlendMode, bool bCastShadow, ERayTracingViewMaskMode MaskMode)
+uint8 BlendModeToRayTracingInstanceMask(const EBlendMode BlendMode, bool bbIsDitherMasked, bool bCastShadow, ERayTracingViewMaskMode MaskMode)
 {
 	uint8 InstanceMask = 0;
 
-	if (IsOpaqueOrMaskedBlendMode(BlendMode))
+	if (IsOpaqueOrMaskedBlendMode(BlendMode) && !bbIsDitherMasked)
 	{
 		InstanceMask |= ComputeRayTracingInstanceMask(ERayTracingInstanceMaskType::Opaque, MaskMode);
 		InstanceMask |= bCastShadow ? ComputeRayTracingInstanceMask(ERayTracingInstanceMaskType::OpaqueShadow, MaskMode) : 0;
@@ -154,7 +154,7 @@ FRayTracingMaskAndFlags BuildRayTracingInstanceMaskAndFlags(TArrayView<const FMe
 			const EBlendMode BlendMode = Material.GetBlendMode();
 			const bool bSegmentCastsShadow = MeshBatch.CastRayTracedShadow && Material.CastsRayTracedShadows() && BlendMode != BLEND_Additive;
 
-			Result.Mask |= BlendModeToRayTracingInstanceMask(BlendMode, bSegmentCastsShadow, MaskMode);
+			Result.Mask |= BlendModeToRayTracingInstanceMask(BlendMode, Material.IsDitherMasked(), bSegmentCastsShadow, MaskMode);
 			bAllSegmentsOpaque &= BlendMode == BLEND_Opaque;
 			bAllSegmentsTranslucent &= BlendMode == BLEND_Translucent;
 			bAnySegmentsCastShadow |= bSegmentCastsShadow;
@@ -253,7 +253,7 @@ void SetupRayTracingMeshCommandMaskAndStatus(FRayTracingMeshCommand& MeshCommand
 	MeshCommand.bIsTranslucent = MaterialResource.GetBlendMode() == EBlendMode::BLEND_Translucent;
 	MeshCommand.bReverseCulling = MeshBatch.ReverseCulling;
 
-	MeshCommand.InstanceMask = BlendModeToRayTracingInstanceMask(MaterialResource.GetBlendMode(), MeshCommand.bCastRayTracedShadows, MaskMode);
+	MeshCommand.InstanceMask = BlendModeToRayTracingInstanceMask(MaterialResource.GetBlendMode(), MaterialResource.IsDitherMasked(), MeshCommand.bCastRayTracedShadows, MaskMode);
 
 	if (!PrimitiveSceneProxy)
 	{
diff --git a/Engine/Source/Runtime/Renderer/Private/RayTracing/RayTracingInstanceMask.h b/Engine/Source/Runtime/Renderer/Private/RayTracing/RayTracingInstanceMask.h
--- a/Engine/Source/Runtime/Renderer/Private/RayTracing/RayTracingInstanceMask.h
+++ b/Engine/Source/Runtime/Renderer/Private/RayTracing/RayTracingInstanceMask.h
@@ -36,7 +36,7 @@ enum class ERayTracingViewMaskMode : uint8
 uint8 ComputeRayTracingInstanceMask(ERayTracingInstanceMaskType MaskType, ERayTracingViewMaskMode MaskMode);
 
 /** Compute the mask based on blend mode for different ray tracing mode*/
-RENDERER_API uint8 BlendModeToRayTracingInstanceMask(const EBlendMode BlendMode, bool bCastShadow, ERayTracingViewMaskMode MaskMode);
+RENDERER_API uint8 BlendModeToRayTracingInstanceMask(const EBlendMode BlendMode, bool bIsDitherMasked, bool bCastShadow, ERayTracingViewMaskMode MaskMode);
 
 
 /** Util struct and function to derive mask related info from scene proxy*/
