Fix incorrect CulledLightDataStart when using r.Forward.LightLinkedListCulling=0.

diff --git a/Engine/Shaders/Private/LightGridInjection.usf b/Engine/Shaders/Private/LightGridInjection.usf
--- a/Engine/Shaders/Private/LightGridInjection.usf
+++ b/Engine/Shaders/Private/LightGridInjection.usf
@@ -399,7 +399,7 @@ void LightGridInjectionCS(
 		// Mega Lights are order dependent
 		WriteReversedLinkedList(LinkOffset, NumVisibleLights, CulledLightDataStart);
 #else
-		uint CulledLightDataStart = NumGridCells * MaxCulledLightsPerCell;
+		uint CulledLightDataStart = GridIndex * MaxCulledLightsPerCell;
 #endif
 
 		RWNumCulledLightsGrid[GridIndex * NUM_CULLED_LIGHTS_GRID_STRIDE + 0] = PackCulledLightsGridHeader0(NumVisibleLights, NumVisibleMegaLights);
@@ -445,7 +445,7 @@ void LightGridInjectionCS(
 		// Reflection captures are order dependent
 		WriteReversedLinkedList(CaptureLinkOffset, NumVisibleCaptures, CulledReflectionCaptureDataStart);
 #else
-		uint CulledReflectionCaptureDataStart = NumGridCells * MaxCulledLightsPerCell;
+		uint CulledReflectionCaptureDataStart = (NumGridCells + GridIndex) * MaxCulledLightsPerCell;
 #endif
 
 		RWNumCulledLightsGrid[(NumGridCells + GridIndex) * NUM_CULLED_LIGHTS_GRID_STRIDE + 0] = NumVisibleCaptures;
