Fix groom cache interplation when triangle strip geometry is enabled.

diff --git a/Engine/Plugins/Runtime/HairStrands/Source/HairStrandsCore/Private/GroomManager.cpp b/Engine/Plugins/Runtime/HairStrands/Source/HairStrandsCore/Private/GroomManager.cpp
--- a/Engine/Plugins/Runtime/HairStrands/Source/HairStrandsCore/Private/GroomManager.cpp
+++ b/Engine/Plugins/Runtime/HairStrands/Source/HairStrandsCore/Private/GroomManager.cpp
@@ -1059,11 +1059,14 @@ static void RunHairStrandsInterpolation_Guide(
 					ShaderMap,
 					InstanceData.Instance->RegisteredIndex,
 					InstanceData.Instance->Guides.RestResource->GetPointCount(),
+					InstanceData.Instance->Guides.RestResource->GetCurveCount(),
 					InterpolationFactor,
 					1.f /*MaxHairRadius*/,
 					CacheResources0,
 					CacheResources1,
 					RegisterAsSRV(GraphBuilder, InstanceData.Instance->Guides.RestResource->PositionBuffer),
+					RegisterAsSRV(GraphBuilder, InstanceData.Instance->Guides.RestResource->CurveBuffer),
+					RegisterAsSRV(GraphBuilder, InstanceData.Instance->Strands.RestResource->PointToCurveBuffer),
 					RegisterAsSRV(GraphBuilder, InstanceData.Instance->Guides.DeformedResource->GetPositionOffsetBuffer(FHairStrandsDeformedResource::EFrameType::Current)),
 					RegisterAsUAV(GraphBuilder, InstanceData.Instance->Guides.DeformedResource->GetBuffer(FHairStrandsDeformedResource::Current)));
 			}
@@ -1541,11 +1544,14 @@ static void RunHairStrandsInterpolation_Strands(
 					ShaderMap,
 					InstanceData.Instance->RegisteredIndex,
 					InstanceData.ActivePointCount,
+					InstanceData.ActiveCurveCount,
 					InterpolationFactor,
 					InstanceData.Instance->Strands.Modifier.HairWidth * 0.5f /*InMaxHairRadius*/,
 					CacheResources0,
 					CacheResources1,
 					RegisterAsSRV(GraphBuilder, InstanceData.Instance->Strands.RestResource->PositionBuffer),
+					RegisterAsSRV(GraphBuilder, InstanceData.Instance->Strands.RestResource->CurveBuffer),
+					RegisterAsSRV(GraphBuilder, InstanceData.Instance->Strands.RestResource->PointToCurveBuffer),
 					RegisterAsSRV(GraphBuilder, InstanceData.Instance->Strands.DeformedResource->GetPositionOffsetBuffer(FHairStrandsDeformedResource::EFrameType::Current)),
 					InstanceData.RDGResources.PositionUAV);
 			}
diff --git a/Engine/Plugins/Runtime/HairStrands/Source/HairStrandsCore/Private/HairStrandsInterpolation.cpp b/Engine/Plugins/Runtime/HairStrands/Source/HairStrandsCore/Private/HairStrandsInterpolation.cpp
--- a/Engine/Plugins/Runtime/HairStrands/Source/HairStrandsCore/Private/HairStrandsInterpolation.cpp
+++ b/Engine/Plugins/Runtime/HairStrands/Source/HairStrandsCore/Private/HairStrandsInterpolation.cpp
@@ -234,15 +234,18 @@ class FGroomCacheUpdatePassCS : public FGlobalShader
 
 	BEGIN_SHADER_PARAMETER_STRUCT(FParameters, )
 		SHADER_PARAMETER(uint32, InstanceRegisteredIndex)
-		SHADER_PARAMETER(uint32, ElementCount)
+		SHADER_PARAMETER(uint32, CachePointCount)
 		SHADER_PARAMETER(uint32, bHasRadiusData)
+		SHADER_PARAMETER(uint32, bHasAddedControlPoint)
 		SHADER_PARAMETER(float, InterpolationFactor)
 		SHADER_PARAMETER(float, MaxHairRadius)
 		SHADER_PARAMETER_RDG_BUFFER_SRV(Buffer, InPosition0Buffer)
 		SHADER_PARAMETER_RDG_BUFFER_SRV(Buffer, InPosition1Buffer)
 		SHADER_PARAMETER_RDG_BUFFER_SRV(Buffer, InRadius0Buffer)
 		SHADER_PARAMETER_RDG_BUFFER_SRV(Buffer, InRadius1Buffer)
-		SHADER_PARAMETER_RDG_BUFFER_SRV(Buffer, InRestPoseBuffer)
+		SHADER_PARAMETER_RDG_BUFFER_SRV(ByteAddressBuffer, InRestPositionBuffer)
+		SHADER_PARAMETER_RDG_BUFFER_SRV(ByteAddressBuffer, InRestCurveBuffer)
+		SHADER_PARAMETER_RDG_BUFFER_SRV(ByteAddressBuffer, InRestPointToCurveBuffer)
 		SHADER_PARAMETER_RDG_BUFFER_SRV(StructuredBuffer, InDeformedOffsetBuffer)
 		SHADER_PARAMETER_RDG_BUFFER_UAV(RWBuffer, OutDeformedBuffer)
 	END_SHADER_PARAMETER_STRUCT()
@@ -265,11 +268,14 @@ void AddGroomCacheUpdatePass(
 	FGlobalShaderMap* ShaderMap,
 	uint32 InstanceRegisteredIndex,
 	uint32 PointCount,
+	uint32 CurveCount,
 	float InterpolationFactor,
 	float InMaxHairRadius,
 	FGroomCacheResources CacheResources0,
 	FGroomCacheResources CacheResources1,
-	FRDGBufferSRVRef InBuffer,
+	FRDGBufferSRVRef InRestPositionBuffer,
+	FRDGBufferSRVRef InRestCurveBuffer,
+	FRDGBufferSRVRef InRestPointToCurveBuffer,
 	FRDGBufferSRVRef InDeformedOffsetBuffer,
 	FRDGBufferUAVRef OutBuffer)
 {
@@ -279,17 +285,20 @@ void AddGroomCacheUpdatePass(
 
 	FGroomCacheUpdatePassCS::FParameters* Parameters = GraphBuilder.AllocParameters<FGroomCacheUpdatePassCS::FParameters>();
 	Parameters->InstanceRegisteredIndex = InstanceRegisteredIndex;
-	Parameters->ElementCount = PointCount;
+	Parameters->CachePointCount = PointCount;
 	Parameters->InPosition0Buffer = CacheResources0.PositionBuffer;
 	Parameters->InPosition1Buffer = CacheResources1.PositionBuffer;
 	Parameters->InRadius0Buffer = CacheResources0.RadiusBuffer;
 	Parameters->InRadius1Buffer = CacheResources1.RadiusBuffer;
-	Parameters->InRestPoseBuffer = InBuffer;
+	Parameters->InRestPositionBuffer = InRestPositionBuffer;
+	Parameters->InRestCurveBuffer = InRestCurveBuffer;
+	Parameters->InRestPointToCurveBuffer = InRestPointToCurveBuffer;
 	Parameters->InDeformedOffsetBuffer = InDeformedOffsetBuffer;
 	Parameters->OutDeformedBuffer = OutBuffer;
 	Parameters->InterpolationFactor = InterpolationFactor;
 	Parameters->MaxHairRadius = InMaxHairRadius;
 	Parameters->bHasRadiusData = CacheResources0.bHasRadiusData ? 1u : 0u;
+	Parameters->bHasAddedControlPoint = GetHairStrandsUsesTriangleStrips() ? 1u : 0u;
 
 	const FIntVector DispatchCount = FIntVector(FMath::DivideAndRoundUp(PointCount, FGroomCacheUpdatePassCS::GetGroupSize()), 1, 1);
 	TShaderMapRef<FGroomCacheUpdatePassCS> ComputeShader(ShaderMap);
diff --git a/Engine/Plugins/Runtime/HairStrands/Source/HairStrandsCore/Public/HairStrandsInterpolation.h b/Engine/Plugins/Runtime/HairStrands/Source/HairStrandsCore/Public/HairStrandsInterpolation.h
--- a/Engine/Plugins/Runtime/HairStrands/Source/HairStrandsCore/Public/HairStrandsInterpolation.h
+++ b/Engine/Plugins/Runtime/HairStrands/Source/HairStrandsCore/Public/HairStrandsInterpolation.h
@@ -82,11 +82,14 @@ void AddGroomCacheUpdatePass(
 	FGlobalShaderMap* ShaderMap,
 	uint32 InstanceRegisteredIndex,
 	uint32 PointCount,
+	uint32 CurveCount,
 	float InterpolationFactor,
-	float MaxHairRadius,
+	float InMaxHairRadius,
 	FGroomCacheResources CacheResources0,
 	FGroomCacheResources CacheResources1,
-	FRDGBufferSRVRef InBuffer,
+	FRDGBufferSRVRef InRestPositionBuffer,
+	FRDGBufferSRVRef InRestCurveBuffer,
+	FRDGBufferSRVRef InRestPointToCurveBuffer,
 	FRDGBufferSRVRef InDeformedOffsetBuffer,
 	FRDGBufferUAVRef OutBuffer);
 
diff --git a/Engine/Plugins/Runtime/HairStrands/Source/HairStrandsEditor/Private/GroomCacheImporter.cpp b/Engine/Plugins/Runtime/HairStrands/Source/HairStrandsEditor/Private/GroomCacheImporter.cpp
--- a/Engine/Plugins/Runtime/HairStrands/Source/HairStrandsEditor/Private/GroomCacheImporter.cpp
+++ b/Engine/Plugins/Runtime/HairStrands/Source/HairStrandsEditor/Private/GroomCacheImporter.cpp
@@ -141,7 +141,9 @@ TArray<UGroomCache*> FGroomCacheImporter::ImportGroomCache(const FString& Source
 			if (Translator->Translate(FrameIndex * AnimInfo.SecondsPerFrame, FrameHairDescription, HairImportContext.ImportOptions->ConversionSettings))
 			{
 				FHairDescriptionGroups HairDescriptionGroups;
-				if (!FGroomBuilder::BuildHairDescriptionGroups(FrameHairDescription, HairDescriptionGroups))
+				// Do not add extra control points at the end of curve when hair strip geometry is enabled. This is because groom cache data are 
+				// serialize within the uasset (i.e. do not used intermediate cached/build data), and need the asset to be compatible with or without hair strip geometry.
+				if (!FGroomBuilder::BuildHairDescriptionGroups(FrameHairDescription, HairDescriptionGroups, false /*bAllowAddEndControlPoint*/))
 				{
 					bSuccess = false;
 					break;
diff --git a/Engine/Shaders/Private/HairStrands/HairStrandsInterpolation.usf b/Engine/Shaders/Private/HairStrands/HairStrandsInterpolation.usf
--- a/Engine/Shaders/Private/HairStrands/HairStrandsInterpolation.usf
+++ b/Engine/Shaders/Private/HairStrands/HairStrandsInterpolation.usf
@@ -914,8 +914,11 @@ void MainCS(uint DispatchThreadId : SV_DispatchThreadID)
 #if SHADER_GROOMCACHE_UPDATE
 
 uint InstanceRegisteredIndex;
-uint ElementCount;
+uint CachePointCount;
+uint PointCount;
+uint CurveCount;
 uint bHasRadiusData;
+uint bHasAddedControlPoint;
 float InterpolationFactor;
 float MaxHairRadius;
 
@@ -924,42 +927,62 @@ Buffer<float> InPosition1Buffer;
 Buffer<float> InRadius0Buffer;
 Buffer<float> InRadius1Buffer;
 
-ByteAddressBuffer  InRestPoseBuffer;
+ByteAddressBuffer InRestCurveBuffer;
+ByteAddressBuffer InRestPositionBuffer;
+ByteAddressBuffer InRestPointToCurveBuffer;
 StructuredBuffer<float4> InDeformedOffsetBuffer;
 RWByteAddressBuffer OutDeformedBuffer;
 
 [numthreads(GROUP_SIZE, 1, 1)]
 void MainCS(uint DispatchThreadId : SV_DispatchThreadID)
 {
-	const uint Index = DispatchThreadId;
-	if (Index >= ElementCount)
+	const uint RestPointIndex = DispatchThreadId;
+	const uint CurveIndex = ReadHairPointToCurveIndex(InRestPointToCurveBuffer, RestPointIndex);
+	const FHairCurve Curve = ReadHairCurve(InRestCurveBuffer, CurveIndex);
+
+	const bool bIsEndPoint = RestPointIndex == (Curve.PointIndex + Curve.PointCount - 1);
+	if (bHasAddedControlPoint && bIsEndPoint)
 	{
-		return;
+		return; // By pass update of the last point
 	}
 
-	const uint Index3 = Index * 3;
-	const float3 InPosition0 = float3(
-		InPosition0Buffer[Index3],
-		InPosition0Buffer[Index3+1],
-		InPosition0Buffer[Index3+2]);
+	// Make the assumption that there is no clamped curve data.
+	// One extra point per curve.
+	const uint NumAddedPoints = bHasAddedControlPoint ? CurveIndex : 0;
+
+	const uint CachePointIndex = RestPointIndex - NumAddedPoints;
+	if (CachePointIndex >= CachePointCount)
+	{
+		return;
+	}
 	
-	const float3 InPosition1 = float3(
-		InPosition1Buffer[Index3],
-		InPosition1Buffer[Index3+1],
-		InPosition1Buffer[Index3+2]);
+	FHairControlPoint CP = ReadHairControlPoint(InRestPositionBuffer, RestPointIndex, float3(0, 0, 0)/*Offset*/, 1/*Radius*/, 1/*RootScale*/, 1/*TipScale*/);
+
+	// Position
+	{
+		const uint Index3 = CachePointIndex * 3;
+		const float3 InPosition0 = float3(
+			InPosition0Buffer[Index3],
+			InPosition0Buffer[Index3+1],
+			InPosition0Buffer[Index3+2]);
 	
-	const float3 WorldVertexPosition = lerp(InPosition0, InPosition1, InterpolationFactor);
-	FHairControlPoint CP = ReadHairControlPoint(InRestPoseBuffer, Index, float3(0, 0, 0)/*Offset*/, 1/*Radius*/, 1/*RootScale*/, 1/*TipScale*/);
-	CP.Position = WorldVertexPosition;
+		const float3 InPosition1 = float3(
+			InPosition1Buffer[Index3],
+			InPosition1Buffer[Index3+1],
+			InPosition1Buffer[Index3+2]);
+
+		CP.Position = lerp(InPosition0, InPosition1, InterpolationFactor);
+	}
 
+	// Radius
 	if (bHasRadiusData == 1)
 	{
-		CP.WorldRadius = lerp(InRadius0Buffer[Index], InRadius1Buffer[Index], InterpolationFactor);
+		CP.WorldRadius = lerp(InRadius0Buffer[CachePointIndex], InRadius1Buffer[CachePointIndex], InterpolationFactor);
 	}
 
 	const float3 DeformedPositionOffset = ReadRenPositionOffset(InDeformedOffsetBuffer, InstanceRegisteredIndex);
 
-	WriteHairControlPointPosition(OutDeformedBuffer, Index, CP, DeformedPositionOffset, MaxHairRadius);
+	WriteHairControlPointPosition(OutDeformedBuffer, RestPointIndex, CP, DeformedPositionOffset, MaxHairRadius);
 }
 
 #endif // SHADER_GROOMCACHE_UPDATE
