Substrate - reworked shadow backface culling logic to account for basis count instead of Closure Count.
Simplified stored data on the header to simply be packed data.

#rnx
#tests none

diff --git a/Engine/Shaders/Private/RayTracing/RayTracingOcclusionRGS.usf b/Engine/Shaders/Private/RayTracing/RayTracingOcclusionRGS.usf
--- a/Engine/Shaders/Private/RayTracing/RayTracingOcclusionRGS.usf
+++ b/Engine/Shaders/Private/RayTracing/RayTracingOcclusionRGS.usf
@@ -735,10 +735,10 @@ RAY_TRACING_ENTRY_RAYGEN_OR_INLINE(Occlusion)
 			ShadingParameters.bBackfaceCulling = !(BSDFType == SUBSTRATE_BSDF_TYPE_EYE || (bHasSSS && SSSType == SSS_TYPE_WRAP)); // isthin?
 
 			// The SSDenoiser reads the top layer normal in order to know how to process the data.
-			// Shadow is traced using the top layer normal when NoL>0 only. So if normals for slabs below do not match, light leaking will occur.
+			// RTShadow traces using the top layer normal when NoL>0 only. So if normals for slabs below do not match, light leaking will occur.
 			// For instance: a flat translucent top layer above a normal mapped bottom layer.
-			// That is why, for complex materials with multiple BSDFs, we also must not do backface culling.
-			if (SubstratePixelHeader.ClosureCount > 1)
+			// That is why, for materials with multiple local bases, we must not do backface culling.
+			if (SubstratePixelHeader.GetLocalBasesCount() > 1)
 			{
 				// Technically, this could only be done when vertical layering is encountered only. 
 				// But we do not have that data available (or it would be costly to infer from closures)
diff --git a/Engine/Shaders/Private/Substrate/Substrate.ush b/Engine/Shaders/Private/Substrate/Substrate.ush
--- a/Engine/Shaders/Private/Substrate/Substrate.ush
+++ b/Engine/Shaders/Private/Substrate/Substrate.ush
@@ -2137,7 +2137,7 @@ struct FSubstratePixelHeader
 #if SUBSTRATE_DEFERRED_SHADING
 	FSubstrateMaterialContainer SubstrateBuffer;
 	uint SharedLocalBasesIndexOffset;
-	uint SharedLocalBasesTypes_PackedHeader; // Stores SharedLocalBasesTypes for non-simple material, and PackedHeader data for simple material
+	uint PackedHeader;	// Used for special simple/eye/etc. packing, or for complex to recover LocalBasesTypes, or use GetLocalBasesCount in rare cases (VSM, RTShadow or visualisation).
 	SUBSTRATE_TOP_LAYER_TYPE PackedTopLayerData;
 #endif
 
@@ -2190,6 +2190,10 @@ struct FSubstratePixelHeader
 	void SetMaterialMode(uint MaterialMode)						{ this.State = (this.State & ~HEADER_MASK_MATERIALMODE) | (MaterialMode & HEADER_MASK_MATERIALMODE); }
 
 	void SetIsComplexSpecialMaterial(bool bIn)					{ this.State |= (bIn ? HEADER_COMPLEXSPECIALPATH_MASK : 0u); }
+
+#if SUBSTRATE_DEFERRED_SHADING
+	uint GetLocalBasesCount()									{ return GetMaterialMode() == HEADER_MATERIALMODE_SLAB_COMPLEX ? HEADER_GETSHAREDLOCALBASESCOUNT(PackedHeader) : 1; }
+#endif
 };
 
 #if SUBSTRATE_DEFERRED_SHADING && SUBSTRATE_MATERIALCONTAINER_IS_VIEWRESOURCE
@@ -2211,7 +2215,7 @@ FSubstratePixelHeader InitialiseSubstratePixelHeader()
 	Out.SubstrateBuffer = DummySubstrateBuffer;
 #endif
 	Out.SharedLocalBasesIndexOffset = 0;
-	Out.SharedLocalBasesTypes_PackedHeader = 0;
+	Out.PackedHeader = 0;
 	Out.PackedTopLayerData = 0;
 #endif
 	Out.State = 0;
@@ -2320,7 +2324,7 @@ float3x3 SubstrateGetBSDFSharedBasis_DeferredShading(in FSubstratePixelHeader Su
 
 	uint PackedBasis = SubstrateLoadUint1AtIndexOffset(SubstratePixelHeader.SubstrateBuffer, SubstrateAddressing, SubstratePixelHeader.SharedLocalBasesIndexOffset + BSDF_GETSHAREDLOCALBASISID(BSDF));
 
-	uint BasisType   = SubstrateGetSharedLocalBasisType(SubstratePixelHeader.SharedLocalBasesTypes_PackedHeader, BSDF_GETSHAREDLOCALBASISID(BSDF));
+	uint BasisType   = SubstrateGetSharedLocalBasisType(HEADER_GETSHAREDLOCALBASESTYPE(SubstratePixelHeader.PackedHeader), BSDF_GETSHAREDLOCALBASISID(BSDF));
 	float3x3 OutTangentBasis;
 	if (BasisType == SUBSTRATE_BASIS_TYPE_NORMAL)
 	{
@@ -3874,8 +3878,8 @@ FSubstratePixelHeader UnpackSubstrateHeaderIn(uint PackedHeader, inout FSubstrat
 		Out.ClosureCount = Out.IsSubstrateMaterial() ? 1 : 0; 
 
 	  #if SUBSTRATE_DEFERRED_SHADING
-		Out.SharedLocalBasesTypes_PackedHeader = PackedHeader; // ?
-		Out.SharedLocalBasesIndexOffset = 0; // Unused
+		Out.PackedHeader = PackedHeader;
+		Out.SharedLocalBasesIndexOffset = 0;	// Unused
 
 		Out.PackedTopLayerData = SubstrateLoadTopLayerData(InSubstrateTopLayerTexture, SubstrateAddressing);
 		SubstrateAddressing.ReadBytes += 4;
@@ -3884,13 +3888,13 @@ FSubstratePixelHeader UnpackSubstrateHeaderIn(uint PackedHeader, inout FSubstrat
 	#if SUBSTRATE_COMPLEXPATH
 	else
 	{
-		uint SharedLocalBasesCount = HEADER_GETSHAREDLOCALBASESCOUNT(PackedHeader); // ?
+		uint SharedLocalBasesCount = HEADER_GETSHAREDLOCALBASESCOUNT(PackedHeader);
 		uint SharedLocalBasesTypes = HEADER_GETSHAREDLOCALBASESTYPE(PackedHeader);
 
 		Out.SetIsComplexSpecialMaterial((PackedHeader & HEADER_COMPLEXSPECIALPATH_MASK) > 0);
 
 	#if SUBSTRATE_DEFERRED_SHADING
-		Out.SharedLocalBasesTypes_PackedHeader = SharedLocalBasesTypes;
+		Out.PackedHeader = PackedHeader;
 		Out.PackedTopLayerData = 0;
 
 		// Only keep the offset to the shared local bases memory. We are going to load them on demand.
@@ -3946,7 +3950,7 @@ FSubstrateBSDF UnpackFastPathSubstrateBSDFIn(FSubstrateMaterialContainer Substra
 	FSubstrateBSDF OutBSDF = (FSubstrateBSDF)0;
 
 #if SUBSTRATE_DEFERRED_SHADING
-	uint Data0 = SubstratePixelHeader.SharedLocalBasesTypes_PackedHeader;
+	uint Data0 = SubstratePixelHeader.PackedHeader;
 	uint Data1 = SubstrateLoadUint1(SubstrateBuffer, SubstrateAddressing);
 	OutBSDF.State = 0;
 	OutBSDF.LuminanceWeightV = 1.0f;
@@ -3990,7 +3994,7 @@ FSubstrateBSDF UnpackFastWaterPathSubstrateBSDFIn(FSubstrateMaterialContainer Su
 	FSubstrateBSDF OutBSDF = (FSubstrateBSDF)0;
 
 #if SUBSTRATE_DEFERRED_SHADING
-	uint Data0 = SubstratePixelHeader.SharedLocalBasesTypes_PackedHeader;
+	uint Data0 = SubstratePixelHeader.PackedHeader;
 	uint Data1 = SubstrateLoadUint1(SubstrateBuffer, SubstrateAddressing);
 	OutBSDF.State = 0;
 	OutBSDF.LuminanceWeightV = 1.0f;
@@ -4057,7 +4061,7 @@ FSubstrateBSDF UnpackSubstrateBSDFIn(FSubstrateMaterialContainer SubstrateBuffer
 	{
 		uint Data = 0;
 	#if SUBSTRATE_DEFERRED_SHADING
-		Data = SubstrateHeader.SharedLocalBasesTypes_PackedHeader;
+		Data = SubstrateHeader.PackedHeader;
 	#endif
 
 		// Restore the BSDF common state
@@ -4086,7 +4090,7 @@ FSubstrateBSDF UnpackSubstrateBSDFIn(FSubstrateMaterialContainer SubstrateBuffer
 			if (OptimisedLegacyMode == SINGLE_OPTLEGACYMODE_CLEARCOAT)
 			{
 			#if SUBSTRATE_DEFERRED_SHADING
-				uint Data0 = SubstrateHeader.SharedLocalBasesTypes_PackedHeader;
+				uint Data0 = SubstrateHeader.PackedHeader;
 				uint Data1 = SubstrateLoadUint1(SubstrateBuffer, SubstrateAddressing);
 				uint Data2 = SubstrateLoadUint1(SubstrateBuffer, SubstrateAddressing);
 
@@ -4117,7 +4121,7 @@ FSubstrateBSDF UnpackSubstrateBSDFIn(FSubstrateMaterialContainer SubstrateBuffer
 			else if (OptimisedLegacyMode == SINGLE_OPTLEGACYMODE_CLOTH)
 			{
 			#if SUBSTRATE_DEFERRED_SHADING
-				uint Data0 = SubstrateHeader.SharedLocalBasesTypes_PackedHeader;
+				uint Data0 = SubstrateHeader.PackedHeader;
 				uint Data1 = SubstrateLoadUint1(SubstrateBuffer, SubstrateAddressing);
 				uint Data2 = SubstrateLoadUint1(SubstrateBuffer, SubstrateAddressing);
 
@@ -4150,7 +4154,7 @@ FSubstrateBSDF UnpackSubstrateBSDFIn(FSubstrateMaterialContainer SubstrateBuffer
 			else if (OptimisedLegacyMode == SINGLE_OPTLEGACYMODE_SSSWRAP || OptimisedLegacyMode == SINGLE_OPTLEGACYMODE_SSSWRAP_THIN)
 			{
 			#if SUBSTRATE_DEFERRED_SHADING
-				uint Data0 = SubstrateHeader.SharedLocalBasesTypes_PackedHeader;
+				uint Data0 = SubstrateHeader.PackedHeader;
 				uint Data1 = SubstrateLoadUint1(SubstrateBuffer, SubstrateAddressing);
 				uint Data2 = SubstrateLoadUint1(SubstrateBuffer, SubstrateAddressing);
 
@@ -4189,7 +4193,7 @@ FSubstrateBSDF UnpackSubstrateBSDFIn(FSubstrateMaterialContainer SubstrateBuffer
 			else if (OptimisedLegacyMode == SINGLE_OPTLEGACYMODE_SSSPROFILE)
 			{
 			#if SUBSTRATE_DEFERRED_SHADING
-				uint Data0 = SubstrateHeader.SharedLocalBasesTypes_PackedHeader;
+				uint Data0 = SubstrateHeader.PackedHeader;
 				uint Data1 = SubstrateLoadUint1(SubstrateBuffer, SubstrateAddressing);
 				uint Data2 = SubstrateLoadUint1(SubstrateBuffer, SubstrateAddressing);
 
@@ -4251,7 +4255,7 @@ FSubstrateBSDF UnpackSubstrateBSDFIn(FSubstrateMaterialContainer SubstrateBuffer
 	{
 		uint Data = 0;
 	#if SUBSTRATE_DEFERRED_SHADING
-		Data = SubstrateHeader.SharedLocalBasesTypes_PackedHeader;
+		Data = SubstrateHeader.PackedHeader;
 	#endif
 		OutBSDF.State = Data & HEADER_EYEENCODING_MASK;
 		OutBSDF.LuminanceWeightV = 1.0f;
@@ -4265,7 +4269,7 @@ FSubstrateBSDF UnpackSubstrateBSDFIn(FSubstrateMaterialContainer SubstrateBuffer
 	{
 		uint Data = 0;
 	#if SUBSTRATE_DEFERRED_SHADING
-		Data = SubstrateHeader.SharedLocalBasesTypes_PackedHeader;
+		Data = SubstrateHeader.PackedHeader;
 	#endif
 		OutBSDF.State = Data & HEADER_HAIRENCODING_MASK;
 		OutBSDF.LuminanceWeightV = 1.0f;
@@ -4396,7 +4400,7 @@ FSubstrateBSDF UnpackSubstrateBSDFIn(FSubstrateMaterialContainer SubstrateBuffer
 	{
 		float4 Data0 = 0;
 	#if SUBSTRATE_DEFERRED_SHADING
-		Data0 = UnpackRGBA8(SubstrateHeader.SharedLocalBasesTypes_PackedHeader);
+		Data0 = UnpackRGBA8(SubstrateHeader.PackedHeader);
 	#endif
 		EYE_IRISMASK(OutBSDF) = Data0.z;
 		EYE_IRISDISTANCE(OutBSDF) = Data0.w;
diff --git a/Engine/Shaders/Private/Substrate/SubstrateDBuffer.usf b/Engine/Shaders/Private/Substrate/SubstrateDBuffer.usf
--- a/Engine/Shaders/Private/Substrate/SubstrateDBuffer.usf
+++ b/Engine/Shaders/Private/Substrate/SubstrateDBuffer.usf
@@ -186,13 +186,13 @@ void MainCS(
 			const uint PackedRoughness8bits= PackR8(Roughness);
 
 			// Data0 (Header_State|Header_AO|Roughness|Diffuse8bits)
-			MaterialTextureArrayUAV[uint3(PixelCoord, 0)] = (SubstratePixelHeader.SharedLocalBasesTypes_PackedHeader & 0x0000FFFF) | (PackedRoughness8bits << 16) | (PackedDiffuse8Bits << 24);
+			MaterialTextureArrayUAV[uint3(PixelCoord, 0)] = (SubstratePixelHeader.PackedHeader & 0x0000FFFF) | (PackedRoughness8bits << 16) | (PackedDiffuse8Bits << 24);
 			MaterialTextureArrayUAV[uint3(PixelCoord, 1)] = PackedF020Bits | (PackedDiffuse12Bits << 20);
 		}
 		#if !SUBSTRATE_FASTPATH
 		else if (SubstratePixelHeader.IsSingleMaterial())
 		{
-			const uint OptimisedLegacyMode = (SubstratePixelHeader.SharedLocalBasesTypes_PackedHeader >> (HEADER_SINGLEENCODING_BIT_COUNT)) & HEADER_SINGLE_OPTLEGACYMODE_BIT_MASK;
+			const uint OptimisedLegacyMode = (SubstratePixelHeader.PackedHeader >> (HEADER_SINGLEENCODING_BIT_COUNT)) & HEADER_SINGLE_OPTLEGACYMODE_BIT_MASK;
 			if (OptimisedLegacyMode == SINGLE_OPTLEGACYMODE_NONE)
 			{
 				MaterialTextureArrayUAV[uint3(PixelCoord, 2)] = (MaterialTextureArrayUAV[uint3(PixelCoord, 2)] & 0xFFFF00FF) | (PackR8(Roughness) << 8);
diff --git a/Engine/Shaders/Private/Substrate/SubstrateVisualizeCommon.ush b/Engine/Shaders/Private/Substrate/SubstrateVisualizeCommon.ush
--- a/Engine/Shaders/Private/Substrate/SubstrateVisualizeCommon.ush
+++ b/Engine/Shaders/Private/Substrate/SubstrateVisualizeCommon.ush
@@ -190,7 +190,7 @@ void SubstratePrintBSDF(
 			}
 			Newline(Context, RectMax);
 
-			const uint BasisType  = SubstrateGetSharedLocalBasisType(Header.SharedLocalBasesTypes_PackedHeader, BSDF_GETSHAREDLOCALBASISID(BSDF));
+			const uint BasisType  = SubstrateGetSharedLocalBasisType(HEADER_GETSHAREDLOCALBASESTYPE(Header.PackedHeader), BSDF_GETSHAREDLOCALBASISID(BSDF));
 
 			Print(Context, TEXT(" "), FontSilver);
 			Print(Context, BSDF_GETSHAREDLOCALBASISID(BSDF), FontBSDF, 2, 1);	Print(Context, TEXT("       "), FontSilver);
@@ -528,7 +528,7 @@ void SubstratePrintMaterialProperties_Internal(
 		if (bIsSingleMaterial)
 		{
 		#if SUBSTRATE_DEFERRED_SHADING
-			uint OptimisedLegacyMode = (Header.SharedLocalBasesTypes_PackedHeader >> (HEADER_SINGLEENCODING_BIT_COUNT)) & HEADER_SINGLE_OPTLEGACYMODE_BIT_MASK;
+			uint OptimisedLegacyMode = (Header.PackedHeader >> (HEADER_SINGLEENCODING_BIT_COUNT)) & HEADER_SINGLE_OPTLEGACYMODE_BIT_MASK;
 			if (OptimisedLegacyMode ==SINGLE_OPTLEGACYMODE_CLEARCOAT)
 			{
 				Print(Context, TEXT(" - Legacy Clear-Coat"), FontLightRed);
@@ -553,7 +553,7 @@ void SubstratePrintMaterialProperties_Internal(
 		}
 		Newline(Context, RectMax);
 
-		Print(Context, TEXT(" AO    IndIrr TopRoughness PreShadow ZeroShadow ContacShadow Ind.Occluder HasSSS"), FontHeaderStateName);
+		Print(Context, TEXT(" AO    IndIrr TopRoughness PreShadow ZeroShadow ContacShadow Ind.Occluder HasSSS BasisCnt"), FontHeaderStateName);
 		Newline(Context, RectMax);
 
 		Print(Context, TEXT(" "));
@@ -564,7 +564,8 @@ void SubstratePrintMaterialProperties_Internal(
 		PrintBool(Context, Header.HasZeroPrecShadowMask());												Print(Context, TEXT("       "));
 		PrintBool(Context, Header.DoesCastContactShadow());												Print(Context, TEXT("         "));
 		PrintBool(Context, Header.HasDynamicIndirectShadowCasterRepresentation());						Print(Context, TEXT("         "));
-		PrintBool(Context, Header.HasSubsurface());
+		PrintBool(Context, Header.HasSubsurface());														Print(Context, TEXT("          "));
+		Print(Context, Header.GetLocalBasesCount());
 		Newline(Context, RectMax);
 		Newline(Context);
 		
diff --git a/Engine/Shaders/Private/VirtualShadowMaps/VirtualShadowMapProjection.usf b/Engine/Shaders/Private/VirtualShadowMaps/VirtualShadowMapProjection.usf
--- a/Engine/Shaders/Private/VirtualShadowMaps/VirtualShadowMapProjection.usf
+++ b/Engine/Shaders/Private/VirtualShadowMaps/VirtualShadowMapProjection.usf
@@ -75,15 +75,13 @@ FProjectionShadingInfo GetProjectionShadingInfo(uint2 PixelPos)
 	Out.WorldNormal = TopLayerData.WorldNormal;
 	Out.bIsSubsurface = SSSType != SSS_TYPE_INVALID; // legacy IsSubsurfaceModel is true for subsurface profile, eyes, etc.
 
-	// VSM traces traces using the top layer normal when NoL>0 only. So if normals for slabs below do not match, light leaking will occur.
+	// VSM traces using the top layer normal when NoL>0 only. So if normals for slabs below do not match, light leaking will occur.
 	// For instance: a flat translucent top layer above a normal mapped bottom layer.
-	// That is why, for complex materials with multiple BSDFs, we also must not do backface culling.
-	if (SubstratePixelHeader.ClosureCount > 1)
+	// That is why, for materials with multiple local bases, we must not do backface culling.
+	if (SubstratePixelHeader.GetLocalBasesCount() > 1)
 	{
-		// Technically, this could only be done when vertical layering is encountered only. 
-		// But we do not have that data available (or it would be costly to infer from closures)
 		Out.bDisableBackFaceCulling = true;
-}
+	}
 #else
 	const FGBufferData GBufferData = GetGBufferDataUint(PixelPos, true);
 	Out.bIsValid = GBufferData.ShadingModelID != SHADINGMODELID_UNLIT;
