Path Tracer: Improve volumetric local light picking

When building the light cdf for volume lights, only add valid samples to the cdf. This reduces the risk of overflow in scenes with many lights as most lights should not overlap with the ray.

Tune the light importance to take into account the constant portion of the equi-angular pdf. This results in better volumetric light picking when lights with a wide region of influence overlap with lights with a narrow region of influence (such as spot lights).

Use uniform sampling on lights with non-inverse square falloff

Take rectlight barn doors into account for volume sampling. The barndoor parameters defines a penumbra region which can be bounded by planes. With four additional ray/plane tests we can further restrict the computation of volumetric lighting to visible regions.

diff --git a/Engine/Shaders/Private/PathTracing/Light/PathTracingLightCommon.ush b/Engine/Shaders/Private/PathTracing/Light/PathTracingLightCommon.ush
--- a/Engine/Shaders/Private/PathTracing/Light/PathTracingLightCommon.ush
+++ b/Engine/Shaders/Private/PathTracing/Light/PathTracingLightCommon.ush
@@ -99,7 +99,7 @@ FVolumeLightSampleSetup CreateEquiAngularSampler(float LightImportance, float3 C
 	Result.ThetaMin = atan((TMin - Result.T0) * RcpHeight);
 	Result.ThetaMax = atan((TMax - Result.T0) * RcpHeight);
 	Result.K = (Result.ThetaMax - Result.ThetaMin) * Result.Height;
-	Result.LightImportance = LightImportance * RcpHeight;
+	Result.LightImportance = LightImportance *  (Result.ThetaMax - Result.ThetaMin) * RcpHeight;
 	return Result;
 }
 
diff --git a/Engine/Shaders/Private/PathTracing/Light/PathTracingPointLight.ush b/Engine/Shaders/Private/PathTracing/Light/PathTracingPointLight.ush
--- a/Engine/Shaders/Private/PathTracing/Light/PathTracingPointLight.ush
+++ b/Engine/Shaders/Private/PathTracing/Light/PathTracingPointLight.ush
@@ -255,7 +255,13 @@ FVolumeLightSampleSetup PointLight_PrepareLightVolumeSample(
 		}
 		float3 C = GetColor(LightId);
 		float LightImportance = max3(C.x, C.y, C.z);
-		return CreateEquiAngularSampler(LightImportance, Center, RayOrigin, RayDirection, TMin, TMax);
+		if ((SceneLights[LightId].Flags & PATHTRACER_FLAG_NON_INVERSE_SQUARE_FALLOFF_MASK) == 0)
+		{
+			// inverse square falloff requires equi-angular sampling
+			return CreateEquiAngularSampler(LightImportance, Center, RayOrigin, RayDirection, TMin, TMax);
+		}
+		// otherwise, default to uniform sampling
+		return CreateUniformSampler(LightImportance, TMin, TMax);
 	}
 	return NullVolumeLightSetup();
 }
diff --git a/Engine/Shaders/Private/PathTracing/Light/PathTracingRectLight.ush b/Engine/Shaders/Private/PathTracing/Light/PathTracingRectLight.ush
--- a/Engine/Shaders/Private/PathTracing/Light/PathTracingRectLight.ush
+++ b/Engine/Shaders/Private/PathTracing/Light/PathTracingRectLight.ush
@@ -195,6 +195,22 @@ float RectLight_EstimateLight(
 	return LightPower * Falloff * Area * NoL * LNoL / LightDistanceSquared;
 }
 
+void ClipRayByPlane(float3 RayOrigin, float3 RayDirection, float3 N, float3 C, inout float TMin, inout float TMax)
+{
+	const float DoN = dot(RayDirection, N);
+	const float TPlane = dot(C - RayOrigin, N) * rcp(DoN);
+	if (DoN > 0.0)
+	{
+		// plane is pointing in the same direction as the ray, clip TMin
+		TMin = max(TMin, TPlane);
+	}
+	if (DoN < 0.0)
+	{
+		// plane is pointing towards the ray, clip TMax
+		TMax = min(TMax, TPlane);
+	}
+}
+
 FVolumeLightSampleSetup RectLight_PrepareLightVolumeSample(
 	int LightId,
 	float3 RayOrigin,
@@ -217,21 +233,41 @@ FVolumeLightSampleSetup RectLight_PrepareLightVolumeSample(
 
 		// now clip the sphere of influence against the plane of the quad and adjust the near/far distance depending on the sign
 		float3 N = GetNormal(LightId);
-		float DoN = dot(RayDirection, N);
-		float TPlane = dot(Center - RayOrigin, N) * rcp(DoN);
-		if (DoN > 0.0)
-			TMin = max(TMin, TPlane);
-		else
-			TMax = min(TMax, TPlane);
-
-		// TODO: clip against barndoors if the ray intersects them?
-
-		float3 C = GetColor(LightId);
-		float Width = GetWidth(LightId);
-		float Height = GetHeight(LightId);
-		float Area = Width * Height;
-		float LightImportance = max3(C.x, C.y, C.z) * Area;
-		return CreateEquiAngularSampler(LightImportance, Center, RayOrigin, RayDirection, TMin, TMax);
+
+		ClipRayByPlane(RayOrigin, RayDirection, N, Center, TMin, TMax);
+
+		const float3 C = GetColor(LightId);
+		const float Width = GetWidth(LightId);
+		const float Height = GetHeight(LightId);
+#if 1
+		const float BarnCos = GetRectLightBarnCosAngle(LightId);
+		const float BarnLen = GetRectLightBarnLength(LightId);
+		if (BarnCos > 0.035)
+		{
+			// Clip ray against barndoor penumbra planes if barndoors are active (see threshold in GetRect)
+			const float BarnSin = sqrt(1 - BarnCos * BarnCos);
+			const float2 BoundingPlaneX = float2(Width  + BarnLen * BarnSin, BarnLen * BarnCos);
+			const float2 BoundingPlaneY = float2(Height + BarnLen * BarnSin, BarnLen * BarnCos);
+
+			const float3 LightdPdu = GetdPdu(LightId);
+			const float3 LightdPdv = GetdPdv(LightId);
+
+			ClipRayByPlane(RayOrigin, RayDirection, BoundingPlaneX.x * N + BoundingPlaneX.y * LightdPdu, Center - LightdPdu * Width * 0.5, TMin, TMax);
+			ClipRayByPlane(RayOrigin, RayDirection, BoundingPlaneX.x * N - BoundingPlaneX.y * LightdPdu, Center + LightdPdu * Width * 0.5, TMin, TMax);
+			ClipRayByPlane(RayOrigin, RayDirection, BoundingPlaneY.x * N + BoundingPlaneY.y * LightdPdv, Center - LightdPdv * Height * 0.5, TMin, TMax);
+			ClipRayByPlane(RayOrigin, RayDirection, BoundingPlaneY.x * N - BoundingPlaneY.y * LightdPdv, Center + LightdPdv * Height * 0.5, TMin, TMax);
+		}
+#endif
+
+		const float Area = Width * Height;
+		const float LightImportance = max3(C.x, C.y, C.z) * Area;
+		if ((SceneLights[LightId].Flags & PATHTRACER_FLAG_NON_INVERSE_SQUARE_FALLOFF_MASK) == 0)
+		{
+			// inverse square falloff requires equi-angular sampling
+			return CreateEquiAngularSampler(LightImportance, Center, RayOrigin, RayDirection, TMin, TMax);
+		}
+		// otherwise, default to uniform sampling
+		return CreateUniformSampler(LightImportance, TMin, TMax);
 	}
 	return NullVolumeLightSetup();
 }
diff --git a/Engine/Shaders/Private/PathTracing/PathTracingCore.ush b/Engine/Shaders/Private/PathTracing/PathTracingCore.ush
--- a/Engine/Shaders/Private/PathTracing/PathTracingCore.ush
+++ b/Engine/Shaders/Private/PathTracing/PathTracingCore.ush
@@ -772,7 +772,7 @@ FPackedPathTracingPayload TraceTransparentRay(inout FPathState PathState, int Bo
 							}
 
 							// if this is not the last bounce, consider the possibility of scattering to give us medium hits
-							if (!bLastBounce)
+							if (!bLastBounce && PathState.ShouldAccumulateVolume())
 							{
 								float3 SigmaS = min(Result.SigmaSRayleigh + Result.SigmaSHG + Result.SigmaSDualHG, SigmaT);
 								// accumulate a signal for the denoiser
@@ -1597,6 +1597,7 @@ bool PathTracingKernel(inout FPathState PathState, int Bounce)
 {
 	// This array will hold a CDF for light picking
 	float LightPickingCdf[RAY_TRACING_LIGHT_COUNT_MAXIMUM];
+	int LightPickingIds[RAY_TRACING_LIGHT_COUNT_MAXIMUM];
 
 #if 0
 	// visualize the shadow ray handling directly
@@ -1620,19 +1621,37 @@ bool PathTracingKernel(inout FPathState PathState, int Bounce)
 		const float VTMin = VolumeSegment.Interval.VolumeTMin;
 		const float VTMax = VolumeSegment.Interval.VolumeTMax;
 
-		const int NumLights = min(SceneLightCount, RAY_TRACING_LIGHT_COUNT_MAXIMUM); // make sure we don't overflow the cdf (TODO: relax this limit)
 		float LightPickingCdfSum = 0.0;
-		for (int LightId = 0; LightId < NumLights; LightId++)
+		int NumLights = 0;
+		for (int LightId = SceneInfiniteLightCount; LightId < SceneLightCount; LightId++)
 		{
 			FVolumeLightSampleSetup LightSetup = PrepareLightVolumeSample(LightId, Ro, Rd, VTMin, VTMax);
 			if (LightSetup.IsValid())
 			{
 				float LightProb = LightSetup.LightImportance * GetVolumetricScatteringIntensity(LightId);
-				LightPickingCdfSum += LightProb;
+				if (LightProb > 0)
+				{
+					LightPickingCdfSum += LightProb;
+					LightPickingCdf[NumLights] = LightPickingCdfSum;
+					LightPickingIds[NumLights] = LightId;
+					NumLights++;
+					if (NumLights == RAY_TRACING_LIGHT_COUNT_MAXIMUM)
+					{
+						// reached max
+						break;
+					}
+				}
 			}
-			LightPickingCdf[LightId] = LightPickingCdfSum;
 		}
 
+		#if 0
+		// TODO: turn this into a debug mode
+		float NumLightsMax = min3(256.0f, SceneLightCount - SceneInfiniteLightCount, RAY_TRACING_LIGHT_COUNT_MAXIMUM);
+		PathState.Radiance = NumLights > 0 ? BlueGreenRedColorMap(saturate(float(NumLights) / NumLightsMax)) : 0.18;
+		PathState.Radiance *= PackedPayload.IsHit() ? abs(dot(PackedPayload.UnpackWorldNormal(), PathState.Ray.Direction)) : 1.0;
+		return false;
+		#endif
+
 		if (LightPickingCdfSum > 0.0)
 		{
 			// at least one light is overlapping with our ray, so we have a chance to sample it
@@ -1640,6 +1659,7 @@ bool PathTracingKernel(inout FPathState PathState, int Bounce)
 			int LightSampleLightId = 0;
 			float LightPickPdf = 0.0;
 			SelectLight(RandSample.x * LightPickingCdfSum, NumLights, LightPickingCdf, LightSampleLightId, LightPickPdf);
+			LightSampleLightId = LightPickingIds[LightSampleLightId]; // map back to original LightId
 			LightPickPdf /= LightPickingCdfSum;
 
 			// picked a light! now use the equi-angular sampler to pick a position along the ray and store it for later
