Fix for virtual shadow on clear coat after 36519540.
Fix tested on/off.

#tests none
#rnx

diff --git a/Engine/Shaders/Private/VirtualShadowMaps/VirtualShadowMapProjection.usf b/Engine/Shaders/Private/VirtualShadowMaps/VirtualShadowMapProjection.usf
--- a/Engine/Shaders/Private/VirtualShadowMaps/VirtualShadowMapProjection.usf
+++ b/Engine/Shaders/Private/VirtualShadowMaps/VirtualShadowMapProjection.usf
@@ -87,6 +87,7 @@ FProjectionShadingInfo GetProjectionShadingInfo(uint2 PixelPos)
 	Out.bIsValid = GBufferData.ShadingModelID != SHADINGMODELID_UNLIT;
 	Out.bIsHair = GBufferData.ShadingModelID == SHADINGMODELID_HAIR;
 	Out.bIsSubsurface = IsSubsurfaceModel(GBufferData.ShadingModelID);
+	Out.bDisableBackFaceCulling = GBufferData.ShadingModelID == SHADINGMODELID_CLEAR_COAT;
 	Out.WorldNormal = GBufferData.WorldNormal;
 #endif
 
