Fixed Screen Position in UI Materials.
 - 36574650 introduced returning 0 instead of the default primitive struct values. The slate shader patched up the primitive flags, but those are now returning 0, which disables WPO.
 - Use the raw WPO method instead since we don't need all the other stuff.

diff --git a/Engine/Shaders/Private/SceneData.ush b/Engine/Shaders/Private/SceneData.ush
--- a/Engine/Shaders/Private/SceneData.ush
+++ b/Engine/Shaders/Private/SceneData.ush
@@ -63,26 +63,6 @@
 	#endif
 #endif
 
-uint GetDefaultPrimitiveSceneDataFlags()
-{
-	// Should match FPrimitiveUniformShaderParametersBuilder::Default()
-	return
-		PRIMITIVE_SCENE_DATA_FLAG_CAST_SHADOWS |
-		PRIMITIVE_SCENE_DATA_FLAG_HAS_CAST_CONTACT_SHADOW |
-		PRIMITIVE_SCENE_DATA_FLAG_EVALUATE_WORLD_POSITION_OFFSET;
-}
-
-uint GetDefaultPrimitiveVisibilityFlags()
-{
-	// Should match FPrimitiveUniformShaderParametersBuilder::Default()
-	return
-		PRIMITIVE_VISIBILITY_FLAG_VISIBLE_IN_GAME |
-		PRIMITIVE_VISIBILITY_FLAG_VISIBLE_IN_EDITOR |
-		PRIMITIVE_VISIBILITY_FLAG_VISIBLE_IN_REFLECTION_CAPTURES |
-		PRIMITIVE_VISIBILITY_FLAG_VISIBLE_IN_REAL_TIME_SKY_CAPTURES |
-		PRIMITIVE_VISIBILITY_FLAG_VISIBLE_IN_RAY_TRACING;
-}
-
 #define INVALID_INSTANCE_PAYLOAD_OFFSET (0xFFFFFFFFu)
 
 #if PLATFORM_ALLOW_SCENE_DATA_COMPRESSED_TRANSFORMS
diff --git a/Engine/Shaders/Private/SlateVertexShader.usf b/Engine/Shaders/Private/SlateVertexShader.usf
--- a/Engine/Shaders/Private/SlateVertexShader.usf
+++ b/Engine/Shaders/Private/SlateVertexShader.usf
@@ -39,10 +39,6 @@ VertexToPixelInterpolants Main(
 	VertexParameters.WorldPosition = WorldPosition.xyz;
 
 	VertexParameters.VertexColor = VOut.Color;
-
-	// Init primitive flags to a good default.
-	VertexParameters.SceneData.Primitive.Flags = GetDefaultPrimitiveSceneDataFlags();
-	VertexParameters.SceneData.Primitive.VisibilityFlags = GetDefaultPrimitiveVisibilityFlags();
 	VertexParameters.LWCData = MakeMaterialLWCData(VertexParameters);
 
 	bool bCalledEvaluateVertexMaterialAttributes = false;
@@ -120,7 +116,7 @@ VertexToPixelInterpolants Main(
 	}
 #endif
 
-	WorldPosition.xyz = GetMaterialWorldPositionOffset(VertexParameters);
+	WorldPosition.xyz = GetMaterialWorldPositionOffsetRaw(VertexParameters);
 #endif
 
 #endif // USE_MATERIALS
