Fix missing SSS on surface with precomputed lighting when Substrate is enabled.

diff --git a/Engine/Shaders/Private/BasePassPixelShader.usf b/Engine/Shaders/Private/BasePassPixelShader.usf
--- a/Engine/Shaders/Private/BasePassPixelShader.usf
+++ b/Engine/Shaders/Private/BasePassPixelShader.usf
@@ -1719,6 +1719,7 @@ void FPixelShaderInOut_MainPS(
 #endif
 	{
 		float3 EmissiveLuminance = 0.0f;
+		float3 ScatteredPrecomputedLuminance = 0.0f;
 		uint2 PixelPos = uint2(In.SvPosition.xy);
 		const float3 WorldBentNormal0 = GetWorldBentNormalZero(MaterialParameters);
 
@@ -1755,9 +1756,16 @@ void FPixelShaderInOut_MainPS(
 			,BasePassInterpolants
 			,LightmapVTPageTableResult
 			,VolumetricLightmapBrickTextureUVs
+			,ScatteredPrecomputedLuminance
 	#endif
 		);
 
+	#if MATERIAL_SUBSTRATE_OPAQUE_PRECOMPUTED_LIGHTING
+		// Add ScatteredPrecomputedLuminance to DiffuseColor in order to handle correctly SSS diffuse with precomputed lighting.
+		// Luminance added to Color and to DiffuseColor are disjoint (emissive luminance vs. scattered precomputed luminance)
+		DiffuseColor = ScatteredPrecomputedLuminance;
+	#endif
+
 		// Write out MRT data
 	#if SUBSTRATE_BASE_PASS_MRT_OUTPUT_COUNT != 3
 	#error Substrate SUBSTRATE_BASE_PASS_MRT_OUTPUT_COUNT has been update but not SubstrateOutput
diff --git a/Engine/Shaders/Private/Substrate/SubstrateExport.ush b/Engine/Shaders/Private/Substrate/SubstrateExport.ush
--- a/Engine/Shaders/Private/Substrate/SubstrateExport.ush
+++ b/Engine/Shaders/Private/Substrate/SubstrateExport.ush
@@ -107,7 +107,7 @@ void PackSubstrateOut(
 	float3 V,
 	float3 WorldBentNormal0,
 	inout bool bSubstrateSubsurfaceEnable,
-	inout float3 EmissiveColor,
+	inout float3 OutEmissiveLuminance,
 	inout FSubstrateSubsurfaceData SSSData,
 	inout FSubstrateTopLayerData TopLayerData,
 	inout FSubstrateOpaqueRoughRefractionData OpaqueRoughRefractionData
@@ -117,6 +117,7 @@ void PackSubstrateOut(
 	,in FBasePassInterpolantsVSToPS BasePassInterpolants
 	,in VTPageTableResult LightmapVTPageTableResult
 	,in float3 VolumetricLightmapBrickTextureUVs
+	,inout float3 OutScatteredPrecomputedLuminance
 #endif
 )
 {
@@ -124,7 +125,10 @@ void PackSubstrateOut(
 #if SUBSTRATE_INLINE_SHADING && SUBSTRATE_CLAMPED_CLOSURE_COUNT > 0
 
 	bSubstrateSubsurfaceEnable = false;
-	EmissiveColor = 0.0f;
+	OutEmissiveLuminance = 0.0f;
+	#if MATERIAL_SUBSTRATE_OPAQUE_PRECOMPUTED_LIGHTING
+	OutScatteredPrecomputedLuminance = 0.0f;
+	#endif
 
 	const float FullyRough = 1.0f;
 
@@ -191,7 +195,7 @@ void PackSubstrateOut(
 					{
 					case SUBSTRATE_BSDF_TYPE_SLAB:
 					{
-						EmissiveColor += BSDF_GETEMISSIVE(BSDF) * BSDF.LuminanceWeightV;
+						OutEmissiveLuminance += BSDF_GETEMISSIVE(BSDF) * BSDF.LuminanceWeightV;
 
 						// Simple and single bsdf do not allow weights other than 1
 						const bool LuminanceWeightEqualOne = all(BSDF.LuminanceWeightV == 1.0f);
@@ -243,7 +247,7 @@ void PackSubstrateOut(
 						bIsOnlySlab = false;
 						bIsSimpleMaterial = false;
 						bIsSingleMaterial = false;
-						EmissiveColor += BSDF_GETEMISSIVE(BSDF) * BSDF.LuminanceWeightV;
+						OutEmissiveLuminance += BSDF_GETEMISSIVE(BSDF) * BSDF.LuminanceWeightV;
 
 						TopLayerTotalWeight				+= BSDF.TopLayerDataWeight;
 						TopLayerData.Roughness			+= BSDF.TopLayerDataWeight * FullyRough;
@@ -277,7 +281,7 @@ void PackSubstrateOut(
 						bIsOnlySlab = false;
 						bIsSimpleMaterial = false;
 						bIsSingleMaterial = false;
-						EmissiveColor += BSDF_GETEMISSIVE(BSDF) * BSDF.LuminanceWeightV;
+						OutEmissiveLuminance += BSDF_GETEMISSIVE(BSDF) * BSDF.LuminanceWeightV;
 
 						TopLayerTotalWeight				+= BSDF.TopLayerDataWeight;
 						TopLayerData.Roughness			+= BSDF.TopLayerDataWeight * SLW_ROUGHNESS(BSDF);
@@ -410,7 +414,7 @@ void PackSubstrateOut(
 				{
 					// Convert BSDF to data ready to evaluate precomputed lighting
 					float2 SpecMultiBounceAO_IndirectIrradiance = 0.0f;
-					EmissiveColor += SubstrateGetBSDFPrecomputedLighting(
+					float3 PrecomputedLighting = SubstrateGetBSDFPrecomputedLighting(
 						Settings,
 						SubstratePixelHeader,
 						BSDF,
@@ -424,6 +428,17 @@ void PackSubstrateOut(
 						SpecMultiBounceAO_IndirectIrradiance);
 					SpecMultiBounceAO_IndirectIrradiance_Acc += SpecMultiBounceAO_IndirectIrradiance;
 
+					// Selectively add PrecomputedLighting to either OutScatteredPrecomputedLuminance or OutEmissiveLuminance 
+					// depending if the light will be processed or not by SSS postprocess
+					if (bSubstrateSubsurfaceEnable)
+					{
+						OutScatteredPrecomputedLuminance = PrecomputedLighting;
+					}
+					else
+					{
+						OutEmissiveLuminance += PrecomputedLighting;
+					}
+
 					// Stop writing if above budget.
 					BSDFCount++;
 				}
