Reverted back to including hidden instances when building instance hierarchy explicit bounds, as they are not rebuilt when visibility is toggled. HLOD streams in instances and then toggles visibility leading to visual glitches  if we skip hidden instances. This is ok with GPU-managed instances as they are not included in the hierarchy.

#rnx
#tests CitySample cooked on a console and Windows Editor/PIE tested.

diff --git a/Engine/Shaders/Private/SceneCulling/SceneCullingBuildExplicitBounds.usf b/Engine/Shaders/Private/SceneCulling/SceneCullingBuildExplicitBounds.usf
--- a/Engine/Shaders/Private/SceneCulling/SceneCullingBuildExplicitBounds.usf
+++ b/Engine/Shaders/Private/SceneCulling/SceneCullingBuildExplicitBounds.usf
@@ -70,7 +70,7 @@ void ComputeExplicitCellBounds(uint3 GroupId : SV_GroupID, uint GroupThreadIndex
 			uint InstanceFlags;
 			LoadInstancePrimitiveIdAndFlags(InstanceId, Scene.GPUScene.InstanceDataSOAStride, PrimitiveId, InstanceFlags);
 
-			if (PrimitiveId == INVALID_PRIMITIVE_ID || (InstanceFlags & INSTANCE_SCENE_DATA_FLAG_HIDDEN) != 0u)
+			if (PrimitiveId == INVALID_PRIMITIVE_ID)
 			{
 				continue;
 			}
