Substrate - Fix for directional light on mobile causing nan with anisotropic lighting in level without dir lights.

#rnx

diff --git a/Engine/Shaders/Private/Substrate/SubstrateMobileForwardLighting.ush b/Engine/Shaders/Private/Substrate/SubstrateMobileForwardLighting.ush
--- a/Engine/Shaders/Private/Substrate/SubstrateMobileForwardLighting.ush
+++ b/Engine/Shaders/Private/Substrate/SubstrateMobileForwardLighting.ush
@@ -218,7 +218,7 @@ float3 SubstrateMobileForwardLighting(
 				////
 				//// Evaluate the single directional light selected for forward lighting.
 				////
-				const bool bHasDirectionalLight = any(MobileDirectionalLight.DirectionalLightColor > 0);
+				const bool bHasDirectionalLight = any(MobileDirectionalLight.DirectionalLightColor.rgb > 0);
 				BRANCH
 				if (bHasDirectionalLight)
 				{
