Fix shader compilation on certain platform preview with Substrate.

diff --git a/Engine/Shaders/Private/DBufferDecalShared.ush b/Engine/Shaders/Private/DBufferDecalShared.ush
--- a/Engine/Shaders/Private/DBufferDecalShared.ush
+++ b/Engine/Shaders/Private/DBufferDecalShared.ush
@@ -97,8 +97,10 @@ void ConvertToMetalness(
 }
 
 // @param BufferUV - UV space in the DBuffer textures
-FSubstrateDBuffer GetDBufferData(float2 BufferUV, uint RTMaskBit)
+FSubstrateDBuffer GetDBufferData(float2 BufferUV, uint RTMaskBit, uint EyeIndex)
 {
+	// SUBSTRATE_TODO: Add support MOBILE_MULTI_VIEW and INSTANCED_STEREO on mobile (using EyeIndex)
+
 	// Setup default values, which mean that no decals are present in DBuffer
 	float4 DBufferA = float4(0.0, 0.0, 0.0, 1.0);
 	float4 DBufferB = float4(128.0f / 255.0f, 128.f / 255.5f, 128.f / 255.5f, 1.0);
@@ -122,6 +124,10 @@ FSubstrateDBuffer GetDBufferData(float2 BufferUV, uint RTMaskBit)
 	return SubstrateUnpackDBuffer(DBufferA, DBufferB, DBufferC);
 }
 
+FSubstrateDBuffer GetDBufferData(float2 BufferUV, uint RTMaskBit)
+{
+	return GetDBufferData(BufferUV, RTMaskBit, 0 /*EyeIndex*/);
+}
 
 // Define how DBuffer is applied onto existing layers:
 // * 0: Apply DBuffer as per-component parameter blending (legacy mode). Default. 
