Path Tracer: Implement custom shading model for volumetric lighting models

This is useful to get a better match on fog cards and other translucent items which do not use the main shading models. The implementation re-uses the unlit shading model for this purpose since we are out of bits for ShadingModelId itself.

The non-directional case is implemented as a uniform phase function, while the directional case uses a wrapped cosine lobe. The directionality term blends between both which can still be efficiently importance sampled. This does not exactly match the deferred renderer, but is guaranteed to be energy conserving.

diff --git a/Engine/Shaders/Private/PathTracing/Material/PathTracingMaterialSampling.ush b/Engine/Shaders/Private/PathTracing/Material/PathTracingMaterialSampling.ush
--- a/Engine/Shaders/Private/PathTracing/Material/PathTracingMaterialSampling.ush
+++ b/Engine/Shaders/Private/PathTracing/Material/PathTracingMaterialSampling.ush
@@ -32,18 +32,10 @@ PathTracingMaterialSampling.usf: Material sampling functions
 
 float3 EstimateMaterialAlbedo(FPathTracingPayload Payload)
 {
-
-	if (Payload.ShadingModelID == SHADINGMODELID_UNLIT)
-	{
-		// this object does not scatter light
-		return 0;
-	}
 	// TODO: make this more accurate by estimating the directional albedo?
 	float3 Albedo = 0;
 	switch (Payload.ShadingModelID)
 	{
-		case SHADINGMODELID_UNLIT:
-			return 0;
 #if SIMPLIFIED_MATERIAL_SHADER
 		case SHADINGMODELID_TWOSIDED_FOLIAGE:
 			Albedo = Payload.BaseColor + Payload.SubsurfaceColor;
@@ -135,10 +127,6 @@ void AccumulateAlbedoNormal(inout FPathTracingPayload Payload, float3 PathThroug
 		// no point tracking past a certain roughness
 		return;
 	}
-	if (Payload.ShadingModelID == SHADINGMODELID_UNLIT)
-	{
-		return;
-	}
 	// figure out how much the current path should contribute to the denoiser aovs
 	PathThroughput *= 1 - saturate((PathRoughness - MinRoughness) / (MaxRoughness - MinRoughness));
 #if PATHTRACING_SUBSTRATE_PAYLOAD
@@ -400,6 +388,14 @@ void AdjustPayloadAfterUnpack(inout FPathTracingPayload HitPayload, bool bApplyO
 #else
 	switch (HitPayload.ShadingModelID)
 	{
+		case SHADINGMODELID_UNLIT:
+		{
+			if (bApplyOverrides)
+			{
+				HitPayload.BaseColor  = HitPayload.BaseColor * View.DiffuseOverrideParameter.w + View.DiffuseOverrideParameter.xyz;
+			}
+			break;
+		}
 		case SHADINGMODELID_SUBSURFACE:
 		case SHADINGMODELID_PREINTEGRATED_SKIN:
 		{
@@ -582,8 +578,7 @@ FMaterialSample SampleMaterial(
 	switch (Payload.ShadingModelID)
 	{
 		case SHADINGMODELID_UNLIT:
-			// Nothing to do, material does not scatter light
-			return NullMaterialSample();
+			return Volumetric_SampleMaterial(V_World, Payload, RandSample);
 #if PATHTRACING_SUBSTRATE_PAYLOAD
 		case SHADINGMODELID_SUBSTRATE:
 			if (Payload.HasRefraction())
@@ -641,7 +636,7 @@ FMaterialEval EvalMaterial(
 	switch (Payload.ShadingModelID)
 	{
 		case SHADINGMODELID_UNLIT:
-			break;
+			return Volumetric_EvalMaterial(V_World, L_World, Payload, DiffuseSpecularScale);
 #if PATHTRACING_SUBSTRATE_PAYLOAD
 		case SHADINGMODELID_SUBSTRATE:
 			if (Payload.HasRefraction())
diff --git a/Engine/Shaders/Private/PathTracing/Material/PathTracingMedium.ush b/Engine/Shaders/Private/PathTracing/Material/PathTracingMedium.ush
--- a/Engine/Shaders/Private/PathTracing/Material/PathTracingMedium.ush
+++ b/Engine/Shaders/Private/PathTracing/Material/PathTracingMedium.ush
@@ -110,3 +110,41 @@ FMaterialEval Medium_EvalMaterial(
 	return Result;
 #endif
 }
+
+
+FMaterialSample Volumetric_SampleMaterial(
+	float3 V_World,
+	FPathTracingPayload Payload,
+	float3 RandSample)
+{
+
+	const float Directionality = Payload.Specular;
+	// phase function = (1 + Directionality * CosTheta) / 4Pi
+	// The expression below gives exact importance sampling
+	const float CosTheta = Directionality > 0 ? (-1.0 + sqrt(Pow2(1 - Directionality) + 4.0 * Directionality * RandSample.x)) * rcp(Directionality) : -1.0 + 2.0 * RandSample.x;
+	const float SinTheta = sqrt(saturate(1.0 - CosTheta * CosTheta));
+	const float Phi = RandSample.y * (2.0 * PI);
+
+	const float3 L_World = TangentToWorld(float3(cos(Phi) * SinTheta, sin(Phi) * SinTheta, CosTheta), Payload.WorldNormal);
+
+	const float3 Color = Payload.GetBaseColor() * Payload.BSDFOpacity;
+	const float Inv4PI = 0.25 / PI;
+	const float Pdf = Inv4PI + Inv4PI * Directionality * CosTheta;
+
+	return CreateMaterialSample(L_World, Color, Pdf, sign(CosTheta), 1.0, PATHTRACER_SCATTER_DIFFUSE);
+}
+
+FMaterialEval Volumetric_EvalMaterial(
+	float3 V_World,
+	float3 L_World,
+	FPathTracingPayload Payload,
+	float2 DiffuseSpecularScale
+)
+{
+	const float Directionality = Payload.Specular;
+	const float3 Color = Payload.GetBaseColor() * Payload.BSDFOpacity;
+	const float CosTheta = dot(L_World, Payload.WorldNormal);
+	const float Inv4PI = 0.25 / PI;
+	const float Pdf = Inv4PI + Inv4PI * Directionality * CosTheta;
+	return CreateMaterialEval(Color * DiffuseSpecularScale.x, Pdf);
+}
diff --git a/Engine/Shaders/Private/PathTracing/PathTracingCommon.ush b/Engine/Shaders/Private/PathTracing/PathTracingCommon.ush
--- a/Engine/Shaders/Private/PathTracing/PathTracingCommon.ush
+++ b/Engine/Shaders/Private/PathTracing/PathTracingCommon.ush
@@ -684,6 +684,7 @@ struct FPathTracingPayload : FMinimalPayload
 		return ShadingModelID == SHADINGMODELID_TWOSIDED_FOLIAGE ||
 			   ShadingModelID == SHADINGMODELID_HAIR ||
 			   ShadingModelID == SHADINGMODELID_MEDIUM ||
+			   ShadingModelID == SHADINGMODELID_UNLIT || // Volumetric Directional/Non-Directional
 			   IsMaterialSolidGlass() ||
 			   IsMaterialThinGlass();
 	}
diff --git a/Engine/Shaders/Private/PathTracing/PathTracingMaterialHitShader.usf b/Engine/Shaders/Private/PathTracing/PathTracingMaterialHitShader.usf
--- a/Engine/Shaders/Private/PathTracing/PathTracingMaterialHitShader.usf
+++ b/Engine/Shaders/Private/PathTracing/PathTracingMaterialHitShader.usf
@@ -272,7 +272,12 @@ RAY_TRACING_ENTRY_CLOSEST_HIT(PathTracingMaterialCHS,
 	/**
 	 * Set common material attributes for both full and simplified materials
 	 **/
+
+#if (TRANSLUCENCY_LIGHTING_VOLUMETRIC_PERVERTEX_NONDIRECTIONAL || TRANSLUCENCY_LIGHTING_VOLUMETRIC_PERVERTEX_DIRECTIONAL || TRANSLUCENCY_LIGHTING_VOLUMETRIC_NONDIRECTIONAL || TRANSLUCENCY_LIGHTING_VOLUMETRIC_DIRECTIONAL)
+	Payload.ShadingModelID = SHADINGMODELID_UNLIT;
+#else
 	Payload.ShadingModelID = GetMaterialShadingModel(PixelMaterialInputs);
+#endif
 
 #if MATERIALBLENDING_ALPHACOMPOSITE
 	Payload.BSDFOpacity = 1.0;
@@ -347,7 +352,11 @@ RAY_TRACING_ENTRY_CLOSEST_HIT(PathTracingMaterialCHS,
 #endif
 	Payload.WorldGeoNormal = GeoNormalSign * GeoNormal;
 	Payload.WorldSmoothNormal = MaterialParameters.TwoSidedSign * TransformTangentNormalToWorld(MaterialParameters.TangentToWorld, float3(0, 0, 1));
+#if (TRANSLUCENCY_LIGHTING_VOLUMETRIC_PERVERTEX_NONDIRECTIONAL || TRANSLUCENCY_LIGHTING_VOLUMETRIC_PERVERTEX_DIRECTIONAL || TRANSLUCENCY_LIGHTING_VOLUMETRIC_NONDIRECTIONAL || TRANSLUCENCY_LIGHTING_VOLUMETRIC_DIRECTIONAL)
+	// no need to adjust normal on volumetric shading models
+#else
 	Payload.WorldSmoothNormal = AdjustShadingNormal(Payload.WorldSmoothNormal, Payload.WorldGeoNormal, WorldRayDirection());
+#endif
 
 
 #if SUBSTRATE_ENABLED
@@ -779,7 +788,11 @@ RAY_TRACING_ENTRY_CLOSEST_HIT(PathTracingMaterialCHS,
 	// already flipped if the material was in tangent space, so add the flip if it wasn't
 	Payload.WorldNormal *= MaterialParameters.TwoSidedSign;
 #endif
+#if (TRANSLUCENCY_LIGHTING_VOLUMETRIC_PERVERTEX_NONDIRECTIONAL || TRANSLUCENCY_LIGHTING_VOLUMETRIC_PERVERTEX_DIRECTIONAL || TRANSLUCENCY_LIGHTING_VOLUMETRIC_NONDIRECTIONAL || TRANSLUCENCY_LIGHTING_VOLUMETRIC_DIRECTIONAL)
+	// volumetric shading doesn't need to adjust normal
+#else
 	Payload.WorldNormal = AdjustShadingNormal(Payload.WorldNormal, Payload.WorldGeoNormal, WorldRayDirection());
+#endif
 
 	// Store the results in local variables and reuse instead of calling the functions multiple times.
 	half3 BaseColor = GetMaterialBaseColor(PixelMaterialInputs);
@@ -819,6 +832,30 @@ RAY_TRACING_ENTRY_CLOSEST_HIT(PathTracingMaterialCHS,
 	Payload.WorldTangent = MaterialParameters.TangentToWorld[2];
 #endif
 
+#if (TRANSLUCENCY_LIGHTING_VOLUMETRIC_PERVERTEX_NONDIRECTIONAL || TRANSLUCENCY_LIGHTING_VOLUMETRIC_PERVERTEX_DIRECTIONAL || TRANSLUCENCY_LIGHTING_VOLUMETRIC_NONDIRECTIONAL || TRANSLUCENCY_LIGHTING_VOLUMETRIC_DIRECTIONAL)
+	if (Payload.ShadingModelID == SHADINGMODELID_UNLIT)
+	{
+		// encode the directionality using the specular field
+#if TRANSLUCENCY_LIGHTING_VOLUMETRIC_PERVERTEX_DIRECTIONAL || TRANSLUCENCY_LIGHTING_VOLUMETRIC_DIRECTIONAL
+		Payload.Specular = GetMaterialTranslucencyDirectionalLightingIntensity();
+#else
+		Payload.Specular = 0;
+#endif
+		Payload.Metallic = 0;
+		Payload.Roughness = 1;
+	}
+	else
+#elif MATERIAL_SHADINGMODEL_UNLIT
+	if (Payload.ShadingModelID == SHADINGMODELID_UNLIT)
+	{
+		// Regular unlit material, make sure the volumetric component is zero'ed out
+		Payload.BaseColor = 0;
+		Payload.Specular = 0;
+	}
+	else
+#endif
+
+
 #if MATERIAL_SHADINGMODEL_CLEAR_COAT
 	if (Payload.ShadingModelID == SHADINGMODELID_CLEAR_COAT)
 	{
