Path Tracer: Fix compilation when Substrate is enabled

Encode volumetric shading model directionality in Anisotropy which is a valid field when substrate is enabled as well.

#rnx

diff --git a/Engine/Shaders/Private/PathTracing/Material/PathTracingMedium.ush b/Engine/Shaders/Private/PathTracing/Material/PathTracingMedium.ush
--- a/Engine/Shaders/Private/PathTracing/Material/PathTracingMedium.ush
+++ b/Engine/Shaders/Private/PathTracing/Material/PathTracingMedium.ush
@@ -118,7 +118,7 @@ FMaterialSample Volumetric_SampleMaterial(
 	float3 RandSample)
 {
 
-	const float Directionality = Payload.Specular;
+	const float Directionality = Payload.Anisotropy;
 	// phase function = (1 + Directionality * CosTheta) / 4Pi
 	// The expression below gives exact importance sampling
 	const float CosTheta = Directionality > 0 ? (-1.0 + sqrt(Pow2(1 - Directionality) + 4.0 * Directionality * RandSample.x)) * rcp(Directionality) : -1.0 + 2.0 * RandSample.x;
@@ -141,7 +141,7 @@ FMaterialEval Volumetric_EvalMaterial(
 	float2 DiffuseSpecularScale
 )
 {
-	const float Directionality = Payload.Specular;
+	const float Directionality = Payload.Anisotropy;
 	const float3 Color = Payload.GetBaseColor() * Payload.BSDFOpacity;
 	const float CosTheta = dot(L_World, Payload.WorldNormal);
 	const float Inv4PI = 0.25 / PI;
diff --git a/Engine/Shaders/Private/PathTracing/PathTracingMaterialHitShader.usf b/Engine/Shaders/Private/PathTracing/PathTracingMaterialHitShader.usf
--- a/Engine/Shaders/Private/PathTracing/PathTracingMaterialHitShader.usf
+++ b/Engine/Shaders/Private/PathTracing/PathTracingMaterialHitShader.usf
@@ -837,10 +837,11 @@ RAY_TRACING_ENTRY_CLOSEST_HIT(PathTracingMaterialCHS,
 	{
 		// encode the directionality using the specular field
 #if TRANSLUCENCY_LIGHTING_VOLUMETRIC_PERVERTEX_DIRECTIONAL || TRANSLUCENCY_LIGHTING_VOLUMETRIC_DIRECTIONAL
-		Payload.Specular = GetMaterialTranslucencyDirectionalLightingIntensity();
+		Payload.Anisotropy = GetMaterialTranslucencyDirectionalLightingIntensity();
 #else
-		Payload.Specular = 0;
+		Payload.Anisotropy = 0;
 #endif
+		Payload.Specular = 0;
 		Payload.Metallic = 0;
 		Payload.Roughness = 1;
 	}
