Fix hair (voxel) shadow causing incorrect shadow.

With Substrate this case incorrect light transmission from thin diffuse surface.

diff --git a/Engine/Shaders/Private/VirtualShadowMaps/VirtualShadowMapProjection.usf b/Engine/Shaders/Private/VirtualShadowMaps/VirtualShadowMapProjection.usf
--- a/Engine/Shaders/Private/VirtualShadowMaps/VirtualShadowMapProjection.usf
+++ b/Engine/Shaders/Private/VirtualShadowMaps/VirtualShadowMapProjection.usf
@@ -440,6 +440,10 @@ FVirtualShadowMapSampleResult ProjectLight(
 					}
 					#endif
 				}
+				else
+				{
+					AccShadowFactor += 1.f;
+				}
 			}
 
 			AccShadowFactor /= float(max(1, min(RayIndex, HairRayCount)));
