Ray Tracing: Light Grid build was not computing the index correctly after the light grid axis change

This could lead to reading from un-initialized memory, or missing lighting.

#rnx

diff --git a/Engine/Shaders/Private/RayTracing/RayTracingBuildLightGrid.usf b/Engine/Shaders/Private/RayTracing/RayTracingBuildLightGrid.usf
--- a/Engine/Shaders/Private/RayTracing/RayTracingBuildLightGrid.usf
+++ b/Engine/Shaders/Private/RayTracing/RayTracingBuildLightGrid.usf
@@ -43,7 +43,7 @@ void RayTracingBuildLightGridCS(uint3 DispatchThreadId : SV_DispatchThreadID)
 
 	uint NumVisible = 0;
 	uint MaxVisible = LightGridMaxCount;
-	uint Offset = LightGridMaxCount * (DispatchThreadId.x + LightGridResolution * DispatchThreadId.y);
+	uint Offset = LightGridMaxCount * (Axis + 3 * (DispatchThreadId.x + LightGridResolution * DispatchThreadId.y));
 	for (uint LightIndex = SceneInfiniteLightCount; LightIndex < SceneLightCount; LightIndex++)
 	{
 		uint LightType = SceneLights[LightIndex].Type;
