Include Visualization.ush for HUEtoRGB such that it can compile without dead code elimination being enabled.

diff --git a/Engine/Shaders/Private/PathTracing/Light/PathTracingLightGrid.ush b/Engine/Shaders/Private/PathTracing/Light/PathTracingLightGrid.ush
--- a/Engine/Shaders/Private/PathTracing/Light/PathTracingLightGrid.ush
+++ b/Engine/Shaders/Private/PathTracing/Light/PathTracingLightGrid.ush
@@ -8,6 +8,9 @@
 
 
 #if USE_PATH_TRACING_LIGHT_GRID
+
+#include "../../Visualization.ush"
+
 // extra shader parameters that must be defined when using the light grid
 uint SceneInfiniteLightCount;
 float3 SceneLightsTranslatedBoundMin;
