Fix divide by zero introducing NaN in BRDF calculation.
Was seen on point lights with SourceLength, where the light axis is along a surface.

diff --git a/Engine/Shaders/Private/ShadingModels.ush b/Engine/Shaders/Private/ShadingModels.ush
--- a/Engine/Shaders/Private/ShadingModels.ush
+++ b/Engine/Shaders/Private/ShadingModels.ush
@@ -186,12 +186,12 @@ float EnergyNormalization( inout float a2, float VoH, FAreaLight AreaLight )
 		float LineCosTwoAlpha = AreaLight.LineCosSubtended;
 		float LineTanAlpha = sqrt( ( 1.0001 - LineCosTwoAlpha ) / ( 1 + LineCosTwoAlpha ) );
 		float Line_a2 = New_a2( Sphere_a2, LineTanAlpha, VoH );
-		Energy *= sqrt( Sphere_a2 / Line_a2 );
+		Energy *= sqrt( Sphere_a2 / max(Line_a2, 1e-5) );
 #else
 		float LineCosTwoAlpha = AreaLight.LineCosSubtended;
 		float LineSinAlpha = sqrt( 0.5 - 0.5 * LineCosTwoAlpha );
 		float Line_a2 = New_a2( Sphere_a2, LineSinAlpha, VoH );
-		Energy *= Sphere_a2 / Line_a2;
+		Energy *= Sphere_a2 / max(Line_a2, 1e-5);
 #endif
 	}
 
