[Light Grid] Support building light grid using one thread group per cell and using two levels (coarse + fine grids).
- Default settings:
    - Use two levels if number of local lights in view > 128 (r.LightCulling.TwoLevel and r.LightCulling.TwoLevel.Threshold).
    - One thread group per cell for coarse grid.
    - Single thread per cell for fine grid (r.LightCulling.WorkloadDistributionMode).

View with 1K local lights running on console: 0.25ms -> 0.1ms

(based on Ola.Olsson previous work)

diff --git a/Engine/Shaders/Private/LightGridInjection.usf b/Engine/Shaders/Private/LightGridInjection.usf
--- a/Engine/Shaders/Private/LightGridInjection.usf
+++ b/Engine/Shaders/Private/LightGridInjection.usf
@@ -25,6 +25,7 @@ RWStructuredBuffer<uint> RWNumCulledLightsGrid;
 	#define RWCulledLightDataGrid	RWCulledLightDataGrid32Bit
 #endif
 
+RWStructuredBuffer<uint> RWNextCulledLightData;
 RWStructuredBuffer<uint> RWNextCulledLightLink;
 RWStructuredBuffer<uint> RWCulledLightLinks;
 
@@ -32,8 +33,10 @@ uint NumReflectionCaptures;
 uint NumLocalLights;
 uint NumGridCells;
 uint MaxCulledLightsPerCell;
+uint NumAvailableLinks;
 uint3 CulledGridSize;
 float3 LightGridZParams;
+uint LightGridZSliceScale;
 uint LightGridPixelSizeShift;
 
 uint MegaLightsSupportedStartIndex;
@@ -42,6 +45,20 @@ StructuredBuffer<float4> LightViewSpacePositionAndRadius;
 StructuredBuffer<float4> LightViewSpaceDirAndPreprocAngle;
 StructuredBuffer<float4> LightViewSpaceRectPlanes;
 
+#if USE_PARENT_LIGHT_GRID
+	StructuredBuffer<uint> ParentNumCulledLightsGrid;
+#	if LIGHT_GRID_USES_16BIT_BUFFERS
+		Buffer<uint>						ParentCulledLightDataGrid16Bit;
+#		define ParentCulledLightDataGrid	ParentCulledLightDataGrid16Bit
+#	else
+		StructuredBuffer<uint>				ParentCulledLightDataGrid32Bit;
+#		define ParentCulledLightDataGrid	ParentCulledLightDataGrid32Bit
+#	endif
+	uint3 ParentGridSize;
+	uint NumParentGridCells;
+	uint ParentGridSizeFactor;
+#endif // USE_PARENT_LIGHT_GRID
+
 #define NUM_PLANES_PER_RECT_LIGHT 4
 
 #define REFINE_SPOTLIGHT_BOUNDS 1
@@ -50,13 +67,11 @@ StructuredBuffer<float4> LightViewSpaceRectPlanes;
 #define REFINE_RECTLIGHT_BOUNDS 1
 #endif
 
-#ifndef USE_LOCALSTORE_FOR_ATOMICS
-#define USE_LOCALSTORE_FOR_ATOMICS 0
-#endif
+#ifdef LightGridInjectionCS
 
-float ComputeCellNearViewDepthFromZSlice(uint ZSlice)
+float ComputeCellNearViewDepthFromZSlice(uint ZSlice, uint ZSliceScale)
 {
-	float SliceDepth = (exp2(ZSlice / LightGridZParams.z) - LightGridZParams.y) / LightGridZParams.x;
+	float SliceDepth = (exp2(ZSlice * ZSliceScale / LightGridZParams.z) - LightGridZParams.y) / LightGridZParams.x;
 
 	if (ZSlice == (uint)CulledGridSize.z)
 	{
@@ -215,6 +230,8 @@ bool IsAabbOutsideInfiniteAcuteConeApprox(float3 ConeVertex, float3 ConeAxis, fl
 
 bool OverlapsLight(uint LocalLightIndex, FCellBounds CellBounds, out bool bIsRectLight, out bool bIsTexturedLight)
 {
+	checkSlow(LocalLightIndex < NumLocalLights);
+
 	float4 LightPositionAndRadius = LightViewSpacePositionAndRadius[LocalLightIndex];
 	float3 ViewSpaceLightPosition = LightPositionAndRadius.xyz;
 	float LightRadius = LightPositionAndRadius.w;
@@ -273,6 +290,8 @@ bool OverlapsLight(uint LocalLightIndex, FCellBounds CellBounds, out bool bIsRec
 
 bool OverlapsReflectionCapture(uint ReflectionCaptureIndex, FCellBounds CellBounds)
 {
+	checkSlow(ReflectionCaptureIndex < NumReflectionCaptures);
+
 	FDFVector3 CaptureWorldPosition = MakeDFVector3(GetReflectionPositionAndRadius(ReflectionCaptureIndex).xyz, GetReflectionPositionLow(ReflectionCaptureIndex).xyz);
 	float3 CaptureTranslatedWorldPosition = DFFastAddDemote(CaptureWorldPosition, PrimaryView.PreViewTranslation);
 	float3 ViewSpaceCapturePosition = mul(float4(CaptureTranslatedWorldPosition, 1), View.TranslatedWorldToView).xyz;
@@ -283,7 +302,7 @@ bool OverlapsReflectionCapture(uint ReflectionCaptureIndex, FCellBounds CellBoun
 	return (BoxDistanceSq < CaptureRadius * CaptureRadius);
 }
 
-void WriteReversedLinkedList(uint LinkOffset, uint NumPrimitives, uint CulledLightDataStart)
+void OutputReversedLinkedList(uint LinkOffset, uint NumPrimitives, uint CulledLightDataStart)
 {
 	uint CurrentIndex = 0;
 
@@ -297,159 +316,491 @@ void WriteReversedLinkedList(uint LinkOffset, uint NumPrimitives, uint CulledLig
 	}
 }
 
-#if SHADER_LIGHT_GRID_INJECTION_CS
+struct FArrayView
+{
+	uint Num;
+	uint Offset;
+};
 
-RWStructuredBuffer<uint> RWNextCulledLightData;
+uint GetPrimitiveIndex(FArrayView ArrayView, uint Index)
+{
+	checkSlow(Index < ArrayView.Num);
+#if USE_PARENT_LIGHT_GRID
+	return ParentCulledLightDataGrid[ArrayView.Offset + Index];
+#else
+	return ArrayView.Offset + Index;
+#endif
+}
 
-[numthreads(THREADGROUP_SIZE, THREADGROUP_SIZE, THREADGROUP_SIZE)]
-void LightGridInjectionCS(
-	uint3 GroupId : SV_GroupID,
-	uint3 DispatchThreadId : SV_DispatchThreadID,
-	uint3 GroupThreadId : SV_GroupThreadID)
+FArrayView InitLocalLightArrayView(uint ParentGridIndex)
 {
-	uint3 GridCoordinate = DispatchThreadId;
+	FArrayView ArrayView;
 
-	if (all(GridCoordinate < (uint3)CulledGridSize))
-	{
-		uint GridIndex = (GridCoordinate.z * CulledGridSize.y + GridCoordinate.y) * CulledGridSize.x + GridCoordinate.x;
+#if USE_PARENT_LIGHT_GRID
+	uint Unused_NumVisibleMegaLights;
+	UnpackCulledLightsGridHeader0(ParentNumCulledLightsGrid[ParentGridIndex * 2 + 0], ArrayView.Num, Unused_NumVisibleMegaLights);
 
-// Disable to pass all lights through for debugging, will hit limits quickly though
-#define CULL_LIGHTS 1
-#if CULL_LIGHTS
+	bool bHasRectLights_Unused;
+	bool bHasTexturedLights_Unused;
+	UnpackCulledLightsGridHeader1(ParentNumCulledLightsGrid[ParentGridIndex * 2 + 1], ArrayView.Offset, bHasRectLights_Unused, bHasTexturedLights_Unused);
+#else
+	ArrayView.Num = NumLocalLights;
+	ArrayView.Offset = 0;
+#endif
 
-		float MinTileZ = ComputeCellNearViewDepthFromZSlice(GridCoordinate.z);
-		float MaxTileZ = ComputeCellNearViewDepthFromZSlice(GridCoordinate.z + 1);
+	return ArrayView;
+}
 
-#if USE_HZB_CULL
-		FScreenRect Rect = ComputeCellCullRect(GridCoordinate, MinTileZ, MaxTileZ);
+FArrayView InitReflectionCaptureArrayView(uint ParentGridIndex)
+{
+	FArrayView ArrayView;
 
- 		if (!IsVisibleHZB(Rect, true /*bSample4x4*/))
- 		{
- 			RWNumCulledLightsGrid[GridIndex * NUM_CULLED_LIGHTS_GRID_STRIDE + 0] = PackCulledLightsGridHeader0(0, 0);
- 			RWNumCulledLightsGrid[GridIndex * NUM_CULLED_LIGHTS_GRID_STRIDE + 1] = PackCulledLightsGridHeader1(GridIndex * MaxCulledLightsPerCell, false, false);
- 			return;
- 		}
-		
-		// tighten cell bounds based on farthest depth in cell
-		{
-			const float FarthestDeviceZ = GetMinDepthFromHZB(Rect, true /*bSample4x4*/);
-			const float FarthestSceneDepth = ConvertFromDeviceZ(FarthestDeviceZ);
+#if USE_PARENT_LIGHT_GRID
+	ArrayView.Num = ParentNumCulledLightsGrid[(NumParentGridCells + ParentGridIndex) * NUM_CULLED_LIGHTS_GRID_STRIDE + 0];
+	ArrayView.Offset = ParentNumCulledLightsGrid[(NumParentGridCells + ParentGridIndex) * NUM_CULLED_LIGHTS_GRID_STRIDE + 1];
+#else
+	ArrayView.Num = NumReflectionCaptures;
+	ArrayView.Offset = 0;
+#endif
 
-			MaxTileZ = min(MaxTileZ, FarthestSceneDepth);
-		}
-#endif // USE_HZB_CULL
-		
-		FCellBounds CellBounds = ComputeCellBounds(GridCoordinate, MinTileZ, MaxTileZ);
-		
-		uint NumAvailableLinks = NumGridCells * MaxCulledLightsPerCell * NUM_CULLED_GRID_PRIMITIVE_TYPES;
+	return ArrayView;
+}
 
-		uint NumVisibleLights = 0;
-		uint NumVisibleMegaLights = 0;
-		uint bHasRectLights = false;
-		uint bHasTexturedLights = false;
+void CullLights_SingleThread(uint GridIndex, FCellBounds CellBounds, uint ParentGridIndex)
+{
+	uint NumVisibleLights = 0;
+	uint NumVisibleMegaLights = 0;
+	uint bHasRectLights = false;
+	uint bHasTexturedLights = false;
 
-		uint LinkOffset = 0xFFFFFFFF;
+	uint LinkOffset = 0xFFFFFFFF;
 
-		LOOP
-		for (uint LocalLightIndex = 0; LocalLightIndex < NumLocalLights; LocalLightIndex++)
+	FArrayView LocalLightArrayView = InitLocalLightArrayView(ParentGridIndex);
+
+	LOOP
+	for (uint Index = 0; Index < LocalLightArrayView.Num; ++Index)
+	{
+		uint LocalLightIndex = GetPrimitiveIndex(LocalLightArrayView, Index);
+
+		bool bIsRectLight;
+		bool bIsTexturedLight;
+		bool bOverlapsLight = OverlapsLight(LocalLightIndex, CellBounds, bIsRectLight, bIsTexturedLight);
+
+		if (bOverlapsLight)
 		{
-			bool bIsRectLight;
-			bool bIsTexturedLight;
-			if (OverlapsLight(LocalLightIndex, CellBounds, bIsRectLight, bIsTexturedLight))
-			{
-				bHasRectLights |= bIsRectLight;
-				bHasTexturedLights |= bIsTexturedLight;
+			bHasRectLights |= bIsRectLight;
+			bHasTexturedLights |= bIsTexturedLight;
 
 #if USE_LINKED_CULL_LIST
-				uint NextLink;
-				InterlockedAdd(RWNextCulledLightLink[0], 1U, NextLink);
+			uint NextLink;
+			WaveInterlockedAddScalar_(RWNextCulledLightLink[0], 1U, NextLink);
 
-				if (NextLink < NumAvailableLinks)
-				{
-					RWCulledLightLinks[NextLink * LIGHT_LINK_STRIDE + 0] = LocalLightIndex;
-					RWCulledLightLinks[NextLink * LIGHT_LINK_STRIDE + 1] = LinkOffset;
+			if (NextLink < NumAvailableLinks)
+			{
+				RWCulledLightLinks[NextLink * LIGHT_LINK_STRIDE + 0] = LocalLightIndex;
+				RWCulledLightLinks[NextLink * LIGHT_LINK_STRIDE + 1] = LinkOffset;
 
-					LinkOffset = NextLink;
+				LinkOffset = NextLink;
 
-					++NumVisibleLights;
-					if (LocalLightIndex >= MegaLightsSupportedStartIndex)
-					{
-						++NumVisibleMegaLights;
-					}
+				++NumVisibleLights;
+				if (LocalLightIndex >= MegaLightsSupportedStartIndex)
+				{
+					++NumVisibleMegaLights;
 				}
+			}
 #else // !USE_LINKED_CULL_LIST
-				if (NumVisibleLights < MaxCulledLightsPerCell)
-				{
-					RWCulledLightDataGrid[GridIndex * MaxCulledLightsPerCell + NumVisibleLights] = LocalLightIndex;
+			if (NumVisibleLights < MaxCulledLightsPerCell)
+			{
+				RWCulledLightDataGrid[GridIndex * MaxCulledLightsPerCell + NumVisibleLights] = LocalLightIndex;
 
-					++NumVisibleLights;
-					if (LocalLightIndex >= MegaLightsSupportedStartIndex)
-					{
-						++NumVisibleMegaLights;
-					}
+				++NumVisibleLights;
+				if (LocalLightIndex >= MegaLightsSupportedStartIndex)
+				{
+					++NumVisibleMegaLights;
 				}
+			}
 #endif // !USE_LINKED_CULL_LIST
+		}
+	}
+
+#if USE_LINKED_CULL_LIST
+	uint CulledLightDataStart;
+	WaveInterlockedAdd_(RWNextCulledLightData[0], NumVisibleLights, CulledLightDataStart);
+
+	// Mega Lights are order dependent
+	OutputReversedLinkedList(LinkOffset, NumVisibleLights, CulledLightDataStart);
+#else
+	uint CulledLightDataStart = GridIndex * MaxCulledLightsPerCell;
+#endif
+
+	RWNumCulledLightsGrid[GridIndex * NUM_CULLED_LIGHTS_GRID_STRIDE + 0] = PackCulledLightsGridHeader0(NumVisibleLights, NumVisibleMegaLights);
+	RWNumCulledLightsGrid[GridIndex * NUM_CULLED_LIGHTS_GRID_STRIDE + 1] = PackCulledLightsGridHeader1(CulledLightDataStart, bHasRectLights, bHasTexturedLights);
+}
+void CullReflectionCaptures_SingleThread(uint GridIndex, FCellBounds CellBounds, uint ParentGridIndex)
+{
+	uint NumVisibleCaptures = 0;
+
+	uint LinkOffset = 0xFFFFFFFF;
+
+	FArrayView ReflectionCaptureArrayView = InitReflectionCaptureArrayView(ParentGridIndex);
+
+	LOOP
+	for (uint Index = 0; Index < ReflectionCaptureArrayView.Num; ++Index)
+	{
+		uint ReflectionCaptureIndex = GetPrimitiveIndex(ReflectionCaptureArrayView, Index);
+
+		bool bOverlapsReflectionCapture = OverlapsReflectionCapture(ReflectionCaptureIndex, CellBounds);
+
+		if (bOverlapsReflectionCapture)
+		{
+#if USE_LINKED_CULL_LIST
+			uint NextLink;
+			WaveInterlockedAddScalar_(RWNextCulledLightLink[0], 1U, NextLink);
+
+			if (NextLink < NumAvailableLinks)
+			{
+				RWCulledLightLinks[NextLink * LIGHT_LINK_STRIDE + 0] = ReflectionCaptureIndex;
+				RWCulledLightLinks[NextLink * LIGHT_LINK_STRIDE + 1] = LinkOffset;
+
+				LinkOffset = NextLink;
+
+				++NumVisibleCaptures;
 			}
+#else // !USE_LINKED_CULL_LIST
+			if (NumVisibleCaptures < MaxCulledLightsPerCell)
+			{
+				RWCulledLightDataGrid[(NumGridCells + GridIndex) * MaxCulledLightsPerCell + NumVisibleCaptures] = ReflectionCaptureIndex;
+
+				++NumVisibleCaptures;
+			}
+#endif // !USE_LINKED_CULL_LIST
 		}
+	}
 
 #if USE_LINKED_CULL_LIST
-		uint CulledLightDataStart;
-		WaveInterlockedAdd_(RWNextCulledLightData[0], NumVisibleLights, CulledLightDataStart);
+	uint CulledReflectionCaptureDataStart;
+	WaveInterlockedAdd_(RWNextCulledLightData[0], NumVisibleCaptures, CulledReflectionCaptureDataStart);
 
-		// Mega Lights are order dependent
-		WriteReversedLinkedList(LinkOffset, NumVisibleLights, CulledLightDataStart);
+	// Reflection captures are order dependent
+	OutputReversedLinkedList(LinkOffset, NumVisibleCaptures, CulledReflectionCaptureDataStart);
 #else
-		uint CulledLightDataStart = GridIndex * MaxCulledLightsPerCell;
+	uint CulledReflectionCaptureDataStart = (NumGridCells + GridIndex) * MaxCulledLightsPerCell;
 #endif
 
-		RWNumCulledLightsGrid[GridIndex * NUM_CULLED_LIGHTS_GRID_STRIDE + 0] = PackCulledLightsGridHeader0(NumVisibleLights, NumVisibleMegaLights);
-		RWNumCulledLightsGrid[GridIndex * NUM_CULLED_LIGHTS_GRID_STRIDE + 1] = PackCulledLightsGridHeader1(CulledLightDataStart, bHasRectLights, bHasTexturedLights);
+	RWNumCulledLightsGrid[(NumGridCells + GridIndex) * NUM_CULLED_LIGHTS_GRID_STRIDE + 0] = NumVisibleCaptures;
+	RWNumCulledLightsGrid[(NumGridCells + GridIndex) * NUM_CULLED_LIGHTS_GRID_STRIDE + 1] = CulledReflectionCaptureDataStart;
+}
+
+#define NUM_THREADS_PER_GROUP THREADGROUP_SIZE
+#include "ThreadGroupPrefixSum.ush"
+
+// When using one thread group per cell, store up to MAX_LDS_ITEMS in LDS before falling back to linked list
+#define USE_LDS (1)
+#define MAX_LDS_ITEMS 256
+
+groupshared uint CellLDSItems[MAX_LDS_ITEMS];
+
+struct FCellWriter
+{
+	uint LinkOffset;
+};
+
+FCellWriter InitCellWriter()
+{
+	FCellWriter CellWriter;
+	CellWriter.LinkOffset = 0xFFFFF;
+	
+	return CellWriter;
+}
+
+bool AddToCell(inout FCellWriter CellWriter, uint PrimitiveIndex, uint IndexInCell)
+{
+#if USE_LDS
+	// use LDS if there's still space
+	if (IndexInCell < MAX_LDS_ITEMS)
+	{
+		CellLDSItems[IndexInCell] = PrimitiveIndex;
+		return true;
+	}
+#endif
+
+	// otherwise store in linked list
+	uint NextLink;
+	WaveInterlockedAddScalar_(RWNextCulledLightLink[0], 1U, NextLink);
+
+	if (NextLink < NumAvailableLinks)
+	{
+		checkSlow(PrimitiveIndex <= 0xFFFFFF);
+		checkSlow(IndexInCell <= 0xFFFF);
+		checkSlow(CellWriter.LinkOffset <= 0xFFFFFF);
+
+		uint IndexInCellL8 = (IndexInCell & 0xFF);
+		uint IndexInCellH8 = ((IndexInCell >> 8) & 0xFF);
+
+		RWCulledLightLinks[NextLink * LIGHT_LINK_STRIDE + 0] = (IndexInCellL8 << 24) | (PrimitiveIndex & 0xFFFFFF);
+		RWCulledLightLinks[NextLink * LIGHT_LINK_STRIDE + 1] = (IndexInCellH8 << 24) | (CellWriter.LinkOffset & 0xFFFFFF);
+
+		CellWriter.LinkOffset = NextLink;
+
+		return true;
+	}
+
+	return false;
+}
+
+void OutputCell(FCellWriter CellWriter, uint NumPrimitives, uint CulledLightDataStart, uint ThreadIndex)
+{
+#if USE_LDS
+	// output entries in LDS to RWCulledLightDataGrid
+	LOOP
+	for (uint OutIndex = ThreadIndex; OutIndex < NumPrimitives; OutIndex += THREADGROUP_SIZE)
+	{
+		RWCulledLightDataGrid[CulledLightDataStart + OutIndex] = CellLDSItems[OutIndex];
+	}
+#endif
+
+	// output entries in linked list to RWCulledLightDataGrid
+
+	while (CellWriter.LinkOffset != 0xFFFFF)
+	{
+		const uint PackedData0 = RWCulledLightLinks[CellWriter.LinkOffset * LIGHT_LINK_STRIDE + 0];
+		const uint PackedData1 = RWCulledLightLinks[CellWriter.LinkOffset * LIGHT_LINK_STRIDE + 1];
+		const uint PrimitiveIndex = (PackedData0 & 0xFFFFFF);
+		const uint LinkOffset = (PackedData1 & 0xFFFFFF);
+		const uint IndexInCellL8 = (PackedData0 >> 24) & 0xFF;
+		const uint IndexInCellH8 = (PackedData1 >> 24) & 0xFF;
+		const uint IndexInCell = (IndexInCellH8 << 8) | IndexInCellL8;
+
+		RWCulledLightDataGrid[CulledLightDataStart + IndexInCell] = PrimitiveIndex;
+		CellWriter.LinkOffset = LinkOffset;
+	}
+}
+
+groupshared uint OutOffset;
+groupshared uint bCellHasRectLights;
+groupshared uint bCellHasTexturedLights;
+groupshared uint CellNumVisibleLights;
+groupshared uint CellNumVisibleMegaLights;
+groupshared uint CellNumVisibleReflectionCaptures;
+
+groupshared uint CulledLightDataStart;
+
+void CullLights_ThreadGroup(uint ThreadIndex, uint GridIndex, FCellBounds CellBounds, uint ParentGridIndex)
+{
+	if (ThreadIndex == 0)
+	{
+		OutOffset = 0;
+		bCellHasRectLights = 0;
+		bCellHasTexturedLights = 0;
+		CellNumVisibleLights = 0;
+		CellNumVisibleMegaLights = 0;
+	}
+	
+	GroupMemoryBarrierWithGroupSync();
+
+	uint NumVisibleLights = 0;
+	uint NumVisibleMegaLights = 0;
+	uint bHasRectLights = false;
+	uint bHasTexturedLights = false;
+
+	FCellWriter CellWriter = InitCellWriter();
+
+	FArrayView LocalLightArrayView = InitLocalLightArrayView(ParentGridIndex);
+
+	LOOP
+	for (uint GroupBaseIndex = 0; GroupBaseIndex < LocalLightArrayView.Num; GroupBaseIndex += THREADGROUP_SIZE)
+	{
+		const uint Index = GroupBaseIndex + ThreadIndex;
+		const uint LocalLightIndex = Index < LocalLightArrayView.Num ? GetPrimitiveIndex(LocalLightArrayView, Index) : 0xFFFFFFFF;
+
+		bool bIsRectLight;
+		bool bIsTexturedLight;
+		bool bOverlapsLight = LocalLightIndex < NumLocalLights && OverlapsLight(LocalLightIndex, CellBounds, bIsRectLight, bIsTexturedLight);
 		
-		uint NumVisibleCaptures = 0;
-		uint CaptureLinkOffset = 0xFFFFFFFF;
+		// NOTE: Cannot be under any divergent branching!
+		uint GroupCount;
+		uint Offset = ThreadGroupPrefixSum(bOverlapsLight ? 1u : 0u, ThreadIndex, GroupCount);
 
-		LOOP
-		for (uint ReflectionCaptureIndex = 0; ReflectionCaptureIndex < NumReflectionCaptures; ReflectionCaptureIndex++)
+		uint IndexInCell = OutOffset + Offset;
+
+		// Wait until all threads are done with 'OutOffset'
+		GroupMemoryBarrierWithGroupSync();
+		
+		if (ThreadIndex == 0)
 		{
-			if (OverlapsReflectionCapture(ReflectionCaptureIndex, CellBounds))
-			{			
-				#if USE_LINKED_CULL_LIST
-					uint NextLink;
-					InterlockedAdd(RWNextCulledLightLink[0], 1U, NextLink);
+			OutOffset += GroupCount;
+		}
 
-					if (NextLink < NumAvailableLinks)
-					{
-						RWCulledLightLinks[NextLink * LIGHT_LINK_STRIDE + 0] = ReflectionCaptureIndex;
-						RWCulledLightLinks[NextLink * LIGHT_LINK_STRIDE + 1] = CaptureLinkOffset;
+		if (bOverlapsLight)
+		{
+			bHasRectLights |= bIsRectLight;
+			bHasTexturedLights |= bIsTexturedLight;
 
-						CaptureLinkOffset = NextLink;
+			if (AddToCell(CellWriter, LocalLightIndex, IndexInCell))
+			{
+				++NumVisibleLights;
+				if (LocalLightIndex >= MegaLightsSupportedStartIndex)
+				{
+					++NumVisibleMegaLights;
+				}
+			}
+		}
+	}
 
-						++NumVisibleCaptures;
-					}
+	InterlockedOr(bCellHasRectLights, bHasRectLights ? 1 : 0);
+	InterlockedOr(bCellHasTexturedLights, bHasTexturedLights ? 1 : 0);
+	InterlockedAdd(CellNumVisibleLights, NumVisibleLights);
+	InterlockedAdd(CellNumVisibleMegaLights, NumVisibleMegaLights);
 
-				#else // !USE_LINKED_CULL_LIST
-					if (NumVisibleCaptures < MaxCulledLightsPerCell)
-					{
-						RWCulledLightDataGrid[(NumGridCells + GridIndex) * MaxCulledLightsPerCell + NumVisibleCaptures] = ReflectionCaptureIndex;
+	GroupMemoryBarrierWithGroupSync();
 
-						++NumVisibleCaptures;
-					}
-				#endif // !USE_LINKED_CULL_LIST
+	if (ThreadIndex == 0)
+	{
+		InterlockedAdd(RWNextCulledLightData[0], CellNumVisibleLights, CulledLightDataStart);
+	}
+
+	GroupMemoryBarrierWithGroupSync();
+
+	OutputCell(CellWriter, CellNumVisibleLights, CulledLightDataStart, ThreadIndex);
+
+	if (ThreadIndex == 0)
+	{
+		RWNumCulledLightsGrid[GridIndex * NUM_CULLED_LIGHTS_GRID_STRIDE + 0] = PackCulledLightsGridHeader0(CellNumVisibleLights, CellNumVisibleMegaLights);
+		RWNumCulledLightsGrid[GridIndex * NUM_CULLED_LIGHTS_GRID_STRIDE + 1] = PackCulledLightsGridHeader1(CulledLightDataStart, bCellHasRectLights, bCellHasTexturedLights);
+	}
+}
+
+void CullReflectionCaptures_ThreadGroup(uint ThreadIndex, uint GridIndex, FCellBounds CellBounds, uint ParentGridIndex)
+{
+	if (ThreadIndex == 0)
+	{
+		OutOffset = 0;
+		CellNumVisibleReflectionCaptures = 0;
+	}
+	
+	GroupMemoryBarrierWithGroupSync();
+
+	uint NumVisibleCaptures = 0;
+
+	FCellWriter CellWriter = InitCellWriter();
+
+	FArrayView ReflectionCaptureArrayView = InitReflectionCaptureArrayView(ParentGridIndex);
+
+	LOOP
+	for (uint GroupBaseIndex = 0; GroupBaseIndex < ReflectionCaptureArrayView.Num; GroupBaseIndex += THREADGROUP_SIZE)
+	{
+		const uint Index = GroupBaseIndex + ThreadIndex;
+		const uint ReflectionCaptureIndex = Index < ReflectionCaptureArrayView.Num ? GetPrimitiveIndex(ReflectionCaptureArrayView, Index) : 0xFFFFFFFF;
+
+		bool bOverlapsReflectionCapture = ReflectionCaptureIndex < NumReflectionCaptures && OverlapsReflectionCapture(ReflectionCaptureIndex, CellBounds);
+		
+		// NOTE: Cannot be under any divergent branching!
+		uint GroupCount;
+		uint Offset = ThreadGroupPrefixSum(bOverlapsReflectionCapture ? 1u : 0u, ThreadIndex, GroupCount);
+		
+		uint IndexInCell = OutOffset + Offset;
+
+		// Wait until all threads are done with 'OutOffset'
+		GroupMemoryBarrierWithGroupSync();
+		
+		if (ThreadIndex == 0)
+		{
+			OutOffset += GroupCount;
+		}
+
+		if (bOverlapsReflectionCapture)
+		{
+			if (AddToCell(CellWriter, ReflectionCaptureIndex, IndexInCell))
+			{
+				++NumVisibleCaptures;
 			}
 		}
+	}
 
-#if USE_LINKED_CULL_LIST
-		uint CulledReflectionCaptureDataStart;
-		WaveInterlockedAdd_(RWNextCulledLightData[0], NumVisibleCaptures, CulledReflectionCaptureDataStart);
+	InterlockedAdd(CellNumVisibleReflectionCaptures, NumVisibleCaptures);
+
+	GroupMemoryBarrierWithGroupSync();
+
+	if (ThreadIndex == 0)
+	{
+		InterlockedAdd(RWNextCulledLightData[0], CellNumVisibleReflectionCaptures, CulledLightDataStart);
+	}
+
+	GroupMemoryBarrierWithGroupSync();
+
+	OutputCell(CellWriter, CellNumVisibleReflectionCaptures, CulledLightDataStart, ThreadIndex);
 
-		// Reflection captures are order dependent
-		WriteReversedLinkedList(CaptureLinkOffset, NumVisibleCaptures, CulledReflectionCaptureDataStart);
+	if (ThreadIndex == 0)
+	{
+		RWNumCulledLightsGrid[(NumGridCells + GridIndex) * NUM_CULLED_LIGHTS_GRID_STRIDE + 0] = CellNumVisibleReflectionCaptures;
+		RWNumCulledLightsGrid[(NumGridCells + GridIndex) * NUM_CULLED_LIGHTS_GRID_STRIDE + 1] = CulledLightDataStart;
+	}
+}
+
+[numthreads(THREADGROUP_SIZE_X, THREADGROUP_SIZE_Y, THREADGROUP_SIZE_Z)]
+void LightGridInjectionCS(
+	uint3 DispatchThreadId : SV_DispatchThreadID,
+	uint3 GroupId : SV_GroupID,
+	uint GroupIndex : SV_GroupIndex)
+{
+#if USE_THREAD_GROUP_PER_CELL
+	const uint3 GridCoordinate = GroupId;
+	const uint ThreadIndex = GroupIndex;
 #else
-		uint CulledReflectionCaptureDataStart = (NumGridCells + GridIndex) * MaxCulledLightsPerCell;
+	const uint3 GridCoordinate = DispatchThreadId;
 #endif
 
-		RWNumCulledLightsGrid[(NumGridCells + GridIndex) * NUM_CULLED_LIGHTS_GRID_STRIDE + 0] = NumVisibleCaptures;
-		RWNumCulledLightsGrid[(NumGridCells + GridIndex) * NUM_CULLED_LIGHTS_GRID_STRIDE + 1] = CulledReflectionCaptureDataStart;
+	if (all(GridCoordinate < (uint3)CulledGridSize))
+	{
+		const uint GridIndex = (GridCoordinate.z * CulledGridSize.y + GridCoordinate.y) * CulledGridSize.x + GridCoordinate.x;
+
+// Disable to pass all lights through for debugging, will hit limits quickly though
+#define CULL_LIGHTS 1
+#if CULL_LIGHTS
+
+		float MinTileZ = ComputeCellNearViewDepthFromZSlice(GridCoordinate.z + 0, LightGridZSliceScale);
+		float MaxTileZ = ComputeCellNearViewDepthFromZSlice(GridCoordinate.z + 1, LightGridZSliceScale);
+
+#if USE_HZB_CULL
+		FScreenRect Rect = ComputeCellCullRect(GridCoordinate, MinTileZ, MaxTileZ);
+
+ 		if (!IsVisibleHZB(Rect, true /*bSample4x4*/))
+ 		{
+#if USE_THREAD_GROUP_PER_CELL
+			if (ThreadIndex == 0)
+#endif
+			{
+ 				RWNumCulledLightsGrid[GridIndex * NUM_CULLED_LIGHTS_GRID_STRIDE + 0] = PackCulledLightsGridHeader0(0, 0);
+ 				RWNumCulledLightsGrid[GridIndex * NUM_CULLED_LIGHTS_GRID_STRIDE + 1] = PackCulledLightsGridHeader1(0, false, false);
+			}
+ 			return;
+ 		}
+		
+		// tighten cell bounds based on farthest depth in cell
+		{
+			const float FarthestDeviceZ = GetMinDepthFromHZB(Rect, true /*bSample4x4*/);
+			const float FarthestSceneDepth = ConvertFromDeviceZ(FarthestDeviceZ);
+
+			MaxTileZ = min(MaxTileZ, FarthestSceneDepth);
+		}
+#endif // USE_HZB_CULL
+		
+		FCellBounds CellBounds = ComputeCellBounds(GridCoordinate, MinTileZ, MaxTileZ);
+
+#if USE_PARENT_LIGHT_GRID
+		const uint3 ParentGridCoordinate = GridCoordinate / ParentGridSizeFactor;
+		const uint ParentGridIndex = (ParentGridCoordinate.z * ParentGridSize.y + ParentGridCoordinate.y) * ParentGridSize.x + ParentGridCoordinate.x;
+#else
+		const uint ParentGridIndex = 0xFFFFFFFF;
+#endif
+
+#if USE_THREAD_GROUP_PER_CELL
+		CullLights_ThreadGroup(ThreadIndex, GridIndex, CellBounds, ParentGridIndex);
+		CullReflectionCaptures_ThreadGroup(ThreadIndex, GridIndex, CellBounds, ParentGridIndex);
+#else
+		CullLights_SingleThread(GridIndex, CellBounds, ParentGridIndex);
+		CullReflectionCaptures_SingleThread(GridIndex, CellBounds, ParentGridIndex);
+#endif
 
 #else // !CULL_LIGHTS
 		LOOP
@@ -461,7 +812,9 @@ void LightGridInjectionCS(
 			}
 		}
 
-		RWNumCulledLightsGrid[GridIndex * NUM_CULLED_LIGHTS_GRID_STRIDE + 0] = PackCulledLightsGridHeader0(NumLocalLights, NumLocalLights - MegaLightsSupportedStartIndex);
+		const uint NumVisibleLights = min(NumLocalLights, MaxCulledLightsPerCell);
+		const uint NumVisibleMegaLights = max(NumVisibleLights - MegaLightsSupportedStartIndex, 0);
+		RWNumCulledLightsGrid[GridIndex * NUM_CULLED_LIGHTS_GRID_STRIDE + 0] = PackCulledLightsGridHeader0(NumVisibleLights, NumVisibleMegaLights);
 		RWNumCulledLightsGrid[GridIndex * NUM_CULLED_LIGHTS_GRID_STRIDE + 1] = PackCulledLightsGridHeader1(GridIndex * MaxCulledLightsPerCell, true, true); // TODO
 		
 		LOOP
@@ -473,13 +826,14 @@ void LightGridInjectionCS(
 			}
 		}
 
-		RWNumCulledLightsGrid[(NumGridCells + GridIndex) * NUM_CULLED_LIGHTS_GRID_STRIDE + 0] = NumReflectionCaptures;
+		const uint NumVisibleReflectionCaptures = min(NumReflectionCaptures, MaxCulledLightsPerCell);
+		RWNumCulledLightsGrid[(NumGridCells + GridIndex) * NUM_CULLED_LIGHTS_GRID_STRIDE + 0] = NumVisibleReflectionCaptures;
 		RWNumCulledLightsGrid[(NumGridCells + GridIndex) * NUM_CULLED_LIGHTS_GRID_STRIDE + 1] = (NumGridCells + GridIndex) * MaxCulledLightsPerCell;
 #endif // !CULL_LIGHTS
 	}
 }
 
-#endif // SHADER_LIGHT_GRID_INJECTION_CS
+#endif // LightGridInjectionCS
 
 
 #ifdef SHADER_DEBUG_LIGHT_GRID_PS
diff --git a/Engine/Source/Runtime/Renderer/Private/LightGridInjection.cpp b/Engine/Source/Runtime/Renderer/Private/LightGridInjection.cpp
--- a/Engine/Source/Runtime/Renderer/Private/LightGridInjection.cpp
+++ b/Engine/Source/Runtime/Renderer/Private/LightGridInjection.cpp
@@ -113,6 +113,37 @@ FAutoConsoleVariableRef CVarLightCullingQuality(
 	ECVF_RenderThreadSafe
 );
 
+static TAutoConsoleVariable<int32> CVarLightCullingWorkloadDistributionMode(
+	TEXT("r.LightCulling.WorkloadDistributionMode"),
+	0,
+	TEXT("0 - single thread per cell.\n")
+	TEXT("1 - thread group per cell (64 threads).\n")
+	TEXT("2 - thread group per cell (32 threads if supported, otherwise single thread).\n")
+	TEXT("(This cvar only applies to fine light grid. When using two levels, coarse grid always uses thread group per cell."),
+	ECVF_RenderThreadSafe
+);
+
+static TAutoConsoleVariable<bool> CVarLightCullingTwoLevel(
+	TEXT("r.LightCulling.TwoLevel"),
+	true,
+	TEXT("Whether to build light grid in two passes."),
+	ECVF_RenderThreadSafe
+);
+
+static TAutoConsoleVariable<int32> CVarLightCullingTwoLevelThreshold(
+	TEXT("r.LightCulling.TwoLevel.Threshold"),
+	128,
+	TEXT("Threshold used to determine whether to use two level culling basedon the number of lights in view."),
+	ECVF_RenderThreadSafe
+);
+
+static TAutoConsoleVariable<int32> CVarLightCullingTwoLevelExponent(
+	TEXT("r.LightCulling.TwoLevel.Exponent"),
+	2,
+	TEXT("Exponent used to derive the coarse grid size (base 2)."),
+	ECVF_RenderThreadSafe
+);
+
 float GLightCullingMaxDistanceOverrideKilometers = -1.0f;
 FAutoConsoleVariableRef CVarLightCullingMaxDistanceOverride(
 	TEXT("r.LightCulling.MaxDistanceOverrideKilometers"),
@@ -203,18 +234,19 @@ typedef uint16 FLightIndexType;
 // UINT_MAX indexable light limit
 typedef uint32 FLightIndexType32;
 
-uint32 LightGridInjectionGroupSize = 4;
-
 
 class FLightGridInjectionCS : public FGlobalShader
 {
 	DECLARE_GLOBAL_SHADER(FLightGridInjectionCS);
 	SHADER_USE_PARAMETER_STRUCT(FLightGridInjectionCS, FGlobalShader)
 public:
-	class FUseLinkedListDim : SHADER_PERMUTATION_BOOL("USE_LINKED_CULL_LIST");
-	class FRefineRectLightBoundsDim : SHADER_PERMUTATION_BOOL("REFINE_RECTLIGHT_BOUNDS");
-	class FUseHZBCullDim : SHADER_PERMUTATION_BOOL("USE_HZB_CULL");
-	using FPermutationDomain = TShaderPermutationDomain<FUseLinkedListDim, FRefineRectLightBoundsDim, FUseHZBCullDim>;
+	class FUseLinkedList : SHADER_PERMUTATION_BOOL("USE_LINKED_CULL_LIST");
+	class FRefineRectLightBounds : SHADER_PERMUTATION_BOOL("REFINE_RECTLIGHT_BOUNDS");
+	class FUseHZBCull : SHADER_PERMUTATION_BOOL("USE_HZB_CULL");
+	class FUseParentLightGrid : SHADER_PERMUTATION_BOOL("USE_PARENT_LIGHT_GRID");
+	class FUseThreadGroupPerCell : SHADER_PERMUTATION_BOOL("USE_THREAD_GROUP_PER_CELL");
+	class FUseThreadGroupSize32 : SHADER_PERMUTATION_BOOL("USE_THREAD_GROUP_SIZE_32");
+	using FPermutationDomain = TShaderPermutationDomain<FUseLinkedList, FRefineRectLightBounds, FUseHZBCull, FUseParentLightGrid, FUseThreadGroupPerCell, FUseThreadGroupSize32>;
 
 	BEGIN_SHADER_PARAMETER_STRUCT(FParameters, )
 		SHADER_PARAMETER_STRUCT_REF(FReflectionCaptureShaderData, ReflectionCapture)
@@ -236,9 +268,19 @@ public:
 		SHADER_PARAMETER(uint32, NumLocalLights)
 		SHADER_PARAMETER(uint32, NumGridCells)
 		SHADER_PARAMETER(uint32, MaxCulledLightsPerCell)
+		SHADER_PARAMETER(uint32, NumAvailableLinks)
 		SHADER_PARAMETER(uint32, LightGridPixelSizeShift)
 		SHADER_PARAMETER(uint32, MegaLightsSupportedStartIndex)
 
+		SHADER_PARAMETER(uint32, LightGridZSliceScale)
+
+		SHADER_PARAMETER_RDG_BUFFER_SRV(StructuredBuffer<uint>, ParentNumCulledLightsGrid)
+		SHADER_PARAMETER_RDG_BUFFER_SRV(StructuredBuffer<uint>, ParentCulledLightDataGrid32Bit)
+		SHADER_PARAMETER_RDG_BUFFER_SRV(Buffer<uint>, ParentCulledLightDataGrid16Bit)
+		SHADER_PARAMETER(FIntVector, ParentGridSize)
+		SHADER_PARAMETER(uint32, NumParentGridCells)
+		SHADER_PARAMETER(uint32, ParentGridSizeFactor)
+
 		SHADER_PARAMETER(FVector2f, HZBSize)
 		SHADER_PARAMETER(FVector2f, HZBViewSize)
 		SHADER_PARAMETER(FIntRect, HZBViewRect)
@@ -252,13 +294,32 @@ public:
 		return true;
 	}
 
+	static FIntVector GetGroupSize(FPermutationDomain PermutationVector)
+	{
+		if (PermutationVector.Get<FUseThreadGroupSize32>())
+		{
+			return FIntVector(4, 4, 2);
+		}
+		else
+		{
+			return FIntVector(4, 4, 4);
+		}
+	}
+
 	static void ModifyCompilationEnvironment(const FGlobalShaderPermutationParameters& Parameters, FShaderCompilerEnvironment& OutEnvironment)
 	{
 		FGlobalShader::ModifyCompilationEnvironment(Parameters, OutEnvironment);
-		OutEnvironment.SetDefine(TEXT("SHADER_LIGHT_GRID_INJECTION_CS"), 1);
-		OutEnvironment.SetDefine(TEXT("THREADGROUP_SIZE"), LightGridInjectionGroupSize);
 		FForwardLightingParameters::ModifyCompilationEnvironment(Parameters.Platform, OutEnvironment);
 		OutEnvironment.SetDefine(TEXT("LIGHT_LINK_STRIDE"), LightLinkStride);
+
+		FPermutationDomain PermutationVector(Parameters.PermutationId);
+
+		FIntVector GroupSize = GetGroupSize(PermutationVector);
+
+		OutEnvironment.SetDefine(TEXT("THREADGROUP_SIZE"), GroupSize.X * GroupSize.Y * GroupSize.Z);
+		OutEnvironment.SetDefine(TEXT("THREADGROUP_SIZE_X"), GroupSize.X);
+		OutEnvironment.SetDefine(TEXT("THREADGROUP_SIZE_Y"), GroupSize.Y);
+		OutEnvironment.SetDefine(TEXT("THREADGROUP_SIZE_Z"), GroupSize.Z);
 	}
 };
 
@@ -476,6 +537,161 @@ static void CalculateRectLightCullingPlanes(const FRectLightSceneProxy* RectProx
 	check(OutPlanes.Num() == NUM_PLANES_PER_RECT_LIGHT);
 }
 
+struct FLightGrid
+{
+	FRDGBufferSRVRef CulledLightDataGridSRV = nullptr;
+	FRDGBufferSRVRef NumCulledLightsGridSRV = nullptr;
+};
+
+FLightGrid LightGridInjection(
+	FRDGBuilder& GraphBuilder,
+	FViewInfo& View,
+	FIntVector GridSize,
+	uint32 LightGridPixelSizeShift,
+	uint32 ZSliceScale,
+	uint32 MaxNumCells,
+	FVector3f ZParams,
+	uint32 NumLocalLights,
+	uint32 NumReflectionCaptures,
+	uint32 MegaLightsSupportedStartIndex,
+	bool bUse16BitBuffers,
+	bool bRefineRectLightBounds,
+	FRDGBufferSRVRef LightViewSpacePositionAndRadiusSRV,
+	FRDGBufferSRVRef LightViewSpaceDirAndPreprocAngleSRV,
+	FRDGBufferSRVRef LightViewSpaceRectPlanesSRV,
+	bool bThreadGroupPerCell,
+	bool bThreadGroupSize32,
+	// parent params
+	FRDGBufferSRVRef ParentNumCulledLightsGridSRV,
+	FRDGBufferSRVRef ParentCulledLightDataGridSRV,
+	uint32 ParentGridSizeFactor)
+{
+	const uint32 NumCulledLightLinks = MaxNumCells * GMaxCulledLightsPerCell * NumCulledGridPrimitiveTypes;
+	
+	if (bThreadGroupPerCell)
+	{
+		checkf(NumLocalLights < 0xFFFFFF, TEXT("NumLocalLights limited to 16M by FCellWriter."));
+		checkf(NumReflectionCaptures < 0xFFFFFF, TEXT("NumReflectionCaptures limited to 16M by FCellWriter."));
+
+		checkf(NumCulledLightLinks < 0xFFFFFF, TEXT("NumCulledLightLinks limited to 16M by FCellWriter."));
+	}
+
+	const FIntVector ParentGridSize = FIntVector::DivideAndRoundUp(GridSize, ParentGridSizeFactor);
+
+	FRDGBufferRef CulledLightLinksBuffer = GraphBuilder.CreateBuffer(FRDGBufferDesc::CreateStructuredDesc(sizeof(uint32), NumCulledLightLinks * LightLinkStride), TEXT("CulledLightLinks"));
+	FRDGBufferRef NextCulledLightLinkBuffer = GraphBuilder.CreateBuffer(FRDGBufferDesc::CreateStructuredDesc(sizeof(uint32), 1), TEXT("NextCulledLightLink"));
+	FRDGBufferRef NextCulledLightDataBuffer = GraphBuilder.CreateBuffer(FRDGBufferDesc::CreateStructuredDesc(sizeof(uint32), 1), TEXT("NextCulledLightData"));
+	FRDGBufferUAVRef NextCulledLightDataUAV = GraphBuilder.CreateUAV(NextCulledLightDataBuffer);
+	FRDGBufferRef NumCulledLightsGrid = GraphBuilder.CreateBuffer(FRDGBufferDesc::CreateStructuredDesc(sizeof(uint32), MaxNumCells * NumCulledLightsGridStride), TEXT("NumCulledLightsGrid"));
+	FRDGBufferUAVRef NumCulledLightsGridUAV = GraphBuilder.CreateUAV(NumCulledLightsGrid);
+
+	FRDGBufferSRVRef CulledLightDataGridSRV;
+	FRDGBufferUAVRef CulledLightDataGridUAV;
+	if (bUse16BitBuffers)
+	{
+		const SIZE_T LightIndexTypeSize = sizeof(FLightIndexType);
+		const EPixelFormat CulledLightDataGridFormat = PF_R16_UINT;
+		FRDGBufferRef CulledLightDataGrid = GraphBuilder.CreateBuffer(FRDGBufferDesc::CreateBufferDesc(LightIndexTypeSize, MaxNumCells * GMaxCulledLightsPerCell), TEXT("CulledLightDataGrid"));
+		CulledLightDataGridSRV = GraphBuilder.CreateSRV(CulledLightDataGrid, CulledLightDataGridFormat);
+		CulledLightDataGridUAV = GraphBuilder.CreateUAV(CulledLightDataGrid, CulledLightDataGridFormat);
+	}
+	else
+	{
+		const SIZE_T LightIndexTypeSize = sizeof(FLightIndexType32);
+		const EPixelFormat CulledLightDataGridFormat = PF_R32_UINT;
+		FRDGBufferRef CulledLightDataGrid = GraphBuilder.CreateBuffer(FRDGBufferDesc::CreateStructuredDesc(LightIndexTypeSize, MaxNumCells * GMaxCulledLightsPerCell), TEXT("CulledLightDataGrid"));
+		CulledLightDataGridSRV = GraphBuilder.CreateSRV(CulledLightDataGrid);
+		CulledLightDataGridUAV = GraphBuilder.CreateUAV(CulledLightDataGrid);
+	}
+
+	FLightGridInjectionCS::FParameters* PassParameters = GraphBuilder.AllocParameters<FLightGridInjectionCS::FParameters>();
+
+	PassParameters->View = View.ViewUniformBuffer;
+
+	if (IsMobilePlatform(View.GetShaderPlatform()))
+	{
+		PassParameters->MobileReflectionCaptureData = View.MobileReflectionCaptureUniformBuffer;
+	}
+	else
+	{
+		PassParameters->ReflectionCapture = View.ReflectionCaptureUniformBuffer;
+	}
+
+	PassParameters->RWNumCulledLightsGrid = NumCulledLightsGridUAV;
+	PassParameters->RWCulledLightDataGrid32Bit = CulledLightDataGridUAV;
+	PassParameters->RWCulledLightDataGrid16Bit = CulledLightDataGridUAV;
+	PassParameters->RWNextCulledLightLink = GraphBuilder.CreateUAV(NextCulledLightLinkBuffer);
+	PassParameters->RWNextCulledLightData = NextCulledLightDataUAV;
+	PassParameters->RWCulledLightLinks = GraphBuilder.CreateUAV(CulledLightLinksBuffer);
+	PassParameters->CulledGridSize = GridSize;
+	PassParameters->LightGridZParams = ZParams;
+	PassParameters->NumReflectionCaptures = NumReflectionCaptures;
+	PassParameters->NumLocalLights = NumLocalLights;
+	PassParameters->MaxCulledLightsPerCell = GMaxCulledLightsPerCell;
+	PassParameters->NumAvailableLinks = NumCulledLightLinks;
+	PassParameters->NumGridCells = GridSize.X * GridSize.Y * GridSize.Z;
+	PassParameters->LightGridPixelSizeShift = LightGridPixelSizeShift;
+	PassParameters->LightGridZSliceScale = ZSliceScale;
+	PassParameters->MegaLightsSupportedStartIndex = MegaLightsSupportedStartIndex;
+
+	PassParameters->ParentNumCulledLightsGrid = ParentNumCulledLightsGridSRV;
+	PassParameters->ParentCulledLightDataGrid32Bit = ParentCulledLightDataGridSRV;
+	PassParameters->ParentCulledLightDataGrid16Bit = ParentCulledLightDataGridSRV;
+	PassParameters->ParentGridSize = ParentGridSize;
+	PassParameters->NumParentGridCells = ParentGridSize.X * ParentGridSize.Y * ParentGridSize.Z;
+	PassParameters->ParentGridSizeFactor = ParentGridSizeFactor;
+
+	PassParameters->LightViewSpacePositionAndRadius = LightViewSpacePositionAndRadiusSRV;
+	PassParameters->LightViewSpaceDirAndPreprocAngle = LightViewSpaceDirAndPreprocAngleSRV;
+	PassParameters->LightViewSpaceRectPlanes = LightViewSpaceRectPlanesSRV;
+
+	{
+		PassParameters->HZBTexture = View.HZB;
+		PassParameters->HZBSampler = TStaticSamplerState< SF_Point, AM_Clamp, AM_Clamp, AM_Clamp >::GetRHI();
+		PassParameters->HZBSize = FVector2f(View.HZBMipmap0Size);
+		PassParameters->HZBViewSize = FVector2f(View.ViewRect.Size());
+		PassParameters->HZBViewRect = FIntRect(0, 0, View.ViewRect.Width(), View.ViewRect.Height());
+	}
+
+	FLightGridInjectionCS::FPermutationDomain PermutationVector;
+	PermutationVector.Set<FLightGridInjectionCS::FUseLinkedList>(GLightLinkedListCulling != 0);
+	PermutationVector.Set<FLightGridInjectionCS::FRefineRectLightBounds>(bRefineRectLightBounds);
+	PermutationVector.Set<FLightGridInjectionCS::FUseHZBCull>(GLightGridHZBCull != 0 && View.HZB != nullptr);
+	PermutationVector.Set<FLightGridInjectionCS::FUseParentLightGrid>(ParentNumCulledLightsGridSRV != nullptr && ParentCulledLightDataGridSRV != nullptr);
+	PermutationVector.Set<FLightGridInjectionCS::FUseThreadGroupPerCell>(bThreadGroupPerCell);
+	PermutationVector.Set<FLightGridInjectionCS::FUseThreadGroupSize32>(bThreadGroupSize32);
+	auto ComputeShader = View.ShaderMap->GetShader<FLightGridInjectionCS>(PermutationVector);
+
+	AddClearUAVPass(GraphBuilder, PassParameters->RWNextCulledLightLink, 0);
+	AddClearUAVPass(GraphBuilder, NextCulledLightDataUAV, 0);
+	AddClearUAVPass(GraphBuilder, NumCulledLightsGridUAV, 0);
+
+	FIntVector NumGroups;
+	if (bThreadGroupPerCell)
+	{
+		NumGroups = GridSize;
+	}
+	else
+	{
+		NumGroups = FComputeShaderUtils::GetGroupCount(GridSize, FLightGridInjectionCS::GetGroupSize(PermutationVector));
+	}
+
+	FComputeShaderUtils::AddPass(
+		GraphBuilder,
+		RDG_EVENT_NAME("LightGridInject %s %s",
+			PermutationVector.Get<FLightGridInjectionCS::FUseLinkedList>() ? TEXT("LinkedList") : TEXT("NoLinkedList"),
+			PermutationVector.Get<FLightGridInjectionCS::FUseThreadGroupPerCell>() ? TEXT("ThreadGroup") : TEXT("SingleThread")),
+		ComputeShader,
+		PassParameters,
+		NumGroups);
+
+	FLightGrid Output;
+	Output.CulledLightDataGridSRV = CulledLightDataGridSRV;
+	Output.NumCulledLightsGridSRV = GraphBuilder.CreateSRV(FRDGBufferSRVDesc(NumCulledLightsGrid));
+
+	return Output;
+}
+
 FComputeLightGridOutput FSceneRenderer::ComputeLightGrid(FRDGBuilder& GraphBuilder, bool bCullLightsToGrid, const FSortedLightSetSceneInfo& SortedLightSet)
 {
 	FComputeLightGridOutput Result = {};
@@ -867,121 +1083,101 @@ FComputeLightGridOutput FSceneRenderer::ComputeLightGrid(FRDGBuilder& GraphBuild
 			}
 		}
 
+		check(ViewSpacePosAndRadiusData.Num() == ForwardLocalLightData.Num());
+		check(ViewSpaceDirAndPreprocAngleData.Num() == ForwardLocalLightData.Num());
+		check(ViewSpaceRectPlanesData.Num() == ForwardLocalLightData.Num() * NUM_PLANES_PER_RECT_LIGHT);
+
+		FRDGBufferRef LightViewSpacePositionAndRadius = CreateStructuredBuffer(GraphBuilder, TEXT("ViewSpacePosAndRadiusData"), TConstArrayView<FVector4f>(ViewSpacePosAndRadiusData));
+		FRDGBufferRef LightViewSpaceDirAndPreprocAngle = CreateStructuredBuffer(GraphBuilder, TEXT("ViewSpaceDirAndPreprocAngleData"), TConstArrayView<FVector4f>(ViewSpaceDirAndPreprocAngleData));
+		FRDGBufferRef LightViewSpaceRectPlanes = CreateStructuredBuffer(GraphBuilder, TEXT("ViewSpaceRectPlanesData"), TConstArrayView<FVector4f>(ViewSpaceRectPlanesData));
+
+		FRDGBufferSRVRef LightViewSpacePositionAndRadiusSRV = GraphBuilder.CreateSRV(FRDGBufferSRVDesc(LightViewSpacePositionAndRadius));
+		FRDGBufferSRVRef LightViewSpaceDirAndPreprocAngleSRV = GraphBuilder.CreateSRV(FRDGBufferSRVDesc(LightViewSpaceDirAndPreprocAngle));
+		FRDGBufferSRVRef LightViewSpaceRectPlanesSRV = GraphBuilder.CreateSRV(FRDGBufferSRVDesc(LightViewSpaceRectPlanes));
+
 		// Allocate buffers using the scene render targets size so we won't reallocate every frame with dynamic resolution
 		const FIntPoint MaxLightGridSizeXY = FIntPoint::DivideAndRoundUp(View.GetSceneTexturesConfig().Extent, GLightGridPixelSize);
 
 		const int32 MaxNumCells = MaxLightGridSizeXY.X * MaxLightGridSizeXY.Y * GLightGridSizeZ * NumCulledGridPrimitiveTypes;
 
-		const FIntVector NumGroups = FIntVector::DivideAndRoundUp(FIntVector(LightGridSizeXY.X, LightGridSizeXY.Y, GLightGridSizeZ), LightGridInjectionGroupSize);
+		RDG_EVENT_SCOPE(GraphBuilder, "CullLights %ux%ux%u NumLights %u NumCaptures %u",
+			ForwardLightData->CulledGridSize.X,
+			ForwardLightData->CulledGridSize.Y,
+			ForwardLightData->CulledGridSize.Z,
+			ForwardLightData->NumLocalLights,
+			ForwardLightData->NumReflectionCaptures);
+
+		FLightGrid ParentLightGrid;
+		uint32 ParentLightGridFactor = 1;
 
+		if(CVarLightCullingTwoLevel.GetValueOnRenderThread() && (int32)ForwardLightData->NumLocalLights > CVarLightCullingTwoLevelThreshold.GetValueOnRenderThread())
 		{
-			RDG_EVENT_SCOPE(GraphBuilder, "CullLights %ux%ux%u NumLights %u NumCaptures %u",
-				ForwardLightData->CulledGridSize.X,
-				ForwardLightData->CulledGridSize.Y,
-				ForwardLightData->CulledGridSize.Z,
+			ParentLightGridFactor = (uint32)FMath::Pow(2.0f, FMath::Clamp(CVarLightCullingTwoLevelExponent.GetValueOnRenderThread(), 1, 4));
+
+			FIntVector ParentLightGridSize = FIntVector::DivideAndRoundUp(ForwardLightData->CulledGridSize, ParentLightGridFactor);
+
+			ParentLightGrid = LightGridInjection(
+				GraphBuilder,
+				View,
+				ParentLightGridSize,
+				FMath::FloorLog2(GLightGridPixelSize * ParentLightGridFactor),
+				ParentLightGridFactor,
+				MaxNumCells, // TODO: could potentially be reduced on coarse grid
+				ForwardLightData->LightGridZParams,
 				ForwardLightData->NumLocalLights,
-				ForwardLightData->NumReflectionCaptures);
-
-			const uint32 CulledLightLinksElements = MaxNumCells * GMaxCulledLightsPerCell * LightLinkStride;
-
-
-			FRDGBufferRef CulledLightLinksBuffer = GraphBuilder.CreateBuffer(FRDGBufferDesc::CreateStructuredDesc(sizeof(uint32), CulledLightLinksElements), TEXT("CulledLightLinks"));
-			FRDGBufferRef NextCulledLightLinkBuffer = GraphBuilder.CreateBuffer(FRDGBufferDesc::CreateStructuredDesc(sizeof(uint32), 1), TEXT("NextCulledLightLink"));
-			FRDGBufferRef NextCulledLightDataBuffer = GraphBuilder.CreateBuffer(FRDGBufferDesc::CreateStructuredDesc(sizeof(uint32), 1), TEXT("NextCulledLightData"));
-			FRDGBufferUAVRef NextCulledLightDataUAV = GraphBuilder.CreateUAV(NextCulledLightDataBuffer);
-			FRDGBufferSRVRef CulledLightDataGridSRV = nullptr;
-			FRDGBufferUAVRef CulledLightDataGridUAV = nullptr;
-			FRDGBufferRef NumCulledLightsGrid       = GraphBuilder.CreateBuffer(FRDGBufferDesc::CreateStructuredDesc(sizeof(uint32), MaxNumCells * NumCulledLightsGridStride), TEXT("NumCulledLightsGrid"));
-			FRDGBufferUAVRef NumCulledLightsGridUAV = GraphBuilder.CreateUAV(NumCulledLightsGrid);
-
-			if (bLightGridUses16BitBuffers)
-			{
-				const SIZE_T LightIndexTypeSize = sizeof(FLightIndexType);
-				const EPixelFormat CulledLightDataGridFormat = PF_R16_UINT;
-				FRDGBufferRef CulledLightDataGrid = GraphBuilder.CreateBuffer(FRDGBufferDesc::CreateBufferDesc(LightIndexTypeSize, MaxNumCells * GMaxCulledLightsPerCell), TEXT("CulledLightDataGrid"));
-				CulledLightDataGridSRV = GraphBuilder.CreateSRV(CulledLightDataGrid, CulledLightDataGridFormat);
-				CulledLightDataGridUAV = GraphBuilder.CreateUAV(CulledLightDataGrid, CulledLightDataGridFormat);
-			}
-			else
-			{
-				const SIZE_T LightIndexTypeSize = sizeof(FLightIndexType32);
-				const EPixelFormat CulledLightDataGridFormat = PF_R32_UINT;
-				FRDGBufferRef CulledLightDataGrid = GraphBuilder.CreateBuffer(FRDGBufferDesc::CreateStructuredDesc(LightIndexTypeSize, MaxNumCells * GMaxCulledLightsPerCell), TEXT("CulledLightDataGrid"));
-				CulledLightDataGridSRV = GraphBuilder.CreateSRV(CulledLightDataGrid);
-				CulledLightDataGridUAV = GraphBuilder.CreateUAV(CulledLightDataGrid);
-			}
-
-			FLightGridInjectionCS::FParameters *PassParameters = GraphBuilder.AllocParameters<FLightGridInjectionCS::FParameters>();
-
-			PassParameters->View                    = View.ViewUniformBuffer;
-			
-			if (IsMobilePlatform(View.GetShaderPlatform()))
-			{
-				PassParameters->MobileReflectionCaptureData = View.MobileReflectionCaptureUniformBuffer;
-			}
-			else
-			{
-				PassParameters->ReflectionCapture = View.ReflectionCaptureUniformBuffer;
-			}
-
-			PassParameters->RWNumCulledLightsGrid   = NumCulledLightsGridUAV;
-			PassParameters->RWCulledLightDataGrid32Bit   = CulledLightDataGridUAV;
-			PassParameters->RWCulledLightDataGrid16Bit = CulledLightDataGridUAV;
-			PassParameters->RWNextCulledLightLink   = GraphBuilder.CreateUAV(NextCulledLightLinkBuffer);
-			PassParameters->RWNextCulledLightData   = NextCulledLightDataUAV;
-			PassParameters->RWCulledLightLinks      = GraphBuilder.CreateUAV(CulledLightLinksBuffer);
-			PassParameters->CulledGridSize          = ForwardLightData->CulledGridSize;
-			PassParameters->LightGridZParams        = ForwardLightData->LightGridZParams;
-			PassParameters->NumReflectionCaptures   = ForwardLightData->NumReflectionCaptures;
-			PassParameters->NumLocalLights          = ForwardLightData->NumLocalLights;
-			PassParameters->MaxCulledLightsPerCell  = GMaxCulledLightsPerCell;
-			PassParameters->NumGridCells            = ForwardLightData->NumGridCells;
-			PassParameters->LightGridPixelSizeShift = ForwardLightData->LightGridPixelSizeShift;
-			PassParameters->MegaLightsSupportedStartIndex = ForwardLightData->MegaLightsSupportedStartIndex;
-
-			check(ViewSpacePosAndRadiusData.Num() == ForwardLocalLightData.Num());
-			check(ViewSpaceDirAndPreprocAngleData.Num() == ForwardLocalLightData.Num());
-			check(ViewSpaceRectPlanesData.Num() == ForwardLocalLightData.Num() * NUM_PLANES_PER_RECT_LIGHT);
-
-			FRDGBufferRef LightViewSpacePositionAndRadius  = CreateStructuredBuffer(GraphBuilder, TEXT("ViewSpacePosAndRadiusData"), TConstArrayView<FVector4f>(ViewSpacePosAndRadiusData));
-			FRDGBufferRef LightViewSpaceDirAndPreprocAngle = CreateStructuredBuffer(GraphBuilder, TEXT("ViewSpacePosAndRadiusData"), TConstArrayView<FVector4f>(ViewSpaceDirAndPreprocAngleData));
-			FRDGBufferRef LightViewSpaceRectPlanes = CreateStructuredBuffer(GraphBuilder, TEXT("ViewSpaceRectPlanesData"), TConstArrayView<FVector4f>(ViewSpaceRectPlanesData));
-
-			PassParameters->LightViewSpacePositionAndRadius  = GraphBuilder.CreateSRV(FRDGBufferSRVDesc(LightViewSpacePositionAndRadius));
-			PassParameters->LightViewSpaceDirAndPreprocAngle = GraphBuilder.CreateSRV(FRDGBufferSRVDesc(LightViewSpaceDirAndPreprocAngle));
-			PassParameters->LightViewSpaceRectPlanes = GraphBuilder.CreateSRV(FRDGBufferSRVDesc(LightViewSpaceRectPlanes));
-
-			{
-				PassParameters->HZBTexture = View.HZB;
-				PassParameters->HZBSampler = TStaticSamplerState< SF_Point, AM_Clamp, AM_Clamp, AM_Clamp >::GetRHI();
-				PassParameters->HZBSize = FVector2f(View.HZBMipmap0Size);
-				PassParameters->HZBViewSize = FVector2f(View.ViewRect.Size());
-				PassParameters->HZBViewRect = FIntRect(0, 0, View.ViewRect.Width(), View.ViewRect.Height());
-			}
+				ForwardLightData->NumReflectionCaptures,
+				ForwardLightData->MegaLightsSupportedStartIndex,
+				bLightGridUses16BitBuffers,
+				bHasRectLights && (GLightGridRefineRectLightBounds != 0),
+				LightViewSpacePositionAndRadiusSRV,
+				LightViewSpaceDirAndPreprocAngleSRV,
+				LightViewSpaceRectPlanesSRV,
+				true,
+				false,
+				nullptr,
+				nullptr,
+				1);
+		}
 
-			FLightGridInjectionCS::FPermutationDomain PermutationVector;
-			PermutationVector.Set<FLightGridInjectionCS::FUseLinkedListDim>(GLightLinkedListCulling != 0);
-			PermutationVector.Set<FLightGridInjectionCS::FRefineRectLightBoundsDim>(bHasRectLights && (GLightGridRefineRectLightBounds != 0));
-			PermutationVector.Set<FLightGridInjectionCS::FUseHZBCullDim>(GLightGridHZBCull != 0 && View.HZB != nullptr);
-			TShaderMapRef<FLightGridInjectionCS> ComputeShader(View.ShaderMap, PermutationVector);
+		const int32 WorkloadDistributionMode = CVarLightCullingWorkloadDistributionMode.GetValueOnRenderThread();
 
-			if (GLightLinkedListCulling != 0)
-			{
-				AddClearUAVPass(GraphBuilder, PassParameters->RWNextCulledLightLink, 0);
-				AddClearUAVPass(GraphBuilder, NextCulledLightDataUAV, 0);
-				AddClearUAVPass(GraphBuilder, NumCulledLightsGridUAV, 0);
-				FComputeShaderUtils::AddPass(GraphBuilder, RDG_EVENT_NAME("LightGridInject:LinkedList"), ComputeShader, PassParameters, NumGroups);
-			}
-			else
-			{
-				AddClearUAVPass(GraphBuilder, NumCulledLightsGridUAV, 0);
-				FComputeShaderUtils::AddPass(GraphBuilder, RDG_EVENT_NAME("LightGridInject:NotLinkedList"), ComputeShader, PassParameters, NumGroups);
-			}
-			ForwardLightData->CulledLightDataGrid32Bit = CulledLightDataGridSRV;
-			ForwardLightData->CulledLightDataGrid16Bit = CulledLightDataGridSRV;
-			ForwardLightData->NumCulledLightsGrid = GraphBuilder.CreateSRV(FRDGBufferSRVDesc(NumCulledLightsGrid));
-			View.ForwardLightingResources.SetUniformBuffer(GraphBuilder.CreateUniformBuffer(ForwardLightData));
+		uint32 NumThreadsPerCell = 1;
+		
+		if (WorkloadDistributionMode == 1) // thread group per cell (64 threads)
+		{
+			NumThreadsPerCell = 64;
 		}
+		else if (WorkloadDistributionMode == 2 && GRHIMinimumWaveSize <= 32) // thread group per cell (32 threads if supported, otherwise single thread).
+		{
+			NumThreadsPerCell = 32;
+		}
+
+		FLightGrid LightGrid = LightGridInjection(
+			GraphBuilder,
+			View,
+			ForwardLightData->CulledGridSize,
+			ForwardLightData->LightGridPixelSizeShift,
+			1,
+			MaxNumCells,
+			ForwardLightData->LightGridZParams,
+			ForwardLightData->NumLocalLights,
+			ForwardLightData->NumReflectionCaptures,
+			ForwardLightData->MegaLightsSupportedStartIndex,
+			bLightGridUses16BitBuffers,
+			bHasRectLights && (GLightGridRefineRectLightBounds != 0),
+			LightViewSpacePositionAndRadiusSRV,
+			LightViewSpaceDirAndPreprocAngleSRV,
+			LightViewSpaceRectPlanesSRV,
+			NumThreadsPerCell > 1,
+			NumThreadsPerCell == 32,
+			ParentLightGrid.NumCulledLightsGridSRV,
+			ParentLightGrid.CulledLightDataGridSRV,
+			ParentLightGridFactor);
+
+		ForwardLightData->CulledLightDataGrid32Bit = LightGrid.CulledLightDataGridSRV;
+		ForwardLightData->CulledLightDataGrid16Bit = LightGrid.CulledLightDataGridSRV;
+		ForwardLightData->NumCulledLightsGrid = LightGrid.NumCulledLightsGridSRV;
+		View.ForwardLightingResources.SetUniformBuffer(GraphBuilder.CreateUniformBuffer(ForwardLightData));
 	}
 
 #if WITH_EDITOR
