Fix VSM stats shader compilation issue where HLSL templates are not supported.

diff --git a/Engine/Shaders/Private/VirtualShadowMaps/VirtualShadowMapBuildPerPageDrawCommands.usf b/Engine/Shaders/Private/VirtualShadowMaps/VirtualShadowMapBuildPerPageDrawCommands.usf
--- a/Engine/Shaders/Private/VirtualShadowMaps/VirtualShadowMapBuildPerPageDrawCommands.usf
+++ b/Engine/Shaders/Private/VirtualShadowMaps/VirtualShadowMapBuildPerPageDrawCommands.usf
@@ -390,8 +390,8 @@ void CullPerPageDrawCommandsCs(uint3 GroupId : SV_GroupID, int GroupThreadIndex
 		}
 		
 		StatsBufferInterlockedAdd(VSM_STAT_NON_NANITE_INSTANCES_DRAWN, ThreadTotalForAllViews);
-		StatsBufferInterlockedEnableFlags<true>(VSM_STAT_OVERFLOW_FLAGS, VSM_STAT_OVERFLOW_FLAG_MARKING_JOB_QUEUE, bMarkingJobQueueOverflow);
-		StatsBufferInterlockedEnableFlags<true>(VSM_STAT_OVERFLOW_FLAGS, VSM_STAT_OVERFLOW_FLAG_VISIBLE_INSTANCES, bVisibleInstancesOverflow);
+		StatsBufferInterlockedEnableFlags(VSM_STAT_OVERFLOW_FLAGS, VSM_STAT_OVERFLOW_FLAG_MARKING_JOB_QUEUE, bMarkingJobQueueOverflow, true);
+		StatsBufferInterlockedEnableFlags(VSM_STAT_OVERFLOW_FLAGS, VSM_STAT_OVERFLOW_FLAG_VISIBLE_INSTANCES, bVisibleInstancesOverflow, true);
 
 		// Accumulate total number of instances for each indirect argument, is also used to allocate space and output compact range of instances later
 		InterlockedAdd(DrawIndirectArgsBufferOut[Payload.IndirectArgIndex * INDIRECT_ARGS_NUM_WORDS + 1], ThreadTotalForAllViews);
diff --git a/Engine/Shaders/Private/VirtualShadowMaps/VirtualShadowMapProjection.usf b/Engine/Shaders/Private/VirtualShadowMaps/VirtualShadowMapProjection.usf
--- a/Engine/Shaders/Private/VirtualShadowMaps/VirtualShadowMapProjection.usf
+++ b/Engine/Shaders/Private/VirtualShadowMaps/VirtualShadowMapProjection.usf
@@ -762,7 +762,7 @@ void VirtualShadowMapProjection(
 	// We can only handle so many lights in our output encoding right now, so no purpose in computing more
 	uint LightCount = min(GetPackedShadowMaskMaxLightCount(), LightGridHeader.NumLights);
 	bool bLightCountOverflow = GetPackedShadowMaskMaxLightCount() < LightGridHeader.NumLights;
-	StatsBufferInterlockedEnableFlags<true>(VSM_STAT_OVERFLOW_FLAGS, VSM_STAT_OVERFLOW_FLAG_OPP_MAX_LIGHTS, bLightCountOverflow);
+	StatsBufferInterlockedEnableFlags(VSM_STAT_OVERFLOW_FLAGS, VSM_STAT_OVERFLOW_FLAG_OPP_MAX_LIGHTS, bLightCountOverflow, true);
 
 	uint4 ShadowMask = InitializePackedShadowMask();
 
diff --git a/Engine/Shaders/Private/VirtualShadowMaps/VirtualShadowMapStats.ush b/Engine/Shaders/Private/VirtualShadowMaps/VirtualShadowMapStats.ush
--- a/Engine/Shaders/Private/VirtualShadowMaps/VirtualShadowMapStats.ush
+++ b/Engine/Shaders/Private/VirtualShadowMaps/VirtualShadowMapStats.ush
@@ -17,8 +17,7 @@
 
 RWStructuredBuffer<uint> OutStatsBuffer;
 
-template<bool bWriteEnabled = VSM_GENERATE_STATS>
-void StatsBufferWrite(uint Index, uint Value)
+void StatsBufferWrite(uint Index, uint Value, bool bWriteEnabled = VSM_GENERATE_STATS)
 {
 	if (bWriteEnabled)
 	{
@@ -26,8 +25,7 @@ void StatsBufferWrite(uint Index, uint Value)
 	}
 }
 
-template<bool bWriteEnabled = VSM_GENERATE_STATS>
-void StatsBufferInterlockedAdd(uint Index, uint Value)
+void StatsBufferInterlockedAdd(uint Index, uint Value, bool bWriteEnabled = VSM_GENERATE_STATS)
 {
 	if (bWriteEnabled)
 	{
@@ -35,8 +33,7 @@ void StatsBufferInterlockedAdd(uint Index, uint Value)
 	}
 }
 
-template<bool bWriteEnabled = VSM_GENERATE_STATS>
-void StatsBufferInterlockedInc(uint Index)
+void StatsBufferInterlockedInc(uint Index, bool bWriteEnabled = VSM_GENERATE_STATS)
 {
 	if (bWriteEnabled)
 	{
@@ -44,8 +41,7 @@ void StatsBufferInterlockedInc(uint Index)
 	}
 }
 
-template<bool bWriteEnabled = VSM_GENERATE_STATS>
-void StatsBufferInterlockedEnableFlags(uint Index, uint Flags, bool bEnabled)
+void StatsBufferInterlockedEnableFlags(uint Index, uint Flags, bool bEnabled, bool bWriteEnabled = VSM_GENERATE_STATS)
 {
 	if (bWriteEnabled)
 	{
