Heterogeneous Volumes: Improve depth precision for AVSMs

Extend the AVSM indirection buffer to handle 32-bit depth. Change AVSM sample depth to 16-bit fixed-point, and compress based on based on depth-range of the indirection data.

diff --git a/Engine/Shaders/Private/HeterogeneousVolumes/HeterogeneousVolumesAdaptiveVolumetricShadowMapSampler.ush b/Engine/Shaders/Private/HeterogeneousVolumes/HeterogeneousVolumesAdaptiveVolumetricShadowMapSampler.ush
--- a/Engine/Shaders/Private/HeterogeneousVolumes/HeterogeneousVolumesAdaptiveVolumetricShadowMapSampler.ush
+++ b/Engine/Shaders/Private/HeterogeneousVolumes/HeterogeneousVolumesAdaptiveVolumetricShadowMapSampler.ush
@@ -54,8 +54,8 @@ FAVSM_Sampler AVSM_Sampler_Create(
 	{
 		int TopLevelSampleCount = CalcTopLevelSampleCount(Iterator.SampleCount);
 		Iterator.SampleIndex = 0;
-		Iterator.SampleData[0] = AVSM_UnpackSampleData(Iterator.ShadowMap.SampleBuffer[Iterator.PixelOffset + Align4(TopLevelSampleCount) + Iterator.SampleIndex]);
-		Iterator.SampleData[1] = AVSM_UnpackSampleData(Iterator.ShadowMap.SampleBuffer[Iterator.PixelOffset + Align4(TopLevelSampleCount) + Iterator.SampleIndex + 1]);
+		Iterator.SampleData[0] = AVSM_UnpackSampleData(Iterator.ShadowMap.SampleBuffer[Iterator.PixelOffset + Align4(TopLevelSampleCount) + Iterator.SampleIndex], IndirectionData);
+		Iterator.SampleData[1] = AVSM_UnpackSampleData(Iterator.ShadowMap.SampleBuffer[Iterator.PixelOffset + Align4(TopLevelSampleCount) + Iterator.SampleIndex + 1], IndirectionData);
 	}
 	else
 	{
@@ -81,7 +81,7 @@ float AVSM_Sampler_Eval(inout FAVSM_Sampler Iterator, float X)
 		{
 			Iterator.SampleIndex++;
 			Iterator.SampleData[0] = Iterator.SampleData[1];
-			Iterator.SampleData[1] = AVSM_UnpackSampleData(Iterator.ShadowMap.SampleBuffer[Iterator.PixelOffset + Align4(TopLevelSampleCount) + Iterator.SampleIndex]);
+			Iterator.SampleData[1] = AVSM_UnpackSampleData(Iterator.ShadowMap.SampleBuffer[Iterator.PixelOffset + Align4(TopLevelSampleCount) + Iterator.SampleIndex], Iterator.IndirectionData);
 		}
 	}
 
@@ -100,7 +100,7 @@ FAVSMSampleData AVSM_Sampler_EvalInverse(inout FAVSM_Sampler Iterator, float Tau
 		{
 			Iterator.SampleIndex++;
 			Iterator.SampleData[0] = Iterator.SampleData[1];
-			Iterator.SampleData[1] = AVSM_UnpackSampleData(Iterator.ShadowMap.SampleBuffer[Iterator.PixelOffset + Align4(TopLevelSampleCount) + Iterator.SampleIndex]);
+			Iterator.SampleData[1] = AVSM_UnpackSampleData(Iterator.ShadowMap.SampleBuffer[Iterator.PixelOffset + Align4(TopLevelSampleCount) + Iterator.SampleIndex], Iterator.IndirectionData);
 		}
 	}
 
@@ -119,7 +119,7 @@ float AVSM_Sampler_Last(inout FAVSM_Sampler Iterator)
 		if (TopLevelSampleCount > 0)
 		{
 			uint SampleIndex = max(Iterator.SampleCount - 2, 0);
-			FAVSMSampleData SampleData = AVSM_UnpackSampleData(Iterator.ShadowMap.SampleBuffer[Iterator.PixelOffset + Align4(TopLevelSampleCount) + SampleIndex]);
+			FAVSMSampleData SampleData = AVSM_UnpackSampleData(Iterator.ShadowMap.SampleBuffer[Iterator.PixelOffset + Align4(TopLevelSampleCount) + SampleIndex], Iterator.IndirectionData);
 
 			Transmittance = SampleData.Tau;
 		}
@@ -134,6 +134,8 @@ struct FAVSM_Sampler4
 	FAdaptiveVolumetricShadowMap ShadowMap;
 	uint4 PixelOffset4;
 	uint SampleCount4;
+	float4 BeginX4;
+	float4 EndX4;
 	uint SampleIndex4;
 
 	// Query cache
@@ -146,40 +148,6 @@ struct FAVSM_Sampler4
 	float DepthOffset;
 };
 
-FAVSMSampleData AVSM_Sampler4_GetSampleData0(inout FAVSM_Sampler4 Sampler4, uint Index)
-{
-	FAVSMSampleData SampleData0 = AVSM_UnpackSampleData(Sampler4.SampleData0Packed4[Index]);
-	SampleData0.X += Sampler4.DepthOffset;
-	return SampleData0;
-}
-
-void AVSM_Sampler4_SetSampleData0(inout FAVSM_Sampler4 Sampler4, uint4 SampleDataPacked4)
-{
-	Sampler4.SampleData0Packed4 = SampleDataPacked4;
-}
-
-void AVSM_Sampler4_SetSampleData0(inout FAVSM_Sampler4 Sampler4, uint Index, uint SampleDataPacked)
-{
-	Sampler4.SampleData0Packed4[Index] = SampleDataPacked;
-}
-
-FAVSMSampleData AVSM_Sampler4_GetSampleData1(inout FAVSM_Sampler4 Sampler4, uint Index)
-{
-	FAVSMSampleData SampleData1 = AVSM_UnpackSampleData(Sampler4.SampleData1Packed4[Index]);
-	SampleData1.X += Sampler4.DepthOffset;
-	return SampleData1;
-}
-
-void AVSM_Sampler4_SetSampleData1(inout FAVSM_Sampler4 Sampler4, uint4 SampleDataPacked4)
-{
-	Sampler4.SampleData1Packed4 = SampleDataPacked4;
-}
-
-void AVSM_Sampler4_SetSampleData1(inout FAVSM_Sampler4 Sampler4, uint Index, uint SampleDataPacked)
-{
-	Sampler4.SampleData1Packed4[Index] = SampleDataPacked;
-}
-
 uint AVSM_Sampler4_GetSampleIndex(inout FAVSM_Sampler4 Sampler4, uint Index)
 {
 	uint SampleIndex = (Sampler4.SampleIndex4 >> (Index * 8u)) & INVALID_SAMPLE_INDEX;
@@ -209,6 +177,18 @@ uint AVSM_Sampler4_GetSampleCount(inout FAVSM_Sampler4 Sampler4, uint Index)
 	return SampleCount;
 }
 
+float AVSM_Sampler4_GetBeginX(inout FAVSM_Sampler4 Sampler4, uint Index)
+{
+	float BeginX = Sampler4.BeginX4[Index];
+	return BeginX;
+}
+
+float AVSM_Sampler4_GetEndX(inout FAVSM_Sampler4 Sampler4, uint Index)
+{
+	float EndX = Sampler4.EndX4[Index];
+	return EndX;
+}
+
 bool AVSM_Sampler4_IsValid(inout FAVSM_Sampler4 Sampler4, uint Index)
 {
 	return AVSM_Sampler4_GetSampleIndex(Sampler4, Index) != INVALID_SAMPLE_INDEX;
@@ -235,6 +215,64 @@ void AVSM_Sampler4_SetPixelOffset(inout FAVSM_Sampler4 Sampler4, uint4 PixelOffs
 	Sampler4.PixelOffset4 = PixelOffset;
 }
 
+void AVSM_Sampler4_SetBeginX(inout FAVSM_Sampler4 Sampler4, float4 BeginX4)
+{
+	Sampler4.BeginX4 = BeginX4;
+}
+
+void AVSM_Sampler4_SetEndX(inout FAVSM_Sampler4 Sampler4, float4 EndX4)
+{
+	Sampler4.EndX4 = EndX4;
+}
+
+FAVSMSampleData AVSM_Sampler4_GetSampleData0(inout FAVSM_Sampler4 Sampler4, uint Index)
+{
+	FAVSMIndirectionData IndirectionData = AVSM_CreateIndirectionData(
+		AVSM_Sampler4_GetPixelOffset(Sampler4, Index),
+		AVSM_Sampler4_GetSampleIndex(Sampler4, Index),
+		AVSM_Sampler4_GetBeginX(Sampler4, Index),
+		AVSM_Sampler4_GetEndX(Sampler4, Index)
+	);
+
+	FAVSMSampleData SampleData0 = AVSM_UnpackSampleData(Sampler4.SampleData0Packed4[Index], IndirectionData);
+	SampleData0.X += Sampler4.DepthOffset;
+	return SampleData0;
+}
+
+void AVSM_Sampler4_SetSampleData0(inout FAVSM_Sampler4 Sampler4, uint4 SampleDataPacked4)
+{
+	Sampler4.SampleData0Packed4 = SampleDataPacked4;
+}
+
+void AVSM_Sampler4_SetSampleData0(inout FAVSM_Sampler4 Sampler4, uint Index, uint SampleDataPacked)
+{
+	Sampler4.SampleData0Packed4[Index] = SampleDataPacked;
+}
+
+FAVSMSampleData AVSM_Sampler4_GetSampleData1(inout FAVSM_Sampler4 Sampler4, uint Index)
+{
+	FAVSMIndirectionData IndirectionData = AVSM_CreateIndirectionData(
+		AVSM_Sampler4_GetPixelOffset(Sampler4, Index),
+		AVSM_Sampler4_GetSampleIndex(Sampler4, Index),
+		AVSM_Sampler4_GetBeginX(Sampler4, Index),
+		AVSM_Sampler4_GetEndX(Sampler4, Index)
+	);
+
+	FAVSMSampleData SampleData1 = AVSM_UnpackSampleData(Sampler4.SampleData1Packed4[Index], IndirectionData);
+	SampleData1.X += Sampler4.DepthOffset;
+	return SampleData1;
+}
+
+void AVSM_Sampler4_SetSampleData1(inout FAVSM_Sampler4 Sampler4, uint4 SampleDataPacked4)
+{
+	Sampler4.SampleData1Packed4 = SampleDataPacked4;
+}
+
+void AVSM_Sampler4_SetSampleData1(inout FAVSM_Sampler4 Sampler4, uint Index, uint SampleDataPacked)
+{
+	Sampler4.SampleData1Packed4[Index] = SampleDataPacked;
+}
+
 FAVSM_Sampler4 AVSM_Sampler4_Create(
 	inout FAdaptiveVolumetricShadowMap ShadowMap,
 	int Face,
@@ -263,11 +301,15 @@ FAVSM_Sampler4 AVSM_Sampler4_Create(
 	uint4 SampleData0Packed4 = 0;
 	uint4 SampleData1Packed4 = 0;
 	uint4 PixelOffset4 = 0;
+	float4 BeginX4 = 0;
+	float4 EndX4 = 0;
 	for (int Index = 0; Index < 4; ++Index)
 	{
 		FAVSMIndirectionData IndirectionData = AVSM_UnpackIndirectionData(ShadowMap.IndirectionBuffer[LinearIndex4[Index]]);
 		PixelOffset4[Index] = IndirectionData.PixelOffset;
 		SampleCount4[Index] = IndirectionData.SampleCount;
+		BeginX4[Index] = IndirectionData.BeginX;
+		EndX4[Index] = IndirectionData.EndX;
 
 		if (SampleCount4[Index] > 0)
 		{
@@ -279,8 +321,8 @@ FAVSM_Sampler4 AVSM_Sampler4_Create(
 		else
 		{
 			SampleIndex4[Index] = INVALID_SAMPLE_INDEX;
-			SampleData0Packed4[Index] = AVSM_PackSampleData(0.0, 1.0);
-			SampleData1Packed4[Index] = AVSM_PackSampleData(0.0, 1.0);
+			SampleData0Packed4[Index] = AVSM_PackSampleData(0.0, 1.0, IndirectionData);
+			SampleData1Packed4[Index] = AVSM_PackSampleData(0.0, 1.0, IndirectionData);
 		}
 	}
 	AVSM_Sampler4_SetSampleCount(Sampler4, SampleCount4);
@@ -288,6 +330,8 @@ FAVSM_Sampler4 AVSM_Sampler4_Create(
 	AVSM_Sampler4_SetSampleData0(Sampler4, SampleData0Packed4);
 	AVSM_Sampler4_SetSampleData1(Sampler4, SampleData1Packed4);
 	AVSM_Sampler4_SetPixelOffset(Sampler4, PixelOffset4);
+	AVSM_Sampler4_SetBeginX(Sampler4, BeginX4);
+	AVSM_Sampler4_SetEndX(Sampler4, EndX4);
 
 	Sampler4.Face = Face;
 	Sampler4.Pixel = Pixel;
@@ -315,11 +359,17 @@ float AVSM_Sampler4_Eval_Internal(inout FAVSM_Sampler4 Sampler4, int Index, floa
 			SampleIndex++;
 			SampleData[0] = SampleData[1];
 			uint PixelOffset = AVSM_Sampler4_GetPixelOffset(Sampler4, Index);
-			SampleData[1] = AVSM_UnpackSampleData(Sampler4.ShadowMap.SampleBuffer[PixelOffset + Align4(TopLevelSampleCount) + SampleIndex]);
+			FAVSMIndirectionData IndirectionData = AVSM_CreateIndirectionData(
+				AVSM_Sampler4_GetPixelOffset(Sampler4, Index),
+				AVSM_Sampler4_GetSampleIndex(Sampler4, Index),
+				AVSM_Sampler4_GetBeginX(Sampler4, Index),
+				AVSM_Sampler4_GetEndX(Sampler4, Index)
+			);
+			SampleData[1] = AVSM_UnpackSampleData(Sampler4.ShadowMap.SampleBuffer[PixelOffset + Align4(TopLevelSampleCount) + SampleIndex], IndirectionData);
 
 			AVSM_Sampler4_SetSampleIndex(Sampler4, Index, SampleIndex);
-			AVSM_Sampler4_SetSampleData0(Sampler4, Index, AVSM_PackSampleData(SampleData[0]));
-			AVSM_Sampler4_SetSampleData1(Sampler4, Index, AVSM_PackSampleData(SampleData[1]));
+			AVSM_Sampler4_SetSampleData0(Sampler4, Index, AVSM_PackSampleData(SampleData[0], IndirectionData));
+			AVSM_Sampler4_SetSampleData1(Sampler4, Index, AVSM_PackSampleData(SampleData[1], IndirectionData));
 		}
 	}
 
@@ -376,10 +426,16 @@ FAVSMSampleData AVSM_Sampler4_EvalInverse_Internal(inout FAVSM_Sampler4 Sampler4
 		{
 			SampleIndex++;
 			SampleData[0] = SampleData[1];
-			SampleData[1] = AVSM_UnpackSampleData(Sampler4.ShadowMap.SampleBuffer[PixelOffset + Align4(TopLevelSampleCount) + SampleIndex]);
-
-			AVSM_Sampler4_SetSampleData0(Sampler4, Index, AVSM_PackSampleData(SampleData[0]));
-			AVSM_Sampler4_SetSampleData1(Sampler4, Index, AVSM_PackSampleData(SampleData[1]));
+			FAVSMIndirectionData IndirectionData = AVSM_CreateIndirectionData(
+				AVSM_Sampler4_GetPixelOffset(Sampler4, Index),
+				AVSM_Sampler4_GetSampleIndex(Sampler4, Index),
+				AVSM_Sampler4_GetBeginX(Sampler4, Index),
+				AVSM_Sampler4_GetEndX(Sampler4, Index)
+			);
+			SampleData[1] = AVSM_UnpackSampleData(Sampler4.ShadowMap.SampleBuffer[PixelOffset + Align4(TopLevelSampleCount) + SampleIndex], IndirectionData);
+
+			AVSM_Sampler4_SetSampleData0(Sampler4, Index, AVSM_PackSampleData(SampleData[0], IndirectionData));
+			AVSM_Sampler4_SetSampleData1(Sampler4, Index, AVSM_PackSampleData(SampleData[1], IndirectionData));
 		}
 	}
 
@@ -416,7 +472,13 @@ FAVSMSampleData AVSM_Sampler4_Last_Internal(inout FAVSM_Sampler4 Sampler4, int I
 		{
 			uint SampleIndex = max(SampleCount - 2, 0);
 			uint PixelOffset = AVSM_Sampler4_GetPixelOffset(Sampler4, Index);
-			SampleData = AVSM_UnpackSampleData(Sampler4.ShadowMap.SampleBuffer[PixelOffset + Align4(TopLevelSampleCount) + SampleIndex]);
+			FAVSMIndirectionData IndirectionData = AVSM_CreateIndirectionData(
+				AVSM_Sampler4_GetPixelOffset(Sampler4, Index),
+				AVSM_Sampler4_GetSampleIndex(Sampler4, Index),
+				AVSM_Sampler4_GetBeginX(Sampler4, Index),
+				AVSM_Sampler4_GetEndX(Sampler4, Index)
+			);
+			SampleData = AVSM_UnpackSampleData(Sampler4.ShadowMap.SampleBuffer[PixelOffset + Align4(TopLevelSampleCount) + SampleIndex], IndirectionData);
 			SampleData.X += Sampler4.DepthOffset;
 		}
 	}
diff --git a/Engine/Shaders/Private/HeterogeneousVolumes/HeterogeneousVolumesAdaptiveVolumetricShadowMapSampling.ush b/Engine/Shaders/Private/HeterogeneousVolumes/HeterogeneousVolumesAdaptiveVolumetricShadowMapSampling.ush
--- a/Engine/Shaders/Private/HeterogeneousVolumes/HeterogeneousVolumesAdaptiveVolumetricShadowMapSampling.ush
+++ b/Engine/Shaders/Private/HeterogeneousVolumes/HeterogeneousVolumesAdaptiveVolumetricShadowMapSampling.ush
@@ -26,7 +26,7 @@ struct FAdaptiveVolumetricShadowMap
 	bool bIsEmpty;
 	bool bIsDirectionalLight;
 
-	StructuredBuffer<uint2> IndirectionBuffer;
+	StructuredBuffer<uint4> IndirectionBuffer;
 	StructuredBuffer<uint> SampleBuffer;
 };
 
@@ -99,16 +99,16 @@ float AVSM_Sample_Reference(
 		if (WorldHitT >= IndirectionData.EndX)
 		{
 			int LastSampleIndex = IndirectionData.PixelOffset + Align4(TopLevelSampleCount) + IndirectionData.SampleCount - 1;
-			return AVSM_UnpackSampleData(ShadowMap.SampleBuffer[LastSampleIndex]).Tau;
+			return AVSM_UnpackSampleData(ShadowMap.SampleBuffer[LastSampleIndex], IndirectionData).Tau;
 		}
 
 		// Initialize from the top-level
-		ShadowData[0] = AVSM_UnpackSampleData(ShadowMap.SampleBuffer[IndirectionData.PixelOffset]);
+		ShadowData[0] = AVSM_UnpackSampleData(ShadowMap.SampleBuffer[IndirectionData.PixelOffset], IndirectionData);
 
 		// Iterate through the linear span of offset transmittance information
 		for (int ElementIndex = 0; ElementIndex < IndirectionData.SampleCount - 1; ++ElementIndex)
 		{
-			ShadowData[1] = AVSM_UnpackSampleData(ShadowMap.SampleBuffer[IndirectionData.PixelOffset + Align4(TopLevelSampleCount) + ElementIndex]);
+			ShadowData[1] = AVSM_UnpackSampleData(ShadowMap.SampleBuffer[IndirectionData.PixelOffset + Align4(TopLevelSampleCount) + ElementIndex], IndirectionData);
 			if (WorldHitT < ShadowData[1].X)
 			{
 				break;
@@ -149,11 +149,11 @@ float AVSM_Sample_BinarySearch(
 		if (WorldHitT >= IndirectionData.EndX)
 		{
 			int LastSampleIndex = IndirectionData.PixelOffset + Align4(TopLevelSampleCount) + IndirectionData.SampleCount - 1;
-			return AVSM_UnpackSampleData(ShadowMap.SampleBuffer[LastSampleIndex]).Tau;
+			return AVSM_UnpackSampleData(ShadowMap.SampleBuffer[LastSampleIndex], IndirectionData).Tau;
 		}
 
 		// Initialize from the top-level
-		ShadowData[0] = AVSM_UnpackSampleData(ShadowMap.SampleBuffer[IndirectionData.PixelOffset]);
+		ShadowData[0] = AVSM_UnpackSampleData(ShadowMap.SampleBuffer[IndirectionData.PixelOffset], IndirectionData);
 
 		// Account for the shifted bottom-level
 		int ElementIndexMin = -1;
@@ -161,7 +161,7 @@ float AVSM_Sample_BinarySearch(
 		while (ElementIndexMax - ElementIndexMin > 1)
 		{
 			int ElementIndex = (ElementIndexMin + ElementIndexMax + 1) / 2;
-			FAVSMSampleData SampleData = AVSM_UnpackSampleData(ShadowMap.SampleBuffer[IndirectionData.PixelOffset + Align4(TopLevelSampleCount) + ElementIndex]);
+			FAVSMSampleData SampleData = AVSM_UnpackSampleData(ShadowMap.SampleBuffer[IndirectionData.PixelOffset + Align4(TopLevelSampleCount) + ElementIndex], IndirectionData);
 
 			if (WorldHitT < SampleData.X)
 			{
@@ -209,15 +209,15 @@ float AVSM_Sample_TwoLevel_Reference(
 		if (WorldHitT >= IndirectionData.EndX)
 		{
 			int LastSampleIndex = IndirectionData.PixelOffset + Align4(TopLevelSampleCount) + IndirectionData.SampleCount - 1;
-			return AVSM_UnpackSampleData(AVSM.SampleBuffer[LastSampleIndex]).Tau;
+			return AVSM_UnpackSampleData(AVSM.SampleBuffer[LastSampleIndex], IndirectionData).Tau;
 		}
 
-		ShadowData[0] = AVSM_UnpackSampleData(ShadowMap.SampleBuffer[IndirectionData.PixelOffset]);
+		ShadowData[0] = AVSM_UnpackSampleData(ShadowMap.SampleBuffer[IndirectionData.PixelOffset], IndirectionData);
 
 		// Iterate through the top-level
 		for (TopLevelSampleIndex = 0; TopLevelSampleIndex < TopLevelSampleCount; ++TopLevelSampleIndex)
 		{
-			ShadowData[1] = AVSM_UnpackSampleData(ShadowMap.SampleBuffer[IndirectionData.PixelOffset + TopLevelSampleIndex]);
+			ShadowData[1] = AVSM_UnpackSampleData(ShadowMap.SampleBuffer[IndirectionData.PixelOffset + TopLevelSampleIndex], IndirectionData);
 			if (WorldHitT < ShadowData[1].X)
 			{
 				break;
@@ -233,7 +233,7 @@ float AVSM_Sample_TwoLevel_Reference(
 	int SampleIndexMax = min(TopLevelSampleIndex * 4, IndirectionData.SampleCount - 1);
 	for (int SampleIndex = SampleIndexMin; SampleIndex < SampleIndexMax; ++SampleIndex)
 	{
-		ShadowData[1] = AVSM_UnpackSampleData(ShadowMap.SampleBuffer[IndirectionData.PixelOffset + Align4(TopLevelSampleCount) + SampleIndex]);
+		ShadowData[1] = AVSM_UnpackSampleData(ShadowMap.SampleBuffer[IndirectionData.PixelOffset + Align4(TopLevelSampleCount) + SampleIndex], IndirectionData);
 		if (WorldHitT < ShadowData[1].X)
 		{
 			break;
@@ -277,7 +277,7 @@ float AVSM_Sample_TwoLevel_Vectorized(
 		if (WorldHitT >= IndirectionData.EndX)
 		{
 			int LastIndex = IndirectionData.PixelOffset + Align4(TopLevelSampleCount) + IndirectionData.SampleCount - 1;
-			return AVSM_UnpackSampleData(ShadowMap.SampleBuffer[LastIndex]).Tau;
+			return AVSM_UnpackSampleData(ShadowMap.SampleBuffer[LastIndex], IndirectionData).Tau;
 		}
 
 		int TopLevelSampleOffset = 0;
@@ -295,17 +295,17 @@ float AVSM_Sample_TwoLevel_Vectorized(
 			);
 
 			float4 Depths = float4(
-				AVSM_UnpackSampleData(PackedShadowData4[0]).X,
-				AVSM_UnpackSampleData(PackedShadowData4[1]).X,
-				AVSM_UnpackSampleData(PackedShadowData4[2]).X,
-				AVSM_UnpackSampleData(PackedShadowData4[3]).X
+				AVSM_UnpackSampleData(PackedShadowData4[0], IndirectionData).X,
+				AVSM_UnpackSampleData(PackedShadowData4[1], IndirectionData).X,
+				AVSM_UnpackSampleData(PackedShadowData4[2], IndirectionData).X,
+				AVSM_UnpackSampleData(PackedShadowData4[3], IndirectionData).X
 			);
 
 			int4 Result = float4(WorldHitT, WorldHitT, WorldHitT, WorldHitT) > Depths;
 			int LocalSampleIndex = Result[0] + Result[1] + Result[2] + Result[3];
 
-			ShadowData[0] = AVSM_UnpackSampleData(PackedShadowData4[max(LocalSampleIndex - 1, 0)]);
-			ShadowData[1] = AVSM_UnpackSampleData(PackedShadowData4[min(LocalSampleIndex, 3)]);
+			ShadowData[0] = AVSM_UnpackSampleData(PackedShadowData4[max(LocalSampleIndex - 1, 0)], IndirectionData);
+			ShadowData[1] = AVSM_UnpackSampleData(PackedShadowData4[min(LocalSampleIndex, 3)], IndirectionData);
 			TopLevelSampleIndex = min(TopLevelSampleOffset + LocalSampleIndex, TopLevelSampleCount);
 #if AVSM_SAMPLE_MODE == AVSM_SAMPLE_MODE_TWO_LEVEL_VECTORIZED_LOOP
 			if (LocalSampleIndex < 4)
@@ -330,10 +330,10 @@ float AVSM_Sample_TwoLevel_Vectorized(
 		);
 
 		float4 Depths = float4(
-			AVSM_UnpackSampleData(PackedShadowData4.x).X,
-			AVSM_UnpackSampleData(PackedShadowData4.y).X,
-			AVSM_UnpackSampleData(PackedShadowData4.z).X,
-			AVSM_UnpackSampleData(PackedShadowData4.w).X
+			AVSM_UnpackSampleData(PackedShadowData4.x, IndirectionData).X,
+			AVSM_UnpackSampleData(PackedShadowData4.y, IndirectionData).X,
+			AVSM_UnpackSampleData(PackedShadowData4.z, IndirectionData).X,
+			AVSM_UnpackSampleData(PackedShadowData4.w, IndirectionData).X
 		);
 
 		int4 Result = float4(WorldHitT, WorldHitT, WorldHitT, WorldHitT) > Depths;
@@ -341,9 +341,9 @@ float AVSM_Sample_TwoLevel_Vectorized(
 
 		if (Index > 0)
 		{
-			ShadowData[0] = AVSM_UnpackSampleData(PackedShadowData4[max(Index - 1, 0)]);
+			ShadowData[0] = AVSM_UnpackSampleData(PackedShadowData4[max(Index - 1, 0)], IndirectionData);
 		}
-		ShadowData[1] = AVSM_UnpackSampleData(PackedShadowData4[min(Index, 3)]);
+		ShadowData[1] = AVSM_UnpackSampleData(PackedShadowData4[min(Index, 3)], IndirectionData);
 	}
 
 	// Lerp the final value
diff --git a/Engine/Shaders/Private/HeterogeneousVolumes/HeterogeneousVolumesAdaptiveVolumetricShadowMapUtils.ush b/Engine/Shaders/Private/HeterogeneousVolumes/HeterogeneousVolumesAdaptiveVolumetricShadowMapUtils.ush
--- a/Engine/Shaders/Private/HeterogeneousVolumes/HeterogeneousVolumesAdaptiveVolumetricShadowMapUtils.ush
+++ b/Engine/Shaders/Private/HeterogeneousVolumes/HeterogeneousVolumesAdaptiveVolumetricShadowMapUtils.ush
@@ -10,7 +10,18 @@ struct FAVSMSampleData
 	float Tau;
 };
 
+#define AVSM_INVALID_PIXEL_OFFSET 0x3FFFFFF
+
+struct FAVSMIndirectionData
+{
+	uint PixelOffset;
+	int SampleCount;
+	float BeginX;
+	float EndX;
+};
+
 /* AVSM Sample Data API */
+
 FAVSMSampleData AVSM_CreateSampleData(float X, float Tau)
 {
 	FAVSMSampleData SampleData;
@@ -19,107 +30,110 @@ FAVSMSampleData AVSM_CreateSampleData(float X, float Tau)
 	return SampleData;
 }
 
-int AVSM_PackSampleData(float X, float3 Tau)
+uint AVSM_PackSampleDataX(float X, FAVSMIndirectionData IndirectionData)
 {
-	int PackedData = f32tof16(X) | (f32tof16(Luminance(Tau)) << 16);
-	return PackedData;
+	float Extent = IndirectionData.EndX - IndirectionData.BeginX;
+	float NormalizedX = (X - IndirectionData.BeginX) / Extent;
+	uint FixedX = saturate(NormalizedX) * MaxHalfFloat;
+	return FixedX;
+}
+
+float AVSM_UnpackSampleDataX(uint FixedX, FAVSMIndirectionData IndirectionData)
+{
+	float Extent = IndirectionData.EndX - IndirectionData.BeginX;
+	float NormalizedX = FixedX / MaxHalfFloat;
+	float X = saturate(NormalizedX) * Extent + IndirectionData.BeginX;
+	return X;
 }
 
-int AVSM_PackSampleData(FAVSMSampleData SampleData)
+int AVSM_PackSampleData(FAVSMSampleData SampleData, FAVSMIndirectionData IndirectionData)
 {
 	int PackedData = 0;
-	PackedData |= f32tof16(SampleData.X);
+	PackedData |= AVSM_PackSampleDataX(SampleData.X, IndirectionData);
 	PackedData |= f32tof16(SampleData.Tau) << 16;
 	return PackedData;
 }
 
-FAVSMSampleData AVSM_UnpackSampleData(int PackedData)
+FAVSMSampleData AVSM_UnpackSampleData(int PackedData, FAVSMIndirectionData IndirectionData)
 {
 	FAVSMSampleData SampleData;
-	SampleData.X = f16tof32(PackedData & 0xFFFF);
+	SampleData.X = AVSM_UnpackSampleDataX(PackedData & 0xFFFF, IndirectionData);
 	SampleData.Tau = f16tof32(PackedData >> 16);
 	return SampleData;
 }
 
+int AVSM_PackSampleData(float X, float Tau, FAVSMIndirectionData IndirectionData)
+{
+	return AVSM_PackSampleData(AVSM_CreateSampleData(X, Tau), IndirectionData);
+}
+
 /* AVSM Linked List API */
 
-FAVSMSampleData AVSM_UnpackSampleData(int2 PackedData)
+FAVSMSampleData AVSM_GetLinkedListSampleData(uint2 PackedData)
 {
-	return AVSM_UnpackSampleData(PackedData.x);
+	FAVSMSampleData SampleData;
+	SampleData.X = asfloat(PackedData.x);
+	SampleData.Tau = f16tof32(PackedData.y >> 16);
+	return SampleData;
 }
 
-int AVSM_GetPrevPtr(int2 PackedData)
+void AVSM_SetLinkedListSampleData(inout uint2 PackedData, FAVSMSampleData SampleData)
+{
+	PackedData.x = asuint(SampleData.X);
+	PackedData.y = (PackedData.y & 0xFFFF) |
+		(f32tof16(SampleData.Tau) << 16);
+}
+
+uint AVSM_GetLinkedListPrevPtr(uint2 PackedData)
 {
 	return (PackedData.y >> 8) & AVSM_NULL_PTR;
 }
 
-int AVSM_GetNextPtr(int2 PackedData)
+uint AVSM_GetLinkedListNextPtr(uint2 PackedData)
 {
 	return PackedData.y & AVSM_NULL_PTR;
 }
 
-void AVSM_SetPrevPtr(inout int2 PackedData, int PrevPtr)
+void AVSM_SetLinkedListPrevPtr(inout uint2 PackedData, uint PrevPtr)
 {
 	PackedData.y = (PackedData.y & 0xFFFF00FF) |
 		((PrevPtr & AVSM_NULL_PTR) << 8);
 }
 
-void AVSM_SetNextPtr(inout int2 PackedData, int NextPtr)
+void AVSM_SetLinkedListNextPtr(inout uint2 PackedData, uint NextPtr)
 {
 	PackedData.y = (PackedData.y & 0xFFFFFF00) |
 		(NextPtr & AVSM_NULL_PTR);
 }
 
-struct FAVSMIndirectionDataContext
-{
-	// TODO: Store begin and length
-	FAVSMSampleData Begin;
-	FAVSMSampleData End;
-};
+/* AVSM Indirection Data API */
 
-FAVSMIndirectionDataContext AVSM_CreateIndirectionDataContext(FAVSMSampleData Begin, FAVSMSampleData End)
-{
-	FAVSMIndirectionDataContext Context;
-	Context.Begin = Begin;
-	Context.End = End;
-	return Context;
-}
-
-#define AVSM_INVALID_PIXEL_OFFSET 0x3FFFFFF
-
-struct FAVSMIndirectionData
-{
-	uint PixelOffset;
-	int SampleCount;
-	float BeginX;
-	float EndX;
-};
-
-FAVSMIndirectionData AVSM_CreateIndirectionData(FAVSMIndirectionDataContext Context, uint PixelOffset, int SampleCount)
+FAVSMIndirectionData AVSM_CreateIndirectionData(uint PixelOffset, int SampleCount, float BeginX, float EndX)
 {
 	FAVSMIndirectionData IndirectionData;
 	IndirectionData.PixelOffset = PixelOffset;
 	IndirectionData.SampleCount = SampleCount;
-	IndirectionData.BeginX = Context.Begin.X;
-	IndirectionData.EndX = Context.End.X;
+	IndirectionData.BeginX = BeginX;
+	IndirectionData.EndX = EndX;
 	return IndirectionData;
 }
 
-uint2 AVSM_PackIndirectionData(FAVSMIndirectionData IndirectionData)
+uint4 AVSM_PackIndirectionData(FAVSMIndirectionData IndirectionData)
 {
-	uint2 PackedData;
+	uint4 PackedData = 0;
 	PackedData.x = ((IndirectionData.PixelOffset & AVSM_INVALID_PIXEL_OFFSET) << 6) | (IndirectionData.SampleCount & 0x3f);
-	PackedData.y = (f32tof16(IndirectionData.BeginX) << 16) | f32tof16(IndirectionData.EndX);
+	PackedData.y = asuint(IndirectionData.BeginX);
+	PackedData.z = asuint(IndirectionData.EndX);
 	return PackedData;
 }
 
-FAVSMIndirectionData AVSM_UnpackIndirectionData(uint2 PackedData)
+FAVSMIndirectionData AVSM_UnpackIndirectionData(uint4 PackedData)
 {
 	FAVSMIndirectionData IndirectionData;
 	IndirectionData.PixelOffset = (PackedData.x >> 6) & AVSM_INVALID_PIXEL_OFFSET;
 	IndirectionData.SampleCount = PackedData.x & 0x3f;
-	IndirectionData.BeginX = f16tof32(PackedData.y >> 16);
-	IndirectionData.EndX = f16tof32(PackedData.y);
+	IndirectionData.BeginX = asfloat(PackedData.y);
+	IndirectionData.EndX = asfloat(PackedData.z);
 	return IndirectionData;
 }
 
@@ -132,4 +146,4 @@ int CalcTopLevelSampleCount(int SampleCount)
 {
 	int TopLevelSampleCount = Align4(SampleCount) >> 2;
 	return TopLevelSampleCount;
-}
+}
\ No newline at end of file
diff --git a/Engine/Shaders/Private/HeterogeneousVolumes/HeterogeneousVolumesLiveShadingShadows.usf b/Engine/Shaders/Private/HeterogeneousVolumes/HeterogeneousVolumesLiveShadingShadows.usf
--- a/Engine/Shaders/Private/HeterogeneousVolumes/HeterogeneousVolumesLiveShadingShadows.usf
+++ b/Engine/Shaders/Private/HeterogeneousVolumes/HeterogeneousVolumesLiveShadingShadows.usf
@@ -224,13 +224,13 @@ int FindSampleIndexWithLowestError(
 	float LowestError = 1.0e6;
 	for (int Index = 1; Index < MaxSampleCount - 1; ++Index)
 	{
-		int PrevIndex = AVSM_GetPrevPtr(RWVolumetricShadowLinkedListBuffer[PixelIndex + Index]);
-		int NextIndex = AVSM_GetNextPtr(RWVolumetricShadowLinkedListBuffer[PixelIndex + Index]);
+		int PrevIndex = AVSM_GetLinkedListPrevPtr(RWVolumetricShadowLinkedListBuffer[PixelIndex + Index]);
+		int NextIndex = AVSM_GetLinkedListNextPtr(RWVolumetricShadowLinkedListBuffer[PixelIndex + Index]);
 
 		FAVSMSampleData Data[3] = {
-			AVSM_UnpackSampleData(RWVolumetricShadowLinkedListBuffer[PixelIndex + PrevIndex]),
-			AVSM_UnpackSampleData(RWVolumetricShadowLinkedListBuffer[PixelIndex + Index]),
-			AVSM_UnpackSampleData(RWVolumetricShadowLinkedListBuffer[PixelIndex + NextIndex])
+			AVSM_GetLinkedListSampleData(RWVolumetricShadowLinkedListBuffer[PixelIndex + PrevIndex]),
+			AVSM_GetLinkedListSampleData(RWVolumetricShadowLinkedListBuffer[PixelIndex + Index]),
+			AVSM_GetLinkedListSampleData(RWVolumetricShadowLinkedListBuffer[PixelIndex + NextIndex])
 		};
 
 		float Error = abs(
@@ -248,11 +248,11 @@ int FindSampleIndexWithLowestError(
 	return LowestErrorIndex;
 }
 
-void AddSample(FAVSMSampleData Data, int PixelIndex, inout int PrevSampleIndex, inout int SampleIndex)
+void AddSample(FAVSMSampleData SampleData, int PixelIndex, inout int PrevSampleIndex, inout int SampleIndex)
 {
-	RWVolumetricShadowLinkedListBuffer[PixelIndex + SampleIndex].x = AVSM_PackSampleData(Data);
-	AVSM_SetPrevPtr(RWVolumetricShadowLinkedListBuffer[PixelIndex + SampleIndex], PrevSampleIndex);
-	AVSM_SetNextPtr(RWVolumetricShadowLinkedListBuffer[PixelIndex + SampleIndex], SampleIndex + 1);
+	AVSM_SetLinkedListSampleData(RWVolumetricShadowLinkedListBuffer[PixelIndex + SampleIndex], SampleData);
+	AVSM_SetLinkedListPrevPtr(RWVolumetricShadowLinkedListBuffer[PixelIndex + SampleIndex], PrevSampleIndex);
+	AVSM_SetLinkedListNextPtr(RWVolumetricShadowLinkedListBuffer[PixelIndex + SampleIndex], SampleIndex + 1);
 
 	PrevSampleIndex = SampleIndex;
 	SampleIndex = SampleIndex + 1;
@@ -268,21 +268,21 @@ void ReplaceSample(
 {
 	// Remove 
 	{
-		int PrevIndex = AVSM_GetPrevPtr(RWVolumetricShadowLinkedListBuffer[PixelIndex + SampleIndex]);
-		int NextIndex = AVSM_GetNextPtr(RWVolumetricShadowLinkedListBuffer[PixelIndex + SampleIndex]);
+		int PrevIndex = AVSM_GetLinkedListPrevPtr(RWVolumetricShadowLinkedListBuffer[PixelIndex + SampleIndex]);
+		int NextIndex = AVSM_GetLinkedListNextPtr(RWVolumetricShadowLinkedListBuffer[PixelIndex + SampleIndex]);
 
-		AVSM_SetNextPtr(RWVolumetricShadowLinkedListBuffer[PixelIndex + PrevIndex], NextIndex);
-		AVSM_SetPrevPtr(RWVolumetricShadowLinkedListBuffer[PixelIndex + NextIndex], PrevIndex);
+		AVSM_SetLinkedListNextPtr(RWVolumetricShadowLinkedListBuffer[PixelIndex + PrevIndex], NextIndex);
+		AVSM_SetLinkedListPrevPtr(RWVolumetricShadowLinkedListBuffer[PixelIndex + NextIndex], PrevIndex);
 	}
 
 	// Replace
-	RWVolumetricShadowLinkedListBuffer[PixelIndex + SampleIndex].x = AVSM_PackSampleData(SampleData);
-	AVSM_SetPrevPtr(RWVolumetricShadowLinkedListBuffer[PixelIndex + SampleIndex], PrevPtrIndex);
-	AVSM_SetNextPtr(RWVolumetricShadowLinkedListBuffer[PixelIndex + SampleIndex], NextPtrIndex);
+	AVSM_SetLinkedListSampleData(RWVolumetricShadowLinkedListBuffer[PixelIndex + SampleIndex], SampleData);
+	AVSM_SetLinkedListPrevPtr(RWVolumetricShadowLinkedListBuffer[PixelIndex + SampleIndex], PrevPtrIndex);
+	AVSM_SetLinkedListNextPtr(RWVolumetricShadowLinkedListBuffer[PixelIndex + SampleIndex], NextPtrIndex);
 
 	// Update neighborhood
-	AVSM_SetPrevPtr(RWVolumetricShadowLinkedListBuffer[PixelIndex + NextPtrIndex], SampleIndex);
-	AVSM_SetNextPtr(RWVolumetricShadowLinkedListBuffer[PixelIndex + PrevPtrIndex], SampleIndex);
+	AVSM_SetLinkedListPrevPtr(RWVolumetricShadowLinkedListBuffer[PixelIndex + NextPtrIndex], SampleIndex);
+	AVSM_SetLinkedListNextPtr(RWVolumetricShadowLinkedListBuffer[PixelIndex + PrevPtrIndex], SampleIndex);
 }
 
 void AddOrReplaceSample(FAVSMSampleData Data, int PixelIndex, inout int PrevSampleIndex, inout int SampleIndex)
@@ -297,16 +297,16 @@ void AddOrReplaceSample(FAVSMSampleData Data, int PixelIndex, inout int PrevSamp
 	{
 #if USE_AVSM_COMPRESSION
 		// Insert fictious sample at MaxSampleCount - 1 to compute error
-		RWVolumetricShadowLinkedListBuffer[PixelIndex + MaxSampleCount - 1].x = AVSM_PackSampleData(Data);
-		AVSM_SetNextPtr(RWVolumetricShadowLinkedListBuffer[PixelIndex + PrevSampleIndex], MaxSampleCount - 1);
-		AVSM_SetPrevPtr(RWVolumetricShadowLinkedListBuffer[PixelIndex + MaxSampleCount - 1], PrevSampleIndex);
+		AVSM_SetLinkedListSampleData(RWVolumetricShadowLinkedListBuffer[PixelIndex + MaxSampleCount - 1], Data);
+		AVSM_SetLinkedListNextPtr(RWVolumetricShadowLinkedListBuffer[PixelIndex + PrevSampleIndex], MaxSampleCount - 1);
+		AVSM_SetLinkedListPrevPtr(RWVolumetricShadowLinkedListBuffer[PixelIndex + MaxSampleCount - 1], PrevSampleIndex);
 
 		int LowestErrorIndex = FindSampleIndexWithLowestError(PixelIndex);
 		if (LowestErrorIndex > 0)
 		{
 			if (LowestErrorIndex == PrevSampleIndex)
 			{
-				RWVolumetricShadowLinkedListBuffer[PixelIndex + PrevSampleIndex].x = AVSM_PackSampleData(Data);
+				AVSM_SetLinkedListSampleData(RWVolumetricShadowLinkedListBuffer[PixelIndex + PrevSampleIndex], Data);
 			}
 			else
 			{
@@ -316,8 +316,8 @@ void AddOrReplaceSample(FAVSMSampleData Data, int PixelIndex, inout int PrevSamp
 			PrevSampleIndex = LowestErrorIndex;
 		}
 #else
-		RWVolumetricShadowLinkedListBuffer[PixelIndex + MaxSampleCount - 1].x = AVSM_PackSampleData(Data);
-		AVSM_SetPrevPtr(RWVolumetricShadowLinkedListBuffer[PixelIndex + MaxSampleCount - 1], PrevSampleIndex);
+		AVSM_SetLinkedListSampleData(RWVolumetricShadowLinkedListBuffer[PixelIndex + MaxSampleCount - 1], Data);
+		AVSM_SetLinkedListPrevPtr(RWVolumetricShadowLinkedListBuffer[PixelIndex + MaxSampleCount - 1], PrevSampleIndex);
 
 		PrevSampleIndex = MaxSampleCount - 1;
 #endif
@@ -491,20 +491,20 @@ void RenderVolumetricShadowMapForLightWithLiveShadingCS(
 		// Stitch up the previous entry
 		if (PrevElementIndex != AVSM_NULL_PTR)
 		{
-			AVSM_SetNextPtr(RWVolumetricShadowLinkedListBuffer[PixelIndex + PrevElementIndex], ElementIndex);
+			AVSM_SetLinkedListNextPtr(RWVolumetricShadowLinkedListBuffer[PixelIndex + PrevElementIndex], ElementIndex);
 
 		}
 
 		// Last Entry
-		RWVolumetricShadowLinkedListBuffer[PixelIndex + ElementIndex].x = AVSM_PackSampleData(WorldHitT, Transmittance);
-		AVSM_SetPrevPtr(RWVolumetricShadowLinkedListBuffer[PixelIndex + ElementIndex], PrevElementIndex);
-		AVSM_SetNextPtr(RWVolumetricShadowLinkedListBuffer[PixelIndex + ElementIndex], AVSM_NULL_PTR);
+		AVSM_SetLinkedListSampleData(RWVolumetricShadowLinkedListBuffer[PixelIndex + ElementIndex], AVSM_CreateSampleData(WorldHitT, Transmittance));
+		AVSM_SetLinkedListPrevPtr(RWVolumetricShadowLinkedListBuffer[PixelIndex + ElementIndex], PrevElementIndex);
+		AVSM_SetLinkedListNextPtr(RWVolumetricShadowLinkedListBuffer[PixelIndex + ElementIndex], AVSM_NULL_PTR);
 	}
 	else
 	{
-		RWVolumetricShadowLinkedListBuffer[PixelIndex].x = AVSM_PackSampleData(0.0, Transmittance);
-		AVSM_SetPrevPtr(RWVolumetricShadowLinkedListBuffer[PixelIndex], AVSM_NULL_PTR);
-		AVSM_SetNextPtr(RWVolumetricShadowLinkedListBuffer[PixelIndex], AVSM_NULL_PTR);
+		AVSM_SetLinkedListSampleData(RWVolumetricShadowLinkedListBuffer[PixelIndex], AVSM_CreateSampleData(0.0, Transmittance));
+		AVSM_SetLinkedListPrevPtr(RWVolumetricShadowLinkedListBuffer[PixelIndex], AVSM_NULL_PTR);
+		AVSM_SetLinkedListNextPtr(RWVolumetricShadowLinkedListBuffer[PixelIndex], AVSM_NULL_PTR);
 	}
 
 	if (Face == ShadowDebugTweak)
diff --git a/Engine/Shaders/Private/HeterogeneousVolumes/HeterogeneousVolumesVoxelGridShadows.usf b/Engine/Shaders/Private/HeterogeneousVolumes/HeterogeneousVolumesVoxelGridShadows.usf
--- a/Engine/Shaders/Private/HeterogeneousVolumes/HeterogeneousVolumesVoxelGridShadows.usf
+++ b/Engine/Shaders/Private/HeterogeneousVolumes/HeterogeneousVolumesVoxelGridShadows.usf
@@ -244,22 +244,22 @@ uint GetLinearIndex(uint Face, uint2 PixelCoord)
 	return LinearIndex;
 }
 
-int FindSampleIndexWithLowestError(
-	int PixelIndex
+uint FindSampleIndexWithLowestError(
+	uint PixelIndex
 )
 {
 	// Find sample, whose removal would mark the lowest overall geometric error
-	int LowestErrorIndex = 0;
+	uint LowestErrorIndex = 0;
 	float LowestError = 1.0e6;
-	for (int Index = 1; Index < MaxSampleCount - 1; ++Index)
+	for (uint Index = 1; Index < MaxSampleCount - 1; ++Index)
 	{
-		int PrevIndex = AVSM_GetPrevPtr(RWVolumetricShadowLinkedListBuffer[PixelIndex + Index]);
-		int NextIndex = AVSM_GetNextPtr(RWVolumetricShadowLinkedListBuffer[PixelIndex + Index]);
+		uint PrevIndex = AVSM_GetLinkedListPrevPtr(RWVolumetricShadowLinkedListBuffer[PixelIndex + Index]);
+		uint NextIndex = AVSM_GetLinkedListNextPtr(RWVolumetricShadowLinkedListBuffer[PixelIndex + Index]);
 
 		FAVSMSampleData Data[3] = {
-			AVSM_UnpackSampleData(RWVolumetricShadowLinkedListBuffer[PixelIndex + PrevIndex]),
-			AVSM_UnpackSampleData(RWVolumetricShadowLinkedListBuffer[PixelIndex + Index]),
-			AVSM_UnpackSampleData(RWVolumetricShadowLinkedListBuffer[PixelIndex + NextIndex])
+			AVSM_GetLinkedListSampleData(RWVolumetricShadowLinkedListBuffer[PixelIndex + PrevIndex]),
+			AVSM_GetLinkedListSampleData(RWVolumetricShadowLinkedListBuffer[PixelIndex + Index]),
+			AVSM_GetLinkedListSampleData(RWVolumetricShadowLinkedListBuffer[PixelIndex + NextIndex])
 		};
 
 		float Error = abs(
@@ -279,9 +279,9 @@ int FindSampleIndexWithLowestError(
 
 void AddSample(FAVSMSampleData Data, int PixelIndex, inout int PrevSampleIndex, inout int SampleIndex)
 {
-	RWVolumetricShadowLinkedListBuffer[PixelIndex + SampleIndex].x = AVSM_PackSampleData(Data);
-	AVSM_SetPrevPtr(RWVolumetricShadowLinkedListBuffer[PixelIndex + SampleIndex], PrevSampleIndex);
-	AVSM_SetNextPtr(RWVolumetricShadowLinkedListBuffer[PixelIndex + SampleIndex], SampleIndex + 1);
+	AVSM_SetLinkedListSampleData(RWVolumetricShadowLinkedListBuffer[PixelIndex + SampleIndex], Data);
+	AVSM_SetLinkedListPrevPtr(RWVolumetricShadowLinkedListBuffer[PixelIndex + SampleIndex], PrevSampleIndex);
+	AVSM_SetLinkedListNextPtr(RWVolumetricShadowLinkedListBuffer[PixelIndex + SampleIndex], SampleIndex + 1);
 
 	PrevSampleIndex = SampleIndex;
 	SampleIndex = SampleIndex + 1;
@@ -297,21 +297,21 @@ void ReplaceSample(
 {
 	// Remove 
 	{
-		int PrevIndex = AVSM_GetPrevPtr(RWVolumetricShadowLinkedListBuffer[PixelIndex + SampleIndex]);
-		int NextIndex = AVSM_GetNextPtr(RWVolumetricShadowLinkedListBuffer[PixelIndex + SampleIndex]);
+		int PrevIndex = AVSM_GetLinkedListPrevPtr(RWVolumetricShadowLinkedListBuffer[PixelIndex + SampleIndex]);
+		int NextIndex = AVSM_GetLinkedListNextPtr(RWVolumetricShadowLinkedListBuffer[PixelIndex + SampleIndex]);
 
-		AVSM_SetNextPtr(RWVolumetricShadowLinkedListBuffer[PixelIndex + PrevIndex], NextIndex);
-		AVSM_SetPrevPtr(RWVolumetricShadowLinkedListBuffer[PixelIndex + NextIndex], PrevIndex);
+		AVSM_SetLinkedListNextPtr(RWVolumetricShadowLinkedListBuffer[PixelIndex + PrevIndex], NextIndex);
+		AVSM_SetLinkedListPrevPtr(RWVolumetricShadowLinkedListBuffer[PixelIndex + NextIndex], PrevIndex);
 	}
 
 	// Replace
-	RWVolumetricShadowLinkedListBuffer[PixelIndex + SampleIndex].x = AVSM_PackSampleData(SampleData);
-	AVSM_SetPrevPtr(RWVolumetricShadowLinkedListBuffer[PixelIndex + SampleIndex], PrevPtrIndex);
-	AVSM_SetNextPtr(RWVolumetricShadowLinkedListBuffer[PixelIndex + SampleIndex], NextPtrIndex);
+	AVSM_SetLinkedListSampleData(RWVolumetricShadowLinkedListBuffer[PixelIndex + SampleIndex], SampleData);
+	AVSM_SetLinkedListPrevPtr(RWVolumetricShadowLinkedListBuffer[PixelIndex + SampleIndex], PrevPtrIndex);
+	AVSM_SetLinkedListNextPtr(RWVolumetricShadowLinkedListBuffer[PixelIndex + SampleIndex], NextPtrIndex);
 
 	// Update neighborhood
-	AVSM_SetNextPtr(RWVolumetricShadowLinkedListBuffer[PixelIndex + PrevPtrIndex], SampleIndex);
-	AVSM_SetPrevPtr(RWVolumetricShadowLinkedListBuffer[PixelIndex + NextPtrIndex], SampleIndex);
+	AVSM_SetLinkedListNextPtr(RWVolumetricShadowLinkedListBuffer[PixelIndex + PrevPtrIndex], SampleIndex);
+	AVSM_SetLinkedListPrevPtr(RWVolumetricShadowLinkedListBuffer[PixelIndex + NextPtrIndex], SampleIndex);
 }
 
 void AddOrReplaceSample(FAVSMSampleData Data, int PixelIndex, inout int PrevSampleIndex, inout int SampleIndex)
@@ -326,16 +326,16 @@ void AddOrReplaceSample(FAVSMSampleData Data, int PixelIndex, inout int PrevSamp
 	{
 #if USE_AVSM_COMPRESSION
 		// Insert fictious sample at MaxSampleCount - 1 to compute error
-		RWVolumetricShadowLinkedListBuffer[PixelIndex + MaxSampleCount - 1].x = AVSM_PackSampleData(Data);
-		AVSM_SetNextPtr(RWVolumetricShadowLinkedListBuffer[PixelIndex + PrevSampleIndex], MaxSampleCount - 1);
-		AVSM_SetPrevPtr(RWVolumetricShadowLinkedListBuffer[PixelIndex + MaxSampleCount - 1], PrevSampleIndex);
+		AVSM_SetLinkedListSampleData(RWVolumetricShadowLinkedListBuffer[PixelIndex + MaxSampleCount - 1], Data);
+		AVSM_SetLinkedListNextPtr(RWVolumetricShadowLinkedListBuffer[PixelIndex + PrevSampleIndex], MaxSampleCount - 1);
+		AVSM_SetLinkedListPrevPtr(RWVolumetricShadowLinkedListBuffer[PixelIndex + MaxSampleCount - 1], PrevSampleIndex);
 
 		int LowestErrorIndex = FindSampleIndexWithLowestError(PixelIndex);
 		if (LowestErrorIndex > 0)
 		{
 			if (LowestErrorIndex == PrevSampleIndex)
 			{
-				RWVolumetricShadowLinkedListBuffer[PixelIndex + PrevSampleIndex].x = AVSM_PackSampleData(Data);
+				AVSM_SetLinkedListSampleData(RWVolumetricShadowLinkedListBuffer[PixelIndex + PrevSampleIndex], Data);
 			}
 			else
 			{
@@ -345,8 +345,8 @@ void AddOrReplaceSample(FAVSMSampleData Data, int PixelIndex, inout int PrevSamp
 			PrevSampleIndex = LowestErrorIndex;
 		}
 #else
-		RWVolumetricShadowLinkedListBuffer[PixelIndex + MaxSampleCount - 1].x = AVSM_PackSampleData(Data);
-		AVSM_SetPrevPtr(RWVolumetricShadowLinkedListBuffer[PixelIndex + MaxSampleCount - 1], PrevSampleIndex);
+		AVSM_SetLinkedListSampleData(RWVolumetricShadowLinkedListBuffer[PixelIndex + MaxSampleCount - 1], Data);
+		AVSM_SetLinkedListPrevPtr(RWVolumetricShadowLinkedListBuffer[PixelIndex + MaxSampleCount - 1], PrevSampleIndex);
 
 		PrevSampleIndex = MaxSampleCount - 1;
 #endif
@@ -408,7 +408,7 @@ void RenderVolumetricShadowMapForLightWithVoxelGridCS(
 	float LocalHitT = RayHitT.x;
 	float WorldHitT = RayHitT.x;
 	float3 Transmittance = 1;
-	int PixelIndex = LinearIndex * MaxSampleCount;
+	uint PixelIndex = LinearIndex * MaxSampleCount;
 
 	float HitSpan = RayHitT.y - RayHitT.x;
 	if (HitSpan > 0.0)
@@ -441,8 +441,8 @@ void RenderVolumetricShadowMapForLightWithVoxelGridCS(
 
 		float3 ExpectedExtinction = 0.0;
 		float3 ExpectedTransmittance = 1.0;
-		int PrevElementIndex = AVSM_NULL_PTR;
-		int ElementIndex = 0;
+		uint PrevElementIndex = AVSM_NULL_PTR;
+		uint ElementIndex = 0;
 
 		for (uint StepIndex = 0; StepIndex < StepCount; ++StepIndex)
 		{
@@ -494,19 +494,19 @@ void RenderVolumetricShadowMapForLightWithVoxelGridCS(
 		// Stitch up the previous entry
 		if (PrevElementIndex != AVSM_NULL_PTR)
 		{
-			AVSM_SetNextPtr(RWVolumetricShadowLinkedListBuffer[PixelIndex + PrevElementIndex], ElementIndex);
+			AVSM_SetLinkedListNextPtr(RWVolumetricShadowLinkedListBuffer[PixelIndex + PrevElementIndex], ElementIndex);
 
 		}
 
-		RWVolumetricShadowLinkedListBuffer[PixelIndex + ElementIndex].x = AVSM_PackSampleData(WorldHitT, Transmittance);
-		AVSM_SetPrevPtr(RWVolumetricShadowLinkedListBuffer[PixelIndex + ElementIndex], PrevElementIndex);
-		AVSM_SetNextPtr(RWVolumetricShadowLinkedListBuffer[PixelIndex + ElementIndex], AVSM_NULL_PTR);
+		AVSM_SetLinkedListSampleData(RWVolumetricShadowLinkedListBuffer[PixelIndex + ElementIndex], AVSM_CreateSampleData(WorldHitT, Transmittance));
+		AVSM_SetLinkedListPrevPtr(RWVolumetricShadowLinkedListBuffer[PixelIndex + ElementIndex], PrevElementIndex);
+		AVSM_SetLinkedListNextPtr(RWVolumetricShadowLinkedListBuffer[PixelIndex + ElementIndex], AVSM_NULL_PTR);
 	}
 	else
 	{
-		RWVolumetricShadowLinkedListBuffer[PixelIndex].x = AVSM_PackSampleData(0.0, Transmittance);
-		AVSM_SetPrevPtr(RWVolumetricShadowLinkedListBuffer[PixelIndex], AVSM_NULL_PTR);
-		AVSM_SetNextPtr(RWVolumetricShadowLinkedListBuffer[PixelIndex], AVSM_NULL_PTR);
+		AVSM_SetLinkedListSampleData(RWVolumetricShadowLinkedListBuffer[PixelIndex], AVSM_CreateSampleData(0.0, Transmittance));
+		AVSM_SetLinkedListPrevPtr(RWVolumetricShadowLinkedListBuffer[PixelIndex], AVSM_NULL_PTR);
+		AVSM_SetLinkedListNextPtr(RWVolumetricShadowLinkedListBuffer[PixelIndex], AVSM_NULL_PTR);
 	}
 
 	if (Face == ShadowDebugTweak)
@@ -518,7 +518,7 @@ void RenderVolumetricShadowMapForLightWithVoxelGridCS(
 StructuredBuffer<uint2> VolumetricShadowLinkedListBuffer;
 
 RWBuffer<uint> RWVolumetricShadowIndirectionAllocatorBuffer;
-RWStructuredBuffer<uint2> RWVolumetricShadowIndirectionBuffer;
+RWStructuredBuffer<uint4> RWVolumetricShadowIndirectionBuffer;
 RWStructuredBuffer<uint> RWVolumetricShadowTransmittanceBuffer;
 
 groupshared int GSWarpSampleCount[THREADGROUP_SIZE_2D * THREADGROUP_SIZE_2D];
@@ -571,18 +571,19 @@ void CompressVolumetricShadowMapCS(
 
 	// Calculate sample count
 	int SampleCount = 0;
-	FAVSMSampleData DataBegin, DataEnd;
+	FAVSMSampleData DataBegin = AVSM_CreateSampleData(0.0, 1.0);
+	FAVSMSampleData DataEnd = DataBegin;
 	if (all(PixelCoord < ShadowResolution))
 	{
-		DataBegin = AVSM_UnpackSampleData(VolumetricShadowLinkedListBuffer[PixelIndex]);
+		DataBegin = AVSM_GetLinkedListSampleData(VolumetricShadowLinkedListBuffer[PixelIndex]);
 		DataEnd = DataBegin;
-		int NextPtr = AVSM_GetNextPtr(VolumetricShadowLinkedListBuffer[PixelIndex]);
+		int NextPtr = AVSM_GetLinkedListNextPtr(VolumetricShadowLinkedListBuffer[PixelIndex]);
 		SampleCount = (NextPtr == AVSM_NULL_PTR) ? 0 : 1;
 		while (NextPtr != AVSM_NULL_PTR)
 		{
 			SampleCount++;
-			DataEnd = AVSM_UnpackSampleData(VolumetricShadowLinkedListBuffer[PixelIndex + NextPtr]);
-			NextPtr = AVSM_GetNextPtr(VolumetricShadowLinkedListBuffer[PixelIndex + NextPtr]);
+			DataEnd = AVSM_GetLinkedListSampleData(VolumetricShadowLinkedListBuffer[PixelIndex + NextPtr]);
+			NextPtr = AVSM_GetLinkedListNextPtr(VolumetricShadowLinkedListBuffer[PixelIndex + NextPtr]);
 		}
 	}
 
@@ -612,48 +613,49 @@ void CompressVolumetricShadowMapCS(
 	uint TopLevelOffset = GSWarpOffset + GSWarpPrefixSum[GroupIndex];
 	uint BottomLevelOffset = GSWarpOffset + GSWarpPrefixSum[GroupIndex] + Align4(TopLevelSampleCount);
 
+	// Write indirection entry
+	FAVSMIndirectionData IndirectionData = AVSM_CreateIndirectionData(
+		GSWarpOffset + GSWarpPrefixSum[GroupIndex],
+		SampleCount,
+		DataBegin.X,
+		DataEnd.X
+	);
+
+	if (all(PixelCoord < ShadowResolution))
+	{
+		RWVolumetricShadowIndirectionBuffer[LinearIndex] = AVSM_PackIndirectionData(IndirectionData);
+	}
+
 	// Pack and fill data per thread
 	FAVSMSampleData ShadowData = AVSM_CreateSampleData(0.0, 1.0);
 	int NextPtr = 0;
 	for (int SampleIndex = 0; SampleIndex < SampleCount; ++SampleIndex)
 	{
-		ShadowData = AVSM_UnpackSampleData(VolumetricShadowLinkedListBuffer[PixelIndex + NextPtr]);
-		NextPtr = AVSM_GetNextPtr(VolumetricShadowLinkedListBuffer[PixelIndex + NextPtr]);
+		ShadowData = AVSM_GetLinkedListSampleData(VolumetricShadowLinkedListBuffer[PixelIndex + NextPtr]);
+		NextPtr = AVSM_GetLinkedListNextPtr(VolumetricShadowLinkedListBuffer[PixelIndex + NextPtr]);
 
 		// Store indexed list offset by 1 element
 		int BottomLevelIndex = max(SampleIndex - 1, 0);
-		RWVolumetricShadowTransmittanceBuffer[BottomLevelOffset + BottomLevelIndex] = AVSM_PackSampleData(ShadowData);
+		RWVolumetricShadowTransmittanceBuffer[BottomLevelOffset + BottomLevelIndex] = AVSM_PackSampleData(ShadowData, IndirectionData);
 
 		// Mark the current sample as an acceleration sample at the top-level
 		if (SampleIndex % 4 == 0)
 		{
 			uint TopLevelIndex = SampleIndex / 4;
-			RWVolumetricShadowTransmittanceBuffer[TopLevelOffset + TopLevelIndex] = AVSM_PackSampleData(ShadowData);
+			RWVolumetricShadowTransmittanceBuffer[TopLevelOffset + TopLevelIndex] = AVSM_PackSampleData(ShadowData, IndirectionData);
 		}
 	}
 
 	// Pad remainder
 	for (int BottomLevelIndex = max(SampleCount - 1, 0); BottomLevelIndex < Align4(SampleCount); ++BottomLevelIndex)
 	{
-		RWVolumetricShadowTransmittanceBuffer[BottomLevelOffset + BottomLevelIndex] = AVSM_PackSampleData(AVSM_CreateSampleData(MaxHalfFloat, ShadowData.Tau));
+		RWVolumetricShadowTransmittanceBuffer[BottomLevelOffset + BottomLevelIndex] = AVSM_PackSampleData(AVSM_CreateSampleData(ShadowData.X, ShadowData.Tau), IndirectionData);
 	}
 
 	// Pad remainder
 	for (int TopLevelIndex = TopLevelSampleCount; TopLevelIndex < Align4(TopLevelSampleCount); ++TopLevelIndex)
 	{
-		RWVolumetricShadowTransmittanceBuffer[TopLevelOffset + TopLevelIndex] = AVSM_PackSampleData(AVSM_CreateSampleData(MaxHalfFloat, ShadowData.Tau));
-	}
-
-	// Write indirection entry
-	if (all(PixelCoord < ShadowResolution))
-	{
-		FAVSMIndirectionDataContext Context = AVSM_CreateIndirectionDataContext(DataBegin, DataEnd);
-		FAVSMIndirectionData IndirectionData = AVSM_CreateIndirectionData(
-			Context,
-			GSWarpOffset + GSWarpPrefixSum[GroupIndex],
-			SampleCount
-		);
-		RWVolumetricShadowIndirectionBuffer[LinearIndex] = AVSM_PackIndirectionData(IndirectionData);
+		RWVolumetricShadowTransmittanceBuffer[TopLevelOffset + TopLevelIndex] = AVSM_PackSampleData(AVSM_CreateSampleData(ShadowData.X, ShadowData.Tau), IndirectionData);
 	}
 }
 
@@ -688,8 +690,8 @@ void CombineVolumetricShadowMapsCS(
 	// Sorted sample insertion
 	while ((IndexPtr0 != AVSM_NULL_PTR) && (IndexPtr1 != AVSM_NULL_PTR))
 	{
-		FAVSMSampleData Data0 = AVSM_UnpackSampleData(VolumetricShadowLinkedListBuffer0[PixelIndex + IndexPtr0]);
-		FAVSMSampleData Data1 = AVSM_UnpackSampleData(VolumetricShadowLinkedListBuffer1[PixelIndex + IndexPtr1]);
+		FAVSMSampleData Data0 = AVSM_GetLinkedListSampleData(VolumetricShadowLinkedListBuffer0[PixelIndex + IndexPtr0]);
+		FAVSMSampleData Data1 = AVSM_GetLinkedListSampleData(VolumetricShadowLinkedListBuffer1[PixelIndex + IndexPtr1]);
 
 		FAVSMSampleData DataToInsert;
 		if (Data0.X < Data1.X)
@@ -698,7 +700,7 @@ void CombineVolumetricShadowMapsCS(
 			DataToInsert.Tau *= Transmittance1;
 			Transmittance0 = Data0.Tau;
 
-			IndexPtr0 = AVSM_GetNextPtr(VolumetricShadowLinkedListBuffer0[PixelIndex + IndexPtr0]);
+			IndexPtr0 = AVSM_GetLinkedListNextPtr(VolumetricShadowLinkedListBuffer0[PixelIndex + IndexPtr0]);
 		}
 		else
 		{
@@ -706,7 +708,7 @@ void CombineVolumetricShadowMapsCS(
 			DataToInsert.Tau *= Transmittance0;
 			Transmittance1 = Data1.Tau;
 
-			IndexPtr1 = AVSM_GetNextPtr(VolumetricShadowLinkedListBuffer1[PixelIndex + IndexPtr1]);
+			IndexPtr1 = AVSM_GetLinkedListNextPtr(VolumetricShadowLinkedListBuffer1[PixelIndex + IndexPtr1]);
 		}
 
 		AddOrReplaceSample(DataToInsert, PixelIndex, PrevSampleIndex, SampleIndex);
@@ -715,26 +717,26 @@ void CombineVolumetricShadowMapsCS(
 	// Sample0 residual
 	while (IndexPtr0 != AVSM_NULL_PTR)
 	{
-		FAVSMSampleData DataToInsert = AVSM_UnpackSampleData(VolumetricShadowLinkedListBuffer0[PixelIndex + IndexPtr0]);
+		FAVSMSampleData DataToInsert = AVSM_GetLinkedListSampleData(VolumetricShadowLinkedListBuffer0[PixelIndex + IndexPtr0]);
 		DataToInsert.Tau *= Transmittance1;
 
-		IndexPtr0 = AVSM_GetNextPtr(VolumetricShadowLinkedListBuffer0[PixelIndex + IndexPtr0]);
+		IndexPtr0 = AVSM_GetLinkedListNextPtr(VolumetricShadowLinkedListBuffer0[PixelIndex + IndexPtr0]);
 		AddOrReplaceSample(DataToInsert, PixelIndex, PrevSampleIndex, SampleIndex);
 	}
 
 	// Sample1 residual
 	while ((IndexPtr1 != AVSM_NULL_PTR))
 	{
-		FAVSMSampleData DataToInsert = AVSM_UnpackSampleData(VolumetricShadowLinkedListBuffer1[PixelIndex + IndexPtr1]);
+		FAVSMSampleData DataToInsert = AVSM_GetLinkedListSampleData(VolumetricShadowLinkedListBuffer1[PixelIndex + IndexPtr1]);
 		DataToInsert.Tau *= Transmittance0;
 
-		IndexPtr1 = AVSM_GetNextPtr(VolumetricShadowLinkedListBuffer1[PixelIndex + IndexPtr1]);
+		IndexPtr1 = AVSM_GetLinkedListNextPtr(VolumetricShadowLinkedListBuffer1[PixelIndex + IndexPtr1]);
 		AddOrReplaceSample(DataToInsert, PixelIndex, PrevSampleIndex, SampleIndex);
 	}
 
 	// Mark the final sample as being terminal
 	if (PrevSampleIndex != AVSM_NULL_PTR)
 	{
-		AVSM_SetNextPtr(RWVolumetricShadowLinkedListBuffer[PixelIndex + PrevSampleIndex], AVSM_NULL_PTR);
+		AVSM_SetLinkedListNextPtr(RWVolumetricShadowLinkedListBuffer[PixelIndex + PrevSampleIndex], AVSM_NULL_PTR);
 	}
 }
diff --git a/Engine/Source/Runtime/Renderer/Private/HeterogeneousVolumes/HeterogeneousVolumes.h b/Engine/Source/Runtime/Renderer/Private/HeterogeneousVolumes/HeterogeneousVolumes.h
--- a/Engine/Source/Runtime/Renderer/Private/HeterogeneousVolumes/HeterogeneousVolumes.h
+++ b/Engine/Source/Runtime/Renderer/Private/HeterogeneousVolumes/HeterogeneousVolumes.h
@@ -411,7 +411,7 @@ struct FAVSMLinkedListPackedData
 
 struct FAVSMIndirectionPackedData
 {
-	uint32 Data[2];
+	uint32 Data[4];
 };
 
 struct FAVSMSamplePackedData
@@ -431,7 +431,7 @@ BEGIN_UNIFORM_BUFFER_STRUCT(FAdaptiveVolumetricShadowMapUniformBufferParameters,
 	SHADER_PARAMETER(int32, bIsDirectionalLight)
 
 	SHADER_PARAMETER_RDG_BUFFER_SRV(StructuredBuffer<uint2>, LinkedListBuffer)
-	SHADER_PARAMETER_RDG_BUFFER_SRV(StructuredBuffer<uint2>, IndirectionBuffer)
+	SHADER_PARAMETER_RDG_BUFFER_SRV(StructuredBuffer<uint4>, IndirectionBuffer)
 	SHADER_PARAMETER_RDG_BUFFER_SRV(StructuredBuffer<uint>, SampleBuffer)
 END_UNIFORM_BUFFER_STRUCT()
 
@@ -447,7 +447,7 @@ BEGIN_SHADER_PARAMETER_STRUCT(FAdaptiveVolumetricShadowMapParameters, RENDERER_A
 	SHADER_PARAMETER(int32, bIsDirectionalLight)
 
 	SHADER_PARAMETER_RDG_BUFFER_SRV(StructuredBuffer<uint2>, LinkedListBuffer)
-	SHADER_PARAMETER_RDG_BUFFER_SRV(StructuredBuffer<uint2>, IndirectionBuffer)
+	SHADER_PARAMETER_RDG_BUFFER_SRV(StructuredBuffer<uint4>, IndirectionBuffer)
 	SHADER_PARAMETER_RDG_BUFFER_SRV(StructuredBuffer<uint>, SampleBuffer)
 END_SHADER_PARAMETER_STRUCT()
 
@@ -849,4 +849,4 @@ void RenderSingleScatteringWithPreshadingHardwareRayTracing(
 	FRDGTextureRef& HeterogeneousVolumeTexture
 );
 
-#endif // RHI_RAYTRACING
\ No newline at end of file
+#endif // RHI_RAYTRACING
diff --git a/Engine/Source/Runtime/Renderer/Private/HeterogeneousVolumes/HeterogeneousVolumesVoxelGridPipeline.cpp b/Engine/Source/Runtime/Renderer/Private/HeterogeneousVolumes/HeterogeneousVolumesVoxelGridPipeline.cpp
--- a/Engine/Source/Runtime/Renderer/Private/HeterogeneousVolumes/HeterogeneousVolumesVoxelGridPipeline.cpp
+++ b/Engine/Source/Runtime/Renderer/Private/HeterogeneousVolumes/HeterogeneousVolumesVoxelGridPipeline.cpp
@@ -2990,7 +2990,7 @@ class FRenderVolumetricShadowMapForLightWithVoxelGridCS : public FGlobalShader
 
 		// Output
 		SHADER_PARAMETER_RDG_BUFFER_UAV(RWBuffer<int>, RWVolumetricShadowLinkedListAllocatorBuffer)
-		SHADER_PARAMETER_RDG_BUFFER_UAV(RWStructuredBuffer<int2>, RWVolumetricShadowLinkedListBuffer)
+		SHADER_PARAMETER_RDG_BUFFER_UAV(RWStructuredBuffer<FAVSMLinkedListPackedData>, RWVolumetricShadowLinkedListBuffer)
 		SHADER_PARAMETER_RDG_TEXTURE_UAV(RWTexture2D<float4>, RWBeerShadowMapTexture)
 
 		// Debug
@@ -3042,7 +3042,7 @@ class FCompressVolumetricShadowMapCS : public FGlobalShader
 		// Input
 		SHADER_PARAMETER(FIntPoint, ShadowResolution)
 		SHADER_PARAMETER(int, MaxSampleCount)
-		SHADER_PARAMETER_RDG_BUFFER_SRV(StructuredBuffer<uint2>, VolumetricShadowLinkedListBuffer)
+		SHADER_PARAMETER_RDG_BUFFER_SRV(StructuredBuffer<FAVSMLinkedListPackedData>, VolumetricShadowLinkedListBuffer)
 
 		// Output
 		SHADER_PARAMETER_RDG_BUFFER_UAV(RWBuffer<int>, RWVolumetricShadowIndirectionAllocatorBuffer)
@@ -3093,8 +3093,8 @@ class FCombineVolumetricShadowMapsCS : public FGlobalShader
 	BEGIN_SHADER_PARAMETER_STRUCT(FParameters, )
 		// Input
 		SHADER_PARAMETER(FIntPoint, ShadowResolution)
-		SHADER_PARAMETER_RDG_BUFFER_SRV(StructuredBuffer<uint2>, VolumetricShadowLinkedListBuffer0)
-		SHADER_PARAMETER_RDG_BUFFER_SRV(StructuredBuffer<uint2>, VolumetricShadowLinkedListBuffer1)
+		SHADER_PARAMETER_RDG_BUFFER_SRV(StructuredBuffer<FAVSMLinkedListPackedData>, VolumetricShadowLinkedListBuffer0)
+		SHADER_PARAMETER_RDG_BUFFER_SRV(StructuredBuffer<FAVSMLinkedListPackedData>, VolumetricShadowLinkedListBuffer1)
 
 		// Volumetric Shadow Map data
 		SHADER_PARAMETER(int, MaxSampleCount)
@@ -3102,7 +3102,7 @@ class FCombineVolumetricShadowMapsCS : public FGlobalShader
 		SHADER_PARAMETER(float, RelativeErrorThreshold)
 
 		// Output
-		SHADER_PARAMETER_RDG_BUFFER_UAV(RWStructuredBuffer<int2>, RWVolumetricShadowLinkedListBuffer)
+		SHADER_PARAMETER_RDG_BUFFER_UAV(RWStructuredBuffer<FAVSMLinkedListPackedData>, RWVolumetricShadowLinkedListBuffer)
 	END_SHADER_PARAMETER_STRUCT()
 
 	static bool ShouldCompilePermutation(
