Path Tracer: Add minor debug mode improvements

Implement MIS for PrimaryRays debug mode to improve sharp specular highlights.  Implement debug modes to visualize light and decal grid axis selection.

#rnx

diff --git a/Engine/Shaders/Private/PathTracing/PathTracingCore.ush b/Engine/Shaders/Private/PathTracing/PathTracingCore.ush
--- a/Engine/Shaders/Private/PathTracing/PathTracingCore.ush
+++ b/Engine/Shaders/Private/PathTracing/PathTracingCore.ush
@@ -1644,14 +1644,6 @@ bool PathTracingKernel(inout FPathState PathState, int Bounce)
 			}
 		}
 
-		#if 0
-		// TODO: turn this into a debug mode
-		float NumLightsMax = min3(256.0f, SceneLightCount - SceneInfiniteLightCount, RAY_TRACING_LIGHT_COUNT_MAXIMUM);
-		PathState.Radiance = NumLights > 0 ? BlueGreenRedColorMap(saturate(float(NumLights) / NumLightsMax)) : 0.18;
-		PathState.Radiance *= PackedPayload.IsHit() ? abs(dot(PackedPayload.UnpackWorldNormal(), PathState.Ray.Direction)) : 1.0;
-		return false;
-		#endif
-
 		if (LightPickingCdfSum > 0.0)
 		{
 			// at least one light is overlapping with our ray, so we have a chance to sample it
diff --git a/Engine/Shaders/Private/PathTracing/PathTracingDebug.usf b/Engine/Shaders/Private/PathTracing/PathTracingDebug.usf
--- a/Engine/Shaders/Private/PathTracing/PathTracingDebug.usf
+++ b/Engine/Shaders/Private/PathTracing/PathTracingDebug.usf
@@ -2,6 +2,8 @@
 
 #define LIGHT_LOOP_METHOD 	2 // 0: stochastic (like ref PT) 1: loop over all lights 2: RIS (take N samples but trace a single shadow ray)
 
+#define USE_BSDF_SAMPLING 1 // 0: light sampling only (no MIS) 1: light sampling and BSDF sampling (with MIS)
+
 #define USE_PATH_TRACING_LIGHT_GRID	1
 #define USE_RAY_TRACING_DECAL_GRID	1
 
@@ -11,6 +13,7 @@
 #include "../RayTracing/RayTracingCommon.ush"
 #include "../RayTracing/RayTracingDecalGrid.ush"
 #include "PathTracingCommon.ush"
+#include "Material/PathTracingMaterialSampling.ush"
 #include "Utilities/PathTracingRandomSequence.ush"
 #include "Utilities/PathTracingRIS.ush"
 #include "Light/PathTracingLightSampling.ush"
@@ -96,13 +99,27 @@ float3 AnisoDebugColor(float InAniso)
 	return float3(0.0, pow(AniG, 2.2), pow(AniB, 2.2));
 }
 
+#define INCLUDE_TRANSLUCENT   0
+
 PathTracingResult PathTracingDebugCore(FRayDesc Ray, uint2 LaunchIndex)
 {
 	const uint RayFlags = RAY_FLAG_CULL_BACK_FACING_TRIANGLES;
+#if INCLUDE_TRANSLUCENT
+	const uint InstanceInclusionMask = PATHTRACER_MASK_CAMERA | PATHTRACER_MASK_CAMERA_TRANSLUCENT;
+#else
 	const uint InstanceInclusionMask = PATHTRACER_MASK_CAMERA;
+#endif
 	const uint MissShaderIndex = 0;
 
 	FPackedPathTracingPayload PackedPayload = InitPathTracingPayload(PATHTRACER_SCATTER_CAMERA, 0.0);
+	{
+		RandomSequence RandSequence;
+		RandomSequence_Initialize(RandSequence, LaunchIndex, 0, View.StateFrameIndex, 1);
+		PackedPayload.SetStochasticSlabRand(RandomSequence_GenerateSample1D(RandSequence));
+	}
+#if !INCLUDE_TRANSLUCENT
+	PackedPayload.SetFlag(PATH_TRACING_PAYLOAD_INPUT_FLAG_IGNORE_TRANSLUCENT);
+#endif
 	TraceRay(
 		TLAS,
 		RayFlags,
@@ -114,19 +131,56 @@ PathTracingResult PathTracingDebugCore(FRayDesc Ray, uint2 LaunchIndex)
 		PackedPayload);
 	float4 Result = float4(0, 0, 0, 1);
 	float Z = 0;
+	if (PackedPayload.IsHit())
+	{
+		const float3 ViewZ = View.ViewToTranslatedWorld[2].xyz;
+		Z = ConvertToDeviceZ(dot(Ray.Origin + PackedPayload.HitT * Ray.Direction, ViewZ));
+	}
+
+
+
+	if (DebugMode == PATH_TRACER_DEBUG_VIZ_VOLUME_LIGHT_COUNT)
+	{
+		float TMax = PackedPayload.IsHit() ? PackedPayload.HitT : Ray.TMax;
+		int NumLights = 0;
+		for (int LightId = SceneInfiniteLightCount; LightId < SceneLightCount; LightId++)
+		{
+			FVolumeLightSampleSetup LightSetup = PrepareLightVolumeSample(LightId, Ray.Origin, Ray.Direction, Ray.TMin, TMax);
+			if (LightSetup.IsValid())
+			{
+				float LightProb = LightSetup.LightImportance * GetVolumetricScatteringIntensity(LightId);
+				if (LightProb > 0)
+				{
+					NumLights++;
+					if (NumLights == RAY_TRACING_LIGHT_COUNT_MAXIMUM)
+					{
+						// reached max
+						break;
+					}
+				}
+			}
+		}
+
+		float NumLightsMax = min(SceneLightCount - SceneInfiniteLightCount, RAY_TRACING_LIGHT_COUNT_MAXIMUM);
+		Result.rgb = NumLights > 0 ? BlueGreenRedColorMap(saturate(float(NumLights) / NumLightsMax)) : 0.18;
+		if (PackedPayload.IsHit())
+		{
+			Result.rgb *= abs(dot(PackedPayload.UnpackWorldNormal(), Ray.Direction));
+		}
+		return MakeResult(Result.rgb, Result.a, Z);
+	}
+
+
 
 	float LightPickingCdf[RAY_TRACING_LIGHT_COUNT_MAXIMUM];
 	float LightPickingCdfSum = 0;
 	float3 PathThroughput = 1;
 	if (PackedPayload.IsHit())
 	{
-		const float3 ViewZ = View.ViewToTranslatedWorld[2].xyz;
 		const float3 TranslatedWorldPos = Ray.Origin + PackedPayload.HitT * Ray.Direction;
 		const float3 WorldNormal = PackedPayload.UnpackWorldNormal();
 		const uint PrimitiveLightingChannelMask = PackedPayload.GetPrimitiveLightingChannelMask();
 
-		Z = ConvertToDeviceZ(dot(TranslatedWorldPos, ViewZ));
-
 		FLightLoopCount LightLoopCount = LightGridLookup(TranslatedWorldPos);
 
 		switch (DebugMode)
@@ -164,20 +218,25 @@ PathTracingResult PathTracingDebugCore(FRayDesc Ray, uint2 LaunchIndex)
 			case PATH_TRACER_DEBUG_VIZ_PRIMARY_RAYS: break; // fallthrough to implementation below
 			case PATH_TRACER_DEBUG_VIZ_WORLD_TANGENT: 				return MakeResult(PackedPayload.UnpackWorldTangent() * 0.5 + 0.5, 0.0, Z);
 			case PATH_TRACER_DEBUG_VIZ_LIGHT_GRID_COUNT:
+			case PATH_TRACER_DEBUG_VIZ_LIGHT_GRID_AXIS:
 			{
-				float3 Color = LightGridVisualize(LightLoopCount, 1);
+
+				float3 Color = LightGridVisualize(LightLoopCount, DebugMode == PATH_TRACER_DEBUG_VIZ_LIGHT_GRID_AXIS ? 4 : 1);
 				Color *= abs(dot(WorldNormal, Ray.Direction));
 				return MakeResult(Color, 0.0, Z);
 			}
 			case PATH_TRACER_DEBUG_VIZ_DECAL_GRID_COUNT:
+			case PATH_TRACER_DEBUG_VIZ_DECAL_GRID_AXIS:
 			{
 				FDecalLoopCount DecalLoopCount = DecalGridLookup(TranslatedWorldPos);
-				float3 Color = DecalGridVisualize(DecalLoopCount, 1);
+				float3 Color = PackedPayload.IsDecalReceiver() ? DecalGridVisualize(DecalLoopCount, DebugMode == PATH_TRACER_DEBUG_VIZ_DECAL_GRID_AXIS ? 3 : 1) : 0.18;
 				Color *= abs(dot(WorldNormal, Ray.Direction));
 				return MakeResult(Color, 0.0, Z);
 			}
 		}
 
+		FPathTracingPayload Payload = UnpackPathTracingPayload(PackedPayload, Ray);
+		AdjustPayloadAfterUnpack(Payload, true);
 #if LIGHT_LOOP_METHOD == 0
 		for (uint Index = 0, Num = LightLoopCount.NumLights; Index < Num; ++Index)
 		{
@@ -199,6 +258,8 @@ PathTracingResult PathTracingDebugCore(FRayDesc Ray, uint2 LaunchIndex)
 				RandomSequence_Initialize(RandSequence, LaunchIndex, Sample, View.StateFrameIndex, NumLightSamples);
 				float4 RandSample = RandomSequence_GenerateSample4D(RandSequence);
 
+				FMaterialSample MaterialSample = SampleMaterial(-Ray.Direction, Payload, RandSample.xyz);
+
 				SelectLight(RandSample.x * LightPickingCdfSum, LightLoopCount.NumLights, LightPickingCdf, LightId, LightPickPdf);
 				LightPickPdf /= LightPickingCdfSum;
 				float2 LightRandSample = RandSample.yz;
@@ -217,10 +278,12 @@ PathTracingResult PathTracingDebugCore(FRayDesc Ray, uint2 LaunchIndex)
 			float3 ShadowRayContrib = 0;
 			uint ShadowRayLightId;
 #endif
+
+			FMaterialSample MaterialSample = SampleMaterial(-Ray.Direction, Payload, RandSample.xyz);
 			for (uint Index = 0, Num = LightLoopCount.NumLights; Index < Num; ++Index)
 			{
 				uint LightId = GetLightId(Index, LightLoopCount);
-				float LightPickPdf = 1;
+				const float LightPickPdf = 1;
 #endif
 
 				FLightSample LightSample = SampleLight(LightId, LightRandSample, TranslatedWorldPos, WorldNormal);
@@ -230,6 +293,14 @@ PathTracingResult PathTracingDebugCore(FRayDesc Ray, uint2 LaunchIndex)
 					LightSample.RadianceOverPdf /= LightPickPdf;
 					LightSample.Pdf *= LightPickPdf;
 
+#if 1
+					FMaterialEval MaterialEval = EvalMaterial(-Ray.Direction, LightSample.Direction, Payload, 1.0);
+					float3 LightContrib = PathThroughput * LightSample.RadianceOverPdf * MaterialEval.Weight * MaterialEval.Pdf;
+#if USE_BSDF_SAMPLING
+					LightContrib *= MISWeightPower(LightSample.Pdf, MaterialEval.Pdf);
+#endif
+
+#else
 #if SUBSTRATE_ENABLED
 					float3 MaterialEvalWeight = PackedPayload.UnpackDiffuseColor();
 #else
@@ -237,6 +308,7 @@ PathTracingResult PathTracingDebugCore(FRayDesc Ray, uint2 LaunchIndex)
 #endif
 					float  MaterialEvalPdf = saturate(dot(WorldNormal, LightSample.Direction)) / PI;
 					float3 LightContrib = PathThroughput * LightSample.RadianceOverPdf * MaterialEvalWeight * MaterialEvalPdf;
+#endif
 					{
 						FRayDesc ShadowRay;
 						ShadowRay.Origin = TranslatedWorldPos;
@@ -246,9 +318,37 @@ PathTracingResult PathTracingDebugCore(FRayDesc Ray, uint2 LaunchIndex)
 
 						ApplyRayBias(ShadowRay.Origin, PackedPayload.HitT, PackedPayload.UnpackWorldGeoNormal());
 
+#if LIGHT_LOOP_METHOD == 2
+						float SelectionWeight = LightContrib.x + LightContrib.y + LightContrib.z;
+						if (TraceSample.Accept(SelectionWeight))
+						{
+							ShadowRaySample = ShadowRay;
+							ShadowRayContrib = LightContrib / SelectionWeight;
+							ShadowRayLightId = LightId;
+						}
+#else
 						const bool bCastShadows = CastsShadow(LightId);
 						const uint MissShaderIndex = GetLightMissShaderIndex(LightId);
+						Result.xyz += LightContrib * TraceShadowRay(ShadowRay, 1.0, 0.0, MissShaderIndex, bCastShadows);
+#endif
+					}
+				}
+#if USE_BSDF_SAMPLING
+				{
+					FRayDesc ShadowRay;
+					ShadowRay.Origin = TranslatedWorldPos;
+					ShadowRay.Direction = MaterialSample.Direction;
+					ShadowRay.TMin = 0;
+					ShadowRay.TMax = POSITIVE_INFINITY;
+
+					ApplyRayBias(ShadowRay.Origin, PackedPayload.HitT, PackedPayload.UnpackWorldGeoNormal());
 
+					FLightHit LightResult = TraceLight(ShadowRay, LightId);
+					if (LightResult.IsHit())
+					{
+						ShadowRay.TMax = LightResult.HitT;
+						float3 LightContrib = PathThroughput * MaterialSample.Weight * LightResult.Radiance;
+						LightContrib *= MISWeightPower(MaterialSample.Pdf, LightResult.Pdf * LightPickPdf);
 #if LIGHT_LOOP_METHOD == 2
 						float SelectionWeight = LightContrib.x + LightContrib.y + LightContrib.z;
 						if (TraceSample.Accept(SelectionWeight))
@@ -258,10 +358,13 @@ PathTracingResult PathTracingDebugCore(FRayDesc Ray, uint2 LaunchIndex)
 							ShadowRayLightId = LightId;
 						}
 #else
+						const bool bCastShadows = CastsShadow(LightId);
+						const uint MissShaderIndex = GetLightMissShaderIndex(LightId);
 						Result.xyz += LightContrib * TraceShadowRay(ShadowRay, 1.0, 0.0, MissShaderIndex, bCastShadows);
 #endif
 					}
 				}
+#endif // USE_BSDF_SAMPLING
 			}
 #if LIGHT_LOOP_METHOD == 2
 			if (TraceSample.HasSample())
diff --git a/Engine/Shaders/Shared/PathTracingDefinitions.h b/Engine/Shaders/Shared/PathTracingDefinitions.h
--- a/Engine/Shaders/Shared/PathTracingDefinitions.h
+++ b/Engine/Shaders/Shared/PathTracingDefinitions.h
@@ -93,4 +93,7 @@
 #define PATH_TRACER_DEBUG_VIZ_WORLD_TANGENT				18 
 #define PATH_TRACER_DEBUG_VIZ_ANISOTROPY				19 
 #define PATH_TRACER_DEBUG_VIZ_LIGHT_GRID_COUNT          20
-#define PATH_TRACER_DEBUG_VIZ_DECAL_GRID_COUNT 			21
+#define PATH_TRACER_DEBUG_VIZ_LIGHT_GRID_AXIS           21
+#define PATH_TRACER_DEBUG_VIZ_DECAL_GRID_COUNT 			22
+#define PATH_TRACER_DEBUG_VIZ_DECAL_GRID_AXIS			23
+#define PATH_TRACER_DEBUG_VIZ_VOLUME_LIGHT_COUNT        24
\ No newline at end of file
diff --git a/Engine/Source/Runtime/Renderer/Private/PathTracing.cpp b/Engine/Source/Runtime/Renderer/Private/PathTracing.cpp
--- a/Engine/Source/Runtime/Renderer/Private/PathTracing.cpp
+++ b/Engine/Source/Runtime/Renderer/Private/PathTracing.cpp
@@ -741,7 +741,11 @@ int GetPathTracingDebugVisualizationMode() {
 		PathTracingDebugModes.Emplace(FName(TEXT("WorldTangent"))		, PATH_TRACER_DEBUG_VIZ_WORLD_TANGENT 			);
 		PathTracingDebugModes.Emplace(FName(TEXT("Anisotropy"))	    	, PATH_TRACER_DEBUG_VIZ_ANISOTROPY 				);
 		PathTracingDebugModes.Emplace(FName(TEXT("LightGridCount"))	   	, PATH_TRACER_DEBUG_VIZ_LIGHT_GRID_COUNT		);
+		PathTracingDebugModes.Emplace(FName(TEXT("LightGridAxis"))	   	, PATH_TRACER_DEBUG_VIZ_LIGHT_GRID_AXIS			);
 		PathTracingDebugModes.Emplace(FName(TEXT("DecalGridCount"))	   	, PATH_TRACER_DEBUG_VIZ_DECAL_GRID_COUNT		);
+		PathTracingDebugModes.Emplace(FName(TEXT("DecalGridAxis"))	   	, PATH_TRACER_DEBUG_VIZ_DECAL_GRID_AXIS			);
+		PathTracingDebugModes.Emplace(FName(TEXT("VolumeLightCount"))	, PATH_TRACER_DEBUG_VIZ_VOLUME_LIGHT_COUNT		);
+		
 	}
 	FString DebugMode = CVarPathTracingDebugVisualize.GetValueOnRenderThread();
 	if (!DebugMode.IsEmpty())
