Made light directionality of volume sampled translucent material energy conservative.

[RN] "Directional Lighting Intensity" is now energy conservative. The default value give the same look as before but other values will result in different result when it comes to visual intensity (and it could not be data converted). So you might have to reauthor materials if you have tweaked this value before.
#tests none

diff --git a/Engine/Shaders/Private/TranslucencyVolumeCommon.ush b/Engine/Shaders/Private/TranslucencyVolumeCommon.ush
--- a/Engine/Shaders/Private/TranslucencyVolumeCommon.ush
+++ b/Engine/Shaders/Private/TranslucencyVolumeCommon.ush
@@ -58,9 +58,6 @@ void GetVolumeLightingNonDirectional(float4 AmbientLightingVector, float3 Diffus
 
 void GetVolumeLightingDirectional(float4 AmbientLightingVector, float3 DirectionalLightingVector, float3 WorldNormal, float3 DiffuseColor, float DirectionalLightingIntensity, inout float3 InterpolatedLighting, out float4 VolumeLighting)
 {
-	AmbientLightingVector.rgb /= DirectionalLightingIntensity;
-	DirectionalLightingVector.rgb *= DirectionalLightingIntensity;
-
 	// Reconstruct the SH coefficients based on what was encoded
 	FTwoBandSHVectorRGB TranslucentLighting;
 	TranslucentLighting.R.V.x = AmbientLightingVector.r;
@@ -74,7 +71,7 @@ void GetVolumeLightingDirectional(float4 AmbientLightingVector, float3 Direction
 	TranslucentLighting.B.V.yzw = DirectionalLightingVector.rgb * NormalizedAmbientColor.b;
 
 	// Compute diffuse lighting which takes the normal into account
-	FTwoBandSHVector DiffuseTransferSH = CalcDiffuseTransferSH(WorldNormal, 1);
+	FTwoBandSHVector DiffuseTransferSH = CalcDiffuseTransferSH(WorldNormal * DirectionalLightingIntensity, 1);
 	VolumeLighting = float4(max(half3(0, 0, 0), DotSH(TranslucentLighting, DiffuseTransferSH)), AmbientLightingVector.a);
 	InterpolatedLighting += DiffuseColor * VolumeLighting.rgb;
 }
