[Exposure Fusion] Fix artifacts due to invalid input to ExposureFusionInverseTonemap(...)

diff --git a/Engine/Shaders/Private/PostProcessTonemap.usf b/Engine/Shaders/Private/PostProcessTonemap.usf
--- a/Engine/Shaders/Private/PostProcessTonemap.usf
+++ b/Engine/Shaders/Private/PostProcessTonemap.usf
@@ -439,7 +439,7 @@ float4 TonemapCommonPS(
 		const float LuminanceVal = CalculateEyeAdaptationLuminance(SceneColor.rgb * OneOverPreExposure * GlobalExposure);
 		
 #if 1
-		LocalExposure = ExposureFusionInverseTonemap(FusionVal) / LuminanceVal;
+		LocalExposure = ExposureFusionInverseTonemap(min(FusionVal, 1.0f)) / LuminanceVal;
 #else
 		const float LuminanceValTonemapped = ExposureFusionTonemap(LuminanceVal);
 		LocalExposure = FusionVal / LuminanceValTonemapped;
