Volumetric Fog: Store Extinction instead of Absorption in VBufferA

This avoids a few redundant Luminance calls and makes the math more accurate when albedo is strongly tinted.

#rnx

diff --git a/Engine/Shaders/Private/VolumetricFog.usf b/Engine/Shaders/Private/VolumetricFog.usf
--- a/Engine/Shaders/Private/VolumetricFog.usf
+++ b/Engine/Shaders/Private/VolumetricFog.usf
@@ -144,7 +144,6 @@ void MaterialSetupCS(
 	float Extinction = max(GlobalDensity * GlobalExtinctionScale, 0);
 
 	float3 Scattering = Albedo * Extinction;
-	float Absorption = max(Extinction - Luminance(Scattering), 0.0f);
 
 	float3 AdditionalEmissive = 0.0f;
 
@@ -217,7 +216,7 @@ void MaterialSetupCS(
 			float3 LFVScattering = LFVExtinction * FogInstance.Albedo;
 
 			Scattering += LFVScattering;
-			Absorption += max(LFVExtinction - Luminance(LFVScattering), 0.0f);
+			Extinction += LFVExtinction;
 
 	#if USE_EMISSIVE
 			// Emissive follow the density of the medium
@@ -229,7 +228,7 @@ void MaterialSetupCS(
 
 	if (all((int3)GridCoordinate < VolumetricFog.ResourceGridSizeInt))
 	{
-		RWVBufferA[GridCoordinate] = float4(Scattering, Absorption);
+		RWVBufferA[GridCoordinate] = float4(Scattering, Extinction);
 #if USE_EMISSIVE
 		float3 HeightFogEmissive = GlobalEmissive.rgb;
 	#if PROJECT_EXPFOG_MATCHES_VFOG
@@ -959,15 +958,15 @@ void LightScatteringCS(
 	}
 	#endif
 
-	float4 MaterialScatteringAndAbsorption = VBufferA[GridCoordinate];
-	float Extinction = MaterialScatteringAndAbsorption.w + Luminance(MaterialScatteringAndAbsorption.xyz);
+	float4 MaterialScatteringAndExtinction = VBufferA[GridCoordinate];
+	float Extinction = MaterialScatteringAndExtinction.w;
 	float3 MaterialEmissive = 0.0f;
 	BRANCH
 	if (UseEmissive > 0)
 	{
 		MaterialEmissive = VBufferB[GridCoordinate].xyz;
 	}
-	float4 PreExposedScatteringAndExtinction = float4(View.PreExposure * (LightScattering * MaterialScatteringAndAbsorption.xyz + MaterialEmissive), Extinction);
+	float4 PreExposedScatteringAndExtinction = float4(View.PreExposure * (LightScattering * MaterialScatteringAndExtinction.xyz + MaterialEmissive), Extinction);
 
 #if USE_TEMPORAL_REPROJECTION
 	BRANCH
diff --git a/Engine/Shaders/Private/VolumetricFogVoxelization.usf b/Engine/Shaders/Private/VolumetricFogVoxelization.usf
--- a/Engine/Shaders/Private/VolumetricFogVoxelization.usf
+++ b/Engine/Shaders/Private/VolumetricFogVoxelization.usf
@@ -331,7 +331,6 @@ void VoxelizePS(
 	float Extinction = SampleExtinctionCoefficients(PixelMaterialInputs);
 
 	float3 Scattering = Albedo * Extinction;
-	float Absorption = max(Extinction - Luminance(Scattering), 0.0f);
 
 	float FadeStart = .6f;
 	float SliceFadeAlpha = 1 - saturate((SvPosition.w / VolumetricFog.MaxDistance - FadeStart) / (1 - FadeStart));
@@ -342,6 +341,6 @@ void VoxelizePS(
 	// Would be faster to branch around the whole material evaluation, but that would cause divergent flow control for gradient operations
 	Scale *= ComputeVolumeShapeMasking(MaterialParameters.WorldPosition_CamRelative, MaterialParameters.PrimitiveId, Interpolants.FactoryInterpolants);
 
-	OutVBufferA = float4(Scattering * Scale, Absorption * Scale);
+	OutVBufferA = float4(Scattering * Scale, Extinction * Scale);
 	OutVBufferB = float4(EmissiveColor * Scale, 0);
 }
