Fix shader warnings

diff --git a/Engine/Shaders/Private/HeterogeneousVolumes/HeterogeneousVolumesAdaptiveVolumetricShadowMapUtils.ush b/Engine/Shaders/Private/HeterogeneousVolumes/HeterogeneousVolumesAdaptiveVolumetricShadowMapUtils.ush
--- a/Engine/Shaders/Private/HeterogeneousVolumes/HeterogeneousVolumesAdaptiveVolumetricShadowMapUtils.ush
+++ b/Engine/Shaders/Private/HeterogeneousVolumes/HeterogeneousVolumesAdaptiveVolumetricShadowMapUtils.ush
@@ -30,6 +30,11 @@ FAVSMSampleData AVSM_CreateSampleData(float X, float Tau)
 	return SampleData;
 }
 
+FAVSMSampleData AVSM_CreateSampleData(float X, float3 Transmittance)
+{
+	return AVSM_CreateSampleData(X, Luminance(Transmittance));
+}
+
 uint AVSM_PackSampleDataX(float X, FAVSMIndirectionData IndirectionData)
 {
 	float Extent = IndirectionData.EndX - IndirectionData.BeginX;
diff --git a/Engine/Shaders/Private/HeterogeneousVolumes/HeterogeneousVolumesLiveShadingShadows.usf b/Engine/Shaders/Private/HeterogeneousVolumes/HeterogeneousVolumesLiveShadingShadows.usf
--- a/Engine/Shaders/Private/HeterogeneousVolumes/HeterogeneousVolumesLiveShadingShadows.usf
+++ b/Engine/Shaders/Private/HeterogeneousVolumes/HeterogeneousVolumesLiveShadingShadows.usf
@@ -455,7 +455,7 @@ void RenderVolumetricShadowMapForLightWithLiveShadingCS(
 			bool bFirstEntry = (ElementIndex == 0);
 			if (bFirstEntry)
 			{
-				FAVSMSampleData SampleData = AVSM_CreateSampleData(WorldHitT, Luminance(Transmittance));
+				FAVSMSampleData SampleData = AVSM_CreateSampleData(WorldHitT, Transmittance);
 				AddSample(SampleData, PixelIndex, PrevElementIndex, ElementIndex);
 
 				// Set expected values
@@ -470,7 +470,7 @@ void RenderVolumetricShadowMapForLightWithLiveShadingCS(
 				float RelativeError = length((Transmittance - ExpectedTransmittance) / ExpectedTransmittance);
 				if ((AbsoluteError > AbsoluteErrorThreshold) && (RelativeError > RelativeErrorThreshold))
 				{
-					FAVSMSampleData SampleData = AVSM_CreateSampleData(WorldHitT, Luminance(Transmittance));
+					FAVSMSampleData SampleData = AVSM_CreateSampleData(WorldHitT, Transmittance);
 					AddOrReplaceSample(SampleData, PixelIndex, PrevElementIndex, ElementIndex);
 
 					// Reset expected values
diff --git a/Engine/Shaders/Private/HeterogeneousVolumes/HeterogeneousVolumesVoxelGridShadows.usf b/Engine/Shaders/Private/HeterogeneousVolumes/HeterogeneousVolumesVoxelGridShadows.usf
--- a/Engine/Shaders/Private/HeterogeneousVolumes/HeterogeneousVolumesVoxelGridShadows.usf
+++ b/Engine/Shaders/Private/HeterogeneousVolumes/HeterogeneousVolumesVoxelGridShadows.usf
@@ -458,7 +458,7 @@ void RenderVolumetricShadowMapForLightWithVoxelGridCS(
 			bool bFirstEntry = (ElementIndex == 0);
 			if (bFirstEntry)
 			{
-				FAVSMSampleData SampleData = AVSM_CreateSampleData(WorldHitT, Luminance(Transmittance));
+				FAVSMSampleData SampleData = AVSM_CreateSampleData(WorldHitT, Transmittance);
 				AddSample(SampleData, PixelIndex, PrevElementIndex, ElementIndex);
 
 				// Set expected values
@@ -473,7 +473,7 @@ void RenderVolumetricShadowMapForLightWithVoxelGridCS(
 				float RelativeError = length((Transmittance - ExpectedTransmittance) / ExpectedTransmittance);
 				if ((AbsoluteError > AbsoluteErrorThreshold) && (RelativeError > RelativeErrorThreshold))
 				{
-					FAVSMSampleData SampleData = AVSM_CreateSampleData(WorldHitT, Luminance(Transmittance));
+					FAVSMSampleData SampleData = AVSM_CreateSampleData(WorldHitT, Transmittance);
 					AddOrReplaceSample(SampleData, PixelIndex, PrevElementIndex, ElementIndex);
 
 					// Reset expected values
