Path Tracer: Fix exponential height fog mismatches with FogFalloff

The code was using the UI limit, but overlooked the fact that the constants are divided by 1000.0. This means the falloff value was effectively being clamped at 1.0

The deferred renderer does not clamp the fog density, it keeps growing while the path tracer was clamping the density below the height. Adjust the clamp to allow up to 8x increase in the mean-free path so we can see the exponential increase further down for large falloff values. Setting the cvar r.PathTracing.FogDensityClamp=1 will restore the old behavior.

diff --git a/Engine/Shaders/Private/PathTracing/Volume/PathTracingFog.ush b/Engine/Shaders/Private/PathTracing/Volume/PathTracingFog.ush
--- a/Engine/Shaders/Private/PathTracing/Volume/PathTracingFog.ush
+++ b/Engine/Shaders/Private/PathTracing/Volume/PathTracingFog.ush
@@ -12,6 +12,7 @@ float2 FogCenter;
 float FogMinZ;
 float FogMaxZ;
 float FogRadius;
+float FogFalloffClamp;
 
 FVolumeIntersection FogIntersect(float3 Origin, float3 Direction, float TMin, float TMax)
 {
@@ -42,6 +43,12 @@ FVolumeIntersection FogIntersect(float3 Origin, float3 Direction, float TMin, fl
 	return CreateEmptyVolumeIntersection();
 }
 
+// returns the density for each bank of fog
+float2 FogGetDualDensity(float Z)
+{
+	return FogDensity * exp2(-max((Z - FogHeight) * FogFalloff, FogFalloffClamp));
+}
+
 FVolumeDensityBounds FogGetDensityBounds(float3 Origin, float3 Direction, float TMin, float TMax)
 {
 #if 0
@@ -49,8 +56,8 @@ FVolumeDensityBounds FogGetDensityBounds(float3 Origin, float3 Direction, float
 	return CreateVolumeDensityBound(0, FogDensity.x + FogDensity.y);
 #else
 	// get bounds at each endpoint of the ray (density is a monotonic function, so just eval at both endpoints and sort)
-	float2 DensityA = FogDensity * exp2(-max(Origin.z + TMin * Direction.z - FogHeight, 0.0) * FogFalloff);
-	float2 DensityB = FogDensity * exp2(-max(Origin.z + TMax * Direction.z - FogHeight, 0.0) * FogFalloff);
+	float2 DensityA = FogGetDualDensity(Origin.z + TMin * Direction.z);
+	float2 DensityB = FogGetDualDensity(Origin.z + TMax * Direction.z);
 
 	float2 DensityBounds = float2(
 		DensityA.x + DensityA.y,
@@ -68,7 +75,7 @@ FVolumeShadedResult FogGetDensity(float3 TranslatedWorldPos)
 {
 	FVolumeShadedResult Result = (FVolumeShadedResult)0;
 
-	float2 Density = FogDensity * exp2(-max(TranslatedWorldPos.z - FogHeight, 0.0) * FogFalloff);
+	float2 Density = FogGetDualDensity(TranslatedWorldPos.z);
 	Result.SigmaT = Density.x + Density.y;
 	Result.SigmaSHG = FogAlbedo.xyz * Result.SigmaT;
 	Result.PhaseG = FogPhaseG;
@@ -78,7 +85,8 @@ FVolumeShadedResult FogGetDensity(float3 TranslatedWorldPos)
 // Return the transmittance along a ray segment
 float3 FogGetTransmittance(float3 Origin, float3 Direction, float TMin, float TMax)
 {
-	// Solve for intersection with Z == FogHeight (or nearest point within slab)
+	// Solve for intersection where (Z - FogHeight) * Falloff == FogFalloffClamp (or nearest point within slab)
+	float2 FogClipZ = FogFalloffClamp / FogFalloff + FogHeight;
 	float2 OH = Origin.z - FogHeight;
 	float Dz = Direction.z;
 
@@ -86,17 +94,20 @@ float3 FogGetTransmittance(float3 Origin, float3 Direction, float TMin, float TM
 	if (abs(Dz) < 1e-3)
 	{
 		// perfectly horizontal ray -- either go through all uniform or all varying, but either way density will be constant
-		Tau = (TMax - TMin) * exp2(-FogFalloff * max(OH, 0.0));
+		Tau = (TMax - TMin) * exp2(-max(FogFalloff * OH, FogFalloffClamp));
 	}
 	else
 	{
 		// figure out optical density through the uniform portion and varying portion respectively
-		float2 TMid = clamp(-OH / Dz, TMin, TMax);
-		Tau = Dz > 0
-			? exp2(-FogFalloff * max(0.0, OH + Dz * TMid)) - exp2(-FogFalloff * max(0.0, OH + Dz * TMax))
-			: exp2(-FogFalloff * max(0.0, OH + Dz * TMin)) - exp2(-FogFalloff * max(0.0, OH + Dz * TMid));
+		float2 TMid = clamp(-(Origin.z - FogClipZ) / Dz, TMin, TMax);
+		float2 Ta = Dz > 0 ? TMid : TMin;
+		float2 Tb = Dz > 0 ? TMax : TMid;
+		// add varying portion (will be 0.0 from clamp of TMid if we don't actually enter the varying portion)
+		Tau = exp2(-max(FogFalloff * (OH + Dz * Ta), FogFalloffClamp)) - exp2(-max(FogFalloff * (OH + Dz * Tb), FogFalloffClamp));
 		Tau *= rcp(Dz * FogFalloff * log(2.0));
-		Tau += Dz > 0 ? TMid - TMin : TMax - TMid;
+
+		// add uniform portion (will be 0.0 from clamp of TMid if we don't actually enter the uniform portion)
+		Tau += (Dz > 0 ? TMid - TMin : TMax - TMid) * exp2(-FogFalloffClamp);
 	}
 
 	return exp(-dot(FogDensity, Tau));
diff --git a/Engine/Source/Runtime/Renderer/Private/PathTracing.cpp b/Engine/Source/Runtime/Renderer/Private/PathTracing.cpp
--- a/Engine/Source/Runtime/Renderer/Private/PathTracing.cpp
+++ b/Engine/Source/Runtime/Renderer/Private/PathTracing.cpp
@@ -317,6 +317,14 @@ TAutoConsoleVariable<int32> CVarPathTracingCloudMapDepth(
 	ECVF_RenderThreadSafe
 );
 
+TAutoConsoleVariable<float> CVarPathTracingFogDensityClamp(
+	TEXT("r.PathTracing.FogDensityClamp"),
+	8.0f,
+	TEXT("Limit the density growth in exponential heightfog (default = 8)\n")
+	TEXT("Instead of allowing the exponential density to increase to infinity vertically, clamp it to some multiplier of the overall density."),
+	ECVF_RenderThreadSafe
+);
+
 TAutoConsoleVariable<int32> CVarPathTracingFrameIndependentTemporalSeed(
 	TEXT("r.PathTracing.FrameIndependentTemporalSeed"),
 	1,
@@ -1037,6 +1045,7 @@ BEGIN_SHADER_PARAMETER_STRUCT(FPathTracingFogParameters, )
 	SHADER_PARAMETER(float, FogMinZ)
 	SHADER_PARAMETER(float, FogMaxZ)
 	SHADER_PARAMETER(float, FogRadius)
+	SHADER_PARAMETER(float, FogFalloffClamp)
 END_SHADER_PARAMETER_STRUCT()
 
 static FPathTracingFogParameters PrepareFogParameters(const FViewInfo& View, const FExponentialHeightFogSceneInfo& FogInfo)
@@ -1046,12 +1055,16 @@ static FPathTracingFogParameters PrepareFogParameters(const FViewInfo& View, con
 
 	const FVector PreViewTranslation = View.ViewMatrices.GetPreViewTranslation();
 
-	Parameters.FogDensity.X = FogInfo.FogData[0].Density * FogInfo.VolumetricFogExtinctionScale;
-	Parameters.FogDensity.Y = FogInfo.FogData[1].Density * FogInfo.VolumetricFogExtinctionScale;
+	// See VolumetricFog.usf - the factor of .5 is needed for a better match to HeightFog behavior
+	const float MatchHeightFogFactor = .5f;
+	Parameters.FogDensity.X = MatchHeightFogFactor * FogInfo.FogData[0].Density * FogInfo.VolumetricFogExtinctionScale;
+	Parameters.FogDensity.Y = MatchHeightFogFactor * FogInfo.FogData[1].Density * FogInfo.VolumetricFogExtinctionScale;
 	Parameters.FogHeight.X = FogInfo.FogData[0].Height + PreViewTranslation.Z;
 	Parameters.FogHeight.Y = FogInfo.FogData[1].Height + PreViewTranslation.Z;
-	Parameters.FogFalloff.X = FMath::Max(FogInfo.FogData[0].HeightFalloff, 0.001f);
-	Parameters.FogFalloff.Y = FMath::Max(FogInfo.FogData[1].HeightFalloff, 0.001f);
+	// Clamp to UI limit to avoid division by 0 in the transmittance calculations
+	// Note that we have to adjust by factor of 1000.0 that is applied in FExponentialHeightFogSceneInfo()
+	Parameters.FogFalloff.X = FMath::Max(FogInfo.FogData[0].HeightFalloff, 0.001f / 1000.0f);
+	Parameters.FogFalloff.Y = FMath::Max(FogInfo.FogData[1].HeightFalloff, 0.001f / 1000.0f);
 	Parameters.FogAlbedo = FogInfo.VolumetricFogAlbedo;
 	Parameters.FogPhaseG = FogInfo.VolumetricFogScatteringDistribution;
 
@@ -1072,6 +1085,7 @@ static FPathTracingFogParameters PrepareFogParameters(const FViewInfo& View, con
 	Parameters.FogMinZ = FMath::Min(ZMin0, ZMin1);
 	Parameters.FogMaxZ = FMath::Max(ZMax0, ZMax1);
 	Parameters.FogRadius = Radius;
+	Parameters.FogFalloffClamp = -FMath::Log2(FMath::Clamp(CVarPathTracingFogDensityClamp.GetValueOnRenderThread(), 1.0f, 256.0f));
 	return Parameters;
 }
 
