Fixed initialization bug in FPerPageDispatchSetup which triggered checks and would cause out of bounds buffer access.

#rnx

diff --git a/Engine/Shaders/Private/VirtualShadowMaps/VirtualShadowMapPerPageDispatch.ush b/Engine/Shaders/Private/VirtualShadowMaps/VirtualShadowMapPerPageDispatch.ush
--- a/Engine/Shaders/Private/VirtualShadowMaps/VirtualShadowMapPerPageDispatch.ush
+++ b/Engine/Shaders/Private/VirtualShadowMaps/VirtualShadowMapPerPageDispatch.ush
@@ -28,6 +28,7 @@ struct FPerPageDispatchSetup
 
 	void Init(uint2 DispatchThreadId)
 	{
+		bValid = false;
 		VirtualShadowMapId = INDEX_NONE;
 		uint LinearThreadId = DispatchThreadId.y * PerPageDispatchDimX + DispatchThreadId.x;
 		if (bUseThreadPerId)
