Path Tracer: Improve decal grid

Implement a more accurate AABB vs OBB overlap test

diff --git a/Engine/Shaders/Private/RayTracing/RayTracingBuildDecalGrid.usf b/Engine/Shaders/Private/RayTracing/RayTracingBuildDecalGrid.usf
--- a/Engine/Shaders/Private/RayTracing/RayTracingBuildDecalGrid.usf
+++ b/Engine/Shaders/Private/RayTracing/RayTracingBuildDecalGrid.usf
@@ -55,6 +55,26 @@ void BuildDecalGridCS(uint3 DispatchThreadId : SV_DispatchThreadID)
 			continue;
 		}
 
+#if 1
+		// Transform the corners of the voxel into decal space so we can cull the "other way" as well and account for rotated decals
+		float3 DecalPos000 = mul(float4(Lo.x, Lo.y, Lo.z, 1), Decal.TranslatedWorldToDecal).xyz;
+		float3 DecalPos001 = mul(float4(Lo.x, Lo.y, Hi.z, 1), Decal.TranslatedWorldToDecal).xyz;
+		float3 DecalPos010 = mul(float4(Lo.x, Hi.y, Lo.z, 1), Decal.TranslatedWorldToDecal).xyz;
+		float3 DecalPos011 = mul(float4(Lo.x, Hi.y, Hi.z, 1), Decal.TranslatedWorldToDecal).xyz;
+		float3 DecalPos100 = mul(float4(Hi.x, Lo.y, Lo.z, 1), Decal.TranslatedWorldToDecal).xyz;
+		float3 DecalPos101 = mul(float4(Hi.x, Lo.y, Hi.z, 1), Decal.TranslatedWorldToDecal).xyz;
+		float3 DecalPos110 = mul(float4(Hi.x, Hi.y, Lo.z, 1), Decal.TranslatedWorldToDecal).xyz;
+		float3 DecalPos111 = mul(float4(Hi.x, Hi.y, Hi.z, 1), Decal.TranslatedWorldToDecal).xyz;
+
+		float3 DecalLo = min(min(min(DecalPos000, DecalPos001), min(DecalPos010, DecalPos011)), min(min(DecalPos100, DecalPos101), min(DecalPos110, DecalPos111)));
+		float3 DecalHi = max(max(max(DecalPos000, DecalPos001), max(DecalPos010, DecalPos011)), max(max(DecalPos100, DecalPos101), max(DecalPos110, DecalPos111)));
+
+		if (any(DecalHi < -1) || any(DecalLo > 1))
+		{
+			continue;
+		}
+#endif
+
 		RWDecalGridData[Offset + NumVisible] = Decal.CallableSlotIndex;
 		NumVisible++;
 		// TODO: handle overflow better?
diff --git a/Engine/Shaders/Shared/RayTracingTypes.h b/Engine/Shaders/Shared/RayTracingTypes.h
--- a/Engine/Shaders/Shared/RayTracingTypes.h
+++ b/Engine/Shaders/Shared/RayTracingTypes.h
@@ -79,12 +79,13 @@ HLSL_STATIC_ASSERT(sizeof(FPathTracingPackedPathState) == 80, "Packed Path State
 
 struct FRayTracingDecal
 {
-	float3 TranslatedBoundMin;
-	uint   Pad0; // keep structure aligned
-	float3 TranslatedBoundMax;
-	uint   CallableSlotIndex;
+	float3 		TranslatedBoundMin;
+	uint   		Pad0; // keep structure aligned
+	float3 		TranslatedBoundMax;
+	uint   		CallableSlotIndex;
+	float4x4	TranslatedWorldToDecal;
 };
-HLSL_STATIC_ASSERT(sizeof(FRayTracingDecal) == 32, "Ray tracing decal structure should be aligned to 32 bytes for optimal access on the GPU");
+HLSL_STATIC_ASSERT(sizeof(FRayTracingDecal) == 96, "Ray tracing decal structure should be aligned to 32 bytes for optimal access on the GPU");
 
 #ifdef __cplusplus
 } // namespace UE::HLSL
diff --git a/Engine/Source/Runtime/Renderer/Private/RayTracing/RayTracingDecals.cpp b/Engine/Source/Runtime/Renderer/Private/RayTracing/RayTracingDecals.cpp
--- a/Engine/Source/Runtime/Renderer/Private/RayTracing/RayTracingDecals.cpp
+++ b/Engine/Source/Runtime/Renderer/Private/RayTracing/RayTracingDecals.cpp
@@ -572,7 +572,10 @@ TRDGUniformBufferRef<FRayTracingDecals> CreateRayTracingDecalData(FRDGBuilder& G
 		DestDecal.TranslatedBoundMin = FVector3f(BoxBounds.Min + View.ViewMatrices.GetPreViewTranslation());
 		DestDecal.TranslatedBoundMax = FVector3f(BoxBounds.Max + View.ViewMatrices.GetPreViewTranslation());
 		DestDecal.CallableSlotIndex = BaseCallableSlotIndex + Index;
-
+		{
+			const FMatrix WorldToDecalMatrix = DecalData.Proxy->ComponentTrans.ToInverseMatrixWithScale();
+			DestDecal.TranslatedWorldToDecal = FMatrix44f(FTranslationMatrix(-View.ViewMatrices.GetPreViewTranslation()) * WorldToDecalMatrix);
+		}
 		TranslatedBoundMin = FVector3f::Min(TranslatedBoundMin, DestDecal.TranslatedBoundMin);
 		TranslatedBoundMax = FVector3f::Max(TranslatedBoundMax, DestDecal.TranslatedBoundMax);
 	}
