Fixed typo in checkSlow causing compile errors if r.Shaders.CheckLevel=2 is enabled.

#rnx

diff --git a/Engine/Shaders/Private/InstanceCulling/InstanceCullingSetup.ush b/Engine/Shaders/Private/InstanceCulling/InstanceCullingSetup.ush
--- a/Engine/Shaders/Private/InstanceCulling/InstanceCullingSetup.ush
+++ b/Engine/Shaders/Private/InstanceCulling/InstanceCullingSetup.ush
@@ -35,7 +35,7 @@ FInstanceCullingSetup LoadInstanceCullingSetup(uint3 GroupId, uint GroupThreadIn
 	if (IsDynamicInstanceDataOffset(Result.InstanceWorkSetup.Item.Payload))
 	{
 		Result.InstanceId += InDynamicInstanceIdOffset;
-		checkSlow(InstanceId < InDynamicInstanceIdMax || (InDynamicInstanceIdOffset == 0 && InDynamicInstanceIdMax == 0));
+		checkSlow(Result.InstanceId < InDynamicInstanceIdMax || (InDynamicInstanceIdOffset == 0 && InDynamicInstanceIdMax == 0));
 		Result.bIsDynamic = true;
 	}
 
