[LocalExposure] New thresholds visualization mode and improved valid lookup mode.

diff --git a/Engine/Shaders/Private/PostProcessVisualizeLocalExposure.usf b/Engine/Shaders/Private/PostProcessVisualizeLocalExposure.usf
--- a/Engine/Shaders/Private/PostProcessVisualizeLocalExposure.usf
+++ b/Engine/Shaders/Private/PostProcessVisualizeLocalExposure.usf
@@ -10,6 +10,7 @@
 StructuredBuffer<float4> EyeAdaptationBuffer;
 
 Texture2D HDRSceneColorTexture;
+Texture2D SceneColorTexture;
 
 Texture3D LumBilateralGrid;
 Texture2D BlurredLogLum;
@@ -18,9 +19,21 @@ SamplerState TextureSampler;
 
 uint DebugMode;
 
+#define VISUALIZE_LOCAL_EXPOSURE_MODE_DEFAULT			(0)
+#define VISUALIZE_LOCAL_EXPOSURE_MODE_THRESHOLDS		(1)
+#define VISUALIZE_LOCAL_EXPOSURE_MODE_BASE_LUMINANCE	(2)
+#define VISUALIZE_LOCAL_EXPOSURE_MODE_DETAIL_LUMINANCE	(3)
+#define VISUALIZE_LOCAL_EXPOSURE_MODE_VALID_LOOKUP		(4)
+
 SCREEN_PASS_TEXTURE_VIEWPORT(Input)
 SCREEN_PASS_TEXTURE_VIEWPORT(Output)
 
+float DiagonalStripesPattern(float2 ScreenPos, float Thickness)
+{    
+    float V = frac((ScreenPos.x + ScreenPos.y) / Thickness);
+    return step(0.5f, V);
+}
+
 float4 MainPS(noperspective float4 UVAndScreenPos : TEXCOORD0, float4 SvPosition : SV_POSITION) : SV_Target0
 {
 	float2 UV = UVAndScreenPos.xy;
@@ -28,9 +41,10 @@ float4 MainPS(noperspective float4 UVAndScreenPos : TEXCOORD0, float4 SvPosition
 
 	float2 ExposureScaleMiddleGrey = EyeAdaptationBuffer[0].xw;
 
-	float4 SceneColor = Texture2DSample(HDRSceneColorTexture, TextureSampler, UV) * View.OneOverPreExposure;
+	float4 HDRSceneColor = Texture2DSample(HDRSceneColorTexture, TextureSampler, UV) * View.OneOverPreExposure;
+	float4 SceneColor = Texture2DSample(SceneColorTexture, TextureSampler, UV);
 
-	float LuminanceVal = CalculateEyeAdaptationLuminance(SceneColor.rgb);
+	float LuminanceVal = CalculateEyeAdaptationLuminance(HDRSceneColor.rgb);
 	float LogLuminance = log2(LuminanceVal);
 	float MiddleGreyLumValue = log2(0.18 * ExposureScaleMiddleGrey.y * LocalExposure_MiddleGreyExposureCompensation);
 
@@ -40,20 +54,42 @@ float4 MainPS(noperspective float4 UVAndScreenPos : TEXCOORD0, float4 SvPosition
 
 	float3 OutValue = 0;
 
-	if (DebugMode == 0)
+	if (DebugMode == VISUALIZE_LOCAL_EXPOSURE_MODE_DEFAULT)
 	{
 		OutValue = saturate(abs(log2(LocalExposure) / 2)) * lerp(float3(1.0f, 0.0f, 0.0f), float3(0.0f, 1.0f, 0.0f), step(LocalExposure, 1));
 	}
-	else if (DebugMode == 1)
+	else if (DebugMode == VISUALIZE_LOCAL_EXPOSURE_MODE_THRESHOLDS)
+	{
+		OutValue = SceneColor.rgb;
+
+		float BaseCentered = (BaseLogLum - MiddleGreyLumValue);
+		if (BaseCentered > 0)
+		{
+			if (BaseCentered - LocalExposure_HighlightThreshold > 0)
+			{
+				OutValue = lerp(OutValue, float3(0,1,0), DiagonalStripesPattern(SvPosition.xy, 32.0f));
+			}
+		}
+		else
+		{
+			if (BaseCentered + LocalExposure_ShadowThreshold < 0)
+			{
+				OutValue = lerp(OutValue, float3(1,0,0), DiagonalStripesPattern(SvPosition.xy, 32.0f));
+			}
+		}
+	}
+	else if (DebugMode == VISUALIZE_LOCAL_EXPOSURE_MODE_BASE_LUMINANCE)
 	{
 		OutValue = exp2(BaseLogLum);
 	}
-	else if (DebugMode == 2)
+	else if (DebugMode == VISUALIZE_LOCAL_EXPOSURE_MODE_DETAIL_LUMINANCE)
 	{
 		OutValue = exp2(DetailLogLum);
 	}
-	else if (DebugMode == 3)
+	else if (DebugMode == VISUALIZE_LOCAL_EXPOSURE_MODE_VALID_LOOKUP)
 	{
+		OutValue = SceneColor.rgb;
+
 		float3 BilateralGridUVW;
 		BilateralGridUVW.xy = ViewportUV * LocalExposure_BilateralGridUVScale;
 		BilateralGridUVW.z = (ComputeHistogramPositionFromLogLuminance(LogLuminance) * (BILATERAL_GRID_DEPTH - 1) + 0.5f) / BILATERAL_GRID_DEPTH;
@@ -67,7 +103,7 @@ float4 MainPS(noperspective float4 UVAndScreenPos : TEXCOORD0, float4 SvPosition
 		}
 		else
 		{
-			OutValue = float3(0,1,0);
+			OutValue = lerp(OutValue, float3(0,1,0), 0.5f);
 		}
 	}
 
diff --git a/Engine/Source/Runtime/Renderer/Private/PostProcess/PostProcessVisualizeLocalExposure.cpp b/Engine/Source/Runtime/Renderer/Private/PostProcess/PostProcessVisualizeLocalExposure.cpp
--- a/Engine/Source/Runtime/Renderer/Private/PostProcess/PostProcessVisualizeLocalExposure.cpp
+++ b/Engine/Source/Runtime/Renderer/Private/PostProcess/PostProcessVisualizeLocalExposure.cpp
@@ -12,9 +12,10 @@ TAutoConsoleVariable<int> CVarLocalExposureVisualizeDebugMode(
 	0,
 	TEXT("When enabling Show->Visualize->Local Exposure is enabled, this flag controls which mode to use.\n")
 	TEXT("    0: Local Exposure\n")
-	TEXT("    1: Base Luminance\n")
-	TEXT("    2: Detail Luminance\n")
-	TEXT("    3: Valid Bilateral Grid Lookup\n"),
+	TEXT("    1: Thresholds\n")
+	TEXT("    2: Base Luminance\n")
+	TEXT("    3: Detail Luminance\n")
+	TEXT("    4: Valid Bilateral Grid Lookup\n"),
 	ECVF_RenderThreadSafe);
 
 class FVisualizeLocalExposurePS : public FGlobalShader
@@ -31,6 +32,7 @@ public:
 		SHADER_PARAMETER_STRUCT(FScreenPassTextureViewportParameters, Output)
 
 		SHADER_PARAMETER_RDG_TEXTURE(Texture2D, HDRSceneColorTexture)
+		SHADER_PARAMETER_RDG_TEXTURE(Texture2D, SceneColorTexture)
 		SHADER_PARAMETER_RDG_BUFFER_SRV(StructuredBuffer<float4>, EyeAdaptationBuffer)
 		SHADER_PARAMETER_RDG_TEXTURE(Texture3D, LumBilateralGrid)
 		SHADER_PARAMETER_RDG_TEXTURE(Texture2D, BlurredLogLum)
@@ -81,6 +83,7 @@ FScreenPassTexture AddVisualizeLocalExposurePass(FRDGBuilder& GraphBuilder, cons
 	PassParameters->Input = GetScreenPassTextureViewportParameters(InputViewport);
 	PassParameters->Output = GetScreenPassTextureViewportParameters(OutputViewport);
 	PassParameters->HDRSceneColorTexture = Inputs.HDRSceneColor.Texture;
+	PassParameters->SceneColorTexture = Inputs.SceneColor.Texture;
 	PassParameters->EyeAdaptationBuffer = GraphBuilder.CreateSRV(Inputs.EyeAdaptationBuffer);
 	PassParameters->LumBilateralGrid = Inputs.LumBilateralGridTexture;
 	PassParameters->BlurredLogLum = Inputs.BlurredLumTexture;
