HeterogeneousVolumes: Adding DDX/DDY support

When selected, the use of analytic derivatives will be supported:
r.HeterogeneousVolumes.UseAnalyticDerivatives [default=1]

diff --git a/Engine/Shaders/Private/HeterogeneousVolumes/HeterogeneousVolumesAmbientOcclusionPipeline.usf b/Engine/Shaders/Private/HeterogeneousVolumes/HeterogeneousVolumesAmbientOcclusionPipeline.usf
--- a/Engine/Shaders/Private/HeterogeneousVolumes/HeterogeneousVolumesAmbientOcclusionPipeline.usf
+++ b/Engine/Shaders/Private/HeterogeneousVolumes/HeterogeneousVolumesAmbientOcclusionPipeline.usf
@@ -134,7 +134,7 @@ struct FVolumeSampleContext
 	FPixelMaterialInputs PixelMaterialInputs;
 };
 
-FVolumeSampleContext CreateVolumeSampleContext(float3 LocalPosition, float3 WorldPosition, float MipLevel)
+FVolumeSampleContext CreateVolumeSampleContext(float3 LocalPosition, float3 WorldPosition, float3 WorldPosition_DDX, float3 WorldPosition_DDY, float MipLevel)
 {
 	FVolumeSampleContext VolumeSampleContext;
 	VolumeSampleContext.MaterialParameters = MakeInitializedMaterialPixelParameters();
@@ -145,11 +145,30 @@ FVolumeSampleContext CreateVolumeSampleContext(float3 LocalPosition, float3 Worl
 	VolumeSampleContext.MaterialParameters.LWCData.LocalToWorld = WSPromote(GetLocalToWorld());
 	VolumeSampleContext.MaterialParameters.LWCData.WorldToLocal = WSPromoteInverse(GetWorldToLocal());
 
-	// Evaluate Material graph
+#if USE_ANALYTIC_DERIVATIVES
+	VolumeSampleContext.MaterialParameters.WorldPosition_DDX = WorldPosition_DDX;
+	VolumeSampleContext.MaterialParameters.WorldPosition_DDY = WorldPosition_DDY;
+	CalcPixelMaterialInputsAnalyticDerivatives(VolumeSampleContext.MaterialParameters, VolumeSampleContext.PixelMaterialInputs);
+#else
 	CalcPixelMaterialInputs(VolumeSampleContext.MaterialParameters, VolumeSampleContext.PixelMaterialInputs);
+#endif
+
 	return VolumeSampleContext;
 }
 
+FVolumeSampleContext CreateVolumeSampleContext(float3 LocalPosition, float3 WorldPosition, float MipLevel)
+{
+	float3 WorldPosition_DDX = 0.0;
+	float3 WorldPosition_DDY = 0.0;
+	return CreateVolumeSampleContext(
+		LocalPosition,
+		WorldPosition,
+		WorldPosition_DDX,
+		WorldPosition_DDY,
+		MipLevel
+	);
+}
+
 FVolumeSampleContext CreateVolumeSampleContext(float3 LocalPosition, float3 WorldPosition, float3 FilterWidth, float MipLevel, float Rand, int StochasticFilteringMode)
 {
 	// TODO: Use FilterWidth determine MipLevel
diff --git a/Engine/Shaders/Private/HeterogeneousVolumes/HeterogeneousVolumesLiveShadingPipeline.usf b/Engine/Shaders/Private/HeterogeneousVolumes/HeterogeneousVolumesLiveShadingPipeline.usf
--- a/Engine/Shaders/Private/HeterogeneousVolumes/HeterogeneousVolumesLiveShadingPipeline.usf
+++ b/Engine/Shaders/Private/HeterogeneousVolumes/HeterogeneousVolumesLiveShadingPipeline.usf
@@ -27,7 +27,6 @@
 #ifndef USE_EXISTENCE_MASK
 #define USE_EXISTENCE_MASK 0
 #endif //
-
 #define USE_CAMERA_AVSM (DIM_USE_ADAPTIVE_VOLUMETRIC_SHADOW_MAP && (SUPPORT_OVERLAPPING_VOLUMES != 0))
 
 #if SUPPORT_OVERLAPPING_VOLUMES
@@ -42,7 +41,6 @@
 #include "HeterogeneousVolumesLightingUtils.ush"
 #include "../BlueNoise.ush"
 
-
 #if VIRTUAL_SHADOW_MAP
 #include "../VirtualShadowMaps/VirtualShadowMapProjectionCommon.ush"
 #endif
@@ -221,7 +219,13 @@ struct FVolumeSampleContext
 	FPixelMaterialInputs PixelMaterialInputs;
 };
 
-FVolumeSampleContext CreateVolumeSampleContext(float3 LocalPosition, float3 WorldPosition, float MipLevel)
+FVolumeSampleContext CreateVolumeSampleContext(
+	float3 LocalPosition,
+	float3 WorldPosition,
+	float3 WorldPosition_DDX,
+	float3 WorldPosition_DDY,
+	float MipLevel
+)
 {
 	FVolumeSampleContext VolumeSampleContext;
 	VolumeSampleContext.MaterialParameters = MakeInitializedMaterialPixelParameters();
@@ -232,12 +236,31 @@ FVolumeSampleContext CreateVolumeSampleContext(float3 LocalPosition, float3 Worl
 	VolumeSampleContext.MaterialParameters.LWCData.LocalToWorld = WSPromote(GetLocalToWorld());
 	VolumeSampleContext.MaterialParameters.LWCData.WorldToLocal = WSPromoteInverse(GetWorldToLocal());
 
-	// Evaluate Material graph
+#if USE_ANALYTIC_DERIVATIVES
+	VolumeSampleContext.MaterialParameters.WorldPosition_DDX = WorldPosition_DDX;
+	VolumeSampleContext.MaterialParameters.WorldPosition_DDY = WorldPosition_DDY;
+	CalcPixelMaterialInputsAnalyticDerivatives(VolumeSampleContext.MaterialParameters, VolumeSampleContext.PixelMaterialInputs);
+#else
 	CalcPixelMaterialInputs(VolumeSampleContext.MaterialParameters, VolumeSampleContext.PixelMaterialInputs);
+#endif
 
 	return VolumeSampleContext;
 }
 
+FVolumeSampleContext CreateVolumeSampleContext(float3 LocalPosition, float3 WorldPosition, float MipLevel)
+{
+	float3 WorldPosition_DDX = 0.0;
+	float3 WorldPosition_DDY = 0.0;
+
+	return CreateVolumeSampleContext(
+		LocalPosition,
+		WorldPosition,
+		WorldPosition_DDX,
+		WorldPosition_DDY,
+		MipLevel
+	);
+}
+
 FVolumeSampleContext CreateVolumeSampleContext(float3 LocalPosition, float3 WorldPosition, float3 FilterWidth, float MipLevel, float Rand, int StochasticFilteringMode)
 {
 	// TODO: Use FilterWidth determine MipLevel
@@ -534,12 +557,23 @@ void RenderSingleScatteringWithLiveShading(
 	float WorldRayLength = length(WorldRayDirection);
 	WorldRayDirection /= WorldRayLength;
 
+	// TODO: Analytical derivate evaluation is not currently supported under orthographic projection
+	float3 WorldRayEnd_DDX = DFHackToFloat(SvPositionToWorld(float4(ViewPixelCoord + ScreenJitter + float2(1.0, 0.0), DeviceZ, 1)));
+	float3 WorldRayDirection_DDX = normalize(WorldRayEnd_DDX - WorldRayOrigin);
+	float3 WorldRayEnd_DDY = DFHackToFloat(SvPositionToWorld(float4(ViewPixelCoord + ScreenJitter + float2(0.0, 1.0), DeviceZ, 1)));
+	float3 WorldRayDirection_DDY = normalize(WorldRayEnd_DDY - WorldRayOrigin);
+
 	float3 LocalRayOrigin = mul(float4(WorldRayOrigin, 1.0), WorldToLocal).xyz;
 	float3 LocalRayEnd = mul(float4(WorldRayEnd, 1.0), WorldToLocal).xyz;
 	float3 LocalRayDirection = LocalRayEnd - LocalRayOrigin;
 	float LocalRayLength = length(LocalRayDirection);
 	LocalRayDirection /= LocalRayLength;
 
+	float3 LocalRayEnd_DDX = mul(float4(WorldRayEnd_DDX, 1.0), WorldToLocal).xyz;
+	float3 LocalRayDirection_DDX = normalize(LocalRayEnd_DDX - LocalRayOrigin);
+	float3 LocalRayEnd_DDY = mul(float4(WorldRayEnd_DDY, 1.0), WorldToLocal).xyz;
+	float3 LocalRayDirection_DDY = normalize(LocalRayEnd_DDY - LocalRayOrigin);
+
 	float3 LocalBoundsMin = LocalBoundsOrigin - LocalBoundsExtent;
 	float3 LocalBoundsMax = LocalBoundsOrigin + LocalBoundsExtent;
 
@@ -553,11 +587,15 @@ void RenderSingleScatteringWithLiveShading(
 			// Local-space
 			LocalRayOrigin,
 			LocalRayDirection,
+			LocalRayDirection_DDX,
+			LocalRayDirection_DDY,
 			HitT.x,
 			HitT.y,
 			// World-space
 			WorldRayOrigin,
 			WorldRayDirection,
+			WorldRayDirection_DDX,
+			WorldRayDirection_DDY,
 			// Ray-step attributes
 			Jitter,
 			CalcStepSize(LocalRayDirection),
diff --git a/Engine/Shaders/Private/HeterogeneousVolumes/HeterogeneousVolumesLiveShadingShadows.usf b/Engine/Shaders/Private/HeterogeneousVolumes/HeterogeneousVolumesLiveShadingShadows.usf
--- a/Engine/Shaders/Private/HeterogeneousVolumes/HeterogeneousVolumesLiveShadingShadows.usf
+++ b/Engine/Shaders/Private/HeterogeneousVolumes/HeterogeneousVolumesLiveShadingShadows.usf
@@ -148,7 +148,13 @@ struct FVolumeSampleContext
 	FPixelMaterialInputs PixelMaterialInputs;
 };
 
-FVolumeSampleContext CreateVolumeSampleContext(float3 LocalPosition, float3 WorldPosition, float MipLevel)
+FVolumeSampleContext CreateVolumeSampleContext(
+	float3 LocalPosition,
+	float3 WorldPosition,
+	float3 WorldPosition_DDX,
+	float3 WorldPosition_DDY,
+	float MipLevel
+)
 {
 	FVolumeSampleContext VolumeSampleContext;
 	VolumeSampleContext.MaterialParameters = MakeInitializedMaterialPixelParameters();
@@ -159,12 +165,31 @@ FVolumeSampleContext CreateVolumeSampleContext(float3 LocalPosition, float3 Worl
 	VolumeSampleContext.MaterialParameters.LWCData.LocalToWorld = WSPromote(GetLocalToWorld());
 	VolumeSampleContext.MaterialParameters.LWCData.WorldToLocal = WSPromoteInverse(GetWorldToLocal());
 
-	// Evaluate Material graph
+#if USE_ANALYTIC_DERIVATIVES
+	VolumeSampleContext.MaterialParameters.WorldPosition_DDX = WorldPosition_DDX;
+	VolumeSampleContext.MaterialParameters.WorldPosition_DDY = WorldPosition_DDY;
+	CalcPixelMaterialInputsAnalyticDerivatives(VolumeSampleContext.MaterialParameters, VolumeSampleContext.PixelMaterialInputs);
+#else
 	CalcPixelMaterialInputs(VolumeSampleContext.MaterialParameters, VolumeSampleContext.PixelMaterialInputs);
+#endif
 
 	return VolumeSampleContext;
 }
 
+FVolumeSampleContext CreateVolumeSampleContext(float3 LocalPosition, float3 WorldPosition, float MipLevel)
+{
+	float3 WorldPosition_DDX = 0.0;
+	float3 WorldPosition_DDY = 0.0;
+
+	return CreateVolumeSampleContext(
+		LocalPosition,
+		WorldPosition,
+		WorldPosition_DDX,
+		WorldPosition_DDY,
+		MipLevel
+	);
+}
+
 FVolumeSampleContext CreateVolumeSampleContext(float3 LocalPosition, float3 WorldPosition, float3 FilterWidth, float MipLevel, float Rand, int StochasticFilteringMode)
 {
 	// TODO: Use FilterWidth determine MipLevel
diff --git a/Engine/Shaders/Private/HeterogeneousVolumes/HeterogeneousVolumesRayMarchingTypes.ush b/Engine/Shaders/Private/HeterogeneousVolumes/HeterogeneousVolumesRayMarchingTypes.ush
--- a/Engine/Shaders/Private/HeterogeneousVolumes/HeterogeneousVolumesRayMarchingTypes.ush
+++ b/Engine/Shaders/Private/HeterogeneousVolumes/HeterogeneousVolumesRayMarchingTypes.ush
@@ -13,6 +13,10 @@ struct FRayMarchingContext
 	// World-space
 	float3 WorldRayOrigin;
 	float3 WorldRayDirection;
+#if USE_ANALYTIC_DERIVATIVES
+	float3 WorldRayDirection_DDX;
+	float3 WorldRayDirection_DDY;
+#endif // USE_ANALYTIC_DERIVATIVES
 	float LocalToWorldScale;
 
 	// Ray-step attributes
@@ -35,11 +39,15 @@ FRayMarchingContext CreateRayMarchingContext(
 	// Local-space
 	float3 LocalRayOrigin,
 	float3 LocalRayDirection,
+	float3 LocalRayDirection_DDX,
+	float3 LocalRayDirection_DDY,
 	float LocalRayTMin,
 	float LocalRayTMax,
 	// World-space
 	float3 WorldRayOrigin,
 	float3 WorldRayDirection,
+	float3 WorldRayDirection_DDX,
+	float3 WorldRayDirection_DDY,
 	// Ray-step attributes
 	float Jitter,
 	float StepSize,
@@ -60,6 +68,10 @@ FRayMarchingContext CreateRayMarchingContext(
 	// World-space
 	Context.WorldRayOrigin = WorldRayOrigin;
 	Context.WorldRayDirection = WorldRayDirection;
+#if USE_ANALYTIC_DERIVATIVES
+	Context.WorldRayDirection_DDX = WorldRayDirection_DDX;
+	Context.WorldRayDirection_DDY = WorldRayDirection_DDY;
+#endif // USE_ANALYTIC_DERIVATIVES
 	// Ray-step attributes
 	Context.Jitter = Jitter;
 	Context.StepSize = StepSize;
@@ -80,6 +92,54 @@ FRayMarchingContext CreateRayMarchingContext(
 	return Context;
 }
 
+FRayMarchingContext CreateRayMarchingContext(
+	// Local-space
+	float3 LocalRayOrigin,
+	float3 LocalRayDirection,
+	float LocalRayTMin,
+	float LocalRayTMax,
+	// World-space
+	float3 WorldRayOrigin,
+	float3 WorldRayDirection,
+	// Ray-step attributes
+	float Jitter,
+	float StepSize,
+	uint MaxStepCount,
+	// Lighting
+	int bApplyEmissionAndTransmittance,
+	int bApplyDirectLighting,
+	int bApplyShadowTransmittance
+)
+{
+	float3 LocalRayDirection_DDX = 0.0;
+	float3 LocalRayDirection_DDY = 0.0;
+	float3 WorldRayDirection_DDX = 0.0;
+	float3 WorldRayDirection_DDY = 0.0;
+
+	return CreateRayMarchingContext(
+		// Local-space
+		LocalRayOrigin,
+		LocalRayDirection,
+		LocalRayDirection_DDX,
+		LocalRayDirection_DDY,
+		LocalRayTMin,
+		LocalRayTMax,
+		// World-space
+		WorldRayOrigin,
+		WorldRayDirection,
+		WorldRayDirection_DDX,
+		WorldRayDirection_DDY,
+		// Ray-step attributes
+		Jitter,
+		StepSize,
+		MaxStepCount,
+		// Lighting
+		bApplyEmissionAndTransmittance,
+		bApplyDirectLighting,
+		bApplyShadowTransmittance
+	);
+}
+
 float CalcStepSize(float3 LocalRayDirection, float LocalStepSize, float LocalStepFactor)
 {
 	if (LocalStepSize <= 0.0f)
diff --git a/Engine/Shaders/Private/HeterogeneousVolumes/HeterogeneousVolumesRayMarchingUtils.ush b/Engine/Shaders/Private/HeterogeneousVolumes/HeterogeneousVolumesRayMarchingUtils.ush
--- a/Engine/Shaders/Private/HeterogeneousVolumes/HeterogeneousVolumesRayMarchingUtils.ush
+++ b/Engine/Shaders/Private/HeterogeneousVolumes/HeterogeneousVolumesRayMarchingUtils.ush
@@ -122,8 +122,17 @@ void RayMarchTransmittance(
 		float LocalHitT = RayMarchingContext.LocalRayTMin + RayMarchingContext.StepSize * (RayMarchingContext.Jitter + StepIndex);
 		float3 LocalPosition = RayMarchingContext.LocalRayOrigin + RayMarchingContext.LocalRayDirection * LocalHitT;
 		float3 WorldPosition = RayMarchingContext.WorldRayOrigin + RayMarchingContext.WorldRayDirection * LocalHitT * RayMarchingContext.LocalToWorldScale;
+#if USE_ANALYTIC_DERIVATIVES
+		float3 WorldPosition_DDX = RayMarchingContext.WorldRayOrigin + RayMarchingContext.WorldRayDirection_DDX * LocalHitT * RayMarchingContext.LocalToWorldScale;
+		float3 WorldPosition_DDY = RayMarchingContext.WorldRayOrigin + RayMarchingContext.WorldRayDirection_DDY * LocalHitT * RayMarchingContext.LocalToWorldScale;
+		WorldPosition_DDX -= WorldPosition;
+		WorldPosition_DDY -= WorldPosition;
+#else
+		float3 WorldPosition_DDX = 0.0;
+		float3 WorldPosition_DDY = 0.0;
+#endif // USE_ANALYTIC_DERIVATIVES
 
-		FVolumeSampleContext SampleContext = CreateVolumeSampleContext(LocalPosition, WorldPosition, RayMarchingContext.MipLevel);
+		FVolumeSampleContext SampleContext = CreateVolumeSampleContext(LocalPosition, WorldPosition, WorldPosition_DDX, WorldPosition_DDY, RayMarchingContext.MipLevel);
 		float3 Extinction = SampleExtinction(SampleContext);
 		Transmittance *= exp(-Extinction * RayMarchingContext.StepSize);
 
@@ -301,8 +310,17 @@ void RayMarchEmissionAbsorption(
 		float LocalHitT = RayMarchingContext.LocalRayTMin + RayMarchingContext.StepSize * (RayMarchingContext.Jitter + StepIndex);
 		float3 LocalPosition = RayMarchingContext.LocalRayOrigin + RayMarchingContext.LocalRayDirection * LocalHitT;
 		float3 WorldPosition = RayMarchingContext.WorldRayOrigin + RayMarchingContext.WorldRayDirection * LocalHitT * RayMarchingContext.LocalToWorldScale;
+#if USE_ANALYTIC_DERIVATIVES
+		float3 WorldPosition_DDX = RayMarchingContext.WorldRayOrigin + RayMarchingContext.WorldRayDirection_DDX * LocalHitT * RayMarchingContext.LocalToWorldScale;
+		float3 WorldPosition_DDY = RayMarchingContext.WorldRayOrigin + RayMarchingContext.WorldRayDirection_DDY * LocalHitT * RayMarchingContext.LocalToWorldScale;
+		WorldPosition_DDX -= WorldPosition;
+		WorldPosition_DDY -= WorldPosition;
+#else
+		float3 WorldPosition_DDX = 0.0;
+		float3 WorldPosition_DDY = 0.0;
+#endif // USE_ANALYTIC_DERIVATIVES
 
-		FVolumeSampleContext SampleContext = CreateVolumeSampleContext(LocalPosition, WorldPosition, RayMarchingContext.MipLevel);
+		FVolumeSampleContext SampleContext = CreateVolumeSampleContext(LocalPosition, WorldPosition, WorldPosition_DDX, WorldPosition_DDY, RayMarchingContext.MipLevel);
 		float3 Extinction = SampleExtinction(SampleContext);
 
 		float3 Emission = SampleEmission(SampleContext);
@@ -359,7 +377,17 @@ void RayMarchSingleScattering(
 		int StochasticFilteringMode = GetStochasticFilteringMode();
 		FVolumeSampleContext SampleContext = CreateVolumeSampleContext(LocalPosition, WorldPosition, FilterWidth, RayMarchingContext.MipLevel, Rand, StochasticFilteringMode);
 #else
-		FVolumeSampleContext SampleContext = CreateVolumeSampleContext(LocalPosition, WorldPosition, RayMarchingContext.MipLevel);
+	#if USE_ANALYTIC_DERIVATIVES
+		float3 WorldPosition_DDX = RayMarchingContext.WorldRayOrigin + RayMarchingContext.WorldRayDirection_DDX * WorldHitT;
+		float3 WorldPosition_DDY = RayMarchingContext.WorldRayOrigin + RayMarchingContext.WorldRayDirection_DDY * WorldHitT;
+		WorldPosition_DDX -= WorldPosition;
+		WorldPosition_DDY -= WorldPosition;
+	#else
+		float3 WorldPosition_DDX = 0.0;
+		float3 WorldPosition_DDY = 0.0;
+	#endif // USE_ANALYTIC_DERIVATIVES
+
+		FVolumeSampleContext SampleContext = CreateVolumeSampleContext(LocalPosition, WorldPosition, WorldPosition_DDX, WorldPosition_DDY, RayMarchingContext.MipLevel);
 #endif // IS_OFFLINE_RENDER
 
 		if (RayMarchingContext.bApplyEmissionAndTransmittance)
@@ -519,7 +547,17 @@ void RayMarchSingleScatteringReferenceFastPath(
 		float WorldHitT = LocalHitT * RayMarchingContext.LocalToWorldScale;
 		float3 LocalPosition = RayMarchingContext.LocalRayOrigin + RayMarchingContext.LocalRayDirection * LocalHitT;
 		float3 WorldPosition = RayMarchingContext.WorldRayOrigin + RayMarchingContext.WorldRayDirection * WorldHitT;
-		FVolumeSampleContext SampleContext = CreateVolumeSampleContext(LocalPosition, WorldPosition, RayMarchingContext.MipLevel);
+#if USE_ANALYTIC_DERIVATIVES
+		float3 WorldPosition_DDX = RayMarchingContext.WorldRayOrigin + RayMarchingContext.WorldRayDirection_DDX * WorldHitT;
+		float3 WorldPosition_DDY = RayMarchingContext.WorldRayOrigin + RayMarchingContext.WorldRayDirection_DDY * WorldHitT;
+		WorldPosition_DDX -= WorldPosition;
+		WorldPosition_DDY -= WorldPosition;
+#else
+		float3 WorldPosition_DDX = 0.0;
+		float3 WorldPosition_DDY = 0.0;
+#endif // USE_ANALYTIC_DERIVATIVES
+
+		FVolumeSampleContext SampleContext = CreateVolumeSampleContext(LocalPosition, WorldPosition, WorldPosition_DDX, WorldPosition_DDY, RayMarchingContext.MipLevel);
 
 		// Radiance from emission
 		if (RayMarchingContext.bApplyEmissionAndTransmittance)
diff --git a/Engine/Shaders/Private/HeterogeneousVolumes/HeterogeneousVolumesSparseVoxelUniformBufferUtils.ush b/Engine/Shaders/Private/HeterogeneousVolumes/HeterogeneousVolumesSparseVoxelUniformBufferUtils.ush
--- a/Engine/Shaders/Private/HeterogeneousVolumes/HeterogeneousVolumesSparseVoxelUniformBufferUtils.ush
+++ b/Engine/Shaders/Private/HeterogeneousVolumes/HeterogeneousVolumesSparseVoxelUniformBufferUtils.ush
@@ -88,7 +88,13 @@ struct FVolumeSampleContext
 	float MipLevel;
 };
 
-FVolumeSampleContext CreateVolumeSampleContext(float3 LocalPosition, float3 WorldPosition, float MipLevel)
+FVolumeSampleContext CreateVolumeSampleContext(
+	float3 LocalPosition,
+	float3 WorldPosition,
+	float3 WorldPosition_DDX,
+	float3 WorldPosition_DDY,
+	float MipLevel
+)
 {
 	float3 LocalBoundsMin = GetLocalBoundsOrigin() - GetLocalBoundsExtent();
 	float3 LocalBoundsMax = GetLocalBoundsOrigin() + GetLocalBoundsExtent();
@@ -100,6 +106,20 @@ FVolumeSampleContext CreateVolumeSampleContext(float3 LocalPosition, float3 Worl
 	return VolumeSampleContext;
 }
 
+FVolumeSampleContext CreateVolumeSampleContext(float3 LocalPosition, float3 WorldPosition, float MipLevel)
+{
+	float3 WorldPosition_DDX = 0.0;
+	float3 WorldPosition_DDY = 0.0;
+
+	return CreateVolumeSampleContext(
+		LocalPosition,
+		WorldPosition,
+		WorldPosition_DDX,
+		WorldPosition_DDY,
+		MipLevel
+	);
+}
+
 FVolumeSampleContext CreateVolumeSampleContext(float3 LocalPosition, float3 WorldPosition, float3 FilterWidth, float MipLevel, float Rand, int SamplingMethod)
 {
 	// TODO: Use FilterWidth determine MipLevel
diff --git a/Engine/Shaders/Private/HeterogeneousVolumes/HeterogeneousVolumesVoxelGridShadows.usf b/Engine/Shaders/Private/HeterogeneousVolumes/HeterogeneousVolumesVoxelGridShadows.usf
--- a/Engine/Shaders/Private/HeterogeneousVolumes/HeterogeneousVolumesVoxelGridShadows.usf
+++ b/Engine/Shaders/Private/HeterogeneousVolumes/HeterogeneousVolumesVoxelGridShadows.usf
@@ -175,7 +175,13 @@ FVolumeSampleContext CreateVolumeSampleContext(float3 LocalPosition, float3 Worl
 	return VolumeSampleContext;
 }
 
-FVolumeSampleContext CreateVolumeSampleContext(float3 LocalPosition, float3 WorldPosition, float MipLevel)
+FVolumeSampleContext CreateVolumeSampleContext(
+	float3 LocalPosition,
+	float3 WorldPosition,
+	float3 WorldPosition_DDX,
+	float3 WorldPosition_DDY,
+	float MipLevel
+)
 {
 	FVolumeSampleContext VolumeSampleContext;
 	VolumeSampleContext.WorldPosition = WorldPosition;
@@ -183,6 +189,20 @@ FVolumeSampleContext CreateVolumeSampleContext(float3 LocalPosition, float3 Worl
 	return VolumeSampleContext;
 }
 
+FVolumeSampleContext CreateVolumeSampleContext(float3 LocalPosition, float3 WorldPosition, float MipLevel)
+{
+	float3 WorldPosition_DDX = 0;
+	float3 WorldPosition_DDY = 0;
+
+	return CreateVolumeSampleContext(
+		LocalPosition,
+		WorldPosition,
+		WorldPosition_DDX,
+		WorldPosition_DDY,
+		MipLevel
+	);
+}
+
 float3 SampleExtinction(inout FVolumeSampleContext Context)
 {
 	return EvalExtinction(Context.WorldPosition);
diff --git a/Engine/Source/Runtime/Renderer/Private/HeterogeneousVolumes/HeterogeneousVolumes.cpp b/Engine/Source/Runtime/Renderer/Private/HeterogeneousVolumes/HeterogeneousVolumes.cpp
--- a/Engine/Source/Runtime/Renderer/Private/HeterogeneousVolumes/HeterogeneousVolumes.cpp
+++ b/Engine/Source/Runtime/Renderer/Private/HeterogeneousVolumes/HeterogeneousVolumes.cpp
@@ -258,6 +258,13 @@ static TAutoConsoleVariable<int32> CVarHeterogeneousVolumesApplyFogInscatteringM
 	ECVF_RenderThreadSafe
 );
 
+static TAutoConsoleVariable<int32> CVarHeterogeneousVolumesUseAnalyticDerivatives(
+	TEXT("r.HeterogeneousVolumes.UseAnalyticDerivatives"),
+	1,
+	TEXT("Writes Heterogeneous Volumes velocity to the feature buffer (Default = 1)"),
+	ECVF_RenderThreadSafe
+);
+
 static TAutoConsoleVariable<int32> CVarHeterogeneousVolumesVelocity(
 	TEXT("r.HeterogeneousVolumes.Velocity"),
 	0,
@@ -563,6 +570,11 @@ namespace HeterogeneousVolumes
 		return static_cast<EFogMode>(FMath::Clamp(CVarHeterogeneousVolumesApplyFogInscatteringMode.GetValueOnRenderThread(), 0, 2));
 	}
 
+	bool UseAnalyticDerivatives()
+	{
+		return CVarHeterogeneousVolumesUseAnalyticDerivatives.GetValueOnRenderThread() != 0;
+	}
+
 	bool ShouldWriteVelocity()
 	{
 		return CVarHeterogeneousVolumesVelocity.GetValueOnRenderThread() != 0;
diff --git a/Engine/Source/Runtime/Renderer/Private/HeterogeneousVolumes/HeterogeneousVolumes.h b/Engine/Source/Runtime/Renderer/Private/HeterogeneousVolumes/HeterogeneousVolumes.h
--- a/Engine/Source/Runtime/Renderer/Private/HeterogeneousVolumes/HeterogeneousVolumes.h
+++ b/Engine/Source/Runtime/Renderer/Private/HeterogeneousVolumes/HeterogeneousVolumes.h
@@ -123,6 +123,7 @@ namespace HeterogeneousVolumes
 	bool SupportsOverlappingVolumes();
 	bool EnableAmbientOcclusion();
 	bool UseExistenceMask();
+	bool UseAnalyticDerivatives();
 
 	enum class EFogMode
 	{
diff --git a/Engine/Source/Runtime/Renderer/Private/HeterogeneousVolumes/HeterogeneousVolumesLiveShadingPipeline.cpp b/Engine/Source/Runtime/Renderer/Private/HeterogeneousVolumes/HeterogeneousVolumesLiveShadingPipeline.cpp
--- a/Engine/Source/Runtime/Renderer/Private/HeterogeneousVolumes/HeterogeneousVolumesLiveShadingPipeline.cpp
+++ b/Engine/Source/Runtime/Renderer/Private/HeterogeneousVolumes/HeterogeneousVolumesLiveShadingPipeline.cpp
@@ -351,6 +351,7 @@ class FRenderSingleScatteringWithLiveShadingCS : public FMeshMaterialShader
 	//class FBilinearInterpolation : SHADER_PERMUTATION_INT("AVSM_BILINEAR_INTERPOLATION", 2);
 	class FIsOfflineRender : SHADER_PERMUTATION_INT("IS_OFFLINE_RENDER", 2);
 	class FApplyFogInscattering : SHADER_PERMUTATION_INT("APPLY_FOG_INSCATTERING", 3);
+	class FUseAnalyticDerivatives : SHADER_PERMUTATION_BOOL("USE_ANALYTIC_DERIVATIVES");
 	using FPermutationDomain = TShaderPermutationDomain<
 		FUseTransmittanceVolume, 
 		FUseInscatteringVolume, 
@@ -360,7 +361,8 @@ class FRenderSingleScatteringWithLiveShadingCS : public FMeshMaterialShader
 		FApplyFogInscattering, 
 		FAVSMSampleMode, 
 		FSupportOverlappingVolumes,
-		FIsOfflineRender
+		FIsOfflineRender,
+		FUseAnalyticDerivatives
 		//FBilinearInterpolation,
 		//FAdaptiveMarch
 	>;
@@ -1267,6 +1269,7 @@ void RenderSingleScatteringWithLiveShading(
 	//PermutationVector.template Set<typename FRenderSingleScatteringWithLiveShadingDispatchTypeCS::FAdaptiveMarch>(static_cast<int32>(HeterogeneousVolumes::ShouldAdaptiveMarch()));
 	PermutationVector.template Set<typename FRenderSingleScatteringWithLiveShadingDispatchTypeCS::FIsOfflineRender>(View.bIsOfflineRender);
 	PermutationVector.template Set<typename FRenderSingleScatteringWithLiveShadingDispatchTypeCS::FApplyFogInscattering>(static_cast<int32>(HeterogeneousVolumes::GetApplyFogInscattering()));
+	PermutationVector.template Set<typename FRenderSingleScatteringWithLiveShadingDispatchTypeCS::FUseAnalyticDerivatives>(HeterogeneousVolumes::UseAnalyticDerivatives());
 	PermutationVector = FRenderSingleScatteringWithLiveShadingDispatchTypeCS::RemapPermutation(PermutationVector);
 	TShaderRef<FRenderSingleScatteringWithLiveShadingDispatchTypeCS> ComputeShader = Material.GetShader<FRenderSingleScatteringWithLiveShadingDispatchTypeCS>(&FLocalVertexFactory::StaticType, PermutationVector, false);
 	if (!ComputeShader.IsNull())
@@ -1695,7 +1698,8 @@ class FRenderVolumetricShadowMapForLightWithLiveShadingCS : public FMeshMaterial
 
 	class FUseAVSMCompression : SHADER_PERMUTATION_BOOL("USE_AVSM_COMPRESSION");
 	class FUseCameraSceneDepth : SHADER_PERMUTATION_BOOL("USE_CAMERA_SCENE_DEPTH");
-	using FPermutationDomain = TShaderPermutationDomain<FUseAVSMCompression, FUseCameraSceneDepth>;
+	class FUseAnalyticDerivatives : SHADER_PERMUTATION_BOOL("USE_ANALYTIC_DERIVATIVES");
+	using FPermutationDomain = TShaderPermutationDomain<FUseAVSMCompression, FUseCameraSceneDepth, FUseAnalyticDerivatives>;
 
 
 	BEGIN_SHADER_PARAMETER_STRUCT(FParameters, )
@@ -2011,6 +2015,7 @@ bool RenderVolumetricShadowMapForLightForHeterogeneousVolumeWithLiveShading(
 	FRenderVolumetricShadowMapForLightWithLiveShadingCS::FPermutationDomain PermutationVector;
 	PermutationVector.Set<FRenderVolumetricShadowMapForLightWithLiveShadingCS::FUseAVSMCompression>(HeterogeneousVolumes::UseAVSMCompression());
 	PermutationVector.Set<FRenderVolumetricShadowMapForLightWithLiveShadingCS::FUseCameraSceneDepth>(bUseCameraSceneDepth && HeterogeneousVolumes::ShadowsUseCameraSceneDepth());
+	PermutationVector.Set<FRenderVolumetricShadowMapForLightWithLiveShadingCS::FUseAnalyticDerivatives>(HeterogeneousVolumes::UseAnalyticDerivatives());
 	TShaderRef<FRenderVolumetricShadowMapForLightWithLiveShadingCS> ComputeShader = Material.GetShader<FRenderVolumetricShadowMapForLightWithLiveShadingCS>(&FLocalVertexFactory::StaticType, PermutationVector, false);
 	if (!ComputeShader.IsNull())
 	{
