Workaround to disable the new path that puts the gpu scene writer parameters in an UB as this is failing on DX11 RHI for unknown reasons.

#rnx

diff --git a/Engine/Shaders/Private/GPUScene/GPUSceneWriter.ush b/Engine/Shaders/Private/GPUScene/GPUSceneWriter.ush
--- a/Engine/Shaders/Private/GPUScene/GPUSceneWriter.ush
+++ b/Engine/Shaders/Private/GPUScene/GPUSceneWriter.ush
@@ -13,29 +13,56 @@
 
 #define USE_GLOBAL_GPU_SCENE_DATA_RW 1
 
+// Workaround for failing tests on DX11 RHI.This negates the hiding of internals for GPU-Scene.
+#define USE_GLOBAL_GPU_SCENE_DATA_RW_UB 0
+
 #include "../SceneData.ush"
 
+#if !USE_GLOBAL_GPU_SCENE_DATA_RW_UB
+uint GPUSceneNumAllocatedInstances;
+uint GPUSceneNumAllocatedPrimitives;
+#endif
+
 void StoreInstanceSceneDataElement(uint InstanceId, uint ElementIndex, float4 DataElement)
 {
+#if USE_GLOBAL_GPU_SCENE_DATA_RW_UB
 	checkSlow(InstanceId < GPUSceneWriter.GPUSceneNumAllocatedInstances);
 	GPUSceneWriter.GPUSceneInstanceSceneDataRW[ElementIndex * GPUSceneWriter.GPUSceneInstanceSceneDataSOAStride + InstanceId] = DataElement;
+#else
+	checkSlow(InstanceId < GPUSceneNumAllocatedInstances);
+	GPUSceneInstanceSceneDataRW[ElementIndex * GPUSceneInstanceSceneDataSOAStride + InstanceId] = DataElement;
+#endif
 }
 
 void StoreInstanceSceneDataElementX(uint InstanceId, uint ElementIndex, float Data)
 {
+#if USE_GLOBAL_GPU_SCENE_DATA_RW_UB
 	// NOTE: Structured buffers allow for DWORD granularity storage, so you don't have to write the whole float4
 	checkSlow(InstanceId < GPUSceneWriter.GPUSceneNumAllocatedInstances);
 	GPUSceneWriter.GPUSceneInstanceSceneDataRW[ElementIndex * GPUSceneWriter.GPUSceneInstanceSceneDataSOAStride + InstanceId].x = Data;
+#else
+	// NOTE: Structured buffers allow for DWORD granularity storage, so you don't have to write the whole float4
+	checkSlow(InstanceId < GPUSceneNumAllocatedInstances);
+	GPUSceneInstanceSceneDataRW[ElementIndex * GPUSceneInstanceSceneDataSOAStride + InstanceId].x = Data;
+#endif
 }
 
 void StoreInstancePayloadDataElement(uint PayloadDataOffset, uint ElementIndex, float4 DataElement)
 {
+#if USE_GLOBAL_GPU_SCENE_DATA_RW_UB
 	GPUSceneWriter.GPUSceneInstancePayloadDataRW[PayloadDataOffset + ElementIndex] = DataElement;
+#else
+	GPUSceneInstancePayloadDataRW[PayloadDataOffset + ElementIndex] = DataElement;
+#endif
 }
 
 void StoreInstancePayloadDataElementZW(uint PayloadDataOffset, uint ElementIndex, float2 DataElement)
 {
+#if USE_GLOBAL_GPU_SCENE_DATA_RW_UB
 	GPUSceneWriter.GPUSceneInstancePayloadDataRW[PayloadDataOffset + ElementIndex].zw = DataElement;
+#else
+	GPUSceneInstancePayloadDataRW[PayloadDataOffset + ElementIndex].zw = DataElement;
+#endif
 }
 
 /**
@@ -43,7 +70,11 @@ void StoreInstancePayloadDataElementZW(uint PayloadDataOffset, uint ElementIndex
  */
 void LoadInstancePrimitiveIdAndFlags(uint InstanceId, inout uint OutPrimitiveId, inout uint OutInstanceFlags)
 {
-	LoadInstancePrimitiveIdAndFlags(InstanceId, GPUSceneWriter.GPUSceneInstanceSceneDataSOAStride, OutPrimitiveId, OutInstanceFlags);
+#if USE_GLOBAL_GPU_SCENE_DATA_RW_UB
+		LoadInstancePrimitiveIdAndFlags(InstanceId, GPUSceneWriter.GPUSceneInstanceSceneDataSOAStride, OutPrimitiveId, OutInstanceFlags);
+#else
+		LoadInstancePrimitiveIdAndFlags(InstanceId, GPUSceneInstanceSceneDataSOAStride, OutPrimitiveId, OutInstanceFlags);
+#endif
 }
 
 void WriteInstancePrimitiveIdAndFlags(uint InstanceId, uint PrimitiveId, uint InstanceFlags)
@@ -168,8 +199,14 @@ void InitializeInstanceSceneDataLS(
  FInstancePayloadDataOffsets GetInstancePayloadDataOffsets(uint InstanceId)
 {
 	uint PrimitiveId, InstanceFlags, InstanceRelativeId, CustomDataCount;
+
+#if USE_GLOBAL_GPU_SCENE_DATA_RW_UB
 	LoadInstancePrimitiveIdAndFlags(InstanceId, GPUSceneWriter.GPUSceneInstanceSceneDataSOAStride, PrimitiveId, InstanceFlags);
 	LoadInstanceRelativeIdAndCustomDataCount(InstanceId, GPUSceneWriter.GPUSceneInstanceSceneDataSOAStride, InstanceRelativeId, CustomDataCount);
+#else
+	LoadInstancePrimitiveIdAndFlags(InstanceId, GPUSceneInstanceSceneDataSOAStride, PrimitiveId, InstanceFlags);
+	LoadInstanceRelativeIdAndCustomDataCount(InstanceId, GPUSceneInstanceSceneDataSOAStride, InstanceRelativeId, CustomDataCount);
+#endif
 	return GetInstancePayloadDataOffsets(PrimitiveId, InstanceFlags, InstanceRelativeId);
 }
 
diff --git a/Engine/Shaders/Private/SceneData.ush b/Engine/Shaders/Private/SceneData.ush
--- a/Engine/Shaders/Private/SceneData.ush
+++ b/Engine/Shaders/Private/SceneData.ush
@@ -309,6 +309,11 @@ FInstanceSceneData WaveReadLaneAt(FInstanceSceneData In, uint SrcIndex)
 // no access to global buffers
 #elif USE_GLOBAL_GPU_SCENE_DATA
 StructuredBuffer<float4> GPUScenePrimitiveSceneData;
+#elif USE_GLOBAL_GPU_SCENE_DATA_RW
+#if !USE_GLOBAL_GPU_SCENE_DATA_RW_UB
+RWStructuredBuffer<float4> GPUScenePrimitiveSceneDataRW;
+#endif
+
 #endif
 
 float4 LoadPrimitivePrimitiveSceneDataElement(uint PrimitiveIndex, uint ItemIndex)
@@ -321,8 +326,14 @@ float4 LoadPrimitivePrimitiveSceneDataElement(uint PrimitiveIndex, uint ItemInde
 	checkStructuredBufferAccessSlow(GPUScenePrimitiveSceneData, TargetIdx);
 	return GPUScenePrimitiveSceneData[TargetIdx];
 #elif USE_GLOBAL_GPU_SCENE_DATA_RW
+#if USE_GLOBAL_GPU_SCENE_DATA_RW_UB
 	checkStructuredBufferAccessSlow(GPUSceneWriter.GPUScenePrimitiveSceneDataRW, TargetIdx);
 	return GPUSceneWriter.GPUScenePrimitiveSceneDataRW[TargetIdx];
+#else
+	checkStructuredBufferAccessSlow(GPUScenePrimitiveSceneDataRW, TargetIdx);
+	return GPUScenePrimitiveSceneDataRW[TargetIdx];
+#endif
+
 #else
 	checkStructuredBufferAccessSlow(Scene.GPUScene.GPUScenePrimitiveSceneData, TargetIdx);
 	return Scene.GPUScene.GPUScenePrimitiveSceneData[TargetIdx];
@@ -557,6 +568,13 @@ struct FInstancePayloadDataOffsets
 StructuredBuffer<float4> GPUSceneInstanceSceneData;
 StructuredBuffer<float4> GPUSceneInstancePayloadData;
 uint GPUSceneFrameNumber;
+#elif USE_GLOBAL_GPU_SCENE_DATA_RW
+#if !USE_GLOBAL_GPU_SCENE_DATA_RW_UB
+RWStructuredBuffer<float4> GPUSceneInstanceSceneDataRW;
+RWStructuredBuffer<float4> GPUSceneInstancePayloadDataRW;
+uint GPUSceneFrameNumber;
+uint GPUSceneInstanceSceneDataSOAStride;
+#endif
 #endif
 
 uint GetGPUSceneFrameNumber()
@@ -564,7 +582,12 @@ uint GetGPUSceneFrameNumber()
 #if USE_GLOBAL_GPU_SCENE_DATA
 	return GPUSceneFrameNumber;
 #elif USE_GLOBAL_GPU_SCENE_DATA_RW
+
+#if USE_GLOBAL_GPU_SCENE_DATA_RW_UB
 	return GPUSceneWriter.GPUSceneFrameNumber;
+#else
+	return GPUSceneFrameNumber;
+#endif
 #else
 	return Scene.GPUScene.GPUSceneFrameNumber;
 #endif
@@ -577,7 +600,12 @@ float4 LoadInstanceSceneDataElement(uint Index)
 #elif USE_GLOBAL_GPU_SCENE_DATA
 	return GPUSceneInstanceSceneData[Index];
 #elif USE_GLOBAL_GPU_SCENE_DATA_RW
+
+#if USE_GLOBAL_GPU_SCENE_DATA_RW_UB
 	return GPUSceneWriter.GPUSceneInstanceSceneDataRW[Index];
+#else
+	return GPUSceneInstanceSceneDataRW[Index];
+#endif
 #else
 	return Scene.GPUScene.GPUSceneInstanceSceneData[Index];
 #endif
@@ -590,7 +618,12 @@ float4 LoadInstancePayloadDataElement(uint Index)
 #elif USE_GLOBAL_GPU_SCENE_DATA
 	return GPUSceneInstancePayloadData[Index];
 #elif USE_GLOBAL_GPU_SCENE_DATA_RW
+
+#if USE_GLOBAL_GPU_SCENE_DATA_RW_UB
 	return GPUSceneWriter.GPUSceneInstancePayloadDataRW[Index];
+#else
+	return GPUSceneInstancePayloadDataRW[Index];
+#endif
 #else
 	return Scene.GPUScene.GPUSceneInstancePayloadData[Index];
 #endif
diff --git a/Engine/Source/Runtime/Renderer/Private/GPUScene.cpp b/Engine/Source/Runtime/Renderer/Private/GPUScene.cpp
--- a/Engine/Source/Runtime/Renderer/Private/GPUScene.cpp
+++ b/Engine/Source/Runtime/Renderer/Private/GPUScene.cpp
@@ -2252,7 +2252,9 @@ class FGPUSceneSetInstancePrimitiveIdCS : public FGlobalShader
 public:
 	BEGIN_SHADER_PARAMETER_STRUCT(FParameters, )
 		SHADER_PARAMETER_STRUCT_INCLUDE(FGPUScene::FInstanceGPULoadBalancer::FShaderParameters, BatcherParameters)
-		SHADER_PARAMETER_RDG_UNIFORM_BUFFER(FGPUSceneWriterUniformParameters, GPUSceneWriterUB)
+		SHADER_PARAMETER_STRUCT_INCLUDE(FGPUSceneWriterParameters, GPUSceneWriterParameters)
+		// TODO: Go back to using the UB. Reverted to use loose parameters due to issues with  DX11 RHI.
+		// SHADER_PARAMETER_RDG_UNIFORM_BUFFER(FGPUSceneWriterUniformParameters, GPUSceneWriterUB)
 	END_SHADER_PARAMETER_STRUCT()
 
 	static constexpr int32 NumThreadsPerGroup = FGPUScene::FInstanceGPULoadBalancer::ThreadGroupSize;
@@ -2281,7 +2283,10 @@ void FGPUScene::AddUpdatePrimitiveIdsPass(FRDGBuilder& GraphBuilder, FInstanceGP
 		FGPUSceneSetInstancePrimitiveIdCS::FParameters* PassParameters = GraphBuilder.AllocParameters<FGPUSceneSetInstancePrimitiveIdCS::FParameters>();
 
 		IdOnlyUpdateItems.Upload(GraphBuilder).GetShaderParameters(GraphBuilder, PassParameters->BatcherParameters);
-		PassParameters->GPUSceneWriterUB = GetWriteParameters(GraphBuilder);
+		
+		PRAGMA_DISABLE_DEPRECATION_WARNINGS
+		GetWriteParameters(GraphBuilder, &PassParameters->GPUSceneWriterParameters);
+		PRAGMA_ENABLE_DEPRECATION_WARNINGS
 		auto ComputeShader = GetGlobalShaderMap(FeatureLevel)->GetShader<FGPUSceneSetInstancePrimitiveIdCS>();
 
 		FComputeShaderUtils::AddPass(
