Add soft transition from glint to GGX to avoid pops when transition from glint to no-glint.

diff --git a/Engine/Shaders/Private/Substrate/SubstrateEvaluation.ush b/Engine/Shaders/Private/Substrate/SubstrateEvaluation.ush
--- a/Engine/Shaders/Private/Substrate/SubstrateEvaluation.ush
+++ b/Engine/Shaders/Private/Substrate/SubstrateEvaluation.ush
@@ -289,6 +289,99 @@ half3 Substrate_F_GGX(half3 F0, half3 F90, half VoH)
 //#endif
 }
 
+// Transition from Glint to GGX (0=Glint <-> 1=GGX)
+float SubstrateTransitionGlintToGGX(float GlintValue)
+{
+#if SUBSTRATE_COMPLEXSPECIALPATH && SUBSTRATE_GLINTS_ALLOWED && SUBSTRATE_GLINTS_ENABLED
+	return saturate(saturate(GlintValue - 0.95f) / 0.05f);
+#else
+	return 1.f;
+#endif
+}
+
+float Substrate_D_Glint(FSubstrateBSDFContext BSDFContext, float SafeRoughness, float GGX_D)
+{
+	float OutD = GGX_D;
+
+	#if SUBSTRATE_COMPLEXSPECIALPATH && SUBSTRATE_GLINTS_ALLOWED && SUBSTRATE_GLINTS_ENABLED
+	BRANCH
+	if (BSDF_GETHASGLINT(BSDFContext.BSDF))
+	{
+		FBeckmannDesc Beckmann = GGXToBeckmann(SafeRoughness);
+		const float Glint_D = f_P(BSDFContext.TangentV, BSDFContext.TangentL, float3(Beckmann.Sigma.xx, Beckmann.Rho), SLAB_GLINT_VALUE(BSDFContext.BSDF), SLAB_GLINT_UV(BSDFContext.BSDF), SLAB_GLINT_UVDDX(BSDFContext.BSDF), SLAB_GLINT_UVDDY(BSDFContext.BSDF));
+
+		// Transition from Glint to GGX
+		OutD = lerp(Glint_D, GGX_D, SubstrateTransitionGlintToGGX(SLAB_GLINT_VALUE(BSDFContext.BSDF)));
+	}
+	#endif // SUBSTRATE_COMPLEXPATH
+
+	return OutD;
+}
+
+float EvaluateGlintRect(FSubstrateBSDFContext BSDFContext, float SafeRoughness, float3 MeanLightWorldDirection)
+{
+	float Out = 1.0f;
+
+	#if SUBSTRATE_COMPLEXSPECIALPATH && SUBSTRATE_GLINTS_ALLOWED && SUBSTRATE_GLINTS_ENABLED
+	BRANCH
+	if (BSDF_GETHASGLINT(BSDFContext.BSDF))
+	{
+		float3 MeanLightLocalDirection = normalize(mul(BSDFContext.TangentBasis, MeanLightWorldDirection));
+		FBeckmannDesc Beckmann = GGXToBeckmann(SafeRoughness);
+
+		// It is wrong to multiply again the GGXevaluation together with the glint function (representing Beckmann D function).
+		// But visually it works and that is an acceptable solution in the meantime we find something more correct.
+		const float GlintSpec = f_P(BSDFContext.TangentV, MeanLightLocalDirection, float3(Beckmann.Sigma.xx, Beckmann.Rho),SLAB_GLINT_VALUE(BSDFContext.BSDF), SLAB_GLINT_UV(BSDFContext.BSDF), SLAB_GLINT_UVDDX(BSDFContext.BSDF), SLAB_GLINT_UVDDY(BSDFContext.BSDF));
+
+		// Transition from Glint to GGX
+		Out = lerp(GlintSpec, 1.f, SubstrateTransitionGlintToGGX(SLAB_GLINT_VALUE(BSDFContext.BSDF)));
+
+	}
+	#endif // SUBSTRATE_COMPLEXPATH
+
+	return Out;
+}
+
+float EvaluateGlintEnv(FSubstrateBSDFContext BSDFContext, float SafeRoughness, inout float3 OutSpecularDirection)
+{
+	float Out = 1.0f;
+
+#if SUBSTRATE_COMPLEXSPECIALPATH && SUBSTRATE_GLINTS_ALLOWED && SUBSTRATE_GLINTS_ENABLED
+	BRANCH
+	if (BSDF_GETHASGLINT(BSDFContext.BSDF))
+	{
+		FBeckmannDesc Beckmann = GGXToBeckmann(SafeRoughness);
+		const float TransitionGlintToGGX = SubstrateTransitionGlintToGGX(SLAB_GLINT_VALUE(BSDFContext.BSDF));
+		float3 MeanLightLocalDirection = normalize(mul(BSDFContext.TangentBasis, OutSpecularDirection));
+
+		// Importance sampling
+		#if SUBSTRATE_GLINTS_IS
+		{
+			const float3 Random = float3(BlueNoiseVec2(BSDFContext.PixelCoord, View.StateFrameIndex), BlueNoiseScalar(BSDFContext.PixelCoord, View.StateFrameIndex));
+					
+			// Sample and refect view according the sampled m
+			const float3 m_local = Sample_P(Random, BSDFContext.TangentV, float3(Beckmann.Sigma, Beckmann.Rho), \
+							SLAB_GLINT_VALUE(BSDFContext.BSDF), SLAB_GLINT_UV(BSDFContext.BSDF), SLAB_GLINT_UVDDX(BSDFContext.BSDF), SLAB_GLINT_UVDDY(BSDFContext.BSDF));
+			MeanLightLocalDirection =  normalize(2 * dot(BSDFContext.TangentV, m_local) * m_local - BSDFContext.TangentV);
+			const float MeanLightWorldDirection = mul(transpose(BSDFContext.TangentBasis), MeanLightLocalDirection);
+			OutSpecularDirection = normalize(lerp(OutSpecularDirection, MeanLightWorldDirection, TransitionGlintToGGX));
+		}
+		#endif
+
+		// It is wrong to multiply again the split env lighting (for split-sum approxiation) together with the glint function (representing Beckmann D function).
+		// But visually it works and that is an acceptable solution in the meantime we find something more correct.
+		float GlintD = f_P(BSDFContext.TangentV, MeanLightLocalDirection, float3(Beckmann.Sigma, Beckmann.Rho), 
+							SLAB_GLINT_VALUE(BSDFContext.BSDF), SLAB_GLINT_UV(BSDFContext.BSDF), SLAB_GLINT_UVDDX(BSDFContext.BSDF), SLAB_GLINT_UVDDY(BSDFContext.BSDF));
+		
+		// We need a way to lerp to the glint direction with a roughness of 0. 
+		// SUBSTRATE_TODO recover glint direction.
+		Out = lerp(saturate(GlintD), 1.f, TransitionGlintToGGX);
+		// We keep the same SpecularSafeRoughness for single glint environment sampling to match the roughness of the surface overall, especially when the glints becomes infinitely small and thus invisible.
+	}
+#endif // SUBSTRATE_COMPLEXSPECIALPATH && SUBSTRATE_GLINTS_ALLOWED && SUBSTRATE_GLINTS_ENABLED
+
+	return Out;
+}
 ///////////////////////////////////////////////////////////////////////////////////////////////////
 // BSDF evaluation (punctual/area)
 
@@ -465,20 +558,14 @@ FSubstrateEvaluateResult SubstrateEvaluateBSDFCommon(FSubstrateBSDFContext BSDFC
 						float2 AlphaXY = 0;
 						GetAnisotropicRoughness(Alpha, SLAB_ANISOTROPY(BSDFContext.BSDF), AlphaXY.x, AlphaXY.y);
 				
-					#if SUBSTRATE_COMPLEXSPECIALPATH && SUBSTRATE_GLINTS_ALLOWED && SUBSTRATE_GLINTS_ENABLED
-						// Glint shading only in the complex path
-						BRANCH
-						if (BSDF_GETHASGLINT(BSDFContext.BSDF))
+						// GGX lobe
 						{
-							const float2 GGXRoughnessXY = float2(sqrtFast(AlphaXY.x), sqrtFast(AlphaXY.y));
-							FBeckmannDesc Beckmann = GGXToBeckmann(GGXRoughnessXY);
-							D = f_P(BSDFContext.TangentV, BSDFContext.TangentL, float3(Beckmann.Sigma, Beckmann.Rho),
-								SLAB_GLINT_VALUE(BSDFContext.BSDF), SLAB_GLINT_UV(BSDFContext.BSDF), SLAB_GLINT_UVDDX(BSDFContext.BSDF), SLAB_GLINT_UVDDY(BSDFContext.BSDF));
+							D = Substrate_D_GGX_Aniso(AlphaXY.x, AlphaXY.y, BSDFContext.Context.NoH, BSDFContext.Context.XoH, BSDFContext.Context.YoH);
 						}
-						else
-					#endif
+
+						// Glint shading only in the complex path
 						{
-							D = Substrate_D_GGX_Aniso(AlphaXY.x, AlphaXY.y, BSDFContext.Context.NoH, BSDFContext.Context.XoH, BSDFContext.Context.YoH);
+							D = Substrate_D_Glint(BSDFContext, SafeRoughness, D);
 						}
 
 						Vis = Substrate_Vis_GGX_Aniso(AlphaXY.x, AlphaXY.y, BSDFContext.Context.NoV, BSDFContext.SatNoL, BSDFContext.Context.XoV, BSDFContext.Context.XoL, BSDFContext.Context.YoV, BSDFContext.Context.YoL);
@@ -495,22 +582,16 @@ FSubstrateEvaluateResult SubstrateEvaluateBSDFCommon(FSubstrateBSDFContext BSDFC
 						// In this case, we set D to 1 and Vis will contain the area light / GGX integration
 						{
 							float3 MeanLightWorldDirection = 0.0f;
-							RectLightSpec = RectGGXApproxLTC(SafeRoughness, F0, F90, BSDFContext.N, BSDFContext.V, AreaLightContext.AreaLight.Rect, AreaLightContext.AreaLight.Texture, MeanLightWorldDirection);
+
+							// GGX lobe
+							{
+								RectLightSpec = RectGGXApproxLTC(SafeRoughness, F0, F90, BSDFContext.N, BSDFContext.V, AreaLightContext.AreaLight.Rect, AreaLightContext.AreaLight.Texture, MeanLightWorldDirection);
+							}
 							
-							// Now combine the rectangular area light with Glints.
-						#if SUBSTRATE_COMPLEXSPECIALPATH && SUBSTRATE_GLINTS_ALLOWED && SUBSTRATE_GLINTS_ENABLED
-							BRANCH
-							if (BSDF_GETHASGLINT(BSDFContext.BSDF))
+							// Glint shading only in the complex path
 							{
-								float3 MeanLightLocalDirection = normalize(mul(BSDFContext.TangentBasis, MeanLightWorldDirection));
-								FBeckmannDesc Beckmann = GGXToBeckmann(SafeRoughness);
-
-								// It is wrong to multiply again the GGXevaluation together with the glint function (representing Beckmann D function).
-								// But visually it works and that is an acceptable solution in the meantime we find something more correct.
-								RectLightSpec *= f_P(BSDFContext.TangentV, MeanLightLocalDirection, float3(Beckmann.Sigma.xx, Beckmann.Rho), 
-									SLAB_GLINT_VALUE(BSDFContext.BSDF), SLAB_GLINT_UV(BSDFContext.BSDF), SLAB_GLINT_UVDDX(BSDFContext.BSDF), SLAB_GLINT_UVDDY(BSDFContext.BSDF));
+								RectLightSpec *= EvaluateGlintRect(BSDFContext, SafeRoughness, MeanLightWorldDirection);
 							}
-						#endif // SUBSTRATE_COMPLEXPATH
 
 							const float H_PDF = VisibleGGXPDF(BSDFContext.TangentV, BSDFContext.TangentH, Alpha2Spec);
 							PDF = RayPDFToReflectionRayPDF(BSDFContext.Context.VoH, H_PDF);
@@ -524,36 +605,33 @@ FSubstrateEvaluateResult SubstrateEvaluateBSDFCommon(FSubstrateBSDFContext BSDFC
 						{
 							if(IntegrationType == INTEGRATION_PUNCTUAL_LIGHT)
 							{
-							#if SUBSTRATE_COMPLEXSPECIALPATH && SUBSTRATE_GLINTS_ALLOWED && SUBSTRATE_GLINTS_ENABLED
-								BRANCH
-								if (BSDF_GETHASGLINT(BSDFContext.BSDF))
+								// GGX lobe
 								{
-									FBeckmannDesc Beckmann = GGXToBeckmann(SafeRoughness);
-									D = f_P(BSDFContext.TangentV, BSDFContext.TangentL, float3(Beckmann.Sigma.xx, Beckmann.Rho),
-										SLAB_GLINT_VALUE(BSDFContext.BSDF), SLAB_GLINT_UV(BSDFContext.BSDF), SLAB_GLINT_UVDDX(BSDFContext.BSDF), SLAB_GLINT_UVDDY(BSDFContext.BSDF));
+									D = Substrate_D_GGX(SafeRoughness, Alpha2Spec, BSDFContext.Context.NoH);
 								}
-								else
-							#endif // SUBSTRATE_COMPLEXPATH
+
+								// Glint shading only in the complex path
 								{
-									D = Substrate_D_GGX(SafeRoughness, Alpha2Spec, BSDFContext.Context.NoH);
+									D = Substrate_D_Glint(BSDFContext, SafeRoughness, D);
 								}
 							}
 							else
 							{
-								const float Energy = EnergyNormalization(Alpha2Spec, BSDFContext.Context.VoH, AreaLightContext.AreaLight);
 							
-							#if SUBSTRATE_COMPLEXSPECIALPATH && SUBSTRATE_GLINTS_ALLOWED && SUBSTRATE_GLINTS_ENABLED
-								BRANCH
-								if (BSDF_GETHASGLINT(BSDFContext.BSDF))
+								// GGX lobe
+								{
+									D = Substrate_D_GGX(SafeRoughness, Alpha2Spec, BSDFContext.Context.NoH);
+								}
+
+								// Glint shading only in the complex path
 								{
-									FBeckmannDesc Beckmann = GGXToBeckmann(SafeRoughness);
-									D = Energy * f_P(BSDFContext.TangentV, BSDFContext.TangentL, float3(Beckmann.Sigma.xx, Beckmann.Rho),
-										SLAB_GLINT_VALUE(BSDFContext.BSDF), SLAB_GLINT_UV(BSDFContext.BSDF), SLAB_GLINT_UVDDX(BSDFContext.BSDF), SLAB_GLINT_UVDDY(BSDFContext.BSDF));
+									D = Substrate_D_Glint(BSDFContext, SafeRoughness, D);
 								}
-								else
-							#endif // SUBSTRATE_COMPLEXPATH
+
+								// Area light energy compensation
 								{
-									D = Energy * Substrate_D_GGX(SafeRoughness, Alpha2Spec, BSDFContext.Context.NoH);
+									const float Energy = EnergyNormalization(Alpha2Spec, BSDFContext.Context.VoH, AreaLightContext.AreaLight);
+									D *= Energy;
 								}
 							}
 
@@ -1552,41 +1630,15 @@ FSubstrateEnvLightResult SubstrateEvaluateForEnvLight(FSubstrateBSDFContext BSDF
 				SubstrateEnvLightResult.SpecularSafeRoughness = SafeRoughness;
 				SubstrateEnvLightResult.SpecularColor = F0;
 
-				// SUBSTRATE_TODO environment light importance sampling
-				// Glints
+				// Glint shading
 			#if SUBSTRATE_COMPLEXSPECIALPATH && SUBSTRATE_GLINTS_ALLOWED && SUBSTRATE_GLINTS_ENABLED
-				BRANCH
-				if (BSDF_GETHASGLINT(BSDFContext.BSDF))
 				{
-					FBeckmannDesc Beckmann = GGXToBeckmann(SafeRoughness);
-
-					float3 MeanLightLocalDirection = normalize(mul(BSDFContext.TangentBasis, SubstrateEnvLightResult.SpecularDirection));
-					#if SUBSTRATE_GLINTS_IS
-					{
-						const float3 Random = float3(BlueNoiseVec2(BSDFContext.PixelCoord, View.StateFrameIndex), BlueNoiseScalar(BSDFContext.PixelCoord, View.StateFrameIndex));
-					
-						// Sample and refect view according the sampled m
-						const float3 m_local = Sample_P(Random, BSDFContext.TangentV, float3(Beckmann.Sigma, Beckmann.Rho), \
-										SLAB_GLINT_VALUE(BSDFContext.BSDF), SLAB_GLINT_UV(BSDFContext.BSDF), SLAB_GLINT_UVDDX(BSDFContext.BSDF), SLAB_GLINT_UVDDY(BSDFContext.BSDF));
-						MeanLightLocalDirection =  normalize(2 * dot(BSDFContext.TangentV, m_local) * m_local - BSDFContext.TangentV);	
-						SubstrateEnvLightResult.SpecularDirection = normalize(mul(transpose(BSDFContext.TangentBasis), MeanLightLocalDirection));
-					}
-					#endif
-
-					// It is wrong to multiply again the split env lighting (for split-sum approxiation) together with the glint function (representing Beckmann D function).
-					// But visually it works and that is an acceptable solution in the meantime we find something more correct.
-					float GlintD = f_P(BSDFContext.TangentV, MeanLightLocalDirection, float3(Beckmann.Sigma, Beckmann.Rho), 
-										SLAB_GLINT_VALUE(BSDFContext.BSDF), SLAB_GLINT_UV(BSDFContext.BSDF), SLAB_GLINT_UVDDX(BSDFContext.BSDF), SLAB_GLINT_UVDDY(BSDFContext.BSDF));
-		
-					// We need a way to lerp to the glint direction with a roughness of 0. 
-					// SUBSTRATE_TODO recover glint direction.
-					SubstrateEnvLightResult.SpecularWeight *= saturate(GlintD);
-					// We keep the same SpecularSafeRoughness for single glint environment sampling to match the roughness of the surface overall, especially when the glints becomes infinitely small and thus invisible.
+					SubstrateEnvLightResult.SpecularWeight *= EvaluateGlintEnv(BSDFContext, SubstrateEnvLightResult.SpecularSafeRoughness, SubstrateEnvLightResult.SpecularDirection);
 				}
-			#endif // SUBSTRATE_COMPLEXSPECIALPATH && SUBSTRATE_GLINTS_ALLOWED && SUBSTRATE_GLINTS_ENABLED
+			#endif				
 
-			#if SUBSTRATE_COMPLEXSPECIALPATH && SUBSTRATE_SPECPROFILE_ENABLED
 				// Specular profile
+			#if SUBSTRATE_COMPLEXSPECIALPATH && SUBSTRATE_SPECPROFILE_ENABLED
 				BRANCH
 				if (BSDF_GETHASSPECPROFILE(BSDFContext.BSDF))
 				{
