Fixed iOS crash in a project that uses depth prepass and a mobile deferred shading
#rnx

diff --git a/Engine/Shaders/Private/DeferredShadingCommon.ush b/Engine/Shaders/Private/DeferredShadingCommon.ush
--- a/Engine/Shaders/Private/DeferredShadingCommon.ush
+++ b/Engine/Shaders/Private/DeferredShadingCommon.ush
@@ -769,23 +769,18 @@ void MobileFetchGBuffer(in float2 UV, out half4 GBufferA, out half4 GBufferB, ou
 	#if MOBILE_EXTENDED_GBUFFER
 		GBufferD = VulkanSubpassFetch4();
 	#endif
-	SceneDepth = ConvertFromDeviceZ(VulkanSubpassDepthFetch());
 #elif (METAL_PROFILE && !MAC && FRAMEBUFFER_FETCH_ACCESS)
 	GBufferA = SubpassFetchRGBA_1(); 
 	GBufferB = SubpassFetchRGBA_2(); 
 	GBufferC = SubpassFetchRGBA_3(); 
 	#if MOBILE_EXTENDED_GBUFFER
 		GBufferD = SubpassFetchRGBA_4();
-		SceneDepth = ConvertFromDeviceZ(SubpassFetchR_5());
-	#else
-		SceneDepth = ConvertFromDeviceZ(SubpassFetchR_4());
 	#endif
 #elif (USE_GLES_FBF_DEFERRED && FRAMEBUFFER_FETCH_ACCESS)
 	GBufferA = GLSubpassFetch1(); 
 	GBufferB = GLSubpassFetch2(); 
 	GBufferC = GLSubpassFetch3();  
 	GBufferD = 0; // PLS is limited to 128bits
-	SceneDepth = ConvertFromDeviceZ(DepthbufferFetchES2());
 #else
 	GBufferA = Texture2DSampleLevel(MobileSceneTextures.GBufferATexture, MobileSceneTextures.GBufferATextureSampler, UV, 0); 
 	GBufferB = Texture2DSampleLevel(MobileSceneTextures.GBufferBTexture, MobileSceneTextures.GBufferBTextureSampler, UV, 0);
@@ -793,6 +788,13 @@ void MobileFetchGBuffer(in float2 UV, out half4 GBufferA, out half4 GBufferB, ou
 	#if MOBILE_EXTENDED_GBUFFER
 		GBufferD = Texture2DSampleLevel(MobileSceneTextures.GBufferDTexture, MobileSceneTextures.GBufferDTextureSampler, UV, 0);
 	#endif
+#endif
+
+#if	(FRAMEBUFFER_FETCH_ACCESS && VULKAN_PROFILE)
+	SceneDepth = ConvertFromDeviceZ(VulkanSubpassDepthFetch());
+#elif (FRAMEBUFFER_FETCH_ACCESS && (USE_SCENE_DEPTH_AUX || USE_GLES_FBF_DEFERRED))
+	SceneDepth = ConvertFromDeviceZ(DepthbufferFetchES2());
+#else
 	SceneDepth = ConvertFromDeviceZ(Texture2DSampleLevel(MobileSceneTextures.SceneDepthTexture, MobileSceneTextures.SceneDepthTextureSampler, UV, 0).r);
 #endif
 }
diff --git a/Engine/Shaders/Private/SceneTexturesCommon.ush b/Engine/Shaders/Private/SceneTexturesCommon.ush
--- a/Engine/Shaders/Private/SceneTexturesCommon.ush
+++ b/Engine/Shaders/Private/SceneTexturesCommon.ush
@@ -219,7 +219,7 @@ float LookupDeviceZ( float2 ScreenUV, float ArrayIndex )
 	return DepthbufferFetchES2();
 #elif MOBILE_DEPTHFECTH && VULKAN_PROFILE
 	return VulkanSubpassDepthFetch();
-#elif MOBILE_DEPTHFECTH && (METAL_PROFILE && !MAC)
+#elif MOBILE_DEPTHFECTH && (METAL_PROFILE && !MAC) && USE_SCENE_DEPTH_AUX
 	return DepthbufferFetchES2();
 #elif (USE_SCENE_DEPTH_AUX && !MOBILE_DEFERRED_SHADING)
 	// SceneDepth texture is discarded after BasePass (with forward shading) 
diff --git a/Engine/Shaders/Public/Platform/Metal/MetalCommon.ush b/Engine/Shaders/Public/Platform/Metal/MetalCommon.ush
--- a/Engine/Shaders/Public/Platform/Metal/MetalCommon.ush
+++ b/Engine/Shaders/Public/Platform/Metal/MetalCommon.ush
@@ -44,9 +44,14 @@
 	// 4 components: RGBA_0
 	Texture2D<float4> gl_LastFragDataRGBA_0;
 	#define SubpassFetchRGBA_0() gl_LastFragDataRGBA_0.Load(uint3(0, 0, 0), 0)
+	
 	// 1 component: R_4
 	Texture2D<float> gl_LastFragDataR_4;
 	#define SubpassFetchR_4() gl_LastFragDataR_4.Load(uint3(0, 0, 0), 0)
+	
+	// 1 component: R_5
+	Texture2D<float> gl_LastFragDataR_5;
+	#define SubpassFetchR_5() gl_LastFragDataR_5.Load(uint3(0, 0, 0), 0)
 
 	Texture2D<float> gl_LastFragDataR_1;
 	#define SubpassFetchR_1() gl_LastFragDataR_1.Load(uint3(0, 0, 0), 0)
@@ -54,7 +59,11 @@
 	// Rest of SubpassFetch can be found in MetalSubpassSupport.ush
 	//
 	#if MOBILE_DEFERRED_SHADING
-		#define DepthbufferFetchES2() SubpassFetchR_4()
+		#if MOBILE_EXTENDED_GBUFFER
+			#define DepthbufferFetchES2() SubpassFetchR_5()
+		#else
+			#define DepthbufferFetchES2() SubpassFetchR_4()
+		#endif
 	#else
 		#define DepthbufferFetchES2() SubpassFetchR_1()
 	#endif
@@ -62,6 +71,7 @@
 	#define SubpassFetchRGBA_0() ((float4)0.0)
 	#define SubpassFetchR_0() (0.0)
 	#define SubpassFetchR_4() (0.0)
+	#define SubpassFetchR_5() (0.0)
 #endif //IOS && PIXELSHADER
 
 #if METAL_SM6_PROFILE == 1
diff --git a/Engine/Shaders/Public/Platform/Metal/MetalSubpassSupport.ush b/Engine/Shaders/Public/Platform/Metal/MetalSubpassSupport.ush
--- a/Engine/Shaders/Public/Platform/Metal/MetalSubpassSupport.ush
+++ b/Engine/Shaders/Public/Platform/Metal/MetalSubpassSupport.ush
@@ -46,10 +46,7 @@
 	Texture2D<float> gl_LastFragDataR_3;
 	#define SubpassFetchR_3() gl_LastFragDataR_3.Load(uint3(0, 0, 0), 0)
     
-	// R_4 in MetalCommon.ush
-
-	Texture2D<float> gl_LastFragDataR_5;
-	#define SubpassFetchR_5() gl_LastFragDataR_5.Load(uint3(0, 0, 0), 0)
+	// R_4 and R_5 in MetalCommon.ush
 
 	Texture2D<float> gl_LastFragDataR_6;
 	#define SubpassFetchR_6() gl_LastFragDataR_6.Load(uint3(0, 0, 0), 0)
@@ -68,7 +65,6 @@
 	#define SubpassFetchR_1() (0.0)
 	#define SubpassFetchR_2() (0.0)
 	#define SubpassFetchR_3() (0.0)
-	#define SubpassFetchR_5() (0.0)
 	#define SubpassFetchR_6() (0.0)
 	#define SubpassFetchR_7() (0.0)
 #endif //IOS && PIXELSHADER 
