Fixed: Spotlight shadow is missing on mobile deferred
#rnx

diff --git a/Engine/Shaders/Private/MobileDeferredShading.usf b/Engine/Shaders/Private/MobileDeferredShading.usf
--- a/Engine/Shaders/Private/MobileDeferredShading.usf
+++ b/Engine/Shaders/Private/MobileDeferredShading.usf
@@ -249,6 +249,10 @@ FDeferredLightData InitDeferredLightFromLightParameters(FLightShaderParameters I
 	Out.bRectLight				= InLightType == LIGHT_TYPE_RECT;
 	Out.HairTransmittance		= InitHairTransmittanceData();
 
+#if SUPPORT_SPOTLIGHTS_SHADOW
+	Out.ShadowedBits			= 1;
+	Out.ShadowMapChannelMask	= 1.f;
+#endif
 	return Out;
 }
 
@@ -285,7 +289,7 @@ void MobileRadialLightPS(
 	}
 
 	half Shadow = 1.0;
-#if SUPPORT_SPOTLIGHTS_SHADOW && IS_SPOT_LIGHT
+#if SUPPORT_SPOTLIGHTS_SHADOW
 	float3 ToLight = LightData.TranslatedWorldPosition - TranslatedWorldPosition;
 	float3 LocalPosition = -ToLight;
 	float DistanceSqr = dot(ToLight, ToLight);
@@ -320,9 +324,7 @@ void MobileRadialLightPS(
 		Settings.TransitionScale = SpotLightSoftTransitionScale * lerp(SpotLightReceiverDepthBias, 1.0, NoL);
 
 		Shadow = MobileShadowPCF(ShadowUVs, Settings);
-
 		Shadow = saturate((Shadow - 0.5) * SpotLightShadowSharpen + 0.5f);
-
 		Shadow = lerp(1.0f, Square(Shadow), SpotLightFadeFraction);
 	}
 #endif
@@ -333,11 +335,10 @@ void MobileRadialLightPS(
 	DirectLighting = AccumulateDynamicLighting(TranslatedWorldPosition, CameraVector, GBuffer, 1, LightData, LightAttenuation, 0, uint2(0, 0), SurfaceShadow);
 
 	half3 Color = DirectLighting.TotalLight * ComputeLightFunctionMultiplier(TranslatedWorldPosition);
-	
 	// MobileHDR applies PreExposure in tonemapper
 	Color *= View.PreExposure;
-
 	Color = SafeGetOutColor(Color);
+
 #if USE_GLES_FBF_DEFERRED
 	OutProxyAdditive.rgb = Color;
 #else
