Substrate - fix shader compilation warnings.

#tests Tested dx12, warnings were gone
#rnx

diff --git a/Engine/Shaders/Private/Substrate/Glint/GlintThirdParty.ush b/Engine/Shaders/Private/Substrate/Glint/GlintThirdParty.ush
--- a/Engine/Shaders/Private/Substrate/Glint/GlintThirdParty.ush
+++ b/Engine/Shaders/Private/Substrate/Glint/GlintThirdParty.ush
@@ -602,11 +602,12 @@ float Sample_P_1D(FGlintLutParams GlintLutParams, float u, int ij, int l)
 	const int Idx = int(fmod(ij, 3)); // ij % 3; // int(mod(ij, 3))
 
 	// Slow version - the CDF is computed inline within the shader (TODO precompute CDF textures)
+	uint it;
 	const uint EntryCount = 64;
 	float Cdfs[EntryCount+1];
 	float Acc = 0;
 	LOOP
-	for (uint it=0;it<EntryCount;++it)
+	for (it=0;it<EntryCount;++it)
 	{
 		const float texCoordX = saturate(float(it) / float(EntryCount-1));
 		float V = View.GlintTexture.SampleLevel(View.GlintSampler, float3(texCoordX, 0.5, l * distPerChannel + distIdxOver3), 0.0f)[Idx];
@@ -617,7 +618,7 @@ float Sample_P_1D(FGlintLutParams GlintLutParams, float u, int ij, int l)
 	Acc = max(0.0001f, Acc);
 
 	LOOP
-	for (uint it = 0; it < EntryCount; ++it)
+	for (it = 0; it < EntryCount; ++it)
 	{
 		Cdfs[it] /= Acc;
 	}
@@ -627,7 +628,7 @@ float Sample_P_1D(FGlintLutParams GlintLutParams, float u, int ij, int l)
 	// Sampled value is in [0, 1)
 	float Sample = 0;
 	LOOP
-	for (uint it = 0; it < EntryCount; ++it)
+	for (it = 0; it < EntryCount; ++it)
 	{
 		if (Cdfs[it] <= u && u <= Cdfs[it + 1])
 		{
@@ -812,8 +813,9 @@ float2 Sample_P22_LOD(
 		ValidCellCount = min(ValidCellCount, MAX_CELL_COUNT);
 
 		// 1.2 Normalized CDF
+		uint it;
 		LOOP
-		for (uint it = 0; it < ValidCellCount; ++it)
+		for (it = 0; it < ValidCellCount; ++it)
 		{
 			Cdfs[it] /= max(Acc, 0.0001f);
 		}
@@ -821,7 +823,7 @@ float2 Sample_P22_LOD(
 
 		// 1.3 Sample CDF
 		LOOP
-		for (uint it = 0; it < ValidCellCount; ++it)
+		for (it = 0; it < ValidCellCount; ++it)
 		{
 			if (Cdfs[it] <= u0 && u0 <= Cdfs[it + 1])
 			{
