First Person: First person pixels are now tagged in the GBuffer. The bit is located in the SelectiveOutputMask and is shared with the precalculated shadows feature, which is only used when r.AllowStaticLighting is enabled. This means that first person rendering tech relying on this GBuffer bit only works if r.AllowStaticLighting is disabled.

diff --git a/Engine/Shaders/Private/BasePassPixelShader.usf b/Engine/Shaders/Private/BasePassPixelShader.usf
--- a/Engine/Shaders/Private/BasePassPixelShader.usf
+++ b/Engine/Shaders/Private/BasePassPixelShader.usf
@@ -976,6 +976,7 @@ void FPixelShaderInOut_MainPS(
 	SubstratePixelHeader.IrradianceAO.MaterialAO = MaterialAO;
 	SubstratePixelHeader.SetCastContactShadow(GetPrimitiveData(MaterialParameters).Flags & PRIMITIVE_SCENE_DATA_FLAG_HAS_CAST_CONTACT_SHADOW);
 	SubstratePixelHeader.SetDynamicIndirectShadowCasterRepresentation(GetPrimitiveData(MaterialParameters).Flags & PRIMITIVE_SCENE_DATA_FLAG_HAS_CAPSULE_REPRESENTATION);
+	SubstratePixelHeader.SetIsFirstPerson(GetPrimitiveData(MaterialParameters).Flags & PRIMITIVE_SCENE_DATA_FLAG_IS_FIRST_PERSON_PRIMITIVE);
 #endif 
 #endif
 
@@ -1110,6 +1111,13 @@ void FPixelShaderInOut_MainPS(
 	// This requires cleanup. Shader code that uses GBuffer.SelectiveOutputMask expects the outputmask to be in
 	// bits [4:7], but it gets packed as bits [0:3] in the flexible gbuffer since we might move it around.
 	GBuffer.SelectiveOutputMask = GetSelectiveOutputMask() >> 4;
+#if !ALLOW_STATIC_LIGHTING
+	// ZERO_PRECSHADOW_MASK is not used when !ALLOW_STATIC_LIGHTING, so we can alias the bit with bIsFirstPerson.
+	if (GetPrimitiveData(MaterialParameters).Flags & PRIMITIVE_SCENE_DATA_FLAG_IS_FIRST_PERSON_PRIMITIVE)
+	{
+		GBuffer.SelectiveOutputMask |= IS_FIRST_PERSON_MASK >> 4;
+	}
+#endif
 	GBuffer.Velocity = 0;
 #endif
 
diff --git a/Engine/Shaders/Private/DeferredShadingCommon.ush b/Engine/Shaders/Private/DeferredShadingCommon.ush
--- a/Engine/Shaders/Private/DeferredShadingCommon.ush
+++ b/Engine/Shaders/Private/DeferredShadingCommon.ush
@@ -405,6 +405,15 @@ bool HasDynamicIndirectShadowCasterRepresentation(FGBufferData GBufferData)
 	return (PackedAlpha & 2) != 0;
 }
 
+bool IsFirstPerson(FGBufferData GBufferData)
+{
+#if !ALLOW_STATIC_LIGHTING
+	return (GBufferData.SelectiveOutputMask & IS_FIRST_PERSON_MASK) != 0;
+#else
+	return false;
+#endif
+}
+
 // High frequency Checkerboard pattern
 // @param PixelPos relative to left top of the rendertarget (not viewport)
 // @return true/false, todo: profile if float 0/1 would be better (need to make sure it's 100% the same)
@@ -992,7 +1001,12 @@ FGBufferData DecodeGBufferData(
 
 	GBuffer.CustomData = HasCustomGBufferData(GBuffer.ShadingModelID) ? InGBufferD : 0;
 
+	// FirstPerson uses a bit in SelectiveOutputMask that is aliased with ZERO_PRECSHADOW_MASK when !ALLOW_STATIC_LIGHTING, so we explicitly skip this logic here.
+#if ALLOW_STATIC_LIGHTING
 	GBuffer.PrecomputedShadowFactors = !(GBuffer.SelectiveOutputMask & SKIP_PRECSHADOW_MASK) ? InGBufferE :  ((GBuffer.SelectiveOutputMask & ZERO_PRECSHADOW_MASK) ? 0 :  1);
+#else
+	GBuffer.PrecomputedShadowFactors = half(1.0f);
+#endif
 	GBuffer.CustomDepth = ConvertFromDeviceZ(CustomNativeDepth);
 	GBuffer.CustomStencil = CustomStencil;
 	GBuffer.Depth = SceneDepth;
diff --git a/Engine/Shaders/Private/GBufferHelpers.ush b/Engine/Shaders/Private/GBufferHelpers.ush
--- a/Engine/Shaders/Private/GBufferHelpers.ush
+++ b/Engine/Shaders/Private/GBufferHelpers.ush
@@ -386,7 +386,12 @@ void GBufferPostDecode(inout FGBufferData Ret, bool bChecker, bool bGetNormalize
 {
 	Ret.CustomData = HasCustomGBufferData(Ret.ShadingModelID) ? Ret.CustomData : half(0.0f);
 
+	// FirstPerson uses a bit in SelectiveOutputMask that is aliased with ZERO_PRECSHADOW_MASK when !ALLOW_STATIC_LIGHTING, so we explicitly skip this logic here.
+#if ALLOW_STATIC_LIGHTING
 	Ret.PrecomputedShadowFactors = !(Ret.SelectiveOutputMask & 0x2) ? Ret.PrecomputedShadowFactors : ((Ret.SelectiveOutputMask & 0x4) ? half(0.0f) : half(1.0f));
+#else
+	Ret.PrecomputedShadowFactors = half(1.0f);
+#endif
 	Ret.Velocity = !(Ret.SelectiveOutputMask & 0x8) ? Ret.Velocity : half(0.0f);
 	bool bHasAnisotropy = (Ret.SelectiveOutputMask & 0x1);
 
diff --git a/Engine/Shaders/Private/ShadingCommon.ush b/Engine/Shaders/Private/ShadingCommon.ush
--- a/Engine/Shaders/Private/ShadingCommon.ush
+++ b/Engine/Shaders/Private/ShadingCommon.ush
@@ -38,6 +38,7 @@
 #define HAS_ANISOTROPY_MASK				(1 << 4)
 #define SKIP_PRECSHADOW_MASK			(1 << 5)
 #define ZERO_PRECSHADOW_MASK			(1 << 6)
+#define IS_FIRST_PERSON_MASK			(1 << 6) // Technically not related to the SelectiveOutputMask, but IS_FIRST_PERSON_MASK shares a bit with ZERO_PRECSHADOW_MASK and consequently only works if ALLOW_STATIC_LIGHTING is 0.
 #define SKIP_VELOCITY_MASK				(1 << 7)
 
 #define SSS_PROFILE_ID_INVALID  256
diff --git a/Engine/Shaders/Private/Substrate/Substrate.ush b/Engine/Shaders/Private/Substrate/Substrate.ush
--- a/Engine/Shaders/Private/Substrate/Substrate.ush
+++ b/Engine/Shaders/Private/Substrate/Substrate.ush
@@ -293,6 +293,7 @@ struct FSubstrateBSDF
 #define HEADER_MASK_MATERIALMODE					(0x7)
 #define HEADER_MASK_HASPRECSHADOWMASK				(1u << 3)
 #define HEADER_MASK_ZEROPRECSHADOWMASK				(1u << 4)
+#define HEADER_MASK_ISFIRSTPERSON					(1u << 4) // FirstPerson and ZEROPRECSHADOWMASK share the same bit, so the FirstPerson bit may only be used if ALLOW_STATIC_LIGHTING is 0.
 #define HEADER_MASK_CASTCONTACTSHADOW				(1u << 5)
 #define HEADER_MASK_HASDYNINDIRECTSHADOWCASTER		(1u << 6)
 #define HEADER_MASK_HASSUBSURFACE					(1u << 7)
@@ -2175,13 +2176,25 @@ struct FSubstratePixelHeader
 	}
 
 	bool HasPrecShadowMask()									{ return (this.State & HEADER_MASK_HASPRECSHADOWMASK) != 0; }
+#if ALLOW_STATIC_LIGHTING
 	bool HasZeroPrecShadowMask()								{ return (this.State & HEADER_MASK_ZEROPRECSHADOWMASK) != 0; }
+	bool IsFirstPerson()										{ return false; }
+#else
+	bool HasZeroPrecShadowMask()								{ return false; }
+	bool IsFirstPerson()										{ return (this.State & HEADER_MASK_ISFIRSTPERSON) != 0; }
+#endif
 	bool DoesCastContactShadow()								{ return (this.State & HEADER_MASK_CASTCONTACTSHADOW) != 0; }
 	bool HasSubsurface()										{ return (this.State & HEADER_MASK_HASSUBSURFACE) != 0; }
 	bool HasDynamicIndirectShadowCasterRepresentation()			{ return (this.State & HEADER_MASK_HASDYNINDIRECTSHADOWCASTER) != 0; }
 
 	void SetHasPrecShadowMask(bool bIn)							{ this.State |= (bIn ? HEADER_MASK_HASPRECSHADOWMASK : 0u); }
+#if ALLOW_STATIC_LIGHTING
 	void SetZeroPrecShadowMask(bool bIn)						{ this.State |= (bIn ? HEADER_MASK_ZEROPRECSHADOWMASK : 0u); }
+	void SetIsFirstPerson(bool bIn)								{ }
+#else
+	void SetZeroPrecShadowMask(bool bIn)						{ }
+	void SetIsFirstPerson(bool bIn)								{ this.State |= (bIn ? HEADER_MASK_ISFIRSTPERSON : 0u); }
+#endif
 	void SetCastContactShadow(bool bIn)							{ this.State |= (bIn ? HEADER_MASK_CASTCONTACTSHADOW : 0u); }
 	void SetHasSubsurface(bool bIn)								{ this.State |= (bIn ? HEADER_MASK_HASSUBSURFACE : 0u); }
 	void SetDynamicIndirectShadowCasterRepresentation(bool bIn)	{ this.State |= (bIn ? HEADER_MASK_HASDYNINDIRECTSHADOWCASTER : 0u); }
diff --git a/Engine/Source/Runtime/Engine/Classes/Engine/RendererSettings.h b/Engine/Source/Runtime/Engine/Classes/Engine/RendererSettings.h
--- a/Engine/Source/Runtime/Engine/Classes/Engine/RendererSettings.h
+++ b/Engine/Source/Runtime/Engine/Classes/Engine/RendererSettings.h
@@ -694,7 +694,7 @@ class URendererSettings : public UDeveloperSettings
 
 	UPROPERTY(config, EditAnywhere, Category=MiscLighting, meta=(
 		ConsoleVariable="r.AllowStaticLighting",
-		ToolTip="Whether to allow any static lighting to be generated and used, like lightmaps and shadowmaps. Games that only use dynamic lighting should set this to 0 to save some static lighting overhead. Disabling allows Material AO and Material BentNormal to work with Lumen Global Illumination.  Changing this setting requires restarting the editor.",
+		ToolTip="Whether to allow any static lighting to be generated and used, like lightmaps and shadowmaps. Games that only use dynamic lighting should set this to 0 to save some static lighting overhead. Disabling allows Material AO and Material BentNormal to work with Lumen Global Illumination. When StaticLighting is disabled, it also enables certain effects related to first person rendering, such as player shadows on the ground. Changing this setting requires restarting the editor.",
 		ConfigRestartRequired=true))
 	uint32 bAllowStaticLighting:1;
 
