Path Tracer: Fix substrate specular highlights

The default glint density was being set incorrectly when not in use, leading to a missing highlight.

#rnx

diff --git a/Engine/Shaders/Private/PathTracing/PathTracingCommon.ush b/Engine/Shaders/Private/PathTracing/PathTracingCommon.ush
--- a/Engine/Shaders/Private/PathTracing/PathTracingCommon.ush
+++ b/Engine/Shaders/Private/PathTracing/PathTracingCommon.ush
@@ -886,6 +886,11 @@ FPathTracingPayload UnpackPathTracingPayload(FPackedPathTracingPayload Input, FR
 		Output.GlintUVdx = UnpackFloat2FromUInt(Input.PackedData[13]);
 		Output.GlintUVdy = UnpackFloat2FromUInt(Input.PackedData[14]);
 	}
+	else
+	{
+		// When glints are not used, default to full density so we get the regular specular highlight
+		Output.GlintValue = 1.0f;
+	}
 
 #else // PATHTRACING_SUBSTRATE_PAYLOAD
 	Output.BaseColor         = Input.UnpackBaseColor();
