Fix directional light shadows in screenspaceshadowmask rendering jagged.
We were using ManualPCF for filtering the shadows rather than MobileShadowPCF which uses lower PCF at the same shadow quality level

diff --git a/Engine/Shaders/Private/ShadowProjectionPixelShader.usf b/Engine/Shaders/Private/ShadowProjectionPixelShader.usf
--- a/Engine/Shaders/Private/ShadowProjectionPixelShader.usf
+++ b/Engine/Shaders/Private/ShadowProjectionPixelShader.usf
@@ -31,6 +31,9 @@
 #if MODULATED_SHADOWS
 	#define MobileSceneTextures MobileBasePass.SceneTextures
 #endif
+#if SHADING_PATH_MOBILE
+#include "MobileLightingCommon.ush"
+#endif
 
 #include "ShadowProjectionCommon.ush"
 #include "ShadowFilteringCommon.ush"
@@ -264,8 +267,11 @@ void Main(
 		Settings.bTreatMaxDepthUnshadowed = false;
 		Settings.DensityMulConstant = 0;
 		Settings.ProjectionDepthBiasParameters = 0;
-
+#if SHADING_PATH_MOBILE
+		Shadow = MobileShadowPCF(ShadowPosition.xy, Settings);
+#else
 		Shadow = ManualPCF(ShadowPosition.xy, Settings);
+#endif
 	}
 	#endif // !USE_PCSS
 		
