Fix numerical precision causing incorrect shading results on certain platform when opaque rough refraction is enabled with Substrate.

diff --git a/Engine/Shaders/Private/Substrate/SubstrateRoughRefraction.usf b/Engine/Shaders/Private/Substrate/SubstrateRoughRefraction.usf
--- a/Engine/Shaders/Private/Substrate/SubstrateRoughRefraction.usf
+++ b/Engine/Shaders/Private/Substrate/SubstrateRoughRefraction.usf
@@ -9,6 +9,9 @@
 // Gaussian function with a=1/(c*sqrt(2*PI)) so it represents the probability density function of a normally distributed random variable with expected value mu = b for some variance and integrale = 1;
 float GaussianFunction(float x, float Mean, float Variance)
 {
+	// To avoid numerical precision issue, clamp the variance (stddev =0.001)
+	Variance = max(0.000001f, Variance);
+
 	float StandardDeviation = sqrt(Variance);
 	float FactorA = 1 / (StandardDeviation * sqrt(2 * PI));
 	float Bell = exp(-abs((x - Mean) * (x - Mean)) / (2 * Variance));
@@ -16,6 +19,9 @@ float GaussianFunction(float x, float Mean, float Variance)
 }
 float GaussianFunction(float2 P, float2 Mean, float Variance)
 {
+	// To avoid numerical precision issue, clamp the variance (stddev =0.001)
+	Variance = max(0.000001f, Variance);
+
 	float StandardDeviation = sqrt(Variance);
 	float FactorA = 1 / (StandardDeviation * sqrt(2 * PI));
 	float Bell = exp(-((P.x - Mean.x) * (P.x - Mean.x)) / (2 * Variance) - ((P.y - Mean.y) * (P.y - Mean.y)) / (2 * Variance));
