Fix incorrect scissor rect handling for Distance Field Shadows in stereo rendering. This could lead to incorrect or missing shadows for lights with small attenuation radii.

diff --git a/Engine/Shaders/Private/DistanceFieldShadowing.usf b/Engine/Shaders/Private/DistanceFieldShadowing.usf
--- a/Engine/Shaders/Private/DistanceFieldShadowing.usf
+++ b/Engine/Shaders/Private/DistanceFieldShadowing.usf
@@ -497,39 +497,37 @@ void CullObjectsToTileWithGather(
 
 	float ExpandRadius = 0;
 
-#if !STEREO_RENDERING
-	// Note: this code is assuming a centered projection, aka no translation present in ViewToClip
-	// Stereo rendering uses an off center projection
-	float2 TanViewFOV = GetTanHalfFieldOfView();
-	// tan(FOV) = HalfUnitPlaneWidth / 1, so TanViewFOV * 2 is the size of the whole unit view plane
-	// We are operating on a subset of that defined by ScissorRectMinAndSize
-	float2 TileSize = TanViewFOV * 2 * ScissorRectMinAndSize.zw / ((float2)View.ViewSizeAndInvSize.xy * NumGroups);
-	float2 UnitPlaneMin = -TanViewFOV + TanViewFOV * 2 * (ScissorRectMinAndSize.xy - View.ViewRectMin.xy) * View.ViewSizeAndInvSize.zw;
-	float2 UnitPlaneTileMin = (GroupId.xy * TileSize + UnitPlaneMin) * float2(1, -1);
-	float2 UnitPlaneTileMax = ((GroupId.xy + 1) * TileSize + UnitPlaneMin) * float2(1, -1);
-
-	ViewTileMin.xy = min(MinTileZ * UnitPlaneTileMin, MaxTileZ2 * UnitPlaneTileMin) - ExpandRadius;
-	ViewTileMax.xy = max(MinTileZ * UnitPlaneTileMax, MaxTileZ2 * UnitPlaneTileMax) + ExpandRadius;
-	ViewTileMin.z = MinTileZ - ExpandRadius;
-	ViewTileMax.z = MaxTileZ2 + ExpandRadius;
-	ViewTileMin2.xy = min(MinTileZ2 * UnitPlaneTileMin, MaxTileZ * UnitPlaneTileMin) - ExpandRadius;
-	ViewTileMax2.xy = max(MinTileZ2 * UnitPlaneTileMax, MaxTileZ * UnitPlaneTileMax) + ExpandRadius;
-	ViewTileMin2.z = MinTileZ2 - ExpandRadius;
-	ViewTileMax2.z = MaxTileZ + ExpandRadius;
-#else
-	// For stereo, we need to do a more expensive transform to account for non-centered projection
-	float2 ScreenTileMin = ((GroupId.xy / NumGroups) - 0.5) * float2(2, -2);
-	float2 ScreenTileMax = (((GroupId.xy + 1) / NumGroups) - 0.5) * float2(2, -2);
+	// We operate within both a view rect (for multiple side-by-side views) and a scissor rect relative to the view rect (circumscribing the light's attenuation radius on screen)
+	float2 TileSize = 2 * ScissorRectMinAndSize.zw / ((float2)View_ViewSizeAndInvSize.xy * NumGroups);
+    float2 ScissorRectMin = float2(-1,-1) + 2 * (ScissorRectMinAndSize.xy - View_ViewRectMin.xy) * View_ViewSizeAndInvSize.zw;
 
+    float2 ScreenTileMin = (GroupId.xy * TileSize + ScissorRectMin) * float2(1, -1);
+    float2 ScreenTileMax = ((GroupId.xy + 1) * TileSize + ScissorRectMin) * float2(1, -1); 
+
+	// Get the bounding box for this tile in view space
+	// Project the corners into view space using both MinZ and MaxZ, and ensure the bounding box contains both results
+#if STEREO_RENDERING
+	// Stereo rendering has asymmetrical FOVs for each eye so we need to account for off-center projection and use full view matrices when projecting into view space
 	ViewTileMin.xy = min(mul(float4(ScreenTileMin * MinTileZ, MinTileZ, 1), ScreenToView).xy, mul(float4(ScreenTileMin * MaxTileZ2, MaxTileZ2, 1), ScreenToView).xy) - ExpandRadius;
 	ViewTileMax.xy = max(mul(float4(ScreenTileMax * MinTileZ, MinTileZ, 1), ScreenToView).xy, mul(float4(ScreenTileMax * MaxTileZ2, MaxTileZ2, 1), ScreenToView).xy) + ExpandRadius;
-	ViewTileMin.z = MinTileZ - ExpandRadius;
-	ViewTileMax.z = MaxTileZ2 + ExpandRadius;
 	ViewTileMin2.xy = min(mul(float4(ScreenTileMin * MinTileZ2, MinTileZ2, 1), ScreenToView).xy, mul(float4(ScreenTileMin * MaxTileZ, MaxTileZ, 1), ScreenToView).xy) - ExpandRadius;
 	ViewTileMax2.xy = max(mul(float4(ScreenTileMax * MinTileZ2, MinTileZ2, 1), ScreenToView).xy, mul(float4(ScreenTileMax * MaxTileZ, MaxTileZ, 1), ScreenToView).xy) + ExpandRadius;
+#else
+	// If we can assume centered projection (symmetrical FOV), we can do a fast trig projection into view space using tan(FOV/2)
+	// We start with coordinates in screen space ranging from (-1,-1) to (1,1)
+	// The (x,y,z) coordinates for a given screen space point (x,y) in view space are (x * z * tan(H_FOV/2), y *  z * tan(V_FOV/2), z)
+	float2 TanViewFOV = GetTanHalfFieldOfView();
+	ViewTileMin.xy = min(MinTileZ * ScreenTileMin * TanViewFOV, MaxTileZ2 * ScreenTileMin * TanViewFOV) - ExpandRadius;
+	ViewTileMax.xy = max(MinTileZ * ScreenTileMax * TanViewFOV, MaxTileZ2 * ScreenTileMax * TanViewFOV) + ExpandRadius;
+	ViewTileMin2.xy = min(MinTileZ2 * ScreenTileMin * TanViewFOV, MaxTileZ * ScreenTileMin * TanViewFOV) - ExpandRadius;
+	ViewTileMax2.xy = max(MinTileZ2 * ScreenTileMax * TanViewFOV, MaxTileZ * ScreenTileMax * TanViewFOV) + ExpandRadius;
+#endif
+	ViewTileMin.z = MinTileZ - ExpandRadius;
+	ViewTileMax.z = MaxTileZ2 + ExpandRadius;
 	ViewTileMin2.z = MinTileZ2 - ExpandRadius;
 	ViewTileMax2.z = MaxTileZ + ExpandRadius;
-#endif
+
+	// Convert the view space bounding box to a world space bounding sphere
 	float3 ViewGroup0Center = (ViewTileMax + ViewTileMin) / 2;
 	float3 TranslatedWorldGroup0Center = mul(float4(ViewGroup0Center, 1), View.ViewToTranslatedWorld).xyz;
 	float Group0BoundingRadius = length(ViewGroup0Center - ViewTileMax);
