Fixed custom stencil buffer access on translucent surfaces.

#tests Tested pc editor with custom stencil on/off with translucent predof//postdof/postmotionblur.
#rnx

diff --git a/Engine/Shaders/Private/MaterialTemplate.ush b/Engine/Shaders/Private/MaterialTemplate.ush
--- a/Engine/Shaders/Private/MaterialTemplate.ush
+++ b/Engine/Shaders/Private/MaterialTemplate.ush
@@ -39,7 +39,7 @@
 
 // Update this GUID to force all material shaders to recompile
 // Merge conflicts on this line should be resolved by generating a new GUID
-#pragma message("UESHADERMETADATA_VERSION 25A8AD39-1BAE-4581-904F-FA19DE92F37E")
+#pragma message("UESHADERMETADATA_VERSION F19145E0-2188-4381-89D4-A25F4E0E0725")
 
 #if USES_SPEEDTREE
 	#include "/Engine/Private/SpeedTreeCommon.ush"
@@ -3145,7 +3145,12 @@ float4 SceneTextureLookup(float2 UV, int SceneTextureIndex, bool bFiltered)
 
 #else // SCENE_TEXTURES_DISABLED
 
+#if MATERIAL_DOMAIN_SURFACE &&  MATERIALBLENDING_ANY_TRANSLUCENT
+	// BufferToSceneTextureScale is only needed for translucent surfaces today and we know that here.
+	const uint2 PixelPos = uint2(UV * View.BufferSizeAndInvSize.xy * View.BufferToSceneTextureScale.xy);
+#else
 	const uint2 PixelPos = uint2(UV * View.BufferSizeAndInvSize.xy);
+#endif
 
 #ifdef POST_PROCESS_PROPAGATE_ALPHA_INPUT
 	// Propagation of SceneColor alpha, to preserve it in materials that don't themselves write to alpha, is supported by storing the propagated alpha to a local variable,
