Fix MobileSSR shader compile error

diff --git a/Engine/Shaders/Private/MobileSSR.ush b/Engine/Shaders/Private/MobileSSR.ush
--- a/Engine/Shaders/Private/MobileSSR.ush
+++ b/Engine/Shaders/Private/MobileSSR.ush
@@ -53,5 +53,6 @@ bool MobileSSR(FGBufferData GBuffer, float4 SvPosition, float3 TranslatedWorldPo
 			return true;
 		}
 	}
+	SSR = 0;
 	return false;
 }
