Added InstanceId to the FInstanceSceneData as a convenience when needing to go from FInstanceSceneData back to ID.

diff --git a/Engine/Shaders/Private/SceneData.ush b/Engine/Shaders/Private/SceneData.ush
--- a/Engine/Shaders/Private/SceneData.ush
+++ b/Engine/Shaders/Private/SceneData.ush
@@ -225,6 +225,7 @@ struct FInstanceSceneEditorData
 
 struct FInstanceSceneData
 {
+	uint InstanceId; // Note: derived data (not loaded from scene, but set to the ID used to load the data for convenience).
 	FDFMatrix LocalToWorld;
 	FDFMatrix PrevLocalToWorld;
 	FDFInverseMatrix WorldToLocal;
@@ -263,6 +264,7 @@ FInstanceSceneData WaveReadLaneAt(FInstanceSceneData In, uint SrcIndex)
 {
 	FInstanceSceneData Result;
 
+	Result.InstanceId					= WaveReadLaneAt(In.InstanceId, SrcIndex);
 	Result.LocalToWorld					= WaveReadLaneAt(In.LocalToWorld, SrcIndex);
 	Result.PrevLocalToWorld				= WaveReadLaneAt(In.PrevLocalToWorld, SrcIndex);
 	Result.WorldToLocal					= WaveReadLaneAt(In.WorldToLocal, SrcIndex);
@@ -1013,11 +1015,11 @@ void ComputeInstanceDerivedData(inout FInstanceSceneData InstanceData, float3 Po
 	InstanceData.WorldToLocal								= MakeDFInverseMatrix(PositionHigh, RelativeWorldToLocal);
 }
 
-// Fetch from scene primitive buffer
 FInstanceSceneData GetInstanceSceneData(uint InstanceId, uint SOAStride, bool bCheckValid = true)
 {
 	FInstanceSceneData InstanceData = (FInstanceSceneData)0;
 
+	InstanceData.InstanceId = InstanceId;
 	//
 	// NOTE: When changing the packed data layout, ensure that GPUScene/GPUSceneWriter.ush is kept in sync!
 	//       Also, please update the GetInstanceSceneData function in GPUScene.cpp for validation purposes.
