Fix half precision mobile preview

diff --git a/Engine/Shaders/Private/Common.ush b/Engine/Shaders/Private/Common.ush
--- a/Engine/Shaders/Private/Common.ush
+++ b/Engine/Shaders/Private/Common.ush
@@ -615,7 +615,11 @@ MaterialFloat4 TextureCubeArraySampleLevel(TextureCubeArray Tex, SamplerState Sa
 #else
 	#define TANGENTTOWORLD0					TEXCOORD10
 	#define TANGENTTOWORLD2					TEXCOORD11
-	#define TANGENTTOWORLD_INTERPOLATOR_BLOCK	half4 TangentToWorld0 : TANGENTTOWORLD0; half4 TangentToWorld2 : TANGENTTOWORLD2;
+	#if MOBILE_EMULATION
+		#define TANGENTTOWORLD_INTERPOLATOR_BLOCK	float4 TangentToWorld0 : TANGENTTOWORLD0; float4 TangentToWorld2 : TANGENTTOWORLD2;
+	#else
+		#define TANGENTTOWORLD_INTERPOLATOR_BLOCK	half4 TangentToWorld0 : TANGENTTOWORLD0; half4 TangentToWorld2 : TANGENTTOWORLD2;
+	#endif
 #endif
 
 
