[Light Grid] Fix incorrect condition in OutputCell(...)

diff --git a/Engine/Shaders/Private/LightGridInjection.usf b/Engine/Shaders/Private/LightGridInjection.usf
--- a/Engine/Shaders/Private/LightGridInjection.usf
+++ b/Engine/Shaders/Private/LightGridInjection.usf
@@ -564,7 +564,7 @@ void OutputCell(FCellWriter CellWriter, uint NumPrimitives, uint CulledLightData
 
 	// output entries in linked list to RWCulledLightDataGrid
 
-	while (CellWriter.LinkOffset != 0xFFFFF)
+	while (CellWriter.LinkOffset < LIGHT_GRID_CELL_WRITER_MAX_LINK_OFFSET)
 	{
 		const uint PackedData0 = RWCulledLightLinks[CellWriter.LinkOffset * LIGHT_LINK_STRIDE + 0];
 		const uint PackedData1 = RWCulledLightLinks[CellWriter.LinkOffset * LIGHT_LINK_STRIDE + 1];
