Substrate - fixed rendertarget output for base color and normal.

#rnx
#tests Tested PC editor all output type with and without substrate.

diff --git a/Engine/Shaders/Private/SceneCapturePixelShader.usf b/Engine/Shaders/Private/SceneCapturePixelShader.usf
--- a/Engine/Shaders/Private/SceneCapturePixelShader.usf
+++ b/Engine/Shaders/Private/SceneCapturePixelShader.usf
@@ -7,6 +7,19 @@
 #include "Common.ush"
 #include "DeferredShadingCommon.ush"
 
+#if SUBSTRATE_ENABLED
+
+#include "/Engine/Private/Substrate/Substrate.ush"
+
+
+#define NEEDS_SCENE_TEXTURES 1
+#define SCENE_TEXTURES_DISABLED 0
+#define SubstratePublicStruct SubstratePublic
+#include "/Engine/Public/SubstratePublic.ush"
+#undef SubstratePublicStruct
+
+#endif // SUBSTRATE_ENABLED
+
 float4 EncodeFloatRGBA( float v )
 {
 	float4 Enc = float4(1.0, 255.0, 65025.0, 16581375.0) * v;
@@ -32,15 +45,32 @@ void Main(
 	OutColor = float4(CalcSceneColor(Input.UV), CalcSceneDepth(Input.UV));
 #elif SOURCE_MODE_SCENE_DEPTH 
 	OutColor = float4(CalcSceneDepth(Input.UV), 0, 0, 0);
-#elif SOURCE_MODE_NORMAL
-	FGBufferData GBufferData = GetGBufferData(Input.UV);
-	OutColor = float4(GBufferData.WorldNormal, 0);
-#elif SOURCE_MODE_BASE_COLOR
-	FGBufferData GBufferData = GetGBufferData(Input.UV);
-	OutColor = float4(GBufferData.BaseColor, 0);
 #elif SOURCE_MODE_DEVICE_DEPTH
 	float Depth = 1 - LookupDeviceZ(Input.UV);
 	OutColor.rgba = EncodeFloatRGBA(Depth).rgba;
 	OutColor.a = 1;
+#else
+
+#if SUBSTRATE_ENABLED
+
+	uint2 ScreenPos = uint2(Input.UV * View.BufferSizeAndInvSize.xy);
+	FSubstrateGBuffer SubstrateGBuffer = SubstratePublic_GetSubstrateGBufferFromFirstBSDF(ScreenPos);
+#if SOURCE_MODE_NORMAL
+	OutColor = float4(SubstrateGBuffer.WorldNormal, 0);
+#elif SOURCE_MODE_BASE_COLOR
+	OutColor = float4(SubstrateGBuffer.BaseColor, 0);
 #endif
+
+#else // SUBSTRATE_ENABLED
+
+	FGBufferData GBufferData = GetGBufferData(Input.UV);
+#if SOURCE_MODE_NORMAL
+	OutColor = float4(GBufferData.WorldNormal, 0);
+#elif SOURCE_MODE_BASE_COLOR
+	OutColor = float4(GBufferData.BaseColor, 0);
+#endif
+
+#endif // SUBSTRATE_ENABLED
+
+#endif // SOURCE_MODE_
 }
diff --git a/Engine/Source/Runtime/Renderer/Private/SceneCaptureRendering.cpp b/Engine/Source/Runtime/Renderer/Private/SceneCaptureRendering.cpp
--- a/Engine/Source/Runtime/Renderer/Private/SceneCaptureRendering.cpp
+++ b/Engine/Source/Runtime/Renderer/Private/SceneCaptureRendering.cpp
@@ -104,6 +104,7 @@ public:
 	BEGIN_SHADER_PARAMETER_STRUCT(FParameters, )
 		SHADER_PARAMETER_STRUCT_REF(FViewUniformShaderParameters, View)
 		SHADER_PARAMETER_STRUCT_INCLUDE(FSceneTextureShaderParameters, SceneTextures)
+		SHADER_PARAMETER_RDG_UNIFORM_BUFFER(FSubstratePublicGlobalUniformParameters, SubstratePublic)
 		RENDER_TARGET_BINDING_SLOTS()
 	END_SHADER_PARAMETER_STRUCT()
 
@@ -283,6 +284,7 @@ void CopySceneCaptureComponentToTarget(
 		FSceneCapturePS::FParameters* PassParameters = GraphBuilder.AllocParameters<FSceneCapturePS::FParameters>();
 		PassParameters->View = View.ViewUniformBuffer;
 		PassParameters->SceneTextures = SceneTextures.GetSceneTextureShaderParameters(ViewFamily.GetFeatureLevel());
+		PassParameters->SubstratePublic = View.SubstrateViewData.SceneData->SubstratePublicGlobalUniformParameters;
 		PassParameters->RenderTargets[0] = FRenderTargetBinding(ViewFamilyTexture, bIsCompositing ? ERenderTargetLoadAction::ELoad : ERenderTargetLoadAction::ENoAction);
 
 		TShaderMapRef<FScreenVS> VertexShader(View.ShaderMap);
@@ -556,7 +558,7 @@ static void UpdateSceneCaptureContentMobile_RenderThread(
 			const FRenderTarget* FamilyTarget = SceneRenderer->ViewFamily.RenderTarget;
 			FRDGTextureRef FamilyTexture = RegisterExternalTexture(GraphBuilder, FamilyTarget->GetRenderTargetTexture(), TEXT("OutputTexture"));
 			const FMinimalSceneTextures& SceneTextures = SceneRenderer->GetActiveSceneTextures();
-
+			
 			RDG_EVENT_SCOPE(GraphBuilder, "CaptureSceneColor");
 			CopySceneCaptureComponentToTarget(
 				GraphBuilder,
