Path Tracer: Fix double-accumulation of emission with multi-layered Substrate materials

The emissive contribution of each slab was being counted in two places. Luckily, the duplicate logic would only be active if the material had more than one slab, so it was rare to see in practice

diff --git a/Engine/Shaders/Private/PathTracing/PathTracingMaterialHitShader.usf b/Engine/Shaders/Private/PathTracing/PathTracingMaterialHitShader.usf
--- a/Engine/Shaders/Private/PathTracing/PathTracingMaterialHitShader.usf
+++ b/Engine/Shaders/Private/PathTracing/PathTracingMaterialHitShader.usf
@@ -434,9 +434,6 @@ RAY_TRACING_ENTRY_CLOSEST_HIT(PathTracingMaterialCHS,
 						const float CoverageAboveAlongN = BSDF.CoverageAboveAlongN;
 						const float3 MaxLobeWeight = WeightV * lerp(1.0f, TransmittanceN, CoverageAboveAlongN); // largest value LobeWeight() could return
 
-						// accumulate visibility
-						Payload.Radiance += WeightV * BSDF_GETEMISSIVE(BSDF);
-
 						// Is the slab even visible at all?
 						if (any(MaxLobeWeight > 0.0))
 						{
