Fix for LightGridInjection compilation on Adreno 6xx.

When the following permutation is compiled:
// #define USE_LINKED_CULL_LIST 1
// #define REFINE_RECTLIGHT_BOUNDS 0
// #define USE_HZB_CULL 0
// #define USE_PARENT_LIGHT_GRID 0
// #define USE_THREAD_GROUP_PER_CELL 0
// #define USE_THREAD_GROUP_SIZE_32 0
// #define LIGHT_LINK_STRIDE 2
// #define THREADGROUP_SIZE 64
// #define THREADGROUP_SIZE_X 4
// #define THREADGROUP_SIZE_Y 4
// #define THREADGROUP_SIZE_Z 4

for Adreno 6xx, driver's reg-alloc exhausts all available UGPRs, while not supporting any spilling.

Hence, it fails with:
AdrenoVK-0: Info log: Assertion failed: false && "No UGPR spill support"

Workarounded to use a slightly less optimal InterlockedAdd on mobile.

#rnx

diff --git a/Engine/Shaders/Private/LightGridInjection.usf b/Engine/Shaders/Private/LightGridInjection.usf
--- a/Engine/Shaders/Private/LightGridInjection.usf
+++ b/Engine/Shaders/Private/LightGridInjection.usf
@@ -436,7 +436,13 @@ void CullLights_SingleThread(uint GridIndex, FCellBounds CellBounds, uint Parent
 
 #if USE_LINKED_CULL_LIST
 	uint CulledLightDataStart;
+	#if FEATURE_LEVEL == FEATURE_LEVEL_ES3_1
+	// Adreno compiler fails to reg-alloc on 6xx for this shader,
+	// and doesn't support UGPR spilling, so just fails to compile.
+	InterlockedAdd(RWNextCulledLightData[0], NumVisibleLights, CulledLightDataStart);
+	#else
 	WaveInterlockedAdd_(RWNextCulledLightData[0], NumVisibleLights, CulledLightDataStart);
+	#endif
 
 	// Mega Lights are order dependent
 	OutputReversedLinkedList(LinkOffset, NumVisibleLights, CulledLightDataStart);
