Clamp tex.coord derivates during Glint evaluation to avoid incorrect pattern in distance.

diff --git a/Engine/Shaders/Private/Substrate/Glint/GlintThirdParty.ush b/Engine/Shaders/Private/Substrate/Glint/GlintThirdParty.ush
--- a/Engine/Shaders/Private/Substrate/Glint/GlintThirdParty.ush
+++ b/Engine/Shaders/Private/Substrate/Glint/GlintThirdParty.ush
@@ -439,6 +439,11 @@ float f_P(float3 wo, float3 wi, float3 sigmas_rho, float GlintDensity, float2 Te
 
 	FGlintLutParams GlintLutParams = GetGlintParams();
 
+	// Clamp tex.coord to [-1..1] in order to avoid incorrect LUT look up. A value > 1, means a pixel footprint larger than 1 pixel, 
+	// which causes the glints dot to be stretched, causing implausible pattern.
+	TexCoordDDX = clamp(TexCoordDDX, -1.f, 1.f);
+	TexCoordDDY = clamp(TexCoordDDY, -1.f, 1.f);
+
 	// Compute texture derivatives
 	float2 dst0 = TexCoordDDX;
 	float2 dst1 = TexCoordDDY;
