[Volumetric Fog] Modified GenerateConservativeDepthBufferCS logic so HZB test results in Volumetric Fog match MegaLights Volume and Light Grid.

diff --git a/Engine/Shaders/Private/GenerateConservativeDepth.usf b/Engine/Shaders/Private/GenerateConservativeDepth.usf
--- a/Engine/Shaders/Private/GenerateConservativeDepth.usf
+++ b/Engine/Shaders/Private/GenerateConservativeDepth.usf
@@ -16,6 +16,30 @@ RWTexture2D<float> ConservativeDepthTextureUAV;
 int2 ConservativeDepthTextureSize;
 int DestinationPixelSizeAtFullRes;
 
+// TODO: Merge this with the similar implementations in LightGridInjection.usf and MegaLightsVolume.ush
+FScreenRect ComputeFroxelCullRect(uint3 GridCoordinate, float FootprintMargin)
+{
+	// Compute extent of tiles in clip-space. Note that the last tile may extend a bit outside of view if view size is not evenly divisible tile size.
+	const float2 InvCulledGridSizeF = DestinationPixelSizeAtFullRes * ResolvedView.ViewSizeAndInvSize.zw;
+	const float2 TileSize = float2(2.0f, -2.0f) * InvCulledGridSizeF.xy;
+	const float2 UnitPlaneMin = float2(-1.0f, 1.0f);
+
+	float2 UnitPlaneTileMin = (GridCoordinate.xy - FootprintMargin) * TileSize + UnitPlaneMin;
+	float2 UnitPlaneTileMax = (GridCoordinate.xy + 1 + FootprintMargin) * TileSize + UnitPlaneMin;
+
+	float3 CullRectMin;
+	CullRectMin.x = min(UnitPlaneTileMin.x, UnitPlaneTileMax.x);
+	CullRectMin.y = min(UnitPlaneTileMin.y, UnitPlaneTileMax.y);
+	CullRectMin.z = 0;
+
+	float3 CullRectMax;
+	CullRectMax.x = max(UnitPlaneTileMin.x, UnitPlaneTileMax.x);
+	CullRectMax.y = max(UnitPlaneTileMin.y, UnitPlaneTileMax.y);
+	CullRectMax.z = 0;
+
+	return GetScreenRect(int4(0, 0, HZBViewSize), CullRectMin, CullRectMax, 4);
+}
+
 [numthreads(8, 8, 1)]
 void GenerateConservativeDepthBufferCS(uint3 DispatchThreadId : SV_DispatchThreadID)
 {
@@ -23,23 +47,7 @@ void GenerateConservativeDepthBufferCS(uint3 DispatchThreadId : SV_DispatchThrea
 
 	if (all(DispatchThreadId.xy < uint2(ConservativeDepthTextureSize)))
 	{
-		FFrustumCullData CullData;
-		// Rect min and max as UVs x=>[0,1] y=>[0,1].
-		// We cover the pixel and the full extent due to bilinear filtering. We also take into account the fact that HZB starts at half resolution.
-		uint2 FullResMin = float2(float2(DispatchThreadId.xy) + float2(1.0f, -1.0f)) * DestinationPixelSizeAtFullRes - float2(DestinationPixelSizeAtFullRes,-DestinationPixelSizeAtFullRes) * 2.0f;
-		uint2 FullResMax = float2(float2(DispatchThreadId.xy) + float2(1.0f, -1.0f)) * DestinationPixelSizeAtFullRes + float2(DestinationPixelSizeAtFullRes,-DestinationPixelSizeAtFullRes) * 2.0f;
-		CullData.RectMin = float3(saturate(FullResMin * ResolvedView.ViewSizeAndInvSize.zw), 0.0f);
-		CullData.RectMax = float3(saturate(FullResMax * ResolvedView.ViewSizeAndInvSize.zw), 1.0f);
-		// clip x=>[-1,1] y=>[1,-1]
-		CullData.RectMin = CullData.RectMin * float3(2.0,-2.0, 1.0) + float3(-1.0f, 1.0f, 0.0f);
-		CullData.RectMax = CullData.RectMax * float3(2.0,-2.0, 1.0) + float3(-1.0f, 1.0f, 0.0f);
-		// Dummy data
-		CullData.bCrossesFarPlane = false;
-		CullData.bCrossesNearPlane= false;
-		CullData.bIsVisible = true;
-
-		const int DesiredFootprintPixels = 4; // Required bias avoid visual artefact in tiny gaps.
-		FScreenRect ScreenRect = GetScreenRect(int4(0,0,ResolvedView.ViewSizeAndInvSize.xy), CullData, DesiredFootprintPixels);
+		FScreenRect ScreenRect = ComputeFroxelCullRect(DispatchThreadId, 0.5f);
 
 		ConservativeDepthTextureUAV[DispatchThreadId.xy] = GetMinDepthFromHZB(ScreenRect, true);
 	}
