[Volumetric Fog] Improvements to history rejection and enable HZB culling (r.VolumetricFog.ConservativeDepth) by default.
- Modify bilinear sample UV to avoid sampling froxels that were occluded in the previous frame.

Perf on console:
- Without MegaLights
    - Without HZB culling: 1.7ms
    - With HZB culling: 1.1ms
- With MegaLights
    - Without HZB culling: 0.9ms
    - With HZB culling: 0.5ms

diff --git a/Engine/Shaders/Private/VolumetricFog.usf b/Engine/Shaders/Private/VolumetricFog.usf
--- a/Engine/Shaders/Private/VolumetricFog.usf
+++ b/Engine/Shaders/Private/VolumetricFog.usf
@@ -684,6 +684,73 @@ float3 MobileDirectionalLightDirection;
 Texture3D<float> RaytracedShadowsVolume;
 Texture3D<float3> MegaLightsVolume;
 
+// Modify UV used to sample LightScatteringHistory to avoid cells that were occluded last frame
+float2 FixupHistoryUV(float2 UV, float2 Size, float PrevCellNDCPositionZ, out bool bHasValidHistory)
+{
+	float2 FullResUV = UV * Size;
+
+	uint2 ScreenCoord = floor(FullResUV - 0.5f);
+	float2 BilinearWeights = frac(FullResUV - 0.5f);
+	float2 GatherUV = (ScreenCoord + 1.0f) / Size;
+
+	float2 FullResOffset = FullResUV - ScreenCoord;
+
+	const float2 PrevConservativeDepthGatherUV = GatherUV / (PrimaryView.VolumetricFogViewGridUVToPrevViewRectUV.xy * PrimaryView.VolumetricFogPrevViewGridRectUVToResourceUV.xy);
+	const float4 PrevSceneDepths = PrevConservativeDepthTexture.Gather(LightScatteringHistorySampler, PrevConservativeDepthGatherUV);
+	const bool4 ValidSamples = PrevSceneDepths < PrevCellNDCPositionZ;
+	
+	bHasValidHistory = true;
+
+	if (all(ValidSamples))
+	{
+		return UV;
+	}
+	
+	// Gather pattern:
+    // wz
+    // xy
+	
+	if (all(ValidSamples.wz))
+	{
+		return (ScreenCoord + float2(FullResOffset.x, 0.5f)) / Size;
+	}
+	if (all(ValidSamples.xy))
+	{
+		return (ScreenCoord + float2(FullResOffset.x, 1.5f)) / Size;
+	}
+	if (all(ValidSamples.wx))
+	{
+		return (ScreenCoord + float2(0.5f, FullResOffset.y)) / Size;
+	}
+	if (all(ValidSamples.zy))
+	{
+		return (ScreenCoord + float2(1.5f, FullResOffset.y)) / Size;
+	}
+
+	if (ValidSamples.x)
+	{
+		return (ScreenCoord + float2(0.5f, 1.5f)) / Size;
+	}
+
+	if (ValidSamples.y)
+	{
+		return (ScreenCoord + float2(1.5f, 1.5f)) / Size;
+	}
+
+	if (ValidSamples.w)
+	{
+		return (ScreenCoord + float2(0.5f, 0.5f)) / Size;
+	}
+
+	if (ValidSamples.z)
+	{
+		return (ScreenCoord + float2(1.5f, 0.5f)) / Size;
+	}
+
+	bHasValidHistory = false;
+	return UV;
+}
+
 #ifdef LightScatteringCS
 [numthreads(THREADGROUP_SIZE_X, THREADGROUP_SIZE_Y, THREADGROUP_SIZE_Z)]
 void LightScatteringCS(
@@ -695,11 +762,16 @@ void LightScatteringCS(
 	float3 LightScattering = 0;
 	uint NumSuperSamples = 1;
 
+#if USE_TEMPORAL_REPROJECTION
+	float3 HistoryUV = ComputeHistoryVolumeUV_DEPRECATED(ComputeCellTranslatedWorldPosition(GridCoordinate, .5f), UnjitteredPrevTranslatedWorldToClip);
+	bool bHasValidHistory = true;
+#endif
+
 	// If the froxel is behind front depth, do not evaluate any lighting.
-	bool CellWasOccluded = false;
 	if (UseConservativeDepthTexture > 0)
 	{
-		const float3 FarDepthOffset = float3(0.5f, 0.5f, -1.0f);
+		// expand by half voxel towards the camera to support bilinear filtering
+		const float3 FarDepthOffset = float3(0.5f, 0.5f, -0.5f);
 		float3 CellTranslatedWorldPosition = ComputeCellTranslatedWorldPosition(GridCoordinate, FarDepthOffset);
 
 		float4 CellNDCPosition = mul(float4(CellTranslatedWorldPosition, 1), PrimaryView.TranslatedWorldToClip);
@@ -714,23 +786,22 @@ void LightScatteringCS(
 			return;
 		}
 
+#if USE_TEMPORAL_REPROJECTION
+		// TODO: investigate whether using CellTranslatedWorldPosition (which is biased towards camera) is correct here
+		// it might be more accurate to transform unbiased translated world position to previous frame clip space and apply bias towards camera in that space?
 		float4 PrevCellNDCPosition = mul(float4(CellTranslatedWorldPosition, 1), UnjitteredPrevTranslatedWorldToClip);
 		PrevCellNDCPosition.xyz /= PrevCellNDCPosition.w;
-		float2 PrevUVs = PrevCellNDCPosition.xy * float2(0.5f, -0.5f) + 0.5f;
-		const float PrevSceneDepth = PrevConservativeDepthTexture.Load(uint3(PrevConservativeDepthTextureSize * PrevUVs, 0));
-		if (PrevSceneDepth > CellNDCPosition.z)
-		{
-			CellWasOccluded = true;
-		}
+
+		HistoryUV.xy = FixupHistoryUV(HistoryUV.xy, View.VolumetricFogPrevResourceGridSize.xy, PrevCellNDCPosition.z, bHasValidHistory);
+#endif
 	}
 
 #if USE_TEMPORAL_REPROJECTION
-	float3 HistoryUV = ComputeHistoryVolumeUV_DEPRECATED(ComputeCellTranslatedWorldPosition(GridCoordinate, .5f), UnjitteredPrevTranslatedWorldToClip);
 	float HistoryAlpha = HistoryWeight;
 
 	// We need to test with View.VolumetricFogPrevUVMaxForTemporalBlend for HistoryUV because that is what we clamp with in ComputeHistoryVolumeUV_DEPRECATED.
 	FLATTEN
-	if (any(HistoryUV < 0) || any(HistoryUV >= float3(View.VolumetricFogPrevUVMaxForTemporalBlend, 1.0f)) || CellWasOccluded)
+	if (any(HistoryUV < 0) || any(HistoryUV >= float3(View.VolumetricFogPrevUVMaxForTemporalBlend, 1.0f)) || !bHasValidHistory)
 	{
 		HistoryAlpha = 0;
 	}
@@ -977,6 +1048,14 @@ void LightScatteringCS(
 		PreExposedScatteringAndExtinction = lerp(PreExposedScatteringAndExtinction, PreExposedHistoryScatteringAndExtinction, HistoryAlpha);
 	}
 	
+#endif
+
+	// Visualize history rejection for debugging purposes
+#if 0 && USE_TEMPORAL_REPROJECTION
+	if (HistoryAlpha < 0.001f)
+	{
+		PreExposedScatteringAndExtinction = float4(View.PreExposure * float3(1,0,0), Extinction);
+	}
 #endif
 
 	if (all(int3(GridCoordinate) < VolumetricFog.ResourceGridSizeInt))
diff --git a/Engine/Source/Runtime/Engine/Public/SceneView.h b/Engine/Source/Runtime/Engine/Public/SceneView.h
--- a/Engine/Source/Runtime/Engine/Public/SceneView.h
+++ b/Engine/Source/Runtime/Engine/Public/SceneView.h
@@ -982,6 +982,7 @@ enum ETranslucencyVolumeCascade
 	VIEW_UNIFORM_BUFFER_MEMBER(FVector2f, VolumetricFogPrevViewGridRectUVToResourceUV) \
 	VIEW_UNIFORM_BUFFER_MEMBER(FVector2f, VolumetricFogPrevUVMax) \
 	VIEW_UNIFORM_BUFFER_MEMBER(FVector2f, VolumetricFogPrevUVMaxForTemporalBlend) \
+	VIEW_UNIFORM_BUFFER_MEMBER(FVector3f, VolumetricFogPrevResourceGridSize) \
 	VIEW_UNIFORM_BUFFER_MEMBER(FVector2f, VolumetricFogScreenToResourceUV) \
 	VIEW_UNIFORM_BUFFER_MEMBER(FVector2f, VolumetricFogUVMax) \
 	VIEW_UNIFORM_BUFFER_MEMBER(float, VolumetricFogMaxDistance) \
diff --git a/Engine/Source/Runtime/Renderer/Private/ScenePrivate.h b/Engine/Source/Runtime/Renderer/Private/ScenePrivate.h
--- a/Engine/Source/Runtime/Renderer/Private/ScenePrivate.h
+++ b/Engine/Source/Runtime/Renderer/Private/ScenePrivate.h
@@ -1106,6 +1106,7 @@ public:
 	FVector2f VolumetricFogPrevViewGridRectUVToResourceUV;
 	FVector2f VolumetricFogPrevUVMax;
 	FVector2f VolumetricFogPrevUVMaxForTemporalBlend;
+	FIntVector VolumetricFogPrevResourceGridSize;
 	TRefCountPtr<IPooledRenderTarget> LightScatteringHistory;
 	TRefCountPtr<IPooledRenderTarget> PrevLightScatteringConservativeDepthTexture;
 
diff --git a/Engine/Source/Runtime/Renderer/Private/SceneRendering.cpp b/Engine/Source/Runtime/Renderer/Private/SceneRendering.cpp
--- a/Engine/Source/Runtime/Renderer/Private/SceneRendering.cpp
+++ b/Engine/Source/Runtime/Renderer/Private/SceneRendering.cpp
@@ -1611,16 +1611,18 @@ void FViewInfo::SetupUniformBufferParameters(
 
 	SetupDefaultGlobalDistanceFieldUniformBufferParameters(ViewUniformShaderParameters);
 
+	int32 VolumetricFogViewGridPixelSize;
+	int32 VolumetricFogResourceGridPixelSize;
+	const FIntVector VolumetricFogResourceGridSize = GetVolumetricFogResourceGridSize(*this, VolumetricFogResourceGridPixelSize);
+	const FIntVector VolumetricFogViewGridSize = GetVolumetricFogViewGridSize(*this, VolumetricFogViewGridPixelSize);
+	const FVector2f ViewRectSize = FVector2f(ViewRect.Size());
+
 	SetupVolumetricFogUniformBufferParameters(ViewUniformShaderParameters);
 	ViewUniformShaderParameters.VolumetricFogViewGridUVToPrevViewRectUV = FVector2f::One();
 	ViewUniformShaderParameters.VolumetricFogPrevViewGridRectUVToResourceUV = FVector2f::One();
 	ViewUniformShaderParameters.VolumetricFogPrevUVMax = FVector2f::One();
 	ViewUniformShaderParameters.VolumetricFogPrevUVMaxForTemporalBlend = FVector2f::One();
-	int32 VolumetricFogViewGridPixelSize;
-	int32 VolumetricFogResourceGridPixelSize;
-	const FIntVector VolumetricFogResourceGridSize = GetVolumetricFogResourceGridSize(*this, VolumetricFogViewGridPixelSize);
-	const FIntVector VolumetricFogViewGridSize = GetVolumetricFogViewGridSize(*this, VolumetricFogResourceGridPixelSize);
-	const FVector2f ViewRectSize = FVector2f(ViewRect.Size());
+	ViewUniformShaderParameters.VolumetricFogPrevResourceGridSize = FVector3f(VolumetricFogResourceGridSize);
 	if (ViewState)
 	{
 		// Compute LightScatteringViewGridUVToViewRectVolumeUV, for the current frame resolution and volume texture resolution according to grid size.
@@ -1633,6 +1635,7 @@ void FViewInfo::SetupUniformBufferParameters(
 		ViewUniformShaderParameters.VolumetricFogPrevViewGridRectUVToResourceUV = ViewState->VolumetricFogPrevViewGridRectUVToResourceUV;
 		ViewUniformShaderParameters.VolumetricFogPrevUVMax = ViewState->VolumetricFogPrevUVMax;
 		ViewUniformShaderParameters.VolumetricFogPrevUVMaxForTemporalBlend = ViewState->VolumetricFogPrevUVMaxForTemporalBlend;
+		ViewUniformShaderParameters.VolumetricFogPrevResourceGridSize = FVector3f(ViewState->VolumetricFogPrevResourceGridSize);
 	}
 	ViewUniformShaderParameters.VolumetricFogScreenToResourceUV = ViewRectSize / (FVector2f(VolumetricFogResourceGridSize.X, VolumetricFogResourceGridSize.Y) * VolumetricFogResourceGridPixelSize);
 	ViewUniformShaderParameters.VolumetricFogUVMax = GetVolumetricFogUVMaxForSampling(ViewRectSize, VolumetricFogResourceGridSize, VolumetricFogResourceGridPixelSize);
diff --git a/Engine/Source/Runtime/Renderer/Private/VolumetricFog.cpp b/Engine/Source/Runtime/Renderer/Private/VolumetricFog.cpp
--- a/Engine/Source/Runtime/Renderer/Private/VolumetricFog.cpp
+++ b/Engine/Source/Runtime/Renderer/Private/VolumetricFog.cpp
@@ -131,7 +131,7 @@ FAutoConsoleVariableRef CVarVolumetricRectLightTexture(
 	ECVF_RenderThreadSafe
 );
 
-int32 GVolumetricFogConservativeDepth = 0;
+int32 GVolumetricFogConservativeDepth = 1;
 FAutoConsoleVariableRef CVarVolumetricFogConservativeDepth(
 	TEXT("r.VolumetricFog.ConservativeDepth"),
 	GVolumetricFogConservativeDepth,
@@ -1359,6 +1359,9 @@ void FSceneRenderer::SetupVolumetricFog()
 				View.ViewState->VolumetricFogPrevViewGridRectUVToResourceUV = FVector2f::One();
 				View.ViewState->VolumetricFogPrevUVMax = FVector2f::One();
 				View.ViewState->VolumetricFogPrevUVMaxForTemporalBlend = FVector2f::One();
+
+				int32 VolumetricFogResourceGridPixelSize;
+				View.ViewState->VolumetricFogPrevResourceGridSize = GetVolumetricFogResourceGridSize(View, VolumetricFogResourceGridPixelSize);
 			}
 		}
 	}
@@ -1877,6 +1880,7 @@ void FSceneRenderer::ComputeVolumetricFog(FRDGBuilder& GraphBuilder,
 			View.ViewState->VolumetricFogPrevViewGridRectUVToResourceUV = FVector2f(VolumetricFogViewGridSize.X, VolumetricFogViewGridSize.Y) / FVector2f(VolumetricFogResourceGridSize.X, VolumetricFogResourceGridSize.Y);
 			View.ViewState->VolumetricFogPrevUVMax = GetVolumetricFogUVMaxForSampling(ViewRectSize, VolumetricFogResourceGridSize, VolumetricFogGridPixelSize);
 			View.ViewState->VolumetricFogPrevUVMaxForTemporalBlend = GetVolumetricFogPrevUVMaxForTemporalBlend(ViewRectSize, VolumetricFogResourceGridSize, VolumetricFogGridPixelSize);
+			View.ViewState->VolumetricFogPrevResourceGridSize = VolumetricFogResourceGridSize;
 		}
 		else if (View.ViewState)
 		{
@@ -1886,6 +1890,7 @@ void FSceneRenderer::ComputeVolumetricFog(FRDGBuilder& GraphBuilder,
 			View.ViewState->VolumetricFogPrevViewGridRectUVToResourceUV = FVector2f::One();
 			View.ViewState->VolumetricFogPrevUVMax = FVector2f::One();
 			View.ViewState->VolumetricFogPrevUVMaxForTemporalBlend = FVector2f::One();
+			View.ViewState->VolumetricFogPrevResourceGridSize = VolumetricFogResourceGridSize;
 		}
 
 		if (bUseTemporalReprojection && GVolumetricFogConservativeDepth > 0)
