Change SpecularProfile to run on complex material and enable on all platforms to reduce visual disparity.

diff --git a/Engine/Config/Android/BaseAndroidEngine.ini b/Engine/Config/Android/BaseAndroidEngine.ini
--- a/Engine/Config/Android/BaseAndroidEngine.ini
+++ b/Engine/Config/Android/BaseAndroidEngine.ini
@@ -31,5 +31,4 @@ r.Substrate.ShadingQuality=2
 r.Substrate.TileCoord8bits=1
 r.Substrate.SheenQuality=2
 r.Substrate.Glints=0
-r.Substrate.SpecularProfile=0
 r.Substrate.ClosuresPerPixel=1
\ No newline at end of file
diff --git a/Engine/Config/BaseDeviceProfiles.ini b/Engine/Config/BaseDeviceProfiles.ini
--- a/Engine/Config/BaseDeviceProfiles.ini
+++ b/Engine/Config/BaseDeviceProfiles.ini
@@ -1299,7 +1299,6 @@ BaseProfileName=Windows
 +CVars=r.Substrate.TileCoord8bits=1
 +CVars=r.Substrate.SheenQuality=2
 +CVars=r.Substrate.Glints=0
-+CVars=r.Substrate.SpecularProfile=0
 +CVars=r.Substrate.ClosuresPerPixel=1
 
 [Windows_Preview_ES31_SDF DeviceProfile]
@@ -1310,5 +1309,4 @@ BaseProfileName=Windows
 +CVars=r.Substrate.TileCoord8bits=1
 +CVars=r.Substrate.SheenQuality=2
 +CVars=r.Substrate.Glints=0
-+CVars=r.Substrate.SpecularProfile=0
 +CVars=r.Substrate.ClosuresPerPixel=1
diff --git a/Engine/Config/IOS/BaseIOSEngine.ini b/Engine/Config/IOS/BaseIOSEngine.ini
--- a/Engine/Config/IOS/BaseIOSEngine.ini
+++ b/Engine/Config/IOS/BaseIOSEngine.ini
@@ -16,7 +16,6 @@ r.Substrate.ShadingQuality=2
 r.Substrate.TileCoord8bits=1
 r.Substrate.SheenQuality=2
 r.Substrate.Glints=0
-r.Substrate.SpecularProfile=0
 r.Substrate.ClosuresPerPixel=1
 ; Disables parallel GDME by default due to not enough threads
 r.Visibility.DynamicMeshElements.Parallel=0
diff --git a/Engine/Shaders/Private/PathTracing/PathTracingMaterialHitShader.usf b/Engine/Shaders/Private/PathTracing/PathTracingMaterialHitShader.usf
--- a/Engine/Shaders/Private/PathTracing/PathTracingMaterialHitShader.usf
+++ b/Engine/Shaders/Private/PathTracing/PathTracingMaterialHitShader.usf
@@ -21,10 +21,7 @@
 // Ensure that SSS albedo comes through in the material
 #define SUBSTRATE_SSS_MATERIAL_OVERRIDE 0
 
-// Ensure specular profile and glints are supported
-// TODO: the derivatives don't appear to be correct in all cases, causing incorrectly blurry lookups - so disable them for now unless using substrate
-//  We don't have an easy way to detect if the material is using glints or not, so use the complex special path which includes glints for now
-#define SUBSTRATE_COMPLEXSPECIALPATH 1
+// The derivatives are used for texture lookup filtering and Substrate's glints
 #define USE_ANALYTIC_DERIVATIVES 1
 
 #include "/Engine/Private/Common.ush"
@@ -659,11 +656,13 @@ RAY_TRACING_ENTRY_CLOSEST_HIT(PathTracingMaterialCHS,
 					Payload.SpecularProfileId = 0;
 					Payload.GlintValue = 1.f;
 
-#if SUBSTRATE_COMPLEXSPECIALPATH
+#if SUBSTRATE_COMPLEXPATH
 					if (BSDF_GETHASSPECPROFILE(BSDF))
 					{
 						Payload.SpecularProfileId = SLAB_SPECPROFILEID(BSDF);
 					}
+#endif
+#if SUBSTRATE_COMPLEXSPECIALPATH
 					if (BSDF_GETHASGLINT(BSDF))
 					{
 						Payload.GlintValue = SLAB_GLINT_VALUE(BSDF);
diff --git a/Engine/Shaders/Private/SpecularProfile.usf b/Engine/Shaders/Private/SpecularProfile.usf
--- a/Engine/Shaders/Private/SpecularProfile.usf
+++ b/Engine/Shaders/Private/SpecularProfile.usf
@@ -53,7 +53,9 @@ void MainCS(uint3 DispatchThreadId : SV_DispatchThreadID)
 uint SampleCount;
 int2 Resolution;
 uint ProfileCount;
-SamplerState TargetSampler;
+#if PERMUTATION_READ_WRITE_ON_ARRAY == 0
+Texture2DArray<float4> SourceTexture;
+#endif
 RWTexture2DArray<float4> TargetTexture;
 StructuredBuffer<float> ProfileSignBuffer;
 
@@ -63,7 +65,11 @@ float3 EvaluateSpecularProfile(float NoV, float NoL, float VoH, float NoH, float
 	uint3 FetchCoord = uint3(ProfileCoord.xy * (Resolution - 1), ProfileCoord.z);
 	FetchCoord.x = clamp(FetchCoord.x, 0u, uint(Resolution.x - 1));
 	FetchCoord.y = clamp(FetchCoord.y, 0u, uint(Resolution.y - 1));
+	#if PERMUTATION_READ_WRITE_ON_ARRAY
 	return TargetTexture[FetchCoord].xyz;
+	#else
+	return SourceTexture[FetchCoord].xyz;
+	#endif
 }
 
 float3 Convolve(float3 V, float3 N, float Roughness, float SpecularProfileId)
@@ -117,7 +123,7 @@ void MainCS(uint3 DispatchThreadId : SV_DispatchThreadID)
 		const float3 ConvolvedColor = Convolve(V, N, Roughness, SpecularProfileId);
 
 		const uint3 OutCoord = uint3(PixelPos, ProfileIndex * SUBSTRATE_SPECULAR_PROFILE_ENTRY_COUNT + SUBSTRATE_SPECULAR_PROFILE_ENTRY_ENV);
-		TargetTexture[OutCoord] = float4(ConvolvedColor, 1);		
+		TargetTexture[OutCoord] = float4(ConvolvedColor, 1);
 	}
 }
 #endif
diff --git a/Engine/Shaders/Private/Substrate/Substrate.ush b/Engine/Shaders/Private/Substrate/Substrate.ush
--- a/Engine/Shaders/Private/Substrate/Substrate.ush
+++ b/Engine/Shaders/Private/Substrate/Substrate.ush
@@ -120,7 +120,7 @@
 
 #define SUBSTRATE_GLINTS_ENABLED (SUBSTRATE_COMPLEXSPECIALPATH && PLATFORM_ENABLES_SUBSTRATE_GLINTS)
 
-#define SUBSTRATE_SPECPROFILE_ENABLED (SUBSTRATE_COMPLEXSPECIALPATH && PLATFORM_ENABLES_SUBSTRATE_SPECULAR_PROFILE)
+#define SUBSTRATE_SPECPROFILE_ENABLED (SUBSTRATE_COMPLEXPATH && PLATFORM_ENABLES_SUBSTRATE_SPECULAR_PROFILE)
 
 #define SUBSTRATE_LEGACY_MATERIAL_APPLIES_FINAL_WEIGHT (!MATERIALBLENDING_MASKED && !MATERIALBLENDING_SOLID && SUBSTRATE_LEGACY_PREMULT_ALPHA_OVERRIDE==0 || SUBSTRATE_MATERIAL_EXPORT_FROM_TRANSLUCENT)
 
@@ -177,6 +177,8 @@ struct FSubstrateBSDF
 
 #if SUBSTRATE_COMPLEXSPECIALPATH	
 	float4 VGPRs[7];
+#elif SUBSTRATE_COMPLEXPATH
+	float4 VGPRs[6];
 #else
 	float4 VGPRs[5];
 #endif
@@ -526,14 +528,16 @@ struct FSubstrateBSDF
 #define SLAB_FUZZ_COLOR(X)				X.VGPRs[4].xyz
 #define SLAB_FUZZ_AMOUNT(X)				X.VGPRs[4].w
 #define SLAB_FUZZ_ROUGHNESS(X)			X.VGPRs[2].w
+#if SUBSTRATE_COMPLEXPATH 
+// Specular profile
+#define SLAB_SPECPROFILEID(X)			X.VGPRs[5].w
+#endif
 #if SUBSTRATE_COMPLEXSPECIALPATH 
 // Glints
 #define SLAB_GLINT_UV(X)				X.VGPRs[5].xy
 #define SLAB_GLINT_VALUE(X)				X.VGPRs[5].z
 #define SLAB_GLINT_UVDDX(X)				X.VGPRs[6].xy
 #define SLAB_GLINT_UVDDY(X)				X.VGPRs[6].zw
-// Specular profile
-#define SLAB_SPECPROFILEID(X)			X.VGPRs[5].w
 #endif
 
 ///////////////////////////////////////////////////////////////////////////////
@@ -1428,7 +1432,6 @@ FSubstrateData InternalGetSubstrateSlabBSDF(
 	BSDF_SETTHICKNESSCM		(SubstrateData.InlinedBSDF,		ThicknessCm);
 #endif
 
-#if SUBSTRATE_COMPLEXSPECIALPATH
 #if SUBSTRATE_GLINTS_ENABLED
 	if (bHasGlint)
 	{
@@ -1446,8 +1449,7 @@ FSubstrateData InternalGetSubstrateSlabBSDF(
 		SLAB_SPECPROFILEID		(SubstrateData.InlinedBSDF) = SpecularProfileId * 255; // Stored as 0-255 index, and sign store parametrization type (View/Light or Half-Angle)
 	}
 #endif // SUBSTRATE_SPECPROFILE_ENABLED
-#endif // SUBSTRATE_COMPLEXSPECIALPATH
-
+	
 #if !SUBSTRATE_FASTPATH
 	// See how UMaterialExpressionSubstrateSlabBSDF node interface is changed
 	BSDF_SETHASF90(SubstrateData.InlinedBSDF, any(F90 < 1.0f));
@@ -2529,7 +2531,6 @@ FSubstrateData SubstrateHorizontalMixingParameterBlending(FSubstrateData Backgro
 	SLAB_FUZZ_ROUGHNESS(ResultBSDF)	=	lerp(SLAB_FUZZ_ROUGHNESS(BackgroundBSDF),		SLAB_FUZZ_ROUGHNESS(ForegroundBSDF),		ForegroundOtherMixFactor);
 	SLAB_HAZINESS(ResultBSDF) =			LerpHazinessParameterBlending(BackgroundBSDF,	ForegroundBSDF,								ForegroundOtherMixFactor);
 
-#if SUBSTRATE_COMPLEXSPECIALPATH
 #if SUBSTRATE_GLINTS_ENABLED
 	BSDF_SETHASGLINT(ResultBSDF,		BSDF_GETHASGLINT(ForegroundBSDF)				| BSDF_GETHASGLINT(BackgroundBSDF));
 	SLAB_GLINT_UV(ResultBSDF) =			lerp(SLAB_GLINT_UV(BackgroundBSDF),				SLAB_GLINT_UV(ForegroundBSDF),				ForegroundOtherMixFactor);
@@ -2541,7 +2542,6 @@ FSubstrateData SubstrateHorizontalMixingParameterBlending(FSubstrateData Backgro
 	BSDF_SETHASSPECPROFILE(ResultBSDF,	BSDF_GETHASSPECPROFILE(ForegroundBSDF)			| BSDF_GETHASSPECPROFILE(BackgroundBSDF));
 	SLAB_SPECPROFILEID(ResultBSDF) =	BSDF_GETHASSPECPROFILE(ForegroundBSDF) ? SLAB_SPECPROFILEID(ForegroundBSDF) : SLAB_SPECPROFILEID(BackgroundBSDF);
 #endif // SUBSTRATE_SPECPROFILE_ENABLED
-#endif // SUBSTRATE_COMPLEXSPECIALPATH
 	
 
 	if (BSDF_GETSSSTYPE(ResultBSDF) == SSS_TYPE_DIFFUSION_PROFILE)
@@ -2660,7 +2660,6 @@ FSubstrateData SubstrateAddParameterBlending(FSubstrateData A, FSubstrateData B,
 	SLAB_FUZZ_ROUGHNESS(ResultBSDF) =	lerp(SLAB_FUZZ_ROUGHNESS(BBSDF),				SLAB_FUZZ_ROUGHNESS(ABSDF),					AMixFactor);
 	SLAB_HAZINESS(ResultBSDF) =			LerpHazinessParameterBlending(BBSDF,			ABSDF,										AMixFactor);
 
-#if SUBSTRATE_COMPLEXSPECIALPATH
 #if SUBSTRATE_GLINTS_ENABLED
 	BSDF_SETHASGLINT(ResultBSDF,		BSDF_GETHASGLINT(BBSDF)							| BSDF_GETHASGLINT(ABSDF));
 	SLAB_GLINT_UV(ResultBSDF) =			lerp(SLAB_GLINT_UV(BBSDF),						SLAB_GLINT_UV(ABSDF), AMixFactor);
@@ -2672,7 +2671,6 @@ FSubstrateData SubstrateAddParameterBlending(FSubstrateData A, FSubstrateData B,
 	BSDF_SETHASSPECPROFILE(ResultBSDF,	BSDF_GETHASSPECPROFILE(BBSDF) | BSDF_GETHASSPECPROFILE(ABSDF));
 	SLAB_SPECPROFILEID(ResultBSDF) = BSDF_GETHASSPECPROFILE(BBSDF) ? SLAB_SPECPROFILEID(BBSDF) : SLAB_SPECPROFILEID(ABSDF);
 #endif // SUBSTRATE_SPECPROFILE_ENABLED
-#endif // SUBSTRATE_COMPLEXSPECIALPATH
 
 	if (BSDF_GETSSSTYPE(ResultBSDF) == SSS_TYPE_DIFFUSION_PROFILE)
 	{
@@ -2801,7 +2799,6 @@ FSubstrateBSDF SubstrateVerticalLayeringParameterBlending(in FSubstrateBSDF TopB
 	SLAB_FUZZ_ROUGHNESS(ResultBSDF) =	lerp(SLAB_FUZZ_ROUGHNESS(TopBSDF),				SLAB_FUZZ_ROUGHNESS(BaseBSDF),TopToBottomLerp);
 	SLAB_HAZINESS(ResultBSDF) =			LerpHazinessParameterBlending(TopBSDF,			BaseBSDF,					TopToBottomLerp);
 
-#if SUBSTRATE_COMPLEXSPECIALPATH
 #if SUBSTRATE_GLINTS_ENABLED
 	BSDF_SETHASGLINT(ResultBSDF,		BSDF_GETHASGLINT(TopBSDF)						| BSDF_GETHASGLINT(BaseBSDF));
 	SLAB_GLINT_UV(ResultBSDF) =			lerp(SLAB_GLINT_UV(TopBSDF),					SLAB_GLINT_UV(BaseBSDF), TopToBottomLerp);
@@ -2813,7 +2810,6 @@ FSubstrateBSDF SubstrateVerticalLayeringParameterBlending(in FSubstrateBSDF TopB
 	BSDF_SETHASSPECPROFILE(ResultBSDF,	BSDF_GETHASSPECPROFILE(TopBSDF) | BSDF_GETHASSPECPROFILE(BaseBSDF));
 	SLAB_SPECPROFILEID(ResultBSDF) = 	BSDF_GETHASSPECPROFILE(TopBSDF) ? SLAB_SPECPROFILEID(TopBSDF) : SLAB_SPECPROFILEID(BaseBSDF);
 #endif // SUBSTRATE_SPECPROFILE_ENABLED
-#endif // SUBSTRATE_COMPLEXSPECIALPATH
 
 	// Anisotropy is disabled with vertical blending. 
 	SLAB_ANISOTROPY(ResultBSDF) =		0.f;
@@ -4395,8 +4391,6 @@ FSubstrateBSDF UnpackSubstrateBSDFIn(FSubstrateMaterialContainer SubstrateBuffer
 			SubstrateLoad_ColorGamma2ToLinearAlphaLinear(SubstrateBuffer, SubstrateAddressing, OutBSDF.TransmittanceAboveAlongN, OutBSDF.CoverageAboveAlongN);
 		}
 
-
-	#if SUBSTRATE_COMPLEXSPECIALPATH
 	#if SUBSTRATE_GLINTS_ENABLED
 		if (BSDF_GETHASGLINT(OutBSDF))
 		{
@@ -4415,7 +4409,6 @@ FSubstrateBSDF UnpackSubstrateBSDFIn(FSubstrateMaterialContainer SubstrateBuffer
 			UnpackSpecularProfile(SubstrateLoadUint1(SubstrateBuffer, SubstrateAddressing), SLAB_SPECPROFILEID(OutBSDF));
 		}
 	#endif // SUBSTRATE_SPECPROFILE_ENABLED
-	#endif // SUBSTRATE_COMPLEXSPECIALPATH
 	}
 	break;
 	case SUBSTRATE_BSDF_TYPE_HAIR:
diff --git a/Engine/Shaders/Private/Substrate/SubstrateEvaluation.ush b/Engine/Shaders/Private/Substrate/SubstrateEvaluation.ush
--- a/Engine/Shaders/Private/Substrate/SubstrateEvaluation.ush
+++ b/Engine/Shaders/Private/Substrate/SubstrateEvaluation.ush
@@ -292,7 +292,7 @@ half3 Substrate_F_GGX(half3 F0, half3 F90, half VoH)
 // Transition from Glint to GGX (0=Glint <-> 1=GGX)
 float SubstrateTransitionGlintToGGX(float GlintValue)
 {
-#if SUBSTRATE_COMPLEXSPECIALPATH && SUBSTRATE_GLINTS_ALLOWED && SUBSTRATE_GLINTS_ENABLED
+#if SUBSTRATE_GLINTS_ALLOWED && SUBSTRATE_GLINTS_ENABLED
 	return saturate(saturate(GlintValue - 0.95f) / 0.05f);
 #else
 	return 1.f;
@@ -303,7 +303,7 @@ float Substrate_D_Glint(FSubstrateBSDFContext BSDFContext, float SafeRoughness,
 {
 	float OutD = GGX_D;
 
-	#if SUBSTRATE_COMPLEXSPECIALPATH && SUBSTRATE_GLINTS_ALLOWED && SUBSTRATE_GLINTS_ENABLED
+	#if SUBSTRATE_GLINTS_ALLOWED && SUBSTRATE_GLINTS_ENABLED
 	BRANCH
 	if (BSDF_GETHASGLINT(BSDFContext.BSDF))
 	{
@@ -322,7 +322,7 @@ float EvaluateGlintRect(FSubstrateBSDFContext BSDFContext, float SafeRoughness,
 {
 	float Out = 1.0f;
 
-	#if SUBSTRATE_COMPLEXSPECIALPATH && SUBSTRATE_GLINTS_ALLOWED && SUBSTRATE_GLINTS_ENABLED
+	#if SUBSTRATE_GLINTS_ALLOWED && SUBSTRATE_GLINTS_ENABLED
 	BRANCH
 	if (BSDF_GETHASGLINT(BSDFContext.BSDF))
 	{
@@ -346,7 +346,7 @@ float EvaluateGlintEnv(FSubstrateBSDFContext BSDFContext, float SafeRoughness, i
 {
 	float Out = 1.0f;
 
-#if SUBSTRATE_COMPLEXSPECIALPATH && SUBSTRATE_GLINTS_ALLOWED && SUBSTRATE_GLINTS_ENABLED
+#if SUBSTRATE_GLINTS_ALLOWED && SUBSTRATE_GLINTS_ENABLED
 	BRANCH
 	if (BSDF_GETHASGLINT(BSDFContext.BSDF))
 	{
@@ -378,7 +378,7 @@ float EvaluateGlintEnv(FSubstrateBSDFContext BSDFContext, float SafeRoughness, i
 		Out = lerp(saturate(GlintD), 1.f, TransitionGlintToGGX);
 		// We keep the same SpecularSafeRoughness for single glint environment sampling to match the roughness of the surface overall, especially when the glints becomes infinitely small and thus invisible.
 	}
-#endif // SUBSTRATE_COMPLEXSPECIALPATH && SUBSTRATE_GLINTS_ALLOWED && SUBSTRATE_GLINTS_ENABLED
+#endif // SUBSTRATE_GLINTS_ALLOWED && SUBSTRATE_GLINTS_ENABLED
 
 	return Out;
 }
@@ -777,7 +777,7 @@ FSubstrateEvaluateResult SubstrateEvaluateBSDFCommon(FSubstrateBSDFContext BSDFC
 			#endif
 
 			// Lighting LUT
-			#if SUBSTRATE_FASTPATH==0 && SUBSTRATE_COMPLEXSPECIALPATH && SUBSTRATE_SPECPROFILE_ENABLED
+			#if SUBSTRATE_FASTPATH==0 && SUBSTRATE_SPECPROFILE_ENABLED
 			BRANCH
 			if (BSDF_GETHASSPECPROFILE(BSDFContext.BSDF))
 			{
@@ -1631,20 +1631,20 @@ FSubstrateEnvLightResult SubstrateEvaluateForEnvLight(FSubstrateBSDFContext BSDF
 				SubstrateEnvLightResult.SpecularColor = F0;
 
 				// Glint shading
-			#if SUBSTRATE_COMPLEXSPECIALPATH && SUBSTRATE_GLINTS_ALLOWED && SUBSTRATE_GLINTS_ENABLED
+			#if SUBSTRATE_GLINTS_ALLOWED && SUBSTRATE_GLINTS_ENABLED
 				{
 					SubstrateEnvLightResult.SpecularWeight *= EvaluateGlintEnv(BSDFContext, SubstrateEnvLightResult.SpecularSafeRoughness, SubstrateEnvLightResult.SpecularDirection);
 				}
-			#endif				
+			#endif
 
 				// Specular profile
-			#if SUBSTRATE_COMPLEXSPECIALPATH && SUBSTRATE_SPECPROFILE_ENABLED
+			#if SUBSTRATE_SPECPROFILE_ENABLED
 				BRANCH
 				if (BSDF_GETHASSPECPROFILE(BSDFContext.BSDF))
 				{
 					SubstrateEnvLightResult.SpecularWeight *= EvaluateSpecularProfileEnv(SLAB_SPECPROFILEID(BSDFContext.BSDF), BSDFContext.SatNoV, SubstrateEnvLightResult.SpecularSafeRoughness);
 				}
-			#endif // SUBSTRATE_COMPLEXSPECIALPATH
+			#endif
 
 				SubstrateEnvLightResult.SSRSpecularWeight = SubstrateEnvLightResult.SpecularWeight;
 
@@ -1701,13 +1701,13 @@ FSubstrateEnvLightResult SubstrateEvaluateForEnvLight(FSubstrateBSDFContext BSDF
 					}
 
 					// Specular profile for haze
-				#if SUBSTRATE_COMPLEXSPECIALPATH && SUBSTRATE_SPECPROFILE_ENABLED
+				#if SUBSTRATE_SPECPROFILE_ENABLED
 					BRANCH
 					if (BSDF_GETHASSPECPROFILE(BSDFContext.BSDF))
 					{
 						SubstrateEnvLightResult.SpecularHazeWeight *= EvaluateSpecularProfileEnv(SLAB_SPECPROFILEID(BSDFContext.BSDF), BSDFContext.SatNoV, SubstrateEnvLightResult.SpecularHazeSafeRoughness);
 					}
-				#endif // SUBSTRATE_COMPLEXSPECIALPATH
+				#endif
 
 				}
 				#endif // SUBSTRATE_FASTPATH==0
diff --git a/Engine/Shaders/Private/Substrate/SubstrateExport.ush b/Engine/Shaders/Private/Substrate/SubstrateExport.ush
--- a/Engine/Shaders/Private/Substrate/SubstrateExport.ush
+++ b/Engine/Shaders/Private/Substrate/SubstrateExport.ush
@@ -201,11 +201,11 @@ void PackSubstrateOut(
 						const bool LuminanceWeightEqualOne = all(BSDF.LuminanceWeightV == 1.0f);
 
 						#if SUBSTRATE_COMPLEXSPECIALPATH
-						const bool bForceComplexSpecialPath = (SUBSTRATE_SPECPROFILE_ENABLED && BSDF_GETHASSPECPROFILE(BSDF)) || (SUBSTRATE_GLINTS_ENABLED && BSDF_GETHASGLINT(BSDF));
+						const bool bForceComplexSpecialPath = (SUBSTRATE_GLINTS_ENABLED && BSDF_GETHASGLINT(BSDF));
 						#else
 						const bool bForceComplexSpecialPath = false;
 						#endif
-						const bool bForceComplexMaterial = BSDF_GETHASANISOTROPY(BSDF) || bForceComplexSpecialPath || !LuminanceWeightEqualOne;
+						const bool bForceComplexMaterial = BSDF_GETHASANISOTROPY(BSDF) || (SUBSTRATE_SPECPROFILE_ENABLED && BSDF_GETHASSPECPROFILE(BSDF)) || bForceComplexSpecialPath || !LuminanceWeightEqualOne;
 
 						// Update simple material compatibility
 						bIsSimpleMaterial = bIsSimpleMaterial && IsSubstrateSlabFastPathCompatible(BSDF) && !bForceComplexMaterial;
@@ -831,7 +831,6 @@ void PackSubstrateOut(
 								SUBSTRATE_STORE_UINT1(PackColorLinearToGamma2AlphaLinear(float4(BSDF.TransmittanceAboveAlongN, BSDF.CoverageAboveAlongN)));
 							}
 
-							#if SUBSTRATE_COMPLEXSPECIALPATH
 							#if SUBSTRATE_GLINTS_ENABLED
 							if (BSDF_GETHASGLINT(BSDF))
 							{
@@ -851,7 +850,6 @@ void PackSubstrateOut(
 								SUBSTRATE_STORE_UINT1(PackSpecularProfile(SLAB_SPECPROFILEID(BSDF)));
 							}
 							#endif // SUBSTRATE_SPECPROFILE_ENABLED
-							#endif // SUBSTRATE_COMPLEXSPECIALPATH
 
 							// 8-28 bytes
 						}
diff --git a/Engine/Source/Runtime/Engine/Private/Materials/HLSLMaterialTranslator.cpp b/Engine/Source/Runtime/Engine/Private/Materials/HLSLMaterialTranslator.cpp
--- a/Engine/Source/Runtime/Engine/Private/Materials/HLSLMaterialTranslator.cpp
+++ b/Engine/Source/Runtime/Engine/Private/Materials/HLSLMaterialTranslator.cpp
@@ -13286,23 +13286,9 @@ bool FHLSLMaterialTranslator::FSubstrateCompilationContext::SubstrateGenerateDer
 					{
 					case SUBSTRATE_BSDF_TYPE_SLAB:
 					{
-						const uint32 SingleMask = 
-							  ESubstrateBsdfFeature_EdgeColor
-							| ESubstrateBsdfFeature_Fuzz
-							| ESubstrateBsdfFeature_SSS
-							| ESubstrateBsdfFeature_SecondRoughnessOrSimpleClearCoat
-							| ESubstrateBsdfFeature_MFPPluggedIn;
-						const uint32 CompexMask = 
-							  ESubstrateBsdfFeature_Anisotropy
-							| ESubstrateBsdfFeature_Eye
-							| ESubstrateBsdfFeature_Hair;
-						const uint32 CompexSpecialMask = 
-							  ESubstrateBsdfFeature_Glint
-							| ESubstrateBsdfFeature_SpecularProfile;
-
-						SubstrateMaterialComplexity.bIsSimple = SubstrateMaterialComplexity.bIsSimple && !bMayHaveColoredWeight && !It.Has(CompexMask) && !It.Has(CompexSpecialMask) && !It.Has(SingleMask);
-						SubstrateMaterialComplexity.bIsSingle = SubstrateMaterialComplexity.bIsSingle && !bMayHaveColoredWeight && !It.Has(CompexMask) && !It.Has(CompexSpecialMask);
-						SubstrateMaterialComplexity.bIsComplexSpecial |= It.Has(CompexSpecialMask);
+						SubstrateMaterialComplexity.bIsSimple = SubstrateMaterialComplexity.bIsSimple && !bMayHaveColoredWeight && !It.Has(ESubstrateBsdfFeature_ComplexMask) && !It.Has(ESubstrateBsdfFeature_ComplexSpecialMask) && !It.Has(ESubstrateBsdfFeature_SingleMask);
+						SubstrateMaterialComplexity.bIsSingle = SubstrateMaterialComplexity.bIsSingle && !bMayHaveColoredWeight && !It.Has(ESubstrateBsdfFeature_ComplexMask) && !It.Has(ESubstrateBsdfFeature_ComplexSpecialMask);
+						SubstrateMaterialComplexity.bIsComplexSpecial |= It.Has(ESubstrateBsdfFeature_ComplexSpecialMask);
 						break;
 					}
 					case SUBSTRATE_BSDF_TYPE_HAIR:
diff --git a/Engine/Source/Runtime/Engine/Private/Materials/MaterialExpressionSubstrate.cpp b/Engine/Source/Runtime/Engine/Private/Materials/MaterialExpressionSubstrate.cpp
--- a/Engine/Source/Runtime/Engine/Private/Materials/MaterialExpressionSubstrate.cpp
+++ b/Engine/Source/Runtime/Engine/Private/Materials/MaterialExpressionSubstrate.cpp
@@ -957,14 +957,34 @@ int32 UMaterialExpressionSubstrateSlabBSDF::CompileDefaultSlab(FMaterialCompiler
 
 FSubstrateMaterialComplexity UMaterialExpressionSubstrateSlabBSDF::GetComplexity() const
 {
-	FSubstrateMaterialComplexity Out;
-	Out.bIsComplexSpecial = HasGlint() || HasSpecularProfile();
-	Out.bIsSingle = !HasAnisotropy() && (HasEdgeColor() || HasFuzz() || HasSecondRoughness() || HasMFPPluggedIn() || HasSSS());
-	Out.bIsSimple = !HasAnisotropy();
+	uint32 FeatureMask = ESubstrateBsdfFeature_None;
+	if (HasGlint()) 			{ FeatureMask |= ESubstrateBsdfFeature_Glint; } 
+	if (HasAnisotropy()) 		{ FeatureMask |= ESubstrateBsdfFeature_Anisotropy; } 
+	if (HasSpecularProfile()) 	{ FeatureMask |= ESubstrateBsdfFeature_SpecularProfile; } 
+	if (HasEdgeColor()) 		{ FeatureMask |= ESubstrateBsdfFeature_EdgeColor; } 
+	if (HasFuzz()) 				{ FeatureMask |= ESubstrateBsdfFeature_Fuzz; } 
+	if (HasSecondRoughness()) 	{ FeatureMask |= ESubstrateBsdfFeature_SecondRoughnessOrSimpleClearCoat; } 
+	if (HasMFPPluggedIn()) 		{ FeatureMask |= ESubstrateBsdfFeature_MFPPluggedIn; } 
+	if (HasSSS()) 				{ FeatureMask |= ESubstrateBsdfFeature_SSS; } 
 
-	// Mask out to only have a single possibility
-	Out.bIsSingle &= !Out.bIsComplexSpecial;
-	Out.bIsSimple &= !Out.bIsComplexSpecial && !Out.bIsSingle;
+	FSubstrateMaterialComplexity Out;
+	Out.Reset();
+	if ((FeatureMask & ESubstrateBsdfFeature_ComplexSpecialMask) != 0)
+	{
+		Out.bIsComplexSpecial = true;
+	}
+	else if ((FeatureMask & ESubstrateBsdfFeature_ComplexMask) != 0)
+	{
+		// Nothing
+	}
+	else if ((FeatureMask & ESubstrateBsdfFeature_SingleMask) != 0)
+	{
+		Out.bIsSingle = true;
+	}
+	else
+	{
+		Out.bIsSimple = true;
+	}
 
 	return Out;
 }
diff --git a/Engine/Source/Runtime/Engine/Private/Rendering/SpecularProfile.cpp b/Engine/Source/Runtime/Engine/Private/Rendering/SpecularProfile.cpp
--- a/Engine/Source/Runtime/Engine/Private/Rendering/SpecularProfile.cpp
+++ b/Engine/Source/Runtime/Engine/Private/Rendering/SpecularProfile.cpp
@@ -50,7 +50,7 @@ static bool ForceUpdateSpecularProfile()
 
 static bool IsSpecularProfileSupport(EShaderPlatform InShaderPlatform)
 {
-	return !IsMobilePlatform(InShaderPlatform) && IsFeatureLevelSupported(InShaderPlatform, ERHIFeatureLevel::SM5);
+	return IsFeatureLevelSupported(InShaderPlatform, ERHIFeatureLevel::SM5);
 }
 
 class FSpecularProfileCopyCS : public FGlobalShader
@@ -95,13 +95,14 @@ public:
 	DECLARE_GLOBAL_SHADER(FSpecularProfileConvolveCS);
 	SHADER_USE_PARAMETER_STRUCT(FSpecularProfileConvolveCS, FGlobalShader);
 
-	using FPermutationDomain = TShaderPermutationDomain<>;
+	class FReadWriteOnArray : SHADER_PERMUTATION_BOOL("PERMUTATION_READ_WRITE_ON_ARRAY");
+	using FPermutationDomain = TShaderPermutationDomain<FReadWriteOnArray>;
 
 	BEGIN_SHADER_PARAMETER_STRUCT(FParameters, )
 		SHADER_PARAMETER(uint32, SampleCount)
 		SHADER_PARAMETER(FIntPoint, Resolution)
 		SHADER_PARAMETER(uint32, ProfileCount)
-		SHADER_PARAMETER_SAMPLER(SamplerState, TargetSampler)
+		SHADER_PARAMETER_RDG_TEXTURE(Texture2DArray, SourceTexture)
 		SHADER_PARAMETER_RDG_TEXTURE_UAV(RWTexture2DArray, TargetTexture)
 		SHADER_PARAMETER_RDG_BUFFER_SRV(StructuredBuffer, ProfileSignBuffer)
 	END_SHADER_PARAMETER_STRUCT()
@@ -361,7 +362,10 @@ IPooledRenderTarget* FSpecularProfileTextureManager::GetAtlasTexture(FRDGBuilder
 		RDG_EVENT_SCOPE(GraphBuilder, "SpecularProfile");
 
 		// 1. Create atlas texture
-		FPooledRenderTargetDesc Desc(FPooledRenderTargetDesc::Create2DArrayDesc(FIntPoint(Resolution, Resolution), PF_FloatR11G11B10, FClearValueBinding::None, TexCreate_None, TexCreate_UAV | TexCreate_ShaderResource, false, LayerCount * SUBSTRATE_SPECULAR_PROFILE_ENTRY_COUNT));
+		const bool bR11G11B10Supported = GPixelFormats[PF_FloatR11G11B10].Supported && UE::PixelFormat::HasCapabilities(PF_FloatR11G11B10, EPixelFormatCapabilities::TextureFilterable);
+		const EPixelFormat Format = bR11G11B10Supported ? PF_FloatR11G11B10 : PF_FloatRGBA;
+
+		FPooledRenderTargetDesc Desc(FPooledRenderTargetDesc::Create2DArrayDesc(FIntPoint(Resolution, Resolution), Format, FClearValueBinding::None, TexCreate_None, TexCreate_UAV | TexCreate_ShaderResource, false, LayerCount * SUBSTRATE_SPECULAR_PROFILE_ENTRY_COUNT));
 		GRenderTargetPool.FindFreeElement(RHICmdList, Desc, GSpecularProfileTextureAtlas, TEXT("SpecularProfileTexture"));
 
 		// 2. Fill in profiles
@@ -425,6 +429,16 @@ IPooledRenderTarget* FSpecularProfileTextureManager::GetAtlasTexture(FRDGBuilder
 		// Convolve all the profiles
 		if (LayerCount > 0)
 		{
+			// Certain platform don't support read/write on the same texture array, even if this done safely on different slices
+			const bool bPlatformSupportReadWriteOnTextureArray = !IsOpenGLPlatform(ShaderPlatform);
+			FRDGTextureRef SourceSpecularProfileTexture = nullptr;
+			if (!bPlatformSupportReadWriteOnTextureArray)
+			{
+				const FIntPoint Extent = SpecularProfileTexture->Desc.Extent;
+				SourceSpecularProfileTexture = GraphBuilder.CreateTexture(SpecularProfileTexture->Desc, TEXT("ConvolvedSpecularProfileTexture"));
+				AddCopyTexturePass(GraphBuilder, SpecularProfileTexture, SourceSpecularProfileTexture);
+			}
+
 			TArray<float> ProfileSign;
 			ProfileSign.SetNum(LayerCount);
 			for (uint32 LayerIt=0; LayerIt < LayerCount; ++LayerIt)
@@ -439,10 +453,12 @@ IPooledRenderTarget* FSpecularProfileTextureManager::GetAtlasTexture(FRDGBuilder
 			PassParameters->ProfileCount = LayerCount;
 			PassParameters->Resolution = SpecularProfileTexture->Desc.Extent;
 			PassParameters->ProfileSignBuffer = GraphBuilder.CreateSRV(ProfileSignBuffer);
-			PassParameters->TargetSampler = TStaticSamplerState<SF_Point, AM_Clamp, AM_Clamp, AM_Clamp>::GetRHI();;
+			PassParameters->SourceTexture = SourceSpecularProfileTexture;
 			PassParameters->TargetTexture = GraphBuilder.CreateUAV(SpecularProfileTexture);
 
-			TShaderMapRef<FSpecularProfileConvolveCS> Shader(GlobalShaderMap);
+			FSpecularProfileConvolveCS::FPermutationDomain PermutationVector;
+			PermutationVector.Set<FSpecularProfileConvolveCS::FReadWriteOnArray>(PassParameters->SourceTexture == nullptr);
+			TShaderMapRef<FSpecularProfileConvolveCS> Shader(GlobalShaderMap, PermutationVector);
 			FComputeShaderUtils::AddPass(GraphBuilder, 
 				RDG_EVENT_NAME("SpecularProfile::Convolve"), 
 				Shader, 
diff --git a/Engine/Source/Runtime/Engine/Public/Rendering/SubstrateMaterialShared.h b/Engine/Source/Runtime/Engine/Public/Rendering/SubstrateMaterialShared.h
--- a/Engine/Source/Runtime/Engine/Public/Rendering/SubstrateMaterialShared.h
+++ b/Engine/Source/Runtime/Engine/Public/Rendering/SubstrateMaterialShared.h
@@ -38,7 +38,24 @@ enum ESubstrateBsdfFeature
 	ESubstrateBsdfFeature_Eye                              = 1u<<8u,
 	ESubstrateBsdfFeature_EyeIrisNormalPluggedIn           = 1u<<9u,
 	ESubstrateBsdfFeature_EyeIrisTangentPluggedIn          = 1u<<10u,
-	ESubstrateBsdfFeature_Hair                             = 1u<<11u
+	ESubstrateBsdfFeature_Hair                             = 1u<<11u,
+
+	// Complexity masks
+	ESubstrateBsdfFeature_SingleMask =	
+		  ESubstrateBsdfFeature_EdgeColor
+		| ESubstrateBsdfFeature_Fuzz
+		| ESubstrateBsdfFeature_SSS
+		| ESubstrateBsdfFeature_SecondRoughnessOrSimpleClearCoat
+		| ESubstrateBsdfFeature_MFPPluggedIn,
+
+	ESubstrateBsdfFeature_ComplexMask = 
+		  ESubstrateBsdfFeature_Anisotropy
+		| ESubstrateBsdfFeature_SpecularProfile
+		| ESubstrateBsdfFeature_Eye
+		| ESubstrateBsdfFeature_Hair,
+
+	ESubstrateBsdfFeature_ComplexSpecialMask = 
+		  ESubstrateBsdfFeature_Glint
 };
 
 struct FSubstrateOperator
diff --git a/Engine/Source/Runtime/Renderer/Private/DeferredShadingRenderer.cpp b/Engine/Source/Runtime/Renderer/Private/DeferredShadingRenderer.cpp
--- a/Engine/Source/Runtime/Renderer/Private/DeferredShadingRenderer.cpp
+++ b/Engine/Source/Runtime/Renderer/Private/DeferredShadingRenderer.cpp
@@ -1562,7 +1562,7 @@ void FDeferredShadingSceneRenderer::Render(FRDGBuilder& GraphBuilder)
 		SCOPE_CYCLE_COUNTER(STAT_FDeferredShadingSceneRenderer_Render_Init);
 		RDG_RHI_GPU_STAT_SCOPE(GraphBuilder, AllocateRendertargets);
 
-		// Force the subsurface profile texture to be updated.
+		// Force the subsurface profiles and specular profiles textures to be updated.
 		SubsurfaceProfile::UpdateSubsurfaceProfileTexture(GraphBuilder, ShaderPlatform);
 		SpecularProfile::UpdateSpecularProfileTextureAtlas(GraphBuilder, ShaderPlatform);
 
diff --git a/Engine/Source/Runtime/Renderer/Private/MobileShadingRenderer.cpp b/Engine/Source/Runtime/Renderer/Private/MobileShadingRenderer.cpp
--- a/Engine/Source/Runtime/Renderer/Private/MobileShadingRenderer.cpp
+++ b/Engine/Source/Runtime/Renderer/Private/MobileShadingRenderer.cpp
@@ -1057,7 +1057,7 @@ void FMobileSceneRenderer::Render(FRDGBuilder& GraphBuilder)
 
 	if (bRendererOutputFinalSceneColor)
 	{
-		// Force the subsurface profile texture to be updated.
+		// Force the subsurface profile & specular profile textures to be updated.
 		SubsurfaceProfile::UpdateSubsurfaceProfileTexture(GraphBuilder, ShaderPlatform);
 		SpecularProfile::UpdateSpecularProfileTextureAtlas(GraphBuilder, ShaderPlatform);
 
