Path Tracer: Improve tangent vector encoding

Instead of storing the tangent vector independently, store it relative to a canonical normal-aligned basis. This requires fewer bits since we are only encoding an angle.

The value is quantized to 16 bits for now, but this could be lowered to 10 bits in the future if we need more room in the payload.

Since we are spending less space on the tangent vector now, shuffle some fields around in the payload for both regular and substrate encodings to take advantage of the extra space.

diff --git a/Engine/Shaders/Private/PathTracing/PathTracingCommon.ush b/Engine/Shaders/Private/PathTracing/PathTracingCommon.ush
--- a/Engine/Shaders/Private/PathTracing/PathTracingCommon.ush
+++ b/Engine/Shaders/Private/PathTracing/PathTracingCommon.ush
@@ -33,7 +33,8 @@
 #define PATH_TRACING_PAYLOAD_INPUT_FLAG_HAS_SCENE_DEPTH               (1 << 4) // Indicates the payload contains a scene depth value (exclusive with dbuffer usage)
 #define PATH_TRACING_PAYLOAD_INPUT_FLAG_IGNORE_TRANSLUCENT            (1 << 5) // Indicates this ray should not process any translucent geometry (so we can skip them in our opaque traces)
 
-#define PATH_TRACING_PAYLOAD_NORMAL_BITS  32
+#define PATH_TRACING_PAYLOAD_NORMAL_BITS  		32
+#define PATH_TRACING_PAYLOAD_TANGENT_ANGLE_BITS 16 						// Could be dropped to 12 or even 10 if we need more room in the future
 
 // shading model extensions beyond the default count
 #define SHADINGMODELID_MEDIUM  			(SHADINGMODELID_NUM + 1)
@@ -70,6 +71,46 @@ float3 PayloadDecodeUnitVector(uint Encoded)
 	return OctahedronToUnitVector(2 * OctF - 1);
 }
 
+// Encode a tangent vector relative to a given Normal
+uint PayloadEncodeTangentAngle(float3 Normal, float3 Tangent)
+{
+	const uint Mask = (1u << (PATH_TRACING_PAYLOAD_TANGENT_ANGLE_BITS - 2)) - 1;
+	const uint SignShiftX = 31 - (PATH_TRACING_PAYLOAD_TANGENT_ANGLE_BITS - 2);
+	const uint SignShiftY = 31 - (PATH_TRACING_PAYLOAD_TANGENT_ANGLE_BITS - 1);
+	const uint SignMaskX = 0x80000000u >> SignShiftX;
+	const uint SignMaskY = 0x80000000u >> SignShiftY;
+	const float Scale = float(Mask);
+
+	float3x3 Basis = GetTangentBasis(Normal);
+	float X = dot(Basis[0], Tangent);
+	float Y = dot(Basis[1], Tangent);
+	// NOTE: We use 'Y' to define the angle so that if Tangent happens to be (0,0,0) it will be decoded back as the "x" axis in the reference basis
+	float PseudoAngle = abs(Y) / (abs(X) + abs(Y)); // Coarse ArcTan, returns value in [0,1]
+	uint QuantizedBits = uint(floor(PseudoAngle * Scale));
+	// Merge with sign bits of X and Y
+	return QuantizedBits | ((asuint(X) >> SignShiftX) & SignMaskX) | ((asuint(Y) >> SignShiftY) & SignMaskY);
+}
+
+// Decode a tangent vector relative to a given Normal
+float3 PayloadDecodeTangentAngle(float3 Normal, uint Bits)
+{
+	const uint Mask = (1u << (PATH_TRACING_PAYLOAD_TANGENT_ANGLE_BITS - 2)) - 1;
+	const uint SignShiftX = 31 - (PATH_TRACING_PAYLOAD_TANGENT_ANGLE_BITS - 2);
+	const uint SignShiftY = 31 - (PATH_TRACING_PAYLOAD_TANGENT_ANGLE_BITS - 1);
+	const uint SignMaskX = 0x80000000u >> SignShiftX;
+	const uint SignMaskY = 0x80000000u >> SignShiftY;
+	const float Scale = float(Mask);
+
+	float3x3 Basis = GetTangentBasis(Normal);
+	float Y = float(Bits & Mask) / Scale;
+	float2 P = float2(1.0f - Y, Y);
+	// restore sign bits
+	P.x = asfloat(asuint(P.x) | ((Bits & SignMaskX) << SignShiftX));
+	P.y = asfloat(asuint(P.y) | ((Bits & SignMaskY) << SignShiftY));
+	P = normalize(P);
+	return Basis[0] * P.x + Basis[1] * P.y;
+}
+
 uint PayloadEncodeHDRColor(float3 rgb)
 {
 	return PackRGB998E6(rgb);
@@ -240,53 +281,48 @@ struct FPackedPathTracingPayload : FMinimalPayload
 
 	void PackWorldNormal(float3 N)   { PackedData[3] = PayloadEncodeUnitVector(N); }
 	void PackRadiance(float3 R)      { PackedData[4] = PayloadEncodeHDRColor(R); }
+
+	float3 UnpackWorldTangent()       { return PayloadDecodeTangentAngle(UnpackWorldNormal(), PackedData[5]); }
+	float UnpackIor()                 { return f16tof32(PackedData[5] >> 16); }
 	
 #if PATHTRACING_SUBSTRATE_PAYLOAD
 
-	float3 UnpackDiffuseColor()       { return PayloadDecodeLDRColor(PackedData[5]); }
-	float3 UnpackSpecularColor()      { return PayloadDecodeLDRColor(PackedData[6]); }
-	float3 UnpackSpecularEdgeColor()  { return PayloadDecodeLDRColor(PackedData[7]); }
-	float4 UnpackRoughnessAniso()     { return PayloadDecodeRoughnessAniso(PackedData[8]); }
-	float3 UnpackWorldTangent()       { return PayloadDecodeUnitVector(PackedData[9]); }
+	float3 UnpackDiffuseColor()       { return PayloadDecodeLDRColor(PackedData[6]); }
+	float3 UnpackSpecularColor()      { return PayloadDecodeLDRColor(PackedData[7]); }
+	float3 UnpackSpecularEdgeColor()  { return PayloadDecodeLDRColor(PackedData[8]); }
+	float4 UnpackRoughnessAniso()     { return PayloadDecodeRoughnessAniso(PackedData[9]); }
 	float3 UnpackTransparencyColor()  { return PayloadDecodeLDRColor(PackedData[10]); }
 	float3 UnpackMeanFreePath()       { return float3(f16tof32(PackedData[11]),
 	                                                  f16tof32(PackedData[11] >> 16),
 	                                                  f16tof32(PackedData[12])); }
 	float UnpackPhaseG()              { return f16tof32(PackedData[12] >> 16); }
 	float3 UnpackFuzzColor()          { return PayloadDecodeLDRColor(PackedData[13]); }
-	float3 UnpackFuzzData()           { return PayloadDecodeLDRColor(PackedData[14]); } // x: UNUSED, y: FuzzAmount, z: FuzzRoughness
+	float2 UnpackFuzzData()           { return UnpackFloat2FromUInt(PackedData[14]); }
 	float3 UnpackWeightV()            { return PayloadDecodeLDRColor(PackedData[15]); }
 	float3 UnpackTransmittanceN()     { return PayloadDecodeLDRColor(PackedData[16]); }
-	float UnpackIor()                 { return Unpack10F(PackedData[17]); }
-	float UnpackCoverageAboveAlongN() { return Unpack10F(PackedData[17] >> 10); }
-	float UnpackSpecularProfileId()   { return Unpack10F(PackedData[17] >> 20); }
 
-	void PackDiffuseColor(float3 C)   { PackedData[5] = PayloadEncodeLDRColor(C); }
-	void PackSpecularColor(float3 C)  { PackedData[6] = PayloadEncodeLDRColor(C); }
-	void PackRoughnessAniso(float4 D) { PackedData[8] = PayloadEncodeRoughnessAniso(D.xyz, D.w); }
+	float UnpackCoverageAboveAlongN() { return abs(f16tof32(PackedData[17] >> 16)); }
+	float UnpackSpecularProfileId()   { return f16tof32(PackedData[17]); }
+
+	void PackDiffuseColor(float3 C)   { PackedData[6] = PayloadEncodeLDRColor(C); }
+	void PackSpecularColor(float3 C)  { PackedData[7] = PayloadEncodeLDRColor(C); }
+	void PackRoughnessAniso(float4 D) { PackedData[9] = PayloadEncodeRoughnessAniso(D.xyz, D.w); }
 
 #else
-	float3 UnpackBaseColor()         { return PayloadDecodeLDRColor(PackedData[5]); }
-	float3 UnpackTransparencyColor() { return PayloadDecodeLDRColor(PackedData[6]); }
-	float UnpackMetallic()           { return f16tof32(PackedData[7]); }
-	float UnpackSpecular()           { return f16tof32(PackedData[7] >> 16); }
-	float UnpackRoughness()          { return f16tof32(PackedData[8]); }
-	float UnpackAnisotropy()         { return f16tof32(PackedData[8] >> 16); }
-	float3 UnpackWorldTangent()      { return PayloadDecodeUnitVector(PackedData[9]); }
-	float4 UnpackCustomData0()       { return float4(f16tof32(PackedData[10]),
-	                                                 f16tof32(PackedData[10] >> 16),
-	                                                 f16tof32(PackedData[11]),
-	                                                 f16tof32(PackedData[11] >> 16)); }
-	float4 UnpackCustomData1()       { return float4(f16tof32(PackedData[12]),
-	                                                 f16tof32(PackedData[12] >> 16),
-	                                                 f16tof32(PackedData[13]),
-	                                                 f16tof32(PackedData[13] >> 16)); }
-	float UnpackIor()                { return f16tof32(PackedData[14]); }
-
-
-	void PackBaseColor(float3 C) { PackedData[5] = PayloadEncodeLDRColor(C); }
-	void PackMetallicSpecular(float M, float S) { PackedData[7] = f32tof16(M) | (f32tof16(S) << 16); }
-	void PackRoughnessAniso(float R, float A)   { PackedData[8] = f32tof16(R) | (f32tof16(A) << 16); }
+	float3 UnpackBaseColor()         { return PayloadDecodeLDRColor(PackedData[6]); }
+	float3 UnpackTransparencyColor() { return PayloadDecodeLDRColor(PackedData[7]); }
+	float UnpackMetallic()           { return f16tof32(PackedData[8]); }
+	float UnpackSpecular()           { return f16tof32(PackedData[8] >> 16); }
+	float UnpackRoughness()          { return f16tof32(PackedData[9]); }
+	float UnpackAnisotropy()         { return f16tof32(PackedData[9] >> 16); }
+	float4 UnpackCustomData0()       { return float4(UnpackFloat2FromUInt(PackedData[10]),
+													 UnpackFloat2FromUInt(PackedData[11])); }
+	float4 UnpackCustomData1()       { return float4(UnpackFloat2FromUInt(PackedData[12]),
+													 UnpackFloat2FromUInt(PackedData[13])); }
+
+	void PackBaseColor(float3 C) { PackedData[6] = PayloadEncodeLDRColor(C); }
+	void PackMetallicSpecular(float M, float S) { PackedData[8] = PackFloat2ToUInt(float2(M, S)); }
+	void PackRoughnessAniso(float R, float A)   { PackedData[9] = PackFloat2ToUInt(float2(R, A)); }
 
 #endif
 	
@@ -310,55 +346,43 @@ struct FPackedPathTracingPayload : FMinimalPayload
 	float4 GetDBufferA()
 	{
 		return float4(
-			f16tof32(PackedData[8]),
-			f16tof32(PackedData[8] >> 16),
-			f16tof32(PackedData[9]),
-			f16tof32(PackedData[9] >> 16)
+			UnpackFloat2FromUInt(PackedData[8]),
+			UnpackFloat2FromUInt(PackedData[9])
 		);
 	}
 
 	void SetDBufferA(float4 DBufferA)
 	{
-		PackedData[8] = f32tof16(DBufferA.x);
-		PackedData[8] |= f32tof16(DBufferA.y) << 16;
-		PackedData[9] = f32tof16(DBufferA.z);
-		PackedData[9] |= f32tof16(DBufferA.w) << 16;
+		PackedData[8] = PackFloat2ToUInt(DBufferA.xy);
+		PackedData[9] = PackFloat2ToUInt(DBufferA.zw);
 	}
 
 	float4 GetDBufferB()
 	{
 		return float4(
-			f16tof32(PackedData[10]),
-			f16tof32(PackedData[10] >> 16),
-			f16tof32(PackedData[11]),
-			f16tof32(PackedData[11] >> 16)
+			UnpackFloat2FromUInt(PackedData[10]),
+			UnpackFloat2FromUInt(PackedData[11])
 		);
 	}
 
 	void SetDBufferB(float4 DBufferB)
 	{
-		PackedData[10] = f32tof16(DBufferB.x);
-		PackedData[10] |= f32tof16(DBufferB.y) << 16;
-		PackedData[11] = f32tof16(DBufferB.z);
-		PackedData[11] |= f32tof16(DBufferB.w) << 16;
+		PackedData[10] = PackFloat2ToUInt(DBufferB.xy);
+		PackedData[11] = PackFloat2ToUInt(DBufferB.zw);
 	}
 
 	float4 GetDBufferC()
 	{
 		return float4(
-			f16tof32(PackedData[12]),
-			f16tof32(PackedData[12] >> 16),
-			f16tof32(PackedData[13]),
-			f16tof32(PackedData[13] >> 16)
+			UnpackFloat2FromUInt(PackedData[12]),
+			UnpackFloat2FromUInt(PackedData[13])
 		);
 	}
 
 	void SetDBufferC(float4 DBufferC)
 	{
-		PackedData[12] = f32tof16(DBufferC.x);
-		PackedData[12] |= f32tof16(DBufferC.y) << 16;
-		PackedData[13] = f32tof16(DBufferC.z);
-		PackedData[13] |= f32tof16(DBufferC.w) << 16;
+		PackedData[12] = PackFloat2ToUInt(DBufferC.xy);
+		PackedData[13] = PackFloat2ToUInt(DBufferC.zw);
 	}
 	
 	void SetStochasticSlabRand(float SlabRand)
@@ -544,12 +568,12 @@ struct FPathTracingPayload : FMinimalPayload
 	// SHADINGMODELID_EYE
 	float GetEyeRoughness() { return RoughnessData.x; }
 	void SetEyeRoughness(float R) { RoughnessData.x = R; }
-	float3 GetEyeCausticNormal() { return 2.0 * FuzzColor - 1.0; } // SUBSTRATE_TODO: not optimal since packed as color
+	float3 GetEyeCausticNormal() { return normalize(2.0 * FuzzColor - 1.0); } // SUBSTRATE_TODO: not optimal since packed as color
 	void SetEyeCausticNormal(float3 CausticNormal) { FuzzColor = 0.5 * CausticNormal + 0.5; }
 	float GetEyeIrisMask() { return Anisotropy; }
 	void SetEyeIrisMask(float IrisMask) { Anisotropy = IrisMask; }
-	float3 GetEyeIrisNormal() { return WorldTangent; }
-	void SetEyeIrisNormal(float3 IrisNormal) { WorldTangent = IrisNormal; }
+	float3 GetEyeIrisNormal() { return normalize(2.0 * SpecularEdgeColor - 1.0); }  // SUBSTRATE_TODO: not optimal since packed as color
+	void SetEyeIrisNormal(float3 IrisNormal) { SpecularEdgeColor = IrisNormal * 0.5 + 0.5; }
 
 	// SHADINGMODELID_HAIR
 	float GetHairLongitudinalRoughness() { return RoughnessData.x; }
@@ -607,12 +631,12 @@ struct FPathTracingPayload : FMinimalPayload
 	// SHADINGMODELID_EYE
 	float GetEyeRoughness() { return Roughness; }
 	void SetEyeRoughness(float R) { Roughness = R; }
-	float3 GetEyeCausticNormal() { return CustomData1.xyz; }
-	void SetEyeCausticNormal(float3 CausticNormal) { CustomData1.xyz = CausticNormal; }
+	float3 GetEyeCausticNormal() { return OctahedronToUnitVector(CustomData1.xy); }
+	void SetEyeCausticNormal(float3 CausticNormal) { CustomData1.xy = UnitVectorToOctahedron(CausticNormal); }
 	float GetEyeIrisMask() { return Anisotropy; }
 	void SetEyeIrisMask(float IrisMask) { Anisotropy = IrisMask; }
-	float3 GetEyeIrisNormal() { return WorldTangent; }
-	void SetEyeIrisNormal(float3 IrisNormal) { WorldTangent = IrisNormal; }
+	float3 GetEyeIrisNormal() { return OctahedronToUnitVector(CustomData1.zw); }
+	void SetEyeIrisNormal(float3 IrisNormal) { CustomData1.zw = UnitVectorToOctahedron(IrisNormal); }
 
 	// SHADINGMODELID_HAIR
 	float GetHairLongitudinalRoughness() { return Roughness; }
@@ -761,26 +785,21 @@ FPackedPathTracingPayload PackPathTracingPayload(FPathTracingPayload Input)
 	Output.PackedData[2] = PayloadEncodeUnitVector(Input.WorldSmoothNormal);
 	Output.PackedData[3] = PayloadEncodeUnitVector(Input.WorldNormal);
 	Output.PackedData[4] = PayloadEncodeHDRColor(Input.Radiance);
-	// PackedData[4] is for Radiance (see PackRadiance call below)
+	Output.PackedData[5] = PayloadEncodeTangentAngle(Output.UnpackWorldNormal(), Input.WorldTangent);
+	Output.PackedData[5]|= f32tof16(Input.Ior) << 16;
 #if PATHTRACING_SUBSTRATE_PAYLOAD
-	Output.PackedData[5] = PayloadEncodeLDRColor(Input.DiffuseColor);
-	Output.PackedData[6] = PayloadEncodeLDRColor(Input.SpecularColor);
-	Output.PackedData[7] = PayloadEncodeLDRColor(Input.SpecularEdgeColor);
-	Output.PackedData[8] = PayloadEncodeRoughnessAniso(Input.RoughnessData, Input.Anisotropy);
-	Output.PackedData[9] = PayloadEncodeUnitVector(Input.WorldTangent);
+	Output.PackedData[6] = PayloadEncodeLDRColor(Input.DiffuseColor);
+	Output.PackedData[7] = PayloadEncodeLDRColor(Input.SpecularColor);
+	Output.PackedData[8] = PayloadEncodeLDRColor(Input.SpecularEdgeColor);
+	Output.PackedData[9] = PayloadEncodeRoughnessAniso(Input.RoughnessData, Input.Anisotropy);
 	Output.PackedData[10] = PayloadEncodeLDRColor(Input.TransparencyColor);
-	Output.PackedData[11] = f32tof16(Input.MeanFreePath.x);
-	Output.PackedData[11]|= f32tof16(Input.MeanFreePath.y) << 16;
-	Output.PackedData[12] = f32tof16(Input.MeanFreePath.z);
-	Output.PackedData[12]|= f32tof16(Input.PhaseG) << 16;
+	Output.PackedData[11] = PackFloat2ToUInt(float2(Input.MeanFreePath.xy));
+	Output.PackedData[12] = PackFloat2ToUInt(float2(Input.MeanFreePath.z, Input.PhaseG));
 	Output.PackedData[13] = PayloadEncodeLDRColor(Input.FuzzColor);
-    Output.PackedData[14] = PayloadEncodeLDRColor(float3(0, Input.FuzzAmount, Input.FuzzRoughness));
+    Output.PackedData[14] = PackFloat2ToUInt(float2(Input.FuzzAmount, Input.FuzzRoughness));
 	Output.PackedData[15] = PayloadEncodeLDRColor(Input.WeightV);
 	Output.PackedData[16] = PayloadEncodeLDRColor(Input.TransmittanceN);
-	Output.PackedData[17] = Pack10F(Input.Ior);
-	Output.PackedData[17]|= Pack10F(Input.CoverageAboveAlongN) << 10;
-	Output.PackedData[17]|= Pack10F(Input.SpecularProfileId) << 20;
-	Output.PackedData[17]|= (Input.GlintValue < 1 ? 1u : 0u) << 30; // Bit flag for glint. 
+	Output.PackedData[17] = PackFloat2ToUInt(float2(Input.SpecularProfileId, abs(Input.CoverageAboveAlongN)));
 
 	// In order to not increase the data payload size, we alias SSS/Fuzz parameters when glint are enabled
 	if (Input.GlintValue < 1)
@@ -790,33 +809,20 @@ FPackedPathTracingPayload PackPathTracingPayload(FPathTracingPayload Input)
 		Output.PackedData[12] = PackedGlintsData.y;
 		Output.PackedData[13] = PackFloat2ToUInt(Input.GlintUVdx);
 		Output.PackedData[14] = PackFloat2ToUInt(Input.GlintUVdx);
+		Output.PackedData[17] |= 0x80000000u; // Encode presence of glints in sign bit of CoverageAboveAlongN
 	}
 
 #else // PATHTRACING_SUBSTRATE_PAYLOAD
 
+	Output.PackedData[ 6] = PayloadEncodeLDRColor(Input.BaseColor);
+	Output.PackedData[ 7] = PayloadEncodeLDRColor(Input.TransparencyColor);
+	Output.PackedData[ 8] = PackFloat2ToUInt(float2(Input.Metallic , Input.Specular  ));
+	Output.PackedData[ 9] = PackFloat2ToUInt(float2(Input.Roughness, Input.Anisotropy));
+	Output.PackedData[10] = PackFloat2ToUInt(Input.CustomData0.xy);
+	Output.PackedData[11] = PackFloat2ToUInt(Input.CustomData0.zw);
+	Output.PackedData[12] = PackFloat2ToUInt(Input.CustomData1.xy);
+	Output.PackedData[13] = PackFloat2ToUInt(Input.CustomData1.zw);
 
-	Output.PackedData[5] = PayloadEncodeLDRColor(Input.BaseColor);
-	Output.PackedData[6] = PayloadEncodeLDRColor(Input.TransparencyColor);
-
-	Output.PackedData[7] = f32tof16(Input.Metallic);
-	Output.PackedData[7] |= f32tof16(Input.Specular) << 16;
-	Output.PackedData[8] = f32tof16(Input.Roughness);
-	Output.PackedData[8] |= f32tof16(Input.Anisotropy) << 16;
-
-	Output.PackedData[9] = PayloadEncodeUnitVector(Input.WorldTangent);
-
-	Output.PackedData[10] = f32tof16(Input.CustomData0.x);
-	Output.PackedData[10] |= f32tof16(Input.CustomData0.y) << 16;
-	Output.PackedData[11] = f32tof16(Input.CustomData0.z);
-	Output.PackedData[11] |= f32tof16(Input.CustomData0.w) << 16;
-
-	Output.PackedData[12] = f32tof16(Input.CustomData1.x);
-	Output.PackedData[12] |= f32tof16(Input.CustomData1.y) << 16;
-	Output.PackedData[13] = f32tof16(Input.CustomData1.z);
-	Output.PackedData[13] |= f32tof16(Input.CustomData1.w) << 16;
-
-	Output.PackedData[14] = f32tof16(Input.Ior);
-	// NOTE: 16 bits left in PackedData[14]
 #endif // PATHTRACING_SUBSTRATE_PAYLOAD
 
 #endif // IS_PAYLOAD_ENABLED(RT_PAYLOAD_TYPE_GPULIGHTMASS)
@@ -847,6 +853,8 @@ FPathTracingPayload UnpackPathTracingPayload(FPackedPathTracingPayload Input, FR
 	Output.WorldSmoothNormal = Input.UnpackWorldSmoothNormal();
 	Output.WorldNormal       = Input.UnpackWorldNormal();
 	Output.Radiance          = Input.UnpackRadiance();
+	Output.WorldTangent      = Input.UnpackWorldTangent();
+	Output.Ior               = Input.UnpackIor();
 
 #if PATHTRACING_SUBSTRATE_PAYLOAD
 	Output.DiffuseColor        = Input.UnpackDiffuseColor();
@@ -855,23 +863,21 @@ FPathTracingPayload UnpackPathTracingPayload(FPackedPathTracingPayload Input, FR
 	float4 RoughnessAniso      = Input.UnpackRoughnessAniso();
 	Output.RoughnessData       = RoughnessAniso.xyz;
 	Output.Anisotropy          = RoughnessAniso.w;
-	Output.WorldTangent        = Input.UnpackWorldTangent();
 	Output.TransparencyColor   = Input.UnpackTransparencyColor();
 	Output.MeanFreePath        = Input.UnpackMeanFreePath();
 	Output.PhaseG              = Input.UnpackPhaseG();
 	Output.FuzzColor           = Input.UnpackFuzzColor();
-	const float3 Data14        = Input.UnpackFuzzData();
-	Output.FuzzAmount          = Data14.y;
-	Output.FuzzRoughness       = Data14.z;
+	const float2 FuzzData      = Input.UnpackFuzzData();
+	Output.FuzzAmount          = FuzzData.x;
+	Output.FuzzRoughness       = FuzzData.y;
 	Output.WeightV             = Input.UnpackWeightV();
 	Output.TransmittanceN      = Input.UnpackTransmittanceN();
-	Output.Ior                 = Input.UnpackIor();
 	Output.CoverageAboveAlongN = Input.UnpackCoverageAboveAlongN();
 	Output.SpecularProfileId   = Input.UnpackSpecularProfileId();
 	
 	// TODO: figure out how to make glint unpacking closer to functions above
 	// In order to not increase the data payload size, we alias SSS/Fuzz parameters when glint are enabled
-	const bool bHasGlints = ((Input.PackedData[17] >> 30u) & 0x1u) > 0u;
+	const bool bHasGlints = Input.PackedData[17] >> 31u;
 	if (bHasGlints)
 	{
 		Output.MeanFreePath.x      = 0.f;
@@ -892,17 +898,16 @@ FPathTracingPayload UnpackPathTracingPayload(FPackedPathTracingPayload Input, FR
 		Output.GlintValue = 1.0f;
 	}
 
-#else // PATHTRACING_SUBSTRATE_PAYLOAD
+#else
+	// !PATHTRACING_SUBSTRATE_PAYLOAD
 	Output.BaseColor         = Input.UnpackBaseColor();
 	Output.TransparencyColor = Input.UnpackTransparencyColor();
 	Output.Metallic          = Input.UnpackMetallic();
 	Output.Specular          = Input.UnpackSpecular();
 	Output.Roughness         = Input.UnpackRoughness();
 	Output.Anisotropy        = Input.UnpackAnisotropy();
-	Output.WorldTangent      = Input.UnpackWorldTangent();
 	Output.CustomData0       = Input.UnpackCustomData0();
 	Output.CustomData1       = Input.UnpackCustomData1();
-	Output.Ior               = Input.UnpackIor();
 
 	// 
 	Output.DiffuseColor  = Output.BaseColor - Output.BaseColor * Output.Metallic;
diff --git a/Engine/Shaders/Private/PathTracing/PathTracingMaterialHitShader.usf b/Engine/Shaders/Private/PathTracing/PathTracingMaterialHitShader.usf
--- a/Engine/Shaders/Private/PathTracing/PathTracingMaterialHitShader.usf
+++ b/Engine/Shaders/Private/PathTracing/PathTracingMaterialHitShader.usf
@@ -551,7 +551,10 @@ RAY_TRACING_ENTRY_CLOSEST_HIT(PathTracingMaterialCHS,
 
 #if SUBSTRATE_CLAMPED_CLOSURE_COUNT > 0
 			Payload.WorldNormal = SubstratePixelHeader.SharedLocalBases.Normals[BSDF_GETSHAREDLOCALBASISID(BSDF)];
-			Payload.WorldTangent = SubstratePixelHeader.SharedLocalBases.Tangents[BSDF_GETSHAREDLOCALBASISID(BSDF)];
+			if (SubstrateGetSharedLocalBasisType(SubstratePixelHeader.SharedLocalBases.Types, BSDFIdx) == SUBSTRATE_BASIS_TYPE_TANGENT)
+			{
+				Payload.WorldTangent = SubstratePixelHeader.SharedLocalBases.Tangents[BSDF_GETSHAREDLOCALBASISID(BSDF)];
+			}
 
 			switch (BSDF_GETTYPE(BSDF))
 			{
