Fix bugs with handling of directional lights in LightData.ush
- Initialize IESAtlasIndex to INDEX_NONE in ConvertToDeferredLight(FDirectionalLightData ...)
- Fix incorrect bRadialLight for directional lights in ConvertToDeferredLight(FLightSceneData ...)

diff --git a/Engine/Shaders/Private/LightData.ush b/Engine/Shaders/Private/LightData.ush
--- a/Engine/Shaders/Private/LightData.ush
+++ b/Engine/Shaders/Private/LightData.ush
@@ -485,6 +485,7 @@ FDeferredLightData ConvertToDeferredLight(
 	Out.HairTransmittance	= InitHairTransmittanceData();
 	Out.SourceRadius		= In.DirectionalLightSourceRadius;
 	Out.SoftSourceRadius	= In.DirectionalLightSoftSourceRadius;
+	Out.IESAtlasIndex		= -1;
 	Out.LightFunctionAtlasLightIndex = In.LightFunctionAtlasLightIndex;
 	Out.bAffectsTranslucentLighting = In.bAffectsTranslucentLighting;
 
@@ -522,7 +523,7 @@ FDeferredLightData ConvertToDeferredLight(FLightSceneData In)
 	Out.bInverseSquared			= In.FalloffExponent == 0;
 	Out.SpecularScale			= In.SpecularScale;
 	Out.DiffuseScale			= In.DiffuseScale;
-	Out.bRadialLight			= true;
+	Out.bRadialLight			= LightType != LIGHT_TYPE_DIRECTIONAL;
 	Out.bSpotLight				= LightType == LIGHT_TYPE_SPOT;
 	Out.bRectLight				= LightType == LIGHT_TYPE_RECT;
 	Out.HairTransmittance		= InitHairTransmittanceData();
