Substrate: Modify tangent basis encoding to give more bits to the normal

Change the bit allocation from 10,10,12 to 11,11,10  which gives slightly more accuracy to the normal when encoding a tangent frame. Also, avoid using the tangent symmetry to preserve the input tangent exactly as the extra bit of precision does not seem to be worth while.

#rnx

diff --git a/Engine/Shaders/Private/Substrate/Substrate.ush b/Engine/Shaders/Private/Substrate/Substrate.ush
--- a/Engine/Shaders/Private/Substrate/Substrate.ush
+++ b/Engine/Shaders/Private/Substrate/Substrate.ush
@@ -3675,96 +3675,39 @@ float3 SubstrateUnpackNormal22(uint PackedNormal)
 }
 
 
-#define SUBSTRATE_TANGENT_SYMMETRY 1
+// 11:11:10 encoding of normal + tangent angle
 uint SubstratePackNormalAndTangent(in float3 Normal, in float3 Tangent)
 {
-	uint X10Bits;
-	uint Y10Bits;
-	uint Z11Bits;
-	uint W1Bit;
-
 	// Oct. encoding for the normal
-	{
-		float l1norm = abs(Normal.x) + abs(Normal.y) + abs(Normal.z);
-		float2 result = Normal.xy * (1.0 / l1norm);
-		if (Normal.z < 0.0)
-		{
-			result = (1.0 - abs(result.yx)) * SignNotZero(result.xy);
-		}
-		X10Bits = uint(clamp(result.x * 511.0f + 512.0f, 0.0f, 1023.0f));
-		Y10Bits = uint(clamp(result.y * 511.0f + 512.0f, 0.0f, 1023.0f));
-	}
-
-	// Oct. encoding for the projected tangent, relative to an generated ortho frame based on normal
-	{
-		// Use the quantize normal to regenerate the ortho. frame. This removes discontinuities
-		float x = float(int(X10Bits) - 512) / 511.0f;
-		float y = float(int(Y10Bits) - 512) / 511.0f;
-		Normal = float3(x, y, 1.0f - abs(x) - abs(y));
-		if (Normal.z < 0.0f)
-		{
-			Normal.xy = (1.0f - abs(Normal.yx)) * SignNotZero(Normal.xy);
-		}
-		Normal = normalize(Normal);
-
-		float3x3 LocalRef = GetTangentBasis(Normal);
-		float2 ProjXY = float2(dot(LocalRef[0], Tangent), dot(LocalRef[1], Tangent));
-	#if SUBSTRATE_TANGENT_SYMMETRY
-		// Use symmetry to improve precision
-		ProjXY = ProjXY.x >= 0 ? ProjXY : -ProjXY.xy;
-	#endif
-		float l1norm = abs(ProjXY.x) + abs(ProjXY.y);
-		float2 result = ProjXY.xy * (1.0 / l1norm);
-
-	#if SUBSTRATE_TANGENT_SYMMETRY
-		Z11Bits = uint(clamp(result.x * 2047.0f, 0.0f, 2047.0f));
-	#else
-		Z11Bits = uint(clamp(result.x * 1023.0f + 1024.0f, 0.0f, 2047.0f));
-	#endif
-		W1Bit   = ProjXY.y >= 0 ? 1u : 0u;
-	}
+	float2 OctF = UnitVectorToOctahedron(Normal);
+	uint XBits = uint(round(OctF.x * 1023.0f + 1023.0f));
+	uint YBits = uint(round(OctF.y * 1023.0f + 1023.0f));
 
-	return
-		 (0x3FF & X10Bits)        | 
-		((0x3FF & Y10Bits) << 10) | 
-		((0x7FF & Z11Bits) << 20) |
-		((0x1   & W1Bit  ) << 31);
+	// Oct. encoding for the projected tangent, relative to a generated ortho frame based on decoded normal
+	float3x3 Basis = GetTangentBasis(OctahedronToUnitVector(float2(XBits, YBits) / 1023.0f - 1.0f));
+	float TX = dot(Basis[0], Tangent);
+	float TY = dot(Basis[1], Tangent);
+	float PseudoAngle = abs(TY) / (abs(TX) + abs(TY)); // Coarse ArcTan, returns value in [0,1]
+	uint ABits = uint(floor(PseudoAngle * 255.0f));
+	// Merge all quantized quantities and sign bits of the tangent
+	return XBits | (YBits << 11) | (ABits << 22) | ((asuint(TX) >> 1) & 0x40000000u) | (asuint(TY) & 0x80000000u);
 }
 
 void SubstrateUnpackNormalAndTangent(inout float3 Normal, inout float3 Tangent, in uint InPacked)
 {
-	uint X10Bits = 0x3FF &  InPacked;
-	uint Y10Bits = 0x3FF & (InPacked>>10);
-	uint Z11bits = 0x7FF & (InPacked>>20);
-	uint W1Bit   = 0x1   & (InPacked>>31);
-
-	// Oct. decoding for the normal
-	{
-		float x = float(int(X10Bits) - 512) / 511.0f;
-		float y = float(int(Y10Bits) - 512) / 511.0f;
-
-		Normal = float3(x, y, 1.0f - abs(x) - abs(y));
-		if (Normal.z < 0.0f)
-		{
-			Normal.xy = (1.0f - abs(Normal.yx)) * SignNotZero(Normal.xy);
-		}
-		Normal = normalize(Normal);
-	}
-
-	// Oct. decoding for the projected tangent, relative to an generated ortho frame based on normal
-	{
-		// Create intermediate basis
-		float3x3 LocalRef = GetTangentBasis(Normal);
-	#if SUBSTRATE_TANGENT_SYMMETRY
-		float x = float(Z11bits) / 2047.0f;
-	#else
-		float x = float(int(Z11bits) - 1024) / 1023.0f;
-	#endif
-		float Sign = (W1Bit > 0 ? 1 : -1);
-		float2 ProjX = float2(x, 1.0f - abs(x));
-		Tangent = LocalRef[0] * ProjX.x + LocalRef[1] * ProjX.y * Sign;
-		Tangent = normalize(Tangent);
-	}
+	uint XBits  = (InPacked      ) & 2047;
+	uint YBits  = (InPacked >> 11) & 2047;
+	uint ABits  = (InPacked >> 22) & 255;
+
+	Normal = OctahedronToUnitVector(float2(XBits, YBits) / 1023.0f - 1.0f);
+	float3x3 Basis = GetTangentBasis(Normal);
+	float Y = float(ABits) / 255.0f;
+	float2 P = float2(1.0f - Y, Y);
+	// restore sign bits
+	P.x = asfloat(asuint(P.x) | ((InPacked << 1) & 0x80000000u));
+	P.y = asfloat(asuint(P.y) | ((InPacked     ) & 0x80000000u));
+	P = normalize(P);
+	Tangent = Basis[0] * P.x + Basis[1] * P.y;
 }
 
 void SubstrateLoad_R11G11B10F(FSubstrateMaterialContainer SubstrateBuffer, inout FSubstrateAddressing SubstrateAddressing, inout float3 A)
