Fixed: Precision issues with bloom on Android OpenGL

diff --git a/Engine/Shaders/Private/PostProcessMobile.usf b/Engine/Shaders/Private/PostProcessMobile.usf
--- a/Engine/Shaders/Private/PostProcessMobile.usf
+++ b/Engine/Shaders/Private/PostProcessMobile.usf
@@ -12,9 +12,9 @@
 #define MOBILE_MULTI_VIEW_FALLBACK (INSTANCED_STEREO && MOBILE_MULTI_VIEW)
 
 #if MOBILE_MULTI_VIEW || MOBILE_MULTI_VIEW_FALLBACK
-#define MultiViewTexture2D Texture2DArray<half4>
+#define MultiViewTexture2D Texture2DArray<float4>
 #else
-#define MultiViewTexture2D Texture2D<half4>
+#define MultiViewTexture2D Texture2D<float4>
 #endif
 
 MultiViewTexture2D SceneColorTexture;
