[Local Exposure] Fix thresholds being incorrectly applied in visualization when using Exposure Fusion.

diff --git a/Engine/Shaders/Private/PostProcessVisualizeLocalExposure.usf b/Engine/Shaders/Private/PostProcessVisualizeLocalExposure.usf
--- a/Engine/Shaders/Private/PostProcessVisualizeLocalExposure.usf
+++ b/Engine/Shaders/Private/PostProcessVisualizeLocalExposure.usf
@@ -65,7 +65,7 @@ float4 MainPS(noperspective float4 UVAndScreenPos : TEXCOORD0, float4 SvPosition
 	float LocalExposure = CalculateLocalExposure(LogLuminance + log2(ExposureScaleMiddleGrey.x), BaseLogLum, MiddleGreyLumValue, LocalExposure_HighlightContrastScale, LocalExposure_ShadowContrastScale, LocalExposure_DetailStrength);
 	float DetailLogLum = (LogLuminance + log2(ExposureScaleMiddleGrey.x)) - BaseLogLum;
 #else
-	float BaseLogLum = LogLuminance;
+	float BaseLogLum = LogLuminance + log2(ExposureScaleMiddleGrey.x);
 
 	float2 ExposureFusionUV = ApplyScreenTransform(UV.xy, ColorToExposureFusion);
 	ExposureFusionUV = clamp(ExposureFusionUV, ExposureFusion_UVViewportBilinearMin, ExposureFusion_UVViewportBilinearMax);
