- Add wrapped group to allow for greater than 65535 * 64 verts in a RT dynamic mesh

diff --git a/Engine/Shaders/Private/RayTracing/RayTracingDynamicMesh.usf b/Engine/Shaders/Private/RayTracing/RayTracingDynamicMesh.usf
--- a/Engine/Shaders/Private/RayTracing/RayTracingDynamicMesh.usf
+++ b/Engine/Shaders/Private/RayTracing/RayTracingDynamicMesh.usf
@@ -8,6 +8,7 @@
 #pragma message("UESHADERMETADATA_VERSION 74B3439C-C6C2-47F8-B07C-05009AB185D1")
 
 #include "/Engine/Private/Common.ush"
+#include "/Engine/Private/ComputeShaderUtils.ush"
 #include "/Engine/Generated/Material.ush"
 #include "/Engine/Generated/VertexFactory.ush"
  
@@ -49,18 +50,20 @@ RWBuffer<float> RWVertexPositions;
 int bApplyWorldPositionOffset;
 
 [numthreads(64, 1, 1)]
-void RayTracingDynamicGeometryConverterCS(uint3 DispatchThreadId : SV_DispatchThreadID)
+void RayTracingDynamicGeometryConverterCS(uint3 GroupId : SV_GroupID, int GroupThreadIndex : SV_GroupIndex)
 { 
-	if (DispatchThreadId.x >= NumVertices) return;
+	uint LinearThreadId = GetUnWrappedDispatchThreadId(GroupId, GroupThreadIndex, 64);
+
+	if (LinearThreadId >= NumVertices) return;
 
 	ResolvedView = ResolveView();
 
 	bool bUsingIndirectDraw = UsingIndirectDraw != 0;
 	uint NumActualVertices = bUsingIndirectDraw ? GetNumRayTracingDynamicMeshVerticesIndirect() : NumVertices;
 
-	uint VertexIndex = MinVertexIndex + DispatchThreadId.x;
+	uint VertexIndex = MinVertexIndex + LinearThreadId;
 
-	if (DispatchThreadId.x < NumActualVertices)
+	if (LinearThreadId < NumActualVertices)
 	{
 		FVertexFactoryInput Input = LoadVertexFactoryInputForDynamicUpdate(VertexIndex / 3, VertexIndex % 3, PrimitiveId, InstanceId);
 
diff --git a/Engine/Source/Runtime/Renderer/Private/RayTracing/RayTracingDynamicGeometry.cpp b/Engine/Source/Runtime/Renderer/Private/RayTracing/RayTracingDynamicGeometry.cpp
--- a/Engine/Source/Runtime/Renderer/Private/RayTracing/RayTracingDynamicGeometry.cpp
+++ b/Engine/Source/Runtime/Renderer/Private/RayTracing/RayTracingDynamicGeometry.cpp
@@ -676,7 +676,8 @@ void FRayTracingDynamicGeometryCollection::DispatchUpdates(FRHICommandList& RHIC
 				Cmd.ShaderBindings.SetParameters(BatchedParameters, &ShaderBindingState);
 				RHICmdList.SetBatchedShaderParameters(CurrentShader, BatchedParameters);
 
-				RHICmdList.DispatchComputeShader(FMath::DivideAndRoundUp<uint32>(Cmd.NumMaxVertices, 64), 1, 1);
+				const FIntVector NumWrappedThreadGroups = FComputeShaderUtils::GetGroupCountWrapped(Cmd.NumMaxVertices, 64);
+				RHICmdList.DispatchComputeShader(NumWrappedThreadGroups.X, NumWrappedThreadGroups.Y, NumWrappedThreadGroups.Z);
 			}
 
 			// Make sure buffers are readable again and disable UAV overlap.
