Fix shader compilation error when Substrate is disabled.

diff --git a/Engine/Shaders/Private/Substrate/SubstrateExportCommon.ush b/Engine/Shaders/Private/Substrate/SubstrateExportCommon.ush
--- a/Engine/Shaders/Private/Substrate/SubstrateExportCommon.ush
+++ b/Engine/Shaders/Private/Substrate/SubstrateExportCommon.ush
@@ -2,8 +2,6 @@
 
 #pragma once
 
-#if SUBSTRATE_ENABLED
-
 ///////////////////////////////////////////////////////////////////////////////
 // Material buffer writing
 
@@ -71,6 +69,4 @@ void WriteDataToBuffer(RWTexture2D<uint> OutBuffer, uint InData, uint2 InCoord,
 		}
 	#endif
 	}
-}
-
-#endif // SUBSTRATE_ENABLED
+}
\ No newline at end of file
