First Person: Fix VSM local light marking logic.

diff --git a/Engine/Shaders/Private/VirtualShadowMaps/VirtualShadowMapPageMarking.usf b/Engine/Shaders/Private/VirtualShadowMaps/VirtualShadowMapPageMarking.usf
--- a/Engine/Shaders/Private/VirtualShadowMaps/VirtualShadowMapPageMarking.usf
+++ b/Engine/Shaders/Private/VirtualShadowMaps/VirtualShadowMapPageMarking.usf
@@ -456,7 +456,7 @@ void GeneratePageFlagsFromPixels(
 
 	// Local lights
 	// Don't mark pixels originating from first person if disabled.
-	if (bIsFirstPersonPixel && !bMarkFirstPersonPixels)
+	if (!bIsFirstPersonPixel || bMarkFirstPersonPixels)
 	{
 		const bool bUsePageDilation = PageDilationBorderSizeLocal > 0.0f;
 		const float2 PageDilationOffset = PageDilationBorderSizeLocal * PageDilationDither;
