Fix hair traversal causing incorrect occluion

diff --git a/Engine/Shaders/Private/HairStrands/HairStrandsVoxelPageTraversal.ush b/Engine/Shaders/Private/HairStrands/HairStrandsVoxelPageTraversal.ush
--- a/Engine/Shaders/Private/HairStrands/HairStrandsVoxelPageTraversal.ush
+++ b/Engine/Shaders/Private/HairStrands/HairStrandsVoxelPageTraversal.ush
@@ -185,6 +185,8 @@ FHairTraversalResult ComputeHairCountVirtualVoxel(
 		const float3 O = lerp(TranslatedWorldPosition0, TranslatedWorldPosition1, HitT.xxx);
 		const float3 E = lerp(TranslatedWorldPosition0, TranslatedWorldPosition1, HitT.yyy);
 		const float OELength = min(length(E - O), InSettings.DistanceThreshold);
+		const float3 PagesBoundMax = InNodeDesc.TranslatedWorldMinAABB + InNodeDesc.PageIndexResolution * InNodeDesc.VoxelWorldSize * InCommonDesc.PageResolution;
+
 		#if VOXEL_TRAVERSAL_DEBUG
 		if (InSettings.bDebugEnabled)
 		{
@@ -206,6 +208,7 @@ FHairTraversalResult ComputeHairCountVirtualVoxel(
 		{
 			const float SteppingWorldSize = max(DeltaWorld * StepScale, InSettings.TanConeAngle * StepIt * InNodeDesc.VoxelWorldSize);
 			const float3 HitP = O + StepIt * D + InSettings.JitterScale * RandomStepJitter * SteppingWorldSize * 0.5f;
+			const bool bIsWithinPageBound = all(HitP < PagesBoundMax);
 
 			#if VOXEL_TRAVERSAL_DEBUG
 			AddStepAABB(InSettings.bDebugEnabled, PreviousP, HitP, float3(0, 1, 0));
@@ -228,7 +231,7 @@ FHairTraversalResult ComputeHairCountVirtualVoxel(
 				}
 			}
 
-			if (bIsPageValid)
+			if (bIsPageValid && bIsWithinPageBound)
 			{
 				const uint3 VoxelPageBase   = (PageCoord << InCommonDesc.PageResolutionLog2);
 				const uint3 VoxelPageOffset = VolumeCoord - (PageIndexCoord << InCommonDesc.PageResolutionLog2);
