Fix groom cache (guide & strands) not working correctly when hair triangle strip geometry is enabled.

diff --git a/Engine/Plugins/Runtime/HairStrands/Source/HairStrandsCore/Private/GroomManager.cpp b/Engine/Plugins/Runtime/HairStrands/Source/HairStrandsCore/Private/GroomManager.cpp
--- a/Engine/Plugins/Runtime/HairStrands/Source/HairStrandsCore/Private/GroomManager.cpp
+++ b/Engine/Plugins/Runtime/HairStrands/Source/HairStrandsCore/Private/GroomManager.cpp
@@ -1066,7 +1066,7 @@ static void RunHairStrandsInterpolation_Guide(
 					CacheResources1,
 					RegisterAsSRV(GraphBuilder, InstanceData.Instance->Guides.RestResource->PositionBuffer),
 					RegisterAsSRV(GraphBuilder, InstanceData.Instance->Guides.RestResource->CurveBuffer),
-					RegisterAsSRV(GraphBuilder, InstanceData.Instance->Strands.RestResource->PointToCurveBuffer),
+					RegisterAsSRV(GraphBuilder, InstanceData.Instance->Guides.RestResource->PointToCurveBuffer),
 					RegisterAsSRV(GraphBuilder, InstanceData.Instance->Guides.DeformedResource->GetPositionOffsetBuffer(FHairStrandsDeformedResource::EFrameType::Current)),
 					RegisterAsUAV(GraphBuilder, InstanceData.Instance->Guides.DeformedResource->GetBuffer(FHairStrandsDeformedResource::Current)));
 			}
diff --git a/Engine/Plugins/Runtime/HairStrands/Source/HairStrandsCore/Private/HairStrandsInterpolation.cpp b/Engine/Plugins/Runtime/HairStrands/Source/HairStrandsCore/Private/HairStrandsInterpolation.cpp
--- a/Engine/Plugins/Runtime/HairStrands/Source/HairStrandsCore/Private/HairStrandsInterpolation.cpp
+++ b/Engine/Plugins/Runtime/HairStrands/Source/HairStrandsCore/Private/HairStrandsInterpolation.cpp
@@ -234,6 +234,8 @@ class FGroomCacheUpdatePassCS : public FGlobalShader
 
 	BEGIN_SHADER_PARAMETER_STRUCT(FParameters, )
 		SHADER_PARAMETER(uint32, InstanceRegisteredIndex)
+		SHADER_PARAMETER(uint32, PointCount)
+		SHADER_PARAMETER(uint32, CurveCount)
 		SHADER_PARAMETER(uint32, CachePointCount)
 		SHADER_PARAMETER(uint32, bHasRadiusData)
 		SHADER_PARAMETER(uint32, bHasAddedControlPoint)
@@ -282,10 +284,15 @@ void AddGroomCacheUpdatePass(
 	TRACE_CPUPROFILER_EVENT_SCOPE(AddGroomCacheUpdatePass);
 
 	if (PointCount == 0 || !CacheResources0.PositionBuffer || !CacheResources1.PositionBuffer) return;
+	
+	// Retrive the cache point count from the SRV.
+	const uint32 CachePointCount = CacheResources0.PositionBuffer->GetParent() ? CacheResources0.PositionBuffer->GetParent()->Desc.NumElements / 3u : PointCount;
 
 	FGroomCacheUpdatePassCS::FParameters* Parameters = GraphBuilder.AllocParameters<FGroomCacheUpdatePassCS::FParameters>();
 	Parameters->InstanceRegisteredIndex = InstanceRegisteredIndex;
-	Parameters->CachePointCount = PointCount;
+	Parameters->PointCount = PointCount;
+	Parameters->CurveCount = CurveCount;
+	Parameters->CachePointCount = CachePointCount;
 	Parameters->InPosition0Buffer = CacheResources0.PositionBuffer;
 	Parameters->InPosition1Buffer = CacheResources1.PositionBuffer;
 	Parameters->InRadius0Buffer = CacheResources0.RadiusBuffer;
diff --git a/Engine/Shaders/Private/HairStrands/HairStrandsInterpolation.usf b/Engine/Shaders/Private/HairStrands/HairStrandsInterpolation.usf
--- a/Engine/Shaders/Private/HairStrands/HairStrandsInterpolation.usf
+++ b/Engine/Shaders/Private/HairStrands/HairStrandsInterpolation.usf
@@ -933,56 +933,92 @@ ByteAddressBuffer InRestPointToCurveBuffer;
 StructuredBuffer<float4> InDeformedOffsetBuffer;
 RWByteAddressBuffer OutDeformedBuffer;
 
+float3 ReadCachePosition(uint InCachePointIndex)
+{
+	const uint Index3 = InCachePointIndex * 3;
+	const float3 InPosition0 = float3(
+		InPosition0Buffer[Index3],
+		InPosition0Buffer[Index3+1],
+		InPosition0Buffer[Index3+2]);
+	
+	const float3 InPosition1 = float3(
+		InPosition1Buffer[Index3],
+		InPosition1Buffer[Index3+1],
+		InPosition1Buffer[Index3+2]);
+	return lerp(InPosition0, InPosition1, InterpolationFactor);
+}
+
+float ReadCacheRadius(uint InCachePointIndex)
+{
+	return lerp(InRadius0Buffer[InCachePointIndex], InRadius1Buffer[InCachePointIndex], InterpolationFactor);
+}
+
 [numthreads(GROUP_SIZE, 1, 1)]
 void MainCS(uint DispatchThreadId : SV_DispatchThreadID)
 {
 	const uint RestPointIndex = DispatchThreadId;
-	const uint CurveIndex = ReadHairPointToCurveIndex(InRestPointToCurveBuffer, RestPointIndex);
-	const FHairCurve Curve = ReadHairCurve(InRestCurveBuffer, CurveIndex);
-
-	const bool bIsEndPoint = RestPointIndex == (Curve.PointIndex + Curve.PointCount - 1);
-	if (bHasAddedControlPoint && bIsEndPoint)
+	if (RestPointIndex >= PointCount)
 	{
-		return; // By pass update of the last point
+		return;
 	}
 
-	// Make the assumption that there is no clamped curve data.
-	// One extra point per curve.
-	const uint NumAddedPoints = bHasAddedControlPoint ? CurveIndex : 0;
+	const uint CurveIndex = ReadHairPointToCurveIndex(InRestPointToCurveBuffer, RestPointIndex);
+	const FHairCurve Curve = ReadHairCurve(InRestCurveBuffer, CurveIndex);
+	const float3 DeformedPositionOffset = ReadRenPositionOffset(InDeformedOffsetBuffer, InstanceRegisteredIndex);
 
-	const uint CachePointIndex = RestPointIndex - NumAddedPoints;
-	if (CachePointIndex >= CachePointCount)
+	// Cache data does not contain extra CP at the end of each curve when UseTriangleStrip geometry is used.
+	// This is because the cache data are directly serialized in to the asset, and needs to be independent of this project settings.
+	// To account for this, we remap the rest position to the cache data. The added CP is set to the last point curve point available in the cache. 
+	// This ensures that the guide cache and the strands cache produce valid results
+	uint CachePointIndex = CachePointCount;
+	bool bExtendedPosition = false;
+	if (bHasAddedControlPoint)
 	{
-		return;
+		// Make the assumption that there is no clamped curve data.
+		// One extra point per curve.
+		const bool bIsEndPoint = RestPointIndex == (Curve.PointIndex + Curve.PointCount - 1);
+		if (bIsEndPoint)
+		{
+			const uint NumAddedPoints  = CurveIndex + 1; // 1 CP for each curve + 1 CP for this last point
+			CachePointIndex = RestPointIndex - NumAddedPoints;
+			bExtendedPosition = Curve.PointCount > 1 && CachePointIndex > 0; // If this is the last point, we need to shift it slightly so that normal computation are correct, otherwise the shader of the last visible CP will be incorrect
+		}
+		else
+		{
+			const uint NumAddedPoints = CurveIndex; // 1 CP for each curve
+			CachePointIndex = RestPointIndex - NumAddedPoints;
+		}
 	}
-	
-	FHairControlPoint CP = ReadHairControlPoint(InRestPositionBuffer, RestPointIndex, float3(0, 0, 0)/*Offset*/, 1/*Radius*/, 1/*RootScale*/, 1/*TipScale*/);
-
-	// Position
+	else
 	{
-		const uint Index3 = CachePointIndex * 3;
-		const float3 InPosition0 = float3(
-			InPosition0Buffer[Index3],
-			InPosition0Buffer[Index3+1],
-			InPosition0Buffer[Index3+2]);
-	
-		const float3 InPosition1 = float3(
-			InPosition1Buffer[Index3],
-			InPosition1Buffer[Index3+1],
-			InPosition1Buffer[Index3+2]);
-
-		CP.Position = lerp(InPosition0, InPosition1, InterpolationFactor);
+		CachePointIndex = RestPointIndex;
 	}
 
-	// Radius
-	if (bHasRadiusData == 1)
+	if (CachePointIndex < CachePointCount)
 	{
-		CP.WorldRadius = lerp(InRadius0Buffer[CachePointIndex], InRadius1Buffer[CachePointIndex], InterpolationFactor);
-	}
+		FHairControlPoint CP = ReadHairControlPoint(InRestPositionBuffer, RestPointIndex, float3(0, 0, 0)/*Offset*/, 1/*Radius*/, 1/*RootScale*/, 1/*TipScale*/);
 
-	const float3 DeformedPositionOffset = ReadRenPositionOffset(InDeformedOffsetBuffer, InstanceRegisteredIndex);
+		// Position
+		if (bExtendedPosition)
+		{
+			const int CachePointIndex0 = CachePointIndex-1;
+			const int CachePointIndex1 = CachePointIndex;
+			float3 P0 = ReadCachePosition(CachePointIndex0);
+			float3 P1 = ReadCachePosition(CachePointIndex1);
+			CP.Position = (P1 - P0) * 0.1f + P1;
+		}
+		else
+		{
+			CP.Position = ReadCachePosition(CachePointIndex);
+		}
 
-	WriteHairControlPointPosition(OutDeformedBuffer, RestPointIndex, CP, DeformedPositionOffset, MaxHairRadius);
+		// Radius
+		if (bHasRadiusData == 1)
+		{
+			CP.WorldRadius = ReadCacheRadius(CachePointIndex);
+		}
+		WriteHairControlPointPosition(OutDeformedBuffer, RestPointIndex, CP, DeformedPositionOffset, MaxHairRadius);
+	}
 }
 
 #endif // SHADER_GROOMCACHE_UPDATE
