Path Tracer: Implement Single Layer Water for Substrate

This is a basic implementation for now which should cover the most common cases in a similar way to how the path tracer handles SLW in the non-substrate case.

diff --git a/Engine/Shaders/Private/PathTracing/Material/PathTracingMaterialSampling.ush b/Engine/Shaders/Private/PathTracing/Material/PathTracingMaterialSampling.ush
--- a/Engine/Shaders/Private/PathTracing/Material/PathTracingMaterialSampling.ush
+++ b/Engine/Shaders/Private/PathTracing/Material/PathTracingMaterialSampling.ush
@@ -43,10 +43,13 @@ float3 EstimateMaterialAlbedo(FPathTracingPayload Payload)
 #elif PATHTRACING_SUBSTRATE_PAYLOAD
 		case SHADINGMODELID_SUBSTRATE:
 		{
-			if (Payload.HasRefraction())
-				Albedo = 1.0 + Payload.FuzzAmount * Payload.FuzzColor;
-			else
-				Albedo = Payload.DiffuseColor + Payload.SubsurfaceColor + Payload.SpecularColor + Payload.FuzzAmount * Payload.FuzzColor;
+			Albedo = Payload.DiffuseColor + Payload.SubsurfaceColor + Payload.SpecularColor + Payload.FuzzAmount * Payload.FuzzColor;
+			Albedo *= Payload.GetMaxLobeWeight();
+			break;
+		}
+		case SHADINGMODELID_SOLID_GLASS:
+		{
+			Albedo = 1.0 + Payload.FuzzAmount * Payload.FuzzColor;
 			Albedo *= Payload.GetMaxLobeWeight();
 			break;
 		}
@@ -144,16 +147,18 @@ void AccumulateAlbedoNormal(inout FPathTracingPayload Payload, float3 PathThroug
 			const float Roughness1 = Payload.RoughnessData.y;
 			const float SpecBlend = Payload.RoughnessData.z;
 			const float RoughnessWeight = lerp(ROUGHNESS_WEIGHT(Roughness0), ROUGHNESS_WEIGHT(Roughness1), SpecBlend);
-			if (Payload.HasRefraction())
-			{
-				float GlassAlbedo = 1.0;
-				Albedo += PathThroughput * (GlassAlbedo * RoughnessWeight + Payload.FuzzAmount * Payload.FuzzColor);
-			}
-			else
-			{
-				Albedo += PathThroughput * (Payload.DiffuseColor + Payload.SubsurfaceColor + Payload.SpecularColor * RoughnessWeight + Payload.FuzzAmount * Payload.FuzzColor);
-			}
-			
+			Albedo += PathThroughput * (Payload.DiffuseColor + Payload.SubsurfaceColor + Payload.SpecularColor * RoughnessWeight + Payload.FuzzAmount * Payload.FuzzColor);
+			Normal += NormThroughput * Payload.WorldNormal;
+			break;
+		}
+		case SHADINGMODELID_SOLID_GLASS:
+		{
+			const float Roughness0 = Payload.RoughnessData.x;
+			const float Roughness1 = Payload.RoughnessData.y;
+			const float SpecBlend = Payload.RoughnessData.z;
+			const float RoughnessWeight = lerp(ROUGHNESS_WEIGHT(Roughness0), ROUGHNESS_WEIGHT(Roughness1), SpecBlend);
+			float GlassAlbedo = 1.0;
+			Albedo += PathThroughput * (GlassAlbedo * RoughnessWeight + Payload.FuzzAmount * Payload.FuzzColor);
 			Normal += NormThroughput * Payload.WorldNormal;
 			break;
 		}
@@ -279,7 +284,6 @@ float GetAverageRoughness(FPathTracingPayload Payload)
 	switch (Payload.ShadingModelID)
 	{
 		case SHADINGMODELID_SUBSTRATE:
-		case SHADINGMODELID_THIN_TRANSLUCENT:
 		{
 			AvgRoughness = max(AvgRoughness, Payload.RoughnessData.y);
 			// Only non-glass slabs contain diffuse
@@ -288,6 +292,13 @@ float GetAverageRoughness(FPathTracingPayload Payload)
 			AvgRoughness = lerp(1.0, AvgRoughness, DiffuseFraction);
 			break;
 		}
+		case SHADINGMODELID_SOLID_GLASS:
+		case SHADINGMODELID_THIN_TRANSLUCENT:
+		{
+			// NOTE: these cases do not contain a diffuse lobe yet
+			AvgRoughness = max(AvgRoughness, Payload.RoughnessData.y);
+			break;
+		}
 		case SHADINGMODELID_CLEAR_COAT:
 		{
 			AvgRoughness = max(AvgRoughness, Payload.RoughnessData.y);
@@ -351,13 +362,10 @@ void AdjustPayloadAfterUnpack(inout FPathTracingPayload HitPayload, bool bApplyO
 	{
 		case SHADINGMODELID_SUBSTRATE:
 		{
-			if (!HitPayload.HasRefraction())
-			{
-				// NOTE: if radius is 0 the sss color will get flipped back into diffuse
-				HitPayload.SubsurfaceColor = HitPayload.DiffuseColor;
-				HitPayload.DiffuseColor = 0;
-				AdjustDiffuseSSSContribution(HitPayload.DiffuseColor, HitPayload.SubsurfaceColor, HitPayload.MeanFreePath);
-			}
+			// NOTE: if radius is 0 the sss color will get flipped back into diffuse
+			HitPayload.SubsurfaceColor = HitPayload.DiffuseColor;
+			HitPayload.DiffuseColor = 0;
+			AdjustDiffuseSSSContribution(HitPayload.DiffuseColor, HitPayload.SubsurfaceColor, HitPayload.MeanFreePath);
 			break;
 		}
 		case SHADINGMODELID_TWOSIDED_FOLIAGE:
@@ -581,10 +589,9 @@ FMaterialSample SampleMaterial(
 			return Volumetric_SampleMaterial(V_World, Payload, RandSample);
 #if PATHTRACING_SUBSTRATE_PAYLOAD
 		case SHADINGMODELID_SUBSTRATE:
-			if (Payload.HasRefraction())
-				return SubstrateSolidGlass_SampleMaterial(V_World, Payload, RandSample);
-			else
-				return Substrate_SampleMaterial(V_World, Payload, RandSample);
+			return Substrate_SampleMaterial(V_World, Payload, RandSample);
+		case SHADINGMODELID_SOLID_GLASS:
+			return SubstrateSolidGlass_SampleMaterial(V_World, Payload, RandSample);
 		case SHADINGMODELID_TWOSIDED_FOLIAGE:
 			return SubstrateFoliage_SampleMaterial(V_World, Payload, RandSample);
 		case SHADINGMODELID_THIN_TRANSLUCENT:
@@ -639,10 +646,9 @@ FMaterialEval EvalMaterial(
 			return Volumetric_EvalMaterial(V_World, L_World, Payload, DiffuseSpecularScale);
 #if PATHTRACING_SUBSTRATE_PAYLOAD
 		case SHADINGMODELID_SUBSTRATE:
-			if (Payload.HasRefraction())
-				return SubstrateSolidGlass_EvalMaterial(V_World, L_World, Payload, DiffuseSpecularScale);
-			else
-				return Substrate_EvalMaterial(V_World, L_World, Payload, DiffuseSpecularScale);
+			return Substrate_EvalMaterial(V_World, L_World, Payload, DiffuseSpecularScale);
+		case SHADINGMODELID_SOLID_GLASS:
+			return SubstrateSolidGlass_EvalMaterial(V_World, L_World, Payload, DiffuseSpecularScale);
 		case SHADINGMODELID_TWOSIDED_FOLIAGE:
 			return SubstrateFoliage_EvalMaterial(V_World, L_World, Payload, DiffuseSpecularScale);
 		case SHADINGMODELID_THIN_TRANSLUCENT:
diff --git a/Engine/Shaders/Private/PathTracing/PathTracingCommon.ush b/Engine/Shaders/Private/PathTracing/PathTracingCommon.ush
--- a/Engine/Shaders/Private/PathTracing/PathTracingCommon.ush
+++ b/Engine/Shaders/Private/PathTracing/PathTracingCommon.ush
@@ -585,7 +585,7 @@ struct FPathTracingPayload : FMinimalPayload
 	float2 GetHairPrimitiveUV() { return SpecularEdgeColor.xy; }
 	void SetHairPrimitiveUV(float2 UV) { SpecularEdgeColor.xy = UV; }
 
-	float3 GetExtinction() { return rcp(max(MeanFreePath, 0.0001)); }
+	float3 GetExtinction() { return rcp(max(MeanFreePath, 1e-5)); }
 
 	// SHADINGMODELID_THIN_TRANSLUCENT
 	float3 GetTransmittanceColor() { return exp(-GetExtinction() * SUBSTRATE_SIMPLEVOLUME_THICKNESS_CM); }
@@ -687,20 +687,12 @@ struct FPathTracingPayload : FMinimalPayload
 
 	bool IsMaterialSolidGlass()
 	{
-#if PATHTRACING_SUBSTRATE_PAYLOAD
-		return ShadingModelID == SHADINGMODELID_SUBSTRATE && HasRefraction();
-#else
 		return ShadingModelID == SHADINGMODELID_SOLID_GLASS;
-#endif
 	}
 
 	bool IsMaterialThinGlass()
 	{
-#if PATHTRACING_SUBSTRATE_PAYLOAD
 		return ShadingModelID == SHADINGMODELID_THIN_TRANSLUCENT && HasRefraction();
-#else
-		return ShadingModelID == SHADINGMODELID_THIN_TRANSLUCENT && HasRefraction();
-#endif
 	}
 
 	bool IsMaterialTransmissive()
diff --git a/Engine/Shaders/Private/PathTracing/PathTracingMaterialHitShader.usf b/Engine/Shaders/Private/PathTracing/PathTracingMaterialHitShader.usf
--- a/Engine/Shaders/Private/PathTracing/PathTracingMaterialHitShader.usf
+++ b/Engine/Shaders/Private/PathTracing/PathTracingMaterialHitShader.usf
@@ -685,6 +685,10 @@ RAY_TRACING_ENTRY_CLOSEST_HIT(PathTracingMaterialCHS,
 					// NOTE: only the bottom most layer can refract, coatings should never include a refraction lobe
 					Payload.Ior					= BSDF.bIsBottom ? GetRefractionIor(PixelMaterialInputs, Payload.SpecularColor) : 0.0;
 					Payload.MeanFreePath        = BSDF.bIsBottom ? SLAB_SSSMFP(BSDF) : Payload.MeanFreePath;
+					if (Payload.Ior > 0 && Payload.ShadingModelID == SHADINGMODELID_SUBSTRATE)
+					{
+						Payload.ShadingModelID = SHADINGMODELID_SOLID_GLASS;
+					}
 #else
 					Payload.Ior                 = 0.0;
 #endif
@@ -749,6 +753,40 @@ RAY_TRACING_ENTRY_CLOSEST_HIT(PathTracingMaterialCHS,
 					Payload.ShadingModelID = SHADINGMODELID_EYE;
 					break;
 				}
+				case SUBSTRATE_BSDF_TYPE_SINGLELAYERWATER:
+				{
+					Payload.ShadingModelID = SHADINGMODELID_SOLID_GLASS;
+					float3 BaseColor = SLW_BASECOLOR(BSDF);
+					float Metallic = SLW_METALLIC(BSDF);
+					float Specular = SLW_SPECULAR(BSDF);
+					float Roughness = SLW_ROUGHNESS(BSDF);
+					float3 Extinction = 0.0;
+#if SUBSTRATE_INLINE_SINGLELAYERWATER
+
+					//float Opacity = SLW_TOPMATERIALOPACITY(BSDF); // TODO: Account for this with an extra random decision?
+					//float3 WaterAlbedo = SLW_WATERALBEDO(BSDF); // TODO: Account for scattering in refractive mediums
+					Extinction = SLW_WATEREXTINCTION(BSDF);
+					//float PhaseG = SLW_WATERPHASEG(BSDF);
+					//float3 ColorBehindWater = SLW_COLORSCALEBEHINDWATER(BSDF);
+#else
+					
+#endif
+					Payload.WeightV             = BSDF.LuminanceWeightV;
+					Payload.TransmittanceN      = BSDF.TransmittanceAboveAlongN;
+					Payload.CoverageAboveAlongN = BSDF.CoverageAboveAlongN;
+
+					Payload.DiffuseColor = ComputeDiffuseAlbedo(BaseColor, Metallic);
+					Payload.SpecularColor = ComputeF0(Specular, BaseColor, Metallic);
+					Payload.RoughnessData.x = MakeRoughnessSafe(Roughness, PackedPayload.GetPathRoughness());
+					Payload.RoughnessData.yz = 0;
+					Payload.Ior = DielectricF0ToIor(DielectricSpecularToF0(Specular));
+					Payload.MeanFreePath = rcp(max(Extinction, 1e-5));
+					Payload.FuzzAmount = 0;
+					Payload.FuzzRoughness = 0;
+					Payload.FuzzColor = 0;
+					Payload.GlintValue = 1;
+					break;
+				}
 				default:
 				{
 					// SUBSTRATE_TODO: handle all other cases
@@ -1347,38 +1385,54 @@ RAY_TRACING_ENTRY_ANY_HIT(PathTracingMaterialAHS,
 		float3 RoughnessData = 0.0;
 		float NoV = 0.0f;
 		float3 LocalSigmaT = 0.0f;
+		float Ior = 0.0;
 		SUBSTRATE_UNROLL_N(SUBSTRATE_CLAMPED_CLOSURE_COUNT)
 		for (int BSDFIdx = 0; BSDFIdx < SubstratePixelHeader.SubstrateTree.BSDFCount; ++BSDFIdx)
 		{
 #define BSDF SubstratePixelHeader.SubstrateTree.BSDFs[BSDFIdx]
 			if (BSDF.bIsBottom)
 			{
-				// SUBSTRATE_TODO: Also consider the MFP to see how much this slab is contributing to refraction ....
-				F0 = F0RGBToF0(SubstrateGetBSDFSpecularF0(BSDF));
-				RoughnessData.x = SLAB_ROUGHNESS(BSDF);
-				if (BSDF_GETHASHAZINESS(BSDF))
-				{
-					// NOTE: It should be safe to assume that the clearcoat case cannot happen here as it is only triggered by legacy clearcoat materials which are unlikely to have refraction
-					const FHaziness Haziness = UnpackHaziness(SLAB_HAZINESS(BSDF));
-					RoughnessData.y = Haziness.Roughness;
-					RoughnessData.z = Haziness.Weight;
-				}
-
 				float3 N = SubstratePixelHeader.SharedLocalBases.Normals[BSDF_GETSHAREDLOCALBASISID(BSDF)];
-				LocalSigmaT = rcp(max(SLAB_SSSMFP(BSDF), 0.0001));
 				NoV = dot(N, V_World);
-				break;
+				switch (BSDF_GETTYPE(BSDF))
+				{
+					case SUBSTRATE_BSDF_TYPE_SLAB:
+					{
+						// SUBSTRATE_TODO: Also consider the MFP to see how much this slab is contributing to refraction ....
+						F0 = F0RGBToF0(SubstrateGetBSDFSpecularF0(BSDF));
+						RoughnessData.x = SLAB_ROUGHNESS(BSDF);
+						if (BSDF_GETHASHAZINESS(BSDF))
+						{
+							// NOTE: It should be safe to assume that the clearcoat case cannot happen here as it is only triggered by legacy clearcoat materials which are unlikely to have refraction
+							const FHaziness Haziness = UnpackHaziness(SLAB_HAZINESS(BSDF));
+							RoughnessData.y = Haziness.Roughness;
+							RoughnessData.z = Haziness.Weight;
+						}
+
+						LocalSigmaT = rcp(max(SLAB_SSSMFP(BSDF), 0.0001));
+#if REFRACTION_USE_INDEX_OF_REFRACTION
+						Ior = GetRefractionIor(PixelMaterialInputs, F0);
+#endif
+						break;
+					}
+					case SUBSTRATE_BSDF_TYPE_SINGLELAYERWATER:
+					{
+						RoughnessData.x = SLW_ROUGHNESS(BSDF);
+						RoughnessData.yz = 0;
+#if SUBSTRATE_INLINE_SINGLELAYERWATER
+						LocalSigmaT = SLW_WATEREXTINCTION(BSDF);
+#endif
+						float F0 = DielectricSpecularToF0(SLW_SPECULAR(BSDF));
+						Ior = DielectricF0ToIor(F0);
+						break;
+					}
+				}
+				break; // Found BSDF bottom, exit loop
 			}
 #undef BSDF
 		}
 
 
-#if REFRACTION_USE_INDEX_OF_REFRACTION
-		float Ior = GetRefractionIor(PixelMaterialInputs, F0);
-#else
-		float Ior = 0.0;
-#endif
-
 #if MATERIAL_ISTHINSURFACE
 		// Thin Glass Fake Caustics
 
