Fix shader compilation failure when Mobile HDR and MMV is enabled in project

diff --git a/Engine/Shaders/Private/PostProcessUpscale.usf b/Engine/Shaders/Private/PostProcessUpscale.usf
--- a/Engine/Shaders/Private/PostProcessUpscale.usf
+++ b/Engine/Shaders/Private/PostProcessUpscale.usf
@@ -112,8 +112,16 @@ float4 AccumulateAndApplyWeight(in float4 DataRGBA, in float Weight, inout float
 	return DataRGBA * Weight;
 }
 
-void MainPS(noperspective float4 UVAndScreenPos : TEXCOORD0, float4 SvPosition : SV_POSITION, out float4 OutColor : SV_Target0)
+void MainPS(noperspective float4 UVAndScreenPos : TEXCOORD0, 
+#if MOBILE_MULTI_VIEW && INSTANCED_STEREO
+	in nointerpolation uint EyeIndex : VIEW_ID,
+#endif
+float4 SvPosition : SV_POSITION, out float4 OutColor : SV_Target0)
 {
+#if MOBILE_MULTI_VIEW && INSTANCED_STEREO
+	ResolvedView = ResolveView(EyeIndex);
+#endif
+
 	OutColor = 0;
 
 #if METHOD == UPSCALE_METHOD_NEAREST
