Path Tracer: Fix rare condition that could lead to infinite loop

Exit the transparency tracing loop when the ray becomes invalid

diff --git a/Engine/Shaders/Private/PathTracing/PathTracingCore.ush b/Engine/Shaders/Private/PathTracing/PathTracingCore.ush
--- a/Engine/Shaders/Private/PathTracing/PathTracingCore.ush
+++ b/Engine/Shaders/Private/PathTracing/PathTracingCore.ush
@@ -942,6 +942,10 @@ FPackedPathTracingPayload TraceTransparentRay(inout FPathState PathState, int Bo
 		// prepare next step around the loop
 		// retrace the exact same ray with TMin one ulp past the hit we just found
 		PathState.Ray.TMin = asfloat(asuint(PackedPayload.HitT) + 1);
+		if (!(PathState.Ray.TMin < PathState.Ray.TMax))
+		{
+			break;
+		}
 
 		if (all(PathState.PathThroughput == 0))
 		{
