Path Tracer: Fix rare fireflies from rect lights

Sample points could get generated too close to the plane of the rect light which would then cause strong fireflies. By offseting the intersection volume slightly inwards, we prevent hits directly on the plane of the rectangle.

diff --git a/Engine/Shaders/Private/PathTracing/Light/PathTracingRectLight.ush b/Engine/Shaders/Private/PathTracing/Light/PathTracingRectLight.ush
--- a/Engine/Shaders/Private/PathTracing/Light/PathTracingRectLight.ush
+++ b/Engine/Shaders/Private/PathTracing/Light/PathTracingRectLight.ush
@@ -227,7 +227,8 @@ FVolumeLightSampleSetup RectLight_PrepareLightVolumeSample(
 		// now clip the sphere of influence against the plane of the quad and adjust the near/far distance depending on the sign
 		float3 N = GetNormal(LightId);
 
-		ClipRayByPlane(RayOrigin, RayDirection, N, Center, T.x, T.y);
+		// NOTE: push quad's plane slightly forward to avoid samples that would lie too close to the light
+		ClipRayByPlane(RayOrigin, RayDirection, N, Center + N * 0.001, T.x, T.y);
 
 		const float3 C = GetColor(LightId);
 		const float Width = GetWidth(LightId);
