Fix Y Flipping for pseudovolumetextures on opengles31 and add packed psuedovolumetexture for optimized packed single texture sample for pseudo volumes.

Packed pseudovolumes are atlases that store the next slice in another channel so that only a single sample is performed for trilinear interpolation. This is helpful for volumetric raymarching where it cuts the texture sample count in half.

diff --git a/Engine/Shaders/Private/Common.ush b/Engine/Shaders/Private/Common.ush
--- a/Engine/Shaders/Private/Common.ush
+++ b/Engine/Shaders/Private/Common.ush
@@ -530,11 +530,6 @@ float4 PseudoVolumeTexture(Texture2D Tex, SamplerState TexSampler, float3 inPos,
 	float2 uvCurFrame = uv + curframe;
 	float2 uvNextFrame = uv + nextframe;
 
-	#if COMPILER_GLSL_ES3_1
-	uvCurFrame.y = 1.0 - uvCurFrame.y;
-	uvNextFrame.y = 1.0 - uvNextFrame.y;
-	#endif
-
 	float4 sampleA = 0, sampleB = 0;
 	switch (mipmode)
 	{
@@ -557,6 +552,49 @@ float4 PseudoVolumeTexture(Texture2D Tex, SamplerState TexSampler, float3 inPos,
 	return lerp(sampleA, sampleB, zphase);
 }
 
+// return a pseudovolume texture sample where the input is a 4 channel packed texure that has neighbour slices.
+// useful for simulating 3D texturing with a 2D texture or as a texture flipbook with lerped transitions
+// treats 2d layout of frames a 3d texture and performs bilinear filtering by blending with an offset Z frame.
+// Wrap repeat mode along XY is not seamless. This is however enough for current sampling use cases all in [0,1].
+// @param Tex          = Input Texture Object storing Volume Data
+// @param inPos        = Input float3 for Position, 0-1
+// @param xysize       = Input float for num frames in x,y directions
+// @param numframes    = Input float for num total frames
+// @param mipmode      = Sampling mode: 0 = use miplevel, 1 = use UV computed gradients, 2 = Use gradients (default=0)
+// @param miplevel     = MIP level to use in mipmode=0 (default 0)
+// @param InDDX, InDDY = Texture gradients in mipmode=2
+float2 PseudoVolumeTexture2ChannelPacked(Texture2D Tex, SamplerState TexSampler, float3 inPos, float2 xysize, float numframes,
+	uint mipmode = 0, float miplevel = 0, float2 InDDX = 0, float2 InDDY = 0)
+{
+	float z = inPos.z - 0.5f / numframes;	// This offset is needed to have a behavior consistent with hardware sampling (voxel value is at their center)
+	float zframe = floor(z * numframes);
+	float zphase = frac(z * numframes);
+
+	float2 uv = frac(inPos.xy) / xysize;
+
+	float2 curframe = Tile1Dto2D(xysize.x, zframe) / xysize;
+
+	float2 uvCurFrame = uv + curframe;
+
+	float4 sample = 0;
+	switch (mipmode)
+	{
+	case 0: // Mip level
+		sample = Tex.SampleLevel(TexSampler, uvCurFrame, miplevel);
+		break;
+	case 1: // Gradients automatic from UV
+		sample = Texture2DSample(Tex, TexSampler, uvCurFrame);
+		break;
+	case 2: // Deriviatives provided
+		sample = Tex.SampleGrad(TexSampler, uvCurFrame,  InDDX, InDDY);
+		break;
+	default:
+		break;
+	}
+
+	return lerp(sample.rb, sample.ga, zphase);
+}
+
 #if FEATURE_LEVEL >= FEATURE_LEVEL_SM5 // Cubemap arrays are not supported in SM4 feature level
 MaterialFloat4 TextureCubeArraySample(TextureCubeArray Tex, SamplerState Sampler, float4 UV)
 {
