fix light injection since some platforms do not have '&&' short circuit working

#rnx

diff --git a/Engine/Shaders/Private/LightGridInjection.usf b/Engine/Shaders/Private/LightGridInjection.usf
--- a/Engine/Shaders/Private/LightGridInjection.usf
+++ b/Engine/Shaders/Private/LightGridInjection.usf
@@ -632,7 +632,11 @@ void CullLights_ThreadGroup(uint ThreadIndex, uint GridIndex, FCellBounds CellBo
 
 		bool bIsRectLight;
 		bool bIsTexturedLight;
-		bool bOverlapsLight = LocalLightIndex < NumLocalLights && OverlapsLight(LocalLightIndex, CellBounds, bIsRectLight, bIsTexturedLight);
+		bool bOverlapsLight = LocalLightIndex < NumLocalLights;
+		if (bOverlapsLight)
+		{
+			bOverlapsLight = OverlapsLight(LocalLightIndex, CellBounds, bIsRectLight, bIsTexturedLight);
+		}
 		
 		// NOTE: Cannot be under any divergent branching!
 		uint GroupCount;
@@ -709,7 +713,11 @@ void CullReflectionCaptures_ThreadGroup(uint ThreadIndex, uint GridIndex, FCellB
 		const uint Index = GroupBaseIndex + ThreadIndex;
 		const uint ReflectionCaptureIndex = Index < ReflectionCaptureArrayView.Num ? GetPrimitiveIndex(ReflectionCaptureArrayView, Index) : 0xFFFFFFFF;
 
-		bool bOverlapsReflectionCapture = ReflectionCaptureIndex < NumReflectionCaptures && OverlapsReflectionCapture(ReflectionCaptureIndex, CellBounds);
+		bool bOverlapsReflectionCapture = ReflectionCaptureIndex < NumReflectionCaptures;
+		if (bOverlapsReflectionCapture)
+		{
+			bOverlapsReflectionCapture = OverlapsReflectionCapture(ReflectionCaptureIndex, CellBounds);
+		}
 		
 		// NOTE: Cannot be under any divergent branching!
 		uint GroupCount;
