[Shaders] Fix compile isse in Common.ush for some backend compilers.

Avoid "sample" identifier. Fixed by using UE coding guidlines of pascal case for identifiers.

#rnx

diff --git a/Engine/Shaders/Private/Common.ush b/Engine/Shaders/Private/Common.ush
--- a/Engine/Shaders/Private/Common.ush
+++ b/Engine/Shaders/Private/Common.ush
@@ -515,41 +515,41 @@ float2 Tile1Dto2D(float xsize, float idx)
 // @param mipmode      = Sampling mode: 0 = use miplevel, 1 = use UV computed gradients, 2 = Use gradients (default=0)
 // @param miplevel     = MIP level to use in mipmode=0 (default 0)
 // @param InDDX, InDDY = Texture gradients in mipmode=2
-float4 PseudoVolumeTexture(Texture2D Tex, SamplerState TexSampler, float3 inPos, float2 xysize, float numframes,
-	uint mipmode = 0, float miplevel = 0, float2 InDDX = 0, float2 InDDY = 0)
+float4 PseudoVolumeTexture(Texture2D Tex, SamplerState TexSampler, float3 InPos, float2 XYSize, float NumFrames,
+	uint MipMode = 0, float MipLevel = 0, float2 InDDX = 0, float2 InDDY = 0)
 {
-	float z = inPos.z - 0.5f / numframes;	// This offset is needed to have a behavior consistent with hardware sampling (voxel value is at their center)
-	float zframe = floor(z * numframes);
-	float zphase = frac(z * numframes);
+	float Z = InPos.z - 0.5f / NumFrames;	// This offset is needed to have a behavior consistent with hardware sampling (voxel value is at their center)
+	float ZFrame = floor(Z * NumFrames);
+	float ZPhase = frac(Z * NumFrames);
 
-	float2 uv = frac(inPos.xy) / xysize;
+	float2 UV = frac(InPos.xy) / XYSize;
 
-	float2 curframe = Tile1Dto2D(xysize.x, zframe) / xysize;
-	float2 nextframe = Tile1Dto2D(xysize.x, zframe + 1) / xysize;
+	float2 CurFrame = Tile1Dto2D(XYSize.x, ZFrame) / XYSize;
+	float2 NextFrame = Tile1Dto2D(XYSize.x, ZFrame + 1) / XYSize;
 
-	float2 uvCurFrame = uv + curframe;
-	float2 uvNextFrame = uv + nextframe;
+	float2 UVCurFrame = UV + CurFrame;
+	float2 UVNextFrame = UV + NextFrame;
 
-	float4 sampleA = 0, sampleB = 0;
-	switch (mipmode)
+	float4 SampleA = 0, SampleB = 0;
+	switch (MipMode)
 	{
 	case 0: // Mip level
-		sampleA = Tex.SampleLevel(TexSampler, uvCurFrame, miplevel);
-		sampleB = Tex.SampleLevel(TexSampler, uvNextFrame, miplevel);
+		SampleA = Tex.SampleLevel(TexSampler, UVCurFrame, MipLevel);
+		SampleB = Tex.SampleLevel(TexSampler, UVNextFrame, MipLevel);
 		break;
 	case 1: // Gradients automatic from UV
-		sampleA = Texture2DSample(Tex, TexSampler, uvCurFrame);
-		sampleB = Texture2DSample(Tex, TexSampler, uvNextFrame);
+		SampleA = Texture2DSample(Tex, TexSampler, UVCurFrame);
+		SampleB = Texture2DSample(Tex, TexSampler, UVNextFrame);
 		break;
 	case 2: // Deriviatives provided
-		sampleA = Tex.SampleGrad(TexSampler, uvCurFrame,  InDDX, InDDY);
-		sampleB = Tex.SampleGrad(TexSampler, uvNextFrame, InDDX, InDDY);
+		SampleA = Tex.SampleGrad(TexSampler, UVCurFrame,  InDDX, InDDY);
+		SampleB = Tex.SampleGrad(TexSampler, UVNextFrame, InDDX, InDDY);
 		break;
 	default:
 		break;
 	}
 
-	return lerp(sampleA, sampleB, zphase);
+	return lerp(SampleA, SampleB, ZPhase);
 }
 
 // return a pseudovolume texture sample where the input is a 4 channel packed texure that has neighbour slices.
@@ -563,36 +563,36 @@ float4 PseudoVolumeTexture(Texture2D Tex, SamplerState TexSampler, float3 inPos,
 // @param mipmode      = Sampling mode: 0 = use miplevel, 1 = use UV computed gradients, 2 = Use gradients (default=0)
 // @param miplevel     = MIP level to use in mipmode=0 (default 0)
 // @param InDDX, InDDY = Texture gradients in mipmode=2
-float2 PseudoVolumeTexture2ChannelPacked(Texture2D Tex, SamplerState TexSampler, float3 inPos, float2 xysize, float numframes,
-	uint mipmode = 0, float miplevel = 0, float2 InDDX = 0, float2 InDDY = 0)
+float2 PseudoVolumeTexture2ChannelPacked(Texture2D Tex, SamplerState TexSampler, float3 InPos, float2 XYSize, float NumFrames,
+	uint MipMode = 0, float MipLevel = 0, float2 InDDX = 0, float2 InDDY = 0)
 {
-	float z = inPos.z - 0.5f / numframes;	// This offset is needed to have a behavior consistent with hardware sampling (voxel value is at their center)
-	float zframe = floor(z * numframes);
-	float zphase = frac(z * numframes);
+	float Z = InPos.z - 0.5f / NumFrames;	// This offset is needed to have a behavior consistent with hardware sampling (voxel value is at their center)
+	float ZFrame = floor(Z * NumFrames);
+	float ZPhase = frac(Z * NumFrames);
 
-	float2 uv = frac(inPos.xy) / xysize;
+	float2 UV = frac(InPos.xy) / XYSize;
 
-	float2 curframe = Tile1Dto2D(xysize.x, zframe) / xysize;
+	float2 CurFrame = Tile1Dto2D(XYSize.x, ZFrame) / XYSize;
 
-	float2 uvCurFrame = uv + curframe;
+	float2 UVCurFrame = UV + CurFrame;
 
-	float4 sample = 0;
-	switch (mipmode)
+	float4 Sample = 0;
+	switch (MipMode)
 	{
 	case 0: // Mip level
-		sample = Tex.SampleLevel(TexSampler, uvCurFrame, miplevel);
+		Sample = Tex.SampleLevel(TexSampler, UVCurFrame, MipLevel);
 		break;
 	case 1: // Gradients automatic from UV
-		sample = Texture2DSample(Tex, TexSampler, uvCurFrame);
+		Sample = Texture2DSample(Tex, TexSampler, UVCurFrame);
 		break;
 	case 2: // Deriviatives provided
-		sample = Tex.SampleGrad(TexSampler, uvCurFrame,  InDDX, InDDY);
+		Sample = Tex.SampleGrad(TexSampler, UVCurFrame,  InDDX, InDDY);
 		break;
 	default:
 		break;
 	}
 
-	return lerp(sample.rb, sample.ga, zphase);
+	return lerp(Sample.rb, Sample.ga, ZPhase);
 }
 
 #if FEATURE_LEVEL >= FEATURE_LEVEL_SM5 // Cubemap arrays are not supported in SM4 feature level
