Add missing #pragma once to ByteBuffer.ush

diff --git a/Engine/Shaders/Private/ByteBuffer.ush b/Engine/Shaders/Private/ByteBuffer.ush
--- a/Engine/Shaders/Private/ByteBuffer.ush
+++ b/Engine/Shaders/Private/ByteBuffer.ush
@@ -1,5 +1,7 @@
 // Copyright Epic Games, Inc. All Rights Reserved.
 
+#pragma once
+
 void LoadAndIncrementOffset(inout uint Value, ByteAddressBuffer Buffer, inout uint Offset)
 {
 	Value = Buffer.Load(Offset);
