[GPU Scene Lights] Move some parameters from FLightSceneData to FLightViewData since they're currently modified outside of GPU Scene tracking.
- Color is modified by [REDACTED_PERSON].
- Light function / IES / rect light atlas parameters are not stable and updated during rendering.

diff --git a/Engine/Shaders/Private/LightData.ush b/Engine/Shaders/Private/LightData.ush
--- a/Engine/Shaders/Private/LightData.ush
+++ b/Engine/Shaders/Private/LightData.ush
@@ -509,11 +509,11 @@ FDeferredLightData ConvertToDeferredLight(FLightSceneData In, FLightViewData InV
 {
 	FDeferredLightData Out = (FDeferredLightData)0;
 
-	const uint LightType = UnpackLightType(In.LightTypeAndShadowMapChannelMaskPacked);
+	const uint LightType = UnpackLightType(InViewData.LightTypeAndShadowMapChannelMaskPacked);
 
 	Out.TranslatedWorldPosition = InViewData.TranslatedWorldPosition;
 	Out.InvRadius				= In.InvRadius;
-	Out.Color					= In.Color.xyz * InViewData.ColorScale;
+	Out.Color					= InViewData.Color;
 	Out.FalloffExponent			= In.FalloffExponent;
 	Out.Direction				= In.Direction;
 	Out.SpotAngles				= In.SpotAngles;
@@ -531,14 +531,14 @@ FDeferredLightData ConvertToDeferredLight(FLightSceneData In, FLightViewData InV
 
 	Out.RectLightData.BarnCosAngle = In.RectLightBarnCosAngle;
 	Out.RectLightData.BarnLength = In.RectLightBarnLength;
-	Out.RectLightData.AtlasData.AtlasMaxLevel = In.RectLightAtlasMaxLevel;
-	Out.RectLightData.AtlasData.AtlasUVOffset = In.RectLightAtlasUVOffset;
-	Out.RectLightData.AtlasData.AtlasUVScale = In.RectLightAtlasUVScale;
+	Out.RectLightData.AtlasData.AtlasMaxLevel = InViewData.RectLightAtlasMaxLevel;
+	Out.RectLightData.AtlasData.AtlasUVOffset = InViewData.RectLightAtlasUVOffset;
+	Out.RectLightData.AtlasData.AtlasUVScale = InViewData.RectLightAtlasUVScale;
 
-	Out.IESAtlasIndex					= In.IESAtlasIndex;
-	Out.LightFunctionAtlasLightIndex	= UnpackLightFunctionAtlasIndex(In.LightTypeAndShadowMapChannelMaskPacked);
+	Out.IESAtlasIndex					= InViewData.IESAtlasIndex;
+	Out.LightFunctionAtlasLightIndex	= UnpackLightFunctionAtlasIndex(InViewData.LightTypeAndShadowMapChannelMaskPacked);
 
-	Out.bAffectsTranslucentLighting = UnpackAffectsTranslucentLighting(In.LightTypeAndShadowMapChannelMaskPacked);
+	Out.bAffectsTranslucentLighting = UnpackAffectsTranslucentLighting(InViewData.LightTypeAndShadowMapChannelMaskPacked);
 
 	// Shadow
 	// * Static shadowing uses ShadowMapChannel, 
@@ -547,18 +547,18 @@ FDeferredLightData ConvertToDeferredLight(FLightSceneData In, FLightViewData InV
 	// * ShadowMapChannelMask	= [3:0]
 	// * PreviewShadowMapChannel= [7:4] 
 	// * LightingChannelMask    = [8:10] 
-	Out.ShadowedBits				= (In.LightTypeAndShadowMapChannelMaskPacked & 0xFF) != 0 ? 1 : 0;
-	Out.ShadowMapChannelMask		= UnpackShadowMapChannelMask(In.LightTypeAndShadowMapChannelMaskPacked);
+	Out.ShadowedBits				= (InViewData.LightTypeAndShadowMapChannelMaskPacked & 0xFF) != 0 ? 1 : 0;
+	Out.ShadowMapChannelMask		= UnpackShadowMapChannelMask(InViewData.LightTypeAndShadowMapChannelMaskPacked);
 	return Out;
 }
 
 // Helper function for converting FLightSceneData into FLightShaderParameters
-FLightShaderParameters ConvertToLightShaderParameters(FLightSceneData In)
+FLightShaderParameters ConvertToLightShaderParameters(FLightSceneData In, FLightViewData InViewData)
 {
 	FLightShaderParameters Out;
 	Out.TranslatedWorldPosition	= DFFastAddDemote(In.WorldPosition, PrimaryView.PreViewTranslation);
 	Out.InvRadius				= In.InvRadius;
-	Out.Color					= In.Color.xyz;
+	Out.Color					= InViewData.Color;
 	Out.FalloffExponent			= In.FalloffExponent;
 	Out.Direction				= In.Direction;
 	Out.Tangent					= In.Tangent;
@@ -570,9 +570,9 @@ FLightShaderParameters ConvertToLightShaderParameters(FLightSceneData In)
 	Out.SourceLength			= In.SourceLength;
 	Out.RectLightBarnCosAngle	= In.RectLightBarnCosAngle;
 	Out.RectLightBarnLength		= In.RectLightBarnLength;
-	Out.RectLightAtlasUVOffset	= In.RectLightAtlasUVOffset;
-	Out.RectLightAtlasUVScale	= In.RectLightAtlasUVScale;
-	Out.RectLightAtlasMaxLevel	= In.RectLightAtlasMaxLevel;
+	Out.RectLightAtlasUVOffset	= InViewData.RectLightAtlasUVOffset;
+	Out.RectLightAtlasUVScale	= InViewData.RectLightAtlasUVScale;
+	Out.RectLightAtlasMaxLevel	= InViewData.RectLightAtlasMaxLevel;
 	return Out;
 }
 
diff --git a/Engine/Shaders/Private/SceneData.ush b/Engine/Shaders/Private/SceneData.ush
--- a/Engine/Shaders/Private/SceneData.ush
+++ b/Engine/Shaders/Private/SceneData.ush
@@ -1404,22 +1404,21 @@ RWByteAddressBuffer GPUSceneLightsRW;
 #define GPUSceneLightsBuffer Scene.GPUScene.GPUSceneLightData
 #endif
 
-FLightSceneData GetLightSceneData(int LightId)
+FLightSceneData GetLightSceneData(uint LightId)
 {
 #if 0
 	return GPUSceneLightsBuffer.Load<FLightSceneData>(LightId * sizeof(FLightSceneData));
 #else
 	FLightSceneData Output;
 
-	uint Offset = LightId * sizeof(FLightSceneData);
+	uint Offset = LightId * (uint)sizeof(FLightSceneData);
 
 	LoadAndIncrementOffset(Output.WorldPosition.High, GPUSceneLightsBuffer, Offset);
 	LoadAndIncrementOffset(Output.WorldPosition.Low, GPUSceneLightsBuffer, Offset);
 	LoadAndIncrementOffset(Output.InvRadius, GPUSceneLightsBuffer, Offset);
 	LoadAndIncrementOffset(Output.FalloffExponent, GPUSceneLightsBuffer, Offset);
-	LoadAndIncrementOffset(Output.Color, GPUSceneLightsBuffer, Offset);
 	LoadAndIncrementOffset(Output.Direction, GPUSceneLightsBuffer, Offset);
-	LoadAndIncrementOffset(Output.LightTypeAndShadowMapChannelMaskPacked, GPUSceneLightsBuffer, Offset);
+	LoadAndIncrementOffset(Output.InverseExposureBlend, GPUSceneLightsBuffer, Offset);
 	LoadAndIncrementOffset(Output.Tangent, GPUSceneLightsBuffer, Offset);
 	LoadAndIncrementOffset(Output.SourceRadius, GPUSceneLightsBuffer, Offset);
 	LoadAndIncrementOffset(Output.SpotAngles, GPUSceneLightsBuffer, Offset);
@@ -1427,11 +1426,6 @@ FLightSceneData GetLightSceneData(int LightId)
 	LoadAndIncrementOffset(Output.SourceLength, GPUSceneLightsBuffer, Offset);
 	LoadAndIncrementOffset(Output.RectLightBarnCosAngle, GPUSceneLightsBuffer, Offset);
 	LoadAndIncrementOffset(Output.RectLightBarnLength, GPUSceneLightsBuffer, Offset);
-	LoadAndIncrementOffset(Output.RectLightAtlasUVOffset, GPUSceneLightsBuffer, Offset);
-	LoadAndIncrementOffset(Output.RectLightAtlasUVScale, GPUSceneLightsBuffer, Offset);
-	LoadAndIncrementOffset(Output.RectLightAtlasMaxLevel, GPUSceneLightsBuffer, Offset);
-	LoadAndIncrementOffset(Output.InverseExposureBlend, GPUSceneLightsBuffer, Offset);
-	LoadAndIncrementOffset(Output.IESAtlasIndex, GPUSceneLightsBuffer, Offset);
 	LoadAndIncrementOffset(Output.SpecularScale, GPUSceneLightsBuffer, Offset);
 	LoadAndIncrementOffset(Output.DiffuseScale, GPUSceneLightsBuffer, Offset);
 #endif
diff --git a/Engine/Shaders/Shared/LightSceneData.h b/Engine/Shaders/Shared/LightSceneData.h
--- a/Engine/Shaders/Shared/LightSceneData.h
+++ b/Engine/Shaders/Shared/LightSceneData.h
@@ -27,14 +27,10 @@ struct FLightSceneData
 	// The exponent for the falloff of the light intensity from the distance.
 	float FalloffExponent;
 
-	// Color of the light.
-	float4 Color;
-
 	// Direction of the light if applies.
 	float3 Direction;
 
-	// Extra fields
-	uint LightTypeAndShadowMapChannelMaskPacked;
+	float InverseExposureBlend;
 
 	// One tangent of the light if applies.
 	// Note: BiTangent is on purpose not stored for memory optimisation purposes.
@@ -58,25 +54,11 @@ struct FLightSceneData
 	// Barn door length for rect light
 	float RectLightBarnLength;
 
-	// Rect. light atlas transformation
-	float2 RectLightAtlasUVOffset;
-	float2 RectLightAtlasUVScale;
-	float RectLightAtlasMaxLevel;
-
-	float InverseExposureBlend;
-
-	// could pack IESAtlasIndex with other data in the future since it doesn't require 32 bits
-	// FLocalLightData packs it in 16 bits (see UnpackLigthIESAtlasIndex(...))
-	// could probably go down to 8 bits (with some logic in GIESTextureManager to warn about overflow)
-	float IESAtlasIndex;
-
 	// Factor to applies on the specular.
 	float SpecularScale;
 
 	// Factor to applies on the diffuse.
 	float DiffuseScale;
-
-	float Padding;
 };
 
 #ifdef __cplusplus
diff --git a/Engine/Shaders/Shared/LightViewData.h b/Engine/Shaders/Shared/LightViewData.h
--- a/Engine/Shaders/Shared/LightViewData.h
+++ b/Engine/Shaders/Shared/LightViewData.h
@@ -16,9 +16,22 @@ struct FLightViewData
 {
 	// TODO: pack members to reduce memory overhead/bandwidth
 	float3 TranslatedWorldPosition;
-	float ColorScale;
+	float3 Color;
 	float VolumetricScatteringIntensity;
 	uint VirtualShadowMapId;
+
+	// Extra fields
+	uint LightTypeAndShadowMapChannelMaskPacked;
+	
+	// Rect. light atlas transformation
+	float2 RectLightAtlasUVOffset;
+	float2 RectLightAtlasUVScale;
+	float RectLightAtlasMaxLevel;
+
+	// could pack IESAtlasIndex with other data in the future since it doesn't require 32 bits
+	// FLocalLightData packs it in 16 bits (see UnpackLigthIESAtlasIndex(...))
+	// could probably go down to 8 bits (with some logic in GIESTextureManager to warn about overflow)
+	float IESAtlasIndex;
 };
 
 #ifdef __cplusplus
diff --git a/Engine/Source/Runtime/Renderer/Private/GPUScene.cpp b/Engine/Source/Runtime/Renderer/Private/GPUScene.cpp
--- a/Engine/Source/Runtime/Renderer/Private/GPUScene.cpp
+++ b/Engine/Source/Runtime/Renderer/Private/GPUScene.cpp
@@ -756,7 +756,6 @@ void FGPUScene::UpdateGPULights(FRDGBuilder& GraphBuilder, const UE::Tasks::FTas
 			else
 			{
 				LightData.WorldPosition = FDFVector3{};
-				LightData.Color = FVector3f::ZeroVector;
 				LightData.InvRadius = 0.0f;
 			}
 
@@ -793,7 +792,6 @@ void FGPUScene::InitLightData(const FLightSceneInfoCompact& LightInfoCompact, bo
 	// FLightRenderParameters fields
 	DataOut.WorldPosition = FDFVector3{ LightParams.WorldPosition };
 	DataOut.InvRadius = LightParams.InvRadius;
-	DataOut.Color = LightParams.Color;
 	DataOut.FalloffExponent = LightParams.FalloffExponent;
 	DataOut.Direction = LightParams.Direction;
 	DataOut.SpecularScale = LightParams.SpecularScale;
@@ -805,12 +803,7 @@ void FGPUScene::InitLightData(const FLightSceneInfoCompact& LightInfoCompact, bo
 	DataOut.SourceLength = LightParams.SourceLength;
 	DataOut.RectLightBarnCosAngle = LightParams.RectLightBarnCosAngle;
 	DataOut.RectLightBarnLength = LightParams.RectLightBarnLength;
-	DataOut.RectLightAtlasUVOffset = LightParams.RectLightAtlasUVOffset;
-	DataOut.RectLightAtlasUVScale = LightParams.RectLightAtlasUVScale;
-	DataOut.RectLightAtlasMaxLevel = LightParams.RectLightAtlasMaxLevel;
 	DataOut.InverseExposureBlend = LightParams.InverseExposureBlend;
-	DataOut.IESAtlasIndex = LightParams.IESAtlasIndex;
-	DataOut.LightTypeAndShadowMapChannelMaskPacked = LightInfo.PackLightTypeAndShadowMapChannelMask(bAllowStaticLighting);
 }
 
 void FGPUScene::UpdateInternal(FRDGBuilder& GraphBuilder, FSceneUniformBuffer& SceneUB, FRDGExternalAccessQueue& ExternalAccessQueue, const UE::Tasks::FTask& UpdateTaskPrerequisites, const FUpdateFromComputeCommands& UpdatesFromCompute)
diff --git a/Engine/Source/Runtime/Renderer/Private/LightGridInjection.cpp b/Engine/Source/Runtime/Renderer/Private/LightGridInjection.cpp
--- a/Engine/Source/Runtime/Renderer/Private/LightGridInjection.cpp
+++ b/Engine/Source/Runtime/Renderer/Private/LightGridInjection.cpp
@@ -1281,6 +1281,8 @@ FComputeLightGridOutput FSceneRenderer::PrepareForwardLightData(FRDGBuilder& Gra
 			{
 				SCOPED_NAMED_EVENT(PrepareLightViewData, FColor::Green);
 
+				const bool bAllowStaticLighting = IsStaticLightingAllowed();
+
 				for (int32 SortedIndex = SortedLightSet.SimpleLightsEnd; SortedIndex < SortedLightSet.SortedLights.Num(); ++SortedIndex)
 				{
 					const FSortedLightSceneInfo& SortedLightInfo = SortedLightSet.SortedLights[SortedIndex];
@@ -1304,6 +1306,8 @@ FComputeLightGridOutput FSceneRenderer::PrepareForwardLightData(FRDGBuilder& Gra
 						continue;
 					}
 
+					FLightViewData& CurrentLightViewData = LightViewData[LightSceneId];
+
 					if (LightSceneInfo->ShouldRenderLight(View))
 					{
 						FLightRenderParameters LightParameters;
@@ -1320,17 +1324,27 @@ FComputeLightGridOutput FSceneRenderer::PrepareForwardLightData(FRDGBuilder& Gra
 
 						const int32 VirtualShadowMapId = ViewFamily.EngineShowFlags.DynamicShadows ? VisibleLightInfos[LightSceneId].GetVirtualShadowMapId(&View) : INDEX_NONE;
 
-						LightViewData[LightSceneId].TranslatedWorldPosition = FVector3f(LightParameters.WorldPosition + View.ViewMatrices.GetPreViewTranslation());
-						LightViewData[LightSceneId].ColorScale = LightFade * LightParameters.GetLightExposureScale(Exposure);
-						LightViewData[LightSceneId].VolumetricScatteringIntensity = VolumetricScatteringIntensity;
-						LightViewData[LightSceneId].VirtualShadowMapId = VirtualShadowMapId;
+						CurrentLightViewData.TranslatedWorldPosition = FVector3f(LightParameters.WorldPosition + View.ViewMatrices.GetPreViewTranslation());
+						CurrentLightViewData.Color = FVector3f(LightParameters.Color) * LightFade * LightParameters.GetLightExposureScale(Exposure);
+						CurrentLightViewData.VolumetricScatteringIntensity = VolumetricScatteringIntensity;
+						CurrentLightViewData.VirtualShadowMapId = VirtualShadowMapId;
+						CurrentLightViewData.LightTypeAndShadowMapChannelMaskPacked = LightSceneInfo->PackLightTypeAndShadowMapChannelMask(bAllowStaticLighting);
+						CurrentLightViewData.RectLightAtlasUVOffset = LightParameters.RectLightAtlasUVOffset;
+						CurrentLightViewData.RectLightAtlasUVScale = LightParameters.RectLightAtlasUVScale;
+						CurrentLightViewData.RectLightAtlasMaxLevel = LightParameters.RectLightAtlasMaxLevel;
+						CurrentLightViewData.IESAtlasIndex = LightParameters.IESAtlasIndex;
 					}
 					else
 					{
-						LightViewData[LightSceneId].TranslatedWorldPosition = FVector3f::ZeroVector;
-						LightViewData[LightSceneId].ColorScale = 0.0f;
-						LightViewData[LightSceneId].VolumetricScatteringIntensity = 0.0f;
-						LightViewData[LightSceneId].VirtualShadowMapId = INDEX_NONE;
+						CurrentLightViewData.TranslatedWorldPosition = FVector3f::ZeroVector;
+						CurrentLightViewData.Color = FVector3f::ZeroVector;
+						CurrentLightViewData.VolumetricScatteringIntensity = 0.0f;
+						CurrentLightViewData.VirtualShadowMapId = INDEX_NONE;
+						CurrentLightViewData.LightTypeAndShadowMapChannelMaskPacked = 0;
+						CurrentLightViewData.RectLightAtlasUVOffset = FVector2f::ZeroVector;
+						CurrentLightViewData.RectLightAtlasUVScale = FVector2f::ZeroVector;
+						CurrentLightViewData.RectLightAtlasMaxLevel = FLightRenderParameters::GetRectLightAtlasInvalidMIPLevel();
+						CurrentLightViewData.IESAtlasIndex = INDEX_NONE;
 					}
 				}
 			}, PrerequisiteTask);
