Remove GPUSceneLightsRW code path since it is unused but causes compilation errors.

diff --git a/Engine/Shaders/Private/SceneData.ush b/Engine/Shaders/Private/SceneData.ush
--- a/Engine/Shaders/Private/SceneData.ush
+++ b/Engine/Shaders/Private/SceneData.ush
@@ -1392,14 +1392,10 @@ FSceneDataIntermediates GetSceneDataIntermediates()
 
 #if USE_GLOBAL_GPU_SCENE_DATA
 ByteAddressBuffer GPUSceneLights;
-#elif USE_GPU_SCENE_DATA_RW
-RWByteAddressBuffer GPUSceneLightsRW;
 #endif
 
 #if USE_GLOBAL_GPU_SCENE_DATA
 #define GPUSceneLightsBuffer GPUSceneLights
-#elif USE_GPU_SCENE_DATA_RW
-#define GPUSceneLightsBuffer GPUSceneLightsRW
 #else
 #define GPUSceneLightsBuffer Scene.GPUScene.GPUSceneLightData
 #endif
