Fix shader compilation due to mismatching argument type (half/float) between forward declaration and function definition.

diff --git a/Engine/Shaders/Private/Substrate/SubstrateExport.ush b/Engine/Shaders/Private/Substrate/SubstrateExport.ush
--- a/Engine/Shaders/Private/Substrate/SubstrateExport.ush
+++ b/Engine/Shaders/Private/Substrate/SubstrateExport.ush
@@ -18,11 +18,11 @@
 
 // Forward declaration
 FShadingOcclusion ApplyBentNormal(
-	in float3 CameraVector,
-	in float3 WorldNormal,
-	in float3 WorldBentNormal0,
-	in float Roughness,
-	in float MaterialAO);
+	in half3 CameraVector, 
+	in half3 WorldNormal, 
+	in float3 WorldBentNormal0, 
+	in half Roughness,
+	in half MaterialAO);
 
 void GetPrecomputedIndirectLightingAndSkyLight(
 	FMaterialPixelParameters MaterialParameters,
diff --git a/Engine/Shaders/Private/Substrate/SubstrateForwardLighting.ush b/Engine/Shaders/Private/Substrate/SubstrateForwardLighting.ush
--- a/Engine/Shaders/Private/Substrate/SubstrateForwardLighting.ush
+++ b/Engine/Shaders/Private/Substrate/SubstrateForwardLighting.ush
@@ -24,11 +24,11 @@ void GetVolumeLightingDirectional(float4 AmbientLightingVector, float3 Direction
 float3 GetWorldBentNormalZero(in FMaterialPixelParameters MaterialParameters);
 
 FShadingOcclusion ApplyBentNormal(
-	in float3 CameraVector,
-	in float3 WorldNormal,
+	in half3 CameraVector,
+	in half3 WorldNormal,
 	in float3 WorldBentNormal0,
-	in float Roughness,
-	in float MaterialAO);
+	in half Roughness,
+	in half MaterialAO);
 
 void GetPrecomputedIndirectLightingAndSkyLight(
 	FMaterialPixelParameters MaterialParameters,
