Fix shader compilation error with Substrate due to incorrect function overloading data type for rough refraction.

diff --git a/Engine/Shaders/Private/Substrate/SubstrateExportCommon.ush b/Engine/Shaders/Private/Substrate/SubstrateExportCommon.ush
--- a/Engine/Shaders/Private/Substrate/SubstrateExportCommon.ush
+++ b/Engine/Shaders/Private/Substrate/SubstrateExportCommon.ush
@@ -40,7 +40,7 @@ void WriteDataToBuffer(RWTexture2DArray<uint> OutBuffer, uint InData, uint3 InCo
 	}
 }
 
-void WriteDataToBuffer(RWTexture2D<uint> OutBuffer, uint InData, uint2 InCoord, uint InQuadPixelWriteMask)
+void WriteDataToBuffer(RWTexture2D<float3> OutBuffer, float3 InData, uint2 InCoord, uint InQuadPixelWriteMask)
 {
 	if (InQuadPixelWriteMask > 0u)
 	{
