Path Tracer: Implement diffuse lobe for Thin Glass case

Substrate allows thin, refractive surface to also have a diffuse lobe contribution. The diffuse color is automatically adjusted relative to the mean free path to maintain energy conservation.

This initial implementation only uses a lambertian lobe, but as a follow up the tabulated BSDF from the Substrate code should be investigated.

diff --git a/Engine/Shaders/Private/PathTracing/Material/PathTracingMaterialSampling.ush b/Engine/Shaders/Private/PathTracing/Material/PathTracingMaterialSampling.ush
--- a/Engine/Shaders/Private/PathTracing/Material/PathTracingMaterialSampling.ush
+++ b/Engine/Shaders/Private/PathTracing/Material/PathTracingMaterialSampling.ush
@@ -61,7 +61,9 @@ float3 EstimateMaterialAlbedo(FPathTracingPayload Payload)
 		}
 		case SHADINGMODELID_THIN_TRANSLUCENT:
 		{
-			Albedo = 1.0 + Payload.FuzzAmount * Payload.FuzzColor;
+			const float F0 = F0RGBToF0(Payload.SpecularColor);
+			Albedo = Payload.HasRefraction() ? lerp(Payload.GetTransmittanceColor(), 1.0, F0) : F0;
+			Albedo += Payload.DiffuseColor + Payload.FuzzAmount * Payload.FuzzColor;
 			Albedo *= Payload.GetMaxLobeWeight();
 			break;
 		}
@@ -180,7 +182,7 @@ void AccumulateAlbedoNormal(inout FPathTracingPayload Payload, float3 PathThroug
 			const float Roughness1 = Payload.RoughnessData.y;
 			const float SpecBlend = Payload.RoughnessData.z;
 			const float RoughnessWeight = lerp(ROUGHNESS_WEIGHT(Roughness0), ROUGHNESS_WEIGHT(Roughness1), SpecBlend);
-			Albedo += PathThroughput * (RoughnessWeight * GlassAlbedo + Payload.FuzzAmount * Payload.FuzzColor);
+			Albedo += PathThroughput * (RoughnessWeight * GlassAlbedo + Payload.DiffuseColor + Payload.FuzzAmount * Payload.FuzzColor);
 			Normal += NormThroughput * RoughnessWeight * Payload.WorldNormal;
 
 		}
@@ -293,10 +295,16 @@ float GetAverageRoughness(FPathTracingPayload Payload)
 			break;
 		}
 		case SHADINGMODELID_SOLID_GLASS:
+		{
+			// NOTE: this case does not contain a diffuse lobe yet
+			AvgRoughness = max(AvgRoughness, Payload.RoughnessData.y);
+			break;
+		}
 		case SHADINGMODELID_THIN_TRANSLUCENT:
 		{
-			// NOTE: these cases do not contain a diffuse lobe yet
 			AvgRoughness = max(AvgRoughness, Payload.RoughnessData.y);
+			float DiffuseFraction = LobeSelectionProb(Payload.DiffuseColor + Payload.FuzzAmount * Payload.FuzzColor, Payload.SpecularColor);
+			AvgRoughness = lerp(1.0, AvgRoughness, DiffuseFraction);
 			break;
 		}
 		case SHADINGMODELID_CLEAR_COAT:
diff --git a/Engine/Shaders/Private/PathTracing/Material/PathTracingSubstrateThinGlass.ush b/Engine/Shaders/Private/PathTracing/Material/PathTracingSubstrateThinGlass.ush
--- a/Engine/Shaders/Private/PathTracing/Material/PathTracingSubstrateThinGlass.ush
+++ b/Engine/Shaders/Private/PathTracing/Material/PathTracingSubstrateThinGlass.ush
@@ -20,6 +20,7 @@ struct FSubstrateThinGlassSlabData
 
 	float2 RoughnessT;
 
+	float  DiffuseWeight;
 	float  F0, F90;
 	float  Eta;
 	float3 TransmissionColor;
@@ -31,8 +32,8 @@ struct FSubstrateThinGlassSlabData
 
 	FSubstrateSheenData Fuzz;
 
-	float2 LobeCdf;
-	float3 LobePdf;
+	float3 LobeCdf;
+	float4 LobePdf;
 
 	void PrepareBSDF(float3 V_World, FPathTracingPayload Payload)
 	{
@@ -64,10 +65,13 @@ struct FSubstrateThinGlassSlabData
 		Fuzz.Prepare(V, Payload.FuzzRoughness, Payload.FuzzAmount);
 
 
+
 		const float3 GlassR = Fuzz.Attenuation * lerp(SpecR0.E * NudgeE.x, SpecR1.E * NudgeE.y, Payload.RoughnessData.z);
 		const float3 GlassT = Fuzz.Attenuation * lerp(SpecT0.E * NudgeE.x, SpecT1.E * NudgeE.y, Payload.RoughnessData.z);
 		const float3 FuzzAlbedo = Fuzz.Scale * Payload.FuzzColor;
 
+		DiffuseWeight = Fuzz.Attenuation * (1.0 - Luminance(GlassR));
+
 		SpecR0W = SpecR0.W * NudgeE.x;
 		SpecR1W = SpecR1.W * NudgeE.y;
 		SpecT0W = SpecT0.W * NudgeE.x;
@@ -77,9 +81,10 @@ struct FSubstrateThinGlassSlabData
 		// NOTE: we put T first so that the reflected lobe (R and fuzz) are consecutive in the PDF
 		float3 MaxLobeWeight = Payload.GetMaxLobeWeight();
 		LobeCdf = LobeSelectionCdf(
-			Payload.Ior > 0.0 ? MaxLobeWeight * GlassT : 0.0,
+			MaxLobeWeight * DiffuseWeight * Payload.DiffuseColor,
+			MaxLobeWeight * FuzzAlbedo,
 			MaxLobeWeight * GlassR,
-			MaxLobeWeight * FuzzAlbedo);
+			Payload.Ior > 0.0 ? MaxLobeWeight * GlassT : 0.0);
 		LobePdf = LobeSelectionPdf(LobeCdf);
 	}
 };
@@ -97,29 +102,44 @@ FMaterialSample SubstrateThinGlass_SampleMaterial(
 	float3 L = 0, H = 0, V = Data.V;
 	float OutRoughness = 1;
 
-	const bool bIsReflection = RandSample.x >= Data.LobeCdf.x;
-	const bool bSampleGlass = RandSample.x < Data.LobeCdf.y;
-	if (bSampleGlass)
+	bool bSampleDiffuse = false, bSampleCloth = false, bSampleGlass = false, bIsReflection = RandSample.x < Data.LobeCdf.z;
+	
+	if (RandSample.x < Data.LobeCdf.x)
+	{
+		bSampleDiffuse = true;
+		RandSample.x = RescaleRandomNumber(RandSample.x, 0.0, Data.LobeCdf.x);
+		// diffuse lobe
+		L = CosineSampleHemisphere(RandSample.xy).xyz;
+		H = normalize(L + V);
+	}
+	else if (RandSample.x < Data.LobeCdf.y)
+	{
+		bSampleCloth = true;
+		// cloth lobe
+		RandSample.x = RescaleRandomNumber(RandSample.x, Data.LobeCdf.x, Data.LobeCdf.y);
+		L = Data.Fuzz.Sample(RandSample.xy);
+		H = normalize(L + V);
+	}
+	else
 	{
+		bSampleGlass = true;
 		// Two-phase lobe selection:
 		// First use the overall CDF to pick between transmission and reflection ...
 		RandSample.x = RescaleRandomNumber(RandSample.x,
-			bIsReflection ? Data.LobeCdf.x : 0.0,
-			bIsReflection ? Data.LobeCdf.y : Data.LobeCdf.x);
+										   bIsReflection ? Data.LobeCdf.y : Data.LobeCdf.z,
+										   bIsReflection ? Data.LobeCdf.z : 1.0);
 
 		// Then use the roughness blend to pick one of the two possible roughness values
 		const bool bUseRoughness1 = RandSample.x < Payload.RoughnessData.z;
 		RandSample.x = RescaleRandomNumber(RandSample.x,
-			bUseRoughness1 ? 0.0 : Payload.RoughnessData.z,
-			bUseRoughness1 ? Payload.RoughnessData.z : 1.0);
+										   bUseRoughness1 ? 0.0 : Payload.RoughnessData.z,
+										   bUseRoughness1 ? Payload.RoughnessData.z : 1.0);
 
 		// 
-		float2 Alpha = bIsReflection
-			? (bUseRoughness1 ? Data.AlphaR1 : Data.AlphaR0)
-			: (bUseRoughness1 ? Data.AlphaT1 : Data.AlphaT0);
-		OutRoughness = bIsReflection
-			? (bUseRoughness1 ? Payload.RoughnessData.y : Payload.RoughnessData.x)
-			: (bUseRoughness1 ? Data.RoughnessT.y : Data.RoughnessT.x);
+		float2 Alpha = bIsReflection ? (bUseRoughness1 ? Data.AlphaR1 : Data.AlphaR0) : (bUseRoughness1 ? Data.AlphaT1 : Data.AlphaT0);
+		OutRoughness = bIsReflection ?
+			(bUseRoughness1 ? Payload.RoughnessData.y : Payload.RoughnessData.x) :
+			(bUseRoughness1 ? Data.RoughnessT.y : Data.RoughnessT.x);
 
 		H = ImportanceSampleVisibleGGX(RandSample.xy, Alpha, V).xyz;
 		L = reflect(-V, H);
@@ -128,14 +148,9 @@ FMaterialSample SubstrateThinGlass_SampleMaterial(
 			// invalid output direction, exit early
 			return NullMaterialSample();
 		}
+
 	}
-	else
-	{
-		// cloth lobe
-		RandSample.x = RescaleRandomNumber(RandSample.x, Data.LobeCdf.y, 1.0);
-		L = Data.Fuzz.Sample(RandSample.xy);
-		H = normalize(L + V);
-	}
+
 	// transform to world space
 	const float3 N_World = Payload.WorldNormal;
 	const float NoV = saturate(V.z);
@@ -145,42 +160,50 @@ FMaterialSample SubstrateThinGlass_SampleMaterial(
 
 	const float PositionBiasSign = bIsReflection ? 1.0 : -1.0;
 	const float3 L_World = normalize(mul(float3(L.xy, PositionBiasSign * L.z), Data.Basis)); // flip reflection to other side if we are refracting
-	FMaterialSample Result = CreateMaterialSample(L_World, 0.0, 0.0, PositionBiasSign, OutRoughness, bIsReflection ? PATHTRACER_SCATTER_SPECULAR : PATHTRACER_SCATTER_REFRACT);
-
-
+	FMaterialSample Result = CreateMaterialSample(L_World, 0.0, 0.0, PositionBiasSign, OutRoughness, bSampleDiffuse ? PATHTRACER_SCATTER_DIFFUSE : (bIsReflection ? PATHTRACER_SCATTER_SPECULAR : PATHTRACER_SCATTER_REFRACT));
 
 	// Specular Lobes (either reflected or transmitted)
 	const float2 GGXResult0 = GGXEvalReflection(L, V, H, bIsReflection ? Data.AlphaR0 : Data.AlphaT0);
 	const float2 GGXResult1 = GGXEvalReflection(L, V, H, bIsReflection ? Data.AlphaR1 : Data.AlphaT1);
-	const float LobePdf = bIsReflection ? Data.LobePdf.y : Data.LobePdf.x;
 	const float4 ClothResult = Data.Fuzz.Eval(L, V, H, Payload.FuzzColor);
 	const float ClothPdf = ClothResult.w;
+	const float DiffusePdf = NoL / PI;
 	if (bSampleGlass)
 	{
 		FThinSlabWeights SlabResult = ComputeThinSlabWeights(Data.TransmissionColor, VoH, Payload.Ior, Data.F0);
 		SlabResult.Reflected *= SubstrateSpecularTint(Payload, NoV, NoL, VoH, NoH);
 
 		const float3 SlabTint = bIsReflection ? SlabResult.Reflected : SlabResult.Transmitted;
-		const float SlabW0 = (bIsReflection ? Data.SpecR0W : Data.SpecT0W);
-		const float SlabW1 = (bIsReflection ? Data.SpecR1W : Data.SpecT1W);
-
+		const float SlabW0    = bIsReflection ? Data.SpecR0W : Data.SpecT0W;
+		const float SlabW1    = bIsReflection ? Data.SpecR1W : Data.SpecT1W;
+		const float LobePdf   = bIsReflection ? Data.LobePdf.z : Data.LobePdf.w;
 
 		Result.AddLobeWithMIS(Data.Fuzz.Attenuation * (1 - Payload.RoughnessData.z) * GGXResult0.x * SlabW0 * SlabTint, GGXResult0.y, LobePdf * (1 - Payload.RoughnessData.z));
 		Result.AddLobeWithMIS(Data.Fuzz.Attenuation * (    Payload.RoughnessData.z) * GGXResult1.x * SlabW1 * SlabTint, GGXResult1.y, LobePdf * (    Payload.RoughnessData.z));
 
 		if (bIsReflection)
 		{
-			Result.Pdf += Data.LobePdf.z * ClothPdf;
+			Result.Pdf += Data.LobePdf.x * DiffusePdf;
+			Result.Pdf += Data.LobePdf.y * ClothPdf;
 		}
 	}
+	else if (bSampleDiffuse)
+	{
+		// Diffuse Lobe
+		Result.AddLobeWithMIS(Data.DiffuseWeight * Payload.DiffuseColor, DiffusePdf, Data.LobePdf.x);
+		Result.Pdf += Data.LobePdf.y * ClothPdf;
+		Result.Pdf += Data.LobePdf.z * (1.0 - Payload.RoughnessData.z) * GGXResult0.y;
+		Result.Pdf += Data.LobePdf.z * (      Payload.RoughnessData.z) * GGXResult1.y;
+	}
 	else
 	{
 		// Cloth Lobe
 		const float ShadowTerminator = ShadowTerminatorTerm(L_World, N_World, Payload.WorldSmoothNormal);
-		Result.AddLobeWithMIS(ClothResult.xyz * ShadowTerminator, ClothPdf, Data.LobePdf.z);
+		Result.AddLobeWithMIS(ClothResult.xyz * ShadowTerminator, ClothPdf, Data.LobePdf.y);
 
-		Result.Pdf += LobePdf * (1.0 - Payload.RoughnessData.z) * GGXResult0.y;
-		Result.Pdf += LobePdf * (      Payload.RoughnessData.z) * GGXResult1.y;
+		Result.Pdf += Data.LobePdf.x * DiffusePdf;
+		Result.Pdf += Data.LobePdf.z * (1.0 - Payload.RoughnessData.z) * GGXResult0.y;
+		Result.Pdf += Data.LobePdf.z * (      Payload.RoughnessData.z) * GGXResult1.y;
 	}
 
 
@@ -222,6 +245,9 @@ FMaterialEval SubstrateThinGlass_EvalMaterial(
 	const float NoH = saturate(H.z);
 
 
+	// Diffuse Lobe
+	const float DiffusePdf = NoL / PI;
+
 	// Specular Lobes (either reflected or transmitted)
 	const float2 GGXResult0 = GGXEvalReflection(L, V, H, bIsReflection ? Data.AlphaR0 : Data.AlphaT0);
 	const float2 GGXResult1 = GGXEvalReflection(L, V, H, bIsReflection ? Data.AlphaR1 : Data.AlphaT1);
@@ -232,7 +258,7 @@ FMaterialEval SubstrateThinGlass_EvalMaterial(
 	const float SlabW0 = (bIsReflection ? Data.SpecR0W : Data.SpecT0W);
 	const float SlabW1 = (bIsReflection ? Data.SpecR1W : Data.SpecT1W);
 
-	const float LobePdf = bIsReflection ? Data.LobePdf.y : Data.LobePdf.x;
+	const float LobePdf = bIsReflection ? Data.LobePdf.z : Data.LobePdf.w;
 
 	FMaterialEval Result = NullMaterialEval();
 
@@ -242,12 +268,16 @@ FMaterialEval SubstrateThinGlass_EvalMaterial(
 
 	if (bIsReflection)
 	{
+		{
+			// Diffuse Lobe
+			Result.AddLobeWithMIS(DiffuseSpecularScale.x * Data.DiffuseWeight * Payload.DiffuseColor, DiffusePdf, Data.LobePdf.x);
+		}
 		{
 			// Cloth Lobe
 			const float ShadowTerminator = ShadowTerminatorTerm(L_World, N_World, Payload.WorldSmoothNormal);
 			const float4 ClothResult = Data.Fuzz.Eval(L, V, H, Payload.FuzzColor);
 			const float ClothPdf = ClothResult.w;
-			Result.AddLobeWithMIS(DiffuseSpecularScale.y * ClothResult.xyz * ShadowTerminator, ClothPdf, Data.LobePdf.z);
+			Result.AddLobeWithMIS(DiffuseSpecularScale.y * ClothResult.xyz * ShadowTerminator, ClothPdf, Data.LobePdf.y);
 		}
 		Result.Weight *= SubstrateLobeWeight(Payload, NoL);
 	}
