[Shaders] Use InBuffer as parameter instead of Buffer to avoid cryptic syntax error in FXC.

#rnx

diff --git a/Engine/Shaders/Private/ByteBuffer.ush b/Engine/Shaders/Private/ByteBuffer.ush
--- a/Engine/Shaders/Private/ByteBuffer.ush
+++ b/Engine/Shaders/Private/ByteBuffer.ush
@@ -2,54 +2,54 @@
 
 #pragma once
 
-void LoadAndIncrementOffset(inout uint Value, ByteAddressBuffer Buffer, inout uint Offset)
+void LoadAndIncrementOffset(inout uint Value, ByteAddressBuffer InBuffer, inout uint Offset)
 {
-	Value = Buffer.Load(Offset);
+	Value = InBuffer.Load(Offset);
 	Offset += 4;
 }
 
-void LoadAndIncrementOffset(inout uint2 Value, ByteAddressBuffer Buffer, inout uint Offset)
+void LoadAndIncrementOffset(inout uint2 Value, ByteAddressBuffer InBuffer, inout uint Offset)
 {
-	Value = Buffer.Load2(Offset);
+	Value = InBuffer.Load2(Offset);
 	Offset += 8;
 }
 
-void LoadAndIncrementOffset(inout uint3 Value, ByteAddressBuffer Buffer, inout uint Offset)
+void LoadAndIncrementOffset(inout uint3 Value, ByteAddressBuffer InBuffer, inout uint Offset)
 {
-	Value = Buffer.Load3(Offset);
+	Value = InBuffer.Load3(Offset);
 	Offset += 12;
 }
 
-void LoadAndIncrementOffset(inout uint4 Value, ByteAddressBuffer Buffer, inout uint Offset)
+void LoadAndIncrementOffset(inout uint4 Value, ByteAddressBuffer InBuffer, inout uint Offset)
 {
-	Value = Buffer.Load4(Offset);
+	Value = InBuffer.Load4(Offset);
 	Offset += 16;
 }
 
-void LoadAndIncrementOffset(inout float Value, ByteAddressBuffer Buffer, inout uint Offset)
+void LoadAndIncrementOffset(inout float Value, ByteAddressBuffer InBuffer, inout uint Offset)
 {
 	uint UintValue;
-	LoadAndIncrementOffset(UintValue, Buffer, Offset);
+	LoadAndIncrementOffset(UintValue, InBuffer, Offset);
 	Value = asfloat(UintValue);
 }
 
-void LoadAndIncrementOffset(inout float2 Value, ByteAddressBuffer Buffer, inout uint Offset)
+void LoadAndIncrementOffset(inout float2 Value, ByteAddressBuffer InBuffer, inout uint Offset)
 {
 	uint2 UintValue;
-	LoadAndIncrementOffset(UintValue, Buffer, Offset);
+	LoadAndIncrementOffset(UintValue, InBuffer, Offset);
 	Value = asfloat(UintValue);
 }
 
-void LoadAndIncrementOffset(inout float3 Value, ByteAddressBuffer Buffer, inout uint Offset)
+void LoadAndIncrementOffset(inout float3 Value, ByteAddressBuffer InBuffer, inout uint Offset)
 {
 	uint3 UintValue;
-	LoadAndIncrementOffset(UintValue, Buffer, Offset);
+	LoadAndIncrementOffset(UintValue, InBuffer, Offset);
 	Value = asfloat(UintValue);
 }
 
-void LoadAndIncrementOffset(inout float4 Value, ByteAddressBuffer Buffer, inout uint Offset)
+void LoadAndIncrementOffset(inout float4 Value, ByteAddressBuffer InBuffer, inout uint Offset)
 {
 	uint4 UintValue;
-	LoadAndIncrementOffset(UintValue, Buffer, Offset);
+	LoadAndIncrementOffset(UintValue, InBuffer, Offset);
 	Value = asfloat(UintValue);
 }
