Add option for translucent materials "Velocity From Depth Only".
This can be set when a translucent material is set to "Output Depth and Velocity". It indicates that we want to calculate the velocity based on the written depth only and not consider the primitive transforms or material world position offset.
It is only intended for special cases such as writing depth when raycasting a translucent volume.

diff --git a/Engine/Shaders/Private/VelocityShader.usf b/Engine/Shaders/Private/VelocityShader.usf
--- a/Engine/Shaders/Private/VelocityShader.usf
+++ b/Engine/Shaders/Private/VelocityShader.usf
@@ -18,10 +18,12 @@ static const float GDepthBias = 0.001f;
 
 struct FVelocityInterpsVSToPS
 {
+#if !TRANSLUCENCY_VELOCITY_FROM_DEPTH
 	// float4(ScreenPos.xyzw)
 	float4 PackedVelocityA : TEXCOORD6;
 	// float4(PrevScreenPos.xyzw)
 	float4 PackedVelocityC : TEXCOORD7;
+#endif
 	FVertexFactoryInterpolantsVSToPS FactoryInterpolants;
 #if USE_WORLD_POSITION_EXCLUDING_SHADER_OFFSETS && !IS_NANITE_PASS
 	float3 PixelPositionExcludingWPO : TEXCOORD8;
@@ -109,8 +111,10 @@ void MainVertexShader(
 		OutGlobalClipPlaneDistance = dot(ResolvedView.GlobalClippingPlane, float4(TranslatedWorldPosition.xyz, 1));
 #endif
 
+#if !TRANSLUCENCY_VELOCITY_FROM_DEPTH
 		Output.Interps.PackedVelocityA = INVARIANT(ScreenPos);
 		Output.Interps.PackedVelocityC = INVARIANT(PrevScreenPosObj);
+#endif
 	}
 
 // Move all geometry a little bit towards the camera (to not get z fighting with existing zbuffer of the same meshes rendered with slightly different float computations).
@@ -201,8 +205,13 @@ void MainPixelShader(
 	GetMaterialClippingVelocity(MaterialParameters, PixelMaterialInputs, bThinTranslucentMaterial, MaterialOpacity);
 #endif
 
+#if !TRANSLUCENCY_VELOCITY_FROM_DEPTH
 	// 3d velocity, includes camera an object motion
 	float3 Velocity = Calculate3DVelocity(Input.PackedVelocityA, Input.PackedVelocityC);
 
 	OutColor = EncodeVelocityToTexture(Velocity, (GetPrimitiveData(MaterialParameters).Flags & PRIMITIVE_SCENE_DATA_FLAG_HAS_PIXEL_ANIMATION) != 0);
+#else
+	// We only want velocity generated from camera movement and depth so write zero which is the special clear color value.
+	OutColor = 0;
+#endif
 }
diff --git a/Engine/Source/Runtime/Engine/Private/Materials/Material.cpp b/Engine/Source/Runtime/Engine/Private/Materials/Material.cpp
--- a/Engine/Source/Runtime/Engine/Private/Materials/Material.cpp
+++ b/Engine/Source/Runtime/Engine/Private/Materials/Material.cpp
@@ -5020,7 +5020,8 @@ bool UMaterial::CanEditChange(const FProperty* InProperty) const
 			|| PropertyName == GET_MEMBER_NAME_STRING_CHECKED(UMaterial, bDisableDepthTest)
 			|| PropertyName == GET_MEMBER_NAME_STRING_CHECKED(UMaterial, bUseTranslucencyVertexFog)
 			|| PropertyName == GET_MEMBER_NAME_STRING_CHECKED(UMaterial, bComputeFogPerPixel)
-			|| PropertyName == GET_MEMBER_NAME_STRING_CHECKED(UMaterial, bOutputTranslucentVelocity))
+			|| PropertyName == GET_MEMBER_NAME_STRING_CHECKED(UMaterial, bOutputTranslucentVelocity)
+			|| PropertyName == GET_MEMBER_NAME_STRING_CHECKED(UMaterial, bIsTranslucencyVelocityFromDepth))
 		{
 			return MaterialDomain != MD_DeferredDecal && IsTranslucentBlendMode(BlendMode);
 		}
@@ -7198,6 +7199,11 @@ bool UMaterial::IsTranslucencyWritingVelocity() const
 	return bOutputTranslucentVelocity && IsTranslucentBlendMode(GetBlendMode());
 }
 
+bool UMaterial::IsTranslucencyVelocityFromDepth() const
+{
+	return bIsTranslucencyVelocityFromDepth;
+}
+
 bool UMaterial::IsTranslucencyWritingFrontLayerTransparency() const
 {
 	return IsTranslucentBlendMode(GetBlendMode())
diff --git a/Engine/Source/Runtime/Engine/Private/Materials/MaterialInstance.cpp b/Engine/Source/Runtime/Engine/Private/Materials/MaterialInstance.cpp
--- a/Engine/Source/Runtime/Engine/Private/Materials/MaterialInstance.cpp
+++ b/Engine/Source/Runtime/Engine/Private/Materials/MaterialInstance.cpp
@@ -4855,6 +4855,11 @@ bool UMaterialInstance::IsTranslucencyWritingVelocity() const
 	return bOutputTranslucentVelocity && IsTranslucentBlendMode(GetBlendMode());
 }
 
+bool UMaterialInstance::IsTranslucencyVelocityFromDepth() const
+{
+	return Parent ? Parent->IsTranslucencyVelocityFromDepth() : false;
+}
+
 bool UMaterialInstance::IsDitheredLODTransition() const
 {
 	return DitheredLODTransition;
diff --git a/Engine/Source/Runtime/Engine/Private/Materials/MaterialInterface.cpp b/Engine/Source/Runtime/Engine/Private/Materials/MaterialInterface.cpp
--- a/Engine/Source/Runtime/Engine/Private/Materials/MaterialInterface.cpp
+++ b/Engine/Source/Runtime/Engine/Private/Materials/MaterialInterface.cpp
@@ -1389,6 +1389,11 @@ bool UMaterialInterface::IsTranslucencyWritingVelocity() const
 	return false;
 }
 
+bool UMaterialInterface::IsTranslucencyVelocityFromDepth() const
+{
+	return false;
+}
+
 bool UMaterialInterface::IsTranslucencyWritingFrontLayerTransparency() const
 {
 	return false;
diff --git a/Engine/Source/Runtime/Engine/Private/Materials/MaterialShared.cpp b/Engine/Source/Runtime/Engine/Private/Materials/MaterialShared.cpp
--- a/Engine/Source/Runtime/Engine/Private/Materials/MaterialShared.cpp
+++ b/Engine/Source/Runtime/Engine/Private/Materials/MaterialShared.cpp
@@ -2098,6 +2098,11 @@ bool FMaterialResource::IsTranslucencyWritingVelocity() const
 	return MaterialInstance ? MaterialInstance->IsTranslucencyWritingVelocity() : Material->IsTranslucencyWritingVelocity();
 }
 
+bool FMaterialResource::IsTranslucencyVelocityFromDepth() const
+{
+	return MaterialInstance ? MaterialInstance->IsTranslucencyVelocityFromDepth() : Material->IsTranslucencyVelocityFromDepth();
+}
+
 bool FMaterialResource::IsTranslucencyWritingFrontLayerTransparency() const
 {
 	// We cannot call UMaterial::IsTranslucencyWritingFrontLayerTransparency because we need to check the instance potentially overriden blend mode.
@@ -2759,6 +2764,7 @@ void FMaterial::SetupMaterialEnvironment(
 	SET_SHADER_DEFINE(OutEnvironment, MATERIAL_OUTPUT_OPACITY_AS_ALPHA, GetBlendableOutputAlpha());
 	SET_SHADER_DEFINE(OutEnvironment, TRANSLUCENT_SHADOW_WITH_MASKED_OPACITY, GetCastDynamicShadowAsMasked());
 	SET_SHADER_DEFINE(OutEnvironment, TRANSLUCENT_WRITING_VELOCITY, IsTranslucencyWritingVelocity());
+	SET_SHADER_DEFINE(OutEnvironment, TRANSLUCENCY_VELOCITY_FROM_DEPTH, IsTranslucencyWritingVelocity() && IsTranslucencyVelocityFromDepth());
 	SET_SHADER_DEFINE(OutEnvironment, TRANSLUCENT_WRITING_FRONT_LAYER_TRANSPARENCY, IsTranslucencyWritingFrontLayerTransparency());
 	SET_SHADER_DEFINE(OutEnvironment, MATERIAL_USE_ALPHA_TO_COVERAGE, IsUsingAlphaToCoverage() ? 1 : 0);
 	SET_SHADER_DEFINE(OutEnvironment, MOBILE_HIGH_QUALITY_BRDF, IsMobileHighQualityBRDFEnabled());
diff --git a/Engine/Source/Runtime/Engine/Public/MaterialShared.h b/Engine/Source/Runtime/Engine/Public/MaterialShared.h
--- a/Engine/Source/Runtime/Engine/Public/MaterialShared.h
+++ b/Engine/Source/Runtime/Engine/Public/MaterialShared.h
@@ -2126,6 +2126,7 @@ public:
 	virtual bool IsDitheredLODTransition() const = 0;
 	virtual bool IsTranslucencyWritingCustomDepth() const { return false; }
 	virtual bool IsTranslucencyWritingVelocity() const { return false; }
+	virtual bool IsTranslucencyVelocityFromDepth() const{ return false; }
 	virtual bool IsTranslucencyWritingFrontLayerTransparency() const { return false; }
 	virtual bool IsTangentSpaceNormal() const { return false; }
 	virtual bool ShouldGenerateSphericalParticleNormals() const { return false; }
@@ -2912,6 +2913,7 @@ public:
 	ENGINE_API virtual bool IsDitheredLODTransition() const override;
 	ENGINE_API virtual bool IsTranslucencyWritingCustomDepth() const override;
 	ENGINE_API virtual bool IsTranslucencyWritingVelocity() const override;
+	ENGINE_API virtual bool IsTranslucencyVelocityFromDepth() const override;
 	ENGINE_API virtual bool IsTranslucencyWritingFrontLayerTransparency() const override;
 	ENGINE_API virtual bool IsTangentSpaceNormal() const override;
 	ENGINE_API virtual bool ShouldGenerateSphericalParticleNormals() const override;
diff --git a/Engine/Source/Runtime/Engine/Public/Materials/Material.h b/Engine/Source/Runtime/Engine/Public/Materials/Material.h
--- a/Engine/Source/Runtime/Engine/Public/Materials/Material.h
+++ b/Engine/Source/Runtime/Engine/Public/Materials/Material.h
@@ -988,6 +988,13 @@ public:
 	UPROPERTY(EditAnywhere, Category = Translucency, meta = (DisplayName = "Output Depth and Velocity"))
 	uint8 bOutputTranslucentVelocity : 1;
 
+	/** 
+	 * When true, translucent materials that write velocity will calculate it only based on the written depth and camera movement. 
+	 * This can be useful for translucent volume raycasts which write pixel depth offset and want to calculate velocity based on the output depth.
+	 */
+	UPROPERTY(EditAnywhere, Category = Translucency, meta = (DisplayName = "Velocity From Depth Only", editcondition = "bOutputTranslucentVelocity"))
+	uint8 bIsTranslucencyVelocityFromDepth : 1;
+
 	/** If true the compilation environment will be changed to remove the global COMPILE_SHADERS_FOR_DEVELOPMENT flag. */
 	UPROPERTY(transient, duplicatetransient)
 	uint8 bAllowDevelopmentShaderCompile : 1;
@@ -1208,6 +1215,7 @@ public:
 	ENGINE_API virtual bool IsDitheredLODTransition() const override;
 	ENGINE_API virtual bool IsTranslucencyWritingCustomDepth() const override;
 	ENGINE_API virtual bool IsTranslucencyWritingVelocity() const override;
+	ENGINE_API virtual bool IsTranslucencyVelocityFromDepth() const override;
 	ENGINE_API virtual bool IsTranslucencyWritingFrontLayerTransparency() const override;
 	ENGINE_API virtual bool IsMasked() const override;
 	ENGINE_API virtual bool IsDeferredDecal() const override;
diff --git a/Engine/Source/Runtime/Engine/Public/Materials/MaterialInstance.h b/Engine/Source/Runtime/Engine/Public/Materials/MaterialInstance.h
--- a/Engine/Source/Runtime/Engine/Public/Materials/MaterialInstance.h
+++ b/Engine/Source/Runtime/Engine/Public/Materials/MaterialInstance.h
@@ -885,6 +885,7 @@ public:
 	ENGINE_API virtual bool IsTwoSided() const override;
 	ENGINE_API virtual bool IsThinSurface() const override;
 	ENGINE_API virtual bool IsTranslucencyWritingVelocity() const override;
+	ENGINE_API virtual bool IsTranslucencyVelocityFromDepth() const override;
 	ENGINE_API virtual bool IsDitheredLODTransition() const override;
 	ENGINE_API virtual bool IsMasked() const override;
 	ENGINE_API virtual bool WritesToRuntimeVirtualTexture() const override;
diff --git a/Engine/Source/Runtime/Engine/Public/Materials/MaterialInterface.h b/Engine/Source/Runtime/Engine/Public/Materials/MaterialInterface.h
--- a/Engine/Source/Runtime/Engine/Public/Materials/MaterialInterface.h
+++ b/Engine/Source/Runtime/Engine/Public/Materials/MaterialInterface.h
@@ -940,6 +940,7 @@ public:
 	ENGINE_API virtual bool IsDitheredLODTransition() const;
 	ENGINE_API virtual bool IsTranslucencyWritingCustomDepth() const;
 	ENGINE_API virtual bool IsTranslucencyWritingVelocity() const;
+	ENGINE_API virtual bool IsTranslucencyVelocityFromDepth() const;
 	ENGINE_API virtual bool IsTranslucencyWritingFrontLayerTransparency() const;
 	ENGINE_API virtual bool IsMasked() const;
 	ENGINE_API virtual bool IsDeferredDecal() const;
