Substrate - fixed PathTracer cloth transmittance to lower lobes to match Substrate rester behavior.

#rnx
#tests PC dx12 Deferred and PathTracer

diff --git a/Engine/Shaders/Private/PathTracing/Material/PathTracingSubstrateCommon.ush b/Engine/Shaders/Private/PathTracing/Material/PathTracingSubstrateCommon.ush
--- a/Engine/Shaders/Private/PathTracing/Material/PathTracingSubstrateCommon.ush
+++ b/Engine/Shaders/Private/PathTracing/Material/PathTracingSubstrateCommon.ush
@@ -87,7 +87,9 @@ struct FSubstrateSheenData
 		aInv = LTC.x;
 		bInv = LTC.y;
 		Scale = FuzzAmount * LTC.z;
-		Attenuation = 1 - FuzzAmount;
+
+		FBxDFEnergyTermsA EnergyTerms = ComputeClothEnergyTermsA(FuzzRoughness, V.z);
+		Attenuation = lerp(1.0f, ComputeEnergyPreservation(EnergyTerms), FuzzAmount);
 #else
 		Roughness = MakeRoughnessSafe(FuzzRoughness, SUBSTRATE_MIN_FUZZ_ROUGHNESS);
 		FBxDFEnergyTermsA Fuzz = ComputeClothEnergyTermsA(Roughness, V.z);
