Substrate - Use the LTC directional albedo instead of ComputeClothEnergyTermsA where the LTC is used.

#rnx
#tests PC dx12 Deferred and PathTracer

diff --git a/Engine/Shaders/Private/PathTracing/Material/PathTracingSubstrateCommon.ush b/Engine/Shaders/Private/PathTracing/Material/PathTracingSubstrateCommon.ush
--- a/Engine/Shaders/Private/PathTracing/Material/PathTracingSubstrateCommon.ush
+++ b/Engine/Shaders/Private/PathTracing/Material/PathTracingSubstrateCommon.ush
@@ -86,9 +86,11 @@ struct FSubstrateSheenData
 		float4 LTC = SheenLTC_Cofficients(V.z, FuzzRoughness, View.SheenLTCTexture, View.SheenLTCSampler);
 		aInv = LTC.x;
 		bInv = LTC.y;
-		Scale = FuzzAmount * LTC.z;
+		const float DirectionalAlbedo = LTC.z;
+		Scale = FuzzAmount * DirectionalAlbedo;
 
-		FBxDFEnergyTermsA EnergyTerms = ComputeClothEnergyTermsA(FuzzRoughness, V.z);
+		FBxDFEnergyTermsA EnergyTerms = (FBxDFEnergyTermsA)1.0f;
+		EnergyTerms.E = DirectionalAlbedo;
 		Attenuation = lerp(1.0f, ComputeEnergyPreservation(EnergyTerms), FuzzAmount);
 #else
 		Roughness = MakeRoughnessSafe(FuzzRoughness, SUBSTRATE_MIN_FUZZ_ROUGHNESS);
