Substrate - fix for anisotropic lighting pop in Lumen reflection.

#rnx
#tests PC - project aurora - dx12 with ray tracing

diff --git a/Engine/Shaders/Private/Substrate/SubstrateEvaluation.ush b/Engine/Shaders/Private/Substrate/SubstrateEvaluation.ush
--- a/Engine/Shaders/Private/Substrate/SubstrateEvaluation.ush
+++ b/Engine/Shaders/Private/Substrate/SubstrateEvaluation.ush
@@ -1366,23 +1366,19 @@ FSubstrateImportanceSampleResult SubstrateImportanceSampleBSDF(FSubstrateBSDFCon
 			// Sample Specular lobe
 			if (Out.Term == SHADING_TERM_SPECULAR)
 			{
-				float2 Alpha = 0;
+				float2 Alpha = Pow2(SpecularRoughness);
 				if (bHasAnisotropy)
 				{
-					GetAnisotropicRoughness(SpecularRoughness, SLAB_ANISOTROPY(BSDFContext.BSDF), Alpha.x, Alpha.y);
+					GetAnisotropicRoughness(Alpha.x, SLAB_ANISOTROPY(BSDFContext.BSDF), Alpha.x, Alpha.y);
 				}
-				else
-				{
-					Alpha = Pow2(SpecularRoughness);
-				}
-				float4 TangentH = ImportanceSampleVisibleGGX(E, Alpha, BSDFContext.TangentV);
-
-				const float HPDF = TangentH.w;
-				const float3 H = mul(TangentH.xyz, BSDFContext.TangentBasis);
-				const float VoH = saturate(dot(BSDFContext.V, H));
 
-				Out.L   = 2 * dot(BSDFContext.V, H) * H - BSDFContext.V;
-				Out.PDF = RayPDFToReflectionRayPDF(VoH, HPDF);
+				float4 GGXSample	= ImportanceSampleVisibleGGX(E, Alpha, BSDFContext.TangentV);
+				const float HPDF	= GGXSample.w;
+				const float3 H		= mul(GGXSample.xyz, BSDFContext.TangentBasis);
+				const float VoH		= saturate(dot(BSDFContext.V, H));
+	
+				Out.L				= 2 * dot(BSDFContext.V, H) * H - BSDFContext.V;
+				Out.PDF				= RayPDFToReflectionRayPDF(VoH, HPDF);
 			}
 
 			// Evaluate lobes
