Fix working color space support in blackbody material node.

diff --git a/Engine/Shaders/Private/MaterialTemplate.ush b/Engine/Shaders/Private/MaterialTemplate.ush
--- a/Engine/Shaders/Private/MaterialTemplate.ush
+++ b/Engine/Shaders/Private/MaterialTemplate.ush
@@ -31,6 +31,7 @@
 #include "/Engine/Private/SparseVolumeTexture/SparseVolumeTextureCommon.ush"
 #include "/Engine/Shared/EnvironmentComponentsFlags.h"
 #include "/Engine/Private/MeshPaintTextureCommon.ush"
+#include "/Engine/Private/ColorSpace.ush"
 
 // Reduce latency of Nanite material base pass pixel shaders caused by register pressure due to the calculation of WPO
 #if PIXELSHADER && IS_NANITE_SHADING_PASS && IS_BASE_PASS
@@ -2045,15 +2046,8 @@ float3 MaterialExpressionBlackBody( float Temp )
 	float X = Y/y * x;
 	float Z = Y/y * z;
 
-	float3x3 XYZtoRGB =
-	{
-		 3.2404542, -1.5371385, -0.4985314,
-		-0.9692660,  1.8760108,  0.0415560,
-		 0.0556434, -0.2040259,  1.0572252,
-	};
-
 	// The XYZtoRGB transform can result in negative values, so we need to clamp here.
-	return max( 0.0f, mul( XYZtoRGB, float3( X, Y, Z ) ) ) * pow( 0.0004 * Temp, 4 );
+	return max( 0.0f, XYZ_2_LinearRGB( float3( X, Y, Z ) ) ) * pow( 0.0004 * Temp, 4 );
 }
 
 float2 MaterialExpressionGetHairRootUV(FMaterialPixelParameters Parameters)
