Tidy up virtual texture feedback opacity code used to determine final opacity for stochastic feedback.
Split out the code that determines _whether_ to apply stochastic opacity, and _how to modify_ opacity for various blend modes.
Added new modificaton which forces base Opacity (before dbuffer modification) to 1 for opaque surfaces.
This is only needed because it is possible to create some exotic layered materials where the material compiler maintains an alpha value even for opaque. That normally doesn't matter because we don't actually use the alpha blend, but it is a problem for virtual texture feedback where a low alpha value will prevent feedback from being written.
Also moved code that chooses a fixed dbuffer opacity for materials that manually sample dbuffer. This now lives with the standard dbuffer opacity calculation, and allows for a tidier implementation of GetFinalVirtualTextureFeedbackOpacity(). It also means we test the decal receiver flag before setting the fixed dbuffer opacity, which allows us to ignore contribution for some cases.
#tests Rendering EngineTests pass

diff --git a/Engine/Shaders/Private/BasePassPixelShader.usf b/Engine/Shaders/Private/BasePassPixelShader.usf
--- a/Engine/Shaders/Private/BasePassPixelShader.usf
+++ b/Engine/Shaders/Private/BasePassPixelShader.usf
@@ -1016,9 +1016,9 @@ void FPixelShaderInOut_MainPS(
 	ShadingModel = MATERIAL_SHADINGMODEL_SINGLELAYERWATER;
 #endif
 
-	float DBufferOpacity = 1.0f;
-#if USE_DBUFFER && MATERIALDECALRESPONSEMASK && !MATERIALBLENDING_ANY_TRANSLUCENT && !MATERIAL_SHADINGMODEL_SINGLELAYERWATER
-	// apply decals from the DBuffer
+	// When applying DBuffer we need to modify the virtual texture feedback opacity so that we don't always overwrite any feedback already written in the DBuffer pass.
+	float DBufferVTFeedbackFactor = 1.0f;
+#if USE_DBUFFER && !MATERIALBLENDING_ANY_TRANSLUCENT && !MATERIAL_SHADINGMODEL_SINGLELAYERWATER
 	#if PC_D3D
 	//Temporary workaround to avoid crashes on AMD, revert back to BRANCH
 	FLATTEN
@@ -1027,6 +1027,8 @@ void FPixelShaderInOut_MainPS(
 	#endif
 	if ((GetPrimitiveData(MaterialParameters).Flags & PRIMITIVE_SCENE_DATA_FLAG_DECAL_RECEIVER) != 0 && View.ShowDecalsMask > 0)
 	{
+	#if MATERIALDECALRESPONSEMASK
+		// Apply decals from the DBuffer.
 		uint ValidDBufferTargetMask = GetDBufferTargetMask(uint2(In.SvPosition.xy)) & MATERIALDECALRESPONSEMASK;
 
 		BRANCH
@@ -1037,14 +1039,19 @@ void FPixelShaderInOut_MainPS(
 			#if SUBSTRATE_INLINE_SHADING && !SUBSTRATE_OPTIMIZED_UNLIT
 			const FSubstrateDBuffer SubstrateBufferData = GetDBufferData(BufferUV, ValidDBufferTargetMask);
 			ApplyDBufferData(SubstrateBufferData, SubstratePixelHeader, SubstrateData);
-			DBufferOpacity = SubstrateBufferData.Coverage;
+			DBufferVTFeedbackFactor = SubstrateBufferData.Coverage;
 			#endif
 		#else
 			FDBufferData DBufferData = GetDBufferData(BufferUV, ValidDBufferTargetMask);
 			ApplyDBufferData(DBufferData, MaterialParameters.WorldNormal, SubsurfaceColor, Roughness, BaseColor, Metallic, Specular);
-			DBufferOpacity = (DBufferData.ColorOpacity + DBufferData.NormalOpacity + DBufferData.RoughnessOpacity) * (1.0f / 3.0f);
+			DBufferVTFeedbackFactor = (DBufferData.ColorOpacity + DBufferData.NormalOpacity + DBufferData.RoughnessOpacity) * (1.0f / 3.0f);
 		#endif // SUBSTRATE_ENABLED
 		}
+	#elif MATERIAL_DECAL_READ_MASK
+		// Material is reading DBuffer directly.
+		// We don't know the opacity that will be used so use a fixed value instead.
+		DBufferVTFeedbackFactor = 0.5f;
+	#endif
 	}
 #endif
 
@@ -2230,7 +2237,7 @@ void FPixelShaderInOut_MainPS(
 		FinalizeVirtualTextureFeedback(
 			MaterialParameters.VirtualTextureFeedback,
 			MaterialParameters.SvPosition,
-			Opacity * DBufferOpacity,
+			Opacity, DBufferVTFeedbackFactor,
 			View.FrameNumber,
 			View.VTFeedbackBuffer
 		);
diff --git a/Engine/Shaders/Private/MobileBasePassPixelShader.usf b/Engine/Shaders/Private/MobileBasePassPixelShader.usf
--- a/Engine/Shaders/Private/MobileBasePassPixelShader.usf
+++ b/Engine/Shaders/Private/MobileBasePassPixelShader.usf
@@ -473,12 +473,14 @@ void Main(
 	float SubsurfaceProfile = 0;
 #endif
 
-	float DBufferOpacity = 1.0f;
-#if USE_DBUFFER && MATERIALDECALRESPONSEMASK && !MATERIALBLENDING_ANY_TRANSLUCENT && !MATERIAL_SHADINGMODEL_SINGLELAYERWATER
-	// apply decals from the DBuffer
+	// When applying DBuffer we need to modify the virtual texture feedback opacity so that we don't always overwrite any feedback already written in the DBuffer pass.
+	float DBufferVTFeedbackFactor = 1.0f;
+#if USE_DBUFFER && !MATERIALBLENDING_ANY_TRANSLUCENT && !MATERIAL_SHADINGMODEL_SINGLELAYERWATER
 	BRANCH
 	if ((GetPrimitiveData(MaterialParameters).Flags & PRIMITIVE_SCENE_DATA_FLAG_DECAL_RECEIVER) != 0 && View.ShowDecalsMask > 0)
 	{
+	#if MATERIALDECALRESPONSEMASK
+		// Apply decals from the DBuffer.
 		uint ValidDBufferTargetMask = GetDBufferTargetMask(uint2(SvPosition.xy)) & MATERIALDECALRESPONSEMASK;
 
 		BRANCH
@@ -495,16 +497,21 @@ void Main(
 			#if SUBSTRATE_INLINE_SHADING && !SUBSTRATE_OPTIMIZED_UNLIT
 			const FSubstrateDBuffer SubstrateBufferData = GetDBufferData(BufferUV, ValidDBufferTargetMask, EyeIndex);
 			ApplyDBufferData(SubstrateBufferData, SubstratePixelHeader, SubstrateData);
-			DBufferOpacity = SubstrateBufferData.Coverage;
+			DBufferVTFeedbackFactor = SubstrateBufferData.Coverage;
 			#endif
 		#else
 			FDBufferData DBufferData = GetDBufferData(BufferUV, ValidDBufferTargetMask, EyeIndex);
 			ApplyDBufferData(DBufferData, MaterialParameters.WorldNormal, SubsurfaceColor, Roughness, BaseColor, Metallic, Specular);
-			DBufferOpacity = (DBufferData.ColorOpacity + DBufferData.NormalOpacity + DBufferData.RoughnessOpacity) * (1.0f / 3.0f);
+			DBufferVTFeedbackFactor = (DBufferData.ColorOpacity + DBufferData.NormalOpacity + DBufferData.RoughnessOpacity) * (1.0f / 3.0f);
 		#endif // SUBSTRATE_ENABLED
 			
 			MaterialParameters.ReflectionVector = ReflectionAboutCustomWorldNormal(MaterialParameters, MaterialParameters.WorldNormal, false);
 		}
+	#elif MATERIAL_DECAL_READ_MASK
+		// Material is reading DBuffer directly.
+		// We don't know the opacity that will be used so use a fixed value instead.
+		DBufferVTFeedbackFactor = 0.5f;
+	#endif
 	}
 #endif
 
@@ -1163,7 +1170,7 @@ void Main(
 	FinalizeVirtualTextureFeedback(
 		MaterialParameters.VirtualTextureFeedback,
 		MaterialParameters.SvPosition,
-		Opacity * DBufferOpacity,
+		Opacity, DBufferVTFeedbackFactor,
 		View.FrameNumber,
 		View.VTFeedbackBuffer
 	);
diff --git a/Engine/Shaders/Private/VirtualTextureCommon.ush b/Engine/Shaders/Private/VirtualTextureCommon.ush
--- a/Engine/Shaders/Private/VirtualTextureCommon.ush
+++ b/Engine/Shaders/Private/VirtualTextureCommon.ush
@@ -54,8 +54,57 @@ void StoreVirtualTextureFeedback(in out FVirtualTextureFeedbackParams Params, ui
 #endif
 }
 
+/** Get if we need stochastic opacity for writing virtual texture feedback. Ideally this test is purely compile time. */
+bool UseVirtualTextureFeedbackOpacity()
+{
+#if MATERIALBLENDING_TRANSLUCENT || \
+	MATERIALBLENDING_ADDITIVE || \
+	MATERIALBLENDING_MODULATE || \
+	MATERIAL_SHADINGMODEL_THIN_TRANSLUCENT || \
+	MATERIAL_SHADINGMODEL_SINGLELAYERWATER || \
+	IS_DECAL || \
+	(USE_DBUFFER && ((MATERIAL_DECAL_READ_MASK != 0) || (MATERIALDECALRESPONSEMASK !=0)))
+	return true;
+#else
+	return false;
+#endif
+}
+
+/** Get the final virtual texture feedback opacity with modifications for special material blend cases. */
+float GetFinalVirtualTextureFeedbackOpacity(float InOpacity, float InDBufferOpacityFactor)
+{
+	float Opacity = InOpacity;
+
+	// Set small ClampMin to always allow some feedback.
+	float ClampMin = 0.02f;
+	float ClampMax = 1.f;
+
+#if (MATERIALBLENDING_ADDITIVE || MATERIALBLENDING_MODULATE || MATERIAL_SHADINGMODEL_THIN_TRANSLUCENT || MATERIAL_SHADINGMODEL_SINGLELAYERWATER)
+	// When doing additive and modulation blending the provide Opacity may not make sense to use as feedback opacity.
+	// Also we cannot know here how many layers that we are drawing.
+	// No value will is guaranteed to give a fair result so we just choose something!
+	Opacity = 0.25f;
+#endif
+
+#if MATERIALBLENDING_SOLID
+	// Fix a layered material compilation case where solid materials can pass in an Opacity value that isn't 1.
+	Opacity = 1.f;
+#endif
+
+#if (IS_DECAL && !IS_DBUFFER_DECAL)
+	// Fully opaque gbuffer decals may only write to subset of gbuffer channels, so we clamp to avoid full overwrite of previous feedback in the base pass.
+	ClampMax = 0.75f;
+#endif
+
+	// Dbuffer contribution is fully opaque when InDBufferOpacityFactor is 0.
+	// But fully opaque dbuffer decals may only feed a subset of material attributes, so we clamp to prevent dbuffer fully blocking feedback from this pass.
+	float DBufferOpacityFactor = max(InDBufferOpacityFactor, 0.25f);
+
+	return clamp(Opacity * DBufferOpacityFactor, ClampMin, ClampMax);
+}
+
 /** This should be called at the end of the pixel shader to write out the gathered VT feedback info to the OutputBuffer. */
-void FinalizeVirtualTextureFeedback(in FVirtualTextureFeedbackParams Params, float4 SvPosition, float Opacity, uint FrameNumber, RWStructuredBuffer<uint> OutputBuffer)
+void FinalizeVirtualTextureFeedback(in FVirtualTextureFeedbackParams Params, float4 SvPosition, float Opacity, float DBufferOpacityFactor, uint FrameNumber, RWStructuredBuffer<uint> OutputBuffer)
 {
 	uint2 PixelTilePos = (uint2)SvPosition.xy & View.VirtualTextureFeedbackMask;
 	uint PixelTileIndex = (PixelTilePos.y << View.VirtualTextureFeedbackShift) + PixelTilePos.x;
@@ -63,52 +112,39 @@ void FinalizeVirtualTextureFeedback(in FVirtualTextureFeedbackParams Params, flo
 	// This code will only run every few pixels...
 	[branch] if (PixelTileIndex == View.VirtualTextureFeedbackJitterOffset)
 	{
-		// TODO use append buffer ?
 		uint2 PixelPos = (uint2)SvPosition.xy >> View.VirtualTextureFeedbackShift;
 		uint FeedbackPos = PixelPos.y * View.VirtualTextureFeedbackStride + PixelPos.x;
 
-		// When using DBuffer, feedback buffer is blending using decal opacity for opaque objects
-		// This way we avoid overwriting values already written to feedback buffer while rendering dbuffer decals
-#if MATERIALBLENDING_TRANSLUCENT || MATERIALBLENDING_ADDITIVE || MATERIALBLENDING_MODULATE || MATERIAL_SHADINGMODEL_SINGLELAYERWATER || IS_DECAL || (USE_DBUFFER && MATERIAL_DECAL_READ_MASK)
-
-		// Stochastic alpha testing
-#if MATERIALBLENDING_ADDITIVE || MATERIALBLENDING_MODULATE || MATERIAL_SHADINGMODEL_THIN_TRANSLUCENT || MATERIAL_SHADINGMODEL_SINGLELAYERWATER
-		// Opacity may not make sense for additive and modulation blending
-		// Use a value of 0.25f so we have up to 4 layers of VT transparancy
-		const float Alpha = 0.25f;  
-#else
-
-#if (IS_DECAL && !IS_DBUFFER_DECAL)
-		// Even fully opaque gbuffer decals may not completely overwrite data from base pass, so don't allow VT feedback to become fully opaque.
-		// TODO - it may be possible to allow fully opaque decals with certain blend-mode/state combinations, if we can guarantee they overwrite all gbuffer channels
-		const float MaxAlpha = 0.5f;
-#elif (USE_DBUFFER && MATERIAL_DECAL_READ_MASK && !MATERIALDECALRESPONSEMASK)
-		// When reading the DBuffer directly (so without a blend alpha available) we don't wan't to allow the base pass to be fully opaque for VT feedback.
-		const float MaxAlpha = 0.5f;
-#else
-		const float MaxAlpha = 1.0f;
-#endif
-		const float MinAlpha = 0.02f;
-		const float Alpha =  clamp(Opacity, MinAlpha, MaxAlpha);
-#endif
+		if (UseVirtualTextureFeedbackOpacity())
+		{
+			const float Alpha = GetFinalVirtualTextureFeedbackOpacity(Opacity, DBufferOpacityFactor);
 
-        const float AlphaThreshold = frac( 
-			PseudoRandom(PixelPos + FrameNumber) + // Random  value in 0-1 on feedback tile. Add in framenumber for extra jitter so the pseudorandom pattern changes over time.
-			SvPosition.w // Add in depth so we pick different thresholds on different depths
-			);  
+			// Stochastic alpha testing
+			const float AlphaThreshold = frac( 
+				PseudoRandom(PixelPos + FrameNumber) + // Random  value in 0-1 on feedback tile. Add in framenumber for extra jitter so the pseudorandom pattern changes over time.
+				SvPosition.w // Add in depth so we pick different thresholds on different depths
+				);  
 											
-        // Threshold is uniform pseudo random this conditional has a higher probability to succeed for higher apha values
-        // that means more opaque surfaces get more pixels in the feedback buffer which is what we want
-        [branch] if (Alpha > AlphaThreshold)				
-        {
-            OutputBuffer[FeedbackPos] = Params.Request;
-        }     
-#else	
-		OutputBuffer[FeedbackPos] = Params.Request;
-#endif
+			// Threshold is uniform pseudo random this conditional has a higher probability to succeed for higher apha values
+			// that means more opaque surfaces get more pixels in the feedback buffer which is what we want
+			[branch] if (Alpha > AlphaThreshold)				
+			{
+				OutputBuffer[FeedbackPos] = Params.Request;
+			}     
+		}
+		else
+		{
+			OutputBuffer[FeedbackPos] = Params.Request;
+		}
 	}	
 }
 
+/* Version of FinalizeVirtualTextureFeedback with no DBufferOpacityFactor. */
+void FinalizeVirtualTextureFeedback(in FVirtualTextureFeedbackParams Params, float4 SvPosition, float Opacity, uint FrameNumber, RWStructuredBuffer<uint> OutputBuffer)
+{
+	FinalizeVirtualTextureFeedback(Params, SvPosition, Opacity, 1, FrameNumber, OutputBuffer);
+}
+
 /**
  * Address Mode Logic
  * Apply before after calculating mip level/derivatives!
