Substrate - fixed emissive vlamp where it should not be. For slabs/unlit, this is fixed during the sanitize function which properly account for MATERIAL_ALLOW_NEGATIVE_EMISSIVECOLOR.

#tests PC editor dx11 and dx12
#rnx

diff --git a/Engine/Shaders/Private/Substrate/SubstrateLegacyConversion.ush b/Engine/Shaders/Private/Substrate/SubstrateLegacyConversion.ush
--- a/Engine/Shaders/Private/Substrate/SubstrateLegacyConversion.ush
+++ b/Engine/Shaders/Private/Substrate/SubstrateLegacyConversion.ush
@@ -471,7 +471,6 @@ FSubstrateData SubstrateConvertLegacyMaterialStatic(
 	const float3 Ones = float3(1, 1, 1);
 
 	// We need to make sure a few values are safe for the transformations into slab input we do.
-	Emissive = max(Emissive, 0.0f.xxx);
 	Metallic = saturate(Metallic);
 	ThinTranslucentTransmittanceColor = saturate(ThinTranslucentTransmittanceColor);
 	ThinTranslucentSurfaceCoverage = saturate(ThinTranslucentSurfaceCoverage);
