Path Tracer: Fix blending of phase function parameter G in atmosphere

The small epsilon to prevent division by zero can skew the value in low density atmosphere. Revert the logic to what it was previously, which just checks for 0.0 explicitly

diff --git a/Engine/Shaders/Private/PathTracing/Volume/PathTracingVolumeCommon.ush b/Engine/Shaders/Private/PathTracing/Volume/PathTracingVolumeCommon.ush
--- a/Engine/Shaders/Private/PathTracing/Volume/PathTracingVolumeCommon.ush
+++ b/Engine/Shaders/Private/PathTracing/Volume/PathTracingVolumeCommon.ush
@@ -268,8 +268,8 @@ void MergeVolumeShadedResult(inout FVolumeShadedResult M, FVolumeShadedResult R,
 		float Q1m = M.SigmaSHG.x + M.SigmaSHG.y + M.SigmaSHG.z;
 		float Q2r = R.SigmaSDualHG.x + R.SigmaSDualHG.y + R.SigmaSDualHG.z;
 		float Q2m = M.SigmaSDualHG.x + M.SigmaSDualHG.y + M.SigmaSDualHG.z;
-		M.PhaseG        = lerp(M.PhaseG       , R.PhaseG       , Q1r / (Q1m + 1e-6));
-		M.DualPhaseData = lerp(M.DualPhaseData, R.DualPhaseData, Q2r / (Q2m + 1e-6));
+		M.PhaseG        = lerp(M.PhaseG       , R.PhaseG       , Q1m > 0 ? Q1r / Q1m : 0.0);
+		M.DualPhaseData = lerp(M.DualPhaseData, R.DualPhaseData, Q2m > 0 ? Q2r / Q2m : 0.0);
 	}
 }
 
