Substrate: Fix discontinuity in tangent vector from new encoding

The quantization should use round instead of floor to avoid a bias when decoding. It is better to round on encode so we preserve 0.0 and 1.0 on decode.

diff --git a/Engine/Shaders/Private/PathTracing/PathTracingCommon.ush b/Engine/Shaders/Private/PathTracing/PathTracingCommon.ush
--- a/Engine/Shaders/Private/PathTracing/PathTracingCommon.ush
+++ b/Engine/Shaders/Private/PathTracing/PathTracingCommon.ush
@@ -86,7 +86,7 @@ uint PayloadEncodeTangentAngle(float3 Normal, float3 Tangent)
 	float Y = dot(Basis[1], Tangent);
 	// NOTE: We use 'Y' to define the angle so that if Tangent happens to be (0,0,0) it will be decoded back as the "x" axis in the reference basis
 	float PseudoAngle = abs(Y) / (abs(X) + abs(Y)); // Coarse ArcTan, returns value in [0,1]
-	uint QuantizedBits = uint(floor(PseudoAngle * Scale));
+	uint QuantizedBits = uint(round(PseudoAngle * Scale));
 	// Merge with sign bits of X and Y
 	return QuantizedBits | ((asuint(X) >> SignShiftX) & SignMaskX) | ((asuint(Y) >> SignShiftY) & SignMaskY);
 }
diff --git a/Engine/Shaders/Private/Substrate/Substrate.ush b/Engine/Shaders/Private/Substrate/Substrate.ush
--- a/Engine/Shaders/Private/Substrate/Substrate.ush
+++ b/Engine/Shaders/Private/Substrate/Substrate.ush
@@ -3688,7 +3688,7 @@ uint SubstratePackNormalAndTangent(in float3 Normal, in float3 Tangent)
 	float TX = dot(Basis[0], Tangent);
 	float TY = dot(Basis[1], Tangent);
 	float PseudoAngle = abs(TY) / (abs(TX) + abs(TY)); // Coarse ArcTan, returns value in [0,1]
-	uint ABits = uint(floor(PseudoAngle * 255.0f));
+	uint ABits = uint(round(PseudoAngle * 255.0f));
 	// Merge all quantized quantities and sign bits of the tangent
 	return XBits | (YBits << 11) | (ABits << 22) | ((asuint(TX) >> 1) & 0x40000000u) | (asuint(TY) & 0x80000000u);
 }
