Fix CIS shader compilation error.

diff --git a/Engine/Shaders/Private/Substrate/SubstrateMaterialSampling.ush b/Engine/Shaders/Private/Substrate/SubstrateMaterialSampling.ush
--- a/Engine/Shaders/Private/Substrate/SubstrateMaterialSampling.ush
+++ b/Engine/Shaders/Private/Substrate/SubstrateMaterialSampling.ush
@@ -5,6 +5,8 @@
 #include "/Engine/Private/Common.ush"
 #include "/Engine/Private/Substrate/Substrate.ush"
 
+#if SUBSTRATE_ENABLED
+
 struct FSubstrateSampledMaterial
 {
 	float3 WorldNormal;
@@ -62,3 +64,5 @@ uint4 PackSampledMaterial(FSubstrateSampledMaterial In)
 	Out.w = f32tof16(In.PDF);
 	return Out;
 }
+
+#endif
\ No newline at end of file
