Path Tracer: Tweak cloth energy conservation logic

Replace expression with an equivalent but cheaper one (keep the original expression as a comment).

#rnx

diff --git a/Engine/Shaders/Private/PathTracing/Material/PathTracingSubstrateCommon.ush b/Engine/Shaders/Private/PathTracing/Material/PathTracingSubstrateCommon.ush
--- a/Engine/Shaders/Private/PathTracing/Material/PathTracingSubstrateCommon.ush
+++ b/Engine/Shaders/Private/PathTracing/Material/PathTracingSubstrateCommon.ush
@@ -88,10 +88,7 @@ struct FSubstrateSheenData
 		bInv = LTC.y;
 		const float DirectionalAlbedo = LTC.z;
 		Scale = FuzzAmount * DirectionalAlbedo;
-
-		FBxDFEnergyTermsA EnergyTerms = (FBxDFEnergyTermsA)1.0f;
-		EnergyTerms.E = DirectionalAlbedo;
-		Attenuation = lerp(1.0f, ComputeEnergyPreservation(EnergyTerms), FuzzAmount);
+		Attenuation = 1 - FuzzAmount * DirectionalAlbedo; // equivalent to: lerp(1.0f, ComputeEnergyPreservation(EnergyTerms), FuzzAmount);
 #else
 		Roughness = MakeRoughnessSafe(FuzzRoughness, SUBSTRATE_MIN_FUZZ_ROUGHNESS);
 		FBxDFEnergyTermsA Fuzz = ComputeClothEnergyTermsA(Roughness, V.z);
diff --git a/Engine/Shaders/Private/Substrate/SubstrateEvaluation.ush b/Engine/Shaders/Private/Substrate/SubstrateEvaluation.ush
--- a/Engine/Shaders/Private/Substrate/SubstrateEvaluation.ush
+++ b/Engine/Shaders/Private/Substrate/SubstrateEvaluation.ush
@@ -1054,7 +1054,7 @@ FSubstrateEvaluateResult SubstrateEvaluateBSDFCommon(FSubstrateBSDFContext BSDFC
 				#else
 				// Charlie cloth modle for D and Ashikhmin for visibility (more efficient than Charlie).
 				{
-					// ComputeEnergyConservation(EnergyTerms) is not taken into account because we do not white white fuzz to reflect all incoming energy
+					// ComputeEnergyConservation(EnergyTerms) is not taken into account because we do not want white fuzz to reflect all incoming energy
 					// The fuzz should have some transmission to the lower lobes. In order to pass a furnace test, a combination of white diffuse+white fuzz is required
 					EnergyTerms = ComputeClothEnergyTermsA(FuzzRoughness, BSDFContext.Context.NoV);
 
@@ -1790,7 +1790,7 @@ FSubstrateEnvLightResult SubstrateEvaluateForEnvLight(FSubstrateBSDFContext BSDF
 				#endif // SUBSTRATE_SHEEN_QUALITY
 
 				// The specular and diffuse components below the fuzz are only attenuated linearly according to the amount of fuzz.
-				const float Cloth_DirectionalAlbedo_SpecularTransmission = lerp(1.0, ComputeEnergyPreservation(ClothEnergyTerms), FuzzAmount);
+				const float Cloth_DirectionalAlbedo_SpecularTransmission = 1 - ClothEnergyTerms.E * FuzzAmount; // equivalent to: lerp(1.0, ComputeEnergyPreservation(ClothEnergyTerms), FuzzAmount);
 
 				SubstrateEnvLightResult.DiffuseColor		*= Cloth_DirectionalAlbedo_SpecularTransmission;
 				SubstrateEnvLightResult.DiffuseWeight		*= Cloth_DirectionalAlbedo_SpecularTransmission;
