Path Tracer: Fix approximate caustics with substrate

The roughness tracking was not producing the expected approximate roughness leading to opaque shadows in some cases. Update the heuristic to match the legacy path more closely.

diff --git a/Engine/Shaders/Private/PathTracing/Material/PathTracingMaterialSampling.ush b/Engine/Shaders/Private/PathTracing/Material/PathTracingMaterialSampling.ush
--- a/Engine/Shaders/Private/PathTracing/Material/PathTracingMaterialSampling.ush
+++ b/Engine/Shaders/Private/PathTracing/Material/PathTracingMaterialSampling.ush
@@ -289,8 +289,7 @@ float GetAverageRoughness(FPathTracingPayload Payload)
 		{
 			AvgRoughness = max(AvgRoughness, Payload.RoughnessData.y);
 			// Only non-glass slabs contain diffuse
-			float3 DiffuseColor = Payload.HasRefraction() ? 0.0 : Payload.SubsurfaceColor + Payload.DiffuseColor;
-			float DiffuseFraction = LobeSelectionProb(DiffuseColor + Payload.FuzzAmount * Payload.FuzzColor, Payload.SpecularColor);
+			float DiffuseFraction = Payload.HasRefraction() ? 1.0 : F0RGBToMetallic(Payload.SpecularColor);
 			AvgRoughness = lerp(1.0, AvgRoughness, DiffuseFraction);
 			break;
 		}
@@ -310,13 +309,13 @@ float GetAverageRoughness(FPathTracingPayload Payload)
 		case SHADINGMODELID_CLEAR_COAT:
 		{
 			AvgRoughness = max(AvgRoughness, Payload.RoughnessData.y);
-			float DiffuseFraction = LobeSelectionProb(Payload.SubsurfaceColor + Payload.DiffuseColor + Payload.FuzzAmount * Payload.FuzzColor, Payload.SpecularColor + Payload.RoughnessData.z * CLEAR_COAT_F0);
+			float DiffuseFraction = F0RGBToMetallic(Payload.SpecularColor + Payload.RoughnessData.z * CLEAR_COAT_F0);
 			AvgRoughness = lerp(1.0, AvgRoughness, DiffuseFraction);
 			break;
 		}
 		case SHADINGMODELID_EYE:
 		{
-			float DiffuseFraction = LobeSelectionProb(Payload.SubsurfaceColor + Payload.DiffuseColor, Payload.SpecularColor);
+			float DiffuseFraction = F0RGBToMetallic(Payload.SpecularColor);
 			AvgRoughness = lerp(1.0, AvgRoughness, DiffuseFraction);
 			break;
 		}
