Fix crash on Mac preview when mobile deferred is activated

diff --git a/Engine/Shaders/Private/DeferredShadingCommon.ush b/Engine/Shaders/Private/DeferredShadingCommon.ush
--- a/Engine/Shaders/Private/DeferredShadingCommon.ush
+++ b/Engine/Shaders/Private/DeferredShadingCommon.ush
@@ -801,7 +801,7 @@ void MobileFetchGBuffer(in float2 UV, out half4 GBufferA, out half4 GBufferB, ou
 
 #if	(FRAMEBUFFER_FETCH_ACCESS && VULKAN_PROFILE)
 	SceneDepth = ConvertFromDeviceZ(VulkanSubpassDepthFetch());
-#elif (FRAMEBUFFER_FETCH_ACCESS && (USE_SCENE_DEPTH_AUX || USE_GLES_FBF_DEFERRED))
+#elif (FRAMEBUFFER_FETCH_ACCESS && !MAC && (USE_SCENE_DEPTH_AUX || USE_GLES_FBF_DEFERRED))
 	SceneDepth = ConvertFromDeviceZ(DepthbufferFetchES2());
 #else
 	SceneDepth = ConvertFromDeviceZ(Texture2DSampleLevel(MobileSceneTextures.SceneDepthTexture, MobileSceneTextures.SceneDepthTextureSampler, UV, 0).r);
