[HWRT] Fix issues when using "Substrate Material Properties" ray tracing visualization.
- modified HasRaytracingDebugViewModeRaytracedOverlay(...) to check if *any* view is using RAY_TRACING_DEBUG_VIZ_SUBSTRATE_DATA.
- fixed sky being rendered as black.
- fixed typo in Substrate visualization name.
- enabled tonemapping in Substrate visualization.

diff --git a/Engine/Shaders/Private/RayTracing/RayTracingDebug.usf b/Engine/Shaders/Private/RayTracing/RayTracingDebug.usf
--- a/Engine/Shaders/Private/RayTracing/RayTracingDebug.usf
+++ b/Engine/Shaders/Private/RayTracing/RayTracingDebug.usf
@@ -61,7 +61,7 @@ float4 RoughnessDebugColor(float InRoughness)
 	return float4(pow(InRoughness, 2.2f).xxx, 0.0f);
 }
 
-float4 DebugRayTracingMaterial(uint2 PixelCoord, FRayDesc Ray, bool bIsFarField, inout bool bWriteOuput)
+float4 DebugRayTracingMaterial(uint2 PixelCoord, FRayDesc Ray, bool bIsFarField)
 {
 	FRayCone RayCone = (FRayCone)0;
 	RayCone.SpreadAngle = View.EyeToPixelSpreadAngle;
@@ -145,7 +145,6 @@ float4 DebugRayTracingMaterial(uint2 PixelCoord, FRayDesc Ray, bool bIsFarField,
 				SerializeSubstratePixelDebugDataEntry(PixelCoord, Payload.SubstrateData, View.FrameNumber);
 
 			}
-			bWriteOuput = false;
 			return float4(0, 0, 0, 1);
 		}
 		else if (SubstratePixelHeader.ClosureCount > 0)
@@ -326,7 +325,7 @@ RAY_TRACING_ENTRY_RAYGEN(RayTracingDebugMainRGS)
 	float OutDepth = 1.0f;
 #endif
 
-	bool bWriteColorOuput = true;
+	bool bWriteColorOuput = VisualizationMode != RAY_TRACING_DEBUG_VIZ_SUBSTRATE_DATA;
 	bool bWriteDepthOutput = VisualizationMode == RAY_TRACING_DEBUG_VIZ_INSTANCE_OVERLAP;
 #if USE_DEBUG_CHS
 	{
@@ -335,14 +334,14 @@ RAY_TRACING_ENTRY_RAYGEN(RayTracingDebugMainRGS)
 #else // USE_DEBUG_CHS
 	{
 		bool bIsFarField = false;
-		Result = DebugRayTracingMaterial(PixelCoord, Ray, bIsFarField, bWriteColorOuput);
+		Result = DebugRayTracingMaterial(PixelCoord, Ray, bIsFarField);
 		
 		const bool bNearFieldMiss = Result.w < 0.0;
 		if (VisualizationMode == RAY_TRACING_DEBUG_VIZ_FAR_FIELD && bNearFieldMiss)
 		{
 			bIsFarField = true;
 			const float4 FarFieldTint = float4(1.0f, 0.5f, 0.5f, 1.0f);
-			Result = DebugRayTracingMaterial(PixelCoord, Ray, bIsFarField, bWriteColorOuput) * FarFieldTint;
+			Result = DebugRayTracingMaterial(PixelCoord, Ray, bIsFarField) * FarFieldTint;
 		}
 	}
 #endif // USE_DEBUG_CHS
diff --git a/Engine/Source/Runtime/Engine/Private/RayTracingVisualizationData.cpp b/Engine/Source/Runtime/Engine/Private/RayTracingVisualizationData.cpp
--- a/Engine/Source/Runtime/Engine/Private/RayTracingVisualizationData.cpp
+++ b/Engine/Source/Runtime/Engine/Private/RayTracingVisualizationData.cpp
@@ -75,7 +75,7 @@ void FRayTracingVisualizationData::Initialize()
 
 			if (Substrate::IsSubstrateEnabled())
 			{
-				AddVisualizationMode(TEXT("SubtrateMaterialProperties"), LOCTEXT("SubtrateMaterialProperties", "Subtrate Material Properties"), FModeType::Other, RAY_TRACING_DEBUG_VIZ_SUBSTRATE_DATA, false);
+				AddVisualizationMode(TEXT("SubstrateMaterialProperties"), LOCTEXT("SubstrateMaterialProperties", "Substrate Material Properties"), FModeType::Other, RAY_TRACING_DEBUG_VIZ_SUBSTRATE_DATA, true);
 			}
 		}
 
diff --git a/Engine/Source/Runtime/Renderer/Private/RayTracing/RayTracingDebug.cpp b/Engine/Source/Runtime/Renderer/Private/RayTracing/RayTracingDebug.cpp
--- a/Engine/Source/Runtime/Renderer/Private/RayTracing/RayTracingDebug.cpp
+++ b/Engine/Source/Runtime/Renderer/Private/RayTracing/RayTracingDebug.cpp
@@ -1386,14 +1386,16 @@ static uint32 GetRaytracingDebugViewMode(const FSceneView& View)
 
 bool HasRaytracingDebugViewModeRaytracedOverlay(const FSceneViewFamily& ViewFamily)
 {
+	bool bAnySubstrate = false;
+
 	for (const FSceneView* View : ViewFamily.Views)
 	{
-		if (View && GetRaytracingDebugViewMode(*View) != RAY_TRACING_DEBUG_VIZ_SUBSTRATE_DATA)
-		{
-			return true;
-		}
+		const uint32 Mode = View != nullptr ? GetRaytracingDebugViewMode(*View) : UINT32_MAX;
+
+		bAnySubstrate |= Mode == RAY_TRACING_DEBUG_VIZ_SUBSTRATE_DATA;
 	}
-	return false;
+
+	return !bAnySubstrate;
 }
 
 void FDeferredShadingSceneRenderer::RenderRayTracingDebug(FRDGBuilder& GraphBuilder, const FViewInfo& View, FRDGTextureRef SceneColorTexture, FRayTracingPickingFeedback& PickingFeedback)
