[GPU Scene Lights] Rename LightTypeAndPackedShadowMapChannelMask to LightSceneInfoExtraDataPacked and associated functions.
- we pack a lot more data in the uint that is not well represented in the old name.

diff --git a/Engine/Shaders/Private/LightData.ush b/Engine/Shaders/Private/LightData.ush
--- a/Engine/Shaders/Private/LightData.ush
+++ b/Engine/Shaders/Private/LightData.ush
@@ -56,7 +56,7 @@ FRectLightData UnpackRectLightData(uint In0, uint In1, uint In2)
 struct FDirectionalLightData
 {
 	uint   HasDirectionalLight;
-	uint   DirectionalLightShadowMapChannelMask;
+	uint   DirectionalLightSceneInfoExtraDataPacked;
 	float2 DirectionalLightDistanceFadeMAD;
 	float3 DirectionalLightColor;
 	float3 DirectionalLightDirection;
@@ -76,7 +76,7 @@ struct FLocalLightData
 	float4 LightPositionAndInvRadius;
 	float4 LightColorAndIdAndFalloffExponent;
 	float4 SpotAnglesAndSourceRadiusPacked;
-	float4 LightDirectionAndShadowMask;
+	float4 LightDirectionAndSceneInfoExtraDataPacked;
 	float4 LightTangentAndIESDataAndSpecularScale;
 	float3 RectData;
 
@@ -204,12 +204,12 @@ float UnpackLightInvRadius(FLocalLightData In)
 
 float3 UnpackLightDirection(FLocalLightData In)
 {
-	return In.LightDirectionAndShadowMask.xyz;
+	return In.LightDirectionAndSceneInfoExtraDataPacked.xyz;
 }
 
-uint UnpackLightTypeAndShadowMask(FLocalLightData In)
+uint UnpackLightSceneInfoExtraDataPacked(FLocalLightData In)
 {
-	return asuint(In.LightDirectionAndShadowMask.w);
+	return asuint(In.LightDirectionAndSceneInfoExtraDataPacked.w);
 }
 
 float3 UnpackLightColor(float2 In)
@@ -234,68 +234,68 @@ uint UnpackLightingChannelMask(uint In)
 
 uint UnpackLightingChannelMask(FLocalLightData In)
 {
-	const uint LightTypeAndPackedShadowMapChannelMask = asuint(In.LightDirectionAndShadowMask.w);
-	return UnpackLightingChannelMask(LightTypeAndPackedShadowMapChannelMask);
+	const uint LightSceneInfoExtraDataPacked = asuint(In.LightDirectionAndSceneInfoExtraDataPacked.w);
+	return UnpackLightingChannelMask(LightSceneInfoExtraDataPacked);
 }
 
-uint UnpackLightType(uint InShadowMapChannelMaskPacked)
+uint UnpackLightType(uint InLightSceneInfoExtraDataPacked)
 {
-	return (InShadowMapChannelMaskPacked >> 16) & 0x3;
+	return (InLightSceneInfoExtraDataPacked >> 16) & 0x3;
 }
 
-bool UnpackCastShadow(uint InShadowMapChannelMaskPacked)
+bool UnpackCastShadow(uint InLightSceneInfoExtraDataPacked)
 {
-	return ((InShadowMapChannelMaskPacked >> 18) & 0x1) != 0;
+	return ((InLightSceneInfoExtraDataPacked >> 18) & 0x1) != 0;
 }
 
 bool UnpackCastShadow(FLocalLightData In)
 {
-	const uint LightTypeAndPackedShadowMapChannelMask = asuint(In.LightDirectionAndShadowMask.w);
-	return UnpackCastShadow(LightTypeAndPackedShadowMapChannelMask);
+	const uint LightSceneInfoExtraDataPacked = asuint(In.LightDirectionAndSceneInfoExtraDataPacked.w);
+	return UnpackCastShadow(LightSceneInfoExtraDataPacked);
 }
 
-bool UnpackHasLightFunction(uint InShadowMapChannelMaskPacked)
+bool UnpackHasLightFunction(uint InLightSceneInfoExtraDataPacked)
 {
-	return ((InShadowMapChannelMaskPacked >> 19) & 0x1) != 0;
+	return ((InLightSceneInfoExtraDataPacked >> 19) & 0x1) != 0;
 }
 
 bool UnpackHasLightFunction(FLocalLightData In)
 {
-	const uint LightTypeAndPackedShadowMapChannelMask = asuint(In.LightDirectionAndShadowMask.w);
-	return UnpackHasLightFunction(LightTypeAndPackedShadowMapChannelMask);
+	const uint LightSceneInfoExtraDataPacked = asuint(In.LightDirectionAndSceneInfoExtraDataPacked.w);
+	return UnpackHasLightFunction(LightSceneInfoExtraDataPacked);
 }
 
-uint UnpackLightFunctionAtlasIndex(uint InShadowMapChannelMaskPacked)
+uint UnpackLightFunctionAtlasIndex(uint InLightSceneInfoExtraDataPacked)
 {
-	return ((InShadowMapChannelMaskPacked >> 20) & 0xFF);
+	return ((InLightSceneInfoExtraDataPacked >> 20) & 0xFF);
 }
 
 uint UnpackLightFunctionAtlasIndex(FLocalLightData In)
 {
-	const uint LightTypeAndPackedShadowMapChannelMask = asuint(In.LightDirectionAndShadowMask.w);
-	return UnpackLightFunctionAtlasIndex(LightTypeAndPackedShadowMapChannelMask);
+	const uint LightSceneInfoExtraDataPacked = asuint(In.LightDirectionAndSceneInfoExtraDataPacked.w);
+	return UnpackLightFunctionAtlasIndex(LightSceneInfoExtraDataPacked);
 }
 
-uint UnpackAffectsTranslucentLighting(uint InShadowMapChannelMaskPacked)
+uint UnpackAffectsTranslucentLighting(uint InLightSceneInfoExtraDataPacked)
 {
-	return ((InShadowMapChannelMaskPacked >> 28) & 0x1);
+	return ((InLightSceneInfoExtraDataPacked >> 28) & 0x1);
 }
 
 uint UnpackAffectsTranslucentLighting(FLocalLightData In)
 {
-	const uint LightTypeAndPackedShadowMapChannelMask = asuint(In.LightDirectionAndShadowMask.w);
-	return UnpackAffectsTranslucentLighting(LightTypeAndPackedShadowMapChannelMask);
+	const uint LightSceneInfoExtraDataPacked = asuint(In.LightDirectionAndSceneInfoExtraDataPacked.w);
+	return UnpackAffectsTranslucentLighting(LightSceneInfoExtraDataPacked);
 }
 
-uint UnpackIsHandledByMegaLights(uint InShadowMapChannelMaskPacked)
+uint UnpackIsHandledByMegaLights(uint InLightSceneInfoExtraDataPacked)
 {
-	return ((InShadowMapChannelMaskPacked >> 29) & 0x1);
+	return ((InLightSceneInfoExtraDataPacked >> 29) & 0x1);
 }
 
 uint UnpackIsHandledByMegaLights(FLocalLightData In)
 {
-	const uint LightTypeAndPackedShadowMapChannelMask = asuint(In.LightDirectionAndShadowMask.w);
-	return UnpackIsHandledByMegaLights(LightTypeAndPackedShadowMapChannelMask);
+	const uint LightSceneInfoExtraDataPacked = asuint(In.LightDirectionAndSceneInfoExtraDataPacked.w);
+	return UnpackIsHandledByMegaLights(LightSceneInfoExtraDataPacked);
 }
 
 float UnpackVolumetricScatteringIntensity(FLocalLightData In)
@@ -365,14 +365,14 @@ FDeferredLightData ConvertToDeferredLight(
 {
 	FDeferredLightData Out = (FDeferredLightData)0;
 
-	const uint LightTypeAndPackedShadowMapChannelMask = asuint(In.LightDirectionAndShadowMask.w);
-	const uint LightType = UnpackLightType(LightTypeAndPackedShadowMapChannelMask);
+	const uint LightSceneInfoExtraDataPacked = asuint(In.LightDirectionAndSceneInfoExtraDataPacked.w);
+	const uint LightType = UnpackLightType(LightSceneInfoExtraDataPacked);
 
 	Out.TranslatedWorldPosition = In.LightPositionAndInvRadius.xyz;
 	Out.InvRadius				= In.LightPositionAndInvRadius.w;
 	Out.Color					= UnpackLightColor(In.LightColorAndIdAndFalloffExponent.xy);
 	Out.FalloffExponent			= In.LightColorAndIdAndFalloffExponent.w;
-	Out.Direction				= In.LightDirectionAndShadowMask.xyz;
+	Out.Direction				= In.LightDirectionAndSceneInfoExtraDataPacked.xyz;
 	Out.SpotAngles				= GetLightSpotAngles(In);
 	Out.SourceRadius			= UnpackLightSourceRadius(In);
 	Out.SourceLength			= UnpackLightSourceLength(In);
@@ -401,10 +401,10 @@ FDeferredLightData ConvertToDeferredLight(
 	// * ShadowMapChannelMask	= [3:0]
 	// * PreviewShadowMapChannel= [7:4] 
 	// * LightincChannelMask    = [8:10] 
-	Out.ShadowedBits				= (LightTypeAndPackedShadowMapChannelMask & 0xFF) != 0 ? 1 : 0;
-	Out.ShadowMapChannelMask		= UnpackShadowMapChannelMask(LightTypeAndPackedShadowMapChannelMask);
-	OutPreviewShadowMapChannelMask	= UnpackShadowMapChannelMask(LightTypeAndPackedShadowMapChannelMask >> 4);
-	OutLightingChannelMask			= UnpackLightingChannelMask(LightTypeAndPackedShadowMapChannelMask);
+	Out.ShadowedBits				= (LightSceneInfoExtraDataPacked & 0xFF) != 0 ? 1 : 0;
+	Out.ShadowMapChannelMask		= UnpackShadowMapChannelMask(LightSceneInfoExtraDataPacked);
+	OutPreviewShadowMapChannelMask	= UnpackShadowMapChannelMask(LightSceneInfoExtraDataPacked >> 4);
+	OutLightingChannelMask			= UnpackLightingChannelMask(LightSceneInfoExtraDataPacked);
 	return Out;
 }
 
@@ -472,7 +472,7 @@ FLightShaderParameters ConvertFromLocal(const FLocalLightData In)
 	Out.InvRadius				= In.LightPositionAndInvRadius.w;
 	Out.Color					= UnpackLightColor(In.LightColorAndIdAndFalloffExponent.xy);
 	Out.FalloffExponent			= In.LightColorAndIdAndFalloffExponent.w;
-	Out.Direction				= In.LightDirectionAndShadowMask.xyz;
+	Out.Direction				= In.LightDirectionAndSceneInfoExtraDataPacked.xyz;
 	Out.Tangent					= In.LightTangentAndIESDataAndSpecularScale.xyz;
 	Out.SpotAngles				= GetLightSpotAngles(In);
 	Out.SpecularScale			= UnpackLightSpecularScale(In);
@@ -514,7 +514,7 @@ FDeferredLightData ConvertToDeferredLight(
 	Out.bRadialLight		= false;
 	Out.SpecularScale		= In.DirectionalLightSpecularScale * InSpecularScale;
 	Out.DiffuseScale		= In.DirectionalLightDiffuseScale;
-	Out.ShadowedBits		= (In.DirectionalLightShadowMapChannelMask & 0xFF) != 0 ? 1 : 0;
+	Out.ShadowedBits		= (In.DirectionalLightSceneInfoExtraDataPacked & 0xFF) != 0 ? 1 : 0;
 	Out.HairTransmittance	= InitHairTransmittanceData();
 	Out.SourceRadius		= In.DirectionalLightSourceRadius;
 	Out.SoftSourceRadius	= In.DirectionalLightSoftSourceRadius;
@@ -529,9 +529,9 @@ FDeferredLightData ConvertToDeferredLight(
 	// * ShadowMapChannelMask	= [3:0]
 	// * PreviewShadowMapChannel= [7:4] 
 	// * LightincChannelMask    = [8:10] 
-	Out.ShadowMapChannelMask		= UnpackShadowMapChannelMask(In.DirectionalLightShadowMapChannelMask);
-	OutPreviewShadowMapChannelMask	= UnpackShadowMapChannelMask(In.DirectionalLightShadowMapChannelMask >> 4);
-	OutLightingChannelMask			= UnpackLightingChannelMask(In.DirectionalLightShadowMapChannelMask);
+	Out.ShadowMapChannelMask		= UnpackShadowMapChannelMask(In.DirectionalLightSceneInfoExtraDataPacked);
+	OutPreviewShadowMapChannelMask	= UnpackShadowMapChannelMask(In.DirectionalLightSceneInfoExtraDataPacked >> 4);
+	OutLightingChannelMask			= UnpackLightingChannelMask(In.DirectionalLightSceneInfoExtraDataPacked);
 
 	return Out;
 }
@@ -541,7 +541,7 @@ FDeferredLightData ConvertToDeferredLight(FLightSceneData In, FLightViewData InV
 {
 	FDeferredLightData Out = (FDeferredLightData)0;
 
-	const uint LightType = UnpackLightType(InViewData.LightTypeAndShadowMapChannelMaskPacked);
+	const uint LightType = UnpackLightType(InViewData.LightSceneInfoExtraDataPacked);
 
 	Out.TranslatedWorldPosition = InViewData.TranslatedWorldPosition;
 	Out.InvRadius				= In.InvRadius;
@@ -568,9 +568,9 @@ FDeferredLightData ConvertToDeferredLight(FLightSceneData In, FLightViewData InV
 	Out.RectLightData.AtlasData.AtlasUVScale = InViewData.RectLightAtlasUVScale;
 
 	Out.IESAtlasIndex					= InViewData.IESAtlasIndex;
-	Out.LightFunctionAtlasLightIndex	= UnpackLightFunctionAtlasIndex(InViewData.LightTypeAndShadowMapChannelMaskPacked);
+	Out.LightFunctionAtlasLightIndex	= UnpackLightFunctionAtlasIndex(InViewData.LightSceneInfoExtraDataPacked);
 
-	Out.bAffectsTranslucentLighting = UnpackAffectsTranslucentLighting(InViewData.LightTypeAndShadowMapChannelMaskPacked);
+	Out.bAffectsTranslucentLighting = UnpackAffectsTranslucentLighting(InViewData.LightSceneInfoExtraDataPacked);
 
 	// Shadow
 	// * Static shadowing uses ShadowMapChannel, 
@@ -579,8 +579,8 @@ FDeferredLightData ConvertToDeferredLight(FLightSceneData In, FLightViewData InV
 	// * ShadowMapChannelMask	= [3:0]
 	// * PreviewShadowMapChannel= [7:4] 
 	// * LightingChannelMask    = [8:10] 
-	Out.ShadowedBits				= (InViewData.LightTypeAndShadowMapChannelMaskPacked & 0xFF) != 0 ? 1 : 0;
-	Out.ShadowMapChannelMask		= UnpackShadowMapChannelMask(InViewData.LightTypeAndShadowMapChannelMaskPacked);
+	Out.ShadowedBits				= (InViewData.LightSceneInfoExtraDataPacked & 0xFF) != 0 ? 1 : 0;
+	Out.ShadowMapChannelMask		= UnpackShadowMapChannelMask(InViewData.LightSceneInfoExtraDataPacked);
 	return Out;
 }
 
diff --git a/Engine/Shaders/Private/LightGridCommon.ush b/Engine/Shaders/Private/LightGridCommon.ush
--- a/Engine/Shaders/Private/LightGridCommon.ush
+++ b/Engine/Shaders/Private/LightGridCommon.ush
@@ -241,33 +241,33 @@ FDirectionalLightData GetDirectionalLightData(uint EyeIndex)
 	{
 #endif
 	
-		Result.HasDirectionalLight					= ForwardLightData.HasDirectionalLight;
-		Result.DirectionalLightShadowMapChannelMask = ForwardLightData.DirectionalLightShadowMapChannelMask;
-		Result.DirectionalLightDistanceFadeMAD		= ForwardLightData.DirectionalLightDistanceFadeMAD;
-		Result.DirectionalLightColor				= ForwardLightData.DirectionalLightColor;
-		Result.DirectionalLightDirection			= ForwardLightData.DirectionalLightDirection;
-		Result.DirectionalLightSourceRadius			= ForwardLightData.DirectionalLightSourceRadius;
-		Result.DirectionalLightSoftSourceRadius		= ForwardLightData.DirectionalLightSoftSourceRadius;
-		Result.DirectionalLightSpecularScale		= ForwardLightData.DirectionalLightSpecularScale;
-		Result.DirectionalLightDiffuseScale			= ForwardLightData.DirectionalLightDiffuseScale;
-		Result.LightFunctionAtlasLightIndex			= ForwardLightData.LightFunctionAtlasLightIndex;
-		Result.bAffectsTranslucentLighting			= ForwardLightData.bAffectsTranslucentLighting;
+		Result.HasDirectionalLight						= ForwardLightData.HasDirectionalLight;
+		Result.DirectionalLightSceneInfoExtraDataPacked = ForwardLightData.DirectionalLightSceneInfoExtraDataPacked;
+		Result.DirectionalLightDistanceFadeMAD			= ForwardLightData.DirectionalLightDistanceFadeMAD;
+		Result.DirectionalLightColor					= ForwardLightData.DirectionalLightColor;
+		Result.DirectionalLightDirection				= ForwardLightData.DirectionalLightDirection;
+		Result.DirectionalLightSourceRadius				= ForwardLightData.DirectionalLightSourceRadius;
+		Result.DirectionalLightSoftSourceRadius			= ForwardLightData.DirectionalLightSoftSourceRadius;
+		Result.DirectionalLightSpecularScale			= ForwardLightData.DirectionalLightSpecularScale;
+		Result.DirectionalLightDiffuseScale				= ForwardLightData.DirectionalLightDiffuseScale;
+		Result.LightFunctionAtlasLightIndex				= ForwardLightData.LightFunctionAtlasLightIndex;
+		Result.bAffectsTranslucentLighting				= ForwardLightData.bAffectsTranslucentLighting;
 
 #if (INSTANCED_STEREO || MOBILE_MULTI_VIEW)
 	}
 	else
 	{
-		Result.HasDirectionalLight					= ForwardLightDataISR.HasDirectionalLight;
-		Result.DirectionalLightShadowMapChannelMask = ForwardLightDataISR.DirectionalLightShadowMapChannelMask;
-		Result.DirectionalLightDistanceFadeMAD		= ForwardLightDataISR.DirectionalLightDistanceFadeMAD;
-		Result.DirectionalLightColor				= ForwardLightDataISR.DirectionalLightColor;
-		Result.DirectionalLightDirection			= ForwardLightDataISR.DirectionalLightDirection;
-		Result.DirectionalLightSourceRadius			= ForwardLightDataISR.DirectionalLightSourceRadius;
-		Result.DirectionalLightSoftSourceRadius		= ForwardLightDataISR.DirectionalLightSoftSourceRadius;
-		Result.DirectionalLightSpecularScale		= ForwardLightDataISR.DirectionalLightSpecularScale;
-		Result.DirectionalLightDiffuseScale			= ForwardLightDataISR.DirectionalLightDiffuseScale;
-		Result.LightFunctionAtlasLightIndex			= ForwardLightDataISR.LightFunctionAtlasLightIndex;
-		Result.bAffectsTranslucentLighting			= ForwardLightDataISR.bAffectsTranslucentLighting;
+		Result.HasDirectionalLight						= ForwardLightDataISR.HasDirectionalLight;
+		Result.DirectionalLightSceneInfoExtraDataPacked = ForwardLightDataISR.DirectionalLightSceneInfoExtraDataPacked;
+		Result.DirectionalLightDistanceFadeMAD			= ForwardLightDataISR.DirectionalLightDistanceFadeMAD;
+		Result.DirectionalLightColor					= ForwardLightDataISR.DirectionalLightColor;
+		Result.DirectionalLightDirection				= ForwardLightDataISR.DirectionalLightDirection;
+		Result.DirectionalLightSourceRadius				= ForwardLightDataISR.DirectionalLightSourceRadius;
+		Result.DirectionalLightSoftSourceRadius			= ForwardLightDataISR.DirectionalLightSoftSourceRadius;
+		Result.DirectionalLightSpecularScale			= ForwardLightDataISR.DirectionalLightSpecularScale;
+		Result.DirectionalLightDiffuseScale				= ForwardLightDataISR.DirectionalLightDiffuseScale;
+		Result.LightFunctionAtlasLightIndex				= ForwardLightDataISR.LightFunctionAtlasLightIndex;
+		Result.bAffectsTranslucentLighting				= ForwardLightDataISR.bAffectsTranslucentLighting;
 	}
 #endif
 
@@ -282,18 +282,18 @@ FLocalLightData GetLocalLightData_Internal(
 {
 	FLocalLightData Out = (FLocalLightData)0;
 
-	Out.bClusteredDeferredSupported				= LocalLightIndex <  ClusteredDeferredSupportedEndIndex;
-	Out.bIsSimpleLight							= LocalLightIndex <  SimpleLightsEndIndex;
-	Out.LightPositionAndInvRadius				= InData0;
-	Out.LightColorAndIdAndFalloffExponent		= InData1;
-	Out.LightDirectionAndShadowMask				= InData2;
-	Out.SpotAnglesAndSourceRadiusPacked			= InData3;
-	Out.LightTangentAndIESDataAndSpecularScale	= InData4;
-	Out.RectData								= InData5.xyz;
+	Out.bClusteredDeferredSupported					= LocalLightIndex <  ClusteredDeferredSupportedEndIndex;
+	Out.bIsSimpleLight								= LocalLightIndex <  SimpleLightsEndIndex;
+	Out.LightPositionAndInvRadius					= InData0;
+	Out.LightColorAndIdAndFalloffExponent			= InData1;
+	Out.LightDirectionAndSceneInfoExtraDataPacked	= InData2;
+	Out.SpotAnglesAndSourceRadiusPacked				= InData3;
+	Out.LightTangentAndIESDataAndSpecularScale		= InData4;
+	Out.RectData									= InData5.xyz;
 	// See PackVirtualShadowMapIdAndPrevLocalLightIndex
-	Out.VirtualShadowMapId						= int(asuint(InData5.w) >> 16U) - 1;
-	Out.PrevLocalLightIndex						= int(asuint(InData5.w) & 0xFFFF) - 1;
-	Out.LightSceneId							= int(Out.LightColorAndIdAndFalloffExponent.z);
+	Out.VirtualShadowMapId							= int(asuint(InData5.w) >> 16U) - 1;
+	Out.PrevLocalLightIndex							= int(asuint(InData5.w) & 0xFFFF) - 1;
+	Out.LightSceneId								= int(Out.LightColorAndIdAndFalloffExponent.z);
 
 	return Out;
 }
diff --git a/Engine/Shaders/Private/MobileLightingCommon.ush b/Engine/Shaders/Private/MobileLightingCommon.ush
--- a/Engine/Shaders/Private/MobileLightingCommon.ush
+++ b/Engine/Shaders/Private/MobileLightingCommon.ush
@@ -602,7 +602,7 @@ void AccumulateLightGridLocalLighting(const FCulledLightsGridHeader InLightGridH
 			continue;
 		}
 
-		uint LightTypeAndPackedShadowMapChannelMask = asuint(LocalLight.LightDirectionAndShadowMask.w);
+		uint LightSceneInfoExtraDataPacked = UnpackLightSceneInfoExtraDataPacked(LocalLight);
 
 #if !MOBILE_DEFERRED_LIGHTING
 		// Check lighting channels for mobile forward
@@ -613,13 +613,13 @@ void AccumulateLightGridLocalLighting(const FCulledLightsGridHeader InLightGridH
 #endif
 
 		// bits [17:16] really
-		uint LightType = UnpackLightType(LightTypeAndPackedShadowMapChannelMask);
+		uint LightType = UnpackLightType(LightSceneInfoExtraDataPacked);
 
 		half DynamicShadowing = 1.0f;
 #if USE_SHADOWMASKTEXTURE
 		if (LightType == LIGHT_TYPE_SPOT || LightType == LIGHT_TYPE_POINT)
 		{
-			half4 PreviewShadowMapChannelMask = UnpackShadowMapChannelMask(LightTypeAndPackedShadowMapChannelMask >> 4);
+			half4 PreviewShadowMapChannelMask = UnpackShadowMapChannelMask(LightSceneInfoExtraDataPacked >> 4);
 			DynamicShadowing = dot(PreviewShadowMapChannelMask, DynamicShadowFactors);
 		}
 #endif
diff --git a/Engine/Shaders/Private/VirtualShadowMaps/VirtualShadowMapPageMarking.usf b/Engine/Shaders/Private/VirtualShadowMaps/VirtualShadowMapPageMarking.usf
--- a/Engine/Shaders/Private/VirtualShadowMaps/VirtualShadowMapPageMarking.usf
+++ b/Engine/Shaders/Private/VirtualShadowMaps/VirtualShadowMapPageMarking.usf
@@ -493,8 +493,8 @@ void GeneratePageFlagsFromPixels(
 				continue;
 			}
 
-			const uint LightTypeAndPackedShadowMapChannelMask = UnpackLightTypeAndShadowMask(LightData);
-			const uint LightType = UnpackLightType(LightTypeAndPackedShadowMapChannelMask);
+			const uint LightSceneInfoExtraDataPacked = UnpackLightSceneInfoExtraDataPacked(LightData);
+			const uint LightType = UnpackLightType(LightSceneInfoExtraDataPacked);
 			
 			const bool bIsRectLight = LightType == LIGHT_TYPE_RECT;
 			if (bIsRectLight)
@@ -663,8 +663,8 @@ void GeneratePageFlagsFromFroxelsCS(uint FroxelIndex : SV_DispatchThreadID)
 				continue;
 			}
 
-			const uint LightTypeAndPackedShadowMapChannelMask = UnpackLightTypeAndShadowMask(LightData);
-			const uint LightType = UnpackLightType(LightTypeAndPackedShadowMapChannelMask);
+			const uint LightSceneInfoExtraDataPacked = UnpackLightSceneInfoExtraDataPacked(LightData);
+			const uint LightType = UnpackLightType(LightSceneInfoExtraDataPacked);
 			
 			const bool bIsRectLight = LightType == LIGHT_TYPE_RECT;
 			if (bIsRectLight)
diff --git a/Engine/Shaders/Shared/LightViewData.h b/Engine/Shaders/Shared/LightViewData.h
--- a/Engine/Shaders/Shared/LightViewData.h
+++ b/Engine/Shaders/Shared/LightViewData.h
@@ -20,8 +20,7 @@ struct FLightViewData
 	float VolumetricScatteringIntensity;
 	uint VirtualShadowMapId;
 
-	// Extra fields
-	uint LightTypeAndShadowMapChannelMaskPacked;
+	uint LightSceneInfoExtraDataPacked;
 	
 	// Rect. light atlas transformation
 	float2 RectLightAtlasUVOffset;
diff --git a/Engine/Source/Runtime/Renderer/Private/BasePassRendering.h b/Engine/Source/Runtime/Renderer/Private/BasePassRendering.h
--- a/Engine/Source/Runtime/Renderer/Private/BasePassRendering.h
+++ b/Engine/Source/Runtime/Renderer/Private/BasePassRendering.h
@@ -61,7 +61,7 @@ class FForwardLocalLightData
 public:
 	FVector4f LightPositionAndInvRadius;
 	FVector4f LightColorAndIdAndFalloffExponent;
-	FVector4f LightDirectionAndShadowMapChannelMask;
+	FVector4f LightDirectionAndSceneInfoExtraDataPacked;
 	FVector4f SpotAnglesAndSourceRadiusPacked;
 	FVector4f LightTangentAndIESDataAndSpecularScale;
 	FVector4f RectDataAndVirtualShadowMapIdOrPrevLocalLightIndex;
diff --git a/Engine/Source/Runtime/Renderer/Private/LightGridInjection.cpp b/Engine/Source/Runtime/Renderer/Private/LightGridInjection.cpp
--- a/Engine/Source/Runtime/Renderer/Private/LightGridInjection.cpp
+++ b/Engine/Source/Runtime/Renderer/Private/LightGridInjection.cpp
@@ -432,8 +432,8 @@ static void PackLocalLightData(
 	const FVector3f LightTranslatedWorldPosition(View.ViewMatrices.GetPreViewTranslation() + SimpleLightPerViewData.Position);
 
 	// No shadowmap channels for simple lights
-	uint32 ShadowMapChannelMask = 0;
-	ShadowMapChannelMask |= SimpleLightLightingChannelMask << 8;
+	uint32 LightSceneInfoExtraDataPacked = 0;
+	LightSceneInfoExtraDataPacked |= SimpleLightLightingChannelMask << 8;
 
 	// Pack both values into a single float to keep float4 alignment
 	const float SimpleLightSourceLength = 0;
@@ -464,7 +464,7 @@ static void PackLocalLightData(
 
 	Out.LightPositionAndInvRadius							= FVector4f(LightTranslatedWorldPosition, 1.0f / FMath::Max(SimpleLight.Radius, KINDA_SMALL_NUMBER));
 	Out.LightColorAndIdAndFalloffExponent					= FVector4f(LightColorPacked.X, LightColorPacked.Y, INDEX_NONE, SimpleLight.Exponent);
-	Out.LightDirectionAndShadowMapChannelMask				= FVector4f(FVector3f(1, 0, 0), FMath::AsFloat(ShadowMapChannelMask));
+	Out.LightDirectionAndSceneInfoExtraDataPacked			= FVector4f(FVector3f(1, 0, 0), FMath::AsFloat(LightSceneInfoExtraDataPacked));
 	Out.SpotAnglesAndSourceRadiusPacked						= FVector4f(-2, 1, FMath::AsFloat(PackedZ), FMath::AsFloat(PackedW));
 	Out.LightTangentAndIESDataAndSpecularScale				= FVector4f(1.0f, 0.0f, 0.0f, FMath::AsFloat(SpecularScale_DiffuseScale_IESData));
 	Out.RectDataAndVirtualShadowMapIdOrPrevLocalLightIndex	= FVector4f(FMath::AsFloat(RectPackedX), FMath::AsFloat(RectPackedY), FMath::AsFloat(RectPackedZ), FMath::AsFloat(VirtualShadowMapIdAndPrevLocalLightIndex));
@@ -474,7 +474,7 @@ static void PackLocalLightData(
 	FForwardLocalLightData& Out, 
 	const FViewInfo& View,
 	const FLightRenderParameters& LightParameters,
-	const uint32 LightTypeAndShadowMapChannelMaskAndLightFunctionIndexPacked,
+	const uint32 LightSceneInfoExtraDataPacked,
 	const int32 LightSceneId,
 	const int32 VirtualShadowMapId,
 	const int32 PrevLocalLightIndex,
@@ -509,7 +509,7 @@ static void PackLocalLightData(
 	// NOTE: SpotAngles needs full-precision for VSM one pass projection
 	Out.LightPositionAndInvRadius							= FVector4f(LightTranslatedWorldPosition, LightParameters.InvRadius);
 	Out.LightColorAndIdAndFalloffExponent					= FVector4f(LightColorPacked.X, LightColorPacked.Y, LightSceneId, LightParameters.FalloffExponent);
-	Out.LightDirectionAndShadowMapChannelMask				= FVector4f(LightParameters.Direction, FMath::AsFloat(LightTypeAndShadowMapChannelMaskAndLightFunctionIndexPacked));
+	Out.LightDirectionAndSceneInfoExtraDataPacked			= FVector4f(LightParameters.Direction, FMath::AsFloat(LightSceneInfoExtraDataPacked));
 	Out.SpotAnglesAndSourceRadiusPacked						= FVector4f(LightParameters.SpotAngles.X, LightParameters.SpotAngles.Y, FMath::AsFloat(PackedZ), FMath::AsFloat(PackedW));
 	Out.LightTangentAndIESDataAndSpecularScale				= FVector4f(LightParameters.Tangent, FMath::AsFloat(SpecularScale_DiffuseScale_IESData));
 	Out.RectDataAndVirtualShadowMapIdOrPrevLocalLightIndex	= FVector4f(FMath::AsFloat(RectPackedX), FMath::AsFloat(RectPackedY), FMath::AsFloat(RectPackedZ), FMath::AsFloat(VirtualShadowMapIdAndPrevLocalLightIndex));
@@ -851,7 +851,7 @@ FComputeLightGridOutput FSceneRenderer::ComputeLightGrid(FRDGBuilder& GraphBuild
 						LightParameters.Color *= LightProxy->GetIndirectLightingScale();
 					}
 
-					uint32 LightTypeAndShadowMapChannelMaskPacked = LightSceneInfo->PackLightTypeAndShadowMapChannelMask(bAllowStaticLighting, SortedLightInfo.SortKey.Fields.bLightFunction, SortedLightInfo.SortKey.Fields.bHandledByMegaLights);
+					uint32 LightSceneInfoExtraDataPacked = LightSceneInfo->PackExtraData(bAllowStaticLighting, SortedLightInfo.SortKey.Fields.bLightFunction, SortedLightInfo.SortKey.Fields.bHandledByMegaLights);
 
 					const bool bDynamicShadows = ViewFamily.EngineShowFlags.DynamicShadows && VisibleLightInfos.IsValidIndex(LightSceneInfo->Id);
 					const int32 VirtualShadowMapId = bDynamicShadows ? VisibleLightInfos[LightSceneInfo->Id].GetVirtualShadowMapId( &View ) : INDEX_NONE;
@@ -896,7 +896,7 @@ FComputeLightGridOutput FSceneRenderer::ComputeLightGrid(FRDGBuilder& GraphBuild
 							LightData,
 							View,
 							LightParameters,
-							LightTypeAndShadowMapChannelMaskPacked,
+							LightSceneInfoExtraDataPacked,
 							LightSceneInfo->Id,
 							VirtualShadowMapId,
 							PrevLocalLightIndex,
@@ -980,7 +980,7 @@ FComputeLightGridOutput FSceneRenderer::ComputeLightGrid(FRDGBuilder& GraphBuild
 							ForwardLightData->DirectionalLightDirection = LightParameters.Direction;
 							ForwardLightData->DirectionalLightSourceRadius = LightParameters.SourceRadius;
 							ForwardLightData->DirectionalLightSoftSourceRadius = LightParameters.SoftSourceRadius;
-							ForwardLightData->DirectionalLightShadowMapChannelMask = LightTypeAndShadowMapChannelMaskPacked;
+							ForwardLightData->DirectionalLightSceneInfoExtraDataPacked = LightSceneInfoExtraDataPacked;
 							ForwardLightData->DirectionalLightVSM = INDEX_NONE;
 							ForwardLightData->LightFunctionAtlasLightIndex = LightParameters.LightFunctionAtlasLightIndex;
 							ForwardLightData->bAffectsTranslucentLighting = LightParameters.bAffectsTranslucentLighting;
@@ -1332,7 +1332,7 @@ FComputeLightGridOutput FSceneRenderer::PrepareForwardLightData(FRDGBuilder& Gra
 						CurrentLightViewData.Color = FVector3f(LightParameters.Color) * LightFade * LightParameters.GetLightExposureScale(Exposure);
 						CurrentLightViewData.VolumetricScatteringIntensity = VolumetricScatteringIntensity;
 						CurrentLightViewData.VirtualShadowMapId = VirtualShadowMapId;
-						CurrentLightViewData.LightTypeAndShadowMapChannelMaskPacked = LightSceneInfo->PackLightTypeAndShadowMapChannelMask(bAllowStaticLighting);
+						CurrentLightViewData.LightSceneInfoExtraDataPacked = LightSceneInfo->PackExtraData(bAllowStaticLighting);
 						CurrentLightViewData.RectLightAtlasUVOffset = LightParameters.RectLightAtlasUVOffset;
 						CurrentLightViewData.RectLightAtlasUVScale = LightParameters.RectLightAtlasUVScale;
 						CurrentLightViewData.RectLightAtlasMaxLevel = LightParameters.RectLightAtlasMaxLevel;
@@ -1344,7 +1344,7 @@ FComputeLightGridOutput FSceneRenderer::PrepareForwardLightData(FRDGBuilder& Gra
 						CurrentLightViewData.Color = FVector3f::ZeroVector;
 						CurrentLightViewData.VolumetricScatteringIntensity = 0.0f;
 						CurrentLightViewData.VirtualShadowMapId = INDEX_NONE;
-						CurrentLightViewData.LightTypeAndShadowMapChannelMaskPacked = 0;
+						CurrentLightViewData.LightSceneInfoExtraDataPacked = 0;
 						CurrentLightViewData.RectLightAtlasUVOffset = FVector2f::ZeroVector;
 						CurrentLightViewData.RectLightAtlasUVScale = FVector2f::ZeroVector;
 						CurrentLightViewData.RectLightAtlasMaxLevel = FLightRenderParameters::GetRectLightAtlasInvalidMIPLevel();
diff --git a/Engine/Source/Runtime/Renderer/Private/LightSceneInfo.cpp b/Engine/Source/Runtime/Renderer/Private/LightSceneInfo.cpp
--- a/Engine/Source/Runtime/Renderer/Private/LightSceneInfo.cpp
+++ b/Engine/Source/Runtime/Renderer/Private/LightSceneInfo.cpp
@@ -380,31 +380,33 @@ bool FLightSceneInfo::ShouldRecordShadowSubjectsForMobile() const
 	return bShouldRecordShadowSubjectsForMobile;
 }
 
-uint32 FLightSceneInfo::PackLightTypeAndShadowMapChannelMask(bool bAllowStaticLighting, bool bLightFunction, bool bMegaLight) const
+uint32 FLightSceneInfo::PackExtraData(bool bAllowStaticLighting, bool bLightFunction, bool bMegaLight) const
 {
 	uint32 Result = 0;
+	uint32 BitOffset = 0;
 
-	// Light type and shadow map
-	int32 ShadowMapChannel = Proxy->GetShadowMapChannel();
-	int32 CurrentDynamicShadowMapChannel = GetDynamicShadowMapChannel();
-
-	if (!bAllowStaticLighting)
 	{
-		ShadowMapChannel = INDEX_NONE;
-	}
+		int32 ShadowMapChannel = Proxy->GetShadowMapChannel();
+		int32 CurrentDynamicShadowMapChannel = GetDynamicShadowMapChannel();
+
+		if (!bAllowStaticLighting)
+		{
+			ShadowMapChannel = INDEX_NONE;
+		}
 
-	// Static shadowing uses ShadowMapChannel, dynamic shadows are packed into light attenuation using DynamicShadowMapChannel
-	Result =
-		(ShadowMapChannel == 0 ? 1 : 0) |
-		(ShadowMapChannel == 1 ? 2 : 0) |
-		(ShadowMapChannel == 2 ? 4 : 0) |
-		(ShadowMapChannel == 3 ? 8 : 0) |
-		(CurrentDynamicShadowMapChannel == 0 ? 16 : 0) |
-		(CurrentDynamicShadowMapChannel == 1 ? 32 : 0) |
-		(CurrentDynamicShadowMapChannel == 2 ? 64 : 0) |
-		(CurrentDynamicShadowMapChannel == 3 ? 128 : 0);
-
-	uint32 BitOffset = 8;
+		// Static shadowing uses ShadowMapChannel, dynamic shadows are packed into light attenuation using DynamicShadowMapChannel
+		uint32 ShadowMapChannelMask =
+			(ShadowMapChannel == 0 ? 1 : 0) |
+			(ShadowMapChannel == 1 ? 2 : 0) |
+			(ShadowMapChannel == 2 ? 4 : 0) |
+			(ShadowMapChannel == 3 ? 8 : 0) |
+			(CurrentDynamicShadowMapChannel == 0 ? 16 : 0) |
+			(CurrentDynamicShadowMapChannel == 1 ? 32 : 0) |
+			(CurrentDynamicShadowMapChannel == 2 ? 64 : 0) |
+			(CurrentDynamicShadowMapChannel == 3 ? 128 : 0);
+
+		Result |= ShadowMapChannelMask << BitOffset;					BitOffset += 8;
+	}
 
 	Result |= Proxy->GetLightingChannelMask() << BitOffset;				BitOffset += 8;					// This could be 3 bits
 
diff --git a/Engine/Source/Runtime/Renderer/Private/LightSceneInfo.h b/Engine/Source/Runtime/Renderer/Private/LightSceneInfo.h
--- a/Engine/Source/Runtime/Renderer/Private/LightSceneInfo.h
+++ b/Engine/Source/Runtime/Renderer/Private/LightSceneInfo.h
@@ -335,7 +335,9 @@ public:
 
 	bool ShouldRecordShadowSubjectsForMobile() const;
 
-	uint32 PackLightTypeAndShadowMapChannelMask(bool bAllowStaticLighting, bool bLightFunction = false, bool bMegaLight = false) const;
+	// Extra data such as shadow channel masks, light type and flags packed into a uint32.
+	// Unpacked in shader using helper functions defined in LightData.ush
+	uint32 PackExtraData(bool bAllowStaticLighting, bool bLightFunction = false, bool bMegaLight = false) const;
 };
 
 /** Defines how the light is stored in the scene's light octree. */
diff --git a/Engine/Source/Runtime/Renderer/Private/SceneRendering.h b/Engine/Source/Runtime/Renderer/Private/SceneRendering.h
--- a/Engine/Source/Runtime/Renderer/Private/SceneRendering.h
+++ b/Engine/Source/Runtime/Renderer/Private/SceneRendering.h
@@ -671,7 +671,7 @@ BEGIN_GLOBAL_SHADER_PARAMETER_STRUCT_WITH_CONSTRUCTOR(FForwardLightData, )
 	SHADER_PARAMETER(float, DirectionalLightVolumetricScatteringIntensity)
 	SHADER_PARAMETER(float, DirectionalLightSpecularScale)
 	SHADER_PARAMETER(float, DirectionalLightDiffuseScale)
-	SHADER_PARAMETER(uint32, DirectionalLightShadowMapChannelMask)
+	SHADER_PARAMETER(uint32, DirectionalLightSceneInfoExtraDataPacked)
 	SHADER_PARAMETER(FVector2f, DirectionalLightDistanceFadeMAD)
 	SHADER_PARAMETER(uint32, NumDirectionalLightCascades)
 	SHADER_PARAMETER(int32, DirectionalLightVSM)
