HeterogeneousVolumes: Conditionally store indirect lighting term in lighting cache rather than evaluate live in single-scattering pass.

Folding computation into the lighting cache couples the evaluation frequency to that of the lighting cache, but removes an additional texture-fetch during the single-scattering pass.

r.HeterogeneousVolumes.IndirectLighting.Mode
Changes where indirect is accumulated in the pipeline (Default = 0)
0: Off
1. Lighting cache
2: Single-scattering

diff --git a/Engine/Config/BaseScalability.ini b/Engine/Config/BaseScalability.ini
--- a/Engine/Config/BaseScalability.ini
+++ b/Engine/Config/BaseScalability.ini
@@ -751,6 +751,7 @@ r.HeterogeneousVolumes.DownsampleFactor=2
 r.HeterogeneousVolumes.MaxStepCount=256
 r.HeterogeneousVolumes.Shadows.Resolution=256
 r.HeterogeneousVolumes.Shadows.MaxSampleCount=8
+r.HeterogeneousVolumes.UseExistenceMask=0
 
 [EffectsQuality@Cine]
 r.TranslucencyLightingVolumeDim=64
@@ -783,6 +784,7 @@ r.HeterogeneousVolumes.DownsampleFactor=1
 r.HeterogeneousVolumes.MaxStepCount=1024
 r.HeterogeneousVolumes.Shadows.Resolution=1024
 r.HeterogeneousVolumes.Shadows.MaxSampleCount=16
+r.HeterogeneousVolumes.UseExistenceMask=1
 
 ;-----------------------------------------------------------------------------------------------------------------
 
diff --git a/Engine/Shaders/Private/HeterogeneousVolumes/HeterogeneousVolumesRayMarchingUtils.ush b/Engine/Shaders/Private/HeterogeneousVolumes/HeterogeneousVolumesRayMarchingUtils.ush
--- a/Engine/Shaders/Private/HeterogeneousVolumes/HeterogeneousVolumesRayMarchingUtils.ush
+++ b/Engine/Shaders/Private/HeterogeneousVolumes/HeterogeneousVolumesRayMarchingUtils.ush
@@ -379,7 +379,11 @@ void RayMarchSingleScattering(
 	inout float3 Transmittance
 )
 {
-	LocalScatterTMin = RayMarchingContext.LocalRayTMax;
+	RayMarchingContext.LocalScatterT0 = GetMaxTraceDistance() / RayMarchingContext.LocalToWorldScale;
+	RayMarchingContext.LocalScatterT1 = 0.0;
+	//RayMarchingContext.LocalScatterT0 = RayMarchingContext.LocalRayTMin;
+	//RayMarchingContext.LocalScatterT1 = RayMarchingContext.LocalRayTMax;
+	RayMarchingContext.Transmittance0 = Transmittance;
 #if APPLY_FOG_INSCATTERING != APPLY_FOG_INSCATTERING_OFF
 	float3 FogRadiance = 0;
 	float3 PrevTransmittance = Transmittance;
@@ -393,6 +397,8 @@ void RayMarchSingleScattering(
 	// Reuse jitter as random sample
 	float RandomSample = RayMarchingContext.Jitter;
 
+	bool bFirstScatterEvent = true;
+
 	for (uint StepIndex = 0; StepIndex < StepCount; ++StepIndex)
 	{
 		// Correct the final step size
@@ -437,7 +443,7 @@ void RayMarchSingleScattering(
 			float3 ScatteringCoefficient = Albedo * Extinction;
 
 			Radiance += IndirectInscattering * ScatteringCoefficient * Transmittance * StepSize;
-#endif // DIM_USE_LUMEN_GI
+#endif // INDIRECT_LIGHTING_MODE
 		}
 
 		float3 Extinction = SampleExtinction(SampleContext);
@@ -455,13 +461,13 @@ void RayMarchSingleScattering(
 				float IsotropicPhase = 1.0 / (4.0 * PI);
 				RadianceSample = Inscattering * ScatteringCoefficient * IsotropicPhase * Transmittance * StepSize;
 			}
-
+#if 0
 			// Stochastically decide to record scatter position
 			if (any(Transmittance < RandomSample))
 			{
 				LocalScatterTMin = min(LocalHitT, LocalScatterTMin);
 			}
-
+#endif
 #if APPLY_FOG_INSCATTERING == APPLY_FOG_INSCATTERING_REFERENCE
 			{
 				float WorldRayT = LocalHitT * RayMarchingContext.LocalToWorldScale;
@@ -503,6 +509,16 @@ void RayMarchSingleScattering(
 #else
 		// Accumulate transmittance for the next evaluation
 		Transmittance *= exp(-Extinction * StepSize);
+		if (any(Extinction > 0.0))
+		{
+			RayMarchingContext.LocalScatterT1 = LocalHitT;
+			if (bFirstScatterEvent)
+			{
+				RayMarchingContext.LocalScatterT0 = LocalScatterTMin = LocalHitT;
+				RayMarchingContext.Transmittance0 = Transmittance;
+				bFirstScatterEvent = false;
+			}
+		}
 #endif
 
 		float Epsilon = 1.0e-7;
@@ -590,6 +606,25 @@ void RayMarchSingleScatteringReferenceFastPath(
 		{
 			float3 Emission = SampleEmission(SampleContext);
 			Radiance += Emission * StepSize * Transmittance;
+
+#if DIM_USE_LUMEN_GI
+			FTwoBandSHVectorRGB TranslucencyGISH = GetTranslucencyGIVolumeLighting(DFPromote(WorldPosition), PrimaryView.WorldToClip, true);
+
+			FTwoBandSHVector RotatedHGZonalHarmonic;
+			float3 CameraVector = normalize(RayMarchingContext.WorldRayDirection);
+			float PhaseG = 0.0;
+			RotatedHGZonalHarmonic.V = float4(1.0f, CameraVector.y, CameraVector.z, CameraVector.x) * float4(1.0f, PhaseG, PhaseG, PhaseG);
+			float3 IndirectInscattering = max(DotSH(TranslucencyGISH, RotatedHGZonalHarmonic), 0) / PI;
+			IndirectInscattering *= GetIndirectInscatteringFactor();
+
+			float3 Extinction = SampleExtinction(SampleContext);
+			float3 Albedo = SampleAlbedo(SampleContext);
+			float3 ScatteringCoefficient = Albedo * Extinction;
+			float AmbientOcclusion = EvalAmbientOcclusion(WorldPosition);
+
+			// Note: Phase is accounted for in spherical hamonic calculation
+			Radiance += IndirectInscattering * AmbientOcclusion * ScatteringCoefficient * Transmittance * StepSize;
+#endif // DIM_USE_LUMEN_GI
 		}
 
 		// Radiance from in-scattering
