Fix shader compiling error

diff --git a/Engine/Shaders/Private/HeterogeneousVolumes/HeterogeneousVolumesVoxelGridShadows.usf b/Engine/Shaders/Private/HeterogeneousVolumes/HeterogeneousVolumesVoxelGridShadows.usf
--- a/Engine/Shaders/Private/HeterogeneousVolumes/HeterogeneousVolumesVoxelGridShadows.usf
+++ b/Engine/Shaders/Private/HeterogeneousVolumes/HeterogeneousVolumesVoxelGridShadows.usf
@@ -139,6 +139,11 @@ int GetStochasticFilteringMode()
 	return StochasticFilteringMode;
 }
 
+float GetIndirectInscatteringFactor()
+{
+	return 1.0;
+}
+
 uint3 GetVolumeResolution()
 {
 	return OrthoGridUniformBuffer.TopLevelGridResolution;
