Fixing shader compilation error

diff --git a/Engine/Shaders/Private/HeterogeneousVolumes/HeterogeneousVolumesRayMarchingUtils.ush b/Engine/Shaders/Private/HeterogeneousVolumes/HeterogeneousVolumesRayMarchingUtils.ush
--- a/Engine/Shaders/Private/HeterogeneousVolumes/HeterogeneousVolumesRayMarchingUtils.ush
+++ b/Engine/Shaders/Private/HeterogeneousVolumes/HeterogeneousVolumesRayMarchingUtils.ush
@@ -379,11 +379,7 @@ void RayMarchSingleScattering(
 	inout float3 Transmittance
 )
 {
-	RayMarchingContext.LocalScatterT0 = GetMaxTraceDistance() / RayMarchingContext.LocalToWorldScale;
-	RayMarchingContext.LocalScatterT1 = 0.0;
-	//RayMarchingContext.LocalScatterT0 = RayMarchingContext.LocalRayTMin;
-	//RayMarchingContext.LocalScatterT1 = RayMarchingContext.LocalRayTMax;
-	RayMarchingContext.Transmittance0 = Transmittance;
+	LocalScatterTMin = RayMarchingContext.LocalRayTMax;
 #if APPLY_FOG_INSCATTERING != APPLY_FOG_INSCATTERING_OFF
 	float3 FogRadiance = 0;
 	float3 PrevTransmittance = Transmittance;
@@ -397,8 +393,6 @@ void RayMarchSingleScattering(
 	// Reuse jitter as random sample
 	float RandomSample = RayMarchingContext.Jitter;
 
-	bool bFirstScatterEvent = true;
-
 	for (uint StepIndex = 0; StepIndex < StepCount; ++StepIndex)
 	{
 		// Correct the final step size
@@ -461,13 +455,13 @@ void RayMarchSingleScattering(
 				float IsotropicPhase = 1.0 / (4.0 * PI);
 				RadianceSample = Inscattering * ScatteringCoefficient * IsotropicPhase * Transmittance * StepSize;
 			}
-#if 0
+
 			// Stochastically decide to record scatter position
 			if (any(Transmittance < RandomSample))
 			{
 				LocalScatterTMin = min(LocalHitT, LocalScatterTMin);
 			}
-#endif
+
 #if APPLY_FOG_INSCATTERING == APPLY_FOG_INSCATTERING_REFERENCE
 			{
 				float WorldRayT = LocalHitT * RayMarchingContext.LocalToWorldScale;
@@ -509,16 +503,6 @@ void RayMarchSingleScattering(
 #else
 		// Accumulate transmittance for the next evaluation
 		Transmittance *= exp(-Extinction * StepSize);
-		if (any(Extinction > 0.0))
-		{
-			RayMarchingContext.LocalScatterT1 = LocalHitT;
-			if (bFirstScatterEvent)
-			{
-				RayMarchingContext.LocalScatterT0 = LocalScatterTMin = LocalHitT;
-				RayMarchingContext.Transmittance0 = Transmittance;
-				bFirstScatterEvent = false;
-			}
-		}
 #endif
 
 		float Epsilon = 1.0e-7;
diff --git a/Engine/Shaders/Private/HeterogeneousVolumes/HeterogeneousVolumesVoxelGridShadows.usf b/Engine/Shaders/Private/HeterogeneousVolumes/HeterogeneousVolumesVoxelGridShadows.usf
--- a/Engine/Shaders/Private/HeterogeneousVolumes/HeterogeneousVolumesVoxelGridShadows.usf
+++ b/Engine/Shaders/Private/HeterogeneousVolumes/HeterogeneousVolumesVoxelGridShadows.usf
@@ -139,6 +139,11 @@ int GetStochasticFilteringMode()
 	return StochasticFilteringMode;
 }
 
+float EvalAmbientOcclusion(float3 WorldPosition)
+{
+	return 1.0;
+}
+
 float GetIndirectInscatteringFactor()
 {
 	return 1.0;
