Fix hair strands voxelization having incorrect voxel size due to missing quantization.

Due to this missing quantization (which happens at packing time), the hair voxel structure is offset causing incorrect traversing during shadow/transmittance computation.

diff --git a/Engine/Shaders/Private/HairStrands/HairStrandsVoxelDebug.usf b/Engine/Shaders/Private/HairStrands/HairStrandsVoxelDebug.usf
--- a/Engine/Shaders/Private/HairStrands/HairStrandsVoxelDebug.usf
+++ b/Engine/Shaders/Private/HairStrands/HairStrandsVoxelDebug.usf
@@ -317,6 +317,8 @@ void MainCS(uint3 DispatchThreadId : SV_DispatchThreadID)
 	}
 	else if (bDrawVoxel)
 	{
+		const bool bIsPixelEnable = all(int2(PixelCoord) == GetCursorPos());
+
 		FHairTraversalResult TraversalResult = (FHairTraversalResult)0;
 
 		const float DistanceThreshold = 10000;
@@ -346,6 +348,7 @@ void MainCS(uint3 DispatchThreadId : SV_DispatchThreadID)
 			TraversalSettings.bIsPrimaryRay		= true;
 			TraversalSettings.PixelRadius		= -1;
 			TraversalSettings.ForcedMip			= ForcedMipLevel;
+			TraversalSettings.bDebugEnabled		= false; //bIsPixelEnable;
 
 			TraversalResult = ComputeHairCountVirtualVoxel(
 				TranslatedWP0,
@@ -360,7 +363,6 @@ void MainCS(uint3 DispatchThreadId : SV_DispatchThreadID)
 			bIsValid = TraversalResult.HairCount > 0;
 		}
 
-		const bool bIsPixelEnable = all(int2(PixelCoord) == GetCursorPos());
 		if (bIsPixelEnable && bIsValid)
 		{
 			FShaderPrintContext Context = InitShaderPrintContext(true, uint2(50, 50));
diff --git a/Engine/Shaders/Private/HairStrands/HairStrandsVoxelPageCommon.ush b/Engine/Shaders/Private/HairStrands/HairStrandsVoxelPageCommon.ush
--- a/Engine/Shaders/Private/HairStrands/HairStrandsVoxelPageCommon.ush
+++ b/Engine/Shaders/Private/HairStrands/HairStrandsVoxelPageCommon.ush
@@ -85,6 +85,22 @@ struct FVirtualVoxelNodeDesc
 	bool    bIsValid;
 };
 
+uint PackVoxelWorldSize(float In)
+{
+	// Encode voxel world size [0..10] onto 10bits, i.e.: 10.f / 1023u
+	const float NormalizedVoxelWorldSize = min(In, 10.f) * 0.1f;
+	return PackUnorm10(NormalizedVoxelWorldSize);
+}
+float UnpackVoxelWorldSize(uint In)
+{
+	return UnpackUnorm10(In) * 10.f;
+}
+
+float QuantizeVoxelWorldSize(float In)
+{
+	return UnpackVoxelWorldSize(PackVoxelWorldSize(In));
+}
+
 FVirtualVoxelNodeDesc UnpackVoxelNode(FPackedVirtualVoxelNodeDesc In, uint InPageResolution)
 {
 	FVirtualVoxelNodeDesc Out;
@@ -93,16 +109,13 @@ FVirtualVoxelNodeDesc UnpackVoxelNode(FPackedVirtualVoxelNodeDesc In, uint InPag
 	Out.PageIndexResolution = uint3(In.PackedPageIndexResolution & 0xFF, (In.PackedPageIndexResolution >> 8) & 0xFF, (In.PackedPageIndexResolution >> 16) & 0xFF);
 	Out.VirtualResolution = Out.PageIndexResolution * InPageResolution;
 	Out.PageIndexOffset = (In.PageIndexOffset_VoxelWorldSize & 0x3FFFFFu);
-	Out.VoxelWorldSize = ((In.PageIndexOffset_VoxelWorldSize >> 22u) & 0x3FFu) * 0.00977517106f; // Encode voxel world size [0..10] onto 10bits, i.e.: 10.f / 1023u
+	Out.VoxelWorldSize = UnpackVoxelWorldSize(In.PageIndexOffset_VoxelWorldSize >> 22u);
 	Out.bIsValid = all(Out.PageIndexResolution != 0);
 	return Out;
 }
 
 FPackedVirtualVoxelNodeDesc PackVoxelNode(FVirtualVoxelNodeDesc In)
 {
-	// Encode voxel world size [0..10] onto 10bits, i.e.: 10.f / 1023u
-	const float NormalizedVoxelWorldSize = min(In.VoxelWorldSize, 10.f) * 0.1f;
-
 	FPackedVirtualVoxelNodeDesc Out;
 	Out.TranslatedWorldMinAABB = In.TranslatedWorldMinAABB;
 	Out.TranslatedWorldMaxAABB = In.TranslatedWorldMaxAABB;
@@ -112,7 +125,7 @@ FPackedVirtualVoxelNodeDesc PackVoxelNode(FVirtualVoxelNodeDesc In)
 		((In.PageIndexResolution.z & 0xFF)<<16);
 	Out.PageIndexOffset_VoxelWorldSize =
 		(In.PageIndexOffset & 0x3FFFFFu) |
-		(uint(NormalizedVoxelWorldSize * 1023.f) << 22u);
+		(PackVoxelWorldSize(In.VoxelWorldSize) << 22u);
 	return Out;
 }
 
diff --git a/Engine/Shaders/Private/HairStrands/HairStrandsVoxelPageTraversal.ush b/Engine/Shaders/Private/HairStrands/HairStrandsVoxelPageTraversal.ush
--- a/Engine/Shaders/Private/HairStrands/HairStrandsVoxelPageTraversal.ush
+++ b/Engine/Shaders/Private/HairStrands/HairStrandsVoxelPageTraversal.ush
@@ -187,11 +187,16 @@ FHairTraversalResult ComputeHairCountVirtualVoxel(
 		const float OELength = min(length(E - O), InSettings.DistanceThreshold);
 		const float3 PagesBoundMax = InNodeDesc.TranslatedWorldMinAABB + InNodeDesc.PageIndexResolution * InNodeDesc.VoxelWorldSize * InCommonDesc.PageResolution;
 
+		#if VOXEL_TRAVERSAL_DEBUG
+		FShaderPrintContext Ctx = InitShaderPrintContext(InSettings.bDebugEnabled, uint2(750, 50));
+		#endif
+
 		#if VOXEL_TRAVERSAL_DEBUG
 		if (InSettings.bDebugEnabled)
 		{
 			AddLine(true, O, E);
 			AddAABBVolume(true, InNodeDesc.TranslatedWorldMinAABB, InNodeDesc.TranslatedWorldMaxAABB);
+			AddAABBTWS(InNodeDesc.TranslatedWorldMinAABB, PagesBoundMax, ColorPurple);
 		}
 		#endif
 
@@ -231,6 +236,15 @@ FHairTraversalResult ComputeHairCountVirtualVoxel(
 				}
 			}
 
+			#if VOXEL_TRAVERSAL_DEBUG
+			if (InSettings.bDebugEnabled && bIsPageValid)
+			{
+				const float3 PageAABBMin = (PageIndexCoord  ) * InCommonDesc.PageResolution * InNodeDesc.VoxelWorldSize + InNodeDesc.TranslatedWorldMinAABB;
+				const float3 PageAABBMax = (PageIndexCoord+1) * InCommonDesc.PageResolution * InNodeDesc.VoxelWorldSize + InNodeDesc.TranslatedWorldMinAABB;
+				AddAABBTWS(Ctx, PageAABBMin, PageAABBMax, ColorRed);
+			}
+			#endif
+
 			if (bIsPageValid && bIsWithinPageBound)
 			{
 				const uint3 VoxelPageBase   = (PageCoord << InCommonDesc.PageResolutionLog2);
diff --git a/Engine/Shaders/Private/HairStrands/HairStrandsVoxelRasterCompute.usf b/Engine/Shaders/Private/HairStrands/HairStrandsVoxelRasterCompute.usf
--- a/Engine/Shaders/Private/HairStrands/HairStrandsVoxelRasterCompute.usf
+++ b/Engine/Shaders/Private/HairStrands/HairStrandsVoxelRasterCompute.usf
@@ -77,7 +77,7 @@ void AllocatePageIndex(uint2 DispatchThreadId : SV_DispatchThreadID)
 			Bound = ReadHairAABB(MacroGroupId, MacroGroupAABBBuffer);
 		}
 
-		const float VoxelWorldSize = bUseCPUVoxelWorldSize ? CPUVoxelWorldSize : max(GPUVoxelWorldSize[0], MacroGroupVoxelSizeBuffer[MacroGroupId]);
+		const float VoxelWorldSize = QuantizeVoxelWorldSize(bUseCPUVoxelWorldSize ? CPUVoxelWorldSize : max(GPUVoxelWorldSize[0], MacroGroupVoxelSizeBuffer[MacroGroupId]));
 		PageWorldSize = VoxelWorldSize * PageResolution;
 
 		if (any(Bound.Min > Bound.Max) || !bSupportVoxelization)
