Fix error in LightData.ush

diff --git a/Engine/Shaders/Private/LightData.ush b/Engine/Shaders/Private/LightData.ush
--- a/Engine/Shaders/Private/LightData.ush
+++ b/Engine/Shaders/Private/LightData.ush
@@ -262,11 +262,16 @@ float3 UnpackLightColor(FLocalLightData In)
 	return UnpackLightColor(In.Internal);
 }
 
-float UnpackLightFalloffExponent(FLocalLightData In)
+float UnpackLightFalloffExponent(FForwardLightData In)
 {
 	return In.LightColorAndIdAndFalloffExponent.w;
 }
 
+float UnpackLightFalloffExponent(FLocalLightData In)
+{
+	return UnpackLightFalloffExponent(In.Internal);
+}
+
 float4 UnpackShadowMapChannelMask(uint In)
 {
 	// 8 ALU
