Fixed an issue where VRS could shade using a pixel position that was outside the rasterized footprint.
It now places the sample point on the first live pixel instead as a sort of "centroid" sampling.
Reenabled VRS on Terrain.

#tests PIE

#changelist validated
#virtualized

diff --git a/Engine/Shaders/Private/ComputeShaderOutputCommon.ush b/Engine/Shaders/Private/ComputeShaderOutputCommon.ush
--- a/Engine/Shaders/Private/ComputeShaderOutputCommon.ush
+++ b/Engine/Shaders/Private/ComputeShaderOutputCommon.ush
@@ -40,7 +40,7 @@ uint GetDataByteOffset()	{ return GetRootConstant3(); }
 	#endif
 #endif
 
-FPixelShaderOut ShadePixel(const float2 SVPositionXY, const uint2 PixelPos, uint QuadIndex, uint QuadPixelWriteMask)
+FPixelShaderOut ShadePixel(const float2 SVPositionXY, uint QuadIndex, uint QuadPixelWriteMask)
 {
 	// Note: Driven by bHighPrecisionGBuffers in GBufferInfo.cpp
 	const bool bHighPrecision = GetHighPrecision() != 0u;
@@ -58,7 +58,7 @@ FPixelShaderOut ShadePixel(const float2 SVPositionXY, const uint2 PixelPos, uint
 
 #if IS_NANITE_PASS && (PIXELSHADEROUTPUT_INTERPOLANTS || PIXELSHADEROUTPUT_BASEPASS)
 	FVertexFactoryInterpolantsVSToPS Interpolants = (FVertexFactoryInterpolantsVSToPS)0;
-	Interpolants.PixelPos = PixelPos;
+	Interpolants.PixelPos = int2(SVPositionXY);
 	Interpolants.ViewIndex = NaniteInterpolants.ViewIndex; // TODO: NANITE_MATERIAL_MULTIVIEW
 
 #if INSTANCED_STEREO
@@ -188,7 +188,7 @@ void ExportPixel(const uint2 PixelPos, FPixelShaderOut ShadedPixel)
 void ProcessPixel(uint ShadingBin, const uint2 PixelPos, const float2 SVPositionXY, uint QuadIndex, uint DispatchIndex, uint PixelWriteMask, uint HelperLaneCount)
 {
 	// All lanes shade (regardless of export - so ddx/ddy are valid)
-	FPixelShaderOut ShadedPixel = ShadePixel(SVPositionXY, PixelPos, QuadIndex, PixelWriteMask);
+	FPixelShaderOut ShadedPixel = ShadePixel(SVPositionXY, QuadIndex, PixelWriteMask);
 
 #if VIS_HELPER_LANES
 	ShadedPixel.MRT[3].rgb = ColorMapTurbo(float(HelperLaneCount) / 4.0f);
@@ -301,8 +301,13 @@ void MainCS(
 	// works. While ignoring the shift does not exactly match HW VRS,
 	// it doesn't seem super important outside of a direct A/B test.
 	const float2 SVPositionXY = PixelPos + ((1u << VRSShift) * 0.5f);
-#else
+#elif 0
+	// Always shade as top-left. For VRS this could be outside the rasterized triangle.
 	const float2 SVPositionXY = PixelPos + 0.5f;
+#else
+	// Centroid-like sampling. Move the sampling point to the center of the first live pixel.
+	const uint WriteMaskFirstIndex = PixelWriteMask ? firstbitlow(PixelWriteMask) : 0u;	// PixelWriteMask can be 0 in quad mode.
+	const float2 SVPositionXY = PixelPos + int2(WriteMaskFirstIndex & 1, (WriteMaskFirstIndex >> 1) & 1) + 0.5f;
 #endif
 
 #if WORKGRAPH_NODE
@@ -317,4 +322,4 @@ void MainCS(
 #endif
 
 	ProcessPixel(ShadingBin, PixelPos, SVPositionXY, ElementIndex, PixelIndex, PixelWriteMask, HelperLaneCount);
-}
+}
\ No newline at end of file
