Heterogeneous Volumes: Add linear interpolated application of in-scattered fog lighting.

Improves run-time performance of live-shading single-scattering integrator by approx 20%.

diff --git a/Engine/Config/BaseScalability.ini b/Engine/Config/BaseScalability.ini
--- a/Engine/Config/BaseScalability.ini
+++ b/Engine/Config/BaseScalability.ini
@@ -780,6 +780,7 @@ r.SkyAtmosphere.TransmittanceLUT.SampleCount=30.0
 r.SkyAtmosphere.MultiScatteringLUT.SampleCount=20.0
 fx.Niagara.QualityLevel=4
 r.Refraction.OffsetQuality=1
+r.HeterogeneousVolumes.ApplyFogInscattering=1
 r.HeterogeneousVolumes.DownsampleFactor=1
 r.HeterogeneousVolumes.MaxStepCount=1024
 r.HeterogeneousVolumes.Shadows.Resolution=1024
diff --git a/Engine/Shaders/Private/HeterogeneousVolumes/HeterogeneousVolumesRayMarchingUtils.ush b/Engine/Shaders/Private/HeterogeneousVolumes/HeterogeneousVolumesRayMarchingUtils.ush
--- a/Engine/Shaders/Private/HeterogeneousVolumes/HeterogeneousVolumesRayMarchingUtils.ush
+++ b/Engine/Shaders/Private/HeterogeneousVolumes/HeterogeneousVolumesRayMarchingUtils.ush
@@ -30,14 +30,6 @@
 #define DIM_USE_ADAPTIVE_VOLUMETRIC_SHADOW_MAP 0
 #endif // DIM_USE_ADAPTIVE_VOLUMETRIC_SHADOW_MAP
 
-#define APPLY_FOG_INSCATTERING_OFF 0
-#define APPLY_FOG_INSCATTERING_REFERENCE 1
-#define APPLY_FOG_INSCATTERING_STOCHASTIC 2
-
-#ifndef APPLY_FOG_INSCATTERING
-#define APPLY_FOG_INSCATTERING APPLY_FOG_INSCATTERING_OFF
-#endif // APPLY_FOG_INSCATTERING
-
 #ifndef IS_OFFLINE_RENDER
 #define IS_OFFLINE_RENDER 0
 #endif // IS_OFFLINE_RENDER
@@ -47,17 +39,40 @@
 #define INDIRECT_LIGHTING_MODE_SINGLE_SCATTERING 2
 
 #ifndef INDIRECT_LIGHTING_MODE
-#define INDIRECT_LIGHTING_MODE 0
+#define INDIRECT_LIGHTING_MODE INDIRECT_LIGHTING_MODE_OFF
 #endif // INDIRECT_LIGHTING_MODE
 
+#define FOG_INSCATTERING_MODE_OFF 0
+#define FOG_INSCATTERING_MODE_REFERENCE 1
+#define FOG_INSCATTERING_MODE_LINEAR_APPROX 2
+
+#ifndef FOG_INSCATTERING_MODE
+#define FOG_INSCATTERING_MODE FOG_INSCATTERING_MODE_OFF
+#endif // FOG_INSCATTERING_MODE
+
 #if DIM_USE_ADAPTIVE_VOLUMETRIC_SHADOW_MAP
 #include "HeterogeneousVolumesAdaptiveVolumetricShadowMapSampling.ush"
 #include "HeterogeneousVolumesAdaptiveVolumetricShadowMapSampler.ush"
 #endif // DIM_USE_ADAPTIVE_VOLUMETRIC_SHADOW_MAP
 
-#if APPLY_FOG_INSCATTERING
+
+#if FOG_INSCATTERING_MODE
 #include "../HeightFogCommon.ush"
-#endif // APPLY_FOG_INSCATTERING
+#endif // FOG_INSCATTERING_MODE
+
+float3 GetWorldVoxelSizeOfLightingCache()
+{
+	float3 WorldVoxelSize = 0;
+	float3 LightingCacheResolutionAsFloat3 = GetLightingCacheResolution();
+	if (all(LightingCacheResolutionAsFloat3 > 0))
+	{
+		float3 VolumeBounds = 2.0f * GetLocalBoundsExtent();
+		float3 VoxelSize = VolumeBounds * rcp(LightingCacheResolutionAsFloat3);
+		WorldVoxelSize = mul(float4(VoxelSize, 0), GetLocalToWorld()).xyz;
+	}
+
+	return WorldVoxelSize;
+}
 
 float CalcShadowBias()
 {
@@ -379,12 +394,47 @@ void RayMarchSingleScattering(
 	inout float3 Transmittance
 )
 {
-	LocalScatterTMin = RayMarchingContext.LocalRayTMax;
-#if APPLY_FOG_INSCATTERING != APPLY_FOG_INSCATTERING_OFF
+#if FOG_INSCATTERING_MODE
 	float3 FogRadiance = 0;
 	float3 PrevTransmittance = Transmittance;
 	float4 PrevCombinedFog = 0;
-#endif // #if APPLY_FOG_INSCATTERING
+#endif // FOG_INSCATTERING_MODE
+
+#if FOG_INSCATTERING_MODE == FOG_INSCATTERING_MODE_LINEAR_APPROX
+	float WorldFogHitT[] = {
+		RayMarchingContext.LocalRayTMin * RayMarchingContext.LocalToWorldScale,
+		RayMarchingContext.LocalRayTMax * RayMarchingContext.LocalToWorldScale
+	};
+
+	float4 FogHitRadiance[] = {
+		float4(0.0, 0.0, 0.0, 0.0),
+		float4(0.0, 0.0, 0.0, 0.0)
+	};
+
+	UNROLL
+	for (int i = 0; i < 2; ++i)
+	{
+		float3 CameraRelative_WorldPosition = RayMarchingContext.WorldRayDirection * WorldFogHitT[i];
+		float4 HeightFog = float4(0, 0, 0, 1);
+		if (ShouldApplyHeightFog())
+		{
+			HeightFog = CalculateHeightFog(CameraRelative_WorldPosition);
+		}
+
+		float4 CombinedFog = HeightFog;
+		if (ShouldApplyVolumetricFog() && FogStruct.ApplyVolumetricFog > 0)
+		{
+			float3 WorldPosition = CameraRelative_WorldPosition + RayMarchingContext.WorldRayOrigin;
+			const uint EyeIndex = 0;
+			float3 VolumeUV = ComputeVolumeUV_DEPRECATED(WorldPosition, DFHackToFloat(PrimaryView.WorldToClip));
+			CombinedFog = CombineVolumetricFog(HeightFog, VolumeUV, EyeIndex, WorldFogHitT[i]);
+		}
+
+		FogHitRadiance[i] = CombinedFog;
+	}
+#endif // FOG_INSCATTERING_MODE
+
+	LocalScatterTMin = RayMarchingContext.LocalRayTMax;
 
 	float WorldShadowBias = CalcShadowBias();
 	//float WorldShadowBias = View.GeneralPurposeTweak;
@@ -462,7 +512,7 @@ void RayMarchSingleScattering(
 				LocalScatterTMin = min(LocalHitT, LocalScatterTMin);
 			}
 
-#if APPLY_FOG_INSCATTERING == APPLY_FOG_INSCATTERING_REFERENCE
+#if FOG_INSCATTERING_MODE == FOG_INSCATTERING_MODE_REFERENCE
 			{
 				float WorldRayT = LocalHitT * RayMarchingContext.LocalToWorldScale;
 				float3 CameraRelative_WorldPosition = RayMarchingContext.WorldRayDirection * WorldRayT;
@@ -491,7 +541,23 @@ void RayMarchSingleScattering(
 
 				PrevCombinedFog = CombinedFog;
 			}
-#endif // APPLY_FOG_INSCATTERING_REFERENCE
+#elif FOG_INSCATTERING_MODE == FOG_INSCATTERING_MODE_LINEAR_APPROX
+			{
+				float WorldRayT = LocalHitT * RayMarchingContext.LocalToWorldScale;
+				float Weight = saturate((WorldRayT - WorldFogHitT[0]) / (WorldFogHitT[1] - WorldFogHitT[0]));
+				float4 CombinedFog = lerp(FogHitRadiance[0], FogHitRadiance[1], Weight);
+
+				// Attenuate RadianceSample
+				RadianceSample *= CombinedFog.a;
+
+				// Add FogRadiance
+				float3 DeltaTransmittance = clamp(Transmittance / PrevTransmittance, 0.0, 1.0);
+				float3 DeltaFogRadiance = CombinedFog.rgb - PrevCombinedFog.rgb;
+				FogRadiance += DeltaFogRadiance * DeltaTransmittance;
+
+				PrevCombinedFog = CombinedFog;
+			}
+#endif // FOG_INSCATTERING_MODE
 
 			// Accumulate radiance
 			Radiance += RadianceSample;
@@ -515,34 +581,8 @@ void RayMarchSingleScattering(
 		LocalHitT += StepSize;
 	}
 
-#if APPLY_FOG_INSCATTERING == APPLY_FOG_INSCATTERING_STOCHASTIC
-	if (LocalScatterTMin < RayMarchingContext.LocalRayTMax)
-	{
-		float WorldRayT = LocalScatterTMin * RayMarchingContext.LocalToWorldScale;
-
-		float3 CameraRelative_WorldPosition = RayMarchingContext.WorldRayDirection * WorldRayT;
-		float4 HeightFog = float4(0, 0, 0, 1);
-		if (ShouldApplyHeightFog())
-		{
-			HeightFog = CalculateHeightFog(CameraRelative_WorldPosition);
-		}
-
-		float4 CombinedFog = HeightFog;
-		if (ShouldApplyVolumetricFog() && FogStruct.ApplyVolumetricFog > 0)
-		{
-			float3 WorldPosition = CameraRelative_WorldPosition + RayMarchingContext.WorldRayOrigin;
-			const uint EyeIndex = 0;
-			float3 VolumeUV = ComputeVolumeUV_DEPRECATED(WorldPosition, DFHackToFloat(PrimaryView.WorldToClip));
-			CombinedFog = CombineVolumetricFog(HeightFog, VolumeUV, EyeIndex, WorldRayT);
-		}
-
-		// Attenuate Radiance by fog
-		Radiance *= CombinedFog.a;
-		FogRadiance = CombinedFog.rgb;
-	}
-#endif // APPLY_FOG_INSCATTERING
 
-#if APPLY_FOG_INSCATTERING != APPLY_FOG_INSCATTERING_OFF
+#if FOG_INSCATTERING_MODE
 	// Simplifying terms to avoid vector division:
 	// float3 DeltaTransmittance = clamp(Transmittance / PrevTransmittance, 0.0, 1.0);
 	// float3 DeltaOpacity = 1.0 - DeltaTransmittance;
@@ -553,7 +593,7 @@ void RayMarchSingleScattering(
 	// for elements behind the subject. To prevent doubled contribution of in-scattering for 
 	// regions of low density (eg: silhouette edges), blending is interpreted as alpha-masking instead.
 	Radiance += FogRadiance.rgb * FogAlpha;
-#endif // APPLY_FOG_INSCATTERING
+#endif // FOG_INSCATTERING_MODE
 }
 
 void RayMarchSingleScatteringReferenceFastPath(
diff --git a/Engine/Source/Runtime/Renderer/Private/HeterogeneousVolumes/HeterogeneousVolumes.cpp b/Engine/Source/Runtime/Renderer/Private/HeterogeneousVolumes/HeterogeneousVolumes.cpp
--- a/Engine/Source/Runtime/Renderer/Private/HeterogeneousVolumes/HeterogeneousVolumes.cpp
+++ b/Engine/Source/Runtime/Renderer/Private/HeterogeneousVolumes/HeterogeneousVolumes.cpp
@@ -288,12 +288,12 @@ static TAutoConsoleVariable<int32> CVarHeterogeneousVolumesApplyVolumetricFog(
 
 static TAutoConsoleVariable<int32> CVarHeterogeneousVolumesApplyFogInscatteringMode(
 	TEXT("r.HeterogeneousVolumes.ApplyFogInscattering"),
-	1,
-	TEXT("Determines the method for applying fog in-scattering (default = 1)\n")
+	2,
+	TEXT("Determines the method for applying fog in-scattering (default = 2)\n")
 	TEXT("0: Off\n")
-	TEXT("1: Reference (evaluated per ray-march step)\n")
-	TEXT("2: Stochastic (evaluated once per ray)\n"),
-	ECVF_RenderThreadSafe
+	TEXT("1: Reference\n")
+	TEXT("2: Linear approximation\n"),
+	ECVF_Scalability | ECVF_RenderThreadSafe
 );
 
 static TAutoConsoleVariable<int32> CVarHeterogeneousVolumesUseAnalyticDerivatives(
@@ -630,7 +630,7 @@ namespace HeterogeneousVolumes
 		return CVarHeterogeneousVolumesApplyVolumetricFog.GetValueOnRenderThread() != 0;
 	}
 
-	EFogMode GetApplyFogInscattering()
+	EFogMode GetFogInscatteringMode()
 	{
 		return static_cast<EFogMode>(FMath::Clamp(CVarHeterogeneousVolumesApplyFogInscatteringMode.GetValueOnRenderThread(), 0, 2));
 	}
diff --git a/Engine/Source/Runtime/Renderer/Private/HeterogeneousVolumes/HeterogeneousVolumes.h b/Engine/Source/Runtime/Renderer/Private/HeterogeneousVolumes/HeterogeneousVolumes.h
--- a/Engine/Source/Runtime/Renderer/Private/HeterogeneousVolumes/HeterogeneousVolumes.h
+++ b/Engine/Source/Runtime/Renderer/Private/HeterogeneousVolumes/HeterogeneousVolumes.h
@@ -135,11 +135,11 @@ namespace HeterogeneousVolumes
 
 	enum class EFogMode
 	{
-		Off,
+		Disabled,
 		Reference,
-		Stochastic
+		LinearApprox
 	};
-	EFogMode GetApplyFogInscattering();
+	EFogMode GetFogInscatteringMode();
 	bool ShouldWriteVelocity();
 
 	bool EnableIndirectionGrid();
diff --git a/Engine/Source/Runtime/Renderer/Private/HeterogeneousVolumes/HeterogeneousVolumesLiveShadingPipeline.cpp b/Engine/Source/Runtime/Renderer/Private/HeterogeneousVolumes/HeterogeneousVolumesLiveShadingPipeline.cpp
--- a/Engine/Source/Runtime/Renderer/Private/HeterogeneousVolumes/HeterogeneousVolumesLiveShadingPipeline.cpp
+++ b/Engine/Source/Runtime/Renderer/Private/HeterogeneousVolumes/HeterogeneousVolumesLiveShadingPipeline.cpp
@@ -429,14 +429,14 @@ class FRenderSingleScatteringWithLiveShadingCS : public FMeshMaterialShader
 	//class FAdaptiveMarch : SHADER_PERMUTATION_INT("ADAPTIVE_MARCH", 2);
 	//class FBilinearInterpolation : SHADER_PERMUTATION_INT("AVSM_BILINEAR_INTERPOLATION", 2);
 	class FIsOfflineRender : SHADER_PERMUTATION_INT("IS_OFFLINE_RENDER", 2);
-	class FApplyFogInscattering : SHADER_PERMUTATION_INT("APPLY_FOG_INSCATTERING", 3);
+	class FFogInscatteringMode : SHADER_PERMUTATION_INT("FOG_INSCATTERING_MODE", 3);
 	class FUseAnalyticDerivatives : SHADER_PERMUTATION_BOOL("USE_ANALYTIC_DERIVATIVES");
 	using FPermutationDomain = TShaderPermutationDomain<
 		//FUseTransmittanceVolume, 
 		FUseInscatteringVolume, 
 		FIndirectLightingMode,
 		FWriteVelocity, 
-		FApplyFogInscattering, 
+		FFogInscatteringMode,
 		FAVSMSampleMode, 
 		FSupportOverlappingVolumes,
 		FIsOfflineRender,
@@ -553,12 +553,6 @@ class FRenderSingleScatteringWithLiveShadingCS : public FMeshMaterialShader
 
 	static FPermutationDomain RemapPermutation(FPermutationDomain PermutationVector)
 	{
-		// Remap Off to Stochastic and turn off individual fog modes
-		if (PermutationVector.template Get<FApplyFogInscattering>() == 0)
-		{
-			PermutationVector.template Set<FApplyFogInscattering>(2);
-		}
-
 		if (PermutationVector.template Get<FIndirectLightingMode>() != static_cast<int32>(HeterogeneousVolumes::EIndirectLightingMode::SingleScatteringPass))
 		{
 			PermutationVector.template Set<FIndirectLightingMode>(static_cast<int32>(HeterogeneousVolumes::EIndirectLightingMode::Disabled));
@@ -1243,14 +1237,8 @@ void RenderSingleScatteringWithLiveShading(
 
 		TRDGUniformBufferRef<FFogUniformParameters> FogBuffer = CreateFogUniformBuffer(GraphBuilder, View);
 		PassParameters->FogStruct = FogBuffer;
-		PassParameters->bApplyHeightFog = 0;
-		PassParameters->bApplyVolumetricFog = 0;
-		if (bApplyEmissionAndTransmittance &&
-			(HeterogeneousVolumes::GetApplyFogInscattering() != HeterogeneousVolumes::EFogMode::Off))
-		{
-			PassParameters->bApplyHeightFog = HeterogeneousVolumes::ShouldApplyHeightFog();
-			PassParameters->bApplyVolumetricFog = HeterogeneousVolumes::ShouldApplyVolumetricFog();
-		}
+		PassParameters->bApplyHeightFog = HeterogeneousVolumes::ShouldApplyHeightFog();
+		PassParameters->bApplyVolumetricFog = HeterogeneousVolumes::ShouldApplyVolumetricFog();
 
 		// Indirect lighting data
 		PassParameters->IndirectInscatteringFactor = HeterogeneousVolumes::GetIndirectLightingFactor();
@@ -1330,7 +1318,7 @@ void RenderSingleScatteringWithLiveShading(
 	//PermutationVector.template Set<typename FRenderSingleScatteringWithLiveShadingDispatchTypeCS::FBilinearInterpolation>(static_cast<int32>(HeterogeneousVolumes::UseBilinearInterpolation()));
 	//PermutationVector.template Set<typename FRenderSingleScatteringWithLiveShadingDispatchTypeCS::FAdaptiveMarch>(static_cast<int32>(HeterogeneousVolumes::ShouldAdaptiveMarch()));
 	PermutationVector.template Set<typename FRenderSingleScatteringWithLiveShadingDispatchTypeCS::FIsOfflineRender>(View.bIsOfflineRender);
-	PermutationVector.template Set<typename FRenderSingleScatteringWithLiveShadingDispatchTypeCS::FApplyFogInscattering>(static_cast<int32>(HeterogeneousVolumes::GetApplyFogInscattering()));
+	PermutationVector.template Set<typename FRenderSingleScatteringWithLiveShadingDispatchTypeCS::FFogInscatteringMode>(static_cast<int32>(HeterogeneousVolumes::GetFogInscatteringMode()));
 	PermutationVector.template Set<typename FRenderSingleScatteringWithLiveShadingDispatchTypeCS::FUseAnalyticDerivatives>(HeterogeneousVolumes::UseAnalyticDerivatives());
 	PermutationVector = FRenderSingleScatteringWithLiveShadingDispatchTypeCS::RemapPermutation(PermutationVector);
 	TShaderRef<FRenderSingleScatteringWithLiveShadingDispatchTypeCS> ComputeShader = Material.GetShader<FRenderSingleScatteringWithLiveShadingDispatchTypeCS>(&FLocalVertexFactory::StaticType, PermutationVector, false);
diff --git a/Engine/Source/Runtime/Renderer/Private/HeterogeneousVolumes/HeterogeneousVolumesPreshadingPipeline.cpp b/Engine/Source/Runtime/Renderer/Private/HeterogeneousVolumes/HeterogeneousVolumesPreshadingPipeline.cpp
--- a/Engine/Source/Runtime/Renderer/Private/HeterogeneousVolumes/HeterogeneousVolumesPreshadingPipeline.cpp
+++ b/Engine/Source/Runtime/Renderer/Private/HeterogeneousVolumes/HeterogeneousVolumesPreshadingPipeline.cpp
@@ -318,12 +318,13 @@ class FRenderSingleScatteringWithPreshadingCS : public FGlobalShader
 	class FVoxelCullingDim : SHADER_PERMUTATION_BOOL("DIM_VOXEL_CULLING");
 	class FSparseVoxelTracingDim : SHADER_PERMUTATION_BOOL("DIM_SPARSE_VOXEL_TRACING");
 	//class FUseTransmittanceVolume : SHADER_PERMUTATION_BOOL("DIM_USE_TRANSMITTANCE_VOLUME");
+	class FFogInscatteringMode : SHADER_PERMUTATION_INT("FOG_INSCATTERING_MODE", 3);
 	class FUseInscatteringVolume : SHADER_PERMUTATION_BOOL("DIM_USE_INSCATTERING_VOLUME");
 	class FIndirectLightingMode : SHADER_PERMUTATION_INT("INDIRECT_LIGHTING_MODE", 3);
 	class FWriteVelocity : SHADER_PERMUTATION_BOOL("DIM_WRITE_VELOCITY");
 	class FUseAdaptiveVolumetricShadowMap : SHADER_PERMUTATION_BOOL("DIM_USE_ADAPTIVE_VOLUMETRIC_SHADOW_MAP");
 	class FDebugDim : SHADER_PERMUTATION_BOOL("DIM_DEBUG");
-	using FPermutationDomain = TShaderPermutationDomain<FApplyShadowTransmittanceDim, FVoxelCullingDim, FSparseVoxelTracingDim, FUseInscatteringVolume, FIndirectLightingMode, FWriteVelocity, FUseAdaptiveVolumetricShadowMap, FDebugDim>;
+	using FPermutationDomain = TShaderPermutationDomain<FApplyShadowTransmittanceDim, FVoxelCullingDim, FSparseVoxelTracingDim, FFogInscatteringMode, FUseInscatteringVolume, FIndirectLightingMode, FWriteVelocity, FUseAdaptiveVolumetricShadowMap, FDebugDim>;
 
 	BEGIN_SHADER_PARAMETER_STRUCT(FParameters, )
 		// Scene data
@@ -433,7 +434,6 @@ class FRenderSingleScatteringWithPreshadingCS : public FGlobalShader
 		FMaterialShader::ModifyCompilationEnvironment(Parameters, OutEnvironment);
 		OutEnvironment.SetDefine(TEXT("THREADGROUP_SIZE_1D"), GetThreadGroupSize1D());
 		OutEnvironment.SetDefine(TEXT("THREADGROUP_SIZE_2D"), GetThreadGroupSize2D());
-		OutEnvironment.SetDefine("APPLY_FOG_INSCATTERING", 1);
 
 		bool bSupportVirtualShadowMap = IsFeatureLevelSupported(Parameters.Platform, ERHIFeatureLevel::SM5);
 		if (bSupportVirtualShadowMap)
@@ -741,6 +741,7 @@ void RenderSingleScatteringWithPreshadingCompute(
 	PermutationVector.Set<FRenderSingleScatteringWithPreshadingCS::FVoxelCullingDim>(HeterogeneousVolumes::UseSparseVoxelPerTileCulling());
 	PermutationVector.Set<FRenderSingleScatteringWithPreshadingCS::FSparseVoxelTracingDim>(HeterogeneousVolumes::UseSparseVoxelPipeline());
 	//PermutationVector.Set<FRenderSingleScatteringWithPreshadingCS::FUseTransmittanceVolume>(HeterogeneousVolumes::UseLightingCacheForTransmittance());
+	PermutationVector.Set<FRenderSingleScatteringWithPreshadingCS::FFogInscatteringMode>(static_cast<int32>(HeterogeneousVolumes::GetFogInscatteringMode()));
 	PermutationVector.Set<FRenderSingleScatteringWithPreshadingCS::FUseInscatteringVolume>(HeterogeneousVolumes::UseLightingCacheForInscattering());
 	PermutationVector.Set<FRenderSingleScatteringWithPreshadingCS::FIndirectLightingMode>(IndirectLightingMode);
 	PermutationVector.Set<FRenderSingleScatteringWithPreshadingCS::FWriteVelocity>(bWriteVelocity);
