Add support for BC7 compression with Alpha.

diff --git a/Engine/Shaders/Private/BCCompressionCommon.ush b/Engine/Shaders/Private/BCCompressionCommon.ush
--- a/Engine/Shaders/Private/BCCompressionCommon.ush
+++ b/Engine/Shaders/Private/BCCompressionCommon.ush
@@ -706,4 +706,141 @@ uint4 CompressBC7Block(in float3 BlockRGB[16])
 	}
 
 	return Block;
-}
\ No newline at end of file
+}
+
+uint4 Quantize7A(float4 X)
+{
+	return (uint4(X * 0xFF)) >> 1;
+}
+
+uint ComputeBC7AIndex(float4 Color, float4 BlockVector, float EndPoint0Pos, float EndPoint1Pos)
+{
+	float Pos = dot(Color, BlockVector);
+	float NormalizedPos = saturate((Pos - EndPoint0Pos) / (EndPoint1Pos - EndPoint0Pos));
+	return (uint)clamp(NormalizedPos * 14.93333f + 0.03333f + 0.5f, 0.0f, 15.0f);
+}
+
+// Least squares optimization to find best endpoints for the selected block indices in a mode 6 BC7 block
+void OptimizeEndpointsBC7A(float4 Texels[16], inout float4 BlockMin, inout float4 BlockMax)
+{
+	float4 BlockVector = BlockMax - BlockMin;
+
+	float EndPoint0Pos = dot(BlockMin, BlockVector);
+	float EndPoint1Pos = dot(BlockMax, BlockVector);
+
+	float4 AlphaTexelSum = 0.0f;
+	float4 BetaTexelSum = 0.0f;
+	float AlphaBetaSum = 0.0f;
+	float AlphaSqSum = 0.0f;
+	float BetaSqSum = 0.0f;
+
+	for (uint TexelIndex = 0; TexelIndex < 16; ++TexelIndex)
+	{
+		uint Index = ComputeBC7AIndex(Texels[TexelIndex], BlockVector, EndPoint0Pos, EndPoint1Pos);
+
+		float Beta = saturate(Index / 15.0f);
+		float Alpha = 1.0f - Beta;
+
+		AlphaTexelSum += Alpha * Texels[TexelIndex];
+		BetaTexelSum += Beta * Texels[TexelIndex];
+
+		AlphaBetaSum += Alpha * Beta;
+
+		AlphaSqSum += Alpha * Alpha;
+		BetaSqSum += Beta * Beta;
+	}
+
+	float Det = AlphaSqSum * BetaSqSum - AlphaBetaSum * AlphaBetaSum;
+
+	if (abs(Det) > 0.1f)
+	{
+		float RcpDet = rcp(Det);
+		BlockMin = saturate(RcpDet * (AlphaTexelSum * BetaSqSum - BetaTexelSum * AlphaBetaSum));
+		BlockMax = saturate(RcpDet * (BetaTexelSum * AlphaSqSum - AlphaTexelSum * AlphaBetaSum));
+	}
+}
+
+// Compress a BC7 color only block. Evaluates only mode 6 for performance
+uint4 CompressBC7ABlock(in float4 BlockRGBA[16])
+{
+	// Compute initial endpoints
+	float4 BlockMin = BlockRGBA[0];
+	float4 BlockMax = BlockRGBA[0];
+	{
+		for (uint TexelIndex = 1; TexelIndex < 16; ++TexelIndex)
+		{
+			BlockMin = min(BlockMin, BlockRGBA[TexelIndex]);
+			BlockMax = max(BlockMax, BlockRGBA[TexelIndex]);
+		}
+	}
+
+	#if LEAST_SQUARES_ENDPOINT_OPTIMIZATION
+	{
+		OptimizeEndpointsBC7A(BlockRGBA, BlockMin, BlockMax);
+	}
+	#endif
+
+	float4 BlockVector = BlockMax - BlockMin;
+	uint4 Endpoint0 = Quantize7A(BlockMin);
+	uint4 Endpoint1 = Quantize7A(BlockMax);
+	float EndPoint0Pos = dot(BlockMin, BlockVector);
+	float EndPoint1Pos = dot(BlockMax, BlockVector);
+
+	// Check if endpoint swap is required
+	uint FixupIndex = ComputeBC7AIndex(BlockRGBA[0], BlockVector, EndPoint0Pos, EndPoint1Pos);
+	if (FixupIndex > 7)
+	{
+		Swap(EndPoint0Pos, EndPoint1Pos);
+		Swap(Endpoint0, Endpoint1);
+	}
+
+	// Compute indices
+	uint Indices[16] = { 0, 0, 0, 0, 0, 0, 0, 0, 0, 0, 0, 0, 0, 0, 0, 0 };
+	for (uint TexelIndex = 0; TexelIndex < 16; ++TexelIndex)
+	{
+		Indices[TexelIndex] = ComputeBC7AIndex(BlockRGBA[TexelIndex], BlockVector, EndPoint0Pos, EndPoint1Pos);
+	}
+
+	uint4 Block = 0;
+	Block.x = 0;
+	Block.y = 0;
+	Block.z = 0;
+	Block.w = 0;
+	{
+		// Encode mode 6 block
+		Block.x = 0x40;
+
+		// Encode endpoints
+		Block.x |= Endpoint0.x << 7;
+		Block.x |= Endpoint1.x << 14;
+		Block.x |= Endpoint0.y << 21;
+		Block.x |= Endpoint1.y << 28;
+		Block.y |= Endpoint1.y >> 4;
+		Block.y |= Endpoint0.z << 3;
+		Block.y |= Endpoint1.z << 10;
+		Block.y |= Endpoint0.w << 17;
+		Block.y |= Endpoint1.w << 24;
+
+		// Encode endpoint p-bit
+
+		// Encode indices
+		Block.z |= Indices[0] << 1;
+		Block.z |= Indices[1] << 4;
+		Block.z |= Indices[2] << 8;
+		Block.z |= Indices[3] << 12;
+		Block.z |= Indices[4] << 16;
+		Block.z |= Indices[5] << 20;
+		Block.z |= Indices[6] << 24;
+		Block.z |= Indices[7] << 28;
+		Block.w |= Indices[8] << 0;
+		Block.w |= Indices[9] << 4;
+		Block.w |= Indices[10] << 8;
+		Block.w |= Indices[11] << 12;
+		Block.w |= Indices[12] << 16;
+		Block.w |= Indices[13] << 20;
+		Block.w |= Indices[14] << 24;
+		Block.w |= Indices[15] << 28;
+	}
+
+	return Block;
+}
diff --git a/Engine/Shaders/Private/BlockCompressionCommon.ush b/Engine/Shaders/Private/BlockCompressionCommon.ush
--- a/Engine/Shaders/Private/BlockCompressionCommon.ush
+++ b/Engine/Shaders/Private/BlockCompressionCommon.ush
@@ -87,6 +87,55 @@ void ReadBlockAlpha(Texture2D<float4> SourceTexture, SamplerState TextureSampler
 	}
 }
 
+void ReadBlockRGBA(Texture2D<float4> SourceTexture, SamplerState TextureSampler, float2 UV, float2 TexelUVSize, out float4 Block[16])
+{
+	{
+		float4 Red = SourceTexture.GatherRed(TextureSampler, UV);
+		float4 Green = SourceTexture.GatherGreen(TextureSampler, UV);
+		float4 Blue = SourceTexture.GatherBlue(TextureSampler, UV);
+		float4 Alpha = SourceTexture.GatherAlpha(TextureSampler, UV);
+
+		Block[0] = float4(Red[3], Green[3], Blue[3], Alpha[3]);
+		Block[1] = float4(Red[2], Green[2], Blue[2], Alpha[2]);
+		Block[4] = float4(Red[0], Green[0], Blue[0], Alpha[0]);
+		Block[5] = float4(Red[1], Green[1], Blue[1], Alpha[1]);
+	}
+	{
+		float2 UVOffset = UV + float2(2.f * TexelUVSize.x, 0);
+		float4 Red = SourceTexture.GatherRed(TextureSampler, UVOffset);
+		float4 Green = SourceTexture.GatherGreen(TextureSampler, UVOffset);
+		float4 Blue = SourceTexture.GatherBlue(TextureSampler, UVOffset);
+		float4 Alpha = SourceTexture.GatherAlpha(TextureSampler, UVOffset);
+
+		Block[2] = float4(Red[3], Green[3], Blue[3], Alpha[3]);
+		Block[3] = float4(Red[2], Green[2], Blue[2], Alpha[2]);
+		Block[6] = float4(Red[0], Green[0], Blue[0], Alpha[0]);
+		Block[7] = float4(Red[1], Green[1], Blue[1], Alpha[1]);
+	}
+	{
+		float2 UVOffset = UV + float2(0, 2.f * TexelUVSize.y);
+		float4 Red = SourceTexture.GatherRed(TextureSampler, UVOffset);
+		float4 Green = SourceTexture.GatherGreen(TextureSampler, UVOffset);
+		float4 Blue = SourceTexture.GatherBlue(TextureSampler, UVOffset);
+		float4 Alpha = SourceTexture.GatherAlpha(TextureSampler, UVOffset);
+		Block[8] = float4(Red[3], Green[3], Blue[3], Alpha[3]);
+		Block[9] = float4(Red[2], Green[2], Blue[2], Alpha[2]);
+		Block[12] = float4(Red[0], Green[0], Blue[0], Alpha[0]);
+		Block[13] = float4(Red[1], Green[1], Blue[1], Alpha[1]);
+	}
+	{
+		float2 UVOffset = UV + 2.f * TexelUVSize;
+		float4 Red = SourceTexture.GatherRed(TextureSampler, UVOffset);
+		float4 Green = SourceTexture.GatherGreen(TextureSampler, UVOffset);
+		float4 Blue = SourceTexture.GatherBlue(TextureSampler, UVOffset);
+		float4 Alpha = SourceTexture.GatherAlpha(TextureSampler, UVOffset);
+		Block[10] = float4(Red[3], Green[3], Blue[3], Alpha[3]);
+		Block[11] = float4(Red[2], Green[2], Blue[2], Alpha[2]);
+		Block[14] = float4(Red[0], Green[0], Blue[0], Alpha[0]);
+		Block[15] = float4(Red[1], Green[1], Blue[1], Alpha[1]);
+	}
+}
+
 // Read a 4x4 color of channel X block ready for BC compression
 void ReadBlockX(Texture2D<float4> SourceTexture, SamplerState TextureSampler, float2 UV, float2 TexelUVSize, out float Block[16])
 {
