Fixed illuminance probe not accounting for DRS and viewport offset.

#tests PC editor dx12 tested
#rnx

diff --git a/Engine/Shaders/Private/DebugProbes.usf b/Engine/Shaders/Private/DebugProbes.usf
--- a/Engine/Shaders/Private/DebugProbes.usf
+++ b/Engine/Shaders/Private/DebugProbes.usf
@@ -265,9 +265,10 @@ void MainPS(
 	// IlluminanceMeter variables must match the one in PostProcessVisualizeHDR.usf
 	const float  IlluminanceMeterSize = 20 * View.ResolutionFractionAndInv.x;
 	const float  IlluminanceMeterHalfSize = IlluminanceMeterSize / 2;
-	const float2 IlluminanceMeterCenter = View.ViewSizeAndInvSize.xy * 0.5f - uint2(0, 100) * View.ResolutionFractionAndInv.x;
+	const float2 IlluminanceMeterPixelCenter = View.ViewRectMin.xy + View.ViewSizeAndInvSize.xy * 0.5f - float2(0, 100) * View.ResolutionFractionAndInv.x;
 
-	if(all(PixelCoord > (IlluminanceMeterCenter - IlluminanceMeterHalfSize)) && all(PixelCoord < (IlluminanceMeterCenter + IlluminanceMeterHalfSize)))
+	const float PixelMargin = 0.0f;
+	if(all(PixelCoord > (IlluminanceMeterPixelCenter - IlluminanceMeterHalfSize - PixelMargin)) && all(PixelCoord < (IlluminanceMeterPixelCenter + IlluminanceMeterHalfSize + PixelMargin)))
 	{
 	#if PERMUTATION_PASS == RENDER_VELOCITYPASS
 
