Fix metallic properties not making it into shader path for mobile materials. The conditions around IsNonmetal are invalid as they only consider input connections but not constant scalar values. In either case, the constant value could be modified at runtime. Moreover, the optimization from this path was limited to a lerp and a multiplication operation.

#rn Allow metallic materials on mobile to be driven by scalar constants as their Metallic/Specular attributes.

diff --git a/Engine/Shaders/Private/MobileBasePassPixelShader.usf b/Engine/Shaders/Private/MobileBasePassPixelShader.usf
--- a/Engine/Shaders/Private/MobileBasePassPixelShader.usf
+++ b/Engine/Shaders/Private/MobileBasePassPixelShader.usf
@@ -556,9 +556,9 @@ void Main(
 	}
 #endif
 
-#if NONMETAL
+#if MOBILE_QL_FORCE_NONMETAL
 	GBuffer.DiffuseColor = GBuffer.BaseColor;
-	GBuffer.SpecularColor = 0.04;
+	GBuffer.SpecularColor = DielectricSpecularToF0(GBuffer.SpecularColor);
 #else
 	GBuffer.SpecularColor = ComputeF0(GBuffer.Specular, GBuffer.BaseColor, GBuffer.Metallic);
 	GBuffer.DiffuseColor = GBuffer.BaseColor - GBuffer.BaseColor * GBuffer.Metallic;
