Fixing compiler error when dead-code stripping is not enabled.

diff --git a/Engine/Shaders/Private/HeterogeneousVolumes/HeterogeneousVolumesAmbientOcclusionPipeline.usf b/Engine/Shaders/Private/HeterogeneousVolumes/HeterogeneousVolumesAmbientOcclusionPipeline.usf
--- a/Engine/Shaders/Private/HeterogeneousVolumes/HeterogeneousVolumesAmbientOcclusionPipeline.usf
+++ b/Engine/Shaders/Private/HeterogeneousVolumes/HeterogeneousVolumesAmbientOcclusionPipeline.usf
@@ -127,6 +127,16 @@ int GetStochasticFilteringMode()
 	return StochasticFilteringMode;
 }
 
+float EvalAmbientOcclusion(float3 WorldPosition)
+{
+	return 0.0;
+}
+
+float GetIndirectInscatteringFactor()
+{
+	return 0.0;
+}
+
 // Material sampling hooks..
 struct FVolumeSampleContext
 {
diff --git a/Engine/Shaders/Private/HeterogeneousVolumes/HeterogeneousVolumesLiveShadingShadows.usf b/Engine/Shaders/Private/HeterogeneousVolumes/HeterogeneousVolumesLiveShadingShadows.usf
--- a/Engine/Shaders/Private/HeterogeneousVolumes/HeterogeneousVolumesLiveShadingShadows.usf
+++ b/Engine/Shaders/Private/HeterogeneousVolumes/HeterogeneousVolumesLiveShadingShadows.usf
@@ -142,6 +142,11 @@ int GetStochasticFilteringMode()
 	return StochasticFilteringMode;
 }
 
+float EvalAmbientOcclusion(float3 WorldPosition)
+{
+	return 0.0;
+}
+
 float GetIndirectInscatteringFactor()
 {
 	return 0.0;
