Fix hair strands voxelization having incorrect voxel size due to missing quantization.

Due to this missing quantization (which happens at packing time), the hair voxel structure is offset causing incorrect traversing during shadow/transmittance computation.

diff --git a/Engine/Shaders/Private/HairStrands/HairStrandsVoxelRasterCompute.usf b/Engine/Shaders/Private/HairStrands/HairStrandsVoxelRasterCompute.usf
--- a/Engine/Shaders/Private/HairStrands/HairStrandsVoxelRasterCompute.usf
+++ b/Engine/Shaders/Private/HairStrands/HairStrandsVoxelRasterCompute.usf
@@ -24,6 +24,23 @@
 
 #if SHADER_ALLOCATEPAGEINDEX 
 
+
+uint PackVoxelWorldSize(float In)
+{
+	// Encode voxel world size [0..10] onto 10bits, i.e.: 10.f / 1023u
+	const float NormalizedVoxelWorldSize = min(In, 10.f) * 0.1f;
+	return PackUnorm10(NormalizedVoxelWorldSize);
+}
+float UnpackVoxelWorldSize(uint In)
+{
+	return UnpackUnorm10(In) * 10.f;
+}
+
+float QuantizeVoxelWorldSize(float In)
+{
+	return UnpackVoxelWorldSize(PackVoxelWorldSize(In));
+}
+
 float				CPUPageWorldSize;
 float				CPUVoxelWorldSize;
 uint				bUseCPUVoxelWorldSize;	// When adaptive voxel size is disabled, we use CPU voxel size value
