Fix CIS shader compilation due to robomerged CL.

diff --git a/Engine/Shaders/Private/HairStrands/HairStrandsVoxelRasterCompute.usf b/Engine/Shaders/Private/HairStrands/HairStrandsVoxelRasterCompute.usf
--- a/Engine/Shaders/Private/HairStrands/HairStrandsVoxelRasterCompute.usf
+++ b/Engine/Shaders/Private/HairStrands/HairStrandsVoxelRasterCompute.usf
@@ -24,23 +24,6 @@
 
 #if SHADER_ALLOCATEPAGEINDEX 
 
-
-uint PackVoxelWorldSize(float In)
-{
-	// Encode voxel world size [0..10] onto 10bits, i.e.: 10.f / 1023u
-	const float NormalizedVoxelWorldSize = min(In, 10.f) * 0.1f;
-	return PackUnorm10(NormalizedVoxelWorldSize);
-}
-float UnpackVoxelWorldSize(uint In)
-{
-	return UnpackUnorm10(In) * 10.f;
-}
-
-float QuantizeVoxelWorldSize(float In)
-{
-	return UnpackVoxelWorldSize(PackVoxelWorldSize(In));
-}
-
 float				CPUPageWorldSize;
 float				CPUVoxelWorldSize;
 uint				bUseCPUVoxelWorldSize;	// When adaptive voxel size is disabled, we use CPU voxel size value
