First Person: Added support for sampling the bIsFirstPerson GBuffer bit from SceneTextures in material graphs.

diff --git a/Engine/Plugins/Experimental/ColorCorrectRegions/Shaders/Private/ColorCorrectRegionsCommon.usf b/Engine/Plugins/Experimental/ColorCorrectRegions/Shaders/Private/ColorCorrectRegionsCommon.usf
--- a/Engine/Plugins/Experimental/ColorCorrectRegions/Shaders/Private/ColorCorrectRegionsCommon.usf
+++ b/Engine/Plugins/Experimental/ColorCorrectRegions/Shaders/Private/ColorCorrectRegionsCommon.usf
@@ -71,6 +71,7 @@ SamplerState PostProcessInput_BilinearSampler;
 #define PPI_Velocity 28
 #define PPI_WorldTangent 29
 #define PPI_Anisotropy 30
+#define PPI_IsFirstPerson 31
 
 //////////////////////////////////////////////////////////////////////////
 
@@ -367,6 +368,8 @@ float4 SceneTextureLookup(float2 UV, int SceneTextureIndex, bool bFiltered)
 		return float4(ScreenSpaceData.GBuffer.WorldTangent, 0);
 	case PPI_Anisotropy:
 		return ScreenSpaceData.GBuffer.Anisotropy;
+	case PPI_IsFirstPerson:
+		return IsFirstPerson(ScreenSpaceData.GBuffer) ? 1.0f : 0.0f;
 #endif
 	default:
 		return float4(0, 0, 0, 0);
diff --git a/Engine/Shaders/Private/MaterialTemplate.ush b/Engine/Shaders/Private/MaterialTemplate.ush
--- a/Engine/Shaders/Private/MaterialTemplate.ush
+++ b/Engine/Shaders/Private/MaterialTemplate.ush
@@ -80,6 +80,7 @@
 #define PPI_Velocity 28
 #define PPI_WorldTangent 29
 #define PPI_Anisotropy 30
+#define PPI_IsFirstPerson 31
 
 #ifdef POST_PROCESS_PROPAGATE_ALPHA_INPUT
 // Arbitrary internal value to fetch local variable propagated alpha stored in the SceneTextureLookup function
@@ -3246,6 +3247,7 @@ float4 SceneTextureLookup(float2 UV, int SceneTextureIndex, bool bFiltered)
 		case PPI_StoredSpecular:		return 0.0f;
 		case PPI_WorldTangent:			return 0.0f;
 		case PPI_Anisotropy:			return 0.0f;
+		case PPI_IsFirstPerson:			return 0.0f;
 		}
 	}
 	#elif SUBSTRATE_ENABLED
@@ -3273,6 +3275,7 @@ float4 SceneTextureLookup(float2 UV, int SceneTextureIndex, bool bFiltered)
 				case PPI_StoredSpecular:		return float4(SubstrateGBuffer.StoredSpecular.rrr, 0);
 				case PPI_WorldTangent:			return float4(SubstrateGBuffer.WorldTangent, 0);
 				case PPI_Anisotropy:			return SubstrateGBuffer.Anisotropy;
+				case PPI_IsFirstPerson:			return SubstrateGBuffer.bIsFirstPerson ? 1.0f : 0.0f;
 			}
 		}
 	}
@@ -3302,6 +3305,7 @@ float4 SceneTextureLookup(float2 UV, int SceneTextureIndex, bool bFiltered)
 			case PPI_StoredSpecular:		return float4(ScreenSpaceData.GBuffer.StoredSpecular.rrr, 0);
 			case PPI_WorldTangent:			return float4(ScreenSpaceData.GBuffer.WorldTangent, 0);
 			case PPI_Anisotropy:			return ScreenSpaceData.GBuffer.Anisotropy;
+			case PPI_IsFirstPerson:			return IsFirstPerson(ScreenSpaceData.GBuffer) ? 1.0f : 0.0f;
 		}
 	}
 	#endif // SUBSTRATE_ENABLED
diff --git a/Engine/Shaders/Public/SubstratePublic.ush b/Engine/Shaders/Public/SubstratePublic.ush
--- a/Engine/Shaders/Public/SubstratePublic.ush
+++ b/Engine/Shaders/Public/SubstratePublic.ush
@@ -38,6 +38,7 @@ struct FSubstrateGBuffer
 	float3 StoredSpecular;
 	float3 WorldTangent;
 	float  Anisotropy;
+	bool   bIsFirstPerson;
 };
 
 FSubstrateGBuffer SubstratePublic_GetSubstrateGBufferFromFirstBSDF(uint2 PixelCoord)
@@ -74,6 +75,7 @@ FSubstrateGBuffer SubstratePublic_GetSubstrateGBufferFromFirstBSDF(uint2 PixelCo
 		SGBuffer.StoredSpecular		= SubstrateGetBSDFSpecular(BSDF).rrr;
 		SGBuffer.WorldTangent		= SubstrateGetWorldTangent(SubstratePixelHeader, BSDF, SubstrateAddressing);
 		SGBuffer.Anisotropy			= SubstrateGetBSDFAnisotropy(BSDF);
+		SGBuffer.bIsFirstPerson		= SubstratePixelHeader.IsFirstPerson();
 
 		SGBuffer.Opacity			= 0.0f;
 		if (SubstrateSubSurfaceHeaderGetSSSType(SSSHeader) != SSS_TYPE_INVALID)
diff --git a/Engine/Source/Runtime/Engine/Public/MaterialSceneTextureId.h b/Engine/Source/Runtime/Engine/Public/MaterialSceneTextureId.h
--- a/Engine/Source/Runtime/Engine/Public/MaterialSceneTextureId.h
+++ b/Engine/Source/Runtime/Engine/Public/MaterialSceneTextureId.h
@@ -73,6 +73,8 @@ enum ESceneTextureId : int
 	PPI_WorldTangent UMETA(DisplayName = "WorldTangent"),
 	/** Material anisotropy, single channel (GBuffer) */
 	PPI_Anisotropy UMETA(DisplayName = "Anisotropy"),
+	/** Whether the pixel is marked as belonging to an opaque first person primitive, single channel (GBuffer) */
+	PPI_IsFirstPerson UMETA(DisplayName = "IsFirstPerson"),
 
 	/** Used internally for user scene texture inputs.  Later remapped to PPI_PostProcessInput0-6 in the shader via procedurally generated defines. */
 	PPI_UserSceneTexture0 = 32 UMETA(Hidden),
diff --git a/Engine/Source/Runtime/Engine/Public/MaterialShared.h b/Engine/Source/Runtime/Engine/Public/MaterialShared.h
--- a/Engine/Source/Runtime/Engine/Public/MaterialShared.h
+++ b/Engine/Source/Runtime/Engine/Public/MaterialShared.h
@@ -850,7 +850,8 @@ public:
 			IsSceneTextureUsed(PPI_ShadingModelID) ||
 			IsSceneTextureUsed(PPI_StoredBaseColor) ||
 			IsSceneTextureUsed(PPI_StoredSpecular) ||
-			IsSceneTextureUsed(PPI_Velocity);
+			IsSceneTextureUsed(PPI_Velocity) ||
+			IsSceneTextureUsed(PPI_IsFirstPerson);
 	}
 
 	/** true if the material uses the SceneDepth lookup */
