Fixed bad world position for lighting for the translucency lighting volume.

#tests Tested pc dx12 and dx11 to update enginetest reference screenshots.
#rnx

#virtualized

diff --git a/Engine/Shaders/Private/TranslucentLightInjectionShaders.usf b/Engine/Shaders/Private/TranslucentLightInjectionShaders.usf
--- a/Engine/Shaders/Private/TranslucentLightInjectionShaders.usf
+++ b/Engine/Shaders/Private/TranslucentLightInjectionShaders.usf
@@ -44,6 +44,12 @@ uint AtmospherePerPixelTransmittanceEnabled;
 #include "HeterogeneousVolumes/HeterogeneousVolumesAdaptiveVolumetricShadowMapSampling.ush"
 #endif // ADAPTIVE_VOLUMETRIC_SHADOW_MAP
 
+// To use DEBUG_ONE_VOXEL, one must enable ShaderPrint from TranslucentLighting.cpp.
+#define DEBUG_ONE_VOXEL 0
+#if DEBUG_ONE_VOXEL
+#include "/Engine/Private/ShaderPrint.ush"
+#endif
+
 float3 GetTranslatedWorldPosition(float2 UV, uint LayerIndex)
 {
 	float ZPosition = View.TranslucencyLightingVolumeMin[VolumeCascadeIndex].z + (LayerIndex + .5f) * View.TranslucencyLightingVolumeInvSize[VolumeCascadeIndex].w;
@@ -94,6 +100,15 @@ float CalcSimpleLightAttenuation(float3 TranslatedWorldPosition)
 	return Attenuation;
 }
 
+float3 GetTranslatedWorldPositionForLighting(float3 TranslatedWorldPosition, float3 NormalizedLightVector)
+{
+	float VoxelSize = View.TranslucencyLightingVolumeInvSize[VolumeCascadeIndex].w;
+	// Push out the position for lighting in the oposite direction of the L vector 
+	// to make sure we always have half a voxel distance from the light to avoid overblown lighting intensities.
+	float3 LightingOffset = -NormalizedLightVector * 0.5f * VoxelSize;
+	return TranslatedWorldPosition + LightingOffset;
+}
+
 /** Pixel shader that calculates direct lighting for a simple light (unshadowed point light) for all the affected voxels of a volume texture. */
 void SimpleLightInjectMainPS(
 	FWriteToSliceGeometryOutput Input,
@@ -106,9 +121,7 @@ void SimpleLightInjectMainPS(
 
 	float3 TranslatedWorldPosition = GetTranslatedWorldPosition(Input.Vertex.UV, Input.LayerIndex);
 	float3 NormalizedLightVector = normalize(SimpleLightPositionAndRadius.xyz - TranslatedWorldPosition);
-
-	float VoxelSize = View.TranslucencyLightingVolumeInvSize[VolumeCascadeIndex].w;
-	float3 TranslatedWorldPositionForLighting = TranslatedWorldPosition + 1 * GetBoxPushout(NormalizedLightVector, .5f * VoxelSize);
+	float3 TranslatedWorldPositionForLighting = GetTranslatedWorldPositionForLighting(TranslatedWorldPosition, NormalizedLightVector);
 
 	float Attenuation = CalcSimpleLightAttenuation(TranslatedWorldPositionForLighting);
 	float3 Lighting = SimpleLightColorAndExponent.rgb / PI * Attenuation;
@@ -134,13 +147,6 @@ float GetBorderLerpFactor(float2 UV, uint LayerIndex)
 	return FinalLerpFactor;
 }
 
-float3 GetTranslatedWorldPositionForLighting(float3 TranslatedWorldPosition, float3 NormalizedLightVector)
-{
-	float VoxelSize = View.TranslucencyLightingVolumeInvSize[VolumeCascadeIndex].w;
-	float3 LightingOffset = 1 * GetBoxPushout(NormalizedLightVector, .5f * VoxelSize);
-	return TranslatedWorldPosition + LightingOffset;
-}
-
 float GetTranslucencyVolumeLocalLightAttenuation(float3 TranslatedWorldPosition, FDeferredLightData LightData)
 {
 	float3 ToLight, L;
@@ -211,6 +217,20 @@ void InjectMainPS(
 
 	// 0: no contribution, 1:full contribution
 	float Masking = 1.0f;
+
+#if DEBUG_ONE_VOXEL
+	const float DebugVoxelSize = View.TranslucencyLightingVolumeInvSize[VolumeCascadeIndex].w;
+	const float DebugHalfVoxelSize = DebugVoxelSize * 0.5f;
+	const float3 WorldPosition = (TranslatedWorldPosition - DFDemote(PrimaryView.PreViewTranslation));
+	const int3 IntegerPos = (int3)(WorldPosition / (DebugVoxelSize));
+
+	// Enable shader print for a single voxel containing VoxelIncludedWorldPosition for cascade 0.
+	const float3 VoxelIncludedWorldPosition = float3(0.5, 0.5, 0.5);
+	FShaderPrintContext Context = InitShaderPrintContext(all((WorldPosition - DebugHalfVoxelSize) <= VoxelIncludedWorldPosition) && all((WorldPosition + DebugHalfVoxelSize) > VoxelIncludedWorldPosition) && VolumeCascadeIndex == 0, uint2(200, 200));
+	AddAABBTWS(Context, TranslatedWorldPosition - DebugHalfVoxelSize, TranslatedWorldPosition + DebugHalfVoxelSize, float4(1, 1, 1, 1));
+	AddSphereTWS(Context, TranslatedWorldPosition, 2.0f, float4(0, 1, 0, 1));
+#endif
+	
 	#if RADIAL_ATTENUATION 
 	{
 		// cull voxels outside the light radius (value is < 0)
@@ -236,6 +256,10 @@ void InjectMainPS(
 	float3 NormalizedLightVector = normalize(LightVector);
 	float3 TranslatedWorldPositionForLighting = GetTranslatedWorldPositionForLighting(TranslatedWorldPosition, NormalizedLightVector);
 
+#if DEBUG_ONE_VOXEL
+	AddSphereTWS(Context, TranslatedWorldPositionForLighting, 5.0f, float4(0, 1, 1, 1));
+#endif
+
 	{
 		float Attenuation = 1.0f;
 		float ShadowFactor = 1.0f;
diff --git a/Engine/Source/Runtime/Renderer/Private/TranslucentLighting.cpp b/Engine/Source/Runtime/Renderer/Private/TranslucentLighting.cpp
--- a/Engine/Source/Runtime/Renderer/Private/TranslucentLighting.cpp
+++ b/Engine/Source/Runtime/Renderer/Private/TranslucentLighting.cpp
@@ -712,6 +712,7 @@ class FTranslucentLightingInjectPS : public FMaterialShader
 
 	BEGIN_SHADER_PARAMETER_STRUCT(FParameters, )
 		SHADER_PARAMETER_STRUCT_REF(FViewUniformShaderParameters, ViewUniformBuffer)
+		SHADER_PARAMETER_STRUCT_INCLUDE(ShaderPrint::FShaderParameters, ShaderPrint)
 		SHADER_PARAMETER_RDG_UNIFORM_BUFFER(FDeferredLightUniformStruct, DeferredLight)
 		SHADER_PARAMETER_STRUCT_INCLUDE(FVolumeShadowingShaderParameters, VolumeShadowingParameters)
 		SHADER_PARAMETER_STRUCT_INCLUDE(FVirtualShadowMapSamplingParameters, VirtualShadowMapSamplingParameters)
@@ -1502,6 +1503,14 @@ void InjectTranslucencyLightingVolume(
 
 				PassParameters->PS.CameraRelativeLightPosition = GetCamRelativeLightPosition(View.ViewMatrices, *LightSceneInfo);
 
+			#if 0 // Enable this to be able to debug using DEBUG_ONE_VOXEL
+				ShaderPrint::SetEnabled(true);
+				ShaderPrint::RequestSpaceForLines(64u);
+				ShaderPrint::RequestSpaceForCharacters(64u);
+				ShaderPrint::RequestSpaceForTriangles(64u);
+				ShaderPrint::SetParameters(GraphBuilder, View.ShaderPrintData, PassParameters->PS.ShaderPrint);
+			#endif
+
 				GraphBuilder.AddPass(
 					RDG_EVENT_NAME("InjectTranslucencyLightingVolume(VolumeCascade=%d%s%s%s)",
 						VolumeCascadeIndex,
