Substrate - fixed wrong screen contact shadows on eye materials.

#tests PC dx12 + substrate
#rnx

diff --git a/Engine/Shaders/Private/DeferredLightingCommon.ush b/Engine/Shaders/Private/DeferredLightingCommon.ush
--- a/Engine/Shaders/Private/DeferredLightingCommon.ush
+++ b/Engine/Shaders/Private/DeferredLightingCommon.ush
@@ -78,7 +78,9 @@ float ShadowRayCast(
 		uint2 PixelPos = View.ViewRectMin.xy + View.ViewSizeAndInvSize.xy * HitUV;
 		FSubstrateAddressing SubstrateAddressing = GetSubstratePixelDataByteOffset(PixelPos, uint2(View.BufferSizeAndInvSize.xy), SubstrateRayCast.MaxBytesPerPixel);
 		FSubstratePixelHeader SubstratePixelHeader = UnpackSubstrateHeaderIn(SubstrateRayCast.MaterialTextureArray, SubstrateAddressing, SubstrateRayCast.TopLayerTexture);
-		bOutHitCastContactShadow = SubstratePixelHeader.DoesCastContactShadow();
+		// We check if the hit pixel has CastContactShadow, and we always exclude EYE BSDFs from there as done by CastContactShadow(SampleGBuffer) in the legacy path below.
+		// SUBSTRATE_TODO: try to move that test to GBuffer export pass (set DoesCastContactShadow = false for EYE BSDFs).
+		bOutHitCastContactShadow = SubstratePixelHeader.DoesCastContactShadow() && !SubstratePixelHeader.SubstrateGetBSDFType() == SUBSTRATE_BSDF_TYPE_EYE;
 #endif
 #else
 		FGBufferData SampleGBuffer = GetGBufferData(HitUV);
