Force rough diffuse energy preservation correction to be achromatic.

Note: on ALU-bound shader, the diffuse correction is less than 1% of the cost as the initial LUT lookup is actually reused between the diff/spec energy conservation sampling and the correction factor. The correction factor only adds extra ALU (7 full rate instructions) for the hardcoded F0/F90 values.

diff --git a/Engine/Shaders/Private/ShadingEnergyConservationTemplate.ush b/Engine/Shaders/Private/ShadingEnergyConservationTemplate.ush
--- a/Engine/Shaders/Private/ShadingEnergyConservationTemplate.ush
+++ b/Engine/Shaders/Private/ShadingEnergyConservationTemplate.ush
@@ -24,6 +24,15 @@ BXDF_ENERGY_TYPE BXDF_ENERGY_SUFFIX(GetF0F90)(float3 InF0)
 #endif
 }
 
+float GetEnergyLuminance(BXDF_ENERGY_TYPE InE)
+{
+#if IS_BXDF_ENERGY_TYPE_ACHROMATIC
+	return InE;
+#else
+	return Luminance(InE);
+#endif
+}
+
 struct BXDF_ENERGY_SUFFIX(FBxDFEnergyTerms)
 {
 	BXDF_ENERGY_TYPE W; // overall weight to scale the lobe BxDF by to ensure energy conservation
@@ -116,9 +125,9 @@ float ComputeEnergyPreservation(BXDF_ENERGY_SUFFIX(FBxDFEnergyTerms) EnergyTerms
 {
 #if USE_ENERGY_CONSERVATION > 0
 #if USE_DEVELOPMENT_SHADERS
-	return View.bShadingEnergyPreservation ? (1 - Luminance(EnergyTerms.E)) : 1.0f;
+	return View.bShadingEnergyPreservation ? (1 - GetEnergyLuminance(EnergyTerms.E)) : 1.0f;
 #else
-	return 1 - Luminance(EnergyTerms.E);
+	return 1 - GetEnergyLuminance(EnergyTerms.E);
 #endif
 #else
 	return 1.0f;
diff --git a/Engine/Shaders/Private/Substrate/SubstrateEvaluation.ush b/Engine/Shaders/Private/Substrate/SubstrateEvaluation.ush
--- a/Engine/Shaders/Private/Substrate/SubstrateEvaluation.ush
+++ b/Engine/Shaders/Private/Substrate/SubstrateEvaluation.ush
@@ -784,7 +784,7 @@ FSubstrateEvaluateResult SubstrateEvaluateBSDFCommon(FSubstrateBSDFContext BSDFC
 				if (Settings.bRoughDiffuseEnabled && BSDF_GETSSSTYPE(BSDFContext.BSDF) != SSS_TYPE_SIMPLEVOLUME)
 				{
 					const float ChanF0 = 0.04;
-					FBxDFEnergyTerms EnergyTerms = ComputeGGXSpecEnergyTerms(SafeRoughness, BSDFContext.Context.NoV, ChanF0, 1.0f);
+					FBxDFEnergyTermsA EnergyTerms = ComputeGGXSpecEnergyTermsA(SafeRoughness, BSDFContext.Context.NoV, ChanF0, 1.0f);
 					const float SpecularTransmitionChanF0ToMaterialF0 = DirectionalAlbedo_SpecularTransmission / max(SUBSTRATE_EPSILON, ComputeEnergyPreservation(EnergyTerms));
 					DirectionalAlbedo_SpecularTransmission = SpecularTransmitionChanF0ToMaterialF0;
 				}
