Volumetric cloud - fix for band of cloud appearing at the horizon when high up in the sky. That special case when using distance form camera tracing mode, was missing the depth test.

#tests PC dx12 editor
#rnx

diff --git a/Engine/Shaders/Private/VolumetricCloud.usf b/Engine/Shaders/Private/VolumetricCloud.usf
--- a/Engine/Shaders/Private/VolumetricCloud.usf
+++ b/Engine/Shaders/Private/VolumetricCloud.usf
@@ -612,6 +612,7 @@ void MainCommon(in FMaterialPixelParameters MaterialParameters, in float4 SvPosi
 
 		float3 DepthBufferTranslatedWorldPos = SvPositionToTranslatedWorld(float4(SvPosition.xy, DeviceZ, 1.0));
 		const float TDepthBufferKm = min(length(DepthBufferTranslatedWorldPos * CENTIMETER_TO_KILOMETER - RayTranslatedWorldOriginKm), MaxHalfFloat);
+		const float TDepthBufferCm = TDepthBufferKm * KILOMETER_TO_CENTIMETER;
 
 		OutDepth.y = TDepthBufferKm;
 
@@ -621,7 +622,7 @@ void MainCommon(in FMaterialPixelParameters MaterialParameters, in float4 SvPosi
 
 		if(!IsOrthoProjection() || TMax < View.OrthoFarPlane)
 		{
-			if (bDoDepthIntersectionAndTest && (TDepthBufferKm * KILOMETER_TO_CENTIMETER) < TMin)
+			if (bDoDepthIntersectionAndTest && TDepthBufferCm < TMin)
 			{
 				//OutColor0 = float4(0.0f, 0.0f, 1.0f, 1.0f);
 				return;
@@ -630,7 +631,7 @@ void MainCommon(in FMaterialPixelParameters MaterialParameters, in float4 SvPosi
 			if (bDoDepthIntersectionAndTest)
 			{
 				// Only trace up to the closest between the cloud layer far distance and the depth buffer 
-				TMax = min(TMax, TDepthBufferKm * KILOMETER_TO_CENTIMETER);
+				TMax = min(TMax, TDepthBufferCm);
 			}
 		}
 
@@ -681,6 +682,9 @@ void MainCommon(in FMaterialPixelParameters MaterialParameters, in float4 SvPosi
 				// Make sure we trace when traveling out of the bottom entry point but travel back again in for the specified range of distance.
 				// But we try to not trace outside of the cloud layer still for highly curved planets.
 				TMax = min(RenderVolumetricCloudParameters.TracingMaxDistance, max(tTop2.x, tTop2.y) * KILOMETER_TO_CENTIMETER);
+
+				// And we must apply the far, i.e. not near, depth buffer constraint to TMax.
+				TMax = min(TMax, TDepthBufferCm);
 			}
 			else
 			{
