Fix missing diffuse contribution on SSS surface (profile, or free MFP) in Lumen reflection with Substrate.

The diffuse split contribution was applied on HitLighting shader, causing the diffuse contribution to be bypassed.

diff --git a/Engine/Shaders/Private/Substrate/SubstrateEvaluation.ush b/Engine/Shaders/Private/Substrate/SubstrateEvaluation.ush
--- a/Engine/Shaders/Private/Substrate/SubstrateEvaluation.ush
+++ b/Engine/Shaders/Private/Substrate/SubstrateEvaluation.ush
@@ -904,7 +904,7 @@ FSubstrateEvaluateResult SubstrateEvaluateBSDFCommon(FSubstrateBSDFContext BSDFC
 			#if SUBSTRATE_FASTPATH==0
 			const bool bIsSubsurfaceDiffusionProfile = BSDF_GETSSSTYPE(BSDFContext.BSDF) == SSS_TYPE_DIFFUSION_PROFILE;
 			const bool bIsSubsurfaceDiffusionMFP = BSDF_GETSSSTYPE(BSDFContext.BSDF) == SSS_TYPE_DIFFUSION;
-			Sample.bPostProcessSubsurface = bIsSubsurfaceDiffusionProfile || bIsSubsurfaceDiffusionMFP;
+			Sample.bPostProcessSubsurface = (bIsSubsurfaceDiffusionProfile || bIsSubsurfaceDiffusionMFP) && SUBSTRATE_SSS_MATERIAL_OVERRIDE;
 			#if SUBSTRATE_SSS_TRANSMISSION
 			if (bIsSubsurfaceDiffusionProfile || bIsSubsurfaceDiffusionMFP)
 			{
@@ -1189,7 +1189,7 @@ FSubstrateEvaluateResult SubstrateEvaluateBSDFCommon(FSubstrateBSDFContext BSDFC
 
 			Sample.DiffuseColor		= DiffuseColor;
 			Sample.DiffusePDF		= BSDFContext.SatNoL / PI;
-			Sample.bPostProcessSubsurface = BSDF_GETSSSTYPE(BSDFContext.BSDF) == SSS_TYPE_DIFFUSION || BSDF_GETSSSTYPE(BSDFContext.BSDF) == SSS_TYPE_DIFFUSION_PROFILE;
+			Sample.bPostProcessSubsurface = (BSDF_GETSSSTYPE(BSDFContext.BSDF) == SSS_TYPE_DIFFUSION || BSDF_GETSSSTYPE(BSDFContext.BSDF) == SSS_TYPE_DIFFUSION_PROFILE) && SUBSTRATE_SSS_MATERIAL_OVERRIDE;
 
 			//
 			// Apply energy conservation on the diffuse component
@@ -1592,7 +1592,7 @@ FSubstrateEnvLightResult SubstrateEvaluateForEnvLight(FSubstrateBSDFContext BSDF
 			const bool bIsSubsurfaceDiffusionProfile = BSDF_GETSSSTYPE(BSDFContext.BSDF) == SSS_TYPE_DIFFUSION_PROFILE;
 			const bool bIsSubsurfaceDiffusionMFP = BSDF_GETSSSTYPE(BSDFContext.BSDF) == SSS_TYPE_DIFFUSION;
 			const bool bIsThin = BSDF_GETISTHIN(BSDFContext.BSDF);
-			SubstrateEnvLightResult.bPostProcessSubsurface = bIsSubsurfaceDiffusionProfile || bIsSubsurfaceDiffusionMFP;
+			SubstrateEnvLightResult.bPostProcessSubsurface = (bIsSubsurfaceDiffusionProfile || bIsSubsurfaceDiffusionMFP) && SUBSTRATE_SSS_MATERIAL_OVERRIDE;
 
 			BRANCH
 			if (BSDF_GETSSSTYPE(BSDFContext.BSDF) == SSS_TYPE_SIMPLEVOLUME)
@@ -1857,7 +1857,7 @@ FSubstrateEnvLightResult SubstrateEvaluateForEnvLight(FSubstrateBSDFContext BSDF
 			}
 
 			SubstrateEnvLightResult.DiffuseNormal			= BSDFContext.N;
-			SubstrateEnvLightResult.bPostProcessSubsurface	= BSDF_GETSSSTYPE(BSDFContext.BSDF) == SSS_TYPE_DIFFUSION || BSDF_GETSSSTYPE(BSDFContext.BSDF) == SSS_TYPE_DIFFUSION_PROFILE;
+			SubstrateEnvLightResult.bPostProcessSubsurface	= (BSDF_GETSSSTYPE(BSDFContext.BSDF) == SSS_TYPE_DIFFUSION || BSDF_GETSSSTYPE(BSDFContext.BSDF) == SSS_TYPE_DIFFUSION_PROFILE) && SUBSTRATE_SSS_MATERIAL_OVERRIDE;
 
 			float SafeRoughness = MakeRoughnessSafe(SLAB_ROUGHNESS(BSDFContext.BSDF));
 
