Replace SUBSTRATE_FORCE_SINGLE_CLOSURE with SUBSTRATE_USE_FULLYSIMPLIFIED_MATERIAL to reduce number of define, and clarify the use case/behavior.

diff --git a/Engine/Shaders/Private/RayTracing/RayTracingLightingMS.usf b/Engine/Shaders/Private/RayTracing/RayTracingLightingMS.usf
--- a/Engine/Shaders/Private/RayTracing/RayTracingLightingMS.usf
+++ b/Engine/Shaders/Private/RayTracing/RayTracingLightingMS.usf
@@ -17,9 +17,6 @@
 #define SUBSTRATE_GLINTS_ALLOWED 0
 #define SUBSTRATE_DIFFUSE_COLOR_BOOST 1
 
-// RT shaders use fully simplified materials, i.e. have a single closure.
-#define SUBSTRATE_FORCE_SINGLE_CLOSURE 1
-
 #ifndef SUPPORT_LIGHT_FUNCTION
 #define SUPPORT_LIGHT_FUNCTION 0
 #endif
diff --git a/Engine/Shaders/Private/Substrate/SubstrateEvaluation.ush b/Engine/Shaders/Private/Substrate/SubstrateEvaluation.ush
--- a/Engine/Shaders/Private/Substrate/SubstrateEvaluation.ush
+++ b/Engine/Shaders/Private/Substrate/SubstrateEvaluation.ush
@@ -60,8 +60,8 @@
 #define SUBSTRATE_DIFFUSE_COLOR_BOOST 0
 #endif
 
-#ifndef SUBSTRATE_FORCE_SINGLE_CLOSURE
-#define SUBSTRATE_FORCE_SINGLE_CLOSURE 0
+#ifndef SUBSTRATE_USE_FULLYSIMPLIFIED_MATERIAL
+#define SUBSTRATE_USE_FULLYSIMPLIFIED_MATERIAL 0
 #endif
 
 ///////////////////////////////////////////////////////////////////////////////
@@ -74,7 +74,7 @@
 //		FSubstrateBSDF BSDF = UnpackSubstrateBSDF(MaterialBuffer, SubstrateAddressing, SubstratePixelHeader);
 //		...
 // }
-#if (SUBSTRATE_FASTPATH == 0 && SUBSTRATE_SINGLEPATH == 0 && SUBSTRATE_FORCE_SINGLE_CLOSURE == 0)
+#if (SUBSTRATE_FASTPATH == 0 && SUBSTRATE_SINGLEPATH == 0 && SUBSTRATE_USE_FULLYSIMPLIFIED_MATERIAL == 0)
 	#if COMPILER_FXC
 		#define Substrate_for(X,Y,Z) [loop] for(X;Y;Z)
 	#else
