Fixed shader compilation failure in alpha invert shader

diff --git a/Engine/Shaders/Private/PostProcessAlphaInvert.usf b/Engine/Shaders/Private/PostProcessAlphaInvert.usf
--- a/Engine/Shaders/Private/PostProcessAlphaInvert.usf
+++ b/Engine/Shaders/Private/PostProcessAlphaInvert.usf
@@ -19,13 +19,13 @@ SamplerState InTextureSampler;
 void AlphaInvert_MainPS(
 	FScreenVertexOutput								Input,
 	out				 float4 OutColor				: SV_Target0
-#if  MOBILE_MULTI_VIEW
+#if MOBILE_MULTI_VIEW && !INSTANCED_STEREO
 	, in nointerpolation uint MultiViewId			: SV_ViewID
 #endif
 )
 {
 	uint ArrayIndex = 0;
-#if MOBILE_MULTI_VIEW
+#if MOBILE_MULTI_VIEW  && !INSTANCED_STEREO
 	ArrayIndex = MultiViewId;
 #endif
 	
