Fixed: Instance occlusion culling not working on mobile platforms

diff --git a/Engine/Shaders/Private/HZB.usf b/Engine/Shaders/Private/HZB.usf
--- a/Engine/Shaders/Private/HZB.usf
+++ b/Engine/Shaders/Private/HZB.usf
@@ -18,18 +18,26 @@ uint4 PixelViewPortMinMax;	// Max is inclusive
 Texture2D		ParentTextureMip;
 SamplerState	ParentTextureMipSampler;
 
+#if PIXELSHADER
+uint SourceMipIndex;
+#endif
 
-float4 Gather4(Texture2D Texture, SamplerState TextureSampler, float2 BufferUV, uint Level)
+float4 Gather4(Texture2D Texture, SamplerState TextureSampler, float2 BufferUV)
 {
 	// min(..., InputViewportMaxBound) because we don't want to sample outside of the viewport
 	// when the view size has odd dimensions on X/Y axis.
 	float2 UV = min(BufferUV + float2(-0.25f, -0.25f) * InvSize, InputViewportMaxBound - InvSize);
-	return Texture.GatherRed(TextureSampler, UV, Level);
-}
 
-float4 Gather4(Texture2D Texture, SamplerState TextureSampler, float2 BufferUV)
-{
-	return Gather4(Texture, TextureSampler, BufferUV, 0);
+#if (PIXELSHADER && DIM_USE_MIPINDEX)
+	float4 UVRect = float4(UV, UV + InvSize);
+	return float4(
+		Texture.SampleLevel(TextureSampler, UVRect.xw, SourceMipIndex).r,
+		Texture.SampleLevel(TextureSampler, UVRect.zw, SourceMipIndex).r,
+		Texture.SampleLevel(TextureSampler, UVRect.zy, SourceMipIndex).r,
+		Texture.SampleLevel(TextureSampler, UVRect.xy, SourceMipIndex).r);
+#else
+	return Texture.GatherRed(TextureSampler, UV);
+#endif
 }
 
 #if COMPUTESHADER
@@ -287,12 +295,10 @@ void HZBBuildCS(
 
 #elif PIXELSHADER
 
-uint SourceMipIndex;
-
 void HZBBuildPS(float4 SvPosition : SV_POSITION, out float4 OutColor : SV_Target0)
 {
 	float2 BufferUV = SvPosition.xy * DispatchThreadIdToBufferUV.xy + DispatchThreadIdToBufferUV.zw;
-	float4 DeviceZ = Gather4(ParentTextureMip, ParentTextureMipSampler, BufferUV, SourceMipIndex);
+	float4 DeviceZ = Gather4(ParentTextureMip, ParentTextureMipSampler, BufferUV);
 
 	float FurthestDeviceZ = min(min(DeviceZ.x, DeviceZ.y), min(DeviceZ.z, DeviceZ.w));
 
diff --git a/Engine/Source/Runtime/Renderer/Private/SceneTextureReductions.cpp b/Engine/Source/Runtime/Renderer/Private/SceneTextureReductions.cpp
--- a/Engine/Source/Runtime/Renderer/Private/SceneTextureReductions.cpp
+++ b/Engine/Source/Runtime/Renderer/Private/SceneTextureReductions.cpp
@@ -38,6 +38,9 @@ class FHZBBuildPS : public FGlobalShader
 	DECLARE_GLOBAL_SHADER(FHZBBuildPS);
 	SHADER_USE_PARAMETER_STRUCT(FHZBBuildPS, FGlobalShader)
 
+	class FDimUseMipIndex : SHADER_PERMUTATION_BOOL("DIM_USE_MIPINDEX");
+	using FPermutationDomain = TShaderPermutationDomain<FDimUseMipIndex>;
+
 	BEGIN_SHADER_PARAMETER_STRUCT( FParameters, )
 		SHADER_PARAMETER_STRUCT_INCLUDE(FSharedHZBParameters, Shared)
 		SHADER_PARAMETER(int32, SourceMipIndex)
@@ -257,11 +260,15 @@ void BuildHZB(
 
 			FHZBBuildPS::FParameters* PassParameters = GraphBuilder.AllocParameters<FHZBBuildPS::FParameters>();
 			PassParameters->Shared = ShaderParameters;
-			PassParameters->SourceMipIndex = GRHISupportsTextureViews ? 0 : StartDestMip - 1;
+			PassParameters->SourceMipIndex = GRHISupportsTextureViews ? 0 : FMath::Max(0, StartDestMip - 1);
 			PassParameters->RenderTargets[0] = FRenderTargetBinding(FurthestHZBTexture, ERenderTargetLoadAction::ENoAction, StartDestMip);
 
+			FHZBBuildPS::FPermutationDomain PermutationVector;
+			// Use fallback with sampling specific HZB mip index if RHI does not support texture views (OpenGL)
+			PermutationVector.Set<FHZBBuildPS::FDimUseMipIndex>(GRHISupportsTextureViews ? false : true);
+
 			FGlobalShaderMap* GlobalShaderMap = GetGlobalShaderMap(FeatureLevel);
-			TShaderMapRef<FHZBBuildPS> PixelShader(GlobalShaderMap);
+			TShaderMapRef<FHZBBuildPS> PixelShader(GlobalShaderMap, PermutationVector);
 
 			FPixelShaderUtils::AddFullscreenPass(
 				GraphBuilder,
