* Rename PrintLine into PrintLineN to print value with return to line
* Add PrintLine to print a value without return to line

diff --git a/Engine/Shaders/Private/PathTracing/PathTracingCore.ush b/Engine/Shaders/Private/PathTracing/PathTracingCore.ush
--- a/Engine/Shaders/Private/PathTracing/PathTracingCore.ush
+++ b/Engine/Shaders/Private/PathTracing/PathTracingCore.ush
@@ -1840,10 +1840,10 @@ bool PathTracingKernel(inout FPathState PathState, int Bounce)
 
 		// Example of debugging :
 		#if SUBSTRATE_ENABLED
-		PrintLine(Ctx, Payload.GlintValue);
-		PrintLine(Ctx, Payload.GlintUV);
-		PrintLine(Ctx, Payload.GlintUVdx);
-		PrintLine(Ctx, Payload.GlintUVdy);
+		PrintLineN(Ctx, Payload.GlintValue);
+		PrintLineN(Ctx, Payload.GlintUV);
+		PrintLineN(Ctx, Payload.GlintUVdx);
+		PrintLineN(Ctx, Payload.GlintUVdy);
 		#endif
 	}
 	#endif
diff --git a/Engine/Shaders/Private/ShaderPrintCommon.ush b/Engine/Shaders/Private/ShaderPrintCommon.ush
--- a/Engine/Shaders/Private/ShaderPrintCommon.ush
+++ b/Engine/Shaders/Private/ShaderPrintCommon.ush
@@ -1482,4 +1482,5 @@ uint2 GetShaderPrintCurrentPixelCoord()
 // Helper functions
 
 // Print a value onto a single line
-#define PrintLine(Ctx, Value) Print(Ctx, TEXT(#Value), FontWhite); Print(Ctx, TEXT(" : "), FontWhite); Print(Ctx, Value, FontYellow); Newline(Ctx);
\ No newline at end of file
+#define PrintLineN(Ctx, Value) Print(Ctx, TEXT(#Value), FontWhite); Print(Ctx, TEXT(" : "), FontWhite); Print(Ctx, Value, FontYellow); Newline(Ctx);
+#define PrintLine(Ctx, Value) Print(Ctx, TEXT(#Value), FontWhite); Print(Ctx, TEXT(" : "), FontWhite); Print(Ctx, Value, FontYellow);
\ No newline at end of file
diff --git a/Engine/Shaders/Private/Substrate/SubstrateMaterialSampling.usf b/Engine/Shaders/Private/Substrate/SubstrateMaterialSampling.usf
--- a/Engine/Shaders/Private/Substrate/SubstrateMaterialSampling.usf
+++ b/Engine/Shaders/Private/Substrate/SubstrateMaterialSampling.usf
@@ -234,17 +234,17 @@ FSubstrateSampledMaterial GetSampledMaterial(uint2 InPixelCoord, uint ClosureInd
 void Print(FShaderPrintContext Ctx, FSubstrateSampledMaterial Out)
 {
 	Newline(Ctx);
-	PrintLine(Ctx, Out.WorldNormal);
-	PrintLine(Ctx, Out.Roughness);
-	PrintLine(Ctx, Out.DiffuseAlbedo);
-	PrintLine(Ctx, Out.SpecularAlbedo);
-	PrintLine(Ctx, Out.ClosureIndex);
-	PrintLine(Ctx, Out.bNeedResolve);
-	PrintLine(Ctx, Out.bIsSimple);
-	PrintLine(Ctx, Out.bIsValid);
-	PrintLine(Ctx, Out.bIsHair);
-	PrintLine(Ctx, Out.bAllowSpatialFilter);
-	PrintLine(Ctx, Out.bNeedsSeparateSubsurfaceLightAccumulation);
+	PrintLineN(Ctx, Out.WorldNormal);
+	PrintLineN(Ctx, Out.Roughness);
+	PrintLineN(Ctx, Out.DiffuseAlbedo);
+	PrintLineN(Ctx, Out.SpecularAlbedo);
+	PrintLineN(Ctx, Out.ClosureIndex);
+	PrintLineN(Ctx, Out.bNeedResolve);
+	PrintLineN(Ctx, Out.bIsSimple);
+	PrintLineN(Ctx, Out.bIsValid);
+	PrintLineN(Ctx, Out.bIsHair);
+	PrintLineN(Ctx, Out.bAllowSpatialFilter);
+	PrintLineN(Ctx, Out.bNeedsSeparateSubsurfaceLightAccumulation);
 }
 #endif
 
@@ -266,8 +266,8 @@ void SubstrateMaterialSamplingCS(uint2 DispatchThreadId : SV_DispatchThreadID)
 		#if DEBUG_ENABLED
 		FShaderPrintContext Ctx = InitShaderPrintContextAtCursor(PixelCoord, uint2(50, 50));
 		Print(Ctx, TEXT("DEBUG"), FontRed); Newline(Ctx);
-		PrintLine(Ctx, PixelCoord);
-		PrintLine(Ctx, ClosureCount);
+		PrintLineN(Ctx, PixelCoord);
+		PrintLineN(Ctx, ClosureCount);
 		#endif
 
 		FSubstrateSampledMaterial Out = (FSubstrateSampledMaterial)0;
@@ -316,7 +316,7 @@ void SubstrateMaterialSamplingCS(uint2 DispatchThreadId : SV_DispatchThreadID)
 					#if DEBUG_ENABLED
 					for (uint CDFIndex = 0; CDFIndex < ClosureCount; ++CDFIndex)
 					{
-						PrintLine(Ctx, SlabCDF[CDFIndex]);
+						PrintLineN(Ctx, SlabCDF[CDFIndex]);
 					}
 					#endif
 				}
