Substrate - disable spec lut with parameter blending.
A second step is to optionally bake the view dependent part ofthe the spec LUT into F0.

#tests pc dx12 editor
#rnx

diff --git a/Engine/Shaders/Private/Substrate/Substrate.ush b/Engine/Shaders/Private/Substrate/Substrate.ush
--- a/Engine/Shaders/Private/Substrate/Substrate.ush
+++ b/Engine/Shaders/Private/Substrate/Substrate.ush
@@ -2774,7 +2774,6 @@ FSubstrateBSDF SubstrateVerticalLayeringParameterBlending(in FSubstrateBSDF TopB
 
 #if SUBSTRATE_COMPLEXSPECIALPATH
 	BSDF_SETHASGLINT(ResultBSDF, 0); // Disable for now
-	BSDF_SETHASSPECPROFILE(ResultBSDF,	0); // Disable for now
 #endif
 
 	// Saved original Base BBSDF MFP to not be changed be changed by EnableSlabBSDFSimpleVolumetric used to compute medium transmittance and scattering.
@@ -2829,8 +2828,10 @@ FSubstrateBSDF SubstrateVerticalLayeringParameterBlending(in FSubstrateBSDF TopB
 	SLAB_GLINT_UVDDY(ResultBSDF) =		lerp(SLAB_GLINT_UVDDY(TopBSDF),					SLAB_GLINT_UVDDY(BaseBSDF), TopToBottomLerp);
 #endif // SUBSTRATE_GLINTS_ENABLED
 #if SUBSTRATE_SPECPROFILE_ENABLED
-	BSDF_SETHASSPECPROFILE(ResultBSDF,	BSDF_GETHASSPECPROFILE(TopBSDF) | BSDF_GETHASSPECPROFILE(BaseBSDF));
-	SLAB_SPECPROFILEID(ResultBSDF) = 	BSDF_GETHASSPECPROFILE(TopBSDF) ? SLAB_SPECPROFILEID(TopBSDF) : SLAB_SPECPROFILEID(BaseBSDF);
+	BSDF_SETHASSPECPROFILE(ResultBSDF, 0);
+	// SUBSTRATE_TODO: bake view direction specular color into F0
+//	BSDF_SETHASSPECPROFILE(ResultBSDF,	BSDF_GETHASSPECPROFILE(TopBSDF) | BSDF_GETHASSPECPROFILE(BaseBSDF)); 
+//	SLAB_SPECPROFILEID(ResultBSDF) = 	BSDF_GETHASSPECPROFILE(TopBSDF) ? SLAB_SPECPROFILEID(TopBSDF) : SLAB_SPECPROFILEID(BaseBSDF);
 #endif // SUBSTRATE_SPECPROFILE_ENABLED
 
 	// Anisotropy is disabled with vertical blending. 
