Subsutrate - Specular profiles can now be baked into F0/F90 when using parameter blending, in order to retain some of the original look when using parameter blending or fully simplified materials.

#tests pc dx12 editor
#rnx

diff --git a/Engine/Shaders/Private/DBufferDecalShared.ush b/Engine/Shaders/Private/DBufferDecalShared.ush
--- a/Engine/Shaders/Private/DBufferDecalShared.ush
+++ b/Engine/Shaders/Private/DBufferDecalShared.ush
@@ -303,9 +303,11 @@ void ApplyDBufferData(in FSubstrateDBuffer SubstrateDBufferData, inout FSubstrat
 					SubstrateDecalHeader.SubstrateConvertFromDBuffer(LocalSubstrateDBufferData, SubstrateDecalData);
 
 					// Parameters blend attributes
-					const float Coverage = SubstrateDecalData.Layers[l].BSDFs[i].Coverage;
-					const float DummyNoV = 1;
-					BSDF = SubstrateVerticalLayeringParameterBlending(SubstrateDecalData.Layers[l].BSDFs[i], BSDF, BasisIndex, DummyNoV);
+					// SUBSTRATE_TODO BasisIndex is not defined... This is an experimental mode that will need revisiting when the time is right.
+					FSubstrateData SubstrateDataForBSDF;
+					SubstrateDataForBSDF.InlinedBSDF = BSDF;
+					const float DummyNoV = 1.0f;
+					BSDF = SubstrateVerticalLayeringParameterBlendingForDecal(SubstrateDecalData.Layers[l].BSDFs[i], SubstrateDataForBSDF, BasisIndex, DummyNoV, DummyNoV).InlinedBSDF; 
 				}
 				#endif // SUBSTRATE_DBUFFER_BLENDING_MODE
 				break;
diff --git a/Engine/Shaders/Private/Substrate/Substrate.ush b/Engine/Shaders/Private/Substrate/Substrate.ush
--- a/Engine/Shaders/Private/Substrate/Substrate.ush
+++ b/Engine/Shaders/Private/Substrate/Substrate.ush
@@ -2509,6 +2509,19 @@ uint SubstrateMergeSSSType(uint SSSType0, uint SSSType1)
 	return max(SSSType0, SSSType1);
 }
 
+void SubstrateBakeSpecularProfileToReflectivity(inout FSubstrateBSDF BSDF, in float NoV)
+{
+	if (BSDF_GETHASSPECPROFILE(BSDF))
+	{
+		// Evaluate the spec lut
+		const uint SpecProfileID = SLAB_SPECPROFILEID(BSDF);
+		const float3 SpecLUTSample = EvaluateSpecularProfile(SpecProfileID, NoV, 1.0f/*NoL*/, 1.0f/*VoH*/, 1.0f/*NoH*/);
+		SLAB_F0(BSDF) *= SpecLUTSample;
+		SLAB_F90(BSDF) *= SpecLUTSample;
+		BSDF_SETHASSPECPROFILE(BSDF, 0);
+	}
+}
+
 // Note about parameter blending
 //  - We can only parameter blend Slab nodes. This error is handled in the compiler.
 //  - We can only parameter blend one BSDF into another single BSDF. More complex topology are not handled and result in a compiler error. This could be extended to a single BSDF per layer.
@@ -2516,7 +2529,7 @@ uint SubstrateMergeSSSType(uint SSSType0, uint SSSType1)
 //  - Slab node with SSSprofile cannot be blendend. Those will have to be specifically handled by the user and not mixed with any other slab if it can be demoted to parameter blending
 
 // ==> NOTE: Always pair with the compiler behavior in SubstrateCompilationInfoHorizontalMixingParamBlend
-FSubstrateData SubstrateHorizontalMixingParameterBlending(FSubstrateData Background, FSubstrateData Foreground, float HorizontalMixValue, float NormalMixCodeChunk, uint NewNormalIndex, inout uint OutSharedLocalBasisTypes)
+FSubstrateData SubstrateHorizontalMixingParameterBlending(FSubstrateData Background, FSubstrateData Foreground, float HorizontalMixValue, float NormalMixCodeChunk, uint NewNormalIndex, inout uint OutSharedLocalBasisTypes, float BackgroundNoV, float ForegroundNoV)
 {
 	const float ForegroundOtherMixFactor = saturate(NormalMixCodeChunk);	// Horizontal mixing value combined with coverage
 	const float ForegroundLightingMixFactor = saturate(HorizontalMixValue);
@@ -2540,6 +2553,13 @@ FSubstrateData SubstrateHorizontalMixingParameterBlending(FSubstrateData Backgro
 	BSDF_SETHASFUZZ(ResultBSDF,			BSDF_GETHASFUZZ(ForegroundBSDF)					| BSDF_GETHASFUZZ(BackgroundBSDF));
 	BSDF_SETHASTRANSABOVE(ResultBSDF,	BSDF_GETHASTRANSABOVE(ForegroundBSDF)			| BSDF_GETHASTRANSABOVE(BackgroundBSDF));
 
+	// Bake the Specular Profiles into F0/F90 of BSDFs to blend in order to retain some of the original look when using parameter blending or fully simplified materials.
+#if SUBSTRATE_SPECPROFILE_ENABLED
+	SubstrateBakeSpecularProfileToReflectivity(BackgroundBSDF, BackgroundNoV);
+	SubstrateBakeSpecularProfileToReflectivity(ForegroundBSDF, ForegroundNoV);
+	BSDF_SETHASSPECPROFILE(ResultBSDF, 0);
+#endif
+
 	BSDF_SETEMISSIVE(ResultBSDF,		lerp(BSDF_GETEMISSIVE(BackgroundBSDF),			BSDF_GETEMISSIVE(ForegroundBSDF),			ForegroundOtherMixFactor));
 	SLAB_DIFFUSEALBEDO(ResultBSDF) =	lerp(SLAB_DIFFUSEALBEDO(BackgroundBSDF),		SLAB_DIFFUSEALBEDO(ForegroundBSDF),			ForegroundOtherMixFactor);
 	SLAB_F0(ResultBSDF) =				lerp(SLAB_F0(BackgroundBSDF),					SLAB_F0(ForegroundBSDF),					ForegroundOtherMixFactor);
@@ -2560,11 +2580,6 @@ FSubstrateData SubstrateHorizontalMixingParameterBlending(FSubstrateData Backgro
 	SLAB_GLINT_UVDDX(ResultBSDF) =		lerp(SLAB_GLINT_UVDDX(BackgroundBSDF),			SLAB_GLINT_UVDDX(ForegroundBSDF),			ForegroundOtherMixFactor);
 	SLAB_GLINT_UVDDY(ResultBSDF) =		lerp(SLAB_GLINT_UVDDY(BackgroundBSDF),			SLAB_GLINT_UVDDY(ForegroundBSDF),			ForegroundOtherMixFactor);
 #endif // SUBSTRATE_GLINTS_ENABLED
-#if SUBSTRATE_SPECPROFILE_ENABLED
-	BSDF_SETHASSPECPROFILE(ResultBSDF,	BSDF_GETHASSPECPROFILE(ForegroundBSDF)			| BSDF_GETHASSPECPROFILE(BackgroundBSDF));
-	SLAB_SPECPROFILEID(ResultBSDF) =	BSDF_GETHASSPECPROFILE(ForegroundBSDF) ? SLAB_SPECPROFILEID(ForegroundBSDF) : SLAB_SPECPROFILEID(BackgroundBSDF);
-#endif // SUBSTRATE_SPECPROFILE_ENABLED
-	
 
 	if (BSDF_GETSSSTYPE(ResultBSDF) == SSS_TYPE_DIFFUSION_PROFILE)
 	{
@@ -2647,7 +2662,7 @@ float AddParameterBlendingBSDFCoverageToNormalMix(FSubstrateData ASubstrate, FSu
 }
 
 // ==> NOTE: Always pair with the compiler behavior in SubstrateCompilationInfoAddParamBlend
-FSubstrateData SubstrateAddParameterBlending(FSubstrateData A, FSubstrateData B, float AMixFactor, uint NewNormalIndex)
+FSubstrateData SubstrateAddParameterBlending(FSubstrateData A, FSubstrateData B, float AMixFactor, uint NewNormalIndex, float ANoV, float BNoV)
 {
 	FSubstrateData Result = GetInitialisedSubstrateData();
 #define ResultBSDF		Result.InlinedBSDF
@@ -2667,6 +2682,13 @@ FSubstrateData SubstrateAddParameterBlending(FSubstrateData A, FSubstrateData B,
 	BSDF_SETHASF90(ResultBSDF,			BSDF_GETHASF90(BBSDF)							| BSDF_GETHASF90(ABSDF));
 	BSDF_SETHASFUZZ(ResultBSDF,			BSDF_GETHASFUZZ(BBSDF)							| BSDF_GETHASFUZZ(ABSDF));
 
+	// Bake the Specular Profiles into F0/F90 of BSDFs to blend in order to retain some of the original look when using parameter blending or fully simplified materials.
+#if SUBSTRATE_SPECPROFILE_ENABLED
+	SubstrateBakeSpecularProfileToReflectivity(ABSDF, ANoV);
+	SubstrateBakeSpecularProfileToReflectivity(BBSDF, BNoV);
+	BSDF_SETHASSPECPROFILE(ResultBSDF, 0);
+#endif
+
 	// Some parameters will contribute to added luminance: in this case we add them together and saturate to not go out of the safe range
 	BSDF_SETEMISSIVE(ResultBSDF,		BSDF_GETEMISSIVE(ABSDF)								+ BSDF_GETEMISSIVE(ABSDF));
 	SLAB_DIFFUSEALBEDO(ResultBSDF) =	saturate(SLAB_DIFFUSEALBEDO(BBSDF)					+ SLAB_DIFFUSEALBEDO(ABSDF));
@@ -2689,10 +2711,6 @@ FSubstrateData SubstrateAddParameterBlending(FSubstrateData A, FSubstrateData B,
 	SLAB_GLINT_UVDDX(ResultBSDF) =		lerp(SLAB_GLINT_UVDDX(BBSDF),					SLAB_GLINT_UVDDX(ABSDF), AMixFactor);
 	SLAB_GLINT_UVDDY(ResultBSDF) =		lerp(SLAB_GLINT_UVDDY(BBSDF),					SLAB_GLINT_UVDDY(ABSDF), AMixFactor);
 #endif // SUBSTRATE_GLINTS_ENABLED
-#if SUBSTRATE_SPECPROFILE_ENABLED
-	BSDF_SETHASSPECPROFILE(ResultBSDF,	BSDF_GETHASSPECPROFILE(BBSDF) | BSDF_GETHASSPECPROFILE(ABSDF));
-	SLAB_SPECPROFILEID(ResultBSDF) = BSDF_GETHASSPECPROFILE(BBSDF) ? SLAB_SPECPROFILEID(BBSDF) : SLAB_SPECPROFILEID(ABSDF);
-#endif // SUBSTRATE_SPECPROFILE_ENABLED
 
 	if (BSDF_GETSSSTYPE(ResultBSDF) == SSS_TYPE_DIFFUSION_PROFILE)
 	{
@@ -2747,10 +2765,13 @@ float VerticalLayeringParameterBlendingBSDFCoverageToNormalMix(FSubstrateData To
 #endif
 }
 
-// ==> NOTE: Always pair with the compiler behavior in SubstrateCompilationInfoVerticalLayeringParamBlend
-FSubstrateBSDF SubstrateVerticalLayeringParameterBlending(in FSubstrateBSDF TopBSDF, in FSubstrateBSDF BaseBSDF, uint NewNormalIndex, float NoV)
+FSubstrateData SubstrateVerticalLayeringParameterBlending(in FSubstrateData Top, in FSubstrateData Base, uint NewNormalIndex, float TopNoV, float BaseNoV)
 {
-	FSubstrateBSDF ResultBSDF = (FSubstrateBSDF)0;
+	FSubstrateData Result = (FSubstrateData)0;
+#define ResultBSDF		Result.InlinedBSDF
+#define TopBSDF			Top.InlinedBSDF
+#define BaseBSDF		Base.InlinedBSDF
+
 #if SUBSTRATE_INLINE_SHADING
 	const float SafeTopCoverage = saturate(TopBSDF.Coverage);
 	const float SafeBaseCoverage = saturate(BaseBSDF.Coverage);
@@ -2776,6 +2797,13 @@ FSubstrateBSDF SubstrateVerticalLayeringParameterBlending(in FSubstrateBSDF TopB
 	BSDF_SETHASGLINT(ResultBSDF, 0); // Disable for now
 #endif
 
+	// Bake the Specular Profiles into F0/F90 of BSDFs to blend in order to retain some of the original look when using parameter blending or fully simplified materials.
+#if SUBSTRATE_SPECPROFILE_ENABLED
+	SubstrateBakeSpecularProfileToReflectivity(TopBSDF, TopNoV);
+	SubstrateBakeSpecularProfileToReflectivity(BaseBSDF, BaseNoV);
+	BSDF_SETHASSPECPROFILE(ResultBSDF, 0);
+#endif
+
 	// Saved original Base BBSDF MFP to not be changed be changed by EnableSlabBSDFSimpleVolumetric used to compute medium transmittance and scattering.
 	// EnableSlabBSDFSimpleVolumetric will be called at then end of the SubstrateTree as needed later.
 	const float3 OriginalBaseMFP = SLAB_SSSMFP(BaseBSDF);
@@ -2785,14 +2813,14 @@ FSubstrateBSDF SubstrateVerticalLayeringParameterBlending(in FSubstrateBSDF TopB
 	TopBSDF.EnableSlabBSDFSimpleVolumetric();
 	BaseBSDF.EnableSlabBSDFSimpleVolumetric();
 	FParticipatingMedia PM				= SubstrateSlabCreateParticipatingMedia(SLAB_DIFFUSEALBEDO(TopBSDF), SLAB_SSSMFP(TopBSDF));
-	const float3 SlabTransmittanceV		= IsotropicMediumSlabTransmittance(PM, SUBSTRATE_SIMPLEVOLUME_THICKNESS_M, NoV);
+	const float3 SlabTransmittanceV		= IsotropicMediumSlabTransmittance(PM, SUBSTRATE_SIMPLEVOLUME_THICKNESS_M, TopNoV);
 	const float3 SlabTransmittanceN		= IsotropicMediumSlabTransmittance(PM, SUBSTRATE_SIMPLEVOLUME_THICKNESS_M, 1.0f /*Normal incidence i.e., NoL == 1.f*/);
 	const float3 SlabTransmittance		= SlabTransmittanceV * SlabTransmittanceN;
 
 	// We need to convert the scattering from the top layer to an additional albedo without light source direction.
 	// There is no right answer here since we are adding energy without accounting for the directionality of directional or other local lights.
 	// IsotropicMediumSlabEnvDirectionalAlbedo seems to set too much energy so instead the one for punctual light is used instead and looks more reasonable after some empirical tests.
-	float3 TopSlabAlbedo = IsotropicMediumSlabPunctualDirectionalAlbedo(PM, NoV, 1.0f);
+	float3 TopSlabAlbedo = IsotropicMediumSlabPunctualDirectionalAlbedo(PM, TopNoV, 1.0f);
 
 	// SUBSTRATE_TODO account for energy conservation here.
 
@@ -2827,12 +2855,6 @@ FSubstrateBSDF SubstrateVerticalLayeringParameterBlending(in FSubstrateBSDF TopB
 	SLAB_GLINT_UVDDX(ResultBSDF) =		lerp(SLAB_GLINT_UVDDX(TopBSDF),					SLAB_GLINT_UVDDX(BaseBSDF), TopToBottomLerp);
 	SLAB_GLINT_UVDDY(ResultBSDF) =		lerp(SLAB_GLINT_UVDDY(TopBSDF),					SLAB_GLINT_UVDDY(BaseBSDF), TopToBottomLerp);
 #endif // SUBSTRATE_GLINTS_ENABLED
-#if SUBSTRATE_SPECPROFILE_ENABLED
-	BSDF_SETHASSPECPROFILE(ResultBSDF, 0);
-	// SUBSTRATE_TODO: bake view direction specular color into F0
-//	BSDF_SETHASSPECPROFILE(ResultBSDF,	BSDF_GETHASSPECPROFILE(TopBSDF) | BSDF_GETHASSPECPROFILE(BaseBSDF)); 
-//	SLAB_SPECPROFILEID(ResultBSDF) = 	BSDF_GETHASSPECPROFILE(TopBSDF) ? SLAB_SPECPROFILEID(TopBSDF) : SLAB_SPECPROFILEID(BaseBSDF);
-#endif // SUBSTRATE_SPECPROFILE_ENABLED
 
 	// Anisotropy is disabled with vertical blending. 
 	SLAB_ANISOTROPY(ResultBSDF) =		0.f;
@@ -2873,16 +2895,14 @@ FSubstrateBSDF SubstrateVerticalLayeringParameterBlending(in FSubstrateBSDF TopB
 #endif
 
 #endif
-	return ResultBSDF;
-}
 
-FSubstrateData SubstrateVerticalLayeringParameterBlending(FSubstrateData Top, FSubstrateData Base, uint NewNormalIndex, float NoV)
-{
-	FSubstrateData Result = GetInitialisedSubstrateData();
-	Result.InlinedBSDF = SubstrateVerticalLayeringParameterBlending(Top.InlinedBSDF, Base.InlinedBSDF, NewNormalIndex, NoV);
 #if USE_DEVELOPMENT_SHADERS
 	Result.PreviewColor = lerp(Top.PreviewColor, Base.PreviewColor, 0.5);
 #endif
+
+#undef ResultBSDF
+#undef TopBSDF
+#undef BaseBSDF
 	return Result;
 }
 
diff --git a/Engine/Source/Runtime/Core/Private/UObject/UE5MainStreamObjectVersion.cpp b/Engine/Source/Runtime/Core/Private/UObject/UE5MainStreamObjectVersion.cpp
--- a/Engine/Source/Runtime/Core/Private/UObject/UE5MainStreamObjectVersion.cpp
+++ b/Engine/Source/Runtime/Core/Private/UObject/UE5MainStreamObjectVersion.cpp
@@ -20,7 +20,7 @@ TMap<FGuid, FGuid> FUE5MainStreamObjectVersion::GetSystemGuids()
 	SystemGuids.Add(DevGuids.POSESEARCHDB_DERIVEDDATA_VER, FGuid("4E595C2AC5E947D6BA9ABC874353E5BC"));
 	SystemGuids.Add(DevGuids.GROOM_BINDING_DERIVED_DATA_VERSION, FGuid("156678D4F9084D7CAE98FADC6AB93573"));
 	SystemGuids.Add(DevGuids.GROOM_DERIVED_DATA_VERSION, FGuid("00B8207248CC44B88F7DEC8F328A5D8B"));
-	SystemGuids.Add(DevGuids.MaterialTranslationDDCVersion, FGuid("80B255B52A0BB2DFB621304F52C68B29"));
+	SystemGuids.Add(DevGuids.MaterialTranslationDDCVersion, FGuid("2EBF326928DF4469BEA514CDBC1081AE"));
 
 	return SystemGuids;
 }
diff --git a/Engine/Source/Runtime/Engine/Private/Materials/HLSLMaterialTranslator.cpp b/Engine/Source/Runtime/Engine/Private/Materials/HLSLMaterialTranslator.cpp
--- a/Engine/Source/Runtime/Engine/Private/Materials/HLSLMaterialTranslator.cpp
+++ b/Engine/Source/Runtime/Engine/Private/Materials/HLSLMaterialTranslator.cpp
@@ -14296,19 +14296,22 @@ int32 FHLSLMaterialTranslator::SubstrateHorizontalMixing(int32 Background, int32
 }
 
 int32 FHLSLMaterialTranslator::SubstrateHorizontalMixingParameterBlending(
-	int32 Background, int32 Foreground, int32 HorizontalMixCodeChunk, int32 NormalMixCodeChunk, const FString& SharedLocalBasisIndexMacro, FSubstrateOperator* PromoteToOperator)
+	int32 Background, int32 Foreground, int32 HorizontalMixCodeChunk, int32 NormalMixCodeChunk, const FString& SharedLocalBasisIndexMacro, int32 BackgroundBSDFNormalCodeChunk, int32 ForegroundBSDFNormalCodeChunk, FSubstrateOperator* PromoteToOperator)
 {
 	if (Foreground == INDEX_NONE || Background == INDEX_NONE || HorizontalMixCodeChunk == INDEX_NONE || NormalMixCodeChunk == INDEX_NONE)
 	{
 		return INDEX_NONE;
 	}
 
+	const int32 BackgroundNoV = Saturate(Dot(BackgroundBSDFNormalCodeChunk, CameraVector()));
+	const int32 ForegroundNoV = Saturate(Dot(ForegroundBSDFNormalCodeChunk, CameraVector()));
+
 	if (PromoteToOperator)
 	{
 		check(PromoteToOperator->Index != INDEX_NONE);
 		check(PromoteToOperator->BSDFIndex != INDEX_NONE);
 		return AddCodeChunk(
-			MCT_Substrate, TEXT("Parameters.%s.PromoteParameterBlendedBSDFToOperator(SubstrateHorizontalMixingParameterBlending(%s, %s, %s, %s, %s, Parameters.%s.Types), %u, %u, %u, %u)"),
+			MCT_Substrate, TEXT("Parameters.%s.PromoteParameterBlendedBSDFToOperator(SubstrateHorizontalMixingParameterBlending(%s, %s, %s, %s, %s, Parameters.%s.Types, %s, %s), %u, %u, %u, %u)"),
 			*GetParametersSubstrateTreeName(CurrentSubstrateCompilationContext),
 			*GetParameterCode(Background),
 			*GetParameterCode(Foreground),
@@ -14316,6 +14319,8 @@ int32 FHLSLMaterialTranslator::SubstrateHorizontalMixingParameterBlending(
 			*GetParameterCode(NormalMixCodeChunk),
 			*SharedLocalBasisIndexMacro,
 			*GetParametersSharedLocalBasesName(CurrentSubstrateCompilationContext),
+			*GetParameterCode(BackgroundNoV),
+			*GetParameterCode(ForegroundNoV),
 			PromoteToOperator->Index,
 			PromoteToOperator->BSDFIndex,
 			PromoteToOperator->LayerDepth,
@@ -14324,13 +14329,15 @@ int32 FHLSLMaterialTranslator::SubstrateHorizontalMixingParameterBlending(
 	}
 
 	return AddCodeChunk(
-		MCT_Substrate, TEXT("SubstrateHorizontalMixingParameterBlending(%s, %s, %s, %s, %s, Parameters.%s.Types)"),
+		MCT_Substrate, TEXT("SubstrateHorizontalMixingParameterBlending(%s, %s, %s, %s, %s, Parameters.%s.Types, %s, %s)"),
 		*GetParameterCode(Background),
 		*GetParameterCode(Foreground),
 		*GetParameterCode(HorizontalMixCodeChunk),
 		*GetParameterCode(NormalMixCodeChunk),
 		*SharedLocalBasisIndexMacro,
-		*GetParametersSharedLocalBasesName(CurrentSubstrateCompilationContext)
+		*GetParametersSharedLocalBasesName(CurrentSubstrateCompilationContext),
+		*GetParameterCode(BackgroundNoV),
+		*GetParameterCode(ForegroundNoV)
 	);
 }
 
@@ -14352,7 +14359,7 @@ int32 FHLSLMaterialTranslator::SubstrateVerticalLayering(int32 Top, int32 Base,
 	);
 }
 
-int32 FHLSLMaterialTranslator::SubstrateVerticalLayeringParameterBlending(int32 Top, int32 Base, int32 Thickness, const FString& SharedLocalBasisIndexMacro, int32 TopBSDFNormalCodeChunk, FSubstrateOperator* PromoteToOperator)
+int32 FHLSLMaterialTranslator::SubstrateVerticalLayeringParameterBlending(int32 Top, int32 Base, int32 Thickness, const FString& SharedLocalBasisIndexMacro, int32 TopBSDFNormalCodeChunk, int32 BaseBSDFNormalCodeChunk, FSubstrateOperator* PromoteToOperator)
 {
 	if (Top == INDEX_NONE || Base == INDEX_NONE)
 	{
@@ -14361,18 +14368,21 @@ int32 FHLSLMaterialTranslator::SubstrateVerticalLayeringParameterBlending(int32
 
 	const FString ThicknessCode = Thickness != INDEX_NONE ? GetParameterCode(Thickness) : TEXT("NONE");
 
+	const int32 TopNoV  = Saturate(Dot(TopBSDFNormalCodeChunk, CameraVector()));
+	const int32 BaseNoV = Saturate(Dot(BaseBSDFNormalCodeChunk, CameraVector()));
+
 	if (PromoteToOperator)
 	{
 		check(PromoteToOperator->Index != INDEX_NONE);
 		check(PromoteToOperator->BSDFIndex != INDEX_NONE);
 		return AddCodeChunk(
-			MCT_Substrate, TEXT("Parameters.%s.PromoteParameterBlendedBSDFToOperator(SubstrateVerticalLayeringParameterBlending(%s, %s, %s, dot(%s, %s)), %u, %u, %u, %u) /* Thickness = %s */"),
+			MCT_Substrate, TEXT("Parameters.%s.PromoteParameterBlendedBSDFToOperator(SubstrateVerticalLayeringParameterBlending(%s, %s, %s, %s, %s), %u, %u, %u, %u) /* Thickness = %s */"),
 			*GetParametersSubstrateTreeName(CurrentSubstrateCompilationContext),
 			*GetParameterCode(Top),
 			*GetParameterCode(Base),
 			*SharedLocalBasisIndexMacro,
-			*GetParameterCode(TopBSDFNormalCodeChunk),
-			*GetParameterCode(CameraVector()),
+			*GetParameterCode(TopNoV),
+			*GetParameterCode(BaseNoV),
 			PromoteToOperator->Index,
 			PromoteToOperator->BSDFIndex,
 			PromoteToOperator->LayerDepth,
@@ -14382,12 +14392,12 @@ int32 FHLSLMaterialTranslator::SubstrateVerticalLayeringParameterBlending(int32
 	}
 
 	return AddCodeChunk(
-		MCT_Substrate, TEXT("SubstrateVerticalLayeringParameterBlending(%s, %s, %s, dot(%s, %s)) /* Thickness = %s */"),
+		MCT_Substrate, TEXT("SubstrateVerticalLayeringParameterBlending(%s.InlinedBSDF, %s.InlinedBSDF, %s, %s, %s) /* Thickness = %s */"),
 		*GetParameterCode(Top),
 		*GetParameterCode(Base),
 		*SharedLocalBasisIndexMacro,
-		*GetParameterCode(TopBSDFNormalCodeChunk),
-		*GetParameterCode(CameraVector()),
+		*GetParameterCode(TopNoV),
+		*GetParameterCode(BaseNoV),
 		*ThicknessCode
 	);
 }
@@ -14408,24 +14418,29 @@ int32 FHLSLMaterialTranslator::SubstrateAdd(int32 A, int32 B, int OperatorIndex,
 	);
 }
 
-int32 FHLSLMaterialTranslator::SubstrateAddParameterBlending(int32 A, int32 B, int32 AMixWeight, const FString& SharedLocalBasisIndexMacro, FSubstrateOperator* PromoteToOperator)
+int32 FHLSLMaterialTranslator::SubstrateAddParameterBlending(int32 A, int32 B, int32 AMixWeight, const FString& SharedLocalBasisIndexMacro, int32 ABSDFNormalCodeChunk, int32 BBSDFNormalCodeChunk, FSubstrateOperator* PromoteToOperator)
 {
 	if (A == INDEX_NONE || B == INDEX_NONE)
 	{
 		return INDEX_NONE;
 	}
 
+	const int32 ANoV = Saturate(Dot(ABSDFNormalCodeChunk, CameraVector()));
+	const int32 BNoV = Saturate(Dot(BBSDFNormalCodeChunk, CameraVector()));
+
 	if (PromoteToOperator)
 	{
 		check(PromoteToOperator->Index != INDEX_NONE);
 		check(PromoteToOperator->BSDFIndex != INDEX_NONE);
 		return AddCodeChunk(
-			MCT_Substrate, TEXT("Parameters.%s.PromoteParameterBlendedBSDFToOperator(SubstrateAddParameterBlending(%s, %s, %s, %s), %u, %u, %u, %u)"),
+			MCT_Substrate, TEXT("Parameters.%s.PromoteParameterBlendedBSDFToOperator(SubstrateAddParameterBlending(%s, %s, %s, %s, %s, %s), %u, %u, %u, %u)"),
 			*GetParametersSubstrateTreeName(CurrentSubstrateCompilationContext),
 			*GetParameterCode(A),
 			*GetParameterCode(B),
 			*GetParameterCode(AMixWeight),
 			*SharedLocalBasisIndexMacro,
+			*GetParameterCode(ANoV),
+			*GetParameterCode(BNoV),
 			PromoteToOperator->Index,
 			PromoteToOperator->BSDFIndex,
 			PromoteToOperator->LayerDepth,
@@ -14434,11 +14449,13 @@ int32 FHLSLMaterialTranslator::SubstrateAddParameterBlending(int32 A, int32 B, i
 	}
 
 	return AddCodeChunk(
-		MCT_Substrate, TEXT("SubstrateAddParameterBlending(%s, %s, %s, %s)"),
+		MCT_Substrate, TEXT("SubstrateAddParameterBlending(%s, %s, %s, %s, %s, %s)"),
 		*GetParameterCode(A),
 		*GetParameterCode(B),
 		*GetParameterCode(AMixWeight),
-		*SharedLocalBasisIndexMacro
+		*SharedLocalBasisIndexMacro,
+		*GetParameterCode(ANoV),
+		*GetParameterCode(BNoV)
 	);
 }
 
diff --git a/Engine/Source/Runtime/Engine/Private/Materials/HLSLMaterialTranslator.h b/Engine/Source/Runtime/Engine/Private/Materials/HLSLMaterialTranslator.h
--- a/Engine/Source/Runtime/Engine/Private/Materials/HLSLMaterialTranslator.h
+++ b/Engine/Source/Runtime/Engine/Private/Materials/HLSLMaterialTranslator.h
@@ -1349,11 +1349,11 @@ protected:
 		int32 EmissiveColor, int32 TopMaterialOpacity, int32 WaterAlbedo, int32 WaterExtinction, int32 WaterPhaseG, 
 		int32 ColorScaleBehindWater, int32 Normal, const FString& SharedLocalBasisIndexMacro, FSubstrateOperator* PromoteToOperator) override;
 	virtual int32 SubstrateHorizontalMixing(int32 Background, int32 Foreground, int32 Mix, int OperatorIndex, uint32 MaxDistanceFromLeaves) override;
-	virtual int32 SubstrateHorizontalMixingParameterBlending(int32 Background, int32 Foreground, int32 HorizontalMixCodeChunk, int32 NormalMixCodeChunk, const FString& SharedLocalBasisIndexMacro, FSubstrateOperator* PromoteToOperator) override;
+	virtual int32 SubstrateHorizontalMixingParameterBlending(int32 Background, int32 Foreground, int32 HorizontalMixCodeChunk, int32 NormalMixCodeChunk, const FString& SharedLocalBasisIndexMacro, int32 BackgroundBSDFNormalCodeChunk, int32 ForegroundBSDFNormalCodeChunk, FSubstrateOperator* PromoteToOperator) override;
 	virtual int32 SubstrateVerticalLayering(int32 Top, int32 Base, int32 Thickness, int OperatorIndex, uint32 MaxDistanceFromLeaves) override;
-	virtual int32 SubstrateVerticalLayeringParameterBlending(int32 Top, int32 Base, int32 Thickness, const FString& SharedLocalBasisIndexMacro, int32 TopBSDFNormalCodeChunk, FSubstrateOperator* PromoteToOperator) override;
+	virtual int32 SubstrateVerticalLayeringParameterBlending(int32 Top, int32 Base, int32 Thickness, const FString& SharedLocalBasisIndexMacro, int32 TopBSDFNormalCodeChunk, int32 BaseBSDFNormalCodeChunk, FSubstrateOperator* PromoteToOperator) override;
 	virtual int32 SubstrateAdd(int32 A, int32 B, int OperatorIndex, uint32 MaxDistanceFromLeaves) override;
-	virtual int32 SubstrateAddParameterBlending(int32 A, int32 B, int32 AMixWeight, const FString& SharedLocalBasisIndexMacro, FSubstrateOperator* PromoteToOperator) override;
+	virtual int32 SubstrateAddParameterBlending(int32 A, int32 B, int32 AMixWeight, const FString& SharedLocalBasisIndexMacro, int32 ABSDFNormalCodeChunk, int32 BBSDFNormalCodeChunk, FSubstrateOperator* PromoteToOperator) override;
 	virtual int32 SubstrateWeight(int32 A, int32 Weight, int OperatorIndex, uint32 MaxDistanceFromLeaves) override;
 	virtual int32 SubstrateWeightParameterBlending(int32 A, int32 Weight, FSubstrateOperator* PromoteToOperator) override;
 	virtual int32 SubstrateSelectParameterBlending(int32 A, int32 B, int32 SelectValue, const FString& SharedLocalBasisIndexMacro, FSubstrateOperator* PromoteToOperator) override;
diff --git a/Engine/Source/Runtime/Engine/Private/Materials/MaterialExpressionSubstrate.cpp b/Engine/Source/Runtime/Engine/Private/Materials/MaterialExpressionSubstrate.cpp
--- a/Engine/Source/Runtime/Engine/Private/Materials/MaterialExpressionSubstrate.cpp
+++ b/Engine/Source/Runtime/Engine/Private/Materials/MaterialExpressionSubstrate.cpp
@@ -2449,6 +2449,7 @@ int32 UMaterialExpressionSubstrateHorizontalMixing::Compile(class FMaterialCompi
 
 		OutputCodeChunk = Compiler->SubstrateHorizontalMixingParameterBlending(
 			BackgroundCodeChunk, ForegroundCodeChunk, HorizontalMixCodeChunk, NormalMixCodeChunk, Compiler->GetSubstrateSharedLocalBasisIndexMacro(NewRegisteredSharedLocalBasis),
+			BackgroundBSDFOperator->BSDFRegisteredSharedLocalBasis.NormalCodeChunk, ForegroundBSDFOperator->BSDFRegisteredSharedLocalBasis.NormalCodeChunk,
 			SubstrateOperator.bRootOfParameterBlendingSubTree ? &SubstrateOperator : nullptr);
 
 		// Propagate the parameter blended normal
@@ -2629,7 +2630,8 @@ int32 UMaterialExpressionSubstrateVerticalLayering::Compile(class FMaterialCompi
 		const FSubstrateRegisteredSharedLocalBasis NewRegisteredSharedLocalBasis = SubstrateCompilationInfoCreateSharedLocalBasis(Compiler, NewNormalCodeChunk, NewTangentCodeChunk);
 
 		OutputCodeChunk = Compiler->SubstrateVerticalLayeringParameterBlending(
-			TopCodeChunk, BaseCodeChunk, ThicknessCodeChunk, Compiler->GetSubstrateSharedLocalBasisIndexMacro(NewRegisteredSharedLocalBasis), TopBSDFOperator->BSDFRegisteredSharedLocalBasis.NormalCodeChunk,
+			TopCodeChunk, BaseCodeChunk, ThicknessCodeChunk, Compiler->GetSubstrateSharedLocalBasisIndexMacro(NewRegisteredSharedLocalBasis), 
+			TopBSDFOperator->BSDFRegisteredSharedLocalBasis.NormalCodeChunk, BaseBSDFOperator->BSDFRegisteredSharedLocalBasis.NormalCodeChunk,
 			SubstrateOperator.bRootOfParameterBlendingSubTree ? &SubstrateOperator : nullptr);
 
 		// Propagate the parameter blended normal
@@ -2831,6 +2833,7 @@ int32 UMaterialExpressionSubstrateAdd::Compile(class FMaterialCompiler* Compiler
 
 		OutputCodeChunk = Compiler->SubstrateAddParameterBlending(
 			ACodeChunk, BCodeChunk, ANormalMixCodeChunk, Compiler->GetSubstrateSharedLocalBasisIndexMacro(NewRegisteredSharedLocalBasis),
+			ABSDFOperator->BSDFRegisteredSharedLocalBasis.NormalCodeChunk, BBSDFOperator->BSDFRegisteredSharedLocalBasis.NormalCodeChunk,
 			SubstrateOperator.bRootOfParameterBlendingSubTree ? &SubstrateOperator : nullptr);
 
 		// Propagate the parameter blended normal
diff --git a/Engine/Source/Runtime/Engine/Public/MaterialCompiler.h b/Engine/Source/Runtime/Engine/Public/MaterialCompiler.h
--- a/Engine/Source/Runtime/Engine/Public/MaterialCompiler.h
+++ b/Engine/Source/Runtime/Engine/Public/MaterialCompiler.h
@@ -737,11 +737,11 @@ public:
 		int32 BaseColor, int32 Metallic, int32 Specular, int32 Roughness, int32 EmissiveColor, int32 TopMaterialOpacity, 
 		int32 WaterAlbedo, int32 WaterExtinction, int32 WaterPhaseG, int32 ColorScaleBehindWater, int32 Normal, const FString& SharedLocalBasisIndexMacro, FSubstrateOperator* PromoteToOperator) = 0;
 	virtual int32 SubstrateHorizontalMixing(int32 Background, int32 Foreground, int32 Mix, int OperatorIndex, uint32 MaxDistanceFromLeaves) = 0;
-	virtual int32 SubstrateHorizontalMixingParameterBlending(int32 Background, int32 Foreground, int32 HorizontalMixCodeChunk, int32 NormalMixCodeChunk, const FString& SharedLocalBasisIndexMacro, FSubstrateOperator* PromoteToOperator) = 0;
+	virtual int32 SubstrateHorizontalMixingParameterBlending(int32 Background, int32 Foreground, int32 HorizontalMixCodeChunk, int32 NormalMixCodeChunk, const FString& SharedLocalBasisIndexMacro, int32 BackgroundBSDFNormalCodeChunk, int32 ForegroundBSDFNormalCodeChunk, FSubstrateOperator* PromoteToOperator) = 0;
 	virtual int32 SubstrateVerticalLayering(int32 Top, int32 Base, int32 Thickness, int OperatorIndex, uint32 MaxDistanceFromLeaves) = 0;
-	virtual int32 SubstrateVerticalLayeringParameterBlending(int32 Top, int32 Base, int32 Thickness, const FString& SharedLocalBasisIndexMacro, int32 TopBSDFNormalCodeChunk, FSubstrateOperator* PromoteToOperator) = 0;
+	virtual int32 SubstrateVerticalLayeringParameterBlending(int32 Top, int32 Base, int32 Thickness, const FString& SharedLocalBasisIndexMacro, int32 TopBSDFNormalCodeChunk, int32 BaseBSDFNormalCodeChunk, FSubstrateOperator* PromoteToOperator) = 0;
 	virtual int32 SubstrateAdd(int32 A, int32 B, int OperatorIndex, uint32 MaxDistanceFromLeaves) = 0;
-	virtual int32 SubstrateAddParameterBlending(int32 A, int32 B, int32 AMixWeight, const FString& SharedLocalBasisIndexMacro, FSubstrateOperator* PromoteToOperator) = 0;
+	virtual int32 SubstrateAddParameterBlending(int32 A, int32 B, int32 AMixWeight, const FString& SharedLocalBasisIndexMacro, int32 ABSDFNormalCodeChunk, int32 BBSDFNormalCodeChunk, FSubstrateOperator* PromoteToOperator) = 0;
 	virtual int32 SubstrateWeight(int32 A, int32 Weight, int OperatorIndex, uint32 MaxDistanceFromLeaves) = 0;
 	virtual int32 SubstrateWeightParameterBlending(int32 A, int32 Weight, FSubstrateOperator* PromoteToOperator) = 0;
 	virtual int32 SubstrateSelectParameterBlending(int32 A, int32 B, int32 SelectValue, const FString& SharedLocalBasisIndexMacro, FSubstrateOperator* PromoteToOperator) = 0;
@@ -1486,9 +1486,9 @@ public:
 		return Compiler->SubstrateHorizontalMixing(Background, Foreground, Mix, OperatorIndex, MaxDistanceFromLeaves);
 	}
 
-	virtual int32 SubstrateHorizontalMixingParameterBlending(int32 Background, int32 Foreground, int32 HorizontalMixCodeChunk, int32 NormalMixCodeChunk, const FString& SharedLocalBasisIndexMacro, FSubstrateOperator* PromoteToOperator) override
+	virtual int32 SubstrateHorizontalMixingParameterBlending(int32 Background, int32 Foreground, int32 HorizontalMixCodeChunk, int32 NormalMixCodeChunk, const FString& SharedLocalBasisIndexMacro, int32 BackgroundBSDFNormalCodeChunk, int32 ForegroundBSDFNormalCodeChunk, FSubstrateOperator* PromoteToOperator) override
 	{
-		return Compiler->SubstrateHorizontalMixingParameterBlending(Background, Foreground, HorizontalMixCodeChunk, NormalMixCodeChunk, SharedLocalBasisIndexMacro, PromoteToOperator);
+		return Compiler->SubstrateHorizontalMixingParameterBlending(Background, Foreground, HorizontalMixCodeChunk, NormalMixCodeChunk, SharedLocalBasisIndexMacro, BackgroundBSDFNormalCodeChunk, ForegroundBSDFNormalCodeChunk, PromoteToOperator);
 	}
 
 	virtual int32 SubstrateVerticalLayering(int32 Top, int32 Base, int32 Thickness, int OperatorIndex, uint32 MaxDistanceFromLeaves) override
@@ -1496,9 +1496,9 @@ public:
 		return Compiler->SubstrateVerticalLayering(Top, Base, Thickness, OperatorIndex, MaxDistanceFromLeaves);
 	}
 
-	virtual int32 SubstrateVerticalLayeringParameterBlending(int32 Top, int32 Base, int32 Thickness, const FString& SharedLocalBasisIndexMacro, int32 TopBSDFNormalCodeChunk, FSubstrateOperator* PromoteToOperator) override
+	virtual int32 SubstrateVerticalLayeringParameterBlending(int32 Top, int32 Base, int32 Thickness, const FString& SharedLocalBasisIndexMacro, int32 TopBSDFNormalCodeChunk, int32 BaseBSDFNormalCodeChunk, FSubstrateOperator* PromoteToOperator) override
 	{
-		return Compiler->SubstrateVerticalLayeringParameterBlending(Top, Base, Thickness, SharedLocalBasisIndexMacro, TopBSDFNormalCodeChunk, PromoteToOperator);
+		return Compiler->SubstrateVerticalLayeringParameterBlending(Top, Base, Thickness, SharedLocalBasisIndexMacro, TopBSDFNormalCodeChunk, BaseBSDFNormalCodeChunk, PromoteToOperator);
 	}
 
 	virtual int32 SubstrateAdd(int32 A, int32 B, int OperatorIndex, uint32 MaxDistanceFromLeaves) override
@@ -1506,9 +1506,9 @@ public:
 		return Compiler->SubstrateAdd(A, B, OperatorIndex, MaxDistanceFromLeaves);
 	}
 
-	virtual int32 SubstrateAddParameterBlending(int32 A, int32 B, int32 AMixWeight, const FString& SharedLocalBasisIndexMacro, FSubstrateOperator* PromoteToOperator) override
+	virtual int32 SubstrateAddParameterBlending(int32 A, int32 B, int32 AMixWeight, const FString& SharedLocalBasisIndexMacro, int32 ABSDFNormalCodeChunk, int32 BBSDFNormalCodeChunk, FSubstrateOperator* PromoteToOperator) override
 	{
-		return Compiler->SubstrateAddParameterBlending(A, B, AMixWeight, SharedLocalBasisIndexMacro, PromoteToOperator);
+		return Compiler->SubstrateAddParameterBlending(A, B, AMixWeight, SharedLocalBasisIndexMacro, ABSDFNormalCodeChunk, BBSDFNormalCodeChunk, PromoteToOperator);
 	}
 
 	virtual int32 SubstrateWeight(int32 A, int32 Weight, int OperatorIndex, uint32 MaxDistanceFromLeaves) override
