Substrate - BSDF now account for glint reflectivity when simplified or parameter blended.
Parameter blending also disables glint rendering now, since this cannot really be correct given the glint parameterisation.
Moved  SubstrateBakeXXX function as member of BSDF.

#rnx
#tests PC DX12 editor

diff --git a/Engine/Shaders/Private/Substrate/Substrate.ush b/Engine/Shaders/Private/Substrate/Substrate.ush
--- a/Engine/Shaders/Private/Substrate/Substrate.ush
+++ b/Engine/Shaders/Private/Substrate/Substrate.ush
@@ -218,6 +218,10 @@ struct FSubstrateBSDF
 
 	// Return true if the BSDF transmits and scatters light on the opposite side
 	bool HasBackScattering();
+
+	void SubstrateBakeSpecularProfileToReflectivity(in float NoV);
+
+	void SubstrateBakeGlintToReflectivity();
 };
 
 ////////////////////////////////////////////////////////////////////////////////////////////////////////////////////////////////////////////////////////////////////////////////////////////////
@@ -1422,6 +1426,7 @@ FSubstrateData InternalGetSubstrateSlabBSDF(
 		ThicknessCm = SUBSTRATE_SIMPLEVOLUME_THICKNESS_CM;
 	}
 
+	// When Glints are enabled, we need to not keep haziness even R0=R1 to avoid pops (second lobe disapearing, more energy from glints).
 	const bool bHasGlint = SUBSTRATE_COMPLEXSPECIALPATH && GlintValue < 1.0f;
 	const uint bHasHaziness = (((Haziness.Roughness != Roughness || bHasGlint) && Haziness.Weight > 0.0f) || Haziness.bSimpleClearCoat) ? 1 : 0;
 	const uint bHasAnisotropy = Anisotropy != 0 ? 1 : 0;
@@ -1479,7 +1484,7 @@ FSubstrateData InternalGetSubstrateSlabBSDF(
 
 	SLAB_DIFFUSEALBEDO(SubstrateData.InlinedBSDF)			= DiffuseAlbedo;
 	SLAB_F0(SubstrateData.InlinedBSDF)						= F0;
-	SLAB_F90(SubstrateData.InlinedBSDF)					= F90;
+	SLAB_F90(SubstrateData.InlinedBSDF)						= F90;
 
 	SLAB_ROUGHNESS			(SubstrateData.InlinedBSDF)		= Roughness;
 
@@ -1518,6 +1523,18 @@ FSubstrateData InternalGetSubstrateSlabBSDF(
 	SubstrateData.PreviewColor = lerp(DiffuseAlbedo, F0, F0RGBToMetallic(F0));
 #endif
 
+#if SUBSTRATE_USE_FULLYSIMPLIFIED_MATERIAL || !PLATFORM_ENABLES_SUBSTRATE_GLINTS
+	if (GlintValue < 1)
+	{
+		// In this case, SUBSTRATE_GLINTS_ENABLED is typically 0. So no glint data will be written to the Substrate buffer.
+		// So it is here that we apply glint coverage to reflectivity, for reflection to better match in Lumen reflection, ray tracing or lower platforms.
+		BSDF_SETHASGLINT(SubstrateData.InlinedBSDF, 0);
+		GlintValue = saturate(GlintValue);
+		SLAB_F0(SubstrateData.InlinedBSDF)  *= GlintValue;
+		SLAB_F90(SubstrateData.InlinedBSDF) *= GlintValue;
+	}
+#endif
+
 	return SubstrateData;
 }
 
@@ -2509,19 +2526,6 @@ uint SubstrateMergeSSSType(uint SSSType0, uint SSSType1)
 	return max(SSSType0, SSSType1);
 }
 
-void SubstrateBakeSpecularProfileToReflectivity(inout FSubstrateBSDF BSDF, in float NoV)
-{
-	if (BSDF_GETHASSPECPROFILE(BSDF))
-	{
-		// Evaluate the spec lut
-		const uint SpecProfileID = SLAB_SPECPROFILEID(BSDF);
-		const float3 SpecLUTSample = EvaluateSpecularProfile(SpecProfileID, NoV, 1.0f/*NoL*/, 1.0f/*VoH*/, 1.0f/*NoH*/);
-		SLAB_F0(BSDF) *= SpecLUTSample;
-		SLAB_F90(BSDF) *= SpecLUTSample;
-		BSDF_SETHASSPECPROFILE(BSDF, 0);
-	}
-}
-
 // Note about parameter blending
 //  - We can only parameter blend Slab nodes. This error is handled in the compiler.
 //  - We can only parameter blend one BSDF into another single BSDF. More complex topology are not handled and result in a compiler error. This could be extended to a single BSDF per layer.
@@ -2555,11 +2559,17 @@ FSubstrateData SubstrateHorizontalMixingParameterBlending(FSubstrateData Backgro
 
 	// Bake the Specular Profiles into F0/F90 of BSDFs to blend in order to retain some of the original look when using parameter blending or fully simplified materials.
 #if SUBSTRATE_SPECPROFILE_ENABLED
-	SubstrateBakeSpecularProfileToReflectivity(BackgroundBSDF, BackgroundNoV);
-	SubstrateBakeSpecularProfileToReflectivity(ForegroundBSDF, ForegroundNoV);
+	BackgroundBSDF.SubstrateBakeSpecularProfileToReflectivity(BackgroundNoV);
+	ForegroundBSDF.SubstrateBakeSpecularProfileToReflectivity(ForegroundNoV);
 	BSDF_SETHASSPECPROFILE(ResultBSDF, 0);
 #endif
 
+#if SUBSTRATE_GLINTS_ENABLED 
+	ForegroundBSDF.SubstrateBakeGlintToReflectivity();
+	BackgroundBSDF.SubstrateBakeGlintToReflectivity();
+	BSDF_SETHASGLINT(ResultBSDF, 0);
+#endif // SUBSTRATE_GLINTS_ENABLED
+
 	BSDF_SETEMISSIVE(ResultBSDF,		lerp(BSDF_GETEMISSIVE(BackgroundBSDF),			BSDF_GETEMISSIVE(ForegroundBSDF),			ForegroundOtherMixFactor));
 	SLAB_DIFFUSEALBEDO(ResultBSDF) =	lerp(SLAB_DIFFUSEALBEDO(BackgroundBSDF),		SLAB_DIFFUSEALBEDO(ForegroundBSDF),			ForegroundOtherMixFactor);
 	SLAB_F0(ResultBSDF) =				lerp(SLAB_F0(BackgroundBSDF),					SLAB_F0(ForegroundBSDF),					ForegroundOtherMixFactor);
@@ -2573,14 +2583,6 @@ FSubstrateData SubstrateHorizontalMixingParameterBlending(FSubstrateData Backgro
 	SLAB_FUZZ_ROUGHNESS(ResultBSDF)	=	lerp(SLAB_FUZZ_ROUGHNESS(BackgroundBSDF),		SLAB_FUZZ_ROUGHNESS(ForegroundBSDF),		ForegroundOtherMixFactor);
 	SLAB_HAZINESS(ResultBSDF) =			LerpHazinessParameterBlending(BackgroundBSDF,	ForegroundBSDF,								ForegroundOtherMixFactor);
 
-#if SUBSTRATE_GLINTS_ENABLED
-	BSDF_SETHASGLINT(ResultBSDF,		BSDF_GETHASGLINT(ForegroundBSDF)				| BSDF_GETHASGLINT(BackgroundBSDF));
-	SLAB_GLINT_UV(ResultBSDF) =			lerp(SLAB_GLINT_UV(BackgroundBSDF),				SLAB_GLINT_UV(ForegroundBSDF),				ForegroundOtherMixFactor);
-	SLAB_GLINT_VALUE(ResultBSDF) =		lerp(SLAB_GLINT_VALUE(BackgroundBSDF),			SLAB_GLINT_VALUE(ForegroundBSDF),			ForegroundOtherMixFactor);
-	SLAB_GLINT_UVDDX(ResultBSDF) =		lerp(SLAB_GLINT_UVDDX(BackgroundBSDF),			SLAB_GLINT_UVDDX(ForegroundBSDF),			ForegroundOtherMixFactor);
-	SLAB_GLINT_UVDDY(ResultBSDF) =		lerp(SLAB_GLINT_UVDDY(BackgroundBSDF),			SLAB_GLINT_UVDDY(ForegroundBSDF),			ForegroundOtherMixFactor);
-#endif // SUBSTRATE_GLINTS_ENABLED
-
 	if (BSDF_GETSSSTYPE(ResultBSDF) == SSS_TYPE_DIFFUSION_PROFILE)
 	{
 		const bool bForegroundHasSSSProfile = BSDF_GETSSSTYPE(ForegroundBSDF) == SSS_TYPE_DIFFUSION_PROFILE;
@@ -2684,11 +2686,17 @@ FSubstrateData SubstrateAddParameterBlending(FSubstrateData A, FSubstrateData B,
 
 	// Bake the Specular Profiles into F0/F90 of BSDFs to blend in order to retain some of the original look when using parameter blending or fully simplified materials.
 #if SUBSTRATE_SPECPROFILE_ENABLED
-	SubstrateBakeSpecularProfileToReflectivity(ABSDF, ANoV);
-	SubstrateBakeSpecularProfileToReflectivity(BBSDF, BNoV);
+	ABSDF.SubstrateBakeSpecularProfileToReflectivity(ANoV);
+	BBSDF.SubstrateBakeSpecularProfileToReflectivity(BNoV);
 	BSDF_SETHASSPECPROFILE(ResultBSDF, 0);
 #endif
 
+#if SUBSTRATE_GLINTS_ENABLED 
+	ABSDF.SubstrateBakeGlintToReflectivity();
+	BBSDF.SubstrateBakeGlintToReflectivity();
+	BSDF_SETHASGLINT(ResultBSDF, 0);
+#endif // SUBSTRATE_GLINTS_ENABLED
+
 	// Some parameters will contribute to added luminance: in this case we add them together and saturate to not go out of the safe range
 	BSDF_SETEMISSIVE(ResultBSDF,		BSDF_GETEMISSIVE(ABSDF)								+ BSDF_GETEMISSIVE(ABSDF));
 	SLAB_DIFFUSEALBEDO(ResultBSDF) =	saturate(SLAB_DIFFUSEALBEDO(BBSDF)					+ SLAB_DIFFUSEALBEDO(ABSDF));
@@ -2704,14 +2712,6 @@ FSubstrateData SubstrateAddParameterBlending(FSubstrateData A, FSubstrateData B,
 	SLAB_FUZZ_ROUGHNESS(ResultBSDF) =	lerp(SLAB_FUZZ_ROUGHNESS(BBSDF),				SLAB_FUZZ_ROUGHNESS(ABSDF),					AMixFactor);
 	SLAB_HAZINESS(ResultBSDF) =			LerpHazinessParameterBlending(BBSDF,			ABSDF,										AMixFactor);
 
-#if SUBSTRATE_GLINTS_ENABLED
-	BSDF_SETHASGLINT(ResultBSDF,		BSDF_GETHASGLINT(BBSDF)							| BSDF_GETHASGLINT(ABSDF));
-	SLAB_GLINT_UV(ResultBSDF) =			lerp(SLAB_GLINT_UV(BBSDF),						SLAB_GLINT_UV(ABSDF), AMixFactor);
-	SLAB_GLINT_VALUE(ResultBSDF) =		lerp(SLAB_GLINT_VALUE(BBSDF),					SLAB_GLINT_VALUE(ABSDF), AMixFactor);
-	SLAB_GLINT_UVDDX(ResultBSDF) =		lerp(SLAB_GLINT_UVDDX(BBSDF),					SLAB_GLINT_UVDDX(ABSDF), AMixFactor);
-	SLAB_GLINT_UVDDY(ResultBSDF) =		lerp(SLAB_GLINT_UVDDY(BBSDF),					SLAB_GLINT_UVDDY(ABSDF), AMixFactor);
-#endif // SUBSTRATE_GLINTS_ENABLED
-
 	if (BSDF_GETSSSTYPE(ResultBSDF) == SSS_TYPE_DIFFUSION_PROFILE)
 	{
 		// As mentioned above, we cannot blend SSSProfiles. We do select the appropriate profile, ABSDF being selected if both have SSSProfile.
@@ -2793,17 +2793,19 @@ FSubstrateData SubstrateVerticalLayeringParameterBlending(in FSubstrateData Top,
 	BSDF_SETHASF90(ResultBSDF,			BSDF_GETHASF90(BaseBSDF)						| BSDF_GETHASF90(TopBSDF));			// F90: keep union of both even though it will be hard to get a match
 	BSDF_SETHASFUZZ(ResultBSDF,			BSDF_GETHASFUZZ(BaseBSDF)						| BSDF_GETHASFUZZ(TopBSDF));		// Fuzz: keep union of both even though it will be hard to get a match
 
-#if SUBSTRATE_COMPLEXSPECIALPATH
-	BSDF_SETHASGLINT(ResultBSDF, 0); // Disable for now
-#endif
-
 	// Bake the Specular Profiles into F0/F90 of BSDFs to blend in order to retain some of the original look when using parameter blending or fully simplified materials.
 #if SUBSTRATE_SPECPROFILE_ENABLED
-	SubstrateBakeSpecularProfileToReflectivity(TopBSDF, TopNoV);
-	SubstrateBakeSpecularProfileToReflectivity(BaseBSDF, BaseNoV);
+	TopBSDF.SubstrateBakeSpecularProfileToReflectivity(TopNoV);
+	BaseBSDF.SubstrateBakeSpecularProfileToReflectivity(BaseNoV);
 	BSDF_SETHASSPECPROFILE(ResultBSDF, 0);
 #endif
 
+#if SUBSTRATE_GLINTS_ENABLED 
+	TopBSDF.SubstrateBakeGlintToReflectivity();
+	BaseBSDF.SubstrateBakeGlintToReflectivity();
+	BSDF_SETHASGLINT(ResultBSDF, 0);
+#endif // SUBSTRATE_GLINTS_ENABLED
+
 	// Saved original Base BBSDF MFP to not be changed be changed by EnableSlabBSDFSimpleVolumetric used to compute medium transmittance and scattering.
 	// EnableSlabBSDFSimpleVolumetric will be called at then end of the SubstrateTree as needed later.
 	const float3 OriginalBaseMFP = SLAB_SSSMFP(BaseBSDF);
@@ -2848,14 +2850,6 @@ FSubstrateData SubstrateVerticalLayeringParameterBlending(in FSubstrateData Top,
 	SLAB_FUZZ_ROUGHNESS(ResultBSDF) =	lerp(SLAB_FUZZ_ROUGHNESS(TopBSDF),				SLAB_FUZZ_ROUGHNESS(BaseBSDF),TopToBottomLerp);
 	SLAB_HAZINESS(ResultBSDF) =			LerpHazinessParameterBlending(TopBSDF,			BaseBSDF,					TopToBottomLerp);
 
-#if SUBSTRATE_GLINTS_ENABLED
-	BSDF_SETHASGLINT(ResultBSDF,		BSDF_GETHASGLINT(TopBSDF)						| BSDF_GETHASGLINT(BaseBSDF));
-	SLAB_GLINT_UV(ResultBSDF) =			lerp(SLAB_GLINT_UV(TopBSDF),					SLAB_GLINT_UV(BaseBSDF), TopToBottomLerp);
-	SLAB_GLINT_VALUE(ResultBSDF) =		lerp(SLAB_GLINT_VALUE(TopBSDF),					SLAB_GLINT_VALUE(BaseBSDF), TopToBottomLerp);
-	SLAB_GLINT_UVDDX(ResultBSDF) =		lerp(SLAB_GLINT_UVDDX(TopBSDF),					SLAB_GLINT_UVDDX(BaseBSDF), TopToBottomLerp);
-	SLAB_GLINT_UVDDY(ResultBSDF) =		lerp(SLAB_GLINT_UVDDY(TopBSDF),					SLAB_GLINT_UVDDY(BaseBSDF), TopToBottomLerp);
-#endif // SUBSTRATE_GLINTS_ENABLED
-
 	// Anisotropy is disabled with vertical blending. 
 	SLAB_ANISOTROPY(ResultBSDF) =		0.f;
 
@@ -3220,6 +3214,36 @@ bool FSubstrateBSDF::HasBackScattering()
 	return SSSType != SSS_TYPE_INVALID && SSSType != SSS_TYPE_SIMPLEVOLUME;
 }
 
+void FSubstrateBSDF::SubstrateBakeSpecularProfileToReflectivity(in float NoV)
+{
+#if SUBSTRATE_SPECPROFILE_ENABLED
+	if (BSDF_GETHASSPECPROFILE(this))
+	{
+		// Evaluate the spec lut
+		const uint SpecProfileID = SLAB_SPECPROFILEID(this);
+		const float3 SpecLUTSample = EvaluateSpecularProfile(SpecProfileID, NoV, 1.0f/*NoL*/, 1.0f/*VoH*/, 1.0f/*NoH*/);
+		SLAB_F0(this) *= SpecLUTSample;
+		SLAB_F90(this) *= SpecLUTSample;
+		BSDF_SETHASSPECPROFILE(this, 0);
+	}
+#endif
+}
+
+void FSubstrateBSDF::SubstrateBakeGlintToReflectivity()
+{
+#if SUBSTRATE_COMPLEXSPECIALPATH
+	// It does not work to blend glints attributes when using parameter blending or fully simplified materials.
+	// So for such case we bake the glints as a specular lobe, reducing reflectivity according to glint amount/coverage.
+	if (BSDF_GETHASGLINT(this))
+	{
+		const float GlintCoverage = saturate(SLAB_GLINT_VALUE(this));
+		SLAB_F0(this) *= GlintCoverage;
+		SLAB_F90(this) *= GlintCoverage;
+		BSDF_SETHASGLINT(this, 0);
+	}
+#endif
+}
+
 // Return 0 if v is zero, 1 otherwise
 float IsNonZeroFast(float v)
 {
