Fixed: Issues with instance self-occlusion in HZB tests due to minor differences in precision

diff --git a/Engine/Shaders/Private/InstanceCulling/BuildInstanceDrawCommands.usf b/Engine/Shaders/Private/InstanceCulling/BuildInstanceDrawCommands.usf
--- a/Engine/Shaders/Private/InstanceCulling/BuildInstanceDrawCommands.usf
+++ b/Engine/Shaders/Private/InstanceCulling/BuildInstanceDrawCommands.usf
@@ -104,6 +104,11 @@ void WriteInstanceForCompaction(uint Offset, bool bVisible, uint InstanceId, FIn
 	CompactInstanceIdsBufferOut[Offset] = PackedId;	
 }
 
+float RoundUpF16(float DeviceZ)
+{
+	return f16tof32(f32tof16(DeviceZ) + 1);
+}
+
 bool IsInstanceVisible(FPrimitiveSceneData PrimitiveData, FInstanceSceneData InstanceData, uint InstanceId, uint ViewIdIndex, bool bAllowOcclusionCulling, FDrawCommandDesc DrawCommandDesc, inout uint CullingFlags)
 {
 	CullingFlags = INSTANCE_CULLING_FLAGS_DEFAULT;
@@ -189,6 +194,8 @@ bool IsInstanceVisible(FPrimitiveSceneData PrimitiveData, FInstanceSceneData Ins
 			if (PrevCull.bIsVisible && !PrevCull.bCrossesNearPlane)
 			{
 				FScreenRect PrevRect = GetScreenRect( NaniteView.HZBTestViewRect, PrevCull, 4 );
+				// Avoid cases where instance might self-occlude the HZB test due to minor precision differences
+				PrevRect.Depth = RoundUpF16(PrevRect.Depth);
 				Cull.bIsVisible = IsVisibleHZB( PrevRect, true );
 			}
 
