[HWRT] Mark instances in TLAS with unsafe transforms as inactive.
- temporary workaround for crash on some platforms.

diff --git a/Engine/Shaders/Private/RayTracing/RayTracingInstanceBufferUtil.usf b/Engine/Shaders/Private/RayTracing/RayTracingInstanceBufferUtil.usf
--- a/Engine/Shaders/Private/RayTracing/RayTracingInstanceBufferUtil.usf
+++ b/Engine/Shaders/Private/RayTracing/RayTracingInstanceBufferUtil.usf
@@ -95,6 +95,11 @@ bool CullInstance(FInstanceSceneData InstanceSceneData, float4x4 LocalToTranslat
 
 #endif // GPU_CULLING
 
+bool IsTransformElementUnsafe(float Element)
+{
+	return isnan(Element) || isinf(Element) || Element > 1e20f;
+}
+
 [numthreads(THREADGROUP_SIZE, 1, 1)]
 void RayTracingBuildInstanceBufferCS(uint3 DispatchThreadId : SV_DispatchThreadID)
 {
@@ -168,6 +173,23 @@ void RayTracingBuildInstanceBufferCS(uint3 DispatchThreadId : SV_DispatchThreadI
 			BlasAddress = AccelerationStructureAddresses.Load2(InputDesc.AccelerationStructureIndex * 8);
 		}
 
+		// temporary workaround for crash on some platforms
+		if(IsTransformElementUnsafe(LocalToTranslatedWorld[0][0])
+			|| IsTransformElementUnsafe(LocalToTranslatedWorld[0][1])
+			|| IsTransformElementUnsafe(LocalToTranslatedWorld[0][2])
+			|| IsTransformElementUnsafe(LocalToTranslatedWorld[0][3])
+			|| IsTransformElementUnsafe(LocalToTranslatedWorld[1][0])
+			|| IsTransformElementUnsafe(LocalToTranslatedWorld[1][1])
+			|| IsTransformElementUnsafe(LocalToTranslatedWorld[1][2])
+			|| IsTransformElementUnsafe(LocalToTranslatedWorld[1][3])
+			|| IsTransformElementUnsafe(LocalToTranslatedWorld[2][0])
+			|| IsTransformElementUnsafe(LocalToTranslatedWorld[2][1])
+			|| IsTransformElementUnsafe(LocalToTranslatedWorld[2][2])
+			|| IsTransformElementUnsafe(LocalToTranslatedWorld[2][3]))
+		{
+			BlasAddress = 0;
+		}
+
 		InstancesDescriptors[InputDesc.OutputDescriptorIndex] = BuildPlatformRayTracingInstanceDesc(
 			InputDesc.InstanceMaskAndFlags & 0xFF,
 			InputDesc.InstanceId,
