Remove ENABLE_SEPARATE_SIMPLE_LIGHT_LOOP

diff --git a/Engine/Shaders/Private/ClusteredDeferredShadingPixelShader.usf b/Engine/Shaders/Private/ClusteredDeferredShadingPixelShader.usf
--- a/Engine/Shaders/Private/ClusteredDeferredShadingPixelShader.usf
+++ b/Engine/Shaders/Private/ClusteredDeferredShadingPixelShader.usf
@@ -64,9 +64,6 @@
 // Temporary work around for shader compilation issue
 #define SUPPORT_SUBSTRATE_LIGHTING (COMPILER_DXC == 1 || COMPILER_PSSL == 1 || COMPILER_METAL == 1)
 
-// Extracts out the simple lights into their own loop that is executed first, independently of shading model.
-#define ENABLE_SEPARATE_SIMPLE_LIGHT_LOOP 0
-
 // Enable cluster debug visualization
 #define GRID_DEBUG_ENABLE 0
 
@@ -114,101 +111,6 @@ bool GetClusterDebugColor(float2 LocalPosition, float SceneDepth, FCulledLightsG
 #endif
 }
 
-/**
- * Adds local lighting of simple lights and returns the index at which non-simple light should starts
- */
-float4 GetSimpleLightLighting(
-#if SUBSTRATE_ENABLED && USE_HAIR_LIGHTING == 0
-	FSubstrateAddressing SubstrateAddressing,
-	FSubstratePixelHeader SubstratePixelHeader,
-	inout float3 OutOpaqueRoughRefractionSceneColor,
-	inout float3 OutSubSurfaceSceneColor,
-#else
-	FScreenSpaceData ScreenSpaceData,
-#endif 
-	const FCulledLightsGridHeader CulledLightGridHeader,
-	float3 TranslatedWorldPosition,
-	float3 CameraVector,
-	uint EyeIndex,
-	inout uint FirstNonSimpleLightIndex)
-{
-	float4 CompositedLighting = 0;
-	FirstNonSimpleLightIndex = 0;
-
-#if ENABLE_SEPARATE_SIMPLE_LIGHT_LOOP
-
-	// Limit max to ForwardLightStruct.NumLocalLights.
-	// This prevents GPU hangs when the PS tries to read from uninitialized NumCulledLightsGrid buffer
-	const uint NumLightsInGridCell = min(CulledLightGridHeader.NumLights, GetMaxLightsPerCell(EyeIndex));
-
-	uint GridLightListIndex = 0;
-	LOOP
-	for (; GridLightListIndex < NumLightsInGridCell; GridLightListIndex++)
-	{
-		const FLocalLightData LocalLight = GetLocalLightData(CulledLightGridHeader.DataStartIndex + GridLightListIndex, EyeIndex);
-		if (!LocalLight.Internal.bIsSimpleLight)
-		{
-			break;
-		}
-
-		if (!IsLightVisible(LocalLight, TranslatedWorldPosition))
-		{
-			continue;
-		}
-
-	#if SUBSTRATE_ENABLED && USE_HAIR_LIGHTING == 0
-		FDeferredLightData LightData = ConvertToDeferredLight(LocalLight);
-		LightData.ShadowedBits = 1;
-		LightData.bRectLight = LightData.bRectLight && USE_RECT_LIGHT;
-
-		FSubstrateShadowTermInputParameters SubstrateShadowTermInputParameters	= GetInitialisedSubstrateShadowTermInputParameters();
-		SubstrateShadowTermInputParameters.bEvaluateShadowTerm					= false;
-
-		const float3 V = -CameraVector;
-		float3 ToLight, L;
-		const float LightMask = GetLocalLightAttenuation(TranslatedWorldPosition, LightData, ToLight, L);
-
-		FSubstrateDeferredLighting SubstrateLighting = SubstrateDeferredLighting(
-			LightData,
-			V,
-			L,
-			ToLight,
-			LightMask,
-			SubstrateShadowTermInputParameters,
-			Substrate.MaterialTextureArray,
-			SubstrateAddressing,
-			SubstratePixelHeader).SceneColor;
-
-		CompositedLighting += SubstrateLighting.SceneColor;
-	#if SUBSTRATE_OPAQUE_ROUGH_REFRACTION_ENABLED
-		OutOpaqueRoughRefractionSceneColor += SubstrateLighting.OpaqueRoughRefractionSceneColor;
-		OutSubSurfaceSceneColor += SubstrateLighting.SubSurfaceSceneColor;
-	#endif
-
-	#else
-
-		FSimpleDeferredLightData SimpleLightData = (FSimpleDeferredLightData)0;
-		SimpleLightData.TranslatedWorldPosition = UnpackLightTranslatedWorldPosition(LocalLight);
-		SimpleLightData.InvRadius = UnpackLightInvRadius(LocalLight);
-		SimpleLightData.Color = UnpackLightColor(LocalLight);
-		SimpleLightData.FalloffExponent = UnpackLightFalloffExponent(LocalLight);
-		SimpleLightData.bInverseSquared = SimpleLightData.FalloffExponent == 0;
-		FSubstrateDeferredLighting SubstrateLighting = GetSimpleDynamicLighting(TranslatedWorldPosition, CameraVector, ScreenSpaceData.GBuffer.WorldNormal, ScreenSpaceData.AmbientOcclusion, ScreenSpaceData.GBuffer.DiffuseColor, ScreenSpaceData.GBuffer.SpecularColor, ScreenSpaceData.GBuffer.Roughness, SimpleLightData);
-
-		CompositedLighting += SubstrateLighting.SceneColor;
-	#if SUBSTRATE_OPAQUE_ROUGH_REFRACTION_ENABLED
-		OutOpaqueRoughRefractionSceneColor += SubstrateLighting.OpaqueRoughRefractionSceneColor;
-		OutSubSurfaceSceneColor += SubstrateLighting.SubSurfaceSceneColor;
-	#endif
-
-	#endif
-	}
-	FirstNonSimpleLightIndex = GridLightListIndex;
-#endif // ENABLE_SEPARATE_SIMPLE_LIGHT_LOOP
-
-	return CompositedLighting;
-}
-
 /**
  * Adds local lighting using the light grid, does not apply directional lights, as they are done elsewhere.
  * Does not support dynamic shadows, as these require the per-light shadow mask.
@@ -229,7 +131,6 @@ float4 GetLightGridLocalLighting(
 	float SceneDepth,
 	uint EyeIndex, 
 	float Dither, 
-	uint FirstNonSimpleLightIndex, 
 	uint SampleIndex=0, 
 	uint TotalSampleCount=0, 
 	uint2 PixelCoord=0,
@@ -254,7 +155,7 @@ float4 GetLightGridLocalLighting(
 	#endif
 	
 	LOOP
-	for (uint GridLightListIndex = FirstNonSimpleLightIndex; GridLightListIndex < NumLightsInGridCell; GridLightListIndex++)
+	for (uint GridLightListIndex = 0; GridLightListIndex < NumLightsInGridCell; GridLightListIndex++)
 	{
 		const FLocalLightData LocalLight = GetLocalLightData(InLightGridHeader.DataStartIndex + GridLightListIndex, EyeIndex);
 
@@ -393,7 +294,7 @@ float4 GetLightGridLocalLighting(
 // NOTE: this is a macro since we want to be sure the 'ScreenSpaceData.GBuffer.ShadingModelID = ShadingModelId' is compile time constant
 //       The rest sort of follows.
 // NOTE 2: The screen-space data is loaded inside the if to try to minimize register pressure, but not clear if it is respected.
-#define GET_LIGHT_GRID_LOCAL_LIGHTING_SINGLE_SM(ShadingModelId, PixelShadingModelId, /*inout float4*/ CompositedLighting, ScreenUV, LightGridHeader, Dither, FirstNonSimpleLightIndex) \
+#define GET_LIGHT_GRID_LOCAL_LIGHTING_SINGLE_SM(ShadingModelId, PixelShadingModelId, /*inout float4*/ CompositedLighting, ScreenUV, LightGridHeader, Dither) \
 {	\
 	BRANCH	\
 	if (PixelShadingModelId == ShadingModelId)	\
@@ -404,7 +305,7 @@ float4 GetLightGridLocalLighting(
 		float3 TranslatedWorldPosition = mul(float4(GetScreenPositionForProjectionType(ScreenPosition, SceneDepth), SceneDepth, 1), PrimaryView.ScreenToTranslatedWorld).xyz;	\
 		float3 CameraVector = normalize(TranslatedWorldPosition - PrimaryView.TranslatedWorldCameraOrigin);	\
 		ScreenSpaceData.GBuffer.ShadingModelID = ShadingModelId;	\
-		CompositedLighting += GetLightGridLocalLighting(ScreenSpaceData, LightGridHeader, TranslatedWorldPosition, CameraVector, ScreenUV, SceneDepth, 0, Dither, FirstNonSimpleLightIndex); \
+		CompositedLighting += GetLightGridLocalLighting(ScreenSpaceData, LightGridHeader, TranslatedWorldPosition, CameraVector, ScreenUV, SceneDepth, 0, Dither); \
 	}	\
 }
 
@@ -459,15 +360,10 @@ void ClusteredShadingPixelShader(
 		const float2 ScreenPosition = ViewportUVToScreenPos(ViewportScreenUV);
 		const float3 TranslatedWorldPosition = mul(float4(GetScreenPositionForProjectionType(ScreenPosition, SceneDepth), SceneDepth, 1), PrimaryView.ScreenToTranslatedWorld).xyz;
 		const float3 CameraVector = normalize(TranslatedWorldPosition - PrimaryView.TranslatedWorldCameraOrigin);
-		uint FirstNonSimpleLightIndex = 0;
-
-		// Simple lights
-		CompositedLighting += GetSimpleLightLighting(SubstrateAddressing, SubstratePixelHeader, OpaqueRoughRefractionSceneColor, SubSurfaceSceneColor,
-			CulledLightGridHeader, TranslatedWorldPosition, CameraVector, EyeIndex, FirstNonSimpleLightIndex);
 
 		// Regular lights
 		CompositedLighting += GetLightGridLocalLighting(SubstrateAddressing, SubstratePixelHeader, OpaqueRoughRefractionSceneColor, SubSurfaceSceneColor,
-			CulledLightGridHeader, TranslatedWorldPosition, CameraVector, ScreenUV, SceneDepth, EyeIndex, Dither, FirstNonSimpleLightIndex);
+			CulledLightGridHeader, TranslatedWorldPosition, CameraVector, ScreenUV, SceneDepth, EyeIndex, Dither);
 	}
 	
 #if !VISUALIZE_LIGHT_CULLING
@@ -517,25 +413,22 @@ void ClusteredShadingPixelShader(
 	float3 CameraVector = normalize(TranslatedWorldPosition - PrimaryView.TranslatedWorldCameraOrigin);
 	uint FirstNonSimpleLightIndex = 0;
 
-	// Simple lights
-	CompositedLighting += GetSimpleLightLighting(GetScreenSpaceData(ScreenUV), CulledLightGridHeader, TranslatedWorldPosition, CameraVector, EyeIndex, FirstNonSimpleLightIndex);
-
 	// Regular lights
 #if USE_PASS_PER_SHADING_MODEL
 
-	GET_LIGHT_GRID_LOCAL_LIGHTING_SINGLE_SM(SHADINGMODELID_DEFAULT_LIT,			PixelShadingModelID, CompositedLighting, ScreenUV, CulledLightGridHeader, Dither, FirstNonSimpleLightIndex);
-	GET_LIGHT_GRID_LOCAL_LIGHTING_SINGLE_SM(SHADINGMODELID_SUBSURFACE,			PixelShadingModelID, CompositedLighting, ScreenUV, CulledLightGridHeader, Dither, FirstNonSimpleLightIndex);
-	GET_LIGHT_GRID_LOCAL_LIGHTING_SINGLE_SM(SHADINGMODELID_PREINTEGRATED_SKIN,	PixelShadingModelID, CompositedLighting, ScreenUV, CulledLightGridHeader, Dither, FirstNonSimpleLightIndex);
-	GET_LIGHT_GRID_LOCAL_LIGHTING_SINGLE_SM(SHADINGMODELID_CLEAR_COAT,			PixelShadingModelID, CompositedLighting, ScreenUV, CulledLightGridHeader, Dither, FirstNonSimpleLightIndex);
-	GET_LIGHT_GRID_LOCAL_LIGHTING_SINGLE_SM(SHADINGMODELID_SUBSURFACE_PROFILE,	PixelShadingModelID, CompositedLighting, ScreenUV, CulledLightGridHeader, Dither, FirstNonSimpleLightIndex);
-	GET_LIGHT_GRID_LOCAL_LIGHTING_SINGLE_SM(SHADINGMODELID_TWOSIDED_FOLIAGE,	PixelShadingModelID, CompositedLighting, ScreenUV, CulledLightGridHeader, Dither, FirstNonSimpleLightIndex);
-	GET_LIGHT_GRID_LOCAL_LIGHTING_SINGLE_SM(SHADINGMODELID_HAIR,				PixelShadingModelID, CompositedLighting, ScreenUV, CulledLightGridHeader, Dither, FirstNonSimpleLightIndex);
-	GET_LIGHT_GRID_LOCAL_LIGHTING_SINGLE_SM(SHADINGMODELID_CLOTH,				PixelShadingModelID, CompositedLighting, ScreenUV, CulledLightGridHeader, Dither, FirstNonSimpleLightIndex);
-	GET_LIGHT_GRID_LOCAL_LIGHTING_SINGLE_SM(SHADINGMODELID_EYE,					PixelShadingModelID, CompositedLighting, ScreenUV, CulledLightGridHeader, Dither, FirstNonSimpleLightIndex);
-	GET_LIGHT_GRID_LOCAL_LIGHTING_SINGLE_SM(SHADINGMODELID_SINGLELAYERWATER,	PixelShadingModelID, CompositedLighting, ScreenUV, CulledLightGridHeader, Dither, FirstNonSimpleLightIndex);
+	GET_LIGHT_GRID_LOCAL_LIGHTING_SINGLE_SM(SHADINGMODELID_DEFAULT_LIT,			PixelShadingModelID, CompositedLighting, ScreenUV, CulledLightGridHeader, Dither);
+	GET_LIGHT_GRID_LOCAL_LIGHTING_SINGLE_SM(SHADINGMODELID_SUBSURFACE,			PixelShadingModelID, CompositedLighting, ScreenUV, CulledLightGridHeader, Dither);
+	GET_LIGHT_GRID_LOCAL_LIGHTING_SINGLE_SM(SHADINGMODELID_PREINTEGRATED_SKIN,	PixelShadingModelID, CompositedLighting, ScreenUV, CulledLightGridHeader, Dither);
+	GET_LIGHT_GRID_LOCAL_LIGHTING_SINGLE_SM(SHADINGMODELID_CLEAR_COAT,			PixelShadingModelID, CompositedLighting, ScreenUV, CulledLightGridHeader, Dither);
+	GET_LIGHT_GRID_LOCAL_LIGHTING_SINGLE_SM(SHADINGMODELID_SUBSURFACE_PROFILE,	PixelShadingModelID, CompositedLighting, ScreenUV, CulledLightGridHeader, Dither);
+	GET_LIGHT_GRID_LOCAL_LIGHTING_SINGLE_SM(SHADINGMODELID_TWOSIDED_FOLIAGE,	PixelShadingModelID, CompositedLighting, ScreenUV, CulledLightGridHeader, Dither);
+	GET_LIGHT_GRID_LOCAL_LIGHTING_SINGLE_SM(SHADINGMODELID_HAIR,				PixelShadingModelID, CompositedLighting, ScreenUV, CulledLightGridHeader, Dither);
+	GET_LIGHT_GRID_LOCAL_LIGHTING_SINGLE_SM(SHADINGMODELID_CLOTH,				PixelShadingModelID, CompositedLighting, ScreenUV, CulledLightGridHeader, Dither);
+	GET_LIGHT_GRID_LOCAL_LIGHTING_SINGLE_SM(SHADINGMODELID_EYE,					PixelShadingModelID, CompositedLighting, ScreenUV, CulledLightGridHeader, Dither);
+	GET_LIGHT_GRID_LOCAL_LIGHTING_SINGLE_SM(SHADINGMODELID_SINGLELAYERWATER,	PixelShadingModelID, CompositedLighting, ScreenUV, CulledLightGridHeader, Dither);
 	// SHADINGMODELID_THIN_TRANSLUCENT - skipping because it can not be opaque
 #else // !USE_PASS_PER_SHADING_MODEL
-	CompositedLighting += GetLightGridLocalLighting(GetScreenSpaceData(ScreenUV), CulledLightGridHeader, TranslatedWorldPosition, CameraVector, ScreenUV, SceneDepth, 0, Dither, FirstNonSimpleLightIndex);
+	CompositedLighting += GetLightGridLocalLighting(GetScreenSpaceData(ScreenUV), CulledLightGridHeader, TranslatedWorldPosition, CameraVector, ScreenUV, SceneDepth, 0, Dither);
 #endif // USE_PASS_PER_SHADING_MODEL
 	
 #if !VISUALIZE_LIGHT_CULLING
@@ -591,8 +484,7 @@ void ClusteredShadingPixelShader(
 	const float3 CameraVector = normalize(TranslatedWorldPosition - PrimaryView.TranslatedWorldCameraOrigin);
 
 	float4 CompositedLighting = 0;
-	uint FirstNonSimpleLightIndex = 0;
-	CompositedLighting += GetLightGridLocalLighting(ScreenSpaceData, CulledLightGridHeader, TranslatedWorldPosition, CameraVector, ScreenUV, SceneDepth, 0, Dither, FirstNonSimpleLightIndex, SampleIndex, TotalSampleCount, PixelCoord, Sample.LightChannelMask);
+	CompositedLighting += GetLightGridLocalLighting(ScreenSpaceData, CulledLightGridHeader, TranslatedWorldPosition, CameraVector, ScreenUV, SceneDepth, 0, Dither, SampleIndex, TotalSampleCount, PixelCoord, Sample.LightChannelMask);
 #if !VISUALIZE_LIGHT_CULLING
 	CompositedLighting *= View.PreExposure;
 #endif
diff --git a/Engine/Shaders/Private/LightData.ush b/Engine/Shaders/Private/LightData.ush
--- a/Engine/Shaders/Private/LightData.ush
+++ b/Engine/Shaders/Private/LightData.ush
@@ -83,8 +83,6 @@ struct FForwardLightData
 	// Derived terms after unpacking
 	/** Flag is true if clustered deferred is supported for this light. They are always first / together.*/
 	bool bClusteredDeferredSupported;
-	/** Flag is true if it is a simple light. They are always first / together.*/
-	bool bIsSimpleLight;
 	/** Virtual shadow map ID or INDEX_NONE if not present.*/
 	int  VirtualShadowMapId;
 	/** Previous frame LocalLightIndex or INDEX_NONE if not present.*/
diff --git a/Engine/Shaders/Private/LightGridCommon.ush b/Engine/Shaders/Private/LightGridCommon.ush
--- a/Engine/Shaders/Private/LightGridCommon.ush
+++ b/Engine/Shaders/Private/LightGridCommon.ush
@@ -277,13 +277,11 @@ FDirectionalLightData GetDirectionalLightData(uint EyeIndex)
 FForwardLightData GetForwardLightData_Internal(
 	uint LightIndex,
 	uint ClusteredDeferredSupportedEndIndex,
-	uint SimpleLightsEndIndex,
 	float4 InData0, float4 InData1, float4 InData2, float4 InData3, float4 InData4, float4 InData5)
 {
 	FForwardLightData Out = (FForwardLightData)0;
 
 	Out.bClusteredDeferredSupported					= LightIndex < ClusteredDeferredSupportedEndIndex;
-	Out.bIsSimpleLight								= LightIndex < SimpleLightsEndIndex;
 	Out.LightPositionAndInvRadius					= InData0;
 	Out.LightColorAndIdAndFalloffExponent			= InData1;
 	Out.LightDirectionAndSceneInfoExtraDataPacked	= InData2;
@@ -305,7 +303,6 @@ FForwardLightData GetForwardLightDataNonStereo(uint LightIndex)
 	return GetForwardLightData_Internal(
 		LightIndex,
 		ForwardLightStruct.ClusteredDeferredSupportedEndIndex,
-		ForwardLightStruct.SimpleLightsEndIndex,
 		ForwardLightStruct.ForwardLightBuffer[LightBaseIndex + 0],
 		ForwardLightStruct.ForwardLightBuffer[LightBaseIndex + 1],
 		ForwardLightStruct.ForwardLightBuffer[LightBaseIndex + 2],
@@ -336,7 +333,6 @@ FLocalLightData GetLocalLightData(uint GridIndex, uint EyeIndex)
 		Out.Internal = GetForwardLightData_Internal(
 			LocalLightIndex, 
 			ForwardLightStruct.ClusteredDeferredSupportedEndIndex,
-			ForwardLightStruct.SimpleLightsEndIndex,
 			ForwardLightStruct.ForwardLightBuffer[LocalLightBaseIndex + 0],
 			ForwardLightStruct.ForwardLightBuffer[LocalLightBaseIndex + 1],
 			ForwardLightStruct.ForwardLightBuffer[LocalLightBaseIndex + 2],
@@ -353,7 +349,6 @@ FLocalLightData GetLocalLightData(uint GridIndex, uint EyeIndex)
 		Out.Internal = GetForwardLightData_Internal(
 			LocalLightIndex,
 			ForwardLightStructISR.ClusteredDeferredSupportedEndIndex,
-			ForwardLightStructISR.SimpleLightsEndIndex,
 			ForwardLightStructISR.ForwardLightBuffer[LocalLightBaseIndex + 0],
 			ForwardLightStructISR.ForwardLightBuffer[LocalLightBaseIndex + 1],
 			ForwardLightStructISR.ForwardLightBuffer[LocalLightBaseIndex + 2],
diff --git a/Engine/Source/Runtime/Renderer/Private/LightGridInjection.cpp b/Engine/Source/Runtime/Renderer/Private/LightGridInjection.cpp
--- a/Engine/Source/Runtime/Renderer/Private/LightGridInjection.cpp
+++ b/Engine/Source/Runtime/Renderer/Private/LightGridInjection.cpp
@@ -1143,7 +1143,6 @@ FComputeLightGridOutput FSceneRenderer::ComputeLightGrid(FRDGBuilder& GraphBuild
 		ForwardLightUniformParameters->CulledGridSize = FIntVector(LightGridSizeXY.X, LightGridSizeXY.Y, GLightGridSizeZ);
 		ForwardLightUniformParameters->MaxCulledLightsPerCell = GLightLinkedListCulling ? NumVisibleLocalLights : GMaxCulledLightsPerCell;
 		ForwardLightUniformParameters->LightGridPixelSizeShift = FMath::FloorLog2(GLightGridPixelSize);
-		ForwardLightUniformParameters->SimpleLightsEndIndex = SimpleLightsEnd;
 		ForwardLightUniformParameters->ClusteredDeferredSupportedEndIndex = ClusteredSupportedEnd;
 		ForwardLightUniformParameters->DirectionalMegaLightsSupportedStartIndex = DirectionalMegaLightsSupportedStart;
 		ForwardLightUniformParameters->DirectLightingShowFlag = ViewFamily.EngineShowFlags.DirectLighting ? 1 : 0;
diff --git a/Engine/Source/Runtime/Renderer/Private/SceneRendering.h b/Engine/Source/Runtime/Renderer/Private/SceneRendering.h
--- a/Engine/Source/Runtime/Renderer/Private/SceneRendering.h
+++ b/Engine/Source/Runtime/Renderer/Private/SceneRendering.h
@@ -683,7 +683,6 @@ BEGIN_GLOBAL_SHADER_PARAMETER_STRUCT_WITH_CONSTRUCTOR(FForwardLightUniformParame
 	SHADER_PARAMETER(float, DirectionalLightDepthBias)
 	SHADER_PARAMETER(uint32, DirectionalLightUseStaticShadowing)
 	SHADER_PARAMETER(uint32, DirectionalLightHandledByMegaLights)
-	SHADER_PARAMETER(uint32, SimpleLightsEndIndex)
 	SHADER_PARAMETER(uint32, ClusteredDeferredSupportedEndIndex)
 	SHADER_PARAMETER(uint32, DirectionalMegaLightsSupportedStartIndex)
 	SHADER_PARAMETER(FVector4f, DirectionalLightStaticShadowBufferSize)
