Fix shader compilation failure of mobile deferred directional light with mobile multi view

diff --git a/Engine/Shaders/Private/MobileDeferredShading.usf b/Engine/Shaders/Private/MobileDeferredShading.usf
--- a/Engine/Shaders/Private/MobileDeferredShading.usf
+++ b/Engine/Shaders/Private/MobileDeferredShading.usf
@@ -137,6 +137,11 @@ void ReflectionEnvironmentSkyLighting(
   
 void MobileDirectionalLightPS(
 	noperspective float4 UVAndScreenPos : TEXCOORD0, 
+#if INSTANCED_STEREO && MOBILE_MULTI_VIEW
+	in nointerpolation uint EyeIndex : VIEW_ID,
+#elif MOBILE_MULTI_VIEW
+	in nointerpolation uint ViewId : SV_ViewID,
+#endif
 	float4 SvPosition : SV_POSITION, 
 #if USE_GLES_FBF_DEFERRED
 	out HALF4_TYPE OutProxyAdditive : SV_Target0,
@@ -148,7 +153,13 @@ void MobileDirectionalLightPS(
 #endif
 )
 {
+#if INSTANCED_STEREO && MOBILE_MULTI_VIEW
+	ResolvedView = ResolveView(EyeIndex);
+#elif MOBILE_MULTI_VIEW
+	ResolvedView = ResolveView(ViewId);
+#else
 	ResolvedView = ResolveView();
+#endif
 
 	FGBufferData GBuffer = MobileFetchAndDecodeGBuffer(UVAndScreenPos.xy, UVAndScreenPos.zw);
 	float2 ScreenPos = UVAndScreenPos.zw;
