Fix legact ClearCoat shading model bottom layer, having incorrect specular response when clear coat value is set to 0 (but using ClearCoat shading model).

The alpha2 value was incorrect, as the area light energy term was computed after the GGX term. The area light energy term computation modifies the alpha2 value, and is the value that should be used by the GGX term (as done in the default lit case).

#tests EngineTests with tests covering clear coat materials (both -sm5 -sm6)

diff --git a/Engine/Shaders/Private/ShadingModels.ush b/Engine/Shaders/Private/ShadingModels.ush
--- a/Engine/Shaders/Private/ShadingModels.ush
+++ b/Engine/Shaders/Private/ShadingModels.ush
@@ -588,6 +588,26 @@ FDirectLighting ClearCoatBxDF( FGBufferData GBuffer, half3 N, half3 V, half3 L,
 		half D2 = 0;
 		half Vis2 = 0;
 
+		// The energy factor for correcting area light shape needs to be computed "before" the actual D_GGX term, 
+		// as calling EnergyNormalization() modify the "a2" value.
+		float Energy = 0;
+		BRANCH
+		if (bHasAnisotropy)
+		{
+			FAreaLight Punctual = AreaLight;
+			Punctual.SphereSinAlpha = 0;
+			Punctual.SphereSinAlphaSoft = 0;
+			Punctual.LineCosSubtended = 1;
+			Punctual.Rect = (FRect)0;
+			Punctual.IsRectAndDiffuseMicroReflWeight = 0;
+			
+			Energy = EnergyNormalization(a2, Context.VoH, Punctual);
+		}
+		else
+		{
+			Energy = EnergyNormalization(a2, Context.VoH, AreaLight);
+		}
+
 		BRANCH
 		if (bHasAnisotropy)
 		{
@@ -612,24 +632,6 @@ FDirectLighting ClearCoatBxDF( FGBufferData GBuffer, half3 N, half3 V, half3 L,
 		half3 F_Bot = F_Schlick(GBuffer.SpecularColor, BottomContext.VoH);
 		half3 F_DefaultLit = F_Schlick(GBuffer.SpecularColor, Context.VoH);
 		
-		float Energy = 0;
-
-		BRANCH
-		if (bHasAnisotropy)
-		{
-			FAreaLight Punctual = AreaLight;
-			Punctual.SphereSinAlpha = 0;
-			Punctual.SphereSinAlphaSoft = 0;
-			Punctual.LineCosSubtended = 1;
-			Punctual.Rect = (FRect)0;
-			Punctual.IsRectAndDiffuseMicroReflWeight = 0;
-			
-			Energy = EnergyNormalization(a2, Context.VoH, Punctual);
-		}
-		else
-		{
-			Energy = EnergyNormalization(a2, Context.VoH, AreaLight);
-		}
 
 		// Note: reusing D and V from refracted context to save computation when ClearCoat < 1
 		float3 CommonSpecular = (Energy * Falloff * NoL * D2 * Vis2) * AreaLight.FalloffColor;
