Substrate - fix for data to unconditionaly be written from Nanite on the first N GBuffer  RTs.

#rnx
#tests Tested pc dx12 editor

diff --git a/Engine/Shaders/Private/ComputeShaderOutputCommon.ush b/Engine/Shaders/Private/ComputeShaderOutputCommon.ush
--- a/Engine/Shaders/Private/ComputeShaderOutputCommon.ush
+++ b/Engine/Shaders/Private/ComputeShaderOutputCommon.ush
@@ -90,30 +90,16 @@ FPixelShaderOut ShadePixel(const float2 SVPositionXY, uint QuadIndex, uint QuadP
 	return PixelShaderOut;
 }
 
-#define CONDITIONAL_EXPORT 1
-
 #if SUBSTRATE_OPAQUE_DEFERRED
 
-void ConditionalExport(bool Forced, uint2 PixelPos, FPixelShaderOut ShadedPixel, uint Index)
+void SubstrateExport(uint2 PixelPos, FPixelShaderOut ShadedPixel, uint Index)
 {
-#if CONDITIONAL_EXPORT
-	BRANCH
-	if (Forced || ShadedPixel.SubstrateOutput[Index] != 0)
-#endif
-	{
-		ComputeShadingOutputs.OutTargets[uint3(PixelPos, Index)] = ShadedPixel.SubstrateOutput[Index];
-	}
+	ComputeShadingOutputs.OutTargets[uint3(PixelPos, Index)] = ShadedPixel.SubstrateOutput[Index];
 }
 
-void ConditionalExport(bool Forced, uint2 PixelPos, SUBSTRATE_TOP_LAYER_TYPE TopLayerData)
+void SubstrateExport(uint2 PixelPos, SUBSTRATE_TOP_LAYER_TYPE TopLayerData)
 {
-#if CONDITIONAL_EXPORT
-	BRANCH
-	if (Forced || any(TopLayerData != 0))
-#endif
-	{
-		ComputeShadingOutputs.OutTopLayerTarget[PixelPos] = TopLayerData;
-	}
+	ComputeShadingOutputs.OutTopLayerTarget[PixelPos] = TopLayerData;
 }
 
 #endif
@@ -153,14 +139,14 @@ void ExportPixel(const uint2 PixelPos, FPixelShaderOut ShadedPixel)
 	// After, Substrate top layer data appended. Remaining Substrate outputs are in the 2d array UAV
 
 	// Export Substrate data
-	// * Layer0 is always exported as it contains the header data, used for classification and pixel prunning
+	// * We unconditionally write all required MRTs. Because, for instance, simple/fast material encoding can have F0=Diffuse=Black making the second uint be 0 when we DO want to write that data. Same for other encoding types.
 	// * SubstrateTopLayerData is also always exported for the same reason
 	UNROLL 
 	for (uint LayerIt=0; LayerIt<SUBSTRATE_BASE_PASS_MRT_OUTPUT_COUNT; ++LayerIt)
 	{
-		ConditionalExport(LayerIt == 0 /*bForceOutput*/,  PixelPos, ShadedPixel, LayerIt);
+		SubstrateExport(PixelPos, ShadedPixel, LayerIt);
 	}
-	ConditionalExport(true /*bForceOutput*/, PixelPos, ShadedPixel.SubstrateTopLayerData);
+	SubstrateExport(PixelPos, ShadedPixel.SubstrateTopLayerData);
 
 #else // SUBSTRATE_OPAQUE_DEFERRED
 
