Real time sky capture - added optimised convolution for mobile path.

#tests PC dx12 public
#rnx

diff --git a/Engine/Shaders/Private/ReflectionEnvironmentShaders.usf b/Engine/Shaders/Private/ReflectionEnvironmentShaders.usf
--- a/Engine/Shaders/Private/ReflectionEnvironmentShaders.usf
+++ b/Engine/Shaders/Private/ReflectionEnvironmentShaders.usf
@@ -539,8 +539,28 @@ void FilterPS(
 	uint CubeSize = 1u << ( NumMips - 1 );
 	const float SolidAngleTexel = 4.0f * PI / float(6.0f * CubeSize * CubeSize) * 2.0f;
 
-	//const uint NumSamples = 1024;
+#if 0
+	// Reference
+	const uint NumSamples = 1024;
+#elif SHADING_PATH_MOBILE
+	// Mobile path - a bit lower quality
+	uint NumSamples = 0;
+	if (Roughness <= 0.1)
+	{
+		NumSamples = 16;
+	}
+	else if (Roughness <= 0.4)
+	{
+		NumSamples = 32;
+	}
+	else
+	{
+		NumSamples = 64;
+	}
+#else
+	// General real time
 	const uint NumSamples = Roughness < 0.1 ? 32 : 64;
+#endif
 	
 	float4 FilteredColor = 0;
 	BRANCH
