Fixed bad translucent lighting world position due to missing ResolvedView data.

#tests PC dx12 editor
#rnx

diff --git a/Engine/Shaders/Private/TranslucentLightInjectionShaders.usf b/Engine/Shaders/Private/TranslucentLightInjectionShaders.usf
--- a/Engine/Shaders/Private/TranslucentLightInjectionShaders.usf
+++ b/Engine/Shaders/Private/TranslucentLightInjectionShaders.usf
@@ -207,6 +207,8 @@ void InjectMainPS(
 	out float4 OutColor1 : SV_Target1
 	)
 {
+	ResolvedView = ResolveView();
+
 	OutColor0 = 0;
 	OutColor1 = 0;
 
