Fix for cloud being visible when SkyAtmosphere or Fog are alpha holdout.

#tests PC dx12 editor
#rnx

diff --git a/Engine/Shaders/Private/VolumetricCloud.usf b/Engine/Shaders/Private/VolumetricCloud.usf
--- a/Engine/Shaders/Private/VolumetricCloud.usf
+++ b/Engine/Shaders/Private/VolumetricCloud.usf
@@ -1492,7 +1492,20 @@ void MainCommon(in FMaterialPixelParameters MaterialParameters, in float4 SvPosi
 	float3 FogSampleWorldPositionRelativeToCameraCm = tAP * Raydir;
 	float4 ClipPos = mul(float4(RayTranslatedWorldOriginKm * KILOMETER_TO_CENTIMETER + FogSampleWorldPositionRelativeToCameraCm, 1.0f), PrimaryView.TranslatedWorldToClip);
 
-	if (RenderVolumetricCloudParameters.EnableAerialPerspectiveSampling && tAPWeightsSum > 0.0f)
+#if SUPPORT_PRIMITIVE_ALPHA_HOLDOUT
+	const bool bApplyFogOnCloud			= tAPWeightsSum;			// Use MeanTransmittance instead of tAPWeightsSum as it is more representative of what can happen since it has not been transformed. (but can result in more fog being applied, this makes alpha holdout correct)
+#if CLOUD_MIN_AND_MAX_DEPTH
+	const bool bApplyFogOnClosestCloud	= ClosestMeanTransmittance < 1.0f;	// Idem
+#endif
+#else // SUPPORT_PRIMITIVE_ALPHA_HOLDOUT
+	// We skip more pixels when using tAPWeightsSum so let;s keep that when alpha output is not enabled for the sake of performance
+	const bool bApplyFogOnCloud			= tAPWeightsSum > 0.0f;
+#if CLOUD_MIN_AND_MAX_DEPTH
+	const bool bApplyFogOnClosestCloud	= ClostestTAPWeightsSum > 0.0f;
+#endif
+#endif // SUPPORT_PRIMITIVE_ALPHA_HOLDOUT
+
+	if (RenderVolumetricCloudParameters.EnableAerialPerspectiveSampling && bApplyFogOnCloud)
 	{
 		float4 AerialPerspective;
 
@@ -1526,7 +1539,7 @@ void MainCommon(in FMaterialPixelParameters MaterialParameters, in float4 SvPosi
 			const float AerialPespectiveCoverage = saturate(1.0 - AerialPespectiveTransmittance);
 
 			// Fade out holdout coverage as a function of the aerial perspective above, if the sky atmosphere is not hold out
-			OutHoldOutCoverage   *= saturate(1.0 - AerialPespectiveCoverage);
+			OutHoldOutCoverage   *= AerialPespectiveTransmittance;
 
 			if(IsSkyAtmosphereHoldout(View.EnvironmentComponentsFlags))
 			{
@@ -1540,7 +1553,7 @@ void MainCommon(in FMaterialPixelParameters MaterialParameters, in float4 SvPosi
 	//
 	// Apply height fog if needed
 	//
-	if (tAPWeightsSum > 0.0f)
+	if (bApplyFogOnCloud)
 	{
 		float4 HeightFogInscatteringAndTransmittance = float4(0, 0, 0, 1);
 
@@ -1601,7 +1614,7 @@ void MainCommon(in FMaterialPixelParameters MaterialParameters, in float4 SvPosi
 		float3 FogSampleWorldPositionRelativeToCameraCm = ClosestTAP * Raydir;
 		float4 ClipPos = mul(float4(RayTranslatedWorldOriginKm * KILOMETER_TO_CENTIMETER + FogSampleWorldPositionRelativeToCameraCm, 1.0f), PrimaryView.TranslatedWorldToClip);
 
-		if (RenderVolumetricCloudParameters.EnableAerialPerspectiveSampling && ClostestTAPWeightsSum > 0.0f)
+		if (RenderVolumetricCloudParameters.EnableAerialPerspectiveSampling && bApplyFogOnClosestCloud)
 		{
 			float4 AerialPerspective;
 
@@ -1633,7 +1646,7 @@ void MainCommon(in FMaterialPixelParameters MaterialParameters, in float4 SvPosi
 		//
 		// Apply height fog if needed
 		//
-		if (ClostestTAPWeightsSum > 0.0f)
+		if (bApplyFogOnClosestCloud)
 		{
 			float4 HeightFogInscatteringAndTransmittance = float4(0, 0, 0, 1);
 
