Fix Lumen screen probe gather + SSR reflection with Lumen.

diff --git a/Engine/Shaders/Private/DiffuseIndirectComposite.usf b/Engine/Shaders/Private/DiffuseIndirectComposite.usf
--- a/Engine/Shaders/Private/DiffuseIndirectComposite.usf
+++ b/Engine/Shaders/Private/DiffuseIndirectComposite.usf
@@ -194,7 +194,7 @@ float3 MixSpecularAndRoughReflections(float InRoughness, bool bHasBackfaceDiffus
 	// (1 - FadeAlpha) applied during Lumen rough specular integration to allow direct lighting specular which shouldn't be lerped out
 	float3 Lighting = RoughReflections;
 	#else
-	float3 Lighting = RoughReflections * (1 - FadeAlpha);
+	float3 Lighting = bLumenReflectionInputIsSSR ? RoughReflections : RoughReflections * (1 - FadeAlpha);
 	#endif
 
 	if (bLumenReflectionInputIsSSR > 0)
@@ -266,7 +266,7 @@ FSubstrateDeferredLighting SubstrateIndirectLighting(
 	#if DIM_APPLY_DIFFUSE_INDIRECT == DIM_APPLY_DIFFUSE_INDIRECT_SCREEN_PROBE_GATHER
 		const float3 InDiffuseLighting			= DiffuseIndirect_Lumen_0.Load(FetchCoord).rgb;
 		const float3 InRoughSpecularLighting	= DiffuseIndirect_Lumen_2.Load(FetchCoord).rgb;
-		const float4 InSpecularLighting			= DiffuseIndirect_Lumen_3.Load(FetchCoord).rgba;
+		const float4 InSpecularLighting			= bLumenReflectionInputIsSSR ? InputSpecularLighting : DiffuseIndirect_Lumen_3.Load(FetchCoord).rgba;
 	#else
 		const float3 InDiffuseLighting			= InputDiffuseLighting;
 		const float3 InRoughSpecularLighting	= InputSpecularLighting.xyz;
