Mobile: Use same BRDF for shading as on SM5

diff --git a/Engine/Shaders/Private/DeferredShadingCommon.ush b/Engine/Shaders/Private/DeferredShadingCommon.ush
--- a/Engine/Shaders/Private/DeferredShadingCommon.ush
+++ b/Engine/Shaders/Private/DeferredShadingCommon.ush
@@ -476,10 +476,6 @@ float4 ComputeIndirectLightingSampleE(uint2 TracingPixelCoord, uint TracingRayIn
 		Hammersley16(TracingRayIndex, TracingRayCount, Seed1));
 }
 
-#ifndef MOBILE_HIGH_QUALITY_BRDF
-#define MOBILE_HIGH_QUALITY_BRDF 0
-#endif
-
 #if SHADING_PATH_MOBILE
 
 #ifndef ENABLE_SHADINGMODEL_SUPPORT_MOBILE_DEFERRED
@@ -1223,17 +1219,10 @@ FScreenSpaceData GetScreenSpaceData(float2 UV, bool bGetNormalizedNormal = true)
 // [ Jimenez et al. 2016, "Practical Realtime Strategies for Accurate Indirect Occlusion" ]
 half3 AOMultiBounce(half3 BaseColor, half AO)
 {
-	if (SHADING_PATH_MOBILE && !MOBILE_HIGH_QUALITY_BRDF)
-	{
-		return AO;
-	}
-	else
-	{
-		half3 a = 2.0404 * BaseColor - 0.3324;
-		half3 b = -4.7951 * BaseColor + 0.6417;
-		half3 c = 2.7552 * BaseColor + 0.6903;
-		return max(AO, ((AO * a + b) * AO + c) * AO);
-	}
+	half3 a = 2.0404 * BaseColor - 0.3324;
+	half3 b = -4.7951 * BaseColor + 0.6417;
+	half3 c = 2.7552 * BaseColor + 0.6903;
+	return max(AO, ((AO * a + b) * AO + c) * AO);
 }
 
 #define LIGHTING_CHANNELS_TEXTURE_DISTANCE_FIELD_REPRESENTATION_BIT 3
\ No newline at end of file
diff --git a/Engine/Shaders/Private/MobileBasePassPixelShader.usf b/Engine/Shaders/Private/MobileBasePassPixelShader.usf
--- a/Engine/Shaders/Private/MobileBasePassPixelShader.usf
+++ b/Engine/Shaders/Private/MobileBasePassPixelShader.usf
@@ -606,10 +606,7 @@ void Main(
 	#endif
 
 	const FShadingOcclusion ShadingOcclusion = ApplyBentNormal(MaterialParameters.CameraVector, InputBentNormal, GetWorldBentNormalZero(MaterialParameters), GBuffer.Roughness, GBuffer.GBufferAO);
-
-#if MOBILE_HIGH_QUALITY_BRDF
 	GBuffer.GBufferAO = AOMultiBounce(Luminance(GBuffer.SpecularColor), ShadingOcclusion.SpecOcclusion).g;
-#endif
 
 #endif // !SUBSTRATE_ENABLED || SUBSTRATE_INLINE_SINGLELAYERWATER
 
diff --git a/Engine/Shaders/Private/MobileLightingCommon.ush b/Engine/Shaders/Private/MobileLightingCommon.ush
--- a/Engine/Shaders/Private/MobileLightingCommon.ush
+++ b/Engine/Shaders/Private/MobileLightingCommon.ush
@@ -484,14 +484,11 @@ void AccumulateReflection(
 
 	half3 TopLayerR = ReflectionVector;
 	half SpecularOcclusion = GBuffer.GBufferAO;
-	if (MOBILE_HIGH_QUALITY_BRDF)
-	{
-		// Point lobe in off-specular peak direction
-		ReflectionVector = GetOffSpecularPeakReflectionDir(N, ReflectionVector, GBuffer.Roughness);
-
-		half RoughnessSq = GBuffer.Roughness * GBuffer.Roughness;
-		SpecularOcclusion = GetSpecularOcclusion(NoV, RoughnessSq, SpecularOcclusion);
-	}
+	
+	// Point lobe in off-specular peak direction
+	ReflectionVector = GetOffSpecularPeakReflectionDir(N, ReflectionVector, GBuffer.Roughness);
+	half RoughnessSq = GBuffer.Roughness * GBuffer.Roughness;
+	SpecularOcclusion = GetSpecularOcclusion(NoV, RoughnessSq, SpecularOcclusion);
 
 	// IBL reflection (primary)
 	half3 SpecularIBL = GetImageBasedReflectionLighting_Mobile(ReflectionVector
diff --git a/Engine/Shaders/Private/ShadingModels.ush b/Engine/Shaders/Private/ShadingModels.ush
--- a/Engine/Shaders/Private/ShadingModels.ush
+++ b/Engine/Shaders/Private/ShadingModels.ush
@@ -40,9 +40,9 @@
 #define MATERIAL_SHADINGMODEL_SINGLELAYERWATER 0
 #endif
 
-//It's always 0 in mobile deferred lighting shaders
+// Always enabled in mobile deferred lighting shaders
 #ifndef MOBILE_USE_PREINTEGRATED_GF
-#define  MOBILE_USE_PREINTEGRATED_GF 0
+#define MOBILE_USE_PREINTEGRATED_GF 1
 #endif
 
 //It's always 1 in mobile deferred lighting shaders
@@ -79,20 +79,6 @@ half3 GetEnvBRDF(half3 SpecularColor, half Roughness, half NoV)
 	return EnvBRDFApprox(SpecularColor, Roughness, NoV);
 #endif
 }
-
-half3 MobileSpecularGGXInner(half D, half3 SpecularColor, half Roughness, half NoV, half NoL, half VoH, bool bHighQualityBRDF)
-{
-	half Vis = (Roughness * 0.25 + 0.25);
-	half3 F = GetEnvBRDF(SpecularColor, Roughness, NoV);
-
-	if (bHighQualityBRDF)
-	{
-		Vis = saturate(Vis_SmithJointApprox(Roughness * Roughness * Roughness * Roughness, NoV, NoL));
-		F = F_Schlick(SpecularColor, VoH);
-	}
-
-	return (D * Vis) * F;
-}
 #endif // SHADING_PATH_MOBILE
 
 #if 0
@@ -229,17 +215,12 @@ float3 SpecularGGX( float Roughness, float3 SpecularColor, BxDFContext Context,
 	float a2 = Pow4( Roughness );
 	float Energy = EnergyNormalization( a2, Context.VoH, AreaLight );
 	
-#if SHADING_PATH_MOBILE
-	half D = D_GGX_Mobile(Roughness, Context.NoH) * Energy;
-	return MobileSpecularGGXInner(D, SpecularColor, Roughness, Context.NoV, NoL, Context.VoH, MOBILE_HIGH_QUALITY_BRDF);
-#else
 	// Generalized microfacet specular
 	float D = D_GGX( a2, Context.NoH ) * Energy;
 	float Vis = Vis_SmithJointApprox( a2, Context.NoV, NoL );
 	float3 F = F_Schlick( SpecularColor, Context.VoH );
 
 	return (D * Vis) * F;
-#endif
 }
 
 half3 DualSpecularGGX(half AverageRoughness, half Lobe0Roughness, half Lobe1Roughness, half LobeMix, half3 SpecularColor, BxDFContext Context, half NoL, FAreaLight AreaLight)
@@ -251,18 +232,12 @@ half3 DualSpecularGGX(half AverageRoughness, half Lobe0Roughness, half Lobe1Roug
 	float Lobe0Energy = EnergyNormalization(Lobe0Alpha2, Context.VoH, AreaLight);
 	float Lobe1Energy = EnergyNormalization(Lobe1Alpha2, Context.VoH, AreaLight);
 
-#if SHADING_PATH_MOBILE
-	half D = lerp(D_GGX_Mobile(Lobe0Roughness, Context.NoH) * Lobe0Energy, D_GGX_Mobile(Lobe1Roughness, Context.NoH) * Lobe1Energy, LobeMix);
-
-	return MobileSpecularGGXInner(D, SpecularColor, AverageRoughness, Context.NoV, NoL, Context.VoH, MOBILE_HIGH_QUALITY_BRDF);
-#else
 	// Generalized microfacet specular
 	float D = lerp(D_GGX(Lobe0Alpha2, Context.NoH) * Lobe0Energy, D_GGX(Lobe1Alpha2, Context.NoH) * Lobe1Energy, LobeMix);
 	float Vis = Vis_SmithJointApprox(AverageAlpha2, Context.NoV, NoL); // Average visibility well approximates using two separate ones (one per lobe).
 	float3 F = F_Schlick(SpecularColor, Context.VoH);
 
 	return (D * Vis) * F;
-#endif
 }
 
 FDirectLighting DefaultLitBxDF( FGBufferData GBuffer, half3 N, half3 V, half3 L, float Falloff, half NoL, FAreaLight AreaLight, FShadowTerms Shadow )
@@ -478,12 +453,7 @@ FDirectLighting ClearCoatBxDF( FGBufferData GBuffer, half3 N, half3 V, half3 L,
 		float a2 = Pow2(Alpha);
 		float ClearCoatEnergy = EnergyNormalization(a2, Context.VoH, AreaLight);
 		half Vis = Vis_SmithJointApprox(a2, Context.NoV, NoL);
-#if SHADING_PATH_MOBILE
-		half D = D_GGX_Mobile(ClearCoatRoughness, Context.NoH) * ClearCoatEnergy;
-		Vis = saturate(Vis);
-#else
 		float D = D_GGX(a2, Context.NoH) * ClearCoatEnergy;
-#endif
 		half Fr1 = (D * Vis) * F;
 		Lighting.Specular = ClearCoat * AreaLight.FalloffColor * (Falloff * NoL * Fr1);
 	}
@@ -622,12 +592,7 @@ FDirectLighting ClearCoatBxDF( FGBufferData GBuffer, half3 N, half3 V, half3 L,
 		{
 			// NoL is chosen to provide better parity with DefaultLit when ClearCoat=0
 			Vis2 = Vis_SmithJointApprox(a2, BottomContext.NoV, NoL);
-#if SHADING_PATH_MOBILE
-			D2 = D_GGX_Mobile(GBuffer.Roughness, BottomContext.NoH);
-			Vis2 = saturate(Vis2);
-#else
 			D2 = D_GGX(a2, BottomContext.NoH);
-#endif
 		}
 		half3 F_Bot = F_Schlick(GBuffer.SpecularColor, BottomContext.VoH);
 		half3 F_DefaultLit = F_Schlick(GBuffer.SpecularColor, Context.VoH);
@@ -847,11 +812,7 @@ FDirectLighting TwoSidedBxDF( FGBufferData GBuffer, half3 N, half3 V, half3 L, f
 
 	// Scatter distribution
 	half VoL = dot(V, L);
-#if SHADING_PATH_MOBILE
-	half Scatter = D_GGX_Mobile(0.6, saturate(-VoL));
-#else
 	float Scatter = D_GGX( 0.6*0.6, saturate( -VoL ) );
-#endif
 
 	Lighting.Transmission = AreaLight.FalloffColor * (Falloff * WrapNoL * Scatter) * SubsurfaceColor;
 
@@ -910,12 +871,7 @@ FDirectLighting EyeBxDF( FGBufferData GBuffer, half3 N, half3 V, half3 L, float
 		float Energy = EnergyNormalization( a2, Context.VoH, AreaLight );
 		// Generalized microfacet specular
 		float Vis = Vis_SmithJointApprox(a2, Context.NoV, NoL);
-#if SHADING_PATH_MOBILE
-		half D = D_GGX_Mobile(GBuffer.Roughness, Context.NoH) * Energy;
-		Vis = saturate(Vis);
-#else
 		float D = D_GGX(a2, Context.NoH) * Energy;
-#endif
 		Lighting.Specular = AreaLight.FalloffColor * (Falloff * NoL) * D * Vis * F;
 	}
 
diff --git a/Engine/Shaders/Private/Substrate/SubstrateEvaluation.ush b/Engine/Shaders/Private/Substrate/SubstrateEvaluation.ush
--- a/Engine/Shaders/Private/Substrate/SubstrateEvaluation.ush
+++ b/Engine/Shaders/Private/Substrate/SubstrateEvaluation.ush
@@ -94,14 +94,8 @@
 // Point lobe in off-specular peak direction
 float3 SubstrateGetOffSpecularPeakReflectionDir(float3 Normal, float3 ReflectionVector, float Roughness)
 {
-	// This function is used for both regular & mobile path. For performance reason, mobile 
-	// version uses (optionally) a simplifed version.
-#if !SHADING_PATH_MOBILE || MOBILE_HIGH_QUALITY_BRDF
 	float a = Square(Roughness);
 	return lerp(Normal, ReflectionVector, (1 - a) * (sqrt(1 - a) + a));
-#else
-	return ReflectionVector;
-#endif
 }
 
 ///////////////////////////////////////////////////////////////////////////////
@@ -261,13 +255,9 @@ float3 LuminanceWeight(in FSubstrateBSDFContext InContext, in FSubstrateBSDF InB
 	return LuminanceWeight(SatNoL, InBSDF);
 }
 
-half Substrate_D_GGX(half Roughness, half a2, half NoH)
+float Substrate_D_GGX(half Roughness, half a2, half NoH)
 {
-#if SHADING_PATH_MOBILE
-	return D_GGX_Mobile(Roughness, NoH);
-#else
 	return D_GGX(a2, NoH);
-#endif
 }
 
 half Substrate_D_GGX_Aniso(float ax, float ay, float NoH, float XoH, float YoH)
@@ -277,9 +267,7 @@ half Substrate_D_GGX_Aniso(float ax, float ay, float NoH, float XoH, float YoH)
 
 half Substrate_Vis_GGX(half Roughness, half a2, half NoV, half NoL)
 {
-#if SHADING_PATH_MOBILE && !MOBILE_HIGH_QUALITY_BRDF
-	return (Roughness * 0.25 + 0.25);
-#elif SUBSTRATE_SHADING_QUALITY > 1
+#if SUBSTRATE_SHADING_QUALITY > 1
 	return Vis_SmithJoint(a2, NoV, NoL);
 #else
 	return Vis_SmithJointApprox(a2, NoV, NoL);
@@ -292,14 +280,9 @@ half Substrate_Vis_GGX_Aniso(float ax, float ay, float NoV, float NoL, float XoV
 	return Vis_SmithJointAniso(ax, ay, NoV, NoL, XoV, XoL, YoV, YoL);
 }
 
-half3 Substrate_F_GGX(half3 F0, half3 F90, half VoH)
+float3 Substrate_F_GGX(half3 F0, half3 F90, half VoH)
 {
-// SUBSTRATE_TODO: Legacy mobile path use EnvBRDFApprox in place of Fresnel evaluation. However its cost seems higher. Need reevaluate if this is needed.
-//#if SHADING_PATH_MOBILE && !MOBILE_HIGH_QUALITY_BRDF
-//	return EnvBRDFApprox(F0, F90, Roughness, NoV);
-//#else
 	return F_Schlick(F0, F90, VoH);
-//#endif
 }
 
 // Transition from Glint to GGX (0=Glint <-> 1=GGX)
diff --git a/Engine/Source/Editor/UnrealEd/Private/Commandlets/DumpMaterialInfo.cpp b/Engine/Source/Editor/UnrealEd/Private/Commandlets/DumpMaterialInfo.cpp
--- a/Engine/Source/Editor/UnrealEd/Private/Commandlets/DumpMaterialInfo.cpp
+++ b/Engine/Source/Editor/UnrealEd/Private/Commandlets/DumpMaterialInfo.cpp
@@ -271,7 +271,6 @@ TArray<MaterialInfo::FPropertySet> GetMaterialInfoProperties()
 		FPropertySet::FromAccessor<bool, &FMaterialResource::IsUsingPlanarForwardReflections>(TEXT("IsUsingPlanarForwardReflections")),
 		FPropertySet::FromAccessor<bool, &FMaterialResource::IsNonmetal>(TEXT("IsNonmetal")),
 		FPropertySet::FromAccessor<bool, &FMaterialResource::UseLmDirectionality>(TEXT("UseLmDirectionality")),
-		FPropertySet::FromAccessor<bool, &FMaterialResource::IsMobileHighQualityBRDFEnabled>(TEXT("IsMobileHighQualityBRDFEnabled")),
 		FPropertySet::FromEnumAccessor<EBlendMode, &FMaterialResource::GetBlendMode>(TEXT("GetBlendMode")),
 		FPropertySet::FromEnumAccessor<ERefractionMode, &FMaterialResource::GetRefractionMode>(TEXT("GetRefractionMode")),
 		FPropertySet::FromAccessor<bool, &FMaterialResource::GetRootNodeOverridesDefaultRefraction>(TEXT("GetRootNodeOverridesDefaultRefraction")),
diff --git a/Engine/Source/Runtime/Engine/Private/Materials/Material.cpp b/Engine/Source/Runtime/Engine/Private/Materials/Material.cpp
--- a/Engine/Source/Runtime/Engine/Private/Materials/Material.cpp
+++ b/Engine/Source/Runtime/Engine/Private/Materials/Material.cpp
@@ -4927,7 +4927,6 @@ bool UMaterial::CanEditChange(const FProperty* InProperty) const
 			PropertyName == GET_MEMBER_NAME_STRING_CHECKED(UMaterial, bUseLightmapDirectionality) ||
 			PropertyName == GET_MEMBER_NAME_STRING_CHECKED(UMaterial, bUseHQForwardReflections) ||
 			PropertyName == GET_MEMBER_NAME_STRING_CHECKED(UMaterial, bForwardBlendsSkyLightCubemaps) ||
-			PropertyName == GET_MEMBER_NAME_STRING_CHECKED(UMaterial, bMobileEnableHighQualityBRDF) ||
 			PropertyName == GET_MEMBER_NAME_STRING_CHECKED(UMaterial, bUsePlanarForwardReflections)
 			)
 		{
diff --git a/Engine/Source/Runtime/Engine/Private/Materials/MaterialShared.cpp b/Engine/Source/Runtime/Engine/Private/Materials/MaterialShared.cpp
--- a/Engine/Source/Runtime/Engine/Private/Materials/MaterialShared.cpp
+++ b/Engine/Source/Runtime/Engine/Private/Materials/MaterialShared.cpp
@@ -1841,11 +1841,6 @@ bool FMaterialResource::UseLmDirectionality() const
 	return Material->bUseLightmapDirectionality;
 }
 
-bool FMaterialResource::IsMobileHighQualityBRDFEnabled() const
-{
-	return Material->bMobileEnableHighQualityBRDF;
-}
-
 /**
  * Should shaders compiled for this material be saved to disk?
  */
@@ -2595,7 +2590,6 @@ void FMaterial::SetupMaterialEnvironment(
 	SET_SHADER_DEFINE(OutEnvironment, TRANSLUCENCY_VELOCITY_FROM_DEPTH, IsTranslucencyWritingVelocity() && IsTranslucencyVelocityFromDepth());
 	SET_SHADER_DEFINE(OutEnvironment, TRANSLUCENT_WRITING_FRONT_LAYER_TRANSPARENCY, IsTranslucencyWritingFrontLayerTransparency());
 	SET_SHADER_DEFINE(OutEnvironment, MATERIAL_USE_ALPHA_TO_COVERAGE, IsUsingAlphaToCoverage() ? 1 : 0);
-	SET_SHADER_DEFINE(OutEnvironment, MOBILE_HIGH_QUALITY_BRDF, IsMobileHighQualityBRDFEnabled());
 	SET_SHADER_DEFINE(OutEnvironment, MATERIAL_TRANSLUCENT_PASS_AFTERMOTIONBLUR, IsTranslucencyAfterMotionBlurEnabled() ? 1 : 0);
 
 	bool bFullPrecisionInMaterial = false;
diff --git a/Engine/Source/Runtime/Engine/Public/MaterialShared.h b/Engine/Source/Runtime/Engine/Public/MaterialShared.h
--- a/Engine/Source/Runtime/Engine/Public/MaterialShared.h
+++ b/Engine/Source/Runtime/Engine/Public/MaterialShared.h
@@ -2183,7 +2183,6 @@ public:
 	virtual bool IsUsingPlanarForwardReflections() const { return false; }
 	virtual bool IsNonmetal() const { return false; }
 	virtual bool UseLmDirectionality() const { return true; }
-	virtual bool IsMobileHighQualityBRDFEnabled() const { return false; }
 	virtual bool IsMasked() const = 0;
 	virtual bool IsDitherMasked() const { return false; }
 	virtual bool AllowNegativeEmissiveColor() const { return false; }
@@ -2968,7 +2967,6 @@ public:
 	ENGINE_API virtual bool IsUsingPlanarForwardReflections() const override;
 	ENGINE_API virtual bool IsNonmetal() const override;
 	ENGINE_API virtual bool UseLmDirectionality() const override;
-	ENGINE_API virtual bool IsMobileHighQualityBRDFEnabled() const override;
 	ENGINE_API virtual enum EBlendMode GetBlendMode() const override;
 	ENGINE_API virtual enum ERefractionMode GetRefractionMode() const override;
 	ENGINE_API virtual bool GetRootNodeOverridesDefaultRefraction() const override;
diff --git a/Engine/Source/Runtime/Engine/Public/Materials/Material.h b/Engine/Source/Runtime/Engine/Public/Materials/Material.h
--- a/Engine/Source/Runtime/Engine/Public/Materials/Material.h
+++ b/Engine/Source/Runtime/Engine/Public/Materials/Material.h
@@ -875,10 +875,6 @@ public:
 	UPROPERTY(EditAnywhere, BlueprintReadOnly, Category=Mobile)
 	uint8 bUseLightmapDirectionality : 1;
 
-	/* Use the high quality brdf functions on mobile to get better visual effects but adds GPU cost. */
-	UPROPERTY(EditAnywhere, BlueprintReadOnly, Category=Mobile, meta = (DisplayName = "Mobile High Quality BRDF"))
-	uint8 bMobileEnableHighQualityBRDF : 1;
-
 	/* Use alpha to coverage for masked material on mobile, make sure MSAA is enabled as well. */
 	UPROPERTY(EditAnywhere, BlueprintReadOnly, Category=Mobile, AdvancedDisplay, meta = (EditCondition = "BlendMode != EBlendMode::BLEND_Opaque"))
 	uint8 bUseAlphaToCoverage : 1;
