MSAA resolve with explicit fmask produces incorrect alpha values
-We were accumulating the sample values into a 0,0,0,1 vector, so there was alway an extra +1/msaaLevel into the resolved pixel value (ie msaa2x alpha ranged from 0.5-1.5 and 4x from 0.25-1.25).  In 4x, for example, the 1.25 would be clamped down to 1, aka full transparency in unreal's convention but the 0.25 would remain meaning we never had full opacity.
#tests Ran Rendering tests on horde, win64 no change to failures. Ran lanestreet, switched it to forward renderer with msaa and enabled alpha propagaion: alpha values looked correct, confirmed this shader was running in that configuration. Locally ran more render tests, seemed to have the same failures with this change.

diff --git a/Engine/Shaders/Private/HdrCustomResolveShaders.usf b/Engine/Shaders/Private/HdrCustomResolveShaders.usf
--- a/Engine/Shaders/Private/HdrCustomResolveShaders.usf
+++ b/Engine/Shaders/Private/HdrCustomResolveShaders.usf
@@ -75,7 +75,7 @@ float4 HdrCustomResolveFMaskPS(float4 Pos : SV_POSITION
 {
 	MSCoords P = MSCoordsFromFloat4(Pos);
 
-	float4 sampleSum = float4(0, 0, 0, 1);
+	float4 sampleSum = float4(0, 0, 0, 0);
 	float weightSum = 0;
 
 	for (int i = 0; i < HDR_RESOLVE_NUM_SAMPLES; ++i)
