[Distance Fields] Fix distance field visualization when using multiple views.
- apply View.ViewRectMin.xy when writing and sampling VisualizeResultTexture.

diff --git a/Engine/Shaders/Private/DistanceFieldVisualization.usf b/Engine/Shaders/Private/DistanceFieldVisualization.usf
--- a/Engine/Shaders/Private/DistanceFieldVisualization.usf
+++ b/Engine/Shaders/Private/DistanceFieldVisualization.usf
@@ -266,7 +266,7 @@ void VisualizeMeshDistanceFieldCS(
 
 	float3 Result = Lighting * View.PreExposure;
 
-	RWVisualizeMeshDistanceFields[DispatchThreadId.xy] = float4(Result, 0);
+	RWVisualizeMeshDistanceFields[DispatchThreadId.xy + View.ViewRectMin.xy / DOWNSAMPLE_FACTOR] = float4(Result, 0);
 }
 
 Texture2D VisualizeDistanceFieldTexture;
@@ -274,10 +274,7 @@ SamplerState VisualizeDistanceFieldSampler;
 
 void VisualizeDistanceFieldUpsamplePS(in float4 UVAndScreenPos : TEXCOORD0, out float4 OutColor : SV_Target0)
 {
-	// Distance field AO was computed at 0,0 regardless of viewrect min
-	float2 DistanceFieldUVs = UVAndScreenPos.xy - View.ViewRectMin.xy * View.BufferSizeAndInvSize.zw;
-
-	float3 Value = Texture2DSampleLevel(VisualizeDistanceFieldTexture, VisualizeDistanceFieldSampler, DistanceFieldUVs, 0).xyz;
+	float3 Value = Texture2DSampleLevel(VisualizeDistanceFieldTexture, VisualizeDistanceFieldSampler, UVAndScreenPos.xy, 0).xyz;
 
 	OutColor = float4(Value, 1);
 }
