Ray Tracing: Refactor shadow tracing for translucency to use AHS

The same trick discovered for path tracing to avoid double-counting of shadow transmission can be applied to raytraced shadows, which greatly simplifies the code.

Also add suport for translucent shadows outside of the RT shadow pass (for example the primary ray debug mode and Lumen hit lighting).

This CL also fixes a few blend modes when using substrate and the path tracer, though there are a few remaining cases to iron out.

With this new mechanism for tracing translucency, the "high quality effects" entry in DDPI is no longer needed so it can be removed.

diff --git a/Engine/Config/VulkanPC/DataDrivenPlatformInfo.ini b/Engine/Config/VulkanPC/DataDrivenPlatformInfo.ini
--- a/Engine/Config/VulkanPC/DataDrivenPlatformInfo.ini
+++ b/Engine/Config/VulkanPC/DataDrivenPlatformInfo.ini
@@ -178,7 +178,6 @@ bSupportsRayTracingShaders = true
 bSupportsPathTracing = true
 bSupportsRayTracingCallableShaders = true
 bSupportsRayTracingProceduralPrimitive = true
-bSupportsHighEndRayTracingEffects=true
 
 bSupportsGPUScene = true
 bSupportsPrimitiveShaders = false
diff --git a/Engine/Config/Windows/DataDrivenPlatformInfo.ini b/Engine/Config/Windows/DataDrivenPlatformInfo.ini
--- a/Engine/Config/Windows/DataDrivenPlatformInfo.ini
+++ b/Engine/Config/Windows/DataDrivenPlatformInfo.ini
@@ -91,7 +91,6 @@ bSupportsRayTracingProceduralPrimitive=true
 bSupportsRayTracingIndirectInstanceData=true
 bSupportsPathTracing=true
 bSupportsShaderExecutionReordering=true
-bSupportsHighEndRayTracingEffects=true
 bSupportsGPUScene=true
 bSupportsNanite=true
 bSupportsSceneDataCompressedTransforms=true
diff --git a/Engine/Shaders/Private/PackUnpack.ush b/Engine/Shaders/Private/PackUnpack.ush
--- a/Engine/Shaders/Private/PackUnpack.ush
+++ b/Engine/Shaders/Private/PackUnpack.ush
@@ -307,6 +307,24 @@ float3 UnpackRGB998E6(uint bits)
 	return result * f;
 }
 
+uint PackRGB111110(float3 rgb)
+{
+	// NOTE: Could add a simple gamma2 here to improve precision of dark colors, but this does not seem necessary for now.
+	// unorm 11/11/10 encoding
+	uint r = uint(floor(saturate(rgb.r) * 2047.0 + 0.5));
+	uint g = uint(floor(saturate(rgb.g) * 2047.0 + 0.5));
+	uint b = uint(floor(saturate(rgb.b) * 1023.0 + 0.5));
+	return (r << 21) | (g << 10) | b;
+}
+
+float3 UnpackRGB111110(uint rgb)
+{
+	float r = float((rgb >> 21)       ) * (1.0 / 2047.0);
+	float g = float((rgb >> 10) & 2047) * (1.0 / 2047.0);
+	float b = float((rgb      ) & 1023) * (1.0 / 1023.0);
+	return float3(r, g, b);
+}
+
 
 //
 //	float4 <-> uint2 as f16.f16.f16.f16
diff --git a/Engine/Shaders/Private/PathTracing/PathTracingCommon.ush b/Engine/Shaders/Private/PathTracing/PathTracingCommon.ush
--- a/Engine/Shaders/Private/PathTracing/PathTracingCommon.ush
+++ b/Engine/Shaders/Private/PathTracing/PathTracingCommon.ush
@@ -123,20 +123,12 @@ float3 PayloadDecodeHDRColor(uint rgb)
 
 uint PayloadEncodeLDRColor(float3 rgb)
 {
-	// NOTE: Could add a simple gamma2 here to improve precision of dark colors, but this does not seem necessary for now.
-	// unorm 11/11/10 encoding
-	uint r = uint(floor(saturate(rgb.r) * 2047.0 + 0.5));
-	uint g = uint(floor(saturate(rgb.g) * 2047.0 + 0.5));
-	uint b = uint(floor(saturate(rgb.b) * 1023.0 + 0.5));
-	return (r << 21) | (g << 10) | b;
+	return PackRGB111110(rgb);
 }
 
 float3 PayloadDecodeLDRColor(uint rgb)
 {
-	float r = float((rgb >> 21)       ) * (1.0 / 2047.0);
-	float g = float((rgb >> 10) & 2047) * (1.0 / 2047.0);
-	float b = float((rgb      ) & 1023) * (1.0 / 1023.0);
-	return float3(r, g, b);
+	return UnpackRGB111110(rgb);
 }
 
 uint PayloadEncodeRoughnessAniso(float3 RoughnessData, float Anisotropy)
diff --git a/Engine/Shaders/Private/PathTracing/PathTracingMaterialHitShader.usf b/Engine/Shaders/Private/PathTracing/PathTracingMaterialHitShader.usf
--- a/Engine/Shaders/Private/PathTracing/PathTracingMaterialHitShader.usf
+++ b/Engine/Shaders/Private/PathTracing/PathTracingMaterialHitShader.usf
@@ -227,7 +227,7 @@ RAY_TRACING_ENTRY_CLOSEST_HIT(PathTracingMaterialCHS,
 		UnlitSurfaceCoverage,
 		UnlitSurfaceTransmittancePreCoverage,
 		UnlitSurfaceNormal);
-	const float3 Transparency = saturate(UnlitSurfaceCoverage + (1.0 - UnlitSurfaceCoverage) * UnlitSurfaceTransmittancePreCoverage);
+	const float3 Transparency = saturate(lerp(1.0, UnlitSurfaceTransmittancePreCoverage, UnlitSurfaceCoverage));
 #else
 	// SUBSTRATE_TODO: These should actually be initialized via the corresponding cvars. To do this properly would require binding FSubstrateSceneData in the path tracing shaders
 	// For now, simply match the default substrate settings
@@ -252,7 +252,7 @@ RAY_TRACING_ENTRY_CLOSEST_HIT(PathTracingMaterialCHS,
 #else
 	// transparency is used for both coverage and transmittance
 	// SUBSTRATE_TODO: Thin case does not seem to be properly handled here, so the TotalTransmittancePreCoverage is recomputed later
-	float3 Transparency = saturate((1 - TotalCoverage) + TotalCoverage * TotalTransmittancePreCoverage);
+	const float3 Transparency = saturate(lerp(1.0, TotalTransmittancePreCoverage, TotalCoverage));
 #endif // REFRACTION_USE_INDEX_OF_REFRACTION
 #endif // SUBSTRATE_OPTIMIZED_UNLIT
 
@@ -1350,12 +1350,7 @@ RAY_TRACING_ENTRY_ANY_HIT(PathTracingMaterialAHS,
 
 		float3 Transparency = 0.0;
 
-#if MATERIAL_IS_SUBSTRATE && !SUBSTRATE_OPTIMIZED_UNLIT
-		// We already figured out how much Transparency there is above
-		// So only do extra logic for refraction and thin refraction ases
-
-
-#if MATERIALBLENDING_TRANSLUCENT && (REFRACTION_USE_INDEX_OF_REFRACTION || MATERIAL_ISTHINSURFACE)
+#if MATERIAL_IS_SUBSTRATE 
 		FSubstratePixelHeader SubstratePixelHeader = MaterialParameters.GetFrontSubstrateHeader();
 		FSubstrateData SubstrateData = PixelMaterialInputs.GetFrontSubstrateData();
 
@@ -1370,6 +1365,7 @@ RAY_TRACING_ENTRY_ANY_HIT(PathTracingMaterialAHS,
 		float3 TotalTransmittancePreCoverage = 0;
 		const float3 V_World = MaterialParameters.CameraVector;
 		SubstratePixelHeader.SubstrateUpdateTree(SubstrateData, V_World, IntegrationSettings, TotalCoverage, TotalTransmittancePreCoverage);
+#if MATERIALBLENDING_TRANSLUCENT && (REFRACTION_USE_INDEX_OF_REFRACTION || MATERIAL_ISTHINSURFACE)
 
 		// SUBSTRATE_TODO: refraction handling
 		float PathRoughness = PackedPayload.GetPathRoughness();
@@ -1458,9 +1454,11 @@ RAY_TRACING_ENTRY_ANY_HIT(PathTracingMaterialAHS,
 
 		// SUBSTRATE_TODO: Is this sanity check necessary?
 		Transparency = saturate(Transparency);
-#endif // MATERIALBLENDING_TRANSLUCENT && (REFRACTION_USE_INDEX_OF_REFRACTION || MATERIAL_ISTHINSURFACE)
-
-#endif // MATERIAL_IS_SUBSTRATE && !SUBSTRATE_OPTIMIZED_UNLIT
+#else // MATERIALBLENDING_TRANSLUCENT && (REFRACTION_USE_INDEX_OF_REFRACTION || MATERIAL_ISTHINSURFACE)
+		// All other blend modes, do simple transparency handling
+		Transparency = lerp(1.0, TotalTransmittancePreCoverage, TotalCoverage);
+#endif
+#endif // MATERIAL_IS_SUBSTRATE
 
 #else // SUBSTRATE_ENABLED
 
diff --git a/Engine/Shaders/Private/RayTracing/RayTracingCommon.ush b/Engine/Shaders/Private/RayTracing/RayTracingCommon.ush
--- a/Engine/Shaders/Private/RayTracing/RayTracingCommon.ush
+++ b/Engine/Shaders/Private/RayTracing/RayTracingCommon.ush
@@ -656,6 +656,26 @@ struct FPackedMaterialClosestHitPayload : FMinimalPayload
 	{ 
 		return UnpackNormalFromUInt32(PackedGeometryNormal); 
 	}
+	
+	float3 GetShadowVisibility()
+	{
+		return UnpackRGB111110(TransmittancePreCoverage); // TODO: Does this interfere with the BRDF evaluation in the lighting miss shader?
+	}
+
+	void SetShadowVisibility(float3 ShadowVisibility)
+	{
+		TransmittancePreCoverage = PackRGB111110(ShadowVisibility);
+	}
+
+	uint GetLightIndex()
+	{
+		return SceneInstanceIndex;
+	}
+
+	void SetLightIndex(uint LightIndex)
+	{
+		SceneInstanceIndex = LightIndex;
+	}
 
 	bool IsValid()							{ return any(SubstrateData.PackedTopLayerData != 0); }
 	uint  GetShadingModelID()				{ return IsValid() ? SHADINGMODELID_SUBSTRATE : SHADINGMODELID_UNLIT; }
@@ -673,8 +693,8 @@ struct FPackedMaterialClosestHitPayload : FMinimalPayload
 	void SetOpacity(float In)						{ Flags  = (Flags & 0x00FFFFFF) | (PackR8(In) << 24u); }
 	void SetTransmittancePreCoverage(float3 In)		{ TransmittancePreCoverage  = PackR11G11B10F(In); }
 
-	void SetRayDirection(float3 In)			{ PackedRayCone = SubstratePackNormal(In); }
-	float3 GetRayDirection()				{ return SubstrateUnpackNormal(PackedRayCone); }
+	void SetRayDirection(float3 In)			{ PackedRayCone = PackNormalToUInt32(In); }
+	float3 GetRayDirection()				{ return UnpackNormalFromUInt32(PackedRayCone); }
 
 	// Member function for (IsLumenPayload() == 1)
 	void SetSceneInstanceIndex(uint In)		{ SceneInstanceIndex = In; }
@@ -857,7 +877,7 @@ struct FPackedMaterialClosestHitPayload : FMinimalPayload
 	{
 		PackedIndirectIrradiance[0]  = f32tof16(InIndirectIrradiance.x);
 		PackedIndirectIrradiance[0] |= f32tof16(InIndirectIrradiance.y) << 16;
-		PackedIndirectIrradiance[1]  = f32tof16(InIndirectIrradiance.z);
+		PackedIndirectIrradiance[1]  = f32tof16(InIndirectIrradiance.z) | (PackedIndirectIrradiance[1] & 0xFFFF0000u);
 	}
 
 	float3 GetIndirectIrradiance()
@@ -870,8 +890,8 @@ struct FPackedMaterialClosestHitPayload : FMinimalPayload
 	}
 	
 	// Alias IndirectIrradiance for storing ray direction
-	void SetRayDirection(float3 In) { SetIndirectIrradiance(In); }
-	float3 GetRayDirection()		{ return GetIndirectIrradiance(); }
+	void SetRayDirection(float3 In) { PackedRayCone = PackNormalToUInt32(In); }
+	float3 GetRayDirection()		{ return UnpackNormalFromUInt32(PackedRayCone); }
 
 	float3 GetWorldNormal()
 	{
@@ -911,8 +931,29 @@ struct FPackedMaterialClosestHitPayload : FMinimalPayload
 		return Result;
 	}
 
+	float3 GetShadowVisibility()
+	{
+		return GetIndirectIrradiance();
+	}
+
+	void SetShadowVisibility(float3 ShadowVisibility)
+	{
+		SetIndirectIrradiance(ShadowVisibility);
+	}
+
+	uint GetLightIndex()
+	{
+		return PackedIndirectIrradiance[1] >> 16;
+	}
+
+	void SetLightIndex(uint LightIndex)
+	{
+		// NOTE: this assumes a maximum of 64k lights
+		PackedIndirectIrradiance[1] = (LightIndex << 16) | (PackedIndirectIrradiance[1] & 0xFFFFu);
+	}
+
 	float GetOpacity()       { return f16tof32(BaseColorAndOpacity[1] >> 16); }
-	void SetOpacity(float Opacity) { BaseColorAndOpacity[1] = f32tof16(Opacity) << 16 | (BaseColorAndOpacity[1] & 0xFF); }
+	void SetOpacity(float Opacity) { BaseColorAndOpacity[1] = f32tof16(Opacity) << 16 | (BaseColorAndOpacity[1] & 0xFFFF); }
 	float GetMetallic()      { return float(MetallicAndSpecularAndRoughness & 0xFF) / 255.0f; }
 	float GetSpecular()      { return float((MetallicAndSpecularAndRoughness >> 8) & 0xFF) / 255.0f; }
 	float GetRoughness()     { return f16tof32(MetallicAndSpecularAndRoughness >> 16); }
@@ -1196,27 +1237,25 @@ FMinimalPayload TraceVisibilityRay(
 
 struct FTranslucentMinimalPayload : FMinimalPayload
 {
-	float Opacity; // The opacity for the translucent material.
+	float3 ShadowVisibility; // What fraction of the ray's radiance is visible? (1.0 - fully visible, 0.0 - fully shadowed)
 };
 
-void TraceTranslucentVisibilityRayPacked(
-	inout FPackedMaterialClosestHitPayload PackedPayload,
+FTranslucentMinimalPayload TraceTranslucentVisibilityRay(
 	in RaytracingAccelerationStructure TLAS,
 	in uint RayFlags,
 	in uint InstanceInclusionMask,
 	in FRayDesc Ray)
 {
+	FPackedMaterialClosestHitPayload PackedPayload = (FPackedMaterialClosestHitPayload)0;
+	PackedPayload.SetShadowRay();
+	PackedPayload.SetMinimalPayloadMode();
+	PackedPayload.HitT = 0; // default value in case we don't run the miss shader
+	PackedPayload.SetShadowVisibility(1.0);
+
 	const uint RayContributionToHitGroupIndex = RAY_TRACING_SHADER_SLOT_SHADOW;
 	const uint MultiplierForGeometryContributionToShaderIndex = RAY_TRACING_NUM_SHADER_SLOTS;
 	const uint MissShaderIndex = 0;
 
-	// By enabling minimal payload mode all other payload information is ignored, meaning these functions need no payload inputs
-	PackedPayload.SetMinimalPayloadMode();
-	PackedPayload.HitT = 0;
-
-	// @todo: add mask to select translucent shadow only.
-	RayFlags |= RAY_TRACING_MASK_TRANSLUCENT;
-
 	// Trace the ray
 	TraceRay(
 		TLAS,
@@ -1227,25 +1266,16 @@ void TraceTranslucentVisibilityRayPacked(
 		MissShaderIndex,
 		Ray.GetNativeDesc(),
 		PackedPayload);
-}
-
-FTranslucentMinimalPayload TraceTranslucentVisibilityRay(
-	in RaytracingAccelerationStructure TLAS,
-	in uint RayFlags,
-	in uint InstanceInclusionMask,
-	in FRayDesc Ray)
-{
-	FPackedMaterialClosestHitPayload PackedPayload = (FPackedMaterialClosestHitPayload)0;
-
-	PackedPayload.SetFlags(RAY_TRACING_PAYLOAD_INPUT_FLAG_SHADOW_RAY);
-
-	TraceTranslucentVisibilityRayPacked(PackedPayload, TLAS, RayFlags, InstanceInclusionMask, Ray);
 
 	// Unpack the payload
 	FTranslucentMinimalPayload MinimalPayload = (FTranslucentMinimalPayload)0;
 
 	MinimalPayload.HitT = PackedPayload.HitT;
-	MinimalPayload.Opacity = PackedPayload.GetOpacity();
+	// only read the throughput if we got a miss (ignored all hits and ran the miss shader which reset HitT)
+	// If we registered a hit, it means that either we ran the closest hit shader which set a valid distance,
+	// or we never ran the miss shader and kept the HitT value from above (or one set by the AHS before the ray was terminated)
+	// in both cases we know the shadow visibility should be counted as zero.
+	MinimalPayload.ShadowVisibility = PackedPayload.IsMiss() ? PackedPayload.GetShadowVisibility() : 0.0;
 
 	return MinimalPayload;
 }
diff --git a/Engine/Shaders/Private/RayTracing/RayTracingLightingCommon.ush b/Engine/Shaders/Private/RayTracing/RayTracingLightingCommon.ush
--- a/Engine/Shaders/Private/RayTracing/RayTracingLightingCommon.ush
+++ b/Engine/Shaders/Private/RayTracing/RayTracingLightingCommon.ush
@@ -323,8 +323,14 @@ float3 ComputeDirectLighting(
 	float3 DirectLighting = (float3)0;
 
 	// Repurpose some fields in the material payload to pass parameters into lighting miss shader.
+	// Backup the values that will be aliased with other entries
 	float3 OldRadiance = Payload.GetRadiance();
-	float3 OldIndirectIrradiance = Payload.GetIndirectIrradiance(); // Aliased with ViewDirection
+#if SUBSTRATE_ENABLED
+	uint OldTransmittancePreCoverage = Payload.TransmittancePreCoverage;
+#else
+	uint2 OldIndirectIrradiance = uint2(Payload.PackedIndirectIrradiance[0], Payload.PackedIndirectIrradiance[1]);
+#endif
+	uint OldRayCone = Payload.PackedRayCone;
 	Payload.SetRadiance(float3(0, 0, 0));
 	Payload.SetRayDirection(ViewDirection);
 
@@ -407,21 +413,22 @@ float3 ComputeDirectLighting(
 
 		uint RayFlags = RAY_FLAG_ACCEPT_FIRST_HIT_AND_END_SEARCH	// Terminate immediatly after detecting a hit
 						| RAY_FLAG_SKIP_CLOSEST_HIT_SHADER;			// Skip Closest Hit shader, only run Any Hit shader
-		const uint InstanceInclusionMask = (ApplyShadow) ? RAY_TRACING_MASK_SHADOW : 0;
+		const uint InstanceInclusionMask = (ApplyShadow) ? RAY_TRACING_MASK_SHADOW | RAY_TRACING_MASK_TRANSLUCENT_SHADOW : 0;
 
-		// Do not include any translucent when tracing for shadow, since we do not handle coverage & transmittance yet.
-		Payload.SetIgnoreTranslucentMaterials();
+		Payload.SetShadowRay();
+		Payload.SetMinimalPayloadMode();
+		Payload.SetShadowVisibility(1.0);
+		Payload.HitT = 0.0;
 
 #if !ENABLE_TWO_SIDED_GEOMETRY
 		RayFlags |= RAY_FLAG_CULL_FRONT_FACING_TRIANGLES;
 #endif // !ENABLE_TWO_SIDED_GEOMETRY
 
-		// Light index is packed into HitT as this component is only accessed by closest hit or miss shaders.
-		// Since closest hit execution is disabled using a ray flag, it is safe to pack custom data here.
-		Payload.HitT = asfloat(LightIndex);
+		// Light index is packed into IndirectIrradiance as well, so is covered by the save/restore of IrradianceColor above
+		Payload.SetLightIndex(LightIndex);
 
-		bool bMissFarField = true;
 #if ENABLE_FAR_FIELD_TRACING
+		bool bMissFarField = true;
 		{
 			// Rebase origin for far-field rays
 			FRayDesc FarFieldShadowRay = ShadowRay;
@@ -432,10 +439,10 @@ float3 ComputeDirectLighting(
 
 			bMissFarField = DefaultPayload.IsMiss();
 		}
-#endif // ENABLE_FAR_FIELD_TRACING
 
 		// use the bound miss shader lighting evaluation if lighting, else nothing
 		if (bMissFarField)
+#endif // ENABLE_FAR_FIELD_TRACING
 		{
 			uint MissShaderIndex = Lit ? LightMissShaderIndex : RAY_TRACING_MISS_SHADER_SLOT_DEFAULT;
 			TraceShadowRayMissShaderLighting(ShadowRay, RayFlags, InstanceInclusionMask, TLAS, MissShaderIndex, Payload);
@@ -444,8 +451,14 @@ float3 ComputeDirectLighting(
 
 	DirectLighting = Payload.GetRadiance();
 	Payload.SetRadiance(OldRadiance);
-	Payload.SetIndirectIrradiance(OldIndirectIrradiance); // Set back indirect irradiance (previously aliased with view direction
-
+	// Restore aliased fields
+#if SUBSTRATE_ENABLED
+	Payload.TransmittancePreCoverage = OldTransmittancePreCoverage;
+#else
+	Payload.PackedIndirectIrradiance[0] = OldIndirectIrradiance.x;  // Set back indirect irradiance (aliased with ShadowVisibility and LightIndex)
+	Payload.PackedIndirectIrradiance[1] = OldIndirectIrradiance.y;
+#endif
+	Payload.PackedRayCone = OldRayCone;
 	return DirectLighting;
 }
 
diff --git a/Engine/Shaders/Private/RayTracing/RayTracingLightingMS.usf b/Engine/Shaders/Private/RayTracing/RayTracingLightingMS.usf
--- a/Engine/Shaders/Private/RayTracing/RayTracingLightingMS.usf
+++ b/Engine/Shaders/Private/RayTracing/RayTracingLightingMS.usf
@@ -58,14 +58,16 @@ RAY_TRACING_ENTRY_MISS(RayTracingLightingMS, FPackedMaterialClosestHitPayload, P
 
 	ResolvedView = ResolveView();
 
+	uint LightIndex = PackedPayload.GetLightIndex();
+	float3 CameraVector = PackedPayload.GetRayDirection();// IndirectIrradiance field is used to pass through the camera vector
 	FMaterialClosestHitPayload Payload = UnpackRayTracingPayload(PackedPayload, DummyRay);
-
-	// Light index is packed into HitT as this component is only accessed by closest hit or miss shaders.
-	// Since closest hit execution is disabled using a ray flag, it is safe to pack custom data here.
-	uint LightIndex = asuint(Payload.HitT);
-	float3 CameraVector = Payload.GetRayDirection();// IndirectIrradiance field is used to pass through the camera vector
 	float3 V = -CameraVector;
-	Payload.IndirectIrradiance = (float3)0;
+	// reset the data that is aliased with shadow visibility
+#if SUBSTRATE_ENABLED
+	Payload.TransmittancePreCoverage = 1; 
+#else
+	Payload.IndirectIrradiance = 0; 
+#endif
 
 	uint LightType = 0;
 	FDeferredLightData LightData = GetRayTracingDeferredLightData(LightIndex, PackedPayload.IsIndirectLightingRay(), LightType);
@@ -133,6 +135,9 @@ RAY_TRACING_ENTRY_MISS(RayTracingLightingMS, FPackedMaterialClosestHitPayload, P
 	LightProfileMultiplier *= dot(LightFunctionMult, .3333f); // convert to greyscale (for compatibility with the appearance in the primary view)
 #endif
 
+	// Apply shadow transparency factor
+	LightContribution *= PackedPayload.GetShadowVisibility();
+
 	float3 AccumulatedRadiance = Payload.Radiance + LightContribution * LightProfileMultiplier;
 
 	PackedPayload.SetRadiance(AccumulatedRadiance);
diff --git a/Engine/Shaders/Private/RayTracing/RayTracingMaterialHitShaders.usf b/Engine/Shaders/Private/RayTracing/RayTracingMaterialHitShaders.usf
--- a/Engine/Shaders/Private/RayTracing/RayTracingMaterialHitShaders.usf
+++ b/Engine/Shaders/Private/RayTracing/RayTracingMaterialHitShaders.usf
@@ -30,6 +30,7 @@
 #include "RayTracingCommon.ush"
 #include "RayTracingHitGroupCommon.ush"
 #include "RayTracingShaderUtils.ush"
+#include "/Engine/Private/PathTracing/Material/PathTracingFresnel.ush"
 
 #include "/Engine/Generated/Material.ush"   
 #include "/Engine/Generated/VertexFactory.ush"
@@ -142,6 +143,35 @@ void GetSkyLighting(
 #endif
 }
 
+#if !SUBSTRATE_ENABLED
+// TODO: Consolidate this definition with similar one in PathTracing Material
+float GetRefractionIor(FPixelMaterialInputs PixelMaterialInputs)
+{
+#if MATERIALBLENDING_TRANSLUCENT && REFRACTION_USE_INDEX_OF_REFRACTION
+	// Is the material using refraction at all?
+#if REFRACTION_ROOT_NODE_OVERRIDES_DEFAULT
+	// Did the user plug something into the Refraction port?
+	// Return the user's IOR value, or 0 (ensure we don't get anything negative)
+	return max(GetMaterialRefractionIOR(GetMaterialRefraction(PixelMaterialInputs)), 0.0);
+#else
+	// User did not override the default, compute it via the amount of specular
+	float3 BaseColor = GetMaterialBaseColor(PixelMaterialInputs);
+	float  Metallic = GetMaterialMetallic(PixelMaterialInputs);
+	float  Specular = GetMaterialSpecular(PixelMaterialInputs);
+	float3 F0 = ComputeF0(Specular, BaseColor, Metallic);
+	return DielectricF0RGBToIor(F0);
+#endif
+
+#else
+	// Material either not translucent, or does use have IOR refraction enabled
+	return 0.0;
+#endif
+}
+#endif // !SUBSTRATE_ENABLED
+
+
+
+
 /** Calculates indirect lighting contribution on this object from precomputed data. */
 void GetPrecomputedIndirectLightingAndSkyLight(
 	FMaterialPixelParameters MaterialParameters, 
@@ -296,8 +326,12 @@ void GetMaterialPayload(
 			TransmittancePreCoverage);
 
 	}
+	#elif MATERIALBLENDING_MASKED && MATERIAL_DITHER_OPACITY_MASK
+	// If we have dithered masking, convert to regular transparency here
+	TotalCoverage *= saturate(GetMaterialMaskInputRaw(PixelMaterialInputs));
 	#endif
 
+
 	Payload.SubstrateData = InitialiseRWSubstrateMaterialContainer();
 	PackSubstrateOut(
 		Payload.SubstrateData,
@@ -333,7 +367,7 @@ void GetMaterialPayload(
 	Payload.Opacity = TotalCoverage;
 	Payload.TransmittancePreCoverage = TransmittancePreCoverage;
 
-	#if MATERIALBLENDING_MASKED
+	#if MATERIALBLENDING_MASKED && !MATERIAL_DITHER_OPACITY_MASK
 	// regular masked mode - binary decision
 	if (GetMaterialMask(PixelMaterialInputs) < 0)
 	{
@@ -351,16 +385,7 @@ void GetMaterialPayload(
 	half  Metallic = GetMaterialMetallic(PixelMaterialInputs);
 	half  Specular = GetMaterialSpecular(PixelMaterialInputs);
 	half Roughness = GetMaterialRoughness(PixelMaterialInputs);
-#if MATERIALBLENDING_TRANSLUCENT && REFRACTION_USE_INDEX_OF_REFRACTION
-	#if REFRACTION_ROOT_NODE_OVERRIDES_DEFAULT
-	half IOR = GetMaterialRefractionIOR(GetMaterialRefraction(PixelMaterialInputs));
-	#else
-	half3 F0 = ComputeF0(Specular, BaseColor, Metallic);
-	half IOR = DielectricF0RGBToIor(F0);
-	#endif
-#else
-	half IOR = 0;
-#endif
+	half IOR = GetRefractionIor(PixelMaterialInputs);
 
 	/**
 	 * Set material attributes for full materials
@@ -372,20 +397,17 @@ void GetMaterialPayload(
 	Payload.Specular = Specular;
 	Payload.Roughness = Roughness;
 	Payload.Metallic = Metallic;
-#if MATERIALBLENDING_MODULATE
+#if MATERIALBLENDING_MASKED && MATERIAL_DITHER_OPACITY_MASK
+	// dithered opacity doesn't make sense for pure raytracing, so replace it with the equivalent translucent behavior
+	Payload.Opacity = saturate(GetMaterialMaskInputRaw(PixelMaterialInputs));
+
+#elif MATERIALBLENDING_MODULATE
 	Payload.Opacity = 0.0;
 #else
+	// All other blending modes
 	Payload.Opacity = GetMaterialOpacity(PixelMaterialInputs);
 #endif
 
-#if MATERIALBLENDING_MASKED
-	// regular masked mode - binary decision
-	if (GetMaterialMask(PixelMaterialInputs) < 0)
-	{
-		Payload.Opacity = 0;
-	}
-#endif // MATERIALBLENDING_MASKED
-
 	Payload.Ior = IOR;
 	Payload.CustomData = float4(GetMaterialCustomData0(MaterialParameters), GetMaterialCustomData1(MaterialParameters), 0, 0);
 	
@@ -606,6 +628,9 @@ RAY_TRACING_ENTRY_CLOSEST_HIT(MaterialCHS,
 		uint MaterialBlendingMode = RAY_TRACING_BLEND_MODE_ADDITIVE;
 	#elif MATERIALBLENDING_MODULATE
 		uint MaterialBlendingMode = RAY_TRACING_BLEND_MODE_MODULATE;
+	#elif MATERIALBLENDING_MASKED && MATERIAL_DITHER_OPACITY_MASK
+		// dithered opacity doesn't make sense for pure raytracing, so replace it with the equivalent translucent behavior
+		uint MaterialBlendingMode = RAY_TRACING_BLEND_MODE_TRANSLUCENT;
 	#else
 		uint MaterialBlendingMode = RAY_TRACING_BLEND_MODE_OPAQUE;
 	#endif
@@ -699,7 +724,7 @@ RAY_TRACING_ENTRY_ANY_HIT(MaterialAHS,
 	// All other blending modes _might_ need material evaluation - so we need
 	// to generate some code
 
-#if MATERIALBLENDING_MASKED
+#if MATERIALBLENDING_MASKED && !MATERIAL_DITHER_OPACITY_MASK
 	// NOTE: Masked mode execute always - regardless of ray flags
 	// This ensures that both material rays and shadow rays "see" the same thing
 #else
@@ -707,14 +732,23 @@ RAY_TRACING_ENTRY_ANY_HIT(MaterialAHS,
 	if (PackedPayload.IsIgnoreTranslucentMaterials())
 	{
 		// special mode used by regular RT shadows - translucent materials are skipped, everything else is treated as opaque
-#if MATERIALBLENDING_TRANSLUCENT
 		IgnoreHit();
-#endif
+		return;
+	}
+	else if (!PackedPayload.IsShadowRay())
+	{
+		// non-opaque blending mode, but we aren't a shadow ray -- continue to CHS
 		return;
 	}
 	else
 	{
-		// non-opaque blending mode, treat as opaque for now
+		// this is a shadow ray and we aren't ignoring translucent materials, so go ahead and run the material logic
+	}
+
+	if (PackedPayload.HitT == RayTCurrent())
+	{
+		// we have processed this intersection already, ignore it
+		IgnoreHit();
 		return;
 	}
 #endif  // !MATERIALBLENDING_MASKED
@@ -759,16 +793,145 @@ RAY_TRACING_ENTRY_ANY_HIT(MaterialAHS,
 	}
 
 	// Now extract the relevant info from PixelMaterialInputs according to the blending mode
-
-#if MATERIALBLENDING_MASKED
+#if MATERIALBLENDING_MASKED && !MATERIAL_DITHER_OPACITY_MASK
 	// Regardless of payload flags -- we always apply this 
 	if (GetMaterialMask(PixelMaterialInputs) < 0)
 	{
 		IgnoreHit();
 	}
-#endif // !MATERIALBLENDING_MASKED
+#else // MATERIALBLENDING_MASKED && !MATERIAL_DITHER_OPACITY_MASK
+
+#if MATERIALBLENDING_MASKED && MATERIAL_DITHER_OPACITY_MASK
+	const float Transparency = 1 - saturate(GetMaterialMaskInputRaw(PixelMaterialInputs));
+#elif SUBSTRATE_ENABLED
+		#if MATERIAL_IS_SUBSTRATE
+		float3 EmissiveLuminance = 0.0f;
+		float TotalCoverage = 1.f;
+		float3 TransmittancePreCoverage = 0.0f;
+		bool bSubstrateSubsurfaceEnable = false;
+		const float3 WorldBentNormal0 = GetWorldBentNormalZero(MaterialParameters);
+		const float3 CamVector = MaterialParameters.CameraVector;
+
+		// Initialise a Substrate header with normal in registers
+		FSubstrateData SubstrateData = PixelMaterialInputs.GetFrontSubstrateData();
+		FSubstratePixelHeader SubstratePixelHeader = MaterialParameters.GetFrontSubstrateHeader();
+
+		#if SUBSTRATE_OPTIMIZED_UNLIT
+		// Unlit BSDF goes through the SubstrateTree to support weighting operations. Technically, layering and mixing could also be supported.
+		float3 UnlitSurfaceLuminancePostCoverage = 0.0f;
+		float UnlitSurfaceCoverage = 0.0f;
+		float3 UnlitSurfaceTransmittancePreCoverage = 0.0f;
+		float3 UnlitSurfaceNormal = 0.0f;
+		SubstratePixelHeader.SubstrateUpdateTreeUnlit(
+			uint2(MaterialParameters.SvPosition.xy),
+			CamVector,
+			SubstrateData,
+			UnlitSurfaceLuminancePostCoverage,
+			UnlitSurfaceCoverage,
+			UnlitSurfaceTransmittancePreCoverage,
+			UnlitSurfaceNormal);
+		EmissiveLuminance = UnlitSurfaceLuminancePostCoverage;
+
+		TotalCoverage = UnlitSurfaceCoverage;
+		TransmittancePreCoverage = UnlitSurfaceTransmittancePreCoverage;
+		const float3 Transparency = lerp(1, TransmittancePreCoverage, TotalCoverage);
+
+		#else // SUBSTRATE_OPTIMIZED_UNLIT
+		SubstratePixelHeader.IrradianceAO.MaterialAO = GetMaterialAmbientOcclusion(PixelMaterialInputs);
+		SubstratePixelHeader.SetCastContactShadow(GetPrimitiveData(MaterialParameters).Flags & PRIMITIVE_SCENE_DATA_FLAG_HAS_CAST_CONTACT_SHADOW);
+		SubstratePixelHeader.SetDynamicIndirectShadowCasterRepresentation(GetPrimitiveData(MaterialParameters).Flags & PRIMITIVE_SCENE_DATA_FLAG_HAS_CAPSULE_REPRESENTATION);
+
+		FSubstrateSubsurfaceData SSSData = (FSubstrateSubsurfaceData)0;
+		FSubstrateTopLayerData TopLayerData = (FSubstrateTopLayerData)0;
+		FSubstrateOpaqueRoughRefractionData OpaqueRoughRefractionData = (FSubstrateOpaqueRoughRefractionData)0;
+		FSubstrateIntegrationSettings Settings = InitSubstrateIntegrationSettings();
+
+		// Generate the Substrate material data to write out
+		FSubstrateAddressing SubstrateAddressing = GetSubstratePixelDataByteOffset(0, 0, 0);
+
+		// Compute TotalCoverage for translucent material
+		if (SubstratePixelHeader.SubstrateTree.BSDFCount > 0)
+		{
+			SubstratePixelHeader.SubstrateUpdateTree(
+				SubstrateData,
+				CamVector,
+				Settings,
+				TotalCoverage,
+				TransmittancePreCoverage);
+
+		}
+		const float3 Transparency = lerp(1, TransmittancePreCoverage, TotalCoverage);
+
+		#endif // SUBSTRATE_OPTIMIZED_UNLIT
+		#else // MATERIAL_IS_SUBSTRATE
+		// SUBSTRATE_TODO: What materials are in this case?
+		const float Transparency = 0;
+		#endif // MATERIAL_IS_SUBSTRATE
+#else // !SUBSTRATE_ENABLED 
+
+#if MATERIALBLENDING_MODULATE
+	const float3 Transparency = GetMaterialEmissive(PixelMaterialInputs);
+#elif MATERIALBLENDING_ALPHACOMPOSITE
+	const float Opacity = GetMaterialOpacity(PixelMaterialInputs);
+	const float Transparency = 1 - Opacity;
+#elif MATERIALBLENDING_TRANSLUCENT
+	const float Opacity = GetMaterialOpacity(PixelMaterialInputs);
+#if MATERIAL_SHADINGMODEL_THIN_TRANSLUCENT
+	float3 Transparency = 1 - Opacity;
+	if (Opacity < 1.0)
+	{
+		float3 Transmission = GetThinTranslucentMaterialOutput0(MaterialParameters);
+		float3 V = WorldRayDirection();
+		float3 N = normalize(MaterialParameters.WorldNormal);
+		float VoN = abs(dot(V, N));
+		// compute transmission through the slab (fresnel + absorption)
+		float3 BaseColor = GetMaterialBaseColor(PixelMaterialInputs);
+		float Metallic = GetMaterialMetallic(PixelMaterialInputs);
+		float Specular = GetMaterialSpecular(PixelMaterialInputs);
+		float F0 = F0RGBToF0(ComputeF0(Specular, BaseColor, Metallic));
+		float Ior = GetRefractionIor(PixelMaterialInputs);
+		Transparency *= ComputeThinSlabWeights(Transmission, VoN, Ior, F0).Transmitted;
+	}
+	// Fade out the Transparency as a function of the SurfaceCoverage.
+	const float SurfaceCoverage = GetThinTranslucentMaterialOutput1(MaterialParameters);
+	Transparency = lerp(1.0.xxx, Transparency, SurfaceCoverage);
+#else // !MATERIAL_SHADINGMODEL_THIN_TRANSLUCENT
+	float Transparency = 1 - Opacity;
+	uint ShadingModelID = GetMaterialShadingModel(PixelMaterialInputs);
+#if MATERIAL_SHADINGMODEL_DEFAULT_LIT && REFRACTION_USE_INDEX_OF_REFRACTION
+	if (ShadingModelID == SHADINGMODELID_DEFAULT_LIT)
+	{
+		float Ior = GetRefractionIor(PixelMaterialInputs);
+		if (Transparency > 0 && Ior > 0.0)
+		{
+			// Material is doing refraction, apply a fake caustic term to represent this
+			float3 BaseColor = GetMaterialBaseColor(PixelMaterialInputs);
+			float  Metallic = GetMaterialMetallic(PixelMaterialInputs);
+			float  Specular = GetMaterialSpecular(PixelMaterialInputs);
+			float  F0 = F0RGBToF0(ComputeF0(Specular, BaseColor, Metallic));
+			float3 N = normalize(MaterialParameters.WorldNormal);
+			float NoV = abs(dot(WorldRayDirection(), N));
+			float Fr = FresnelReflectance(NoV, Ior, F0);
+			Transparency *= Pow2(1 - Fr);
+		}
+	}
+#endif  // MATERIAL_SHADINGMODEL_DEFAULT_LIT
+#endif // !MATERIAL_SHADINGMODEL_THIN_TRANSLUCENT
+#endif // MATERIALBLENDING_TRANSLUCENT
+#endif // SUBSTRATE_ENABLED
+
+	float3 ShadowVis = PackedPayload.GetShadowVisibility();
+	ShadowVis *= Transparency;
+	PackedPayload.SetShadowVisibility(ShadowVis);
+	PackedPayload.HitT = RayTCurrent();
+	// Keep going if we haven't reached full opacity
+	if (any(ShadowVis > 0.0))
+	{
+		IgnoreHit();
+	}
+#endif // !MATERIALBLENDING_MASKED || MATERIAL_DITHER_OPACITY_MASK
 
-#endif // !MATERIALBLENDING_SOLID && !MATERIALBLENDING_ALPHAHOLDOUT && !MATERIALBLENDING_ADDITIVE
+#endif // MATERIALBLENDING_SOLID
 
 #endif // USE_MATERIAL_ANY_HIT_SHADER
 }
diff --git a/Engine/Shaders/Private/RayTracing/RayTracingOcclusionRGS.usf b/Engine/Shaders/Private/RayTracing/RayTracingOcclusionRGS.usf
--- a/Engine/Shaders/Private/RayTracing/RayTracingOcclusionRGS.usf
+++ b/Engine/Shaders/Private/RayTracing/RayTracingOcclusionRGS.usf
@@ -24,7 +24,7 @@
 #include "RayTracingCapsuleLight.ush"
 #include "/Engine/Shared/LightDefinitions.h"
 
-#define USE_TRANSLUCENT_SHADOW  (!COMPUTESHADER && PLATFORM_SUPPORTS_RAY_TRACING_HIGH_END_EFFECTS) // support translucent shadow path in high end ray tracing without inline ray tracing
+#define USE_TRANSLUCENT_SHADOW  !COMPUTESHADER // support translucent shadow path if any-hit shading is available
 
 #define ShadowMaskType_Opaque 0
 #define ShadowMaskType_Hair 1
@@ -49,7 +49,7 @@ uint TransmissionMeanFreePathType;
 uint RejectionSamplingTrials;
 uint bAcceptFirstHit;
 uint bTwoSidedGeometry;
-uint TranslucentShadow;
+uint bTranslucentShadow;
 uint MaxTranslucencyHitCount;
 
 #define TransmissionSampling_ConstantTrackingInfiniteDomain 0
@@ -440,13 +440,9 @@ FOcclusionResult ComputeOcclusion(
 		}
 
 		// Attenuation check
-		if (LightParameters.InvRadius > 0.0)
+		if (Ray.TMax * LightParameters.InvRadius > 1.0)
 		{
-			const float MaxAttenuationDistance = 1.0 / LightParameters.InvRadius;
-			if (Ray.TMax > MaxAttenuationDistance)
-			{
-				continue;
-			}
+			continue;
 		}
 
 		uint RayFlags = 0;
@@ -468,7 +464,6 @@ FOcclusionResult ComputeOcclusion(
 		uint RayFlagsForOpaque = bAcceptFirstHit != 0 ? RAY_FLAG_ACCEPT_FIRST_HIT_AND_END_SEARCH : 0;
 
 		FMinimalPayload MinimalPayload;
-		const bool NeedToTraceTranslucentShadow = TranslucentShadow == 1 && MaxTranslucencyHitCount != 0;
 
 #if AVOID_SELF_INTERSECTION_TRACE
 		{
@@ -564,58 +559,20 @@ FOcclusionResult ComputeOcclusion(
 #if USE_TRANSLUCENT_SHADOW
 			bool bHasClosestEffectiveTranslucentHit = false;
 
-			if (NeedToTraceTranslucentShadow)
+			if (bTranslucentShadow)
 			{
 				// When there is no intersection, test for translucent material hit if translucent shadow
 				// is supported. If there is a hit, record the accumulated opacity.
-				FTranslucentMinimalPayload TranslucentMinimalPayload = (FTranslucentMinimalPayload)0;
-				float Visibility = 1.0f;
-				float HitT = 0;
-				int TranslucentRayTriangleHitCount = 0;
-				FRayDesc TranslucentRay = Ray;
-				
-				while (MaxTranslucencyHitCount != 0 &&
-					TranslucentRay.TMin < TranslucentRay.TMax &&
-					TranslucentRayTriangleHitCount != MaxTranslucencyHitCount &&
-					Visibility > 0.0f)
-				{
-					TranslucentMinimalPayload = TraceTranslucentVisibilityRay(
-						TLAS,
-						RayFlags, // Accept first hit will cause artifact.
-						RaytracingMask,
-						TranslucentRay);
-
-					if (!TranslucentMinimalPayload.IsHit())
-					{
-						break;
-					}
-
-					TranslucentRay.TMin = asfloat(asuint(TranslucentMinimalPayload.HitT) + 1);
-					TranslucentRayTriangleHitCount++;
-
-					// Record first hit with opacity > 0
-					if (!bHasClosestEffectiveTranslucentHit && 
-						TranslucentMinimalPayload.Opacity > 0)
-					{
-						HitT = TranslucentMinimalPayload.HitT;
-						bHasClosestEffectiveTranslucentHit = true;
-					}
-
-					Visibility *= 1 - TranslucentMinimalPayload.Opacity;
-				}
-
-				if (bHasClosestEffectiveTranslucentHit)
-				{
-					Out.ClosestRayDistance =
-						(Out.ClosestRayDistance == DENOISER_INVALID_HIT_DISTANCE) ||
-						(HitT < Out.ClosestRayDistance) ? HitT : Out.ClosestRayDistance;
-					Out.SumRayDistance += HitT;
-					Out.HitCount += 1.0f;
-					Out.Visibility += Visibility;
-				}
+				FTranslucentMinimalPayload TranslucentMinimalPayload = TraceTranslucentVisibilityRay(
+					TLAS,
+					RAY_FLAG_ACCEPT_FIRST_HIT_AND_END_SEARCH | RAY_FLAG_SKIP_CLOSEST_HIT_SHADER,
+					RAY_TRACING_MASK_TRANSLUCENT_SHADOW,
+					Ray);
+				Out.ClosestRayDistance = (Out.ClosestRayDistance == DENOISER_INVALID_HIT_DISTANCE) ? DENOISER_MISS_HIT_DISTANCE : Out.ClosestRayDistance;
+				Out.TransmissionDistance += 1.0;
+				Out.Visibility += Luminance(TranslucentMinimalPayload.ShadowVisibility);
 			}
-			
-			if (!bHasClosestEffectiveTranslucentHit)
+			else
 #endif
 			{
 				Out.ClosestRayDistance = (Out.ClosestRayDistance == DENOISER_INVALID_HIT_DISTANCE) ? DENOISER_MISS_HIT_DISTANCE : Out.ClosestRayDistance;
diff --git a/Engine/Source/Runtime/Engine/Private/ShaderCompiler/ShaderCompiler.cpp b/Engine/Source/Runtime/Engine/Private/ShaderCompiler/ShaderCompiler.cpp
--- a/Engine/Source/Runtime/Engine/Private/ShaderCompiler/ShaderCompiler.cpp
+++ b/Engine/Source/Runtime/Engine/Private/ShaderCompiler/ShaderCompiler.cpp
@@ -3911,7 +3911,6 @@ void GlobalBeginCompileShader(
 	SET_SHADER_DEFINE(Input.Environment, PLATFORM_ALLOW_SCENE_DATA_COMPRESSED_TRANSFORMS, FDataDrivenShaderPlatformInfo::GetSupportSceneDataCompressedTransforms(EShaderPlatform(Target.Platform)) ? 1 : 0);
 	SET_SHADER_DEFINE(Input.Environment, PLATFORM_SUPPORTS_BUFFER_LOAD_TYPE_CONVERSION, RHISupportsBufferLoadTypeConversion(ShaderPlatform) ? 1 : 0);
 	SET_SHADER_DEFINE(Input.Environment, PLATFORM_SUPPORTS_UNIFORM_BUFFER_OBJECTS, FDataDrivenShaderPlatformInfo::GetSupportsUniformBufferObjects(EShaderPlatform(Target.Platform)) ? 1 : 0);
-	SET_SHADER_DEFINE(Input.Environment, PLATFORM_SUPPORTS_RAY_TRACING_HIGH_END_EFFECTS, FDataDrivenShaderPlatformInfo::GetSupportsHighEndRayTracingEffects(EShaderPlatform(Target.Platform)) ? 1 : 0);
 	SET_SHADER_DEFINE(Input.Environment, COMPILER_SUPPORTS_BARYCENTRIC_INTRINSICS, FDataDrivenShaderPlatformInfo::GetSupportsBarycentricsIntrinsics(EShaderPlatform(Target.Platform)));
 	SET_SHADER_DEFINE(Input.Environment, PLATFORM_SUPPORTS_BARYCENTRICS_SEMANTIC, FDataDrivenShaderPlatformInfo::GetSupportsBarycentricsSemantic(EShaderPlatform(Target.Platform)) != ERHIFeatureSupport::Unsupported);
 
diff --git a/Engine/Source/Runtime/RHI/Private/DataDrivenShaderPlatformInfo.cpp b/Engine/Source/Runtime/RHI/Private/DataDrivenShaderPlatformInfo.cpp
--- a/Engine/Source/Runtime/RHI/Private/DataDrivenShaderPlatformInfo.cpp
+++ b/Engine/Source/Runtime/RHI/Private/DataDrivenShaderPlatformInfo.cpp
@@ -272,7 +272,6 @@ void FGenericDataDrivenShaderPlatformInfo::ParseDataDrivenShaderInfo(const FConf
 	GET_SECTION_BOOL_HELPER(bSupportsRayTracingIndirectInstanceData);
 	GET_SECTION_BOOL_HELPER(bSupportsPathTracing);
 	GET_SECTION_BOOL_HELPER(bSupportsShaderExecutionReordering);
-	GET_SECTION_BOOL_HELPER(bSupportsHighEndRayTracingEffects);
 	GET_SECTION_BOOL_HELPER(bSupportsUnrestrictedHalfFloatBuffers)
 	GET_SECTION_BOOL_HELPER(bSupportsGPUScene);
 	GET_SECTION_BOOL_HELPER(bSupportsPrimitiveShaders);
diff --git a/Engine/Source/Runtime/RHI/Public/DataDrivenShaderPlatformInfo.h b/Engine/Source/Runtime/RHI/Public/DataDrivenShaderPlatformInfo.h
--- a/Engine/Source/Runtime/RHI/Public/DataDrivenShaderPlatformInfo.h
+++ b/Engine/Source/Runtime/RHI/Public/DataDrivenShaderPlatformInfo.h
@@ -50,7 +50,6 @@ class FGenericDataDrivenShaderPlatformInfo
 	uint32 bSupportsRayTracingProceduralPrimitive : 1;
 	uint32 bSupportsRayTracingTraversalStatistics : 1;
 	uint32 bSupportsRayTracingIndirectInstanceData : 1; // Whether instance transforms can be copied from the GPU to the TLAS instances buffer
-	uint32 bSupportsHighEndRayTracingEffects : 1; // Whether fully-featured RT effects can be used on the platform (with translucent shadow, etc.)
 	uint32 bSupportsPathTracing : 1; // Whether real-time path tracer is supported on this platform (avoids compiling unnecessary shaders)
 	uint32 bSupportsShaderExecutionReordering : 1; // Does the platform support Shader Execution Reordering extensions?
 	uint32 bSupportsGPUScene : 1;
@@ -413,12 +412,6 @@ public:
 		return Infos[Platform].bSupportsShaderExecutionReordering;
 	}
 
-	static FORCEINLINE_DEBUGGABLE const bool GetSupportsHighEndRayTracingEffects(const FStaticShaderPlatform Platform)
-	{
-		check(IsValid(Platform));
-		return Infos[Platform].bSupportsRayTracing && Infos[Platform].bSupportsHighEndRayTracingEffects;
-	}
-
 	static FORCEINLINE_DEBUGGABLE const bool GetSupportsComputeFramework(const FStaticShaderPlatform Platform)
 	{
 		check(IsValid(Platform));
diff --git a/Engine/Source/Runtime/Renderer/Private/RayTracing/RayTracingInstanceMask.cpp b/Engine/Source/Runtime/Renderer/Private/RayTracing/RayTracingInstanceMask.cpp
--- a/Engine/Source/Runtime/Renderer/Private/RayTracing/RayTracingInstanceMask.cpp
+++ b/Engine/Source/Runtime/Renderer/Private/RayTracing/RayTracingInstanceMask.cpp
@@ -209,8 +209,8 @@ FRayTracingMaskAndFlags BuildRayTracingInstanceMaskAndFlags(const FRayTracingIns
 			const bool bSegmentCastsShadow = MaskInfo.bVisibleToShadow && MeshBatch.CastRayTracedShadow && Material.CastsRayTracedShadows() && BlendMode != BLEND_Additive;
 
 			Result.Mask |= BlendModeToRayTracingInstanceMask(BlendMode, Material.IsDitherMasked(), bSegmentCastsShadow, RayTracingType);
-			bAllSegmentsOpaque &= BlendMode == BLEND_Opaque;
-			bAllSegmentsTranslucent &= BlendMode == BLEND_Translucent;
+			bAllSegmentsOpaque &= BlendMode == EBlendMode::BLEND_Opaque;
+			bAllSegmentsTranslucent &= !(PrimitiveSceneProxy.IsOpaqueOrMasked() && !Material.IsDitherMasked());
 			bAnySegmentsCastShadow |= bSegmentCastsShadow;
 			bAllSegmentsCastShadow &= bSegmentCastsShadow;
 			bAnySegmentsDecal |= Material.IsDeferredDecal();
@@ -276,7 +276,7 @@ void SetupRayTracingMeshCommandMaskAndStatus(FRayTracingMeshCommand& MeshCommand
 
 	MeshCommand.bCastRayTracedShadows = MeshBatch.CastRayTracedShadow && MaterialResource.CastsRayTracedShadows() && MaterialResource.GetBlendMode() != BLEND_Additive;
 	MeshCommand.bOpaque = MaterialResource.GetBlendMode() == EBlendMode::BLEND_Opaque && !(VertexFactory->GetType()->SupportsRayTracingProceduralPrimitive() && FDataDrivenShaderPlatformInfo::GetSupportsRayTracingProceduralPrimitive(GMaxRHIShaderPlatform));
-	MeshCommand.bAlphaMasked = MaterialResource.GetBlendMode() == EBlendMode::BLEND_Masked;
+	MeshCommand.bAlphaMasked = MaterialResource.GetBlendMode() == EBlendMode::BLEND_Masked && !MaterialResource.IsDitherMasked();
 	MeshCommand.bDecal = MaterialResource.IsDeferredDecal();
 	MeshCommand.bIsSky = MaterialResource.IsSky();
 	MeshCommand.bTwoSided = MaterialResource.IsTwoSided();
diff --git a/Engine/Source/Runtime/Renderer/Private/RayTracing/RayTracingMaterialHitShaders.cpp b/Engine/Source/Runtime/Renderer/Private/RayTracing/RayTracingMaterialHitShaders.cpp
--- a/Engine/Source/Runtime/Renderer/Private/RayTracing/RayTracingMaterialHitShaders.cpp
+++ b/Engine/Source/Runtime/Renderer/Private/RayTracing/RayTracingMaterialHitShaders.cpp
@@ -125,6 +125,22 @@ public:
 	}
 };
 
+static bool NeedsAnyHitShader(EBlendMode BlendMode)
+{
+	switch (BlendMode)
+	{
+		case BLEND_Opaque: 							return false; // always hit
+		case BLEND_Masked: 							return true;  // runs shader (NOTE: dithered masking gets turned into translucent for the path tracer)
+		case BLEND_Translucent: 					return true;  // casts transparent (colored) shadows depending on the shading model setup (fake caustics or transparent shadows)
+		case BLEND_Additive: 						return false; // never hit for shadows, goes through the default shader instead, so no need to use AHS for primary rays
+		case BLEND_Modulate: 						return true;  // casts colored shadows
+		case BLEND_AlphaComposite: 					return true;
+		case BLEND_AlphaHoldout: 					return false; // treat as opaque for shadows
+		case BLEND_TranslucentColoredTransmittance: return true;  // NOTE: Substrate only
+		default: checkf(false, TEXT("Unhandled blend mode %d"), int(BlendMode)); return false;
+	}
+}
+
 template<typename LightMapPolicyType, bool UseAnyHitShader, bool UseIntersectionShader, bool UseRayConeTextureLod>
 class TMaterialCHS : public FMaterialCHS
 {
@@ -144,8 +160,13 @@ public:
 			return false;
 		}
 
+		if (Parameters.MaterialParameters.MaterialDomain != MD_Surface)
+		{
+			return false;
+		}
+
 		static FShaderPlatformCachedIniValue<int32> CVarCompileMaterialAHS(TEXT("r.RayTracing.CompileMaterialAHS"));
-		const bool bWantAnyHitShader = ((CVarCompileMaterialAHS.Get(Parameters.Platform) != 0) && (Parameters.MaterialParameters.bIsMasked || IsTranslucentOnlyBlendMode(Parameters.MaterialParameters)));
+		const bool bWantAnyHitShader = CVarCompileMaterialAHS.Get(Parameters.Platform) != 0 && NeedsAnyHitShader(Parameters.MaterialParameters.BlendMode);
 		const bool bSupportProceduralPrimitive = Parameters.VertexFactoryType->SupportsRayTracingProceduralPrimitive() && FDataDrivenShaderPlatformInfo::GetSupportsRayTracingProceduralPrimitive(Parameters.Platform);
 
 		return IsSupportedVertexFactoryType(Parameters.VertexFactoryType)
@@ -160,10 +181,9 @@ public:
 	{
 		// NOTE: Any CVars that are used in this function must be handled in ShaderMapAppendKeyString() to ensure shaders are recompiled when necessary.
 		static FShaderPlatformCachedIniValue<int32> CVarCompileMaterialCHS(TEXT("r.RayTracing.CompileMaterialCHS"));
-		static FShaderPlatformCachedIniValue<int32> CVarCompileMaterialAHS(TEXT("r.RayTracing.CompileMaterialAHS"));
 
 		OutEnvironment.SetDefine(TEXT("USE_MATERIAL_CLOSEST_HIT_SHADER"), CVarCompileMaterialCHS.Get(Parameters.Platform) ? 1 : 0);
-		OutEnvironment.SetDefine(TEXT("USE_MATERIAL_ANY_HIT_SHADER"), CVarCompileMaterialAHS.Get(Parameters.Platform) ? 1 : 0);
+		OutEnvironment.SetDefine(TEXT("USE_MATERIAL_ANY_HIT_SHADER"), UseAnyHitShader ? 1 : 0);
 		OutEnvironment.SetDefine(TEXT("USE_MATERIAL_INTERSECTION_SHADER"), UseIntersectionShader ? 1 : 0);
 		OutEnvironment.SetDefine(TEXT("USE_RAYTRACED_TEXTURE_RAYCONE_LOD"), UseRayConeTextureLod ? 1 : 0);
 		OutEnvironment.SetDefine(TEXT("SCENE_TEXTURES_DISABLED"), 1);
@@ -330,7 +350,7 @@ static bool GetMaterialHitShader(const FMaterial& RESTRICT MaterialResource, con
 	FMaterialShaderTypes ShaderTypes;
 	const FVertexFactoryType* VFType = VertexFactory->GetType();
 	const bool bUseIntersectionShader = VFType->HasFlags(EVertexFactoryFlags::SupportsRayTracingProceduralPrimitive) && FDataDrivenShaderPlatformInfo::GetSupportsRayTracingProceduralPrimitive(GMaxRHIShaderPlatform);
-	const bool UseAnyHitShader = (MaterialResource.IsMasked() || IsTranslucentOnlyBlendMode(MaterialResource)) && GCompileRayTracingMaterialAHS;
+	const bool UseAnyHitShader = (MaterialResource.IsMasked() || NeedsAnyHitShader(MaterialResource.GetBlendMode())) && GCompileRayTracingMaterialAHS;
 
 	GetMaterialHitShader_AnyHit_Intersection_TextureLOD<LightMapPolicyType>(ShaderTypes, UseAnyHitShader, bUseIntersectionShader, UseTextureLod);
 
diff --git a/Engine/Source/Runtime/Renderer/Private/RayTracing/RayTracingShadows.cpp b/Engine/Source/Runtime/Renderer/Private/RayTracing/RayTracingShadows.cpp
--- a/Engine/Source/Runtime/Renderer/Private/RayTracing/RayTracingShadows.cpp
+++ b/Engine/Source/Runtime/Renderer/Private/RayTracing/RayTracingShadows.cpp
@@ -107,7 +107,7 @@ static TAutoConsoleVariable<int32> CVarRayTracingShadowsLODTransitionEnd(
 static TAutoConsoleVariable<int32> CVarRayTracingShadowsAcceptFirstHit(
 	TEXT("r.RayTracing.Shadows.AcceptFirstHit"),
 	1,
-	TEXT("Whether to allow shadow rays to terminate early, on first intersected primitive. This may result in worse denoising quality in some cases. (default = 0)"),
+	TEXT("Whether to allow shadow rays to terminate early, on first intersected primitive. This may result in worse denoising quality in some cases. (default = 1)"),
 	ECVF_RenderThreadSafe
 );
 
@@ -225,7 +225,7 @@ class FOcclusionRGS : public FGlobalShader
 		SHADER_PARAMETER(uint32, RejectionSamplingTrials)
 		SHADER_PARAMETER(uint32, bAcceptFirstHit)
 		SHADER_PARAMETER(uint32, bTwoSidedGeometry)
-		SHADER_PARAMETER(uint32, TranslucentShadow)
+		SHADER_PARAMETER(uint32, bTranslucentShadow)
 		SHADER_PARAMETER(uint32, MaxTranslucencyHitCount)
 
 		SHADER_PARAMETER_STRUCT(FLightShaderParameters, Light)
@@ -439,7 +439,7 @@ void FDeferredShadingSceneRenderer::RenderRayTracingShadows(
 		CommonPassParameters->AvoidSelfIntersectionTraceDistance = GRayTracingShadowsAvoidSelfIntersectionTraceDistance;
 		CommonPassParameters->bAcceptFirstHit = CVarRayTracingShadowsAcceptFirstHit.GetValueOnRenderThread();
 		CommonPassParameters->bTwoSidedGeometry = EnableRayTracingShadowTwoSidedGeometry() ? 1 : 0;
-		CommonPassParameters->TranslucentShadow = CVarRayTracingShadowsTranslucency.GetValueOnRenderThread() != 0;
+		CommonPassParameters->bTranslucentShadow = CVarRayTracingShadowsTranslucency.GetValueOnRenderThread() != 0;
 		CommonPassParameters->MaxTranslucencyHitCount = GetRayTracingShadowsMaxTranslucencyHitCount();
 		CommonPassParameters->TLAS = View.GetRayTracingSceneLayerViewChecked(ERayTracingSceneLayer::Base);
 		CommonPassParameters->ViewUniformBuffer = View.ViewUniformBuffer;
