Fix a bug in Light Function of Volumetric Fog, it doesn't work properly with world position as input.

diff --git a/Engine/Shaders/Private/VolumetricFogLightFunction.usf b/Engine/Shaders/Private/VolumetricFogLightFunction.usf
--- a/Engine/Shaders/Private/VolumetricFogLightFunction.usf
+++ b/Engine/Shaders/Private/VolumetricFogLightFunction.usf
@@ -21,6 +21,8 @@ void Main(
 	out float4 OutColor : SV_Target0
 	)
 {
+	ResolvedView = ResolveView();
+
 	float4 TranslatedWorldPosition = float4(0,0,0,0);
 	float3 LightVector = float3(0,0,0);
 
