Fix exponential height fog rendering issue in the right eye

diff --git a/Engine/Shaders/Private/Common.ush b/Engine/Shaders/Private/Common.ush
--- a/Engine/Shaders/Private/Common.ush
+++ b/Engine/Shaders/Private/Common.ush
@@ -1474,6 +1474,14 @@ float3 ScreenVectorFromScreenRect(float4 Position)
 	return select(IsOrthoProjection(), View.ViewForward, mul(Position,View.ScreenToTranslatedWorld).xyz);
 }
 
+float3 ScreenVectorFromScreenRect(float4 Position, ViewState InView)
+{
+	//Screen rect does not include scene depth. Perspective projection requires ScreenToTranslated world matrix
+	//Ortho Screen does not need the calculation as the screenvector is always the camera direction, so uses ViewForward
+
+	return select(IsOrthoProjection(), InView.ViewForward, mul(Position, InView.ScreenToTranslatedWorld).xyz);
+}
+
 // Returns the screen position for projection matrix calculations depending on the type of projection the view is using
 float2 GetScreenPositionForProjectionType(float2 ScreenPosition, float SceneDepth)
 {
diff --git a/Engine/Shaders/Private/MobileFog.usf b/Engine/Shaders/Private/MobileFog.usf
--- a/Engine/Shaders/Private/MobileFog.usf
+++ b/Engine/Shaders/Private/MobileFog.usf
@@ -54,7 +54,7 @@ void MobileFogVS(
 	// texture coord from vb
 	OutTexCoord = InPosition * ResolvedView.ScreenPositionScaleBias.xy + ResolvedView.ScreenPositionScaleBias.wz;
 	// deproject to world space
-	OutScreenVector = ScreenVectorFromScreenRect(float4(InPosition,1,0));
+	OutScreenVector = ScreenVectorFromScreenRect(float4(InPosition,1,0), ResolvedView);
 }
 
 void MobileFogPS(
