Fix RT hit lighting with Substrate.

LightIndex was aliased on SceneInstanceIndex which is used. For now, moved the LightIndex storage at the end of the Substrate Data.

diff --git a/Engine/Shaders/Private/RayTracing/RayTracingCommon.ush b/Engine/Shaders/Private/RayTracing/RayTracingCommon.ush
--- a/Engine/Shaders/Private/RayTracing/RayTracingCommon.ush
+++ b/Engine/Shaders/Private/RayTracing/RayTracingCommon.ush
@@ -669,12 +669,12 @@ struct FPackedMaterialClosestHitPayload : FMinimalPayload
 
 	uint GetLightIndex()
 	{
-		return SceneInstanceIndex;
+		return (SubstrateData.Data[SUBSTRATE_RT_PAYLOAD_NUM_UINTS - 1] & 0xFFFF);
 	}
 
 	void SetLightIndex(uint LightIndex)
 	{
-		SceneInstanceIndex = LightIndex;
+		SubstrateData.Data[SUBSTRATE_RT_PAYLOAD_NUM_UINTS - 1] = LightIndex & 0xFFFF;
 	}
 
 	bool IsValid()							{ return any(SubstrateData.PackedTopLayerData != 0); }
