Proper fix for RT hit lighting with Substrate.

Revert CL 39061693 and fix missing restore aliased SceneInstanceIndex value on RT payload when Substrate is enabled.

diff --git a/Engine/Shaders/Private/RayTracing/RayTracingCommon.ush b/Engine/Shaders/Private/RayTracing/RayTracingCommon.ush
--- a/Engine/Shaders/Private/RayTracing/RayTracingCommon.ush
+++ b/Engine/Shaders/Private/RayTracing/RayTracingCommon.ush
@@ -669,12 +669,12 @@ struct FPackedMaterialClosestHitPayload : FMinimalPayload
 
 	uint GetLightIndex()
 	{
-		return (SubstrateData.Data[SUBSTRATE_RT_PAYLOAD_NUM_UINTS - 1] & 0xFFFF);
+		return SceneInstanceIndex;
 	}
 
 	void SetLightIndex(uint LightIndex)
 	{
-		SubstrateData.Data[SUBSTRATE_RT_PAYLOAD_NUM_UINTS - 1] = LightIndex & 0xFFFF;
+		SceneInstanceIndex = LightIndex;
 	}
 
 	bool IsValid()							{ return any(SubstrateData.PackedTopLayerData != 0); }
diff --git a/Engine/Shaders/Private/RayTracing/RayTracingLightingCommon.ush b/Engine/Shaders/Private/RayTracing/RayTracingLightingCommon.ush
--- a/Engine/Shaders/Private/RayTracing/RayTracingLightingCommon.ush
+++ b/Engine/Shaders/Private/RayTracing/RayTracingLightingCommon.ush
@@ -327,6 +327,7 @@ float3 ComputeDirectLighting(
 	float3 OldRadiance = Payload.GetRadiance();
 #if SUBSTRATE_ENABLED
 	uint OldTransmittancePreCoverage = Payload.TransmittancePreCoverage;
+	uint OldSceneInstanceIndex = Payload.SceneInstanceIndex;
 #else
 	uint2 OldIndirectIrradiance = uint2(Payload.PackedIndirectIrradiance[0], Payload.PackedIndirectIrradiance[1]);
 #endif
@@ -453,6 +454,7 @@ float3 ComputeDirectLighting(
 	Payload.SetRadiance(OldRadiance);
 	// Restore aliased fields
 #if SUBSTRATE_ENABLED
+	Payload.SceneInstanceIndex = OldSceneInstanceIndex;
 	Payload.TransmittancePreCoverage = OldTransmittancePreCoverage;
 #else
 	Payload.PackedIndirectIrradiance[0] = OldIndirectIrradiance.x;  // Set back indirect irradiance (aliased with ShadowVisibility and LightIndex)
