Fix Hair tangent visualization.

Tangent arrows were displayed on random pixels screen. This is because now the hair node / hair coverage buffer are no longer fully clearered, but rely on tile dispatch. This CL fixes the  issue by reading the HairOnlyDepth texture to prune invalid pixels.

diff --git a/Engine/Shaders/Private/HairStrands/HairStrandsDebug.usf b/Engine/Shaders/Private/HairStrands/HairStrandsDebug.usf
--- a/Engine/Shaders/Private/HairStrands/HairStrandsDebug.usf
+++ b/Engine/Shaders/Private/HairStrands/HairStrandsDebug.usf
@@ -1848,6 +1848,8 @@ FTangentDesc AddTangentPlot(uint2 PixelCoord, float2 UV)
 	FTangentDesc Out = (FTangentDesc)0;
 
 	// Strands
+	const float HairDeviceZ = HairStrands.HairOnlyDepthTexture.Load(uint3(PixelCoord,0)).x;
+	if (HairDeviceZ > 0)
 	{
 		// Find closest
 		float3 Closest_Tangent = 0;
