Fix PT hair shading not taking into account per pixel tangent variation.

diff --git a/Engine/Shaders/Private/PathTracing/PathTracingMaterialHitShader.usf b/Engine/Shaders/Private/PathTracing/PathTracingMaterialHitShader.usf
--- a/Engine/Shaders/Private/PathTracing/PathTracingMaterialHitShader.usf
+++ b/Engine/Shaders/Private/PathTracing/PathTracingMaterialHitShader.usf
@@ -867,6 +867,9 @@ RAY_TRACING_ENTRY_CLOSEST_HIT(PathTracingMaterialCHS,
 #if HAIR_STRAND_MESH_FACTORY
 	Payload.WorldSmoothNormal = Payload.WorldNormal = Payload.WorldGeoNormal;
 	Payload.WorldTangent = MaterialParameters.TangentToWorld[2];
+	// Take into account per-pixel normal shading variation
+	// Note: For hair shading the tangent is actually stored as the normal input
+	Payload.WorldTangent = normalize(MaterialParameters.WorldNormal);
 #endif
 
 #if (TRANSLUCENCY_LIGHTING_VOLUMETRIC_PERVERTEX_NONDIRECTIONAL || TRANSLUCENCY_LIGHTING_VOLUMETRIC_PERVERTEX_DIRECTIONAL || TRANSLUCENCY_LIGHTING_VOLUMETRIC_NONDIRECTIONAL || TRANSLUCENCY_LIGHTING_VOLUMETRIC_DIRECTIONAL)
