Fix race condition in FMotionBlurVelocityFlattenCS for GPUs with native wave size < 32 such as Arm Mali

diff --git a/Engine/Shaders/Private/MotionBlur/MotionBlurVelocityFlatten.ush b/Engine/Shaders/Private/MotionBlur/MotionBlurVelocityFlatten.ush
--- a/Engine/Shaders/Private/MotionBlur/MotionBlurVelocityFlatten.ush
+++ b/Engine/Shaders/Private/MotionBlur/MotionBlurVelocityFlatten.ush
@@ -102,12 +102,32 @@ void ReduceVelocityFlattenTile(uint GroupIndex, float2 Velocity, out FVelocityRa
 	GroupMemoryBarrierWithGroupSync();
 #endif
 
-	// Safe for vector sizes 32 or larger, AMD and NV
-	// TODO Intel variable size vector
+#if FEATURE_LEVEL >= FEATURE_LEVEL_SM6 || PLATFORM_SUPPORTS_SM6_0_WAVE_OPERATIONS
+	const uint LaneCount = WaveGetLaneCount();
+#else
+	// Actual wave size is unknown, assume the worst
+	const uint LaneCount = 0u;
+#endif
+
 	VelocityFlattenStep(GroupIndex, /* GroupSize = */ 32, /* inout */ VelocityPolarRange);
+	if (LaneCount < 32)
+	{
+		GroupMemoryBarrierWithGroupSync();
+	}
+	
 	VelocityFlattenStep(GroupIndex, /* GroupSize = */ 16, /* inout */ VelocityPolarRange);
+	if (LaneCount < 16)
+	{
+		GroupMemoryBarrierWithGroupSync();
+	}
+	
 	VelocityFlattenStep(GroupIndex, /* GroupSize = */  8, /* inout */ VelocityPolarRange);
-	VelocityFlattenStep(GroupIndex, /* GroupSize = */  4, /* inout */ VelocityPolarRange);
+	if (LaneCount < 8)
+	{
+		GroupMemoryBarrierWithGroupSync();
+	}
+	
+	VelocityFlattenStep(GroupIndex, /* GroupSize = */  4, /* inout */ VelocityPolarRange);	
 	VelocityFlattenStep(GroupIndex, /* GroupSize = */  2, /* inout */ VelocityPolarRange);
 	VelocityFlattenStep(GroupIndex, /* GroupSize = */  1, /* inout */ VelocityPolarRange);
 
