Substrate - fixed add operator with parameter blending not propagating both input emissive.

#rnx
#tests PC editor dx12

diff --git a/Engine/Shaders/Private/Substrate/Substrate.ush b/Engine/Shaders/Private/Substrate/Substrate.ush
--- a/Engine/Shaders/Private/Substrate/Substrate.ush
+++ b/Engine/Shaders/Private/Substrate/Substrate.ush
@@ -2698,7 +2698,7 @@ FSubstrateData SubstrateAddParameterBlending(FSubstrateData A, FSubstrateData B,
 #endif // SUBSTRATE_GLINTS_ENABLED
 
 	// Some parameters will contribute to added luminance: in this case we add them together and saturate to not go out of the safe range
-	BSDF_SETEMISSIVE(ResultBSDF,		BSDF_GETEMISSIVE(ABSDF)								+ BSDF_GETEMISSIVE(ABSDF));
+	BSDF_SETEMISSIVE(ResultBSDF,		BSDF_GETEMISSIVE(BBSDF)								+ BSDF_GETEMISSIVE(ABSDF));
 	SLAB_DIFFUSEALBEDO(ResultBSDF) =	saturate(SLAB_DIFFUSEALBEDO(BBSDF)					+ SLAB_DIFFUSEALBEDO(ABSDF));
 	SLAB_F0(ResultBSDF) =				saturate(SLAB_F0(BBSDF)								+ SLAB_F0(ABSDF));
 	SLAB_F90(ResultBSDF) =				saturate(SLAB_F90(BBSDF)							+ SLAB_F90(ABSDF));
