[HWRT] Unify near/far culling for Nanite and non-Nanite ISM:
* Path tracer now runs GPU culling pass just to do min/max draw distance culling.
* Fix non-Nanite ISM pass to do culling in the same way as Nanite ISM.

#rnx

diff --git a/Engine/Shaders/Private/RayTracing/RayTracingInstanceBufferUtil.usf b/Engine/Shaders/Private/RayTracing/RayTracingInstanceBufferUtil.usf
--- a/Engine/Shaders/Private/RayTracing/RayTracingInstanceBufferUtil.usf
+++ b/Engine/Shaders/Private/RayTracing/RayTracingInstanceBufferUtil.usf
@@ -56,7 +56,7 @@ float AngleThresholdRatioSq;
 float3 ViewOrigin;
 uint CullingMode;
 
-bool CullInstance(FInstanceSceneData InstanceSceneData, float4x4 LocalToTranslatedWorld4x4, bool bFarFieldInstance)
+bool CullInstance(FInstanceSceneData InstanceSceneData, FPrimitiveSceneData PrimitiveData, float4x4 LocalToTranslatedWorld4x4, bool bFarFieldInstance)
 {
 	bool bIsVisible = true;
 
@@ -84,13 +84,20 @@ bool CullInstance(FInstanceSceneData InstanceSceneData, float4x4 LocalToTranslat
 		{
 			bIsVisible = bIsVisibleDistance && bIsVisibleAngle;
 		}
-		else
+		else if (CullingMode == 2)
 		{
 			bIsVisible = bIsVisibleDistance || bIsVisibleAngle;	
-		}
+		}		
 		
-	}	
+		const bool bDrawDistance = (PrimitiveData.Flags & PRIMITIVE_SCENE_DATA_FLAG_INSTANCE_DRAW_DISTANCE_CULL) != 0u;
+		if (bIsVisible && bDrawDistance)
+		{
+			const float MinDistanceSq = PrimitiveData.InstanceDrawDistanceMinMaxSquared.x;
+			const float MaxDistanceSq = PrimitiveData.InstanceDrawDistanceMinMaxSquared.y;
 
+			bIsVisible = InstanceDrawDistSq >= MinDistanceSq && InstanceDrawDistSq <= MaxDistanceSq;
+		}
+	}	
 	return !bIsVisible;
 }
 
@@ -126,7 +133,7 @@ void RayTracingBuildInstanceBufferCS(uint GroupIndex : SV_GroupIndex, uint3 Grou
 		const bool bHasGroupId = PrimitiveData.Flags & PRIMITIVE_SCENE_DATA_FLAG_RAYTRACING_HAS_GROUPID;
 		if (!bHasGroupId)
 		{
-			bCulled = CullInstance(InstanceSceneData, LocalToTranslatedWorld4x4, bFarFieldInstance);
+			bCulled = CullInstance(InstanceSceneData, PrimitiveData, LocalToTranslatedWorld4x4, bFarFieldInstance);
 		}
 #endif // GPU_CULLING
 
diff --git a/Engine/Source/Runtime/Engine/Private/InstancedStaticMesh.cpp b/Engine/Source/Runtime/Engine/Private/InstancedStaticMesh.cpp
--- a/Engine/Source/Runtime/Engine/Private/InstancedStaticMesh.cpp
+++ b/Engine/Source/Runtime/Engine/Private/InstancedStaticMesh.cpp
@@ -1763,9 +1763,10 @@ void FInstancedStaticMeshSceneProxy::GetDynamicRayTracingInstances(FRayTracingIn
 	{
 		// whether to use angular culling instead of distance, angle is halved as it is compared against the projection of the radius rather than the diameter
 		const float CullAngle = FMath::Min(CVarRayTracingInstancesCullAngle.GetValueOnRenderThread(), 179.9f) * 0.5f;
+		const float MaxDrawDistanceScale = GetCachedScalabilityCVars().ViewDistanceScale;
 
-		const float UserNearCullDistance = UserData_AllInstances.MinDrawDistance;
-		const float UserFarCullDistance = UserData_AllInstances.EndCullDistance;
+		const float UserNearCullDistance = UserData_AllInstances.MinDrawDistance * MaxDrawDistanceScale;
+		const float UserFarCullDistance = UserData_AllInstances.EndCullDistance * MaxDrawDistanceScale;
 		if (CullAngle < 0.0f || UserFarCullDistance > 0.0f)
 		{
 			//
@@ -1786,7 +1787,7 @@ void FInstancedStaticMeshSceneProxy::GetDynamicRayTracingInstances(FRayTracingIn
 				GetDistanceToInstance(InstanceIndex, InstanceRadius, DistanceToInstanceCenter, DistanceToInstanceStart);
 
 				// Cull instance based on distance
-				if ((DistanceToInstanceStart < UserNearCullDistance || DistanceToInstanceStart >= BVHFarCullRadius) && ApplyGeneralCulling)
+				if (((DistanceToInstanceCenter < UserNearCullDistance) || (DistanceToInstanceCenter > BVHFarCullRadius)) && ApplyGeneralCulling)
 					continue;
 
 				// Special culling for small scale objects
diff --git a/Engine/Source/Runtime/Renderer/Private/RayTracing/RayTracingScene.cpp b/Engine/Source/Runtime/Renderer/Private/RayTracing/RayTracingScene.cpp
--- a/Engine/Source/Runtime/Renderer/Private/RayTracing/RayTracingScene.cpp
+++ b/Engine/Source/Runtime/Renderer/Private/RayTracing/RayTracingScene.cpp
@@ -358,7 +358,7 @@ void FRayTracingScene::Create(FRDGBuilder& GraphBuilder, const FViewInfo& View,
 						Layer.TransformUploadSRV,
 						NumNativeGPUSceneInstances,
 						NumNativeCPUInstances,
-						CullingParameters.bUseInstanceCulling ? &CullingParameters : nullptr,
+						&CullingParameters,
 						PassParams->OutputStats ? PassParams->OutputStats->GetRHI() : nullptr,
 						PassParams->DebugInstanceGPUSceneIndexBuffer ? PassParams->DebugInstanceGPUSceneIndexBuffer->GetRHI() : nullptr);
 				});
