[HWRT] Refactored RayTracingBuildInstanceBufferCS and FillRayTracingInstanceUploadBuffer to move logic processing each instance to a helper function.
- this makes upcoming changes cleaner.

diff --git a/Engine/Shaders/Private/RayTracing/RayTracingInstanceBufferUtil.usf b/Engine/Shaders/Private/RayTracing/RayTracingInstanceBufferUtil.usf
--- a/Engine/Shaders/Private/RayTracing/RayTracingInstanceBufferUtil.usf
+++ b/Engine/Shaders/Private/RayTracing/RayTracingInstanceBufferUtil.usf
@@ -108,113 +108,120 @@ bool IsTransformElementUnsafe(float Element)
 	return isnan(Element) || isinf(Element) || Element > 1e20f;
 }
 
-[numthreads(THREADGROUP_SIZE, 1, 1)]
-void RayTracingBuildInstanceBufferCS(uint GroupIndex : SV_GroupIndex, uint3 GroupId : SV_GroupID)
+void ProcessInstance(FRayTracingInstanceDescriptor InputDesc)
 {
 	const FDFVector3 PreViewTranslation = MakeDFVector3(PreViewTranslationHigh, PreViewTranslationLow);
 
-	const uint Index = GetUnWrappedDispatchThreadId(GroupId, GroupIndex, THREADGROUP_SIZE);
-
 	bool bCulled = false;
 
-	if (Index < NumInstances)
-	{
-		FRayTracingInstanceDescriptor InputDesc = InstanceDescriptors[BaseInstanceDescriptorIndex + Index];
-
 #if USE_GPUSCENE
-		FInstanceSceneData InstanceSceneData = GetInstanceSceneData(InputDesc.GPUSceneInstanceOrTransformIndex);
-		FPrimitiveSceneData PrimitiveData = GetPrimitiveData(InstanceSceneData.PrimitiveId);
+	FInstanceSceneData InstanceSceneData = GetInstanceSceneData(InputDesc.GPUSceneInstanceOrTransformIndex);
+	FPrimitiveSceneData PrimitiveData = GetPrimitiveData(InstanceSceneData.PrimitiveId);
 
-		float4x4 LocalToTranslatedWorld4x4 = DFFastToTranslatedWorld(InstanceSceneData.LocalToWorld, PreViewTranslation);
+	float4x4 LocalToTranslatedWorld4x4 = DFFastToTranslatedWorld(InstanceSceneData.LocalToWorld, PreViewTranslation);
 
-		const bool bFarFieldInstance = PrimitiveData.Flags & PRIMITIVE_SCENE_DATA_FLAG_RAYTRACING_FAR_FIELD;
+	const bool bFarFieldInstance = PrimitiveData.Flags & PRIMITIVE_SCENE_DATA_FLAG_RAYTRACING_FAR_FIELD;
 
 #if GPU_CULLING
-		const bool bHasGroupId = PrimitiveData.Flags & PRIMITIVE_SCENE_DATA_FLAG_RAYTRACING_HAS_GROUPID;
-		if (!bHasGroupId)
-		{
-			bCulled = CullInstance(InstanceSceneData, PrimitiveData, LocalToTranslatedWorld4x4, bFarFieldInstance);
-		}
+	const bool bHasGroupId = PrimitiveData.Flags & PRIMITIVE_SCENE_DATA_FLAG_RAYTRACING_HAS_GROUPID;
+	if (!bHasGroupId)
+	{
+		bCulled = CullInstance(InstanceSceneData, PrimitiveData, LocalToTranslatedWorld4x4, bFarFieldInstance);
+	}
 #endif // GPU_CULLING
 
-		if (bFarFieldInstance)
-		{
-			LocalToTranslatedWorld4x4[3].xyz += FarFieldReferencePos;
-		}
+	if (bFarFieldInstance)
+	{
+		LocalToTranslatedWorld4x4[3].xyz += FarFieldReferencePos;
+	}
 
-		if (InputDesc.bApplyLocalBoundsTransform)
-		{
-			float4x4 ProcT;
-			ProcT[0] = float4(InstanceSceneData.LocalBoundsExtent.x * 2, 0, 0, 0);
-			ProcT[1] = float4(0, InstanceSceneData.LocalBoundsExtent.y * 2, 0, 0);
-			ProcT[2] = float4(0, 0, InstanceSceneData.LocalBoundsExtent.z * 2, 0);
-			ProcT[3] = float4(InstanceSceneData.LocalBoundsCenter, 1);
+	if (InputDesc.bApplyLocalBoundsTransform)
+	{
+		float4x4 ProcT;
+		ProcT[0] = float4(InstanceSceneData.LocalBoundsExtent.x * 2, 0, 0, 0);
+		ProcT[1] = float4(0, InstanceSceneData.LocalBoundsExtent.y * 2, 0, 0);
+		ProcT[2] = float4(0, 0, InstanceSceneData.LocalBoundsExtent.z * 2, 0);
+		ProcT[3] = float4(InstanceSceneData.LocalBoundsCenter, 1);
 
-			LocalToTranslatedWorld4x4 = mul(ProcT, LocalToTranslatedWorld4x4);
-		}
+		LocalToTranslatedWorld4x4 = mul(ProcT, LocalToTranslatedWorld4x4);
+	}
 
-		LocalToTranslatedWorld4x4 = transpose(LocalToTranslatedWorld4x4);
+	LocalToTranslatedWorld4x4 = transpose(LocalToTranslatedWorld4x4);
 
-		float3x4 LocalToTranslatedWorld;
-		LocalToTranslatedWorld[0] = LocalToTranslatedWorld4x4[0];
-		LocalToTranslatedWorld[1] = LocalToTranslatedWorld4x4[1];
-		LocalToTranslatedWorld[2] = LocalToTranslatedWorld4x4[2];
+	float3x4 LocalToTranslatedWorld;
+	LocalToTranslatedWorld[0] = LocalToTranslatedWorld4x4[0];
+	LocalToTranslatedWorld[1] = LocalToTranslatedWorld4x4[1];
+	LocalToTranslatedWorld[2] = LocalToTranslatedWorld4x4[2];
 
 #if OUTPUT_INSTANCE_GPUSCENE_INDEX
-		if (InputDesc.AccelerationStructureIndex != 0xFFFFFFFF && !bCulled)
-		{
-			RWDebugInstanceGPUSceneIndices[InputDesc.OutputDescriptorIndex] = InputDesc.GPUSceneInstanceOrTransformIndex;
-		}
+	if (InputDesc.AccelerationStructureIndex != 0xFFFFFFFF && !bCulled)
+	{
+		RWDebugInstanceGPUSceneIndices[InputDesc.OutputDescriptorIndex] = InputDesc.GPUSceneInstanceOrTransformIndex;
+	}
 #endif // OUTPUT_INSTANCE_GPUSCENE_INDEX
 
 #else
-		float3x4 LocalToTranslatedWorld;
-		LocalToTranslatedWorld[0] = InstanceTransforms[InputDesc.GPUSceneInstanceOrTransformIndex * 3 + 0];
-		LocalToTranslatedWorld[1] = InstanceTransforms[InputDesc.GPUSceneInstanceOrTransformIndex * 3 + 1];
-		LocalToTranslatedWorld[2] = InstanceTransforms[InputDesc.GPUSceneInstanceOrTransformIndex * 3 + 2];
+	float3x4 LocalToTranslatedWorld;
+	LocalToTranslatedWorld[0] = InstanceTransforms[InputDesc.GPUSceneInstanceOrTransformIndex * 3 + 0];
+	LocalToTranslatedWorld[1] = InstanceTransforms[InputDesc.GPUSceneInstanceOrTransformIndex * 3 + 1];
+	LocalToTranslatedWorld[2] = InstanceTransforms[InputDesc.GPUSceneInstanceOrTransformIndex * 3 + 2];
 #endif // USE_GPUSCENE
 
-		uint2 BlasAddress = 0;
+	uint2 BlasAddress = 0;
 
-		if (InputDesc.AccelerationStructureIndex != 0xFFFFFFFF && !bCulled)
-		{
-			BlasAddress = AccelerationStructureAddresses.Load2(InputDesc.AccelerationStructureIndex * 8);
-		}
+	if (InputDesc.AccelerationStructureIndex != 0xFFFFFFFF && !bCulled)
+	{
+		BlasAddress = AccelerationStructureAddresses.Load2(InputDesc.AccelerationStructureIndex * 8);
+	}
 
-		// temporary workaround for crash on some platforms
-		if(IsTransformElementUnsafe(LocalToTranslatedWorld[0][0])
-			|| IsTransformElementUnsafe(LocalToTranslatedWorld[0][1])
-			|| IsTransformElementUnsafe(LocalToTranslatedWorld[0][2])
-			|| IsTransformElementUnsafe(LocalToTranslatedWorld[0][3])
-			|| IsTransformElementUnsafe(LocalToTranslatedWorld[1][0])
-			|| IsTransformElementUnsafe(LocalToTranslatedWorld[1][1])
-			|| IsTransformElementUnsafe(LocalToTranslatedWorld[1][2])
-			|| IsTransformElementUnsafe(LocalToTranslatedWorld[1][3])
-			|| IsTransformElementUnsafe(LocalToTranslatedWorld[2][0])
-			|| IsTransformElementUnsafe(LocalToTranslatedWorld[2][1])
-			|| IsTransformElementUnsafe(LocalToTranslatedWorld[2][2])
-			|| IsTransformElementUnsafe(LocalToTranslatedWorld[2][3]))
-		{
-			BlasAddress = 0;
-		}
+	// temporary workaround for crash on some platforms
+	if(IsTransformElementUnsafe(LocalToTranslatedWorld[0][0])
+		|| IsTransformElementUnsafe(LocalToTranslatedWorld[0][1])
+		|| IsTransformElementUnsafe(LocalToTranslatedWorld[0][2])
+		|| IsTransformElementUnsafe(LocalToTranslatedWorld[0][3])
+		|| IsTransformElementUnsafe(LocalToTranslatedWorld[1][0])
+		|| IsTransformElementUnsafe(LocalToTranslatedWorld[1][1])
+		|| IsTransformElementUnsafe(LocalToTranslatedWorld[1][2])
+		|| IsTransformElementUnsafe(LocalToTranslatedWorld[1][3])
+		|| IsTransformElementUnsafe(LocalToTranslatedWorld[2][0])
+		|| IsTransformElementUnsafe(LocalToTranslatedWorld[2][1])
+		|| IsTransformElementUnsafe(LocalToTranslatedWorld[2][2])
+		|| IsTransformElementUnsafe(LocalToTranslatedWorld[2][3]))
+	{
+		BlasAddress = 0;
+	}
 
-		OutPlatformInstanceDescriptors[InputDesc.OutputDescriptorIndex] = BuildPlatformRayTracingInstanceDesc(
-			InputDesc.InstanceMaskAndFlags & 0xFF,
-			InputDesc.InstanceId,
-			TranslateRayTracingInstanceFlags((InputDesc.InstanceMaskAndFlags >> 8) & 0xFF),
-			InputDesc.InstanceContributionToHitGroupIndex,
-			LocalToTranslatedWorld,
-			BlasAddress);
+	OutPlatformInstanceDescriptors[InputDesc.OutputDescriptorIndex] = BuildPlatformRayTracingInstanceDesc(
+		InputDesc.InstanceMaskAndFlags & 0xFF,
+		InputDesc.InstanceId,
+		TranslateRayTracingInstanceFlags((InputDesc.InstanceMaskAndFlags >> 8) & 0xFF),
+		InputDesc.InstanceContributionToHitGroupIndex,
+		LocalToTranslatedWorld,
+		BlasAddress);
 
 #if OUTPUT_STATS
-		if (!bCulled)
-		{
+	if (!bCulled)
+	{
 #if USE_WAVE_OPS
-			WaveInterlockedAddScalar(RWOutputStats[0], 1);
+		WaveInterlockedAddScalar(RWOutputStats[0], 1);
 #else
-			InterlockedAdd(RWOutputStats[0], 1);
+		InterlockedAdd(RWOutputStats[0], 1);
 #endif
-		}
+	}
 #endif
+}
+
+[numthreads(THREADGROUP_SIZE, 1, 1)]
+void RayTracingBuildInstanceBufferCS(uint GroupThreadIndex : SV_GroupIndex, uint3 GroupId : SV_GroupID)
+{
+	const uint GroupIndex = GetUnWrappedDispatchGroupId(GroupId);
+	const uint Index = GetUnWrappedDispatchThreadId(GroupId, GroupThreadIndex, THREADGROUP_SIZE);
+
+	if (Index >= NumInstances)
+	{
+		return;
 	}
+
+	FRayTracingInstanceDescriptor InputDesc = InstanceDescriptors[BaseInstanceDescriptorIndex + Index];
+	ProcessInstance(InputDesc);
 }
diff --git a/Engine/Source/Runtime/Renderer/Private/RayTracing/RayTracingInstanceBufferUtil.cpp b/Engine/Source/Runtime/Renderer/Private/RayTracing/RayTracingInstanceBufferUtil.cpp
--- a/Engine/Source/Runtime/Renderer/Private/RayTracing/RayTracingInstanceBufferUtil.cpp
+++ b/Engine/Source/Runtime/Renderer/Private/RayTracing/RayTracingInstanceBufferUtil.cpp
@@ -95,6 +95,64 @@ FRayTracingSceneInitializationData BuildRayTracingSceneInitializationData(TConst
 	return MoveTemp(Output);
 }
 
+void WriteInstanceDescriptor(
+	const FRayTracingGeometryInstance& SceneInstance,
+	uint32 TransformIndex,
+	uint32 AccelerationStructureIndex,
+	bool bGpuSceneInstance,
+	bool bUseUniqueUserData,
+	uint32 BaseInstanceIndex,
+	uint32 BaseTransformIndex,
+	FRayTracingInstanceDescriptor& OutInstanceDescriptor)
+{
+	FRayTracingInstanceDescriptor InstanceDesc;
+
+	if (bGpuSceneInstance)
+	{
+		if (SceneInstance.BaseInstanceSceneDataOffset != -1)
+		{
+			InstanceDesc.GPUSceneInstanceOrTransformIndex = SceneInstance.BaseInstanceSceneDataOffset + TransformIndex;
+		}
+		else
+		{
+			InstanceDesc.GPUSceneInstanceOrTransformIndex = SceneInstance.InstanceSceneDataOffsets[TransformIndex];
+		}
+	}
+	else
+	{
+		InstanceDesc.GPUSceneInstanceOrTransformIndex = BaseTransformIndex + TransformIndex;
+	}
+
+	uint32 UserData;
+
+	if (bUseUniqueUserData)
+	{
+		UserData = SceneInstance.UserData[TransformIndex];
+	}
+	else
+	{
+		UserData = SceneInstance.DefaultUserData;
+
+		if (SceneInstance.bIncrementUserDataPerInstance)
+		{
+			UserData += TransformIndex;
+		}
+	}
+
+	InstanceDesc.OutputDescriptorIndex = BaseInstanceIndex + TransformIndex;
+	InstanceDesc.AccelerationStructureIndex = AccelerationStructureIndex;
+	InstanceDesc.InstanceId = UserData;
+	InstanceDesc.InstanceMaskAndFlags = SceneInstance.Mask | ((uint32)SceneInstance.Flags << 8);
+	InstanceDesc.InstanceContributionToHitGroupIndex = SceneInstance.InstanceContributionToHitGroupIndex;
+	InstanceDesc.bApplyLocalBoundsTransform = SceneInstance.bApplyLocalBoundsTransform;
+
+	checkf(InstanceDesc.InstanceId <= 0xFFFFFF, TEXT("InstanceId must fit in 24 bits."));
+	checkf(InstanceDesc.InstanceContributionToHitGroupIndex <= 0xFFFFFF, TEXT("InstanceContributionToHitGroupIndex must fit in 24 bits."));
+
+	// copy at the end to avoid reading from OutInstanceDescriptor in the checks above
+	OutInstanceDescriptor = InstanceDesc;
+}
+
 void FillRayTracingInstanceUploadBuffer(
 	FVector PreViewTranslation,
 	TConstArrayView<FRayTracingGeometryInstance> Instances,
@@ -144,71 +202,40 @@ void FillRayTracingInstanceUploadBuffer(
 			const uint32 BaseInstanceIndex = BaseInstancePrefixSum[SceneInstanceIndex];
 			const uint32 BaseTransformIndex = bCpuInstance ? BaseUploadBufferOffsets[SceneInstanceIndex] : 0;
 
-			uint32 BaseDescriptorIndex = BaseUploadBufferOffsets[SceneInstanceIndex];
-
-			// Upload buffer is split into 2 sections [GPUSceneInstances][CPUInstances]
-			if (!bGpuSceneInstance)
 			{
-				BaseDescriptorIndex += NumNativeGPUSceneInstances;
-			}
+				uint32 BaseDescriptorIndex = BaseUploadBufferOffsets[SceneInstanceIndex];
 
-			for (uint32 TransformIndex = 0; TransformIndex < NumTransforms; ++TransformIndex)
-			{
-				FRayTracingInstanceDescriptor InstanceDesc;
-
-				if (bGpuSceneInstance)
-				{
-					if (SceneInstance.BaseInstanceSceneDataOffset != -1)
-					{
-						InstanceDesc.GPUSceneInstanceOrTransformIndex = SceneInstance.BaseInstanceSceneDataOffset + TransformIndex;
-					}
-					else
-					{
-						InstanceDesc.GPUSceneInstanceOrTransformIndex = SceneInstance.InstanceSceneDataOffsets[TransformIndex];
-					}
-				}
-				else
+				// Upload buffer is split into 2 sections [GPUSceneInstances][CPUInstances]
+				if (!bGpuSceneInstance)
 				{
-					InstanceDesc.GPUSceneInstanceOrTransformIndex = BaseTransformIndex + TransformIndex;
+					BaseDescriptorIndex += NumNativeGPUSceneInstances;
 				}
 
-				uint32 UserData;
-
-				if (bUseUniqueUserData)
+				for (uint32 TransformIndex = 0; TransformIndex < NumTransforms; ++TransformIndex)
 				{
-					UserData = SceneInstance.UserData[TransformIndex];
-				}
-				else
-				{
-					UserData = SceneInstance.DefaultUserData;
-
-					if (SceneInstance.bIncrementUserDataPerInstance)
-					{
-						UserData += TransformIndex;
-					}
+					WriteInstanceDescriptor(
+						SceneInstance,
+						TransformIndex,
+						AccelerationStructureIndex,
+						bGpuSceneInstance,
+						bUseUniqueUserData,
+						BaseInstanceIndex,
+						BaseTransformIndex,
+						OutInstanceUploadData[BaseDescriptorIndex + TransformIndex]);
 				}
+			}
 
-				InstanceDesc.OutputDescriptorIndex = BaseInstanceIndex + TransformIndex;
-				InstanceDesc.AccelerationStructureIndex = AccelerationStructureIndex;
-				InstanceDesc.InstanceId = UserData;
-				InstanceDesc.InstanceMaskAndFlags = SceneInstance.Mask | ((uint32)SceneInstance.Flags << 8);
-				InstanceDesc.InstanceContributionToHitGroupIndex = SceneInstance.InstanceContributionToHitGroupIndex;
-				InstanceDesc.bApplyLocalBoundsTransform = SceneInstance.bApplyLocalBoundsTransform;
-
-				checkf(InstanceDesc.InstanceId <= 0xFFFFFF, TEXT("InstanceId must fit in 24 bits."));
-				checkf(InstanceDesc.InstanceContributionToHitGroupIndex <= 0xFFFFFF, TEXT("InstanceContributionToHitGroupIndex must fit in 24 bits."));
-
-				if (bCpuInstance)
+			if (bCpuInstance)
+			{
+				for (uint32 TransformIndex = 0; TransformIndex < NumTransforms; ++TransformIndex)
 				{
-					const uint32 TransformDataOffset = InstanceDesc.GPUSceneInstanceOrTransformIndex * 3;
+					const uint32 TransformDataOffset = (BaseTransformIndex + TransformIndex) * 3;
 					FMatrix LocalToTranslatedWorld = SceneInstance.Transforms[TransformIndex].ConcatTranslation(PreViewTranslation);
 					const FMatrix44f LocalToTranslatedWorldF = FMatrix44f(LocalToTranslatedWorld.GetTransposed());
 					OutTransformData[TransformDataOffset + 0] = *(FVector4f*)&LocalToTranslatedWorldF.M[0];
 					OutTransformData[TransformDataOffset + 1] = *(FVector4f*)&LocalToTranslatedWorldF.M[1];
 					OutTransformData[TransformDataOffset + 2] = *(FVector4f*)&LocalToTranslatedWorldF.M[2];
 				}
-
-				OutInstanceUploadData[BaseDescriptorIndex + TransformIndex] = InstanceDesc;
 			}
 		});
 }
