Substrate - is now compatible with material export. It seems that even wihtout substrate, material baking is not working perfectly (unlesss you are simply baking Disney material). Substrate is a match to legacy at least, and further improvement can come later.
Material preview and lightmass have been tested and are still functional.

#rnx
#tests PC dx12 editor

diff --git a/Engine/Shaders/Private/BasePassPixelShader.usf b/Engine/Shaders/Private/BasePassPixelShader.usf
--- a/Engine/Shaders/Private/BasePassPixelShader.usf
+++ b/Engine/Shaders/Private/BasePassPixelShader.usf
@@ -1643,64 +1643,81 @@ void FPixelShaderInOut_MainPS(
 	// Export non-unlit materials (opaque or translucent)
 #if SUBSTRATE_MATERIAL_EXPORT_EXECUTED
 
-	float  Coverage = 0;
-	float3 TransmittancePreCoverage = 0;
-	float3 BaseColorPostCoverage = 0;
-	float3 WorldNormal = 0;
-	float3 EmissiveLuminance = 0;
 
-	FSubstrateIntegrationSettings Settings = InitSubstrateIntegrationSettings(true /*bForceFullyRough*/, false /*SubstrateStruct.bRoughDiffuse*/, SubstrateStruct.PeelLayersAboveDepth, SubstrateStruct.bRoughnessTracking);
+	FSubstrateIntegrationSettings Settings = InitSubstrateIntegrationSettings(false /*bForceFullyRough*/, false /*SubstrateStruct.bRoughDiffuse*/, SubstrateStruct.PeelLayersAboveDepth, SubstrateStruct.bRoughnessTracking);
 
 	const float3 SurfaceWorldNormal = MaterialParameters.TangentToWorld[2].xyz;
-	SubstrateMaterialExportOut(
+	FExportResult Export = SubstrateMaterialExportOut(
 		Settings,
 		SubstratePixelHeader,
 		SubstrateData,
-		SurfaceWorldNormal,
-		Coverage,
-		TransmittancePreCoverage,
-		BaseColorPostCoverage,
-		WorldNormal,
-		EmissiveLuminance);
+		SurfaceWorldNormal);
 
 	#if SUBSTRATE_MATERIAL_EXPORT_FROM_OPAQUE
-		#if	  SUBSTRATE_MATERIAL_EXPORT_TYPE == SUBSTRATE_MATERIAL_EXPORT_DIFFUSE_COLOR
-			SubstrateMaterialExport = float4(BaseColorPostCoverage, 1.0);
+		#if	  SUBSTRATE_MATERIAL_EXPORT_TYPE == SUBSTRATE_MATERIAL_EXPORT_BASE_COLOR_POST_COVERAGE
+			SubstrateMaterialExport = float4(Export.BaseColorPostCoverage, 1.0);
+		#elif SUBSTRATE_MATERIAL_EXPORT_TYPE == SUBSTRATE_MATERIAL_EXPORT_BASE_COLOR
+			SubstrateMaterialExport = float4(Export.BaseColor, 1.0);
 		#elif SUBSTRATE_MATERIAL_EXPORT_TYPE == SUBSTRATE_MATERIAL_EXPORT_NORMAL
-			SubstrateMaterialExport = float4(WorldNormal * 0.5 + 0.5, 1.0);
+			SubstrateMaterialExport = float4(Export.WorldNormal * 0.5 + 0.5, 1.0);
 		#elif SUBSTRATE_MATERIAL_EXPORT_TYPE == SUBSTRATE_MATERIAL_EXPORT_EMISSIVE
-			SubstrateMaterialExport = float4(EmissiveLuminance, 1.0);
+			SubstrateMaterialExport = float4(Export.EmissiveLuminance, 1.0);
+		#elif SUBSTRATE_MATERIAL_EXPORT_TYPE == SUBSTRATE_MATERIAL_EXPORT_SPECULAR
+			SubstrateMaterialExport = float4(Export.Specular.xxx, 1.0f);
+		#elif SUBSTRATE_MATERIAL_EXPORT_TYPE == SUBSTRATE_MATERIAL_EXPORT_ROUGHNESS
+			SubstrateMaterialExport = float4(Export.Roughness.xxx, 1.0f);
+		#elif SUBSTRATE_MATERIAL_EXPORT_TYPE == SUBSTRATE_MATERIAL_EXPORT_ANISOTROPY
+			SubstrateMaterialExport = float4(Export.Anisotropy.xxx, 1.0f);
+		#elif SUBSTRATE_MATERIAL_EXPORT_TYPE == SUBSTRATE_MATERIAL_EXPORT_METALLIC
+			SubstrateMaterialExport = float4(Export.Metallic.xxx, 1.0f);
 		#elif SUBSTRATE_MATERIAL_EXPORT_TYPE == SUBSTRATE_MATERIAL_EXPORT_TRANSMITTANCE
-			SubstrateMaterialExport = float4(PixelMaterialInputs.OpacityMask.xxx, 1.0f); // when not masked, OpacityMask default to 1.
+			SubstrateMaterialExport = float4(Export.TransmittancePreCoverage * Export.Coverage + (1.0f - Export.Coverage), 1.0f);
+		#elif SUBSTRATE_MATERIAL_EXPORT_TYPE == SUBSTRATE_MATERIAL_EXPORT_OPACITY
+			SubstrateMaterialExport = float4(Export.Coverage.xxx, 1.0f);
+			#if SUBSTRATE_USE_PREMULTALPHA_OVERRIDE
+			SubstrateMaterialExport = float4(GetMaterialOpacity(PixelMaterialInputs).xxx, 1.0f);
+			#endif
+		#elif SUBSTRATE_MATERIAL_EXPORT_TYPE == SUBSTRATE_MATERIAL_EXPORT_OPACITYMASK
+			SubstrateMaterialExport = float4(PixelMaterialInputs.OpacityMask.xxx, 1.0f);
+		#elif SUBSTRATE_MATERIAL_EXPORT_TYPE == SUBSTRATE_MATERIAL_EXPORT_CUSTOMDATA0
+			SubstrateMaterialExport = Export.CustomData0;
+		#elif SUBSTRATE_MATERIAL_EXPORT_TYPE == SUBSTRATE_MATERIAL_EXPORT_CUSTOMDATA1
+			SubstrateMaterialExport = Export.CustomData1;
+		#elif SUBSTRATE_MATERIAL_EXPORT_TYPE == SUBSTRATE_MATERIAL_EXPORT_SUBSURFACECOLOR
+			SubstrateMaterialExport = float4(Export.SubsurfaceColor, 1.0f);
+		#elif SUBSTRATE_MATERIAL_EXPORT_TYPE == SUBSTRATE_MATERIAL_EXPORT_TANGENT
+			SubstrateMaterialExport = float4(Export.Tangent * 0.5 + 0.5, 1.0f);
+		#elif SUBSTRATE_MATERIAL_EXPORT_TYPE == SUBSTRATE_MATERIAL_EXPORT_SHADINGMODEL
+			SubstrateMaterialExport = float4(Export.ShadingModel.xxx, 1.0f);
 		#elif SUBSTRATE_MATERIAL_EXPORT_TYPE == SUBSTRATE_MATERIAL_EXPORT_MATERIAL_PREVIEW
 			SubstrateMaterialExport = float4(PixelMaterialInputs.EmissiveColor, 1.0f);
 		#endif
 	#elif SUBSTRATE_MATERIAL_EXPORT_FROM_TRANSLUCENT
 		#if SUBSTRATE_USES_CONVERSION_FROM_LEGACY
-			#if	  SUBSTRATE_MATERIAL_EXPORT_TYPE == SUBSTRATE_MATERIAL_EXPORT_DIFFUSE_COLOR
+			#if	  SUBSTRATE_MATERIAL_EXPORT_TYPE == SUBSTRATE_MATERIAL_EXPORT_BASE_COLOR || SUBSTRATE_MATERIAL_EXPORT_TYPE == SUBSTRATE_MATERIAL_EXPORT_BASE_COLOR_POST_COVERAGE
 				SubstrateMaterialExport = float4(0.0, 0.0, 0.0, 1.0);
 			#elif SUBSTRATE_MATERIAL_EXPORT_TYPE == SUBSTRATE_MATERIAL_EXPORT_NORMAL
 				SubstrateMaterialExport = float4(0.0, 0.0, 0.0, 1.0);
 			#elif SUBSTRATE_MATERIAL_EXPORT_TYPE == SUBSTRATE_MATERIAL_EXPORT_EMISSIVE
-				SubstrateMaterialExport = float4(EmissiveLuminance, 1.0);
+				SubstrateMaterialExport = float4(Export.EmissiveLuminance, 1.0);
 			#elif SUBSTRATE_MATERIAL_EXPORT_TYPE == SUBSTRATE_MATERIAL_EXPORT_TRANSMITTANCE
 				#if SUBSTRATE_MATERIAL_EXPORT_LEGACY_BLEND_MODE == SUBSTRATE_MATERIAL_EXPORT_LEGACY_BLEND_MODULATE
-					SubstrateMaterialExport = float4(BaseColorPostCoverage, 1.0f);
+					SubstrateMaterialExport = float4(Export.BaseColorPostCoverage, 1.0f);
 				#else
-					SubstrateMaterialExport = float4(BaseColorPostCoverage + saturate(1.0f - Coverage), 1.0f);
+					SubstrateMaterialExport = float4(Export.BaseColorPostCoverage + saturate(1.0f - Export.Coverage), 1.0f);
 				#endif
 			#elif SUBSTRATE_MATERIAL_EXPORT_TYPE == SUBSTRATE_MATERIAL_EXPORT_MATERIAL_PREVIEW
 				SubstrateMaterialExport = float4(PixelMaterialInputs.EmissiveColor, 1.0f);
 			#endif
 		#else
-			#if	  SUBSTRATE_MATERIAL_EXPORT_TYPE == SUBSTRATE_MATERIAL_EXPORT_DIFFUSE_COLOR
-				SubstrateMaterialExport = float4(BaseColorPostCoverage, 1.0);
+			#if	  SUBSTRATE_MATERIAL_EXPORT_TYPE == SUBSTRATE_MATERIAL_EXPORT_BASE_COLOR || SUBSTRATE_MATERIAL_EXPORT_TYPE == SUBSTRATE_MATERIAL_EXPORT_BASE_COLOR_POST_COVERAGE
+				SubstrateMaterialExport = float4(Export.BaseColorPostCoverage, 1.0);
 			#elif SUBSTRATE_MATERIAL_EXPORT_TYPE == SUBSTRATE_MATERIAL_EXPORT_NORMAL
-				SubstrateMaterialExport = float4(WorldNormal * 0.5 + 0.5, 1.0);
+				SubstrateMaterialExport = float4(Export.WorldNormal * 0.5 + 0.5, 1.0);
 			#elif SUBSTRATE_MATERIAL_EXPORT_TYPE == SUBSTRATE_MATERIAL_EXPORT_EMISSIVE
-				SubstrateMaterialExport = float4(EmissiveLuminance, 1.0);
+				SubstrateMaterialExport = float4(Export.EmissiveLuminance, 1.0);
 			#elif SUBSTRATE_MATERIAL_EXPORT_TYPE == SUBSTRATE_MATERIAL_EXPORT_TRANSMITTANCE
-				SubstrateMaterialExport = float4(TransmittancePreCoverage * Coverage + (1.0f - Coverage), 1.0f);
+				SubstrateMaterialExport = float4(Export.TransmittancePreCoverage * Coverage + (1.0f - Coverage), 1.0f);
 			#elif SUBSTRATE_MATERIAL_EXPORT_TYPE == SUBSTRATE_MATERIAL_EXPORT_MATERIAL_PREVIEW
 				SubstrateMaterialExport = float4(PixelMaterialInputs.EmissiveColor, 1.0f);
 			#endif
@@ -1892,7 +1909,7 @@ void FPixelShaderInOut_MainPS(
 #if ((SUBSTRATE_OPAQUE_DEFERRED || (SUBSTRATE_OPAQUE_MATERIAL && SUBSTRATE_FORWARD_SHADING))) && !SUBSTRATE_MATERIAL_EXPORT_FROM_TRANSLUCENT
 	Opacity = 1.0f;
 
-#if   SUBSTRATE_MATERIAL_EXPORT_TYPE == SUBSTRATE_MATERIAL_EXPORT_DIFFUSE_COLOR
+#if	  SUBSTRATE_MATERIAL_EXPORT_TYPE == SUBSTRATE_MATERIAL_EXPORT_BASE_COLOR || SUBSTRATE_MATERIAL_EXPORT_TYPE == SUBSTRATE_MATERIAL_EXPORT_BASE_COLOR_POST_COVERAGE
 	SubstrateMaterialExport = float4(0.0, 0.0, 0.0, 1.0);
 #elif SUBSTRATE_MATERIAL_EXPORT_TYPE == SUBSTRATE_MATERIAL_EXPORT_NORMAL
 	SubstrateMaterialExport = float4(0.0, 0.0, 0.0, 1.0);
@@ -1911,7 +1928,7 @@ void FPixelShaderInOut_MainPS(
 	Opacity = 1.0f;
 
 #if SUBSTRATE_USES_CONVERSION_FROM_LEGACY
-	#if   SUBSTRATE_MATERIAL_EXPORT_TYPE == SUBSTRATE_MATERIAL_EXPORT_DIFFUSE_COLOR
+#if	  SUBSTRATE_MATERIAL_EXPORT_TYPE == SUBSTRATE_MATERIAL_EXPORT_BASE_COLOR || SUBSTRATE_MATERIAL_EXPORT_TYPE == SUBSTRATE_MATERIAL_EXPORT_BASE_COLOR_POST_COVERAGE
 		SubstrateMaterialExport = float4(0.0, 0.0, 0.0, 1.0);
 	#elif SUBSTRATE_MATERIAL_EXPORT_TYPE == SUBSTRATE_MATERIAL_EXPORT_NORMAL
 		SubstrateMaterialExport = float4(0.0, 0.0, 0.0, 1.0);
@@ -1928,7 +1945,7 @@ void FPixelShaderInOut_MainPS(
 	#endif
 #else // SUBSTRATE_USES_CONVERSION_FROM_LEGACY
 	// The legacy unlit to lightmass conversion is weird and very specialized.
-	#if   SUBSTRATE_MATERIAL_EXPORT_TYPE == SUBSTRATE_MATERIAL_EXPORT_DIFFUSE_COLOR
+#if	  SUBSTRATE_MATERIAL_EXPORT_TYPE == SUBSTRATE_MATERIAL_EXPORT_BASE_COLOR || SUBSTRATE_MATERIAL_EXPORT_TYPE == SUBSTRATE_MATERIAL_EXPORT_BASE_COLOR_POST_COVERAGE
 		SubstrateMaterialExport = float4(0.0, 0.0, 0.0, 1.0);
 	#elif SUBSTRATE_MATERIAL_EXPORT_TYPE == SUBSTRATE_MATERIAL_EXPORT_NORMAL
 		SubstrateMaterialExport = float4(0.0, 0.0, 0.0, 1.0);
diff --git a/Engine/Shaders/Private/Definitions.usf b/Engine/Shaders/Private/Definitions.usf
--- a/Engine/Shaders/Private/Definitions.usf
+++ b/Engine/Shaders/Private/Definitions.usf
@@ -278,15 +278,33 @@
 #endif
 
 #ifndef SUBSTRATE_USE_FULLYSIMPLIFIED_MATERIAL
+#if (SUBSTRATE_MATERIAL_EXPORT_TYPE > 0) // 0 == SUBSTRATE_MATERIAL_EXPORT_NONE
+// When exporting anything, we always use the fully simplified Substrate materials.
+#define SUBSTRATE_USE_FULLYSIMPLIFIED_MATERIAL				1
+#else
 #define SUBSTRATE_USE_FULLYSIMPLIFIED_MATERIAL				0
 #endif
+#endif
 
 // Those define usage should match the meaning of ESubstrateMaterialExport enum from MaterialCompiler.h
-#define SUBSTRATE_MATERIAL_EXPORT_DIFFUSE_COLOR					1
-#define SUBSTRATE_MATERIAL_EXPORT_NORMAL						2
-#define SUBSTRATE_MATERIAL_EXPORT_EMISSIVE						3
-#define SUBSTRATE_MATERIAL_EXPORT_TRANSMITTANCE					4
-#define SUBSTRATE_MATERIAL_EXPORT_MATERIAL_PREVIEW				5
+#define SUBSTRATE_MATERIAL_EXPORT_NONE							0
+#define SUBSTRATE_MATERIAL_EXPORT_BASE_COLOR_POST_COVERAGE		1
+#define SUBSTRATE_MATERIAL_EXPORT_BASE_COLOR					2
+#define SUBSTRATE_MATERIAL_EXPORT_NORMAL						3
+#define SUBSTRATE_MATERIAL_EXPORT_EMISSIVE						4
+#define SUBSTRATE_MATERIAL_EXPORT_SPECULAR						5
+#define SUBSTRATE_MATERIAL_EXPORT_ROUGHNESS						6
+#define SUBSTRATE_MATERIAL_EXPORT_ANISOTROPY					7
+#define SUBSTRATE_MATERIAL_EXPORT_METALLIC						8
+#define SUBSTRATE_MATERIAL_EXPORT_TRANSMITTANCE					9
+#define SUBSTRATE_MATERIAL_EXPORT_OPACITY						10
+#define SUBSTRATE_MATERIAL_EXPORT_OPACITYMASK					11
+#define SUBSTRATE_MATERIAL_EXPORT_CUSTOMDATA0					12
+#define SUBSTRATE_MATERIAL_EXPORT_CUSTOMDATA1					13
+#define SUBSTRATE_MATERIAL_EXPORT_SUBSURFACECOLOR				14
+#define SUBSTRATE_MATERIAL_EXPORT_TANGENT						15
+#define SUBSTRATE_MATERIAL_EXPORT_SHADINGMODEL					16
+#define SUBSTRATE_MATERIAL_EXPORT_MATERIAL_PREVIEW				17
 
 #define SUBSTRATE_MATERIAL_EXPORT_CONTEXT_OPAQUE				0
 #define SUBSTRATE_MATERIAL_EXPORT_CONTEXT_TRANSLUCENT			1
diff --git a/Engine/Shaders/Private/PathTracing/PathTracingGPULightmassMaterialHitShader.usf b/Engine/Shaders/Private/PathTracing/PathTracingGPULightmassMaterialHitShader.usf
--- a/Engine/Shaders/Private/PathTracing/PathTracingGPULightmassMaterialHitShader.usf
+++ b/Engine/Shaders/Private/PathTracing/PathTracingGPULightmassMaterialHitShader.usf
@@ -87,12 +87,10 @@ RAY_TRACING_ENTRY_CLOSEST_HIT(GPULightmassMaterialCHS,
 #if MATERIAL_IS_SUBSTRATE
 	FSubstrateData SubstrateData = PixelMaterialInputs.GetFrontSubstrateData();
 	FSubstratePixelHeader SubstratePixelHeader = MaterialParameters.GetFrontSubstrateHeader();
-
-
 	FSubstrateIntegrationSettings Settings = InitSubstrateIntegrationSettings(true /*bForceFullyRough*/, false /*SubstrateStruct.bRoughDiffuse*/, -1/*SubstrateStruct.PeelLayersAboveDepth*/, false/*SubstrateStruct.bRoughnessTracking*/);
 
 	const float3 SurfaceWorldNormal = GeoNormal;
-	SubstrateMaterialExportOut(
+	FExportResult Export = SubstrateMaterialExportOut(
 		Settings,
 		SubstratePixelHeader,
 		SubstrateData,
@@ -102,6 +100,12 @@ RAY_TRACING_ENTRY_CLOSEST_HIT(GPULightmassMaterialCHS,
 		BaseColorPostCoverage,
 		WorldNormal,
 		EmissiveLuminance);
+
+	Coverage = Export.Coverage;
+	TransmittancePreCoverage = Export.TransmittancePreCoverage;
+	BaseColorPostCoverage = Export.BaseColorPostCoverage;
+	WorldNormal = Export.WorldNormal;
+	EmissiveLuminance = Export.EmissiveLuminance;
 #endif // MATERIAL_IS_SUBSTRATE
 
 
@@ -405,7 +409,7 @@ RAY_TRACING_ENTRY_ANY_HIT(GPULightmassMaterialAHS,
 	FSubstrateIntegrationSettings Settings = InitSubstrateIntegrationSettings(true /*bForceFullyRough*/, false /*SubstrateStruct.bRoughDiffuse*/, -1/*SubstrateStruct.PeelLayersAboveDepth*/, false/*SubstrateStruct.bRoughnessTracking*/);
 
 	const float3 SurfaceWorldNormal = GeoNormal;
-	SubstrateMaterialExportOut(
+	FExportResult Export = SubstrateMaterialExportOut(
 		Settings,
 		SubstratePixelHeader,
 		SubstrateData,
@@ -416,6 +420,12 @@ RAY_TRACING_ENTRY_ANY_HIT(GPULightmassMaterialAHS,
 		WorldNormal,
 		EmissiveLuminance);
 
+	Coverage = Export.Coverage;
+	TransmittancePreCoverage = Export.TransmittancePreCoverage;
+	BaseColorPostCoverage = Export.BaseColorPostCoverage;
+	WorldNormal = Export.WorldNormal;
+	EmissiveLuminance = Export.EmissiveLuminance;
+
 	float3 Transparency = saturate((1.0f - Coverage) + Coverage * TransmittancePreCoverage);
 #else
 #error "Shouldn't need to compile AHS for non-substrate material"
diff --git a/Engine/Shaders/Private/Substrate/SubstrateExport.ush b/Engine/Shaders/Private/Substrate/SubstrateExport.ush
--- a/Engine/Shaders/Private/Substrate/SubstrateExport.ush
+++ b/Engine/Shaders/Private/Substrate/SubstrateExport.ush
@@ -934,35 +934,45 @@ void PackSubstrateOut(
 
 #if SUBSTRATE_MATERIAL_EXPORT_TYPE > 0 || SUBSTRATE_GPU_LIGHTMASS > 0
 
-void SubstrateMaterialExportOut(
+struct FExportResult
+{
+	float  Coverage;
+	float3 TransmittancePreCoverage;
+	float3 BaseColorPostCoverage;
+	float3 BaseColor;
+	float3 WorldNormal;
+	float3 EmissiveLuminance;
+	float  Specular;
+	float  Roughness;
+	float  Anisotropy;
+	float  Metallic;
+	float3 CustomData0;
+	float3 CustomData1;
+	float3 SubsurfaceColor;
+	float3 Tangent;
+	float  ShadingModel;
+};
+
+// This is used to export very simple data for lightmaps, material baking or preview on editor node.
+// In this case we force full simplification and only recover a single BSDF in the end/ 
+FExportResult SubstrateMaterialExportOut(
 	in FSubstrateIntegrationSettings Settings,
 	in FSubstratePixelHeader SubstratePixelHeader,
 	in FSubstrateData SubstrateData,
-	in float3 SurfaceWorldNormal,
-	inout float  OutCoverage, 
-	inout float3 OutTransmittancePreCoverage, 
-	inout float3 OutBaseColorPostCoverage, 
-	inout float3 OutWorldNormal,
-	inout float3 OutEmissiveLuminance)
+	in float3 SurfaceWorldNormal)
 {
+	FExportResult Out = (FExportResult)0;
+
 	// Align V with N to not have any view dependent effect when exporting material data.
 	const float3 FakeV = SurfaceWorldNormal;
-
-	OutCoverage = 0.0f;
-	OutTransmittancePreCoverage = 0.0f;
-	OutBaseColorPostCoverage = 0.0f;
-	OutWorldNormal = 0.0f;
-	OutEmissiveLuminance = 0.0f;
 	
 	if (SubstratePixelHeader.SubstrateTree.BSDFCount > 0)
 	{
 		// Update tree (coverage/transmittance/luminace weights)
-		SubstratePixelHeader.SubstrateUpdateTree(SubstrateData, FakeV, Settings, OutCoverage, OutTransmittancePreCoverage); 
+		SubstratePixelHeader.SubstrateUpdateTree(SubstrateData, FakeV, Settings, Out.Coverage, Out.TransmittancePreCoverage);
 
-		//
-		// Forward lighting
-		//
-		Substrate_for_unroll(int BSDFIdx = 0, BSDFIdx < SubstratePixelHeader.SubstrateTree.BSDFCount, ++BSDFIdx)
+		const int BSDFIdx = 0;
+		if(SubstratePixelHeader.SubstrateTree.BSDFCount > 0)
 		{
 		#define CurrentBSDF SubstratePixelHeader.SubstrateTree.BSDFs[BSDFIdx]
 
@@ -970,21 +980,31 @@ void SubstrateMaterialExportOut(
 			{
 				const float3x3 TangentBasis = SubstrateGetBSDFSharedBasis_InlineShading(SubstratePixelHeader, CurrentBSDF);
 				const float3 N = TangentBasis[2];
-
-				const float GreyScaleWeight = dot(CurrentBSDF.LuminanceWeightV * CurrentBSDF.TransmittanceAboveAlongN, (1.0/3.0).xxx);
-
-				// Simply sum all BaseColors according to transmittance and coverage.
-				OutBaseColorPostCoverage += SubstrateGetBSDFBaseColor(CurrentBSDF) * CurrentBSDF.LuminanceWeightV * CurrentBSDF.TransmittanceAboveAlongN;
-
-				OutWorldNormal += N * GreyScaleWeight;
-
-				OutEmissiveLuminance += BSDF_GETEMISSIVE(CurrentBSDF) * CurrentBSDF.LuminanceWeightV;
+				const float3 T = TangentBasis[0];
+
+				const float GreyScaleWeight = dot(CurrentBSDF.LuminanceWeightV, (1.0 / 3.0).xxx);
+
+				Out.BaseColorPostCoverage	= SubstrateGetBSDFBaseColor(CurrentBSDF)	* CurrentBSDF.LuminanceWeightV;
+				Out.BaseColor				= SubstrateGetBSDFBaseColor(CurrentBSDF);
+				Out.WorldNormal				= N;
+				Out.EmissiveLuminance		= BSDF_GETEMISSIVE(CurrentBSDF)				* CurrentBSDF.LuminanceWeightV;
+				Out.Specular				= SubstrateGetBSDFSpecular(CurrentBSDF);
+				Out.Roughness				= SubstrateGetBSDFRoughness(CurrentBSDF);
+				Out.Anisotropy				= SubstrateGetBSDFAnisotropy(CurrentBSDF);
+				Out.Metallic				= SubstrateGetBSDFMetallic(CurrentBSDF);
+				Out.CustomData0				= 0;											// SUBSTRATE_TODO not tested since not exposed in the only use case so far: material baking.
+				Out.CustomData1				= 0;											// idem
+				Out.SubsurfaceColor			= SubstrateGetBSDFSubSurfaceColor(CurrentBSDF);	// idem
+				Out.Tangent					= T;
+				Out.ShadingModel			= SubstrateGetLegacyShadingModels(CurrentBSDF);	// idem
 			}
 		#undef CurrentBSDF
 		}
 
-		OutWorldNormal = all(OutWorldNormal == 0) ? float3(0.0, 0.0, 0.0) : normalize(OutWorldNormal);
+		Out.WorldNormal = all(Out.WorldNormal == 0) ? float3(0.0, 0.0, 0.0) : normalize(Out.WorldNormal);
 	}
+
+	return Out;
 }
 
 #endif // SUBSTRATE_MATERIAL_EXPORT_TYPE
diff --git a/Engine/Shaders/Private/Substrate/SubstrateLegacyConversion.ush b/Engine/Shaders/Private/Substrate/SubstrateLegacyConversion.ush
--- a/Engine/Shaders/Private/Substrate/SubstrateLegacyConversion.ush
+++ b/Engine/Shaders/Private/Substrate/SubstrateLegacyConversion.ush
@@ -8,6 +8,8 @@
 // This is for matching legacy behavior. See ShadingModelsMaterial.ush.
 #define SUBSURFACE_PROFILE_OPACITY_THRESHOLD 1
 
+#define SUBSTRATE_LEGACYCONVERT_OUTPUT_OPACITY (SUBSTRATE_TRANSLUCENT_MATERIAL || SUBSTRATE_MATERIAL_EXPORT_TYPE == SUBSTRATE_MATERIAL_EXPORT_OPACITY)
+
 ///////////////////////////////////////////////////////////////////////////////////////////////////
 // Helpers
 
@@ -237,7 +239,7 @@ FSubstrateData SubstrateConvertLegacyMaterialDynamic(
 		LegacySlab.Emissive = Emissive;
 
 		FinalWeight = 1.0f;
-#if SUBSTRATE_TRANSLUCENT_MATERIAL
+#if SUBSTRATE_LEGACYCONVERT_OUTPUT_OPACITY
 		FinalWeight = Opacity;
 #endif
 	}
@@ -251,7 +253,7 @@ FSubstrateData SubstrateConvertLegacyMaterialDynamic(
 		LegacySlab.Emissive = Emissive;
 
 		FinalWeight = 1.0f;
-#if SUBSTRATE_TRANSLUCENT_MATERIAL
+#if SUBSTRATE_LEGACYCONVERT_OUTPUT_OPACITY
 		FinalWeight = Opacity;
 #endif
 	}
@@ -274,7 +276,7 @@ FSubstrateData SubstrateConvertLegacyMaterialDynamic(
 
 		FinalWeight = 1.0f;
 
-#if SUBSTRATE_TRANSLUCENT_MATERIAL
+#if SUBSTRATE_LEGACYCONVERT_OUTPUT_OPACITY
 		// SUBSTRATE_TODO 
 		// If we try to reduce the diffusion according to transmittance that will increase mfp and make the material disapear (make it optically thin).
 		// Instead for now we use a mfp of 0 and make the material vanish according to opacity.
@@ -307,7 +309,7 @@ FSubstrateData SubstrateConvertLegacyMaterialDynamic(
 		LegacySlab.Emissive = Emissive;
 
 		FinalWeight = 1.0f;
-#if SUBSTRATE_TRANSLUCENT_MATERIAL
+#if SUBSTRATE_LEGACYCONVERT_OUTPUT_OPACITY
 		FinalWeight = Opacity;
 #endif
 	}
@@ -330,7 +332,7 @@ FSubstrateData SubstrateConvertLegacyMaterialDynamic(
 		LegacySlab.Emissive = Emissive;
 
 		FinalWeight = 1.0f;
-#if SUBSTRATE_TRANSLUCENT_MATERIAL
+#if SUBSTRATE_LEGACYCONVERT_OUTPUT_OPACITY
 		FinalWeight = Opacity;
 #endif
 	}
@@ -351,7 +353,7 @@ FSubstrateData SubstrateConvertLegacyMaterialDynamic(
 		LegacySlab.Weight = 1.f;
 
 		FinalWeight = 1.0f;
-#if SUBSTRATE_TRANSLUCENT_MATERIAL
+#if SUBSTRATE_LEGACYCONVERT_OUTPUT_OPACITY
 		FinalWeight = Opacity;
 #endif
 	}
@@ -376,7 +378,7 @@ FSubstrateData SubstrateConvertLegacyMaterialDynamic(
 		LegacySlab.bIsThinSurface = true;
 
 		FinalWeight = 1.0f;
-#if SUBSTRATE_TRANSLUCENT_MATERIAL
+#if SUBSTRATE_LEGACYCONVERT_OUTPUT_OPACITY
 		FinalWeight = Opacity;
 #endif
 	}
@@ -395,7 +397,7 @@ FSubstrateData SubstrateConvertLegacyMaterialDynamic(
 		LegacySlab.FuzzColor = SubSurfaceColor;
 
 		FinalWeight = 1.0f;
-#if SUBSTRATE_TRANSLUCENT_MATERIAL
+#if SUBSTRATE_LEGACYCONVERT_OUTPUT_OPACITY
 		FinalWeight = Opacity;
 #endif
 	}
diff --git a/Engine/Source/Developer/MaterialBaking/Private/ExportMaterialProxy.h b/Engine/Source/Developer/MaterialBaking/Private/ExportMaterialProxy.h
--- a/Engine/Source/Developer/MaterialBaking/Private/ExportMaterialProxy.h
+++ b/Engine/Source/Developer/MaterialBaking/Private/ExportMaterialProxy.h
@@ -356,6 +356,53 @@ public:
 			FExportMaterialCompiler ProxyCompiler(Compiler);
 			const uint32 ForceCast_Exact_Replicate = MFCF_ForceCast | MFCF_ExactMatch | MFCF_ReplicateValue;
 
+			switch (PropertyToCompile)
+			{
+			case MP_EmissiveColor:
+				Compiler->SetSubstrateMaterialExportType(SME_Emissive, ESubstrateMaterialExportContext::SMEC_Opaque, BLEND_Opaque);
+				break;
+			case MP_BaseColor:
+				Compiler->SetSubstrateMaterialExportType(SME_BaseColor, ESubstrateMaterialExportContext::SMEC_Opaque, BLEND_Opaque);
+				break;
+			case MP_Specular:
+				Compiler->SetSubstrateMaterialExportType(SME_Specular, ESubstrateMaterialExportContext::SMEC_Opaque, BLEND_Opaque);
+				break;
+			case MP_Roughness:
+				Compiler->SetSubstrateMaterialExportType(SME_Roughness, ESubstrateMaterialExportContext::SMEC_Opaque, BLEND_Opaque);
+				break;
+			case MP_Anisotropy:
+				Compiler->SetSubstrateMaterialExportType(SME_Anisotropy, ESubstrateMaterialExportContext::SMEC_Opaque, BLEND_Opaque);
+				break;
+			case MP_Metallic:
+				Compiler->SetSubstrateMaterialExportType(SME_Metallic, ESubstrateMaterialExportContext::SMEC_Opaque, BLEND_Opaque);
+				break;
+			case MP_Opacity:
+				Compiler->SetSubstrateMaterialExportType(SME_Opacity, ESubstrateMaterialExportContext::SMEC_Opaque, BLEND_Opaque);
+				break;
+			case MP_OpacityMask:
+				Compiler->SetSubstrateMaterialExportType(SME_OpacityMask, ESubstrateMaterialExportContext::SMEC_Opaque, BLEND_Opaque);
+				break;
+			case MP_CustomData0:
+				Compiler->SetSubstrateMaterialExportType(SME_CustomData0, ESubstrateMaterialExportContext::SMEC_Opaque, BLEND_Opaque);
+				break;
+			case MP_CustomData1:
+				Compiler->SetSubstrateMaterialExportType(SME_CustomData1, ESubstrateMaterialExportContext::SMEC_Opaque, BLEND_Opaque);
+				break;
+			case MP_SubsurfaceColor:
+				Compiler->SetSubstrateMaterialExportType(SME_SubsurfaceColor, ESubstrateMaterialExportContext::SMEC_Opaque, BLEND_Opaque);
+				break;
+			case MP_Normal:
+				Compiler->SetSubstrateMaterialExportType(SME_Normal, ESubstrateMaterialExportContext::SMEC_Opaque, BLEND_Opaque);
+				break;
+			case MP_Tangent:
+				Compiler->SetSubstrateMaterialExportType(SME_Tangent, ESubstrateMaterialExportContext::SMEC_Opaque, BLEND_Opaque);
+				break;
+			case MP_ShadingModel:
+				Compiler->SetSubstrateMaterialExportType(SME_ShadingModel, ESubstrateMaterialExportContext::SMEC_Opaque, BLEND_Opaque);
+				break;
+			}
+
+
 			switch (PropertyToCompile)
 			{
 			case MP_EmissiveColor:
@@ -394,6 +441,17 @@ public:
 					return CompileInputForCustomOutput(&ProxyCompiler, CustomAttribute, InputIndex, ForceCast_Exact_Replicate);
 				}
 				break;
+			case MP_FrontMaterial:
+				if (Substrate::IsSubstrateEnabled())
+				{
+					// When using Substrate, material property always compile from material. 
+					// We cannot use rediction so instead we instruct the compiler the type of data export we are looking for.
+					return MaterialInterface->CompileProperty(&ProxyCompiler, MP_FrontMaterial);
+				}
+				else
+				{
+					return ProxyCompiler.SubstrateCreateAndRegisterNullMaterial();
+				}
 			default:
 				return Compiler->Constant(1.0f);
 			}
@@ -410,10 +468,13 @@ public:
 			// Pass through customized UVs
 			return MaterialInterface->CompileProperty(Compiler, Property);
 		}
+		else if (Property == MP_OpacityMask)
+		{
+			return MaterialInterface->CompileProperty(Compiler, MP_OpacityMask);
+		}
 		else if (Property == MP_ShadingModel)
 		{
 			return MaterialInterface->CompileProperty(Compiler, MP_ShadingModel);
-
 		}
 		else if (Property == MP_SurfaceThickness)
 		{
@@ -421,10 +482,23 @@ public:
 		}
 		else if (Property == MP_FrontMaterial)
 		{
-			return MaterialInterface->CompileProperty(Compiler, MP_FrontMaterial);
+			if (Substrate::IsSubstrateEnabled())
+			{
+				return MaterialInterface->CompileProperty(Compiler, MP_FrontMaterial);
+			}
+			else
+			{
+				return Compiler->SubstrateCreateAndRegisterNullMaterial();
+			}
 		}
 		else
 		{
+			if (Substrate::IsSubstrateEnabled() && Property == MP_Roughness)
+			{
+				// In this case return non 0.5 to not trigger the fully rough path from the Translator.
+				// This is because we do not only compile EmissiveColor only anymore since we use the specified shading model (see GetShadingModels below).
+				return Compiler->Constant(0.5f);
+			}
 			return Compiler->Constant(1.0f);
 		}
 	}
@@ -520,7 +594,14 @@ public:
 	virtual enum EBlendMode GetBlendMode() const override { return ProxyBlendMode; }
 	virtual enum ERefractionMode GetRefractionMode() const override { return Material ? (ERefractionMode)Material->RefractionMethod : RM_None; }
 	virtual bool GetRootNodeOverridesDefaultRefraction()const override { return Material ? Material->bRootNodeOverridesDefaultDistortion : false; }
-	virtual FMaterialShadingModelField GetShadingModels() const override { return MSM_DefaultLit; }
+
+	virtual FMaterialShadingModelField GetShadingModels() const override
+	{
+		// Substrate needs the real material shading model since the expressions access GetMaterialShadingModels() through the compiler to generate the substrate operators, 
+		// and we do not want unlit materials when it is a Slab or a ShadingModel node.
+		return Substrate::IsSubstrateEnabled() ? Material->GetShadingModels() : MSM_Unlit;
+	}
+
 	virtual bool IsShadingModelFromMaterialExpression() const override { return false; }
 	virtual float GetOpacityMaskClipValue() const override { return 0.5f; }
 	virtual bool GetCastDynamicShadowAsMasked() const override { return false; }
diff --git a/Engine/Source/Developer/MeshMergeUtilities/Private/MeshMergeUtilities.cpp b/Engine/Source/Developer/MeshMergeUtilities/Private/MeshMergeUtilities.cpp
--- a/Engine/Source/Developer/MeshMergeUtilities/Private/MeshMergeUtilities.cpp
+++ b/Engine/Source/Developer/MeshMergeUtilities/Private/MeshMergeUtilities.cpp
@@ -315,7 +315,11 @@ void FMeshMergeUtilities::BakeMaterialsForComponent(TArray<TWeakObjectPtr<UObjec
 				// Add all user defined properties for baking out
 				for (const FPropertyEntry& Entry : MaterialOptions->Properties)
 				{
-					if (!Entry.bUseConstantValue && Material->IsPropertyActive(Entry.Property) && Entry.Property != MP_MAX)
+					// With substrate, we cannot use connection on the root node to infer Active or Constant properties. So we activate all material input for now. 
+					// SUBSTRATE_TODO: gather information about the material used input, maybe using Material->GetSubstrateCompilationConfig().
+					const float bIsSubstrateEnabled = Substrate::IsSubstrateEnabled();
+
+					if ((bIsSubstrateEnabled || !Entry.bUseConstantValue && Material->IsPropertyActive(Entry.Property)) && Entry.Property != MP_MAX)
 					{
 						MaterialSettings.PropertySizes.Add(Entry.Property, Entry.bUseCustomSize ? Entry.CustomSize : MaterialOptions->TextureSize);
 					}
diff --git a/Engine/Source/Editor/UnrealEd/Private/Lightmass/LightmassRender.cpp b/Engine/Source/Editor/UnrealEd/Private/Lightmass/LightmassRender.cpp
--- a/Engine/Source/Editor/UnrealEd/Private/Lightmass/LightmassRender.cpp
+++ b/Engine/Source/Editor/UnrealEd/Private/Lightmass/LightmassRender.cpp
@@ -389,7 +389,7 @@ public:
 
 			if (Usage == EMaterialShaderMapUsage::LightmassExportDiffuse)
 			{
-				Compiler->SetSubstrateMaterialExportType(SME_Diffuse, SubstrateMaterialExportContext, BlendMode);
+				Compiler->SetSubstrateMaterialExportType(SME_BaseColorPostCoverage, SubstrateMaterialExportContext, BlendMode);
 			}
 			else if (Usage == EMaterialShaderMapUsage::LightmassExportNormal)
 			{
diff --git a/Engine/Source/Runtime/Core/Private/UObject/FortniteMainBranchObjectVersion.cpp b/Engine/Source/Runtime/Core/Private/UObject/FortniteMainBranchObjectVersion.cpp
--- a/Engine/Source/Runtime/Core/Private/UObject/FortniteMainBranchObjectVersion.cpp
+++ b/Engine/Source/Runtime/Core/Private/UObject/FortniteMainBranchObjectVersion.cpp
@@ -14,6 +14,6 @@ TMap<FGuid, FGuid> FFortniteMainBranchObjectVersion::GetSystemGuids()
 	SystemGuids.Add(DevGuids.Niagara_LatestScriptCompileVersion, FGuid("393081597FA7234C92F00869DE262D73"));
 	SystemGuids.Add(DevGuids.SkeletalMeshDerivedDataVersion, FGuid("2F90180769024483BB7D90AACEB2798D"));
 	SystemGuids.Add(DevGuids.STATICMESH_DERIVEDDATA_VER, FGuid("656E338C78C34641A5D3BFAE303997AF"));
-	SystemGuids.Add(DevGuids.MaterialTranslationDDCVersion, FGuid("33FAE17575F24CBEB88622EEB1450EF4"));
+	SystemGuids.Add(DevGuids.MaterialTranslationDDCVersion, FGuid("0B923787D28841B8A53F4B28C8063EBA"));
 	return SystemGuids;
 }
diff --git a/Engine/Source/Runtime/Engine/Private/Materials/HLSLMaterialTranslator.cpp b/Engine/Source/Runtime/Engine/Private/Materials/HLSLMaterialTranslator.cpp
--- a/Engine/Source/Runtime/Engine/Private/Materials/HLSLMaterialTranslator.cpp
+++ b/Engine/Source/Runtime/Engine/Private/Materials/HLSLMaterialTranslator.cpp
@@ -14756,6 +14756,7 @@ int32 FHLSLMaterialTranslator::SubstrateCompilePreview(int32 SubstrateDataCodeCh
 	// Instead of using a preview color, we might go with a sphere+lighting preview. 
 	// The only problem we would have to solve is to process the Substrate tree before for the sub tree. 
 	// Or have dedicated lighting preview functions for SubstrateData that have not be filled in by the Substrate tree processing.
+	// PreviewColor is filled up using SME_MaterialPreview, driving SUBSTRATE_MATERIAL_EXPORT_MATERIAL_PREVIEW.
 	int32 PreviewCodeChunk = AddCodeChunk(MCT_Float3, TEXT("%s.PreviewColor"), *GetParameterCode(SubstrateDataCodeChunk));
 	return PreviewCodeChunk;
 }
@@ -14763,6 +14764,7 @@ int32 FHLSLMaterialTranslator::SubstrateCompilePreview(int32 SubstrateDataCodeCh
 bool FHLSLMaterialTranslator::SubstrateSkipsOpacityEvaluation()
 {
 	return !IsTranslucentBlendMode(Material)
+		&& GetSubstrateMaterialExportType() == ESubstrateMaterialExport::SME_None	// We never skip when exporting information
 		&& Material->GetShadingModels().CountShadingModels() == 1
 		&& !Material->GetShadingModels().HasShadingModel(MSM_SingleLayerWater)
 		&& !Material->GetShadingModels().HasShadingModel(MSM_Subsurface)
diff --git a/Engine/Source/Runtime/Engine/Public/MaterialCompiler.h b/Engine/Source/Runtime/Engine/Public/MaterialCompiler.h
--- a/Engine/Source/Runtime/Engine/Public/MaterialCompiler.h
+++ b/Engine/Source/Runtime/Engine/Public/MaterialCompiler.h
@@ -67,13 +67,26 @@ enum class EMaterialCompilerType
 /** Whether we need some data export from a Substrate material from spatially varying properties, e.g. diffuse color for Lighmass to generate lightmaps. */
 enum ESubstrateMaterialExport : uint8
 {
-	SME_None			= 0,
-	SME_Diffuse			= 1,
-	SME_Normal			= 2,
-	SME_Emissive		= 3,
-	SME_Transmittance	= 4,
-	SME_MaterialPreview	= 5
+	SME_None					= 0,
+	SME_BaseColorPostCoverage	= 1,
+	SME_BaseColor				= 2,
+	SME_Normal					= 3,
+	SME_Emissive				= 4,
+	SME_Specular				= 5,
+	SME_Roughness				= 6,
+	SME_Anisotropy				= 7,
+	SME_Metallic				= 8,
+	SME_Transmittance			= 9,
+	SME_Opacity					= 10,
+	SME_OpacityMask				= 11,
+	SME_CustomData0				= 12,
+	SME_CustomData1				= 13,
+	SME_SubsurfaceColor			= 14,
+	SME_Tangent					= 15,
+	SME_ShadingModel			= 16,
+	SME_MaterialPreview			= 17
 };
+
 /** Exported materials are all opaque unlit. This is used to give some context to the export logic.*/
 enum ESubstrateMaterialExportContext : uint8
 {
