Ray Tracing: Fix artifacts in refraction when using ray traced translucency

This was caused by not saving/restoring the flags field of the payload during lighting calculations.

#rnx

diff --git a/Engine/Shaders/Private/RayTracing/RayTracingLightingCommon.ush b/Engine/Shaders/Private/RayTracing/RayTracingLightingCommon.ush
--- a/Engine/Shaders/Private/RayTracing/RayTracingLightingCommon.ush
+++ b/Engine/Shaders/Private/RayTracing/RayTracingLightingCommon.ush
@@ -328,8 +328,10 @@ float3 ComputeDirectLighting(
 #if SUBSTRATE_ENABLED
 	uint OldTransmittancePreCoverage = Payload.TransmittancePreCoverage;
 	uint OldSceneInstanceIndex = Payload.SceneInstanceIndex;
+	uint OldFlags = Payload.Flags;
 #else
 	uint2 OldIndirectIrradiance = uint2(Payload.PackedIndirectIrradiance[0], Payload.PackedIndirectIrradiance[1]);
+	uint OldFlags = Payload.FlagsAndMipBias;
 #endif
 	uint OldRayCone = Payload.PackedRayCone;
 	Payload.SetRadiance(float3(0, 0, 0));
@@ -456,7 +458,9 @@ float3 ComputeDirectLighting(
 #if SUBSTRATE_ENABLED
 	Payload.SceneInstanceIndex = OldSceneInstanceIndex;
 	Payload.TransmittancePreCoverage = OldTransmittancePreCoverage;
+	Payload.Flags = OldFlags;
 #else
+	Payload.FlagsAndMipBias = OldFlags;
 	Payload.PackedIndirectIrradiance[0] = OldIndirectIrradiance.x;  // Set back indirect irradiance (aliased with ShadowVisibility and LightIndex)
 	Payload.PackedIndirectIrradiance[1] = OldIndirectIrradiance.y;
 #endif
