ACES: Minor cleanup to unify the 1000 & 2000 nits paths since those are already parametrized.

diff --git a/Engine/Shaders/Private/PostProcessCombineLUTs.usf b/Engine/Shaders/Private/PostProcessCombineLUTs.usf
--- a/Engine/Shaders/Private/PostProcessCombineLUTs.usf
+++ b/Engine/Shaders/Private/PostProcessCombineLUTs.usf
@@ -346,37 +346,26 @@ float4 CombineLUTsCommon(float2 InUV, uint InLayerIndex)
 		OutDeviceColor = LinearTo709Branchless( OutputGamutColor );
 	}
 
-	// ACES 1000nit transform with PQ/2084 encoding, user specified gamut 
-	else if( GetOutputDevice() == TONEMAPPER_OUTPUT_ACES1000nitST2084 || GetOutputDevice() == TONEMAPPER_OUTPUT_ACES1000nitScRGB)
+	// ACES transform with PQ/2084 encoding, user specified gamut 
+	else if( GetOutputDevice() == TONEMAPPER_OUTPUT_ACES1000nitST2084 || GetOutputDevice() == TONEMAPPER_OUTPUT_ACES2000nitST2084)
 	{		
-		// 1000 nit ODT
 		FACESTonemapParams AcesParams = ComputeACESTonemapParams(ACESMinMaxData, ACESMidData, ACESCoefsLow_0, ACESCoefsHigh_0, ACESCoefsLow_4, ACESCoefsHigh_4, ACESSceneColorMultiplier, ACESGamutCompression);
-		float3 ODTColor = ACESOutputTransforms1000( GradedColor, (float3x3)WorkingColorSpace.ToAP0, AcesParams);
+		float3 ODTColor = ACESOutputTransform( GradedColor, (float3x3)WorkingColorSpace.ToAP0, AcesParams);
 
-		// Convert from AP1 to specified output gamut except for ScRGB as it might contain negative values. 
-		// In this case, the AP1->output conversion happens through ST2084ToScRGB in the tonemap function
-		if( GetOutputDevice() != TONEMAPPER_OUTPUT_ACES1000nitScRGB )
-		{
-			ODTColor = mul(AP1_2_Output, ODTColor);
-		}
+		ODTColor = mul(AP1_2_Output, ODTColor);
 
 		// Apply conversion to ST-2084 (Dolby PQ)
 		OutDeviceColor = LinearToST2084( ODTColor );
 	}
 
-	// ACES 2000nit transform with PQ/2084 encoding, user specified gamut 
-	else if( GetOutputDevice() == TONEMAPPER_OUTPUT_ACES2000nitST2084 || GetOutputDevice() == TONEMAPPER_OUTPUT_ACES2000nitScRGB)
+	// ACES transform with PQ/2084 encoding, AP1 gamut
+	else if( GetOutputDevice() == TONEMAPPER_OUTPUT_ACES1000nitScRGB || GetOutputDevice() == TONEMAPPER_OUTPUT_ACES2000nitScRGB)
 	{		
-		// 2000 nit ODT
 		FACESTonemapParams AcesParams = ComputeACESTonemapParams(ACESMinMaxData, ACESMidData, ACESCoefsLow_0, ACESCoefsHigh_0, ACESCoefsLow_4, ACESCoefsHigh_4, ACESSceneColorMultiplier, ACESGamutCompression);
-		float3 ODTColor = ACESOutputTransforms2000( GradedColor, (float3x3)WorkingColorSpace.ToAP0, AcesParams);
+		float3 ODTColor = ACESOutputTransform( GradedColor, (float3x3)WorkingColorSpace.ToAP0, AcesParams);
 
-		// Convert from AP1 to specified output gamut except for ScRGB as it might contain negative values. 
+		// No conversion from AP1 to specified output gamut for ScRGB as it might contain negative values. 
 		// In this case, the AP1->output conversion happens through ST2084ToScRGB in the tonemap function
-		if ( GetOutputDevice() != TONEMAPPER_OUTPUT_ACES2000nitScRGB)
-		{
-			ODTColor = mul(AP1_2_Output, ODTColor);
-		}
 
 		// Apply conversion to ST-2084 (Dolby PQ)
 		OutDeviceColor = LinearToST2084( ODTColor );
diff --git a/Engine/Shaders/Private/PostProcessDeviceEncodingOnly.usf b/Engine/Shaders/Private/PostProcessDeviceEncodingOnly.usf
--- a/Engine/Shaders/Private/PostProcessDeviceEncodingOnly.usf
+++ b/Engine/Shaders/Private/PostProcessDeviceEncodingOnly.usf
@@ -95,12 +95,12 @@ float4 DeviceEncodingOnlyCommonPS(
 		OutDeviceColor = LinearTo709Branchless( OutputGamutColor );
 	}
 
-	// ACES 1000nit transform with PQ/2084 encoding, user specified gamut 
-	else if( GetOutputDevice() == TONEMAPPER_OUTPUT_ACES1000nitST2084 )
+	// ACES transform with PQ/2084 encoding, user specified gamut 
+	else if (GetOutputDevice() == TONEMAPPER_OUTPUT_ACES1000nitST2084 || GetOutputDevice() == TONEMAPPER_OUTPUT_ACES2000nitST2084)
 	{		 
-		// 1000 nit ODT
+		// ACES ODT
 		FACESTonemapParams AcesParams = ComputeACESTonemapParams(ACESMinMaxData, ACESMidData, ACESCoefsLow_0, ACESCoefsHigh_0, ACESCoefsLow_4, ACESCoefsHigh_4, ACESSceneColorMultiplier, ACESGamutCompression);
-		float3 ODTColor = ACESOutputTransforms1000( FilmColor, (float3x3)WorkingColorSpace.ToAP0, AcesParams);
+		float3 ODTColor = ACESOutputTransform( FilmColor, (float3x3)WorkingColorSpace.ToAP0, AcesParams);
 
 		// Convert from AP1 to specified output gamut
 		ODTColor = mul( AP1_2_Output, ODTColor );
@@ -109,43 +109,16 @@ float4 DeviceEncodingOnlyCommonPS(
 		OutDeviceColor = LinearToST2084( ODTColor );
 	}
 	 
-	// ACES 2000nit transform with PQ/2084 encoding, user specified gamut 
-	else if( GetOutputDevice() == TONEMAPPER_OUTPUT_ACES2000nitST2084 )
-	{		
-		FACESTonemapParams AcesParams = ComputeACESTonemapParams(ACESMinMaxData, ACESMidData, ACESCoefsLow_0, ACESCoefsHigh_0, ACESCoefsLow_4, ACESCoefsHigh_4, ACESSceneColorMultiplier, ACESGamutCompression);
-		
-		// 2000 nit ODT
-		float3 ODTColor = ACESOutputTransforms2000( FilmColor, (float3x3)WorkingColorSpace.ToAP0, AcesParams);
-
-		// Convert from AP1 to specified output gamut
-		ODTColor = mul( AP1_2_Output, ODTColor );
-
-		// Apply conversion to ST-2084 (Dolby PQ)
-		OutDeviceColor = LinearToST2084( ODTColor );
-	}
-
-	// ACES 1000nit transform to linear ScRGB
-	else if(GetOutputDevice() == TONEMAPPER_OUTPUT_ACES1000nitScRGB)
-	{
-		FACESTonemapParams AcesParams = ComputeACESTonemapParams(ACESMinMaxData, ACESMidData, ACESCoefsLow_0, ACESCoefsHigh_0, ACESCoefsLow_4, ACESCoefsHigh_4, ACESSceneColorMultiplier, ACESGamutCompression);
-
-		// 1000 nit ODT
-		float3 ODTColor = ACESOutputTransforms1000( FilmColor, (float3x3)WorkingColorSpace.ToAP0, AcesParams);
-
-		// Apply conversion to ScRGB
-		OutDeviceColor = ST2084ToScRGB( LinearToST2084( ODTColor ), TONEMAPPER_OUTPUT_ACES1000nitScRGB, OutputMaxLuminance );
-	}
-
-	// ACES 2000nit transform to linear ScRGB
-	else if(GetOutputDevice() == TONEMAPPER_OUTPUT_ACES2000nitScRGB)
+	// ACES transform to linear ScRGB
+	else if (GetOutputDevice() == TONEMAPPER_OUTPUT_ACES1000nitScRGB || GetOutputDevice() == TONEMAPPER_OUTPUT_ACES2000nitScRGB)
 	{
 		FACESTonemapParams AcesParams = ComputeACESTonemapParams(ACESMinMaxData, ACESMidData, ACESCoefsLow_0, ACESCoefsHigh_0, ACESCoefsLow_4, ACESCoefsHigh_4, ACESSceneColorMultiplier, ACESGamutCompression);
 
-		// 2000 nit ODT
-		float3 ODTColor = ACESOutputTransforms2000( FilmColor, (float3x3)WorkingColorSpace.ToAP0, AcesParams);
+		// ACES ODT
+		float3 ODTColor = ACESOutputTransform(FilmColor, (float3x3) WorkingColorSpace.ToAP0, AcesParams);
 
 		// Apply conversion to ScRGB
-		OutDeviceColor = ST2084ToScRGB( LinearToST2084( ODTColor ), TONEMAPPER_OUTPUT_ACES2000nitScRGB, OutputMaxLuminance );
+		OutDeviceColor = ST2084ToScRGB( LinearToST2084( ODTColor ), GetOutputDevice(), OutputMaxLuminance );
 	}
 	
 	else if( GetOutputDevice() == TONEMAPPER_OUTPUT_LinearEXR)
diff --git a/Engine/Shaders/Private/TonemapCommon.ush b/Engine/Shaders/Private/TonemapCommon.ush
--- a/Engine/Shaders/Private/TonemapCommon.ush
+++ b/Engine/Shaders/Private/TonemapCommon.ush
@@ -401,11 +401,11 @@ FACESTonemapParams ComputeACESTonemapParams(float4 InACESMinMaxData, float4 InAC
 }
 
 //
-// ACES D65 1000 nit Output Transform - Forward
+// ACES D65 nit Output Transform - Forward
 //  Input is scene-referred linear values in the working space gamut
 //  Output is output-referred linear values in the AP1 gamut
 //
-float3 ACESOutputTransforms1000( float3 SceneReferredLinearColor, const float3x3 WCS_2_AP0, FACESTonemapParams InACESTonemapParams)
+float3 ACESOutputTransform( float3 SceneReferredLinearColor, const float3x3 WCS_2_AP0, FACESTonemapParams InACESTonemapParams)
 {
 	float3 aces = mul( WCS_2_AP0, SceneReferredLinearColor * InACESTonemapParams.SceneColorMultiplier);
 	float3 AcesCompressed = compressColor(aces);
@@ -419,37 +419,11 @@ float3 ACESOutputTransforms1000( float3 SceneReferredLinearColor, const float3x3
 	return OutputReferredLinearAP1Color;
 }
 
-float3 ACESOutputTransforms1000( float3 SceneReferredLinearsRGBColor, FACESTonemapParams InACESTonemapParams)
+float3 ACESOutputTransform( float3 SceneReferredLinearsRGBColor, FACESTonemapParams InACESTonemapParams)
 {
 	const float3x3 sRGB_2_AP0 = mul( XYZ_2_AP0_MAT, mul( D65_2_D60_CAT, sRGB_2_XYZ_MAT ) );
 
-	return ACESOutputTransforms1000( SceneReferredLinearsRGBColor, sRGB_2_AP0, InACESTonemapParams);
-}
-
-//
-// ACES D65 2000 nit Output Transform - Forward
-//  Input is scene-referred linear values in the working space gamut
-//  Output is output-referred linear values in the AP1 gamut
-//
-float3 ACESOutputTransforms2000( float3 SceneReferredLinearColor, const float3x3 WCS_2_AP0, FACESTonemapParams InACESTonemapParams)
-{
-	float3 aces = mul( WCS_2_AP0, SceneReferredLinearColor * InACESTonemapParams.SceneColorMultiplier);
-	float3 AcesCompressed = compressColor(aces);
-	aces = lerp(aces, AcesCompressed, InACESTonemapParams.GamutCompression);
-
-	const FColorSpace UE_ACES_AP1_DISPLAY_PRI = GetColorSpace(TONEMAPPER_GAMUT_ACEScg_D60);
-	const FColorSpace UE_ACES_AP1_LIMITING_PRI = GetColorSpace(TONEMAPPER_GAMUT_ACEScg_D60);
-	float3 OutputReferredLinearAP1Color = outputTransform(aces, InACESTonemapParams.AcesParams, UE_ACES_AP1_DISPLAY_PRI, UE_ACES_AP1_LIMITING_PRI, UE_ACES_EOTF,
-											 			  UE_ACES_SURROUND, UE_ACES_STRETCH_BLACK, UE_ACES_D60_SIM, UE_ACES_LEGAL_RANGE);
-
-	return OutputReferredLinearAP1Color;
-}
-
-float3 ACESOutputTransforms2000( float3 SceneReferredLinearsRGBColor, FACESTonemapParams InACESTonemapParams)
-{
-	const float3x3 sRGB_2_AP0 = mul( XYZ_2_AP0_MAT, mul( D65_2_D60_CAT, sRGB_2_XYZ_MAT ) );
-
-	return ACESOutputTransforms2000( SceneReferredLinearsRGBColor, sRGB_2_AP0, InACESTonemapParams);
+	return ACESOutputTransform( SceneReferredLinearsRGBColor, sRGB_2_AP0, InACESTonemapParams);
 }
 
 static const float3x3 GamutMappingIdentityMatrix = { 1.0, 0.0, 0.0, 0.0, 1.0, 0.0, 0.0, 0.0, 1.0 };
@@ -509,7 +483,7 @@ float3x3 OuputInverseGamutMappingMatrix( uint OutputGamut )
 	return GamutMappingMatrix;
 }
 
-float3 ST2084ToScRGB(float3 Color, uint OutputDevice, float MaxODTNits)
+float3 ST2084ToScRGB(float3 ColorAP1, uint OutputDevice, float MaxODTNits)
 {
 	// Nvidia HDR encoding - Remove PQ, convert to linear scRGB
 	const float3x3 AP1_2_sRGB = mul(XYZ_2_sRGB_MAT, mul(D60_2_D65_CAT, AP1_2_XYZ_MAT));
@@ -517,7 +491,7 @@ float3 ST2084ToScRGB(float3 Color, uint OutputDevice, float MaxODTNits)
 
 
 	float MinODTNits = UE_ACES_HDR_Y_MIN;
-	float3 OutColor = ST2084ToLinear(Color);
+	float3 OutColor = ST2084ToLinear(ColorAP1);
 		
 	OutColor = clamp(OutColor, MinODTNits, MaxODTNits);
 	OutColor.x = Y_2_linCV(OutColor.x, MaxODTNits, MinODTNits);
