[GPULM] Fix shader compile error when Substrate is enabled.

#rnx

diff --git a/Engine/Shaders/Private/PathTracing/PathTracingGPULightmassMaterialHitShader.usf b/Engine/Shaders/Private/PathTracing/PathTracingGPULightmassMaterialHitShader.usf
--- a/Engine/Shaders/Private/PathTracing/PathTracingGPULightmassMaterialHitShader.usf
+++ b/Engine/Shaders/Private/PathTracing/PathTracingGPULightmassMaterialHitShader.usf
@@ -94,12 +94,7 @@ RAY_TRACING_ENTRY_CLOSEST_HIT(GPULightmassMaterialCHS,
 		Settings,
 		SubstratePixelHeader,
 		SubstrateData,
-		SurfaceWorldNormal,
-		Coverage,
-		TransmittancePreCoverage,
-		BaseColorPostCoverage,
-		WorldNormal,
-		EmissiveLuminance);
+		SurfaceWorldNormal);
 
 	Coverage = Export.Coverage;
 	TransmittancePreCoverage = Export.TransmittancePreCoverage;
@@ -413,12 +408,7 @@ RAY_TRACING_ENTRY_ANY_HIT(GPULightmassMaterialAHS,
 		Settings,
 		SubstratePixelHeader,
 		SubstrateData,
-		SurfaceWorldNormal,
-		Coverage,
-		TransmittancePreCoverage,
-		BaseColorPostCoverage,
-		WorldNormal,
-		EmissiveLuminance);
+		SurfaceWorldNormal);
 
 	Coverage = Export.Coverage;
 	TransmittancePreCoverage = Export.TransmittancePreCoverage;
