Substrate - Fix for non unlit emissive being clamped in non pre exposed space.

#tests PC dx12 editor
#rnx

diff --git a/Engine/Shaders/Private/Substrate/Substrate.ush b/Engine/Shaders/Private/Substrate/Substrate.ush
--- a/Engine/Shaders/Private/Substrate/Substrate.ush
+++ b/Engine/Shaders/Private/Substrate/Substrate.ush
@@ -3063,7 +3063,9 @@ void FSubstrateBSDF::SubstrateSanitizeBSDF()
 	{
 		case SUBSTRATE_BSDF_TYPE_SLAB:
 		{
-			BSDF_SETEMISSIVE(this, clamp(BSDF_GETEMISSIVE(this), 0.0, Max10BitsFloat));
+			// We do not clamp emissive luminance upper bound. This is done later in the base pass within pre-exposed luminance space.
+			BSDF_SETEMISSIVE(this, max(BSDF_GETEMISSIVE(this), 0.0));
+
 			SLAB_DIFFUSEALBEDO(this) = SanitizeDiffuseAlbedo(SLAB_DIFFUSEALBEDO(this));
 			SLAB_F0(this) = SanitizeF0(SLAB_F0(this));
 			SLAB_F90(this) = saturate(SLAB_F90(this));
@@ -3100,7 +3102,9 @@ void FSubstrateBSDF::SubstrateSanitizeBSDF()
 		}
 		case SUBSTRATE_BSDF_TYPE_HAIR:
 		{
-			BSDF_SETEMISSIVE(this, clamp(BSDF_GETEMISSIVE(this), 0.0, Max10BitsFloat));
+			// We do not clamp emissive luminance upper bound. This is done later in the base pass within pre-exposed luminance space.
+			BSDF_SETEMISSIVE(this, max(BSDF_GETEMISSIVE(this), 0.0));
+
 			HAIR_BASECOLOR(this) = SanitizeDiffuseAlbedo(HAIR_BASECOLOR(this));
 			HAIR_SCATTER(this) = saturate(HAIR_SCATTER(this));
 			HAIR_ROUGHNESS(this) = SanitizeRoughness(HAIR_ROUGHNESS(this));
@@ -3110,7 +3114,9 @@ void FSubstrateBSDF::SubstrateSanitizeBSDF()
 		}
 		case SUBSTRATE_BSDF_TYPE_EYE:
 		{
-			BSDF_SETEMISSIVE(this, clamp(BSDF_GETEMISSIVE(this), 0.0, Max10BitsFloat));
+			// We do not clamp emissive luminance upper bound. This is done later in the base pass within pre-exposed luminance space.
+			BSDF_SETEMISSIVE(this, max(BSDF_GETEMISSIVE(this), 0.0));
+
 			EYE_DIFFUSEALBEDO(this) = SanitizeDiffuseAlbedo(EYE_DIFFUSEALBEDO(this));
 			EYE_ROUGHNESS(this) = SanitizeRoughness(EYE_ROUGHNESS(this));
 			EYE_IRISDISTANCE(this) = saturate(EYE_IRISDISTANCE(this));
@@ -3121,7 +3127,9 @@ void FSubstrateBSDF::SubstrateSanitizeBSDF()
 		}
 		case SUBSTRATE_BSDF_TYPE_SINGLELAYERWATER:
 		{
-			BSDF_SETEMISSIVE(this, clamp(BSDF_GETEMISSIVE(this), 0.0, Max10BitsFloat));
+			// We do not clamp emissive luminance upper bound. This is done later in the base pass within pre-exposed luminance space.
+			BSDF_SETEMISSIVE(this, max(BSDF_GETEMISSIVE(this), 0.0));
+
 			SLW_BASECOLOR(this) = SanitizeDiffuseAlbedo(SLW_BASECOLOR(this));
 			SLW_METALLIC(this) = saturate(SLW_METALLIC(this));
 			SLW_SPECULAR(this) = SanitizeF0(SLW_SPECULAR(this));
@@ -3137,7 +3145,7 @@ void FSubstrateBSDF::SubstrateSanitizeBSDF()
 		}
 		case SUBSTRATE_BSDF_TYPE_UNLIT:
 		{
-		#if MATERIAL_ALLOW_NEGATIVE_EMISSIVECOLOR == 0
+		#if MATERIAL_ALLOW_NEGATIVE_EMISSIVECOLOR == 0 // Only Unlit can have negative emissive if specified
 			// We do not clamp emissive luminance upper bound. This is done later in the base pass within pre-exposed luminance space.
 			BSDF_SETEMISSIVE(this, max(BSDF_GETEMISSIVE(this), 0.0));
 		#endif
