PostProcessTonemap - Fix corruption resulting from Sharpening on Mac when we send negative value to LinearToST2084

#tests Local testing on Mac

diff --git a/Engine/Shaders/Private/PostProcessTonemap.usf b/Engine/Shaders/Private/PostProcessTonemap.usf
--- a/Engine/Shaders/Private/PostProcessTonemap.usf
+++ b/Engine/Shaders/Private/PostProcessTonemap.usf
@@ -190,7 +190,8 @@ half3 ColorLookupTable( half3 LinearColor )
 	{
 		// ST2084 expects to receive linear values 0-10000 in nits.
 		// So the linear value must be multiplied by a scale factor to convert to nits.
-		LUTEncodedColor = LinearToST2084(LinearColor * LinearToNitsScale);
+		// We don't send negative values to LinearToST2084 as it will result in NaN because of pow.
+		LUTEncodedColor = LinearToST2084(max(0, LinearColor) * LinearToNitsScale);
 	}
 	else
 	{
