Ray Tracing: Fix InstanceLocalToWorld and InstanceWorldToLocal transforms not being propagated

The ray tracing codepath did not properly forward these matrices from the interpolants to the pixel parameters, leading to incorrect results in the path tracer and ray tracing contexts.

diff --git a/Engine/Shaders/Private/LocalVertexFactory.ush b/Engine/Shaders/Private/LocalVertexFactory.ush
--- a/Engine/Shaders/Private/LocalVertexFactory.ush
+++ b/Engine/Shaders/Private/LocalVertexFactory.ush
@@ -1320,6 +1320,15 @@ float4 VertexFactoryGetPreviousWorldPosition(FVertexFactoryInput Input, FVertexF
 		O.InterpolantsVSToPS.PerInstanceRandom = a.InterpolantsVSToPS.PerInstanceRandom;
 #endif
 
+#if HAS_INSTANCE_LOCAL_TO_WORLD_PS
+		O.InterpolantsVSToPS.InstanceLocalToWorld = a.InterpolantsVSToPS.InstanceLocalToWorld;
+#endif
+
+#if HAS_INSTANCE_WORLD_TO_LOCAL_PS
+		O.InterpolantsVSToPS.InstanceWorldToLocal = a.InterpolantsVSToPS.InstanceWorldToLocal;
+#endif
+
+
 		// Do we really need to interpolate TangentToWorld2 here? It should be replaced by the
 		// interpolated normal from 'whatever' interpolation scheme we're using
 	
