[HWRT] Fix incorrect instance flags when using ray tracing picker:
* Picker was incorrectly showing platform instance flags as opposed to engine instance flags.
* Add new function to translate back from platform to engine instance flags.

#rnx

diff --git a/Engine/Shaders/Private/RayTracing/RayTracingCommon.ush b/Engine/Shaders/Private/RayTracing/RayTracingCommon.ush
--- a/Engine/Shaders/Private/RayTracing/RayTracingCommon.ush
+++ b/Engine/Shaders/Private/RayTracing/RayTracingCommon.ush
@@ -552,6 +552,33 @@ uint TranslateRayTracingInstanceFlags(uint InFlags)
 	return Result;
 }
 
+uint TranslateRayTracingPlatformInstanceFlags(uint InPlatformFlags)
+{
+	uint Result = RAYTRACING_INSTANCE_FLAG_NONE;
+
+	if (InPlatformFlags & PLATFORM_RAYTRACING_INSTANCE_FLAG_TRIANGLE_CULL_DISABLE)
+	{
+		Result |= RAYTRACING_INSTANCE_FLAG_TRIANGLE_CULL_DISABLE;
+	}
+
+	if ((InPlatformFlags & PLATFORM_RAYTRACING_INSTANCE_FLAG_TRIANGLE_FRONT_COUNTERCLOCKWISE) == 0)
+	{
+		Result |= RAYTRACING_INSTANCE_FLAG_TRIANGLE_CULL_REVERSE;
+	}
+
+	if (InPlatformFlags & PLATFORM_RAYTRACING_INSTANCE_FLAG_FORCE_OPAQUE)
+	{
+		Result |= RAYTRACING_INSTANCE_FLAG_FORCE_OPAQUE;
+	}
+
+	if (InPlatformFlags & PLATFORM_RAYTRACING_INSTANCE_FLAG_FORCE_NON_OPAQUE)
+	{
+		Result |= RAYTRACING_INSTANCE_FLAG_FORCE_NON_OPAQUE;
+	}
+
+	return Result;
+}
+
 ///////////////////////////////////////////////////////////////////////////////////////////////////
 // Material payload
 
diff --git a/Engine/Shaders/Private/RayTracing/RayTracingDebugPicking.usf b/Engine/Shaders/Private/RayTracing/RayTracingDebugPicking.usf
--- a/Engine/Shaders/Private/RayTracing/RayTracingDebugPicking.usf
+++ b/Engine/Shaders/Private/RayTracing/RayTracingDebugPicking.usf
@@ -51,7 +51,7 @@ RAY_TRACING_ENTRY_RAYGEN(RayTracingDebugPickingRGS)
 		PickingFeedback.GeometryIndex = Payload.GeometryIndex;
 
 		PickingFeedback.InstanceId = GetRayTracingInstanceDescriptorInstanceId(InstanceDescriptor);
-		PickingFeedback.Flags = GetRayTracingInstanceDescriptorFlags(InstanceDescriptor);
+		PickingFeedback.Flags = TranslateRayTracingPlatformInstanceFlags(GetRayTracingInstanceDescriptorFlags(InstanceDescriptor));
 		PickingFeedback.Mask = GetRayTracingInstanceDescriptorMask(InstanceDescriptor);
 		PickingFeedback.InstanceContributionToHitGroupIndex = GetRayTracingInstanceDescriptorInstanceContributionToHitGroupIndex(InstanceDescriptor);
 	}
