Substrate - implementation of mobile deferred using conversion to legacy gbuffer for now.
It works for many materials and static/stationary/movable lighting modes.
This is a start, more validation and checks needs to be done.

#tests PC editor, mobile preview android and D3D
#rnx

diff --git a/Engine/Shaders/Private/MobileBasePassPixelShader.usf b/Engine/Shaders/Private/MobileBasePassPixelShader.usf
--- a/Engine/Shaders/Private/MobileBasePassPixelShader.usf
+++ b/Engine/Shaders/Private/MobileBasePassPixelShader.usf
@@ -128,7 +128,7 @@
 #endif
 
 // SUBSTRATE_TODO: Shared with BasePassPixelShader.usf
-#define MATERIAL_SUBSTRATE_OPAQUE_PRECOMPUTED_LIGHTING 0
+#define MATERIAL_SUBSTRATE_OPAQUE_PRECOMPUTED_LIGHTING (MATERIAL_IS_SUBSTRATE && SUBSTRATE_ENABLED && SUBSTRATE_OPAQUE_DEFERRED)
 
 #if SUBSTRATE_TRANSLUCENT_FORWARD || SUBSTRATE_FORWARD_SHADING || MATERIAL_SUBSTRATE_OPAQUE_PRECOMPUTED_LIGHTING || SUBSTRATE_MATERIAL_EXPORT_EXECUTED
 #include "/Engine/Private/Substrate/SubstrateEvaluation.ush"
@@ -136,6 +136,10 @@
 #if SUBSTRATE_TRANSLUCENT_FORWARD || SUBSTRATE_FORWARD_SHADING
 #include "/Engine/Private/Substrate/SubstrateMobileForwardLighting.ush"
 #endif
+
+#ifdef MATERIAL_SUBSTRATE_OPAQUE_PRECOMPUTED_LIGHTING
+	#define SUBSTRATE_MATERIAL_EXPORT_REQUESTED 1
+#endif
 #if MATERIAL_SUBSTRATE_OPAQUE_PRECOMPUTED_LIGHTING || SUBSTRATE_MATERIAL_EXPORT_EXECUTED
 #include "/Engine/Private/Substrate/SubstrateExport.ush"
 #endif
@@ -683,7 +687,208 @@ void Main(
 	half3 CameraVector = -MaterialParameters.CameraVector;
 	float2 ScreenUV = ScreenPositionToBufferUV(MaterialParameters.ScreenPosition);
 
-#if (SUBSTRATE_TRANSLUCENT_FORWARD || SUBSTRATE_FORWARD_SHADING)
+
+#if MATERIAL_SUBSTRATE_OPAQUE_PRECOMPUTED_LIGHTING
+
+#if SUBSTRATE_INLINE_SINGLELAYERWATER==0
+	// Need to reset color to make sure Substrate material are only lit using Substrate lighting passes.
+	// Except for water doing some specialised and simplified lighting during the base pass
+	Color = 0;
+#endif
+
+	{
+		SubstratePixelHeader.IrradianceAO									= InitIrradianceAndOcclusion();
+		SubstratePixelHeader.IrradianceAO.MaterialAO						= MaterialAO;	// This is the material AO, it will be converted to multibounce AO using ShadingOcclusion.SpecOcclusion in the PackSubstrateOut function, also accounting for bent normal.
+		SubstratePixelHeader.IrradianceAO.IndirectIrradiance				= 0;			// SUBSTRATE_TODO
+		SubstratePixelHeader.IrradianceAO.DiffuseIndirectSampleOcclusion	= 0;			// SUBSTRATE_TODO GetDiffuseIndirectSampleOcclusion(SubstratePixelHeader.SharedLocalBases, MaterialParameters.CameraVector, SvPosition.xy, SubstratePixelHeader.IrradianceAO.MaterialAO);
+	}
+
+	// We only rely on this GBuffer structure for the write out to legacy render target such as velocity or precomputed shadow factors
+	FGBufferData GBuffer = (FGBufferData)0;
+	GBuffer.PrecomputedShadowFactors = GetPrimaryPrecomputedShadowMask(LightmapVTPageTableResult, Interpolants, MaterialParameters);
+
+#if OUTPUT_PIXEL_DEPTH_OFFSET && SUBSTRATE_INLINE_SINGLELAYERWATER==0 && 0	// We use the legacy GBuffer format for now, so this is disabled.
+	// When in deferred, opaque materials with pixel depth offset must execute a custom depth test in order to avoid conflicting UAV writes.
+	// This is however not done for single layer water as the depth will never match. Note: that can lead overlapping SLW polys to cause flicker on screen.
+	if (ManualDepthTestEqual(SvPosition, OutDepth))
+#endif
+	{
+		uint2 PixelPos = uint2(SvPosition.xy);
+		FSubstrateSubsurfaceData SSSData = (FSubstrateSubsurfaceData)0;
+		FSubstrateTopLayerData TopLayerData = (FSubstrateTopLayerData)0;
+		FSubstrateOpaqueRoughRefractionData OpaqueRoughRefractionData = (FSubstrateOpaqueRoughRefractionData)0;
+		FSubstrateIntegrationSettings Settings = InitSubstrateIntegrationSettings(FULLY_ROUGH, SubstrateStruct.bRoughDiffuse, SubstrateStruct.PeelLayersAboveDepth, SubstrateStruct.bRoughnessTracking);
+	#if SUBSTRATE_ADVANCED_DEBUG_ENABLED
+		Settings.SliceStoringDebugSubstrateTreeData = all(uint2(float2(View.CursorPosition) * View.ViewResolutionFraction) == PixelPos) ? SubstrateStruct.SliceStoringDebugSubstrateTreeDataWithoutMRT : Settings.SliceStoringDebugSubstrateTreeData;
+	#endif
+
+		const float3 SurfaceWorldNormal = MaterialParameters.TangentToWorld[2].xyz;
+		FExportResult Export = SubstrateMaterialExportOut(
+			Settings,
+			SubstratePixelHeader,
+			SubstrateData,
+			SurfaceWorldNormal);
+
+		GBuffer.GBufferAO = MaterialAO;
+		GBuffer.Depth = MaterialParameters.ScreenPosition.w;
+
+		// Patch material parameter to not have to change the function for now
+		MaterialParameters.WorldNormal = Export.WorldNormal;
+		MaterialParameters.WorldTangent = Export.Tangent;
+	#if USE_WORLDVERTEXNORMAL_CENTER_INTERPOLATION
+		MaterialParameters.WorldVertexNormal_Center = Export.WorldNormal;	// SUBSTRATE_TODO: how to evluate it for center?
+	#endif
+		MaterialParameters.TangentToWorld[2].xyz = Export.Tangent;
+
+		SetGBufferForShadingModel(
+			GBuffer,
+			MaterialParameters,
+			Export.Coverage,		// Opacity
+			Export.BaseColor,
+			Export.Metallic,
+			Export.Specular,
+			Export.Roughness,
+			Export.Anisotropy,
+			Export.SubsurfaceColor,
+			0,						// SUBSTRATE_TODO SubsurfaceProfile
+			0.0f,					// Dither
+			Export.ShadingModel
+		);
+
+		GBuffer.StoredBaseColor = GBuffer.BaseColor;
+		GBuffer.StoredMetallic = GBuffer.Metallic;
+		GBuffer.StoredSpecular = GBuffer.Specular;
+
+	#if MATERIAL_SHADINGMODEL_EYE
+		if( GBuffer.ShadingModelID == SHADINGMODELID_EYE )
+		{
+			GBuffer.Metallic = 0.0;
+			#if IRIS_NORMAL
+			GBuffer.Specular = 0.25;
+			#endif
+		}
+	#endif
+
+	#if MOBILE_QL_FORCE_NONMETAL
+		GBuffer.DiffuseColor = GBuffer.BaseColor;
+		GBuffer.SpecularColor = DielectricSpecularToF0(GBuffer.SpecularColor);
+	#else
+		GBuffer.SpecularColor = ComputeF0(GBuffer.Specular, GBuffer.BaseColor, GBuffer.Metallic);
+		GBuffer.DiffuseColor = GBuffer.BaseColor - GBuffer.BaseColor * GBuffer.Metallic;
+	#endif
+
+	#if FULLY_ROUGH
+		EnvBRDFApproxFullyRough(GBuffer.DiffuseColor, GBuffer.SpecularColor);
+	#endif
+
+	#if MOBILE_EMULATION && !MOBILE_DEFERRED_LIGHTING
+		{
+			// this feature is only needed for development/editor - we can compile it out for a shipping build (see r.CompileShadersForDevelopment cvar help)
+			GBuffer.DiffuseColor = GBuffer.DiffuseColor * ResolvedView.DiffuseOverrideParameter.w + ResolvedView.DiffuseOverrideParameter.xyz;
+			GBuffer.SpecularColor = GBuffer.SpecularColor * ResolvedView.SpecularOverrideParameter.w + ResolvedView.SpecularOverrideParameter.xyz;
+		}
+	#endif
+
+	#if MATERIAL_NORMAL_CURVATURE_TO_ROUGHNESS
+		// Curvature-to-roughness uses derivatives of the WorldVertexNormal, which is incompatible with centroid interpolation because
+		// the samples are not uniformly distributed. Therefore we use WorldVertexNormal_Center which is guaranteed to be center interpolated.
+		#if USE_WORLDVERTEXNORMAL_CENTER_INTERPOLATION
+		float GeometricAARoughness = NormalCurvatureToRoughness(MaterialParameters.WorldVertexNormal_Center);
+		#else
+		float GeometricAARoughness = NormalCurvatureToRoughness(MaterialParameters.TangentToWorld[2].xyz);
+		#endif
+		GBuffer.Roughness = max(GBuffer.Roughness, GeometricAARoughness);
+
+		#if MATERIAL_SHADINGMODEL_CLEAR_COAT
+		if (GBuffer.ShadingModelID == SHADINGMODELID_CLEAR_COAT)
+		{
+			GBuffer.CustomData.y = max(GBuffer.CustomData.y, GeometricAARoughness);
+		}
+		#endif
+	#endif
+
+
+		half3 InputBentNormal = MaterialParameters.WorldNormal;
+		// Clear Coat Bottom Normal
+		#if MATERIAL_SHADINGMODEL_CLEAR_COAT
+		BRANCH if (GBuffer.ShadingModelID == SHADINGMODELID_CLEAR_COAT && CLEAR_COAT_BOTTOM_NORMAL)
+		{
+			const float2 oct1 = ((float2(GBuffer.CustomData.a, GBuffer.CustomData.z) * 4) - (512.0 / 255.0)) + UnitVectorToOctahedron(GBuffer.WorldNormal);
+			InputBentNormal = OctahedronToUnitVector(oct1);
+		}
+		#endif
+		const float3 WorldBentNormal0 = GetWorldBentNormalZero(MaterialParameters);
+		const FShadingOcclusion ShadingOcclusion = ApplyBentNormal(MaterialParameters.CameraVector, InputBentNormal, WorldBentNormal0, GBuffer.Roughness, GBuffer.GBufferAO);
+		GBuffer.GBufferAO = AOMultiBounce(Luminance(GBuffer.SpecularColor), ShadingOcclusion.SpecOcclusion).g;
+
+		half3 SubsurfaceColor = 0;// SUBSTRATE_TODO
+		half3 WaterDiffuseIndirectLuminance = 0;
+		half IndirectIrradiance = 0;
+		half3 DiffuseIndirectLighting = 0;
+		half3 SubsurfaceIndirectLighting = 0;
+
+		half3 DiffuseColorForIndirect = GBuffer.DiffuseColor;
+		half3 DiffuseDir = ShadingOcclusion.BentNormal;
+
+	#if MATERIAL_SHADINGMODEL_SUBSURFACE || MATERIAL_SHADINGMODEL_PREINTEGRATED_SKIN
+		if (GBuffer.ShadingModelID == SHADINGMODELID_SUBSURFACE || GBuffer.ShadingModelID == SHADINGMODELID_PREINTEGRATED_SKIN)
+		{
+			// Add subsurface energy to diffuse
+			//@todo - better subsurface handling for these shading models with skylight and precomputed GI
+			DiffuseColorForIndirect += SubsurfaceColor;
+		}
+	#endif
+
+	#if MATERIAL_SHADINGMODEL_CLOTH
+		if (GBuffer.ShadingModelID == SHADINGMODELID_CLOTH)
+		{
+			DiffuseColorForIndirect += SubsurfaceColor * saturate(GetMaterialCustomData0(MaterialParameters));
+		}
+	#endif
+
+	#if MATERIAL_SHADINGMODEL_HAIR
+		if (GBuffer.ShadingModelID == SHADINGMODELID_HAIR)
+		{
+			const half3 N = MaterialParameters.WorldNormal;
+			const half3 V = MaterialParameters.CameraVector;
+			DiffuseColorForIndirect = EvaluateEnvHair(GBuffer, V, N, DiffuseDir);
+		}
+	#endif
+
+		//half NoL = max(0, dot(GBuffer.WorldNormal, MobileDirectionalLight.DirectionalLightDirectionAndShadowTransition.xyz));
+
+		const bool bEvaluateBackface = GetShadingModelRequiresBackfaceLighting(GBuffer.ShadingModelID);
+
+		GetPrecomputedIndirectLightingAndSkyLight(LightmapVTPageTableResult, Interpolants, bEvaluateBackface, DiffuseDir, ResolvedView.SkyLightColor.rgb, DiffuseIndirectLighting, SubsurfaceIndirectLighting, IndirectIrradiance, WaterDiffuseIndirectLuminance);
+
+		// Apply MaterialAO since we don't have the DiffuseIndirectComposite pass on mobile deferred.
+		IndirectIrradiance *= GBuffer.GBufferAO;
+
+		half3 DiffuseColor = (DiffuseIndirectLighting * DiffuseColorForIndirect + SubsurfaceIndirectLighting * SubsurfaceColor) * AOMultiBounce(GBuffer.BaseColor, ShadingOcclusion.DiffOcclusion);
+
+	#if MATERIAL_SHADINGMODEL_SINGLELAYERWATER
+		// Fade out diffuse as this will be handled by the single scattering lighting in water material.
+		// We do this after the just above GetPrecomputedIndirectLightingAndSkyLight to keep ambiant lighting avialable.
+		// We also keep the SpecularColor for sun/water interactions.
+		GBuffer.DiffuseColor *= BaseMaterialCoverageOverWater;
+		DiffuseColor *= BaseMaterialCoverageOverWater;
+		DiffuseColorForIndirect *= BaseMaterialCoverageOverWater;
+	#endif
+
+		// This is deferred, should od the same for Substrate GBuffer
+
+		FLightAccumulator DirectLighting = (FLightAccumulator)0;
+		LightAccumulator_AddSplit(DirectLighting, DiffuseColor, 0.0f, DiffuseColor, 1.0f, false);
+
+		GBuffer.IndirectIrradiance = IndirectIrradiance;
+		MobileEncodeGBuffer(GBuffer, OutGBufferA, OutGBufferB, OutGBufferC, OutGBufferD);
+
+
+		Color = DirectLighting.TotalLight;
+	}
+#endif // MATERIAL_SUBSTRATE_OPAQUE_PRECOMPUTED_LIGHTING
+
+#if (SUBSTRATE_TRANSLUCENT_FORWARD || SUBSTRATE_FORWARD_SHADING) && !MATERIAL_SUBSTRATE_OPAQUE_PRECOMPUTED_LIGHTING
 	#define SUBSTRATE_DEBUG_ENABLE_LIGHTING 1
 	#if SUBSTRATE_DEBUG_ENABLE_LIGHTING
 	if (SubstratePixelHeader.ClosureCount > 0)
diff --git a/Engine/Shaders/Private/MobileDeferredShading.usf b/Engine/Shaders/Private/MobileDeferredShading.usf
--- a/Engine/Shaders/Private/MobileDeferredShading.usf
+++ b/Engine/Shaders/Private/MobileDeferredShading.usf
@@ -2,6 +2,11 @@
 
 #define MOBILE_DEFERRED_LIGHTING 1
 
+#if SUBSTRATE_ENABLED
+// This is needed when lights are rendered deferred on mobile with Substrate because we use the legacy GBuffer representation and code process.
+#define ALLOW_LOCAL_LIGHT_DISTANCE_ATTENUATION 1
+#endif
+
 #include "Common.ush"
 #include "SHCommon.ush"
 
diff --git a/Engine/Shaders/Private/Substrate/SubstrateExport.ush b/Engine/Shaders/Private/Substrate/SubstrateExport.ush
--- a/Engine/Shaders/Private/Substrate/SubstrateExport.ush
+++ b/Engine/Shaders/Private/Substrate/SubstrateExport.ush
@@ -2,7 +2,8 @@
 
 #pragma once
 
-#if SUBSTRATE_ENABLED
+// Do not include that on mobile since FMobileBasePassInterpolantsVSToPS is not the same structure as FBasePassInterpolantsVSToPS
+#if SUBSTRATE_ENABLED && !defined(FMobileBasePassInterpolantsVSToPS)
 
 // This is used to evaluate some precomputed lighting data while we pack out Substrate material.
 // It is also used to remove reference to unwanted resourced when not needed.
@@ -928,11 +929,13 @@ void PackSubstrateOut(
 
 }
 
+#endif // SUBSTRATE_ENABLED && !defined(FMobileBasePassInterpolantsVSToPS)
+
 #ifndef SUBSTRATE_GPU_LIGHTMASS
 #define SUBSTRATE_GPU_LIGHTMASS 0
 #endif
 
-#if SUBSTRATE_MATERIAL_EXPORT_TYPE > 0 || SUBSTRATE_GPU_LIGHTMASS > 0
+#if SUBSTRATE_ENABLED && (SUBSTRATE_MATERIAL_EXPORT_TYPE > 0 || SUBSTRATE_GPU_LIGHTMASS > 0 || defined(SUBSTRATE_MATERIAL_EXPORT_REQUESTED))
 
 struct FExportResult
 {
@@ -1007,6 +1010,5 @@ FExportResult SubstrateMaterialExportOut(
 	return Out;
 }
 
-#endif // SUBSTRATE_MATERIAL_EXPORT_TYPE
+#endif // SUBSTRATE_ENABLED, SUBSTRATE_MATERIAL_EXPORT_TYPE, etc.
 
-#endif // SUBSTRATE_ENABLED
diff --git a/Engine/Source/Runtime/RenderCore/Private/GBufferInfo.cpp b/Engine/Source/Runtime/RenderCore/Private/GBufferInfo.cpp
--- a/Engine/Source/Runtime/RenderCore/Private/GBufferInfo.cpp
+++ b/Engine/Source/Runtime/RenderCore/Private/GBufferInfo.cpp
@@ -52,6 +52,10 @@ static bool IsGBufferItemEqual(const FGBufferItem& Lhs, const FGBufferItem& Rhs)
 		Lhs.Checker == Rhs.Checker;
 }
 
+static bool IsMobileDeferred(const FGBufferParams& Params)
+{
+	return IsMobilePlatform(Params.ShaderPlatform) && IsMobileDeferredShadingEnabled(Params.ShaderPlatform);
+}
 
 bool RENDERCORE_API IsGBufferInfoEqual(const FGBufferInfo& Lhs, const FGBufferInfo& Rhs)
 {
@@ -259,7 +263,7 @@ FGBufferInfo RENDERCORE_API FetchLegacyGBufferInfo(const FGBufferParams& Params)
 	int32 TargetSeparatedMainDirLight = -1;
 
 	// Substrate outputs material data through UAV. Only SceneColor, PrecalcShadow & Velocity data are still emitted through RenderTargets
-	if (Substrate::IsSubstrateEnabled())
+	if (Substrate::IsSubstrateEnabled() && !IsMobileDeferred(Params))
 	{
 		TargetGBufferA = -1;
 		TargetGBufferB = -1;
@@ -585,7 +589,7 @@ FGBufferInfo RENDERCORE_API FetchMobileGBufferInfo(const FGBufferParams& Params)
 FGBufferInfo RENDERCORE_API FetchFullGBufferInfo(const FGBufferParams& Params)
 {
 	// For now, we are only doing legacy. But next, we will have a switch between the old and new formats.
-	if (IsMobilePlatform(Params.ShaderPlatform) && IsMobileDeferredShadingEnabled(Params.ShaderPlatform))
+	if (IsMobileDeferred(Params))
 	{
 		return FetchMobileGBufferInfo(Params);
 	}
