Substrate - fixed export naming.

#tests PC editor dx12 and mobile preview
#rnx

diff --git a/Engine/Shaders/Private/BasePassPixelShader.usf b/Engine/Shaders/Private/BasePassPixelShader.usf
--- a/Engine/Shaders/Private/BasePassPixelShader.usf
+++ b/Engine/Shaders/Private/BasePassPixelShader.usf
@@ -1679,9 +1679,9 @@ void FPixelShaderInOut_MainPS(
 		#elif SUBSTRATE_MATERIAL_EXPORT_TYPE == SUBSTRATE_MATERIAL_EXPORT_SUBSURFACECOLOR
 			SubstrateMaterialExport = float4(Export.SubsurfaceColor, 1.0f);
 		#elif SUBSTRATE_MATERIAL_EXPORT_TYPE == SUBSTRATE_MATERIAL_EXPORT_TANGENT
-			SubstrateMaterialExport = float4(Export.Tangent * 0.5 + 0.5, 1.0f);
+			SubstrateMaterialExport = float4(Export.WorldTangent * 0.5 + 0.5, 1.0f);
 		#elif SUBSTRATE_MATERIAL_EXPORT_TYPE == SUBSTRATE_MATERIAL_EXPORT_SHADINGMODEL
-			SubstrateMaterialExport = float4(Export.ShadingModel.xxx, 1.0f);
+			SubstrateMaterialExport = float4(Export.ShadingModelID.xxx, 1.0f);
 		#elif SUBSTRATE_MATERIAL_EXPORT_TYPE == SUBSTRATE_MATERIAL_EXPORT_MATERIAL_PREVIEW
 			SubstrateMaterialExport = float4(PixelMaterialInputs.EmissiveColor, 1.0f);
 		#endif
diff --git a/Engine/Shaders/Private/MobileBasePassPixelShader.usf b/Engine/Shaders/Private/MobileBasePassPixelShader.usf
--- a/Engine/Shaders/Private/MobileBasePassPixelShader.usf
+++ b/Engine/Shaders/Private/MobileBasePassPixelShader.usf
@@ -734,11 +734,11 @@ void Main(
 
 		// Patch material parameter to not have to change the function for now
 		MaterialParameters.WorldNormal = Export.WorldNormal;
-		MaterialParameters.WorldTangent = Export.Tangent;
+		MaterialParameters.WorldTangent = Export.WorldTangent;
 	#if USE_WORLDVERTEXNORMAL_CENTER_INTERPOLATION
 		MaterialParameters.WorldVertexNormal_Center = Export.WorldNormal;	// SUBSTRATE_TODO: how to evluate it for center?
 	#endif
-		MaterialParameters.TangentToWorld[2].xyz = Export.Tangent;
+		MaterialParameters.TangentToWorld[2].xyz = Export.WorldTangent;
 
 		SetGBufferForShadingModel(
 			GBuffer,
@@ -752,7 +752,7 @@ void Main(
 			Export.SubsurfaceColor,
 			0,						// SUBSTRATE_TODO SubsurfaceProfile
 			0.0f,					// Dither
-			Export.ShadingModel
+			Export.ShadingModelID
 		);
 
 		GBuffer.StoredBaseColor = GBuffer.BaseColor;
diff --git a/Engine/Shaders/Private/Substrate/SubstrateExport.ush b/Engine/Shaders/Private/Substrate/SubstrateExport.ush
--- a/Engine/Shaders/Private/Substrate/SubstrateExport.ush
+++ b/Engine/Shaders/Private/Substrate/SubstrateExport.ush
@@ -952,8 +952,8 @@ struct FExportResult
 	float3 CustomData0;
 	float3 CustomData1;
 	float3 SubsurfaceColor;
-	float3 Tangent;
-	float  ShadingModel;
+	float3 WorldTangent;
+	float  ShadingModelID;
 };
 
 // This is used to export very simple data for lightmaps, material baking or preview on editor node.
@@ -998,8 +998,8 @@ FExportResult SubstrateMaterialExportOut(
 				Out.CustomData0				= 0;											// SUBSTRATE_TODO not tested since not exposed in the only use case so far: material baking.
 				Out.CustomData1				= 0;											// idem
 				Out.SubsurfaceColor			= SubstrateGetBSDFSubSurfaceColor(CurrentBSDF);	// idem
-				Out.Tangent					= T;
-				Out.ShadingModel			= SubstrateGetLegacyShadingModels(CurrentBSDF);	// idem
+				Out.WorldTangent			= T;
+				Out.ShadingModelID			= SubstrateGetLegacyShadingModels(CurrentBSDF);	// idem
 			}
 		#undef CurrentBSDF
 		}
