Fix bug in mesh paint seam painting when not writing to all color channels.
Brush mask needed to be calculated without colors, only writing white to pixels affected by the brush.
Dilation pass then needed correct logic to ignore pixels not in the brush mask and in the seam mask.
Also made the masks single channel 8 bit to save a bit on memory/bandwidth.

diff --git a/Engine/Plugins/MeshPainting/Source/MeshPaintingToolset/Private/MeshTexturePaintingTool.cpp b/Engine/Plugins/MeshPainting/Source/MeshPaintingToolset/Private/MeshTexturePaintingTool.cpp
--- a/Engine/Plugins/MeshPainting/Source/MeshPaintingToolset/Private/MeshTexturePaintingTool.cpp
+++ b/Engine/Plugins/MeshPainting/Source/MeshPaintingToolset/Private/MeshTexturePaintingTool.cpp
@@ -809,7 +809,7 @@ void UMeshTexturePaintingTool::StartPaintingTexture(UMeshComponent* InMeshCompon
 					const bool bForceLinearGamma = true;
 					TextureData->BrushMaskRenderTargetTexture->ClearColor = FLinearColor::Black;
 					TextureData->BrushMaskRenderTargetTexture->bNeedsTwoCopies = true;
-					TextureData->BrushMaskRenderTargetTexture->InitCustomFormat(BrushTargetTextureWidth, BrushTargetTextureHeight, PF_B8G8R8A8, bForceLinearGamma);
+					TextureData->BrushMaskRenderTargetTexture->InitCustomFormat(BrushTargetTextureWidth, BrushTargetTextureHeight, PF_G8, bForceLinearGamma);
 					TextureData->BrushMaskRenderTargetTexture->UpdateResourceImmediate();
 					TextureData->BrushMaskRenderTargetTexture->AddressX = TextureData->PaintRenderTargetTexture->AddressX;
 					TextureData->BrushMaskRenderTargetTexture->AddressY = TextureData->PaintRenderTargetTexture->AddressY;
@@ -827,7 +827,7 @@ void UMeshTexturePaintingTool::StartPaintingTexture(UMeshComponent* InMeshCompon
 					const bool bForceLinearGamma = true;
 					TextureData->SeamMaskRenderTargetTexture->ClearColor = FLinearColor::Black;
 					TextureData->SeamMaskRenderTargetTexture->bNeedsTwoCopies = true;
-					TextureData->SeamMaskRenderTargetTexture->InitCustomFormat(BrushTargetTextureWidth, BrushTargetTextureHeight, PF_B8G8R8A8, bForceLinearGamma);
+					TextureData->SeamMaskRenderTargetTexture->InitCustomFormat(BrushTargetTextureWidth, BrushTargetTextureHeight, PF_G8, bForceLinearGamma);
 					TextureData->SeamMaskRenderTargetTexture->UpdateResourceImmediate();
 					TextureData->SeamMaskRenderTargetTexture->AddressX = TextureData->PaintRenderTargetTexture->AddressX;
 					TextureData->SeamMaskRenderTargetTexture->AddressY = TextureData->PaintRenderTargetTexture->AddressY;
diff --git a/Engine/Shaders/Private/MeshPaintPixelShader.usf b/Engine/Shaders/Private/MeshPaintPixelShader.usf
--- a/Engine/Shaders/Private/MeshPaintPixelShader.usf
+++ b/Engine/Shaders/Private/MeshPaintPixelShader.usf
@@ -29,6 +29,8 @@ float c_BrushStrength;
 float4 c_BrushColor;
 float4 c_ChannelFlags;
 
+float c_GenerateMaskFlag;
+
 // @todo MeshPaint: Remove this?
 float c_Gamma;
 
@@ -49,8 +51,6 @@ void Main( float4 InPosition : SV_POSITION,
 	float BrushDepth = c_BrushMetrics.z;
 	float BrushDepthFalloffRange = c_BrushMetrics.w;
 
-
-
 	float4x4 WorldToBrushMatrix = c_WorldToBrushMatrix;
 
 	// Project the pixel into the plane of the brush
@@ -58,10 +58,13 @@ void Main( float4 InPosition : SV_POSITION,
 	float3 BrushSpaceVertexPosition = mul(float4( WorldSpaceVertexPosition, 1.0f ), WorldToBrushMatrix).xyz;
 	float2 BrushSpaceVertexPosition2D = BrushSpaceVertexPosition.xy;
 
+	// All paths will compute amount of paint to apply
+	float PaintAmount = 0;
 
 	// Start by using the old color
 	float4 NewColor = OldColor;
-	#if MESH_PAINT_USE_FILL_BUCKET
+
+#if MESH_PAINT_USE_FILL_BUCKET
 
 #if MESH_PAINT_USE_PAINTBRUSH
 	// Flood with a texture selected
@@ -91,8 +94,7 @@ void Main( float4 InPosition : SV_POSITION,
 	LocalTextureCoordinate.xy /= 2.0f;
 	LocalTextureCoordinate.xy += 0.5f;
 	
-	
-	float PaintAmount = 1.0f;
+	PaintAmount = 1.0f;
 	PaintAmount *= c_BrushStrength;
 	
 	// Determine an alpha blended color which is somewhere between our old color and our new color
@@ -106,7 +108,7 @@ void Main( float4 InPosition : SV_POSITION,
 #else
 
 	// Simple flood
-	float PaintAmount = 1.0f;
+	PaintAmount = 1.0f;
 	PaintAmount *= c_BrushStrength;
 
 	float4 BrushColor = c_BrushColor;
@@ -146,7 +148,7 @@ void Main( float4 InPosition : SV_POSITION,
  		if( VertexDepthToBrush <= BrushDepth )
 		{
 			// Compute amount of paint to apply
-			float PaintAmount = 1.0f;
+			PaintAmount = 1.0f;
 
 			// For distance falloff, keep it a square instead of a circle
 			float DistanceToVertex2D = max(abs(BrushSpaceVertexPosition2D.x), abs(BrushSpaceVertexPosition2D.y));
@@ -196,7 +198,7 @@ void Main( float4 InPosition : SV_POSITION,
  		if( VertexDepthToBrush <= BrushDepth )
 		{
 			// Compute amount of paint to apply
-			float PaintAmount = 1.0f;
+			PaintAmount = 1.0f;
 
 			// Apply radial-based falloff
 			{
@@ -233,6 +235,12 @@ void Main( float4 InPosition : SV_POSITION,
 
 	OutColor = NewColor;
 
+	// Mask will write out 1 for anywhere painted.
+	if ( c_GenerateMaskFlag == 1.0 )
+	{
+		OutColor.rgb = PaintAmount > 0 ? 1 : 0;
+	}
+
 	// Gamma correct the output color.
 	if( c_Gamma != 1.0 )
 	{
diff --git a/Engine/Shaders/Private/meshpaintdilatepixelshader.usf b/Engine/Shaders/Private/meshpaintdilatepixelshader.usf
--- a/Engine/Shaders/Private/meshpaintdilatepixelshader.usf
+++ b/Engine/Shaders/Private/meshpaintdilatepixelshader.usf
@@ -7,23 +7,22 @@
 #include "Common.ush"
 
 /** BrushRenderTargetTexture - Texture that has our accumulated paint data */
-Texture2D Texture0;
+Texture2D<float4> Texture0;
 SamplerState Texture0Sampler;
 
 /** SeamMaskRenderTargetTexture - Texture that stores the UV seam mask */
-Texture2D Texture1;
+Texture2D<float> Texture1;
 SamplerState Texture1Sampler;
 
-/** BrushMaskRenderTargetTexture - Texture that stores info about the area we incrementally painted */
-Texture2D Texture2;
+/** BrushMaskRenderTargetTexture - Texture that stores mask of where we incrementally painted */
+Texture2D<float> Texture2;
 SamplerState Texture2Sampler;
 
+/** Width of a pixel in UV space. */
 float WidthPixelOffset;
+/** Height of a pixel in UV space. */
 float HeightPixelOffset;
 
-float Gamma;
-
-
 void Main(
 	in float2 InTextureCoordinates : TEXCOORD0,
 	in float3 InWorldSpaceVertexPosition : TEXCOORD1,
@@ -31,48 +30,45 @@ void Main(
 	out float4 OutColor : SV_Target0 
 	)
 {
-
-	// Texture0 will be the brush texture
-	// Texture1 will be the seam mask
-	// Texture2 will be the brush mask
-
-	// First sample the source values from brush texture and seam mask
-	float4 OldColor0 = Texture2DSample( Texture0, Texture0Sampler, InTextureCoordinates );
-	float4 OldColor1 = Texture2DSample( Texture1, Texture1Sampler, InTextureCoordinates );
-
-	float4 CenterColor = OldColor0;
-
-	float2 OffsetConst = float2( WidthPixelOffset, HeightPixelOffset );
-
-	float2 OffsetIdx;
-
-	// Arbitrary large number that we will compare to against the squared sample distance.
-	//   The largest we will hit in 25 samples should be 8 (distance squared = x*x+y*y)
-	int LastValidDistance = 10;
-
-	// Sample pixel data from a 5x5 grid around our current pixel
-	for( OffsetIdx.y = 2; OffsetIdx.y >= -2; OffsetIdx.y-- )
+	// Sample the initial value from the brush texture.
+	OutColor = Texture0.Sample( Texture0Sampler, InTextureCoordinates );
+	
+	// Only need to change the value if the current pixel falls within the white area of the texture seam mask.
+	float SeamMask = Texture1.Sample( Texture1Sampler, InTextureCoordinates );
+	if ( SeamMask != 0 )
 	{
-		for( OffsetIdx.x = -2; OffsetIdx.x <= 2; OffsetIdx.x++)
-		{
-			// Compute offset coordinates for the pixel and sample the color
-			float2 Coord = InTextureCoordinates + (OffsetIdx * OffsetConst);
-			float4 SampleColor = Texture2DSample( Texture0, Texture0Sampler, Coord );
-
-			// A sample is only valid when it falls in the white area of the brush mask and the current
-			//   pixel we are processing falls within the white area of the texture seam mask.
-			bool IsValid = all( Texture2DSample( Texture2, Texture2Sampler, Coord ).rgb != float3(0, 0, 0) ) && all( OldColor1.rgb != float3(0, 0, 0) );
+		float2 OffsetConst = float2( WidthPixelOffset, HeightPixelOffset );
 
-			// Calculate the squared sample distance from the center
-			int SampleDist = OffsetIdx.x * OffsetIdx.x + OffsetIdx.y * OffsetIdx.y;
+		// Arbitrary large number that we will compare to against the squared sample distance.
+		// The largest we will hit in 25 samples should be 8 (distance squared = x*x+y*y)
+		int LastValidDistance = 10;
 
-			// @todo MeshPaint:If needed, we can improve this by averaging the pixel colors at the same distance
-			if( IsValid && SampleDist < LastValidDistance)
+		// Sample pixel data from a 5x5 grid around our current pixel.
+		float2 OffsetIdx;
+		for( OffsetIdx.y = 2; OffsetIdx.y >= -2; OffsetIdx.y-- )
+		{
+			for( OffsetIdx.x = -2; OffsetIdx.x <= 2; OffsetIdx.x++ )
 			{
-				CenterColor = SampleColor;
-				LastValidDistance = SampleDist;
+				// Calculate the squared sample distance from the center.
+				int SampleDist = OffsetIdx.x * OffsetIdx.x + OffsetIdx.y * OffsetIdx.y;
+
+				// @todo MeshPaint:If needed, we can improve this by averaging the pixel colors at the same distance
+				if ( SampleDist < LastValidDistance )
+				{
+					// Compute offset coordinates for the pixel and sample the color from the brush mask.
+					float2 Coord = InTextureCoordinates + ( OffsetIdx * OffsetConst );
+
+					// A sample is only valid when it falls in the black area of the seam mask and the white area of the brush mask.
+					float OffsetSeamMask = Texture1.Sample( Texture1Sampler, Coord );
+					float OffsetBrushMask = Texture2.Sample( Texture2Sampler, Coord );
+
+					if ( OffsetSeamMask == 0 && OffsetBrushMask != 0 )
+					{
+						OutColor = Texture0.Sample( Texture0Sampler, Coord );
+						LastValidDistance = SampleDist;
+					}
+				}
 			}
 		}
 	} 
-	OutColor = CenterColor;
 }
