Substrate - added support for decal on mobile deferred renderer, surface and mesh. Made Legacy and Substrate decal function available both at the same time when substrate is enable for simplicity.

#tests Tested PC and mobile preview, deferred / forward, substrate on/off.
#rnx

diff --git a/Engine/Shaders/Private/BasePassPixelShader.usf b/Engine/Shaders/Private/BasePassPixelShader.usf
--- a/Engine/Shaders/Private/BasePassPixelShader.usf
+++ b/Engine/Shaders/Private/BasePassPixelShader.usf
@@ -1030,7 +1030,7 @@ void FPixelShaderInOut_MainPS(
 			float2 BufferUV = SvPositionToBufferUV(In.SvPosition);
 		#if SUBSTRATE_ENABLED 
 			#if SUBSTRATE_INLINE_SHADING && !SUBSTRATE_OPTIMIZED_UNLIT
-			const FSubstrateDBuffer SubstrateBufferData = GetDBufferData(BufferUV, ValidDBufferTargetMask);
+			const FSubstrateDBuffer SubstrateBufferData = SubstrateGetDBufferData(BufferUV, ValidDBufferTargetMask);
 			ApplyDBufferData(SubstrateBufferData, SubstratePixelHeader, SubstrateData);
 			DBufferVTFeedbackFactor = SubstrateBufferData.Coverage;
 			#endif
diff --git a/Engine/Shaders/Private/DBufferDecalShared.ush b/Engine/Shaders/Private/DBufferDecalShared.ush
--- a/Engine/Shaders/Private/DBufferDecalShared.ush
+++ b/Engine/Shaders/Private/DBufferDecalShared.ush
@@ -97,7 +97,7 @@ void ConvertToMetalness(
 }
 
 // @param BufferUV - UV space in the DBuffer textures
-FSubstrateDBuffer GetDBufferData(float2 BufferUV, uint RTMaskBit, uint EyeIndex)
+FSubstrateDBuffer SubstrateGetDBufferData(float2 BufferUV, uint RTMaskBit, uint EyeIndex)
 {
 	// SUBSTRATE_TODO: Add support MOBILE_MULTI_VIEW and INSTANCED_STEREO on mobile (using EyeIndex)
 
@@ -124,9 +124,9 @@ FSubstrateDBuffer GetDBufferData(float2 BufferUV, uint RTMaskBit, uint EyeIndex)
 	return SubstrateUnpackDBuffer(DBufferA, DBufferB, DBufferC);
 }
 
-FSubstrateDBuffer GetDBufferData(float2 BufferUV, uint RTMaskBit)
+FSubstrateDBuffer SubstrateGetDBufferData(float2 BufferUV, uint RTMaskBit)
 {
-	return GetDBufferData(BufferUV, RTMaskBit, 0 /*EyeIndex*/);
+	return SubstrateGetDBufferData(BufferUV, RTMaskBit, 0 /*EyeIndex*/);
 }
 
 // Define how DBuffer is applied onto existing layers:
@@ -357,7 +357,9 @@ void ApplyDBufferData(in FSubstrateDBuffer SubstrateDBufferData, inout FSubstrat
 	}
 }
 #endif // MATERIAL_IS_SUBSTRATE
-#else // SUBSTRATE_ENABLED
+
+#endif // SUBSTRATE_ENABLED
+
 struct FDBufferData
 {
 	// 0..1, premultiplied with ColorOpacity
@@ -510,4 +512,3 @@ void ApplyDBufferData(
 	Specular = Specular * DBufferData.RoughnessOpacity + DBufferData.PreMulSpecular;
 #endif
 }
-#endif // SUBSTRATE_ENABLED
\ No newline at end of file
diff --git a/Engine/Shaders/Private/DecalCommon.ush b/Engine/Shaders/Private/DecalCommon.ush
--- a/Engine/Shaders/Private/DecalCommon.ush
+++ b/Engine/Shaders/Private/DecalCommon.ush
@@ -84,7 +84,8 @@ FSubstrateDBuffer DecalCommonOutput(
 
 	return OutDBuffer;
 }
-#else // SUBSTRATE_ENABLED
+#endif // SUBSTRATE_ENABLED
+
 void DecalCommonOutput(inout FPixelShaderIn In, inout FPixelShaderOut Out, float3 Color, float Opacity, FGBufferData Data)
 {
 #if DECAL_RENDERSTAGE == DECAL_RENDERSTAGE_MOBILE
@@ -165,4 +166,3 @@ void DecalCommonOutput(inout FPixelShaderIn In, inout FPixelShaderOut Out, float
 	Out.MRT[3] = 0;
 #endif
 }
-#endif // SUBSTRATE_ENABLED
diff --git a/Engine/Shaders/Private/DeferredDecal.usf b/Engine/Shaders/Private/DeferredDecal.usf
--- a/Engine/Shaders/Private/DeferredDecal.usf
+++ b/Engine/Shaders/Private/DeferredDecal.usf
@@ -35,6 +35,14 @@
 #include "MobileLightingCommon.ush"
 #endif
 
+#define MOBILE_SUBSTRATE_USES_GBUFFER (DECAL_RENDERSTAGE == DECAL_RENDERSTAGE_MOBILEBEFORELIGHTING)
+
+#if MOBILE_SUBSTRATE_USES_GBUFFER
+#define SUBSTRATE_MATERIAL_EXPORT_REQUESTED 1
+#define MATERIAL_SUBSTRATE_OPAQUE_PRECOMPUTED_LIGHTING 0
+#include "/Engine/Private/Substrate/SubstrateExport.ush"
+#endif
+
 #endif // PIXELSHADER
 
 // large world coordinate tile position
@@ -178,7 +186,52 @@ void FPixelShaderInOut_MainPS(inout FPixelShaderIn In, inout FPixelShaderOut Out
 
 #if SUBSTRATE_ENABLED    
 	FSubstratePixelHeader SubstrateHeader = MaterialParameters.GetFrontSubstrateHeader();
-	SubstrateHeader.IrradianceAO.MaterialAO = GetMaterialAmbientOcclusion(PixelMaterialInputs);
+	const float MaterialAO = GetMaterialAmbientOcclusion(PixelMaterialInputs);
+	SubstrateHeader.IrradianceAO.MaterialAO = MaterialAO;
+
+#if MOBILE_SUBSTRATE_USES_GBUFFER
+
+	float3 EmissiveColor = 0.0f.xxx;
+	float Opacity = 0.0f;
+	FGBufferData DBufferData;
+	if (SubstrateHeader.SubstrateTree.BSDFCount > 0)
+	{
+		const FSubstrateIntegrationSettings Settings = InitSubstrateIntegrationSettings(false /*bForceFullyRough*/, false /*bRoughDiffuseEnabled*/, 0, false/*bRoughnessTracking*/);
+		const float3 SurfaceWorldNormal = MaterialParameters.TangentToWorld[2].xyz;
+		FExportResult Export = SubstrateMaterialExportOut(
+			Settings,
+			SubstrateHeader,
+			PixelMaterialInputs.GetFrontSubstrateData(),
+			SurfaceWorldNormal);
+		
+		// Either propagate the BSDF opacity, or use the one for AlphaComposite (premultiplied alpha)
+		float CoverageOverBackground = Export.Coverage * DecalFading;
+	#if SUBSTRATE_USE_PREMULTALPHA_OVERRIDE // AlphaComposite - Premultiplied alpha blending
+		CoverageOverBackground = GetMaterialOpacity(PixelMaterialInputs) * DecalFading;
+	#endif 
+		Opacity										= CoverageOverBackground;
+		EmissiveColor								= Export.EmissiveLuminance;
+
+		FGBufferData DBufferData					= (FGBufferData)0;
+		DBufferData.WorldNormal						= Export.WorldNormal;
+		DBufferData.BaseColor						= Export.BaseColor;
+		DBufferData.Metallic						= Export.Metallic;
+		DBufferData.Specular						= Export.Specular;
+		DBufferData.Roughness						= Export.Roughness;
+		DBufferData.CustomData						= 0;
+		DBufferData.IndirectIrradiance				= 0;
+		DBufferData.PrecomputedShadowFactors		= 1;
+		DBufferData.GBufferAO						= MaterialAO;
+		DBufferData.ShadingModelID					= Export.ShadingModelID;
+		DBufferData.SelectiveOutputMask				= 0;
+		DBufferData.PerObjectGBufferData			= 1;
+		DBufferData.DiffuseIndirectSampleOcclusion	= 0;
+		DBufferData.Velocity						= 0;
+
+		DecalCommonOutput(In, Out, EmissiveColor, Opacity, DBufferData);
+	}
+
+#else
 
 	float Coverage = 0.0;
 	FSubstrateDBuffer DBufferData =(FSubstrateDBuffer)0;
@@ -192,16 +245,20 @@ void FPixelShaderInOut_MainPS(inout FPixelShaderIn In, inout FPixelShaderOut Out
 
 		// Either propagate the BSDF opacity, or use the one for AlphaComposite (premultiplied alpha)
 		float CoverageOverBackground = SubstrateHeader.SubstrateTree.BSDFs[0].Coverage * DecalFading;
-#if SUBSTRATE_USE_PREMULTALPHA_OVERRIDE // AlphaComposite - Premultiplied alpha blending
-		CoverageOverBackground = GetMaterialOpacity(PixelMaterialInputs);
-#endif 
+	#if SUBSTRATE_USE_PREMULTALPHA_OVERRIDE // AlphaComposite - Premultiplied alpha blending
+		CoverageOverBackground = GetMaterialOpacity(PixelMaterialInputs) * DecalFading;
+	#endif 
 
 		// Since the Convert-To-Decal node forces parameter blending for the entire Substrate tree, the BSDF we are interested in will be in 0.
 		DBufferData = DecalCommonOutput(Out, SubstrateHeader, SubstrateHeader.SubstrateTree.BSDFs[0], DecalFading, CoverageOverBackground);
 	}
 	const float Opacity = DBufferData.Coverage;
 	const float3 EmissiveColor = DBufferData.Emissive;
+
+#endif
+
 #else // SUBSTRATE_ENABLED
+
 	float3 EmissiveColor = GetMaterialEmissive(PixelMaterialInputs);
 	
 	float Opacity = GetMaterialOpacity(PixelMaterialInputs) * DecalFading;
diff --git a/Engine/Shaders/Private/MaterialTemplate.ush b/Engine/Shaders/Private/MaterialTemplate.ush
--- a/Engine/Shaders/Private/MaterialTemplate.ush
+++ b/Engine/Shaders/Private/MaterialTemplate.ush
@@ -3387,7 +3387,7 @@ float4 MaterialExpressionDBufferTextureLookup(FMaterialPixelParameters Parameter
 		uint2 PixelPos = uint2(BufferUV * View.BufferSizeAndInvSize.xy);
 		uint ValidDBufferTargetMask = GetDBufferTargetMask(PixelPos) & (1u << DBufferTextureIndex);
 #if SUBSTRATE_ENABLED
-		const FSubstrateDBuffer DBufferData = GetDBufferData(BufferUV, ValidDBufferTargetMask);
+		const FSubstrateDBuffer DBufferData = SubstrateGetDBufferData(BufferUV, ValidDBufferTargetMask);
 		switch (DBufferTextureIndex)
 		{
 		// All data (BaseColor, WorldNormal, Roughness) are pre-multiplied by their coverage
diff --git a/Engine/Shaders/Private/MeshDecals.usf b/Engine/Shaders/Private/MeshDecals.usf
--- a/Engine/Shaders/Private/MeshDecals.usf
+++ b/Engine/Shaders/Private/MeshDecals.usf
@@ -23,6 +23,14 @@
 #include "/Engine/Generated/VertexFactory.ush"
 
 
+#define MOBILE_SUBSTRATE_USES_GBUFFER (SHADING_PATH_MOBILE) // When using forward shading, the base pass pixel shader is used with blending.
+
+#if SUBSTRATE_ENABLED && MOBILE_SUBSTRATE_USES_GBUFFER
+#define SUBSTRATE_MATERIAL_EXPORT_REQUESTED 1
+#define MATERIAL_SUBSTRATE_OPAQUE_PRECOMPUTED_LIGHTING 0
+#include "/Engine/Private/Substrate/SubstrateExport.ush"
+#endif
+
 // If virtual texture feedback is enabled then use early depth test so that UAV feedback buffer writes respect the depth test
 #if MATERIAL_VIRTUALTEXTURE_FEEDBACK
 	#if COMPILER_SUPPORTS_DEPTHSTENCIL_EARLYTEST_LATEWRITE
@@ -132,8 +140,50 @@ void FPixelShaderInOut_MainPS(
 
 #if SUBSTRATE_ENABLED
 	FSubstratePixelHeader SubstrateHeader = MaterialParameters.GetFrontSubstrateHeader();
-	SubstrateHeader.IrradianceAO.MaterialAO = GetMaterialAmbientOcclusion(PixelMaterialInputs);
+	const float MaterialAO = GetMaterialAmbientOcclusion(PixelMaterialInputs);
+	SubstrateHeader.IrradianceAO.MaterialAO = MaterialAO;
+
+#if MOBILE_SUBSTRATE_USES_GBUFFER
+	float3 EmissiveColor = 0.0f.xxx;
+	float Opacity = 0.0f;
+	FGBufferData DBufferData;
+	if (SubstrateHeader.SubstrateTree.BSDFCount > 0)
+	{
+		const FSubstrateIntegrationSettings Settings = InitSubstrateIntegrationSettings(false /*bForceFullyRough*/, false /*bRoughDiffuseEnabled*/, 0, false/*bRoughnessTracking*/);
+		const float3 SurfaceWorldNormal = MaterialParameters.TangentToWorld[2].xyz;
+		FExportResult Export = SubstrateMaterialExportOut(
+			Settings,
+			SubstrateHeader,
+			PixelMaterialInputs.GetFrontSubstrateData(),
+			SurfaceWorldNormal);
+		
+		// Either propagate the BSDF opacity, or use the one for AlphaComposite (premultiplied alpha)
+		float CoverageOverBackground = Export.Coverage;
+	#if SUBSTRATE_USE_PREMULTALPHA_OVERRIDE // AlphaComposite - Premultiplied alpha blending
+		CoverageOverBackground = GetMaterialOpacity(PixelMaterialInputs);
+	#endif 
+		Opacity										= CoverageOverBackground;
+		EmissiveColor								= Export.EmissiveLuminance;
+
+		FGBufferData DBufferData					= (FGBufferData)0;
+		DBufferData.WorldNormal						= Export.WorldNormal;
+		DBufferData.BaseColor						= Export.BaseColor;
+		DBufferData.Metallic						= Export.Metallic;
+		DBufferData.Specular						= Export.Specular;
+		DBufferData.Roughness						= Export.Roughness;
+		DBufferData.CustomData						= 0;
+		DBufferData.IndirectIrradiance				= 0;
+		DBufferData.PrecomputedShadowFactors		= 1;
+		DBufferData.GBufferAO						= MaterialAO;
+		DBufferData.ShadingModelID					= Export.ShadingModelID;
+		DBufferData.SelectiveOutputMask				= 0;
+		DBufferData.PerObjectGBufferData			= 1;
+		DBufferData.DiffuseIndirectSampleOcclusion	= 0;
+		DBufferData.Velocity						= 0;
 
+		DecalCommonOutput(In, Out, EmissiveColor, Opacity, DBufferData);
+	}
+#else
 	float Coverage = 0.0;
 	if (SubstrateHeader.SubstrateTree.BSDFCount > 0)
 	{
@@ -153,7 +203,10 @@ void FPixelShaderInOut_MainPS(
 		Coverage = DBuffer.Coverage;
 	}
 	const float Opacity = Coverage;
+#endif // MOBILE_SUBSTRATE_USES_GBUFFER
+
 #else // SUBSTRATE_ENABLED
+
 	float3 Color = GetMaterialEmissive(PixelMaterialInputs);
 
 	float Opacity = GetMaterialOpacity(PixelMaterialInputs);
diff --git a/Engine/Shaders/Private/MobileBasePassPixelShader.usf b/Engine/Shaders/Private/MobileBasePassPixelShader.usf
--- a/Engine/Shaders/Private/MobileBasePassPixelShader.usf
+++ b/Engine/Shaders/Private/MobileBasePassPixelShader.usf
@@ -492,7 +492,7 @@ void Main(
 
 		#if SUBSTRATE_ENABLED 
 			#if SUBSTRATE_INLINE_SHADING && !SUBSTRATE_OPTIMIZED_UNLIT
-			const FSubstrateDBuffer SubstrateBufferData = GetDBufferData(BufferUV, ValidDBufferTargetMask, EyeIndex);
+			const FSubstrateDBuffer SubstrateBufferData = SubstrateGetDBufferData(BufferUV, ValidDBufferTargetMask, EyeIndex);
 			ApplyDBufferData(SubstrateBufferData, SubstratePixelHeader, SubstrateData);
 			DBufferVTFeedbackFactor = SubstrateBufferData.Coverage;
 			#endif
diff --git a/Engine/Shaders/Private/Substrate/SubstrateDBuffer.usf b/Engine/Shaders/Private/Substrate/SubstrateDBuffer.usf
--- a/Engine/Shaders/Private/Substrate/SubstrateDBuffer.usf
+++ b/Engine/Shaders/Private/Substrate/SubstrateDBuffer.usf
@@ -122,7 +122,7 @@ void MainCS(
 	);
 
 	const float2 BufferUV = SvPositionToBufferUV(float4(PixelCoord + 0.5f, 0, 0));
-	const FSubstrateDBuffer DBufferData = GetDBufferData(BufferUV, ValidDBufferTargetMask);
+	const FSubstrateDBuffer DBufferData = SubstrateGetDBufferData(BufferUV, ValidDBufferTargetMask);
 
 	// Early out if DBuffer has no coverage
 	BRANCH
