Fixed incorrect color component swizzling on Skin Cache tangent recomputation

#12644

diff --git a/Engine/Shaders/Private/RecomputeTangentsPerVertexPass.usf b/Engine/Shaders/Private/RecomputeTangentsPerVertexPass.usf
--- a/Engine/Shaders/Private/RecomputeTangentsPerVertexPass.usf
+++ b/Engine/Shaders/Private/RecomputeTangentsPerVertexPass.usf
@@ -4,6 +4,7 @@
 	RecomputeTangentsPerVertexPass.usf: Recompute Skin Tangents
 =============================================================================*/
 
+#include "/Engine/Public/Platform.ush"
 #include "RecomputeTangentsCommon.ush"
 
 // 0:fast and simple / 1:higher quality lighting with extreme mesh distortions
@@ -155,7 +156,7 @@ void MainCS(
 #if BLEND_USING_VERTEX_COLOR
 		float4 InputTangentX = TangentBias_SkinCache(TangentInputBuffer[2 * Offset + GPUSKIN_RWBUFFER_OFFSET_TANGENT_X]);
 		float4 InputTangentZ = TangentBias_SkinCache(TangentInputBuffer[2 * Offset + GPUSKIN_RWBUFFER_OFFSET_TANGENT_Z]);
-		float4 BlendFactors = ColorInputBuffer[DispatchThreadId + InputStreamStart];
+		float4 BlendFactors = ColorInputBuffer[DispatchThreadId + InputStreamStart] FMANUALFETCH_COLOR_COMPONENT_SWIZZLE;
 		float BlendFactor = (VertexColorChannel == 0) ? BlendFactors.r : ((VertexColorChannel == 1) ? BlendFactors.g : BlendFactors.b);
 
 		// Linear interpolation
