Substrate - added subsurface profile ID to export. We need to wait for further development with mobile deferred to expend export.

#rnx
#tests PD editor dx12, mobile preview

diff --git a/Engine/Shaders/Private/MobileBasePassPixelShader.usf b/Engine/Shaders/Private/MobileBasePassPixelShader.usf
--- a/Engine/Shaders/Private/MobileBasePassPixelShader.usf
+++ b/Engine/Shaders/Private/MobileBasePassPixelShader.usf
@@ -740,6 +740,8 @@ void Main(
 	#endif
 		MaterialParameters.TangentToWorld[2].xyz = Export.WorldTangent;
 
+		// SUBSTRATE_TODO when clear coat or cloth are clarified for mobile deferred, SetGBufferForShadingModel should read data directly from Export. It will depend on material conversion features and mobile deferred support.
+		// Right now it uses GetMaterialCustomData0(MaterialParameters) to get some specific data. Which does not work with Substrate (but it also do not work for mobile deferred it seems).
 		SetGBufferForShadingModel(
 			GBuffer,
 			MaterialParameters,
@@ -750,7 +752,7 @@ void Main(
 			Export.Roughness,
 			Export.Anisotropy,
 			Export.SubsurfaceColor,
-			0,						// SUBSTRATE_TODO SubsurfaceProfile
+			Export.SubsurfaceProfileID,
 			0.0f,					// Dither
 			Export.ShadingModelID
 		);
diff --git a/Engine/Shaders/Private/Substrate/Substrate.ush b/Engine/Shaders/Private/Substrate/Substrate.ush
--- a/Engine/Shaders/Private/Substrate/Substrate.ush
+++ b/Engine/Shaders/Private/Substrate/Substrate.ush
@@ -3474,6 +3474,23 @@ float3 SubstrateGetBSDFSubSurfaceColor(in FSubstrateBSDF BSDF)
 	return 0.f;
 }
 
+uint SubstrateGetBSDFSubSurfaceProfileID(in FSubstrateBSDF BSDF)
+{
+	switch (BSDF_GETTYPE(BSDF))
+	{
+	case SUBSTRATE_BSDF_TYPE_SLAB:
+	case SUBSTRATE_BSDF_TYPE_EYE:
+	{
+		if (BSDF_GETSSSTYPE(BSDF) != SSS_TYPE_DIFFUSION_PROFILE)
+		{
+			return SubstrateSubsurfaceProfileIdTo8bits(SLAB_SSSPROFILEID(BSDF));
+		}
+		return 0;
+	}
+	}
+	return 0;
+}
+
 float SubstrateGetBSDFSpecular(in FSubstrateBSDF BSDF)
 {
 	switch (BSDF_GETTYPE(BSDF))
diff --git a/Engine/Shaders/Private/Substrate/SubstrateExport.ush b/Engine/Shaders/Private/Substrate/SubstrateExport.ush
--- a/Engine/Shaders/Private/Substrate/SubstrateExport.ush
+++ b/Engine/Shaders/Private/Substrate/SubstrateExport.ush
@@ -949,9 +949,10 @@ struct FExportResult
 	float  Roughness;
 	float  Anisotropy;
 	float  Metallic;
-	float3 CustomData0;
-	float3 CustomData1;
+	float3 CustomData0;	// Not sur CustomDataX memeber here is the best approach yet.
+	float3 CustomData1;	// Idem
 	float3 SubsurfaceColor;
+	uint   SubsurfaceProfileID;
 	float3 WorldTangent;
 	float  ShadingModelID;
 };
@@ -995,11 +996,12 @@ FExportResult SubstrateMaterialExportOut(
 				Out.Roughness				= SubstrateGetBSDFRoughness(CurrentBSDF);
 				Out.Anisotropy				= SubstrateGetBSDFAnisotropy(CurrentBSDF);
 				Out.Metallic				= SubstrateGetBSDFMetallic(CurrentBSDF);
+				Out.SubsurfaceColor			= SubstrateGetBSDFSubSurfaceColor(CurrentBSDF);
+				Out.SubsurfaceProfileID		= SubstrateGetBSDFSubSurfaceProfileID(CurrentBSDF);
+				Out.WorldTangent			= T;
+				Out.ShadingModelID			= SubstrateGetLegacyShadingModels(CurrentBSDF);
 				Out.CustomData0				= 0;											// SUBSTRATE_TODO not tested since not exposed in the only use case so far: material baking.
 				Out.CustomData1				= 0;											// idem
-				Out.SubsurfaceColor			= SubstrateGetBSDFSubSurfaceColor(CurrentBSDF);	// idem
-				Out.WorldTangent			= T;
-				Out.ShadingModelID			= SubstrateGetLegacyShadingModels(CurrentBSDF);	// idem
 			}
 		#undef CurrentBSDF
 		}
