Fix issue where DoF focal plane is inconsistent between different screenpercentages at low post-process quality settings.
The CoCRadius used in calculations through the DoF code can be expressed as pixels in the viewbuffer, pixels in the half-res buffer, or pixels after upscaling (screenpercentage).
This CL attempts to fix inconsistencies in the units used in the gather and recombine passes.

diff --git a/Engine/Shaders/Private/DiaphragmDOF/DOFCommon.ush b/Engine/Shaders/Private/DiaphragmDOF/DOFCommon.ush
--- a/Engine/Shaders/Private/DiaphragmDOF/DOFCommon.ush
+++ b/Engine/Shaders/Private/DiaphragmDOF/DOFCommon.ush
@@ -160,6 +160,7 @@ static const int2 kOffsetsSquare3x3[9] = {
 
 //------------------------------------------------------- PARAMETERS
 
+// See SetCocModelParameters, CocRadiusBasis
 float CocInfinityRadius;
 float CocInFocusRadius;
 uint bCocEnableDynamicRadiusOffset;
@@ -204,6 +205,7 @@ float GetCocOffset(float CocRadius)
 #endif
 
 // Returns the coc radius
+// The unit for the output, i.e. 'encoded' pixels or physical pixels, depends on what CocRadiusBasis was used.
 float SceneDepthToCocRadius(float SceneDepth)
 {
 	const float Focus = View.DepthOfFieldFocalDistance; // TODO: rename this to DepthOfFieldFocusDistance.
diff --git a/Engine/Shaders/Private/DiaphragmDOF/DOFRecombine.usf b/Engine/Shaders/Private/DiaphragmDOF/DOFRecombine.usf
--- a/Engine/Shaders/Private/DiaphragmDOF/DOFRecombine.usf
+++ b/Engine/Shaders/Private/DiaphragmDOF/DOFRecombine.usf
@@ -140,6 +140,7 @@ float2 DOFBufferUVMax;
 float4 SeparateTranslucencyBilinearUVMinMax;
 uint SeparateTranslucencyUpscaling;
 float EncodedCocRadiusToRecombineCocRadius;
+float MaxRecombineAbsCocRadius;
 
 Texture2D BokehLUT;
 
@@ -225,6 +226,9 @@ void AmendAdditiveColorWithMaxTranslucency(inout float4 Color, inout float Trans
 	}
 }
 
+// TODO: most of the math below uses CocRadius as expressed in physical pixels, rather than scaling-independent 'encoded' units.
+// It's not clear if the math is always accounting for this, as many of these expressions contain magic numbers.
+
 
 // Compute sample weight according's to its Coc radius.
 float ComputeSampleWeight(float CocRadius)
@@ -246,8 +250,6 @@ float ComputeSampleWeight(float CocRadius)
 	#endif
 }
 
-
-// Computes full resolution sample intersection.
 float ComputeSampleIntersection(float SampleCocRadius, float SampleDistance)
 {
 	#if 0 // DEBUG
@@ -278,7 +280,7 @@ float ComputeSampleIntersection(float SampleCocRadius, float SampleDistance)
 float ComputeBackgroundSampleOpacity(float CocRadius)
 {
 	//return CocRadius < MAX_RECOMBINE_ABS_COC_RADIUS ? 1 : 0;
-	return saturate(MAX_RECOMBINE_ABS_COC_RADIUS - CocRadius);
+	return saturate(MaxRecombineAbsCocRadius - CocRadius);
 }
 
 // Returns the opacity to use to transition to background.
@@ -303,7 +305,7 @@ float ComputeForegroundSampleOpacity(float CocRadius)
 // Returns the opacity to use to transition to background.
 float IsConsideredForegroundSample(float CocRadius)
 {
-	return ComputeForegroundSampleOpacity(CocRadius) * saturate(MAX_RECOMBINE_ABS_COC_RADIUS + CocRadius);
+	return ComputeForegroundSampleOpacity(CocRadius) * saturate(MaxRecombineAbsCocRadius + CocRadius);
 }
 
 
@@ -391,7 +393,7 @@ FGatherSampleDerivedParameters ComputeSampleDerivates(in FFullResGatherAccumulat
 	else if (Accumulator.LayerProcessing == LAYER_PROCESSING_SLIGHT_OUT_OF_FOCUS)
 	{
 		DerivedA.Opacity = ComputeBackgroundSampleOpacity(A.CocRadius);
-		DerivedA.IsConsidered = saturate(MAX_RECOMBINE_ABS_COC_RADIUS - abs(A.CocRadius));
+		DerivedA.IsConsidered = saturate(MaxRecombineAbsCocRadius - abs(A.CocRadius));
 	}
 	else if (Accumulator.LayerProcessing == LAYER_PROCESSING_BACKGROUND_ONLY)
 	{
@@ -775,11 +777,11 @@ void RecombineMainCS(
 		// Grab the smallest slightly out of focus Coc radius of the tile.
 		float TileMaxConsideredAbsCocRadius;
 		{
-			float MaxConsideredAbsCocRadius = abs(CenterSample.CocRadius) < MAX_RECOMBINE_ABS_COC_RADIUS ? abs(CenterSample.CocRadius) : 0;
+			float MaxConsideredAbsCocRadius = abs(CenterSample.CocRadius) < MaxRecombineAbsCocRadius ? abs(CenterSample.CocRadius) : 0;
 
 			for (uint j = 0; j < 4; j++)
 			{
-				const float2 SamplePixelOffset = float2(kOffsetsCross3x3[j]) * CocDistanceToFullResPixelDistance(MAX_RECOMBINE_ABS_COC_RADIUS);
+				const float2 SamplePixelOffset = float2(kOffsetsCross3x3[j]) * CocDistanceToFullResPixelDistance(MaxRecombineAbsCocRadius);
 
 				float2 SampleUVOffset = View.BufferSizeAndInvSize.zw * SamplePixelOffset;
 				float2 SampleUV = InputParameters.SceneBufferUV + SampleUVOffset;
@@ -809,7 +811,7 @@ void RecombineMainCS(
 					float MinCocRadiusSquare = dot(OutsideGroupCocOffset, OutsideGroupCocOffset);
 
 					// Not interested if the CocRadius is too large, or does not overlap with the group's tile.
-					if (SampleAbsCocRadius < MAX_RECOMBINE_ABS_COC_RADIUS &&
+					if (SampleAbsCocRadius < MaxRecombineAbsCocRadius &&
 						SampleAbsCocRadius * SampleAbsCocRadius > MinCocRadiusSquare)
 					{
 						MaxConsideredAbsCocRadius = max(MaxConsideredAbsCocRadius, SampleAbsCocRadius);
@@ -817,7 +819,7 @@ void RecombineMainCS(
 				}
 				#else
 				{
-					MaxConsideredAbsCocRadius = max(MaxConsideredAbsCocRadius, SampleAbsCocRadius < MAX_RECOMBINE_ABS_COC_RADIUS ? SampleAbsCocRadius : MaxConsideredAbsCocRadius);
+					MaxConsideredAbsCocRadius = max(MaxConsideredAbsCocRadius, SampleAbsCocRadius < MaxRecombineAbsCocRadius ? SampleAbsCocRadius : MaxConsideredAbsCocRadius);
 				}
 				#endif
 			}
@@ -971,15 +973,17 @@ void RecombineMainCS(
 			}
 			#endif
 
+			float ClampWeight = saturate(CenterSample.CocRadius * CenterSample.CocRadius * 4.0);
+
 			// Clamp color
 			float3 ClampedColor = clamp(Accumulator.Color.rgb, Min.rgb * Accumulator.ColorWeight, Max.rgb * Accumulator.ColorWeight);
 
-			Accumulator.Color.rgb = lerp(Accumulator.Color.rgb, ClampedColor, saturate(CenterSample.CocRadius * CenterSample.CocRadius * 4.0));
+			Accumulator.Color.rgb = lerp(Accumulator.Color.rgb, ClampedColor, ClampWeight);
 
 			// Clamp opacity.			
 			float ClampedOpacity = clamp(Accumulator.Opacity, Min.a * Accumulator.OpacityWeight, Max.a * Accumulator.OpacityWeight);
 
-			Accumulator.Opacity = lerp(Accumulator.Opacity, ClampedOpacity, saturate(CenterSample.CocRadius * CenterSample.CocRadius * 4.0));
+			Accumulator.Opacity = lerp(Accumulator.Opacity, ClampedOpacity, ClampWeight);
 		}
 		#endif
 
@@ -1072,7 +1076,7 @@ void RecombineMainCS(
 	if (GatherForegroundTranslucency < 1 && BackgroundValidity != 1.0)
 	{
 		#if CONFIG_GATHER_PAIR_COUNT == 0
-			float MaxTranslucency = saturate(MAX_RECOMBINE_ABS_COC_RADIUS + CenterSample.CocRadius);
+			float MaxTranslucency = saturate(MaxRecombineAbsCocRadius + CenterSample.CocRadius);
 		#else
 			float MaxTranslucency = 1 - ComputeForegroundSampleOpacity(CenterSample.CocRadius);
 		#endif
diff --git a/Engine/Source/Runtime/Renderer/Private/PostProcess/DiaphragmDOF.cpp b/Engine/Source/Runtime/Renderer/Private/PostProcess/DiaphragmDOF.cpp
--- a/Engine/Source/Runtime/Renderer/Private/PostProcess/DiaphragmDOF.cpp
+++ b/Engine/Source/Runtime/Renderer/Private/PostProcess/DiaphragmDOF.cpp
@@ -1371,6 +1371,7 @@ class FDiaphragmDOFRecombineCS : public FDiaphragmDOFShader
 		SHADER_PARAMETER(FVector4f, SeparateTranslucencyBilinearUVMinMax)
 		SHADER_PARAMETER(int32, SeparateTranslucencyUpscaling)
 		SHADER_PARAMETER(float, EncodedCocRadiusToRecombineCocRadius)
+		SHADER_PARAMETER(float, MaxRecombineAbsCocRadius)
 		
 		SHADER_PARAMETER_RDG_TEXTURE(Texture2D, BokehLUT)
 
@@ -2363,11 +2364,8 @@ bool DiaphragmDOF::AddPasses(
 			FVector2f ConsiderCocRadiusAffineTransformation1 = kContantlyPassingAffineTransformation;
 			FVector2f ConsiderAbsCocRadiusAffineTransformation = kContantlyPassingAffineTransformation;
 			{
-				// Gathering scalability.
-				const float GatheringScalingDownFactor = float(PreprocessViewSize.X) / float(GatheringViewSize.X);
-
 				// Coc radius considered.
-				const float RecombineCocRadiusBorder = GatheringScalingDownFactor * (kMaxSlightOutOfFocusCocRadius - 1.0f);
+				const float RecombineCocRadiusBorder = kMaxSlightOutOfFocusCocRadius - 1.0f;
 
 				if (ConvolutionSettings.LayerProcessing == EDiaphragmDOFLayerProcessing::ForegroundOnly)
 				{
@@ -2393,7 +2391,7 @@ bool DiaphragmDOF::AddPasses(
 				else if (ConvolutionSettings.LayerProcessing == EDiaphragmDOFLayerProcessing::SlightOutOfFocus)
 				{
 					ConsiderAbsCocRadiusAffineTransformation = GenerateSaturatedAffineTransformation(
-						RecombineCocRadiusBorder + GatheringScalingDownFactor * 1.0f, RecombineCocRadiusBorder);
+						RecombineCocRadiusBorder + 1.0f, RecombineCocRadiusBorder);
 				}
 				else
 				{
@@ -2415,15 +2413,15 @@ bool DiaphragmDOF::AddPasses(
 				float(SrcSize.X * GatheringViewSize.X) / float(PreprocessViewSize.X),
 				float(SrcSize.Y * GatheringViewSize.Y) / float(PreprocessViewSize.Y));
 			PassParameters->MipBias = FMath::Log2(float(PreprocessViewSize.X) / float(GatheringViewSize.X));
-			PassParameters->MaxRecombineAbsCocRadius = float(kMaxSlightOutOfFocusCocRadius) / (GatheringCocRadiusBasis / EncodedCocRadiusBasis);
+			PassParameters->MaxRecombineAbsCocRadius = float(kMaxSlightOutOfFocusCocRadius);
 			PassParameters->CocSqueeze = CocModel.Squeeze;
 			PassParameters->CocInvSqueeze = 1.0f / CocModel.Squeeze;
 			
 			PassParameters->Petzval = CocModel.Petzval;
 			PassParameters->PetzvalFalloffPower = CocModel.PetzvalFalloffPower;
 
-			PassParameters->TileDecisionParameters.MinGatherRadius = PassParameters->MaxRecombineAbsCocRadius - 1;
-			PassParameters->TileDecisionParameters.SlightOutOfFocusRadiusBoundary = float(kMaxSlightOutOfFocusCocRadius) / (GatheringCocRadiusBasis / EncodedCocRadiusBasis);
+			PassParameters->TileDecisionParameters.MinGatherRadius = float(kMaxSlightOutOfFocusCocRadius) - 1;
+			PassParameters->TileDecisionParameters.SlightOutOfFocusRadiusBoundary = float(kMaxSlightOutOfFocusCocRadius);
 			PassParameters->CommonParameters = CommonParameters;
 		
 			PassParameters->GatherInputSize = FVector4f(SrcSize.X, SrcSize.Y, 1.0f / SrcSize.X, 1.0f / SrcSize.Y);
@@ -2497,7 +2495,7 @@ bool DiaphragmDOF::AddPasses(
 			PermutationVector.Set<FDiaphragmDOFPostfilterCS::FTileOptimization>(true); // TODO
 			PermutationVector = FDiaphragmDOFPostfilterCS::RemapPermutationVector(PermutationVector);
 			
-			float MaxRecombineAbsCocRadius = 3.0 * float(PreprocessViewSize.X) / float(GatheringViewSize.X);
+			float MaxRecombineAbsCocRadius = float(kMaxSlightOutOfFocusCocRadius);
 
 			FDiaphragmDOFPostfilterCS::FParameters* PassParameters = GraphBuilder.AllocParameters<FDiaphragmDOFPostfilterCS::FParameters>();
 			PassParameters->ViewportRect = FIntRect(0, 0, GatheringViewSize.X, GatheringViewSize.Y);
@@ -2798,6 +2796,7 @@ bool DiaphragmDOF::AddPasses(
 			(GatheringViewSize.X - 0.5f) / float(RefBufferSize.X),
 			(GatheringViewSize.Y - 0.5f) / float(RefBufferSize.Y));
 		PassParameters->EncodedCocRadiusToRecombineCocRadius = float(GatheringViewSize.X) / EncodedCocRadiusBasis; 
+		PassParameters->MaxRecombineAbsCocRadius = float(kMaxSlightOutOfFocusCocRadius) * PassParameters->EncodedCocRadiusToRecombineCocRadius;
 
 		PassParameters->SeparateTranslucencyBilinearUVMinMax.X = (SeparateTranslucencyViewport.Rect.Min.X + 0.5f) / float(SeparateTranslucencyViewport.Extent.X);
 		PassParameters->SeparateTranslucencyBilinearUVMinMax.Y = (SeparateTranslucencyViewport.Rect.Min.Y + 0.5f) / float(SeparateTranslucencyViewport.Extent.Y);
