Modified PR #12223: Fix 2D unwrapped tonemapping LUT going out of bounds for extremely bright colors.

diff --git a/Engine/Shaders/Private/PostProcessTonemap.usf b/Engine/Shaders/Private/PostProcessTonemap.usf
--- a/Engine/Shaders/Private/PostProcessTonemap.usf
+++ b/Engine/Shaders/Private/PostProcessTonemap.usf
@@ -191,7 +191,7 @@ half3 ColorLookupTable( half3 LinearColor )
 		// ST2084 expects to receive linear values 0-10000 in nits.
 		// So the linear value must be multiplied by a scale factor to convert to nits.
 		// We don't send negative values to LinearToST2084 as it will result in NaN because of pow.
-		LUTEncodedColor = LinearToST2084(max(0, LinearColor) * LinearToNitsScale);
+		LUTEncodedColor = LinearToST2084(clamp(LinearColor * LinearToNitsScale, 0.0, 10000.0));
 	}
 	else
 	{
