[Lit Translucency] Fix incorrect world position calculation in TLV injection shaders.
- caused lighting to be offset by half a voxel.

diff --git a/Engine/Shaders/Private/TranslucentLightInjectionShaders.usf b/Engine/Shaders/Private/TranslucentLightInjectionShaders.usf
--- a/Engine/Shaders/Private/TranslucentLightInjectionShaders.usf
+++ b/Engine/Shaders/Private/TranslucentLightInjectionShaders.usf
@@ -53,7 +53,7 @@ uint AtmospherePerPixelTransmittanceEnabled;
 float3 GetTranslatedWorldPosition(float2 UV, uint LayerIndex)
 {
 	float ZPosition = View.TranslucencyLightingVolumeMin[VolumeCascadeIndex].z + (LayerIndex + .5f) * View.TranslucencyLightingVolumeInvSize[VolumeCascadeIndex].w;
-	float3 TranslatedWorldPosition = float3(View.TranslucencyLightingVolumeMin[VolumeCascadeIndex].xy + UV / View.TranslucencyLightingVolumeInvSize[VolumeCascadeIndex].xy - .5f * View.TranslucencyLightingVolumeInvSize[VolumeCascadeIndex].w, ZPosition);
+	float3 TranslatedWorldPosition = float3(View.TranslucencyLightingVolumeMin[VolumeCascadeIndex].xy + UV / View.TranslucencyLightingVolumeInvSize[VolumeCascadeIndex].xy, ZPosition);
 	return TranslatedWorldPosition;
 }
 
