Substrate - Metalness / F0 conversion update.

#tests PS editor dx12 x/xo Substrate, mobile preview
#rnx

diff --git a/Engine/Shaders/Private/ShadingCommon.ush b/Engine/Shaders/Private/ShadingCommon.ush
--- a/Engine/Shaders/Private/ShadingCommon.ush
+++ b/Engine/Shaders/Private/ShadingCommon.ush
@@ -103,7 +103,7 @@ float F0ToDielectricSpecular(float F0)
 
 float F0RGBToF0(float3 F0)
 {
-	return dot(0.3333333.xxx, F0);
+	return max3(F0.r, F0.g, F0.b);
 }
 
 float F0RGBToDielectricSpecular(float3 F0)
@@ -135,13 +135,14 @@ float DielectricIorToF0(float Ior)
 }
 
 // Anything with Specular less than 2% is physically impossible and is instead considered to be shadowing.
+// See F_Schlick implementation.
 float  GetF0MicroOcclusionThreshold() { return 0.02f; }
 float  F0ToMicroOcclusion(float F0)   { return saturate(50.0 * F0); }
 float3 F0ToMicroOcclusion(float3 F0)  { return saturate(50.0 * F0); }
 
 float F0RGBToMicroOcclusion(float3 F0)
 {
-	return F0ToMicroOcclusion(max(F0.r, max(F0.g, F0.b)));
+	return F0ToMicroOcclusion(F0RGBToF0(F0));
 }
 
 half3 ComputeF0(half Specular, half3 BaseColor, half Metallic)
@@ -176,7 +177,7 @@ float F0ToMetallic(float F0)
 
 float F0RGBToMetallic(float3 F0)
 {
-	return F0ToMetallic(max(F0.r, max(F0.g, F0.b)));
+	return F0ToMetallic(F0RGBToF0(F0));
 }
 
 
diff --git a/Engine/Shaders/Private/ShadingEnergyConservationTemplate.ush b/Engine/Shaders/Private/ShadingEnergyConservationTemplate.ush
--- a/Engine/Shaders/Private/ShadingEnergyConservationTemplate.ush
+++ b/Engine/Shaders/Private/ShadingEnergyConservationTemplate.ush
@@ -12,6 +12,8 @@
 #error  IS_BXDF_ENERGY_TYPE_ACHROMATIC needs to be defined
 #endif
 
+#include "ShadingCommon.ush"
+
 ///////////////////////////////////////////////////////////////////////////////////////////////////
 // Energy terms
 
@@ -79,7 +81,7 @@ BXDF_ENERGY_SUFFIX(FBxDFEnergyTerms) BXDF_ENERGY_SUFFIX(ComputeGGXSpecEnergyTerm
 
 BXDF_ENERGY_SUFFIX(FBxDFEnergyTerms) BXDF_ENERGY_SUFFIX(ComputeGGXSpecEnergyTerms)(float Roughness, float NoV, float3 F0)
 {
-	const float F90 = saturate(50.0 * F0.g); // See F_Schlick implementation
+	const float F90 = F0RGBToMicroOcclusion(F0);
 	return BXDF_ENERGY_SUFFIX(ComputeGGXSpecEnergyTerms)(Roughness, NoV, F0, F90);
 }
 
diff --git a/Engine/Shaders/Private/Substrate/Substrate.ush b/Engine/Shaders/Private/Substrate/Substrate.ush
--- a/Engine/Shaders/Private/Substrate/Substrate.ush
+++ b/Engine/Shaders/Private/Substrate/Substrate.ush
@@ -3433,9 +3433,9 @@ void SubstrateSetBSDFSpecularF0(in FSubstrateBSDF BSDF, float3 F0)
 	switch (BSDF_GETTYPE(BSDF))
 	{
 	case SUBSTRATE_BSDF_TYPE_SLAB: 				SLAB_F0(BSDF) = F0; break;
-	case SUBSTRATE_BSDF_TYPE_HAIR: 				HAIR_SPECULAR(BSDF) = F0ToDielectricSpecular(max3(F0.x, F0.y, F0.z)); break;
+	case SUBSTRATE_BSDF_TYPE_HAIR: 				HAIR_SPECULAR(BSDF) = F0RGBToDielectricSpecular(F0); break;
 	case SUBSTRATE_BSDF_TYPE_EYE: 				/* EYE_F0 us hardcoded */ break;
-	case SUBSTRATE_BSDF_TYPE_SINGLELAYERWATER:	SLW_SPECULAR(BSDF) = F0ToDielectricSpecular(max3(F0.x, F0.y, F0.z)); break;
+	case SUBSTRATE_BSDF_TYPE_SINGLELAYERWATER:	SLW_SPECULAR(BSDF) = F0RGBToDielectricSpecular(F0); break;
 	}
 }
 
@@ -3495,7 +3495,7 @@ float SubstrateGetBSDFSpecular(in FSubstrateBSDF BSDF)
 {
 	switch (BSDF_GETTYPE(BSDF))
 	{
-	case SUBSTRATE_BSDF_TYPE_SLAB: 				return F0ToDielectricSpecular(SLAB_F0(BSDF).x);
+	case SUBSTRATE_BSDF_TYPE_SLAB: 				return F0RGBToDielectricSpecular(SLAB_F0(BSDF));
 	case SUBSTRATE_BSDF_TYPE_HAIR: 				return HAIR_SPECULAR(BSDF);
 	case SUBSTRATE_BSDF_TYPE_EYE: 				return F0ToDielectricSpecular(EYE_F0(BSDF));
 	case SUBSTRATE_BSDF_TYPE_SINGLELAYERWATER:	return SLW_SPECULAR(BSDF);
