Substrate - removed legacy style function calling F0 SpecularColor.

#tests PD editor with Substrate
#rnx

diff --git a/Engine/Shaders/Private/RayTracing/RayTracingDebug.usf b/Engine/Shaders/Private/RayTracing/RayTracingDebug.usf
--- a/Engine/Shaders/Private/RayTracing/RayTracingDebug.usf
+++ b/Engine/Shaders/Private/RayTracing/RayTracingDebug.usf
@@ -158,7 +158,7 @@ float4 DebugRayTracingMaterial(uint2 PixelCoord, FRayDesc Ray, bool bIsFarField)
 			case RAY_TRACING_DEBUG_VIZ_SHADING_MODEL:	return float4(GetShadingModelColor(SubstrateGetLegacyShadingModels(BSDF)), 0.0f);
 			case RAY_TRACING_DEBUG_VIZ_BASE_COLOR:		return float4(SubstrateGetBSDFBaseColor(BSDF), 0.0f);
 			case RAY_TRACING_DEBUG_VIZ_DIFFUSE_COLOR:	return float4(SubstrateGetBSDFDiffuseColor(BSDF), 0.0f);
-			case RAY_TRACING_DEBUG_VIZ_SPECULAR_COLOR:	return float4(SubstrateGetBSDFSpecularColor(BSDF).xxx, 0.0f);
+			case RAY_TRACING_DEBUG_VIZ_SPECULAR_COLOR:	return float4(SubstrateGetBSDFSpecularF0(BSDF).xxx, 0.0f);
 			case RAY_TRACING_DEBUG_VIZ_METALLIC:		return float4(SubstrateGetBSDFMetallic(BSDF).xxx, 0.0f);
 			case RAY_TRACING_DEBUG_VIZ_SPECULAR:		return float4(SubstrateGetBSDFSpecular(BSDF).xxx, 0.0f);
 			case RAY_TRACING_DEBUG_VIZ_ROUGHNESS:		return RoughnessDebugColor(SubstrateGetBSDFRoughness(BSDF));
diff --git a/Engine/Shaders/Private/RayTracing/RayTracingPrimaryRays.usf b/Engine/Shaders/Private/RayTracing/RayTracingPrimaryRays.usf
--- a/Engine/Shaders/Private/RayTracing/RayTracingPrimaryRays.usf
+++ b/Engine/Shaders/Private/RayTracing/RayTracingPrimaryRays.usf
@@ -194,7 +194,7 @@ RAY_TRACING_ENTRY_RAYGEN(RayTracingPrimaryRaysRGS)
 
 				PayloadRoughness = SLAB_ROUGHNESS(BSDF);
 				PayloadDiffuseColor = SubstrateGetBSDFDiffuseColor(BSDF);
-				PayloadSpecularColor = SubstrateGetBSDFSpecularColor(BSDF);
+				PayloadSpecularColor = SubstrateGetBSDFSpecularF0(BSDF);
 				PayloadWorldTangent = SubstrateBSDFContext.X;
 				PayloadAnisotropy = SLAB_ANISOTROPY(BSDF);
 
diff --git a/Engine/Shaders/Private/Substrate/Substrate.ush b/Engine/Shaders/Private/Substrate/Substrate.ush
--- a/Engine/Shaders/Private/Substrate/Substrate.ush
+++ b/Engine/Shaders/Private/Substrate/Substrate.ush
@@ -3414,7 +3414,6 @@ float3 SubstrateGetBSDFSpecularF0(in FSubstrateBSDF BSDF)
 	}
 	return 0.f;
 }
-float3 SubstrateGetBSDFSpecularColor(in FSubstrateBSDF BSDF) { return SubstrateGetBSDFSpecularF0(BSDF);  }
 
 float3 SubstrateGetBSDFSpecularF90(in FSubstrateBSDF BSDF)
 {
diff --git a/Engine/Shaders/Private/Substrate/SubstrateTree.ush b/Engine/Shaders/Private/Substrate/SubstrateTree.ush
--- a/Engine/Shaders/Private/Substrate/SubstrateTree.ush
+++ b/Engine/Shaders/Private/Substrate/SubstrateTree.ush
@@ -94,7 +94,7 @@ void UpdateSingleBSDFOperatorCoverageTransmittance(
 	Op.TransmittanceAlongN	= BSDFEvaluate.TransmittanceAlongN;
 
 	// Evaluate refraction lobes.
-	const FBxDFEnergyTerms EnergyTerms = ComputeGGXSpecEnergyTerms(Roughness, SubstrateBSDFContext.SatNoV, SubstrateGetBSDFSpecularColor(CurrentBSDF), SubstrateGetBSDFSpecularF90(CurrentBSDF));
+	const FBxDFEnergyTerms EnergyTerms = ComputeGGXSpecEnergyTerms(Roughness, SubstrateBSDFContext.SatNoV, SubstrateGetBSDFSpecularF0(CurrentBSDF), SubstrateGetBSDFSpecularF90(CurrentBSDF));
 	{
 		const float TopDownInterfaceEta = CurrentBSDF.bIsTop ? AIR_IOR / WATER_IOR : 1.0f;
 		FSubstrateLobeStatistic RefractedLobe = SubstrateGetRefractedLobe(SubstrateGetDiracLobe(V), EnergyTerms.E, Roughness, TopDownInterfaceEta);
diff --git a/Engine/Shaders/Public/SubstratePublic.ush b/Engine/Shaders/Public/SubstratePublic.ush
--- a/Engine/Shaders/Public/SubstratePublic.ush
+++ b/Engine/Shaders/Public/SubstratePublic.ush
@@ -61,7 +61,7 @@ FSubstrateGBuffer SubstratePublic_GetSubstrateGBufferFromFirstBSDF(uint2 PixelCo
 	
 		SGBuffer.bIsValid			= true;
 		SGBuffer.DiffuseColor		= SubstrateGetBSDFDiffuseColor(BSDF);
-		SGBuffer.SpecularColor		= SubstrateGetBSDFSpecularColor(BSDF);
+		SGBuffer.SpecularColor		= SubstrateGetBSDFSpecularF0(BSDF);
 		SGBuffer.SubsurfaceColor	= SubstrateGetBSDFSubSurfaceColor(BSDF);
 		SGBuffer.BaseColor			= SubstrateGetBSDFBaseColor(BSDF);
 		SGBuffer.Specular			= SubstrateGetBSDFSpecular(BSDF);
