Substrate - improved mobile hair shading gbuffer setup.

#tests PC editor dx12 + mobile preview
#rnx

diff --git a/Engine/Shaders/Private/MobileBasePassPixelShader.usf b/Engine/Shaders/Private/MobileBasePassPixelShader.usf
--- a/Engine/Shaders/Private/MobileBasePassPixelShader.usf
+++ b/Engine/Shaders/Private/MobileBasePassPixelShader.usf
@@ -756,6 +756,7 @@ void Main(
 			0.0f,					// Dither
 			Export.ShadingModelID
 		);
+		GBuffer.CustomData = Export.CustomData;
 
 		GBuffer.StoredBaseColor = GBuffer.BaseColor;
 		GBuffer.StoredMetallic = GBuffer.Metallic;
diff --git a/Engine/Shaders/Private/Substrate/SubstrateExport.ush b/Engine/Shaders/Private/Substrate/SubstrateExport.ush
--- a/Engine/Shaders/Private/Substrate/SubstrateExport.ush
+++ b/Engine/Shaders/Private/Substrate/SubstrateExport.ush
@@ -949,8 +949,7 @@ struct FExportResult
 	float  Roughness;
 	float  Anisotropy;
 	float  Metallic;
-	float3 CustomData0;	// Not sur CustomDataX memeber here is the best approach yet.
-	float3 CustomData1;	// Idem
+	float4 CustomData;
 	float3 SubsurfaceColor;
 	uint   SubsurfaceProfileID;
 	float3 WorldTangent;
@@ -982,6 +981,8 @@ FExportResult SubstrateMaterialExportOut(
 
 			if (SubstrateIsBSDFVisible(CurrentBSDF))
 			{
+				const uint BSDFType = BSDF_GETTYPE(CurrentBSDF);
+
 				const float3x3 TangentBasis = SubstrateGetBSDFSharedBasis_InlineShading(SubstratePixelHeader, CurrentBSDF);
 				const float3 N = TangentBasis[2];
 				const float3 T = TangentBasis[0];
@@ -995,13 +996,21 @@ FExportResult SubstrateMaterialExportOut(
 				Out.Specular				= SubstrateGetBSDFSpecular(CurrentBSDF);
 				Out.Roughness				= SubstrateGetBSDFRoughness(CurrentBSDF);
 				Out.Anisotropy				= SubstrateGetBSDFAnisotropy(CurrentBSDF);
+
 				Out.Metallic				= SubstrateGetBSDFMetallic(CurrentBSDF);
+				if (BSDFType == SUBSTRATE_BSDF_TYPE_HAIR)
+				{
+					// SUBSTRATE_TODO Decide is export is a GBuffer (could use FGBuffer then?) or more ligical export of data (have scatter & backlit parameters inside).
+					Out.Metallic = HAIR_SCATTER(CurrentBSDF);
+					Out.CustomData.z = HAIR_BACKLIT(CurrentBSDF);
+				}
+
 				Out.SubsurfaceColor			= SubstrateGetBSDFSubSurfaceColor(CurrentBSDF);
 				Out.SubsurfaceProfileID		= SubstrateGetBSDFSubSurfaceProfileID(CurrentBSDF);
 				Out.WorldTangent			= T;
 				Out.ShadingModelID			= SubstrateGetLegacyShadingModels(CurrentBSDF);
-				Out.CustomData0				= 0;											// SUBSTRATE_TODO not tested since not exposed in the only use case so far: material baking.
-				Out.CustomData1				= 0;											// idem
+
+				// SUBSTRATE_TODO CustomData not tested entirely since not used by material baking and mobile defered seems to have limited materials.
 			}
 		#undef CurrentBSDF
 		}
