Fix pan normal material in vulkan ISR material

diff --git a/Engine/Shaders/Private/DistortApplyScreenPS.usf b/Engine/Shaders/Private/DistortApplyScreenPS.usf
--- a/Engine/Shaders/Private/DistortApplyScreenPS.usf
+++ b/Engine/Shaders/Private/DistortApplyScreenPS.usf
@@ -56,9 +56,18 @@ SamplerState		StandardTranslucentTransmittanceTextureSampler;
 float RefractionMaxMipLevel;
 float StandardDevInPixelForMip0;
 
+// TexCoords [0,1] where 0 is the left edge and 1 is the right edge.
+// This means that in ISR they go over the whole doublewide.
+int2 TexToPixCoords(float2 TexCoord)
+{
+	float2 SampleCoords = (TexCoord.xy * View.BufferSizeAndInvSize.xy);
+	return int2(SampleCoords);
+}
+
 /** distorts screen texture using accumulated distortion offsets */
 void Main(
 	in noperspective float4 TexCoord: TEXCOORD0,
+	in FStereoPSInput StereoInput,
 	out float4 OutColor : SV_Target0
 #if USE_MSAA
 	, in uint SampleIndex : SV_SampleIndex
@@ -66,8 +75,9 @@ void Main(
 	)
 {
 	// sample accumulated distortion and apply inverse scale
+	const int2 PixCoord = TexToPixCoords(TexCoord.xy);
 #if USE_MSAA
-	half4 AccumDist = DistortionMSAATexture.Load(int2(TexCoord.xy * View.BufferSizeAndInvSize.xy), SampleIndex);
+	half4 AccumDist = DistortionMSAATexture.Load(PixCoord, SampleIndex);
 #else
 	half4 AccumDist = Texture2DSample(DistortionTexture,DistortionTextureSampler,TexCoord.xy);
 #endif
@@ -82,8 +92,8 @@ void Main(
 #if USE_ROUGH_REFRACTION && SUBSTRATE_ENABLED
 
 #if USE_MSAA
-	half2 RefractionVarianceCoverage= RefractionVarianceCoverageMSAATexture.Load(int2(TexCoord.xy * View.BufferSizeAndInvSize.xy), SampleIndex).xy;
-	half  RefractionDepthMeter		= RefractionClosestDepthMeterMSAATexture.Load(int2(TexCoord.xy * View.BufferSizeAndInvSize.xy), SampleIndex).x;
+	half2 RefractionVarianceCoverage= RefractionVarianceCoverageMSAATexture.Load(PixCoord, SampleIndex).xy;
+	half  RefractionDepthMeter		= RefractionClosestDepthMeterMSAATexture.Load(PixCoord, SampleIndex).x;
 #else
 	half2 RefractionVarianceCoverage= Texture2DSample(RefractionVarianceCoverageTexture, RefractionVarianceCoverageSampler, TexCoord.xy).xy;
 	half  RefractionDepthMeter		= Texture2DSample(RefractionClosestDepthMeterTexture, RefractionClosestDepthMeterSampler, TexCoord.xy).x;
@@ -192,7 +202,7 @@ void Main(
 
 	// sample screen using offset coords
 #if USE_MSAA
-	OutColor = SceneColorMSAATexture.Load(int2(NewBufferUV * View.BufferSizeAndInvSize.xy), SampleIndex);
+	OutColor = SceneColorMSAATexture.Load(TexToPixCoords(NewBufferUV), SampleIndex);
 #else
 	OutColor = Texture2DSample(SceneColorTexture, SceneColorTextureSampler, NewBufferUV);
 #endif
@@ -202,6 +212,7 @@ void Main(
 
 void Merge(
 	in noperspective float4 TexCoord: TEXCOORD0,
+	in FStereoPSInput StereoInput,
 	out float4 OutColor : SV_Target0
 #if USE_MSAA
 	, in uint SampleIndex : SV_SampleIndex
@@ -209,8 +220,7 @@ void Merge(
 	)
 {
 #if USE_MSAA
-
-	int2 PixCoord = int2(View.BufferSizeAndInvSize.xy * TexCoord.xy);
+	int2 PixCoord = TexToPixCoords(TexCoord.xy);
 	OutColor = SceneColorMSAATexture.Load(PixCoord, SampleIndex);
 #if COMPOSITE_STANDARD_TRANSLUCENT
 	const float4 TranslucentColorAlpha		= StandardTranslucentColorAlphaMSAATexture.Load(PixCoord, SampleIndex);
diff --git a/Engine/Source/Runtime/Renderer/Private/DistortionRendering.cpp b/Engine/Source/Runtime/Renderer/Private/DistortionRendering.cpp
--- a/Engine/Source/Runtime/Renderer/Private/DistortionRendering.cpp
+++ b/Engine/Source/Runtime/Renderer/Private/DistortionRendering.cpp
@@ -892,7 +892,7 @@ void FDeferredShadingSceneRenderer::RenderDistortion(
 			Viewport.Rect = View.ViewRect;
 
 			ClearUnusedGraphResources(ApplyPixelShader, PassParameters);
-			AddDrawScreenPass(GraphBuilder, {}, View, Viewport, Viewport, PipelineState, PassParameters,
+			AddDrawScreenPass(GraphBuilder, {}, FScreenPassViewInfo(View, 1), Viewport, Viewport, PipelineState, PassParameters,
 				[ApplyPixelShader, PassParameters](FRHICommandList& RHICmdList)
 			{
 				SetShaderParameters(RHICmdList, ApplyPixelShader, ApplyPixelShader.GetPixelShader(), *PassParameters);
@@ -945,7 +945,7 @@ void FDeferredShadingSceneRenderer::RenderDistortion(
 			Viewport.Rect = View.ViewRect;
 
 			ClearUnusedGraphResources(MergePixelShader, PassParameters);
-			AddDrawScreenPass(GraphBuilder, {}, View, Viewport, Viewport, PipelineState, PassParameters,
+			AddDrawScreenPass(GraphBuilder, {}, FScreenPassViewInfo(View, 1), Viewport, Viewport, PipelineState, PassParameters,
 				[MergePixelShader, PassParameters](FRHICommandList& RHICmdList)
 			{
 				SetShaderParameters(RHICmdList, MergePixelShader, MergePixelShader.GetPixelShader(), *PassParameters);
diff --git a/Engine/Source/Runtime/Renderer/Public/ScreenPass.h b/Engine/Source/Runtime/Renderer/Public/ScreenPass.h
--- a/Engine/Source/Runtime/Renderer/Public/ScreenPass.h
+++ b/Engine/Source/Runtime/Renderer/Public/ScreenPass.h
@@ -452,6 +452,14 @@ struct FScreenPassViewInfo
 		, bHMDHiddenAreaMaskActive(View.bHMDHiddenAreaMaskActive)
 	{
 	}
+
+	FScreenPassViewInfo(const FSceneView& View, const uint32 OverrideInstanceCount)
+		: FeatureLevel(View.FeatureLevel)
+		, StereoViewIndex(View.StereoViewIndex)
+		, InstanceCount(OverrideInstanceCount)
+		, bHMDHiddenAreaMaskActive(View.bHMDHiddenAreaMaskActive)
+	{
+	}
 };
 
 RENDERER_API void DrawScreenPass_PostSetup(
