Add detail on Subsurface behavior for Substrate.

diff --git a/Engine/Shaders/Private/Substrate/Substrate.ush b/Engine/Shaders/Private/Substrate/Substrate.ush
--- a/Engine/Shaders/Private/Substrate/Substrate.ush
+++ b/Engine/Shaders/Private/Substrate/Substrate.ush
@@ -155,6 +155,59 @@
 // Disable SSS diffusion based when the pixel's footprint is larger than x times the MFPs
 #define SUBSTRATE_DISABLE_SSS_DIFFUSION_BASED_ON_PIXEL_FOOTPRINT 0
 
+///////////////////////////////////////////////////////////////////////////////
+// Sub-surface scattering
+//
+//  Material configuration 
+//  ----------------------
+//  * Use Subsurface Diffusion (on Slab)
+//  * Is Thin Surface (on Material)
+//  * Top layer/Bottom layer (on Material topology)
+//
+//  Material Parameters (Slab)
+//  -------------------
+//  * Diffuse Albedo
+//  * SSS MFP
+//  * SSS MFP scale
+//  * SSS Phase aniso
+// 
+//  SSS techniques
+//  --------------
+//  * Diffusion         
+//  * Diffusion Profile     
+//  * Wrap (old Foliage)    
+//  * Wrap (ols Subsurface) 
+//  * Simple Volume         
+// 
+//  Techniques used for each configuration 
+//
+// 	|-------------------------------------------------------------------------------------------|-------------------------------------------------------------------------------------------|
+// 	| Diffusion                                                                                 | !Diffusion                                                                                | 
+// 	|---------------------------------------------|---------------------------------------------|---------------------------------------------|---------------------------------------------|
+// 	| !IsThin                                     | IsThin                                      | !IsThin                                     | IsThin                                      |
+// 	|---------------------------------------------|---------------------------------------------|---------------------------------------------|---------------------------------------------|
+// 	| Top                  | Bottom               | Top                  | Bottom               | Top                  | Bottom               | Top                  | Bottom               |
+// 	|----------------------|----------------------|----------------------|----------------------|----------------------|----------------------|----------------------|----------------------|
+// 	| Simple Volume        | Diffusion            | Simple Volume        | Diffusion            | Simple Volume        | Wrap (old Subsurface)| Simple Volume        | Wrap (old Foliage)   |
+// 
+//
+// How are stored the MFP data?
+// ---------------------------
+//
+//                          | MaterialData                                             | SSSData
+//  ------------------------ ---------------------------------------------------------- -----------------
+//  * Diffusion Profile     | ProfileID + Scale                                        | BaseColor + ProfileID + Scale
+//  * Diffusion             | Rescaled MFP (11/11/10)                                  | BaseColor + MFP x Coverage
+//  * Wrap (old Foliage)    | Rescaled MFP (11/11/10)                                  | BaseColor + Opacity
+//                          | Rescaled MFP + Opacity (30bits + 7bits) - (LegacyLayout) | BaseColor + Opacity
+//  * Wrap (ols Subsurface) | Rescaled MFP (11/11/10)                                  | BaseColor + Opacity
+//                          | Rescaled MFP + Opacity (30bits + 7bits) - (LegacyLayout) | BaseColor + Opacity
+//  * Simple Volume         | Rescaled MFP                                             | NoStorage
+//
+// Note: When using non-profile SSS diffusion, rescale MFP so that we take into account the slab thickness. 
+//       This is done 'after' storing the SSSHeader/SSSData, as the original MFP needs to be preserved for correct post-process diffusion.
+//       This rescaled MFP value is only used for transmission evaluation, not the diffusion
+
 ///////////////////////////////////////////////////////////////////////////////
 // BSDFs representation
 
@@ -3156,30 +3209,6 @@ void FSubstrateBSDF::SubstrateSanitizeBSDF()
 	}
 }
 
-// Description of SSS conversions
-//
-// Opaque
-// ======
-//               |                     Bottom                      |         !Bottom          |
-//               |          Thin          |          Thick         |                          |
-//               |----------------------------------------------------------------------------|
-// SSS Diffusion |        SSS / MFP       |       SSS / MFP        |  SimpleVolume / Res.MFP  | 
-// SSS Profile   |        SSS / MFP       |       SSS / MFP        |  SimpleVolume / Res.MFP  | 
-// SSS Wrap      |      SSS / Res.MFP     |     SSS / Res.MFP      |  SimpleVolume / Res.MFP  | 
-// SSS Invalid   |                        |                        |                          | 
-//
-// Translucent
-// ===========
-//               |                     Bottom                      |         !Bottom          |
-//               |          Thin          |          Thick         |                          |
-//               |----------------------------------------------------------------------------|
-// SSS Diffusion | SimpleVolume / Res.MFP | SimpleVolume / Res.MFP |  SimpleVolume / Res.MFP  | 
-// SSS Profile   | SimpleVolume / Res.MFP | SimpleVolume / Res.MFP |  SimpleVolume / Res.MFP  | 
-// SSS Wrap      | SimpleVolume / Res.MFP | SimpleVolume / Res.MFP |  SimpleVolume / Res.MFP  | 
-// SSS Invalid   |                        |                        |                          | 
-
-//  Res.MFP : Rescaled MFP in computation space i.e., MFP is rescaled for a thickness of 1m
-
 // Convert to simple volume to get some transmittance
 void FSubstrateBSDF::PostProcessSlabBeforeLighting(bool bIsSubstrateOpaqueMaterial)
 {
