Substrate - now support Subsurface, Two sided, clear coat and cloth legacy shading models.

#tests PC editor dx12, mobile preview vulkan
#rnx

diff --git a/Engine/Shaders/Private/BasePassPixelShader.usf b/Engine/Shaders/Private/BasePassPixelShader.usf
--- a/Engine/Shaders/Private/BasePassPixelShader.usf
+++ b/Engine/Shaders/Private/BasePassPixelShader.usf
@@ -1644,7 +1644,9 @@ void FPixelShaderInOut_MainPS(
 		Settings,
 		SubstratePixelHeader,
 		SubstrateData,
-		SurfaceWorldNormal);
+		SurfaceWorldNormal,
+		MaterialParameters.WorldPosition_CamRelative,
+		0.0f /*MobileShadingPathCurvature*/);
 
 	#if SUBSTRATE_MATERIAL_EXPORT_FROM_OPAQUE
 		#if	  SUBSTRATE_MATERIAL_EXPORT_TYPE == SUBSTRATE_MATERIAL_EXPORT_BASE_COLOR_POST_COVERAGE
diff --git a/Engine/Shaders/Private/DeferredDecal.usf b/Engine/Shaders/Private/DeferredDecal.usf
--- a/Engine/Shaders/Private/DeferredDecal.usf
+++ b/Engine/Shaders/Private/DeferredDecal.usf
@@ -202,7 +202,9 @@ void FPixelShaderInOut_MainPS(inout FPixelShaderIn In, inout FPixelShaderOut Out
 			Settings,
 			SubstrateHeader,
 			PixelMaterialInputs.GetFrontSubstrateData(),
-			SurfaceWorldNormal);
+			SurfaceWorldNormal,
+			MaterialParameters.WorldPosition_CamRelative,
+			0.0f /*MobileShadingPathCurvature*/);
 		
 		// Either propagate the BSDF opacity, or use the one for AlphaComposite (premultiplied alpha)
 		float CoverageOverBackground = Export.Coverage * DecalFading;
diff --git a/Engine/Shaders/Private/MeshDecals.usf b/Engine/Shaders/Private/MeshDecals.usf
--- a/Engine/Shaders/Private/MeshDecals.usf
+++ b/Engine/Shaders/Private/MeshDecals.usf
@@ -155,7 +155,9 @@ void FPixelShaderInOut_MainPS(
 			Settings,
 			SubstrateHeader,
 			PixelMaterialInputs.GetFrontSubstrateData(),
-			SurfaceWorldNormal);
+			SurfaceWorldNormal,
+			MaterialParameters.WorldPosition_CamRelative,
+			0.0f /*MobileShadingPathCurvature*/);
 		
 		// Either propagate the BSDF opacity, or use the one for AlphaComposite (premultiplied alpha)
 		float CoverageOverBackground = Export.Coverage;
diff --git a/Engine/Shaders/Private/MobileBasePassPixelShader.usf b/Engine/Shaders/Private/MobileBasePassPixelShader.usf
--- a/Engine/Shaders/Private/MobileBasePassPixelShader.usf
+++ b/Engine/Shaders/Private/MobileBasePassPixelShader.usf
@@ -727,7 +727,14 @@ void Main(
 			Settings,
 			SubstratePixelHeader,
 			SubstrateData,
-			SurfaceWorldNormal);
+			SurfaceWorldNormal,
+			MaterialParameters.WorldPosition_CamRelative,
+		#if SHADING_PATH_MOBILE && MATERIAL_SUBSURFACE_PROFILE_USE_CURVATURE
+			GetMaterialCustomData0(MaterialParameters)
+		#else
+			0.0f
+		#endif
+			);
 
 		GBuffer.GBufferAO = MaterialAO;
 		GBuffer.Depth = MaterialParameters.ScreenPosition.w;
@@ -740,8 +747,6 @@ void Main(
 	#endif
 		MaterialParameters.TangentToWorld[2].xyz = Export.WorldTangent;
 
-		// SUBSTRATE_TODO when clear coat or cloth are clarified for mobile deferred, SetGBufferForShadingModel should read data directly from Export. It will depend on material conversion features and mobile deferred support.
-		// Right now it uses GetMaterialCustomData0(MaterialParameters) to get some specific data. Which does not work with Substrate (but it also do not work for mobile deferred it seems).
 		SetGBufferForShadingModel(
 			GBuffer,
 			MaterialParameters,
@@ -756,6 +761,7 @@ void Main(
 			0.0f,					// Dither
 			Export.ShadingModelID
 		);
+		// Set custom data from Substrate export after SetGBufferForShadingModel since it access MaterialParameters to set specific data not available with Substrate
 		GBuffer.CustomData = Export.CustomData;
 
 		GBuffer.StoredBaseColor = GBuffer.BaseColor;
@@ -883,7 +889,7 @@ void Main(
 		FLightAccumulator DirectLighting = (FLightAccumulator)0;
 		LightAccumulator_AddSplit(DirectLighting, DiffuseColor, 0.0f, DiffuseColor, 1.0f, false);
 
-		GBuffer.IndirectIrradiance = IndirectIrradiance;
+		GBuffer.IndirectIrradiance = IndirectIrradiance * Export.IndirectIrradianceScale;
 		MobileEncodeGBuffer(GBuffer, OutGBufferA, OutGBufferB, OutGBufferC, OutGBufferD);
 
 
diff --git a/Engine/Shaders/Private/PathTracing/PathTracingGPULightmassMaterialHitShader.usf b/Engine/Shaders/Private/PathTracing/PathTracingGPULightmassMaterialHitShader.usf
--- a/Engine/Shaders/Private/PathTracing/PathTracingGPULightmassMaterialHitShader.usf
+++ b/Engine/Shaders/Private/PathTracing/PathTracingGPULightmassMaterialHitShader.usf
@@ -94,7 +94,9 @@ RAY_TRACING_ENTRY_CLOSEST_HIT(GPULightmassMaterialCHS,
 		Settings,
 		SubstratePixelHeader,
 		SubstrateData,
-		SurfaceWorldNormal);
+		SurfaceWorldNormal,
+		MaterialParameters.WorldPosition_CamRelative,
+		0.0f /*MobileShadingPathCurvature*/);
 
 	Coverage = Export.Coverage;
 	TransmittancePreCoverage = Export.TransmittancePreCoverage;
@@ -408,7 +410,9 @@ RAY_TRACING_ENTRY_ANY_HIT(GPULightmassMaterialAHS,
 		Settings,
 		SubstratePixelHeader,
 		SubstrateData,
-		SurfaceWorldNormal);
+		SurfaceWorldNormal,
+		MaterialParameters.WorldPosition_CamRelative,
+		0.0f /*MobileShadingPathCurvature*/);
 
 	Coverage = Export.Coverage;
 	TransmittancePreCoverage = Export.TransmittancePreCoverage;
diff --git a/Engine/Shaders/Private/Substrate/Substrate.ush b/Engine/Shaders/Private/Substrate/Substrate.ush
--- a/Engine/Shaders/Private/Substrate/Substrate.ush
+++ b/Engine/Shaders/Private/Substrate/Substrate.ush
@@ -3509,11 +3509,10 @@ uint SubstrateGetBSDFSubSurfaceProfileID(in FSubstrateBSDF BSDF)
 	case SUBSTRATE_BSDF_TYPE_SLAB:
 	case SUBSTRATE_BSDF_TYPE_EYE:
 	{
-		if (BSDF_GETSSSTYPE(BSDF) != SSS_TYPE_DIFFUSION_PROFILE)
+		if (BSDF_GETSSSTYPE(BSDF) == SSS_TYPE_DIFFUSION_PROFILE)
 		{
 			return SubstrateSubsurfaceProfileIdTo8bits(SLAB_SSSPROFILEID(BSDF));
 		}
-		return 0;
 	}
 	}
 	return 0;
diff --git a/Engine/Shaders/Private/Substrate/SubstrateExport.ush b/Engine/Shaders/Private/Substrate/SubstrateExport.ush
--- a/Engine/Shaders/Private/Substrate/SubstrateExport.ush
+++ b/Engine/Shaders/Private/Substrate/SubstrateExport.ush
@@ -954,6 +954,8 @@ struct FExportResult
 	uint   SubsurfaceProfileID;
 	float3 WorldTangent;
 	float  ShadingModelID;
+	float IndirectIrradianceScale;
+	float Curvature;
 };
 
 // This is used to export very simple data for lightmaps, material baking or preview on editor node.
@@ -962,9 +964,12 @@ FExportResult SubstrateMaterialExportOut(
 	in FSubstrateIntegrationSettings Settings,
 	in FSubstratePixelHeader SubstratePixelHeader,
 	in FSubstrateData SubstrateData,
-	in float3 SurfaceWorldNormal)
+	in float3 SurfaceWorldNormal,
+	in float3 WorldPosition_CamRelative,
+	in float MobileShadingPathCurvature)
 {
 	FExportResult Out = (FExportResult)0;
+	Out.IndirectIrradianceScale = 1.0;
 
 	// Align V with N to not have any view dependent effect when exporting material data.
 	const float3 FakeV = SurfaceWorldNormal;
@@ -982,6 +987,7 @@ FExportResult SubstrateMaterialExportOut(
 			if (SubstrateIsBSDFVisible(CurrentBSDF))
 			{
 				const uint BSDFType = BSDF_GETTYPE(CurrentBSDF);
+				const uint SSSType  = BSDF_GETSSSTYPE(CurrentBSDF);
 
 				const float3x3 TangentBasis = SubstrateGetBSDFSharedBasis_InlineShading(SubstratePixelHeader, CurrentBSDF);
 				const float3 N = TangentBasis[2];
@@ -996,21 +1002,156 @@ FExportResult SubstrateMaterialExportOut(
 				Out.Specular				= SubstrateGetBSDFSpecular(CurrentBSDF);
 				Out.Roughness				= SubstrateGetBSDFRoughness(CurrentBSDF);
 				Out.Anisotropy				= SubstrateGetBSDFAnisotropy(CurrentBSDF);
-
 				Out.Metallic				= SubstrateGetBSDFMetallic(CurrentBSDF);
-				if (BSDFType == SUBSTRATE_BSDF_TYPE_HAIR)
-				{
-					// SUBSTRATE_TODO Decide is export is a GBuffer (could use FGBuffer then?) or more ligical export of data (have scatter & backlit parameters inside).
-					Out.Metallic = HAIR_SCATTER(CurrentBSDF);
-					Out.CustomData.z = HAIR_BACKLIT(CurrentBSDF);
-				}
+
 
 				Out.SubsurfaceColor			= SubstrateGetBSDFSubSurfaceColor(CurrentBSDF);
 				Out.SubsurfaceProfileID		= SubstrateGetBSDFSubSurfaceProfileID(CurrentBSDF);
 				Out.WorldTangent			= T;
 				Out.ShadingModelID			= SubstrateGetLegacyShadingModels(CurrentBSDF);
 
-				// SUBSTRATE_TODO CustomData not tested entirely since not used by material baking and mobile defered seems to have limited materials.
+				if (BSDFType == SUBSTRATE_BSDF_TYPE_HAIR)
+				{
+					Out.Metallic	= HAIR_SCATTER(CurrentBSDF);
+					Out.CustomData.z= HAIR_BACKLIT(CurrentBSDF);
+				}
+				else if (SSSType == SSS_TYPE_WRAP) // Thin and not thin, Subsurface and TwoSided Foliage
+				{
+					Out.CustomData.rgb = EncodeSubsurfaceColor(Out.SubsurfaceColor);
+					Out.CustomData.a   = SubstrateSubSurfaceGetWrapOpacityFromAnisotropy(SLAB_SSSPHASEANISOTROPY(CurrentBSDF));
+				}
+				else if (BSDF_GETHASFUZZ(CurrentBSDF))
+				{
+					Out.CustomData.rgb = EncodeSubsurfaceColor(SLAB_FUZZ_COLOR(CurrentBSDF));
+					Out.CustomData.a   = SLAB_FUZZ_AMOUNT(CurrentBSDF); 
+					//MakeRoughnessSafe(SLAB_FUZZ_ROUGHNESS(CurrentBSDF), SUBSTRATE_MIN_FUZZ_ROUGHNESS);
+					Out.IndirectIrradianceScale *= 1 - Out.CustomData.a;
+				}
+				else if (BSDF_GETHASHAZINESS(CurrentBSDF))
+				{
+					FHaziness LocalHaziness = UnpackHaziness(SLAB_HAZINESS(CurrentBSDF));
+					if (LocalHaziness.bSimpleClearCoat)
+					{
+						Out.CustomData.x = LocalHaziness.Weight;
+						Out.CustomData.y = LocalHaziness.Roughness;
+
+					//#if CLEAR_COAT_BOTTOM_NORMAL
+					//	// SUBSTRATE_TODO, both top and bottom normals have been set on two slabs that have been simplified...
+					//	float2 oct2 = UnitVectorToOctahedron(N);
+					//
+					//	#if NUM_MATERIAL_OUTPUTS_CLEARCOATBOTTOMNORMAL > 0
+					//		#if MATERIAL_TANGENTSPACENORMAL
+					//			float3 tempnormal = normalize(TransformTangentVectorToWorld( MaterialParameters.TangentToWorld, ClearCoatBottomNormal0(MaterialParameters) ));
+					//		#else
+					//			float3 tempnormal =  ClearCoatBottomNormal0(MaterialParameters);
+					//		#endif
+					//
+					//		float2 oct1 = UnitVectorToOctahedron(tempnormal);
+					//		float2 oct3 = ( (oct1 - oct2) *  0.25 ) + (128.0/255.0);
+					//		Out.CustomData.a = oct3.x;
+					//		Out.CustomData.z = oct3.y;
+					//	#else
+							Out.CustomData.a = 128.0/255.0;
+							Out.CustomData.z = 128.0/255.0;
+					//	#endif
+					//#endif
+					}
+				}
+				else if (BSDFType == SUBSTRATE_BSDF_TYPE_EYE)
+				{
+					const float3 IrisNormal		= EYE_IRISNORMAL(CurrentBSDF);
+					const float3 IrisPlaneNormal= EYE_IRISPLANENORMAL(CurrentBSDF);
+					const float IrisMask		= saturate(EYE_IRISMASK(CurrentBSDF));
+					const float IrisDistance	= saturate(EYE_IRISDISTANCE(CurrentBSDF));
+
+					Out.CustomData.x = EncodeSubsurfaceProfile(Out.SubsurfaceProfileID).x;
+					Out.CustomData.w = 1.0f - IrisMask;	// Opacity
+
+				#if IRIS_NORMAL
+					const float3 GBufferWorldNormal = N;
+					float2 WorldNormalOct = UnitVectorToOctahedron(GBufferWorldNormal);
+
+					// CausticNormal stored as octahedron
+					#if NUM_MATERIAL_OUTPUTS_GETTANGENTOUTPUT > 0
+						// Blend in the negative intersection normal to create some concavity
+						// Not great as it ties the concavity to the convexity of the cornea surface
+						// No good justification for that. On the other hand, if we're just looking to
+						// introduce some concavity, this does the job.
+						float3 PlaneNormal = normalize( T /*GetTangentOutput0(MaterialParameters)*/);
+						float3 CausticNormal = normalize( lerp( PlaneNormal, -GBufferWorldNormal, IrisMask*IrisDistance ) );
+						float2 CausticNormalOct  = UnitVectorToOctahedron( CausticNormal );
+						float2 CausticNormalDelta = ( CausticNormalOct - WorldNormalOct ) * 0.5 + (128.0/255.0);
+						Out.Metallic = CausticNormalDelta.x;
+						Out.Specular = CausticNormalDelta.y;
+					#else
+						float3 PlaneNormal = GBufferWorldNormal;
+						Out.Metallic = 128.0/255.0;
+						Out.Specular = 128.0/255.0;
+					#endif
+
+					// IrisNormal CustomData.yz
+					//#if NUM_MATERIAL_OUTPUTS_CLEARCOATBOTTOMNORMAL > 0
+					//	float3 IrisNormal = normalize( ClearCoatBottomNormal0(MaterialParameters) );
+					//	#if MATERIAL_TANGENTSPACENORMAL
+					//		IrisNormal = normalize( TransformTangentVectorToWorld( MaterialParameters.TangentToWorld, IrisNormal ) );
+					//	#endif
+					//#else
+					//	float3 IrisNormal = PlaneNormal;
+					//#endif
+
+					float2 IrisNormalOct  = UnitVectorToOctahedron( IrisNormal );
+					float2 IrisNormalDelta = ( IrisNormalOct - WorldNormalOct ) * 0.5 + (128.0/255.0);
+					Out.CustomData.yz = IrisNormalDelta;
+				#else
+					Out.Metallic = IrisDistance;
+
+					#if NUM_MATERIAL_OUTPUTS_GETTANGENTOUTPUT > 0
+						float3 Tangent = T /*GetTangentOutput0(MaterialParameters)*/;
+						Out.CustomData.yz = UnitVectorToOctahedron( normalize(Tangent) ) * 0.5 + 0.5;
+					#endif
+				#endif
+
+				#if SHADING_PATH_MOBILE
+					#if MATERIAL_SHADINGMODEL_EYE_USE_CURVATURE
+						Out.Curvature = Metallic;
+					#else
+						Out.Curvature = CalculateCurvature(GBufferWorldNormal, WorldPosition_CamRelative);
+					#endif
+
+						Out.Curvature = clamp(Out.Curvature, 0.001f, 1.0f);
+				#endif
+				}
+				else if (SSSType == SSS_TYPE_DIFFUSION_PROFILE) // Thin and not thin, Subsurface and TwoSided Foliage
+				{
+					const float3 GBufferWorldNormal = N;
+					const float Opacity = SLAB_SSSPROFILERADIUSSCALE(CurrentBSDF);
+					
+					// Optimization: if opacity is 0 then revert to default shading model 
+				#if SUBSURFACE_PROFILE_OPACITY_THRESHOLD
+					if (Opacity > SSSS_OPACITY_THRESHOLD_EPS)
+				#endif
+					{
+						Out.CustomData.rgb = EncodeSubsurfaceColor(Out.SubsurfaceColor);
+						Out.CustomData.a = Opacity;
+
+					#if SHADING_PATH_MOBILE
+						#if MATERIAL_SUBSURFACE_PROFILE_USE_CURVATURE
+						Out.Curvature = MobileShadingPathCurvature;
+						#else
+						Out.Curvature = CalculateCurvature(GBufferWorldNormal, WorldPosition_CamRelative);
+						#endif
+
+						Out.Curvature = clamp(Out.Curvature, 0.001f, 1.0f);
+					#endif
+					}
+				#if SUBSURFACE_PROFILE_OPACITY_THRESHOLD
+					else
+					{
+						Out.ShadingModelID = SHADINGMODELID_DEFAULT_LIT;
+						Out.CustomData = 0;
+					}
+				#endif
+				}
 			}
 		#undef CurrentBSDF
 		}
