Substrate - Fixed SSR not showing when there is no skylight in the scene. (EnvBRDF energy was skipped).

#tests PC dx12 editor with/wo Lumen.
#rnx

diff --git a/Engine/Shaders/Private/ReflectionEnvironmentPixelShader.usf b/Engine/Shaders/Private/ReflectionEnvironmentPixelShader.usf
--- a/Engine/Shaders/Private/ReflectionEnvironmentPixelShader.usf
+++ b/Engine/Shaders/Private/ReflectionEnvironmentPixelShader.usf
@@ -476,7 +476,8 @@ void ReflectionEnvironmentSkyLighting(
 		uint NumCulledReflectionCaptures = 0;
 	#endif
 
-		const bool bEnableSpecular = ReflectionStruct.SkyLightParameters.y > 0.0f || NumCulledReflectionCaptures > 0;
+		const bool bPixelMayHaveSSR = REFLECTION_SOURCE_TYPE == REFLECTION_SOURCE_LUMEN_STANDALONE || TopLayerSpecularContributionFactor < 1.0f;
+		const bool bEnableSpecular = ReflectionStruct.SkyLightParameters.y > 0.0f || NumCulledReflectionCaptures > 0 || bPixelMayHaveSSR;
 	
 		Substrate_for (uint ClosureIndex = 0, ClosureIndex < SubstratePixelHeader.ClosureCount, ++ClosureIndex)
 		{
