[TSR] Fix TSR visualization mode flickering temporal analysis (r.tsr.visualize 7) not working.
* Match the logic to ComputeMoireError function due to the shift to non-gama encoding space.

#tests Test done in city sample, we can now see the moire antiflickering region on detailed building.
[RN] Fix TSR visualization mode flickering temporal analysis not working

diff --git a/Engine/Shaders/Private/TemporalSuperResolution/TSRColorSpace.ush b/Engine/Shaders/Private/TemporalSuperResolution/TSRColorSpace.ush
--- a/Engine/Shaders/Private/TemporalSuperResolution/TSRColorSpace.ush
+++ b/Engine/Shaders/Private/TemporalSuperResolution/TSRColorSpace.ush
@@ -98,6 +98,13 @@ tsr_half4 PreExposureCorrectGCS(tsr_half4 GColor, const float PreExposureCorrect
 
 //------------------------------------------------------- GUIDE COLOR SPACE <-> SHADING MEASUREMENT COLOR SPACE
 
+CALL_SITE_DEBUGLOC
+tsr_half GCSToSMCS(tsr_half GColor)
+{
+	tsr_half SMColor = GColor * GColor;
+	return SMColor;
+}
+
 CALL_SITE_DEBUGLOC
 tsr_half3 GCSToSMCS(tsr_half3 GColor)
 {
diff --git a/Engine/Shaders/Private/TemporalSuperResolution/TSRVisualize.usf b/Engine/Shaders/Private/TemporalSuperResolution/TSRVisualize.usf
--- a/Engine/Shaders/Private/TemporalSuperResolution/TSRVisualize.usf
+++ b/Engine/Shaders/Private/TemporalSuperResolution/TSRVisualize.usf
@@ -332,7 +332,7 @@ tsr_half3 DisplayAntiFlickering(FContext Ctx)
 	tsr_half DisableAntiFlickering = tsr_half(1.0) - tsr_half(IsMovingMaskTexture[Ctx.InputPixelPos]) * rcp(tsr_half(255.0));
 
 	// Unpack history
-	tsr_half FlickeringHistory   = MoireHistory[0];
+	tsr_half FlickeringHistory   = GCSToSMCS(MoireHistory[0]);
 	tsr_half Gradient            = MoireHistory[1] * tsr_half(255.0 / 127.0) - tsr_half(1.0);
 	tsr_half TotalVariation      = MoireHistory[2] * (tsr_half(1.0) - BlendFinal);
 	tsr_half TotalVariationCount = MoireHistory[3] * MaxPrevTotalVariationCount * (tsr_half(1.0) - BlendFinal);
@@ -349,8 +349,10 @@ tsr_half3 DisplayAntiFlickering(FContext Ctx)
 		tsr_half SceneColorLuma = dot(OriginalOpaqueInput.rgb, kMoireSMCSWeights);
 		tsr_half TranslucencyChange = dot(abs(OriginalOpaqueInput.rgb - OriginalInput.rgb), kMoireSMCSWeights);
 		
+		tsr_half InputLuma = GCSToSMCS(MoireInput);
+
 		InputLumaModification = tsr_half(0.0);
-		InputLumaModification = max(InputLumaModification, abs(MoireInput - SceneColorLuma) - InputLumaEncodingError);
+		InputLumaModification = max(InputLumaModification, abs(InputLuma - SceneColorLuma) - InputLumaEncodingError);
 		InputLumaModification = max(InputLumaModification, TranslucencyChange);
 		InputLumaModification = max(InputLumaModification, tsr_half(1.0) - OriginalTranslucencyInput[3]);
 	}
