Changed FVirtualShadowMapProjectionShaderData loading to use helper functions in ByteBuffer.ush and add float4x4 loading helper.

#rnx

diff --git a/Engine/Shaders/Private/ByteBuffer.ush b/Engine/Shaders/Private/ByteBuffer.ush
--- a/Engine/Shaders/Private/ByteBuffer.ush
+++ b/Engine/Shaders/Private/ByteBuffer.ush
@@ -53,3 +53,11 @@ void LoadAndIncrementOffset(inout float4 Value, ByteAddressBuffer InBuffer, inou
 	LoadAndIncrementOffset(UintValue, InBuffer, Offset);
 	Value = asfloat(UintValue);
 }
+
+void LoadAndIncrementOffset(inout float4x4 Value, ByteAddressBuffer InBuffer, inout uint Offset)
+{	
+	LoadAndIncrementOffset(Value[0], InBuffer, Offset);
+	LoadAndIncrementOffset(Value[1], InBuffer, Offset);
+	LoadAndIncrementOffset(Value[2], InBuffer, Offset);
+	LoadAndIncrementOffset(Value[3], InBuffer, Offset);
+}
diff --git a/Engine/Shaders/Private/VirtualShadowMaps/VirtualShadowMapProjectionStructs.ush b/Engine/Shaders/Private/VirtualShadowMaps/VirtualShadowMapProjectionStructs.ush
--- a/Engine/Shaders/Private/VirtualShadowMaps/VirtualShadowMapProjectionStructs.ush
+++ b/Engine/Shaders/Private/VirtualShadowMaps/VirtualShadowMapProjectionStructs.ush
@@ -3,6 +3,7 @@
 #pragma once
 
 #include "../LargeWorldCoordinates.ush"
+#include "../ByteBuffer.ush"
 #include "/Engine/Shared/VirtualShadowMapDefinitions.h"
 
 // Rougly the same structure as in VirtualShadowMapArray.h, with some logic around LWC values
@@ -48,60 +49,53 @@ FVirtualShadowMapProjectionShaderData DecodeVirtualShadowMapProjectionData(ByteA
 
 	// Seems the FMatrix forces 16-byte alignment, so ensure padding matches the C++ structure size
 	const uint Stride = 16 * 18;
-	const uint Base = VirtualShadowMapId * Stride;
-
-	Result.ShadowViewToClipMatrix[0]				= asfloat(ProjectionData.Load4(Base + 16*0));
-	Result.ShadowViewToClipMatrix[1]				= asfloat(ProjectionData.Load4(Base + 16*1));
-	Result.ShadowViewToClipMatrix[2]				= asfloat(ProjectionData.Load4(Base + 16*2));
-	Result.ShadowViewToClipMatrix[3]				= asfloat(ProjectionData.Load4(Base + 16*3));
-	Result.TranslatedWorldToShadowUVMatrix[0]		= asfloat(ProjectionData.Load4(Base + 16*4));
-	Result.TranslatedWorldToShadowUVMatrix[1]		= asfloat(ProjectionData.Load4(Base + 16*5));
-	Result.TranslatedWorldToShadowUVMatrix[2]		= asfloat(ProjectionData.Load4(Base + 16*6));
-	Result.TranslatedWorldToShadowUVMatrix[3]		= asfloat(ProjectionData.Load4(Base + 16*7));
-	Result.TranslatedWorldToShadowUVNormalMatrix[0]	= asfloat(ProjectionData.Load4(Base + 16*8));
-	Result.TranslatedWorldToShadowUVNormalMatrix[1]	= asfloat(ProjectionData.Load4(Base + 16*9));
-	Result.TranslatedWorldToShadowUVNormalMatrix[2]	= asfloat(ProjectionData.Load4(Base + 16*10));
-	Result.TranslatedWorldToShadowUVNormalMatrix[3]	= asfloat(ProjectionData.Load4(Base + 16*11));
+	uint CurrentLoadOffset = VirtualShadowMapId * Stride;
+	LoadAndIncrementOffset(Result.ShadowViewToClipMatrix, ProjectionData, CurrentLoadOffset);
+	LoadAndIncrementOffset(Result.TranslatedWorldToShadowUVMatrix, ProjectionData, CurrentLoadOffset);
+	LoadAndIncrementOffset(Result.TranslatedWorldToShadowUVNormalMatrix, ProjectionData, CurrentLoadOffset);
 
 	// NOTE: Stick with struct-element-sized loads for the moment since we may only be using subsets
 	// of the data in the calling code and we want to ensure the compiler has simple DCE options.
-
-	Result.LightDirection					= asfloat(ProjectionData.Load3(Base + (16*12 + 4*0)));
-	Result.LightType						=		 (ProjectionData.Load (Base + (16*12 + 4*3)));
+	LoadAndIncrementOffset(Result.LightDirection, ProjectionData, CurrentLoadOffset);
+	LoadAndIncrementOffset(Result.LightType, ProjectionData, CurrentLoadOffset);
 		 
-	float3 PreViewTranslationHigh			= asfloat(ProjectionData.Load3(Base + (16*13 + 4*0)));
-	Result.LightRadius						= asfloat(ProjectionData.Load (Base + (16*13 + 4*3)));
+	float3 PreViewTranslationHigh;
+	LoadAndIncrementOffset(PreViewTranslationHigh, ProjectionData, CurrentLoadOffset);
+	LoadAndIncrementOffset(Result.LightRadius, ProjectionData, CurrentLoadOffset);
 
-	float3 PreViewTranslationLow			= asfloat(ProjectionData.Load3(Base + (16*14 + 4*0)));
-	Result.ResolutionLodBias				= asfloat(ProjectionData.Load (Base + (16*14 + 4*3)));
+	float3 PreViewTranslationLow;
+	LoadAndIncrementOffset(PreViewTranslationLow, ProjectionData, CurrentLoadOffset);
+	LoadAndIncrementOffset(Result.ResolutionLodBias, ProjectionData, CurrentLoadOffset);
 
-	float3 NegativeClipmapWorldOriginOffset = asfloat(ProjectionData.Load3(Base + (16*15 + 4*0)));
-	Result.LightSourceRadius				= asfloat(ProjectionData.Load (Base + (16*15 + 4*3)));
+	float3 NegativeClipmapWorldOriginOffset;
+	LoadAndIncrementOffset(NegativeClipmapWorldOriginOffset, ProjectionData, CurrentLoadOffset);
+	LoadAndIncrementOffset(Result.LightSourceRadius, ProjectionData, CurrentLoadOffset);
 	
-	Result.ClipmapCornerRelativeOffset		= asint  (ProjectionData.Load2(Base + (16*16 + 4*0)));
+	LoadAndIncrementOffset(Result.ClipmapCornerRelativeOffset, ProjectionData, CurrentLoadOffset);
 
-	int ClipmapLevel_ClipmapLevelCountRemaining = asint  (ProjectionData.Load (Base + (16*16 + 4*2)));
-	if (ClipmapLevel_ClipmapLevelCountRemaining != -1)
+	uint ClipmapLevel_ClipmapLevelCountRemaining = 0;
+	LoadAndIncrementOffset(ClipmapLevel_ClipmapLevelCountRemaining, ProjectionData, CurrentLoadOffset);
+	if (int(ClipmapLevel_ClipmapLevelCountRemaining) != -1)
 	{
-		Result.ClipmapLevel						= (ClipmapLevel_ClipmapLevelCountRemaining >> 16u) - VSM_PACKED_CLIP_LEVEL_BIAS;
-		Result.ClipmapLevelCountRemaining		= ClipmapLevel_ClipmapLevelCountRemaining & 0xFFFF;
+		Result.ClipmapLevel						= int(ClipmapLevel_ClipmapLevelCountRemaining >> 16u) - VSM_PACKED_CLIP_LEVEL_BIAS;
+		Result.ClipmapLevelCountRemaining		= int(ClipmapLevel_ClipmapLevelCountRemaining & 0xFFFFu);
 	}
 	else
 	{
 		Result.ClipmapLevel						= VSM_PACKED_CLIP_LEVEL_BIAS;
 		Result.ClipmapLevelCountRemaining		= -1;
 	}
-	uint PackedCullingViewId		= (ProjectionData.Load (Base + (16*16 + 4*3)));
+	uint PackedCullingViewId;
+	LoadAndIncrementOffset(PackedCullingViewId, ProjectionData, CurrentLoadOffset);
 
 	Result.SceneRendererPrimaryViewId = int(PackedCullingViewId >> 16u) - 1;
 	Result.PersistentViewId = int(PackedCullingViewId & 0xFFFFu) - 1;
 
-	Result.Flags 							=		 (ProjectionData.Load (Base + (16*17 + 4*0)));
-	Result.ClipmapLevelWPODistanceDisabledThresholdSquared
-		                                    = asfloat(ProjectionData.Load (Base + (16*17 + 4*1)));
+	LoadAndIncrementOffset(Result.Flags, ProjectionData, CurrentLoadOffset);
+	LoadAndIncrementOffset(Result.ClipmapLevelWPODistanceDisabledThresholdSquared, ProjectionData, CurrentLoadOffset);
 	
-	Result.TexelDitherScale					= asfloat(ProjectionData.Load (Base + (16*17 + 4*2)));
-	Result.MinMipLevel						= ProjectionData.Load (Base + (16*17 + 4*3));
+	LoadAndIncrementOffset(Result.TexelDitherScale, ProjectionData, CurrentLoadOffset);
+	LoadAndIncrementOffset(Result.MinMipLevel, ProjectionData, CurrentLoadOffset);
 	
 	// Computed data
 	Result.PreViewTranslation				= MakeDFVector3(PreViewTranslationHigh, PreViewTranslationLow);
