LocalFogVolume - smoothly transition from the front clip planes of LFV to avoid bad straircase visual effect. (basically scale extinction as a fonction of the LFV front clip plane position inside the froxel along V).

#tests PC editor dx12
#rnx

diff --git a/Engine/Shaders/Private/VolumetricFog.usf b/Engine/Shaders/Private/VolumetricFog.usf
--- a/Engine/Shaders/Private/VolumetricFog.usf
+++ b/Engine/Shaders/Private/VolumetricFog.usf
@@ -150,9 +150,20 @@ void MaterialSetupCS(
 #if USE_LOCAL_FOG_VOLUMES
 	float CellOffset = .5f;
 	float SceneDepth;
-	float3 TranslatedWorldPosition = ComputeCellTranslatedWorldPosition(GridCoordinate, CellOffset, SceneDepth);
-	if(GetDistanceToCameraFromViewVector(TranslatedWorldPosition - PrimaryView.TranslatedWorldCameraOrigin) > LFVGlobalStartDistance)
+	const float3 TranslatedWorldPosition			= ComputeCellTranslatedWorldPosition(GridCoordinate, CellOffset, SceneDepth);
+	const float3 BackCellTranslatedWorldPosition	= ComputeCellTranslatedWorldPosition(GridCoordinate, float3(CellOffset, CellOffset, 1.0f), SceneDepth);
+	const float  DistanceToCellBack					= GetDistanceToCameraFromViewVector(BackCellTranslatedWorldPosition - PrimaryView.TranslatedWorldCameraOrigin);
+	if(DistanceToCellBack > LFVGlobalStartDistance)
 	{
+		const float3 FrontCellTranslatedWorldPosition	= ComputeCellTranslatedWorldPosition(GridCoordinate, float3(CellOffset, CellOffset, 0.0f), SceneDepth);
+		const float DistanceToCellFront					= GetDistanceToCameraFromViewVector(FrontCellTranslatedWorldPosition - PrimaryView.TranslatedWorldCameraOrigin);
+		float SoftDensityScale = 1;
+		if (DistanceToCellFront < LFVGlobalStartDistance && LFVGlobalStartDistance < DistanceToCellBack)
+		{
+			// Smoothly transition from the front clip planes of LFV to avoid bad straircase/aliasing visual effect.
+			// Extinction is basically scaled as a fonction of the start distance position inside the froxel along V.
+			SoftDensityScale = (DistanceToCellBack - LFVGlobalStartDistance) / max(0.0001, DistanceToCellBack - DistanceToCellFront);
+		}
 
 		float4 VertexClipSpacePosition = mul(float4(TranslatedWorldPosition.xyz, 1), PrimaryView.TranslatedWorldToClip);
 		float2 SvPosition = (VertexClipSpacePosition.xy / VertexClipSpacePosition.w * float2(.5f, -.5f) + .5f) * PrimaryView.ViewSizeAndInvSize.xy;
@@ -211,7 +222,7 @@ void MaterialSetupCS(
 			// Convert extinction from unit volume 1/m into extinction for volumetric fog as 1/cm
 			LFVExtinction *= (1.0f / METER_TO_CENTIMETER);
 
-			LFVExtinction = min(LFVMaxDensityIntoVolumetricFog, LFVExtinction);
+			LFVExtinction = min(LFVMaxDensityIntoVolumetricFog, LFVExtinction) * SoftDensityScale;
 
 			float3 LFVScattering = LFVExtinction * FogInstance.Albedo;
 
